[GLTFExporter] chore: initial empty gltf exporter plugin

commit fa26e028a7106aef83356184ea13a361df4b92cc
Date: Thu, 19 Mar 2020 08:11:14 +0100

    chore: initial empty gltf exporter plugin

diff --git a/Engine/Plugins/Enterprise/GLTFExporter/MigratedFromGit/GLTFExporter.uplugin b/Engine/Plugins/Enterprise/GLTFExporter/MigratedFromGit/GLTFExporter.uplugin
new file mode 100644
--- /dev/null
+++ b/Engine/Plugins/Enterprise/GLTFExporter/MigratedFromGit/GLTFExporter.uplugin
@@ -0,0 +1,24 @@
+{
+	"FileVersion": 3,
+	"Version": 1.0,
+	"VersionName": "1.0",
+	"FriendlyName": "glTF Exporter",
+	"Description": "An exporter for Khronos glTF 3D models.",
+	"Category": "Exporters",
+	"CreatedBy": "Epic Games, Inc.",
+	"CreatedByURL": "http://epicgames.com",
+	"EnabledByDefault": false,
+	"CanContainContent": true,
+	"IsBetaVersion": true,
+	"Modules":
+	[
+		{
+			"Name": "GLTFCore",
+			"Type": "Runtime"
+		},
+		{
+			"Name": "GLTFExporter",
+			"Type": "Editor"
+		}
+	]
+}
diff --git a/Engine/Plugins/Enterprise/GLTFExporter/MigratedFromGit/Source/GLTFCore/GLTFCore.Build.cs b/Engine/Plugins/Enterprise/GLTFExporter/MigratedFromGit/Source/GLTFCore/GLTFCore.Build.cs
new file mode 100644
--- /dev/null
+++ b/Engine/Plugins/Enterprise/GLTFExporter/MigratedFromGit/Source/GLTFCore/GLTFCore.Build.cs
@@ -0,0 +1,29 @@
+// Copyright Epic Games, Inc. All Rights Reserved.
+
+namespace UnrealBuildTool.Rules
+{
+    public class GLTFCore : ModuleRules
+    {
+        public GLTFCore(ReadOnlyTargetRules Target) : base(Target)
+        {
+            PublicDependencyModuleNames.AddRange(
+                new string[]
+                {
+                    "Core",
+                    "CoreUObject",
+                }
+                );
+
+            PrivateDependencyModuleNames.AddRange(
+                new string[]
+                {
+                    "Engine",
+                    "MeshDescription",
+					"StaticMeshDescription",
+                    "Json",
+                    "RenderCore",
+                }
+                );
+        }
+    }
+}
diff --git a/Engine/Plugins/Enterprise/GLTFExporter/MigratedFromGit/Source/GLTFCore/Private/GLTFCoreModule.cpp b/Engine/Plugins/Enterprise/GLTFExporter/MigratedFromGit/Source/GLTFCore/Private/GLTFCoreModule.cpp
new file mode 100644
--- /dev/null
+++ b/Engine/Plugins/Enterprise/GLTFExporter/MigratedFromGit/Source/GLTFCore/Private/GLTFCoreModule.cpp
@@ -0,0 +1,20 @@
+// Copyright Epic Games, Inc. All Rights Reserved.
+
+#include "GLTFCoreModule.h"
+
+#include "CoreMinimal.h"
+#include "Modules/ModuleManager.h"
+
+/**
+ * glTF Core module implementation (private)
+ */
+class FGLTFCoreModule : public IGLTFCoreModule
+{
+
+public:
+	virtual void StartupModule() override {}
+
+	virtual void ShutdownModule() override {}
+};
+
+IMPLEMENT_MODULE(FGLTFCoreModule, GLTFCore);
\ No newline at end of file
diff --git a/Engine/Plugins/Enterprise/GLTFExporter/MigratedFromGit/Source/GLTFCore/Public/GLTFCoreModule.h b/Engine/Plugins/Enterprise/GLTFExporter/MigratedFromGit/Source/GLTFCore/Public/GLTFCoreModule.h
new file mode 100644
--- /dev/null
+++ b/Engine/Plugins/Enterprise/GLTFExporter/MigratedFromGit/Source/GLTFCore/Public/GLTFCoreModule.h
@@ -0,0 +1,34 @@
+// Copyright Epic Games, Inc. All Rights Reserved.
+
+#pragma once
+
+#include "CoreMinimal.h"
+#include "Modules/ModuleInterface.h"
+#include "Modules/ModuleManager.h"
+
+/**
+ * The public interface of the GLTFCore module
+ */
+class IGLTFCoreModule : public IModuleInterface
+{
+public:
+	/**
+	 * Singleton-like access to IGLTFCore
+	 *
+	 * @return Returns IGLTFCore singleton instance, loading the module on demand if needed
+	 */
+	static IGLTFCoreModule& Get()
+	{
+		return FModuleManager::LoadModuleChecked<IGLTFCoreModule>("GLTFCore");
+	}
+
+	/**
+	 * Checks to see if this module is loaded and ready.  It is only valid to call Get() if IsAvailable() returns true.
+	 *
+	 * @return True if the module is loaded and ready to use
+	 */
+	static bool IsAvailable()
+	{
+		return FModuleManager::Get().IsModuleLoaded("GLTFCore");
+	}
+};
diff --git a/Engine/Plugins/Enterprise/GLTFExporter/MigratedFromGit/Source/GLTFExporter/GLTFExporter.Build.cs b/Engine/Plugins/Enterprise/GLTFExporter/MigratedFromGit/Source/GLTFExporter/GLTFExporter.Build.cs
new file mode 100644
--- /dev/null
+++ b/Engine/Plugins/Enterprise/GLTFExporter/MigratedFromGit/Source/GLTFExporter/GLTFExporter.Build.cs
@@ -0,0 +1,54 @@
+// Copyright Epic Games, Inc. All Rights Reserved.
+
+namespace UnrealBuildTool.Rules
+{
+    public class GLTFExporter : ModuleRules
+    {
+        public GLTFExporter(ReadOnlyTargetRules Target) : base(Target)
+        {
+            PublicIncludePaths.AddRange(
+                new string[] {
+                }
+                );
+
+            PrivateIncludePaths.AddRange(
+                new string[] {
+                }
+                );
+
+            PublicDependencyModuleNames.AddRange(
+                new string[]
+                {
+                }
+                );
+
+            PrivateDependencyModuleNames.AddRange(
+                new string[]
+                {
+                    "Core",
+                    "CoreUObject",
+                    "Engine",
+                    "UnrealEd",
+                    "MeshDescription",
+					"StaticMeshDescription",
+                    "MeshUtilities",
+                    "MessageLog",
+                    "Json",
+                    "MaterialEditor",
+                    "Slate",
+                    "SlateCore",
+                    "Mainframe",
+                    "InputCore",
+                    "EditorStyle",
+					"GLTFCore",
+                }
+                );
+
+            DynamicallyLoadedModuleNames.AddRange(
+                new string[]
+                {
+                }
+                );
+        }
+    }
+}
diff --git a/Engine/Plugins/Enterprise/GLTFExporter/MigratedFromGit/Source/GLTFExporter/Private/GLTFExporterModule.cpp b/Engine/Plugins/Enterprise/GLTFExporter/MigratedFromGit/Source/GLTFExporter/Private/GLTFExporterModule.cpp
new file mode 100644
--- /dev/null
+++ b/Engine/Plugins/Enterprise/GLTFExporter/MigratedFromGit/Source/GLTFExporter/Private/GLTFExporterModule.cpp
@@ -0,0 +1,20 @@
+// Copyright Epic Games, Inc. All Rights Reserved.
+
+#include "GLTFExporterModule.h"
+
+#include "CoreMinimal.h"
+#include "Modules/ModuleManager.h"
+
+/**
+ * glTF Exporter module implementation (private)
+ */
+class FGLTFExporterModule : public IGLTFExporterModule
+{
+
+public:
+	virtual void StartupModule() override {}
+
+	virtual void ShutdownModule() override {}
+};
+
+IMPLEMENT_MODULE(FGLTFExporterModule, GLTFExporter);
diff --git a/Engine/Plugins/Enterprise/GLTFExporter/MigratedFromGit/Source/GLTFExporter/Public/GLTFExporterModule.h b/Engine/Plugins/Enterprise/GLTFExporter/MigratedFromGit/Source/GLTFExporter/Public/GLTFExporterModule.h
new file mode 100644
--- /dev/null
+++ b/Engine/Plugins/Enterprise/GLTFExporter/MigratedFromGit/Source/GLTFExporter/Public/GLTFExporterModule.h
@@ -0,0 +1,34 @@
+// Copyright Epic Games, Inc. All Rights Reserved.
+
+#pragma once
+
+#include "CoreMinimal.h"
+#include "Modules/ModuleInterface.h"
+#include "Modules/ModuleManager.h"
+
+/**
+ * The public interface of the GLTFExporter module
+ */
+class IGLTFExporterModule : public IModuleInterface
+{
+public:
+	/**
+	 * Singleton-like access to IGLTFExporter
+	 *
+	 * @return Returns IGLTFExporter singleton instance, loading the module on demand if needed
+	 */
+	static IGLTFExporterModule& Get()
+	{
+		return FModuleManager::LoadModuleChecked<IGLTFExporterModule>("GLTFExporter");
+	}
+
+	/**
+	 * Checks to see if this module is loaded and ready.  It is only valid to call Get() if IsAvailable() returns true.
+	 *
+	 * @return True if the module is loaded and ready to use
+	 */
+	static bool IsAvailable()
+	{
+		return FModuleManager::Get().IsModuleLoaded("GLTFExporter");
+	}
+};
