[GLTFExporter] feat: add ULevelExporterGLTF with temporary logging

commit 971a3afa90ac5912353f3c600500decac51fc968
Date: Thu, 19 Mar 2020 20:14:50 +0100

    feat: add ULevelExporterGLTF with temporary logging

#skipci

diff --git a/Engine/Plugins/Enterprise/GLTFExporter/MigratedFromGit/Source/GLTFExporter/Private/LevelExporterGLTF.cpp b/Engine/Plugins/Enterprise/GLTFExporter/MigratedFromGit/Source/GLTFExporter/Private/LevelExporterGLTF.cpp
new file mode 100644
--- /dev/null
+++ b/Engine/Plugins/Enterprise/GLTFExporter/MigratedFromGit/Source/GLTFExporter/Private/LevelExporterGLTF.cpp
@@ -0,0 +1,128 @@
+// Copyright Epic Games, Inc. All Rights Reserved.
+
+#include "LevelExporterGLTF.h"
+#include "UnrealExporter.h"
+#include "AssetExportTask.h"
+#include "Engine/World.h"
+#include "Engine.h"
+
+DEFINE_LOG_CATEGORY(LogGLTFExporter);
+
+ULevelExporterGLTF::ULevelExporterGLTF(const FObjectInitializer& ObjectInitializer)
+	: Super(ObjectInitializer)
+{
+	SupportedClass = UWorld::StaticClass();
+	bText = false;
+	PreferredFormatIndex = 0;
+	FormatExtension.Add(TEXT("gltf"));
+	FormatExtension.Add(TEXT("glb"));
+	FormatDescription.Add(TEXT("GL Transmission Format"));
+	FormatDescription.Add(TEXT("GL Transmission Format (Binary)"));
+}
+
+bool ULevelExporterGLTF::ExportText(const FExportObjectInnerContext* Context, UObject* Object, const TCHAR* Type, FOutputDevice& Ar, FFeedbackContext* Warn, uint32 PortFlags)
+{
+	UE_LOG(LogGLTFExporter, Warning, TEXT("ULevelExporterGLTF::ExportText"));
+	UE_LOG(LogGLTFExporter, Warning, TEXT("Object: %s (%s)"), *(Object->GetName()), *(Object->GetClass()->GetName()));
+	UE_LOG(LogGLTFExporter, Warning, TEXT("Type: %s"), Type);
+	UE_LOG(LogGLTFExporter, Warning, TEXT("PortFlags: %d"), PortFlags);
+	return true;
+}
+
+bool ULevelExporterGLTF::ExportBinary(UObject* Object, const TCHAR* Type, FArchive& Ar, FFeedbackContext* Warn, int32 FileIndex, uint32 PortFlags)
+{
+	UE_LOG(LogGLTFExporter, Warning, TEXT("ULevelExporterGLTF::ExportBinary"));
+	UE_LOG(LogGLTFExporter, Warning, TEXT("Object: %s (%s)"), *(Object->GetName()), *(Object->GetClass()->GetName()));
+	UE_LOG(LogGLTFExporter, Warning, TEXT("Type: %s"), Type);
+	UE_LOG(LogGLTFExporter, Warning, TEXT("FileIndex: %d"), FileIndex);
+	UE_LOG(LogGLTFExporter, Warning, TEXT("PortFlags: %d"), PortFlags);
+	return true;
+}
+
+void ULevelExporterGLTF::ExportComponentExtra(const FExportObjectInnerContext* Context, const TArray<UActorComponent*>& Components, FOutputDevice& Ar, uint32 PortFlags)
+{
+	UE_LOG(LogGLTFExporter, Warning, TEXT("ULevelExporterGLTF::ExportComponentExtra"));
+	UE_LOG(LogGLTFExporter, Warning, TEXT("Components: %d"), Components.Num());
+	for (int32 i = 0; i != Components.Num(); ++i)
+	{
+		auto& Component = Components[i];
+		UE_LOG(LogGLTFExporter, Warning, TEXT("Components[%d]: %s (%s)"), i, *(Component->GetName()), *(Component->GetClass()->GetName()));
+	}
+	UE_LOG(LogGLTFExporter, Warning, TEXT("PortFlags: %d"), PortFlags);
+}
+
+void ULevelExporterGLTF::ExportPackageObject(FExportPackageParams& ExpPackageParams)
+{
+	UE_LOG(LogGLTFExporter, Warning, TEXT("ULevelExporterGLTF::ExportPackageObject"));
+	UE_LOG(LogGLTFExporter, Warning, TEXT("ExpPackageParams.Type: %s"), ExpPackageParams.Type);
+	UE_LOG(LogGLTFExporter, Warning, TEXT("ExpPackageParams.RootMapPackageName: %d"), *(ExpPackageParams.RootMapPackageName));
+	UE_LOG(LogGLTFExporter, Warning, TEXT("PortFlags: %d"), ExpPackageParams.PortFlags);
+}
+
+void ULevelExporterGLTF::ExportPackageInners(FExportPackageParams& ExpPackageParams)
+{
+	UE_LOG(LogGLTFExporter, Warning, TEXT("ULevelExporterGLTF::ExportPackageInners"));
+	UE_LOG(LogGLTFExporter, Warning, TEXT("ExpPackageParams.Type: %s"), ExpPackageParams.Type);
+	UE_LOG(LogGLTFExporter, Warning, TEXT("ExpPackageParams.RootMapPackageName: %d"), *(ExpPackageParams.RootMapPackageName));
+	UE_LOG(LogGLTFExporter, Warning, TEXT("PortFlags: %d"), ExpPackageParams.PortFlags);
+}
+
+bool ULevelExporterGLTF::SupportsObject(UObject* Object) const
+{
+	UE_LOG(LogGLTFExporter, Warning, TEXT("ULevelExporterGLTF::SupportsObject"));
+	UE_LOG(LogGLTFExporter, Warning, TEXT("Object: %s (%s)"), *(Object->GetName()), *(Object->GetClass()->GetName()));
+	return UExporter::SupportsObject(Object);
+}
+
+int32 ULevelExporterGLTF::GetFileCount() const
+{
+	UE_LOG(LogGLTFExporter, Warning, TEXT("ULevelExporterGLTF::GetFileCount"));
+	return UExporter::GetFileCount();
+}
+
+FString ULevelExporterGLTF::GetUniqueFilename(const TCHAR * Filename, int32 FileIndex)
+{
+	UE_LOG(LogGLTFExporter, Warning, TEXT("ULevelExporterGLTF::GetUniqueFilename"));
+	UE_LOG(LogGLTFExporter, Warning, TEXT("Filename: %s"), Filename);
+	UE_LOG(LogGLTFExporter, Warning, TEXT("FileIndex: %d"), FileIndex);
+	return UExporter::GetUniqueFilename(Filename, FileIndex);
+}
+
+bool ULevelExporterGLTF::GetBatchMode() const
+{
+	UE_LOG(LogGLTFExporter, Warning, TEXT("ULevelExporterGLTF::GetBatchMode"));
+	return UExporter::GetBatchMode();
+}
+
+void ULevelExporterGLTF::SetBatchMode(bool InBatchExportMode)
+{
+	UE_LOG(LogGLTFExporter, Warning, TEXT("ULevelExporterGLTF::SetBatchMode"));
+	UE_LOG(LogGLTFExporter, Warning, TEXT("InBatchExportMode: %s"), (InBatchExportMode ? TEXT("true") : TEXT("false")));
+	UExporter::SetBatchMode(InBatchExportMode);
+}
+
+bool ULevelExporterGLTF::GetCancelBatch() const
+{
+	UE_LOG(LogGLTFExporter, Warning, TEXT("ULevelExporterGLTF::GetCancelBatch"));
+	return UExporter::GetCancelBatch();
+}
+
+void ULevelExporterGLTF::SetCancelBatch(bool InCancelBatch)
+{
+	UE_LOG(LogGLTFExporter, Warning, TEXT("ULevelExporterGLTF::SetCancelBatch"));
+	UE_LOG(LogGLTFExporter, Warning, TEXT("InCancelBatch: %s"), (InCancelBatch ? TEXT("true") : TEXT("false")));
+	UExporter::SetCancelBatch(InCancelBatch);
+}
+
+bool ULevelExporterGLTF::GetShowExportOption() const
+{
+	UE_LOG(LogGLTFExporter, Warning, TEXT("ULevelExporterGLTF::GetShowExportOption"));
+	return UExporter::GetShowExportOption();
+}
+
+void ULevelExporterGLTF::SetShowExportOption(bool InShowExportOption)
+{
+	UE_LOG(LogGLTFExporter, Warning, TEXT("ULevelExporterGLTF::SetShowExportOption"));
+	UE_LOG(LogGLTFExporter, Warning, TEXT("InShowExportOption: %s"), (InShowExportOption ? TEXT("true") : TEXT("false")));
+	UExporter::SetShowExportOption(InShowExportOption);
+}
diff --git a/Engine/Plugins/Enterprise/GLTFExporter/MigratedFromGit/Source/GLTFExporter/Public/LevelExporterGLTF.h b/Engine/Plugins/Enterprise/GLTFExporter/MigratedFromGit/Source/GLTFExporter/Public/LevelExporterGLTF.h
new file mode 100644
--- /dev/null
+++ b/Engine/Plugins/Enterprise/GLTFExporter/MigratedFromGit/Source/GLTFExporter/Public/LevelExporterGLTF.h
@@ -0,0 +1,48 @@
+// Copyright Epic Games, Inc. All Rights Reserved.
+
+#pragma once
+
+#include "CoreMinimal.h"
+#include "UObject/ObjectMacros.h"
+#include "UObject/UObjectGlobals.h"
+#include "Exporters/Exporter.h"
+#include "LevelExporterGLTF.generated.h"
+
+class FExportObjectInnerContext;
+class UActorComponent;
+
+DECLARE_LOG_CATEGORY_EXTERN(LogGLTFExporter, Log, All);
+
+UCLASS()
+class ULevelExporterGLTF : public UExporter
+{
+public:
+	GENERATED_BODY()
+
+public:
+	ULevelExporterGLTF(const FObjectInitializer& ObjectInitializer = FObjectInitializer());
+
+	//~ Begin UExporter Interface
+
+	virtual bool SupportsObject(UObject* Object) const override;
+
+	virtual bool ExportText(const FExportObjectInnerContext* Context, UObject* Object, const TCHAR* Type, FOutputDevice& Ar, FFeedbackContext* Warn, uint32 PortFlags = 0) override;
+	virtual bool ExportBinary(UObject* Object, const TCHAR* Type, FArchive& Ar, FFeedbackContext* Warn, int32 FileIndex = 0, uint32 PortFlags = 0) override;
+	virtual void ExportComponentExtra(const FExportObjectInnerContext* Context, const TArray<UActorComponent*>& Components, FOutputDevice& Ar, uint32 PortFlags) override;
+
+	virtual void ExportPackageObject(FExportPackageParams& ExpPackageParams) override;
+	virtual void ExportPackageInners(FExportPackageParams& ExpPackageParams) override;
+
+	virtual int32 GetFileCount(void) const override;
+	virtual FString GetUniqueFilename(const TCHAR* Filename, int32 FileIndex) override;
+
+	virtual bool GetBatchMode() const override;
+	virtual void SetBatchMode(bool InBatchExportMode) override;
+	virtual bool GetCancelBatch() const override;
+	virtual void SetCancelBatch(bool InCancelBatch) override;
+
+	virtual bool GetShowExportOption() const override;
+	virtual void SetShowExportOption(bool InShowExportOption) override;
+
+	//~ End UExporter Interface
+};
