[GLTFExporter] refactor: rename LevelExporterGLTF to GLTFLevelExporter and add base class GLTFExporter

commit 75ddf40665e3a6005040f43d83a947218305989b
Date: Thu, 19 Mar 2020 22:29:27 +0100

    refactor: rename LevelExporterGLTF to GLTFLevelExporter and add base class GLTFExporter

#skipci

diff --git a/Engine/Plugins/Enterprise/GLTFExporter/MigratedFromGit/Source/GLTFExporter/Private/GLTFExporter.cpp b/Engine/Plugins/Enterprise/GLTFExporter/MigratedFromGit/Source/GLTFExporter/Private/GLTFExporter.cpp
new file mode 100644
--- /dev/null
+++ b/Engine/Plugins/Enterprise/GLTFExporter/MigratedFromGit/Source/GLTFExporter/Private/GLTFExporter.cpp
@@ -0,0 +1,29 @@
+// Copyright Epic Games, Inc. All Rights Reserved.
+
+#include "GLTFExporter.h"
+#include "UnrealExporter.h"
+#include "Engine.h"
+
+DEFINE_LOG_CATEGORY(LogGLTFExporter);
+
+UGLTFExporter::UGLTFExporter(const FObjectInitializer& ObjectInitializer)
+	: Super(ObjectInitializer)
+{
+	SupportedClass = nullptr;
+	bText = false;
+	PreferredFormatIndex = 0;
+	FormatExtension.Add(TEXT("gltf"));
+	FormatExtension.Add(TEXT("glb"));
+	FormatDescription.Add(TEXT("GL Transmission Format"));
+	FormatDescription.Add(TEXT("GL Transmission Format (Binary)"));
+}
+
+bool UGLTFExporter::ExportBinary(UObject* Object, const TCHAR* Type, FArchive& Ar, FFeedbackContext* Warn, int32 FileIndex, uint32 PortFlags)
+{
+	UE_LOG(LogGLTFExporter, Warning, TEXT("%s::ExportBinary"), *(this->GetClass()->GetName()));
+	UE_LOG(LogGLTFExporter, Warning, TEXT("Object: %s (%s)"), *(Object->GetName()), *(Object->GetClass()->GetName()));
+	UE_LOG(LogGLTFExporter, Warning, TEXT("Type: %s"), Type);
+	UE_LOG(LogGLTFExporter, Warning, TEXT("FileIndex: %d"), FileIndex);
+	UE_LOG(LogGLTFExporter, Warning, TEXT("PortFlags: %d"), PortFlags);
+	return true;
+}
diff --git a/Engine/Plugins/Enterprise/GLTFExporter/MigratedFromGit/Source/GLTFExporter/Private/GLTFExporterModule.cpp b/Engine/Plugins/Enterprise/GLTFExporter/MigratedFromGit/Source/GLTFExporter/Private/GLTFExporterModule.cpp
--- a/Engine/Plugins/Enterprise/GLTFExporter/MigratedFromGit/Source/GLTFExporter/Private/GLTFExporterModule.cpp
+++ b/Engine/Plugins/Enterprise/GLTFExporter/MigratedFromGit/Source/GLTFExporter/Private/GLTFExporterModule.cpp
@@ -5,8 +5,6 @@
 #include "CoreMinimal.h"
 #include "Modules/ModuleManager.h"
 
-DEFINE_LOG_CATEGORY(LogGLTFExporter);
-
 /**
  * glTF Exporter module implementation (private)
  */
diff --git a/Engine/Plugins/Enterprise/GLTFExporter/MigratedFromGit/Source/GLTFExporter/Private/GLTFLevelExporter.cpp b/Engine/Plugins/Enterprise/GLTFExporter/MigratedFromGit/Source/GLTFExporter/Private/GLTFLevelExporter.cpp
new file mode 100644
--- /dev/null
+++ b/Engine/Plugins/Enterprise/GLTFExporter/MigratedFromGit/Source/GLTFExporter/Private/GLTFLevelExporter.cpp
@@ -0,0 +1,16 @@
+// Copyright Epic Games, Inc. All Rights Reserved.
+
+#include "GLTFLevelExporter.h"
+#include "Engine/World.h"
+
+UGLTFLevelExporter::UGLTFLevelExporter(const FObjectInitializer& ObjectInitializer)
+	: Super(ObjectInitializer)
+{
+	SupportedClass = UWorld::StaticClass();
+}
+
+bool UGLTFLevelExporter::ExportBinary(UObject* Object, const TCHAR* Type, FArchive& Ar, FFeedbackContext* Warn, int32 FileIndex, uint32 PortFlags)
+{
+	UWorld* World = CastChecked<UWorld>(Object);
+	return Super::ExportBinary(World, Type, Ar, Warn, FileIndex, PortFlags);
+}
diff --git a/Engine/Plugins/Enterprise/GLTFExporter/MigratedFromGit/Source/GLTFExporter/Private/LevelExporterGLTF.cpp b/Engine/Plugins/Enterprise/GLTFExporter/MigratedFromGit/Source/GLTFExporter/Private/LevelExporterGLTF.cpp
deleted file mode 100644
diff --git a/Engine/Plugins/Enterprise/GLTFExporter/MigratedFromGit/Source/GLTFExporter/Public/GLTFExporter.h b/Engine/Plugins/Enterprise/GLTFExporter/MigratedFromGit/Source/GLTFExporter/Public/GLTFExporter.h
new file mode 100644
--- /dev/null
+++ b/Engine/Plugins/Enterprise/GLTFExporter/MigratedFromGit/Source/GLTFExporter/Public/GLTFExporter.h
@@ -0,0 +1,25 @@
+// Copyright Epic Games, Inc. All Rights Reserved.
+
+#pragma once
+
+#include "CoreMinimal.h"
+#include "UObject/ObjectMacros.h"
+#include "UObject/UObjectGlobals.h"
+#include "Exporters/Exporter.h"
+#include "GLTFExporter.generated.h"
+
+DECLARE_LOG_CATEGORY_EXTERN(LogGLTFExporter, Log, All);
+
+UCLASS()
+class UGLTFExporter : public UExporter
+{
+public:
+	GENERATED_BODY()
+
+public:
+	UGLTFExporter(const FObjectInitializer& ObjectInitializer = FObjectInitializer());
+
+	//~ Begin UExporter Interface
+	virtual bool ExportBinary(UObject* Object, const TCHAR* Type, FArchive& Ar, FFeedbackContext* Warn, int32 FileIndex = 0, uint32 PortFlags = 0) override;
+	//~ End UExporter Interface
+};
diff --git a/Engine/Plugins/Enterprise/GLTFExporter/MigratedFromGit/Source/GLTFExporter/Public/GLTFExporterModule.h b/Engine/Plugins/Enterprise/GLTFExporter/MigratedFromGit/Source/GLTFExporter/Public/GLTFExporterModule.h
--- a/Engine/Plugins/Enterprise/GLTFExporter/MigratedFromGit/Source/GLTFExporter/Public/GLTFExporterModule.h
+++ b/Engine/Plugins/Enterprise/GLTFExporter/MigratedFromGit/Source/GLTFExporter/Public/GLTFExporterModule.h
@@ -6,8 +6,6 @@
 #include "Modules/ModuleInterface.h"
 #include "Modules/ModuleManager.h"
 
-DECLARE_LOG_CATEGORY_EXTERN(LogGLTFExporter, Log, All);
-
 /**
  * The public interface of the GLTFExporter module
  */
diff --git a/Engine/Plugins/Enterprise/GLTFExporter/MigratedFromGit/Source/GLTFExporter/Public/GLTFLevelExporter.h b/Engine/Plugins/Enterprise/GLTFExporter/MigratedFromGit/Source/GLTFExporter/Public/GLTFLevelExporter.h
new file mode 100644
--- /dev/null
+++ b/Engine/Plugins/Enterprise/GLTFExporter/MigratedFromGit/Source/GLTFExporter/Public/GLTFLevelExporter.h
@@ -0,0 +1,20 @@
+// Copyright Epic Games, Inc. All Rights Reserved.
+
+#pragma once
+
+#include "GLTFExporter.h"
+#include "GLTFLevelExporter.generated.h"
+
+UCLASS()
+class UGLTFLevelExporter : public UGLTFExporter
+{
+public:
+	GENERATED_BODY()
+
+public:
+	UGLTFLevelExporter(const FObjectInitializer& ObjectInitializer = FObjectInitializer());
+
+	//~ Begin UExporter Interface
+	virtual bool ExportBinary(UObject* Object, const TCHAR* Type, FArchive& Ar, FFeedbackContext* Warn, int32 FileIndex = 0, uint32 PortFlags = 0) override;
+	//~ End UExporter Interface
+};
diff --git a/Engine/Plugins/Enterprise/GLTFExporter/MigratedFromGit/Source/GLTFExporter/Public/LevelExporterGLTF.h b/Engine/Plugins/Enterprise/GLTFExporter/MigratedFromGit/Source/GLTFExporter/Public/LevelExporterGLTF.h
deleted file mode 100644
