[GLTFExporter] feat: add exporters for AnimSequence, SkeletalMesh, and StaticMesh

commit 771053d151dde251f9974dc008e0d854a4914330
Date: Thu, 19 Mar 2020 22:29:56 +0100

    feat: add exporters for AnimSequence, SkeletalMesh, and StaticMesh

#skipci

diff --git a/Engine/Plugins/Enterprise/GLTFExporter/MigratedFromGit/Source/GLTFExporter/Private/GLTFAnimSequenceExporter.cpp b/Engine/Plugins/Enterprise/GLTFExporter/MigratedFromGit/Source/GLTFExporter/Private/GLTFAnimSequenceExporter.cpp
new file mode 100644
--- /dev/null
+++ b/Engine/Plugins/Enterprise/GLTFExporter/MigratedFromGit/Source/GLTFExporter/Private/GLTFAnimSequenceExporter.cpp
@@ -0,0 +1,16 @@
+// Copyright Epic Games, Inc. All Rights Reserved.
+
+#include "GLTFAnimSequenceExporter.h"
+#include "Animation/AnimSequence.h"
+
+UGLTFAnimSequenceExporter::UGLTFAnimSequenceExporter(const FObjectInitializer& ObjectInitializer)
+	: Super(ObjectInitializer)
+{
+	SupportedClass = UAnimSequence::StaticClass();
+}
+
+bool UGLTFAnimSequenceExporter::ExportBinary(UObject* Object, const TCHAR* Type, FArchive& Ar, FFeedbackContext* Warn, int32 FileIndex, uint32 PortFlags)
+{
+	UAnimSequence* AnimSequence = CastChecked<UAnimSequence>(Object);
+	return Super::ExportBinary(AnimSequence, Type, Ar, Warn, FileIndex, PortFlags);
+}
diff --git a/Engine/Plugins/Enterprise/GLTFExporter/MigratedFromGit/Source/GLTFExporter/Private/GLTFSkeletalMeshExporter.cpp b/Engine/Plugins/Enterprise/GLTFExporter/MigratedFromGit/Source/GLTFExporter/Private/GLTFSkeletalMeshExporter.cpp
new file mode 100644
--- /dev/null
+++ b/Engine/Plugins/Enterprise/GLTFExporter/MigratedFromGit/Source/GLTFExporter/Private/GLTFSkeletalMeshExporter.cpp
@@ -0,0 +1,16 @@
+// Copyright Epic Games, Inc. All Rights Reserved.
+
+#include "GLTFSkeletalMeshExporter.h"
+#include "Engine/SkeletalMesh.h"
+
+UGLTFSkeletalMeshExporter::UGLTFSkeletalMeshExporter(const FObjectInitializer& ObjectInitializer)
+	: Super(ObjectInitializer)
+{
+	SupportedClass = USkeletalMesh::StaticClass();
+}
+
+bool UGLTFSkeletalMeshExporter::ExportBinary(UObject* Object, const TCHAR* Type, FArchive& Ar, FFeedbackContext* Warn, int32 FileIndex, uint32 PortFlags)
+{
+	USkeletalMesh* SkeletalMesh = CastChecked<USkeletalMesh>(Object);
+	return Super::ExportBinary(SkeletalMesh, Type, Ar, Warn, FileIndex, PortFlags);
+}
diff --git a/Engine/Plugins/Enterprise/GLTFExporter/MigratedFromGit/Source/GLTFExporter/Private/GLTFStaticMeshExporter.cpp b/Engine/Plugins/Enterprise/GLTFExporter/MigratedFromGit/Source/GLTFExporter/Private/GLTFStaticMeshExporter.cpp
new file mode 100644
--- /dev/null
+++ b/Engine/Plugins/Enterprise/GLTFExporter/MigratedFromGit/Source/GLTFExporter/Private/GLTFStaticMeshExporter.cpp
@@ -0,0 +1,16 @@
+// Copyright Epic Games, Inc. All Rights Reserved.
+
+#include "GLTFStaticMeshExporter.h"
+#include "Engine/StaticMesh.h"
+
+UGLTFStaticMeshExporter::UGLTFStaticMeshExporter(const FObjectInitializer& ObjectInitializer)
+	: Super(ObjectInitializer)
+{
+	SupportedClass = UStaticMesh::StaticClass();
+}
+
+bool UGLTFStaticMeshExporter::ExportBinary(UObject* Object, const TCHAR* Type, FArchive& Ar, FFeedbackContext* Warn, int32 FileIndex, uint32 PortFlags)
+{
+	UStaticMesh* StaticMesh = CastChecked<UStaticMesh>(Object);
+	return Super::ExportBinary(StaticMesh, Type, Ar, Warn, FileIndex, PortFlags);
+}
diff --git a/Engine/Plugins/Enterprise/GLTFExporter/MigratedFromGit/Source/GLTFExporter/Public/GLTFAnimSequenceExporter.h b/Engine/Plugins/Enterprise/GLTFExporter/MigratedFromGit/Source/GLTFExporter/Public/GLTFAnimSequenceExporter.h
new file mode 100644
--- /dev/null
+++ b/Engine/Plugins/Enterprise/GLTFExporter/MigratedFromGit/Source/GLTFExporter/Public/GLTFAnimSequenceExporter.h
@@ -0,0 +1,20 @@
+// Copyright Epic Games, Inc. All Rights Reserved.
+
+#pragma once
+
+#include "GLTFExporter.h"
+#include "GLTFAnimSequenceExporter.generated.h"
+
+UCLASS()
+class UGLTFAnimSequenceExporter : public UGLTFExporter
+{
+public:
+	GENERATED_BODY()
+
+public:
+	UGLTFAnimSequenceExporter(const FObjectInitializer& ObjectInitializer = FObjectInitializer());
+
+	//~ Begin UExporter Interface
+	virtual bool ExportBinary(UObject* Object, const TCHAR* Type, FArchive& Ar, FFeedbackContext* Warn, int32 FileIndex = 0, uint32 PortFlags = 0) override;
+	//~ End UExporter Interface
+};
diff --git a/Engine/Plugins/Enterprise/GLTFExporter/MigratedFromGit/Source/GLTFExporter/Public/GLTFSkeletalMeshExporter.h b/Engine/Plugins/Enterprise/GLTFExporter/MigratedFromGit/Source/GLTFExporter/Public/GLTFSkeletalMeshExporter.h
new file mode 100644
--- /dev/null
+++ b/Engine/Plugins/Enterprise/GLTFExporter/MigratedFromGit/Source/GLTFExporter/Public/GLTFSkeletalMeshExporter.h
@@ -0,0 +1,20 @@
+// Copyright Epic Games, Inc. All Rights Reserved.
+
+#pragma once
+
+#include "GLTFExporter.h"
+#include "GLTFSkeletalMeshExporter.generated.h"
+
+UCLASS()
+class UGLTFSkeletalMeshExporter : public UGLTFExporter
+{
+public:
+	GENERATED_BODY()
+
+public:
+	UGLTFSkeletalMeshExporter(const FObjectInitializer& ObjectInitializer = FObjectInitializer());
+
+	//~ Begin UExporter Interface
+	virtual bool ExportBinary(UObject* Object, const TCHAR* Type, FArchive& Ar, FFeedbackContext* Warn, int32 FileIndex = 0, uint32 PortFlags = 0) override;
+	//~ End UExporter Interface
+};
diff --git a/Engine/Plugins/Enterprise/GLTFExporter/MigratedFromGit/Source/GLTFExporter/Public/GLTFStaticMeshExporter.h b/Engine/Plugins/Enterprise/GLTFExporter/MigratedFromGit/Source/GLTFExporter/Public/GLTFStaticMeshExporter.h
new file mode 100644
--- /dev/null
+++ b/Engine/Plugins/Enterprise/GLTFExporter/MigratedFromGit/Source/GLTFExporter/Public/GLTFStaticMeshExporter.h
@@ -0,0 +1,20 @@
+// Copyright Epic Games, Inc. All Rights Reserved.
+
+#pragma once
+
+#include "GLTFExporter.h"
+#include "GLTFStaticMeshExporter.generated.h"
+
+UCLASS()
+class UGLTFStaticMeshExporter : public UGLTFExporter
+{
+public:
+	GENERATED_BODY()
+
+public:
+	UGLTFStaticMeshExporter(const FObjectInitializer& ObjectInitializer = FObjectInitializer());
+
+	//~ Begin UExporter Interface
+	virtual bool ExportBinary(UObject* Object, const TCHAR* Type, FArchive& Ar, FFeedbackContext* Warn, int32 FileIndex = 0, uint32 PortFlags = 0) override;
+	//~ End UExporter Interface
+};
