[GLTFExporter] refactor: remove module GLTFCore

commit 6f4fb6240015875037fff49340203efcdf83d7dd
Date: Thu, 19 Mar 2020 22:35:27 +0100

    refactor: remove module GLTFCore

#skipci

diff --git a/Engine/Plugins/Enterprise/GLTFExporter/MigratedFromGit/GLTFExporter.uplugin b/Engine/Plugins/Enterprise/GLTFExporter/MigratedFromGit/GLTFExporter.uplugin
--- a/Engine/Plugins/Enterprise/GLTFExporter/MigratedFromGit/GLTFExporter.uplugin
+++ b/Engine/Plugins/Enterprise/GLTFExporter/MigratedFromGit/GLTFExporter.uplugin
@@ -12,10 +12,6 @@
 	"IsBetaVersion": true,
 	"Modules":
 	[
-		{
-			"Name": "GLTFCore",
-			"Type": "Runtime"
-		},
 		{
 			"Name": "GLTFExporter",
 			"Type": "Editor"
diff --git a/Engine/Plugins/Enterprise/GLTFExporter/MigratedFromGit/Source/GLTFCore/GLTFCore.Build.cs b/Engine/Plugins/Enterprise/GLTFExporter/MigratedFromGit/Source/GLTFCore/GLTFCore.Build.cs
deleted file mode 100644
diff --git a/Engine/Plugins/Enterprise/GLTFExporter/MigratedFromGit/Source/GLTFCore/Private/GLTFCoreModule.cpp b/Engine/Plugins/Enterprise/GLTFExporter/MigratedFromGit/Source/GLTFCore/Private/GLTFCoreModule.cpp
deleted file mode 100644
diff --git a/Engine/Plugins/Enterprise/GLTFExporter/MigratedFromGit/Source/GLTFCore/Public/GLTFCoreModule.h b/Engine/Plugins/Enterprise/GLTFExporter/MigratedFromGit/Source/GLTFCore/Public/GLTFCoreModule.h
deleted file mode 100644
diff --git a/Engine/Plugins/Enterprise/GLTFExporter/MigratedFromGit/Source/GLTFExporter/GLTFExporter.Build.cs b/Engine/Plugins/Enterprise/GLTFExporter/MigratedFromGit/Source/GLTFExporter/GLTFExporter.Build.cs
--- a/Engine/Plugins/Enterprise/GLTFExporter/MigratedFromGit/Source/GLTFExporter/GLTFExporter.Build.cs
+++ b/Engine/Plugins/Enterprise/GLTFExporter/MigratedFromGit/Source/GLTFExporter/GLTFExporter.Build.cs
@@ -39,8 +39,7 @@ namespace UnrealBuildTool.Rules
                     "SlateCore",
                     "Mainframe",
                     "InputCore",
-                    "EditorStyle",
-					"GLTFCore",
+                    "EditorStyle"
                 }
                 );
 
