[GLTFExporter] feat: add glTF export options

commit 0a2b6017979ebc53edb849749cdccac2e076c9b6
Date: Fri, 20 Mar 2020 00:34:08 +0100

    feat: add glTF export options

#skipci

diff --git a/Engine/Plugins/Enterprise/GLTFExporter/MigratedFromGit/Source/GLTFExporter/Private/GLTFExportOptions.cpp b/Engine/Plugins/Enterprise/GLTFExporter/MigratedFromGit/Source/GLTFExporter/Private/GLTFExportOptions.cpp
new file mode 100644
--- /dev/null
+++ b/Engine/Plugins/Enterprise/GLTFExporter/MigratedFromGit/Source/GLTFExporter/Private/GLTFExportOptions.cpp
@@ -0,0 +1,38 @@
+// Copyriyright Epic Games, Inc. All Rights Reserved.
+
+#include "GLTFExportOptions.h"
+
+UGLTFExportOptions::UGLTFExportOptions(const FObjectInitializer& ObjectInitializer)
+	: Super(ObjectInitializer)
+{
+	bEmbedVertexData = true;
+	bExportVertexColor = true;
+	bExportPreviewMesh = true;
+	bExportUnlitMaterial = true;
+	bExportClearCoatMaterial = true;
+	bBakeMaterialInputs = true;
+	BakedMaterialInputSize = EGLTFExporterTextureSize::POT_512;
+	bEmbedTextures = true;
+	TextureFormat = EGLTFExporterTextureFormat::PNG;
+	bExportLightmaps = true;
+	bExportAnyActor = true;
+	bExportLight = true;
+	bExportCamera = true;
+	bExportReflectionCapture = true;
+	bExportHDRIBackdrop = true;
+	bExportVariantSets = true;
+	bExportAnimationTrigger = true;
+}
+
+void UGLTFExportOptions::ResetToDefault()
+{
+	ReloadConfig();
+}
+
+void UGLTFExportOptions::LoadOptions()
+{
+}
+
+void UGLTFExportOptions::SaveOptions()
+{
+}
diff --git a/Engine/Plugins/Enterprise/GLTFExporter/MigratedFromGit/Source/GLTFExporter/Public/GLTFExportOptions.h b/Engine/Plugins/Enterprise/GLTFExporter/MigratedFromGit/Source/GLTFExporter/Public/GLTFExportOptions.h
new file mode 100644
--- /dev/null
+++ b/Engine/Plugins/Enterprise/GLTFExporter/MigratedFromGit/Source/GLTFExporter/Public/GLTFExportOptions.h
@@ -0,0 +1,97 @@
+// Copyright Epic Games, Inc. All Rights Reserved.
+
+#pragma once
+
+#include "CoreMinimal.h"
+#include "UObject/ObjectMacros.h"
+#include "GLTFExportOptions.generated.h"
+
+UENUM(BlueprintType)
+enum class EGLTFExporterTextureFormat : uint8
+{
+	PNG,
+	JPEG
+};
+
+UENUM(BlueprintType)
+enum class EGLTFExporterTextureSize : uint8
+{
+	POT_1 UMETA(DisplayName = "1 x 1"),
+	POT_2 UMETA(DisplayName = "2 x 2"),
+	POT_4 UMETA(DisplayName = "4 x 4"),
+	POT_8 UMETA(DisplayName = "8 x 8"),
+	POT_16 UMETA(DisplayName = "16 x 16"),
+	POT_32 UMETA(DisplayName = "32 x 32"),
+	POT_64 UMETA(DisplayName = "64 x 64"),
+	POT_128 UMETA(DisplayName = "128 x 128"),
+	POT_256 UMETA(DisplayName = "256 x 256"),
+	POT_512 UMETA(DisplayName = "512 x 512"),
+	POT_1024 UMETA(DisplayName = "1024 x 1024"),
+	POT_2048 UMETA(DisplayName = "2048 x 2048"),
+	POT_4096 UMETA(DisplayName = "4096 x 4096"),
+	POT_8192 UMETA(DisplayName = "8192 x 8192")
+};
+
+UCLASS(config = EditorPerProjectUserSettings, MinimalAPI)
+class UGLTFExportOptions : public UObject
+{
+	GENERATED_UCLASS_BODY()
+
+public:
+	UPROPERTY(EditAnywhere, BlueprintReadWrite, config, Category = Mesh)
+	uint32 bEmbedVertexData : 1;
+
+	UPROPERTY(EditAnywhere, BlueprintReadWrite, config, category = Mesh)
+	uint32 bExportVertexColor : 1;
+
+	UPROPERTY(EditAnywhere, BlueprintReadWrite, config, category = Animation)
+	uint32 bExportPreviewMesh : 1;
+
+	UPROPERTY(EditAnywhere, BlueprintReadWrite, config, Category = Material)
+	uint32 bExportUnlitMaterial : 1;
+
+	UPROPERTY(EditAnywhere, BlueprintReadWrite, config, Category = Material)
+	uint32 bExportClearCoatMaterial : 1;
+
+	UPROPERTY(EditAnywhere, BlueprintReadWrite, config, category = Material)
+	uint32 bBakeMaterialInputs : 1;
+
+	UPROPERTY(EditAnywhere, BlueprintReadWrite, config, category = Material)
+	EGLTFExporterTextureSize BakedMaterialInputSize;
+
+	UPROPERTY(EditAnywhere, BlueprintReadWrite, config, Category = Texture)
+	uint32 bEmbedTextures : 1;
+
+	UPROPERTY(EditAnywhere, BlueprintReadWrite, config, Category = Texture)
+	EGLTFExporterTextureFormat TextureFormat;
+
+	UPROPERTY(EditAnywhere, BlueprintReadWrite, config, Category = Texture)
+	uint32 bExportLightmaps : 1;
+
+	UPROPERTY(EditAnywhere, BlueprintReadWrite, config, Category = Scene)
+	uint32 bExportAnyActor : 1;
+
+	UPROPERTY(EditAnywhere, BlueprintReadWrite, config, Category = Scene)
+	uint32 bExportLight : 1;
+
+	UPROPERTY(EditAnywhere, BlueprintReadWrite, config, Category = Scene)
+	uint32 bExportCamera : 1;
+
+	UPROPERTY(EditAnywhere, BlueprintReadWrite, config, Category = Scene)
+	uint32 bExportReflectionCapture : 1;
+
+	UPROPERTY(EditAnywhere, BlueprintReadWrite, config, Category = Scene)
+	uint32 bExportHDRIBackdrop : 1;
+
+	UPROPERTY(EditAnywhere, BlueprintReadWrite, config, Category = Scene)
+	uint32 bExportVariantSets : 1;
+
+	UPROPERTY(EditAnywhere, BlueprintReadWrite, config, Category = Scene)
+	uint32 bExportAnimationTrigger : 1;
+
+	void ResetToDefault();
+
+	void SaveOptions();
+
+	void LoadOptions();
+};
