[GLTFExporter] feat: add dialog window for glTF export options

commit ffc80fbcd212f9fdf10f9fafdcd2f5e3ae62877f
Date: Fri, 20 Mar 2020 00:34:34 +0100

    feat: add dialog window for glTF export options

#skipci

diff --git a/Engine/Plugins/Enterprise/GLTFExporter/MigratedFromGit/Source/GLTFExporter/Private/GLTFExportOptions.cpp b/Engine/Plugins/Enterprise/GLTFExporter/MigratedFromGit/Source/GLTFExporter/Private/GLTFExportOptions.cpp
--- a/Engine/Plugins/Enterprise/GLTFExporter/MigratedFromGit/Source/GLTFExporter/Private/GLTFExportOptions.cpp
+++ b/Engine/Plugins/Enterprise/GLTFExporter/MigratedFromGit/Source/GLTFExporter/Private/GLTFExportOptions.cpp
@@ -1,6 +1,7 @@
 // Copyriyright Epic Games, Inc. All Rights Reserved.
 
 #include "GLTFExportOptions.h"
+#include "GLTFExportOptionsWindow.h"
 
 UGLTFExportOptions::UGLTFExportOptions(const FObjectInitializer& ObjectInitializer)
 	: Super(ObjectInitializer)
@@ -36,3 +37,20 @@ void UGLTFExportOptions::LoadOptions()
 void UGLTFExportOptions::SaveOptions()
 {
 }
+
+void UGLTFExportOptions::FillOptions(bool bShowOptionDialog, const FString& FullPath, bool BatchMode, bool& OutOperationCanceled, bool& bOutExportAll)
+{
+	OutOperationCanceled = false;
+
+	LoadOptions();
+
+	if (!bShowOptionDialog || GIsAutomationTesting || FApp::IsUnattended())
+	{
+		return;
+	}
+
+	bOutExportAll = false;
+
+	SGLTFExportOptionsWindow::ShowDialog(this, FullPath, BatchMode, OutOperationCanceled, bOutExportAll);
+	SaveOptions();
+}
diff --git a/Engine/Plugins/Enterprise/GLTFExporter/MigratedFromGit/Source/GLTFExporter/Private/GLTFExportOptionsWindow.cpp b/Engine/Plugins/Enterprise/GLTFExporter/MigratedFromGit/Source/GLTFExporter/Private/GLTFExportOptionsWindow.cpp
new file mode 100644
--- /dev/null
+++ b/Engine/Plugins/Enterprise/GLTFExporter/MigratedFromGit/Source/GLTFExporter/Private/GLTFExportOptionsWindow.cpp
@@ -0,0 +1,199 @@
+// Copyright 1998-2019 Epic Games, Inc. All Rights Reserved.
+
+#include "GLTFExportOptionsWindow.h"
+#include "Modules/ModuleManager.h"
+#include "Widgets/Layout/SBorder.h"
+#include "Widgets/Text/STextBlock.h"
+#include "Widgets/Layout/SBox.h"
+#include "Widgets/Layout/SUniformGridPanel.h"
+#include "Widgets/Input/SButton.h"
+#include "EditorStyleSet.h"
+#include "IDocumentation.h"
+#include "PropertyEditorModule.h"
+#include "IDetailsView.h"
+#include "Widgets/SWindow.h"
+#include "Framework/Application/SlateApplication.h"
+#include "Interfaces/IMainFrameModule.h"
+
+#define LOCTEXT_NAMESPACE "GLTFExportOptions"
+
+void SGLTFExportOptionsWindow::Construct(const FArguments& InArgs)
+{
+	ExportOptions = InArgs._ExportOptions;
+	WidgetWindow = InArgs._WidgetWindow;
+
+	check (ExportOptions);
+	
+	FText CancelText = InArgs._BatchMode ? LOCTEXT("GLTFExportOptionsWindow_CancelBatch", "Cancel All") : LOCTEXT("GLTFExportOptionsWindow_Cancel", "Cancel");
+	FText CancelTooltipText = InArgs._BatchMode ? LOCTEXT("GLTFExportOptionsWindow_Cancel_ToolTip_Batch", "Cancel the batch export.") : LOCTEXT("GLTFExportOptionsWindow_Cancel_ToolTip", "Cancel the current glTF export.");
+
+	TSharedPtr<SBox> HeaderToolBox;
+	TSharedPtr<SHorizontalBox> HeaderButtons;
+	TSharedPtr<SBox> InspectorBox;
+	this->ChildSlot
+	[
+		SNew(SBox)
+		[
+			SNew(SVerticalBox)
+			+ SVerticalBox::Slot()
+			.AutoHeight()
+			.Padding(2)
+			[
+				SAssignNew(HeaderToolBox, SBox)
+			]
+			+SVerticalBox::Slot()
+			.AutoHeight()
+			.Padding(2)
+			[
+				SNew(SBorder)
+				.Padding(FMargin(3))
+				.BorderImage(FEditorStyle::GetBrush("ToolPanel.GroupBorder"))
+				[
+					SNew(SHorizontalBox)
+					+SHorizontalBox::Slot()
+					.AutoWidth()
+					[
+						SNew(STextBlock)
+						.Font(FEditorStyle::GetFontStyle("CurveEd.LabelFont"))
+						.Text(LOCTEXT("Export_CurrentFileTitle", "Current File: "))
+					]
+					+SHorizontalBox::Slot()
+					.Padding(5, 0, 0, 0)
+					.AutoWidth()
+					.VAlign(VAlign_Center)
+					[
+						SNew(STextBlock)
+						.Font(FEditorStyle::GetFontStyle("CurveEd.InfoFont"))
+						.Text(InArgs._FullPath)
+					]
+				]
+			]
+			+ SVerticalBox::Slot()
+			.FillHeight(1.0f)
+			.Padding(2)
+			[
+				SAssignNew(InspectorBox, SBox)
+			]
+			+ SVerticalBox::Slot()
+			.AutoHeight()
+			.HAlign(HAlign_Right)
+			.Padding(2)
+			[
+				SNew(SUniformGridPanel)
+				.SlotPadding(2)
+				+ SUniformGridPanel::Slot(0, 0)
+				[
+					SNew(SButton)
+					.HAlign(HAlign_Center)
+					.Text(LOCTEXT("GLTFExportOptionsWindow_ExportAll", "Export All"))
+					.ToolTipText(LOCTEXT("GLTExportOptionsWindow_ExportAll_ToolTip", "Export all files with these same settings"))
+					.Visibility(InArgs._BatchMode ? EVisibility::All : EVisibility::Hidden)
+					.OnClicked(this, &SGLTFExportOptionsWindow::OnExportAll)
+				]
+				+ SUniformGridPanel::Slot(1, 0)
+				[
+					SAssignNew(ImportButton, SButton)
+					.HAlign(HAlign_Center)
+					.Text(LOCTEXT("GLTExportOptionsWindow_Export", "Export"))
+					.OnClicked(this, &SGLTFExportOptionsWindow::OnExport)
+				]
+				+ SUniformGridPanel::Slot(2, 0)
+				[
+					SNew(SButton)
+					.HAlign(HAlign_Center)
+					.Text(CancelText)
+					.ToolTipText(CancelTooltipText)
+					.OnClicked(this, &SGLTFExportOptionsWindow::OnCancel)
+				]
+			]
+		]
+	];
+
+	FPropertyEditorModule& PropertyEditorModule = FModuleManager::GetModuleChecked<FPropertyEditorModule>("PropertyEditor");
+	FDetailsViewArgs DetailsViewArgs;
+	DetailsViewArgs.bAllowSearch = false;
+	DetailsViewArgs.NameAreaSettings = FDetailsViewArgs::HideNameArea;
+	DetailsView = PropertyEditorModule.CreateDetailView(DetailsViewArgs);
+
+	InspectorBox->SetContent(DetailsView->AsShared());
+
+	HeaderToolBox->SetContent(
+		SNew(SBorder)
+		.Padding(FMargin(3))
+		.BorderImage(FEditorStyle::GetBrush("ToolPanel.GroupBorder"))
+		[
+			SNew(SHorizontalBox)
+			+ SHorizontalBox::Slot()
+			[
+				SNew(SBox)
+				.HAlign(HAlign_Right)
+				[
+					SAssignNew(HeaderButtons, SHorizontalBox)
+					+ SHorizontalBox::Slot()
+					.AutoWidth()
+					.Padding(FMargin(2.0f, 0.0f))
+					[
+						SNew(SButton)
+						.Text(LOCTEXT("ExportOptionsWindow_ResetOptions", "Reset to Default"))
+						.OnClicked(this, &SGLTFExportOptionsWindow::OnResetToDefaultClick)
+					]
+				]
+			]
+		]
+	);
+
+	DetailsView->SetObject(ExportOptions);
+}
+
+FReply SGLTFExportOptionsWindow::OnResetToDefaultClick() const
+{
+	ExportOptions->ResetToDefault();
+	//Refresh the view to make sure the custom UI are updating correctly
+	DetailsView->SetObject(ExportOptions, true);
+	return FReply::Handled();
+}
+
+void SGLTFExportOptionsWindow::ShowDialog(UGLTFExportOptions* ExportOptions, const FString& FullPath, bool BatchMode, bool& OutOperationCanceled, bool& bOutExportAll)
+{
+	check(ExportOptions != nullptr);
+
+	OutOperationCanceled = false;
+	bOutExportAll = false;
+
+	TSharedPtr<SWindow> ParentWindow;
+
+	if (FModuleManager::Get().IsModuleLoaded("MainFrame"))
+	{
+		IMainFrameModule& MainFrame = FModuleManager::LoadModuleChecked<IMainFrameModule>("MainFrame");
+		ParentWindow = MainFrame.GetParentWindow();
+	}
+
+	TSharedRef<SWindow> Window = SNew(SWindow)
+		.Title(NSLOCTEXT("UnrealEd", "GLTFExportOptionsTitle", "glTF Export Options"))
+		.SizingRule(ESizingRule::UserSized)
+		.AutoCenter(EAutoCenter::PrimaryWorkArea)
+		.ClientSize(FVector2D(500, 445));
+
+	TSharedPtr<SGLTFExportOptionsWindow> OptionsWindow;
+	Window->SetContent
+	(
+		SAssignNew(OptionsWindow, SGLTFExportOptionsWindow)
+		.ExportOptions(ExportOptions)
+		.WidgetWindow(Window)
+		.FullPath(FText::FromString(FullPath))
+		.BatchMode(BatchMode)
+	);
+
+	FSlateApplication::Get().AddModalWindow(Window, ParentWindow, false);
+
+	if (OptionsWindow->ShouldExport())
+	{
+		bOutExportAll = OptionsWindow->ShouldExportAll();
+	}
+	else
+	{
+		OutOperationCanceled = true;
+	}
+}
+
+#undef LOCTEXT_NAMESPACE
diff --git a/Engine/Plugins/Enterprise/GLTFExporter/MigratedFromGit/Source/GLTFExporter/Public/GLTFExportOptions.h b/Engine/Plugins/Enterprise/GLTFExporter/MigratedFromGit/Source/GLTFExporter/Public/GLTFExportOptions.h
--- a/Engine/Plugins/Enterprise/GLTFExporter/MigratedFromGit/Source/GLTFExporter/Public/GLTFExportOptions.h
+++ b/Engine/Plugins/Enterprise/GLTFExporter/MigratedFromGit/Source/GLTFExporter/Public/GLTFExportOptions.h
@@ -94,4 +94,6 @@ public:
 	void SaveOptions();
 
 	void LoadOptions();
+
+	void FillOptions(bool bShowOptionDialog, const FString& FullPath, bool BatchMode, bool& OutOperationCanceled, bool& bOutExportAll);
 };
diff --git a/Engine/Plugins/Enterprise/GLTFExporter/MigratedFromGit/Source/GLTFExporter/Public/GLTFExportOptionsWindow.h b/Engine/Plugins/Enterprise/GLTFExporter/MigratedFromGit/Source/GLTFExporter/Public/GLTFExportOptionsWindow.h
new file mode 100644
--- /dev/null
+++ b/Engine/Plugins/Enterprise/GLTFExporter/MigratedFromGit/Source/GLTFExporter/Public/GLTFExportOptionsWindow.h
@@ -0,0 +1,101 @@
+// Copyright 1998-2019 Epic Games, Inc. All Rights Reserved.
+
+#pragma once
+
+#include "CoreMinimal.h"
+#include "InputCoreTypes.h"
+#include "Widgets/DeclarativeSyntaxSupport.h"
+#include "Input/Reply.h"
+#include "Widgets/SCompoundWidget.h"
+#include "Widgets/SWindow.h"
+#include "GLTFExportOptions.h"
+
+class SButton;
+
+class SGLTFExportOptionsWindow : public SCompoundWidget
+{
+public:
+	SLATE_BEGIN_ARGS(SGLTFExportOptionsWindow)
+		: _ExportOptions(nullptr)
+		, _WidgetWindow()
+		, _FullPath()
+		, _BatchMode()
+		{}
+
+		SLATE_ARGUMENT( UGLTFExportOptions*, ExportOptions )
+		SLATE_ARGUMENT( TSharedPtr<SWindow>, WidgetWindow )
+		SLATE_ARGUMENT( FText, FullPath )
+		SLATE_ARGUMENT( bool, BatchMode )
+	SLATE_END_ARGS()
+
+public:
+	void Construct(const FArguments& InArgs);
+	virtual bool SupportsKeyboardFocus() const override { return true; }
+
+	FReply OnExport()
+	{
+		bShouldExport = true;
+		if ( WidgetWindow.IsValid() )
+		{
+			WidgetWindow.Pin()->RequestDestroyWindow();
+		}
+		return FReply::Handled();
+	}
+
+	FReply OnExportAll()
+	{
+		bShouldExportAll = true;
+		return OnExport();
+	}
+
+	FReply OnCancel()
+	{
+		bShouldExport = false;
+		bShouldExportAll = false;
+		if ( WidgetWindow.IsValid() )
+		{
+			WidgetWindow.Pin()->RequestDestroyWindow();
+		}
+		return FReply::Handled();
+	}
+
+	virtual FReply OnKeyDown( const FGeometry& MyGeometry, const FKeyEvent& InKeyEvent ) override
+	{
+		if( InKeyEvent.GetKey() == EKeys::Escape )
+		{
+			return OnCancel();
+		}
+
+		return FReply::Unhandled();
+	}
+
+	bool ShouldExport() const
+	{
+		return bShouldExport;
+	}
+
+	bool ShouldExportAll() const
+	{
+		return bShouldExportAll;
+	}
+
+	SGLTFExportOptionsWindow()
+		: ExportOptions(nullptr)
+		, bShouldExport(false)
+		, bShouldExportAll(false)
+	{}
+
+public:
+	static void ShowDialog(UGLTFExportOptions* ExportOptions, const FString& FullPath, bool BatchMode, bool& OutOperationCanceled, bool& bOutExportAll);
+
+private:
+	FReply OnResetToDefaultClick() const;
+
+private:
+	UGLTFExportOptions* ExportOptions;
+	TSharedPtr<class IDetailsView> DetailsView;
+	TWeakPtr< SWindow > WidgetWindow;
+	TSharedPtr< SButton > ImportButton;
+	bool			bShouldExport;
+	bool			bShouldExportAll;
+};
