[GLTFExporter] feat: implement LoadOptions and SaveOptions

commit 23cc2b9013b12adff4ff3a46e8136b3bc146f5a6
Date: Fri, 20 Mar 2020 00:43:57 +0100

    feat: implement LoadOptions and SaveOptions

#skipci

diff --git a/Engine/Plugins/Enterprise/GLTFExporter/MigratedFromGit/Source/GLTFExporter/Private/GLTFExportOptions.cpp b/Engine/Plugins/Enterprise/GLTFExporter/MigratedFromGit/Source/GLTFExporter/Private/GLTFExportOptions.cpp
--- a/Engine/Plugins/Enterprise/GLTFExporter/MigratedFromGit/Source/GLTFExporter/Private/GLTFExportOptions.cpp
+++ b/Engine/Plugins/Enterprise/GLTFExporter/MigratedFromGit/Source/GLTFExporter/Private/GLTFExportOptions.cpp
@@ -3,6 +3,15 @@
 #include "GLTFExportOptions.h"
 #include "GLTFExportOptionsWindow.h"
 
+#include "UObject/ObjectMacros.h"
+#include "UObject/UObjectGlobals.h"
+#include "UObject/Object.h"
+#include "UObject/UObjectBaseUtility.h"
+#include "UObject/UObjectHash.h"
+#include "UObject/Class.h"
+#include "UObject/UnrealType.h"
+#include "Misc/ConfigCacheIni.h"
+
 UGLTFExportOptions::UGLTFExportOptions(const FObjectInitializer& ObjectInitializer)
 	: Super(ObjectInitializer)
 {
@@ -32,10 +41,145 @@ void UGLTFExportOptions::ResetToDefault()
 
 void UGLTFExportOptions::LoadOptions()
 {
+	int32 PortFlags = 0;
+
+	for (UProperty* Property = GetClass()->PropertyLink; Property; Property = Property->PropertyLinkNext)
+	{
+		if (!Property->HasAnyPropertyFlags(CPF_Config))
+		{
+			continue;
+		}
+		FString Section = GetClass()->GetName();
+		FString Key = Property->GetName();
+
+		const bool bIsPropertyInherited = Property->GetOwnerClass() != GetClass();
+		UObject* SuperClassDefaultObject = GetClass()->GetSuperClass()->GetDefaultObject();
+
+		const FString& PropFileName = GEditorPerProjectIni;
+
+		UArrayProperty* Array = dynamic_cast<UArrayProperty*>(Property);
+		if (Array)
+		{
+			FConfigSection* Sec = GConfig->GetSectionPrivate(*Section, 0, 1, *GEditorPerProjectIni);
+			if (Sec != nullptr)
+			{
+				TArray<FConfigValue> List;
+				const FName KeyName(*Key, FNAME_Find);
+				Sec->MultiFind(KeyName, List);
+
+				FScriptArrayHelper_InContainer ArrayHelper(Array, this);
+				// Only override default properties if there is something to override them with.
+				if (List.Num() > 0)
+				{
+					ArrayHelper.EmptyAndAddValues(List.Num());
+					for (int32 i = List.Num() - 1, c = 0; i >= 0; i--, c++)
+					{
+						Array->Inner->ImportText(*List[i].GetValue(), ArrayHelper.GetRawPtr(c), PortFlags, this);
+					}
+				}
+				else
+				{
+					int32 Index = 0;
+					const FConfigValue* ElementValue = nullptr;
+					do
+					{
+						// Add array index number to end of key
+						FString IndexedKey = FString::Printf(TEXT("%s[%i]"), *Key, Index);
+
+						// Try to find value of key
+						const FName IndexedName(*IndexedKey, FNAME_Find);
+						if (IndexedName == NAME_None)
+						{
+							break;
+						}
+						ElementValue = Sec->Find(IndexedName);
+
+						// If found, import the element
+						if (ElementValue != nullptr)
+						{
+							// expand the array if necessary so that Index is a valid element
+							ArrayHelper.ExpandForIndex(Index);
+							Array->Inner->ImportText(*ElementValue->GetValue(), ArrayHelper.GetRawPtr(Index), PortFlags, this);
+						}
+
+						Index++;
+					} while (ElementValue || Index < ArrayHelper.Num());
+				}
+			}
+		}
+		else
+		{
+			for (int32 i = 0; i < Property->ArrayDim; i++)
+			{
+				if (Property->ArrayDim != 1)
+				{
+					Key = FString::Printf(TEXT("%s[%i]"), *Property->GetName(), i);
+				}
+
+				FString Value;
+				bool bFoundValue = GConfig->GetString(*Section, *Key, Value, *GEditorPerProjectIni);
+
+				if (bFoundValue)
+				{
+					if (Property->ImportText(*Value, Property->ContainerPtrToValuePtr<uint8>(this, i), PortFlags, this) == NULL)
+					{
+						// this should be an error as the properties from the .ini / .int file are not correctly being read in and probably are affecting things in subtle ways
+					}
+				}
+			}
+		}
+	}
 }
 
 void UGLTFExportOptions::SaveOptions()
 {
+	int32 PortFlags = 0;
+
+	for (UProperty* Property = GetClass()->PropertyLink; Property; Property = Property->PropertyLinkNext)
+	{
+		if (!Property->HasAnyPropertyFlags(CPF_Config))
+		{
+			continue;
+		}
+		FString Section = GetClass()->GetName();
+		FString Key = Property->GetName();
+
+		const bool bIsPropertyInherited = Property->GetOwnerClass() != GetClass();
+		UObject* SuperClassDefaultObject = GetClass()->GetSuperClass()->GetDefaultObject();
+
+		UArrayProperty* Array = dynamic_cast<UArrayProperty*>(Property);
+		if (Array)
+		{
+			FConfigSection* Sec = GConfig->GetSectionPrivate(*Section, 1, 0, *GEditorPerProjectIni);
+			check(Sec);
+			Sec->Remove(*Key);
+
+			FScriptArrayHelper_InContainer ArrayHelper(Array, this);
+			for (int32 i = 0; i < ArrayHelper.Num(); i++)
+			{
+				FString	Buffer;
+				Array->Inner->ExportTextItem(Buffer, ArrayHelper.GetRawPtr(i), ArrayHelper.GetRawPtr(i), this, PortFlags);
+				Sec->Add(*Key, *Buffer);
+			}
+		}
+		else
+		{
+			TCHAR TempKey[MAX_SPRINTF] = TEXT("");
+			for (int32 Index = 0; Index < Property->ArrayDim; Index++)
+			{
+				if (Property->ArrayDim != 1)
+				{
+					FCString::Sprintf(TempKey, TEXT("%s[%i]"), *Property->GetName(), Index);
+					Key = TempKey;
+				}
+
+				FString	Value;
+				Property->ExportText_InContainer(Index, Value, this, this, this, PortFlags);
+				GConfig->SetString(*Section, *Key, *Value, *GEditorPerProjectIni);
+			}
+		}
+	}
+	GConfig->Flush(0);
 }
 
 void UGLTFExportOptions::FillOptions(bool bShowOptionDialog, const FString& FullPath, bool BatchMode, bool& OutOperationCanceled, bool& bOutExportAll)
@@ -44,6 +188,7 @@ void UGLTFExportOptions::FillOptions(bool bShowOptionDialog, const FString& Full
 
 	LoadOptions();
 
+	//Return if we do not show the export options or we are running automation test or we are unattended
 	if (!bShowOptionDialog || GIsAutomationTesting || FApp::IsUnattended())
 	{
 		return;
diff --git a/Engine/Plugins/Enterprise/GLTFExporter/MigratedFromGit/Source/GLTFExporter/Public/GLTFExportOptions.h b/Engine/Plugins/Enterprise/GLTFExporter/MigratedFromGit/Source/GLTFExporter/Public/GLTFExportOptions.h
--- a/Engine/Plugins/Enterprise/GLTFExporter/MigratedFromGit/Source/GLTFExporter/Public/GLTFExportOptions.h
+++ b/Engine/Plugins/Enterprise/GLTFExporter/MigratedFromGit/Source/GLTFExporter/Public/GLTFExportOptions.h
@@ -89,11 +89,22 @@ public:
 	UPROPERTY(EditAnywhere, BlueprintReadWrite, config, Category = Scene)
 	uint32 bExportAnimationTrigger : 1;
 
+	/* Set all the UProperty to the CDO value */
 	void ResetToDefault();
 
+	/* Save the UProperty to a local ini to retrieve the value the next time we call function LoadOptions() */
 	void SaveOptions();
 
+	/* Load the UProperty data from a local ini which the value was store by the function SaveOptions() */
 	void LoadOptions();
 
+	/**
+	* Load the export option from the last save state and show the dialog if bShowOptionDialog is true.
+	* FullPath is the export file path we display it in the dialog
+	* If user cancel the dialog, the OutOperationCanceled will be true
+	* bOutExportAll will be true if the user want to use the same option for all other asset he want to export
+	*
+	* The function is saving the dialog state in a user ini file and reload it from there. It is not changing the CDO.
+	*/
 	void FillOptions(bool bShowOptionDialog, const FString& FullPath, bool BatchMode, bool& OutOperationCanceled, bool& bOutExportAll);
 };
