[GLTFExporter] feat: choose window title based on gltf text or binary

commit e161102dfb5c238ec505932a072d1b4424e9d47d
Date: Fri, 20 Mar 2020 01:09:54 +0100

    feat: choose window title based on gltf text or binary

#skipci

diff --git a/Engine/Plugins/Enterprise/GLTFExporter/MigratedFromGit/Source/GLTFExporter/Private/GLTFExportOptionsWindow.cpp b/Engine/Plugins/Enterprise/GLTFExporter/MigratedFromGit/Source/GLTFExporter/Private/GLTFExportOptionsWindow.cpp
--- a/Engine/Plugins/Enterprise/GLTFExporter/MigratedFromGit/Source/GLTFExporter/Private/GLTFExportOptionsWindow.cpp
+++ b/Engine/Plugins/Enterprise/GLTFExporter/MigratedFromGit/Source/GLTFExporter/Private/GLTFExportOptionsWindow.cpp
@@ -168,8 +168,13 @@ void SGLTFExportOptionsWindow::ShowDialog(UGLTFExportOptions* ExportOptions, con
 		ParentWindow = MainFrame.GetParentWindow();
 	}
 
+	bool bBinaryFile = FPaths::GetExtension(FullPath, false).Compare(TEXT("glb"), ESearchCase::IgnoreCase) == 0;
+	auto Title = bBinaryFile ?
+		NSLOCTEXT("UnrealEd", "GLTFExportOptionsTitleBinary", "glTF (Binary) Export Options") :
+		NSLOCTEXT("UnrealEd", "GLTFExportOptionsTitle", "glTF Export Options");
+
 	TSharedRef<SWindow> Window = SNew(SWindow)
-		.Title(NSLOCTEXT("UnrealEd", "GLTFExportOptionsTitle", "glTF Export Options"))
+		.Title(Title)
 		.SizingRule(ESizingRule::UserSized)
 		.AutoCenter(EAutoCenter::PrimaryWorkArea)
 		.ClientSize(FVector2D(500, 445));
