[GLTFExporter] refactor: minor code cleanup

commit 49a0611ca5cd8288f5a3968b5a36e8a29f30da5c
Date: Fri, 20 Mar 2020 01:10:11 +0100

    refactor: minor code cleanup

#skipci

diff --git a/Engine/Plugins/Enterprise/GLTFExporter/MigratedFromGit/Source/GLTFExporter/Private/GLTFExportOptions.cpp b/Engine/Plugins/Enterprise/GLTFExporter/MigratedFromGit/Source/GLTFExporter/Private/GLTFExportOptions.cpp
--- a/Engine/Plugins/Enterprise/GLTFExporter/MigratedFromGit/Source/GLTFExporter/Private/GLTFExportOptions.cpp
+++ b/Engine/Plugins/Enterprise/GLTFExporter/MigratedFromGit/Source/GLTFExporter/Private/GLTFExportOptions.cpp
@@ -182,9 +182,9 @@ void UGLTFExportOptions::SaveOptions()
 	GConfig->Flush(0);
 }
 
-void UGLTFExportOptions::FillOptions(bool bShowOptionDialog, const FString& FullPath, bool BatchMode, bool& OutOperationCanceled, bool& bOutExportAll)
+void UGLTFExportOptions::FillOptions(bool bShowOptionDialog, const FString& FullPath, bool bBatchMode, bool& bOutOperationCanceled, bool& bOutExportAll)
 {
-	OutOperationCanceled = false;
+	bOutOperationCanceled = false;
 
 	LoadOptions();
 
@@ -196,6 +196,6 @@ void UGLTFExportOptions::FillOptions(bool bShowOptionDialog, const FString& Full
 
 	bOutExportAll = false;
 
-	SGLTFExportOptionsWindow::ShowDialog(this, FullPath, BatchMode, OutOperationCanceled, bOutExportAll);
+	SGLTFExportOptionsWindow::ShowDialog(this, FullPath, bBatchMode, bOutOperationCanceled, bOutExportAll);
 	SaveOptions();
 }
diff --git a/Engine/Plugins/Enterprise/GLTFExporter/MigratedFromGit/Source/GLTFExporter/Private/GLTFExportOptionsWindow.cpp b/Engine/Plugins/Enterprise/GLTFExporter/MigratedFromGit/Source/GLTFExporter/Private/GLTFExportOptionsWindow.cpp
--- a/Engine/Plugins/Enterprise/GLTFExporter/MigratedFromGit/Source/GLTFExporter/Private/GLTFExportOptionsWindow.cpp
+++ b/Engine/Plugins/Enterprise/GLTFExporter/MigratedFromGit/Source/GLTFExporter/Private/GLTFExportOptionsWindow.cpp
@@ -153,11 +153,64 @@ FReply SGLTFExportOptionsWindow::OnResetToDefaultClick() const
 	return FReply::Handled();
 }
 
-void SGLTFExportOptionsWindow::ShowDialog(UGLTFExportOptions* ExportOptions, const FString& FullPath, bool BatchMode, bool& OutOperationCanceled, bool& bOutExportAll)
+inline FReply SGLTFExportOptionsWindow::OnExport()
+{
+	bShouldExport = true;
+	if (WidgetWindow.IsValid())
+	{
+		WidgetWindow.Pin()->RequestDestroyWindow();
+	}
+	return FReply::Handled();
+}
+
+inline FReply SGLTFExportOptionsWindow::OnExportAll()
+{
+	bShouldExportAll = true;
+	return OnExport();
+}
+
+inline FReply SGLTFExportOptionsWindow::OnCancel()
+{
+	bShouldExport = false;
+	bShouldExportAll = false;
+	if (WidgetWindow.IsValid())
+	{
+		WidgetWindow.Pin()->RequestDestroyWindow();
+	}
+	return FReply::Handled();
+}
+
+inline FReply SGLTFExportOptionsWindow::OnKeyDown(const FGeometry & MyGeometry, const FKeyEvent & InKeyEvent)
+{
+	if (InKeyEvent.GetKey() == EKeys::Escape)
+	{
+		return OnCancel();
+	}
+
+	return FReply::Unhandled();
+}
+
+inline bool SGLTFExportOptionsWindow::ShouldExport() const
+{
+	return bShouldExport;
+}
+
+inline bool SGLTFExportOptionsWindow::ShouldExportAll() const
+{
+	return bShouldExportAll;
+}
+
+inline SGLTFExportOptionsWindow::SGLTFExportOptionsWindow()
+	: ExportOptions(nullptr)
+	, bShouldExport(false)
+	, bShouldExportAll(false)
+{}
+
+void SGLTFExportOptionsWindow::ShowDialog(UGLTFExportOptions* ExportOptions, const FString& FullPath, bool bBatchMode, bool& bOutOperationCanceled, bool& bOutExportAll)
 {
 	check(ExportOptions != nullptr);
 
-	OutOperationCanceled = false;
+	bOutOperationCanceled = false;
 	bOutExportAll = false;
 
 	TSharedPtr<SWindow> ParentWindow;
@@ -186,7 +239,7 @@ void SGLTFExportOptionsWindow::ShowDialog(UGLTFExportOptions* ExportOptions, con
 		.ExportOptions(ExportOptions)
 		.WidgetWindow(Window)
 		.FullPath(FText::FromString(FullPath))
-		.BatchMode(BatchMode)
+		.BatchMode(bBatchMode)
 	);
 
 	FSlateApplication::Get().AddModalWindow(Window, ParentWindow, false);
@@ -197,7 +250,7 @@ void SGLTFExportOptionsWindow::ShowDialog(UGLTFExportOptions* ExportOptions, con
 	}
 	else
 	{
-		OutOperationCanceled = true;
+		bOutOperationCanceled = true;
 	}
 }
 
diff --git a/Engine/Plugins/Enterprise/GLTFExporter/MigratedFromGit/Source/GLTFExporter/Public/GLTFExportOptions.h b/Engine/Plugins/Enterprise/GLTFExporter/MigratedFromGit/Source/GLTFExporter/Public/GLTFExportOptions.h
--- a/Engine/Plugins/Enterprise/GLTFExporter/MigratedFromGit/Source/GLTFExporter/Public/GLTFExportOptions.h
+++ b/Engine/Plugins/Enterprise/GLTFExporter/MigratedFromGit/Source/GLTFExporter/Public/GLTFExportOptions.h
@@ -106,5 +106,5 @@ public:
 	*
 	* The function is saving the dialog state in a user ini file and reload it from there. It is not changing the CDO.
 	*/
-	void FillOptions(bool bShowOptionDialog, const FString& FullPath, bool BatchMode, bool& OutOperationCanceled, bool& bOutExportAll);
+	void FillOptions(bool bShowOptionDialog, const FString& FullPath, bool bBatchMode, bool& bOutOperationCanceled, bool& bOutExportAll);
 };
diff --git a/Engine/Plugins/Enterprise/GLTFExporter/MigratedFromGit/Source/GLTFExporter/Public/GLTFExportOptionsWindow.h b/Engine/Plugins/Enterprise/GLTFExporter/MigratedFromGit/Source/GLTFExporter/Public/GLTFExportOptionsWindow.h
--- a/Engine/Plugins/Enterprise/GLTFExporter/MigratedFromGit/Source/GLTFExporter/Public/GLTFExportOptionsWindow.h
+++ b/Engine/Plugins/Enterprise/GLTFExporter/MigratedFromGit/Source/GLTFExporter/Public/GLTFExportOptionsWindow.h
@@ -30,63 +30,25 @@ public:
 
 public:
 	void Construct(const FArguments& InArgs);
+
 	virtual bool SupportsKeyboardFocus() const override { return true; }
 
-	FReply OnExport()
-	{
-		bShouldExport = true;
-		if ( WidgetWindow.IsValid() )
-		{
-			WidgetWindow.Pin()->RequestDestroyWindow();
-		}
-		return FReply::Handled();
-	}
-
-	FReply OnExportAll()
-	{
-		bShouldExportAll = true;
-		return OnExport();
-	}
-
-	FReply OnCancel()
-	{
-		bShouldExport = false;
-		bShouldExportAll = false;
-		if ( WidgetWindow.IsValid() )
-		{
-			WidgetWindow.Pin()->RequestDestroyWindow();
-		}
-		return FReply::Handled();
-	}
-
-	virtual FReply OnKeyDown( const FGeometry& MyGeometry, const FKeyEvent& InKeyEvent ) override
-	{
-		if( InKeyEvent.GetKey() == EKeys::Escape )
-		{
-			return OnCancel();
-		}
-
-		return FReply::Unhandled();
-	}
-
-	bool ShouldExport() const
-	{
-		return bShouldExport;
-	}
-
-	bool ShouldExportAll() const
-	{
-		return bShouldExportAll;
-	}
-
-	SGLTFExportOptionsWindow()
-		: ExportOptions(nullptr)
-		, bShouldExport(false)
-		, bShouldExportAll(false)
-	{}
+	FReply OnExport();
+
+	FReply OnExportAll();
+
+	FReply OnCancel();
+
+	virtual FReply OnKeyDown(const FGeometry& MyGeometry, const FKeyEvent& InKeyEvent) override;
+
+	bool ShouldExport() const;
+
+	bool ShouldExportAll() const;
+
+	SGLTFExportOptionsWindow();
 
 public:
-	static void ShowDialog(UGLTFExportOptions* ExportOptions, const FString& FullPath, bool BatchMode, bool& OutOperationCanceled, bool& bOutExportAll);
+	static void ShowDialog(UGLTFExportOptions* ExportOptions, const FString& FullPath, bool BatchMode, bool& bOutOperationCanceled, bool& bOutExportAll);
 
 private:
 	FReply OnResetToDefaultClick() const;
