[GLTFExporter] chore: add GLTFEXPORTER_API macro to public classes

commit 8c83bf6afa0caccfbff73de34e87183b8c792b51
Date: Mon, 23 Mar 2020 21:30:29 +0100

    chore: add GLTFEXPORTER_API macro to public classes

#skipci

diff --git a/Engine/Plugins/Enterprise/GLTFExporter/MigratedFromGit/Source/GLTFExporter/Public/GLTFAnimSequenceExporter.h b/Engine/Plugins/Enterprise/GLTFExporter/MigratedFromGit/Source/GLTFExporter/Public/GLTFAnimSequenceExporter.h
--- a/Engine/Plugins/Enterprise/GLTFExporter/MigratedFromGit/Source/GLTFExporter/Public/GLTFAnimSequenceExporter.h
+++ b/Engine/Plugins/Enterprise/GLTFExporter/MigratedFromGit/Source/GLTFExporter/Public/GLTFAnimSequenceExporter.h
@@ -6,7 +6,7 @@
 #include "GLTFAnimSequenceExporter.generated.h"
 
 UCLASS()
-class UGLTFAnimSequenceExporter : public UGLTFExporter
+class GLTFEXPORTER_API UGLTFAnimSequenceExporter : public UGLTFExporter
 {
 public:
 	GENERATED_BODY()
diff --git a/Engine/Plugins/Enterprise/GLTFExporter/MigratedFromGit/Source/GLTFExporter/Public/GLTFExportOptionsWindow.h b/Engine/Plugins/Enterprise/GLTFExporter/MigratedFromGit/Source/GLTFExporter/Public/GLTFExportOptionsWindow.h
--- a/Engine/Plugins/Enterprise/GLTFExporter/MigratedFromGit/Source/GLTFExporter/Public/GLTFExportOptionsWindow.h
+++ b/Engine/Plugins/Enterprise/GLTFExporter/MigratedFromGit/Source/GLTFExporter/Public/GLTFExportOptionsWindow.h
@@ -12,7 +12,7 @@
 
 class SButton;
 
-class SGLTFExportOptionsWindow : public SCompoundWidget
+class GLTFEXPORTER_API SGLTFExportOptionsWindow : public SCompoundWidget
 {
 public:
 	SLATE_BEGIN_ARGS(SGLTFExportOptionsWindow)
diff --git a/Engine/Plugins/Enterprise/GLTFExporter/MigratedFromGit/Source/GLTFExporter/Public/GLTFExporter.h b/Engine/Plugins/Enterprise/GLTFExporter/MigratedFromGit/Source/GLTFExporter/Public/GLTFExporter.h
--- a/Engine/Plugins/Enterprise/GLTFExporter/MigratedFromGit/Source/GLTFExporter/Public/GLTFExporter.h
+++ b/Engine/Plugins/Enterprise/GLTFExporter/MigratedFromGit/Source/GLTFExporter/Public/GLTFExporter.h
@@ -11,7 +11,7 @@
 DECLARE_LOG_CATEGORY_EXTERN(LogGLTFExporter, Log, All);
 
 UCLASS()
-class UGLTFExporter : public UExporter
+class GLTFEXPORTER_API UGLTFExporter : public UExporter
 {
 public:
 	GENERATED_BODY()
diff --git a/Engine/Plugins/Enterprise/GLTFExporter/MigratedFromGit/Source/GLTFExporter/Public/GLTFLevelExporter.h b/Engine/Plugins/Enterprise/GLTFExporter/MigratedFromGit/Source/GLTFExporter/Public/GLTFLevelExporter.h
--- a/Engine/Plugins/Enterprise/GLTFExporter/MigratedFromGit/Source/GLTFExporter/Public/GLTFLevelExporter.h
+++ b/Engine/Plugins/Enterprise/GLTFExporter/MigratedFromGit/Source/GLTFExporter/Public/GLTFLevelExporter.h
@@ -6,7 +6,7 @@
 #include "GLTFLevelExporter.generated.h"
 
 UCLASS()
-class UGLTFLevelExporter : public UGLTFExporter
+class GLTFEXPORTER_API UGLTFLevelExporter : public UGLTFExporter
 {
 public:
 	GENERATED_BODY()
diff --git a/Engine/Plugins/Enterprise/GLTFExporter/MigratedFromGit/Source/GLTFExporter/Public/GLTFSkeletalMeshExporter.h b/Engine/Plugins/Enterprise/GLTFExporter/MigratedFromGit/Source/GLTFExporter/Public/GLTFSkeletalMeshExporter.h
--- a/Engine/Plugins/Enterprise/GLTFExporter/MigratedFromGit/Source/GLTFExporter/Public/GLTFSkeletalMeshExporter.h
+++ b/Engine/Plugins/Enterprise/GLTFExporter/MigratedFromGit/Source/GLTFExporter/Public/GLTFSkeletalMeshExporter.h
@@ -6,7 +6,7 @@
 #include "GLTFSkeletalMeshExporter.generated.h"
 
 UCLASS()
-class UGLTFSkeletalMeshExporter : public UGLTFExporter
+class GLTFEXPORTER_API UGLTFSkeletalMeshExporter : public UGLTFExporter
 {
 public:
 	GENERATED_BODY()
diff --git a/Engine/Plugins/Enterprise/GLTFExporter/MigratedFromGit/Source/GLTFExporter/Public/GLTFStaticMeshExporter.h b/Engine/Plugins/Enterprise/GLTFExporter/MigratedFromGit/Source/GLTFExporter/Public/GLTFStaticMeshExporter.h
--- a/Engine/Plugins/Enterprise/GLTFExporter/MigratedFromGit/Source/GLTFExporter/Public/GLTFStaticMeshExporter.h
+++ b/Engine/Plugins/Enterprise/GLTFExporter/MigratedFromGit/Source/GLTFExporter/Public/GLTFStaticMeshExporter.h
@@ -6,7 +6,7 @@
 #include "GLTFStaticMeshExporter.generated.h"
 
 UCLASS()
-class UGLTFStaticMeshExporter : public UGLTFExporter
+class GLTFEXPORTER_API UGLTFStaticMeshExporter : public UGLTFExporter
 {
 public:
 	GENERATED_BODY()
