[GLTFExporter] feat: add bBundleWebViewer to GLTFExportOptions

commit 6f28f6a4f4eed68f499cc7c8421fff07030f6029
Date: Mon, 23 Mar 2020 21:34:27 +0100

    feat: add bBundleWebViewer to GLTFExportOptions

#skipci

diff --git a/Engine/Plugins/Enterprise/GLTFExporter/MigratedFromGit/Source/GLTFExporter/Private/GLTFExportOptions.cpp b/Engine/Plugins/Enterprise/GLTFExporter/MigratedFromGit/Source/GLTFExporter/Private/GLTFExportOptions.cpp
--- a/Engine/Plugins/Enterprise/GLTFExporter/MigratedFromGit/Source/GLTFExporter/Private/GLTFExportOptions.cpp
+++ b/Engine/Plugins/Enterprise/GLTFExporter/MigratedFromGit/Source/GLTFExporter/Private/GLTFExportOptions.cpp
@@ -15,6 +15,7 @@
 UGLTFExportOptions::UGLTFExportOptions(const FObjectInitializer& ObjectInitializer)
 	: Super(ObjectInitializer)
 {
+	bBundleWebViewer = true;
 	bEmbedVertexData = true;
 	bExportVertexColor = true;
 	bExportPreviewMesh = true;
diff --git a/Engine/Plugins/Enterprise/GLTFExporter/MigratedFromGit/Source/GLTFExporter/Public/GLTFExportOptions.h b/Engine/Plugins/Enterprise/GLTFExporter/MigratedFromGit/Source/GLTFExporter/Public/GLTFExportOptions.h
--- a/Engine/Plugins/Enterprise/GLTFExporter/MigratedFromGit/Source/GLTFExporter/Public/GLTFExportOptions.h
+++ b/Engine/Plugins/Enterprise/GLTFExporter/MigratedFromGit/Source/GLTFExporter/Public/GLTFExportOptions.h
@@ -38,6 +38,9 @@ class UGLTFExportOptions : public UObject
 	GENERATED_UCLASS_BODY()
 
 public:
+	UPROPERTY(EditAnywhere, BlueprintReadWrite, config, Category = Exporter)
+	uint32 bBundleWebViewer : 1;
+
 	UPROPERTY(EditAnywhere, BlueprintReadWrite, config, Category = Mesh)
 	uint32 bEmbedVertexData : 1;
 
