[GLTFExporter] feat: add bExtensionsRequired to GLTFExportOptions

commit ba2c4a1273ca1c0a046212699c2b7c47962aa992
Date: Mon, 23 Mar 2020 22:08:06 +0100

    feat: add bExtensionsRequired to GLTFExportOptions

#skipci

diff --git a/Engine/Plugins/Enterprise/GLTFExporter/MigratedFromGit/Source/GLTFExporter/Private/GLTFExportOptions.cpp b/Engine/Plugins/Enterprise/GLTFExporter/MigratedFromGit/Source/GLTFExporter/Private/GLTFExportOptions.cpp
--- a/Engine/Plugins/Enterprise/GLTFExporter/MigratedFromGit/Source/GLTFExporter/Private/GLTFExportOptions.cpp
+++ b/Engine/Plugins/Enterprise/GLTFExporter/MigratedFromGit/Source/GLTFExporter/Private/GLTFExportOptions.cpp
@@ -16,9 +16,10 @@ UGLTFExportOptions::UGLTFExportOptions(const FObjectInitializer& ObjectInitializ
 	: Super(ObjectInitializer)
 {
 	bBundleWebViewer = true;
+	bExportPreviewMesh = true;
+	bExtensionsRequired = true;
 	bEmbedVertexData = true;
 	bExportVertexColor = true;
-	bExportPreviewMesh = true;
 	bExportUnlitMaterial = true;
 	bExportClearCoatMaterial = true;
 	bBakeMaterialInputs = true;
diff --git a/Engine/Plugins/Enterprise/GLTFExporter/MigratedFromGit/Source/GLTFExporter/Public/GLTFExportOptions.h b/Engine/Plugins/Enterprise/GLTFExporter/MigratedFromGit/Source/GLTFExporter/Public/GLTFExportOptions.h
--- a/Engine/Plugins/Enterprise/GLTFExporter/MigratedFromGit/Source/GLTFExporter/Public/GLTFExportOptions.h
+++ b/Engine/Plugins/Enterprise/GLTFExporter/MigratedFromGit/Source/GLTFExporter/Public/GLTFExportOptions.h
@@ -32,8 +32,8 @@ enum class EGLTFExporterTextureSize : uint8
 	POT_8192 UMETA(DisplayName = "8192 x 8192")
 };
 
-UCLASS(config = EditorPerProjectUserSettings, MinimalAPI)
-class UGLTFExportOptions : public UObject
+UCLASS(config = EditorPerProjectUserSettings, HideCategories = (DebugProperty))
+class GLTFEXPORTER_API UGLTFExportOptions : public UObject
 {
 	GENERATED_UCLASS_BODY()
 
@@ -41,15 +41,18 @@ public:
 	UPROPERTY(EditAnywhere, BlueprintReadWrite, config, Category = Exporter)
 	uint32 bBundleWebViewer : 1;
 
+	UPROPERTY(EditAnywhere, BlueprintReadWrite, config, category = Exporter)
+	uint32 bExportPreviewMesh : 1;
+
+	UPROPERTY(EditAnywhere, BlueprintReadWrite, config, Category = Exporter)
+	uint32 bExtensionsRequired : 1;
+
 	UPROPERTY(EditAnywhere, BlueprintReadWrite, config, Category = Mesh)
 	uint32 bEmbedVertexData : 1;
 
 	UPROPERTY(EditAnywhere, BlueprintReadWrite, config, category = Mesh)
 	uint32 bExportVertexColor : 1;
 
-	UPROPERTY(EditAnywhere, BlueprintReadWrite, config, category = Animation)
-	uint32 bExportPreviewMesh : 1;
-
 	UPROPERTY(EditAnywhere, BlueprintReadWrite, config, Category = Material)
 	uint32 bExportUnlitMaterial : 1;
 
