[GLTFExporter] feat: add bOverwriteExisting to GLTFExportOptions

commit d7310fc94e43cd8275a03453c8adaa6128d30023
Date: Mon, 23 Mar 2020 22:59:35 +0100

    feat: add bOverwriteExisting to GLTFExportOptions

#skipci

diff --git a/Engine/Plugins/Enterprise/GLTFExporter/MigratedFromGit/Source/GLTFExporter/Private/GLTFExportOptions.cpp b/Engine/Plugins/Enterprise/GLTFExporter/MigratedFromGit/Source/GLTFExporter/Private/GLTFExportOptions.cpp
--- a/Engine/Plugins/Enterprise/GLTFExporter/MigratedFromGit/Source/GLTFExporter/Private/GLTFExportOptions.cpp
+++ b/Engine/Plugins/Enterprise/GLTFExporter/MigratedFromGit/Source/GLTFExporter/Private/GLTFExportOptions.cpp
@@ -15,6 +15,7 @@
 UGLTFExportOptions::UGLTFExportOptions(const FObjectInitializer& ObjectInitializer)
 	: Super(ObjectInitializer)
 {
+	bOverwriteExisting = true;
 	bBundleWebViewer = true;
 	bExportPreviewMesh = true;
 	bExtensionsRequired = true;
diff --git a/Engine/Plugins/Enterprise/GLTFExporter/MigratedFromGit/Source/GLTFExporter/Public/GLTFExportOptions.h b/Engine/Plugins/Enterprise/GLTFExporter/MigratedFromGit/Source/GLTFExporter/Public/GLTFExportOptions.h
--- a/Engine/Plugins/Enterprise/GLTFExporter/MigratedFromGit/Source/GLTFExporter/Public/GLTFExportOptions.h
+++ b/Engine/Plugins/Enterprise/GLTFExporter/MigratedFromGit/Source/GLTFExporter/Public/GLTFExportOptions.h
@@ -38,6 +38,9 @@ class GLTFEXPORTER_API UGLTFExportOptions : public UObject
 	GENERATED_UCLASS_BODY()
 
 public:
+	UPROPERTY(EditAnywhere, BlueprintReadWrite, config, Category = Exporter)
+	uint32 bOverwriteExisting : 1;
+
 	UPROPERTY(EditAnywhere, BlueprintReadWrite, config, Category = Exporter)
 	uint32 bBundleWebViewer : 1;
 
