[GLTFExporter] feat: add ExportScale to GLTFExportOptions

commit ade7daca9869a7637e816d67ffe4dee650a5f13c
Date: Mon, 23 Mar 2020 22:59:46 +0100

    feat: add ExportScale to GLTFExportOptions

#skipci

diff --git a/Engine/Plugins/Enterprise/GLTFExporter/MigratedFromGit/Source/GLTFExporter/Private/GLTFExportOptions.cpp b/Engine/Plugins/Enterprise/GLTFExporter/MigratedFromGit/Source/GLTFExporter/Private/GLTFExportOptions.cpp
--- a/Engine/Plugins/Enterprise/GLTFExporter/MigratedFromGit/Source/GLTFExporter/Private/GLTFExportOptions.cpp
+++ b/Engine/Plugins/Enterprise/GLTFExporter/MigratedFromGit/Source/GLTFExporter/Private/GLTFExportOptions.cpp
@@ -28,6 +28,7 @@ UGLTFExportOptions::UGLTFExportOptions(const FObjectInitializer& ObjectInitializ
 	bEmbedTextures = true;
 	TextureFormat = EGLTFExporterTextureFormat::PNG;
 	bExportLightmaps = true;
+	ExportScale = 0.01;
 	bExportAnyActor = true;
 	bExportLight = true;
 	bExportCamera = true;
diff --git a/Engine/Plugins/Enterprise/GLTFExporter/MigratedFromGit/Source/GLTFExporter/Public/GLTFExportOptions.h b/Engine/Plugins/Enterprise/GLTFExporter/MigratedFromGit/Source/GLTFExporter/Public/GLTFExportOptions.h
--- a/Engine/Plugins/Enterprise/GLTFExporter/MigratedFromGit/Source/GLTFExporter/Public/GLTFExportOptions.h
+++ b/Engine/Plugins/Enterprise/GLTFExporter/MigratedFromGit/Source/GLTFExporter/Public/GLTFExportOptions.h
@@ -77,6 +77,11 @@ public:
 	UPROPERTY(EditAnywhere, BlueprintReadWrite, config, Category = Texture)
 	uint32 bExportLightmaps : 1;
 
+	UPROPERTY(EditAnywhere, BlueprintReadWrite, config, Category = Scene, meta =
+		(DisplayName = "Import Uniform Scale",
+		ToolTip = "Scale factor used for exporting assets, by default: 0.01, for conversion from centimeters(Unreal default) to meters(glTF)."))
+	float ExportScale;
+
 	UPROPERTY(EditAnywhere, BlueprintReadWrite, config, Category = Scene)
 	uint32 bExportAnyActor : 1;
 
