[GLTFExporter] refactor!: remove unnecessary options bEmbedVertexData and bEmbedTextures

commit 2fb9df7b7c008091067b9687484f6495885e5e1b
Date: Wed, 8 Apr 2020 16:26:12 +0200

    refactor!: remove unnecessary options bEmbedVertexData and bEmbedTextures

#skipci

diff --git a/Engine/Plugins/Enterprise/GLTFExporter/MigratedFromGit/Source/GLTFExporter/Private/GLTFExportOptions.cpp b/Engine/Plugins/Enterprise/GLTFExporter/MigratedFromGit/Source/GLTFExporter/Private/GLTFExportOptions.cpp
--- a/Engine/Plugins/Enterprise/GLTFExporter/MigratedFromGit/Source/GLTFExporter/Private/GLTFExportOptions.cpp
+++ b/Engine/Plugins/Enterprise/GLTFExporter/MigratedFromGit/Source/GLTFExporter/Private/GLTFExportOptions.cpp
@@ -19,13 +19,11 @@ UGLTFExportOptions::UGLTFExportOptions(const FObjectInitializer& ObjectInitializ
 	bBundleWebViewer = true;
 	bExportPreviewMesh = true;
 	bExtensionsRequired = true;
-	bEmbedVertexData = true;
 	bExportVertexColor = true;
 	bExportUnlitMaterial = true;
 	bExportClearCoatMaterial = true;
 	bBakeMaterialInputs = true;
 	BakedMaterialInputSize = EGLTFExporterTextureSize::POT_512;
-	bEmbedTextures = true;
 	TextureFormat = EGLTFExporterTextureFormat::PNG;
 	bExportLightmaps = true;
 	ExportScale = 0.01;
diff --git a/Engine/Plugins/Enterprise/GLTFExporter/MigratedFromGit/Source/GLTFExporter/Public/GLTFExportOptions.h b/Engine/Plugins/Enterprise/GLTFExporter/MigratedFromGit/Source/GLTFExporter/Public/GLTFExportOptions.h
--- a/Engine/Plugins/Enterprise/GLTFExporter/MigratedFromGit/Source/GLTFExporter/Public/GLTFExportOptions.h
+++ b/Engine/Plugins/Enterprise/GLTFExporter/MigratedFromGit/Source/GLTFExporter/Public/GLTFExportOptions.h
@@ -50,9 +50,6 @@ public:
 	UPROPERTY(EditAnywhere, BlueprintReadWrite, config, Category = Exporter)
 	uint32 bExtensionsRequired : 1;
 
-	UPROPERTY(EditAnywhere, BlueprintReadWrite, config, Category = Mesh)
-	uint32 bEmbedVertexData : 1;
-
 	UPROPERTY(EditAnywhere, BlueprintReadWrite, config, category = Mesh)
 	uint32 bExportVertexColor : 1;
 
@@ -68,9 +65,6 @@ public:
 	UPROPERTY(EditAnywhere, BlueprintReadWrite, config, category = Material)
 	EGLTFExporterTextureSize BakedMaterialInputSize;
 
-	UPROPERTY(EditAnywhere, BlueprintReadWrite, config, Category = Texture)
-	uint32 bEmbedTextures : 1;
-
 	UPROPERTY(EditAnywhere, BlueprintReadWrite, config, Category = Texture)
 	EGLTFExporterTextureFormat TextureFormat;
 
