[GLTFExporter] chore: rename option bExportAnyActor to bExportGenericActor

commit f169e6c56279702dba681de5073f378c3f6e9a59
Date: Wed, 8 Apr 2020 16:27:40 +0200

    chore: rename option bExportAnyActor to bExportGenericActor

#skipci

diff --git a/Engine/Plugins/Enterprise/GLTFExporter/MigratedFromGit/Source/GLTFExporter/Private/GLTFExportOptions.cpp b/Engine/Plugins/Enterprise/GLTFExporter/MigratedFromGit/Source/GLTFExporter/Private/GLTFExportOptions.cpp
--- a/Engine/Plugins/Enterprise/GLTFExporter/MigratedFromGit/Source/GLTFExporter/Private/GLTFExportOptions.cpp
+++ b/Engine/Plugins/Enterprise/GLTFExporter/MigratedFromGit/Source/GLTFExporter/Private/GLTFExportOptions.cpp
@@ -27,7 +27,7 @@ UGLTFExportOptions::UGLTFExportOptions(const FObjectInitializer& ObjectInitializ
 	TextureFormat = EGLTFExporterTextureFormat::PNG;
 	bExportLightmaps = true;
 	ExportScale = 0.01;
-	bExportAnyActor = true;
+	bExportGenericActor = true;
 	bExportLight = true;
 	bExportCamera = true;
 	bExportReflectionCapture = true;
diff --git a/Engine/Plugins/Enterprise/GLTFExporter/MigratedFromGit/Source/GLTFExporter/Public/GLTFExportOptions.h b/Engine/Plugins/Enterprise/GLTFExporter/MigratedFromGit/Source/GLTFExporter/Public/GLTFExportOptions.h
--- a/Engine/Plugins/Enterprise/GLTFExporter/MigratedFromGit/Source/GLTFExporter/Public/GLTFExportOptions.h
+++ b/Engine/Plugins/Enterprise/GLTFExporter/MigratedFromGit/Source/GLTFExporter/Public/GLTFExportOptions.h
@@ -77,7 +77,7 @@ public:
 	float ExportScale;
 
 	UPROPERTY(EditAnywhere, BlueprintReadWrite, config, Category = Scene)
-	uint32 bExportAnyActor : 1;
+	uint32 bExportGenericActor : 1;
 
 	UPROPERTY(EditAnywhere, BlueprintReadWrite, config, Category = Scene)
 	uint32 bExportLight : 1;
