[GLTFExporter] feat: add options for different HDR texture encoding

commit c9887b192fb3bb0594b91f734b92a0d695e1737f
Date: Wed, 8 Apr 2020 16:28:48 +0200

    feat: add options for different HDR texture encoding

#skipci

diff --git a/Engine/Plugins/Enterprise/GLTFExporter/MigratedFromGit/Source/GLTFExporter/Private/GLTFExportOptions.cpp b/Engine/Plugins/Enterprise/GLTFExporter/MigratedFromGit/Source/GLTFExporter/Private/GLTFExportOptions.cpp
--- a/Engine/Plugins/Enterprise/GLTFExporter/MigratedFromGit/Source/GLTFExporter/Private/GLTFExportOptions.cpp
+++ b/Engine/Plugins/Enterprise/GLTFExporter/MigratedFromGit/Source/GLTFExporter/Private/GLTFExportOptions.cpp
@@ -25,6 +25,7 @@ UGLTFExportOptions::UGLTFExportOptions(const FObjectInitializer& ObjectInitializ
 	bBakeMaterialInputs = true;
 	BakedMaterialInputSize = EGLTFExporterTextureSize::POT_512;
 	TextureFormat = EGLTFExporterTextureFormat::PNG;
+	TextureHDREncoding = EGLTFExporterTextureHDREncoding::RGBD;
 	bExportLightmaps = true;
 	ExportScale = 0.01;
 	bExportGenericActor = true;
diff --git a/Engine/Plugins/Enterprise/GLTFExporter/MigratedFromGit/Source/GLTFExporter/Public/GLTFExportOptions.h b/Engine/Plugins/Enterprise/GLTFExporter/MigratedFromGit/Source/GLTFExporter/Public/GLTFExportOptions.h
--- a/Engine/Plugins/Enterprise/GLTFExporter/MigratedFromGit/Source/GLTFExporter/Public/GLTFExportOptions.h
+++ b/Engine/Plugins/Enterprise/GLTFExporter/MigratedFromGit/Source/GLTFExporter/Public/GLTFExportOptions.h
@@ -13,6 +13,16 @@ enum class EGLTFExporterTextureFormat : uint8
 	JPEG
 };
 
+UENUM(BlueprintType)
+enum class EGLTFExporterTextureHDREncoding : uint8
+{
+	None,
+	HDR,
+	RGBM,
+	RGBD,
+	RGBE
+};
+
 UENUM(BlueprintType)
 enum class EGLTFExporterTextureSize : uint8
 {
@@ -68,6 +78,9 @@ public:
 	UPROPERTY(EditAnywhere, BlueprintReadWrite, config, Category = Texture)
 	EGLTFExporterTextureFormat TextureFormat;
 
+	UPROPERTY(EditAnywhere, BlueprintReadWrite, config, Category = Texture)
+	EGLTFExporterTextureHDREncoding TextureHDREncoding;
+
 	UPROPERTY(EditAnywhere, BlueprintReadWrite, config, Category = Texture)
 	uint32 bExportLightmaps : 1;
 
