[GLTFExporter] feat: show options when exporting any asset

commit d72f39f6c11c1f00d1d6b0f2df9a3db4a572d7df
Date: Wed, 15 Apr 2020 17:45:41 +0200

    feat: show options when exporting any asset

#skipci

diff --git a/Engine/Plugins/Enterprise/GLTFExporter/MigratedFromGit/Source/GLTFExporter/Private/GLTFExportOptions.cpp b/Engine/Plugins/Enterprise/GLTFExporter/MigratedFromGit/Source/GLTFExporter/Private/GLTFExportOptions.cpp
--- a/Engine/Plugins/Enterprise/GLTFExporter/MigratedFromGit/Source/GLTFExporter/Private/GLTFExportOptions.cpp
+++ b/Engine/Plugins/Enterprise/GLTFExporter/MigratedFromGit/Source/GLTFExporter/Private/GLTFExportOptions.cpp
@@ -185,7 +185,7 @@ void UGLTFExportOptions::SaveOptions()
 	GConfig->Flush(0);
 }
 
-void UGLTFExportOptions::FillOptions(bool bShowOptionDialog, const FString& FullPath, bool bBatchMode, bool& bOutOperationCanceled, bool& bOutExportAll)
+void UGLTFExportOptions::FillOptions(bool bBatchMode, bool bShowOptionDialog, const FString& FullPath, bool& bOutOperationCanceled, bool& bOutExportAll)
 {
 	bOutOperationCanceled = false;
 
diff --git a/Engine/Plugins/Enterprise/GLTFExporter/MigratedFromGit/Source/GLTFExporter/Private/GLTFExporter.cpp b/Engine/Plugins/Enterprise/GLTFExporter/MigratedFromGit/Source/GLTFExporter/Private/GLTFExporter.cpp
--- a/Engine/Plugins/Enterprise/GLTFExporter/MigratedFromGit/Source/GLTFExporter/Private/GLTFExporter.cpp
+++ b/Engine/Plugins/Enterprise/GLTFExporter/MigratedFromGit/Source/GLTFExporter/Private/GLTFExporter.cpp
@@ -2,6 +2,8 @@
 
 #include "GLTFExporter.h"
 #include "UnrealExporter.h"
+#include "GLTFExportOptions.h"
+#include "UObject/StrongObjectPtr.h"
 #include "Engine.h"
 
 DEFINE_LOG_CATEGORY(LogGLTFExporter);
@@ -20,10 +22,29 @@ UGLTFExporter::UGLTFExporter(const FObjectInitializer& ObjectInitializer)
 
 bool UGLTFExporter::ExportBinary(UObject* Object, const TCHAR* Type, FArchive& Ar, FFeedbackContext* Warn, int32 FileIndex, uint32 PortFlags)
 {
+	// TODO: remove temp logging
 	UE_LOG(LogGLTFExporter, Warning, TEXT("%s::ExportBinary"), *(this->GetClass()->GetName()));
 	UE_LOG(LogGLTFExporter, Warning, TEXT("Object: %s (%s)"), *(Object->GetName()), *(Object->GetClass()->GetName()));
 	UE_LOG(LogGLTFExporter, Warning, TEXT("Type: %s"), Type);
 	UE_LOG(LogGLTFExporter, Warning, TEXT("FileIndex: %d"), FileIndex);
 	UE_LOG(LogGLTFExporter, Warning, TEXT("PortFlags: %d"), PortFlags);
-	return true;
+
+	// Populate export options
+	TStrongObjectPtr<UGLTFExportOptions> Options(NewObject<UGLTFExportOptions>(GetTransientPackage(), TEXT("GLTF Export Options")));
+	{
+		bool ExportAll = GetBatchMode() && !GetShowExportOption();
+		bool ExportCancel = false;
+		Options->FillOptions(GetBatchMode(), GetShowExportOption(), UExporter::CurrentFilename, ExportCancel, ExportAll);
+		if (ExportCancel)
+		{
+			SetCancelBatch(GetBatchMode());
+
+			// User cancelled the FBX export
+			return false;
+		}
+		SetShowExportOption(!ExportAll);
+	}
+
+	// TODO: implement
+	return false;
 }
diff --git a/Engine/Plugins/Enterprise/GLTFExporter/MigratedFromGit/Source/GLTFExporter/Public/GLTFExportOptions.h b/Engine/Plugins/Enterprise/GLTFExporter/MigratedFromGit/Source/GLTFExporter/Public/GLTFExportOptions.h
--- a/Engine/Plugins/Enterprise/GLTFExporter/MigratedFromGit/Source/GLTFExporter/Public/GLTFExportOptions.h
+++ b/Engine/Plugins/Enterprise/GLTFExporter/MigratedFromGit/Source/GLTFExporter/Public/GLTFExportOptions.h
@@ -127,5 +127,5 @@ public:
 	*
 	* The function is saving the dialog state in a user ini file and reload it from there. It is not changing the CDO.
 	*/
-	void FillOptions(bool bShowOptionDialog, const FString& FullPath, bool bBatchMode, bool& bOutOperationCanceled, bool& bOutExportAll);
+	void FillOptions(bool BatchMode, bool bShowOptionDialog, const FString& FullPath, bool& OutOperationCanceled, bool& bOutExportAll);
 };
