[GLTFExporter] refactor: move populating options to own function

commit d5ee23c563daa456c87854ef573f62a53aad3a6f
Date: Fri, 17 Apr 2020 12:59:22 +0200

    refactor: move populating options to own function

#skipci

diff --git a/Engine/Plugins/Enterprise/GLTFExporter/MigratedFromGit/Source/GLTFExporter/Private/GLTFExporter.cpp b/Engine/Plugins/Enterprise/GLTFExporter/MigratedFromGit/Source/GLTFExporter/Private/GLTFExporter.cpp
--- a/Engine/Plugins/Enterprise/GLTFExporter/MigratedFromGit/Source/GLTFExporter/Private/GLTFExporter.cpp
+++ b/Engine/Plugins/Enterprise/GLTFExporter/MigratedFromGit/Source/GLTFExporter/Private/GLTFExporter.cpp
@@ -2,8 +2,6 @@
 
 #include "GLTFExporter.h"
 #include "UnrealExporter.h"
-#include "GLTFExportOptions.h"
-#include "UObject/StrongObjectPtr.h"
 #include "Engine.h"
 
 DEFINE_LOG_CATEGORY(LogGLTFExporter);
@@ -18,6 +16,7 @@ UGLTFExporter::UGLTFExporter(const FObjectInitializer& ObjectInitializer)
 	FormatExtension.Add(TEXT("glb"));
 	FormatDescription.Add(TEXT("GL Transmission Format"));
 	FormatDescription.Add(TEXT("GL Transmission Format (Binary)"));
+	ExportOptions = NewObject<UGLTFExportOptions>(this, TEXT("GLTF Export Options"));
 }
 
 bool UGLTFExporter::ExportBinary(UObject* Object, const TCHAR* Type, FArchive& Ar, FFeedbackContext* Warn, int32 FileIndex, uint32 PortFlags)
@@ -29,22 +28,28 @@ bool UGLTFExporter::ExportBinary(UObject* Object, const TCHAR* Type, FArchive& A
 	UE_LOG(LogGLTFExporter, Warning, TEXT("FileIndex: %d"), FileIndex);
 	UE_LOG(LogGLTFExporter, Warning, TEXT("PortFlags: %d"), PortFlags);
 
-	// Populate export options
-	TStrongObjectPtr<UGLTFExportOptions> Options(NewObject<UGLTFExportOptions>(GetTransientPackage(), TEXT("GLTF Export Options")));
+	if (!FillExportOptions())
 	{
-		bool ExportAll = GetBatchMode() && !GetShowExportOption();
-		bool ExportCancel = false;
-		Options->FillOptions(GetBatchMode(), GetShowExportOption(), UExporter::CurrentFilename, ExportCancel, ExportAll);
-		if (ExportCancel)
-		{
-			SetCancelBatch(GetBatchMode());
-
-			// User cancelled the FBX export
-			return false;
-		}
-		SetShowExportOption(!ExportAll);
+		// User cancelled the export
+		return false;
 	}
-
+	
 	// TODO: implement
 	return false;
 }
+
+bool UGLTFExporter::FillExportOptions()
+{
+	bool ExportAll = GetBatchMode() && !GetShowExportOption();
+	bool ExportCancel = false;
+
+	ExportOptions->FillOptions(GetBatchMode(), GetShowExportOption(), UExporter::CurrentFilename, ExportCancel, ExportAll);
+	if (ExportCancel)
+	{
+		SetCancelBatch(GetBatchMode());
+		return false;
+	}
+
+	SetShowExportOption(!ExportAll);
+	return true;
+}
diff --git a/Engine/Plugins/Enterprise/GLTFExporter/MigratedFromGit/Source/GLTFExporter/Public/GLTFExporter.h b/Engine/Plugins/Enterprise/GLTFExporter/MigratedFromGit/Source/GLTFExporter/Public/GLTFExporter.h
--- a/Engine/Plugins/Enterprise/GLTFExporter/MigratedFromGit/Source/GLTFExporter/Public/GLTFExporter.h
+++ b/Engine/Plugins/Enterprise/GLTFExporter/MigratedFromGit/Source/GLTFExporter/Public/GLTFExporter.h
@@ -6,6 +6,7 @@
 #include "UObject/ObjectMacros.h"
 #include "UObject/UObjectGlobals.h"
 #include "Exporters/Exporter.h"
+#include "GLTFExportOptions.h"
 #include "GLTFExporter.generated.h"
 
 DECLARE_LOG_CATEGORY_EXTERN(LogGLTFExporter, Log, All);
@@ -22,4 +23,11 @@ public:
 	//~ Begin UExporter Interface
 	virtual bool ExportBinary(UObject* Object, const TCHAR* Type, FArchive& Ar, FFeedbackContext* Warn, int32 FileIndex = 0, uint32 PortFlags = 0) override;
 	//~ End UExporter Interface
+
+protected:
+
+	bool FillExportOptions();
+
+	UPROPERTY(transient)
+	UGLTFExportOptions* ExportOptions;
 };
