[GLTFExporter] feat: add sanity check in UGLTFExporter::ExportBinary

commit 842810820e1e6d8334d9650daa891dd5c82887a5
Date: Wed, 22 Apr 2020 20:44:53 +0200

    feat: add sanity check in UGLTFExporter::ExportBinary

#skipci

diff --git a/Engine/Plugins/Enterprise/GLTFExporter/MigratedFromGit/Source/GLTFExporter/Private/GLTFExporter.cpp b/Engine/Plugins/Enterprise/GLTFExporter/MigratedFromGit/Source/GLTFExporter/Private/GLTFExporter.cpp
--- a/Engine/Plugins/Enterprise/GLTFExporter/MigratedFromGit/Source/GLTFExporter/Private/GLTFExporter.cpp
+++ b/Engine/Plugins/Enterprise/GLTFExporter/MigratedFromGit/Source/GLTFExporter/Private/GLTFExporter.cpp
@@ -21,6 +21,8 @@ UGLTFExporter::UGLTFExporter(const FObjectInitializer& ObjectInitializer)
 
 bool UGLTFExporter::ExportBinary(UObject* Object, const TCHAR* Type, FArchive& Ar, FFeedbackContext* Warn, int32 FileIndex, uint32 PortFlags)
 {
+	check(Object);
+
 	// TODO: remove temp logging
 	UE_LOG(LogGLTFExporter, Warning, TEXT("%s::ExportBinary"), *(this->GetClass()->GetName()));
 	UE_LOG(LogGLTFExporter, Warning, TEXT("Object: %s (%s)"), *(Object->GetName()), *(Object->GetClass()->GetName()));
