[GLTFExporter] refactor: move exporter log category to module

commit ef462214a312d63603a524bccb85f2931cbd10ff
Date: Fri, 24 Apr 2020 07:36:55 +0200

    refactor: move exporter log category to module

#skipci

diff --git a/Engine/Plugins/Enterprise/GLTFExporter/MigratedFromGit/Source/GLTFExporter/Private/GLTFExporter.cpp b/Engine/Plugins/Enterprise/GLTFExporter/MigratedFromGit/Source/GLTFExporter/Private/GLTFExporter.cpp
--- a/Engine/Plugins/Enterprise/GLTFExporter/MigratedFromGit/Source/GLTFExporter/Private/GLTFExporter.cpp
+++ b/Engine/Plugins/Enterprise/GLTFExporter/MigratedFromGit/Source/GLTFExporter/Private/GLTFExporter.cpp
@@ -4,8 +4,6 @@
 #include "UnrealExporter.h"
 #include "Engine.h"
 
-DEFINE_LOG_CATEGORY(LogGLTFExporter);
-
 UGLTFExporter::UGLTFExporter(const FObjectInitializer& ObjectInitializer)
 	: Super(ObjectInitializer)
 {
@@ -23,19 +21,12 @@ bool UGLTFExporter::ExportBinary(UObject* Object, const TCHAR* Type, FArchive& A
 {
 	check(Object);
 
-	// TODO: remove temp logging
-	UE_LOG(LogGLTFExporter, Warning, TEXT("%s::ExportBinary"), *(this->GetClass()->GetName()));
-	UE_LOG(LogGLTFExporter, Warning, TEXT("Object: %s (%s)"), *(Object->GetName()), *(Object->GetClass()->GetName()));
-	UE_LOG(LogGLTFExporter, Warning, TEXT("Type: %s"), Type);
-	UE_LOG(LogGLTFExporter, Warning, TEXT("FileIndex: %d"), FileIndex);
-	UE_LOG(LogGLTFExporter, Warning, TEXT("PortFlags: %d"), PortFlags);
-
 	if (!FillExportOptions())
 	{
 		// User cancelled the export
 		return false;
 	}
-	
+
 	// TODO: implement
 	return false;
 }
diff --git a/Engine/Plugins/Enterprise/GLTFExporter/MigratedFromGit/Source/GLTFExporter/Private/GLTFExporterModule.cpp b/Engine/Plugins/Enterprise/GLTFExporter/MigratedFromGit/Source/GLTFExporter/Private/GLTFExporterModule.cpp
--- a/Engine/Plugins/Enterprise/GLTFExporter/MigratedFromGit/Source/GLTFExporter/Private/GLTFExporterModule.cpp
+++ b/Engine/Plugins/Enterprise/GLTFExporter/MigratedFromGit/Source/GLTFExporter/Private/GLTFExporterModule.cpp
@@ -5,6 +5,8 @@
 #include "CoreMinimal.h"
 #include "Modules/ModuleManager.h"
 
+DEFINE_LOG_CATEGORY(LogGLTFExporter);
+
 /**
  * glTF Exporter module implementation (private)
  */
diff --git a/Engine/Plugins/Enterprise/GLTFExporter/MigratedFromGit/Source/GLTFExporter/Public/GLTFExporter.h b/Engine/Plugins/Enterprise/GLTFExporter/MigratedFromGit/Source/GLTFExporter/Public/GLTFExporter.h
--- a/Engine/Plugins/Enterprise/GLTFExporter/MigratedFromGit/Source/GLTFExporter/Public/GLTFExporter.h
+++ b/Engine/Plugins/Enterprise/GLTFExporter/MigratedFromGit/Source/GLTFExporter/Public/GLTFExporter.h
@@ -9,8 +9,6 @@
 #include "GLTFExportOptions.h"
 #include "GLTFExporter.generated.h"
 
-DECLARE_LOG_CATEGORY_EXTERN(LogGLTFExporter, Log, All);
-
 UCLASS()
 class GLTFEXPORTER_API UGLTFExporter : public UExporter
 {
diff --git a/Engine/Plugins/Enterprise/GLTFExporter/MigratedFromGit/Source/GLTFExporter/Public/GLTFExporterModule.h b/Engine/Plugins/Enterprise/GLTFExporter/MigratedFromGit/Source/GLTFExporter/Public/GLTFExporterModule.h
--- a/Engine/Plugins/Enterprise/GLTFExporter/MigratedFromGit/Source/GLTFExporter/Public/GLTFExporterModule.h
+++ b/Engine/Plugins/Enterprise/GLTFExporter/MigratedFromGit/Source/GLTFExporter/Public/GLTFExporterModule.h
@@ -6,6 +6,8 @@
 #include "Modules/ModuleInterface.h"
 #include "Modules/ModuleManager.h"
 
+DECLARE_LOG_CATEGORY_EXTERN(LogGLTFExporter, Log, All);
+
 /**
  * The public interface of the GLTFExporter module
  */
