[GLTFExporter] refactor: remove base implementation UGLTFExporter::ExportBinary

commit 9583d39aaea4a917dd029201509ad993854d938b
Date: Fri, 24 Apr 2020 07:37:52 +0200

    refactor: remove base implementation UGLTFExporter::ExportBinary

#skipci

diff --git a/Engine/Plugins/Enterprise/GLTFExporter/MigratedFromGit/Source/GLTFExporter/Private/GLTFAnimSequenceExporter.cpp b/Engine/Plugins/Enterprise/GLTFExporter/MigratedFromGit/Source/GLTFExporter/Private/GLTFAnimSequenceExporter.cpp
--- a/Engine/Plugins/Enterprise/GLTFExporter/MigratedFromGit/Source/GLTFExporter/Private/GLTFAnimSequenceExporter.cpp
+++ b/Engine/Plugins/Enterprise/GLTFExporter/MigratedFromGit/Source/GLTFExporter/Private/GLTFAnimSequenceExporter.cpp
@@ -9,8 +9,16 @@ UGLTFAnimSequenceExporter::UGLTFAnimSequenceExporter(const FObjectInitializer& O
 	SupportedClass = UAnimSequence::StaticClass();
 }
 
-bool UGLTFAnimSequenceExporter::ExportBinary(UObject* Object, const TCHAR* Type, FArchive& Ar, FFeedbackContext* Warn, int32 FileIndex, uint32 PortFlags)
+bool UGLTFAnimSequenceExporter::ExportBinary(UObject* Object, const TCHAR* Type, FArchive& Archive, FFeedbackContext* Warn, int32 FileIndex, uint32 PortFlags)
 {
 	UAnimSequence* AnimSequence = CastChecked<UAnimSequence>(Object);
-	return Super::ExportBinary(AnimSequence, Type, Ar, Warn, FileIndex, PortFlags);
+
+	if (!FillExportOptions())
+	{
+		// User cancelled the export
+		return false;
+	}
+
+	// TODO: implement
+	return true;
 }
diff --git a/Engine/Plugins/Enterprise/GLTFExporter/MigratedFromGit/Source/GLTFExporter/Private/GLTFExporter.cpp b/Engine/Plugins/Enterprise/GLTFExporter/MigratedFromGit/Source/GLTFExporter/Private/GLTFExporter.cpp
--- a/Engine/Plugins/Enterprise/GLTFExporter/MigratedFromGit/Source/GLTFExporter/Private/GLTFExporter.cpp
+++ b/Engine/Plugins/Enterprise/GLTFExporter/MigratedFromGit/Source/GLTFExporter/Private/GLTFExporter.cpp
@@ -17,20 +17,6 @@ UGLTFExporter::UGLTFExporter(const FObjectInitializer& ObjectInitializer)
 	ExportOptions = NewObject<UGLTFExportOptions>(this, TEXT("GLTF Export Options"));
 }
 
-bool UGLTFExporter::ExportBinary(UObject* Object, const TCHAR* Type, FArchive& Ar, FFeedbackContext* Warn, int32 FileIndex, uint32 PortFlags)
-{
-	check(Object);
-
-	if (!FillExportOptions())
-	{
-		// User cancelled the export
-		return false;
-	}
-
-	// TODO: implement
-	return false;
-}
-
 bool UGLTFExporter::FillExportOptions()
 {
 	bool ExportAll = GetBatchMode() && !GetShowExportOption();
diff --git a/Engine/Plugins/Enterprise/GLTFExporter/MigratedFromGit/Source/GLTFExporter/Private/GLTFLevelExporter.cpp b/Engine/Plugins/Enterprise/GLTFExporter/MigratedFromGit/Source/GLTFExporter/Private/GLTFLevelExporter.cpp
--- a/Engine/Plugins/Enterprise/GLTFExporter/MigratedFromGit/Source/GLTFExporter/Private/GLTFLevelExporter.cpp
+++ b/Engine/Plugins/Enterprise/GLTFExporter/MigratedFromGit/Source/GLTFExporter/Private/GLTFLevelExporter.cpp
@@ -9,8 +9,16 @@ UGLTFLevelExporter::UGLTFLevelExporter(const FObjectInitializer& ObjectInitializ
 	SupportedClass = UWorld::StaticClass();
 }
 
-bool UGLTFLevelExporter::ExportBinary(UObject* Object, const TCHAR* Type, FArchive& Ar, FFeedbackContext* Warn, int32 FileIndex, uint32 PortFlags)
+bool UGLTFLevelExporter::ExportBinary(UObject* Object, const TCHAR* Type, FArchive& Archive, FFeedbackContext* Warn, int32 FileIndex, uint32 PortFlags)
 {
 	UWorld* World = CastChecked<UWorld>(Object);
-	return Super::ExportBinary(World, Type, Ar, Warn, FileIndex, PortFlags);
+
+	if (!FillExportOptions())
+	{
+		// User cancelled the export
+		return false;
+	}
+
+	// TODO: implement
+	return true;
 }
diff --git a/Engine/Plugins/Enterprise/GLTFExporter/MigratedFromGit/Source/GLTFExporter/Private/GLTFSkeletalMeshExporter.cpp b/Engine/Plugins/Enterprise/GLTFExporter/MigratedFromGit/Source/GLTFExporter/Private/GLTFSkeletalMeshExporter.cpp
--- a/Engine/Plugins/Enterprise/GLTFExporter/MigratedFromGit/Source/GLTFExporter/Private/GLTFSkeletalMeshExporter.cpp
+++ b/Engine/Plugins/Enterprise/GLTFExporter/MigratedFromGit/Source/GLTFExporter/Private/GLTFSkeletalMeshExporter.cpp
@@ -9,8 +9,16 @@ UGLTFSkeletalMeshExporter::UGLTFSkeletalMeshExporter(const FObjectInitializer& O
 	SupportedClass = USkeletalMesh::StaticClass();
 }
 
-bool UGLTFSkeletalMeshExporter::ExportBinary(UObject* Object, const TCHAR* Type, FArchive& Ar, FFeedbackContext* Warn, int32 FileIndex, uint32 PortFlags)
+bool UGLTFSkeletalMeshExporter::ExportBinary(UObject* Object, const TCHAR* Type, FArchive& Archive, FFeedbackContext* Warn, int32 FileIndex, uint32 PortFlags)
 {
 	USkeletalMesh* SkeletalMesh = CastChecked<USkeletalMesh>(Object);
-	return Super::ExportBinary(SkeletalMesh, Type, Ar, Warn, FileIndex, PortFlags);
+
+	if (!FillExportOptions())
+	{
+		// User cancelled the export
+		return false;
+	}
+
+	// TODO: implement
+	return true;
 }
diff --git a/Engine/Plugins/Enterprise/GLTFExporter/MigratedFromGit/Source/GLTFExporter/Private/GLTFStaticMeshExporter.cpp b/Engine/Plugins/Enterprise/GLTFExporter/MigratedFromGit/Source/GLTFExporter/Private/GLTFStaticMeshExporter.cpp
--- a/Engine/Plugins/Enterprise/GLTFExporter/MigratedFromGit/Source/GLTFExporter/Private/GLTFStaticMeshExporter.cpp
+++ b/Engine/Plugins/Enterprise/GLTFExporter/MigratedFromGit/Source/GLTFExporter/Private/GLTFStaticMeshExporter.cpp
@@ -9,8 +9,16 @@ UGLTFStaticMeshExporter::UGLTFStaticMeshExporter(const FObjectInitializer& Objec
 	SupportedClass = UStaticMesh::StaticClass();
 }
 
-bool UGLTFStaticMeshExporter::ExportBinary(UObject* Object, const TCHAR* Type, FArchive& Ar, FFeedbackContext* Warn, int32 FileIndex, uint32 PortFlags)
+bool UGLTFStaticMeshExporter::ExportBinary(UObject* Object, const TCHAR* Type, FArchive& Archive, FFeedbackContext* Warn, int32 FileIndex, uint32 PortFlags)
 {
 	UStaticMesh* StaticMesh = CastChecked<UStaticMesh>(Object);
-	return Super::ExportBinary(StaticMesh, Type, Ar, Warn, FileIndex, PortFlags);
+
+	if (!FillExportOptions())
+	{
+		// User cancelled the export
+		return false;
+	}
+
+	// TODO: implement
+	return true;
 }
diff --git a/Engine/Plugins/Enterprise/GLTFExporter/MigratedFromGit/Source/GLTFExporter/Public/GLTFAnimSequenceExporter.h b/Engine/Plugins/Enterprise/GLTFExporter/MigratedFromGit/Source/GLTFExporter/Public/GLTFAnimSequenceExporter.h
--- a/Engine/Plugins/Enterprise/GLTFExporter/MigratedFromGit/Source/GLTFExporter/Public/GLTFAnimSequenceExporter.h
+++ b/Engine/Plugins/Enterprise/GLTFExporter/MigratedFromGit/Source/GLTFExporter/Public/GLTFAnimSequenceExporter.h
@@ -15,6 +15,6 @@ public:
 	UGLTFAnimSequenceExporter(const FObjectInitializer& ObjectInitializer = FObjectInitializer());
 
 	//~ Begin UExporter Interface
-	virtual bool ExportBinary(UObject* Object, const TCHAR* Type, FArchive& Ar, FFeedbackContext* Warn, int32 FileIndex = 0, uint32 PortFlags = 0) override;
+	virtual bool ExportBinary(UObject* Object, const TCHAR* Type, FArchive& Archive, FFeedbackContext* Warn, int32 FileIndex = 0, uint32 PortFlags = 0) override;
 	//~ End UExporter Interface
 };
diff --git a/Engine/Plugins/Enterprise/GLTFExporter/MigratedFromGit/Source/GLTFExporter/Public/GLTFExporter.h b/Engine/Plugins/Enterprise/GLTFExporter/MigratedFromGit/Source/GLTFExporter/Public/GLTFExporter.h
--- a/Engine/Plugins/Enterprise/GLTFExporter/MigratedFromGit/Source/GLTFExporter/Public/GLTFExporter.h
+++ b/Engine/Plugins/Enterprise/GLTFExporter/MigratedFromGit/Source/GLTFExporter/Public/GLTFExporter.h
@@ -18,10 +18,6 @@ public:
 public:
 	UGLTFExporter(const FObjectInitializer& ObjectInitializer = FObjectInitializer());
 
-	//~ Begin UExporter Interface
-	virtual bool ExportBinary(UObject* Object, const TCHAR* Type, FArchive& Ar, FFeedbackContext* Warn, int32 FileIndex = 0, uint32 PortFlags = 0) override;
-	//~ End UExporter Interface
-
 protected:
 
 	bool FillExportOptions();
diff --git a/Engine/Plugins/Enterprise/GLTFExporter/MigratedFromGit/Source/GLTFExporter/Public/GLTFLevelExporter.h b/Engine/Plugins/Enterprise/GLTFExporter/MigratedFromGit/Source/GLTFExporter/Public/GLTFLevelExporter.h
--- a/Engine/Plugins/Enterprise/GLTFExporter/MigratedFromGit/Source/GLTFExporter/Public/GLTFLevelExporter.h
+++ b/Engine/Plugins/Enterprise/GLTFExporter/MigratedFromGit/Source/GLTFExporter/Public/GLTFLevelExporter.h
@@ -15,6 +15,6 @@ public:
 	UGLTFLevelExporter(const FObjectInitializer& ObjectInitializer = FObjectInitializer());
 
 	//~ Begin UExporter Interface
-	virtual bool ExportBinary(UObject* Object, const TCHAR* Type, FArchive& Ar, FFeedbackContext* Warn, int32 FileIndex = 0, uint32 PortFlags = 0) override;
+	virtual bool ExportBinary(UObject* Object, const TCHAR* Type, FArchive& Archive, FFeedbackContext* Warn, int32 FileIndex = 0, uint32 PortFlags = 0) override;
 	//~ End UExporter Interface
 };
diff --git a/Engine/Plugins/Enterprise/GLTFExporter/MigratedFromGit/Source/GLTFExporter/Public/GLTFSkeletalMeshExporter.h b/Engine/Plugins/Enterprise/GLTFExporter/MigratedFromGit/Source/GLTFExporter/Public/GLTFSkeletalMeshExporter.h
--- a/Engine/Plugins/Enterprise/GLTFExporter/MigratedFromGit/Source/GLTFExporter/Public/GLTFSkeletalMeshExporter.h
+++ b/Engine/Plugins/Enterprise/GLTFExporter/MigratedFromGit/Source/GLTFExporter/Public/GLTFSkeletalMeshExporter.h
@@ -15,6 +15,6 @@ public:
 	UGLTFSkeletalMeshExporter(const FObjectInitializer& ObjectInitializer = FObjectInitializer());
 
 	//~ Begin UExporter Interface
-	virtual bool ExportBinary(UObject* Object, const TCHAR* Type, FArchive& Ar, FFeedbackContext* Warn, int32 FileIndex = 0, uint32 PortFlags = 0) override;
+	virtual bool ExportBinary(UObject* Object, const TCHAR* Type, FArchive& Archive, FFeedbackContext* Warn, int32 FileIndex = 0, uint32 PortFlags = 0) override;
 	//~ End UExporter Interface
 };
diff --git a/Engine/Plugins/Enterprise/GLTFExporter/MigratedFromGit/Source/GLTFExporter/Public/GLTFStaticMeshExporter.h b/Engine/Plugins/Enterprise/GLTFExporter/MigratedFromGit/Source/GLTFExporter/Public/GLTFStaticMeshExporter.h
--- a/Engine/Plugins/Enterprise/GLTFExporter/MigratedFromGit/Source/GLTFExporter/Public/GLTFStaticMeshExporter.h
+++ b/Engine/Plugins/Enterprise/GLTFExporter/MigratedFromGit/Source/GLTFExporter/Public/GLTFStaticMeshExporter.h
@@ -15,6 +15,6 @@ public:
 	UGLTFStaticMeshExporter(const FObjectInitializer& ObjectInitializer = FObjectInitializer());
 
 	//~ Begin UExporter Interface
-	virtual bool ExportBinary(UObject* Object, const TCHAR* Type, FArchive& Ar, FFeedbackContext* Warn, int32 FileIndex = 0, uint32 PortFlags = 0) override;
+	virtual bool ExportBinary(UObject* Object, const TCHAR* Type, FArchive& Archive, FFeedbackContext* Warn, int32 FileIndex = 0, uint32 PortFlags = 0) override;
 	//~ End UExporter Interface
 };
