[GLTFExporter] feat: add glTF JSON enum types

commit 12471a17e1cceadeefe4ca4662393a6e83adf8ec
Date: Fri, 24 Apr 2020 07:46:05 +0200

    feat: add glTF JSON enum types

#skipci

diff --git a/Engine/Plugins/Enterprise/GLTFExporter/MigratedFromGit/Source/GLTFExporter/Public/GLTFJsonEnums.h b/Engine/Plugins/Enterprise/GLTFExporter/MigratedFromGit/Source/GLTFExporter/Public/GLTFJsonEnums.h
new file mode 100644
--- /dev/null
+++ b/Engine/Plugins/Enterprise/GLTFExporter/MigratedFromGit/Source/GLTFExporter/Public/GLTFJsonEnums.h
@@ -0,0 +1,97 @@
+// Copyright Epic Games, Inc. All Rights Reserved.
+
+#pragma once
+
+#include "CoreMinimal.h"
+
+enum class EGLTFJsonAccessorType
+{
+	None = -1,
+	Scalar,
+	Vec2,
+	Vec3,
+	Vec4,
+	Mat2,
+	Mat3,
+	Mat4
+};
+
+enum class EGLTFJsonComponentType
+{
+	None = -1,
+	S8 = 5120, // signed byte
+	U8 = 5121, // unsigned byte
+	S16 = 5122, // signed short
+	U16 = 5123, // unsigned short
+	S32 = 5124, // unused
+	U32 = 5125, // unsigned int -- only valid for indices, not attributes
+	F32 = 5126  // float
+};
+
+enum class EGLTFJsonBufferTarget
+{
+	None = -1,
+	ArrayBuffer = 34962,
+	ElementArrayBuffer = 34963
+};
+
+enum class EGLTFJsonPrimitiveMode
+{
+	None = -1,
+	// unsupported
+	Points = 0,
+	Lines = 1,
+	LineLoop = 2,
+	LineStrip = 3,
+	// initially supported
+	Triangles = 4,
+	// will be supported prior to release
+	TriangleStrip = 5,
+	TriangleFan = 6
+};
+
+enum class EGLTFJsonAlphaMode
+{
+	None = -1,
+	Opaque,
+	Mask,
+	Blend
+};
+
+enum class EGLTFJsonSamplerFilter
+{
+	None = -1,
+	// valid for Min & Mag
+	Nearest = 9728,
+	Linear = 9729,
+	// valid for Min only
+	NearestMipmapNearest = 9984,
+	LinearMipmapNearest = 9985,
+	NearestMipmapLinear = 9986,
+	LinearMipmapLinear = 9987
+};
+
+enum class EGLTFJsonSamplerWrap
+{
+	None = -1,
+	Repeat = 10497,
+	MirroredRepeat = 33648,
+	ClampToEdge = 33071
+};
+
+enum class EGLTFJsonAnimationInterpolation
+{
+	None = -1,
+	Linear,
+	Step,
+	CubicSpline,
+};
+
+enum class EGLTFJsonAnimationPath
+{
+	None = -1,
+	Translation,
+	Rotation,
+	Scale,
+	Weights
+};
