[GLTFExporter] feat: add basic json objects GLTFJsonObject and GLTFJsonBuffer

commit 051541f51b1b38d4e8ed36a8bacdf5cb15fc3a0e
Date: Fri, 24 Apr 2020 07:50:59 +0200

    feat: add basic json objects GLTFJsonObject and GLTFJsonBuffer

#skipci

diff --git a/Engine/Plugins/Enterprise/GLTFExporter/MigratedFromGit/Source/GLTFExporter/Public/GLTFJsonBuffer.h b/Engine/Plugins/Enterprise/GLTFExporter/MigratedFromGit/Source/GLTFExporter/Public/GLTFJsonBuffer.h
new file mode 100644
--- /dev/null
+++ b/Engine/Plugins/Enterprise/GLTFExporter/MigratedFromGit/Source/GLTFExporter/Public/GLTFJsonBuffer.h
@@ -0,0 +1,18 @@
+// Copyright Epic Games, Inc. All Rights Reserved.
+
+#pragma once
+
+#include "GLTFJsonObject.h"
+
+struct GLTFEXPORTER_API FGLTFJsonBuffer : FGLTFJsonObject
+{
+	FString Name;
+
+	FString URI;
+	int64   ByteLength;
+
+	FGLTFJsonBuffer()
+		: ByteLength(0)
+	{
+	}
+};
diff --git a/Engine/Plugins/Enterprise/GLTFExporter/MigratedFromGit/Source/GLTFExporter/Public/GLTFJsonObject.h b/Engine/Plugins/Enterprise/GLTFExporter/MigratedFromGit/Source/GLTFExporter/Public/GLTFJsonObject.h
new file mode 100644
--- /dev/null
+++ b/Engine/Plugins/Enterprise/GLTFExporter/MigratedFromGit/Source/GLTFExporter/Public/GLTFJsonObject.h
@@ -0,0 +1,12 @@
+// Copyright Epic Games, Inc. All Rights Reserved.
+
+#pragma once
+
+#include "CoreMinimal.h"
+#include "Serialization/JsonSerializer.h"
+
+typedef int32 FGLTFJsonIndex;
+
+struct GLTFEXPORTER_API FGLTFJsonObject
+{
+};
