[GLTFExporter] feat: add json objects GLTFJsonAccessor and GLTFJsonBufferView

commit 226c191ad7eb84fa72f201694cb0df480a870b05
Date: Fri, 24 Apr 2020 07:51:31 +0200

    feat: add json objects GLTFJsonAccessor and GLTFJsonBufferView

#skipci

diff --git a/Engine/Plugins/Enterprise/GLTFExporter/MigratedFromGit/Source/GLTFExporter/Public/GLTFJsonAccessor.h b/Engine/Plugins/Enterprise/GLTFExporter/MigratedFromGit/Source/GLTFExporter/Public/GLTFJsonAccessor.h
new file mode 100644
--- /dev/null
+++ b/Engine/Plugins/Enterprise/GLTFExporter/MigratedFromGit/Source/GLTFExporter/Public/GLTFJsonAccessor.h
@@ -0,0 +1,31 @@
+// Copyright Epic Games, Inc. All Rights Reserved.
+
+#pragma once
+
+#include "GLTFJsonEnums.h"
+#include "GLTFJsonObject.h"
+
+struct GLTFEXPORTER_API FGLTFJsonAccessor : FGLTFJsonObject
+{
+	FString Name;
+
+	FGLTFJsonIndex         BufferView;
+	int32                  Count;
+	EGLTFJsonAccessorType  Type;
+	EGLTFJsonComponentType ComponentType;
+	bool                   Normalized;
+
+	int32 MinMaxLength;
+	float Min[16];
+	float Max[16];
+
+	FGLTFJsonAccessor()
+		: BufferView(INDEX_NONE)
+		, Count(0)
+		, Type(EGLTFJsonAccessorType::None)
+		, ComponentType(EGLTFJsonComponentType::None)
+		, Normalized(false)
+		, MinMaxLength(0)
+	{
+	}
+};
diff --git a/Engine/Plugins/Enterprise/GLTFExporter/MigratedFromGit/Source/GLTFExporter/Public/GLTFJsonBufferView.h b/Engine/Plugins/Enterprise/GLTFExporter/MigratedFromGit/Source/GLTFExporter/Public/GLTFJsonBufferView.h
new file mode 100644
--- /dev/null
+++ b/Engine/Plugins/Enterprise/GLTFExporter/MigratedFromGit/Source/GLTFExporter/Public/GLTFJsonBufferView.h
@@ -0,0 +1,28 @@
+// Copyright Epic Games, Inc. All Rights Reserved.
+
+#pragma once
+
+#include "GLTFJsonEnums.h"
+#include "GLTFJsonObject.h"
+
+struct GLTFEXPORTER_API FGLTFJsonBufferView : FGLTFJsonObject
+{
+	FString Name;
+
+	FGLTFJsonIndex Buffer;
+
+	int64 ByteLength;
+	int64 ByteOffset;
+	int32 ByteStride;
+
+	EGLTFJsonBufferTarget Target;
+
+	FGLTFJsonBufferView()
+		: Buffer(INDEX_NONE)
+		, ByteLength(0)
+		, ByteOffset(0)
+		, ByteStride(0)
+		, Target(EGLTFJsonBufferTarget::None)
+	{
+	}
+};
