[GLTFExporter] feat: add json objects GLTFJsonMesh, GLTFJsonPrimitive, and GLTFJsonAttributes

commit 1854d06ed92549dd0df511d0fe9dc6f396d478ef
Date: Fri, 24 Apr 2020 07:52:24 +0200

    feat: add json objects GLTFJsonMesh, GLTFJsonPrimitive, and GLTFJsonAttributes

#skipci

diff --git a/Engine/Plugins/Enterprise/GLTFExporter/MigratedFromGit/Source/GLTFExporter/Public/GLTFJsonMesh.h b/Engine/Plugins/Enterprise/GLTFExporter/MigratedFromGit/Source/GLTFExporter/Public/GLTFJsonMesh.h
new file mode 100644
--- /dev/null
+++ b/Engine/Plugins/Enterprise/GLTFExporter/MigratedFromGit/Source/GLTFExporter/Public/GLTFJsonMesh.h
@@ -0,0 +1,53 @@
+// Copyright Epic Games, Inc. All Rights Reserved.
+
+#pragma once
+
+#include "GLTFJsonEnums.h"
+#include "GLTFJsonObject.h"
+
+struct GLTFEXPORTER_API FGLTFJsonAttributes : FGLTFJsonObject
+{
+	FGLTFJsonIndex Position; // always required
+	FGLTFJsonIndex Normal;
+	FGLTFJsonIndex Tangent;
+	FGLTFJsonIndex TexCoord0;
+	FGLTFJsonIndex TexCoord1;
+	FGLTFJsonIndex Color0;
+	// skeletal mesh attributes
+	FGLTFJsonIndex Joints0;
+	FGLTFJsonIndex Weights0;
+
+	FGLTFJsonAttributes()
+		: Position(INDEX_NONE)
+		, Normal(INDEX_NONE)
+		, Tangent(INDEX_NONE)
+		, TexCoord0(INDEX_NONE)
+		, TexCoord1(INDEX_NONE)
+		, Color0(INDEX_NONE)
+		, Joints0(INDEX_NONE)
+		, Weights0(INDEX_NONE)
+	{
+	}
+};
+
+struct GLTFEXPORTER_API FGLTFJsonPrimitive : FGLTFJsonObject
+{
+	FGLTFJsonIndex         Indices;
+	FGLTFJsonIndex         Material;
+	EGLTFJsonPrimitiveMode Mode;
+	FGLTFJsonAttributes    Attributes;
+
+	FGLTFJsonPrimitive()
+		: Indices(INDEX_NONE)
+		, Material(INDEX_NONE)
+		, Mode(EGLTFJsonPrimitiveMode::None)
+	{
+	}
+};
+
+struct GLTFEXPORTER_API FGLTFJsonMesh : FGLTFJsonObject
+{
+	FString Name;
+
+	TArray<FGLTFJsonPrimitive> Primitives;
+};
