[GLTFExporter] feat: add json objects GLTFJsonNode and GLTFJsonScene

commit b5a656a21c4068e865b16e559068abd5eb50db43
Date: Fri, 24 Apr 2020 07:52:44 +0200

    feat: add json objects GLTFJsonNode and GLTFJsonScene

#skipci

diff --git a/Engine/Plugins/Enterprise/GLTFExporter/MigratedFromGit/Source/GLTFExporter/Public/GLTFJsonNode.h b/Engine/Plugins/Enterprise/GLTFExporter/MigratedFromGit/Source/GLTFExporter/Public/GLTFJsonNode.h
new file mode 100644
--- /dev/null
+++ b/Engine/Plugins/Enterprise/GLTFExporter/MigratedFromGit/Source/GLTFExporter/Public/GLTFJsonNode.h
@@ -0,0 +1,30 @@
+// Copyright Epic Games, Inc. All Rights Reserved.
+
+#pragma once
+
+#include "GLTFJsonObject.h"
+
+struct GLTFEXPORTER_API FGLTFJsonNode : FGLTFJsonObject
+{
+	FString Name;
+
+	FVector Translation;
+	FQuat   Rotation;
+	FVector Scale;
+
+	FGLTFJsonIndex Camera;
+	FGLTFJsonIndex Skin;
+	FGLTFJsonIndex Mesh;
+
+	TArray<FGLTFJsonIndex> Children;
+
+	FGLTFJsonNode()
+		: Translation(FVector::ZeroVector)
+		, Rotation(FQuat::Identity)
+		, Scale(FVector::OneVector)
+		, Camera(INDEX_NONE)
+		, Skin(INDEX_NONE)
+		, Mesh(INDEX_NONE)
+	{
+	}
+};
diff --git a/Engine/Plugins/Enterprise/GLTFExporter/MigratedFromGit/Source/GLTFExporter/Public/GLTFJsonScene.h b/Engine/Plugins/Enterprise/GLTFExporter/MigratedFromGit/Source/GLTFExporter/Public/GLTFJsonScene.h
new file mode 100644
--- /dev/null
+++ b/Engine/Plugins/Enterprise/GLTFExporter/MigratedFromGit/Source/GLTFExporter/Public/GLTFJsonScene.h
@@ -0,0 +1,12 @@
+// Copyright Epic Games, Inc. All Rights Reserved.
+
+#pragma once
+
+#include "GLTFJsonObject.h"
+
+struct GLTFEXPORTER_API FGLTFJsonScene: FGLTFJsonObject
+{
+	FString Name;
+
+	TArray<FGLTFJsonIndex> Nodes;
+};
