[GLTFExporter] feat: add json objects GLTFJsonRoot and GLTFJsonAsset

commit 9c7bed5b511994279f0c52b83a21dd83807260e1
Date: Fri, 24 Apr 2020 07:53:04 +0200

    feat: add json objects GLTFJsonRoot and GLTFJsonAsset

#skipci

diff --git a/Engine/Plugins/Enterprise/GLTFExporter/MigratedFromGit/Source/GLTFExporter/Public/GLTFJsonRoot.h b/Engine/Plugins/Enterprise/GLTFExporter/MigratedFromGit/Source/GLTFExporter/Public/GLTFJsonRoot.h
new file mode 100644
--- /dev/null
+++ b/Engine/Plugins/Enterprise/GLTFExporter/MigratedFromGit/Source/GLTFExporter/Public/GLTFJsonRoot.h
@@ -0,0 +1,45 @@
+// Copyright Epic Games, Inc. All Rights Reserved.
+
+#pragma once
+
+#include "GLTFJsonObject.h"
+#include "GLTFJsonAccessor.h"
+#include "GLTFJsonBufferView.h"
+#include "GLTFJsonBuffer.h"
+#include "GLTFJsonMesh.h"
+#include "GLTFJsonScene.h"
+#include "GLTFJsonNode.h"
+
+#include "Runtime/Launch/Resources/Version.h"
+
+
+struct GLTFEXPORTER_API FGLTFJsonAsset : FGLTFJsonObject
+{
+	FString Version;
+	FString Generator;
+	FString Copyright;
+
+	FGLTFJsonAsset()
+		: Version(TEXT("2.0"))
+		, Generator(TEXT(EPIC_PRODUCT_NAME) TEXT(" ") ENGINE_VERSION_STRING)
+	{
+	}
+};
+
+struct GLTFEXPORTER_API FGLTFJsonRoot
+{
+	FGLTFJsonAsset Asset;
+	FGLTFJsonIndex DefaultScene;
+
+	TArray<FGLTFJsonAccessor>   Accessors;
+	TArray<FGLTFJsonBuffer>     Buffers;
+	TArray<FGLTFJsonBufferView> BufferViews;
+	TArray<FGLTFJsonMesh>       Meshes;
+	TArray<FGLTFJsonNode>       Nodes;
+	TArray<FGLTFJsonScene>      Scenes;
+
+	FGLTFJsonRoot()
+		: DefaultScene(INDEX_NONE)
+	{
+	}
+};
