[GLTFExporter] feat: add GLTFJsonUtilities for converting enum values to json

commit 3107ee09a7bff4e55534a1985b915594d007db54
Date: Fri, 24 Apr 2020 07:56:05 +0200

    feat: add GLTFJsonUtilities for converting enum values to json

#skipci

diff --git a/Engine/Plugins/Enterprise/GLTFExporter/MigratedFromGit/Source/GLTFExporter/Public/GLTFJsonUtilities.h b/Engine/Plugins/Enterprise/GLTFExporter/MigratedFromGit/Source/GLTFExporter/Public/GLTFJsonUtilities.h
new file mode 100644
--- /dev/null
+++ b/Engine/Plugins/Enterprise/GLTFExporter/MigratedFromGit/Source/GLTFExporter/Public/GLTFJsonUtilities.h
@@ -0,0 +1,60 @@
+// Copyright Epic Games, Inc. All Rights Reserved.
+
+#pragma once
+
+#include "GLTFJsonEnums.h"
+
+inline const TCHAR* AlphaModeToString(EGLTFJsonAlphaMode Mode)
+{
+	switch (Mode)
+	{
+		case EGLTFJsonAlphaMode::Opaque: return TEXT("OPAQUE");
+		case EGLTFJsonAlphaMode::Blend:  return TEXT("BLEND");
+		case EGLTFJsonAlphaMode::Mask:   return TEXT("MASK");
+	}
+
+	check(false);
+	return TEXT("UNKNOWN");
+}
+
+inline int32 PrimitiveModeToNumber(EGLTFJsonPrimitiveMode Mode)
+{
+	return static_cast<int32>(Mode);
+}
+
+inline int32 FilterToNumber(EGLTFJsonSamplerFilter Filter)
+{
+	return static_cast<int32>(Filter);
+}
+
+inline int32 WrapModeToNumber(EGLTFJsonSamplerWrap Wrap)
+{
+	return static_cast<int32>(Wrap);
+}
+
+inline const TCHAR* AccessorTypeToString(EGLTFJsonAccessorType Type)
+{
+	switch (Type)
+	{
+		case EGLTFJsonAccessorType::Scalar: return TEXT("SCALAR");
+		case EGLTFJsonAccessorType::Vec2:   return TEXT("VEC2");
+		case EGLTFJsonAccessorType::Vec3:   return TEXT("VEC3");
+		case EGLTFJsonAccessorType::Vec4:   return TEXT("VEC4");
+		case EGLTFJsonAccessorType::Mat2:   return TEXT("MAT2");
+		case EGLTFJsonAccessorType::Mat3:   return TEXT("MAT3");
+		case EGLTFJsonAccessorType::Mat4:   return TEXT("MAT4");
+	}
+
+	check(false);
+	return TEXT("UNKNOWN");
+}
+
+inline int32 ComponentTypeToNumber(EGLTFJsonComponentType Type)
+{
+	return static_cast<int32>(Type);
+}
+
+inline int32 BufferTargetToNumber(EGLTFJsonBufferTarget Target)
+{
+	return static_cast<int32>(Target);
+}
