[GLTFExporter] feat: add a Write method to each GLTFJson object

commit d409bc3140d546b546650c97795efd050c20b17d
Date: Fri, 24 Apr 2020 07:57:05 +0200

    feat: add a Write method to each GLTFJson object

#skipci

diff --git a/Engine/Plugins/Enterprise/GLTFExporter/MigratedFromGit/Source/GLTFExporter/Public/GLTFJsonAccessor.h b/Engine/Plugins/Enterprise/GLTFExporter/MigratedFromGit/Source/GLTFExporter/Public/GLTFJsonAccessor.h
--- a/Engine/Plugins/Enterprise/GLTFExporter/MigratedFromGit/Source/GLTFExporter/Public/GLTFJsonAccessor.h
+++ b/Engine/Plugins/Enterprise/GLTFExporter/MigratedFromGit/Source/GLTFExporter/Public/GLTFJsonAccessor.h
@@ -4,6 +4,7 @@
 
 #include "GLTFJsonEnums.h"
 #include "GLTFJsonObject.h"
+#include "GLTFJsonUtilities.h"
 
 struct GLTFEXPORTER_API FGLTFJsonAccessor : FGLTFJsonObject
 {
@@ -28,4 +29,36 @@ struct GLTFEXPORTER_API FGLTFJsonAccessor : FGLTFJsonObject
 		, MinMaxLength(0)
 	{
 	}
+
+	template <class CharType = TCHAR, class PrintPolicy = TPrettyJsonPrintPolicy<CharType>>
+	void Write(TJsonWriter<CharType, PrintPolicy>& JsonWriter) const
+	{
+		JsonWriter.WriteObjectStart();
+
+		if (!Name.IsEmpty()) JsonWriter.WriteValue(TEXT("name"), Name);
+
+		JsonWriter.WriteValue(TEXT("bufferView"), BufferView);
+		JsonWriter.WriteValue(TEXT("count"), Count);
+		JsonWriter.WriteValue(TEXT("type"), AccessorTypeToString(Type));
+		JsonWriter.WriteValue(TEXT("componentType"), ComponentTypeToNumber(ComponentType));
+
+		if (MinMaxLength > 0)
+		{
+			JsonWriter.WriteArrayStart(TEXT("min"));
+			for (int32 ComponentIndex = 0; ComponentIndex < MinMaxLength; ++ComponentIndex)
+			{
+				JsonWriter.WriteValue(Min[ComponentIndex]);
+			}
+			JsonWriter.WriteArrayEnd();
+
+			JsonWriter.WriteArrayStart(TEXT("max"));
+			for (int32 ComponentIndex = 0; ComponentIndex < MinMaxLength; ++ComponentIndex)
+			{
+				JsonWriter.WriteValue(Max[ComponentIndex]);
+			}
+			JsonWriter.WriteArrayEnd();
+		}
+
+		JsonWriter.WriteObjectEnd();
+	}
 };
diff --git a/Engine/Plugins/Enterprise/GLTFExporter/MigratedFromGit/Source/GLTFExporter/Public/GLTFJsonBuffer.h b/Engine/Plugins/Enterprise/GLTFExporter/MigratedFromGit/Source/GLTFExporter/Public/GLTFJsonBuffer.h
--- a/Engine/Plugins/Enterprise/GLTFExporter/MigratedFromGit/Source/GLTFExporter/Public/GLTFJsonBuffer.h
+++ b/Engine/Plugins/Enterprise/GLTFExporter/MigratedFromGit/Source/GLTFExporter/Public/GLTFJsonBuffer.h
@@ -15,4 +15,17 @@ struct GLTFEXPORTER_API FGLTFJsonBuffer : FGLTFJsonObject
 		: ByteLength(0)
 	{
 	}
+
+	template <class CharType = TCHAR, class PrintPolicy = TPrettyJsonPrintPolicy<CharType>>
+	void Write(TJsonWriter<CharType, PrintPolicy>& JsonWriter) const
+	{
+		JsonWriter.WriteObjectStart();
+
+		if (!Name.IsEmpty()) JsonWriter.WriteValue(TEXT("name"), Name);
+
+		if (!URI.IsEmpty()) JsonWriter.WriteValue(TEXT("uri"), URI);
+		JsonWriter.WriteValue(TEXT("byteLength"), ByteLength);
+
+		JsonWriter.WriteObjectEnd();
+	}
 };
diff --git a/Engine/Plugins/Enterprise/GLTFExporter/MigratedFromGit/Source/GLTFExporter/Public/GLTFJsonBufferView.h b/Engine/Plugins/Enterprise/GLTFExporter/MigratedFromGit/Source/GLTFExporter/Public/GLTFJsonBufferView.h
--- a/Engine/Plugins/Enterprise/GLTFExporter/MigratedFromGit/Source/GLTFExporter/Public/GLTFJsonBufferView.h
+++ b/Engine/Plugins/Enterprise/GLTFExporter/MigratedFromGit/Source/GLTFExporter/Public/GLTFJsonBufferView.h
@@ -4,6 +4,7 @@
 
 #include "GLTFJsonEnums.h"
 #include "GLTFJsonObject.h"
+#include "GLTFJsonUtilities.h"
 
 struct GLTFEXPORTER_API FGLTFJsonBufferView : FGLTFJsonObject
 {
@@ -24,5 +25,24 @@ struct GLTFEXPORTER_API FGLTFJsonBufferView : FGLTFJsonObject
 		, ByteStride(0)
 		, Target(EGLTFJsonBufferTarget::None)
 	{
+		// check that view fits completely inside the buffer
+	}
+
+	template <class CharType = TCHAR, class PrintPolicy = TPrettyJsonPrintPolicy<CharType>>
+	void Write(TJsonWriter<CharType, PrintPolicy>& JsonWriter) const
+	{
+		JsonWriter.WriteObjectStart();
+
+		if (!Name.IsEmpty()) JsonWriter.WriteValue(TEXT("name"), Name);
+
+		JsonWriter.WriteValue(TEXT("buffer"), Buffer);
+		JsonWriter.WriteValue(TEXT("byteLength"), ByteLength);
+
+		if (ByteOffset != 0) JsonWriter.WriteValue(TEXT("byteOffset"), ByteOffset);
+		if (ByteStride != 0) JsonWriter.WriteValue(TEXT("byteStride"), ByteStride);
+
+		if (Target != EGLTFJsonBufferTarget::None) JsonWriter.WriteValue(TEXT("target"), BufferTargetToNumber(Target));
+
+		JsonWriter.WriteObjectEnd();
 	}
 };
diff --git a/Engine/Plugins/Enterprise/GLTFExporter/MigratedFromGit/Source/GLTFExporter/Public/GLTFJsonMesh.h b/Engine/Plugins/Enterprise/GLTFExporter/MigratedFromGit/Source/GLTFExporter/Public/GLTFJsonMesh.h
--- a/Engine/Plugins/Enterprise/GLTFExporter/MigratedFromGit/Source/GLTFExporter/Public/GLTFJsonMesh.h
+++ b/Engine/Plugins/Enterprise/GLTFExporter/MigratedFromGit/Source/GLTFExporter/Public/GLTFJsonMesh.h
@@ -28,6 +28,23 @@ struct GLTFEXPORTER_API FGLTFJsonAttributes : FGLTFJsonObject
 		, Weights0(INDEX_NONE)
 	{
 	}
+
+	template <class CharType = TCHAR, class PrintPolicy = TPrettyJsonPrintPolicy<CharType>>
+	void Write(TJsonWriter<CharType, PrintPolicy>& JsonWriter) const
+	{
+		JsonWriter.WriteObjectStart();
+
+		JsonWriter.WriteValue(TEXT("POSITION"), Position);
+		if (Normal != INDEX_NONE) JsonWriter.WriteValue(TEXT("NORMAL"), Normal);
+		if (Tangent != INDEX_NONE) JsonWriter.WriteValue(TEXT("TANGENT"), Tangent);
+		if (TexCoord0 != INDEX_NONE) JsonWriter.WriteValue(TEXT("TEXCOORD_0"), TexCoord0);
+		if (TexCoord1 != INDEX_NONE) JsonWriter.WriteValue(TEXT("TEXCOORD_1"), TexCoord1);
+		if (Color0 != INDEX_NONE) JsonWriter.WriteValue(TEXT("COLOR_0"), Color0);
+		if (Joints0 != INDEX_NONE) JsonWriter.WriteValue(TEXT("JOINTS_0"), Joints0);
+		if (Weights0 != INDEX_NONE) JsonWriter.WriteValue(TEXT("WEIGHTS_0"), Weights0);
+
+		JsonWriter.WriteObjectEnd();
+	}
 };
 
 struct GLTFEXPORTER_API FGLTFJsonPrimitive : FGLTFJsonObject
@@ -43,6 +60,21 @@ struct GLTFEXPORTER_API FGLTFJsonPrimitive : FGLTFJsonObject
 		, Mode(EGLTFJsonPrimitiveMode::None)
 	{
 	}
+
+	template <class CharType = TCHAR, class PrintPolicy = TPrettyJsonPrintPolicy<CharType>>
+	void Write(TJsonWriter<CharType, PrintPolicy>& JsonWriter) const
+	{
+		JsonWriter.WriteObjectStart();
+
+		JsonWriter.WriteIdentifierPrefix(TEXT("attributes"));
+		Attributes.Write(JsonWriter);
+
+		if (Indices != INDEX_NONE) JsonWriter.WriteValue(TEXT("indices"), Indices);
+		if (Material != INDEX_NONE) JsonWriter.WriteValue(TEXT("material"), Material);
+		if (Mode != EGLTFJsonPrimitiveMode::None) JsonWriter.WriteValue(TEXT("mode"), PrimitiveModeToNumber(Mode));
+
+		JsonWriter.WriteObjectEnd();
+	}
 };
 
 struct GLTFEXPORTER_API FGLTFJsonMesh : FGLTFJsonObject
@@ -50,4 +82,21 @@ struct GLTFEXPORTER_API FGLTFJsonMesh : FGLTFJsonObject
 	FString Name;
 
 	TArray<FGLTFJsonPrimitive> Primitives;
+
+	template <class CharType = TCHAR, class PrintPolicy = TPrettyJsonPrintPolicy<CharType>>
+	void Write(TJsonWriter<CharType, PrintPolicy>& JsonWriter) const
+	{
+		JsonWriter.WriteObjectStart();
+
+		if (!Name.IsEmpty()) JsonWriter.WriteValue(TEXT("name"), Name);
+
+		JsonWriter.WriteArrayStart(TEXT("primitives"));
+		for (const FGLTFJsonPrimitive& Primitive : Primitives)
+		{
+			Primitive.Write(JsonWriter);
+		}
+		JsonWriter.WriteArrayEnd();
+
+		JsonWriter.WriteObjectEnd();
+	}
 };
diff --git a/Engine/Plugins/Enterprise/GLTFExporter/MigratedFromGit/Source/GLTFExporter/Public/GLTFJsonNode.h b/Engine/Plugins/Enterprise/GLTFExporter/MigratedFromGit/Source/GLTFExporter/Public/GLTFJsonNode.h
--- a/Engine/Plugins/Enterprise/GLTFExporter/MigratedFromGit/Source/GLTFExporter/Public/GLTFJsonNode.h
+++ b/Engine/Plugins/Enterprise/GLTFExporter/MigratedFromGit/Source/GLTFExporter/Public/GLTFJsonNode.h
@@ -27,4 +27,51 @@ struct GLTFEXPORTER_API FGLTFJsonNode : FGLTFJsonObject
 		, Mesh(INDEX_NONE)
 	{
 	}
+
+	template <class CharType = TCHAR, class PrintPolicy = TPrettyJsonPrintPolicy<CharType>>
+	void Write(TJsonWriter<CharType, PrintPolicy>& JsonWriter) const
+	{
+		JsonWriter.WriteObjectStart();
+
+		if (!Name.IsEmpty()) JsonWriter.WriteValue(TEXT("name"), Name);
+
+		if (Translation != FVector::ZeroVector)
+		{
+			JsonWriter.WriteArrayStart(TEXT("translation"));
+			JsonWriter.WriteValue(Translation.X);
+			JsonWriter.WriteValue(Translation.Y);
+			JsonWriter.WriteValue(Translation.Z);
+			JsonWriter.WriteArrayEnd();
+		}
+
+		if (Rotation != FQuat::Identity)
+		{
+			JsonWriter.WriteArrayStart(TEXT("rotation"));
+			JsonWriter.WriteValue(Rotation.X);
+			JsonWriter.WriteValue(Rotation.Y);
+			JsonWriter.WriteValue(Rotation.Z);
+			JsonWriter.WriteValue(Rotation.W);
+			JsonWriter.WriteArrayEnd();
+		}
+
+		if (Scale != FVector::OneVector)
+		{
+			JsonWriter.WriteArrayStart(TEXT("scale"));
+			JsonWriter.WriteValue(Scale.X);
+			JsonWriter.WriteValue(Scale.Y);
+			JsonWriter.WriteValue(Scale.Z);
+			JsonWriter.WriteArrayEnd();
+		}
+
+		if (Camera != INDEX_NONE) JsonWriter.WriteValue(TEXT("camera"), Camera);
+		if (Skin != INDEX_NONE) JsonWriter.WriteValue(TEXT("skin"), Skin);
+		if (Mesh != INDEX_NONE) JsonWriter.WriteValue(TEXT("mesh"), Mesh);
+
+		if (Children.Num() > 0)
+		{
+			JsonWriter.WriteValue(TEXT("children"), Children);
+		}
+
+		JsonWriter.WriteObjectEnd();
+	}
 };
diff --git a/Engine/Plugins/Enterprise/GLTFExporter/MigratedFromGit/Source/GLTFExporter/Public/GLTFJsonRoot.h b/Engine/Plugins/Enterprise/GLTFExporter/MigratedFromGit/Source/GLTFExporter/Public/GLTFJsonRoot.h
--- a/Engine/Plugins/Enterprise/GLTFExporter/MigratedFromGit/Source/GLTFExporter/Public/GLTFJsonRoot.h
+++ b/Engine/Plugins/Enterprise/GLTFExporter/MigratedFromGit/Source/GLTFExporter/Public/GLTFJsonRoot.h
@@ -10,6 +10,8 @@
 #include "GLTFJsonScene.h"
 #include "GLTFJsonNode.h"
 
+#include "Containers/Set.h"
+#include "Policies/CondensedJsonPrintPolicy.h"
 #include "Runtime/Launch/Resources/Version.h"
 
 
@@ -24,6 +26,18 @@ struct GLTFEXPORTER_API FGLTFJsonAsset : FGLTFJsonObject
 		, Generator(TEXT(EPIC_PRODUCT_NAME) TEXT(" ") ENGINE_VERSION_STRING)
 	{
 	}
+
+	template <class CharType = TCHAR, class PrintPolicy = TPrettyJsonPrintPolicy<CharType>>
+	void Write(TJsonWriter<CharType, PrintPolicy>& JsonWriter) const
+	{
+		JsonWriter.WriteObjectStart();
+
+		JsonWriter.WriteValue(TEXT("version"), Version);
+		if (!Generator.IsEmpty()) JsonWriter.WriteValue(TEXT("generator"), Generator);
+		if (!Copyright.IsEmpty()) JsonWriter.WriteValue(TEXT("copyright"), Copyright);
+
+		JsonWriter.WriteObjectEnd();
+	}
 };
 
 struct GLTFEXPORTER_API FGLTFJsonRoot
@@ -42,4 +56,77 @@ struct GLTFEXPORTER_API FGLTFJsonRoot
 		: DefaultScene(INDEX_NONE)
 	{
 	}
+
+	template <class CharType = TCHAR, class PrintPolicy = TPrettyJsonPrintPolicy<CharType>>
+	void Write(TJsonWriter<CharType, PrintPolicy>& JsonWriter) const
+	{
+		JsonWriter.WriteObjectStart();
+
+		JsonWriter.WriteIdentifierPrefix(TEXT("asset"));
+		Asset.Write(JsonWriter);
+
+		if (DefaultScene != INDEX_NONE) JsonWriter.WriteValue(TEXT("scene"), DefaultScene);
+
+		if (Accessors.Num() > 0)
+		{
+			JsonWriter.WriteArrayStart(TEXT("accessors"));
+			for (const FGLTFJsonAccessor& Accessor : Accessors)
+			{
+				Accessor.Write(JsonWriter);
+			}
+			JsonWriter.WriteArrayEnd();
+		}
+
+		if (BufferViews.Num() > 0)
+		{
+			JsonWriter.WriteArrayStart(TEXT("bufferViews"));
+			for (const FGLTFJsonBufferView& BufferView : BufferViews)
+			{
+				BufferView.Write(JsonWriter);
+			}
+			JsonWriter.WriteArrayEnd();
+		}
+
+		if (Meshes.Num() > 0)
+		{
+			JsonWriter.WriteArrayStart(TEXT("meshes"));
+			for (const FGLTFJsonMesh& Mesh : Meshes)
+			{
+				Mesh.Write(JsonWriter);
+			}
+			JsonWriter.WriteArrayEnd();
+		}
+
+		if (Nodes.Num() > 0)
+		{
+			JsonWriter.WriteArrayStart(TEXT("nodes"));
+			for (const FGLTFJsonNode& Node : Nodes)
+			{
+				Node.Write(JsonWriter);
+			}
+			JsonWriter.WriteArrayEnd();
+		}
+
+		if (Scenes.Num() > 0)
+		{
+			JsonWriter.WriteArrayStart(TEXT("scenes"));
+			for (const FGLTFJsonScene& Scene : Scenes)
+			{
+				Scene.Write(JsonWriter);
+			}
+			JsonWriter.WriteArrayEnd();
+		}
+
+		if (Buffers.Num() > 0)
+		{
+			JsonWriter.WriteArrayStart(TEXT("buffers"));
+			for (const FGLTFJsonBuffer& Buffer : Buffers)
+			{
+				Buffer.Write(JsonWriter);
+			}
+			JsonWriter.WriteArrayEnd();
+		}
+
+		JsonWriter.WriteObjectEnd();
+	}
 };
diff --git a/Engine/Plugins/Enterprise/GLTFExporter/MigratedFromGit/Source/GLTFExporter/Public/GLTFJsonScene.h b/Engine/Plugins/Enterprise/GLTFExporter/MigratedFromGit/Source/GLTFExporter/Public/GLTFJsonScene.h
--- a/Engine/Plugins/Enterprise/GLTFExporter/MigratedFromGit/Source/GLTFExporter/Public/GLTFJsonScene.h
+++ b/Engine/Plugins/Enterprise/GLTFExporter/MigratedFromGit/Source/GLTFExporter/Public/GLTFJsonScene.h
@@ -9,4 +9,19 @@ struct GLTFEXPORTER_API FGLTFJsonScene: FGLTFJsonObject
 	FString Name;
 
 	TArray<FGLTFJsonIndex> Nodes;
+
+	template <class CharType = TCHAR, class PrintPolicy = TPrettyJsonPrintPolicy<CharType>>
+	void Write(TJsonWriter<CharType, PrintPolicy>& JsonWriter) const
+	{
+		JsonWriter.WriteObjectStart();
+
+		if (!Name.IsEmpty()) JsonWriter.WriteValue(TEXT("name"), Name);
+
+		if (Nodes.Num() > 0)
+		{
+			JsonWriter.WriteValue(TEXT("nodes"), Nodes);
+		}
+
+		JsonWriter.WriteObjectEnd();
+	}
 };
