[GLTFExporter] feat: initial json export of StaticMesh without any vertex data

commit 7f6d82744db385f77e44ff70de1cfd903ec6e452
Date: Fri, 24 Apr 2020 08:11:29 +0200

    feat: initial json export of StaticMesh without any vertex data

#skipci

diff --git a/Engine/Plugins/Enterprise/GLTFExporter/MigratedFromGit/Source/GLTFExporter/Private/GLTFStaticMeshExporter.cpp b/Engine/Plugins/Enterprise/GLTFExporter/MigratedFromGit/Source/GLTFExporter/Private/GLTFStaticMeshExporter.cpp
--- a/Engine/Plugins/Enterprise/GLTFExporter/MigratedFromGit/Source/GLTFExporter/Private/GLTFStaticMeshExporter.cpp
+++ b/Engine/Plugins/Enterprise/GLTFExporter/MigratedFromGit/Source/GLTFExporter/Private/GLTFStaticMeshExporter.cpp
@@ -1,6 +1,7 @@
 // Copyright Epic Games, Inc. All Rights Reserved.
 
 #include "GLTFStaticMeshExporter.h"
+#include "GLTFJsonRoot.h"
 #include "Engine/StaticMesh.h"
 
 UGLTFStaticMeshExporter::UGLTFStaticMeshExporter(const FObjectInitializer& ObjectInitializer)
@@ -19,6 +20,35 @@ bool UGLTFStaticMeshExporter::ExportBinary(UObject* Object, const TCHAR* Type, F
 		return false;
 	}
 
-	// TODO: implement
+	FGLTFJsonRoot Root;
+
+	FGLTFJsonBuffer Buffer;
+	FGLTFJsonIndex BufferIndex = Root.Buffers.Add(Buffer);
+
+	FGLTFJsonBufferView BufferView;
+	BufferView.Buffer = BufferIndex;
+	FGLTFJsonIndex BufferViewIndex = Root.BufferViews.Add(BufferView);
+
+	FGLTFJsonAccessor Accessor;
+	Accessor.BufferView = BufferViewIndex;
+	FGLTFJsonIndex AccessorIndex = Root.Accessors.Add(Accessor);
+
+	FGLTFJsonMesh Mesh;
+	FGLTFJsonPrimitive Primitive;
+	Primitive.Attributes.Position = AccessorIndex;
+	Mesh.Primitives.Add(Primitive);
+	FGLTFJsonIndex MeshIndex = Root.Meshes.Add(Mesh);
+
+	FGLTFJsonNode Node;
+	Node.Mesh = MeshIndex;
+	FGLTFJsonIndex NodeIndex = Root.Nodes.Add(Node);
+
+	FGLTFJsonScene Scene;
+	Scene.Nodes.Add(NodeIndex);
+	FGLTFJsonIndex SceneIndex = Root.Scenes.Add(Scene);
+
+	Root.DefaultScene = SceneIndex;
+
+	Root.Serialize(&Archive, true);
 	return true;
 }
diff --git a/Engine/Plugins/Enterprise/GLTFExporter/MigratedFromGit/Source/GLTFExporter/Public/GLTFJsonRoot.h b/Engine/Plugins/Enterprise/GLTFExporter/MigratedFromGit/Source/GLTFExporter/Public/GLTFJsonRoot.h
--- a/Engine/Plugins/Enterprise/GLTFExporter/MigratedFromGit/Source/GLTFExporter/Public/GLTFJsonRoot.h
+++ b/Engine/Plugins/Enterprise/GLTFExporter/MigratedFromGit/Source/GLTFExporter/Public/GLTFJsonRoot.h
@@ -129,4 +129,24 @@ struct GLTFEXPORTER_API FGLTFJsonRoot
 
 		JsonWriter.WriteObjectEnd();
 	}
+
+	template <class CharType = TCHAR, class PrintPolicy = TPrettyJsonPrintPolicy<CharType>>
+	void Serialize(FArchive* const Archive) const
+	{
+		TSharedRef<TJsonWriter<CharType, PrintPolicy>> JsonWriter = TJsonWriterFactory<CharType, PrintPolicy>::Create(Archive);
+		Write(*JsonWriter);
+		JsonWriter->Close();
+	}
+
+	void Serialize(FArchive* const Archive, bool bPrettyPrint = true) const
+	{
+		if (bPrettyPrint)
+		{
+			Serialize<UTF8CHAR, TPrettyJsonPrintPolicy<UTF8CHAR>>(Archive);
+		}
+		else
+		{
+			Serialize<UTF8CHAR, TCondensedJsonPrintPolicy<UTF8CHAR>>(Archive);
+		}
+	}
 };
