[GLTFExporter] feat: add GLTFBufferBuilder

commit 74ee4e3847d4dbc55e8e13dec72f6ff9870da378
Date: Fri, 24 Apr 2020 08:11:48 +0200

    feat: add GLTFBufferBuilder

#skipci

diff --git a/Engine/Plugins/Enterprise/GLTFExporter/MigratedFromGit/Source/GLTFExporter/Private/GLTFBufferBuilder.cpp b/Engine/Plugins/Enterprise/GLTFExporter/MigratedFromGit/Source/GLTFExporter/Private/GLTFBufferBuilder.cpp
new file mode 100644
--- /dev/null
+++ b/Engine/Plugins/Enterprise/GLTFExporter/MigratedFromGit/Source/GLTFExporter/Private/GLTFBufferBuilder.cpp
@@ -0,0 +1,35 @@
+// Copyright Epic Games, Inc. All Rights Reserved.
+
+#include "GLTFBufferBuilder.h"
+#include "Misc/Base64.h"
+
+inline FGLTFBufferBuilder::FGLTFBufferBuilder(FGLTFJsonRoot& JsonRoot, const FString& Name)
+	: JsonRoot(JsonRoot)
+	, BufferIndex(JsonRoot.Buffers.AddDefaulted(1))
+{
+	JsonRoot.Buffers[BufferIndex].Name = Name;
+}
+
+inline FGLTFJsonIndex FGLTFBufferBuilder::AppendBufferView(const void* RawData, uint64 ByteLength, const FString& Name, EGLTFJsonBufferTarget BufferTarget)
+{
+	FGLTFJsonBufferView BufferView;
+	BufferView.Name = Name;
+	BufferView.Buffer = BufferIndex;
+	BufferView.ByteOffset = Data.Num();
+	BufferView.ByteLength = ByteLength;
+	BufferView.Target = BufferTarget;
+
+	Data.Append(static_cast<const uint8*>(RawData), ByteLength);
+
+	FGLTFJsonIndex BufferViewIndex = JsonRoot.BufferViews.Add(BufferView);
+	return BufferViewIndex;
+}
+
+void FGLTFBufferBuilder::Close()
+{
+	FString DataBase64 = FBase64::Encode(Data.GetData(), Data.Num());
+
+	FGLTFJsonBuffer& Buffer = JsonRoot.Buffers[BufferIndex];
+	Buffer.ByteLength = Data.Num();
+	Buffer.URI = TEXT("data:application/octet-stream;base64,") + DataBase64;
+}
diff --git a/Engine/Plugins/Enterprise/GLTFExporter/MigratedFromGit/Source/GLTFExporter/Public/GLTFBufferBuilder.h b/Engine/Plugins/Enterprise/GLTFExporter/MigratedFromGit/Source/GLTFExporter/Public/GLTFBufferBuilder.h
new file mode 100644
--- /dev/null
+++ b/Engine/Plugins/Enterprise/GLTFExporter/MigratedFromGit/Source/GLTFExporter/Public/GLTFBufferBuilder.h
@@ -0,0 +1,26 @@
+// Copyright Epic Games, Inc. All Rights Reserved.
+
+#pragma once
+
+#include "GLTFJsonRoot.h"
+#include "CoreMinimal.h"
+
+struct GLTFEXPORTER_API FGLTFBufferBuilder
+{
+	FGLTFJsonRoot& JsonRoot;
+	const FGLTFJsonIndex BufferIndex;
+
+	TArray<uint8> Data;
+
+	FGLTFBufferBuilder(FGLTFJsonRoot& JsonRoot, const FString& Name = TEXT(""));
+
+	template <class ElementType>
+	FGLTFJsonIndex AppendBufferView(const TArray<ElementType>& Array, const FString& Name = TEXT(""), EGLTFJsonBufferTarget BufferTarget = EGLTFJsonBufferTarget::ArrayBuffer)
+	{
+		return AppendBufferView(Array.GetData(), Array.Num() * sizeof(ElementType), Name, BufferTarget);
+	}
+
+	FGLTFJsonIndex AppendBufferView(const void* RawData, uint64 ByteLength, const FString& Name = TEXT(""), EGLTFJsonBufferTarget BufferTarget = EGLTFJsonBufferTarget::ArrayBuffer);
+
+	void Close();
+};
