[GLTFExporter] feat: add GLTFConversionUtilities for converting vectors etc

commit 6d9f122d37476d85399a5c4a29e2bf4b7e0c8942
Date: Fri, 24 Apr 2020 08:12:36 +0200

    feat: add GLTFConversionUtilities for converting vectors etc

#skipci

diff --git a/Engine/Plugins/Enterprise/GLTFExporter/MigratedFromGit/Source/GLTFExporter/Public/GLTFConversionUtilities.h b/Engine/Plugins/Enterprise/GLTFExporter/MigratedFromGit/Source/GLTFExporter/Public/GLTFConversionUtilities.h
new file mode 100644
--- /dev/null
+++ b/Engine/Plugins/Enterprise/GLTFExporter/MigratedFromGit/Source/GLTFExporter/Public/GLTFConversionUtilities.h
@@ -0,0 +1,81 @@
+// Copyright Epic Games, Inc. All Rights Reserved.
+
+#pragma once
+
+#include "GLTFJsonEnums.h"
+#include "Engine/EngineTypes.h"
+
+inline FVector ConvertVector(const FVector& Vector)
+{
+	// UE4 uses a left-handed coordinate system, with Z up.
+	// glTF uses a right-handed coordinate system, with Y up.
+	return { Vector.X, Vector.Z, Vector.Y };
+}
+
+inline FVector ConvertPosition(const FVector& Position)
+{
+	return ConvertVector(Position);
+}
+
+inline FVector ConvertSize(const FVector& Position)
+{
+	return ConvertVector(Position);
+}
+
+inline FVector4 ConvertTangent(const FVector4& Tangent)
+{
+	// glTF stores tangent as Vec4, with W component indicating handedness of tangent basis.
+
+	return FVector4(Tangent.X, Tangent.Z, Tangent.Y, Tangent.W);
+}
+
+inline FColor ConvertColor(const FColor& Color)
+{
+	// UE4 uses ABGR.
+	// glTF uses RGBA.
+	return { Color.R, Color.G, Color.B, Color.A };
+}
+
+inline FQuat ConvertQuat(const FQuat& Quat)
+{
+	// glTF uses a right-handed coordinate system, with Y up.
+	// UE4 uses a left-handed coordinate system, with Z up.
+	// Quat = (qX, qY, qZ, qW) = (sin(angle/2) * aX, sin(angle/2) * aY, sin(angle/2) * aZ, cons(angle/2))
+	// where (aX, aY, aZ) - rotation axis, angle - rotation angle
+	// Y swapped with Z between these coordinate systems
+	// also, as handedness is changed rotation is inversed - hence negation
+	// therefore QuatUE = (-qX, -qZ, -qY, qw)
+
+	FQuat Result(-Quat.X, -Quat.Z, -Quat.Y, Quat.W);
+	// Not checking if quaternion is normalized
+	// e.g. some sources use non-unit Quats for rotation tangents
+	return Result;
+}
+
+inline FMatrix ConvertMat(const FMatrix& Matrix)
+{
+	// glTF stores matrix elements in column major order
+	// Unreal's FMatrix is row major
+
+	FMatrix Result;
+	for (int32 Row = 0; Row < 4; ++Row)
+	{
+		for (int32 Col = 0; Col < 4; ++Col)
+		{
+			Result.M[Col][Row] = Matrix.M[Row][Col];
+		}
+	}
+	return Result;
+}
+
+inline EGLTFJsonAlphaMode ConvertAlphaMode(EBlendMode Mode)
+{
+	switch (Mode)
+	{
+		case EBlendMode::BLEND_Opaque:      return EGLTFJsonAlphaMode::Opaque;
+		case EBlendMode::BLEND_Translucent: return EGLTFJsonAlphaMode::Blend;
+		case EBlendMode::BLEND_Masked:      return EGLTFJsonAlphaMode::Mask;
+	}
+
+	return EGLTFJsonAlphaMode::Opaque; // fallback
+}
