[GLTFExporter] feat: add GLTFMeshConverter for converting mesh to json data

commit 2af492d6a6c7225eb5b2fd9dc3c710eaac957647
Date: Fri, 24 Apr 2020 08:13:12 +0200

    feat: add GLTFMeshConverter for converting mesh to json data

#skipci

diff --git a/Engine/Plugins/Enterprise/GLTFExporter/MigratedFromGit/Source/GLTFExporter/Private/GLTFMeshConverter.cpp b/Engine/Plugins/Enterprise/GLTFExporter/MigratedFromGit/Source/GLTFExporter/Private/GLTFMeshConverter.cpp
new file mode 100644
--- /dev/null
+++ b/Engine/Plugins/Enterprise/GLTFExporter/MigratedFromGit/Source/GLTFExporter/Private/GLTFMeshConverter.cpp
@@ -0,0 +1,245 @@
+// Copyright Epic Games, Inc. All Rights Reserved.
+
+#include "GLTFMeshConverter.h"
+#include "GLTFConversionUtilities.h"
+
+FGLTFConvertedSection::FGLTFConvertedSection(const FString& SectionName, const FStaticMeshLODResources& LODMesh, int32 SectionIndex)
+	: Name(SectionName)
+{
+	FIndexArrayView RawIndices = LODMesh.IndexBuffer.GetArrayView();
+	const FStaticMeshSection& Section = LODMesh.Sections[SectionIndex];
+
+	const uint32 TriangleCount = Section.NumTriangles;
+	for (uint32 TriangleIndex = 0; TriangleIndex < TriangleCount; ++TriangleIndex)
+	{
+		uint32 TriangleCorner0 = RawIndices[Section.FirstIndex + ((TriangleIndex * 3) + 0)];
+		uint32 TriangleCorner1 = RawIndices[Section.FirstIndex + ((TriangleIndex * 3) + 1)];
+		uint32 TriangleCorner2 = RawIndices[Section.FirstIndex + ((TriangleIndex * 3) + 2)];
+		Indices.Add(TriangleCorner0);
+		Indices.Add(TriangleCorner1);
+		Indices.Add(TriangleCorner2);
+	}
+}
+
+FGLTFJsonIndex FGLTFConvertedSection::AppendAccessorForIndices(FGLTFJsonRoot& Root, FGLTFBufferBuilder& BufferBuilder) const
+{
+	if (Indices.Num() == 0) return INDEX_NONE;
+
+	FString AttributeName = Name + TEXT("_Indices");
+	FGLTFJsonIndex BufferViewIndex = BufferBuilder.AppendBufferView(Indices, AttributeName, EGLTFJsonBufferTarget::ElementArrayBuffer);
+
+	FGLTFJsonAccessor Accessor;
+	Accessor.Name = AttributeName;
+	Accessor.BufferView = BufferViewIndex;
+	Accessor.ComponentType = EGLTFJsonComponentType::U32;
+	Accessor.Count = Indices.Num();
+	Accessor.Type = EGLTFJsonAccessorType::Scalar;
+
+	FGLTFJsonIndex AccessorIndex = Root.Accessors.Add(Accessor);
+	return AccessorIndex;
+}
+
+FGLTFConvertedMesh::FGLTFConvertedMesh(const UStaticMesh* StaticMesh, int32 LODIndex)
+{
+	StaticMesh->GetName(Name);
+
+	const FStaticMeshLODResources& LODMesh = StaticMesh->GetLODForExport(LODIndex);
+
+	const int32 SectionCount = LODMesh.Sections.Num();
+	for (int32 SectionIndex = 0; SectionIndex < SectionCount; ++SectionIndex)
+	{
+		FString SectionName = Name + TEXT("_Section") + FString::FromInt(SectionIndex);
+		Sections.Add(FGLTFConvertedSection(SectionName, LODMesh, SectionIndex));
+	}
+
+	const int32 PositionCount = LODMesh.VertexBuffers.PositionVertexBuffer.GetNumVertices();
+	for (int32 PosIndex = 0; PosIndex < PositionCount; ++PosIndex)
+	{
+		FVector Position = LODMesh.VertexBuffers.PositionVertexBuffer.VertexPosition(PosIndex);
+		Positions.Add(ConvertPosition(Position));
+	}
+
+	const int32 VertexCount = LODMesh.VertexBuffers.StaticMeshVertexBuffer.GetNumVertices();
+	for (int32 VertIndex = 0; VertIndex < VertexCount; ++VertIndex)
+	{
+		FVector Normal = LODMesh.VertexBuffers.StaticMeshVertexBuffer.VertexTangentZ(VertIndex);
+		FVector4 Tangent = LODMesh.VertexBuffers.StaticMeshVertexBuffer.VertexTangentX(VertIndex);
+		Normals.Add(ConvertVector(Normal));
+		Tangents.Add(ConvertTangent(Tangent));
+	}
+
+	const int32 ColorCount = LODMesh.VertexBuffers.ColorVertexBuffer.GetNumVertices();
+	if (ColorCount > 0)
+	{
+		for (int32 VertIndex = 0; VertIndex < VertexCount; ++VertIndex)
+		{
+			FColor Color = LODMesh.VertexBuffers.ColorVertexBuffer.VertexColor(VertIndex);
+			Colors.Add(ConvertColor(Color));
+		}
+	}
+
+	const int32 UVCount = LODMesh.GetNumTexCoords();
+	if (UVCount >= 1)
+	{
+		for (int32 VertIndex = 0; VertIndex < VertexCount; ++VertIndex)
+		{
+			FVector2D UV0 = LODMesh.VertexBuffers.StaticMeshVertexBuffer.GetVertexUV(VertIndex, 0);
+			UV0s.Add(UV0);
+		}
+
+		if (UVCount >= 2)
+		{
+			for (int32 VertIndex = 0; VertIndex < VertexCount; ++VertIndex)
+			{
+				FVector2D UV1 = LODMesh.VertexBuffers.StaticMeshVertexBuffer.GetVertexUV(VertIndex, 1);
+				UV1s.Add(UV1);
+			}
+		}
+	}
+
+	BoundingBox = StaticMesh->GetBoundingBox();
+}
+
+FGLTFJsonIndex FGLTFConvertedMesh::AppendAccessorForPositions(FGLTFJsonRoot& Root, FGLTFBufferBuilder& BufferBuilder) const
+{
+	if (Positions.Num() == 0) return INDEX_NONE;
+
+	FString AttributeName = Name + TEXT("_Positions");
+	FGLTFJsonIndex BufferViewIndex = BufferBuilder.AppendBufferView(Positions, AttributeName);
+
+	FGLTFJsonAccessor Accessor;
+	Accessor.Name = AttributeName;
+	Accessor.BufferView = BufferViewIndex;
+	Accessor.ComponentType = EGLTFJsonComponentType::F32;
+	Accessor.Count = Positions.Num();
+	Accessor.Type = EGLTFJsonAccessorType::Vec3;
+
+	FVector Max = ConvertSize(BoundingBox.Max);
+	FVector Min = ConvertSize(BoundingBox.Min);
+	Accessor.Max[0] = Max.X;
+	Accessor.Max[1] = Max.Y;
+	Accessor.Max[2] = Max.Z;
+	Accessor.Min[0] = Min.X;
+	Accessor.Min[1] = Min.Y;
+	Accessor.Min[2] = Min.Z;
+	Accessor.MinMaxLength = 3;
+
+	FGLTFJsonIndex AccessorIndex = Root.Accessors.Add(Accessor);
+	return AccessorIndex;
+}
+
+FGLTFJsonIndex FGLTFConvertedMesh::AppendAccessorForNormals(FGLTFJsonRoot& Root, FGLTFBufferBuilder& BufferBuilder) const
+{
+	if (Normals.Num() == 0) return INDEX_NONE;
+
+	FString AttributeName = Name + TEXT("_Normals");
+	FGLTFJsonIndex BufferViewIndex = BufferBuilder.AppendBufferView(Normals, AttributeName);
+
+	FGLTFJsonAccessor Accessor;
+	Accessor.Name = AttributeName;
+	Accessor.BufferView = BufferViewIndex;
+	Accessor.ComponentType = EGLTFJsonComponentType::F32;
+	Accessor.Count = Normals.Num();
+	Accessor.Type = EGLTFJsonAccessorType::Vec3;
+	FGLTFJsonIndex AccessorIndex = Root.Accessors.Add(Accessor);
+
+	return AccessorIndex;
+}
+
+FGLTFJsonIndex FGLTFConvertedMesh::AppendAccessorForTangents(FGLTFJsonRoot& Root, FGLTFBufferBuilder& BufferBuilder) const
+{
+	if (Tangents.Num() == 0) return INDEX_NONE;
+
+	FString AttributeName = Name + TEXT("_Tangents");
+	FGLTFJsonIndex BufferViewIndex = BufferBuilder.AppendBufferView(Tangents, AttributeName);
+
+	FGLTFJsonAccessor Accessor;
+	Accessor.Name = AttributeName;
+	Accessor.BufferView = BufferViewIndex;
+	Accessor.ComponentType = EGLTFJsonComponentType::F32;
+	Accessor.Count = Tangents.Num();
+	Accessor.Type = EGLTFJsonAccessorType::Vec4;
+
+	FGLTFJsonIndex AccessorIndex = Root.Accessors.Add(Accessor);
+	return AccessorIndex;
+}
+
+FGLTFJsonIndex FGLTFConvertedMesh::AppendAccessorForColors(FGLTFJsonRoot& Root, FGLTFBufferBuilder& BufferBuilder) const
+{
+	if (Colors.Num() == 0) return INDEX_NONE;
+
+	FString AttributeName = Name + TEXT("_Colors");
+	FGLTFJsonIndex BufferViewIndex = BufferBuilder.AppendBufferView(Colors, AttributeName);
+
+	FGLTFJsonAccessor Accessor;
+	Accessor.Name = AttributeName;
+	Accessor.BufferView = BufferViewIndex;
+	Accessor.ComponentType = EGLTFJsonComponentType::U8;
+	Accessor.Count = Colors.Num();
+	Accessor.Type = EGLTFJsonAccessorType::Vec4;
+
+	FGLTFJsonIndex AccessorIndex = Root.Accessors.Add(Accessor);
+	return AccessorIndex;
+}
+
+FGLTFJsonIndex FGLTFConvertedMesh::AppendAccessorForUV0s(FGLTFJsonRoot& Root, FGLTFBufferBuilder& BufferBuilder) const
+{
+	if (UV0s.Num() == 0) return INDEX_NONE;
+
+	FString AttributeName = Name + TEXT("_UV0s");
+	FGLTFJsonIndex BufferViewIndex = BufferBuilder.AppendBufferView(UV0s, AttributeName);
+
+	FGLTFJsonAccessor Accessor;
+	Accessor.Name = AttributeName;
+	Accessor.BufferView = BufferViewIndex;
+	Accessor.ComponentType = EGLTFJsonComponentType::F32;
+	Accessor.Count = UV0s.Num();
+	Accessor.Type = EGLTFJsonAccessorType::Vec2;
+
+	FGLTFJsonIndex AccessorIndex = Root.Accessors.Add(Accessor);
+	return AccessorIndex;
+}
+
+FGLTFJsonIndex FGLTFConvertedMesh::AppendAccessorForUV1s(FGLTFJsonRoot& Root, FGLTFBufferBuilder& BufferBuilder) const
+{
+	if (UV1s.Num() == 0) return INDEX_NONE;
+
+	FString AttributeName = Name + TEXT("_UV1s");
+	FGLTFJsonIndex BufferViewIndex = BufferBuilder.AppendBufferView(UV1s, AttributeName);
+
+	FGLTFJsonAccessor Accessor;
+	Accessor.Name = AttributeName;
+	Accessor.BufferView = BufferViewIndex;
+	Accessor.ComponentType = EGLTFJsonComponentType::F32;
+	Accessor.Count = UV1s.Num();
+	Accessor.Type = EGLTFJsonAccessorType::Vec2;
+
+	FGLTFJsonIndex AccessorIndex = Root.Accessors.Add(Accessor);
+	return AccessorIndex;
+}
+
+FGLTFJsonIndex FGLTFConvertedMesh::AppendMesh(FGLTFJsonRoot& Root, FGLTFBufferBuilder& BufferBuilder) const
+{
+	FGLTFJsonMesh Mesh;
+	Mesh.Name = Name;
+
+	FGLTFJsonAttributes Attributes;
+	Attributes.Position = AppendAccessorForPositions(Root, BufferBuilder);
+	Attributes.Normal = AppendAccessorForNormals(Root, BufferBuilder);
+	Attributes.Tangent = AppendAccessorForTangents(Root, BufferBuilder);
+	Attributes.Color0 = AppendAccessorForColors(Root, BufferBuilder);
+	Attributes.TexCoord0 = AppendAccessorForUV0s(Root, BufferBuilder);
+	Attributes.TexCoord1 = AppendAccessorForUV1s(Root, BufferBuilder);
+
+	for (const FGLTFConvertedSection& Section : Sections)
+	{
+		FGLTFJsonPrimitive Primitive;
+		Primitive.Attributes = Attributes;
+
+		Primitive.Indices = Section.AppendAccessorForIndices(Root, BufferBuilder);
+		Mesh.Primitives.Add(Primitive);
+	}
+
+	FGLTFJsonIndex MeshIndex = Root.Meshes.Add(Mesh);
+	return MeshIndex;
+}
diff --git a/Engine/Plugins/Enterprise/GLTFExporter/MigratedFromGit/Source/GLTFExporter/Public/GLTFMeshConverter.h b/Engine/Plugins/Enterprise/GLTFExporter/MigratedFromGit/Source/GLTFExporter/Public/GLTFMeshConverter.h
new file mode 100644
--- /dev/null
+++ b/Engine/Plugins/Enterprise/GLTFExporter/MigratedFromGit/Source/GLTFExporter/Public/GLTFMeshConverter.h
@@ -0,0 +1,62 @@
+// Copyright Epic Games, Inc. All Rights Reserved.
+
+#pragma once
+
+#include "GLTFJsonRoot.h"
+#include "GLTFBufferBuilder.h"
+#include "Engine.h"
+
+struct GLTFEXPORTER_API FGLTFConvertedSection
+{
+	FString Name;
+
+	TArray<uint32> Indices;
+
+	FGLTFConvertedSection(const FString& SectionName, const FStaticMeshLODResources& LODMesh, int32 SectionIndex);
+
+	FGLTFJsonIndex AppendAccessorForIndices(FGLTFJsonRoot& JsonRoot, FGLTFBufferBuilder& BufferBuilder) const;
+};
+
+struct GLTFEXPORTER_API FGLTFConvertedMesh
+{
+	FString Name;
+
+	TArray<FGLTFConvertedSection> Sections;
+
+	TArray<FVector>   Positions;
+	TArray<FVector>   Normals;
+	TArray<FVector4>  Tangents;
+	TArray<FColor>    Colors;
+	TArray<FVector2D> UV0s;
+	TArray<FVector2D> UV1s;
+
+	FBox BoundingBox;
+
+	FGLTFConvertedMesh(const UStaticMesh* StaticMesh, int32 LODIndex);
+
+	FGLTFJsonIndex AppendAccessorForPositions(FGLTFJsonRoot& JsonRoot, FGLTFBufferBuilder& BufferBuilder) const;
+	FGLTFJsonIndex AppendAccessorForNormals(FGLTFJsonRoot& JsonRoot, FGLTFBufferBuilder& BufferBuilder) const;
+	FGLTFJsonIndex AppendAccessorForTangents(FGLTFJsonRoot& JsonRoot, FGLTFBufferBuilder& BufferBuilder) const;
+	FGLTFJsonIndex AppendAccessorForColors(FGLTFJsonRoot& JsonRoot, FGLTFBufferBuilder& BufferBuilder) const;
+	FGLTFJsonIndex AppendAccessorForUV0s(FGLTFJsonRoot& JsonRoot, FGLTFBufferBuilder& BufferBuilder) const;
+	FGLTFJsonIndex AppendAccessorForUV1s(FGLTFJsonRoot& JsonRoot, FGLTFBufferBuilder& BufferBuilder) const;
+
+	FGLTFJsonIndex AppendMesh(FGLTFJsonRoot& JsonRoot, FGLTFBufferBuilder& BufferBuilder) const;
+};
+
+struct GLTFEXPORTER_API FGLTFMeshConverter
+{
+	FGLTFJsonRoot& JsonRoot;
+	FGLTFBufferBuilder& BufferBuilder;
+
+	FGLTFMeshConverter(FGLTFJsonRoot& JsonRoot, FGLTFBufferBuilder& BufferBuilder)
+		: JsonRoot(JsonRoot)
+		, BufferBuilder(BufferBuilder)
+	{
+	}
+
+	FGLTFJsonIndex AppendMesh(const UStaticMesh* StaticMesh, int32 LODIndex)
+	{
+		return FGLTFConvertedMesh(StaticMesh, LODIndex).AppendMesh(JsonRoot, BufferBuilder);
+	}
+};
