[GLTFExporter] feat: export all StaticMesh vertex data to embedded gltf json

commit 46ea6f2f2ba26cff32cce6e755e34223b86244fe
Date: Fri, 24 Apr 2020 08:16:53 +0200

    feat: export all StaticMesh vertex data to embedded gltf json

#skipci

diff --git a/Engine/Plugins/Enterprise/GLTFExporter/MigratedFromGit/Source/GLTFExporter/Private/GLTFStaticMeshExporter.cpp b/Engine/Plugins/Enterprise/GLTFExporter/MigratedFromGit/Source/GLTFExporter/Private/GLTFStaticMeshExporter.cpp
--- a/Engine/Plugins/Enterprise/GLTFExporter/MigratedFromGit/Source/GLTFExporter/Private/GLTFStaticMeshExporter.cpp
+++ b/Engine/Plugins/Enterprise/GLTFExporter/MigratedFromGit/Source/GLTFExporter/Private/GLTFStaticMeshExporter.cpp
@@ -2,6 +2,8 @@
 
 #include "GLTFStaticMeshExporter.h"
 #include "GLTFJsonRoot.h"
+#include "GLTFBufferBuilder.h"
+#include "GLTFMeshConverter.h"
 #include "Engine/StaticMesh.h"
 
 UGLTFStaticMeshExporter::UGLTFStaticMeshExporter(const FObjectInitializer& ObjectInitializer)
@@ -21,23 +23,10 @@ bool UGLTFStaticMeshExporter::ExportBinary(UObject* Object, const TCHAR* Type, F
 	}
 
 	FGLTFJsonRoot Root;
+	FGLTFBufferBuilder BufferBuilder(Root);
+	FGLTFMeshConverter MeshConverter(Root, BufferBuilder);
 
-	FGLTFJsonBuffer Buffer;
-	FGLTFJsonIndex BufferIndex = Root.Buffers.Add(Buffer);
-
-	FGLTFJsonBufferView BufferView;
-	BufferView.Buffer = BufferIndex;
-	FGLTFJsonIndex BufferViewIndex = Root.BufferViews.Add(BufferView);
-
-	FGLTFJsonAccessor Accessor;
-	Accessor.BufferView = BufferViewIndex;
-	FGLTFJsonIndex AccessorIndex = Root.Accessors.Add(Accessor);
-
-	FGLTFJsonMesh Mesh;
-	FGLTFJsonPrimitive Primitive;
-	Primitive.Attributes.Position = AccessorIndex;
-	Mesh.Primitives.Add(Primitive);
-	FGLTFJsonIndex MeshIndex = Root.Meshes.Add(Mesh);
+	FGLTFJsonIndex MeshIndex = MeshConverter.AppendMesh(StaticMesh, 0);
 
 	FGLTFJsonNode Node;
 	Node.Mesh = MeshIndex;
@@ -49,6 +38,7 @@ bool UGLTFStaticMeshExporter::ExportBinary(UObject* Object, const TCHAR* Type, F
 
 	Root.DefaultScene = SceneIndex;
 
+	BufferBuilder.Close();
 	Root.Serialize(&Archive, true);
 	return true;
 }
