[GLTFExporter] feat: optimize export by preinitializing vertex data array sizes

commit 6ae561a582385650c99c943333d3e78c23705164
Date: Fri, 24 Apr 2020 10:11:36 +0200

    feat: optimize export by preinitializing vertex data array sizes

#skipci

diff --git a/Engine/Plugins/Enterprise/GLTFExporter/MigratedFromGit/Source/GLTFExporter/Private/GLTFMeshConverter.cpp b/Engine/Plugins/Enterprise/GLTFExporter/MigratedFromGit/Source/GLTFExporter/Private/GLTFMeshConverter.cpp
--- a/Engine/Plugins/Enterprise/GLTFExporter/MigratedFromGit/Source/GLTFExporter/Private/GLTFMeshConverter.cpp
+++ b/Engine/Plugins/Enterprise/GLTFExporter/MigratedFromGit/Source/GLTFExporter/Private/GLTFMeshConverter.cpp
@@ -9,15 +9,12 @@ FGLTFConvertedSection::FGLTFConvertedSection(const FString& SectionName, const F
 	FIndexArrayView RawIndices = LODMesh.IndexBuffer.GetArrayView();
 	const FStaticMeshSection& Section = LODMesh.Sections[SectionIndex];
 
-	const uint32 TriangleCount = Section.NumTriangles;
-	for (uint32 TriangleIndex = 0; TriangleIndex < TriangleCount; ++TriangleIndex)
+	const uint32 IndexCount = Section.NumTriangles * 3;
+	Indices.AddUninitialized(IndexCount);
+
+	for (uint32 Index = 0; Index < IndexCount; ++Index)
 	{
-		uint32 TriangleCorner0 = RawIndices[Section.FirstIndex + ((TriangleIndex * 3) + 0)];
-		uint32 TriangleCorner1 = RawIndices[Section.FirstIndex + ((TriangleIndex * 3) + 1)];
-		uint32 TriangleCorner2 = RawIndices[Section.FirstIndex + ((TriangleIndex * 3) + 2)];
-		Indices.Add(TriangleCorner0);
-		Indices.Add(TriangleCorner1);
-		Indices.Add(TriangleCorner2);
+		Indices[Index] = RawIndices[Section.FirstIndex + Index];
 	}
 }
 
@@ -53,46 +50,48 @@ FGLTFConvertedMesh::FGLTFConvertedMesh(const UStaticMesh* StaticMesh, int32 LODI
 	}
 
 	const int32 PositionCount = LODMesh.VertexBuffers.PositionVertexBuffer.GetNumVertices();
-	for (int32 PosIndex = 0; PosIndex < PositionCount; ++PosIndex)
-	{
-		FVector Position = LODMesh.VertexBuffers.PositionVertexBuffer.VertexPosition(PosIndex);
-		Positions.Add(ConvertPosition(Position));
-	}
+	Positions.AddUninitialized(PositionCount);
 
-	const int32 VertexCount = LODMesh.VertexBuffers.StaticMeshVertexBuffer.GetNumVertices();
-	for (int32 VertIndex = 0; VertIndex < VertexCount; ++VertIndex)
+	for (int32 PosIndex = 0; PosIndex < PositionCount; ++PosIndex)
 	{
-		FVector Normal = LODMesh.VertexBuffers.StaticMeshVertexBuffer.VertexTangentZ(VertIndex);
-		FVector4 Tangent = LODMesh.VertexBuffers.StaticMeshVertexBuffer.VertexTangentX(VertIndex);
-		Normals.Add(ConvertVector(Normal));
-		Tangents.Add(ConvertTangent(Tangent));
+		Positions[PosIndex] = ConvertPosition(LODMesh.VertexBuffers.PositionVertexBuffer.VertexPosition(PosIndex));
 	}
 
 	const int32 ColorCount = LODMesh.VertexBuffers.ColorVertexBuffer.GetNumVertices();
 	if (ColorCount > 0)
 	{
-		for (int32 VertIndex = 0; VertIndex < VertexCount; ++VertIndex)
+		Colors.AddUninitialized(ColorCount);
+		for (int32 ColorIndex = 0; ColorIndex < ColorCount; ++ColorIndex)
 		{
-			FColor Color = LODMesh.VertexBuffers.ColorVertexBuffer.VertexColor(VertIndex);
-			Colors.Add(ConvertColor(Color));
+			Colors[ColorIndex] = ConvertColor(LODMesh.VertexBuffers.ColorVertexBuffer.VertexColor(ColorIndex));
 		}
 	}
 
+	const int32 VertexCount = LODMesh.VertexBuffers.StaticMeshVertexBuffer.GetNumVertices();
+	Normals.AddUninitialized(VertexCount);
+	Tangents.AddUninitialized(VertexCount);
+
+	for (int32 VertIndex = 0; VertIndex < VertexCount; ++VertIndex)
+	{
+		Normals[VertIndex] = ConvertVector(LODMesh.VertexBuffers.StaticMeshVertexBuffer.VertexTangentZ(VertIndex));
+		Tangents[VertIndex] = ConvertTangent(LODMesh.VertexBuffers.StaticMeshVertexBuffer.VertexTangentX(VertIndex));
+	}
+
 	const int32 UVCount = LODMesh.GetNumTexCoords();
 	if (UVCount >= 1)
 	{
+		UV0s.AddUninitialized(VertexCount);
 		for (int32 VertIndex = 0; VertIndex < VertexCount; ++VertIndex)
 		{
-			FVector2D UV0 = LODMesh.VertexBuffers.StaticMeshVertexBuffer.GetVertexUV(VertIndex, 0);
-			UV0s.Add(UV0);
+			UV0s[VertIndex] = LODMesh.VertexBuffers.StaticMeshVertexBuffer.GetVertexUV(VertIndex, 0);
 		}
 
 		if (UVCount >= 2)
 		{
+			UV1s.AddUninitialized(VertexCount);
 			for (int32 VertIndex = 0; VertIndex < VertexCount; ++VertIndex)
 			{
-				FVector2D UV1 = LODMesh.VertexBuffers.StaticMeshVertexBuffer.GetVertexUV(VertIndex, 1);
-				UV1s.Add(UV1);
+				UV1s[VertIndex] = LODMesh.VertexBuffers.StaticMeshVertexBuffer.GetVertexUV(VertIndex, 1);
 			}
 		}
 	}
@@ -146,36 +145,36 @@ FGLTFJsonIndex FGLTFConvertedMesh::AppendAccessorForNormals(FGLTFJsonRoot& Root,
 	return AccessorIndex;
 }
 
-FGLTFJsonIndex FGLTFConvertedMesh::AppendAccessorForTangents(FGLTFJsonRoot& Root, FGLTFBufferBuilder& BufferBuilder) const
+FGLTFJsonIndex FGLTFConvertedMesh::AppendAccessorForColors(FGLTFJsonRoot& Root, FGLTFBufferBuilder& BufferBuilder) const
 {
-	if (Tangents.Num() == 0) return INDEX_NONE;
+	if (Colors.Num() == 0) return INDEX_NONE;
 
-	FString AttributeName = Name + TEXT("_Tangents");
-	FGLTFJsonIndex BufferViewIndex = BufferBuilder.AppendBufferView(Tangents, AttributeName);
+	FString AttributeName = Name + TEXT("_Colors");
+	FGLTFJsonIndex BufferViewIndex = BufferBuilder.AppendBufferView(Colors, AttributeName);
 
 	FGLTFJsonAccessor Accessor;
 	Accessor.Name = AttributeName;
 	Accessor.BufferView = BufferViewIndex;
-	Accessor.ComponentType = EGLTFJsonComponentType::F32;
-	Accessor.Count = Tangents.Num();
+	Accessor.ComponentType = EGLTFJsonComponentType::U8;
+	Accessor.Count = Colors.Num();
 	Accessor.Type = EGLTFJsonAccessorType::Vec4;
 
 	FGLTFJsonIndex AccessorIndex = Root.Accessors.Add(Accessor);
 	return AccessorIndex;
 }
 
-FGLTFJsonIndex FGLTFConvertedMesh::AppendAccessorForColors(FGLTFJsonRoot& Root, FGLTFBufferBuilder& BufferBuilder) const
+FGLTFJsonIndex FGLTFConvertedMesh::AppendAccessorForTangents(FGLTFJsonRoot& Root, FGLTFBufferBuilder& BufferBuilder) const
 {
-	if (Colors.Num() == 0) return INDEX_NONE;
+	if (Tangents.Num() == 0) return INDEX_NONE;
 
-	FString AttributeName = Name + TEXT("_Colors");
-	FGLTFJsonIndex BufferViewIndex = BufferBuilder.AppendBufferView(Colors, AttributeName);
+	FString AttributeName = Name + TEXT("_Tangents");
+	FGLTFJsonIndex BufferViewIndex = BufferBuilder.AppendBufferView(Tangents, AttributeName);
 
 	FGLTFJsonAccessor Accessor;
 	Accessor.Name = AttributeName;
 	Accessor.BufferView = BufferViewIndex;
-	Accessor.ComponentType = EGLTFJsonComponentType::U8;
-	Accessor.Count = Colors.Num();
+	Accessor.ComponentType = EGLTFJsonComponentType::F32;
+	Accessor.Count = Tangents.Num();
 	Accessor.Type = EGLTFJsonAccessorType::Vec4;
 
 	FGLTFJsonIndex AccessorIndex = Root.Accessors.Add(Accessor);
@@ -225,9 +224,9 @@ FGLTFJsonIndex FGLTFConvertedMesh::AppendMesh(FGLTFJsonRoot& Root, FGLTFBufferBu
 
 	FGLTFJsonAttributes Attributes;
 	Attributes.Position = AppendAccessorForPositions(Root, BufferBuilder);
+	Attributes.Color0 = AppendAccessorForColors(Root, BufferBuilder);
 	Attributes.Normal = AppendAccessorForNormals(Root, BufferBuilder);
 	Attributes.Tangent = AppendAccessorForTangents(Root, BufferBuilder);
-	Attributes.Color0 = AppendAccessorForColors(Root, BufferBuilder);
 	Attributes.TexCoord0 = AppendAccessorForUV0s(Root, BufferBuilder);
 	Attributes.TexCoord1 = AppendAccessorForUV1s(Root, BufferBuilder);
 
