[GLTFExporter] refactor: add separate typedef for each kind of json index to improve readability

commit 9551a065621b7593bf7eb6b70c94b21ec6b9632f
Date: Fri, 24 Apr 2020 10:48:01 +0200

    refactor: add separate typedef for each kind of json index to improve readability

#skipci

diff --git a/Engine/Plugins/Enterprise/GLTFExporter/MigratedFromGit/Source/GLTFExporter/Private/GLTFBufferBuilder.cpp b/Engine/Plugins/Enterprise/GLTFExporter/MigratedFromGit/Source/GLTFExporter/Private/GLTFBufferBuilder.cpp
--- a/Engine/Plugins/Enterprise/GLTFExporter/MigratedFromGit/Source/GLTFExporter/Private/GLTFBufferBuilder.cpp
+++ b/Engine/Plugins/Enterprise/GLTFExporter/MigratedFromGit/Source/GLTFExporter/Private/GLTFBufferBuilder.cpp
@@ -10,7 +10,7 @@ inline FGLTFBufferBuilder::FGLTFBufferBuilder(FGLTFJsonRoot& JsonRoot, const FSt
 	JsonRoot.Buffers[BufferIndex].Name = Name;
 }
 
-inline FGLTFJsonIndex FGLTFBufferBuilder::AppendBufferView(const void* RawData, uint64 ByteLength, const FString& Name, EGLTFJsonBufferTarget BufferTarget)
+inline FGLTFJsonBufferViewIndex FGLTFBufferBuilder::AppendBufferView(const void* RawData, uint64 ByteLength, const FString& Name, EGLTFJsonBufferTarget BufferTarget)
 {
 	FGLTFJsonBufferView BufferView;
 	BufferView.Name = Name;
@@ -21,8 +21,7 @@ inline FGLTFJsonIndex FGLTFBufferBuilder::AppendBufferView(const void* RawData,
 
 	Data.Append(static_cast<const uint8*>(RawData), ByteLength);
 
-	FGLTFJsonIndex BufferViewIndex = JsonRoot.BufferViews.Add(BufferView);
-	return BufferViewIndex;
+	return JsonRoot.BufferViews.Add(BufferView);
 }
 
 void FGLTFBufferBuilder::Close()
diff --git a/Engine/Plugins/Enterprise/GLTFExporter/MigratedFromGit/Source/GLTFExporter/Private/GLTFMeshConverter.cpp b/Engine/Plugins/Enterprise/GLTFExporter/MigratedFromGit/Source/GLTFExporter/Private/GLTFMeshConverter.cpp
--- a/Engine/Plugins/Enterprise/GLTFExporter/MigratedFromGit/Source/GLTFExporter/Private/GLTFMeshConverter.cpp
+++ b/Engine/Plugins/Enterprise/GLTFExporter/MigratedFromGit/Source/GLTFExporter/Private/GLTFMeshConverter.cpp
@@ -18,12 +18,12 @@ FGLTFConvertedSection::FGLTFConvertedSection(const FString& SectionName, const F
 	}
 }
 
-FGLTFJsonIndex FGLTFConvertedSection::AppendAccessorForIndices(FGLTFJsonRoot& Root, FGLTFBufferBuilder& BufferBuilder) const
+FGLTFJsonAccessorIndex FGLTFConvertedSection::AppendAccessorForIndices(FGLTFJsonRoot& Root, FGLTFBufferBuilder& BufferBuilder) const
 {
 	if (Indices.Num() == 0) return INDEX_NONE;
 
 	FString AttributeName = Name + TEXT("_Indices");
-	FGLTFJsonIndex BufferViewIndex = BufferBuilder.AppendBufferView(Indices, AttributeName, EGLTFJsonBufferTarget::ElementArrayBuffer);
+	FGLTFJsonBufferViewIndex BufferViewIndex = BufferBuilder.AppendBufferView(Indices, AttributeName, EGLTFJsonBufferTarget::ElementArrayBuffer);
 
 	FGLTFJsonAccessor Accessor;
 	Accessor.Name = AttributeName;
@@ -32,8 +32,7 @@ FGLTFJsonIndex FGLTFConvertedSection::AppendAccessorForIndices(FGLTFJsonRoot& Ro
 	Accessor.Count = Indices.Num();
 	Accessor.Type = EGLTFJsonAccessorType::Scalar;
 
-	FGLTFJsonIndex AccessorIndex = Root.Accessors.Add(Accessor);
-	return AccessorIndex;
+	return Root.Accessors.Add(Accessor);
 }
 
 FGLTFConvertedMesh::FGLTFConvertedMesh(const UStaticMesh* StaticMesh, int32 LODIndex)
@@ -99,12 +98,12 @@ FGLTFConvertedMesh::FGLTFConvertedMesh(const UStaticMesh* StaticMesh, int32 LODI
 	BoundingBox = StaticMesh->GetBoundingBox();
 }
 
-FGLTFJsonIndex FGLTFConvertedMesh::AppendAccessorForPositions(FGLTFJsonRoot& Root, FGLTFBufferBuilder& BufferBuilder) const
+FGLTFJsonAccessorIndex FGLTFConvertedMesh::AppendAccessorForPositions(FGLTFJsonRoot& Root, FGLTFBufferBuilder& BufferBuilder) const
 {
 	if (Positions.Num() == 0) return INDEX_NONE;
 
 	FString AttributeName = Name + TEXT("_Positions");
-	FGLTFJsonIndex BufferViewIndex = BufferBuilder.AppendBufferView(Positions, AttributeName);
+	FGLTFJsonBufferViewIndex BufferViewIndex = BufferBuilder.AppendBufferView(Positions, AttributeName);
 
 	FGLTFJsonAccessor Accessor;
 	Accessor.Name = AttributeName;
@@ -123,16 +122,15 @@ FGLTFJsonIndex FGLTFConvertedMesh::AppendAccessorForPositions(FGLTFJsonRoot& Roo
 	Accessor.Min[2] = Min.Z;
 	Accessor.MinMaxLength = 3;
 
-	FGLTFJsonIndex AccessorIndex = Root.Accessors.Add(Accessor);
-	return AccessorIndex;
+	return Root.Accessors.Add(Accessor);
 }
 
-FGLTFJsonIndex FGLTFConvertedMesh::AppendAccessorForNormals(FGLTFJsonRoot& Root, FGLTFBufferBuilder& BufferBuilder) const
+FGLTFJsonAccessorIndex FGLTFConvertedMesh::AppendAccessorForNormals(FGLTFJsonRoot& Root, FGLTFBufferBuilder& BufferBuilder) const
 {
 	if (Normals.Num() == 0) return INDEX_NONE;
 
 	FString AttributeName = Name + TEXT("_Normals");
-	FGLTFJsonIndex BufferViewIndex = BufferBuilder.AppendBufferView(Normals, AttributeName);
+	FGLTFJsonBufferViewIndex BufferViewIndex = BufferBuilder.AppendBufferView(Normals, AttributeName);
 
 	FGLTFJsonAccessor Accessor;
 	Accessor.Name = AttributeName;
@@ -140,17 +138,16 @@ FGLTFJsonIndex FGLTFConvertedMesh::AppendAccessorForNormals(FGLTFJsonRoot& Root,
 	Accessor.ComponentType = EGLTFJsonComponentType::F32;
 	Accessor.Count = Normals.Num();
 	Accessor.Type = EGLTFJsonAccessorType::Vec3;
-	FGLTFJsonIndex AccessorIndex = Root.Accessors.Add(Accessor);
 
-	return AccessorIndex;
+	return Root.Accessors.Add(Accessor);
 }
 
-FGLTFJsonIndex FGLTFConvertedMesh::AppendAccessorForColors(FGLTFJsonRoot& Root, FGLTFBufferBuilder& BufferBuilder) const
+FGLTFJsonAccessorIndex FGLTFConvertedMesh::AppendAccessorForColors(FGLTFJsonRoot& Root, FGLTFBufferBuilder& BufferBuilder) const
 {
 	if (Colors.Num() == 0) return INDEX_NONE;
 
 	FString AttributeName = Name + TEXT("_Colors");
-	FGLTFJsonIndex BufferViewIndex = BufferBuilder.AppendBufferView(Colors, AttributeName);
+	FGLTFJsonBufferViewIndex BufferViewIndex = BufferBuilder.AppendBufferView(Colors, AttributeName);
 
 	FGLTFJsonAccessor Accessor;
 	Accessor.Name = AttributeName;
@@ -159,16 +156,15 @@ FGLTFJsonIndex FGLTFConvertedMesh::AppendAccessorForColors(FGLTFJsonRoot& Root,
 	Accessor.Count = Colors.Num();
 	Accessor.Type = EGLTFJsonAccessorType::Vec4;
 
-	FGLTFJsonIndex AccessorIndex = Root.Accessors.Add(Accessor);
-	return AccessorIndex;
+	return Root.Accessors.Add(Accessor);
 }
 
-FGLTFJsonIndex FGLTFConvertedMesh::AppendAccessorForTangents(FGLTFJsonRoot& Root, FGLTFBufferBuilder& BufferBuilder) const
+FGLTFJsonAccessorIndex FGLTFConvertedMesh::AppendAccessorForTangents(FGLTFJsonRoot& Root, FGLTFBufferBuilder& BufferBuilder) const
 {
 	if (Tangents.Num() == 0) return INDEX_NONE;
 
 	FString AttributeName = Name + TEXT("_Tangents");
-	FGLTFJsonIndex BufferViewIndex = BufferBuilder.AppendBufferView(Tangents, AttributeName);
+	FGLTFJsonBufferViewIndex BufferViewIndex = BufferBuilder.AppendBufferView(Tangents, AttributeName);
 
 	FGLTFJsonAccessor Accessor;
 	Accessor.Name = AttributeName;
@@ -177,16 +173,15 @@ FGLTFJsonIndex FGLTFConvertedMesh::AppendAccessorForTangents(FGLTFJsonRoot& Root
 	Accessor.Count = Tangents.Num();
 	Accessor.Type = EGLTFJsonAccessorType::Vec4;
 
-	FGLTFJsonIndex AccessorIndex = Root.Accessors.Add(Accessor);
-	return AccessorIndex;
+	return Root.Accessors.Add(Accessor);
 }
 
-FGLTFJsonIndex FGLTFConvertedMesh::AppendAccessorForUV0s(FGLTFJsonRoot& Root, FGLTFBufferBuilder& BufferBuilder) const
+FGLTFJsonAccessorIndex FGLTFConvertedMesh::AppendAccessorForUV0s(FGLTFJsonRoot& Root, FGLTFBufferBuilder& BufferBuilder) const
 {
 	if (UV0s.Num() == 0) return INDEX_NONE;
 
 	FString AttributeName = Name + TEXT("_UV0s");
-	FGLTFJsonIndex BufferViewIndex = BufferBuilder.AppendBufferView(UV0s, AttributeName);
+	FGLTFJsonBufferViewIndex BufferViewIndex = BufferBuilder.AppendBufferView(UV0s, AttributeName);
 
 	FGLTFJsonAccessor Accessor;
 	Accessor.Name = AttributeName;
@@ -195,16 +190,15 @@ FGLTFJsonIndex FGLTFConvertedMesh::AppendAccessorForUV0s(FGLTFJsonRoot& Root, FG
 	Accessor.Count = UV0s.Num();
 	Accessor.Type = EGLTFJsonAccessorType::Vec2;
 
-	FGLTFJsonIndex AccessorIndex = Root.Accessors.Add(Accessor);
-	return AccessorIndex;
+	return Root.Accessors.Add(Accessor);
 }
 
-FGLTFJsonIndex FGLTFConvertedMesh::AppendAccessorForUV1s(FGLTFJsonRoot& Root, FGLTFBufferBuilder& BufferBuilder) const
+FGLTFJsonAccessorIndex FGLTFConvertedMesh::AppendAccessorForUV1s(FGLTFJsonRoot& Root, FGLTFBufferBuilder& BufferBuilder) const
 {
 	if (UV1s.Num() == 0) return INDEX_NONE;
 
 	FString AttributeName = Name + TEXT("_UV1s");
-	FGLTFJsonIndex BufferViewIndex = BufferBuilder.AppendBufferView(UV1s, AttributeName);
+	FGLTFJsonBufferViewIndex BufferViewIndex = BufferBuilder.AppendBufferView(UV1s, AttributeName);
 
 	FGLTFJsonAccessor Accessor;
 	Accessor.Name = AttributeName;
@@ -213,11 +207,10 @@ FGLTFJsonIndex FGLTFConvertedMesh::AppendAccessorForUV1s(FGLTFJsonRoot& Root, FG
 	Accessor.Count = UV1s.Num();
 	Accessor.Type = EGLTFJsonAccessorType::Vec2;
 
-	FGLTFJsonIndex AccessorIndex = Root.Accessors.Add(Accessor);
-	return AccessorIndex;
+	return Root.Accessors.Add(Accessor);
 }
 
-FGLTFJsonIndex FGLTFConvertedMesh::AppendMesh(FGLTFJsonRoot& Root, FGLTFBufferBuilder& BufferBuilder) const
+FGLTFJsonMeshIndex FGLTFConvertedMesh::AppendMesh(FGLTFJsonRoot& Root, FGLTFBufferBuilder& BufferBuilder) const
 {
 	FGLTFJsonMesh Mesh;
 	Mesh.Name = Name;
@@ -239,6 +232,5 @@ FGLTFJsonIndex FGLTFConvertedMesh::AppendMesh(FGLTFJsonRoot& Root, FGLTFBufferBu
 		Mesh.Primitives.Add(Primitive);
 	}
 
-	FGLTFJsonIndex MeshIndex = Root.Meshes.Add(Mesh);
-	return MeshIndex;
+	return Root.Meshes.Add(Mesh);
 }
diff --git a/Engine/Plugins/Enterprise/GLTFExporter/MigratedFromGit/Source/GLTFExporter/Private/GLTFStaticMeshExporter.cpp b/Engine/Plugins/Enterprise/GLTFExporter/MigratedFromGit/Source/GLTFExporter/Private/GLTFStaticMeshExporter.cpp
--- a/Engine/Plugins/Enterprise/GLTFExporter/MigratedFromGit/Source/GLTFExporter/Private/GLTFStaticMeshExporter.cpp
+++ b/Engine/Plugins/Enterprise/GLTFExporter/MigratedFromGit/Source/GLTFExporter/Private/GLTFStaticMeshExporter.cpp
@@ -26,15 +26,15 @@ bool UGLTFStaticMeshExporter::ExportBinary(UObject* Object, const TCHAR* Type, F
 	FGLTFBufferBuilder BufferBuilder(Root);
 	FGLTFMeshConverter MeshConverter(Root, BufferBuilder);
 
-	FGLTFJsonIndex MeshIndex = MeshConverter.AppendMesh(StaticMesh, 0);
+	FGLTFJsonMeshIndex MeshIndex = MeshConverter.AppendMesh(StaticMesh, 0);
 
 	FGLTFJsonNode Node;
 	Node.Mesh = MeshIndex;
-	FGLTFJsonIndex NodeIndex = Root.Nodes.Add(Node);
+	FGLTFJsonNodeIndex NodeIndex = Root.Nodes.Add(Node);
 
 	FGLTFJsonScene Scene;
 	Scene.Nodes.Add(NodeIndex);
-	FGLTFJsonIndex SceneIndex = Root.Scenes.Add(Scene);
+	FGLTFJsonSceneIndex SceneIndex = Root.Scenes.Add(Scene);
 
 	Root.DefaultScene = SceneIndex;
 
diff --git a/Engine/Plugins/Enterprise/GLTFExporter/MigratedFromGit/Source/GLTFExporter/Public/GLTFBufferBuilder.h b/Engine/Plugins/Enterprise/GLTFExporter/MigratedFromGit/Source/GLTFExporter/Public/GLTFBufferBuilder.h
--- a/Engine/Plugins/Enterprise/GLTFExporter/MigratedFromGit/Source/GLTFExporter/Public/GLTFBufferBuilder.h
+++ b/Engine/Plugins/Enterprise/GLTFExporter/MigratedFromGit/Source/GLTFExporter/Public/GLTFBufferBuilder.h
@@ -8,19 +8,19 @@
 struct GLTFEXPORTER_API FGLTFBufferBuilder
 {
 	FGLTFJsonRoot& JsonRoot;
-	const FGLTFJsonIndex BufferIndex;
+	const FGLTFJsonBufferIndex BufferIndex;
 
 	TArray<uint8> Data;
 
 	FGLTFBufferBuilder(FGLTFJsonRoot& JsonRoot, const FString& Name = TEXT(""));
 
 	template <class ElementType>
-	FGLTFJsonIndex AppendBufferView(const TArray<ElementType>& Array, const FString& Name = TEXT(""), EGLTFJsonBufferTarget BufferTarget = EGLTFJsonBufferTarget::ArrayBuffer)
+	FGLTFJsonBufferViewIndex AppendBufferView(const TArray<ElementType>& Array, const FString& Name = TEXT(""), EGLTFJsonBufferTarget BufferTarget = EGLTFJsonBufferTarget::ArrayBuffer)
 	{
 		return AppendBufferView(Array.GetData(), Array.Num() * sizeof(ElementType), Name, BufferTarget);
 	}
 
-	FGLTFJsonIndex AppendBufferView(const void* RawData, uint64 ByteLength, const FString& Name = TEXT(""), EGLTFJsonBufferTarget BufferTarget = EGLTFJsonBufferTarget::ArrayBuffer);
+	FGLTFJsonBufferViewIndex AppendBufferView(const void* RawData, uint64 ByteLength, const FString& Name = TEXT(""), EGLTFJsonBufferTarget BufferTarget = EGLTFJsonBufferTarget::ArrayBuffer);
 
 	void Close();
 };
diff --git a/Engine/Plugins/Enterprise/GLTFExporter/MigratedFromGit/Source/GLTFExporter/Public/GLTFJsonAccessor.h b/Engine/Plugins/Enterprise/GLTFExporter/MigratedFromGit/Source/GLTFExporter/Public/GLTFJsonAccessor.h
--- a/Engine/Plugins/Enterprise/GLTFExporter/MigratedFromGit/Source/GLTFExporter/Public/GLTFJsonAccessor.h
+++ b/Engine/Plugins/Enterprise/GLTFExporter/MigratedFromGit/Source/GLTFExporter/Public/GLTFJsonAccessor.h
@@ -3,18 +3,19 @@
 #pragma once
 
 #include "GLTFJsonEnums.h"
-#include "GLTFJsonObject.h"
+#include "GLTFJsonIndex.h"
 #include "GLTFJsonUtilities.h"
+#include "Serialization/JsonSerializer.h"
 
-struct GLTFEXPORTER_API FGLTFJsonAccessor : FGLTFJsonObject
+struct GLTFEXPORTER_API FGLTFJsonAccessor
 {
 	FString Name;
 
-	FGLTFJsonIndex         BufferView;
-	int32                  Count;
-	EGLTFJsonAccessorType  Type;
-	EGLTFJsonComponentType ComponentType;
-	bool                   Normalized;
+	FGLTFJsonBufferViewIndex BufferView;
+	int32                    Count;
+	EGLTFJsonAccessorType    Type;
+	EGLTFJsonComponentType   ComponentType;
+	bool                     Normalized;
 
 	int32 MinMaxLength;
 	float Min[16];
diff --git a/Engine/Plugins/Enterprise/GLTFExporter/MigratedFromGit/Source/GLTFExporter/Public/GLTFJsonBuffer.h b/Engine/Plugins/Enterprise/GLTFExporter/MigratedFromGit/Source/GLTFExporter/Public/GLTFJsonBuffer.h
--- a/Engine/Plugins/Enterprise/GLTFExporter/MigratedFromGit/Source/GLTFExporter/Public/GLTFJsonBuffer.h
+++ b/Engine/Plugins/Enterprise/GLTFExporter/MigratedFromGit/Source/GLTFExporter/Public/GLTFJsonBuffer.h
@@ -2,9 +2,10 @@
 
 #pragma once
 
-#include "GLTFJsonObject.h"
+#include "GLTFJsonIndex.h"
+#include "Serialization/JsonSerializer.h"
 
-struct GLTFEXPORTER_API FGLTFJsonBuffer : FGLTFJsonObject
+struct GLTFEXPORTER_API FGLTFJsonBuffer
 {
 	FString Name;
 
diff --git a/Engine/Plugins/Enterprise/GLTFExporter/MigratedFromGit/Source/GLTFExporter/Public/GLTFJsonBufferView.h b/Engine/Plugins/Enterprise/GLTFExporter/MigratedFromGit/Source/GLTFExporter/Public/GLTFJsonBufferView.h
--- a/Engine/Plugins/Enterprise/GLTFExporter/MigratedFromGit/Source/GLTFExporter/Public/GLTFJsonBufferView.h
+++ b/Engine/Plugins/Enterprise/GLTFExporter/MigratedFromGit/Source/GLTFExporter/Public/GLTFJsonBufferView.h
@@ -3,14 +3,15 @@
 #pragma once
 
 #include "GLTFJsonEnums.h"
-#include "GLTFJsonObject.h"
+#include "GLTFJsonIndex.h"
 #include "GLTFJsonUtilities.h"
+#include "Serialization/JsonSerializer.h"
 
-struct GLTFEXPORTER_API FGLTFJsonBufferView : FGLTFJsonObject
+struct GLTFEXPORTER_API FGLTFJsonBufferView
 {
 	FString Name;
 
-	FGLTFJsonIndex Buffer;
+	FGLTFJsonBufferIndex Buffer;
 
 	int64 ByteLength;
 	int64 ByteOffset;
diff --git a/Engine/Plugins/Enterprise/GLTFExporter/MigratedFromGit/Source/GLTFExporter/Public/GLTFJsonIndex.h b/Engine/Plugins/Enterprise/GLTFExporter/MigratedFromGit/Source/GLTFExporter/Public/GLTFJsonIndex.h
new file mode 100644
--- /dev/null
+++ b/Engine/Plugins/Enterprise/GLTFExporter/MigratedFromGit/Source/GLTFExporter/Public/GLTFJsonIndex.h
@@ -0,0 +1,15 @@
+// Copyright Epic Games, Inc. All Rights Reserved.
+
+#pragma once
+
+#include "CoreMinimal.h"
+
+typedef int32 FGLTFJsonAccessorIndex;
+typedef int32 FGLTFJsonBufferIndex;
+typedef int32 FGLTFJsonBufferViewIndex;
+typedef int32 FGLTFJsonCameraIndex;
+typedef int32 FGLTFJsonMaterialIndex;
+typedef int32 FGLTFJsonMeshIndex;
+typedef int32 FGLTFJsonNodeIndex;
+typedef int32 FGLTFJsonSceneIndex;
+typedef int32 FGLTFJsonSkinIndex;
diff --git a/Engine/Plugins/Enterprise/GLTFExporter/MigratedFromGit/Source/GLTFExporter/Public/GLTFJsonMesh.h b/Engine/Plugins/Enterprise/GLTFExporter/MigratedFromGit/Source/GLTFExporter/Public/GLTFJsonMesh.h
--- a/Engine/Plugins/Enterprise/GLTFExporter/MigratedFromGit/Source/GLTFExporter/Public/GLTFJsonMesh.h
+++ b/Engine/Plugins/Enterprise/GLTFExporter/MigratedFromGit/Source/GLTFExporter/Public/GLTFJsonMesh.h
@@ -3,19 +3,20 @@
 #pragma once
 
 #include "GLTFJsonEnums.h"
-#include "GLTFJsonObject.h"
+#include "GLTFJsonIndex.h"
+#include "Serialization/JsonSerializer.h"
 
-struct GLTFEXPORTER_API FGLTFJsonAttributes : FGLTFJsonObject
+struct GLTFEXPORTER_API FGLTFJsonAttributes
 {
-	FGLTFJsonIndex Position; // always required
-	FGLTFJsonIndex Normal;
-	FGLTFJsonIndex Tangent;
-	FGLTFJsonIndex TexCoord0;
-	FGLTFJsonIndex TexCoord1;
-	FGLTFJsonIndex Color0;
+	FGLTFJsonAccessorIndex Position; // always required
+	FGLTFJsonAccessorIndex Normal;
+	FGLTFJsonAccessorIndex Tangent;
+	FGLTFJsonAccessorIndex TexCoord0;
+	FGLTFJsonAccessorIndex TexCoord1;
+	FGLTFJsonAccessorIndex Color0;
 	// skeletal mesh attributes
-	FGLTFJsonIndex Joints0;
-	FGLTFJsonIndex Weights0;
+	FGLTFJsonAccessorIndex Joints0;
+	FGLTFJsonAccessorIndex Weights0;
 
 	FGLTFJsonAttributes()
 		: Position(INDEX_NONE)
@@ -47,10 +48,10 @@ struct GLTFEXPORTER_API FGLTFJsonAttributes : FGLTFJsonObject
 	}
 };
 
-struct GLTFEXPORTER_API FGLTFJsonPrimitive : FGLTFJsonObject
+struct GLTFEXPORTER_API FGLTFJsonPrimitive
 {
-	FGLTFJsonIndex         Indices;
-	FGLTFJsonIndex         Material;
+	FGLTFJsonAccessorIndex Indices;
+	FGLTFJsonMaterialIndex Material;
 	EGLTFJsonPrimitiveMode Mode;
 	FGLTFJsonAttributes    Attributes;
 
@@ -77,7 +78,7 @@ struct GLTFEXPORTER_API FGLTFJsonPrimitive : FGLTFJsonObject
 	}
 };
 
-struct GLTFEXPORTER_API FGLTFJsonMesh : FGLTFJsonObject
+struct GLTFEXPORTER_API FGLTFJsonMesh
 {
 	FString Name;
 
diff --git a/Engine/Plugins/Enterprise/GLTFExporter/MigratedFromGit/Source/GLTFExporter/Public/GLTFJsonNode.h b/Engine/Plugins/Enterprise/GLTFExporter/MigratedFromGit/Source/GLTFExporter/Public/GLTFJsonNode.h
--- a/Engine/Plugins/Enterprise/GLTFExporter/MigratedFromGit/Source/GLTFExporter/Public/GLTFJsonNode.h
+++ b/Engine/Plugins/Enterprise/GLTFExporter/MigratedFromGit/Source/GLTFExporter/Public/GLTFJsonNode.h
@@ -2,9 +2,10 @@
 
 #pragma once
 
-#include "GLTFJsonObject.h"
+#include "GLTFJsonIndex.h"
+#include "Serialization/JsonSerializer.h"
 
-struct GLTFEXPORTER_API FGLTFJsonNode : FGLTFJsonObject
+struct GLTFEXPORTER_API FGLTFJsonNode
 {
 	FString Name;
 
@@ -12,11 +13,11 @@ struct GLTFEXPORTER_API FGLTFJsonNode : FGLTFJsonObject
 	FQuat   Rotation;
 	FVector Scale;
 
-	FGLTFJsonIndex Camera;
-	FGLTFJsonIndex Skin;
-	FGLTFJsonIndex Mesh;
+	FGLTFJsonCameraIndex Camera;
+	FGLTFJsonSkinIndex   Skin;
+	FGLTFJsonMeshIndex   Mesh;
 
-	TArray<FGLTFJsonIndex> Children;
+	TArray<FGLTFJsonNodeIndex> Children;
 
 	FGLTFJsonNode()
 		: Translation(FVector::ZeroVector)
diff --git a/Engine/Plugins/Enterprise/GLTFExporter/MigratedFromGit/Source/GLTFExporter/Public/GLTFJsonObject.h b/Engine/Plugins/Enterprise/GLTFExporter/MigratedFromGit/Source/GLTFExporter/Public/GLTFJsonObject.h
deleted file mode 100644
diff --git a/Engine/Plugins/Enterprise/GLTFExporter/MigratedFromGit/Source/GLTFExporter/Public/GLTFJsonRoot.h b/Engine/Plugins/Enterprise/GLTFExporter/MigratedFromGit/Source/GLTFExporter/Public/GLTFJsonRoot.h
--- a/Engine/Plugins/Enterprise/GLTFExporter/MigratedFromGit/Source/GLTFExporter/Public/GLTFJsonRoot.h
+++ b/Engine/Plugins/Enterprise/GLTFExporter/MigratedFromGit/Source/GLTFExporter/Public/GLTFJsonRoot.h
@@ -2,7 +2,7 @@
 
 #pragma once
 
-#include "GLTFJsonObject.h"
+#include "GLTFJsonIndex.h"
 #include "GLTFJsonAccessor.h"
 #include "GLTFJsonBufferView.h"
 #include "GLTFJsonBuffer.h"
@@ -15,7 +15,7 @@
 #include "Runtime/Launch/Resources/Version.h"
 
 
-struct GLTFEXPORTER_API FGLTFJsonAsset : FGLTFJsonObject
+struct GLTFEXPORTER_API FGLTFJsonAsset
 {
 	FString Version;
 	FString Generator;
@@ -42,8 +42,8 @@ struct GLTFEXPORTER_API FGLTFJsonAsset : FGLTFJsonObject
 
 struct GLTFEXPORTER_API FGLTFJsonRoot
 {
-	FGLTFJsonAsset Asset;
-	FGLTFJsonIndex DefaultScene;
+	FGLTFJsonAsset      Asset;
+	FGLTFJsonSceneIndex DefaultScene;
 
 	TArray<FGLTFJsonAccessor>   Accessors;
 	TArray<FGLTFJsonBuffer>     Buffers;
diff --git a/Engine/Plugins/Enterprise/GLTFExporter/MigratedFromGit/Source/GLTFExporter/Public/GLTFJsonScene.h b/Engine/Plugins/Enterprise/GLTFExporter/MigratedFromGit/Source/GLTFExporter/Public/GLTFJsonScene.h
--- a/Engine/Plugins/Enterprise/GLTFExporter/MigratedFromGit/Source/GLTFExporter/Public/GLTFJsonScene.h
+++ b/Engine/Plugins/Enterprise/GLTFExporter/MigratedFromGit/Source/GLTFExporter/Public/GLTFJsonScene.h
@@ -2,13 +2,14 @@
 
 #pragma once
 
-#include "GLTFJsonObject.h"
+#include "GLTFJsonIndex.h"
+#include "Serialization/JsonSerializer.h"
 
-struct GLTFEXPORTER_API FGLTFJsonScene: FGLTFJsonObject
+struct GLTFEXPORTER_API FGLTFJsonScene
 {
 	FString Name;
 
-	TArray<FGLTFJsonIndex> Nodes;
+	TArray<FGLTFJsonNodeIndex> Nodes;
 
 	template <class CharType = TCHAR, class PrintPolicy = TPrettyJsonPrintPolicy<CharType>>
 	void Write(TJsonWriter<CharType, PrintPolicy>& JsonWriter) const
diff --git a/Engine/Plugins/Enterprise/GLTFExporter/MigratedFromGit/Source/GLTFExporter/Public/GLTFMeshConverter.h b/Engine/Plugins/Enterprise/GLTFExporter/MigratedFromGit/Source/GLTFExporter/Public/GLTFMeshConverter.h
--- a/Engine/Plugins/Enterprise/GLTFExporter/MigratedFromGit/Source/GLTFExporter/Public/GLTFMeshConverter.h
+++ b/Engine/Plugins/Enterprise/GLTFExporter/MigratedFromGit/Source/GLTFExporter/Public/GLTFMeshConverter.h
@@ -14,7 +14,7 @@ struct GLTFEXPORTER_API FGLTFConvertedSection
 
 	FGLTFConvertedSection(const FString& SectionName, const FStaticMeshLODResources& LODMesh, int32 SectionIndex);
 
-	FGLTFJsonIndex AppendAccessorForIndices(FGLTFJsonRoot& JsonRoot, FGLTFBufferBuilder& BufferBuilder) const;
+	FGLTFJsonAccessorIndex AppendAccessorForIndices(FGLTFJsonRoot& JsonRoot, FGLTFBufferBuilder& BufferBuilder) const;
 };
 
 struct GLTFEXPORTER_API FGLTFConvertedMesh
@@ -34,14 +34,14 @@ struct GLTFEXPORTER_API FGLTFConvertedMesh
 
 	FGLTFConvertedMesh(const UStaticMesh* StaticMesh, int32 LODIndex);
 
-	FGLTFJsonIndex AppendAccessorForPositions(FGLTFJsonRoot& JsonRoot, FGLTFBufferBuilder& BufferBuilder) const;
-	FGLTFJsonIndex AppendAccessorForColors(FGLTFJsonRoot& JsonRoot, FGLTFBufferBuilder& BufferBuilder) const;
-	FGLTFJsonIndex AppendAccessorForNormals(FGLTFJsonRoot& JsonRoot, FGLTFBufferBuilder& BufferBuilder) const;
-	FGLTFJsonIndex AppendAccessorForTangents(FGLTFJsonRoot& JsonRoot, FGLTFBufferBuilder& BufferBuilder) const;
-	FGLTFJsonIndex AppendAccessorForUV0s(FGLTFJsonRoot& JsonRoot, FGLTFBufferBuilder& BufferBuilder) const;
-	FGLTFJsonIndex AppendAccessorForUV1s(FGLTFJsonRoot& JsonRoot, FGLTFBufferBuilder& BufferBuilder) const;
+	FGLTFJsonAccessorIndex AppendAccessorForPositions(FGLTFJsonRoot& JsonRoot, FGLTFBufferBuilder& BufferBuilder) const;
+	FGLTFJsonAccessorIndex AppendAccessorForColors(FGLTFJsonRoot& JsonRoot, FGLTFBufferBuilder& BufferBuilder) const;
+	FGLTFJsonAccessorIndex AppendAccessorForNormals(FGLTFJsonRoot& JsonRoot, FGLTFBufferBuilder& BufferBuilder) const;
+	FGLTFJsonAccessorIndex AppendAccessorForTangents(FGLTFJsonRoot& JsonRoot, FGLTFBufferBuilder& BufferBuilder) const;
+	FGLTFJsonAccessorIndex AppendAccessorForUV0s(FGLTFJsonRoot& JsonRoot, FGLTFBufferBuilder& BufferBuilder) const;
+	FGLTFJsonAccessorIndex AppendAccessorForUV1s(FGLTFJsonRoot& JsonRoot, FGLTFBufferBuilder& BufferBuilder) const;
 
-	FGLTFJsonIndex AppendMesh(FGLTFJsonRoot& JsonRoot, FGLTFBufferBuilder& BufferBuilder) const;
+	FGLTFJsonMeshIndex AppendMesh(FGLTFJsonRoot& JsonRoot, FGLTFBufferBuilder& BufferBuilder) const;
 };
 
 struct GLTFEXPORTER_API FGLTFMeshConverter
@@ -55,7 +55,7 @@ struct GLTFEXPORTER_API FGLTFMeshConverter
 	{
 	}
 
-	FGLTFJsonIndex AppendMesh(const UStaticMesh* StaticMesh, int32 LODIndex)
+	FGLTFJsonMeshIndex AppendMesh(const UStaticMesh* StaticMesh, int32 LODIndex)
 	{
 		return FGLTFConvertedMesh(StaticMesh, LODIndex).AppendMesh(JsonRoot, BufferBuilder);
 	}
