[GLTFExporter] refactor: remove unnecessary inlines

commit fe04a7603eec36e3d03a1a4732d937adeaad6665
Date: Fri, 24 Apr 2020 15:00:49 +0200

    refactor: remove unnecessary inlines

#skipci

diff --git a/Engine/Plugins/Enterprise/GLTFExporter/MigratedFromGit/Source/GLTFExporter/Private/GLTFBufferBuilder.cpp b/Engine/Plugins/Enterprise/GLTFExporter/MigratedFromGit/Source/GLTFExporter/Private/GLTFBufferBuilder.cpp
--- a/Engine/Plugins/Enterprise/GLTFExporter/MigratedFromGit/Source/GLTFExporter/Private/GLTFBufferBuilder.cpp
+++ b/Engine/Plugins/Enterprise/GLTFExporter/MigratedFromGit/Source/GLTFExporter/Private/GLTFBufferBuilder.cpp
@@ -3,14 +3,14 @@
 #include "GLTFBufferBuilder.h"
 #include "Misc/Base64.h"
 
-inline FGLTFBufferBuilder::FGLTFBufferBuilder(FGLTFJsonRoot& JsonRoot, const FString& Name)
+FGLTFBufferBuilder::FGLTFBufferBuilder(FGLTFJsonRoot& JsonRoot, const FString& Name)
 	: JsonRoot(JsonRoot)
 	, BufferIndex(JsonRoot.Buffers.AddDefaulted(1))
 {
 	JsonRoot.Buffers[BufferIndex].Name = Name;
 }
 
-inline FGLTFJsonBufferViewIndex FGLTFBufferBuilder::AppendBufferView(const void* RawData, uint64 ByteLength, const FString& Name, EGLTFJsonBufferTarget BufferTarget)
+FGLTFJsonBufferViewIndex FGLTFBufferBuilder::AppendBufferView(const void* RawData, uint64 ByteLength, const FString& Name, EGLTFJsonBufferTarget BufferTarget)
 {
 	FGLTFJsonBufferView BufferView;
 	BufferView.Name = Name;
diff --git a/Engine/Plugins/Enterprise/GLTFExporter/MigratedFromGit/Source/GLTFExporter/Private/GLTFExportOptionsWindow.cpp b/Engine/Plugins/Enterprise/GLTFExporter/MigratedFromGit/Source/GLTFExporter/Private/GLTFExportOptionsWindow.cpp
--- a/Engine/Plugins/Enterprise/GLTFExporter/MigratedFromGit/Source/GLTFExporter/Private/GLTFExportOptionsWindow.cpp
+++ b/Engine/Plugins/Enterprise/GLTFExporter/MigratedFromGit/Source/GLTFExporter/Private/GLTFExportOptionsWindow.cpp
@@ -153,7 +153,7 @@ FReply SGLTFExportOptionsWindow::OnResetToDefaultClick() const
 	return FReply::Handled();
 }
 
-inline FReply SGLTFExportOptionsWindow::OnExport()
+FReply SGLTFExportOptionsWindow::OnExport()
 {
 	bShouldExport = true;
 	if (WidgetWindow.IsValid())
@@ -163,13 +163,13 @@ inline FReply SGLTFExportOptionsWindow::OnExport()
 	return FReply::Handled();
 }
 
-inline FReply SGLTFExportOptionsWindow::OnExportAll()
+FReply SGLTFExportOptionsWindow::OnExportAll()
 {
 	bShouldExportAll = true;
 	return OnExport();
 }
 
-inline FReply SGLTFExportOptionsWindow::OnCancel()
+FReply SGLTFExportOptionsWindow::OnCancel()
 {
 	bShouldExport = false;
 	bShouldExportAll = false;
@@ -180,7 +180,7 @@ inline FReply SGLTFExportOptionsWindow::OnCancel()
 	return FReply::Handled();
 }
 
-inline FReply SGLTFExportOptionsWindow::OnKeyDown(const FGeometry & MyGeometry, const FKeyEvent & InKeyEvent)
+FReply SGLTFExportOptionsWindow::OnKeyDown(const FGeometry & MyGeometry, const FKeyEvent & InKeyEvent)
 {
 	if (InKeyEvent.GetKey() == EKeys::Escape)
 	{
@@ -190,17 +190,17 @@ inline FReply SGLTFExportOptionsWindow::OnKeyDown(const FGeometry & MyGeometry,
 	return FReply::Unhandled();
 }
 
-inline bool SGLTFExportOptionsWindow::ShouldExport() const
+bool SGLTFExportOptionsWindow::ShouldExport() const
 {
 	return bShouldExport;
 }
 
-inline bool SGLTFExportOptionsWindow::ShouldExportAll() const
+bool SGLTFExportOptionsWindow::ShouldExportAll() const
 {
 	return bShouldExportAll;
 }
 
-inline SGLTFExportOptionsWindow::SGLTFExportOptionsWindow()
+SGLTFExportOptionsWindow::SGLTFExportOptionsWindow()
 	: ExportOptions(nullptr)
 	, bShouldExport(false)
 	, bShouldExportAll(false)
