[GLTFExporter] refactor: add braces to single-line if-statement, as per the code standard section "Code Formatting".

commit 654c9035860c94522330f2da2764d19bae3c8348
Date: Fri, 24 Apr 2020 15:03:17 +0200

    refactor: add braces to single-line if-statement, as per the code standard section "Code Formatting".

#skipci

diff --git a/Engine/Plugins/Enterprise/GLTFExporter/MigratedFromGit/Source/GLTFExporter/Private/GLTFMeshConverter.cpp b/Engine/Plugins/Enterprise/GLTFExporter/MigratedFromGit/Source/GLTFExporter/Private/GLTFMeshConverter.cpp
--- a/Engine/Plugins/Enterprise/GLTFExporter/MigratedFromGit/Source/GLTFExporter/Private/GLTFMeshConverter.cpp
+++ b/Engine/Plugins/Enterprise/GLTFExporter/MigratedFromGit/Source/GLTFExporter/Private/GLTFMeshConverter.cpp
@@ -20,7 +20,10 @@ FGLTFConvertedSection::FGLTFConvertedSection(const FString& SectionName, const F
 
 FGLTFJsonAccessorIndex FGLTFConvertedSection::AppendAccessorForIndices(FGLTFJsonRoot& Root, FGLTFBufferBuilder& BufferBuilder) const
 {
-	if (Indices.Num() == 0) return INDEX_NONE;
+	if (Indices.Num() == 0)
+	{
+		return INDEX_NONE;
+	}
 
 	FString AttributeName = Name + TEXT("_Indices");
 	FGLTFJsonBufferViewIndex BufferViewIndex = BufferBuilder.AppendBufferView(Indices, AttributeName, EGLTFJsonBufferTarget::ElementArrayBuffer);
@@ -100,7 +103,10 @@ FGLTFConvertedMesh::FGLTFConvertedMesh(const UStaticMesh* StaticMesh, int32 LODI
 
 FGLTFJsonAccessorIndex FGLTFConvertedMesh::AppendAccessorForPositions(FGLTFJsonRoot& Root, FGLTFBufferBuilder& BufferBuilder) const
 {
-	if (Positions.Num() == 0) return INDEX_NONE;
+	if (Positions.Num() == 0)
+	{
+		return INDEX_NONE;
+	}
 
 	FString AttributeName = Name + TEXT("_Positions");
 	FGLTFJsonBufferViewIndex BufferViewIndex = BufferBuilder.AppendBufferView(Positions, AttributeName);
@@ -127,7 +133,10 @@ FGLTFJsonAccessorIndex FGLTFConvertedMesh::AppendAccessorForPositions(FGLTFJsonR
 
 FGLTFJsonAccessorIndex FGLTFConvertedMesh::AppendAccessorForNormals(FGLTFJsonRoot& Root, FGLTFBufferBuilder& BufferBuilder) const
 {
-	if (Normals.Num() == 0) return INDEX_NONE;
+	if (Normals.Num() == 0)
+	{
+		return INDEX_NONE;
+	}
 
 	FString AttributeName = Name + TEXT("_Normals");
 	FGLTFJsonBufferViewIndex BufferViewIndex = BufferBuilder.AppendBufferView(Normals, AttributeName);
@@ -144,7 +153,10 @@ FGLTFJsonAccessorIndex FGLTFConvertedMesh::AppendAccessorForNormals(FGLTFJsonRoo
 
 FGLTFJsonAccessorIndex FGLTFConvertedMesh::AppendAccessorForColors(FGLTFJsonRoot& Root, FGLTFBufferBuilder& BufferBuilder) const
 {
-	if (Colors.Num() == 0) return INDEX_NONE;
+	if (Colors.Num() == 0)
+	{
+		return INDEX_NONE;
+	}
 
 	FString AttributeName = Name + TEXT("_Colors");
 	FGLTFJsonBufferViewIndex BufferViewIndex = BufferBuilder.AppendBufferView(Colors, AttributeName);
@@ -161,7 +173,10 @@ FGLTFJsonAccessorIndex FGLTFConvertedMesh::AppendAccessorForColors(FGLTFJsonRoot
 
 FGLTFJsonAccessorIndex FGLTFConvertedMesh::AppendAccessorForTangents(FGLTFJsonRoot& Root, FGLTFBufferBuilder& BufferBuilder) const
 {
-	if (Tangents.Num() == 0) return INDEX_NONE;
+	if (Tangents.Num() == 0)
+	{
+		return INDEX_NONE;
+	}
 
 	FString AttributeName = Name + TEXT("_Tangents");
 	FGLTFJsonBufferViewIndex BufferViewIndex = BufferBuilder.AppendBufferView(Tangents, AttributeName);
@@ -178,7 +193,10 @@ FGLTFJsonAccessorIndex FGLTFConvertedMesh::AppendAccessorForTangents(FGLTFJsonRo
 
 FGLTFJsonAccessorIndex FGLTFConvertedMesh::AppendAccessorForUV0s(FGLTFJsonRoot& Root, FGLTFBufferBuilder& BufferBuilder) const
 {
-	if (UV0s.Num() == 0) return INDEX_NONE;
+	if (UV0s.Num() == 0)
+	{
+		return INDEX_NONE;
+	}
 
 	FString AttributeName = Name + TEXT("_UV0s");
 	FGLTFJsonBufferViewIndex BufferViewIndex = BufferBuilder.AppendBufferView(UV0s, AttributeName);
@@ -195,7 +213,10 @@ FGLTFJsonAccessorIndex FGLTFConvertedMesh::AppendAccessorForUV0s(FGLTFJsonRoot&
 
 FGLTFJsonAccessorIndex FGLTFConvertedMesh::AppendAccessorForUV1s(FGLTFJsonRoot& Root, FGLTFBufferBuilder& BufferBuilder) const
 {
-	if (UV1s.Num() == 0) return INDEX_NONE;
+	if (UV1s.Num() == 0)
+	{
+		return INDEX_NONE;
+	}
 
 	FString AttributeName = Name + TEXT("_UV1s");
 	FGLTFJsonBufferViewIndex BufferViewIndex = BufferBuilder.AppendBufferView(UV1s, AttributeName);
