[GLTFExporter] refactor: make non-changing local variable FString AttributeName to const

commit 13c03f9a9f95266bceb3713f853705c9e7b320cb
Date: Fri, 24 Apr 2020 15:04:56 +0200

    refactor: make non-changing local variable FString AttributeName to const

#skipci

diff --git a/Engine/Plugins/Enterprise/GLTFExporter/MigratedFromGit/Source/GLTFExporter/Private/GLTFMeshConverter.cpp b/Engine/Plugins/Enterprise/GLTFExporter/MigratedFromGit/Source/GLTFExporter/Private/GLTFMeshConverter.cpp
--- a/Engine/Plugins/Enterprise/GLTFExporter/MigratedFromGit/Source/GLTFExporter/Private/GLTFMeshConverter.cpp
+++ b/Engine/Plugins/Enterprise/GLTFExporter/MigratedFromGit/Source/GLTFExporter/Private/GLTFMeshConverter.cpp
@@ -25,7 +25,7 @@ FGLTFJsonAccessorIndex FGLTFConvertedSection::AppendAccessorForIndices(FGLTFJson
 		return INDEX_NONE;
 	}
 
-	FString AttributeName = Name + TEXT("_Indices");
+	const FString AttributeName = Name + TEXT("_Indices");
 	FGLTFJsonBufferViewIndex BufferViewIndex = BufferBuilder.AppendBufferView(Indices, AttributeName, EGLTFJsonBufferTarget::ElementArrayBuffer);
 
 	FGLTFJsonAccessor Accessor;
@@ -108,7 +108,7 @@ FGLTFJsonAccessorIndex FGLTFConvertedMesh::AppendAccessorForPositions(FGLTFJsonR
 		return INDEX_NONE;
 	}
 
-	FString AttributeName = Name + TEXT("_Positions");
+	const FString AttributeName = Name + TEXT("_Positions");
 	FGLTFJsonBufferViewIndex BufferViewIndex = BufferBuilder.AppendBufferView(Positions, AttributeName);
 
 	FGLTFJsonAccessor Accessor;
@@ -138,7 +138,7 @@ FGLTFJsonAccessorIndex FGLTFConvertedMesh::AppendAccessorForNormals(FGLTFJsonRoo
 		return INDEX_NONE;
 	}
 
-	FString AttributeName = Name + TEXT("_Normals");
+	const FString AttributeName = Name + TEXT("_Normals");
 	FGLTFJsonBufferViewIndex BufferViewIndex = BufferBuilder.AppendBufferView(Normals, AttributeName);
 
 	FGLTFJsonAccessor Accessor;
@@ -158,7 +158,7 @@ FGLTFJsonAccessorIndex FGLTFConvertedMesh::AppendAccessorForColors(FGLTFJsonRoot
 		return INDEX_NONE;
 	}
 
-	FString AttributeName = Name + TEXT("_Colors");
+	const FString AttributeName = Name + TEXT("_Colors");
 	FGLTFJsonBufferViewIndex BufferViewIndex = BufferBuilder.AppendBufferView(Colors, AttributeName);
 
 	FGLTFJsonAccessor Accessor;
@@ -178,7 +178,7 @@ FGLTFJsonAccessorIndex FGLTFConvertedMesh::AppendAccessorForTangents(FGLTFJsonRo
 		return INDEX_NONE;
 	}
 
-	FString AttributeName = Name + TEXT("_Tangents");
+	const FString AttributeName = Name + TEXT("_Tangents");
 	FGLTFJsonBufferViewIndex BufferViewIndex = BufferBuilder.AppendBufferView(Tangents, AttributeName);
 
 	FGLTFJsonAccessor Accessor;
@@ -198,7 +198,7 @@ FGLTFJsonAccessorIndex FGLTFConvertedMesh::AppendAccessorForUV0s(FGLTFJsonRoot&
 		return INDEX_NONE;
 	}
 
-	FString AttributeName = Name + TEXT("_UV0s");
+	const FString AttributeName = Name + TEXT("_UV0s");
 	FGLTFJsonBufferViewIndex BufferViewIndex = BufferBuilder.AppendBufferView(UV0s, AttributeName);
 
 	FGLTFJsonAccessor Accessor;
@@ -218,7 +218,7 @@ FGLTFJsonAccessorIndex FGLTFConvertedMesh::AppendAccessorForUV1s(FGLTFJsonRoot&
 		return INDEX_NONE;
 	}
 
-	FString AttributeName = Name + TEXT("_UV1s");
+	const FString AttributeName = Name + TEXT("_UV1s");
 	FGLTFJsonBufferViewIndex BufferViewIndex = BufferBuilder.AppendBufferView(UV1s, AttributeName);
 
 	FGLTFJsonAccessor Accessor;
