[GLTFExporter] refactor: merge GLTFBufferBuilder and GLTFMeshConverter into GLTFContainer

commit 970b1bff357d3093930b146c3b57ec5033ba676c
Date: Fri, 24 Apr 2020 15:35:16 +0200

    refactor: merge GLTFBufferBuilder and GLTFMeshConverter into GLTFContainer

#skipci

diff --git a/Engine/Plugins/Enterprise/GLTFExporter/MigratedFromGit/Source/GLTFExporter/Private/GLTFBufferBuilder.cpp b/Engine/Plugins/Enterprise/GLTFExporter/MigratedFromGit/Source/GLTFExporter/Private/GLTFBufferBuilder.cpp
deleted file mode 100644
diff --git a/Engine/Plugins/Enterprise/GLTFExporter/MigratedFromGit/Source/GLTFExporter/Private/GLTFContainer.cpp b/Engine/Plugins/Enterprise/GLTFExporter/MigratedFromGit/Source/GLTFExporter/Private/GLTFContainer.cpp
new file mode 100644
--- /dev/null
+++ b/Engine/Plugins/Enterprise/GLTFExporter/MigratedFromGit/Source/GLTFExporter/Private/GLTFContainer.cpp
@@ -0,0 +1,47 @@
+// Copyright Epic Games, Inc. All Rights Reserved.
+
+#include "GLTFContainer.h"
+#include "GLTFConversionMesh.h"
+#include "Misc/Base64.h"
+
+FGLTFContainer::FGLTFContainer()
+	: MergedBufferIndex(JsonRoot.Buffers.AddDefaulted(1))
+{
+}
+
+FGLTFJsonBufferViewIndex FGLTFContainer::AppendBufferView(const void* RawData, uint64 ByteLength, const FString& Name, EGLTFJsonBufferTarget BufferTarget)
+{
+	FGLTFJsonBufferView BufferView;
+	BufferView.Name = Name;
+	BufferView.Buffer = MergedBufferIndex;
+	BufferView.ByteOffset = MergedBufferData.Num();
+	BufferView.ByteLength = ByteLength;
+	BufferView.Target = BufferTarget;
+
+	MergedBufferData.Append(static_cast<const uint8*>(RawData), ByteLength);
+
+	return JsonRoot.BufferViews.Add(BufferView);
+}
+
+void FGLTFContainer::UpdateMergedBuffer()
+{
+	FGLTFJsonBuffer& Buffer = JsonRoot.Buffers[MergedBufferIndex];
+	if (Buffer.ByteLength != MergedBufferData.Num())
+	{
+		FString DataBase64 = FBase64::Encode(MergedBufferData.GetData(), MergedBufferData.Num());
+
+		Buffer.URI = TEXT("data:application/octet-stream;base64,") + DataBase64;
+		Buffer.ByteLength = MergedBufferData.Num();
+	}
+}
+
+void FGLTFContainer::Serialize(FArchive& Archive)
+{
+	UpdateMergedBuffer();
+	JsonRoot.Serialize(&Archive, true);
+}
+
+FGLTFJsonMeshIndex FGLTFContainer::AppendMesh(const UStaticMesh* StaticMesh, int32 LODIndex)
+{
+	return FGLTFConversionMesh(StaticMesh, LODIndex).AppendMesh(*this);
+}
diff --git a/Engine/Plugins/Enterprise/GLTFExporter/MigratedFromGit/Source/GLTFExporter/Private/GLTFConversionMesh.cpp b/Engine/Plugins/Enterprise/GLTFExporter/MigratedFromGit/Source/GLTFExporter/Private/GLTFConversionMesh.cpp
new file mode 100644
--- /dev/null
+++ b/Engine/Plugins/Enterprise/GLTFExporter/MigratedFromGit/Source/GLTFExporter/Private/GLTFConversionMesh.cpp
@@ -0,0 +1,257 @@
+// Copyright Epic Games, Inc. All Rights Reserved.
+
+#include "GLTFConversionMesh.h"
+#include "GLTFConversionUtilities.h"
+
+FGLTFConversionSection::FGLTFConversionSection(const FString& SectionName, const FStaticMeshLODResources& LODMesh, int32 SectionIndex)
+	: Name(SectionName)
+{
+	FIndexArrayView RawIndices = LODMesh.IndexBuffer.GetArrayView();
+	const FStaticMeshSection& Section = LODMesh.Sections[SectionIndex];
+
+	const uint32 IndexCount = Section.NumTriangles * 3;
+	Indices.AddUninitialized(IndexCount);
+
+	for (uint32 Index = 0; Index < IndexCount; ++Index)
+	{
+		Indices[Index] = RawIndices[Section.FirstIndex + Index];
+	}
+}
+
+FGLTFJsonAccessorIndex FGLTFConversionSection::AppendAccessorForIndices(FGLTFContainer& Container) const
+{
+	if (Indices.Num() == 0)
+	{
+		return INDEX_NONE;
+	}
+
+	const FString AttributeName = Name + TEXT("_Indices");
+	FGLTFJsonBufferViewIndex BufferViewIndex = Container.AppendBufferView(Indices, AttributeName, EGLTFJsonBufferTarget::ElementArrayBuffer);
+
+	FGLTFJsonAccessor Accessor;
+	Accessor.Name = AttributeName;
+	Accessor.BufferView = BufferViewIndex;
+	Accessor.ComponentType = EGLTFJsonComponentType::U32;
+	Accessor.Count = Indices.Num();
+	Accessor.Type = EGLTFJsonAccessorType::Scalar;
+
+	return Container.JsonRoot.Accessors.Add(Accessor);
+}
+
+FGLTFConversionMesh::FGLTFConversionMesh(const UStaticMesh* StaticMesh, int32 LODIndex)
+{
+	StaticMesh->GetName(Name);
+
+	const FStaticMeshLODResources& LODMesh = StaticMesh->GetLODForExport(LODIndex);
+
+	const int32 SectionCount = LODMesh.Sections.Num();
+	for (int32 SectionIndex = 0; SectionIndex < SectionCount; ++SectionIndex)
+	{
+		FString SectionName = Name + TEXT("_Section") + FString::FromInt(SectionIndex);
+		Sections.Add(FGLTFConversionSection(SectionName, LODMesh, SectionIndex));
+	}
+
+	const int32 PositionCount = LODMesh.VertexBuffers.PositionVertexBuffer.GetNumVertices();
+	Positions.AddUninitialized(PositionCount);
+
+	for (int32 PosIndex = 0; PosIndex < PositionCount; ++PosIndex)
+	{
+		Positions[PosIndex] = ConvertPosition(LODMesh.VertexBuffers.PositionVertexBuffer.VertexPosition(PosIndex));
+	}
+
+	const int32 ColorCount = LODMesh.VertexBuffers.ColorVertexBuffer.GetNumVertices();
+	if (ColorCount > 0)
+	{
+		Colors.AddUninitialized(ColorCount);
+		for (int32 ColorIndex = 0; ColorIndex < ColorCount; ++ColorIndex)
+		{
+			Colors[ColorIndex] = ConvertColor(LODMesh.VertexBuffers.ColorVertexBuffer.VertexColor(ColorIndex));
+		}
+	}
+
+	const int32 VertexCount = LODMesh.VertexBuffers.StaticMeshVertexBuffer.GetNumVertices();
+	Normals.AddUninitialized(VertexCount);
+	Tangents.AddUninitialized(VertexCount);
+
+	for (int32 VertIndex = 0; VertIndex < VertexCount; ++VertIndex)
+	{
+		Normals[VertIndex] = ConvertVector(LODMesh.VertexBuffers.StaticMeshVertexBuffer.VertexTangentZ(VertIndex));
+		Tangents[VertIndex] = ConvertTangent(LODMesh.VertexBuffers.StaticMeshVertexBuffer.VertexTangentX(VertIndex));
+	}
+
+	const int32 UVCount = LODMesh.GetNumTexCoords();
+	if (UVCount >= 1)
+	{
+		UV0s.AddUninitialized(VertexCount);
+		for (int32 VertIndex = 0; VertIndex < VertexCount; ++VertIndex)
+		{
+			UV0s[VertIndex] = LODMesh.VertexBuffers.StaticMeshVertexBuffer.GetVertexUV(VertIndex, 0);
+		}
+
+		if (UVCount >= 2)
+		{
+			UV1s.AddUninitialized(VertexCount);
+			for (int32 VertIndex = 0; VertIndex < VertexCount; ++VertIndex)
+			{
+				UV1s[VertIndex] = LODMesh.VertexBuffers.StaticMeshVertexBuffer.GetVertexUV(VertIndex, 1);
+			}
+		}
+	}
+
+	BoundingBox = StaticMesh->GetBoundingBox();
+}
+
+FGLTFJsonAccessorIndex FGLTFConversionMesh::AppendAccessorForPositions(FGLTFContainer& Container) const
+{
+	if (Positions.Num() == 0)
+	{
+		return INDEX_NONE;
+	}
+
+	const FString AttributeName = Name + TEXT("_Positions");
+	FGLTFJsonBufferViewIndex BufferViewIndex = Container.AppendBufferView(Positions, AttributeName);
+
+	FGLTFJsonAccessor Accessor;
+	Accessor.Name = AttributeName;
+	Accessor.BufferView = BufferViewIndex;
+	Accessor.ComponentType = EGLTFJsonComponentType::F32;
+	Accessor.Count = Positions.Num();
+	Accessor.Type = EGLTFJsonAccessorType::Vec3;
+
+	FVector Max = ConvertSize(BoundingBox.Max);
+	FVector Min = ConvertSize(BoundingBox.Min);
+	Accessor.Max[0] = Max.X;
+	Accessor.Max[1] = Max.Y;
+	Accessor.Max[2] = Max.Z;
+	Accessor.Min[0] = Min.X;
+	Accessor.Min[1] = Min.Y;
+	Accessor.Min[2] = Min.Z;
+	Accessor.MinMaxLength = 3;
+
+	return Container.JsonRoot.Accessors.Add(Accessor);
+}
+
+FGLTFJsonAccessorIndex FGLTFConversionMesh::AppendAccessorForNormals(FGLTFContainer& Container) const
+{
+	if (Normals.Num() == 0)
+	{
+		return INDEX_NONE;
+	}
+
+	const FString AttributeName = Name + TEXT("_Normals");
+	FGLTFJsonBufferViewIndex BufferViewIndex = Container.AppendBufferView(Normals, AttributeName);
+
+	FGLTFJsonAccessor Accessor;
+	Accessor.Name = AttributeName;
+	Accessor.BufferView = BufferViewIndex;
+	Accessor.ComponentType = EGLTFJsonComponentType::F32;
+	Accessor.Count = Normals.Num();
+	Accessor.Type = EGLTFJsonAccessorType::Vec3;
+
+	return Container.JsonRoot.Accessors.Add(Accessor);
+}
+
+FGLTFJsonAccessorIndex FGLTFConversionMesh::AppendAccessorForColors(FGLTFContainer& Container) const
+{
+	if (Colors.Num() == 0)
+	{
+		return INDEX_NONE;
+	}
+
+	const FString AttributeName = Name + TEXT("_Colors");
+	FGLTFJsonBufferViewIndex BufferViewIndex = Container.AppendBufferView(Colors, AttributeName);
+
+	FGLTFJsonAccessor Accessor;
+	Accessor.Name = AttributeName;
+	Accessor.BufferView = BufferViewIndex;
+	Accessor.ComponentType = EGLTFJsonComponentType::U8;
+	Accessor.Count = Colors.Num();
+	Accessor.Type = EGLTFJsonAccessorType::Vec4;
+
+	return Container.JsonRoot.Accessors.Add(Accessor);
+}
+
+FGLTFJsonAccessorIndex FGLTFConversionMesh::AppendAccessorForTangents(FGLTFContainer& Container) const
+{
+	if (Tangents.Num() == 0)
+	{
+		return INDEX_NONE;
+	}
+
+	const FString AttributeName = Name + TEXT("_Tangents");
+	FGLTFJsonBufferViewIndex BufferViewIndex = Container.AppendBufferView(Tangents, AttributeName);
+
+	FGLTFJsonAccessor Accessor;
+	Accessor.Name = AttributeName;
+	Accessor.BufferView = BufferViewIndex;
+	Accessor.ComponentType = EGLTFJsonComponentType::F32;
+	Accessor.Count = Tangents.Num();
+	Accessor.Type = EGLTFJsonAccessorType::Vec4;
+
+	return Container.JsonRoot.Accessors.Add(Accessor);
+}
+
+FGLTFJsonAccessorIndex FGLTFConversionMesh::AppendAccessorForUV0s(FGLTFContainer& Container) const
+{
+	if (UV0s.Num() == 0)
+	{
+		return INDEX_NONE;
+	}
+
+	const FString AttributeName = Name + TEXT("_UV0s");
+	FGLTFJsonBufferViewIndex BufferViewIndex = Container.AppendBufferView(UV0s, AttributeName);
+
+	FGLTFJsonAccessor Accessor;
+	Accessor.Name = AttributeName;
+	Accessor.BufferView = BufferViewIndex;
+	Accessor.ComponentType = EGLTFJsonComponentType::F32;
+	Accessor.Count = UV0s.Num();
+	Accessor.Type = EGLTFJsonAccessorType::Vec2;
+
+	return Container.JsonRoot.Accessors.Add(Accessor);
+}
+
+FGLTFJsonAccessorIndex FGLTFConversionMesh::AppendAccessorForUV1s(FGLTFContainer& Container) const
+{
+	if (UV1s.Num() == 0)
+	{
+		return INDEX_NONE;
+	}
+
+	const FString AttributeName = Name + TEXT("_UV1s");
+	FGLTFJsonBufferViewIndex BufferViewIndex = Container.AppendBufferView(UV1s, AttributeName);
+
+	FGLTFJsonAccessor Accessor;
+	Accessor.Name = AttributeName;
+	Accessor.BufferView = BufferViewIndex;
+	Accessor.ComponentType = EGLTFJsonComponentType::F32;
+	Accessor.Count = UV1s.Num();
+	Accessor.Type = EGLTFJsonAccessorType::Vec2;
+
+	return Container.JsonRoot.Accessors.Add(Accessor);
+}
+
+FGLTFJsonMeshIndex FGLTFConversionMesh::AppendMesh(FGLTFContainer& Container) const
+{
+	FGLTFJsonMesh Mesh;
+	Mesh.Name = Name;
+
+	FGLTFJsonAttributes Attributes;
+	Attributes.Position = AppendAccessorForPositions(Container);
+	Attributes.Color0 = AppendAccessorForColors(Container);
+	Attributes.Normal = AppendAccessorForNormals(Container);
+	Attributes.Tangent = AppendAccessorForTangents(Container);
+	Attributes.TexCoord0 = AppendAccessorForUV0s(Container);
+	Attributes.TexCoord1 = AppendAccessorForUV1s(Container);
+
+	for (const FGLTFConversionSection& Section : Sections)
+	{
+		FGLTFJsonPrimitive Primitive;
+		Primitive.Attributes = Attributes;
+
+		Primitive.Indices = Section.AppendAccessorForIndices(Container);
+		Mesh.Primitives.Add(Primitive);
+	}
+
+	return Container.JsonRoot.Meshes.Add(Mesh);
+}
diff --git a/Engine/Plugins/Enterprise/GLTFExporter/MigratedFromGit/Source/GLTFExporter/Private/GLTFMeshConverter.cpp b/Engine/Plugins/Enterprise/GLTFExporter/MigratedFromGit/Source/GLTFExporter/Private/GLTFMeshConverter.cpp
deleted file mode 100644
diff --git a/Engine/Plugins/Enterprise/GLTFExporter/MigratedFromGit/Source/GLTFExporter/Private/GLTFStaticMeshExporter.cpp b/Engine/Plugins/Enterprise/GLTFExporter/MigratedFromGit/Source/GLTFExporter/Private/GLTFStaticMeshExporter.cpp
--- a/Engine/Plugins/Enterprise/GLTFExporter/MigratedFromGit/Source/GLTFExporter/Private/GLTFStaticMeshExporter.cpp
+++ b/Engine/Plugins/Enterprise/GLTFExporter/MigratedFromGit/Source/GLTFExporter/Private/GLTFStaticMeshExporter.cpp
@@ -2,8 +2,7 @@
 
 #include "GLTFStaticMeshExporter.h"
 #include "GLTFJsonRoot.h"
-#include "GLTFBufferBuilder.h"
-#include "GLTFMeshConverter.h"
+#include "GLTFContainer.h"
 #include "Engine/StaticMesh.h"
 
 UGLTFStaticMeshExporter::UGLTFStaticMeshExporter(const FObjectInitializer& ObjectInitializer)
@@ -22,23 +21,20 @@ bool UGLTFStaticMeshExporter::ExportBinary(UObject* Object, const TCHAR* Type, F
 		return false;
 	}
 
-	FGLTFJsonRoot Root;
-	FGLTFBufferBuilder BufferBuilder(Root);
-	FGLTFMeshConverter MeshConverter(Root, BufferBuilder);
+	FGLTFContainer Container;
 
-	FGLTFJsonMeshIndex MeshIndex = MeshConverter.AppendMesh(StaticMesh, 0);
+	FGLTFJsonMeshIndex MeshIndex = Container.AppendMesh(StaticMesh, 0);
 
 	FGLTFJsonNode Node;
 	Node.Mesh = MeshIndex;
-	FGLTFJsonNodeIndex NodeIndex = Root.Nodes.Add(Node);
+	FGLTFJsonNodeIndex NodeIndex = Container.JsonRoot.Nodes.Add(Node);
 
 	FGLTFJsonScene Scene;
 	Scene.Nodes.Add(NodeIndex);
-	FGLTFJsonSceneIndex SceneIndex = Root.Scenes.Add(Scene);
+	FGLTFJsonSceneIndex SceneIndex = Container.JsonRoot.Scenes.Add(Scene);
 
-	Root.DefaultScene = SceneIndex;
+	Container.JsonRoot.DefaultScene = SceneIndex;
 
-	BufferBuilder.Close();
-	Root.Serialize(&Archive, true);
+	Container.Serialize(Archive);
 	return true;
 }
diff --git a/Engine/Plugins/Enterprise/GLTFExporter/MigratedFromGit/Source/GLTFExporter/Public/GLTFBufferBuilder.h b/Engine/Plugins/Enterprise/GLTFExporter/MigratedFromGit/Source/GLTFExporter/Public/GLTFBufferBuilder.h
deleted file mode 100644
diff --git a/Engine/Plugins/Enterprise/GLTFExporter/MigratedFromGit/Source/GLTFExporter/Public/GLTFContainer.h b/Engine/Plugins/Enterprise/GLTFExporter/MigratedFromGit/Source/GLTFExporter/Public/GLTFContainer.h
new file mode 100644
--- /dev/null
+++ b/Engine/Plugins/Enterprise/GLTFExporter/MigratedFromGit/Source/GLTFExporter/Public/GLTFContainer.h
@@ -0,0 +1,29 @@
+// Copyright Epic Games, Inc. All Rights Reserved.
+
+#pragma once
+
+#include "GLTFJsonRoot.h"
+#include "Engine.h"
+
+struct GLTFEXPORTER_API FGLTFContainer
+{
+	FGLTFJsonRoot JsonRoot;
+
+	const FGLTFJsonBufferIndex MergedBufferIndex;
+	TArray<uint8> MergedBufferData;
+
+	FGLTFContainer();
+
+	template <class ElementType>
+	FGLTFJsonBufferViewIndex AppendBufferView(const TArray<ElementType>& Array, const FString& Name = TEXT(""), EGLTFJsonBufferTarget BufferTarget = EGLTFJsonBufferTarget::ArrayBuffer)
+	{
+		return AppendBufferView(Array.GetData(), Array.Num() * sizeof(ElementType), Name, BufferTarget);
+	}
+
+	FGLTFJsonBufferViewIndex AppendBufferView(const void* RawData, uint64 ByteLength, const FString& Name = TEXT(""), EGLTFJsonBufferTarget BufferTarget = EGLTFJsonBufferTarget::ArrayBuffer);
+
+	FGLTFJsonMeshIndex AppendMesh(const UStaticMesh* StaticMesh, int32 LODIndex);
+
+	void UpdateMergedBuffer();
+	void Serialize(FArchive& Archive);
+};
diff --git a/Engine/Plugins/Enterprise/GLTFExporter/MigratedFromGit/Source/GLTFExporter/Public/GLTFConversionMesh.h b/Engine/Plugins/Enterprise/GLTFExporter/MigratedFromGit/Source/GLTFExporter/Public/GLTFConversionMesh.h
new file mode 100644
--- /dev/null
+++ b/Engine/Plugins/Enterprise/GLTFExporter/MigratedFromGit/Source/GLTFExporter/Public/GLTFConversionMesh.h
@@ -0,0 +1,45 @@
+// Copyright Epic Games, Inc. All Rights Reserved.
+
+#pragma once
+
+#include "GLTFJsonRoot.h"
+#include "GLTFContainer.h"
+#include "Engine.h"
+
+struct GLTFEXPORTER_API FGLTFConversionSection
+{
+	FString Name;
+
+	TArray<uint32> Indices;
+
+	FGLTFConversionSection(const FString& SectionName, const FStaticMeshLODResources& LODMesh, int32 SectionIndex);
+
+	FGLTFJsonAccessorIndex AppendAccessorForIndices(FGLTFContainer& Container) const;
+};
+
+struct GLTFEXPORTER_API FGLTFConversionMesh
+{
+	FString Name;
+
+	TArray<FGLTFConversionSection> Sections;
+
+	TArray<FVector>   Positions;
+	TArray<FColor>    Colors;
+	TArray<FVector>   Normals;
+	TArray<FVector4>  Tangents;
+	TArray<FVector2D> UV0s;
+	TArray<FVector2D> UV1s;
+
+	FBox BoundingBox;
+
+	FGLTFConversionMesh(const UStaticMesh* StaticMesh, int32 LODIndex);
+
+	FGLTFJsonAccessorIndex AppendAccessorForPositions(FGLTFContainer& Container) const;
+	FGLTFJsonAccessorIndex AppendAccessorForColors(FGLTFContainer& Container) const;
+	FGLTFJsonAccessorIndex AppendAccessorForNormals(FGLTFContainer& Container) const;
+	FGLTFJsonAccessorIndex AppendAccessorForTangents(FGLTFContainer& Container) const;
+	FGLTFJsonAccessorIndex AppendAccessorForUV0s(FGLTFContainer& Container) const;
+	FGLTFJsonAccessorIndex AppendAccessorForUV1s(FGLTFContainer& Container) const;
+
+	FGLTFJsonMeshIndex AppendMesh(FGLTFContainer& Container) const;
+};
diff --git a/Engine/Plugins/Enterprise/GLTFExporter/MigratedFromGit/Source/GLTFExporter/Public/GLTFMeshConverter.h b/Engine/Plugins/Enterprise/GLTFExporter/MigratedFromGit/Source/GLTFExporter/Public/GLTFMeshConverter.h
deleted file mode 100644
