[GLTFExporter] feat: add initial export tests

commit 5e7dd26e061de402672a3544cbae8a17fa578840
Date: Fri, 24 Apr 2020 18:22:37 +0200

    feat: add initial export tests
    (cherry picked from commit 097ab66bf673dfc46dd15e81422e3a2dbbf0f049​)

#skipci

diff --git a/Engine/Plugins/Enterprise/GLTFExporter/MigratedFromGit/Source/GLTFExporter/Private/Tests/GLTFExporterTest.cpp b/Engine/Plugins/Enterprise/GLTFExporter/MigratedFromGit/Source/GLTFExporter/Private/Tests/GLTFExporterTest.cpp
new file mode 100644
--- /dev/null
+++ b/Engine/Plugins/Enterprise/GLTFExporter/MigratedFromGit/Source/GLTFExporter/Private/Tests/GLTFExporterTest.cpp
@@ -0,0 +1,80 @@
+// Copyright 1998-2019 Epic Games, Inc. All Rights Reserved.
+
+#include "CoreMinimal.h"
+#include "Misc/AutomationTest.h"
+#include "UObject/StrongObjectPtr.h"
+#include "UnrealExporter.h"
+#include "Serialization/BufferArchive.h"
+#include "Engine/StaticMesh.h"
+#include "Engine/Classes/Engine/DataTable.h"
+#include "Tests/GLTFExporterTestTargetTableRow.h"
+
+#include "GLTFExporter.h"
+
+#if WITH_DEV_AUTOMATION_TESTS
+
+namespace {
+
+// TODO: move to config
+const TCHAR* TestTargetTableAssetPath = TEXT("DataTable'/Game/GLTFExportTestTargets.GLTFExportTestTargets'");
+const TCHAR* Delimiter = TEXT("kaviarmacka");
+
+} // anonymous namespace
+
+IMPLEMENT_COMPLEX_AUTOMATION_TEST(FGLTFExporterTest, "Unreal2glTF.Export Test", EAutomationTestFlags::EditorContext | EAutomationTestFlags::EngineFilter)
+
+void FGLTFExporterTest::GetTests(TArray<FString>& OutBeautifiedNames, TArray <FString>& OutTestCommands) const
+{
+	const UDataTable* TestTargetTable = LoadObject<UDataTable>(NULL, TestTargetTableAssetPath, NULL, LOAD_None, NULL);
+
+	if (TestTargetTable == nullptr) {
+		return;
+	}
+
+	const TArray<FName> TableRowNames = TestTargetTable->GetRowNames();
+	const int32 TableRowCount = TableRowNames.Num();
+	const FString ContextString = "";
+
+	for (int32 TableRowIndex = 0; TableRowIndex < TableRowCount; ++TableRowIndex) {
+		const FName TableRowName = TableRowNames[TableRowIndex];
+		const FGLTFExporterTestTargetTableRow* TableRow = TestTargetTable->FindRow<FGLTFExporterTestTargetTableRow>(TableRowName, ContextString);
+
+		if (TableRow == nullptr) {
+			continue;
+		}
+
+		const FString Elements[] = { TableRow->TargetStaticMesh->GetPathName(), TableRow->ExpectedOutput };
+
+		OutBeautifiedNames.Add(FString(TEXT("Test target with index ")) + FString::FromInt(TableRowIndex));
+		OutTestCommands.Add(FString::Join(Elements, Delimiter));
+	}
+}
+
+bool FGLTFExporterTest::RunTest(const FString& Parameters)
+{
+	FString TargetStaticMeshPath = "";
+	FString ExpectedOutput = "";
+	Parameters.Split(Delimiter, &TargetStaticMeshPath, &ExpectedOutput);
+
+	UStaticMesh* ObjectToExport = LoadObject<UStaticMesh>(NULL, *TargetStaticMeshPath, NULL, LOAD_None, NULL);
+
+	if (ObjectToExport == nullptr) {
+		AddError(FString::Printf(TEXT("Failed to find test asset %s"), *TargetStaticMeshPath));
+
+		return false;
+	}
+
+	FBufferArchive BufferArchive;
+	UExporter::ExportToArchive(ObjectToExport, nullptr, BufferArchive, TEXT("gltf"), 0);
+	const FString ExportedText = BytesToString(BufferArchive.GetData(), BufferArchive.Num());
+
+	if (!ExportedText.Equals(ExpectedOutput)) {
+		AddError(FString::Printf(TEXT("Exported GLTF for the asset %s did not match the expected result"), *TargetStaticMeshPath));
+
+		return false;
+	}
+
+	return true;
+}
+
+#endif //WITH_DEV_AUTOMATION_TESTS
diff --git a/Engine/Plugins/Enterprise/GLTFExporter/MigratedFromGit/Source/GLTFExporter/Public/Tests/GLTFExporterTestTargetTableRow.h b/Engine/Plugins/Enterprise/GLTFExporter/MigratedFromGit/Source/GLTFExporter/Public/Tests/GLTFExporterTestTargetTableRow.h
new file mode 100644
--- /dev/null
+++ b/Engine/Plugins/Enterprise/GLTFExporter/MigratedFromGit/Source/GLTFExporter/Public/Tests/GLTFExporterTestTargetTableRow.h
@@ -0,0 +1,18 @@
+// Copyright Epic Games, Inc. All Rights Reserved.
+
+#include "Engine/DataTable.h"
+#include "GLTFExporterTestTargetTableRow.generated.h"
+
+class UStaticMesh;
+
+USTRUCT(BlueprintType)
+struct FGLTFExporterTestTargetTableRow : public FTableRowBase
+{
+	GENERATED_USTRUCT_BODY()
+
+	UPROPERTY(EditAnywhere, BlueprintReadWrite)
+	UStaticMesh* TargetStaticMesh;
+
+	UPROPERTY(EditAnywhere, BlueprintReadWrite)
+	FString ExpectedOutput;
+};
