[GLTFExporter] feat(test): compare export result to actual gltf files

commit b64bb8feda0a1276ca6b22c052f5ac9d922f8890
Date: Mon, 27 Apr 2020 15:42:53 +0200

    feat(test): compare export result to actual gltf files
    (cherry picked from commit d47e7cf835b42e39b9ef78befd15909512fac495​)

#skipci

diff --git a/Engine/Plugins/Enterprise/GLTFExporter/MigratedFromGit/Source/GLTFExporter/Private/Tests/GLTFExporterTest.cpp b/Engine/Plugins/Enterprise/GLTFExporter/MigratedFromGit/Source/GLTFExporter/Private/Tests/GLTFExporterTest.cpp
--- a/Engine/Plugins/Enterprise/GLTFExporter/MigratedFromGit/Source/GLTFExporter/Private/Tests/GLTFExporterTest.cpp
+++ b/Engine/Plugins/Enterprise/GLTFExporter/MigratedFromGit/Source/GLTFExporter/Private/Tests/GLTFExporterTest.cpp
@@ -2,12 +2,10 @@
 
 #include "CoreMinimal.h"
 #include "Misc/AutomationTest.h"
-#include "UObject/StrongObjectPtr.h"
 #include "UnrealExporter.h"
 #include "Serialization/BufferArchive.h"
-#include "Engine/StaticMesh.h"
-#include "Engine/Classes/Engine/DataTable.h"
-#include "Tests/GLTFExporterTestTargetTableRow.h"
+#include "Misc/FileHelper.h"
+#include "Misc/Paths.h"
 
 #include "GLTFExporter.h"
 
@@ -15,9 +13,17 @@
 
 namespace {
 
-// TODO: move to config
-const TCHAR* TestTargetTableAssetPath = TEXT("DataTable'/Game/GLTFExportTestTargets.GLTFExportTestTargets'");
-const TCHAR* Delimiter = TEXT("kaviarmacka");
+struct TestDefinition {
+	TCHAR* AssetPath;
+	TCHAR* ControlFilePath;
+} TestDefinitions[] = {
+	{
+		TEXT("StaticMesh'/Engine/EngineMeshes/Cube.Cube'"),
+		TEXT("Tests/Cube/Cube.gltf")
+	}
+};
+
+const TCHAR* ParamDelimiter = TEXT("c8a4fd9d525c0ac433fd7d24ce2a3eca");
 
 } // anonymous namespace
 
@@ -25,41 +31,27 @@ IMPLEMENT_COMPLEX_AUTOMATION_TEST(FGLTFExporterTest, "Unreal2glTF.Export Test",
 
 void FGLTFExporterTest::GetTests(TArray<FString>& OutBeautifiedNames, TArray <FString>& OutTestCommands) const
 {
-	const UDataTable* TestTargetTable = LoadObject<UDataTable>(NULL, TestTargetTableAssetPath, NULL, LOAD_None, NULL);
-
-	if (TestTargetTable == nullptr) {
-		return;
-	}
-
-	const TArray<FName> TableRowNames = TestTargetTable->GetRowNames();
-	const int32 TableRowCount = TableRowNames.Num();
-	const FString ContextString = "";
-
-	for (int32 TableRowIndex = 0; TableRowIndex < TableRowCount; ++TableRowIndex) {
-		const FName TableRowName = TableRowNames[TableRowIndex];
-		const FGLTFExporterTestTargetTableRow* TableRow = TestTargetTable->FindRow<FGLTFExporterTestTargetTableRow>(TableRowName, ContextString);
+	const int32 TestDefinitionCount = sizeof(TestDefinitions) / sizeof(TestDefinition);
 
-		if (TableRow == nullptr) {
-			continue;
-		}
+	for (int32 TestDefinitionIndex = 0; TestDefinitionIndex < TestDefinitionCount; ++TestDefinitionIndex) {
+		const TestDefinition& CurrentTestDefinition = TestDefinitions[TestDefinitionIndex];
+		const FString Elements[] = { CurrentTestDefinition.AssetPath, CurrentTestDefinition.ControlFilePath };
 
-		const FString Elements[] = { TableRow->TargetStaticMesh->GetPathName(), TableRow->ExpectedOutput };
-
-		OutBeautifiedNames.Add(FString(TEXT("Test target with index ")) + FString::FromInt(TableRowIndex));
-		OutTestCommands.Add(FString::Join(Elements, Delimiter));
+		OutBeautifiedNames.Add(FString(TEXT("Test target with index ")) + FString::FromInt(TestDefinitionIndex));
+		OutTestCommands.Add(FString::Join(Elements, ParamDelimiter));
 	}
 }
 
 bool FGLTFExporterTest::RunTest(const FString& Parameters)
 {
-	FString TargetStaticMeshPath = "";
-	FString ExpectedOutput = "";
-	Parameters.Split(Delimiter, &TargetStaticMeshPath, &ExpectedOutput);
+	FString AssetPath = "";
+	FString ControlFilePath = "";
+	Parameters.Split(ParamDelimiter, &AssetPath, &ControlFilePath);
 
-	UStaticMesh* ObjectToExport = LoadObject<UStaticMesh>(NULL, *TargetStaticMeshPath, NULL, LOAD_None, NULL);
+	UObject* ObjectToExport = LoadObject<UObject>(NULL, *AssetPath, NULL, LOAD_None, NULL);
 
 	if (ObjectToExport == nullptr) {
-		AddError(FString::Printf(TEXT("Failed to find test asset %s"), *TargetStaticMeshPath));
+		AddError(FString::Printf(TEXT("Failed to find test asset %s"), *AssetPath));
 
 		return false;
 	}
@@ -68,8 +60,17 @@ bool FGLTFExporterTest::RunTest(const FString& Parameters)
 	UExporter::ExportToArchive(ObjectToExport, nullptr, BufferArchive, TEXT("gltf"), 0);
 	const FString ExportedText = BytesToString(BufferArchive.GetData(), BufferArchive.Num());
 
-	if (!ExportedText.Equals(ExpectedOutput)) {
-		AddError(FString::Printf(TEXT("Exported GLTF for the asset %s did not match the expected result"), *TargetStaticMeshPath));
+	const FString AbsoluteControlFilePath = FPaths::ConvertRelativePathToFull(FPaths::ProjectDir(), ControlFilePath);
+	FString ControlFileContent = "";
+
+	if (!FFileHelper::LoadFileToString(ControlFileContent, *AbsoluteControlFilePath)) {
+		AddError(FString::Printf(TEXT("Failed to find test control file %s"), *AbsoluteControlFilePath));
+
+		return false;
+	}
+
+	if (!ExportedText.Equals(ControlFileContent)) {
+		AddError(FString::Printf(TEXT("Exported GLTF for the asset %s did not match the expected result"), *AssetPath));
 
 		return false;
 	}
diff --git a/Engine/Plugins/Enterprise/GLTFExporter/MigratedFromGit/Source/GLTFExporter/Public/Tests/GLTFExporterTestTargetTableRow.h b/Engine/Plugins/Enterprise/GLTFExporter/MigratedFromGit/Source/GLTFExporter/Public/Tests/GLTFExporterTestTargetTableRow.h
deleted file mode 100644
