[GLTFExporter] feat(test): move test target list into json file

commit 418956790e6b4e9b004edac51d7bba12ced9774c
Date: Mon, 27 Apr 2020 21:58:01 +0200

    feat(test): move test target list into json file
    (cherry picked from commit 214022559506c241ea360d7488ea98e5e96fe538​)

#skipci

diff --git a/Engine/Plugins/Enterprise/GLTFExporter/MigratedFromGit/Source/GLTFExporter/Private/Tests/GLTFExporterTest.cpp b/Engine/Plugins/Enterprise/GLTFExporter/MigratedFromGit/Source/GLTFExporter/Private/Tests/GLTFExporterTest.cpp
--- a/Engine/Plugins/Enterprise/GLTFExporter/MigratedFromGit/Source/GLTFExporter/Private/Tests/GLTFExporterTest.cpp
+++ b/Engine/Plugins/Enterprise/GLTFExporter/MigratedFromGit/Source/GLTFExporter/Private/Tests/GLTFExporterTest.cpp
@@ -6,6 +6,7 @@
 #include "Serialization/BufferArchive.h"
 #include "Misc/FileHelper.h"
 #include "Misc/Paths.h"
+#include "Serialization/JsonSerializer.h"
 
 #include "GLTFExporter.h"
 
@@ -13,16 +14,8 @@
 
 namespace {
 
-struct TestDefinition {
-	TCHAR* AssetPath;
-	TCHAR* ControlFilePath;
-} TestDefinitions[] = {
-	{
-		TEXT("StaticMesh'/Engine/EngineMeshes/Cube.Cube'"),
-		TEXT("Tests/Cube/Cube.gltf")
-	}
-};
-
+// TODO Turn into config variable?
+const TCHAR* TargetFilePath = TEXT("Tests/targets.json");
 const TCHAR* ParamDelimiter = TEXT("c8a4fd9d525c0ac433fd7d24ce2a3eca");
 
 } // anonymous namespace
@@ -31,14 +24,26 @@ IMPLEMENT_COMPLEX_AUTOMATION_TEST(FGLTFExporterTest, "Unreal2glTF.Export Test",
 
 void FGLTFExporterTest::GetTests(TArray<FString>& OutBeautifiedNames, TArray <FString>& OutTestCommands) const
 {
-	const int32 TestDefinitionCount = sizeof(TestDefinitions) / sizeof(TestDefinition);
+	const FString TargetFilePathAbsolute = FPaths::ConvertRelativePathToFull(FPaths::ProjectDir(), TargetFilePath);
+	FString TargetFileContent = "";
+	FFileHelper::LoadFileToString(TargetFileContent, *TargetFilePathAbsolute);
+
+	TArray<TSharedPtr<FJsonValue>> TargetJsonRoot = {};
+	const TSharedRef<TJsonReader<>> JsonReader = TJsonReaderFactory<>::Create(TargetFileContent);
+
+	if (FJsonSerializer::Deserialize(JsonReader, TargetJsonRoot))
+	{
+		const int32 TargetJsonElementCount = TargetJsonRoot.Num();
 
-	for (int32 TestDefinitionIndex = 0; TestDefinitionIndex < TestDefinitionCount; ++TestDefinitionIndex) {
-		const TestDefinition& CurrentTestDefinition = TestDefinitions[TestDefinitionIndex];
-		const FString Elements[] = { CurrentTestDefinition.AssetPath, CurrentTestDefinition.ControlFilePath };
+		for (int32 TargetJsonArrayElementIndex = 0; TargetJsonArrayElementIndex < TargetJsonElementCount; ++TargetJsonArrayElementIndex) {
+			const TSharedPtr<FJsonObject>& TargetJsonElement = TargetJsonRoot[TargetJsonArrayElementIndex]->AsObject();
+			const FString AssetField = TargetJsonElement->GetStringField("asset");
+			const FString ControlField = TargetJsonElement->GetStringField("control");
+			const FString ParamElements[] = { AssetField, ControlField };
 
-		OutBeautifiedNames.Add(FString(TEXT("Test target with index ")) + FString::FromInt(TestDefinitionIndex));
-		OutTestCommands.Add(FString::Join(Elements, ParamDelimiter));
+			OutBeautifiedNames.Add(FString(TEXT("Test target with index ")) + FString::FromInt(TargetJsonArrayElementIndex));
+			OutTestCommands.Add(FString::Join(ParamElements, ParamDelimiter));
+		}
 	}
 }
 
@@ -56,21 +61,27 @@ bool FGLTFExporterTest::RunTest(const FString& Parameters)
 		return false;
 	}
 
-	FBufferArchive BufferArchive;
-	UExporter::ExportToArchive(ObjectToExport, nullptr, BufferArchive, TEXT("gltf"), 0);
-	const FString ExportedText = BytesToString(BufferArchive.GetData(), BufferArchive.Num());
+	FBufferArchive BufferArchive = {};
 
-	const FString AbsoluteControlFilePath = FPaths::ConvertRelativePathToFull(FPaths::ProjectDir(), ControlFilePath);
+	if (!UExporter::ExportToArchive(ObjectToExport, nullptr, BufferArchive, TEXT("gltf"), 0)) {
+		AddError(FString::Printf(TEXT("Export failed for asset %s"), *AssetPath));
+
+		return false;
+	}
+
+	const FString ControlFilePathAbsolute = FPaths::ConvertRelativePathToFull(FPaths::ProjectDir(), ControlFilePath);
 	FString ControlFileContent = "";
 
-	if (!FFileHelper::LoadFileToString(ControlFileContent, *AbsoluteControlFilePath)) {
-		AddError(FString::Printf(TEXT("Failed to find test control file %s"), *AbsoluteControlFilePath));
+	if (!FFileHelper::LoadFileToString(ControlFileContent, *ControlFilePathAbsolute)) {
+		AddError(FString::Printf(TEXT("Failed to find test control file %s"), *ControlFilePathAbsolute));
 
 		return false;
 	}
 
+	const FString ExportedText = FString(BufferArchive.Num(), UTF8_TO_TCHAR(BufferArchive.GetData()));
+
 	if (!ExportedText.Equals(ControlFileContent)) {
-		AddError(FString::Printf(TEXT("Exported GLTF for the asset %s did not match the expected result"), *AssetPath));
+		AddError(FString::Printf(TEXT("Exported GLTF for the asset %s did not match the control"), *AssetPath));
 
 		return false;
 	}
