[GLTFExporter] refactor(test): fix issues from code analysis

commit 9cd816e719dfff9e6a6d1a0773bd21216b5fed39
Date: Tue, 28 Apr 2020 11:12:35 +0200

    refactor(test): fix issues from code analysis
    (cherry picked from commit 836917ace5f781760f509f4e4b6a73777f6859af​)

#skipci

diff --git a/Engine/Plugins/Enterprise/GLTFExporter/MigratedFromGit/Source/GLTFExporter/Private/Tests/GLTFExporterTest.cpp b/Engine/Plugins/Enterprise/GLTFExporter/MigratedFromGit/Source/GLTFExporter/Private/Tests/GLTFExporterTest.cpp
--- a/Engine/Plugins/Enterprise/GLTFExporter/MigratedFromGit/Source/GLTFExporter/Private/Tests/GLTFExporterTest.cpp
+++ b/Engine/Plugins/Enterprise/GLTFExporter/MigratedFromGit/Source/GLTFExporter/Private/Tests/GLTFExporterTest.cpp
@@ -2,7 +2,6 @@
 
 #include "CoreMinimal.h"
 #include "Misc/AutomationTest.h"
-#include "UnrealExporter.h"
 #include "Serialization/BufferArchive.h"
 #include "Misc/FileHelper.h"
 #include "Misc/Paths.h"
@@ -53,7 +52,7 @@ bool FGLTFExporterTest::RunTest(const FString& Parameters)
 	FString ControlFilePath = "";
 	Parameters.Split(ParamDelimiter, &AssetPath, &ControlFilePath);
 
-	UObject* ObjectToExport = LoadObject<UObject>(NULL, *AssetPath, NULL, LOAD_None, NULL);
+	UObject* ObjectToExport = LoadObject<UObject>(nullptr, *AssetPath, nullptr, LOAD_None, nullptr);
 
 	if (ObjectToExport == nullptr) {
 		AddError(FString::Printf(TEXT("Failed to find test asset %s"), *AssetPath));
