[GLTFExporter] feat: add ConvertScale and improve comments

commit 92b5fc3788cb2b0cda179bc183204181921f6051
Date: Tue, 28 Apr 2020 11:04:29 +0200

    feat: add ConvertScale and improve comments
    (cherry picked from commit 1212c2576ae050c6e8b3c31de9347e092491e787​)

#skipci

diff --git a/Engine/Plugins/Enterprise/GLTFExporter/MigratedFromGit/Source/GLTFExporter/Public/GLTFConversionUtilities.h b/Engine/Plugins/Enterprise/GLTFExporter/MigratedFromGit/Source/GLTFExporter/Public/GLTFConversionUtilities.h
--- a/Engine/Plugins/Enterprise/GLTFExporter/MigratedFromGit/Source/GLTFExporter/Public/GLTFConversionUtilities.h
+++ b/Engine/Plugins/Enterprise/GLTFExporter/MigratedFromGit/Source/GLTFExporter/Public/GLTFConversionUtilities.h
@@ -17,34 +17,37 @@ inline FVector ConvertPosition(const FVector& Position)
 	return ConvertVector(Position);
 }
 
-inline FVector ConvertSize(const FVector& Position)
+inline FVector ConvertSize(const FVector& Size)
 {
-	return ConvertVector(Position);
+	return ConvertVector(Size);
+}
+
+inline FVector ConvertScale(const FVector& Scale)
+{
+	return ConvertVector(Scale);
 }
 
 inline FVector4 ConvertTangent(const FVector4& Tangent)
 {
 	// glTF stores tangent as Vec4, with W component indicating handedness of tangent basis.
-
-	return FVector4(Tangent.X, Tangent.Z, Tangent.Y, Tangent.W);
+	return FVector4(ConvertVector(Tangent), Tangent.W);
 }
 
 inline FColor ConvertColor(const FColor& Color)
 {
-	// UE4 uses ABGR.
-	// glTF uses RGBA.
+	// UE4 uses ABGR while glTF uses RGBA.
 	return { Color.R, Color.G, Color.B, Color.A };
 }
 
-inline FQuat ConvertQuat(const FQuat& Quat)
+inline FQuat ConvertRotation(const FQuat& Quat)
 {
-	// glTF uses a right-handed coordinate system, with Y up.
 	// UE4 uses a left-handed coordinate system, with Z up.
+	// glTF uses a right-handed coordinate system, with Y up.
 	// Quat = (qX, qY, qZ, qW) = (sin(angle/2) * aX, sin(angle/2) * aY, sin(angle/2) * aZ, cons(angle/2))
 	// where (aX, aY, aZ) - rotation axis, angle - rotation angle
 	// Y swapped with Z between these coordinate systems
 	// also, as handedness is changed rotation is inversed - hence negation
-	// therefore QuatUE = (-qX, -qZ, -qY, qw)
+	// therefore glTFQuat = (-qX, -qZ, -qY, qw)
 
 	FQuat Result(-Quat.X, -Quat.Z, -Quat.Y, Quat.W);
 	// Not checking if quaternion is normalized
@@ -52,10 +55,10 @@ inline FQuat ConvertQuat(const FQuat& Quat)
 	return Result;
 }
 
-inline FMatrix ConvertMat(const FMatrix& Matrix)
+inline FMatrix ConvertMatrix(const FMatrix& Matrix)
 {
-	// glTF stores matrix elements in column major order
-	// Unreal's FMatrix is row major
+	// Unreal stores matrix elements in row major order.
+	// glTF stores matrix elements in column major order.
 
 	FMatrix Result;
 	for (int32 Row = 0; Row < 4; ++Row)
