[GLTFExporter] feat: minor optimize of GLTFConversionMesh

commit fc66bc174ba14327f4d6cb92c991d77fc6db965e
Date: Tue, 28 Apr 2020 11:04:55 +0200

    feat: minor optimize of GLTFConversionMesh
    (cherry picked from commit 3a797e317b710cfe798dab9de5610a5fced92f29​)

#skipci

diff --git a/Engine/Plugins/Enterprise/GLTFExporter/MigratedFromGit/Source/GLTFExporter/Private/GLTFConversionMesh.cpp b/Engine/Plugins/Enterprise/GLTFExporter/MigratedFromGit/Source/GLTFExporter/Private/GLTFConversionMesh.cpp
--- a/Engine/Plugins/Enterprise/GLTFExporter/MigratedFromGit/Source/GLTFExporter/Private/GLTFConversionMesh.cpp
+++ b/Engine/Plugins/Enterprise/GLTFExporter/MigratedFromGit/Source/GLTFExporter/Private/GLTFConversionMesh.cpp
@@ -3,18 +3,16 @@
 #include "GLTFConversionMesh.h"
 #include "GLTFConversionUtilities.h"
 
-FGLTFConversionSection::FGLTFConversionSection(const FString& SectionName, const FStaticMeshLODResources& LODMesh, int32 SectionIndex)
+FGLTFConversionSection::FGLTFConversionSection(const FString& SectionName, const FStaticMeshSection& MeshSection, const FIndexArrayView& IndexArray)
 	: Name(SectionName)
 {
-	FIndexArrayView RawIndices = LODMesh.IndexBuffer.GetArrayView();
-	const FStaticMeshSection& Section = LODMesh.Sections[SectionIndex];
-
-	const uint32 IndexCount = Section.NumTriangles * 3;
+	const uint32 IndexCount = MeshSection.NumTriangles * 3;
+	const uint32 FirstIndex = MeshSection.FirstIndex;
 	Indices.AddUninitialized(IndexCount);
 
 	for (uint32 Index = 0; Index < IndexCount; ++Index)
 	{
-		Indices[Index] = RawIndices[Section.FirstIndex + Index];
+		Indices[Index] = IndexArray[FirstIndex + Index];
 	}
 }
 
@@ -44,56 +42,69 @@ FGLTFConversionMesh::FGLTFConversionMesh(const UStaticMesh* StaticMesh, int32 LO
 
 	const FStaticMeshLODResources& LODMesh = StaticMesh->GetLODForExport(LODIndex);
 
+	const FIndexArrayView IndexArray = LODMesh.IndexBuffer.GetArrayView();
 	const int32 SectionCount = LODMesh.Sections.Num();
+
+	Sections.Reserve(SectionCount);
+
 	for (int32 SectionIndex = 0; SectionIndex < SectionCount; ++SectionIndex)
 	{
 		FString SectionName = Name + TEXT("_Section") + FString::FromInt(SectionIndex);
-		Sections.Add(FGLTFConversionSection(SectionName, LODMesh, SectionIndex));
+		Sections.Emplace(SectionName, LODMesh.Sections[SectionIndex], IndexArray);
 	}
 
-	const int32 PositionCount = LODMesh.VertexBuffers.PositionVertexBuffer.GetNumVertices();
+	const FPositionVertexBuffer& PositionVertexBuffer = LODMesh.VertexBuffers.PositionVertexBuffer;
+	const int32 PositionCount = PositionVertexBuffer.GetNumVertices();
+
 	Positions.AddUninitialized(PositionCount);
 
 	for (int32 PosIndex = 0; PosIndex < PositionCount; ++PosIndex)
 	{
-		Positions[PosIndex] = ConvertPosition(LODMesh.VertexBuffers.PositionVertexBuffer.VertexPosition(PosIndex));
+		Positions[PosIndex] = ConvertPosition(PositionVertexBuffer.VertexPosition(PosIndex));
 	}
 
-	const int32 ColorCount = LODMesh.VertexBuffers.ColorVertexBuffer.GetNumVertices();
+	const FColorVertexBuffer& ColorVertexBuffer = LODMesh.VertexBuffers.ColorVertexBuffer;
+	const int32 ColorCount = ColorVertexBuffer.GetNumVertices();
+
 	if (ColorCount > 0)
 	{
 		Colors.AddUninitialized(ColorCount);
+
 		for (int32 ColorIndex = 0; ColorIndex < ColorCount; ++ColorIndex)
 		{
-			Colors[ColorIndex] = ConvertColor(LODMesh.VertexBuffers.ColorVertexBuffer.VertexColor(ColorIndex));
+			Colors[ColorIndex] = ConvertColor(ColorVertexBuffer.VertexColor(ColorIndex));
 		}
 	}
 
-	const int32 VertexCount = LODMesh.VertexBuffers.StaticMeshVertexBuffer.GetNumVertices();
+	const FStaticMeshVertexBuffer& StaticMeshVertexBuffer = LODMesh.VertexBuffers.StaticMeshVertexBuffer;
+	const int32 VertexCount = StaticMeshVertexBuffer.GetNumVertices();
+
 	Normals.AddUninitialized(VertexCount);
 	Tangents.AddUninitialized(VertexCount);
 
 	for (int32 VertIndex = 0; VertIndex < VertexCount; ++VertIndex)
 	{
-		Normals[VertIndex] = ConvertVector(LODMesh.VertexBuffers.StaticMeshVertexBuffer.VertexTangentZ(VertIndex));
-		Tangents[VertIndex] = ConvertTangent(LODMesh.VertexBuffers.StaticMeshVertexBuffer.VertexTangentX(VertIndex));
+		Normals[VertIndex] = ConvertVector(StaticMeshVertexBuffer.VertexTangentZ(VertIndex));
+		Tangents[VertIndex] = ConvertTangent(StaticMeshVertexBuffer.VertexTangentX(VertIndex));
 	}
 
-	const int32 UVCount = LODMesh.GetNumTexCoords();
+	const int32 UVCount = StaticMeshVertexBuffer.GetNumTexCoords();
 	if (UVCount >= 1)
 	{
 		UV0s.AddUninitialized(VertexCount);
+
 		for (int32 VertIndex = 0; VertIndex < VertexCount; ++VertIndex)
 		{
-			UV0s[VertIndex] = LODMesh.VertexBuffers.StaticMeshVertexBuffer.GetVertexUV(VertIndex, 0);
+			UV0s[VertIndex] = StaticMeshVertexBuffer.GetVertexUV(VertIndex, 0);
 		}
 
 		if (UVCount >= 2)
 		{
 			UV1s.AddUninitialized(VertexCount);
+
 			for (int32 VertIndex = 0; VertIndex < VertexCount; ++VertIndex)
 			{
-				UV1s[VertIndex] = LODMesh.VertexBuffers.StaticMeshVertexBuffer.GetVertexUV(VertIndex, 1);
+				UV1s[VertIndex] = StaticMeshVertexBuffer.GetVertexUV(VertIndex, 1);
 			}
 		}
 	}
diff --git a/Engine/Plugins/Enterprise/GLTFExporter/MigratedFromGit/Source/GLTFExporter/Public/GLTFConversionMesh.h b/Engine/Plugins/Enterprise/GLTFExporter/MigratedFromGit/Source/GLTFExporter/Public/GLTFConversionMesh.h
--- a/Engine/Plugins/Enterprise/GLTFExporter/MigratedFromGit/Source/GLTFExporter/Public/GLTFConversionMesh.h
+++ b/Engine/Plugins/Enterprise/GLTFExporter/MigratedFromGit/Source/GLTFExporter/Public/GLTFConversionMesh.h
@@ -12,7 +12,7 @@ struct GLTFEXPORTER_API FGLTFConversionSection
 
 	TArray<uint32> Indices;
 
-	FGLTFConversionSection(const FString& SectionName, const FStaticMeshLODResources& LODMesh, int32 SectionIndex);
+	FGLTFConversionSection(const FString& SectionName, const FStaticMeshSection& MeshSection, const FIndexArrayView& IndexArray);
 
 	FGLTFJsonAccessorIndex AppendAccessorForIndices(FGLTFContainer& Container) const;
 };
