[GLTFExporter] refactor: rename GLTFContainer -> GLTFBuilder

commit 9f2974d92b2ebcd2e66d44653fb154fa93d3f925
Date: Tue, 28 Apr 2020 11:31:58 +0200

    refactor: rename GLTFContainer -> GLTFBuilder
    (cherry picked from commit d4ba3f2d7f32fc76c45d00a11bd5a950ea1aac1e​)

#skipci

diff --git a/Engine/Plugins/Enterprise/GLTFExporter/MigratedFromGit/Source/GLTFExporter/Private/GLTFBuilder.cpp b/Engine/Plugins/Enterprise/GLTFExporter/MigratedFromGit/Source/GLTFExporter/Private/GLTFBuilder.cpp
new file mode 100644
--- /dev/null
+++ b/Engine/Plugins/Enterprise/GLTFExporter/MigratedFromGit/Source/GLTFExporter/Private/GLTFBuilder.cpp
@@ -0,0 +1,47 @@
+// Copyright Epic Games, Inc. All Rights Reserved.
+
+#include "GLTFBuilder.h"
+#include "GLTFConversionMesh.h"
+#include "Misc/Base64.h"
+
+FGLTFBuilder::FGLTFBuilder()
+	: MergedBufferIndex(JsonRoot.Buffers.AddDefaulted(1))
+{
+}
+
+FGLTFJsonBufferViewIndex FGLTFBuilder::AppendBufferView(const void* RawData, uint64 ByteLength, const FString& Name, EGLTFJsonBufferTarget BufferTarget)
+{
+	FGLTFJsonBufferView BufferView;
+	BufferView.Name = Name;
+	BufferView.Buffer = MergedBufferIndex;
+	BufferView.ByteOffset = MergedBufferData.Num();
+	BufferView.ByteLength = ByteLength;
+	BufferView.Target = BufferTarget;
+
+	MergedBufferData.Append(static_cast<const uint8*>(RawData), ByteLength);
+
+	return JsonRoot.BufferViews.Add(BufferView);
+}
+
+void FGLTFBuilder::UpdateMergedBuffer()
+{
+	FGLTFJsonBuffer& Buffer = JsonRoot.Buffers[MergedBufferIndex];
+	if (Buffer.ByteLength != MergedBufferData.Num())
+	{
+		FString DataBase64 = FBase64::Encode(MergedBufferData.GetData(), MergedBufferData.Num());
+
+		Buffer.URI = TEXT("data:application/octet-stream;base64,") + DataBase64;
+		Buffer.ByteLength = MergedBufferData.Num();
+	}
+}
+
+void FGLTFBuilder::Serialize(FArchive& Archive)
+{
+	UpdateMergedBuffer();
+	JsonRoot.Serialize(&Archive, true);
+}
+
+FGLTFJsonMeshIndex FGLTFBuilder::AppendMesh(const UStaticMesh* StaticMesh, int32 LODIndex)
+{
+	return FGLTFConversionMesh(StaticMesh, LODIndex).AppendMesh(*this);
+}
diff --git a/Engine/Plugins/Enterprise/GLTFExporter/MigratedFromGit/Source/GLTFExporter/Private/GLTFContainer.cpp b/Engine/Plugins/Enterprise/GLTFExporter/MigratedFromGit/Source/GLTFExporter/Private/GLTFContainer.cpp
deleted file mode 100644
diff --git a/Engine/Plugins/Enterprise/GLTFExporter/MigratedFromGit/Source/GLTFExporter/Private/GLTFConversionMesh.cpp b/Engine/Plugins/Enterprise/GLTFExporter/MigratedFromGit/Source/GLTFExporter/Private/GLTFConversionMesh.cpp
--- a/Engine/Plugins/Enterprise/GLTFExporter/MigratedFromGit/Source/GLTFExporter/Private/GLTFConversionMesh.cpp
+++ b/Engine/Plugins/Enterprise/GLTFExporter/MigratedFromGit/Source/GLTFExporter/Private/GLTFConversionMesh.cpp
@@ -16,7 +16,7 @@ FGLTFConversionSection::FGLTFConversionSection(const FString& SectionName, const
 	}
 }
 
-FGLTFJsonAccessorIndex FGLTFConversionSection::AppendAccessorForIndices(FGLTFContainer& Container) const
+FGLTFJsonAccessorIndex FGLTFConversionSection::AppendAccessorForIndices(FGLTFBuilder& Builder) const
 {
 	if (Indices.Num() == 0)
 	{
@@ -24,7 +24,7 @@ FGLTFJsonAccessorIndex FGLTFConversionSection::AppendAccessorForIndices(FGLTFCon
 	}
 
 	const FString AttributeName = Name + TEXT("_Indices");
-	FGLTFJsonBufferViewIndex BufferViewIndex = Container.AppendBufferView(Indices, AttributeName, EGLTFJsonBufferTarget::ElementArrayBuffer);
+	FGLTFJsonBufferViewIndex BufferViewIndex = Builder.AppendBufferView(Indices, AttributeName, EGLTFJsonBufferTarget::ElementArrayBuffer);
 
 	FGLTFJsonAccessor Accessor;
 	Accessor.Name = AttributeName;
@@ -33,7 +33,7 @@ FGLTFJsonAccessorIndex FGLTFConversionSection::AppendAccessorForIndices(FGLTFCon
 	Accessor.Count = Indices.Num();
 	Accessor.Type = EGLTFJsonAccessorType::Scalar;
 
-	return Container.JsonRoot.Accessors.Add(Accessor);
+	return Builder.JsonRoot.Accessors.Add(Accessor);
 }
 
 FGLTFConversionMesh::FGLTFConversionMesh(const UStaticMesh* StaticMesh, int32 LODIndex)
@@ -112,7 +112,7 @@ FGLTFConversionMesh::FGLTFConversionMesh(const UStaticMesh* StaticMesh, int32 LO
 	BoundingBox = StaticMesh->GetBoundingBox();
 }
 
-FGLTFJsonAccessorIndex FGLTFConversionMesh::AppendAccessorForPositions(FGLTFContainer& Container) const
+FGLTFJsonAccessorIndex FGLTFConversionMesh::AppendAccessorForPositions(FGLTFBuilder& Builder) const
 {
 	if (Positions.Num() == 0)
 	{
@@ -120,7 +120,7 @@ FGLTFJsonAccessorIndex FGLTFConversionMesh::AppendAccessorForPositions(FGLTFCont
 	}
 
 	const FString AttributeName = Name + TEXT("_Positions");
-	FGLTFJsonBufferViewIndex BufferViewIndex = Container.AppendBufferView(Positions, AttributeName);
+	FGLTFJsonBufferViewIndex BufferViewIndex = Builder.AppendBufferView(Positions, AttributeName);
 
 	FGLTFJsonAccessor Accessor;
 	Accessor.Name = AttributeName;
@@ -139,10 +139,10 @@ FGLTFJsonAccessorIndex FGLTFConversionMesh::AppendAccessorForPositions(FGLTFCont
 	Accessor.Min[2] = Min.Z;
 	Accessor.MinMaxLength = 3;
 
-	return Container.JsonRoot.Accessors.Add(Accessor);
+	return Builder.JsonRoot.Accessors.Add(Accessor);
 }
 
-FGLTFJsonAccessorIndex FGLTFConversionMesh::AppendAccessorForNormals(FGLTFContainer& Container) const
+FGLTFJsonAccessorIndex FGLTFConversionMesh::AppendAccessorForNormals(FGLTFBuilder& Builder) const
 {
 	if (Normals.Num() == 0)
 	{
@@ -150,7 +150,7 @@ FGLTFJsonAccessorIndex FGLTFConversionMesh::AppendAccessorForNormals(FGLTFContai
 	}
 
 	const FString AttributeName = Name + TEXT("_Normals");
-	FGLTFJsonBufferViewIndex BufferViewIndex = Container.AppendBufferView(Normals, AttributeName);
+	FGLTFJsonBufferViewIndex BufferViewIndex = Builder.AppendBufferView(Normals, AttributeName);
 
 	FGLTFJsonAccessor Accessor;
 	Accessor.Name = AttributeName;
@@ -159,10 +159,10 @@ FGLTFJsonAccessorIndex FGLTFConversionMesh::AppendAccessorForNormals(FGLTFContai
 	Accessor.Count = Normals.Num();
 	Accessor.Type = EGLTFJsonAccessorType::Vec3;
 
-	return Container.JsonRoot.Accessors.Add(Accessor);
+	return Builder.JsonRoot.Accessors.Add(Accessor);
 }
 
-FGLTFJsonAccessorIndex FGLTFConversionMesh::AppendAccessorForColors(FGLTFContainer& Container) const
+FGLTFJsonAccessorIndex FGLTFConversionMesh::AppendAccessorForColors(FGLTFBuilder& Builder) const
 {
 	if (Colors.Num() == 0)
 	{
@@ -170,7 +170,7 @@ FGLTFJsonAccessorIndex FGLTFConversionMesh::AppendAccessorForColors(FGLTFContain
 	}
 
 	const FString AttributeName = Name + TEXT("_Colors");
-	FGLTFJsonBufferViewIndex BufferViewIndex = Container.AppendBufferView(Colors, AttributeName);
+	FGLTFJsonBufferViewIndex BufferViewIndex = Builder.AppendBufferView(Colors, AttributeName);
 
 	FGLTFJsonAccessor Accessor;
 	Accessor.Name = AttributeName;
@@ -179,10 +179,10 @@ FGLTFJsonAccessorIndex FGLTFConversionMesh::AppendAccessorForColors(FGLTFContain
 	Accessor.Count = Colors.Num();
 	Accessor.Type = EGLTFJsonAccessorType::Vec4;
 
-	return Container.JsonRoot.Accessors.Add(Accessor);
+	return Builder.JsonRoot.Accessors.Add(Accessor);
 }
 
-FGLTFJsonAccessorIndex FGLTFConversionMesh::AppendAccessorForTangents(FGLTFContainer& Container) const
+FGLTFJsonAccessorIndex FGLTFConversionMesh::AppendAccessorForTangents(FGLTFBuilder& Builder) const
 {
 	if (Tangents.Num() == 0)
 	{
@@ -190,7 +190,7 @@ FGLTFJsonAccessorIndex FGLTFConversionMesh::AppendAccessorForTangents(FGLTFConta
 	}
 
 	const FString AttributeName = Name + TEXT("_Tangents");
-	FGLTFJsonBufferViewIndex BufferViewIndex = Container.AppendBufferView(Tangents, AttributeName);
+	FGLTFJsonBufferViewIndex BufferViewIndex = Builder.AppendBufferView(Tangents, AttributeName);
 
 	FGLTFJsonAccessor Accessor;
 	Accessor.Name = AttributeName;
@@ -199,10 +199,10 @@ FGLTFJsonAccessorIndex FGLTFConversionMesh::AppendAccessorForTangents(FGLTFConta
 	Accessor.Count = Tangents.Num();
 	Accessor.Type = EGLTFJsonAccessorType::Vec4;
 
-	return Container.JsonRoot.Accessors.Add(Accessor);
+	return Builder.JsonRoot.Accessors.Add(Accessor);
 }
 
-FGLTFJsonAccessorIndex FGLTFConversionMesh::AppendAccessorForUV0s(FGLTFContainer& Container) const
+FGLTFJsonAccessorIndex FGLTFConversionMesh::AppendAccessorForUV0s(FGLTFBuilder& Builder) const
 {
 	if (UV0s.Num() == 0)
 	{
@@ -210,7 +210,7 @@ FGLTFJsonAccessorIndex FGLTFConversionMesh::AppendAccessorForUV0s(FGLTFContainer
 	}
 
 	const FString AttributeName = Name + TEXT("_UV0s");
-	FGLTFJsonBufferViewIndex BufferViewIndex = Container.AppendBufferView(UV0s, AttributeName);
+	FGLTFJsonBufferViewIndex BufferViewIndex = Builder.AppendBufferView(UV0s, AttributeName);
 
 	FGLTFJsonAccessor Accessor;
 	Accessor.Name = AttributeName;
@@ -219,10 +219,10 @@ FGLTFJsonAccessorIndex FGLTFConversionMesh::AppendAccessorForUV0s(FGLTFContainer
 	Accessor.Count = UV0s.Num();
 	Accessor.Type = EGLTFJsonAccessorType::Vec2;
 
-	return Container.JsonRoot.Accessors.Add(Accessor);
+	return Builder.JsonRoot.Accessors.Add(Accessor);
 }
 
-FGLTFJsonAccessorIndex FGLTFConversionMesh::AppendAccessorForUV1s(FGLTFContainer& Container) const
+FGLTFJsonAccessorIndex FGLTFConversionMesh::AppendAccessorForUV1s(FGLTFBuilder& Builder) const
 {
 	if (UV1s.Num() == 0)
 	{
@@ -230,7 +230,7 @@ FGLTFJsonAccessorIndex FGLTFConversionMesh::AppendAccessorForUV1s(FGLTFContainer
 	}
 
 	const FString AttributeName = Name + TEXT("_UV1s");
-	FGLTFJsonBufferViewIndex BufferViewIndex = Container.AppendBufferView(UV1s, AttributeName);
+	FGLTFJsonBufferViewIndex BufferViewIndex = Builder.AppendBufferView(UV1s, AttributeName);
 
 	FGLTFJsonAccessor Accessor;
 	Accessor.Name = AttributeName;
@@ -239,30 +239,30 @@ FGLTFJsonAccessorIndex FGLTFConversionMesh::AppendAccessorForUV1s(FGLTFContainer
 	Accessor.Count = UV1s.Num();
 	Accessor.Type = EGLTFJsonAccessorType::Vec2;
 
-	return Container.JsonRoot.Accessors.Add(Accessor);
+	return Builder.JsonRoot.Accessors.Add(Accessor);
 }
 
-FGLTFJsonMeshIndex FGLTFConversionMesh::AppendMesh(FGLTFContainer& Container) const
+FGLTFJsonMeshIndex FGLTFConversionMesh::AppendMesh(FGLTFBuilder& Builder) const
 {
 	FGLTFJsonMesh Mesh;
 	Mesh.Name = Name;
 
 	FGLTFJsonAttributes Attributes;
-	Attributes.Position = AppendAccessorForPositions(Container);
-	Attributes.Color0 = AppendAccessorForColors(Container);
-	Attributes.Normal = AppendAccessorForNormals(Container);
-	Attributes.Tangent = AppendAccessorForTangents(Container);
-	Attributes.TexCoord0 = AppendAccessorForUV0s(Container);
-	Attributes.TexCoord1 = AppendAccessorForUV1s(Container);
+	Attributes.Position = AppendAccessorForPositions(Builder);
+	Attributes.Color0 = AppendAccessorForColors(Builder);
+	Attributes.Normal = AppendAccessorForNormals(Builder);
+	Attributes.Tangent = AppendAccessorForTangents(Builder);
+	Attributes.TexCoord0 = AppendAccessorForUV0s(Builder);
+	Attributes.TexCoord1 = AppendAccessorForUV1s(Builder);
 
 	for (const FGLTFConversionSection& Section : Sections)
 	{
 		FGLTFJsonPrimitive Primitive;
 		Primitive.Attributes = Attributes;
 
-		Primitive.Indices = Section.AppendAccessorForIndices(Container);
+		Primitive.Indices = Section.AppendAccessorForIndices(Builder);
 		Mesh.Primitives.Add(Primitive);
 	}
 
-	return Container.JsonRoot.Meshes.Add(Mesh);
+	return Builder.JsonRoot.Meshes.Add(Mesh);
 }
diff --git a/Engine/Plugins/Enterprise/GLTFExporter/MigratedFromGit/Source/GLTFExporter/Private/GLTFStaticMeshExporter.cpp b/Engine/Plugins/Enterprise/GLTFExporter/MigratedFromGit/Source/GLTFExporter/Private/GLTFStaticMeshExporter.cpp
--- a/Engine/Plugins/Enterprise/GLTFExporter/MigratedFromGit/Source/GLTFExporter/Private/GLTFStaticMeshExporter.cpp
+++ b/Engine/Plugins/Enterprise/GLTFExporter/MigratedFromGit/Source/GLTFExporter/Private/GLTFStaticMeshExporter.cpp
@@ -2,7 +2,7 @@
 
 #include "GLTFStaticMeshExporter.h"
 #include "GLTFJsonRoot.h"
-#include "GLTFContainer.h"
+#include "GLTFBuilder.h"
 #include "Engine/StaticMesh.h"
 
 UGLTFStaticMeshExporter::UGLTFStaticMeshExporter(const FObjectInitializer& ObjectInitializer)
@@ -21,20 +21,20 @@ bool UGLTFStaticMeshExporter::ExportBinary(UObject* Object, const TCHAR* Type, F
 		return false;
 	}
 
-	FGLTFContainer Container;
+	FGLTFBuilder Builder;
 
-	FGLTFJsonMeshIndex MeshIndex = Container.AppendMesh(StaticMesh, 0);
+	FGLTFJsonMeshIndex MeshIndex = Builder.AppendMesh(StaticMesh, 0);
 
 	FGLTFJsonNode Node;
 	Node.Mesh = MeshIndex;
-	FGLTFJsonNodeIndex NodeIndex = Container.JsonRoot.Nodes.Add(Node);
+	FGLTFJsonNodeIndex NodeIndex = Builder.JsonRoot.Nodes.Add(Node);
 
 	FGLTFJsonScene Scene;
 	Scene.Nodes.Add(NodeIndex);
-	FGLTFJsonSceneIndex SceneIndex = Container.JsonRoot.Scenes.Add(Scene);
+	FGLTFJsonSceneIndex SceneIndex = Builder.JsonRoot.Scenes.Add(Scene);
 
-	Container.JsonRoot.DefaultScene = SceneIndex;
+	Builder.JsonRoot.DefaultScene = SceneIndex;
 
-	Container.Serialize(Archive);
+	Builder.Serialize(Archive);
 	return true;
 }
diff --git a/Engine/Plugins/Enterprise/GLTFExporter/MigratedFromGit/Source/GLTFExporter/Public/GLTFBuilder.h b/Engine/Plugins/Enterprise/GLTFExporter/MigratedFromGit/Source/GLTFExporter/Public/GLTFBuilder.h
new file mode 100644
--- /dev/null
+++ b/Engine/Plugins/Enterprise/GLTFExporter/MigratedFromGit/Source/GLTFExporter/Public/GLTFBuilder.h
@@ -0,0 +1,29 @@
+// Copyright Epic Games, Inc. All Rights Reserved.
+
+#pragma once
+
+#include "GLTFJsonRoot.h"
+#include "Engine.h"
+
+struct GLTFEXPORTER_API FGLTFBuilder
+{
+	FGLTFJsonRoot JsonRoot;
+
+	const FGLTFJsonBufferIndex MergedBufferIndex;
+	TArray<uint8> MergedBufferData;
+
+	FGLTFBuilder();
+
+	template <class ElementType>
+	FGLTFJsonBufferViewIndex AppendBufferView(const TArray<ElementType>& Array, const FString& Name = TEXT(""), EGLTFJsonBufferTarget BufferTarget = EGLTFJsonBufferTarget::ArrayBuffer)
+	{
+		return AppendBufferView(Array.GetData(), Array.Num() * sizeof(ElementType), Name, BufferTarget);
+	}
+
+	FGLTFJsonBufferViewIndex AppendBufferView(const void* RawData, uint64 ByteLength, const FString& Name = TEXT(""), EGLTFJsonBufferTarget BufferTarget = EGLTFJsonBufferTarget::ArrayBuffer);
+
+	FGLTFJsonMeshIndex AppendMesh(const UStaticMesh* StaticMesh, int32 LODIndex);
+
+	void UpdateMergedBuffer();
+	void Serialize(FArchive& Archive);
+};
diff --git a/Engine/Plugins/Enterprise/GLTFExporter/MigratedFromGit/Source/GLTFExporter/Public/GLTFContainer.h b/Engine/Plugins/Enterprise/GLTFExporter/MigratedFromGit/Source/GLTFExporter/Public/GLTFContainer.h
deleted file mode 100644
diff --git a/Engine/Plugins/Enterprise/GLTFExporter/MigratedFromGit/Source/GLTFExporter/Public/GLTFConversionMesh.h b/Engine/Plugins/Enterprise/GLTFExporter/MigratedFromGit/Source/GLTFExporter/Public/GLTFConversionMesh.h
--- a/Engine/Plugins/Enterprise/GLTFExporter/MigratedFromGit/Source/GLTFExporter/Public/GLTFConversionMesh.h
+++ b/Engine/Plugins/Enterprise/GLTFExporter/MigratedFromGit/Source/GLTFExporter/Public/GLTFConversionMesh.h
@@ -3,7 +3,7 @@
 #pragma once
 
 #include "GLTFJsonRoot.h"
-#include "GLTFContainer.h"
+#include "GLTFBuilder.h"
 #include "Engine.h"
 
 struct GLTFEXPORTER_API FGLTFConversionSection
@@ -14,7 +14,7 @@ struct GLTFEXPORTER_API FGLTFConversionSection
 
 	FGLTFConversionSection(const FString& SectionName, const FStaticMeshSection& MeshSection, const FIndexArrayView& IndexArray);
 
-	FGLTFJsonAccessorIndex AppendAccessorForIndices(FGLTFContainer& Container) const;
+	FGLTFJsonAccessorIndex AppendAccessorForIndices(FGLTFBuilder& Builder) const;
 };
 
 struct GLTFEXPORTER_API FGLTFConversionMesh
@@ -34,12 +34,12 @@ struct GLTFEXPORTER_API FGLTFConversionMesh
 
 	FGLTFConversionMesh(const UStaticMesh* StaticMesh, int32 LODIndex);
 
-	FGLTFJsonAccessorIndex AppendAccessorForPositions(FGLTFContainer& Container) const;
-	FGLTFJsonAccessorIndex AppendAccessorForColors(FGLTFContainer& Container) const;
-	FGLTFJsonAccessorIndex AppendAccessorForNormals(FGLTFContainer& Container) const;
-	FGLTFJsonAccessorIndex AppendAccessorForTangents(FGLTFContainer& Container) const;
-	FGLTFJsonAccessorIndex AppendAccessorForUV0s(FGLTFContainer& Container) const;
-	FGLTFJsonAccessorIndex AppendAccessorForUV1s(FGLTFContainer& Container) const;
+	FGLTFJsonAccessorIndex AppendAccessorForPositions(FGLTFBuilder& Builder) const;
+	FGLTFJsonAccessorIndex AppendAccessorForColors(FGLTFBuilder& Builder) const;
+	FGLTFJsonAccessorIndex AppendAccessorForNormals(FGLTFBuilder& Builder) const;
+	FGLTFJsonAccessorIndex AppendAccessorForTangents(FGLTFBuilder& Builder) const;
+	FGLTFJsonAccessorIndex AppendAccessorForUV0s(FGLTFBuilder& Builder) const;
+	FGLTFJsonAccessorIndex AppendAccessorForUV1s(FGLTFBuilder& Builder) const;
 
-	FGLTFJsonMeshIndex AppendMesh(FGLTFContainer& Container) const;
+	FGLTFJsonMeshIndex AppendMesh(FGLTFBuilder& Builder) const;
 };
