[GLTFExporter] refactor: rename GLTFConversion* -> GLTF*Builder

commit 6c90377f77ad69a36d2407eed8e5ee0690f15912
Date: Tue, 28 Apr 2020 11:36:04 +0200

    refactor: rename GLTFConversion* -> GLTF*Builder
    (cherry picked from commit bf4bb558a12dbf5c619c5e916fcf326dda5252b0​)

#skipci

diff --git a/Engine/Plugins/Enterprise/GLTFExporter/MigratedFromGit/Source/GLTFExporter/Private/GLTFBuilder.cpp b/Engine/Plugins/Enterprise/GLTFExporter/MigratedFromGit/Source/GLTFExporter/Private/GLTFBuilder.cpp
--- a/Engine/Plugins/Enterprise/GLTFExporter/MigratedFromGit/Source/GLTFExporter/Private/GLTFBuilder.cpp
+++ b/Engine/Plugins/Enterprise/GLTFExporter/MigratedFromGit/Source/GLTFExporter/Private/GLTFBuilder.cpp
@@ -1,7 +1,7 @@
 // Copyright Epic Games, Inc. All Rights Reserved.
 
 #include "GLTFBuilder.h"
-#include "GLTFConversionMesh.h"
+#include "GLTFMeshBuilder.h"
 #include "Misc/Base64.h"
 
 FGLTFBuilder::FGLTFBuilder()
@@ -43,5 +43,5 @@ void FGLTFBuilder::Serialize(FArchive& Archive)
 
 FGLTFJsonMeshIndex FGLTFBuilder::AppendMesh(const UStaticMesh* StaticMesh, int32 LODIndex)
 {
-	return FGLTFConversionMesh(StaticMesh, LODIndex).AppendMesh(*this);
+	return FGLTFMeshBuilder(StaticMesh, LODIndex).AppendMesh(*this);
 }
diff --git a/Engine/Plugins/Enterprise/GLTFExporter/MigratedFromGit/Source/GLTFExporter/Private/GLTFConversionMesh.cpp b/Engine/Plugins/Enterprise/GLTFExporter/MigratedFromGit/Source/GLTFExporter/Private/GLTFConversionMesh.cpp
deleted file mode 100644
diff --git a/Engine/Plugins/Enterprise/GLTFExporter/MigratedFromGit/Source/GLTFExporter/Private/GLTFMeshBuilder.cpp b/Engine/Plugins/Enterprise/GLTFExporter/MigratedFromGit/Source/GLTFExporter/Private/GLTFMeshBuilder.cpp
new file mode 100644
--- /dev/null
+++ b/Engine/Plugins/Enterprise/GLTFExporter/MigratedFromGit/Source/GLTFExporter/Private/GLTFMeshBuilder.cpp
@@ -0,0 +1,268 @@
+// Copyright Epic Games, Inc. All Rights Reserved.
+
+#include "GLTFMeshBuilder.h"
+#include "GLTFConversionUtilities.h"
+
+FGLTFSectionBuilder::FGLTFSectionBuilder(const FString& SectionName, const FStaticMeshSection& MeshSection, const FIndexArrayView& IndexArray)
+	: Name(SectionName)
+{
+	const uint32 IndexCount = MeshSection.NumTriangles * 3;
+	const uint32 FirstIndex = MeshSection.FirstIndex;
+	Indices.AddUninitialized(IndexCount);
+
+	for (uint32 Index = 0; Index < IndexCount; ++Index)
+	{
+		Indices[Index] = IndexArray[FirstIndex + Index];
+	}
+}
+
+FGLTFJsonAccessorIndex FGLTFSectionBuilder::AppendAccessorForIndices(FGLTFBuilder& Builder) const
+{
+	if (Indices.Num() == 0)
+	{
+		return INDEX_NONE;
+	}
+
+	const FString AttributeName = Name + TEXT("_Indices");
+	FGLTFJsonBufferViewIndex BufferViewIndex = Builder.AppendBufferView(Indices, AttributeName, EGLTFJsonBufferTarget::ElementArrayBuffer);
+
+	FGLTFJsonAccessor Accessor;
+	Accessor.Name = AttributeName;
+	Accessor.BufferView = BufferViewIndex;
+	Accessor.ComponentType = EGLTFJsonComponentType::U32;
+	Accessor.Count = Indices.Num();
+	Accessor.Type = EGLTFJsonAccessorType::Scalar;
+
+	return Builder.JsonRoot.Accessors.Add(Accessor);
+}
+
+FGLTFMeshBuilder::FGLTFMeshBuilder(const UStaticMesh* StaticMesh, int32 LODIndex)
+{
+	StaticMesh->GetName(Name);
+
+	const FStaticMeshLODResources& LODMesh = StaticMesh->GetLODForExport(LODIndex);
+
+	const FIndexArrayView IndexArray = LODMesh.IndexBuffer.GetArrayView();
+	const int32 SectionCount = LODMesh.Sections.Num();
+
+	Sections.Reserve(SectionCount);
+
+	for (int32 SectionIndex = 0; SectionIndex < SectionCount; ++SectionIndex)
+	{
+		FString SectionName = Name + TEXT("_Section") + FString::FromInt(SectionIndex);
+		Sections.Emplace(SectionName, LODMesh.Sections[SectionIndex], IndexArray);
+	}
+
+	const FPositionVertexBuffer& PositionVertexBuffer = LODMesh.VertexBuffers.PositionVertexBuffer;
+	const int32 PositionCount = PositionVertexBuffer.GetNumVertices();
+
+	Positions.AddUninitialized(PositionCount);
+
+	for (int32 PosIndex = 0; PosIndex < PositionCount; ++PosIndex)
+	{
+		Positions[PosIndex] = ConvertPosition(PositionVertexBuffer.VertexPosition(PosIndex));
+	}
+
+	const FColorVertexBuffer& ColorVertexBuffer = LODMesh.VertexBuffers.ColorVertexBuffer;
+	const int32 ColorCount = ColorVertexBuffer.GetNumVertices();
+
+	if (ColorCount > 0)
+	{
+		Colors.AddUninitialized(ColorCount);
+
+		for (int32 ColorIndex = 0; ColorIndex < ColorCount; ++ColorIndex)
+		{
+			Colors[ColorIndex] = ConvertColor(ColorVertexBuffer.VertexColor(ColorIndex));
+		}
+	}
+
+	const FStaticMeshVertexBuffer& StaticMeshVertexBuffer = LODMesh.VertexBuffers.StaticMeshVertexBuffer;
+	const int32 VertexCount = StaticMeshVertexBuffer.GetNumVertices();
+
+	Normals.AddUninitialized(VertexCount);
+	Tangents.AddUninitialized(VertexCount);
+
+	for (int32 VertIndex = 0; VertIndex < VertexCount; ++VertIndex)
+	{
+		Normals[VertIndex] = ConvertVector(StaticMeshVertexBuffer.VertexTangentZ(VertIndex));
+		Tangents[VertIndex] = ConvertTangent(StaticMeshVertexBuffer.VertexTangentX(VertIndex));
+	}
+
+	const int32 UVCount = StaticMeshVertexBuffer.GetNumTexCoords();
+	if (UVCount >= 1)
+	{
+		UV0s.AddUninitialized(VertexCount);
+
+		for (int32 VertIndex = 0; VertIndex < VertexCount; ++VertIndex)
+		{
+			UV0s[VertIndex] = StaticMeshVertexBuffer.GetVertexUV(VertIndex, 0);
+		}
+
+		if (UVCount >= 2)
+		{
+			UV1s.AddUninitialized(VertexCount);
+
+			for (int32 VertIndex = 0; VertIndex < VertexCount; ++VertIndex)
+			{
+				UV1s[VertIndex] = StaticMeshVertexBuffer.GetVertexUV(VertIndex, 1);
+			}
+		}
+	}
+
+	BoundingBox = StaticMesh->GetBoundingBox();
+}
+
+FGLTFJsonAccessorIndex FGLTFMeshBuilder::AppendAccessorForPositions(FGLTFBuilder& Builder) const
+{
+	if (Positions.Num() == 0)
+	{
+		return INDEX_NONE;
+	}
+
+	const FString AttributeName = Name + TEXT("_Positions");
+	FGLTFJsonBufferViewIndex BufferViewIndex = Builder.AppendBufferView(Positions, AttributeName);
+
+	FGLTFJsonAccessor Accessor;
+	Accessor.Name = AttributeName;
+	Accessor.BufferView = BufferViewIndex;
+	Accessor.ComponentType = EGLTFJsonComponentType::F32;
+	Accessor.Count = Positions.Num();
+	Accessor.Type = EGLTFJsonAccessorType::Vec3;
+
+	FVector Max = ConvertSize(BoundingBox.Max);
+	FVector Min = ConvertSize(BoundingBox.Min);
+	Accessor.Max[0] = Max.X;
+	Accessor.Max[1] = Max.Y;
+	Accessor.Max[2] = Max.Z;
+	Accessor.Min[0] = Min.X;
+	Accessor.Min[1] = Min.Y;
+	Accessor.Min[2] = Min.Z;
+	Accessor.MinMaxLength = 3;
+
+	return Builder.JsonRoot.Accessors.Add(Accessor);
+}
+
+FGLTFJsonAccessorIndex FGLTFMeshBuilder::AppendAccessorForNormals(FGLTFBuilder& Builder) const
+{
+	if (Normals.Num() == 0)
+	{
+		return INDEX_NONE;
+	}
+
+	const FString AttributeName = Name + TEXT("_Normals");
+	FGLTFJsonBufferViewIndex BufferViewIndex = Builder.AppendBufferView(Normals, AttributeName);
+
+	FGLTFJsonAccessor Accessor;
+	Accessor.Name = AttributeName;
+	Accessor.BufferView = BufferViewIndex;
+	Accessor.ComponentType = EGLTFJsonComponentType::F32;
+	Accessor.Count = Normals.Num();
+	Accessor.Type = EGLTFJsonAccessorType::Vec3;
+
+	return Builder.JsonRoot.Accessors.Add(Accessor);
+}
+
+FGLTFJsonAccessorIndex FGLTFMeshBuilder::AppendAccessorForColors(FGLTFBuilder& Builder) const
+{
+	if (Colors.Num() == 0)
+	{
+		return INDEX_NONE;
+	}
+
+	const FString AttributeName = Name + TEXT("_Colors");
+	FGLTFJsonBufferViewIndex BufferViewIndex = Builder.AppendBufferView(Colors, AttributeName);
+
+	FGLTFJsonAccessor Accessor;
+	Accessor.Name = AttributeName;
+	Accessor.BufferView = BufferViewIndex;
+	Accessor.ComponentType = EGLTFJsonComponentType::U8;
+	Accessor.Count = Colors.Num();
+	Accessor.Type = EGLTFJsonAccessorType::Vec4;
+
+	return Builder.JsonRoot.Accessors.Add(Accessor);
+}
+
+FGLTFJsonAccessorIndex FGLTFMeshBuilder::AppendAccessorForTangents(FGLTFBuilder& Builder) const
+{
+	if (Tangents.Num() == 0)
+	{
+		return INDEX_NONE;
+	}
+
+	const FString AttributeName = Name + TEXT("_Tangents");
+	FGLTFJsonBufferViewIndex BufferViewIndex = Builder.AppendBufferView(Tangents, AttributeName);
+
+	FGLTFJsonAccessor Accessor;
+	Accessor.Name = AttributeName;
+	Accessor.BufferView = BufferViewIndex;
+	Accessor.ComponentType = EGLTFJsonComponentType::F32;
+	Accessor.Count = Tangents.Num();
+	Accessor.Type = EGLTFJsonAccessorType::Vec4;
+
+	return Builder.JsonRoot.Accessors.Add(Accessor);
+}
+
+FGLTFJsonAccessorIndex FGLTFMeshBuilder::AppendAccessorForUV0s(FGLTFBuilder& Builder) const
+{
+	if (UV0s.Num() == 0)
+	{
+		return INDEX_NONE;
+	}
+
+	const FString AttributeName = Name + TEXT("_UV0s");
+	FGLTFJsonBufferViewIndex BufferViewIndex = Builder.AppendBufferView(UV0s, AttributeName);
+
+	FGLTFJsonAccessor Accessor;
+	Accessor.Name = AttributeName;
+	Accessor.BufferView = BufferViewIndex;
+	Accessor.ComponentType = EGLTFJsonComponentType::F32;
+	Accessor.Count = UV0s.Num();
+	Accessor.Type = EGLTFJsonAccessorType::Vec2;
+
+	return Builder.JsonRoot.Accessors.Add(Accessor);
+}
+
+FGLTFJsonAccessorIndex FGLTFMeshBuilder::AppendAccessorForUV1s(FGLTFBuilder& Builder) const
+{
+	if (UV1s.Num() == 0)
+	{
+		return INDEX_NONE;
+	}
+
+	const FString AttributeName = Name + TEXT("_UV1s");
+	FGLTFJsonBufferViewIndex BufferViewIndex = Builder.AppendBufferView(UV1s, AttributeName);
+
+	FGLTFJsonAccessor Accessor;
+	Accessor.Name = AttributeName;
+	Accessor.BufferView = BufferViewIndex;
+	Accessor.ComponentType = EGLTFJsonComponentType::F32;
+	Accessor.Count = UV1s.Num();
+	Accessor.Type = EGLTFJsonAccessorType::Vec2;
+
+	return Builder.JsonRoot.Accessors.Add(Accessor);
+}
+
+FGLTFJsonMeshIndex FGLTFMeshBuilder::AppendMesh(FGLTFBuilder& Builder) const
+{
+	FGLTFJsonMesh Mesh;
+	Mesh.Name = Name;
+
+	FGLTFJsonAttributes Attributes;
+	Attributes.Position = AppendAccessorForPositions(Builder);
+	Attributes.Color0 = AppendAccessorForColors(Builder);
+	Attributes.Normal = AppendAccessorForNormals(Builder);
+	Attributes.Tangent = AppendAccessorForTangents(Builder);
+	Attributes.TexCoord0 = AppendAccessorForUV0s(Builder);
+	Attributes.TexCoord1 = AppendAccessorForUV1s(Builder);
+
+	for (const FGLTFSectionBuilder& Section : Sections)
+	{
+		FGLTFJsonPrimitive Primitive;
+		Primitive.Attributes = Attributes;
+
+		Primitive.Indices = Section.AppendAccessorForIndices(Builder);
+		Mesh.Primitives.Add(Primitive);
+	}
+
+	return Builder.JsonRoot.Meshes.Add(Mesh);
+}
diff --git a/Engine/Plugins/Enterprise/GLTFExporter/MigratedFromGit/Source/GLTFExporter/Public/GLTFConversionMesh.h b/Engine/Plugins/Enterprise/GLTFExporter/MigratedFromGit/Source/GLTFExporter/Public/GLTFConversionMesh.h
deleted file mode 100644
diff --git a/Engine/Plugins/Enterprise/GLTFExporter/MigratedFromGit/Source/GLTFExporter/Public/GLTFMeshBuilder.h b/Engine/Plugins/Enterprise/GLTFExporter/MigratedFromGit/Source/GLTFExporter/Public/GLTFMeshBuilder.h
new file mode 100644
--- /dev/null
+++ b/Engine/Plugins/Enterprise/GLTFExporter/MigratedFromGit/Source/GLTFExporter/Public/GLTFMeshBuilder.h
@@ -0,0 +1,45 @@
+// Copyright Epic Games, Inc. All Rights Reserved.
+
+#pragma once
+
+#include "GLTFJsonRoot.h"
+#include "GLTFBuilder.h"
+#include "Engine.h"
+
+struct GLTFEXPORTER_API FGLTFSectionBuilder
+{
+	FString Name;
+
+	TArray<uint32> Indices;
+
+	FGLTFSectionBuilder(const FString& SectionName, const FStaticMeshSection& MeshSection, const FIndexArrayView& IndexArray);
+
+	FGLTFJsonAccessorIndex AppendAccessorForIndices(FGLTFBuilder& Builder) const;
+};
+
+struct GLTFEXPORTER_API FGLTFMeshBuilder
+{
+	FString Name;
+
+	TArray<FGLTFSectionBuilder> Sections;
+
+	TArray<FVector>   Positions;
+	TArray<FColor>    Colors;
+	TArray<FVector>   Normals;
+	TArray<FVector4>  Tangents;
+	TArray<FVector2D> UV0s;
+	TArray<FVector2D> UV1s;
+
+	FBox BoundingBox;
+
+	FGLTFMeshBuilder(const UStaticMesh* StaticMesh, int32 LODIndex);
+
+	FGLTFJsonAccessorIndex AppendAccessorForPositions(FGLTFBuilder& Builder) const;
+	FGLTFJsonAccessorIndex AppendAccessorForColors(FGLTFBuilder& Builder) const;
+	FGLTFJsonAccessorIndex AppendAccessorForNormals(FGLTFBuilder& Builder) const;
+	FGLTFJsonAccessorIndex AppendAccessorForTangents(FGLTFBuilder& Builder) const;
+	FGLTFJsonAccessorIndex AppendAccessorForUV0s(FGLTFBuilder& Builder) const;
+	FGLTFJsonAccessorIndex AppendAccessorForUV1s(FGLTFBuilder& Builder) const;
+
+	FGLTFJsonMeshIndex AppendMesh(FGLTFBuilder& Builder) const;
+};
