[GLTFExporter] refactor: rename *Builder::Append* -> *Builder::Add*

commit 2a3c00530f459ad4d7202213e50b71fca8cc8655
Date: Tue, 28 Apr 2020 11:39:13 +0200

    refactor: rename *Builder::Append* -> *Builder::Add*
    (cherry picked from commit 8ef3350bfbd499a3c1a8a54845a74b2caafc7f45​)

#skipci

diff --git a/Engine/Plugins/Enterprise/GLTFExporter/MigratedFromGit/Source/GLTFExporter/Private/GLTFBuilder.cpp b/Engine/Plugins/Enterprise/GLTFExporter/MigratedFromGit/Source/GLTFExporter/Private/GLTFBuilder.cpp
--- a/Engine/Plugins/Enterprise/GLTFExporter/MigratedFromGit/Source/GLTFExporter/Private/GLTFBuilder.cpp
+++ b/Engine/Plugins/Enterprise/GLTFExporter/MigratedFromGit/Source/GLTFExporter/Private/GLTFBuilder.cpp
@@ -9,7 +9,7 @@ FGLTFBuilder::FGLTFBuilder()
 {
 }
 
-FGLTFJsonBufferViewIndex FGLTFBuilder::AppendBufferView(const void* RawData, uint64 ByteLength, const FString& Name, EGLTFJsonBufferTarget BufferTarget)
+FGLTFJsonBufferViewIndex FGLTFBuilder::AddBufferView(const void* RawData, uint64 ByteLength, const FString& Name, EGLTFJsonBufferTarget BufferTarget)
 {
 	FGLTFJsonBufferView BufferView;
 	BufferView.Name = Name;
@@ -41,7 +41,7 @@ void FGLTFBuilder::Serialize(FArchive& Archive)
 	JsonRoot.Serialize(&Archive, true);
 }
 
-FGLTFJsonMeshIndex FGLTFBuilder::AppendMesh(const UStaticMesh* StaticMesh, int32 LODIndex)
+FGLTFJsonMeshIndex FGLTFBuilder::AddMesh(const UStaticMesh* StaticMesh, int32 LODIndex)
 {
-	return FGLTFMeshBuilder(StaticMesh, LODIndex).AppendMesh(*this);
+	return FGLTFMeshBuilder(StaticMesh, LODIndex).AddMesh(*this);
 }
diff --git a/Engine/Plugins/Enterprise/GLTFExporter/MigratedFromGit/Source/GLTFExporter/Private/GLTFMeshBuilder.cpp b/Engine/Plugins/Enterprise/GLTFExporter/MigratedFromGit/Source/GLTFExporter/Private/GLTFMeshBuilder.cpp
--- a/Engine/Plugins/Enterprise/GLTFExporter/MigratedFromGit/Source/GLTFExporter/Private/GLTFMeshBuilder.cpp
+++ b/Engine/Plugins/Enterprise/GLTFExporter/MigratedFromGit/Source/GLTFExporter/Private/GLTFMeshBuilder.cpp
@@ -16,7 +16,7 @@ FGLTFSectionBuilder::FGLTFSectionBuilder(const FString& SectionName, const FStat
 	}
 }
 
-FGLTFJsonAccessorIndex FGLTFSectionBuilder::AppendAccessorForIndices(FGLTFBuilder& Builder) const
+FGLTFJsonAccessorIndex FGLTFSectionBuilder::AddAccessorForIndices(FGLTFBuilder& Builder) const
 {
 	if (Indices.Num() == 0)
 	{
@@ -24,7 +24,7 @@ FGLTFJsonAccessorIndex FGLTFSectionBuilder::AppendAccessorForIndices(FGLTFBuilde
 	}
 
 	const FString AttributeName = Name + TEXT("_Indices");
-	FGLTFJsonBufferViewIndex BufferViewIndex = Builder.AppendBufferView(Indices, AttributeName, EGLTFJsonBufferTarget::ElementArrayBuffer);
+	FGLTFJsonBufferViewIndex BufferViewIndex = Builder.AddBufferView(Indices, AttributeName, EGLTFJsonBufferTarget::ElementArrayBuffer);
 
 	FGLTFJsonAccessor Accessor;
 	Accessor.Name = AttributeName;
@@ -112,7 +112,7 @@ FGLTFMeshBuilder::FGLTFMeshBuilder(const UStaticMesh* StaticMesh, int32 LODIndex
 	BoundingBox = StaticMesh->GetBoundingBox();
 }
 
-FGLTFJsonAccessorIndex FGLTFMeshBuilder::AppendAccessorForPositions(FGLTFBuilder& Builder) const
+FGLTFJsonAccessorIndex FGLTFMeshBuilder::AddAccessorForPositions(FGLTFBuilder& Builder) const
 {
 	if (Positions.Num() == 0)
 	{
@@ -120,7 +120,7 @@ FGLTFJsonAccessorIndex FGLTFMeshBuilder::AppendAccessorForPositions(FGLTFBuilder
 	}
 
 	const FString AttributeName = Name + TEXT("_Positions");
-	FGLTFJsonBufferViewIndex BufferViewIndex = Builder.AppendBufferView(Positions, AttributeName);
+	FGLTFJsonBufferViewIndex BufferViewIndex = Builder.AddBufferView(Positions, AttributeName);
 
 	FGLTFJsonAccessor Accessor;
 	Accessor.Name = AttributeName;
@@ -142,7 +142,7 @@ FGLTFJsonAccessorIndex FGLTFMeshBuilder::AppendAccessorForPositions(FGLTFBuilder
 	return Builder.JsonRoot.Accessors.Add(Accessor);
 }
 
-FGLTFJsonAccessorIndex FGLTFMeshBuilder::AppendAccessorForNormals(FGLTFBuilder& Builder) const
+FGLTFJsonAccessorIndex FGLTFMeshBuilder::AddAccessorForNormals(FGLTFBuilder& Builder) const
 {
 	if (Normals.Num() == 0)
 	{
@@ -150,7 +150,7 @@ FGLTFJsonAccessorIndex FGLTFMeshBuilder::AppendAccessorForNormals(FGLTFBuilder&
 	}
 
 	const FString AttributeName = Name + TEXT("_Normals");
-	FGLTFJsonBufferViewIndex BufferViewIndex = Builder.AppendBufferView(Normals, AttributeName);
+	FGLTFJsonBufferViewIndex BufferViewIndex = Builder.AddBufferView(Normals, AttributeName);
 
 	FGLTFJsonAccessor Accessor;
 	Accessor.Name = AttributeName;
@@ -162,7 +162,7 @@ FGLTFJsonAccessorIndex FGLTFMeshBuilder::AppendAccessorForNormals(FGLTFBuilder&
 	return Builder.JsonRoot.Accessors.Add(Accessor);
 }
 
-FGLTFJsonAccessorIndex FGLTFMeshBuilder::AppendAccessorForColors(FGLTFBuilder& Builder) const
+FGLTFJsonAccessorIndex FGLTFMeshBuilder::AddAccessorForColors(FGLTFBuilder& Builder) const
 {
 	if (Colors.Num() == 0)
 	{
@@ -170,7 +170,7 @@ FGLTFJsonAccessorIndex FGLTFMeshBuilder::AppendAccessorForColors(FGLTFBuilder& B
 	}
 
 	const FString AttributeName = Name + TEXT("_Colors");
-	FGLTFJsonBufferViewIndex BufferViewIndex = Builder.AppendBufferView(Colors, AttributeName);
+	FGLTFJsonBufferViewIndex BufferViewIndex = Builder.AddBufferView(Colors, AttributeName);
 
 	FGLTFJsonAccessor Accessor;
 	Accessor.Name = AttributeName;
@@ -182,7 +182,7 @@ FGLTFJsonAccessorIndex FGLTFMeshBuilder::AppendAccessorForColors(FGLTFBuilder& B
 	return Builder.JsonRoot.Accessors.Add(Accessor);
 }
 
-FGLTFJsonAccessorIndex FGLTFMeshBuilder::AppendAccessorForTangents(FGLTFBuilder& Builder) const
+FGLTFJsonAccessorIndex FGLTFMeshBuilder::AddAccessorForTangents(FGLTFBuilder& Builder) const
 {
 	if (Tangents.Num() == 0)
 	{
@@ -190,7 +190,7 @@ FGLTFJsonAccessorIndex FGLTFMeshBuilder::AppendAccessorForTangents(FGLTFBuilder&
 	}
 
 	const FString AttributeName = Name + TEXT("_Tangents");
-	FGLTFJsonBufferViewIndex BufferViewIndex = Builder.AppendBufferView(Tangents, AttributeName);
+	FGLTFJsonBufferViewIndex BufferViewIndex = Builder.AddBufferView(Tangents, AttributeName);
 
 	FGLTFJsonAccessor Accessor;
 	Accessor.Name = AttributeName;
@@ -202,7 +202,7 @@ FGLTFJsonAccessorIndex FGLTFMeshBuilder::AppendAccessorForTangents(FGLTFBuilder&
 	return Builder.JsonRoot.Accessors.Add(Accessor);
 }
 
-FGLTFJsonAccessorIndex FGLTFMeshBuilder::AppendAccessorForUV0s(FGLTFBuilder& Builder) const
+FGLTFJsonAccessorIndex FGLTFMeshBuilder::AddAccessorForUV0s(FGLTFBuilder& Builder) const
 {
 	if (UV0s.Num() == 0)
 	{
@@ -210,7 +210,7 @@ FGLTFJsonAccessorIndex FGLTFMeshBuilder::AppendAccessorForUV0s(FGLTFBuilder& Bui
 	}
 
 	const FString AttributeName = Name + TEXT("_UV0s");
-	FGLTFJsonBufferViewIndex BufferViewIndex = Builder.AppendBufferView(UV0s, AttributeName);
+	FGLTFJsonBufferViewIndex BufferViewIndex = Builder.AddBufferView(UV0s, AttributeName);
 
 	FGLTFJsonAccessor Accessor;
 	Accessor.Name = AttributeName;
@@ -222,7 +222,7 @@ FGLTFJsonAccessorIndex FGLTFMeshBuilder::AppendAccessorForUV0s(FGLTFBuilder& Bui
 	return Builder.JsonRoot.Accessors.Add(Accessor);
 }
 
-FGLTFJsonAccessorIndex FGLTFMeshBuilder::AppendAccessorForUV1s(FGLTFBuilder& Builder) const
+FGLTFJsonAccessorIndex FGLTFMeshBuilder::AddAccessorForUV1s(FGLTFBuilder& Builder) const
 {
 	if (UV1s.Num() == 0)
 	{
@@ -230,7 +230,7 @@ FGLTFJsonAccessorIndex FGLTFMeshBuilder::AppendAccessorForUV1s(FGLTFBuilder& Bui
 	}
 
 	const FString AttributeName = Name + TEXT("_UV1s");
-	FGLTFJsonBufferViewIndex BufferViewIndex = Builder.AppendBufferView(UV1s, AttributeName);
+	FGLTFJsonBufferViewIndex BufferViewIndex = Builder.AddBufferView(UV1s, AttributeName);
 
 	FGLTFJsonAccessor Accessor;
 	Accessor.Name = AttributeName;
@@ -242,25 +242,25 @@ FGLTFJsonAccessorIndex FGLTFMeshBuilder::AppendAccessorForUV1s(FGLTFBuilder& Bui
 	return Builder.JsonRoot.Accessors.Add(Accessor);
 }
 
-FGLTFJsonMeshIndex FGLTFMeshBuilder::AppendMesh(FGLTFBuilder& Builder) const
+FGLTFJsonMeshIndex FGLTFMeshBuilder::AddMesh(FGLTFBuilder& Builder) const
 {
 	FGLTFJsonMesh Mesh;
 	Mesh.Name = Name;
 
 	FGLTFJsonAttributes Attributes;
-	Attributes.Position = AppendAccessorForPositions(Builder);
-	Attributes.Color0 = AppendAccessorForColors(Builder);
-	Attributes.Normal = AppendAccessorForNormals(Builder);
-	Attributes.Tangent = AppendAccessorForTangents(Builder);
-	Attributes.TexCoord0 = AppendAccessorForUV0s(Builder);
-	Attributes.TexCoord1 = AppendAccessorForUV1s(Builder);
+	Attributes.Position = AddAccessorForPositions(Builder);
+	Attributes.Color0 = AddAccessorForColors(Builder);
+	Attributes.Normal = AddAccessorForNormals(Builder);
+	Attributes.Tangent = AddAccessorForTangents(Builder);
+	Attributes.TexCoord0 = AddAccessorForUV0s(Builder);
+	Attributes.TexCoord1 = AddAccessorForUV1s(Builder);
 
 	for (const FGLTFSectionBuilder& Section : Sections)
 	{
 		FGLTFJsonPrimitive Primitive;
 		Primitive.Attributes = Attributes;
 
-		Primitive.Indices = Section.AppendAccessorForIndices(Builder);
+		Primitive.Indices = Section.AddAccessorForIndices(Builder);
 		Mesh.Primitives.Add(Primitive);
 	}
 
diff --git a/Engine/Plugins/Enterprise/GLTFExporter/MigratedFromGit/Source/GLTFExporter/Private/GLTFStaticMeshExporter.cpp b/Engine/Plugins/Enterprise/GLTFExporter/MigratedFromGit/Source/GLTFExporter/Private/GLTFStaticMeshExporter.cpp
--- a/Engine/Plugins/Enterprise/GLTFExporter/MigratedFromGit/Source/GLTFExporter/Private/GLTFStaticMeshExporter.cpp
+++ b/Engine/Plugins/Enterprise/GLTFExporter/MigratedFromGit/Source/GLTFExporter/Private/GLTFStaticMeshExporter.cpp
@@ -23,7 +23,7 @@ bool UGLTFStaticMeshExporter::ExportBinary(UObject* Object, const TCHAR* Type, F
 
 	FGLTFBuilder Builder;
 
-	FGLTFJsonMeshIndex MeshIndex = Builder.AppendMesh(StaticMesh, 0);
+	FGLTFJsonMeshIndex MeshIndex = Builder.AddMesh(StaticMesh, 0);
 
 	FGLTFJsonNode Node;
 	Node.Mesh = MeshIndex;
diff --git a/Engine/Plugins/Enterprise/GLTFExporter/MigratedFromGit/Source/GLTFExporter/Public/GLTFBuilder.h b/Engine/Plugins/Enterprise/GLTFExporter/MigratedFromGit/Source/GLTFExporter/Public/GLTFBuilder.h
--- a/Engine/Plugins/Enterprise/GLTFExporter/MigratedFromGit/Source/GLTFExporter/Public/GLTFBuilder.h
+++ b/Engine/Plugins/Enterprise/GLTFExporter/MigratedFromGit/Source/GLTFExporter/Public/GLTFBuilder.h
@@ -15,14 +15,14 @@ struct GLTFEXPORTER_API FGLTFBuilder
 	FGLTFBuilder();
 
 	template <class ElementType>
-	FGLTFJsonBufferViewIndex AppendBufferView(const TArray<ElementType>& Array, const FString& Name = TEXT(""), EGLTFJsonBufferTarget BufferTarget = EGLTFJsonBufferTarget::ArrayBuffer)
+	FGLTFJsonBufferViewIndex AddBufferView(const TArray<ElementType>& Array, const FString& Name = TEXT(""), EGLTFJsonBufferTarget BufferTarget = EGLTFJsonBufferTarget::ArrayBuffer)
 	{
-		return AppendBufferView(Array.GetData(), Array.Num() * sizeof(ElementType), Name, BufferTarget);
+		return AddBufferView(Array.GetData(), Array.Num() * sizeof(ElementType), Name, BufferTarget);
 	}
 
-	FGLTFJsonBufferViewIndex AppendBufferView(const void* RawData, uint64 ByteLength, const FString& Name = TEXT(""), EGLTFJsonBufferTarget BufferTarget = EGLTFJsonBufferTarget::ArrayBuffer);
+	FGLTFJsonBufferViewIndex AddBufferView(const void* RawData, uint64 ByteLength, const FString& Name = TEXT(""), EGLTFJsonBufferTarget BufferTarget = EGLTFJsonBufferTarget::ArrayBuffer);
 
-	FGLTFJsonMeshIndex AppendMesh(const UStaticMesh* StaticMesh, int32 LODIndex);
+	FGLTFJsonMeshIndex AddMesh(const UStaticMesh* StaticMesh, int32 LODIndex);
 
 	void UpdateMergedBuffer();
 	void Serialize(FArchive& Archive);
diff --git a/Engine/Plugins/Enterprise/GLTFExporter/MigratedFromGit/Source/GLTFExporter/Public/GLTFMeshBuilder.h b/Engine/Plugins/Enterprise/GLTFExporter/MigratedFromGit/Source/GLTFExporter/Public/GLTFMeshBuilder.h
--- a/Engine/Plugins/Enterprise/GLTFExporter/MigratedFromGit/Source/GLTFExporter/Public/GLTFMeshBuilder.h
+++ b/Engine/Plugins/Enterprise/GLTFExporter/MigratedFromGit/Source/GLTFExporter/Public/GLTFMeshBuilder.h
@@ -14,7 +14,7 @@ struct GLTFEXPORTER_API FGLTFSectionBuilder
 
 	FGLTFSectionBuilder(const FString& SectionName, const FStaticMeshSection& MeshSection, const FIndexArrayView& IndexArray);
 
-	FGLTFJsonAccessorIndex AppendAccessorForIndices(FGLTFBuilder& Builder) const;
+	FGLTFJsonAccessorIndex AddAccessorForIndices(FGLTFBuilder& Builder) const;
 };
 
 struct GLTFEXPORTER_API FGLTFMeshBuilder
@@ -34,12 +34,12 @@ struct GLTFEXPORTER_API FGLTFMeshBuilder
 
 	FGLTFMeshBuilder(const UStaticMesh* StaticMesh, int32 LODIndex);
 
-	FGLTFJsonAccessorIndex AppendAccessorForPositions(FGLTFBuilder& Builder) const;
-	FGLTFJsonAccessorIndex AppendAccessorForColors(FGLTFBuilder& Builder) const;
-	FGLTFJsonAccessorIndex AppendAccessorForNormals(FGLTFBuilder& Builder) const;
-	FGLTFJsonAccessorIndex AppendAccessorForTangents(FGLTFBuilder& Builder) const;
-	FGLTFJsonAccessorIndex AppendAccessorForUV0s(FGLTFBuilder& Builder) const;
-	FGLTFJsonAccessorIndex AppendAccessorForUV1s(FGLTFBuilder& Builder) const;
+	FGLTFJsonAccessorIndex AddAccessorForPositions(FGLTFBuilder& Builder) const;
+	FGLTFJsonAccessorIndex AddAccessorForColors(FGLTFBuilder& Builder) const;
+	FGLTFJsonAccessorIndex AddAccessorForNormals(FGLTFBuilder& Builder) const;
+	FGLTFJsonAccessorIndex AddAccessorForTangents(FGLTFBuilder& Builder) const;
+	FGLTFJsonAccessorIndex AddAccessorForUV0s(FGLTFBuilder& Builder) const;
+	FGLTFJsonAccessorIndex AddAccessorForUV1s(FGLTFBuilder& Builder) const;
 
-	FGLTFJsonMeshIndex AppendMesh(FGLTFBuilder& Builder) const;
+	FGLTFJsonMeshIndex AddMesh(FGLTFBuilder& Builder) const;
 };
