[GLTFExporter] refactor: add GLTFBufferBuilder

commit e9fa64138852d7081f54b436132f9158bb5e0089
Date: Tue, 28 Apr 2020 13:36:24 +0200

    refactor: add GLTFBufferBuilder
    (cherry picked from commit 3de571db2a507ddb47839d2d15254ef456a83e03​)

#skipci

diff --git a/Engine/Plugins/Enterprise/GLTFExporter/MigratedFromGit/Source/GLTFExporter/Private/GLTFBufferBuilder.cpp b/Engine/Plugins/Enterprise/GLTFExporter/MigratedFromGit/Source/GLTFExporter/Private/GLTFBufferBuilder.cpp
new file mode 100644
--- /dev/null
+++ b/Engine/Plugins/Enterprise/GLTFExporter/MigratedFromGit/Source/GLTFExporter/Private/GLTFBufferBuilder.cpp
@@ -0,0 +1,36 @@
+// Copyright Epic Games, Inc. All Rights Reserved.
+
+#include "GLTFBufferBuilder.h"
+#include "GLTFMeshBuilder.h"
+#include "Misc/Base64.h"
+
+FGLTFBufferBuilder::FGLTFBufferBuilder(FGLTFJsonRoot& JsonRoot)
+	: MergedBufferIndex(JsonRoot.Buffers.AddDefaulted(1))
+{
+}
+
+FGLTFJsonBufferViewIndex FGLTFBufferBuilder::AddBufferView(const void* RawData, uint64 ByteLength, const FString& Name, EGLTFJsonBufferTarget BufferTarget)
+{
+	FGLTFJsonBufferView BufferView;
+	BufferView.Name = Name;
+	BufferView.Buffer = MergedBufferIndex;
+	BufferView.ByteOffset = MergedBufferData.Num();
+	BufferView.ByteLength = ByteLength;
+	BufferView.Target = BufferTarget;
+
+	MergedBufferData.Append(static_cast<const uint8*>(RawData), ByteLength);
+
+	return JsonRoot.BufferViews.Add(BufferView);
+}
+
+void FGLTFBufferBuilder::UpdateMergedBuffer()
+{
+	FGLTFJsonBuffer& Buffer = JsonRoot.Buffers[MergedBufferIndex];
+	if (Buffer.ByteLength != MergedBufferData.Num())
+	{
+		FString DataBase64 = FBase64::Encode(MergedBufferData.GetData(), MergedBufferData.Num());
+
+		Buffer.URI = TEXT("data:application/octet-stream;base64,") + DataBase64;
+		Buffer.ByteLength = MergedBufferData.Num();
+	}
+}
diff --git a/Engine/Plugins/Enterprise/GLTFExporter/MigratedFromGit/Source/GLTFExporter/Private/GLTFBuilder.cpp b/Engine/Plugins/Enterprise/GLTFExporter/MigratedFromGit/Source/GLTFExporter/Private/GLTFBuilder.cpp
--- a/Engine/Plugins/Enterprise/GLTFExporter/MigratedFromGit/Source/GLTFExporter/Private/GLTFBuilder.cpp
+++ b/Engine/Plugins/Enterprise/GLTFExporter/MigratedFromGit/Source/GLTFExporter/Private/GLTFBuilder.cpp
@@ -5,39 +5,18 @@
 #include "Misc/Base64.h"
 
 FGLTFBuilder::FGLTFBuilder()
-	: MergedBufferIndex(JsonRoot.Buffers.AddDefaulted(1))
+	: BufferBuilder(JsonRoot)
 {
 }
 
 FGLTFJsonBufferViewIndex FGLTFBuilder::AddBufferView(const void* RawData, uint64 ByteLength, const FString& Name, EGLTFJsonBufferTarget BufferTarget)
 {
-	FGLTFJsonBufferView BufferView;
-	BufferView.Name = Name;
-	BufferView.Buffer = MergedBufferIndex;
-	BufferView.ByteOffset = MergedBufferData.Num();
-	BufferView.ByteLength = ByteLength;
-	BufferView.Target = BufferTarget;
-
-	MergedBufferData.Append(static_cast<const uint8*>(RawData), ByteLength);
-
-	return JsonRoot.BufferViews.Add(BufferView);
-}
-
-void FGLTFBuilder::UpdateMergedBuffer()
-{
-	FGLTFJsonBuffer& Buffer = JsonRoot.Buffers[MergedBufferIndex];
-	if (Buffer.ByteLength != MergedBufferData.Num())
-	{
-		FString DataBase64 = FBase64::Encode(MergedBufferData.GetData(), MergedBufferData.Num());
-
-		Buffer.URI = TEXT("data:application/octet-stream;base64,") + DataBase64;
-		Buffer.ByteLength = MergedBufferData.Num();
-	}
+	return BufferBuilder.AddBufferView(RawData, ByteLength, Name, BufferTarget);
 }
 
 void FGLTFBuilder::Serialize(FArchive& Archive)
 {
-	UpdateMergedBuffer();
+	BufferBuilder.UpdateMergedBuffer();
 	JsonRoot.Serialize(&Archive, true);
 }
 
diff --git a/Engine/Plugins/Enterprise/GLTFExporter/MigratedFromGit/Source/GLTFExporter/Public/GLTFBufferBuilder.h b/Engine/Plugins/Enterprise/GLTFExporter/MigratedFromGit/Source/GLTFExporter/Public/GLTFBufferBuilder.h
new file mode 100644
--- /dev/null
+++ b/Engine/Plugins/Enterprise/GLTFExporter/MigratedFromGit/Source/GLTFExporter/Public/GLTFBufferBuilder.h
@@ -0,0 +1,20 @@
+// Copyright Epic Games, Inc. All Rights Reserved.
+
+#pragma once
+
+#include "GLTFJsonRoot.h"
+#include "Engine.h"
+
+struct GLTFEXPORTER_API FGLTFBufferBuilder
+{
+	FGLTFJsonRoot JsonRoot;
+
+	const FGLTFJsonBufferIndex MergedBufferIndex;
+	TArray<uint8> MergedBufferData;
+
+	FGLTFBufferBuilder(FGLTFJsonRoot& JsonRoot);
+
+	FGLTFJsonBufferViewIndex AddBufferView(const void* RawData, uint64 ByteLength, const FString& Name, EGLTFJsonBufferTarget BufferTarget);
+
+	void UpdateMergedBuffer();
+};
diff --git a/Engine/Plugins/Enterprise/GLTFExporter/MigratedFromGit/Source/GLTFExporter/Public/GLTFBuilder.h b/Engine/Plugins/Enterprise/GLTFExporter/MigratedFromGit/Source/GLTFExporter/Public/GLTFBuilder.h
--- a/Engine/Plugins/Enterprise/GLTFExporter/MigratedFromGit/Source/GLTFExporter/Public/GLTFBuilder.h
+++ b/Engine/Plugins/Enterprise/GLTFExporter/MigratedFromGit/Source/GLTFExporter/Public/GLTFBuilder.h
@@ -3,27 +3,25 @@
 #pragma once
 
 #include "GLTFJsonRoot.h"
+#include "GLTFBufferBuilder.h"
 #include "Engine.h"
 
 struct GLTFEXPORTER_API FGLTFBuilder
 {
 	FGLTFJsonRoot JsonRoot;
-
-	const FGLTFJsonBufferIndex MergedBufferIndex;
-	TArray<uint8> MergedBufferData;
+	FGLTFBufferBuilder BufferBuilder;
 
 	FGLTFBuilder();
 
+	FGLTFJsonBufferViewIndex AddBufferView(const void* RawData, uint64 ByteLength, const FString& Name = TEXT(""), EGLTFJsonBufferTarget BufferTarget = EGLTFJsonBufferTarget::ArrayBuffer);
+
 	template <class ElementType>
 	FGLTFJsonBufferViewIndex AddBufferView(const TArray<ElementType>& Array, const FString& Name = TEXT(""), EGLTFJsonBufferTarget BufferTarget = EGLTFJsonBufferTarget::ArrayBuffer)
 	{
-		return AddBufferView(Array.GetData(), Array.Num() * sizeof(ElementType), Name, BufferTarget);
+		return BufferBuilder.AddBufferView(Array.GetData(), Array.Num() * sizeof(ElementType), Name, BufferTarget);
 	}
 
-	FGLTFJsonBufferViewIndex AddBufferView(const void* RawData, uint64 ByteLength, const FString& Name = TEXT(""), EGLTFJsonBufferTarget BufferTarget = EGLTFJsonBufferTarget::ArrayBuffer);
-
 	FGLTFJsonMeshIndex AddMesh(const UStaticMesh* StaticMesh, int32 LODIndex);
 
-	void UpdateMergedBuffer();
 	void Serialize(FArchive& Archive);
 };
