[GLTFExporter] refactor: rename GLTFBuilder -> GLTFContainerBuilder

commit 87b9d067d49394d1d668742b5f9a7106d36bb90e
Date: Tue, 28 Apr 2020 14:12:52 +0200

    refactor: rename GLTFBuilder -> GLTFContainerBuilder
    (cherry picked from commit f9e4e285f9ff93979c809f68919acbb28c0becf5​)

#skipci

diff --git a/Engine/Plugins/Enterprise/GLTFExporter/MigratedFromGit/Source/GLTFExporter/Private/GLTFBufferBuilder.cpp b/Engine/Plugins/Enterprise/GLTFExporter/MigratedFromGit/Source/GLTFExporter/Private/GLTFBufferBuilder.cpp
--- a/Engine/Plugins/Enterprise/GLTFExporter/MigratedFromGit/Source/GLTFExporter/Private/GLTFBufferBuilder.cpp
+++ b/Engine/Plugins/Enterprise/GLTFExporter/MigratedFromGit/Source/GLTFExporter/Private/GLTFBufferBuilder.cpp
@@ -5,8 +5,9 @@
 #include "Misc/Base64.h"
 
 FGLTFBufferBuilder::FGLTFBufferBuilder(FGLTFJsonRoot& JsonRoot)
-	: MergedBufferIndex(JsonRoot.Buffers.AddDefaulted(1))
+	: JsonRoot(JsonRoot)
 {
+	MergedBufferIndex = JsonRoot.Buffers.AddDefaulted(1);
 }
 
 FGLTFJsonBufferViewIndex FGLTFBufferBuilder::AddBufferView(const void* RawData, uint64 ByteLength, const FString& Name, EGLTFJsonBufferTarget BufferTarget)
diff --git a/Engine/Plugins/Enterprise/GLTFExporter/MigratedFromGit/Source/GLTFExporter/Private/GLTFBuilder.cpp b/Engine/Plugins/Enterprise/GLTFExporter/MigratedFromGit/Source/GLTFExporter/Private/GLTFBuilder.cpp
deleted file mode 100644
diff --git a/Engine/Plugins/Enterprise/GLTFExporter/MigratedFromGit/Source/GLTFExporter/Private/GLTFContainerBuilder.cpp b/Engine/Plugins/Enterprise/GLTFExporter/MigratedFromGit/Source/GLTFExporter/Private/GLTFContainerBuilder.cpp
new file mode 100644
--- /dev/null
+++ b/Engine/Plugins/Enterprise/GLTFExporter/MigratedFromGit/Source/GLTFExporter/Private/GLTFContainerBuilder.cpp
@@ -0,0 +1,26 @@
+// Copyright Epic Games, Inc. All Rights Reserved.
+
+#include "GLTFContainerBuilder.h"
+#include "GLTFMeshBuilder.h"
+#include "Misc/Base64.h"
+
+FGLTFContainerBuilder::FGLTFContainerBuilder()
+	: BufferBuilder(JsonRoot)
+{
+}
+
+FGLTFJsonBufferViewIndex FGLTFContainerBuilder::AddBufferView(const void* RawData, uint64 ByteLength, const FString& Name, EGLTFJsonBufferTarget BufferTarget)
+{
+	return BufferBuilder.AddBufferView(RawData, ByteLength, Name, BufferTarget);
+}
+
+void FGLTFContainerBuilder::Serialize(FArchive& Archive)
+{
+	BufferBuilder.UpdateMergedBuffer();
+	JsonRoot.Serialize(&Archive, true);
+}
+
+FGLTFJsonMeshIndex FGLTFContainerBuilder::AddMesh(const UStaticMesh* StaticMesh, int32 LODIndex)
+{
+	return FGLTFMeshBuilder(StaticMesh, LODIndex).AddMesh(*this);
+}
diff --git a/Engine/Plugins/Enterprise/GLTFExporter/MigratedFromGit/Source/GLTFExporter/Private/GLTFMeshBuilder.cpp b/Engine/Plugins/Enterprise/GLTFExporter/MigratedFromGit/Source/GLTFExporter/Private/GLTFMeshBuilder.cpp
--- a/Engine/Plugins/Enterprise/GLTFExporter/MigratedFromGit/Source/GLTFExporter/Private/GLTFMeshBuilder.cpp
+++ b/Engine/Plugins/Enterprise/GLTFExporter/MigratedFromGit/Source/GLTFExporter/Private/GLTFMeshBuilder.cpp
@@ -16,7 +16,7 @@ FGLTFSectionBuilder::FGLTFSectionBuilder(const FString& SectionName, const FStat
 	}
 }
 
-FGLTFJsonAccessorIndex FGLTFSectionBuilder::AddAccessorForIndices(FGLTFBuilder& Builder) const
+FGLTFJsonAccessorIndex FGLTFSectionBuilder::AddAccessorForIndices(FGLTFContainerBuilder& Container) const
 {
 	if (Indices.Num() == 0)
 	{
@@ -24,7 +24,7 @@ FGLTFJsonAccessorIndex FGLTFSectionBuilder::AddAccessorForIndices(FGLTFBuilder&
 	}
 
 	const FString AttributeName = Name + TEXT("_Indices");
-	FGLTFJsonBufferViewIndex BufferViewIndex = Builder.AddBufferView(Indices, AttributeName, EGLTFJsonBufferTarget::ElementArrayBuffer);
+	FGLTFJsonBufferViewIndex BufferViewIndex = Container.AddBufferView(Indices, AttributeName, EGLTFJsonBufferTarget::ElementArrayBuffer);
 
 	FGLTFJsonAccessor Accessor;
 	Accessor.Name = AttributeName;
@@ -33,7 +33,7 @@ FGLTFJsonAccessorIndex FGLTFSectionBuilder::AddAccessorForIndices(FGLTFBuilder&
 	Accessor.Count = Indices.Num();
 	Accessor.Type = EGLTFJsonAccessorType::Scalar;
 
-	return Builder.JsonRoot.Accessors.Add(Accessor);
+	return Container.JsonRoot.Accessors.Add(Accessor);
 }
 
 FGLTFMeshBuilder::FGLTFMeshBuilder(const UStaticMesh* StaticMesh, int32 LODIndex)
@@ -112,7 +112,7 @@ FGLTFMeshBuilder::FGLTFMeshBuilder(const UStaticMesh* StaticMesh, int32 LODIndex
 	BoundingBox = StaticMesh->GetBoundingBox();
 }
 
-FGLTFJsonAccessorIndex FGLTFMeshBuilder::AddAccessorForPositions(FGLTFBuilder& Builder) const
+FGLTFJsonAccessorIndex FGLTFMeshBuilder::AddAccessorForPositions(FGLTFContainerBuilder& Container) const
 {
 	if (Positions.Num() == 0)
 	{
@@ -120,7 +120,7 @@ FGLTFJsonAccessorIndex FGLTFMeshBuilder::AddAccessorForPositions(FGLTFBuilder& B
 	}
 
 	const FString AttributeName = Name + TEXT("_Positions");
-	FGLTFJsonBufferViewIndex BufferViewIndex = Builder.AddBufferView(Positions, AttributeName);
+	FGLTFJsonBufferViewIndex BufferViewIndex = Container.AddBufferView(Positions, AttributeName);
 
 	FGLTFJsonAccessor Accessor;
 	Accessor.Name = AttributeName;
@@ -139,10 +139,10 @@ FGLTFJsonAccessorIndex FGLTFMeshBuilder::AddAccessorForPositions(FGLTFBuilder& B
 	Accessor.Min[2] = Min.Z;
 	Accessor.MinMaxLength = 3;
 
-	return Builder.JsonRoot.Accessors.Add(Accessor);
+	return Container.JsonRoot.Accessors.Add(Accessor);
 }
 
-FGLTFJsonAccessorIndex FGLTFMeshBuilder::AddAccessorForNormals(FGLTFBuilder& Builder) const
+FGLTFJsonAccessorIndex FGLTFMeshBuilder::AddAccessorForNormals(FGLTFContainerBuilder& Container) const
 {
 	if (Normals.Num() == 0)
 	{
@@ -150,7 +150,7 @@ FGLTFJsonAccessorIndex FGLTFMeshBuilder::AddAccessorForNormals(FGLTFBuilder& Bui
 	}
 
 	const FString AttributeName = Name + TEXT("_Normals");
-	FGLTFJsonBufferViewIndex BufferViewIndex = Builder.AddBufferView(Normals, AttributeName);
+	FGLTFJsonBufferViewIndex BufferViewIndex = Container.AddBufferView(Normals, AttributeName);
 
 	FGLTFJsonAccessor Accessor;
 	Accessor.Name = AttributeName;
@@ -159,10 +159,10 @@ FGLTFJsonAccessorIndex FGLTFMeshBuilder::AddAccessorForNormals(FGLTFBuilder& Bui
 	Accessor.Count = Normals.Num();
 	Accessor.Type = EGLTFJsonAccessorType::Vec3;
 
-	return Builder.JsonRoot.Accessors.Add(Accessor);
+	return Container.JsonRoot.Accessors.Add(Accessor);
 }
 
-FGLTFJsonAccessorIndex FGLTFMeshBuilder::AddAccessorForColors(FGLTFBuilder& Builder) const
+FGLTFJsonAccessorIndex FGLTFMeshBuilder::AddAccessorForColors(FGLTFContainerBuilder& Container) const
 {
 	if (Colors.Num() == 0)
 	{
@@ -170,7 +170,7 @@ FGLTFJsonAccessorIndex FGLTFMeshBuilder::AddAccessorForColors(FGLTFBuilder& Buil
 	}
 
 	const FString AttributeName = Name + TEXT("_Colors");
-	FGLTFJsonBufferViewIndex BufferViewIndex = Builder.AddBufferView(Colors, AttributeName);
+	FGLTFJsonBufferViewIndex BufferViewIndex = Container.AddBufferView(Colors, AttributeName);
 
 	FGLTFJsonAccessor Accessor;
 	Accessor.Name = AttributeName;
@@ -179,10 +179,10 @@ FGLTFJsonAccessorIndex FGLTFMeshBuilder::AddAccessorForColors(FGLTFBuilder& Buil
 	Accessor.Count = Colors.Num();
 	Accessor.Type = EGLTFJsonAccessorType::Vec4;
 
-	return Builder.JsonRoot.Accessors.Add(Accessor);
+	return Container.JsonRoot.Accessors.Add(Accessor);
 }
 
-FGLTFJsonAccessorIndex FGLTFMeshBuilder::AddAccessorForTangents(FGLTFBuilder& Builder) const
+FGLTFJsonAccessorIndex FGLTFMeshBuilder::AddAccessorForTangents(FGLTFContainerBuilder& Container) const
 {
 	if (Tangents.Num() == 0)
 	{
@@ -190,7 +190,7 @@ FGLTFJsonAccessorIndex FGLTFMeshBuilder::AddAccessorForTangents(FGLTFBuilder& Bu
 	}
 
 	const FString AttributeName = Name + TEXT("_Tangents");
-	FGLTFJsonBufferViewIndex BufferViewIndex = Builder.AddBufferView(Tangents, AttributeName);
+	FGLTFJsonBufferViewIndex BufferViewIndex = Container.AddBufferView(Tangents, AttributeName);
 
 	FGLTFJsonAccessor Accessor;
 	Accessor.Name = AttributeName;
@@ -199,10 +199,10 @@ FGLTFJsonAccessorIndex FGLTFMeshBuilder::AddAccessorForTangents(FGLTFBuilder& Bu
 	Accessor.Count = Tangents.Num();
 	Accessor.Type = EGLTFJsonAccessorType::Vec4;
 
-	return Builder.JsonRoot.Accessors.Add(Accessor);
+	return Container.JsonRoot.Accessors.Add(Accessor);
 }
 
-FGLTFJsonAccessorIndex FGLTFMeshBuilder::AddAccessorForUV0s(FGLTFBuilder& Builder) const
+FGLTFJsonAccessorIndex FGLTFMeshBuilder::AddAccessorForUV0s(FGLTFContainerBuilder& Container) const
 {
 	if (UV0s.Num() == 0)
 	{
@@ -210,7 +210,7 @@ FGLTFJsonAccessorIndex FGLTFMeshBuilder::AddAccessorForUV0s(FGLTFBuilder& Builde
 	}
 
 	const FString AttributeName = Name + TEXT("_UV0s");
-	FGLTFJsonBufferViewIndex BufferViewIndex = Builder.AddBufferView(UV0s, AttributeName);
+	FGLTFJsonBufferViewIndex BufferViewIndex = Container.AddBufferView(UV0s, AttributeName);
 
 	FGLTFJsonAccessor Accessor;
 	Accessor.Name = AttributeName;
@@ -219,10 +219,10 @@ FGLTFJsonAccessorIndex FGLTFMeshBuilder::AddAccessorForUV0s(FGLTFBuilder& Builde
 	Accessor.Count = UV0s.Num();
 	Accessor.Type = EGLTFJsonAccessorType::Vec2;
 
-	return Builder.JsonRoot.Accessors.Add(Accessor);
+	return Container.JsonRoot.Accessors.Add(Accessor);
 }
 
-FGLTFJsonAccessorIndex FGLTFMeshBuilder::AddAccessorForUV1s(FGLTFBuilder& Builder) const
+FGLTFJsonAccessorIndex FGLTFMeshBuilder::AddAccessorForUV1s(FGLTFContainerBuilder& Container) const
 {
 	if (UV1s.Num() == 0)
 	{
@@ -230,7 +230,7 @@ FGLTFJsonAccessorIndex FGLTFMeshBuilder::AddAccessorForUV1s(FGLTFBuilder& Builde
 	}
 
 	const FString AttributeName = Name + TEXT("_UV1s");
-	FGLTFJsonBufferViewIndex BufferViewIndex = Builder.AddBufferView(UV1s, AttributeName);
+	FGLTFJsonBufferViewIndex BufferViewIndex = Container.AddBufferView(UV1s, AttributeName);
 
 	FGLTFJsonAccessor Accessor;
 	Accessor.Name = AttributeName;
@@ -239,30 +239,30 @@ FGLTFJsonAccessorIndex FGLTFMeshBuilder::AddAccessorForUV1s(FGLTFBuilder& Builde
 	Accessor.Count = UV1s.Num();
 	Accessor.Type = EGLTFJsonAccessorType::Vec2;
 
-	return Builder.JsonRoot.Accessors.Add(Accessor);
+	return Container.JsonRoot.Accessors.Add(Accessor);
 }
 
-FGLTFJsonMeshIndex FGLTFMeshBuilder::AddMesh(FGLTFBuilder& Builder) const
+FGLTFJsonMeshIndex FGLTFMeshBuilder::AddMesh(FGLTFContainerBuilder& Container) const
 {
 	FGLTFJsonMesh Mesh;
 	Mesh.Name = Name;
 
 	FGLTFJsonAttributes Attributes;
-	Attributes.Position = AddAccessorForPositions(Builder);
-	Attributes.Color0 = AddAccessorForColors(Builder);
-	Attributes.Normal = AddAccessorForNormals(Builder);
-	Attributes.Tangent = AddAccessorForTangents(Builder);
-	Attributes.TexCoord0 = AddAccessorForUV0s(Builder);
-	Attributes.TexCoord1 = AddAccessorForUV1s(Builder);
+	Attributes.Position = AddAccessorForPositions(Container);
+	Attributes.Color0 = AddAccessorForColors(Container);
+	Attributes.Normal = AddAccessorForNormals(Container);
+	Attributes.Tangent = AddAccessorForTangents(Container);
+	Attributes.TexCoord0 = AddAccessorForUV0s(Container);
+	Attributes.TexCoord1 = AddAccessorForUV1s(Container);
 
 	for (const FGLTFSectionBuilder& Section : Sections)
 	{
 		FGLTFJsonPrimitive Primitive;
 		Primitive.Attributes = Attributes;
 
-		Primitive.Indices = Section.AddAccessorForIndices(Builder);
+		Primitive.Indices = Section.AddAccessorForIndices(Container);
 		Mesh.Primitives.Add(Primitive);
 	}
 
-	return Builder.JsonRoot.Meshes.Add(Mesh);
+	return Container.JsonRoot.Meshes.Add(Mesh);
 }
diff --git a/Engine/Plugins/Enterprise/GLTFExporter/MigratedFromGit/Source/GLTFExporter/Private/GLTFStaticMeshExporter.cpp b/Engine/Plugins/Enterprise/GLTFExporter/MigratedFromGit/Source/GLTFExporter/Private/GLTFStaticMeshExporter.cpp
--- a/Engine/Plugins/Enterprise/GLTFExporter/MigratedFromGit/Source/GLTFExporter/Private/GLTFStaticMeshExporter.cpp
+++ b/Engine/Plugins/Enterprise/GLTFExporter/MigratedFromGit/Source/GLTFExporter/Private/GLTFStaticMeshExporter.cpp
@@ -2,7 +2,7 @@
 
 #include "GLTFStaticMeshExporter.h"
 #include "GLTFJsonRoot.h"
-#include "GLTFBuilder.h"
+#include "GLTFContainerBuilder.h"
 #include "Engine/StaticMesh.h"
 
 UGLTFStaticMeshExporter::UGLTFStaticMeshExporter(const FObjectInitializer& ObjectInitializer)
@@ -21,20 +21,20 @@ bool UGLTFStaticMeshExporter::ExportBinary(UObject* Object, const TCHAR* Type, F
 		return false;
 	}
 
-	FGLTFBuilder Builder;
+	FGLTFContainerBuilder Container;
 
-	FGLTFJsonMeshIndex MeshIndex = Builder.AddMesh(StaticMesh, 0);
+	FGLTFJsonMeshIndex MeshIndex = Container.AddMesh(StaticMesh, 0);
 
 	FGLTFJsonNode Node;
 	Node.Mesh = MeshIndex;
-	FGLTFJsonNodeIndex NodeIndex = Builder.JsonRoot.Nodes.Add(Node);
+	FGLTFJsonNodeIndex NodeIndex = Container.JsonRoot.Nodes.Add(Node);
 
 	FGLTFJsonScene Scene;
 	Scene.Nodes.Add(NodeIndex);
-	FGLTFJsonSceneIndex SceneIndex = Builder.JsonRoot.Scenes.Add(Scene);
+	FGLTFJsonSceneIndex SceneIndex = Container.JsonRoot.Scenes.Add(Scene);
 
-	Builder.JsonRoot.DefaultScene = SceneIndex;
+	Container.JsonRoot.DefaultScene = SceneIndex;
 
-	Builder.Serialize(Archive);
+	Container.Serialize(Archive);
 	return true;
 }
diff --git a/Engine/Plugins/Enterprise/GLTFExporter/MigratedFromGit/Source/GLTFExporter/Public/GLTFBufferBuilder.h b/Engine/Plugins/Enterprise/GLTFExporter/MigratedFromGit/Source/GLTFExporter/Public/GLTFBufferBuilder.h
--- a/Engine/Plugins/Enterprise/GLTFExporter/MigratedFromGit/Source/GLTFExporter/Public/GLTFBufferBuilder.h
+++ b/Engine/Plugins/Enterprise/GLTFExporter/MigratedFromGit/Source/GLTFExporter/Public/GLTFBufferBuilder.h
@@ -7,9 +7,9 @@
 
 struct GLTFEXPORTER_API FGLTFBufferBuilder
 {
-	FGLTFJsonRoot JsonRoot;
+	FGLTFJsonRoot& JsonRoot;
 
-	const FGLTFJsonBufferIndex MergedBufferIndex;
+	FGLTFJsonBufferIndex MergedBufferIndex;
 	TArray<uint8> MergedBufferData;
 
 	FGLTFBufferBuilder(FGLTFJsonRoot& JsonRoot);
diff --git a/Engine/Plugins/Enterprise/GLTFExporter/MigratedFromGit/Source/GLTFExporter/Public/GLTFBuilder.h b/Engine/Plugins/Enterprise/GLTFExporter/MigratedFromGit/Source/GLTFExporter/Public/GLTFBuilder.h
deleted file mode 100644
diff --git a/Engine/Plugins/Enterprise/GLTFExporter/MigratedFromGit/Source/GLTFExporter/Public/GLTFContainerBuilder.h b/Engine/Plugins/Enterprise/GLTFExporter/MigratedFromGit/Source/GLTFExporter/Public/GLTFContainerBuilder.h
new file mode 100644
--- /dev/null
+++ b/Engine/Plugins/Enterprise/GLTFExporter/MigratedFromGit/Source/GLTFExporter/Public/GLTFContainerBuilder.h
@@ -0,0 +1,27 @@
+// Copyright Epic Games, Inc. All Rights Reserved.
+
+#pragma once
+
+#include "GLTFJsonRoot.h"
+#include "GLTFBufferBuilder.h"
+#include "Engine.h"
+
+struct GLTFEXPORTER_API FGLTFContainerBuilder
+{
+	FGLTFJsonRoot JsonRoot;
+	FGLTFBufferBuilder BufferBuilder;
+
+	FGLTFContainerBuilder();
+
+	FGLTFJsonBufferViewIndex AddBufferView(const void* RawData, uint64 ByteLength, const FString& Name = TEXT(""), EGLTFJsonBufferTarget BufferTarget = EGLTFJsonBufferTarget::ArrayBuffer);
+
+	template <class ElementType>
+	FGLTFJsonBufferViewIndex AddBufferView(const TArray<ElementType>& Array, const FString& Name = TEXT(""), EGLTFJsonBufferTarget BufferTarget = EGLTFJsonBufferTarget::ArrayBuffer)
+	{
+		return BufferBuilder.AddBufferView(Array.GetData(), Array.Num() * sizeof(ElementType), Name, BufferTarget);
+	}
+
+	FGLTFJsonMeshIndex AddMesh(const UStaticMesh* StaticMesh, int32 LODIndex);
+
+	void Serialize(FArchive& Archive);
+};
diff --git a/Engine/Plugins/Enterprise/GLTFExporter/MigratedFromGit/Source/GLTFExporter/Public/GLTFMeshBuilder.h b/Engine/Plugins/Enterprise/GLTFExporter/MigratedFromGit/Source/GLTFExporter/Public/GLTFMeshBuilder.h
--- a/Engine/Plugins/Enterprise/GLTFExporter/MigratedFromGit/Source/GLTFExporter/Public/GLTFMeshBuilder.h
+++ b/Engine/Plugins/Enterprise/GLTFExporter/MigratedFromGit/Source/GLTFExporter/Public/GLTFMeshBuilder.h
@@ -3,7 +3,7 @@
 #pragma once
 
 #include "GLTFJsonRoot.h"
-#include "GLTFBuilder.h"
+#include "GLTFContainerBuilder.h"
 #include "Engine.h"
 
 struct GLTFEXPORTER_API FGLTFSectionBuilder
@@ -14,7 +14,7 @@ struct GLTFEXPORTER_API FGLTFSectionBuilder
 
 	FGLTFSectionBuilder(const FString& SectionName, const FStaticMeshSection& MeshSection, const FIndexArrayView& IndexArray);
 
-	FGLTFJsonAccessorIndex AddAccessorForIndices(FGLTFBuilder& Builder) const;
+	FGLTFJsonAccessorIndex AddAccessorForIndices(FGLTFContainerBuilder& Container) const;
 };
 
 struct GLTFEXPORTER_API FGLTFMeshBuilder
@@ -34,12 +34,12 @@ struct GLTFEXPORTER_API FGLTFMeshBuilder
 
 	FGLTFMeshBuilder(const UStaticMesh* StaticMesh, int32 LODIndex);
 
-	FGLTFJsonAccessorIndex AddAccessorForPositions(FGLTFBuilder& Builder) const;
-	FGLTFJsonAccessorIndex AddAccessorForColors(FGLTFBuilder& Builder) const;
-	FGLTFJsonAccessorIndex AddAccessorForNormals(FGLTFBuilder& Builder) const;
-	FGLTFJsonAccessorIndex AddAccessorForTangents(FGLTFBuilder& Builder) const;
-	FGLTFJsonAccessorIndex AddAccessorForUV0s(FGLTFBuilder& Builder) const;
-	FGLTFJsonAccessorIndex AddAccessorForUV1s(FGLTFBuilder& Builder) const;
+	FGLTFJsonAccessorIndex AddAccessorForPositions(FGLTFContainerBuilder& Container) const;
+	FGLTFJsonAccessorIndex AddAccessorForColors(FGLTFContainerBuilder& Container) const;
+	FGLTFJsonAccessorIndex AddAccessorForNormals(FGLTFContainerBuilder& Container) const;
+	FGLTFJsonAccessorIndex AddAccessorForTangents(FGLTFContainerBuilder& Container) const;
+	FGLTFJsonAccessorIndex AddAccessorForUV0s(FGLTFContainerBuilder& Container) const;
+	FGLTFJsonAccessorIndex AddAccessorForUV1s(FGLTFContainerBuilder& Container) const;
 
-	FGLTFJsonMeshIndex AddMesh(FGLTFBuilder& Builder) const;
+	FGLTFJsonMeshIndex AddMesh(FGLTFContainerBuilder& Container) const;
 };
