[GLTFExporter] refactor: fix code inspection issues

commit 4ceef3ba25ced5e0fecc57b802d72b7102528c36
Date: Wed, 29 Apr 2020 22:09:03 +0200

    refactor: fix code inspection issues
    (cherry picked from commit e70c0f8f9cfea07367620aac1f14a6b6ea18cdb8​)

#skipci

diff --git a/Engine/Plugins/Enterprise/GLTFExporter/MigratedFromGit/Source/GLTFExporter/Private/GLTFBufferBuilder.cpp b/Engine/Plugins/Enterprise/GLTFExporter/MigratedFromGit/Source/GLTFExporter/Private/GLTFBufferBuilder.cpp
--- a/Engine/Plugins/Enterprise/GLTFExporter/MigratedFromGit/Source/GLTFExporter/Private/GLTFBufferBuilder.cpp
+++ b/Engine/Plugins/Enterprise/GLTFExporter/MigratedFromGit/Source/GLTFExporter/Private/GLTFBufferBuilder.cpp
@@ -29,7 +29,7 @@ void FGLTFBufferBuilder::UpdateMergedBuffer()
 	FGLTFJsonBuffer& Buffer = JsonRoot.Buffers[MergedBufferIndex];
 	if (Buffer.ByteLength != MergedBufferData.Num())
 	{
-		FString DataBase64 = FBase64::Encode(MergedBufferData.GetData(), MergedBufferData.Num());
+		const FString DataBase64 = FBase64::Encode(MergedBufferData.GetData(), MergedBufferData.Num());
 
 		Buffer.URI = TEXT("data:application/octet-stream;base64,") + DataBase64;
 		Buffer.ByteLength = MergedBufferData.Num();
diff --git a/Engine/Plugins/Enterprise/GLTFExporter/MigratedFromGit/Source/GLTFExporter/Private/GLTFContainerBuilder.cpp b/Engine/Plugins/Enterprise/GLTFExporter/MigratedFromGit/Source/GLTFExporter/Private/GLTFContainerBuilder.cpp
--- a/Engine/Plugins/Enterprise/GLTFExporter/MigratedFromGit/Source/GLTFExporter/Private/GLTFContainerBuilder.cpp
+++ b/Engine/Plugins/Enterprise/GLTFExporter/MigratedFromGit/Source/GLTFExporter/Private/GLTFContainerBuilder.cpp
@@ -2,7 +2,6 @@
 
 #include "GLTFContainerBuilder.h"
 #include "GLTFMeshBuilder.h"
-#include "Misc/Base64.h"
 
 FGLTFContainerBuilder::FGLTFContainerBuilder()
 	: BufferBuilder(JsonRoot)
diff --git a/Engine/Plugins/Enterprise/GLTFExporter/MigratedFromGit/Source/GLTFExporter/Private/GLTFExportOptions.cpp b/Engine/Plugins/Enterprise/GLTFExporter/MigratedFromGit/Source/GLTFExporter/Private/GLTFExportOptions.cpp
--- a/Engine/Plugins/Enterprise/GLTFExporter/MigratedFromGit/Source/GLTFExporter/Private/GLTFExportOptions.cpp
+++ b/Engine/Plugins/Enterprise/GLTFExporter/MigratedFromGit/Source/GLTFExporter/Private/GLTFExportOptions.cpp
@@ -1,7 +1,8 @@
-// Copyriyright Epic Games, Inc. All Rights Reserved.
+// Copyright Epic Games, Inc. All Rights Reserved.
 
 #include "GLTFExportOptions.h"
 #include "GLTFExportOptionsWindow.h"
+#include "Misc/App.h"
 
 #include "UObject/ObjectMacros.h"
 #include "UObject/UObjectGlobals.h"
@@ -44,7 +45,7 @@ void UGLTFExportOptions::ResetToDefault()
 
 void UGLTFExportOptions::LoadOptions()
 {
-	int32 PortFlags = 0;
+	const int32 PortFlags = 0;
 
 	for (UProperty* Property = GetClass()->PropertyLink; Property; Property = Property->PropertyLinkNext)
 	{
@@ -55,15 +56,10 @@ void UGLTFExportOptions::LoadOptions()
 		FString Section = GetClass()->GetName();
 		FString Key = Property->GetName();
 
-		const bool bIsPropertyInherited = Property->GetOwnerClass() != GetClass();
-		UObject* SuperClassDefaultObject = GetClass()->GetSuperClass()->GetDefaultObject();
-
-		const FString& PropFileName = GEditorPerProjectIni;
-
 		UArrayProperty* Array = dynamic_cast<UArrayProperty*>(Property);
 		if (Array)
 		{
-			FConfigSection* Sec = GConfig->GetSectionPrivate(*Section, 0, 1, *GEditorPerProjectIni);
+			FConfigSection* Sec = GConfig->GetSectionPrivate(*Section, false, true, *GEditorPerProjectIni);
 			if (Sec != nullptr)
 			{
 				TArray<FConfigValue> List;
@@ -83,7 +79,7 @@ void UGLTFExportOptions::LoadOptions()
 				else
 				{
 					int32 Index = 0;
-					const FConfigValue* ElementValue = nullptr;
+					const FConfigValue* ElementValue;
 					do
 					{
 						// Add array index number to end of key
@@ -120,11 +116,11 @@ void UGLTFExportOptions::LoadOptions()
 				}
 
 				FString Value;
-				bool bFoundValue = GConfig->GetString(*Section, *Key, Value, *GEditorPerProjectIni);
+				const bool bFoundValue = GConfig->GetString(*Section, *Key, Value, *GEditorPerProjectIni);
 
 				if (bFoundValue)
 				{
-					if (Property->ImportText(*Value, Property->ContainerPtrToValuePtr<uint8>(this, i), PortFlags, this) == NULL)
+					if (Property->ImportText(*Value, Property->ContainerPtrToValuePtr<uint8>(this, i), PortFlags, this) == nullptr)
 					{
 						// this should be an error as the properties from the .ini / .int file are not correctly being read in and probably are affecting things in subtle ways
 					}
@@ -136,7 +132,7 @@ void UGLTFExportOptions::LoadOptions()
 
 void UGLTFExportOptions::SaveOptions()
 {
-	int32 PortFlags = 0;
+	const int32 PortFlags = 0;
 
 	for (UProperty* Property = GetClass()->PropertyLink; Property; Property = Property->PropertyLinkNext)
 	{
@@ -147,13 +143,10 @@ void UGLTFExportOptions::SaveOptions()
 		FString Section = GetClass()->GetName();
 		FString Key = Property->GetName();
 
-		const bool bIsPropertyInherited = Property->GetOwnerClass() != GetClass();
-		UObject* SuperClassDefaultObject = GetClass()->GetSuperClass()->GetDefaultObject();
-
 		UArrayProperty* Array = dynamic_cast<UArrayProperty*>(Property);
 		if (Array)
 		{
-			FConfigSection* Sec = GConfig->GetSectionPrivate(*Section, 1, 0, *GEditorPerProjectIni);
+			FConfigSection* Sec = GConfig->GetSectionPrivate(*Section, true, false, *GEditorPerProjectIni);
 			check(Sec);
 			Sec->Remove(*Key);
 
@@ -182,7 +175,7 @@ void UGLTFExportOptions::SaveOptions()
 			}
 		}
 	}
-	GConfig->Flush(0);
+	GConfig->Flush(false);
 }
 
 void UGLTFExportOptions::FillOptions(bool bBatchMode, bool bShowOptionDialog, const FString& FullPath, bool& bOutOperationCanceled, bool& bOutExportAll)
diff --git a/Engine/Plugins/Enterprise/GLTFExporter/MigratedFromGit/Source/GLTFExporter/Private/GLTFExportOptionsWindow.cpp b/Engine/Plugins/Enterprise/GLTFExporter/MigratedFromGit/Source/GLTFExporter/Private/GLTFExportOptionsWindow.cpp
--- a/Engine/Plugins/Enterprise/GLTFExporter/MigratedFromGit/Source/GLTFExporter/Private/GLTFExportOptionsWindow.cpp
+++ b/Engine/Plugins/Enterprise/GLTFExporter/MigratedFromGit/Source/GLTFExporter/Private/GLTFExportOptionsWindow.cpp
@@ -24,8 +24,8 @@ void SGLTFExportOptionsWindow::Construct(const FArguments& InArgs)
 
 	check (ExportOptions);
 	
-	FText CancelText = InArgs._BatchMode ? LOCTEXT("GLTFExportOptionsWindow_CancelBatch", "Cancel All") : LOCTEXT("GLTFExportOptionsWindow_Cancel", "Cancel");
-	FText CancelTooltipText = InArgs._BatchMode ? LOCTEXT("GLTFExportOptionsWindow_Cancel_ToolTip_Batch", "Cancel the batch export.") : LOCTEXT("GLTFExportOptionsWindow_Cancel_ToolTip", "Cancel the current glTF export.");
+	const FText CancelText = InArgs._BatchMode ? LOCTEXT("GLTFExportOptionsWindow_CancelBatch", "Cancel All") : LOCTEXT("GLTFExportOptionsWindow_Cancel", "Cancel");
+	const FText CancelTooltipText = InArgs._BatchMode ? LOCTEXT("GLTFExportOptionsWindow_Cancel_ToolTip_Batch", "Cancel the batch export.") : LOCTEXT("GLTFExportOptionsWindow_Cancel_ToolTip", "Cancel the current glTF export.");
 
 	TSharedPtr<SBox> HeaderToolBox;
 	TSharedPtr<SHorizontalBox> HeaderButtons;
@@ -221,8 +221,8 @@ void SGLTFExportOptionsWindow::ShowDialog(UGLTFExportOptions* ExportOptions, con
 		ParentWindow = MainFrame.GetParentWindow();
 	}
 
-	bool bBinaryFile = FPaths::GetExtension(FullPath, false).Compare(TEXT("glb"), ESearchCase::IgnoreCase) == 0;
-	auto Title = bBinaryFile ?
+	const bool bBinaryFile = FPaths::GetExtension(FullPath, false).Compare(TEXT("glb"), ESearchCase::IgnoreCase) == 0;
+	const FText Title = bBinaryFile ?
 		NSLOCTEXT("UnrealEd", "GLTFExportOptionsTitleBinary", "glTF (Binary) Export Options") :
 		NSLOCTEXT("UnrealEd", "GLTFExportOptionsTitle", "glTF Export Options");
 
diff --git a/Engine/Plugins/Enterprise/GLTFExporter/MigratedFromGit/Source/GLTFExporter/Private/GLTFMeshBuilder.cpp b/Engine/Plugins/Enterprise/GLTFExporter/MigratedFromGit/Source/GLTFExporter/Private/GLTFMeshBuilder.cpp
--- a/Engine/Plugins/Enterprise/GLTFExporter/MigratedFromGit/Source/GLTFExporter/Private/GLTFMeshBuilder.cpp
+++ b/Engine/Plugins/Enterprise/GLTFExporter/MigratedFromGit/Source/GLTFExporter/Private/GLTFMeshBuilder.cpp
@@ -24,7 +24,7 @@ FGLTFJsonAccessorIndex FGLTFSectionBuilder::AddAccessorForIndices(FGLTFContainer
 	}
 
 	const FString AttributeName = Name + TEXT("_Indices");
-	FGLTFJsonBufferViewIndex BufferViewIndex = Container.AddBufferView(Indices, AttributeName, EGLTFJsonBufferTarget::ElementArrayBuffer);
+	const FGLTFJsonBufferViewIndex BufferViewIndex = Container.AddBufferView(Indices, AttributeName, EGLTFJsonBufferTarget::ElementArrayBuffer);
 
 	FGLTFJsonAccessor Accessor;
 	Accessor.Name = AttributeName;
@@ -120,7 +120,7 @@ FGLTFJsonAccessorIndex FGLTFMeshBuilder::AddAccessorForPositions(FGLTFContainerB
 	}
 
 	const FString AttributeName = Name + TEXT("_Positions");
-	FGLTFJsonBufferViewIndex BufferViewIndex = Container.AddBufferView(Positions, AttributeName);
+	const FGLTFJsonBufferViewIndex BufferViewIndex = Container.AddBufferView(Positions, AttributeName);
 
 	FGLTFJsonAccessor Accessor;
 	Accessor.Name = AttributeName;
@@ -129,8 +129,8 @@ FGLTFJsonAccessorIndex FGLTFMeshBuilder::AddAccessorForPositions(FGLTFContainerB
 	Accessor.Count = Positions.Num();
 	Accessor.Type = EGLTFJsonAccessorType::Vec3;
 
-	FVector Max = ConvertSize(BoundingBox.Max);
-	FVector Min = ConvertSize(BoundingBox.Min);
+	const FVector Max = ConvertSize(BoundingBox.Max);
+	const FVector Min = ConvertSize(BoundingBox.Min);
 	Accessor.Max[0] = Max.X;
 	Accessor.Max[1] = Max.Y;
 	Accessor.Max[2] = Max.Z;
@@ -150,7 +150,7 @@ FGLTFJsonAccessorIndex FGLTFMeshBuilder::AddAccessorForNormals(FGLTFContainerBui
 	}
 
 	const FString AttributeName = Name + TEXT("_Normals");
-	FGLTFJsonBufferViewIndex BufferViewIndex = Container.AddBufferView(Normals, AttributeName);
+	const FGLTFJsonBufferViewIndex BufferViewIndex = Container.AddBufferView(Normals, AttributeName);
 
 	FGLTFJsonAccessor Accessor;
 	Accessor.Name = AttributeName;
@@ -170,7 +170,7 @@ FGLTFJsonAccessorIndex FGLTFMeshBuilder::AddAccessorForColors(FGLTFContainerBuil
 	}
 
 	const FString AttributeName = Name + TEXT("_Colors");
-	FGLTFJsonBufferViewIndex BufferViewIndex = Container.AddBufferView(Colors, AttributeName);
+	const FGLTFJsonBufferViewIndex BufferViewIndex = Container.AddBufferView(Colors, AttributeName);
 
 	FGLTFJsonAccessor Accessor;
 	Accessor.Name = AttributeName;
@@ -190,7 +190,7 @@ FGLTFJsonAccessorIndex FGLTFMeshBuilder::AddAccessorForTangents(FGLTFContainerBu
 	}
 
 	const FString AttributeName = Name + TEXT("_Tangents");
-	FGLTFJsonBufferViewIndex BufferViewIndex = Container.AddBufferView(Tangents, AttributeName);
+	const FGLTFJsonBufferViewIndex BufferViewIndex = Container.AddBufferView(Tangents, AttributeName);
 
 	FGLTFJsonAccessor Accessor;
 	Accessor.Name = AttributeName;
@@ -210,7 +210,7 @@ FGLTFJsonAccessorIndex FGLTFMeshBuilder::AddAccessorForUV0s(FGLTFContainerBuilde
 	}
 
 	const FString AttributeName = Name + TEXT("_UV0s");
-	FGLTFJsonBufferViewIndex BufferViewIndex = Container.AddBufferView(UV0s, AttributeName);
+	const FGLTFJsonBufferViewIndex BufferViewIndex = Container.AddBufferView(UV0s, AttributeName);
 
 	FGLTFJsonAccessor Accessor;
 	Accessor.Name = AttributeName;
@@ -230,7 +230,7 @@ FGLTFJsonAccessorIndex FGLTFMeshBuilder::AddAccessorForUV1s(FGLTFContainerBuilde
 	}
 
 	const FString AttributeName = Name + TEXT("_UV1s");
-	FGLTFJsonBufferViewIndex BufferViewIndex = Container.AddBufferView(UV1s, AttributeName);
+	const FGLTFJsonBufferViewIndex BufferViewIndex = Container.AddBufferView(UV1s, AttributeName);
 
 	FGLTFJsonAccessor Accessor;
 	Accessor.Name = AttributeName;
diff --git a/Engine/Plugins/Enterprise/GLTFExporter/MigratedFromGit/Source/GLTFExporter/Public/GLTFConversionUtilities.h b/Engine/Plugins/Enterprise/GLTFExporter/MigratedFromGit/Source/GLTFExporter/Public/GLTFConversionUtilities.h
--- a/Engine/Plugins/Enterprise/GLTFExporter/MigratedFromGit/Source/GLTFExporter/Public/GLTFConversionUtilities.h
+++ b/Engine/Plugins/Enterprise/GLTFExporter/MigratedFromGit/Source/GLTFExporter/Public/GLTFConversionUtilities.h
@@ -39,17 +39,17 @@ inline FColor ConvertColor(const FColor& Color)
 	return { Color.R, Color.G, Color.B, Color.A };
 }
 
-inline FQuat ConvertRotation(const FQuat& Quat)
+inline FQuat ConvertRotation(const FQuat& Rotation)
 {
 	// UE4 uses a left-handed coordinate system, with Z up.
 	// glTF uses a right-handed coordinate system, with Y up.
-	// Quat = (qX, qY, qZ, qW) = (sin(angle/2) * aX, sin(angle/2) * aY, sin(angle/2) * aZ, cons(angle/2))
+	// Rotation = (qX, qY, qZ, qW) = (sin(angle/2) * aX, sin(angle/2) * aY, sin(angle/2) * aZ, cons(angle/2))
 	// where (aX, aY, aZ) - rotation axis, angle - rotation angle
 	// Y swapped with Z between these coordinate systems
 	// also, as handedness is changed rotation is inversed - hence negation
-	// therefore glTFQuat = (-qX, -qZ, -qY, qw)
+	// therefore glTFRotation = (-qX, -qZ, -qY, qw)
 
-	FQuat Result(-Quat.X, -Quat.Z, -Quat.Y, Quat.W);
+	const FQuat Result(-Rotation.X, -Rotation.Z, -Rotation.Y, Rotation.W);
 	// Not checking if quaternion is normalized
 	// e.g. some sources use non-unit Quats for rotation tangents
 	return Result;
@@ -71,14 +71,13 @@ inline FMatrix ConvertMatrix(const FMatrix& Matrix)
 	return Result;
 }
 
-inline EGLTFJsonAlphaMode ConvertAlphaMode(EBlendMode Mode)
+inline EGLTFJsonAlphaMode ConvertAlphaMode(const EBlendMode Mode)
 {
 	switch (Mode)
 	{
 		case EBlendMode::BLEND_Opaque:      return EGLTFJsonAlphaMode::Opaque;
 		case EBlendMode::BLEND_Translucent: return EGLTFJsonAlphaMode::Blend;
 		case EBlendMode::BLEND_Masked:      return EGLTFJsonAlphaMode::Mask;
+		default:                            return EGLTFJsonAlphaMode::Opaque; // fallback
 	}
-
-	return EGLTFJsonAlphaMode::Opaque; // fallback
 }
diff --git a/Engine/Plugins/Enterprise/GLTFExporter/MigratedFromGit/Source/GLTFExporter/Public/GLTFExportOptions.h b/Engine/Plugins/Enterprise/GLTFExporter/MigratedFromGit/Source/GLTFExporter/Public/GLTFExportOptions.h
--- a/Engine/Plugins/Enterprise/GLTFExporter/MigratedFromGit/Source/GLTFExporter/Public/GLTFExportOptions.h
+++ b/Engine/Plugins/Enterprise/GLTFExporter/MigratedFromGit/Source/GLTFExporter/Public/GLTFExportOptions.h
@@ -3,7 +3,6 @@
 #pragma once
 
 #include "CoreMinimal.h"
-#include "UObject/ObjectMacros.h"
 #include "GLTFExportOptions.generated.h"
 
 UENUM(BlueprintType)
diff --git a/Engine/Plugins/Enterprise/GLTFExporter/MigratedFromGit/Source/GLTFExporter/Public/GLTFExportOptionsWindow.h b/Engine/Plugins/Enterprise/GLTFExporter/MigratedFromGit/Source/GLTFExporter/Public/GLTFExportOptionsWindow.h
--- a/Engine/Plugins/Enterprise/GLTFExporter/MigratedFromGit/Source/GLTFExporter/Public/GLTFExportOptionsWindow.h
+++ b/Engine/Plugins/Enterprise/GLTFExporter/MigratedFromGit/Source/GLTFExporter/Public/GLTFExportOptionsWindow.h
@@ -17,8 +17,6 @@ class GLTFEXPORTER_API SGLTFExportOptionsWindow : public SCompoundWidget
 public:
 	SLATE_BEGIN_ARGS(SGLTFExportOptionsWindow)
 		: _ExportOptions(nullptr)
-		, _WidgetWindow()
-		, _FullPath()
 		, _BatchMode()
 		{}
 
diff --git a/Engine/Plugins/Enterprise/GLTFExporter/MigratedFromGit/Source/GLTFExporter/Public/GLTFJsonAccessor.h b/Engine/Plugins/Enterprise/GLTFExporter/MigratedFromGit/Source/GLTFExporter/Public/GLTFJsonAccessor.h
--- a/Engine/Plugins/Enterprise/GLTFExporter/MigratedFromGit/Source/GLTFExporter/Public/GLTFJsonAccessor.h
+++ b/Engine/Plugins/Enterprise/GLTFExporter/MigratedFromGit/Source/GLTFExporter/Public/GLTFJsonAccessor.h
@@ -28,6 +28,8 @@ struct GLTFEXPORTER_API FGLTFJsonAccessor
 		, ComponentType(EGLTFJsonComponentType::None)
 		, Normalized(false)
 		, MinMaxLength(0)
+		, Min{0}
+		, Max{0}
 	{
 	}
 
diff --git a/Engine/Plugins/Enterprise/GLTFExporter/MigratedFromGit/Source/GLTFExporter/Public/GLTFJsonUtilities.h b/Engine/Plugins/Enterprise/GLTFExporter/MigratedFromGit/Source/GLTFExporter/Public/GLTFJsonUtilities.h
--- a/Engine/Plugins/Enterprise/GLTFExporter/MigratedFromGit/Source/GLTFExporter/Public/GLTFJsonUtilities.h
+++ b/Engine/Plugins/Enterprise/GLTFExporter/MigratedFromGit/Source/GLTFExporter/Public/GLTFJsonUtilities.h
@@ -11,10 +11,8 @@ inline const TCHAR* AlphaModeToString(EGLTFJsonAlphaMode Mode)
 		case EGLTFJsonAlphaMode::Opaque: return TEXT("OPAQUE");
 		case EGLTFJsonAlphaMode::Blend:  return TEXT("BLEND");
 		case EGLTFJsonAlphaMode::Mask:   return TEXT("MASK");
+		default:                         return TEXT("UNKNOWN");
 	}
-
-	check(false);
-	return TEXT("UNKNOWN");
 }
 
 inline int32 PrimitiveModeToNumber(EGLTFJsonPrimitiveMode Mode)
@@ -43,10 +41,8 @@ inline const TCHAR* AccessorTypeToString(EGLTFJsonAccessorType Type)
 		case EGLTFJsonAccessorType::Mat2:   return TEXT("MAT2");
 		case EGLTFJsonAccessorType::Mat3:   return TEXT("MAT3");
 		case EGLTFJsonAccessorType::Mat4:   return TEXT("MAT4");
+		default:                            return TEXT("UNKNOWN");
 	}
-
-	check(false);
-	return TEXT("UNKNOWN");
 }
 
 inline int32 ComponentTypeToNumber(EGLTFJsonComponentType Type)
