[GLTFExporter] refactor: rename FJson*::Write -> FJson*::WriteObject

commit 60b9f8324829a907667fd4ab3b6e6d29c63d8247
Date: Wed, 29 Apr 2020 23:04:56 +0200

    refactor: rename FJson*::Write -> FJson*::WriteObject
    (cherry picked from commit f6a222bb2717a5998c97d5fb86827241d6b65983​)

#skipci

diff --git a/Engine/Plugins/Enterprise/GLTFExporter/MigratedFromGit/Source/GLTFExporter/Public/GLTFJsonAccessor.h b/Engine/Plugins/Enterprise/GLTFExporter/MigratedFromGit/Source/GLTFExporter/Public/GLTFJsonAccessor.h
--- a/Engine/Plugins/Enterprise/GLTFExporter/MigratedFromGit/Source/GLTFExporter/Public/GLTFJsonAccessor.h
+++ b/Engine/Plugins/Enterprise/GLTFExporter/MigratedFromGit/Source/GLTFExporter/Public/GLTFJsonAccessor.h
@@ -34,7 +34,7 @@ struct GLTFEXPORTER_API FGLTFJsonAccessor
 	}
 
 	template <class CharType = TCHAR, class PrintPolicy = TPrettyJsonPrintPolicy<CharType>>
-	void Write(TJsonWriter<CharType, PrintPolicy>& JsonWriter) const
+	void WriteObject(TJsonWriter<CharType, PrintPolicy>& JsonWriter) const
 	{
 		JsonWriter.WriteObjectStart();
 
diff --git a/Engine/Plugins/Enterprise/GLTFExporter/MigratedFromGit/Source/GLTFExporter/Public/GLTFJsonBuffer.h b/Engine/Plugins/Enterprise/GLTFExporter/MigratedFromGit/Source/GLTFExporter/Public/GLTFJsonBuffer.h
--- a/Engine/Plugins/Enterprise/GLTFExporter/MigratedFromGit/Source/GLTFExporter/Public/GLTFJsonBuffer.h
+++ b/Engine/Plugins/Enterprise/GLTFExporter/MigratedFromGit/Source/GLTFExporter/Public/GLTFJsonBuffer.h
@@ -18,7 +18,7 @@ struct GLTFEXPORTER_API FGLTFJsonBuffer
 	}
 
 	template <class CharType = TCHAR, class PrintPolicy = TPrettyJsonPrintPolicy<CharType>>
-	void Write(TJsonWriter<CharType, PrintPolicy>& JsonWriter) const
+	void WriteObject(TJsonWriter<CharType, PrintPolicy>& JsonWriter) const
 	{
 		JsonWriter.WriteObjectStart();
 
diff --git a/Engine/Plugins/Enterprise/GLTFExporter/MigratedFromGit/Source/GLTFExporter/Public/GLTFJsonBufferView.h b/Engine/Plugins/Enterprise/GLTFExporter/MigratedFromGit/Source/GLTFExporter/Public/GLTFJsonBufferView.h
--- a/Engine/Plugins/Enterprise/GLTFExporter/MigratedFromGit/Source/GLTFExporter/Public/GLTFJsonBufferView.h
+++ b/Engine/Plugins/Enterprise/GLTFExporter/MigratedFromGit/Source/GLTFExporter/Public/GLTFJsonBufferView.h
@@ -30,7 +30,7 @@ struct GLTFEXPORTER_API FGLTFJsonBufferView
 	}
 
 	template <class CharType = TCHAR, class PrintPolicy = TPrettyJsonPrintPolicy<CharType>>
-	void Write(TJsonWriter<CharType, PrintPolicy>& JsonWriter) const
+	void WriteObject(TJsonWriter<CharType, PrintPolicy>& JsonWriter) const
 	{
 		JsonWriter.WriteObjectStart();
 
diff --git a/Engine/Plugins/Enterprise/GLTFExporter/MigratedFromGit/Source/GLTFExporter/Public/GLTFJsonMesh.h b/Engine/Plugins/Enterprise/GLTFExporter/MigratedFromGit/Source/GLTFExporter/Public/GLTFJsonMesh.h
--- a/Engine/Plugins/Enterprise/GLTFExporter/MigratedFromGit/Source/GLTFExporter/Public/GLTFJsonMesh.h
+++ b/Engine/Plugins/Enterprise/GLTFExporter/MigratedFromGit/Source/GLTFExporter/Public/GLTFJsonMesh.h
@@ -31,7 +31,7 @@ struct GLTFEXPORTER_API FGLTFJsonAttributes
 	}
 
 	template <class CharType = TCHAR, class PrintPolicy = TPrettyJsonPrintPolicy<CharType>>
-	void Write(TJsonWriter<CharType, PrintPolicy>& JsonWriter) const
+	void WriteObject(TJsonWriter<CharType, PrintPolicy>& JsonWriter) const
 	{
 		JsonWriter.WriteObjectStart();
 
@@ -63,12 +63,12 @@ struct GLTFEXPORTER_API FGLTFJsonPrimitive
 	}
 
 	template <class CharType = TCHAR, class PrintPolicy = TPrettyJsonPrintPolicy<CharType>>
-	void Write(TJsonWriter<CharType, PrintPolicy>& JsonWriter) const
+	void WriteObject(TJsonWriter<CharType, PrintPolicy>& JsonWriter) const
 	{
 		JsonWriter.WriteObjectStart();
 
 		JsonWriter.WriteIdentifierPrefix(TEXT("attributes"));
-		Attributes.Write(JsonWriter);
+		Attributes.WriteObject(JsonWriter);
 
 		if (Indices != INDEX_NONE) JsonWriter.WriteValue(TEXT("indices"), Indices);
 		if (Material != INDEX_NONE) JsonWriter.WriteValue(TEXT("material"), Material);
@@ -85,7 +85,7 @@ struct GLTFEXPORTER_API FGLTFJsonMesh
 	TArray<FGLTFJsonPrimitive> Primitives;
 
 	template <class CharType = TCHAR, class PrintPolicy = TPrettyJsonPrintPolicy<CharType>>
-	void Write(TJsonWriter<CharType, PrintPolicy>& JsonWriter) const
+	void WriteObject(TJsonWriter<CharType, PrintPolicy>& JsonWriter) const
 	{
 		JsonWriter.WriteObjectStart();
 
@@ -94,7 +94,7 @@ struct GLTFEXPORTER_API FGLTFJsonMesh
 		JsonWriter.WriteArrayStart(TEXT("primitives"));
 		for (const FGLTFJsonPrimitive& Primitive : Primitives)
 		{
-			Primitive.Write(JsonWriter);
+			Primitive.WriteObject(JsonWriter);
 		}
 		JsonWriter.WriteArrayEnd();
 
diff --git a/Engine/Plugins/Enterprise/GLTFExporter/MigratedFromGit/Source/GLTFExporter/Public/GLTFJsonNode.h b/Engine/Plugins/Enterprise/GLTFExporter/MigratedFromGit/Source/GLTFExporter/Public/GLTFJsonNode.h
--- a/Engine/Plugins/Enterprise/GLTFExporter/MigratedFromGit/Source/GLTFExporter/Public/GLTFJsonNode.h
+++ b/Engine/Plugins/Enterprise/GLTFExporter/MigratedFromGit/Source/GLTFExporter/Public/GLTFJsonNode.h
@@ -30,7 +30,7 @@ struct GLTFEXPORTER_API FGLTFJsonNode
 	}
 
 	template <class CharType = TCHAR, class PrintPolicy = TPrettyJsonPrintPolicy<CharType>>
-	void Write(TJsonWriter<CharType, PrintPolicy>& JsonWriter) const
+	void WriteObject(TJsonWriter<CharType, PrintPolicy>& JsonWriter) const
 	{
 		JsonWriter.WriteObjectStart();
 
diff --git a/Engine/Plugins/Enterprise/GLTFExporter/MigratedFromGit/Source/GLTFExporter/Public/GLTFJsonRoot.h b/Engine/Plugins/Enterprise/GLTFExporter/MigratedFromGit/Source/GLTFExporter/Public/GLTFJsonRoot.h
--- a/Engine/Plugins/Enterprise/GLTFExporter/MigratedFromGit/Source/GLTFExporter/Public/GLTFJsonRoot.h
+++ b/Engine/Plugins/Enterprise/GLTFExporter/MigratedFromGit/Source/GLTFExporter/Public/GLTFJsonRoot.h
@@ -28,7 +28,7 @@ struct GLTFEXPORTER_API FGLTFJsonAsset
 	}
 
 	template <class CharType = TCHAR, class PrintPolicy = TPrettyJsonPrintPolicy<CharType>>
-	void Write(TJsonWriter<CharType, PrintPolicy>& JsonWriter) const
+	void WriteObject(TJsonWriter<CharType, PrintPolicy>& JsonWriter) const
 	{
 		JsonWriter.WriteObjectStart();
 
@@ -58,12 +58,12 @@ struct GLTFEXPORTER_API FGLTFJsonRoot
 	}
 
 	template <class CharType = TCHAR, class PrintPolicy = TPrettyJsonPrintPolicy<CharType>>
-	void Write(TJsonWriter<CharType, PrintPolicy>& JsonWriter) const
+	void WriteObject(TJsonWriter<CharType, PrintPolicy>& JsonWriter) const
 	{
 		JsonWriter.WriteObjectStart();
 
 		JsonWriter.WriteIdentifierPrefix(TEXT("asset"));
-		Asset.Write(JsonWriter);
+		Asset.WriteObject(JsonWriter);
 
 		if (DefaultScene != INDEX_NONE) JsonWriter.WriteValue(TEXT("scene"), DefaultScene);
 
@@ -72,7 +72,7 @@ struct GLTFEXPORTER_API FGLTFJsonRoot
 			JsonWriter.WriteArrayStart(TEXT("accessors"));
 			for (const FGLTFJsonAccessor& Accessor : Accessors)
 			{
-				Accessor.Write(JsonWriter);
+				Accessor.WriteObject(JsonWriter);
 			}
 			JsonWriter.WriteArrayEnd();
 		}
@@ -82,7 +82,7 @@ struct GLTFEXPORTER_API FGLTFJsonRoot
 			JsonWriter.WriteArrayStart(TEXT("bufferViews"));
 			for (const FGLTFJsonBufferView& BufferView : BufferViews)
 			{
-				BufferView.Write(JsonWriter);
+				BufferView.WriteObject(JsonWriter);
 			}
 			JsonWriter.WriteArrayEnd();
 		}
@@ -92,7 +92,7 @@ struct GLTFEXPORTER_API FGLTFJsonRoot
 			JsonWriter.WriteArrayStart(TEXT("meshes"));
 			for (const FGLTFJsonMesh& Mesh : Meshes)
 			{
-				Mesh.Write(JsonWriter);
+				Mesh.WriteObject(JsonWriter);
 			}
 			JsonWriter.WriteArrayEnd();
 		}
@@ -102,7 +102,7 @@ struct GLTFEXPORTER_API FGLTFJsonRoot
 			JsonWriter.WriteArrayStart(TEXT("nodes"));
 			for (const FGLTFJsonNode& Node : Nodes)
 			{
-				Node.Write(JsonWriter);
+				Node.WriteObject(JsonWriter);
 			}
 			JsonWriter.WriteArrayEnd();
 		}
@@ -112,7 +112,7 @@ struct GLTFEXPORTER_API FGLTFJsonRoot
 			JsonWriter.WriteArrayStart(TEXT("scenes"));
 			for (const FGLTFJsonScene& Scene : Scenes)
 			{
-				Scene.Write(JsonWriter);
+				Scene.WriteObject(JsonWriter);
 			}
 			JsonWriter.WriteArrayEnd();
 		}
@@ -122,7 +122,7 @@ struct GLTFEXPORTER_API FGLTFJsonRoot
 			JsonWriter.WriteArrayStart(TEXT("buffers"));
 			for (const FGLTFJsonBuffer& Buffer : Buffers)
 			{
-				Buffer.Write(JsonWriter);
+				Buffer.WriteObject(JsonWriter);
 			}
 			JsonWriter.WriteArrayEnd();
 		}
@@ -134,7 +134,7 @@ struct GLTFEXPORTER_API FGLTFJsonRoot
 	void Serialize(FArchive* const Archive) const
 	{
 		TSharedRef<TJsonWriter<CharType, PrintPolicy>> JsonWriter = TJsonWriterFactory<CharType, PrintPolicy>::Create(Archive);
-		Write(*JsonWriter);
+		WriteObject(*JsonWriter);
 		JsonWriter->Close();
 	}
 
diff --git a/Engine/Plugins/Enterprise/GLTFExporter/MigratedFromGit/Source/GLTFExporter/Public/GLTFJsonScene.h b/Engine/Plugins/Enterprise/GLTFExporter/MigratedFromGit/Source/GLTFExporter/Public/GLTFJsonScene.h
--- a/Engine/Plugins/Enterprise/GLTFExporter/MigratedFromGit/Source/GLTFExporter/Public/GLTFJsonScene.h
+++ b/Engine/Plugins/Enterprise/GLTFExporter/MigratedFromGit/Source/GLTFExporter/Public/GLTFJsonScene.h
@@ -12,7 +12,7 @@ struct GLTFEXPORTER_API FGLTFJsonScene
 	TArray<FGLTFJsonNodeIndex> Nodes;
 
 	template <class CharType = TCHAR, class PrintPolicy = TPrettyJsonPrintPolicy<CharType>>
-	void Write(TJsonWriter<CharType, PrintPolicy>& JsonWriter) const
+	void WriteObject(TJsonWriter<CharType, PrintPolicy>& JsonWriter) const
 	{
 		JsonWriter.WriteObjectStart();
 
