[GLTFExporter] refactor(test): adjust code according to feedback

commit 7ac18359a0f5bef625708c3cdf1573c7fd6e3c55
Date: Mon, 4 May 2020 12:37:12 +0200

    refactor(test): adjust code according to feedback
    (cherry picked from commit 9983087086c705e048aab7f33c4616c5b16d695f​)

#skipci

diff --git a/Engine/Plugins/Enterprise/GLTFExporter/MigratedFromGit/Source/GLTFExporter/Private/Tests/GLTFExporterTest.cpp b/Engine/Plugins/Enterprise/GLTFExporter/MigratedFromGit/Source/GLTFExporter/Private/Tests/GLTFExporterTest.cpp
--- a/Engine/Plugins/Enterprise/GLTFExporter/MigratedFromGit/Source/GLTFExporter/Private/Tests/GLTFExporterTest.cpp
+++ b/Engine/Plugins/Enterprise/GLTFExporter/MigratedFromGit/Source/GLTFExporter/Private/Tests/GLTFExporterTest.cpp
@@ -11,11 +11,11 @@
 
 #if WITH_DEV_AUTOMATION_TESTS
 
-namespace {
-
-// TODO Turn into config variable?
-const TCHAR* TargetFilePath = TEXT("Tests/targets.json");
-const TCHAR* ParamDelimiter = TEXT("c8a4fd9d525c0ac433fd7d24ce2a3eca");
+namespace
+{
+	// TODO Turn into config variable?
+	const TCHAR* TargetFilePath = TEXT("Tests/targets.json");
+	const TCHAR* ParameterDelimiter = TEXT("c8a4fd9d525c0ac433fd7d24ce2a3eca");
 
 } // anonymous namespace
 
@@ -36,26 +36,26 @@ void FGLTFExporterTest::GetTests(TArray<FString>& OutBeautifiedNames, TArray <FS
 
 		for (int32 TargetJsonArrayElementIndex = 0; TargetJsonArrayElementIndex < TargetJsonElementCount; ++TargetJsonArrayElementIndex) {
 			const TSharedPtr<FJsonObject>& TargetJsonElement = TargetJsonRoot[TargetJsonArrayElementIndex]->AsObject();
-			const FString AssetField = TargetJsonElement->GetStringField("asset");
-			const FString ControlField = TargetJsonElement->GetStringField("control");
-			const FString ParamElements[] = { AssetField, ControlField };
+			const FString InputField = TargetJsonElement->GetStringField("input");
+			const FString ExpectedOutputField = TargetJsonElement->GetStringField("expectedoutput");
+			const FString ParameterElements[] = { InputField, ExpectedOutputField };
 
 			OutBeautifiedNames.Add(FString(TEXT("Test target with index ")) + FString::FromInt(TargetJsonArrayElementIndex));
-			OutTestCommands.Add(FString::Join(ParamElements, ParamDelimiter));
+			OutTestCommands.Add(FString::Join(ParameterElements, ParameterDelimiter));
 		}
 	}
 }
 
 bool FGLTFExporterTest::RunTest(const FString& Parameters)
 {
-	FString AssetPath = "";
-	FString ControlFilePath = "";
-	Parameters.Split(ParamDelimiter, &AssetPath, &ControlFilePath);
+	FString InputAssetPath = "";
+	FString ExpectedOutputFilePath = "";
+	Parameters.Split(ParameterDelimiter, &InputAssetPath, &ExpectedOutputFilePath);
 
-	UObject* ObjectToExport = LoadObject<UObject>(nullptr, *AssetPath, nullptr, LOAD_None, nullptr);
+	UObject* ObjectToExport = LoadObject<UObject>(nullptr, *InputAssetPath, nullptr, LOAD_None, nullptr);
 
 	if (ObjectToExport == nullptr) {
-		AddError(FString::Printf(TEXT("Failed to find test asset %s"), *AssetPath));
+		AddError(FString::Printf(TEXT("Failed to find input asset %s"), *InputAssetPath));
 
 		return false;
 	}
@@ -63,24 +63,24 @@ bool FGLTFExporterTest::RunTest(const FString& Parameters)
 	FBufferArchive BufferArchive = {};
 
 	if (!UExporter::ExportToArchive(ObjectToExport, nullptr, BufferArchive, TEXT("gltf"), 0)) {
-		AddError(FString::Printf(TEXT("Export failed for asset %s"), *AssetPath));
+		AddError(FString::Printf(TEXT("Export failed for input asset %s"), *InputAssetPath));
 
 		return false;
 	}
 
-	const FString ControlFilePathAbsolute = FPaths::ConvertRelativePathToFull(FPaths::ProjectDir(), ControlFilePath);
-	FString ControlFileContent = "";
+	const FString ExpectedOutputFilePathAbsolute = FPaths::ConvertRelativePathToFull(FPaths::ProjectDir(), ExpectedOutputFilePath);
+	FString ExpectedOutputFileContent = "";
 
-	if (!FFileHelper::LoadFileToString(ControlFileContent, *ControlFilePathAbsolute)) {
-		AddError(FString::Printf(TEXT("Failed to find test control file %s"), *ControlFilePathAbsolute));
+	if (!FFileHelper::LoadFileToString(ExpectedOutputFileContent, *ExpectedOutputFilePathAbsolute)) {
+		AddError(FString::Printf(TEXT("Failed to find expected output file %s"), *ExpectedOutputFilePathAbsolute));
 
 		return false;
 	}
 
 	const FString ExportedText = FString(BufferArchive.Num(), UTF8_TO_TCHAR(BufferArchive.GetData()));
 
-	if (!ExportedText.Equals(ControlFileContent)) {
-		AddError(FString::Printf(TEXT("Exported GLTF for the asset %s did not match the control"), *AssetPath));
+	if (!ExportedText.Equals(ExpectedOutputFileContent)) {
+		AddError(FString::Printf(TEXT("Exported GLTF for the asset %s did not match the expected output"), *InputAssetPath));
 
 		return false;
 	}
