[GLTFExporter] refactor(test): change structure of test-data and how it's iterated based on feedback

commit 29d8a7c8ce3a48a2b7f20ef0742077e2d94fdb93
Date: Tue, 5 May 2020 08:33:24 +0200

    refactor(test): change structure of test-data and how it's iterated based on feedback
    (cherry picked from commit 2da4c794ed4fc6db94deaf1ffb3c7091c26748bf​)

#skipci

diff --git a/Engine/Plugins/Enterprise/GLTFExporter/MigratedFromGit/Source/GLTFExporter/Private/Tests/GLTFExporterTest.cpp b/Engine/Plugins/Enterprise/GLTFExporter/MigratedFromGit/Source/GLTFExporter/Private/Tests/GLTFExporterTest.cpp
--- a/Engine/Plugins/Enterprise/GLTFExporter/MigratedFromGit/Source/GLTFExporter/Private/Tests/GLTFExporterTest.cpp
+++ b/Engine/Plugins/Enterprise/GLTFExporter/MigratedFromGit/Source/GLTFExporter/Private/Tests/GLTFExporterTest.cpp
@@ -27,21 +27,19 @@ void FGLTFExporterTest::GetTests(TArray<FString>& OutBeautifiedNames, TArray <FS
 	FString TargetFileContent = "";
 	FFileHelper::LoadFileToString(TargetFileContent, *TargetFilePathAbsolute);
 
-	TArray<TSharedPtr<FJsonValue>> TargetJsonRoot = {};
+	TSharedPtr<FJsonObject> TargetJsonRoot;
 	const TSharedRef<TJsonReader<>> JsonReader = TJsonReaderFactory<>::Create(TargetFileContent);
 
 	if (FJsonSerializer::Deserialize(JsonReader, TargetJsonRoot))
 	{
-		const int32 TargetJsonElementCount = TargetJsonRoot.Num();
-
-		for (int32 TargetJsonArrayElementIndex = 0; TargetJsonArrayElementIndex < TargetJsonElementCount; ++TargetJsonArrayElementIndex)
+		for (const TPair<FString, TSharedPtr<FJsonValue>>& TestJsonValue : TargetJsonRoot->Values)
 		{
-			const TSharedPtr<FJsonObject>& TargetJsonElement = TargetJsonRoot[TargetJsonArrayElementIndex]->AsObject();
+			const TSharedPtr<FJsonObject>& TargetJsonElement = TestJsonValue.Value->AsObject();
 			const FString InputField = TargetJsonElement->GetStringField("input");
 			const FString ExpectedOutputField = TargetJsonElement->GetStringField("expectedoutput");
 			const FString ParameterElements[] = { InputField, ExpectedOutputField };
 
-			OutBeautifiedNames.Add(FString(TEXT("Test target with index ")) + FString::FromInt(TargetJsonArrayElementIndex));
+			OutBeautifiedNames.Add(TestJsonValue.Key);
 			OutTestCommands.Add(FString::Join(ParameterElements, ParameterDelimiter));
 		}
 	}
