[GLTFExporter] refactor(test): remove redundant explicit initialization of variables based on feedback

commit b592a20bfecce720ba046193752a1a76c9b48226
Date: Tue, 5 May 2020 08:50:24 +0200

    refactor(test): remove redundant explicit initialization of variables based on feedback
    (cherry picked from commit bbd793fccd04f14fbede94f43a962621b66b9064​)

#skipci

diff --git a/Engine/Plugins/Enterprise/GLTFExporter/MigratedFromGit/Source/GLTFExporter/Private/Tests/GLTFExporterTest.cpp b/Engine/Plugins/Enterprise/GLTFExporter/MigratedFromGit/Source/GLTFExporter/Private/Tests/GLTFExporterTest.cpp
--- a/Engine/Plugins/Enterprise/GLTFExporter/MigratedFromGit/Source/GLTFExporter/Private/Tests/GLTFExporterTest.cpp
+++ b/Engine/Plugins/Enterprise/GLTFExporter/MigratedFromGit/Source/GLTFExporter/Private/Tests/GLTFExporterTest.cpp
@@ -24,7 +24,7 @@ IMPLEMENT_COMPLEX_AUTOMATION_TEST(FGLTFExporterTest, "Unreal2glTF.Export Test",
 void FGLTFExporterTest::GetTests(TArray<FString>& OutBeautifiedNames, TArray <FString>& OutTestCommands) const
 {
 	const FString TargetFilePathAbsolute = FPaths::ConvertRelativePathToFull(FPaths::ProjectDir(), TargetFilePath);
-	FString TargetFileContent = "";
+	FString TargetFileContent;
 	FFileHelper::LoadFileToString(TargetFileContent, *TargetFilePathAbsolute);
 
 	TSharedPtr<FJsonObject> TargetJsonRoot;
@@ -47,8 +47,8 @@ void FGLTFExporterTest::GetTests(TArray<FString>& OutBeautifiedNames, TArray <FS
 
 bool FGLTFExporterTest::RunTest(const FString& Parameters)
 {
-	FString InputAssetPath = "";
-	FString ExpectedOutputFilePath = "";
+	FString InputAssetPath;
+	FString ExpectedOutputFilePath;
 	Parameters.Split(ParameterDelimiter, &InputAssetPath, &ExpectedOutputFilePath);
 
 	UObject* ObjectToExport = LoadObject<UObject>(nullptr, *InputAssetPath, nullptr, LOAD_None, nullptr);
@@ -60,7 +60,7 @@ bool FGLTFExporterTest::RunTest(const FString& Parameters)
 		return false;
 	}
 
-	FBufferArchive BufferArchive = {};
+	FBufferArchive BufferArchive;
 
 	if (!UExporter::ExportToArchive(ObjectToExport, nullptr, BufferArchive, TEXT("gltf"), 0))
 	{
@@ -70,7 +70,7 @@ bool FGLTFExporterTest::RunTest(const FString& Parameters)
 	}
 
 	const FString ExpectedOutputFilePathAbsolute = FPaths::ConvertRelativePathToFull(FPaths::ProjectDir(), ExpectedOutputFilePath);
-	FString ExpectedOutputFileContent = "";
+	FString ExpectedOutputFileContent;
 
 	if (!FFileHelper::LoadFileToString(ExpectedOutputFileContent, *ExpectedOutputFilePathAbsolute))
 	{
