[GLTFExporter] refactor(test): rename variables and test-definitions file based on feedback

commit 0f2a962aa3c9aa1224411edd948c4e53e10d672a
Date: Tue, 5 May 2020 11:04:35 +0200

    refactor(test): rename variables and test-definitions file based on feedback
    (cherry picked from commit 65cb880f9550ff3e51c517176dd35bbd2699ff17​)

#skipci

diff --git a/Engine/Plugins/Enterprise/GLTFExporter/MigratedFromGit/Source/GLTFExporter/Private/Tests/GLTFExporterTest.cpp b/Engine/Plugins/Enterprise/GLTFExporter/MigratedFromGit/Source/GLTFExporter/Private/Tests/GLTFExporterTest.cpp
--- a/Engine/Plugins/Enterprise/GLTFExporter/MigratedFromGit/Source/GLTFExporter/Private/Tests/GLTFExporterTest.cpp
+++ b/Engine/Plugins/Enterprise/GLTFExporter/MigratedFromGit/Source/GLTFExporter/Private/Tests/GLTFExporterTest.cpp
@@ -14,7 +14,7 @@
 namespace
 {
 	// TODO Turn into config variable?
-	const TCHAR* TargetFilePath = TEXT("Tests/targets.json");
+	const TCHAR* TestDefinitionsFilePath = TEXT("Tests/ExportTests.json");
 	const TCHAR* ParameterDelimiter = TEXT("c8a4fd9d525c0ac433fd7d24ce2a3eca");
 
 } // anonymous namespace
@@ -23,23 +23,23 @@ IMPLEMENT_COMPLEX_AUTOMATION_TEST(FGLTFExporterTest, "Unreal2glTF.Export Test",
 
 void FGLTFExporterTest::GetTests(TArray<FString>& OutBeautifiedNames, TArray <FString>& OutTestCommands) const
 {
-	const FString TargetFilePathAbsolute = FPaths::ConvertRelativePathToFull(FPaths::ProjectDir(), TargetFilePath);
-	FString TargetFileContent;
-	FFileHelper::LoadFileToString(TargetFileContent, *TargetFilePathAbsolute);
+	const FString FilePathAbsolute = FPaths::ConvertRelativePathToFull(FPaths::ProjectDir(), TestDefinitionsFilePath);
+	FString FileContent;
+	FFileHelper::LoadFileToString(FileContent, *FilePathAbsolute);
 
-	TSharedPtr<FJsonObject> TargetJsonRoot;
-	const TSharedRef<TJsonReader<>> JsonReader = TJsonReaderFactory<>::Create(TargetFileContent);
+	TSharedPtr<FJsonObject> JsonRoot;
+	const TSharedRef<TJsonReader<>> JsonReader = TJsonReaderFactory<>::Create(FileContent);
 
-	if (FJsonSerializer::Deserialize(JsonReader, TargetJsonRoot))
+	if (FJsonSerializer::Deserialize(JsonReader, JsonRoot))
 	{
-		for (const TPair<FString, TSharedPtr<FJsonValue>>& TestJsonValue : TargetJsonRoot->Values)
+		for (const TPair<FString, TSharedPtr<FJsonValue>>& TestPair : JsonRoot->Values)
 		{
-			const TSharedPtr<FJsonObject>& TargetJsonElement = TestJsonValue.Value->AsObject();
-			const FString InputField = TargetJsonElement->GetStringField("input");
-			const FString ExpectedOutputField = TargetJsonElement->GetStringField("expectedoutput");
+			const TSharedPtr<FJsonObject>& TestJsonObject = TestPair.Value->AsObject();
+			const FString InputField = TestJsonObject->GetStringField("input");
+			const FString ExpectedOutputField = TestJsonObject->GetStringField("expectedoutput");
 			const FString ParameterElements[] = { InputField, ExpectedOutputField };
 
-			OutBeautifiedNames.Add(TestJsonValue.Key);
+			OutBeautifiedNames.Add(TestPair.Key);
 			OutTestCommands.Add(FString::Join(ParameterElements, ParameterDelimiter));
 		}
 	}
