[GLTFExporter] refactor(test): rename test-class and modify the prefix for tests based on feedback

commit dfbc618cbe28875b5b10095f8a7ca55657e72cff
Date: Fri, 8 May 2020 08:19:21 +0200

    refactor(test): rename test-class and modify the prefix for tests based on feedback
    (cherry picked from commit 501c14b1c5dacf8d174a63c6f861ef4d7d6e420f​)

#skipci

diff --git a/Engine/Plugins/Enterprise/GLTFExporter/MigratedFromGit/Source/GLTFExporter/Private/Tests/GLTFExportTests.cpp b/Engine/Plugins/Enterprise/GLTFExporter/MigratedFromGit/Source/GLTFExporter/Private/Tests/GLTFExportTests.cpp
new file mode 100644
--- /dev/null
+++ b/Engine/Plugins/Enterprise/GLTFExporter/MigratedFromGit/Source/GLTFExporter/Private/Tests/GLTFExportTests.cpp
@@ -0,0 +1,94 @@
+// Copyright 1998-2019 Epic Games, Inc. All Rights Reserved.
+
+#include "CoreMinimal.h"
+#include "Misc/AutomationTest.h"
+#include "Serialization/BufferArchive.h"
+#include "Misc/FileHelper.h"
+#include "Misc/Paths.h"
+#include "Serialization/JsonSerializer.h"
+
+#include "GLTFExporter.h"
+
+#if WITH_DEV_AUTOMATION_TESTS
+
+namespace
+{
+	// TODO Turn into config variable?
+	const TCHAR* TestDefinitionsFilePath = TEXT("Tests/ExportTests.json");
+	const TCHAR* ParameterDelimiter = TEXT("c8a4fd9d525c0ac433fd7d24ce2a3eca");
+
+} // anonymous namespace
+
+IMPLEMENT_COMPLEX_AUTOMATION_TEST(FGLTFExportTests, "Unreal2glTF.ExportTests", EAutomationTestFlags::EditorContext | EAutomationTestFlags::EngineFilter)
+
+void FGLTFExportTests::GetTests(TArray<FString>& OutBeautifiedNames, TArray <FString>& OutTestCommands) const
+{
+	const FString FilePathAbsolute = FPaths::ConvertRelativePathToFull(FPaths::ProjectDir(), TestDefinitionsFilePath);
+	FString FileContent;
+	FFileHelper::LoadFileToString(FileContent, *FilePathAbsolute);
+
+	TSharedPtr<FJsonObject> JsonRoot;
+	const TSharedRef<TJsonReader<>> JsonReader = TJsonReaderFactory<>::Create(FileContent);
+
+	if (FJsonSerializer::Deserialize(JsonReader, JsonRoot))
+	{
+		for (const TPair<FString, TSharedPtr<FJsonValue>>& TestPair : JsonRoot->Values)
+		{
+			const TSharedPtr<FJsonObject>& TestJsonObject = TestPair.Value->AsObject();
+			const FString InputField = TestJsonObject->GetStringField("input");
+			const FString ExpectedOutputField = TestJsonObject->GetStringField("expectedoutput");
+			const FString ParameterElements[] = { InputField, ExpectedOutputField };
+
+			OutBeautifiedNames.Add(TestPair.Key);
+			OutTestCommands.Add(FString::Join(ParameterElements, ParameterDelimiter));
+		}
+	}
+}
+
+bool FGLTFExportTests::RunTest(const FString& Parameters)
+{
+	FString InputAssetPath;
+	FString ExpectedOutputFilePath;
+	Parameters.Split(ParameterDelimiter, &InputAssetPath, &ExpectedOutputFilePath);
+
+	UObject* ObjectToExport = LoadObject<UObject>(nullptr, *InputAssetPath, nullptr, LOAD_None, nullptr);
+
+	if (ObjectToExport == nullptr)
+	{
+		AddError(FString::Printf(TEXT("Failed to find input asset %s"), *InputAssetPath));
+
+		return false;
+	}
+
+	FBufferArchive BufferArchive;
+
+	if (!UExporter::ExportToArchive(ObjectToExport, nullptr, BufferArchive, TEXT("gltf"), 0))
+	{
+		AddError(FString::Printf(TEXT("Export failed for input asset %s"), *InputAssetPath));
+
+		return false;
+	}
+
+	const FString ExpectedOutputFilePathAbsolute = FPaths::ConvertRelativePathToFull(FPaths::ProjectDir(), ExpectedOutputFilePath);
+	FString ExpectedOutputFileContent;
+
+	if (!FFileHelper::LoadFileToString(ExpectedOutputFileContent, *ExpectedOutputFilePathAbsolute))
+	{
+		AddError(FString::Printf(TEXT("Failed to find expected output file %s"), *ExpectedOutputFilePathAbsolute));
+
+		return false;
+	}
+
+	const FString ExportedText = FString(BufferArchive.Num(), UTF8_TO_TCHAR(BufferArchive.GetData()));
+
+	if (!ExportedText.Equals(ExpectedOutputFileContent))
+	{
+		AddError(FString::Printf(TEXT("Exported GLTF for the asset %s did not match the expected output"), *InputAssetPath));
+
+		return false;
+	}
+
+	return true;
+}
+
+#endif //WITH_DEV_AUTOMATION_TESTS
diff --git a/Engine/Plugins/Enterprise/GLTFExporter/MigratedFromGit/Source/GLTFExporter/Private/Tests/GLTFExporterTest.cpp b/Engine/Plugins/Enterprise/GLTFExporter/MigratedFromGit/Source/GLTFExporter/Private/Tests/GLTFExporterTest.cpp
deleted file mode 100644
