[GLTFExporter] fix(test): replace the custom argument-delimiter with ";" based on feedback

commit 36165f3bf588c09fa38a6379ba61cb6f412d3798
Date: Fri, 8 May 2020 08:20:52 +0200

    fix(test): replace the custom argument-delimiter with ";" based on feedback
    (cherry picked from commit e72ddd0fe41a238699399a2f0f1eeb4d4b8b21fd​)

#skipci

diff --git a/Engine/Plugins/Enterprise/GLTFExporter/MigratedFromGit/Source/GLTFExporter/Private/Tests/GLTFExportTests.cpp b/Engine/Plugins/Enterprise/GLTFExporter/MigratedFromGit/Source/GLTFExporter/Private/Tests/GLTFExportTests.cpp
--- a/Engine/Plugins/Enterprise/GLTFExporter/MigratedFromGit/Source/GLTFExporter/Private/Tests/GLTFExportTests.cpp
+++ b/Engine/Plugins/Enterprise/GLTFExporter/MigratedFromGit/Source/GLTFExporter/Private/Tests/GLTFExportTests.cpp
@@ -15,7 +15,7 @@ namespace
 {
 	// TODO Turn into config variable?
 	const TCHAR* TestDefinitionsFilePath = TEXT("Tests/ExportTests.json");
-	const TCHAR* ParameterDelimiter = TEXT("c8a4fd9d525c0ac433fd7d24ce2a3eca");
+	const TCHAR* ParameterDelimiter = TEXT(";");
 
 } // anonymous namespace
 
