[GLTFExporter] refactor: add missing "b"-prefix for some boolean variables and parameters

commit 7b4fdedca73f1ac03b403ad99380bcc943fde7aa
Date: Tue, 5 May 2020 00:39:11 +0200

    refactor: add missing "b"-prefix for some boolean variables and parameters
    (cherry picked from commit 3b8d1f9d3ff43ba8ce99a4cf48a589c59ea96e11​)

#skipci

diff --git a/Engine/Plugins/Enterprise/GLTFExporter/MigratedFromGit/Source/GLTFExporter/Private/GLTFExporter.cpp b/Engine/Plugins/Enterprise/GLTFExporter/MigratedFromGit/Source/GLTFExporter/Private/GLTFExporter.cpp
--- a/Engine/Plugins/Enterprise/GLTFExporter/MigratedFromGit/Source/GLTFExporter/Private/GLTFExporter.cpp
+++ b/Engine/Plugins/Enterprise/GLTFExporter/MigratedFromGit/Source/GLTFExporter/Private/GLTFExporter.cpp
@@ -19,16 +19,16 @@ UGLTFExporter::UGLTFExporter(const FObjectInitializer& ObjectInitializer)
 
 bool UGLTFExporter::FillExportOptions()
 {
-	bool ExportAll = GetBatchMode() && !GetShowExportOption();
-	bool ExportCancel = false;
+	bool bExportAll = GetBatchMode() && !GetShowExportOption();
+	bool bExportCancel = false;
 
-	ExportOptions->FillOptions(GetBatchMode(), GetShowExportOption(), UExporter::CurrentFilename, ExportCancel, ExportAll);
-	if (ExportCancel)
+	ExportOptions->FillOptions(GetBatchMode(), GetShowExportOption(), UExporter::CurrentFilename, bExportCancel, bExportAll);
+	if (bExportCancel)
 	{
 		SetCancelBatch(GetBatchMode());
 		return false;
 	}
 
-	SetShowExportOption(!ExportAll);
+	SetShowExportOption(!bExportAll);
 	return true;
 }
diff --git a/Engine/Plugins/Enterprise/GLTFExporter/MigratedFromGit/Source/GLTFExporter/Public/GLTFExportOptions.h b/Engine/Plugins/Enterprise/GLTFExporter/MigratedFromGit/Source/GLTFExporter/Public/GLTFExportOptions.h
--- a/Engine/Plugins/Enterprise/GLTFExporter/MigratedFromGit/Source/GLTFExporter/Public/GLTFExportOptions.h
+++ b/Engine/Plugins/Enterprise/GLTFExporter/MigratedFromGit/Source/GLTFExporter/Public/GLTFExportOptions.h
@@ -126,5 +126,5 @@ public:
 	*
 	* The function is saving the dialog state in a user ini file and reload it from there. It is not changing the CDO.
 	*/
-	void FillOptions(bool BatchMode, bool bShowOptionDialog, const FString& FullPath, bool& OutOperationCanceled, bool& bOutExportAll);
+	void FillOptions(bool bBatchMode, bool bShowOptionDialog, const FString& FullPath, bool& bOutOperationCanceled, bool& bOutExportAll);
 };
diff --git a/Engine/Plugins/Enterprise/GLTFExporter/MigratedFromGit/Source/GLTFExporter/Public/GLTFExportOptionsWindow.h b/Engine/Plugins/Enterprise/GLTFExporter/MigratedFromGit/Source/GLTFExporter/Public/GLTFExportOptionsWindow.h
--- a/Engine/Plugins/Enterprise/GLTFExporter/MigratedFromGit/Source/GLTFExporter/Public/GLTFExportOptionsWindow.h
+++ b/Engine/Plugins/Enterprise/GLTFExporter/MigratedFromGit/Source/GLTFExporter/Public/GLTFExportOptionsWindow.h
@@ -46,7 +46,7 @@ public:
 	SGLTFExportOptionsWindow();
 
 public:
-	static void ShowDialog(UGLTFExportOptions* ExportOptions, const FString& FullPath, bool BatchMode, bool& bOutOperationCanceled, bool& bOutExportAll);
+	static void ShowDialog(UGLTFExportOptions* ExportOptions, const FString& FullPath, bool bBatchMode, bool& bOutOperationCanceled, bool& bOutExportAll);
 
 private:
 	FReply OnResetToDefaultClick() const;
@@ -56,6 +56,6 @@ private:
 	TSharedPtr<class IDetailsView> DetailsView;
 	TWeakPtr< SWindow > WidgetWindow;
 	TSharedPtr< SButton > ImportButton;
-	bool			bShouldExport;
-	bool			bShouldExportAll;
+	bool bShouldExport;
+	bool bShouldExportAll;
 };
diff --git a/Engine/Plugins/Enterprise/GLTFExporter/MigratedFromGit/Source/GLTFExporter/Public/GLTFJsonAccessor.h b/Engine/Plugins/Enterprise/GLTFExporter/MigratedFromGit/Source/GLTFExporter/Public/GLTFJsonAccessor.h
--- a/Engine/Plugins/Enterprise/GLTFExporter/MigratedFromGit/Source/GLTFExporter/Public/GLTFJsonAccessor.h
+++ b/Engine/Plugins/Enterprise/GLTFExporter/MigratedFromGit/Source/GLTFExporter/Public/GLTFJsonAccessor.h
@@ -15,7 +15,7 @@ struct GLTFEXPORTER_API FGLTFJsonAccessor
 	int32                    Count;
 	EGLTFJsonAccessorType    Type;
 	EGLTFJsonComponentType   ComponentType;
-	bool                     Normalized;
+	bool                     bNormalized;
 
 	int32 MinMaxLength;
 	float Min[16];
@@ -26,7 +26,7 @@ struct GLTFEXPORTER_API FGLTFJsonAccessor
 		, Count(0)
 		, Type(EGLTFJsonAccessorType::None)
 		, ComponentType(EGLTFJsonComponentType::None)
-		, Normalized(false)
+		, bNormalized(false)
 		, MinMaxLength(0)
 		, Min{0}
 		, Max{0}
