[GLTFExporter] feat: fix vertex color issue (accessor should be normalized)

commit ab6c7db8afdaaa0af02231ed12ec688ecfc96bd4
Date: Tue, 5 May 2020 00:48:13 +0200

    feat: fix vertex color issue (accessor should be normalized)
    (cherry picked from commit b41d86090c1e275f84aa6d0b835e45281fe14e3a​)

#skipci

diff --git a/Engine/Plugins/Enterprise/GLTFExporter/MigratedFromGit/Source/GLTFExporter/Private/GLTFMeshBuilder.cpp b/Engine/Plugins/Enterprise/GLTFExporter/MigratedFromGit/Source/GLTFExporter/Private/GLTFMeshBuilder.cpp
--- a/Engine/Plugins/Enterprise/GLTFExporter/MigratedFromGit/Source/GLTFExporter/Private/GLTFMeshBuilder.cpp
+++ b/Engine/Plugins/Enterprise/GLTFExporter/MigratedFromGit/Source/GLTFExporter/Private/GLTFMeshBuilder.cpp
@@ -178,6 +178,7 @@ FGLTFJsonAccessorIndex FGLTFMeshBuilder::AddAccessorForColors(FGLTFContainerBuil
 	Accessor.ComponentType = EGLTFJsonComponentType::U8;
 	Accessor.Count = Colors.Num();
 	Accessor.Type = EGLTFJsonAccessorType::Vec4;
+	Accessor.bNormalized = true;
 
 	return Container.JsonRoot.Accessors.Add(Accessor);
 }
diff --git a/Engine/Plugins/Enterprise/GLTFExporter/MigratedFromGit/Source/GLTFExporter/Public/GLTFJsonAccessor.h b/Engine/Plugins/Enterprise/GLTFExporter/MigratedFromGit/Source/GLTFExporter/Public/GLTFJsonAccessor.h
--- a/Engine/Plugins/Enterprise/GLTFExporter/MigratedFromGit/Source/GLTFExporter/Public/GLTFJsonAccessor.h
+++ b/Engine/Plugins/Enterprise/GLTFExporter/MigratedFromGit/Source/GLTFExporter/Public/GLTFJsonAccessor.h
@@ -44,6 +44,7 @@ struct GLTFEXPORTER_API FGLTFJsonAccessor
 		JsonWriter.WriteValue(TEXT("count"), Count);
 		JsonWriter.WriteValue(TEXT("type"), AccessorTypeToString(Type));
 		JsonWriter.WriteValue(TEXT("componentType"), ComponentTypeToNumber(ComponentType));
+		if (bNormalized) JsonWriter.WriteValue(TEXT("componentType"), bNormalized);
 
 		if (MinMaxLength > 0)
 		{
