[GLTFExporter] refactor: change some implicit nullptr checks to explicit

commit 015d22dba0888260f59669a7ff31b012ae48e15d
Date: Tue, 5 May 2020 00:49:47 +0200

    refactor: change some implicit nullptr checks to explicit
    (cherry picked from commit e8e24d61bc3ab6a5187c0a4f5a715c930d194f53​)

#skipci

diff --git a/Engine/Plugins/Enterprise/GLTFExporter/MigratedFromGit/Source/GLTFExporter/Private/GLTFExportOptions.cpp b/Engine/Plugins/Enterprise/GLTFExporter/MigratedFromGit/Source/GLTFExporter/Private/GLTFExportOptions.cpp
--- a/Engine/Plugins/Enterprise/GLTFExporter/MigratedFromGit/Source/GLTFExporter/Private/GLTFExportOptions.cpp
+++ b/Engine/Plugins/Enterprise/GLTFExporter/MigratedFromGit/Source/GLTFExporter/Private/GLTFExportOptions.cpp
@@ -57,7 +57,7 @@ void UGLTFExportOptions::LoadOptions()
 		FString Key = Property->GetName();
 
 		UArrayProperty* Array = dynamic_cast<UArrayProperty*>(Property);
-		if (Array)
+		if (Array != nullptr)
 		{
 			FConfigSection* Sec = GConfig->GetSectionPrivate(*Section, false, true, *GEditorPerProjectIni);
 			if (Sec != nullptr)
@@ -144,7 +144,7 @@ void UGLTFExportOptions::SaveOptions()
 		FString Key = Property->GetName();
 
 		UArrayProperty* Array = dynamic_cast<UArrayProperty*>(Property);
-		if (Array)
+		if (Array != nullptr)
 		{
 			FConfigSection* Sec = GConfig->GetSectionPrivate(*Section, true, false, *GEditorPerProjectIni);
 			check(Sec);
