[GLTFExporter] fix: fix incorrect flag-name for normalized vertex-color

commit fa05c6293c677b665a00d04b2b76e0c7141603fc
Date: Tue, 5 May 2020 14:58:42 +0200

    fix: fix incorrect flag-name for normalized vertex-color
    (cherry picked from commit a8c2eb8aeb67ac1772e3961baf4997b285edb0c1​)

#skipci

diff --git a/Engine/Plugins/Enterprise/GLTFExporter/MigratedFromGit/Source/GLTFExporter/Public/GLTFJsonAccessor.h b/Engine/Plugins/Enterprise/GLTFExporter/MigratedFromGit/Source/GLTFExporter/Public/GLTFJsonAccessor.h
--- a/Engine/Plugins/Enterprise/GLTFExporter/MigratedFromGit/Source/GLTFExporter/Public/GLTFJsonAccessor.h
+++ b/Engine/Plugins/Enterprise/GLTFExporter/MigratedFromGit/Source/GLTFExporter/Public/GLTFJsonAccessor.h
@@ -44,7 +44,7 @@ struct GLTFEXPORTER_API FGLTFJsonAccessor
 		JsonWriter.WriteValue(TEXT("count"), Count);
 		JsonWriter.WriteValue(TEXT("type"), AccessorTypeToString(Type));
 		JsonWriter.WriteValue(TEXT("componentType"), ComponentTypeToNumber(ComponentType));
-		if (bNormalized) JsonWriter.WriteValue(TEXT("componentType"), bNormalized);
+		if (bNormalized) JsonWriter.WriteValue(TEXT("normalized"), bNormalized);
 
 		if (MinMaxLength > 0)
 		{
