[GLTFExporter] refactor: add GLTFContainerBuilder::AddAccessor, AddMesh, AddNode, and AddScene

commit 9a2aea7aa51bd3aec456fec9dce0e9ec49545fc2
Date: Sun, 10 May 2020 11:54:07 +0200

    refactor: add GLTFContainerBuilder::AddAccessor, AddMesh, AddNode, and AddScene
    (cherry picked from commit 6f0e698b5ed4cc35b9ff95f2e254d3969b47c8ef​)

#skipci

diff --git a/Engine/Plugins/Enterprise/GLTFExporter/MigratedFromGit/Source/GLTFExporter/Private/GLTFContainerBuilder.cpp b/Engine/Plugins/Enterprise/GLTFExporter/MigratedFromGit/Source/GLTFExporter/Private/GLTFContainerBuilder.cpp
--- a/Engine/Plugins/Enterprise/GLTFExporter/MigratedFromGit/Source/GLTFExporter/Private/GLTFContainerBuilder.cpp
+++ b/Engine/Plugins/Enterprise/GLTFExporter/MigratedFromGit/Source/GLTFExporter/Private/GLTFContainerBuilder.cpp
@@ -8,6 +8,26 @@ FGLTFContainerBuilder::FGLTFContainerBuilder()
 {
 }
 
+FGLTFJsonAccessorIndex FGLTFContainerBuilder::AddAccessor(const FGLTFJsonAccessor& JsonAccessor)
+{
+	return JsonRoot.Accessors.Add(JsonAccessor);
+}
+
+FGLTFJsonMeshIndex FGLTFContainerBuilder::AddMesh(const FGLTFJsonMesh& JsonMesh)
+{
+	return JsonRoot.Meshes.Add(JsonMesh);
+}
+
+FGLTFJsonNodeIndex FGLTFContainerBuilder::AddNode(const FGLTFJsonNode& JsonNode)
+{
+	return JsonRoot.Nodes.Add(JsonNode);
+}
+
+FGLTFJsonSceneIndex FGLTFContainerBuilder::AddScene(const FGLTFJsonScene& JsonScene)
+{
+	return JsonRoot.Scenes.Add(JsonScene);
+}
+
 FGLTFJsonBufferViewIndex FGLTFContainerBuilder::AddBufferView(const void* RawData, uint64 ByteLength, const FString& Name, EGLTFJsonBufferTarget BufferTarget)
 {
 	return BufferBuilder.AddBufferView(RawData, ByteLength, Name, BufferTarget);
diff --git a/Engine/Plugins/Enterprise/GLTFExporter/MigratedFromGit/Source/GLTFExporter/Private/GLTFMeshBuilder.cpp b/Engine/Plugins/Enterprise/GLTFExporter/MigratedFromGit/Source/GLTFExporter/Private/GLTFMeshBuilder.cpp
--- a/Engine/Plugins/Enterprise/GLTFExporter/MigratedFromGit/Source/GLTFExporter/Private/GLTFMeshBuilder.cpp
+++ b/Engine/Plugins/Enterprise/GLTFExporter/MigratedFromGit/Source/GLTFExporter/Private/GLTFMeshBuilder.cpp
@@ -33,7 +33,7 @@ FGLTFJsonAccessorIndex FGLTFSectionBuilder::AddAccessorForIndices(FGLTFContainer
 	Accessor.Count = Indices.Num();
 	Accessor.Type = EGLTFJsonAccessorType::Scalar;
 
-	return Container.JsonRoot.Accessors.Add(Accessor);
+	return Container.AddAccessor(Accessor);
 }
 
 FGLTFMeshBuilder::FGLTFMeshBuilder(const UStaticMesh* StaticMesh, int32 LODIndex)
@@ -138,8 +138,8 @@ FGLTFJsonAccessorIndex FGLTFMeshBuilder::AddAccessorForPositions(FGLTFContainerB
 	Accessor.Min[1] = Min.Y;
 	Accessor.Min[2] = Min.Z;
 	Accessor.MinMaxLength = 3;
-
-	return Container.JsonRoot.Accessors.Add(Accessor);
+	
+	return Container.AddAccessor(Accessor);
 }
 
 FGLTFJsonAccessorIndex FGLTFMeshBuilder::AddAccessorForNormals(FGLTFContainerBuilder& Container) const
@@ -158,8 +158,8 @@ FGLTFJsonAccessorIndex FGLTFMeshBuilder::AddAccessorForNormals(FGLTFContainerBui
 	Accessor.ComponentType = EGLTFJsonComponentType::F32;
 	Accessor.Count = Normals.Num();
 	Accessor.Type = EGLTFJsonAccessorType::Vec3;
-
-	return Container.JsonRoot.Accessors.Add(Accessor);
+	
+	return Container.AddAccessor(Accessor);
 }
 
 FGLTFJsonAccessorIndex FGLTFMeshBuilder::AddAccessorForColors(FGLTFContainerBuilder& Container) const
@@ -179,8 +179,8 @@ FGLTFJsonAccessorIndex FGLTFMeshBuilder::AddAccessorForColors(FGLTFContainerBuil
 	Accessor.Count = Colors.Num();
 	Accessor.Type = EGLTFJsonAccessorType::Vec4;
 	Accessor.bNormalized = true;
-
-	return Container.JsonRoot.Accessors.Add(Accessor);
+	
+	return Container.AddAccessor(Accessor);
 }
 
 FGLTFJsonAccessorIndex FGLTFMeshBuilder::AddAccessorForTangents(FGLTFContainerBuilder& Container) const
@@ -199,8 +199,8 @@ FGLTFJsonAccessorIndex FGLTFMeshBuilder::AddAccessorForTangents(FGLTFContainerBu
 	Accessor.ComponentType = EGLTFJsonComponentType::F32;
 	Accessor.Count = Tangents.Num();
 	Accessor.Type = EGLTFJsonAccessorType::Vec4;
-
-	return Container.JsonRoot.Accessors.Add(Accessor);
+	
+	return Container.AddAccessor(Accessor);
 }
 
 FGLTFJsonAccessorIndex FGLTFMeshBuilder::AddAccessorForUV0s(FGLTFContainerBuilder& Container) const
@@ -219,8 +219,8 @@ FGLTFJsonAccessorIndex FGLTFMeshBuilder::AddAccessorForUV0s(FGLTFContainerBuilde
 	Accessor.ComponentType = EGLTFJsonComponentType::F32;
 	Accessor.Count = UV0s.Num();
 	Accessor.Type = EGLTFJsonAccessorType::Vec2;
-
-	return Container.JsonRoot.Accessors.Add(Accessor);
+	
+	return Container.AddAccessor(Accessor);
 }
 
 FGLTFJsonAccessorIndex FGLTFMeshBuilder::AddAccessorForUV1s(FGLTFContainerBuilder& Container) const
@@ -239,8 +239,8 @@ FGLTFJsonAccessorIndex FGLTFMeshBuilder::AddAccessorForUV1s(FGLTFContainerBuilde
 	Accessor.ComponentType = EGLTFJsonComponentType::F32;
 	Accessor.Count = UV1s.Num();
 	Accessor.Type = EGLTFJsonAccessorType::Vec2;
-
-	return Container.JsonRoot.Accessors.Add(Accessor);
+	
+	return Container.AddAccessor(Accessor);
 }
 
 FGLTFJsonMeshIndex FGLTFMeshBuilder::AddMesh(FGLTFContainerBuilder& Container) const
@@ -264,6 +264,6 @@ FGLTFJsonMeshIndex FGLTFMeshBuilder::AddMesh(FGLTFContainerBuilder& Container) c
 		Primitive.Indices = Section.AddAccessorForIndices(Container);
 		Mesh.Primitives.Add(Primitive);
 	}
-
-	return Container.JsonRoot.Meshes.Add(Mesh);
+	
+	return Container.AddMesh(Mesh);
 }
diff --git a/Engine/Plugins/Enterprise/GLTFExporter/MigratedFromGit/Source/GLTFExporter/Private/GLTFStaticMeshExporter.cpp b/Engine/Plugins/Enterprise/GLTFExporter/MigratedFromGit/Source/GLTFExporter/Private/GLTFStaticMeshExporter.cpp
--- a/Engine/Plugins/Enterprise/GLTFExporter/MigratedFromGit/Source/GLTFExporter/Private/GLTFStaticMeshExporter.cpp
+++ b/Engine/Plugins/Enterprise/GLTFExporter/MigratedFromGit/Source/GLTFExporter/Private/GLTFStaticMeshExporter.cpp
@@ -27,11 +27,11 @@ bool UGLTFStaticMeshExporter::ExportBinary(UObject* Object, const TCHAR* Type, F
 
 	FGLTFJsonNode Node;
 	Node.Mesh = MeshIndex;
-	FGLTFJsonNodeIndex NodeIndex = Container.JsonRoot.Nodes.Add(Node);
+	FGLTFJsonNodeIndex NodeIndex = Container.AddNode(Node);
 
 	FGLTFJsonScene Scene;
 	Scene.Nodes.Add(NodeIndex);
-	FGLTFJsonSceneIndex SceneIndex = Container.JsonRoot.Scenes.Add(Scene);
+	FGLTFJsonSceneIndex SceneIndex = Container.AddScene(Scene);
 
 	Container.JsonRoot.DefaultScene = SceneIndex;
 
diff --git a/Engine/Plugins/Enterprise/GLTFExporter/MigratedFromGit/Source/GLTFExporter/Public/GLTFContainerBuilder.h b/Engine/Plugins/Enterprise/GLTFExporter/MigratedFromGit/Source/GLTFExporter/Public/GLTFContainerBuilder.h
--- a/Engine/Plugins/Enterprise/GLTFExporter/MigratedFromGit/Source/GLTFExporter/Public/GLTFContainerBuilder.h
+++ b/Engine/Plugins/Enterprise/GLTFExporter/MigratedFromGit/Source/GLTFExporter/Public/GLTFContainerBuilder.h
@@ -13,6 +13,11 @@ struct GLTFEXPORTER_API FGLTFContainerBuilder
 
 	FGLTFContainerBuilder();
 
+	FGLTFJsonAccessorIndex AddAccessor(const FGLTFJsonAccessor& JsonAccessor);
+	FGLTFJsonMeshIndex AddMesh(const FGLTFJsonMesh& JsonMesh);
+	FGLTFJsonNodeIndex AddNode(const FGLTFJsonNode& JsonNode);
+	FGLTFJsonSceneIndex AddScene(const FGLTFJsonScene& JsonScene);
+
 	FGLTFJsonBufferViewIndex AddBufferView(const void* RawData, uint64 ByteLength, const FString& Name = TEXT(""), EGLTFJsonBufferTarget BufferTarget = EGLTFJsonBufferTarget::ArrayBuffer);
 
 	template <class ElementType>
