[GLTFExporter] refactor: add FGLTFContainerBuilder::AddBuffer and AddBufferView

commit af199c630aaf96760f0e92509caffef1c8a69e88
Date: Sun, 10 May 2020 12:03:53 +0200

    refactor: add FGLTFContainerBuilder::AddBuffer and AddBufferView
    (cherry picked from commit f0afdf5c1af641cf405829daacad5db8568e3ba5​)

#skipci

diff --git a/Engine/Plugins/Enterprise/GLTFExporter/MigratedFromGit/Source/GLTFExporter/Private/GLTFBufferBuilder.cpp b/Engine/Plugins/Enterprise/GLTFExporter/MigratedFromGit/Source/GLTFExporter/Private/GLTFBufferBuilder.cpp
--- a/Engine/Plugins/Enterprise/GLTFExporter/MigratedFromGit/Source/GLTFExporter/Private/GLTFBufferBuilder.cpp
+++ b/Engine/Plugins/Enterprise/GLTFExporter/MigratedFromGit/Source/GLTFExporter/Private/GLTFBufferBuilder.cpp
@@ -4,34 +4,33 @@
 #include "GLTFMeshBuilder.h"
 #include "Misc/Base64.h"
 
-FGLTFBufferBuilder::FGLTFBufferBuilder(FGLTFJsonRoot& JsonRoot)
-	: JsonRoot(JsonRoot)
+FGLTFBufferBuilder::FGLTFBufferBuilder(FGLTFJsonBufferIndex BufferIndex)
+	: BufferIndex(BufferIndex)
 {
-	MergedBufferIndex = JsonRoot.Buffers.AddDefaulted(1);
 }
 
-FGLTFJsonBufferViewIndex FGLTFBufferBuilder::AddBufferView(const void* RawData, uint64 ByteLength, const FString& Name, EGLTFJsonBufferTarget BufferTarget)
+FGLTFJsonBufferViewIndex FGLTFBufferBuilder::AddBufferView(FGLTFContainerBuilder& Container, const void* RawData, uint64 ByteLength, const FString& Name, EGLTFJsonBufferTarget BufferTarget)
 {
 	FGLTFJsonBufferView BufferView;
 	BufferView.Name = Name;
-	BufferView.Buffer = MergedBufferIndex;
-	BufferView.ByteOffset = MergedBufferData.Num();
+	BufferView.Buffer = BufferIndex;
+	BufferView.ByteOffset = BufferData.Num();
 	BufferView.ByteLength = ByteLength;
 	BufferView.Target = BufferTarget;
 
-	MergedBufferData.Append(static_cast<const uint8*>(RawData), ByteLength);
+	BufferData.Append(static_cast<const uint8*>(RawData), ByteLength);
 
-	return JsonRoot.BufferViews.Add(BufferView);
+	return Container.AddBufferView(BufferView);
 }
 
-void FGLTFBufferBuilder::UpdateMergedBuffer()
+void FGLTFBufferBuilder::UpdateBuffer(FGLTFJsonBuffer& JsonBuffer)
 {
-	FGLTFJsonBuffer& Buffer = JsonRoot.Buffers[MergedBufferIndex];
-	if (Buffer.ByteLength != MergedBufferData.Num())
+	const int32 ByteLength = BufferData.Num();
+	if (JsonBuffer.ByteLength != ByteLength)
 	{
-		const FString DataBase64 = FBase64::Encode(MergedBufferData.GetData(), MergedBufferData.Num());
-
-		Buffer.URI = TEXT("data:application/octet-stream;base64,") + DataBase64;
-		Buffer.ByteLength = MergedBufferData.Num();
+		JsonBuffer.ByteLength = ByteLength;
+		
+		const FString DataBase64 = FBase64::Encode(BufferData.GetData(), ByteLength);
+		JsonBuffer.URI = TEXT("data:application/octet-stream;base64,") + DataBase64;
 	}
 }
diff --git a/Engine/Plugins/Enterprise/GLTFExporter/MigratedFromGit/Source/GLTFExporter/Private/GLTFContainerBuilder.cpp b/Engine/Plugins/Enterprise/GLTFExporter/MigratedFromGit/Source/GLTFExporter/Private/GLTFContainerBuilder.cpp
--- a/Engine/Plugins/Enterprise/GLTFExporter/MigratedFromGit/Source/GLTFExporter/Private/GLTFContainerBuilder.cpp
+++ b/Engine/Plugins/Enterprise/GLTFExporter/MigratedFromGit/Source/GLTFExporter/Private/GLTFContainerBuilder.cpp
@@ -4,7 +4,7 @@
 #include "GLTFMeshBuilder.h"
 
 FGLTFContainerBuilder::FGLTFContainerBuilder()
-	: BufferBuilder(JsonRoot)
+	: BufferBuilder(AddBuffer(FGLTFJsonBuffer()))
 {
 }
 
@@ -13,6 +13,16 @@ FGLTFJsonAccessorIndex FGLTFContainerBuilder::AddAccessor(const FGLTFJsonAccesso
 	return JsonRoot.Accessors.Add(JsonAccessor);
 }
 
+FGLTFJsonBufferIndex FGLTFContainerBuilder::AddBuffer(const FGLTFJsonBuffer& JsonBuffer)
+{
+	return JsonRoot.Buffers.Add(JsonBuffer);
+}
+
+FGLTFJsonBufferViewIndex FGLTFContainerBuilder::AddBufferView(const FGLTFJsonBufferView& JsonBufferView)
+{
+	return JsonRoot.BufferViews.Add(JsonBufferView);
+}
+
 FGLTFJsonMeshIndex FGLTFContainerBuilder::AddMesh(const FGLTFJsonMesh& JsonMesh)
 {
 	return JsonRoot.Meshes.Add(JsonMesh);
@@ -30,12 +40,13 @@ FGLTFJsonSceneIndex FGLTFContainerBuilder::AddScene(const FGLTFJsonScene& JsonSc
 
 FGLTFJsonBufferViewIndex FGLTFContainerBuilder::AddBufferView(const void* RawData, uint64 ByteLength, const FString& Name, EGLTFJsonBufferTarget BufferTarget)
 {
-	return BufferBuilder.AddBufferView(RawData, ByteLength, Name, BufferTarget);
+	return BufferBuilder.AddBufferView(*this, RawData, ByteLength, Name, BufferTarget);
 }
 
 void FGLTFContainerBuilder::Serialize(FArchive& Archive)
 {
-	BufferBuilder.UpdateMergedBuffer();
+	FGLTFJsonBuffer& JsonBuffer = JsonRoot.Buffers[BufferBuilder.BufferIndex];
+	BufferBuilder.UpdateBuffer(JsonBuffer);
 	JsonRoot.Serialize(&Archive, true);
 }
 
diff --git a/Engine/Plugins/Enterprise/GLTFExporter/MigratedFromGit/Source/GLTFExporter/Public/GLTFBufferBuilder.h b/Engine/Plugins/Enterprise/GLTFExporter/MigratedFromGit/Source/GLTFExporter/Public/GLTFBufferBuilder.h
--- a/Engine/Plugins/Enterprise/GLTFExporter/MigratedFromGit/Source/GLTFExporter/Public/GLTFBufferBuilder.h
+++ b/Engine/Plugins/Enterprise/GLTFExporter/MigratedFromGit/Source/GLTFExporter/Public/GLTFBufferBuilder.h
@@ -5,16 +5,16 @@
 #include "GLTFJsonRoot.h"
 #include "Engine.h"
 
+struct FGLTFContainerBuilder;
+
 struct GLTFEXPORTER_API FGLTFBufferBuilder
 {
-	FGLTFJsonRoot& JsonRoot;
-
-	FGLTFJsonBufferIndex MergedBufferIndex;
-	TArray<uint8> MergedBufferData;
+	FGLTFJsonBufferIndex BufferIndex;
+	TArray<uint8> BufferData;
 
-	FGLTFBufferBuilder(FGLTFJsonRoot& JsonRoot);
+	FGLTFBufferBuilder(FGLTFJsonBufferIndex BufferIndex);
 
-	FGLTFJsonBufferViewIndex AddBufferView(const void* RawData, uint64 ByteLength, const FString& Name, EGLTFJsonBufferTarget BufferTarget);
+	FGLTFJsonBufferViewIndex AddBufferView(FGLTFContainerBuilder& Container, const void* RawData, uint64 ByteLength, const FString& Name, EGLTFJsonBufferTarget BufferTarget);
 
-	void UpdateMergedBuffer();
+	void UpdateBuffer(FGLTFJsonBuffer& JsonBuffer);
 };
diff --git a/Engine/Plugins/Enterprise/GLTFExporter/MigratedFromGit/Source/GLTFExporter/Public/GLTFContainerBuilder.h b/Engine/Plugins/Enterprise/GLTFExporter/MigratedFromGit/Source/GLTFExporter/Public/GLTFContainerBuilder.h
--- a/Engine/Plugins/Enterprise/GLTFExporter/MigratedFromGit/Source/GLTFExporter/Public/GLTFContainerBuilder.h
+++ b/Engine/Plugins/Enterprise/GLTFExporter/MigratedFromGit/Source/GLTFExporter/Public/GLTFContainerBuilder.h
@@ -14,6 +14,8 @@ struct GLTFEXPORTER_API FGLTFContainerBuilder
 	FGLTFContainerBuilder();
 
 	FGLTFJsonAccessorIndex AddAccessor(const FGLTFJsonAccessor& JsonAccessor);
+	FGLTFJsonBufferIndex AddBuffer(const FGLTFJsonBuffer& JsonBuffer);
+	FGLTFJsonBufferViewIndex AddBufferView(const FGLTFJsonBufferView& JsonBufferView);
 	FGLTFJsonMeshIndex AddMesh(const FGLTFJsonMesh& JsonMesh);
 	FGLTFJsonNodeIndex AddNode(const FGLTFJsonNode& JsonNode);
 	FGLTFJsonSceneIndex AddScene(const FGLTFJsonScene& JsonScene);
@@ -23,7 +25,7 @@ struct GLTFEXPORTER_API FGLTFContainerBuilder
 	template <class ElementType>
 	FGLTFJsonBufferViewIndex AddBufferView(const TArray<ElementType>& Array, const FString& Name = TEXT(""), EGLTFJsonBufferTarget BufferTarget = EGLTFJsonBufferTarget::ArrayBuffer)
 	{
-		return BufferBuilder.AddBufferView(Array.GetData(), Array.Num() * sizeof(ElementType), Name, BufferTarget);
+		return AddBufferView(Array.GetData(), Array.Num() * sizeof(ElementType), Name, BufferTarget);
 	}
 
 	FGLTFJsonMeshIndex AddMesh(const UStaticMesh* StaticMesh, int32 LODIndex);
