[GLTFExporter] refactor: improve type-safety by make each kind of json index into a unique type

commit 13fe6aa8215c0bf16dabf41f870e2489934f1137
Date: Sun, 10 May 2020 12:12:59 +0200

    refactor: improve type-safety by make each kind of json index into a unique type
    (cherry picked from commit 1f0992d50c1616dd82500562099f5d2fcd99c8f6​)

#skipci

diff --git a/Engine/Plugins/Enterprise/GLTFExporter/MigratedFromGit/Source/GLTFExporter/Private/GLTFContainerBuilder.cpp b/Engine/Plugins/Enterprise/GLTFExporter/MigratedFromGit/Source/GLTFExporter/Private/GLTFContainerBuilder.cpp
--- a/Engine/Plugins/Enterprise/GLTFExporter/MigratedFromGit/Source/GLTFExporter/Private/GLTFContainerBuilder.cpp
+++ b/Engine/Plugins/Enterprise/GLTFExporter/MigratedFromGit/Source/GLTFExporter/Private/GLTFContainerBuilder.cpp
@@ -10,32 +10,32 @@ FGLTFContainerBuilder::FGLTFContainerBuilder()
 
 FGLTFJsonAccessorIndex FGLTFContainerBuilder::AddAccessor(const FGLTFJsonAccessor& JsonAccessor)
 {
-	return JsonRoot.Accessors.Add(JsonAccessor);
+	return FGLTFJsonAccessorIndex(JsonRoot.Accessors.Add(JsonAccessor));
 }
 
 FGLTFJsonBufferIndex FGLTFContainerBuilder::AddBuffer(const FGLTFJsonBuffer& JsonBuffer)
 {
-	return JsonRoot.Buffers.Add(JsonBuffer);
+	return FGLTFJsonBufferIndex(JsonRoot.Buffers.Add(JsonBuffer));
 }
 
 FGLTFJsonBufferViewIndex FGLTFContainerBuilder::AddBufferView(const FGLTFJsonBufferView& JsonBufferView)
 {
-	return JsonRoot.BufferViews.Add(JsonBufferView);
+	return FGLTFJsonBufferViewIndex(JsonRoot.BufferViews.Add(JsonBufferView));
 }
 
 FGLTFJsonMeshIndex FGLTFContainerBuilder::AddMesh(const FGLTFJsonMesh& JsonMesh)
 {
-	return JsonRoot.Meshes.Add(JsonMesh);
+	return FGLTFJsonMeshIndex(JsonRoot.Meshes.Add(JsonMesh));
 }
 
 FGLTFJsonNodeIndex FGLTFContainerBuilder::AddNode(const FGLTFJsonNode& JsonNode)
 {
-	return JsonRoot.Nodes.Add(JsonNode);
+	return FGLTFJsonNodeIndex(JsonRoot.Nodes.Add(JsonNode));
 }
 
 FGLTFJsonSceneIndex FGLTFContainerBuilder::AddScene(const FGLTFJsonScene& JsonScene)
 {
-	return JsonRoot.Scenes.Add(JsonScene);
+	return FGLTFJsonSceneIndex(JsonRoot.Scenes.Add(JsonScene));
 }
 
 FGLTFJsonBufferViewIndex FGLTFContainerBuilder::AddBufferView(const void* RawData, uint64 ByteLength, const FString& Name, EGLTFJsonBufferTarget BufferTarget)
diff --git a/Engine/Plugins/Enterprise/GLTFExporter/MigratedFromGit/Source/GLTFExporter/Private/GLTFMeshBuilder.cpp b/Engine/Plugins/Enterprise/GLTFExporter/MigratedFromGit/Source/GLTFExporter/Private/GLTFMeshBuilder.cpp
--- a/Engine/Plugins/Enterprise/GLTFExporter/MigratedFromGit/Source/GLTFExporter/Private/GLTFMeshBuilder.cpp
+++ b/Engine/Plugins/Enterprise/GLTFExporter/MigratedFromGit/Source/GLTFExporter/Private/GLTFMeshBuilder.cpp
@@ -20,7 +20,7 @@ FGLTFJsonAccessorIndex FGLTFSectionBuilder::AddAccessorForIndices(FGLTFContainer
 {
 	if (Indices.Num() == 0)
 	{
-		return INDEX_NONE;
+		return FGLTFJsonAccessorIndex(INDEX_NONE);
 	}
 
 	const FString AttributeName = Name + TEXT("_Indices");
@@ -116,7 +116,7 @@ FGLTFJsonAccessorIndex FGLTFMeshBuilder::AddAccessorForPositions(FGLTFContainerB
 {
 	if (Positions.Num() == 0)
 	{
-		return INDEX_NONE;
+		return FGLTFJsonAccessorIndex(INDEX_NONE);
 	}
 
 	const FString AttributeName = Name + TEXT("_Positions");
@@ -146,7 +146,7 @@ FGLTFJsonAccessorIndex FGLTFMeshBuilder::AddAccessorForNormals(FGLTFContainerBui
 {
 	if (Normals.Num() == 0)
 	{
-		return INDEX_NONE;
+		return FGLTFJsonAccessorIndex(INDEX_NONE);
 	}
 
 	const FString AttributeName = Name + TEXT("_Normals");
@@ -166,7 +166,7 @@ FGLTFJsonAccessorIndex FGLTFMeshBuilder::AddAccessorForColors(FGLTFContainerBuil
 {
 	if (Colors.Num() == 0)
 	{
-		return INDEX_NONE;
+		return FGLTFJsonAccessorIndex(INDEX_NONE);
 	}
 
 	const FString AttributeName = Name + TEXT("_Colors");
@@ -187,7 +187,7 @@ FGLTFJsonAccessorIndex FGLTFMeshBuilder::AddAccessorForTangents(FGLTFContainerBu
 {
 	if (Tangents.Num() == 0)
 	{
-		return INDEX_NONE;
+		return FGLTFJsonAccessorIndex(INDEX_NONE);
 	}
 
 	const FString AttributeName = Name + TEXT("_Tangents");
@@ -207,7 +207,7 @@ FGLTFJsonAccessorIndex FGLTFMeshBuilder::AddAccessorForUV0s(FGLTFContainerBuilde
 {
 	if (UV0s.Num() == 0)
 	{
-		return INDEX_NONE;
+		return FGLTFJsonAccessorIndex(INDEX_NONE);
 	}
 
 	const FString AttributeName = Name + TEXT("_UV0s");
@@ -227,7 +227,7 @@ FGLTFJsonAccessorIndex FGLTFMeshBuilder::AddAccessorForUV1s(FGLTFContainerBuilde
 {
 	if (UV1s.Num() == 0)
 	{
-		return INDEX_NONE;
+		return FGLTFJsonAccessorIndex(INDEX_NONE);
 	}
 
 	const FString AttributeName = Name + TEXT("_UV1s");
diff --git a/Engine/Plugins/Enterprise/GLTFExporter/MigratedFromGit/Source/GLTFExporter/Public/GLTFJsonIndex.h b/Engine/Plugins/Enterprise/GLTFExporter/MigratedFromGit/Source/GLTFExporter/Public/GLTFJsonIndex.h
--- a/Engine/Plugins/Enterprise/GLTFExporter/MigratedFromGit/Source/GLTFExporter/Public/GLTFJsonIndex.h
+++ b/Engine/Plugins/Enterprise/GLTFExporter/MigratedFromGit/Source/GLTFExporter/Public/GLTFJsonIndex.h
@@ -4,12 +4,58 @@
 
 #include "CoreMinimal.h"
 
-typedef int32 FGLTFJsonAccessorIndex;
-typedef int32 FGLTFJsonBufferIndex;
-typedef int32 FGLTFJsonBufferViewIndex;
-typedef int32 FGLTFJsonCameraIndex;
-typedef int32 FGLTFJsonMaterialIndex;
-typedef int32 FGLTFJsonMeshIndex;
-typedef int32 FGLTFJsonNodeIndex;
-typedef int32 FGLTFJsonSceneIndex;
-typedef int32 FGLTFJsonSkinIndex;
+template <typename DerivedType>
+struct FGLTFJsonIndex
+{
+    int32 Value;
+
+    FGLTFJsonIndex()
+        : Value(INDEX_NONE)
+    {
+    }
+
+    explicit FGLTFJsonIndex(int32 Value)
+        : Value(Value)
+    {
+    }
+
+    FORCEINLINE operator int32() const
+    {
+        return Value;
+    }
+
+    FORCEINLINE bool operator==(int32 OtherValue) const
+    {
+        return Value == OtherValue;
+    }
+
+    FORCEINLINE bool operator!=(int32 OtherValue) const
+    {
+        return Value != OtherValue;
+    }
+
+    FORCEINLINE bool operator==(const DerivedType& Other) const
+    {
+        return Value == Other.Value;
+    }
+
+    FORCEINLINE bool operator!=(const DerivedType& Other) const
+    {
+        return Value != Other.Value;
+    }
+
+    FORCEINLINE friend uint32 GetTypeHash(const DerivedType& Other)
+    {
+        return GetTypeHash(Other.Value);
+    }
+};
+
+struct FGLTFJsonAccessorIndex : FGLTFJsonIndex<FGLTFJsonAccessorIndex> { using FGLTFJsonIndex::FGLTFJsonIndex; };
+struct FGLTFJsonBufferIndex : FGLTFJsonIndex<FGLTFJsonBufferIndex> { using FGLTFJsonIndex::FGLTFJsonIndex; };
+struct FGLTFJsonBufferViewIndex : FGLTFJsonIndex<FGLTFJsonBufferViewIndex> { using FGLTFJsonIndex::FGLTFJsonIndex; };
+struct FGLTFJsonCameraIndex : FGLTFJsonIndex<FGLTFJsonCameraIndex> { using FGLTFJsonIndex::FGLTFJsonIndex; };
+struct FGLTFJsonMaterialIndex : FGLTFJsonIndex<FGLTFJsonMaterialIndex> { using FGLTFJsonIndex::FGLTFJsonIndex; };
+struct FGLTFJsonMeshIndex : FGLTFJsonIndex<FGLTFJsonMeshIndex> { using FGLTFJsonIndex::FGLTFJsonIndex; };
+struct FGLTFJsonNodeIndex : FGLTFJsonIndex<FGLTFJsonNodeIndex> { using FGLTFJsonIndex::FGLTFJsonIndex; };
+struct FGLTFJsonSceneIndex : FGLTFJsonIndex<FGLTFJsonSceneIndex> { using FGLTFJsonIndex::FGLTFJsonIndex; };
+struct FGLTFJsonSkinIndex : FGLTFJsonIndex<FGLTFJsonSkinIndex> { using FGLTFJsonIndex::FGLTFJsonIndex; };
