[GLTFExporter] refactor: remove redundant prefix UExporter::

commit 43627bc08c8a67a5837d3f4a96ce81db81f42570
Date: Sun, 10 May 2020 12:23:30 +0200

    refactor: remove redundant prefix UExporter::
    (cherry picked from commit 67182a8f96a77626346d2a88685a065ec5ad8fd6​)

#skipci

diff --git a/Engine/Plugins/Enterprise/GLTFExporter/MigratedFromGit/Source/GLTFExporter/Private/GLTFExporter.cpp b/Engine/Plugins/Enterprise/GLTFExporter/MigratedFromGit/Source/GLTFExporter/Private/GLTFExporter.cpp
--- a/Engine/Plugins/Enterprise/GLTFExporter/MigratedFromGit/Source/GLTFExporter/Private/GLTFExporter.cpp
+++ b/Engine/Plugins/Enterprise/GLTFExporter/MigratedFromGit/Source/GLTFExporter/Private/GLTFExporter.cpp
@@ -22,7 +22,7 @@ bool UGLTFExporter::FillExportOptions()
 	bool bExportAll = GetBatchMode() && !GetShowExportOption();
 	bool bExportCancel = false;
 
-	ExportOptions->FillOptions(GetBatchMode(), GetShowExportOption(), UExporter::CurrentFilename, bExportCancel, bExportAll);
+	ExportOptions->FillOptions(GetBatchMode(), GetShowExportOption(), CurrentFilename, bExportCancel, bExportAll);
 	if (bExportCancel)
 	{
 		SetCancelBatch(GetBatchMode());
