[GLTFExporter] feat(level): add partially implemented GLTFSceneBuilder and GLTFNodeBuilder with some debug logging

commit 2167927cda7c95c202e6cc348923d0543659f1cb
Date: Sun, 10 May 2020 18:37:16 +0200

    feat(level): add partially implemented GLTFSceneBuilder and GLTFNodeBuilder with some debug logging
    (cherry picked from commit ed32d7183f6750f17c9ad576a3cf8af0f289fab6​)

#skipci

diff --git a/Engine/Plugins/Enterprise/GLTFExporter/MigratedFromGit/Source/GLTFExporter/Private/GLTFLevelExporter.cpp b/Engine/Plugins/Enterprise/GLTFExporter/MigratedFromGit/Source/GLTFExporter/Private/GLTFLevelExporter.cpp
--- a/Engine/Plugins/Enterprise/GLTFExporter/MigratedFromGit/Source/GLTFExporter/Private/GLTFLevelExporter.cpp
+++ b/Engine/Plugins/Enterprise/GLTFExporter/MigratedFromGit/Source/GLTFExporter/Private/GLTFLevelExporter.cpp
@@ -1,6 +1,7 @@
 // Copyright Epic Games, Inc. All Rights Reserved.
 
 #include "GLTFLevelExporter.h"
+#include "GLTFSceneBuilder.h"
 #include "Engine/World.h"
 
 UGLTFLevelExporter::UGLTFLevelExporter(const FObjectInitializer& ObjectInitializer)
@@ -19,6 +20,8 @@ bool UGLTFLevelExporter::ExportBinary(UObject* Object, const TCHAR* Type, FArchi
 		return false;
 	}
 
+	FGLTFSceneBuilder(World, bSelectedOnly)._DebugLog();
+
 	// TODO: implement
 	return true;
 }
diff --git a/Engine/Plugins/Enterprise/GLTFExporter/MigratedFromGit/Source/GLTFExporter/Private/GLTFSceneBuilder.cpp b/Engine/Plugins/Enterprise/GLTFExporter/MigratedFromGit/Source/GLTFExporter/Private/GLTFSceneBuilder.cpp
new file mode 100644
--- /dev/null
+++ b/Engine/Plugins/Enterprise/GLTFExporter/MigratedFromGit/Source/GLTFExporter/Private/GLTFSceneBuilder.cpp
@@ -0,0 +1,97 @@
+// Copyright Epic Games, Inc. All Rights Reserved.
+
+#include "GLTFSceneBuilder.h"
+#include "GLTFConversionUtilities.h"
+#include "GLTFExporterModule.h"
+
+FGLTFNodeBuilder::FGLTFNodeBuilder(const USceneComponent* SceneComponent, bool bSelectedOnly, bool bTopLevel)
+	: SceneComponent(SceneComponent)
+	, bTopLevel(bTopLevel)
+{
+	const AActor* Owner = SceneComponent->GetOwner();
+	Name = Owner->GetName() + TEXT("_") + SceneComponent->GetName();
+
+	const TArray<USceneComponent*>& Children = SceneComponent->GetAttachChildren();
+	for (USceneComponent* ChildComponent : Children)
+	{
+		if (ChildComponent != nullptr && (!bSelectedOnly || ChildComponent->GetOwner()->IsSelected()))
+		{
+			AttachedComponents.Add(FGLTFNodeBuilder(ChildComponent, bSelectedOnly));
+		}
+	}
+}
+
+FGLTFJsonNodeIndex FGLTFNodeBuilder::AddNode(FGLTFContainerBuilder& Container)
+{
+	FGLTFJsonNode Node;
+
+	// TODO:
+
+	return Container.AddNode(Node);
+}
+
+void FGLTFNodeBuilder::_DebugLog()
+{
+	FString Path;
+	auto Component = SceneComponent;
+
+	do
+	{
+		auto Actor = Component->GetOwner();
+		auto ActorName = Actor ? Actor->GetName() : TEXT("null");
+		Path = TEXT(" / ") + ActorName + TEXT("_") + Component->GetName() + Path;
+		Component = Component->GetAttachParent();
+	} while (Component != nullptr);
+
+	UE_LOG(LogGLTFExporter, Log, TEXT("%s [%d] %s"), *Path, AttachedComponents.Num(), bTopLevel ? TEXT("*") : TEXT(""));
+
+	for (auto& AttachedComponent : AttachedComponents)
+	{
+		AttachedComponent._DebugLog();
+	}
+}
+
+FGLTFSceneBuilder::FGLTFSceneBuilder(const UWorld* World, bool bSelectedOnly)
+{
+	World->GetName(Name);
+
+	const ULevel* Level = World->PersistentLevel;
+
+	// TODO: export Level->Model
+
+	const TArray<AActor*>& Actors = Level->Actors;
+	for (const AActor* Actor : Actors)
+	{
+		if (Actor != nullptr && (!bSelectedOnly || Actor->IsSelected()))
+		{
+			const USceneComponent* RootComponent = Actor->GetRootComponent();
+			if (RootComponent != nullptr)
+			{
+				const USceneComponent* ParentComponent = RootComponent->GetAttachParent();
+				if (ParentComponent == nullptr || (bSelectedOnly && !ParentComponent->GetOwner()->IsSelected()))
+				{
+					TopLevelComponents.Add(FGLTFNodeBuilder(RootComponent, bSelectedOnly, true));
+				}
+			}
+		}
+	}
+}
+
+FGLTFJsonSceneIndex FGLTFSceneBuilder::AddScene(FGLTFContainerBuilder& Container)
+{
+	FGLTFJsonScene Scene;
+
+	// TODO:
+
+	return Container.AddScene(Scene);
+}
+
+void FGLTFSceneBuilder::_DebugLog()
+{
+	UE_LOG(LogGLTFExporter, Log, TEXT("Level %s:"), *Name);
+
+	for (auto& TopLevelComponent : TopLevelComponents)
+	{
+		TopLevelComponent._DebugLog();
+	}
+}
diff --git a/Engine/Plugins/Enterprise/GLTFExporter/MigratedFromGit/Source/GLTFExporter/Public/GLTFSceneBuilder.h b/Engine/Plugins/Enterprise/GLTFExporter/MigratedFromGit/Source/GLTFExporter/Public/GLTFSceneBuilder.h
new file mode 100644
--- /dev/null
+++ b/Engine/Plugins/Enterprise/GLTFExporter/MigratedFromGit/Source/GLTFExporter/Public/GLTFSceneBuilder.h
@@ -0,0 +1,37 @@
+// Copyright Epic Games, Inc. All Rights Reserved.
+
+#pragma once
+
+#include "GLTFJsonIndex.h"
+#include "GLTFContainerBuilder.h"
+#include "Engine.h"
+
+struct GLTFEXPORTER_API FGLTFNodeBuilder
+{
+	FString Name;
+
+	const USceneComponent* SceneComponent;
+
+	bool bTopLevel;
+
+	TArray<FGLTFNodeBuilder> AttachedComponents;
+	
+	FGLTFNodeBuilder(const USceneComponent* SceneComponent, bool bSelectedOnly, bool bTopLevel = false);
+
+	FGLTFJsonNodeIndex AddNode(FGLTFContainerBuilder& Container);
+
+	void _DebugLog();
+};
+
+struct GLTFEXPORTER_API FGLTFSceneBuilder
+{
+	FString Name;
+
+	TArray<FGLTFNodeBuilder> TopLevelComponents;
+
+	FGLTFSceneBuilder(const UWorld* World, bool bSelectedOnly);
+
+	FGLTFJsonSceneIndex AddScene(FGLTFContainerBuilder& Container);
+
+	void _DebugLog();
+};
