[GLTFExporter] feat(level): implement GLTFSceneBuilder::AddScene and GLTFNodeBuilder::AddNode

commit 33e45601e490a8b6e340fa9711dbd7d3b9cc836e
Date: Sun, 10 May 2020 18:42:16 +0200

    feat(level): implement GLTFSceneBuilder::AddScene and GLTFNodeBuilder::AddNode
    (cherry picked from commit 8a9f3b9a01368b2bbd113d1df744a0780d1589d2​)

#skipci

diff --git a/Engine/Plugins/Enterprise/GLTFExporter/MigratedFromGit/Source/GLTFExporter/Private/GLTFSceneBuilder.cpp b/Engine/Plugins/Enterprise/GLTFExporter/MigratedFromGit/Source/GLTFExporter/Private/GLTFSceneBuilder.cpp
--- a/Engine/Plugins/Enterprise/GLTFExporter/MigratedFromGit/Source/GLTFExporter/Private/GLTFSceneBuilder.cpp
+++ b/Engine/Plugins/Enterprise/GLTFExporter/MigratedFromGit/Source/GLTFExporter/Private/GLTFSceneBuilder.cpp
@@ -12,7 +12,7 @@ FGLTFNodeBuilder::FGLTFNodeBuilder(const USceneComponent* SceneComponent, bool b
 	Name = Owner->GetName() + TEXT("_") + SceneComponent->GetName();
 
 	const TArray<USceneComponent*>& Children = SceneComponent->GetAttachChildren();
-	for (USceneComponent* ChildComponent : Children)
+	for (const USceneComponent* ChildComponent : Children)
 	{
 		if (ChildComponent != nullptr && (!bSelectedOnly || ChildComponent->GetOwner()->IsSelected()))
 		{
@@ -21,16 +21,26 @@ FGLTFNodeBuilder::FGLTFNodeBuilder(const USceneComponent* SceneComponent, bool b
 	}
 }
 
-FGLTFJsonNodeIndex FGLTFNodeBuilder::AddNode(FGLTFContainerBuilder& Container)
+FGLTFJsonNodeIndex FGLTFNodeBuilder::AddNode(FGLTFContainerBuilder& Container) const
 {
 	FGLTFJsonNode Node;
+	Node.Name = Name;
 
-	// TODO:
+	FTransform Transform = bTopLevel ? SceneComponent->GetComponentTransform() : SceneComponent->GetRelativeTransform();
+	Node.Translation = ConvertPosition(Transform.GetTranslation());
+	Node.Rotation = ConvertRotation(Transform.GetRotation());
+	Node.Scale = ConvertScale(Transform.GetScale3D());
+
+	for (const FGLTFNodeBuilder& AttachedComponent : AttachedComponents)
+	{
+		FGLTFJsonNodeIndex ChildIndex = AttachedComponent.AddNode(Container);
+		if (ChildIndex != INDEX_NONE) Node.Children.Add(ChildIndex);
+	}
 
 	return Container.AddNode(Node);
 }
 
-void FGLTFNodeBuilder::_DebugLog()
+void FGLTFNodeBuilder::_DebugLog() const
 {
 	FString Path;
 	auto Component = SceneComponent;
@@ -75,18 +85,25 @@ FGLTFSceneBuilder::FGLTFSceneBuilder(const UWorld* World, bool bSelectedOnly)
 			}
 		}
 	}
+
+	_DebugLog(); // TODO: remove
 }
 
-FGLTFJsonSceneIndex FGLTFSceneBuilder::AddScene(FGLTFContainerBuilder& Container)
+FGLTFJsonSceneIndex FGLTFSceneBuilder::AddScene(FGLTFContainerBuilder& Container) const
 {
 	FGLTFJsonScene Scene;
+	Scene.Name = Name;
 
-	// TODO:
+	for (const FGLTFNodeBuilder& TopLevelComponent : TopLevelComponents)
+	{
+		FGLTFJsonNodeIndex NodeIndex = TopLevelComponent.AddNode(Container);
+		if (NodeIndex != INDEX_NONE) Scene.Nodes.Add(NodeIndex);
+	}
 
 	return Container.AddScene(Scene);
 }
 
-void FGLTFSceneBuilder::_DebugLog()
+void FGLTFSceneBuilder::_DebugLog() const
 {
 	UE_LOG(LogGLTFExporter, Log, TEXT("Level %s:"), *Name);
 
diff --git a/Engine/Plugins/Enterprise/GLTFExporter/MigratedFromGit/Source/GLTFExporter/Public/GLTFSceneBuilder.h b/Engine/Plugins/Enterprise/GLTFExporter/MigratedFromGit/Source/GLTFExporter/Public/GLTFSceneBuilder.h
--- a/Engine/Plugins/Enterprise/GLTFExporter/MigratedFromGit/Source/GLTFExporter/Public/GLTFSceneBuilder.h
+++ b/Engine/Plugins/Enterprise/GLTFExporter/MigratedFromGit/Source/GLTFExporter/Public/GLTFSceneBuilder.h
@@ -18,9 +18,9 @@ struct GLTFEXPORTER_API FGLTFNodeBuilder
 	
 	FGLTFNodeBuilder(const USceneComponent* SceneComponent, bool bSelectedOnly, bool bTopLevel = false);
 
-	FGLTFJsonNodeIndex AddNode(FGLTFContainerBuilder& Container);
+	FGLTFJsonNodeIndex AddNode(FGLTFContainerBuilder& Container) const;
 
-	void _DebugLog();
+	void _DebugLog() const;
 };
 
 struct GLTFEXPORTER_API FGLTFSceneBuilder
@@ -31,7 +31,7 @@ struct GLTFEXPORTER_API FGLTFSceneBuilder
 
 	FGLTFSceneBuilder(const UWorld* World, bool bSelectedOnly);
 
-	FGLTFJsonSceneIndex AddScene(FGLTFContainerBuilder& Container);
+	FGLTFJsonSceneIndex AddScene(FGLTFContainerBuilder& Container) const;
 
-	void _DebugLog();
+	void _DebugLog() const;
 };
