[GLTFExporter] feat(level): add and use GLTFContainerBuilder::AddScene

commit 885acf50c5779a1878e380200b5679dd545d7c8e
Date: Sun, 10 May 2020 18:53:10 +0200

    feat(level): add and use GLTFContainerBuilder::AddScene
    (cherry picked from commit f68bda9cf6d29a943e2673086de6b721ef438047​)

#skipci

diff --git a/Engine/Plugins/Enterprise/GLTFExporter/MigratedFromGit/Source/GLTFExporter/Private/GLTFContainerBuilder.cpp b/Engine/Plugins/Enterprise/GLTFExporter/MigratedFromGit/Source/GLTFExporter/Private/GLTFContainerBuilder.cpp
--- a/Engine/Plugins/Enterprise/GLTFExporter/MigratedFromGit/Source/GLTFExporter/Private/GLTFContainerBuilder.cpp
+++ b/Engine/Plugins/Enterprise/GLTFExporter/MigratedFromGit/Source/GLTFExporter/Private/GLTFContainerBuilder.cpp
@@ -2,6 +2,7 @@
 
 #include "GLTFContainerBuilder.h"
 #include "GLTFMeshBuilder.h"
+#include "GLTFSceneBuilder.h"
 
 FGLTFContainerBuilder::FGLTFContainerBuilder()
 	: BufferBuilder(AddBuffer(FGLTFJsonBuffer()))
@@ -54,3 +55,8 @@ FGLTFJsonMeshIndex FGLTFContainerBuilder::AddMesh(const UStaticMesh* StaticMesh,
 {
 	return FGLTFMeshBuilder(StaticMesh, LODIndex).AddMesh(*this);
 }
+
+FGLTFJsonSceneIndex FGLTFContainerBuilder::AddScene(const UWorld* World, bool bSelectedOnly)
+{
+	return FGLTFSceneBuilder(World, bSelectedOnly).AddScene(*this);
+}
diff --git a/Engine/Plugins/Enterprise/GLTFExporter/MigratedFromGit/Source/GLTFExporter/Private/GLTFLevelExporter.cpp b/Engine/Plugins/Enterprise/GLTFExporter/MigratedFromGit/Source/GLTFExporter/Private/GLTFLevelExporter.cpp
--- a/Engine/Plugins/Enterprise/GLTFExporter/MigratedFromGit/Source/GLTFExporter/Private/GLTFLevelExporter.cpp
+++ b/Engine/Plugins/Enterprise/GLTFExporter/MigratedFromGit/Source/GLTFExporter/Private/GLTFLevelExporter.cpp
@@ -1,7 +1,7 @@
 // Copyright Epic Games, Inc. All Rights Reserved.
 
 #include "GLTFLevelExporter.h"
-#include "GLTFSceneBuilder.h"
+#include "GLTFContainerBuilder.h"
 #include "Engine/World.h"
 
 UGLTFLevelExporter::UGLTFLevelExporter(const FObjectInitializer& ObjectInitializer)
@@ -20,8 +20,10 @@ bool UGLTFLevelExporter::ExportBinary(UObject* Object, const TCHAR* Type, FArchi
 		return false;
 	}
 
-	FGLTFSceneBuilder(World, bSelectedOnly)._DebugLog();
+	FGLTFContainerBuilder Container;
+	const FGLTFJsonSceneIndex SceneIndex = Container.AddScene(World, bSelectedOnly);
+	Container.JsonRoot.DefaultScene = SceneIndex;
 
-	// TODO: implement
+	Container.Serialize(Archive);
 	return true;
 }
diff --git a/Engine/Plugins/Enterprise/GLTFExporter/MigratedFromGit/Source/GLTFExporter/Public/GLTFContainerBuilder.h b/Engine/Plugins/Enterprise/GLTFExporter/MigratedFromGit/Source/GLTFExporter/Public/GLTFContainerBuilder.h
--- a/Engine/Plugins/Enterprise/GLTFExporter/MigratedFromGit/Source/GLTFExporter/Public/GLTFContainerBuilder.h
+++ b/Engine/Plugins/Enterprise/GLTFExporter/MigratedFromGit/Source/GLTFExporter/Public/GLTFContainerBuilder.h
@@ -30,5 +30,7 @@ struct GLTFEXPORTER_API FGLTFContainerBuilder
 
 	FGLTFJsonMeshIndex AddMesh(const UStaticMesh* StaticMesh, int32 LODIndex);
 
+	FGLTFJsonSceneIndex AddScene(const UWorld* World, bool bSelectedOnly = false);
+
 	void Serialize(FArchive& Archive);
 };
