[GLTFExporter] refactor(level): remove debug logging in GLTFSceneBuilder and GLTFNodeBuilder

commit 79b888fb145accff34857911f3fdf54c22722ff0
Date: Sun, 10 May 2020 18:50:48 +0200

    refactor(level): remove debug logging in GLTFSceneBuilder and GLTFNodeBuilder
    (cherry picked from commit bd3418ec0ef34c3344ce8acdf098f312cd045c6a​)

#skipci

diff --git a/Engine/Plugins/Enterprise/GLTFExporter/MigratedFromGit/Source/GLTFExporter/Private/GLTFSceneBuilder.cpp b/Engine/Plugins/Enterprise/GLTFExporter/MigratedFromGit/Source/GLTFExporter/Private/GLTFSceneBuilder.cpp
--- a/Engine/Plugins/Enterprise/GLTFExporter/MigratedFromGit/Source/GLTFExporter/Private/GLTFSceneBuilder.cpp
+++ b/Engine/Plugins/Enterprise/GLTFExporter/MigratedFromGit/Source/GLTFExporter/Private/GLTFSceneBuilder.cpp
@@ -40,27 +40,6 @@ FGLTFJsonNodeIndex FGLTFNodeBuilder::AddNode(FGLTFContainerBuilder& Container) c
 	return Container.AddNode(Node);
 }
 
-void FGLTFNodeBuilder::_DebugLog() const
-{
-	FString Path;
-	auto Component = SceneComponent;
-
-	do
-	{
-		auto Actor = Component->GetOwner();
-		auto ActorName = Actor ? Actor->GetName() : TEXT("null");
-		Path = TEXT(" / ") + ActorName + TEXT("_") + Component->GetName() + Path;
-		Component = Component->GetAttachParent();
-	} while (Component != nullptr);
-
-	UE_LOG(LogGLTFExporter, Log, TEXT("%s [%d] %s"), *Path, AttachedComponents.Num(), bTopLevel ? TEXT("*") : TEXT(""));
-
-	for (auto& AttachedComponent : AttachedComponents)
-	{
-		AttachedComponent._DebugLog();
-	}
-}
-
 FGLTFSceneBuilder::FGLTFSceneBuilder(const UWorld* World, bool bSelectedOnly)
 {
 	World->GetName(Name);
@@ -85,8 +64,6 @@ FGLTFSceneBuilder::FGLTFSceneBuilder(const UWorld* World, bool bSelectedOnly)
 			}
 		}
 	}
-
-	_DebugLog(); // TODO: remove
 }
 
 FGLTFJsonSceneIndex FGLTFSceneBuilder::AddScene(FGLTFContainerBuilder& Container) const
@@ -102,13 +79,3 @@ FGLTFJsonSceneIndex FGLTFSceneBuilder::AddScene(FGLTFContainerBuilder& Container
 
 	return Container.AddScene(Scene);
 }
-
-void FGLTFSceneBuilder::_DebugLog() const
-{
-	UE_LOG(LogGLTFExporter, Log, TEXT("Level %s:"), *Name);
-
-	for (auto& TopLevelComponent : TopLevelComponents)
-	{
-		TopLevelComponent._DebugLog();
-	}
-}
diff --git a/Engine/Plugins/Enterprise/GLTFExporter/MigratedFromGit/Source/GLTFExporter/Public/GLTFSceneBuilder.h b/Engine/Plugins/Enterprise/GLTFExporter/MigratedFromGit/Source/GLTFExporter/Public/GLTFSceneBuilder.h
--- a/Engine/Plugins/Enterprise/GLTFExporter/MigratedFromGit/Source/GLTFExporter/Public/GLTFSceneBuilder.h
+++ b/Engine/Plugins/Enterprise/GLTFExporter/MigratedFromGit/Source/GLTFExporter/Public/GLTFSceneBuilder.h
@@ -19,8 +19,6 @@ struct GLTFEXPORTER_API FGLTFNodeBuilder
 	FGLTFNodeBuilder(const USceneComponent* SceneComponent, bool bSelectedOnly, bool bTopLevel = false);
 
 	FGLTFJsonNodeIndex AddNode(FGLTFContainerBuilder& Container) const;
-
-	void _DebugLog() const;
 };
 
 struct GLTFEXPORTER_API FGLTFSceneBuilder
@@ -32,6 +30,4 @@ struct GLTFEXPORTER_API FGLTFSceneBuilder
 	FGLTFSceneBuilder(const UWorld* World, bool bSelectedOnly);
 
 	FGLTFJsonSceneIndex AddScene(FGLTFContainerBuilder& Container) const;
-
-	void _DebugLog() const;
 };
