[GLTFExporter] feat(level): support static mesh component when exporting node

commit 23352277d83924c65039f672d44730f24b319950
Date: Sun, 10 May 2020 18:54:26 +0200

    feat(level): support static mesh component when exporting node
    (cherry picked from commit 60d1e6cb13341c9357f8188b3b2e5158cd1dfb69​)

#skipci

diff --git a/Engine/Plugins/Enterprise/GLTFExporter/MigratedFromGit/Source/GLTFExporter/Private/GLTFContainerBuilder.cpp b/Engine/Plugins/Enterprise/GLTFExporter/MigratedFromGit/Source/GLTFExporter/Private/GLTFContainerBuilder.cpp
--- a/Engine/Plugins/Enterprise/GLTFExporter/MigratedFromGit/Source/GLTFExporter/Private/GLTFContainerBuilder.cpp
+++ b/Engine/Plugins/Enterprise/GLTFExporter/MigratedFromGit/Source/GLTFExporter/Private/GLTFContainerBuilder.cpp
@@ -56,6 +56,13 @@ FGLTFJsonMeshIndex FGLTFContainerBuilder::AddMesh(const UStaticMesh* StaticMesh,
 	return FGLTFMeshBuilder(StaticMesh, LODIndex).AddMesh(*this);
 }
 
+FGLTFJsonMeshIndex FGLTFContainerBuilder::AddMesh(const UStaticMeshComponent* StaticMeshComponent)
+{
+	const UStaticMesh* StaticMesh = StaticMeshComponent->GetStaticMesh();
+	const int32 LODIndex = StaticMeshComponent->ForcedLodModel > 0 ? StaticMeshComponent->ForcedLodModel - 1 : /* auto-select */ 0;
+	return AddMesh(StaticMesh, LODIndex);
+}
+
 FGLTFJsonSceneIndex FGLTFContainerBuilder::AddScene(const UWorld* World, bool bSelectedOnly)
 {
 	return FGLTFSceneBuilder(World, bSelectedOnly).AddScene(*this);
diff --git a/Engine/Plugins/Enterprise/GLTFExporter/MigratedFromGit/Source/GLTFExporter/Private/GLTFSceneBuilder.cpp b/Engine/Plugins/Enterprise/GLTFExporter/MigratedFromGit/Source/GLTFExporter/Private/GLTFSceneBuilder.cpp
--- a/Engine/Plugins/Enterprise/GLTFExporter/MigratedFromGit/Source/GLTFExporter/Private/GLTFSceneBuilder.cpp
+++ b/Engine/Plugins/Enterprise/GLTFExporter/MigratedFromGit/Source/GLTFExporter/Private/GLTFSceneBuilder.cpp
@@ -31,6 +31,12 @@ FGLTFJsonNodeIndex FGLTFNodeBuilder::AddNode(FGLTFContainerBuilder& Container) c
 	Node.Rotation = ConvertRotation(Transform.GetRotation());
 	Node.Scale = ConvertScale(Transform.GetScale3D());
 
+	const UStaticMeshComponent* StaticMeshComponent = Cast<UStaticMeshComponent>(SceneComponent);
+	if (StaticMeshComponent != nullptr)
+	{
+		Node.Mesh = Container.AddMesh(StaticMeshComponent);
+	}
+
 	for (const FGLTFNodeBuilder& AttachedComponent : AttachedComponents)
 	{
 		FGLTFJsonNodeIndex ChildIndex = AttachedComponent.AddNode(Container);
diff --git a/Engine/Plugins/Enterprise/GLTFExporter/MigratedFromGit/Source/GLTFExporter/Public/GLTFContainerBuilder.h b/Engine/Plugins/Enterprise/GLTFExporter/MigratedFromGit/Source/GLTFExporter/Public/GLTFContainerBuilder.h
--- a/Engine/Plugins/Enterprise/GLTFExporter/MigratedFromGit/Source/GLTFExporter/Public/GLTFContainerBuilder.h
+++ b/Engine/Plugins/Enterprise/GLTFExporter/MigratedFromGit/Source/GLTFExporter/Public/GLTFContainerBuilder.h
@@ -29,6 +29,7 @@ struct GLTFEXPORTER_API FGLTFContainerBuilder
 	}
 
 	FGLTFJsonMeshIndex AddMesh(const UStaticMesh* StaticMesh, int32 LODIndex);
+	FGLTFJsonMeshIndex AddMesh(const UStaticMeshComponent* StaticMeshComponent);
 
 	FGLTFJsonSceneIndex AddScene(const UWorld* World, bool bSelectedOnly = false);
 
