[GLTFExporter] feat: add support for overriding vertex colors when exporting static mesh

commit ab756c6ec9d960db764b18331e76f14641987ce5
Date: Sun, 10 May 2020 18:59:17 +0200

    feat: add support for overriding vertex colors when exporting static mesh
    (cherry picked from commit 72d6903919f256623cf3aaeea0c9d8cc8a0393dd​)

#skipci

diff --git a/Engine/Plugins/Enterprise/GLTFExporter/MigratedFromGit/Source/GLTFExporter/Private/GLTFContainerBuilder.cpp b/Engine/Plugins/Enterprise/GLTFExporter/MigratedFromGit/Source/GLTFExporter/Private/GLTFContainerBuilder.cpp
--- a/Engine/Plugins/Enterprise/GLTFExporter/MigratedFromGit/Source/GLTFExporter/Private/GLTFContainerBuilder.cpp
+++ b/Engine/Plugins/Enterprise/GLTFExporter/MigratedFromGit/Source/GLTFExporter/Private/GLTFContainerBuilder.cpp
@@ -51,16 +51,17 @@ void FGLTFContainerBuilder::Serialize(FArchive& Archive)
 	JsonRoot.Serialize(&Archive, true);
 }
 
-FGLTFJsonMeshIndex FGLTFContainerBuilder::AddMesh(const UStaticMesh* StaticMesh, int32 LODIndex)
+FGLTFJsonMeshIndex FGLTFContainerBuilder::AddMesh(const UStaticMesh* StaticMesh, int32 LODIndex, const FColorVertexBuffer* OverrideVertexColors)
 {
-	return FGLTFMeshBuilder(StaticMesh, LODIndex).AddMesh(*this);
+	return FGLTFMeshBuilder(StaticMesh, LODIndex, OverrideVertexColors).AddMesh(*this);
 }
 
 FGLTFJsonMeshIndex FGLTFContainerBuilder::AddMesh(const UStaticMeshComponent* StaticMeshComponent)
 {
 	const UStaticMesh* StaticMesh = StaticMeshComponent->GetStaticMesh();
 	const int32 LODIndex = StaticMeshComponent->ForcedLodModel > 0 ? StaticMeshComponent->ForcedLodModel - 1 : /* auto-select */ 0;
-	return AddMesh(StaticMesh, LODIndex);
+	const FColorVertexBuffer* OverrideVertexColors = LODIndex < StaticMeshComponent->LODData.Num() ? StaticMeshComponent->LODData[LODIndex].OverrideVertexColors : nullptr;
+	return AddMesh(StaticMesh, LODIndex, OverrideVertexColors);
 }
 
 FGLTFJsonSceneIndex FGLTFContainerBuilder::AddScene(const UWorld* World, bool bSelectedOnly)
diff --git a/Engine/Plugins/Enterprise/GLTFExporter/MigratedFromGit/Source/GLTFExporter/Private/GLTFMeshBuilder.cpp b/Engine/Plugins/Enterprise/GLTFExporter/MigratedFromGit/Source/GLTFExporter/Private/GLTFMeshBuilder.cpp
--- a/Engine/Plugins/Enterprise/GLTFExporter/MigratedFromGit/Source/GLTFExporter/Private/GLTFMeshBuilder.cpp
+++ b/Engine/Plugins/Enterprise/GLTFExporter/MigratedFromGit/Source/GLTFExporter/Private/GLTFMeshBuilder.cpp
@@ -36,7 +36,7 @@ FGLTFJsonAccessorIndex FGLTFSectionBuilder::AddAccessorForIndices(FGLTFContainer
 	return Container.AddAccessor(Accessor);
 }
 
-FGLTFMeshBuilder::FGLTFMeshBuilder(const UStaticMesh* StaticMesh, int32 LODIndex)
+FGLTFMeshBuilder::FGLTFMeshBuilder(const UStaticMesh* StaticMesh, int32 LODIndex, const FColorVertexBuffer* OverrideVertexColors)
 {
 	StaticMesh->GetName(Name);
 
@@ -63,7 +63,7 @@ FGLTFMeshBuilder::FGLTFMeshBuilder(const UStaticMesh* StaticMesh, int32 LODIndex
 		Positions[PosIndex] = ConvertPosition(PositionVertexBuffer.VertexPosition(PosIndex));
 	}
 
-	const FColorVertexBuffer& ColorVertexBuffer = LODMesh.VertexBuffers.ColorVertexBuffer;
+	const FColorVertexBuffer& ColorVertexBuffer = OverrideVertexColors != nullptr ? *OverrideVertexColors : LODMesh.VertexBuffers.ColorVertexBuffer;
 	const int32 ColorCount = ColorVertexBuffer.GetNumVertices();
 
 	if (ColorCount > 0)
diff --git a/Engine/Plugins/Enterprise/GLTFExporter/MigratedFromGit/Source/GLTFExporter/Public/GLTFContainerBuilder.h b/Engine/Plugins/Enterprise/GLTFExporter/MigratedFromGit/Source/GLTFExporter/Public/GLTFContainerBuilder.h
--- a/Engine/Plugins/Enterprise/GLTFExporter/MigratedFromGit/Source/GLTFExporter/Public/GLTFContainerBuilder.h
+++ b/Engine/Plugins/Enterprise/GLTFExporter/MigratedFromGit/Source/GLTFExporter/Public/GLTFContainerBuilder.h
@@ -28,7 +28,7 @@ struct GLTFEXPORTER_API FGLTFContainerBuilder
 		return AddBufferView(Array.GetData(), Array.Num() * sizeof(ElementType), Name, BufferTarget);
 	}
 
-	FGLTFJsonMeshIndex AddMesh(const UStaticMesh* StaticMesh, int32 LODIndex);
+	FGLTFJsonMeshIndex AddMesh(const UStaticMesh* StaticMesh, int32 LODIndex = 0, const FColorVertexBuffer* OverrideVertexColors = nullptr);
 	FGLTFJsonMeshIndex AddMesh(const UStaticMeshComponent* StaticMeshComponent);
 
 	FGLTFJsonSceneIndex AddScene(const UWorld* World, bool bSelectedOnly = false);
diff --git a/Engine/Plugins/Enterprise/GLTFExporter/MigratedFromGit/Source/GLTFExporter/Public/GLTFMeshBuilder.h b/Engine/Plugins/Enterprise/GLTFExporter/MigratedFromGit/Source/GLTFExporter/Public/GLTFMeshBuilder.h
--- a/Engine/Plugins/Enterprise/GLTFExporter/MigratedFromGit/Source/GLTFExporter/Public/GLTFMeshBuilder.h
+++ b/Engine/Plugins/Enterprise/GLTFExporter/MigratedFromGit/Source/GLTFExporter/Public/GLTFMeshBuilder.h
@@ -32,7 +32,7 @@ struct GLTFEXPORTER_API FGLTFMeshBuilder
 
 	FBox BoundingBox;
 
-	FGLTFMeshBuilder(const UStaticMesh* StaticMesh, int32 LODIndex);
+	FGLTFMeshBuilder(const UStaticMesh* StaticMesh, int32 LODIndex, const FColorVertexBuffer* OverrideVertexColors);
 
 	FGLTFJsonAccessorIndex AddAccessorForPositions(FGLTFContainerBuilder& Container) const;
 	FGLTFJsonAccessorIndex AddAccessorForColors(FGLTFContainerBuilder& Container) const;
