[GLTFExporter] feat: add IsSkySphereBlueprint and IsHDRIBackdropBlueprint utility functions

commit c47db8cbd368b2621ea7d9a99361c11b959e8e56
Date: Sun, 10 May 2020 18:59:44 +0200

    feat: add IsSkySphereBlueprint and IsHDRIBackdropBlueprint utility functions
    (cherry picked from commit c13f1208add816a3dd4c76b14f03c4f0004949cc​)

#skipci

diff --git a/Engine/Plugins/Enterprise/GLTFExporter/MigratedFromGit/Source/GLTFExporter/Private/GLTFConversionUtilities.cpp b/Engine/Plugins/Enterprise/GLTFExporter/MigratedFromGit/Source/GLTFExporter/Private/GLTFConversionUtilities.cpp
new file mode 100644
--- /dev/null
+++ b/Engine/Plugins/Enterprise/GLTFExporter/MigratedFromGit/Source/GLTFExporter/Private/GLTFConversionUtilities.cpp
@@ -0,0 +1,15 @@
+// Copyright Epic Games, Inc. All Rights Reserved.
+
+#include "GLTFConversionUtilities.h"
+
+bool IsSkySphereBlueprint(const UBlueprint* Blueprint)
+{
+	static const UBlueprint* SkySphere = LoadObject<UBlueprint>(nullptr, TEXT("/Engine/EngineSky/BP_Sky_Sphere.BP_Sky_Sphere"), nullptr, LOAD_None, nullptr);
+	return Blueprint != nullptr && Blueprint == SkySphere;
+}
+
+bool IsHDRIBackdropBlueprint(const UBlueprint* Blueprint)
+{
+	static const UBlueprint* HDRIBackdrop = LoadObject<UBlueprint>(nullptr, TEXT("/HDRIBackdrop/Blueprints/HDRIBackdrop.HDRIBackdrop"), nullptr, LOAD_None, nullptr);
+	return Blueprint != nullptr && Blueprint == HDRIBackdrop;
+}
diff --git a/Engine/Plugins/Enterprise/GLTFExporter/MigratedFromGit/Source/GLTFExporter/Public/GLTFConversionUtilities.h b/Engine/Plugins/Enterprise/GLTFExporter/MigratedFromGit/Source/GLTFExporter/Public/GLTFConversionUtilities.h
--- a/Engine/Plugins/Enterprise/GLTFExporter/MigratedFromGit/Source/GLTFExporter/Public/GLTFConversionUtilities.h
+++ b/Engine/Plugins/Enterprise/GLTFExporter/MigratedFromGit/Source/GLTFExporter/Public/GLTFConversionUtilities.h
@@ -81,3 +81,7 @@ inline EGLTFJsonAlphaMode ConvertAlphaMode(const EBlendMode Mode)
 		default:                            return EGLTFJsonAlphaMode::Opaque; // fallback
 	}
 }
+
+bool IsSkySphereBlueprint(const UBlueprint* Blueprint);
+
+bool IsHDRIBackdropBlueprint(const UBlueprint* Blueprint);
