[GLTFExporter] feat(level): add checks for special blueprints (SkySphere and HDRIBackdrop)

commit 574922203c077d414726b97e8ee5de485f55bd48
Date: Sun, 10 May 2020 19:03:58 +0200

    feat(level): add checks for special blueprints (SkySphere and HDRIBackdrop)
    (cherry picked from commit 002534b520fa679f507d05e1e36a4d89c771637b​)

#skipci

diff --git a/Engine/Plugins/Enterprise/GLTFExporter/MigratedFromGit/Source/GLTFExporter/Private/GLTFSceneBuilder.cpp b/Engine/Plugins/Enterprise/GLTFExporter/MigratedFromGit/Source/GLTFExporter/Private/GLTFSceneBuilder.cpp
--- a/Engine/Plugins/Enterprise/GLTFExporter/MigratedFromGit/Source/GLTFExporter/Private/GLTFSceneBuilder.cpp
+++ b/Engine/Plugins/Enterprise/GLTFExporter/MigratedFromGit/Source/GLTFExporter/Private/GLTFSceneBuilder.cpp
@@ -4,19 +4,23 @@
 #include "GLTFConversionUtilities.h"
 #include "GLTFExporterModule.h"
 
-FGLTFNodeBuilder::FGLTFNodeBuilder(const USceneComponent* SceneComponent, bool bSelectedOnly, bool bTopLevel)
+FGLTFNodeBuilder::FGLTFNodeBuilder(const USceneComponent* SceneComponent, const AActor* ComponentOwner, bool bSelectedOnly, bool bTopLevel)
 	: SceneComponent(SceneComponent)
+	, ComponentOwner(ComponentOwner)
 	, bTopLevel(bTopLevel)
 {
-	const AActor* Owner = SceneComponent->GetOwner();
-	Name = Owner->GetName() + TEXT("_") + SceneComponent->GetName();
+	Name = ComponentOwner->GetName() + TEXT("_") + SceneComponent->GetName();
 
 	const TArray<USceneComponent*>& Children = SceneComponent->GetAttachChildren();
 	for (const USceneComponent* ChildComponent : Children)
 	{
-		if (ChildComponent != nullptr && (!bSelectedOnly || ChildComponent->GetOwner()->IsSelected()))
+		if (ChildComponent != nullptr)
 		{
-			AttachedComponents.Add(FGLTFNodeBuilder(ChildComponent, bSelectedOnly));
+			const AActor* ChildOwner = ChildComponent->GetOwner();
+			if (!bSelectedOnly || ChildOwner->IsSelected())
+			{
+				AttachedComponents.Add(FGLTFNodeBuilder(ChildComponent, ChildOwner, bSelectedOnly));
+			}
 		}
 	}
 }
@@ -31,11 +35,21 @@ FGLTFJsonNodeIndex FGLTFNodeBuilder::AddNode(FGLTFContainerBuilder& Container) c
 	Node.Rotation = ConvertRotation(Transform.GetRotation());
 	Node.Scale = ConvertScale(Transform.GetScale3D());
 
-	const UStaticMeshComponent* StaticMeshComponent = Cast<UStaticMeshComponent>(SceneComponent);
-	if (StaticMeshComponent != nullptr)
+	const UClass* OwnerClass = ComponentOwner->GetClass();
+	const UBlueprint* Blueprint = UBlueprint::GetBlueprintFromClass(OwnerClass);
+	
+	if (IsSkySphereBlueprint(Blueprint))
+	{
+		// Ignore components
+	}
+	else if (const UStaticMeshComponent* StaticMeshComponent = Cast<UStaticMeshComponent>(SceneComponent))
 	{
 		Node.Mesh = Container.AddMesh(StaticMeshComponent);
 	}
+	else if (IsHDRIBackdropBlueprint(Blueprint) && ComponentOwner->GetRootComponent() == SceneComponent)
+	{
+		// TODO: add support for backdrop
+	}
 
 	for (const FGLTFNodeBuilder& AttachedComponent : AttachedComponents)
 	{
@@ -62,10 +76,10 @@ FGLTFSceneBuilder::FGLTFSceneBuilder(const UWorld* World, bool bSelectedOnly)
 			const USceneComponent* RootComponent = Actor->GetRootComponent();
 			if (RootComponent != nullptr)
 			{
-				const USceneComponent* ParentComponent = RootComponent->GetAttachParent();
-				if (ParentComponent == nullptr || (bSelectedOnly && !ParentComponent->GetOwner()->IsSelected()))
+				const AActor* ParentActor = Actor->GetParentActor();
+				if (ParentActor == nullptr || (bSelectedOnly && !ParentActor->IsSelected()))
 				{
-					TopLevelComponents.Add(FGLTFNodeBuilder(RootComponent, bSelectedOnly, true));
+					TopLevelComponents.Add(FGLTFNodeBuilder(RootComponent, Actor, bSelectedOnly, true));
 				}
 			}
 		}
diff --git a/Engine/Plugins/Enterprise/GLTFExporter/MigratedFromGit/Source/GLTFExporter/Public/GLTFSceneBuilder.h b/Engine/Plugins/Enterprise/GLTFExporter/MigratedFromGit/Source/GLTFExporter/Public/GLTFSceneBuilder.h
--- a/Engine/Plugins/Enterprise/GLTFExporter/MigratedFromGit/Source/GLTFExporter/Public/GLTFSceneBuilder.h
+++ b/Engine/Plugins/Enterprise/GLTFExporter/MigratedFromGit/Source/GLTFExporter/Public/GLTFSceneBuilder.h
@@ -11,12 +11,13 @@ struct GLTFEXPORTER_API FGLTFNodeBuilder
 	FString Name;
 
 	const USceneComponent* SceneComponent;
+	const AActor* ComponentOwner;
 
 	bool bTopLevel;
 
 	TArray<FGLTFNodeBuilder> AttachedComponents;
 	
-	FGLTFNodeBuilder(const USceneComponent* SceneComponent, bool bSelectedOnly, bool bTopLevel = false);
+	FGLTFNodeBuilder(const USceneComponent* SceneComponent, const AActor* ComponentOwner, bool bSelectedOnly, bool bTopLevel = false);
 
 	FGLTFJsonNodeIndex AddNode(FGLTFContainerBuilder& Container) const;
 };
