[GLTFExporter] refactor(level): use glTF-based identifiers and minor code cleanup

commit 0b0e85fa4c52a501b0a6c34812a664de66e1cff0
Date: Sun, 10 May 2020 19:19:25 +0200

    refactor(level): use glTF-based identifiers and minor code cleanup
    (cherry picked from commit 9a18a12cd971e37b493e26d6a83f7eeb1e8ba791​)

#skipci

diff --git a/Engine/Plugins/Enterprise/GLTFExporter/MigratedFromGit/Source/GLTFExporter/Private/GLTFSceneBuilder.cpp b/Engine/Plugins/Enterprise/GLTFExporter/MigratedFromGit/Source/GLTFExporter/Private/GLTFSceneBuilder.cpp
--- a/Engine/Plugins/Enterprise/GLTFExporter/MigratedFromGit/Source/GLTFExporter/Private/GLTFSceneBuilder.cpp
+++ b/Engine/Plugins/Enterprise/GLTFExporter/MigratedFromGit/Source/GLTFExporter/Private/GLTFSceneBuilder.cpp
@@ -4,10 +4,10 @@
 #include "GLTFConversionUtilities.h"
 #include "GLTFExporterModule.h"
 
-FGLTFNodeBuilder::FGLTFNodeBuilder(const USceneComponent* SceneComponent, const AActor* ComponentOwner, bool bSelectedOnly, bool bTopLevel)
+FGLTFNodeBuilder::FGLTFNodeBuilder( const USceneComponent* SceneComponent, const AActor* ComponentOwner, bool bSelectedOnly, bool bRootNode)
 	: SceneComponent(SceneComponent)
 	, ComponentOwner(ComponentOwner)
-	, bTopLevel(bTopLevel)
+	, bRootNode(bRootNode)
 {
 	Name = ComponentOwner->GetName() + TEXT("_") + SceneComponent->GetName();
 
@@ -30,23 +30,24 @@ FGLTFJsonNodeIndex FGLTFNodeBuilder::AddNode(FGLTFContainerBuilder& Container) c
 	FGLTFJsonNode Node;
 	Node.Name = Name;
 
-	FTransform Transform = bTopLevel ? SceneComponent->GetComponentTransform() : SceneComponent->GetRelativeTransform();
+	const FTransform Transform = bRootNode ? SceneComponent->GetComponentTransform() : SceneComponent->GetRelativeTransform();
 	Node.Translation = ConvertPosition(Transform.GetTranslation());
 	Node.Rotation = ConvertRotation(Transform.GetRotation());
 	Node.Scale = ConvertScale(Transform.GetScale3D());
 
 	const UClass* OwnerClass = ComponentOwner->GetClass();
 	const UBlueprint* Blueprint = UBlueprint::GetBlueprintFromClass(OwnerClass);
-	
+	const bool bIsRootComponent = ComponentOwner->GetRootComponent() == SceneComponent;
+
 	if (IsSkySphereBlueprint(Blueprint))
 	{
-		// Ignore components
+		// Avoid exporting any SkySphere actor components (like StaticMeshComponent)
 	}
 	else if (const UStaticMeshComponent* StaticMeshComponent = Cast<UStaticMeshComponent>(SceneComponent))
 	{
 		Node.Mesh = Container.AddMesh(StaticMeshComponent);
 	}
-	else if (IsHDRIBackdropBlueprint(Blueprint) && ComponentOwner->GetRootComponent() == SceneComponent)
+	else if (bIsRootComponent && IsHDRIBackdropBlueprint(Blueprint))
 	{
 		// TODO: add support for backdrop
 	}
@@ -66,7 +67,7 @@ FGLTFSceneBuilder::FGLTFSceneBuilder(const UWorld* World, bool bSelectedOnly)
 
 	const ULevel* Level = World->PersistentLevel;
 
-	// TODO: export Level->Model
+	// TODO: export Level->Model ?
 
 	const TArray<AActor*>& Actors = Level->Actors;
 	for (const AActor* Actor : Actors)
@@ -79,7 +80,7 @@ FGLTFSceneBuilder::FGLTFSceneBuilder(const UWorld* World, bool bSelectedOnly)
 				const AActor* ParentActor = Actor->GetParentActor();
 				if (ParentActor == nullptr || (bSelectedOnly && !ParentActor->IsSelected()))
 				{
-					TopLevelComponents.Add(FGLTFNodeBuilder(RootComponent, Actor, bSelectedOnly, true));
+					RootNodes.Add(FGLTFNodeBuilder(RootComponent, Actor, bSelectedOnly, true));
 				}
 			}
 		}
@@ -91,7 +92,7 @@ FGLTFJsonSceneIndex FGLTFSceneBuilder::AddScene(FGLTFContainerBuilder& Container
 	FGLTFJsonScene Scene;
 	Scene.Name = Name;
 
-	for (const FGLTFNodeBuilder& TopLevelComponent : TopLevelComponents)
+	for (const FGLTFNodeBuilder& TopLevelComponent : RootNodes)
 	{
 		FGLTFJsonNodeIndex NodeIndex = TopLevelComponent.AddNode(Container);
 		if (NodeIndex != INDEX_NONE) Scene.Nodes.Add(NodeIndex);
diff --git a/Engine/Plugins/Enterprise/GLTFExporter/MigratedFromGit/Source/GLTFExporter/Public/GLTFSceneBuilder.h b/Engine/Plugins/Enterprise/GLTFExporter/MigratedFromGit/Source/GLTFExporter/Public/GLTFSceneBuilder.h
--- a/Engine/Plugins/Enterprise/GLTFExporter/MigratedFromGit/Source/GLTFExporter/Public/GLTFSceneBuilder.h
+++ b/Engine/Plugins/Enterprise/GLTFExporter/MigratedFromGit/Source/GLTFExporter/Public/GLTFSceneBuilder.h
@@ -13,11 +13,11 @@ struct GLTFEXPORTER_API FGLTFNodeBuilder
 	const USceneComponent* SceneComponent;
 	const AActor* ComponentOwner;
 
-	bool bTopLevel;
+	bool bRootNode;
 
 	TArray<FGLTFNodeBuilder> AttachedComponents;
-	
-	FGLTFNodeBuilder(const USceneComponent* SceneComponent, const AActor* ComponentOwner, bool bSelectedOnly, bool bTopLevel = false);
+
+	FGLTFNodeBuilder(const USceneComponent* SceneComponent, const AActor* ComponentOwner, bool bSelectedOnly, bool bRootNode = false);
 
 	FGLTFJsonNodeIndex AddNode(FGLTFContainerBuilder& Container) const;
 };
@@ -26,7 +26,7 @@ struct GLTFEXPORTER_API FGLTFSceneBuilder
 {
 	FString Name;
 
-	TArray<FGLTFNodeBuilder> TopLevelComponents;
+	TArray<FGLTFNodeBuilder> RootNodes;
 
 	FGLTFSceneBuilder(const UWorld* World, bool bSelectedOnly);
 
