[GLTFExporter] fix: avoid glTF validation issue by using more accurate bounding box

commit 4cb28d3f580db97c2ec736622ffa5df70dc6adee
Date: Sun, 10 May 2020 21:09:47 +0200

    fix: avoid glTF validation issue by using more accurate bounding box
    (cherry picked from commit e5eea1ed4b74e29d6800b21270cee30df5ed8e0f​)

#skipci

diff --git a/Engine/Plugins/Enterprise/GLTFExporter/MigratedFromGit/Source/GLTFExporter/Private/GLTFMeshBuilder.cpp b/Engine/Plugins/Enterprise/GLTFExporter/MigratedFromGit/Source/GLTFExporter/Private/GLTFMeshBuilder.cpp
--- a/Engine/Plugins/Enterprise/GLTFExporter/MigratedFromGit/Source/GLTFExporter/Private/GLTFMeshBuilder.cpp
+++ b/Engine/Plugins/Enterprise/GLTFExporter/MigratedFromGit/Source/GLTFExporter/Private/GLTFMeshBuilder.cpp
@@ -109,7 +109,13 @@ FGLTFMeshBuilder::FGLTFMeshBuilder(const UStaticMesh* StaticMesh, int32 LODIndex
 		}
 	}
 
-	BoundingBox = StaticMesh->GetBoundingBox();
+	BoundingBox.Init();
+
+	// Avoid glTF validation issue by using more accurate bounding box (based on raw vertex values)
+	for (const FVector& Position : Positions)
+	{
+		BoundingBox += Position;
+	}
 }
 
 FGLTFJsonAccessorIndex FGLTFMeshBuilder::AddAccessorForPositions(FGLTFContainerBuilder& Container) const
@@ -129,14 +135,12 @@ FGLTFJsonAccessorIndex FGLTFMeshBuilder::AddAccessorForPositions(FGLTFContainerB
 	Accessor.Count = Positions.Num();
 	Accessor.Type = EGLTFJsonAccessorType::Vec3;
 
-	const FVector Max = ConvertSize(BoundingBox.Max);
-	const FVector Min = ConvertSize(BoundingBox.Min);
-	Accessor.Max[0] = Max.X;
-	Accessor.Max[1] = Max.Y;
-	Accessor.Max[2] = Max.Z;
-	Accessor.Min[0] = Min.X;
-	Accessor.Min[1] = Min.Y;
-	Accessor.Min[2] = Min.Z;
+	Accessor.Max[0] = BoundingBox.Max.X;
+	Accessor.Max[1] = BoundingBox.Max.Y;
+	Accessor.Max[2] = BoundingBox.Max.Z;
+	Accessor.Min[0] = BoundingBox.Min.X;
+	Accessor.Min[1] = BoundingBox.Min.Y;
+	Accessor.Min[2] = BoundingBox.Min.Z;
 	Accessor.MinMaxLength = 3;
 
 	return Container.AddAccessor(Accessor);
