[GLTFExporter] refactor: remove unused utility function ConvertSize

commit 868553f7c41d95eb3a97121eabb864c251cb10a2
Date: Sun, 10 May 2020 21:06:26 +0200

    refactor: remove unused utility function ConvertSize
    (cherry picked from commit 4b9468e2df303dd393a6c288652c3c06d1fe634a​)

#skipci

diff --git a/Engine/Plugins/Enterprise/GLTFExporter/MigratedFromGit/Source/GLTFExporter/Public/GLTFConversionUtilities.h b/Engine/Plugins/Enterprise/GLTFExporter/MigratedFromGit/Source/GLTFExporter/Public/GLTFConversionUtilities.h
--- a/Engine/Plugins/Enterprise/GLTFExporter/MigratedFromGit/Source/GLTFExporter/Public/GLTFConversionUtilities.h
+++ b/Engine/Plugins/Enterprise/GLTFExporter/MigratedFromGit/Source/GLTFExporter/Public/GLTFConversionUtilities.h
@@ -17,11 +17,6 @@ inline FVector ConvertPosition(const FVector& Position)
 	return ConvertVector(Position);
 }
 
-inline FVector ConvertSize(const FVector& Size)
-{
-	return ConvertVector(Size);
-}
-
 inline FVector ConvertScale(const FVector& Scale)
 {
 	return ConvertVector(Scale);
