[GLTFExporter] refactor: move conversion utilities to own class

commit bc3feda90f910c57474c9403e67a6e4fee99982b
Date: Mon, 11 May 2020 21:24:40 +0200

    refactor: move conversion utilities to own class
    (cherry picked from commit 6860233e806e151a77d5638aa93b841a2ffbd3dc​)

#skipci

diff --git a/Engine/Plugins/Enterprise/GLTFExporter/MigratedFromGit/Source/GLTFExporter/Private/GLTFConversionUtilities.cpp b/Engine/Plugins/Enterprise/GLTFExporter/MigratedFromGit/Source/GLTFExporter/Private/GLTFConversionUtilities.cpp
deleted file mode 100644
diff --git a/Engine/Plugins/Enterprise/GLTFExporter/MigratedFromGit/Source/GLTFExporter/Private/GLTFConverterUtility.cpp b/Engine/Plugins/Enterprise/GLTFExporter/MigratedFromGit/Source/GLTFExporter/Private/GLTFConverterUtility.cpp
new file mode 100644
--- /dev/null
+++ b/Engine/Plugins/Enterprise/GLTFExporter/MigratedFromGit/Source/GLTFExporter/Private/GLTFConverterUtility.cpp
@@ -0,0 +1,15 @@
+// Copyright Epic Games, Inc. All Rights Reserved.
+
+#include "GLTFConverterUtility.h"
+
+bool FGLTFConverterUtility::IsSkySphereBlueprint(const UBlueprint* Blueprint)
+{
+	static const UBlueprint* SkySphere = LoadObject<UBlueprint>(nullptr, TEXT("/Engine/EngineSky/BP_Sky_Sphere.BP_Sky_Sphere"), nullptr, LOAD_None, nullptr);
+	return Blueprint != nullptr && Blueprint == SkySphere;
+}
+
+bool FGLTFConverterUtility::IsHDRIBackdropBlueprint(const UBlueprint* Blueprint)
+{
+	static const UBlueprint* HDRIBackdrop = LoadObject<UBlueprint>(nullptr, TEXT("/HDRIBackdrop/Blueprints/HDRIBackdrop.HDRIBackdrop"), nullptr, LOAD_None, nullptr);
+	return Blueprint != nullptr && Blueprint == HDRIBackdrop;
+}
diff --git a/Engine/Plugins/Enterprise/GLTFExporter/MigratedFromGit/Source/GLTFExporter/Private/GLTFMeshBuilder.cpp b/Engine/Plugins/Enterprise/GLTFExporter/MigratedFromGit/Source/GLTFExporter/Private/GLTFMeshBuilder.cpp
--- a/Engine/Plugins/Enterprise/GLTFExporter/MigratedFromGit/Source/GLTFExporter/Private/GLTFMeshBuilder.cpp
+++ b/Engine/Plugins/Enterprise/GLTFExporter/MigratedFromGit/Source/GLTFExporter/Private/GLTFMeshBuilder.cpp
@@ -1,7 +1,7 @@
 // Copyright Epic Games, Inc. All Rights Reserved.
 
 #include "GLTFMeshBuilder.h"
-#include "GLTFConversionUtilities.h"
+#include "GLTFConverterUtility.h"
 
 FGLTFSectionBuilder::FGLTFSectionBuilder(const FString& SectionName, const FStaticMeshSection& MeshSection, const FIndexArrayView& IndexArray)
 	: Name(SectionName)
@@ -60,7 +60,7 @@ FGLTFMeshBuilder::FGLTFMeshBuilder(const UStaticMesh* StaticMesh, int32 LODIndex
 
 	for (int32 PosIndex = 0; PosIndex < PositionCount; ++PosIndex)
 	{
-		Positions[PosIndex] = ConvertPosition(PositionVertexBuffer.VertexPosition(PosIndex));
+		Positions[PosIndex] = FGLTFConverterUtility::ConvertPosition(PositionVertexBuffer.VertexPosition(PosIndex));
 	}
 
 	const FColorVertexBuffer& ColorVertexBuffer = OverrideVertexColors != nullptr ? *OverrideVertexColors : LODMesh.VertexBuffers.ColorVertexBuffer;
@@ -72,7 +72,7 @@ FGLTFMeshBuilder::FGLTFMeshBuilder(const UStaticMesh* StaticMesh, int32 LODIndex
 
 		for (int32 ColorIndex = 0; ColorIndex < ColorCount; ++ColorIndex)
 		{
-			Colors[ColorIndex] = ConvertColor(ColorVertexBuffer.VertexColor(ColorIndex));
+			Colors[ColorIndex] = FGLTFConverterUtility::ConvertColor(ColorVertexBuffer.VertexColor(ColorIndex));
 		}
 	}
 
@@ -84,8 +84,8 @@ FGLTFMeshBuilder::FGLTFMeshBuilder(const UStaticMesh* StaticMesh, int32 LODIndex
 
 	for (int32 VertIndex = 0; VertIndex < VertexCount; ++VertIndex)
 	{
-		Normals[VertIndex] = ConvertVector(StaticMeshVertexBuffer.VertexTangentZ(VertIndex));
-		Tangents[VertIndex] = ConvertTangent(StaticMeshVertexBuffer.VertexTangentX(VertIndex));
+		Normals[VertIndex] = FGLTFConverterUtility::ConvertVector(StaticMeshVertexBuffer.VertexTangentZ(VertIndex));
+		Tangents[VertIndex] = FGLTFConverterUtility::ConvertTangent(StaticMeshVertexBuffer.VertexTangentX(VertIndex));
 	}
 
 	const int32 UVCount = StaticMeshVertexBuffer.GetNumTexCoords();
diff --git a/Engine/Plugins/Enterprise/GLTFExporter/MigratedFromGit/Source/GLTFExporter/Private/GLTFSceneBuilder.cpp b/Engine/Plugins/Enterprise/GLTFExporter/MigratedFromGit/Source/GLTFExporter/Private/GLTFSceneBuilder.cpp
--- a/Engine/Plugins/Enterprise/GLTFExporter/MigratedFromGit/Source/GLTFExporter/Private/GLTFSceneBuilder.cpp
+++ b/Engine/Plugins/Enterprise/GLTFExporter/MigratedFromGit/Source/GLTFExporter/Private/GLTFSceneBuilder.cpp
@@ -1,7 +1,7 @@
 // Copyright Epic Games, Inc. All Rights Reserved.
 
 #include "GLTFSceneBuilder.h"
-#include "GLTFConversionUtilities.h"
+#include "GLTFConverterUtility.h"
 #include "GLTFExporterModule.h"
 
 FGLTFNodeBuilder::FGLTFNodeBuilder(const USceneComponent* SceneComponent, const AActor* ComponentOwner, bool bSelectedOnly, bool bRootNode)
@@ -31,15 +31,15 @@ FGLTFJsonNodeIndex FGLTFNodeBuilder::AddNode(FGLTFContainerBuilder& Container) c
 	Node.Name = Name;
 
 	const FTransform Transform = bRootNode ? SceneComponent->GetComponentTransform() : SceneComponent->GetRelativeTransform();
-	Node.Translation = ConvertPosition(Transform.GetTranslation());
-	Node.Rotation = ConvertRotation(Transform.GetRotation());
-	Node.Scale = ConvertScale(Transform.GetScale3D());
+	Node.Translation = FGLTFConverterUtility::ConvertPosition(Transform.GetTranslation());
+	Node.Rotation = FGLTFConverterUtility::ConvertRotation(Transform.GetRotation());
+	Node.Scale = FGLTFConverterUtility::ConvertScale(Transform.GetScale3D());
 
 	const UClass* OwnerClass = ComponentOwner->GetClass();
 	const UBlueprint* Blueprint = UBlueprint::GetBlueprintFromClass(OwnerClass);
 	const bool bIsRootComponent = ComponentOwner->GetRootComponent() == SceneComponent;
 
-	if (IsSkySphereBlueprint(Blueprint))
+	if (FGLTFConverterUtility::IsSkySphereBlueprint(Blueprint))
 	{
 		// Avoid exporting any SkySphere actor components (like StaticMeshComponent)
 	}
@@ -47,7 +47,7 @@ FGLTFJsonNodeIndex FGLTFNodeBuilder::AddNode(FGLTFContainerBuilder& Container) c
 	{
 		Node.Mesh = Container.AddMesh(StaticMeshComponent);
 	}
-	else if (bIsRootComponent && IsHDRIBackdropBlueprint(Blueprint))
+	else if (bIsRootComponent && FGLTFConverterUtility::IsHDRIBackdropBlueprint(Blueprint))
 	{
 		// TODO: add support for backdrop
 	}
diff --git a/Engine/Plugins/Enterprise/GLTFExporter/MigratedFromGit/Source/GLTFExporter/Public/GLTFConversionUtilities.h b/Engine/Plugins/Enterprise/GLTFExporter/MigratedFromGit/Source/GLTFExporter/Public/GLTFConversionUtilities.h
deleted file mode 100644
diff --git a/Engine/Plugins/Enterprise/GLTFExporter/MigratedFromGit/Source/GLTFExporter/Public/GLTFConverterUtility.h b/Engine/Plugins/Enterprise/GLTFExporter/MigratedFromGit/Source/GLTFExporter/Public/GLTFConverterUtility.h
new file mode 100644
--- /dev/null
+++ b/Engine/Plugins/Enterprise/GLTFExporter/MigratedFromGit/Source/GLTFExporter/Public/GLTFConverterUtility.h
@@ -0,0 +1,85 @@
+// Copyright Epic Games, Inc. All Rights Reserved.
+
+#pragma once
+
+#include "GLTFJsonEnums.h"
+#include "Engine/EngineTypes.h"
+
+struct GLTFEXPORTER_API FGLTFConverterUtility
+{
+	static inline FVector ConvertVector(const FVector& Vector)
+	{
+		// UE4 uses a left-handed coordinate system, with Z up.
+		// glTF uses a right-handed coordinate system, with Y up.
+		return { Vector.X, Vector.Z, Vector.Y };
+	}
+
+	static inline FVector ConvertPosition(const FVector& Position)
+	{
+		return ConvertVector(Position);
+	}
+
+	static inline FVector ConvertScale(const FVector& Scale)
+	{
+		return ConvertVector(Scale);
+	}
+
+	static inline FVector4 ConvertTangent(const FVector4& Tangent)
+	{
+		// glTF stores tangent as Vec4, with W component indicating handedness of tangent basis.
+		return FVector4(ConvertVector(Tangent), Tangent.W);
+	}
+
+	static inline FColor ConvertColor(const FColor& Color)
+	{
+		// UE4 uses ABGR while glTF uses RGBA.
+		return { Color.R, Color.G, Color.B, Color.A };
+	}
+
+	static inline FQuat ConvertRotation(const FQuat& Rotation)
+	{
+		// UE4 uses a left-handed coordinate system, with Z up.
+		// glTF uses a right-handed coordinate system, with Y up.
+		// Rotation = (qX, qY, qZ, qW) = (sin(angle/2) * aX, sin(angle/2) * aY, sin(angle/2) * aZ, cons(angle/2))
+		// where (aX, aY, aZ) - rotation axis, angle - rotation angle
+		// Y swapped with Z between these coordinate systems
+		// also, as handedness is changed rotation is inversed - hence negation
+		// therefore glTFRotation = (-qX, -qZ, -qY, qw)
+
+		const FQuat Result(-Rotation.X, -Rotation.Z, -Rotation.Y, Rotation.W);
+		// Not checking if quaternion is normalized
+		// e.g. some sources use non-unit Quats for rotation tangents
+		return Result;
+	}
+
+	static inline FMatrix ConvertMatrix(const FMatrix& Matrix)
+	{
+		// Unreal stores matrix elements in row major order.
+		// glTF stores matrix elements in column major order.
+
+		FMatrix Result;
+		for (int32 Row = 0; Row < 4; ++Row)
+		{
+			for (int32 Col = 0; Col < 4; ++Col)
+			{
+				Result.M[Col][Row] = Matrix.M[Row][Col];
+			}
+		}
+		return Result;
+	}
+
+	static inline EGLTFJsonAlphaMode ConvertAlphaMode(const EBlendMode Mode)
+	{
+		switch (Mode)
+		{
+			case EBlendMode::BLEND_Opaque:      return EGLTFJsonAlphaMode::Opaque;
+			case EBlendMode::BLEND_Translucent: return EGLTFJsonAlphaMode::Blend;
+			case EBlendMode::BLEND_Masked:      return EGLTFJsonAlphaMode::Mask;
+			default:                            return EGLTFJsonAlphaMode::Opaque; // fallback
+		}
+	}
+
+	static bool IsSkySphereBlueprint(const UBlueprint* Blueprint);
+
+	static bool IsHDRIBackdropBlueprint(const UBlueprint* Blueprint);
+};
