[GLTFExporter] refactor: move json utilities to own class

commit f39231ae74731392b930d303fd76749f1e3a67cc
Date: Mon, 11 May 2020 21:34:30 +0200

    refactor: move json utilities to own class
    (cherry picked from commit 96e7f1686deb1a54ac383cb0874297ba520693e1​)

#skipci

diff --git a/Engine/Plugins/Enterprise/GLTFExporter/MigratedFromGit/Source/GLTFExporter/Public/GLTFJsonAccessor.h b/Engine/Plugins/Enterprise/GLTFExporter/MigratedFromGit/Source/GLTFExporter/Public/GLTFJsonAccessor.h
--- a/Engine/Plugins/Enterprise/GLTFExporter/MigratedFromGit/Source/GLTFExporter/Public/GLTFJsonAccessor.h
+++ b/Engine/Plugins/Enterprise/GLTFExporter/MigratedFromGit/Source/GLTFExporter/Public/GLTFJsonAccessor.h
@@ -4,7 +4,7 @@
 
 #include "GLTFJsonEnums.h"
 #include "GLTFJsonIndex.h"
-#include "GLTFJsonUtilities.h"
+#include "GLTFJsonUtility.h"
 #include "Serialization/JsonSerializer.h"
 
 struct GLTFEXPORTER_API FGLTFJsonAccessor
@@ -42,8 +42,8 @@ struct GLTFEXPORTER_API FGLTFJsonAccessor
 
 		JsonWriter.WriteValue(TEXT("bufferView"), BufferView);
 		JsonWriter.WriteValue(TEXT("count"), Count);
-		JsonWriter.WriteValue(TEXT("type"), AccessorTypeToString(Type));
-		JsonWriter.WriteValue(TEXT("componentType"), ComponentTypeToNumber(ComponentType));
+		JsonWriter.WriteValue(TEXT("type"), FGLTFJsonUtility::AccessorTypeToString(Type));
+		JsonWriter.WriteValue(TEXT("componentType"), FGLTFJsonUtility::ComponentTypeToNumber(ComponentType));
 		if (bNormalized) JsonWriter.WriteValue(TEXT("normalized"), bNormalized);
 
 		if (MinMaxLength > 0)
diff --git a/Engine/Plugins/Enterprise/GLTFExporter/MigratedFromGit/Source/GLTFExporter/Public/GLTFJsonBufferView.h b/Engine/Plugins/Enterprise/GLTFExporter/MigratedFromGit/Source/GLTFExporter/Public/GLTFJsonBufferView.h
--- a/Engine/Plugins/Enterprise/GLTFExporter/MigratedFromGit/Source/GLTFExporter/Public/GLTFJsonBufferView.h
+++ b/Engine/Plugins/Enterprise/GLTFExporter/MigratedFromGit/Source/GLTFExporter/Public/GLTFJsonBufferView.h
@@ -4,7 +4,7 @@
 
 #include "GLTFJsonEnums.h"
 #include "GLTFJsonIndex.h"
-#include "GLTFJsonUtilities.h"
+#include "GLTFJsonUtility.h"
 #include "Serialization/JsonSerializer.h"
 
 struct GLTFEXPORTER_API FGLTFJsonBufferView
@@ -42,7 +42,7 @@ struct GLTFEXPORTER_API FGLTFJsonBufferView
 		if (ByteOffset != 0) JsonWriter.WriteValue(TEXT("byteOffset"), ByteOffset);
 		if (ByteStride != 0) JsonWriter.WriteValue(TEXT("byteStride"), ByteStride);
 
-		if (Target != EGLTFJsonBufferTarget::None) JsonWriter.WriteValue(TEXT("target"), BufferTargetToNumber(Target));
+		if (Target != EGLTFJsonBufferTarget::None) JsonWriter.WriteValue(TEXT("target"), FGLTFJsonUtility::BufferTargetToNumber(Target));
 
 		JsonWriter.WriteObjectEnd();
 	}
diff --git a/Engine/Plugins/Enterprise/GLTFExporter/MigratedFromGit/Source/GLTFExporter/Public/GLTFJsonMesh.h b/Engine/Plugins/Enterprise/GLTFExporter/MigratedFromGit/Source/GLTFExporter/Public/GLTFJsonMesh.h
--- a/Engine/Plugins/Enterprise/GLTFExporter/MigratedFromGit/Source/GLTFExporter/Public/GLTFJsonMesh.h
+++ b/Engine/Plugins/Enterprise/GLTFExporter/MigratedFromGit/Source/GLTFExporter/Public/GLTFJsonMesh.h
@@ -72,7 +72,7 @@ struct GLTFEXPORTER_API FGLTFJsonPrimitive
 
 		if (Indices != INDEX_NONE) JsonWriter.WriteValue(TEXT("indices"), Indices);
 		if (Material != INDEX_NONE) JsonWriter.WriteValue(TEXT("material"), Material);
-		if (Mode != EGLTFJsonPrimitiveMode::None) JsonWriter.WriteValue(TEXT("mode"), PrimitiveModeToNumber(Mode));
+		if (Mode != EGLTFJsonPrimitiveMode::None) JsonWriter.WriteValue(TEXT("mode"), FGLTFJsonUtility::PrimitiveModeToNumber(Mode));
 
 		JsonWriter.WriteObjectEnd();
 	}
diff --git a/Engine/Plugins/Enterprise/GLTFExporter/MigratedFromGit/Source/GLTFExporter/Public/GLTFJsonUtilities.h b/Engine/Plugins/Enterprise/GLTFExporter/MigratedFromGit/Source/GLTFExporter/Public/GLTFJsonUtilities.h
deleted file mode 100644
diff --git a/Engine/Plugins/Enterprise/GLTFExporter/MigratedFromGit/Source/GLTFExporter/Public/GLTFJsonUtility.h b/Engine/Plugins/Enterprise/GLTFExporter/MigratedFromGit/Source/GLTFExporter/Public/GLTFJsonUtility.h
new file mode 100644
--- /dev/null
+++ b/Engine/Plugins/Enterprise/GLTFExporter/MigratedFromGit/Source/GLTFExporter/Public/GLTFJsonUtility.h
@@ -0,0 +1,59 @@
+// Copyright Epic Games, Inc. All Rights Reserved.
+
+#pragma once
+
+#include "GLTFJsonEnums.h"
+
+struct GLTFEXPORTER_API FGLTFJsonUtility
+{
+	static inline const TCHAR* AlphaModeToString(EGLTFJsonAlphaMode Mode)
+	{
+		switch (Mode)
+		{
+			case EGLTFJsonAlphaMode::Opaque: return TEXT("OPAQUE");
+			case EGLTFJsonAlphaMode::Blend:  return TEXT("BLEND");
+			case EGLTFJsonAlphaMode::Mask:   return TEXT("MASK");
+			default:                         return TEXT("UNKNOWN");
+		}
+	}
+
+	static inline int32 PrimitiveModeToNumber(EGLTFJsonPrimitiveMode Mode)
+	{
+		return static_cast<int32>(Mode);
+	}
+
+	static inline int32 FilterToNumber(EGLTFJsonSamplerFilter Filter)
+	{
+		return static_cast<int32>(Filter);
+	}
+
+	static inline int32 WrapModeToNumber(EGLTFJsonSamplerWrap Wrap)
+	{
+		return static_cast<int32>(Wrap);
+	}
+
+	static inline const TCHAR* AccessorTypeToString(EGLTFJsonAccessorType Type)
+	{
+		switch (Type)
+		{
+			case EGLTFJsonAccessorType::Scalar: return TEXT("SCALAR");
+			case EGLTFJsonAccessorType::Vec2:   return TEXT("VEC2");
+			case EGLTFJsonAccessorType::Vec3:   return TEXT("VEC3");
+			case EGLTFJsonAccessorType::Vec4:   return TEXT("VEC4");
+			case EGLTFJsonAccessorType::Mat2:   return TEXT("MAT2");
+			case EGLTFJsonAccessorType::Mat3:   return TEXT("MAT3");
+			case EGLTFJsonAccessorType::Mat4:   return TEXT("MAT4");
+			default:                            return TEXT("UNKNOWN");
+		}
+	}
+
+	static inline int32 ComponentTypeToNumber(EGLTFJsonComponentType Type)
+	{
+		return static_cast<int32>(Type);
+	}
+
+	static inline int32 BufferTargetToNumber(EGLTFJsonBufferTarget Target)
+	{
+		return static_cast<int32>(Target);
+	}
+};
