[GLTFExporter] fix(level): avoid loading blueprints into memory for comparison

commit 82c6d234a60ad61def1958f1b71865a793138caf
Date: Mon, 11 May 2020 21:35:38 +0200

    fix(level): avoid loading blueprints into memory for comparison
    (cherry picked from commit d103d03ccfd199c1eaf35559789a01626f7a017e​)

#skipci

diff --git a/Engine/Plugins/Enterprise/GLTFExporter/MigratedFromGit/Source/GLTFExporter/Private/GLTFConverterUtility.cpp b/Engine/Plugins/Enterprise/GLTFExporter/MigratedFromGit/Source/GLTFExporter/Private/GLTFConverterUtility.cpp
--- a/Engine/Plugins/Enterprise/GLTFExporter/MigratedFromGit/Source/GLTFExporter/Private/GLTFConverterUtility.cpp
+++ b/Engine/Plugins/Enterprise/GLTFExporter/MigratedFromGit/Source/GLTFExporter/Private/GLTFConverterUtility.cpp
@@ -1,15 +1,14 @@
 // Copyright Epic Games, Inc. All Rights Reserved.
 
 #include "GLTFConverterUtility.h"
+#include "Engine/Blueprint.h"
 
 bool FGLTFConverterUtility::IsSkySphereBlueprint(const UBlueprint* Blueprint)
 {
-	static const UBlueprint* SkySphere = LoadObject<UBlueprint>(nullptr, TEXT("/Engine/EngineSky/BP_Sky_Sphere.BP_Sky_Sphere"), nullptr, LOAD_None, nullptr);
-	return Blueprint != nullptr && Blueprint == SkySphere;
+	return Blueprint != nullptr && Blueprint->GetPathName().Equals(TEXT("/Engine/EngineSky/BP_Sky_Sphere.BP_Sky_Sphere"));
 }
 
 bool FGLTFConverterUtility::IsHDRIBackdropBlueprint(const UBlueprint* Blueprint)
 {
-	static const UBlueprint* HDRIBackdrop = LoadObject<UBlueprint>(nullptr, TEXT("/HDRIBackdrop/Blueprints/HDRIBackdrop.HDRIBackdrop"), nullptr, LOAD_None, nullptr);
-	return Blueprint != nullptr && Blueprint == HDRIBackdrop;
+	return Blueprint != nullptr && Blueprint->GetPathName().Equals(TEXT("/HDRIBackdrop/Blueprints/HDRIBackdrop.HDRIBackdrop"));
 }
