[GLTFExporter] refactor(level): move node name assignment to initializer list

commit 4cba4f95a92a559e44a411f59bd763bea33a99e3
Date: Mon, 11 May 2020 21:38:23 +0200

    refactor(level): move node name assignment to initializer list
    (cherry picked from commit 8f999dbedb86aed5c41ce2017957c900c7dea781​)

#skipci

diff --git a/Engine/Plugins/Enterprise/GLTFExporter/MigratedFromGit/Source/GLTFExporter/Private/GLTFSceneBuilder.cpp b/Engine/Plugins/Enterprise/GLTFExporter/MigratedFromGit/Source/GLTFExporter/Private/GLTFSceneBuilder.cpp
--- a/Engine/Plugins/Enterprise/GLTFExporter/MigratedFromGit/Source/GLTFExporter/Private/GLTFSceneBuilder.cpp
+++ b/Engine/Plugins/Enterprise/GLTFExporter/MigratedFromGit/Source/GLTFExporter/Private/GLTFSceneBuilder.cpp
@@ -5,12 +5,11 @@
 #include "GLTFExporterModule.h"
 
 FGLTFNodeBuilder::FGLTFNodeBuilder(const USceneComponent* SceneComponent, const AActor* ComponentOwner, bool bSelectedOnly, bool bRootNode)
-	: SceneComponent(SceneComponent)
+	: Name(ComponentOwner->GetName() + TEXT("_") + SceneComponent->GetName())
+	, SceneComponent(SceneComponent)
 	, ComponentOwner(ComponentOwner)
 	, bRootNode(bRootNode)
 {
-	Name = ComponentOwner->GetName() + TEXT("_") + SceneComponent->GetName();
-
 	const TArray<USceneComponent*>& Children = SceneComponent->GetAttachChildren();
 	for (const USceneComponent* ChildComponent : Children)
 	{
