[GLTFExporter] refactor: change single if-statements to include braces on new line

commit 35c8e1f795ac75c0587fa0d5e881b16d3afc54f1
Date: Mon, 11 May 2020 21:46:36 +0200

    refactor: change single if-statements to include braces on new line
    (cherry picked from commit db7b45bed6f8fb47e08baa8a40ea7a05f4ed1e05​)

#skipci

diff --git a/Engine/Plugins/Enterprise/GLTFExporter/MigratedFromGit/Source/GLTFExporter/Private/GLTFSceneBuilder.cpp b/Engine/Plugins/Enterprise/GLTFExporter/MigratedFromGit/Source/GLTFExporter/Private/GLTFSceneBuilder.cpp
--- a/Engine/Plugins/Enterprise/GLTFExporter/MigratedFromGit/Source/GLTFExporter/Private/GLTFSceneBuilder.cpp
+++ b/Engine/Plugins/Enterprise/GLTFExporter/MigratedFromGit/Source/GLTFExporter/Private/GLTFSceneBuilder.cpp
@@ -54,7 +54,10 @@ FGLTFJsonNodeIndex FGLTFNodeBuilder::AddNode(FGLTFContainerBuilder& Container) c
 	for (const FGLTFNodeBuilder& AttachedComponent : AttachedComponents)
 	{
 		FGLTFJsonNodeIndex ChildIndex = AttachedComponent.AddNode(Container);
-		if (ChildIndex != INDEX_NONE) Node.Children.Add(ChildIndex);
+		if (ChildIndex != INDEX_NONE)
+		{
+			Node.Children.Add(ChildIndex);
+		}
 	}
 
 	return Container.AddNode(Node);
@@ -94,7 +97,10 @@ FGLTFJsonSceneIndex FGLTFSceneBuilder::AddScene(FGLTFContainerBuilder& Container
 	for (const FGLTFNodeBuilder& TopLevelComponent : RootNodes)
 	{
 		FGLTFJsonNodeIndex NodeIndex = TopLevelComponent.AddNode(Container);
-		if (NodeIndex != INDEX_NONE) Scene.Nodes.Add(NodeIndex);
+		if (NodeIndex != INDEX_NONE)
+		{
+			Scene.Nodes.Add(NodeIndex);
+		}
 	}
 
 	return Container.AddScene(Scene);
diff --git a/Engine/Plugins/Enterprise/GLTFExporter/MigratedFromGit/Source/GLTFExporter/Public/GLTFJsonAccessor.h b/Engine/Plugins/Enterprise/GLTFExporter/MigratedFromGit/Source/GLTFExporter/Public/GLTFJsonAccessor.h
--- a/Engine/Plugins/Enterprise/GLTFExporter/MigratedFromGit/Source/GLTFExporter/Public/GLTFJsonAccessor.h
+++ b/Engine/Plugins/Enterprise/GLTFExporter/MigratedFromGit/Source/GLTFExporter/Public/GLTFJsonAccessor.h
@@ -38,13 +38,20 @@ struct GLTFEXPORTER_API FGLTFJsonAccessor
 	{
 		JsonWriter.WriteObjectStart();
 
-		if (!Name.IsEmpty()) JsonWriter.WriteValue(TEXT("name"), Name);
+		if (!Name.IsEmpty())
+		{
+			JsonWriter.WriteValue(TEXT("name"), Name);
+		}
 
 		JsonWriter.WriteValue(TEXT("bufferView"), BufferView);
 		JsonWriter.WriteValue(TEXT("count"), Count);
 		JsonWriter.WriteValue(TEXT("type"), FGLTFJsonUtility::AccessorTypeToString(Type));
 		JsonWriter.WriteValue(TEXT("componentType"), FGLTFJsonUtility::ComponentTypeToNumber(ComponentType));
-		if (bNormalized) JsonWriter.WriteValue(TEXT("normalized"), bNormalized);
+
+		if (bNormalized) 
+		{
+			JsonWriter.WriteValue(TEXT("normalized"), bNormalized);
+		}
 
 		if (MinMaxLength > 0)
 		{
diff --git a/Engine/Plugins/Enterprise/GLTFExporter/MigratedFromGit/Source/GLTFExporter/Public/GLTFJsonBuffer.h b/Engine/Plugins/Enterprise/GLTFExporter/MigratedFromGit/Source/GLTFExporter/Public/GLTFJsonBuffer.h
--- a/Engine/Plugins/Enterprise/GLTFExporter/MigratedFromGit/Source/GLTFExporter/Public/GLTFJsonBuffer.h
+++ b/Engine/Plugins/Enterprise/GLTFExporter/MigratedFromGit/Source/GLTFExporter/Public/GLTFJsonBuffer.h
@@ -22,9 +22,16 @@ struct GLTFEXPORTER_API FGLTFJsonBuffer
 	{
 		JsonWriter.WriteObjectStart();
 
-		if (!Name.IsEmpty()) JsonWriter.WriteValue(TEXT("name"), Name);
+		if (!Name.IsEmpty())
+		{
+			JsonWriter.WriteValue(TEXT("name"), Name);
+		}
+
+		if (!URI.IsEmpty())
+		{
+			JsonWriter.WriteValue(TEXT("uri"), URI);
+		}
 
-		if (!URI.IsEmpty()) JsonWriter.WriteValue(TEXT("uri"), URI);
 		JsonWriter.WriteValue(TEXT("byteLength"), ByteLength);
 
 		JsonWriter.WriteObjectEnd();
diff --git a/Engine/Plugins/Enterprise/GLTFExporter/MigratedFromGit/Source/GLTFExporter/Public/GLTFJsonBufferView.h b/Engine/Plugins/Enterprise/GLTFExporter/MigratedFromGit/Source/GLTFExporter/Public/GLTFJsonBufferView.h
--- a/Engine/Plugins/Enterprise/GLTFExporter/MigratedFromGit/Source/GLTFExporter/Public/GLTFJsonBufferView.h
+++ b/Engine/Plugins/Enterprise/GLTFExporter/MigratedFromGit/Source/GLTFExporter/Public/GLTFJsonBufferView.h
@@ -34,15 +34,28 @@ struct GLTFEXPORTER_API FGLTFJsonBufferView
 	{
 		JsonWriter.WriteObjectStart();
 
-		if (!Name.IsEmpty()) JsonWriter.WriteValue(TEXT("name"), Name);
+		if (!Name.IsEmpty())
+		{
+			JsonWriter.WriteValue(TEXT("name"), Name);
+		}
 
 		JsonWriter.WriteValue(TEXT("buffer"), Buffer);
 		JsonWriter.WriteValue(TEXT("byteLength"), ByteLength);
 
-		if (ByteOffset != 0) JsonWriter.WriteValue(TEXT("byteOffset"), ByteOffset);
-		if (ByteStride != 0) JsonWriter.WriteValue(TEXT("byteStride"), ByteStride);
+		if (ByteOffset != 0)
+		{
+			JsonWriter.WriteValue(TEXT("byteOffset"), ByteOffset);
+		}
 
-		if (Target != EGLTFJsonBufferTarget::None) JsonWriter.WriteValue(TEXT("target"), FGLTFJsonUtility::BufferTargetToNumber(Target));
+		if (ByteStride != 0)
+		{
+			JsonWriter.WriteValue(TEXT("byteStride"), ByteStride);
+		}
+
+		if (Target != EGLTFJsonBufferTarget::None)
+		{
+			JsonWriter.WriteValue(TEXT("target"), FGLTFJsonUtility::BufferTargetToNumber(Target));
+		}
 
 		JsonWriter.WriteObjectEnd();
 	}
diff --git a/Engine/Plugins/Enterprise/GLTFExporter/MigratedFromGit/Source/GLTFExporter/Public/GLTFJsonMesh.h b/Engine/Plugins/Enterprise/GLTFExporter/MigratedFromGit/Source/GLTFExporter/Public/GLTFJsonMesh.h
--- a/Engine/Plugins/Enterprise/GLTFExporter/MigratedFromGit/Source/GLTFExporter/Public/GLTFJsonMesh.h
+++ b/Engine/Plugins/Enterprise/GLTFExporter/MigratedFromGit/Source/GLTFExporter/Public/GLTFJsonMesh.h
@@ -70,9 +70,20 @@ struct GLTFEXPORTER_API FGLTFJsonPrimitive
 		JsonWriter.WriteIdentifierPrefix(TEXT("attributes"));
 		Attributes.WriteObject(JsonWriter);
 
-		if (Indices != INDEX_NONE) JsonWriter.WriteValue(TEXT("indices"), Indices);
-		if (Material != INDEX_NONE) JsonWriter.WriteValue(TEXT("material"), Material);
-		if (Mode != EGLTFJsonPrimitiveMode::None) JsonWriter.WriteValue(TEXT("mode"), FGLTFJsonUtility::PrimitiveModeToNumber(Mode));
+		if (Indices != INDEX_NONE)
+		{
+			JsonWriter.WriteValue(TEXT("indices"), Indices);
+		}
+
+		if (Material != INDEX_NONE)
+		{
+			JsonWriter.WriteValue(TEXT("material"), Material);
+		}
+
+		if (Mode != EGLTFJsonPrimitiveMode::None)
+		{
+			JsonWriter.WriteValue(TEXT("mode"), FGLTFJsonUtility::PrimitiveModeToNumber(Mode));
+		}
 
 		JsonWriter.WriteObjectEnd();
 	}
@@ -89,7 +100,10 @@ struct GLTFEXPORTER_API FGLTFJsonMesh
 	{
 		JsonWriter.WriteObjectStart();
 
-		if (!Name.IsEmpty()) JsonWriter.WriteValue(TEXT("name"), Name);
+		if (!Name.IsEmpty())
+		{
+			JsonWriter.WriteValue(TEXT("name"), Name);
+		}
 
 		JsonWriter.WriteArrayStart(TEXT("primitives"));
 		for (const FGLTFJsonPrimitive& Primitive : Primitives)
diff --git a/Engine/Plugins/Enterprise/GLTFExporter/MigratedFromGit/Source/GLTFExporter/Public/GLTFJsonNode.h b/Engine/Plugins/Enterprise/GLTFExporter/MigratedFromGit/Source/GLTFExporter/Public/GLTFJsonNode.h
--- a/Engine/Plugins/Enterprise/GLTFExporter/MigratedFromGit/Source/GLTFExporter/Public/GLTFJsonNode.h
+++ b/Engine/Plugins/Enterprise/GLTFExporter/MigratedFromGit/Source/GLTFExporter/Public/GLTFJsonNode.h
@@ -34,7 +34,10 @@ struct GLTFEXPORTER_API FGLTFJsonNode
 	{
 		JsonWriter.WriteObjectStart();
 
-		if (!Name.IsEmpty()) JsonWriter.WriteValue(TEXT("name"), Name);
+		if (!Name.IsEmpty())
+		{
+			JsonWriter.WriteValue(TEXT("name"), Name);
+		}
 
 		if (Translation != FVector::ZeroVector)
 		{
@@ -64,9 +67,20 @@ struct GLTFEXPORTER_API FGLTFJsonNode
 			JsonWriter.WriteArrayEnd();
 		}
 
-		if (Camera != INDEX_NONE) JsonWriter.WriteValue(TEXT("camera"), Camera);
-		if (Skin != INDEX_NONE) JsonWriter.WriteValue(TEXT("skin"), Skin);
-		if (Mesh != INDEX_NONE) JsonWriter.WriteValue(TEXT("mesh"), Mesh);
+		if (Camera != INDEX_NONE)
+		{
+			JsonWriter.WriteValue(TEXT("camera"), Camera);
+		}
+
+		if (Skin != INDEX_NONE)
+		{
+			JsonWriter.WriteValue(TEXT("skin"), Skin);
+		}
+
+		if (Mesh != INDEX_NONE)
+		{
+			JsonWriter.WriteValue(TEXT("mesh"), Mesh);
+		}
 
 		if (Children.Num() > 0)
 		{
diff --git a/Engine/Plugins/Enterprise/GLTFExporter/MigratedFromGit/Source/GLTFExporter/Public/GLTFJsonRoot.h b/Engine/Plugins/Enterprise/GLTFExporter/MigratedFromGit/Source/GLTFExporter/Public/GLTFJsonRoot.h
--- a/Engine/Plugins/Enterprise/GLTFExporter/MigratedFromGit/Source/GLTFExporter/Public/GLTFJsonRoot.h
+++ b/Engine/Plugins/Enterprise/GLTFExporter/MigratedFromGit/Source/GLTFExporter/Public/GLTFJsonRoot.h
@@ -33,8 +33,16 @@ struct GLTFEXPORTER_API FGLTFJsonAsset
 		JsonWriter.WriteObjectStart();
 
 		JsonWriter.WriteValue(TEXT("version"), Version);
-		if (!Generator.IsEmpty()) JsonWriter.WriteValue(TEXT("generator"), Generator);
-		if (!Copyright.IsEmpty()) JsonWriter.WriteValue(TEXT("copyright"), Copyright);
+
+		if (!Generator.IsEmpty())
+		{
+			JsonWriter.WriteValue(TEXT("generator"), Generator);
+		}
+
+		if (!Copyright.IsEmpty())
+		{
+			JsonWriter.WriteValue(TEXT("copyright"), Copyright);
+		}
 
 		JsonWriter.WriteObjectEnd();
 	}
@@ -65,7 +73,10 @@ struct GLTFEXPORTER_API FGLTFJsonRoot
 		JsonWriter.WriteIdentifierPrefix(TEXT("asset"));
 		Asset.WriteObject(JsonWriter);
 
-		if (DefaultScene != INDEX_NONE) JsonWriter.WriteValue(TEXT("scene"), DefaultScene);
+		if (DefaultScene != INDEX_NONE)
+		{
+			JsonWriter.WriteValue(TEXT("scene"), DefaultScene);
+		}
 
 		if (Accessors.Num() > 0)
 		{
diff --git a/Engine/Plugins/Enterprise/GLTFExporter/MigratedFromGit/Source/GLTFExporter/Public/GLTFJsonScene.h b/Engine/Plugins/Enterprise/GLTFExporter/MigratedFromGit/Source/GLTFExporter/Public/GLTFJsonScene.h
--- a/Engine/Plugins/Enterprise/GLTFExporter/MigratedFromGit/Source/GLTFExporter/Public/GLTFJsonScene.h
+++ b/Engine/Plugins/Enterprise/GLTFExporter/MigratedFromGit/Source/GLTFExporter/Public/GLTFJsonScene.h
@@ -16,7 +16,10 @@ struct GLTFEXPORTER_API FGLTFJsonScene
 	{
 		JsonWriter.WriteObjectStart();
 
-		if (!Name.IsEmpty()) JsonWriter.WriteValue(TEXT("name"), Name);
+		if (!Name.IsEmpty())
+		{
+			JsonWriter.WriteValue(TEXT("name"), Name);
+		}
 
 		if (Nodes.Num() > 0)
 		{
