[GLTFExporter] fix(options): remove bOverwriteExisting because this should be handled by UAssetExportTask::bReplaceIdentical

commit bc8fc0edd02c333e7f81df3d2b244b41cc4dfdd3
Date: Wed, 13 May 2020 00:19:59 +0200

    fix(options): remove bOverwriteExisting because this should be handled by UAssetExportTask::bReplaceIdentical
    (cherry picked from commit a7880916f27118820961624813c21fa271acd7e4​)

#skipci

diff --git a/Engine/Plugins/Enterprise/GLTFExporter/MigratedFromGit/Source/GLTFExporter/Private/GLTFExportOptions.cpp b/Engine/Plugins/Enterprise/GLTFExporter/MigratedFromGit/Source/GLTFExporter/Private/GLTFExportOptions.cpp
--- a/Engine/Plugins/Enterprise/GLTFExporter/MigratedFromGit/Source/GLTFExporter/Private/GLTFExportOptions.cpp
+++ b/Engine/Plugins/Enterprise/GLTFExporter/MigratedFromGit/Source/GLTFExporter/Private/GLTFExportOptions.cpp
@@ -16,7 +16,6 @@
 UGLTFExportOptions::UGLTFExportOptions(const FObjectInitializer& ObjectInitializer)
 	: Super(ObjectInitializer)
 {
-	bOverwriteExisting = true;
 	bBundleWebViewer = true;
 	bExportPreviewMesh = true;
 	bExtensionsRequired = true;
diff --git a/Engine/Plugins/Enterprise/GLTFExporter/MigratedFromGit/Source/GLTFExporter/Public/GLTFExportOptions.h b/Engine/Plugins/Enterprise/GLTFExporter/MigratedFromGit/Source/GLTFExporter/Public/GLTFExportOptions.h
--- a/Engine/Plugins/Enterprise/GLTFExporter/MigratedFromGit/Source/GLTFExporter/Public/GLTFExportOptions.h
+++ b/Engine/Plugins/Enterprise/GLTFExporter/MigratedFromGit/Source/GLTFExporter/Public/GLTFExportOptions.h
@@ -47,9 +47,6 @@ class GLTFEXPORTER_API UGLTFExportOptions : public UObject
 	GENERATED_UCLASS_BODY()
 
 public:
-	UPROPERTY(EditAnywhere, BlueprintReadWrite, config, Category = Exporter)
-	uint32 bOverwriteExisting : 1;
-
 	UPROPERTY(EditAnywhere, BlueprintReadWrite, config, Category = Exporter)
 	uint32 bBundleWebViewer : 1;
 
