[GLTFExporter] fix(options): remove bExportGenericActor because this should always be true

commit 0e7091d4835a66095746e3b9c7f506fdc7f5f1ea
Date: Wed, 13 May 2020 00:23:59 +0200

    fix(options): remove bExportGenericActor because this should always be true
    (cherry picked from commit 13a2ea02d974e98b6472d18f9e04f81ffc40b8d4​)

#skipci

diff --git a/Engine/Plugins/Enterprise/GLTFExporter/MigratedFromGit/Source/GLTFExporter/Private/GLTFExportOptions.cpp b/Engine/Plugins/Enterprise/GLTFExporter/MigratedFromGit/Source/GLTFExporter/Private/GLTFExportOptions.cpp
--- a/Engine/Plugins/Enterprise/GLTFExporter/MigratedFromGit/Source/GLTFExporter/Private/GLTFExportOptions.cpp
+++ b/Engine/Plugins/Enterprise/GLTFExporter/MigratedFromGit/Source/GLTFExporter/Private/GLTFExportOptions.cpp
@@ -28,7 +28,6 @@ UGLTFExportOptions::UGLTFExportOptions(const FObjectInitializer& ObjectInitializ
 	TextureHDREncoding = EGLTFExporterTextureHDREncoding::RGBD;
 	bExportLightmaps = true;
 	ExportScale = 0.01;
-	bExportGenericActor = true;
 	bExportLight = true;
 	bExportCamera = true;
 	bExportReflectionCapture = true;
diff --git a/Engine/Plugins/Enterprise/GLTFExporter/MigratedFromGit/Source/GLTFExporter/Public/GLTFExportOptions.h b/Engine/Plugins/Enterprise/GLTFExporter/MigratedFromGit/Source/GLTFExporter/Public/GLTFExportOptions.h
--- a/Engine/Plugins/Enterprise/GLTFExporter/MigratedFromGit/Source/GLTFExporter/Public/GLTFExportOptions.h
+++ b/Engine/Plugins/Enterprise/GLTFExporter/MigratedFromGit/Source/GLTFExporter/Public/GLTFExportOptions.h
@@ -85,9 +85,6 @@ public:
 		 ToolTip = "Scale factor used for exporting assets, by default: 0.01, for conversion from centimeters(Unreal default) to meters(glTF)."))
 	float ExportScale;
 
-	UPROPERTY(EditAnywhere, BlueprintReadWrite, config, Category = Scene)
-	uint32 bExportGenericActor : 1;
-
 	UPROPERTY(EditAnywhere, BlueprintReadWrite, config, Category = Scene)
 	uint32 bExportLight : 1;
 
