[GLTFExporter] fix(options): rename properties from inconsistently singular to consistently plural and rename bExportAnimationTrigger to bExportInteractionHotspots

commit d3a14b52f3dd75eba678daba4870921081e57fb4
Date: Wed, 13 May 2020 00:29:41 +0200

    fix(options): rename properties from inconsistently singular to consistently plural and rename bExportAnimationTrigger to bExportInteractionHotspots
    (cherry picked from commit e1516f408d05b1e13be8fe72cc4a7a8df3fa3f4c​)

#skipci

diff --git a/Engine/Plugins/Enterprise/GLTFExporter/MigratedFromGit/Source/GLTFExporter/Private/GLTFExportOptions.cpp b/Engine/Plugins/Enterprise/GLTFExporter/MigratedFromGit/Source/GLTFExporter/Private/GLTFExportOptions.cpp
--- a/Engine/Plugins/Enterprise/GLTFExporter/MigratedFromGit/Source/GLTFExporter/Private/GLTFExportOptions.cpp
+++ b/Engine/Plugins/Enterprise/GLTFExporter/MigratedFromGit/Source/GLTFExporter/Private/GLTFExportOptions.cpp
@@ -19,21 +19,21 @@ UGLTFExportOptions::UGLTFExportOptions(const FObjectInitializer& ObjectInitializ
 	bBundleWebViewer = true;
 	bExportPreviewMesh = true;
 	bExtensionsRequired = true;
-	bExportVertexColor = true;
-	bExportUnlitMaterial = true;
-	bExportClearCoatMaterial = true;
+	bExportVertexColors = true;
+	bExportUnlitMaterials = true;
+	bExportClearCoatMaterials = true;
 	bBakeMaterialInputs = true;
 	BakedMaterialInputSize = EGLTFExporterTextureSize::POT_512;
 	TextureFormat = EGLTFExporterTextureFormat::PNG;
 	TextureHDREncoding = EGLTFExporterTextureHDREncoding::RGBD;
 	bExportLightmaps = true;
 	ExportScale = 0.01;
-	bExportLight = true;
-	bExportCamera = true;
-	bExportReflectionCapture = true;
-	bExportHDRIBackdrop = true;
+	bExportLights = true;
+	bExportCameras = true;
+	bExportReflectionCaptures = true;
+	bExportHDRIBackdrops = true;
 	bExportVariantSets = true;
-	bExportAnimationTrigger = true;
+	bExportInteractionHotspots = true;
 }
 
 void UGLTFExportOptions::ResetToDefault()
diff --git a/Engine/Plugins/Enterprise/GLTFExporter/MigratedFromGit/Source/GLTFExporter/Public/GLTFExportOptions.h b/Engine/Plugins/Enterprise/GLTFExporter/MigratedFromGit/Source/GLTFExporter/Public/GLTFExportOptions.h
--- a/Engine/Plugins/Enterprise/GLTFExporter/MigratedFromGit/Source/GLTFExporter/Public/GLTFExportOptions.h
+++ b/Engine/Plugins/Enterprise/GLTFExporter/MigratedFromGit/Source/GLTFExporter/Public/GLTFExportOptions.h
@@ -57,13 +57,13 @@ public:
 	uint32 bExtensionsRequired : 1;
 
 	UPROPERTY(EditAnywhere, BlueprintReadWrite, config, category = Mesh)
-	uint32 bExportVertexColor : 1;
+	uint32 bExportVertexColors : 1;
 
 	UPROPERTY(EditAnywhere, BlueprintReadWrite, config, Category = Material)
-	uint32 bExportUnlitMaterial : 1;
+	uint32 bExportUnlitMaterials : 1;
 
 	UPROPERTY(EditAnywhere, BlueprintReadWrite, config, Category = Material)
-	uint32 bExportClearCoatMaterial : 1;
+	uint32 bExportClearCoatMaterials : 1;
 
 	UPROPERTY(EditAnywhere, BlueprintReadWrite, config, category = Material)
 	uint32 bBakeMaterialInputs : 1;
@@ -81,27 +81,27 @@ public:
 	uint32 bExportLightmaps : 1;
 
 	UPROPERTY(EditAnywhere, BlueprintReadWrite, config, Category = Scene, meta =
-		(DisplayName = "Import Uniform Scale",
+		(DisplayName = "Export Uniform Scale",
 		 ToolTip = "Scale factor used for exporting assets, by default: 0.01, for conversion from centimeters(Unreal default) to meters(glTF)."))
 	float ExportScale;
 
 	UPROPERTY(EditAnywhere, BlueprintReadWrite, config, Category = Scene)
-	uint32 bExportLight : 1;
+	uint32 bExportLights : 1;
 
 	UPROPERTY(EditAnywhere, BlueprintReadWrite, config, Category = Scene)
-	uint32 bExportCamera : 1;
+	uint32 bExportCameras : 1;
 
 	UPROPERTY(EditAnywhere, BlueprintReadWrite, config, Category = Scene)
-	uint32 bExportReflectionCapture : 1;
+	uint32 bExportReflectionCaptures : 1;
 
 	UPROPERTY(EditAnywhere, BlueprintReadWrite, config, Category = Scene)
-	uint32 bExportHDRIBackdrop : 1;
+	uint32 bExportHDRIBackdrops : 1;
 
 	UPROPERTY(EditAnywhere, BlueprintReadWrite, config, Category = Scene)
 	uint32 bExportVariantSets : 1;
 
 	UPROPERTY(EditAnywhere, BlueprintReadWrite, config, Category = Scene)
-	uint32 bExportAnimationTrigger : 1;
+	uint32 bExportInteractionHotspots : 1;
 
 	/* Set all the UProperty to the CDO value */
 	void ResetToDefault();
