[GLTFExporter] feat(options): add properties DefaultLevelOfDetail and bMapSkeletalMotionToRoot

commit 6640261a34f1e67605246522b3a9db78572cd503
Date: Wed, 13 May 2020 00:30:21 +0200

    feat(options): add properties DefaultLevelOfDetail and bMapSkeletalMotionToRoot
    (cherry picked from commit 3e5edd5fbd2f2671ed3470b69328a1499d14bc8c​)

#skipci

diff --git a/Engine/Plugins/Enterprise/GLTFExporter/MigratedFromGit/Source/GLTFExporter/Private/GLTFExportOptions.cpp b/Engine/Plugins/Enterprise/GLTFExporter/MigratedFromGit/Source/GLTFExporter/Private/GLTFExportOptions.cpp
--- a/Engine/Plugins/Enterprise/GLTFExporter/MigratedFromGit/Source/GLTFExporter/Private/GLTFExportOptions.cpp
+++ b/Engine/Plugins/Enterprise/GLTFExporter/MigratedFromGit/Source/GLTFExporter/Private/GLTFExportOptions.cpp
@@ -20,6 +20,8 @@ UGLTFExportOptions::UGLTFExportOptions(const FObjectInitializer& ObjectInitializ
 	bExportPreviewMesh = true;
 	bExtensionsRequired = true;
 	bExportVertexColors = true;
+	DefaultLevelOfDetail = 0;
+	bMapSkeletalMotionToRoot = false;
 	bExportUnlitMaterials = true;
 	bExportClearCoatMaterials = true;
 	bBakeMaterialInputs = true;
diff --git a/Engine/Plugins/Enterprise/GLTFExporter/MigratedFromGit/Source/GLTFExporter/Public/GLTFExportOptions.h b/Engine/Plugins/Enterprise/GLTFExporter/MigratedFromGit/Source/GLTFExporter/Public/GLTFExportOptions.h
--- a/Engine/Plugins/Enterprise/GLTFExporter/MigratedFromGit/Source/GLTFExporter/Public/GLTFExportOptions.h
+++ b/Engine/Plugins/Enterprise/GLTFExporter/MigratedFromGit/Source/GLTFExporter/Public/GLTFExportOptions.h
@@ -59,6 +59,12 @@ public:
 	UPROPERTY(EditAnywhere, BlueprintReadWrite, config, category = Mesh)
 	uint32 bExportVertexColors : 1;
 
+	UPROPERTY(EditAnywhere, BlueprintReadWrite, config, category = Mesh)
+	int32 DefaultLevelOfDetail;
+
+	UPROPERTY(EditAnywhere, BlueprintReadWrite, config, category = Animation)
+	uint32 bMapSkeletalMotionToRoot : 1;
+
 	UPROPERTY(EditAnywhere, BlueprintReadWrite, config, Category = Material)
 	uint32 bExportUnlitMaterials : 1;
 
