[GLTFExporter] feat(options): add property bExportOrbitalCameras

commit 01782dc0e5c83756f077490ae118ead9cc200d28
Date: Wed, 13 May 2020 00:57:31 +0200

    feat(options): add property bExportOrbitalCameras
    (cherry picked from commit 506535a608cdfdf12f84b757fcb28224060c4dd3​)

#skipci

diff --git a/Engine/Plugins/Enterprise/GLTFExporter/MigratedFromGit/Source/GLTFExporter/Private/GLTFExportOptions.cpp b/Engine/Plugins/Enterprise/GLTFExporter/MigratedFromGit/Source/GLTFExporter/Private/GLTFExportOptions.cpp
--- a/Engine/Plugins/Enterprise/GLTFExporter/MigratedFromGit/Source/GLTFExporter/Private/GLTFExportOptions.cpp
+++ b/Engine/Plugins/Enterprise/GLTFExporter/MigratedFromGit/Source/GLTFExporter/Private/GLTFExportOptions.cpp
@@ -36,6 +36,7 @@ UGLTFExportOptions::UGLTFExportOptions(const FObjectInitializer& ObjectInitializ
 	bExportHDRIBackdrops = true;
 	bExportVariantSets = true;
 	bExportInteractionHotspots = true;
+	bExportOrbitalCameras = true;
 }
 
 void UGLTFExportOptions::ResetToDefault()
diff --git a/Engine/Plugins/Enterprise/GLTFExporter/MigratedFromGit/Source/GLTFExporter/Public/GLTFExportOptions.h b/Engine/Plugins/Enterprise/GLTFExporter/MigratedFromGit/Source/GLTFExporter/Public/GLTFExportOptions.h
--- a/Engine/Plugins/Enterprise/GLTFExporter/MigratedFromGit/Source/GLTFExporter/Public/GLTFExportOptions.h
+++ b/Engine/Plugins/Enterprise/GLTFExporter/MigratedFromGit/Source/GLTFExporter/Public/GLTFExportOptions.h
@@ -109,6 +109,9 @@ public:
 	UPROPERTY(EditAnywhere, BlueprintReadWrite, config, Category = Scene)
 	uint32 bExportInteractionHotspots : 1;
 
+	UPROPERTY(EditAnywhere, BlueprintReadWrite, config, Category = Scene)
+    uint32 bExportOrbitalCameras : 1;
+
 	/* Set all the UProperty to the CDO value */
 	void ResetToDefault();
 
