[GLTFExporter] refactor(options): improve public code documentation for the functions

commit 66cb441af95ff2f519c97de7c249b830d8137ed4
Date: Wed, 13 May 2020 01:05:22 +0200

    refactor(options): improve public code documentation for the functions
    (cherry picked from commit d4099f2f9aeec294843d7b5c39043a0074eac332​)

#skipci

diff --git a/Engine/Plugins/Enterprise/GLTFExporter/MigratedFromGit/Source/GLTFExporter/Public/GLTFExportOptions.h b/Engine/Plugins/Enterprise/GLTFExporter/MigratedFromGit/Source/GLTFExporter/Public/GLTFExportOptions.h
--- a/Engine/Plugins/Enterprise/GLTFExporter/MigratedFromGit/Source/GLTFExporter/Public/GLTFExportOptions.h
+++ b/Engine/Plugins/Enterprise/GLTFExporter/MigratedFromGit/Source/GLTFExporter/Public/GLTFExportOptions.h
@@ -112,19 +112,19 @@ public:
 	UPROPERTY(EditAnywhere, BlueprintReadWrite, config, Category = Scene)
     uint32 bExportOrbitalCameras : 1;
 
-	/* Set all the UProperty to the CDO value */
+	/* Set all the each property to the CDO value */
 	void ResetToDefault();
 
-	/* Save the UProperty to a local ini to retrieve the value the next time we call function LoadOptions() */
+	/* Save the each property to a local ini to retrieve the value the next time we call function LoadOptions() */
 	void SaveOptions();
 
-	/* Load the UProperty data from a local ini which the value was store by the function SaveOptions() */
+	/* Load the each property from a local ini which the value was store by the function SaveOptions() */
 	void LoadOptions();
 
 	/**
 	* Load the export option from the last save state and show the dialog if bShowOptionDialog is true.
 	* FullPath is the export file path we display it in the dialog
-	* If user cancel the dialog, the OutOperationCanceled will be true
+	* If user cancel the dialog, the bOutOperationCanceled will be true
 	* bOutExportAll will be true if the user want to use the same option for all other asset he want to export
 	*
 	* The function is saving the dialog state in a user ini file and reload it from there. It is not changing the CDO.
