[GLTFExporter] refactor: change local asset pointer (to be exported) to const

commit 78563ece632295b7955746801d1656b16a1e0707
Date: Wed, 13 May 2020 17:18:27 +0200

    refactor: change local asset pointer (to be exported) to const
    (cherry picked from commit 19ec30d1fe7d589deb89f6bf8c1a9b130df5ed1e​)

#skipci

diff --git a/Engine/Plugins/Enterprise/GLTFExporter/MigratedFromGit/Source/GLTFExporter/Private/GLTFAnimSequenceExporter.cpp b/Engine/Plugins/Enterprise/GLTFExporter/MigratedFromGit/Source/GLTFExporter/Private/GLTFAnimSequenceExporter.cpp
--- a/Engine/Plugins/Enterprise/GLTFExporter/MigratedFromGit/Source/GLTFExporter/Private/GLTFAnimSequenceExporter.cpp
+++ b/Engine/Plugins/Enterprise/GLTFExporter/MigratedFromGit/Source/GLTFExporter/Private/GLTFAnimSequenceExporter.cpp
@@ -11,7 +11,7 @@ UGLTFAnimSequenceExporter::UGLTFAnimSequenceExporter(const FObjectInitializer& O
 
 bool UGLTFAnimSequenceExporter::ExportBinary(UObject* Object, const TCHAR* Type, FArchive& Archive, FFeedbackContext* Warn, int32 FileIndex, uint32 PortFlags)
 {
-	UAnimSequence* AnimSequence = CastChecked<UAnimSequence>(Object);
+	const UAnimSequence* AnimSequence = CastChecked<UAnimSequence>(Object);
 
 	if (!FillExportOptions())
 	{
diff --git a/Engine/Plugins/Enterprise/GLTFExporter/MigratedFromGit/Source/GLTFExporter/Private/GLTFLevelExporter.cpp b/Engine/Plugins/Enterprise/GLTFExporter/MigratedFromGit/Source/GLTFExporter/Private/GLTFLevelExporter.cpp
--- a/Engine/Plugins/Enterprise/GLTFExporter/MigratedFromGit/Source/GLTFExporter/Private/GLTFLevelExporter.cpp
+++ b/Engine/Plugins/Enterprise/GLTFExporter/MigratedFromGit/Source/GLTFExporter/Private/GLTFLevelExporter.cpp
@@ -12,7 +12,7 @@ UGLTFLevelExporter::UGLTFLevelExporter(const FObjectInitializer& ObjectInitializ
 
 bool UGLTFLevelExporter::ExportBinary(UObject* Object, const TCHAR* Type, FArchive& Archive, FFeedbackContext* Warn, int32 FileIndex, uint32 PortFlags)
 {
-	UWorld* World = CastChecked<UWorld>(Object);
+	const UWorld* World = CastChecked<UWorld>(Object);
 
 	if (!FillExportOptions())
 	{
diff --git a/Engine/Plugins/Enterprise/GLTFExporter/MigratedFromGit/Source/GLTFExporter/Private/GLTFSkeletalMeshExporter.cpp b/Engine/Plugins/Enterprise/GLTFExporter/MigratedFromGit/Source/GLTFExporter/Private/GLTFSkeletalMeshExporter.cpp
--- a/Engine/Plugins/Enterprise/GLTFExporter/MigratedFromGit/Source/GLTFExporter/Private/GLTFSkeletalMeshExporter.cpp
+++ b/Engine/Plugins/Enterprise/GLTFExporter/MigratedFromGit/Source/GLTFExporter/Private/GLTFSkeletalMeshExporter.cpp
@@ -11,7 +11,7 @@ UGLTFSkeletalMeshExporter::UGLTFSkeletalMeshExporter(const FObjectInitializer& O
 
 bool UGLTFSkeletalMeshExporter::ExportBinary(UObject* Object, const TCHAR* Type, FArchive& Archive, FFeedbackContext* Warn, int32 FileIndex, uint32 PortFlags)
 {
-	USkeletalMesh* SkeletalMesh = CastChecked<USkeletalMesh>(Object);
+	const USkeletalMesh* SkeletalMesh = CastChecked<USkeletalMesh>(Object);
 
 	if (!FillExportOptions())
 	{
diff --git a/Engine/Plugins/Enterprise/GLTFExporter/MigratedFromGit/Source/GLTFExporter/Private/GLTFStaticMeshExporter.cpp b/Engine/Plugins/Enterprise/GLTFExporter/MigratedFromGit/Source/GLTFExporter/Private/GLTFStaticMeshExporter.cpp
--- a/Engine/Plugins/Enterprise/GLTFExporter/MigratedFromGit/Source/GLTFExporter/Private/GLTFStaticMeshExporter.cpp
+++ b/Engine/Plugins/Enterprise/GLTFExporter/MigratedFromGit/Source/GLTFExporter/Private/GLTFStaticMeshExporter.cpp
@@ -13,7 +13,7 @@ UGLTFStaticMeshExporter::UGLTFStaticMeshExporter(const FObjectInitializer& Objec
 
 bool UGLTFStaticMeshExporter::ExportBinary(UObject* Object, const TCHAR* Type, FArchive& Archive, FFeedbackContext* Warn, int32 FileIndex, uint32 PortFlags)
 {
-	UStaticMesh* StaticMesh = CastChecked<UStaticMesh>(Object);
+	const UStaticMesh* StaticMesh = CastChecked<UStaticMesh>(Object);
 
 	if (!FillExportOptions())
 	{
