[GLTFExporter] fix: don't register glb file format until implemented

commit 1cba1137d39f5406e920dc84981bb2bee71f364a
Date: Wed, 13 May 2020 21:24:57 +0200

    fix: don't register glb file format until implemented
    (cherry picked from commit ca0b3639dd80b6be232a90960cda9fafcd6cf077​)

#skipci

diff --git a/Engine/Plugins/Enterprise/GLTFExporter/MigratedFromGit/Source/GLTFExporter/Private/GLTFExporter.cpp b/Engine/Plugins/Enterprise/GLTFExporter/MigratedFromGit/Source/GLTFExporter/Private/GLTFExporter.cpp
--- a/Engine/Plugins/Enterprise/GLTFExporter/MigratedFromGit/Source/GLTFExporter/Private/GLTFExporter.cpp
+++ b/Engine/Plugins/Enterprise/GLTFExporter/MigratedFromGit/Source/GLTFExporter/Private/GLTFExporter.cpp
@@ -10,10 +10,14 @@ UGLTFExporter::UGLTFExporter(const FObjectInitializer& ObjectInitializer)
 	SupportedClass = nullptr;
 	bText = false;
 	PreferredFormatIndex = 0;
+
 	FormatExtension.Add(TEXT("gltf"));
-	FormatExtension.Add(TEXT("glb"));
 	FormatDescription.Add(TEXT("GL Transmission Format"));
-	FormatDescription.Add(TEXT("GL Transmission Format (Binary)"));
+
+	// TODO: uncomment when support for .glb is implemented
+	// FormatExtension.Add(TEXT("glb"));
+	// FormatDescription.Add(TEXT("GL Transmission Format (Binary)"));
+
 	ExportOptions = NewObject<UGLTFExportOptions>(this, TEXT("GLTF Export Options"));
 }
 
