[GLTFExporter] refactor: breakout GLTFAttributesBuilder from GLTFMeshBuilder and rename GLTFSectionBuilder -> GLTFPrimitiveBuilder

commit 3c0ec899e88a9ae974966489a89d5ac46d821bf0
Date: Sun, 10 May 2020 21:11:14 +0200

    refactor: breakout GLTFAttributesBuilder from GLTFMeshBuilder and rename GLTFSectionBuilder -> GLTFPrimitiveBuilder
    (cherry picked from commit 93ac853ed7d3c7b0b6fb4349c8b7acbdaa699d29​)
    # Conflicts from cherry-pick:
    #	Plugins/GLTFExporter/Source/GLTFExporter/Private/GLTFMeshBuilder.cpp

#skipci

diff --git a/Engine/Plugins/Enterprise/GLTFExporter/MigratedFromGit/Source/GLTFExporter/Private/GLTFMeshBuilder.cpp b/Engine/Plugins/Enterprise/GLTFExporter/MigratedFromGit/Source/GLTFExporter/Private/GLTFMeshBuilder.cpp
--- a/Engine/Plugins/Enterprise/GLTFExporter/MigratedFromGit/Source/GLTFExporter/Private/GLTFMeshBuilder.cpp
+++ b/Engine/Plugins/Enterprise/GLTFExporter/MigratedFromGit/Source/GLTFExporter/Private/GLTFMeshBuilder.cpp
@@ -3,56 +3,9 @@
 #include "GLTFMeshBuilder.h"
 #include "GLTFConverterUtility.h"
 
-FGLTFSectionBuilder::FGLTFSectionBuilder(const FString& SectionName, const FStaticMeshSection& MeshSection, const FIndexArrayView& IndexArray)
-	: Name(SectionName)
-{
-	const uint32 IndexCount = MeshSection.NumTriangles * 3;
-	const uint32 FirstIndex = MeshSection.FirstIndex;
-	Indices.AddUninitialized(IndexCount);
-
-	for (uint32 Index = 0; Index < IndexCount; ++Index)
-	{
-		Indices[Index] = IndexArray[FirstIndex + Index];
-	}
-}
-
-FGLTFJsonAccessorIndex FGLTFSectionBuilder::AddAccessorForIndices(FGLTFContainerBuilder& Container) const
-{
-	if (Indices.Num() == 0)
-	{
-		return FGLTFJsonAccessorIndex(INDEX_NONE);
-	}
-
-	const FString AttributeName = Name + TEXT("_Indices");
-	const FGLTFJsonBufferViewIndex BufferViewIndex = Container.AddBufferView(Indices, AttributeName, EGLTFJsonBufferTarget::ElementArrayBuffer);
-
-	FGLTFJsonAccessor Accessor;
-	Accessor.Name = AttributeName;
-	Accessor.BufferView = BufferViewIndex;
-	Accessor.ComponentType = EGLTFJsonComponentType::U32;
-	Accessor.Count = Indices.Num();
-	Accessor.Type = EGLTFJsonAccessorType::Scalar;
-
-	return Container.AddAccessor(Accessor);
-}
-
-FGLTFMeshBuilder::FGLTFMeshBuilder(const UStaticMesh* StaticMesh, int32 LODIndex, const FColorVertexBuffer* OverrideVertexColors)
+FGLTFAttributesBuilder::FGLTFAttributesBuilder(const FString& Name, const FStaticMeshLODResources& LODMesh, const FColorVertexBuffer* OverrideVertexColors)
+	: Name(Name)
 {
-	StaticMesh->GetName(Name);
-
-	const FStaticMeshLODResources& LODMesh = StaticMesh->GetLODForExport(LODIndex);
-
-	const FIndexArrayView IndexArray = LODMesh.IndexBuffer.GetArrayView();
-	const int32 SectionCount = LODMesh.Sections.Num();
-
-	Sections.Reserve(SectionCount);
-
-	for (int32 SectionIndex = 0; SectionIndex < SectionCount; ++SectionIndex)
-	{
-		FString SectionName = Name + TEXT("_Section") + FString::FromInt(SectionIndex);
-		Sections.Emplace(SectionName, LODMesh.Sections[SectionIndex], IndexArray);
-	}
-
 	const FPositionVertexBuffer& PositionVertexBuffer = LODMesh.VertexBuffers.PositionVertexBuffer;
 	const int32 PositionCount = PositionVertexBuffer.GetNumVertices();
 
@@ -63,6 +16,14 @@ FGLTFMeshBuilder::FGLTFMeshBuilder(const UStaticMesh* StaticMesh, int32 LODIndex
 		Positions[PosIndex] = FGLTFConverterUtility::ConvertPosition(PositionVertexBuffer.VertexPosition(PosIndex));
 	}
 
+	BoundingBox.Init();
+
+	// Avoid glTF validation issue by using more accurate bounding box (based on raw vertex values)
+	for (const FVector& Position : Positions)
+	{
+		BoundingBox += Position;
+	}
+
 	const FColorVertexBuffer& ColorVertexBuffer = OverrideVertexColors != nullptr ? *OverrideVertexColors : LODMesh.VertexBuffers.ColorVertexBuffer;
 	const int32 ColorCount = ColorVertexBuffer.GetNumVertices();
 
@@ -108,17 +69,9 @@ FGLTFMeshBuilder::FGLTFMeshBuilder(const UStaticMesh* StaticMesh, int32 LODIndex
 			}
 		}
 	}
-
-	BoundingBox.Init();
-
-	// Avoid glTF validation issue by using more accurate bounding box (based on raw vertex values)
-	for (const FVector& Position : Positions)
-	{
-		BoundingBox += Position;
-	}
 }
 
-FGLTFJsonAccessorIndex FGLTFMeshBuilder::AddAccessorForPositions(FGLTFContainerBuilder& Container) const
+FGLTFJsonAccessorIndex FGLTFAttributesBuilder::AddPositions(FGLTFContainerBuilder& Container) const
 {
 	if (Positions.Num() == 0)
 	{
@@ -134,19 +87,19 @@ FGLTFJsonAccessorIndex FGLTFMeshBuilder::AddAccessorForPositions(FGLTFContainerB
 	Accessor.ComponentType = EGLTFJsonComponentType::F32;
 	Accessor.Count = Positions.Num();
 	Accessor.Type = EGLTFJsonAccessorType::Vec3;
-
-	Accessor.Max[0] = BoundingBox.Max.X;
-	Accessor.Max[1] = BoundingBox.Max.Y;
-	Accessor.Max[2] = BoundingBox.Max.Z;
-	Accessor.Min[0] = BoundingBox.Min.X;
-	Accessor.Min[1] = BoundingBox.Min.Y;
-	Accessor.Min[2] = BoundingBox.Min.Z;
-	Accessor.MinMaxLength = 3;
-
-	return Container.AddAccessor(Accessor);
+	
+	JsonAccessor.MinMaxLength = 3;
+	JsonAccessor.Max[0] = BoundingBox.Max.X;
+	JsonAccessor.Max[1] = BoundingBox.Max.Y;
+	JsonAccessor.Max[2] = BoundingBox.Max.Z;
+	JsonAccessor.Min[0] = BoundingBox.Min.X;
+	JsonAccessor.Min[1] = BoundingBox.Min.Y;
+	JsonAccessor.Min[2] = BoundingBox.Min.Z;
+	
+	return Container.AddAccessor(JsonAccessor);
 }
 
-FGLTFJsonAccessorIndex FGLTFMeshBuilder::AddAccessorForNormals(FGLTFContainerBuilder& Container) const
+FGLTFJsonAccessorIndex FGLTFAttributesBuilder::AddNormals(FGLTFContainerBuilder& Container) const
 {
 	if (Normals.Num() == 0)
 	{
@@ -156,17 +109,17 @@ FGLTFJsonAccessorIndex FGLTFMeshBuilder::AddAccessorForNormals(FGLTFContainerBui
 	const FString AttributeName = Name + TEXT("_Normals");
 	const FGLTFJsonBufferViewIndex BufferViewIndex = Container.AddBufferView(Normals, AttributeName);
 
-	FGLTFJsonAccessor Accessor;
-	Accessor.Name = AttributeName;
-	Accessor.BufferView = BufferViewIndex;
-	Accessor.ComponentType = EGLTFJsonComponentType::F32;
-	Accessor.Count = Normals.Num();
-	Accessor.Type = EGLTFJsonAccessorType::Vec3;
-
-	return Container.AddAccessor(Accessor);
+	FGLTFJsonAccessor JsonAccessor;
+	JsonAccessor.Name = AttributeName;
+	JsonAccessor.BufferView = BufferViewIndex;
+	JsonAccessor.ComponentType = EGLTFJsonComponentType::F32;
+	JsonAccessor.Count = Normals.Num();
+	JsonAccessor.Type = EGLTFJsonAccessorType::Vec3;
+	
+	return Container.AddAccessor(JsonAccessor);
 }
 
-FGLTFJsonAccessorIndex FGLTFMeshBuilder::AddAccessorForColors(FGLTFContainerBuilder& Container) const
+FGLTFJsonAccessorIndex FGLTFAttributesBuilder::AddColors(FGLTFContainerBuilder& Container) const
 {
 	if (Colors.Num() == 0)
 	{
@@ -176,18 +129,18 @@ FGLTFJsonAccessorIndex FGLTFMeshBuilder::AddAccessorForColors(FGLTFContainerBuil
 	const FString AttributeName = Name + TEXT("_Colors");
 	const FGLTFJsonBufferViewIndex BufferViewIndex = Container.AddBufferView(Colors, AttributeName);
 
-	FGLTFJsonAccessor Accessor;
-	Accessor.Name = AttributeName;
-	Accessor.BufferView = BufferViewIndex;
-	Accessor.ComponentType = EGLTFJsonComponentType::U8;
-	Accessor.Count = Colors.Num();
-	Accessor.Type = EGLTFJsonAccessorType::Vec4;
-	Accessor.bNormalized = true;
+	FGLTFJsonAccessor JsonAccessor;
+	JsonAccessor.Name = AttributeName;
+	JsonAccessor.BufferView = BufferViewIndex;
+	JsonAccessor.ComponentType = EGLTFJsonComponentType::U8;
+	JsonAccessor.Count = Colors.Num();
+	JsonAccessor.Type = EGLTFJsonAccessorType::Vec4;
+	JsonAccessor.bNormalized = true;
 
-	return Container.AddAccessor(Accessor);
+	return Container.AddAccessor(JsonAccessor);
 }
 
-FGLTFJsonAccessorIndex FGLTFMeshBuilder::AddAccessorForTangents(FGLTFContainerBuilder& Container) const
+FGLTFJsonAccessorIndex FGLTFAttributesBuilder::AddTangents(FGLTFContainerBuilder& Container) const
 {
 	if (Tangents.Num() == 0)
 	{
@@ -197,17 +150,17 @@ FGLTFJsonAccessorIndex FGLTFMeshBuilder::AddAccessorForTangents(FGLTFContainerBu
 	const FString AttributeName = Name + TEXT("_Tangents");
 	const FGLTFJsonBufferViewIndex BufferViewIndex = Container.AddBufferView(Tangents, AttributeName);
 
-	FGLTFJsonAccessor Accessor;
-	Accessor.Name = AttributeName;
-	Accessor.BufferView = BufferViewIndex;
-	Accessor.ComponentType = EGLTFJsonComponentType::F32;
-	Accessor.Count = Tangents.Num();
-	Accessor.Type = EGLTFJsonAccessorType::Vec4;
+	FGLTFJsonAccessor JsonAccessor;
+	JsonAccessor.Name = AttributeName;
+	JsonAccessor.BufferView = BufferViewIndex;
+	JsonAccessor.ComponentType = EGLTFJsonComponentType::F32;
+	JsonAccessor.Count = Tangents.Num();
+	JsonAccessor.Type = EGLTFJsonAccessorType::Vec4;
 
-	return Container.AddAccessor(Accessor);
+	return Container.AddAccessor(JsonAccessor);
 }
 
-FGLTFJsonAccessorIndex FGLTFMeshBuilder::AddAccessorForUV0s(FGLTFContainerBuilder& Container) const
+FGLTFJsonAccessorIndex FGLTFAttributesBuilder::AddUV0s(FGLTFContainerBuilder& Container) const
 {
 	if (UV0s.Num() == 0)
 	{
@@ -217,17 +170,17 @@ FGLTFJsonAccessorIndex FGLTFMeshBuilder::AddAccessorForUV0s(FGLTFContainerBuilde
 	const FString AttributeName = Name + TEXT("_UV0s");
 	const FGLTFJsonBufferViewIndex BufferViewIndex = Container.AddBufferView(UV0s, AttributeName);
 
-	FGLTFJsonAccessor Accessor;
-	Accessor.Name = AttributeName;
-	Accessor.BufferView = BufferViewIndex;
-	Accessor.ComponentType = EGLTFJsonComponentType::F32;
-	Accessor.Count = UV0s.Num();
-	Accessor.Type = EGLTFJsonAccessorType::Vec2;
+	FGLTFJsonAccessor JsonAccessor;
+	JsonAccessor.Name = AttributeName;
+	JsonAccessor.BufferView = BufferViewIndex;
+	JsonAccessor.ComponentType = EGLTFJsonComponentType::F32;
+	JsonAccessor.Count = UV0s.Num();
+	JsonAccessor.Type = EGLTFJsonAccessorType::Vec2;
 
-	return Container.AddAccessor(Accessor);
+	return Container.AddAccessor(JsonAccessor);
 }
 
-FGLTFJsonAccessorIndex FGLTFMeshBuilder::AddAccessorForUV1s(FGLTFContainerBuilder& Container) const
+FGLTFJsonAccessorIndex FGLTFAttributesBuilder::AddUV1s(FGLTFContainerBuilder& Container) const
 {
 	if (UV1s.Num() == 0)
 	{
@@ -237,36 +190,100 @@ FGLTFJsonAccessorIndex FGLTFMeshBuilder::AddAccessorForUV1s(FGLTFContainerBuilde
 	const FString AttributeName = Name + TEXT("_UV1s");
 	const FGLTFJsonBufferViewIndex BufferViewIndex = Container.AddBufferView(UV1s, AttributeName);
 
-	FGLTFJsonAccessor Accessor;
-	Accessor.Name = AttributeName;
-	Accessor.BufferView = BufferViewIndex;
-	Accessor.ComponentType = EGLTFJsonComponentType::F32;
-	Accessor.Count = UV1s.Num();
-	Accessor.Type = EGLTFJsonAccessorType::Vec2;
+	FGLTFJsonAccessor JsonAccessor;
+	JsonAccessor.Name = AttributeName;
+	JsonAccessor.BufferView = BufferViewIndex;
+	JsonAccessor.ComponentType = EGLTFJsonComponentType::F32;
+	JsonAccessor.Count = UV1s.Num();
+	JsonAccessor.Type = EGLTFJsonAccessorType::Vec2;
+
+	return Container.AddAccessor(JsonAccessor);
+}
+
+FGLTFJsonAttributes FGLTFAttributesBuilder::AddAttributes(FGLTFContainerBuilder& Container) const
+{
+	FGLTFJsonAttributes JsonAttributes;
+	JsonAttributes.Position = AddPositions(Container);
+	JsonAttributes.Color0 = AddColors(Container);
+	JsonAttributes.Normal = AddNormals(Container);
+	JsonAttributes.Tangent = AddTangents(Container);
+	JsonAttributes.TexCoord0 = AddUV0s(Container);
+	JsonAttributes.TexCoord1 = AddUV1s(Container);
+	return JsonAttributes;
+}
+
+FGLTFPrimitiveBuilder::FGLTFPrimitiveBuilder(const FString& SectionName, const FStaticMeshSection& MeshSection, const FIndexArrayView& IndexArray)
+	: Name(SectionName)
+{
+	const uint32 IndexCount = MeshSection.NumTriangles * 3;
+	const uint32 FirstIndex = MeshSection.FirstIndex;
+
+	Indices.AddUninitialized(IndexCount);
 
-	return Container.AddAccessor(Accessor);
+	for (uint32 Index = 0; Index < IndexCount; ++Index)
+	{
+		Indices[Index] = IndexArray[FirstIndex + Index];
+	}
+}
+
+FGLTFJsonAccessorIndex FGLTFPrimitiveBuilder::AddIndices(FGLTFContainerBuilder& Container) const
+{
+	if (Indices.Num() == 0)
+	{
+		return FGLTFJsonAccessorIndex(INDEX_NONE);
+	}
+
+	const FString AttributeName = Name + TEXT("_Indices");
+	const FGLTFJsonBufferViewIndex BufferViewIndex = Container.AddBufferView(Indices, AttributeName, EGLTFJsonBufferTarget::ElementArrayBuffer);
+
+	FGLTFJsonAccessor JsonAccessor;
+	JsonAccessor.Name = AttributeName;
+	JsonAccessor.BufferView = BufferViewIndex;
+	JsonAccessor.ComponentType = EGLTFJsonComponentType::U32;
+	JsonAccessor.Count = Indices.Num();
+	JsonAccessor.Type = EGLTFJsonAccessorType::Scalar;
+
+	return Container.AddAccessor(JsonAccessor);
+}
+
+FGLTFJsonPrimitive FGLTFPrimitiveBuilder::AddPrimitive(FGLTFContainerBuilder& Container, const FGLTFJsonAttributes& JsonAttributes) const
+{
+	FGLTFJsonPrimitive JsonPrimitive;
+	JsonPrimitive.Indices = AddIndices(Container);
+	JsonPrimitive.Attributes = JsonAttributes;
+	return JsonPrimitive;
+}
+
+FGLTFMeshBuilder::FGLTFMeshBuilder(const UStaticMesh* StaticMesh, int32 LODIndex, const FColorVertexBuffer* OverrideVertexColors)
+	: FGLTFMeshBuilder(StaticMesh->GetName() + (LODIndex != 0 ? TEXT("_LOD") + FString::FromInt(LODIndex) : FString()), StaticMesh->GetLODForExport(LODIndex), OverrideVertexColors)
+{
+}
+
+FGLTFMeshBuilder::FGLTFMeshBuilder(const FString& Name, const FStaticMeshLODResources& LODMesh, const FColorVertexBuffer* OverrideVertexColors)
+	: Attributes(Name, LODMesh, OverrideVertexColors)
+{
+	const FIndexArrayView IndexArray = LODMesh.IndexBuffer.GetArrayView();
+	const int32 SectionCount = LODMesh.Sections.Num();
+
+	Primitives.Reserve(SectionCount);
+
+	for (int32 SectionIndex = 0; SectionIndex < SectionCount; ++SectionIndex)
+	{
+		FString SectionName = Name + TEXT("_Section") + FString::FromInt(SectionIndex);
+		Primitives.Emplace(SectionName, LODMesh.Sections[SectionIndex], IndexArray);
+	}
 }
 
 FGLTFJsonMeshIndex FGLTFMeshBuilder::AddMesh(FGLTFContainerBuilder& Container) const
 {
+	const FGLTFJsonAttributes JsonAttributes = Attributes.AddAttributes(Container);
+	
 	FGLTFJsonMesh Mesh;
 	Mesh.Name = Name;
 
-	FGLTFJsonAttributes Attributes;
-	Attributes.Position = AddAccessorForPositions(Container);
-	Attributes.Color0 = AddAccessorForColors(Container);
-	Attributes.Normal = AddAccessorForNormals(Container);
-	Attributes.Tangent = AddAccessorForTangents(Container);
-	Attributes.TexCoord0 = AddAccessorForUV0s(Container);
-	Attributes.TexCoord1 = AddAccessorForUV1s(Container);
-
-	for (const FGLTFSectionBuilder& Section : Sections)
+	for (const FGLTFPrimitiveBuilder& Primitive : Primitives)
 	{
-		FGLTFJsonPrimitive Primitive;
-		Primitive.Attributes = Attributes;
-
-		Primitive.Indices = Section.AddAccessorForIndices(Container);
-		Mesh.Primitives.Add(Primitive);
+		Mesh.Primitives.Add(Primitive.AddPrimitive(Container, JsonAttributes));
 	}
 
 	return Container.AddMesh(Mesh);
diff --git a/Engine/Plugins/Enterprise/GLTFExporter/MigratedFromGit/Source/GLTFExporter/Public/GLTFMeshBuilder.h b/Engine/Plugins/Enterprise/GLTFExporter/MigratedFromGit/Source/GLTFExporter/Public/GLTFMeshBuilder.h
--- a/Engine/Plugins/Enterprise/GLTFExporter/MigratedFromGit/Source/GLTFExporter/Public/GLTFMeshBuilder.h
+++ b/Engine/Plugins/Enterprise/GLTFExporter/MigratedFromGit/Source/GLTFExporter/Public/GLTFMeshBuilder.h
@@ -6,40 +6,53 @@
 #include "GLTFContainerBuilder.h"
 #include "Engine.h"
 
-struct GLTFEXPORTER_API FGLTFSectionBuilder
+struct GLTFEXPORTER_API FGLTFAttributesBuilder
 {
 	FString Name;
 
-	TArray<uint32> Indices;
+	FBox BoundingBox;
 
-	FGLTFSectionBuilder(const FString& SectionName, const FStaticMeshSection& MeshSection, const FIndexArrayView& IndexArray);
+	TArray<FVector>   Positions;
+	TArray<FColor>    Colors;
+	TArray<FVector>   Normals;
+	TArray<FVector4>  Tangents;
+	TArray<FVector2D> UV0s;
+	TArray<FVector2D> UV1s;
+
+	FGLTFAttributesBuilder(const FString& Name, const FStaticMeshLODResources& LODMesh, const FColorVertexBuffer* OverrideVertexColors);
+
+	FGLTFJsonAccessorIndex AddPositions(FGLTFContainerBuilder& Container) const;
+	FGLTFJsonAccessorIndex AddColors(FGLTFContainerBuilder& Container) const;
+	FGLTFJsonAccessorIndex AddNormals(FGLTFContainerBuilder& Container) const;
+	FGLTFJsonAccessorIndex AddTangents(FGLTFContainerBuilder& Container) const;
+	FGLTFJsonAccessorIndex AddUV0s(FGLTFContainerBuilder& Container) const;
+	FGLTFJsonAccessorIndex AddUV1s(FGLTFContainerBuilder& Container) const;
 
-	FGLTFJsonAccessorIndex AddAccessorForIndices(FGLTFContainerBuilder& Container) const;
+	FGLTFJsonAttributes AddAttributes(FGLTFContainerBuilder& Container) const;
 };
 
-struct GLTFEXPORTER_API FGLTFMeshBuilder
+struct GLTFEXPORTER_API FGLTFPrimitiveBuilder
 {
 	FString Name;
 
-	TArray<FGLTFSectionBuilder> Sections;
+	TArray<uint32> Indices;
 
-	TArray<FVector>   Positions;
-	TArray<FColor>    Colors;
-	TArray<FVector>   Normals;
-	TArray<FVector4>  Tangents;
-	TArray<FVector2D> UV0s;
-	TArray<FVector2D> UV1s;
+	FGLTFPrimitiveBuilder(const FString& Name, const FStaticMeshSection& MeshSection, const FIndexArrayView& IndexArray);
 
-	FBox BoundingBox;
+	FGLTFJsonAccessorIndex AddIndices(FGLTFContainerBuilder& Container) const;
 
-	FGLTFMeshBuilder(const UStaticMesh* StaticMesh, int32 LODIndex, const FColorVertexBuffer* OverrideVertexColors);
+	FGLTFJsonPrimitive AddPrimitive(FGLTFContainerBuilder& Container, const FGLTFJsonAttributes& JsonAttributes) const;
+};
+
+struct GLTFEXPORTER_API FGLTFMeshBuilder
+{
+	FString Name;
+	
+	FGLTFAttributesBuilder Attributes;
+	TArray<FGLTFPrimitiveBuilder> Primitives;
 
-	FGLTFJsonAccessorIndex AddAccessorForPositions(FGLTFContainerBuilder& Container) const;
-	FGLTFJsonAccessorIndex AddAccessorForColors(FGLTFContainerBuilder& Container) const;
-	FGLTFJsonAccessorIndex AddAccessorForNormals(FGLTFContainerBuilder& Container) const;
-	FGLTFJsonAccessorIndex AddAccessorForTangents(FGLTFContainerBuilder& Container) const;
-	FGLTFJsonAccessorIndex AddAccessorForUV0s(FGLTFContainerBuilder& Container) const;
-	FGLTFJsonAccessorIndex AddAccessorForUV1s(FGLTFContainerBuilder& Container) const;
+	FGLTFMeshBuilder(const UStaticMesh* StaticMesh, int32 LODIndex, const FColorVertexBuffer* OverrideVertexColors);
+	FGLTFMeshBuilder(const FString& Name, const FStaticMeshLODResources& LODMesh, const FColorVertexBuffer* OverrideVertexColors);
 
 	FGLTFJsonMeshIndex AddMesh(FGLTFContainerBuilder& Container) const;
 };
