[GLTFExporter] feat: add GLTFIndexBuilder to avoid duplicate mesh export

commit 9ed36d9d4d0a06359829e518f39477b3c0b5d588
Date: Sun, 10 May 2020 12:35:08 +0200

    feat: add GLTFIndexBuilder to avoid duplicate mesh export
    (cherry picked from commit 13a008ea48fc35bc703860673c5d8e95dbee4e11​)

#skipci

diff --git a/Engine/Plugins/Enterprise/GLTFExporter/MigratedFromGit/Source/GLTFExporter/Private/GLTFBufferBuilder.cpp b/Engine/Plugins/Enterprise/GLTFExporter/MigratedFromGit/Source/GLTFExporter/Private/GLTFBufferBuilder.cpp
--- a/Engine/Plugins/Enterprise/GLTFExporter/MigratedFromGit/Source/GLTFExporter/Private/GLTFBufferBuilder.cpp
+++ b/Engine/Plugins/Enterprise/GLTFExporter/MigratedFromGit/Source/GLTFExporter/Private/GLTFBufferBuilder.cpp
@@ -1,7 +1,7 @@
 // Copyright Epic Games, Inc. All Rights Reserved.
 
 #include "GLTFBufferBuilder.h"
-#include "GLTFMeshBuilder.h"
+#include "GLTFContainerBuilder.h"
 #include "Misc/Base64.h"
 
 FGLTFBufferBuilder::FGLTFBufferBuilder(FGLTFJsonBufferIndex BufferIndex)
diff --git a/Engine/Plugins/Enterprise/GLTFExporter/MigratedFromGit/Source/GLTFExporter/Private/GLTFContainerBuilder.cpp b/Engine/Plugins/Enterprise/GLTFExporter/MigratedFromGit/Source/GLTFExporter/Private/GLTFContainerBuilder.cpp
--- a/Engine/Plugins/Enterprise/GLTFExporter/MigratedFromGit/Source/GLTFExporter/Private/GLTFContainerBuilder.cpp
+++ b/Engine/Plugins/Enterprise/GLTFExporter/MigratedFromGit/Source/GLTFExporter/Private/GLTFContainerBuilder.cpp
@@ -53,7 +53,7 @@ void FGLTFContainerBuilder::Serialize(FArchive& Archive)
 
 FGLTFJsonMeshIndex FGLTFContainerBuilder::AddMesh(const UStaticMesh* StaticMesh, int32 LODIndex, const FColorVertexBuffer* OverrideVertexColors)
 {
-	return FGLTFMeshBuilder(StaticMesh, LODIndex, OverrideVertexColors).AddMesh(*this);
+	return IndexBuilder.FindOrAdd(FGLTFStaticMeshKey(StaticMesh, LODIndex, OverrideVertexColors), *this);
 }
 
 FGLTFJsonMeshIndex FGLTFContainerBuilder::AddMesh(const UStaticMeshComponent* StaticMeshComponent)
diff --git a/Engine/Plugins/Enterprise/GLTFExporter/MigratedFromGit/Source/GLTFExporter/Private/GLTFIndexBuilder.cpp b/Engine/Plugins/Enterprise/GLTFExporter/MigratedFromGit/Source/GLTFExporter/Private/GLTFIndexBuilder.cpp
new file mode 100644
--- /dev/null
+++ b/Engine/Plugins/Enterprise/GLTFExporter/MigratedFromGit/Source/GLTFExporter/Private/GLTFIndexBuilder.cpp
@@ -0,0 +1,24 @@
+// Copyright Epic Games, Inc. All Rights Reserved.
+
+#include "GLTFIndexBuilder.h"
+#include "GLTFContainerBuilder.h"
+#include "GLTFMeshBuilder.h"
+
+FGLTFJsonMeshIndex FGLTFIndexBuilder::Find(const FGLTFStaticMeshKey& Key) const
+{
+	const FGLTFJsonMeshIndex* Index = StaticMeshes.Find(Key);
+	return Index != nullptr ? *Index : FGLTFJsonMeshIndex(INDEX_NONE);
+}
+
+FGLTFJsonMeshIndex FGLTFIndexBuilder::FindOrAdd(const FGLTFStaticMeshKey& Key, FGLTFContainerBuilder& Container)
+{
+	FGLTFJsonMeshIndex Index = Find(Key);
+
+	if (Index == INDEX_NONE)
+	{
+		Index = FGLTFMeshBuilder(Key.StaticMesh, Key.LODIndex, Key.OverrideVertexColors).AddMesh(Container);
+		StaticMeshes.Add(Key, Index);
+	}
+	
+	return Index;
+}
diff --git a/Engine/Plugins/Enterprise/GLTFExporter/MigratedFromGit/Source/GLTFExporter/Public/GLTFContainerBuilder.h b/Engine/Plugins/Enterprise/GLTFExporter/MigratedFromGit/Source/GLTFExporter/Public/GLTFContainerBuilder.h
--- a/Engine/Plugins/Enterprise/GLTFExporter/MigratedFromGit/Source/GLTFExporter/Public/GLTFContainerBuilder.h
+++ b/Engine/Plugins/Enterprise/GLTFExporter/MigratedFromGit/Source/GLTFExporter/Public/GLTFContainerBuilder.h
@@ -4,12 +4,14 @@
 
 #include "GLTFJsonRoot.h"
 #include "GLTFBufferBuilder.h"
+#include "GLTFIndexBuilder.h"
 #include "Engine.h"
 
 struct GLTFEXPORTER_API FGLTFContainerBuilder
 {
 	FGLTFJsonRoot JsonRoot;
 	FGLTFBufferBuilder BufferBuilder;
+	FGLTFIndexBuilder IndexBuilder;
 
 	FGLTFContainerBuilder();
 
diff --git a/Engine/Plugins/Enterprise/GLTFExporter/MigratedFromGit/Source/GLTFExporter/Public/GLTFIndexBuilder.h b/Engine/Plugins/Enterprise/GLTFExporter/MigratedFromGit/Source/GLTFExporter/Public/GLTFIndexBuilder.h
new file mode 100644
--- /dev/null
+++ b/Engine/Plugins/Enterprise/GLTFExporter/MigratedFromGit/Source/GLTFExporter/Public/GLTFIndexBuilder.h
@@ -0,0 +1,42 @@
+// Copyright Epic Games, Inc. All Rights Reserved.
+
+#pragma once
+
+#include "GLTFJsonIndex.h"
+#include "Engine.h"
+
+struct FGLTFContainerBuilder;
+
+struct FGLTFStaticMeshKey
+{
+	const UStaticMesh* StaticMesh;
+	int32 LODIndex;
+	const FColorVertexBuffer* OverrideVertexColors;
+
+	FGLTFStaticMeshKey(const UStaticMesh* StaticMesh, int32 LODIndex, const FColorVertexBuffer* OverrideVertexColors)
+		: StaticMesh(StaticMesh)
+		, LODIndex(LODIndex)
+		, OverrideVertexColors(OverrideVertexColors)
+	{
+	}
+
+	FORCEINLINE bool operator==(const FGLTFStaticMeshKey& Other) const
+	{
+		return StaticMesh           == Other.StaticMesh
+		    && LODIndex             == Other.LODIndex
+			&& OverrideVertexColors == Other.OverrideVertexColors;
+	}
+
+	FORCEINLINE friend uint32 GetTypeHash(const FGLTFStaticMeshKey& Other)
+	{
+		return HashCombine(GetTypeHash(Other.StaticMesh), HashCombine(GetTypeHash(Other.LODIndex), GetTypeHash(Other.OverrideVertexColors)));
+	}
+};
+
+struct GLTFEXPORTER_API FGLTFIndexBuilder
+{
+	TMap<FGLTFStaticMeshKey, FGLTFJsonMeshIndex> StaticMeshes;
+
+	FGLTFJsonMeshIndex Find(const FGLTFStaticMeshKey& Key) const;
+	FGLTFJsonMeshIndex FindOrAdd(const FGLTFStaticMeshKey& Key, FGLTFContainerBuilder& Container);
+};
