[GLTFExporter] feat: allow each mesh accessor to be reusable

commit 137bd67d78032e3f46be11489fae6958757d06dd
Date: Sun, 10 May 2020 12:40:24 +0200

    feat: allow each mesh accessor to be reusable
    (cherry picked from commit b2b86166970c4831fb38cfce60bc5fb4c14619e1​)
    # Conflicts from cherry-pick:
    #	Plugins/GLTFExporter/Source/GLTFExporter/Private/GLTFMeshBuilder.cpp
    #	Plugins/GLTFExporter/Source/GLTFExporter/Private/GLTFSceneBuilder.cpp

#skipci

diff --git a/Engine/Plugins/Enterprise/GLTFExporter/MigratedFromGit/Source/GLTFExporter/Private/GLTFBufferBuilder.cpp b/Engine/Plugins/Enterprise/GLTFExporter/MigratedFromGit/Source/GLTFExporter/Private/GLTFBufferBuilder.cpp
--- a/Engine/Plugins/Enterprise/GLTFExporter/MigratedFromGit/Source/GLTFExporter/Private/GLTFBufferBuilder.cpp
+++ b/Engine/Plugins/Enterprise/GLTFExporter/MigratedFromGit/Source/GLTFExporter/Private/GLTFBufferBuilder.cpp
@@ -9,7 +9,7 @@ FGLTFBufferBuilder::FGLTFBufferBuilder(FGLTFJsonBufferIndex BufferIndex)
 {
 }
 
-FGLTFJsonBufferViewIndex FGLTFBufferBuilder::AddBufferView(FGLTFContainerBuilder& Container, const void* RawData, uint64 ByteLength, const FString& Name, EGLTFJsonBufferTarget BufferTarget)
+FGLTFJsonBufferViewIndex FGLTFBufferBuilder::CreateBufferView(FGLTFContainerBuilder& Container, const void* RawData, uint64 ByteLength, const FString& Name, EGLTFJsonBufferTarget BufferTarget)
 {
 	FGLTFJsonBufferView BufferView;
 	BufferView.Name = Name;
@@ -20,7 +20,7 @@ FGLTFJsonBufferViewIndex FGLTFBufferBuilder::AddBufferView(FGLTFContainerBuilder
 
 	BufferData.Append(static_cast<const uint8*>(RawData), ByteLength);
 
-	return Container.AddBufferView(BufferView);
+	return Container.CreateBufferView(BufferView);
 }
 
 void FGLTFBufferBuilder::UpdateBuffer(FGLTFJsonBuffer& JsonBuffer)
diff --git a/Engine/Plugins/Enterprise/GLTFExporter/MigratedFromGit/Source/GLTFExporter/Private/GLTFContainerBuilder.cpp b/Engine/Plugins/Enterprise/GLTFExporter/MigratedFromGit/Source/GLTFExporter/Private/GLTFContainerBuilder.cpp
--- a/Engine/Plugins/Enterprise/GLTFExporter/MigratedFromGit/Source/GLTFExporter/Private/GLTFContainerBuilder.cpp
+++ b/Engine/Plugins/Enterprise/GLTFExporter/MigratedFromGit/Source/GLTFExporter/Private/GLTFContainerBuilder.cpp
@@ -1,47 +1,86 @@
 // Copyright Epic Games, Inc. All Rights Reserved.
 
 #include "GLTFContainerBuilder.h"
-#include "GLTFMeshBuilder.h"
 #include "GLTFSceneBuilder.h"
 
 FGLTFContainerBuilder::FGLTFContainerBuilder()
-	: BufferBuilder(AddBuffer(FGLTFJsonBuffer()))
+	: BufferBuilder(CreateBuffer(FGLTFJsonBuffer()))
 {
 }
 
-FGLTFJsonAccessorIndex FGLTFContainerBuilder::AddAccessor(const FGLTFJsonAccessor& JsonAccessor)
+FGLTFJsonAccessorIndex FGLTFContainerBuilder::CreateAccessor(const FGLTFJsonAccessor& JsonAccessor)
 {
 	return FGLTFJsonAccessorIndex(JsonRoot.Accessors.Add(JsonAccessor));
 }
 
-FGLTFJsonBufferIndex FGLTFContainerBuilder::AddBuffer(const FGLTFJsonBuffer& JsonBuffer)
+FGLTFJsonBufferIndex FGLTFContainerBuilder::CreateBuffer(const FGLTFJsonBuffer& JsonBuffer)
 {
 	return FGLTFJsonBufferIndex(JsonRoot.Buffers.Add(JsonBuffer));
 }
 
-FGLTFJsonBufferViewIndex FGLTFContainerBuilder::AddBufferView(const FGLTFJsonBufferView& JsonBufferView)
+FGLTFJsonBufferViewIndex FGLTFContainerBuilder::CreateBufferView(const FGLTFJsonBufferView& JsonBufferView)
 {
 	return FGLTFJsonBufferViewIndex(JsonRoot.BufferViews.Add(JsonBufferView));
 }
 
-FGLTFJsonMeshIndex FGLTFContainerBuilder::AddMesh(const FGLTFJsonMesh& JsonMesh)
+FGLTFJsonMeshIndex FGLTFContainerBuilder::CreateMesh(const FGLTFJsonMesh& JsonMesh)
 {
 	return FGLTFJsonMeshIndex(JsonRoot.Meshes.Add(JsonMesh));
 }
 
-FGLTFJsonNodeIndex FGLTFContainerBuilder::AddNode(const FGLTFJsonNode& JsonNode)
+FGLTFJsonNodeIndex FGLTFContainerBuilder::CreateNode(const FGLTFJsonNode& JsonNode)
 {
 	return FGLTFJsonNodeIndex(JsonRoot.Nodes.Add(JsonNode));
 }
 
-FGLTFJsonSceneIndex FGLTFContainerBuilder::AddScene(const FGLTFJsonScene& JsonScene)
+FGLTFJsonSceneIndex FGLTFContainerBuilder::CreateScene(const FGLTFJsonScene& JsonScene)
 {
 	return FGLTFJsonSceneIndex(JsonRoot.Scenes.Add(JsonScene));
 }
 
-FGLTFJsonBufferViewIndex FGLTFContainerBuilder::AddBufferView(const void* RawData, uint64 ByteLength, const FString& Name, EGLTFJsonBufferTarget BufferTarget)
+FGLTFJsonBufferViewIndex FGLTFContainerBuilder::CreateBufferView(const void* RawData, uint64 ByteLength, const FString& Name, EGLTFJsonBufferTarget BufferTarget)
 {
-	return BufferBuilder.AddBufferView(*this, RawData, ByteLength, Name, BufferTarget);
+	return BufferBuilder.CreateBufferView(*this, RawData, ByteLength, Name, BufferTarget);
+}
+
+FGLTFJsonAccessorIndex FGLTFContainerBuilder::AddPositionAccessor(const FPositionVertexBuffer* VertexBuffer, const FString& Name)
+{
+	return IndexBuilder.FindOrConvertPositionAccessor(FGLTFPositionVertexBufferKey(VertexBuffer, Name), *this);
+}
+
+FGLTFJsonAccessorIndex FGLTFContainerBuilder::AddNormalAccessor(const FStaticMeshVertexBuffer* VertexBuffer, const FString& Name)
+{
+	return IndexBuilder.FindOrConvertNormalAccessor(FGLTFStaticMeshVertexBufferKey(VertexBuffer, Name), *this);
+}
+
+FGLTFJsonAccessorIndex FGLTFContainerBuilder::AddTangentAccessor(const FStaticMeshVertexBuffer* VertexBuffer, const FString& Name)
+{
+	return IndexBuilder.FindOrConvertTangentAccessor(FGLTFStaticMeshVertexBufferKey(VertexBuffer, Name), *this);
+}
+
+FGLTFJsonAccessorIndex FGLTFContainerBuilder::AddUV0Accessor(const FStaticMeshVertexBuffer* VertexBuffer, const FString& Name)
+{
+	return IndexBuilder.FindOrConvertUV0Accessor(FGLTFStaticMeshVertexBufferKey(VertexBuffer, Name), *this);
+}
+
+FGLTFJsonAccessorIndex FGLTFContainerBuilder::AddUV1Accessor(const FStaticMeshVertexBuffer* VertexBuffer, const FString& Name)
+{
+	return IndexBuilder.FindOrConvertUV1Accessor(FGLTFStaticMeshVertexBufferKey(VertexBuffer, Name), *this);
+}
+
+FGLTFJsonAccessorIndex FGLTFContainerBuilder::AddColorAccessor(const FColorVertexBuffer* VertexBuffer, const FString& Name)
+{
+	return IndexBuilder.FindOrConvertColorAccessor(FGLTFColorVertexBufferKey(VertexBuffer, Name), *this);
+}
+
+FGLTFJsonBufferViewIndex FGLTFContainerBuilder::AddIndexBufferView(const FRawStaticIndexBuffer* IndexBuffer, const FString& Name)
+{
+	return IndexBuilder.FindOrConvertIndexBufferView(FGLTFRawStaticIndexBufferKey(IndexBuffer, Name), *this);
+}
+
+FGLTFJsonAccessorIndex FGLTFContainerBuilder::AddIndexAccessor(const FStaticMeshSection* MeshSection, const FRawStaticIndexBuffer* IndexBuffer, const FString& Name)
+{
+	return IndexBuilder.FindOrConvertIndexAccessor(FGLTFStaticMeshSectionKey(MeshSection, IndexBuffer, Name), *this);
 }
 
 void FGLTFContainerBuilder::Serialize(FArchive& Archive)
@@ -53,7 +92,7 @@ void FGLTFContainerBuilder::Serialize(FArchive& Archive)
 
 FGLTFJsonMeshIndex FGLTFContainerBuilder::AddMesh(const UStaticMesh* StaticMesh, int32 LODIndex, const FColorVertexBuffer* OverrideVertexColors)
 {
-	return IndexBuilder.FindOrAdd(FGLTFStaticMeshKey(StaticMesh, LODIndex, OverrideVertexColors), *this);
+	return IndexBuilder.FindOrConvertMesh(FGLTFStaticMeshKey(StaticMesh, LODIndex, OverrideVertexColors), *this);
 }
 
 FGLTFJsonMeshIndex FGLTFContainerBuilder::AddMesh(const UStaticMeshComponent* StaticMeshComponent)
diff --git a/Engine/Plugins/Enterprise/GLTFExporter/MigratedFromGit/Source/GLTFExporter/Private/GLTFIndexBuilder.cpp b/Engine/Plugins/Enterprise/GLTFExporter/MigratedFromGit/Source/GLTFExporter/Private/GLTFIndexBuilder.cpp
--- a/Engine/Plugins/Enterprise/GLTFExporter/MigratedFromGit/Source/GLTFExporter/Private/GLTFIndexBuilder.cpp
+++ b/Engine/Plugins/Enterprise/GLTFExporter/MigratedFromGit/Source/GLTFExporter/Private/GLTFIndexBuilder.cpp
@@ -2,23 +2,176 @@
 
 #include "GLTFIndexBuilder.h"
 #include "GLTFContainerBuilder.h"
-#include "GLTFMeshBuilder.h"
+#include "GLTFMeshConverter.h"
 
-FGLTFJsonMeshIndex FGLTFIndexBuilder::Find(const FGLTFStaticMeshKey& Key) const
+FGLTFJsonMeshIndex FGLTFIndexBuilder::FindMesh(const FGLTFStaticMeshKey& Key) const
 {
 	const FGLTFJsonMeshIndex* Index = StaticMeshes.Find(Key);
 	return Index != nullptr ? *Index : FGLTFJsonMeshIndex(INDEX_NONE);
 }
 
-FGLTFJsonMeshIndex FGLTFIndexBuilder::FindOrAdd(const FGLTFStaticMeshKey& Key, FGLTFContainerBuilder& Container)
+FGLTFJsonAccessorIndex FGLTFIndexBuilder::FindPositionAccessor(const FGLTFPositionVertexBufferKey& Key) const
 {
-	FGLTFJsonMeshIndex Index = Find(Key);
+	const FGLTFJsonAccessorIndex* Index = PositionVertexBuffers.Find(Key);
+	return Index != nullptr ? *Index : FGLTFJsonAccessorIndex(INDEX_NONE);
+}
+
+FGLTFJsonAccessorIndex FGLTFIndexBuilder::FindNormalAccessor(const FGLTFStaticMeshVertexBufferKey& Key) const
+{
+	const FGLTFJsonAccessorIndex* Index = StaticNormalVertexBuffers.Find(Key);
+	return Index != nullptr ? *Index : FGLTFJsonAccessorIndex(INDEX_NONE);
+}
+
+FGLTFJsonAccessorIndex FGLTFIndexBuilder::FindTangentAccessor(const FGLTFStaticMeshVertexBufferKey& Key) const
+{
+	const FGLTFJsonAccessorIndex* Index = StaticTangentVertexBuffers.Find(Key);
+	return Index != nullptr ? *Index : FGLTFJsonAccessorIndex(INDEX_NONE);
+}
+
+FGLTFJsonAccessorIndex FGLTFIndexBuilder::FindUV0Accessor(const FGLTFStaticMeshVertexBufferKey& Key) const
+{
+	const FGLTFJsonAccessorIndex* Index = StaticUV0VertexBuffers.Find(Key);
+	return Index != nullptr ? *Index : FGLTFJsonAccessorIndex(INDEX_NONE);
+}
+
+FGLTFJsonAccessorIndex FGLTFIndexBuilder::FindUV1Accessor(const FGLTFStaticMeshVertexBufferKey& Key) const
+{
+	const FGLTFJsonAccessorIndex* Index = StaticUV1VertexBuffers.Find(Key);
+	return Index != nullptr ? *Index : FGLTFJsonAccessorIndex(INDEX_NONE);
+}
+
+FGLTFJsonAccessorIndex FGLTFIndexBuilder::FindColorAccessor(const FGLTFColorVertexBufferKey& Key) const
+{
+	const FGLTFJsonAccessorIndex* Index = ColorVertexBuffers.Find(Key);
+	return Index != nullptr ? *Index : FGLTFJsonAccessorIndex(INDEX_NONE);
+}
+
+FGLTFJsonBufferViewIndex FGLTFIndexBuilder::FindIndexBufferView(const FGLTFRawStaticIndexBufferKey& Key) const
+{
+	const FGLTFJsonBufferViewIndex* Index = StaticIndexBuffers.Find(Key);
+	return Index != nullptr ? *Index : FGLTFJsonBufferViewIndex(INDEX_NONE);
+}
+
+FGLTFJsonAccessorIndex FGLTFIndexBuilder::FindIndexAccessor(const FGLTFStaticMeshSectionKey& Key) const
+{
+	const FGLTFJsonAccessorIndex* Index = StaticMeshSections.Find(Key);
+	return Index != nullptr ? *Index : FGLTFJsonAccessorIndex(INDEX_NONE);
+}
+
+
+FGLTFJsonMeshIndex FGLTFIndexBuilder::FindOrConvertMesh(const FGLTFStaticMeshKey& Key, FGLTFContainerBuilder& Container)
+{
+	FGLTFJsonMeshIndex Index = FindMesh(Key);
 
 	if (Index == INDEX_NONE)
 	{
-		Index = FGLTFMeshBuilder(Key.StaticMesh, Key.LODIndex, Key.OverrideVertexColors).AddMesh(Container);
+		Index = FGLTFMeshConverter::ConvertMesh(Container, Key.StaticMesh, Key.LODIndex, Key.OverrideVertexColors);
 		StaticMeshes.Add(Key, Index);
 	}
 	
 	return Index;
 }
+
+FGLTFJsonAccessorIndex FGLTFIndexBuilder::FindOrConvertPositionAccessor(const FGLTFPositionVertexBufferKey& Key, FGLTFContainerBuilder& Container)
+{
+	FGLTFJsonAccessorIndex Index = FindPositionAccessor(Key);
+
+	if (Index == INDEX_NONE)
+	{
+		Index = FGLTFMeshConverter::ConvertPositionAccessor(Container, Key.VertexBuffer, Key.Name);
+		PositionVertexBuffers.Add(Key, Index);
+	}
+	
+	return Index;
+}
+
+FGLTFJsonAccessorIndex FGLTFIndexBuilder::FindOrConvertNormalAccessor(const FGLTFStaticMeshVertexBufferKey& Key, FGLTFContainerBuilder& Container)
+{
+	FGLTFJsonAccessorIndex Index = FindNormalAccessor(Key);
+
+	if (Index == INDEX_NONE)
+	{
+		Index = FGLTFMeshConverter::ConvertNormalAccessor(Container, Key.VertexBuffer, Key.Name);
+		StaticNormalVertexBuffers.Add(Key, Index);
+	}
+	
+	return Index;
+}
+
+FGLTFJsonAccessorIndex FGLTFIndexBuilder::FindOrConvertTangentAccessor(const FGLTFStaticMeshVertexBufferKey& Key, FGLTFContainerBuilder& Container)
+{
+	FGLTFJsonAccessorIndex Index = FindTangentAccessor(Key);
+
+	if (Index == INDEX_NONE)
+	{
+		Index = FGLTFMeshConverter::ConvertTangentAccessor(Container, Key.VertexBuffer, Key.Name);
+		StaticTangentVertexBuffers.Add(Key, Index);
+	}
+	
+	return Index;
+}
+
+FGLTFJsonAccessorIndex FGLTFIndexBuilder::FindOrConvertUV0Accessor(const FGLTFStaticMeshVertexBufferKey& Key, FGLTFContainerBuilder& Container)
+{
+	FGLTFJsonAccessorIndex Index = FindUV0Accessor(Key);
+
+	if (Index == INDEX_NONE)
+	{
+		Index = FGLTFMeshConverter::ConvertUV0Accessor(Container, Key.VertexBuffer, Key.Name);
+		StaticUV0VertexBuffers.Add(Key, Index);
+	}
+	
+	return Index;
+}
+
+FGLTFJsonAccessorIndex FGLTFIndexBuilder::FindOrConvertUV1Accessor(const FGLTFStaticMeshVertexBufferKey& Key, FGLTFContainerBuilder& Container)
+{
+	FGLTFJsonAccessorIndex Index = FindUV1Accessor(Key);
+
+	if (Index == INDEX_NONE)
+	{
+		Index = FGLTFMeshConverter::ConvertUV1Accessor(Container, Key.VertexBuffer, Key.Name);
+		StaticUV1VertexBuffers.Add(Key, Index);
+	}
+	
+	return Index;
+}
+
+FGLTFJsonAccessorIndex FGLTFIndexBuilder::FindOrConvertColorAccessor(const FGLTFColorVertexBufferKey& Key, FGLTFContainerBuilder& Container)
+{
+	FGLTFJsonAccessorIndex Index = FindColorAccessor(Key);
+
+	if (Index == INDEX_NONE)
+	{
+		Index = FGLTFMeshConverter::ConvertColorAccessor(Container, Key.VertexBuffer, Key.Name);
+		ColorVertexBuffers.Add(Key, Index);
+	}
+	
+	return Index;
+}
+
+FGLTFJsonBufferViewIndex FGLTFIndexBuilder::FindOrConvertIndexBufferView(const FGLTFRawStaticIndexBufferKey& Key, FGLTFContainerBuilder& Container)
+{
+	FGLTFJsonBufferViewIndex Index = FindIndexBufferView(Key);
+
+	if (Index == INDEX_NONE)
+	{
+		Index = FGLTFMeshConverter::ConvertIndexBufferView(Container, Key.IndexBuffer, Key.Name);
+		StaticIndexBuffers.Add(Key, Index);
+	}
+	
+	return Index;
+}
+
+FGLTFJsonAccessorIndex FGLTFIndexBuilder::FindOrConvertIndexAccessor(const FGLTFStaticMeshSectionKey& Key, FGLTFContainerBuilder& Container)
+{
+	FGLTFJsonAccessorIndex Index = FindIndexAccessor(Key);
+
+	if (Index == INDEX_NONE)
+	{
+		Index = FGLTFMeshConverter::ConvertIndexAccessor(Container, Key.MeshSection, Key.IndexBuffer, Key.Name);
+		StaticMeshSections.Add(Key, Index);
+	}
+	
+	return Index;
+}
diff --git a/Engine/Plugins/Enterprise/GLTFExporter/MigratedFromGit/Source/GLTFExporter/Private/GLTFMeshBuilder.cpp b/Engine/Plugins/Enterprise/GLTFExporter/MigratedFromGit/Source/GLTFExporter/Private/GLTFMeshBuilder.cpp
deleted file mode 100644
diff --git a/Engine/Plugins/Enterprise/GLTFExporter/MigratedFromGit/Source/GLTFExporter/Private/GLTFMeshConverter.cpp b/Engine/Plugins/Enterprise/GLTFExporter/MigratedFromGit/Source/GLTFExporter/Private/GLTFMeshConverter.cpp
new file mode 100644
--- /dev/null
+++ b/Engine/Plugins/Enterprise/GLTFExporter/MigratedFromGit/Source/GLTFExporter/Private/GLTFMeshConverter.cpp
@@ -0,0 +1,248 @@
+// Copyright Epic Games, Inc. All Rights Reserved.
+
+#include "GLTFMeshConverter.h"
+#include "GLTFConversionUtilities.h"
+
+FGLTFJsonAccessorIndex FGLTFMeshConverter::ConvertPositionAccessor(FGLTFContainerBuilder& Container, const FPositionVertexBuffer* VertexBuffer, const FString& Name)
+{
+	const uint32 VertexCount = VertexBuffer->GetNumVertices();
+	if (VertexCount == 0)
+	{
+		return FGLTFJsonAccessorIndex(INDEX_NONE);
+	}
+	
+	TArray<FVector> Positions;
+	Positions.AddUninitialized(VertexCount);
+
+	for (uint32 VertexIndex = 0; VertexIndex < VertexCount; ++VertexIndex)
+	{
+		Positions[VertexIndex] = ConvertPosition(VertexBuffer->VertexPosition(VertexIndex));
+	}
+
+	FBox BoundingBox;
+
+	// More accurate bounding box if based on raw vertex values
+	for (const FVector& Position : Positions)
+	{
+		BoundingBox += Position;
+	}
+
+	FGLTFJsonAccessor JsonAccessor;
+	JsonAccessor.Name = Name;
+	JsonAccessor.BufferView = Container.CreateBufferView(Positions, Name);
+	JsonAccessor.ComponentType = EGLTFJsonComponentType::F32;
+	JsonAccessor.Count = VertexCount;
+	JsonAccessor.Type = EGLTFJsonAccessorType::Vec3;
+	
+	JsonAccessor.MinMaxLength = 3;
+	JsonAccessor.Max[0] = BoundingBox.Max.X;
+	JsonAccessor.Max[1] = BoundingBox.Max.Y;
+	JsonAccessor.Max[2] = BoundingBox.Max.Z;
+	JsonAccessor.Min[0] = BoundingBox.Min.X;
+	JsonAccessor.Min[1] = BoundingBox.Min.Y;
+	JsonAccessor.Min[2] = BoundingBox.Min.Z;
+	
+	return Container.CreateAccessor(JsonAccessor);
+}
+
+FGLTFJsonAccessorIndex FGLTFMeshConverter::ConvertNormalAccessor(FGLTFContainerBuilder& Container, const FStaticMeshVertexBuffer* VertexBuffer, const FString& Name)
+{
+	const uint32 VertexCount = VertexBuffer->GetNumVertices();
+	if (VertexCount == 0)
+	{
+		return FGLTFJsonAccessorIndex(INDEX_NONE);
+	}
+
+	TArray<FVector> Normals;
+	Normals.AddUninitialized(VertexCount);
+	
+	for (uint32 VertexIndex = 0; VertexIndex < VertexCount; ++VertexIndex)
+	{
+		Normals[VertexIndex] = ConvertVector(VertexBuffer->VertexTangentZ(VertexIndex));
+	}
+
+	FGLTFJsonAccessor JsonAccessor;
+	JsonAccessor.Name = Name;
+	JsonAccessor.BufferView = Container.CreateBufferView(Normals, Name);
+	JsonAccessor.ComponentType = EGLTFJsonComponentType::F32;
+	JsonAccessor.Count = VertexCount;
+	JsonAccessor.Type = EGLTFJsonAccessorType::Vec3;
+	
+	return Container.CreateAccessor(JsonAccessor);
+}
+
+FGLTFJsonAccessorIndex FGLTFMeshConverter::ConvertTangentAccessor(FGLTFContainerBuilder& Container, const FStaticMeshVertexBuffer* VertexBuffer, const FString& Name)
+{
+	const uint32 VertexCount = VertexBuffer->GetNumVertices();
+	if (VertexCount == 0)
+	{
+		return FGLTFJsonAccessorIndex(INDEX_NONE);
+	}
+
+	TArray<FVector4> Tangents;
+	Tangents.AddUninitialized(VertexCount);
+	
+	for (uint32 VertexIndex = 0; VertexIndex < VertexCount; ++VertexIndex)
+	{
+		Tangents[VertexIndex] = ConvertVector(VertexBuffer->VertexTangentX(VertexIndex));
+	}
+
+	FGLTFJsonAccessor JsonAccessor;
+	JsonAccessor.Name = Name;
+	JsonAccessor.BufferView = Container.CreateBufferView(Tangents, Name);
+	JsonAccessor.ComponentType = EGLTFJsonComponentType::F32;
+	JsonAccessor.Count = VertexCount;
+	JsonAccessor.Type = EGLTFJsonAccessorType::Vec4;
+	
+	return Container.CreateAccessor(JsonAccessor);
+}
+
+FGLTFJsonAccessorIndex FGLTFMeshConverter::ConvertUV0Accessor(FGLTFContainerBuilder& Container, const FStaticMeshVertexBuffer* VertexBuffer, const FString& Name)
+{
+	const uint32 VertexCount = VertexBuffer->GetNumVertices();
+	if (VertexCount == 0 || VertexBuffer->GetNumTexCoords() < 1)
+	{
+		return FGLTFJsonAccessorIndex(INDEX_NONE);
+	}
+	
+	TArray<FVector2D> UV0s;
+	UV0s.AddUninitialized(VertexCount);
+
+	for (uint32 VertexIndex = 0; VertexIndex < VertexCount; ++VertexIndex)
+	{
+		UV0s[VertexIndex] = VertexBuffer->GetVertexUV(VertexIndex, 0);
+	}
+
+	FGLTFJsonAccessor JsonAccessor;
+	JsonAccessor.Name = Name;
+	JsonAccessor.BufferView = Container.CreateBufferView(UV0s, Name);
+	JsonAccessor.ComponentType = EGLTFJsonComponentType::F32;
+	JsonAccessor.Count = VertexCount;
+	JsonAccessor.Type = EGLTFJsonAccessorType::Vec2;
+
+	return Container.CreateAccessor(JsonAccessor);
+}
+
+FGLTFJsonAccessorIndex FGLTFMeshConverter::ConvertUV1Accessor(FGLTFContainerBuilder& Container, const FStaticMeshVertexBuffer* VertexBuffer, const FString& Name)
+{
+	const uint32 VertexCount = VertexBuffer->GetNumVertices();
+	if (VertexCount == 0 || VertexBuffer->GetNumTexCoords() < 2)
+	{
+		return FGLTFJsonAccessorIndex(INDEX_NONE);
+	}
+
+	TArray<FVector2D> UV1s;
+	UV1s.AddUninitialized(VertexCount);
+
+	for (uint32 VertexIndex = 0; VertexIndex < VertexCount; ++VertexIndex)
+	{
+		UV1s[VertexIndex] = VertexBuffer->GetVertexUV(VertexIndex, 1);
+	}
+
+	FGLTFJsonAccessor JsonAccessor;
+	JsonAccessor.Name = Name;
+	JsonAccessor.BufferView = Container.CreateBufferView(UV1s, Name);
+	JsonAccessor.ComponentType = EGLTFJsonComponentType::F32;
+	JsonAccessor.Count = VertexCount;
+	JsonAccessor.Type = EGLTFJsonAccessorType::Vec2;
+
+	return Container.CreateAccessor(JsonAccessor);
+}
+
+FGLTFJsonAccessorIndex FGLTFMeshConverter::ConvertColorAccessor(FGLTFContainerBuilder& Container, const FColorVertexBuffer* VertexBuffer, const FString& Name)
+{
+	const uint32 VertexCount = VertexBuffer->GetNumVertices();
+	if (VertexCount == 0)
+	{
+		return FGLTFJsonAccessorIndex(INDEX_NONE);
+	}
+
+	TArray<FColor> Colors;
+	Colors.AddUninitialized(VertexCount);
+
+	for (uint32 VertexIndex = 0; VertexIndex < VertexCount; ++VertexIndex)
+	{
+		Colors[VertexIndex] = ConvertColor(VertexBuffer->VertexColor(VertexIndex));
+	}
+
+	FGLTFJsonAccessor JsonAccessor;
+	JsonAccessor.Name = Name;
+	JsonAccessor.BufferView = Container.CreateBufferView(Colors, Name);
+	JsonAccessor.ComponentType = EGLTFJsonComponentType::U8;
+	JsonAccessor.Count = VertexCount;
+	JsonAccessor.Type = EGLTFJsonAccessorType::Vec4;
+	JsonAccessor.bNormalized = true;
+
+	return Container.CreateAccessor(JsonAccessor);
+}
+
+FGLTFJsonBufferViewIndex FGLTFMeshConverter::ConvertIndexBufferView(FGLTFContainerBuilder& Container, const FRawStaticIndexBuffer* IndexBuffer, const FString& Name)
+{
+	if (IndexBuffer->GetNumIndices() == 0)
+	{
+		return FGLTFJsonBufferViewIndex(INDEX_NONE);
+	}
+
+	TArray<uint32> Indices;
+	IndexBuffer->GetCopy(Indices);
+	return Container.CreateBufferView(Indices, Name, EGLTFJsonBufferTarget::ElementArrayBuffer);
+}
+
+FGLTFJsonAccessorIndex FGLTFMeshConverter::ConvertIndexAccessor(FGLTFContainerBuilder& Container, const FStaticMeshSection* MeshSection, const FRawStaticIndexBuffer* IndexBuffer, const FString& Name)
+{	
+	const uint32 TriangleCount = MeshSection->NumTriangles;
+	if (TriangleCount == 0)
+	{
+		return FGLTFJsonAccessorIndex(INDEX_NONE);
+	}
+
+	FGLTFJsonAccessor JsonAccessor;
+	JsonAccessor.Name = Name;
+	JsonAccessor.BufferView = Container.AddIndexBufferView(IndexBuffer);
+	JsonAccessor.ByteOffset = MeshSection->FirstIndex * sizeof(uint32);
+	JsonAccessor.ComponentType = EGLTFJsonComponentType::U32;
+	JsonAccessor.Count = TriangleCount * 3;
+	JsonAccessor.Type = EGLTFJsonAccessorType::Scalar;
+
+	return Container.CreateAccessor(JsonAccessor);
+}
+
+FGLTFJsonMeshIndex FGLTFMeshConverter::ConvertMesh(FGLTFContainerBuilder& Container, const UStaticMesh* StaticMesh, int32 LODIndex, const FColorVertexBuffer* OverrideVertexColors)
+{
+	FString Name;
+	StaticMesh->GetName(Name);
+	if (LODIndex != 0) Name += TEXT("_LOD") + FString::FromInt(LODIndex);
+
+	const FStaticMeshLODResources& LODMesh = StaticMesh->GetLODForExport(LODIndex);
+	const FPositionVertexBuffer* PositionBuffer = &LODMesh.VertexBuffers.PositionVertexBuffer;
+	const FStaticMeshVertexBuffer* VertexBuffer = &LODMesh.VertexBuffers.StaticMeshVertexBuffer;
+	const FColorVertexBuffer* ColorBuffer = OverrideVertexColors != nullptr ? OverrideVertexColors : &LODMesh.VertexBuffers.ColorVertexBuffer;
+	
+	FGLTFJsonMesh JsonMesh;
+	JsonMesh.Name = Name;
+
+	FGLTFJsonAttributes JsonAttributes;
+	JsonAttributes.Position = Container.AddPositionAccessor(PositionBuffer, Name + TEXT("_Positions"));
+	JsonAttributes.Normal = Container.AddNormalAccessor(VertexBuffer, Name + TEXT("_Normals"));
+	JsonAttributes.Tangent = Container.AddTangentAccessor(VertexBuffer, Name + TEXT("_Tangents"));
+	JsonAttributes.TexCoord0 = Container.AddUV0Accessor(VertexBuffer, Name + TEXT("_UV0s"));
+	JsonAttributes.TexCoord1 = Container.AddUV1Accessor(VertexBuffer, Name + TEXT("_UV1s"));
+	JsonAttributes.Color0 = Container.AddColorAccessor(ColorBuffer, Name + TEXT("_Colors"));
+
+	const FRawStaticIndexBuffer* IndexBuffer = &LODMesh.IndexBuffer;
+	Container.AddIndexBufferView(IndexBuffer, Name + TEXT("_Indices"));
+
+	const int32 SectionCount = LODMesh.Sections.Num();
+	JsonMesh.Primitives.AddDefaulted(SectionCount);
+
+	for (int32 SectionIndex = 0; SectionIndex < SectionCount; ++SectionIndex)
+	{
+		FGLTFJsonPrimitive& JsonPrimitive = JsonMesh.Primitives[SectionIndex];
+		JsonPrimitive.Attributes = JsonAttributes;
+		
+		const FString SectionName = Name + (SectionCount != 0 ? TEXT("_Indices_Section") + FString::FromInt(SectionIndex) : TEXT("_Indices"));
+		JsonPrimitive.Indices = Container.AddIndexAccessor(&LODMesh.Sections[SectionIndex], IndexBuffer, SectionName);
+	}
+	
+	return Container.CreateMesh(JsonMesh);
+}
diff --git a/Engine/Plugins/Enterprise/GLTFExporter/MigratedFromGit/Source/GLTFExporter/Private/GLTFSceneBuilder.cpp b/Engine/Plugins/Enterprise/GLTFExporter/MigratedFromGit/Source/GLTFExporter/Private/GLTFSceneBuilder.cpp
--- a/Engine/Plugins/Enterprise/GLTFExporter/MigratedFromGit/Source/GLTFExporter/Private/GLTFSceneBuilder.cpp
+++ b/Engine/Plugins/Enterprise/GLTFExporter/MigratedFromGit/Source/GLTFExporter/Private/GLTFSceneBuilder.cpp
@@ -60,7 +60,7 @@ FGLTFJsonNodeIndex FGLTFNodeBuilder::AddNode(FGLTFContainerBuilder& Container) c
 		}
 	}
 
-	return Container.AddNode(Node);
+	return Container.CreateNode(Node);
 }
 
 FGLTFSceneBuilder::FGLTFSceneBuilder(const UWorld* World, bool bSelectedOnly)
@@ -94,14 +94,11 @@ FGLTFJsonSceneIndex FGLTFSceneBuilder::AddScene(FGLTFContainerBuilder& Container
 	FGLTFJsonScene Scene;
 	Scene.Name = Name;
 
-	for (const FGLTFNodeBuilder& TopLevelComponent : RootNodes)
+	for (const FGLTFNodeBuilder& RootNode : RootNodes)
 	{
-		FGLTFJsonNodeIndex NodeIndex = TopLevelComponent.AddNode(Container);
-		if (NodeIndex != INDEX_NONE)
-		{
-			Scene.Nodes.Add(NodeIndex);
-		}
+		FGLTFJsonNodeIndex NodeIndex = RootNode.AddNode(Container);
+		if (NodeIndex != INDEX_NONE) Scene.Nodes.Add(NodeIndex);
 	}
 
-	return Container.AddScene(Scene);
+	return Container.CreateScene(Scene);
 }
diff --git a/Engine/Plugins/Enterprise/GLTFExporter/MigratedFromGit/Source/GLTFExporter/Private/GLTFStaticMeshExporter.cpp b/Engine/Plugins/Enterprise/GLTFExporter/MigratedFromGit/Source/GLTFExporter/Private/GLTFStaticMeshExporter.cpp
--- a/Engine/Plugins/Enterprise/GLTFExporter/MigratedFromGit/Source/GLTFExporter/Private/GLTFStaticMeshExporter.cpp
+++ b/Engine/Plugins/Enterprise/GLTFExporter/MigratedFromGit/Source/GLTFExporter/Private/GLTFStaticMeshExporter.cpp
@@ -27,11 +27,11 @@ bool UGLTFStaticMeshExporter::ExportBinary(UObject* Object, const TCHAR* Type, F
 
 	FGLTFJsonNode Node;
 	Node.Mesh = MeshIndex;
-	FGLTFJsonNodeIndex NodeIndex = Container.AddNode(Node);
+	FGLTFJsonNodeIndex NodeIndex = Container.CreateNode(Node);
 
 	FGLTFJsonScene Scene;
 	Scene.Nodes.Add(NodeIndex);
-	FGLTFJsonSceneIndex SceneIndex = Container.AddScene(Scene);
+	FGLTFJsonSceneIndex SceneIndex = Container.CreateScene(Scene);
 
 	Container.JsonRoot.DefaultScene = SceneIndex;
 
diff --git a/Engine/Plugins/Enterprise/GLTFExporter/MigratedFromGit/Source/GLTFExporter/Public/GLTFBufferBuilder.h b/Engine/Plugins/Enterprise/GLTFExporter/MigratedFromGit/Source/GLTFExporter/Public/GLTFBufferBuilder.h
--- a/Engine/Plugins/Enterprise/GLTFExporter/MigratedFromGit/Source/GLTFExporter/Public/GLTFBufferBuilder.h
+++ b/Engine/Plugins/Enterprise/GLTFExporter/MigratedFromGit/Source/GLTFExporter/Public/GLTFBufferBuilder.h
@@ -14,7 +14,7 @@ struct GLTFEXPORTER_API FGLTFBufferBuilder
 
 	FGLTFBufferBuilder(FGLTFJsonBufferIndex BufferIndex);
 
-	FGLTFJsonBufferViewIndex AddBufferView(FGLTFContainerBuilder& Container, const void* RawData, uint64 ByteLength, const FString& Name, EGLTFJsonBufferTarget BufferTarget);
+	FGLTFJsonBufferViewIndex CreateBufferView(FGLTFContainerBuilder& Container, const void* RawData, uint64 ByteLength, const FString& Name, EGLTFJsonBufferTarget BufferTarget);
 
 	void UpdateBuffer(FGLTFJsonBuffer& JsonBuffer);
 };
diff --git a/Engine/Plugins/Enterprise/GLTFExporter/MigratedFromGit/Source/GLTFExporter/Public/GLTFContainerBuilder.h b/Engine/Plugins/Enterprise/GLTFExporter/MigratedFromGit/Source/GLTFExporter/Public/GLTFContainerBuilder.h
--- a/Engine/Plugins/Enterprise/GLTFExporter/MigratedFromGit/Source/GLTFExporter/Public/GLTFContainerBuilder.h
+++ b/Engine/Plugins/Enterprise/GLTFExporter/MigratedFromGit/Source/GLTFExporter/Public/GLTFContainerBuilder.h
@@ -15,21 +15,31 @@ struct GLTFEXPORTER_API FGLTFContainerBuilder
 
 	FGLTFContainerBuilder();
 
-	FGLTFJsonAccessorIndex AddAccessor(const FGLTFJsonAccessor& JsonAccessor);
-	FGLTFJsonBufferIndex AddBuffer(const FGLTFJsonBuffer& JsonBuffer);
-	FGLTFJsonBufferViewIndex AddBufferView(const FGLTFJsonBufferView& JsonBufferView);
-	FGLTFJsonMeshIndex AddMesh(const FGLTFJsonMesh& JsonMesh);
-	FGLTFJsonNodeIndex AddNode(const FGLTFJsonNode& JsonNode);
-	FGLTFJsonSceneIndex AddScene(const FGLTFJsonScene& JsonScene);
+	FGLTFJsonAccessorIndex CreateAccessor(const FGLTFJsonAccessor& JsonAccessor);
+	FGLTFJsonBufferIndex CreateBuffer(const FGLTFJsonBuffer& JsonBuffer);
+	FGLTFJsonBufferViewIndex CreateBufferView(const FGLTFJsonBufferView& JsonBufferView);
+	FGLTFJsonMeshIndex CreateMesh(const FGLTFJsonMesh& JsonMesh);
+	FGLTFJsonNodeIndex CreateNode(const FGLTFJsonNode& JsonNode);
+	FGLTFJsonSceneIndex CreateScene(const FGLTFJsonScene& JsonScene);
 
-	FGLTFJsonBufferViewIndex AddBufferView(const void* RawData, uint64 ByteLength, const FString& Name = TEXT(""), EGLTFJsonBufferTarget BufferTarget = EGLTFJsonBufferTarget::ArrayBuffer);
+	FGLTFJsonBufferViewIndex CreateBufferView(const void* RawData, uint64 ByteLength, const FString& Name = TEXT(""), EGLTFJsonBufferTarget BufferTarget = EGLTFJsonBufferTarget::ArrayBuffer);
 
 	template <class ElementType>
-	FGLTFJsonBufferViewIndex AddBufferView(const TArray<ElementType>& Array, const FString& Name = TEXT(""), EGLTFJsonBufferTarget BufferTarget = EGLTFJsonBufferTarget::ArrayBuffer)
+	FGLTFJsonBufferViewIndex CreateBufferView(const TArray<ElementType>& Array, const FString& Name = TEXT(""), EGLTFJsonBufferTarget BufferTarget = EGLTFJsonBufferTarget::ArrayBuffer)
 	{
-		return AddBufferView(Array.GetData(), Array.Num() * sizeof(ElementType), Name, BufferTarget);
+		return CreateBufferView(Array.GetData(), Array.Num() * sizeof(ElementType), Name, BufferTarget);
 	}
 
+	FGLTFJsonAccessorIndex AddPositionAccessor(const FPositionVertexBuffer* VertexBuffer, const FString& Name = TEXT(""));
+	FGLTFJsonAccessorIndex AddNormalAccessor(const FStaticMeshVertexBuffer* VertexBuffer, const FString& Name = TEXT(""));
+	FGLTFJsonAccessorIndex AddTangentAccessor(const FStaticMeshVertexBuffer* VertexBuffer, const FString& Name = TEXT(""));
+	FGLTFJsonAccessorIndex AddUV0Accessor(const FStaticMeshVertexBuffer* VertexBuffer, const FString& Name = TEXT(""));
+	FGLTFJsonAccessorIndex AddUV1Accessor(const FStaticMeshVertexBuffer* VertexBuffer, const FString& Name = TEXT(""));
+	FGLTFJsonAccessorIndex AddColorAccessor(const FColorVertexBuffer* VertexBuffer, const FString& Name = TEXT(""));
+
+	FGLTFJsonBufferViewIndex AddIndexBufferView(const FRawStaticIndexBuffer* IndexBuffer, const FString& Name = TEXT(""));
+	FGLTFJsonAccessorIndex AddIndexAccessor(const FStaticMeshSection* MeshSection, const FRawStaticIndexBuffer* IndexBuffer, const FString& Name = TEXT(""));
+
 	FGLTFJsonMeshIndex AddMesh(const UStaticMesh* StaticMesh, int32 LODIndex = 0, const FColorVertexBuffer* OverrideVertexColors = nullptr);
 	FGLTFJsonMeshIndex AddMesh(const UStaticMeshComponent* StaticMeshComponent);
 
diff --git a/Engine/Plugins/Enterprise/GLTFExporter/MigratedFromGit/Source/GLTFExporter/Public/GLTFIndexBuilder.h b/Engine/Plugins/Enterprise/GLTFExporter/MigratedFromGit/Source/GLTFExporter/Public/GLTFIndexBuilder.h
--- a/Engine/Plugins/Enterprise/GLTFExporter/MigratedFromGit/Source/GLTFExporter/Public/GLTFIndexBuilder.h
+++ b/Engine/Plugins/Enterprise/GLTFExporter/MigratedFromGit/Source/GLTFExporter/Public/GLTFIndexBuilder.h
@@ -23,7 +23,7 @@ struct FGLTFStaticMeshKey
 	FORCEINLINE bool operator==(const FGLTFStaticMeshKey& Other) const
 	{
 		return StaticMesh           == Other.StaticMesh
-		    && LODIndex             == Other.LODIndex
+			&& LODIndex             == Other.LODIndex
 			&& OverrideVertexColors == Other.OverrideVertexColors;
 	}
 
@@ -33,10 +33,148 @@ struct FGLTFStaticMeshKey
 	}
 };
 
+struct FGLTFPositionVertexBufferKey
+{
+	const FPositionVertexBuffer* VertexBuffer;
+	const FString Name;
+
+	FGLTFPositionVertexBufferKey(const FPositionVertexBuffer* VertexBuffer, const FString& Name)
+		: VertexBuffer(VertexBuffer)
+		, Name(Name)
+	{
+	}
+
+	FORCEINLINE bool operator==(const FGLTFPositionVertexBufferKey& Other) const
+	{
+		return VertexBuffer == Other.VertexBuffer;
+	}
+
+	FORCEINLINE friend uint32 GetTypeHash(const FGLTFPositionVertexBufferKey& Other)
+	{
+		return GetTypeHash(Other.VertexBuffer);
+	}
+};
+
+struct FGLTFStaticMeshVertexBufferKey
+{
+	const FStaticMeshVertexBuffer* VertexBuffer;
+	const FString Name;
+
+	FGLTFStaticMeshVertexBufferKey(const FStaticMeshVertexBuffer* VertexBuffer, const FString& Name)
+		: VertexBuffer(VertexBuffer)
+		, Name(Name)
+	{
+	}
+
+	FORCEINLINE bool operator==(const FGLTFStaticMeshVertexBufferKey& Other) const
+	{
+		return VertexBuffer == Other.VertexBuffer;
+	}
+
+	FORCEINLINE friend uint32 GetTypeHash(const FGLTFStaticMeshVertexBufferKey& Other)
+	{
+		return GetTypeHash(Other.VertexBuffer);
+	}
+};
+
+struct FGLTFColorVertexBufferKey
+{
+	const FColorVertexBuffer* VertexBuffer;
+	const FString Name;
+
+	FGLTFColorVertexBufferKey(const FColorVertexBuffer* VertexBuffer, const FString& Name)
+		: VertexBuffer(VertexBuffer)
+		, Name(Name)
+	{
+	}
+
+	FORCEINLINE bool operator==(const FGLTFColorVertexBufferKey& Other) const
+	{
+		return VertexBuffer == Other.VertexBuffer;
+	}
+
+	FORCEINLINE friend uint32 GetTypeHash(const FGLTFColorVertexBufferKey& Other)
+	{
+		return GetTypeHash(Other.VertexBuffer);
+	}
+};
+
+struct FGLTFRawStaticIndexBufferKey
+{
+	const FRawStaticIndexBuffer* IndexBuffer;
+	const FString Name;
+
+	FGLTFRawStaticIndexBufferKey(const FRawStaticIndexBuffer* IndexBuffer, const FString& Name)
+		: IndexBuffer(IndexBuffer)
+		, Name(Name)
+	{
+	}
+
+	FORCEINLINE bool operator==(const FGLTFRawStaticIndexBufferKey& Other) const
+	{
+		return IndexBuffer == Other.IndexBuffer;
+	}
+
+	FORCEINLINE friend uint32 GetTypeHash(const FGLTFRawStaticIndexBufferKey& Other)
+	{
+		return GetTypeHash(Other.IndexBuffer);
+	}
+};
+
+struct FGLTFStaticMeshSectionKey
+{
+	const FStaticMeshSection* MeshSection;
+	const FRawStaticIndexBuffer* IndexBuffer;
+	const FString Name;
+
+	FGLTFStaticMeshSectionKey(const FStaticMeshSection* MeshSection, const FRawStaticIndexBuffer* IndexBuffer, const FString& Name)
+		: MeshSection(MeshSection)
+		, IndexBuffer(IndexBuffer)
+		, Name(Name)
+	{
+	}
+
+	FORCEINLINE bool operator==(const FGLTFStaticMeshSectionKey& Other) const
+	{
+		return MeshSection == Other.MeshSection
+			&& IndexBuffer == Other.IndexBuffer;
+	}
+
+	FORCEINLINE friend uint32 GetTypeHash(const FGLTFStaticMeshSectionKey& Other)
+	{
+		return HashCombine(GetTypeHash(Other.MeshSection), GetTypeHash(Other.IndexBuffer));
+	}
+};
+
 struct GLTFEXPORTER_API FGLTFIndexBuilder
 {
 	TMap<FGLTFStaticMeshKey, FGLTFJsonMeshIndex> StaticMeshes;
+	TMap<FGLTFPositionVertexBufferKey, FGLTFJsonAccessorIndex> PositionVertexBuffers;
+	TMap<FGLTFStaticMeshVertexBufferKey, FGLTFJsonAccessorIndex> StaticNormalVertexBuffers;
+	TMap<FGLTFStaticMeshVertexBufferKey, FGLTFJsonAccessorIndex> StaticTangentVertexBuffers;
+	TMap<FGLTFStaticMeshVertexBufferKey, FGLTFJsonAccessorIndex> StaticUV0VertexBuffers;
+	TMap<FGLTFStaticMeshVertexBufferKey, FGLTFJsonAccessorIndex> StaticUV1VertexBuffers;
+	TMap<FGLTFColorVertexBufferKey, FGLTFJsonAccessorIndex> ColorVertexBuffers;
+	TMap<FGLTFRawStaticIndexBufferKey, FGLTFJsonBufferViewIndex> StaticIndexBuffers;
+	TMap<FGLTFStaticMeshSectionKey, FGLTFJsonAccessorIndex> StaticMeshSections;
+
+	FGLTFJsonMeshIndex FindMesh(const FGLTFStaticMeshKey& Key) const;
+	FGLTFJsonAccessorIndex FindPositionAccessor(const FGLTFPositionVertexBufferKey& Key) const;
+	FGLTFJsonAccessorIndex FindNormalAccessor(const FGLTFStaticMeshVertexBufferKey& Key) const;
+	FGLTFJsonAccessorIndex FindTangentAccessor(const FGLTFStaticMeshVertexBufferKey& Key) const;
+	FGLTFJsonAccessorIndex FindUV0Accessor(const FGLTFStaticMeshVertexBufferKey& Key) const;
+	FGLTFJsonAccessorIndex FindUV1Accessor(const FGLTFStaticMeshVertexBufferKey& Key) const;
+	FGLTFJsonAccessorIndex FindColorAccessor(const FGLTFColorVertexBufferKey& Key) const;
+	FGLTFJsonBufferViewIndex FindIndexBufferView(const FGLTFRawStaticIndexBufferKey& Key) const;
+	FGLTFJsonAccessorIndex FindIndexAccessor(const FGLTFStaticMeshSectionKey& Key) const;
 
-	FGLTFJsonMeshIndex Find(const FGLTFStaticMeshKey& Key) const;
-	FGLTFJsonMeshIndex FindOrAdd(const FGLTFStaticMeshKey& Key, FGLTFContainerBuilder& Container);
+	FGLTFJsonMeshIndex FindOrConvertMesh(const FGLTFStaticMeshKey& Key, FGLTFContainerBuilder& Container);
+	FGLTFJsonAccessorIndex FindOrConvertPositionAccessor(const FGLTFPositionVertexBufferKey& Key, FGLTFContainerBuilder& Container);
+	FGLTFJsonAccessorIndex FindOrConvertNormalAccessor(const FGLTFStaticMeshVertexBufferKey& Key, FGLTFContainerBuilder& Container);
+	FGLTFJsonAccessorIndex FindOrConvertTangentAccessor(const FGLTFStaticMeshVertexBufferKey& Key, FGLTFContainerBuilder& Container);
+	FGLTFJsonAccessorIndex FindOrConvertUV0Accessor(const FGLTFStaticMeshVertexBufferKey& Key, FGLTFContainerBuilder& Container);
+	FGLTFJsonAccessorIndex FindOrConvertUV1Accessor(const FGLTFStaticMeshVertexBufferKey& Key, FGLTFContainerBuilder& Container);
+	FGLTFJsonAccessorIndex FindOrConvertColorAccessor(const FGLTFColorVertexBufferKey& Key, FGLTFContainerBuilder& Container);
+	FGLTFJsonBufferViewIndex FindOrConvertIndexBufferView(const FGLTFRawStaticIndexBufferKey& Key, FGLTFContainerBuilder& Container);
+	FGLTFJsonAccessorIndex FindOrConvertIndexAccessor(const FGLTFStaticMeshSectionKey& Key, FGLTFContainerBuilder& Container);
 };
diff --git a/Engine/Plugins/Enterprise/GLTFExporter/MigratedFromGit/Source/GLTFExporter/Public/GLTFJsonAccessor.h b/Engine/Plugins/Enterprise/GLTFExporter/MigratedFromGit/Source/GLTFExporter/Public/GLTFJsonAccessor.h
--- a/Engine/Plugins/Enterprise/GLTFExporter/MigratedFromGit/Source/GLTFExporter/Public/GLTFJsonAccessor.h
+++ b/Engine/Plugins/Enterprise/GLTFExporter/MigratedFromGit/Source/GLTFExporter/Public/GLTFJsonAccessor.h
@@ -12,6 +12,7 @@ struct GLTFEXPORTER_API FGLTFJsonAccessor
 	FString Name;
 
 	FGLTFJsonBufferViewIndex BufferView;
+	int64                    ByteOffset;
 	int32                    Count;
 	EGLTFJsonAccessorType    Type;
 	EGLTFJsonComponentType   ComponentType;
@@ -23,6 +24,7 @@ struct GLTFEXPORTER_API FGLTFJsonAccessor
 
 	FGLTFJsonAccessor()
 		: BufferView(INDEX_NONE)
+		, ByteOffset(0)
 		, Count(0)
 		, Type(EGLTFJsonAccessorType::None)
 		, ComponentType(EGLTFJsonComponentType::None)
@@ -44,6 +46,7 @@ struct GLTFEXPORTER_API FGLTFJsonAccessor
 		}
 
 		JsonWriter.WriteValue(TEXT("bufferView"), BufferView);
+		if (ByteOffset != 0) JsonWriter.WriteValue(TEXT("byteOffset"), ByteOffset);
 		JsonWriter.WriteValue(TEXT("count"), Count);
 		JsonWriter.WriteValue(TEXT("type"), FGLTFJsonUtility::AccessorTypeToString(Type));
 		JsonWriter.WriteValue(TEXT("componentType"), FGLTFJsonUtility::ComponentTypeToNumber(ComponentType));
diff --git a/Engine/Plugins/Enterprise/GLTFExporter/MigratedFromGit/Source/GLTFExporter/Public/GLTFMeshBuilder.h b/Engine/Plugins/Enterprise/GLTFExporter/MigratedFromGit/Source/GLTFExporter/Public/GLTFMeshBuilder.h
deleted file mode 100644
diff --git a/Engine/Plugins/Enterprise/GLTFExporter/MigratedFromGit/Source/GLTFExporter/Public/GLTFMeshConverter.h b/Engine/Plugins/Enterprise/GLTFExporter/MigratedFromGit/Source/GLTFExporter/Public/GLTFMeshConverter.h
new file mode 100644
--- /dev/null
+++ b/Engine/Plugins/Enterprise/GLTFExporter/MigratedFromGit/Source/GLTFExporter/Public/GLTFMeshConverter.h
@@ -0,0 +1,22 @@
+// Copyright Epic Games, Inc. All Rights Reserved.
+
+#pragma once
+
+#include "GLTFJsonRoot.h"
+#include "GLTFContainerBuilder.h"
+#include "Engine.h"
+
+struct GLTFEXPORTER_API FGLTFMeshConverter
+{
+	static FGLTFJsonAccessorIndex ConvertPositionAccessor(FGLTFContainerBuilder& Container, const FPositionVertexBuffer* VertexBuffer, const FString& Name);
+	static FGLTFJsonAccessorIndex ConvertNormalAccessor(FGLTFContainerBuilder& Container, const FStaticMeshVertexBuffer* VertexBuffer, const FString& Name);
+	static FGLTFJsonAccessorIndex ConvertTangentAccessor(FGLTFContainerBuilder& Container, const FStaticMeshVertexBuffer* VertexBuffer, const FString& Name);
+	static FGLTFJsonAccessorIndex ConvertUV0Accessor(FGLTFContainerBuilder& Container, const FStaticMeshVertexBuffer* VertexBuffer, const FString& Name);
+	static FGLTFJsonAccessorIndex ConvertUV1Accessor(FGLTFContainerBuilder& Container, const FStaticMeshVertexBuffer* VertexBuffer, const FString& Name);
+	static FGLTFJsonAccessorIndex ConvertColorAccessor(FGLTFContainerBuilder& Container, const FColorVertexBuffer* VertexBuffer, const FString& Name);
+
+	static FGLTFJsonBufferViewIndex ConvertIndexBufferView(FGLTFContainerBuilder& Container, const FRawStaticIndexBuffer* IndexBuffer, const FString& Name);
+	static FGLTFJsonAccessorIndex ConvertIndexAccessor(FGLTFContainerBuilder& Container, const FStaticMeshSection* MeshSection, const FRawStaticIndexBuffer* IndexBuffer, const FString& Name);
+
+	static FGLTFJsonMeshIndex ConvertMesh(FGLTFContainerBuilder& Container, const UStaticMesh* StaticMesh, int32 LODIndex, const FColorVertexBuffer* OverrideVertexColors);
+};
