[GLTFExporter] refactor: add template TGLTFIndexedConvert to reduce code duplication and rename GLTFIndexBuilder -> GLTFIndexedConverts

commit 8a0942b247be843e2b6505849761fbcc8232c0bb
Date: Sun, 10 May 2020 14:42:50 +0200

    refactor: add template TGLTFIndexedConvert to reduce code duplication and rename GLTFIndexBuilder -> GLTFIndexedConverts
    (cherry picked from commit bc782d485d223ed953068dec0d16fc8f97f43c89​)

#skipci

diff --git a/Engine/Plugins/Enterprise/GLTFExporter/MigratedFromGit/Source/GLTFExporter/Private/GLTFBufferBuilder.cpp b/Engine/Plugins/Enterprise/GLTFExporter/MigratedFromGit/Source/GLTFExporter/Private/GLTFBufferBuilder.cpp
--- a/Engine/Plugins/Enterprise/GLTFExporter/MigratedFromGit/Source/GLTFExporter/Private/GLTFBufferBuilder.cpp
+++ b/Engine/Plugins/Enterprise/GLTFExporter/MigratedFromGit/Source/GLTFExporter/Private/GLTFBufferBuilder.cpp
@@ -9,7 +9,7 @@ FGLTFBufferBuilder::FGLTFBufferBuilder(FGLTFJsonBufferIndex BufferIndex)
 {
 }
 
-FGLTFJsonBufferViewIndex FGLTFBufferBuilder::CreateBufferView(FGLTFContainerBuilder& Container, const void* RawData, uint64 ByteLength, const FString& Name, EGLTFJsonBufferTarget BufferTarget)
+FGLTFJsonBufferViewIndex FGLTFBufferBuilder::AddBufferView(FGLTFContainerBuilder& Container, const void* RawData, uint64 ByteLength, const FString& Name, EGLTFJsonBufferTarget BufferTarget)
 {
 	FGLTFJsonBufferView BufferView;
 	BufferView.Name = Name;
@@ -20,7 +20,7 @@ FGLTFJsonBufferViewIndex FGLTFBufferBuilder::CreateBufferView(FGLTFContainerBuil
 
 	BufferData.Append(static_cast<const uint8*>(RawData), ByteLength);
 
-	return Container.CreateBufferView(BufferView);
+	return Container.AddBufferView(BufferView);
 }
 
 void FGLTFBufferBuilder::UpdateBuffer(FGLTFJsonBuffer& JsonBuffer)
diff --git a/Engine/Plugins/Enterprise/GLTFExporter/MigratedFromGit/Source/GLTFExporter/Private/GLTFContainerBuilder.cpp b/Engine/Plugins/Enterprise/GLTFExporter/MigratedFromGit/Source/GLTFExporter/Private/GLTFContainerBuilder.cpp
--- a/Engine/Plugins/Enterprise/GLTFExporter/MigratedFromGit/Source/GLTFExporter/Private/GLTFContainerBuilder.cpp
+++ b/Engine/Plugins/Enterprise/GLTFExporter/MigratedFromGit/Source/GLTFExporter/Private/GLTFContainerBuilder.cpp
@@ -4,106 +4,106 @@
 #include "GLTFSceneBuilder.h"
 
 FGLTFContainerBuilder::FGLTFContainerBuilder()
-	: BufferBuilder(CreateBuffer(FGLTFJsonBuffer()))
+	: BufferBuilder(AddBuffer(FGLTFJsonBuffer()))
 {
 }
 
-FGLTFJsonAccessorIndex FGLTFContainerBuilder::CreateAccessor(const FGLTFJsonAccessor& JsonAccessor)
+FGLTFJsonAccessorIndex FGLTFContainerBuilder::AddAccessor(const FGLTFJsonAccessor& JsonAccessor)
 {
 	return FGLTFJsonAccessorIndex(JsonRoot.Accessors.Add(JsonAccessor));
 }
 
-FGLTFJsonBufferIndex FGLTFContainerBuilder::CreateBuffer(const FGLTFJsonBuffer& JsonBuffer)
+FGLTFJsonBufferIndex FGLTFContainerBuilder::AddBuffer(const FGLTFJsonBuffer& JsonBuffer)
 {
 	return FGLTFJsonBufferIndex(JsonRoot.Buffers.Add(JsonBuffer));
 }
 
-FGLTFJsonBufferViewIndex FGLTFContainerBuilder::CreateBufferView(const FGLTFJsonBufferView& JsonBufferView)
+FGLTFJsonBufferViewIndex FGLTFContainerBuilder::AddBufferView(const FGLTFJsonBufferView& JsonBufferView)
 {
 	return FGLTFJsonBufferViewIndex(JsonRoot.BufferViews.Add(JsonBufferView));
 }
 
-FGLTFJsonMeshIndex FGLTFContainerBuilder::CreateMesh(const FGLTFJsonMesh& JsonMesh)
+FGLTFJsonMeshIndex FGLTFContainerBuilder::AddMesh(const FGLTFJsonMesh& JsonMesh)
 {
 	return FGLTFJsonMeshIndex(JsonRoot.Meshes.Add(JsonMesh));
 }
 
-FGLTFJsonNodeIndex FGLTFContainerBuilder::CreateNode(const FGLTFJsonNode& JsonNode)
+FGLTFJsonNodeIndex FGLTFContainerBuilder::AddNode(const FGLTFJsonNode& JsonNode)
 {
 	return FGLTFJsonNodeIndex(JsonRoot.Nodes.Add(JsonNode));
 }
 
-FGLTFJsonSceneIndex FGLTFContainerBuilder::CreateScene(const FGLTFJsonScene& JsonScene)
+FGLTFJsonSceneIndex FGLTFContainerBuilder::AddScene(const FGLTFJsonScene& JsonScene)
 {
 	return FGLTFJsonSceneIndex(JsonRoot.Scenes.Add(JsonScene));
 }
 
-FGLTFJsonBufferViewIndex FGLTFContainerBuilder::CreateBufferView(const void* RawData, uint64 ByteLength, const FString& Name, EGLTFJsonBufferTarget BufferTarget)
+FGLTFJsonBufferViewIndex FGLTFContainerBuilder::AddBufferView(const void* RawData, uint64 ByteLength, const FString& Name, EGLTFJsonBufferTarget BufferTarget)
 {
-	return BufferBuilder.CreateBufferView(*this, RawData, ByteLength, Name, BufferTarget);
+	return BufferBuilder.AddBufferView(*this, RawData, ByteLength, Name, BufferTarget);
 }
 
-FGLTFJsonAccessorIndex FGLTFContainerBuilder::AddPositionAccessor(const FPositionVertexBuffer* VertexBuffer, const FString& Name)
+FGLTFJsonAccessorIndex FGLTFContainerBuilder::ConvertPositionAccessor(const FPositionVertexBuffer* VertexBuffer, const FString& DesiredName)
 {
-	return IndexBuilder.FindOrConvertPositionAccessor(FGLTFPositionVertexBufferKey(VertexBuffer, Name), *this);
+	return IndexedConverts.PositionVertexBuffers.Convert(*this, DesiredName, VertexBuffer);
 }
 
-FGLTFJsonAccessorIndex FGLTFContainerBuilder::AddNormalAccessor(const FStaticMeshVertexBuffer* VertexBuffer, const FString& Name)
+FGLTFJsonAccessorIndex FGLTFContainerBuilder::ConvertColorAccessor(const FColorVertexBuffer* VertexBuffer, const FString& DesiredName)
 {
-	return IndexBuilder.FindOrConvertNormalAccessor(FGLTFStaticMeshVertexBufferKey(VertexBuffer, Name), *this);
+	return IndexedConverts.ColorVertexBuffers.Convert(*this, DesiredName, VertexBuffer);
 }
 
-FGLTFJsonAccessorIndex FGLTFContainerBuilder::AddTangentAccessor(const FStaticMeshVertexBuffer* VertexBuffer, const FString& Name)
+FGLTFJsonAccessorIndex FGLTFContainerBuilder::ConvertNormalAccessor(const FStaticMeshVertexBuffer* VertexBuffer, const FString& DesiredName)
 {
-	return IndexBuilder.FindOrConvertTangentAccessor(FGLTFStaticMeshVertexBufferKey(VertexBuffer, Name), *this);
+	return IndexedConverts.StaticMeshNormalVertexBuffers.Convert(*this, DesiredName, VertexBuffer);
 }
 
-FGLTFJsonAccessorIndex FGLTFContainerBuilder::AddUV0Accessor(const FStaticMeshVertexBuffer* VertexBuffer, const FString& Name)
+FGLTFJsonAccessorIndex FGLTFContainerBuilder::ConvertTangentAccessor(const FStaticMeshVertexBuffer* VertexBuffer, const FString& DesiredName)
 {
-	return IndexBuilder.FindOrConvertUV0Accessor(FGLTFStaticMeshVertexBufferKey(VertexBuffer, Name), *this);
+	return IndexedConverts.StaticMeshTangentVertexBuffers.Convert(*this, DesiredName, VertexBuffer);
 }
 
-FGLTFJsonAccessorIndex FGLTFContainerBuilder::AddUV1Accessor(const FStaticMeshVertexBuffer* VertexBuffer, const FString& Name)
+FGLTFJsonAccessorIndex FGLTFContainerBuilder::ConvertUV0Accessor(const FStaticMeshVertexBuffer* VertexBuffer, const FString& DesiredName)
 {
-	return IndexBuilder.FindOrConvertUV1Accessor(FGLTFStaticMeshVertexBufferKey(VertexBuffer, Name), *this);
+	return IndexedConverts.StaticMeshUV0VertexBuffers.Convert(*this, DesiredName, VertexBuffer);
 }
 
-FGLTFJsonAccessorIndex FGLTFContainerBuilder::AddColorAccessor(const FColorVertexBuffer* VertexBuffer, const FString& Name)
+FGLTFJsonAccessorIndex FGLTFContainerBuilder::ConvertUV1Accessor(const FStaticMeshVertexBuffer* VertexBuffer, const FString& DesiredName)
 {
-	return IndexBuilder.FindOrConvertColorAccessor(FGLTFColorVertexBufferKey(VertexBuffer, Name), *this);
+	return IndexedConverts.StaticMeshUV1VertexBuffers.Convert(*this, DesiredName, VertexBuffer);
 }
 
-FGLTFJsonBufferViewIndex FGLTFContainerBuilder::AddIndexBufferView(const FRawStaticIndexBuffer* IndexBuffer, const FString& Name)
+FGLTFJsonBufferViewIndex FGLTFContainerBuilder::ConvertIndexBufferView(const FRawStaticIndexBuffer* IndexBuffer, const FString& DesiredName)
 {
-	return IndexBuilder.FindOrConvertIndexBufferView(FGLTFRawStaticIndexBufferKey(IndexBuffer, Name), *this);
+	return IndexedConverts.StaticMeshIndexBuffers.Convert(*this, DesiredName, IndexBuffer);
 }
 
-FGLTFJsonAccessorIndex FGLTFContainerBuilder::AddIndexAccessor(const FStaticMeshSection* MeshSection, const FRawStaticIndexBuffer* IndexBuffer, const FString& Name)
+FGLTFJsonAccessorIndex FGLTFContainerBuilder::ConvertIndexAccessor(const FStaticMeshSection* MeshSection, const FRawStaticIndexBuffer* IndexBuffer, const FString& DesiredName)
 {
-	return IndexBuilder.FindOrConvertIndexAccessor(FGLTFStaticMeshSectionKey(MeshSection, IndexBuffer, Name), *this);
+	return IndexedConverts.StaticMeshSections.Convert(*this, DesiredName, MeshSection, IndexBuffer);
 }
 
-void FGLTFContainerBuilder::Serialize(FArchive& Archive)
-{
-	FGLTFJsonBuffer& JsonBuffer = JsonRoot.Buffers[BufferBuilder.BufferIndex];
-	BufferBuilder.UpdateBuffer(JsonBuffer);
-	JsonRoot.Serialize(&Archive, true);
-}
-
-FGLTFJsonMeshIndex FGLTFContainerBuilder::AddMesh(const UStaticMesh* StaticMesh, int32 LODIndex, const FColorVertexBuffer* OverrideVertexColors)
+FGLTFJsonMeshIndex FGLTFContainerBuilder::ConvertMesh(const UStaticMesh* StaticMesh, int32 LODIndex, const FColorVertexBuffer* OverrideVertexColors, const FString& DesiredName)
 {
-	return IndexBuilder.FindOrConvertMesh(FGLTFStaticMeshKey(StaticMesh, LODIndex, OverrideVertexColors), *this);
+	return IndexedConverts.StaticMeshes.Convert(*this, DesiredName, StaticMesh, LODIndex, OverrideVertexColors);
 }
 
-FGLTFJsonMeshIndex FGLTFContainerBuilder::AddMesh(const UStaticMeshComponent* StaticMeshComponent)
+FGLTFJsonMeshIndex FGLTFContainerBuilder::ConvertMesh(const UStaticMeshComponent* StaticMeshComponent, const FString& DesiredName)
 {
 	const UStaticMesh* StaticMesh = StaticMeshComponent->GetStaticMesh();
 	const int32 LODIndex = StaticMeshComponent->ForcedLodModel > 0 ? StaticMeshComponent->ForcedLodModel - 1 : /* auto-select */ 0;
 	const FColorVertexBuffer* OverrideVertexColors = LODIndex < StaticMeshComponent->LODData.Num() ? StaticMeshComponent->LODData[LODIndex].OverrideVertexColors : nullptr;
-	return AddMesh(StaticMesh, LODIndex, OverrideVertexColors);
+	return ConvertMesh(StaticMesh, LODIndex, OverrideVertexColors, DesiredName);
 }
 
 FGLTFJsonSceneIndex FGLTFContainerBuilder::AddScene(const UWorld* World, bool bSelectedOnly)
 {
 	return FGLTFSceneBuilder(World, bSelectedOnly).AddScene(*this);
 }
+
+void FGLTFContainerBuilder::Serialize(FArchive& Archive)
+{
+	FGLTFJsonBuffer& JsonBuffer = JsonRoot.Buffers[BufferBuilder.BufferIndex];
+	BufferBuilder.UpdateBuffer(JsonBuffer);
+	JsonRoot.Serialize(&Archive, true);
+}
diff --git a/Engine/Plugins/Enterprise/GLTFExporter/MigratedFromGit/Source/GLTFExporter/Private/GLTFIndexBuilder.cpp b/Engine/Plugins/Enterprise/GLTFExporter/MigratedFromGit/Source/GLTFExporter/Private/GLTFIndexBuilder.cpp
deleted file mode 100644
diff --git a/Engine/Plugins/Enterprise/GLTFExporter/MigratedFromGit/Source/GLTFExporter/Private/GLTFMeshConverter.cpp b/Engine/Plugins/Enterprise/GLTFExporter/MigratedFromGit/Source/GLTFExporter/Private/GLTFMeshConverter.cpp
deleted file mode 100644
diff --git a/Engine/Plugins/Enterprise/GLTFExporter/MigratedFromGit/Source/GLTFExporter/Private/GLTFSceneBuilder.cpp b/Engine/Plugins/Enterprise/GLTFExporter/MigratedFromGit/Source/GLTFExporter/Private/GLTFSceneBuilder.cpp
--- a/Engine/Plugins/Enterprise/GLTFExporter/MigratedFromGit/Source/GLTFExporter/Private/GLTFSceneBuilder.cpp
+++ b/Engine/Plugins/Enterprise/GLTFExporter/MigratedFromGit/Source/GLTFExporter/Private/GLTFSceneBuilder.cpp
@@ -44,7 +44,7 @@ FGLTFJsonNodeIndex FGLTFNodeBuilder::AddNode(FGLTFContainerBuilder& Container) c
 	}
 	else if (const UStaticMeshComponent* StaticMeshComponent = Cast<UStaticMeshComponent>(SceneComponent))
 	{
-		Node.Mesh = Container.AddMesh(StaticMeshComponent);
+		Node.Mesh = Container.ConvertMesh(StaticMeshComponent);
 	}
 	else if (bIsRootComponent && FGLTFConverterUtility::IsHDRIBackdropBlueprint(Blueprint))
 	{
@@ -60,7 +60,7 @@ FGLTFJsonNodeIndex FGLTFNodeBuilder::AddNode(FGLTFContainerBuilder& Container) c
 		}
 	}
 
-	return Container.CreateNode(Node);
+	return Container.AddNode(Node);
 }
 
 FGLTFSceneBuilder::FGLTFSceneBuilder(const UWorld* World, bool bSelectedOnly)
@@ -100,5 +100,5 @@ FGLTFJsonSceneIndex FGLTFSceneBuilder::AddScene(FGLTFContainerBuilder& Container
 		if (NodeIndex != INDEX_NONE) Scene.Nodes.Add(NodeIndex);
 	}
 
-	return Container.CreateScene(Scene);
+	return Container.AddScene(Scene);
 }
diff --git a/Engine/Plugins/Enterprise/GLTFExporter/MigratedFromGit/Source/GLTFExporter/Private/GLTFStaticMeshConverters.cpp b/Engine/Plugins/Enterprise/GLTFExporter/MigratedFromGit/Source/GLTFExporter/Private/GLTFStaticMeshConverters.cpp
new file mode 100644
--- /dev/null
+++ b/Engine/Plugins/Enterprise/GLTFExporter/MigratedFromGit/Source/GLTFExporter/Private/GLTFStaticMeshConverters.cpp
@@ -0,0 +1,251 @@
+// Copyright Epic Games, Inc. All Rights Reserved.
+
+#include "GLTFStaticMeshConverters.h"
+#include "GLTFContainerBuilder.h"
+#include "GLTFConversionUtilities.h"
+
+FGLTFJsonAccessorIndex FGLTFPositionVertexBufferConverter::Convert(FGLTFContainerBuilder& Container, const FString& Name, const FPositionVertexBuffer* VertexBuffer)
+{
+	const uint32 VertexCount = VertexBuffer->GetNumVertices();
+	if (VertexCount == 0)
+	{
+		return FGLTFJsonAccessorIndex(INDEX_NONE);
+	}
+	
+	TArray<FVector> Positions;
+	Positions.AddUninitialized(VertexCount);
+
+	for (uint32 VertexIndex = 0; VertexIndex < VertexCount; ++VertexIndex)
+	{
+		Positions[VertexIndex] = ConvertPosition(VertexBuffer->VertexPosition(VertexIndex));
+	}
+
+	FBox BoundingBox;
+
+	// More accurate bounding box if based on raw vertex values
+	for (const FVector& Position : Positions)
+	{
+		BoundingBox += Position;
+	}
+
+	FGLTFJsonAccessor JsonAccessor;
+	JsonAccessor.Name = Name;
+	JsonAccessor.BufferView = Container.AddBufferView(Positions, Name);
+	JsonAccessor.ComponentType = EGLTFJsonComponentType::F32;
+	JsonAccessor.Count = VertexCount;
+	JsonAccessor.Type = EGLTFJsonAccessorType::Vec3;
+	
+	JsonAccessor.MinMaxLength = 3;
+	JsonAccessor.Max[0] = BoundingBox.Max.X;
+	JsonAccessor.Max[1] = BoundingBox.Max.Y;
+	JsonAccessor.Max[2] = BoundingBox.Max.Z;
+	JsonAccessor.Min[0] = BoundingBox.Min.X;
+	JsonAccessor.Min[1] = BoundingBox.Min.Y;
+	JsonAccessor.Min[2] = BoundingBox.Min.Z;
+	
+	return Container.AddAccessor(JsonAccessor);
+}
+
+FGLTFJsonAccessorIndex FGLTFColorVertexBufferConverter::Convert(FGLTFContainerBuilder& Container, const FString& Name, const FColorVertexBuffer* VertexBuffer)
+{
+	const uint32 VertexCount = VertexBuffer->GetNumVertices();
+	if (VertexCount == 0)
+	{
+		return FGLTFJsonAccessorIndex(INDEX_NONE);
+	}
+
+	TArray<FColor> Colors;
+	Colors.AddUninitialized(VertexCount);
+
+	for (uint32 VertexIndex = 0; VertexIndex < VertexCount; ++VertexIndex)
+	{
+		Colors[VertexIndex] = ConvertColor(VertexBuffer->VertexColor(VertexIndex));
+	}
+
+	FGLTFJsonAccessor JsonAccessor;
+	JsonAccessor.Name = Name;
+	JsonAccessor.BufferView = Container.AddBufferView(Colors, Name);
+	JsonAccessor.ComponentType = EGLTFJsonComponentType::U8;
+	JsonAccessor.Count = VertexCount;
+	JsonAccessor.Type = EGLTFJsonAccessorType::Vec4;
+	JsonAccessor.bNormalized = true;
+
+	return Container.AddAccessor(JsonAccessor);
+}
+
+FGLTFJsonAccessorIndex FGLTFStaticMeshNormalVertexBufferConverter::Convert(FGLTFContainerBuilder& Container, const FString& Name, const FStaticMeshVertexBuffer* VertexBuffer)
+{
+	const uint32 VertexCount = VertexBuffer->GetNumVertices();
+	if (VertexCount == 0)
+	{
+		return FGLTFJsonAccessorIndex(INDEX_NONE);
+	}
+
+	TArray<FVector> Normals;
+	Normals.AddUninitialized(VertexCount);
+	
+	for (uint32 VertexIndex = 0; VertexIndex < VertexCount; ++VertexIndex)
+	{
+		Normals[VertexIndex] = ConvertVector(VertexBuffer->VertexTangentZ(VertexIndex));
+	}
+
+	FGLTFJsonAccessor JsonAccessor;
+	JsonAccessor.Name = Name;
+	JsonAccessor.BufferView = Container.AddBufferView(Normals, Name);
+	JsonAccessor.ComponentType = EGLTFJsonComponentType::F32;
+	JsonAccessor.Count = VertexCount;
+	JsonAccessor.Type = EGLTFJsonAccessorType::Vec3;
+	
+	return Container.AddAccessor(JsonAccessor);
+}
+
+FGLTFJsonAccessorIndex FGLTFStaticMeshTangentVertexBufferConverter::Convert(FGLTFContainerBuilder& Container, const FString& Name, const FStaticMeshVertexBuffer* VertexBuffer)
+{
+	const uint32 VertexCount = VertexBuffer->GetNumVertices();
+	if (VertexCount == 0)
+	{
+		return FGLTFJsonAccessorIndex(INDEX_NONE);
+	}
+
+	TArray<FVector4> Tangents;
+	Tangents.AddUninitialized(VertexCount);
+	
+	for (uint32 VertexIndex = 0; VertexIndex < VertexCount; ++VertexIndex)
+	{
+		Tangents[VertexIndex] = ConvertVector(VertexBuffer->VertexTangentX(VertexIndex));
+	}
+
+	FGLTFJsonAccessor JsonAccessor;
+	JsonAccessor.Name = Name;
+	JsonAccessor.BufferView = Container.AddBufferView(Tangents, Name);
+	JsonAccessor.ComponentType = EGLTFJsonComponentType::F32;
+	JsonAccessor.Count = VertexCount;
+	JsonAccessor.Type = EGLTFJsonAccessorType::Vec4;
+	
+	return Container.AddAccessor(JsonAccessor);
+}
+
+FGLTFJsonAccessorIndex FGLTFStaticMeshUV0VertexBufferConverter::Convert(FGLTFContainerBuilder& Container, const FString& Name, const FStaticMeshVertexBuffer* VertexBuffer)
+{
+	const uint32 VertexCount = VertexBuffer->GetNumVertices();
+	if (VertexCount == 0 || VertexBuffer->GetNumTexCoords() < 1)
+	{
+		return FGLTFJsonAccessorIndex(INDEX_NONE);
+	}
+	
+	TArray<FVector2D> UV0s;
+	UV0s.AddUninitialized(VertexCount);
+
+	for (uint32 VertexIndex = 0; VertexIndex < VertexCount; ++VertexIndex)
+	{
+		UV0s[VertexIndex] = VertexBuffer->GetVertexUV(VertexIndex, 0);
+	}
+
+	FGLTFJsonAccessor JsonAccessor;
+	JsonAccessor.Name = Name;
+	JsonAccessor.BufferView = Container.AddBufferView(UV0s, Name);
+	JsonAccessor.ComponentType = EGLTFJsonComponentType::F32;
+	JsonAccessor.Count = VertexCount;
+	JsonAccessor.Type = EGLTFJsonAccessorType::Vec2;
+
+	return Container.AddAccessor(JsonAccessor);
+}
+
+FGLTFJsonAccessorIndex FGLTFStaticMeshUV1VertexBufferConverter::Convert(FGLTFContainerBuilder& Container, const FString& Name, const FStaticMeshVertexBuffer* VertexBuffer)
+{
+	const uint32 VertexCount = VertexBuffer->GetNumVertices();
+	if (VertexCount == 0 || VertexBuffer->GetNumTexCoords() < 2)
+	{
+		return FGLTFJsonAccessorIndex(INDEX_NONE);
+	}
+
+	TArray<FVector2D> UV1s;
+	UV1s.AddUninitialized(VertexCount);
+
+	for (uint32 VertexIndex = 0; VertexIndex < VertexCount; ++VertexIndex)
+	{
+		UV1s[VertexIndex] = VertexBuffer->GetVertexUV(VertexIndex, 1);
+	}
+
+	FGLTFJsonAccessor JsonAccessor;
+	JsonAccessor.Name = Name;
+	JsonAccessor.BufferView = Container.AddBufferView(UV1s, Name);
+	JsonAccessor.ComponentType = EGLTFJsonComponentType::F32;
+	JsonAccessor.Count = VertexCount;
+	JsonAccessor.Type = EGLTFJsonAccessorType::Vec2;
+
+	return Container.AddAccessor(JsonAccessor);
+}
+
+FGLTFJsonBufferViewIndex FGLTFStaticMeshIndexBufferConverter::Convert(FGLTFContainerBuilder& Container, const FString& Name, const FRawStaticIndexBuffer* IndexBuffer)
+{
+	if (IndexBuffer->GetNumIndices() == 0)
+	{
+		return FGLTFJsonBufferViewIndex(INDEX_NONE);
+	}
+
+	TArray<uint32> Indices;
+	IndexBuffer->GetCopy(Indices);
+	return Container.CreateBufferView(Indices, Name, EGLTFJsonBufferTarget::ElementArrayBuffer);
+}
+
+FGLTFJsonAccessorIndex FGLTFStaticMeshSectionConverter::Convert(FGLTFContainerBuilder& Container, const FString& Name, const FStaticMeshSection* MeshSection, const FRawStaticIndexBuffer* IndexBuffer)
+{	
+	const uint32 TriangleCount = MeshSection->NumTriangles;
+	if (TriangleCount == 0)
+	{
+		return FGLTFJsonAccessorIndex(INDEX_NONE);
+	}
+
+	FGLTFJsonAccessor JsonAccessor;
+	JsonAccessor.Name = Name;
+	JsonAccessor.BufferView = Container.AddIndexBufferView(IndexBuffer);
+	JsonAccessor.ByteOffset = MeshSection->FirstIndex * sizeof(uint32);
+	JsonAccessor.ComponentType = EGLTFJsonComponentType::U32;
+	JsonAccessor.Count = TriangleCount * 3;
+	JsonAccessor.Type = EGLTFJsonAccessorType::Scalar;
+
+	return Container.AddAccessor(JsonAccessor);
+}
+
+FGLTFJsonMeshIndex FGLTFStaticMeshConverter::Convert(FGLTFContainerBuilder& Container, const FString& Name, const UStaticMesh* StaticMesh, int32 LODIndex, const FColorVertexBuffer* OverrideVertexColors)
+{
+	const FStaticMeshLODResources& LODMesh = StaticMesh->GetLODForExport(LODIndex);
+	const FPositionVertexBuffer* PositionBuffer = &LODMesh.VertexBuffers.PositionVertexBuffer;
+	const FStaticMeshVertexBuffer* VertexBuffer = &LODMesh.VertexBuffers.StaticMeshVertexBuffer;
+	const FColorVertexBuffer* ColorBuffer = OverrideVertexColors != nullptr ? OverrideVertexColors : &LODMesh.VertexBuffers.ColorVertexBuffer;
+
+	FGLTFJsonMesh JsonMesh;
+	JsonMesh.Name = Name;
+
+	if (JsonMesh.Name.IsEmpty())
+	{
+		StaticMesh->GetName(JsonMesh.Name);
+		if (LODIndex != 0) JsonMesh.Name += TEXT("_LOD") + FString::FromInt(LODIndex);
+	}
+
+	FGLTFJsonAttributes JsonAttributes;
+	JsonAttributes.Position = Container.ConvertPositionAccessor(PositionBuffer, Name + TEXT("_Positions"));
+	JsonAttributes.Normal = Container.ConvertNormalAccessor(VertexBuffer, Name + TEXT("_Normals"));
+	JsonAttributes.Tangent = Container.ConvertTangentAccessor(VertexBuffer, Name + TEXT("_Tangents"));
+	JsonAttributes.TexCoord0 = Container.ConvertUV0Accessor(VertexBuffer, Name + TEXT("_UV0s"));
+	JsonAttributes.TexCoord1 = Container.ConvertUV1Accessor(VertexBuffer, Name + TEXT("_UV1s"));
+	JsonAttributes.Color0 = Container.ConvertColorAccessor(ColorBuffer, Name + TEXT("_Colors"));
+
+	const FRawStaticIndexBuffer* IndexBuffer = &LODMesh.IndexBuffer;
+	Container.ConvertIndexBufferView(IndexBuffer, Name + TEXT("_Indices"));
+
+	const int32 SectionCount = LODMesh.Sections.Num();
+	JsonMesh.Primitives.AddDefaulted(SectionCount);
+
+	for (int32 SectionIndex = 0; SectionIndex < SectionCount; ++SectionIndex)
+	{
+		FGLTFJsonPrimitive& JsonPrimitive = JsonMesh.Primitives[SectionIndex];
+		JsonPrimitive.Attributes = JsonAttributes;
+
+		const FString SectionName = JsonMesh.Name + (SectionCount != 0 ? TEXT("_Indices_Section") + FString::FromInt(SectionIndex) : TEXT("_Indices"));
+		JsonPrimitive.Indices = Container.ConvertIndexAccessor(&LODMesh.Sections[SectionIndex], IndexBuffer, SectionName);
+	}
+
+	return Container.AddMesh(JsonMesh);
+}
diff --git a/Engine/Plugins/Enterprise/GLTFExporter/MigratedFromGit/Source/GLTFExporter/Private/GLTFStaticMeshExporter.cpp b/Engine/Plugins/Enterprise/GLTFExporter/MigratedFromGit/Source/GLTFExporter/Private/GLTFStaticMeshExporter.cpp
--- a/Engine/Plugins/Enterprise/GLTFExporter/MigratedFromGit/Source/GLTFExporter/Private/GLTFStaticMeshExporter.cpp
+++ b/Engine/Plugins/Enterprise/GLTFExporter/MigratedFromGit/Source/GLTFExporter/Private/GLTFStaticMeshExporter.cpp
@@ -23,15 +23,15 @@ bool UGLTFStaticMeshExporter::ExportBinary(UObject* Object, const TCHAR* Type, F
 
 	FGLTFContainerBuilder Container;
 
-	FGLTFJsonMeshIndex MeshIndex = Container.AddMesh(StaticMesh, 0);
+	FGLTFJsonMeshIndex MeshIndex = Container.ConvertMesh(StaticMesh);
 
 	FGLTFJsonNode Node;
 	Node.Mesh = MeshIndex;
-	FGLTFJsonNodeIndex NodeIndex = Container.CreateNode(Node);
+	FGLTFJsonNodeIndex NodeIndex = Container.AddNode(Node);
 
 	FGLTFJsonScene Scene;
 	Scene.Nodes.Add(NodeIndex);
-	FGLTFJsonSceneIndex SceneIndex = Container.CreateScene(Scene);
+	FGLTFJsonSceneIndex SceneIndex = Container.AddScene(Scene);
 
 	Container.JsonRoot.DefaultScene = SceneIndex;
 
diff --git a/Engine/Plugins/Enterprise/GLTFExporter/MigratedFromGit/Source/GLTFExporter/Public/GLTFBufferBuilder.h b/Engine/Plugins/Enterprise/GLTFExporter/MigratedFromGit/Source/GLTFExporter/Public/GLTFBufferBuilder.h
--- a/Engine/Plugins/Enterprise/GLTFExporter/MigratedFromGit/Source/GLTFExporter/Public/GLTFBufferBuilder.h
+++ b/Engine/Plugins/Enterprise/GLTFExporter/MigratedFromGit/Source/GLTFExporter/Public/GLTFBufferBuilder.h
@@ -14,7 +14,7 @@ struct GLTFEXPORTER_API FGLTFBufferBuilder
 
 	FGLTFBufferBuilder(FGLTFJsonBufferIndex BufferIndex);
 
-	FGLTFJsonBufferViewIndex CreateBufferView(FGLTFContainerBuilder& Container, const void* RawData, uint64 ByteLength, const FString& Name, EGLTFJsonBufferTarget BufferTarget);
+	FGLTFJsonBufferViewIndex AddBufferView(FGLTFContainerBuilder& Container, const void* RawData, uint64 ByteLength, const FString& Name, EGLTFJsonBufferTarget BufferTarget);
 
 	void UpdateBuffer(FGLTFJsonBuffer& JsonBuffer);
 };
diff --git a/Engine/Plugins/Enterprise/GLTFExporter/MigratedFromGit/Source/GLTFExporter/Public/GLTFContainerBuilder.h b/Engine/Plugins/Enterprise/GLTFExporter/MigratedFromGit/Source/GLTFExporter/Public/GLTFContainerBuilder.h
--- a/Engine/Plugins/Enterprise/GLTFExporter/MigratedFromGit/Source/GLTFExporter/Public/GLTFContainerBuilder.h
+++ b/Engine/Plugins/Enterprise/GLTFExporter/MigratedFromGit/Source/GLTFExporter/Public/GLTFContainerBuilder.h
@@ -4,44 +4,44 @@
 
 #include "GLTFJsonRoot.h"
 #include "GLTFBufferBuilder.h"
-#include "GLTFIndexBuilder.h"
+#include "GLTFIndexedConverts.h"
 #include "Engine.h"
 
 struct GLTFEXPORTER_API FGLTFContainerBuilder
 {
 	FGLTFJsonRoot JsonRoot;
 	FGLTFBufferBuilder BufferBuilder;
-	FGLTFIndexBuilder IndexBuilder;
+	FGLTFIndexedConverts IndexedConverts;
 
 	FGLTFContainerBuilder();
 
-	FGLTFJsonAccessorIndex CreateAccessor(const FGLTFJsonAccessor& JsonAccessor);
-	FGLTFJsonBufferIndex CreateBuffer(const FGLTFJsonBuffer& JsonBuffer);
-	FGLTFJsonBufferViewIndex CreateBufferView(const FGLTFJsonBufferView& JsonBufferView);
-	FGLTFJsonMeshIndex CreateMesh(const FGLTFJsonMesh& JsonMesh);
-	FGLTFJsonNodeIndex CreateNode(const FGLTFJsonNode& JsonNode);
-	FGLTFJsonSceneIndex CreateScene(const FGLTFJsonScene& JsonScene);
+	FGLTFJsonAccessorIndex AddAccessor(const FGLTFJsonAccessor& JsonAccessor);
+	FGLTFJsonBufferIndex AddBuffer(const FGLTFJsonBuffer& JsonBuffer);
+	FGLTFJsonBufferViewIndex AddBufferView(const FGLTFJsonBufferView& JsonBufferView);
+	FGLTFJsonMeshIndex AddMesh(const FGLTFJsonMesh& JsonMesh);
+	FGLTFJsonNodeIndex AddNode(const FGLTFJsonNode& JsonNode);
+	FGLTFJsonSceneIndex AddScene(const FGLTFJsonScene& JsonScene);
 
-	FGLTFJsonBufferViewIndex CreateBufferView(const void* RawData, uint64 ByteLength, const FString& Name = TEXT(""), EGLTFJsonBufferTarget BufferTarget = EGLTFJsonBufferTarget::ArrayBuffer);
+	FGLTFJsonBufferViewIndex AddBufferView(const void* RawData, uint64 ByteLength, const FString& Name = TEXT(""), EGLTFJsonBufferTarget BufferTarget = EGLTFJsonBufferTarget::ArrayBuffer);
 
 	template <class ElementType>
-	FGLTFJsonBufferViewIndex CreateBufferView(const TArray<ElementType>& Array, const FString& Name = TEXT(""), EGLTFJsonBufferTarget BufferTarget = EGLTFJsonBufferTarget::ArrayBuffer)
+	FGLTFJsonBufferViewIndex AddBufferView(const TArray<ElementType>& Array, const FString& Name = TEXT(""), EGLTFJsonBufferTarget BufferTarget = EGLTFJsonBufferTarget::ArrayBuffer)
 	{
-		return CreateBufferView(Array.GetData(), Array.Num() * sizeof(ElementType), Name, BufferTarget);
+		return AddBufferView(Array.GetData(), Array.Num() * sizeof(ElementType), Name, BufferTarget);
 	}
 
-	FGLTFJsonAccessorIndex AddPositionAccessor(const FPositionVertexBuffer* VertexBuffer, const FString& Name = TEXT(""));
-	FGLTFJsonAccessorIndex AddNormalAccessor(const FStaticMeshVertexBuffer* VertexBuffer, const FString& Name = TEXT(""));
-	FGLTFJsonAccessorIndex AddTangentAccessor(const FStaticMeshVertexBuffer* VertexBuffer, const FString& Name = TEXT(""));
-	FGLTFJsonAccessorIndex AddUV0Accessor(const FStaticMeshVertexBuffer* VertexBuffer, const FString& Name = TEXT(""));
-	FGLTFJsonAccessorIndex AddUV1Accessor(const FStaticMeshVertexBuffer* VertexBuffer, const FString& Name = TEXT(""));
-	FGLTFJsonAccessorIndex AddColorAccessor(const FColorVertexBuffer* VertexBuffer, const FString& Name = TEXT(""));
+	FGLTFJsonAccessorIndex ConvertPositionAccessor(const FPositionVertexBuffer* VertexBuffer, const FString& DesiredName = TEXT(""));
+	FGLTFJsonAccessorIndex ConvertNormalAccessor(const FStaticMeshVertexBuffer* VertexBuffer, const FString& DesiredName = TEXT(""));
+	FGLTFJsonAccessorIndex ConvertTangentAccessor(const FStaticMeshVertexBuffer* VertexBuffer, const FString& DesiredName = TEXT(""));
+	FGLTFJsonAccessorIndex ConvertUV0Accessor(const FStaticMeshVertexBuffer* VertexBuffer, const FString& DesiredName = TEXT(""));
+	FGLTFJsonAccessorIndex ConvertUV1Accessor(const FStaticMeshVertexBuffer* VertexBuffer, const FString& DesiredName = TEXT(""));
+	FGLTFJsonAccessorIndex ConvertColorAccessor(const FColorVertexBuffer* VertexBuffer, const FString& DesiredName = TEXT(""));
 
-	FGLTFJsonBufferViewIndex AddIndexBufferView(const FRawStaticIndexBuffer* IndexBuffer, const FString& Name = TEXT(""));
-	FGLTFJsonAccessorIndex AddIndexAccessor(const FStaticMeshSection* MeshSection, const FRawStaticIndexBuffer* IndexBuffer, const FString& Name = TEXT(""));
+	FGLTFJsonBufferViewIndex ConvertIndexBufferView(const FRawStaticIndexBuffer* IndexBuffer, const FString& DesiredName = TEXT(""));
+	FGLTFJsonAccessorIndex ConvertIndexAccessor(const FStaticMeshSection* MeshSection, const FRawStaticIndexBuffer* IndexBuffer, const FString& DesiredName = TEXT(""));
 
-	FGLTFJsonMeshIndex AddMesh(const UStaticMesh* StaticMesh, int32 LODIndex = 0, const FColorVertexBuffer* OverrideVertexColors = nullptr);
-	FGLTFJsonMeshIndex AddMesh(const UStaticMeshComponent* StaticMeshComponent);
+	FGLTFJsonMeshIndex ConvertMesh(const UStaticMesh* StaticMesh, int32 LODIndex = 0, const FColorVertexBuffer* OverrideVertexColors = nullptr, const FString& DesiredName = TEXT(""));
+	FGLTFJsonMeshIndex ConvertMesh(const UStaticMeshComponent* StaticMeshComponent, const FString& DesiredName = TEXT(""));
 
 	FGLTFJsonSceneIndex AddScene(const UWorld* World, bool bSelectedOnly = false);
 
diff --git a/Engine/Plugins/Enterprise/GLTFExporter/MigratedFromGit/Source/GLTFExporter/Public/GLTFIndexBuilder.h b/Engine/Plugins/Enterprise/GLTFExporter/MigratedFromGit/Source/GLTFExporter/Public/GLTFIndexBuilder.h
deleted file mode 100644
diff --git a/Engine/Plugins/Enterprise/GLTFExporter/MigratedFromGit/Source/GLTFExporter/Public/GLTFIndexedConverts.h b/Engine/Plugins/Enterprise/GLTFExporter/MigratedFromGit/Source/GLTFExporter/Public/GLTFIndexedConverts.h
new file mode 100644
--- /dev/null
+++ b/Engine/Plugins/Enterprise/GLTFExporter/MigratedFromGit/Source/GLTFExporter/Public/GLTFIndexedConverts.h
@@ -0,0 +1,50 @@
+// Copyright Epic Games, Inc. All Rights Reserved.
+
+#pragma once
+
+#include "GLTFJsonIndex.h"
+#include "GLTFStaticMeshConverters.h"
+#include "Engine.h"
+
+struct FGLTFContainerBuilder;
+
+template <class IndexType, class KeyType, class ConverterType>
+struct GLTFEXPORTER_API TGLTFIndexedConvert
+{
+	TMap<KeyType, IndexType> IndexLookup;
+	
+	template <typename... ArgsType>
+    FORCEINLINE IndexType Find(ArgsType&&... Args) const
+	{
+		const KeyType Key(Forward<ArgsType>(Args)...);
+		return IndexLookup.FindRef(Key);
+	}
+
+	template <typename... ArgsType>
+    FORCEINLINE IndexType Convert(FGLTFContainerBuilder& Container, const FString& DesiredName, ArgsType&&... Args)
+	{
+		const KeyType Key(Forward<ArgsType>(Args)...);
+
+		IndexType Index = IndexLookup.FindRef(Key);
+		if (Index == INDEX_NONE)
+		{
+			Index = ConverterType::Convert(Container, DesiredName, Forward<ArgsType>(Args)...);
+			IndexLookup.Add(Key, Index);
+		}
+
+		return Index;
+	}
+};
+
+struct GLTFEXPORTER_API FGLTFIndexedConverts
+{
+	TGLTFIndexedConvert<FGLTFJsonAccessorIndex, TTuple<const FPositionVertexBuffer*>, FGLTFPositionVertexBufferConverter> PositionVertexBuffers;
+	TGLTFIndexedConvert<FGLTFJsonAccessorIndex, TTuple<const FColorVertexBuffer*>, FGLTFColorVertexBufferConverter> ColorVertexBuffers;
+	TGLTFIndexedConvert<FGLTFJsonAccessorIndex, TTuple<const FStaticMeshVertexBuffer*>, FGLTFStaticMeshNormalVertexBufferConverter> StaticMeshNormalVertexBuffers;
+	TGLTFIndexedConvert<FGLTFJsonAccessorIndex, TTuple<const FStaticMeshVertexBuffer*>, FGLTFStaticMeshTangentVertexBufferConverter> StaticMeshTangentVertexBuffers;
+	TGLTFIndexedConvert<FGLTFJsonAccessorIndex, TTuple<const FStaticMeshVertexBuffer*>, FGLTFStaticMeshUV0VertexBufferConverter> StaticMeshUV0VertexBuffers;
+	TGLTFIndexedConvert<FGLTFJsonAccessorIndex, TTuple<const FStaticMeshVertexBuffer*>, FGLTFStaticMeshUV1VertexBufferConverter> StaticMeshUV1VertexBuffers;
+	TGLTFIndexedConvert<FGLTFJsonBufferViewIndex, TTuple<const FRawStaticIndexBuffer*>, FGLTFStaticMeshIndexBufferConverter> StaticMeshIndexBuffers;
+	TGLTFIndexedConvert<FGLTFJsonAccessorIndex, TTuple<const FStaticMeshSection*, const FRawStaticIndexBuffer*>, FGLTFStaticMeshSectionConverter> StaticMeshSections;
+	TGLTFIndexedConvert<FGLTFJsonMeshIndex, TTuple<const UStaticMesh*, int32, const FColorVertexBuffer*>, FGLTFStaticMeshConverter> StaticMeshes;
+};
diff --git a/Engine/Plugins/Enterprise/GLTFExporter/MigratedFromGit/Source/GLTFExporter/Public/GLTFMeshConverter.h b/Engine/Plugins/Enterprise/GLTFExporter/MigratedFromGit/Source/GLTFExporter/Public/GLTFMeshConverter.h
deleted file mode 100644
diff --git a/Engine/Plugins/Enterprise/GLTFExporter/MigratedFromGit/Source/GLTFExporter/Public/GLTFStaticMeshConverters.h b/Engine/Plugins/Enterprise/GLTFExporter/MigratedFromGit/Source/GLTFExporter/Public/GLTFStaticMeshConverters.h
new file mode 100644
--- /dev/null
+++ b/Engine/Plugins/Enterprise/GLTFExporter/MigratedFromGit/Source/GLTFExporter/Public/GLTFStaticMeshConverters.h
@@ -0,0 +1,53 @@
+// Copyright Epic Games, Inc. All Rights Reserved.
+
+#pragma once
+
+#include "GLTFJsonIndex.h"
+#include "Engine.h"
+
+struct FGLTFContainerBuilder;
+
+struct GLTFEXPORTER_API FGLTFPositionVertexBufferConverter
+{
+	static FGLTFJsonAccessorIndex Convert(FGLTFContainerBuilder& Container, const FString& Name, const FPositionVertexBuffer* VertexBuffer);
+};
+
+struct GLTFEXPORTER_API FGLTFColorVertexBufferConverter
+{
+	static FGLTFJsonAccessorIndex Convert(FGLTFContainerBuilder& Container, const FString& Name, const FColorVertexBuffer* VertexBuffer);
+};
+
+struct GLTFEXPORTER_API FGLTFStaticMeshNormalVertexBufferConverter
+{
+	static FGLTFJsonAccessorIndex Convert(FGLTFContainerBuilder& Container, const FString& Name, const FStaticMeshVertexBuffer* VertexBuffer);
+};
+
+struct GLTFEXPORTER_API FGLTFStaticMeshTangentVertexBufferConverter
+{
+	static FGLTFJsonAccessorIndex Convert(FGLTFContainerBuilder& Container, const FString& Name, const FStaticMeshVertexBuffer* VertexBuffer);
+};
+
+struct GLTFEXPORTER_API FGLTFStaticMeshUV0VertexBufferConverter
+{
+	static FGLTFJsonAccessorIndex Convert(FGLTFContainerBuilder& Container, const FString& Name, const FStaticMeshVertexBuffer* VertexBuffer);
+};
+
+struct GLTFEXPORTER_API FGLTFStaticMeshUV1VertexBufferConverter
+{
+	static FGLTFJsonAccessorIndex Convert(FGLTFContainerBuilder& Container, const FString& Name, const FStaticMeshVertexBuffer* VertexBuffer);
+};
+
+struct GLTFEXPORTER_API FGLTFStaticMeshIndexBufferConverter
+{
+	static FGLTFJsonBufferViewIndex Convert(FGLTFContainerBuilder& Container, const FString& Name, const FRawStaticIndexBuffer* IndexBuffer);
+};
+
+struct GLTFEXPORTER_API FGLTFStaticMeshSectionConverter
+{
+	static FGLTFJsonAccessorIndex Convert(FGLTFContainerBuilder& Container, const FString& Name, const FStaticMeshSection* MeshSection, const FRawStaticIndexBuffer* IndexBuffer);
+};
+
+struct GLTFEXPORTER_API FGLTFStaticMeshConverter
+{
+	static FGLTFJsonMeshIndex Convert(FGLTFContainerBuilder& Container, const FString& Name, const UStaticMesh* StaticMesh, int32 LODIndex, const FColorVertexBuffer* OverrideVertexColors);
+};
