[GLTFExporter] refactor: allow empty names for static mesh if explicit

commit d7e9c70c6922ed1d2c625561d561679fbd1003c1
Date: Sun, 10 May 2020 14:58:58 +0200

    refactor: allow empty names for static mesh if explicit
    (cherry picked from commit 3324f202a53d393c8ce425fd990a47241b8bc653​)

#skipci

diff --git a/Engine/Plugins/Enterprise/GLTFExporter/MigratedFromGit/Source/GLTFExporter/Private/GLTFContainerBuilder.cpp b/Engine/Plugins/Enterprise/GLTFExporter/MigratedFromGit/Source/GLTFExporter/Private/GLTFContainerBuilder.cpp
--- a/Engine/Plugins/Enterprise/GLTFExporter/MigratedFromGit/Source/GLTFExporter/Private/GLTFContainerBuilder.cpp
+++ b/Engine/Plugins/Enterprise/GLTFExporter/MigratedFromGit/Source/GLTFExporter/Private/GLTFContainerBuilder.cpp
@@ -83,9 +83,15 @@ FGLTFJsonAccessorIndex FGLTFContainerBuilder::ConvertIndexAccessor(const FStatic
 	return IndexedConverts.StaticMeshSections.Convert(*this, DesiredName, MeshSection, IndexBuffer);
 }
 
+FGLTFJsonMeshIndex FGLTFContainerBuilder::ConvertMesh(const FStaticMeshLODResources* StaticMeshLOD, const FColorVertexBuffer* OverrideVertexColors, const FString& DesiredName)
+{
+	return IndexedConverts.StaticMeshes.Convert(*this, DesiredName, StaticMeshLOD, OverrideVertexColors);
+}
+
 FGLTFJsonMeshIndex FGLTFContainerBuilder::ConvertMesh(const UStaticMesh* StaticMesh, int32 LODIndex, const FColorVertexBuffer* OverrideVertexColors, const FString& DesiredName)
 {
-	return IndexedConverts.StaticMeshes.Convert(*this, DesiredName, StaticMesh, LODIndex, OverrideVertexColors);
+	const FStaticMeshLODResources* StaticMeshLOD = &StaticMesh->GetLODForExport(LODIndex);
+	return ConvertMesh(StaticMeshLOD, OverrideVertexColors, DesiredName.IsEmpty() ? StaticMesh->GetName() + TEXT("_LOD") + FString::FromInt(LODIndex) : DesiredName);
 }
 
 FGLTFJsonMeshIndex FGLTFContainerBuilder::ConvertMesh(const UStaticMeshComponent* StaticMeshComponent, const FString& DesiredName)
diff --git a/Engine/Plugins/Enterprise/GLTFExporter/MigratedFromGit/Source/GLTFExporter/Private/GLTFStaticMeshConverters.cpp b/Engine/Plugins/Enterprise/GLTFExporter/MigratedFromGit/Source/GLTFExporter/Private/GLTFStaticMeshConverters.cpp
--- a/Engine/Plugins/Enterprise/GLTFExporter/MigratedFromGit/Source/GLTFExporter/Private/GLTFStaticMeshConverters.cpp
+++ b/Engine/Plugins/Enterprise/GLTFExporter/MigratedFromGit/Source/GLTFExporter/Private/GLTFStaticMeshConverters.cpp
@@ -208,34 +208,28 @@ FGLTFJsonAccessorIndex FGLTFStaticMeshSectionConverter::Convert(FGLTFContainerBu
 	return Container.AddAccessor(JsonAccessor);
 }
 
-FGLTFJsonMeshIndex FGLTFStaticMeshConverter::Convert(FGLTFContainerBuilder& Container, const FString& Name, const UStaticMesh* StaticMesh, int32 LODIndex, const FColorVertexBuffer* OverrideVertexColors)
+FGLTFJsonMeshIndex FGLTFStaticMeshConverter::Convert(FGLTFContainerBuilder& Container, const FString& Name, const FStaticMeshLODResources* StaticMeshLOD, const FColorVertexBuffer* OverrideVertexColors)
 {
-	const FStaticMeshLODResources& LODMesh = StaticMesh->GetLODForExport(LODIndex);
-	const FPositionVertexBuffer* PositionBuffer = &LODMesh.VertexBuffers.PositionVertexBuffer;
-	const FStaticMeshVertexBuffer* VertexBuffer = &LODMesh.VertexBuffers.StaticMeshVertexBuffer;
-	const FColorVertexBuffer* ColorBuffer = OverrideVertexColors != nullptr ? OverrideVertexColors : &LODMesh.VertexBuffers.ColorVertexBuffer;
-
 	FGLTFJsonMesh JsonMesh;
 	JsonMesh.Name = Name;
 
-	if (JsonMesh.Name.IsEmpty())
-	{
-		StaticMesh->GetName(JsonMesh.Name);
-		if (LODIndex != 0) JsonMesh.Name += TEXT("_LOD") + FString::FromInt(LODIndex);
-	}
+	const FPositionVertexBuffer* PositionBuffer = &StaticMeshLOD->VertexBuffers.PositionVertexBuffer;
+	const FStaticMeshVertexBuffer* VertexBuffer = &StaticMeshLOD->VertexBuffers.StaticMeshVertexBuffer;
+	const FColorVertexBuffer* ColorBuffer = OverrideVertexColors != nullptr ? OverrideVertexColors : &StaticMeshLOD->VertexBuffers.ColorVertexBuffer;
 
 	FGLTFJsonAttributes JsonAttributes;
-	JsonAttributes.Position = Container.ConvertPositionAccessor(PositionBuffer, Name + TEXT("_Positions"));
-	JsonAttributes.Normal = Container.ConvertNormalAccessor(VertexBuffer, Name + TEXT("_Normals"));
-	JsonAttributes.Tangent = Container.ConvertTangentAccessor(VertexBuffer, Name + TEXT("_Tangents"));
-	JsonAttributes.TexCoord0 = Container.ConvertUV0Accessor(VertexBuffer, Name + TEXT("_UV0s"));
-	JsonAttributes.TexCoord1 = Container.ConvertUV1Accessor(VertexBuffer, Name + TEXT("_UV1s"));
-	JsonAttributes.Color0 = Container.ConvertColorAccessor(ColorBuffer, Name + TEXT("_Colors"));
-
-	const FRawStaticIndexBuffer* IndexBuffer = &LODMesh.IndexBuffer;
-	Container.ConvertIndexBufferView(IndexBuffer, Name + TEXT("_Indices"));
-
-	const int32 SectionCount = LODMesh.Sections.Num();
+	JsonAttributes.Position = Container.ConvertPositionAccessor(PositionBuffer, Name.IsEmpty() ? Name : Name + TEXT("_Positions"));
+	JsonAttributes.Normal = Container.ConvertNormalAccessor(VertexBuffer, Name.IsEmpty() ? Name : Name + TEXT("_Normals"));
+	JsonAttributes.Tangent = Container.ConvertTangentAccessor(VertexBuffer, Name.IsEmpty() ? Name : Name + TEXT("_Tangents"));
+	JsonAttributes.TexCoord0 = Container.ConvertUV0Accessor(VertexBuffer, Name.IsEmpty() ? Name : Name + TEXT("_UV0s"));
+	JsonAttributes.TexCoord1 = Container.ConvertUV1Accessor(VertexBuffer, Name.IsEmpty() ? Name : Name + TEXT("_UV1s"));
+	JsonAttributes.Color0 = Container.ConvertColorAccessor(ColorBuffer, Name.IsEmpty() ? Name : Name + TEXT("_Colors"));
+
+	const FRawStaticIndexBuffer* IndexBuffer = &StaticMeshLOD->IndexBuffer;
+	Container.ConvertIndexBufferView(IndexBuffer, Name.IsEmpty() ? Name : Name + TEXT("_Indices"));
+
+	const TArray<FStaticMeshSection>& Sections = StaticMeshLOD->Sections;
+	const int32 SectionCount = Sections.Num();
 	JsonMesh.Primitives.AddDefaulted(SectionCount);
 
 	for (int32 SectionIndex = 0; SectionIndex < SectionCount; ++SectionIndex)
@@ -243,8 +237,8 @@ FGLTFJsonMeshIndex FGLTFStaticMeshConverter::Convert(FGLTFContainerBuilder& Cont
 		FGLTFJsonPrimitive& JsonPrimitive = JsonMesh.Primitives[SectionIndex];
 		JsonPrimitive.Attributes = JsonAttributes;
 
-		const FString SectionName = JsonMesh.Name + (SectionCount != 0 ? TEXT("_Indices_Section") + FString::FromInt(SectionIndex) : TEXT("_Indices"));
-		JsonPrimitive.Indices = Container.ConvertIndexAccessor(&LODMesh.Sections[SectionIndex], IndexBuffer, SectionName);
+		JsonPrimitive.Indices = Container.ConvertIndexAccessor(&Sections[SectionIndex], IndexBuffer,
+			Name.IsEmpty() ? Name : JsonMesh.Name + (SectionCount != 0 ? TEXT("_Indices_Section") + FString::FromInt(SectionIndex) : TEXT("_Indices")));
 	}
 
 	return Container.AddMesh(JsonMesh);
diff --git a/Engine/Plugins/Enterprise/GLTFExporter/MigratedFromGit/Source/GLTFExporter/Public/GLTFContainerBuilder.h b/Engine/Plugins/Enterprise/GLTFExporter/MigratedFromGit/Source/GLTFExporter/Public/GLTFContainerBuilder.h
--- a/Engine/Plugins/Enterprise/GLTFExporter/MigratedFromGit/Source/GLTFExporter/Public/GLTFContainerBuilder.h
+++ b/Engine/Plugins/Enterprise/GLTFExporter/MigratedFromGit/Source/GLTFExporter/Public/GLTFContainerBuilder.h
@@ -40,6 +40,7 @@ struct GLTFEXPORTER_API FGLTFContainerBuilder
 	FGLTFJsonBufferViewIndex ConvertIndexBufferView(const FRawStaticIndexBuffer* IndexBuffer, const FString& DesiredName = TEXT(""));
 	FGLTFJsonAccessorIndex ConvertIndexAccessor(const FStaticMeshSection* MeshSection, const FRawStaticIndexBuffer* IndexBuffer, const FString& DesiredName = TEXT(""));
 
+	FGLTFJsonMeshIndex ConvertMesh(const FStaticMeshLODResources* StaticMeshLOD, const FColorVertexBuffer* OverrideVertexColors = nullptr, const FString& DesiredName = TEXT(""));
 	FGLTFJsonMeshIndex ConvertMesh(const UStaticMesh* StaticMesh, int32 LODIndex = 0, const FColorVertexBuffer* OverrideVertexColors = nullptr, const FString& DesiredName = TEXT(""));
 	FGLTFJsonMeshIndex ConvertMesh(const UStaticMeshComponent* StaticMeshComponent, const FString& DesiredName = TEXT(""));
 
diff --git a/Engine/Plugins/Enterprise/GLTFExporter/MigratedFromGit/Source/GLTFExporter/Public/GLTFIndexedConverts.h b/Engine/Plugins/Enterprise/GLTFExporter/MigratedFromGit/Source/GLTFExporter/Public/GLTFIndexedConverts.h
--- a/Engine/Plugins/Enterprise/GLTFExporter/MigratedFromGit/Source/GLTFExporter/Public/GLTFIndexedConverts.h
+++ b/Engine/Plugins/Enterprise/GLTFExporter/MigratedFromGit/Source/GLTFExporter/Public/GLTFIndexedConverts.h
@@ -46,5 +46,5 @@ struct GLTFEXPORTER_API FGLTFIndexedConverts
 	TGLTFIndexedConvert<FGLTFJsonAccessorIndex, TTuple<const FStaticMeshVertexBuffer*>, FGLTFStaticMeshUV1VertexBufferConverter> StaticMeshUV1VertexBuffers;
 	TGLTFIndexedConvert<FGLTFJsonBufferViewIndex, TTuple<const FRawStaticIndexBuffer*>, FGLTFStaticMeshIndexBufferConverter> StaticMeshIndexBuffers;
 	TGLTFIndexedConvert<FGLTFJsonAccessorIndex, TTuple<const FStaticMeshSection*, const FRawStaticIndexBuffer*>, FGLTFStaticMeshSectionConverter> StaticMeshSections;
-	TGLTFIndexedConvert<FGLTFJsonMeshIndex, TTuple<const UStaticMesh*, int32, const FColorVertexBuffer*>, FGLTFStaticMeshConverter> StaticMeshes;
+	TGLTFIndexedConvert<FGLTFJsonMeshIndex, TTuple<const FStaticMeshLODResources*, const FColorVertexBuffer*>, FGLTFStaticMeshConverter> StaticMeshes;
 };
diff --git a/Engine/Plugins/Enterprise/GLTFExporter/MigratedFromGit/Source/GLTFExporter/Public/GLTFStaticMeshConverters.h b/Engine/Plugins/Enterprise/GLTFExporter/MigratedFromGit/Source/GLTFExporter/Public/GLTFStaticMeshConverters.h
--- a/Engine/Plugins/Enterprise/GLTFExporter/MigratedFromGit/Source/GLTFExporter/Public/GLTFStaticMeshConverters.h
+++ b/Engine/Plugins/Enterprise/GLTFExporter/MigratedFromGit/Source/GLTFExporter/Public/GLTFStaticMeshConverters.h
@@ -49,5 +49,5 @@ struct GLTFEXPORTER_API FGLTFStaticMeshSectionConverter
 
 struct GLTFEXPORTER_API FGLTFStaticMeshConverter
 {
-	static FGLTFJsonMeshIndex Convert(FGLTFContainerBuilder& Container, const FString& Name, const UStaticMesh* StaticMesh, int32 LODIndex, const FColorVertexBuffer* OverrideVertexColors);
+	static FGLTFJsonMeshIndex Convert(FGLTFContainerBuilder& Container, const FString& Name, const FStaticMeshLODResources* StaticMeshLOD, const FColorVertexBuffer* OverrideVertexColors);
 };
