[GLTFExporter] feat: make ULevel and USceneComponent converts reusable

commit a8ba9f8ef9223ae24f2a6a44d7a0d19e6385201b
Date: Sun, 10 May 2020 15:43:11 +0200

    feat: make ULevel and USceneComponent converts reusable
    (cherry picked from commit 2271170ed058f6d4d0d29094f8607c3018eb2a38​)
    # Conflicts from cherry-pick:
    #	Plugins/GLTFExporter/Source/GLTFExporter/Private/GLTFSceneBuilder.cpp

#skipci

diff --git a/Engine/Plugins/Enterprise/GLTFExporter/MigratedFromGit/Source/GLTFExporter/Private/GLTFContainerBuilder.cpp b/Engine/Plugins/Enterprise/GLTFExporter/MigratedFromGit/Source/GLTFExporter/Private/GLTFContainerBuilder.cpp
--- a/Engine/Plugins/Enterprise/GLTFExporter/MigratedFromGit/Source/GLTFExporter/Private/GLTFContainerBuilder.cpp
+++ b/Engine/Plugins/Enterprise/GLTFExporter/MigratedFromGit/Source/GLTFExporter/Private/GLTFContainerBuilder.cpp
@@ -1,7 +1,6 @@
 // Copyright Epic Games, Inc. All Rights Reserved.
 
 #include "GLTFContainerBuilder.h"
-#include "GLTFSceneBuilder.h"
 
 FGLTFContainerBuilder::FGLTFContainerBuilder()
 	: BufferBuilder(AddBuffer(FGLTFJsonBuffer()))
@@ -102,9 +101,20 @@ FGLTFJsonMeshIndex FGLTFContainerBuilder::ConvertMesh(const UStaticMeshComponent
 	return ConvertMesh(StaticMesh, LODIndex, OverrideVertexColors, DesiredName);
 }
 
-FGLTFJsonSceneIndex FGLTFContainerBuilder::AddScene(const UWorld* World, bool bSelectedOnly)
+FGLTFJsonNodeIndex FGLTFContainerBuilder::ConvertNode(const USceneComponent* SceneComponent, bool bSelectedOnly, bool bRootNode, const FString& DesiredName)
 {
-	return FGLTFSceneBuilder(World, bSelectedOnly).AddScene(*this);
+	return IndexedConverts.SceneComponents.Convert(*this, DesiredName, SceneComponent, bSelectedOnly, bRootNode);
+}
+
+FGLTFJsonSceneIndex FGLTFContainerBuilder::ConvertScene(const ULevel* Level, bool bSelectedOnly, const FString& DesiredName)
+{
+	return IndexedConverts.Levels.Convert(*this, DesiredName, Level, bSelectedOnly);
+}
+
+FGLTFJsonSceneIndex FGLTFContainerBuilder::ConvertScene(const UWorld* World, bool bSelectedOnly, const FString& DesiredName)
+{
+	const ULevel* Level = World->PersistentLevel;
+	return ConvertScene(Level, bSelectedOnly, DesiredName.IsEmpty() ? World->GetName() : DesiredName);
 }
 
 void FGLTFContainerBuilder::Serialize(FArchive& Archive)
diff --git a/Engine/Plugins/Enterprise/GLTFExporter/MigratedFromGit/Source/GLTFExporter/Private/GLTFLevelConverters.cpp b/Engine/Plugins/Enterprise/GLTFExporter/MigratedFromGit/Source/GLTFExporter/Private/GLTFLevelConverters.cpp
new file mode 100644
--- /dev/null
+++ b/Engine/Plugins/Enterprise/GLTFExporter/MigratedFromGit/Source/GLTFExporter/Private/GLTFLevelConverters.cpp
@@ -0,0 +1,79 @@
+// Copyright Epic Games, Inc. All Rights Reserved.
+
+//#include "GLTFLevelConverter.h"
+
+#include "GLTFLevelConverters.h"
+#include "GLTFContainerBuilder.h"
+#include "GLTFConversionUtilities.h"
+
+FGLTFJsonNodeIndex FGLTFSceneComponentConverter::Convert(FGLTFContainerBuilder& Container, const FString& Name, const USceneComponent* SceneComponent, bool bSelectedOnly, bool bRootNode)
+{
+	const AActor* Owner = SceneComponent->GetOwner();
+	const bool bIsRootComponent = Owner != nullptr && Owner->GetRootComponent() == SceneComponent;
+	const UBlueprint* Blueprint = UBlueprint::GetBlueprintFromClass(Owner != nullptr ? Owner->GetClass() : nullptr);
+
+	FGLTFJsonNode Node;
+	Node.Name = Name.IsEmpty() ? Owner->GetName() + TEXT("_") + SceneComponent->GetName() : Name;
+
+	const FTransform Transform = bRootNode ? SceneComponent->GetComponentTransform() : SceneComponent->GetRelativeTransform();
+	Node.Translation = ConvertPosition(Transform.GetTranslation());
+	Node.Rotation = ConvertRotation(Transform.GetRotation());
+	Node.Scale = ConvertScale(Transform.GetScale3D());
+	
+	if (IsSkySphereBlueprint(Blueprint))
+	{
+		// Ignore components
+	}
+	else if (const UStaticMeshComponent* StaticMeshComponent = Cast<UStaticMeshComponent>(SceneComponent))
+	{
+		Node.Mesh = Container.ConvertMesh(StaticMeshComponent);
+	}
+	else if (IsHDRIBackdropBlueprint(Blueprint) && bIsRootComponent)
+	{
+		// TODO: add support for backdrop
+	}
+
+	const TArray<USceneComponent*>& Children = SceneComponent->GetAttachChildren();
+	for (const USceneComponent* ChildComponent : Children)
+	{
+		if (ChildComponent != nullptr)
+		{
+			const AActor* ChildOwner = ChildComponent->GetOwner();
+			if (!bSelectedOnly || ChildOwner->IsSelected())
+			{
+				FGLTFJsonNodeIndex NodeIndex = Container.ConvertNode(ChildComponent, bSelectedOnly, false);
+				if (NodeIndex != INDEX_NONE) Node.Children.Add(NodeIndex);
+			}
+		}
+	}
+
+	return Container.AddNode(Node);
+}
+
+FGLTFJsonSceneIndex FGLTFLevelConverter::Convert(FGLTFContainerBuilder& Container, const FString& Name, const ULevel* Level, bool bSelectedOnly)
+{
+	FGLTFJsonScene Scene;
+	Scene.Name = Name;
+	
+	// TODO: export Level->Model ?
+
+	const TArray<AActor*>& Actors = Level->Actors;
+	for (const AActor* Actor : Actors)
+	{
+		if (Actor != nullptr && (!bSelectedOnly || Actor->IsSelected()))
+		{
+			const USceneComponent* RootComponent = Actor->GetRootComponent();
+			if (RootComponent != nullptr)
+			{
+				const AActor* ParentActor = Actor->GetParentActor();
+				if (ParentActor == nullptr || (bSelectedOnly && !ParentActor->IsSelected()))
+				{
+					FGLTFJsonNodeIndex NodeIndex = Container.ConvertNode(RootComponent, bSelectedOnly, true);
+					if (NodeIndex != INDEX_NONE) Scene.Nodes.Add(NodeIndex);
+				}
+			}
+		}
+	}
+
+	return Container.AddScene(Scene);
+}
diff --git a/Engine/Plugins/Enterprise/GLTFExporter/MigratedFromGit/Source/GLTFExporter/Private/GLTFLevelExporter.cpp b/Engine/Plugins/Enterprise/GLTFExporter/MigratedFromGit/Source/GLTFExporter/Private/GLTFLevelExporter.cpp
--- a/Engine/Plugins/Enterprise/GLTFExporter/MigratedFromGit/Source/GLTFExporter/Private/GLTFLevelExporter.cpp
+++ b/Engine/Plugins/Enterprise/GLTFExporter/MigratedFromGit/Source/GLTFExporter/Private/GLTFLevelExporter.cpp
@@ -21,7 +21,7 @@ bool UGLTFLevelExporter::ExportBinary(UObject* Object, const TCHAR* Type, FArchi
 	}
 
 	FGLTFContainerBuilder Container;
-	const FGLTFJsonSceneIndex SceneIndex = Container.AddScene(World, bSelectedOnly);
+	const FGLTFJsonSceneIndex SceneIndex = Container.ConvertScene(World, bSelectedOnly);
 	Container.JsonRoot.DefaultScene = SceneIndex;
 
 	Container.Serialize(Archive);
diff --git a/Engine/Plugins/Enterprise/GLTFExporter/MigratedFromGit/Source/GLTFExporter/Private/GLTFSceneBuilder.cpp b/Engine/Plugins/Enterprise/GLTFExporter/MigratedFromGit/Source/GLTFExporter/Private/GLTFSceneBuilder.cpp
deleted file mode 100644
diff --git a/Engine/Plugins/Enterprise/GLTFExporter/MigratedFromGit/Source/GLTFExporter/Public/GLTFContainerBuilder.h b/Engine/Plugins/Enterprise/GLTFExporter/MigratedFromGit/Source/GLTFExporter/Public/GLTFContainerBuilder.h
--- a/Engine/Plugins/Enterprise/GLTFExporter/MigratedFromGit/Source/GLTFExporter/Public/GLTFContainerBuilder.h
+++ b/Engine/Plugins/Enterprise/GLTFExporter/MigratedFromGit/Source/GLTFExporter/Public/GLTFContainerBuilder.h
@@ -44,7 +44,9 @@ struct GLTFEXPORTER_API FGLTFContainerBuilder
 	FGLTFJsonMeshIndex ConvertMesh(const UStaticMesh* StaticMesh, int32 LODIndex = 0, const FColorVertexBuffer* OverrideVertexColors = nullptr, const FString& DesiredName = TEXT(""));
 	FGLTFJsonMeshIndex ConvertMesh(const UStaticMeshComponent* StaticMeshComponent, const FString& DesiredName = TEXT(""));
 
-	FGLTFJsonSceneIndex AddScene(const UWorld* World, bool bSelectedOnly = false);
+	FGLTFJsonNodeIndex ConvertNode(const USceneComponent* SceneComponent, bool bSelectedOnly, bool bRootNode = false, const FString& DesiredName = TEXT(""));
+	FGLTFJsonSceneIndex ConvertScene(const ULevel* Level, bool bSelectedOnly, const FString& DesiredName = TEXT(""));
+	FGLTFJsonSceneIndex ConvertScene(const UWorld* World, bool bSelectedOnly, const FString& DesiredName = TEXT(""));
 
 	void Serialize(FArchive& Archive);
 };
diff --git a/Engine/Plugins/Enterprise/GLTFExporter/MigratedFromGit/Source/GLTFExporter/Public/GLTFIndexedConverts.h b/Engine/Plugins/Enterprise/GLTFExporter/MigratedFromGit/Source/GLTFExporter/Public/GLTFIndexedConverts.h
--- a/Engine/Plugins/Enterprise/GLTFExporter/MigratedFromGit/Source/GLTFExporter/Public/GLTFIndexedConverts.h
+++ b/Engine/Plugins/Enterprise/GLTFExporter/MigratedFromGit/Source/GLTFExporter/Public/GLTFIndexedConverts.h
@@ -4,6 +4,7 @@
 
 #include "GLTFJsonIndex.h"
 #include "GLTFStaticMeshConverters.h"
+#include "GLTFLevelConverters.h"
 #include "Engine.h"
 
 struct FGLTFContainerBuilder;
@@ -47,4 +48,7 @@ struct GLTFEXPORTER_API FGLTFIndexedConverts
 	TGLTFIndexedConvert<FGLTFJsonBufferViewIndex, TTuple<const FRawStaticIndexBuffer*>, FGLTFStaticMeshIndexBufferConverter> StaticMeshIndexBuffers;
 	TGLTFIndexedConvert<FGLTFJsonAccessorIndex, TTuple<const FStaticMeshSection*, const FRawStaticIndexBuffer*>, FGLTFStaticMeshSectionConverter> StaticMeshSections;
 	TGLTFIndexedConvert<FGLTFJsonMeshIndex, TTuple<const FStaticMeshLODResources*, const FColorVertexBuffer*>, FGLTFStaticMeshConverter> StaticMeshes;
+
+	TGLTFIndexedConvert<FGLTFJsonNodeIndex, TTuple<const USceneComponent*, bool, bool>, FGLTFSceneComponentConverter> SceneComponents;
+	TGLTFIndexedConvert<FGLTFJsonSceneIndex, TTuple<const ULevel*, bool>, FGLTFLevelConverter> Levels;
 };
diff --git a/Engine/Plugins/Enterprise/GLTFExporter/MigratedFromGit/Source/GLTFExporter/Public/GLTFLevelConverters.h b/Engine/Plugins/Enterprise/GLTFExporter/MigratedFromGit/Source/GLTFExporter/Public/GLTFLevelConverters.h
new file mode 100644
--- /dev/null
+++ b/Engine/Plugins/Enterprise/GLTFExporter/MigratedFromGit/Source/GLTFExporter/Public/GLTFLevelConverters.h
@@ -0,0 +1,18 @@
+// Copyright Epic Games, Inc. All Rights Reserved.
+
+#pragma once
+
+#include "GLTFJsonIndex.h"
+#include "Engine.h"
+
+struct FGLTFContainerBuilder;
+
+struct GLTFEXPORTER_API FGLTFSceneComponentConverter
+{
+	static FGLTFJsonNodeIndex Convert(FGLTFContainerBuilder& Container, const FString& Name, const USceneComponent* SceneComponent, bool bSelectedOnly, bool bRootNode = false);
+};
+
+struct GLTFEXPORTER_API FGLTFLevelConverter
+{
+	static FGLTFJsonSceneIndex Convert(FGLTFContainerBuilder& Container, const FString& Name, const ULevel* Level, bool bSelectedOnly);
+};
diff --git a/Engine/Plugins/Enterprise/GLTFExporter/MigratedFromGit/Source/GLTFExporter/Public/GLTFSceneBuilder.h b/Engine/Plugins/Enterprise/GLTFExporter/MigratedFromGit/Source/GLTFExporter/Public/GLTFSceneBuilder.h
deleted file mode 100644
