[GLTFExporter] refactor: split GLTFContainerBuilder into GLTFJsonBuilder, GLTFBufferBuilder, and GLTFIndexedBuilder

commit 2738be33a6683a38c3387f7e497374dca9fa20b4
Date: Sun, 10 May 2020 17:02:10 +0200

    refactor: split GLTFContainerBuilder into GLTFJsonBuilder, GLTFBufferBuilder, and GLTFIndexedBuilder
    (cherry picked from commit f6631fcd06e4ec0e31571e01c2285fc6d19bb941​)

#skipci

diff --git a/Engine/Plugins/Enterprise/GLTFExporter/MigratedFromGit/Source/GLTFExporter/Private/GLTFBufferBuilder.cpp b/Engine/Plugins/Enterprise/GLTFExporter/MigratedFromGit/Source/GLTFExporter/Private/GLTFBufferBuilder.cpp
--- a/Engine/Plugins/Enterprise/GLTFExporter/MigratedFromGit/Source/GLTFExporter/Private/GLTFBufferBuilder.cpp
+++ b/Engine/Plugins/Enterprise/GLTFExporter/MigratedFromGit/Source/GLTFExporter/Private/GLTFBufferBuilder.cpp
@@ -1,31 +1,32 @@
 // Copyright Epic Games, Inc. All Rights Reserved.
 
 #include "GLTFBufferBuilder.h"
-#include "GLTFContainerBuilder.h"
 #include "Misc/Base64.h"
 
-FGLTFBufferBuilder::FGLTFBufferBuilder(FGLTFJsonBufferIndex BufferIndex)
-	: BufferIndex(BufferIndex)
+FGLTFBufferBuilder::FGLTFBufferBuilder()
 {
+	BufferIndex = AddBuffer(FGLTFJsonBuffer());
 }
 
-FGLTFJsonBufferViewIndex FGLTFBufferBuilder::AddBufferView(FGLTFContainerBuilder& Container, const void* RawData, uint64 ByteLength, const FString& Name, EGLTFJsonBufferTarget BufferTarget)
+FGLTFJsonBufferViewIndex FGLTFBufferBuilder::AddBufferView(const void* RawData, uint64 ByteLength, const FString& Name, EGLTFJsonBufferTarget BufferTarget)
 {
-	FGLTFJsonBufferView BufferView;
-	BufferView.Name = Name;
-	BufferView.Buffer = BufferIndex;
-	BufferView.ByteOffset = BufferData.Num();
-	BufferView.ByteLength = ByteLength;
-	BufferView.Target = BufferTarget;
+	FGLTFJsonBufferView JsonBufferView;
+	JsonBufferView.Name = Name;
+	JsonBufferView.Buffer = BufferIndex;
+	JsonBufferView.ByteOffset = BufferData.Num();
+	JsonBufferView.ByteLength = ByteLength;
+	JsonBufferView.Target = BufferTarget;
 
 	BufferData.Append(static_cast<const uint8*>(RawData), ByteLength);
 
-	return Container.AddBufferView(BufferView);
+	return FGLTFJsonBuilder::AddBufferView(JsonBufferView);
 }
 
-void FGLTFBufferBuilder::UpdateBuffer(FGLTFJsonBuffer& JsonBuffer)
+void FGLTFBufferBuilder::UpdateBuffer()
 {
+	FGLTFJsonBuffer& JsonBuffer = JsonRoot.Buffers[BufferIndex];
 	const int32 ByteLength = BufferData.Num();
+
 	if (JsonBuffer.ByteLength != ByteLength)
 	{
 		JsonBuffer.ByteLength = ByteLength;
@@ -34,3 +35,9 @@ void FGLTFBufferBuilder::UpdateBuffer(FGLTFJsonBuffer& JsonBuffer)
 		JsonBuffer.URI = TEXT("data:application/octet-stream;base64,") + DataBase64;
 	}
 }
+
+void FGLTFBufferBuilder::Serialize(FArchive& Archive)
+{
+	UpdateBuffer();
+	FGLTFJsonBuilder::Serialize(Archive);
+}
diff --git a/Engine/Plugins/Enterprise/GLTFExporter/MigratedFromGit/Source/GLTFExporter/Private/GLTFContainerBuilder.cpp b/Engine/Plugins/Enterprise/GLTFExporter/MigratedFromGit/Source/GLTFExporter/Private/GLTFContainerBuilder.cpp
--- a/Engine/Plugins/Enterprise/GLTFExporter/MigratedFromGit/Source/GLTFExporter/Private/GLTFContainerBuilder.cpp
+++ b/Engine/Plugins/Enterprise/GLTFExporter/MigratedFromGit/Source/GLTFExporter/Private/GLTFContainerBuilder.cpp
@@ -1,125 +1,3 @@
 // Copyright Epic Games, Inc. All Rights Reserved.
 
 #include "GLTFContainerBuilder.h"
-
-FGLTFContainerBuilder::FGLTFContainerBuilder()
-	: BufferBuilder(AddBuffer(FGLTFJsonBuffer()))
-{
-}
-
-FGLTFJsonAccessorIndex FGLTFContainerBuilder::AddAccessor(const FGLTFJsonAccessor& JsonAccessor)
-{
-	return FGLTFJsonAccessorIndex(JsonRoot.Accessors.Add(JsonAccessor));
-}
-
-FGLTFJsonBufferIndex FGLTFContainerBuilder::AddBuffer(const FGLTFJsonBuffer& JsonBuffer)
-{
-	return FGLTFJsonBufferIndex(JsonRoot.Buffers.Add(JsonBuffer));
-}
-
-FGLTFJsonBufferViewIndex FGLTFContainerBuilder::AddBufferView(const FGLTFJsonBufferView& JsonBufferView)
-{
-	return FGLTFJsonBufferViewIndex(JsonRoot.BufferViews.Add(JsonBufferView));
-}
-
-FGLTFJsonMeshIndex FGLTFContainerBuilder::AddMesh(const FGLTFJsonMesh& JsonMesh)
-{
-	return FGLTFJsonMeshIndex(JsonRoot.Meshes.Add(JsonMesh));
-}
-
-FGLTFJsonNodeIndex FGLTFContainerBuilder::AddNode(const FGLTFJsonNode& JsonNode)
-{
-	return FGLTFJsonNodeIndex(JsonRoot.Nodes.Add(JsonNode));
-}
-
-FGLTFJsonSceneIndex FGLTFContainerBuilder::AddScene(const FGLTFJsonScene& JsonScene)
-{
-	return FGLTFJsonSceneIndex(JsonRoot.Scenes.Add(JsonScene));
-}
-
-FGLTFJsonBufferViewIndex FGLTFContainerBuilder::AddBufferView(const void* RawData, uint64 ByteLength, const FString& Name, EGLTFJsonBufferTarget BufferTarget)
-{
-	return BufferBuilder.AddBufferView(*this, RawData, ByteLength, Name, BufferTarget);
-}
-
-FGLTFJsonAccessorIndex FGLTFContainerBuilder::ConvertPositionAccessor(const FPositionVertexBuffer* VertexBuffer, const FString& DesiredName)
-{
-	return IndexedConverts.PositionVertexBuffers.Convert(*this, DesiredName, VertexBuffer);
-}
-
-FGLTFJsonAccessorIndex FGLTFContainerBuilder::ConvertColorAccessor(const FColorVertexBuffer* VertexBuffer, const FString& DesiredName)
-{
-	return IndexedConverts.ColorVertexBuffers.Convert(*this, DesiredName, VertexBuffer);
-}
-
-FGLTFJsonAccessorIndex FGLTFContainerBuilder::ConvertNormalAccessor(const FStaticMeshVertexBuffer* VertexBuffer, const FString& DesiredName)
-{
-	return IndexedConverts.StaticMeshNormalVertexBuffers.Convert(*this, DesiredName, VertexBuffer);
-}
-
-FGLTFJsonAccessorIndex FGLTFContainerBuilder::ConvertTangentAccessor(const FStaticMeshVertexBuffer* VertexBuffer, const FString& DesiredName)
-{
-	return IndexedConverts.StaticMeshTangentVertexBuffers.Convert(*this, DesiredName, VertexBuffer);
-}
-
-FGLTFJsonAccessorIndex FGLTFContainerBuilder::ConvertUV0Accessor(const FStaticMeshVertexBuffer* VertexBuffer, const FString& DesiredName)
-{
-	return IndexedConverts.StaticMeshUV0VertexBuffers.Convert(*this, DesiredName, VertexBuffer);
-}
-
-FGLTFJsonAccessorIndex FGLTFContainerBuilder::ConvertUV1Accessor(const FStaticMeshVertexBuffer* VertexBuffer, const FString& DesiredName)
-{
-	return IndexedConverts.StaticMeshUV1VertexBuffers.Convert(*this, DesiredName, VertexBuffer);
-}
-
-FGLTFJsonBufferViewIndex FGLTFContainerBuilder::ConvertIndexBufferView(const FRawStaticIndexBuffer* IndexBuffer, const FString& DesiredName)
-{
-	return IndexedConverts.StaticMeshIndexBuffers.Convert(*this, DesiredName, IndexBuffer);
-}
-
-FGLTFJsonAccessorIndex FGLTFContainerBuilder::ConvertIndexAccessor(const FStaticMeshSection* MeshSection, const FRawStaticIndexBuffer* IndexBuffer, const FString& DesiredName)
-{
-	return IndexedConverts.StaticMeshSections.Convert(*this, DesiredName, MeshSection, IndexBuffer);
-}
-
-FGLTFJsonMeshIndex FGLTFContainerBuilder::ConvertMesh(const FStaticMeshLODResources* StaticMeshLOD, const FColorVertexBuffer* OverrideVertexColors, const FString& DesiredName)
-{
-	return IndexedConverts.StaticMeshes.Convert(*this, DesiredName, StaticMeshLOD, OverrideVertexColors);
-}
-
-FGLTFJsonMeshIndex FGLTFContainerBuilder::ConvertMesh(const UStaticMesh* StaticMesh, int32 LODIndex, const FColorVertexBuffer* OverrideVertexColors, const FString& DesiredName)
-{
-	const FStaticMeshLODResources* StaticMeshLOD = &StaticMesh->GetLODForExport(LODIndex);
-	return ConvertMesh(StaticMeshLOD, OverrideVertexColors, DesiredName.IsEmpty() ? StaticMesh->GetName() + TEXT("_LOD") + FString::FromInt(LODIndex) : DesiredName);
-}
-
-FGLTFJsonMeshIndex FGLTFContainerBuilder::ConvertMesh(const UStaticMeshComponent* StaticMeshComponent, const FString& DesiredName)
-{
-	const UStaticMesh* StaticMesh = StaticMeshComponent->GetStaticMesh();
-	const int32 LODIndex = StaticMeshComponent->ForcedLodModel > 0 ? StaticMeshComponent->ForcedLodModel - 1 : /* auto-select */ 0;
-	const FColorVertexBuffer* OverrideVertexColors = LODIndex < StaticMeshComponent->LODData.Num() ? StaticMeshComponent->LODData[LODIndex].OverrideVertexColors : nullptr;
-	return ConvertMesh(StaticMesh, LODIndex, OverrideVertexColors, DesiredName);
-}
-
-FGLTFJsonNodeIndex FGLTFContainerBuilder::ConvertNode(const USceneComponent* SceneComponent, bool bSelectedOnly, bool bRootNode, const FString& DesiredName)
-{
-	return IndexedConverts.SceneComponents.Convert(*this, DesiredName, SceneComponent, bSelectedOnly, bRootNode);
-}
-
-FGLTFJsonSceneIndex FGLTFContainerBuilder::ConvertScene(const ULevel* Level, bool bSelectedOnly, const FString& DesiredName)
-{
-	return IndexedConverts.Levels.Convert(*this, DesiredName, Level, bSelectedOnly);
-}
-
-FGLTFJsonSceneIndex FGLTFContainerBuilder::ConvertScene(const UWorld* World, bool bSelectedOnly, const FString& DesiredName)
-{
-	const ULevel* Level = World->PersistentLevel;
-	return ConvertScene(Level, bSelectedOnly, DesiredName.IsEmpty() ? World->GetName() : DesiredName);
-}
-
-void FGLTFContainerBuilder::Serialize(FArchive& Archive)
-{
-	FGLTFJsonBuffer& JsonBuffer = JsonRoot.Buffers[BufferBuilder.BufferIndex];
-	BufferBuilder.UpdateBuffer(JsonBuffer);
-	JsonRoot.Serialize(&Archive, true);
-}
diff --git a/Engine/Plugins/Enterprise/GLTFExporter/MigratedFromGit/Source/GLTFExporter/Private/GLTFIndexedBuilder.cpp b/Engine/Plugins/Enterprise/GLTFExporter/MigratedFromGit/Source/GLTFExporter/Private/GLTFIndexedBuilder.cpp
new file mode 100644
--- /dev/null
+++ b/Engine/Plugins/Enterprise/GLTFExporter/MigratedFromGit/Source/GLTFExporter/Private/GLTFIndexedBuilder.cpp
@@ -0,0 +1,78 @@
+// Copyright Epic Games, Inc. All Rights Reserved.
+
+#include "GLTFIndexedBuilder.h"
+
+FGLTFJsonAccessorIndex FGLTFIndexedBuilder::ConvertPositionAccessor(const FPositionVertexBuffer* VertexBuffer, const FString& DesiredName)
+{
+	return PositionVertexBuffers.Convert(*this, DesiredName, VertexBuffer);
+}
+
+FGLTFJsonAccessorIndex FGLTFIndexedBuilder::ConvertColorAccessor(const FColorVertexBuffer* VertexBuffer, const FString& DesiredName)
+{
+	return ColorVertexBuffers.Convert(*this, DesiredName, VertexBuffer);
+}
+
+FGLTFJsonAccessorIndex FGLTFIndexedBuilder::ConvertNormalAccessor(const FStaticMeshVertexBuffer* VertexBuffer, const FString& DesiredName)
+{
+	return StaticMeshNormalVertexBuffers.Convert(*this, DesiredName, VertexBuffer);
+}
+
+FGLTFJsonAccessorIndex FGLTFIndexedBuilder::ConvertTangentAccessor(const FStaticMeshVertexBuffer* VertexBuffer, const FString& DesiredName)
+{
+	return StaticMeshTangentVertexBuffers.Convert(*this, DesiredName, VertexBuffer);
+}
+
+FGLTFJsonAccessorIndex FGLTFIndexedBuilder::ConvertUV0Accessor(const FStaticMeshVertexBuffer* VertexBuffer, const FString& DesiredName)
+{
+	return StaticMeshUV0VertexBuffers.Convert(*this, DesiredName, VertexBuffer);
+}
+
+FGLTFJsonAccessorIndex FGLTFIndexedBuilder::ConvertUV1Accessor(const FStaticMeshVertexBuffer* VertexBuffer, const FString& DesiredName)
+{
+	return StaticMeshUV1VertexBuffers.Convert(*this, DesiredName, VertexBuffer);
+}
+
+FGLTFJsonBufferViewIndex FGLTFIndexedBuilder::ConvertIndexBufferView(const FRawStaticIndexBuffer* IndexBuffer, const FString& DesiredName)
+{
+	return StaticMeshIndexBuffers.Convert(*this, DesiredName, IndexBuffer);
+}
+
+FGLTFJsonAccessorIndex FGLTFIndexedBuilder::ConvertIndexAccessor(const FStaticMeshSection* MeshSection, const FRawStaticIndexBuffer* IndexBuffer, const FString& DesiredName)
+{
+	return StaticMeshSections.Convert(*this, DesiredName, MeshSection, IndexBuffer);
+}
+
+FGLTFJsonMeshIndex FGLTFIndexedBuilder::ConvertMesh(const FStaticMeshLODResources* StaticMeshLOD, const FColorVertexBuffer* OverrideVertexColors, const FString& DesiredName)
+{
+	return StaticMeshes.Convert(*this, DesiredName, StaticMeshLOD, OverrideVertexColors);
+}
+
+FGLTFJsonMeshIndex FGLTFIndexedBuilder::ConvertMesh(const UStaticMesh* StaticMesh, int32 LODIndex, const FColorVertexBuffer* OverrideVertexColors, const FString& DesiredName)
+{
+	const FStaticMeshLODResources* StaticMeshLOD = &StaticMesh->GetLODForExport(LODIndex);
+	return ConvertMesh(StaticMeshLOD, OverrideVertexColors, DesiredName.IsEmpty() ? StaticMesh->GetName() + TEXT("_LOD") + FString::FromInt(LODIndex) : DesiredName);
+}
+
+FGLTFJsonMeshIndex FGLTFIndexedBuilder::ConvertMesh(const UStaticMeshComponent* StaticMeshComponent, const FString& DesiredName)
+{
+	const UStaticMesh* StaticMesh = StaticMeshComponent->GetStaticMesh();
+	const int32 LODIndex = StaticMeshComponent->ForcedLodModel > 0 ? StaticMeshComponent->ForcedLodModel - 1 : /* auto-select */ 0;
+	const FColorVertexBuffer* OverrideVertexColors = LODIndex < StaticMeshComponent->LODData.Num() ? StaticMeshComponent->LODData[LODIndex].OverrideVertexColors : nullptr;
+	return ConvertMesh(StaticMesh, LODIndex, OverrideVertexColors, DesiredName);
+}
+
+FGLTFJsonNodeIndex FGLTFIndexedBuilder::ConvertNode(const USceneComponent* SceneComponent, bool bSelectedOnly, bool bRootNode, const FString& DesiredName)
+{
+	return SceneComponents.Convert(*this, DesiredName, SceneComponent, bSelectedOnly, bRootNode);
+}
+
+FGLTFJsonSceneIndex FGLTFIndexedBuilder::ConvertScene(const ULevel* Level, bool bSelectedOnly, const FString& DesiredName)
+{
+	return Levels.Convert(*this, DesiredName, Level, bSelectedOnly);
+}
+
+FGLTFJsonSceneIndex FGLTFIndexedBuilder::ConvertScene(const UWorld* World, bool bSelectedOnly, const FString& DesiredName)
+{
+	const ULevel* Level = World->PersistentLevel;
+	return ConvertScene(Level, bSelectedOnly, DesiredName.IsEmpty() ? World->GetName() : DesiredName);
+}
diff --git a/Engine/Plugins/Enterprise/GLTFExporter/MigratedFromGit/Source/GLTFExporter/Private/GLTFJsonBuilder.cpp b/Engine/Plugins/Enterprise/GLTFExporter/MigratedFromGit/Source/GLTFExporter/Private/GLTFJsonBuilder.cpp
new file mode 100644
--- /dev/null
+++ b/Engine/Plugins/Enterprise/GLTFExporter/MigratedFromGit/Source/GLTFExporter/Private/GLTFJsonBuilder.cpp
@@ -0,0 +1,46 @@
+// Copyright Epic Games, Inc. All Rights Reserved.
+
+#include "GLTFJsonBuilder.h"
+
+FGLTFJsonBuilder::FGLTFJsonBuilder()
+{
+}
+
+FGLTFJsonBuilder::~FGLTFJsonBuilder()
+{
+}
+
+FGLTFJsonAccessorIndex FGLTFJsonBuilder::AddAccessor(const FGLTFJsonAccessor& JsonAccessor)
+{
+	return FGLTFJsonAccessorIndex(JsonRoot.Accessors.Add(JsonAccessor));
+}
+
+FGLTFJsonBufferIndex FGLTFJsonBuilder::AddBuffer(const FGLTFJsonBuffer& JsonBuffer)
+{
+	return FGLTFJsonBufferIndex(JsonRoot.Buffers.Add(JsonBuffer));
+}
+
+FGLTFJsonBufferViewIndex FGLTFJsonBuilder::AddBufferView(const FGLTFJsonBufferView& JsonBufferView)
+{
+	return FGLTFJsonBufferViewIndex(JsonRoot.BufferViews.Add(JsonBufferView));
+}
+
+FGLTFJsonMeshIndex FGLTFJsonBuilder::AddMesh(const FGLTFJsonMesh& JsonMesh)
+{
+	return FGLTFJsonMeshIndex(JsonRoot.Meshes.Add(JsonMesh));
+}
+
+FGLTFJsonNodeIndex FGLTFJsonBuilder::AddNode(const FGLTFJsonNode& JsonNode)
+{
+	return FGLTFJsonNodeIndex(JsonRoot.Nodes.Add(JsonNode));
+}
+
+FGLTFJsonSceneIndex FGLTFJsonBuilder::AddScene(const FGLTFJsonScene& JsonScene)
+{
+	return FGLTFJsonSceneIndex(JsonRoot.Scenes.Add(JsonScene));
+}
+
+void FGLTFJsonBuilder::Serialize(FArchive& Archive)
+{
+	JsonRoot.Serialize(&Archive, true);
+}
diff --git a/Engine/Plugins/Enterprise/GLTFExporter/MigratedFromGit/Source/GLTFExporter/Private/GLTFLevelConverters.cpp b/Engine/Plugins/Enterprise/GLTFExporter/MigratedFromGit/Source/GLTFExporter/Private/GLTFLevelConverters.cpp
--- a/Engine/Plugins/Enterprise/GLTFExporter/MigratedFromGit/Source/GLTFExporter/Private/GLTFLevelConverters.cpp
+++ b/Engine/Plugins/Enterprise/GLTFExporter/MigratedFromGit/Source/GLTFExporter/Private/GLTFLevelConverters.cpp
@@ -6,7 +6,7 @@
 #include "GLTFContainerBuilder.h"
 #include "GLTFConversionUtilities.h"
 
-FGLTFJsonNodeIndex FGLTFSceneComponentConverter::Convert(FGLTFContainerBuilder& Container, const FString& Name, const USceneComponent* SceneComponent, bool bSelectedOnly, bool bRootNode)
+FGLTFJsonNodeIndex FGLTFSceneComponentConverter::Convert(FGLTFIndexedBuilder& Builder, const FString& Name, const USceneComponent* SceneComponent, bool bSelectedOnly, bool bRootNode)
 {
 	const AActor* Owner = SceneComponent->GetOwner();
 	const bool bIsRootComponent = Owner != nullptr && Owner->GetRootComponent() == SceneComponent;
@@ -26,7 +26,7 @@ FGLTFJsonNodeIndex FGLTFSceneComponentConverter::Convert(FGLTFContainerBuilder&
 	}
 	else if (const UStaticMeshComponent* StaticMeshComponent = Cast<UStaticMeshComponent>(SceneComponent))
 	{
-		Node.Mesh = Container.ConvertMesh(StaticMeshComponent);
+		Node.Mesh = Builder.ConvertMesh(StaticMeshComponent);
 	}
 	else if (IsHDRIBackdropBlueprint(Blueprint) && bIsRootComponent)
 	{
@@ -41,16 +41,16 @@ FGLTFJsonNodeIndex FGLTFSceneComponentConverter::Convert(FGLTFContainerBuilder&
 			const AActor* ChildOwner = ChildComponent->GetOwner();
 			if (!bSelectedOnly || ChildOwner->IsSelected())
 			{
-				FGLTFJsonNodeIndex NodeIndex = Container.ConvertNode(ChildComponent, bSelectedOnly, false);
+				FGLTFJsonNodeIndex NodeIndex = Builder.ConvertNode(ChildComponent, bSelectedOnly, false);
 				if (NodeIndex != INDEX_NONE) Node.Children.Add(NodeIndex);
 			}
 		}
 	}
 
-	return Container.AddNode(Node);
+	return Builder.AddNode(Node);
 }
 
-FGLTFJsonSceneIndex FGLTFLevelConverter::Convert(FGLTFContainerBuilder& Container, const FString& Name, const ULevel* Level, bool bSelectedOnly)
+FGLTFJsonSceneIndex FGLTFLevelConverter::Convert(FGLTFIndexedBuilder& Builder, const FString& Name, const ULevel* Level, bool bSelectedOnly)
 {
 	FGLTFJsonScene Scene;
 	Scene.Name = Name;
@@ -68,12 +68,12 @@ FGLTFJsonSceneIndex FGLTFLevelConverter::Convert(FGLTFContainerBuilder& Containe
 				const AActor* ParentActor = Actor->GetParentActor();
 				if (ParentActor == nullptr || (bSelectedOnly && !ParentActor->IsSelected()))
 				{
-					FGLTFJsonNodeIndex NodeIndex = Container.ConvertNode(RootComponent, bSelectedOnly, true);
+					FGLTFJsonNodeIndex NodeIndex = Builder.ConvertNode(RootComponent, bSelectedOnly, true);
 					if (NodeIndex != INDEX_NONE) Scene.Nodes.Add(NodeIndex);
 				}
 			}
 		}
 	}
 
-	return Container.AddScene(Scene);
+	return Builder.AddScene(Scene);
 }
diff --git a/Engine/Plugins/Enterprise/GLTFExporter/MigratedFromGit/Source/GLTFExporter/Private/GLTFLevelExporter.cpp b/Engine/Plugins/Enterprise/GLTFExporter/MigratedFromGit/Source/GLTFExporter/Private/GLTFLevelExporter.cpp
--- a/Engine/Plugins/Enterprise/GLTFExporter/MigratedFromGit/Source/GLTFExporter/Private/GLTFLevelExporter.cpp
+++ b/Engine/Plugins/Enterprise/GLTFExporter/MigratedFromGit/Source/GLTFExporter/Private/GLTFLevelExporter.cpp
@@ -20,10 +20,10 @@ bool UGLTFLevelExporter::ExportBinary(UObject* Object, const TCHAR* Type, FArchi
 		return false;
 	}
 
-	FGLTFContainerBuilder Container;
-	const FGLTFJsonSceneIndex SceneIndex = Container.ConvertScene(World, bSelectedOnly);
-	Container.JsonRoot.DefaultScene = SceneIndex;
+	FGLTFContainerBuilder Builder;
+	const FGLTFJsonSceneIndex SceneIndex = Builder.ConvertScene(World, bSelectedOnly);
+	Builder.JsonRoot.DefaultScene = SceneIndex;
 
-	Container.Serialize(Archive);
+	Builder.Serialize(Archive);
 	return true;
 }
diff --git a/Engine/Plugins/Enterprise/GLTFExporter/MigratedFromGit/Source/GLTFExporter/Private/GLTFStaticMeshConverters.cpp b/Engine/Plugins/Enterprise/GLTFExporter/MigratedFromGit/Source/GLTFExporter/Private/GLTFStaticMeshConverters.cpp
--- a/Engine/Plugins/Enterprise/GLTFExporter/MigratedFromGit/Source/GLTFExporter/Private/GLTFStaticMeshConverters.cpp
+++ b/Engine/Plugins/Enterprise/GLTFExporter/MigratedFromGit/Source/GLTFExporter/Private/GLTFStaticMeshConverters.cpp
@@ -4,7 +4,7 @@
 #include "GLTFContainerBuilder.h"
 #include "GLTFConversionUtilities.h"
 
-FGLTFJsonAccessorIndex FGLTFPositionVertexBufferConverter::Convert(FGLTFContainerBuilder& Container, const FString& Name, const FPositionVertexBuffer* VertexBuffer)
+FGLTFJsonAccessorIndex FGLTFPositionVertexBufferConverter::Convert(FGLTFIndexedBuilder& Builder, const FString& Name, const FPositionVertexBuffer* VertexBuffer)
 {
 	const uint32 VertexCount = VertexBuffer->GetNumVertices();
 	if (VertexCount == 0)
@@ -30,7 +30,7 @@ FGLTFJsonAccessorIndex FGLTFPositionVertexBufferConverter::Convert(FGLTFContaine
 
 	FGLTFJsonAccessor JsonAccessor;
 	JsonAccessor.Name = Name;
-	JsonAccessor.BufferView = Container.AddBufferView(Positions, Name);
+	JsonAccessor.BufferView = Builder.AddBufferView(Positions, Name);
 	JsonAccessor.ComponentType = EGLTFJsonComponentType::F32;
 	JsonAccessor.Count = VertexCount;
 	JsonAccessor.Type = EGLTFJsonAccessorType::Vec3;
@@ -43,10 +43,10 @@ FGLTFJsonAccessorIndex FGLTFPositionVertexBufferConverter::Convert(FGLTFContaine
 	JsonAccessor.Min[1] = BoundingBox.Min.Y;
 	JsonAccessor.Min[2] = BoundingBox.Min.Z;
 	
-	return Container.AddAccessor(JsonAccessor);
+	return Builder.AddAccessor(JsonAccessor);
 }
 
-FGLTFJsonAccessorIndex FGLTFColorVertexBufferConverter::Convert(FGLTFContainerBuilder& Container, const FString& Name, const FColorVertexBuffer* VertexBuffer)
+FGLTFJsonAccessorIndex FGLTFColorVertexBufferConverter::Convert(FGLTFIndexedBuilder& Builder, const FString& Name, const FColorVertexBuffer* VertexBuffer)
 {
 	const uint32 VertexCount = VertexBuffer->GetNumVertices();
 	if (VertexCount == 0)
@@ -64,16 +64,16 @@ FGLTFJsonAccessorIndex FGLTFColorVertexBufferConverter::Convert(FGLTFContainerBu
 
 	FGLTFJsonAccessor JsonAccessor;
 	JsonAccessor.Name = Name;
-	JsonAccessor.BufferView = Container.AddBufferView(Colors, Name);
+	JsonAccessor.BufferView = Builder.AddBufferView(Colors, Name);
 	JsonAccessor.ComponentType = EGLTFJsonComponentType::U8;
 	JsonAccessor.Count = VertexCount;
 	JsonAccessor.Type = EGLTFJsonAccessorType::Vec4;
 	JsonAccessor.bNormalized = true;
 
-	return Container.AddAccessor(JsonAccessor);
+	return Builder.AddAccessor(JsonAccessor);
 }
 
-FGLTFJsonAccessorIndex FGLTFStaticMeshNormalVertexBufferConverter::Convert(FGLTFContainerBuilder& Container, const FString& Name, const FStaticMeshVertexBuffer* VertexBuffer)
+FGLTFJsonAccessorIndex FGLTFStaticMeshNormalVertexBufferConverter::Convert(FGLTFIndexedBuilder& Builder, const FString& Name, const FStaticMeshVertexBuffer* VertexBuffer)
 {
 	const uint32 VertexCount = VertexBuffer->GetNumVertices();
 	if (VertexCount == 0)
@@ -91,15 +91,15 @@ FGLTFJsonAccessorIndex FGLTFStaticMeshNormalVertexBufferConverter::Convert(FGLTF
 
 	FGLTFJsonAccessor JsonAccessor;
 	JsonAccessor.Name = Name;
-	JsonAccessor.BufferView = Container.AddBufferView(Normals, Name);
+	JsonAccessor.BufferView = Builder.AddBufferView(Normals, Name);
 	JsonAccessor.ComponentType = EGLTFJsonComponentType::F32;
 	JsonAccessor.Count = VertexCount;
 	JsonAccessor.Type = EGLTFJsonAccessorType::Vec3;
 	
-	return Container.AddAccessor(JsonAccessor);
+	return Builder.AddAccessor(JsonAccessor);
 }
 
-FGLTFJsonAccessorIndex FGLTFStaticMeshTangentVertexBufferConverter::Convert(FGLTFContainerBuilder& Container, const FString& Name, const FStaticMeshVertexBuffer* VertexBuffer)
+FGLTFJsonAccessorIndex FGLTFStaticMeshTangentVertexBufferConverter::Convert(FGLTFIndexedBuilder& Builder, const FString& Name, const FStaticMeshVertexBuffer* VertexBuffer)
 {
 	const uint32 VertexCount = VertexBuffer->GetNumVertices();
 	if (VertexCount == 0)
@@ -117,15 +117,15 @@ FGLTFJsonAccessorIndex FGLTFStaticMeshTangentVertexBufferConverter::Convert(FGLT
 
 	FGLTFJsonAccessor JsonAccessor;
 	JsonAccessor.Name = Name;
-	JsonAccessor.BufferView = Container.AddBufferView(Tangents, Name);
+	JsonAccessor.BufferView = Builder.AddBufferView(Tangents, Name);
 	JsonAccessor.ComponentType = EGLTFJsonComponentType::F32;
 	JsonAccessor.Count = VertexCount;
 	JsonAccessor.Type = EGLTFJsonAccessorType::Vec4;
 	
-	return Container.AddAccessor(JsonAccessor);
+	return Builder.AddAccessor(JsonAccessor);
 }
 
-FGLTFJsonAccessorIndex FGLTFStaticMeshUV0VertexBufferConverter::Convert(FGLTFContainerBuilder& Container, const FString& Name, const FStaticMeshVertexBuffer* VertexBuffer)
+FGLTFJsonAccessorIndex FGLTFStaticMeshUV0VertexBufferConverter::Convert(FGLTFIndexedBuilder& Builder, const FString& Name, const FStaticMeshVertexBuffer* VertexBuffer)
 {
 	const uint32 VertexCount = VertexBuffer->GetNumVertices();
 	if (VertexCount == 0 || VertexBuffer->GetNumTexCoords() < 1)
@@ -143,15 +143,15 @@ FGLTFJsonAccessorIndex FGLTFStaticMeshUV0VertexBufferConverter::Convert(FGLTFCon
 
 	FGLTFJsonAccessor JsonAccessor;
 	JsonAccessor.Name = Name;
-	JsonAccessor.BufferView = Container.AddBufferView(UV0s, Name);
+	JsonAccessor.BufferView = Builder.AddBufferView(UV0s, Name);
 	JsonAccessor.ComponentType = EGLTFJsonComponentType::F32;
 	JsonAccessor.Count = VertexCount;
 	JsonAccessor.Type = EGLTFJsonAccessorType::Vec2;
 
-	return Container.AddAccessor(JsonAccessor);
+	return Builder.AddAccessor(JsonAccessor);
 }
 
-FGLTFJsonAccessorIndex FGLTFStaticMeshUV1VertexBufferConverter::Convert(FGLTFContainerBuilder& Container, const FString& Name, const FStaticMeshVertexBuffer* VertexBuffer)
+FGLTFJsonAccessorIndex FGLTFStaticMeshUV1VertexBufferConverter::Convert(FGLTFIndexedBuilder& Builder, const FString& Name, const FStaticMeshVertexBuffer* VertexBuffer)
 {
 	const uint32 VertexCount = VertexBuffer->GetNumVertices();
 	if (VertexCount == 0 || VertexBuffer->GetNumTexCoords() < 2)
@@ -169,15 +169,15 @@ FGLTFJsonAccessorIndex FGLTFStaticMeshUV1VertexBufferConverter::Convert(FGLTFCon
 
 	FGLTFJsonAccessor JsonAccessor;
 	JsonAccessor.Name = Name;
-	JsonAccessor.BufferView = Container.AddBufferView(UV1s, Name);
+	JsonAccessor.BufferView = Builder.AddBufferView(UV1s, Name);
 	JsonAccessor.ComponentType = EGLTFJsonComponentType::F32;
 	JsonAccessor.Count = VertexCount;
 	JsonAccessor.Type = EGLTFJsonAccessorType::Vec2;
 
-	return Container.AddAccessor(JsonAccessor);
+	return Builder.AddAccessor(JsonAccessor);
 }
 
-FGLTFJsonBufferViewIndex FGLTFStaticMeshIndexBufferConverter::Convert(FGLTFContainerBuilder& Container, const FString& Name, const FRawStaticIndexBuffer* IndexBuffer)
+FGLTFJsonBufferViewIndex FGLTFStaticMeshIndexBufferConverter::Convert(FGLTFIndexedBuilder& Builder, const FString& Name, const FRawStaticIndexBuffer* IndexBuffer)
 {
 	if (IndexBuffer->GetNumIndices() == 0)
 	{
@@ -186,10 +186,10 @@ FGLTFJsonBufferViewIndex FGLTFStaticMeshIndexBufferConverter::Convert(FGLTFConta
 
 	TArray<uint32> Indices;
 	IndexBuffer->GetCopy(Indices);
-	return Container.CreateBufferView(Indices, Name, EGLTFJsonBufferTarget::ElementArrayBuffer);
+	return Builder.AddBufferView(Indices, Name, EGLTFJsonBufferTarget::ElementArrayBuffer);
 }
 
-FGLTFJsonAccessorIndex FGLTFStaticMeshSectionConverter::Convert(FGLTFContainerBuilder& Container, const FString& Name, const FStaticMeshSection* MeshSection, const FRawStaticIndexBuffer* IndexBuffer)
+FGLTFJsonAccessorIndex FGLTFStaticMeshSectionConverter::Convert(FGLTFIndexedBuilder& Builder, const FString& Name, const FStaticMeshSection* MeshSection, const FRawStaticIndexBuffer* IndexBuffer)
 {	
 	const uint32 TriangleCount = MeshSection->NumTriangles;
 	if (TriangleCount == 0)
@@ -199,16 +199,16 @@ FGLTFJsonAccessorIndex FGLTFStaticMeshSectionConverter::Convert(FGLTFContainerBu
 
 	FGLTFJsonAccessor JsonAccessor;
 	JsonAccessor.Name = Name;
-	JsonAccessor.BufferView = Container.AddIndexBufferView(IndexBuffer);
+	JsonAccessor.BufferView = Builder.ConvertIndexBufferView(IndexBuffer);
 	JsonAccessor.ByteOffset = MeshSection->FirstIndex * sizeof(uint32);
 	JsonAccessor.ComponentType = EGLTFJsonComponentType::U32;
 	JsonAccessor.Count = TriangleCount * 3;
 	JsonAccessor.Type = EGLTFJsonAccessorType::Scalar;
 
-	return Container.AddAccessor(JsonAccessor);
+	return Builder.AddAccessor(JsonAccessor);
 }
 
-FGLTFJsonMeshIndex FGLTFStaticMeshConverter::Convert(FGLTFContainerBuilder& Container, const FString& Name, const FStaticMeshLODResources* StaticMeshLOD, const FColorVertexBuffer* OverrideVertexColors)
+FGLTFJsonMeshIndex FGLTFStaticMeshConverter::Convert(FGLTFIndexedBuilder& Builder, const FString& Name, const FStaticMeshLODResources* StaticMeshLOD, const FColorVertexBuffer* OverrideVertexColors)
 {
 	FGLTFJsonMesh JsonMesh;
 	JsonMesh.Name = Name;
@@ -218,15 +218,15 @@ FGLTFJsonMeshIndex FGLTFStaticMeshConverter::Convert(FGLTFContainerBuilder& Cont
 	const FColorVertexBuffer* ColorBuffer = OverrideVertexColors != nullptr ? OverrideVertexColors : &StaticMeshLOD->VertexBuffers.ColorVertexBuffer;
 
 	FGLTFJsonAttributes JsonAttributes;
-	JsonAttributes.Position = Container.ConvertPositionAccessor(PositionBuffer, Name.IsEmpty() ? Name : Name + TEXT("_Positions"));
-	JsonAttributes.Normal = Container.ConvertNormalAccessor(VertexBuffer, Name.IsEmpty() ? Name : Name + TEXT("_Normals"));
-	JsonAttributes.Tangent = Container.ConvertTangentAccessor(VertexBuffer, Name.IsEmpty() ? Name : Name + TEXT("_Tangents"));
-	JsonAttributes.TexCoord0 = Container.ConvertUV0Accessor(VertexBuffer, Name.IsEmpty() ? Name : Name + TEXT("_UV0s"));
-	JsonAttributes.TexCoord1 = Container.ConvertUV1Accessor(VertexBuffer, Name.IsEmpty() ? Name : Name + TEXT("_UV1s"));
-	JsonAttributes.Color0 = Container.ConvertColorAccessor(ColorBuffer, Name.IsEmpty() ? Name : Name + TEXT("_Colors"));
+	JsonAttributes.Position = Builder.ConvertPositionAccessor(PositionBuffer, Name.IsEmpty() ? Name : Name + TEXT("_Positions"));
+	JsonAttributes.Normal = Builder.ConvertNormalAccessor(VertexBuffer, Name.IsEmpty() ? Name : Name + TEXT("_Normals"));
+	JsonAttributes.Tangent = Builder.ConvertTangentAccessor(VertexBuffer, Name.IsEmpty() ? Name : Name + TEXT("_Tangents"));
+	JsonAttributes.TexCoord0 = Builder.ConvertUV0Accessor(VertexBuffer, Name.IsEmpty() ? Name : Name + TEXT("_UV0s"));
+	JsonAttributes.TexCoord1 = Builder.ConvertUV1Accessor(VertexBuffer, Name.IsEmpty() ? Name : Name + TEXT("_UV1s"));
+	JsonAttributes.Color0 = Builder.ConvertColorAccessor(ColorBuffer, Name.IsEmpty() ? Name : Name + TEXT("_Colors"));
 
 	const FRawStaticIndexBuffer* IndexBuffer = &StaticMeshLOD->IndexBuffer;
-	Container.ConvertIndexBufferView(IndexBuffer, Name.IsEmpty() ? Name : Name + TEXT("_Indices"));
+	Builder.ConvertIndexBufferView(IndexBuffer, Name.IsEmpty() ? Name : Name + TEXT("_Indices"));
 
 	const TArray<FStaticMeshSection>& Sections = StaticMeshLOD->Sections;
 	const int32 SectionCount = Sections.Num();
@@ -237,9 +237,9 @@ FGLTFJsonMeshIndex FGLTFStaticMeshConverter::Convert(FGLTFContainerBuilder& Cont
 		FGLTFJsonPrimitive& JsonPrimitive = JsonMesh.Primitives[SectionIndex];
 		JsonPrimitive.Attributes = JsonAttributes;
 
-		JsonPrimitive.Indices = Container.ConvertIndexAccessor(&Sections[SectionIndex], IndexBuffer,
+		JsonPrimitive.Indices = Builder.ConvertIndexAccessor(&Sections[SectionIndex], IndexBuffer,
 			Name.IsEmpty() ? Name : JsonMesh.Name + (SectionCount != 0 ? TEXT("_Indices_Section") + FString::FromInt(SectionIndex) : TEXT("_Indices")));
 	}
 
-	return Container.AddMesh(JsonMesh);
+	return Builder.AddMesh(JsonMesh);
 }
diff --git a/Engine/Plugins/Enterprise/GLTFExporter/MigratedFromGit/Source/GLTFExporter/Private/GLTFStaticMeshExporter.cpp b/Engine/Plugins/Enterprise/GLTFExporter/MigratedFromGit/Source/GLTFExporter/Private/GLTFStaticMeshExporter.cpp
--- a/Engine/Plugins/Enterprise/GLTFExporter/MigratedFromGit/Source/GLTFExporter/Private/GLTFStaticMeshExporter.cpp
+++ b/Engine/Plugins/Enterprise/GLTFExporter/MigratedFromGit/Source/GLTFExporter/Private/GLTFStaticMeshExporter.cpp
@@ -21,20 +21,20 @@ bool UGLTFStaticMeshExporter::ExportBinary(UObject* Object, const TCHAR* Type, F
 		return false;
 	}
 
-	FGLTFContainerBuilder Container;
+	FGLTFContainerBuilder Builder;
 
-	FGLTFJsonMeshIndex MeshIndex = Container.ConvertMesh(StaticMesh);
+	FGLTFJsonMeshIndex MeshIndex = Builder.ConvertMesh(StaticMesh);
 
 	FGLTFJsonNode Node;
 	Node.Mesh = MeshIndex;
-	FGLTFJsonNodeIndex NodeIndex = Container.AddNode(Node);
+	FGLTFJsonNodeIndex NodeIndex = Builder.AddNode(Node);
 
 	FGLTFJsonScene Scene;
 	Scene.Nodes.Add(NodeIndex);
-	FGLTFJsonSceneIndex SceneIndex = Container.AddScene(Scene);
+	FGLTFJsonSceneIndex SceneIndex = Builder.AddScene(Scene);
 
-	Container.JsonRoot.DefaultScene = SceneIndex;
+	Builder.JsonRoot.DefaultScene = SceneIndex;
 
-	Container.Serialize(Archive);
+	Builder.Serialize(Archive);
 	return true;
 }
diff --git a/Engine/Plugins/Enterprise/GLTFExporter/MigratedFromGit/Source/GLTFExporter/Public/GLTFBufferBuilder.h b/Engine/Plugins/Enterprise/GLTFExporter/MigratedFromGit/Source/GLTFExporter/Public/GLTFBufferBuilder.h
--- a/Engine/Plugins/Enterprise/GLTFExporter/MigratedFromGit/Source/GLTFExporter/Public/GLTFBufferBuilder.h
+++ b/Engine/Plugins/Enterprise/GLTFExporter/MigratedFromGit/Source/GLTFExporter/Public/GLTFBufferBuilder.h
@@ -2,19 +2,23 @@
 
 #pragma once
 
-#include "GLTFJsonRoot.h"
-#include "Engine.h"
+#include "GLTFJsonBuilder.h"
 
-struct FGLTFContainerBuilder;
-
-struct GLTFEXPORTER_API FGLTFBufferBuilder
+struct GLTFEXPORTER_API FGLTFBufferBuilder : public FGLTFJsonBuilder
 {
 	FGLTFJsonBufferIndex BufferIndex;
 	TArray<uint8> BufferData;
 
-	FGLTFBufferBuilder(FGLTFJsonBufferIndex BufferIndex);
+	FGLTFBufferBuilder();
+
+	FGLTFJsonBufferViewIndex AddBufferView(const void* RawData, uint64 ByteLength, const FString& Name = TEXT(""), EGLTFJsonBufferTarget BufferTarget = EGLTFJsonBufferTarget::ArrayBuffer);
 
-	FGLTFJsonBufferViewIndex AddBufferView(FGLTFContainerBuilder& Container, const void* RawData, uint64 ByteLength, const FString& Name, EGLTFJsonBufferTarget BufferTarget);
+	template <class ElementType>
+    FGLTFJsonBufferViewIndex AddBufferView(const TArray<ElementType>& Array, const FString& Name = TEXT(""), EGLTFJsonBufferTarget BufferTarget = EGLTFJsonBufferTarget::ArrayBuffer)
+	{
+		return AddBufferView(Array.GetData(), Array.Num() * sizeof(ElementType), Name, BufferTarget);
+	}
 
-	void UpdateBuffer(FGLTFJsonBuffer& JsonBuffer);
+	void UpdateBuffer();
+	virtual void Serialize(FArchive& Archive) override;
 };
diff --git a/Engine/Plugins/Enterprise/GLTFExporter/MigratedFromGit/Source/GLTFExporter/Public/GLTFContainerBuilder.h b/Engine/Plugins/Enterprise/GLTFExporter/MigratedFromGit/Source/GLTFExporter/Public/GLTFContainerBuilder.h
--- a/Engine/Plugins/Enterprise/GLTFExporter/MigratedFromGit/Source/GLTFExporter/Public/GLTFContainerBuilder.h
+++ b/Engine/Plugins/Enterprise/GLTFExporter/MigratedFromGit/Source/GLTFExporter/Public/GLTFContainerBuilder.h
@@ -2,51 +2,8 @@
 
 #pragma once
 
-#include "GLTFJsonRoot.h"
-#include "GLTFBufferBuilder.h"
-#include "GLTFIndexedConverts.h"
-#include "Engine.h"
+#include "GLTFIndexedBuilder.h"
 
-struct GLTFEXPORTER_API FGLTFContainerBuilder
-{
-	FGLTFJsonRoot JsonRoot;
-	FGLTFBufferBuilder BufferBuilder;
-	FGLTFIndexedConverts IndexedConverts;
-
-	FGLTFContainerBuilder();
-
-	FGLTFJsonAccessorIndex AddAccessor(const FGLTFJsonAccessor& JsonAccessor);
-	FGLTFJsonBufferIndex AddBuffer(const FGLTFJsonBuffer& JsonBuffer);
-	FGLTFJsonBufferViewIndex AddBufferView(const FGLTFJsonBufferView& JsonBufferView);
-	FGLTFJsonMeshIndex AddMesh(const FGLTFJsonMesh& JsonMesh);
-	FGLTFJsonNodeIndex AddNode(const FGLTFJsonNode& JsonNode);
-	FGLTFJsonSceneIndex AddScene(const FGLTFJsonScene& JsonScene);
-
-	FGLTFJsonBufferViewIndex AddBufferView(const void* RawData, uint64 ByteLength, const FString& Name = TEXT(""), EGLTFJsonBufferTarget BufferTarget = EGLTFJsonBufferTarget::ArrayBuffer);
-
-	template <class ElementType>
-	FGLTFJsonBufferViewIndex AddBufferView(const TArray<ElementType>& Array, const FString& Name = TEXT(""), EGLTFJsonBufferTarget BufferTarget = EGLTFJsonBufferTarget::ArrayBuffer)
-	{
-		return AddBufferView(Array.GetData(), Array.Num() * sizeof(ElementType), Name, BufferTarget);
-	}
-
-	FGLTFJsonAccessorIndex ConvertPositionAccessor(const FPositionVertexBuffer* VertexBuffer, const FString& DesiredName = TEXT(""));
-	FGLTFJsonAccessorIndex ConvertNormalAccessor(const FStaticMeshVertexBuffer* VertexBuffer, const FString& DesiredName = TEXT(""));
-	FGLTFJsonAccessorIndex ConvertTangentAccessor(const FStaticMeshVertexBuffer* VertexBuffer, const FString& DesiredName = TEXT(""));
-	FGLTFJsonAccessorIndex ConvertUV0Accessor(const FStaticMeshVertexBuffer* VertexBuffer, const FString& DesiredName = TEXT(""));
-	FGLTFJsonAccessorIndex ConvertUV1Accessor(const FStaticMeshVertexBuffer* VertexBuffer, const FString& DesiredName = TEXT(""));
-	FGLTFJsonAccessorIndex ConvertColorAccessor(const FColorVertexBuffer* VertexBuffer, const FString& DesiredName = TEXT(""));
-
-	FGLTFJsonBufferViewIndex ConvertIndexBufferView(const FRawStaticIndexBuffer* IndexBuffer, const FString& DesiredName = TEXT(""));
-	FGLTFJsonAccessorIndex ConvertIndexAccessor(const FStaticMeshSection* MeshSection, const FRawStaticIndexBuffer* IndexBuffer, const FString& DesiredName = TEXT(""));
-
-	FGLTFJsonMeshIndex ConvertMesh(const FStaticMeshLODResources* StaticMeshLOD, const FColorVertexBuffer* OverrideVertexColors = nullptr, const FString& DesiredName = TEXT(""));
-	FGLTFJsonMeshIndex ConvertMesh(const UStaticMesh* StaticMesh, int32 LODIndex = 0, const FColorVertexBuffer* OverrideVertexColors = nullptr, const FString& DesiredName = TEXT(""));
-	FGLTFJsonMeshIndex ConvertMesh(const UStaticMeshComponent* StaticMeshComponent, const FString& DesiredName = TEXT(""));
-
-	FGLTFJsonNodeIndex ConvertNode(const USceneComponent* SceneComponent, bool bSelectedOnly, bool bRootNode = false, const FString& DesiredName = TEXT(""));
-	FGLTFJsonSceneIndex ConvertScene(const ULevel* Level, bool bSelectedOnly, const FString& DesiredName = TEXT(""));
-	FGLTFJsonSceneIndex ConvertScene(const UWorld* World, bool bSelectedOnly, const FString& DesiredName = TEXT(""));
-
-	void Serialize(FArchive& Archive);
+struct GLTFEXPORTER_API FGLTFContainerBuilder : public FGLTFIndexedBuilder
+{	
 };
diff --git a/Engine/Plugins/Enterprise/GLTFExporter/MigratedFromGit/Source/GLTFExporter/Public/GLTFIndexedBuilder.h b/Engine/Plugins/Enterprise/GLTFExporter/MigratedFromGit/Source/GLTFExporter/Public/GLTFIndexedBuilder.h
new file mode 100644
--- /dev/null
+++ b/Engine/Plugins/Enterprise/GLTFExporter/MigratedFromGit/Source/GLTFExporter/Public/GLTFIndexedBuilder.h
@@ -0,0 +1,42 @@
+// Copyright Epic Games, Inc. All Rights Reserved.
+
+#pragma once
+
+#include "GLTFBufferBuilder.h"
+#include "GLTFIndexedObjects.h"
+#include "GLTFStaticMeshConverters.h"
+#include "GLTFLevelConverters.h"
+
+struct GLTFEXPORTER_API FGLTFIndexedBuilder : public FGLTFBufferBuilder
+{
+	TGLTFIndexedObjects<FGLTFJsonAccessorIndex, TTuple<const FPositionVertexBuffer*>, FGLTFPositionVertexBufferConverter> PositionVertexBuffers;
+	TGLTFIndexedObjects<FGLTFJsonAccessorIndex, TTuple<const FColorVertexBuffer*>, FGLTFColorVertexBufferConverter> ColorVertexBuffers;
+	TGLTFIndexedObjects<FGLTFJsonAccessorIndex, TTuple<const FStaticMeshVertexBuffer*>, FGLTFStaticMeshNormalVertexBufferConverter> StaticMeshNormalVertexBuffers;
+	TGLTFIndexedObjects<FGLTFJsonAccessorIndex, TTuple<const FStaticMeshVertexBuffer*>, FGLTFStaticMeshTangentVertexBufferConverter> StaticMeshTangentVertexBuffers;
+	TGLTFIndexedObjects<FGLTFJsonAccessorIndex, TTuple<const FStaticMeshVertexBuffer*>, FGLTFStaticMeshUV0VertexBufferConverter> StaticMeshUV0VertexBuffers;
+	TGLTFIndexedObjects<FGLTFJsonAccessorIndex, TTuple<const FStaticMeshVertexBuffer*>, FGLTFStaticMeshUV1VertexBufferConverter> StaticMeshUV1VertexBuffers;
+	TGLTFIndexedObjects<FGLTFJsonBufferViewIndex, TTuple<const FRawStaticIndexBuffer*>, FGLTFStaticMeshIndexBufferConverter> StaticMeshIndexBuffers;
+	TGLTFIndexedObjects<FGLTFJsonAccessorIndex, TTuple<const FStaticMeshSection*, const FRawStaticIndexBuffer*>, FGLTFStaticMeshSectionConverter> StaticMeshSections;
+	TGLTFIndexedObjects<FGLTFJsonMeshIndex, TTuple<const FStaticMeshLODResources*, const FColorVertexBuffer*>, FGLTFStaticMeshConverter> StaticMeshes;
+
+	TGLTFIndexedObjects<FGLTFJsonNodeIndex, TTuple<const USceneComponent*, bool, bool>, FGLTFSceneComponentConverter> SceneComponents;
+	TGLTFIndexedObjects<FGLTFJsonSceneIndex, TTuple<const ULevel*, bool>, FGLTFLevelConverter> Levels;
+
+	FGLTFJsonAccessorIndex ConvertPositionAccessor(const FPositionVertexBuffer* VertexBuffer, const FString& DesiredName = TEXT(""));
+	FGLTFJsonAccessorIndex ConvertNormalAccessor(const FStaticMeshVertexBuffer* VertexBuffer, const FString& DesiredName = TEXT(""));
+	FGLTFJsonAccessorIndex ConvertTangentAccessor(const FStaticMeshVertexBuffer* VertexBuffer, const FString& DesiredName = TEXT(""));
+	FGLTFJsonAccessorIndex ConvertUV0Accessor(const FStaticMeshVertexBuffer* VertexBuffer, const FString& DesiredName = TEXT(""));
+	FGLTFJsonAccessorIndex ConvertUV1Accessor(const FStaticMeshVertexBuffer* VertexBuffer, const FString& DesiredName = TEXT(""));
+	FGLTFJsonAccessorIndex ConvertColorAccessor(const FColorVertexBuffer* VertexBuffer, const FString& DesiredName = TEXT(""));
+
+	FGLTFJsonBufferViewIndex ConvertIndexBufferView(const FRawStaticIndexBuffer* IndexBuffer, const FString& DesiredName = TEXT(""));
+	FGLTFJsonAccessorIndex ConvertIndexAccessor(const FStaticMeshSection* MeshSection, const FRawStaticIndexBuffer* IndexBuffer, const FString& DesiredName = TEXT(""));
+
+	FGLTFJsonMeshIndex ConvertMesh(const FStaticMeshLODResources* StaticMeshLOD, const FColorVertexBuffer* OverrideVertexColors = nullptr, const FString& DesiredName = TEXT(""));
+	FGLTFJsonMeshIndex ConvertMesh(const UStaticMesh* StaticMesh, int32 LODIndex = 0, const FColorVertexBuffer* OverrideVertexColors = nullptr, const FString& DesiredName = TEXT(""));
+	FGLTFJsonMeshIndex ConvertMesh(const UStaticMeshComponent* StaticMeshComponent, const FString& DesiredName = TEXT(""));
+
+	FGLTFJsonNodeIndex ConvertNode(const USceneComponent* SceneComponent, bool bSelectedOnly, bool bRootNode = false, const FString& DesiredName = TEXT(""));
+	FGLTFJsonSceneIndex ConvertScene(const ULevel* Level, bool bSelectedOnly, const FString& DesiredName = TEXT(""));
+	FGLTFJsonSceneIndex ConvertScene(const UWorld* World, bool bSelectedOnly, const FString& DesiredName = TEXT(""));
+};
diff --git a/Engine/Plugins/Enterprise/GLTFExporter/MigratedFromGit/Source/GLTFExporter/Public/GLTFIndexedConverts.h b/Engine/Plugins/Enterprise/GLTFExporter/MigratedFromGit/Source/GLTFExporter/Public/GLTFIndexedConverts.h
deleted file mode 100644
diff --git a/Engine/Plugins/Enterprise/GLTFExporter/MigratedFromGit/Source/GLTFExporter/Public/GLTFIndexedObjects.h b/Engine/Plugins/Enterprise/GLTFExporter/MigratedFromGit/Source/GLTFExporter/Public/GLTFIndexedObjects.h
new file mode 100644
--- /dev/null
+++ b/Engine/Plugins/Enterprise/GLTFExporter/MigratedFromGit/Source/GLTFExporter/Public/GLTFIndexedObjects.h
@@ -0,0 +1,35 @@
+// Copyright Epic Games, Inc. All Rights Reserved.
+
+#pragma once
+
+#include "CoreMinimal.h"
+
+struct FGLTFIndexedBuilder;
+
+template <class IndexType, class KeyType, class ConverterType>
+struct GLTFEXPORTER_API TGLTFIndexedObjects
+{
+	TMap<KeyType, IndexType> IndexLookup;
+	
+	template <typename... ArgsType>
+    FORCEINLINE IndexType Find(ArgsType&&... Args) const
+	{
+		const KeyType Key(Forward<ArgsType>(Args)...);
+		return IndexLookup.FindRef(Key);
+	}
+
+	template <typename... ArgsType>
+    FORCEINLINE IndexType Convert(FGLTFIndexedBuilder& Builder, const FString& DesiredName, ArgsType&&... Args)
+	{
+		const KeyType Key(Forward<ArgsType>(Args)...);
+
+		IndexType Index = IndexLookup.FindRef(Key);
+		if (Index == INDEX_NONE)
+		{
+			Index = ConverterType::Convert(Builder, DesiredName, Forward<ArgsType>(Args)...);
+			IndexLookup.Add(Key, Index);
+		}
+
+		return Index;
+	}
+};
diff --git a/Engine/Plugins/Enterprise/GLTFExporter/MigratedFromGit/Source/GLTFExporter/Public/GLTFJsonBuilder.h b/Engine/Plugins/Enterprise/GLTFExporter/MigratedFromGit/Source/GLTFExporter/Public/GLTFJsonBuilder.h
new file mode 100644
--- /dev/null
+++ b/Engine/Plugins/Enterprise/GLTFExporter/MigratedFromGit/Source/GLTFExporter/Public/GLTFJsonBuilder.h
@@ -0,0 +1,22 @@
+// Copyright Epic Games, Inc. All Rights Reserved.
+
+#pragma once
+
+#include "GLTFJsonRoot.h"
+
+struct GLTFEXPORTER_API FGLTFJsonBuilder
+{
+	FGLTFJsonRoot JsonRoot;
+
+	FGLTFJsonBuilder();
+	virtual ~FGLTFJsonBuilder();
+
+	FGLTFJsonAccessorIndex AddAccessor(const FGLTFJsonAccessor& JsonAccessor);
+	FGLTFJsonBufferIndex AddBuffer(const FGLTFJsonBuffer& JsonBuffer);
+	FGLTFJsonBufferViewIndex AddBufferView(const FGLTFJsonBufferView& JsonBufferView);
+	FGLTFJsonMeshIndex AddMesh(const FGLTFJsonMesh& JsonMesh);
+	FGLTFJsonNodeIndex AddNode(const FGLTFJsonNode& JsonNode);
+	FGLTFJsonSceneIndex AddScene(const FGLTFJsonScene& JsonScene);
+
+	virtual void Serialize(FArchive& Archive);
+};
diff --git a/Engine/Plugins/Enterprise/GLTFExporter/MigratedFromGit/Source/GLTFExporter/Public/GLTFLevelConverters.h b/Engine/Plugins/Enterprise/GLTFExporter/MigratedFromGit/Source/GLTFExporter/Public/GLTFLevelConverters.h
--- a/Engine/Plugins/Enterprise/GLTFExporter/MigratedFromGit/Source/GLTFExporter/Public/GLTFLevelConverters.h
+++ b/Engine/Plugins/Enterprise/GLTFExporter/MigratedFromGit/Source/GLTFExporter/Public/GLTFLevelConverters.h
@@ -5,14 +5,14 @@
 #include "GLTFJsonIndex.h"
 #include "Engine.h"
 
-struct FGLTFContainerBuilder;
+struct FGLTFIndexedBuilder;
 
 struct GLTFEXPORTER_API FGLTFSceneComponentConverter
 {
-	static FGLTFJsonNodeIndex Convert(FGLTFContainerBuilder& Container, const FString& Name, const USceneComponent* SceneComponent, bool bSelectedOnly, bool bRootNode = false);
+	static FGLTFJsonNodeIndex Convert(FGLTFIndexedBuilder& Builder, const FString& Name, const USceneComponent* SceneComponent, bool bSelectedOnly, bool bRootNode = false);
 };
 
 struct GLTFEXPORTER_API FGLTFLevelConverter
 {
-	static FGLTFJsonSceneIndex Convert(FGLTFContainerBuilder& Container, const FString& Name, const ULevel* Level, bool bSelectedOnly);
+	static FGLTFJsonSceneIndex Convert(FGLTFIndexedBuilder& Builder, const FString& Name, const ULevel* Level, bool bSelectedOnly);
 };
diff --git a/Engine/Plugins/Enterprise/GLTFExporter/MigratedFromGit/Source/GLTFExporter/Public/GLTFStaticMeshConverters.h b/Engine/Plugins/Enterprise/GLTFExporter/MigratedFromGit/Source/GLTFExporter/Public/GLTFStaticMeshConverters.h
--- a/Engine/Plugins/Enterprise/GLTFExporter/MigratedFromGit/Source/GLTFExporter/Public/GLTFStaticMeshConverters.h
+++ b/Engine/Plugins/Enterprise/GLTFExporter/MigratedFromGit/Source/GLTFExporter/Public/GLTFStaticMeshConverters.h
@@ -5,49 +5,49 @@
 #include "GLTFJsonIndex.h"
 #include "Engine.h"
 
-struct FGLTFContainerBuilder;
+struct FGLTFIndexedBuilder;
 
 struct GLTFEXPORTER_API FGLTFPositionVertexBufferConverter
 {
-	static FGLTFJsonAccessorIndex Convert(FGLTFContainerBuilder& Container, const FString& Name, const FPositionVertexBuffer* VertexBuffer);
+	static FGLTFJsonAccessorIndex Convert(FGLTFIndexedBuilder& Builder, const FString& Name, const FPositionVertexBuffer* VertexBuffer);
 };
 
 struct GLTFEXPORTER_API FGLTFColorVertexBufferConverter
 {
-	static FGLTFJsonAccessorIndex Convert(FGLTFContainerBuilder& Container, const FString& Name, const FColorVertexBuffer* VertexBuffer);
+	static FGLTFJsonAccessorIndex Convert(FGLTFIndexedBuilder& Builder, const FString& Name, const FColorVertexBuffer* VertexBuffer);
 };
 
 struct GLTFEXPORTER_API FGLTFStaticMeshNormalVertexBufferConverter
 {
-	static FGLTFJsonAccessorIndex Convert(FGLTFContainerBuilder& Container, const FString& Name, const FStaticMeshVertexBuffer* VertexBuffer);
+	static FGLTFJsonAccessorIndex Convert(FGLTFIndexedBuilder& Builder, const FString& Name, const FStaticMeshVertexBuffer* VertexBuffer);
 };
 
 struct GLTFEXPORTER_API FGLTFStaticMeshTangentVertexBufferConverter
 {
-	static FGLTFJsonAccessorIndex Convert(FGLTFContainerBuilder& Container, const FString& Name, const FStaticMeshVertexBuffer* VertexBuffer);
+	static FGLTFJsonAccessorIndex Convert(FGLTFIndexedBuilder& Builder, const FString& Name, const FStaticMeshVertexBuffer* VertexBuffer);
 };
 
 struct GLTFEXPORTER_API FGLTFStaticMeshUV0VertexBufferConverter
 {
-	static FGLTFJsonAccessorIndex Convert(FGLTFContainerBuilder& Container, const FString& Name, const FStaticMeshVertexBuffer* VertexBuffer);
+	static FGLTFJsonAccessorIndex Convert(FGLTFIndexedBuilder& Builder, const FString& Name, const FStaticMeshVertexBuffer* VertexBuffer);
 };
 
 struct GLTFEXPORTER_API FGLTFStaticMeshUV1VertexBufferConverter
 {
-	static FGLTFJsonAccessorIndex Convert(FGLTFContainerBuilder& Container, const FString& Name, const FStaticMeshVertexBuffer* VertexBuffer);
+	static FGLTFJsonAccessorIndex Convert(FGLTFIndexedBuilder& Builder, const FString& Name, const FStaticMeshVertexBuffer* VertexBuffer);
 };
 
 struct GLTFEXPORTER_API FGLTFStaticMeshIndexBufferConverter
 {
-	static FGLTFJsonBufferViewIndex Convert(FGLTFContainerBuilder& Container, const FString& Name, const FRawStaticIndexBuffer* IndexBuffer);
+	static FGLTFJsonBufferViewIndex Convert(FGLTFIndexedBuilder& Builder, const FString& Name, const FRawStaticIndexBuffer* IndexBuffer);
 };
 
 struct GLTFEXPORTER_API FGLTFStaticMeshSectionConverter
 {
-	static FGLTFJsonAccessorIndex Convert(FGLTFContainerBuilder& Container, const FString& Name, const FStaticMeshSection* MeshSection, const FRawStaticIndexBuffer* IndexBuffer);
+	static FGLTFJsonAccessorIndex Convert(FGLTFIndexedBuilder& Builder, const FString& Name, const FStaticMeshSection* MeshSection, const FRawStaticIndexBuffer* IndexBuffer);
 };
 
 struct GLTFEXPORTER_API FGLTFStaticMeshConverter
 {
-	static FGLTFJsonMeshIndex Convert(FGLTFContainerBuilder& Container, const FString& Name, const FStaticMeshLODResources* StaticMeshLOD, const FColorVertexBuffer* OverrideVertexColors);
+	static FGLTFJsonMeshIndex Convert(FGLTFIndexedBuilder& Builder, const FString& Name, const FStaticMeshLODResources* StaticMeshLOD, const FColorVertexBuffer* OverrideVertexColors);
 };
