[GLTFExporter] fix: uninitialized boundingbox and mesh naming

commit b794f128fce4d5e2aa5584f8913dfb547e93e577
Date: Mon, 11 May 2020 00:59:02 +0200

    fix: uninitialized boundingbox and mesh naming
    (cherry picked from commit 4b384437919825bb6b7084d64de85e77bf835f8e​)

#skipci

diff --git a/Engine/Plugins/Enterprise/GLTFExporter/MigratedFromGit/Source/GLTFExporter/Private/GLTFIndexedBuilder.cpp b/Engine/Plugins/Enterprise/GLTFExporter/MigratedFromGit/Source/GLTFExporter/Private/GLTFIndexedBuilder.cpp
--- a/Engine/Plugins/Enterprise/GLTFExporter/MigratedFromGit/Source/GLTFExporter/Private/GLTFIndexedBuilder.cpp
+++ b/Engine/Plugins/Enterprise/GLTFExporter/MigratedFromGit/Source/GLTFExporter/Private/GLTFIndexedBuilder.cpp
@@ -50,7 +50,9 @@ FGLTFJsonMeshIndex FGLTFIndexedBuilder::ConvertMesh(const FStaticMeshLODResource
 FGLTFJsonMeshIndex FGLTFIndexedBuilder::ConvertMesh(const UStaticMesh* StaticMesh, int32 LODIndex, const FColorVertexBuffer* OverrideVertexColors, const FString& DesiredName)
 {
 	const FStaticMeshLODResources* StaticMeshLOD = &StaticMesh->GetLODForExport(LODIndex);
-	return ConvertMesh(StaticMeshLOD, OverrideVertexColors, DesiredName.IsEmpty() ? StaticMesh->GetName() + TEXT("_LOD") + FString::FromInt(LODIndex) : DesiredName);
+	FString Name = DesiredName.IsEmpty() ? StaticMesh->GetName() : DesiredName;
+	if (LODIndex != 0) Name += TEXT("_LOD") + FString::FromInt(LODIndex);
+	return ConvertMesh(StaticMeshLOD, OverrideVertexColors, Name);
 }
 
 FGLTFJsonMeshIndex FGLTFIndexedBuilder::ConvertMesh(const UStaticMeshComponent* StaticMeshComponent, const FString& DesiredName)
diff --git a/Engine/Plugins/Enterprise/GLTFExporter/MigratedFromGit/Source/GLTFExporter/Private/GLTFStaticMeshConverters.cpp b/Engine/Plugins/Enterprise/GLTFExporter/MigratedFromGit/Source/GLTFExporter/Private/GLTFStaticMeshConverters.cpp
--- a/Engine/Plugins/Enterprise/GLTFExporter/MigratedFromGit/Source/GLTFExporter/Private/GLTFStaticMeshConverters.cpp
+++ b/Engine/Plugins/Enterprise/GLTFExporter/MigratedFromGit/Source/GLTFExporter/Private/GLTFStaticMeshConverters.cpp
@@ -21,6 +21,7 @@ FGLTFJsonAccessorIndex FGLTFPositionVertexBufferConverter::Convert(FGLTFIndexedB
 	}
 
 	FBox BoundingBox;
+	BoundingBox.Init();
 
 	// More accurate bounding box if based on raw vertex values
 	for (const FVector& Position : Positions)
