[GLTFExporter] Merge remote-tracking branch 'origin/feature/export-level' into feature/avoid-asset-duplication

commit 2caf400739ac1d5c6d79ec5567891d91dc7b9d12
Date: Tue, 12 May 2020 08:39:37 +0200

    Merge remote-tracking branch 'origin/feature/export-level' into feature/avoid-asset-duplication
    # Conflicts from merge:
    #	Plugins/GLTFExporter/Source/GLTFExporter/Private/GLTFMeshBuilder.cpp
    #	Plugins/GLTFExporter/Source/GLTFExporter/Private/GLTFSceneBuilder.cpp
    (cherry picked from commit c99a6408300cbeb673b218330509d6bfab241a29​)
    # Conflicts from cherry-pick:
    #	Plugins/GLTFExporter/Source/GLTFExporter/Public/GLTFExportOptions.h

#skipci

diff --git a/Engine/Plugins/Enterprise/GLTFExporter/MigratedFromGit/Source/GLTFExporter/Private/GLTFIndexedBuilder.cpp b/Engine/Plugins/Enterprise/GLTFExporter/MigratedFromGit/Source/GLTFExporter/Private/GLTFIndexedBuilder.cpp
--- a/Engine/Plugins/Enterprise/GLTFExporter/MigratedFromGit/Source/GLTFExporter/Private/GLTFIndexedBuilder.cpp
+++ b/Engine/Plugins/Enterprise/GLTFExporter/MigratedFromGit/Source/GLTFExporter/Private/GLTFIndexedBuilder.cpp
@@ -51,7 +51,12 @@ FGLTFJsonMeshIndex FGLTFIndexedBuilder::ConvertMesh(const UStaticMesh* StaticMes
 {
 	const FStaticMeshLODResources* StaticMeshLOD = &StaticMesh->GetLODForExport(LODIndex);
 	FString Name = DesiredName.IsEmpty() ? StaticMesh->GetName() : DesiredName;
-	if (LODIndex != 0) Name += TEXT("_LOD") + FString::FromInt(LODIndex);
+
+	if (LODIndex != 0)
+	{
+		Name += TEXT("_LOD") + FString::FromInt(LODIndex);
+	}
+
 	return ConvertMesh(StaticMeshLOD, OverrideVertexColors, Name);
 }
 
@@ -60,6 +65,7 @@ FGLTFJsonMeshIndex FGLTFIndexedBuilder::ConvertMesh(const UStaticMeshComponent*
 	const UStaticMesh* StaticMesh = StaticMeshComponent->GetStaticMesh();
 	const int32 LODIndex = StaticMeshComponent->ForcedLodModel > 0 ? StaticMeshComponent->ForcedLodModel - 1 : /* auto-select */ 0;
 	const FColorVertexBuffer* OverrideVertexColors = LODIndex < StaticMeshComponent->LODData.Num() ? StaticMeshComponent->LODData[LODIndex].OverrideVertexColors : nullptr;
+
 	return ConvertMesh(StaticMesh, LODIndex, OverrideVertexColors, DesiredName);
 }
 
@@ -75,6 +81,5 @@ FGLTFJsonSceneIndex FGLTFIndexedBuilder::ConvertScene(const ULevel* Level, bool
 
 FGLTFJsonSceneIndex FGLTFIndexedBuilder::ConvertScene(const UWorld* World, bool bSelectedOnly, const FString& DesiredName)
 {
-	const ULevel* Level = World->PersistentLevel;
-	return ConvertScene(Level, bSelectedOnly, DesiredName.IsEmpty() ? World->GetName() : DesiredName);
+	return ConvertScene(World->PersistentLevel, bSelectedOnly, DesiredName.IsEmpty() ? World->GetName() : DesiredName);
 }
diff --git a/Engine/Plugins/Enterprise/GLTFExporter/MigratedFromGit/Source/GLTFExporter/Private/GLTFLevelConverters.cpp b/Engine/Plugins/Enterprise/GLTFExporter/MigratedFromGit/Source/GLTFExporter/Private/GLTFLevelConverters.cpp
--- a/Engine/Plugins/Enterprise/GLTFExporter/MigratedFromGit/Source/GLTFExporter/Private/GLTFLevelConverters.cpp
+++ b/Engine/Plugins/Enterprise/GLTFExporter/MigratedFromGit/Source/GLTFExporter/Private/GLTFLevelConverters.cpp
@@ -4,7 +4,7 @@
 
 #include "GLTFLevelConverters.h"
 #include "GLTFContainerBuilder.h"
-#include "GLTFConversionUtilities.h"
+#include "GLTFConverterUtility.h"
 
 FGLTFJsonNodeIndex FGLTFSceneComponentConverter::Convert(FGLTFIndexedBuilder& Builder, const FString& Name, const USceneComponent* SceneComponent, bool bSelectedOnly, bool bRootNode)
 {
@@ -16,11 +16,11 @@ FGLTFJsonNodeIndex FGLTFSceneComponentConverter::Convert(FGLTFIndexedBuilder& Bu
 	Node.Name = Name.IsEmpty() ? Owner->GetName() + TEXT("_") + SceneComponent->GetName() : Name;
 
 	const FTransform Transform = bRootNode ? SceneComponent->GetComponentTransform() : SceneComponent->GetRelativeTransform();
-	Node.Translation = ConvertPosition(Transform.GetTranslation());
-	Node.Rotation = ConvertRotation(Transform.GetRotation());
-	Node.Scale = ConvertScale(Transform.GetScale3D());
+	Node.Translation = FGLTFConverterUtility::ConvertPosition(Transform.GetTranslation());
+	Node.Rotation = FGLTFConverterUtility::ConvertRotation(Transform.GetRotation());
+	Node.Scale = FGLTFConverterUtility::ConvertScale(Transform.GetScale3D());
 	
-	if (IsSkySphereBlueprint(Blueprint))
+	if (FGLTFConverterUtility::IsSkySphereBlueprint(Blueprint))
 	{
 		// Ignore components
 	}
@@ -28,7 +28,7 @@ FGLTFJsonNodeIndex FGLTFSceneComponentConverter::Convert(FGLTFIndexedBuilder& Bu
 	{
 		Node.Mesh = Builder.ConvertMesh(StaticMeshComponent);
 	}
-	else if (IsHDRIBackdropBlueprint(Blueprint) && bIsRootComponent)
+	else if (FGLTFConverterUtility::IsHDRIBackdropBlueprint(Blueprint) && bIsRootComponent)
 	{
 		// TODO: add support for backdrop
 	}
@@ -42,7 +42,10 @@ FGLTFJsonNodeIndex FGLTFSceneComponentConverter::Convert(FGLTFIndexedBuilder& Bu
 			if (!bSelectedOnly || ChildOwner->IsSelected())
 			{
 				FGLTFJsonNodeIndex NodeIndex = Builder.ConvertNode(ChildComponent, bSelectedOnly, false);
-				if (NodeIndex != INDEX_NONE) Node.Children.Add(NodeIndex);
+				if (NodeIndex != INDEX_NONE)
+				{
+					Node.Children.Add(NodeIndex);
+				}
 			}
 		}
 	}
@@ -69,7 +72,10 @@ FGLTFJsonSceneIndex FGLTFLevelConverter::Convert(FGLTFIndexedBuilder& Builder, c
 				if (ParentActor == nullptr || (bSelectedOnly && !ParentActor->IsSelected()))
 				{
 					FGLTFJsonNodeIndex NodeIndex = Builder.ConvertNode(RootComponent, bSelectedOnly, true);
-					if (NodeIndex != INDEX_NONE) Scene.Nodes.Add(NodeIndex);
+					if (NodeIndex != INDEX_NONE)
+					{
+						Scene.Nodes.Add(NodeIndex);
+					}
 				}
 			}
 		}
diff --git a/Engine/Plugins/Enterprise/GLTFExporter/MigratedFromGit/Source/GLTFExporter/Private/GLTFStaticMeshConverters.cpp b/Engine/Plugins/Enterprise/GLTFExporter/MigratedFromGit/Source/GLTFExporter/Private/GLTFStaticMeshConverters.cpp
--- a/Engine/Plugins/Enterprise/GLTFExporter/MigratedFromGit/Source/GLTFExporter/Private/GLTFStaticMeshConverters.cpp
+++ b/Engine/Plugins/Enterprise/GLTFExporter/MigratedFromGit/Source/GLTFExporter/Private/GLTFStaticMeshConverters.cpp
@@ -2,7 +2,7 @@
 
 #include "GLTFStaticMeshConverters.h"
 #include "GLTFContainerBuilder.h"
-#include "GLTFConversionUtilities.h"
+#include "GLTFConverterUtility.h"
 
 FGLTFJsonAccessorIndex FGLTFPositionVertexBufferConverter::Convert(FGLTFIndexedBuilder& Builder, const FString& Name, const FPositionVertexBuffer* VertexBuffer)
 {
@@ -17,7 +17,7 @@ FGLTFJsonAccessorIndex FGLTFPositionVertexBufferConverter::Convert(FGLTFIndexedB
 
 	for (uint32 VertexIndex = 0; VertexIndex < VertexCount; ++VertexIndex)
 	{
-		Positions[VertexIndex] = ConvertPosition(VertexBuffer->VertexPosition(VertexIndex));
+		Positions[VertexIndex] = FGLTFConverterUtility::ConvertPosition(VertexBuffer->VertexPosition(VertexIndex));
 	}
 
 	FBox BoundingBox;
@@ -60,7 +60,7 @@ FGLTFJsonAccessorIndex FGLTFColorVertexBufferConverter::Convert(FGLTFIndexedBuil
 
 	for (uint32 VertexIndex = 0; VertexIndex < VertexCount; ++VertexIndex)
 	{
-		Colors[VertexIndex] = ConvertColor(VertexBuffer->VertexColor(VertexIndex));
+		Colors[VertexIndex] = FGLTFConverterUtility::ConvertColor(VertexBuffer->VertexColor(VertexIndex));
 	}
 
 	FGLTFJsonAccessor JsonAccessor;
@@ -87,7 +87,7 @@ FGLTFJsonAccessorIndex FGLTFStaticMeshNormalVertexBufferConverter::Convert(FGLTF
 	
 	for (uint32 VertexIndex = 0; VertexIndex < VertexCount; ++VertexIndex)
 	{
-		Normals[VertexIndex] = ConvertVector(VertexBuffer->VertexTangentZ(VertexIndex));
+		Normals[VertexIndex] = FGLTFConverterUtility::ConvertVector(VertexBuffer->VertexTangentZ(VertexIndex));
 	}
 
 	FGLTFJsonAccessor JsonAccessor;
@@ -113,7 +113,7 @@ FGLTFJsonAccessorIndex FGLTFStaticMeshTangentVertexBufferConverter::Convert(FGLT
 	
 	for (uint32 VertexIndex = 0; VertexIndex < VertexCount; ++VertexIndex)
 	{
-		Tangents[VertexIndex] = ConvertVector(VertexBuffer->VertexTangentX(VertexIndex));
+		Tangents[VertexIndex] = FGLTFConverterUtility::ConvertVector(VertexBuffer->VertexTangentX(VertexIndex));
 	}
 
 	FGLTFJsonAccessor JsonAccessor;
diff --git a/Engine/Plugins/Enterprise/GLTFExporter/MigratedFromGit/Source/GLTFExporter/Public/GLTFJsonAccessor.h b/Engine/Plugins/Enterprise/GLTFExporter/MigratedFromGit/Source/GLTFExporter/Public/GLTFJsonAccessor.h
--- a/Engine/Plugins/Enterprise/GLTFExporter/MigratedFromGit/Source/GLTFExporter/Public/GLTFJsonAccessor.h
+++ b/Engine/Plugins/Enterprise/GLTFExporter/MigratedFromGit/Source/GLTFExporter/Public/GLTFJsonAccessor.h
@@ -46,7 +46,12 @@ struct GLTFEXPORTER_API FGLTFJsonAccessor
 		}
 
 		JsonWriter.WriteValue(TEXT("bufferView"), BufferView);
-		if (ByteOffset != 0) JsonWriter.WriteValue(TEXT("byteOffset"), ByteOffset);
+
+		if (ByteOffset != 0)
+		{
+			JsonWriter.WriteValue(TEXT("byteOffset"), ByteOffset);
+		}
+
 		JsonWriter.WriteValue(TEXT("count"), Count);
 		JsonWriter.WriteValue(TEXT("type"), FGLTFJsonUtility::AccessorTypeToString(Type));
 		JsonWriter.WriteValue(TEXT("componentType"), FGLTFJsonUtility::ComponentTypeToNumber(ComponentType));
