[GLTFExporter] refactor: rename GLTFIndexedBuilder::Convert* -> GLTFIndexedBuilder::GetOrAdd*

commit e8a8a76e4a59d062047d41e1f9e0d4a37690ded7
Date: Mon, 11 May 2020 22:28:02 +0200

    refactor: rename GLTFIndexedBuilder::Convert* -> GLTFIndexedBuilder::GetOrAdd*
    (cherry picked from commit 830c8adc0a397a0a360f1f5024251229f27b2b34​)

#skipci

diff --git a/Engine/Plugins/Enterprise/GLTFExporter/MigratedFromGit/Source/GLTFExporter/Private/GLTFIndexedBuilder.cpp b/Engine/Plugins/Enterprise/GLTFExporter/MigratedFromGit/Source/GLTFExporter/Private/GLTFIndexedBuilder.cpp
--- a/Engine/Plugins/Enterprise/GLTFExporter/MigratedFromGit/Source/GLTFExporter/Private/GLTFIndexedBuilder.cpp
+++ b/Engine/Plugins/Enterprise/GLTFExporter/MigratedFromGit/Source/GLTFExporter/Private/GLTFIndexedBuilder.cpp
@@ -2,52 +2,52 @@
 
 #include "GLTFIndexedBuilder.h"
 
-FGLTFJsonAccessorIndex FGLTFIndexedBuilder::ConvertPositionAccessor(const FPositionVertexBuffer* VertexBuffer, const FString& DesiredName)
+FGLTFJsonAccessorIndex FGLTFIndexedBuilder::GetOrAddPositionAccessor(const FPositionVertexBuffer* VertexBuffer, const FString& DesiredName)
 {
 	return PositionVertexBuffers.Convert(*this, DesiredName, VertexBuffer);
 }
 
-FGLTFJsonAccessorIndex FGLTFIndexedBuilder::ConvertColorAccessor(const FColorVertexBuffer* VertexBuffer, const FString& DesiredName)
+FGLTFJsonAccessorIndex FGLTFIndexedBuilder::GetOrAddColorAccessor(const FColorVertexBuffer* VertexBuffer, const FString& DesiredName)
 {
 	return ColorVertexBuffers.Convert(*this, DesiredName, VertexBuffer);
 }
 
-FGLTFJsonAccessorIndex FGLTFIndexedBuilder::ConvertNormalAccessor(const FStaticMeshVertexBuffer* VertexBuffer, const FString& DesiredName)
+FGLTFJsonAccessorIndex FGLTFIndexedBuilder::GetOrAddNormalAccessor(const FStaticMeshVertexBuffer* VertexBuffer, const FString& DesiredName)
 {
 	return StaticMeshNormalVertexBuffers.Convert(*this, DesiredName, VertexBuffer);
 }
 
-FGLTFJsonAccessorIndex FGLTFIndexedBuilder::ConvertTangentAccessor(const FStaticMeshVertexBuffer* VertexBuffer, const FString& DesiredName)
+FGLTFJsonAccessorIndex FGLTFIndexedBuilder::GetOrAddTangentAccessor(const FStaticMeshVertexBuffer* VertexBuffer, const FString& DesiredName)
 {
 	return StaticMeshTangentVertexBuffers.Convert(*this, DesiredName, VertexBuffer);
 }
 
-FGLTFJsonAccessorIndex FGLTFIndexedBuilder::ConvertUV0Accessor(const FStaticMeshVertexBuffer* VertexBuffer, const FString& DesiredName)
+FGLTFJsonAccessorIndex FGLTFIndexedBuilder::GetOrAddUV0Accessor(const FStaticMeshVertexBuffer* VertexBuffer, const FString& DesiredName)
 {
 	return StaticMeshUV0VertexBuffers.Convert(*this, DesiredName, VertexBuffer);
 }
 
-FGLTFJsonAccessorIndex FGLTFIndexedBuilder::ConvertUV1Accessor(const FStaticMeshVertexBuffer* VertexBuffer, const FString& DesiredName)
+FGLTFJsonAccessorIndex FGLTFIndexedBuilder::GetOrAddUV1Accessor(const FStaticMeshVertexBuffer* VertexBuffer, const FString& DesiredName)
 {
 	return StaticMeshUV1VertexBuffers.Convert(*this, DesiredName, VertexBuffer);
 }
 
-FGLTFJsonBufferViewIndex FGLTFIndexedBuilder::ConvertIndexBufferView(const FRawStaticIndexBuffer* IndexBuffer, const FString& DesiredName)
+FGLTFJsonBufferViewIndex FGLTFIndexedBuilder::GetOrAddIndexBufferView(const FRawStaticIndexBuffer* IndexBuffer, const FString& DesiredName)
 {
 	return StaticMeshIndexBuffers.Convert(*this, DesiredName, IndexBuffer);
 }
 
-FGLTFJsonAccessorIndex FGLTFIndexedBuilder::ConvertIndexAccessor(const FStaticMeshSection* MeshSection, const FRawStaticIndexBuffer* IndexBuffer, const FString& DesiredName)
+FGLTFJsonAccessorIndex FGLTFIndexedBuilder::GetOrAddIndexAccessor(const FStaticMeshSection* MeshSection, const FRawStaticIndexBuffer* IndexBuffer, const FString& DesiredName)
 {
 	return StaticMeshSections.Convert(*this, DesiredName, MeshSection, IndexBuffer);
 }
 
-FGLTFJsonMeshIndex FGLTFIndexedBuilder::ConvertMesh(const FStaticMeshLODResources* StaticMeshLOD, const FColorVertexBuffer* OverrideVertexColors, const FString& DesiredName)
+FGLTFJsonMeshIndex FGLTFIndexedBuilder::GetOrAddMesh(const FStaticMeshLODResources* StaticMeshLOD, const FColorVertexBuffer* OverrideVertexColors, const FString& DesiredName)
 {
 	return StaticMeshes.Convert(*this, DesiredName, StaticMeshLOD, OverrideVertexColors);
 }
 
-FGLTFJsonMeshIndex FGLTFIndexedBuilder::ConvertMesh(const UStaticMesh* StaticMesh, int32 LODIndex, const FColorVertexBuffer* OverrideVertexColors, const FString& DesiredName)
+FGLTFJsonMeshIndex FGLTFIndexedBuilder::GetOrAddMesh(const UStaticMesh* StaticMesh, int32 LODIndex, const FColorVertexBuffer* OverrideVertexColors, const FString& DesiredName)
 {
 	const FStaticMeshLODResources* StaticMeshLOD = &StaticMesh->GetLODForExport(LODIndex);
 	FString Name = DesiredName.IsEmpty() ? StaticMesh->GetName() : DesiredName;
@@ -57,29 +57,29 @@ FGLTFJsonMeshIndex FGLTFIndexedBuilder::ConvertMesh(const UStaticMesh* StaticMes
 		Name += TEXT("_LOD") + FString::FromInt(LODIndex);
 	}
 
-	return ConvertMesh(StaticMeshLOD, OverrideVertexColors, Name);
+	return GetOrAddMesh(StaticMeshLOD, OverrideVertexColors, Name);
 }
 
-FGLTFJsonMeshIndex FGLTFIndexedBuilder::ConvertMesh(const UStaticMeshComponent* StaticMeshComponent, const FString& DesiredName)
+FGLTFJsonMeshIndex FGLTFIndexedBuilder::GetOrAddMesh(const UStaticMeshComponent* StaticMeshComponent, const FString& DesiredName)
 {
 	const UStaticMesh* StaticMesh = StaticMeshComponent->GetStaticMesh();
 	const int32 LODIndex = StaticMeshComponent->ForcedLodModel > 0 ? StaticMeshComponent->ForcedLodModel - 1 : /* auto-select */ 0;
 	const FColorVertexBuffer* OverrideVertexColors = LODIndex < StaticMeshComponent->LODData.Num() ? StaticMeshComponent->LODData[LODIndex].OverrideVertexColors : nullptr;
 
-	return ConvertMesh(StaticMesh, LODIndex, OverrideVertexColors, DesiredName);
+	return GetOrAddMesh(StaticMesh, LODIndex, OverrideVertexColors, DesiredName);
 }
 
-FGLTFJsonNodeIndex FGLTFIndexedBuilder::ConvertNode(const USceneComponent* SceneComponent, bool bSelectedOnly, bool bRootNode, const FString& DesiredName)
+FGLTFJsonNodeIndex FGLTFIndexedBuilder::GetOrAddNode(const USceneComponent* SceneComponent, bool bSelectedOnly, bool bRootNode, const FString& DesiredName)
 {
 	return SceneComponents.Convert(*this, DesiredName, SceneComponent, bSelectedOnly, bRootNode);
 }
 
-FGLTFJsonSceneIndex FGLTFIndexedBuilder::ConvertScene(const ULevel* Level, bool bSelectedOnly, const FString& DesiredName)
+FGLTFJsonSceneIndex FGLTFIndexedBuilder::GetOrAddScene(const ULevel* Level, bool bSelectedOnly, const FString& DesiredName)
 {
 	return Levels.Convert(*this, DesiredName, Level, bSelectedOnly);
 }
 
-FGLTFJsonSceneIndex FGLTFIndexedBuilder::ConvertScene(const UWorld* World, bool bSelectedOnly, const FString& DesiredName)
+FGLTFJsonSceneIndex FGLTFIndexedBuilder::GetOrAddScene(const UWorld* World, bool bSelectedOnly, const FString& DesiredName)
 {
-	return ConvertScene(World->PersistentLevel, bSelectedOnly, DesiredName.IsEmpty() ? World->GetName() : DesiredName);
+	return GetOrAddScene(World->PersistentLevel, bSelectedOnly, DesiredName.IsEmpty() ? World->GetName() : DesiredName);
 }
diff --git a/Engine/Plugins/Enterprise/GLTFExporter/MigratedFromGit/Source/GLTFExporter/Private/GLTFLevelConverters.cpp b/Engine/Plugins/Enterprise/GLTFExporter/MigratedFromGit/Source/GLTFExporter/Private/GLTFLevelConverters.cpp
--- a/Engine/Plugins/Enterprise/GLTFExporter/MigratedFromGit/Source/GLTFExporter/Private/GLTFLevelConverters.cpp
+++ b/Engine/Plugins/Enterprise/GLTFExporter/MigratedFromGit/Source/GLTFExporter/Private/GLTFLevelConverters.cpp
@@ -26,7 +26,7 @@ FGLTFJsonNodeIndex FGLTFSceneComponentConverter::Convert(FGLTFIndexedBuilder& Bu
 	}
 	else if (const UStaticMeshComponent* StaticMeshComponent = Cast<UStaticMeshComponent>(SceneComponent))
 	{
-		Node.Mesh = Builder.ConvertMesh(StaticMeshComponent);
+		Node.Mesh = Builder.GetOrAddMesh(StaticMeshComponent);
 	}
 	else if (FGLTFConverterUtility::IsHDRIBackdropBlueprint(Blueprint) && bIsRootComponent)
 	{
@@ -41,7 +41,7 @@ FGLTFJsonNodeIndex FGLTFSceneComponentConverter::Convert(FGLTFIndexedBuilder& Bu
 			const AActor* ChildOwner = ChildComponent->GetOwner();
 			if (!bSelectedOnly || ChildOwner->IsSelected())
 			{
-				FGLTFJsonNodeIndex NodeIndex = Builder.ConvertNode(ChildComponent, bSelectedOnly, false);
+				FGLTFJsonNodeIndex NodeIndex = Builder.GetOrAddNode(ChildComponent, bSelectedOnly, false);
 				if (NodeIndex != INDEX_NONE)
 				{
 					Node.Children.Add(NodeIndex);
@@ -71,7 +71,7 @@ FGLTFJsonSceneIndex FGLTFLevelConverter::Convert(FGLTFIndexedBuilder& Builder, c
 				const AActor* ParentActor = Actor->GetParentActor();
 				if (ParentActor == nullptr || (bSelectedOnly && !ParentActor->IsSelected()))
 				{
-					FGLTFJsonNodeIndex NodeIndex = Builder.ConvertNode(RootComponent, bSelectedOnly, true);
+					FGLTFJsonNodeIndex NodeIndex = Builder.GetOrAddNode(RootComponent, bSelectedOnly, true);
 					if (NodeIndex != INDEX_NONE)
 					{
 						Scene.Nodes.Add(NodeIndex);
diff --git a/Engine/Plugins/Enterprise/GLTFExporter/MigratedFromGit/Source/GLTFExporter/Private/GLTFLevelExporter.cpp b/Engine/Plugins/Enterprise/GLTFExporter/MigratedFromGit/Source/GLTFExporter/Private/GLTFLevelExporter.cpp
--- a/Engine/Plugins/Enterprise/GLTFExporter/MigratedFromGit/Source/GLTFExporter/Private/GLTFLevelExporter.cpp
+++ b/Engine/Plugins/Enterprise/GLTFExporter/MigratedFromGit/Source/GLTFExporter/Private/GLTFLevelExporter.cpp
@@ -21,7 +21,7 @@ bool UGLTFLevelExporter::ExportBinary(UObject* Object, const TCHAR* Type, FArchi
 	}
 
 	FGLTFContainerBuilder Builder;
-	const FGLTFJsonSceneIndex SceneIndex = Builder.ConvertScene(World, bSelectedOnly);
+	const FGLTFJsonSceneIndex SceneIndex = Builder.GetOrAddScene(World, bSelectedOnly);
 	Builder.JsonRoot.DefaultScene = SceneIndex;
 
 	Builder.Serialize(Archive);
diff --git a/Engine/Plugins/Enterprise/GLTFExporter/MigratedFromGit/Source/GLTFExporter/Private/GLTFStaticMeshConverters.cpp b/Engine/Plugins/Enterprise/GLTFExporter/MigratedFromGit/Source/GLTFExporter/Private/GLTFStaticMeshConverters.cpp
--- a/Engine/Plugins/Enterprise/GLTFExporter/MigratedFromGit/Source/GLTFExporter/Private/GLTFStaticMeshConverters.cpp
+++ b/Engine/Plugins/Enterprise/GLTFExporter/MigratedFromGit/Source/GLTFExporter/Private/GLTFStaticMeshConverters.cpp
@@ -200,7 +200,7 @@ FGLTFJsonAccessorIndex FGLTFStaticMeshSectionConverter::Convert(FGLTFIndexedBuil
 
 	FGLTFJsonAccessor JsonAccessor;
 	JsonAccessor.Name = Name;
-	JsonAccessor.BufferView = Builder.ConvertIndexBufferView(IndexBuffer);
+	JsonAccessor.BufferView = Builder.GetOrAddIndexBufferView(IndexBuffer);
 	JsonAccessor.ByteOffset = MeshSection->FirstIndex * sizeof(uint32);
 	JsonAccessor.ComponentType = EGLTFJsonComponentType::U32;
 	JsonAccessor.Count = TriangleCount * 3;
@@ -219,15 +219,15 @@ FGLTFJsonMeshIndex FGLTFStaticMeshConverter::Convert(FGLTFIndexedBuilder& Builde
 	const FColorVertexBuffer* ColorBuffer = OverrideVertexColors != nullptr ? OverrideVertexColors : &StaticMeshLOD->VertexBuffers.ColorVertexBuffer;
 
 	FGLTFJsonAttributes JsonAttributes;
-	JsonAttributes.Position = Builder.ConvertPositionAccessor(PositionBuffer, Name.IsEmpty() ? Name : Name + TEXT("_Positions"));
-	JsonAttributes.Normal = Builder.ConvertNormalAccessor(VertexBuffer, Name.IsEmpty() ? Name : Name + TEXT("_Normals"));
-	JsonAttributes.Tangent = Builder.ConvertTangentAccessor(VertexBuffer, Name.IsEmpty() ? Name : Name + TEXT("_Tangents"));
-	JsonAttributes.TexCoord0 = Builder.ConvertUV0Accessor(VertexBuffer, Name.IsEmpty() ? Name : Name + TEXT("_UV0s"));
-	JsonAttributes.TexCoord1 = Builder.ConvertUV1Accessor(VertexBuffer, Name.IsEmpty() ? Name : Name + TEXT("_UV1s"));
-	JsonAttributes.Color0 = Builder.ConvertColorAccessor(ColorBuffer, Name.IsEmpty() ? Name : Name + TEXT("_Colors"));
+	JsonAttributes.Position = Builder.GetOrAddPositionAccessor(PositionBuffer, Name.IsEmpty() ? Name : Name + TEXT("_Positions"));
+	JsonAttributes.Normal = Builder.GetOrAddNormalAccessor(VertexBuffer, Name.IsEmpty() ? Name : Name + TEXT("_Normals"));
+	JsonAttributes.Tangent = Builder.GetOrAddTangentAccessor(VertexBuffer, Name.IsEmpty() ? Name : Name + TEXT("_Tangents"));
+	JsonAttributes.TexCoord0 = Builder.GetOrAddUV0Accessor(VertexBuffer, Name.IsEmpty() ? Name : Name + TEXT("_UV0s"));
+	JsonAttributes.TexCoord1 = Builder.GetOrAddUV1Accessor(VertexBuffer, Name.IsEmpty() ? Name : Name + TEXT("_UV1s"));
+	JsonAttributes.Color0 = Builder.GetOrAddColorAccessor(ColorBuffer, Name.IsEmpty() ? Name : Name + TEXT("_Colors"));
 
 	const FRawStaticIndexBuffer* IndexBuffer = &StaticMeshLOD->IndexBuffer;
-	Builder.ConvertIndexBufferView(IndexBuffer, Name.IsEmpty() ? Name : Name + TEXT("_Indices"));
+	Builder.GetOrAddIndexBufferView(IndexBuffer, Name.IsEmpty() ? Name : Name + TEXT("_Indices"));
 
 	const TArray<FStaticMeshSection>& Sections = StaticMeshLOD->Sections;
 	const int32 SectionCount = Sections.Num();
@@ -238,7 +238,7 @@ FGLTFJsonMeshIndex FGLTFStaticMeshConverter::Convert(FGLTFIndexedBuilder& Builde
 		FGLTFJsonPrimitive& JsonPrimitive = JsonMesh.Primitives[SectionIndex];
 		JsonPrimitive.Attributes = JsonAttributes;
 
-		JsonPrimitive.Indices = Builder.ConvertIndexAccessor(&Sections[SectionIndex], IndexBuffer,
+		JsonPrimitive.Indices = Builder.GetOrAddIndexAccessor(&Sections[SectionIndex], IndexBuffer,
 			Name.IsEmpty() ? Name : JsonMesh.Name + (SectionCount != 0 ? TEXT("_Indices_Section") + FString::FromInt(SectionIndex) : TEXT("_Indices")));
 	}
 
diff --git a/Engine/Plugins/Enterprise/GLTFExporter/MigratedFromGit/Source/GLTFExporter/Private/GLTFStaticMeshExporter.cpp b/Engine/Plugins/Enterprise/GLTFExporter/MigratedFromGit/Source/GLTFExporter/Private/GLTFStaticMeshExporter.cpp
--- a/Engine/Plugins/Enterprise/GLTFExporter/MigratedFromGit/Source/GLTFExporter/Private/GLTFStaticMeshExporter.cpp
+++ b/Engine/Plugins/Enterprise/GLTFExporter/MigratedFromGit/Source/GLTFExporter/Private/GLTFStaticMeshExporter.cpp
@@ -23,7 +23,7 @@ bool UGLTFStaticMeshExporter::ExportBinary(UObject* Object, const TCHAR* Type, F
 
 	FGLTFContainerBuilder Builder;
 
-	FGLTFJsonMeshIndex MeshIndex = Builder.ConvertMesh(StaticMesh);
+	FGLTFJsonMeshIndex MeshIndex = Builder.GetOrAddMesh(StaticMesh);
 
 	FGLTFJsonNode Node;
 	Node.Mesh = MeshIndex;
diff --git a/Engine/Plugins/Enterprise/GLTFExporter/MigratedFromGit/Source/GLTFExporter/Public/GLTFIndexedBuilder.h b/Engine/Plugins/Enterprise/GLTFExporter/MigratedFromGit/Source/GLTFExporter/Public/GLTFIndexedBuilder.h
--- a/Engine/Plugins/Enterprise/GLTFExporter/MigratedFromGit/Source/GLTFExporter/Public/GLTFIndexedBuilder.h
+++ b/Engine/Plugins/Enterprise/GLTFExporter/MigratedFromGit/Source/GLTFExporter/Public/GLTFIndexedBuilder.h
@@ -22,21 +22,21 @@ struct GLTFEXPORTER_API FGLTFIndexedBuilder : public FGLTFBufferBuilder
 	TGLTFIndexedObjects<FGLTFJsonNodeIndex, TTuple<const USceneComponent*, bool, bool>, FGLTFSceneComponentConverter> SceneComponents;
 	TGLTFIndexedObjects<FGLTFJsonSceneIndex, TTuple<const ULevel*, bool>, FGLTFLevelConverter> Levels;
 
-	FGLTFJsonAccessorIndex ConvertPositionAccessor(const FPositionVertexBuffer* VertexBuffer, const FString& DesiredName = TEXT(""));
-	FGLTFJsonAccessorIndex ConvertNormalAccessor(const FStaticMeshVertexBuffer* VertexBuffer, const FString& DesiredName = TEXT(""));
-	FGLTFJsonAccessorIndex ConvertTangentAccessor(const FStaticMeshVertexBuffer* VertexBuffer, const FString& DesiredName = TEXT(""));
-	FGLTFJsonAccessorIndex ConvertUV0Accessor(const FStaticMeshVertexBuffer* VertexBuffer, const FString& DesiredName = TEXT(""));
-	FGLTFJsonAccessorIndex ConvertUV1Accessor(const FStaticMeshVertexBuffer* VertexBuffer, const FString& DesiredName = TEXT(""));
-	FGLTFJsonAccessorIndex ConvertColorAccessor(const FColorVertexBuffer* VertexBuffer, const FString& DesiredName = TEXT(""));
-
-	FGLTFJsonBufferViewIndex ConvertIndexBufferView(const FRawStaticIndexBuffer* IndexBuffer, const FString& DesiredName = TEXT(""));
-	FGLTFJsonAccessorIndex ConvertIndexAccessor(const FStaticMeshSection* MeshSection, const FRawStaticIndexBuffer* IndexBuffer, const FString& DesiredName = TEXT(""));
-
-	FGLTFJsonMeshIndex ConvertMesh(const FStaticMeshLODResources* StaticMeshLOD, const FColorVertexBuffer* OverrideVertexColors = nullptr, const FString& DesiredName = TEXT(""));
-	FGLTFJsonMeshIndex ConvertMesh(const UStaticMesh* StaticMesh, int32 LODIndex = 0, const FColorVertexBuffer* OverrideVertexColors = nullptr, const FString& DesiredName = TEXT(""));
-	FGLTFJsonMeshIndex ConvertMesh(const UStaticMeshComponent* StaticMeshComponent, const FString& DesiredName = TEXT(""));
-
-	FGLTFJsonNodeIndex ConvertNode(const USceneComponent* SceneComponent, bool bSelectedOnly, bool bRootNode = false, const FString& DesiredName = TEXT(""));
-	FGLTFJsonSceneIndex ConvertScene(const ULevel* Level, bool bSelectedOnly, const FString& DesiredName = TEXT(""));
-	FGLTFJsonSceneIndex ConvertScene(const UWorld* World, bool bSelectedOnly, const FString& DesiredName = TEXT(""));
+	FGLTFJsonAccessorIndex GetOrAddPositionAccessor(const FPositionVertexBuffer* VertexBuffer, const FString& DesiredName = TEXT(""));
+	FGLTFJsonAccessorIndex GetOrAddNormalAccessor(const FStaticMeshVertexBuffer* VertexBuffer, const FString& DesiredName = TEXT(""));
+	FGLTFJsonAccessorIndex GetOrAddTangentAccessor(const FStaticMeshVertexBuffer* VertexBuffer, const FString& DesiredName = TEXT(""));
+	FGLTFJsonAccessorIndex GetOrAddUV0Accessor(const FStaticMeshVertexBuffer* VertexBuffer, const FString& DesiredName = TEXT(""));
+	FGLTFJsonAccessorIndex GetOrAddUV1Accessor(const FStaticMeshVertexBuffer* VertexBuffer, const FString& DesiredName = TEXT(""));
+	FGLTFJsonAccessorIndex GetOrAddColorAccessor(const FColorVertexBuffer* VertexBuffer, const FString& DesiredName = TEXT(""));
+
+	FGLTFJsonBufferViewIndex GetOrAddIndexBufferView(const FRawStaticIndexBuffer* IndexBuffer, const FString& DesiredName = TEXT(""));
+	FGLTFJsonAccessorIndex GetOrAddIndexAccessor(const FStaticMeshSection* MeshSection, const FRawStaticIndexBuffer* IndexBuffer, const FString& DesiredName = TEXT(""));
+
+	FGLTFJsonMeshIndex GetOrAddMesh(const FStaticMeshLODResources* StaticMeshLOD, const FColorVertexBuffer* OverrideVertexColors = nullptr, const FString& DesiredName = TEXT(""));
+	FGLTFJsonMeshIndex GetOrAddMesh(const UStaticMesh* StaticMesh, int32 LODIndex = 0, const FColorVertexBuffer* OverrideVertexColors = nullptr, const FString& DesiredName = TEXT(""));
+	FGLTFJsonMeshIndex GetOrAddMesh(const UStaticMeshComponent* StaticMeshComponent, const FString& DesiredName = TEXT(""));
+
+	FGLTFJsonNodeIndex GetOrAddNode(const USceneComponent* SceneComponent, bool bSelectedOnly, bool bRootNode = false, const FString& DesiredName = TEXT(""));
+	FGLTFJsonSceneIndex GetOrAddScene(const ULevel* Level, bool bSelectedOnly, const FString& DesiredName = TEXT(""));
+	FGLTFJsonSceneIndex GetOrAddScene(const UWorld* World, bool bSelectedOnly, const FString& DesiredName = TEXT(""));
 };
