[GLTFExporter] refactor: rename *Converter::Convert* -> *Converter::Add*

commit db4419348568163eae03ab426d20cafa89a16d3d
Date: Mon, 11 May 2020 22:34:35 +0200

    refactor: rename *Converter::Convert* -> *Converter::Add*
    (cherry picked from commit d281a2f0b14c7bb2dfefdabed1f018e1814862e8​)

#skipci

diff --git a/Engine/Plugins/Enterprise/GLTFExporter/MigratedFromGit/Source/GLTFExporter/Private/GLTFIndexedBuilder.cpp b/Engine/Plugins/Enterprise/GLTFExporter/MigratedFromGit/Source/GLTFExporter/Private/GLTFIndexedBuilder.cpp
--- a/Engine/Plugins/Enterprise/GLTFExporter/MigratedFromGit/Source/GLTFExporter/Private/GLTFIndexedBuilder.cpp
+++ b/Engine/Plugins/Enterprise/GLTFExporter/MigratedFromGit/Source/GLTFExporter/Private/GLTFIndexedBuilder.cpp
@@ -4,47 +4,47 @@
 
 FGLTFJsonAccessorIndex FGLTFIndexedBuilder::GetOrAddPositionAccessor(const FPositionVertexBuffer* VertexBuffer, const FString& DesiredName)
 {
-	return PositionVertexBuffers.Convert(*this, DesiredName, VertexBuffer);
+	return PositionVertexBuffers.GetOrAdd(*this, DesiredName, VertexBuffer);
 }
 
 FGLTFJsonAccessorIndex FGLTFIndexedBuilder::GetOrAddColorAccessor(const FColorVertexBuffer* VertexBuffer, const FString& DesiredName)
 {
-	return ColorVertexBuffers.Convert(*this, DesiredName, VertexBuffer);
+	return ColorVertexBuffers.GetOrAdd(*this, DesiredName, VertexBuffer);
 }
 
 FGLTFJsonAccessorIndex FGLTFIndexedBuilder::GetOrAddNormalAccessor(const FStaticMeshVertexBuffer* VertexBuffer, const FString& DesiredName)
 {
-	return StaticMeshNormalVertexBuffers.Convert(*this, DesiredName, VertexBuffer);
+	return StaticMeshNormalVertexBuffers.GetOrAdd(*this, DesiredName, VertexBuffer);
 }
 
 FGLTFJsonAccessorIndex FGLTFIndexedBuilder::GetOrAddTangentAccessor(const FStaticMeshVertexBuffer* VertexBuffer, const FString& DesiredName)
 {
-	return StaticMeshTangentVertexBuffers.Convert(*this, DesiredName, VertexBuffer);
+	return StaticMeshTangentVertexBuffers.GetOrAdd(*this, DesiredName, VertexBuffer);
 }
 
 FGLTFJsonAccessorIndex FGLTFIndexedBuilder::GetOrAddUV0Accessor(const FStaticMeshVertexBuffer* VertexBuffer, const FString& DesiredName)
 {
-	return StaticMeshUV0VertexBuffers.Convert(*this, DesiredName, VertexBuffer);
+	return StaticMeshUV0VertexBuffers.GetOrAdd(*this, DesiredName, VertexBuffer);
 }
 
 FGLTFJsonAccessorIndex FGLTFIndexedBuilder::GetOrAddUV1Accessor(const FStaticMeshVertexBuffer* VertexBuffer, const FString& DesiredName)
 {
-	return StaticMeshUV1VertexBuffers.Convert(*this, DesiredName, VertexBuffer);
+	return StaticMeshUV1VertexBuffers.GetOrAdd(*this, DesiredName, VertexBuffer);
 }
 
 FGLTFJsonBufferViewIndex FGLTFIndexedBuilder::GetOrAddIndexBufferView(const FRawStaticIndexBuffer* IndexBuffer, const FString& DesiredName)
 {
-	return StaticMeshIndexBuffers.Convert(*this, DesiredName, IndexBuffer);
+	return StaticMeshIndexBuffers.GetOrAdd(*this, DesiredName, IndexBuffer);
 }
 
 FGLTFJsonAccessorIndex FGLTFIndexedBuilder::GetOrAddIndexAccessor(const FStaticMeshSection* MeshSection, const FRawStaticIndexBuffer* IndexBuffer, const FString& DesiredName)
 {
-	return StaticMeshSections.Convert(*this, DesiredName, MeshSection, IndexBuffer);
+	return StaticMeshSections.GetOrAdd(*this, DesiredName, MeshSection, IndexBuffer);
 }
 
 FGLTFJsonMeshIndex FGLTFIndexedBuilder::GetOrAddMesh(const FStaticMeshLODResources* StaticMeshLOD, const FColorVertexBuffer* OverrideVertexColors, const FString& DesiredName)
 {
-	return StaticMeshes.Convert(*this, DesiredName, StaticMeshLOD, OverrideVertexColors);
+	return StaticMeshes.GetOrAdd(*this, DesiredName, StaticMeshLOD, OverrideVertexColors);
 }
 
 FGLTFJsonMeshIndex FGLTFIndexedBuilder::GetOrAddMesh(const UStaticMesh* StaticMesh, int32 LODIndex, const FColorVertexBuffer* OverrideVertexColors, const FString& DesiredName)
@@ -71,12 +71,12 @@ FGLTFJsonMeshIndex FGLTFIndexedBuilder::GetOrAddMesh(const UStaticMeshComponent*
 
 FGLTFJsonNodeIndex FGLTFIndexedBuilder::GetOrAddNode(const USceneComponent* SceneComponent, bool bSelectedOnly, bool bRootNode, const FString& DesiredName)
 {
-	return SceneComponents.Convert(*this, DesiredName, SceneComponent, bSelectedOnly, bRootNode);
+	return SceneComponents.GetOrAdd(*this, DesiredName, SceneComponent, bSelectedOnly, bRootNode);
 }
 
 FGLTFJsonSceneIndex FGLTFIndexedBuilder::GetOrAddScene(const ULevel* Level, bool bSelectedOnly, const FString& DesiredName)
 {
-	return Levels.Convert(*this, DesiredName, Level, bSelectedOnly);
+	return Levels.GetOrAdd(*this, DesiredName, Level, bSelectedOnly);
 }
 
 FGLTFJsonSceneIndex FGLTFIndexedBuilder::GetOrAddScene(const UWorld* World, bool bSelectedOnly, const FString& DesiredName)
diff --git a/Engine/Plugins/Enterprise/GLTFExporter/MigratedFromGit/Source/GLTFExporter/Private/GLTFLevelConverters.cpp b/Engine/Plugins/Enterprise/GLTFExporter/MigratedFromGit/Source/GLTFExporter/Private/GLTFLevelConverters.cpp
--- a/Engine/Plugins/Enterprise/GLTFExporter/MigratedFromGit/Source/GLTFExporter/Private/GLTFLevelConverters.cpp
+++ b/Engine/Plugins/Enterprise/GLTFExporter/MigratedFromGit/Source/GLTFExporter/Private/GLTFLevelConverters.cpp
@@ -6,7 +6,7 @@
 #include "GLTFContainerBuilder.h"
 #include "GLTFConverterUtility.h"
 
-FGLTFJsonNodeIndex FGLTFSceneComponentConverter::Convert(FGLTFIndexedBuilder& Builder, const FString& Name, const USceneComponent* SceneComponent, bool bSelectedOnly, bool bRootNode)
+FGLTFJsonNodeIndex FGLTFSceneComponentConverter::Add(FGLTFIndexedBuilder& Builder, const FString& Name, const USceneComponent* SceneComponent, bool bSelectedOnly, bool bRootNode)
 {
 	const AActor* Owner = SceneComponent->GetOwner();
 	const bool bIsRootComponent = Owner != nullptr && Owner->GetRootComponent() == SceneComponent;
@@ -53,7 +53,7 @@ FGLTFJsonNodeIndex FGLTFSceneComponentConverter::Convert(FGLTFIndexedBuilder& Bu
 	return Builder.AddNode(Node);
 }
 
-FGLTFJsonSceneIndex FGLTFLevelConverter::Convert(FGLTFIndexedBuilder& Builder, const FString& Name, const ULevel* Level, bool bSelectedOnly)
+FGLTFJsonSceneIndex FGLTFLevelConverter::Add(FGLTFIndexedBuilder& Builder, const FString& Name, const ULevel* Level, bool bSelectedOnly)
 {
 	FGLTFJsonScene Scene;
 	Scene.Name = Name;
diff --git a/Engine/Plugins/Enterprise/GLTFExporter/MigratedFromGit/Source/GLTFExporter/Private/GLTFStaticMeshConverters.cpp b/Engine/Plugins/Enterprise/GLTFExporter/MigratedFromGit/Source/GLTFExporter/Private/GLTFStaticMeshConverters.cpp
--- a/Engine/Plugins/Enterprise/GLTFExporter/MigratedFromGit/Source/GLTFExporter/Private/GLTFStaticMeshConverters.cpp
+++ b/Engine/Plugins/Enterprise/GLTFExporter/MigratedFromGit/Source/GLTFExporter/Private/GLTFStaticMeshConverters.cpp
@@ -4,7 +4,7 @@
 #include "GLTFContainerBuilder.h"
 #include "GLTFConverterUtility.h"
 
-FGLTFJsonAccessorIndex FGLTFPositionVertexBufferConverter::Convert(FGLTFIndexedBuilder& Builder, const FString& Name, const FPositionVertexBuffer* VertexBuffer)
+FGLTFJsonAccessorIndex FGLTFPositionVertexBufferConverter::Add(FGLTFIndexedBuilder& Builder, const FString& Name, const FPositionVertexBuffer* VertexBuffer)
 {
 	const uint32 VertexCount = VertexBuffer->GetNumVertices();
 	if (VertexCount == 0)
@@ -47,7 +47,7 @@ FGLTFJsonAccessorIndex FGLTFPositionVertexBufferConverter::Convert(FGLTFIndexedB
 	return Builder.AddAccessor(JsonAccessor);
 }
 
-FGLTFJsonAccessorIndex FGLTFColorVertexBufferConverter::Convert(FGLTFIndexedBuilder& Builder, const FString& Name, const FColorVertexBuffer* VertexBuffer)
+FGLTFJsonAccessorIndex FGLTFColorVertexBufferConverter::Add(FGLTFIndexedBuilder& Builder, const FString& Name, const FColorVertexBuffer* VertexBuffer)
 {
 	const uint32 VertexCount = VertexBuffer->GetNumVertices();
 	if (VertexCount == 0)
@@ -74,7 +74,7 @@ FGLTFJsonAccessorIndex FGLTFColorVertexBufferConverter::Convert(FGLTFIndexedBuil
 	return Builder.AddAccessor(JsonAccessor);
 }
 
-FGLTFJsonAccessorIndex FGLTFStaticMeshNormalVertexBufferConverter::Convert(FGLTFIndexedBuilder& Builder, const FString& Name, const FStaticMeshVertexBuffer* VertexBuffer)
+FGLTFJsonAccessorIndex FGLTFStaticMeshNormalVertexBufferConverter::Add(FGLTFIndexedBuilder& Builder, const FString& Name, const FStaticMeshVertexBuffer* VertexBuffer)
 {
 	const uint32 VertexCount = VertexBuffer->GetNumVertices();
 	if (VertexCount == 0)
@@ -100,7 +100,7 @@ FGLTFJsonAccessorIndex FGLTFStaticMeshNormalVertexBufferConverter::Convert(FGLTF
 	return Builder.AddAccessor(JsonAccessor);
 }
 
-FGLTFJsonAccessorIndex FGLTFStaticMeshTangentVertexBufferConverter::Convert(FGLTFIndexedBuilder& Builder, const FString& Name, const FStaticMeshVertexBuffer* VertexBuffer)
+FGLTFJsonAccessorIndex FGLTFStaticMeshTangentVertexBufferConverter::Add(FGLTFIndexedBuilder& Builder, const FString& Name, const FStaticMeshVertexBuffer* VertexBuffer)
 {
 	const uint32 VertexCount = VertexBuffer->GetNumVertices();
 	if (VertexCount == 0)
@@ -126,7 +126,7 @@ FGLTFJsonAccessorIndex FGLTFStaticMeshTangentVertexBufferConverter::Convert(FGLT
 	return Builder.AddAccessor(JsonAccessor);
 }
 
-FGLTFJsonAccessorIndex FGLTFStaticMeshUV0VertexBufferConverter::Convert(FGLTFIndexedBuilder& Builder, const FString& Name, const FStaticMeshVertexBuffer* VertexBuffer)
+FGLTFJsonAccessorIndex FGLTFStaticMeshUV0VertexBufferConverter::Add(FGLTFIndexedBuilder& Builder, const FString& Name, const FStaticMeshVertexBuffer* VertexBuffer)
 {
 	const uint32 VertexCount = VertexBuffer->GetNumVertices();
 	if (VertexCount == 0 || VertexBuffer->GetNumTexCoords() < 1)
@@ -152,7 +152,7 @@ FGLTFJsonAccessorIndex FGLTFStaticMeshUV0VertexBufferConverter::Convert(FGLTFInd
 	return Builder.AddAccessor(JsonAccessor);
 }
 
-FGLTFJsonAccessorIndex FGLTFStaticMeshUV1VertexBufferConverter::Convert(FGLTFIndexedBuilder& Builder, const FString& Name, const FStaticMeshVertexBuffer* VertexBuffer)
+FGLTFJsonAccessorIndex FGLTFStaticMeshUV1VertexBufferConverter::Add(FGLTFIndexedBuilder& Builder, const FString& Name, const FStaticMeshVertexBuffer* VertexBuffer)
 {
 	const uint32 VertexCount = VertexBuffer->GetNumVertices();
 	if (VertexCount == 0 || VertexBuffer->GetNumTexCoords() < 2)
@@ -178,7 +178,7 @@ FGLTFJsonAccessorIndex FGLTFStaticMeshUV1VertexBufferConverter::Convert(FGLTFInd
 	return Builder.AddAccessor(JsonAccessor);
 }
 
-FGLTFJsonBufferViewIndex FGLTFStaticMeshIndexBufferConverter::Convert(FGLTFIndexedBuilder& Builder, const FString& Name, const FRawStaticIndexBuffer* IndexBuffer)
+FGLTFJsonBufferViewIndex FGLTFStaticMeshIndexBufferConverter::Add(FGLTFIndexedBuilder& Builder, const FString& Name, const FRawStaticIndexBuffer* IndexBuffer)
 {
 	if (IndexBuffer->GetNumIndices() == 0)
 	{
@@ -190,7 +190,7 @@ FGLTFJsonBufferViewIndex FGLTFStaticMeshIndexBufferConverter::Convert(FGLTFIndex
 	return Builder.AddBufferView(Indices, Name, EGLTFJsonBufferTarget::ElementArrayBuffer);
 }
 
-FGLTFJsonAccessorIndex FGLTFStaticMeshSectionConverter::Convert(FGLTFIndexedBuilder& Builder, const FString& Name, const FStaticMeshSection* MeshSection, const FRawStaticIndexBuffer* IndexBuffer)
+FGLTFJsonAccessorIndex FGLTFStaticMeshSectionConverter::Add(FGLTFIndexedBuilder& Builder, const FString& Name, const FStaticMeshSection* MeshSection, const FRawStaticIndexBuffer* IndexBuffer)
 {
 	const uint32 TriangleCount = MeshSection->NumTriangles;
 	if (TriangleCount == 0)
@@ -209,7 +209,7 @@ FGLTFJsonAccessorIndex FGLTFStaticMeshSectionConverter::Convert(FGLTFIndexedBuil
 	return Builder.AddAccessor(JsonAccessor);
 }
 
-FGLTFJsonMeshIndex FGLTFStaticMeshConverter::Convert(FGLTFIndexedBuilder& Builder, const FString& Name, const FStaticMeshLODResources* StaticMeshLOD, const FColorVertexBuffer* OverrideVertexColors)
+FGLTFJsonMeshIndex FGLTFStaticMeshConverter::Add(FGLTFIndexedBuilder& Builder, const FString& Name, const FStaticMeshLODResources* StaticMeshLOD, const FColorVertexBuffer* OverrideVertexColors)
 {
 	FGLTFJsonMesh JsonMesh;
 	JsonMesh.Name = Name;
diff --git a/Engine/Plugins/Enterprise/GLTFExporter/MigratedFromGit/Source/GLTFExporter/Public/GLTFIndexedObjects.h b/Engine/Plugins/Enterprise/GLTFExporter/MigratedFromGit/Source/GLTFExporter/Public/GLTFIndexedObjects.h
--- a/Engine/Plugins/Enterprise/GLTFExporter/MigratedFromGit/Source/GLTFExporter/Public/GLTFIndexedObjects.h
+++ b/Engine/Plugins/Enterprise/GLTFExporter/MigratedFromGit/Source/GLTFExporter/Public/GLTFIndexedObjects.h
@@ -6,27 +6,27 @@
 
 struct FGLTFIndexedBuilder;
 
-template <class IndexType, class KeyType, class ConverterType>
+template <class IndexType, class KeyType, class AdderType>
 struct GLTFEXPORTER_API TGLTFIndexedObjects
 {
 	TMap<KeyType, IndexType> IndexLookup;
 
 	template <typename... ArgsType>
-	FORCEINLINE IndexType Find(ArgsType&&... Args) const
+	FORCEINLINE IndexType Get(ArgsType&&... Args) const
 	{
 		const KeyType Key(Forward<ArgsType>(Args)...);
 		return IndexLookup.FindRef(Key);
 	}
 
 	template <typename... ArgsType>
-	FORCEINLINE IndexType Convert(FGLTFIndexedBuilder& Builder, const FString& DesiredName, ArgsType&&... Args)
+	FORCEINLINE IndexType GetOrAdd(FGLTFIndexedBuilder& Builder, const FString& DesiredName, ArgsType&&... Args)
 	{
 		const KeyType Key(Forward<ArgsType>(Args)...);
 
 		IndexType Index = IndexLookup.FindRef(Key);
 		if (Index == INDEX_NONE)
 		{
-			Index = ConverterType::Convert(Builder, DesiredName, Forward<ArgsType>(Args)...);
+			Index = AdderType::Add(Builder, DesiredName, Forward<ArgsType>(Args)...);
 			IndexLookup.Add(Key, Index);
 		}
 
diff --git a/Engine/Plugins/Enterprise/GLTFExporter/MigratedFromGit/Source/GLTFExporter/Public/GLTFLevelConverters.h b/Engine/Plugins/Enterprise/GLTFExporter/MigratedFromGit/Source/GLTFExporter/Public/GLTFLevelConverters.h
--- a/Engine/Plugins/Enterprise/GLTFExporter/MigratedFromGit/Source/GLTFExporter/Public/GLTFLevelConverters.h
+++ b/Engine/Plugins/Enterprise/GLTFExporter/MigratedFromGit/Source/GLTFExporter/Public/GLTFLevelConverters.h
@@ -9,10 +9,10 @@ struct FGLTFIndexedBuilder;
 
 struct GLTFEXPORTER_API FGLTFSceneComponentConverter
 {
-	static FGLTFJsonNodeIndex Convert(FGLTFIndexedBuilder& Builder, const FString& Name, const USceneComponent* SceneComponent, bool bSelectedOnly, bool bRootNode = false);
+	static FGLTFJsonNodeIndex Add(FGLTFIndexedBuilder& Builder, const FString& Name, const USceneComponent* SceneComponent, bool bSelectedOnly, bool bRootNode = false);
 };
 
 struct GLTFEXPORTER_API FGLTFLevelConverter
 {
-	static FGLTFJsonSceneIndex Convert(FGLTFIndexedBuilder& Builder, const FString& Name, const ULevel* Level, bool bSelectedOnly);
+	static FGLTFJsonSceneIndex Add(FGLTFIndexedBuilder& Builder, const FString& Name, const ULevel* Level, bool bSelectedOnly);
 };
diff --git a/Engine/Plugins/Enterprise/GLTFExporter/MigratedFromGit/Source/GLTFExporter/Public/GLTFStaticMeshConverters.h b/Engine/Plugins/Enterprise/GLTFExporter/MigratedFromGit/Source/GLTFExporter/Public/GLTFStaticMeshConverters.h
--- a/Engine/Plugins/Enterprise/GLTFExporter/MigratedFromGit/Source/GLTFExporter/Public/GLTFStaticMeshConverters.h
+++ b/Engine/Plugins/Enterprise/GLTFExporter/MigratedFromGit/Source/GLTFExporter/Public/GLTFStaticMeshConverters.h
@@ -9,45 +9,45 @@ struct FGLTFIndexedBuilder;
 
 struct GLTFEXPORTER_API FGLTFPositionVertexBufferConverter
 {
-	static FGLTFJsonAccessorIndex Convert(FGLTFIndexedBuilder& Builder, const FString& Name, const FPositionVertexBuffer* VertexBuffer);
+	static FGLTFJsonAccessorIndex Add(FGLTFIndexedBuilder& Builder, const FString& Name, const FPositionVertexBuffer* VertexBuffer);
 };
 
 struct GLTFEXPORTER_API FGLTFColorVertexBufferConverter
 {
-	static FGLTFJsonAccessorIndex Convert(FGLTFIndexedBuilder& Builder, const FString& Name, const FColorVertexBuffer* VertexBuffer);
+	static FGLTFJsonAccessorIndex Add(FGLTFIndexedBuilder& Builder, const FString& Name, const FColorVertexBuffer* VertexBuffer);
 };
 
 struct GLTFEXPORTER_API FGLTFStaticMeshNormalVertexBufferConverter
 {
-	static FGLTFJsonAccessorIndex Convert(FGLTFIndexedBuilder& Builder, const FString& Name, const FStaticMeshVertexBuffer* VertexBuffer);
+	static FGLTFJsonAccessorIndex Add(FGLTFIndexedBuilder& Builder, const FString& Name, const FStaticMeshVertexBuffer* VertexBuffer);
 };
 
 struct GLTFEXPORTER_API FGLTFStaticMeshTangentVertexBufferConverter
 {
-	static FGLTFJsonAccessorIndex Convert(FGLTFIndexedBuilder& Builder, const FString& Name, const FStaticMeshVertexBuffer* VertexBuffer);
+	static FGLTFJsonAccessorIndex Add(FGLTFIndexedBuilder& Builder, const FString& Name, const FStaticMeshVertexBuffer* VertexBuffer);
 };
 
 struct GLTFEXPORTER_API FGLTFStaticMeshUV0VertexBufferConverter
 {
-	static FGLTFJsonAccessorIndex Convert(FGLTFIndexedBuilder& Builder, const FString& Name, const FStaticMeshVertexBuffer* VertexBuffer);
+	static FGLTFJsonAccessorIndex Add(FGLTFIndexedBuilder& Builder, const FString& Name, const FStaticMeshVertexBuffer* VertexBuffer);
 };
 
 struct GLTFEXPORTER_API FGLTFStaticMeshUV1VertexBufferConverter
 {
-	static FGLTFJsonAccessorIndex Convert(FGLTFIndexedBuilder& Builder, const FString& Name, const FStaticMeshVertexBuffer* VertexBuffer);
+	static FGLTFJsonAccessorIndex Add(FGLTFIndexedBuilder& Builder, const FString& Name, const FStaticMeshVertexBuffer* VertexBuffer);
 };
 
 struct GLTFEXPORTER_API FGLTFStaticMeshIndexBufferConverter
 {
-	static FGLTFJsonBufferViewIndex Convert(FGLTFIndexedBuilder& Builder, const FString& Name, const FRawStaticIndexBuffer* IndexBuffer);
+	static FGLTFJsonBufferViewIndex Add(FGLTFIndexedBuilder& Builder, const FString& Name, const FRawStaticIndexBuffer* IndexBuffer);
 };
 
 struct GLTFEXPORTER_API FGLTFStaticMeshSectionConverter
 {
-	static FGLTFJsonAccessorIndex Convert(FGLTFIndexedBuilder& Builder, const FString& Name, const FStaticMeshSection* MeshSection, const FRawStaticIndexBuffer* IndexBuffer);
+	static FGLTFJsonAccessorIndex Add(FGLTFIndexedBuilder& Builder, const FString& Name, const FStaticMeshSection* MeshSection, const FRawStaticIndexBuffer* IndexBuffer);
 };
 
 struct GLTFEXPORTER_API FGLTFStaticMeshConverter
 {
-	static FGLTFJsonMeshIndex Convert(FGLTFIndexedBuilder& Builder, const FString& Name, const FStaticMeshLODResources* StaticMeshLOD, const FColorVertexBuffer* OverrideVertexColors);
+	static FGLTFJsonMeshIndex Add(FGLTFIndexedBuilder& Builder, const FString& Name, const FStaticMeshLODResources* StaticMeshLOD, const FColorVertexBuffer* OverrideVertexColors);
 };
