[GLTFExporter] fix: fix minor index accessor naming quirk (should ignore adding "_Section#" prefix if only one section)

commit fa80cb583ee1adf2991d315014458f686d818a26
Date: Tue, 12 May 2020 08:52:32 +0200

    fix: fix minor index accessor naming quirk (should ignore adding "_Section#" prefix if only one section)
    (cherry picked from commit 6ce32f54a3aba37914afd734f27b6e30ca4002ca​)

#skipci

diff --git a/Engine/Plugins/Enterprise/GLTFExporter/MigratedFromGit/Source/GLTFExporter/Private/GLTFStaticMeshConverters.cpp b/Engine/Plugins/Enterprise/GLTFExporter/MigratedFromGit/Source/GLTFExporter/Private/GLTFStaticMeshConverters.cpp
--- a/Engine/Plugins/Enterprise/GLTFExporter/MigratedFromGit/Source/GLTFExporter/Private/GLTFStaticMeshConverters.cpp
+++ b/Engine/Plugins/Enterprise/GLTFExporter/MigratedFromGit/Source/GLTFExporter/Private/GLTFStaticMeshConverters.cpp
@@ -239,7 +239,7 @@ FGLTFJsonMeshIndex FGLTFStaticMeshConverter::Add(FGLTFIndexedBuilder& Builder, c
 		JsonPrimitive.Attributes = JsonAttributes;
 
 		JsonPrimitive.Indices = Builder.GetOrAddIndexAccessor(&Sections[SectionIndex], IndexBuffer,
-			Name.IsEmpty() ? Name : JsonMesh.Name + (SectionCount != 0 ? TEXT("_Indices_Section") + FString::FromInt(SectionIndex) : TEXT("_Indices")));
+			Name.IsEmpty() ? Name : JsonMesh.Name + (SectionCount != 1 ? TEXT("_Indices_Section") + FString::FromInt(SectionIndex) : TEXT("_Indices")));
 	}
 
 	return Builder.AddMesh(JsonMesh);
