[GLTFExporter] refactor: don't declare local reference to FStaticMeshLODResources::Sections since this API changes type in 4.25 and therefore more likely do so in the future as well

commit 063dfdc425f2dfe2a12be1e51584e6d7748fbddd
Date: Tue, 12 May 2020 09:47:24 +0200

    refactor: don't declare local reference to FStaticMeshLODResources::Sections since this API changes type in 4.25 and therefore more likely do so in the future as well
    (cherry picked from commit 8c2b4264e5cf17b27eb7d684c4a22e7b805049d2​)

#skipci

diff --git a/Engine/Plugins/Enterprise/GLTFExporter/MigratedFromGit/Source/GLTFExporter/Private/GLTFStaticMeshConverters.cpp b/Engine/Plugins/Enterprise/GLTFExporter/MigratedFromGit/Source/GLTFExporter/Private/GLTFStaticMeshConverters.cpp
--- a/Engine/Plugins/Enterprise/GLTFExporter/MigratedFromGit/Source/GLTFExporter/Private/GLTFStaticMeshConverters.cpp
+++ b/Engine/Plugins/Enterprise/GLTFExporter/MigratedFromGit/Source/GLTFExporter/Private/GLTFStaticMeshConverters.cpp
@@ -229,8 +229,7 @@ FGLTFJsonMeshIndex FGLTFStaticMeshConverter::Add(FGLTFIndexedBuilder& Builder, c
 	const FRawStaticIndexBuffer* IndexBuffer = &StaticMeshLOD->IndexBuffer;
 	Builder.GetOrAddIndexBufferView(IndexBuffer, Name.IsEmpty() ? Name : Name + TEXT("_Indices"));
 
-	const TArray<FStaticMeshSection>& Sections = StaticMeshLOD->Sections;
-	const int32 SectionCount = Sections.Num();
+	const int32 SectionCount = StaticMeshLOD->Sections.Num();
 	JsonMesh.Primitives.AddDefaulted(SectionCount);
 
 	for (int32 SectionIndex = 0; SectionIndex < SectionCount; ++SectionIndex)
@@ -238,7 +237,7 @@ FGLTFJsonMeshIndex FGLTFStaticMeshConverter::Add(FGLTFIndexedBuilder& Builder, c
 		FGLTFJsonPrimitive& JsonPrimitive = JsonMesh.Primitives[SectionIndex];
 		JsonPrimitive.Attributes = JsonAttributes;
 
-		JsonPrimitive.Indices = Builder.GetOrAddIndexAccessor(&Sections[SectionIndex], IndexBuffer,
+		JsonPrimitive.Indices = Builder.GetOrAddIndexAccessor(&StaticMeshLOD->Sections[SectionIndex], IndexBuffer,
 			Name.IsEmpty() ? Name : JsonMesh.Name + (SectionCount != 1 ? TEXT("_Indices_Section") + FString::FromInt(SectionIndex) : TEXT("_Indices")));
 	}
 
