[GLTFExporter] refactor(level): remove old commented include

commit 49ccf77e834c4ff85356aa9020cf211f1b06ad47
Date: Wed, 13 May 2020 17:24:49 +0200

    refactor(level): remove old commented include
    (cherry picked from commit 94ec583d22e13398bc64bae4f36a2bfdc732485d​)

#skipci

diff --git a/Engine/Plugins/Enterprise/GLTFExporter/MigratedFromGit/Source/GLTFExporter/Private/GLTFLevelConverters.cpp b/Engine/Plugins/Enterprise/GLTFExporter/MigratedFromGit/Source/GLTFExporter/Private/GLTFLevelConverters.cpp
--- a/Engine/Plugins/Enterprise/GLTFExporter/MigratedFromGit/Source/GLTFExporter/Private/GLTFLevelConverters.cpp
+++ b/Engine/Plugins/Enterprise/GLTFExporter/MigratedFromGit/Source/GLTFExporter/Private/GLTFLevelConverters.cpp
@@ -1,7 +1,5 @@
 // Copyright Epic Games, Inc. All Rights Reserved.
 
-//#include "GLTFLevelConverter.h"
-
 #include "GLTFLevelConverters.h"
 #include "GLTFContainerBuilder.h"
 #include "GLTFConverterUtility.h"
