[GLTFExporter] feat(bin-file): export binary data to separate external .bin file

commit d54c9c095a2e4d8de4f23901897287efaf7cfb7b
Date: Wed, 13 May 2020 23:45:05 +0200

    feat(bin-file): export binary data to separate external .bin file
    (cherry picked from commit d3a14d710395ba52926ba161e501665f7bbcbef7​)

#skipci

diff --git a/Engine/Plugins/Enterprise/GLTFExporter/MigratedFromGit/Source/GLTFExporter/Private/GLTFBufferBuilder.cpp b/Engine/Plugins/Enterprise/GLTFExporter/MigratedFromGit/Source/GLTFExporter/Private/GLTFBufferBuilder.cpp
--- a/Engine/Plugins/Enterprise/GLTFExporter/MigratedFromGit/Source/GLTFExporter/Private/GLTFBufferBuilder.cpp
+++ b/Engine/Plugins/Enterprise/GLTFExporter/MigratedFromGit/Source/GLTFExporter/Private/GLTFBufferBuilder.cpp
@@ -1,7 +1,8 @@
 // Copyright Epic Games, Inc. All Rights Reserved.
 
 #include "GLTFBufferBuilder.h"
-#include "Misc/Base64.h"
+#include "Misc/FileHelper.h"
+#include "Misc/Paths.h"
 
 FGLTFBufferBuilder::FGLTFBufferBuilder()
 {
@@ -22,22 +23,23 @@ FGLTFJsonBufferViewIndex FGLTFBufferBuilder::AddBufferView(const void* RawData,
 	return FGLTFJsonBuilder::AddBufferView(JsonBufferView);
 }
 
-void FGLTFBufferBuilder::UpdateBuffer()
+void FGLTFBufferBuilder::UpdateJsonBufferObject(const FString& BinaryFilePath)
 {
 	FGLTFJsonBuffer& JsonBuffer = JsonRoot.Buffers[BufferIndex];
-	const int32 ByteLength = BufferData.Num();
+	JsonBuffer.URI = FPaths::GetCleanFilename(BinaryFilePath);
+	JsonBuffer.ByteLength = BufferData.Num();
+}
 
-	if (JsonBuffer.ByteLength != ByteLength)
-	{
-		JsonBuffer.ByteLength = ByteLength;
+bool FGLTFBufferBuilder::Serialize(FArchive& Archive, const FString& FilePath)
+{
+	const FString BinaryFilePath = FPaths::ChangeExtension(FilePath, TEXT(".bin"));
 
-		const FString DataBase64 = FBase64::Encode(BufferData.GetData(), ByteLength);
-		JsonBuffer.URI = TEXT("data:application/octet-stream;base64,") + DataBase64;
+	if(!FFileHelper::SaveArrayToFile(BufferData, *BinaryFilePath))
+	{
+		// TODO: report error
+		return false;
 	}
-}
 
-void FGLTFBufferBuilder::Serialize(FArchive& Archive)
-{
-	UpdateBuffer();
-	FGLTFJsonBuilder::Serialize(Archive);
+	UpdateJsonBufferObject(BinaryFilePath);
+	return FGLTFJsonBuilder::Serialize(Archive, FilePath);
 }
diff --git a/Engine/Plugins/Enterprise/GLTFExporter/MigratedFromGit/Source/GLTFExporter/Private/GLTFJsonBuilder.cpp b/Engine/Plugins/Enterprise/GLTFExporter/MigratedFromGit/Source/GLTFExporter/Private/GLTFJsonBuilder.cpp
--- a/Engine/Plugins/Enterprise/GLTFExporter/MigratedFromGit/Source/GLTFExporter/Private/GLTFJsonBuilder.cpp
+++ b/Engine/Plugins/Enterprise/GLTFExporter/MigratedFromGit/Source/GLTFExporter/Private/GLTFJsonBuilder.cpp
@@ -40,7 +40,8 @@ FGLTFJsonSceneIndex FGLTFJsonBuilder::AddScene(const FGLTFJsonScene& JsonScene)
 	return FGLTFJsonSceneIndex(JsonRoot.Scenes.Add(JsonScene));
 }
 
-void FGLTFJsonBuilder::Serialize(FArchive& Archive)
+bool FGLTFJsonBuilder::Serialize(FArchive& Archive, const FString& FilePath)
 {
 	JsonRoot.Serialize(&Archive, true);
+	return true;
 }
diff --git a/Engine/Plugins/Enterprise/GLTFExporter/MigratedFromGit/Source/GLTFExporter/Private/GLTFLevelExporter.cpp b/Engine/Plugins/Enterprise/GLTFExporter/MigratedFromGit/Source/GLTFExporter/Private/GLTFLevelExporter.cpp
--- a/Engine/Plugins/Enterprise/GLTFExporter/MigratedFromGit/Source/GLTFExporter/Private/GLTFLevelExporter.cpp
+++ b/Engine/Plugins/Enterprise/GLTFExporter/MigratedFromGit/Source/GLTFExporter/Private/GLTFLevelExporter.cpp
@@ -24,6 +24,5 @@ bool UGLTFLevelExporter::ExportBinary(UObject* Object, const TCHAR* Type, FArchi
 	const FGLTFJsonSceneIndex SceneIndex = Builder.GetOrAddScene(World, bSelectedOnly);
 	Builder.JsonRoot.DefaultScene = SceneIndex;
 
-	Builder.Serialize(Archive);
-	return true;
+	return Builder.Serialize(Archive, CurrentFilename);
 }
diff --git a/Engine/Plugins/Enterprise/GLTFExporter/MigratedFromGit/Source/GLTFExporter/Private/GLTFStaticMeshExporter.cpp b/Engine/Plugins/Enterprise/GLTFExporter/MigratedFromGit/Source/GLTFExporter/Private/GLTFStaticMeshExporter.cpp
--- a/Engine/Plugins/Enterprise/GLTFExporter/MigratedFromGit/Source/GLTFExporter/Private/GLTFStaticMeshExporter.cpp
+++ b/Engine/Plugins/Enterprise/GLTFExporter/MigratedFromGit/Source/GLTFExporter/Private/GLTFStaticMeshExporter.cpp
@@ -35,6 +35,5 @@ bool UGLTFStaticMeshExporter::ExportBinary(UObject* Object, const TCHAR* Type, F
 
 	Builder.JsonRoot.DefaultScene = SceneIndex;
 
-	Builder.Serialize(Archive);
-	return true;
+	return Builder.Serialize(Archive, CurrentFilename);
 }
diff --git a/Engine/Plugins/Enterprise/GLTFExporter/MigratedFromGit/Source/GLTFExporter/Public/GLTFBufferBuilder.h b/Engine/Plugins/Enterprise/GLTFExporter/MigratedFromGit/Source/GLTFExporter/Public/GLTFBufferBuilder.h
--- a/Engine/Plugins/Enterprise/GLTFExporter/MigratedFromGit/Source/GLTFExporter/Public/GLTFBufferBuilder.h
+++ b/Engine/Plugins/Enterprise/GLTFExporter/MigratedFromGit/Source/GLTFExporter/Public/GLTFBufferBuilder.h
@@ -19,6 +19,6 @@ struct GLTFEXPORTER_API FGLTFBufferBuilder : public FGLTFJsonBuilder
 		return AddBufferView(Array.GetData(), Array.Num() * sizeof(ElementType), Name, BufferTarget);
 	}
 
-	void UpdateBuffer();
-	virtual void Serialize(FArchive& Archive) override;
+	void UpdateJsonBufferObject(const FString& BinaryFilePath);
+	virtual bool Serialize(FArchive& Archive, const FString& FilePath) override;
 };
diff --git a/Engine/Plugins/Enterprise/GLTFExporter/MigratedFromGit/Source/GLTFExporter/Public/GLTFJsonBuilder.h b/Engine/Plugins/Enterprise/GLTFExporter/MigratedFromGit/Source/GLTFExporter/Public/GLTFJsonBuilder.h
--- a/Engine/Plugins/Enterprise/GLTFExporter/MigratedFromGit/Source/GLTFExporter/Public/GLTFJsonBuilder.h
+++ b/Engine/Plugins/Enterprise/GLTFExporter/MigratedFromGit/Source/GLTFExporter/Public/GLTFJsonBuilder.h
@@ -18,5 +18,5 @@ struct GLTFEXPORTER_API FGLTFJsonBuilder
 	FGLTFJsonNodeIndex AddNode(const FGLTFJsonNode& JsonNode);
 	FGLTFJsonSceneIndex AddScene(const FGLTFJsonScene& JsonScene);
 
-	virtual void Serialize(FArchive& Archive);
+	virtual bool Serialize(FArchive& Archive, const FString& FilePath);
 };
