[GLTFExporter] refactor: add GLTFConverterUtility::ConvertNormal

commit 6eb155ed171b3cd5824e9f17fd2902e676388b62
Date: Sat, 16 May 2020 11:35:14 +0200

    refactor: add GLTFConverterUtility::ConvertNormal
    (cherry picked from commit f3a2a332f746833720fcac56a8c11e826765eca4​)

#skipci

diff --git a/Engine/Plugins/Enterprise/GLTFExporter/MigratedFromGit/Source/GLTFExporter/Public/GLTFConverterUtility.h b/Engine/Plugins/Enterprise/GLTFExporter/MigratedFromGit/Source/GLTFExporter/Public/GLTFConverterUtility.h
--- a/Engine/Plugins/Enterprise/GLTFExporter/MigratedFromGit/Source/GLTFExporter/Public/GLTFConverterUtility.h
+++ b/Engine/Plugins/Enterprise/GLTFExporter/MigratedFromGit/Source/GLTFExporter/Public/GLTFConverterUtility.h
@@ -2,7 +2,6 @@
 
 #pragma once
 
-#include "GLTFJsonEnums.h"
 #include "Engine/EngineTypes.h"
 
 struct GLTFEXPORTER_API FGLTFConverterUtility
@@ -24,10 +23,15 @@ struct GLTFEXPORTER_API FGLTFConverterUtility
 		return ConvertVector(Scale);
 	}
 
-	static inline FVector4 ConvertTangent(const FVector4& Tangent)
+	static inline FVector ConvertNormal(const FVector& Normal)
+	{
+		return ConvertVector(Normal);
+	}
+
+	static inline FVector4 ConvertTangent(const FVector& Tangent)
 	{
 		// glTF stores tangent as Vec4, with W component indicating handedness of tangent basis.
-		return FVector4(ConvertVector(Tangent), Tangent.W);
+		return FVector4(ConvertVector(Tangent), 1.0f);
 	}
 
 	static inline FColor ConvertColor(const FColor& Color)
