[GLTFExporter] refactor(static-mesh): use ConvertNormal and ConvertTangent

commit ff6ae79ac4d4fb510a29a3629ee1a1989ca239b8
Date: Sat, 16 May 2020 11:36:52 +0200

    refactor(static-mesh): use ConvertNormal and ConvertTangent
    (cherry picked from commit 20590ceaeea692a28409fd9be92c99d1e91b739b​)

#skipci

diff --git a/Engine/Plugins/Enterprise/GLTFExporter/MigratedFromGit/Source/GLTFExporter/Private/GLTFStaticMeshConverters.cpp b/Engine/Plugins/Enterprise/GLTFExporter/MigratedFromGit/Source/GLTFExporter/Private/GLTFStaticMeshConverters.cpp
--- a/Engine/Plugins/Enterprise/GLTFExporter/MigratedFromGit/Source/GLTFExporter/Private/GLTFStaticMeshConverters.cpp
+++ b/Engine/Plugins/Enterprise/GLTFExporter/MigratedFromGit/Source/GLTFExporter/Private/GLTFStaticMeshConverters.cpp
@@ -87,7 +87,7 @@ FGLTFJsonAccessorIndex FGLTFStaticMeshNormalVertexBufferConverter::Add(FGLTFInde
 
 	for (uint32 VertexIndex = 0; VertexIndex < VertexCount; ++VertexIndex)
 	{
-		Normals[VertexIndex] = FGLTFConverterUtility::ConvertVector(VertexBuffer->VertexTangentZ(VertexIndex));
+		Normals[VertexIndex] = FGLTFConverterUtility::ConvertNormal(VertexBuffer->VertexTangentZ(VertexIndex));
 	}
 
 	FGLTFJsonAccessor JsonAccessor;
@@ -113,7 +113,7 @@ FGLTFJsonAccessorIndex FGLTFStaticMeshTangentVertexBufferConverter::Add(FGLTFInd
 
 	for (uint32 VertexIndex = 0; VertexIndex < VertexCount; ++VertexIndex)
 	{
-		Tangents[VertexIndex] = FGLTFConverterUtility::ConvertVector(VertexBuffer->VertexTangentX(VertexIndex));
+		Tangents[VertexIndex] = FGLTFConverterUtility::ConvertTangent(VertexBuffer->VertexTangentX(VertexIndex));
 	}
 
 	FGLTFJsonAccessor JsonAccessor;
