[GLTFExporter] feat(options): add bTangentDataQuantization and bNormalizeUVCoordinates and rename bExtensionsRequired -> bAllExtensionsRequired

commit 754d9a5a29909195248e119b361719818e0b02e4
Date: Sat, 16 May 2020 11:58:49 +0200

    feat(options): add bTangentDataQuantization and bNormalizeUVCoordinates and rename bExtensionsRequired -> bAllExtensionsRequired
    (cherry picked from commit da3e62f1e480d8d0cbd531ddcbaa3fdee0fe9496​)

#skipci

diff --git a/Engine/Plugins/Enterprise/GLTFExporter/MigratedFromGit/Source/GLTFExporter/Private/GLTFExportOptions.cpp b/Engine/Plugins/Enterprise/GLTFExporter/MigratedFromGit/Source/GLTFExporter/Private/GLTFExportOptions.cpp
--- a/Engine/Plugins/Enterprise/GLTFExporter/MigratedFromGit/Source/GLTFExporter/Private/GLTFExportOptions.cpp
+++ b/Engine/Plugins/Enterprise/GLTFExporter/MigratedFromGit/Source/GLTFExporter/Private/GLTFExportOptions.cpp
@@ -18,8 +18,10 @@ UGLTFExportOptions::UGLTFExportOptions(const FObjectInitializer& ObjectInitializ
 {
 	bBundleWebViewer = true;
 	bExportPreviewMesh = true;
-	bExtensionsRequired = true;
+	bAllExtensionsRequired = true;
 	bExportVertexColors = true;
+	bTangentDataQuantization = false;
+	bNormalizeUVCoordinates = EGLTFExporterNormalizeUVCoordinates::Auto;
 	DefaultLevelOfDetail = 0;
 	bMapSkeletalMotionToRoot = false;
 	bExportUnlitMaterials = true;
diff --git a/Engine/Plugins/Enterprise/GLTFExporter/MigratedFromGit/Source/GLTFExporter/Public/GLTFExportOptions.h b/Engine/Plugins/Enterprise/GLTFExporter/MigratedFromGit/Source/GLTFExporter/Public/GLTFExportOptions.h
--- a/Engine/Plugins/Enterprise/GLTFExporter/MigratedFromGit/Source/GLTFExporter/Public/GLTFExportOptions.h
+++ b/Engine/Plugins/Enterprise/GLTFExporter/MigratedFromGit/Source/GLTFExporter/Public/GLTFExportOptions.h
@@ -5,6 +5,14 @@
 #include "CoreMinimal.h"
 #include "GLTFExportOptions.generated.h"
 
+UENUM(BlueprintType)
+enum class EGLTFExporterNormalizeUVCoordinates : uint8
+{
+	Always,
+    Auto,
+	Never,
+};
+
 UENUM(BlueprintType)
 enum class EGLTFExporterTextureFormat : uint8
 {
@@ -54,11 +62,17 @@ public:
 	uint32 bExportPreviewMesh : 1;
 
 	UPROPERTY(EditAnywhere, BlueprintReadWrite, config, Category = Exporter)
-	uint32 bExtensionsRequired : 1;
+	uint32 bAllExtensionsRequired : 1;
 
 	UPROPERTY(EditAnywhere, BlueprintReadWrite, config, category = Mesh)
 	uint32 bExportVertexColors : 1;
 
+	UPROPERTY(EditAnywhere, BlueprintReadWrite, config, Category = Mesh)
+	uint32 bTangentDataQuantization : 1;
+
+	UPROPERTY(EditAnywhere, BlueprintReadWrite, config, Category = Mesh)
+    EGLTFExporterNormalizeUVCoordinates bNormalizeUVCoordinates;
+
 	UPROPERTY(EditAnywhere, BlueprintReadWrite, config, category = Mesh)
 	int32 DefaultLevelOfDetail;
 
