[GLTFExporter] feat: add support for listing extensions

commit f8549c0c410084af892727ad12b3064671fe41c7
Date: Sun, 17 May 2020 14:03:39 +0200

    feat: add support for listing extensions
    (cherry picked from commit e597b0c8be4553711be7b9cc8a0e5806794f3c86​)

#skipci

diff --git a/Engine/Plugins/Enterprise/GLTFExporter/MigratedFromGit/Source/GLTFExporter/Public/GLTFJsonEnums.h b/Engine/Plugins/Enterprise/GLTFExporter/MigratedFromGit/Source/GLTFExporter/Public/GLTFJsonEnums.h
--- a/Engine/Plugins/Enterprise/GLTFExporter/MigratedFromGit/Source/GLTFExporter/Public/GLTFJsonEnums.h
+++ b/Engine/Plugins/Enterprise/GLTFExporter/MigratedFromGit/Source/GLTFExporter/Public/GLTFJsonEnums.h
@@ -4,6 +4,14 @@
 
 #include "CoreMinimal.h"
 
+enum class EGLTFExtension
+{
+	KHR_LightsPunctual,
+	KHR_MaterialsUnlit,
+	KHR_MaterialsClearCoat,
+	KHR_MeshQuantization
+};
+
 enum class EGLTFJsonAccessorType
 {
 	None = -1,
diff --git a/Engine/Plugins/Enterprise/GLTFExporter/MigratedFromGit/Source/GLTFExporter/Public/GLTFJsonRoot.h b/Engine/Plugins/Enterprise/GLTFExporter/MigratedFromGit/Source/GLTFExporter/Public/GLTFJsonRoot.h
--- a/Engine/Plugins/Enterprise/GLTFExporter/MigratedFromGit/Source/GLTFExporter/Public/GLTFJsonRoot.h
+++ b/Engine/Plugins/Enterprise/GLTFExporter/MigratedFromGit/Source/GLTFExporter/Public/GLTFJsonRoot.h
@@ -50,7 +50,11 @@ struct GLTFEXPORTER_API FGLTFJsonAsset
 
 struct GLTFEXPORTER_API FGLTFJsonRoot
 {
-	FGLTFJsonAsset      Asset;
+	FGLTFJsonAsset Asset;
+
+	TSet<EGLTFExtension> ExtensionsUsed;
+	TSet<EGLTFExtension> ExtensionsRequired;
+
 	FGLTFJsonSceneIndex DefaultScene;
 
 	TArray<FGLTFJsonAccessor>   Accessors;
@@ -73,6 +77,28 @@ struct GLTFEXPORTER_API FGLTFJsonRoot
 		JsonWriter.WriteIdentifierPrefix(TEXT("asset"));
 		Asset.WriteObject(JsonWriter);
 
+		if (ExtensionsUsed.Num() > 0)
+		{
+			JsonWriter.WriteIdentifierPrefix(TEXT("extensionsUsed"));
+			JsonWriter.WriteArrayStart();
+			for (const EGLTFExtension& Extension : ExtensionsUsed)
+			{
+				JsonWriter.WriteValue(FGLTFJsonUtility::ExtensionToString(Extension));
+			}
+			JsonWriter.WriteArrayEnd();
+		}
+
+		if (ExtensionsRequired.Num() > 0)
+		{
+			JsonWriter.WriteIdentifierPrefix(TEXT("extensionsRequired"));
+			JsonWriter.WriteArrayStart();
+			for (const EGLTFExtension& Extension : ExtensionsRequired)
+			{
+				JsonWriter.WriteValue(FGLTFJsonUtility::ExtensionToString(Extension));
+			}
+			JsonWriter.WriteArrayEnd();
+		}
+
 		if (DefaultScene != INDEX_NONE)
 		{
 			JsonWriter.WriteValue(TEXT("scene"), DefaultScene);
diff --git a/Engine/Plugins/Enterprise/GLTFExporter/MigratedFromGit/Source/GLTFExporter/Public/GLTFJsonUtility.h b/Engine/Plugins/Enterprise/GLTFExporter/MigratedFromGit/Source/GLTFExporter/Public/GLTFJsonUtility.h
--- a/Engine/Plugins/Enterprise/GLTFExporter/MigratedFromGit/Source/GLTFExporter/Public/GLTFJsonUtility.h
+++ b/Engine/Plugins/Enterprise/GLTFExporter/MigratedFromGit/Source/GLTFExporter/Public/GLTFJsonUtility.h
@@ -6,6 +6,18 @@
 
 struct GLTFEXPORTER_API FGLTFJsonUtility
 {
+	static inline const TCHAR* ExtensionToString(EGLTFExtension Extension)
+	{
+		switch (Extension)
+		{
+			case EGLTFExtension::KHR_LightsPunctual:     return TEXT("KHR_lights_punctual");
+			case EGLTFExtension::KHR_MaterialsUnlit:     return TEXT("KHR_materials_unlit");
+			case EGLTFExtension::KHR_MaterialsClearCoat: return TEXT("KHR_materials_clearcoat");
+			case EGLTFExtension::KHR_MeshQuantization:   return TEXT("KHR_mesh_quantization");
+			default:                                     return TEXT("unknown");
+		}
+	}
+
 	static inline const TCHAR* AlphaModeToString(EGLTFJsonAlphaMode Mode)
 	{
 		switch (Mode)
