[GLTFExporter] feat(options): add option for normal export method

commit 94f45ed8d58a2106b8aea91d8e089cfe3049a1d4
Date: Sun, 17 May 2020 13:56:40 +0200

    feat(options): add option for normal export method
    (cherry picked from commit d481f1f88036a284074da1e4aeccb64b2e1ce9ac​)

#skipci

diff --git a/Engine/Plugins/Enterprise/GLTFExporter/MigratedFromGit/Source/GLTFExporter/Private/GLTFExportOptions.cpp b/Engine/Plugins/Enterprise/GLTFExporter/MigratedFromGit/Source/GLTFExporter/Private/GLTFExportOptions.cpp
--- a/Engine/Plugins/Enterprise/GLTFExporter/MigratedFromGit/Source/GLTFExporter/Private/GLTFExportOptions.cpp
+++ b/Engine/Plugins/Enterprise/GLTFExporter/MigratedFromGit/Source/GLTFExporter/Private/GLTFExportOptions.cpp
@@ -20,6 +20,7 @@ UGLTFExportOptions::UGLTFExportOptions(const FObjectInitializer& ObjectInitializ
 	bExportPreviewMesh = true;
 	bAllExtensionsRequired = true;
 	bExportVertexColors = true;
+	NormalExportMethod = EGLTFExporterNormalExportMethod::NormalsAndTangents;
 	bTangentDataQuantization = false;
 	bNormalizeUVCoordinates = EGLTFExporterNormalizeUVCoordinates::Auto;
 	DefaultLevelOfDetail = 0;
diff --git a/Engine/Plugins/Enterprise/GLTFExporter/MigratedFromGit/Source/GLTFExporter/Public/GLTFExportOptions.h b/Engine/Plugins/Enterprise/GLTFExporter/MigratedFromGit/Source/GLTFExporter/Public/GLTFExportOptions.h
--- a/Engine/Plugins/Enterprise/GLTFExporter/MigratedFromGit/Source/GLTFExporter/Public/GLTFExportOptions.h
+++ b/Engine/Plugins/Enterprise/GLTFExporter/MigratedFromGit/Source/GLTFExporter/Public/GLTFExportOptions.h
@@ -13,6 +13,14 @@ enum class EGLTFExporterNormalizeUVCoordinates : uint8
 	Never,
 };
 
+UENUM(BlueprintType)
+enum class EGLTFExporterNormalExportMethod : uint8
+{
+	None UMETA(DisplayName = "None"),
+	Normals UMETA(DisplayName = "Normals"),
+	NormalsAndTangents UMETA(DisplayName = "Normals and Tangents")
+};
+
 UENUM(BlueprintType)
 enum class EGLTFExporterTextureFormat : uint8
 {
@@ -67,6 +75,9 @@ public:
 	UPROPERTY(EditAnywhere, BlueprintReadWrite, config, category = Mesh)
 	uint32 bExportVertexColors : 1;
 
+	UPROPERTY(EditAnywhere, BlueprintReadWrite, config, Category = Mesh)
+	EGLTFExporterNormalExportMethod NormalExportMethod;
+
 	UPROPERTY(EditAnywhere, BlueprintReadWrite, config, Category = Mesh)
 	uint32 bTangentDataQuantization : 1;
 
