[GLTFExporter] refactor(options): minor code cleanup and tweaking order of options

commit 86ef42e81e471e91065788cb68e729d0df724321
Date: Sun, 17 May 2020 14:00:41 +0200

    refactor(options): minor code cleanup and tweaking order of options
    (cherry picked from commit 070d8cb0a51b107a64c7b7d900e4959f8907d32d​)

#skipci

diff --git a/Engine/Plugins/Enterprise/GLTFExporter/MigratedFromGit/Source/GLTFExporter/Private/GLTFExportOptions.cpp b/Engine/Plugins/Enterprise/GLTFExporter/MigratedFromGit/Source/GLTFExporter/Private/GLTFExportOptions.cpp
--- a/Engine/Plugins/Enterprise/GLTFExporter/MigratedFromGit/Source/GLTFExporter/Private/GLTFExportOptions.cpp
+++ b/Engine/Plugins/Enterprise/GLTFExporter/MigratedFromGit/Source/GLTFExporter/Private/GLTFExportOptions.cpp
@@ -28,18 +28,18 @@ UGLTFExportOptions::UGLTFExportOptions(const FObjectInitializer& ObjectInitializ
 	bExportUnlitMaterials = true;
 	bExportClearCoatMaterials = true;
 	bBakeMaterialInputs = true;
-	BakedMaterialInputSize = EGLTFExporterTextureSize::POT_512;
+	DefaultMaterialBakeSize = EGLTFExporterMaterialBakeSize::POT_1024;
 	TextureFormat = EGLTFExporterTextureFormat::PNG;
-	TextureHDREncoding = EGLTFExporterTextureHDREncoding::RGBD;
 	bExportLightmaps = true;
+	TextureHDREncoding = EGLTFExporterTextureHDREncoding::RGBD;
 	ExportScale = 0.01;
 	bExportLights = true;
 	bExportCameras = true;
+	bExportOrbitalCameras = true;
 	bExportReflectionCaptures = true;
 	bExportHDRIBackdrops = true;
 	bExportVariantSets = true;
 	bExportInteractionHotspots = true;
-	bExportOrbitalCameras = true;
 }
 
 void UGLTFExportOptions::ResetToDefault()
diff --git a/Engine/Plugins/Enterprise/GLTFExporter/MigratedFromGit/Source/GLTFExporter/Public/GLTFExportOptions.h b/Engine/Plugins/Enterprise/GLTFExporter/MigratedFromGit/Source/GLTFExporter/Public/GLTFExportOptions.h
--- a/Engine/Plugins/Enterprise/GLTFExporter/MigratedFromGit/Source/GLTFExporter/Public/GLTFExportOptions.h
+++ b/Engine/Plugins/Enterprise/GLTFExporter/MigratedFromGit/Source/GLTFExporter/Public/GLTFExportOptions.h
@@ -9,8 +9,8 @@ UENUM(BlueprintType)
 enum class EGLTFExporterNormalizeUVCoordinates : uint8
 {
 	Always,
-    Auto,
-	Never,
+	Auto,
+	Never
 };
 
 UENUM(BlueprintType)
@@ -22,24 +22,7 @@ enum class EGLTFExporterNormalExportMethod : uint8
 };
 
 UENUM(BlueprintType)
-enum class EGLTFExporterTextureFormat : uint8
-{
-	PNG,
-	JPEG
-};
-
-UENUM(BlueprintType)
-enum class EGLTFExporterTextureHDREncoding : uint8
-{
-	None,
-	HDR,
-	RGBM,
-	RGBD,
-	RGBE
-};
-
-UENUM(BlueprintType)
-enum class EGLTFExporterTextureSize : uint8
+enum class EGLTFExporterMaterialBakeSize : uint8
 {
 	POT_1 UMETA(DisplayName = "1 x 1"),
 	POT_2 UMETA(DisplayName = "2 x 2"),
@@ -57,6 +40,24 @@ enum class EGLTFExporterTextureSize : uint8
 	POT_8192 UMETA(DisplayName = "8192 x 8192")
 };
 
+UENUM(BlueprintType)
+enum class EGLTFExporterTextureFormat : uint8
+{
+	None,
+	PNG,
+	JPEG
+};
+
+UENUM(BlueprintType)
+enum class EGLTFExporterTextureHDREncoding : uint8
+{
+	None,
+	RGBM,
+	RGBD,
+	RGBE,
+	HDR
+};
+
 UCLASS(config = EditorPerProjectUserSettings, HideCategories = (DebugProperty))
 class GLTFEXPORTER_API UGLTFExportOptions : public UObject
 {
@@ -66,13 +67,13 @@ public:
 	UPROPERTY(EditAnywhere, BlueprintReadWrite, config, Category = Exporter)
 	uint32 bBundleWebViewer : 1;
 
-	UPROPERTY(EditAnywhere, BlueprintReadWrite, config, category = Exporter)
+	UPROPERTY(EditAnywhere, BlueprintReadWrite, config, Category = Exporter)
 	uint32 bExportPreviewMesh : 1;
 
 	UPROPERTY(EditAnywhere, BlueprintReadWrite, config, Category = Exporter)
 	uint32 bAllExtensionsRequired : 1;
 
-	UPROPERTY(EditAnywhere, BlueprintReadWrite, config, category = Mesh)
+	UPROPERTY(EditAnywhere, BlueprintReadWrite, config, Category = Mesh)
 	uint32 bExportVertexColors : 1;
 
 	UPROPERTY(EditAnywhere, BlueprintReadWrite, config, Category = Mesh)
@@ -82,12 +83,12 @@ public:
 	uint32 bTangentDataQuantization : 1;
 
 	UPROPERTY(EditAnywhere, BlueprintReadWrite, config, Category = Mesh)
-    EGLTFExporterNormalizeUVCoordinates bNormalizeUVCoordinates;
+	EGLTFExporterNormalizeUVCoordinates bNormalizeUVCoordinates;
 
-	UPROPERTY(EditAnywhere, BlueprintReadWrite, config, category = Mesh)
+	UPROPERTY(EditAnywhere, BlueprintReadWrite, config, Category = Mesh)
 	int32 DefaultLevelOfDetail;
 
-	UPROPERTY(EditAnywhere, BlueprintReadWrite, config, category = Animation)
+	UPROPERTY(EditAnywhere, BlueprintReadWrite, config, Category = Animation)
 	uint32 bMapSkeletalMotionToRoot : 1;
 
 	UPROPERTY(EditAnywhere, BlueprintReadWrite, config, Category = Material)
@@ -96,20 +97,20 @@ public:
 	UPROPERTY(EditAnywhere, BlueprintReadWrite, config, Category = Material)
 	uint32 bExportClearCoatMaterials : 1;
 
-	UPROPERTY(EditAnywhere, BlueprintReadWrite, config, category = Material)
+	UPROPERTY(EditAnywhere, BlueprintReadWrite, config, Category = Material)
 	uint32 bBakeMaterialInputs : 1;
 
-	UPROPERTY(EditAnywhere, BlueprintReadWrite, config, category = Material)
-	EGLTFExporterTextureSize BakedMaterialInputSize;
+	UPROPERTY(EditAnywhere, BlueprintReadWrite, config, Category = Material)
+	EGLTFExporterMaterialBakeSize DefaultMaterialBakeSize;
 
 	UPROPERTY(EditAnywhere, BlueprintReadWrite, config, Category = Texture)
 	EGLTFExporterTextureFormat TextureFormat;
 
 	UPROPERTY(EditAnywhere, BlueprintReadWrite, config, Category = Texture)
-	EGLTFExporterTextureHDREncoding TextureHDREncoding;
+	uint32 bExportLightmaps : 1;
 
 	UPROPERTY(EditAnywhere, BlueprintReadWrite, config, Category = Texture)
-	uint32 bExportLightmaps : 1;
+	EGLTFExporterTextureHDREncoding TextureHDREncoding;
 
 	UPROPERTY(EditAnywhere, BlueprintReadWrite, config, Category = Scene, meta =
 		(DisplayName = "Export Uniform Scale",
@@ -122,6 +123,9 @@ public:
 	UPROPERTY(EditAnywhere, BlueprintReadWrite, config, Category = Scene)
 	uint32 bExportCameras : 1;
 
+	UPROPERTY(EditAnywhere, BlueprintReadWrite, config, Category = Scene)
+	uint32 bExportOrbitalCameras : 1;
+
 	UPROPERTY(EditAnywhere, BlueprintReadWrite, config, Category = Scene)
 	uint32 bExportReflectionCaptures : 1;
 
@@ -134,9 +138,6 @@ public:
 	UPROPERTY(EditAnywhere, BlueprintReadWrite, config, Category = Scene)
 	uint32 bExportInteractionHotspots : 1;
 
-	UPROPERTY(EditAnywhere, BlueprintReadWrite, config, Category = Scene)
-    uint32 bExportOrbitalCameras : 1;
-
 	/* Set all the each property to the CDO value */
 	void ResetToDefault();
 
