[GLTFExporter] fix(options): change option defaults to represent current "hard-coded" settings

commit 9d7c8340fee808737ececabbecbd9aa3ed1b2a67
Date: Sun, 17 May 2020 14:15:39 +0200

    fix(options): change option defaults to represent current "hard-coded" settings
    (cherry picked from commit 675d5fed33198bd6b37dd06f2b9a0e1b2f845d9e​)

#skipci

diff --git a/Engine/Plugins/Enterprise/GLTFExporter/MigratedFromGit/Source/GLTFExporter/Private/GLTFExportOptions.cpp b/Engine/Plugins/Enterprise/GLTFExporter/MigratedFromGit/Source/GLTFExporter/Private/GLTFExportOptions.cpp
--- a/Engine/Plugins/Enterprise/GLTFExporter/MigratedFromGit/Source/GLTFExporter/Private/GLTFExportOptions.cpp
+++ b/Engine/Plugins/Enterprise/GLTFExporter/MigratedFromGit/Source/GLTFExporter/Private/GLTFExportOptions.cpp
@@ -16,30 +16,30 @@
 UGLTFExportOptions::UGLTFExportOptions(const FObjectInitializer& ObjectInitializer)
 	: Super(ObjectInitializer)
 {
-	bBundleWebViewer = true;
-	bExportPreviewMesh = true;
-	bAllExtensionsRequired = true;
+	bBundleWebViewer = false;
+	bExportPreviewMesh = false;
+	bAllExtensionsRequired = false;
 	bExportVertexColors = true;
 	NormalExportMethod = EGLTFExporterNormalExportMethod::NormalsAndTangents;
 	bTangentDataQuantization = false;
-	bNormalizeUVCoordinates = EGLTFExporterNormalizeUVCoordinates::Auto;
+	bNormalizeUVCoordinates = EGLTFExporterNormalizeUVCoordinates::Never;
 	DefaultLevelOfDetail = 0;
 	bMapSkeletalMotionToRoot = false;
-	bExportUnlitMaterials = true;
-	bExportClearCoatMaterials = true;
-	bBakeMaterialInputs = true;
+	bExportUnlitMaterials = false;
+	bExportClearCoatMaterials = false;
+	bBakeMaterialInputs = false;
 	DefaultMaterialBakeSize = EGLTFExporterMaterialBakeSize::POT_1024;
-	TextureFormat = EGLTFExporterTextureFormat::PNG;
-	bExportLightmaps = true;
-	TextureHDREncoding = EGLTFExporterTextureHDREncoding::RGBD;
+	TextureFormat = EGLTFExporterTextureFormat::None;
+	bExportLightmaps = false;
+	TextureHDREncoding = EGLTFExporterTextureHDREncoding::None;
 	ExportScale = 0.01;
-	bExportLights = true;
-	bExportCameras = true;
-	bExportOrbitalCameras = true;
-	bExportReflectionCaptures = true;
-	bExportHDRIBackdrops = true;
-	bExportVariantSets = true;
-	bExportInteractionHotspots = true;
+	bExportLights = false;
+	bExportCameras = false;
+	bExportOrbitalCameras = false;
+	bExportReflectionCaptures = false;
+	bExportHDRIBackdrops = false;
+	bExportVariantSets = false;
+	bExportInteractionHotspots = false;
 }
 
 void UGLTFExportOptions::ResetToDefault()
