[GLTFExporter] fix(options): disable saving/loading export options until actually used

commit 7921b1fe6b41d5d0043e7940b5c6f43cf281171b
Date: Sun, 17 May 2020 15:20:51 +0200

    fix(options): disable saving/loading export options until actually used
    (cherry picked from commit 1ad07fb7569ff4c4e179d956503b012b4307097d​)

#skipci

diff --git a/Engine/Plugins/Enterprise/GLTFExporter/MigratedFromGit/Source/GLTFExporter/Private/GLTFExportOptions.cpp b/Engine/Plugins/Enterprise/GLTFExporter/MigratedFromGit/Source/GLTFExporter/Private/GLTFExportOptions.cpp
--- a/Engine/Plugins/Enterprise/GLTFExporter/MigratedFromGit/Source/GLTFExporter/Private/GLTFExportOptions.cpp
+++ b/Engine/Plugins/Enterprise/GLTFExporter/MigratedFromGit/Source/GLTFExporter/Private/GLTFExportOptions.cpp
@@ -186,7 +186,7 @@ void UGLTFExportOptions::FillOptions(bool bBatchMode, bool bShowOptionDialog, co
 {
 	bOutOperationCanceled = false;
 
-	LoadOptions();
+	// TODO: call LoadOptions();
 
 	//Return if we do not show the export options or we are running automation test or we are unattended
 	if (!bShowOptionDialog || GIsAutomationTesting || FApp::IsUnattended())
@@ -197,5 +197,5 @@ void UGLTFExportOptions::FillOptions(bool bBatchMode, bool bShowOptionDialog, co
 	bOutExportAll = false;
 
 	SGLTFExportOptionsWindow::ShowDialog(this, FullPath, bBatchMode, bOutOperationCanceled, bOutExportAll);
-	SaveOptions();
+	// TODO: call SaveOptions();
 }
