[GLTFExporter] fix: conversion from UE endian-dependenct color layout to endian-independent glTF

commit 3113537fea4df1e4aa8bf876460612951c492f63
Date: Sun, 17 May 2020 15:09:43 +0200

    fix: conversion from UE endian-dependenct color layout to endian-independent glTF
    (cherry picked from commit 829f77db65697bc3af65fd6cbbdb449f39566a72​)

#skipci

diff --git a/Engine/Plugins/Enterprise/GLTFExporter/MigratedFromGit/Source/GLTFExporter/Private/GLTFStaticMeshConverters.cpp b/Engine/Plugins/Enterprise/GLTFExporter/MigratedFromGit/Source/GLTFExporter/Private/GLTFStaticMeshConverters.cpp
--- a/Engine/Plugins/Enterprise/GLTFExporter/MigratedFromGit/Source/GLTFExporter/Private/GLTFStaticMeshConverters.cpp
+++ b/Engine/Plugins/Enterprise/GLTFExporter/MigratedFromGit/Source/GLTFExporter/Private/GLTFStaticMeshConverters.cpp
@@ -55,7 +55,7 @@ FGLTFJsonAccessorIndex FGLTFColorVertexBufferConverter::Add(FGLTFIndexedBuilder&
 		return FGLTFJsonAccessorIndex(INDEX_NONE);
 	}
 
-	TArray<FColor> Colors;
+	TArray<FGLTFColor> Colors;
 	Colors.AddUninitialized(VertexCount);
 
 	for (uint32 VertexIndex = 0; VertexIndex < VertexCount; ++VertexIndex)
diff --git a/Engine/Plugins/Enterprise/GLTFExporter/MigratedFromGit/Source/GLTFExporter/Public/GLTFColor.h b/Engine/Plugins/Enterprise/GLTFExporter/MigratedFromGit/Source/GLTFExporter/Public/GLTFColor.h
new file mode 100644
--- /dev/null
+++ b/Engine/Plugins/Enterprise/GLTFExporter/MigratedFromGit/Source/GLTFExporter/Public/GLTFColor.h
@@ -0,0 +1,17 @@
+// Copyright Epic Games, Inc. All Rights Reserved.
+
+#pragma once
+
+#include "CoreTypes.h"
+
+// UE4 uses an endianness dependent memory layout (ABGR for little endian, ARGB for big endian)
+// glTF uses RGBA independent of endianness, hence why we need this dedicated struct
+struct GLTFEXPORTER_API FGLTFColor
+{
+	uint8 R,G,B,A;
+
+	FORCEINLINE FGLTFColor(uint8 R, uint8 G, uint8 B, uint8 A = 255)
+		: R(R), G(G), B(B), A(A)
+	{
+	}
+};
diff --git a/Engine/Plugins/Enterprise/GLTFExporter/MigratedFromGit/Source/GLTFExporter/Public/GLTFConverterUtility.h b/Engine/Plugins/Enterprise/GLTFExporter/MigratedFromGit/Source/GLTFExporter/Public/GLTFConverterUtility.h
--- a/Engine/Plugins/Enterprise/GLTFExporter/MigratedFromGit/Source/GLTFExporter/Public/GLTFConverterUtility.h
+++ b/Engine/Plugins/Enterprise/GLTFExporter/MigratedFromGit/Source/GLTFExporter/Public/GLTFConverterUtility.h
@@ -2,6 +2,7 @@
 
 #pragma once
 
+#include "GLTFColor.h"
 #include "Engine/EngineTypes.h"
 
 struct GLTFEXPORTER_API FGLTFConverterUtility
@@ -34,9 +35,10 @@ struct GLTFEXPORTER_API FGLTFConverterUtility
 		return FVector4(ConvertVector(Tangent), 1.0f);
 	}
 
-	static inline FColor ConvertColor(const FColor& Color)
+	static inline FGLTFColor ConvertColor(const FColor& Color)
 	{
-		// UE4 uses ABGR while glTF uses RGBA.
+		// UE4 uses ABGR or ARGB depending on endianness.
+		// glTF always uses RGBA independent of endianness.
 		return { Color.R, Color.G, Color.B, Color.A };
 	}
 
