[GLTFExporter] fix: convert position to export scale (temporarily hardcoded)

commit 292009ea3ac418dcb87260864456127287311487
Date: Sun, 17 May 2020 14:16:32 +0200

    fix: convert position to export scale (temporarily hardcoded)
    (cherry picked from commit f5886b4ebc1abbad048a5fb559c709d68f6ee0f5​)

#skipci

diff --git a/Engine/Plugins/Enterprise/GLTFExporter/MigratedFromGit/Source/GLTFExporter/Public/GLTFConverterUtility.h b/Engine/Plugins/Enterprise/GLTFExporter/MigratedFromGit/Source/GLTFExporter/Public/GLTFConverterUtility.h
--- a/Engine/Plugins/Enterprise/GLTFExporter/MigratedFromGit/Source/GLTFExporter/Public/GLTFConverterUtility.h
+++ b/Engine/Plugins/Enterprise/GLTFExporter/MigratedFromGit/Source/GLTFExporter/Public/GLTFConverterUtility.h
@@ -16,7 +16,7 @@ struct GLTFEXPORTER_API FGLTFConverterUtility
 
 	static inline FVector ConvertPosition(const FVector& Position)
 	{
-		return ConvertVector(Position);
+		return ConvertVector(Position) * 0.01f; // TODO: use options export scale instead of hardcoded value
 	}
 
 	static inline FVector ConvertScale(const FVector& Scale)
