[GLTFExporter] refactor: remove redundant inlining in utilities

commit 4ecfc9646dceb74e05fe685363ae2da14f5b882d
Date: Sun, 17 May 2020 15:42:38 +0200

    refactor: remove redundant inlining in utilities
    (cherry picked from commit 59579f03c1886827e8f8938f875640f80422531a​)

#skipci

diff --git a/Engine/Plugins/Enterprise/GLTFExporter/MigratedFromGit/Source/GLTFExporter/Public/GLTFConverterUtility.h b/Engine/Plugins/Enterprise/GLTFExporter/MigratedFromGit/Source/GLTFExporter/Public/GLTFConverterUtility.h
--- a/Engine/Plugins/Enterprise/GLTFExporter/MigratedFromGit/Source/GLTFExporter/Public/GLTFConverterUtility.h
+++ b/Engine/Plugins/Enterprise/GLTFExporter/MigratedFromGit/Source/GLTFExporter/Public/GLTFConverterUtility.h
@@ -7,42 +7,42 @@
 
 struct GLTFEXPORTER_API FGLTFConverterUtility
 {
-	static inline FVector ConvertVector(const FVector& Vector)
+	static FVector ConvertVector(const FVector& Vector)
 	{
 		// UE4 uses a left-handed coordinate system, with Z up.
 		// glTF uses a right-handed coordinate system, with Y up.
 		return { Vector.X, Vector.Z, Vector.Y };
 	}
 
-	static inline FVector ConvertPosition(const FVector& Position)
+	static FVector ConvertPosition(const FVector& Position)
 	{
 		return ConvertVector(Position) * 0.01f; // TODO: use options export scale instead of hardcoded value
 	}
 
-	static inline FVector ConvertScale(const FVector& Scale)
+	static FVector ConvertScale(const FVector& Scale)
 	{
 		return ConvertVector(Scale);
 	}
 
-	static inline FVector ConvertNormal(const FVector& Normal)
+	static FVector ConvertNormal(const FVector& Normal)
 	{
 		return ConvertVector(Normal);
 	}
 
-	static inline FVector4 ConvertTangent(const FVector& Tangent)
+	static FVector4 ConvertTangent(const FVector& Tangent)
 	{
 		// glTF stores tangent as Vec4, with W component indicating handedness of tangent basis.
 		return FVector4(ConvertVector(Tangent), 1.0f);
 	}
 
-	static inline FGLTFColor ConvertColor(const FColor& Color)
+	static FGLTFColor ConvertColor(const FColor& Color)
 	{
 		// UE4 uses ABGR or ARGB depending on endianness.
 		// glTF always uses RGBA independent of endianness.
 		return { Color.R, Color.G, Color.B, Color.A };
 	}
 
-	static inline FQuat ConvertRotation(const FQuat& Rotation)
+	static FQuat ConvertRotation(const FQuat& Rotation)
 	{
 		// UE4 uses a left-handed coordinate system, with Z up.
 		// glTF uses a right-handed coordinate system, with Y up.
@@ -58,7 +58,7 @@ struct GLTFEXPORTER_API FGLTFConverterUtility
 		return Result;
 	}
 
-	static inline FMatrix ConvertMatrix(const FMatrix& Matrix)
+	static FMatrix ConvertMatrix(const FMatrix& Matrix)
 	{
 		// Unreal stores matrix elements in row major order.
 		// glTF stores matrix elements in column major order.
@@ -74,7 +74,7 @@ struct GLTFEXPORTER_API FGLTFConverterUtility
 		return Result;
 	}
 
-	static inline EGLTFJsonAlphaMode ConvertAlphaMode(const EBlendMode Mode)
+	static EGLTFJsonAlphaMode ConvertAlphaMode(const EBlendMode Mode)
 	{
 		switch (Mode)
 		{
diff --git a/Engine/Plugins/Enterprise/GLTFExporter/MigratedFromGit/Source/GLTFExporter/Public/GLTFJsonUtility.h b/Engine/Plugins/Enterprise/GLTFExporter/MigratedFromGit/Source/GLTFExporter/Public/GLTFJsonUtility.h
--- a/Engine/Plugins/Enterprise/GLTFExporter/MigratedFromGit/Source/GLTFExporter/Public/GLTFJsonUtility.h
+++ b/Engine/Plugins/Enterprise/GLTFExporter/MigratedFromGit/Source/GLTFExporter/Public/GLTFJsonUtility.h
@@ -6,7 +6,7 @@
 
 struct GLTFEXPORTER_API FGLTFJsonUtility
 {
-	static inline const TCHAR* ExtensionToString(EGLTFExtension Extension)
+	static const TCHAR* ExtensionToString(EGLTFExtension Extension)
 	{
 		switch (Extension)
 		{
@@ -18,7 +18,7 @@ struct GLTFEXPORTER_API FGLTFJsonUtility
 		}
 	}
 
-	static inline const TCHAR* AlphaModeToString(EGLTFJsonAlphaMode Mode)
+	static const TCHAR* AlphaModeToString(EGLTFJsonAlphaMode Mode)
 	{
 		switch (Mode)
 		{
@@ -29,22 +29,22 @@ struct GLTFEXPORTER_API FGLTFJsonUtility
 		}
 	}
 
-	static inline int32 PrimitiveModeToNumber(EGLTFJsonPrimitiveMode Mode)
+	static int32 PrimitiveModeToNumber(EGLTFJsonPrimitiveMode Mode)
 	{
 		return static_cast<int32>(Mode);
 	}
 
-	static inline int32 FilterToNumber(EGLTFJsonSamplerFilter Filter)
+	static int32 FilterToNumber(EGLTFJsonSamplerFilter Filter)
 	{
 		return static_cast<int32>(Filter);
 	}
 
-	static inline int32 WrapModeToNumber(EGLTFJsonSamplerWrap Wrap)
+	static int32 WrapModeToNumber(EGLTFJsonSamplerWrap Wrap)
 	{
 		return static_cast<int32>(Wrap);
 	}
 
-	static inline const TCHAR* AccessorTypeToString(EGLTFJsonAccessorType Type)
+	static const TCHAR* AccessorTypeToString(EGLTFJsonAccessorType Type)
 	{
 		switch (Type)
 		{
@@ -59,12 +59,12 @@ struct GLTFEXPORTER_API FGLTFJsonUtility
 		}
 	}
 
-	static inline int32 ComponentTypeToNumber(EGLTFJsonComponentType Type)
+	static int32 ComponentTypeToNumber(EGLTFJsonComponentType Type)
 	{
 		return static_cast<int32>(Type);
 	}
 
-	static inline int32 BufferTargetToNumber(EGLTFJsonBufferTarget Target)
+	static int32 BufferTargetToNumber(EGLTFJsonBufferTarget Target)
 	{
 		return static_cast<int32>(Target);
 	}
