[GLTFExporter] refactor: rename GLTFIndexedBuilder to GLTFConvertBuilder and replace GLTFIndexedObjects with GLTFConverter used as base class for all converters

commit e56b798c9ff5e95f469f16ddbd2884a821d2d38d
Date: Sun, 17 May 2020 17:15:11 +0200

    refactor: rename GLTFIndexedBuilder to GLTFConvertBuilder and replace GLTFIndexedObjects with GLTFConverter used as base class for all converters
    (cherry picked from commit aff358a20bbcda5a115dad84d0819d958497a7cb​)

#skipci

diff --git a/Engine/Plugins/Enterprise/GLTFExporter/MigratedFromGit/Source/GLTFExporter/Private/GLTFConvertBuilder.cpp b/Engine/Plugins/Enterprise/GLTFExporter/MigratedFromGit/Source/GLTFExporter/Private/GLTFConvertBuilder.cpp
new file mode 100644
--- /dev/null
+++ b/Engine/Plugins/Enterprise/GLTFExporter/MigratedFromGit/Source/GLTFExporter/Private/GLTFConvertBuilder.cpp
@@ -0,0 +1,85 @@
+// Copyright Epic Games, Inc. All Rights Reserved.
+
+#include "GLTFConvertBuilder.h"
+
+FGLTFJsonAccessorIndex FGLTFConvertBuilder::GetOrAddPositionAccessor(const FPositionVertexBuffer* VertexBuffer, const FString& DesiredName)
+{
+	return PositionVertexBufferConverter.GetOrAdd(*this, DesiredName, VertexBuffer);
+}
+
+FGLTFJsonAccessorIndex FGLTFConvertBuilder::GetOrAddColorAccessor(const FColorVertexBuffer* VertexBuffer, const FString& DesiredName)
+{
+	return ColorVertexBufferConverter.GetOrAdd(*this, DesiredName, VertexBuffer);
+}
+
+FGLTFJsonAccessorIndex FGLTFConvertBuilder::GetOrAddNormalAccessor(const FStaticMeshVertexBuffer* VertexBuffer, const FString& DesiredName)
+{
+	return StaticMeshNormalVertexBufferConverter.GetOrAdd(*this, DesiredName, VertexBuffer);
+}
+
+FGLTFJsonAccessorIndex FGLTFConvertBuilder::GetOrAddTangentAccessor(const FStaticMeshVertexBuffer* VertexBuffer, const FString& DesiredName)
+{
+	return StaticMeshTangentVertexBufferConverter.GetOrAdd(*this, DesiredName, VertexBuffer);
+}
+
+FGLTFJsonAccessorIndex FGLTFConvertBuilder::GetOrAddUV0Accessor(const FStaticMeshVertexBuffer* VertexBuffer, const FString& DesiredName)
+{
+	return StaticMeshUV0VertexBufferConverter.GetOrAdd(*this, DesiredName, VertexBuffer);
+}
+
+FGLTFJsonAccessorIndex FGLTFConvertBuilder::GetOrAddUV1Accessor(const FStaticMeshVertexBuffer* VertexBuffer, const FString& DesiredName)
+{
+	return StaticMeshUV1VertexBufferConverter.GetOrAdd(*this, DesiredName, VertexBuffer);
+}
+
+FGLTFJsonBufferViewIndex FGLTFConvertBuilder::GetOrAddIndexBufferView(const FRawStaticIndexBuffer* IndexBuffer, const FString& DesiredName)
+{
+	return StaticMeshIndexBufferConverter.GetOrAdd(*this, DesiredName, IndexBuffer);
+}
+
+FGLTFJsonAccessorIndex FGLTFConvertBuilder::GetOrAddIndexAccessor(const FStaticMeshSection* MeshSection, const FRawStaticIndexBuffer* IndexBuffer, const FString& DesiredName)
+{
+	return StaticMeshSectionConverter.GetOrAdd(*this, DesiredName, MeshSection, IndexBuffer);
+}
+
+FGLTFJsonMeshIndex FGLTFConvertBuilder::GetOrAddMesh(const FStaticMeshLODResources* StaticMeshLOD, const FColorVertexBuffer* OverrideVertexColors, const FString& DesiredName)
+{
+	return StaticMeshConverter.GetOrAdd(*this, DesiredName, StaticMeshLOD, OverrideVertexColors);
+}
+
+FGLTFJsonMeshIndex FGLTFConvertBuilder::GetOrAddMesh(const UStaticMesh* StaticMesh, int32 LODIndex, const FColorVertexBuffer* OverrideVertexColors, const FString& DesiredName)
+{
+	const FStaticMeshLODResources* StaticMeshLOD = &StaticMesh->GetLODForExport(LODIndex);
+	FString Name = DesiredName.IsEmpty() ? StaticMesh->GetName() : DesiredName;
+
+	if (LODIndex != 0)
+	{
+		Name += TEXT("_LOD") + FString::FromInt(LODIndex);
+	}
+
+	return GetOrAddMesh(StaticMeshLOD, OverrideVertexColors, Name);
+}
+
+FGLTFJsonMeshIndex FGLTFConvertBuilder::GetOrAddMesh(const UStaticMeshComponent* StaticMeshComponent, const FString& DesiredName)
+{
+	const UStaticMesh* StaticMesh = StaticMeshComponent->GetStaticMesh();
+	const int32 LODIndex = StaticMeshComponent->ForcedLodModel > 0 ? StaticMeshComponent->ForcedLodModel - 1 : /* auto-select */ 0;
+	const FColorVertexBuffer* OverrideVertexColors = LODIndex < StaticMeshComponent->LODData.Num() ? StaticMeshComponent->LODData[LODIndex].OverrideVertexColors : nullptr;
+
+	return GetOrAddMesh(StaticMesh, LODIndex, OverrideVertexColors, DesiredName);
+}
+
+FGLTFJsonNodeIndex FGLTFConvertBuilder::GetOrAddNode(const USceneComponent* SceneComponent, bool bSelectedOnly, bool bRootNode, const FString& DesiredName)
+{
+	return SceneComponentConverter.GetOrAdd(*this, DesiredName, SceneComponent, bSelectedOnly, bRootNode);
+}
+
+FGLTFJsonSceneIndex FGLTFConvertBuilder::GetOrAddScene(const ULevel* Level, bool bSelectedOnly, const FString& DesiredName)
+{
+	return LevelConverter.GetOrAdd(*this, DesiredName, Level, bSelectedOnly);
+}
+
+FGLTFJsonSceneIndex FGLTFConvertBuilder::GetOrAddScene(const UWorld* World, bool bSelectedOnly, const FString& DesiredName)
+{
+	return GetOrAddScene(World->PersistentLevel, bSelectedOnly, DesiredName.IsEmpty() ? World->GetName() : DesiredName);
+}
diff --git a/Engine/Plugins/Enterprise/GLTFExporter/MigratedFromGit/Source/GLTFExporter/Private/GLTFIndexedBuilder.cpp b/Engine/Plugins/Enterprise/GLTFExporter/MigratedFromGit/Source/GLTFExporter/Private/GLTFIndexedBuilder.cpp
deleted file mode 100644
diff --git a/Engine/Plugins/Enterprise/GLTFExporter/MigratedFromGit/Source/GLTFExporter/Private/GLTFLevelConverters.cpp b/Engine/Plugins/Enterprise/GLTFExporter/MigratedFromGit/Source/GLTFExporter/Private/GLTFLevelConverters.cpp
--- a/Engine/Plugins/Enterprise/GLTFExporter/MigratedFromGit/Source/GLTFExporter/Private/GLTFLevelConverters.cpp
+++ b/Engine/Plugins/Enterprise/GLTFExporter/MigratedFromGit/Source/GLTFExporter/Private/GLTFLevelConverters.cpp
@@ -4,8 +4,12 @@
 #include "GLTFContainerBuilder.h"
 #include "GLTFConverterUtility.h"
 
-FGLTFJsonNodeIndex FGLTFSceneComponentConverter::Add(FGLTFIndexedBuilder& Builder, const FString& Name, const USceneComponent* SceneComponent, bool bSelectedOnly, bool bRootNode)
+FGLTFJsonNodeIndex FGLTFSceneComponentConverter::Add(FGLTFConvertBuilder& Builder, const FString& Name, TTuple<const USceneComponent*, bool, bool> Params)
 {
+	const USceneComponent* SceneComponent = Params.Get<0>();
+	bool bSelectedOnly = Params.Get<1>();
+	bool bRootNode = Params.Get<2>();
+
 	const AActor* Owner = SceneComponent->GetOwner();
 	const bool bIsRootComponent = Owner != nullptr && Owner->GetRootComponent() == SceneComponent;
 	const UBlueprint* Blueprint = UBlueprint::GetBlueprintFromClass(Owner != nullptr ? Owner->GetClass() : nullptr);
@@ -51,8 +55,11 @@ FGLTFJsonNodeIndex FGLTFSceneComponentConverter::Add(FGLTFIndexedBuilder& Builde
 	return Builder.AddNode(Node);
 }
 
-FGLTFJsonSceneIndex FGLTFLevelConverter::Add(FGLTFIndexedBuilder& Builder, const FString& Name, const ULevel* Level, bool bSelectedOnly)
+FGLTFJsonSceneIndex FGLTFLevelConverter::Add(FGLTFConvertBuilder& Builder, const FString& Name, TTuple<const ULevel*, bool> Params)
 {
+	const ULevel* Level = Params.Get<0>();
+	bool bSelectedOnly = Params.Get<1>();
+
 	FGLTFJsonScene Scene;
 	Scene.Name = Name;
 
diff --git a/Engine/Plugins/Enterprise/GLTFExporter/MigratedFromGit/Source/GLTFExporter/Private/GLTFStaticMeshConverters.cpp b/Engine/Plugins/Enterprise/GLTFExporter/MigratedFromGit/Source/GLTFExporter/Private/GLTFStaticMeshConverters.cpp
--- a/Engine/Plugins/Enterprise/GLTFExporter/MigratedFromGit/Source/GLTFExporter/Private/GLTFStaticMeshConverters.cpp
+++ b/Engine/Plugins/Enterprise/GLTFExporter/MigratedFromGit/Source/GLTFExporter/Private/GLTFStaticMeshConverters.cpp
@@ -1,10 +1,10 @@
 // Copyright Epic Games, Inc. All Rights Reserved.
 
 #include "GLTFStaticMeshConverters.h"
-#include "GLTFContainerBuilder.h"
+#include "GLTFConvertBuilder.h"
 #include "GLTFConverterUtility.h"
 
-FGLTFJsonAccessorIndex FGLTFPositionVertexBufferConverter::Add(FGLTFIndexedBuilder& Builder, const FString& Name, const FPositionVertexBuffer* VertexBuffer)
+FGLTFJsonAccessorIndex FGLTFPositionVertexBufferConverter::Add(FGLTFConvertBuilder& Builder, const FString& Name, const FPositionVertexBuffer* VertexBuffer)
 {
 	const uint32 VertexCount = VertexBuffer->GetNumVertices();
 	if (VertexCount == 0)
@@ -47,7 +47,7 @@ FGLTFJsonAccessorIndex FGLTFPositionVertexBufferConverter::Add(FGLTFIndexedBuild
 	return Builder.AddAccessor(JsonAccessor);
 }
 
-FGLTFJsonAccessorIndex FGLTFColorVertexBufferConverter::Add(FGLTFIndexedBuilder& Builder, const FString& Name, const FColorVertexBuffer* VertexBuffer)
+FGLTFJsonAccessorIndex FGLTFColorVertexBufferConverter::Add(FGLTFConvertBuilder& Builder, const FString& Name, const FColorVertexBuffer* VertexBuffer)
 {
 	const uint32 VertexCount = VertexBuffer->GetNumVertices();
 	if (VertexCount == 0)
@@ -74,7 +74,7 @@ FGLTFJsonAccessorIndex FGLTFColorVertexBufferConverter::Add(FGLTFIndexedBuilder&
 	return Builder.AddAccessor(JsonAccessor);
 }
 
-FGLTFJsonAccessorIndex FGLTFStaticMeshNormalVertexBufferConverter::Add(FGLTFIndexedBuilder& Builder, const FString& Name, const FStaticMeshVertexBuffer* VertexBuffer)
+FGLTFJsonAccessorIndex FGLTFStaticMeshNormalVertexBufferConverter::Add(FGLTFConvertBuilder& Builder, const FString& Name, const FStaticMeshVertexBuffer* VertexBuffer)
 {
 	const uint32 VertexCount = VertexBuffer->GetNumVertices();
 	if (VertexCount == 0)
@@ -100,7 +100,7 @@ FGLTFJsonAccessorIndex FGLTFStaticMeshNormalVertexBufferConverter::Add(FGLTFInde
 	return Builder.AddAccessor(JsonAccessor);
 }
 
-FGLTFJsonAccessorIndex FGLTFStaticMeshTangentVertexBufferConverter::Add(FGLTFIndexedBuilder& Builder, const FString& Name, const FStaticMeshVertexBuffer* VertexBuffer)
+FGLTFJsonAccessorIndex FGLTFStaticMeshTangentVertexBufferConverter::Add(FGLTFConvertBuilder& Builder, const FString& Name, const FStaticMeshVertexBuffer* VertexBuffer)
 {
 	const uint32 VertexCount = VertexBuffer->GetNumVertices();
 	if (VertexCount == 0)
@@ -126,7 +126,7 @@ FGLTFJsonAccessorIndex FGLTFStaticMeshTangentVertexBufferConverter::Add(FGLTFInd
 	return Builder.AddAccessor(JsonAccessor);
 }
 
-FGLTFJsonAccessorIndex FGLTFStaticMeshUV0VertexBufferConverter::Add(FGLTFIndexedBuilder& Builder, const FString& Name, const FStaticMeshVertexBuffer* VertexBuffer)
+FGLTFJsonAccessorIndex FGLTFStaticMeshUV0VertexBufferConverter::Add(FGLTFConvertBuilder& Builder, const FString& Name, const FStaticMeshVertexBuffer* VertexBuffer)
 {
 	const uint32 VertexCount = VertexBuffer->GetNumVertices();
 	if (VertexCount == 0 || VertexBuffer->GetNumTexCoords() < 1)
@@ -152,7 +152,7 @@ FGLTFJsonAccessorIndex FGLTFStaticMeshUV0VertexBufferConverter::Add(FGLTFIndexed
 	return Builder.AddAccessor(JsonAccessor);
 }
 
-FGLTFJsonAccessorIndex FGLTFStaticMeshUV1VertexBufferConverter::Add(FGLTFIndexedBuilder& Builder, const FString& Name, const FStaticMeshVertexBuffer* VertexBuffer)
+FGLTFJsonAccessorIndex FGLTFStaticMeshUV1VertexBufferConverter::Add(FGLTFConvertBuilder& Builder, const FString& Name, const FStaticMeshVertexBuffer* VertexBuffer)
 {
 	const uint32 VertexCount = VertexBuffer->GetNumVertices();
 	if (VertexCount == 0 || VertexBuffer->GetNumTexCoords() < 2)
@@ -178,7 +178,7 @@ FGLTFJsonAccessorIndex FGLTFStaticMeshUV1VertexBufferConverter::Add(FGLTFIndexed
 	return Builder.AddAccessor(JsonAccessor);
 }
 
-FGLTFJsonBufferViewIndex FGLTFStaticMeshIndexBufferConverter::Add(FGLTFIndexedBuilder& Builder, const FString& Name, const FRawStaticIndexBuffer* IndexBuffer)
+FGLTFJsonBufferViewIndex FGLTFStaticMeshIndexBufferConverter::Add(FGLTFConvertBuilder& Builder, const FString& Name, const FRawStaticIndexBuffer* IndexBuffer)
 {
 	if (IndexBuffer->GetNumIndices() == 0)
 	{
@@ -190,8 +190,11 @@ FGLTFJsonBufferViewIndex FGLTFStaticMeshIndexBufferConverter::Add(FGLTFIndexedBu
 	return Builder.AddBufferView(Indices, Name, EGLTFJsonBufferTarget::ElementArrayBuffer);
 }
 
-FGLTFJsonAccessorIndex FGLTFStaticMeshSectionConverter::Add(FGLTFIndexedBuilder& Builder, const FString& Name, const FStaticMeshSection* MeshSection, const FRawStaticIndexBuffer* IndexBuffer)
+FGLTFJsonAccessorIndex FGLTFStaticMeshSectionConverter::Add(FGLTFConvertBuilder& Builder, const FString& Name, TTuple<const FStaticMeshSection*, const FRawStaticIndexBuffer*> Params)
 {
+	const FStaticMeshSection* MeshSection = Params.Get<0>();
+	const FRawStaticIndexBuffer* IndexBuffer = Params.Get<1>();
+
 	const uint32 TriangleCount = MeshSection->NumTriangles;
 	if (TriangleCount == 0)
 	{
@@ -209,8 +212,11 @@ FGLTFJsonAccessorIndex FGLTFStaticMeshSectionConverter::Add(FGLTFIndexedBuilder&
 	return Builder.AddAccessor(JsonAccessor);
 }
 
-FGLTFJsonMeshIndex FGLTFStaticMeshConverter::Add(FGLTFIndexedBuilder& Builder, const FString& Name, const FStaticMeshLODResources* StaticMeshLOD, const FColorVertexBuffer* OverrideVertexColors)
+FGLTFJsonMeshIndex FGLTFStaticMeshConverter::Add(FGLTFConvertBuilder& Builder, const FString& Name, TTuple<const FStaticMeshLODResources*, const FColorVertexBuffer*> Params)
 {
+	const FStaticMeshLODResources* StaticMeshLOD = Params.Get<0>();
+	const FColorVertexBuffer* OverrideVertexColors = Params.Get<1>();
+
 	FGLTFJsonMesh JsonMesh;
 	JsonMesh.Name = Name;
 
@@ -238,7 +244,7 @@ FGLTFJsonMeshIndex FGLTFStaticMeshConverter::Add(FGLTFIndexedBuilder& Builder, c
 		JsonPrimitive.Attributes = JsonAttributes;
 
 		JsonPrimitive.Indices = Builder.GetOrAddIndexAccessor(&StaticMeshLOD->Sections[SectionIndex], IndexBuffer,
-			Name.IsEmpty() ? Name : JsonMesh.Name + (SectionCount != 1 ? TEXT("_Indices_Section") + FString::FromInt(SectionIndex) : TEXT("_Indices")));
+            Name.IsEmpty() ? Name : JsonMesh.Name + (SectionCount != 1 ? TEXT("_Indices_Section") + FString::FromInt(SectionIndex) : TEXT("_Indices")));
 	}
 
 	return Builder.AddMesh(JsonMesh);
diff --git a/Engine/Plugins/Enterprise/GLTFExporter/MigratedFromGit/Source/GLTFExporter/Public/GLTFContainerBuilder.h b/Engine/Plugins/Enterprise/GLTFExporter/MigratedFromGit/Source/GLTFExporter/Public/GLTFContainerBuilder.h
--- a/Engine/Plugins/Enterprise/GLTFExporter/MigratedFromGit/Source/GLTFExporter/Public/GLTFContainerBuilder.h
+++ b/Engine/Plugins/Enterprise/GLTFExporter/MigratedFromGit/Source/GLTFExporter/Public/GLTFContainerBuilder.h
@@ -2,8 +2,8 @@
 
 #pragma once
 
-#include "GLTFIndexedBuilder.h"
+#include "GLTFConvertBuilder.h"
 
-struct GLTFEXPORTER_API FGLTFContainerBuilder : public FGLTFIndexedBuilder
+struct GLTFEXPORTER_API FGLTFContainerBuilder : public FGLTFConvertBuilder
 {
 };
diff --git a/Engine/Plugins/Enterprise/GLTFExporter/MigratedFromGit/Source/GLTFExporter/Public/GLTFConvertBuilder.h b/Engine/Plugins/Enterprise/GLTFExporter/MigratedFromGit/Source/GLTFExporter/Public/GLTFConvertBuilder.h
new file mode 100644
--- /dev/null
+++ b/Engine/Plugins/Enterprise/GLTFExporter/MigratedFromGit/Source/GLTFExporter/Public/GLTFConvertBuilder.h
@@ -0,0 +1,42 @@
+// Copyright Epic Games, Inc. All Rights Reserved.
+
+#pragma once
+
+#include "GLTFBufferBuilder.h"
+#include "GLTFConverter.h"
+#include "GLTFStaticMeshConverters.h"
+#include "GLTFLevelConverters.h"
+
+struct GLTFEXPORTER_API FGLTFConvertBuilder : public FGLTFBufferBuilder
+{
+	FGLTFPositionVertexBufferConverter PositionVertexBufferConverter;
+	FGLTFColorVertexBufferConverter ColorVertexBufferConverter;
+	FGLTFStaticMeshNormalVertexBufferConverter StaticMeshNormalVertexBufferConverter;
+	FGLTFStaticMeshTangentVertexBufferConverter StaticMeshTangentVertexBufferConverter;
+	FGLTFStaticMeshUV0VertexBufferConverter StaticMeshUV0VertexBufferConverter;
+	FGLTFStaticMeshUV1VertexBufferConverter StaticMeshUV1VertexBufferConverter;
+	FGLTFStaticMeshIndexBufferConverter StaticMeshIndexBufferConverter;
+	FGLTFStaticMeshSectionConverter StaticMeshSectionConverter;
+	FGLTFStaticMeshConverter StaticMeshConverter;
+
+	FGLTFSceneComponentConverter SceneComponentConverter;
+	FGLTFLevelConverter LevelConverter;
+
+	FGLTFJsonAccessorIndex GetOrAddPositionAccessor(const FPositionVertexBuffer* VertexBuffer, const FString& DesiredName = TEXT(""));
+	FGLTFJsonAccessorIndex GetOrAddNormalAccessor(const FStaticMeshVertexBuffer* VertexBuffer, const FString& DesiredName = TEXT(""));
+	FGLTFJsonAccessorIndex GetOrAddTangentAccessor(const FStaticMeshVertexBuffer* VertexBuffer, const FString& DesiredName = TEXT(""));
+	FGLTFJsonAccessorIndex GetOrAddUV0Accessor(const FStaticMeshVertexBuffer* VertexBuffer, const FString& DesiredName = TEXT(""));
+	FGLTFJsonAccessorIndex GetOrAddUV1Accessor(const FStaticMeshVertexBuffer* VertexBuffer, const FString& DesiredName = TEXT(""));
+	FGLTFJsonAccessorIndex GetOrAddColorAccessor(const FColorVertexBuffer* VertexBuffer, const FString& DesiredName = TEXT(""));
+
+	FGLTFJsonBufferViewIndex GetOrAddIndexBufferView(const FRawStaticIndexBuffer* IndexBuffer, const FString& DesiredName = TEXT(""));
+	FGLTFJsonAccessorIndex GetOrAddIndexAccessor(const FStaticMeshSection* MeshSection, const FRawStaticIndexBuffer* IndexBuffer, const FString& DesiredName = TEXT(""));
+
+	FGLTFJsonMeshIndex GetOrAddMesh(const FStaticMeshLODResources* StaticMeshLOD, const FColorVertexBuffer* OverrideVertexColors = nullptr, const FString& DesiredName = TEXT(""));
+	FGLTFJsonMeshIndex GetOrAddMesh(const UStaticMesh* StaticMesh, int32 LODIndex = 0, const FColorVertexBuffer* OverrideVertexColors = nullptr, const FString& DesiredName = TEXT(""));
+	FGLTFJsonMeshIndex GetOrAddMesh(const UStaticMeshComponent* StaticMeshComponent, const FString& DesiredName = TEXT(""));
+
+	FGLTFJsonNodeIndex GetOrAddNode(const USceneComponent* SceneComponent, bool bSelectedOnly, bool bRootNode = false, const FString& DesiredName = TEXT(""));
+	FGLTFJsonSceneIndex GetOrAddScene(const ULevel* Level, bool bSelectedOnly, const FString& DesiredName = TEXT(""));
+	FGLTFJsonSceneIndex GetOrAddScene(const UWorld* World, bool bSelectedOnly, const FString& DesiredName = TEXT(""));
+};
diff --git a/Engine/Plugins/Enterprise/GLTFExporter/MigratedFromGit/Source/GLTFExporter/Public/GLTFConverter.h b/Engine/Plugins/Enterprise/GLTFExporter/MigratedFromGit/Source/GLTFExporter/Public/GLTFConverter.h
new file mode 100644
--- /dev/null
+++ b/Engine/Plugins/Enterprise/GLTFExporter/MigratedFromGit/Source/GLTFExporter/Public/GLTFConverter.h
@@ -0,0 +1,42 @@
+// Copyright Epic Games, Inc. All Rights Reserved.
+
+#pragma once
+
+#include "CoreMinimal.h"
+
+struct FGLTFConvertBuilder;
+
+template <class IndexType, class KeyType>
+class GLTFEXPORTER_API TGLTFConverter
+{
+public:
+	virtual ~TGLTFConverter() = default;
+
+	template <typename... ArgsType>
+    FORCEINLINE IndexType Get(ArgsType&&... Args) const
+	{
+		const KeyType Key(Forward<ArgsType>(Args)...);
+		return IndexLookup.FindRef(Key);
+	}
+
+	template <typename... ArgsType>
+    FORCEINLINE IndexType GetOrAdd(FGLTFConvertBuilder& Builder, const FString& DesiredName, ArgsType&&... Args)
+	{
+		const KeyType Key(Forward<ArgsType>(Args)...);
+
+		IndexType Index = IndexLookup.FindRef(Key);
+		if (Index == INDEX_NONE)
+		{
+			Index = Add(Builder, DesiredName, Key);
+			IndexLookup.Add(Key, Index);
+		}
+
+		return Index;
+	}
+
+private:
+
+	TMap<KeyType, IndexType> IndexLookup;
+
+	virtual IndexType Add(FGLTFConvertBuilder& Builder, const FString& Name, KeyType Key) = 0;
+};
diff --git a/Engine/Plugins/Enterprise/GLTFExporter/MigratedFromGit/Source/GLTFExporter/Public/GLTFIndexedBuilder.h b/Engine/Plugins/Enterprise/GLTFExporter/MigratedFromGit/Source/GLTFExporter/Public/GLTFIndexedBuilder.h
deleted file mode 100644
diff --git a/Engine/Plugins/Enterprise/GLTFExporter/MigratedFromGit/Source/GLTFExporter/Public/GLTFIndexedObjects.h b/Engine/Plugins/Enterprise/GLTFExporter/MigratedFromGit/Source/GLTFExporter/Public/GLTFIndexedObjects.h
deleted file mode 100644
diff --git a/Engine/Plugins/Enterprise/GLTFExporter/MigratedFromGit/Source/GLTFExporter/Public/GLTFLevelConverters.h b/Engine/Plugins/Enterprise/GLTFExporter/MigratedFromGit/Source/GLTFExporter/Public/GLTFLevelConverters.h
--- a/Engine/Plugins/Enterprise/GLTFExporter/MigratedFromGit/Source/GLTFExporter/Public/GLTFLevelConverters.h
+++ b/Engine/Plugins/Enterprise/GLTFExporter/MigratedFromGit/Source/GLTFExporter/Public/GLTFLevelConverters.h
@@ -3,16 +3,15 @@
 #pragma once
 
 #include "GLTFJsonIndex.h"
+#include "GLTFConverter.h"
 #include "Engine.h"
 
-struct FGLTFIndexedBuilder;
-
-struct GLTFEXPORTER_API FGLTFSceneComponentConverter
+class GLTFEXPORTER_API FGLTFSceneComponentConverter final : public TGLTFConverter<FGLTFJsonNodeIndex, TTuple<const USceneComponent*, bool, bool>>
 {
-	static FGLTFJsonNodeIndex Add(FGLTFIndexedBuilder& Builder, const FString& Name, const USceneComponent* SceneComponent, bool bSelectedOnly, bool bRootNode = false);
+	FGLTFJsonNodeIndex Add(FGLTFConvertBuilder& Builder, const FString& Name, TTuple<const USceneComponent*, bool, bool> Params) override;
 };
 
-struct GLTFEXPORTER_API FGLTFLevelConverter
+class GLTFEXPORTER_API FGLTFLevelConverter final : public TGLTFConverter<FGLTFJsonSceneIndex, TTuple<const ULevel*, bool>>
 {
-	static FGLTFJsonSceneIndex Add(FGLTFIndexedBuilder& Builder, const FString& Name, const ULevel* Level, bool bSelectedOnly);
+	FGLTFJsonSceneIndex Add(FGLTFConvertBuilder& Builder, const FString& Name, TTuple<const ULevel*, bool> Params) override;
 };
diff --git a/Engine/Plugins/Enterprise/GLTFExporter/MigratedFromGit/Source/GLTFExporter/Public/GLTFStaticMeshConverters.h b/Engine/Plugins/Enterprise/GLTFExporter/MigratedFromGit/Source/GLTFExporter/Public/GLTFStaticMeshConverters.h
--- a/Engine/Plugins/Enterprise/GLTFExporter/MigratedFromGit/Source/GLTFExporter/Public/GLTFStaticMeshConverters.h
+++ b/Engine/Plugins/Enterprise/GLTFExporter/MigratedFromGit/Source/GLTFExporter/Public/GLTFStaticMeshConverters.h
@@ -3,51 +3,50 @@
 #pragma once
 
 #include "GLTFJsonIndex.h"
+#include "GLTFConverter.h"
 #include "Engine.h"
 
-struct FGLTFIndexedBuilder;
-
-struct GLTFEXPORTER_API FGLTFPositionVertexBufferConverter
+class GLTFEXPORTER_API FGLTFPositionVertexBufferConverter final : public TGLTFConverter<FGLTFJsonAccessorIndex, const FPositionVertexBuffer*>
 {
-	static FGLTFJsonAccessorIndex Add(FGLTFIndexedBuilder& Builder, const FString& Name, const FPositionVertexBuffer* VertexBuffer);
+	FGLTFJsonAccessorIndex Add(FGLTFConvertBuilder& Builder, const FString& Name, const FPositionVertexBuffer* VertexBuffer) override;
 };
 
-struct GLTFEXPORTER_API FGLTFColorVertexBufferConverter
+class GLTFEXPORTER_API FGLTFColorVertexBufferConverter final : public TGLTFConverter<FGLTFJsonAccessorIndex, const FColorVertexBuffer*>
 {
-	static FGLTFJsonAccessorIndex Add(FGLTFIndexedBuilder& Builder, const FString& Name, const FColorVertexBuffer* VertexBuffer);
+	FGLTFJsonAccessorIndex Add(FGLTFConvertBuilder& Builder, const FString& Name, const FColorVertexBuffer* VertexBuffer) override;
 };
 
-struct GLTFEXPORTER_API FGLTFStaticMeshNormalVertexBufferConverter
+class GLTFEXPORTER_API FGLTFStaticMeshNormalVertexBufferConverter final : public TGLTFConverter<FGLTFJsonAccessorIndex, const FStaticMeshVertexBuffer*>
 {
-	static FGLTFJsonAccessorIndex Add(FGLTFIndexedBuilder& Builder, const FString& Name, const FStaticMeshVertexBuffer* VertexBuffer);
+	FGLTFJsonAccessorIndex Add(FGLTFConvertBuilder& Builder, const FString& Name, const FStaticMeshVertexBuffer* VertexBuffer) override;
 };
 
-struct GLTFEXPORTER_API FGLTFStaticMeshTangentVertexBufferConverter
+class GLTFEXPORTER_API FGLTFStaticMeshTangentVertexBufferConverter final : public TGLTFConverter<FGLTFJsonAccessorIndex, const FStaticMeshVertexBuffer*>
 {
-	static FGLTFJsonAccessorIndex Add(FGLTFIndexedBuilder& Builder, const FString& Name, const FStaticMeshVertexBuffer* VertexBuffer);
+	FGLTFJsonAccessorIndex Add(FGLTFConvertBuilder& Builder, const FString& Name, const FStaticMeshVertexBuffer* VertexBuffer) override;
 };
 
-struct GLTFEXPORTER_API FGLTFStaticMeshUV0VertexBufferConverter
+class GLTFEXPORTER_API FGLTFStaticMeshUV0VertexBufferConverter final : public TGLTFConverter<FGLTFJsonAccessorIndex, const FStaticMeshVertexBuffer*>
 {
-	static FGLTFJsonAccessorIndex Add(FGLTFIndexedBuilder& Builder, const FString& Name, const FStaticMeshVertexBuffer* VertexBuffer);
+	FGLTFJsonAccessorIndex Add(FGLTFConvertBuilder& Builder, const FString& Name, const FStaticMeshVertexBuffer* VertexBuffer) override;
 };
 
-struct GLTFEXPORTER_API FGLTFStaticMeshUV1VertexBufferConverter
+class GLTFEXPORTER_API FGLTFStaticMeshUV1VertexBufferConverter final : public TGLTFConverter<FGLTFJsonAccessorIndex, const FStaticMeshVertexBuffer*>
 {
-	static FGLTFJsonAccessorIndex Add(FGLTFIndexedBuilder& Builder, const FString& Name, const FStaticMeshVertexBuffer* VertexBuffer);
+	FGLTFJsonAccessorIndex Add(FGLTFConvertBuilder& Builder, const FString& Name, const FStaticMeshVertexBuffer* VertexBuffer) override;
 };
 
-struct GLTFEXPORTER_API FGLTFStaticMeshIndexBufferConverter
+class GLTFEXPORTER_API FGLTFStaticMeshIndexBufferConverter final : public TGLTFConverter<FGLTFJsonBufferViewIndex, const FRawStaticIndexBuffer*>
 {
-	static FGLTFJsonBufferViewIndex Add(FGLTFIndexedBuilder& Builder, const FString& Name, const FRawStaticIndexBuffer* IndexBuffer);
+	FGLTFJsonBufferViewIndex Add(FGLTFConvertBuilder& Builder, const FString& Name, const FRawStaticIndexBuffer* IndexBuffer) override;
 };
 
-struct GLTFEXPORTER_API FGLTFStaticMeshSectionConverter
+class GLTFEXPORTER_API FGLTFStaticMeshSectionConverter final : public TGLTFConverter<FGLTFJsonAccessorIndex, TTuple<const FStaticMeshSection*, const FRawStaticIndexBuffer*>>
 {
-	static FGLTFJsonAccessorIndex Add(FGLTFIndexedBuilder& Builder, const FString& Name, const FStaticMeshSection* MeshSection, const FRawStaticIndexBuffer* IndexBuffer);
+	FGLTFJsonAccessorIndex Add(FGLTFConvertBuilder& Builder, const FString& Name, TTuple<const FStaticMeshSection*, const FRawStaticIndexBuffer*> Params) override;
 };
 
-struct GLTFEXPORTER_API FGLTFStaticMeshConverter
+class GLTFEXPORTER_API FGLTFStaticMeshConverter final : public TGLTFConverter<FGLTFJsonMeshIndex, TTuple<const FStaticMeshLODResources*, const FColorVertexBuffer*>>
 {
-	static FGLTFJsonMeshIndex Add(FGLTFIndexedBuilder& Builder, const FString& Name, const FStaticMeshLODResources* StaticMeshLOD, const FColorVertexBuffer* OverrideVertexColors);
+	FGLTFJsonMeshIndex Add(FGLTFConvertBuilder& Builder, const FString& Name, TTuple<const FStaticMeshLODResources*, const FColorVertexBuffer*> Params) override;
 };
