[GLTFExporter] refactor: organize files into subfolders (Builders, Converters, Exporters, Json, UI)

commit e2fb24a49a98e9e5e695bdbeff3ffdcceabd7815
Date: Sun, 17 May 2020 18:13:42 +0200

    refactor: organize files into subfolders (Builders, Converters, Exporters, Json, UI)
    (cherry picked from commit 6f679b0cb889cf1687ebacf9b09f29dac5e649a8​)

#skipci

diff --git a/Engine/Plugins/Enterprise/GLTFExporter/MigratedFromGit/Source/GLTFExporter/Private/Builders/GLTFBufferBuilder.cpp b/Engine/Plugins/Enterprise/GLTFExporter/MigratedFromGit/Source/GLTFExporter/Private/Builders/GLTFBufferBuilder.cpp
new file mode 100644
--- /dev/null
+++ b/Engine/Plugins/Enterprise/GLTFExporter/MigratedFromGit/Source/GLTFExporter/Private/Builders/GLTFBufferBuilder.cpp
@@ -0,0 +1,45 @@
+// Copyright Epic Games, Inc. All Rights Reserved.
+
+#include "Builders/GLTFBufferBuilder.h"
+#include "Misc/FileHelper.h"
+#include "Misc/Paths.h"
+
+FGLTFBufferBuilder::FGLTFBufferBuilder()
+{
+	BufferIndex = AddBuffer(FGLTFJsonBuffer());
+}
+
+FGLTFJsonBufferViewIndex FGLTFBufferBuilder::AddBufferView(const void* RawData, uint64 ByteLength, const FString& Name, EGLTFJsonBufferTarget BufferTarget)
+{
+	FGLTFJsonBufferView JsonBufferView;
+	JsonBufferView.Name = Name;
+	JsonBufferView.Buffer = BufferIndex;
+	JsonBufferView.ByteOffset = BufferData.Num();
+	JsonBufferView.ByteLength = ByteLength;
+	JsonBufferView.Target = BufferTarget;
+
+	BufferData.Append(static_cast<const uint8*>(RawData), ByteLength);
+
+	return FGLTFJsonBuilder::AddBufferView(JsonBufferView);
+}
+
+void FGLTFBufferBuilder::UpdateJsonBufferObject(const FString& BinaryFilePath)
+{
+	FGLTFJsonBuffer& JsonBuffer = JsonRoot.Buffers[BufferIndex];
+	JsonBuffer.URI = FPaths::GetCleanFilename(BinaryFilePath);
+	JsonBuffer.ByteLength = BufferData.Num();
+}
+
+bool FGLTFBufferBuilder::Serialize(FArchive& Archive, const FString& FilePath)
+{
+	const FString BinaryFilePath = FPaths::ChangeExtension(FilePath, TEXT(".bin"));
+
+	if(!FFileHelper::SaveArrayToFile(BufferData, *BinaryFilePath))
+	{
+		// TODO: report error
+		return false;
+	}
+
+	UpdateJsonBufferObject(BinaryFilePath);
+	return FGLTFJsonBuilder::Serialize(Archive, FilePath);
+}
diff --git a/Engine/Plugins/Enterprise/GLTFExporter/MigratedFromGit/Source/GLTFExporter/Private/Builders/GLTFContainerBuilder.cpp b/Engine/Plugins/Enterprise/GLTFExporter/MigratedFromGit/Source/GLTFExporter/Private/Builders/GLTFContainerBuilder.cpp
new file mode 100644
--- /dev/null
+++ b/Engine/Plugins/Enterprise/GLTFExporter/MigratedFromGit/Source/GLTFExporter/Private/Builders/GLTFContainerBuilder.cpp
@@ -0,0 +1,3 @@
+// Copyright Epic Games, Inc. All Rights Reserved.
+
+#include "Builders/GLTFContainerBuilder.h"
diff --git a/Engine/Plugins/Enterprise/GLTFExporter/MigratedFromGit/Source/GLTFExporter/Private/Builders/GLTFConvertBuilder.cpp b/Engine/Plugins/Enterprise/GLTFExporter/MigratedFromGit/Source/GLTFExporter/Private/Builders/GLTFConvertBuilder.cpp
new file mode 100644
--- /dev/null
+++ b/Engine/Plugins/Enterprise/GLTFExporter/MigratedFromGit/Source/GLTFExporter/Private/Builders/GLTFConvertBuilder.cpp
@@ -0,0 +1,85 @@
+// Copyright Epic Games, Inc. All Rights Reserved.
+
+#include "Builders/GLTFConvertBuilder.h"
+
+FGLTFJsonAccessorIndex FGLTFConvertBuilder::GetOrAddPositionAccessor(const FPositionVertexBuffer* VertexBuffer, const FString& DesiredName)
+{
+	return PositionVertexBufferConverter.GetOrAdd(*this, DesiredName, VertexBuffer);
+}
+
+FGLTFJsonAccessorIndex FGLTFConvertBuilder::GetOrAddColorAccessor(const FColorVertexBuffer* VertexBuffer, const FString& DesiredName)
+{
+	return ColorVertexBufferConverter.GetOrAdd(*this, DesiredName, VertexBuffer);
+}
+
+FGLTFJsonAccessorIndex FGLTFConvertBuilder::GetOrAddNormalAccessor(const FStaticMeshVertexBuffer* VertexBuffer, const FString& DesiredName)
+{
+	return StaticMeshNormalVertexBufferConverter.GetOrAdd(*this, DesiredName, VertexBuffer);
+}
+
+FGLTFJsonAccessorIndex FGLTFConvertBuilder::GetOrAddTangentAccessor(const FStaticMeshVertexBuffer* VertexBuffer, const FString& DesiredName)
+{
+	return StaticMeshTangentVertexBufferConverter.GetOrAdd(*this, DesiredName, VertexBuffer);
+}
+
+FGLTFJsonAccessorIndex FGLTFConvertBuilder::GetOrAddUV0Accessor(const FStaticMeshVertexBuffer* VertexBuffer, const FString& DesiredName)
+{
+	return StaticMeshUV0VertexBufferConverter.GetOrAdd(*this, DesiredName, VertexBuffer);
+}
+
+FGLTFJsonAccessorIndex FGLTFConvertBuilder::GetOrAddUV1Accessor(const FStaticMeshVertexBuffer* VertexBuffer, const FString& DesiredName)
+{
+	return StaticMeshUV1VertexBufferConverter.GetOrAdd(*this, DesiredName, VertexBuffer);
+}
+
+FGLTFJsonBufferViewIndex FGLTFConvertBuilder::GetOrAddIndexBufferView(const FRawStaticIndexBuffer* IndexBuffer, const FString& DesiredName)
+{
+	return StaticMeshIndexBufferConverter.GetOrAdd(*this, DesiredName, IndexBuffer);
+}
+
+FGLTFJsonAccessorIndex FGLTFConvertBuilder::GetOrAddIndexAccessor(const FStaticMeshSection* MeshSection, const FRawStaticIndexBuffer* IndexBuffer, const FString& DesiredName)
+{
+	return StaticMeshSectionConverter.GetOrAdd(*this, DesiredName, MeshSection, IndexBuffer);
+}
+
+FGLTFJsonMeshIndex FGLTFConvertBuilder::GetOrAddMesh(const FStaticMeshLODResources* StaticMeshLOD, const FColorVertexBuffer* OverrideVertexColors, const FString& DesiredName)
+{
+	return StaticMeshConverter.GetOrAdd(*this, DesiredName, StaticMeshLOD, OverrideVertexColors);
+}
+
+FGLTFJsonMeshIndex FGLTFConvertBuilder::GetOrAddMesh(const UStaticMesh* StaticMesh, int32 LODIndex, const FColorVertexBuffer* OverrideVertexColors, const FString& DesiredName)
+{
+	const FStaticMeshLODResources* StaticMeshLOD = &StaticMesh->GetLODForExport(LODIndex);
+	FString Name = DesiredName.IsEmpty() ? StaticMesh->GetName() : DesiredName;
+
+	if (LODIndex != 0)
+	{
+		Name += TEXT("_LOD") + FString::FromInt(LODIndex);
+	}
+
+	return GetOrAddMesh(StaticMeshLOD, OverrideVertexColors, Name);
+}
+
+FGLTFJsonMeshIndex FGLTFConvertBuilder::GetOrAddMesh(const UStaticMeshComponent* StaticMeshComponent, const FString& DesiredName)
+{
+	const UStaticMesh* StaticMesh = StaticMeshComponent->GetStaticMesh();
+	const int32 LODIndex = StaticMeshComponent->ForcedLodModel > 0 ? StaticMeshComponent->ForcedLodModel - 1 : /* auto-select */ 0;
+	const FColorVertexBuffer* OverrideVertexColors = LODIndex < StaticMeshComponent->LODData.Num() ? StaticMeshComponent->LODData[LODIndex].OverrideVertexColors : nullptr;
+
+	return GetOrAddMesh(StaticMesh, LODIndex, OverrideVertexColors, DesiredName);
+}
+
+FGLTFJsonNodeIndex FGLTFConvertBuilder::GetOrAddNode(const USceneComponent* SceneComponent, bool bSelectedOnly, bool bRootNode, const FString& DesiredName)
+{
+	return SceneComponentConverter.GetOrAdd(*this, DesiredName, SceneComponent, bSelectedOnly, bRootNode);
+}
+
+FGLTFJsonSceneIndex FGLTFConvertBuilder::GetOrAddScene(const ULevel* Level, bool bSelectedOnly, const FString& DesiredName)
+{
+	return LevelConverter.GetOrAdd(*this, DesiredName, Level, bSelectedOnly);
+}
+
+FGLTFJsonSceneIndex FGLTFConvertBuilder::GetOrAddScene(const UWorld* World, bool bSelectedOnly, const FString& DesiredName)
+{
+	return GetOrAddScene(World->PersistentLevel, bSelectedOnly, DesiredName.IsEmpty() ? World->GetName() : DesiredName);
+}
diff --git a/Engine/Plugins/Enterprise/GLTFExporter/MigratedFromGit/Source/GLTFExporter/Private/Builders/GLTFJsonBuilder.cpp b/Engine/Plugins/Enterprise/GLTFExporter/MigratedFromGit/Source/GLTFExporter/Private/Builders/GLTFJsonBuilder.cpp
new file mode 100644
--- /dev/null
+++ b/Engine/Plugins/Enterprise/GLTFExporter/MigratedFromGit/Source/GLTFExporter/Private/Builders/GLTFJsonBuilder.cpp
@@ -0,0 +1,47 @@
+// Copyright Epic Games, Inc. All Rights Reserved.
+
+#include "Builders/GLTFJsonBuilder.h"
+
+FGLTFJsonBuilder::FGLTFJsonBuilder()
+{
+}
+
+FGLTFJsonBuilder::~FGLTFJsonBuilder()
+{
+}
+
+FGLTFJsonAccessorIndex FGLTFJsonBuilder::AddAccessor(const FGLTFJsonAccessor& JsonAccessor)
+{
+	return FGLTFJsonAccessorIndex(JsonRoot.Accessors.Add(JsonAccessor));
+}
+
+FGLTFJsonBufferIndex FGLTFJsonBuilder::AddBuffer(const FGLTFJsonBuffer& JsonBuffer)
+{
+	return FGLTFJsonBufferIndex(JsonRoot.Buffers.Add(JsonBuffer));
+}
+
+FGLTFJsonBufferViewIndex FGLTFJsonBuilder::AddBufferView(const FGLTFJsonBufferView& JsonBufferView)
+{
+	return FGLTFJsonBufferViewIndex(JsonRoot.BufferViews.Add(JsonBufferView));
+}
+
+FGLTFJsonMeshIndex FGLTFJsonBuilder::AddMesh(const FGLTFJsonMesh& JsonMesh)
+{
+	return FGLTFJsonMeshIndex(JsonRoot.Meshes.Add(JsonMesh));
+}
+
+FGLTFJsonNodeIndex FGLTFJsonBuilder::AddNode(const FGLTFJsonNode& JsonNode)
+{
+	return FGLTFJsonNodeIndex(JsonRoot.Nodes.Add(JsonNode));
+}
+
+FGLTFJsonSceneIndex FGLTFJsonBuilder::AddScene(const FGLTFJsonScene& JsonScene)
+{
+	return FGLTFJsonSceneIndex(JsonRoot.Scenes.Add(JsonScene));
+}
+
+bool FGLTFJsonBuilder::Serialize(FArchive& Archive, const FString& FilePath)
+{
+	JsonRoot.Serialize(&Archive, true);
+	return true;
+}
diff --git a/Engine/Plugins/Enterprise/GLTFExporter/MigratedFromGit/Source/GLTFExporter/Private/Converters/GLTFConverterUtility.cpp b/Engine/Plugins/Enterprise/GLTFExporter/MigratedFromGit/Source/GLTFExporter/Private/Converters/GLTFConverterUtility.cpp
new file mode 100644
--- /dev/null
+++ b/Engine/Plugins/Enterprise/GLTFExporter/MigratedFromGit/Source/GLTFExporter/Private/Converters/GLTFConverterUtility.cpp
@@ -0,0 +1,14 @@
+// Copyright Epic Games, Inc. All Rights Reserved.
+
+#include "Converters/GLTFConverterUtility.h"
+#include "Engine/Blueprint.h"
+
+bool FGLTFConverterUtility::IsSkySphereBlueprint(const UBlueprint* Blueprint)
+{
+	return Blueprint != nullptr && Blueprint->GetPathName().Equals(TEXT("/Engine/EngineSky/BP_Sky_Sphere.BP_Sky_Sphere"));
+}
+
+bool FGLTFConverterUtility::IsHDRIBackdropBlueprint(const UBlueprint* Blueprint)
+{
+	return Blueprint != nullptr && Blueprint->GetPathName().Equals(TEXT("/HDRIBackdrop/Blueprints/HDRIBackdrop.HDRIBackdrop"));
+}
diff --git a/Engine/Plugins/Enterprise/GLTFExporter/MigratedFromGit/Source/GLTFExporter/Private/Converters/GLTFLevelConverters.cpp b/Engine/Plugins/Enterprise/GLTFExporter/MigratedFromGit/Source/GLTFExporter/Private/Converters/GLTFLevelConverters.cpp
new file mode 100644
--- /dev/null
+++ b/Engine/Plugins/Enterprise/GLTFExporter/MigratedFromGit/Source/GLTFExporter/Private/Converters/GLTFLevelConverters.cpp
@@ -0,0 +1,90 @@
+// Copyright Epic Games, Inc. All Rights Reserved.
+
+#include "Converters/GLTFLevelConverters.h"
+#include "Builders/GLTFContainerBuilder.h"
+#include "Converters/GLTFConverterUtility.h"
+
+FGLTFJsonNodeIndex FGLTFSceneComponentConverter::Add(FGLTFConvertBuilder& Builder, const FString& Name, TTuple<const USceneComponent*, bool, bool> Params)
+{
+	const USceneComponent* SceneComponent = Params.Get<0>();
+	bool bSelectedOnly = Params.Get<1>();
+	bool bRootNode = Params.Get<2>();
+
+	const AActor* Owner = SceneComponent->GetOwner();
+	const bool bIsRootComponent = Owner != nullptr && Owner->GetRootComponent() == SceneComponent;
+	const UBlueprint* Blueprint = UBlueprint::GetBlueprintFromClass(Owner != nullptr ? Owner->GetClass() : nullptr);
+
+	FGLTFJsonNode Node;
+	Node.Name = Name.IsEmpty() ? Owner->GetName() + TEXT("_") + SceneComponent->GetName() : Name;
+
+	const FTransform Transform = bRootNode ? SceneComponent->GetComponentTransform() : SceneComponent->GetRelativeTransform();
+	Node.Translation = FGLTFConverterUtility::ConvertPosition(Transform.GetTranslation());
+	Node.Rotation = FGLTFConverterUtility::ConvertRotation(Transform.GetRotation());
+	Node.Scale = FGLTFConverterUtility::ConvertScale(Transform.GetScale3D());
+
+	if (FGLTFConverterUtility::IsSkySphereBlueprint(Blueprint))
+	{
+		// Ignore components
+	}
+	else if (const UStaticMeshComponent* StaticMeshComponent = Cast<UStaticMeshComponent>(SceneComponent))
+	{
+		Node.Mesh = Builder.GetOrAddMesh(StaticMeshComponent);
+	}
+	else if (FGLTFConverterUtility::IsHDRIBackdropBlueprint(Blueprint) && bIsRootComponent)
+	{
+		// TODO: add support for backdrop
+	}
+
+	const TArray<USceneComponent*>& Children = SceneComponent->GetAttachChildren();
+	for (const USceneComponent* ChildComponent : Children)
+	{
+		if (ChildComponent != nullptr)
+		{
+			const AActor* ChildOwner = ChildComponent->GetOwner();
+			if (!bSelectedOnly || ChildOwner->IsSelected())
+			{
+				FGLTFJsonNodeIndex NodeIndex = Builder.GetOrAddNode(ChildComponent, bSelectedOnly, false);
+				if (NodeIndex != INDEX_NONE)
+				{
+					Node.Children.Add(NodeIndex);
+				}
+			}
+		}
+	}
+
+	return Builder.AddNode(Node);
+}
+
+FGLTFJsonSceneIndex FGLTFLevelConverter::Add(FGLTFConvertBuilder& Builder, const FString& Name, TTuple<const ULevel*, bool> Params)
+{
+	const ULevel* Level = Params.Get<0>();
+	bool bSelectedOnly = Params.Get<1>();
+
+	FGLTFJsonScene Scene;
+	Scene.Name = Name;
+
+	// TODO: export Level->Model ?
+
+	const TArray<AActor*>& Actors = Level->Actors;
+	for (const AActor* Actor : Actors)
+	{
+		if (Actor != nullptr && (!bSelectedOnly || Actor->IsSelected()))
+		{
+			const USceneComponent* RootComponent = Actor->GetRootComponent();
+			if (RootComponent != nullptr)
+			{
+				const AActor* ParentActor = Actor->GetParentActor();
+				if (ParentActor == nullptr || (bSelectedOnly && !ParentActor->IsSelected()))
+				{
+					FGLTFJsonNodeIndex NodeIndex = Builder.GetOrAddNode(RootComponent, bSelectedOnly, true);
+					if (NodeIndex != INDEX_NONE)
+					{
+						Scene.Nodes.Add(NodeIndex);
+					}
+				}
+			}
+		}
+	}
+
+	return Builder.AddScene(Scene);
+}
diff --git a/Engine/Plugins/Enterprise/GLTFExporter/MigratedFromGit/Source/GLTFExporter/Private/Converters/GLTFStaticMeshConverters.cpp b/Engine/Plugins/Enterprise/GLTFExporter/MigratedFromGit/Source/GLTFExporter/Private/Converters/GLTFStaticMeshConverters.cpp
new file mode 100644
--- /dev/null
+++ b/Engine/Plugins/Enterprise/GLTFExporter/MigratedFromGit/Source/GLTFExporter/Private/Converters/GLTFStaticMeshConverters.cpp
@@ -0,0 +1,251 @@
+// Copyright Epic Games, Inc. All Rights Reserved.
+
+#include "Converters/GLTFStaticMeshConverters.h"
+#include "Builders/GLTFConvertBuilder.h"
+#include "Converters/GLTFConverterUtility.h"
+
+FGLTFJsonAccessorIndex FGLTFPositionVertexBufferConverter::Add(FGLTFConvertBuilder& Builder, const FString& Name, const FPositionVertexBuffer* VertexBuffer)
+{
+	const uint32 VertexCount = VertexBuffer->GetNumVertices();
+	if (VertexCount == 0)
+	{
+		return FGLTFJsonAccessorIndex(INDEX_NONE);
+	}
+
+	TArray<FVector> Positions;
+	Positions.AddUninitialized(VertexCount);
+
+	for (uint32 VertexIndex = 0; VertexIndex < VertexCount; ++VertexIndex)
+	{
+		Positions[VertexIndex] = FGLTFConverterUtility::ConvertPosition(VertexBuffer->VertexPosition(VertexIndex));
+	}
+
+	FBox BoundingBox;
+	BoundingBox.Init();
+
+	// More accurate bounding box if based on raw vertex values
+	for (const FVector& Position : Positions)
+	{
+		BoundingBox += Position;
+	}
+
+	FGLTFJsonAccessor JsonAccessor;
+	JsonAccessor.Name = Name;
+	JsonAccessor.BufferView = Builder.AddBufferView(Positions, Name);
+	JsonAccessor.ComponentType = EGLTFJsonComponentType::F32;
+	JsonAccessor.Count = VertexCount;
+	JsonAccessor.Type = EGLTFJsonAccessorType::Vec3;
+
+	JsonAccessor.MinMaxLength = 3;
+	JsonAccessor.Max[0] = BoundingBox.Max.X;
+	JsonAccessor.Max[1] = BoundingBox.Max.Y;
+	JsonAccessor.Max[2] = BoundingBox.Max.Z;
+	JsonAccessor.Min[0] = BoundingBox.Min.X;
+	JsonAccessor.Min[1] = BoundingBox.Min.Y;
+	JsonAccessor.Min[2] = BoundingBox.Min.Z;
+
+	return Builder.AddAccessor(JsonAccessor);
+}
+
+FGLTFJsonAccessorIndex FGLTFColorVertexBufferConverter::Add(FGLTFConvertBuilder& Builder, const FString& Name, const FColorVertexBuffer* VertexBuffer)
+{
+	const uint32 VertexCount = VertexBuffer->GetNumVertices();
+	if (VertexCount == 0)
+	{
+		return FGLTFJsonAccessorIndex(INDEX_NONE);
+	}
+
+	TArray<FGLTFColor> Colors;
+	Colors.AddUninitialized(VertexCount);
+
+	for (uint32 VertexIndex = 0; VertexIndex < VertexCount; ++VertexIndex)
+	{
+		Colors[VertexIndex] = FGLTFConverterUtility::ConvertColor(VertexBuffer->VertexColor(VertexIndex));
+	}
+
+	FGLTFJsonAccessor JsonAccessor;
+	JsonAccessor.Name = Name;
+	JsonAccessor.BufferView = Builder.AddBufferView(Colors, Name);
+	JsonAccessor.ComponentType = EGLTFJsonComponentType::U8;
+	JsonAccessor.Count = VertexCount;
+	JsonAccessor.Type = EGLTFJsonAccessorType::Vec4;
+	JsonAccessor.bNormalized = true;
+
+	return Builder.AddAccessor(JsonAccessor);
+}
+
+FGLTFJsonAccessorIndex FGLTFStaticMeshNormalVertexBufferConverter::Add(FGLTFConvertBuilder& Builder, const FString& Name, const FStaticMeshVertexBuffer* VertexBuffer)
+{
+	const uint32 VertexCount = VertexBuffer->GetNumVertices();
+	if (VertexCount == 0)
+	{
+		return FGLTFJsonAccessorIndex(INDEX_NONE);
+	}
+
+	TArray<FVector> Normals;
+	Normals.AddUninitialized(VertexCount);
+
+	for (uint32 VertexIndex = 0; VertexIndex < VertexCount; ++VertexIndex)
+	{
+		Normals[VertexIndex] = FGLTFConverterUtility::ConvertNormal(VertexBuffer->VertexTangentZ(VertexIndex));
+	}
+
+	FGLTFJsonAccessor JsonAccessor;
+	JsonAccessor.Name = Name;
+	JsonAccessor.BufferView = Builder.AddBufferView(Normals, Name);
+	JsonAccessor.ComponentType = EGLTFJsonComponentType::F32;
+	JsonAccessor.Count = VertexCount;
+	JsonAccessor.Type = EGLTFJsonAccessorType::Vec3;
+
+	return Builder.AddAccessor(JsonAccessor);
+}
+
+FGLTFJsonAccessorIndex FGLTFStaticMeshTangentVertexBufferConverter::Add(FGLTFConvertBuilder& Builder, const FString& Name, const FStaticMeshVertexBuffer* VertexBuffer)
+{
+	const uint32 VertexCount = VertexBuffer->GetNumVertices();
+	if (VertexCount == 0)
+	{
+		return FGLTFJsonAccessorIndex(INDEX_NONE);
+	}
+
+	TArray<FVector4> Tangents;
+	Tangents.AddUninitialized(VertexCount);
+
+	for (uint32 VertexIndex = 0; VertexIndex < VertexCount; ++VertexIndex)
+	{
+		Tangents[VertexIndex] = FGLTFConverterUtility::ConvertTangent(VertexBuffer->VertexTangentX(VertexIndex));
+	}
+
+	FGLTFJsonAccessor JsonAccessor;
+	JsonAccessor.Name = Name;
+	JsonAccessor.BufferView = Builder.AddBufferView(Tangents, Name);
+	JsonAccessor.ComponentType = EGLTFJsonComponentType::F32;
+	JsonAccessor.Count = VertexCount;
+	JsonAccessor.Type = EGLTFJsonAccessorType::Vec4;
+
+	return Builder.AddAccessor(JsonAccessor);
+}
+
+FGLTFJsonAccessorIndex FGLTFStaticMeshUV0VertexBufferConverter::Add(FGLTFConvertBuilder& Builder, const FString& Name, const FStaticMeshVertexBuffer* VertexBuffer)
+{
+	const uint32 VertexCount = VertexBuffer->GetNumVertices();
+	if (VertexCount == 0 || VertexBuffer->GetNumTexCoords() < 1)
+	{
+		return FGLTFJsonAccessorIndex(INDEX_NONE);
+	}
+
+	TArray<FVector2D> UV0s;
+	UV0s.AddUninitialized(VertexCount);
+
+	for (uint32 VertexIndex = 0; VertexIndex < VertexCount; ++VertexIndex)
+	{
+		UV0s[VertexIndex] = VertexBuffer->GetVertexUV(VertexIndex, 0);
+	}
+
+	FGLTFJsonAccessor JsonAccessor;
+	JsonAccessor.Name = Name;
+	JsonAccessor.BufferView = Builder.AddBufferView(UV0s, Name);
+	JsonAccessor.ComponentType = EGLTFJsonComponentType::F32;
+	JsonAccessor.Count = VertexCount;
+	JsonAccessor.Type = EGLTFJsonAccessorType::Vec2;
+
+	return Builder.AddAccessor(JsonAccessor);
+}
+
+FGLTFJsonAccessorIndex FGLTFStaticMeshUV1VertexBufferConverter::Add(FGLTFConvertBuilder& Builder, const FString& Name, const FStaticMeshVertexBuffer* VertexBuffer)
+{
+	const uint32 VertexCount = VertexBuffer->GetNumVertices();
+	if (VertexCount == 0 || VertexBuffer->GetNumTexCoords() < 2)
+	{
+		return FGLTFJsonAccessorIndex(INDEX_NONE);
+	}
+
+	TArray<FVector2D> UV1s;
+	UV1s.AddUninitialized(VertexCount);
+
+	for (uint32 VertexIndex = 0; VertexIndex < VertexCount; ++VertexIndex)
+	{
+		UV1s[VertexIndex] = VertexBuffer->GetVertexUV(VertexIndex, 1);
+	}
+
+	FGLTFJsonAccessor JsonAccessor;
+	JsonAccessor.Name = Name;
+	JsonAccessor.BufferView = Builder.AddBufferView(UV1s, Name);
+	JsonAccessor.ComponentType = EGLTFJsonComponentType::F32;
+	JsonAccessor.Count = VertexCount;
+	JsonAccessor.Type = EGLTFJsonAccessorType::Vec2;
+
+	return Builder.AddAccessor(JsonAccessor);
+}
+
+FGLTFJsonBufferViewIndex FGLTFStaticMeshIndexBufferConverter::Add(FGLTFConvertBuilder& Builder, const FString& Name, const FRawStaticIndexBuffer* IndexBuffer)
+{
+	if (IndexBuffer->GetNumIndices() == 0)
+	{
+		return FGLTFJsonBufferViewIndex(INDEX_NONE);
+	}
+
+	TArray<uint32> Indices;
+	IndexBuffer->GetCopy(Indices);
+	return Builder.AddBufferView(Indices, Name, EGLTFJsonBufferTarget::ElementArrayBuffer);
+}
+
+FGLTFJsonAccessorIndex FGLTFStaticMeshSectionConverter::Add(FGLTFConvertBuilder& Builder, const FString& Name, TTuple<const FStaticMeshSection*, const FRawStaticIndexBuffer*> Params)
+{
+	const FStaticMeshSection* MeshSection = Params.Get<0>();
+	const FRawStaticIndexBuffer* IndexBuffer = Params.Get<1>();
+
+	const uint32 TriangleCount = MeshSection->NumTriangles;
+	if (TriangleCount == 0)
+	{
+		return FGLTFJsonAccessorIndex(INDEX_NONE);
+	}
+
+	FGLTFJsonAccessor JsonAccessor;
+	JsonAccessor.Name = Name;
+	JsonAccessor.BufferView = Builder.GetOrAddIndexBufferView(IndexBuffer);
+	JsonAccessor.ByteOffset = MeshSection->FirstIndex * sizeof(uint32);
+	JsonAccessor.ComponentType = EGLTFJsonComponentType::U32;
+	JsonAccessor.Count = TriangleCount * 3;
+	JsonAccessor.Type = EGLTFJsonAccessorType::Scalar;
+
+	return Builder.AddAccessor(JsonAccessor);
+}
+
+FGLTFJsonMeshIndex FGLTFStaticMeshConverter::Add(FGLTFConvertBuilder& Builder, const FString& Name, TTuple<const FStaticMeshLODResources*, const FColorVertexBuffer*> Params)
+{
+	const FStaticMeshLODResources* StaticMeshLOD = Params.Get<0>();
+	const FColorVertexBuffer* OverrideVertexColors = Params.Get<1>();
+
+	FGLTFJsonMesh JsonMesh;
+	JsonMesh.Name = Name;
+
+	const FPositionVertexBuffer* PositionBuffer = &StaticMeshLOD->VertexBuffers.PositionVertexBuffer;
+	const FStaticMeshVertexBuffer* VertexBuffer = &StaticMeshLOD->VertexBuffers.StaticMeshVertexBuffer;
+	const FColorVertexBuffer* ColorBuffer = OverrideVertexColors != nullptr ? OverrideVertexColors : &StaticMeshLOD->VertexBuffers.ColorVertexBuffer;
+
+	FGLTFJsonAttributes JsonAttributes;
+	JsonAttributes.Position = Builder.GetOrAddPositionAccessor(PositionBuffer, Name.IsEmpty() ? Name : Name + TEXT("_Positions"));
+	JsonAttributes.Normal = Builder.GetOrAddNormalAccessor(VertexBuffer, Name.IsEmpty() ? Name : Name + TEXT("_Normals"));
+	JsonAttributes.Tangent = Builder.GetOrAddTangentAccessor(VertexBuffer, Name.IsEmpty() ? Name : Name + TEXT("_Tangents"));
+	JsonAttributes.TexCoord0 = Builder.GetOrAddUV0Accessor(VertexBuffer, Name.IsEmpty() ? Name : Name + TEXT("_UV0s"));
+	JsonAttributes.TexCoord1 = Builder.GetOrAddUV1Accessor(VertexBuffer, Name.IsEmpty() ? Name : Name + TEXT("_UV1s"));
+	JsonAttributes.Color0 = Builder.GetOrAddColorAccessor(ColorBuffer, Name.IsEmpty() ? Name : Name + TEXT("_Colors"));
+
+	const FRawStaticIndexBuffer* IndexBuffer = &StaticMeshLOD->IndexBuffer;
+	Builder.GetOrAddIndexBufferView(IndexBuffer, Name.IsEmpty() ? Name : Name + TEXT("_Indices"));
+
+	const int32 SectionCount = StaticMeshLOD->Sections.Num();
+	JsonMesh.Primitives.AddDefaulted(SectionCount);
+
+	for (int32 SectionIndex = 0; SectionIndex < SectionCount; ++SectionIndex)
+	{
+		FGLTFJsonPrimitive& JsonPrimitive = JsonMesh.Primitives[SectionIndex];
+		JsonPrimitive.Attributes = JsonAttributes;
+
+		JsonPrimitive.Indices = Builder.GetOrAddIndexAccessor(&StaticMeshLOD->Sections[SectionIndex], IndexBuffer,
+            Name.IsEmpty() ? Name : JsonMesh.Name + (SectionCount != 1 ? TEXT("_Indices_Section") + FString::FromInt(SectionIndex) : TEXT("_Indices")));
+	}
+
+	return Builder.AddMesh(JsonMesh);
+}
diff --git a/Engine/Plugins/Enterprise/GLTFExporter/MigratedFromGit/Source/GLTFExporter/Private/Exporters/GLTFAnimSequenceExporter.cpp b/Engine/Plugins/Enterprise/GLTFExporter/MigratedFromGit/Source/GLTFExporter/Private/Exporters/GLTFAnimSequenceExporter.cpp
new file mode 100644
--- /dev/null
+++ b/Engine/Plugins/Enterprise/GLTFExporter/MigratedFromGit/Source/GLTFExporter/Private/Exporters/GLTFAnimSequenceExporter.cpp
@@ -0,0 +1,24 @@
+// Copyright Epic Games, Inc. All Rights Reserved.
+
+#include "Exporters/GLTFAnimSequenceExporter.h"
+#include "Animation/AnimSequence.h"
+
+UGLTFAnimSequenceExporter::UGLTFAnimSequenceExporter(const FObjectInitializer& ObjectInitializer)
+	: Super(ObjectInitializer)
+{
+	SupportedClass = UAnimSequence::StaticClass();
+}
+
+bool UGLTFAnimSequenceExporter::ExportBinary(UObject* Object, const TCHAR* Type, FArchive& Archive, FFeedbackContext* Warn, int32 FileIndex, uint32 PortFlags)
+{
+	const UAnimSequence* AnimSequence = CastChecked<UAnimSequence>(Object);
+
+	if (!FillExportOptions())
+	{
+		// User cancelled the export
+		return false;
+	}
+
+	// TODO: implement
+	return true;
+}
diff --git a/Engine/Plugins/Enterprise/GLTFExporter/MigratedFromGit/Source/GLTFExporter/Private/Exporters/GLTFExporter.cpp b/Engine/Plugins/Enterprise/GLTFExporter/MigratedFromGit/Source/GLTFExporter/Private/Exporters/GLTFExporter.cpp
new file mode 100644
--- /dev/null
+++ b/Engine/Plugins/Enterprise/GLTFExporter/MigratedFromGit/Source/GLTFExporter/Private/Exporters/GLTFExporter.cpp
@@ -0,0 +1,38 @@
+// Copyright Epic Games, Inc. All Rights Reserved.
+
+#include "Exporters/GLTFExporter.h"
+#include "UnrealExporter.h"
+#include "Engine.h"
+
+UGLTFExporter::UGLTFExporter(const FObjectInitializer& ObjectInitializer)
+	: Super(ObjectInitializer)
+{
+	SupportedClass = nullptr;
+	bText = false;
+	PreferredFormatIndex = 0;
+
+	FormatExtension.Add(TEXT("gltf"));
+	FormatDescription.Add(TEXT("GL Transmission Format"));
+
+	// TODO: uncomment when support for .glb is implemented
+	// FormatExtension.Add(TEXT("glb"));
+	// FormatDescription.Add(TEXT("GL Transmission Format (Binary)"));
+
+	ExportOptions = NewObject<UGLTFExportOptions>(this, TEXT("GLTF Export Options"));
+}
+
+bool UGLTFExporter::FillExportOptions()
+{
+	bool bExportAll = GetBatchMode() && !GetShowExportOption();
+	bool bExportCancel = false;
+
+	ExportOptions->FillOptions(GetBatchMode(), GetShowExportOption(), CurrentFilename, bExportCancel, bExportAll);
+	if (bExportCancel)
+	{
+		SetCancelBatch(GetBatchMode());
+		return false;
+	}
+
+	SetShowExportOption(!bExportAll);
+	return true;
+}
diff --git a/Engine/Plugins/Enterprise/GLTFExporter/MigratedFromGit/Source/GLTFExporter/Private/Exporters/GLTFLevelExporter.cpp b/Engine/Plugins/Enterprise/GLTFExporter/MigratedFromGit/Source/GLTFExporter/Private/Exporters/GLTFLevelExporter.cpp
new file mode 100644
--- /dev/null
+++ b/Engine/Plugins/Enterprise/GLTFExporter/MigratedFromGit/Source/GLTFExporter/Private/Exporters/GLTFLevelExporter.cpp
@@ -0,0 +1,28 @@
+// Copyright Epic Games, Inc. All Rights Reserved.
+
+#include "Exporters/GLTFLevelExporter.h"
+#include "Builders/GLTFContainerBuilder.h"
+#include "Engine/World.h"
+
+UGLTFLevelExporter::UGLTFLevelExporter(const FObjectInitializer& ObjectInitializer)
+	: Super(ObjectInitializer)
+{
+	SupportedClass = UWorld::StaticClass();
+}
+
+bool UGLTFLevelExporter::ExportBinary(UObject* Object, const TCHAR* Type, FArchive& Archive, FFeedbackContext* Warn, int32 FileIndex, uint32 PortFlags)
+{
+	const UWorld* World = CastChecked<UWorld>(Object);
+
+	if (!FillExportOptions())
+	{
+		// User cancelled the export
+		return false;
+	}
+
+	FGLTFContainerBuilder Builder;
+	const FGLTFJsonSceneIndex SceneIndex = Builder.GetOrAddScene(World, bSelectedOnly);
+	Builder.JsonRoot.DefaultScene = SceneIndex;
+
+	return Builder.Serialize(Archive, CurrentFilename);
+}
diff --git a/Engine/Plugins/Enterprise/GLTFExporter/MigratedFromGit/Source/GLTFExporter/Private/Exporters/GLTFSkeletalMeshExporter.cpp b/Engine/Plugins/Enterprise/GLTFExporter/MigratedFromGit/Source/GLTFExporter/Private/Exporters/GLTFSkeletalMeshExporter.cpp
new file mode 100644
--- /dev/null
+++ b/Engine/Plugins/Enterprise/GLTFExporter/MigratedFromGit/Source/GLTFExporter/Private/Exporters/GLTFSkeletalMeshExporter.cpp
@@ -0,0 +1,24 @@
+// Copyright Epic Games, Inc. All Rights Reserved.
+
+#include "Exporters/GLTFSkeletalMeshExporter.h"
+#include "Engine/SkeletalMesh.h"
+
+UGLTFSkeletalMeshExporter::UGLTFSkeletalMeshExporter(const FObjectInitializer& ObjectInitializer)
+	: Super(ObjectInitializer)
+{
+	SupportedClass = USkeletalMesh::StaticClass();
+}
+
+bool UGLTFSkeletalMeshExporter::ExportBinary(UObject* Object, const TCHAR* Type, FArchive& Archive, FFeedbackContext* Warn, int32 FileIndex, uint32 PortFlags)
+{
+	const USkeletalMesh* SkeletalMesh = CastChecked<USkeletalMesh>(Object);
+
+	if (!FillExportOptions())
+	{
+		// User cancelled the export
+		return false;
+	}
+
+	// TODO: implement
+	return true;
+}
diff --git a/Engine/Plugins/Enterprise/GLTFExporter/MigratedFromGit/Source/GLTFExporter/Private/Exporters/GLTFStaticMeshExporter.cpp b/Engine/Plugins/Enterprise/GLTFExporter/MigratedFromGit/Source/GLTFExporter/Private/Exporters/GLTFStaticMeshExporter.cpp
new file mode 100644
--- /dev/null
+++ b/Engine/Plugins/Enterprise/GLTFExporter/MigratedFromGit/Source/GLTFExporter/Private/Exporters/GLTFStaticMeshExporter.cpp
@@ -0,0 +1,38 @@
+// Copyright Epic Games, Inc. All Rights Reserved.
+
+#include "Exporters/GLTFStaticMeshExporter.h"
+#include "Builders/GLTFContainerBuilder.h"
+#include "Engine/StaticMesh.h"
+
+UGLTFStaticMeshExporter::UGLTFStaticMeshExporter(const FObjectInitializer& ObjectInitializer)
+	: Super(ObjectInitializer)
+{
+	SupportedClass = UStaticMesh::StaticClass();
+}
+
+bool UGLTFStaticMeshExporter::ExportBinary(UObject* Object, const TCHAR* Type, FArchive& Archive, FFeedbackContext* Warn, int32 FileIndex, uint32 PortFlags)
+{
+	const UStaticMesh* StaticMesh = CastChecked<UStaticMesh>(Object);
+
+	if (!FillExportOptions())
+	{
+		// User cancelled the export
+		return false;
+	}
+
+	FGLTFContainerBuilder Builder;
+
+	FGLTFJsonMeshIndex MeshIndex = Builder.GetOrAddMesh(StaticMesh);
+
+	FGLTFJsonNode Node;
+	Node.Mesh = MeshIndex;
+	FGLTFJsonNodeIndex NodeIndex = Builder.AddNode(Node);
+
+	FGLTFJsonScene Scene;
+	Scene.Nodes.Add(NodeIndex);
+	FGLTFJsonSceneIndex SceneIndex = Builder.AddScene(Scene);
+
+	Builder.JsonRoot.DefaultScene = SceneIndex;
+
+	return Builder.Serialize(Archive, CurrentFilename);
+}
diff --git a/Engine/Plugins/Enterprise/GLTFExporter/MigratedFromGit/Source/GLTFExporter/Private/GLTFAnimSequenceExporter.cpp b/Engine/Plugins/Enterprise/GLTFExporter/MigratedFromGit/Source/GLTFExporter/Private/GLTFAnimSequenceExporter.cpp
deleted file mode 100644
diff --git a/Engine/Plugins/Enterprise/GLTFExporter/MigratedFromGit/Source/GLTFExporter/Private/GLTFBufferBuilder.cpp b/Engine/Plugins/Enterprise/GLTFExporter/MigratedFromGit/Source/GLTFExporter/Private/GLTFBufferBuilder.cpp
deleted file mode 100644
diff --git a/Engine/Plugins/Enterprise/GLTFExporter/MigratedFromGit/Source/GLTFExporter/Private/GLTFContainerBuilder.cpp b/Engine/Plugins/Enterprise/GLTFExporter/MigratedFromGit/Source/GLTFExporter/Private/GLTFContainerBuilder.cpp
deleted file mode 100644
diff --git a/Engine/Plugins/Enterprise/GLTFExporter/MigratedFromGit/Source/GLTFExporter/Private/GLTFConvertBuilder.cpp b/Engine/Plugins/Enterprise/GLTFExporter/MigratedFromGit/Source/GLTFExporter/Private/GLTFConvertBuilder.cpp
deleted file mode 100644
diff --git a/Engine/Plugins/Enterprise/GLTFExporter/MigratedFromGit/Source/GLTFExporter/Private/GLTFConverterUtility.cpp b/Engine/Plugins/Enterprise/GLTFExporter/MigratedFromGit/Source/GLTFExporter/Private/GLTFConverterUtility.cpp
deleted file mode 100644
diff --git a/Engine/Plugins/Enterprise/GLTFExporter/MigratedFromGit/Source/GLTFExporter/Private/GLTFExportOptions.cpp b/Engine/Plugins/Enterprise/GLTFExporter/MigratedFromGit/Source/GLTFExporter/Private/GLTFExportOptions.cpp
--- a/Engine/Plugins/Enterprise/GLTFExporter/MigratedFromGit/Source/GLTFExporter/Private/GLTFExportOptions.cpp
+++ b/Engine/Plugins/Enterprise/GLTFExporter/MigratedFromGit/Source/GLTFExporter/Private/GLTFExportOptions.cpp
@@ -1,7 +1,7 @@
 // Copyright Epic Games, Inc. All Rights Reserved.
 
 #include "GLTFExportOptions.h"
-#include "GLTFExportOptionsWindow.h"
+#include "UI/GLTFExportOptionsWindow.h"
 #include "Misc/App.h"
 
 #include "UObject/ObjectMacros.h"
diff --git a/Engine/Plugins/Enterprise/GLTFExporter/MigratedFromGit/Source/GLTFExporter/Private/GLTFExportOptionsWindow.cpp b/Engine/Plugins/Enterprise/GLTFExporter/MigratedFromGit/Source/GLTFExporter/Private/GLTFExportOptionsWindow.cpp
deleted file mode 100644
diff --git a/Engine/Plugins/Enterprise/GLTFExporter/MigratedFromGit/Source/GLTFExporter/Private/GLTFExporter.cpp b/Engine/Plugins/Enterprise/GLTFExporter/MigratedFromGit/Source/GLTFExporter/Private/GLTFExporter.cpp
deleted file mode 100644
diff --git a/Engine/Plugins/Enterprise/GLTFExporter/MigratedFromGit/Source/GLTFExporter/Private/GLTFJsonBuilder.cpp b/Engine/Plugins/Enterprise/GLTFExporter/MigratedFromGit/Source/GLTFExporter/Private/GLTFJsonBuilder.cpp
deleted file mode 100644
diff --git a/Engine/Plugins/Enterprise/GLTFExporter/MigratedFromGit/Source/GLTFExporter/Private/GLTFLevelConverters.cpp b/Engine/Plugins/Enterprise/GLTFExporter/MigratedFromGit/Source/GLTFExporter/Private/GLTFLevelConverters.cpp
deleted file mode 100644
diff --git a/Engine/Plugins/Enterprise/GLTFExporter/MigratedFromGit/Source/GLTFExporter/Private/GLTFLevelExporter.cpp b/Engine/Plugins/Enterprise/GLTFExporter/MigratedFromGit/Source/GLTFExporter/Private/GLTFLevelExporter.cpp
deleted file mode 100644
diff --git a/Engine/Plugins/Enterprise/GLTFExporter/MigratedFromGit/Source/GLTFExporter/Private/GLTFSkeletalMeshExporter.cpp b/Engine/Plugins/Enterprise/GLTFExporter/MigratedFromGit/Source/GLTFExporter/Private/GLTFSkeletalMeshExporter.cpp
deleted file mode 100644
diff --git a/Engine/Plugins/Enterprise/GLTFExporter/MigratedFromGit/Source/GLTFExporter/Private/GLTFStaticMeshConverters.cpp b/Engine/Plugins/Enterprise/GLTFExporter/MigratedFromGit/Source/GLTFExporter/Private/GLTFStaticMeshConverters.cpp
deleted file mode 100644
diff --git a/Engine/Plugins/Enterprise/GLTFExporter/MigratedFromGit/Source/GLTFExporter/Private/GLTFStaticMeshExporter.cpp b/Engine/Plugins/Enterprise/GLTFExporter/MigratedFromGit/Source/GLTFExporter/Private/GLTFStaticMeshExporter.cpp
deleted file mode 100644
diff --git a/Engine/Plugins/Enterprise/GLTFExporter/MigratedFromGit/Source/GLTFExporter/Private/UI/GLTFExportOptionsWindow.cpp b/Engine/Plugins/Enterprise/GLTFExporter/MigratedFromGit/Source/GLTFExporter/Private/UI/GLTFExportOptionsWindow.cpp
new file mode 100644
--- /dev/null
+++ b/Engine/Plugins/Enterprise/GLTFExporter/MigratedFromGit/Source/GLTFExporter/Private/UI/GLTFExportOptionsWindow.cpp
@@ -0,0 +1,256 @@
+// Copyright 1998-2019 Epic Games, Inc. All Rights Reserved.
+
+#include "UI/GLTFExportOptionsWindow.h"
+#include "Modules/ModuleManager.h"
+#include "Widgets/Layout/SBorder.h"
+#include "Widgets/Text/STextBlock.h"
+#include "Widgets/Layout/SBox.h"
+#include "Widgets/Layout/SUniformGridPanel.h"
+#include "Widgets/Input/SButton.h"
+#include "EditorStyleSet.h"
+#include "PropertyEditorModule.h"
+#include "IDetailsView.h"
+#include "Widgets/SWindow.h"
+#include "Framework/Application/SlateApplication.h"
+#include "Interfaces/IMainFrameModule.h"
+
+#define LOCTEXT_NAMESPACE "GLTFExportOptions"
+
+void SGLTFExportOptionsWindow::Construct(const FArguments& InArgs)
+{
+	ExportOptions = InArgs._ExportOptions;
+	WidgetWindow = InArgs._WidgetWindow;
+
+	check(ExportOptions);
+
+	const FText CancelText = InArgs._BatchMode ? LOCTEXT("GLTFExportOptionsWindow_CancelBatch", "Cancel All") : LOCTEXT("GLTFExportOptionsWindow_Cancel", "Cancel");
+	const FText CancelTooltipText = InArgs._BatchMode ? LOCTEXT("GLTFExportOptionsWindow_Cancel_ToolTip_Batch", "Cancel the batch export.") : LOCTEXT("GLTFExportOptionsWindow_Cancel_ToolTip", "Cancel the current glTF export.");
+
+	TSharedPtr<SBox> HeaderToolBox;
+	TSharedPtr<SHorizontalBox> HeaderButtons;
+	TSharedPtr<SBox> InspectorBox;
+	this->ChildSlot
+	[
+		SNew(SBox)
+		[
+			SNew(SVerticalBox)
+			+ SVerticalBox::Slot()
+			.AutoHeight()
+			.Padding(2)
+			[
+				SAssignNew(HeaderToolBox, SBox)
+			]
+			+SVerticalBox::Slot()
+			.AutoHeight()
+			.Padding(2)
+			[
+				SNew(SBorder)
+				.Padding(FMargin(3))
+				.BorderImage(FEditorStyle::GetBrush("ToolPanel.GroupBorder"))
+				[
+					SNew(SHorizontalBox)
+					+SHorizontalBox::Slot()
+					.AutoWidth()
+					[
+						SNew(STextBlock)
+						.Font(FEditorStyle::GetFontStyle("CurveEd.LabelFont"))
+						.Text(LOCTEXT("Export_CurrentFileTitle", "Current File: "))
+					]
+					+SHorizontalBox::Slot()
+					.Padding(5, 0, 0, 0)
+					.AutoWidth()
+					.VAlign(VAlign_Center)
+					[
+						SNew(STextBlock)
+						.Font(FEditorStyle::GetFontStyle("CurveEd.InfoFont"))
+						.Text(InArgs._FullPath)
+					]
+				]
+			]
+			+ SVerticalBox::Slot()
+			.FillHeight(1.0f)
+			.Padding(2)
+			[
+				SAssignNew(InspectorBox, SBox)
+			]
+			+ SVerticalBox::Slot()
+			.AutoHeight()
+			.HAlign(HAlign_Right)
+			.Padding(2)
+			[
+				SNew(SUniformGridPanel)
+				.SlotPadding(2)
+				+ SUniformGridPanel::Slot(0, 0)
+				[
+					SNew(SButton)
+					.HAlign(HAlign_Center)
+					.Text(LOCTEXT("GLTFExportOptionsWindow_ExportAll", "Export All"))
+					.ToolTipText(LOCTEXT("GLTExportOptionsWindow_ExportAll_ToolTip", "Export all files with these same settings"))
+					.Visibility(InArgs._BatchMode ? EVisibility::All : EVisibility::Hidden)
+					.OnClicked(this, &SGLTFExportOptionsWindow::OnExportAll)
+				]
+				+ SUniformGridPanel::Slot(1, 0)
+				[
+					SAssignNew(ImportButton, SButton)
+					.HAlign(HAlign_Center)
+					.Text(LOCTEXT("GLTExportOptionsWindow_Export", "Export"))
+					.OnClicked(this, &SGLTFExportOptionsWindow::OnExport)
+				]
+				+ SUniformGridPanel::Slot(2, 0)
+				[
+					SNew(SButton)
+					.HAlign(HAlign_Center)
+					.Text(CancelText)
+					.ToolTipText(CancelTooltipText)
+					.OnClicked(this, &SGLTFExportOptionsWindow::OnCancel)
+				]
+			]
+		]
+	];
+
+	FPropertyEditorModule& PropertyEditorModule = FModuleManager::GetModuleChecked<FPropertyEditorModule>("PropertyEditor");
+	FDetailsViewArgs DetailsViewArgs;
+	DetailsViewArgs.bAllowSearch = false;
+	DetailsViewArgs.NameAreaSettings = FDetailsViewArgs::HideNameArea;
+	DetailsView = PropertyEditorModule.CreateDetailView(DetailsViewArgs);
+
+	InspectorBox->SetContent(DetailsView->AsShared());
+
+	HeaderToolBox->SetContent(
+		SNew(SBorder)
+		.Padding(FMargin(3))
+		.BorderImage(FEditorStyle::GetBrush("ToolPanel.GroupBorder"))
+		[
+			SNew(SHorizontalBox)
+			+ SHorizontalBox::Slot()
+			[
+				SNew(SBox)
+				.HAlign(HAlign_Right)
+				[
+					SAssignNew(HeaderButtons, SHorizontalBox)
+					+ SHorizontalBox::Slot()
+					.AutoWidth()
+					.Padding(FMargin(2.0f, 0.0f))
+					[
+						SNew(SButton)
+						.Text(LOCTEXT("ExportOptionsWindow_ResetOptions", "Reset to Default"))
+						.OnClicked(this, &SGLTFExportOptionsWindow::OnResetToDefaultClick)
+					]
+				]
+			]
+		]
+	);
+
+	DetailsView->SetObject(ExportOptions);
+}
+
+FReply SGLTFExportOptionsWindow::OnResetToDefaultClick() const
+{
+	ExportOptions->ResetToDefault();
+	//Refresh the view to make sure the custom UI are updating correctly
+	DetailsView->SetObject(ExportOptions, true);
+	return FReply::Handled();
+}
+
+FReply SGLTFExportOptionsWindow::OnExport()
+{
+	bShouldExport = true;
+	if (WidgetWindow.IsValid())
+	{
+		WidgetWindow.Pin()->RequestDestroyWindow();
+	}
+	return FReply::Handled();
+}
+
+FReply SGLTFExportOptionsWindow::OnExportAll()
+{
+	bShouldExportAll = true;
+	return OnExport();
+}
+
+FReply SGLTFExportOptionsWindow::OnCancel()
+{
+	bShouldExport = false;
+	bShouldExportAll = false;
+	if (WidgetWindow.IsValid())
+	{
+		WidgetWindow.Pin()->RequestDestroyWindow();
+	}
+	return FReply::Handled();
+}
+
+FReply SGLTFExportOptionsWindow::OnKeyDown(const FGeometry & MyGeometry, const FKeyEvent & InKeyEvent)
+{
+	if (InKeyEvent.GetKey() == EKeys::Escape)
+	{
+		return OnCancel();
+	}
+
+	return FReply::Unhandled();
+}
+
+bool SGLTFExportOptionsWindow::ShouldExport() const
+{
+	return bShouldExport;
+}
+
+bool SGLTFExportOptionsWindow::ShouldExportAll() const
+{
+	return bShouldExportAll;
+}
+
+SGLTFExportOptionsWindow::SGLTFExportOptionsWindow()
+	: ExportOptions(nullptr)
+	, bShouldExport(false)
+	, bShouldExportAll(false)
+{}
+
+void SGLTFExportOptionsWindow::ShowDialog(UGLTFExportOptions* ExportOptions, const FString& FullPath, bool bBatchMode, bool& bOutOperationCanceled, bool& bOutExportAll)
+{
+	check(ExportOptions != nullptr);
+
+	bOutOperationCanceled = false;
+	bOutExportAll = false;
+
+	TSharedPtr<SWindow> ParentWindow;
+
+	if (FModuleManager::Get().IsModuleLoaded("MainFrame"))
+	{
+		IMainFrameModule& MainFrame = FModuleManager::LoadModuleChecked<IMainFrameModule>("MainFrame");
+		ParentWindow = MainFrame.GetParentWindow();
+	}
+
+	const bool bBinaryFile = FPaths::GetExtension(FullPath, false).Compare(TEXT("glb"), ESearchCase::IgnoreCase) == 0;
+	const FText Title = bBinaryFile ?
+		NSLOCTEXT("UnrealEd", "GLTFExportOptionsTitleBinary", "glTF (Binary) Export Options") :
+		NSLOCTEXT("UnrealEd", "GLTFExportOptionsTitle", "glTF Export Options");
+
+	TSharedRef<SWindow> Window = SNew(SWindow)
+		.Title(Title)
+		.SizingRule(ESizingRule::UserSized)
+		.AutoCenter(EAutoCenter::PrimaryWorkArea)
+		.ClientSize(FVector2D(500, 445));
+
+	TSharedPtr<SGLTFExportOptionsWindow> OptionsWindow;
+	Window->SetContent
+	(
+		SAssignNew(OptionsWindow, SGLTFExportOptionsWindow)
+		.ExportOptions(ExportOptions)
+		.WidgetWindow(Window)
+		.FullPath(FText::FromString(FullPath))
+		.BatchMode(bBatchMode)
+	);
+
+	FSlateApplication::Get().AddModalWindow(Window, ParentWindow, false);
+
+	if (OptionsWindow->ShouldExport())
+	{
+		bOutExportAll = OptionsWindow->ShouldExportAll();
+	}
+	else
+	{
+		bOutOperationCanceled = true;
+	}
+}
+
+#undef LOCTEXT_NAMESPACE
diff --git a/Engine/Plugins/Enterprise/GLTFExporter/MigratedFromGit/Source/GLTFExporter/Public/Builders/GLTFBufferBuilder.h b/Engine/Plugins/Enterprise/GLTFExporter/MigratedFromGit/Source/GLTFExporter/Public/Builders/GLTFBufferBuilder.h
new file mode 100644
--- /dev/null
+++ b/Engine/Plugins/Enterprise/GLTFExporter/MigratedFromGit/Source/GLTFExporter/Public/Builders/GLTFBufferBuilder.h
@@ -0,0 +1,24 @@
+// Copyright Epic Games, Inc. All Rights Reserved.
+
+#pragma once
+
+#include "Builders/GLTFJsonBuilder.h"
+
+struct GLTFEXPORTER_API FGLTFBufferBuilder : public FGLTFJsonBuilder
+{
+	FGLTFJsonBufferIndex BufferIndex;
+	TArray<uint8> BufferData;
+
+	FGLTFBufferBuilder();
+
+	FGLTFJsonBufferViewIndex AddBufferView(const void* RawData, uint64 ByteLength, const FString& Name = TEXT(""), EGLTFJsonBufferTarget BufferTarget = EGLTFJsonBufferTarget::ArrayBuffer);
+
+	template <class ElementType>
+	FGLTFJsonBufferViewIndex AddBufferView(const TArray<ElementType>& Array, const FString& Name = TEXT(""), EGLTFJsonBufferTarget BufferTarget = EGLTFJsonBufferTarget::ArrayBuffer)
+	{
+		return AddBufferView(Array.GetData(), Array.Num() * sizeof(ElementType), Name, BufferTarget);
+	}
+
+	void UpdateJsonBufferObject(const FString& BinaryFilePath);
+	virtual bool Serialize(FArchive& Archive, const FString& FilePath) override;
+};
diff --git a/Engine/Plugins/Enterprise/GLTFExporter/MigratedFromGit/Source/GLTFExporter/Public/Builders/GLTFContainerBuilder.h b/Engine/Plugins/Enterprise/GLTFExporter/MigratedFromGit/Source/GLTFExporter/Public/Builders/GLTFContainerBuilder.h
new file mode 100644
--- /dev/null
+++ b/Engine/Plugins/Enterprise/GLTFExporter/MigratedFromGit/Source/GLTFExporter/Public/Builders/GLTFContainerBuilder.h
@@ -0,0 +1,9 @@
+// Copyright Epic Games, Inc. All Rights Reserved.
+
+#pragma once
+
+#include "Builders/GLTFConvertBuilder.h"
+
+struct GLTFEXPORTER_API FGLTFContainerBuilder : public FGLTFConvertBuilder
+{
+};
diff --git a/Engine/Plugins/Enterprise/GLTFExporter/MigratedFromGit/Source/GLTFExporter/Public/Builders/GLTFConvertBuilder.h b/Engine/Plugins/Enterprise/GLTFExporter/MigratedFromGit/Source/GLTFExporter/Public/Builders/GLTFConvertBuilder.h
new file mode 100644
--- /dev/null
+++ b/Engine/Plugins/Enterprise/GLTFExporter/MigratedFromGit/Source/GLTFExporter/Public/Builders/GLTFConvertBuilder.h
@@ -0,0 +1,41 @@
+// Copyright Epic Games, Inc. All Rights Reserved.
+
+#pragma once
+
+#include "Builders/GLTFBufferBuilder.h"
+#include "Converters/GLTFStaticMeshConverters.h"
+#include "Converters/GLTFLevelConverters.h"
+
+struct GLTFEXPORTER_API FGLTFConvertBuilder : public FGLTFBufferBuilder
+{
+	FGLTFPositionVertexBufferConverter PositionVertexBufferConverter;
+	FGLTFColorVertexBufferConverter ColorVertexBufferConverter;
+	FGLTFStaticMeshNormalVertexBufferConverter StaticMeshNormalVertexBufferConverter;
+	FGLTFStaticMeshTangentVertexBufferConverter StaticMeshTangentVertexBufferConverter;
+	FGLTFStaticMeshUV0VertexBufferConverter StaticMeshUV0VertexBufferConverter;
+	FGLTFStaticMeshUV1VertexBufferConverter StaticMeshUV1VertexBufferConverter;
+	FGLTFStaticMeshIndexBufferConverter StaticMeshIndexBufferConverter;
+	FGLTFStaticMeshSectionConverter StaticMeshSectionConverter;
+	FGLTFStaticMeshConverter StaticMeshConverter;
+
+	FGLTFSceneComponentConverter SceneComponentConverter;
+	FGLTFLevelConverter LevelConverter;
+
+	FGLTFJsonAccessorIndex GetOrAddPositionAccessor(const FPositionVertexBuffer* VertexBuffer, const FString& DesiredName = TEXT(""));
+	FGLTFJsonAccessorIndex GetOrAddNormalAccessor(const FStaticMeshVertexBuffer* VertexBuffer, const FString& DesiredName = TEXT(""));
+	FGLTFJsonAccessorIndex GetOrAddTangentAccessor(const FStaticMeshVertexBuffer* VertexBuffer, const FString& DesiredName = TEXT(""));
+	FGLTFJsonAccessorIndex GetOrAddUV0Accessor(const FStaticMeshVertexBuffer* VertexBuffer, const FString& DesiredName = TEXT(""));
+	FGLTFJsonAccessorIndex GetOrAddUV1Accessor(const FStaticMeshVertexBuffer* VertexBuffer, const FString& DesiredName = TEXT(""));
+	FGLTFJsonAccessorIndex GetOrAddColorAccessor(const FColorVertexBuffer* VertexBuffer, const FString& DesiredName = TEXT(""));
+
+	FGLTFJsonBufferViewIndex GetOrAddIndexBufferView(const FRawStaticIndexBuffer* IndexBuffer, const FString& DesiredName = TEXT(""));
+	FGLTFJsonAccessorIndex GetOrAddIndexAccessor(const FStaticMeshSection* MeshSection, const FRawStaticIndexBuffer* IndexBuffer, const FString& DesiredName = TEXT(""));
+
+	FGLTFJsonMeshIndex GetOrAddMesh(const FStaticMeshLODResources* StaticMeshLOD, const FColorVertexBuffer* OverrideVertexColors = nullptr, const FString& DesiredName = TEXT(""));
+	FGLTFJsonMeshIndex GetOrAddMesh(const UStaticMesh* StaticMesh, int32 LODIndex = 0, const FColorVertexBuffer* OverrideVertexColors = nullptr, const FString& DesiredName = TEXT(""));
+	FGLTFJsonMeshIndex GetOrAddMesh(const UStaticMeshComponent* StaticMeshComponent, const FString& DesiredName = TEXT(""));
+
+	FGLTFJsonNodeIndex GetOrAddNode(const USceneComponent* SceneComponent, bool bSelectedOnly, bool bRootNode = false, const FString& DesiredName = TEXT(""));
+	FGLTFJsonSceneIndex GetOrAddScene(const ULevel* Level, bool bSelectedOnly, const FString& DesiredName = TEXT(""));
+	FGLTFJsonSceneIndex GetOrAddScene(const UWorld* World, bool bSelectedOnly, const FString& DesiredName = TEXT(""));
+};
diff --git a/Engine/Plugins/Enterprise/GLTFExporter/MigratedFromGit/Source/GLTFExporter/Public/Builders/GLTFJsonBuilder.h b/Engine/Plugins/Enterprise/GLTFExporter/MigratedFromGit/Source/GLTFExporter/Public/Builders/GLTFJsonBuilder.h
new file mode 100644
--- /dev/null
+++ b/Engine/Plugins/Enterprise/GLTFExporter/MigratedFromGit/Source/GLTFExporter/Public/Builders/GLTFJsonBuilder.h
@@ -0,0 +1,22 @@
+// Copyright Epic Games, Inc. All Rights Reserved.
+
+#pragma once
+
+#include "Json/GLTFJsonRoot.h"
+
+struct GLTFEXPORTER_API FGLTFJsonBuilder
+{
+	FGLTFJsonRoot JsonRoot;
+
+	FGLTFJsonBuilder();
+	virtual ~FGLTFJsonBuilder();
+
+	FGLTFJsonAccessorIndex AddAccessor(const FGLTFJsonAccessor& JsonAccessor);
+	FGLTFJsonBufferIndex AddBuffer(const FGLTFJsonBuffer& JsonBuffer);
+	FGLTFJsonBufferViewIndex AddBufferView(const FGLTFJsonBufferView& JsonBufferView);
+	FGLTFJsonMeshIndex AddMesh(const FGLTFJsonMesh& JsonMesh);
+	FGLTFJsonNodeIndex AddNode(const FGLTFJsonNode& JsonNode);
+	FGLTFJsonSceneIndex AddScene(const FGLTFJsonScene& JsonScene);
+
+	virtual bool Serialize(FArchive& Archive, const FString& FilePath);
+};
diff --git a/Engine/Plugins/Enterprise/GLTFExporter/MigratedFromGit/Source/GLTFExporter/Public/Converters/GLTFColor.h b/Engine/Plugins/Enterprise/GLTFExporter/MigratedFromGit/Source/GLTFExporter/Public/Converters/GLTFColor.h
new file mode 100644
--- /dev/null
+++ b/Engine/Plugins/Enterprise/GLTFExporter/MigratedFromGit/Source/GLTFExporter/Public/Converters/GLTFColor.h
@@ -0,0 +1,17 @@
+// Copyright Epic Games, Inc. All Rights Reserved.
+
+#pragma once
+
+#include "CoreTypes.h"
+
+// UE4 uses an endianness dependent memory layout (ABGR for little endian, ARGB for big endian)
+// glTF uses RGBA independent of endianness, hence why we need this dedicated struct
+struct GLTFEXPORTER_API FGLTFColor
+{
+	uint8 R,G,B,A;
+
+	FORCEINLINE FGLTFColor(uint8 R, uint8 G, uint8 B, uint8 A = 255)
+		: R(R), G(G), B(B), A(A)
+	{
+	}
+};
diff --git a/Engine/Plugins/Enterprise/GLTFExporter/MigratedFromGit/Source/GLTFExporter/Public/Converters/GLTFConverter.h b/Engine/Plugins/Enterprise/GLTFExporter/MigratedFromGit/Source/GLTFExporter/Public/Converters/GLTFConverter.h
new file mode 100644
--- /dev/null
+++ b/Engine/Plugins/Enterprise/GLTFExporter/MigratedFromGit/Source/GLTFExporter/Public/Converters/GLTFConverter.h
@@ -0,0 +1,42 @@
+// Copyright Epic Games, Inc. All Rights Reserved.
+
+#pragma once
+
+#include "CoreMinimal.h"
+
+struct FGLTFConvertBuilder;
+
+template <class IndexType, class KeyType>
+class GLTFEXPORTER_API TGLTFConverter
+{
+public:
+	virtual ~TGLTFConverter() = default;
+
+	template <typename... ArgsType>
+    FORCEINLINE IndexType Get(ArgsType&&... Args) const
+	{
+		const KeyType Key(Forward<ArgsType>(Args)...);
+		return IndexLookup.FindRef(Key);
+	}
+
+	template <typename... ArgsType>
+    FORCEINLINE IndexType GetOrAdd(FGLTFConvertBuilder& Builder, const FString& DesiredName, ArgsType&&... Args)
+	{
+		const KeyType Key(Forward<ArgsType>(Args)...);
+
+		IndexType Index = IndexLookup.FindRef(Key);
+		if (Index == INDEX_NONE)
+		{
+			Index = Add(Builder, DesiredName, Key);
+			IndexLookup.Add(Key, Index);
+		}
+
+		return Index;
+	}
+
+private:
+
+	TMap<KeyType, IndexType> IndexLookup;
+
+	virtual IndexType Add(FGLTFConvertBuilder& Builder, const FString& Name, KeyType Key) = 0;
+};
diff --git a/Engine/Plugins/Enterprise/GLTFExporter/MigratedFromGit/Source/GLTFExporter/Public/Converters/GLTFConverterUtility.h b/Engine/Plugins/Enterprise/GLTFExporter/MigratedFromGit/Source/GLTFExporter/Public/Converters/GLTFConverterUtility.h
new file mode 100644
--- /dev/null
+++ b/Engine/Plugins/Enterprise/GLTFExporter/MigratedFromGit/Source/GLTFExporter/Public/Converters/GLTFConverterUtility.h
@@ -0,0 +1,92 @@
+// Copyright Epic Games, Inc. All Rights Reserved.
+
+#pragma once
+
+#include "Converters/GLTFColor.h"
+#include "Json/GLTFJsonEnums.h"
+#include "Engine/EngineTypes.h"
+
+struct GLTFEXPORTER_API FGLTFConverterUtility
+{
+	static FVector ConvertVector(const FVector& Vector)
+	{
+		// UE4 uses a left-handed coordinate system, with Z up.
+		// glTF uses a right-handed coordinate system, with Y up.
+		return { Vector.X, Vector.Z, Vector.Y };
+	}
+
+	static FVector ConvertPosition(const FVector& Position)
+	{
+		return ConvertVector(Position) * 0.01f; // TODO: use options export scale instead of hardcoded value
+	}
+
+	static FVector ConvertScale(const FVector& Scale)
+	{
+		return ConvertVector(Scale);
+	}
+
+	static FVector ConvertNormal(const FVector& Normal)
+	{
+		return ConvertVector(Normal);
+	}
+
+	static FVector4 ConvertTangent(const FVector& Tangent)
+	{
+		// glTF stores tangent as Vec4, with W component indicating handedness of tangent basis.
+		return FVector4(ConvertVector(Tangent), 1.0f);
+	}
+
+	static FGLTFColor ConvertColor(const FColor& Color)
+	{
+		// UE4 uses ABGR or ARGB depending on endianness.
+		// glTF always uses RGBA independent of endianness.
+		return { Color.R, Color.G, Color.B, Color.A };
+	}
+
+	static FQuat ConvertRotation(const FQuat& Rotation)
+	{
+		// UE4 uses a left-handed coordinate system, with Z up.
+		// glTF uses a right-handed coordinate system, with Y up.
+		// Rotation = (qX, qY, qZ, qW) = (sin(angle/2) * aX, sin(angle/2) * aY, sin(angle/2) * aZ, cons(angle/2))
+		// where (aX, aY, aZ) - rotation axis, angle - rotation angle
+		// Y swapped with Z between these coordinate systems
+		// also, as handedness is changed rotation is inversed - hence negation
+		// therefore glTFRotation = (-qX, -qZ, -qY, qw)
+
+		const FQuat Result(-Rotation.X, -Rotation.Z, -Rotation.Y, Rotation.W);
+		// Not checking if quaternion is normalized
+		// e.g. some sources use non-unit Quats for rotation tangents
+		return Result;
+	}
+
+	static FMatrix ConvertMatrix(const FMatrix& Matrix)
+	{
+		// Unreal stores matrix elements in row major order.
+		// glTF stores matrix elements in column major order.
+
+		FMatrix Result;
+		for (int32 Row = 0; Row < 4; ++Row)
+		{
+			for (int32 Col = 0; Col < 4; ++Col)
+			{
+				Result.M[Col][Row] = Matrix.M[Row][Col];
+			}
+		}
+		return Result;
+	}
+
+	static EGLTFJsonAlphaMode ConvertAlphaMode(const EBlendMode Mode)
+	{
+		switch (Mode)
+		{
+			case EBlendMode::BLEND_Opaque:      return EGLTFJsonAlphaMode::Opaque;
+			case EBlendMode::BLEND_Translucent: return EGLTFJsonAlphaMode::Blend;
+			case EBlendMode::BLEND_Masked:      return EGLTFJsonAlphaMode::Mask;
+			default:                            return EGLTFJsonAlphaMode::Opaque; // fallback
+		}
+	}
+
+	static bool IsSkySphereBlueprint(const UBlueprint* Blueprint);
+
+	static bool IsHDRIBackdropBlueprint(const UBlueprint* Blueprint);
+};
diff --git a/Engine/Plugins/Enterprise/GLTFExporter/MigratedFromGit/Source/GLTFExporter/Public/Converters/GLTFLevelConverters.h b/Engine/Plugins/Enterprise/GLTFExporter/MigratedFromGit/Source/GLTFExporter/Public/Converters/GLTFLevelConverters.h
new file mode 100644
--- /dev/null
+++ b/Engine/Plugins/Enterprise/GLTFExporter/MigratedFromGit/Source/GLTFExporter/Public/Converters/GLTFLevelConverters.h
@@ -0,0 +1,17 @@
+// Copyright Epic Games, Inc. All Rights Reserved.
+
+#pragma once
+
+#include "Json/GLTFJsonIndex.h"
+#include "Converters/GLTFConverter.h"
+#include "Engine.h"
+
+class GLTFEXPORTER_API FGLTFSceneComponentConverter final : public TGLTFConverter<FGLTFJsonNodeIndex, TTuple<const USceneComponent*, bool, bool>>
+{
+	FGLTFJsonNodeIndex Add(FGLTFConvertBuilder& Builder, const FString& Name, TTuple<const USceneComponent*, bool, bool> Params) override;
+};
+
+class GLTFEXPORTER_API FGLTFLevelConverter final : public TGLTFConverter<FGLTFJsonSceneIndex, TTuple<const ULevel*, bool>>
+{
+	FGLTFJsonSceneIndex Add(FGLTFConvertBuilder& Builder, const FString& Name, TTuple<const ULevel*, bool> Params) override;
+};
diff --git a/Engine/Plugins/Enterprise/GLTFExporter/MigratedFromGit/Source/GLTFExporter/Public/Converters/GLTFStaticMeshConverters.h b/Engine/Plugins/Enterprise/GLTFExporter/MigratedFromGit/Source/GLTFExporter/Public/Converters/GLTFStaticMeshConverters.h
new file mode 100644
--- /dev/null
+++ b/Engine/Plugins/Enterprise/GLTFExporter/MigratedFromGit/Source/GLTFExporter/Public/Converters/GLTFStaticMeshConverters.h
@@ -0,0 +1,52 @@
+// Copyright Epic Games, Inc. All Rights Reserved.
+
+#pragma once
+
+#include "Json/GLTFJsonIndex.h"
+#include "Converters/GLTFConverter.h"
+#include "Engine.h"
+
+class GLTFEXPORTER_API FGLTFPositionVertexBufferConverter final : public TGLTFConverter<FGLTFJsonAccessorIndex, const FPositionVertexBuffer*>
+{
+	FGLTFJsonAccessorIndex Add(FGLTFConvertBuilder& Builder, const FString& Name, const FPositionVertexBuffer* VertexBuffer) override;
+};
+
+class GLTFEXPORTER_API FGLTFColorVertexBufferConverter final : public TGLTFConverter<FGLTFJsonAccessorIndex, const FColorVertexBuffer*>
+{
+	FGLTFJsonAccessorIndex Add(FGLTFConvertBuilder& Builder, const FString& Name, const FColorVertexBuffer* VertexBuffer) override;
+};
+
+class GLTFEXPORTER_API FGLTFStaticMeshNormalVertexBufferConverter final : public TGLTFConverter<FGLTFJsonAccessorIndex, const FStaticMeshVertexBuffer*>
+{
+	FGLTFJsonAccessorIndex Add(FGLTFConvertBuilder& Builder, const FString& Name, const FStaticMeshVertexBuffer* VertexBuffer) override;
+};
+
+class GLTFEXPORTER_API FGLTFStaticMeshTangentVertexBufferConverter final : public TGLTFConverter<FGLTFJsonAccessorIndex, const FStaticMeshVertexBuffer*>
+{
+	FGLTFJsonAccessorIndex Add(FGLTFConvertBuilder& Builder, const FString& Name, const FStaticMeshVertexBuffer* VertexBuffer) override;
+};
+
+class GLTFEXPORTER_API FGLTFStaticMeshUV0VertexBufferConverter final : public TGLTFConverter<FGLTFJsonAccessorIndex, const FStaticMeshVertexBuffer*>
+{
+	FGLTFJsonAccessorIndex Add(FGLTFConvertBuilder& Builder, const FString& Name, const FStaticMeshVertexBuffer* VertexBuffer) override;
+};
+
+class GLTFEXPORTER_API FGLTFStaticMeshUV1VertexBufferConverter final : public TGLTFConverter<FGLTFJsonAccessorIndex, const FStaticMeshVertexBuffer*>
+{
+	FGLTFJsonAccessorIndex Add(FGLTFConvertBuilder& Builder, const FString& Name, const FStaticMeshVertexBuffer* VertexBuffer) override;
+};
+
+class GLTFEXPORTER_API FGLTFStaticMeshIndexBufferConverter final : public TGLTFConverter<FGLTFJsonBufferViewIndex, const FRawStaticIndexBuffer*>
+{
+	FGLTFJsonBufferViewIndex Add(FGLTFConvertBuilder& Builder, const FString& Name, const FRawStaticIndexBuffer* IndexBuffer) override;
+};
+
+class GLTFEXPORTER_API FGLTFStaticMeshSectionConverter final : public TGLTFConverter<FGLTFJsonAccessorIndex, TTuple<const FStaticMeshSection*, const FRawStaticIndexBuffer*>>
+{
+	FGLTFJsonAccessorIndex Add(FGLTFConvertBuilder& Builder, const FString& Name, TTuple<const FStaticMeshSection*, const FRawStaticIndexBuffer*> Params) override;
+};
+
+class GLTFEXPORTER_API FGLTFStaticMeshConverter final : public TGLTFConverter<FGLTFJsonMeshIndex, TTuple<const FStaticMeshLODResources*, const FColorVertexBuffer*>>
+{
+	FGLTFJsonMeshIndex Add(FGLTFConvertBuilder& Builder, const FString& Name, TTuple<const FStaticMeshLODResources*, const FColorVertexBuffer*> Params) override;
+};
diff --git a/Engine/Plugins/Enterprise/GLTFExporter/MigratedFromGit/Source/GLTFExporter/Public/Exporters/GLTFAnimSequenceExporter.h b/Engine/Plugins/Enterprise/GLTFExporter/MigratedFromGit/Source/GLTFExporter/Public/Exporters/GLTFAnimSequenceExporter.h
new file mode 100644
--- /dev/null
+++ b/Engine/Plugins/Enterprise/GLTFExporter/MigratedFromGit/Source/GLTFExporter/Public/Exporters/GLTFAnimSequenceExporter.h
@@ -0,0 +1,20 @@
+// Copyright Epic Games, Inc. All Rights Reserved.
+
+#pragma once
+
+#include "Exporters/GLTFExporter.h"
+#include "GLTFAnimSequenceExporter.generated.h"
+
+UCLASS()
+class GLTFEXPORTER_API UGLTFAnimSequenceExporter : public UGLTFExporter
+{
+public:
+	GENERATED_BODY()
+
+public:
+	UGLTFAnimSequenceExporter(const FObjectInitializer& ObjectInitializer = FObjectInitializer());
+
+	//~ Begin UExporter Interface
+	virtual bool ExportBinary(UObject* Object, const TCHAR* Type, FArchive& Archive, FFeedbackContext* Warn, int32 FileIndex = 0, uint32 PortFlags = 0) override;
+	//~ End UExporter Interface
+};
diff --git a/Engine/Plugins/Enterprise/GLTFExporter/MigratedFromGit/Source/GLTFExporter/Public/Exporters/GLTFExporter.h b/Engine/Plugins/Enterprise/GLTFExporter/MigratedFromGit/Source/GLTFExporter/Public/Exporters/GLTFExporter.h
new file mode 100644
--- /dev/null
+++ b/Engine/Plugins/Enterprise/GLTFExporter/MigratedFromGit/Source/GLTFExporter/Public/Exporters/GLTFExporter.h
@@ -0,0 +1,27 @@
+// Copyright Epic Games, Inc. All Rights Reserved.
+
+#pragma once
+
+#include "CoreMinimal.h"
+#include "UObject/ObjectMacros.h"
+#include "UObject/UObjectGlobals.h"
+#include "Exporters/Exporter.h"
+#include "GLTFExportOptions.h"
+#include "GLTFExporter.generated.h"
+
+UCLASS()
+class GLTFEXPORTER_API UGLTFExporter : public UExporter
+{
+public:
+	GENERATED_BODY()
+
+public:
+	UGLTFExporter(const FObjectInitializer& ObjectInitializer = FObjectInitializer());
+
+protected:
+
+	bool FillExportOptions();
+
+	UPROPERTY(transient)
+	UGLTFExportOptions* ExportOptions;
+};
diff --git a/Engine/Plugins/Enterprise/GLTFExporter/MigratedFromGit/Source/GLTFExporter/Public/Exporters/GLTFLevelExporter.h b/Engine/Plugins/Enterprise/GLTFExporter/MigratedFromGit/Source/GLTFExporter/Public/Exporters/GLTFLevelExporter.h
new file mode 100644
--- /dev/null
+++ b/Engine/Plugins/Enterprise/GLTFExporter/MigratedFromGit/Source/GLTFExporter/Public/Exporters/GLTFLevelExporter.h
@@ -0,0 +1,20 @@
+// Copyright Epic Games, Inc. All Rights Reserved.
+
+#pragma once
+
+#include "Exporters/GLTFExporter.h"
+#include "GLTFLevelExporter.generated.h"
+
+UCLASS()
+class GLTFEXPORTER_API UGLTFLevelExporter : public UGLTFExporter
+{
+public:
+	GENERATED_BODY()
+
+public:
+	UGLTFLevelExporter(const FObjectInitializer& ObjectInitializer = FObjectInitializer());
+
+	//~ Begin UExporter Interface
+	virtual bool ExportBinary(UObject* Object, const TCHAR* Type, FArchive& Archive, FFeedbackContext* Warn, int32 FileIndex = 0, uint32 PortFlags = 0) override;
+	//~ End UExporter Interface
+};
diff --git a/Engine/Plugins/Enterprise/GLTFExporter/MigratedFromGit/Source/GLTFExporter/Public/Exporters/GLTFSkeletalMeshExporter.h b/Engine/Plugins/Enterprise/GLTFExporter/MigratedFromGit/Source/GLTFExporter/Public/Exporters/GLTFSkeletalMeshExporter.h
new file mode 100644
--- /dev/null
+++ b/Engine/Plugins/Enterprise/GLTFExporter/MigratedFromGit/Source/GLTFExporter/Public/Exporters/GLTFSkeletalMeshExporter.h
@@ -0,0 +1,20 @@
+// Copyright Epic Games, Inc. All Rights Reserved.
+
+#pragma once
+
+#include "Exporters/GLTFExporter.h"
+#include "GLTFSkeletalMeshExporter.generated.h"
+
+UCLASS()
+class GLTFEXPORTER_API UGLTFSkeletalMeshExporter : public UGLTFExporter
+{
+public:
+	GENERATED_BODY()
+
+public:
+	UGLTFSkeletalMeshExporter(const FObjectInitializer& ObjectInitializer = FObjectInitializer());
+
+	//~ Begin UExporter Interface
+	virtual bool ExportBinary(UObject* Object, const TCHAR* Type, FArchive& Archive, FFeedbackContext* Warn, int32 FileIndex = 0, uint32 PortFlags = 0) override;
+	//~ End UExporter Interface
+};
diff --git a/Engine/Plugins/Enterprise/GLTFExporter/MigratedFromGit/Source/GLTFExporter/Public/Exporters/GLTFStaticMeshExporter.h b/Engine/Plugins/Enterprise/GLTFExporter/MigratedFromGit/Source/GLTFExporter/Public/Exporters/GLTFStaticMeshExporter.h
new file mode 100644
--- /dev/null
+++ b/Engine/Plugins/Enterprise/GLTFExporter/MigratedFromGit/Source/GLTFExporter/Public/Exporters/GLTFStaticMeshExporter.h
@@ -0,0 +1,20 @@
+// Copyright Epic Games, Inc. All Rights Reserved.
+
+#pragma once
+
+#include "Exporters/GLTFExporter.h"
+#include "GLTFStaticMeshExporter.generated.h"
+
+UCLASS()
+class GLTFEXPORTER_API UGLTFStaticMeshExporter : public UGLTFExporter
+{
+public:
+	GENERATED_BODY()
+
+public:
+	UGLTFStaticMeshExporter(const FObjectInitializer& ObjectInitializer = FObjectInitializer());
+
+	//~ Begin UExporter Interface
+	virtual bool ExportBinary(UObject* Object, const TCHAR* Type, FArchive& Archive, FFeedbackContext* Warn, int32 FileIndex = 0, uint32 PortFlags = 0) override;
+	//~ End UExporter Interface
+};
diff --git a/Engine/Plugins/Enterprise/GLTFExporter/MigratedFromGit/Source/GLTFExporter/Public/GLTFAnimSequenceExporter.h b/Engine/Plugins/Enterprise/GLTFExporter/MigratedFromGit/Source/GLTFExporter/Public/GLTFAnimSequenceExporter.h
deleted file mode 100644
diff --git a/Engine/Plugins/Enterprise/GLTFExporter/MigratedFromGit/Source/GLTFExporter/Public/GLTFBufferBuilder.h b/Engine/Plugins/Enterprise/GLTFExporter/MigratedFromGit/Source/GLTFExporter/Public/GLTFBufferBuilder.h
deleted file mode 100644
diff --git a/Engine/Plugins/Enterprise/GLTFExporter/MigratedFromGit/Source/GLTFExporter/Public/GLTFColor.h b/Engine/Plugins/Enterprise/GLTFExporter/MigratedFromGit/Source/GLTFExporter/Public/GLTFColor.h
deleted file mode 100644
diff --git a/Engine/Plugins/Enterprise/GLTFExporter/MigratedFromGit/Source/GLTFExporter/Public/GLTFContainerBuilder.h b/Engine/Plugins/Enterprise/GLTFExporter/MigratedFromGit/Source/GLTFExporter/Public/GLTFContainerBuilder.h
deleted file mode 100644
diff --git a/Engine/Plugins/Enterprise/GLTFExporter/MigratedFromGit/Source/GLTFExporter/Public/GLTFConvertBuilder.h b/Engine/Plugins/Enterprise/GLTFExporter/MigratedFromGit/Source/GLTFExporter/Public/GLTFConvertBuilder.h
deleted file mode 100644
diff --git a/Engine/Plugins/Enterprise/GLTFExporter/MigratedFromGit/Source/GLTFExporter/Public/GLTFConverter.h b/Engine/Plugins/Enterprise/GLTFExporter/MigratedFromGit/Source/GLTFExporter/Public/GLTFConverter.h
deleted file mode 100644
diff --git a/Engine/Plugins/Enterprise/GLTFExporter/MigratedFromGit/Source/GLTFExporter/Public/GLTFConverterUtility.h b/Engine/Plugins/Enterprise/GLTFExporter/MigratedFromGit/Source/GLTFExporter/Public/GLTFConverterUtility.h
deleted file mode 100644
diff --git a/Engine/Plugins/Enterprise/GLTFExporter/MigratedFromGit/Source/GLTFExporter/Public/GLTFExportOptionsWindow.h b/Engine/Plugins/Enterprise/GLTFExporter/MigratedFromGit/Source/GLTFExporter/Public/GLTFExportOptionsWindow.h
deleted file mode 100644
diff --git a/Engine/Plugins/Enterprise/GLTFExporter/MigratedFromGit/Source/GLTFExporter/Public/GLTFExporter.h b/Engine/Plugins/Enterprise/GLTFExporter/MigratedFromGit/Source/GLTFExporter/Public/GLTFExporter.h
deleted file mode 100644
diff --git a/Engine/Plugins/Enterprise/GLTFExporter/MigratedFromGit/Source/GLTFExporter/Public/GLTFJsonAccessor.h b/Engine/Plugins/Enterprise/GLTFExporter/MigratedFromGit/Source/GLTFExporter/Public/GLTFJsonAccessor.h
deleted file mode 100644
diff --git a/Engine/Plugins/Enterprise/GLTFExporter/MigratedFromGit/Source/GLTFExporter/Public/GLTFJsonBuffer.h b/Engine/Plugins/Enterprise/GLTFExporter/MigratedFromGit/Source/GLTFExporter/Public/GLTFJsonBuffer.h
deleted file mode 100644
diff --git a/Engine/Plugins/Enterprise/GLTFExporter/MigratedFromGit/Source/GLTFExporter/Public/GLTFJsonBufferView.h b/Engine/Plugins/Enterprise/GLTFExporter/MigratedFromGit/Source/GLTFExporter/Public/GLTFJsonBufferView.h
deleted file mode 100644
diff --git a/Engine/Plugins/Enterprise/GLTFExporter/MigratedFromGit/Source/GLTFExporter/Public/GLTFJsonBuilder.h b/Engine/Plugins/Enterprise/GLTFExporter/MigratedFromGit/Source/GLTFExporter/Public/GLTFJsonBuilder.h
deleted file mode 100644
diff --git a/Engine/Plugins/Enterprise/GLTFExporter/MigratedFromGit/Source/GLTFExporter/Public/GLTFJsonEnums.h b/Engine/Plugins/Enterprise/GLTFExporter/MigratedFromGit/Source/GLTFExporter/Public/GLTFJsonEnums.h
deleted file mode 100644
diff --git a/Engine/Plugins/Enterprise/GLTFExporter/MigratedFromGit/Source/GLTFExporter/Public/GLTFJsonIndex.h b/Engine/Plugins/Enterprise/GLTFExporter/MigratedFromGit/Source/GLTFExporter/Public/GLTFJsonIndex.h
deleted file mode 100644
diff --git a/Engine/Plugins/Enterprise/GLTFExporter/MigratedFromGit/Source/GLTFExporter/Public/GLTFJsonMesh.h b/Engine/Plugins/Enterprise/GLTFExporter/MigratedFromGit/Source/GLTFExporter/Public/GLTFJsonMesh.h
deleted file mode 100644
diff --git a/Engine/Plugins/Enterprise/GLTFExporter/MigratedFromGit/Source/GLTFExporter/Public/GLTFJsonNode.h b/Engine/Plugins/Enterprise/GLTFExporter/MigratedFromGit/Source/GLTFExporter/Public/GLTFJsonNode.h
deleted file mode 100644
diff --git a/Engine/Plugins/Enterprise/GLTFExporter/MigratedFromGit/Source/GLTFExporter/Public/GLTFJsonRoot.h b/Engine/Plugins/Enterprise/GLTFExporter/MigratedFromGit/Source/GLTFExporter/Public/GLTFJsonRoot.h
deleted file mode 100644
diff --git a/Engine/Plugins/Enterprise/GLTFExporter/MigratedFromGit/Source/GLTFExporter/Public/GLTFJsonScene.h b/Engine/Plugins/Enterprise/GLTFExporter/MigratedFromGit/Source/GLTFExporter/Public/GLTFJsonScene.h
deleted file mode 100644
diff --git a/Engine/Plugins/Enterprise/GLTFExporter/MigratedFromGit/Source/GLTFExporter/Public/GLTFJsonUtility.h b/Engine/Plugins/Enterprise/GLTFExporter/MigratedFromGit/Source/GLTFExporter/Public/GLTFJsonUtility.h
deleted file mode 100644
diff --git a/Engine/Plugins/Enterprise/GLTFExporter/MigratedFromGit/Source/GLTFExporter/Public/GLTFLevelConverters.h b/Engine/Plugins/Enterprise/GLTFExporter/MigratedFromGit/Source/GLTFExporter/Public/GLTFLevelConverters.h
deleted file mode 100644
diff --git a/Engine/Plugins/Enterprise/GLTFExporter/MigratedFromGit/Source/GLTFExporter/Public/GLTFLevelExporter.h b/Engine/Plugins/Enterprise/GLTFExporter/MigratedFromGit/Source/GLTFExporter/Public/GLTFLevelExporter.h
deleted file mode 100644
diff --git a/Engine/Plugins/Enterprise/GLTFExporter/MigratedFromGit/Source/GLTFExporter/Public/GLTFSkeletalMeshExporter.h b/Engine/Plugins/Enterprise/GLTFExporter/MigratedFromGit/Source/GLTFExporter/Public/GLTFSkeletalMeshExporter.h
deleted file mode 100644
diff --git a/Engine/Plugins/Enterprise/GLTFExporter/MigratedFromGit/Source/GLTFExporter/Public/GLTFStaticMeshConverters.h b/Engine/Plugins/Enterprise/GLTFExporter/MigratedFromGit/Source/GLTFExporter/Public/GLTFStaticMeshConverters.h
deleted file mode 100644
diff --git a/Engine/Plugins/Enterprise/GLTFExporter/MigratedFromGit/Source/GLTFExporter/Public/GLTFStaticMeshExporter.h b/Engine/Plugins/Enterprise/GLTFExporter/MigratedFromGit/Source/GLTFExporter/Public/GLTFStaticMeshExporter.h
deleted file mode 100644
diff --git a/Engine/Plugins/Enterprise/GLTFExporter/MigratedFromGit/Source/GLTFExporter/Public/Json/GLTFJsonAccessor.h b/Engine/Plugins/Enterprise/GLTFExporter/MigratedFromGit/Source/GLTFExporter/Public/Json/GLTFJsonAccessor.h
new file mode 100644
--- /dev/null
+++ b/Engine/Plugins/Enterprise/GLTFExporter/MigratedFromGit/Source/GLTFExporter/Public/Json/GLTFJsonAccessor.h
@@ -0,0 +1,83 @@
+// Copyright Epic Games, Inc. All Rights Reserved.
+
+#pragma once
+
+#include "Json/GLTFJsonEnums.h"
+#include "Json/GLTFJsonIndex.h"
+#include "Json/GLTFJsonUtility.h"
+#include "Serialization/JsonSerializer.h"
+
+struct GLTFEXPORTER_API FGLTFJsonAccessor
+{
+	FString Name;
+
+	FGLTFJsonBufferViewIndex BufferView;
+	int64                    ByteOffset;
+	int32                    Count;
+	EGLTFJsonAccessorType    Type;
+	EGLTFJsonComponentType   ComponentType;
+	bool                     bNormalized;
+
+	int32 MinMaxLength;
+	float Min[16];
+	float Max[16];
+
+	FGLTFJsonAccessor()
+		: BufferView(INDEX_NONE)
+		, ByteOffset(0)
+		, Count(0)
+		, Type(EGLTFJsonAccessorType::None)
+		, ComponentType(EGLTFJsonComponentType::None)
+		, bNormalized(false)
+		, MinMaxLength(0)
+		, Min{0}
+		, Max{0}
+	{
+	}
+
+	template <class CharType = TCHAR, class PrintPolicy = TPrettyJsonPrintPolicy<CharType>>
+	void WriteObject(TJsonWriter<CharType, PrintPolicy>& JsonWriter) const
+	{
+		JsonWriter.WriteObjectStart();
+
+		if (!Name.IsEmpty())
+		{
+			JsonWriter.WriteValue(TEXT("name"), Name);
+		}
+
+		JsonWriter.WriteValue(TEXT("bufferView"), BufferView);
+
+		if (ByteOffset != 0)
+		{
+			JsonWriter.WriteValue(TEXT("byteOffset"), ByteOffset);
+		}
+
+		JsonWriter.WriteValue(TEXT("count"), Count);
+		JsonWriter.WriteValue(TEXT("type"), FGLTFJsonUtility::AccessorTypeToString(Type));
+		JsonWriter.WriteValue(TEXT("componentType"), FGLTFJsonUtility::ComponentTypeToNumber(ComponentType));
+
+		if (bNormalized)
+		{
+			JsonWriter.WriteValue(TEXT("normalized"), bNormalized);
+		}
+
+		if (MinMaxLength > 0)
+		{
+			JsonWriter.WriteArrayStart(TEXT("min"));
+			for (int32 ComponentIndex = 0; ComponentIndex < MinMaxLength; ++ComponentIndex)
+			{
+				JsonWriter.WriteValue(Min[ComponentIndex]);
+			}
+			JsonWriter.WriteArrayEnd();
+
+			JsonWriter.WriteArrayStart(TEXT("max"));
+			for (int32 ComponentIndex = 0; ComponentIndex < MinMaxLength; ++ComponentIndex)
+			{
+				JsonWriter.WriteValue(Max[ComponentIndex]);
+			}
+			JsonWriter.WriteArrayEnd();
+		}
+
+		JsonWriter.WriteObjectEnd();
+	}
+};
diff --git a/Engine/Plugins/Enterprise/GLTFExporter/MigratedFromGit/Source/GLTFExporter/Public/Json/GLTFJsonBuffer.h b/Engine/Plugins/Enterprise/GLTFExporter/MigratedFromGit/Source/GLTFExporter/Public/Json/GLTFJsonBuffer.h
new file mode 100644
--- /dev/null
+++ b/Engine/Plugins/Enterprise/GLTFExporter/MigratedFromGit/Source/GLTFExporter/Public/Json/GLTFJsonBuffer.h
@@ -0,0 +1,39 @@
+// Copyright Epic Games, Inc. All Rights Reserved.
+
+#pragma once
+
+#include "Json/GLTFJsonIndex.h"
+#include "Serialization/JsonSerializer.h"
+
+struct GLTFEXPORTER_API FGLTFJsonBuffer
+{
+	FString Name;
+
+	FString URI;
+	int64   ByteLength;
+
+	FGLTFJsonBuffer()
+		: ByteLength(0)
+	{
+	}
+
+	template <class CharType = TCHAR, class PrintPolicy = TPrettyJsonPrintPolicy<CharType>>
+	void WriteObject(TJsonWriter<CharType, PrintPolicy>& JsonWriter) const
+	{
+		JsonWriter.WriteObjectStart();
+
+		if (!Name.IsEmpty())
+		{
+			JsonWriter.WriteValue(TEXT("name"), Name);
+		}
+
+		if (!URI.IsEmpty())
+		{
+			JsonWriter.WriteValue(TEXT("uri"), URI);
+		}
+
+		JsonWriter.WriteValue(TEXT("byteLength"), ByteLength);
+
+		JsonWriter.WriteObjectEnd();
+	}
+};
diff --git a/Engine/Plugins/Enterprise/GLTFExporter/MigratedFromGit/Source/GLTFExporter/Public/Json/GLTFJsonBufferView.h b/Engine/Plugins/Enterprise/GLTFExporter/MigratedFromGit/Source/GLTFExporter/Public/Json/GLTFJsonBufferView.h
new file mode 100644
--- /dev/null
+++ b/Engine/Plugins/Enterprise/GLTFExporter/MigratedFromGit/Source/GLTFExporter/Public/Json/GLTFJsonBufferView.h
@@ -0,0 +1,62 @@
+// Copyright Epic Games, Inc. All Rights Reserved.
+
+#pragma once
+
+#include "Json/GLTFJsonEnums.h"
+#include "Json/GLTFJsonIndex.h"
+#include "Json/GLTFJsonUtility.h"
+#include "Serialization/JsonSerializer.h"
+
+struct GLTFEXPORTER_API FGLTFJsonBufferView
+{
+	FString Name;
+
+	FGLTFJsonBufferIndex Buffer;
+
+	int64 ByteLength;
+	int64 ByteOffset;
+	int32 ByteStride;
+
+	EGLTFJsonBufferTarget Target;
+
+	FGLTFJsonBufferView()
+		: Buffer(INDEX_NONE)
+		, ByteLength(0)
+		, ByteOffset(0)
+		, ByteStride(0)
+		, Target(EGLTFJsonBufferTarget::None)
+	{
+		// check that view fits completely inside the buffer
+	}
+
+	template <class CharType = TCHAR, class PrintPolicy = TPrettyJsonPrintPolicy<CharType>>
+	void WriteObject(TJsonWriter<CharType, PrintPolicy>& JsonWriter) const
+	{
+		JsonWriter.WriteObjectStart();
+
+		if (!Name.IsEmpty())
+		{
+			JsonWriter.WriteValue(TEXT("name"), Name);
+		}
+
+		JsonWriter.WriteValue(TEXT("buffer"), Buffer);
+		JsonWriter.WriteValue(TEXT("byteLength"), ByteLength);
+
+		if (ByteOffset != 0)
+		{
+			JsonWriter.WriteValue(TEXT("byteOffset"), ByteOffset);
+		}
+
+		if (ByteStride != 0)
+		{
+			JsonWriter.WriteValue(TEXT("byteStride"), ByteStride);
+		}
+
+		if (Target != EGLTFJsonBufferTarget::None)
+		{
+			JsonWriter.WriteValue(TEXT("target"), FGLTFJsonUtility::BufferTargetToNumber(Target));
+		}
+
+		JsonWriter.WriteObjectEnd();
+	}
+};
diff --git a/Engine/Plugins/Enterprise/GLTFExporter/MigratedFromGit/Source/GLTFExporter/Public/Json/GLTFJsonEnums.h b/Engine/Plugins/Enterprise/GLTFExporter/MigratedFromGit/Source/GLTFExporter/Public/Json/GLTFJsonEnums.h
new file mode 100644
--- /dev/null
+++ b/Engine/Plugins/Enterprise/GLTFExporter/MigratedFromGit/Source/GLTFExporter/Public/Json/GLTFJsonEnums.h
@@ -0,0 +1,105 @@
+// Copyright Epic Games, Inc. All Rights Reserved.
+
+#pragma once
+
+#include "CoreMinimal.h"
+
+enum class EGLTFExtension
+{
+	KHR_LightsPunctual,
+	KHR_MaterialsUnlit,
+	KHR_MaterialsClearCoat,
+	KHR_MeshQuantization
+};
+
+enum class EGLTFJsonAccessorType
+{
+	None = -1,
+	Scalar,
+	Vec2,
+	Vec3,
+	Vec4,
+	Mat2,
+	Mat3,
+	Mat4
+};
+
+enum class EGLTFJsonComponentType
+{
+	None = -1,
+	S8 = 5120, // signed byte
+	U8 = 5121, // unsigned byte
+	S16 = 5122, // signed short
+	U16 = 5123, // unsigned short
+	S32 = 5124, // unused
+	U32 = 5125, // unsigned int -- only valid for indices, not attributes
+	F32 = 5126  // float
+};
+
+enum class EGLTFJsonBufferTarget
+{
+	None = -1,
+	ArrayBuffer = 34962,
+	ElementArrayBuffer = 34963
+};
+
+enum class EGLTFJsonPrimitiveMode
+{
+	None = -1,
+	// unsupported
+	Points = 0,
+	Lines = 1,
+	LineLoop = 2,
+	LineStrip = 3,
+	// initially supported
+	Triangles = 4,
+	// will be supported prior to release
+	TriangleStrip = 5,
+	TriangleFan = 6
+};
+
+enum class EGLTFJsonAlphaMode
+{
+	None = -1,
+	Opaque,
+	Mask,
+	Blend
+};
+
+enum class EGLTFJsonSamplerFilter
+{
+	None = -1,
+	// valid for Min & Mag
+	Nearest = 9728,
+	Linear = 9729,
+	// valid for Min only
+	NearestMipmapNearest = 9984,
+	LinearMipmapNearest = 9985,
+	NearestMipmapLinear = 9986,
+	LinearMipmapLinear = 9987
+};
+
+enum class EGLTFJsonSamplerWrap
+{
+	None = -1,
+	Repeat = 10497,
+	MirroredRepeat = 33648,
+	ClampToEdge = 33071
+};
+
+enum class EGLTFJsonAnimationInterpolation
+{
+	None = -1,
+	Linear,
+	Step,
+	CubicSpline,
+};
+
+enum class EGLTFJsonAnimationPath
+{
+	None = -1,
+	Translation,
+	Rotation,
+	Scale,
+	Weights
+};
diff --git a/Engine/Plugins/Enterprise/GLTFExporter/MigratedFromGit/Source/GLTFExporter/Public/Json/GLTFJsonIndex.h b/Engine/Plugins/Enterprise/GLTFExporter/MigratedFromGit/Source/GLTFExporter/Public/Json/GLTFJsonIndex.h
new file mode 100644
--- /dev/null
+++ b/Engine/Plugins/Enterprise/GLTFExporter/MigratedFromGit/Source/GLTFExporter/Public/Json/GLTFJsonIndex.h
@@ -0,0 +1,61 @@
+// Copyright Epic Games, Inc. All Rights Reserved.
+
+#pragma once
+
+#include "CoreMinimal.h"
+
+template <typename DerivedType>
+struct FGLTFJsonIndex
+{
+	int32 Value;
+
+	FGLTFJsonIndex()
+		: Value(INDEX_NONE)
+	{
+	}
+
+	explicit FGLTFJsonIndex(int32 Value)
+		: Value(Value)
+	{
+	}
+
+	FORCEINLINE operator int32() const
+	{
+		return Value;
+	}
+
+	FORCEINLINE bool operator==(int32 OtherValue) const
+	{
+		return Value == OtherValue;
+	}
+
+	FORCEINLINE bool operator!=(int32 OtherValue) const
+	{
+		return Value != OtherValue;
+	}
+
+	FORCEINLINE bool operator==(const DerivedType& Other) const
+	{
+		return Value == Other.Value;
+	}
+
+	FORCEINLINE bool operator!=(const DerivedType& Other) const
+	{
+		return Value != Other.Value;
+	}
+
+	FORCEINLINE friend uint32 GetTypeHash(const DerivedType& Other)
+	{
+		return GetTypeHash(Other.Value);
+	}
+};
+
+struct FGLTFJsonAccessorIndex : FGLTFJsonIndex<FGLTFJsonAccessorIndex> { using FGLTFJsonIndex::FGLTFJsonIndex; };
+struct FGLTFJsonBufferIndex : FGLTFJsonIndex<FGLTFJsonBufferIndex> { using FGLTFJsonIndex::FGLTFJsonIndex; };
+struct FGLTFJsonBufferViewIndex : FGLTFJsonIndex<FGLTFJsonBufferViewIndex> { using FGLTFJsonIndex::FGLTFJsonIndex; };
+struct FGLTFJsonCameraIndex : FGLTFJsonIndex<FGLTFJsonCameraIndex> { using FGLTFJsonIndex::FGLTFJsonIndex; };
+struct FGLTFJsonMaterialIndex : FGLTFJsonIndex<FGLTFJsonMaterialIndex> { using FGLTFJsonIndex::FGLTFJsonIndex; };
+struct FGLTFJsonMeshIndex : FGLTFJsonIndex<FGLTFJsonMeshIndex> { using FGLTFJsonIndex::FGLTFJsonIndex; };
+struct FGLTFJsonNodeIndex : FGLTFJsonIndex<FGLTFJsonNodeIndex> { using FGLTFJsonIndex::FGLTFJsonIndex; };
+struct FGLTFJsonSceneIndex : FGLTFJsonIndex<FGLTFJsonSceneIndex> { using FGLTFJsonIndex::FGLTFJsonIndex; };
+struct FGLTFJsonSkinIndex : FGLTFJsonIndex<FGLTFJsonSkinIndex> { using FGLTFJsonIndex::FGLTFJsonIndex; };
diff --git a/Engine/Plugins/Enterprise/GLTFExporter/MigratedFromGit/Source/GLTFExporter/Public/Json/GLTFJsonMesh.h b/Engine/Plugins/Enterprise/GLTFExporter/MigratedFromGit/Source/GLTFExporter/Public/Json/GLTFJsonMesh.h
new file mode 100644
--- /dev/null
+++ b/Engine/Plugins/Enterprise/GLTFExporter/MigratedFromGit/Source/GLTFExporter/Public/Json/GLTFJsonMesh.h
@@ -0,0 +1,117 @@
+// Copyright Epic Games, Inc. All Rights Reserved.
+
+#pragma once
+
+#include "Json/GLTFJsonEnums.h"
+#include "Json/GLTFJsonIndex.h"
+#include "Serialization/JsonSerializer.h"
+
+struct GLTFEXPORTER_API FGLTFJsonAttributes
+{
+	FGLTFJsonAccessorIndex Position; // always required
+	FGLTFJsonAccessorIndex Normal;
+	FGLTFJsonAccessorIndex Tangent;
+	FGLTFJsonAccessorIndex TexCoord0;
+	FGLTFJsonAccessorIndex TexCoord1;
+	FGLTFJsonAccessorIndex Color0;
+	// skeletal mesh attributes
+	FGLTFJsonAccessorIndex Joints0;
+	FGLTFJsonAccessorIndex Weights0;
+
+	FGLTFJsonAttributes()
+		: Position(INDEX_NONE)
+		, Normal(INDEX_NONE)
+		, Tangent(INDEX_NONE)
+		, TexCoord0(INDEX_NONE)
+		, TexCoord1(INDEX_NONE)
+		, Color0(INDEX_NONE)
+		, Joints0(INDEX_NONE)
+		, Weights0(INDEX_NONE)
+	{
+	}
+
+	template <class CharType = TCHAR, class PrintPolicy = TPrettyJsonPrintPolicy<CharType>>
+	void WriteObject(TJsonWriter<CharType, PrintPolicy>& JsonWriter) const
+	{
+		JsonWriter.WriteObjectStart();
+
+		JsonWriter.WriteValue(TEXT("POSITION"), Position);
+		if (Normal != INDEX_NONE) JsonWriter.WriteValue(TEXT("NORMAL"), Normal);
+		if (Tangent != INDEX_NONE) JsonWriter.WriteValue(TEXT("TANGENT"), Tangent);
+		if (TexCoord0 != INDEX_NONE) JsonWriter.WriteValue(TEXT("TEXCOORD_0"), TexCoord0);
+		if (TexCoord1 != INDEX_NONE) JsonWriter.WriteValue(TEXT("TEXCOORD_1"), TexCoord1);
+		if (Color0 != INDEX_NONE) JsonWriter.WriteValue(TEXT("COLOR_0"), Color0);
+		if (Joints0 != INDEX_NONE) JsonWriter.WriteValue(TEXT("JOINTS_0"), Joints0);
+		if (Weights0 != INDEX_NONE) JsonWriter.WriteValue(TEXT("WEIGHTS_0"), Weights0);
+
+		JsonWriter.WriteObjectEnd();
+	}
+};
+
+struct GLTFEXPORTER_API FGLTFJsonPrimitive
+{
+	FGLTFJsonAccessorIndex Indices;
+	FGLTFJsonMaterialIndex Material;
+	EGLTFJsonPrimitiveMode Mode;
+	FGLTFJsonAttributes    Attributes;
+
+	FGLTFJsonPrimitive()
+		: Indices(INDEX_NONE)
+		, Material(INDEX_NONE)
+		, Mode(EGLTFJsonPrimitiveMode::None)
+	{
+	}
+
+	template <class CharType = TCHAR, class PrintPolicy = TPrettyJsonPrintPolicy<CharType>>
+	void WriteObject(TJsonWriter<CharType, PrintPolicy>& JsonWriter) const
+	{
+		JsonWriter.WriteObjectStart();
+
+		JsonWriter.WriteIdentifierPrefix(TEXT("attributes"));
+		Attributes.WriteObject(JsonWriter);
+
+		if (Indices != INDEX_NONE)
+		{
+			JsonWriter.WriteValue(TEXT("indices"), Indices);
+		}
+
+		if (Material != INDEX_NONE)
+		{
+			JsonWriter.WriteValue(TEXT("material"), Material);
+		}
+
+		if (Mode != EGLTFJsonPrimitiveMode::None)
+		{
+			JsonWriter.WriteValue(TEXT("mode"), FGLTFJsonUtility::PrimitiveModeToNumber(Mode));
+		}
+
+		JsonWriter.WriteObjectEnd();
+	}
+};
+
+struct GLTFEXPORTER_API FGLTFJsonMesh
+{
+	FString Name;
+
+	TArray<FGLTFJsonPrimitive> Primitives;
+
+	template <class CharType = TCHAR, class PrintPolicy = TPrettyJsonPrintPolicy<CharType>>
+	void WriteObject(TJsonWriter<CharType, PrintPolicy>& JsonWriter) const
+	{
+		JsonWriter.WriteObjectStart();
+
+		if (!Name.IsEmpty())
+		{
+			JsonWriter.WriteValue(TEXT("name"), Name);
+		}
+
+		JsonWriter.WriteArrayStart(TEXT("primitives"));
+		for (const FGLTFJsonPrimitive& Primitive : Primitives)
+		{
+			Primitive.WriteObject(JsonWriter);
+		}
+		JsonWriter.WriteArrayEnd();
+
+		JsonWriter.WriteObjectEnd();
+	}
+};
diff --git a/Engine/Plugins/Enterprise/GLTFExporter/MigratedFromGit/Source/GLTFExporter/Public/Json/GLTFJsonNode.h b/Engine/Plugins/Enterprise/GLTFExporter/MigratedFromGit/Source/GLTFExporter/Public/Json/GLTFJsonNode.h
new file mode 100644
--- /dev/null
+++ b/Engine/Plugins/Enterprise/GLTFExporter/MigratedFromGit/Source/GLTFExporter/Public/Json/GLTFJsonNode.h
@@ -0,0 +1,92 @@
+// Copyright Epic Games, Inc. All Rights Reserved.
+
+#pragma once
+
+#include "Json/GLTFJsonIndex.h"
+#include "Serialization/JsonSerializer.h"
+
+struct GLTFEXPORTER_API FGLTFJsonNode
+{
+	FString Name;
+
+	FVector Translation;
+	FQuat   Rotation;
+	FVector Scale;
+
+	FGLTFJsonCameraIndex Camera;
+	FGLTFJsonSkinIndex   Skin;
+	FGLTFJsonMeshIndex   Mesh;
+
+	TArray<FGLTFJsonNodeIndex> Children;
+
+	FGLTFJsonNode()
+		: Translation(FVector::ZeroVector)
+		, Rotation(FQuat::Identity)
+		, Scale(FVector::OneVector)
+		, Camera(INDEX_NONE)
+		, Skin(INDEX_NONE)
+		, Mesh(INDEX_NONE)
+	{
+	}
+
+	template <class CharType = TCHAR, class PrintPolicy = TPrettyJsonPrintPolicy<CharType>>
+	void WriteObject(TJsonWriter<CharType, PrintPolicy>& JsonWriter) const
+	{
+		JsonWriter.WriteObjectStart();
+
+		if (!Name.IsEmpty())
+		{
+			JsonWriter.WriteValue(TEXT("name"), Name);
+		}
+
+		if (Translation != FVector::ZeroVector)
+		{
+			JsonWriter.WriteArrayStart(TEXT("translation"));
+			JsonWriter.WriteValue(Translation.X);
+			JsonWriter.WriteValue(Translation.Y);
+			JsonWriter.WriteValue(Translation.Z);
+			JsonWriter.WriteArrayEnd();
+		}
+
+		if (Rotation != FQuat::Identity)
+		{
+			JsonWriter.WriteArrayStart(TEXT("rotation"));
+			JsonWriter.WriteValue(Rotation.X);
+			JsonWriter.WriteValue(Rotation.Y);
+			JsonWriter.WriteValue(Rotation.Z);
+			JsonWriter.WriteValue(Rotation.W);
+			JsonWriter.WriteArrayEnd();
+		}
+
+		if (Scale != FVector::OneVector)
+		{
+			JsonWriter.WriteArrayStart(TEXT("scale"));
+			JsonWriter.WriteValue(Scale.X);
+			JsonWriter.WriteValue(Scale.Y);
+			JsonWriter.WriteValue(Scale.Z);
+			JsonWriter.WriteArrayEnd();
+		}
+
+		if (Camera != INDEX_NONE)
+		{
+			JsonWriter.WriteValue(TEXT("camera"), Camera);
+		}
+
+		if (Skin != INDEX_NONE)
+		{
+			JsonWriter.WriteValue(TEXT("skin"), Skin);
+		}
+
+		if (Mesh != INDEX_NONE)
+		{
+			JsonWriter.WriteValue(TEXT("mesh"), Mesh);
+		}
+
+		if (Children.Num() > 0)
+		{
+			JsonWriter.WriteValue(TEXT("children"), Children);
+		}
+
+		JsonWriter.WriteObjectEnd();
+	}
+};
diff --git a/Engine/Plugins/Enterprise/GLTFExporter/MigratedFromGit/Source/GLTFExporter/Public/Json/GLTFJsonRoot.h b/Engine/Plugins/Enterprise/GLTFExporter/MigratedFromGit/Source/GLTFExporter/Public/Json/GLTFJsonRoot.h
new file mode 100644
--- /dev/null
+++ b/Engine/Plugins/Enterprise/GLTFExporter/MigratedFromGit/Source/GLTFExporter/Public/Json/GLTFJsonRoot.h
@@ -0,0 +1,189 @@
+// Copyright Epic Games, Inc. All Rights Reserved.
+
+#pragma once
+
+#include "Json/GLTFJsonIndex.h"
+#include "Json/GLTFJsonAccessor.h"
+#include "Json/GLTFJsonBufferView.h"
+#include "Json/GLTFJsonBuffer.h"
+#include "Json/GLTFJsonMesh.h"
+#include "Json/GLTFJsonScene.h"
+#include "Json/GLTFJsonNode.h"
+
+#include "Containers/Set.h"
+#include "Policies/CondensedJsonPrintPolicy.h"
+#include "Runtime/Launch/Resources/Version.h"
+
+
+struct GLTFEXPORTER_API FGLTFJsonAsset
+{
+	FString Version;
+	FString Generator;
+	FString Copyright;
+
+	FGLTFJsonAsset()
+		: Version(TEXT("2.0"))
+		, Generator(TEXT(EPIC_PRODUCT_NAME) TEXT(" ") ENGINE_VERSION_STRING)
+	{
+	}
+
+	template <class CharType = TCHAR, class PrintPolicy = TPrettyJsonPrintPolicy<CharType>>
+	void WriteObject(TJsonWriter<CharType, PrintPolicy>& JsonWriter) const
+	{
+		JsonWriter.WriteObjectStart();
+
+		JsonWriter.WriteValue(TEXT("version"), Version);
+
+		if (!Generator.IsEmpty())
+		{
+			JsonWriter.WriteValue(TEXT("generator"), Generator);
+		}
+
+		if (!Copyright.IsEmpty())
+		{
+			JsonWriter.WriteValue(TEXT("copyright"), Copyright);
+		}
+
+		JsonWriter.WriteObjectEnd();
+	}
+};
+
+struct GLTFEXPORTER_API FGLTFJsonRoot
+{
+	FGLTFJsonAsset Asset;
+
+	TSet<EGLTFExtension> ExtensionsUsed;
+	TSet<EGLTFExtension> ExtensionsRequired;
+
+	FGLTFJsonSceneIndex DefaultScene;
+
+	TArray<FGLTFJsonAccessor>   Accessors;
+	TArray<FGLTFJsonBuffer>     Buffers;
+	TArray<FGLTFJsonBufferView> BufferViews;
+	TArray<FGLTFJsonMesh>       Meshes;
+	TArray<FGLTFJsonNode>       Nodes;
+	TArray<FGLTFJsonScene>      Scenes;
+
+	FGLTFJsonRoot()
+		: DefaultScene(INDEX_NONE)
+	{
+	}
+
+	template <class CharType = TCHAR, class PrintPolicy = TPrettyJsonPrintPolicy<CharType>>
+	void WriteObject(TJsonWriter<CharType, PrintPolicy>& JsonWriter) const
+	{
+		JsonWriter.WriteObjectStart();
+
+		JsonWriter.WriteIdentifierPrefix(TEXT("asset"));
+		Asset.WriteObject(JsonWriter);
+
+		if (ExtensionsUsed.Num() > 0)
+		{
+			JsonWriter.WriteIdentifierPrefix(TEXT("extensionsUsed"));
+			JsonWriter.WriteArrayStart();
+			for (const EGLTFExtension& Extension : ExtensionsUsed)
+			{
+				JsonWriter.WriteValue(FGLTFJsonUtility::ExtensionToString(Extension));
+			}
+			JsonWriter.WriteArrayEnd();
+		}
+
+		if (ExtensionsRequired.Num() > 0)
+		{
+			JsonWriter.WriteIdentifierPrefix(TEXT("extensionsRequired"));
+			JsonWriter.WriteArrayStart();
+			for (const EGLTFExtension& Extension : ExtensionsRequired)
+			{
+				JsonWriter.WriteValue(FGLTFJsonUtility::ExtensionToString(Extension));
+			}
+			JsonWriter.WriteArrayEnd();
+		}
+
+		if (DefaultScene != INDEX_NONE)
+		{
+			JsonWriter.WriteValue(TEXT("scene"), DefaultScene);
+		}
+
+		if (Accessors.Num() > 0)
+		{
+			JsonWriter.WriteArrayStart(TEXT("accessors"));
+			for (const FGLTFJsonAccessor& Accessor : Accessors)
+			{
+				Accessor.WriteObject(JsonWriter);
+			}
+			JsonWriter.WriteArrayEnd();
+		}
+
+		if (BufferViews.Num() > 0)
+		{
+			JsonWriter.WriteArrayStart(TEXT("bufferViews"));
+			for (const FGLTFJsonBufferView& BufferView : BufferViews)
+			{
+				BufferView.WriteObject(JsonWriter);
+			}
+			JsonWriter.WriteArrayEnd();
+		}
+
+		if (Meshes.Num() > 0)
+		{
+			JsonWriter.WriteArrayStart(TEXT("meshes"));
+			for (const FGLTFJsonMesh& Mesh : Meshes)
+			{
+				Mesh.WriteObject(JsonWriter);
+			}
+			JsonWriter.WriteArrayEnd();
+		}
+
+		if (Nodes.Num() > 0)
+		{
+			JsonWriter.WriteArrayStart(TEXT("nodes"));
+			for (const FGLTFJsonNode& Node : Nodes)
+			{
+				Node.WriteObject(JsonWriter);
+			}
+			JsonWriter.WriteArrayEnd();
+		}
+
+		if (Scenes.Num() > 0)
+		{
+			JsonWriter.WriteArrayStart(TEXT("scenes"));
+			for (const FGLTFJsonScene& Scene : Scenes)
+			{
+				Scene.WriteObject(JsonWriter);
+			}
+			JsonWriter.WriteArrayEnd();
+		}
+
+		if (Buffers.Num() > 0)
+		{
+			JsonWriter.WriteArrayStart(TEXT("buffers"));
+			for (const FGLTFJsonBuffer& Buffer : Buffers)
+			{
+				Buffer.WriteObject(JsonWriter);
+			}
+			JsonWriter.WriteArrayEnd();
+		}
+
+		JsonWriter.WriteObjectEnd();
+	}
+
+	template <class CharType = TCHAR, class PrintPolicy = TPrettyJsonPrintPolicy<CharType>>
+	void Serialize(FArchive* const Archive) const
+	{
+		TSharedRef<TJsonWriter<CharType, PrintPolicy>> JsonWriter = TJsonWriterFactory<CharType, PrintPolicy>::Create(Archive);
+		WriteObject(*JsonWriter);
+		JsonWriter->Close();
+	}
+
+	void Serialize(FArchive* const Archive, bool bPrettyPrint = true) const
+	{
+		if (bPrettyPrint)
+		{
+			Serialize<UTF8CHAR, TPrettyJsonPrintPolicy<UTF8CHAR>>(Archive);
+		}
+		else
+		{
+			Serialize<UTF8CHAR, TCondensedJsonPrintPolicy<UTF8CHAR>>(Archive);
+		}
+	}
+};
diff --git a/Engine/Plugins/Enterprise/GLTFExporter/MigratedFromGit/Source/GLTFExporter/Public/Json/GLTFJsonScene.h b/Engine/Plugins/Enterprise/GLTFExporter/MigratedFromGit/Source/GLTFExporter/Public/Json/GLTFJsonScene.h
new file mode 100644
--- /dev/null
+++ b/Engine/Plugins/Enterprise/GLTFExporter/MigratedFromGit/Source/GLTFExporter/Public/Json/GLTFJsonScene.h
@@ -0,0 +1,31 @@
+// Copyright Epic Games, Inc. All Rights Reserved.
+
+#pragma once
+
+#include "Json/GLTFJsonIndex.h"
+#include "Serialization/JsonSerializer.h"
+
+struct GLTFEXPORTER_API FGLTFJsonScene
+{
+	FString Name;
+
+	TArray<FGLTFJsonNodeIndex> Nodes;
+
+	template <class CharType = TCHAR, class PrintPolicy = TPrettyJsonPrintPolicy<CharType>>
+	void WriteObject(TJsonWriter<CharType, PrintPolicy>& JsonWriter) const
+	{
+		JsonWriter.WriteObjectStart();
+
+		if (!Name.IsEmpty())
+		{
+			JsonWriter.WriteValue(TEXT("name"), Name);
+		}
+
+		if (Nodes.Num() > 0)
+		{
+			JsonWriter.WriteValue(TEXT("nodes"), Nodes);
+		}
+
+		JsonWriter.WriteObjectEnd();
+	}
+};
diff --git a/Engine/Plugins/Enterprise/GLTFExporter/MigratedFromGit/Source/GLTFExporter/Public/Json/GLTFJsonUtility.h b/Engine/Plugins/Enterprise/GLTFExporter/MigratedFromGit/Source/GLTFExporter/Public/Json/GLTFJsonUtility.h
new file mode 100644
--- /dev/null
+++ b/Engine/Plugins/Enterprise/GLTFExporter/MigratedFromGit/Source/GLTFExporter/Public/Json/GLTFJsonUtility.h
@@ -0,0 +1,71 @@
+// Copyright Epic Games, Inc. All Rights Reserved.
+
+#pragma once
+
+#include "Json/GLTFJsonEnums.h"
+
+struct GLTFEXPORTER_API FGLTFJsonUtility
+{
+	static const TCHAR* ExtensionToString(EGLTFExtension Extension)
+	{
+		switch (Extension)
+		{
+			case EGLTFExtension::KHR_LightsPunctual:     return TEXT("KHR_lights_punctual");
+			case EGLTFExtension::KHR_MaterialsUnlit:     return TEXT("KHR_materials_unlit");
+			case EGLTFExtension::KHR_MaterialsClearCoat: return TEXT("KHR_materials_clearcoat");
+			case EGLTFExtension::KHR_MeshQuantization:   return TEXT("KHR_mesh_quantization");
+			default:                                     return TEXT("unknown");
+		}
+	}
+
+	static const TCHAR* AlphaModeToString(EGLTFJsonAlphaMode Mode)
+	{
+		switch (Mode)
+		{
+			case EGLTFJsonAlphaMode::Opaque: return TEXT("OPAQUE");
+			case EGLTFJsonAlphaMode::Blend:  return TEXT("BLEND");
+			case EGLTFJsonAlphaMode::Mask:   return TEXT("MASK");
+			default:                         return TEXT("UNKNOWN");
+		}
+	}
+
+	static int32 PrimitiveModeToNumber(EGLTFJsonPrimitiveMode Mode)
+	{
+		return static_cast<int32>(Mode);
+	}
+
+	static int32 FilterToNumber(EGLTFJsonSamplerFilter Filter)
+	{
+		return static_cast<int32>(Filter);
+	}
+
+	static int32 WrapModeToNumber(EGLTFJsonSamplerWrap Wrap)
+	{
+		return static_cast<int32>(Wrap);
+	}
+
+	static const TCHAR* AccessorTypeToString(EGLTFJsonAccessorType Type)
+	{
+		switch (Type)
+		{
+			case EGLTFJsonAccessorType::Scalar: return TEXT("SCALAR");
+			case EGLTFJsonAccessorType::Vec2:   return TEXT("VEC2");
+			case EGLTFJsonAccessorType::Vec3:   return TEXT("VEC3");
+			case EGLTFJsonAccessorType::Vec4:   return TEXT("VEC4");
+			case EGLTFJsonAccessorType::Mat2:   return TEXT("MAT2");
+			case EGLTFJsonAccessorType::Mat3:   return TEXT("MAT3");
+			case EGLTFJsonAccessorType::Mat4:   return TEXT("MAT4");
+			default:                            return TEXT("UNKNOWN");
+		}
+	}
+
+	static int32 ComponentTypeToNumber(EGLTFJsonComponentType Type)
+	{
+		return static_cast<int32>(Type);
+	}
+
+	static int32 BufferTargetToNumber(EGLTFJsonBufferTarget Target)
+	{
+		return static_cast<int32>(Target);
+	}
+};
diff --git a/Engine/Plugins/Enterprise/GLTFExporter/MigratedFromGit/Source/GLTFExporter/Public/UI/GLTFExportOptionsWindow.h b/Engine/Plugins/Enterprise/GLTFExporter/MigratedFromGit/Source/GLTFExporter/Public/UI/GLTFExportOptionsWindow.h
new file mode 100644
--- /dev/null
+++ b/Engine/Plugins/Enterprise/GLTFExporter/MigratedFromGit/Source/GLTFExporter/Public/UI/GLTFExportOptionsWindow.h
@@ -0,0 +1,60 @@
+// Copyright 1998-2019 Epic Games, Inc. All Rights Reserved.
+
+#pragma once
+
+#include "GLTFExportOptions.h"
+#include "InputCoreTypes.h"
+#include "Widgets/DeclarativeSyntaxSupport.h"
+#include "Input/Reply.h"
+#include "Widgets/SCompoundWidget.h"
+#include "Widgets/SWindow.h"
+
+class SButton;
+
+class GLTFEXPORTER_API SGLTFExportOptionsWindow : public SCompoundWidget
+{
+public:
+	SLATE_BEGIN_ARGS(SGLTFExportOptionsWindow)
+		: _ExportOptions(nullptr)
+		, _BatchMode()
+		{}
+
+		SLATE_ARGUMENT( UGLTFExportOptions*, ExportOptions )
+		SLATE_ARGUMENT( TSharedPtr<SWindow>, WidgetWindow )
+		SLATE_ARGUMENT( FText, FullPath )
+		SLATE_ARGUMENT( bool, BatchMode )
+	SLATE_END_ARGS()
+
+public:
+	void Construct(const FArguments& InArgs);
+
+	virtual bool SupportsKeyboardFocus() const override { return true; }
+
+	FReply OnExport();
+
+	FReply OnExportAll();
+
+	FReply OnCancel();
+
+	virtual FReply OnKeyDown(const FGeometry& MyGeometry, const FKeyEvent& InKeyEvent) override;
+
+	bool ShouldExport() const;
+
+	bool ShouldExportAll() const;
+
+	SGLTFExportOptionsWindow();
+
+public:
+	static void ShowDialog(UGLTFExportOptions* ExportOptions, const FString& FullPath, bool bBatchMode, bool& bOutOperationCanceled, bool& bOutExportAll);
+
+private:
+	FReply OnResetToDefaultClick() const;
+
+private:
+	UGLTFExportOptions* ExportOptions;
+	TSharedPtr<class IDetailsView> DetailsView;
+	TWeakPtr<SWindow> WidgetWindow;
+	TSharedPtr<SButton> ImportButton;
+	bool bShouldExport;
+	bool bShouldExportAll;
+};
