[GLTFExporter] refactor: move internal types into Source/Private and remove GLTFEXPORTER_API

commit ca0781f420fa7446860343cacfda576042f210d7
Date: Sun, 17 May 2020 18:16:25 +0200

    refactor: move internal types into Source/Private and remove GLTFEXPORTER_API
    (cherry picked from commit 8efc253b551ac0c6ab5abbb4c2ea02c512071544​)

#skipci

diff --git a/Engine/Plugins/Enterprise/GLTFExporter/MigratedFromGit/Source/GLTFExporter/Private/Builders/GLTFBufferBuilder.h b/Engine/Plugins/Enterprise/GLTFExporter/MigratedFromGit/Source/GLTFExporter/Private/Builders/GLTFBufferBuilder.h
new file mode 100644
--- /dev/null
+++ b/Engine/Plugins/Enterprise/GLTFExporter/MigratedFromGit/Source/GLTFExporter/Private/Builders/GLTFBufferBuilder.h
@@ -0,0 +1,24 @@
+// Copyright Epic Games, Inc. All Rights Reserved.
+
+#pragma once
+
+#include "Builders/GLTFJsonBuilder.h"
+
+struct FGLTFBufferBuilder : public FGLTFJsonBuilder
+{
+	FGLTFJsonBufferIndex BufferIndex;
+	TArray<uint8> BufferData;
+
+	FGLTFBufferBuilder();
+
+	FGLTFJsonBufferViewIndex AddBufferView(const void* RawData, uint64 ByteLength, const FString& Name = TEXT(""), EGLTFJsonBufferTarget BufferTarget = EGLTFJsonBufferTarget::ArrayBuffer);
+
+	template <class ElementType>
+	FGLTFJsonBufferViewIndex AddBufferView(const TArray<ElementType>& Array, const FString& Name = TEXT(""), EGLTFJsonBufferTarget BufferTarget = EGLTFJsonBufferTarget::ArrayBuffer)
+	{
+		return AddBufferView(Array.GetData(), Array.Num() * sizeof(ElementType), Name, BufferTarget);
+	}
+
+	void UpdateJsonBufferObject(const FString& BinaryFilePath);
+	virtual bool Serialize(FArchive& Archive, const FString& FilePath) override;
+};
diff --git a/Engine/Plugins/Enterprise/GLTFExporter/MigratedFromGit/Source/GLTFExporter/Private/Builders/GLTFContainerBuilder.h b/Engine/Plugins/Enterprise/GLTFExporter/MigratedFromGit/Source/GLTFExporter/Private/Builders/GLTFContainerBuilder.h
new file mode 100644
--- /dev/null
+++ b/Engine/Plugins/Enterprise/GLTFExporter/MigratedFromGit/Source/GLTFExporter/Private/Builders/GLTFContainerBuilder.h
@@ -0,0 +1,9 @@
+// Copyright Epic Games, Inc. All Rights Reserved.
+
+#pragma once
+
+#include "Builders/GLTFConvertBuilder.h"
+
+struct FGLTFContainerBuilder : public FGLTFConvertBuilder
+{
+};
diff --git a/Engine/Plugins/Enterprise/GLTFExporter/MigratedFromGit/Source/GLTFExporter/Private/Builders/GLTFConvertBuilder.h b/Engine/Plugins/Enterprise/GLTFExporter/MigratedFromGit/Source/GLTFExporter/Private/Builders/GLTFConvertBuilder.h
new file mode 100644
--- /dev/null
+++ b/Engine/Plugins/Enterprise/GLTFExporter/MigratedFromGit/Source/GLTFExporter/Private/Builders/GLTFConvertBuilder.h
@@ -0,0 +1,41 @@
+// Copyright Epic Games, Inc. All Rights Reserved.
+
+#pragma once
+
+#include "Builders/GLTFBufferBuilder.h"
+#include "Converters/GLTFStaticMeshConverters.h"
+#include "Converters/GLTFLevelConverters.h"
+
+struct FGLTFConvertBuilder : public FGLTFBufferBuilder
+{
+	FGLTFPositionVertexBufferConverter PositionVertexBufferConverter;
+	FGLTFColorVertexBufferConverter ColorVertexBufferConverter;
+	FGLTFStaticMeshNormalVertexBufferConverter StaticMeshNormalVertexBufferConverter;
+	FGLTFStaticMeshTangentVertexBufferConverter StaticMeshTangentVertexBufferConverter;
+	FGLTFStaticMeshUV0VertexBufferConverter StaticMeshUV0VertexBufferConverter;
+	FGLTFStaticMeshUV1VertexBufferConverter StaticMeshUV1VertexBufferConverter;
+	FGLTFStaticMeshIndexBufferConverter StaticMeshIndexBufferConverter;
+	FGLTFStaticMeshSectionConverter StaticMeshSectionConverter;
+	FGLTFStaticMeshConverter StaticMeshConverter;
+
+	FGLTFSceneComponentConverter SceneComponentConverter;
+	FGLTFLevelConverter LevelConverter;
+
+	FGLTFJsonAccessorIndex GetOrAddPositionAccessor(const FPositionVertexBuffer* VertexBuffer, const FString& DesiredName = TEXT(""));
+	FGLTFJsonAccessorIndex GetOrAddNormalAccessor(const FStaticMeshVertexBuffer* VertexBuffer, const FString& DesiredName = TEXT(""));
+	FGLTFJsonAccessorIndex GetOrAddTangentAccessor(const FStaticMeshVertexBuffer* VertexBuffer, const FString& DesiredName = TEXT(""));
+	FGLTFJsonAccessorIndex GetOrAddUV0Accessor(const FStaticMeshVertexBuffer* VertexBuffer, const FString& DesiredName = TEXT(""));
+	FGLTFJsonAccessorIndex GetOrAddUV1Accessor(const FStaticMeshVertexBuffer* VertexBuffer, const FString& DesiredName = TEXT(""));
+	FGLTFJsonAccessorIndex GetOrAddColorAccessor(const FColorVertexBuffer* VertexBuffer, const FString& DesiredName = TEXT(""));
+
+	FGLTFJsonBufferViewIndex GetOrAddIndexBufferView(const FRawStaticIndexBuffer* IndexBuffer, const FString& DesiredName = TEXT(""));
+	FGLTFJsonAccessorIndex GetOrAddIndexAccessor(const FStaticMeshSection* MeshSection, const FRawStaticIndexBuffer* IndexBuffer, const FString& DesiredName = TEXT(""));
+
+	FGLTFJsonMeshIndex GetOrAddMesh(const FStaticMeshLODResources* StaticMeshLOD, const FColorVertexBuffer* OverrideVertexColors = nullptr, const FString& DesiredName = TEXT(""));
+	FGLTFJsonMeshIndex GetOrAddMesh(const UStaticMesh* StaticMesh, int32 LODIndex = 0, const FColorVertexBuffer* OverrideVertexColors = nullptr, const FString& DesiredName = TEXT(""));
+	FGLTFJsonMeshIndex GetOrAddMesh(const UStaticMeshComponent* StaticMeshComponent, const FString& DesiredName = TEXT(""));
+
+	FGLTFJsonNodeIndex GetOrAddNode(const USceneComponent* SceneComponent, bool bSelectedOnly, bool bRootNode = false, const FString& DesiredName = TEXT(""));
+	FGLTFJsonSceneIndex GetOrAddScene(const ULevel* Level, bool bSelectedOnly, const FString& DesiredName = TEXT(""));
+	FGLTFJsonSceneIndex GetOrAddScene(const UWorld* World, bool bSelectedOnly, const FString& DesiredName = TEXT(""));
+};
diff --git a/Engine/Plugins/Enterprise/GLTFExporter/MigratedFromGit/Source/GLTFExporter/Private/Builders/GLTFJsonBuilder.h b/Engine/Plugins/Enterprise/GLTFExporter/MigratedFromGit/Source/GLTFExporter/Private/Builders/GLTFJsonBuilder.h
new file mode 100644
--- /dev/null
+++ b/Engine/Plugins/Enterprise/GLTFExporter/MigratedFromGit/Source/GLTFExporter/Private/Builders/GLTFJsonBuilder.h
@@ -0,0 +1,22 @@
+// Copyright Epic Games, Inc. All Rights Reserved.
+
+#pragma once
+
+#include "Json/GLTFJsonRoot.h"
+
+struct FGLTFJsonBuilder
+{
+	FGLTFJsonRoot JsonRoot;
+
+	FGLTFJsonBuilder();
+	virtual ~FGLTFJsonBuilder();
+
+	FGLTFJsonAccessorIndex AddAccessor(const FGLTFJsonAccessor& JsonAccessor);
+	FGLTFJsonBufferIndex AddBuffer(const FGLTFJsonBuffer& JsonBuffer);
+	FGLTFJsonBufferViewIndex AddBufferView(const FGLTFJsonBufferView& JsonBufferView);
+	FGLTFJsonMeshIndex AddMesh(const FGLTFJsonMesh& JsonMesh);
+	FGLTFJsonNodeIndex AddNode(const FGLTFJsonNode& JsonNode);
+	FGLTFJsonSceneIndex AddScene(const FGLTFJsonScene& JsonScene);
+
+	virtual bool Serialize(FArchive& Archive, const FString& FilePath);
+};
diff --git a/Engine/Plugins/Enterprise/GLTFExporter/MigratedFromGit/Source/GLTFExporter/Private/Converters/GLTFColor.h b/Engine/Plugins/Enterprise/GLTFExporter/MigratedFromGit/Source/GLTFExporter/Private/Converters/GLTFColor.h
new file mode 100644
--- /dev/null
+++ b/Engine/Plugins/Enterprise/GLTFExporter/MigratedFromGit/Source/GLTFExporter/Private/Converters/GLTFColor.h
@@ -0,0 +1,17 @@
+// Copyright Epic Games, Inc. All Rights Reserved.
+
+#pragma once
+
+#include "CoreTypes.h"
+
+// UE4 uses an endianness dependent memory layout (ABGR for little endian, ARGB for big endian)
+// glTF uses RGBA independent of endianness, hence why we need this dedicated struct
+struct FGLTFColor
+{
+	uint8 R,G,B,A;
+
+	FORCEINLINE FGLTFColor(uint8 R, uint8 G, uint8 B, uint8 A = 255)
+		: R(R), G(G), B(B), A(A)
+	{
+	}
+};
diff --git a/Engine/Plugins/Enterprise/GLTFExporter/MigratedFromGit/Source/GLTFExporter/Private/Converters/GLTFConverter.h b/Engine/Plugins/Enterprise/GLTFExporter/MigratedFromGit/Source/GLTFExporter/Private/Converters/GLTFConverter.h
new file mode 100644
--- /dev/null
+++ b/Engine/Plugins/Enterprise/GLTFExporter/MigratedFromGit/Source/GLTFExporter/Private/Converters/GLTFConverter.h
@@ -0,0 +1,42 @@
+// Copyright Epic Games, Inc. All Rights Reserved.
+
+#pragma once
+
+#include "CoreMinimal.h"
+
+struct FGLTFConvertBuilder;
+
+template <class IndexType, class KeyType>
+class TGLTFConverter
+{
+public:
+	virtual ~TGLTFConverter() = default;
+
+	template <typename... ArgsType>
+    FORCEINLINE IndexType Get(ArgsType&&... Args) const
+	{
+		const KeyType Key(Forward<ArgsType>(Args)...);
+		return IndexLookup.FindRef(Key);
+	}
+
+	template <typename... ArgsType>
+    FORCEINLINE IndexType GetOrAdd(FGLTFConvertBuilder& Builder, const FString& DesiredName, ArgsType&&... Args)
+	{
+		const KeyType Key(Forward<ArgsType>(Args)...);
+
+		IndexType Index = IndexLookup.FindRef(Key);
+		if (Index == INDEX_NONE)
+		{
+			Index = Add(Builder, DesiredName, Key);
+			IndexLookup.Add(Key, Index);
+		}
+
+		return Index;
+	}
+
+private:
+
+	TMap<KeyType, IndexType> IndexLookup;
+
+	virtual IndexType Add(FGLTFConvertBuilder& Builder, const FString& Name, KeyType Key) = 0;
+};
diff --git a/Engine/Plugins/Enterprise/GLTFExporter/MigratedFromGit/Source/GLTFExporter/Private/Converters/GLTFConverterUtility.h b/Engine/Plugins/Enterprise/GLTFExporter/MigratedFromGit/Source/GLTFExporter/Private/Converters/GLTFConverterUtility.h
new file mode 100644
--- /dev/null
+++ b/Engine/Plugins/Enterprise/GLTFExporter/MigratedFromGit/Source/GLTFExporter/Private/Converters/GLTFConverterUtility.h
@@ -0,0 +1,92 @@
+// Copyright Epic Games, Inc. All Rights Reserved.
+
+#pragma once
+
+#include "Converters/GLTFColor.h"
+#include "Json/GLTFJsonEnums.h"
+#include "Engine/EngineTypes.h"
+
+struct FGLTFConverterUtility
+{
+	static FVector ConvertVector(const FVector& Vector)
+	{
+		// UE4 uses a left-handed coordinate system, with Z up.
+		// glTF uses a right-handed coordinate system, with Y up.
+		return { Vector.X, Vector.Z, Vector.Y };
+	}
+
+	static FVector ConvertPosition(const FVector& Position)
+	{
+		return ConvertVector(Position) * 0.01f; // TODO: use options export scale instead of hardcoded value
+	}
+
+	static FVector ConvertScale(const FVector& Scale)
+	{
+		return ConvertVector(Scale);
+	}
+
+	static FVector ConvertNormal(const FVector& Normal)
+	{
+		return ConvertVector(Normal);
+	}
+
+	static FVector4 ConvertTangent(const FVector& Tangent)
+	{
+		// glTF stores tangent as Vec4, with W component indicating handedness of tangent basis.
+		return FVector4(ConvertVector(Tangent), 1.0f);
+	}
+
+	static FGLTFColor ConvertColor(const FColor& Color)
+	{
+		// UE4 uses ABGR or ARGB depending on endianness.
+		// glTF always uses RGBA independent of endianness.
+		return { Color.R, Color.G, Color.B, Color.A };
+	}
+
+	static FQuat ConvertRotation(const FQuat& Rotation)
+	{
+		// UE4 uses a left-handed coordinate system, with Z up.
+		// glTF uses a right-handed coordinate system, with Y up.
+		// Rotation = (qX, qY, qZ, qW) = (sin(angle/2) * aX, sin(angle/2) * aY, sin(angle/2) * aZ, cons(angle/2))
+		// where (aX, aY, aZ) - rotation axis, angle - rotation angle
+		// Y swapped with Z between these coordinate systems
+		// also, as handedness is changed rotation is inversed - hence negation
+		// therefore glTFRotation = (-qX, -qZ, -qY, qw)
+
+		const FQuat Result(-Rotation.X, -Rotation.Z, -Rotation.Y, Rotation.W);
+		// Not checking if quaternion is normalized
+		// e.g. some sources use non-unit Quats for rotation tangents
+		return Result;
+	}
+
+	static FMatrix ConvertMatrix(const FMatrix& Matrix)
+	{
+		// Unreal stores matrix elements in row major order.
+		// glTF stores matrix elements in column major order.
+
+		FMatrix Result;
+		for (int32 Row = 0; Row < 4; ++Row)
+		{
+			for (int32 Col = 0; Col < 4; ++Col)
+			{
+				Result.M[Col][Row] = Matrix.M[Row][Col];
+			}
+		}
+		return Result;
+	}
+
+	static EGLTFJsonAlphaMode ConvertAlphaMode(const EBlendMode Mode)
+	{
+		switch (Mode)
+		{
+			case EBlendMode::BLEND_Opaque:      return EGLTFJsonAlphaMode::Opaque;
+			case EBlendMode::BLEND_Translucent: return EGLTFJsonAlphaMode::Blend;
+			case EBlendMode::BLEND_Masked:      return EGLTFJsonAlphaMode::Mask;
+			default:                            return EGLTFJsonAlphaMode::Opaque; // fallback
+		}
+	}
+
+	static bool IsSkySphereBlueprint(const UBlueprint* Blueprint);
+
+	static bool IsHDRIBackdropBlueprint(const UBlueprint* Blueprint);
+};
diff --git a/Engine/Plugins/Enterprise/GLTFExporter/MigratedFromGit/Source/GLTFExporter/Private/Converters/GLTFLevelConverters.h b/Engine/Plugins/Enterprise/GLTFExporter/MigratedFromGit/Source/GLTFExporter/Private/Converters/GLTFLevelConverters.h
new file mode 100644
--- /dev/null
+++ b/Engine/Plugins/Enterprise/GLTFExporter/MigratedFromGit/Source/GLTFExporter/Private/Converters/GLTFLevelConverters.h
@@ -0,0 +1,17 @@
+// Copyright Epic Games, Inc. All Rights Reserved.
+
+#pragma once
+
+#include "Json/GLTFJsonIndex.h"
+#include "Converters/GLTFConverter.h"
+#include "Engine.h"
+
+class FGLTFSceneComponentConverter final : public TGLTFConverter<FGLTFJsonNodeIndex, TTuple<const USceneComponent*, bool, bool>>
+{
+	FGLTFJsonNodeIndex Add(FGLTFConvertBuilder& Builder, const FString& Name, TTuple<const USceneComponent*, bool, bool> Params) override;
+};
+
+class FGLTFLevelConverter final : public TGLTFConverter<FGLTFJsonSceneIndex, TTuple<const ULevel*, bool>>
+{
+	FGLTFJsonSceneIndex Add(FGLTFConvertBuilder& Builder, const FString& Name, TTuple<const ULevel*, bool> Params) override;
+};
diff --git a/Engine/Plugins/Enterprise/GLTFExporter/MigratedFromGit/Source/GLTFExporter/Private/Converters/GLTFStaticMeshConverters.h b/Engine/Plugins/Enterprise/GLTFExporter/MigratedFromGit/Source/GLTFExporter/Private/Converters/GLTFStaticMeshConverters.h
new file mode 100644
--- /dev/null
+++ b/Engine/Plugins/Enterprise/GLTFExporter/MigratedFromGit/Source/GLTFExporter/Private/Converters/GLTFStaticMeshConverters.h
@@ -0,0 +1,52 @@
+// Copyright Epic Games, Inc. All Rights Reserved.
+
+#pragma once
+
+#include "Json/GLTFJsonIndex.h"
+#include "Converters/GLTFConverter.h"
+#include "Engine.h"
+
+class FGLTFPositionVertexBufferConverter final : public TGLTFConverter<FGLTFJsonAccessorIndex, const FPositionVertexBuffer*>
+{
+	FGLTFJsonAccessorIndex Add(FGLTFConvertBuilder& Builder, const FString& Name, const FPositionVertexBuffer* VertexBuffer) override;
+};
+
+class FGLTFColorVertexBufferConverter final : public TGLTFConverter<FGLTFJsonAccessorIndex, const FColorVertexBuffer*>
+{
+	FGLTFJsonAccessorIndex Add(FGLTFConvertBuilder& Builder, const FString& Name, const FColorVertexBuffer* VertexBuffer) override;
+};
+
+class FGLTFStaticMeshNormalVertexBufferConverter final : public TGLTFConverter<FGLTFJsonAccessorIndex, const FStaticMeshVertexBuffer*>
+{
+	FGLTFJsonAccessorIndex Add(FGLTFConvertBuilder& Builder, const FString& Name, const FStaticMeshVertexBuffer* VertexBuffer) override;
+};
+
+class FGLTFStaticMeshTangentVertexBufferConverter final : public TGLTFConverter<FGLTFJsonAccessorIndex, const FStaticMeshVertexBuffer*>
+{
+	FGLTFJsonAccessorIndex Add(FGLTFConvertBuilder& Builder, const FString& Name, const FStaticMeshVertexBuffer* VertexBuffer) override;
+};
+
+class FGLTFStaticMeshUV0VertexBufferConverter final : public TGLTFConverter<FGLTFJsonAccessorIndex, const FStaticMeshVertexBuffer*>
+{
+	FGLTFJsonAccessorIndex Add(FGLTFConvertBuilder& Builder, const FString& Name, const FStaticMeshVertexBuffer* VertexBuffer) override;
+};
+
+class FGLTFStaticMeshUV1VertexBufferConverter final : public TGLTFConverter<FGLTFJsonAccessorIndex, const FStaticMeshVertexBuffer*>
+{
+	FGLTFJsonAccessorIndex Add(FGLTFConvertBuilder& Builder, const FString& Name, const FStaticMeshVertexBuffer* VertexBuffer) override;
+};
+
+class FGLTFStaticMeshIndexBufferConverter final : public TGLTFConverter<FGLTFJsonBufferViewIndex, const FRawStaticIndexBuffer*>
+{
+	FGLTFJsonBufferViewIndex Add(FGLTFConvertBuilder& Builder, const FString& Name, const FRawStaticIndexBuffer* IndexBuffer) override;
+};
+
+class FGLTFStaticMeshSectionConverter final : public TGLTFConverter<FGLTFJsonAccessorIndex, TTuple<const FStaticMeshSection*, const FRawStaticIndexBuffer*>>
+{
+	FGLTFJsonAccessorIndex Add(FGLTFConvertBuilder& Builder, const FString& Name, TTuple<const FStaticMeshSection*, const FRawStaticIndexBuffer*> Params) override;
+};
+
+class FGLTFStaticMeshConverter final : public TGLTFConverter<FGLTFJsonMeshIndex, TTuple<const FStaticMeshLODResources*, const FColorVertexBuffer*>>
+{
+	FGLTFJsonMeshIndex Add(FGLTFConvertBuilder& Builder, const FString& Name, TTuple<const FStaticMeshLODResources*, const FColorVertexBuffer*> Params) override;
+};
diff --git a/Engine/Plugins/Enterprise/GLTFExporter/MigratedFromGit/Source/GLTFExporter/Private/Json/GLTFJsonAccessor.h b/Engine/Plugins/Enterprise/GLTFExporter/MigratedFromGit/Source/GLTFExporter/Private/Json/GLTFJsonAccessor.h
new file mode 100644
--- /dev/null
+++ b/Engine/Plugins/Enterprise/GLTFExporter/MigratedFromGit/Source/GLTFExporter/Private/Json/GLTFJsonAccessor.h
@@ -0,0 +1,83 @@
+// Copyright Epic Games, Inc. All Rights Reserved.
+
+#pragma once
+
+#include "Json/GLTFJsonEnums.h"
+#include "Json/GLTFJsonIndex.h"
+#include "Json/GLTFJsonUtility.h"
+#include "Serialization/JsonSerializer.h"
+
+struct FGLTFJsonAccessor
+{
+	FString Name;
+
+	FGLTFJsonBufferViewIndex BufferView;
+	int64                    ByteOffset;
+	int32                    Count;
+	EGLTFJsonAccessorType    Type;
+	EGLTFJsonComponentType   ComponentType;
+	bool                     bNormalized;
+
+	int32 MinMaxLength;
+	float Min[16];
+	float Max[16];
+
+	FGLTFJsonAccessor()
+		: BufferView(INDEX_NONE)
+		, ByteOffset(0)
+		, Count(0)
+		, Type(EGLTFJsonAccessorType::None)
+		, ComponentType(EGLTFJsonComponentType::None)
+		, bNormalized(false)
+		, MinMaxLength(0)
+		, Min{0}
+		, Max{0}
+	{
+	}
+
+	template <class CharType = TCHAR, class PrintPolicy = TPrettyJsonPrintPolicy<CharType>>
+	void WriteObject(TJsonWriter<CharType, PrintPolicy>& JsonWriter) const
+	{
+		JsonWriter.WriteObjectStart();
+
+		if (!Name.IsEmpty())
+		{
+			JsonWriter.WriteValue(TEXT("name"), Name);
+		}
+
+		JsonWriter.WriteValue(TEXT("bufferView"), BufferView);
+
+		if (ByteOffset != 0)
+		{
+			JsonWriter.WriteValue(TEXT("byteOffset"), ByteOffset);
+		}
+
+		JsonWriter.WriteValue(TEXT("count"), Count);
+		JsonWriter.WriteValue(TEXT("type"), FGLTFJsonUtility::AccessorTypeToString(Type));
+		JsonWriter.WriteValue(TEXT("componentType"), FGLTFJsonUtility::ComponentTypeToNumber(ComponentType));
+
+		if (bNormalized)
+		{
+			JsonWriter.WriteValue(TEXT("normalized"), bNormalized);
+		}
+
+		if (MinMaxLength > 0)
+		{
+			JsonWriter.WriteArrayStart(TEXT("min"));
+			for (int32 ComponentIndex = 0; ComponentIndex < MinMaxLength; ++ComponentIndex)
+			{
+				JsonWriter.WriteValue(Min[ComponentIndex]);
+			}
+			JsonWriter.WriteArrayEnd();
+
+			JsonWriter.WriteArrayStart(TEXT("max"));
+			for (int32 ComponentIndex = 0; ComponentIndex < MinMaxLength; ++ComponentIndex)
+			{
+				JsonWriter.WriteValue(Max[ComponentIndex]);
+			}
+			JsonWriter.WriteArrayEnd();
+		}
+
+		JsonWriter.WriteObjectEnd();
+	}
+};
diff --git a/Engine/Plugins/Enterprise/GLTFExporter/MigratedFromGit/Source/GLTFExporter/Private/Json/GLTFJsonBuffer.h b/Engine/Plugins/Enterprise/GLTFExporter/MigratedFromGit/Source/GLTFExporter/Private/Json/GLTFJsonBuffer.h
new file mode 100644
--- /dev/null
+++ b/Engine/Plugins/Enterprise/GLTFExporter/MigratedFromGit/Source/GLTFExporter/Private/Json/GLTFJsonBuffer.h
@@ -0,0 +1,39 @@
+// Copyright Epic Games, Inc. All Rights Reserved.
+
+#pragma once
+
+#include "Json/GLTFJsonIndex.h"
+#include "Serialization/JsonSerializer.h"
+
+struct FGLTFJsonBuffer
+{
+	FString Name;
+
+	FString URI;
+	int64   ByteLength;
+
+	FGLTFJsonBuffer()
+		: ByteLength(0)
+	{
+	}
+
+	template <class CharType = TCHAR, class PrintPolicy = TPrettyJsonPrintPolicy<CharType>>
+	void WriteObject(TJsonWriter<CharType, PrintPolicy>& JsonWriter) const
+	{
+		JsonWriter.WriteObjectStart();
+
+		if (!Name.IsEmpty())
+		{
+			JsonWriter.WriteValue(TEXT("name"), Name);
+		}
+
+		if (!URI.IsEmpty())
+		{
+			JsonWriter.WriteValue(TEXT("uri"), URI);
+		}
+
+		JsonWriter.WriteValue(TEXT("byteLength"), ByteLength);
+
+		JsonWriter.WriteObjectEnd();
+	}
+};
diff --git a/Engine/Plugins/Enterprise/GLTFExporter/MigratedFromGit/Source/GLTFExporter/Private/Json/GLTFJsonBufferView.h b/Engine/Plugins/Enterprise/GLTFExporter/MigratedFromGit/Source/GLTFExporter/Private/Json/GLTFJsonBufferView.h
new file mode 100644
--- /dev/null
+++ b/Engine/Plugins/Enterprise/GLTFExporter/MigratedFromGit/Source/GLTFExporter/Private/Json/GLTFJsonBufferView.h
@@ -0,0 +1,62 @@
+// Copyright Epic Games, Inc. All Rights Reserved.
+
+#pragma once
+
+#include "Json/GLTFJsonEnums.h"
+#include "Json/GLTFJsonIndex.h"
+#include "Json/GLTFJsonUtility.h"
+#include "Serialization/JsonSerializer.h"
+
+struct FGLTFJsonBufferView
+{
+	FString Name;
+
+	FGLTFJsonBufferIndex Buffer;
+
+	int64 ByteLength;
+	int64 ByteOffset;
+	int32 ByteStride;
+
+	EGLTFJsonBufferTarget Target;
+
+	FGLTFJsonBufferView()
+		: Buffer(INDEX_NONE)
+		, ByteLength(0)
+		, ByteOffset(0)
+		, ByteStride(0)
+		, Target(EGLTFJsonBufferTarget::None)
+	{
+		// check that view fits completely inside the buffer
+	}
+
+	template <class CharType = TCHAR, class PrintPolicy = TPrettyJsonPrintPolicy<CharType>>
+	void WriteObject(TJsonWriter<CharType, PrintPolicy>& JsonWriter) const
+	{
+		JsonWriter.WriteObjectStart();
+
+		if (!Name.IsEmpty())
+		{
+			JsonWriter.WriteValue(TEXT("name"), Name);
+		}
+
+		JsonWriter.WriteValue(TEXT("buffer"), Buffer);
+		JsonWriter.WriteValue(TEXT("byteLength"), ByteLength);
+
+		if (ByteOffset != 0)
+		{
+			JsonWriter.WriteValue(TEXT("byteOffset"), ByteOffset);
+		}
+
+		if (ByteStride != 0)
+		{
+			JsonWriter.WriteValue(TEXT("byteStride"), ByteStride);
+		}
+
+		if (Target != EGLTFJsonBufferTarget::None)
+		{
+			JsonWriter.WriteValue(TEXT("target"), FGLTFJsonUtility::BufferTargetToNumber(Target));
+		}
+
+		JsonWriter.WriteObjectEnd();
+	}
+};
diff --git a/Engine/Plugins/Enterprise/GLTFExporter/MigratedFromGit/Source/GLTFExporter/Private/Json/GLTFJsonEnums.h b/Engine/Plugins/Enterprise/GLTFExporter/MigratedFromGit/Source/GLTFExporter/Private/Json/GLTFJsonEnums.h
new file mode 100644
--- /dev/null
+++ b/Engine/Plugins/Enterprise/GLTFExporter/MigratedFromGit/Source/GLTFExporter/Private/Json/GLTFJsonEnums.h
@@ -0,0 +1,105 @@
+// Copyright Epic Games, Inc. All Rights Reserved.
+
+#pragma once
+
+#include "CoreMinimal.h"
+
+enum class EGLTFExtension
+{
+	KHR_LightsPunctual,
+	KHR_MaterialsUnlit,
+	KHR_MaterialsClearCoat,
+	KHR_MeshQuantization
+};
+
+enum class EGLTFJsonAccessorType
+{
+	None = -1,
+	Scalar,
+	Vec2,
+	Vec3,
+	Vec4,
+	Mat2,
+	Mat3,
+	Mat4
+};
+
+enum class EGLTFJsonComponentType
+{
+	None = -1,
+	S8 = 5120, // signed byte
+	U8 = 5121, // unsigned byte
+	S16 = 5122, // signed short
+	U16 = 5123, // unsigned short
+	S32 = 5124, // unused
+	U32 = 5125, // unsigned int -- only valid for indices, not attributes
+	F32 = 5126  // float
+};
+
+enum class EGLTFJsonBufferTarget
+{
+	None = -1,
+	ArrayBuffer = 34962,
+	ElementArrayBuffer = 34963
+};
+
+enum class EGLTFJsonPrimitiveMode
+{
+	None = -1,
+	// unsupported
+	Points = 0,
+	Lines = 1,
+	LineLoop = 2,
+	LineStrip = 3,
+	// initially supported
+	Triangles = 4,
+	// will be supported prior to release
+	TriangleStrip = 5,
+	TriangleFan = 6
+};
+
+enum class EGLTFJsonAlphaMode
+{
+	None = -1,
+	Opaque,
+	Mask,
+	Blend
+};
+
+enum class EGLTFJsonSamplerFilter
+{
+	None = -1,
+	// valid for Min & Mag
+	Nearest = 9728,
+	Linear = 9729,
+	// valid for Min only
+	NearestMipmapNearest = 9984,
+	LinearMipmapNearest = 9985,
+	NearestMipmapLinear = 9986,
+	LinearMipmapLinear = 9987
+};
+
+enum class EGLTFJsonSamplerWrap
+{
+	None = -1,
+	Repeat = 10497,
+	MirroredRepeat = 33648,
+	ClampToEdge = 33071
+};
+
+enum class EGLTFJsonAnimationInterpolation
+{
+	None = -1,
+	Linear,
+	Step,
+	CubicSpline,
+};
+
+enum class EGLTFJsonAnimationPath
+{
+	None = -1,
+	Translation,
+	Rotation,
+	Scale,
+	Weights
+};
diff --git a/Engine/Plugins/Enterprise/GLTFExporter/MigratedFromGit/Source/GLTFExporter/Private/Json/GLTFJsonIndex.h b/Engine/Plugins/Enterprise/GLTFExporter/MigratedFromGit/Source/GLTFExporter/Private/Json/GLTFJsonIndex.h
new file mode 100644
--- /dev/null
+++ b/Engine/Plugins/Enterprise/GLTFExporter/MigratedFromGit/Source/GLTFExporter/Private/Json/GLTFJsonIndex.h
@@ -0,0 +1,61 @@
+// Copyright Epic Games, Inc. All Rights Reserved.
+
+#pragma once
+
+#include "CoreMinimal.h"
+
+template <typename DerivedType>
+struct FGLTFJsonIndex
+{
+	int32 Value;
+
+	FGLTFJsonIndex()
+		: Value(INDEX_NONE)
+	{
+	}
+
+	explicit FGLTFJsonIndex(int32 Value)
+		: Value(Value)
+	{
+	}
+
+	FORCEINLINE operator int32() const
+	{
+		return Value;
+	}
+
+	FORCEINLINE bool operator==(int32 OtherValue) const
+	{
+		return Value == OtherValue;
+	}
+
+	FORCEINLINE bool operator!=(int32 OtherValue) const
+	{
+		return Value != OtherValue;
+	}
+
+	FORCEINLINE bool operator==(const DerivedType& Other) const
+	{
+		return Value == Other.Value;
+	}
+
+	FORCEINLINE bool operator!=(const DerivedType& Other) const
+	{
+		return Value != Other.Value;
+	}
+
+	FORCEINLINE friend uint32 GetTypeHash(const DerivedType& Other)
+	{
+		return GetTypeHash(Other.Value);
+	}
+};
+
+struct FGLTFJsonAccessorIndex : FGLTFJsonIndex<FGLTFJsonAccessorIndex> { using FGLTFJsonIndex::FGLTFJsonIndex; };
+struct FGLTFJsonBufferIndex : FGLTFJsonIndex<FGLTFJsonBufferIndex> { using FGLTFJsonIndex::FGLTFJsonIndex; };
+struct FGLTFJsonBufferViewIndex : FGLTFJsonIndex<FGLTFJsonBufferViewIndex> { using FGLTFJsonIndex::FGLTFJsonIndex; };
+struct FGLTFJsonCameraIndex : FGLTFJsonIndex<FGLTFJsonCameraIndex> { using FGLTFJsonIndex::FGLTFJsonIndex; };
+struct FGLTFJsonMaterialIndex : FGLTFJsonIndex<FGLTFJsonMaterialIndex> { using FGLTFJsonIndex::FGLTFJsonIndex; };
+struct FGLTFJsonMeshIndex : FGLTFJsonIndex<FGLTFJsonMeshIndex> { using FGLTFJsonIndex::FGLTFJsonIndex; };
+struct FGLTFJsonNodeIndex : FGLTFJsonIndex<FGLTFJsonNodeIndex> { using FGLTFJsonIndex::FGLTFJsonIndex; };
+struct FGLTFJsonSceneIndex : FGLTFJsonIndex<FGLTFJsonSceneIndex> { using FGLTFJsonIndex::FGLTFJsonIndex; };
+struct FGLTFJsonSkinIndex : FGLTFJsonIndex<FGLTFJsonSkinIndex> { using FGLTFJsonIndex::FGLTFJsonIndex; };
diff --git a/Engine/Plugins/Enterprise/GLTFExporter/MigratedFromGit/Source/GLTFExporter/Private/Json/GLTFJsonMesh.h b/Engine/Plugins/Enterprise/GLTFExporter/MigratedFromGit/Source/GLTFExporter/Private/Json/GLTFJsonMesh.h
new file mode 100644
--- /dev/null
+++ b/Engine/Plugins/Enterprise/GLTFExporter/MigratedFromGit/Source/GLTFExporter/Private/Json/GLTFJsonMesh.h
@@ -0,0 +1,117 @@
+// Copyright Epic Games, Inc. All Rights Reserved.
+
+#pragma once
+
+#include "Json/GLTFJsonEnums.h"
+#include "Json/GLTFJsonIndex.h"
+#include "Serialization/JsonSerializer.h"
+
+struct FGLTFJsonAttributes
+{
+	FGLTFJsonAccessorIndex Position; // always required
+	FGLTFJsonAccessorIndex Normal;
+	FGLTFJsonAccessorIndex Tangent;
+	FGLTFJsonAccessorIndex TexCoord0;
+	FGLTFJsonAccessorIndex TexCoord1;
+	FGLTFJsonAccessorIndex Color0;
+	// skeletal mesh attributes
+	FGLTFJsonAccessorIndex Joints0;
+	FGLTFJsonAccessorIndex Weights0;
+
+	FGLTFJsonAttributes()
+		: Position(INDEX_NONE)
+		, Normal(INDEX_NONE)
+		, Tangent(INDEX_NONE)
+		, TexCoord0(INDEX_NONE)
+		, TexCoord1(INDEX_NONE)
+		, Color0(INDEX_NONE)
+		, Joints0(INDEX_NONE)
+		, Weights0(INDEX_NONE)
+	{
+	}
+
+	template <class CharType = TCHAR, class PrintPolicy = TPrettyJsonPrintPolicy<CharType>>
+	void WriteObject(TJsonWriter<CharType, PrintPolicy>& JsonWriter) const
+	{
+		JsonWriter.WriteObjectStart();
+
+		JsonWriter.WriteValue(TEXT("POSITION"), Position);
+		if (Normal != INDEX_NONE) JsonWriter.WriteValue(TEXT("NORMAL"), Normal);
+		if (Tangent != INDEX_NONE) JsonWriter.WriteValue(TEXT("TANGENT"), Tangent);
+		if (TexCoord0 != INDEX_NONE) JsonWriter.WriteValue(TEXT("TEXCOORD_0"), TexCoord0);
+		if (TexCoord1 != INDEX_NONE) JsonWriter.WriteValue(TEXT("TEXCOORD_1"), TexCoord1);
+		if (Color0 != INDEX_NONE) JsonWriter.WriteValue(TEXT("COLOR_0"), Color0);
+		if (Joints0 != INDEX_NONE) JsonWriter.WriteValue(TEXT("JOINTS_0"), Joints0);
+		if (Weights0 != INDEX_NONE) JsonWriter.WriteValue(TEXT("WEIGHTS_0"), Weights0);
+
+		JsonWriter.WriteObjectEnd();
+	}
+};
+
+struct FGLTFJsonPrimitive
+{
+	FGLTFJsonAccessorIndex Indices;
+	FGLTFJsonMaterialIndex Material;
+	EGLTFJsonPrimitiveMode Mode;
+	FGLTFJsonAttributes    Attributes;
+
+	FGLTFJsonPrimitive()
+		: Indices(INDEX_NONE)
+		, Material(INDEX_NONE)
+		, Mode(EGLTFJsonPrimitiveMode::None)
+	{
+	}
+
+	template <class CharType = TCHAR, class PrintPolicy = TPrettyJsonPrintPolicy<CharType>>
+	void WriteObject(TJsonWriter<CharType, PrintPolicy>& JsonWriter) const
+	{
+		JsonWriter.WriteObjectStart();
+
+		JsonWriter.WriteIdentifierPrefix(TEXT("attributes"));
+		Attributes.WriteObject(JsonWriter);
+
+		if (Indices != INDEX_NONE)
+		{
+			JsonWriter.WriteValue(TEXT("indices"), Indices);
+		}
+
+		if (Material != INDEX_NONE)
+		{
+			JsonWriter.WriteValue(TEXT("material"), Material);
+		}
+
+		if (Mode != EGLTFJsonPrimitiveMode::None)
+		{
+			JsonWriter.WriteValue(TEXT("mode"), FGLTFJsonUtility::PrimitiveModeToNumber(Mode));
+		}
+
+		JsonWriter.WriteObjectEnd();
+	}
+};
+
+struct FGLTFJsonMesh
+{
+	FString Name;
+
+	TArray<FGLTFJsonPrimitive> Primitives;
+
+	template <class CharType = TCHAR, class PrintPolicy = TPrettyJsonPrintPolicy<CharType>>
+	void WriteObject(TJsonWriter<CharType, PrintPolicy>& JsonWriter) const
+	{
+		JsonWriter.WriteObjectStart();
+
+		if (!Name.IsEmpty())
+		{
+			JsonWriter.WriteValue(TEXT("name"), Name);
+		}
+
+		JsonWriter.WriteArrayStart(TEXT("primitives"));
+		for (const FGLTFJsonPrimitive& Primitive : Primitives)
+		{
+			Primitive.WriteObject(JsonWriter);
+		}
+		JsonWriter.WriteArrayEnd();
+
+		JsonWriter.WriteObjectEnd();
+	}
+};
diff --git a/Engine/Plugins/Enterprise/GLTFExporter/MigratedFromGit/Source/GLTFExporter/Private/Json/GLTFJsonNode.h b/Engine/Plugins/Enterprise/GLTFExporter/MigratedFromGit/Source/GLTFExporter/Private/Json/GLTFJsonNode.h
new file mode 100644
--- /dev/null
+++ b/Engine/Plugins/Enterprise/GLTFExporter/MigratedFromGit/Source/GLTFExporter/Private/Json/GLTFJsonNode.h
@@ -0,0 +1,92 @@
+// Copyright Epic Games, Inc. All Rights Reserved.
+
+#pragma once
+
+#include "Json/GLTFJsonIndex.h"
+#include "Serialization/JsonSerializer.h"
+
+struct FGLTFJsonNode
+{
+	FString Name;
+
+	FVector Translation;
+	FQuat   Rotation;
+	FVector Scale;
+
+	FGLTFJsonCameraIndex Camera;
+	FGLTFJsonSkinIndex   Skin;
+	FGLTFJsonMeshIndex   Mesh;
+
+	TArray<FGLTFJsonNodeIndex> Children;
+
+	FGLTFJsonNode()
+		: Translation(FVector::ZeroVector)
+		, Rotation(FQuat::Identity)
+		, Scale(FVector::OneVector)
+		, Camera(INDEX_NONE)
+		, Skin(INDEX_NONE)
+		, Mesh(INDEX_NONE)
+	{
+	}
+
+	template <class CharType = TCHAR, class PrintPolicy = TPrettyJsonPrintPolicy<CharType>>
+	void WriteObject(TJsonWriter<CharType, PrintPolicy>& JsonWriter) const
+	{
+		JsonWriter.WriteObjectStart();
+
+		if (!Name.IsEmpty())
+		{
+			JsonWriter.WriteValue(TEXT("name"), Name);
+		}
+
+		if (Translation != FVector::ZeroVector)
+		{
+			JsonWriter.WriteArrayStart(TEXT("translation"));
+			JsonWriter.WriteValue(Translation.X);
+			JsonWriter.WriteValue(Translation.Y);
+			JsonWriter.WriteValue(Translation.Z);
+			JsonWriter.WriteArrayEnd();
+		}
+
+		if (Rotation != FQuat::Identity)
+		{
+			JsonWriter.WriteArrayStart(TEXT("rotation"));
+			JsonWriter.WriteValue(Rotation.X);
+			JsonWriter.WriteValue(Rotation.Y);
+			JsonWriter.WriteValue(Rotation.Z);
+			JsonWriter.WriteValue(Rotation.W);
+			JsonWriter.WriteArrayEnd();
+		}
+
+		if (Scale != FVector::OneVector)
+		{
+			JsonWriter.WriteArrayStart(TEXT("scale"));
+			JsonWriter.WriteValue(Scale.X);
+			JsonWriter.WriteValue(Scale.Y);
+			JsonWriter.WriteValue(Scale.Z);
+			JsonWriter.WriteArrayEnd();
+		}
+
+		if (Camera != INDEX_NONE)
+		{
+			JsonWriter.WriteValue(TEXT("camera"), Camera);
+		}
+
+		if (Skin != INDEX_NONE)
+		{
+			JsonWriter.WriteValue(TEXT("skin"), Skin);
+		}
+
+		if (Mesh != INDEX_NONE)
+		{
+			JsonWriter.WriteValue(TEXT("mesh"), Mesh);
+		}
+
+		if (Children.Num() > 0)
+		{
+			JsonWriter.WriteValue(TEXT("children"), Children);
+		}
+
+		JsonWriter.WriteObjectEnd();
+	}
+};
diff --git a/Engine/Plugins/Enterprise/GLTFExporter/MigratedFromGit/Source/GLTFExporter/Private/Json/GLTFJsonRoot.h b/Engine/Plugins/Enterprise/GLTFExporter/MigratedFromGit/Source/GLTFExporter/Private/Json/GLTFJsonRoot.h
new file mode 100644
--- /dev/null
+++ b/Engine/Plugins/Enterprise/GLTFExporter/MigratedFromGit/Source/GLTFExporter/Private/Json/GLTFJsonRoot.h
@@ -0,0 +1,189 @@
+// Copyright Epic Games, Inc. All Rights Reserved.
+
+#pragma once
+
+#include "Json/GLTFJsonIndex.h"
+#include "Json/GLTFJsonAccessor.h"
+#include "Json/GLTFJsonBufferView.h"
+#include "Json/GLTFJsonBuffer.h"
+#include "Json/GLTFJsonMesh.h"
+#include "Json/GLTFJsonScene.h"
+#include "Json/GLTFJsonNode.h"
+
+#include "Containers/Set.h"
+#include "Policies/CondensedJsonPrintPolicy.h"
+#include "Runtime/Launch/Resources/Version.h"
+
+
+struct FGLTFJsonAsset
+{
+	FString Version;
+	FString Generator;
+	FString Copyright;
+
+	FGLTFJsonAsset()
+		: Version(TEXT("2.0"))
+		, Generator(TEXT(EPIC_PRODUCT_NAME) TEXT(" ") ENGINE_VERSION_STRING)
+	{
+	}
+
+	template <class CharType = TCHAR, class PrintPolicy = TPrettyJsonPrintPolicy<CharType>>
+	void WriteObject(TJsonWriter<CharType, PrintPolicy>& JsonWriter) const
+	{
+		JsonWriter.WriteObjectStart();
+
+		JsonWriter.WriteValue(TEXT("version"), Version);
+
+		if (!Generator.IsEmpty())
+		{
+			JsonWriter.WriteValue(TEXT("generator"), Generator);
+		}
+
+		if (!Copyright.IsEmpty())
+		{
+			JsonWriter.WriteValue(TEXT("copyright"), Copyright);
+		}
+
+		JsonWriter.WriteObjectEnd();
+	}
+};
+
+struct FGLTFJsonRoot
+{
+	FGLTFJsonAsset Asset;
+
+	TSet<EGLTFExtension> ExtensionsUsed;
+	TSet<EGLTFExtension> ExtensionsRequired;
+
+	FGLTFJsonSceneIndex DefaultScene;
+
+	TArray<FGLTFJsonAccessor>   Accessors;
+	TArray<FGLTFJsonBuffer>     Buffers;
+	TArray<FGLTFJsonBufferView> BufferViews;
+	TArray<FGLTFJsonMesh>       Meshes;
+	TArray<FGLTFJsonNode>       Nodes;
+	TArray<FGLTFJsonScene>      Scenes;
+
+	FGLTFJsonRoot()
+		: DefaultScene(INDEX_NONE)
+	{
+	}
+
+	template <class CharType = TCHAR, class PrintPolicy = TPrettyJsonPrintPolicy<CharType>>
+	void WriteObject(TJsonWriter<CharType, PrintPolicy>& JsonWriter) const
+	{
+		JsonWriter.WriteObjectStart();
+
+		JsonWriter.WriteIdentifierPrefix(TEXT("asset"));
+		Asset.WriteObject(JsonWriter);
+
+		if (ExtensionsUsed.Num() > 0)
+		{
+			JsonWriter.WriteIdentifierPrefix(TEXT("extensionsUsed"));
+			JsonWriter.WriteArrayStart();
+			for (const EGLTFExtension& Extension : ExtensionsUsed)
+			{
+				JsonWriter.WriteValue(FGLTFJsonUtility::ExtensionToString(Extension));
+			}
+			JsonWriter.WriteArrayEnd();
+		}
+
+		if (ExtensionsRequired.Num() > 0)
+		{
+			JsonWriter.WriteIdentifierPrefix(TEXT("extensionsRequired"));
+			JsonWriter.WriteArrayStart();
+			for (const EGLTFExtension& Extension : ExtensionsRequired)
+			{
+				JsonWriter.WriteValue(FGLTFJsonUtility::ExtensionToString(Extension));
+			}
+			JsonWriter.WriteArrayEnd();
+		}
+
+		if (DefaultScene != INDEX_NONE)
+		{
+			JsonWriter.WriteValue(TEXT("scene"), DefaultScene);
+		}
+
+		if (Accessors.Num() > 0)
+		{
+			JsonWriter.WriteArrayStart(TEXT("accessors"));
+			for (const FGLTFJsonAccessor& Accessor : Accessors)
+			{
+				Accessor.WriteObject(JsonWriter);
+			}
+			JsonWriter.WriteArrayEnd();
+		}
+
+		if (BufferViews.Num() > 0)
+		{
+			JsonWriter.WriteArrayStart(TEXT("bufferViews"));
+			for (const FGLTFJsonBufferView& BufferView : BufferViews)
+			{
+				BufferView.WriteObject(JsonWriter);
+			}
+			JsonWriter.WriteArrayEnd();
+		}
+
+		if (Meshes.Num() > 0)
+		{
+			JsonWriter.WriteArrayStart(TEXT("meshes"));
+			for (const FGLTFJsonMesh& Mesh : Meshes)
+			{
+				Mesh.WriteObject(JsonWriter);
+			}
+			JsonWriter.WriteArrayEnd();
+		}
+
+		if (Nodes.Num() > 0)
+		{
+			JsonWriter.WriteArrayStart(TEXT("nodes"));
+			for (const FGLTFJsonNode& Node : Nodes)
+			{
+				Node.WriteObject(JsonWriter);
+			}
+			JsonWriter.WriteArrayEnd();
+		}
+
+		if (Scenes.Num() > 0)
+		{
+			JsonWriter.WriteArrayStart(TEXT("scenes"));
+			for (const FGLTFJsonScene& Scene : Scenes)
+			{
+				Scene.WriteObject(JsonWriter);
+			}
+			JsonWriter.WriteArrayEnd();
+		}
+
+		if (Buffers.Num() > 0)
+		{
+			JsonWriter.WriteArrayStart(TEXT("buffers"));
+			for (const FGLTFJsonBuffer& Buffer : Buffers)
+			{
+				Buffer.WriteObject(JsonWriter);
+			}
+			JsonWriter.WriteArrayEnd();
+		}
+
+		JsonWriter.WriteObjectEnd();
+	}
+
+	template <class CharType = TCHAR, class PrintPolicy = TPrettyJsonPrintPolicy<CharType>>
+	void Serialize(FArchive* const Archive) const
+	{
+		TSharedRef<TJsonWriter<CharType, PrintPolicy>> JsonWriter = TJsonWriterFactory<CharType, PrintPolicy>::Create(Archive);
+		WriteObject(*JsonWriter);
+		JsonWriter->Close();
+	}
+
+	void Serialize(FArchive* const Archive, bool bPrettyPrint = true) const
+	{
+		if (bPrettyPrint)
+		{
+			Serialize<UTF8CHAR, TPrettyJsonPrintPolicy<UTF8CHAR>>(Archive);
+		}
+		else
+		{
+			Serialize<UTF8CHAR, TCondensedJsonPrintPolicy<UTF8CHAR>>(Archive);
+		}
+	}
+};
diff --git a/Engine/Plugins/Enterprise/GLTFExporter/MigratedFromGit/Source/GLTFExporter/Private/Json/GLTFJsonScene.h b/Engine/Plugins/Enterprise/GLTFExporter/MigratedFromGit/Source/GLTFExporter/Private/Json/GLTFJsonScene.h
new file mode 100644
--- /dev/null
+++ b/Engine/Plugins/Enterprise/GLTFExporter/MigratedFromGit/Source/GLTFExporter/Private/Json/GLTFJsonScene.h
@@ -0,0 +1,31 @@
+// Copyright Epic Games, Inc. All Rights Reserved.
+
+#pragma once
+
+#include "Json/GLTFJsonIndex.h"
+#include "Serialization/JsonSerializer.h"
+
+struct FGLTFJsonScene
+{
+	FString Name;
+
+	TArray<FGLTFJsonNodeIndex> Nodes;
+
+	template <class CharType = TCHAR, class PrintPolicy = TPrettyJsonPrintPolicy<CharType>>
+	void WriteObject(TJsonWriter<CharType, PrintPolicy>& JsonWriter) const
+	{
+		JsonWriter.WriteObjectStart();
+
+		if (!Name.IsEmpty())
+		{
+			JsonWriter.WriteValue(TEXT("name"), Name);
+		}
+
+		if (Nodes.Num() > 0)
+		{
+			JsonWriter.WriteValue(TEXT("nodes"), Nodes);
+		}
+
+		JsonWriter.WriteObjectEnd();
+	}
+};
diff --git a/Engine/Plugins/Enterprise/GLTFExporter/MigratedFromGit/Source/GLTFExporter/Private/Json/GLTFJsonUtility.h b/Engine/Plugins/Enterprise/GLTFExporter/MigratedFromGit/Source/GLTFExporter/Private/Json/GLTFJsonUtility.h
new file mode 100644
--- /dev/null
+++ b/Engine/Plugins/Enterprise/GLTFExporter/MigratedFromGit/Source/GLTFExporter/Private/Json/GLTFJsonUtility.h
@@ -0,0 +1,71 @@
+// Copyright Epic Games, Inc. All Rights Reserved.
+
+#pragma once
+
+#include "Json/GLTFJsonEnums.h"
+
+struct FGLTFJsonUtility
+{
+	static const TCHAR* ExtensionToString(EGLTFExtension Extension)
+	{
+		switch (Extension)
+		{
+			case EGLTFExtension::KHR_LightsPunctual:     return TEXT("KHR_lights_punctual");
+			case EGLTFExtension::KHR_MaterialsUnlit:     return TEXT("KHR_materials_unlit");
+			case EGLTFExtension::KHR_MaterialsClearCoat: return TEXT("KHR_materials_clearcoat");
+			case EGLTFExtension::KHR_MeshQuantization:   return TEXT("KHR_mesh_quantization");
+			default:                                     return TEXT("unknown");
+		}
+	}
+
+	static const TCHAR* AlphaModeToString(EGLTFJsonAlphaMode Mode)
+	{
+		switch (Mode)
+		{
+			case EGLTFJsonAlphaMode::Opaque: return TEXT("OPAQUE");
+			case EGLTFJsonAlphaMode::Blend:  return TEXT("BLEND");
+			case EGLTFJsonAlphaMode::Mask:   return TEXT("MASK");
+			default:                         return TEXT("UNKNOWN");
+		}
+	}
+
+	static int32 PrimitiveModeToNumber(EGLTFJsonPrimitiveMode Mode)
+	{
+		return static_cast<int32>(Mode);
+	}
+
+	static int32 FilterToNumber(EGLTFJsonSamplerFilter Filter)
+	{
+		return static_cast<int32>(Filter);
+	}
+
+	static int32 WrapModeToNumber(EGLTFJsonSamplerWrap Wrap)
+	{
+		return static_cast<int32>(Wrap);
+	}
+
+	static const TCHAR* AccessorTypeToString(EGLTFJsonAccessorType Type)
+	{
+		switch (Type)
+		{
+			case EGLTFJsonAccessorType::Scalar: return TEXT("SCALAR");
+			case EGLTFJsonAccessorType::Vec2:   return TEXT("VEC2");
+			case EGLTFJsonAccessorType::Vec3:   return TEXT("VEC3");
+			case EGLTFJsonAccessorType::Vec4:   return TEXT("VEC4");
+			case EGLTFJsonAccessorType::Mat2:   return TEXT("MAT2");
+			case EGLTFJsonAccessorType::Mat3:   return TEXT("MAT3");
+			case EGLTFJsonAccessorType::Mat4:   return TEXT("MAT4");
+			default:                            return TEXT("UNKNOWN");
+		}
+	}
+
+	static int32 ComponentTypeToNumber(EGLTFJsonComponentType Type)
+	{
+		return static_cast<int32>(Type);
+	}
+
+	static int32 BufferTargetToNumber(EGLTFJsonBufferTarget Target)
+	{
+		return static_cast<int32>(Target);
+	}
+};
diff --git a/Engine/Plugins/Enterprise/GLTFExporter/MigratedFromGit/Source/GLTFExporter/Private/UI/GLTFExportOptionsWindow.h b/Engine/Plugins/Enterprise/GLTFExporter/MigratedFromGit/Source/GLTFExporter/Private/UI/GLTFExportOptionsWindow.h
new file mode 100644
--- /dev/null
+++ b/Engine/Plugins/Enterprise/GLTFExporter/MigratedFromGit/Source/GLTFExporter/Private/UI/GLTFExportOptionsWindow.h
@@ -0,0 +1,60 @@
+// Copyright 1998-2019 Epic Games, Inc. All Rights Reserved.
+
+#pragma once
+
+#include "GLTFExportOptions.h"
+#include "InputCoreTypes.h"
+#include "Widgets/DeclarativeSyntaxSupport.h"
+#include "Input/Reply.h"
+#include "Widgets/SCompoundWidget.h"
+#include "Widgets/SWindow.h"
+
+class SButton;
+
+class GLTFEXPORTER_API SGLTFExportOptionsWindow : public SCompoundWidget
+{
+public:
+	SLATE_BEGIN_ARGS(SGLTFExportOptionsWindow)
+		: _ExportOptions(nullptr)
+		, _BatchMode()
+		{}
+
+		SLATE_ARGUMENT( UGLTFExportOptions*, ExportOptions )
+		SLATE_ARGUMENT( TSharedPtr<SWindow>, WidgetWindow )
+		SLATE_ARGUMENT( FText, FullPath )
+		SLATE_ARGUMENT( bool, BatchMode )
+	SLATE_END_ARGS()
+
+public:
+	void Construct(const FArguments& InArgs);
+
+	virtual bool SupportsKeyboardFocus() const override { return true; }
+
+	FReply OnExport();
+
+	FReply OnExportAll();
+
+	FReply OnCancel();
+
+	virtual FReply OnKeyDown(const FGeometry& MyGeometry, const FKeyEvent& InKeyEvent) override;
+
+	bool ShouldExport() const;
+
+	bool ShouldExportAll() const;
+
+	SGLTFExportOptionsWindow();
+
+public:
+	static void ShowDialog(UGLTFExportOptions* ExportOptions, const FString& FullPath, bool bBatchMode, bool& bOutOperationCanceled, bool& bOutExportAll);
+
+private:
+	FReply OnResetToDefaultClick() const;
+
+private:
+	UGLTFExportOptions* ExportOptions;
+	TSharedPtr<class IDetailsView> DetailsView;
+	TWeakPtr<SWindow> WidgetWindow;
+	TSharedPtr<SButton> ImportButton;
+	bool bShouldExport;
+	bool bShouldExportAll;
+};
diff --git a/Engine/Plugins/Enterprise/GLTFExporter/MigratedFromGit/Source/GLTFExporter/Public/Builders/GLTFBufferBuilder.h b/Engine/Plugins/Enterprise/GLTFExporter/MigratedFromGit/Source/GLTFExporter/Public/Builders/GLTFBufferBuilder.h
deleted file mode 100644
diff --git a/Engine/Plugins/Enterprise/GLTFExporter/MigratedFromGit/Source/GLTFExporter/Public/Builders/GLTFContainerBuilder.h b/Engine/Plugins/Enterprise/GLTFExporter/MigratedFromGit/Source/GLTFExporter/Public/Builders/GLTFContainerBuilder.h
deleted file mode 100644
diff --git a/Engine/Plugins/Enterprise/GLTFExporter/MigratedFromGit/Source/GLTFExporter/Public/Builders/GLTFConvertBuilder.h b/Engine/Plugins/Enterprise/GLTFExporter/MigratedFromGit/Source/GLTFExporter/Public/Builders/GLTFConvertBuilder.h
deleted file mode 100644
diff --git a/Engine/Plugins/Enterprise/GLTFExporter/MigratedFromGit/Source/GLTFExporter/Public/Builders/GLTFJsonBuilder.h b/Engine/Plugins/Enterprise/GLTFExporter/MigratedFromGit/Source/GLTFExporter/Public/Builders/GLTFJsonBuilder.h
deleted file mode 100644
diff --git a/Engine/Plugins/Enterprise/GLTFExporter/MigratedFromGit/Source/GLTFExporter/Public/Converters/GLTFColor.h b/Engine/Plugins/Enterprise/GLTFExporter/MigratedFromGit/Source/GLTFExporter/Public/Converters/GLTFColor.h
deleted file mode 100644
diff --git a/Engine/Plugins/Enterprise/GLTFExporter/MigratedFromGit/Source/GLTFExporter/Public/Converters/GLTFConverter.h b/Engine/Plugins/Enterprise/GLTFExporter/MigratedFromGit/Source/GLTFExporter/Public/Converters/GLTFConverter.h
deleted file mode 100644
diff --git a/Engine/Plugins/Enterprise/GLTFExporter/MigratedFromGit/Source/GLTFExporter/Public/Converters/GLTFConverterUtility.h b/Engine/Plugins/Enterprise/GLTFExporter/MigratedFromGit/Source/GLTFExporter/Public/Converters/GLTFConverterUtility.h
deleted file mode 100644
diff --git a/Engine/Plugins/Enterprise/GLTFExporter/MigratedFromGit/Source/GLTFExporter/Public/Converters/GLTFLevelConverters.h b/Engine/Plugins/Enterprise/GLTFExporter/MigratedFromGit/Source/GLTFExporter/Public/Converters/GLTFLevelConverters.h
deleted file mode 100644
diff --git a/Engine/Plugins/Enterprise/GLTFExporter/MigratedFromGit/Source/GLTFExporter/Public/Converters/GLTFStaticMeshConverters.h b/Engine/Plugins/Enterprise/GLTFExporter/MigratedFromGit/Source/GLTFExporter/Public/Converters/GLTFStaticMeshConverters.h
deleted file mode 100644
diff --git a/Engine/Plugins/Enterprise/GLTFExporter/MigratedFromGit/Source/GLTFExporter/Public/Json/GLTFJsonAccessor.h b/Engine/Plugins/Enterprise/GLTFExporter/MigratedFromGit/Source/GLTFExporter/Public/Json/GLTFJsonAccessor.h
deleted file mode 100644
diff --git a/Engine/Plugins/Enterprise/GLTFExporter/MigratedFromGit/Source/GLTFExporter/Public/Json/GLTFJsonBuffer.h b/Engine/Plugins/Enterprise/GLTFExporter/MigratedFromGit/Source/GLTFExporter/Public/Json/GLTFJsonBuffer.h
deleted file mode 100644
diff --git a/Engine/Plugins/Enterprise/GLTFExporter/MigratedFromGit/Source/GLTFExporter/Public/Json/GLTFJsonBufferView.h b/Engine/Plugins/Enterprise/GLTFExporter/MigratedFromGit/Source/GLTFExporter/Public/Json/GLTFJsonBufferView.h
deleted file mode 100644
diff --git a/Engine/Plugins/Enterprise/GLTFExporter/MigratedFromGit/Source/GLTFExporter/Public/Json/GLTFJsonEnums.h b/Engine/Plugins/Enterprise/GLTFExporter/MigratedFromGit/Source/GLTFExporter/Public/Json/GLTFJsonEnums.h
deleted file mode 100644
diff --git a/Engine/Plugins/Enterprise/GLTFExporter/MigratedFromGit/Source/GLTFExporter/Public/Json/GLTFJsonIndex.h b/Engine/Plugins/Enterprise/GLTFExporter/MigratedFromGit/Source/GLTFExporter/Public/Json/GLTFJsonIndex.h
deleted file mode 100644
diff --git a/Engine/Plugins/Enterprise/GLTFExporter/MigratedFromGit/Source/GLTFExporter/Public/Json/GLTFJsonMesh.h b/Engine/Plugins/Enterprise/GLTFExporter/MigratedFromGit/Source/GLTFExporter/Public/Json/GLTFJsonMesh.h
deleted file mode 100644
diff --git a/Engine/Plugins/Enterprise/GLTFExporter/MigratedFromGit/Source/GLTFExporter/Public/Json/GLTFJsonNode.h b/Engine/Plugins/Enterprise/GLTFExporter/MigratedFromGit/Source/GLTFExporter/Public/Json/GLTFJsonNode.h
deleted file mode 100644
diff --git a/Engine/Plugins/Enterprise/GLTFExporter/MigratedFromGit/Source/GLTFExporter/Public/Json/GLTFJsonRoot.h b/Engine/Plugins/Enterprise/GLTFExporter/MigratedFromGit/Source/GLTFExporter/Public/Json/GLTFJsonRoot.h
deleted file mode 100644
diff --git a/Engine/Plugins/Enterprise/GLTFExporter/MigratedFromGit/Source/GLTFExporter/Public/Json/GLTFJsonScene.h b/Engine/Plugins/Enterprise/GLTFExporter/MigratedFromGit/Source/GLTFExporter/Public/Json/GLTFJsonScene.h
deleted file mode 100644
diff --git a/Engine/Plugins/Enterprise/GLTFExporter/MigratedFromGit/Source/GLTFExporter/Public/Json/GLTFJsonUtility.h b/Engine/Plugins/Enterprise/GLTFExporter/MigratedFromGit/Source/GLTFExporter/Public/Json/GLTFJsonUtility.h
deleted file mode 100644
diff --git a/Engine/Plugins/Enterprise/GLTFExporter/MigratedFromGit/Source/GLTFExporter/Public/UI/GLTFExportOptionsWindow.h b/Engine/Plugins/Enterprise/GLTFExporter/MigratedFromGit/Source/GLTFExporter/Public/UI/GLTFExportOptionsWindow.h
deleted file mode 100644
