[GLTFExporter] refactor: fix access encapsulation for remaining types

commit 7d5b7ec52e057467d8c72b6ea68ea99b009d6be3
Date: Sun, 17 May 2020 18:31:22 +0200

    refactor: fix access encapsulation for remaining types
    (cherry picked from commit cbb452c76277fe40a4234ae20232dfa736850a74​)

#skipci

diff --git a/Engine/Plugins/Enterprise/GLTFExporter/MigratedFromGit/Source/GLTFExporter/Private/Builders/GLTFBufferBuilder.h b/Engine/Plugins/Enterprise/GLTFExporter/MigratedFromGit/Source/GLTFExporter/Private/Builders/GLTFBufferBuilder.h
--- a/Engine/Plugins/Enterprise/GLTFExporter/MigratedFromGit/Source/GLTFExporter/Private/Builders/GLTFBufferBuilder.h
+++ b/Engine/Plugins/Enterprise/GLTFExporter/MigratedFromGit/Source/GLTFExporter/Private/Builders/GLTFBufferBuilder.h
@@ -4,10 +4,9 @@
 
 #include "Builders/GLTFJsonBuilder.h"
 
-struct FGLTFBufferBuilder : public FGLTFJsonBuilder
+class FGLTFBufferBuilder : public FGLTFJsonBuilder
 {
-	FGLTFJsonBufferIndex BufferIndex;
-	TArray<uint8> BufferData;
+public:
 
 	FGLTFBufferBuilder();
 
@@ -21,4 +20,9 @@ struct FGLTFBufferBuilder : public FGLTFJsonBuilder
 
 	void UpdateJsonBufferObject(const FString& BinaryFilePath);
 	virtual bool Serialize(FArchive& Archive, const FString& FilePath) override;
+
+private:
+
+	FGLTFJsonBufferIndex BufferIndex;
+	TArray<uint8> BufferData;
 };
diff --git a/Engine/Plugins/Enterprise/GLTFExporter/MigratedFromGit/Source/GLTFExporter/Private/Builders/GLTFContainerBuilder.h b/Engine/Plugins/Enterprise/GLTFExporter/MigratedFromGit/Source/GLTFExporter/Private/Builders/GLTFContainerBuilder.h
--- a/Engine/Plugins/Enterprise/GLTFExporter/MigratedFromGit/Source/GLTFExporter/Private/Builders/GLTFContainerBuilder.h
+++ b/Engine/Plugins/Enterprise/GLTFExporter/MigratedFromGit/Source/GLTFExporter/Private/Builders/GLTFContainerBuilder.h
@@ -4,6 +4,6 @@
 
 #include "Builders/GLTFConvertBuilder.h"
 
-struct FGLTFContainerBuilder : public FGLTFConvertBuilder
+class FGLTFContainerBuilder : public FGLTFConvertBuilder
 {
 };
diff --git a/Engine/Plugins/Enterprise/GLTFExporter/MigratedFromGit/Source/GLTFExporter/Private/Builders/GLTFConvertBuilder.h b/Engine/Plugins/Enterprise/GLTFExporter/MigratedFromGit/Source/GLTFExporter/Private/Builders/GLTFConvertBuilder.h
--- a/Engine/Plugins/Enterprise/GLTFExporter/MigratedFromGit/Source/GLTFExporter/Private/Builders/GLTFConvertBuilder.h
+++ b/Engine/Plugins/Enterprise/GLTFExporter/MigratedFromGit/Source/GLTFExporter/Private/Builders/GLTFConvertBuilder.h
@@ -6,20 +6,9 @@
 #include "Converters/GLTFStaticMeshConverters.h"
 #include "Converters/GLTFLevelConverters.h"
 
-struct FGLTFConvertBuilder : public FGLTFBufferBuilder
+class FGLTFConvertBuilder : public FGLTFBufferBuilder
 {
-	FGLTFPositionVertexBufferConverter PositionVertexBufferConverter;
-	FGLTFColorVertexBufferConverter ColorVertexBufferConverter;
-	FGLTFStaticMeshNormalVertexBufferConverter StaticMeshNormalVertexBufferConverter;
-	FGLTFStaticMeshTangentVertexBufferConverter StaticMeshTangentVertexBufferConverter;
-	FGLTFStaticMeshUV0VertexBufferConverter StaticMeshUV0VertexBufferConverter;
-	FGLTFStaticMeshUV1VertexBufferConverter StaticMeshUV1VertexBufferConverter;
-	FGLTFStaticMeshIndexBufferConverter StaticMeshIndexBufferConverter;
-	FGLTFStaticMeshSectionConverter StaticMeshSectionConverter;
-	FGLTFStaticMeshConverter StaticMeshConverter;
-
-	FGLTFSceneComponentConverter SceneComponentConverter;
-	FGLTFLevelConverter LevelConverter;
+public:
 
 	FGLTFJsonAccessorIndex GetOrAddPositionAccessor(const FPositionVertexBuffer* VertexBuffer, const FString& DesiredName = TEXT(""));
 	FGLTFJsonAccessorIndex GetOrAddNormalAccessor(const FStaticMeshVertexBuffer* VertexBuffer, const FString& DesiredName = TEXT(""));
@@ -38,4 +27,19 @@ struct FGLTFConvertBuilder : public FGLTFBufferBuilder
 	FGLTFJsonNodeIndex GetOrAddNode(const USceneComponent* SceneComponent, bool bSelectedOnly, bool bRootNode = false, const FString& DesiredName = TEXT(""));
 	FGLTFJsonSceneIndex GetOrAddScene(const ULevel* Level, bool bSelectedOnly, const FString& DesiredName = TEXT(""));
 	FGLTFJsonSceneIndex GetOrAddScene(const UWorld* World, bool bSelectedOnly, const FString& DesiredName = TEXT(""));
+
+private:
+
+	FGLTFPositionVertexBufferConverter PositionVertexBufferConverter;
+	FGLTFColorVertexBufferConverter ColorVertexBufferConverter;
+	FGLTFStaticMeshNormalVertexBufferConverter StaticMeshNormalVertexBufferConverter;
+	FGLTFStaticMeshTangentVertexBufferConverter StaticMeshTangentVertexBufferConverter;
+	FGLTFStaticMeshUV0VertexBufferConverter StaticMeshUV0VertexBufferConverter;
+	FGLTFStaticMeshUV1VertexBufferConverter StaticMeshUV1VertexBufferConverter;
+	FGLTFStaticMeshIndexBufferConverter StaticMeshIndexBufferConverter;
+	FGLTFStaticMeshSectionConverter StaticMeshSectionConverter;
+	FGLTFStaticMeshConverter StaticMeshConverter;
+
+	FGLTFSceneComponentConverter SceneComponentConverter;
+	FGLTFLevelConverter LevelConverter;
 };
diff --git a/Engine/Plugins/Enterprise/GLTFExporter/MigratedFromGit/Source/GLTFExporter/Private/Builders/GLTFJsonBuilder.cpp b/Engine/Plugins/Enterprise/GLTFExporter/MigratedFromGit/Source/GLTFExporter/Private/Builders/GLTFJsonBuilder.cpp
--- a/Engine/Plugins/Enterprise/GLTFExporter/MigratedFromGit/Source/GLTFExporter/Private/Builders/GLTFJsonBuilder.cpp
+++ b/Engine/Plugins/Enterprise/GLTFExporter/MigratedFromGit/Source/GLTFExporter/Private/Builders/GLTFJsonBuilder.cpp
@@ -3,6 +3,7 @@
 #include "Builders/GLTFJsonBuilder.h"
 
 FGLTFJsonBuilder::FGLTFJsonBuilder()
+	: DefaultScene(JsonRoot.DefaultScene)
 {
 }
 
diff --git a/Engine/Plugins/Enterprise/GLTFExporter/MigratedFromGit/Source/GLTFExporter/Private/Builders/GLTFJsonBuilder.h b/Engine/Plugins/Enterprise/GLTFExporter/MigratedFromGit/Source/GLTFExporter/Private/Builders/GLTFJsonBuilder.h
--- a/Engine/Plugins/Enterprise/GLTFExporter/MigratedFromGit/Source/GLTFExporter/Private/Builders/GLTFJsonBuilder.h
+++ b/Engine/Plugins/Enterprise/GLTFExporter/MigratedFromGit/Source/GLTFExporter/Private/Builders/GLTFJsonBuilder.h
@@ -4,9 +4,11 @@
 
 #include "Json/GLTFJsonRoot.h"
 
-struct FGLTFJsonBuilder
+class FGLTFJsonBuilder
 {
-	FGLTFJsonRoot JsonRoot;
+public:
+
+	FGLTFJsonSceneIndex& DefaultScene;
 
 	FGLTFJsonBuilder();
 	virtual ~FGLTFJsonBuilder();
@@ -19,4 +21,8 @@ struct FGLTFJsonBuilder
 	FGLTFJsonSceneIndex AddScene(const FGLTFJsonScene& JsonScene);
 
 	virtual bool Serialize(FArchive& Archive, const FString& FilePath);
+
+protected:
+
+	FGLTFJsonRoot JsonRoot;
 };
diff --git a/Engine/Plugins/Enterprise/GLTFExporter/MigratedFromGit/Source/GLTFExporter/Private/Converters/GLTFConverter.h b/Engine/Plugins/Enterprise/GLTFExporter/MigratedFromGit/Source/GLTFExporter/Private/Converters/GLTFConverter.h
--- a/Engine/Plugins/Enterprise/GLTFExporter/MigratedFromGit/Source/GLTFExporter/Private/Converters/GLTFConverter.h
+++ b/Engine/Plugins/Enterprise/GLTFExporter/MigratedFromGit/Source/GLTFExporter/Private/Converters/GLTFConverter.h
@@ -4,7 +4,7 @@
 
 #include "CoreMinimal.h"
 
-struct FGLTFConvertBuilder;
+class FGLTFConvertBuilder;
 
 template <class IndexType, class KeyType>
 class TGLTFConverter
diff --git a/Engine/Plugins/Enterprise/GLTFExporter/MigratedFromGit/Source/GLTFExporter/Private/Exporters/GLTFLevelExporter.cpp b/Engine/Plugins/Enterprise/GLTFExporter/MigratedFromGit/Source/GLTFExporter/Private/Exporters/GLTFLevelExporter.cpp
--- a/Engine/Plugins/Enterprise/GLTFExporter/MigratedFromGit/Source/GLTFExporter/Private/Exporters/GLTFLevelExporter.cpp
+++ b/Engine/Plugins/Enterprise/GLTFExporter/MigratedFromGit/Source/GLTFExporter/Private/Exporters/GLTFLevelExporter.cpp
@@ -22,7 +22,7 @@ bool UGLTFLevelExporter::ExportBinary(UObject* Object, const TCHAR* Type, FArchi
 
 	FGLTFContainerBuilder Builder;
 	const FGLTFJsonSceneIndex SceneIndex = Builder.GetOrAddScene(World, bSelectedOnly);
-	Builder.JsonRoot.DefaultScene = SceneIndex;
+	Builder.DefaultScene = SceneIndex;
 
 	return Builder.Serialize(Archive, CurrentFilename);
 }
diff --git a/Engine/Plugins/Enterprise/GLTFExporter/MigratedFromGit/Source/GLTFExporter/Private/Exporters/GLTFStaticMeshExporter.cpp b/Engine/Plugins/Enterprise/GLTFExporter/MigratedFromGit/Source/GLTFExporter/Private/Exporters/GLTFStaticMeshExporter.cpp
--- a/Engine/Plugins/Enterprise/GLTFExporter/MigratedFromGit/Source/GLTFExporter/Private/Exporters/GLTFStaticMeshExporter.cpp
+++ b/Engine/Plugins/Enterprise/GLTFExporter/MigratedFromGit/Source/GLTFExporter/Private/Exporters/GLTFStaticMeshExporter.cpp
@@ -32,7 +32,7 @@ bool UGLTFStaticMeshExporter::ExportBinary(UObject* Object, const TCHAR* Type, F
 	Scene.Nodes.Add(NodeIndex);
 	FGLTFJsonSceneIndex SceneIndex = Builder.AddScene(Scene);
 
-	Builder.JsonRoot.DefaultScene = SceneIndex;
+	Builder.DefaultScene = SceneIndex;
 
 	return Builder.Serialize(Archive, CurrentFilename);
 }
diff --git a/Engine/Plugins/Enterprise/GLTFExporter/MigratedFromGit/Source/GLTFExporter/Private/Json/GLTFJsonIndex.h b/Engine/Plugins/Enterprise/GLTFExporter/MigratedFromGit/Source/GLTFExporter/Private/Json/GLTFJsonIndex.h
--- a/Engine/Plugins/Enterprise/GLTFExporter/MigratedFromGit/Source/GLTFExporter/Private/Json/GLTFJsonIndex.h
+++ b/Engine/Plugins/Enterprise/GLTFExporter/MigratedFromGit/Source/GLTFExporter/Private/Json/GLTFJsonIndex.h
@@ -7,8 +7,6 @@
 template <typename DerivedType>
 struct FGLTFJsonIndex
 {
-	int32 Value;
-
 	FGLTFJsonIndex()
 		: Value(INDEX_NONE)
 	{
@@ -48,6 +46,10 @@ struct FGLTFJsonIndex
 	{
 		return GetTypeHash(Other.Value);
 	}
+
+private:
+
+	int32 Value;
 };
 
 struct FGLTFJsonAccessorIndex : FGLTFJsonIndex<FGLTFJsonAccessorIndex> { using FGLTFJsonIndex::FGLTFJsonIndex; };
diff --git a/Engine/Plugins/Enterprise/GLTFExporter/MigratedFromGit/Source/GLTFExporter/Public/Exporters/GLTFAnimSequenceExporter.h b/Engine/Plugins/Enterprise/GLTFExporter/MigratedFromGit/Source/GLTFExporter/Public/Exporters/GLTFAnimSequenceExporter.h
--- a/Engine/Plugins/Enterprise/GLTFExporter/MigratedFromGit/Source/GLTFExporter/Public/Exporters/GLTFAnimSequenceExporter.h
+++ b/Engine/Plugins/Enterprise/GLTFExporter/MigratedFromGit/Source/GLTFExporter/Public/Exporters/GLTFAnimSequenceExporter.h
@@ -9,9 +9,9 @@ UCLASS()
 class GLTFEXPORTER_API UGLTFAnimSequenceExporter : public UGLTFExporter
 {
 public:
+
 	GENERATED_BODY()
 
-public:
 	UGLTFAnimSequenceExporter(const FObjectInitializer& ObjectInitializer = FObjectInitializer());
 
 	//~ Begin UExporter Interface
diff --git a/Engine/Plugins/Enterprise/GLTFExporter/MigratedFromGit/Source/GLTFExporter/Public/Exporters/GLTFExporter.h b/Engine/Plugins/Enterprise/GLTFExporter/MigratedFromGit/Source/GLTFExporter/Public/Exporters/GLTFExporter.h
--- a/Engine/Plugins/Enterprise/GLTFExporter/MigratedFromGit/Source/GLTFExporter/Public/Exporters/GLTFExporter.h
+++ b/Engine/Plugins/Enterprise/GLTFExporter/MigratedFromGit/Source/GLTFExporter/Public/Exporters/GLTFExporter.h
@@ -13,9 +13,9 @@ UCLASS()
 class GLTFEXPORTER_API UGLTFExporter : public UExporter
 {
 public:
+
 	GENERATED_BODY()
 
-public:
 	UGLTFExporter(const FObjectInitializer& ObjectInitializer = FObjectInitializer());
 
 protected:
diff --git a/Engine/Plugins/Enterprise/GLTFExporter/MigratedFromGit/Source/GLTFExporter/Public/Exporters/GLTFLevelExporter.h b/Engine/Plugins/Enterprise/GLTFExporter/MigratedFromGit/Source/GLTFExporter/Public/Exporters/GLTFLevelExporter.h
--- a/Engine/Plugins/Enterprise/GLTFExporter/MigratedFromGit/Source/GLTFExporter/Public/Exporters/GLTFLevelExporter.h
+++ b/Engine/Plugins/Enterprise/GLTFExporter/MigratedFromGit/Source/GLTFExporter/Public/Exporters/GLTFLevelExporter.h
@@ -9,9 +9,9 @@ UCLASS()
 class GLTFEXPORTER_API UGLTFLevelExporter : public UGLTFExporter
 {
 public:
+
 	GENERATED_BODY()
 
-public:
 	UGLTFLevelExporter(const FObjectInitializer& ObjectInitializer = FObjectInitializer());
 
 	//~ Begin UExporter Interface
diff --git a/Engine/Plugins/Enterprise/GLTFExporter/MigratedFromGit/Source/GLTFExporter/Public/Exporters/GLTFSkeletalMeshExporter.h b/Engine/Plugins/Enterprise/GLTFExporter/MigratedFromGit/Source/GLTFExporter/Public/Exporters/GLTFSkeletalMeshExporter.h
--- a/Engine/Plugins/Enterprise/GLTFExporter/MigratedFromGit/Source/GLTFExporter/Public/Exporters/GLTFSkeletalMeshExporter.h
+++ b/Engine/Plugins/Enterprise/GLTFExporter/MigratedFromGit/Source/GLTFExporter/Public/Exporters/GLTFSkeletalMeshExporter.h
@@ -9,9 +9,9 @@ UCLASS()
 class GLTFEXPORTER_API UGLTFSkeletalMeshExporter : public UGLTFExporter
 {
 public:
+
 	GENERATED_BODY()
 
-public:
 	UGLTFSkeletalMeshExporter(const FObjectInitializer& ObjectInitializer = FObjectInitializer());
 
 	//~ Begin UExporter Interface
diff --git a/Engine/Plugins/Enterprise/GLTFExporter/MigratedFromGit/Source/GLTFExporter/Public/Exporters/GLTFStaticMeshExporter.h b/Engine/Plugins/Enterprise/GLTFExporter/MigratedFromGit/Source/GLTFExporter/Public/Exporters/GLTFStaticMeshExporter.h
--- a/Engine/Plugins/Enterprise/GLTFExporter/MigratedFromGit/Source/GLTFExporter/Public/Exporters/GLTFStaticMeshExporter.h
+++ b/Engine/Plugins/Enterprise/GLTFExporter/MigratedFromGit/Source/GLTFExporter/Public/Exporters/GLTFStaticMeshExporter.h
@@ -9,9 +9,9 @@ UCLASS()
 class GLTFEXPORTER_API UGLTFStaticMeshExporter : public UGLTFExporter
 {
 public:
+
 	GENERATED_BODY()
 
-public:
 	UGLTFStaticMeshExporter(const FObjectInitializer& ObjectInitializer = FObjectInitializer());
 
 	//~ Begin UExporter Interface
