[GLTFExporter] refactor: reduce code duplication in exporters

commit bee1c27f40bc00699c4177a635952760415c53bb
Date: Sun, 17 May 2020 19:57:57 +0200

    refactor: reduce code duplication in exporters
    (cherry picked from commit a1e418a0934e0d4ff99ac9cd3087481b75ffa7a7​)

#skipci

diff --git a/Engine/Plugins/Enterprise/GLTFExporter/MigratedFromGit/Source/GLTFExporter/Private/Exporters/GLTFAnimSequenceExporter.cpp b/Engine/Plugins/Enterprise/GLTFExporter/MigratedFromGit/Source/GLTFExporter/Private/Exporters/GLTFAnimSequenceExporter.cpp
--- a/Engine/Plugins/Enterprise/GLTFExporter/MigratedFromGit/Source/GLTFExporter/Private/Exporters/GLTFAnimSequenceExporter.cpp
+++ b/Engine/Plugins/Enterprise/GLTFExporter/MigratedFromGit/Source/GLTFExporter/Private/Exporters/GLTFAnimSequenceExporter.cpp
@@ -9,16 +9,9 @@ UGLTFAnimSequenceExporter::UGLTFAnimSequenceExporter(const FObjectInitializer& O
 	SupportedClass = UAnimSequence::StaticClass();
 }
 
-bool UGLTFAnimSequenceExporter::ExportBinary(UObject* Object, const TCHAR* Type, FArchive& Archive, FFeedbackContext* Warn, int32 FileIndex, uint32 PortFlags)
+bool UGLTFAnimSequenceExporter::Add(FGLTFContainerBuilder& Builder, const UObject* Object)
 {
 	const UAnimSequence* AnimSequence = CastChecked<UAnimSequence>(Object);
-
-	if (!FillExportOptions())
-	{
-		// User cancelled the export
-		return false;
-	}
-
 	// TODO: implement
 	return true;
 }
diff --git a/Engine/Plugins/Enterprise/GLTFExporter/MigratedFromGit/Source/GLTFExporter/Private/Exporters/GLTFExporter.cpp b/Engine/Plugins/Enterprise/GLTFExporter/MigratedFromGit/Source/GLTFExporter/Private/Exporters/GLTFExporter.cpp
--- a/Engine/Plugins/Enterprise/GLTFExporter/MigratedFromGit/Source/GLTFExporter/Private/Exporters/GLTFExporter.cpp
+++ b/Engine/Plugins/Enterprise/GLTFExporter/MigratedFromGit/Source/GLTFExporter/Private/Exporters/GLTFExporter.cpp
@@ -1,8 +1,10 @@
 // Copyright Epic Games, Inc. All Rights Reserved.
 
 #include "Exporters/GLTFExporter.h"
-#include "UnrealExporter.h"
-#include "Engine.h"
+#include "UObject/GCObjectScopeGuard.h"
+#include "Builders/GLTFContainerBuilder.h"
+#include "GLTFExportOptions.h"
+
 
 UGLTFExporter::UGLTFExporter(const FObjectInitializer& ObjectInitializer)
 	: Super(ObjectInitializer)
@@ -17,11 +19,35 @@ UGLTFExporter::UGLTFExporter(const FObjectInitializer& ObjectInitializer)
 	// TODO: uncomment when support for .glb is implemented
 	// FormatExtension.Add(TEXT("glb"));
 	// FormatDescription.Add(TEXT("GL Transmission Format (Binary)"));
+}
+
+bool UGLTFExporter::ExportBinary(UObject* Object, const TCHAR* Type, FArchive& Archive, FFeedbackContext* Warn, int32 FileIndex, uint32 PortFlags)
+{
+	UGLTFExportOptions* Options = NewObject<UGLTFExportOptions>();
+	FGCObjectScopeGuard OptionsGuard(Options);
+
+	if (!FillExportOptions(Options))
+	{
+		// User cancelled the export
+		return false;
+	}
 
-	ExportOptions = NewObject<UGLTFExportOptions>(this, TEXT("GLTF Export Options"));
+	FGLTFContainerBuilder Builder;
+	if (!Add(Builder, Object))
+	{
+		// TODO: Report error
+		return false;
+	}
+
+	return Builder.Serialize(Archive, CurrentFilename);
+}
+
+bool UGLTFExporter::Add(FGLTFContainerBuilder& Builder, const UObject* Object)
+{
+	return false;
 }
 
-bool UGLTFExporter::FillExportOptions()
+bool UGLTFExporter::FillExportOptions(UGLTFExportOptions* ExportOptions)
 {
 	bool bExportAll = GetBatchMode() && !GetShowExportOption();
 	bool bExportCancel = false;
diff --git a/Engine/Plugins/Enterprise/GLTFExporter/MigratedFromGit/Source/GLTFExporter/Private/Exporters/GLTFLevelExporter.cpp b/Engine/Plugins/Enterprise/GLTFExporter/MigratedFromGit/Source/GLTFExporter/Private/Exporters/GLTFLevelExporter.cpp
--- a/Engine/Plugins/Enterprise/GLTFExporter/MigratedFromGit/Source/GLTFExporter/Private/Exporters/GLTFLevelExporter.cpp
+++ b/Engine/Plugins/Enterprise/GLTFExporter/MigratedFromGit/Source/GLTFExporter/Private/Exporters/GLTFLevelExporter.cpp
@@ -10,19 +10,11 @@ UGLTFLevelExporter::UGLTFLevelExporter(const FObjectInitializer& ObjectInitializ
 	SupportedClass = UWorld::StaticClass();
 }
 
-bool UGLTFLevelExporter::ExportBinary(UObject* Object, const TCHAR* Type, FArchive& Archive, FFeedbackContext* Warn, int32 FileIndex, uint32 PortFlags)
+bool UGLTFLevelExporter::Add(FGLTFContainerBuilder& Builder, const UObject* Object)
 {
 	const UWorld* World = CastChecked<UWorld>(Object);
-
-	if (!FillExportOptions())
-	{
-		// User cancelled the export
-		return false;
-	}
-
-	FGLTFContainerBuilder Builder;
 	const FGLTFJsonSceneIndex SceneIndex = Builder.GetOrAddScene(World, bSelectedOnly);
-	Builder.DefaultScene = SceneIndex;
 
-	return Builder.Serialize(Archive, CurrentFilename);
+	Builder.DefaultScene = SceneIndex;
+	return true;
 }
diff --git a/Engine/Plugins/Enterprise/GLTFExporter/MigratedFromGit/Source/GLTFExporter/Private/Exporters/GLTFSkeletalMeshExporter.cpp b/Engine/Plugins/Enterprise/GLTFExporter/MigratedFromGit/Source/GLTFExporter/Private/Exporters/GLTFSkeletalMeshExporter.cpp
--- a/Engine/Plugins/Enterprise/GLTFExporter/MigratedFromGit/Source/GLTFExporter/Private/Exporters/GLTFSkeletalMeshExporter.cpp
+++ b/Engine/Plugins/Enterprise/GLTFExporter/MigratedFromGit/Source/GLTFExporter/Private/Exporters/GLTFSkeletalMeshExporter.cpp
@@ -9,16 +9,9 @@ UGLTFSkeletalMeshExporter::UGLTFSkeletalMeshExporter(const FObjectInitializer& O
 	SupportedClass = USkeletalMesh::StaticClass();
 }
 
-bool UGLTFSkeletalMeshExporter::ExportBinary(UObject* Object, const TCHAR* Type, FArchive& Archive, FFeedbackContext* Warn, int32 FileIndex, uint32 PortFlags)
+bool UGLTFSkeletalMeshExporter::Add(FGLTFContainerBuilder& Builder, const UObject* Object)
 {
 	const USkeletalMesh* SkeletalMesh = CastChecked<USkeletalMesh>(Object);
-
-	if (!FillExportOptions())
-	{
-		// User cancelled the export
-		return false;
-	}
-
 	// TODO: implement
 	return true;
 }
diff --git a/Engine/Plugins/Enterprise/GLTFExporter/MigratedFromGit/Source/GLTFExporter/Private/Exporters/GLTFStaticMeshExporter.cpp b/Engine/Plugins/Enterprise/GLTFExporter/MigratedFromGit/Source/GLTFExporter/Private/Exporters/GLTFStaticMeshExporter.cpp
--- a/Engine/Plugins/Enterprise/GLTFExporter/MigratedFromGit/Source/GLTFExporter/Private/Exporters/GLTFStaticMeshExporter.cpp
+++ b/Engine/Plugins/Enterprise/GLTFExporter/MigratedFromGit/Source/GLTFExporter/Private/Exporters/GLTFStaticMeshExporter.cpp
@@ -10,18 +10,9 @@ UGLTFStaticMeshExporter::UGLTFStaticMeshExporter(const FObjectInitializer& Objec
 	SupportedClass = UStaticMesh::StaticClass();
 }
 
-bool UGLTFStaticMeshExporter::ExportBinary(UObject* Object, const TCHAR* Type, FArchive& Archive, FFeedbackContext* Warn, int32 FileIndex, uint32 PortFlags)
+bool UGLTFStaticMeshExporter::Add(FGLTFContainerBuilder& Builder, const UObject* Object)
 {
 	const UStaticMesh* StaticMesh = CastChecked<UStaticMesh>(Object);
-
-	if (!FillExportOptions())
-	{
-		// User cancelled the export
-		return false;
-	}
-
-	FGLTFContainerBuilder Builder;
-
 	FGLTFJsonMeshIndex MeshIndex = Builder.GetOrAddMesh(StaticMesh);
 
 	FGLTFJsonNode Node;
@@ -33,6 +24,5 @@ bool UGLTFStaticMeshExporter::ExportBinary(UObject* Object, const TCHAR* Type, F
 	FGLTFJsonSceneIndex SceneIndex = Builder.AddScene(Scene);
 
 	Builder.DefaultScene = SceneIndex;
-
-	return Builder.Serialize(Archive, CurrentFilename);
+	return true;
 }
diff --git a/Engine/Plugins/Enterprise/GLTFExporter/MigratedFromGit/Source/GLTFExporter/Public/Exporters/GLTFAnimSequenceExporter.h b/Engine/Plugins/Enterprise/GLTFExporter/MigratedFromGit/Source/GLTFExporter/Public/Exporters/GLTFAnimSequenceExporter.h
--- a/Engine/Plugins/Enterprise/GLTFExporter/MigratedFromGit/Source/GLTFExporter/Public/Exporters/GLTFAnimSequenceExporter.h
+++ b/Engine/Plugins/Enterprise/GLTFExporter/MigratedFromGit/Source/GLTFExporter/Public/Exporters/GLTFAnimSequenceExporter.h
@@ -6,7 +6,7 @@
 #include "GLTFAnimSequenceExporter.generated.h"
 
 UCLASS()
-class GLTFEXPORTER_API UGLTFAnimSequenceExporter : public UGLTFExporter
+class GLTFEXPORTER_API UGLTFAnimSequenceExporter final : public UGLTFExporter
 {
 public:
 
@@ -14,7 +14,5 @@ public:
 
 	UGLTFAnimSequenceExporter(const FObjectInitializer& ObjectInitializer = FObjectInitializer());
 
-	//~ Begin UExporter Interface
-	virtual bool ExportBinary(UObject* Object, const TCHAR* Type, FArchive& Archive, FFeedbackContext* Warn, int32 FileIndex = 0, uint32 PortFlags = 0) override;
-	//~ End UExporter Interface
+	bool Add(FGLTFContainerBuilder& Builder, const UObject* Object) override;
 };
diff --git a/Engine/Plugins/Enterprise/GLTFExporter/MigratedFromGit/Source/GLTFExporter/Public/Exporters/GLTFExporter.h b/Engine/Plugins/Enterprise/GLTFExporter/MigratedFromGit/Source/GLTFExporter/Public/Exporters/GLTFExporter.h
--- a/Engine/Plugins/Enterprise/GLTFExporter/MigratedFromGit/Source/GLTFExporter/Public/Exporters/GLTFExporter.h
+++ b/Engine/Plugins/Enterprise/GLTFExporter/MigratedFromGit/Source/GLTFExporter/Public/Exporters/GLTFExporter.h
@@ -6,9 +6,11 @@
 #include "UObject/ObjectMacros.h"
 #include "UObject/UObjectGlobals.h"
 #include "Exporters/Exporter.h"
-#include "GLTFExportOptions.h"
 #include "GLTFExporter.generated.h"
 
+class FGLTFContainerBuilder;
+class UGLTFExportOptions;
+
 UCLASS()
 class GLTFEXPORTER_API UGLTFExporter : public UExporter
 {
@@ -16,12 +18,15 @@ public:
 
 	GENERATED_BODY()
 
-	UGLTFExporter(const FObjectInitializer& ObjectInitializer = FObjectInitializer());
+	UGLTFExporter(const FObjectInitializer& ObjectInitializer);
+
+	virtual bool ExportBinary(UObject* Object, const TCHAR* Type, FArchive& Archive, FFeedbackContext* Warn, int32 FileIndex, uint32 PortFlags) override;
 
 protected:
 
-	bool FillExportOptions();
+	virtual bool Add(FGLTFContainerBuilder& Builder, const UObject* Object);
+
+private:
 
-	UPROPERTY(transient)
-	UGLTFExportOptions* ExportOptions;
+	bool FillExportOptions(UGLTFExportOptions* ExportOptions);
 };
diff --git a/Engine/Plugins/Enterprise/GLTFExporter/MigratedFromGit/Source/GLTFExporter/Public/Exporters/GLTFLevelExporter.h b/Engine/Plugins/Enterprise/GLTFExporter/MigratedFromGit/Source/GLTFExporter/Public/Exporters/GLTFLevelExporter.h
--- a/Engine/Plugins/Enterprise/GLTFExporter/MigratedFromGit/Source/GLTFExporter/Public/Exporters/GLTFLevelExporter.h
+++ b/Engine/Plugins/Enterprise/GLTFExporter/MigratedFromGit/Source/GLTFExporter/Public/Exporters/GLTFLevelExporter.h
@@ -6,7 +6,7 @@
 #include "GLTFLevelExporter.generated.h"
 
 UCLASS()
-class GLTFEXPORTER_API UGLTFLevelExporter : public UGLTFExporter
+class GLTFEXPORTER_API UGLTFLevelExporter final : public UGLTFExporter
 {
 public:
 
@@ -14,7 +14,5 @@ public:
 
 	UGLTFLevelExporter(const FObjectInitializer& ObjectInitializer = FObjectInitializer());
 
-	//~ Begin UExporter Interface
-	virtual bool ExportBinary(UObject* Object, const TCHAR* Type, FArchive& Archive, FFeedbackContext* Warn, int32 FileIndex = 0, uint32 PortFlags = 0) override;
-	//~ End UExporter Interface
+	bool Add(FGLTFContainerBuilder& Builder, const UObject* Object) override;
 };
diff --git a/Engine/Plugins/Enterprise/GLTFExporter/MigratedFromGit/Source/GLTFExporter/Public/Exporters/GLTFSkeletalMeshExporter.h b/Engine/Plugins/Enterprise/GLTFExporter/MigratedFromGit/Source/GLTFExporter/Public/Exporters/GLTFSkeletalMeshExporter.h
--- a/Engine/Plugins/Enterprise/GLTFExporter/MigratedFromGit/Source/GLTFExporter/Public/Exporters/GLTFSkeletalMeshExporter.h
+++ b/Engine/Plugins/Enterprise/GLTFExporter/MigratedFromGit/Source/GLTFExporter/Public/Exporters/GLTFSkeletalMeshExporter.h
@@ -6,7 +6,7 @@
 #include "GLTFSkeletalMeshExporter.generated.h"
 
 UCLASS()
-class GLTFEXPORTER_API UGLTFSkeletalMeshExporter : public UGLTFExporter
+class GLTFEXPORTER_API UGLTFSkeletalMeshExporter final : public UGLTFExporter
 {
 public:
 
@@ -14,7 +14,5 @@ public:
 
 	UGLTFSkeletalMeshExporter(const FObjectInitializer& ObjectInitializer = FObjectInitializer());
 
-	//~ Begin UExporter Interface
-	virtual bool ExportBinary(UObject* Object, const TCHAR* Type, FArchive& Archive, FFeedbackContext* Warn, int32 FileIndex = 0, uint32 PortFlags = 0) override;
-	//~ End UExporter Interface
+	bool Add(FGLTFContainerBuilder& Builder, const UObject* Object) override;
 };
diff --git a/Engine/Plugins/Enterprise/GLTFExporter/MigratedFromGit/Source/GLTFExporter/Public/Exporters/GLTFStaticMeshExporter.h b/Engine/Plugins/Enterprise/GLTFExporter/MigratedFromGit/Source/GLTFExporter/Public/Exporters/GLTFStaticMeshExporter.h
--- a/Engine/Plugins/Enterprise/GLTFExporter/MigratedFromGit/Source/GLTFExporter/Public/Exporters/GLTFStaticMeshExporter.h
+++ b/Engine/Plugins/Enterprise/GLTFExporter/MigratedFromGit/Source/GLTFExporter/Public/Exporters/GLTFStaticMeshExporter.h
@@ -6,7 +6,7 @@
 #include "GLTFStaticMeshExporter.generated.h"
 
 UCLASS()
-class GLTFEXPORTER_API UGLTFStaticMeshExporter : public UGLTFExporter
+class GLTFEXPORTER_API UGLTFStaticMeshExporter final : public UGLTFExporter
 {
 public:
 
@@ -14,7 +14,5 @@ public:
 
 	UGLTFStaticMeshExporter(const FObjectInitializer& ObjectInitializer = FObjectInitializer());
 
-	//~ Begin UExporter Interface
-	virtual bool ExportBinary(UObject* Object, const TCHAR* Type, FArchive& Archive, FFeedbackContext* Warn, int32 FileIndex = 0, uint32 PortFlags = 0) override;
-	//~ End UExporter Interface
+	bool Add(FGLTFContainerBuilder& Builder, const UObject* Object) override;
 };
