[GLTFExporter] refactor: add 'explicit' specifier to make constructors non-converting

commit bc029ff842544a4c9d2710b85c1b953e8135fc78
Date: Sun, 17 May 2020 20:27:04 +0200

    refactor: add 'explicit' specifier to make constructors non-converting
    (cherry picked from commit 68caf829cced4d20f8e7a48955bd5cd446deb8de​)

#skipci

diff --git a/Engine/Plugins/Enterprise/GLTFExporter/MigratedFromGit/Source/GLTFExporter/Public/Exporters/GLTFAnimSequenceExporter.h b/Engine/Plugins/Enterprise/GLTFExporter/MigratedFromGit/Source/GLTFExporter/Public/Exporters/GLTFAnimSequenceExporter.h
--- a/Engine/Plugins/Enterprise/GLTFExporter/MigratedFromGit/Source/GLTFExporter/Public/Exporters/GLTFAnimSequenceExporter.h
+++ b/Engine/Plugins/Enterprise/GLTFExporter/MigratedFromGit/Source/GLTFExporter/Public/Exporters/GLTFAnimSequenceExporter.h
@@ -12,7 +12,7 @@ public:
 
 	GENERATED_BODY()
 
-	UGLTFAnimSequenceExporter(const FObjectInitializer& ObjectInitializer = FObjectInitializer());
+	explicit UGLTFAnimSequenceExporter(const FObjectInitializer& ObjectInitializer = FObjectInitializer());
 
 	bool Add(FGLTFContainerBuilder& Builder, const UObject* Object) override;
 };
diff --git a/Engine/Plugins/Enterprise/GLTFExporter/MigratedFromGit/Source/GLTFExporter/Public/Exporters/GLTFExporter.h b/Engine/Plugins/Enterprise/GLTFExporter/MigratedFromGit/Source/GLTFExporter/Public/Exporters/GLTFExporter.h
--- a/Engine/Plugins/Enterprise/GLTFExporter/MigratedFromGit/Source/GLTFExporter/Public/Exporters/GLTFExporter.h
+++ b/Engine/Plugins/Enterprise/GLTFExporter/MigratedFromGit/Source/GLTFExporter/Public/Exporters/GLTFExporter.h
@@ -18,7 +18,7 @@ public:
 
 	GENERATED_BODY()
 
-	UGLTFExporter(const FObjectInitializer& ObjectInitializer);
+	explicit UGLTFExporter(const FObjectInitializer& ObjectInitializer);
 
 	virtual bool ExportBinary(UObject* Object, const TCHAR* Type, FArchive& Archive, FFeedbackContext* Warn, int32 FileIndex, uint32 PortFlags) override;
 
diff --git a/Engine/Plugins/Enterprise/GLTFExporter/MigratedFromGit/Source/GLTFExporter/Public/Exporters/GLTFLevelExporter.h b/Engine/Plugins/Enterprise/GLTFExporter/MigratedFromGit/Source/GLTFExporter/Public/Exporters/GLTFLevelExporter.h
--- a/Engine/Plugins/Enterprise/GLTFExporter/MigratedFromGit/Source/GLTFExporter/Public/Exporters/GLTFLevelExporter.h
+++ b/Engine/Plugins/Enterprise/GLTFExporter/MigratedFromGit/Source/GLTFExporter/Public/Exporters/GLTFLevelExporter.h
@@ -12,7 +12,7 @@ public:
 
 	GENERATED_BODY()
 
-	UGLTFLevelExporter(const FObjectInitializer& ObjectInitializer = FObjectInitializer());
+	explicit UGLTFLevelExporter(const FObjectInitializer& ObjectInitializer = FObjectInitializer());
 
 	bool Add(FGLTFContainerBuilder& Builder, const UObject* Object) override;
 };
diff --git a/Engine/Plugins/Enterprise/GLTFExporter/MigratedFromGit/Source/GLTFExporter/Public/Exporters/GLTFSkeletalMeshExporter.h b/Engine/Plugins/Enterprise/GLTFExporter/MigratedFromGit/Source/GLTFExporter/Public/Exporters/GLTFSkeletalMeshExporter.h
--- a/Engine/Plugins/Enterprise/GLTFExporter/MigratedFromGit/Source/GLTFExporter/Public/Exporters/GLTFSkeletalMeshExporter.h
+++ b/Engine/Plugins/Enterprise/GLTFExporter/MigratedFromGit/Source/GLTFExporter/Public/Exporters/GLTFSkeletalMeshExporter.h
@@ -12,7 +12,7 @@ public:
 
 	GENERATED_BODY()
 
-	UGLTFSkeletalMeshExporter(const FObjectInitializer& ObjectInitializer = FObjectInitializer());
+	explicit UGLTFSkeletalMeshExporter(const FObjectInitializer& ObjectInitializer = FObjectInitializer());
 
 	bool Add(FGLTFContainerBuilder& Builder, const UObject* Object) override;
 };
diff --git a/Engine/Plugins/Enterprise/GLTFExporter/MigratedFromGit/Source/GLTFExporter/Public/Exporters/GLTFStaticMeshExporter.h b/Engine/Plugins/Enterprise/GLTFExporter/MigratedFromGit/Source/GLTFExporter/Public/Exporters/GLTFStaticMeshExporter.h
--- a/Engine/Plugins/Enterprise/GLTFExporter/MigratedFromGit/Source/GLTFExporter/Public/Exporters/GLTFStaticMeshExporter.h
+++ b/Engine/Plugins/Enterprise/GLTFExporter/MigratedFromGit/Source/GLTFExporter/Public/Exporters/GLTFStaticMeshExporter.h
@@ -12,7 +12,7 @@ public:
 
 	GENERATED_BODY()
 
-	UGLTFStaticMeshExporter(const FObjectInitializer& ObjectInitializer = FObjectInitializer());
+	explicit UGLTFStaticMeshExporter(const FObjectInitializer& ObjectInitializer = FObjectInitializer());
 
 	bool Add(FGLTFContainerBuilder& Builder, const UObject* Object) override;
 };
