[GLTFExporter] refactor: add final to plugin module implementation

commit edf6cecb73c3e8344c13b0c437ee9540d7b83688
Date: Sun, 17 May 2020 20:28:21 +0200

    refactor: add final to plugin module implementation
    (cherry picked from commit f7abf7b67e6053d82d436b33e594f22ea2677f75​)

#skipci

diff --git a/Engine/Plugins/Enterprise/GLTFExporter/MigratedFromGit/Source/GLTFExporter/Private/GLTFExporterModule.cpp b/Engine/Plugins/Enterprise/GLTFExporter/MigratedFromGit/Source/GLTFExporter/Private/GLTFExporterModule.cpp
--- a/Engine/Plugins/Enterprise/GLTFExporter/MigratedFromGit/Source/GLTFExporter/Private/GLTFExporterModule.cpp
+++ b/Engine/Plugins/Enterprise/GLTFExporter/MigratedFromGit/Source/GLTFExporter/Private/GLTFExporterModule.cpp
@@ -10,7 +10,7 @@ DEFINE_LOG_CATEGORY(LogGLTFExporter);
 /**
  * glTF Exporter module implementation (private)
  */
-class FGLTFExporterModule : public IGLTFExporterModule
+class FGLTFExporterModule final : public IGLTFExporterModule
 {
 
 public:
