[GLTFExporter] refactor: add 'const' specifier to some local variables

commit fce60663dafd1008eb3a5a7ebf95c38c433f85ce
Date: Sun, 17 May 2020 20:31:11 +0200

    refactor: add 'const' specifier to some local variables
    (cherry picked from commit 79d300f7fb1f8a01f45a6434b193f25c786ad853​)

#skipci

diff --git a/Engine/Plugins/Enterprise/GLTFExporter/MigratedFromGit/Source/GLTFExporter/Private/Converters/GLTFLevelConverters.cpp b/Engine/Plugins/Enterprise/GLTFExporter/MigratedFromGit/Source/GLTFExporter/Private/Converters/GLTFLevelConverters.cpp
--- a/Engine/Plugins/Enterprise/GLTFExporter/MigratedFromGit/Source/GLTFExporter/Private/Converters/GLTFLevelConverters.cpp
+++ b/Engine/Plugins/Enterprise/GLTFExporter/MigratedFromGit/Source/GLTFExporter/Private/Converters/GLTFLevelConverters.cpp
@@ -7,8 +7,8 @@
 FGLTFJsonNodeIndex FGLTFSceneComponentConverter::Add(FGLTFConvertBuilder& Builder, const FString& Name, TTuple<const USceneComponent*, bool, bool> Params)
 {
 	const USceneComponent* SceneComponent = Params.Get<0>();
-	bool bSelectedOnly = Params.Get<1>();
-	bool bRootNode = Params.Get<2>();
+	const bool bSelectedOnly = Params.Get<1>();
+	const bool bRootNode = Params.Get<2>();
 
 	const AActor* Owner = SceneComponent->GetOwner();
 	const bool bIsRootComponent = Owner != nullptr && Owner->GetRootComponent() == SceneComponent;
@@ -58,7 +58,7 @@ FGLTFJsonNodeIndex FGLTFSceneComponentConverter::Add(FGLTFConvertBuilder& Builde
 FGLTFJsonSceneIndex FGLTFLevelConverter::Add(FGLTFConvertBuilder& Builder, const FString& Name, TTuple<const ULevel*, bool> Params)
 {
 	const ULevel* Level = Params.Get<0>();
-	bool bSelectedOnly = Params.Get<1>();
+	const bool bSelectedOnly = Params.Get<1>();
 
 	FGLTFJsonScene Scene;
 	Scene.Name = Name;
diff --git a/Engine/Plugins/Enterprise/GLTFExporter/MigratedFromGit/Source/GLTFExporter/Private/Exporters/GLTFStaticMeshExporter.cpp b/Engine/Plugins/Enterprise/GLTFExporter/MigratedFromGit/Source/GLTFExporter/Private/Exporters/GLTFStaticMeshExporter.cpp
--- a/Engine/Plugins/Enterprise/GLTFExporter/MigratedFromGit/Source/GLTFExporter/Private/Exporters/GLTFStaticMeshExporter.cpp
+++ b/Engine/Plugins/Enterprise/GLTFExporter/MigratedFromGit/Source/GLTFExporter/Private/Exporters/GLTFStaticMeshExporter.cpp
@@ -13,15 +13,15 @@ UGLTFStaticMeshExporter::UGLTFStaticMeshExporter(const FObjectInitializer& Objec
 bool UGLTFStaticMeshExporter::Add(FGLTFContainerBuilder& Builder, const UObject* Object)
 {
 	const UStaticMesh* StaticMesh = CastChecked<UStaticMesh>(Object);
-	FGLTFJsonMeshIndex MeshIndex = Builder.GetOrAddMesh(StaticMesh);
+	const FGLTFJsonMeshIndex MeshIndex = Builder.GetOrAddMesh(StaticMesh);
 
 	FGLTFJsonNode Node;
 	Node.Mesh = MeshIndex;
-	FGLTFJsonNodeIndex NodeIndex = Builder.AddNode(Node);
+	const FGLTFJsonNodeIndex NodeIndex = Builder.AddNode(Node);
 
 	FGLTFJsonScene Scene;
 	Scene.Nodes.Add(NodeIndex);
-	FGLTFJsonSceneIndex SceneIndex = Builder.AddScene(Scene);
+	const FGLTFJsonSceneIndex SceneIndex = Builder.AddScene(Scene);
 
 	Builder.DefaultScene = SceneIndex;
 	return true;
