[GLTFExporter] refactor: remove redundant access specifiers

commit a7895dc34e0bb1f230eae9cef645264b7fa07482
Date: Sun, 17 May 2020 20:31:54 +0200

    refactor: remove redundant access specifiers
    (cherry picked from commit d83c4215f022df2b7affd7585843e2990b6a4a5d​)

#skipci

diff --git a/Engine/Plugins/Enterprise/GLTFExporter/MigratedFromGit/Source/GLTFExporter/Private/UI/GLTFExportOptionsWindow.h b/Engine/Plugins/Enterprise/GLTFExporter/MigratedFromGit/Source/GLTFExporter/Private/UI/GLTFExportOptionsWindow.h
--- a/Engine/Plugins/Enterprise/GLTFExporter/MigratedFromGit/Source/GLTFExporter/Private/UI/GLTFExportOptionsWindow.h
+++ b/Engine/Plugins/Enterprise/GLTFExporter/MigratedFromGit/Source/GLTFExporter/Private/UI/GLTFExportOptionsWindow.h
@@ -14,6 +14,7 @@ class SButton;
 class GLTFEXPORTER_API SGLTFExportOptionsWindow : public SCompoundWidget
 {
 public:
+
 	SLATE_BEGIN_ARGS(SGLTFExportOptionsWindow)
 		: _ExportOptions(nullptr)
 		, _BatchMode()
@@ -25,7 +26,6 @@ public:
 		SLATE_ARGUMENT( bool, BatchMode )
 	SLATE_END_ARGS()
 
-public:
 	void Construct(const FArguments& InArgs);
 
 	virtual bool SupportsKeyboardFocus() const override { return true; }
@@ -44,13 +44,12 @@ public:
 
 	SGLTFExportOptionsWindow();
 
-public:
 	static void ShowDialog(UGLTFExportOptions* ExportOptions, const FString& FullPath, bool bBatchMode, bool& bOutOperationCanceled, bool& bOutExportAll);
 
 private:
+
 	FReply OnResetToDefaultClick() const;
 
-private:
 	UGLTFExportOptions* ExportOptions;
 	TSharedPtr<class IDetailsView> DetailsView;
 	TWeakPtr<SWindow> WidgetWindow;
