[GLTFExporter] refactor: add GLTFBuilderUtility to simplify name generation

commit 36f7658c6310addfbfb7eabc5f9c715fa15e8a10
Date: Sun, 17 May 2020 20:53:13 +0200

    refactor: add GLTFBuilderUtility to simplify name generation
    (cherry picked from commit 169d64846c9c51743d00cf88250271a8b5d176fc​)

#skipci

diff --git a/Engine/Plugins/Enterprise/GLTFExporter/MigratedFromGit/Source/GLTFExporter/Private/Builders/GLTFBuilderUtility.h b/Engine/Plugins/Enterprise/GLTFExporter/MigratedFromGit/Source/GLTFExporter/Private/Builders/GLTFBuilderUtility.h
new file mode 100644
--- /dev/null
+++ b/Engine/Plugins/Enterprise/GLTFExporter/MigratedFromGit/Source/GLTFExporter/Private/Builders/GLTFBuilderUtility.h
@@ -0,0 +1,16 @@
+// Copyright Epic Games, Inc. All Rights Reserved.
+
+#pragma once
+
+#include "Engine.h"
+
+struct FGLTFBuilderUtility
+{
+	static FString GetMeshName(const UStaticMesh* StaticMesh, int32 LODIndex)
+	{
+		FString Name;
+		StaticMesh->GetName(Name);
+		if (LODIndex != 0) Name += TEXT("_LOD") + FString::FromInt(LODIndex);
+		return Name;
+	}
+};
diff --git a/Engine/Plugins/Enterprise/GLTFExporter/MigratedFromGit/Source/GLTFExporter/Private/Builders/GLTFConvertBuilder.cpp b/Engine/Plugins/Enterprise/GLTFExporter/MigratedFromGit/Source/GLTFExporter/Private/Builders/GLTFConvertBuilder.cpp
--- a/Engine/Plugins/Enterprise/GLTFExporter/MigratedFromGit/Source/GLTFExporter/Private/Builders/GLTFConvertBuilder.cpp
+++ b/Engine/Plugins/Enterprise/GLTFExporter/MigratedFromGit/Source/GLTFExporter/Private/Builders/GLTFConvertBuilder.cpp
@@ -1,6 +1,7 @@
 // Copyright Epic Games, Inc. All Rights Reserved.
 
 #include "Builders/GLTFConvertBuilder.h"
+#include "GLTFBuilderUtility.h"
 
 FGLTFJsonAccessorIndex FGLTFConvertBuilder::GetOrAddPositionAccessor(const FPositionVertexBuffer* VertexBuffer, const FString& DesiredName)
 {
@@ -50,14 +51,8 @@ FGLTFJsonMeshIndex FGLTFConvertBuilder::GetOrAddMesh(const FStaticMeshLODResourc
 FGLTFJsonMeshIndex FGLTFConvertBuilder::GetOrAddMesh(const UStaticMesh* StaticMesh, int32 LODIndex, const FColorVertexBuffer* OverrideVertexColors, const FString& DesiredName)
 {
 	const FStaticMeshLODResources* StaticMeshLOD = &StaticMesh->GetLODForExport(LODIndex);
-	FString Name = DesiredName.IsEmpty() ? StaticMesh->GetName() : DesiredName;
 
-	if (LODIndex != 0)
-	{
-		Name += TEXT("_LOD") + FString::FromInt(LODIndex);
-	}
-
-	return GetOrAddMesh(StaticMeshLOD, OverrideVertexColors, Name);
+	return GetOrAddMesh(StaticMeshLOD, OverrideVertexColors, DesiredName.IsEmpty() ? FGLTFBuilderUtility::GetMeshName(StaticMesh, LODIndex) : DesiredName);
 }
 
 FGLTFJsonMeshIndex FGLTFConvertBuilder::GetOrAddMesh(const UStaticMeshComponent* StaticMeshComponent, const FString& DesiredName)
