[GLTFExporter] refactor: share same code path for converting both UV sets

commit eefb1ad0c30cae219a718808aaa0b674ce97f965
Date: Sun, 17 May 2020 21:47:37 +0200

    refactor: share same code path for converting both UV sets
    (cherry picked from commit aed371e367d764db8f7edd846472e5c2bbbc082d​)

#skipci

diff --git a/Engine/Plugins/Enterprise/GLTFExporter/MigratedFromGit/Source/GLTFExporter/Private/Builders/GLTFConvertBuilder.cpp b/Engine/Plugins/Enterprise/GLTFExporter/MigratedFromGit/Source/GLTFExporter/Private/Builders/GLTFConvertBuilder.cpp
--- a/Engine/Plugins/Enterprise/GLTFExporter/MigratedFromGit/Source/GLTFExporter/Private/Builders/GLTFConvertBuilder.cpp
+++ b/Engine/Plugins/Enterprise/GLTFExporter/MigratedFromGit/Source/GLTFExporter/Private/Builders/GLTFConvertBuilder.cpp
@@ -23,14 +23,9 @@ FGLTFJsonAccessorIndex FGLTFConvertBuilder::GetOrAddTangentAccessor(const FStati
 	return StaticMeshTangentVertexBufferConverter.GetOrAdd(*this, DesiredName, VertexBuffer);
 }
 
-FGLTFJsonAccessorIndex FGLTFConvertBuilder::GetOrAddUV0Accessor(const FStaticMeshVertexBuffer* VertexBuffer, const FString& DesiredName)
+FGLTFJsonAccessorIndex FGLTFConvertBuilder::GetOrAddUVAccessor(const FStaticMeshVertexBuffer* VertexBuffer, int32 UVIndex, const FString& DesiredName)
 {
-	return StaticMeshUV0VertexBufferConverter.GetOrAdd(*this, DesiredName, VertexBuffer);
-}
-
-FGLTFJsonAccessorIndex FGLTFConvertBuilder::GetOrAddUV1Accessor(const FStaticMeshVertexBuffer* VertexBuffer, const FString& DesiredName)
-{
-	return StaticMeshUV1VertexBufferConverter.GetOrAdd(*this, DesiredName, VertexBuffer);
+	return StaticMeshUVVertexBufferConverter.GetOrAdd(*this, DesiredName, VertexBuffer, UVIndex);
 }
 
 FGLTFJsonBufferViewIndex FGLTFConvertBuilder::GetOrAddIndexBufferView(const FRawStaticIndexBuffer* IndexBuffer, const FString& DesiredName)
@@ -43,16 +38,9 @@ FGLTFJsonAccessorIndex FGLTFConvertBuilder::GetOrAddIndexAccessor(const FStaticM
 	return StaticMeshSectionConverter.GetOrAdd(*this, DesiredName, MeshSection, IndexBuffer);
 }
 
-FGLTFJsonMeshIndex FGLTFConvertBuilder::GetOrAddMesh(const FStaticMeshLODResources* StaticMeshLOD, const FColorVertexBuffer* OverrideVertexColors, const FString& DesiredName)
-{
-	return StaticMeshConverter.GetOrAdd(*this, DesiredName, StaticMeshLOD, OverrideVertexColors);
-}
-
 FGLTFJsonMeshIndex FGLTFConvertBuilder::GetOrAddMesh(const UStaticMesh* StaticMesh, int32 LODIndex, const FColorVertexBuffer* OverrideVertexColors, const FString& DesiredName)
 {
-	const FStaticMeshLODResources* StaticMeshLOD = &StaticMesh->GetLODForExport(LODIndex);
-
-	return GetOrAddMesh(StaticMeshLOD, OverrideVertexColors, DesiredName.IsEmpty() ? FGLTFBuilderUtility::GetMeshName(StaticMesh, LODIndex) : DesiredName);
+	return StaticMeshConverter.GetOrAdd(*this, DesiredName.IsEmpty() ? FGLTFBuilderUtility::GetMeshName(StaticMesh, LODIndex) : DesiredName, StaticMesh, LODIndex, OverrideVertexColors);
 }
 
 FGLTFJsonMeshIndex FGLTFConvertBuilder::GetOrAddMesh(const UStaticMeshComponent* StaticMeshComponent, const FString& DesiredName)
diff --git a/Engine/Plugins/Enterprise/GLTFExporter/MigratedFromGit/Source/GLTFExporter/Private/Builders/GLTFConvertBuilder.h b/Engine/Plugins/Enterprise/GLTFExporter/MigratedFromGit/Source/GLTFExporter/Private/Builders/GLTFConvertBuilder.h
--- a/Engine/Plugins/Enterprise/GLTFExporter/MigratedFromGit/Source/GLTFExporter/Private/Builders/GLTFConvertBuilder.h
+++ b/Engine/Plugins/Enterprise/GLTFExporter/MigratedFromGit/Source/GLTFExporter/Private/Builders/GLTFConvertBuilder.h
@@ -13,18 +13,17 @@ public:
 	FGLTFJsonAccessorIndex GetOrAddPositionAccessor(const FPositionVertexBuffer* VertexBuffer, const FString& DesiredName = TEXT(""));
 	FGLTFJsonAccessorIndex GetOrAddNormalAccessor(const FStaticMeshVertexBuffer* VertexBuffer, const FString& DesiredName = TEXT(""));
 	FGLTFJsonAccessorIndex GetOrAddTangentAccessor(const FStaticMeshVertexBuffer* VertexBuffer, const FString& DesiredName = TEXT(""));
-	FGLTFJsonAccessorIndex GetOrAddUV0Accessor(const FStaticMeshVertexBuffer* VertexBuffer, const FString& DesiredName = TEXT(""));
-	FGLTFJsonAccessorIndex GetOrAddUV1Accessor(const FStaticMeshVertexBuffer* VertexBuffer, const FString& DesiredName = TEXT(""));
+	FGLTFJsonAccessorIndex GetOrAddUVAccessor(const FStaticMeshVertexBuffer* VertexBuffer, int32 UVIndex, const FString& DesiredName = TEXT(""));
 	FGLTFJsonAccessorIndex GetOrAddColorAccessor(const FColorVertexBuffer* VertexBuffer, const FString& DesiredName = TEXT(""));
 
 	FGLTFJsonBufferViewIndex GetOrAddIndexBufferView(const FRawStaticIndexBuffer* IndexBuffer, const FString& DesiredName = TEXT(""));
 	FGLTFJsonAccessorIndex GetOrAddIndexAccessor(const FStaticMeshSection* MeshSection, const FRawStaticIndexBuffer* IndexBuffer, const FString& DesiredName = TEXT(""));
 
-	FGLTFJsonMeshIndex GetOrAddMesh(const FStaticMeshLODResources* StaticMeshLOD, const FColorVertexBuffer* OverrideVertexColors = nullptr, const FString& DesiredName = TEXT(""));
 	FGLTFJsonMeshIndex GetOrAddMesh(const UStaticMesh* StaticMesh, int32 LODIndex = 0, const FColorVertexBuffer* OverrideVertexColors = nullptr, const FString& DesiredName = TEXT(""));
 	FGLTFJsonMeshIndex GetOrAddMesh(const UStaticMeshComponent* StaticMeshComponent, const FString& DesiredName = TEXT(""));
 
 	FGLTFJsonNodeIndex GetOrAddNode(const USceneComponent* SceneComponent, bool bSelectedOnly, bool bRootNode = false, const FString& DesiredName = TEXT(""));
+
 	FGLTFJsonSceneIndex GetOrAddScene(const ULevel* Level, bool bSelectedOnly, const FString& DesiredName = TEXT(""));
 	FGLTFJsonSceneIndex GetOrAddScene(const UWorld* World, bool bSelectedOnly, const FString& DesiredName = TEXT(""));
 
@@ -34,8 +33,7 @@ private:
 	FGLTFColorVertexBufferConverter ColorVertexBufferConverter;
 	FGLTFStaticMeshNormalVertexBufferConverter StaticMeshNormalVertexBufferConverter;
 	FGLTFStaticMeshTangentVertexBufferConverter StaticMeshTangentVertexBufferConverter;
-	FGLTFStaticMeshUV0VertexBufferConverter StaticMeshUV0VertexBufferConverter;
-	FGLTFStaticMeshUV1VertexBufferConverter StaticMeshUV1VertexBufferConverter;
+	FGLTFStaticMeshUVVertexBufferConverter StaticMeshUVVertexBufferConverter;
 	FGLTFStaticMeshIndexBufferConverter StaticMeshIndexBufferConverter;
 	FGLTFStaticMeshSectionConverter StaticMeshSectionConverter;
 	FGLTFStaticMeshConverter StaticMeshConverter;
diff --git a/Engine/Plugins/Enterprise/GLTFExporter/MigratedFromGit/Source/GLTFExporter/Private/Converters/GLTFStaticMeshConverters.cpp b/Engine/Plugins/Enterprise/GLTFExporter/MigratedFromGit/Source/GLTFExporter/Private/Converters/GLTFStaticMeshConverters.cpp
--- a/Engine/Plugins/Enterprise/GLTFExporter/MigratedFromGit/Source/GLTFExporter/Private/Converters/GLTFStaticMeshConverters.cpp
+++ b/Engine/Plugins/Enterprise/GLTFExporter/MigratedFromGit/Source/GLTFExporter/Private/Converters/GLTFStaticMeshConverters.cpp
@@ -126,51 +126,28 @@ FGLTFJsonAccessorIndex FGLTFStaticMeshTangentVertexBufferConverter::Add(FGLTFCon
 	return Builder.AddAccessor(JsonAccessor);
 }
 
-FGLTFJsonAccessorIndex FGLTFStaticMeshUV0VertexBufferConverter::Add(FGLTFConvertBuilder& Builder, const FString& Name, const FStaticMeshVertexBuffer* VertexBuffer)
+FGLTFJsonAccessorIndex FGLTFStaticMeshUVVertexBufferConverter::Add(FGLTFConvertBuilder& Builder, const FString& Name, TTuple<const FStaticMeshVertexBuffer*, int32> Params)
 {
-	const uint32 VertexCount = VertexBuffer->GetNumVertices();
-	if (VertexCount == 0 || VertexBuffer->GetNumTexCoords() < 1)
-	{
-		return FGLTFJsonAccessorIndex(INDEX_NONE);
-	}
-
-	TArray<FVector2D> UV0s;
-	UV0s.AddUninitialized(VertexCount);
-
-	for (uint32 VertexIndex = 0; VertexIndex < VertexCount; ++VertexIndex)
-	{
-		UV0s[VertexIndex] = VertexBuffer->GetVertexUV(VertexIndex, 0);
-	}
-
-	FGLTFJsonAccessor JsonAccessor;
-	JsonAccessor.Name = Name;
-	JsonAccessor.BufferView = Builder.AddBufferView(UV0s, Name);
-	JsonAccessor.ComponentType = EGLTFJsonComponentType::F32;
-	JsonAccessor.Count = VertexCount;
-	JsonAccessor.Type = EGLTFJsonAccessorType::Vec2;
+	const FStaticMeshVertexBuffer* VertexBuffer = Params.Get<0>();
+	const int32 UVIndex = Params.Get<1>();
 
-	return Builder.AddAccessor(JsonAccessor);
-}
-
-FGLTFJsonAccessorIndex FGLTFStaticMeshUV1VertexBufferConverter::Add(FGLTFConvertBuilder& Builder, const FString& Name, const FStaticMeshVertexBuffer* VertexBuffer)
-{
 	const uint32 VertexCount = VertexBuffer->GetNumVertices();
-	if (VertexCount == 0 || VertexBuffer->GetNumTexCoords() < 2)
+	if (VertexCount == 0 || UVIndex < 0 || VertexBuffer->GetNumTexCoords() <= static_cast<uint32>(UVIndex))
 	{
 		return FGLTFJsonAccessorIndex(INDEX_NONE);
 	}
 
-	TArray<FVector2D> UV1s;
-	UV1s.AddUninitialized(VertexCount);
+	TArray<FVector2D> UVs;
+	UVs.AddUninitialized(VertexCount);
 
 	for (uint32 VertexIndex = 0; VertexIndex < VertexCount; ++VertexIndex)
 	{
-		UV1s[VertexIndex] = VertexBuffer->GetVertexUV(VertexIndex, 1);
+		UVs[VertexIndex] = VertexBuffer->GetVertexUV(VertexIndex, UVIndex);
 	}
 
 	FGLTFJsonAccessor JsonAccessor;
 	JsonAccessor.Name = Name;
-	JsonAccessor.BufferView = Builder.AddBufferView(UV1s, Name);
+	JsonAccessor.BufferView = Builder.AddBufferView(UVs, Name);
 	JsonAccessor.ComponentType = EGLTFJsonComponentType::F32;
 	JsonAccessor.Count = VertexCount;
 	JsonAccessor.Type = EGLTFJsonAccessorType::Vec2;
@@ -212,30 +189,35 @@ FGLTFJsonAccessorIndex FGLTFStaticMeshSectionConverter::Add(FGLTFConvertBuilder&
 	return Builder.AddAccessor(JsonAccessor);
 }
 
-FGLTFJsonMeshIndex FGLTFStaticMeshConverter::Add(FGLTFConvertBuilder& Builder, const FString& Name, TTuple<const FStaticMeshLODResources*, const FColorVertexBuffer*> Params)
+FGLTFJsonMeshIndex FGLTFStaticMeshConverter::Add(FGLTFConvertBuilder& Builder, const FString& Name, TTuple<const UStaticMesh*, int32, const FColorVertexBuffer*> Params)
 {
-	const FStaticMeshLODResources* StaticMeshLOD = Params.Get<0>();
-	const FColorVertexBuffer* OverrideVertexColors = Params.Get<1>();
+	const UStaticMesh* StaticMesh = Params.Get<0>();
+	const int32 LODIndex = Params.Get<1>();
+	const FColorVertexBuffer* OverrideVertexColors = Params.Get<2>();
+
+	const int32 PrimaryUVIndex = 0; // TODO: make this configurable?
+	const int32 LightmapUVIndex = StaticMesh->LightMapCoordinateIndex != PrimaryUVIndex ? StaticMesh->LightMapCoordinateIndex : INDEX_NONE;
+	const FStaticMeshLODResources& LODResources = StaticMesh->GetLODForExport(LODIndex);
 
 	FGLTFJsonMesh JsonMesh;
 	JsonMesh.Name = Name;
 
-	const FPositionVertexBuffer* PositionBuffer = &StaticMeshLOD->VertexBuffers.PositionVertexBuffer;
-	const FStaticMeshVertexBuffer* VertexBuffer = &StaticMeshLOD->VertexBuffers.StaticMeshVertexBuffer;
-	const FColorVertexBuffer* ColorBuffer = OverrideVertexColors != nullptr ? OverrideVertexColors : &StaticMeshLOD->VertexBuffers.ColorVertexBuffer;
+	const FPositionVertexBuffer* PositionBuffer = &LODResources.VertexBuffers.PositionVertexBuffer;
+	const FStaticMeshVertexBuffer* VertexBuffer = &LODResources.VertexBuffers.StaticMeshVertexBuffer;
+	const FColorVertexBuffer* ColorBuffer = OverrideVertexColors != nullptr ? OverrideVertexColors : &LODResources.VertexBuffers.ColorVertexBuffer;
 
 	FGLTFJsonAttributes JsonAttributes;
-	JsonAttributes.Position = Builder.GetOrAddPositionAccessor(PositionBuffer, Name.IsEmpty() ? Name : Name + TEXT("_Positions"));
-	JsonAttributes.Normal = Builder.GetOrAddNormalAccessor(VertexBuffer, Name.IsEmpty() ? Name : Name + TEXT("_Normals"));
-	JsonAttributes.Tangent = Builder.GetOrAddTangentAccessor(VertexBuffer, Name.IsEmpty() ? Name : Name + TEXT("_Tangents"));
-	JsonAttributes.TexCoord0 = Builder.GetOrAddUV0Accessor(VertexBuffer, Name.IsEmpty() ? Name : Name + TEXT("_UV0s"));
-	JsonAttributes.TexCoord1 = Builder.GetOrAddUV1Accessor(VertexBuffer, Name.IsEmpty() ? Name : Name + TEXT("_UV1s"));
-	JsonAttributes.Color0 = Builder.GetOrAddColorAccessor(ColorBuffer, Name.IsEmpty() ? Name : Name + TEXT("_Colors"));
+	JsonAttributes.Position = Builder.GetOrAddPositionAccessor(PositionBuffer, Name + TEXT("_Positions"));
+	JsonAttributes.Color0 = Builder.GetOrAddColorAccessor(ColorBuffer, Name + TEXT("_Colors"));
+	JsonAttributes.Normal = Builder.GetOrAddNormalAccessor(VertexBuffer, Name + TEXT("_Normals"));
+	JsonAttributes.Tangent = Builder.GetOrAddTangentAccessor(VertexBuffer, Name + TEXT("_Tangents"));
+	JsonAttributes.TexCoord0 = Builder.GetOrAddUVAccessor(VertexBuffer, PrimaryUVIndex, Name + TEXT("_UV") + FString::FromInt(PrimaryUVIndex) + TEXT("s"));
+	JsonAttributes.TexCoord1 = Builder.GetOrAddUVAccessor(VertexBuffer, LightmapUVIndex, Name + TEXT("_UV") + FString::FromInt(LightmapUVIndex) + TEXT("s"));
 
-	const FRawStaticIndexBuffer* IndexBuffer = &StaticMeshLOD->IndexBuffer;
-	Builder.GetOrAddIndexBufferView(IndexBuffer, Name.IsEmpty() ? Name : Name + TEXT("_Indices"));
+	const FRawStaticIndexBuffer* IndexBuffer = &LODResources.IndexBuffer;
+	Builder.GetOrAddIndexBufferView(IndexBuffer, Name + TEXT("_Indices"));
 
-	const int32 SectionCount = StaticMeshLOD->Sections.Num();
+	const int32 SectionCount = LODResources.Sections.Num();
 	JsonMesh.Primitives.AddDefaulted(SectionCount);
 
 	for (int32 SectionIndex = 0; SectionIndex < SectionCount; ++SectionIndex)
@@ -243,8 +225,8 @@ FGLTFJsonMeshIndex FGLTFStaticMeshConverter::Add(FGLTFConvertBuilder& Builder, c
 		FGLTFJsonPrimitive& JsonPrimitive = JsonMesh.Primitives[SectionIndex];
 		JsonPrimitive.Attributes = JsonAttributes;
 
-		JsonPrimitive.Indices = Builder.GetOrAddIndexAccessor(&StaticMeshLOD->Sections[SectionIndex], IndexBuffer,
-            Name.IsEmpty() ? Name : JsonMesh.Name + (SectionCount != 1 ? TEXT("_Indices_Section") + FString::FromInt(SectionIndex) : TEXT("_Indices")));
+		JsonPrimitive.Indices = Builder.GetOrAddIndexAccessor(&LODResources.Sections[SectionIndex], IndexBuffer,
+            Name + (SectionCount != 1 ? TEXT("_Indices_Section") + FString::FromInt(SectionIndex) : TEXT("_Indices")));
 	}
 
 	return Builder.AddMesh(JsonMesh);
diff --git a/Engine/Plugins/Enterprise/GLTFExporter/MigratedFromGit/Source/GLTFExporter/Private/Converters/GLTFStaticMeshConverters.h b/Engine/Plugins/Enterprise/GLTFExporter/MigratedFromGit/Source/GLTFExporter/Private/Converters/GLTFStaticMeshConverters.h
--- a/Engine/Plugins/Enterprise/GLTFExporter/MigratedFromGit/Source/GLTFExporter/Private/Converters/GLTFStaticMeshConverters.h
+++ b/Engine/Plugins/Enterprise/GLTFExporter/MigratedFromGit/Source/GLTFExporter/Private/Converters/GLTFStaticMeshConverters.h
@@ -26,14 +26,9 @@ class FGLTFStaticMeshTangentVertexBufferConverter final : public TGLTFConverter<
 	FGLTFJsonAccessorIndex Add(FGLTFConvertBuilder& Builder, const FString& Name, const FStaticMeshVertexBuffer* VertexBuffer) override;
 };
 
-class FGLTFStaticMeshUV0VertexBufferConverter final : public TGLTFConverter<FGLTFJsonAccessorIndex, const FStaticMeshVertexBuffer*>
+class FGLTFStaticMeshUVVertexBufferConverter final : public TGLTFConverter<FGLTFJsonAccessorIndex,TTuple<const FStaticMeshVertexBuffer*, int32>>
 {
-	FGLTFJsonAccessorIndex Add(FGLTFConvertBuilder& Builder, const FString& Name, const FStaticMeshVertexBuffer* VertexBuffer) override;
-};
-
-class FGLTFStaticMeshUV1VertexBufferConverter final : public TGLTFConverter<FGLTFJsonAccessorIndex, const FStaticMeshVertexBuffer*>
-{
-	FGLTFJsonAccessorIndex Add(FGLTFConvertBuilder& Builder, const FString& Name, const FStaticMeshVertexBuffer* VertexBuffer) override;
+	FGLTFJsonAccessorIndex Add(FGLTFConvertBuilder& Builder, const FString& Name, TTuple<const FStaticMeshVertexBuffer*, int32> Params) override;
 };
 
 class FGLTFStaticMeshIndexBufferConverter final : public TGLTFConverter<FGLTFJsonBufferViewIndex, const FRawStaticIndexBuffer*>
@@ -46,7 +41,7 @@ class FGLTFStaticMeshSectionConverter final : public TGLTFConverter<FGLTFJsonAcc
 	FGLTFJsonAccessorIndex Add(FGLTFConvertBuilder& Builder, const FString& Name, TTuple<const FStaticMeshSection*, const FRawStaticIndexBuffer*> Params) override;
 };
 
-class FGLTFStaticMeshConverter final : public TGLTFConverter<FGLTFJsonMeshIndex, TTuple<const FStaticMeshLODResources*, const FColorVertexBuffer*>>
+class FGLTFStaticMeshConverter final : public TGLTFConverter<FGLTFJsonMeshIndex, TTuple<const UStaticMesh*, int32, const FColorVertexBuffer*>>
 {
-	FGLTFJsonMeshIndex Add(FGLTFConvertBuilder& Builder, const FString& Name, TTuple<const FStaticMeshLODResources*, const FColorVertexBuffer*> Params) override;
+	FGLTFJsonMeshIndex Add(FGLTFConvertBuilder& Builder, const FString& Name, TTuple<const UStaticMesh*, int32, const FColorVertexBuffer*> Params) override;
 };
