[GLTFExporter] feat: don't try converting StaticMesh LOD index that doesn't exist

commit 06f0c26f3e1d809f987387a717a73c364c6a86f6
Date: Sun, 17 May 2020 22:06:42 +0200

    feat: don't try converting StaticMesh LOD index that doesn't exist
    (cherry picked from commit 4aecf6a3a98a43162a031629e83f6553af0e5864​)

#skipci

diff --git a/Engine/Plugins/Enterprise/GLTFExporter/MigratedFromGit/Source/GLTFExporter/Private/Converters/GLTFStaticMeshConverters.cpp b/Engine/Plugins/Enterprise/GLTFExporter/MigratedFromGit/Source/GLTFExporter/Private/Converters/GLTFStaticMeshConverters.cpp
--- a/Engine/Plugins/Enterprise/GLTFExporter/MigratedFromGit/Source/GLTFExporter/Private/Converters/GLTFStaticMeshConverters.cpp
+++ b/Engine/Plugins/Enterprise/GLTFExporter/MigratedFromGit/Source/GLTFExporter/Private/Converters/GLTFStaticMeshConverters.cpp
@@ -195,6 +195,11 @@ FGLTFJsonMeshIndex FGLTFStaticMeshConverter::Add(FGLTFConvertBuilder& Builder, c
 	const int32 LODIndex = Params.Get<1>();
 	const FColorVertexBuffer* OverrideVertexColors = Params.Get<2>();
 
+	if (LODIndex < 0 || StaticMesh->GetNumLODs() <= LODIndex)
+	{
+		return FGLTFJsonMeshIndex(INDEX_NONE);
+	}
+
 	const int32 PrimaryUVIndex = 0; // TODO: make this configurable?
 	const int32 LightmapUVIndex = StaticMesh->LightMapCoordinateIndex != PrimaryUVIndex ? StaticMesh->LightMapCoordinateIndex : INDEX_NONE;
 	const FStaticMeshLODResources& LODResources = StaticMesh->GetLODForExport(LODIndex);
