[GLTFExporter] feat: update project and plugin to UE 4.25

commit c0b6991296ae00eed4940bb5d86a54527739b040
Date: Tue, 12 May 2020 09:49:05 +0200

    feat: update project and plugin to UE 4.25
    (cherry picked from commit e05ad6d00d2f783ba331513b6db9c52bc4b67ec3​)
    # Conflicts from cherry-pick:
    #	Plugins/GLTFExporter/Source/GLTFExporter/Private/GLTFExportOptions.cpp

#skipci

diff --git a/Engine/Plugins/Enterprise/GLTFExporter/MigratedFromGit/Source/GLTFExporter/Private/GLTFExportOptions.cpp b/Engine/Plugins/Enterprise/GLTFExporter/MigratedFromGit/Source/GLTFExporter/Private/GLTFExportOptions.cpp
--- a/Engine/Plugins/Enterprise/GLTFExporter/MigratedFromGit/Source/GLTFExporter/Private/GLTFExportOptions.cpp
+++ b/Engine/Plugins/Enterprise/GLTFExporter/MigratedFromGit/Source/GLTFExporter/Private/GLTFExportOptions.cpp
@@ -50,7 +50,7 @@ void UGLTFExportOptions::LoadOptions()
 {
 	const int32 PortFlags = 0;
 
-	for (UProperty* Property = GetClass()->PropertyLink; Property; Property = Property->PropertyLinkNext)
+	for (FProperty* Property = GetClass()->PropertyLink; Property; Property = Property->PropertyLinkNext)
 	{
 		if (!Property->HasAnyPropertyFlags(CPF_Config))
 		{
@@ -59,8 +59,13 @@ void UGLTFExportOptions::LoadOptions()
 		FString Section = GetClass()->GetName();
 		FString Key = Property->GetName();
 
-		UArrayProperty* Array = dynamic_cast<UArrayProperty*>(Property);
-		if (Array != nullptr)
+		const bool bIsPropertyInherited = Property->GetOwnerClass() != GetClass();
+		UObject* SuperClassDefaultObject = GetClass()->GetSuperClass()->GetDefaultObject();
+
+		const FString& PropFileName = GEditorPerProjectIni;
+
+		FArrayProperty* Array = CastField<FArrayProperty>(Property);
+		if (Array)
 		{
 			FConfigSection* Sec = GConfig->GetSectionPrivate(*Section, false, true, *GEditorPerProjectIni);
 			if (Sec != nullptr)
@@ -82,7 +87,7 @@ void UGLTFExportOptions::LoadOptions()
 				else
 				{
 					int32 Index = 0;
-					const FConfigValue* ElementValue;
+					const FConfigValue* ElementValue = nullptr;
 					do
 					{
 						// Add array index number to end of key
@@ -137,7 +142,7 @@ void UGLTFExportOptions::SaveOptions()
 {
 	const int32 PortFlags = 0;
 
-	for (UProperty* Property = GetClass()->PropertyLink; Property; Property = Property->PropertyLinkNext)
+	for (FProperty* Property = GetClass()->PropertyLink; Property; Property = Property->PropertyLinkNext)
 	{
 		if (!Property->HasAnyPropertyFlags(CPF_Config))
 		{
@@ -146,8 +151,11 @@ void UGLTFExportOptions::SaveOptions()
 		FString Section = GetClass()->GetName();
 		FString Key = Property->GetName();
 
-		UArrayProperty* Array = dynamic_cast<UArrayProperty*>(Property);
-		if (Array != nullptr)
+		const bool bIsPropertyInherited = Property->GetOwnerClass() != GetClass();
+		UObject* SuperClassDefaultObject = GetClass()->GetSuperClass()->GetDefaultObject();
+
+		FArrayProperty* Array = CastField<FArrayProperty>(Property);
+		if (Array)
 		{
 			FConfigSection* Sec = GConfig->GetSectionPrivate(*Section, true, false, *GEditorPerProjectIni);
 			check(Sec);
