[GLTFExporter] refactor: rename GLTFColor -> GLTFPackedColor

commit c6b28239dfadfdd952e44b3d736dcfaba8a7a686
Date: Mon, 18 May 2020 23:57:16 +0200

    refactor: rename GLTFColor -> GLTFPackedColor
    (cherry picked from commit 33356c613e3724baf3acffa97b14525cc8e92b7a​)

#skipci

diff --git a/Engine/Plugins/Enterprise/GLTFExporter/MigratedFromGit/Source/GLTFExporter/Private/Converters/GLTFColor.h b/Engine/Plugins/Enterprise/GLTFExporter/MigratedFromGit/Source/GLTFExporter/Private/Converters/GLTFColor.h
deleted file mode 100644
diff --git a/Engine/Plugins/Enterprise/GLTFExporter/MigratedFromGit/Source/GLTFExporter/Private/Converters/GLTFConverterUtility.h b/Engine/Plugins/Enterprise/GLTFExporter/MigratedFromGit/Source/GLTFExporter/Private/Converters/GLTFConverterUtility.h
--- a/Engine/Plugins/Enterprise/GLTFExporter/MigratedFromGit/Source/GLTFExporter/Private/Converters/GLTFConverterUtility.h
+++ b/Engine/Plugins/Enterprise/GLTFExporter/MigratedFromGit/Source/GLTFExporter/Private/Converters/GLTFConverterUtility.h
@@ -2,7 +2,7 @@
 
 #pragma once
 
-#include "Converters/GLTFColor.h"
+#include "Converters/GLTFPackedColor.h"
 #include "Json/GLTFJsonEnums.h"
 #include "Engine/EngineTypes.h"
 
@@ -36,7 +36,7 @@ struct FGLTFConverterUtility
 		return FVector4(ConvertVector(Tangent), 1.0f);
 	}
 
-	static FGLTFColor ConvertColor(const FColor& Color)
+	static FGLTFPackedColor ConvertColor(const FColor& Color)
 	{
 		// UE4 uses ABGR or ARGB depending on endianness.
 		// glTF always uses RGBA independent of endianness.
diff --git a/Engine/Plugins/Enterprise/GLTFExporter/MigratedFromGit/Source/GLTFExporter/Private/Converters/GLTFPackedColor.h b/Engine/Plugins/Enterprise/GLTFExporter/MigratedFromGit/Source/GLTFExporter/Private/Converters/GLTFPackedColor.h
new file mode 100644
--- /dev/null
+++ b/Engine/Plugins/Enterprise/GLTFExporter/MigratedFromGit/Source/GLTFExporter/Private/Converters/GLTFPackedColor.h
@@ -0,0 +1,17 @@
+// Copyright Epic Games, Inc. All Rights Reserved.
+
+#pragma once
+
+#include "CoreTypes.h"
+
+// UE4 uses an endianness dependent memory layout (ABGR for little endian, ARGB for big endian)
+// glTF uses RGBA independent of endianness, hence why we need this dedicated struct
+struct FGLTFPackedColor
+{
+	uint8 R,G,B,A;
+
+	FGLTFPackedColor(uint8 R, uint8 G, uint8 B, uint8 A = 255)
+		: R(R), G(G), B(B), A(A)
+	{
+	}
+};
diff --git a/Engine/Plugins/Enterprise/GLTFExporter/MigratedFromGit/Source/GLTFExporter/Private/Converters/GLTFStaticMeshConverters.cpp b/Engine/Plugins/Enterprise/GLTFExporter/MigratedFromGit/Source/GLTFExporter/Private/Converters/GLTFStaticMeshConverters.cpp
--- a/Engine/Plugins/Enterprise/GLTFExporter/MigratedFromGit/Source/GLTFExporter/Private/Converters/GLTFStaticMeshConverters.cpp
+++ b/Engine/Plugins/Enterprise/GLTFExporter/MigratedFromGit/Source/GLTFExporter/Private/Converters/GLTFStaticMeshConverters.cpp
@@ -55,7 +55,7 @@ FGLTFJsonAccessorIndex FGLTFColorVertexBufferConverter::Add(FGLTFConvertBuilder&
 		return FGLTFJsonAccessorIndex(INDEX_NONE);
 	}
 
-	TArray<FGLTFColor> Colors;
+	TArray<FGLTFPackedColor> Colors;
 	Colors.AddUninitialized(VertexCount);
 
 	for (uint32 VertexIndex = 0; VertexIndex < VertexCount; ++VertexIndex)
