[GLTFExporter] refactor: add GLTFJsonVector2, GLTFJsonVector3, GLTFJsonVector4, and GLTFJsonQuaternion

commit 5f999c7af93775c2662922b1df72779772b8b351
Date: Mon, 18 May 2020 23:48:51 +0200

    refactor: add GLTFJsonVector2, GLTFJsonVector3, GLTFJsonVector4, and GLTFJsonQuaternion
    (cherry picked from commit cfa4be01c0aa124e1c120444bf419521d1b3f57f​)

#skipci

diff --git a/Engine/Plugins/Enterprise/GLTFExporter/MigratedFromGit/Source/GLTFExporter/Private/Json/GLTFJsonQuaternion.cpp b/Engine/Plugins/Enterprise/GLTFExporter/MigratedFromGit/Source/GLTFExporter/Private/Json/GLTFJsonQuaternion.cpp
new file mode 100644
--- /dev/null
+++ b/Engine/Plugins/Enterprise/GLTFExporter/MigratedFromGit/Source/GLTFExporter/Private/Json/GLTFJsonQuaternion.cpp
@@ -0,0 +1,5 @@
+// Copyright Epic Games, Inc. All Rights Reserved.
+
+#include "Json/GLTFJsonQuaternion.h"
+
+const FGLTFJsonQuaternion FGLTFJsonQuaternion::Identity(0, 0, 0, 1);
diff --git a/Engine/Plugins/Enterprise/GLTFExporter/MigratedFromGit/Source/GLTFExporter/Private/Json/GLTFJsonQuaternion.h b/Engine/Plugins/Enterprise/GLTFExporter/MigratedFromGit/Source/GLTFExporter/Private/Json/GLTFJsonQuaternion.h
new file mode 100644
--- /dev/null
+++ b/Engine/Plugins/Enterprise/GLTFExporter/MigratedFromGit/Source/GLTFExporter/Private/Json/GLTFJsonQuaternion.h
@@ -0,0 +1,44 @@
+// Copyright Epic Games, Inc. All Rights Reserved.
+
+#pragma once
+
+#include "Serialization/JsonSerializer.h"
+
+struct FGLTFJsonQuaternion
+{
+	float X, Y, Z, W;
+
+	static const FGLTFJsonQuaternion Identity;
+
+	FGLTFJsonQuaternion(float X, float Y, float Z, float W)
+		: X(X), Y(Y), Z(Z), W(W)
+	{
+	}
+
+	template <class CharType = TCHAR, class PrintPolicy = TPrettyJsonPrintPolicy<CharType>>
+	void WriteArray(TJsonWriter<CharType, PrintPolicy>& JsonWriter) const
+	{
+		JsonWriter.WriteArrayStart();
+		JsonWriter.WriteValue(X);
+		JsonWriter.WriteValue(Y);
+		JsonWriter.WriteValue(Z);
+		JsonWriter.WriteValue(W);
+		JsonWriter.WriteArrayEnd();
+	}
+
+	bool operator==(const FGLTFJsonQuaternion& Other) const
+	{
+		return X == Other.X
+			&& Y == Other.Y
+			&& Z == Other.Z
+			&& W == Other.W;
+	}
+
+	bool operator!=(const FGLTFJsonQuaternion& Other) const
+	{
+		return X != Other.X
+            || Y != Other.Y
+            || Z != Other.Z
+            || W != Other.W;
+	}
+};
diff --git a/Engine/Plugins/Enterprise/GLTFExporter/MigratedFromGit/Source/GLTFExporter/Private/Json/GLTFJsonVector2.cpp b/Engine/Plugins/Enterprise/GLTFExporter/MigratedFromGit/Source/GLTFExporter/Private/Json/GLTFJsonVector2.cpp
new file mode 100644
--- /dev/null
+++ b/Engine/Plugins/Enterprise/GLTFExporter/MigratedFromGit/Source/GLTFExporter/Private/Json/GLTFJsonVector2.cpp
@@ -0,0 +1,6 @@
+// Copyright Epic Games, Inc. All Rights Reserved.
+
+#include "Json/GLTFJsonVector2.h"
+
+const FGLTFJsonVector2 FGLTFJsonVector2::Zero(0, 0);
+const FGLTFJsonVector2 FGLTFJsonVector2::One(1, 1);
diff --git a/Engine/Plugins/Enterprise/GLTFExporter/MigratedFromGit/Source/GLTFExporter/Private/Json/GLTFJsonVector2.h b/Engine/Plugins/Enterprise/GLTFExporter/MigratedFromGit/Source/GLTFExporter/Private/Json/GLTFJsonVector2.h
new file mode 100644
--- /dev/null
+++ b/Engine/Plugins/Enterprise/GLTFExporter/MigratedFromGit/Source/GLTFExporter/Private/Json/GLTFJsonVector2.h
@@ -0,0 +1,39 @@
+// Copyright Epic Games, Inc. All Rights Reserved.
+
+#pragma once
+
+#include "Serialization/JsonSerializer.h"
+
+struct FGLTFJsonVector2
+{
+	float X, Y;
+
+	static const FGLTFJsonVector2 Zero;
+	static const FGLTFJsonVector2 One;
+
+	FGLTFJsonVector2(float X, float Y)
+		: X(X), Y(Y)
+	{
+	}
+
+	template <class CharType = TCHAR, class PrintPolicy = TPrettyJsonPrintPolicy<CharType>>
+	void WriteArray(TJsonWriter<CharType, PrintPolicy>& JsonWriter) const
+	{
+		JsonWriter.WriteArrayStart();
+		JsonWriter.WriteValue(X);
+		JsonWriter.WriteValue(Y);
+		JsonWriter.WriteArrayEnd();
+	}
+
+	bool operator==(const FGLTFJsonVector2& Other) const
+	{
+		return X == Other.X
+			&& Y == Other.Y;
+	}
+
+	bool operator!=(const FGLTFJsonVector2& Other) const
+	{
+		return X != Other.X
+            || Y != Other.Y;
+	}
+};
diff --git a/Engine/Plugins/Enterprise/GLTFExporter/MigratedFromGit/Source/GLTFExporter/Private/Json/GLTFJsonVector3.cpp b/Engine/Plugins/Enterprise/GLTFExporter/MigratedFromGit/Source/GLTFExporter/Private/Json/GLTFJsonVector3.cpp
new file mode 100644
--- /dev/null
+++ b/Engine/Plugins/Enterprise/GLTFExporter/MigratedFromGit/Source/GLTFExporter/Private/Json/GLTFJsonVector3.cpp
@@ -0,0 +1,6 @@
+// Copyright Epic Games, Inc. All Rights Reserved.
+
+#include "Json/GLTFJsonVector3.h"
+
+const FGLTFJsonVector3 FGLTFJsonVector3::Zero(0, 0, 0);
+const FGLTFJsonVector3 FGLTFJsonVector3::One(1, 1, 1);
diff --git a/Engine/Plugins/Enterprise/GLTFExporter/MigratedFromGit/Source/GLTFExporter/Private/Json/GLTFJsonVector3.h b/Engine/Plugins/Enterprise/GLTFExporter/MigratedFromGit/Source/GLTFExporter/Private/Json/GLTFJsonVector3.h
new file mode 100644
--- /dev/null
+++ b/Engine/Plugins/Enterprise/GLTFExporter/MigratedFromGit/Source/GLTFExporter/Private/Json/GLTFJsonVector3.h
@@ -0,0 +1,42 @@
+// Copyright Epic Games, Inc. All Rights Reserved.
+
+#pragma once
+
+#include "Serialization/JsonSerializer.h"
+
+struct FGLTFJsonVector3
+{
+	float X, Y, Z;
+
+	static const FGLTFJsonVector3 Zero;
+	static const FGLTFJsonVector3 One;
+
+	FGLTFJsonVector3(float X, float Y, float Z)
+		: X(X), Y(Y), Z(Z)
+	{
+	}
+
+	template <class CharType = TCHAR, class PrintPolicy = TPrettyJsonPrintPolicy<CharType>>
+	void WriteArray(TJsonWriter<CharType, PrintPolicy>& JsonWriter) const
+	{
+		JsonWriter.WriteArrayStart();
+		JsonWriter.WriteValue(X);
+		JsonWriter.WriteValue(Y);
+		JsonWriter.WriteValue(Z);
+		JsonWriter.WriteArrayEnd();
+	}
+
+	bool operator==(const FGLTFJsonVector3& Other) const
+	{
+		return X == Other.X
+			&& Y == Other.Y
+			&& Z == Other.Z;
+	}
+
+	bool operator!=(const FGLTFJsonVector3& Other) const
+	{
+		return X != Other.X
+            || Y != Other.Y
+            || Z != Other.Z;
+	}
+};
diff --git a/Engine/Plugins/Enterprise/GLTFExporter/MigratedFromGit/Source/GLTFExporter/Private/Json/GLTFJsonVector4.cpp b/Engine/Plugins/Enterprise/GLTFExporter/MigratedFromGit/Source/GLTFExporter/Private/Json/GLTFJsonVector4.cpp
new file mode 100644
--- /dev/null
+++ b/Engine/Plugins/Enterprise/GLTFExporter/MigratedFromGit/Source/GLTFExporter/Private/Json/GLTFJsonVector4.cpp
@@ -0,0 +1,6 @@
+// Copyright Epic Games, Inc. All Rights Reserved.
+
+#include "Json/GLTFJsonVector4.h"
+
+const FGLTFJsonVector4 FGLTFJsonVector4::Zero(0, 0, 0, 0);
+const FGLTFJsonVector4 FGLTFJsonVector4::One(1, 1, 1, 1);
diff --git a/Engine/Plugins/Enterprise/GLTFExporter/MigratedFromGit/Source/GLTFExporter/Private/Json/GLTFJsonVector4.h b/Engine/Plugins/Enterprise/GLTFExporter/MigratedFromGit/Source/GLTFExporter/Private/Json/GLTFJsonVector4.h
new file mode 100644
--- /dev/null
+++ b/Engine/Plugins/Enterprise/GLTFExporter/MigratedFromGit/Source/GLTFExporter/Private/Json/GLTFJsonVector4.h
@@ -0,0 +1,51 @@
+// Copyright Epic Games, Inc. All Rights Reserved.
+
+#pragma once
+
+#include "Json/GLTFJsonVector3.h"
+#include "Serialization/JsonSerializer.h"
+
+struct FGLTFJsonVector4
+{
+	float X, Y, Z, W;
+
+	static const FGLTFJsonVector4 Zero;
+	static const FGLTFJsonVector4 One;
+
+	FGLTFJsonVector4(float X, float Y, float Z, float W)
+        : X(X), Y(Y), Z(Z), W(W)
+	{
+	}
+
+	FGLTFJsonVector4(const FGLTFJsonVector3& Vector3, float W)
+        : X(Vector3.X), Y(Vector3.Y), Z(Vector3.Z), W(W)
+	{
+	}
+
+	template <class CharType = TCHAR, class PrintPolicy = TPrettyJsonPrintPolicy<CharType>>
+	void WriteArray(TJsonWriter<CharType, PrintPolicy>& JsonWriter) const
+	{
+		JsonWriter.WriteArrayStart();
+		JsonWriter.WriteValue(X);
+		JsonWriter.WriteValue(Y);
+		JsonWriter.WriteValue(Z);
+		JsonWriter.WriteValue(W);
+		JsonWriter.WriteArrayEnd();
+	}
+
+	bool operator==(const FGLTFJsonVector4& Other) const
+	{
+		return X == Other.X
+            && Y == Other.Y
+            && Z == Other.Z
+            && W == Other.W;
+	}
+
+	bool operator!=(const FGLTFJsonVector4& Other) const
+	{
+		return X != Other.X
+            || Y != Other.Y
+            || Z != Other.Z
+            || W != Other.W;
+	}
+};
