[GLTFExporter] refactor: convert FVector2D, FVector, FVector4, and FQuat to GLTFJsonVector* and GLTFJsonQuaternion

commit f47ffe28454d77a74216b5139a3ccab6faae9290
Date: Mon, 18 May 2020 23:54:59 +0200

    refactor: convert FVector2D, FVector, FVector4, and FQuat to GLTFJsonVector* and GLTFJsonQuaternion
    (cherry picked from commit 06705c2951ed9311f99feb1f42cd5df1582019e8​)

#skipci

diff --git a/Engine/Plugins/Enterprise/GLTFExporter/MigratedFromGit/Source/GLTFExporter/Private/Converters/GLTFConverterUtility.h b/Engine/Plugins/Enterprise/GLTFExporter/MigratedFromGit/Source/GLTFExporter/Private/Converters/GLTFConverterUtility.h
--- a/Engine/Plugins/Enterprise/GLTFExporter/MigratedFromGit/Source/GLTFExporter/Private/Converters/GLTFConverterUtility.h
+++ b/Engine/Plugins/Enterprise/GLTFExporter/MigratedFromGit/Source/GLTFExporter/Private/Converters/GLTFConverterUtility.h
@@ -4,36 +4,46 @@
 
 #include "Converters/GLTFPackedColor.h"
 #include "Json/GLTFJsonEnums.h"
+#include "Json/GLTFJsonVector2.h"
+#include "Json/GLTFJsonVector3.h"
+#include "Json/GLTFJsonVector4.h"
+#include "Json/GLTFJsonQuaternion.h"
 #include "Engine/EngineTypes.h"
 
 struct FGLTFConverterUtility
 {
-	static FVector ConvertVector(const FVector& Vector)
+	static FGLTFJsonVector3 ConvertVector(const FVector& Vector)
 	{
 		// UE4 uses a left-handed coordinate system, with Z up.
 		// glTF uses a right-handed coordinate system, with Y up.
 		return { Vector.X, Vector.Z, Vector.Y };
 	}
 
-	static FVector ConvertPosition(const FVector& Position)
+	static FGLTFJsonVector3 ConvertPosition(const FVector& Position)
 	{
-		return ConvertVector(Position) * 0.01f; // TODO: use options export scale instead of hardcoded value
+		return ConvertVector(Position * 0.01f); // TODO: use options export scale instead of hardcoded value
 	}
 
-	static FVector ConvertScale(const FVector& Scale)
+	static FGLTFJsonVector3 ConvertScale(const FVector& Scale)
 	{
 		return ConvertVector(Scale);
 	}
 
-	static FVector ConvertNormal(const FVector& Normal)
+	static FGLTFJsonVector3 ConvertNormal(const FVector& Normal)
 	{
 		return ConvertVector(Normal);
 	}
 
-	static FVector4 ConvertTangent(const FVector& Tangent)
+	static FGLTFJsonVector4 ConvertTangent(const FVector& Tangent)
 	{
 		// glTF stores tangent as Vec4, with W component indicating handedness of tangent basis.
-		return FVector4(ConvertVector(Tangent), 1.0f);
+		return { ConvertVector(Tangent), 1.0f };
+	}
+
+	static FGLTFJsonVector2 ConvertUV(const FVector2D& UV)
+	{
+		// No conversion actually needed, this is primarily for type-safety.
+		return { UV.X, UV.Y };
 	}
 
 	static FGLTFPackedColor ConvertColor(const FColor& Color)
@@ -43,7 +53,7 @@ struct FGLTFConverterUtility
 		return { Color.R, Color.G, Color.B, Color.A };
 	}
 
-	static FQuat ConvertRotation(const FQuat& Rotation)
+	static FGLTFJsonQuaternion ConvertRotation(const FQuat& Rotation)
 	{
 		// UE4 uses a left-handed coordinate system, with Z up.
 		// glTF uses a right-handed coordinate system, with Y up.
@@ -53,10 +63,9 @@ struct FGLTFConverterUtility
 		// also, as handedness is changed rotation is inversed - hence negation
 		// therefore glTFRotation = (-qX, -qZ, -qY, qw)
 
-		const FQuat Result(-Rotation.X, -Rotation.Z, -Rotation.Y, Rotation.W);
 		// Not checking if quaternion is normalized
 		// e.g. some sources use non-unit Quats for rotation tangents
-		return Result;
+		return { -Rotation.X, -Rotation.Z, -Rotation.Y, Rotation.W };
 	}
 
 	static FMatrix ConvertMatrix(const FMatrix& Matrix)
diff --git a/Engine/Plugins/Enterprise/GLTFExporter/MigratedFromGit/Source/GLTFExporter/Private/Converters/GLTFStaticMeshConverters.cpp b/Engine/Plugins/Enterprise/GLTFExporter/MigratedFromGit/Source/GLTFExporter/Private/Converters/GLTFStaticMeshConverters.cpp
--- a/Engine/Plugins/Enterprise/GLTFExporter/MigratedFromGit/Source/GLTFExporter/Private/Converters/GLTFStaticMeshConverters.cpp
+++ b/Engine/Plugins/Enterprise/GLTFExporter/MigratedFromGit/Source/GLTFExporter/Private/Converters/GLTFStaticMeshConverters.cpp
@@ -12,7 +12,7 @@ FGLTFJsonAccessorIndex FGLTFPositionVertexBufferConverter::Add(FGLTFConvertBuild
 		return FGLTFJsonAccessorIndex(INDEX_NONE);
 	}
 
-	TArray<FVector> Positions;
+	TArray<FGLTFJsonVector3> Positions;
 	Positions.AddUninitialized(VertexCount);
 
 	for (uint32 VertexIndex = 0; VertexIndex < VertexCount; ++VertexIndex)
@@ -20,13 +20,19 @@ FGLTFJsonAccessorIndex FGLTFPositionVertexBufferConverter::Add(FGLTFConvertBuild
 		Positions[VertexIndex] = FGLTFConverterUtility::ConvertPosition(VertexBuffer->VertexPosition(VertexIndex));
 	}
 
-	FBox BoundingBox;
-	BoundingBox.Init();
-
 	// More accurate bounding box if based on raw vertex values
-	for (const FVector& Position : Positions)
+	FGLTFJsonVector3 MinPosition = Positions[0];
+	FGLTFJsonVector3 MaxPosition = Positions[0];
+
+	for (uint32 VertexIndex = 1; VertexIndex < VertexCount; ++VertexIndex)
 	{
-		BoundingBox += Position;
+		const FGLTFJsonVector3& Position = Positions[VertexIndex];
+		MinPosition.X = FMath::Min(MinPosition.X, Position.X);
+		MinPosition.Y = FMath::Min(MinPosition.Y, Position.Y);
+		MinPosition.Z = FMath::Min(MinPosition.Z, Position.Z);
+		MaxPosition.X = FMath::Max(MaxPosition.X, Position.X);
+		MaxPosition.Y = FMath::Max(MaxPosition.Y, Position.Y);
+		MaxPosition.Z = FMath::Max(MaxPosition.Z, Position.Z);
 	}
 
 	FGLTFJsonAccessor JsonAccessor;
@@ -37,12 +43,12 @@ FGLTFJsonAccessorIndex FGLTFPositionVertexBufferConverter::Add(FGLTFConvertBuild
 	JsonAccessor.Type = EGLTFJsonAccessorType::Vec3;
 
 	JsonAccessor.MinMaxLength = 3;
-	JsonAccessor.Max[0] = BoundingBox.Max.X;
-	JsonAccessor.Max[1] = BoundingBox.Max.Y;
-	JsonAccessor.Max[2] = BoundingBox.Max.Z;
-	JsonAccessor.Min[0] = BoundingBox.Min.X;
-	JsonAccessor.Min[1] = BoundingBox.Min.Y;
-	JsonAccessor.Min[2] = BoundingBox.Min.Z;
+	JsonAccessor.Min[0] = MinPosition.X;
+	JsonAccessor.Min[1] = MinPosition.Y;
+	JsonAccessor.Min[2] = MinPosition.Z;
+	JsonAccessor.Max[0] = MaxPosition.X;
+	JsonAccessor.Max[1] = MaxPosition.Y;
+	JsonAccessor.Max[2] = MaxPosition.Z;
 
 	return Builder.AddAccessor(JsonAccessor);
 }
@@ -82,7 +88,7 @@ FGLTFJsonAccessorIndex FGLTFStaticMeshNormalVertexBufferConverter::Add(FGLTFConv
 		return FGLTFJsonAccessorIndex(INDEX_NONE);
 	}
 
-	TArray<FVector> Normals;
+	TArray<FGLTFJsonVector3> Normals;
 	Normals.AddUninitialized(VertexCount);
 
 	for (uint32 VertexIndex = 0; VertexIndex < VertexCount; ++VertexIndex)
@@ -108,7 +114,7 @@ FGLTFJsonAccessorIndex FGLTFStaticMeshTangentVertexBufferConverter::Add(FGLTFCon
 		return FGLTFJsonAccessorIndex(INDEX_NONE);
 	}
 
-	TArray<FVector4> Tangents;
+	TArray<FGLTFJsonVector4> Tangents;
 	Tangents.AddUninitialized(VertexCount);
 
 	for (uint32 VertexIndex = 0; VertexIndex < VertexCount; ++VertexIndex)
@@ -137,12 +143,12 @@ FGLTFJsonAccessorIndex FGLTFStaticMeshUVVertexBufferConverter::Add(FGLTFConvertB
 		return FGLTFJsonAccessorIndex(INDEX_NONE);
 	}
 
-	TArray<FVector2D> UVs;
+	TArray<FGLTFJsonVector2> UVs;
 	UVs.AddUninitialized(VertexCount);
 
 	for (uint32 VertexIndex = 0; VertexIndex < VertexCount; ++VertexIndex)
 	{
-		UVs[VertexIndex] = VertexBuffer->GetVertexUV(VertexIndex, UVIndex);
+		UVs[VertexIndex] = FGLTFConverterUtility::ConvertUV(VertexBuffer->GetVertexUV(VertexIndex, UVIndex));
 	}
 
 	FGLTFJsonAccessor JsonAccessor;
diff --git a/Engine/Plugins/Enterprise/GLTFExporter/MigratedFromGit/Source/GLTFExporter/Private/Json/GLTFJsonNode.h b/Engine/Plugins/Enterprise/GLTFExporter/MigratedFromGit/Source/GLTFExporter/Private/Json/GLTFJsonNode.h
--- a/Engine/Plugins/Enterprise/GLTFExporter/MigratedFromGit/Source/GLTFExporter/Private/Json/GLTFJsonNode.h
+++ b/Engine/Plugins/Enterprise/GLTFExporter/MigratedFromGit/Source/GLTFExporter/Private/Json/GLTFJsonNode.h
@@ -3,15 +3,17 @@
 #pragma once
 
 #include "Json/GLTFJsonIndex.h"
+#include "Json/GLTFJsonVector3.h"
+#include "Json/GLTFJsonQuaternion.h"
 #include "Serialization/JsonSerializer.h"
 
 struct FGLTFJsonNode
 {
 	FString Name;
 
-	FVector Translation;
-	FQuat   Rotation;
-	FVector Scale;
+	FGLTFJsonVector3    Translation;
+	FGLTFJsonQuaternion Rotation;
+	FGLTFJsonVector3    Scale;
 
 	FGLTFJsonCameraIndex Camera;
 	FGLTFJsonSkinIndex   Skin;
@@ -20,9 +22,9 @@ struct FGLTFJsonNode
 	TArray<FGLTFJsonNodeIndex> Children;
 
 	FGLTFJsonNode()
-		: Translation(FVector::ZeroVector)
-		, Rotation(FQuat::Identity)
-		, Scale(FVector::OneVector)
+		: Translation(FGLTFJsonVector3::Zero)
+		, Rotation(FGLTFJsonQuaternion::Identity)
+		, Scale(FGLTFJsonVector3::One)
 		, Camera(INDEX_NONE)
 		, Skin(INDEX_NONE)
 		, Mesh(INDEX_NONE)
@@ -39,32 +41,22 @@ struct FGLTFJsonNode
 			JsonWriter.WriteValue(TEXT("name"), Name);
 		}
 
-		if (Translation != FVector::ZeroVector)
+		if (Translation != FGLTFJsonVector3::Zero)
 		{
-			JsonWriter.WriteArrayStart(TEXT("translation"));
-			JsonWriter.WriteValue(Translation.X);
-			JsonWriter.WriteValue(Translation.Y);
-			JsonWriter.WriteValue(Translation.Z);
-			JsonWriter.WriteArrayEnd();
+			JsonWriter.WriteIdentifierPrefix(TEXT("translation"));
+			Translation.WriteArray(JsonWriter);
 		}
 
-		if (Rotation != FQuat::Identity)
+		if (Rotation != FGLTFJsonQuaternion::Identity)
 		{
-			JsonWriter.WriteArrayStart(TEXT("rotation"));
-			JsonWriter.WriteValue(Rotation.X);
-			JsonWriter.WriteValue(Rotation.Y);
-			JsonWriter.WriteValue(Rotation.Z);
-			JsonWriter.WriteValue(Rotation.W);
-			JsonWriter.WriteArrayEnd();
+			JsonWriter.WriteIdentifierPrefix(TEXT("rotation"));
+			Rotation.WriteArray(JsonWriter);
 		}
 
-		if (Scale != FVector::OneVector)
+		if (Scale != FGLTFJsonVector3::One)
 		{
-			JsonWriter.WriteArrayStart(TEXT("scale"));
-			JsonWriter.WriteValue(Scale.X);
-			JsonWriter.WriteValue(Scale.Y);
-			JsonWriter.WriteValue(Scale.Z);
-			JsonWriter.WriteArrayEnd();
+			JsonWriter.WriteIdentifierPrefix(TEXT("scale"));
+			Scale.WriteArray(JsonWriter);
 		}
 
 		if (Camera != INDEX_NONE)
