[GLTFExporter] fix(options): disable all options until implemented

commit 8d48a8cc47e94b938340a38e41a769b3dd6b0d7a
Date: Tue, 19 May 2020 09:25:02 +0200

    fix(options): disable all options until implemented
    (cherry picked from commit 091e5587dec939c14b4aba00f44c0af924cc1174​)

#skipci

diff --git a/Engine/Plugins/Enterprise/GLTFExporter/MigratedFromGit/Source/GLTFExporter/Private/GLTFExportOptions.cpp b/Engine/Plugins/Enterprise/GLTFExporter/MigratedFromGit/Source/GLTFExporter/Private/GLTFExportOptions.cpp
--- a/Engine/Plugins/Enterprise/GLTFExporter/MigratedFromGit/Source/GLTFExporter/Private/GLTFExportOptions.cpp
+++ b/Engine/Plugins/Enterprise/GLTFExporter/MigratedFromGit/Source/GLTFExporter/Private/GLTFExportOptions.cpp
@@ -206,3 +206,8 @@ void UGLTFExportOptions::FillOptions(bool bBatchMode, bool bShowOptionDialog, co
 	SGLTFExportOptionsWindow::ShowDialog(this, FullPath, bBatchMode, bOutOperationCanceled, bOutExportAll);
 	// TODO: call SaveOptions();
 }
+
+bool UGLTFExportOptions::CanEditChange(const FProperty* InProperty) const
+{
+	return false; // TODO: change to true for each option implemented and support by the object to be exported.
+}
diff --git a/Engine/Plugins/Enterprise/GLTFExporter/MigratedFromGit/Source/GLTFExporter/Public/GLTFExportOptions.h b/Engine/Plugins/Enterprise/GLTFExporter/MigratedFromGit/Source/GLTFExporter/Public/GLTFExportOptions.h
--- a/Engine/Plugins/Enterprise/GLTFExporter/MigratedFromGit/Source/GLTFExporter/Public/GLTFExportOptions.h
+++ b/Engine/Plugins/Enterprise/GLTFExporter/MigratedFromGit/Source/GLTFExporter/Public/GLTFExportOptions.h
@@ -157,4 +157,6 @@ public:
 	* The function is saving the dialog state in a user ini file and reload it from there. It is not changing the CDO.
 	*/
 	void FillOptions(bool bBatchMode, bool bShowOptionDialog, const FString& FullPath, bool& bOutOperationCanceled, bool& bOutExportAll);
+
+	virtual bool CanEditChange(const FProperty* InProperty) const override;
 };
