[GLTFExporter] fix(levels): fix invalid node hierarchy bug uncovered by new test assets

commit 20a8a7bb746ae439ff636fcf492cc02bb0ae223f
Date: Tue, 19 May 2020 21:03:57 +0200

    fix(levels): fix invalid node hierarchy bug uncovered by new test assets
    (cherry picked from commit 5710e9f60a19374929a50b7e19df3843d766b0e2​)

#skipci

diff --git a/Engine/Plugins/Enterprise/GLTFExporter/MigratedFromGit/Source/GLTFExporter/Private/Converters/GLTFConverterUtility.cpp b/Engine/Plugins/Enterprise/GLTFExporter/MigratedFromGit/Source/GLTFExporter/Private/Converters/GLTFConverterUtility.cpp
--- a/Engine/Plugins/Enterprise/GLTFExporter/MigratedFromGit/Source/GLTFExporter/Private/Converters/GLTFConverterUtility.cpp
+++ b/Engine/Plugins/Enterprise/GLTFExporter/MigratedFromGit/Source/GLTFExporter/Private/Converters/GLTFConverterUtility.cpp
@@ -12,3 +12,14 @@ bool FGLTFConverterUtility::IsHDRIBackdropBlueprint(const UBlueprint* Blueprint)
 {
 	return Blueprint != nullptr && Blueprint->GetPathName().Equals(TEXT("/HDRIBackdrop/Blueprints/HDRIBackdrop.HDRIBackdrop"));
 }
+
+bool FGLTFConverterUtility::IsSelected(const UActorComponent* ActorComponent)
+{
+	if (ActorComponent == nullptr)
+	{
+		return false;
+	}
+
+	const AActor* Owner = ActorComponent->GetOwner();
+	return Owner != nullptr && Owner->IsSelected();
+}
diff --git a/Engine/Plugins/Enterprise/GLTFExporter/MigratedFromGit/Source/GLTFExporter/Private/Converters/GLTFConverterUtility.h b/Engine/Plugins/Enterprise/GLTFExporter/MigratedFromGit/Source/GLTFExporter/Private/Converters/GLTFConverterUtility.h
--- a/Engine/Plugins/Enterprise/GLTFExporter/MigratedFromGit/Source/GLTFExporter/Private/Converters/GLTFConverterUtility.h
+++ b/Engine/Plugins/Enterprise/GLTFExporter/MigratedFromGit/Source/GLTFExporter/Private/Converters/GLTFConverterUtility.h
@@ -98,4 +98,6 @@ struct FGLTFConverterUtility
 	static bool IsSkySphereBlueprint(const UBlueprint* Blueprint);
 
 	static bool IsHDRIBackdropBlueprint(const UBlueprint* Blueprint);
+
+	static bool IsSelected(const UActorComponent* ActorComponent);
 };
diff --git a/Engine/Plugins/Enterprise/GLTFExporter/MigratedFromGit/Source/GLTFExporter/Private/Converters/GLTFLevelConverters.cpp b/Engine/Plugins/Enterprise/GLTFExporter/MigratedFromGit/Source/GLTFExporter/Private/Converters/GLTFLevelConverters.cpp
--- a/Engine/Plugins/Enterprise/GLTFExporter/MigratedFromGit/Source/GLTFExporter/Private/Converters/GLTFLevelConverters.cpp
+++ b/Engine/Plugins/Enterprise/GLTFExporter/MigratedFromGit/Source/GLTFExporter/Private/Converters/GLTFLevelConverters.cpp
@@ -38,16 +38,12 @@ FGLTFJsonNodeIndex FGLTFSceneComponentConverter::Add(FGLTFConvertBuilder& Builde
 	const TArray<USceneComponent*>& Children = SceneComponent->GetAttachChildren();
 	for (const USceneComponent* ChildComponent : Children)
 	{
-		if (ChildComponent != nullptr)
+		if (ChildComponent != nullptr && (!bSelectedOnly || FGLTFConverterUtility::IsSelected(ChildComponent)))
 		{
-			const AActor* ChildOwner = ChildComponent->GetOwner();
-			if (!bSelectedOnly || ChildOwner->IsSelected())
+			FGLTFJsonNodeIndex NodeIndex = Builder.GetOrAddNode(ChildComponent, bSelectedOnly, false);
+			if (NodeIndex != INDEX_NONE)
 			{
-				FGLTFJsonNodeIndex NodeIndex = Builder.GetOrAddNode(ChildComponent, bSelectedOnly, false);
-				if (NodeIndex != INDEX_NONE)
-				{
-					Node.Children.Add(NodeIndex);
-				}
+				Node.Children.Add(NodeIndex);
 			}
 		}
 	}
@@ -73,8 +69,8 @@ FGLTFJsonSceneIndex FGLTFLevelConverter::Add(FGLTFConvertBuilder& Builder, const
 			const USceneComponent* RootComponent = Actor->GetRootComponent();
 			if (RootComponent != nullptr)
 			{
-				const AActor* ParentActor = Actor->GetParentActor();
-				if (ParentActor == nullptr || (bSelectedOnly && !ParentActor->IsSelected()))
+				const USceneComponent* ParentComponent = RootComponent->GetAttachParent();
+				if (ParentComponent == nullptr || (bSelectedOnly && !FGLTFConverterUtility::IsSelected(ParentComponent)))
 				{
 					FGLTFJsonNodeIndex NodeIndex = Builder.GetOrAddNode(RootComponent, bSelectedOnly, true);
 					if (NodeIndex != INDEX_NONE)
