[GLTFExporter] fix(levels): fix HDRI backdrop detection (shouldn't export static mesh naively)

commit e13f3b404af2584c36e426f2b34bd2c3cd3a14ca
Date: Tue, 19 May 2020 23:30:09 +0200

    fix(levels): fix HDRI backdrop detection (shouldn't export static mesh naively)
    (cherry picked from commit 1d348d6d6f36f4acff1431bacea3f8c2449e0fa0​)

#skipci

diff --git a/Engine/Plugins/Enterprise/GLTFExporter/MigratedFromGit/Source/GLTFExporter/Private/Converters/GLTFLevelConverters.cpp b/Engine/Plugins/Enterprise/GLTFExporter/MigratedFromGit/Source/GLTFExporter/Private/Converters/GLTFLevelConverters.cpp
--- a/Engine/Plugins/Enterprise/GLTFExporter/MigratedFromGit/Source/GLTFExporter/Private/Converters/GLTFLevelConverters.cpp
+++ b/Engine/Plugins/Enterprise/GLTFExporter/MigratedFromGit/Source/GLTFExporter/Private/Converters/GLTFLevelConverters.cpp
@@ -24,15 +24,18 @@ FGLTFJsonNodeIndex FGLTFSceneComponentConverter::Add(FGLTFConvertBuilder& Builde
 
 	if (FGLTFConverterUtility::IsSkySphereBlueprint(Blueprint))
 	{
-		// Ignore components
+		// Ignore mesh and light components that are part of BP_SkySphere
 	}
-	else if (const UStaticMeshComponent* StaticMeshComponent = Cast<UStaticMeshComponent>(SceneComponent))
+	else if (FGLTFConverterUtility::IsHDRIBackdropBlueprint(Blueprint))
 	{
-		Node.Mesh = Builder.GetOrAddMesh(StaticMeshComponent);
+		if (bIsRootComponent)
+		{
+			// TODO: add support for backdrop
+		}
 	}
-	else if (FGLTFConverterUtility::IsHDRIBackdropBlueprint(Blueprint) && bIsRootComponent)
+	else if (const UStaticMeshComponent* StaticMeshComponent = Cast<UStaticMeshComponent>(SceneComponent))
 	{
-		// TODO: add support for backdrop
+		Node.Mesh = Builder.GetOrAddMesh(StaticMeshComponent);
 	}
 
 	const TArray<USceneComponent*>& Children = SceneComponent->GetAttachChildren();
