[GLTFExporter] feat(static-mesh): add support for 16bit vertex indices

commit 9b0b91be99e4e3ed44e10aa3a6521ef5835f9879
Date: Tue, 12 May 2020 08:46:38 +0200

    feat(static-mesh): add support for 16bit vertex indices
    (cherry picked from commit f59efb9e462093e34a10bdb0ec527c69e445a513​)

#skipci

diff --git a/Engine/Plugins/Enterprise/GLTFExporter/MigratedFromGit/Source/GLTFExporter/Private/Converters/GLTFStaticMeshConverters.cpp b/Engine/Plugins/Enterprise/GLTFExporter/MigratedFromGit/Source/GLTFExporter/Private/Converters/GLTFStaticMeshConverters.cpp
--- a/Engine/Plugins/Enterprise/GLTFExporter/MigratedFromGit/Source/GLTFExporter/Private/Converters/GLTFStaticMeshConverters.cpp
+++ b/Engine/Plugins/Enterprise/GLTFExporter/MigratedFromGit/Source/GLTFExporter/Private/Converters/GLTFStaticMeshConverters.cpp
@@ -168,9 +168,16 @@ FGLTFJsonBufferViewIndex FGLTFStaticMeshIndexBufferConverter::Add(FGLTFConvertBu
 		return FGLTFJsonBufferViewIndex(INDEX_NONE);
 	}
 
-	TArray<uint32> Indices;
-	IndexBuffer->GetCopy(Indices);
-	return Builder.AddBufferView(Indices, Name, EGLTFJsonBufferTarget::ElementArrayBuffer);
+	if (IndexBuffer->Is32Bit())
+	{
+		TArray<uint32> Indices;
+		IndexBuffer->GetCopy(Indices);
+		return Builder.AddBufferView(Indices, Name, EGLTFJsonBufferTarget::ElementArrayBuffer);
+	}
+
+	const uint16* IndexData = IndexBuffer->AccessStream16();
+	const int32 IndexDataSize = IndexBuffer->GetIndexDataSize();
+	return Builder.AddBufferView(IndexData, IndexDataSize, Name, EGLTFJsonBufferTarget::ElementArrayBuffer);
 }
 
 FGLTFJsonAccessorIndex FGLTFStaticMeshSectionConverter::Add(FGLTFConvertBuilder& Builder, const FString& Name, TTuple<const FStaticMeshSection*, const FRawStaticIndexBuffer*> Params)
@@ -184,11 +191,14 @@ FGLTFJsonAccessorIndex FGLTFStaticMeshSectionConverter::Add(FGLTFConvertBuilder&
 		return FGLTFJsonAccessorIndex(INDEX_NONE);
 	}
 
+	const uint32 FirstIndex = MeshSection->FirstIndex;
+	const bool bIs32Bit = IndexBuffer->Is32Bit();
+
 	FGLTFJsonAccessor JsonAccessor;
 	JsonAccessor.Name = Name;
 	JsonAccessor.BufferView = Builder.GetOrAddIndexBufferView(IndexBuffer);
-	JsonAccessor.ByteOffset = MeshSection->FirstIndex * sizeof(uint32);
-	JsonAccessor.ComponentType = EGLTFJsonComponentType::U32;
+	JsonAccessor.ByteOffset = FirstIndex * (bIs32Bit ? sizeof(uint32) : sizeof(uint16));
+	JsonAccessor.ComponentType = bIs32Bit ? EGLTFJsonComponentType::U32 : EGLTFJsonComponentType::U16;
 	JsonAccessor.Count = TriangleCount * 3;
 	JsonAccessor.Type = EGLTFJsonAccessorType::Scalar;
 
