[GLTFExporter] fix(static-mesh): fix byte alignment for non-32bit component types

commit 20d63348094e247e774fb6bdf5b7ea06ee16b9ce
Date: Sat, 16 May 2020 11:28:49 +0200

    fix(static-mesh): fix byte alignment for non-32bit component types
    (cherry picked from commit abab321a71a04c2f195fe9659b6a7ef9d678e4df​)

#skipci

diff --git a/Engine/Plugins/Enterprise/GLTFExporter/MigratedFromGit/Source/GLTFExporter/Private/Builders/GLTFBufferBuilder.cpp b/Engine/Plugins/Enterprise/GLTFExporter/MigratedFromGit/Source/GLTFExporter/Private/Builders/GLTFBufferBuilder.cpp
--- a/Engine/Plugins/Enterprise/GLTFExporter/MigratedFromGit/Source/GLTFExporter/Private/Builders/GLTFBufferBuilder.cpp
+++ b/Engine/Plugins/Enterprise/GLTFExporter/MigratedFromGit/Source/GLTFExporter/Private/Builders/GLTFBufferBuilder.cpp
@@ -9,17 +9,27 @@ FGLTFBufferBuilder::FGLTFBufferBuilder()
 	BufferIndex = AddBuffer(FGLTFJsonBuffer());
 }
 
-FGLTFJsonBufferViewIndex FGLTFBufferBuilder::AddBufferView(const void* RawData, uint64 ByteLength, const FString& Name, EGLTFJsonBufferTarget BufferTarget)
+FGLTFJsonBufferViewIndex FGLTFBufferBuilder::AddBufferView(const void* RawData, uint64 ByteLength, const FString& Name, uint8 DataAlignment, EGLTFJsonBufferTarget BufferTarget)
 {
+	uint64 ByteOffset = BufferData.Num();
+
+	// Data offset must be a multiple of the size of the glTF component type (given by ByteAlignment).
+	const uint8 Padding = (DataAlignment - (ByteOffset % DataAlignment)) % DataAlignment;
+	if (Padding > 0)
+	{
+		ByteOffset += Padding;
+		BufferData.AddZeroed(Padding);
+	}
+
+	BufferData.Append(static_cast<const uint8*>(RawData), ByteLength);
+
 	FGLTFJsonBufferView JsonBufferView;
 	JsonBufferView.Name = Name;
 	JsonBufferView.Buffer = BufferIndex;
-	JsonBufferView.ByteOffset = BufferData.Num();
+	JsonBufferView.ByteOffset = ByteOffset;
 	JsonBufferView.ByteLength = ByteLength;
 	JsonBufferView.Target = BufferTarget;
 
-	BufferData.Append(static_cast<const uint8*>(RawData), ByteLength);
-
 	return FGLTFJsonBuilder::AddBufferView(JsonBufferView);
 }
 
diff --git a/Engine/Plugins/Enterprise/GLTFExporter/MigratedFromGit/Source/GLTFExporter/Private/Builders/GLTFBufferBuilder.h b/Engine/Plugins/Enterprise/GLTFExporter/MigratedFromGit/Source/GLTFExporter/Private/Builders/GLTFBufferBuilder.h
--- a/Engine/Plugins/Enterprise/GLTFExporter/MigratedFromGit/Source/GLTFExporter/Private/Builders/GLTFBufferBuilder.h
+++ b/Engine/Plugins/Enterprise/GLTFExporter/MigratedFromGit/Source/GLTFExporter/Private/Builders/GLTFBufferBuilder.h
@@ -10,12 +10,12 @@ public:
 
 	FGLTFBufferBuilder();
 
-	FGLTFJsonBufferViewIndex AddBufferView(const void* RawData, uint64 ByteLength, const FString& Name = TEXT(""), EGLTFJsonBufferTarget BufferTarget = EGLTFJsonBufferTarget::ArrayBuffer);
+	FGLTFJsonBufferViewIndex AddBufferView(const void* RawData, uint64 ByteLength, const FString& Name = TEXT(""), uint8 DataAlignment = 4, EGLTFJsonBufferTarget BufferTarget = EGLTFJsonBufferTarget::ArrayBuffer);
 
 	template <class ElementType>
-	FGLTFJsonBufferViewIndex AddBufferView(const TArray<ElementType>& Array, const FString& Name = TEXT(""), EGLTFJsonBufferTarget BufferTarget = EGLTFJsonBufferTarget::ArrayBuffer)
+	FGLTFJsonBufferViewIndex AddBufferView(const TArray<ElementType>& Array, const FString& Name = TEXT(""), uint8 DataAlignment = 4, EGLTFJsonBufferTarget BufferTarget = EGLTFJsonBufferTarget::ArrayBuffer)
 	{
-		return AddBufferView(Array.GetData(), Array.Num() * sizeof(ElementType), Name, BufferTarget);
+		return AddBufferView(Array.GetData(), Array.Num() * sizeof(ElementType), Name, DataAlignment, BufferTarget);
 	}
 
 	void UpdateJsonBufferObject(const FString& BinaryFilePath);
diff --git a/Engine/Plugins/Enterprise/GLTFExporter/MigratedFromGit/Source/GLTFExporter/Private/Converters/GLTFStaticMeshConverters.cpp b/Engine/Plugins/Enterprise/GLTFExporter/MigratedFromGit/Source/GLTFExporter/Private/Converters/GLTFStaticMeshConverters.cpp
--- a/Engine/Plugins/Enterprise/GLTFExporter/MigratedFromGit/Source/GLTFExporter/Private/Converters/GLTFStaticMeshConverters.cpp
+++ b/Engine/Plugins/Enterprise/GLTFExporter/MigratedFromGit/Source/GLTFExporter/Private/Converters/GLTFStaticMeshConverters.cpp
@@ -172,12 +172,12 @@ FGLTFJsonBufferViewIndex FGLTFStaticMeshIndexBufferConverter::Add(FGLTFConvertBu
 	{
 		TArray<uint32> Indices;
 		IndexBuffer->GetCopy(Indices);
-		return Builder.AddBufferView(Indices, Name, EGLTFJsonBufferTarget::ElementArrayBuffer);
+		return Builder.AddBufferView(Indices, Name, sizeof(uint32), EGLTFJsonBufferTarget::ElementArrayBuffer);
 	}
 
 	const uint16* IndexData = IndexBuffer->AccessStream16();
 	const int32 IndexDataSize = IndexBuffer->GetIndexDataSize();
-	return Builder.AddBufferView(IndexData, IndexDataSize, Name, EGLTFJsonBufferTarget::ElementArrayBuffer);
+	return Builder.AddBufferView(IndexData, IndexDataSize, Name, sizeof(uint16), EGLTFJsonBufferTarget::ElementArrayBuffer);
 }
 
 FGLTFJsonAccessorIndex FGLTFStaticMeshSectionConverter::Add(FGLTFConvertBuilder& Builder, const FString& Name, TTuple<const FStaticMeshSection*, const FRawStaticIndexBuffer*> Params)
