[GLTFExporter] refactor: move conversion code-path of 16-bit indices into else block

commit 3014d76e339c559244fd0b4b6e655f902fc0f0d2
Date: Tue, 19 May 2020 04:16:11 +0200

    refactor: move conversion code-path of 16-bit indices into else block
    (cherry picked from commit c2377e19345df17eb4119f3da8d2f1169a249e49​)

#skipci

diff --git a/Engine/Plugins/Enterprise/GLTFExporter/MigratedFromGit/Source/GLTFExporter/Private/Converters/GLTFStaticMeshConverters.cpp b/Engine/Plugins/Enterprise/GLTFExporter/MigratedFromGit/Source/GLTFExporter/Private/Converters/GLTFStaticMeshConverters.cpp
--- a/Engine/Plugins/Enterprise/GLTFExporter/MigratedFromGit/Source/GLTFExporter/Private/Converters/GLTFStaticMeshConverters.cpp
+++ b/Engine/Plugins/Enterprise/GLTFExporter/MigratedFromGit/Source/GLTFExporter/Private/Converters/GLTFStaticMeshConverters.cpp
@@ -174,10 +174,12 @@ FGLTFJsonBufferViewIndex FGLTFStaticMeshIndexBufferConverter::Add(FGLTFConvertBu
 		IndexBuffer->GetCopy(Indices);
 		return Builder.AddBufferView(Indices, Name, sizeof(uint32), EGLTFJsonBufferTarget::ElementArrayBuffer);
 	}
-
-	const uint16* IndexData = IndexBuffer->AccessStream16();
-	const int32 IndexDataSize = IndexBuffer->GetIndexDataSize();
-	return Builder.AddBufferView(IndexData, IndexDataSize, Name, sizeof(uint16), EGLTFJsonBufferTarget::ElementArrayBuffer);
+	else
+	{
+		const uint16* IndexData = IndexBuffer->AccessStream16();
+		const int32 IndexDataSize = IndexBuffer->GetIndexDataSize();
+		return Builder.AddBufferView(IndexData, IndexDataSize, Name, sizeof(uint16), EGLTFJsonBufferTarget::ElementArrayBuffer);
+	}
 }
 
 FGLTFJsonAccessorIndex FGLTFStaticMeshSectionConverter::Add(FGLTFConvertBuilder& Builder, const FString& Name, TTuple<const FStaticMeshSection*, const FRawStaticIndexBuffer*> Params)
