[GLTFExporter] chore: ensure that README.md for the plugin is included in the packaged plugin

commit caa24dce477d5854b6fa68741975e45415ca09cc
Date: Wed, 20 May 2020 08:24:07 +0200

    chore: ensure that README.md for the plugin is included in the packaged plugin
    (cherry picked from commit 3574db907d79b349017bd058bd6d20d2e5e1a187​)

#skipci

diff --git a/Engine/Plugins/Enterprise/GLTFExporter/MigratedFromGit/Config/FilterPlugin.ini b/Engine/Plugins/Enterprise/GLTFExporter/MigratedFromGit/Config/FilterPlugin.ini
new file mode 100644
--- /dev/null
+++ b/Engine/Plugins/Enterprise/GLTFExporter/MigratedFromGit/Config/FilterPlugin.ini
@@ -0,0 +1,10 @@
+[FilterPlugin]
+; This section lists additional files which will be packaged along with your plugin. Paths should be listed relative to the root plugin directory, and
+; may include "...", "*", and "?" wildcards to match directories, files, and individual characters respectively.
+;
+; Examples:
+;    /README.txt
+;    /Extras/...
+;    /Binaries/ThirdParty/*.dll
+
+/README.md
