[GLTFExporter] chore: set plugin version to 0.1 to indicate that it's still in alpha

commit f3932d51849ebd481c97b408205b17e77c87c4f6
Date: Wed, 20 May 2020 12:39:02 +0200

    chore: set plugin version to 0.1 to indicate that it's still in alpha
    (cherry picked from commit e3690b4674e095cf56d3aa0f09e73b6c7d3bd229​)

#skipci

diff --git a/Engine/Plugins/Enterprise/GLTFExporter/MigratedFromGit/GLTFExporter.uplugin b/Engine/Plugins/Enterprise/GLTFExporter/MigratedFromGit/GLTFExporter.uplugin
--- a/Engine/Plugins/Enterprise/GLTFExporter/MigratedFromGit/GLTFExporter.uplugin
+++ b/Engine/Plugins/Enterprise/GLTFExporter/MigratedFromGit/GLTFExporter.uplugin
@@ -1,7 +1,7 @@
 {
 	"FileVersion": 3,
-	"Version": 1.0,
-	"VersionName": "1.0",
+	"Version": 0.1,
+	"VersionName": "0.1",
 	"FriendlyName": "glTF Exporter",
 	"Description": "An exporter for Khronos glTF 3D models.",
 	"Category": "Exporters",
