[GLTFExporter] refactor: disable SkeletalMesh and AnimSequence exporters until implemented

commit 74444280f3194152d6d46e654198d9af668a3169
Date: Thu, 21 May 2020 03:29:10 +0200

    refactor: disable SkeletalMesh and AnimSequence exporters until implemented
    (cherry picked from commit 9a52935b6e3bc65d8530cf13b8446d622c333821​)

#skipci

diff --git a/Engine/Plugins/Enterprise/GLTFExporter/MigratedFromGit/Source/GLTFExporter/Private/Exporters/GLTFAnimSequenceExporter.cpp b/Engine/Plugins/Enterprise/GLTFExporter/MigratedFromGit/Source/GLTFExporter/Private/Exporters/GLTFAnimSequenceExporter.cpp
--- a/Engine/Plugins/Enterprise/GLTFExporter/MigratedFromGit/Source/GLTFExporter/Private/Exporters/GLTFAnimSequenceExporter.cpp
+++ b/Engine/Plugins/Enterprise/GLTFExporter/MigratedFromGit/Source/GLTFExporter/Private/Exporters/GLTFAnimSequenceExporter.cpp
@@ -6,7 +6,8 @@
 UGLTFAnimSequenceExporter::UGLTFAnimSequenceExporter(const FObjectInitializer& ObjectInitializer)
 	: Super(ObjectInitializer)
 {
-	SupportedClass = UAnimSequence::StaticClass();
+	// TODO: uncomment when support is implemented
+	// SupportedClass = UAnimSequence::StaticClass();
 }
 
 bool UGLTFAnimSequenceExporter::Add(FGLTFContainerBuilder& Builder, const UObject* Object)
diff --git a/Engine/Plugins/Enterprise/GLTFExporter/MigratedFromGit/Source/GLTFExporter/Private/Exporters/GLTFSkeletalMeshExporter.cpp b/Engine/Plugins/Enterprise/GLTFExporter/MigratedFromGit/Source/GLTFExporter/Private/Exporters/GLTFSkeletalMeshExporter.cpp
--- a/Engine/Plugins/Enterprise/GLTFExporter/MigratedFromGit/Source/GLTFExporter/Private/Exporters/GLTFSkeletalMeshExporter.cpp
+++ b/Engine/Plugins/Enterprise/GLTFExporter/MigratedFromGit/Source/GLTFExporter/Private/Exporters/GLTFSkeletalMeshExporter.cpp
@@ -6,7 +6,8 @@
 UGLTFSkeletalMeshExporter::UGLTFSkeletalMeshExporter(const FObjectInitializer& ObjectInitializer)
 	: Super(ObjectInitializer)
 {
-	SupportedClass = USkeletalMesh::StaticClass();
+	// TODO: uncomment when support is implemented
+	// SupportedClass = USkeletalMesh::StaticClass();
 }
 
 bool UGLTFSkeletalMeshExporter::Add(FGLTFContainerBuilder& Builder, const UObject* Object)
