[GLTFExporter] fix: handle null assets passed to builder

commit b3e8e2e6d41900e1bbbf5a607a9568689ffddbde
Date: Thu, 21 May 2020 04:44:21 +0200

    fix: handle null assets passed to builder
    (cherry picked from commit 71b3e6a2ceb64001d9b0d62668ea238187185bd8​)

#skipci

diff --git a/Engine/Plugins/Enterprise/GLTFExporter/MigratedFromGit/Source/GLTFExporter/Private/Builders/GLTFBuilderUtility.h b/Engine/Plugins/Enterprise/GLTFExporter/MigratedFromGit/Source/GLTFExporter/Private/Builders/GLTFBuilderUtility.h
--- a/Engine/Plugins/Enterprise/GLTFExporter/MigratedFromGit/Source/GLTFExporter/Private/Builders/GLTFBuilderUtility.h
+++ b/Engine/Plugins/Enterprise/GLTFExporter/MigratedFromGit/Source/GLTFExporter/Private/Builders/GLTFBuilderUtility.h
@@ -9,8 +9,13 @@ struct FGLTFBuilderUtility
 	static FString GetMeshName(const UStaticMesh* StaticMesh, int32 LODIndex)
 	{
 		FString Name;
-		StaticMesh->GetName(Name);
-		if (LODIndex != 0) Name += TEXT("_LOD") + FString::FromInt(LODIndex);
+
+		if (StaticMesh != nullptr)
+		{
+			StaticMesh->GetName(Name);
+			if (LODIndex != 0) Name += TEXT("_LOD") + FString::FromInt(LODIndex);
+		}
+
 		return Name;
 	}
 };
diff --git a/Engine/Plugins/Enterprise/GLTFExporter/MigratedFromGit/Source/GLTFExporter/Private/Builders/GLTFConvertBuilder.cpp b/Engine/Plugins/Enterprise/GLTFExporter/MigratedFromGit/Source/GLTFExporter/Private/Builders/GLTFConvertBuilder.cpp
--- a/Engine/Plugins/Enterprise/GLTFExporter/MigratedFromGit/Source/GLTFExporter/Private/Builders/GLTFConvertBuilder.cpp
+++ b/Engine/Plugins/Enterprise/GLTFExporter/MigratedFromGit/Source/GLTFExporter/Private/Builders/GLTFConvertBuilder.cpp
@@ -5,46 +5,91 @@
 
 FGLTFJsonAccessorIndex FGLTFConvertBuilder::GetOrAddPositionAccessor(const FPositionVertexBuffer* VertexBuffer, const FString& DesiredName)
 {
+	if (VertexBuffer == nullptr)
+	{
+		return FGLTFJsonAccessorIndex(INDEX_NONE);
+	}
+
 	return PositionVertexBufferConverter.GetOrAdd(*this, DesiredName, VertexBuffer);
 }
 
 FGLTFJsonAccessorIndex FGLTFConvertBuilder::GetOrAddColorAccessor(const FColorVertexBuffer* VertexBuffer, const FString& DesiredName)
 {
+	if (VertexBuffer == nullptr)
+	{
+		return FGLTFJsonAccessorIndex(INDEX_NONE);
+	}
+
 	return ColorVertexBufferConverter.GetOrAdd(*this, DesiredName, VertexBuffer);
 }
 
 FGLTFJsonAccessorIndex FGLTFConvertBuilder::GetOrAddNormalAccessor(const FStaticMeshVertexBuffer* VertexBuffer, const FString& DesiredName)
 {
+	if (VertexBuffer == nullptr)
+	{
+		return FGLTFJsonAccessorIndex(INDEX_NONE);
+	}
+
 	return StaticMeshNormalVertexBufferConverter.GetOrAdd(*this, DesiredName, VertexBuffer);
 }
 
 FGLTFJsonAccessorIndex FGLTFConvertBuilder::GetOrAddTangentAccessor(const FStaticMeshVertexBuffer* VertexBuffer, const FString& DesiredName)
 {
+	if (VertexBuffer == nullptr)
+	{
+		return FGLTFJsonAccessorIndex(INDEX_NONE);
+	}
+
 	return StaticMeshTangentVertexBufferConverter.GetOrAdd(*this, DesiredName, VertexBuffer);
 }
 
 FGLTFJsonAccessorIndex FGLTFConvertBuilder::GetOrAddUVAccessor(const FStaticMeshVertexBuffer* VertexBuffer, int32 UVIndex, const FString& DesiredName)
 {
+	if (VertexBuffer == nullptr)
+	{
+		return FGLTFJsonAccessorIndex(INDEX_NONE);
+	}
+
 	return StaticMeshUVVertexBufferConverter.GetOrAdd(*this, DesiredName, VertexBuffer, UVIndex);
 }
 
 FGLTFJsonBufferViewIndex FGLTFConvertBuilder::GetOrAddIndexBufferView(const FRawStaticIndexBuffer* IndexBuffer, const FString& DesiredName)
 {
+	if (IndexBuffer == nullptr)
+	{
+		return FGLTFJsonBufferViewIndex(INDEX_NONE);
+	}
+
 	return StaticMeshIndexBufferConverter.GetOrAdd(*this, DesiredName, IndexBuffer);
 }
 
 FGLTFJsonAccessorIndex FGLTFConvertBuilder::GetOrAddIndexAccessor(const FStaticMeshSection* MeshSection, const FRawStaticIndexBuffer* IndexBuffer, const FString& DesiredName)
 {
+	if (MeshSection == nullptr)
+	{
+		return FGLTFJsonAccessorIndex(INDEX_NONE);
+	}
+
 	return StaticMeshSectionConverter.GetOrAdd(*this, DesiredName, MeshSection, IndexBuffer);
 }
 
 FGLTFJsonMeshIndex FGLTFConvertBuilder::GetOrAddMesh(const UStaticMesh* StaticMesh, int32 LODIndex, const FColorVertexBuffer* OverrideVertexColors, const FString& DesiredName)
 {
+	if (StaticMesh == nullptr)
+	{
+		return FGLTFJsonMeshIndex(INDEX_NONE);
+	}
+
 	return StaticMeshConverter.GetOrAdd(*this, DesiredName.IsEmpty() ? FGLTFBuilderUtility::GetMeshName(StaticMesh, LODIndex) : DesiredName, StaticMesh, LODIndex, OverrideVertexColors);
 }
 
 FGLTFJsonMeshIndex FGLTFConvertBuilder::GetOrAddMesh(const UStaticMeshComponent* StaticMeshComponent, const FString& DesiredName)
 {
+	if (StaticMeshComponent == nullptr)
+	{
+		return FGLTFJsonMeshIndex(INDEX_NONE);
+	}
+
 	const UStaticMesh* StaticMesh = StaticMeshComponent->GetStaticMesh();
 	const int32 LODIndex = StaticMeshComponent->ForcedLodModel > 0 ? StaticMeshComponent->ForcedLodModel - 1 : /* auto-select */ 0;
 	const FColorVertexBuffer* OverrideVertexColors = LODIndex < StaticMeshComponent->LODData.Num() ? StaticMeshComponent->LODData[LODIndex].OverrideVertexColors : nullptr;
@@ -54,15 +99,30 @@ FGLTFJsonMeshIndex FGLTFConvertBuilder::GetOrAddMesh(const UStaticMeshComponent*
 
 FGLTFJsonNodeIndex FGLTFConvertBuilder::GetOrAddNode(const USceneComponent* SceneComponent, bool bSelectedOnly, bool bRootNode, const FString& DesiredName)
 {
+	if (SceneComponent == nullptr)
+	{
+		return FGLTFJsonNodeIndex(INDEX_NONE);
+	}
+
 	return SceneComponentConverter.GetOrAdd(*this, DesiredName, SceneComponent, bSelectedOnly, bRootNode);
 }
 
 FGLTFJsonSceneIndex FGLTFConvertBuilder::GetOrAddScene(const ULevel* Level, bool bSelectedOnly, const FString& DesiredName)
 {
+	if (Level == nullptr)
+	{
+		return FGLTFJsonSceneIndex(INDEX_NONE);
+	}
+
 	return LevelConverter.GetOrAdd(*this, DesiredName, Level, bSelectedOnly);
 }
 
 FGLTFJsonSceneIndex FGLTFConvertBuilder::GetOrAddScene(const UWorld* World, bool bSelectedOnly, const FString& DesiredName)
 {
+	if (World == nullptr)
+	{
+		return FGLTFJsonSceneIndex(INDEX_NONE);
+	}
+
 	return GetOrAddScene(World->PersistentLevel, bSelectedOnly, DesiredName.IsEmpty() ? World->GetName() : DesiredName);
 }
