[GLTFExporter] refactor: TGLTFConverter template base class to automatically use TTuple internally

commit aedbf531e337e0157bdb1317665147331f2e4f8d
Date: Thu, 21 May 2020 23:33:56 +0200

    refactor: TGLTFConverter template base class to automatically use TTuple internally
    (cherry picked from commit d75d96d1fcde2069558a1217e4bcc5ee508cfbda​)

#skipci

diff --git a/Engine/Plugins/Enterprise/GLTFExporter/MigratedFromGit/Source/GLTFExporter/Private/Converters/GLTFConverter.h b/Engine/Plugins/Enterprise/GLTFExporter/MigratedFromGit/Source/GLTFExporter/Private/Converters/GLTFConverter.h
--- a/Engine/Plugins/Enterprise/GLTFExporter/MigratedFromGit/Source/GLTFExporter/Private/Converters/GLTFConverter.h
+++ b/Engine/Plugins/Enterprise/GLTFExporter/MigratedFromGit/Source/GLTFExporter/Private/Converters/GLTFConverter.h
@@ -6,28 +6,29 @@
 
 class FGLTFConvertBuilder;
 
-template <class IndexType, class KeyType>
+template <typename IndexType, typename... ArgTypes>
 class TGLTFConverter
 {
+	typedef TTuple<ArgTypes...> KeyType;
+
 public:
+
 	virtual ~TGLTFConverter() = default;
 
-	template <typename... ArgsType>
-    FORCEINLINE IndexType Get(ArgsType&&... Args) const
+	FORCEINLINE IndexType Get(ArgTypes&&... Args) const
 	{
-		const KeyType Key(Forward<ArgsType>(Args)...);
+		const KeyType Key(Forward<ArgTypes>(Args)...);
 		return IndexLookup.FindRef(Key);
 	}
 
-	template <typename... ArgsType>
-    FORCEINLINE IndexType GetOrAdd(FGLTFConvertBuilder& Builder, const FString& DesiredName, ArgsType&&... Args)
+	FORCEINLINE IndexType GetOrAdd(FGLTFConvertBuilder& Builder, const FString& DesiredName, ArgTypes... Args)
 	{
-		const KeyType Key(Forward<ArgsType>(Args)...);
+		const KeyType Key(Forward<ArgTypes>(Args)...);
 
 		IndexType Index = IndexLookup.FindRef(Key);
 		if (Index == INDEX_NONE)
 		{
-			Index = Add(Builder, DesiredName, Key);
+			Index = Add(Builder, DesiredName, Forward<ArgTypes>(Args)...);
 			IndexLookup.Add(Key, Index);
 		}
 
@@ -36,7 +37,7 @@ public:
 
 private:
 
-	TMap<KeyType, IndexType> IndexLookup;
+	virtual IndexType Add(FGLTFConvertBuilder& Builder, const FString& Name, ArgTypes... Args) = 0;
 
-	virtual IndexType Add(FGLTFConvertBuilder& Builder, const FString& Name, KeyType Key) = 0;
+	TMap<KeyType, IndexType> IndexLookup;
 };
diff --git a/Engine/Plugins/Enterprise/GLTFExporter/MigratedFromGit/Source/GLTFExporter/Private/Converters/GLTFLevelConverters.cpp b/Engine/Plugins/Enterprise/GLTFExporter/MigratedFromGit/Source/GLTFExporter/Private/Converters/GLTFLevelConverters.cpp
--- a/Engine/Plugins/Enterprise/GLTFExporter/MigratedFromGit/Source/GLTFExporter/Private/Converters/GLTFLevelConverters.cpp
+++ b/Engine/Plugins/Enterprise/GLTFExporter/MigratedFromGit/Source/GLTFExporter/Private/Converters/GLTFLevelConverters.cpp
@@ -4,12 +4,8 @@
 #include "Builders/GLTFContainerBuilder.h"
 #include "Converters/GLTFConverterUtility.h"
 
-FGLTFJsonNodeIndex FGLTFSceneComponentConverter::Add(FGLTFConvertBuilder& Builder, const FString& Name, TTuple<const USceneComponent*, bool, bool> Params)
+FGLTFJsonNodeIndex FGLTFSceneComponentConverter::Add(FGLTFConvertBuilder& Builder, const FString& Name, const USceneComponent* SceneComponent, bool bSelectedOnly, bool bRootNode)
 {
-	const USceneComponent* SceneComponent = Params.Get<0>();
-	const bool bSelectedOnly = Params.Get<1>();
-	const bool bRootNode = Params.Get<2>();
-
 	const AActor* Owner = SceneComponent->GetOwner();
 	const bool bIsRootComponent = Owner != nullptr && Owner->GetRootComponent() == SceneComponent;
 	const UBlueprint* Blueprint = UBlueprint::GetBlueprintFromClass(Owner != nullptr ? Owner->GetClass() : nullptr);
@@ -54,11 +50,8 @@ FGLTFJsonNodeIndex FGLTFSceneComponentConverter::Add(FGLTFConvertBuilder& Builde
 	return Builder.AddNode(Node);
 }
 
-FGLTFJsonSceneIndex FGLTFLevelConverter::Add(FGLTFConvertBuilder& Builder, const FString& Name, TTuple<const ULevel*, bool> Params)
+FGLTFJsonSceneIndex FGLTFLevelConverter::Add(FGLTFConvertBuilder& Builder, const FString& Name, const ULevel* Level, bool bSelectedOnly)
 {
-	const ULevel* Level = Params.Get<0>();
-	const bool bSelectedOnly = Params.Get<1>();
-
 	FGLTFJsonScene Scene;
 	Scene.Name = Name;
 
diff --git a/Engine/Plugins/Enterprise/GLTFExporter/MigratedFromGit/Source/GLTFExporter/Private/Converters/GLTFLevelConverters.h b/Engine/Plugins/Enterprise/GLTFExporter/MigratedFromGit/Source/GLTFExporter/Private/Converters/GLTFLevelConverters.h
--- a/Engine/Plugins/Enterprise/GLTFExporter/MigratedFromGit/Source/GLTFExporter/Private/Converters/GLTFLevelConverters.h
+++ b/Engine/Plugins/Enterprise/GLTFExporter/MigratedFromGit/Source/GLTFExporter/Private/Converters/GLTFLevelConverters.h
@@ -6,12 +6,12 @@
 #include "Converters/GLTFConverter.h"
 #include "Engine.h"
 
-class FGLTFSceneComponentConverter final : public TGLTFConverter<FGLTFJsonNodeIndex, TTuple<const USceneComponent*, bool, bool>>
+class FGLTFSceneComponentConverter final : public TGLTFConverter<FGLTFJsonNodeIndex, const USceneComponent*, bool, bool>
 {
-	FGLTFJsonNodeIndex Add(FGLTFConvertBuilder& Builder, const FString& Name, TTuple<const USceneComponent*, bool, bool> Params) override;
+	FGLTFJsonNodeIndex Add(FGLTFConvertBuilder& Builder, const FString& Name, const USceneComponent* SceneComponent, bool bSelectedOnly, bool bRootNode) override;
 };
 
-class FGLTFLevelConverter final : public TGLTFConverter<FGLTFJsonSceneIndex, TTuple<const ULevel*, bool>>
+class FGLTFLevelConverter final : public TGLTFConverter<FGLTFJsonSceneIndex, const ULevel*, bool>
 {
-	FGLTFJsonSceneIndex Add(FGLTFConvertBuilder& Builder, const FString& Name, TTuple<const ULevel*, bool> Params) override;
+	FGLTFJsonSceneIndex Add(FGLTFConvertBuilder& Builder, const FString& Name, const ULevel* Level, bool bSelectedOnly) override;
 };
diff --git a/Engine/Plugins/Enterprise/GLTFExporter/MigratedFromGit/Source/GLTFExporter/Private/Converters/GLTFStaticMeshConverters.cpp b/Engine/Plugins/Enterprise/GLTFExporter/MigratedFromGit/Source/GLTFExporter/Private/Converters/GLTFStaticMeshConverters.cpp
--- a/Engine/Plugins/Enterprise/GLTFExporter/MigratedFromGit/Source/GLTFExporter/Private/Converters/GLTFStaticMeshConverters.cpp
+++ b/Engine/Plugins/Enterprise/GLTFExporter/MigratedFromGit/Source/GLTFExporter/Private/Converters/GLTFStaticMeshConverters.cpp
@@ -132,11 +132,8 @@ FGLTFJsonAccessorIndex FGLTFStaticMeshTangentVertexBufferConverter::Add(FGLTFCon
 	return Builder.AddAccessor(JsonAccessor);
 }
 
-FGLTFJsonAccessorIndex FGLTFStaticMeshUVVertexBufferConverter::Add(FGLTFConvertBuilder& Builder, const FString& Name, TTuple<const FStaticMeshVertexBuffer*, int32> Params)
+FGLTFJsonAccessorIndex FGLTFStaticMeshUVVertexBufferConverter::Add(FGLTFConvertBuilder& Builder, const FString& Name, const FStaticMeshVertexBuffer* VertexBuffer, int32 UVIndex)
 {
-	const FStaticMeshVertexBuffer* VertexBuffer = Params.Get<0>();
-	const int32 UVIndex = Params.Get<1>();
-
 	const uint32 VertexCount = VertexBuffer->GetNumVertices();
 	if (VertexCount == 0 || UVIndex < 0 || VertexBuffer->GetNumTexCoords() <= static_cast<uint32>(UVIndex))
 	{
@@ -182,11 +179,8 @@ FGLTFJsonBufferViewIndex FGLTFStaticMeshIndexBufferConverter::Add(FGLTFConvertBu
 	}
 }
 
-FGLTFJsonAccessorIndex FGLTFStaticMeshSectionConverter::Add(FGLTFConvertBuilder& Builder, const FString& Name, TTuple<const FStaticMeshSection*, const FRawStaticIndexBuffer*> Params)
+FGLTFJsonAccessorIndex FGLTFStaticMeshSectionConverter::Add(FGLTFConvertBuilder& Builder, const FString& Name, const FStaticMeshSection* MeshSection, const FRawStaticIndexBuffer* IndexBuffer)
 {
-	const FStaticMeshSection* MeshSection = Params.Get<0>();
-	const FRawStaticIndexBuffer* IndexBuffer = Params.Get<1>();
-
 	const uint32 TriangleCount = MeshSection->NumTriangles;
 	if (TriangleCount == 0)
 	{
@@ -207,12 +201,8 @@ FGLTFJsonAccessorIndex FGLTFStaticMeshSectionConverter::Add(FGLTFConvertBuilder&
 	return Builder.AddAccessor(JsonAccessor);
 }
 
-FGLTFJsonMeshIndex FGLTFStaticMeshConverter::Add(FGLTFConvertBuilder& Builder, const FString& Name, TTuple<const UStaticMesh*, int32, const FColorVertexBuffer*> Params)
+FGLTFJsonMeshIndex FGLTFStaticMeshConverter::Add(FGLTFConvertBuilder& Builder, const FString& Name, const UStaticMesh* StaticMesh, int32 LODIndex, const FColorVertexBuffer* OverrideVertexColors)
 {
-	const UStaticMesh* StaticMesh = Params.Get<0>();
-	const int32 LODIndex = Params.Get<1>();
-	const FColorVertexBuffer* OverrideVertexColors = Params.Get<2>();
-
 	if (LODIndex < 0 || StaticMesh->GetNumLODs() <= LODIndex)
 	{
 		return FGLTFJsonMeshIndex(INDEX_NONE);
@@ -249,7 +239,7 @@ FGLTFJsonMeshIndex FGLTFStaticMeshConverter::Add(FGLTFConvertBuilder& Builder, c
 		JsonPrimitive.Attributes = JsonAttributes;
 
 		JsonPrimitive.Indices = Builder.GetOrAddIndexAccessor(&LODResources.Sections[SectionIndex], IndexBuffer,
-            Name + (SectionCount != 1 ? TEXT("_Indices_Section") + FString::FromInt(SectionIndex) : TEXT("_Indices")));
+			Name + (SectionCount != 1 ? TEXT("_Indices_Section") + FString::FromInt(SectionIndex) : TEXT("_Indices")));
 	}
 
 	return Builder.AddMesh(JsonMesh);
diff --git a/Engine/Plugins/Enterprise/GLTFExporter/MigratedFromGit/Source/GLTFExporter/Private/Converters/GLTFStaticMeshConverters.h b/Engine/Plugins/Enterprise/GLTFExporter/MigratedFromGit/Source/GLTFExporter/Private/Converters/GLTFStaticMeshConverters.h
--- a/Engine/Plugins/Enterprise/GLTFExporter/MigratedFromGit/Source/GLTFExporter/Private/Converters/GLTFStaticMeshConverters.h
+++ b/Engine/Plugins/Enterprise/GLTFExporter/MigratedFromGit/Source/GLTFExporter/Private/Converters/GLTFStaticMeshConverters.h
@@ -26,9 +26,9 @@ class FGLTFStaticMeshTangentVertexBufferConverter final : public TGLTFConverter<
 	FGLTFJsonAccessorIndex Add(FGLTFConvertBuilder& Builder, const FString& Name, const FStaticMeshVertexBuffer* VertexBuffer) override;
 };
 
-class FGLTFStaticMeshUVVertexBufferConverter final : public TGLTFConverter<FGLTFJsonAccessorIndex,TTuple<const FStaticMeshVertexBuffer*, int32>>
+class FGLTFStaticMeshUVVertexBufferConverter final : public TGLTFConverter<FGLTFJsonAccessorIndex, const FStaticMeshVertexBuffer*, int32>
 {
-	FGLTFJsonAccessorIndex Add(FGLTFConvertBuilder& Builder, const FString& Name, TTuple<const FStaticMeshVertexBuffer*, int32> Params) override;
+	FGLTFJsonAccessorIndex Add(FGLTFConvertBuilder& Builder, const FString& Name, const FStaticMeshVertexBuffer* VertexBuffer, int32 UVIndex) override;
 };
 
 class FGLTFStaticMeshIndexBufferConverter final : public TGLTFConverter<FGLTFJsonBufferViewIndex, const FRawStaticIndexBuffer*>
@@ -36,12 +36,12 @@ class FGLTFStaticMeshIndexBufferConverter final : public TGLTFConverter<FGLTFJso
 	FGLTFJsonBufferViewIndex Add(FGLTFConvertBuilder& Builder, const FString& Name, const FRawStaticIndexBuffer* IndexBuffer) override;
 };
 
-class FGLTFStaticMeshSectionConverter final : public TGLTFConverter<FGLTFJsonAccessorIndex, TTuple<const FStaticMeshSection*, const FRawStaticIndexBuffer*>>
+class FGLTFStaticMeshSectionConverter final : public TGLTFConverter<FGLTFJsonAccessorIndex, const FStaticMeshSection*, const FRawStaticIndexBuffer*>
 {
-	FGLTFJsonAccessorIndex Add(FGLTFConvertBuilder& Builder, const FString& Name, TTuple<const FStaticMeshSection*, const FRawStaticIndexBuffer*> Params) override;
+	FGLTFJsonAccessorIndex Add(FGLTFConvertBuilder& Builder, const FString& Name, const FStaticMeshSection* MeshSection, const FRawStaticIndexBuffer* IndexBuffer) override;
 };
 
-class FGLTFStaticMeshConverter final : public TGLTFConverter<FGLTFJsonMeshIndex, TTuple<const UStaticMesh*, int32, const FColorVertexBuffer*>>
+class FGLTFStaticMeshConverter final : public TGLTFConverter<FGLTFJsonMeshIndex, const UStaticMesh*, int32, const FColorVertexBuffer*>
 {
-	FGLTFJsonMeshIndex Add(FGLTFConvertBuilder& Builder, const FString& Name, TTuple<const UStaticMesh*, int32, const FColorVertexBuffer*> Params) override;
+	FGLTFJsonMeshIndex Add(FGLTFConvertBuilder& Builder, const FString& Name, const UStaticMesh* StaticMesh, int32 LODIndex, const FColorVertexBuffer* OverrideVertexColors) override;
 };
