[GLTFExporter] refactor: rename EGLTFExtension -> EGLTFJsonExtension

commit 16d1e6e929c780efb52f60c52f1160ef86e66e26
Date: Fri, 22 May 2020 19:29:36 +0200

    refactor: rename EGLTFExtension -> EGLTFJsonExtension
    (cherry picked from commit f8aa8cbe07b25c5943040a2b4ab6b4ffe6b00008​)

#skipci

diff --git a/Engine/Plugins/Enterprise/GLTFExporter/MigratedFromGit/Source/GLTFExporter/Private/Json/GLTFJsonEnums.h b/Engine/Plugins/Enterprise/GLTFExporter/MigratedFromGit/Source/GLTFExporter/Private/Json/GLTFJsonEnums.h
--- a/Engine/Plugins/Enterprise/GLTFExporter/MigratedFromGit/Source/GLTFExporter/Private/Json/GLTFJsonEnums.h
+++ b/Engine/Plugins/Enterprise/GLTFExporter/MigratedFromGit/Source/GLTFExporter/Private/Json/GLTFJsonEnums.h
@@ -4,7 +4,7 @@
 
 #include "CoreMinimal.h"
 
-enum class EGLTFExtension
+enum class EGLTFJsonExtension
 {
 	KHR_LightsPunctual,
 	KHR_MaterialsUnlit,
diff --git a/Engine/Plugins/Enterprise/GLTFExporter/MigratedFromGit/Source/GLTFExporter/Private/Json/GLTFJsonRoot.h b/Engine/Plugins/Enterprise/GLTFExporter/MigratedFromGit/Source/GLTFExporter/Private/Json/GLTFJsonRoot.h
--- a/Engine/Plugins/Enterprise/GLTFExporter/MigratedFromGit/Source/GLTFExporter/Private/Json/GLTFJsonRoot.h
+++ b/Engine/Plugins/Enterprise/GLTFExporter/MigratedFromGit/Source/GLTFExporter/Private/Json/GLTFJsonRoot.h
@@ -52,8 +52,8 @@ struct FGLTFJsonRoot
 {
 	FGLTFJsonAsset Asset;
 
-	TSet<EGLTFExtension> ExtensionsUsed;
-	TSet<EGLTFExtension> ExtensionsRequired;
+	TSet<EGLTFJsonExtension> ExtensionsUsed;
+	TSet<EGLTFJsonExtension> ExtensionsRequired;
 
 	FGLTFJsonSceneIndex DefaultScene;
 
@@ -81,7 +81,7 @@ struct FGLTFJsonRoot
 		{
 			JsonWriter.WriteIdentifierPrefix(TEXT("extensionsUsed"));
 			JsonWriter.WriteArrayStart();
-			for (const EGLTFExtension& Extension : ExtensionsUsed)
+			for (const EGLTFJsonExtension& Extension : ExtensionsUsed)
 			{
 				JsonWriter.WriteValue(FGLTFJsonUtility::ExtensionToString(Extension));
 			}
@@ -92,7 +92,7 @@ struct FGLTFJsonRoot
 		{
 			JsonWriter.WriteIdentifierPrefix(TEXT("extensionsRequired"));
 			JsonWriter.WriteArrayStart();
-			for (const EGLTFExtension& Extension : ExtensionsRequired)
+			for (const EGLTFJsonExtension& Extension : ExtensionsRequired)
 			{
 				JsonWriter.WriteValue(FGLTFJsonUtility::ExtensionToString(Extension));
 			}
diff --git a/Engine/Plugins/Enterprise/GLTFExporter/MigratedFromGit/Source/GLTFExporter/Private/Json/GLTFJsonUtility.h b/Engine/Plugins/Enterprise/GLTFExporter/MigratedFromGit/Source/GLTFExporter/Private/Json/GLTFJsonUtility.h
--- a/Engine/Plugins/Enterprise/GLTFExporter/MigratedFromGit/Source/GLTFExporter/Private/Json/GLTFJsonUtility.h
+++ b/Engine/Plugins/Enterprise/GLTFExporter/MigratedFromGit/Source/GLTFExporter/Private/Json/GLTFJsonUtility.h
@@ -6,14 +6,14 @@
 
 struct FGLTFJsonUtility
 {
-	static const TCHAR* ExtensionToString(EGLTFExtension Extension)
+	static const TCHAR* ExtensionToString(EGLTFJsonExtension Extension)
 	{
 		switch (Extension)
 		{
-			case EGLTFExtension::KHR_LightsPunctual:     return TEXT("KHR_lights_punctual");
-			case EGLTFExtension::KHR_MaterialsUnlit:     return TEXT("KHR_materials_unlit");
-			case EGLTFExtension::KHR_MaterialsClearCoat: return TEXT("KHR_materials_clearcoat");
-			case EGLTFExtension::KHR_MeshQuantization:   return TEXT("KHR_mesh_quantization");
+			case EGLTFJsonExtension::KHR_LightsPunctual:     return TEXT("KHR_lights_punctual");
+			case EGLTFJsonExtension::KHR_MaterialsUnlit:     return TEXT("KHR_materials_unlit");
+			case EGLTFJsonExtension::KHR_MaterialsClearCoat: return TEXT("KHR_materials_clearcoat");
+			case EGLTFJsonExtension::KHR_MeshQuantization:   return TEXT("KHR_mesh_quantization");
 			default:                                     return TEXT("unknown");
 		}
 	}
