[GLTFExporter] refactor: remove redundant default initializations of FGLTFJson*Index

commit df8beb3fb9d2c501a08e9ad644277c70edfe6d16
Date: Fri, 22 May 2020 20:02:22 +0200

    refactor: remove redundant default initializations of FGLTFJson*Index
    (cherry picked from commit a03b59f4eff39890de84d3953b1e4f4d2d924421​)

#skipci

diff --git a/Engine/Plugins/Enterprise/GLTFExporter/MigratedFromGit/Source/GLTFExporter/Private/Json/GLTFJsonAccessor.h b/Engine/Plugins/Enterprise/GLTFExporter/MigratedFromGit/Source/GLTFExporter/Private/Json/GLTFJsonAccessor.h
--- a/Engine/Plugins/Enterprise/GLTFExporter/MigratedFromGit/Source/GLTFExporter/Private/Json/GLTFJsonAccessor.h
+++ b/Engine/Plugins/Enterprise/GLTFExporter/MigratedFromGit/Source/GLTFExporter/Private/Json/GLTFJsonAccessor.h
@@ -23,8 +23,7 @@ struct FGLTFJsonAccessor
 	float Max[16];
 
 	FGLTFJsonAccessor()
-		: BufferView(INDEX_NONE)
-		, ByteOffset(0)
+		: ByteOffset(0)
 		, Count(0)
 		, Type(EGLTFJsonAccessorType::None)
 		, ComponentType(EGLTFJsonComponentType::None)
diff --git a/Engine/Plugins/Enterprise/GLTFExporter/MigratedFromGit/Source/GLTFExporter/Private/Json/GLTFJsonBufferView.h b/Engine/Plugins/Enterprise/GLTFExporter/MigratedFromGit/Source/GLTFExporter/Private/Json/GLTFJsonBufferView.h
--- a/Engine/Plugins/Enterprise/GLTFExporter/MigratedFromGit/Source/GLTFExporter/Private/Json/GLTFJsonBufferView.h
+++ b/Engine/Plugins/Enterprise/GLTFExporter/MigratedFromGit/Source/GLTFExporter/Private/Json/GLTFJsonBufferView.h
@@ -20,8 +20,7 @@ struct FGLTFJsonBufferView
 	EGLTFJsonBufferTarget Target;
 
 	FGLTFJsonBufferView()
-		: Buffer(INDEX_NONE)
-		, ByteLength(0)
+		: ByteLength(0)
 		, ByteOffset(0)
 		, ByteStride(0)
 		, Target(EGLTFJsonBufferTarget::None)
diff --git a/Engine/Plugins/Enterprise/GLTFExporter/MigratedFromGit/Source/GLTFExporter/Private/Json/GLTFJsonMesh.h b/Engine/Plugins/Enterprise/GLTFExporter/MigratedFromGit/Source/GLTFExporter/Private/Json/GLTFJsonMesh.h
--- a/Engine/Plugins/Enterprise/GLTFExporter/MigratedFromGit/Source/GLTFExporter/Private/Json/GLTFJsonMesh.h
+++ b/Engine/Plugins/Enterprise/GLTFExporter/MigratedFromGit/Source/GLTFExporter/Private/Json/GLTFJsonMesh.h
@@ -18,18 +18,6 @@ struct FGLTFJsonAttributes
 	FGLTFJsonAccessorIndex Joints0;
 	FGLTFJsonAccessorIndex Weights0;
 
-	FGLTFJsonAttributes()
-		: Position(INDEX_NONE)
-		, Normal(INDEX_NONE)
-		, Tangent(INDEX_NONE)
-		, TexCoord0(INDEX_NONE)
-		, TexCoord1(INDEX_NONE)
-		, Color0(INDEX_NONE)
-		, Joints0(INDEX_NONE)
-		, Weights0(INDEX_NONE)
-	{
-	}
-
 	template <class CharType = TCHAR, class PrintPolicy = TPrettyJsonPrintPolicy<CharType>>
 	void WriteObject(TJsonWriter<CharType, PrintPolicy>& JsonWriter) const
 	{
@@ -56,9 +44,7 @@ struct FGLTFJsonPrimitive
 	FGLTFJsonAttributes    Attributes;
 
 	FGLTFJsonPrimitive()
-		: Indices(INDEX_NONE)
-		, Material(INDEX_NONE)
-		, Mode(EGLTFJsonPrimitiveMode::None)
+		: Mode(EGLTFJsonPrimitiveMode::None)
 	{
 	}
 
diff --git a/Engine/Plugins/Enterprise/GLTFExporter/MigratedFromGit/Source/GLTFExporter/Private/Json/GLTFJsonNode.h b/Engine/Plugins/Enterprise/GLTFExporter/MigratedFromGit/Source/GLTFExporter/Private/Json/GLTFJsonNode.h
--- a/Engine/Plugins/Enterprise/GLTFExporter/MigratedFromGit/Source/GLTFExporter/Private/Json/GLTFJsonNode.h
+++ b/Engine/Plugins/Enterprise/GLTFExporter/MigratedFromGit/Source/GLTFExporter/Private/Json/GLTFJsonNode.h
@@ -25,9 +25,6 @@ struct FGLTFJsonNode
 		: Translation(FGLTFJsonVector3::Zero)
 		, Rotation(FGLTFJsonQuaternion::Identity)
 		, Scale(FGLTFJsonVector3::One)
-		, Camera(INDEX_NONE)
-		, Skin(INDEX_NONE)
-		, Mesh(INDEX_NONE)
 	{
 	}
 
diff --git a/Engine/Plugins/Enterprise/GLTFExporter/MigratedFromGit/Source/GLTFExporter/Private/Json/GLTFJsonRoot.h b/Engine/Plugins/Enterprise/GLTFExporter/MigratedFromGit/Source/GLTFExporter/Private/Json/GLTFJsonRoot.h
--- a/Engine/Plugins/Enterprise/GLTFExporter/MigratedFromGit/Source/GLTFExporter/Private/Json/GLTFJsonRoot.h
+++ b/Engine/Plugins/Enterprise/GLTFExporter/MigratedFromGit/Source/GLTFExporter/Private/Json/GLTFJsonRoot.h
@@ -64,11 +64,6 @@ struct FGLTFJsonRoot
 	TArray<FGLTFJsonNode>       Nodes;
 	TArray<FGLTFJsonScene>      Scenes;
 
-	FGLTFJsonRoot()
-		: DefaultScene(INDEX_NONE)
-	{
-	}
-
 	template <class CharType = TCHAR, class PrintPolicy = TPrettyJsonPrintPolicy<CharType>>
 	void WriteObject(TJsonWriter<CharType, PrintPolicy>& JsonWriter) const
 	{
