[GLTFExporter] refactor(json): simplify json utility function names and unify does with same behavior

commit 4b0ed666af672dd86ca997ef3e792d283049e5c9
Date: Fri, 22 May 2020 20:44:58 +0200

    refactor(json): simplify json utility function names and unify does with same behavior
    (cherry picked from commit 25811a8ec2563bfa4a8d8dc8b1f246ab0a5b0285​)

#skipci

diff --git a/Engine/Plugins/Enterprise/GLTFExporter/MigratedFromGit/Source/GLTFExporter/Private/Json/GLTFJsonAccessor.h b/Engine/Plugins/Enterprise/GLTFExporter/MigratedFromGit/Source/GLTFExporter/Private/Json/GLTFJsonAccessor.h
--- a/Engine/Plugins/Enterprise/GLTFExporter/MigratedFromGit/Source/GLTFExporter/Private/Json/GLTFJsonAccessor.h
+++ b/Engine/Plugins/Enterprise/GLTFExporter/MigratedFromGit/Source/GLTFExporter/Private/Json/GLTFJsonAccessor.h
@@ -52,8 +52,8 @@ struct FGLTFJsonAccessor
 		}
 
 		JsonWriter.WriteValue(TEXT("count"), Count);
-		JsonWriter.WriteValue(TEXT("type"), FGLTFJsonUtility::AccessorTypeToString(Type));
-		JsonWriter.WriteValue(TEXT("componentType"), FGLTFJsonUtility::ComponentTypeToNumber(ComponentType));
+		JsonWriter.WriteValue(TEXT("type"), FGLTFJsonUtility::ToString(Type));
+		JsonWriter.WriteValue(TEXT("componentType"), FGLTFJsonUtility::ToNumber(ComponentType));
 
 		if (bNormalized)
 		{
diff --git a/Engine/Plugins/Enterprise/GLTFExporter/MigratedFromGit/Source/GLTFExporter/Private/Json/GLTFJsonBufferView.h b/Engine/Plugins/Enterprise/GLTFExporter/MigratedFromGit/Source/GLTFExporter/Private/Json/GLTFJsonBufferView.h
--- a/Engine/Plugins/Enterprise/GLTFExporter/MigratedFromGit/Source/GLTFExporter/Private/Json/GLTFJsonBufferView.h
+++ b/Engine/Plugins/Enterprise/GLTFExporter/MigratedFromGit/Source/GLTFExporter/Private/Json/GLTFJsonBufferView.h
@@ -53,7 +53,7 @@ struct FGLTFJsonBufferView
 
 		if (Target != EGLTFJsonBufferTarget::None)
 		{
-			JsonWriter.WriteValue(TEXT("target"), FGLTFJsonUtility::BufferTargetToNumber(Target));
+			JsonWriter.WriteValue(TEXT("target"), FGLTFJsonUtility::ToNumber(Target));
 		}
 
 		JsonWriter.WriteObjectEnd();
diff --git a/Engine/Plugins/Enterprise/GLTFExporter/MigratedFromGit/Source/GLTFExporter/Private/Json/GLTFJsonMesh.h b/Engine/Plugins/Enterprise/GLTFExporter/MigratedFromGit/Source/GLTFExporter/Private/Json/GLTFJsonMesh.h
--- a/Engine/Plugins/Enterprise/GLTFExporter/MigratedFromGit/Source/GLTFExporter/Private/Json/GLTFJsonMesh.h
+++ b/Engine/Plugins/Enterprise/GLTFExporter/MigratedFromGit/Source/GLTFExporter/Private/Json/GLTFJsonMesh.h
@@ -68,7 +68,7 @@ struct FGLTFJsonPrimitive
 
 		if (Mode != EGLTFJsonPrimitiveMode::None)
 		{
-			JsonWriter.WriteValue(TEXT("mode"), FGLTFJsonUtility::PrimitiveModeToNumber(Mode));
+			JsonWriter.WriteValue(TEXT("mode"), FGLTFJsonUtility::ToNumber(Mode));
 		}
 
 		JsonWriter.WriteObjectEnd();
diff --git a/Engine/Plugins/Enterprise/GLTFExporter/MigratedFromGit/Source/GLTFExporter/Private/Json/GLTFJsonRoot.h b/Engine/Plugins/Enterprise/GLTFExporter/MigratedFromGit/Source/GLTFExporter/Private/Json/GLTFJsonRoot.h
--- a/Engine/Plugins/Enterprise/GLTFExporter/MigratedFromGit/Source/GLTFExporter/Private/Json/GLTFJsonRoot.h
+++ b/Engine/Plugins/Enterprise/GLTFExporter/MigratedFromGit/Source/GLTFExporter/Private/Json/GLTFJsonRoot.h
@@ -78,7 +78,7 @@ struct FGLTFJsonRoot
 			JsonWriter.WriteArrayStart();
 			for (const EGLTFJsonExtension& Extension : ExtensionsUsed)
 			{
-				JsonWriter.WriteValue(FGLTFJsonUtility::ExtensionToString(Extension));
+				JsonWriter.WriteValue(FGLTFJsonUtility::ToString(Extension));
 			}
 			JsonWriter.WriteArrayEnd();
 		}
@@ -89,7 +89,7 @@ struct FGLTFJsonRoot
 			JsonWriter.WriteArrayStart();
 			for (const EGLTFJsonExtension& Extension : ExtensionsRequired)
 			{
-				JsonWriter.WriteValue(FGLTFJsonUtility::ExtensionToString(Extension));
+				JsonWriter.WriteValue(FGLTFJsonUtility::ToString(Extension));
 			}
 			JsonWriter.WriteArrayEnd();
 		}
diff --git a/Engine/Plugins/Enterprise/GLTFExporter/MigratedFromGit/Source/GLTFExporter/Private/Json/GLTFJsonUtility.h b/Engine/Plugins/Enterprise/GLTFExporter/MigratedFromGit/Source/GLTFExporter/Private/Json/GLTFJsonUtility.h
--- a/Engine/Plugins/Enterprise/GLTFExporter/MigratedFromGit/Source/GLTFExporter/Private/Json/GLTFJsonUtility.h
+++ b/Engine/Plugins/Enterprise/GLTFExporter/MigratedFromGit/Source/GLTFExporter/Private/Json/GLTFJsonUtility.h
@@ -6,21 +6,27 @@
 
 struct FGLTFJsonUtility
 {
-	static const TCHAR* ExtensionToString(EGLTFJsonExtension Extension)
+	template <typename EnumType>
+	static int32 ToNumber(EnumType Value)
 	{
-		switch (Extension)
+		return static_cast<int32>(Value);
+	}
+
+	static const TCHAR* ToString(EGLTFJsonExtension Value)
+	{
+		switch (Value)
 		{
 			case EGLTFJsonExtension::KHR_LightsPunctual:     return TEXT("KHR_lights_punctual");
 			case EGLTFJsonExtension::KHR_MaterialsUnlit:     return TEXT("KHR_materials_unlit");
 			case EGLTFJsonExtension::KHR_MaterialsClearCoat: return TEXT("KHR_materials_clearcoat");
 			case EGLTFJsonExtension::KHR_MeshQuantization:   return TEXT("KHR_mesh_quantization");
-			default:                                     return TEXT("unknown");
+			default:                                         return TEXT("unknown");
 		}
 	}
 
-	static const TCHAR* AlphaModeToString(EGLTFJsonAlphaMode Mode)
+	static const TCHAR* ToString(EGLTFJsonAlphaMode Value)
 	{
-		switch (Mode)
+		switch (Value)
 		{
 			case EGLTFJsonAlphaMode::Opaque: return TEXT("OPAQUE");
 			case EGLTFJsonAlphaMode::Blend:  return TEXT("BLEND");
@@ -29,24 +35,9 @@ struct FGLTFJsonUtility
 		}
 	}
 
-	static int32 PrimitiveModeToNumber(EGLTFJsonPrimitiveMode Mode)
-	{
-		return static_cast<int32>(Mode);
-	}
-
-	static int32 FilterToNumber(EGLTFJsonSamplerFilter Filter)
-	{
-		return static_cast<int32>(Filter);
-	}
-
-	static int32 WrapModeToNumber(EGLTFJsonSamplerWrap Wrap)
+	static const TCHAR* ToString(EGLTFJsonAccessorType Value)
 	{
-		return static_cast<int32>(Wrap);
-	}
-
-	static const TCHAR* AccessorTypeToString(EGLTFJsonAccessorType Type)
-	{
-		switch (Type)
+		switch (Value)
 		{
 			case EGLTFJsonAccessorType::Scalar: return TEXT("SCALAR");
 			case EGLTFJsonAccessorType::Vec2:   return TEXT("VEC2");
@@ -58,14 +49,4 @@ struct FGLTFJsonUtility
 			default:                            return TEXT("UNKNOWN");
 		}
 	}
-
-	static int32 ComponentTypeToNumber(EGLTFJsonComponentType Type)
-	{
-		return static_cast<int32>(Type);
-	}
-
-	static int32 BufferTargetToNumber(EGLTFJsonBufferTarget Target)
-	{
-		return static_cast<int32>(Target);
-	}
 };
