[GLTFExporter] refactor(json): rename GLTFJsonUtility::ToNumber -> ToInteger

commit 3ebefefa75b31af774e1d3c1ac31afba91958591
Date: Sat, 23 May 2020 14:20:26 +0200

    refactor(json): rename GLTFJsonUtility::ToNumber -> ToInteger
    (cherry picked from commit 6f855c216168e6d30925b0b420940e54c179f185​)

#skipci

diff --git a/Engine/Plugins/Enterprise/GLTFExporter/MigratedFromGit/Source/GLTFExporter/Private/Json/GLTFJsonAccessor.h b/Engine/Plugins/Enterprise/GLTFExporter/MigratedFromGit/Source/GLTFExporter/Private/Json/GLTFJsonAccessor.h
--- a/Engine/Plugins/Enterprise/GLTFExporter/MigratedFromGit/Source/GLTFExporter/Private/Json/GLTFJsonAccessor.h
+++ b/Engine/Plugins/Enterprise/GLTFExporter/MigratedFromGit/Source/GLTFExporter/Private/Json/GLTFJsonAccessor.h
@@ -53,7 +53,7 @@ struct FGLTFJsonAccessor
 
 		JsonWriter.WriteValue(TEXT("count"), Count);
 		JsonWriter.WriteValue(TEXT("type"), FGLTFJsonUtility::ToString(Type));
-		JsonWriter.WriteValue(TEXT("componentType"), FGLTFJsonUtility::ToNumber(ComponentType));
+		JsonWriter.WriteValue(TEXT("componentType"), FGLTFJsonUtility::ToInteger(ComponentType));
 
 		if (bNormalized)
 		{
diff --git a/Engine/Plugins/Enterprise/GLTFExporter/MigratedFromGit/Source/GLTFExporter/Private/Json/GLTFJsonBufferView.h b/Engine/Plugins/Enterprise/GLTFExporter/MigratedFromGit/Source/GLTFExporter/Private/Json/GLTFJsonBufferView.h
--- a/Engine/Plugins/Enterprise/GLTFExporter/MigratedFromGit/Source/GLTFExporter/Private/Json/GLTFJsonBufferView.h
+++ b/Engine/Plugins/Enterprise/GLTFExporter/MigratedFromGit/Source/GLTFExporter/Private/Json/GLTFJsonBufferView.h
@@ -53,7 +53,7 @@ struct FGLTFJsonBufferView
 
 		if (Target != EGLTFJsonBufferTarget::None)
 		{
-			JsonWriter.WriteValue(TEXT("target"), FGLTFJsonUtility::ToNumber(Target));
+			JsonWriter.WriteValue(TEXT("target"), FGLTFJsonUtility::ToInteger(Target));
 		}
 
 		JsonWriter.WriteObjectEnd();
diff --git a/Engine/Plugins/Enterprise/GLTFExporter/MigratedFromGit/Source/GLTFExporter/Private/Json/GLTFJsonMesh.h b/Engine/Plugins/Enterprise/GLTFExporter/MigratedFromGit/Source/GLTFExporter/Private/Json/GLTFJsonMesh.h
--- a/Engine/Plugins/Enterprise/GLTFExporter/MigratedFromGit/Source/GLTFExporter/Private/Json/GLTFJsonMesh.h
+++ b/Engine/Plugins/Enterprise/GLTFExporter/MigratedFromGit/Source/GLTFExporter/Private/Json/GLTFJsonMesh.h
@@ -68,7 +68,7 @@ struct FGLTFJsonPrimitive
 
 		if (Mode != EGLTFJsonPrimitiveMode::None)
 		{
-			JsonWriter.WriteValue(TEXT("mode"), FGLTFJsonUtility::ToNumber(Mode));
+			JsonWriter.WriteValue(TEXT("mode"), FGLTFJsonUtility::ToInteger(Mode));
 		}
 
 		JsonWriter.WriteObjectEnd();
diff --git a/Engine/Plugins/Enterprise/GLTFExporter/MigratedFromGit/Source/GLTFExporter/Private/Json/GLTFJsonUtility.h b/Engine/Plugins/Enterprise/GLTFExporter/MigratedFromGit/Source/GLTFExporter/Private/Json/GLTFJsonUtility.h
--- a/Engine/Plugins/Enterprise/GLTFExporter/MigratedFromGit/Source/GLTFExporter/Private/Json/GLTFJsonUtility.h
+++ b/Engine/Plugins/Enterprise/GLTFExporter/MigratedFromGit/Source/GLTFExporter/Private/Json/GLTFJsonUtility.h
@@ -7,7 +7,7 @@
 struct FGLTFJsonUtility
 {
 	template <typename EnumType>
-	static int32 ToNumber(EnumType Value)
+	static int32 ToInteger(EnumType Value)
 	{
 		return static_cast<int32>(Value);
 	}
