[GLTFExporter] fix(json): write float as exact string representation to avoid gltf validation errors

commit 33100421b3f9e9991ff85ffed997d5cb83851f8c
Date: Sat, 23 May 2020 14:35:33 +0200

    fix(json): write float as exact string representation to avoid gltf validation errors
    (cherry picked from commit af2409940d947479cde3085ce1a83e39577700e9​)

#skipci

diff --git a/Engine/Plugins/Enterprise/GLTFExporter/MigratedFromGit/Source/GLTFExporter/Private/Json/GLTFJsonAccessor.h b/Engine/Plugins/Enterprise/GLTFExporter/MigratedFromGit/Source/GLTFExporter/Private/Json/GLTFJsonAccessor.h
--- a/Engine/Plugins/Enterprise/GLTFExporter/MigratedFromGit/Source/GLTFExporter/Private/Json/GLTFJsonAccessor.h
+++ b/Engine/Plugins/Enterprise/GLTFExporter/MigratedFromGit/Source/GLTFExporter/Private/Json/GLTFJsonAccessor.h
@@ -65,14 +65,14 @@ struct FGLTFJsonAccessor
 			JsonWriter.WriteArrayStart(TEXT("min"));
 			for (int32 ComponentIndex = 0; ComponentIndex < MinMaxLength; ++ComponentIndex)
 			{
-				JsonWriter.WriteValue(Min[ComponentIndex]);
+				FGLTFJsonUtility::WriteExactValue(JsonWriter, Min[ComponentIndex]);
 			}
 			JsonWriter.WriteArrayEnd();
 
 			JsonWriter.WriteArrayStart(TEXT("max"));
 			for (int32 ComponentIndex = 0; ComponentIndex < MinMaxLength; ++ComponentIndex)
 			{
-				JsonWriter.WriteValue(Max[ComponentIndex]);
+				FGLTFJsonUtility::WriteExactValue(JsonWriter, Max[ComponentIndex]);
 			}
 			JsonWriter.WriteArrayEnd();
 		}
diff --git a/Engine/Plugins/Enterprise/GLTFExporter/MigratedFromGit/Source/GLTFExporter/Private/Json/GLTFJsonQuaternion.h b/Engine/Plugins/Enterprise/GLTFExporter/MigratedFromGit/Source/GLTFExporter/Private/Json/GLTFJsonQuaternion.h
--- a/Engine/Plugins/Enterprise/GLTFExporter/MigratedFromGit/Source/GLTFExporter/Private/Json/GLTFJsonQuaternion.h
+++ b/Engine/Plugins/Enterprise/GLTFExporter/MigratedFromGit/Source/GLTFExporter/Private/Json/GLTFJsonQuaternion.h
@@ -2,6 +2,7 @@
 
 #pragma once
 
+#include "Json/GLTFJsonUtility.h"
 #include "Serialization/JsonSerializer.h"
 
 struct FGLTFJsonQuaternion
@@ -19,10 +20,10 @@ struct FGLTFJsonQuaternion
 	void WriteArray(TJsonWriter<CharType, PrintPolicy>& JsonWriter) const
 	{
 		JsonWriter.WriteArrayStart();
-		JsonWriter.WriteValue(X);
-		JsonWriter.WriteValue(Y);
-		JsonWriter.WriteValue(Z);
-		JsonWriter.WriteValue(W);
+		FGLTFJsonUtility::WriteExactValue(JsonWriter, X);
+		FGLTFJsonUtility::WriteExactValue(JsonWriter, Y);
+		FGLTFJsonUtility::WriteExactValue(JsonWriter, Z);
+		FGLTFJsonUtility::WriteExactValue(JsonWriter, W);
 		JsonWriter.WriteArrayEnd();
 	}
 
diff --git a/Engine/Plugins/Enterprise/GLTFExporter/MigratedFromGit/Source/GLTFExporter/Private/Json/GLTFJsonUtility.h b/Engine/Plugins/Enterprise/GLTFExporter/MigratedFromGit/Source/GLTFExporter/Private/Json/GLTFJsonUtility.h
--- a/Engine/Plugins/Enterprise/GLTFExporter/MigratedFromGit/Source/GLTFExporter/Private/Json/GLTFJsonUtility.h
+++ b/Engine/Plugins/Enterprise/GLTFExporter/MigratedFromGit/Source/GLTFExporter/Private/Json/GLTFJsonUtility.h
@@ -49,4 +49,11 @@ struct FGLTFJsonUtility
 			default:                            return TEXT("UNKNOWN");
 		}
 	}
+
+	template <class CharType = TCHAR, class PrintPolicy = TPrettyJsonPrintPolicy<CharType>>
+    static void WriteExactValue(TJsonWriter<CharType, PrintPolicy>& JsonWriter, float Value)
+	{
+		FString ExactStringRepresentation = FString::Printf(TEXT("%.9g"), Value);
+		JsonWriter.WriteRawJSONValue(ExactStringRepresentation);
+	}
 };
diff --git a/Engine/Plugins/Enterprise/GLTFExporter/MigratedFromGit/Source/GLTFExporter/Private/Json/GLTFJsonVector2.h b/Engine/Plugins/Enterprise/GLTFExporter/MigratedFromGit/Source/GLTFExporter/Private/Json/GLTFJsonVector2.h
--- a/Engine/Plugins/Enterprise/GLTFExporter/MigratedFromGit/Source/GLTFExporter/Private/Json/GLTFJsonVector2.h
+++ b/Engine/Plugins/Enterprise/GLTFExporter/MigratedFromGit/Source/GLTFExporter/Private/Json/GLTFJsonVector2.h
@@ -2,6 +2,7 @@
 
 #pragma once
 
+#include "Json/GLTFJsonUtility.h"
 #include "Serialization/JsonSerializer.h"
 
 struct FGLTFJsonVector2
@@ -20,8 +21,8 @@ struct FGLTFJsonVector2
 	void WriteArray(TJsonWriter<CharType, PrintPolicy>& JsonWriter) const
 	{
 		JsonWriter.WriteArrayStart();
-		JsonWriter.WriteValue(X);
-		JsonWriter.WriteValue(Y);
+		FGLTFJsonUtility::WriteExactValue(JsonWriter, X);
+		FGLTFJsonUtility::WriteExactValue(JsonWriter, Y);
 		JsonWriter.WriteArrayEnd();
 	}
 
diff --git a/Engine/Plugins/Enterprise/GLTFExporter/MigratedFromGit/Source/GLTFExporter/Private/Json/GLTFJsonVector3.h b/Engine/Plugins/Enterprise/GLTFExporter/MigratedFromGit/Source/GLTFExporter/Private/Json/GLTFJsonVector3.h
--- a/Engine/Plugins/Enterprise/GLTFExporter/MigratedFromGit/Source/GLTFExporter/Private/Json/GLTFJsonVector3.h
+++ b/Engine/Plugins/Enterprise/GLTFExporter/MigratedFromGit/Source/GLTFExporter/Private/Json/GLTFJsonVector3.h
@@ -2,6 +2,8 @@
 
 #pragma once
 
+#include "Json/GLTFJsonVector2.h"
+#include "Json/GLTFJsonUtility.h"
 #include "Serialization/JsonSerializer.h"
 
 struct FGLTFJsonVector3
@@ -16,13 +18,18 @@ struct FGLTFJsonVector3
 	{
 	}
 
+	FGLTFJsonVector3(const FGLTFJsonVector2& Vector3, float Z)
+        : X(Vector3.X), Y(Vector3.Y), Z(Z)
+	{
+	}
+
 	template <class CharType = TCHAR, class PrintPolicy = TPrettyJsonPrintPolicy<CharType>>
 	void WriteArray(TJsonWriter<CharType, PrintPolicy>& JsonWriter) const
 	{
 		JsonWriter.WriteArrayStart();
-		JsonWriter.WriteValue(X);
-		JsonWriter.WriteValue(Y);
-		JsonWriter.WriteValue(Z);
+		FGLTFJsonUtility::WriteExactValue(JsonWriter, X);
+		FGLTFJsonUtility::WriteExactValue(JsonWriter, Y);
+		FGLTFJsonUtility::WriteExactValue(JsonWriter, Z);
 		JsonWriter.WriteArrayEnd();
 	}
 
diff --git a/Engine/Plugins/Enterprise/GLTFExporter/MigratedFromGit/Source/GLTFExporter/Private/Json/GLTFJsonVector4.h b/Engine/Plugins/Enterprise/GLTFExporter/MigratedFromGit/Source/GLTFExporter/Private/Json/GLTFJsonVector4.h
--- a/Engine/Plugins/Enterprise/GLTFExporter/MigratedFromGit/Source/GLTFExporter/Private/Json/GLTFJsonVector4.h
+++ b/Engine/Plugins/Enterprise/GLTFExporter/MigratedFromGit/Source/GLTFExporter/Private/Json/GLTFJsonVector4.h
@@ -3,6 +3,7 @@
 #pragma once
 
 #include "Json/GLTFJsonVector3.h"
+#include "Json/GLTFJsonUtility.h"
 #include "Serialization/JsonSerializer.h"
 
 struct FGLTFJsonVector4
@@ -26,10 +27,10 @@ struct FGLTFJsonVector4
 	void WriteArray(TJsonWriter<CharType, PrintPolicy>& JsonWriter) const
 	{
 		JsonWriter.WriteArrayStart();
-		JsonWriter.WriteValue(X);
-		JsonWriter.WriteValue(Y);
-		JsonWriter.WriteValue(Z);
-		JsonWriter.WriteValue(W);
+		FGLTFJsonUtility::WriteExactValue(JsonWriter, X);
+		FGLTFJsonUtility::WriteExactValue(JsonWriter, Y);
+		FGLTFJsonUtility::WriteExactValue(JsonWriter, Z);
+		FGLTFJsonUtility::WriteExactValue(JsonWriter, W);
 		JsonWriter.WriteArrayEnd();
 	}
 
