[GLTFExporter] refactor(json): fix minor identifier typo in JsonVector3

commit 5e7b08f79d35ff3149323101fa664aebc56841c2
Date: Sat, 23 May 2020 14:48:52 +0200

    refactor(json): fix minor identifier typo in JsonVector3
    (cherry picked from commit 1e1484c49baab95ec3be6645fb48fc894a8e9754​)

#skipci

diff --git a/Engine/Plugins/Enterprise/GLTFExporter/MigratedFromGit/Source/GLTFExporter/Private/Json/GLTFJsonVector3.h b/Engine/Plugins/Enterprise/GLTFExporter/MigratedFromGit/Source/GLTFExporter/Private/Json/GLTFJsonVector3.h
--- a/Engine/Plugins/Enterprise/GLTFExporter/MigratedFromGit/Source/GLTFExporter/Private/Json/GLTFJsonVector3.h
+++ b/Engine/Plugins/Enterprise/GLTFExporter/MigratedFromGit/Source/GLTFExporter/Private/Json/GLTFJsonVector3.h
@@ -18,8 +18,8 @@ struct FGLTFJsonVector3
 	{
 	}
 
-	FGLTFJsonVector3(const FGLTFJsonVector2& Vector3, float Z)
-        : X(Vector3.X), Y(Vector3.Y), Z(Z)
+	FGLTFJsonVector3(const FGLTFJsonVector2& Vector2, float Z)
+        : X(Vector2.X), Y(Vector2.Y), Z(Z)
 	{
 	}
 
