[GLTFExporter] refactor(static-mesh): minor readability improvement of conditional expression

commit da4949ad2aad50bbcbfdf7e1198411a19ad726da
Date: Mon, 25 May 2020 05:59:26 +0200

    refactor(static-mesh): minor readability improvement of conditional expression
    (cherry picked from commit 69a6706b54f08a59f5b1fc58ca4e413c7cc5e311​)

#skipci

diff --git a/Engine/Plugins/Enterprise/GLTFExporter/MigratedFromGit/Source/GLTFExporter/Private/Converters/GLTFStaticMeshConverters.cpp b/Engine/Plugins/Enterprise/GLTFExporter/MigratedFromGit/Source/GLTFExporter/Private/Converters/GLTFStaticMeshConverters.cpp
--- a/Engine/Plugins/Enterprise/GLTFExporter/MigratedFromGit/Source/GLTFExporter/Private/Converters/GLTFStaticMeshConverters.cpp
+++ b/Engine/Plugins/Enterprise/GLTFExporter/MigratedFromGit/Source/GLTFExporter/Private/Converters/GLTFStaticMeshConverters.cpp
@@ -238,8 +238,9 @@ FGLTFJsonMeshIndex FGLTFStaticMeshConverter::Add(FGLTFConvertBuilder& Builder, c
 		FGLTFJsonPrimitive& JsonPrimitive = JsonMesh.Primitives[SectionIndex];
 		JsonPrimitive.Attributes = JsonAttributes;
 
-		JsonPrimitive.Indices = Builder.GetOrAddIndexAccessor(&LODResources.Sections[SectionIndex], IndexBuffer,
-			Name + (SectionCount != 1 ? TEXT("_Indices_Section") + FString::FromInt(SectionIndex) : TEXT("_Indices")));
+		const FStaticMeshSection& Section = LODResources.Sections[SectionIndex];
+		JsonPrimitive.Indices = Builder.GetOrAddIndexAccessor(&Section, IndexBuffer,
+			Name + (SectionCount == 1 ? TEXT("_Indices") : TEXT("_Indices_Section") + FString::FromInt(SectionIndex)));
 	}
 
 	return Builder.AddMesh(JsonMesh);
