[GLTFExporter] refactor: move out non-static-mesh-specific converters to new files

commit a2020d21228f04e07cac312d28af195193f58478
Date: Mon, 25 May 2020 19:44:19 +0200

    refactor: move out non-static-mesh-specific converters to new files
    (cherry picked from commit 6546299732469ea5e26054114c23f68de2bc1a78​)

#skipci

diff --git a/Engine/Plugins/Enterprise/GLTFExporter/MigratedFromGit/Source/GLTFExporter/Private/Builders/GLTFConvertBuilder.cpp b/Engine/Plugins/Enterprise/GLTFExporter/MigratedFromGit/Source/GLTFExporter/Private/Builders/GLTFConvertBuilder.cpp
--- a/Engine/Plugins/Enterprise/GLTFExporter/MigratedFromGit/Source/GLTFExporter/Private/Builders/GLTFConvertBuilder.cpp
+++ b/Engine/Plugins/Enterprise/GLTFExporter/MigratedFromGit/Source/GLTFExporter/Private/Builders/GLTFConvertBuilder.cpp
@@ -30,7 +30,7 @@ FGLTFJsonAccessorIndex FGLTFConvertBuilder::GetOrAddNormalAccessor(const FStatic
 		return FGLTFJsonAccessorIndex(INDEX_NONE);
 	}
 
-	return StaticMeshNormalVertexBufferConverter.GetOrAdd(*this, DesiredName, VertexBuffer);
+	return NormalVertexBufferConverter.GetOrAdd(*this, DesiredName, VertexBuffer);
 }
 
 FGLTFJsonAccessorIndex FGLTFConvertBuilder::GetOrAddTangentAccessor(const FStaticMeshVertexBuffer* VertexBuffer, const FString& DesiredName)
@@ -40,7 +40,7 @@ FGLTFJsonAccessorIndex FGLTFConvertBuilder::GetOrAddTangentAccessor(const FStati
 		return FGLTFJsonAccessorIndex(INDEX_NONE);
 	}
 
-	return StaticMeshTangentVertexBufferConverter.GetOrAdd(*this, DesiredName, VertexBuffer);
+	return TangentVertexBufferConverter.GetOrAdd(*this, DesiredName, VertexBuffer);
 }
 
 FGLTFJsonAccessorIndex FGLTFConvertBuilder::GetOrAddUVAccessor(const FStaticMeshVertexBuffer* VertexBuffer, int32 UVIndex, const FString& DesiredName)
@@ -50,7 +50,7 @@ FGLTFJsonAccessorIndex FGLTFConvertBuilder::GetOrAddUVAccessor(const FStaticMesh
 		return FGLTFJsonAccessorIndex(INDEX_NONE);
 	}
 
-	return StaticMeshUVVertexBufferConverter.GetOrAdd(*this, DesiredName, VertexBuffer, UVIndex);
+	return UVVertexBufferConverter.GetOrAdd(*this, DesiredName, VertexBuffer, UVIndex);
 }
 
 FGLTFJsonBufferViewIndex FGLTFConvertBuilder::GetOrAddIndexBufferView(const FRawStaticIndexBuffer* IndexBuffer, const FString& DesiredName)
@@ -60,7 +60,7 @@ FGLTFJsonBufferViewIndex FGLTFConvertBuilder::GetOrAddIndexBufferView(const FRaw
 		return FGLTFJsonBufferViewIndex(INDEX_NONE);
 	}
 
-	return StaticMeshIndexBufferConverter.GetOrAdd(*this, DesiredName, IndexBuffer);
+	return IndexBufferConverter.GetOrAdd(*this, DesiredName, IndexBuffer);
 }
 
 FGLTFJsonAccessorIndex FGLTFConvertBuilder::GetOrAddIndexAccessor(const FStaticMeshSection* MeshSection, const FRawStaticIndexBuffer* IndexBuffer, const FString& DesiredName)
@@ -92,7 +92,7 @@ FGLTFJsonMeshIndex FGLTFConvertBuilder::GetOrAddMesh(const UStaticMeshComponent*
 
 	const UStaticMesh* StaticMesh = StaticMeshComponent->GetStaticMesh();
 	const int32 LODIndex = StaticMeshComponent->ForcedLodModel > 0 ? StaticMeshComponent->ForcedLodModel - 1 : /* auto-select */ 0;
-	const FColorVertexBuffer* OverrideVertexColors = LODIndex < StaticMeshComponent->LODData.Num() ? StaticMeshComponent->LODData[LODIndex].OverrideVertexColors : nullptr;
+	const FColorVertexBuffer* OverrideVertexColors = StaticMeshComponent->LODData.IsValidIndex(LODIndex) ? StaticMeshComponent->LODData[LODIndex].OverrideVertexColors : nullptr;
 
 	return GetOrAddMesh(StaticMesh, LODIndex, OverrideVertexColors, DesiredName);
 }
diff --git a/Engine/Plugins/Enterprise/GLTFExporter/MigratedFromGit/Source/GLTFExporter/Private/Builders/GLTFConvertBuilder.h b/Engine/Plugins/Enterprise/GLTFExporter/MigratedFromGit/Source/GLTFExporter/Private/Builders/GLTFConvertBuilder.h
--- a/Engine/Plugins/Enterprise/GLTFExporter/MigratedFromGit/Source/GLTFExporter/Private/Builders/GLTFConvertBuilder.h
+++ b/Engine/Plugins/Enterprise/GLTFExporter/MigratedFromGit/Source/GLTFExporter/Private/Builders/GLTFConvertBuilder.h
@@ -3,6 +3,7 @@
 #pragma once
 
 #include "Builders/GLTFBufferBuilder.h"
+#include "Converters/GLTFVertexBufferConverters.h"
 #include "Converters/GLTFStaticMeshConverters.h"
 #include "Converters/GLTFLevelConverters.h"
 
@@ -11,19 +12,17 @@ class FGLTFConvertBuilder : public FGLTFBufferBuilder
 public:
 
 	FGLTFJsonAccessorIndex GetOrAddPositionAccessor(const FPositionVertexBuffer* VertexBuffer, const FString& DesiredName = TEXT(""));
+	FGLTFJsonAccessorIndex GetOrAddColorAccessor(const FColorVertexBuffer* VertexBuffer, const FString& DesiredName = TEXT(""));
 	FGLTFJsonAccessorIndex GetOrAddNormalAccessor(const FStaticMeshVertexBuffer* VertexBuffer, const FString& DesiredName = TEXT(""));
 	FGLTFJsonAccessorIndex GetOrAddTangentAccessor(const FStaticMeshVertexBuffer* VertexBuffer, const FString& DesiredName = TEXT(""));
 	FGLTFJsonAccessorIndex GetOrAddUVAccessor(const FStaticMeshVertexBuffer* VertexBuffer, int32 UVIndex, const FString& DesiredName = TEXT(""));
-	FGLTFJsonAccessorIndex GetOrAddColorAccessor(const FColorVertexBuffer* VertexBuffer, const FString& DesiredName = TEXT(""));
 
 	FGLTFJsonBufferViewIndex GetOrAddIndexBufferView(const FRawStaticIndexBuffer* IndexBuffer, const FString& DesiredName = TEXT(""));
 	FGLTFJsonAccessorIndex GetOrAddIndexAccessor(const FStaticMeshSection* MeshSection, const FRawStaticIndexBuffer* IndexBuffer, const FString& DesiredName = TEXT(""));
-
 	FGLTFJsonMeshIndex GetOrAddMesh(const UStaticMesh* StaticMesh, int32 LODIndex = 0, const FColorVertexBuffer* OverrideVertexColors = nullptr, const FString& DesiredName = TEXT(""));
 	FGLTFJsonMeshIndex GetOrAddMesh(const UStaticMeshComponent* StaticMeshComponent, const FString& DesiredName = TEXT(""));
 
 	FGLTFJsonNodeIndex GetOrAddNode(const USceneComponent* SceneComponent, bool bSelectedOnly, bool bRootNode = false, const FString& DesiredName = TEXT(""));
-
 	FGLTFJsonSceneIndex GetOrAddScene(const ULevel* Level, bool bSelectedOnly, const FString& DesiredName = TEXT(""));
 	FGLTFJsonSceneIndex GetOrAddScene(const UWorld* World, bool bSelectedOnly, const FString& DesiredName = TEXT(""));
 
@@ -31,10 +30,11 @@ private:
 
 	FGLTFPositionVertexBufferConverter PositionVertexBufferConverter;
 	FGLTFColorVertexBufferConverter ColorVertexBufferConverter;
-	FGLTFStaticMeshNormalVertexBufferConverter StaticMeshNormalVertexBufferConverter;
-	FGLTFStaticMeshTangentVertexBufferConverter StaticMeshTangentVertexBufferConverter;
-	FGLTFStaticMeshUVVertexBufferConverter StaticMeshUVVertexBufferConverter;
-	FGLTFStaticMeshIndexBufferConverter StaticMeshIndexBufferConverter;
+	FGLTFNormalVertexBufferConverter NormalVertexBufferConverter;
+	FGLTFTangentVertexBufferConverter TangentVertexBufferConverter;
+	FGLTFUVVertexBufferConverter UVVertexBufferConverter;
+
+	FGLTFIndexBufferConverter IndexBufferConverter;
 	FGLTFStaticMeshSectionConverter StaticMeshSectionConverter;
 	FGLTFStaticMeshConverter StaticMeshConverter;
 
diff --git a/Engine/Plugins/Enterprise/GLTFExporter/MigratedFromGit/Source/GLTFExporter/Private/Converters/GLTFStaticMeshConverters.cpp b/Engine/Plugins/Enterprise/GLTFExporter/MigratedFromGit/Source/GLTFExporter/Private/Converters/GLTFStaticMeshConverters.cpp
--- a/Engine/Plugins/Enterprise/GLTFExporter/MigratedFromGit/Source/GLTFExporter/Private/Converters/GLTFStaticMeshConverters.cpp
+++ b/Engine/Plugins/Enterprise/GLTFExporter/MigratedFromGit/Source/GLTFExporter/Private/Converters/GLTFStaticMeshConverters.cpp
@@ -1,164 +1,10 @@
 // Copyright Epic Games, Inc. All Rights Reserved.
 
 #include "Converters/GLTFStaticMeshConverters.h"
-#include "Builders/GLTFConvertBuilder.h"
 #include "Converters/GLTFConverterUtility.h"
+#include "Builders/GLTFConvertBuilder.h"
 
-FGLTFJsonAccessorIndex FGLTFPositionVertexBufferConverter::Add(FGLTFConvertBuilder& Builder, const FString& Name, const FPositionVertexBuffer* VertexBuffer)
-{
-	const uint32 VertexCount = VertexBuffer->GetNumVertices();
-	if (VertexCount == 0)
-	{
-		return FGLTFJsonAccessorIndex(INDEX_NONE);
-	}
-
-	TArray<FGLTFJsonVector3> Positions;
-	Positions.AddUninitialized(VertexCount);
-
-	for (uint32 VertexIndex = 0; VertexIndex < VertexCount; ++VertexIndex)
-	{
-		Positions[VertexIndex] = FGLTFConverterUtility::ConvertPosition(VertexBuffer->VertexPosition(VertexIndex));
-	}
-
-	// More accurate bounding box if based on raw vertex values
-	FGLTFJsonVector3 MinPosition = Positions[0];
-	FGLTFJsonVector3 MaxPosition = Positions[0];
-
-	for (uint32 VertexIndex = 1; VertexIndex < VertexCount; ++VertexIndex)
-	{
-		const FGLTFJsonVector3& Position = Positions[VertexIndex];
-		MinPosition.X = FMath::Min(MinPosition.X, Position.X);
-		MinPosition.Y = FMath::Min(MinPosition.Y, Position.Y);
-		MinPosition.Z = FMath::Min(MinPosition.Z, Position.Z);
-		MaxPosition.X = FMath::Max(MaxPosition.X, Position.X);
-		MaxPosition.Y = FMath::Max(MaxPosition.Y, Position.Y);
-		MaxPosition.Z = FMath::Max(MaxPosition.Z, Position.Z);
-	}
-
-	FGLTFJsonAccessor JsonAccessor;
-	JsonAccessor.Name = Name;
-	JsonAccessor.BufferView = Builder.AddBufferView(Positions, Name);
-	JsonAccessor.ComponentType = EGLTFJsonComponentType::F32;
-	JsonAccessor.Count = VertexCount;
-	JsonAccessor.Type = EGLTFJsonAccessorType::Vec3;
-
-	JsonAccessor.MinMaxLength = 3;
-	JsonAccessor.Min[0] = MinPosition.X;
-	JsonAccessor.Min[1] = MinPosition.Y;
-	JsonAccessor.Min[2] = MinPosition.Z;
-	JsonAccessor.Max[0] = MaxPosition.X;
-	JsonAccessor.Max[1] = MaxPosition.Y;
-	JsonAccessor.Max[2] = MaxPosition.Z;
-
-	return Builder.AddAccessor(JsonAccessor);
-}
-
-FGLTFJsonAccessorIndex FGLTFColorVertexBufferConverter::Add(FGLTFConvertBuilder& Builder, const FString& Name, const FColorVertexBuffer* VertexBuffer)
-{
-	const uint32 VertexCount = VertexBuffer->GetNumVertices();
-	if (VertexCount == 0)
-	{
-		return FGLTFJsonAccessorIndex(INDEX_NONE);
-	}
-
-	TArray<FGLTFPackedColor> Colors;
-	Colors.AddUninitialized(VertexCount);
-
-	for (uint32 VertexIndex = 0; VertexIndex < VertexCount; ++VertexIndex)
-	{
-		Colors[VertexIndex] = FGLTFConverterUtility::ConvertColor(VertexBuffer->VertexColor(VertexIndex));
-	}
-
-	FGLTFJsonAccessor JsonAccessor;
-	JsonAccessor.Name = Name;
-	JsonAccessor.BufferView = Builder.AddBufferView(Colors, Name);
-	JsonAccessor.ComponentType = EGLTFJsonComponentType::U8;
-	JsonAccessor.Count = VertexCount;
-	JsonAccessor.Type = EGLTFJsonAccessorType::Vec4;
-	JsonAccessor.bNormalized = true;
-
-	return Builder.AddAccessor(JsonAccessor);
-}
-
-FGLTFJsonAccessorIndex FGLTFStaticMeshNormalVertexBufferConverter::Add(FGLTFConvertBuilder& Builder, const FString& Name, const FStaticMeshVertexBuffer* VertexBuffer)
-{
-	const uint32 VertexCount = VertexBuffer->GetNumVertices();
-	if (VertexCount == 0)
-	{
-		return FGLTFJsonAccessorIndex(INDEX_NONE);
-	}
-
-	TArray<FGLTFJsonVector3> Normals;
-	Normals.AddUninitialized(VertexCount);
-
-	for (uint32 VertexIndex = 0; VertexIndex < VertexCount; ++VertexIndex)
-	{
-		Normals[VertexIndex] = FGLTFConverterUtility::ConvertNormal(VertexBuffer->VertexTangentZ(VertexIndex));
-	}
-
-	FGLTFJsonAccessor JsonAccessor;
-	JsonAccessor.Name = Name;
-	JsonAccessor.BufferView = Builder.AddBufferView(Normals, Name);
-	JsonAccessor.ComponentType = EGLTFJsonComponentType::F32;
-	JsonAccessor.Count = VertexCount;
-	JsonAccessor.Type = EGLTFJsonAccessorType::Vec3;
-
-	return Builder.AddAccessor(JsonAccessor);
-}
-
-FGLTFJsonAccessorIndex FGLTFStaticMeshTangentVertexBufferConverter::Add(FGLTFConvertBuilder& Builder, const FString& Name, const FStaticMeshVertexBuffer* VertexBuffer)
-{
-	const uint32 VertexCount = VertexBuffer->GetNumVertices();
-	if (VertexCount == 0)
-	{
-		return FGLTFJsonAccessorIndex(INDEX_NONE);
-	}
-
-	TArray<FGLTFJsonVector4> Tangents;
-	Tangents.AddUninitialized(VertexCount);
-
-	for (uint32 VertexIndex = 0; VertexIndex < VertexCount; ++VertexIndex)
-	{
-		Tangents[VertexIndex] = FGLTFConverterUtility::ConvertTangent(VertexBuffer->VertexTangentX(VertexIndex));
-	}
-
-	FGLTFJsonAccessor JsonAccessor;
-	JsonAccessor.Name = Name;
-	JsonAccessor.BufferView = Builder.AddBufferView(Tangents, Name);
-	JsonAccessor.ComponentType = EGLTFJsonComponentType::F32;
-	JsonAccessor.Count = VertexCount;
-	JsonAccessor.Type = EGLTFJsonAccessorType::Vec4;
-
-	return Builder.AddAccessor(JsonAccessor);
-}
-
-FGLTFJsonAccessorIndex FGLTFStaticMeshUVVertexBufferConverter::Add(FGLTFConvertBuilder& Builder, const FString& Name, const FStaticMeshVertexBuffer* VertexBuffer, int32 UVIndex)
-{
-	const uint32 VertexCount = VertexBuffer->GetNumVertices();
-	if (VertexCount == 0 || UVIndex < 0 || VertexBuffer->GetNumTexCoords() <= static_cast<uint32>(UVIndex))
-	{
-		return FGLTFJsonAccessorIndex(INDEX_NONE);
-	}
-
-	TArray<FGLTFJsonVector2> UVs;
-	UVs.AddUninitialized(VertexCount);
-
-	for (uint32 VertexIndex = 0; VertexIndex < VertexCount; ++VertexIndex)
-	{
-		UVs[VertexIndex] = FGLTFConverterUtility::ConvertUV(VertexBuffer->GetVertexUV(VertexIndex, UVIndex));
-	}
-
-	FGLTFJsonAccessor JsonAccessor;
-	JsonAccessor.Name = Name;
-	JsonAccessor.BufferView = Builder.AddBufferView(UVs, Name);
-	JsonAccessor.ComponentType = EGLTFJsonComponentType::F32;
-	JsonAccessor.Count = VertexCount;
-	JsonAccessor.Type = EGLTFJsonAccessorType::Vec2;
-
-	return Builder.AddAccessor(JsonAccessor);
-}
-
-FGLTFJsonBufferViewIndex FGLTFStaticMeshIndexBufferConverter::Add(FGLTFConvertBuilder& Builder, const FString& Name, const FRawStaticIndexBuffer* IndexBuffer)
+FGLTFJsonBufferViewIndex FGLTFIndexBufferConverter::Add(FGLTFConvertBuilder& Builder, const FString& Name, const FRawStaticIndexBuffer* IndexBuffer)
 {
 	if (IndexBuffer->GetNumIndices() == 0)
 	{
diff --git a/Engine/Plugins/Enterprise/GLTFExporter/MigratedFromGit/Source/GLTFExporter/Private/Converters/GLTFStaticMeshConverters.h b/Engine/Plugins/Enterprise/GLTFExporter/MigratedFromGit/Source/GLTFExporter/Private/Converters/GLTFStaticMeshConverters.h
--- a/Engine/Plugins/Enterprise/GLTFExporter/MigratedFromGit/Source/GLTFExporter/Private/Converters/GLTFStaticMeshConverters.h
+++ b/Engine/Plugins/Enterprise/GLTFExporter/MigratedFromGit/Source/GLTFExporter/Private/Converters/GLTFStaticMeshConverters.h
@@ -6,32 +6,7 @@
 #include "Converters/GLTFConverter.h"
 #include "Engine.h"
 
-class FGLTFPositionVertexBufferConverter final : public TGLTFConverter<FGLTFJsonAccessorIndex, const FPositionVertexBuffer*>
-{
-	FGLTFJsonAccessorIndex Add(FGLTFConvertBuilder& Builder, const FString& Name, const FPositionVertexBuffer* VertexBuffer) override;
-};
-
-class FGLTFColorVertexBufferConverter final : public TGLTFConverter<FGLTFJsonAccessorIndex, const FColorVertexBuffer*>
-{
-	FGLTFJsonAccessorIndex Add(FGLTFConvertBuilder& Builder, const FString& Name, const FColorVertexBuffer* VertexBuffer) override;
-};
-
-class FGLTFStaticMeshNormalVertexBufferConverter final : public TGLTFConverter<FGLTFJsonAccessorIndex, const FStaticMeshVertexBuffer*>
-{
-	FGLTFJsonAccessorIndex Add(FGLTFConvertBuilder& Builder, const FString& Name, const FStaticMeshVertexBuffer* VertexBuffer) override;
-};
-
-class FGLTFStaticMeshTangentVertexBufferConverter final : public TGLTFConverter<FGLTFJsonAccessorIndex, const FStaticMeshVertexBuffer*>
-{
-	FGLTFJsonAccessorIndex Add(FGLTFConvertBuilder& Builder, const FString& Name, const FStaticMeshVertexBuffer* VertexBuffer) override;
-};
-
-class FGLTFStaticMeshUVVertexBufferConverter final : public TGLTFConverter<FGLTFJsonAccessorIndex, const FStaticMeshVertexBuffer*, int32>
-{
-	FGLTFJsonAccessorIndex Add(FGLTFConvertBuilder& Builder, const FString& Name, const FStaticMeshVertexBuffer* VertexBuffer, int32 UVIndex) override;
-};
-
-class FGLTFStaticMeshIndexBufferConverter final : public TGLTFConverter<FGLTFJsonBufferViewIndex, const FRawStaticIndexBuffer*>
+class FGLTFIndexBufferConverter final : public TGLTFConverter<FGLTFJsonBufferViewIndex, const FRawStaticIndexBuffer*>
 {
 	FGLTFJsonBufferViewIndex Add(FGLTFConvertBuilder& Builder, const FString& Name, const FRawStaticIndexBuffer* IndexBuffer) override;
 };
diff --git a/Engine/Plugins/Enterprise/GLTFExporter/MigratedFromGit/Source/GLTFExporter/Private/Converters/GLTFVertexBufferConverters.cpp b/Engine/Plugins/Enterprise/GLTFExporter/MigratedFromGit/Source/GLTFExporter/Private/Converters/GLTFVertexBufferConverters.cpp
new file mode 100644
--- /dev/null
+++ b/Engine/Plugins/Enterprise/GLTFExporter/MigratedFromGit/Source/GLTFExporter/Private/Converters/GLTFVertexBufferConverters.cpp
@@ -0,0 +1,160 @@
+// Copyright Epic Games, Inc. All Rights Reserved.
+
+#include "Converters/GLTFVertexBufferConverters.h"
+#include "Converters/GLTFConverterUtility.h"
+#include "Builders/GLTFConvertBuilder.h"
+
+FGLTFJsonAccessorIndex FGLTFPositionVertexBufferConverter::Add(FGLTFConvertBuilder& Builder, const FString& Name, const FPositionVertexBuffer* VertexBuffer)
+{
+	const uint32 VertexCount = VertexBuffer->GetNumVertices();
+	if (VertexCount == 0)
+	{
+		return FGLTFJsonAccessorIndex(INDEX_NONE);
+	}
+
+	TArray<FGLTFJsonVector3> Positions;
+	Positions.AddUninitialized(VertexCount);
+
+	for (uint32 VertexIndex = 0; VertexIndex < VertexCount; ++VertexIndex)
+	{
+		Positions[VertexIndex] = FGLTFConverterUtility::ConvertPosition(VertexBuffer->VertexPosition(VertexIndex));
+	}
+
+	// More accurate bounding box if based on raw vertex values
+	FGLTFJsonVector3 MinPosition = Positions[0];
+	FGLTFJsonVector3 MaxPosition = Positions[0];
+
+	for (uint32 VertexIndex = 1; VertexIndex < VertexCount; ++VertexIndex)
+	{
+		const FGLTFJsonVector3& Position = Positions[VertexIndex];
+		MinPosition.X = FMath::Min(MinPosition.X, Position.X);
+		MinPosition.Y = FMath::Min(MinPosition.Y, Position.Y);
+		MinPosition.Z = FMath::Min(MinPosition.Z, Position.Z);
+		MaxPosition.X = FMath::Max(MaxPosition.X, Position.X);
+		MaxPosition.Y = FMath::Max(MaxPosition.Y, Position.Y);
+		MaxPosition.Z = FMath::Max(MaxPosition.Z, Position.Z);
+	}
+
+	FGLTFJsonAccessor JsonAccessor;
+	JsonAccessor.Name = Name;
+	JsonAccessor.BufferView = Builder.AddBufferView(Positions, Name);
+	JsonAccessor.ComponentType = EGLTFJsonComponentType::F32;
+	JsonAccessor.Count = VertexCount;
+	JsonAccessor.Type = EGLTFJsonAccessorType::Vec3;
+
+	JsonAccessor.MinMaxLength = 3;
+	JsonAccessor.Min[0] = MinPosition.X;
+	JsonAccessor.Min[1] = MinPosition.Y;
+	JsonAccessor.Min[2] = MinPosition.Z;
+	JsonAccessor.Max[0] = MaxPosition.X;
+	JsonAccessor.Max[1] = MaxPosition.Y;
+	JsonAccessor.Max[2] = MaxPosition.Z;
+
+	return Builder.AddAccessor(JsonAccessor);
+}
+
+FGLTFJsonAccessorIndex FGLTFColorVertexBufferConverter::Add(FGLTFConvertBuilder& Builder, const FString& Name, const FColorVertexBuffer* VertexBuffer)
+{
+	const uint32 VertexCount = VertexBuffer->GetNumVertices();
+	if (VertexCount == 0)
+	{
+		return FGLTFJsonAccessorIndex(INDEX_NONE);
+	}
+
+	TArray<FGLTFPackedColor> Colors;
+	Colors.AddUninitialized(VertexCount);
+
+	for (uint32 VertexIndex = 0; VertexIndex < VertexCount; ++VertexIndex)
+	{
+		Colors[VertexIndex] = FGLTFConverterUtility::ConvertColor(VertexBuffer->VertexColor(VertexIndex));
+	}
+
+	FGLTFJsonAccessor JsonAccessor;
+	JsonAccessor.Name = Name;
+	JsonAccessor.BufferView = Builder.AddBufferView(Colors, Name);
+	JsonAccessor.ComponentType = EGLTFJsonComponentType::U8;
+	JsonAccessor.Count = VertexCount;
+	JsonAccessor.Type = EGLTFJsonAccessorType::Vec4;
+	JsonAccessor.bNormalized = true;
+
+	return Builder.AddAccessor(JsonAccessor);
+}
+
+FGLTFJsonAccessorIndex FGLTFNormalVertexBufferConverter::Add(FGLTFConvertBuilder& Builder, const FString& Name, const FStaticMeshVertexBuffer* VertexBuffer)
+{
+	const uint32 VertexCount = VertexBuffer->GetNumVertices();
+	if (VertexCount == 0)
+	{
+		return FGLTFJsonAccessorIndex(INDEX_NONE);
+	}
+
+	TArray<FGLTFJsonVector3> Normals;
+	Normals.AddUninitialized(VertexCount);
+
+	for (uint32 VertexIndex = 0; VertexIndex < VertexCount; ++VertexIndex)
+	{
+		Normals[VertexIndex] = FGLTFConverterUtility::ConvertNormal(VertexBuffer->VertexTangentZ(VertexIndex));
+	}
+
+	FGLTFJsonAccessor JsonAccessor;
+	JsonAccessor.Name = Name;
+	JsonAccessor.BufferView = Builder.AddBufferView(Normals, Name);
+	JsonAccessor.ComponentType = EGLTFJsonComponentType::F32;
+	JsonAccessor.Count = VertexCount;
+	JsonAccessor.Type = EGLTFJsonAccessorType::Vec3;
+
+	return Builder.AddAccessor(JsonAccessor);
+}
+
+FGLTFJsonAccessorIndex FGLTFTangentVertexBufferConverter::Add(FGLTFConvertBuilder& Builder, const FString& Name, const FStaticMeshVertexBuffer* VertexBuffer)
+{
+	const uint32 VertexCount = VertexBuffer->GetNumVertices();
+	if (VertexCount == 0)
+	{
+		return FGLTFJsonAccessorIndex(INDEX_NONE);
+	}
+
+	TArray<FGLTFJsonVector4> Tangents;
+	Tangents.AddUninitialized(VertexCount);
+
+	for (uint32 VertexIndex = 0; VertexIndex < VertexCount; ++VertexIndex)
+	{
+		Tangents[VertexIndex] = FGLTFConverterUtility::ConvertTangent(VertexBuffer->VertexTangentX(VertexIndex));
+	}
+
+	FGLTFJsonAccessor JsonAccessor;
+	JsonAccessor.Name = Name;
+	JsonAccessor.BufferView = Builder.AddBufferView(Tangents, Name);
+	JsonAccessor.ComponentType = EGLTFJsonComponentType::F32;
+	JsonAccessor.Count = VertexCount;
+	JsonAccessor.Type = EGLTFJsonAccessorType::Vec4;
+
+	return Builder.AddAccessor(JsonAccessor);
+}
+
+FGLTFJsonAccessorIndex FGLTFUVVertexBufferConverter::Add(FGLTFConvertBuilder& Builder, const FString& Name, const FStaticMeshVertexBuffer* VertexBuffer, int32 UVIndex)
+{
+	const uint32 VertexCount = VertexBuffer->GetNumVertices();
+	if (VertexCount == 0 || UVIndex < 0 || VertexBuffer->GetNumTexCoords() <= static_cast<uint32>(UVIndex))
+	{
+		return FGLTFJsonAccessorIndex(INDEX_NONE);
+	}
+
+	TArray<FGLTFJsonVector2> UVs;
+	UVs.AddUninitialized(VertexCount);
+
+	for (uint32 VertexIndex = 0; VertexIndex < VertexCount; ++VertexIndex)
+	{
+		UVs[VertexIndex] = FGLTFConverterUtility::ConvertUV(VertexBuffer->GetVertexUV(VertexIndex, UVIndex));
+	}
+
+	FGLTFJsonAccessor JsonAccessor;
+	JsonAccessor.Name = Name;
+	JsonAccessor.BufferView = Builder.AddBufferView(UVs, Name);
+	JsonAccessor.ComponentType = EGLTFJsonComponentType::F32;
+	JsonAccessor.Count = VertexCount;
+	JsonAccessor.Type = EGLTFJsonAccessorType::Vec2;
+
+	return Builder.AddAccessor(JsonAccessor);
+}
+
diff --git a/Engine/Plugins/Enterprise/GLTFExporter/MigratedFromGit/Source/GLTFExporter/Private/Converters/GLTFVertexBufferConverters.h b/Engine/Plugins/Enterprise/GLTFExporter/MigratedFromGit/Source/GLTFExporter/Private/Converters/GLTFVertexBufferConverters.h
new file mode 100644
--- /dev/null
+++ b/Engine/Plugins/Enterprise/GLTFExporter/MigratedFromGit/Source/GLTFExporter/Private/Converters/GLTFVertexBufferConverters.h
@@ -0,0 +1,32 @@
+// Copyright Epic Games, Inc. All Rights Reserved.
+
+#pragma once
+
+#include "Json/GLTFJsonIndex.h"
+#include "Converters/GLTFConverter.h"
+#include "Engine.h"
+
+class FGLTFPositionVertexBufferConverter final : public TGLTFConverter<FGLTFJsonAccessorIndex, const FPositionVertexBuffer*>
+{
+	FGLTFJsonAccessorIndex Add(FGLTFConvertBuilder& Builder, const FString& Name, const FPositionVertexBuffer* VertexBuffer) override;
+};
+
+class FGLTFColorVertexBufferConverter final : public TGLTFConverter<FGLTFJsonAccessorIndex, const FColorVertexBuffer*>
+{
+	FGLTFJsonAccessorIndex Add(FGLTFConvertBuilder& Builder, const FString& Name, const FColorVertexBuffer* VertexBuffer) override;
+};
+
+class FGLTFNormalVertexBufferConverter final : public TGLTFConverter<FGLTFJsonAccessorIndex, const FStaticMeshVertexBuffer*>
+{
+	FGLTFJsonAccessorIndex Add(FGLTFConvertBuilder& Builder, const FString& Name, const FStaticMeshVertexBuffer* VertexBuffer) override;
+};
+
+class FGLTFTangentVertexBufferConverter final : public TGLTFConverter<FGLTFJsonAccessorIndex, const FStaticMeshVertexBuffer*>
+{
+	FGLTFJsonAccessorIndex Add(FGLTFConvertBuilder& Builder, const FString& Name, const FStaticMeshVertexBuffer* VertexBuffer) override;
+};
+
+class FGLTFUVVertexBufferConverter final : public TGLTFConverter<FGLTFJsonAccessorIndex, const FStaticMeshVertexBuffer*, int32>
+{
+	FGLTFJsonAccessorIndex Add(FGLTFConvertBuilder& Builder, const FString& Name, const FStaticMeshVertexBuffer* VertexBuffer, int32 UVIndex) override;
+};
