[GLTFExporter] feat(static-mesh): add support for variable number of UV channels

commit 36b943b9ec32bed7dac1523932656b0fb09fbbe3
Date: Mon, 25 May 2020 19:52:24 +0200

    feat(static-mesh): add support for variable number of UV channels
    (cherry picked from commit d6fdd15c6d09485b0e09332780d5bc2ae62e82ca​)

#skipci

diff --git a/Engine/Plugins/Enterprise/GLTFExporter/MigratedFromGit/Source/GLTFExporter/Private/Converters/GLTFStaticMeshConverters.cpp b/Engine/Plugins/Enterprise/GLTFExporter/MigratedFromGit/Source/GLTFExporter/Private/Converters/GLTFStaticMeshConverters.cpp
--- a/Engine/Plugins/Enterprise/GLTFExporter/MigratedFromGit/Source/GLTFExporter/Private/Converters/GLTFStaticMeshConverters.cpp
+++ b/Engine/Plugins/Enterprise/GLTFExporter/MigratedFromGit/Source/GLTFExporter/Private/Converters/GLTFStaticMeshConverters.cpp
@@ -54,8 +54,6 @@ FGLTFJsonMeshIndex FGLTFStaticMeshConverter::Add(FGLTFConvertBuilder& Builder, c
 		return FGLTFJsonMeshIndex(INDEX_NONE);
 	}
 
-	const int32 PrimaryUVIndex = 0; // TODO: make this configurable?
-	const int32 LightmapUVIndex = StaticMesh->LightMapCoordinateIndex != PrimaryUVIndex ? StaticMesh->LightMapCoordinateIndex : INDEX_NONE;
 	const FStaticMeshLODResources& LODResources = StaticMesh->GetLODForExport(LODIndex);
 
 	FGLTFJsonMesh JsonMesh;
@@ -70,8 +68,14 @@ FGLTFJsonMeshIndex FGLTFStaticMeshConverter::Add(FGLTFConvertBuilder& Builder, c
 	JsonAttributes.Color0 = Builder.GetOrAddColorAccessor(ColorBuffer, Name + TEXT("_Colors"));
 	JsonAttributes.Normal = Builder.GetOrAddNormalAccessor(VertexBuffer, Name + TEXT("_Normals"));
 	JsonAttributes.Tangent = Builder.GetOrAddTangentAccessor(VertexBuffer, Name + TEXT("_Tangents"));
-	JsonAttributes.TexCoord0 = Builder.GetOrAddUVAccessor(VertexBuffer, PrimaryUVIndex, Name + TEXT("_UV") + FString::FromInt(PrimaryUVIndex) + TEXT("s"));
-	JsonAttributes.TexCoord1 = Builder.GetOrAddUVAccessor(VertexBuffer, LightmapUVIndex, Name + TEXT("_UV") + FString::FromInt(LightmapUVIndex) + TEXT("s"));
+
+	const uint32 UVCount = VertexBuffer->GetNumTexCoords();
+	JsonAttributes.TexCoords.AddUninitialized(UVCount);
+
+	for (uint32 UVIndex = 0; UVIndex < UVCount; ++UVIndex)
+	{
+		JsonAttributes.TexCoords[UVIndex] = Builder.GetOrAddUVAccessor(VertexBuffer, UVIndex, Name + TEXT("_UV") + FString::FromInt(UVIndex) + TEXT("s"));
+	}
 
 	const FRawStaticIndexBuffer* IndexBuffer = &LODResources.IndexBuffer;
 	Builder.GetOrAddIndexBufferView(IndexBuffer, Name + TEXT("_Indices"));
diff --git a/Engine/Plugins/Enterprise/GLTFExporter/MigratedFromGit/Source/GLTFExporter/Private/Json/GLTFJsonMesh.h b/Engine/Plugins/Enterprise/GLTFExporter/MigratedFromGit/Source/GLTFExporter/Private/Json/GLTFJsonMesh.h
--- a/Engine/Plugins/Enterprise/GLTFExporter/MigratedFromGit/Source/GLTFExporter/Private/Json/GLTFJsonMesh.h
+++ b/Engine/Plugins/Enterprise/GLTFExporter/MigratedFromGit/Source/GLTFExporter/Private/Json/GLTFJsonMesh.h
@@ -9,11 +9,11 @@
 struct FGLTFJsonAttributes
 {
 	FGLTFJsonAccessorIndex Position; // always required
+	FGLTFJsonAccessorIndex Color0;
 	FGLTFJsonAccessorIndex Normal;
 	FGLTFJsonAccessorIndex Tangent;
-	FGLTFJsonAccessorIndex TexCoord0;
-	FGLTFJsonAccessorIndex TexCoord1;
-	FGLTFJsonAccessorIndex Color0;
+	TArray<FGLTFJsonAccessorIndex> TexCoords;
+
 	// skeletal mesh attributes
 	FGLTFJsonAccessorIndex Joints0;
 	FGLTFJsonAccessorIndex Weights0;
@@ -24,11 +24,16 @@ struct FGLTFJsonAttributes
 		JsonWriter.WriteObjectStart();
 
 		JsonWriter.WriteValue(TEXT("POSITION"), Position);
+
+		if (Color0 != INDEX_NONE) JsonWriter.WriteValue(TEXT("COLOR_0"), Color0);
 		if (Normal != INDEX_NONE) JsonWriter.WriteValue(TEXT("NORMAL"), Normal);
 		if (Tangent != INDEX_NONE) JsonWriter.WriteValue(TEXT("TANGENT"), Tangent);
-		if (TexCoord0 != INDEX_NONE) JsonWriter.WriteValue(TEXT("TEXCOORD_0"), TexCoord0);
-		if (TexCoord1 != INDEX_NONE) JsonWriter.WriteValue(TEXT("TEXCOORD_1"), TexCoord1);
-		if (Color0 != INDEX_NONE) JsonWriter.WriteValue(TEXT("COLOR_0"), Color0);
+
+		for (int32 Index = 0; Index < TexCoords.Num(); ++Index)
+		{
+			JsonWriter.WriteValue(TEXT("TEXCOORD_") + FString::FromInt(Index), TexCoords[Index]);
+		}
+
 		if (Joints0 != INDEX_NONE) JsonWriter.WriteValue(TEXT("JOINTS_0"), Joints0);
 		if (Weights0 != INDEX_NONE) JsonWriter.WriteValue(TEXT("WEIGHTS_0"), Weights0);
 
