[GLTFExporter] feat(material): add GLTFJsonColor3 and GLTFJsonColor4

commit d332d71b00f7f696ff62c4d33929a9a2069b6485
Date: Tue, 19 May 2020 00:07:46 +0200

    feat(material): add GLTFJsonColor3 and GLTFJsonColor4
    (cherry picked from commit e6b2bcff0a8497a63b09170634324914228585de​)

#skipci

diff --git a/Engine/Plugins/Enterprise/GLTFExporter/MigratedFromGit/Source/GLTFExporter/Private/Json/GLTFJsonColor3.cpp b/Engine/Plugins/Enterprise/GLTFExporter/MigratedFromGit/Source/GLTFExporter/Private/Json/GLTFJsonColor3.cpp
new file mode 100644
--- /dev/null
+++ b/Engine/Plugins/Enterprise/GLTFExporter/MigratedFromGit/Source/GLTFExporter/Private/Json/GLTFJsonColor3.cpp
@@ -0,0 +1,6 @@
+// Copyright Epic Games, Inc. All Rights Reserved.
+
+#include "Json/GLTFJsonColor3.h"
+
+const FGLTFJsonColor3 FGLTFJsonColor3::Black(0, 0, 0);
+const FGLTFJsonColor3 FGLTFJsonColor3::White(1, 1, 1);
diff --git a/Engine/Plugins/Enterprise/GLTFExporter/MigratedFromGit/Source/GLTFExporter/Private/Json/GLTFJsonColor3.h b/Engine/Plugins/Enterprise/GLTFExporter/MigratedFromGit/Source/GLTFExporter/Private/Json/GLTFJsonColor3.h
new file mode 100644
--- /dev/null
+++ b/Engine/Plugins/Enterprise/GLTFExporter/MigratedFromGit/Source/GLTFExporter/Private/Json/GLTFJsonColor3.h
@@ -0,0 +1,42 @@
+// Copyright Epic Games, Inc. All Rights Reserved.
+
+#pragma once
+
+#include "Serialization/JsonSerializer.h"
+
+struct FGLTFJsonColor3
+{
+	float R, G, B;
+
+	static const FGLTFJsonColor3 Black;
+	static const FGLTFJsonColor3 White;
+
+	FGLTFJsonColor3(float R, float G, float B)
+		: R(R), G(G), B(B)
+	{
+	}
+
+	template <class CharType = TCHAR, class PrintPolicy = TPrettyJsonPrintPolicy<CharType>>
+	void WriteArray(TJsonWriter<CharType, PrintPolicy>& JsonWriter) const
+	{
+		JsonWriter.WriteArrayStart();
+		JsonWriter.WriteValue(R);
+		JsonWriter.WriteValue(G);
+		JsonWriter.WriteValue(B);
+		JsonWriter.WriteArrayEnd();
+	}
+
+	bool operator==(const FGLTFJsonColor3& Other) const
+	{
+		return R == Other.R
+			&& G == Other.G
+			&& B == Other.B;
+	}
+
+	bool operator!=(const FGLTFJsonColor3& Other) const
+	{
+		return R != Other.R
+            || G != Other.G
+            || B != Other.B;
+	}
+};
diff --git a/Engine/Plugins/Enterprise/GLTFExporter/MigratedFromGit/Source/GLTFExporter/Private/Json/GLTFJsonColor4.cpp b/Engine/Plugins/Enterprise/GLTFExporter/MigratedFromGit/Source/GLTFExporter/Private/Json/GLTFJsonColor4.cpp
new file mode 100644
--- /dev/null
+++ b/Engine/Plugins/Enterprise/GLTFExporter/MigratedFromGit/Source/GLTFExporter/Private/Json/GLTFJsonColor4.cpp
@@ -0,0 +1,6 @@
+// Copyright Epic Games, Inc. All Rights Reserved.
+
+#include "Json/GLTFJsonColor4.h"
+
+const FGLTFJsonColor4 FGLTFJsonColor4::Black(0, 0, 0, 1);
+const FGLTFJsonColor4 FGLTFJsonColor4::White(1, 1, 1, 1);
diff --git a/Engine/Plugins/Enterprise/GLTFExporter/MigratedFromGit/Source/GLTFExporter/Private/Json/GLTFJsonColor4.h b/Engine/Plugins/Enterprise/GLTFExporter/MigratedFromGit/Source/GLTFExporter/Private/Json/GLTFJsonColor4.h
new file mode 100644
--- /dev/null
+++ b/Engine/Plugins/Enterprise/GLTFExporter/MigratedFromGit/Source/GLTFExporter/Private/Json/GLTFJsonColor4.h
@@ -0,0 +1,51 @@
+// CopGright Epic Games, Inc. All Rights Reserved.
+
+#pragma once
+
+#include "Json/GLTFJsonColor3.h"
+#include "Serialization/JsonSerializer.h"
+
+struct FGLTFJsonColor4
+{
+	float R, G, B, A;
+
+	static const FGLTFJsonColor4 Black;
+	static const FGLTFJsonColor4 White;
+
+	FGLTFJsonColor4(float R, float G, float B, float A)
+        : R(R), G(G), B(B), A(A)
+	{
+	}
+
+	FGLTFJsonColor4(const FGLTFJsonColor3& Color3, float A)
+        : R(Color3.R), G(Color3.G), B(Color3.B), A(A)
+	{
+	}
+
+	template <class CharType = TCHAR, class PrintPolicy = TPrettyJsonPrintPolicy<CharType>>
+    void WriteArray(TJsonWriter<CharType, PrintPolicy>& JsonWriter) const
+	{
+		JsonWriter.WriteArrayStart();
+		JsonWriter.WriteValue(R);
+		JsonWriter.WriteValue(G);
+		JsonWriter.WriteValue(B);
+		JsonWriter.WriteValue(A);
+		JsonWriter.WriteArrayEnd();
+	}
+
+	bool operator==(const FGLTFJsonColor4& Other) const
+	{
+		return R == Other.R
+            && G == Other.G
+            && B == Other.B
+            && A == Other.A;
+	}
+
+	bool operator!=(const FGLTFJsonColor4& Other) const
+	{
+		return R != Other.R
+            || G != Other.G
+            || B != Other.B
+            || A != Other.A;
+	}
+};
