[GLTFExporter] feat(material): add initial JSON related data structures for glTF materials

commit 1dba67711afa59a1d89b660bb2eb650e16cf447e
Date: Tue, 19 May 2020 00:15:03 +0200

    feat(material): add initial JSON related data structures for glTF materials
    (cherry picked from commit 0840583b82f7b15c43529962e4642118c338d644​)

#skipci

diff --git a/Engine/Plugins/Enterprise/GLTFExporter/MigratedFromGit/Source/GLTFExporter/Private/Json/GLTFJsonMaterial.h b/Engine/Plugins/Enterprise/GLTFExporter/MigratedFromGit/Source/GLTFExporter/Private/Json/GLTFJsonMaterial.h
new file mode 100644
--- /dev/null
+++ b/Engine/Plugins/Enterprise/GLTFExporter/MigratedFromGit/Source/GLTFExporter/Private/Json/GLTFJsonMaterial.h
@@ -0,0 +1,107 @@
+// Copyright Epic Games, Inc. All Rights Reserved.
+
+#pragma once
+
+#include "GLTFJsonUtility.h"
+#include "Json/GLTFJsonEnums.h"
+#include "Json/GLTFJsonColor3.h"
+#include "Json/GLTFJsonColor4.h"
+#include "Serialization/JsonSerializer.h"
+
+struct FGLTFJsonPBRMetallicRoughness
+{
+	FGLTFJsonColor4 BaseColorFactor;
+
+	float MetallicFactor;
+	float RoughnessFactor;
+
+	FGLTFJsonPBRMetallicRoughness()
+		: BaseColorFactor(FGLTFJsonColor4::White)
+		, MetallicFactor(1)
+		, RoughnessFactor(1)
+	{
+	}
+
+	template <class CharType = TCHAR, class PrintPolicy = TPrettyJsonPrintPolicy<CharType>>
+    void WriteObject(TJsonWriter<CharType, PrintPolicy>& JsonWriter) const
+	{
+		JsonWriter.WriteObjectStart();
+
+		if (BaseColorFactor != FGLTFJsonColor4::White)
+		{
+			JsonWriter.WriteIdentifierPrefix(TEXT("baseColorFactor"));
+			BaseColorFactor.WriteArray(JsonWriter);
+		}
+
+		if (MetallicFactor != 1)
+		{
+			JsonWriter.WriteValue(TEXT("metallicFactor"), MetallicFactor);
+		}
+
+		if (RoughnessFactor != 1)
+		{
+			JsonWriter.WriteValue(TEXT("roughnessFactor"), RoughnessFactor);
+		}
+
+		JsonWriter.WriteObjectEnd();
+	}
+};
+
+struct FGLTFJsonMaterial
+{
+	FString Name;
+
+	FGLTFJsonPBRMetallicRoughness PBRMetallicRoughness;
+
+	FGLTFJsonColor3 EmissiveFactor;
+
+	EGLTFJsonAlphaMode AlphaMode;
+	float AlphaCutoff;
+
+	bool DoubleSided;
+
+	FGLTFJsonMaterial()
+        : EmissiveFactor(FGLTFJsonColor3::Black)
+        , AlphaMode(EGLTFJsonAlphaMode::Opaque)
+        , AlphaCutoff(0.5f)
+		, DoubleSided(false)
+	{
+	}
+
+	template <class CharType = TCHAR, class PrintPolicy = TPrettyJsonPrintPolicy<CharType>>
+	void WriteObject(TJsonWriter<CharType, PrintPolicy>& JsonWriter) const
+	{
+		JsonWriter.WriteObjectStart();
+
+		if (!Name.IsEmpty())
+		{
+			JsonWriter.WriteValue(TEXT("name"), Name);
+		}
+
+		JsonWriter.WriteIdentifierPrefix(TEXT("pbrMetallicRoughness"));
+		PBRMetallicRoughness.WriteObject(JsonWriter);
+
+		if (EmissiveFactor != FGLTFJsonColor3::Black)
+		{
+			JsonWriter.WriteIdentifierPrefix(TEXT("emissiveFactor"));
+			EmissiveFactor.WriteArray(JsonWriter);
+		}
+
+		if (AlphaMode != EGLTFJsonAlphaMode::Opaque)
+		{
+			JsonWriter.WriteValue(TEXT("alphaMode"), FGLTFJsonUtility::AlphaModeToString(AlphaMode));
+		}
+
+		if (AlphaMode == EGLTFJsonAlphaMode::Mask && AlphaCutoff != 0.5f)
+		{
+			JsonWriter.WriteValue(TEXT("alphaCutoff"), AlphaCutoff);
+		}
+
+		if (DoubleSided)
+		{
+			JsonWriter.WriteValue(TEXT("doubleSided"), DoubleSided);
+		}
+
+		JsonWriter.WriteObjectEnd();
+	}
+};
diff --git a/Engine/Plugins/Enterprise/GLTFExporter/MigratedFromGit/Source/GLTFExporter/Private/Json/GLTFJsonRoot.h b/Engine/Plugins/Enterprise/GLTFExporter/MigratedFromGit/Source/GLTFExporter/Private/Json/GLTFJsonRoot.h
--- a/Engine/Plugins/Enterprise/GLTFExporter/MigratedFromGit/Source/GLTFExporter/Private/Json/GLTFJsonRoot.h
+++ b/Engine/Plugins/Enterprise/GLTFExporter/MigratedFromGit/Source/GLTFExporter/Private/Json/GLTFJsonRoot.h
@@ -6,6 +6,7 @@
 #include "Json/GLTFJsonAccessor.h"
 #include "Json/GLTFJsonBufferView.h"
 #include "Json/GLTFJsonBuffer.h"
+#include "Json/GLTFJsonMaterial.h"
 #include "Json/GLTFJsonMesh.h"
 #include "Json/GLTFJsonScene.h"
 #include "Json/GLTFJsonNode.h"
@@ -60,6 +61,7 @@ struct FGLTFJsonRoot
 	TArray<FGLTFJsonAccessor>   Accessors;
 	TArray<FGLTFJsonBuffer>     Buffers;
 	TArray<FGLTFJsonBufferView> BufferViews;
+	TArray<FGLTFJsonMaterial>   Materials;
 	TArray<FGLTFJsonMesh>       Meshes;
 	TArray<FGLTFJsonNode>       Nodes;
 	TArray<FGLTFJsonScene>      Scenes;
@@ -119,6 +121,16 @@ struct FGLTFJsonRoot
 			JsonWriter.WriteArrayEnd();
 		}
 
+		if (Materials.Num() > 0)
+		{
+			JsonWriter.WriteArrayStart(TEXT("materials"));
+			for (const FGLTFJsonMaterial& Material : Materials)
+			{
+				Material.WriteObject(JsonWriter);
+			}
+			JsonWriter.WriteArrayEnd();
+		}
+
 		if (Meshes.Num() > 0)
 		{
 			JsonWriter.WriteArrayStart(TEXT("meshes"));
