[GLTFExporter] feat(material): add empty initial converter for materials

commit 0818b9c5689eb0f62f65984e0648c17599d0de2b
Date: Mon, 18 May 2020 23:59:11 +0200

    feat(material): add empty initial converter for materials
    (cherry picked from commit c5cfab80d607af4c299270ba134782af0166ef6b​)
    # Conflicts from cherry-pick:
    #	Plugins/GLTFExporter/Source/GLTFExporter/Private/Builders/GLTFConvertBuilder.h

#skipci

diff --git a/Engine/Plugins/Enterprise/GLTFExporter/MigratedFromGit/Source/GLTFExporter/Private/Builders/GLTFConvertBuilder.cpp b/Engine/Plugins/Enterprise/GLTFExporter/MigratedFromGit/Source/GLTFExporter/Private/Builders/GLTFConvertBuilder.cpp
--- a/Engine/Plugins/Enterprise/GLTFExporter/MigratedFromGit/Source/GLTFExporter/Private/Builders/GLTFConvertBuilder.cpp
+++ b/Engine/Plugins/Enterprise/GLTFExporter/MigratedFromGit/Source/GLTFExporter/Private/Builders/GLTFConvertBuilder.cpp
@@ -73,6 +73,11 @@ FGLTFJsonAccessorIndex FGLTFConvertBuilder::GetOrAddIndexAccessor(const FStaticM
 	return StaticMeshSectionConverter.GetOrAdd(*this, DesiredName, MeshSection, IndexBuffer);
 }
 
+FGLTFJsonMaterialIndex FGLTFConvertBuilder::GetOrAddMaterial(const UMaterialInterface* Material, const FString& DesiredName)
+{
+	return MaterialConverter.GetOrAdd(*this, DesiredName.IsEmpty() ? Material->GetName() : DesiredName, Material);
+}
+
 FGLTFJsonMeshIndex FGLTFConvertBuilder::GetOrAddMesh(const UStaticMesh* StaticMesh, int32 LODIndex, const FColorVertexBuffer* OverrideVertexColors, const FString& DesiredName)
 {
 	if (StaticMesh == nullptr)
diff --git a/Engine/Plugins/Enterprise/GLTFExporter/MigratedFromGit/Source/GLTFExporter/Private/Builders/GLTFConvertBuilder.h b/Engine/Plugins/Enterprise/GLTFExporter/MigratedFromGit/Source/GLTFExporter/Private/Builders/GLTFConvertBuilder.h
--- a/Engine/Plugins/Enterprise/GLTFExporter/MigratedFromGit/Source/GLTFExporter/Private/Builders/GLTFConvertBuilder.h
+++ b/Engine/Plugins/Enterprise/GLTFExporter/MigratedFromGit/Source/GLTFExporter/Private/Builders/GLTFConvertBuilder.h
@@ -5,6 +5,7 @@
 #include "Builders/GLTFBufferBuilder.h"
 #include "Converters/GLTFVertexBufferConverters.h"
 #include "Converters/GLTFStaticMeshConverters.h"
+#include "Converters/GLTFMaterialConverters.h"
 #include "Converters/GLTFLevelConverters.h"
 
 class FGLTFConvertBuilder : public FGLTFBufferBuilder
@@ -19,6 +20,9 @@ public:
 
 	FGLTFJsonBufferViewIndex GetOrAddIndexBufferView(const FRawStaticIndexBuffer* IndexBuffer, const FString& DesiredName = TEXT(""));
 	FGLTFJsonAccessorIndex GetOrAddIndexAccessor(const FStaticMeshSection* MeshSection, const FRawStaticIndexBuffer* IndexBuffer, const FString& DesiredName = TEXT(""));
+
+	FGLTFJsonMaterialIndex GetOrAddMaterial(const UMaterialInterface* Material, const FString& DesiredName = TEXT(""));
+
 	FGLTFJsonMeshIndex GetOrAddMesh(const UStaticMesh* StaticMesh, int32 LODIndex = 0, const FColorVertexBuffer* OverrideVertexColors = nullptr, const FString& DesiredName = TEXT(""));
 	FGLTFJsonMeshIndex GetOrAddMesh(const UStaticMeshComponent* StaticMeshComponent, const FString& DesiredName = TEXT(""));
 
@@ -38,6 +42,8 @@ private:
 	FGLTFStaticMeshSectionConverter StaticMeshSectionConverter;
 	FGLTFStaticMeshConverter StaticMeshConverter;
 
+	FGLTFMaterialConverter MaterialConverter;
+
 	FGLTFSceneComponentConverter SceneComponentConverter;
 	FGLTFLevelConverter LevelConverter;
 };
diff --git a/Engine/Plugins/Enterprise/GLTFExporter/MigratedFromGit/Source/GLTFExporter/Private/Builders/GLTFJsonBuilder.cpp b/Engine/Plugins/Enterprise/GLTFExporter/MigratedFromGit/Source/GLTFExporter/Private/Builders/GLTFJsonBuilder.cpp
--- a/Engine/Plugins/Enterprise/GLTFExporter/MigratedFromGit/Source/GLTFExporter/Private/Builders/GLTFJsonBuilder.cpp
+++ b/Engine/Plugins/Enterprise/GLTFExporter/MigratedFromGit/Source/GLTFExporter/Private/Builders/GLTFJsonBuilder.cpp
@@ -26,6 +26,11 @@ FGLTFJsonBufferViewIndex FGLTFJsonBuilder::AddBufferView(const FGLTFJsonBufferVi
 	return FGLTFJsonBufferViewIndex(JsonRoot.BufferViews.Add(JsonBufferView));
 }
 
+FGLTFJsonMaterialIndex FGLTFJsonBuilder::AddMaterial(const FGLTFJsonMaterial& JsonMaterial)
+{
+	return FGLTFJsonMaterialIndex(JsonRoot.Materials.Add(JsonMaterial));
+}
+
 FGLTFJsonMeshIndex FGLTFJsonBuilder::AddMesh(const FGLTFJsonMesh& JsonMesh)
 {
 	return FGLTFJsonMeshIndex(JsonRoot.Meshes.Add(JsonMesh));
diff --git a/Engine/Plugins/Enterprise/GLTFExporter/MigratedFromGit/Source/GLTFExporter/Private/Builders/GLTFJsonBuilder.h b/Engine/Plugins/Enterprise/GLTFExporter/MigratedFromGit/Source/GLTFExporter/Private/Builders/GLTFJsonBuilder.h
--- a/Engine/Plugins/Enterprise/GLTFExporter/MigratedFromGit/Source/GLTFExporter/Private/Builders/GLTFJsonBuilder.h
+++ b/Engine/Plugins/Enterprise/GLTFExporter/MigratedFromGit/Source/GLTFExporter/Private/Builders/GLTFJsonBuilder.h
@@ -16,6 +16,7 @@ public:
 	FGLTFJsonAccessorIndex AddAccessor(const FGLTFJsonAccessor& JsonAccessor);
 	FGLTFJsonBufferIndex AddBuffer(const FGLTFJsonBuffer& JsonBuffer);
 	FGLTFJsonBufferViewIndex AddBufferView(const FGLTFJsonBufferView& JsonBufferView);
+	FGLTFJsonMaterialIndex AddMaterial(const FGLTFJsonMaterial& JsonMaterial);
 	FGLTFJsonMeshIndex AddMesh(const FGLTFJsonMesh& JsonMesh);
 	FGLTFJsonNodeIndex AddNode(const FGLTFJsonNode& JsonNode);
 	FGLTFJsonSceneIndex AddScene(const FGLTFJsonScene& JsonScene);
diff --git a/Engine/Plugins/Enterprise/GLTFExporter/MigratedFromGit/Source/GLTFExporter/Private/Converters/GLTFMaterialConverters.cpp b/Engine/Plugins/Enterprise/GLTFExporter/MigratedFromGit/Source/GLTFExporter/Private/Converters/GLTFMaterialConverters.cpp
new file mode 100644
--- /dev/null
+++ b/Engine/Plugins/Enterprise/GLTFExporter/MigratedFromGit/Source/GLTFExporter/Private/Converters/GLTFMaterialConverters.cpp
@@ -0,0 +1,13 @@
+// Copyright Epic Games, Inc. All Rights Reserved.
+
+#include "Converters/GLTFMaterialConverters.h"
+#include "Builders/GLTFContainerBuilder.h"
+#include "Converters/GLTFConverterUtility.h"
+
+FGLTFJsonMaterialIndex FGLTFMaterialConverter::Add(FGLTFConvertBuilder& Builder, const FString& Name, const UMaterialInterface* MaterialInterface)
+{
+	FGLTFJsonMaterial JsonMaterial;
+	JsonMaterial.Name = Name;
+
+	return Builder.AddMaterial(JsonMaterial);
+}
diff --git a/Engine/Plugins/Enterprise/GLTFExporter/MigratedFromGit/Source/GLTFExporter/Private/Converters/GLTFMaterialConverters.h b/Engine/Plugins/Enterprise/GLTFExporter/MigratedFromGit/Source/GLTFExporter/Private/Converters/GLTFMaterialConverters.h
new file mode 100644
--- /dev/null
+++ b/Engine/Plugins/Enterprise/GLTFExporter/MigratedFromGit/Source/GLTFExporter/Private/Converters/GLTFMaterialConverters.h
@@ -0,0 +1,12 @@
+// Copyright Epic Games, Inc. All Rights Reserved.
+
+#pragma once
+
+#include "Json/GLTFJsonIndex.h"
+#include "Converters/GLTFConverter.h"
+#include "Engine.h"
+
+class FGLTFMaterialConverter final : public TGLTFConverter<FGLTFJsonMaterialIndex, const UMaterialInterface*>
+{
+	FGLTFJsonMaterialIndex Add(FGLTFConvertBuilder& Builder, const FString& Name, const UMaterialInterface* Material) override;
+};
