[GLTFExporter] feat(material): add conversion operator for GLTFJsonColor4 -> GLTFJsonColor3 and GLTFJsonVector4 -> GLTFJsonVector3

commit b2371b3d729e407475a0f9c07ae90f7750462e21
Date: Tue, 19 May 2020 04:35:11 +0200

    feat(material): add conversion operator for GLTFJsonColor4 -> GLTFJsonColor3 and GLTFJsonVector4 -> GLTFJsonVector3
    (cherry picked from commit 40d06d4213826981020bed089d560df1cafd71c7​)

#skipci

diff --git a/Engine/Plugins/Enterprise/GLTFExporter/MigratedFromGit/Source/GLTFExporter/Private/Json/GLTFJsonColor4.h b/Engine/Plugins/Enterprise/GLTFExporter/MigratedFromGit/Source/GLTFExporter/Private/Json/GLTFJsonColor4.h
--- a/Engine/Plugins/Enterprise/GLTFExporter/MigratedFromGit/Source/GLTFExporter/Private/Json/GLTFJsonColor4.h
+++ b/Engine/Plugins/Enterprise/GLTFExporter/MigratedFromGit/Source/GLTFExporter/Private/Json/GLTFJsonColor4.h
@@ -12,12 +12,12 @@ struct FGLTFJsonColor4
 	static const FGLTFJsonColor4 Black;
 	static const FGLTFJsonColor4 White;
 
-	FGLTFJsonColor4(float R, float G, float B, float A)
+	FGLTFJsonColor4(float R, float G, float B, float A = 1.0f)
         : R(R), G(G), B(B), A(A)
 	{
 	}
 
-	FGLTFJsonColor4(const FGLTFJsonColor3& Color3, float A)
+	FGLTFJsonColor4(const FGLTFJsonColor3& Color3, float A = 1.0f)
         : R(Color3.R), G(Color3.G), B(Color3.B), A(A)
 	{
 	}
@@ -33,6 +33,11 @@ struct FGLTFJsonColor4
 		JsonWriter.WriteArrayEnd();
 	}
 
+	operator FGLTFJsonColor3() const
+	{
+		return { R, G, B };
+	}
+
 	bool operator==(const FGLTFJsonColor4& Other) const
 	{
 		return R == Other.R
diff --git a/Engine/Plugins/Enterprise/GLTFExporter/MigratedFromGit/Source/GLTFExporter/Private/Json/GLTFJsonVector4.h b/Engine/Plugins/Enterprise/GLTFExporter/MigratedFromGit/Source/GLTFExporter/Private/Json/GLTFJsonVector4.h
--- a/Engine/Plugins/Enterprise/GLTFExporter/MigratedFromGit/Source/GLTFExporter/Private/Json/GLTFJsonVector4.h
+++ b/Engine/Plugins/Enterprise/GLTFExporter/MigratedFromGit/Source/GLTFExporter/Private/Json/GLTFJsonVector4.h
@@ -34,6 +34,11 @@ struct FGLTFJsonVector4
 		JsonWriter.WriteArrayEnd();
 	}
 
+	operator FGLTFJsonVector3() const
+	{
+		return { X, Y, Z };
+	}
+
 	bool operator==(const FGLTFJsonVector4& Other) const
 	{
 		return X == Other.X
