[GLTFExporter] feat(material): implement initial material converters that support constant values

commit 57c55f2eb599fe4678eedaf09ac9afaa96ca09ba
Date: Tue, 19 May 2020 04:35:58 +0200

    feat(material): implement initial material converters that support constant values
    (cherry picked from commit 71f03678d3cc3d19ddc3ac07790a8485446360d8​)
    # Conflicts from cherry-pick:
    #	Plugins/GLTFExporter/Source/GLTFExporter/Private/Converters/GLTFStaticMeshConverters.cpp

#skipci

diff --git a/Engine/Plugins/Enterprise/GLTFExporter/MigratedFromGit/Source/GLTFExporter/Private/Converters/GLTFConverterUtility.h b/Engine/Plugins/Enterprise/GLTFExporter/MigratedFromGit/Source/GLTFExporter/Private/Converters/GLTFConverterUtility.h
--- a/Engine/Plugins/Enterprise/GLTFExporter/MigratedFromGit/Source/GLTFExporter/Private/Converters/GLTFConverterUtility.h
+++ b/Engine/Plugins/Enterprise/GLTFExporter/MigratedFromGit/Source/GLTFExporter/Private/Converters/GLTFConverterUtility.h
@@ -7,6 +7,7 @@
 #include "Json/GLTFJsonVector2.h"
 #include "Json/GLTFJsonVector3.h"
 #include "Json/GLTFJsonVector4.h"
+#include "Json/GLTFJsonColor4.h"
 #include "Json/GLTFJsonQuaternion.h"
 #include "Engine/EngineTypes.h"
 
@@ -46,6 +47,12 @@ struct FGLTFConverterUtility
 		return { UV.X, UV.Y };
 	}
 
+	static FGLTFJsonColor4 ConvertColor(const FLinearColor& Color)
+	{
+		// No conversion actually needed, this is primarily for type-safety.
+		return { Color.R, Color.G, Color.B, Color.A };
+	}
+
 	static FGLTFPackedColor ConvertColor(const FColor& Color)
 	{
 		// UE4 uses ABGR or ARGB depending on endianness.
@@ -84,7 +91,7 @@ struct FGLTFConverterUtility
 		return Result;
 	}
 
-	static EGLTFJsonAlphaMode ConvertAlphaMode(const EBlendMode Mode)
+	static EGLTFJsonAlphaMode ConvertBlendMode(const EBlendMode Mode)
 	{
 		switch (Mode)
 		{
diff --git a/Engine/Plugins/Enterprise/GLTFExporter/MigratedFromGit/Source/GLTFExporter/Private/Converters/GLTFMaterialConverters.cpp b/Engine/Plugins/Enterprise/GLTFExporter/MigratedFromGit/Source/GLTFExporter/Private/Converters/GLTFMaterialConverters.cpp
--- a/Engine/Plugins/Enterprise/GLTFExporter/MigratedFromGit/Source/GLTFExporter/Private/Converters/GLTFMaterialConverters.cpp
+++ b/Engine/Plugins/Enterprise/GLTFExporter/MigratedFromGit/Source/GLTFExporter/Private/Converters/GLTFMaterialConverters.cpp
@@ -3,11 +3,115 @@
 #include "Converters/GLTFMaterialConverters.h"
 #include "Builders/GLTFContainerBuilder.h"
 #include "Converters/GLTFConverterUtility.h"
+#include "Materials/MaterialExpressionConstant.h"
+#include "Materials/MaterialExpressionVectorParameter.h"
+#include "Materials/MaterialExpressionConstant2Vector.h"
+#include "Materials/MaterialExpressionConstant3Vector.h"
+#include "Materials/MaterialExpressionConstant4Vector.h"
+#include "Materials/MaterialExpressionScalarParameter.h"
+#include "Materials/MaterialExpressionTextureSample.h"
 
 FGLTFJsonMaterialIndex FGLTFMaterialConverter::Add(FGLTFConvertBuilder& Builder, const FString& Name, const UMaterialInterface* MaterialInterface)
 {
 	FGLTFJsonMaterial JsonMaterial;
 	JsonMaterial.Name = Name;
 
+	const UMaterial* Material = MaterialInterface->GetMaterial();
+	JsonMaterial.PBRMetallicRoughness.BaseColorFactor = FGLTFConverterUtility::ConvertColor(GetConstantColor(Material->BaseColor));
+	JsonMaterial.PBRMetallicRoughness.MetallicFactor = GetConstantScalar(Material->Metallic);
+	JsonMaterial.PBRMetallicRoughness.RoughnessFactor = GetConstantScalar(Material->Roughness);
+	JsonMaterial.EmissiveFactor = FGLTFConverterUtility::ConvertColor(GetConstantColor(Material->EmissiveColor));
+
+	JsonMaterial.AlphaMode = FGLTFConverterUtility::ConvertBlendMode(Material->BlendMode);
+	if (JsonMaterial.AlphaMode == EGLTFJsonAlphaMode::Blend)
+	{
+		JsonMaterial.PBRMetallicRoughness.BaseColorFactor.A = GetConstantScalar(Material->Opacity);
+	}
+	else if (JsonMaterial.AlphaMode == EGLTFJsonAlphaMode::Mask)
+	{
+		JsonMaterial.PBRMetallicRoughness.BaseColorFactor.A = GetConstantScalar(Material->OpacityMask);
+	}
+
+	JsonMaterial.AlphaCutoff = Material->OpacityMaskClipValue;
+	JsonMaterial.DoubleSided = Material->TwoSided;
+
 	return Builder.AddMaterial(JsonMaterial);
 }
+
+FLinearColor FGLTFMaterialConverter::GetConstantColor(const FColorMaterialInput& MaterialInput) const
+{
+	const UMaterialExpression* Expression = MaterialInput.Expression;
+	if (Expression != nullptr)
+	{
+		if (const UMaterialExpressionVectorParameter* VectorParameter = Cast<UMaterialExpressionVectorParameter>(Expression))
+		{
+			return VectorParameter->DefaultValue;
+		}
+
+		if (const UMaterialExpressionScalarParameter* ScalarParameter = Cast<UMaterialExpressionScalarParameter>(Expression))
+		{
+			return { ScalarParameter->DefaultValue, ScalarParameter->DefaultValue, ScalarParameter->DefaultValue };
+		}
+
+		if (const UMaterialExpressionConstant4Vector* Constant4Vector =Cast<UMaterialExpressionConstant4Vector>(Expression))
+		{
+			return Constant4Vector->Constant;
+		}
+
+		if (const UMaterialExpressionConstant3Vector* Constant3Vector = Cast<UMaterialExpressionConstant3Vector>(Expression))
+		{
+			return Constant3Vector->Constant;
+		}
+
+		if (const UMaterialExpressionConstant2Vector* Constant2Vector =Cast<UMaterialExpressionConstant2Vector>(Expression))
+		{
+			return { Constant2Vector->R, Constant2Vector->G, 0 }; // TODO: is this case ever possible?
+		}
+
+		if (const UMaterialExpressionConstant* Constant = Cast<UMaterialExpressionConstant>(Expression))
+		{
+			return { Constant->R, Constant->R, Constant->R };
+		}
+	}
+
+	return MaterialInput.Constant;
+}
+
+float FGLTFMaterialConverter::GetConstantScalar(const FScalarMaterialInput& MaterialInput) const
+{
+	const UMaterialExpression* Expression = MaterialInput.Expression;
+	if (Expression != nullptr)
+	{
+		if (const UMaterialExpressionVectorParameter* VectorParameter = Cast<UMaterialExpressionVectorParameter>(Expression))
+		{
+			return VectorParameter->DefaultValue.R;
+		}
+
+		if (const UMaterialExpressionScalarParameter* ScalarParameter = Cast<UMaterialExpressionScalarParameter>(Expression))
+		{
+			return ScalarParameter->DefaultValue;
+		}
+
+		if (const UMaterialExpressionConstant4Vector* Constant4Vector =Cast<UMaterialExpressionConstant4Vector>(Expression))
+		{
+			return Constant4Vector->Constant.R;
+		}
+
+		if (const UMaterialExpressionConstant3Vector* Constant3Vector = Cast<UMaterialExpressionConstant3Vector>(Expression))
+		{
+			return Constant3Vector->Constant.R;
+		}
+
+		if (const UMaterialExpressionConstant2Vector* Constant2Vector =Cast<UMaterialExpressionConstant2Vector>(Expression))
+		{
+			return Constant2Vector->R;
+		}
+
+		if (const UMaterialExpressionConstant* Constant = Cast<UMaterialExpressionConstant>(Expression))
+		{
+			return Constant->R;
+		}
+	}
+
+	return MaterialInput.Constant;
+}
diff --git a/Engine/Plugins/Enterprise/GLTFExporter/MigratedFromGit/Source/GLTFExporter/Private/Converters/GLTFMaterialConverters.h b/Engine/Plugins/Enterprise/GLTFExporter/MigratedFromGit/Source/GLTFExporter/Private/Converters/GLTFMaterialConverters.h
--- a/Engine/Plugins/Enterprise/GLTFExporter/MigratedFromGit/Source/GLTFExporter/Private/Converters/GLTFMaterialConverters.h
+++ b/Engine/Plugins/Enterprise/GLTFExporter/MigratedFromGit/Source/GLTFExporter/Private/Converters/GLTFMaterialConverters.h
@@ -9,4 +9,7 @@
 class FGLTFMaterialConverter final : public TGLTFConverter<FGLTFJsonMaterialIndex, const UMaterialInterface*>
 {
 	FGLTFJsonMaterialIndex Add(FGLTFConvertBuilder& Builder, const FString& Name, const UMaterialInterface* Material) override;
+
+	FLinearColor GetConstantColor(const FColorMaterialInput& MaterialInput) const;
+	float GetConstantScalar(const FScalarMaterialInput& MaterialInput) const;
 };
diff --git a/Engine/Plugins/Enterprise/GLTFExporter/MigratedFromGit/Source/GLTFExporter/Private/Converters/GLTFStaticMeshConverters.cpp b/Engine/Plugins/Enterprise/GLTFExporter/MigratedFromGit/Source/GLTFExporter/Private/Converters/GLTFStaticMeshConverters.cpp
--- a/Engine/Plugins/Enterprise/GLTFExporter/MigratedFromGit/Source/GLTFExporter/Private/Converters/GLTFStaticMeshConverters.cpp
+++ b/Engine/Plugins/Enterprise/GLTFExporter/MigratedFromGit/Source/GLTFExporter/Private/Converters/GLTFStaticMeshConverters.cpp
@@ -91,6 +91,9 @@ FGLTFJsonMeshIndex FGLTFStaticMeshConverter::Add(FGLTFConvertBuilder& Builder, c
 		const FStaticMeshSection& Section = LODResources.Sections[SectionIndex];
 		JsonPrimitive.Indices = Builder.GetOrAddIndexAccessor(&Section, IndexBuffer,
 			Name + (SectionCount == 1 ? TEXT("_Indices") : TEXT("_Indices_Section") + FString::FromInt(SectionIndex)));
+
+		const UMaterialInterface* Material = StaticMesh->GetMaterial(Section.MaterialIndex);
+		if (Material != nullptr) JsonPrimitive.Material = Builder.GetOrAddMaterial(Material);
 	}
 
 	return Builder.AddMesh(JsonMesh);
