[GLTFExporter] feat(material): add support for material overrides in StaticMeshComponents

commit 0dcb9be4c354a26492aab074f71e559a35a78a4c
Date: Tue, 19 May 2020 04:39:14 +0200

    feat(material): add support for material overrides in StaticMeshComponents
    (cherry picked from commit 342c1dad7cf359f50833610365702b9d5c0f66ca​)
    # Conflicts from cherry-pick:
    #	Plugins/GLTFExporter/Source/GLTFExporter/Private/Builders/GLTFConvertBuilder.cpp
    #	Plugins/GLTFExporter/Source/GLTFExporter/Private/Converters/GLTFStaticMeshConverters.cpp
    #	Plugins/GLTFExporter/Source/GLTFExporter/Private/Converters/GLTFStaticMeshConverters.h

#skipci

diff --git a/Engine/Plugins/Enterprise/GLTFExporter/MigratedFromGit/Source/GLTFExporter/Private/Builders/GLTFBuilderUtility.h b/Engine/Plugins/Enterprise/GLTFExporter/MigratedFromGit/Source/GLTFExporter/Private/Builders/GLTFBuilderUtility.h
--- a/Engine/Plugins/Enterprise/GLTFExporter/MigratedFromGit/Source/GLTFExporter/Private/Builders/GLTFBuilderUtility.h
+++ b/Engine/Plugins/Enterprise/GLTFExporter/MigratedFromGit/Source/GLTFExporter/Private/Builders/GLTFBuilderUtility.h
@@ -18,4 +18,14 @@ struct FGLTFBuilderUtility
 
 		return Name;
 	}
+
+	template <typename ElementType, typename AllocatorType>
+	static TArray<const ElementType*, AllocatorType> MakeArrayOfPointersConst(const TArray<ElementType*, AllocatorType>& ArrayOfPointers)
+	{
+		TArray<const ElementType*, AllocatorType> ArrayOfConstPointers;
+		ArrayOfConstPointers.AddUninitialized(ArrayOfPointers.Num());
+
+		FMemory::Memcpy(ArrayOfConstPointers.GetData(), ArrayOfPointers.GetData(), ArrayOfPointers.Num() * sizeof(ElementType*));
+		return ArrayOfConstPointers;
+	}
 };
diff --git a/Engine/Plugins/Enterprise/GLTFExporter/MigratedFromGit/Source/GLTFExporter/Private/Builders/GLTFConvertBuilder.cpp b/Engine/Plugins/Enterprise/GLTFExporter/MigratedFromGit/Source/GLTFExporter/Private/Builders/GLTFConvertBuilder.cpp
--- a/Engine/Plugins/Enterprise/GLTFExporter/MigratedFromGit/Source/GLTFExporter/Private/Builders/GLTFConvertBuilder.cpp
+++ b/Engine/Plugins/Enterprise/GLTFExporter/MigratedFromGit/Source/GLTFExporter/Private/Builders/GLTFConvertBuilder.cpp
@@ -78,14 +78,14 @@ FGLTFJsonMaterialIndex FGLTFConvertBuilder::GetOrAddMaterial(const UMaterialInte
 	return MaterialConverter.GetOrAdd(*this, DesiredName.IsEmpty() ? Material->GetName() : DesiredName, Material);
 }
 
-FGLTFJsonMeshIndex FGLTFConvertBuilder::GetOrAddMesh(const UStaticMesh* StaticMesh, int32 LODIndex, const FColorVertexBuffer* OverrideVertexColors, const FString& DesiredName)
+FGLTFJsonMeshIndex FGLTFConvertBuilder::GetOrAddMesh(const UStaticMesh* StaticMesh, int32 LODIndex, const FColorVertexBuffer* OverrideVertexColors, TArray<const UMaterialInterface*> OverrideMaterials, const FString& DesiredName)
 {
 	if (StaticMesh == nullptr)
 	{
 		return FGLTFJsonMeshIndex(INDEX_NONE);
 	}
 
-	return StaticMeshConverter.GetOrAdd(*this, DesiredName.IsEmpty() ? FGLTFBuilderUtility::GetMeshName(StaticMesh, LODIndex) : DesiredName, StaticMesh, LODIndex, OverrideVertexColors);
+	return StaticMeshConverter.GetOrAdd(*this, DesiredName.IsEmpty() ? FGLTFBuilderUtility::GetMeshName(StaticMesh, LODIndex) : DesiredName, StaticMesh, LODIndex, OverrideVertexColors, OverrideMaterials);
 }
 
 FGLTFJsonMeshIndex FGLTFConvertBuilder::GetOrAddMesh(const UStaticMeshComponent* StaticMeshComponent, const FString& DesiredName)
@@ -98,8 +98,9 @@ FGLTFJsonMeshIndex FGLTFConvertBuilder::GetOrAddMesh(const UStaticMeshComponent*
 	const UStaticMesh* StaticMesh = StaticMeshComponent->GetStaticMesh();
 	const int32 LODIndex = StaticMeshComponent->ForcedLodModel > 0 ? StaticMeshComponent->ForcedLodModel - 1 : /* auto-select */ 0;
 	const FColorVertexBuffer* OverrideVertexColors = StaticMeshComponent->LODData.IsValidIndex(LODIndex) ? StaticMeshComponent->LODData[LODIndex].OverrideVertexColors : nullptr;
+	const TArray<const UMaterialInterface*> OverrideMaterials = FGLTFBuilderUtility::MakeArrayOfPointersConst(StaticMeshComponent->OverrideMaterials);
 
-	return GetOrAddMesh(StaticMesh, LODIndex, OverrideVertexColors, DesiredName);
+	return GetOrAddMesh(StaticMesh, LODIndex, OverrideVertexColors, OverrideMaterials, DesiredName);
 }
 
 FGLTFJsonNodeIndex FGLTFConvertBuilder::GetOrAddNode(const USceneComponent* SceneComponent, bool bSelectedOnly, bool bRootNode, const FString& DesiredName)
diff --git a/Engine/Plugins/Enterprise/GLTFExporter/MigratedFromGit/Source/GLTFExporter/Private/Builders/GLTFConvertBuilder.h b/Engine/Plugins/Enterprise/GLTFExporter/MigratedFromGit/Source/GLTFExporter/Private/Builders/GLTFConvertBuilder.h
--- a/Engine/Plugins/Enterprise/GLTFExporter/MigratedFromGit/Source/GLTFExporter/Private/Builders/GLTFConvertBuilder.h
+++ b/Engine/Plugins/Enterprise/GLTFExporter/MigratedFromGit/Source/GLTFExporter/Private/Builders/GLTFConvertBuilder.h
@@ -23,7 +23,7 @@ public:
 
 	FGLTFJsonMaterialIndex GetOrAddMaterial(const UMaterialInterface* Material, const FString& DesiredName = TEXT(""));
 
-	FGLTFJsonMeshIndex GetOrAddMesh(const UStaticMesh* StaticMesh, int32 LODIndex = 0, const FColorVertexBuffer* OverrideVertexColors = nullptr, const FString& DesiredName = TEXT(""));
+	FGLTFJsonMeshIndex GetOrAddMesh(const UStaticMesh* StaticMesh, int32 LODIndex = 0, const FColorVertexBuffer* OverrideVertexColors = nullptr, TArray<const UMaterialInterface*> OverrideMaterials = {}, const FString& DesiredName = TEXT(""));
 	FGLTFJsonMeshIndex GetOrAddMesh(const UStaticMeshComponent* StaticMeshComponent, const FString& DesiredName = TEXT(""));
 
 	FGLTFJsonNodeIndex GetOrAddNode(const USceneComponent* SceneComponent, bool bSelectedOnly, bool bRootNode = false, const FString& DesiredName = TEXT(""));
diff --git a/Engine/Plugins/Enterprise/GLTFExporter/MigratedFromGit/Source/GLTFExporter/Private/Converters/GLTFStaticMeshConverters.cpp b/Engine/Plugins/Enterprise/GLTFExporter/MigratedFromGit/Source/GLTFExporter/Private/Converters/GLTFStaticMeshConverters.cpp
--- a/Engine/Plugins/Enterprise/GLTFExporter/MigratedFromGit/Source/GLTFExporter/Private/Converters/GLTFStaticMeshConverters.cpp
+++ b/Engine/Plugins/Enterprise/GLTFExporter/MigratedFromGit/Source/GLTFExporter/Private/Converters/GLTFStaticMeshConverters.cpp
@@ -47,7 +47,7 @@ FGLTFJsonAccessorIndex FGLTFStaticMeshSectionConverter::Add(FGLTFConvertBuilder&
 	return Builder.AddAccessor(JsonAccessor);
 }
 
-FGLTFJsonMeshIndex FGLTFStaticMeshConverter::Add(FGLTFConvertBuilder& Builder, const FString& Name, const UStaticMesh* StaticMesh, int32 LODIndex, const FColorVertexBuffer* OverrideVertexColors)
+FGLTFJsonMeshIndex FGLTFStaticMeshConverter::Add(FGLTFConvertBuilder& Builder, const FString& Name, const UStaticMesh* StaticMesh, int32 LODIndex, const FColorVertexBuffer* OverrideVertexColors, TArray<const UMaterialInterface*> OverrideMaterials)
 {
 	if (LODIndex < 0 || StaticMesh->GetNumLODs() <= LODIndex)
 	{
@@ -92,7 +92,8 @@ FGLTFJsonMeshIndex FGLTFStaticMeshConverter::Add(FGLTFConvertBuilder& Builder, c
 		JsonPrimitive.Indices = Builder.GetOrAddIndexAccessor(&Section, IndexBuffer,
 			Name + (SectionCount == 1 ? TEXT("_Indices") : TEXT("_Indices_Section") + FString::FromInt(SectionIndex)));
 
-		const UMaterialInterface* Material = StaticMesh->GetMaterial(Section.MaterialIndex);
+		const int32 MaterialIndex = Section.MaterialIndex;
+		const UMaterialInterface* Material = MaterialIndex < OverrideMaterials.Num() ? OverrideMaterials[MaterialIndex] : StaticMesh->GetMaterial(MaterialIndex);
 		if (Material != nullptr) JsonPrimitive.Material = Builder.GetOrAddMaterial(Material);
 	}
 
diff --git a/Engine/Plugins/Enterprise/GLTFExporter/MigratedFromGit/Source/GLTFExporter/Private/Converters/GLTFStaticMeshConverters.h b/Engine/Plugins/Enterprise/GLTFExporter/MigratedFromGit/Source/GLTFExporter/Private/Converters/GLTFStaticMeshConverters.h
--- a/Engine/Plugins/Enterprise/GLTFExporter/MigratedFromGit/Source/GLTFExporter/Private/Converters/GLTFStaticMeshConverters.h
+++ b/Engine/Plugins/Enterprise/GLTFExporter/MigratedFromGit/Source/GLTFExporter/Private/Converters/GLTFStaticMeshConverters.h
@@ -18,5 +18,16 @@ class FGLTFStaticMeshSectionConverter final : public TGLTFConverter<FGLTFJsonAcc
 
 class FGLTFStaticMeshConverter final : public TGLTFConverter<FGLTFJsonMeshIndex, const UStaticMesh*, int32, const FColorVertexBuffer*>
 {
-	FGLTFJsonMeshIndex Add(FGLTFConvertBuilder& Builder, const FString& Name, const UStaticMesh* StaticMesh, int32 LODIndex, const FColorVertexBuffer* OverrideVertexColors) override;
+	FGLTFJsonMeshIndex Add(FGLTFConvertBuilder& Builder, const FString& Name, const UStaticMesh* StaticMesh, int32 LODIndex, const FColorVertexBuffer* OverrideVertexColors, TArray<const UMaterialInterface*> OverrideMaterials) override;
+
+	// TODO: move this to somewhere cleaner
+	friend inline uint32 GetTypeHash(TArray<const UMaterialInterface*> Materials)
+	{
+		uint32 Hash = GetTypeHash(Materials.Num());
+		for (const UMaterialInterface* Material : Materials)
+		{
+			Hash = HashCombine(Hash, GetTypeHash(Material));
+		}
+		return Hash;
+	}
 };
