[GLTFExporter] feat(material): add support for material instance parameter lookup

commit f1ea6c275b9e270906e53d5853abbc8e686ac937
Date: Tue, 19 May 2020 05:00:08 +0200

    feat(material): add support for material instance parameter lookup
    (cherry picked from commit 6cfa398bfcc42423721aeb00ee2a1f74ac8fa5aa​)

#skipci

diff --git a/Engine/Plugins/Enterprise/GLTFExporter/MigratedFromGit/Source/GLTFExporter/Private/Converters/GLTFMaterialConverters.cpp b/Engine/Plugins/Enterprise/GLTFExporter/MigratedFromGit/Source/GLTFExporter/Private/Converters/GLTFMaterialConverters.cpp
--- a/Engine/Plugins/Enterprise/GLTFExporter/MigratedFromGit/Source/GLTFExporter/Private/Converters/GLTFMaterialConverters.cpp
+++ b/Engine/Plugins/Enterprise/GLTFExporter/MigratedFromGit/Source/GLTFExporter/Private/Converters/GLTFMaterialConverters.cpp
@@ -16,102 +16,209 @@ FGLTFJsonMaterialIndex FGLTFMaterialConverter::Add(FGLTFConvertBuilder& Builder,
 	FGLTFJsonMaterial JsonMaterial;
 	JsonMaterial.Name = Name;
 
+	JsonMaterial.AlphaMode = FGLTFConverterUtility::ConvertBlendMode(MaterialInterface->GetBlendMode());
+	JsonMaterial.AlphaCutoff = MaterialInterface->GetOpacityMaskClipValue();
+	JsonMaterial.DoubleSided = MaterialInterface->IsTwoSided();
+
 	const UMaterial* Material = MaterialInterface->GetMaterial();
-	JsonMaterial.PBRMetallicRoughness.BaseColorFactor = FGLTFConverterUtility::ConvertColor(GetConstantColor(Material->BaseColor));
-	JsonMaterial.PBRMetallicRoughness.MetallicFactor = GetConstantScalar(Material->Metallic);
-	JsonMaterial.PBRMetallicRoughness.RoughnessFactor = GetConstantScalar(Material->Roughness);
-	JsonMaterial.EmissiveFactor = FGLTFConverterUtility::ConvertColor(GetConstantColor(Material->EmissiveColor));
+	const UMaterialInstance* MaterialInstance = Cast<UMaterialInstance>(MaterialInterface);
 
-	JsonMaterial.AlphaMode = FGLTFConverterUtility::ConvertBlendMode(Material->BlendMode);
-	if (JsonMaterial.AlphaMode == EGLTFJsonAlphaMode::Blend)
+	if (!TryGetConstantColor(JsonMaterial.PBRMetallicRoughness.BaseColorFactor, Material->BaseColor, MaterialInstance))
 	{
-		JsonMaterial.PBRMetallicRoughness.BaseColorFactor.A = GetConstantScalar(Material->Opacity);
+		// TODO: temporary fallback until texture and baking support
+		JsonMaterial.PBRMetallicRoughness.BaseColorFactor = FGLTFConverterUtility::ConvertColor(FLinearColor(Material->BaseColor.Constant));
+	}
+
+	if (JsonMaterial.AlphaMode == EGLTFJsonAlphaMode::Blend || JsonMaterial.AlphaMode == EGLTFJsonAlphaMode::Mask)
+	{
+		const FScalarMaterialInput& OpacityInput = JsonMaterial.AlphaMode == EGLTFJsonAlphaMode::Blend ? Material->Opacity : Material->OpacityMask;
+		if (!TryGetConstantScalar(JsonMaterial.PBRMetallicRoughness.BaseColorFactor.A, OpacityInput, MaterialInstance))
+		{
+			// TODO: temporary fallback until texture and baking support
+			JsonMaterial.PBRMetallicRoughness.BaseColorFactor.A = OpacityInput.Constant;
+		}
 	}
-	else if (JsonMaterial.AlphaMode == EGLTFJsonAlphaMode::Mask)
+
+	if (!TryGetConstantScalar(JsonMaterial.PBRMetallicRoughness.MetallicFactor, Material->Metallic, MaterialInstance))
 	{
-		JsonMaterial.PBRMetallicRoughness.BaseColorFactor.A = GetConstantScalar(Material->OpacityMask);
+		// TODO: temporary fallback until texture and baking support
+		JsonMaterial.PBRMetallicRoughness.MetallicFactor = Material->Metallic.Constant;
 	}
 
-	JsonMaterial.AlphaCutoff = Material->OpacityMaskClipValue;
-	JsonMaterial.DoubleSided = Material->TwoSided;
+	if (!TryGetConstantScalar(JsonMaterial.PBRMetallicRoughness.RoughnessFactor, Material->Roughness, MaterialInstance))
+	{
+		// TODO: temporary fallback until texture and baking support
+		JsonMaterial.PBRMetallicRoughness.RoughnessFactor = Material->Roughness.Constant;
+	}
+
+	if (!TryGetConstantColor(JsonMaterial.EmissiveFactor, Material->EmissiveColor, MaterialInstance))
+	{
+		// TODO: temporary fallback until texture and baking support
+		JsonMaterial.EmissiveFactor = FGLTFConverterUtility::ConvertColor(FLinearColor(Material->EmissiveColor.Constant));
+	}
 
 	return Builder.AddMaterial(JsonMaterial);
 }
 
-FLinearColor FGLTFMaterialConverter::GetConstantColor(const FColorMaterialInput& MaterialInput) const
+bool FGLTFMaterialConverter::TryGetConstantColor(FGLTFJsonColor3& OutValue, const FColorMaterialInput& MaterialInput, const UMaterialInstance* MaterialInstance) const
+{
+	FLinearColor Value;
+	if (TryGetConstantColor(Value, MaterialInput, MaterialInstance))
+	{
+		OutValue = FGLTFConverterUtility::ConvertColor(Value);
+		return true;
+	}
+
+	return false;
+}
+
+bool FGLTFMaterialConverter::TryGetConstantColor(FGLTFJsonColor4& OutValue, const FColorMaterialInput& MaterialInput, const UMaterialInstance* MaterialInstance) const
+{
+	FLinearColor Value;
+	if (TryGetConstantColor(Value, MaterialInput, MaterialInstance))
+	{
+		OutValue = FGLTFConverterUtility::ConvertColor(Value);
+		return true;
+	}
+
+	return false;
+}
+
+bool FGLTFMaterialConverter::TryGetConstantColor(FLinearColor& OutValue, const FColorMaterialInput& MaterialInput, const UMaterialInstance* MaterialInstance) const
 {
 	const UMaterialExpression* Expression = MaterialInput.Expression;
-	if (Expression != nullptr)
+	if (Expression == nullptr)
 	{
-		if (const UMaterialExpressionVectorParameter* VectorParameter = Cast<UMaterialExpressionVectorParameter>(Expression))
-		{
-			return VectorParameter->DefaultValue;
-		}
+		OutValue = MaterialInput.Constant;
+		return true;
+	}
 
-		if (const UMaterialExpressionScalarParameter* ScalarParameter = Cast<UMaterialExpressionScalarParameter>(Expression))
-		{
-			return { ScalarParameter->DefaultValue, ScalarParameter->DefaultValue, ScalarParameter->DefaultValue };
-		}
+	if (const UMaterialExpressionVectorParameter* VectorParameter = Cast<UMaterialExpressionVectorParameter>(Expression))
+	{
+		FLinearColor Value = VectorParameter->DefaultValue;
 
-		if (const UMaterialExpressionConstant4Vector* Constant4Vector =Cast<UMaterialExpressionConstant4Vector>(Expression))
+		if (MaterialInstance != nullptr)
 		{
-			return Constant4Vector->Constant;
+			const FHashedMaterialParameterInfo ParameterInfo(VectorParameter->GetParameterName());
+			if (!MaterialInstance->GetVectorParameterValue(ParameterInfo, Value))
+			{
+				// TODO: should not happen
+			}
 		}
 
-		if (const UMaterialExpressionConstant3Vector* Constant3Vector = Cast<UMaterialExpressionConstant3Vector>(Expression))
-		{
-			return Constant3Vector->Constant;
-		}
+		OutValue = Value;
+		return true;
+	}
 
-		if (const UMaterialExpressionConstant2Vector* Constant2Vector =Cast<UMaterialExpressionConstant2Vector>(Expression))
-		{
-			return { Constant2Vector->R, Constant2Vector->G, 0 }; // TODO: is this case ever possible?
-		}
+	if (const UMaterialExpressionScalarParameter* ScalarParameter = Cast<UMaterialExpressionScalarParameter>(Expression))
+	{
+		float Value = ScalarParameter->DefaultValue;
 
-		if (const UMaterialExpressionConstant* Constant = Cast<UMaterialExpressionConstant>(Expression))
+		if (MaterialInstance != nullptr)
 		{
-			return { Constant->R, Constant->R, Constant->R };
+			const FHashedMaterialParameterInfo ParameterInfo(ScalarParameter->GetParameterName());
+			if (!MaterialInstance->GetScalarParameterValue(ParameterInfo, Value))
+			{
+				// TODO: should not happen
+			}
 		}
+
+		OutValue = { Value, Value, Value, Value };
+		return true;
 	}
 
-	return MaterialInput.Constant;
+	if (const UMaterialExpressionConstant4Vector* Constant4Vector =Cast<UMaterialExpressionConstant4Vector>(Expression))
+	{
+		OutValue = Constant4Vector->Constant;
+		return true;
+	}
+
+	if (const UMaterialExpressionConstant3Vector* Constant3Vector = Cast<UMaterialExpressionConstant3Vector>(Expression))
+	{
+		OutValue = Constant3Vector->Constant;
+		return true;
+	}
+
+	if (const UMaterialExpressionConstant2Vector* Constant2Vector =Cast<UMaterialExpressionConstant2Vector>(Expression))
+	{
+		OutValue = { Constant2Vector->R, Constant2Vector->G, 0, 0 }; // TODO: is this case ever possible?
+		return true;
+	}
+
+	if (const UMaterialExpressionConstant* Constant = Cast<UMaterialExpressionConstant>(Expression))
+	{
+		OutValue = { Constant->R, Constant->R, Constant->R, Constant->R };
+		return true;
+	}
+
+	return false;
 }
 
-float FGLTFMaterialConverter::GetConstantScalar(const FScalarMaterialInput& MaterialInput) const
+bool FGLTFMaterialConverter::TryGetConstantScalar(float& OutValue, const FScalarMaterialInput& MaterialInput, const UMaterialInstance* MaterialInstance) const
 {
 	const UMaterialExpression* Expression = MaterialInput.Expression;
-	if (Expression != nullptr)
+	if (Expression == nullptr)
 	{
-		if (const UMaterialExpressionVectorParameter* VectorParameter = Cast<UMaterialExpressionVectorParameter>(Expression))
-		{
-			return VectorParameter->DefaultValue.R;
-		}
+		return false;
+	}
 
-		if (const UMaterialExpressionScalarParameter* ScalarParameter = Cast<UMaterialExpressionScalarParameter>(Expression))
-		{
-			return ScalarParameter->DefaultValue;
-		}
+	if (const UMaterialExpressionVectorParameter* VectorParameter = Cast<UMaterialExpressionVectorParameter>(Expression))
+	{
+		FLinearColor Value = VectorParameter->DefaultValue;
 
-		if (const UMaterialExpressionConstant4Vector* Constant4Vector =Cast<UMaterialExpressionConstant4Vector>(Expression))
+		if (MaterialInstance != nullptr)
 		{
-			return Constant4Vector->Constant.R;
+			const FHashedMaterialParameterInfo ParameterInfo(VectorParameter->GetParameterName());
+			if (!MaterialInstance->GetVectorParameterValue(ParameterInfo, Value))
+			{
+				// TODO: should not happen
+			}
 		}
 
-		if (const UMaterialExpressionConstant3Vector* Constant3Vector = Cast<UMaterialExpressionConstant3Vector>(Expression))
-		{
-			return Constant3Vector->Constant.R;
-		}
+		OutValue = Value.R;
+		return true;
+	}
 
-		if (const UMaterialExpressionConstant2Vector* Constant2Vector =Cast<UMaterialExpressionConstant2Vector>(Expression))
-		{
-			return Constant2Vector->R;
-		}
+	if (const UMaterialExpressionScalarParameter* ScalarParameter = Cast<UMaterialExpressionScalarParameter>(Expression))
+	{
+		float Value = ScalarParameter->DefaultValue;
 
-		if (const UMaterialExpressionConstant* Constant = Cast<UMaterialExpressionConstant>(Expression))
+		if (MaterialInstance != nullptr)
 		{
-			return Constant->R;
+			const FHashedMaterialParameterInfo ParameterInfo(ScalarParameter->GetParameterName());
+			if (!MaterialInstance->GetScalarParameterValue(ParameterInfo, Value))
+			{
+				// TODO: should not happen
+			}
 		}
+
+		OutValue = Value;
+		return true;
+	}
+
+	if (const UMaterialExpressionConstant4Vector* Constant4Vector =Cast<UMaterialExpressionConstant4Vector>(Expression))
+	{
+		OutValue = Constant4Vector->Constant.R;
+		return true;
+	}
+
+	if (const UMaterialExpressionConstant3Vector* Constant3Vector = Cast<UMaterialExpressionConstant3Vector>(Expression))
+	{
+		OutValue = Constant3Vector->Constant.R;
+		return true;
+	}
+
+	if (const UMaterialExpressionConstant2Vector* Constant2Vector =Cast<UMaterialExpressionConstant2Vector>(Expression))
+	{
+		OutValue = Constant2Vector->R; // TODO: is this case ever possible?
+		return true;
+	}
+
+	if (const UMaterialExpressionConstant* Constant = Cast<UMaterialExpressionConstant>(Expression))
+	{
+		OutValue = Constant->R;
+		return true;
 	}
 
-	return MaterialInput.Constant;
+	return false;
 }
diff --git a/Engine/Plugins/Enterprise/GLTFExporter/MigratedFromGit/Source/GLTFExporter/Private/Converters/GLTFMaterialConverters.h b/Engine/Plugins/Enterprise/GLTFExporter/MigratedFromGit/Source/GLTFExporter/Private/Converters/GLTFMaterialConverters.h
--- a/Engine/Plugins/Enterprise/GLTFExporter/MigratedFromGit/Source/GLTFExporter/Private/Converters/GLTFMaterialConverters.h
+++ b/Engine/Plugins/Enterprise/GLTFExporter/MigratedFromGit/Source/GLTFExporter/Private/Converters/GLTFMaterialConverters.h
@@ -3,6 +3,8 @@
 #pragma once
 
 #include "Json/GLTFJsonIndex.h"
+#include "Json/GLTFJsonColor3.h"
+#include "Json/GLTFJsonColor4.h"
 #include "Converters/GLTFConverter.h"
 #include "Engine.h"
 
@@ -10,6 +12,8 @@ class FGLTFMaterialConverter final : public TGLTFConverter<FGLTFJsonMaterialInde
 {
 	FGLTFJsonMaterialIndex Add(FGLTFConvertBuilder& Builder, const FString& Name, const UMaterialInterface* Material) override;
 
-	FLinearColor GetConstantColor(const FColorMaterialInput& MaterialInput) const;
-	float GetConstantScalar(const FScalarMaterialInput& MaterialInput) const;
+	bool TryGetConstantColor(FGLTFJsonColor3& OutValue, const FColorMaterialInput& MaterialInput, const UMaterialInstance* MaterialInstance) const;
+	bool TryGetConstantColor(FGLTFJsonColor4& OutValue, const FColorMaterialInput& MaterialInput, const UMaterialInstance* MaterialInstance) const;
+	bool TryGetConstantColor(FLinearColor& OutValue, const FColorMaterialInput& MaterialInput, const UMaterialInstance* MaterialInstance) const;
+	bool TryGetConstantScalar(float& OutValue, const FScalarMaterialInput& MaterialInput, const UMaterialInstance* MaterialInstance) const;
 };
