[GLTFExporter] feat(material): add support for exporting a material asset directly in the editor (without attaching it to another asset)

commit d300a62f4ceb1167875b868c2ba1d8e50bbca100
Date: Tue, 19 May 2020 13:00:22 +0200

    feat(material): add support for exporting a material asset directly in the editor (without attaching it to another asset)
    (cherry picked from commit 2ed5c331efbc2c7f7f7aa17dbc281beefe3e1e69​)

#skipci

diff --git a/Engine/Plugins/Enterprise/GLTFExporter/MigratedFromGit/Source/GLTFExporter/Private/Exporters/GLTFExporterUtility.cpp b/Engine/Plugins/Enterprise/GLTFExporter/MigratedFromGit/Source/GLTFExporter/Private/Exporters/GLTFExporterUtility.cpp
new file mode 100644
--- /dev/null
+++ b/Engine/Plugins/Enterprise/GLTFExporter/MigratedFromGit/Source/GLTFExporter/Private/Exporters/GLTFExporterUtility.cpp
@@ -0,0 +1,30 @@
+// Copyright Epic Games, Inc. All Rights Reserved.
+
+#include "Exporters/GLTFExporterUtility.h"
+#include "UnrealEd.h"
+
+const UStaticMesh* FGLTFExporterUtility::GetPreviewMesh(const UMaterialInterface* Material)
+{
+	const UStaticMesh* PreviewMesh = Cast<UStaticMesh>(Material->PreviewMesh.TryLoad());
+	if (PreviewMesh == nullptr)
+	{
+		// Attempt to use the parent material's preview mesh if the instance's preview mesh is invalid
+		const UMaterialInstance* MaterialInstance = Cast<UMaterialInstance>(Material);
+		if (MaterialInstance != nullptr)
+		{
+			const UMaterialInterface* ParentMaterial = MaterialInstance->Parent;
+			if (ParentMaterial != nullptr)
+			{
+				PreviewMesh = Cast<UStaticMesh>(ParentMaterial->PreviewMesh.TryLoad());
+			}
+		}
+	}
+
+	if (PreviewMesh == nullptr)
+	{
+		// Use a default sphere instead if the parent material's preview mesh is also invalid
+		PreviewMesh = GUnrealEd->GetThumbnailManager()->EditorSphere;
+	}
+
+	return PreviewMesh;
+}
diff --git a/Engine/Plugins/Enterprise/GLTFExporter/MigratedFromGit/Source/GLTFExporter/Private/Exporters/GLTFExporterUtility.h b/Engine/Plugins/Enterprise/GLTFExporter/MigratedFromGit/Source/GLTFExporter/Private/Exporters/GLTFExporterUtility.h
new file mode 100644
--- /dev/null
+++ b/Engine/Plugins/Enterprise/GLTFExporter/MigratedFromGit/Source/GLTFExporter/Private/Exporters/GLTFExporterUtility.h
@@ -0,0 +1,11 @@
+// Copyright Epic Games, Inc. All Rights Reserved.
+
+#pragma once
+
+class UStaticMesh;
+class UMaterialInterface;
+
+struct FGLTFExporterUtility
+{
+	static const UStaticMesh* GetPreviewMesh(const UMaterialInterface* MaterialInterface);
+};
diff --git a/Engine/Plugins/Enterprise/GLTFExporter/MigratedFromGit/Source/GLTFExporter/Private/Exporters/GLTFMaterialExporter.cpp b/Engine/Plugins/Enterprise/GLTFExporter/MigratedFromGit/Source/GLTFExporter/Private/Exporters/GLTFMaterialExporter.cpp
new file mode 100644
--- /dev/null
+++ b/Engine/Plugins/Enterprise/GLTFExporter/MigratedFromGit/Source/GLTFExporter/Private/Exporters/GLTFMaterialExporter.cpp
@@ -0,0 +1,36 @@
+// Copyright Epic Games, Inc. All Rights Reserved.
+
+#include "Exporters/GLTFMaterialExporter.h"
+#include "Exporters/GLTFExporterUtility.h"
+#include "Builders/GLTFContainerBuilder.h"
+#include "Materials/MaterialInterface.h"
+
+UGLTFMaterialExporter::UGLTFMaterialExporter(const FObjectInitializer& ObjectInitializer)
+	: Super(ObjectInitializer)
+{
+	SupportedClass = UMaterialInterface::StaticClass();
+}
+
+bool UGLTFMaterialExporter::Add(FGLTFContainerBuilder& Builder, const UObject* Object)
+{
+	const UMaterialInterface* Material = CastChecked<UMaterialInterface>(Object);
+	const FGLTFJsonMaterialIndex MaterialIndex = Builder.GetOrAddMaterial(Material);
+
+	const UStaticMesh* PreviewMesh = FGLTFExporterUtility::GetPreviewMesh(Material);
+	if (PreviewMesh != nullptr)
+	{
+		const FGLTFJsonMeshIndex MeshIndex = Builder.GetOrAddMesh(PreviewMesh, 0, nullptr, { Material }, "PreviewMesh");
+
+		FGLTFJsonNode Node;
+		Node.Mesh = MeshIndex;
+		const FGLTFJsonNodeIndex NodeIndex = Builder.AddNode(Node);
+
+		FGLTFJsonScene Scene;
+		Scene.Nodes.Add(NodeIndex);
+		const FGLTFJsonSceneIndex SceneIndex = Builder.AddScene(Scene);
+
+		Builder.DefaultScene = SceneIndex;
+	}
+
+	return true;
+}
diff --git a/Engine/Plugins/Enterprise/GLTFExporter/MigratedFromGit/Source/GLTFExporter/Public/Exporters/GLTFMaterialExporter.h b/Engine/Plugins/Enterprise/GLTFExporter/MigratedFromGit/Source/GLTFExporter/Public/Exporters/GLTFMaterialExporter.h
new file mode 100644
--- /dev/null
+++ b/Engine/Plugins/Enterprise/GLTFExporter/MigratedFromGit/Source/GLTFExporter/Public/Exporters/GLTFMaterialExporter.h
@@ -0,0 +1,18 @@
+// Copyright Epic Games, Inc. All Rights Reserved.
+
+#pragma once
+
+#include "Exporters/GLTFExporter.h"
+#include "GLTFMaterialExporter.generated.h"
+
+UCLASS()
+class GLTFEXPORTER_API UGLTFMaterialExporter final : public UGLTFExporter
+{
+public:
+
+	GENERATED_BODY()
+
+	explicit UGLTFMaterialExporter(const FObjectInitializer& ObjectInitializer = FObjectInitializer());
+
+	bool Add(FGLTFContainerBuilder& Builder, const UObject* Object) override;
+};
