[GLTFExporter] refactor: add some TODOs in the code for missing support

commit f7b43403cef5c5532940d8cf39431a33aaa0fd32
Date: Wed, 20 May 2020 12:38:28 +0200

    refactor: add some TODOs in the code for missing support
    (cherry picked from commit f5f5f691e995bdcf0c6bf6ac184d76015830a797​)

#skipci

diff --git a/Engine/Plugins/Enterprise/GLTFExporter/MigratedFromGit/Source/GLTFExporter/Private/Converters/GLTFLevelConverters.cpp b/Engine/Plugins/Enterprise/GLTFExporter/MigratedFromGit/Source/GLTFExporter/Private/Converters/GLTFLevelConverters.cpp
--- a/Engine/Plugins/Enterprise/GLTFExporter/MigratedFromGit/Source/GLTFExporter/Private/Converters/GLTFLevelConverters.cpp
+++ b/Engine/Plugins/Enterprise/GLTFExporter/MigratedFromGit/Source/GLTFExporter/Private/Converters/GLTFLevelConverters.cpp
@@ -13,6 +13,7 @@ FGLTFJsonNodeIndex FGLTFSceneComponentConverter::Add(FGLTFConvertBuilder& Builde
 	FGLTFJsonNode Node;
 	Node.Name = Name.IsEmpty() ? Owner->GetName() + TEXT("_") + SceneComponent->GetName() : Name;
 
+	// TODO: add support for non-uniform scaling (Unreal doesn't treat combined transforms as simple matrix multiplication)
 	const FTransform Transform = bRootNode ? SceneComponent->GetComponentTransform() : SceneComponent->GetRelativeTransform();
 	Node.Translation = FGLTFConverterUtility::ConvertPosition(Transform.GetTranslation());
 	Node.Rotation = FGLTFConverterUtility::ConvertRotation(Transform.GetRotation());
diff --git a/Engine/Plugins/Enterprise/GLTFExporter/MigratedFromGit/Source/GLTFExporter/Private/Converters/GLTFMaterialConverters.cpp b/Engine/Plugins/Enterprise/GLTFExporter/MigratedFromGit/Source/GLTFExporter/Private/Converters/GLTFMaterialConverters.cpp
--- a/Engine/Plugins/Enterprise/GLTFExporter/MigratedFromGit/Source/GLTFExporter/Private/Converters/GLTFMaterialConverters.cpp
+++ b/Engine/Plugins/Enterprise/GLTFExporter/MigratedFromGit/Source/GLTFExporter/Private/Converters/GLTFMaterialConverters.cpp
@@ -16,7 +16,11 @@ FGLTFJsonMaterialIndex FGLTFMaterialConverter::Add(FGLTFConvertBuilder& Builder,
 	FGLTFJsonMaterial JsonMaterial;
 	JsonMaterial.Name = Name;
 
+	// TODO: add support for different shading models (Default Lit, Unlit, Clear Coat)
+
+	// TODO: add support for additional blend modes (at least Additive and Modulate)
 	JsonMaterial.AlphaMode = FGLTFConverterUtility::ConvertBlendMode(MaterialInterface->GetBlendMode());
+
 	JsonMaterial.AlphaCutoff = MaterialInterface->GetOpacityMaskClipValue();
 	JsonMaterial.DoubleSided = MaterialInterface->IsTwoSided();
 
