[GLTFExporter] fix(material): only use built-in input constant if material input has no expression

commit 66e522f908dd6847ebeacdbb9d4fb4a54b81a6c0
Date: Sun, 24 May 2020 20:36:57 +0200

    fix(material): only use built-in input constant if material input has no expression
    (cherry picked from commit 220238e1cad46368a470a4c3206ffcbc41a0086e​)

#skipci

diff --git a/Engine/Plugins/Enterprise/GLTFExporter/MigratedFromGit/Source/GLTFExporter/Private/Converters/GLTFMaterialConverters.cpp b/Engine/Plugins/Enterprise/GLTFExporter/MigratedFromGit/Source/GLTFExporter/Private/Converters/GLTFMaterialConverters.cpp
--- a/Engine/Plugins/Enterprise/GLTFExporter/MigratedFromGit/Source/GLTFExporter/Private/Converters/GLTFMaterialConverters.cpp
+++ b/Engine/Plugins/Enterprise/GLTFExporter/MigratedFromGit/Source/GLTFExporter/Private/Converters/GLTFMaterialConverters.cpp
@@ -30,7 +30,6 @@ FGLTFJsonMaterialIndex FGLTFMaterialConverter::Add(FGLTFConvertBuilder& Builder,
 	if (!TryGetConstantColor(JsonMaterial.PBRMetallicRoughness.BaseColorFactor, Material->BaseColor, MaterialInstance))
 	{
 		// TODO: temporary fallback until texture and baking support
-		JsonMaterial.PBRMetallicRoughness.BaseColorFactor = FGLTFConverterUtility::ConvertColor(FLinearColor(Material->BaseColor.Constant));
 	}
 
 	if (JsonMaterial.AlphaMode == EGLTFJsonAlphaMode::Blend || JsonMaterial.AlphaMode == EGLTFJsonAlphaMode::Mask)
@@ -39,26 +38,22 @@ FGLTFJsonMaterialIndex FGLTFMaterialConverter::Add(FGLTFConvertBuilder& Builder,
 		if (!TryGetConstantScalar(JsonMaterial.PBRMetallicRoughness.BaseColorFactor.A, OpacityInput, MaterialInstance))
 		{
 			// TODO: temporary fallback until texture and baking support
-			JsonMaterial.PBRMetallicRoughness.BaseColorFactor.A = OpacityInput.Constant;
 		}
 	}
 
 	if (!TryGetConstantScalar(JsonMaterial.PBRMetallicRoughness.MetallicFactor, Material->Metallic, MaterialInstance))
 	{
 		// TODO: temporary fallback until texture and baking support
-		JsonMaterial.PBRMetallicRoughness.MetallicFactor = Material->Metallic.Constant;
 	}
 
 	if (!TryGetConstantScalar(JsonMaterial.PBRMetallicRoughness.RoughnessFactor, Material->Roughness, MaterialInstance))
 	{
 		// TODO: temporary fallback until texture and baking support
-		JsonMaterial.PBRMetallicRoughness.RoughnessFactor = Material->Roughness.Constant;
 	}
 
 	if (!TryGetConstantColor(JsonMaterial.EmissiveFactor, Material->EmissiveColor, MaterialInstance))
 	{
 		// TODO: temporary fallback until texture and baking support
-		JsonMaterial.EmissiveFactor = FGLTFConverterUtility::ConvertColor(FLinearColor(Material->EmissiveColor.Constant));
 	}
 
 	return Builder.AddMaterial(JsonMaterial);
@@ -93,7 +88,8 @@ bool FGLTFMaterialConverter::TryGetConstantColor(FLinearColor& OutValue, const F
 	const UMaterialExpression* Expression = MaterialInput.Expression;
 	if (Expression == nullptr)
 	{
-		OutValue = MaterialInput.Constant;
+		// TODO: is this assumption correct?
+		OutValue = FLinearColor(MaterialInput.Constant);
 		return true;
 	}
 
@@ -163,7 +159,9 @@ bool FGLTFMaterialConverter::TryGetConstantScalar(float& OutValue, const FScalar
 	const UMaterialExpression* Expression = MaterialInput.Expression;
 	if (Expression == nullptr)
 	{
-		return false;
+		// TODO: is this assumption correct?
+		OutValue = MaterialInput.Constant;
+		return true;
 	}
 
 	if (const UMaterialExpressionVectorParameter* VectorParameter = Cast<UMaterialExpressionVectorParameter>(Expression))
