[GLTFExporter] fix(material): make sure all expression casts are exact to avoid misinterpreting future derived expressions that aren't supported

commit 2d003d346bc3987742f4e36c0dbd46d6f7c64a81
Date: Sun, 24 May 2020 21:21:36 +0200

    fix(material): make sure all expression casts are exact to avoid misinterpreting future derived expressions that aren't supported
    (cherry picked from commit abf444e8520a770af3f016c15af8f08f5720d413​)

#skipci

diff --git a/Engine/Plugins/Enterprise/GLTFExporter/MigratedFromGit/Source/GLTFExporter/Private/Converters/GLTFMaterialConverters.cpp b/Engine/Plugins/Enterprise/GLTFExporter/MigratedFromGit/Source/GLTFExporter/Private/Converters/GLTFMaterialConverters.cpp
--- a/Engine/Plugins/Enterprise/GLTFExporter/MigratedFromGit/Source/GLTFExporter/Private/Converters/GLTFMaterialConverters.cpp
+++ b/Engine/Plugins/Enterprise/GLTFExporter/MigratedFromGit/Source/GLTFExporter/Private/Converters/GLTFMaterialConverters.cpp
@@ -93,7 +93,7 @@ bool FGLTFMaterialConverter::TryGetConstantColor(FLinearColor& OutValue, const F
 		return true;
 	}
 
-	if (const UMaterialExpressionVectorParameter* VectorParameter = Cast<UMaterialExpressionVectorParameter>(Expression))
+	if (const UMaterialExpressionVectorParameter* VectorParameter = ExactCast<UMaterialExpressionVectorParameter>(Expression))
 	{
 		FLinearColor Value = VectorParameter->DefaultValue;
 
@@ -110,7 +110,7 @@ bool FGLTFMaterialConverter::TryGetConstantColor(FLinearColor& OutValue, const F
 		return true;
 	}
 
-	if (const UMaterialExpressionScalarParameter* ScalarParameter = Cast<UMaterialExpressionScalarParameter>(Expression))
+	if (const UMaterialExpressionScalarParameter* ScalarParameter = ExactCast<UMaterialExpressionScalarParameter>(Expression))
 	{
 		float Value = ScalarParameter->DefaultValue;
 
@@ -127,25 +127,25 @@ bool FGLTFMaterialConverter::TryGetConstantColor(FLinearColor& OutValue, const F
 		return true;
 	}
 
-	if (const UMaterialExpressionConstant4Vector* Constant4Vector =Cast<UMaterialExpressionConstant4Vector>(Expression))
+	if (const UMaterialExpressionConstant4Vector* Constant4Vector = ExactCast<UMaterialExpressionConstant4Vector>(Expression))
 	{
 		OutValue = Constant4Vector->Constant;
 		return true;
 	}
 
-	if (const UMaterialExpressionConstant3Vector* Constant3Vector = Cast<UMaterialExpressionConstant3Vector>(Expression))
+	if (const UMaterialExpressionConstant3Vector* Constant3Vector = ExactCast<UMaterialExpressionConstant3Vector>(Expression))
 	{
 		OutValue = Constant3Vector->Constant;
 		return true;
 	}
 
-	if (const UMaterialExpressionConstant2Vector* Constant2Vector =Cast<UMaterialExpressionConstant2Vector>(Expression))
+	if (const UMaterialExpressionConstant2Vector* Constant2Vector = ExactCast<UMaterialExpressionConstant2Vector>(Expression))
 	{
 		OutValue = { Constant2Vector->R, Constant2Vector->G, 0, 0 }; // TODO: is this case ever possible?
 		return true;
 	}
 
-	if (const UMaterialExpressionConstant* Constant = Cast<UMaterialExpressionConstant>(Expression))
+	if (const UMaterialExpressionConstant* Constant = ExactCast<UMaterialExpressionConstant>(Expression))
 	{
 		OutValue = { Constant->R, Constant->R, Constant->R, Constant->R };
 		return true;
@@ -164,7 +164,7 @@ bool FGLTFMaterialConverter::TryGetConstantScalar(float& OutValue, const FScalar
 		return true;
 	}
 
-	if (const UMaterialExpressionVectorParameter* VectorParameter = Cast<UMaterialExpressionVectorParameter>(Expression))
+	if (const UMaterialExpressionVectorParameter* VectorParameter = ExactCast<UMaterialExpressionVectorParameter>(Expression))
 	{
 		FLinearColor Value = VectorParameter->DefaultValue;
 
@@ -181,7 +181,7 @@ bool FGLTFMaterialConverter::TryGetConstantScalar(float& OutValue, const FScalar
 		return true;
 	}
 
-	if (const UMaterialExpressionScalarParameter* ScalarParameter = Cast<UMaterialExpressionScalarParameter>(Expression))
+	if (const UMaterialExpressionScalarParameter* ScalarParameter = ExactCast<UMaterialExpressionScalarParameter>(Expression))
 	{
 		float Value = ScalarParameter->DefaultValue;
 
@@ -198,25 +198,25 @@ bool FGLTFMaterialConverter::TryGetConstantScalar(float& OutValue, const FScalar
 		return true;
 	}
 
-	if (const UMaterialExpressionConstant4Vector* Constant4Vector =Cast<UMaterialExpressionConstant4Vector>(Expression))
+	if (const UMaterialExpressionConstant4Vector* Constant4Vector = ExactCast<UMaterialExpressionConstant4Vector>(Expression))
 	{
 		OutValue = Constant4Vector->Constant.R;
 		return true;
 	}
 
-	if (const UMaterialExpressionConstant3Vector* Constant3Vector = Cast<UMaterialExpressionConstant3Vector>(Expression))
+	if (const UMaterialExpressionConstant3Vector* Constant3Vector = ExactCast<UMaterialExpressionConstant3Vector>(Expression))
 	{
 		OutValue = Constant3Vector->Constant.R;
 		return true;
 	}
 
-	if (const UMaterialExpressionConstant2Vector* Constant2Vector =Cast<UMaterialExpressionConstant2Vector>(Expression))
+	if (const UMaterialExpressionConstant2Vector* Constant2Vector = ExactCast<UMaterialExpressionConstant2Vector>(Expression))
 	{
 		OutValue = Constant2Vector->R; // TODO: is this case ever possible?
 		return true;
 	}
 
-	if (const UMaterialExpressionConstant* Constant = Cast<UMaterialExpressionConstant>(Expression))
+	if (const UMaterialExpressionConstant* Constant = ExactCast<UMaterialExpressionConstant>(Expression))
 	{
 		OutValue = Constant->R;
 		return true;
