[GLTFExporter] refactor(materail): add TODO comments for future support

commit 08fe32cd44af31cb6bc45ab64d463a2ae96fbf94
Date: Sun, 24 May 2020 21:35:45 +0200

    refactor(materail): add TODO comments for future support
    (cherry picked from commit af3debc3f006a66e1cbcaa2283affd360105d084​)

#skipci

diff --git a/Engine/Plugins/Enterprise/GLTFExporter/MigratedFromGit/Source/GLTFExporter/Private/Converters/GLTFMaterialConverters.cpp b/Engine/Plugins/Enterprise/GLTFExporter/MigratedFromGit/Source/GLTFExporter/Private/Converters/GLTFMaterialConverters.cpp
--- a/Engine/Plugins/Enterprise/GLTFExporter/MigratedFromGit/Source/GLTFExporter/Private/Converters/GLTFMaterialConverters.cpp
+++ b/Engine/Plugins/Enterprise/GLTFExporter/MigratedFromGit/Source/GLTFExporter/Private/Converters/GLTFMaterialConverters.cpp
@@ -102,10 +102,11 @@ bool FGLTFMaterialConverter::TryGetConstantColor(FLinearColor& OutValue, const F
 			const FHashedMaterialParameterInfo ParameterInfo(VectorParameter->GetParameterName());
 			if (!MaterialInstance->GetVectorParameterValue(ParameterInfo, Value))
 			{
-				// TODO: should not happen
+				// TODO: can this happen?
 			}
 		}
 
+		// TODO: add support for output mask
 		OutValue = Value;
 		return true;
 	}
@@ -119,7 +120,7 @@ bool FGLTFMaterialConverter::TryGetConstantColor(FLinearColor& OutValue, const F
 			const FHashedMaterialParameterInfo ParameterInfo(ScalarParameter->GetParameterName());
 			if (!MaterialInstance->GetScalarParameterValue(ParameterInfo, Value))
 			{
-				// TODO: should not happen
+				// TODO: can this happen?
 			}
 		}
 
@@ -129,19 +130,22 @@ bool FGLTFMaterialConverter::TryGetConstantColor(FLinearColor& OutValue, const F
 
 	if (const UMaterialExpressionConstant4Vector* Constant4Vector = ExactCast<UMaterialExpressionConstant4Vector>(Expression))
 	{
+		// TODO: add support for output mask
 		OutValue = Constant4Vector->Constant;
 		return true;
 	}
 
 	if (const UMaterialExpressionConstant3Vector* Constant3Vector = ExactCast<UMaterialExpressionConstant3Vector>(Expression))
 	{
+		// TODO: add support for output mask
 		OutValue = Constant3Vector->Constant;
 		return true;
 	}
 
 	if (const UMaterialExpressionConstant2Vector* Constant2Vector = ExactCast<UMaterialExpressionConstant2Vector>(Expression))
 	{
-		OutValue = { Constant2Vector->R, Constant2Vector->G, 0, 0 }; // TODO: is this case ever possible?
+		// TODO: add support for output mask
+		OutValue = { Constant2Vector->R, Constant2Vector->G, 0, 0 };
 		return true;
 	}
 
@@ -173,10 +177,11 @@ bool FGLTFMaterialConverter::TryGetConstantScalar(float& OutValue, const FScalar
 			const FHashedMaterialParameterInfo ParameterInfo(VectorParameter->GetParameterName());
 			if (!MaterialInstance->GetVectorParameterValue(ParameterInfo, Value))
 			{
-				// TODO: should not happen
+				// TODO: can this happen?
 			}
 		}
 
+		// TODO: add support for output mask
 		OutValue = Value.R;
 		return true;
 	}
@@ -190,7 +195,7 @@ bool FGLTFMaterialConverter::TryGetConstantScalar(float& OutValue, const FScalar
 			const FHashedMaterialParameterInfo ParameterInfo(ScalarParameter->GetParameterName());
 			if (!MaterialInstance->GetScalarParameterValue(ParameterInfo, Value))
 			{
-				// TODO: should not happen
+				// TODO: can this happen?
 			}
 		}
 
@@ -200,19 +205,22 @@ bool FGLTFMaterialConverter::TryGetConstantScalar(float& OutValue, const FScalar
 
 	if (const UMaterialExpressionConstant4Vector* Constant4Vector = ExactCast<UMaterialExpressionConstant4Vector>(Expression))
 	{
+		// TODO: add support for output mask
 		OutValue = Constant4Vector->Constant.R;
 		return true;
 	}
 
 	if (const UMaterialExpressionConstant3Vector* Constant3Vector = ExactCast<UMaterialExpressionConstant3Vector>(Expression))
 	{
+		// TODO: add support for output mask
 		OutValue = Constant3Vector->Constant.R;
 		return true;
 	}
 
 	if (const UMaterialExpressionConstant2Vector* Constant2Vector = ExactCast<UMaterialExpressionConstant2Vector>(Expression))
 	{
-		OutValue = Constant2Vector->R; // TODO: is this case ever possible?
+		// TODO: add support for output mask
+		OutValue = Constant2Vector->R;
 		return true;
 	}
 
