[GLTFExporter] fix(static-mesh): use static mesh's original material if override is null

commit 452f52b1c612fb08927be6158fb855893318f5d0
Date: Mon, 25 May 2020 06:06:19 +0200

    fix(static-mesh): use static mesh's original material if override is null
    (cherry picked from commit e066e5f77305bdfcbf82c6d1aa57d53a00abc8d3​)
    # Conflicts from cherry-pick:
    #	Plugins/GLTFExporter/Source/GLTFExporter/Private/Builders/GLTFConvertBuilder.cpp

#skipci

diff --git a/Engine/Plugins/Enterprise/GLTFExporter/MigratedFromGit/Source/GLTFExporter/Private/Converters/GLTFStaticMeshConverters.cpp b/Engine/Plugins/Enterprise/GLTFExporter/MigratedFromGit/Source/GLTFExporter/Private/Converters/GLTFStaticMeshConverters.cpp
--- a/Engine/Plugins/Enterprise/GLTFExporter/MigratedFromGit/Source/GLTFExporter/Private/Converters/GLTFStaticMeshConverters.cpp
+++ b/Engine/Plugins/Enterprise/GLTFExporter/MigratedFromGit/Source/GLTFExporter/Private/Converters/GLTFStaticMeshConverters.cpp
@@ -93,8 +93,13 @@ FGLTFJsonMeshIndex FGLTFStaticMeshConverter::Add(FGLTFConvertBuilder& Builder, c
 			Name + (SectionCount == 1 ? TEXT("_Indices") : TEXT("_Indices_Section") + FString::FromInt(SectionIndex)));
 
 		const int32 MaterialIndex = Section.MaterialIndex;
-		const UMaterialInterface* Material = MaterialIndex < OverrideMaterials.Num() ? OverrideMaterials[MaterialIndex] : StaticMesh->GetMaterial(MaterialIndex);
-		if (Material != nullptr) JsonPrimitive.Material = Builder.GetOrAddMaterial(Material);
+		const UMaterialInterface* Material = OverrideMaterials.IsValidIndex(MaterialIndex) && OverrideMaterials[MaterialIndex] != nullptr ?
+			OverrideMaterials[MaterialIndex] : StaticMesh->GetMaterial(MaterialIndex);
+
+		if (Material != nullptr)
+		{
+			JsonPrimitive.Material = Builder.GetOrAddMaterial(Material);
+		}
 	}
 
 	return Builder.AddMesh(JsonMesh);
