[GLTFExporter] fix(material): add GLTFMaterialArray to fix issue when forwarding a TArray by copy twice (in FGLTFConverter::GetOrAdd), which results in a empty array

commit 28889299cf5e3773cc8e68143486f8b3bf18e514
Date: Mon, 25 May 2020 07:51:06 +0200

    fix(material): add GLTFMaterialArray to fix issue when forwarding a TArray by copy twice (in FGLTFConverter::GetOrAdd), which results in a empty array
    (cherry picked from commit acd886794fbb4f7d7ab5434993ff825531045df9​)
    # Conflicts from cherry-pick:
    #	Plugins/GLTFExporter/Source/GLTFExporter/Private/Builders/GLTFConvertBuilder.cpp
    #	Plugins/GLTFExporter/Source/GLTFExporter/Private/Converters/GLTFStaticMeshConverters.h

#skipci

diff --git a/Engine/Plugins/Enterprise/GLTFExporter/MigratedFromGit/Source/GLTFExporter/Private/Builders/GLTFBuilderUtility.h b/Engine/Plugins/Enterprise/GLTFExporter/MigratedFromGit/Source/GLTFExporter/Private/Builders/GLTFBuilderUtility.h
--- a/Engine/Plugins/Enterprise/GLTFExporter/MigratedFromGit/Source/GLTFExporter/Private/Builders/GLTFBuilderUtility.h
+++ b/Engine/Plugins/Enterprise/GLTFExporter/MigratedFromGit/Source/GLTFExporter/Private/Builders/GLTFBuilderUtility.h
@@ -18,14 +18,4 @@ struct FGLTFBuilderUtility
 
 		return Name;
 	}
-
-	template <typename ElementType, typename AllocatorType>
-	static TArray<const ElementType*, AllocatorType> MakeArrayOfPointersConst(const TArray<ElementType*, AllocatorType>& ArrayOfPointers)
-	{
-		TArray<const ElementType*, AllocatorType> ArrayOfConstPointers;
-		ArrayOfConstPointers.AddUninitialized(ArrayOfPointers.Num());
-
-		FMemory::Memcpy(ArrayOfConstPointers.GetData(), ArrayOfPointers.GetData(), ArrayOfPointers.Num() * sizeof(ElementType*));
-		return ArrayOfConstPointers;
-	}
 };
diff --git a/Engine/Plugins/Enterprise/GLTFExporter/MigratedFromGit/Source/GLTFExporter/Private/Builders/GLTFConvertBuilder.cpp b/Engine/Plugins/Enterprise/GLTFExporter/MigratedFromGit/Source/GLTFExporter/Private/Builders/GLTFConvertBuilder.cpp
--- a/Engine/Plugins/Enterprise/GLTFExporter/MigratedFromGit/Source/GLTFExporter/Private/Builders/GLTFConvertBuilder.cpp
+++ b/Engine/Plugins/Enterprise/GLTFExporter/MigratedFromGit/Source/GLTFExporter/Private/Builders/GLTFConvertBuilder.cpp
@@ -78,7 +78,7 @@ FGLTFJsonMaterialIndex FGLTFConvertBuilder::GetOrAddMaterial(const UMaterialInte
 	return MaterialConverter.GetOrAdd(*this, DesiredName.IsEmpty() ? Material->GetName() : DesiredName, Material);
 }
 
-FGLTFJsonMeshIndex FGLTFConvertBuilder::GetOrAddMesh(const UStaticMesh* StaticMesh, int32 LODIndex, const FColorVertexBuffer* OverrideVertexColors, TArray<const UMaterialInterface*> OverrideMaterials, const FString& DesiredName)
+FGLTFJsonMeshIndex FGLTFConvertBuilder::GetOrAddMesh(const UStaticMesh* StaticMesh, int32 LODIndex, const FColorVertexBuffer* OverrideVertexColors, const FGLTFMaterialArray& OverrideMaterials, const FString& DesiredName)
 {
 	if (StaticMesh == nullptr)
 	{
@@ -98,7 +98,7 @@ FGLTFJsonMeshIndex FGLTFConvertBuilder::GetOrAddMesh(const UStaticMeshComponent*
 	const UStaticMesh* StaticMesh = StaticMeshComponent->GetStaticMesh();
 	const int32 LODIndex = StaticMeshComponent->ForcedLodModel > 0 ? StaticMeshComponent->ForcedLodModel - 1 : /* auto-select */ 0;
 	const FColorVertexBuffer* OverrideVertexColors = StaticMeshComponent->LODData.IsValidIndex(LODIndex) ? StaticMeshComponent->LODData[LODIndex].OverrideVertexColors : nullptr;
-	const TArray<const UMaterialInterface*> OverrideMaterials = FGLTFBuilderUtility::MakeArrayOfPointersConst(StaticMeshComponent->OverrideMaterials);
+	const FGLTFMaterialArray OverrideMaterials = FGLTFMaterialArray(StaticMeshComponent->OverrideMaterials);
 
 	return GetOrAddMesh(StaticMesh, LODIndex, OverrideVertexColors, OverrideMaterials, DesiredName);
 }
diff --git a/Engine/Plugins/Enterprise/GLTFExporter/MigratedFromGit/Source/GLTFExporter/Private/Builders/GLTFConvertBuilder.h b/Engine/Plugins/Enterprise/GLTFExporter/MigratedFromGit/Source/GLTFExporter/Private/Builders/GLTFConvertBuilder.h
--- a/Engine/Plugins/Enterprise/GLTFExporter/MigratedFromGit/Source/GLTFExporter/Private/Builders/GLTFConvertBuilder.h
+++ b/Engine/Plugins/Enterprise/GLTFExporter/MigratedFromGit/Source/GLTFExporter/Private/Builders/GLTFConvertBuilder.h
@@ -23,7 +23,7 @@ public:
 
 	FGLTFJsonMaterialIndex GetOrAddMaterial(const UMaterialInterface* Material, const FString& DesiredName = TEXT(""));
 
-	FGLTFJsonMeshIndex GetOrAddMesh(const UStaticMesh* StaticMesh, int32 LODIndex = 0, const FColorVertexBuffer* OverrideVertexColors = nullptr, TArray<const UMaterialInterface*> OverrideMaterials = {}, const FString& DesiredName = TEXT(""));
+	FGLTFJsonMeshIndex GetOrAddMesh(const UStaticMesh* StaticMesh, int32 LODIndex = 0, const FColorVertexBuffer* OverrideVertexColors = nullptr, const FGLTFMaterialArray& OverrideMaterials = {}, const FString& DesiredName = TEXT(""));
 	FGLTFJsonMeshIndex GetOrAddMesh(const UStaticMeshComponent* StaticMeshComponent, const FString& DesiredName = TEXT(""));
 
 	FGLTFJsonNodeIndex GetOrAddNode(const USceneComponent* SceneComponent, bool bSelectedOnly, bool bRootNode = false, const FString& DesiredName = TEXT(""));
diff --git a/Engine/Plugins/Enterprise/GLTFExporter/MigratedFromGit/Source/GLTFExporter/Private/Converters/GLTFConverter.h b/Engine/Plugins/Enterprise/GLTFExporter/MigratedFromGit/Source/GLTFExporter/Private/Converters/GLTFConverter.h
--- a/Engine/Plugins/Enterprise/GLTFExporter/MigratedFromGit/Source/GLTFExporter/Private/Converters/GLTFConverter.h
+++ b/Engine/Plugins/Enterprise/GLTFExporter/MigratedFromGit/Source/GLTFExporter/Private/Converters/GLTFConverter.h
@@ -15,13 +15,13 @@ public:
 
 	virtual ~TGLTFConverter() = default;
 
-	FORCEINLINE IndexType Get(ArgTypes&&... Args) const
+	IndexType Get(ArgTypes&&... Args) const
 	{
 		const KeyType Key(Forward<ArgTypes>(Args)...);
 		return IndexLookup.FindRef(Key);
 	}
 
-	FORCEINLINE IndexType GetOrAdd(FGLTFConvertBuilder& Builder, const FString& DesiredName, ArgTypes... Args)
+	IndexType GetOrAdd(FGLTFConvertBuilder& Builder, const FString& DesiredName, ArgTypes... Args)
 	{
 		const KeyType Key(Forward<ArgTypes>(Args)...);
 
diff --git a/Engine/Plugins/Enterprise/GLTFExporter/MigratedFromGit/Source/GLTFExporter/Private/Converters/GLTFMaterialArray.h b/Engine/Plugins/Enterprise/GLTFExporter/MigratedFromGit/Source/GLTFExporter/Private/Converters/GLTFMaterialArray.h
new file mode 100644
--- /dev/null
+++ b/Engine/Plugins/Enterprise/GLTFExporter/MigratedFromGit/Source/GLTFExporter/Private/Converters/GLTFMaterialArray.h
@@ -0,0 +1,77 @@
+// Copyright Epic Games, Inc. All Rights Reserved.
+
+#pragma once
+
+#include "CoreMinimal.h"
+
+class UMaterialInterface;
+
+// The reason behind a custom material array type are two:
+// 1. Forwarding a TArray by copy in FGLTFConverter::GetOrAdd twice will result in a empty array
+// 2. There is no GetTypeHash for generic arrays
+struct FGLTFMaterialArray : TArray<const UMaterialInterface*>
+{
+	FGLTFMaterialArray()
+	{
+	}
+
+	FGLTFMaterialArray(const TArray& Other)
+		: TArray(Other)
+	{
+	}
+
+	FGLTFMaterialArray(const FGLTFMaterialArray& Other)
+        : TArray(Other)
+	{
+	}
+
+	template <typename ElementType, typename AllocatorType>
+    explicit FGLTFMaterialArray(const TArray<ElementType, AllocatorType>& Other)
+        : TArray(Other)
+	{
+	}
+
+	FGLTFMaterialArray(std::initializer_list<const UMaterialInterface*> InitList)
+		: TArray(InitList)
+	{
+	}
+
+	FGLTFMaterialArray& operator=(std::initializer_list<const UMaterialInterface*> InitList)
+	{
+		static_cast<TArray&>(*this) = InitList;
+		return *this;
+	}
+
+	FGLTFMaterialArray& operator=(const TArray& Other)
+	{
+		static_cast<TArray&>(*this) = Other;
+		return *this;
+	}
+
+	template <typename ElementType, typename AllocatorType>
+    FGLTFMaterialArray& operator=(const TArray<ElementType, AllocatorType>& Other)
+	{
+		static_cast<TArray&>(*this) = Other;
+		return *this;
+	}
+
+	bool operator==(const TArray& Other) const
+	{
+		return static_cast<const TArray&>(*this) == Other;
+	}
+
+	bool operator!=(const TArray& Other) const
+	{
+		return static_cast<const TArray&>(*this) != Other;
+	}
+
+	friend uint32 GetTypeHash(const TArray& Array)
+	{
+		uint32 Hash = GetTypeHash(Array.Num());
+		for (const UMaterialInterface* Material : Array)
+		{
+			Hash = HashCombine(Hash, GetTypeHash(Material));
+		}
+		return Hash;
+	}
+};
diff --git a/Engine/Plugins/Enterprise/GLTFExporter/MigratedFromGit/Source/GLTFExporter/Private/Converters/GLTFStaticMeshConverters.cpp b/Engine/Plugins/Enterprise/GLTFExporter/MigratedFromGit/Source/GLTFExporter/Private/Converters/GLTFStaticMeshConverters.cpp
--- a/Engine/Plugins/Enterprise/GLTFExporter/MigratedFromGit/Source/GLTFExporter/Private/Converters/GLTFStaticMeshConverters.cpp
+++ b/Engine/Plugins/Enterprise/GLTFExporter/MigratedFromGit/Source/GLTFExporter/Private/Converters/GLTFStaticMeshConverters.cpp
@@ -47,7 +47,7 @@ FGLTFJsonAccessorIndex FGLTFStaticMeshSectionConverter::Add(FGLTFConvertBuilder&
 	return Builder.AddAccessor(JsonAccessor);
 }
 
-FGLTFJsonMeshIndex FGLTFStaticMeshConverter::Add(FGLTFConvertBuilder& Builder, const FString& Name, const UStaticMesh* StaticMesh, int32 LODIndex, const FColorVertexBuffer* OverrideVertexColors, TArray<const UMaterialInterface*> OverrideMaterials)
+FGLTFJsonMeshIndex FGLTFStaticMeshConverter::Add(FGLTFConvertBuilder& Builder, const FString& Name, const UStaticMesh* StaticMesh, int32 LODIndex, const FColorVertexBuffer* OverrideVertexColors, FGLTFMaterialArray OverrideMaterials)
 {
 	if (LODIndex < 0 || StaticMesh->GetNumLODs() <= LODIndex)
 	{
diff --git a/Engine/Plugins/Enterprise/GLTFExporter/MigratedFromGit/Source/GLTFExporter/Private/Converters/GLTFStaticMeshConverters.h b/Engine/Plugins/Enterprise/GLTFExporter/MigratedFromGit/Source/GLTFExporter/Private/Converters/GLTFStaticMeshConverters.h
--- a/Engine/Plugins/Enterprise/GLTFExporter/MigratedFromGit/Source/GLTFExporter/Private/Converters/GLTFStaticMeshConverters.h
+++ b/Engine/Plugins/Enterprise/GLTFExporter/MigratedFromGit/Source/GLTFExporter/Private/Converters/GLTFStaticMeshConverters.h
@@ -4,6 +4,7 @@
 
 #include "Json/GLTFJsonIndex.h"
 #include "Converters/GLTFConverter.h"
+#include "Converters/GLTFMaterialArray.h"
 #include "Engine.h"
 
 class FGLTFIndexBufferConverter final : public TGLTFConverter<FGLTFJsonBufferViewIndex, const FRawStaticIndexBuffer*>
@@ -18,16 +19,5 @@ class FGLTFStaticMeshSectionConverter final : public TGLTFConverter<FGLTFJsonAcc
 
 class FGLTFStaticMeshConverter final : public TGLTFConverter<FGLTFJsonMeshIndex, const UStaticMesh*, int32, const FColorVertexBuffer*>
 {
-	FGLTFJsonMeshIndex Add(FGLTFConvertBuilder& Builder, const FString& Name, const UStaticMesh* StaticMesh, int32 LODIndex, const FColorVertexBuffer* OverrideVertexColors, TArray<const UMaterialInterface*> OverrideMaterials) override;
-
-	// TODO: move this to somewhere cleaner
-	friend inline uint32 GetTypeHash(TArray<const UMaterialInterface*> Materials)
-	{
-		uint32 Hash = GetTypeHash(Materials.Num());
-		for (const UMaterialInterface* Material : Materials)
-		{
-			Hash = HashCombine(Hash, GetTypeHash(Material));
-		}
-		return Hash;
-	}
+	FGLTFJsonMeshIndex Add(FGLTFConvertBuilder& Builder, const FString& Name, const UStaticMesh* StaticMesh, int32 LODIndex, const FColorVertexBuffer* OverrideVertexColors, FGLTFMaterialArray OverrideMaterials) override;
 };
