[GLTFExporter] refactor(material): add comment explaining why GLTFExporterUtility::GetPreviewMesh doesn't check all parents

commit aa6097150c19e541075d3c5de5dcb0216e7282a0
Date: Wed, 27 May 2020 23:43:15 +0200

    refactor(material): add comment explaining why GLTFExporterUtility::GetPreviewMesh doesn't check all parents
    (cherry picked from commit 6dce08ec516fd6d8b1e3819a7a96cee733406f9b​)

#skipci

diff --git a/Engine/Plugins/Enterprise/GLTFExporter/MigratedFromGit/Source/GLTFExporter/Private/Exporters/GLTFExporterUtility.cpp b/Engine/Plugins/Enterprise/GLTFExporter/MigratedFromGit/Source/GLTFExporter/Private/Exporters/GLTFExporterUtility.cpp
--- a/Engine/Plugins/Enterprise/GLTFExporter/MigratedFromGit/Source/GLTFExporter/Private/Exporters/GLTFExporterUtility.cpp
+++ b/Engine/Plugins/Enterprise/GLTFExporter/MigratedFromGit/Source/GLTFExporter/Private/Exporters/GLTFExporterUtility.cpp
@@ -5,6 +5,8 @@
 
 const UStaticMesh* FGLTFExporterUtility::GetPreviewMesh(const UMaterialInterface* Material)
 {
+	// The following implementation is based of FMaterialInstanceEditor::RefreshPreviewAsset
+
 	const UStaticMesh* PreviewMesh = Cast<UStaticMesh>(Material->PreviewMesh.TryLoad());
 	if (PreviewMesh == nullptr)
 	{
