[GLTFExporter] Merge branch 'feature/constant-material-colors' into 'master'

commit 9bd00914ffccecefa883895028d4a2e1e1e2682b
Date: Thu, 28 May 2020 08:46:45 +0200

    Merge branch 'feature/constant-material-colors' into 'master'
    Add initial support for constant material colors
    (cherry picked from commit 370e13b01fe1c4336ea087afef306fd98086d12b​)
    # Conflicts from cherry-pick:
    #	Plugins/GLTFExporter/Source/GLTFExporter/Private/Builders/GLTFConvertBuilder.cpp
    #	Plugins/GLTFExporter/Source/GLTFExporter/Private/Builders/GLTFConvertBuilder.h
    #	Plugins/GLTFExporter/Source/GLTFExporter/Private/Json/GLTFJsonColor3.h
    #	Plugins/GLTFExporter/Source/GLTFExporter/Private/Json/GLTFJsonColor4.h
    #	Plugins/GLTFExporter/Source/GLTFExporter/Private/Json/GLTFJsonMaterial.h

#skipci

diff --git a/Engine/Plugins/Enterprise/GLTFExporter/MigratedFromGit/Source/GLTFExporter/Private/Builders/GLTFConvertBuilder.cpp b/Engine/Plugins/Enterprise/GLTFExporter/MigratedFromGit/Source/GLTFExporter/Private/Builders/GLTFConvertBuilder.cpp
--- a/Engine/Plugins/Enterprise/GLTFExporter/MigratedFromGit/Source/GLTFExporter/Private/Builders/GLTFConvertBuilder.cpp
+++ b/Engine/Plugins/Enterprise/GLTFExporter/MigratedFromGit/Source/GLTFExporter/Private/Builders/GLTFConvertBuilder.cpp
@@ -73,11 +73,6 @@ FGLTFJsonAccessorIndex FGLTFConvertBuilder::GetOrAddIndexAccessor(const FStaticM
 	return StaticMeshSectionConverter.GetOrAdd(*this, DesiredName, MeshSection, IndexBuffer);
 }
 
-FGLTFJsonMaterialIndex FGLTFConvertBuilder::GetOrAddMaterial(const UMaterialInterface* Material, const FString& DesiredName)
-{
-	return MaterialConverter.GetOrAdd(*this, DesiredName.IsEmpty() ? Material->GetName() : DesiredName, Material);
-}
-
 FGLTFJsonMeshIndex FGLTFConvertBuilder::GetOrAddMesh(const UStaticMesh* StaticMesh, int32 LODIndex, const FColorVertexBuffer* OverrideVertexColors, const FGLTFMaterialArray& OverrideMaterials, const FString& DesiredName)
 {
 	if (StaticMesh == nullptr)
@@ -103,6 +98,16 @@ FGLTFJsonMeshIndex FGLTFConvertBuilder::GetOrAddMesh(const UStaticMeshComponent*
 	return GetOrAddMesh(StaticMesh, LODIndex, OverrideVertexColors, OverrideMaterials, DesiredName);
 }
 
+FGLTFJsonMaterialIndex FGLTFConvertBuilder::GetOrAddMaterial(const UMaterialInterface* Material, const FString& DesiredName)
+{
+	if (Material == nullptr)
+	{
+		return FGLTFJsonMaterialIndex(INDEX_NONE);
+	}
+
+	return MaterialConverter.GetOrAdd(*this, DesiredName.IsEmpty() ? Material->GetName() : DesiredName, Material);
+}
+
 FGLTFJsonNodeIndex FGLTFConvertBuilder::GetOrAddNode(const USceneComponent* SceneComponent, bool bSelectedOnly, bool bRootNode, const FString& DesiredName)
 {
 	if (SceneComponent == nullptr)
diff --git a/Engine/Plugins/Enterprise/GLTFExporter/MigratedFromGit/Source/GLTFExporter/Private/Converters/GLTFStaticMeshConverters.h b/Engine/Plugins/Enterprise/GLTFExporter/MigratedFromGit/Source/GLTFExporter/Private/Converters/GLTFStaticMeshConverters.h
--- a/Engine/Plugins/Enterprise/GLTFExporter/MigratedFromGit/Source/GLTFExporter/Private/Converters/GLTFStaticMeshConverters.h
+++ b/Engine/Plugins/Enterprise/GLTFExporter/MigratedFromGit/Source/GLTFExporter/Private/Converters/GLTFStaticMeshConverters.h
@@ -17,7 +17,7 @@ class FGLTFStaticMeshSectionConverter final : public TGLTFConverter<FGLTFJsonAcc
 	FGLTFJsonAccessorIndex Add(FGLTFConvertBuilder& Builder, const FString& Name, const FStaticMeshSection* MeshSection, const FRawStaticIndexBuffer* IndexBuffer) override;
 };
 
-class FGLTFStaticMeshConverter final : public TGLTFConverter<FGLTFJsonMeshIndex, const UStaticMesh*, int32, const FColorVertexBuffer*>
+class FGLTFStaticMeshConverter final : public TGLTFConverter<FGLTFJsonMeshIndex, const UStaticMesh*, int32, const FColorVertexBuffer*, FGLTFMaterialArray>
 {
 	FGLTFJsonMeshIndex Add(FGLTFConvertBuilder& Builder, const FString& Name, const UStaticMesh* StaticMesh, int32 LODIndex, const FColorVertexBuffer* OverrideVertexColors, FGLTFMaterialArray OverrideMaterials) override;
 };
diff --git a/Engine/Plugins/Enterprise/GLTFExporter/MigratedFromGit/Source/GLTFExporter/Private/Json/GLTFJsonColor3.h b/Engine/Plugins/Enterprise/GLTFExporter/MigratedFromGit/Source/GLTFExporter/Private/Json/GLTFJsonColor3.h
--- a/Engine/Plugins/Enterprise/GLTFExporter/MigratedFromGit/Source/GLTFExporter/Private/Json/GLTFJsonColor3.h
+++ b/Engine/Plugins/Enterprise/GLTFExporter/MigratedFromGit/Source/GLTFExporter/Private/Json/GLTFJsonColor3.h
@@ -2,6 +2,7 @@
 
 #pragma once
 
+#include "Json/GLTFJsonUtility.h"
 #include "Serialization/JsonSerializer.h"
 
 struct FGLTFJsonColor3
@@ -20,9 +21,9 @@ struct FGLTFJsonColor3
 	void WriteArray(TJsonWriter<CharType, PrintPolicy>& JsonWriter) const
 	{
 		JsonWriter.WriteArrayStart();
-		JsonWriter.WriteValue(R);
-		JsonWriter.WriteValue(G);
-		JsonWriter.WriteValue(B);
+		FGLTFJsonUtility::WriteExactValue(JsonWriter, R);
+		FGLTFJsonUtility::WriteExactValue(JsonWriter, G);
+		FGLTFJsonUtility::WriteExactValue(JsonWriter, B);
 		JsonWriter.WriteArrayEnd();
 	}
 
diff --git a/Engine/Plugins/Enterprise/GLTFExporter/MigratedFromGit/Source/GLTFExporter/Private/Json/GLTFJsonColor4.h b/Engine/Plugins/Enterprise/GLTFExporter/MigratedFromGit/Source/GLTFExporter/Private/Json/GLTFJsonColor4.h
--- a/Engine/Plugins/Enterprise/GLTFExporter/MigratedFromGit/Source/GLTFExporter/Private/Json/GLTFJsonColor4.h
+++ b/Engine/Plugins/Enterprise/GLTFExporter/MigratedFromGit/Source/GLTFExporter/Private/Json/GLTFJsonColor4.h
@@ -3,6 +3,7 @@
 #pragma once
 
 #include "Json/GLTFJsonColor3.h"
+#include "Json/GLTFJsonUtility.h"
 #include "Serialization/JsonSerializer.h"
 
 struct FGLTFJsonColor4
@@ -26,10 +27,10 @@ struct FGLTFJsonColor4
 	void WriteArray(TJsonWriter<CharType, PrintPolicy>& JsonWriter) const
 	{
 		JsonWriter.WriteArrayStart();
-		JsonWriter.WriteValue(R);
-		JsonWriter.WriteValue(G);
-		JsonWriter.WriteValue(B);
-		JsonWriter.WriteValue(A);
+		FGLTFJsonUtility::WriteExactValue(JsonWriter, R);
+		FGLTFJsonUtility::WriteExactValue(JsonWriter, G);
+		FGLTFJsonUtility::WriteExactValue(JsonWriter, B);
+		FGLTFJsonUtility::WriteExactValue(JsonWriter, A);
 		JsonWriter.WriteArrayEnd();
 	}
 
diff --git a/Engine/Plugins/Enterprise/GLTFExporter/MigratedFromGit/Source/GLTFExporter/Private/Json/GLTFJsonMaterial.h b/Engine/Plugins/Enterprise/GLTFExporter/MigratedFromGit/Source/GLTFExporter/Private/Json/GLTFJsonMaterial.h
--- a/Engine/Plugins/Enterprise/GLTFExporter/MigratedFromGit/Source/GLTFExporter/Private/Json/GLTFJsonMaterial.h
+++ b/Engine/Plugins/Enterprise/GLTFExporter/MigratedFromGit/Source/GLTFExporter/Private/Json/GLTFJsonMaterial.h
@@ -35,12 +35,12 @@ struct FGLTFJsonPBRMetallicRoughness
 
 		if (MetallicFactor != 1)
 		{
-			JsonWriter.WriteValue(TEXT("metallicFactor"), MetallicFactor);
+			FGLTFJsonUtility::WriteExactValue(JsonWriter, TEXT("metallicFactor"), MetallicFactor);
 		}
 
 		if (RoughnessFactor != 1)
 		{
-			JsonWriter.WriteValue(TEXT("roughnessFactor"), RoughnessFactor);
+			FGLTFJsonUtility::WriteExactValue(JsonWriter, TEXT("roughnessFactor"), RoughnessFactor);
 		}
 
 		JsonWriter.WriteObjectEnd();
@@ -89,12 +89,12 @@ struct FGLTFJsonMaterial
 
 		if (AlphaMode != EGLTFJsonAlphaMode::Opaque)
 		{
-			JsonWriter.WriteValue(TEXT("alphaMode"), FGLTFJsonUtility::AlphaModeToString(AlphaMode));
+			JsonWriter.WriteValue(TEXT("alphaMode"), FGLTFJsonUtility::ToString(AlphaMode));
 		}
 
 		if (AlphaMode == EGLTFJsonAlphaMode::Mask && AlphaCutoff != 0.5f)
 		{
-			JsonWriter.WriteValue(TEXT("alphaCutoff"), AlphaCutoff);
+			FGLTFJsonUtility::WriteExactValue(JsonWriter, TEXT("alphaCutoff"), AlphaCutoff);
 		}
 
 		if (DoubleSided)
diff --git a/Engine/Plugins/Enterprise/GLTFExporter/MigratedFromGit/Source/GLTFExporter/Private/Json/GLTFJsonUtility.h b/Engine/Plugins/Enterprise/GLTFExporter/MigratedFromGit/Source/GLTFExporter/Private/Json/GLTFJsonUtility.h
--- a/Engine/Plugins/Enterprise/GLTFExporter/MigratedFromGit/Source/GLTFExporter/Private/Json/GLTFJsonUtility.h
+++ b/Engine/Plugins/Enterprise/GLTFExporter/MigratedFromGit/Source/GLTFExporter/Private/Json/GLTFJsonUtility.h
@@ -56,4 +56,11 @@ struct FGLTFJsonUtility
 		FString ExactStringRepresentation = FString::Printf(TEXT("%.9g"), Value);
 		JsonWriter.WriteRawJSONValue(ExactStringRepresentation);
 	}
+
+	template <class CharType = TCHAR, class PrintPolicy = TPrettyJsonPrintPolicy<CharType>>
+    static void WriteExactValue(TJsonWriter<CharType, PrintPolicy>& JsonWriter, const FString& Identifier, float Value)
+	{
+		FString ExactStringRepresentation = FString::Printf(TEXT("%.9g"), Value);
+		JsonWriter.WriteRawJSONValue(Identifier, ExactStringRepresentation);
+	}
 };
