[GLTFExporter] refactor: simplify serialization of top-level json object arrays

commit 833e2a347d7fe40bd1d7fa6ef6903b300d9494cb
Date: Fri, 29 May 2020 22:55:12 +0200

    refactor: simplify serialization of top-level json object arrays
    (cherry picked from commit ab8fdac1d5d183a1385f625887dc13a587f29d0a​)

#skipci

diff --git a/Engine/Plugins/Enterprise/GLTFExporter/MigratedFromGit/Source/GLTFExporter/Private/Json/GLTFJsonRoot.h b/Engine/Plugins/Enterprise/GLTFExporter/MigratedFromGit/Source/GLTFExporter/Private/Json/GLTFJsonRoot.h
--- a/Engine/Plugins/Enterprise/GLTFExporter/MigratedFromGit/Source/GLTFExporter/Private/Json/GLTFJsonRoot.h
+++ b/Engine/Plugins/Enterprise/GLTFExporter/MigratedFromGit/Source/GLTFExporter/Private/Json/GLTFJsonRoot.h
@@ -4,12 +4,12 @@
 
 #include "Json/GLTFJsonIndex.h"
 #include "Json/GLTFJsonAccessor.h"
-#include "Json/GLTFJsonBufferView.h"
 #include "Json/GLTFJsonBuffer.h"
+#include "Json/GLTFJsonBufferView.h"
 #include "Json/GLTFJsonMaterial.h"
 #include "Json/GLTFJsonMesh.h"
-#include "Json/GLTFJsonScene.h"
 #include "Json/GLTFJsonNode.h"
+#include "Json/GLTFJsonScene.h"
 
 #include "Containers/Set.h"
 #include "Policies/CondensedJsonPrintPolicy.h"
@@ -101,75 +101,13 @@ struct FGLTFJsonRoot
 			JsonWriter.WriteValue(TEXT("scene"), DefaultScene);
 		}
 
-		if (Accessors.Num() > 0)
-		{
-			JsonWriter.WriteArrayStart(TEXT("accessors"));
-			for (const FGLTFJsonAccessor& Accessor : Accessors)
-			{
-				Accessor.WriteObject(JsonWriter);
-			}
-			JsonWriter.WriteArrayEnd();
-		}
-
-		if (BufferViews.Num() > 0)
-		{
-			JsonWriter.WriteArrayStart(TEXT("bufferViews"));
-			for (const FGLTFJsonBufferView& BufferView : BufferViews)
-			{
-				BufferView.WriteObject(JsonWriter);
-			}
-			JsonWriter.WriteArrayEnd();
-		}
-
-		if (Materials.Num() > 0)
-		{
-			JsonWriter.WriteArrayStart(TEXT("materials"));
-			for (const FGLTFJsonMaterial& Material : Materials)
-			{
-				Material.WriteObject(JsonWriter);
-			}
-			JsonWriter.WriteArrayEnd();
-		}
-
-		if (Meshes.Num() > 0)
-		{
-			JsonWriter.WriteArrayStart(TEXT("meshes"));
-			for (const FGLTFJsonMesh& Mesh : Meshes)
-			{
-				Mesh.WriteObject(JsonWriter);
-			}
-			JsonWriter.WriteArrayEnd();
-		}
-
-		if (Nodes.Num() > 0)
-		{
-			JsonWriter.WriteArrayStart(TEXT("nodes"));
-			for (const FGLTFJsonNode& Node : Nodes)
-			{
-				Node.WriteObject(JsonWriter);
-			}
-			JsonWriter.WriteArrayEnd();
-		}
-
-		if (Scenes.Num() > 0)
-		{
-			JsonWriter.WriteArrayStart(TEXT("scenes"));
-			for (const FGLTFJsonScene& Scene : Scenes)
-			{
-				Scene.WriteObject(JsonWriter);
-			}
-			JsonWriter.WriteArrayEnd();
-		}
-
-		if (Buffers.Num() > 0)
-		{
-			JsonWriter.WriteArrayStart(TEXT("buffers"));
-			for (const FGLTFJsonBuffer& Buffer : Buffers)
-			{
-				Buffer.WriteObject(JsonWriter);
-			}
-			JsonWriter.WriteArrayEnd();
-		}
+		FGLTFJsonUtility::WriteObjectArray(JsonWriter, TEXT("accessors"), Accessors);
+		FGLTFJsonUtility::WriteObjectArray(JsonWriter, TEXT("buffers"), Buffers);
+		FGLTFJsonUtility::WriteObjectArray(JsonWriter, TEXT("bufferViews"), BufferViews);
+		FGLTFJsonUtility::WriteObjectArray(JsonWriter, TEXT("materials"), Materials);
+		FGLTFJsonUtility::WriteObjectArray(JsonWriter, TEXT("meshes"), Meshes);
+		FGLTFJsonUtility::WriteObjectArray(JsonWriter, TEXT("nodes"), Nodes);
+		FGLTFJsonUtility::WriteObjectArray(JsonWriter, TEXT("scenes"), Scenes);
 
 		JsonWriter.WriteObjectEnd();
 	}
diff --git a/Engine/Plugins/Enterprise/GLTFExporter/MigratedFromGit/Source/GLTFExporter/Private/Json/GLTFJsonUtility.h b/Engine/Plugins/Enterprise/GLTFExporter/MigratedFromGit/Source/GLTFExporter/Private/Json/GLTFJsonUtility.h
--- a/Engine/Plugins/Enterprise/GLTFExporter/MigratedFromGit/Source/GLTFExporter/Private/Json/GLTFJsonUtility.h
+++ b/Engine/Plugins/Enterprise/GLTFExporter/MigratedFromGit/Source/GLTFExporter/Private/Json/GLTFJsonUtility.h
@@ -63,4 +63,18 @@ struct FGLTFJsonUtility
 		FString ExactStringRepresentation = FString::Printf(TEXT("%.9g"), Value);
 		JsonWriter.WriteRawJSONValue(Identifier, ExactStringRepresentation);
 	}
+
+	template <class WriterType, class ElementType, class AllocatorType>
+	static void WriteObjectArray(WriterType& JsonWriter, const FString& Identifier, const TArray<ElementType, AllocatorType>& ObjectArray)
+	{
+		if (ObjectArray.Num() > 0)
+		{
+			JsonWriter.WriteArrayStart(Identifier);
+			for (const ElementType& Object : ObjectArray)
+			{
+				Object.WriteObject(JsonWriter);
+			}
+			JsonWriter.WriteArrayEnd();
+		}
+	}
 };
