[GLTFExporter] feat(texture): add initial data structure for serializing top-level json objects related to images, samplers, and textures

commit 1f816ebf6c7033241e7b9efa8e843dfb82aa9d49
Date: Sat, 23 May 2020 18:06:29 +0200

    feat(texture): add initial data structure for serializing top-level json objects related to images, samplers, and textures
    (cherry picked from commit c82e08fe8402c58463366e01d0433910de80f55e​)

#skipci

diff --git a/Engine/Plugins/Enterprise/GLTFExporter/MigratedFromGit/Source/GLTFExporter/Private/Converters/GLTFConverterUtility.h b/Engine/Plugins/Enterprise/GLTFExporter/MigratedFromGit/Source/GLTFExporter/Private/Converters/GLTFConverterUtility.h
--- a/Engine/Plugins/Enterprise/GLTFExporter/MigratedFromGit/Source/GLTFExporter/Private/Converters/GLTFConverterUtility.h
+++ b/Engine/Plugins/Enterprise/GLTFExporter/MigratedFromGit/Source/GLTFExporter/Private/Converters/GLTFConverterUtility.h
@@ -102,6 +102,29 @@ struct FGLTFConverterUtility
 		}
 	}
 
+	static EGLTFJsonTextureFilter ConvertFilter(TextureFilter Filter)
+	{
+		switch (Filter)
+		{
+			case TextureFilter::TF_Nearest:   return EGLTFJsonTextureFilter::Nearest;
+			case TextureFilter::TF_Bilinear:  return EGLTFJsonTextureFilter::LinearMipmapNearest;
+			case TextureFilter::TF_Trilinear: return EGLTFJsonTextureFilter::LinearMipmapLinear;
+			case TextureFilter::TF_Default:   return EGLTFJsonTextureFilter::None; // TODO: handle this case better
+			default:                          return EGLTFJsonTextureFilter::None;
+		}
+	}
+
+	static EGLTFJsonTextureWrap ConvertWrap(TextureAddress Wrap)
+	{
+		switch (Wrap)
+		{
+			case TextureAddress::TA_Wrap:   return EGLTFJsonTextureWrap::Repeat;
+			case TextureAddress::TA_Mirror: return EGLTFJsonTextureWrap::MirroredRepeat;
+			case TextureAddress::TA_Clamp:  return EGLTFJsonTextureWrap::ClampToEdge;
+			default:                        return EGLTFJsonTextureWrap::Repeat; // fallback
+		}
+	}
+
 	static bool IsSkySphereBlueprint(const UBlueprint* Blueprint);
 
 	static bool IsHDRIBackdropBlueprint(const UBlueprint* Blueprint);
diff --git a/Engine/Plugins/Enterprise/GLTFExporter/MigratedFromGit/Source/GLTFExporter/Private/Json/GLTFJsonEnums.h b/Engine/Plugins/Enterprise/GLTFExporter/MigratedFromGit/Source/GLTFExporter/Private/Json/GLTFJsonEnums.h
--- a/Engine/Plugins/Enterprise/GLTFExporter/MigratedFromGit/Source/GLTFExporter/Private/Json/GLTFJsonEnums.h
+++ b/Engine/Plugins/Enterprise/GLTFExporter/MigratedFromGit/Source/GLTFExporter/Private/Json/GLTFJsonEnums.h
@@ -45,7 +45,6 @@ enum class EGLTFJsonBufferTarget
 
 enum class EGLTFJsonPrimitiveMode
 {
-	None = -1,
 	// unsupported
 	Points = 0,
 	Lines = 1,
@@ -60,13 +59,19 @@ enum class EGLTFJsonPrimitiveMode
 
 enum class EGLTFJsonAlphaMode
 {
-	None = -1,
 	Opaque,
 	Mask,
 	Blend
 };
 
-enum class EGLTFJsonSamplerFilter
+enum class EGLTFJsonMimeType
+{
+	None = -1,
+    PNG,
+    JPEG
+};
+
+enum class EGLTFJsonTextureFilter
 {
 	None = -1,
 	// valid for Min & Mag
@@ -79,9 +84,8 @@ enum class EGLTFJsonSamplerFilter
 	LinearMipmapLinear = 9987
 };
 
-enum class EGLTFJsonSamplerWrap
+enum class EGLTFJsonTextureWrap
 {
-	None = -1,
 	Repeat = 10497,
 	MirroredRepeat = 33648,
 	ClampToEdge = 33071
@@ -89,7 +93,6 @@ enum class EGLTFJsonSamplerWrap
 
 enum class EGLTFJsonAnimationInterpolation
 {
-	None = -1,
 	Linear,
 	Step,
 	CubicSpline,
diff --git a/Engine/Plugins/Enterprise/GLTFExporter/MigratedFromGit/Source/GLTFExporter/Private/Json/GLTFJsonImage.h b/Engine/Plugins/Enterprise/GLTFExporter/MigratedFromGit/Source/GLTFExporter/Private/Json/GLTFJsonImage.h
new file mode 100644
--- /dev/null
+++ b/Engine/Plugins/Enterprise/GLTFExporter/MigratedFromGit/Source/GLTFExporter/Private/Json/GLTFJsonImage.h
@@ -0,0 +1,51 @@
+// Copyright Epic Games, Inc. All Rights Reserved.
+
+#pragma once
+
+#include "Json/GLTFJsonEnums.h"
+#include "Json/GLTFJsonIndex.h"
+#include "Json/GLTFJsonUtility.h"
+#include "Serialization/JsonSerializer.h"
+
+struct FGLTFJsonImage
+{
+	FString Name;
+	FString Uri;
+
+	EGLTFJsonMimeType MimeType;
+
+	FGLTFJsonBufferViewIndex BufferView;
+
+	FGLTFJsonImage()
+		: MimeType(EGLTFJsonMimeType::None)
+	{
+	}
+
+	template <class CharType = TCHAR, class PrintPolicy = TPrettyJsonPrintPolicy<CharType>>
+    void WriteObject(TJsonWriter<CharType, PrintPolicy>& JsonWriter) const
+	{
+		JsonWriter.WriteObjectStart();
+
+		if (!Name.IsEmpty())
+		{
+			JsonWriter.WriteValue(TEXT("name"), Name);
+		}
+
+		if (!Uri.IsEmpty())
+		{
+			JsonWriter.WriteValue(TEXT("uri"), Uri);
+		}
+
+		if (MimeType != EGLTFJsonMimeType::None)
+		{
+			JsonWriter.WriteValue(TEXT("mimeType"), FGLTFJsonUtility::ToString(MimeType));
+		}
+
+		if (BufferView != INDEX_NONE)
+		{
+			JsonWriter.WriteValue(TEXT("bufferView"), BufferView);
+		}
+
+		JsonWriter.WriteObjectEnd();
+	}
+};
diff --git a/Engine/Plugins/Enterprise/GLTFExporter/MigratedFromGit/Source/GLTFExporter/Private/Json/GLTFJsonIndex.h b/Engine/Plugins/Enterprise/GLTFExporter/MigratedFromGit/Source/GLTFExporter/Private/Json/GLTFJsonIndex.h
--- a/Engine/Plugins/Enterprise/GLTFExporter/MigratedFromGit/Source/GLTFExporter/Private/Json/GLTFJsonIndex.h
+++ b/Engine/Plugins/Enterprise/GLTFExporter/MigratedFromGit/Source/GLTFExporter/Private/Json/GLTFJsonIndex.h
@@ -56,8 +56,11 @@ struct FGLTFJsonAccessorIndex : FGLTFJsonIndex<FGLTFJsonAccessorIndex> { using F
 struct FGLTFJsonBufferIndex : FGLTFJsonIndex<FGLTFJsonBufferIndex> { using FGLTFJsonIndex::FGLTFJsonIndex; };
 struct FGLTFJsonBufferViewIndex : FGLTFJsonIndex<FGLTFJsonBufferViewIndex> { using FGLTFJsonIndex::FGLTFJsonIndex; };
 struct FGLTFJsonCameraIndex : FGLTFJsonIndex<FGLTFJsonCameraIndex> { using FGLTFJsonIndex::FGLTFJsonIndex; };
+struct FGLTFJsonImageIndex : FGLTFJsonIndex<FGLTFJsonImageIndex> { using FGLTFJsonIndex::FGLTFJsonIndex; };
 struct FGLTFJsonMaterialIndex : FGLTFJsonIndex<FGLTFJsonMaterialIndex> { using FGLTFJsonIndex::FGLTFJsonIndex; };
 struct FGLTFJsonMeshIndex : FGLTFJsonIndex<FGLTFJsonMeshIndex> { using FGLTFJsonIndex::FGLTFJsonIndex; };
 struct FGLTFJsonNodeIndex : FGLTFJsonIndex<FGLTFJsonNodeIndex> { using FGLTFJsonIndex::FGLTFJsonIndex; };
+struct FGLTFJsonSamplerIndex : FGLTFJsonIndex<FGLTFJsonSamplerIndex> { using FGLTFJsonIndex::FGLTFJsonIndex; };
 struct FGLTFJsonSceneIndex : FGLTFJsonIndex<FGLTFJsonSceneIndex> { using FGLTFJsonIndex::FGLTFJsonIndex; };
 struct FGLTFJsonSkinIndex : FGLTFJsonIndex<FGLTFJsonSkinIndex> { using FGLTFJsonIndex::FGLTFJsonIndex; };
+struct FGLTFJsonTextureIndex : FGLTFJsonIndex<FGLTFJsonTextureIndex> { using FGLTFJsonIndex::FGLTFJsonIndex; };
diff --git a/Engine/Plugins/Enterprise/GLTFExporter/MigratedFromGit/Source/GLTFExporter/Private/Json/GLTFJsonMesh.h b/Engine/Plugins/Enterprise/GLTFExporter/MigratedFromGit/Source/GLTFExporter/Private/Json/GLTFJsonMesh.h
--- a/Engine/Plugins/Enterprise/GLTFExporter/MigratedFromGit/Source/GLTFExporter/Private/Json/GLTFJsonMesh.h
+++ b/Engine/Plugins/Enterprise/GLTFExporter/MigratedFromGit/Source/GLTFExporter/Private/Json/GLTFJsonMesh.h
@@ -49,7 +49,7 @@ struct FGLTFJsonPrimitive
 	FGLTFJsonAttributes    Attributes;
 
 	FGLTFJsonPrimitive()
-		: Mode(EGLTFJsonPrimitiveMode::None)
+		: Mode(EGLTFJsonPrimitiveMode::Triangles)
 	{
 	}
 
@@ -71,7 +71,7 @@ struct FGLTFJsonPrimitive
 			JsonWriter.WriteValue(TEXT("material"), Material);
 		}
 
-		if (Mode != EGLTFJsonPrimitiveMode::None)
+		if (Mode != EGLTFJsonPrimitiveMode::Triangles)
 		{
 			JsonWriter.WriteValue(TEXT("mode"), FGLTFJsonUtility::ToInteger(Mode));
 		}
diff --git a/Engine/Plugins/Enterprise/GLTFExporter/MigratedFromGit/Source/GLTFExporter/Private/Json/GLTFJsonRoot.h b/Engine/Plugins/Enterprise/GLTFExporter/MigratedFromGit/Source/GLTFExporter/Private/Json/GLTFJsonRoot.h
--- a/Engine/Plugins/Enterprise/GLTFExporter/MigratedFromGit/Source/GLTFExporter/Private/Json/GLTFJsonRoot.h
+++ b/Engine/Plugins/Enterprise/GLTFExporter/MigratedFromGit/Source/GLTFExporter/Private/Json/GLTFJsonRoot.h
@@ -6,10 +6,13 @@
 #include "Json/GLTFJsonAccessor.h"
 #include "Json/GLTFJsonBuffer.h"
 #include "Json/GLTFJsonBufferView.h"
+#include "Json/GLTFJsonImage.h"
 #include "Json/GLTFJsonMaterial.h"
 #include "Json/GLTFJsonMesh.h"
 #include "Json/GLTFJsonNode.h"
+#include "Json/GLTFJsonSampler.h"
 #include "Json/GLTFJsonScene.h"
+#include "Json/GLTFJsonTexture.h"
 
 #include "Containers/Set.h"
 #include "Policies/CondensedJsonPrintPolicy.h"
@@ -64,7 +67,10 @@ struct FGLTFJsonRoot
 	TArray<FGLTFJsonMaterial>   Materials;
 	TArray<FGLTFJsonMesh>       Meshes;
 	TArray<FGLTFJsonNode>       Nodes;
+	TArray<FGLTFJsonImage>      Images;
+	TArray<FGLTFJsonSampler>    Samplers;
 	TArray<FGLTFJsonScene>      Scenes;
+	TArray<FGLTFJsonTexture>    Textures;
 
 	template <class CharType = TCHAR, class PrintPolicy = TPrettyJsonPrintPolicy<CharType>>
 	void WriteObject(TJsonWriter<CharType, PrintPolicy>& JsonWriter) const
@@ -104,10 +110,13 @@ struct FGLTFJsonRoot
 		FGLTFJsonUtility::WriteObjectArray(JsonWriter, TEXT("accessors"), Accessors);
 		FGLTFJsonUtility::WriteObjectArray(JsonWriter, TEXT("buffers"), Buffers);
 		FGLTFJsonUtility::WriteObjectArray(JsonWriter, TEXT("bufferViews"), BufferViews);
+		FGLTFJsonUtility::WriteObjectArray(JsonWriter, TEXT("images"), Images);
 		FGLTFJsonUtility::WriteObjectArray(JsonWriter, TEXT("materials"), Materials);
 		FGLTFJsonUtility::WriteObjectArray(JsonWriter, TEXT("meshes"), Meshes);
 		FGLTFJsonUtility::WriteObjectArray(JsonWriter, TEXT("nodes"), Nodes);
+		FGLTFJsonUtility::WriteObjectArray(JsonWriter, TEXT("samplers"), Samplers);
 		FGLTFJsonUtility::WriteObjectArray(JsonWriter, TEXT("scenes"), Scenes);
+		FGLTFJsonUtility::WriteObjectArray(JsonWriter, TEXT("textures"), Textures);
 
 		JsonWriter.WriteObjectEnd();
 	}
diff --git a/Engine/Plugins/Enterprise/GLTFExporter/MigratedFromGit/Source/GLTFExporter/Private/Json/GLTFJsonSampler.h b/Engine/Plugins/Enterprise/GLTFExporter/MigratedFromGit/Source/GLTFExporter/Private/Json/GLTFJsonSampler.h
new file mode 100644
--- /dev/null
+++ b/Engine/Plugins/Enterprise/GLTFExporter/MigratedFromGit/Source/GLTFExporter/Private/Json/GLTFJsonSampler.h
@@ -0,0 +1,59 @@
+// Copyright Epic Games, Inc. All Rights Reserved.
+
+#pragma once
+
+#include "Json/GLTFJsonEnums.h"
+#include "Json/GLTFJsonUtility.h"
+#include "Serialization/JsonSerializer.h"
+
+struct FGLTFJsonSampler
+{
+	FString Name;
+
+	EGLTFJsonTextureFilter MinFilter;
+	EGLTFJsonTextureFilter MagFilter;
+
+	EGLTFJsonTextureWrap WrapS;
+	EGLTFJsonTextureWrap WrapT;
+
+	FGLTFJsonSampler()
+		: MinFilter(EGLTFJsonTextureFilter::None)
+		, MagFilter(EGLTFJsonTextureFilter::None)
+		, WrapS(EGLTFJsonTextureWrap::Repeat)
+		, WrapT(EGLTFJsonTextureWrap::Repeat)
+	{
+	}
+
+	template <class CharType = TCHAR, class PrintPolicy = TPrettyJsonPrintPolicy<CharType>>
+    void WriteObject(TJsonWriter<CharType, PrintPolicy>& JsonWriter) const
+	{
+		JsonWriter.WriteObjectStart();
+
+		if (!Name.IsEmpty())
+		{
+			JsonWriter.WriteValue(TEXT("name"), Name);
+		}
+
+		if (MinFilter != EGLTFJsonTextureFilter::None)
+		{
+			JsonWriter.WriteValue(TEXT("minFilter"), FGLTFJsonUtility::ToInteger(MinFilter));
+		}
+
+		if (MagFilter != EGLTFJsonTextureFilter::None)
+		{
+			JsonWriter.WriteValue(TEXT("magFilter"), FGLTFJsonUtility::ToInteger(MagFilter));
+		}
+
+		if (WrapS != EGLTFJsonTextureWrap::Repeat)
+		{
+			JsonWriter.WriteValue(TEXT("wrapS"), FGLTFJsonUtility::ToInteger(WrapS));
+		}
+
+		if (WrapT != EGLTFJsonTextureWrap::Repeat)
+		{
+			JsonWriter.WriteValue(TEXT("wrapT"), FGLTFJsonUtility::ToInteger(WrapT));
+		}
+
+		JsonWriter.WriteObjectEnd();
+	}
+};
diff --git a/Engine/Plugins/Enterprise/GLTFExporter/MigratedFromGit/Source/GLTFExporter/Private/Json/GLTFJsonTexture.h b/Engine/Plugins/Enterprise/GLTFExporter/MigratedFromGit/Source/GLTFExporter/Private/Json/GLTFJsonTexture.h
new file mode 100644
--- /dev/null
+++ b/Engine/Plugins/Enterprise/GLTFExporter/MigratedFromGit/Source/GLTFExporter/Private/Json/GLTFJsonTexture.h
@@ -0,0 +1,38 @@
+// Copyright Epic Games, Inc. All Rights Reserved.
+
+#pragma once
+
+#include "Json/GLTFJsonIndex.h"
+#include "Serialization/JsonSerializer.h"
+
+struct FGLTFJsonTexture
+{
+	FString Name;
+
+	FGLTFJsonSamplerIndex Sampler;
+
+	FGLTFJsonImageIndex Source;
+
+	template <class CharType = TCHAR, class PrintPolicy = TPrettyJsonPrintPolicy<CharType>>
+    void WriteObject(TJsonWriter<CharType, PrintPolicy>& JsonWriter) const
+	{
+		JsonWriter.WriteObjectStart();
+
+		if (!Name.IsEmpty())
+		{
+			JsonWriter.WriteValue(TEXT("name"), Name);
+		}
+
+		if (Sampler != INDEX_NONE)
+		{
+			JsonWriter.WriteValue(TEXT("sampler"), Sampler);
+		}
+
+		if (Source != INDEX_NONE)
+		{
+			JsonWriter.WriteValue(TEXT("source"), Source);
+		}
+
+		JsonWriter.WriteObjectEnd();
+	}
+};
diff --git a/Engine/Plugins/Enterprise/GLTFExporter/MigratedFromGit/Source/GLTFExporter/Private/Json/GLTFJsonUtility.h b/Engine/Plugins/Enterprise/GLTFExporter/MigratedFromGit/Source/GLTFExporter/Private/Json/GLTFJsonUtility.h
--- a/Engine/Plugins/Enterprise/GLTFExporter/MigratedFromGit/Source/GLTFExporter/Private/Json/GLTFJsonUtility.h
+++ b/Engine/Plugins/Enterprise/GLTFExporter/MigratedFromGit/Source/GLTFExporter/Private/Json/GLTFJsonUtility.h
@@ -35,6 +35,16 @@ struct FGLTFJsonUtility
 		}
 	}
 
+	static const TCHAR* ToString(EGLTFJsonMimeType Value)
+	{
+		switch (Value)
+		{
+			case EGLTFJsonMimeType::PNG:  return TEXT("image/png");
+			case EGLTFJsonMimeType::JPEG: return TEXT("image/jpeg");
+			default:                      return TEXT("unknown");
+		}
+	}
+
 	static const TCHAR* ToString(EGLTFJsonAccessorType Value)
 	{
 		switch (Value)
