[GLTFExporter] feat(texture): add data structures for serializing texture info objects to json

commit c4d9adcb570c38155edcdf48f7fec6e7c9ac0626
Date: Sat, 23 May 2020 18:09:57 +0200

    feat(texture): add data structures for serializing texture info objects to json
    (cherry picked from commit 0843623d420e7cc1bc3defea1ea438506638385b​)

#skipci

diff --git a/Engine/Plugins/Enterprise/GLTFExporter/MigratedFromGit/Source/GLTFExporter/Private/Json/GLTFJsonMaterial.h b/Engine/Plugins/Enterprise/GLTFExporter/MigratedFromGit/Source/GLTFExporter/Private/Json/GLTFJsonMaterial.h
--- a/Engine/Plugins/Enterprise/GLTFExporter/MigratedFromGit/Source/GLTFExporter/Private/Json/GLTFJsonMaterial.h
+++ b/Engine/Plugins/Enterprise/GLTFExporter/MigratedFromGit/Source/GLTFExporter/Private/Json/GLTFJsonMaterial.h
@@ -8,12 +8,100 @@
 #include "Json/GLTFJsonColor4.h"
 #include "Serialization/JsonSerializer.h"
 
+struct FGLTFJsonTextureInfo
+{
+	FGLTFJsonTextureIndex Index;
+	int32 TexCoord;
+
+	FGLTFJsonTextureInfo()
+        : TexCoord(0)
+	{
+	}
+
+	template <class CharType = TCHAR, class PrintPolicy = TPrettyJsonPrintPolicy<CharType>>
+    void WriteObject(TJsonWriter<CharType, PrintPolicy>& JsonWriter) const
+	{
+		JsonWriter.WriteObjectStart();
+
+		JsonWriter.WriteValue(TEXT("index"), Index);
+
+		if (TexCoord != 0)
+		{
+			JsonWriter.WriteValue(TEXT("texCoord"), TexCoord);
+		}
+
+		JsonWriter.WriteObjectEnd();
+	}
+};
+
+struct FGLTFJsonNormalTextureInfo : FGLTFJsonTextureInfo
+{
+	float Scale;
+
+	FGLTFJsonNormalTextureInfo()
+        : Scale(1)
+	{
+	}
+
+	template <class CharType = TCHAR, class PrintPolicy = TPrettyJsonPrintPolicy<CharType>>
+    void WriteObject(TJsonWriter<CharType, PrintPolicy>& JsonWriter) const
+	{
+		JsonWriter.WriteObjectStart();
+
+		JsonWriter.WriteValue(TEXT("index"), Index);
+
+		if (TexCoord != 0)
+		{
+			JsonWriter.WriteValue(TEXT("texCoord"), TexCoord);
+		}
+
+		if (Scale != 1)
+		{
+			FGLTFJsonUtility::WriteExactValue(JsonWriter, TEXT("scale"), Scale);
+		}
+
+		JsonWriter.WriteObjectEnd();
+	}
+};
+
+struct FGLTFJsonOcclusionTextureInfo : FGLTFJsonTextureInfo
+{
+	float Strength;
+
+	FGLTFJsonOcclusionTextureInfo()
+        : Strength(1)
+	{
+	}
+
+	template <class CharType = TCHAR, class PrintPolicy = TPrettyJsonPrintPolicy<CharType>>
+    void WriteObject(TJsonWriter<CharType, PrintPolicy>& JsonWriter) const
+	{
+		JsonWriter.WriteObjectStart();
+
+		JsonWriter.WriteValue(TEXT("index"), Index);
+
+		if (TexCoord != 0)
+		{
+			JsonWriter.WriteValue(TEXT("texCoord"), TexCoord);
+		}
+
+		if (Strength != 1)
+		{
+			FGLTFJsonUtility::WriteExactValue(JsonWriter, TEXT("strength"), Strength);
+		}
+
+		JsonWriter.WriteObjectEnd();
+	}
+};
+
 struct FGLTFJsonPBRMetallicRoughness
 {
 	FGLTFJsonColor4 BaseColorFactor;
+	FGLTFJsonTextureInfo BaseColorTexture;
 
 	float MetallicFactor;
 	float RoughnessFactor;
+	FGLTFJsonTextureInfo MetallicRoughnessTexture;
 
 	FGLTFJsonPBRMetallicRoughness()
 		: BaseColorFactor(FGLTFJsonColor4::White)
@@ -33,6 +121,12 @@ struct FGLTFJsonPBRMetallicRoughness
 			BaseColorFactor.WriteArray(JsonWriter);
 		}
 
+		if (BaseColorTexture.Index != INDEX_NONE)
+		{
+			JsonWriter.WriteIdentifierPrefix(TEXT("baseColorTexture"));
+			BaseColorTexture.WriteObject(JsonWriter);
+		}
+
 		if (MetallicFactor != 1)
 		{
 			FGLTFJsonUtility::WriteExactValue(JsonWriter, TEXT("metallicFactor"), MetallicFactor);
@@ -43,6 +137,12 @@ struct FGLTFJsonPBRMetallicRoughness
 			FGLTFJsonUtility::WriteExactValue(JsonWriter, TEXT("roughnessFactor"), RoughnessFactor);
 		}
 
+		if (MetallicRoughnessTexture.Index != INDEX_NONE)
+		{
+			JsonWriter.WriteIdentifierPrefix(TEXT("metallicRoughnessTexture"));
+			MetallicRoughnessTexture.WriteObject(JsonWriter);
+		}
+
 		JsonWriter.WriteObjectEnd();
 	}
 };
@@ -53,6 +153,10 @@ struct FGLTFJsonMaterial
 
 	FGLTFJsonPBRMetallicRoughness PBRMetallicRoughness;
 
+	FGLTFJsonNormalTextureInfo NormalTexture;
+	FGLTFJsonOcclusionTextureInfo OcclusionTexture;
+
+	FGLTFJsonTextureInfo EmissiveTexture;
 	FGLTFJsonColor3 EmissiveFactor;
 
 	EGLTFJsonAlphaMode AlphaMode;
@@ -81,6 +185,24 @@ struct FGLTFJsonMaterial
 		JsonWriter.WriteIdentifierPrefix(TEXT("pbrMetallicRoughness"));
 		PBRMetallicRoughness.WriteObject(JsonWriter);
 
+		if (NormalTexture.Index != INDEX_NONE)
+		{
+			JsonWriter.WriteIdentifierPrefix(TEXT("normalTexture"));
+			NormalTexture.WriteObject(JsonWriter);
+		}
+
+		if (OcclusionTexture.Index != INDEX_NONE)
+		{
+			JsonWriter.WriteIdentifierPrefix(TEXT("occlusionTexture"));
+			OcclusionTexture.WriteObject(JsonWriter);
+		}
+
+		if (EmissiveTexture.Index != INDEX_NONE)
+		{
+			JsonWriter.WriteIdentifierPrefix(TEXT("emissiveTexture"));
+			EmissiveTexture.WriteObject(JsonWriter);
+		}
+
 		if (EmissiveFactor != FGLTFJsonColor3::Black)
 		{
 			JsonWriter.WriteIdentifierPrefix(TEXT("emissiveFactor"));
