[GLTFExporter] feat(material): add GLTFImageBuilder to compress and export images

commit a1d9f2f9222f29ea04592bfd6ac8ab7022bf422d
Date: Sun, 24 May 2020 23:51:00 +0200

    feat(material): add GLTFImageBuilder to compress and export images
    (cherry picked from commit 04b72e10b0edf62ae93237ca9e76085ba257a56b​)

#skipci

diff --git a/Engine/Plugins/Enterprise/GLTFExporter/MigratedFromGit/Source/GLTFExporter/Private/Builders/GLTFBufferBuilder.h b/Engine/Plugins/Enterprise/GLTFExporter/MigratedFromGit/Source/GLTFExporter/Private/Builders/GLTFBufferBuilder.h
--- a/Engine/Plugins/Enterprise/GLTFExporter/MigratedFromGit/Source/GLTFExporter/Private/Builders/GLTFBufferBuilder.h
+++ b/Engine/Plugins/Enterprise/GLTFExporter/MigratedFromGit/Source/GLTFExporter/Private/Builders/GLTFBufferBuilder.h
@@ -24,5 +24,5 @@ public:
 private:
 
 	FGLTFJsonBufferIndex BufferIndex;
-	TArray<uint8> BufferData;
+	TArray64<uint8> BufferData;
 };
diff --git a/Engine/Plugins/Enterprise/GLTFExporter/MigratedFromGit/Source/GLTFExporter/Private/Builders/GLTFBuilderUtility.cpp b/Engine/Plugins/Enterprise/GLTFExporter/MigratedFromGit/Source/GLTFExporter/Private/Builders/GLTFBuilderUtility.cpp
new file mode 100644
--- /dev/null
+++ b/Engine/Plugins/Enterprise/GLTFExporter/MigratedFromGit/Source/GLTFExporter/Private/Builders/GLTFBuilderUtility.cpp
@@ -0,0 +1,28 @@
+// Copyright Epic Games, Inc. All Rights Reserved.
+
+#include "Builders/GLTFBuilderUtility.h"
+#include "IImageWrapperModule.h"
+#include "IImageWrapper.h"
+
+bool FGLTFBuilderUtility::CompressImage(const TArray64<uint8>& InRawData, int32 InWidth, int32 InHeight, ERGBFormat InRawFormat, int32 InBitDepth, TArray64<uint8>& OutCompressedData, EImageFormat OutCompressionFormat, int32 OutCompressionQuality)
+{
+	IImageWrapperModule& ImageWrapperModule = FModuleManager::Get().LoadModuleChecked<IImageWrapperModule>(TEXT("ImageWrapper"));
+
+	TSharedPtr<IImageWrapper> ImageWrapper = ImageWrapperModule.CreateImageWrapper(OutCompressionFormat);
+	if (!ImageWrapper.IsValid())
+	{
+		// TODO: report error
+		return false;
+	}
+
+	const bool bFormatSupported = ImageWrapper->SetRaw(InRawData.GetData(), InRawData.Num(), InWidth, InHeight, InRawFormat, InBitDepth);
+	if (!bFormatSupported)
+	{
+		// TODO: report error
+		return false;
+	}
+
+	const TArray64<uint8>& CompressedData = ImageWrapper->GetCompressed(OutCompressionQuality);
+	OutCompressedData.Append(CompressedData);
+	return true;
+}
diff --git a/Engine/Plugins/Enterprise/GLTFExporter/MigratedFromGit/Source/GLTFExporter/Private/Builders/GLTFBuilderUtility.h b/Engine/Plugins/Enterprise/GLTFExporter/MigratedFromGit/Source/GLTFExporter/Private/Builders/GLTFBuilderUtility.h
--- a/Engine/Plugins/Enterprise/GLTFExporter/MigratedFromGit/Source/GLTFExporter/Private/Builders/GLTFBuilderUtility.h
+++ b/Engine/Plugins/Enterprise/GLTFExporter/MigratedFromGit/Source/GLTFExporter/Private/Builders/GLTFBuilderUtility.h
@@ -3,6 +3,10 @@
 #pragma once
 
 #include "Engine.h"
+#include "Json/GLTFJsonEnums.h"
+
+enum class EImageFormat : signed char;
+enum class ERGBFormat : int8;
 
 struct FGLTFBuilderUtility
 {
@@ -18,4 +22,16 @@ struct FGLTFBuilderUtility
 
 		return Name;
 	}
+
+	static bool CompressImage(const TArray64<uint8>& InRawData, int32 InWidth, int32 InHeight, ERGBFormat InRawFormat, int32 InBitDepth, TArray64<uint8>& OutCompressedData, EImageFormat OutCompressionFormat, int32 OutCompressionQuality);
+
+	static const TCHAR* GetFileExtension(EGLTFJsonMimeType MimeType)
+	{
+		switch (MimeType)
+		{
+			case EGLTFJsonMimeType::PNG:  return TEXT(".png");
+			case EGLTFJsonMimeType::JPEG: return TEXT(".jpg");
+			default:                      return TEXT(".unknown");
+		}
+	}
 };
diff --git a/Engine/Plugins/Enterprise/GLTFExporter/MigratedFromGit/Source/GLTFExporter/Private/Builders/GLTFConvertBuilder.h b/Engine/Plugins/Enterprise/GLTFExporter/MigratedFromGit/Source/GLTFExporter/Private/Builders/GLTFConvertBuilder.h
--- a/Engine/Plugins/Enterprise/GLTFExporter/MigratedFromGit/Source/GLTFExporter/Private/Builders/GLTFConvertBuilder.h
+++ b/Engine/Plugins/Enterprise/GLTFExporter/MigratedFromGit/Source/GLTFExporter/Private/Builders/GLTFConvertBuilder.h
@@ -2,13 +2,13 @@
 
 #pragma once
 
-#include "Builders/GLTFBufferBuilder.h"
+#include "Builders/GLTFImageBuilder.h"
 #include "Converters/GLTFVertexBufferConverters.h"
 #include "Converters/GLTFStaticMeshConverters.h"
 #include "Converters/GLTFMaterialConverters.h"
 #include "Converters/GLTFLevelConverters.h"
 
-class FGLTFConvertBuilder : public FGLTFBufferBuilder
+class FGLTFConvertBuilder : public FGLTFImageBuilder
 {
 public:
 
diff --git a/Engine/Plugins/Enterprise/GLTFExporter/MigratedFromGit/Source/GLTFExporter/Private/Builders/GLTFImageBuilder.cpp b/Engine/Plugins/Enterprise/GLTFExporter/MigratedFromGit/Source/GLTFExporter/Private/Builders/GLTFImageBuilder.cpp
new file mode 100644
--- /dev/null
+++ b/Engine/Plugins/Enterprise/GLTFExporter/MigratedFromGit/Source/GLTFExporter/Private/Builders/GLTFImageBuilder.cpp
@@ -0,0 +1,109 @@
+// Copyright Epic Games, Inc. All Rights Reserved.
+
+#include "Builders/GLTFImageBuilder.h"
+#include "Builders/GLTFBuilderUtility.h"
+#include "IImageWrapper.h"
+#include "Engine/Texture.h"
+
+FGLTFJsonImageIndex FGLTFImageBuilder::AddImage(const TArray64<uint8>& RawData, int32 Width, int32 Height, ERGBFormat RawFormat, int32 BitDepth, bool bFloatFormat, const FString& Name, EGLTFJsonMimeType MimeType, int32 Quality)
+{
+	TArray64<uint8> ImageData;
+
+	if (bFloatFormat)
+	{
+		// TODO: implement support
+		return FGLTFJsonImageIndex(INDEX_NONE);
+	}
+	else
+	{
+		switch (MimeType)
+		{
+		case EGLTFJsonMimeType::PNG:
+            if (!FGLTFBuilderUtility::CompressImage(RawData, Width, Height, RawFormat, BitDepth, ImageData, EImageFormat::PNG, Quality))
+            {
+                return FGLTFJsonImageIndex(INDEX_NONE);
+            }
+			break;
+		case EGLTFJsonMimeType::JPEG:
+            if (!FGLTFBuilderUtility::CompressImage(RawData, Width, Height, RawFormat, BitDepth, ImageData, EImageFormat::JPEG, Quality))
+            {
+                return FGLTFJsonImageIndex(INDEX_NONE);
+            }
+			break;
+		default:
+            // TODO: report error
+            return FGLTFJsonImageIndex(INDEX_NONE);
+		}
+	}
+
+	FGLTFJsonImage Image;
+	Image.Name = Name;
+	Image.MimeType = MimeType;
+
+	const FGLTFJsonImageIndex ImageIndex = FGLTFJsonBuilder::AddImage(Image);
+	ImageDataLookup.Add(ImageIndex, ImageData);
+	return ImageIndex;
+}
+
+FGLTFJsonImageIndex FGLTFImageBuilder::AddImage(const FTextureSource& Image, const FString& Name, EGLTFJsonMimeType MimeType, int32 Quality)
+{
+	// TODO: are these always zero?
+	const int32 BlockIndex = 0;
+	const int32 LayerIndex = 0;
+	const int32 MipIndex = 0;
+
+	ERGBFormat RawFormat;
+	int32 BitDepth;
+	bool bFloatFormat;
+
+	const ETextureSourceFormat SourceFormat = Image.GetFormat(LayerIndex);
+	switch (SourceFormat)
+	{
+        case TSF_BGRA8:   RawFormat = ERGBFormat::BGRA; BitDepth = 8;  bFloatFormat = false; break;
+        case TSF_RGBA8:   RawFormat = ERGBFormat::RGBA; BitDepth = 8;  bFloatFormat = false; break;
+		case TSF_BGRE8:   RawFormat = ERGBFormat::BGRA; BitDepth = 8;  bFloatFormat = false; break;
+        case TSF_RGBE8:   RawFormat = ERGBFormat::RGBA; BitDepth = 8;  bFloatFormat = false; break;
+        case TSF_RGBA16:  RawFormat = ERGBFormat::RGBA; BitDepth = 16; bFloatFormat = false; break;
+        case TSF_RGBA16F: RawFormat = ERGBFormat::RGBA; BitDepth = 16; bFloatFormat = true;  break;
+        case TSF_G8:      RawFormat = ERGBFormat::Gray; BitDepth = 8;  bFloatFormat = false; break;
+        case TSF_G16:     RawFormat = ERGBFormat::Gray; BitDepth = 16; bFloatFormat = false; break;
+		default:
+			// TODO: report error
+			return FGLTFJsonImageIndex(INDEX_NONE);
+	}
+
+	TArray64<uint8> RawData;
+	if (!const_cast<FTextureSource&>(Image).GetMipData(RawData, BlockIndex, LayerIndex, MipIndex))
+	{
+		// TODO: report error
+		return FGLTFJsonImageIndex(INDEX_NONE);
+	}
+
+	return AddImage(RawData, Image.GetSizeX(), Image.GetSizeY(), RawFormat, BitDepth, bFloatFormat, Name, MimeType, Quality);
+}
+
+bool FGLTFImageBuilder::Serialize(FArchive& Archive, const FString& FilePath)
+{
+	const FString ImageDir = FPaths::GetPath(FilePath);
+
+	for (const auto& DataPair : ImageDataLookup)
+	{
+		const TArray64<uint8>& ImageData = DataPair.Value;
+		FGLTFJsonImage& JsonImage = JsonRoot.Images[DataPair.Key];
+
+		const FString ImageUri = JsonImage.Name + FGLTFBuilderUtility::GetFileExtension(JsonImage.MimeType);
+		const FString ImagePath = FPaths::Combine(ImageDir, ImageUri);
+
+		if (!FFileHelper::SaveArrayToFile(ImageData, *ImagePath))
+		{
+			// TODO: report error
+			continue;
+		}
+
+		JsonImage.Uri = ImageUri;
+		JsonImage.Name.Empty(); // URI already contains name
+		JsonImage.MimeType = EGLTFJsonMimeType::None; // Not required if external file
+	}
+
+	return FGLTFBufferBuilder::Serialize(Archive, FilePath);
+}
diff --git a/Engine/Plugins/Enterprise/GLTFExporter/MigratedFromGit/Source/GLTFExporter/Private/Builders/GLTFImageBuilder.h b/Engine/Plugins/Enterprise/GLTFExporter/MigratedFromGit/Source/GLTFExporter/Private/Builders/GLTFImageBuilder.h
new file mode 100644
--- /dev/null
+++ b/Engine/Plugins/Enterprise/GLTFExporter/MigratedFromGit/Source/GLTFExporter/Private/Builders/GLTFImageBuilder.h
@@ -0,0 +1,22 @@
+// Copyright Epic Games, Inc. All Rights Reserved.
+
+#pragma once
+
+#include "Builders/GLTFBufferBuilder.h"
+
+enum class ERGBFormat : int8;
+
+class FGLTFImageBuilder : public FGLTFBufferBuilder
+{
+public:
+
+	FGLTFJsonImageIndex AddImage(const TArray64<uint8>& RawData, int32 Width, int32 Height, ERGBFormat RawFormat, int32 BitDepth, bool bFloatFormat, const FString& Name = TEXT(""), EGLTFJsonMimeType MimeType = EGLTFJsonMimeType::PNG, int32 Quality = 100);
+
+	FGLTFJsonImageIndex AddImage(const FTextureSource& Image, const FString& Name = TEXT(""), EGLTFJsonMimeType MimeType = EGLTFJsonMimeType::PNG, int32 Quality = 100);
+
+	virtual bool Serialize(FArchive& Archive, const FString& FilePath) override;
+
+private:
+
+	TMap<FGLTFJsonImageIndex, TArray64<uint8>> ImageDataLookup;
+};
diff --git a/Engine/Plugins/Enterprise/GLTFExporter/MigratedFromGit/Source/GLTFExporter/Private/Builders/GLTFJsonBuilder.cpp b/Engine/Plugins/Enterprise/GLTFExporter/MigratedFromGit/Source/GLTFExporter/Private/Builders/GLTFJsonBuilder.cpp
--- a/Engine/Plugins/Enterprise/GLTFExporter/MigratedFromGit/Source/GLTFExporter/Private/Builders/GLTFJsonBuilder.cpp
+++ b/Engine/Plugins/Enterprise/GLTFExporter/MigratedFromGit/Source/GLTFExporter/Private/Builders/GLTFJsonBuilder.cpp
@@ -26,6 +26,11 @@ FGLTFJsonBufferViewIndex FGLTFJsonBuilder::AddBufferView(const FGLTFJsonBufferVi
 	return FGLTFJsonBufferViewIndex(JsonRoot.BufferViews.Add(JsonBufferView));
 }
 
+FGLTFJsonImageIndex FGLTFJsonBuilder::AddImage(const FGLTFJsonImage& JsonImage)
+{
+	return FGLTFJsonImageIndex(JsonRoot.Images.Add(JsonImage));
+}
+
 FGLTFJsonMaterialIndex FGLTFJsonBuilder::AddMaterial(const FGLTFJsonMaterial& JsonMaterial)
 {
 	return FGLTFJsonMaterialIndex(JsonRoot.Materials.Add(JsonMaterial));
@@ -41,11 +46,21 @@ FGLTFJsonNodeIndex FGLTFJsonBuilder::AddNode(const FGLTFJsonNode& JsonNode)
 	return FGLTFJsonNodeIndex(JsonRoot.Nodes.Add(JsonNode));
 }
 
+FGLTFJsonSamplerIndex FGLTFJsonBuilder::AddSampler(const FGLTFJsonSampler& JsonSampler)
+{
+	return FGLTFJsonSamplerIndex(JsonRoot.Samplers.Add(JsonSampler));
+}
+
 FGLTFJsonSceneIndex FGLTFJsonBuilder::AddScene(const FGLTFJsonScene& JsonScene)
 {
 	return FGLTFJsonSceneIndex(JsonRoot.Scenes.Add(JsonScene));
 }
 
+FGLTFJsonTextureIndex FGLTFJsonBuilder::AddTexture(const FGLTFJsonTexture& JsonTexture)
+{
+	return FGLTFJsonTextureIndex(JsonRoot.Textures.Add(JsonTexture));
+}
+
 bool FGLTFJsonBuilder::Serialize(FArchive& Archive, const FString& FilePath)
 {
 	JsonRoot.Serialize(&Archive, true);
diff --git a/Engine/Plugins/Enterprise/GLTFExporter/MigratedFromGit/Source/GLTFExporter/Private/Builders/GLTFJsonBuilder.h b/Engine/Plugins/Enterprise/GLTFExporter/MigratedFromGit/Source/GLTFExporter/Private/Builders/GLTFJsonBuilder.h
--- a/Engine/Plugins/Enterprise/GLTFExporter/MigratedFromGit/Source/GLTFExporter/Private/Builders/GLTFJsonBuilder.h
+++ b/Engine/Plugins/Enterprise/GLTFExporter/MigratedFromGit/Source/GLTFExporter/Private/Builders/GLTFJsonBuilder.h
@@ -16,10 +16,13 @@ public:
 	FGLTFJsonAccessorIndex AddAccessor(const FGLTFJsonAccessor& JsonAccessor);
 	FGLTFJsonBufferIndex AddBuffer(const FGLTFJsonBuffer& JsonBuffer);
 	FGLTFJsonBufferViewIndex AddBufferView(const FGLTFJsonBufferView& JsonBufferView);
+	FGLTFJsonImageIndex AddImage(const FGLTFJsonImage& JsonImage);
 	FGLTFJsonMaterialIndex AddMaterial(const FGLTFJsonMaterial& JsonMaterial);
 	FGLTFJsonMeshIndex AddMesh(const FGLTFJsonMesh& JsonMesh);
 	FGLTFJsonNodeIndex AddNode(const FGLTFJsonNode& JsonNode);
+	FGLTFJsonSamplerIndex AddSampler(const FGLTFJsonSampler& JsonSampler);
 	FGLTFJsonSceneIndex AddScene(const FGLTFJsonScene& JsonScene);
+	FGLTFJsonTextureIndex AddTexture(const FGLTFJsonTexture& JsonTexture);
 
 	virtual bool Serialize(FArchive& Archive, const FString& FilePath);
 
