[GLTFExporter] feat(material): add & convert textures and samplers

commit 36734289c0fe23ec6b2970980869efffed3c2193
Date: Sun, 24 May 2020 23:52:27 +0200

    feat(material): add & convert textures and samplers
    (cherry picked from commit 926699fc24bf20c9d8009d14f9d862062aa376a6​)

#skipci

diff --git a/Engine/Plugins/Enterprise/GLTFExporter/MigratedFromGit/Source/GLTFExporter/Private/Builders/GLTFConvertBuilder.cpp b/Engine/Plugins/Enterprise/GLTFExporter/MigratedFromGit/Source/GLTFExporter/Private/Builders/GLTFConvertBuilder.cpp
--- a/Engine/Plugins/Enterprise/GLTFExporter/MigratedFromGit/Source/GLTFExporter/Private/Builders/GLTFConvertBuilder.cpp
+++ b/Engine/Plugins/Enterprise/GLTFExporter/MigratedFromGit/Source/GLTFExporter/Private/Builders/GLTFConvertBuilder.cpp
@@ -108,6 +108,16 @@ FGLTFJsonMaterialIndex FGLTFConvertBuilder::GetOrAddMaterial(const UMaterialInte
 	return MaterialConverter.GetOrAdd(*this, DesiredName.IsEmpty() ? Material->GetName() : DesiredName, Material);
 }
 
+FGLTFJsonTextureIndex FGLTFConvertBuilder::GetOrAddTexture(const UTexture2D* Texture, const FString& DesiredName)
+{
+	if (Texture == nullptr)
+	{
+		return FGLTFJsonTextureIndex(INDEX_NONE);
+	}
+
+	return Texture2DConverter.GetOrAdd(*this, DesiredName.IsEmpty() ? Texture->GetName() : DesiredName, Texture);
+}
+
 FGLTFJsonNodeIndex FGLTFConvertBuilder::GetOrAddNode(const USceneComponent* SceneComponent, bool bSelectedOnly, bool bRootNode, const FString& DesiredName)
 {
 	if (SceneComponent == nullptr)
diff --git a/Engine/Plugins/Enterprise/GLTFExporter/MigratedFromGit/Source/GLTFExporter/Private/Builders/GLTFConvertBuilder.h b/Engine/Plugins/Enterprise/GLTFExporter/MigratedFromGit/Source/GLTFExporter/Private/Builders/GLTFConvertBuilder.h
--- a/Engine/Plugins/Enterprise/GLTFExporter/MigratedFromGit/Source/GLTFExporter/Private/Builders/GLTFConvertBuilder.h
+++ b/Engine/Plugins/Enterprise/GLTFExporter/MigratedFromGit/Source/GLTFExporter/Private/Builders/GLTFConvertBuilder.h
@@ -6,6 +6,7 @@
 #include "Converters/GLTFVertexBufferConverters.h"
 #include "Converters/GLTFStaticMeshConverters.h"
 #include "Converters/GLTFMaterialConverters.h"
+#include "Converters/GLTFTextureConverters.h"
 #include "Converters/GLTFLevelConverters.h"
 
 class FGLTFConvertBuilder : public FGLTFImageBuilder
@@ -23,12 +24,13 @@ public:
 	FGLTFJsonMeshIndex GetOrAddMesh(const UStaticMesh* StaticMesh, int32 LODIndex = 0, const FColorVertexBuffer* OverrideVertexColors = nullptr, const FGLTFMaterialArray& OverrideMaterials = {}, const FString& DesiredName = TEXT(""));
 	FGLTFJsonMeshIndex GetOrAddMesh(const UStaticMeshComponent* StaticMeshComponent, const FString& DesiredName = TEXT(""));
 
+	FGLTFJsonMaterialIndex GetOrAddMaterial(const UMaterialInterface* Material, const FString& DesiredName = TEXT(""));
+	FGLTFJsonTextureIndex GetOrAddTexture(const UTexture2D* Texture, const FString& DesiredName = TEXT(""));
+
 	FGLTFJsonNodeIndex GetOrAddNode(const USceneComponent* SceneComponent, bool bSelectedOnly, bool bRootNode = false, const FString& DesiredName = TEXT(""));
 	FGLTFJsonSceneIndex GetOrAddScene(const ULevel* Level, bool bSelectedOnly, const FString& DesiredName = TEXT(""));
 	FGLTFJsonSceneIndex GetOrAddScene(const UWorld* World, bool bSelectedOnly, const FString& DesiredName = TEXT(""));
 
-	FGLTFJsonMaterialIndex GetOrAddMaterial(const UMaterialInterface* Material, const FString& DesiredName = TEXT(""));
-
 private:
 
 	FGLTFPositionVertexBufferConverter PositionVertexBufferConverter;
@@ -42,6 +44,7 @@ private:
 	FGLTFStaticMeshConverter StaticMeshConverter;
 
 	FGLTFMaterialConverter MaterialConverter;
+	FGLTFTexture2DConverter Texture2DConverter;
 
 	FGLTFSceneComponentConverter SceneComponentConverter;
 	FGLTFLevelConverter LevelConverter;
diff --git a/Engine/Plugins/Enterprise/GLTFExporter/MigratedFromGit/Source/GLTFExporter/Private/Converters/GLTFConverterUtility.cpp b/Engine/Plugins/Enterprise/GLTFExporter/MigratedFromGit/Source/GLTFExporter/Private/Converters/GLTFConverterUtility.cpp
--- a/Engine/Plugins/Enterprise/GLTFExporter/MigratedFromGit/Source/GLTFExporter/Private/Converters/GLTFConverterUtility.cpp
+++ b/Engine/Plugins/Enterprise/GLTFExporter/MigratedFromGit/Source/GLTFExporter/Private/Converters/GLTFConverterUtility.cpp
@@ -1,8 +1,17 @@
 // Copyright Epic Games, Inc. All Rights Reserved.
 
 #include "Converters/GLTFConverterUtility.h"
+#include "Engine/TextureLODSettings.h"
 #include "Engine/Blueprint.h"
 
+ETextureSamplerFilter FGLTFConverterUtility::GetDefaultFilter(TextureGroup LODGroup)
+{
+	const ITargetPlatform* RunningPlatform = GetTargetPlatformManagerRef().GetRunningTargetPlatform(); // TODO: should this be the running platform?
+	const UTextureLODSettings& TextureLODSettings = RunningPlatform->GetTextureLODSettings();
+	const FTextureLODGroup& TextureLODGroup = TextureLODSettings.GetTextureLODGroup(LODGroup);
+	return TextureLODGroup.Filter;
+}
+
 bool FGLTFConverterUtility::IsSkySphereBlueprint(const UBlueprint* Blueprint)
 {
 	return Blueprint != nullptr && Blueprint->GetPathName().Equals(TEXT("/Engine/EngineSky/BP_Sky_Sphere.BP_Sky_Sphere"));
diff --git a/Engine/Plugins/Enterprise/GLTFExporter/MigratedFromGit/Source/GLTFExporter/Private/Converters/GLTFConverterUtility.h b/Engine/Plugins/Enterprise/GLTFExporter/MigratedFromGit/Source/GLTFExporter/Private/Converters/GLTFConverterUtility.h
--- a/Engine/Plugins/Enterprise/GLTFExporter/MigratedFromGit/Source/GLTFExporter/Private/Converters/GLTFConverterUtility.h
+++ b/Engine/Plugins/Enterprise/GLTFExporter/MigratedFromGit/Source/GLTFExporter/Private/Converters/GLTFConverterUtility.h
@@ -102,21 +102,21 @@ struct FGLTFConverterUtility
 		}
 	}
 
-	static EGLTFJsonTextureFilter ConvertFilter(TextureFilter Filter)
+	static EGLTFJsonTextureFilter ConvertFilter(TextureFilter Filter, TextureGroup LODGroup)
 	{
 		switch (Filter)
 		{
 			case TextureFilter::TF_Nearest:   return EGLTFJsonTextureFilter::Nearest;
 			case TextureFilter::TF_Bilinear:  return EGLTFJsonTextureFilter::LinearMipmapNearest;
 			case TextureFilter::TF_Trilinear: return EGLTFJsonTextureFilter::LinearMipmapLinear;
-			case TextureFilter::TF_Default:   return EGLTFJsonTextureFilter::None; // TODO: handle this case better
+			case TextureFilter::TF_Default:   return ConvertFilter(GetDefaultFilter(LODGroup));
 			default:                          return EGLTFJsonTextureFilter::None;
 		}
 	}
 
-	static EGLTFJsonTextureWrap ConvertWrap(TextureAddress Wrap)
+	static EGLTFJsonTextureWrap ConvertWrap(TextureAddress Address)
 	{
-		switch (Wrap)
+		switch (Address)
 		{
 			case TextureAddress::TA_Wrap:   return EGLTFJsonTextureWrap::Repeat;
 			case TextureAddress::TA_Mirror: return EGLTFJsonTextureWrap::MirroredRepeat;
@@ -125,6 +125,21 @@ struct FGLTFConverterUtility
 		}
 	}
 
+	static EGLTFJsonTextureFilter ConvertFilter(ETextureSamplerFilter Filter)
+	{
+		switch (Filter)
+		{
+			case ETextureSamplerFilter::Point:             return EGLTFJsonTextureFilter::Nearest;
+			case ETextureSamplerFilter::Bilinear:          return EGLTFJsonTextureFilter::LinearMipmapNearest;
+			case ETextureSamplerFilter::Trilinear:         return EGLTFJsonTextureFilter::LinearMipmapLinear;
+			case ETextureSamplerFilter::AnisotropicPoint:  return EGLTFJsonTextureFilter::Nearest; // TODO: is this correct?
+			case ETextureSamplerFilter::AnisotropicLinear: return EGLTFJsonTextureFilter::LinearMipmapLinear; // TODO: is this correct?
+			default:                                       return EGLTFJsonTextureFilter::None;
+		}
+	}
+
+	static ETextureSamplerFilter GetDefaultFilter(TextureGroup Group);
+
 	static bool IsSkySphereBlueprint(const UBlueprint* Blueprint);
 
 	static bool IsHDRIBackdropBlueprint(const UBlueprint* Blueprint);
diff --git a/Engine/Plugins/Enterprise/GLTFExporter/MigratedFromGit/Source/GLTFExporter/Private/Converters/GLTFTextureConverters.cpp b/Engine/Plugins/Enterprise/GLTFExporter/MigratedFromGit/Source/GLTFExporter/Private/Converters/GLTFTextureConverters.cpp
new file mode 100644
--- /dev/null
+++ b/Engine/Plugins/Enterprise/GLTFExporter/MigratedFromGit/Source/GLTFExporter/Private/Converters/GLTFTextureConverters.cpp
@@ -0,0 +1,46 @@
+// Copyright Epic Games, Inc. All Rights Reserved.
+
+#include "Converters/GLTFTextureConverters.h"
+#include "Builders/GLTFContainerBuilder.h"
+#include "Converters/GLTFConverterUtility.h"
+
+FGLTFJsonTextureIndex FGLTFTexture2DConverter::Add(FGLTFConvertBuilder& Builder, const FString& Name, const UTexture2D* Texture2D)
+{
+	// TODO: maybe we should reuse existing samplers?
+	FGLTFJsonSampler JsonSampler;
+	JsonSampler.Name = Name;
+	JsonSampler.MagFilter = JsonSampler.MinFilter = FGLTFConverterUtility::ConvertFilter(Texture2D->Filter, Texture2D->LODGroup);
+	JsonSampler.WrapS = FGLTFConverterUtility::ConvertWrap(Texture2D->AddressX);
+	JsonSampler.WrapT = FGLTFConverterUtility::ConvertWrap(Texture2D->AddressY);
+
+	FGLTFJsonTexture JsonTexture;
+	JsonTexture.Name = Name;
+	JsonTexture.Sampler = Builder.AddSampler(JsonSampler);
+	JsonTexture.Source = Builder.AddImage(Texture2D->Source, Name);
+
+	return Builder.AddTexture(JsonTexture);
+}
+
+FGLTFJsonTextureIndex FGLTFTextureCubeConverter::Add(FGLTFConvertBuilder& Builder, const FString& Name, const UTextureCube* TextureCube)
+{
+	// TODO: implement support
+	return FGLTFJsonTextureIndex(INDEX_NONE);
+}
+
+FGLTFJsonTextureIndex FGLTFTextureRenderTarget2DConverter::Add(FGLTFConvertBuilder& Builder, const FString& Name, const UTextureRenderTarget2D* RenderTarget2D)
+{
+	// TODO: implement support
+	return FGLTFJsonTextureIndex(INDEX_NONE);
+}
+
+FGLTFJsonTextureIndex FGLTFTextureRenderTargetCubeConverter::Add(FGLTFConvertBuilder& Builder, const FString& Name, const UTextureRenderTargetCube* RenderTargetCube)
+{
+	// TODO: implement support
+	return FGLTFJsonTextureIndex(INDEX_NONE);
+}
+
+FGLTFJsonTextureIndex FGLTFLightMapTexture2DConverter::Add(FGLTFConvertBuilder& Builder, const FString& Name, const ULightMapTexture2D* LightMap)
+{
+	// TODO: implement support
+	return FGLTFJsonTextureIndex(INDEX_NONE);
+}
diff --git a/Engine/Plugins/Enterprise/GLTFExporter/MigratedFromGit/Source/GLTFExporter/Private/Converters/GLTFTextureConverters.h b/Engine/Plugins/Enterprise/GLTFExporter/MigratedFromGit/Source/GLTFExporter/Private/Converters/GLTFTextureConverters.h
new file mode 100644
--- /dev/null
+++ b/Engine/Plugins/Enterprise/GLTFExporter/MigratedFromGit/Source/GLTFExporter/Private/Converters/GLTFTextureConverters.h
@@ -0,0 +1,32 @@
+// Copyright Epic Games, Inc. All Rights Reserved.
+
+#pragma once
+
+#include "Json/GLTFJsonIndex.h"
+#include "Converters/GLTFConverter.h"
+#include "Engine.h"
+
+class FGLTFTexture2DConverter final : public TGLTFConverter<FGLTFJsonTextureIndex, const UTexture2D*>
+{
+	FGLTFJsonTextureIndex Add(FGLTFConvertBuilder& Builder, const FString& Name, const UTexture2D* Texture2D) override;
+};
+
+class FGLTFTextureCubeConverter final : public TGLTFConverter<FGLTFJsonTextureIndex, const UTextureCube*>
+{
+	FGLTFJsonTextureIndex Add(FGLTFConvertBuilder& Builder, const FString& Name, const UTextureCube* TextureCube) override;
+};
+
+class FGLTFTextureRenderTarget2DConverter final : public TGLTFConverter<FGLTFJsonTextureIndex, const UTextureRenderTarget2D*>
+{
+	FGLTFJsonTextureIndex Add(FGLTFConvertBuilder& Builder, const FString& Name, const UTextureRenderTarget2D* RenderTarget2D) override;
+};
+
+class FGLTFTextureRenderTargetCubeConverter final : public TGLTFConverter<FGLTFJsonTextureIndex, const UTextureRenderTargetCube*>
+{
+	FGLTFJsonTextureIndex Add(FGLTFConvertBuilder& Builder, const FString& Name, const UTextureRenderTargetCube* RenderTargetCube) override;
+};
+
+class FGLTFLightMapTexture2DConverter final : public TGLTFConverter<FGLTFJsonTextureIndex, const ULightMapTexture2D*>
+{
+	FGLTFJsonTextureIndex Add(FGLTFConvertBuilder& Builder, const FString& Name, const ULightMapTexture2D* LightMap) override;
+};
