[GLTFExporter] feat(material): add initial limited texture support in material converter

commit dd3ac97cc9e32a275de2315ddbd8da70fc750338
Date: Sun, 24 May 2020 23:53:09 +0200

    feat(material): add initial limited texture support in material converter
    (cherry picked from commit 743f502895521e2c8070bd0109bb4997b8518869​)

#skipci

diff --git a/Engine/Plugins/Enterprise/GLTFExporter/MigratedFromGit/Source/GLTFExporter/Private/Converters/GLTFMaterialConverters.cpp b/Engine/Plugins/Enterprise/GLTFExporter/MigratedFromGit/Source/GLTFExporter/Private/Converters/GLTFMaterialConverters.cpp
--- a/Engine/Plugins/Enterprise/GLTFExporter/MigratedFromGit/Source/GLTFExporter/Private/Converters/GLTFMaterialConverters.cpp
+++ b/Engine/Plugins/Enterprise/GLTFExporter/MigratedFromGit/Source/GLTFExporter/Private/Converters/GLTFMaterialConverters.cpp
@@ -10,53 +10,99 @@
 #include "Materials/MaterialExpressionScalarParameter.h"
 #include "Materials/MaterialExpressionVectorParameter.h"
 #include "Materials/MaterialExpressionTextureSample.h"
+#include "Materials/MaterialExpressionTextureSampleParameter2D.h"
 
 FGLTFJsonMaterialIndex FGLTFMaterialConverter::Add(FGLTFConvertBuilder& Builder, const FString& Name, const UMaterialInterface* MaterialInterface)
 {
+	const UMaterial* Material = MaterialInterface->GetMaterial();
+	const UMaterialInstance* MaterialInstance = Cast<UMaterialInstance>(MaterialInterface);
+
 	FGLTFJsonMaterial JsonMaterial;
 	JsonMaterial.Name = Name;
-
 	// TODO: add support for different shading models (Default Lit, Unlit, Clear Coat)
 
 	// TODO: add support for additional blend modes (at least Additive and Modulate)
 	JsonMaterial.AlphaMode = FGLTFConverterUtility::ConvertBlendMode(MaterialInterface->GetBlendMode());
-
 	JsonMaterial.AlphaCutoff = MaterialInterface->GetOpacityMaskClipValue();
 	JsonMaterial.DoubleSided = MaterialInterface->IsTwoSided();
 
-	const UMaterial* Material = MaterialInterface->GetMaterial();
-	const UMaterialInstance* MaterialInstance = Cast<UMaterialInstance>(MaterialInterface);
+	if (JsonMaterial.AlphaMode == EGLTFJsonAlphaMode::Blend || JsonMaterial.AlphaMode == EGLTFJsonAlphaMode::Mask)
+	{
+		const FScalarMaterialInput& OpacityInput = JsonMaterial.AlphaMode == EGLTFJsonAlphaMode::Blend ? Material->Opacity : Material->OpacityMask;
+		if (!TryGetSimpleBaseColorOpacity(Builder, JsonMaterial, Material->BaseColor, OpacityInput, MaterialInstance))
+		{
+			// TODO: add fallback to material baking
+		}
+	}
+	else
+	{
+		if (!TryGetConstantColor(JsonMaterial.PBRMetallicRoughness.BaseColorFactor, Material->BaseColor, MaterialInstance))
+		{
+			if (!TryGetSourceTexture(Builder, JsonMaterial.PBRMetallicRoughness.BaseColorTexture, Material->BaseColor, MaterialInstance))
+			{
+				// TODO: add fallback to material baking
+			}
+		}
+	}
 
-	if (!TryGetConstantColor(JsonMaterial.PBRMetallicRoughness.BaseColorFactor, Material->BaseColor, MaterialInstance))
+	if (!TryGetSimpleMetallicRoughness(Builder, JsonMaterial.PBRMetallicRoughness, Material->Metallic, Material->Roughness, MaterialInstance))
 	{
-		// TODO: temporary fallback until texture and baking support
+		// TODO: add fallback to material baking
 	}
 
-	if (JsonMaterial.AlphaMode == EGLTFJsonAlphaMode::Blend || JsonMaterial.AlphaMode == EGLTFJsonAlphaMode::Mask)
+	if (!TryGetConstantColor(JsonMaterial.EmissiveFactor, Material->EmissiveColor, MaterialInstance))
 	{
-		const FScalarMaterialInput& OpacityInput = JsonMaterial.AlphaMode == EGLTFJsonAlphaMode::Blend ? Material->Opacity : Material->OpacityMask;
-		if (!TryGetConstantScalar(JsonMaterial.PBRMetallicRoughness.BaseColorFactor.A, OpacityInput, MaterialInstance))
+		if (!TryGetSourceTexture(Builder, JsonMaterial.EmissiveTexture, Material->EmissiveColor, MaterialInstance))
 		{
-			// TODO: temporary fallback until texture and baking support
+			// TODO: add fallback to material baking
 		}
 	}
 
-	if (!TryGetConstantScalar(JsonMaterial.PBRMetallicRoughness.MetallicFactor, Material->Metallic, MaterialInstance))
+	if (Material->Normal.Expression != nullptr)
 	{
-		// TODO: temporary fallback until texture and baking support
+		if (!TryGetSourceTexture(Builder, JsonMaterial.NormalTexture, Material->Normal, MaterialInstance))
+		{
+			// TODO: add fallback to material baking
+		}
 	}
 
-	if (!TryGetConstantScalar(JsonMaterial.PBRMetallicRoughness.RoughnessFactor, Material->Roughness, MaterialInstance))
+	if (Material->AmbientOcclusion.Expression != nullptr)
 	{
-		// TODO: temporary fallback until texture and baking support
+		if (!TryGetSourceTexture(Builder, JsonMaterial.OcclusionTexture, Material->AmbientOcclusion, MaterialInstance))
+		{
+			// TODO: add fallback to material baking
+		}
 	}
 
-	if (!TryGetConstantColor(JsonMaterial.EmissiveFactor, Material->EmissiveColor, MaterialInstance))
+	return Builder.AddMaterial(JsonMaterial);
+}
+
+bool FGLTFMaterialConverter::TryGetSimpleBaseColorOpacity(FGLTFConvertBuilder& Builder, FGLTFJsonMaterial& OutValue, const FColorMaterialInput& BaseColorInput, const FScalarMaterialInput& OpacityInput, const UMaterialInstance* MaterialInstance) const
+{
+	const bool bIsBaseColorConstant = TryGetConstantColor(OutValue.PBRMetallicRoughness.BaseColorFactor, BaseColorInput, MaterialInstance);
+	const bool bIsOpacityConstant = TryGetConstantScalar(OutValue.PBRMetallicRoughness.BaseColorFactor.A, OpacityInput, MaterialInstance);
+
+	if (bIsBaseColorConstant && bIsOpacityConstant)
 	{
-		// TODO: temporary fallback until texture and baking support
+		return true;
 	}
 
-	return Builder.AddMaterial(JsonMaterial);
+	// TODO: add support for other combos
+	return false;
+}
+
+bool FGLTFMaterialConverter::TryGetSimpleMetallicRoughness(FGLTFConvertBuilder& Builder, FGLTFJsonPBRMetallicRoughness& OutValue, const FScalarMaterialInput& MetallicInput, const FScalarMaterialInput& RoughnessInput, const UMaterialInstance* MaterialInstance) const
+{
+	const bool bIsMetallicConstant = TryGetConstantScalar(OutValue.MetallicFactor, MetallicInput, MaterialInstance);
+	const bool bIsRoughnessConstant = TryGetConstantScalar(OutValue.RoughnessFactor, RoughnessInput, MaterialInstance);
+
+	if (bIsMetallicConstant && bIsRoughnessConstant)
+	{
+		return true;
+	}
+
+	// TODO: add support for other combos
+	return false;
 }
 
 bool FGLTFMaterialConverter::TryGetConstantColor(FGLTFJsonColor3& OutValue, const FColorMaterialInput& MaterialInput, const UMaterialInstance* MaterialInstance) const
@@ -232,3 +278,89 @@ bool FGLTFMaterialConverter::TryGetConstantScalar(float& OutValue, const FScalar
 
 	return false;
 }
+
+bool FGLTFMaterialConverter::TryGetSourceTexture(FGLTFConvertBuilder& Builder, FGLTFJsonTextureInfo& OutValue, const FExpressionInput& MaterialInput, const UMaterialInstance* MaterialInstance) const
+{
+	const UTexture2D* Texture;
+	int32 TexCoord;
+	int32 Mask;
+
+	if (TryGetSourceTexture(Texture, TexCoord, Mask, MaterialInput, MaterialInstance))
+	{
+		// TODO: add support for output mask
+		OutValue.Index = Builder.GetOrAddTexture(Texture);
+		OutValue.TexCoord = TexCoord;
+		return true;
+	}
+
+	return false;
+}
+
+bool FGLTFMaterialConverter::TryGetSourceTexture(const UTexture2D*& OutTexture, int32& OutTexCoord, int32& OutMask, const FExpressionInput& MaterialInput, const UMaterialInstance* MaterialInstance) const
+{
+	const UMaterialExpression* Expression = MaterialInput.Expression;
+	if (Expression == nullptr)
+	{
+		return false;
+	}
+
+	if (const UMaterialExpressionTextureSampleParameter2D* TextureParameter = ExactCast<UMaterialExpressionTextureSampleParameter2D>(Expression))
+	{
+		if (MaterialInstance != nullptr)
+		{
+			UTexture* ParameterValue;
+
+			const FHashedMaterialParameterInfo ParameterInfo(TextureParameter->GetParameterName());
+			if (!MaterialInstance->GetTextureParameterValue(ParameterInfo, ParameterValue))
+			{
+				// TODO: can this happen?
+			}
+
+			OutTexture = Cast<UTexture2D>(ParameterValue);
+		}
+
+		if (OutTexture == nullptr)
+		{
+			OutTexture = Cast<UTexture2D>(TextureParameter->Texture);
+		}
+
+		if (OutTexture == nullptr)
+		{
+			// TODO: report material as broken
+			return false;
+		}
+
+		// TODO: add support for texture coordinate input expression
+		OutTexCoord = TextureParameter->ConstCoordinate;
+
+		// TODO: add support for output mask
+		return true;
+	}
+
+	if (const UMaterialExpressionTextureSample* TextureSample = ExactCast<UMaterialExpressionTextureSample>(Expression))
+	{
+		// TODO: add support for texture object input expression
+		OutTexture = Cast<UTexture2D>(TextureSample->Texture);
+
+		if (OutTexture == nullptr)
+		{
+			if (TextureSample->Texture == nullptr)
+			{
+				// TODO: report material as broken
+			}
+			else
+			{
+				// TODO: report incorrect texture type
+			}
+			return false;
+		}
+
+		// TODO: add support for texture coordinate input expression
+		OutTexCoord = TextureSample->ConstCoordinate;
+
+		// TODO: add support for output mask
+		return true;
+	}
+
+	return false;
+}
diff --git a/Engine/Plugins/Enterprise/GLTFExporter/MigratedFromGit/Source/GLTFExporter/Private/Converters/GLTFMaterialConverters.h b/Engine/Plugins/Enterprise/GLTFExporter/MigratedFromGit/Source/GLTFExporter/Private/Converters/GLTFMaterialConverters.h
--- a/Engine/Plugins/Enterprise/GLTFExporter/MigratedFromGit/Source/GLTFExporter/Private/Converters/GLTFMaterialConverters.h
+++ b/Engine/Plugins/Enterprise/GLTFExporter/MigratedFromGit/Source/GLTFExporter/Private/Converters/GLTFMaterialConverters.h
@@ -3,8 +3,7 @@
 #pragma once
 
 #include "Json/GLTFJsonIndex.h"
-#include "Json/GLTFJsonColor3.h"
-#include "Json/GLTFJsonColor4.h"
+#include "Json/GLTFJsonMaterial.h"
 #include "Converters/GLTFConverter.h"
 #include "Engine.h"
 
@@ -12,8 +11,14 @@ class FGLTFMaterialConverter final : public TGLTFConverter<FGLTFJsonMaterialInde
 {
 	FGLTFJsonMaterialIndex Add(FGLTFConvertBuilder& Builder, const FString& Name, const UMaterialInterface* Material) override;
 
+	bool TryGetSimpleBaseColorOpacity(FGLTFConvertBuilder& Builder, FGLTFJsonMaterial& OutValue, const FColorMaterialInput& BaseColorInput, const FScalarMaterialInput& OpacityInput, const UMaterialInstance* MaterialInstance) const;
+	bool TryGetSimpleMetallicRoughness(FGLTFConvertBuilder& Builder, FGLTFJsonPBRMetallicRoughness& OutValue, const FScalarMaterialInput& MetallicInput, const FScalarMaterialInput& RoughnessInput, const UMaterialInstance* MaterialInstance) const;
+
 	bool TryGetConstantColor(FGLTFJsonColor3& OutValue, const FColorMaterialInput& MaterialInput, const UMaterialInstance* MaterialInstance) const;
 	bool TryGetConstantColor(FGLTFJsonColor4& OutValue, const FColorMaterialInput& MaterialInput, const UMaterialInstance* MaterialInstance) const;
 	bool TryGetConstantColor(FLinearColor& OutValue, const FColorMaterialInput& MaterialInput, const UMaterialInstance* MaterialInstance) const;
 	bool TryGetConstantScalar(float& OutValue, const FScalarMaterialInput& MaterialInput, const UMaterialInstance* MaterialInstance) const;
+
+	bool TryGetSourceTexture(FGLTFConvertBuilder& Builder, FGLTFJsonTextureInfo& OutValue, const FExpressionInput& MaterialInput, const UMaterialInstance* MaterialInstance) const;
+	bool TryGetSourceTexture(const UTexture2D*& OutTexture, int32& OutTexCoord, int32& OutMask, const FExpressionInput& MaterialInput, const UMaterialInstance* MaterialInstance) const;
 };
