[GLTFExporter] feat(material): add support for single texture used for combined params

commit 3d75516c4191d82371f3d0bc479a0509035350b3
Date: Wed, 27 May 2020 23:28:14 +0200

    feat(material): add support for single texture used for combined params
    (cherry picked from commit 3b5e6e8f96087f3a1d58de47df91d4abeca85f82​)

#skipci

diff --git a/Engine/Plugins/Enterprise/GLTFExporter/MigratedFromGit/Source/GLTFExporter/Private/Converters/GLTFMaterialConverters.cpp b/Engine/Plugins/Enterprise/GLTFExporter/MigratedFromGit/Source/GLTFExporter/Private/Converters/GLTFMaterialConverters.cpp
--- a/Engine/Plugins/Enterprise/GLTFExporter/MigratedFromGit/Source/GLTFExporter/Private/Converters/GLTFMaterialConverters.cpp
+++ b/Engine/Plugins/Enterprise/GLTFExporter/MigratedFromGit/Source/GLTFExporter/Private/Converters/GLTFMaterialConverters.cpp
@@ -29,7 +29,7 @@ FGLTFJsonMaterialIndex FGLTFMaterialConverter::Add(FGLTFConvertBuilder& Builder,
 	if (JsonMaterial.AlphaMode == EGLTFJsonAlphaMode::Blend || JsonMaterial.AlphaMode == EGLTFJsonAlphaMode::Mask)
 	{
 		const FScalarMaterialInput& OpacityInput = JsonMaterial.AlphaMode == EGLTFJsonAlphaMode::Blend ? Material->Opacity : Material->OpacityMask;
-		if (!TryGetSimpleBaseColorOpacity(Builder, JsonMaterial, Material->BaseColor, OpacityInput, MaterialInstance))
+		if (!TryGetBaseColorAndOpacity(Builder, JsonMaterial.PBRMetallicRoughness, Material->BaseColor, OpacityInput, MaterialInstance))
 		{
 			// TODO: add fallback to material baking
 		}
@@ -45,7 +45,7 @@ FGLTFJsonMaterialIndex FGLTFMaterialConverter::Add(FGLTFConvertBuilder& Builder,
 		}
 	}
 
-	if (!TryGetSimpleMetallicRoughness(Builder, JsonMaterial.PBRMetallicRoughness, Material->Metallic, Material->Roughness, MaterialInstance))
+	if (!TryGetMetallicAndRoughness(Builder, JsonMaterial.PBRMetallicRoughness, Material->Metallic, Material->Roughness, MaterialInstance))
 	{
 		// TODO: add fallback to material baking
 	}
@@ -56,6 +56,10 @@ FGLTFJsonMaterialIndex FGLTFMaterialConverter::Add(FGLTFConvertBuilder& Builder,
 		{
 			// TODO: add fallback to material baking
 		}
+		else
+		{
+			JsonMaterial.EmissiveFactor = FGLTFJsonColor3::White; // make sure texture is not multiplied with black
+		}
 	}
 
 	if (Material->Normal.Expression != nullptr)
@@ -77,31 +81,97 @@ FGLTFJsonMaterialIndex FGLTFMaterialConverter::Add(FGLTFConvertBuilder& Builder,
 	return Builder.AddMaterial(JsonMaterial);
 }
 
-bool FGLTFMaterialConverter::TryGetSimpleBaseColorOpacity(FGLTFConvertBuilder& Builder, FGLTFJsonMaterial& OutValue, const FColorMaterialInput& BaseColorInput, const FScalarMaterialInput& OpacityInput, const UMaterialInstance* MaterialInstance) const
+bool FGLTFMaterialConverter::TryGetBaseColorAndOpacity(FGLTFConvertBuilder& Builder, FGLTFJsonPBRMetallicRoughness& OutPBRParams, const FColorMaterialInput& BaseColorInput, const FScalarMaterialInput& OpacityInput, const UMaterialInstance* MaterialInstance) const
 {
-	const bool bIsBaseColorConstant = TryGetConstantColor(OutValue.PBRMetallicRoughness.BaseColorFactor, BaseColorInput, MaterialInstance);
-	const bool bIsOpacityConstant = TryGetConstantScalar(OutValue.PBRMetallicRoughness.BaseColorFactor.A, OpacityInput, MaterialInstance);
+	const bool bIsBaseColorConstant = TryGetConstantColor(OutPBRParams.BaseColorFactor, BaseColorInput, MaterialInstance);
+	const bool bIsOpacityConstant = TryGetConstantScalar(OutPBRParams.BaseColorFactor.A, OpacityInput, MaterialInstance);
 
 	if (bIsBaseColorConstant && bIsOpacityConstant)
 	{
 		return true;
 	}
 
-	// TODO: add support for other combos
+	const UTexture2D* BaseColorTexture;
+	const UTexture2D* OpacityTexture;
+	int32 BaseColorTexCoord;
+	int32 OpacityTexCoord;
+
+	const bool bHasBaseColorSourceTexture = TryGetSourceTexture(BaseColorTexture, BaseColorTexCoord, BaseColorInput, MaterialInstance);
+	const bool bHasOpacitySourceTexture = TryGetSourceTexture(OpacityTexture, OpacityTexCoord, OpacityInput, MaterialInstance);
+
+	if (bHasBaseColorSourceTexture && bHasOpacitySourceTexture)
+	{
+		if (BaseColorTexture == OpacityTexture && BaseColorTexCoord == OpacityTexCoord)
+		{
+			// TODO: make sure textures are correctly masked
+			OutPBRParams.BaseColorTexture.Index = Builder.GetOrAddTexture(BaseColorTexture);
+			OutPBRParams.BaseColorTexture.TexCoord = BaseColorTexCoord;
+			return true;
+		}
+
+		// TODO: add support for combining two textures
+		return false;
+	}
+
+	if (bHasBaseColorSourceTexture && bIsOpacityConstant)
+	{
+		// TODO: add support for combining constant with texture
+		return false;
+	}
+
+	if (bIsBaseColorConstant && bHasOpacitySourceTexture)
+	{
+		// TODO: add support for combining constant with texture
+		return false;
+	}
+
 	return false;
 }
 
-bool FGLTFMaterialConverter::TryGetSimpleMetallicRoughness(FGLTFConvertBuilder& Builder, FGLTFJsonPBRMetallicRoughness& OutValue, const FScalarMaterialInput& MetallicInput, const FScalarMaterialInput& RoughnessInput, const UMaterialInstance* MaterialInstance) const
+bool FGLTFMaterialConverter::TryGetMetallicAndRoughness(FGLTFConvertBuilder& Builder, FGLTFJsonPBRMetallicRoughness& OutPBRParams, const FScalarMaterialInput& MetallicInput, const FScalarMaterialInput& RoughnessInput, const UMaterialInstance* MaterialInstance) const
 {
-	const bool bIsMetallicConstant = TryGetConstantScalar(OutValue.MetallicFactor, MetallicInput, MaterialInstance);
-	const bool bIsRoughnessConstant = TryGetConstantScalar(OutValue.RoughnessFactor, RoughnessInput, MaterialInstance);
+	const bool bIsMetallicConstant = TryGetConstantScalar(OutPBRParams.MetallicFactor, MetallicInput, MaterialInstance);
+	const bool bIsRoughnessConstant = TryGetConstantScalar(OutPBRParams.RoughnessFactor, RoughnessInput, MaterialInstance);
 
 	if (bIsMetallicConstant && bIsRoughnessConstant)
 	{
 		return true;
 	}
 
-	// TODO: add support for other combos
+	const UTexture2D* MetallicTexture;
+	const UTexture2D* RoughnessTexture;
+	int32 MetallicTexCoord;
+	int32 RoughnessTexCoord;
+
+	const bool bHasMetallicSourceTexture = TryGetSourceTexture(MetallicTexture, MetallicTexCoord, MetallicInput, MaterialInstance);
+	const bool bHasRoughnessSourceTexture = TryGetSourceTexture(RoughnessTexture, RoughnessTexCoord, RoughnessInput, MaterialInstance);
+
+	if (bHasMetallicSourceTexture && bHasRoughnessSourceTexture)
+	{
+		if (MetallicTexture == RoughnessTexture && MetallicTexCoord == RoughnessTexCoord)
+		{
+			// TODO: make sure textures are correctly masked
+			OutPBRParams.MetallicRoughnessTexture.Index = Builder.GetOrAddTexture(MetallicTexture);
+			OutPBRParams.MetallicRoughnessTexture.TexCoord = MetallicTexCoord;
+			return true;
+		}
+
+		// TODO: add support for combining two textures
+		return false;
+	}
+
+	if (bHasMetallicSourceTexture && bIsRoughnessConstant)
+	{
+		// TODO: add support for combining constant with texture
+		return false;
+	}
+
+	if (bIsMetallicConstant && bHasRoughnessSourceTexture)
+	{
+		// TODO: add support for combining constant with texture
+		return false;
+	}
+
 	return false;
 }
 
@@ -279,24 +349,23 @@ bool FGLTFMaterialConverter::TryGetConstantScalar(float& OutValue, const FScalar
 	return false;
 }
 
-bool FGLTFMaterialConverter::TryGetSourceTexture(FGLTFConvertBuilder& Builder, FGLTFJsonTextureInfo& OutValue, const FExpressionInput& MaterialInput, const UMaterialInstance* MaterialInstance) const
+bool FGLTFMaterialConverter::TryGetSourceTexture(FGLTFConvertBuilder& Builder, FGLTFJsonTextureInfo& OutTexInfo, const FExpressionInput& MaterialInput, const UMaterialInstance* MaterialInstance) const
 {
 	const UTexture2D* Texture;
 	int32 TexCoord;
-	int32 Mask;
 
-	if (TryGetSourceTexture(Texture, TexCoord, Mask, MaterialInput, MaterialInstance))
+	if (TryGetSourceTexture(Texture, TexCoord, MaterialInput, MaterialInstance))
 	{
 		// TODO: add support for output mask
-		OutValue.Index = Builder.GetOrAddTexture(Texture);
-		OutValue.TexCoord = TexCoord;
+		OutTexInfo.Index = Builder.GetOrAddTexture(Texture);
+		OutTexInfo.TexCoord = TexCoord;
 		return true;
 	}
 
 	return false;
 }
 
-bool FGLTFMaterialConverter::TryGetSourceTexture(const UTexture2D*& OutTexture, int32& OutTexCoord, int32& OutMask, const FExpressionInput& MaterialInput, const UMaterialInstance* MaterialInstance) const
+bool FGLTFMaterialConverter::TryGetSourceTexture(const UTexture2D*& OutTexture, int32& OutTexCoord, const FExpressionInput& MaterialInput, const UMaterialInstance* MaterialInstance) const
 {
 	const UMaterialExpression* Expression = MaterialInput.Expression;
 	if (Expression == nullptr)
diff --git a/Engine/Plugins/Enterprise/GLTFExporter/MigratedFromGit/Source/GLTFExporter/Private/Converters/GLTFMaterialConverters.h b/Engine/Plugins/Enterprise/GLTFExporter/MigratedFromGit/Source/GLTFExporter/Private/Converters/GLTFMaterialConverters.h
--- a/Engine/Plugins/Enterprise/GLTFExporter/MigratedFromGit/Source/GLTFExporter/Private/Converters/GLTFMaterialConverters.h
+++ b/Engine/Plugins/Enterprise/GLTFExporter/MigratedFromGit/Source/GLTFExporter/Private/Converters/GLTFMaterialConverters.h
@@ -11,14 +11,14 @@ class FGLTFMaterialConverter final : public TGLTFConverter<FGLTFJsonMaterialInde
 {
 	FGLTFJsonMaterialIndex Add(FGLTFConvertBuilder& Builder, const FString& Name, const UMaterialInterface* Material) override;
 
-	bool TryGetSimpleBaseColorOpacity(FGLTFConvertBuilder& Builder, FGLTFJsonMaterial& OutValue, const FColorMaterialInput& BaseColorInput, const FScalarMaterialInput& OpacityInput, const UMaterialInstance* MaterialInstance) const;
-	bool TryGetSimpleMetallicRoughness(FGLTFConvertBuilder& Builder, FGLTFJsonPBRMetallicRoughness& OutValue, const FScalarMaterialInput& MetallicInput, const FScalarMaterialInput& RoughnessInput, const UMaterialInstance* MaterialInstance) const;
+	bool TryGetBaseColorAndOpacity(FGLTFConvertBuilder& Builder, FGLTFJsonPBRMetallicRoughness& OutPBRParams, const FColorMaterialInput& BaseColorInput, const FScalarMaterialInput& OpacityInput, const UMaterialInstance* MaterialInstance) const;
+	bool TryGetMetallicAndRoughness(FGLTFConvertBuilder& Builder, FGLTFJsonPBRMetallicRoughness& OutPBRParams, const FScalarMaterialInput& MetallicInput, const FScalarMaterialInput& RoughnessInput, const UMaterialInstance* MaterialInstance) const;
 
 	bool TryGetConstantColor(FGLTFJsonColor3& OutValue, const FColorMaterialInput& MaterialInput, const UMaterialInstance* MaterialInstance) const;
 	bool TryGetConstantColor(FGLTFJsonColor4& OutValue, const FColorMaterialInput& MaterialInput, const UMaterialInstance* MaterialInstance) const;
 	bool TryGetConstantColor(FLinearColor& OutValue, const FColorMaterialInput& MaterialInput, const UMaterialInstance* MaterialInstance) const;
 	bool TryGetConstantScalar(float& OutValue, const FScalarMaterialInput& MaterialInput, const UMaterialInstance* MaterialInstance) const;
 
-	bool TryGetSourceTexture(FGLTFConvertBuilder& Builder, FGLTFJsonTextureInfo& OutValue, const FExpressionInput& MaterialInput, const UMaterialInstance* MaterialInstance) const;
-	bool TryGetSourceTexture(const UTexture2D*& OutTexture, int32& OutTexCoord, int32& OutMask, const FExpressionInput& MaterialInput, const UMaterialInstance* MaterialInstance) const;
+	bool TryGetSourceTexture(FGLTFConvertBuilder& Builder, FGLTFJsonTextureInfo& OutTexInfo, const FExpressionInput& MaterialInput, const UMaterialInstance* MaterialInstance) const;
+	bool TryGetSourceTexture(const UTexture2D*& OutTexture, int32& OutTexCoord, const FExpressionInput& MaterialInput, const UMaterialInstance* MaterialInstance) const;
 };
