[GLTFExporter] feat(material): add builder function for adding image by pixel format

commit 9b4fd7a849a56d055ea5d1a550931a57d4ceda4c
Date: Fri, 29 May 2020 23:29:11 +0200

    feat(material): add builder function for adding image by pixel format
    (cherry picked from commit 2a735ae1627a1eed4b9775a8fe2b21413c4e8e2d​)

#skipci

diff --git a/Engine/Plugins/Enterprise/GLTFExporter/MigratedFromGit/Source/GLTFExporter/Private/Builders/GLTFBuilderUtility.cpp b/Engine/Plugins/Enterprise/GLTFExporter/MigratedFromGit/Source/GLTFExporter/Private/Builders/GLTFBuilderUtility.cpp
--- a/Engine/Plugins/Enterprise/GLTFExporter/MigratedFromGit/Source/GLTFExporter/Private/Builders/GLTFBuilderUtility.cpp
+++ b/Engine/Plugins/Enterprise/GLTFExporter/MigratedFromGit/Source/GLTFExporter/Private/Builders/GLTFBuilderUtility.cpp
@@ -4,7 +4,7 @@
 #include "IImageWrapperModule.h"
 #include "IImageWrapper.h"
 
-bool FGLTFBuilderUtility::CompressImage(const TArray64<uint8>& InRawData, int32 InWidth, int32 InHeight, ERGBFormat InRawFormat, int32 InBitDepth, TArray64<uint8>& OutCompressedData, EImageFormat OutCompressionFormat, int32 OutCompressionQuality)
+bool FGLTFBuilderUtility::CompressImage(const void* RawData, int64 ByteLength, int32 InWidth, int32 InHeight, ERGBFormat InRawFormat, int32 InBitDepth, TArray64<uint8>& OutCompressedData, EImageFormat OutCompressionFormat, int32 OutCompressionQuality)
 {
 	IImageWrapperModule& ImageWrapperModule = FModuleManager::Get().LoadModuleChecked<IImageWrapperModule>(TEXT("ImageWrapper"));
 
@@ -15,7 +15,7 @@ bool FGLTFBuilderUtility::CompressImage(const TArray64<uint8>& InRawData, int32
 		return false;
 	}
 
-	const bool bFormatSupported = ImageWrapper->SetRaw(InRawData.GetData(), InRawData.Num(), InWidth, InHeight, InRawFormat, InBitDepth);
+	const bool bFormatSupported = ImageWrapper->SetRaw(RawData, ByteLength, InWidth, InHeight, InRawFormat, InBitDepth);
 	if (!bFormatSupported)
 	{
 		// TODO: report error
diff --git a/Engine/Plugins/Enterprise/GLTFExporter/MigratedFromGit/Source/GLTFExporter/Private/Builders/GLTFBuilderUtility.h b/Engine/Plugins/Enterprise/GLTFExporter/MigratedFromGit/Source/GLTFExporter/Private/Builders/GLTFBuilderUtility.h
--- a/Engine/Plugins/Enterprise/GLTFExporter/MigratedFromGit/Source/GLTFExporter/Private/Builders/GLTFBuilderUtility.h
+++ b/Engine/Plugins/Enterprise/GLTFExporter/MigratedFromGit/Source/GLTFExporter/Private/Builders/GLTFBuilderUtility.h
@@ -23,7 +23,7 @@ struct FGLTFBuilderUtility
 		return Name;
 	}
 
-	static bool CompressImage(const TArray64<uint8>& InRawData, int32 InWidth, int32 InHeight, ERGBFormat InRawFormat, int32 InBitDepth, TArray64<uint8>& OutCompressedData, EImageFormat OutCompressionFormat, int32 OutCompressionQuality);
+	static bool CompressImage(const void* RawData, int64 ByteLength, int32 InWidth, int32 InHeight, ERGBFormat InRawFormat, int32 InBitDepth, TArray64<uint8>& OutCompressedData, EImageFormat OutCompressionFormat, int32 OutCompressionQuality);
 
 	static const TCHAR* GetFileExtension(EGLTFJsonMimeType MimeType)
 	{
diff --git a/Engine/Plugins/Enterprise/GLTFExporter/MigratedFromGit/Source/GLTFExporter/Private/Builders/GLTFImageBuilder.cpp b/Engine/Plugins/Enterprise/GLTFExporter/MigratedFromGit/Source/GLTFExporter/Private/Builders/GLTFImageBuilder.cpp
--- a/Engine/Plugins/Enterprise/GLTFExporter/MigratedFromGit/Source/GLTFExporter/Private/Builders/GLTFImageBuilder.cpp
+++ b/Engine/Plugins/Enterprise/GLTFExporter/MigratedFromGit/Source/GLTFExporter/Private/Builders/GLTFImageBuilder.cpp
@@ -5,7 +5,7 @@
 #include "IImageWrapper.h"
 #include "Engine/Texture.h"
 
-FGLTFJsonImageIndex FGLTFImageBuilder::AddImage(const TArray64<uint8>& RawData, int32 Width, int32 Height, ERGBFormat RawFormat, int32 BitDepth, bool bFloatFormat, const FString& Name, EGLTFJsonMimeType MimeType, int32 Quality)
+FGLTFJsonImageIndex FGLTFImageBuilder::AddImage(const void* RawData, int64 ByteLength, int32 Width, int32 Height, ERGBFormat RawFormat, int32 BitDepth, bool bFloatFormat, const FString& Name, EGLTFJsonMimeType MimeType, int32 Quality)
 {
 	TArray64<uint8> ImageData;
 
@@ -19,20 +19,20 @@ FGLTFJsonImageIndex FGLTFImageBuilder::AddImage(const TArray64<uint8>& RawData,
 		switch (MimeType)
 		{
 		case EGLTFJsonMimeType::PNG:
-            if (!FGLTFBuilderUtility::CompressImage(RawData, Width, Height, RawFormat, BitDepth, ImageData, EImageFormat::PNG, Quality))
-            {
-                return FGLTFJsonImageIndex(INDEX_NONE);
-            }
+			if (!FGLTFBuilderUtility::CompressImage(RawData, ByteLength, Width, Height, RawFormat, BitDepth, ImageData, EImageFormat::PNG, Quality))
+			{
+				return FGLTFJsonImageIndex(INDEX_NONE);
+			}
 			break;
 		case EGLTFJsonMimeType::JPEG:
-            if (!FGLTFBuilderUtility::CompressImage(RawData, Width, Height, RawFormat, BitDepth, ImageData, EImageFormat::JPEG, Quality))
-            {
-                return FGLTFJsonImageIndex(INDEX_NONE);
-            }
+			if (!FGLTFBuilderUtility::CompressImage(RawData, ByteLength, Width, Height, RawFormat, BitDepth, ImageData, EImageFormat::JPEG, Quality))
+			{
+				return FGLTFJsonImageIndex(INDEX_NONE);
+			}
 			break;
 		default:
-            // TODO: report error
-            return FGLTFJsonImageIndex(INDEX_NONE);
+			// TODO: report error
+			return FGLTFJsonImageIndex(INDEX_NONE);
 		}
 	}
 
@@ -45,6 +45,25 @@ FGLTFJsonImageIndex FGLTFImageBuilder::AddImage(const TArray64<uint8>& RawData,
 	return ImageIndex;
 }
 
+FGLTFJsonImageIndex FGLTFImageBuilder::AddImage(const TArray<FColor>& Pixels, int32 Width, int32 Height, EPixelFormat PixelFormat, const FString& Name, EGLTFJsonMimeType MimeType, int32 Quality)
+{
+	ERGBFormat RawFormat;
+	int32 BitDepth;
+	bool bFloatFormat;
+
+	switch (PixelFormat)
+	{
+		case PF_B8G8R8A8: RawFormat = ERGBFormat::BGRA; BitDepth = 8;  bFloatFormat = false; break;
+		case PF_R8G8B8A8: RawFormat = ERGBFormat::RGBA; BitDepth = 8;  bFloatFormat = false; break;
+		case PF_G8:       RawFormat = ERGBFormat::Gray; BitDepth = 8;  bFloatFormat = false; break;
+		default:
+			// TODO: report error
+			return FGLTFJsonImageIndex(INDEX_NONE);
+	}
+
+	return AddImage(Pixels.GetData(), Pixels.Num() * Pixels.GetTypeSize(), Width, Height, RawFormat, BitDepth, bFloatFormat, Name, MimeType, Quality);
+}
+
 FGLTFJsonImageIndex FGLTFImageBuilder::AddImage(const FTextureSource& Image, const FString& Name, EGLTFJsonMimeType MimeType, int32 Quality)
 {
 	// TODO: are these always zero?
@@ -59,14 +78,14 @@ FGLTFJsonImageIndex FGLTFImageBuilder::AddImage(const FTextureSource& Image, con
 	const ETextureSourceFormat SourceFormat = Image.GetFormat(LayerIndex);
 	switch (SourceFormat)
 	{
-        case TSF_BGRA8:   RawFormat = ERGBFormat::BGRA; BitDepth = 8;  bFloatFormat = false; break;
-        case TSF_RGBA8:   RawFormat = ERGBFormat::RGBA; BitDepth = 8;  bFloatFormat = false; break;
+		case TSF_BGRA8:   RawFormat = ERGBFormat::BGRA; BitDepth = 8;  bFloatFormat = false; break;
+		case TSF_RGBA8:   RawFormat = ERGBFormat::RGBA; BitDepth = 8;  bFloatFormat = false; break;
 		case TSF_BGRE8:   RawFormat = ERGBFormat::BGRA; BitDepth = 8;  bFloatFormat = false; break;
-        case TSF_RGBE8:   RawFormat = ERGBFormat::RGBA; BitDepth = 8;  bFloatFormat = false; break;
-        case TSF_RGBA16:  RawFormat = ERGBFormat::RGBA; BitDepth = 16; bFloatFormat = false; break;
-        case TSF_RGBA16F: RawFormat = ERGBFormat::RGBA; BitDepth = 16; bFloatFormat = true;  break;
-        case TSF_G8:      RawFormat = ERGBFormat::Gray; BitDepth = 8;  bFloatFormat = false; break;
-        case TSF_G16:     RawFormat = ERGBFormat::Gray; BitDepth = 16; bFloatFormat = false; break;
+		case TSF_RGBE8:   RawFormat = ERGBFormat::RGBA; BitDepth = 8;  bFloatFormat = false; break;
+		case TSF_RGBA16:  RawFormat = ERGBFormat::RGBA; BitDepth = 16; bFloatFormat = false; break;
+		case TSF_RGBA16F: RawFormat = ERGBFormat::RGBA; BitDepth = 16; bFloatFormat = true;  break;
+		case TSF_G8:      RawFormat = ERGBFormat::Gray; BitDepth = 8;  bFloatFormat = false; break;
+		case TSF_G16:     RawFormat = ERGBFormat::Gray; BitDepth = 16; bFloatFormat = false; break;
 		default:
 			// TODO: report error
 			return FGLTFJsonImageIndex(INDEX_NONE);
@@ -79,7 +98,7 @@ FGLTFJsonImageIndex FGLTFImageBuilder::AddImage(const FTextureSource& Image, con
 		return FGLTFJsonImageIndex(INDEX_NONE);
 	}
 
-	return AddImage(RawData, Image.GetSizeX(), Image.GetSizeY(), RawFormat, BitDepth, bFloatFormat, Name, MimeType, Quality);
+	return AddImage(RawData.GetData(), RawData.Num(), Image.GetSizeX(), Image.GetSizeY(), RawFormat, BitDepth, bFloatFormat, Name, MimeType, Quality);
 }
 
 bool FGLTFImageBuilder::Serialize(FArchive& Archive, const FString& FilePath)
diff --git a/Engine/Plugins/Enterprise/GLTFExporter/MigratedFromGit/Source/GLTFExporter/Private/Builders/GLTFImageBuilder.h b/Engine/Plugins/Enterprise/GLTFExporter/MigratedFromGit/Source/GLTFExporter/Private/Builders/GLTFImageBuilder.h
--- a/Engine/Plugins/Enterprise/GLTFExporter/MigratedFromGit/Source/GLTFExporter/Private/Builders/GLTFImageBuilder.h
+++ b/Engine/Plugins/Enterprise/GLTFExporter/MigratedFromGit/Source/GLTFExporter/Private/Builders/GLTFImageBuilder.h
@@ -10,7 +10,9 @@ class FGLTFImageBuilder : public FGLTFBufferBuilder
 {
 public:
 
-	FGLTFJsonImageIndex AddImage(const TArray64<uint8>& RawData, int32 Width, int32 Height, ERGBFormat RawFormat, int32 BitDepth, bool bFloatFormat, const FString& Name = TEXT(""), EGLTFJsonMimeType MimeType = EGLTFJsonMimeType::PNG, int32 Quality = 100);
+	FGLTFJsonImageIndex AddImage(const void* RawData, int64 ByteLength, int32 Width, int32 Height, ERGBFormat RawFormat, int32 BitDepth, bool bFloatFormat, const FString& Name = TEXT(""), EGLTFJsonMimeType MimeType = EGLTFJsonMimeType::PNG, int32 Quality = 100);
+
+	FGLTFJsonImageIndex AddImage(const TArray<FColor>& Pixels, int32 Width, int32 Height, EPixelFormat PixelFormat = EPixelFormat::PF_B8G8R8A8, const FString& Name = TEXT(""), EGLTFJsonMimeType MimeType = EGLTFJsonMimeType::PNG, int32 Quality = 100);
 
 	FGLTFJsonImageIndex AddImage(const FTextureSource& Image, const FString& Name = TEXT(""), EGLTFJsonMimeType MimeType = EGLTFJsonMimeType::PNG, int32 Quality = 100);
 
