[GLTFExporter] feat(material): initial support for baked metallic-roughness texture

commit 6cb91f2d298329ea275320f2e4c7af486078155d
Date: Fri, 29 May 2020 23:29:25 +0200

    feat(material): initial support for baked metallic-roughness texture
    (cherry picked from commit c6db8c25b13ce2f9444a03ba24726185020d2d16​)

#skipci

diff --git a/Engine/Plugins/Enterprise/GLTFExporter/MigratedFromGit/Source/GLTFExporter/GLTFExporter.Build.cs b/Engine/Plugins/Enterprise/GLTFExporter/MigratedFromGit/Source/GLTFExporter/GLTFExporter.Build.cs
--- a/Engine/Plugins/Enterprise/GLTFExporter/MigratedFromGit/Source/GLTFExporter/GLTFExporter.Build.cs
+++ b/Engine/Plugins/Enterprise/GLTFExporter/MigratedFromGit/Source/GLTFExporter/GLTFExporter.Build.cs
@@ -39,7 +39,10 @@ namespace UnrealBuildTool.Rules
 					"SlateCore",
 					"Mainframe",
 					"InputCore",
-					"EditorStyle"
+					"EditorStyle",
+					"RenderCore",
+					"RHI",
+					"MaterialBaking"
 				}
 				);
 
diff --git a/Engine/Plugins/Enterprise/GLTFExporter/MigratedFromGit/Source/GLTFExporter/Private/Converters/GLTFMaterialConverters.cpp b/Engine/Plugins/Enterprise/GLTFExporter/MigratedFromGit/Source/GLTFExporter/Private/Converters/GLTFMaterialConverters.cpp
--- a/Engine/Plugins/Enterprise/GLTFExporter/MigratedFromGit/Source/GLTFExporter/Private/Converters/GLTFMaterialConverters.cpp
+++ b/Engine/Plugins/Enterprise/GLTFExporter/MigratedFromGit/Source/GLTFExporter/Private/Converters/GLTFMaterialConverters.cpp
@@ -1,8 +1,9 @@
 // Copyright Epic Games, Inc. All Rights Reserved.
 
 #include "Converters/GLTFMaterialConverters.h"
-#include "Builders/GLTFContainerBuilder.h"
 #include "Converters/GLTFConverterUtility.h"
+#include "Converters/GLTFMaterialUtility.h"
+#include "Builders/GLTFContainerBuilder.h"
 #include "Materials/MaterialExpressionConstant.h"
 #include "Materials/MaterialExpressionConstant2Vector.h"
 #include "Materials/MaterialExpressionConstant3Vector.h"
@@ -45,9 +46,9 @@ FGLTFJsonMaterialIndex FGLTFMaterialConverter::Add(FGLTFConvertBuilder& Builder,
 		}
 	}
 
-	if (!TryGetMetallicAndRoughness(Builder, JsonMaterial.PBRMetallicRoughness, Material->Metallic, Material->Roughness, MaterialInstance))
+	if (!TryGetMetallicAndRoughness(Builder, JsonMaterial.PBRMetallicRoughness, Material->Metallic, Material->Roughness, MaterialInterface))
 	{
-		// TODO: add fallback to material baking
+		// TODO: handle failure?
 	}
 
 	if (!TryGetConstantColor(JsonMaterial.EmissiveFactor, Material->EmissiveColor, MaterialInstance))
@@ -128,8 +129,10 @@ bool FGLTFMaterialConverter::TryGetBaseColorAndOpacity(FGLTFConvertBuilder& Buil
 	return false;
 }
 
-bool FGLTFMaterialConverter::TryGetMetallicAndRoughness(FGLTFConvertBuilder& Builder, FGLTFJsonPBRMetallicRoughness& OutPBRParams, const FScalarMaterialInput& MetallicInput, const FScalarMaterialInput& RoughnessInput, const UMaterialInstance* MaterialInstance) const
+bool FGLTFMaterialConverter::TryGetMetallicAndRoughness(FGLTFConvertBuilder& Builder, FGLTFJsonPBRMetallicRoughness& OutPBRParams, const FScalarMaterialInput& MetallicInput, const FScalarMaterialInput& RoughnessInput, const UMaterialInterface* MaterialInterface) const
 {
+	const UMaterialInstance* MaterialInstance = Cast<UMaterialInstance>(MaterialInterface);
+
 	const bool bIsMetallicConstant = TryGetConstantScalar(OutPBRParams.MetallicFactor, MetallicInput, MaterialInstance);
 	const bool bIsRoughnessConstant = TryGetConstantScalar(OutPBRParams.RoughnessFactor, RoughnessInput, MaterialInstance);
 
@@ -138,6 +141,12 @@ bool FGLTFMaterialConverter::TryGetMetallicAndRoughness(FGLTFConvertBuilder& Bui
 		return true;
 	}
 
+	// NOTE: since we always bake the properties (for now) when atleast property is non-const, we need
+	// to reset the constant factors to their defaults. Otherwise the baked value of a constant property
+	// would be scaled with the factor, i.e a double scaling.
+	OutPBRParams.MetallicFactor = 1;
+	OutPBRParams.RoughnessFactor = 1;
+
 	const UTexture2D* MetallicTexture;
 	const UTexture2D* RoughnessTexture;
 	int32 MetallicTexCoord;
@@ -146,33 +155,61 @@ bool FGLTFMaterialConverter::TryGetMetallicAndRoughness(FGLTFConvertBuilder& Bui
 	const bool bHasMetallicSourceTexture = TryGetSourceTexture(MetallicTexture, MetallicTexCoord, MetallicInput, MaterialInstance);
 	const bool bHasRoughnessSourceTexture = TryGetSourceTexture(RoughnessTexture, RoughnessTexCoord, RoughnessInput, MaterialInstance);
 
+	int32 TexCoord = 0;
+	FIntPoint TextureSize(512, 512);	// TODO: make default baking-resolution configurable
+	EGLTFJsonTextureWrap TextureWrap = EGLTFJsonTextureWrap::ClampToEdge;
+	EGLTFJsonTextureFilter TextureFilter = EGLTFJsonTextureFilter::LinearMipmapLinear;
+	const EPixelFormat PixelFormat = PF_B8G8R8A8;
+
 	if (bHasMetallicSourceTexture && bHasRoughnessSourceTexture)
 	{
-		if (MetallicTexture == RoughnessTexture && MetallicTexCoord == RoughnessTexCoord)
+		const bool bAreTexturesCompatible = MetallicTexCoord == RoughnessTexCoord &&
+			MetallicTexture->AddressX == RoughnessTexture->AddressX &&
+			MetallicTexture->AddressY == RoughnessTexture->AddressY;
+
+		if (!bAreTexturesCompatible)
 		{
-			// TODO: make sure textures are correctly masked
-			OutPBRParams.MetallicRoughnessTexture.Index = Builder.GetOrAddTexture(MetallicTexture);
-			OutPBRParams.MetallicRoughnessTexture.TexCoord = MetallicTexCoord;
-			return true;
+			// TODO: handle differences in wrapping or uv-coords
+			return false;
 		}
 
-		// TODO: add support for combining two textures
-		return false;
-	}
+		TexCoord = MetallicTexCoord;
+		TextureSize = FIntPoint(
+			FMath::Max(MetallicTexture->GetSizeX(), RoughnessTexture->GetSizeX()),
+			FMath::Max(MetallicTexture->GetSizeY(), RoughnessTexture->GetSizeY()));
 
-	if (bHasMetallicSourceTexture && bIsRoughnessConstant)
+		TextureWrap = FGLTFConverterUtility::ConvertWrap(MetallicTexture->AddressX);
+		TextureFilter = FGLTFConverterUtility::ConvertFilter(MetallicTexture->Filter, MetallicTexture->LODGroup);
+	}
+	else if (bHasMetallicSourceTexture)
 	{
-		// TODO: add support for combining constant with texture
-		return false;
+		TexCoord = MetallicTexCoord;
+		TextureSize = FIntPoint(MetallicTexture->GetSizeX(), MetallicTexture->GetSizeY());
+		TextureWrap = FGLTFConverterUtility::ConvertWrap(MetallicTexture->AddressX);
+		TextureFilter = FGLTFConverterUtility::ConvertFilter(MetallicTexture->Filter, MetallicTexture->LODGroup);
 	}
-
-	if (bIsMetallicConstant && bHasRoughnessSourceTexture)
+	else if (bHasRoughnessSourceTexture)
 	{
-		// TODO: add support for combining constant with texture
-		return false;
+		TexCoord = RoughnessTexCoord;
+		TextureSize = FIntPoint(RoughnessTexture->GetSizeX(), RoughnessTexture->GetSizeY());
+		TextureWrap = FGLTFConverterUtility::ConvertWrap(RoughnessTexture->AddressX);
+		TextureFilter = FGLTFConverterUtility::ConvertFilter(RoughnessTexture->Filter, RoughnessTexture->LODGroup);
 	}
 
-	return false;
+	RoughnessTexture = FGLTFMaterialUtility::BakeMaterialProperty(TextureSize, MP_Roughness, MaterialInterface);
+	MetallicTexture = FGLTFMaterialUtility::BakeMaterialProperty(TextureSize, MP_Metallic, MaterialInterface);
+
+	const FGLTFJsonTextureIndex TextureIndex = FGLTFMaterialUtility::AddMetallicRoughnessTexture(
+		Builder,
+		MetallicTexture,
+		RoughnessTexture,
+		TextureFilter,
+		TextureWrap);
+
+	OutPBRParams.MetallicRoughnessTexture.TexCoord = TexCoord;
+	OutPBRParams.MetallicRoughnessTexture.Index = TextureIndex;
+
+	return true;
 }
 
 bool FGLTFMaterialConverter::TryGetConstantColor(FGLTFJsonColor3& OutValue, const FColorMaterialInput& MaterialInput, const UMaterialInstance* MaterialInstance) const
diff --git a/Engine/Plugins/Enterprise/GLTFExporter/MigratedFromGit/Source/GLTFExporter/Private/Converters/GLTFMaterialConverters.h b/Engine/Plugins/Enterprise/GLTFExporter/MigratedFromGit/Source/GLTFExporter/Private/Converters/GLTFMaterialConverters.h
--- a/Engine/Plugins/Enterprise/GLTFExporter/MigratedFromGit/Source/GLTFExporter/Private/Converters/GLTFMaterialConverters.h
+++ b/Engine/Plugins/Enterprise/GLTFExporter/MigratedFromGit/Source/GLTFExporter/Private/Converters/GLTFMaterialConverters.h
@@ -12,7 +12,7 @@ class FGLTFMaterialConverter final : public TGLTFConverter<FGLTFJsonMaterialInde
 	FGLTFJsonMaterialIndex Add(FGLTFConvertBuilder& Builder, const FString& Name, const UMaterialInterface* Material) override;
 
 	bool TryGetBaseColorAndOpacity(FGLTFConvertBuilder& Builder, FGLTFJsonPBRMetallicRoughness& OutPBRParams, const FColorMaterialInput& BaseColorInput, const FScalarMaterialInput& OpacityInput, const UMaterialInstance* MaterialInstance) const;
-	bool TryGetMetallicAndRoughness(FGLTFConvertBuilder& Builder, FGLTFJsonPBRMetallicRoughness& OutPBRParams, const FScalarMaterialInput& MetallicInput, const FScalarMaterialInput& RoughnessInput, const UMaterialInstance* MaterialInstance) const;
+	bool TryGetMetallicAndRoughness(FGLTFConvertBuilder& Builder, FGLTFJsonPBRMetallicRoughness& OutPBRParams, const FScalarMaterialInput& MetallicInput, const FScalarMaterialInput& RoughnessInput, const UMaterialInterface* MaterialInterface) const;
 
 	bool TryGetConstantColor(FGLTFJsonColor3& OutValue, const FColorMaterialInput& MaterialInput, const UMaterialInstance* MaterialInstance) const;
 	bool TryGetConstantColor(FGLTFJsonColor4& OutValue, const FColorMaterialInput& MaterialInput, const UMaterialInstance* MaterialInstance) const;
diff --git a/Engine/Plugins/Enterprise/GLTFExporter/MigratedFromGit/Source/GLTFExporter/Private/Converters/GLTFMaterialUtility.cpp b/Engine/Plugins/Enterprise/GLTFExporter/MigratedFromGit/Source/GLTFExporter/Private/Converters/GLTFMaterialUtility.cpp
new file mode 100644
--- /dev/null
+++ b/Engine/Plugins/Enterprise/GLTFExporter/MigratedFromGit/Source/GLTFExporter/Private/Converters/GLTFMaterialUtility.cpp
@@ -0,0 +1,150 @@
+// Copyright Epic Games, Inc. All Rights Reserved.
+
+#include "GLTFMaterialUtility.h"
+
+#include "Engine/TextureRenderTarget2D.h"
+#include "CanvasItem.h"
+#include "CanvasTypes.h"
+#include "Modules/ModuleManager.h"
+
+#include "IMaterialBakingModule.h"
+#include "MaterialBakingStructures.h"
+
+UTexture2D* FGLTFMaterialUtility::CreateTransientTexture(const TArray<FColor>& Pixels, const FIntPoint& TextureSize, const EPixelFormat& TextureFormat)
+{
+	check(TextureSize.X * TextureSize.Y == Pixels.Num());
+
+	UTexture2D* Texture = UTexture2D::CreateTransient(TextureSize.X, TextureSize.Y, TextureFormat);
+	void* MipData = Texture->PlatformData->Mips[0].BulkData.Lock(LOCK_READ_WRITE);
+
+	FMemory::Memcpy(MipData, Pixels.GetData(), Pixels.Num() * Pixels.GetTypeSize());
+	Texture->PlatformData->Mips[0].BulkData.Unlock();
+
+	// Texture->SRGB = false;	// TODO: Do we need this?
+	Texture->CompressionNone = true;
+	Texture->MipGenSettings = TMGS_NoMipmaps;
+	Texture->CompressionSettings = TC_Default;
+
+	Texture->UpdateResource();
+	return Texture;
+}
+
+bool FGLTFMaterialUtility::CombineTextures(TArray<FColor>& OutPixels, const TArray<FGLTFTextureCombineSource>& Sources, const FIntPoint& OutputSize, const EPixelFormat OutputPixelFormat)
+{
+	UTextureRenderTarget2D* RenderTarget2D = NewObject<UTextureRenderTarget2D>();
+	check(RenderTarget2D);
+
+	RenderTarget2D->AddToRoot();
+	RenderTarget2D->ClearColor = FLinearColor(0.0f, 0.0f, 0.0f, 0.0f);
+	RenderTarget2D->InitCustomFormat(OutputSize.X, OutputSize.Y, OutputPixelFormat, false);
+	RenderTarget2D->TargetGamma = 0;
+
+	FRenderTarget* RenderTarget = RenderTarget2D->GameThread_GetRenderTargetResource();
+	FCanvas* Canvas = new FCanvas(RenderTarget, nullptr, 0, 0, 0, GMaxRHIFeatureLevel);
+
+	Canvas->SetRenderTarget_GameThread(RenderTarget);
+	Canvas->Clear(FLinearColor(0, 0, 0, 0));
+
+	const FVector2D TileSize(OutputSize.X, OutputSize.Y);
+	const FVector2D TilePosition(0, 0);
+
+	for (const FGLTFTextureCombineSource& Source: Sources)
+	{
+		check(Source.Texture);
+
+		FCanvasTileItem TileItem(TilePosition, Source.Texture->Resource, TileSize, Source.TintColor);
+
+		TileItem.BlendMode = Source.BlendMode;
+		TileItem.Draw(Canvas);
+	}
+
+	Canvas->Flush_GameThread();
+	FlushRenderingCommands();
+	Canvas->SetRenderTarget_GameThread(nullptr);
+	FlushRenderingCommands();
+
+	const bool bReadSuccessful = RenderTarget->ReadPixels(OutPixels);
+
+	// Clean up.
+	RenderTarget2D->ReleaseResource();
+	RenderTarget2D->RemoveFromRoot();
+	RenderTarget2D = nullptr;
+	delete Canvas;
+
+	return bReadSuccessful;
+}
+
+UTexture2D* FGLTFMaterialUtility::BakeMaterialProperty(const FIntPoint OutputSize, const EMaterialProperty& MaterialProperty, const UMaterialInterface* Material)
+{
+	TArray<FMeshData*> MeshSettings;
+
+	FMeshData* MeshSet = new FMeshData();
+
+	MeshSet->TextureCoordinateBox = FBox2D(FVector2D(0, 0), FVector2D(1, 1));
+
+	// TODO: Do we need to fill in any info in MeshSet?
+
+	MeshSettings.Add(MeshSet);
+
+	TArray<FMaterialData*> MatSettings;
+
+	FMaterialData* MatSet = new FMaterialData();
+	MatSet->Material = const_cast<UMaterialInterface*>(Material);
+	MatSet->PropertySizes.Add(MaterialProperty, OutputSize);
+
+	MatSettings.Add(MatSet);
+
+	TArray<FBakeOutput> BakeOutputs;
+	IMaterialBakingModule& Module = FModuleManager::Get().LoadModuleChecked<IMaterialBakingModule>("MaterialBaking");
+
+	Module.BakeMaterials(MatSettings, MeshSettings, BakeOutputs);
+
+	const FBakeOutput& Output = BakeOutputs[0];
+
+	// NOTE: The texture-format is mapped per property in the MaterialBaking-module via PerPropertyFormat.
+	// It's PF_B8G8R8A8 for all properties except for MP_EmissiveColor (which is PF_FloatRGBA).
+	// TODO: Should we handle cases the format *may* be something other that what we expect?
+	return CreateTransientTexture(
+		Output.PropertyData.FindChecked(MaterialProperty),
+		Output.PropertySizes.FindChecked(MaterialProperty),
+		PF_B8G8R8A8);
+}
+
+FGLTFJsonTextureIndex FGLTFMaterialUtility::AddMetallicRoughnessTexture(FGLTFConvertBuilder& Builder, const UTexture2D* MetallicTexture, const UTexture2D* RoughnessTexture, EGLTFJsonTextureFilter Filter, EGLTFJsonTextureWrap Wrap)
+{
+	check(MetallicTexture->GetPixelFormat() == RoughnessTexture->GetPixelFormat());
+
+	// TODO: Smarter calculation for output-size?
+	const FIntPoint TextureSize(
+		FMath::Max(MetallicTexture->GetSizeX(), RoughnessTexture->GetSizeX()),
+		FMath::Max(MetallicTexture->GetSizeY(), RoughnessTexture->GetSizeY()));
+
+	TArray<FColor> Pixels;
+	const EPixelFormat PixelFormat = MetallicTexture->GetPixelFormat();
+
+	const TArray<FGLTFTextureCombineSource> CombineSources =
+	{
+		{MetallicTexture, FLinearColor(0, 0, 1, 0)},
+		{RoughnessTexture, FLinearColor(0, 1, 0, 0)}
+	};
+
+	if (!CombineTextures(Pixels, CombineSources, TextureSize, PixelFormat))
+	{
+		return FGLTFJsonTextureIndex(INDEX_NONE);
+	}
+
+	const FString Name = TEXT("MetallicRoughness");
+
+	// TODO: maybe we should reuse existing samplers?
+	FGLTFJsonSampler JsonSampler;
+	JsonSampler.Name = Name;
+	JsonSampler.MagFilter = JsonSampler.MinFilter = Filter;
+	JsonSampler.WrapS = JsonSampler.WrapT = Wrap;
+
+	FGLTFJsonTexture JsonTexture;
+	JsonTexture.Name = Name;
+	JsonTexture.Sampler = Builder.AddSampler(JsonSampler);
+	JsonTexture.Source = Builder.AddImage(Pixels, TextureSize.X, TextureSize.Y, PixelFormat, Name);
+
+	return Builder.AddTexture(JsonTexture);
+}
diff --git a/Engine/Plugins/Enterprise/GLTFExporter/MigratedFromGit/Source/GLTFExporter/Private/Converters/GLTFMaterialUtility.h b/Engine/Plugins/Enterprise/GLTFExporter/MigratedFromGit/Source/GLTFExporter/Private/Converters/GLTFMaterialUtility.h
new file mode 100644
--- /dev/null
+++ b/Engine/Plugins/Enterprise/GLTFExporter/MigratedFromGit/Source/GLTFExporter/Private/Converters/GLTFMaterialUtility.h
@@ -0,0 +1,26 @@
+// Copyright Epic Games, Inc. All Rights Reserved.
+
+#pragma once
+
+#include "Builders/GLTFContainerBuilder.h"
+#include "Engine/Texture2D.h"
+
+struct FGLTFTextureCombineSource
+{
+	constexpr FORCEINLINE FGLTFTextureCombineSource(const UTexture2D* Texture, FLinearColor TintColor = FLinearColor(1, 1, 1, 1), ESimpleElementBlendMode BlendMode = SE_BLEND_Additive)
+		: Texture(Texture), TintColor(TintColor), BlendMode(BlendMode)
+	{}
+
+	const UTexture2D* Texture;
+	FLinearColor TintColor;
+	ESimpleElementBlendMode BlendMode;
+};
+
+struct FGLTFMaterialUtility
+{
+	static UTexture2D* CreateTransientTexture(const TArray<FColor>& Pixels, const FIntPoint& TextureSize, const EPixelFormat& TextureFormat);
+	static bool CombineTextures(TArray<FColor>& OutPixels, const TArray<FGLTFTextureCombineSource>& Sources, const FIntPoint& OutputSize, const EPixelFormat OutputPixelFormat);
+	static UTexture2D* BakeMaterialProperty(const FIntPoint OutputSize, const EMaterialProperty& MaterialProperty, const UMaterialInterface* Material);
+
+	static FGLTFJsonTextureIndex AddMetallicRoughnessTexture(FGLTFConvertBuilder& Builder, const UTexture2D* MetallicTexture, const UTexture2D* RoughnessTexture, EGLTFJsonTextureFilter Filter, EGLTFJsonTextureWrap Wrap);
+};
