[GLTFExporter] feat(material): add utility for converting channel-masks to a vector for easier comparison

commit 0a90237a4e054e3cd789c6dd70c7e09fcb212d47
Date: Sat, 30 May 2020 09:28:03 +0200

    feat(material): add utility for converting channel-masks to a vector for easier comparison
    (cherry picked from commit d00340033ebf56a00d04c5f8f7a21dac08945dec​)

#skipci

diff --git a/Engine/Plugins/Enterprise/GLTFExporter/MigratedFromGit/Source/GLTFExporter/Private/Converters/GLTFMaterialUtility.cpp b/Engine/Plugins/Enterprise/GLTFExporter/MigratedFromGit/Source/GLTFExporter/Private/Converters/GLTFMaterialUtility.cpp
--- a/Engine/Plugins/Enterprise/GLTFExporter/MigratedFromGit/Source/GLTFExporter/Private/Converters/GLTFMaterialUtility.cpp
+++ b/Engine/Plugins/Enterprise/GLTFExporter/MigratedFromGit/Source/GLTFExporter/Private/Converters/GLTFMaterialUtility.cpp
@@ -148,3 +148,10 @@ FGLTFJsonTextureIndex FGLTFMaterialUtility::AddMetallicRoughnessTexture(FGLTFCon
 
 	return Builder.AddTexture(JsonTexture);
 }
+
+FIntVector4 FGLTFMaterialUtility::ConvertMaskToVector(const FExpressionInput& ExpressionInput)
+{
+	// TODO: Do we need to check MaterialInput.Mask first (which seems to always be 0 or 1, not an actual bitfield)?
+	// Or is it enough to assume that MaskR, MaskG, MaskB and MaskA will only be 1 when Mask is 1?
+	return FIntVector4(ExpressionInput.MaskR, ExpressionInput.MaskG, ExpressionInput.MaskB, ExpressionInput.MaskA);
+}
diff --git a/Engine/Plugins/Enterprise/GLTFExporter/MigratedFromGit/Source/GLTFExporter/Private/Converters/GLTFMaterialUtility.h b/Engine/Plugins/Enterprise/GLTFExporter/MigratedFromGit/Source/GLTFExporter/Private/Converters/GLTFMaterialUtility.h
--- a/Engine/Plugins/Enterprise/GLTFExporter/MigratedFromGit/Source/GLTFExporter/Private/Converters/GLTFMaterialUtility.h
+++ b/Engine/Plugins/Enterprise/GLTFExporter/MigratedFromGit/Source/GLTFExporter/Private/Converters/GLTFMaterialUtility.h
@@ -23,4 +23,5 @@ struct FGLTFMaterialUtility
 	static UTexture2D* BakeMaterialProperty(const FIntPoint OutputSize, const EMaterialProperty& MaterialProperty, const UMaterialInterface* Material);
 
 	static FGLTFJsonTextureIndex AddMetallicRoughnessTexture(FGLTFConvertBuilder& Builder, const UTexture2D* MetallicTexture, const UTexture2D* RoughnessTexture, EGLTFJsonTextureFilter Filter, EGLTFJsonTextureWrap Wrap);
+	static FIntVector4 ConvertMaskToVector(const FExpressionInput& ExpressionInput);
 };
