[GLTFExporter] feat(material): properly detect when an existing metallic-roughness texture can be used without conversion

commit 1acd1b5406c9748df57d9edb9bf5b3a380428e95
Date: Sat, 30 May 2020 09:29:56 +0200

    feat(material): properly detect when an existing metallic-roughness texture can be used without conversion
    (cherry picked from commit 58b4a309a48b74d69db3b953ae67fe7f63007168​)

#skipci

diff --git a/Engine/Plugins/Enterprise/GLTFExporter/MigratedFromGit/Source/GLTFExporter/Private/Converters/GLTFMaterialConverters.cpp b/Engine/Plugins/Enterprise/GLTFExporter/MigratedFromGit/Source/GLTFExporter/Private/Converters/GLTFMaterialConverters.cpp
--- a/Engine/Plugins/Enterprise/GLTFExporter/MigratedFromGit/Source/GLTFExporter/Private/Converters/GLTFMaterialConverters.cpp
+++ b/Engine/Plugins/Enterprise/GLTFExporter/MigratedFromGit/Source/GLTFExporter/Private/Converters/GLTFMaterialConverters.cpp
@@ -155,6 +155,23 @@ bool FGLTFMaterialConverter::TryGetMetallicAndRoughness(FGLTFConvertBuilder& Bui
 	const bool bHasMetallicSourceTexture = TryGetSourceTexture(MetallicTexture, MetallicTexCoord, MetallicInput, MaterialInstance);
 	const bool bHasRoughnessSourceTexture = TryGetSourceTexture(RoughnessTexture, RoughnessTexCoord, RoughnessInput, MaterialInstance);
 
+	// Detect the "happy path" where both inputs share the same texture, and both inputs correctly masked.
+	if (bHasMetallicSourceTexture && bHasRoughnessSourceTexture)
+	{
+		const FIntVector4 MetallicMask(0, 0, 1, 0);
+		const FIntVector4 RoughnessMask(0, 1, 0, 0);
+
+		if (MetallicTexture == RoughnessTexture &&
+			MetallicTexCoord == RoughnessTexCoord &&
+			MetallicMask == FGLTFMaterialUtility::ConvertMaskToVector(MetallicInput) &&
+			RoughnessMask == FGLTFMaterialUtility::ConvertMaskToVector(RoughnessInput))
+		{
+			OutPBRParams.MetallicRoughnessTexture.Index = Builder.GetOrAddTexture(MetallicTexture);
+			OutPBRParams.MetallicRoughnessTexture.TexCoord = MetallicTexCoord;
+			return true;
+		}
+	}
+
 	int32 TexCoord = 0;
 	FIntPoint TextureSize(512, 512);	// TODO: make default baking-resolution configurable
 	EGLTFJsonTextureWrap TextureWrap = EGLTFJsonTextureWrap::ClampToEdge;
