[GLTFExporter] feat(material): add more mask-utilities for materials

commit f21a82cbedb5a8dc208107b628de860d47ce505a
Date: Tue, 2 Jun 2020 15:42:36 +0200

    feat(material): add more mask-utilities for materials
    (cherry picked from commit 9c8364a1f7918fb7e123b682227b53346dfd11a4​)

#skipci

diff --git a/Engine/Plugins/Enterprise/GLTFExporter/MigratedFromGit/Source/GLTFExporter/Private/Converters/GLTFMaterialUtility.cpp b/Engine/Plugins/Enterprise/GLTFExporter/MigratedFromGit/Source/GLTFExporter/Private/Converters/GLTFMaterialUtility.cpp
--- a/Engine/Plugins/Enterprise/GLTFExporter/MigratedFromGit/Source/GLTFExporter/Private/Converters/GLTFMaterialUtility.cpp
+++ b/Engine/Plugins/Enterprise/GLTFExporter/MigratedFromGit/Source/GLTFExporter/Private/Converters/GLTFMaterialUtility.cpp
@@ -155,3 +155,15 @@ FIntVector4 FGLTFMaterialUtility::ConvertMaskToVector(const FExpressionInput& Ex
 	// Or is it enough to assume that MaskR, MaskG, MaskB and MaskA will only be 1 when Mask is 1?
 	return FIntVector4(ExpressionInput.MaskR, ExpressionInput.MaskG, ExpressionInput.MaskB, ExpressionInput.MaskA);
 }
+
+FLinearColor FGLTFMaterialUtility::ConvertMaskToColor(const FExpressionInput& ExpressionInput)
+{
+	const FIntVector4 Mask = ConvertMaskToVector(ExpressionInput);
+
+	return FLinearColor(Mask.X, Mask.Y, Mask.Z, Mask.W);
+}
+
+uint32 FGLTFMaterialUtility::GetMaskComponentCount(const FExpressionInput& ExpressionInput)
+{
+	return ExpressionInput.MaskR + ExpressionInput.MaskG + ExpressionInput.MaskB + ExpressionInput.MaskA;
+}
diff --git a/Engine/Plugins/Enterprise/GLTFExporter/MigratedFromGit/Source/GLTFExporter/Private/Converters/GLTFMaterialUtility.h b/Engine/Plugins/Enterprise/GLTFExporter/MigratedFromGit/Source/GLTFExporter/Private/Converters/GLTFMaterialUtility.h
--- a/Engine/Plugins/Enterprise/GLTFExporter/MigratedFromGit/Source/GLTFExporter/Private/Converters/GLTFMaterialUtility.h
+++ b/Engine/Plugins/Enterprise/GLTFExporter/MigratedFromGit/Source/GLTFExporter/Private/Converters/GLTFMaterialUtility.h
@@ -24,4 +24,6 @@ struct FGLTFMaterialUtility
 
 	static FGLTFJsonTextureIndex AddMetallicRoughnessTexture(FGLTFConvertBuilder& Builder, const UTexture2D* MetallicTexture, const UTexture2D* RoughnessTexture, EGLTFJsonTextureFilter Filter, EGLTFJsonTextureWrap Wrap);
 	static FIntVector4 ConvertMaskToVector(const FExpressionInput& ExpressionInput);
+	static FLinearColor ConvertMaskToColor(const FExpressionInput& ExpressionInput);
+	static uint32 GetMaskComponentCount(const FExpressionInput& ExpressionInput);
 };
