[GLTFExporter] chore(material): remove incorrect TODO:s for output-masks

commit c11179754f45d8a3859943f088da4f6a9b44434d
Date: Tue, 2 Jun 2020 16:25:20 +0200

    chore(material): remove incorrect TODO:s for output-masks
    (cherry picked from commit f4da441a9479d60c4a7a8646328c3c4c80ea6958​)

#skipci

diff --git a/Engine/Plugins/Enterprise/GLTFExporter/MigratedFromGit/Source/GLTFExporter/Private/Converters/GLTFMaterialConverters.cpp b/Engine/Plugins/Enterprise/GLTFExporter/MigratedFromGit/Source/GLTFExporter/Private/Converters/GLTFMaterialConverters.cpp
--- a/Engine/Plugins/Enterprise/GLTFExporter/MigratedFromGit/Source/GLTFExporter/Private/Converters/GLTFMaterialConverters.cpp
+++ b/Engine/Plugins/Enterprise/GLTFExporter/MigratedFromGit/Source/GLTFExporter/Private/Converters/GLTFMaterialConverters.cpp
@@ -300,21 +300,18 @@ bool FGLTFMaterialConverter::TryGetConstantColor(FLinearColor& OutValue, const F
 
 	if (const UMaterialExpressionConstant4Vector* Constant4Vector = ExactCast<UMaterialExpressionConstant4Vector>(Expression))
 	{
-		// TODO: add support for output mask
 		OutValue = Constant4Vector->Constant;
 		return true;
 	}
 
 	if (const UMaterialExpressionConstant3Vector* Constant3Vector = ExactCast<UMaterialExpressionConstant3Vector>(Expression))
 	{
-		// TODO: add support for output mask
 		OutValue = Constant3Vector->Constant;
 		return true;
 	}
 
 	if (const UMaterialExpressionConstant2Vector* Constant2Vector = ExactCast<UMaterialExpressionConstant2Vector>(Expression))
 	{
-		// TODO: add support for output mask
 		OutValue = { Constant2Vector->R, Constant2Vector->G, 0, 0 };
 		return true;
 	}
@@ -375,21 +372,18 @@ bool FGLTFMaterialConverter::TryGetConstantScalar(float& OutValue, const FScalar
 
 	if (const UMaterialExpressionConstant4Vector* Constant4Vector = ExactCast<UMaterialExpressionConstant4Vector>(Expression))
 	{
-		// TODO: add support for output mask
 		OutValue = Constant4Vector->Constant.R;
 		return true;
 	}
 
 	if (const UMaterialExpressionConstant3Vector* Constant3Vector = ExactCast<UMaterialExpressionConstant3Vector>(Expression))
 	{
-		// TODO: add support for output mask
 		OutValue = Constant3Vector->Constant.R;
 		return true;
 	}
 
 	if (const UMaterialExpressionConstant2Vector* Constant2Vector = ExactCast<UMaterialExpressionConstant2Vector>(Expression))
 	{
-		// TODO: add support for output mask
 		OutValue = Constant2Vector->R;
 		return true;
 	}
