[GLTFExporter] feat(material): handle output-mask for UMaterialExpressionVectorParameter

commit 850799ac5f5d6a217d1675c238d06537b0500685
Date: Tue, 2 Jun 2020 16:25:42 +0200

    feat(material): handle output-mask for UMaterialExpressionVectorParameter
    (cherry picked from commit e371f840f7f49f83859f894d4f1d0951dad3a652​)

#skipci

diff --git a/Engine/Plugins/Enterprise/GLTFExporter/MigratedFromGit/Source/GLTFExporter/Private/Converters/GLTFMaterialConverters.cpp b/Engine/Plugins/Enterprise/GLTFExporter/MigratedFromGit/Source/GLTFExporter/Private/Converters/GLTFMaterialConverters.cpp
--- a/Engine/Plugins/Enterprise/GLTFExporter/MigratedFromGit/Source/GLTFExporter/Private/Converters/GLTFMaterialConverters.cpp
+++ b/Engine/Plugins/Enterprise/GLTFExporter/MigratedFromGit/Source/GLTFExporter/Private/Converters/GLTFMaterialConverters.cpp
@@ -276,7 +276,21 @@ bool FGLTFMaterialConverter::TryGetConstantColor(FLinearColor& OutValue, const F
 			}
 		}
 
-		// TODO: add support for output mask
+		const uint32 MaskComponentCount = FGLTFMaterialUtility::GetMaskComponentCount(MaterialInput);
+
+		if (MaskComponentCount > 0)
+		{
+			const FLinearColor Mask = FGLTFMaterialUtility::ConvertMaskToColor(MaterialInput);
+
+			Value *= Mask;
+
+			if (MaskComponentCount == 1)
+			{
+				const float ComponentValue = Value.R + Value.G + Value.B + Value.A;
+				Value = { ComponentValue, ComponentValue, ComponentValue, ComponentValue };
+			}
+		}
+
 		OutValue = Value;
 		return true;
 	}
@@ -348,8 +362,16 @@ bool FGLTFMaterialConverter::TryGetConstantScalar(float& OutValue, const FScalar
 			}
 		}
 
-		// TODO: add support for output mask
-		OutValue = Value.R;
+		const uint32 MaskComponentCount = FGLTFMaterialUtility::GetMaskComponentCount(MaterialInput);
+
+		if (MaskComponentCount > 0)
+		{
+			const FLinearColor Mask = FGLTFMaterialUtility::ConvertMaskToColor(MaterialInput);
+			Value *= Mask;
+		}
+
+		// TODO: is this a correct assumption, that the max component should be used as value?
+		OutValue = Value.GetMax();
 		return true;
 	}
 
