[GLTFExporter] fix(material): prefix the name of MetallicRoughness-textures with the material name to avoid name-collisions

commit 1bd16dccab19cb8d9717f4c1ed67a73b57e35006
Date: Tue, 2 Jun 2020 16:42:14 +0200

    fix(material): prefix the name of MetallicRoughness-textures with the material name to avoid name-collisions
    (cherry picked from commit 9b33827f579618e29fe092c33db991a171eba810​)

#skipci

diff --git a/Engine/Plugins/Enterprise/GLTFExporter/MigratedFromGit/Source/GLTFExporter/Private/Converters/GLTFMaterialConverters.cpp b/Engine/Plugins/Enterprise/GLTFExporter/MigratedFromGit/Source/GLTFExporter/Private/Converters/GLTFMaterialConverters.cpp
--- a/Engine/Plugins/Enterprise/GLTFExporter/MigratedFromGit/Source/GLTFExporter/Private/Converters/GLTFMaterialConverters.cpp
+++ b/Engine/Plugins/Enterprise/GLTFExporter/MigratedFromGit/Source/GLTFExporter/Private/Converters/GLTFMaterialConverters.cpp
@@ -216,12 +216,15 @@ bool FGLTFMaterialConverter::TryGetMetallicAndRoughness(FGLTFConvertBuilder& Bui
 	RoughnessTexture = FGLTFMaterialUtility::BakeMaterialProperty(TextureSize, MP_Roughness, MaterialInterface);
 	MetallicTexture = FGLTFMaterialUtility::BakeMaterialProperty(TextureSize, MP_Metallic, MaterialInterface);
 
+	const FString TextureName = FString::Printf(TEXT("%s_MetallicRoughness"), *MaterialInterface->GetName());
+
 	const FGLTFJsonTextureIndex TextureIndex = FGLTFMaterialUtility::AddMetallicRoughnessTexture(
 		Builder,
 		MetallicTexture,
 		RoughnessTexture,
 		TextureFilter,
-		TextureWrap);
+		TextureWrap,
+		TextureName);
 
 	OutPBRParams.MetallicRoughnessTexture.TexCoord = TexCoord;
 	OutPBRParams.MetallicRoughnessTexture.Index = TextureIndex;
diff --git a/Engine/Plugins/Enterprise/GLTFExporter/MigratedFromGit/Source/GLTFExporter/Private/Converters/GLTFMaterialUtility.cpp b/Engine/Plugins/Enterprise/GLTFExporter/MigratedFromGit/Source/GLTFExporter/Private/Converters/GLTFMaterialUtility.cpp
--- a/Engine/Plugins/Enterprise/GLTFExporter/MigratedFromGit/Source/GLTFExporter/Private/Converters/GLTFMaterialUtility.cpp
+++ b/Engine/Plugins/Enterprise/GLTFExporter/MigratedFromGit/Source/GLTFExporter/Private/Converters/GLTFMaterialUtility.cpp
@@ -110,7 +110,7 @@ UTexture2D* FGLTFMaterialUtility::BakeMaterialProperty(const FIntPoint OutputSiz
 		PF_B8G8R8A8);
 }
 
-FGLTFJsonTextureIndex FGLTFMaterialUtility::AddMetallicRoughnessTexture(FGLTFConvertBuilder& Builder, const UTexture2D* MetallicTexture, const UTexture2D* RoughnessTexture, EGLTFJsonTextureFilter Filter, EGLTFJsonTextureWrap Wrap)
+FGLTFJsonTextureIndex FGLTFMaterialUtility::AddMetallicRoughnessTexture(FGLTFConvertBuilder& Builder, const UTexture2D* MetallicTexture, const UTexture2D* RoughnessTexture, EGLTFJsonTextureFilter Filter, EGLTFJsonTextureWrap Wrap, const FString& TextureName)
 {
 	check(MetallicTexture->GetPixelFormat() == RoughnessTexture->GetPixelFormat());
 
@@ -133,18 +133,16 @@ FGLTFJsonTextureIndex FGLTFMaterialUtility::AddMetallicRoughnessTexture(FGLTFCon
 		return FGLTFJsonTextureIndex(INDEX_NONE);
 	}
 
-	const FString Name = TEXT("MetallicRoughness");
-
 	// TODO: maybe we should reuse existing samplers?
 	FGLTFJsonSampler JsonSampler;
-	JsonSampler.Name = Name;
+	JsonSampler.Name = TextureName;
 	JsonSampler.MagFilter = JsonSampler.MinFilter = Filter;
 	JsonSampler.WrapS = JsonSampler.WrapT = Wrap;
 
 	FGLTFJsonTexture JsonTexture;
-	JsonTexture.Name = Name;
+	JsonTexture.Name = TextureName;
 	JsonTexture.Sampler = Builder.AddSampler(JsonSampler);
-	JsonTexture.Source = Builder.AddImage(Pixels, TextureSize.X, TextureSize.Y, PixelFormat, Name);
+	JsonTexture.Source = Builder.AddImage(Pixels, TextureSize.X, TextureSize.Y, PixelFormat, TextureName);
 
 	return Builder.AddTexture(JsonTexture);
 }
diff --git a/Engine/Plugins/Enterprise/GLTFExporter/MigratedFromGit/Source/GLTFExporter/Private/Converters/GLTFMaterialUtility.h b/Engine/Plugins/Enterprise/GLTFExporter/MigratedFromGit/Source/GLTFExporter/Private/Converters/GLTFMaterialUtility.h
--- a/Engine/Plugins/Enterprise/GLTFExporter/MigratedFromGit/Source/GLTFExporter/Private/Converters/GLTFMaterialUtility.h
+++ b/Engine/Plugins/Enterprise/GLTFExporter/MigratedFromGit/Source/GLTFExporter/Private/Converters/GLTFMaterialUtility.h
@@ -22,7 +22,7 @@ struct FGLTFMaterialUtility
 	static bool CombineTextures(TArray<FColor>& OutPixels, const TArray<FGLTFTextureCombineSource>& Sources, const FIntPoint& OutputSize, const EPixelFormat OutputPixelFormat);
 	static UTexture2D* BakeMaterialProperty(const FIntPoint OutputSize, const EMaterialProperty& MaterialProperty, const UMaterialInterface* Material);
 
-	static FGLTFJsonTextureIndex AddMetallicRoughnessTexture(FGLTFConvertBuilder& Builder, const UTexture2D* MetallicTexture, const UTexture2D* RoughnessTexture, EGLTFJsonTextureFilter Filter, EGLTFJsonTextureWrap Wrap);
+	static FGLTFJsonTextureIndex AddMetallicRoughnessTexture(FGLTFConvertBuilder& Builder, const UTexture2D* MetallicTexture, const UTexture2D* RoughnessTexture, EGLTFJsonTextureFilter Filter, EGLTFJsonTextureWrap Wrap, const FString& TextureName = TEXT("MetallicRoughness"));
 	static FIntVector4 ConvertMaskToVector(const FExpressionInput& ExpressionInput);
 	static FLinearColor ConvertMaskToColor(const FExpressionInput& ExpressionInput);
 	static uint32 GetMaskComponentCount(const FExpressionInput& ExpressionInput);
