[GLTFExporter] refactor(material): generalize the function that creates and adds combined textures

commit 2a7c449f2002962b545c24e66ddc2c7aa621157b
Date: Thu, 4 Jun 2020 11:04:12 +0200

    refactor(material): generalize the function that creates and adds combined textures
    this function will be used for more properties than just metallic+roughness
    (cherry picked from commit e979e867df1fcff8e072400773a4309a76648a48​)

#skipci

diff --git a/Engine/Plugins/Enterprise/GLTFExporter/MigratedFromGit/Source/GLTFExporter/Private/Converters/GLTFMaterialConverters.cpp b/Engine/Plugins/Enterprise/GLTFExporter/MigratedFromGit/Source/GLTFExporter/Private/Converters/GLTFMaterialConverters.cpp
--- a/Engine/Plugins/Enterprise/GLTFExporter/MigratedFromGit/Source/GLTFExporter/Private/Converters/GLTFMaterialConverters.cpp
+++ b/Engine/Plugins/Enterprise/GLTFExporter/MigratedFromGit/Source/GLTFExporter/Private/Converters/GLTFMaterialConverters.cpp
@@ -216,15 +216,31 @@ bool FGLTFMaterialConverter::TryGetMetallicAndRoughness(FGLTFConvertBuilder& Bui
 	RoughnessTexture = FGLTFMaterialUtility::BakeMaterialProperty(TextureSize, MP_Roughness, MaterialInterface);
 	MetallicTexture = FGLTFMaterialUtility::BakeMaterialProperty(TextureSize, MP_Metallic, MaterialInterface);
 
+	// NOTE: the baked textures may have a different (smaller) size than requested, so we update the
+	// texture-size to fit both textures without wasting too much space.
+	TextureSize = {
+		FMath::Max(MetallicTexture->GetSizeX(), RoughnessTexture->GetSizeX()),
+		FMath::Max(MetallicTexture->GetSizeY(), RoughnessTexture->GetSizeY())
+	};
+
+	// TODO: handle the case where TextureSize is 1x1. In this case, both properties are constants and we should
+	// extract the value of each property from the texture and use as-is instead of exporting a combined texture.
+
 	const FString TextureName = FString::Printf(TEXT("%s_MetallicRoughness"), *MaterialInterface->GetName());
 
-	const FGLTFJsonTextureIndex TextureIndex = FGLTFMaterialUtility::AddMetallicRoughnessTexture(
+	const TArray<FGLTFTextureCombineSource> CombineSources =
+	{
+		{MetallicTexture, FLinearColor(0, 0, 1, 0)},
+		{RoughnessTexture, FLinearColor(0, 1, 0, 0)}
+	};
+
+	const FGLTFJsonTextureIndex TextureIndex = FGLTFMaterialUtility::AddCombinedTexture(
 		Builder,
-		MetallicTexture,
-		RoughnessTexture,
+		CombineSources,
+		TextureSize,
+		TextureName,
 		TextureFilter,
-		TextureWrap,
-		TextureName);
+		TextureWrap);
 
 	OutPBRParams.MetallicRoughnessTexture.TexCoord = TexCoord;
 	OutPBRParams.MetallicRoughnessTexture.Index = TextureIndex;
diff --git a/Engine/Plugins/Enterprise/GLTFExporter/MigratedFromGit/Source/GLTFExporter/Private/Converters/GLTFMaterialUtility.cpp b/Engine/Plugins/Enterprise/GLTFExporter/MigratedFromGit/Source/GLTFExporter/Private/Converters/GLTFMaterialUtility.cpp
--- a/Engine/Plugins/Enterprise/GLTFExporter/MigratedFromGit/Source/GLTFExporter/Private/Converters/GLTFMaterialUtility.cpp
+++ b/Engine/Plugins/Enterprise/GLTFExporter/MigratedFromGit/Source/GLTFExporter/Private/Converters/GLTFMaterialUtility.cpp
@@ -110,23 +110,12 @@ UTexture2D* FGLTFMaterialUtility::BakeMaterialProperty(const FIntPoint OutputSiz
 		PF_B8G8R8A8);
 }
 
-FGLTFJsonTextureIndex FGLTFMaterialUtility::AddMetallicRoughnessTexture(FGLTFConvertBuilder& Builder, const UTexture2D* MetallicTexture, const UTexture2D* RoughnessTexture, EGLTFJsonTextureFilter Filter, EGLTFJsonTextureWrap Wrap, const FString& TextureName)
+FGLTFJsonTextureIndex FGLTFMaterialUtility::AddCombinedTexture(FGLTFConvertBuilder& Builder, const TArray<FGLTFTextureCombineSource>& CombineSources, const FIntPoint TextureSize, const FString& TextureName, const EGLTFJsonTextureFilter Filter, const EGLTFJsonTextureWrap Wrap)
 {
-	check(MetallicTexture->GetPixelFormat() == RoughnessTexture->GetPixelFormat());
-
-	// TODO: Smarter calculation for output-size?
-	const FIntPoint TextureSize(
-		FMath::Max(MetallicTexture->GetSizeX(), RoughnessTexture->GetSizeX()),
-		FMath::Max(MetallicTexture->GetSizeY(), RoughnessTexture->GetSizeY()));
+	check(CombineSources.Num() > 0);
 
 	TArray<FColor> Pixels;
-	const EPixelFormat PixelFormat = MetallicTexture->GetPixelFormat();
-
-	const TArray<FGLTFTextureCombineSource> CombineSources =
-	{
-		{MetallicTexture, FLinearColor(0, 0, 1, 0)},
-		{RoughnessTexture, FLinearColor(0, 1, 0, 0)}
-	};
+	const EPixelFormat PixelFormat = CombineSources[0].Texture->GetPixelFormat();
 
 	if (!CombineTextures(Pixels, CombineSources, TextureSize, PixelFormat))
 	{
diff --git a/Engine/Plugins/Enterprise/GLTFExporter/MigratedFromGit/Source/GLTFExporter/Private/Converters/GLTFMaterialUtility.h b/Engine/Plugins/Enterprise/GLTFExporter/MigratedFromGit/Source/GLTFExporter/Private/Converters/GLTFMaterialUtility.h
--- a/Engine/Plugins/Enterprise/GLTFExporter/MigratedFromGit/Source/GLTFExporter/Private/Converters/GLTFMaterialUtility.h
+++ b/Engine/Plugins/Enterprise/GLTFExporter/MigratedFromGit/Source/GLTFExporter/Private/Converters/GLTFMaterialUtility.h
@@ -22,7 +22,7 @@ struct FGLTFMaterialUtility
 	static bool CombineTextures(TArray<FColor>& OutPixels, const TArray<FGLTFTextureCombineSource>& Sources, const FIntPoint& OutputSize, const EPixelFormat OutputPixelFormat);
 	static UTexture2D* BakeMaterialProperty(const FIntPoint OutputSize, const EMaterialProperty& MaterialProperty, const UMaterialInterface* Material);
 
-	static FGLTFJsonTextureIndex AddMetallicRoughnessTexture(FGLTFConvertBuilder& Builder, const UTexture2D* MetallicTexture, const UTexture2D* RoughnessTexture, EGLTFJsonTextureFilter Filter, EGLTFJsonTextureWrap Wrap, const FString& TextureName = TEXT("MetallicRoughness"));
+	static FGLTFJsonTextureIndex AddCombinedTexture(FGLTFConvertBuilder& Builder, const TArray<FGLTFTextureCombineSource>& CombineSources, const FIntPoint TextureSize, const FString& TextureName, const EGLTFJsonTextureFilter Filter, const EGLTFJsonTextureWrap Wrap);
 	static FIntVector4 ConvertMaskToVector(const FExpressionInput& ExpressionInput);
 	static FLinearColor ConvertMaskToColor(const FExpressionInput& ExpressionInput);
 	static uint32 GetMaskComponentCount(const FExpressionInput& ExpressionInput);
