[GLTFExporter] refactor(material): simplify handling of masks for materials

commit 767d9004fd2eeccf1042d148afc0d0d85a364c7e
Date: Thu, 4 Jun 2020 13:14:29 +0200

    refactor(material): simplify handling of masks for materials
    (cherry picked from commit d3814361bf0e7d1280df3543d74f369ad891bd9c​)

#skipci

diff --git a/Engine/Plugins/Enterprise/GLTFExporter/MigratedFromGit/Source/GLTFExporter/Private/Converters/GLTFMaterialConverters.cpp b/Engine/Plugins/Enterprise/GLTFExporter/MigratedFromGit/Source/GLTFExporter/Private/Converters/GLTFMaterialConverters.cpp
--- a/Engine/Plugins/Enterprise/GLTFExporter/MigratedFromGit/Source/GLTFExporter/Private/Converters/GLTFMaterialConverters.cpp
+++ b/Engine/Plugins/Enterprise/GLTFExporter/MigratedFromGit/Source/GLTFExporter/Private/Converters/GLTFMaterialConverters.cpp
@@ -155,16 +155,16 @@ bool FGLTFMaterialConverter::TryGetMetallicAndRoughness(FGLTFConvertBuilder& Bui
 	const bool bHasMetallicSourceTexture = TryGetSourceTexture(MetallicTexture, MetallicTexCoord, MetallicInput, MaterialInstance);
 	const bool bHasRoughnessSourceTexture = TryGetSourceTexture(RoughnessTexture, RoughnessTexCoord, RoughnessInput, MaterialInstance);
 
+	const FLinearColor MetallicMask(0, 0, 1, 0);
+	const FLinearColor RoughnessMask(0, 1, 0, 0);
+
 	// Detect the "happy path" where both inputs share the same texture, and both inputs correctly masked.
 	if (bHasMetallicSourceTexture && bHasRoughnessSourceTexture)
 	{
-		const FIntVector4 MetallicMask(0, 0, 1, 0);
-		const FIntVector4 RoughnessMask(0, 1, 0, 0);
-
 		if (MetallicTexture == RoughnessTexture &&
 			MetallicTexCoord == RoughnessTexCoord &&
-			MetallicMask == FGLTFMaterialUtility::ConvertMaskToVector(MetallicInput) &&
-			RoughnessMask == FGLTFMaterialUtility::ConvertMaskToVector(RoughnessInput))
+			MetallicMask.Equals(FGLTFMaterialUtility::GetMask(MetallicInput)) &&
+			RoughnessMask.Equals(FGLTFMaterialUtility::GetMask(RoughnessInput)))
 		{
 			OutPBRParams.MetallicRoughnessTexture.Index = Builder.GetOrAddTexture(MetallicTexture);
 			OutPBRParams.MetallicRoughnessTexture.TexCoord = MetallicTexCoord;
@@ -230,8 +230,8 @@ bool FGLTFMaterialConverter::TryGetMetallicAndRoughness(FGLTFConvertBuilder& Bui
 
 	const TArray<FGLTFTextureCombineSource> CombineSources =
 	{
-		{MetallicTexture, FLinearColor(0, 0, 1, 0)},
-		{RoughnessTexture, FLinearColor(0, 1, 0, 0)}
+		{MetallicTexture, MetallicMask},
+		{RoughnessTexture, RoughnessMask}
 	};
 
 	const FGLTFJsonTextureIndex TextureIndex = FGLTFMaterialUtility::AddCombinedTexture(
@@ -299,7 +299,7 @@ bool FGLTFMaterialConverter::TryGetConstantColor(FLinearColor& OutValue, const F
 
 		if (MaskComponentCount > 0)
 		{
-			const FLinearColor Mask = FGLTFMaterialUtility::ConvertMaskToColor(MaterialInput);
+			const FLinearColor Mask = FGLTFMaterialUtility::GetMask(MaterialInput);
 
 			Value *= Mask;
 
@@ -385,7 +385,7 @@ bool FGLTFMaterialConverter::TryGetConstantScalar(float& OutValue, const FScalar
 
 		if (MaskComponentCount > 0)
 		{
-			const FLinearColor Mask = FGLTFMaterialUtility::ConvertMaskToColor(MaterialInput);
+			const FLinearColor Mask = FGLTFMaterialUtility::GetMask(MaterialInput);
 			Value *= Mask;
 		}
 
diff --git a/Engine/Plugins/Enterprise/GLTFExporter/MigratedFromGit/Source/GLTFExporter/Private/Converters/GLTFMaterialUtility.cpp b/Engine/Plugins/Enterprise/GLTFExporter/MigratedFromGit/Source/GLTFExporter/Private/Converters/GLTFMaterialUtility.cpp
--- a/Engine/Plugins/Enterprise/GLTFExporter/MigratedFromGit/Source/GLTFExporter/Private/Converters/GLTFMaterialUtility.cpp
+++ b/Engine/Plugins/Enterprise/GLTFExporter/MigratedFromGit/Source/GLTFExporter/Private/Converters/GLTFMaterialUtility.cpp
@@ -136,18 +136,9 @@ FGLTFJsonTextureIndex FGLTFMaterialUtility::AddCombinedTexture(FGLTFConvertBuild
 	return Builder.AddTexture(JsonTexture);
 }
 
-FIntVector4 FGLTFMaterialUtility::ConvertMaskToVector(const FExpressionInput& ExpressionInput)
+FLinearColor FGLTFMaterialUtility::GetMask(const FExpressionInput& ExpressionInput)
 {
-	// TODO: Do we need to check MaterialInput.Mask first (which seems to always be 0 or 1, not an actual bitfield)?
-	// Or is it enough to assume that MaskR, MaskG, MaskB and MaskA will only be 1 when Mask is 1?
-	return FIntVector4(ExpressionInput.MaskR, ExpressionInput.MaskG, ExpressionInput.MaskB, ExpressionInput.MaskA);
-}
-
-FLinearColor FGLTFMaterialUtility::ConvertMaskToColor(const FExpressionInput& ExpressionInput)
-{
-	const FIntVector4 Mask = ConvertMaskToVector(ExpressionInput);
-
-	return FLinearColor(Mask.X, Mask.Y, Mask.Z, Mask.W);
+	return FLinearColor(ExpressionInput.MaskR, ExpressionInput.MaskG, ExpressionInput.MaskB, ExpressionInput.MaskA);
 }
 
 uint32 FGLTFMaterialUtility::GetMaskComponentCount(const FExpressionInput& ExpressionInput)
diff --git a/Engine/Plugins/Enterprise/GLTFExporter/MigratedFromGit/Source/GLTFExporter/Private/Converters/GLTFMaterialUtility.h b/Engine/Plugins/Enterprise/GLTFExporter/MigratedFromGit/Source/GLTFExporter/Private/Converters/GLTFMaterialUtility.h
--- a/Engine/Plugins/Enterprise/GLTFExporter/MigratedFromGit/Source/GLTFExporter/Private/Converters/GLTFMaterialUtility.h
+++ b/Engine/Plugins/Enterprise/GLTFExporter/MigratedFromGit/Source/GLTFExporter/Private/Converters/GLTFMaterialUtility.h
@@ -23,7 +23,6 @@ struct FGLTFMaterialUtility
 	static UTexture2D* BakeMaterialProperty(const FIntPoint OutputSize, const EMaterialProperty& MaterialProperty, const UMaterialInterface* Material);
 
 	static FGLTFJsonTextureIndex AddCombinedTexture(FGLTFConvertBuilder& Builder, const TArray<FGLTFTextureCombineSource>& CombineSources, const FIntPoint TextureSize, const FString& TextureName, const EGLTFJsonTextureFilter Filter, const EGLTFJsonTextureWrap Wrap);
-	static FIntVector4 ConvertMaskToVector(const FExpressionInput& ExpressionInput);
-	static FLinearColor ConvertMaskToColor(const FExpressionInput& ExpressionInput);
+	static FLinearColor GetMask(const FExpressionInput& ExpressionInput);
 	static uint32 GetMaskComponentCount(const FExpressionInput& ExpressionInput);
 };
