[GLTFExporter] feat(material): add ability to copy values from the red-channel to alpha when baking material-properties

commit 5c2b5fd0e9e67b7456badd7bf487d801092b39fb
Date: Thu, 4 Jun 2020 15:31:34 +0200

    feat(material): add ability to copy values from the red-channel to alpha when baking material-properties
    (cherry picked from commit 879360835be94f09089fdf9ac34dd65659e46ce3​)

#skipci

diff --git a/Engine/Plugins/Enterprise/GLTFExporter/MigratedFromGit/Source/GLTFExporter/Private/Converters/GLTFMaterialUtility.cpp b/Engine/Plugins/Enterprise/GLTFExporter/MigratedFromGit/Source/GLTFExporter/Private/Converters/GLTFMaterialUtility.cpp
--- a/Engine/Plugins/Enterprise/GLTFExporter/MigratedFromGit/Source/GLTFExporter/Private/Converters/GLTFMaterialUtility.cpp
+++ b/Engine/Plugins/Enterprise/GLTFExporter/MigratedFromGit/Source/GLTFExporter/Private/Converters/GLTFMaterialUtility.cpp
@@ -74,7 +74,7 @@ bool FGLTFMaterialUtility::CombineTextures(TArray<FColor>& OutPixels, const TArr
 	return bReadSuccessful;
 }
 
-UTexture2D* FGLTFMaterialUtility::BakeMaterialProperty(const FIntPoint OutputSize, const EMaterialProperty& MaterialProperty, const UMaterialInterface* Material)
+UTexture2D* FGLTFMaterialUtility::BakeMaterialProperty(const FIntPoint OutputSize, const EMaterialProperty& MaterialProperty, const UMaterialInterface* Material, const bool bCopyAlphaFromRedChannel)
 {
 	TArray<FMeshData*> MeshSettings;
 
@@ -101,13 +101,28 @@ UTexture2D* FGLTFMaterialUtility::BakeMaterialProperty(const FIntPoint OutputSiz
 
 	const FBakeOutput& Output = BakeOutputs[0];
 
+	TArray<FColor> BakedPixels = Output.PropertyData.FindChecked(MaterialProperty);
+	const FIntPoint BakedSize = Output.PropertySizes.FindChecked(MaterialProperty);
+
+	if (bCopyAlphaFromRedChannel)
+	{
+		// NOTE: Baked property-textures only have values in RGB, not A.
+		// If someone wants to use the alpha-channel of the baked texture (for compositing),
+		// it needs to be filled with values first.
+		// Since scalar properties are baked to R + G + B, it doesn't matter which channel
+		// we use as source. For now it's set to R.
+		// TODO: Determine if we can get the baking module to somehow output values in A automatically.
+		// If not, see if the copy-operation can be performed via the canvas (using a proxy-material or similar).
+		for (FColor& Pixel: BakedPixels)
+		{
+			Pixel.A = Pixel.R;
+		}
+	}
+
 	// NOTE: The texture-format is mapped per property in the MaterialBaking-module via PerPropertyFormat.
 	// It's PF_B8G8R8A8 for all properties except for MP_EmissiveColor (which is PF_FloatRGBA).
 	// TODO: Should we handle cases the format *may* be something other that what we expect?
-	return CreateTransientTexture(
-		Output.PropertyData.FindChecked(MaterialProperty),
-		Output.PropertySizes.FindChecked(MaterialProperty),
-		PF_B8G8R8A8);
+	return CreateTransientTexture(BakedPixels, BakedSize, PF_B8G8R8A8);
 }
 
 FGLTFJsonTextureIndex FGLTFMaterialUtility::AddCombinedTexture(FGLTFConvertBuilder& Builder, const TArray<FGLTFTextureCombineSource>& CombineSources, const FIntPoint TextureSize, const FString& TextureName, const EGLTFJsonTextureFilter Filter, const EGLTFJsonTextureWrap Wrap)
diff --git a/Engine/Plugins/Enterprise/GLTFExporter/MigratedFromGit/Source/GLTFExporter/Private/Converters/GLTFMaterialUtility.h b/Engine/Plugins/Enterprise/GLTFExporter/MigratedFromGit/Source/GLTFExporter/Private/Converters/GLTFMaterialUtility.h
--- a/Engine/Plugins/Enterprise/GLTFExporter/MigratedFromGit/Source/GLTFExporter/Private/Converters/GLTFMaterialUtility.h
+++ b/Engine/Plugins/Enterprise/GLTFExporter/MigratedFromGit/Source/GLTFExporter/Private/Converters/GLTFMaterialUtility.h
@@ -20,7 +20,7 @@ struct FGLTFMaterialUtility
 {
 	static UTexture2D* CreateTransientTexture(const TArray<FColor>& Pixels, const FIntPoint& TextureSize, const EPixelFormat& TextureFormat);
 	static bool CombineTextures(TArray<FColor>& OutPixels, const TArray<FGLTFTextureCombineSource>& Sources, const FIntPoint& OutputSize, const EPixelFormat OutputPixelFormat);
-	static UTexture2D* BakeMaterialProperty(const FIntPoint OutputSize, const EMaterialProperty& MaterialProperty, const UMaterialInterface* Material);
+	static UTexture2D* BakeMaterialProperty(const FIntPoint OutputSize, const EMaterialProperty& MaterialProperty, const UMaterialInterface* Material, const bool bCopyAlphaFromRedChannel = false);
 
 	static FGLTFJsonTextureIndex AddCombinedTexture(FGLTFConvertBuilder& Builder, const TArray<FGLTFTextureCombineSource>& CombineSources, const FIntPoint TextureSize, const FString& TextureName, const EGLTFJsonTextureFilter Filter, const EGLTFJsonTextureWrap Wrap);
 	static FLinearColor GetMask(const FExpressionInput& ExpressionInput);
