[GLTFExporter] feat(material): initial support for baked basecolor-opacity texture

commit 011994777b02ae5f3c857b82ca092f3c30788a34
Date: Thu, 4 Jun 2020 15:32:39 +0200

    feat(material): initial support for baked basecolor-opacity texture
    (cherry picked from commit 0d9e4091fcc0a056697907f438750a7701847064​)

#skipci

diff --git a/Engine/Plugins/Enterprise/GLTFExporter/MigratedFromGit/Source/GLTFExporter/Private/Converters/GLTFMaterialConverters.cpp b/Engine/Plugins/Enterprise/GLTFExporter/MigratedFromGit/Source/GLTFExporter/Private/Converters/GLTFMaterialConverters.cpp
--- a/Engine/Plugins/Enterprise/GLTFExporter/MigratedFromGit/Source/GLTFExporter/Private/Converters/GLTFMaterialConverters.cpp
+++ b/Engine/Plugins/Enterprise/GLTFExporter/MigratedFromGit/Source/GLTFExporter/Private/Converters/GLTFMaterialConverters.cpp
@@ -30,9 +30,9 @@ FGLTFJsonMaterialIndex FGLTFMaterialConverter::Add(FGLTFConvertBuilder& Builder,
 	if (JsonMaterial.AlphaMode == EGLTFJsonAlphaMode::Blend || JsonMaterial.AlphaMode == EGLTFJsonAlphaMode::Mask)
 	{
 		const FScalarMaterialInput& OpacityInput = JsonMaterial.AlphaMode == EGLTFJsonAlphaMode::Blend ? Material->Opacity : Material->OpacityMask;
-		if (!TryGetBaseColorAndOpacity(Builder, JsonMaterial.PBRMetallicRoughness, Material->BaseColor, OpacityInput, MaterialInstance))
+		if (!TryGetBaseColorAndOpacity(Builder, JsonMaterial.PBRMetallicRoughness, Material->BaseColor, OpacityInput, MaterialInterface))
 		{
-			// TODO: add fallback to material baking
+			// TODO: handle failure?
 		}
 	}
 	else
@@ -82,8 +82,10 @@ FGLTFJsonMaterialIndex FGLTFMaterialConverter::Add(FGLTFConvertBuilder& Builder,
 	return Builder.AddMaterial(JsonMaterial);
 }
 
-bool FGLTFMaterialConverter::TryGetBaseColorAndOpacity(FGLTFConvertBuilder& Builder, FGLTFJsonPBRMetallicRoughness& OutPBRParams, const FColorMaterialInput& BaseColorInput, const FScalarMaterialInput& OpacityInput, const UMaterialInstance* MaterialInstance) const
+bool FGLTFMaterialConverter::TryGetBaseColorAndOpacity(FGLTFConvertBuilder& Builder, FGLTFJsonPBRMetallicRoughness& OutPBRParams, const FColorMaterialInput& BaseColorInput, const FScalarMaterialInput& OpacityInput, const UMaterialInterface* MaterialInterface) const
 {
+	const UMaterialInstance* MaterialInstance = Cast<UMaterialInstance>(MaterialInterface);
+
 	const bool bIsBaseColorConstant = TryGetConstantColor(OutPBRParams.BaseColorFactor, BaseColorInput, MaterialInstance);
 	const bool bIsOpacityConstant = TryGetConstantScalar(OutPBRParams.BaseColorFactor.A, OpacityInput, MaterialInstance);
 
@@ -92,6 +94,11 @@ bool FGLTFMaterialConverter::TryGetBaseColorAndOpacity(FGLTFConvertBuilder& Buil
 		return true;
 	}
 
+	// NOTE: since we always bake the properties (for now) when atleast property is non-const, we need
+	// to reset the constant factors to their defaults. Otherwise the baked value of a constant property
+	// would be scaled with the factor, i.e a double scaling.
+	OutPBRParams.BaseColorFactor = {1, 1, 1, 1};
+
 	const UTexture2D* BaseColorTexture;
 	const UTexture2D* OpacityTexture;
 	int32 BaseColorTexCoord;
@@ -100,33 +107,98 @@ bool FGLTFMaterialConverter::TryGetBaseColorAndOpacity(FGLTFConvertBuilder& Buil
 	const bool bHasBaseColorSourceTexture = TryGetSourceTexture(BaseColorTexture, BaseColorTexCoord, BaseColorInput, MaterialInstance);
 	const bool bHasOpacitySourceTexture = TryGetSourceTexture(OpacityTexture, OpacityTexCoord, OpacityInput, MaterialInstance);
 
+	const FLinearColor BaseColorMask(1, 1, 1, 0);
+	const FLinearColor OpacityMask(0, 0, 0, 1);
+
+	// Detect the "happy path" where both inputs share the same texture, and both inputs correctly masked.
 	if (bHasBaseColorSourceTexture && bHasOpacitySourceTexture)
 	{
-		if (BaseColorTexture == OpacityTexture && BaseColorTexCoord == OpacityTexCoord)
+		if (BaseColorTexture == OpacityTexture &&
+			BaseColorTexCoord == OpacityTexCoord &&
+			BaseColorMask.Equals(FGLTFMaterialUtility::GetMask(BaseColorInput)) &&
+			OpacityMask.Equals(FGLTFMaterialUtility::GetMask(OpacityInput)))
 		{
-			// TODO: make sure textures are correctly masked
 			OutPBRParams.BaseColorTexture.Index = Builder.GetOrAddTexture(BaseColorTexture);
 			OutPBRParams.BaseColorTexture.TexCoord = BaseColorTexCoord;
 			return true;
 		}
-
-		// TODO: add support for combining two textures
-		return false;
 	}
 
-	if (bHasBaseColorSourceTexture && bIsOpacityConstant)
+	int32 TexCoord = 0;
+	FIntPoint TextureSize(512, 512);	// TODO: make default baking-resolution configurable
+	EGLTFJsonTextureWrap TextureWrap = EGLTFJsonTextureWrap::ClampToEdge;
+	EGLTFJsonTextureFilter TextureFilter = EGLTFJsonTextureFilter::LinearMipmapLinear;
+
+	if (bHasBaseColorSourceTexture && bHasOpacitySourceTexture)
 	{
-		// TODO: add support for combining constant with texture
-		return false;
-	}
+		const bool bAreTexturesCompatible = BaseColorTexCoord == OpacityTexCoord &&
+			BaseColorTexture->AddressX == OpacityTexture->AddressX &&
+			BaseColorTexture->AddressY == OpacityTexture->AddressY;
+
+		if (!bAreTexturesCompatible)
+		{
+			// TODO: handle differences in wrapping or uv-coords
+			return false;
+		}
 
-	if (bIsBaseColorConstant && bHasOpacitySourceTexture)
+		TexCoord = BaseColorTexCoord;
+		TextureSize = FIntPoint(
+			FMath::Max(BaseColorTexture->GetSizeX(), OpacityTexture->GetSizeX()),
+			FMath::Max(BaseColorTexture->GetSizeY(), OpacityTexture->GetSizeY()));
+
+		TextureWrap = FGLTFConverterUtility::ConvertWrap(BaseColorTexture->AddressX);
+		TextureFilter = FGLTFConverterUtility::ConvertFilter(BaseColorTexture->Filter, BaseColorTexture->LODGroup);
+	}
+	else if (bHasBaseColorSourceTexture)
 	{
-		// TODO: add support for combining constant with texture
-		return false;
+		TexCoord = BaseColorTexCoord;
+		TextureSize = FIntPoint(BaseColorTexture->GetSizeX(), BaseColorTexture->GetSizeY());
+		TextureWrap = FGLTFConverterUtility::ConvertWrap(BaseColorTexture->AddressX);
+		TextureFilter = FGLTFConverterUtility::ConvertFilter(BaseColorTexture->Filter, BaseColorTexture->LODGroup);
+	}
+	else if (bHasOpacitySourceTexture)
+	{
+		TexCoord = OpacityTexCoord;
+		TextureSize = FIntPoint(OpacityTexture->GetSizeX(), OpacityTexture->GetSizeY());
+		TextureWrap = FGLTFConverterUtility::ConvertWrap(OpacityTexture->AddressX);
+		TextureFilter = FGLTFConverterUtility::ConvertFilter(OpacityTexture->Filter, OpacityTexture->LODGroup);
 	}
 
-	return false;
+	const EMaterialProperty OpacityProperty = MaterialInterface->GetBlendMode() == BLEND_Masked ? MP_OpacityMask : MP_Opacity;
+
+	BaseColorTexture = FGLTFMaterialUtility::BakeMaterialProperty(TextureSize, MP_BaseColor, MaterialInterface);
+	OpacityTexture = FGLTFMaterialUtility::BakeMaterialProperty(TextureSize, OpacityProperty, MaterialInterface, true);
+
+	// NOTE: the baked textures may have a different (smaller) size than requested, so we update the
+	// texture-size to fit both textures without wasting too much space.
+	TextureSize = {
+		FMath::Max(BaseColorTexture->GetSizeX(), OpacityTexture->GetSizeX()),
+		FMath::Max(BaseColorTexture->GetSizeY(), OpacityTexture->GetSizeY())
+	};
+
+	// TODO: handle the case where TextureSize is 1x1. In this case, both properties are constants and we should
+	// extract the value of each property from the texture and use as-is instead of exporting a combined texture.
+
+	const FString TextureName = FString::Printf(TEXT("%s_BaseColor"), *MaterialInterface->GetName());
+
+	const TArray<FGLTFTextureCombineSource> CombineSources =
+	{
+		{OpacityTexture, OpacityMask, SE_BLEND_Opaque},
+		{BaseColorTexture, BaseColorMask}
+	};
+
+	const FGLTFJsonTextureIndex TextureIndex = FGLTFMaterialUtility::AddCombinedTexture(
+		Builder,
+		CombineSources,
+		TextureSize,
+		TextureName,
+		TextureFilter,
+		TextureWrap);
+
+	OutPBRParams.BaseColorTexture.TexCoord = TexCoord;
+	OutPBRParams.BaseColorTexture.Index = TextureIndex;
+
+	return true;
 }
 
 bool FGLTFMaterialConverter::TryGetMetallicAndRoughness(FGLTFConvertBuilder& Builder, FGLTFJsonPBRMetallicRoughness& OutPBRParams, const FScalarMaterialInput& MetallicInput, const FScalarMaterialInput& RoughnessInput, const UMaterialInterface* MaterialInterface) const
diff --git a/Engine/Plugins/Enterprise/GLTFExporter/MigratedFromGit/Source/GLTFExporter/Private/Converters/GLTFMaterialConverters.h b/Engine/Plugins/Enterprise/GLTFExporter/MigratedFromGit/Source/GLTFExporter/Private/Converters/GLTFMaterialConverters.h
--- a/Engine/Plugins/Enterprise/GLTFExporter/MigratedFromGit/Source/GLTFExporter/Private/Converters/GLTFMaterialConverters.h
+++ b/Engine/Plugins/Enterprise/GLTFExporter/MigratedFromGit/Source/GLTFExporter/Private/Converters/GLTFMaterialConverters.h
@@ -11,7 +11,7 @@ class FGLTFMaterialConverter final : public TGLTFConverter<FGLTFJsonMaterialInde
 {
 	FGLTFJsonMaterialIndex Add(FGLTFConvertBuilder& Builder, const FString& Name, const UMaterialInterface* Material) override;
 
-	bool TryGetBaseColorAndOpacity(FGLTFConvertBuilder& Builder, FGLTFJsonPBRMetallicRoughness& OutPBRParams, const FColorMaterialInput& BaseColorInput, const FScalarMaterialInput& OpacityInput, const UMaterialInstance* MaterialInstance) const;
+	bool TryGetBaseColorAndOpacity(FGLTFConvertBuilder& Builder, FGLTFJsonPBRMetallicRoughness& OutPBRParams, const FColorMaterialInput& BaseColorInput, const FScalarMaterialInput& OpacityInput, const UMaterialInterface* MaterialInterface) const;
 	bool TryGetMetallicAndRoughness(FGLTFConvertBuilder& Builder, FGLTFJsonPBRMetallicRoughness& OutPBRParams, const FScalarMaterialInput& MetallicInput, const FScalarMaterialInput& RoughnessInput, const UMaterialInterface* MaterialInterface) const;
 
 	bool TryGetConstantColor(FGLTFJsonColor3& OutValue, const FColorMaterialInput& MaterialInput, const UMaterialInstance* MaterialInstance) const;
