[GLTFExporter] refactor(material): remove reduntant const specifiers on parameters

commit 964e9145c65f0c937b398bad366fecefecd76841
Date: Fri, 5 Jun 2020 08:09:27 +0200

    refactor(material): remove reduntant const specifiers on parameters
    (cherry picked from commit fff2a357e93f7e47bbf50c717eea9a059c690349​)

#skipci

diff --git a/Engine/Plugins/Enterprise/GLTFExporter/MigratedFromGit/Source/GLTFExporter/Private/Converters/GLTFMaterialUtility.cpp b/Engine/Plugins/Enterprise/GLTFExporter/MigratedFromGit/Source/GLTFExporter/Private/Converters/GLTFMaterialUtility.cpp
--- a/Engine/Plugins/Enterprise/GLTFExporter/MigratedFromGit/Source/GLTFExporter/Private/Converters/GLTFMaterialUtility.cpp
+++ b/Engine/Plugins/Enterprise/GLTFExporter/MigratedFromGit/Source/GLTFExporter/Private/Converters/GLTFMaterialUtility.cpp
@@ -10,7 +10,7 @@
 #include "IMaterialBakingModule.h"
 #include "MaterialBakingStructures.h"
 
-UTexture2D* FGLTFMaterialUtility::CreateTransientTexture(const TArray<FColor>& Pixels, const FIntPoint& TextureSize, const EPixelFormat& TextureFormat)
+UTexture2D* FGLTFMaterialUtility::CreateTransientTexture(const TArray<FColor>& Pixels, const FIntPoint& TextureSize, EPixelFormat TextureFormat)
 {
 	check(TextureSize.X * TextureSize.Y == Pixels.Num());
 
@@ -74,7 +74,7 @@ bool FGLTFMaterialUtility::CombineTextures(TArray<FColor>& OutPixels, const TArr
 	return bReadSuccessful;
 }
 
-UTexture2D* FGLTFMaterialUtility::BakeMaterialProperty(const FIntPoint OutputSize, const EMaterialProperty& MaterialProperty, const UMaterialInterface* Material, const bool bCopyAlphaFromRedChannel)
+UTexture2D* FGLTFMaterialUtility::BakeMaterialProperty(const FIntPoint& OutputSize, EMaterialProperty MaterialProperty, const UMaterialInterface* Material, bool bCopyAlphaFromRedChannel)
 {
 	TArray<FMeshData*> MeshSettings;
 
@@ -125,7 +125,7 @@ UTexture2D* FGLTFMaterialUtility::BakeMaterialProperty(const FIntPoint OutputSiz
 	return CreateTransientTexture(BakedPixels, BakedSize, PF_B8G8R8A8);
 }
 
-FGLTFJsonTextureIndex FGLTFMaterialUtility::AddCombinedTexture(FGLTFConvertBuilder& Builder, const TArray<FGLTFTextureCombineSource>& CombineSources, const FIntPoint TextureSize, const FString& TextureName, const EGLTFJsonTextureFilter Filter, const EGLTFJsonTextureWrap Wrap)
+FGLTFJsonTextureIndex FGLTFMaterialUtility::AddCombinedTexture(FGLTFConvertBuilder& Builder, const TArray<FGLTFTextureCombineSource>& CombineSources, const FIntPoint& TextureSize, const FString& TextureName, EGLTFJsonTextureFilter Filter, EGLTFJsonTextureWrap Wrap)
 {
 	check(CombineSources.Num() > 0);
 
diff --git a/Engine/Plugins/Enterprise/GLTFExporter/MigratedFromGit/Source/GLTFExporter/Private/Converters/GLTFMaterialUtility.h b/Engine/Plugins/Enterprise/GLTFExporter/MigratedFromGit/Source/GLTFExporter/Private/Converters/GLTFMaterialUtility.h
--- a/Engine/Plugins/Enterprise/GLTFExporter/MigratedFromGit/Source/GLTFExporter/Private/Converters/GLTFMaterialUtility.h
+++ b/Engine/Plugins/Enterprise/GLTFExporter/MigratedFromGit/Source/GLTFExporter/Private/Converters/GLTFMaterialUtility.h
@@ -18,11 +18,11 @@ struct FGLTFTextureCombineSource
 
 struct FGLTFMaterialUtility
 {
-	static UTexture2D* CreateTransientTexture(const TArray<FColor>& Pixels, const FIntPoint& TextureSize, const EPixelFormat& TextureFormat);
-	static bool CombineTextures(TArray<FColor>& OutPixels, const TArray<FGLTFTextureCombineSource>& Sources, const FIntPoint& OutputSize, const EPixelFormat OutputPixelFormat);
-	static UTexture2D* BakeMaterialProperty(const FIntPoint OutputSize, const EMaterialProperty& MaterialProperty, const UMaterialInterface* Material, const bool bCopyAlphaFromRedChannel = false);
+	static UTexture2D* CreateTransientTexture(const TArray<FColor>& Pixels, const FIntPoint& TextureSize, EPixelFormat TextureFormat);
+	static bool CombineTextures(TArray<FColor>& OutPixels, const TArray<FGLTFTextureCombineSource>& Sources, const FIntPoint& OutputSize, EPixelFormat OutputPixelFormat);
+	static UTexture2D* BakeMaterialProperty(const FIntPoint& OutputSize, EMaterialProperty MaterialProperty, const UMaterialInterface* Material, bool bCopyAlphaFromRedChannel = false);
 
-	static FGLTFJsonTextureIndex AddCombinedTexture(FGLTFConvertBuilder& Builder, const TArray<FGLTFTextureCombineSource>& CombineSources, const FIntPoint TextureSize, const FString& TextureName, const EGLTFJsonTextureFilter Filter, const EGLTFJsonTextureWrap Wrap);
+	static FGLTFJsonTextureIndex AddCombinedTexture(FGLTFConvertBuilder& Builder, const TArray<FGLTFTextureCombineSource>& CombineSources, const FIntPoint& TextureSize, const FString& TextureName, EGLTFJsonTextureFilter Filter, EGLTFJsonTextureWrap Wrap);
 	static FLinearColor GetMask(const FExpressionInput& ExpressionInput);
 	static uint32 GetMaskComponentCount(const FExpressionInput& ExpressionInput);
 };
