[GLTFExporter] refactor(material): make baking-functionality more flexible

commit 1f8e245f493172059a438e14289f6d351d2a9068
Date: Fri, 5 Jun 2020 10:45:20 +0200

    refactor(material): make baking-functionality more flexible
    (cherry picked from commit b14d3132e135558b64d2d1c28143ceaf18935747​)

#skipci

diff --git a/Engine/Plugins/Enterprise/GLTFExporter/MigratedFromGit/Source/GLTFExporter/Private/Converters/GLTFMaterialConverters.cpp b/Engine/Plugins/Enterprise/GLTFExporter/MigratedFromGit/Source/GLTFExporter/Private/Converters/GLTFMaterialConverters.cpp
--- a/Engine/Plugins/Enterprise/GLTFExporter/MigratedFromGit/Source/GLTFExporter/Private/Converters/GLTFMaterialConverters.cpp
+++ b/Engine/Plugins/Enterprise/GLTFExporter/MigratedFromGit/Source/GLTFExporter/Private/Converters/GLTFMaterialConverters.cpp
@@ -166,8 +166,8 @@ bool FGLTFMaterialConverter::TryGetBaseColorAndOpacity(FGLTFConvertBuilder& Buil
 
 	const EMaterialProperty OpacityProperty = MaterialInterface->GetBlendMode() == BLEND_Masked ? MP_OpacityMask : MP_Opacity;
 
-	BaseColorTexture = FGLTFMaterialUtility::BakeMaterialProperty(TextureSize, MP_BaseColor, MaterialInterface);
-	OpacityTexture = FGLTFMaterialUtility::BakeMaterialProperty(TextureSize, OpacityProperty, MaterialInterface, true);
+	BaseColorTexture = FGLTFMaterialUtility::BakeMaterialPropertyToTexture(TextureSize, MP_BaseColor, MaterialInterface);
+	OpacityTexture = FGLTFMaterialUtility::BakeMaterialPropertyToTexture(TextureSize, OpacityProperty, MaterialInterface, true);
 
 	// NOTE: the baked textures may have a different (smaller) size than requested, so we update the
 	// texture-size to fit both textures without wasting too much space.
@@ -285,8 +285,8 @@ bool FGLTFMaterialConverter::TryGetMetallicAndRoughness(FGLTFConvertBuilder& Bui
 		TextureFilter = FGLTFConverterUtility::ConvertFilter(RoughnessTexture->Filter, RoughnessTexture->LODGroup);
 	}
 
-	RoughnessTexture = FGLTFMaterialUtility::BakeMaterialProperty(TextureSize, MP_Roughness, MaterialInterface);
-	MetallicTexture = FGLTFMaterialUtility::BakeMaterialProperty(TextureSize, MP_Metallic, MaterialInterface);
+	RoughnessTexture = FGLTFMaterialUtility::BakeMaterialPropertyToTexture(TextureSize, MP_Roughness, MaterialInterface);
+	MetallicTexture = FGLTFMaterialUtility::BakeMaterialPropertyToTexture(TextureSize, MP_Metallic, MaterialInterface);
 
 	// NOTE: the baked textures may have a different (smaller) size than requested, so we update the
 	// texture-size to fit both textures without wasting too much space.
diff --git a/Engine/Plugins/Enterprise/GLTFExporter/MigratedFromGit/Source/GLTFExporter/Private/Converters/GLTFMaterialUtility.cpp b/Engine/Plugins/Enterprise/GLTFExporter/MigratedFromGit/Source/GLTFExporter/Private/Converters/GLTFMaterialUtility.cpp
--- a/Engine/Plugins/Enterprise/GLTFExporter/MigratedFromGit/Source/GLTFExporter/Private/Converters/GLTFMaterialUtility.cpp
+++ b/Engine/Plugins/Enterprise/GLTFExporter/MigratedFromGit/Source/GLTFExporter/Private/Converters/GLTFMaterialUtility.cpp
@@ -74,7 +74,7 @@ bool FGLTFMaterialUtility::CombineTextures(TArray<FColor>& OutPixels, const TArr
 	return bReadSuccessful;
 }
 
-UTexture2D* FGLTFMaterialUtility::BakeMaterialProperty(const FIntPoint& OutputSize, EMaterialProperty MaterialProperty, const UMaterialInterface* Material, bool bCopyAlphaFromRedChannel)
+FGLTFPropertyBakeOutput FGLTFMaterialUtility::BakeMaterialProperty(const FIntPoint& OutputSize, EMaterialProperty MaterialProperty, const UMaterialInterface* Material, bool bCopyAlphaFromRedChannel)
 {
 	TArray<FMeshData*> MeshSettings;
 
@@ -99,10 +99,8 @@ UTexture2D* FGLTFMaterialUtility::BakeMaterialProperty(const FIntPoint& OutputSi
 
 	Module.BakeMaterials(MatSettings, MeshSettings, BakeOutputs);
 
-	const FBakeOutput& Output = BakeOutputs[0];
-
-	TArray<FColor> BakedPixels = Output.PropertyData.FindChecked(MaterialProperty);
-	const FIntPoint BakedSize = Output.PropertySizes.FindChecked(MaterialProperty);
+	TArray<FColor> BakedPixels = MoveTemp(BakeOutputs[0].PropertyData.FindChecked(MaterialProperty));
+	const FIntPoint BakedSize = BakeOutputs[0].PropertySizes.FindChecked(MaterialProperty);
 
 	if (bCopyAlphaFromRedChannel)
 	{
@@ -119,10 +117,28 @@ UTexture2D* FGLTFMaterialUtility::BakeMaterialProperty(const FIntPoint& OutputSi
 		}
 	}
 
+	EPixelFormat PixelFormat;
+
 	// NOTE: The texture-format is mapped per property in the MaterialBaking-module via PerPropertyFormat.
 	// It's PF_B8G8R8A8 for all properties except for MP_EmissiveColor (which is PF_FloatRGBA).
-	// TODO: Should we handle cases the format *may* be something other that what we expect?
-	return CreateTransientTexture(BakedPixels, BakedSize, PF_B8G8R8A8);
+	switch (MaterialProperty)
+	{
+		case MP_EmissiveColor: PixelFormat = PF_FloatRGBA; break;
+		default: PixelFormat = PF_B8G8R8A8; break;
+	}
+
+	return {
+		MaterialProperty,
+		PixelFormat,
+		BakedPixels,
+		BakedSize
+	};
+}
+
+UTexture2D* FGLTFMaterialUtility::BakeMaterialPropertyToTexture(const FIntPoint& OutputSize, EMaterialProperty MaterialProperty, const UMaterialInterface* Material, bool bCopyAlphaFromRedChannel)
+{
+	const FGLTFPropertyBakeOutput PropertyBakeOutput = BakeMaterialProperty(OutputSize, MaterialProperty, Material, bCopyAlphaFromRedChannel);
+	return CreateTransientTexture(PropertyBakeOutput.Pixels, PropertyBakeOutput.Size, PropertyBakeOutput.PixelFormat);
 }
 
 FGLTFJsonTextureIndex FGLTFMaterialUtility::AddCombinedTexture(FGLTFConvertBuilder& Builder, const TArray<FGLTFTextureCombineSource>& CombineSources, const FIntPoint& TextureSize, const FString& TextureName, EGLTFJsonTextureFilter Filter, EGLTFJsonTextureWrap Wrap)
@@ -137,6 +153,11 @@ FGLTFJsonTextureIndex FGLTFMaterialUtility::AddCombinedTexture(FGLTFConvertBuild
 		return FGLTFJsonTextureIndex(INDEX_NONE);
 	}
 
+	return AddTexture(Builder, Pixels, TextureSize, TextureName, PixelFormat, Filter, Wrap);
+}
+
+FGLTFJsonTextureIndex FGLTFMaterialUtility::AddTexture(FGLTFConvertBuilder& Builder, const TArray<FColor>& Pixels, const FIntPoint& TextureSize, const FString& TextureName, EPixelFormat PixelFormat, EGLTFJsonTextureFilter Filter, EGLTFJsonTextureWrap Wrap)
+{
 	// TODO: maybe we should reuse existing samplers?
 	FGLTFJsonSampler JsonSampler;
 	JsonSampler.Name = TextureName;
diff --git a/Engine/Plugins/Enterprise/GLTFExporter/MigratedFromGit/Source/GLTFExporter/Private/Converters/GLTFMaterialUtility.h b/Engine/Plugins/Enterprise/GLTFExporter/MigratedFromGit/Source/GLTFExporter/Private/Converters/GLTFMaterialUtility.h
--- a/Engine/Plugins/Enterprise/GLTFExporter/MigratedFromGit/Source/GLTFExporter/Private/Converters/GLTFMaterialUtility.h
+++ b/Engine/Plugins/Enterprise/GLTFExporter/MigratedFromGit/Source/GLTFExporter/Private/Converters/GLTFMaterialUtility.h
@@ -16,13 +16,28 @@ struct FGLTFTextureCombineSource
 	ESimpleElementBlendMode BlendMode;
 };
 
+struct FGLTFPropertyBakeOutput
+{
+	FORCEINLINE FGLTFPropertyBakeOutput(EMaterialProperty Property, EPixelFormat PixelFormat, TArray<FColor>& Pixels, FIntPoint Size)
+		: Property(Property), PixelFormat(PixelFormat), Pixels(Pixels), Size(Size)
+	{}
+
+	EMaterialProperty Property;
+	EPixelFormat PixelFormat;
+	TArray<FColor> Pixels;
+	FIntPoint Size;
+};
+
 struct FGLTFMaterialUtility
 {
 	static UTexture2D* CreateTransientTexture(const TArray<FColor>& Pixels, const FIntPoint& TextureSize, EPixelFormat TextureFormat);
 	static bool CombineTextures(TArray<FColor>& OutPixels, const TArray<FGLTFTextureCombineSource>& Sources, const FIntPoint& OutputSize, EPixelFormat OutputPixelFormat);
-	static UTexture2D* BakeMaterialProperty(const FIntPoint& OutputSize, EMaterialProperty MaterialProperty, const UMaterialInterface* Material, bool bCopyAlphaFromRedChannel = false);
+	static FGLTFPropertyBakeOutput BakeMaterialProperty(const FIntPoint& OutputSize, EMaterialProperty MaterialProperty, const UMaterialInterface* Material, bool bCopyAlphaFromRedChannel = false);
+	static UTexture2D* BakeMaterialPropertyToTexture(const FIntPoint& OutputSize, EMaterialProperty MaterialProperty, const UMaterialInterface* Material, bool bCopyAlphaFromRedChannel = false);
 
 	static FGLTFJsonTextureIndex AddCombinedTexture(FGLTFConvertBuilder& Builder, const TArray<FGLTFTextureCombineSource>& CombineSources, const FIntPoint& TextureSize, const FString& TextureName, EGLTFJsonTextureFilter Filter, EGLTFJsonTextureWrap Wrap);
+	static FGLTFJsonTextureIndex AddTexture(FGLTFConvertBuilder& Builder, const TArray<FColor>& Pixels, const FIntPoint& TextureSize, const FString& TextureName, EPixelFormat PixelFormat, EGLTFJsonTextureFilter Filter, EGLTFJsonTextureWrap Wrap);
+
 	static FLinearColor GetMask(const FExpressionInput& ExpressionInput);
 	static uint32 GetMaskComponentCount(const FExpressionInput& ExpressionInput);
 };
