[GLTFExporter] fix(material): use the same pixel-format for all baked properties

commit 8842f55a255eacff6787e9c62b2778090c0a46cc
Date: Fri, 5 Jun 2020 11:25:18 +0200

    fix(material): use the same pixel-format for all baked properties
    (cherry picked from commit b7f87a99bcb3224208d8e8df7f64815b75d0bb97​)

#skipci

diff --git a/Engine/Plugins/Enterprise/GLTFExporter/MigratedFromGit/Source/GLTFExporter/Private/Converters/GLTFMaterialUtility.cpp b/Engine/Plugins/Enterprise/GLTFExporter/MigratedFromGit/Source/GLTFExporter/Private/Converters/GLTFMaterialUtility.cpp
--- a/Engine/Plugins/Enterprise/GLTFExporter/MigratedFromGit/Source/GLTFExporter/Private/Converters/GLTFMaterialUtility.cpp
+++ b/Engine/Plugins/Enterprise/GLTFExporter/MigratedFromGit/Source/GLTFExporter/Private/Converters/GLTFMaterialUtility.cpp
@@ -117,19 +117,9 @@ FGLTFPropertyBakeOutput FGLTFMaterialUtility::BakeMaterialProperty(const FIntPoi
 		}
 	}
 
-	EPixelFormat PixelFormat;
-
-	// NOTE: The texture-format is mapped per property in the MaterialBaking-module via PerPropertyFormat.
-	// It's PF_B8G8R8A8 for all properties except for MP_EmissiveColor (which is PF_FloatRGBA).
-	switch (MaterialProperty)
-	{
-		case MP_EmissiveColor: PixelFormat = PF_FloatRGBA; break;
-		default: PixelFormat = PF_B8G8R8A8; break;
-	}
-
 	return {
 		MaterialProperty,
-		PixelFormat,
+		PF_B8G8R8A8,
 		BakedPixels,
 		BakedSize
 	};
