[GLTFExporter] feat(material): add initial support for baking "simple" properties

commit b02ba659ae94646080fc7adf14e780ed5c302ba7
Date: Fri, 5 Jun 2020 11:46:08 +0200

    feat(material): add initial support for baking "simple" properties
    (cherry picked from commit 9902921052e14c52e918e7e08a9dd6b53960952a​)

#skipci

diff --git a/Engine/Plugins/Enterprise/GLTFExporter/MigratedFromGit/Source/GLTFExporter/Private/Converters/GLTFMaterialConverters.cpp b/Engine/Plugins/Enterprise/GLTFExporter/MigratedFromGit/Source/GLTFExporter/Private/Converters/GLTFMaterialConverters.cpp
--- a/Engine/Plugins/Enterprise/GLTFExporter/MigratedFromGit/Source/GLTFExporter/Private/Converters/GLTFMaterialConverters.cpp
+++ b/Engine/Plugins/Enterprise/GLTFExporter/MigratedFromGit/Source/GLTFExporter/Private/Converters/GLTFMaterialConverters.cpp
@@ -41,7 +41,10 @@ FGLTFJsonMaterialIndex FGLTFMaterialConverter::Add(FGLTFConvertBuilder& Builder,
 		{
 			if (!TryGetSourceTexture(Builder, JsonMaterial.PBRMetallicRoughness.BaseColorTexture, Material->BaseColor, MaterialInstance))
 			{
-				// TODO: add fallback to material baking
+				if (!TryGetBakedTexture(Builder, JsonMaterial.PBRMetallicRoughness.BaseColorTexture, MP_BaseColor, MaterialInterface))
+				{
+					// TODO: handle failure?
+				}
 			}
 		}
 	}
@@ -55,7 +58,10 @@ FGLTFJsonMaterialIndex FGLTFMaterialConverter::Add(FGLTFConvertBuilder& Builder,
 	{
 		if (!TryGetSourceTexture(Builder, JsonMaterial.EmissiveTexture, Material->EmissiveColor, MaterialInstance))
 		{
-			// TODO: add fallback to material baking
+			if (!TryGetBakedTexture(Builder, JsonMaterial.EmissiveTexture, MP_EmissiveColor, MaterialInterface))
+			{
+				// TODO: handle failure?
+			}
 		}
 		else
 		{
@@ -67,7 +73,10 @@ FGLTFJsonMaterialIndex FGLTFMaterialConverter::Add(FGLTFConvertBuilder& Builder,
 	{
 		if (!TryGetSourceTexture(Builder, JsonMaterial.NormalTexture, Material->Normal, MaterialInstance))
 		{
-			// TODO: add fallback to material baking
+			if (!TryGetBakedTexture(Builder, JsonMaterial.NormalTexture, MP_Normal, MaterialInterface))
+			{
+				// TODO: handle failure?
+			}
 		}
 	}
 
@@ -75,7 +84,10 @@ FGLTFJsonMaterialIndex FGLTFMaterialConverter::Add(FGLTFConvertBuilder& Builder,
 	{
 		if (!TryGetSourceTexture(Builder, JsonMaterial.OcclusionTexture, Material->AmbientOcclusion, MaterialInstance))
 		{
-			// TODO: add fallback to material baking
+			if (!TryGetBakedTexture(Builder, JsonMaterial.OcclusionTexture, MP_AmbientOcclusion, MaterialInterface))
+			{
+				// TODO: handle failure?
+			}
 		}
 	}
 
@@ -510,6 +522,56 @@ bool FGLTFMaterialConverter::TryGetConstantScalar(float& OutValue, const FScalar
 	return false;
 }
 
+bool FGLTFMaterialConverter::TryGetBakedTexture(FGLTFConvertBuilder& Builder, FGLTFJsonTextureInfo& OutTexInfo, EMaterialProperty MaterialProperty, const UMaterialInterface* MaterialInterface) const
+{
+	// TODO: make default baking-resolution configurable
+	const FIntPoint TextureSize(512, 512);
+
+	// TODO: add support for detecting the correct tex-coord for this property
+	const uint32 TexCoord = 0;
+
+	const FGLTFPropertyBakeOutput PropertyBakeOutput = FGLTFMaterialUtility::BakeMaterialProperty(
+		TextureSize,
+		MaterialProperty,
+		MaterialInterface);
+
+	// TODO: handle the case where TextureSize is 1x1. In this case, both properties are constants and we should
+	// extract the value of each property from the texture and use as-is instead of exporting a texture.
+
+	TCHAR* PropertyName;
+
+	switch (MaterialProperty)
+	{
+		case MP_BaseColor: PropertyName = TEXT("BaseColor"); break;
+		case MP_Normal: PropertyName = TEXT("Normal"); break;
+		case MP_EmissiveColor: PropertyName = TEXT("EmissiveColor"); break;
+		case MP_AmbientOcclusion: PropertyName = TEXT("AmbientOcclusion"); break;
+		default:
+			// TODO: support for more properties
+			return false;
+	}
+
+	const FString TextureName = FString::Printf(TEXT("%s_%s"), *MaterialInterface->GetName(), PropertyName);
+
+	// TODO: support for other wrapping / filters?
+	const EGLTFJsonTextureWrap TextureWrap = EGLTFJsonTextureWrap::ClampToEdge;
+	const EGLTFJsonTextureFilter TextureFilter = EGLTFJsonTextureFilter::LinearMipmapLinear;
+
+	const FGLTFJsonTextureIndex TextureIndex = FGLTFMaterialUtility::AddTexture(
+		Builder,
+		PropertyBakeOutput.Pixels,
+		PropertyBakeOutput.Size,
+		TextureName,
+		PropertyBakeOutput.PixelFormat,
+		TextureFilter,
+		TextureWrap);
+
+	OutTexInfo.TexCoord = TexCoord;
+	OutTexInfo.Index = TextureIndex;
+
+	return true;
+}
+
 bool FGLTFMaterialConverter::TryGetSourceTexture(FGLTFConvertBuilder& Builder, FGLTFJsonTextureInfo& OutTexInfo, const FExpressionInput& MaterialInput, const UMaterialInstance* MaterialInstance) const
 {
 	const UTexture2D* Texture;
diff --git a/Engine/Plugins/Enterprise/GLTFExporter/MigratedFromGit/Source/GLTFExporter/Private/Converters/GLTFMaterialConverters.h b/Engine/Plugins/Enterprise/GLTFExporter/MigratedFromGit/Source/GLTFExporter/Private/Converters/GLTFMaterialConverters.h
--- a/Engine/Plugins/Enterprise/GLTFExporter/MigratedFromGit/Source/GLTFExporter/Private/Converters/GLTFMaterialConverters.h
+++ b/Engine/Plugins/Enterprise/GLTFExporter/MigratedFromGit/Source/GLTFExporter/Private/Converters/GLTFMaterialConverters.h
@@ -21,4 +21,5 @@ class FGLTFMaterialConverter final : public TGLTFConverter<FGLTFJsonMaterialInde
 
 	bool TryGetSourceTexture(FGLTFConvertBuilder& Builder, FGLTFJsonTextureInfo& OutTexInfo, const FExpressionInput& MaterialInput, const UMaterialInstance* MaterialInstance) const;
 	bool TryGetSourceTexture(const UTexture2D*& OutTexture, int32& OutTexCoord, const FExpressionInput& MaterialInput, const UMaterialInstance* MaterialInstance) const;
+	bool TryGetBakedTexture(FGLTFConvertBuilder& Builder, FGLTFJsonTextureInfo& OutTexInfo, EMaterialProperty MaterialProperty, const UMaterialInterface* MaterialInterface) const;
 };
