[GLTFExporter] feat(material): only use source-textures as-is if they use masks that are compatible with glTF

commit 0311dc22c71ba1e93a9b484fe50c9eccf70d6616
Date: Mon, 8 Jun 2020 10:22:23 +0200

    feat(material): only use source-textures as-is if they use masks that are compatible with glTF
    (cherry picked from commit ee8411b13379b10ce0e679ac92dea81508276a6f​)

#skipci

diff --git a/Engine/Plugins/Enterprise/GLTFExporter/MigratedFromGit/Source/GLTFExporter/Private/Converters/GLTFMaterialConverters.cpp b/Engine/Plugins/Enterprise/GLTFExporter/MigratedFromGit/Source/GLTFExporter/Private/Converters/GLTFMaterialConverters.cpp
--- a/Engine/Plugins/Enterprise/GLTFExporter/MigratedFromGit/Source/GLTFExporter/Private/Converters/GLTFMaterialConverters.cpp
+++ b/Engine/Plugins/Enterprise/GLTFExporter/MigratedFromGit/Source/GLTFExporter/Private/Converters/GLTFMaterialConverters.cpp
@@ -27,6 +27,10 @@ FGLTFJsonMaterialIndex FGLTFMaterialConverter::Add(FGLTFConvertBuilder& Builder,
 	JsonMaterial.AlphaCutoff = MaterialInterface->GetOpacityMaskClipValue();
 	JsonMaterial.DoubleSided = MaterialInterface->IsTwoSided();
 
+	const FLinearColor RgbaMask(1.0f, 1.0f, 1.0f, 1.0f);
+	const FLinearColor RgbMask(1.0f, 1.0f, 1.0f, 0.0f);
+	const FLinearColor RMask(1.0f, 0.0f, 0.0f, 0.0f);
+
 	if (JsonMaterial.AlphaMode == EGLTFJsonAlphaMode::Blend || JsonMaterial.AlphaMode == EGLTFJsonAlphaMode::Mask)
 	{
 		const FScalarMaterialInput& OpacityInput = JsonMaterial.AlphaMode == EGLTFJsonAlphaMode::Blend ? Material->Opacity : Material->OpacityMask;
@@ -39,7 +43,7 @@ FGLTFJsonMaterialIndex FGLTFMaterialConverter::Add(FGLTFConvertBuilder& Builder,
 	{
 		if (!TryGetConstantColor(JsonMaterial.PBRMetallicRoughness.BaseColorFactor, Material->BaseColor, MaterialInstance))
 		{
-			if (!TryGetSourceTexture(Builder, JsonMaterial.PBRMetallicRoughness.BaseColorTexture, Material->BaseColor, MaterialInstance))
+			if (!TryGetSourceTexture(Builder, JsonMaterial.PBRMetallicRoughness.BaseColorTexture, Material->BaseColor, MaterialInstance, {RgbaMask, RgbMask}))
 			{
 				if (!TryGetBakedTexture(Builder, JsonMaterial.PBRMetallicRoughness.BaseColorTexture, MP_BaseColor, MaterialInterface))
 				{
@@ -56,7 +60,7 @@ FGLTFJsonMaterialIndex FGLTFMaterialConverter::Add(FGLTFConvertBuilder& Builder,
 
 	if (!TryGetConstantColor(JsonMaterial.EmissiveFactor, Material->EmissiveColor, MaterialInstance))
 	{
-		if (!TryGetSourceTexture(Builder, JsonMaterial.EmissiveTexture, Material->EmissiveColor, MaterialInstance))
+		if (!TryGetSourceTexture(Builder, JsonMaterial.EmissiveTexture, Material->EmissiveColor, MaterialInstance, {RgbaMask, RgbMask}))
 		{
 			if (!TryGetBakedTexture(Builder, JsonMaterial.EmissiveTexture, MP_EmissiveColor, MaterialInterface))
 			{
@@ -71,7 +75,7 @@ FGLTFJsonMaterialIndex FGLTFMaterialConverter::Add(FGLTFConvertBuilder& Builder,
 
 	if (Material->Normal.Expression != nullptr)
 	{
-		if (!TryGetSourceTexture(Builder, JsonMaterial.NormalTexture, Material->Normal, MaterialInstance))
+		if (!TryGetSourceTexture(Builder, JsonMaterial.NormalTexture, Material->Normal, MaterialInstance, {RgbaMask, RgbMask}))
 		{
 			if (!TryGetBakedTexture(Builder, JsonMaterial.NormalTexture, MP_Normal, MaterialInterface))
 			{
@@ -82,7 +86,7 @@ FGLTFJsonMaterialIndex FGLTFMaterialConverter::Add(FGLTFConvertBuilder& Builder,
 
 	if (Material->AmbientOcclusion.Expression != nullptr)
 	{
-		if (!TryGetSourceTexture(Builder, JsonMaterial.OcclusionTexture, Material->AmbientOcclusion, MaterialInstance))
+		if (!TryGetSourceTexture(Builder, JsonMaterial.OcclusionTexture, Material->AmbientOcclusion, MaterialInstance, {RMask}))
 		{
 			if (!TryGetBakedTexture(Builder, JsonMaterial.OcclusionTexture, MP_AmbientOcclusion, MaterialInterface))
 			{
@@ -116,24 +120,21 @@ bool FGLTFMaterialConverter::TryGetBaseColorAndOpacity(FGLTFConvertBuilder& Buil
 	int32 BaseColorTexCoord;
 	int32 OpacityTexCoord;
 
-	const bool bHasBaseColorSourceTexture = TryGetSourceTexture(BaseColorTexture, BaseColorTexCoord, BaseColorInput, MaterialInstance);
-	const bool bHasOpacitySourceTexture = TryGetSourceTexture(OpacityTexture, OpacityTexCoord, OpacityInput, MaterialInstance);
-
 	const FLinearColor BaseColorMask(1, 1, 1, 0);
 	const FLinearColor OpacityMask(0, 0, 0, 1);
 
+	const bool bHasBaseColorSourceTexture = TryGetSourceTexture(BaseColorTexture, BaseColorTexCoord, BaseColorInput, MaterialInstance, {BaseColorMask});
+	const bool bHasOpacitySourceTexture = TryGetSourceTexture(OpacityTexture, OpacityTexCoord, OpacityInput, MaterialInstance, {OpacityMask});
+
 	// Detect the "happy path" where both inputs share the same texture, and both inputs correctly masked.
-	if (bHasBaseColorSourceTexture && bHasOpacitySourceTexture)
+	if (bHasBaseColorSourceTexture &&
+		bHasOpacitySourceTexture &&
+		BaseColorTexture == OpacityTexture &&
+		BaseColorTexCoord == OpacityTexCoord)
 	{
-		if (BaseColorTexture == OpacityTexture &&
-			BaseColorTexCoord == OpacityTexCoord &&
-			BaseColorMask.Equals(FGLTFMaterialUtility::GetMask(BaseColorInput)) &&
-			OpacityMask.Equals(FGLTFMaterialUtility::GetMask(OpacityInput)))
-		{
-			OutPBRParams.BaseColorTexture.Index = Builder.GetOrAddTexture(BaseColorTexture);
-			OutPBRParams.BaseColorTexture.TexCoord = BaseColorTexCoord;
-			return true;
-		}
+		OutPBRParams.BaseColorTexture.Index = Builder.GetOrAddTexture(BaseColorTexture);
+		OutPBRParams.BaseColorTexture.TexCoord = BaseColorTexCoord;
+		return true;
 	}
 
 	int32 TexCoord = 0;
@@ -236,31 +237,27 @@ bool FGLTFMaterialConverter::TryGetMetallicAndRoughness(FGLTFConvertBuilder& Bui
 	int32 MetallicTexCoord;
 	int32 RoughnessTexCoord;
 
-	const bool bHasMetallicSourceTexture = TryGetSourceTexture(MetallicTexture, MetallicTexCoord, MetallicInput, MaterialInstance);
-	const bool bHasRoughnessSourceTexture = TryGetSourceTexture(RoughnessTexture, RoughnessTexCoord, RoughnessInput, MaterialInstance);
-
 	const FLinearColor MetallicMask(0, 0, 1, 0);
 	const FLinearColor RoughnessMask(0, 1, 0, 0);
 
-	// Detect the "happy path" where both inputs share the same texture, and both inputs correctly masked.
-	if (bHasMetallicSourceTexture && bHasRoughnessSourceTexture)
+	const bool bHasMetallicSourceTexture = TryGetSourceTexture(MetallicTexture, MetallicTexCoord, MetallicInput, MaterialInstance, {MetallicMask});
+	const bool bHasRoughnessSourceTexture = TryGetSourceTexture(RoughnessTexture, RoughnessTexCoord, RoughnessInput, MaterialInstance, {RoughnessMask});
+
+	// Detect the "happy path" where both inputs share the same texture, and both inputs are correctly masked.
+	if (bHasMetallicSourceTexture &&
+		bHasRoughnessSourceTexture &&
+		MetallicTexture == RoughnessTexture &&
+		MetallicTexCoord == RoughnessTexCoord)
 	{
-		if (MetallicTexture == RoughnessTexture &&
-			MetallicTexCoord == RoughnessTexCoord &&
-			MetallicMask.Equals(FGLTFMaterialUtility::GetMask(MetallicInput)) &&
-			RoughnessMask.Equals(FGLTFMaterialUtility::GetMask(RoughnessInput)))
-		{
-			OutPBRParams.MetallicRoughnessTexture.Index = Builder.GetOrAddTexture(MetallicTexture);
-			OutPBRParams.MetallicRoughnessTexture.TexCoord = MetallicTexCoord;
-			return true;
-		}
+		OutPBRParams.MetallicRoughnessTexture.Index = Builder.GetOrAddTexture(MetallicTexture);
+		OutPBRParams.MetallicRoughnessTexture.TexCoord = MetallicTexCoord;
+		return true;
 	}
 
 	int32 TexCoord = 0;
 	FIntPoint TextureSize(512, 512);	// TODO: make default baking-resolution configurable
 	EGLTFJsonTextureWrap TextureWrap = EGLTFJsonTextureWrap::ClampToEdge;
 	EGLTFJsonTextureFilter TextureFilter = EGLTFJsonTextureFilter::LinearMipmapLinear;
-	const EPixelFormat PixelFormat = PF_B8G8R8A8;
 
 	if (bHasMetallicSourceTexture && bHasRoughnessSourceTexture)
 	{
@@ -572,12 +569,12 @@ bool FGLTFMaterialConverter::TryGetBakedTexture(FGLTFConvertBuilder& Builder, FG
 	return true;
 }
 
-bool FGLTFMaterialConverter::TryGetSourceTexture(FGLTFConvertBuilder& Builder, FGLTFJsonTextureInfo& OutTexInfo, const FExpressionInput& MaterialInput, const UMaterialInstance* MaterialInstance) const
+bool FGLTFMaterialConverter::TryGetSourceTexture(FGLTFConvertBuilder& Builder, FGLTFJsonTextureInfo& OutTexInfo, const FExpressionInput& MaterialInput, const UMaterialInstance* MaterialInstance, const TArray<FLinearColor>& AllowedMasks) const
 {
 	const UTexture2D* Texture;
 	int32 TexCoord;
 
-	if (TryGetSourceTexture(Texture, TexCoord, MaterialInput, MaterialInstance))
+	if (TryGetSourceTexture(Texture, TexCoord, MaterialInput, MaterialInstance, AllowedMasks))
 	{
 		// TODO: add support for output mask
 		OutTexInfo.Index = Builder.GetOrAddTexture(Texture);
@@ -588,8 +585,14 @@ bool FGLTFMaterialConverter::TryGetSourceTexture(FGLTFConvertBuilder& Builder, F
 	return false;
 }
 
-bool FGLTFMaterialConverter::TryGetSourceTexture(const UTexture2D*& OutTexture, int32& OutTexCoord, const FExpressionInput& MaterialInput, const UMaterialInstance* MaterialInstance) const
+bool FGLTFMaterialConverter::TryGetSourceTexture(const UTexture2D*& OutTexture, int32& OutTexCoord, const FExpressionInput& MaterialInput, const UMaterialInstance* MaterialInstance, const TArray<FLinearColor>& AllowedMasks) const
 {
+	const FLinearColor InputMask = FGLTFMaterialUtility::GetMask(MaterialInput);
+	if (AllowedMasks.Num() > 0 && !AllowedMasks.Contains(InputMask))
+	{
+		return false;
+	}
+
 	const UMaterialExpression* Expression = MaterialInput.Expression;
 	if (Expression == nullptr)
 	{
diff --git a/Engine/Plugins/Enterprise/GLTFExporter/MigratedFromGit/Source/GLTFExporter/Private/Converters/GLTFMaterialConverters.h b/Engine/Plugins/Enterprise/GLTFExporter/MigratedFromGit/Source/GLTFExporter/Private/Converters/GLTFMaterialConverters.h
--- a/Engine/Plugins/Enterprise/GLTFExporter/MigratedFromGit/Source/GLTFExporter/Private/Converters/GLTFMaterialConverters.h
+++ b/Engine/Plugins/Enterprise/GLTFExporter/MigratedFromGit/Source/GLTFExporter/Private/Converters/GLTFMaterialConverters.h
@@ -19,7 +19,7 @@ class FGLTFMaterialConverter final : public TGLTFConverter<FGLTFJsonMaterialInde
 	bool TryGetConstantColor(FLinearColor& OutValue, const FColorMaterialInput& MaterialInput, const UMaterialInstance* MaterialInstance) const;
 	bool TryGetConstantScalar(float& OutValue, const FScalarMaterialInput& MaterialInput, const UMaterialInstance* MaterialInstance) const;
 
-	bool TryGetSourceTexture(FGLTFConvertBuilder& Builder, FGLTFJsonTextureInfo& OutTexInfo, const FExpressionInput& MaterialInput, const UMaterialInstance* MaterialInstance) const;
-	bool TryGetSourceTexture(const UTexture2D*& OutTexture, int32& OutTexCoord, const FExpressionInput& MaterialInput, const UMaterialInstance* MaterialInstance) const;
+	bool TryGetSourceTexture(FGLTFConvertBuilder& Builder, FGLTFJsonTextureInfo& OutTexInfo, const FExpressionInput& MaterialInput, const UMaterialInstance* MaterialInstance, const TArray<FLinearColor>& AllowedMasks = {}) const;
+	bool TryGetSourceTexture(const UTexture2D*& OutTexture, int32& OutTexCoord, const FExpressionInput& MaterialInput, const UMaterialInstance* MaterialInstance, const TArray<FLinearColor>& AllowedMasks = {}) const;
 	bool TryGetBakedTexture(FGLTFConvertBuilder& Builder, FGLTFJsonTextureInfo& OutTexInfo, EMaterialProperty MaterialProperty, const UMaterialInterface* MaterialInterface) const;
 };
