[GLTFExporter] refactor(material): include EmissiveScale in output when baking a property (for future use)

commit c9829e0c153e5aa38adddba07b59033fdbc65213
Date: Mon, 8 Jun 2020 10:49:33 +0200

    refactor(material): include EmissiveScale in output when baking a property (for future use)
    (cherry picked from commit 1b3ddd106c15d5f7871c7b46f468ecc96f669b00​)

#skipci

diff --git a/Engine/Plugins/Enterprise/GLTFExporter/MigratedFromGit/Source/GLTFExporter/Private/Converters/GLTFMaterialUtility.cpp b/Engine/Plugins/Enterprise/GLTFExporter/MigratedFromGit/Source/GLTFExporter/Private/Converters/GLTFMaterialUtility.cpp
--- a/Engine/Plugins/Enterprise/GLTFExporter/MigratedFromGit/Source/GLTFExporter/Private/Converters/GLTFMaterialUtility.cpp
+++ b/Engine/Plugins/Enterprise/GLTFExporter/MigratedFromGit/Source/GLTFExporter/Private/Converters/GLTFMaterialUtility.cpp
@@ -82,7 +82,7 @@ FGLTFPropertyBakeOutput FGLTFMaterialUtility::BakeMaterialProperty(const FIntPoi
 
 	MeshSet->TextureCoordinateBox = FBox2D(FVector2D(0, 0), FVector2D(1, 1));
 
-	// TODO: Do we need to fill in any info in MeshSet?
+	// TODO: Do we need to fill in any more info in MeshSet?
 
 	MeshSettings.Add(MeshSet);
 
@@ -99,8 +99,11 @@ FGLTFPropertyBakeOutput FGLTFMaterialUtility::BakeMaterialProperty(const FIntPoi
 
 	Module.BakeMaterials(MatSettings, MeshSettings, BakeOutputs);
 
-	TArray<FColor> BakedPixels = MoveTemp(BakeOutputs[0].PropertyData.FindChecked(MaterialProperty));
-	const FIntPoint BakedSize = BakeOutputs[0].PropertySizes.FindChecked(MaterialProperty);
+	FBakeOutput BakeOutput = BakeOutputs[0];
+
+	TArray<FColor> BakedPixels = MoveTemp(BakeOutput.PropertyData.FindChecked(MaterialProperty));
+	const FIntPoint BakedSize = BakeOutput.PropertySizes.FindChecked(MaterialProperty);
+	const float EmissiveScale = BakeOutput.EmissiveScale;
 
 	if (bCopyAlphaFromRedChannel)
 	{
@@ -121,7 +124,8 @@ FGLTFPropertyBakeOutput FGLTFMaterialUtility::BakeMaterialProperty(const FIntPoi
 		MaterialProperty,
 		PF_B8G8R8A8,
 		BakedPixels,
-		BakedSize
+		BakedSize,
+		EmissiveScale
 	};
 }
 
diff --git a/Engine/Plugins/Enterprise/GLTFExporter/MigratedFromGit/Source/GLTFExporter/Private/Converters/GLTFMaterialUtility.h b/Engine/Plugins/Enterprise/GLTFExporter/MigratedFromGit/Source/GLTFExporter/Private/Converters/GLTFMaterialUtility.h
--- a/Engine/Plugins/Enterprise/GLTFExporter/MigratedFromGit/Source/GLTFExporter/Private/Converters/GLTFMaterialUtility.h
+++ b/Engine/Plugins/Enterprise/GLTFExporter/MigratedFromGit/Source/GLTFExporter/Private/Converters/GLTFMaterialUtility.h
@@ -18,14 +18,15 @@ struct FGLTFTextureCombineSource
 
 struct FGLTFPropertyBakeOutput
 {
-	FORCEINLINE FGLTFPropertyBakeOutput(EMaterialProperty Property, EPixelFormat PixelFormat, TArray<FColor>& Pixels, FIntPoint Size)
-		: Property(Property), PixelFormat(PixelFormat), Pixels(Pixels), Size(Size)
+	FORCEINLINE FGLTFPropertyBakeOutput(EMaterialProperty Property, EPixelFormat PixelFormat, TArray<FColor>& Pixels, FIntPoint Size, float EmissiveScale)
+		: Property(Property), PixelFormat(PixelFormat), Pixels(Pixels), Size(Size), EmissiveScale(EmissiveScale)
 	{}
 
 	EMaterialProperty Property;
 	EPixelFormat PixelFormat;
 	TArray<FColor> Pixels;
 	FIntPoint Size;
+	float EmissiveScale;
 };
 
 struct FGLTFMaterialUtility
