[GLTFExporter] chore(material): update some TODOs for missing functionality

commit 9fda3f2dcae0934a29dfd12c5d4240acf4235889
Date: Mon, 8 Jun 2020 10:50:09 +0200

    chore(material): update some TODOs for missing functionality
    (cherry picked from commit 02065426fde1a5f35abe38e825c63340f52ce562​)

#skipci

diff --git a/Engine/Plugins/Enterprise/GLTFExporter/MigratedFromGit/Source/GLTFExporter/Private/Converters/GLTFMaterialConverters.cpp b/Engine/Plugins/Enterprise/GLTFExporter/MigratedFromGit/Source/GLTFExporter/Private/Converters/GLTFMaterialConverters.cpp
--- a/Engine/Plugins/Enterprise/GLTFExporter/MigratedFromGit/Source/GLTFExporter/Private/Converters/GLTFMaterialConverters.cpp
+++ b/Engine/Plugins/Enterprise/GLTFExporter/MigratedFromGit/Source/GLTFExporter/Private/Converters/GLTFMaterialConverters.cpp
@@ -126,7 +126,7 @@ bool FGLTFMaterialConverter::TryGetBaseColorAndOpacity(FGLTFConvertBuilder& Buil
 	const bool bHasBaseColorSourceTexture = TryGetSourceTexture(BaseColorTexture, BaseColorTexCoord, BaseColorInput, MaterialInstance, {BaseColorMask});
 	const bool bHasOpacitySourceTexture = TryGetSourceTexture(OpacityTexture, OpacityTexCoord, OpacityInput, MaterialInstance, {OpacityMask});
 
-	// Detect the "happy path" where both inputs share the same texture, and both inputs correctly masked.
+	// Detect the "happy path" where both inputs share the same texture and are correctly masked.
 	if (bHasBaseColorSourceTexture &&
 		bHasOpacitySourceTexture &&
 		BaseColorTexture == OpacityTexture &&
@@ -137,8 +137,12 @@ bool FGLTFMaterialConverter::TryGetBaseColorAndOpacity(FGLTFConvertBuilder& Buil
 		return true;
 	}
 
+	// TODO: make default baking-resolution configurable
+	// TODO: add support for detecting the correct tex-coord for this property based on connected nodes
+	// TODO: add support for calculating the ideal resolution to use for baking based on connected (texture) nodes
 	int32 TexCoord = 0;
-	FIntPoint TextureSize(512, 512);	// TODO: make default baking-resolution configurable
+	FIntPoint TextureSize(512, 512);
+
 	EGLTFJsonTextureWrap TextureWrap = EGLTFJsonTextureWrap::ClampToEdge;
 	EGLTFJsonTextureFilter TextureFilter = EGLTFJsonTextureFilter::LinearMipmapLinear;
 
@@ -243,7 +247,7 @@ bool FGLTFMaterialConverter::TryGetMetallicAndRoughness(FGLTFConvertBuilder& Bui
 	const bool bHasMetallicSourceTexture = TryGetSourceTexture(MetallicTexture, MetallicTexCoord, MetallicInput, MaterialInstance, {MetallicMask});
 	const bool bHasRoughnessSourceTexture = TryGetSourceTexture(RoughnessTexture, RoughnessTexCoord, RoughnessInput, MaterialInstance, {RoughnessMask});
 
-	// Detect the "happy path" where both inputs share the same texture, and both inputs are correctly masked.
+	// Detect the "happy path" where both inputs share the same texture and are correctly masked.
 	if (bHasMetallicSourceTexture &&
 		bHasRoughnessSourceTexture &&
 		MetallicTexture == RoughnessTexture &&
@@ -254,8 +258,12 @@ bool FGLTFMaterialConverter::TryGetMetallicAndRoughness(FGLTFConvertBuilder& Bui
 		return true;
 	}
 
+	// TODO: make default baking-resolution configurable
+	// TODO: add support for detecting the correct tex-coord for this property based on connected nodes
+	// TODO: add support for calculating the ideal resolution to use for baking based on connected (texture) nodes
 	int32 TexCoord = 0;
-	FIntPoint TextureSize(512, 512);	// TODO: make default baking-resolution configurable
+	FIntPoint TextureSize(512, 512);
+
 	EGLTFJsonTextureWrap TextureWrap = EGLTFJsonTextureWrap::ClampToEdge;
 	EGLTFJsonTextureFilter TextureFilter = EGLTFJsonTextureFilter::LinearMipmapLinear;
 
@@ -355,6 +363,8 @@ bool FGLTFMaterialConverter::TryGetConstantColor(FGLTFJsonColor4& OutValue, cons
 
 bool FGLTFMaterialConverter::TryGetConstantColor(FLinearColor& OutValue, const FColorMaterialInput& MaterialInput, const UMaterialInstance* MaterialInstance) const
 {
+	// TODO: handle emissive color-values above 1.0
+
 	const UMaterialExpression* Expression = MaterialInput.Expression;
 	if (Expression == nullptr)
 	{
@@ -522,9 +532,9 @@ bool FGLTFMaterialConverter::TryGetConstantScalar(float& OutValue, const FScalar
 bool FGLTFMaterialConverter::TryGetBakedTexture(FGLTFConvertBuilder& Builder, FGLTFJsonTextureInfo& OutTexInfo, EMaterialProperty MaterialProperty, const UMaterialInterface* MaterialInterface) const
 {
 	// TODO: make default baking-resolution configurable
+	// TODO: add support for detecting the correct tex-coord for this property based on connected nodes
+	// TODO: add support for calculating the ideal resolution to use for baking based on connected (texture) nodes
 	const FIntPoint TextureSize(512, 512);
-
-	// TODO: add support for detecting the correct tex-coord for this property
 	const uint32 TexCoord = 0;
 
 	const FGLTFPropertyBakeOutput PropertyBakeOutput = FGLTFMaterialUtility::BakeMaterialProperty(
@@ -532,8 +542,8 @@ bool FGLTFMaterialConverter::TryGetBakedTexture(FGLTFConvertBuilder& Builder, FG
 		MaterialProperty,
 		MaterialInterface);
 
-	// TODO: handle the case where TextureSize is 1x1. In this case, both properties are constants and we should
-	// extract the value of each property from the texture and use as-is instead of exporting a texture.
+	// TODO: handle cases where PropertyBakeOutput.EmissiveScale is not 1.0 (when baking EmissiveColor)
+	// TODO: handle the case where TextureSize is 1x1, which means is could be stored in a constant instead of in a texture
 
 	TCHAR* PropertyName;
 
@@ -576,7 +586,6 @@ bool FGLTFMaterialConverter::TryGetSourceTexture(FGLTFConvertBuilder& Builder, F
 
 	if (TryGetSourceTexture(Texture, TexCoord, MaterialInput, MaterialInstance, AllowedMasks))
 	{
-		// TODO: add support for output mask
 		OutTexInfo.Index = Builder.GetOrAddTexture(Texture);
 		OutTexInfo.TexCoord = TexCoord;
 		return true;
@@ -628,7 +637,6 @@ bool FGLTFMaterialConverter::TryGetSourceTexture(const UTexture2D*& OutTexture,
 		// TODO: add support for texture coordinate input expression
 		OutTexCoord = TextureParameter->ConstCoordinate;
 
-		// TODO: add support for output mask
 		return true;
 	}
 
@@ -653,7 +661,6 @@ bool FGLTFMaterialConverter::TryGetSourceTexture(const UTexture2D*& OutTexture,
 		// TODO: add support for texture coordinate input expression
 		OutTexCoord = TextureSample->ConstCoordinate;
 
-		// TODO: add support for output mask
 		return true;
 	}
 
