[GLTFExporter] fix(material): ensure that EmissiveFactor is set to 1,1,1,1 for all cases where a texture is used

commit 4f4722259eb8c4ce1960349439def1155f65cefa
Date: Mon, 8 Jun 2020 11:22:43 +0200

    fix(material): ensure that EmissiveFactor is set to 1,1,1,1 for all cases where a texture is used
    (cherry picked from commit a513ab3c0853797a58d48ef2bde32c98978f2f4c​)

#skipci

diff --git a/Engine/Plugins/Enterprise/GLTFExporter/MigratedFromGit/Source/GLTFExporter/Private/Converters/GLTFMaterialConverters.cpp b/Engine/Plugins/Enterprise/GLTFExporter/MigratedFromGit/Source/GLTFExporter/Private/Converters/GLTFMaterialConverters.cpp
--- a/Engine/Plugins/Enterprise/GLTFExporter/MigratedFromGit/Source/GLTFExporter/Private/Converters/GLTFMaterialConverters.cpp
+++ b/Engine/Plugins/Enterprise/GLTFExporter/MigratedFromGit/Source/GLTFExporter/Private/Converters/GLTFMaterialConverters.cpp
@@ -60,16 +60,16 @@ FGLTFJsonMaterialIndex FGLTFMaterialConverter::Add(FGLTFConvertBuilder& Builder,
 
 	if (!TryGetConstantColor(JsonMaterial.EmissiveFactor, Material->EmissiveColor, MaterialInstance))
 	{
-		if (!TryGetSourceTexture(Builder, JsonMaterial.EmissiveTexture, Material->EmissiveColor, MaterialInstance, {RgbaMask, RgbMask}))
+		const bool bTextureFound = TryGetSourceTexture(Builder, JsonMaterial.EmissiveTexture, Material->EmissiveColor, MaterialInstance, {RgbaMask, RgbMask}) ||
+			TryGetBakedTexture(Builder, JsonMaterial.EmissiveTexture, MP_EmissiveColor, MaterialInterface);
+
+		if (bTextureFound)
 		{
-			if (!TryGetBakedTexture(Builder, JsonMaterial.EmissiveTexture, MP_EmissiveColor, MaterialInterface))
-			{
-				// TODO: handle failure?
-			}
+			JsonMaterial.EmissiveFactor = FGLTFJsonColor3::White; // make sure texture is not multiplied with black
 		}
 		else
 		{
-			JsonMaterial.EmissiveFactor = FGLTFJsonColor3::White; // make sure texture is not multiplied with black
+			// TODO: handle failure?
 		}
 	}
 
