[GLTFExporter] refactor: rename GLTFExporter::Add -> GLTFExporter::AddObject

commit 3f6409b3b4255b70cc1d58d79fa137ceeb53eea0
Date: Sun, 31 May 2020 22:05:46 +0200

    refactor: rename GLTFExporter::Add -> GLTFExporter::AddObject
    (cherry picked from commit 6ae90618f09fe7c2c24de841a5bbf669a07313d0​)

#skipci

diff --git a/Engine/Plugins/Enterprise/GLTFExporter/MigratedFromGit/Source/GLTFExporter/Private/Exporters/GLTFAnimSequenceExporter.cpp b/Engine/Plugins/Enterprise/GLTFExporter/MigratedFromGit/Source/GLTFExporter/Private/Exporters/GLTFAnimSequenceExporter.cpp
--- a/Engine/Plugins/Enterprise/GLTFExporter/MigratedFromGit/Source/GLTFExporter/Private/Exporters/GLTFAnimSequenceExporter.cpp
+++ b/Engine/Plugins/Enterprise/GLTFExporter/MigratedFromGit/Source/GLTFExporter/Private/Exporters/GLTFAnimSequenceExporter.cpp
@@ -10,7 +10,7 @@ UGLTFAnimSequenceExporter::UGLTFAnimSequenceExporter(const FObjectInitializer& O
 	// SupportedClass = UAnimSequence::StaticClass();
 }
 
-bool UGLTFAnimSequenceExporter::Add(FGLTFContainerBuilder& Builder, const UObject* Object)
+bool UGLTFAnimSequenceExporter::AddObject(FGLTFContainerBuilder& Builder, const UObject* Object)
 {
 	const UAnimSequence* AnimSequence = CastChecked<UAnimSequence>(Object);
 	// TODO: implement
diff --git a/Engine/Plugins/Enterprise/GLTFExporter/MigratedFromGit/Source/GLTFExporter/Private/Exporters/GLTFExporter.cpp b/Engine/Plugins/Enterprise/GLTFExporter/MigratedFromGit/Source/GLTFExporter/Private/Exporters/GLTFExporter.cpp
--- a/Engine/Plugins/Enterprise/GLTFExporter/MigratedFromGit/Source/GLTFExporter/Private/Exporters/GLTFExporter.cpp
+++ b/Engine/Plugins/Enterprise/GLTFExporter/MigratedFromGit/Source/GLTFExporter/Private/Exporters/GLTFExporter.cpp
@@ -33,7 +33,7 @@ bool UGLTFExporter::ExportBinary(UObject* Object, const TCHAR* Type, FArchive& A
 	}
 
 	FGLTFContainerBuilder Builder;
-	if (!Add(Builder, Object))
+	if (!AddObject(Builder, Object))
 	{
 		// TODO: Report error
 		return false;
@@ -42,7 +42,7 @@ bool UGLTFExporter::ExportBinary(UObject* Object, const TCHAR* Type, FArchive& A
 	return Builder.Serialize(Archive, CurrentFilename);
 }
 
-bool UGLTFExporter::Add(FGLTFContainerBuilder& Builder, const UObject* Object)
+bool UGLTFExporter::AddObject(FGLTFContainerBuilder& Builder, const UObject* Object)
 {
 	return false;
 }
diff --git a/Engine/Plugins/Enterprise/GLTFExporter/MigratedFromGit/Source/GLTFExporter/Private/Exporters/GLTFLevelExporter.cpp b/Engine/Plugins/Enterprise/GLTFExporter/MigratedFromGit/Source/GLTFExporter/Private/Exporters/GLTFLevelExporter.cpp
--- a/Engine/Plugins/Enterprise/GLTFExporter/MigratedFromGit/Source/GLTFExporter/Private/Exporters/GLTFLevelExporter.cpp
+++ b/Engine/Plugins/Enterprise/GLTFExporter/MigratedFromGit/Source/GLTFExporter/Private/Exporters/GLTFLevelExporter.cpp
@@ -10,7 +10,7 @@ UGLTFLevelExporter::UGLTFLevelExporter(const FObjectInitializer& ObjectInitializ
 	SupportedClass = UWorld::StaticClass();
 }
 
-bool UGLTFLevelExporter::Add(FGLTFContainerBuilder& Builder, const UObject* Object)
+bool UGLTFLevelExporter::AddObject(FGLTFContainerBuilder& Builder, const UObject* Object)
 {
 	const UWorld* World = CastChecked<UWorld>(Object);
 	const FGLTFJsonSceneIndex SceneIndex = Builder.GetOrAddScene(World, bSelectedOnly);
diff --git a/Engine/Plugins/Enterprise/GLTFExporter/MigratedFromGit/Source/GLTFExporter/Private/Exporters/GLTFMaterialExporter.cpp b/Engine/Plugins/Enterprise/GLTFExporter/MigratedFromGit/Source/GLTFExporter/Private/Exporters/GLTFMaterialExporter.cpp
--- a/Engine/Plugins/Enterprise/GLTFExporter/MigratedFromGit/Source/GLTFExporter/Private/Exporters/GLTFMaterialExporter.cpp
+++ b/Engine/Plugins/Enterprise/GLTFExporter/MigratedFromGit/Source/GLTFExporter/Private/Exporters/GLTFMaterialExporter.cpp
@@ -11,7 +11,7 @@ UGLTFMaterialExporter::UGLTFMaterialExporter(const FObjectInitializer& ObjectIni
 	SupportedClass = UMaterialInterface::StaticClass();
 }
 
-bool UGLTFMaterialExporter::Add(FGLTFContainerBuilder& Builder, const UObject* Object)
+bool UGLTFMaterialExporter::AddObject(FGLTFContainerBuilder& Builder, const UObject* Object)
 {
 	const UMaterialInterface* Material = CastChecked<UMaterialInterface>(Object);
 	const FGLTFJsonMaterialIndex MaterialIndex = Builder.GetOrAddMaterial(Material);
diff --git a/Engine/Plugins/Enterprise/GLTFExporter/MigratedFromGit/Source/GLTFExporter/Private/Exporters/GLTFSkeletalMeshExporter.cpp b/Engine/Plugins/Enterprise/GLTFExporter/MigratedFromGit/Source/GLTFExporter/Private/Exporters/GLTFSkeletalMeshExporter.cpp
--- a/Engine/Plugins/Enterprise/GLTFExporter/MigratedFromGit/Source/GLTFExporter/Private/Exporters/GLTFSkeletalMeshExporter.cpp
+++ b/Engine/Plugins/Enterprise/GLTFExporter/MigratedFromGit/Source/GLTFExporter/Private/Exporters/GLTFSkeletalMeshExporter.cpp
@@ -10,7 +10,7 @@ UGLTFSkeletalMeshExporter::UGLTFSkeletalMeshExporter(const FObjectInitializer& O
 	// SupportedClass = USkeletalMesh::StaticClass();
 }
 
-bool UGLTFSkeletalMeshExporter::Add(FGLTFContainerBuilder& Builder, const UObject* Object)
+bool UGLTFSkeletalMeshExporter::AddObject(FGLTFContainerBuilder& Builder, const UObject* Object)
 {
 	const USkeletalMesh* SkeletalMesh = CastChecked<USkeletalMesh>(Object);
 	// TODO: implement
diff --git a/Engine/Plugins/Enterprise/GLTFExporter/MigratedFromGit/Source/GLTFExporter/Private/Exporters/GLTFStaticMeshExporter.cpp b/Engine/Plugins/Enterprise/GLTFExporter/MigratedFromGit/Source/GLTFExporter/Private/Exporters/GLTFStaticMeshExporter.cpp
--- a/Engine/Plugins/Enterprise/GLTFExporter/MigratedFromGit/Source/GLTFExporter/Private/Exporters/GLTFStaticMeshExporter.cpp
+++ b/Engine/Plugins/Enterprise/GLTFExporter/MigratedFromGit/Source/GLTFExporter/Private/Exporters/GLTFStaticMeshExporter.cpp
@@ -10,7 +10,7 @@ UGLTFStaticMeshExporter::UGLTFStaticMeshExporter(const FObjectInitializer& Objec
 	SupportedClass = UStaticMesh::StaticClass();
 }
 
-bool UGLTFStaticMeshExporter::Add(FGLTFContainerBuilder& Builder, const UObject* Object)
+bool UGLTFStaticMeshExporter::AddObject(FGLTFContainerBuilder& Builder, const UObject* Object)
 {
 	const UStaticMesh* StaticMesh = CastChecked<UStaticMesh>(Object);
 	const FGLTFJsonMeshIndex MeshIndex = Builder.GetOrAddMesh(StaticMesh);
diff --git a/Engine/Plugins/Enterprise/GLTFExporter/MigratedFromGit/Source/GLTFExporter/Public/Exporters/GLTFAnimSequenceExporter.h b/Engine/Plugins/Enterprise/GLTFExporter/MigratedFromGit/Source/GLTFExporter/Public/Exporters/GLTFAnimSequenceExporter.h
--- a/Engine/Plugins/Enterprise/GLTFExporter/MigratedFromGit/Source/GLTFExporter/Public/Exporters/GLTFAnimSequenceExporter.h
+++ b/Engine/Plugins/Enterprise/GLTFExporter/MigratedFromGit/Source/GLTFExporter/Public/Exporters/GLTFAnimSequenceExporter.h
@@ -14,5 +14,5 @@ public:
 
 	explicit UGLTFAnimSequenceExporter(const FObjectInitializer& ObjectInitializer = FObjectInitializer());
 
-	bool Add(FGLTFContainerBuilder& Builder, const UObject* Object) override;
+	bool AddObject(FGLTFContainerBuilder& Builder, const UObject* Object) override;
 };
diff --git a/Engine/Plugins/Enterprise/GLTFExporter/MigratedFromGit/Source/GLTFExporter/Public/Exporters/GLTFExporter.h b/Engine/Plugins/Enterprise/GLTFExporter/MigratedFromGit/Source/GLTFExporter/Public/Exporters/GLTFExporter.h
--- a/Engine/Plugins/Enterprise/GLTFExporter/MigratedFromGit/Source/GLTFExporter/Public/Exporters/GLTFExporter.h
+++ b/Engine/Plugins/Enterprise/GLTFExporter/MigratedFromGit/Source/GLTFExporter/Public/Exporters/GLTFExporter.h
@@ -3,8 +3,6 @@
 #pragma once
 
 #include "CoreMinimal.h"
-#include "UObject/ObjectMacros.h"
-#include "UObject/UObjectGlobals.h"
 #include "Exporters/Exporter.h"
 #include "GLTFExporter.generated.h"
 
@@ -24,7 +22,7 @@ public:
 
 protected:
 
-	virtual bool Add(FGLTFContainerBuilder& Builder, const UObject* Object);
+	virtual bool AddObject(FGLTFContainerBuilder& Builder, const UObject* Object);
 
 private:
 
diff --git a/Engine/Plugins/Enterprise/GLTFExporter/MigratedFromGit/Source/GLTFExporter/Public/Exporters/GLTFLevelExporter.h b/Engine/Plugins/Enterprise/GLTFExporter/MigratedFromGit/Source/GLTFExporter/Public/Exporters/GLTFLevelExporter.h
--- a/Engine/Plugins/Enterprise/GLTFExporter/MigratedFromGit/Source/GLTFExporter/Public/Exporters/GLTFLevelExporter.h
+++ b/Engine/Plugins/Enterprise/GLTFExporter/MigratedFromGit/Source/GLTFExporter/Public/Exporters/GLTFLevelExporter.h
@@ -14,5 +14,5 @@ public:
 
 	explicit UGLTFLevelExporter(const FObjectInitializer& ObjectInitializer = FObjectInitializer());
 
-	bool Add(FGLTFContainerBuilder& Builder, const UObject* Object) override;
+	bool AddObject(FGLTFContainerBuilder& Builder, const UObject* Object) override;
 };
diff --git a/Engine/Plugins/Enterprise/GLTFExporter/MigratedFromGit/Source/GLTFExporter/Public/Exporters/GLTFMaterialExporter.h b/Engine/Plugins/Enterprise/GLTFExporter/MigratedFromGit/Source/GLTFExporter/Public/Exporters/GLTFMaterialExporter.h
--- a/Engine/Plugins/Enterprise/GLTFExporter/MigratedFromGit/Source/GLTFExporter/Public/Exporters/GLTFMaterialExporter.h
+++ b/Engine/Plugins/Enterprise/GLTFExporter/MigratedFromGit/Source/GLTFExporter/Public/Exporters/GLTFMaterialExporter.h
@@ -14,5 +14,5 @@ public:
 
 	explicit UGLTFMaterialExporter(const FObjectInitializer& ObjectInitializer = FObjectInitializer());
 
-	bool Add(FGLTFContainerBuilder& Builder, const UObject* Object) override;
+	bool AddObject(FGLTFContainerBuilder& Builder, const UObject* Object) override;
 };
diff --git a/Engine/Plugins/Enterprise/GLTFExporter/MigratedFromGit/Source/GLTFExporter/Public/Exporters/GLTFSkeletalMeshExporter.h b/Engine/Plugins/Enterprise/GLTFExporter/MigratedFromGit/Source/GLTFExporter/Public/Exporters/GLTFSkeletalMeshExporter.h
--- a/Engine/Plugins/Enterprise/GLTFExporter/MigratedFromGit/Source/GLTFExporter/Public/Exporters/GLTFSkeletalMeshExporter.h
+++ b/Engine/Plugins/Enterprise/GLTFExporter/MigratedFromGit/Source/GLTFExporter/Public/Exporters/GLTFSkeletalMeshExporter.h
@@ -14,5 +14,5 @@ public:
 
 	explicit UGLTFSkeletalMeshExporter(const FObjectInitializer& ObjectInitializer = FObjectInitializer());
 
-	bool Add(FGLTFContainerBuilder& Builder, const UObject* Object) override;
+	bool AddObject(FGLTFContainerBuilder& Builder, const UObject* Object) override;
 };
diff --git a/Engine/Plugins/Enterprise/GLTFExporter/MigratedFromGit/Source/GLTFExporter/Public/Exporters/GLTFStaticMeshExporter.h b/Engine/Plugins/Enterprise/GLTFExporter/MigratedFromGit/Source/GLTFExporter/Public/Exporters/GLTFStaticMeshExporter.h
--- a/Engine/Plugins/Enterprise/GLTFExporter/MigratedFromGit/Source/GLTFExporter/Public/Exporters/GLTFStaticMeshExporter.h
+++ b/Engine/Plugins/Enterprise/GLTFExporter/MigratedFromGit/Source/GLTFExporter/Public/Exporters/GLTFStaticMeshExporter.h
@@ -14,5 +14,5 @@ public:
 
 	explicit UGLTFStaticMeshExporter(const FObjectInitializer& ObjectInitializer = FObjectInitializer());
 
-	bool Add(FGLTFContainerBuilder& Builder, const UObject* Object) override;
+	bool AddObject(FGLTFContainerBuilder& Builder, const UObject* Object) override;
 };
