[GLTFExporter] refactor: mark base exporter class as abstract to make sure it is not considered a exporter implementation.

commit 5d95acc6f62d3cdacfec2eca43241cf0d3dbc78f
Date: Fri, 5 Jun 2020 20:53:47 +0200

    refactor: mark base exporter class as abstract to make sure it is not considered a exporter implementation.
    (cherry picked from commit 29bb0343e8ab104ed145925db59f7478b0e5a3fe​)

#skipci

diff --git a/Engine/Plugins/Enterprise/GLTFExporter/MigratedFromGit/Source/GLTFExporter/Public/Exporters/GLTFExporter.h b/Engine/Plugins/Enterprise/GLTFExporter/MigratedFromGit/Source/GLTFExporter/Public/Exporters/GLTFExporter.h
--- a/Engine/Plugins/Enterprise/GLTFExporter/MigratedFromGit/Source/GLTFExporter/Public/Exporters/GLTFExporter.h
+++ b/Engine/Plugins/Enterprise/GLTFExporter/MigratedFromGit/Source/GLTFExporter/Public/Exporters/GLTFExporter.h
@@ -9,7 +9,7 @@
 class FGLTFContainerBuilder;
 class UGLTFExportOptions;
 
-UCLASS()
+UCLASS(Abstract)
 class GLTFEXPORTER_API UGLTFExporter : public UExporter
 {
 public:
