[GLTFExporter] refactor: make bSelectedOnly global and constant (throughout an export) via GLTFConvertBuilder

commit 7547ed7793fc0a506173a936a53d527b85d523d7
Date: Fri, 5 Jun 2020 22:11:52 +0200

    refactor: make bSelectedOnly global and constant (throughout an export) via GLTFConvertBuilder
    (cherry picked from commit e107b6cd3c839b198fea7f99273a62b77c9ebe35​)
    # Conflicts from cherry-pick:
    #	Plugins/GLTFExporter/Source/GLTFExporter/Private/Builders/GLTFConvertBuilder.cpp
    #	Plugins/GLTFExporter/Source/GLTFExporter/Private/Builders/GLTFConvertBuilder.h

#skipci

diff --git a/Engine/Plugins/Enterprise/GLTFExporter/MigratedFromGit/Source/GLTFExporter/Private/Builders/GLTFContainerBuilder.h b/Engine/Plugins/Enterprise/GLTFExporter/MigratedFromGit/Source/GLTFExporter/Private/Builders/GLTFContainerBuilder.h
--- a/Engine/Plugins/Enterprise/GLTFExporter/MigratedFromGit/Source/GLTFExporter/Private/Builders/GLTFContainerBuilder.h
+++ b/Engine/Plugins/Enterprise/GLTFExporter/MigratedFromGit/Source/GLTFExporter/Private/Builders/GLTFContainerBuilder.h
@@ -6,4 +6,10 @@
 
 class FGLTFContainerBuilder : public FGLTFConvertBuilder
 {
+public:
+
+	FGLTFContainerBuilder(bool bSelectedActorsOnly)
+        : FGLTFConvertBuilder(bSelectedActorsOnly)
+	{
+	}
 };
diff --git a/Engine/Plugins/Enterprise/GLTFExporter/MigratedFromGit/Source/GLTFExporter/Private/Builders/GLTFConvertBuilder.cpp b/Engine/Plugins/Enterprise/GLTFExporter/MigratedFromGit/Source/GLTFExporter/Private/Builders/GLTFConvertBuilder.cpp
--- a/Engine/Plugins/Enterprise/GLTFExporter/MigratedFromGit/Source/GLTFExporter/Private/Builders/GLTFConvertBuilder.cpp
+++ b/Engine/Plugins/Enterprise/GLTFExporter/MigratedFromGit/Source/GLTFExporter/Private/Builders/GLTFConvertBuilder.cpp
@@ -125,25 +125,25 @@ FGLTFJsonNodeIndex FGLTFConvertBuilder::GetOrAddNode(const USceneComponent* Scen
 		return FGLTFJsonNodeIndex(INDEX_NONE);
 	}
 
-	return SceneComponentConverter.GetOrAdd(*this, DesiredName, SceneComponent, bSelectedOnly, bRootNode);
+	return SceneComponentConverter.GetOrAdd(*this, DesiredName, SceneComponent);
 }
 
-FGLTFJsonSceneIndex FGLTFConvertBuilder::GetOrAddScene(const ULevel* Level, bool bSelectedOnly, const FString& DesiredName)
+FGLTFJsonSceneIndex FGLTFConvertBuilder::GetOrAddScene(const ULevel* Level, const FString& DesiredName)
 {
 	if (Level == nullptr)
 	{
 		return FGLTFJsonSceneIndex(INDEX_NONE);
 	}
 
-	return LevelConverter.GetOrAdd(*this, DesiredName, Level, bSelectedOnly);
+	return LevelConverter.GetOrAdd(*this, DesiredName, Level);
 }
 
-FGLTFJsonSceneIndex FGLTFConvertBuilder::GetOrAddScene(const UWorld* World, bool bSelectedOnly, const FString& DesiredName)
+FGLTFJsonSceneIndex FGLTFConvertBuilder::GetOrAddScene(const UWorld* World, const FString& DesiredName)
 {
 	if (World == nullptr)
 	{
 		return FGLTFJsonSceneIndex(INDEX_NONE);
 	}
 
-	return GetOrAddScene(World->PersistentLevel, bSelectedOnly, DesiredName.IsEmpty() ? World->GetName() : DesiredName);
+	return GetOrAddScene(World->PersistentLevel, DesiredName.IsEmpty() ? World->GetName() : DesiredName);
 }
diff --git a/Engine/Plugins/Enterprise/GLTFExporter/MigratedFromGit/Source/GLTFExporter/Private/Builders/GLTFConvertBuilder.h b/Engine/Plugins/Enterprise/GLTFExporter/MigratedFromGit/Source/GLTFExporter/Private/Builders/GLTFConvertBuilder.h
--- a/Engine/Plugins/Enterprise/GLTFExporter/MigratedFromGit/Source/GLTFExporter/Private/Builders/GLTFConvertBuilder.h
+++ b/Engine/Plugins/Enterprise/GLTFExporter/MigratedFromGit/Source/GLTFExporter/Private/Builders/GLTFConvertBuilder.h
@@ -13,6 +13,13 @@ class FGLTFConvertBuilder : public FGLTFImageBuilder
 {
 public:
 
+    const bool bSelectedActorsOnly;
+
+	FGLTFConvertBuilder(bool bSelectedActorsOnly)
+		: bSelectedActorsOnly(bSelectedActorsOnly)
+	{
+	}
+
 	FGLTFJsonAccessorIndex GetOrAddPositionAccessor(const FPositionVertexBuffer* VertexBuffer, const FString& DesiredName = TEXT(""));
 	FGLTFJsonAccessorIndex GetOrAddColorAccessor(const FColorVertexBuffer* VertexBuffer, const FString& DesiredName = TEXT(""));
 	FGLTFJsonAccessorIndex GetOrAddNormalAccessor(const FStaticMeshVertexBuffer* VertexBuffer, const FString& DesiredName = TEXT(""));
diff --git a/Engine/Plugins/Enterprise/GLTFExporter/MigratedFromGit/Source/GLTFExporter/Private/Converters/GLTFLevelConverters.cpp b/Engine/Plugins/Enterprise/GLTFExporter/MigratedFromGit/Source/GLTFExporter/Private/Converters/GLTFLevelConverters.cpp
--- a/Engine/Plugins/Enterprise/GLTFExporter/MigratedFromGit/Source/GLTFExporter/Private/Converters/GLTFLevelConverters.cpp
+++ b/Engine/Plugins/Enterprise/GLTFExporter/MigratedFromGit/Source/GLTFExporter/Private/Converters/GLTFLevelConverters.cpp
@@ -4,12 +4,16 @@
 #include "Builders/GLTFContainerBuilder.h"
 #include "Converters/GLTFConverterUtility.h"
 
-FGLTFJsonNodeIndex FGLTFSceneComponentConverter::Add(FGLTFConvertBuilder& Builder, const FString& Name, const USceneComponent* SceneComponent, bool bSelectedOnly, bool bRootNode)
+FGLTFJsonNodeIndex FGLTFSceneComponentConverter::Add(FGLTFConvertBuilder& Builder, const FString& Name, const USceneComponent* SceneComponent)
 {
 	const AActor* Owner = SceneComponent->GetOwner();
-	const bool bIsRootComponent = Owner != nullptr && Owner->GetRootComponent() == SceneComponent;
+	const USceneComponent* ParentComponent = SceneComponent->GetAttachParent();
 	const UBlueprint* Blueprint = UBlueprint::GetBlueprintFromClass(Owner != nullptr ? Owner->GetClass() : nullptr);
 
+	const bool bSelectedOnly = Builder.bSelectedActorsOnly;
+	const bool bIsRootComponent = Owner != nullptr && Owner->GetRootComponent() == SceneComponent;
+	const bool bRootNode = bIsRootComponent && (ParentComponent == nullptr || (bSelectedOnly && !FGLTFConverterUtility::IsSelected(ParentComponent)));
+
 	FGLTFJsonNode Node;
 	Node.Name = Name.IsEmpty() ? Owner->GetName() + TEXT("_") + SceneComponent->GetName() : Name;
 
@@ -40,7 +44,7 @@ FGLTFJsonNodeIndex FGLTFSceneComponentConverter::Add(FGLTFConvertBuilder& Builde
 	{
 		if (ChildComponent != nullptr && (!bSelectedOnly || FGLTFConverterUtility::IsSelected(ChildComponent)))
 		{
-			FGLTFJsonNodeIndex NodeIndex = Builder.GetOrAddNode(ChildComponent, bSelectedOnly, false);
+			FGLTFJsonNodeIndex NodeIndex = Builder.GetOrAddNode(ChildComponent);
 			if (NodeIndex != INDEX_NONE)
 			{
 				Node.Children.Add(NodeIndex);
@@ -51,8 +55,10 @@ FGLTFJsonNodeIndex FGLTFSceneComponentConverter::Add(FGLTFConvertBuilder& Builde
 	return Builder.AddNode(Node);
 }
 
-FGLTFJsonSceneIndex FGLTFLevelConverter::Add(FGLTFConvertBuilder& Builder, const FString& Name, const ULevel* Level, bool bSelectedOnly)
+FGLTFJsonSceneIndex FGLTFLevelConverter::Add(FGLTFConvertBuilder& Builder, const FString& Name, const ULevel* Level)
 {
+	const bool bSelectedOnly = Builder.bSelectedActorsOnly;
+
 	FGLTFJsonScene Scene;
 	Scene.Name = Name;
 
@@ -69,7 +75,7 @@ FGLTFJsonSceneIndex FGLTFLevelConverter::Add(FGLTFConvertBuilder& Builder, const
 				const USceneComponent* ParentComponent = RootComponent->GetAttachParent();
 				if (ParentComponent == nullptr || (bSelectedOnly && !FGLTFConverterUtility::IsSelected(ParentComponent)))
 				{
-					FGLTFJsonNodeIndex NodeIndex = Builder.GetOrAddNode(RootComponent, bSelectedOnly, true);
+					FGLTFJsonNodeIndex NodeIndex = Builder.GetOrAddNode(RootComponent);
 					if (NodeIndex != INDEX_NONE)
 					{
 						Scene.Nodes.Add(NodeIndex);
diff --git a/Engine/Plugins/Enterprise/GLTFExporter/MigratedFromGit/Source/GLTFExporter/Private/Converters/GLTFLevelConverters.h b/Engine/Plugins/Enterprise/GLTFExporter/MigratedFromGit/Source/GLTFExporter/Private/Converters/GLTFLevelConverters.h
--- a/Engine/Plugins/Enterprise/GLTFExporter/MigratedFromGit/Source/GLTFExporter/Private/Converters/GLTFLevelConverters.h
+++ b/Engine/Plugins/Enterprise/GLTFExporter/MigratedFromGit/Source/GLTFExporter/Private/Converters/GLTFLevelConverters.h
@@ -6,12 +6,12 @@
 #include "Converters/GLTFConverter.h"
 #include "Engine.h"
 
-class FGLTFSceneComponentConverter final : public TGLTFConverter<FGLTFJsonNodeIndex, const USceneComponent*, bool, bool>
+class FGLTFSceneComponentConverter final : public TGLTFConverter<FGLTFJsonNodeIndex, const USceneComponent*>
 {
-	FGLTFJsonNodeIndex Add(FGLTFConvertBuilder& Builder, const FString& Name, const USceneComponent* SceneComponent, bool bSelectedOnly, bool bRootNode) override;
+	FGLTFJsonNodeIndex Add(FGLTFConvertBuilder& Builder, const FString& Name, const USceneComponent* SceneComponent) override;
 };
 
-class FGLTFLevelConverter final : public TGLTFConverter<FGLTFJsonSceneIndex, const ULevel*, bool>
+class FGLTFLevelConverter final : public TGLTFConverter<FGLTFJsonSceneIndex, const ULevel*>
 {
-	FGLTFJsonSceneIndex Add(FGLTFConvertBuilder& Builder, const FString& Name, const ULevel* Level, bool bSelectedOnly) override;
+	FGLTFJsonSceneIndex Add(FGLTFConvertBuilder& Builder, const FString& Name, const ULevel* Level) override;
 };
diff --git a/Engine/Plugins/Enterprise/GLTFExporter/MigratedFromGit/Source/GLTFExporter/Private/Exporters/GLTFExporter.cpp b/Engine/Plugins/Enterprise/GLTFExporter/MigratedFromGit/Source/GLTFExporter/Private/Exporters/GLTFExporter.cpp
--- a/Engine/Plugins/Enterprise/GLTFExporter/MigratedFromGit/Source/GLTFExporter/Private/Exporters/GLTFExporter.cpp
+++ b/Engine/Plugins/Enterprise/GLTFExporter/MigratedFromGit/Source/GLTFExporter/Private/Exporters/GLTFExporter.cpp
@@ -31,7 +31,7 @@ bool UGLTFExporter::ExportBinary(UObject* Object, const TCHAR* Type, FArchive& A
 		return false;
 	}
 
-	FGLTFContainerBuilder Builder;
+	FGLTFContainerBuilder Builder(bSelectedOnly);
 	if (!AddObject(Builder, Object))
 	{
 		// TODO: Report error
diff --git a/Engine/Plugins/Enterprise/GLTFExporter/MigratedFromGit/Source/GLTFExporter/Private/Exporters/GLTFLevelExporter.cpp b/Engine/Plugins/Enterprise/GLTFExporter/MigratedFromGit/Source/GLTFExporter/Private/Exporters/GLTFLevelExporter.cpp
--- a/Engine/Plugins/Enterprise/GLTFExporter/MigratedFromGit/Source/GLTFExporter/Private/Exporters/GLTFLevelExporter.cpp
+++ b/Engine/Plugins/Enterprise/GLTFExporter/MigratedFromGit/Source/GLTFExporter/Private/Exporters/GLTFLevelExporter.cpp
@@ -13,7 +13,7 @@ UGLTFLevelExporter::UGLTFLevelExporter(const FObjectInitializer& ObjectInitializ
 bool UGLTFLevelExporter::AddObject(FGLTFContainerBuilder& Builder, const UObject* Object)
 {
 	const UWorld* World = CastChecked<UWorld>(Object);
-	const FGLTFJsonSceneIndex SceneIndex = Builder.GetOrAddScene(World, bSelectedOnly);
+	const FGLTFJsonSceneIndex SceneIndex = Builder.GetOrAddScene(World);
 
 	Builder.DefaultScene = SceneIndex;
 	return true;
