[GLTFExporter] feat(json): remove engine patch version from generator string and add plugin name and version

commit be754ad0d5d46df4e794243f7c288171eb4c1364
Date: Fri, 5 Jun 2020 23:35:28 +0200

    feat(json): remove engine patch version from generator string and add plugin name and version
    (cherry picked from commit 40550f238a72a9f4e70c5a122642686d2fd8957f​)
    # Conflicts from cherry-pick:
    #	Plugins/GLTFExporter/Source/GLTFExporter/GLTFExporter.Build.cs

#skipci

diff --git a/Engine/Plugins/Enterprise/GLTFExporter/MigratedFromGit/Source/GLTFExporter/GLTFExporter.Build.cs b/Engine/Plugins/Enterprise/GLTFExporter/MigratedFromGit/Source/GLTFExporter/GLTFExporter.Build.cs
--- a/Engine/Plugins/Enterprise/GLTFExporter/MigratedFromGit/Source/GLTFExporter/GLTFExporter.Build.cs
+++ b/Engine/Plugins/Enterprise/GLTFExporter/MigratedFromGit/Source/GLTFExporter/GLTFExporter.Build.cs
@@ -6,22 +6,6 @@ namespace UnrealBuildTool.Rules
 	{
 		public GLTFExporter(ReadOnlyTargetRules Target) : base(Target)
 		{
-			PublicIncludePaths.AddRange(
-				new string[] {
-				}
-				);
-
-			PrivateIncludePaths.AddRange(
-				new string[] {
-				}
-				);
-
-			PublicDependencyModuleNames.AddRange(
-				new string[]
-				{
-				}
-				);
-
 			PrivateDependencyModuleNames.AddRange(
 				new string[]
 				{
@@ -40,6 +24,7 @@ namespace UnrealBuildTool.Rules
 					"Mainframe",
 					"InputCore",
 					"EditorStyle",
+					"Projects",
 					"RenderCore",
 					"RHI",
 					"MaterialBaking"
diff --git a/Engine/Plugins/Enterprise/GLTFExporter/MigratedFromGit/Source/GLTFExporter/Private/Json/GLTFJsonRoot.cpp b/Engine/Plugins/Enterprise/GLTFExporter/MigratedFromGit/Source/GLTFExporter/Private/Json/GLTFJsonRoot.cpp
new file mode 100644
--- /dev/null
+++ b/Engine/Plugins/Enterprise/GLTFExporter/MigratedFromGit/Source/GLTFExporter/Private/Json/GLTFJsonRoot.cpp
@@ -0,0 +1,16 @@
+// Copyright Epic Games, Inc. All Rights Reserved.
+
+#include "Json/GLTFJsonRoot.h"
+#include "GLTFExporterModule.h"
+#include "Interfaces/IPluginManager.h"
+#include "Runtime/Launch/Resources/Version.h"
+
+FGLTFJsonAsset::FGLTFJsonAsset()
+	: Version(TEXT("2.0"))
+{
+	const TSharedPtr<IPlugin> Plugin = IPluginManager::Get().FindPlugin(GLTFEXPORTER_MODULE_NAME);
+	const FPluginDescriptor& PluginDescriptor = Plugin->GetDescriptor();
+
+	Generator = TEXT(EPIC_PRODUCT_NAME) TEXT(" ") VERSION_STRINGIFY(ENGINE_MAJOR_VERSION) TEXT(".") VERSION_STRINGIFY(ENGINE_MINOR_VERSION)
+		TEXT(" ") + PluginDescriptor.FriendlyName + TEXT(" ") + PluginDescriptor.VersionName;
+}
diff --git a/Engine/Plugins/Enterprise/GLTFExporter/MigratedFromGit/Source/GLTFExporter/Private/Json/GLTFJsonRoot.h b/Engine/Plugins/Enterprise/GLTFExporter/MigratedFromGit/Source/GLTFExporter/Private/Json/GLTFJsonRoot.h
--- a/Engine/Plugins/Enterprise/GLTFExporter/MigratedFromGit/Source/GLTFExporter/Private/Json/GLTFJsonRoot.h
+++ b/Engine/Plugins/Enterprise/GLTFExporter/MigratedFromGit/Source/GLTFExporter/Private/Json/GLTFJsonRoot.h
@@ -16,7 +16,6 @@
 
 #include "Containers/Set.h"
 #include "Policies/CondensedJsonPrintPolicy.h"
-#include "Runtime/Launch/Resources/Version.h"
 
 
 struct FGLTFJsonAsset
@@ -25,11 +24,7 @@ struct FGLTFJsonAsset
 	FString Generator;
 	FString Copyright;
 
-	FGLTFJsonAsset()
-		: Version(TEXT("2.0"))
-		, Generator(TEXT(EPIC_PRODUCT_NAME) TEXT(" ") ENGINE_VERSION_STRING)
-	{
-	}
+	FGLTFJsonAsset();
 
 	template <class CharType = TCHAR, class PrintPolicy = TPrettyJsonPrintPolicy<CharType>>
 	void WriteObject(TJsonWriter<CharType, PrintPolicy>& JsonWriter) const
diff --git a/Engine/Plugins/Enterprise/GLTFExporter/MigratedFromGit/Source/GLTFExporter/Public/GLTFExporterModule.h b/Engine/Plugins/Enterprise/GLTFExporter/MigratedFromGit/Source/GLTFExporter/Public/GLTFExporterModule.h
--- a/Engine/Plugins/Enterprise/GLTFExporter/MigratedFromGit/Source/GLTFExporter/Public/GLTFExporterModule.h
+++ b/Engine/Plugins/Enterprise/GLTFExporter/MigratedFromGit/Source/GLTFExporter/Public/GLTFExporterModule.h
@@ -6,6 +6,8 @@
 #include "Modules/ModuleInterface.h"
 #include "Modules/ModuleManager.h"
 
+#define GLTFEXPORTER_MODULE_NAME TEXT("GLTFExporter")
+
 DECLARE_LOG_CATEGORY_EXTERN(LogGLTFExporter, Log, All);
 
 /**
@@ -21,7 +23,7 @@ public:
 	 */
 	static IGLTFExporterModule& Get()
 	{
-		return FModuleManager::LoadModuleChecked<IGLTFExporterModule>("GLTFExporter");
+		return FModuleManager::LoadModuleChecked<IGLTFExporterModule>(GLTFEXPORTER_MODULE_NAME);
 	}
 
 	/**
@@ -31,6 +33,6 @@ public:
 	 */
 	static bool IsAvailable()
 	{
-		return FModuleManager::Get().IsModuleLoaded("GLTFExporter");
+		return FModuleManager::Get().IsModuleLoaded(GLTFEXPORTER_MODULE_NAME);
 	}
 };
