[GLTFExporter] Merge branch 'master' into feature/initial-support-for-baked-materials

commit daf4a253d58217ccc6c1aeca6937bf99fe975592
Date: Mon, 8 Jun 2020 12:18:35 +0200

    Merge branch 'master' into feature/initial-support-for-baked-materials
    # Conflicts from merge:
    #	Plugins/GLTFExporter/Source/GLTFExporter/GLTFExporter.Build.cs
    #	Plugins/GLTFExporter/Source/GLTFExporter/Private/Builders/GLTFConvertBuilder.cpp
    #	Plugins/GLTFExporter/Source/GLTFExporter/Private/Builders/GLTFConvertBuilder.h
    (cherry picked from commit 8548728605587fc4c0d0acf35c48985808c2889b​)

#skipci

diff --git a/Engine/Plugins/Enterprise/GLTFExporter/MigratedFromGit/Source/GLTFExporter/GLTFExporter.Build.cs b/Engine/Plugins/Enterprise/GLTFExporter/MigratedFromGit/Source/GLTFExporter/GLTFExporter.Build.cs
--- a/Engine/Plugins/Enterprise/GLTFExporter/MigratedFromGit/Source/GLTFExporter/GLTFExporter.Build.cs
+++ b/Engine/Plugins/Enterprise/GLTFExporter/MigratedFromGit/Source/GLTFExporter/GLTFExporter.Build.cs
@@ -30,12 +30,6 @@ namespace UnrealBuildTool.Rules
 					"MaterialBaking"
 				}
 				);
-
-			DynamicallyLoadedModuleNames.AddRange(
-				new string[]
-				{
-				}
-				);
 		}
 	}
 }
diff --git a/Engine/Plugins/Enterprise/GLTFExporter/MigratedFromGit/Source/GLTFExporter/Private/Builders/GLTFConvertBuilder.cpp b/Engine/Plugins/Enterprise/GLTFExporter/MigratedFromGit/Source/GLTFExporter/Private/Builders/GLTFConvertBuilder.cpp
--- a/Engine/Plugins/Enterprise/GLTFExporter/MigratedFromGit/Source/GLTFExporter/Private/Builders/GLTFConvertBuilder.cpp
+++ b/Engine/Plugins/Enterprise/GLTFExporter/MigratedFromGit/Source/GLTFExporter/Private/Builders/GLTFConvertBuilder.cpp
@@ -118,7 +118,7 @@ FGLTFJsonTextureIndex FGLTFConvertBuilder::GetOrAddTexture(const UTexture2D* Tex
 	return Texture2DConverter.GetOrAdd(*this, DesiredName.IsEmpty() ? Texture->GetName() : DesiredName, Texture);
 }
 
-FGLTFJsonNodeIndex FGLTFConvertBuilder::GetOrAddNode(const USceneComponent* SceneComponent, bool bSelectedOnly, bool bRootNode, const FString& DesiredName)
+FGLTFJsonNodeIndex FGLTFConvertBuilder::GetOrAddNode(const USceneComponent* SceneComponent, const FString& DesiredName)
 {
 	if (SceneComponent == nullptr)
 	{
diff --git a/Engine/Plugins/Enterprise/GLTFExporter/MigratedFromGit/Source/GLTFExporter/Private/Builders/GLTFConvertBuilder.h b/Engine/Plugins/Enterprise/GLTFExporter/MigratedFromGit/Source/GLTFExporter/Private/Builders/GLTFConvertBuilder.h
--- a/Engine/Plugins/Enterprise/GLTFExporter/MigratedFromGit/Source/GLTFExporter/Private/Builders/GLTFConvertBuilder.h
+++ b/Engine/Plugins/Enterprise/GLTFExporter/MigratedFromGit/Source/GLTFExporter/Private/Builders/GLTFConvertBuilder.h
@@ -34,9 +34,9 @@ public:
 	FGLTFJsonMaterialIndex GetOrAddMaterial(const UMaterialInterface* Material, const FString& DesiredName = TEXT(""));
 	FGLTFJsonTextureIndex GetOrAddTexture(const UTexture2D* Texture, const FString& DesiredName = TEXT(""));
 
-	FGLTFJsonNodeIndex GetOrAddNode(const USceneComponent* SceneComponent, bool bSelectedOnly, bool bRootNode = false, const FString& DesiredName = TEXT(""));
-	FGLTFJsonSceneIndex GetOrAddScene(const ULevel* Level, bool bSelectedOnly, const FString& DesiredName = TEXT(""));
-	FGLTFJsonSceneIndex GetOrAddScene(const UWorld* World, bool bSelectedOnly, const FString& DesiredName = TEXT(""));
+	FGLTFJsonNodeIndex GetOrAddNode(const USceneComponent* SceneComponent, const FString& DesiredName = TEXT(""));
+	FGLTFJsonSceneIndex GetOrAddScene(const ULevel* Level, const FString& DesiredName = TEXT(""));
+	FGLTFJsonSceneIndex GetOrAddScene(const UWorld* World, const FString& DesiredName = TEXT(""));
 
 private:
 
