[GLTFExporter] refactor(material): create the FCanvas-instance on the stack instead of on the heap based on feedback

commit ecbbcdaed16a163f2b3015ab851dffb5448989bb
Date: Mon, 8 Jun 2020 14:09:37 +0200

    refactor(material): create the FCanvas-instance on the stack instead of on the heap based on feedback
    (cherry picked from commit 98a4c29ae134d983487be5ef54e49a2006513014​)

#skipci

diff --git a/Engine/Plugins/Enterprise/GLTFExporter/MigratedFromGit/Source/GLTFExporter/Private/Converters/GLTFMaterialUtility.cpp b/Engine/Plugins/Enterprise/GLTFExporter/MigratedFromGit/Source/GLTFExporter/Private/Converters/GLTFMaterialUtility.cpp
--- a/Engine/Plugins/Enterprise/GLTFExporter/MigratedFromGit/Source/GLTFExporter/Private/Converters/GLTFMaterialUtility.cpp
+++ b/Engine/Plugins/Enterprise/GLTFExporter/MigratedFromGit/Source/GLTFExporter/Private/Converters/GLTFMaterialUtility.cpp
@@ -40,10 +40,10 @@ bool FGLTFMaterialUtility::CombineTextures(TArray<FColor>& OutPixels, const TArr
 	RenderTarget2D->TargetGamma = 0;
 
 	FRenderTarget* RenderTarget = RenderTarget2D->GameThread_GetRenderTargetResource();
-	FCanvas* Canvas = new FCanvas(RenderTarget, nullptr, 0, 0, 0, GMaxRHIFeatureLevel);
+	FCanvas Canvas(RenderTarget, nullptr, 0, 0, 0, GMaxRHIFeatureLevel);
 
-	Canvas->SetRenderTarget_GameThread(RenderTarget);
-	Canvas->Clear(FLinearColor(0, 0, 0, 0));
+	Canvas.SetRenderTarget_GameThread(RenderTarget);
+	Canvas.Clear(FLinearColor(0, 0, 0, 0));
 
 	const FVector2D TileSize(OutputSize.X, OutputSize.Y);
 	const FVector2D TilePosition(0, 0);
@@ -55,12 +55,12 @@ bool FGLTFMaterialUtility::CombineTextures(TArray<FColor>& OutPixels, const TArr
 		FCanvasTileItem TileItem(TilePosition, Source.Texture->Resource, TileSize, Source.TintColor);
 
 		TileItem.BlendMode = Source.BlendMode;
-		TileItem.Draw(Canvas);
+		TileItem.Draw(&Canvas);
 	}
 
-	Canvas->Flush_GameThread();
+	Canvas.Flush_GameThread();
 	FlushRenderingCommands();
-	Canvas->SetRenderTarget_GameThread(nullptr);
+	Canvas.SetRenderTarget_GameThread(nullptr);
 	FlushRenderingCommands();
 
 	const bool bReadSuccessful = RenderTarget->ReadPixels(OutPixels);
@@ -69,7 +69,6 @@ bool FGLTFMaterialUtility::CombineTextures(TArray<FColor>& OutPixels, const TArr
 	RenderTarget2D->ReleaseResource();
 	RenderTarget2D->RemoveFromRoot();
 	RenderTarget2D = nullptr;
-	delete Canvas;
 
 	return bReadSuccessful;
 }
