[GLTFExporter] refactor(material): remove redundant check() after creating a rendertarget

commit ab54a7fa9aefd6f23027a90e2f6365ded2a41d5c
Date: Mon, 8 Jun 2020 14:59:52 +0200

    refactor(material): remove redundant check() after creating a rendertarget
    (cherry picked from commit 21fb1ae98339bddd1e2e9ed49e704b34fe5e1a6b​)

#skipci

diff --git a/Engine/Plugins/Enterprise/GLTFExporter/MigratedFromGit/Source/GLTFExporter/Private/Converters/GLTFMaterialUtility.cpp b/Engine/Plugins/Enterprise/GLTFExporter/MigratedFromGit/Source/GLTFExporter/Private/Converters/GLTFMaterialUtility.cpp
--- a/Engine/Plugins/Enterprise/GLTFExporter/MigratedFromGit/Source/GLTFExporter/Private/Converters/GLTFMaterialUtility.cpp
+++ b/Engine/Plugins/Enterprise/GLTFExporter/MigratedFromGit/Source/GLTFExporter/Private/Converters/GLTFMaterialUtility.cpp
@@ -32,7 +32,6 @@ UTexture2D* FGLTFMaterialUtility::CreateTransientTexture(const TArray<FColor>& P
 bool FGLTFMaterialUtility::CombineTextures(TArray<FColor>& OutPixels, const TArray<FGLTFTextureCombineSource>& Sources, const FIntPoint& OutputSize, const EPixelFormat OutputPixelFormat)
 {
 	UTextureRenderTarget2D* RenderTarget2D = NewObject<UTextureRenderTarget2D>();
-	check(RenderTarget2D);
 
 	RenderTarget2D->AddToRoot();
 	RenderTarget2D->ClearColor = FLinearColor(0.0f, 0.0f, 0.0f, 0.0f);
