[GLTFExporter] refactor(material): make SRGB-format for transient textures settable by optional parameter

commit ea9cb867cbfb6e2bbebcc6bd3bb2adb6a443f7d3
Date: Mon, 8 Jun 2020 15:09:53 +0200

    refactor(material): make SRGB-format for transient textures settable by optional parameter
    (cherry picked from commit 1238849a6e7d8d6c0eb5510b36ed4d67434e3cdc​)

#skipci

diff --git a/Engine/Plugins/Enterprise/GLTFExporter/MigratedFromGit/Source/GLTFExporter/Private/Converters/GLTFMaterialUtility.cpp b/Engine/Plugins/Enterprise/GLTFExporter/MigratedFromGit/Source/GLTFExporter/Private/Converters/GLTFMaterialUtility.cpp
--- a/Engine/Plugins/Enterprise/GLTFExporter/MigratedFromGit/Source/GLTFExporter/Private/Converters/GLTFMaterialUtility.cpp
+++ b/Engine/Plugins/Enterprise/GLTFExporter/MigratedFromGit/Source/GLTFExporter/Private/Converters/GLTFMaterialUtility.cpp
@@ -10,7 +10,7 @@
 #include "IMaterialBakingModule.h"
 #include "MaterialBakingStructures.h"
 
-UTexture2D* FGLTFMaterialUtility::CreateTransientTexture(const TArray<FColor>& Pixels, const FIntPoint& TextureSize, EPixelFormat TextureFormat)
+UTexture2D* FGLTFMaterialUtility::CreateTransientTexture(const TArray<FColor>& Pixels, const FIntPoint& TextureSize, EPixelFormat TextureFormat, bool bUseSRGB)
 {
 	check(TextureSize.X * TextureSize.Y == Pixels.Num());
 
@@ -20,7 +20,7 @@ UTexture2D* FGLTFMaterialUtility::CreateTransientTexture(const TArray<FColor>& P
 	FMemory::Memcpy(MipData, Pixels.GetData(), Pixels.Num() * Pixels.GetTypeSize());
 	Texture->PlatformData->Mips[0].BulkData.Unlock();
 
-	// Texture->SRGB = false;	// TODO: Do we need this?
+	Texture->SRGB = bUseSRGB ? 1 : 0;
 	Texture->CompressionNone = true;
 	Texture->MipGenSettings = TMGS_NoMipmaps;
 	Texture->CompressionSettings = TC_Default;
@@ -127,10 +127,10 @@ FGLTFPropertyBakeOutput FGLTFMaterialUtility::BakeMaterialProperty(const FIntPoi
 	};
 }
 
-UTexture2D* FGLTFMaterialUtility::BakeMaterialPropertyToTexture(const FIntPoint& OutputSize, EMaterialProperty MaterialProperty, const UMaterialInterface* Material, bool bCopyAlphaFromRedChannel)
+UTexture2D* FGLTFMaterialUtility::BakeMaterialPropertyToTexture(const FIntPoint& OutputSize, EMaterialProperty MaterialProperty, const UMaterialInterface* Material, bool bCopyAlphaFromRedChannel, bool bUseSRGB)
 {
 	const FGLTFPropertyBakeOutput PropertyBakeOutput = BakeMaterialProperty(OutputSize, MaterialProperty, Material, bCopyAlphaFromRedChannel);
-	return CreateTransientTexture(PropertyBakeOutput.Pixels, PropertyBakeOutput.Size, PropertyBakeOutput.PixelFormat);
+	return CreateTransientTexture(PropertyBakeOutput.Pixels, PropertyBakeOutput.Size, PropertyBakeOutput.PixelFormat, bUseSRGB);
 }
 
 FGLTFJsonTextureIndex FGLTFMaterialUtility::AddCombinedTexture(FGLTFConvertBuilder& Builder, const TArray<FGLTFTextureCombineSource>& CombineSources, const FIntPoint& TextureSize, const FString& TextureName, EGLTFJsonTextureFilter Filter, EGLTFJsonTextureWrap Wrap)
diff --git a/Engine/Plugins/Enterprise/GLTFExporter/MigratedFromGit/Source/GLTFExporter/Private/Converters/GLTFMaterialUtility.h b/Engine/Plugins/Enterprise/GLTFExporter/MigratedFromGit/Source/GLTFExporter/Private/Converters/GLTFMaterialUtility.h
--- a/Engine/Plugins/Enterprise/GLTFExporter/MigratedFromGit/Source/GLTFExporter/Private/Converters/GLTFMaterialUtility.h
+++ b/Engine/Plugins/Enterprise/GLTFExporter/MigratedFromGit/Source/GLTFExporter/Private/Converters/GLTFMaterialUtility.h
@@ -31,10 +31,10 @@ struct FGLTFPropertyBakeOutput
 
 struct FGLTFMaterialUtility
 {
-	static UTexture2D* CreateTransientTexture(const TArray<FColor>& Pixels, const FIntPoint& TextureSize, EPixelFormat TextureFormat);
+	static UTexture2D* CreateTransientTexture(const TArray<FColor>& Pixels, const FIntPoint& TextureSize, EPixelFormat TextureFormat, bool bUseSRGB = false);
 	static bool CombineTextures(TArray<FColor>& OutPixels, const TArray<FGLTFTextureCombineSource>& Sources, const FIntPoint& OutputSize, EPixelFormat OutputPixelFormat);
 	static FGLTFPropertyBakeOutput BakeMaterialProperty(const FIntPoint& OutputSize, EMaterialProperty MaterialProperty, const UMaterialInterface* Material, bool bCopyAlphaFromRedChannel = false);
-	static UTexture2D* BakeMaterialPropertyToTexture(const FIntPoint& OutputSize, EMaterialProperty MaterialProperty, const UMaterialInterface* Material, bool bCopyAlphaFromRedChannel = false);
+	static UTexture2D* BakeMaterialPropertyToTexture(const FIntPoint& OutputSize, EMaterialProperty MaterialProperty, const UMaterialInterface* Material, bool bCopyAlphaFromRedChannel = false, bool bUseSRGB = false);
 
 	static FGLTFJsonTextureIndex AddCombinedTexture(FGLTFConvertBuilder& Builder, const TArray<FGLTFTextureCombineSource>& CombineSources, const FIntPoint& TextureSize, const FString& TextureName, EGLTFJsonTextureFilter Filter, EGLTFJsonTextureWrap Wrap);
 	static FGLTFJsonTextureIndex AddTexture(FGLTFConvertBuilder& Builder, const TArray<FColor>& Pixels, const FIntPoint& TextureSize, const FString& TextureName, EPixelFormat PixelFormat, EGLTFJsonTextureFilter Filter, EGLTFJsonTextureWrap Wrap);
