[GLTFExporter] refactor(material): rename some baked textures to better match the glTF-standard

commit 183a4d23715afdcb02c401b1fb1915c3d267624e
Date: Mon, 8 Jun 2020 15:32:20 +0200

    refactor(material): rename some baked textures to better match the glTF-standard
    (cherry picked from commit 594f9dc038ed981c1957cdb248c3b4b685740c04​)

#skipci

diff --git a/Engine/Plugins/Enterprise/GLTFExporter/MigratedFromGit/Source/GLTFExporter/Private/Converters/GLTFMaterialConverters.cpp b/Engine/Plugins/Enterprise/GLTFExporter/MigratedFromGit/Source/GLTFExporter/Private/Converters/GLTFMaterialConverters.cpp
--- a/Engine/Plugins/Enterprise/GLTFExporter/MigratedFromGit/Source/GLTFExporter/Private/Converters/GLTFMaterialConverters.cpp
+++ b/Engine/Plugins/Enterprise/GLTFExporter/MigratedFromGit/Source/GLTFExporter/Private/Converters/GLTFMaterialConverters.cpp
@@ -551,8 +551,8 @@ bool FGLTFMaterialConverter::TryGetBakedTexture(FGLTFConvertBuilder& Builder, FG
 	{
 		case MP_BaseColor: PropertyName = TEXT("BaseColor"); break;
 		case MP_Normal: PropertyName = TEXT("Normal"); break;
-		case MP_EmissiveColor: PropertyName = TEXT("EmissiveColor"); break;
-		case MP_AmbientOcclusion: PropertyName = TEXT("AmbientOcclusion"); break;
+		case MP_EmissiveColor: PropertyName = TEXT("Emissive"); break;
+		case MP_AmbientOcclusion: PropertyName = TEXT("Occlusion"); break;
 		default:
 			// TODO: support for more properties
 			return false;
