[GLTFExporter] refactor(texture): remove FGLTFTextureRenderTarget2DConverter and FGLTFTextureRenderTargetCubeConverter, since the primary purpose of adding them was for baking, which no longer seems to need them. If we need FGLTFTextureRenderTargetCubeConverter for reflection captures later then we'll just add it.

commit ba895d015f18ed793ec0e13e5af6ab82c75374c7
Date: Mon, 8 Jun 2020 22:23:22 +0200

    refactor(texture): remove FGLTFTextureRenderTarget2DConverter and FGLTFTextureRenderTargetCubeConverter, since the primary purpose of adding them was for baking, which no longer seems to need them. If we need FGLTFTextureRenderTargetCubeConverter for reflection captures later then we'll just add it.
    (cherry picked from commit adcb88744d424eacec2a5d13d89bdd834d205070​)

#skipci

diff --git a/Engine/Plugins/Enterprise/GLTFExporter/MigratedFromGit/Source/GLTFExporter/Private/Converters/GLTFTextureConverters.cpp b/Engine/Plugins/Enterprise/GLTFExporter/MigratedFromGit/Source/GLTFExporter/Private/Converters/GLTFTextureConverters.cpp
--- a/Engine/Plugins/Enterprise/GLTFExporter/MigratedFromGit/Source/GLTFExporter/Private/Converters/GLTFTextureConverters.cpp
+++ b/Engine/Plugins/Enterprise/GLTFExporter/MigratedFromGit/Source/GLTFExporter/Private/Converters/GLTFTextureConverters.cpp
@@ -27,18 +27,6 @@ FGLTFJsonTextureIndex FGLTFTextureCubeConverter::Add(FGLTFConvertBuilder& Builde
 	return FGLTFJsonTextureIndex(INDEX_NONE);
 }
 
-FGLTFJsonTextureIndex FGLTFTextureRenderTarget2DConverter::Add(FGLTFConvertBuilder& Builder, const FString& Name, const UTextureRenderTarget2D* RenderTarget2D)
-{
-	// TODO: implement support
-	return FGLTFJsonTextureIndex(INDEX_NONE);
-}
-
-FGLTFJsonTextureIndex FGLTFTextureRenderTargetCubeConverter::Add(FGLTFConvertBuilder& Builder, const FString& Name, const UTextureRenderTargetCube* RenderTargetCube)
-{
-	// TODO: implement support
-	return FGLTFJsonTextureIndex(INDEX_NONE);
-}
-
 FGLTFJsonTextureIndex FGLTFLightMapTexture2DConverter::Add(FGLTFConvertBuilder& Builder, const FString& Name, const ULightMapTexture2D* LightMap)
 {
 	// TODO: implement support
diff --git a/Engine/Plugins/Enterprise/GLTFExporter/MigratedFromGit/Source/GLTFExporter/Private/Converters/GLTFTextureConverters.h b/Engine/Plugins/Enterprise/GLTFExporter/MigratedFromGit/Source/GLTFExporter/Private/Converters/GLTFTextureConverters.h
--- a/Engine/Plugins/Enterprise/GLTFExporter/MigratedFromGit/Source/GLTFExporter/Private/Converters/GLTFTextureConverters.h
+++ b/Engine/Plugins/Enterprise/GLTFExporter/MigratedFromGit/Source/GLTFExporter/Private/Converters/GLTFTextureConverters.h
@@ -16,16 +16,6 @@ class FGLTFTextureCubeConverter final : public TGLTFConverter<FGLTFJsonTextureIn
 	FGLTFJsonTextureIndex Add(FGLTFConvertBuilder& Builder, const FString& Name, const UTextureCube* TextureCube) override;
 };
 
-class FGLTFTextureRenderTarget2DConverter final : public TGLTFConverter<FGLTFJsonTextureIndex, const UTextureRenderTarget2D*>
-{
-	FGLTFJsonTextureIndex Add(FGLTFConvertBuilder& Builder, const FString& Name, const UTextureRenderTarget2D* RenderTarget2D) override;
-};
-
-class FGLTFTextureRenderTargetCubeConverter final : public TGLTFConverter<FGLTFJsonTextureIndex, const UTextureRenderTargetCube*>
-{
-	FGLTFJsonTextureIndex Add(FGLTFConvertBuilder& Builder, const FString& Name, const UTextureRenderTargetCube* RenderTargetCube) override;
-};
-
 class FGLTFLightMapTexture2DConverter final : public TGLTFConverter<FGLTFJsonTextureIndex, const ULightMapTexture2D*>
 {
 	FGLTFJsonTextureIndex Add(FGLTFConvertBuilder& Builder, const FString& Name, const ULightMapTexture2D* LightMap) override;
