[GLTFExporter] refactor(texture): fix minor typo in forward declaration of EImageFormat

commit 595385e1be7af71509bd8e8be58ff1c5ad6d45cb
Date: Mon, 8 Jun 2020 22:24:48 +0200

    refactor(texture): fix minor typo in forward declaration of EImageFormat
    (cherry picked from commit 547446723f69275cc316fc8350500c7a6a7c0593​)

#skipci

diff --git a/Engine/Plugins/Enterprise/GLTFExporter/MigratedFromGit/Source/GLTFExporter/Private/Builders/GLTFBuilderUtility.h b/Engine/Plugins/Enterprise/GLTFExporter/MigratedFromGit/Source/GLTFExporter/Private/Builders/GLTFBuilderUtility.h
--- a/Engine/Plugins/Enterprise/GLTFExporter/MigratedFromGit/Source/GLTFExporter/Private/Builders/GLTFBuilderUtility.h
+++ b/Engine/Plugins/Enterprise/GLTFExporter/MigratedFromGit/Source/GLTFExporter/Private/Builders/GLTFBuilderUtility.h
@@ -5,7 +5,7 @@
 #include "Engine.h"
 #include "Json/GLTFJsonEnums.h"
 
-enum class EImageFormat : signed char;
+enum class EImageFormat : int8;
 enum class ERGBFormat : int8;
 
 struct FGLTFBuilderUtility
