[GLTFExporter] refactor(texture): generalize FGLTFImageBuilder::AddImage to support more EPixelFormat

commit 87070611d857edca75903323f195317cf597c84d
Date: Mon, 8 Jun 2020 22:48:44 +0200

    refactor(texture): generalize FGLTFImageBuilder::AddImage to support more EPixelFormat
    (cherry picked from commit 9fc51cb41d3ffb04fd7c62f21f67f61af15e4a29​)

#skipci

diff --git a/Engine/Plugins/Enterprise/GLTFExporter/MigratedFromGit/Source/GLTFExporter/Private/Builders/GLTFImageBuilder.cpp b/Engine/Plugins/Enterprise/GLTFExporter/MigratedFromGit/Source/GLTFExporter/Private/Builders/GLTFImageBuilder.cpp
--- a/Engine/Plugins/Enterprise/GLTFExporter/MigratedFromGit/Source/GLTFExporter/Private/Builders/GLTFImageBuilder.cpp
+++ b/Engine/Plugins/Enterprise/GLTFExporter/MigratedFromGit/Source/GLTFExporter/Private/Builders/GLTFImageBuilder.cpp
@@ -21,12 +21,14 @@ FGLTFJsonImageIndex FGLTFImageBuilder::AddImage(const void* RawData, int64 ByteL
 		case EGLTFJsonMimeType::PNG:
 			if (!FGLTFBuilderUtility::CompressImage(RawData, ByteLength, Width, Height, RawFormat, BitDepth, ImageData, EImageFormat::PNG, Quality))
 			{
+				// TODO: report error
 				return FGLTFJsonImageIndex(INDEX_NONE);
 			}
 			break;
 		case EGLTFJsonMimeType::JPEG:
 			if (!FGLTFBuilderUtility::CompressImage(RawData, ByteLength, Width, Height, RawFormat, BitDepth, ImageData, EImageFormat::JPEG, Quality))
 			{
+				// TODO: report error
 				return FGLTFJsonImageIndex(INDEX_NONE);
 			}
 			break;
@@ -45,7 +47,7 @@ FGLTFJsonImageIndex FGLTFImageBuilder::AddImage(const void* RawData, int64 ByteL
 	return ImageIndex;
 }
 
-FGLTFJsonImageIndex FGLTFImageBuilder::AddImage(const TArray<FColor>& Pixels, int32 Width, int32 Height, EPixelFormat PixelFormat, const FString& Name, EGLTFJsonMimeType MimeType, int32 Quality)
+FGLTFJsonImageIndex FGLTFImageBuilder::AddImage(const void* RawData, int64 ByteLength, int32 Width, int32 Height, EPixelFormat PixelFormat, const FString& Name, EGLTFJsonMimeType MimeType, int32 Quality)
 {
 	ERGBFormat RawFormat;
 	int32 BitDepth;
@@ -53,15 +55,22 @@ FGLTFJsonImageIndex FGLTFImageBuilder::AddImage(const TArray<FColor>& Pixels, in
 
 	switch (PixelFormat)
 	{
-		case PF_B8G8R8A8: RawFormat = ERGBFormat::BGRA; BitDepth = 8;  bFloatFormat = false; break;
-		case PF_R8G8B8A8: RawFormat = ERGBFormat::RGBA; BitDepth = 8;  bFloatFormat = false; break;
-		case PF_G8:       RawFormat = ERGBFormat::Gray; BitDepth = 8;  bFloatFormat = false; break;
+		case PF_B8G8R8A8:           RawFormat = ERGBFormat::BGRA; BitDepth = 8;  bFloatFormat = false; break;
+		case PF_R8G8B8A8:           RawFormat = ERGBFormat::RGBA; BitDepth = 8;  bFloatFormat = false; break;
+		case PF_R16G16B16A16_UNORM: RawFormat = ERGBFormat::RGBA; BitDepth = 16; bFloatFormat = false; break;
+		case PF_FloatRGBA:          RawFormat = ERGBFormat::RGBA; BitDepth = 16; bFloatFormat = true;  break;
+		case PF_L8:                 RawFormat = ERGBFormat::Gray; BitDepth = 8;  bFloatFormat = false; break;
+		case PF_G8:                 RawFormat = ERGBFormat::Gray; BitDepth = 8;  bFloatFormat = false; break;
+		case PF_G16:                RawFormat = ERGBFormat::Gray; BitDepth = 16; bFloatFormat = false; break;
+		case PF_R16F:               RawFormat = ERGBFormat::Gray; BitDepth = 16; bFloatFormat = true;  break;
+		case PF_R16F_FILTER:        RawFormat = ERGBFormat::Gray; BitDepth = 16; bFloatFormat = true;  break;
+		case PF_R32_FLOAT:          RawFormat = ERGBFormat::Gray; BitDepth = 32; bFloatFormat = true;  break;
 		default:
 			// TODO: report error
 			return FGLTFJsonImageIndex(INDEX_NONE);
 	}
 
-	return AddImage(Pixels.GetData(), Pixels.Num() * Pixels.GetTypeSize(), Width, Height, RawFormat, BitDepth, bFloatFormat, Name, MimeType, Quality);
+	return AddImage(RawData, ByteLength, Width, Height, RawFormat, BitDepth, bFloatFormat, Name, MimeType, Quality);
 }
 
 FGLTFJsonImageIndex FGLTFImageBuilder::AddImage(const FTextureSource& Image, const FString& Name, EGLTFJsonMimeType MimeType, int32 Quality)
diff --git a/Engine/Plugins/Enterprise/GLTFExporter/MigratedFromGit/Source/GLTFExporter/Private/Builders/GLTFImageBuilder.h b/Engine/Plugins/Enterprise/GLTFExporter/MigratedFromGit/Source/GLTFExporter/Private/Builders/GLTFImageBuilder.h
--- a/Engine/Plugins/Enterprise/GLTFExporter/MigratedFromGit/Source/GLTFExporter/Private/Builders/GLTFImageBuilder.h
+++ b/Engine/Plugins/Enterprise/GLTFExporter/MigratedFromGit/Source/GLTFExporter/Private/Builders/GLTFImageBuilder.h
@@ -12,7 +12,7 @@ public:
 
 	FGLTFJsonImageIndex AddImage(const void* RawData, int64 ByteLength, int32 Width, int32 Height, ERGBFormat RawFormat, int32 BitDepth, bool bFloatFormat, const FString& Name = TEXT(""), EGLTFJsonMimeType MimeType = EGLTFJsonMimeType::PNG, int32 Quality = 100);
 
-	FGLTFJsonImageIndex AddImage(const TArray<FColor>& Pixels, int32 Width, int32 Height, EPixelFormat PixelFormat = EPixelFormat::PF_B8G8R8A8, const FString& Name = TEXT(""), EGLTFJsonMimeType MimeType = EGLTFJsonMimeType::PNG, int32 Quality = 100);
+	FGLTFJsonImageIndex AddImage(const void* RawData, int64 ByteLength, int32 Width, int32 Height, EPixelFormat PixelFormat = EPixelFormat::PF_B8G8R8A8, const FString& Name = TEXT(""), EGLTFJsonMimeType MimeType = EGLTFJsonMimeType::PNG, int32 Quality = 100);
 
 	FGLTFJsonImageIndex AddImage(const FTextureSource& Image, const FString& Name = TEXT(""), EGLTFJsonMimeType MimeType = EGLTFJsonMimeType::PNG, int32 Quality = 100);
 
diff --git a/Engine/Plugins/Enterprise/GLTFExporter/MigratedFromGit/Source/GLTFExporter/Private/Converters/GLTFMaterialUtility.cpp b/Engine/Plugins/Enterprise/GLTFExporter/MigratedFromGit/Source/GLTFExporter/Private/Converters/GLTFMaterialUtility.cpp
--- a/Engine/Plugins/Enterprise/GLTFExporter/MigratedFromGit/Source/GLTFExporter/Private/Converters/GLTFMaterialUtility.cpp
+++ b/Engine/Plugins/Enterprise/GLTFExporter/MigratedFromGit/Source/GLTFExporter/Private/Converters/GLTFMaterialUtility.cpp
@@ -159,7 +159,7 @@ FGLTFJsonTextureIndex FGLTFMaterialUtility::AddTexture(FGLTFConvertBuilder& Buil
 	FGLTFJsonTexture JsonTexture;
 	JsonTexture.Name = TextureName;
 	JsonTexture.Sampler = Builder.AddSampler(JsonSampler);
-	JsonTexture.Source = Builder.AddImage(Pixels, TextureSize.X, TextureSize.Y, PixelFormat, TextureName);
+	JsonTexture.Source = Builder.AddImage(Pixels.GetData(), Pixels.Num() * sizeof(FColor), TextureSize.X, TextureSize.Y, PixelFormat, TextureName);
 
 	return Builder.AddTexture(JsonTexture);
 }
