[GLTFExporter] fix(texture): split ConvertFilter into ConvertMinFilter and ConvertMagFilter to match gltf specs

commit 115084f1c3251e902fd695bc8a6e3f8dfa928d79
Date: Mon, 8 Jun 2020 23:30:31 +0200

    fix(texture): split ConvertFilter into ConvertMinFilter and ConvertMagFilter to match gltf specs
    (cherry picked from commit ce741f42a6efa6cb5e17020ad48e6041e09bb463​)

#skipci

diff --git a/Engine/Plugins/Enterprise/GLTFExporter/MigratedFromGit/Source/GLTFExporter/Private/Converters/GLTFConverterUtility.cpp b/Engine/Plugins/Enterprise/GLTFExporter/MigratedFromGit/Source/GLTFExporter/Private/Converters/GLTFConverterUtility.cpp
--- a/Engine/Plugins/Enterprise/GLTFExporter/MigratedFromGit/Source/GLTFExporter/Private/Converters/GLTFConverterUtility.cpp
+++ b/Engine/Plugins/Enterprise/GLTFExporter/MigratedFromGit/Source/GLTFExporter/Private/Converters/GLTFConverterUtility.cpp
@@ -4,12 +4,22 @@
 #include "Engine/TextureLODSettings.h"
 #include "Engine/Blueprint.h"
 
-ETextureSamplerFilter FGLTFConverterUtility::GetDefaultFilter(TextureGroup LODGroup)
+TextureFilter FGLTFConverterUtility::GetDefaultFilter(TextureGroup LODGroup)
 {
-	const ITargetPlatform* RunningPlatform = GetTargetPlatformManagerRef().GetRunningTargetPlatform(); // TODO: should this be the running platform?
+	// TODO: should this be the running platform?
+	const ITargetPlatform* RunningPlatform = GetTargetPlatformManagerRef().GetRunningTargetPlatform();
 	const UTextureLODSettings& TextureLODSettings = RunningPlatform->GetTextureLODSettings();
 	const FTextureLODGroup& TextureLODGroup = TextureLODSettings.GetTextureLODGroup(LODGroup);
-	return TextureLODGroup.Filter;
+
+	switch (TextureLODGroup.Filter)
+	{
+		case ETextureSamplerFilter::Point:             return TextureFilter::TF_Nearest;
+		case ETextureSamplerFilter::Bilinear:          return TextureFilter::TF_Bilinear;
+		case ETextureSamplerFilter::Trilinear:         return TextureFilter::TF_Trilinear;
+		case ETextureSamplerFilter::AnisotropicPoint:  return TextureFilter::TF_Nearest;
+		case ETextureSamplerFilter::AnisotropicLinear: return TextureFilter::TF_Trilinear;
+		default:                                       return TextureFilter::TF_Default; // fallback
+	}
 }
 
 bool FGLTFConverterUtility::IsSkySphereBlueprint(const UBlueprint* Blueprint)
diff --git a/Engine/Plugins/Enterprise/GLTFExporter/MigratedFromGit/Source/GLTFExporter/Private/Converters/GLTFConverterUtility.h b/Engine/Plugins/Enterprise/GLTFExporter/MigratedFromGit/Source/GLTFExporter/Private/Converters/GLTFConverterUtility.h
--- a/Engine/Plugins/Enterprise/GLTFExporter/MigratedFromGit/Source/GLTFExporter/Private/Converters/GLTFConverterUtility.h
+++ b/Engine/Plugins/Enterprise/GLTFExporter/MigratedFromGit/Source/GLTFExporter/Private/Converters/GLTFConverterUtility.h
@@ -102,18 +102,6 @@ struct FGLTFConverterUtility
 		}
 	}
 
-	static EGLTFJsonTextureFilter ConvertFilter(TextureFilter Filter, TextureGroup LODGroup)
-	{
-		switch (Filter)
-		{
-			case TextureFilter::TF_Nearest:   return EGLTFJsonTextureFilter::Nearest;
-			case TextureFilter::TF_Bilinear:  return EGLTFJsonTextureFilter::LinearMipmapNearest;
-			case TextureFilter::TF_Trilinear: return EGLTFJsonTextureFilter::LinearMipmapLinear;
-			case TextureFilter::TF_Default:   return ConvertFilter(GetDefaultFilter(LODGroup));
-			default:                          return EGLTFJsonTextureFilter::None;
-		}
-	}
-
 	static EGLTFJsonTextureWrap ConvertWrap(TextureAddress Address)
 	{
 		switch (Address)
@@ -125,20 +113,39 @@ struct FGLTFConverterUtility
 		}
 	}
 
-	static EGLTFJsonTextureFilter ConvertFilter(ETextureSamplerFilter Filter)
+	static EGLTFJsonTextureFilter ConvertMinFilter(TextureFilter Filter)
 	{
 		switch (Filter)
 		{
-			case ETextureSamplerFilter::Point:             return EGLTFJsonTextureFilter::Nearest;
-			case ETextureSamplerFilter::Bilinear:          return EGLTFJsonTextureFilter::LinearMipmapNearest;
-			case ETextureSamplerFilter::Trilinear:         return EGLTFJsonTextureFilter::LinearMipmapLinear;
-			case ETextureSamplerFilter::AnisotropicPoint:  return EGLTFJsonTextureFilter::Nearest; // TODO: is this correct?
-			case ETextureSamplerFilter::AnisotropicLinear: return EGLTFJsonTextureFilter::LinearMipmapLinear; // TODO: is this correct?
-			default:                                       return EGLTFJsonTextureFilter::None;
+			case TextureFilter::TF_Nearest:   return EGLTFJsonTextureFilter::NearestMipmapNearest;
+			case TextureFilter::TF_Bilinear:  return EGLTFJsonTextureFilter::LinearMipmapNearest;
+			case TextureFilter::TF_Trilinear: return EGLTFJsonTextureFilter::LinearMipmapLinear;
+			default:                          return EGLTFJsonTextureFilter::None;
 		}
 	}
 
-	static ETextureSamplerFilter GetDefaultFilter(TextureGroup Group);
+	static EGLTFJsonTextureFilter ConvertMagFilter(TextureFilter Filter)
+	{
+		switch (Filter)
+		{
+			case TextureFilter::TF_Nearest:   return EGLTFJsonTextureFilter::Nearest;
+			case TextureFilter::TF_Bilinear:  return EGLTFJsonTextureFilter::Linear;
+			case TextureFilter::TF_Trilinear: return EGLTFJsonTextureFilter::Linear;
+			default:                          return EGLTFJsonTextureFilter::None;
+		}
+	}
+
+	static EGLTFJsonTextureFilter ConvertMinFilter(TextureFilter Filter, TextureGroup LODGroup)
+	{
+		return ConvertMinFilter(Filter == TextureFilter::TF_Default ? GetDefaultFilter(LODGroup) : Filter);
+	}
+
+	static EGLTFJsonTextureFilter ConvertMagFilter(TextureFilter Filter, TextureGroup LODGroup)
+	{
+		return ConvertMagFilter(Filter == TextureFilter::TF_Default ? GetDefaultFilter(LODGroup) : Filter);
+	}
+
+	static TextureFilter GetDefaultFilter(TextureGroup Group);
 
 	static bool IsSkySphereBlueprint(const UBlueprint* Blueprint);
 
diff --git a/Engine/Plugins/Enterprise/GLTFExporter/MigratedFromGit/Source/GLTFExporter/Private/Converters/GLTFMaterialConverters.cpp b/Engine/Plugins/Enterprise/GLTFExporter/MigratedFromGit/Source/GLTFExporter/Private/Converters/GLTFMaterialConverters.cpp
--- a/Engine/Plugins/Enterprise/GLTFExporter/MigratedFromGit/Source/GLTFExporter/Private/Converters/GLTFMaterialConverters.cpp
+++ b/Engine/Plugins/Enterprise/GLTFExporter/MigratedFromGit/Source/GLTFExporter/Private/Converters/GLTFMaterialConverters.cpp
@@ -144,7 +144,8 @@ bool FGLTFMaterialConverter::TryGetBaseColorAndOpacity(FGLTFConvertBuilder& Buil
 	FIntPoint TextureSize(512, 512);
 
 	EGLTFJsonTextureWrap TextureWrap = EGLTFJsonTextureWrap::ClampToEdge;
-	EGLTFJsonTextureFilter TextureFilter = EGLTFJsonTextureFilter::LinearMipmapLinear;
+	EGLTFJsonTextureFilter TextureMinFilter = EGLTFJsonTextureFilter::LinearMipmapLinear;
+	EGLTFJsonTextureFilter TextureMagFilter = EGLTFJsonTextureFilter::Linear;
 
 	if (bHasBaseColorSourceTexture && bHasOpacitySourceTexture)
 	{
@@ -164,21 +165,24 @@ bool FGLTFMaterialConverter::TryGetBaseColorAndOpacity(FGLTFConvertBuilder& Buil
 			FMath::Max(BaseColorTexture->GetSizeY(), OpacityTexture->GetSizeY()));
 
 		TextureWrap = FGLTFConverterUtility::ConvertWrap(BaseColorTexture->AddressX);
-		TextureFilter = FGLTFConverterUtility::ConvertFilter(BaseColorTexture->Filter, BaseColorTexture->LODGroup);
+		TextureMinFilter = FGLTFConverterUtility::ConvertMinFilter(BaseColorTexture->Filter, BaseColorTexture->LODGroup);
+		TextureMagFilter = FGLTFConverterUtility::ConvertMagFilter(BaseColorTexture->Filter, BaseColorTexture->LODGroup);
 	}
 	else if (bHasBaseColorSourceTexture)
 	{
 		TexCoord = BaseColorTexCoord;
 		TextureSize = FIntPoint(BaseColorTexture->GetSizeX(), BaseColorTexture->GetSizeY());
 		TextureWrap = FGLTFConverterUtility::ConvertWrap(BaseColorTexture->AddressX);
-		TextureFilter = FGLTFConverterUtility::ConvertFilter(BaseColorTexture->Filter, BaseColorTexture->LODGroup);
+		TextureMinFilter = FGLTFConverterUtility::ConvertMinFilter(BaseColorTexture->Filter, BaseColorTexture->LODGroup);
+		TextureMagFilter = FGLTFConverterUtility::ConvertMagFilter(BaseColorTexture->Filter, BaseColorTexture->LODGroup);
 	}
 	else if (bHasOpacitySourceTexture)
 	{
 		TexCoord = OpacityTexCoord;
 		TextureSize = FIntPoint(OpacityTexture->GetSizeX(), OpacityTexture->GetSizeY());
 		TextureWrap = FGLTFConverterUtility::ConvertWrap(OpacityTexture->AddressX);
-		TextureFilter = FGLTFConverterUtility::ConvertFilter(OpacityTexture->Filter, OpacityTexture->LODGroup);
+		TextureMinFilter = FGLTFConverterUtility::ConvertMinFilter(OpacityTexture->Filter, OpacityTexture->LODGroup);
+		TextureMagFilter = FGLTFConverterUtility::ConvertMagFilter(OpacityTexture->Filter, OpacityTexture->LODGroup);
 	}
 
 	const EMaterialProperty OpacityProperty = MaterialInterface->GetBlendMode() == BLEND_Masked ? MP_OpacityMask : MP_Opacity;
@@ -209,7 +213,8 @@ bool FGLTFMaterialConverter::TryGetBaseColorAndOpacity(FGLTFConvertBuilder& Buil
 		CombineSources,
 		TextureSize,
 		TextureName,
-		TextureFilter,
+		TextureMinFilter,
+		TextureMagFilter,
 		TextureWrap);
 
 	OutPBRParams.BaseColorTexture.TexCoord = TexCoord;
@@ -265,7 +270,8 @@ bool FGLTFMaterialConverter::TryGetMetallicAndRoughness(FGLTFConvertBuilder& Bui
 	FIntPoint TextureSize(512, 512);
 
 	EGLTFJsonTextureWrap TextureWrap = EGLTFJsonTextureWrap::ClampToEdge;
-	EGLTFJsonTextureFilter TextureFilter = EGLTFJsonTextureFilter::LinearMipmapLinear;
+	EGLTFJsonTextureFilter TextureMinFilter = EGLTFJsonTextureFilter::LinearMipmapLinear;
+	EGLTFJsonTextureFilter TextureMagFilter = EGLTFJsonTextureFilter::Linear;
 
 	if (bHasMetallicSourceTexture && bHasRoughnessSourceTexture)
 	{
@@ -285,21 +291,24 @@ bool FGLTFMaterialConverter::TryGetMetallicAndRoughness(FGLTFConvertBuilder& Bui
 			FMath::Max(MetallicTexture->GetSizeY(), RoughnessTexture->GetSizeY()));
 
 		TextureWrap = FGLTFConverterUtility::ConvertWrap(MetallicTexture->AddressX);
-		TextureFilter = FGLTFConverterUtility::ConvertFilter(MetallicTexture->Filter, MetallicTexture->LODGroup);
+		TextureMinFilter = FGLTFConverterUtility::ConvertMinFilter(MetallicTexture->Filter, MetallicTexture->LODGroup);
+		TextureMagFilter = FGLTFConverterUtility::ConvertMagFilter(MetallicTexture->Filter, MetallicTexture->LODGroup);
 	}
 	else if (bHasMetallicSourceTexture)
 	{
 		TexCoord = MetallicTexCoord;
 		TextureSize = FIntPoint(MetallicTexture->GetSizeX(), MetallicTexture->GetSizeY());
 		TextureWrap = FGLTFConverterUtility::ConvertWrap(MetallicTexture->AddressX);
-		TextureFilter = FGLTFConverterUtility::ConvertFilter(MetallicTexture->Filter, MetallicTexture->LODGroup);
+		TextureMinFilter = FGLTFConverterUtility::ConvertMinFilter(MetallicTexture->Filter, MetallicTexture->LODGroup);
+		TextureMagFilter = FGLTFConverterUtility::ConvertMagFilter(MetallicTexture->Filter, MetallicTexture->LODGroup);
 	}
 	else if (bHasRoughnessSourceTexture)
 	{
 		TexCoord = RoughnessTexCoord;
 		TextureSize = FIntPoint(RoughnessTexture->GetSizeX(), RoughnessTexture->GetSizeY());
 		TextureWrap = FGLTFConverterUtility::ConvertWrap(RoughnessTexture->AddressX);
-		TextureFilter = FGLTFConverterUtility::ConvertFilter(RoughnessTexture->Filter, RoughnessTexture->LODGroup);
+		TextureMinFilter = FGLTFConverterUtility::ConvertMinFilter(RoughnessTexture->Filter, RoughnessTexture->LODGroup);
+		TextureMagFilter = FGLTFConverterUtility::ConvertMagFilter(RoughnessTexture->Filter, RoughnessTexture->LODGroup);
 	}
 
 	RoughnessTexture = FGLTFMaterialUtility::BakeMaterialPropertyToTexture(TextureSize, MP_Roughness, MaterialInterface);
@@ -328,7 +337,8 @@ bool FGLTFMaterialConverter::TryGetMetallicAndRoughness(FGLTFConvertBuilder& Bui
 		CombineSources,
 		TextureSize,
 		TextureName,
-		TextureFilter,
+		TextureMinFilter,
+		TextureMagFilter,
 		TextureWrap);
 
 	OutPBRParams.MetallicRoughnessTexture.TexCoord = TexCoord;
@@ -562,7 +572,8 @@ bool FGLTFMaterialConverter::TryGetBakedTexture(FGLTFConvertBuilder& Builder, FG
 
 	// TODO: support for other wrapping / filters?
 	const EGLTFJsonTextureWrap TextureWrap = EGLTFJsonTextureWrap::ClampToEdge;
-	const EGLTFJsonTextureFilter TextureFilter = EGLTFJsonTextureFilter::LinearMipmapLinear;
+	const EGLTFJsonTextureFilter TextureMinFilter = EGLTFJsonTextureFilter::LinearMipmapLinear;
+	const EGLTFJsonTextureFilter TextureMagFilter = EGLTFJsonTextureFilter::Linear;
 
 	const FGLTFJsonTextureIndex TextureIndex = FGLTFMaterialUtility::AddTexture(
 		Builder,
@@ -570,7 +581,8 @@ bool FGLTFMaterialConverter::TryGetBakedTexture(FGLTFConvertBuilder& Builder, FG
 		PropertyBakeOutput.Size,
 		TextureName,
 		PropertyBakeOutput.PixelFormat,
-		TextureFilter,
+		TextureMinFilter,
+		TextureMagFilter,
 		TextureWrap);
 
 	OutTexInfo.TexCoord = TexCoord;
diff --git a/Engine/Plugins/Enterprise/GLTFExporter/MigratedFromGit/Source/GLTFExporter/Private/Converters/GLTFMaterialUtility.cpp b/Engine/Plugins/Enterprise/GLTFExporter/MigratedFromGit/Source/GLTFExporter/Private/Converters/GLTFMaterialUtility.cpp
--- a/Engine/Plugins/Enterprise/GLTFExporter/MigratedFromGit/Source/GLTFExporter/Private/Converters/GLTFMaterialUtility.cpp
+++ b/Engine/Plugins/Enterprise/GLTFExporter/MigratedFromGit/Source/GLTFExporter/Private/Converters/GLTFMaterialUtility.cpp
@@ -133,7 +133,7 @@ UTexture2D* FGLTFMaterialUtility::BakeMaterialPropertyToTexture(const FIntPoint&
 	return CreateTransientTexture(PropertyBakeOutput.Pixels, PropertyBakeOutput.Size, PropertyBakeOutput.PixelFormat, bUseSRGB);
 }
 
-FGLTFJsonTextureIndex FGLTFMaterialUtility::AddCombinedTexture(FGLTFConvertBuilder& Builder, const TArray<FGLTFTextureCombineSource>& CombineSources, const FIntPoint& TextureSize, const FString& TextureName, EGLTFJsonTextureFilter Filter, EGLTFJsonTextureWrap Wrap)
+FGLTFJsonTextureIndex FGLTFMaterialUtility::AddCombinedTexture(FGLTFConvertBuilder& Builder, const TArray<FGLTFTextureCombineSource>& CombineSources, const FIntPoint& TextureSize, const FString& TextureName, EGLTFJsonTextureFilter MinFilter, EGLTFJsonTextureFilter MagFilter, EGLTFJsonTextureWrap Wrap)
 {
 	check(CombineSources.Num() > 0);
 
@@ -145,15 +145,16 @@ FGLTFJsonTextureIndex FGLTFMaterialUtility::AddCombinedTexture(FGLTFConvertBuild
 		return FGLTFJsonTextureIndex(INDEX_NONE);
 	}
 
-	return AddTexture(Builder, Pixels, TextureSize, TextureName, PixelFormat, Filter, Wrap);
+	return AddTexture(Builder, Pixels, TextureSize, TextureName, PixelFormat, MinFilter, MagFilter, Wrap);
 }
 
-FGLTFJsonTextureIndex FGLTFMaterialUtility::AddTexture(FGLTFConvertBuilder& Builder, const TArray<FColor>& Pixels, const FIntPoint& TextureSize, const FString& TextureName, EPixelFormat PixelFormat, EGLTFJsonTextureFilter Filter, EGLTFJsonTextureWrap Wrap)
+FGLTFJsonTextureIndex FGLTFMaterialUtility::AddTexture(FGLTFConvertBuilder& Builder, const TArray<FColor>& Pixels, const FIntPoint& TextureSize, const FString& TextureName, EPixelFormat PixelFormat, EGLTFJsonTextureFilter MinFilter, EGLTFJsonTextureFilter MagFilter, EGLTFJsonTextureWrap Wrap)
 {
 	// TODO: maybe we should reuse existing samplers?
 	FGLTFJsonSampler JsonSampler;
 	JsonSampler.Name = TextureName;
-	JsonSampler.MagFilter = JsonSampler.MinFilter = Filter;
+	JsonSampler.MinFilter = MinFilter;
+	JsonSampler.MagFilter = MagFilter;
 	JsonSampler.WrapS = JsonSampler.WrapT = Wrap;
 
 	FGLTFJsonTexture JsonTexture;
diff --git a/Engine/Plugins/Enterprise/GLTFExporter/MigratedFromGit/Source/GLTFExporter/Private/Converters/GLTFMaterialUtility.h b/Engine/Plugins/Enterprise/GLTFExporter/MigratedFromGit/Source/GLTFExporter/Private/Converters/GLTFMaterialUtility.h
--- a/Engine/Plugins/Enterprise/GLTFExporter/MigratedFromGit/Source/GLTFExporter/Private/Converters/GLTFMaterialUtility.h
+++ b/Engine/Plugins/Enterprise/GLTFExporter/MigratedFromGit/Source/GLTFExporter/Private/Converters/GLTFMaterialUtility.h
@@ -36,8 +36,8 @@ struct FGLTFMaterialUtility
 	static FGLTFPropertyBakeOutput BakeMaterialProperty(const FIntPoint& OutputSize, EMaterialProperty MaterialProperty, const UMaterialInterface* Material, bool bCopyAlphaFromRedChannel = false);
 	static UTexture2D* BakeMaterialPropertyToTexture(const FIntPoint& OutputSize, EMaterialProperty MaterialProperty, const UMaterialInterface* Material, bool bCopyAlphaFromRedChannel = false, bool bUseSRGB = false);
 
-	static FGLTFJsonTextureIndex AddCombinedTexture(FGLTFConvertBuilder& Builder, const TArray<FGLTFTextureCombineSource>& CombineSources, const FIntPoint& TextureSize, const FString& TextureName, EGLTFJsonTextureFilter Filter, EGLTFJsonTextureWrap Wrap);
-	static FGLTFJsonTextureIndex AddTexture(FGLTFConvertBuilder& Builder, const TArray<FColor>& Pixels, const FIntPoint& TextureSize, const FString& TextureName, EPixelFormat PixelFormat, EGLTFJsonTextureFilter Filter, EGLTFJsonTextureWrap Wrap);
+	static FGLTFJsonTextureIndex AddCombinedTexture(FGLTFConvertBuilder& Builder, const TArray<FGLTFTextureCombineSource>& CombineSources, const FIntPoint& TextureSize, const FString& TextureName, EGLTFJsonTextureFilter MinFilter, EGLTFJsonTextureFilter MagFilter, EGLTFJsonTextureWrap Wrap);
+	static FGLTFJsonTextureIndex AddTexture(FGLTFConvertBuilder& Builder, const TArray<FColor>& Pixels, const FIntPoint& TextureSize, const FString& TextureName, EPixelFormat PixelFormat, EGLTFJsonTextureFilter MinFilter, EGLTFJsonTextureFilter MagFilter, EGLTFJsonTextureWrap Wrap);
 
 	static FLinearColor GetMask(const FExpressionInput& ExpressionInput);
 	static uint32 GetMaskComponentCount(const FExpressionInput& ExpressionInput);
diff --git a/Engine/Plugins/Enterprise/GLTFExporter/MigratedFromGit/Source/GLTFExporter/Private/Converters/GLTFTextureConverters.cpp b/Engine/Plugins/Enterprise/GLTFExporter/MigratedFromGit/Source/GLTFExporter/Private/Converters/GLTFTextureConverters.cpp
--- a/Engine/Plugins/Enterprise/GLTFExporter/MigratedFromGit/Source/GLTFExporter/Private/Converters/GLTFTextureConverters.cpp
+++ b/Engine/Plugins/Enterprise/GLTFExporter/MigratedFromGit/Source/GLTFExporter/Private/Converters/GLTFTextureConverters.cpp
@@ -9,7 +9,8 @@ FGLTFJsonTextureIndex FGLTFTexture2DConverter::Add(FGLTFConvertBuilder& Builder,
 	// TODO: maybe we should reuse existing samplers?
 	FGLTFJsonSampler JsonSampler;
 	JsonSampler.Name = Name;
-	JsonSampler.MagFilter = JsonSampler.MinFilter = FGLTFConverterUtility::ConvertFilter(Texture2D->Filter, Texture2D->LODGroup);
+	JsonSampler.MinFilter = FGLTFConverterUtility::ConvertMinFilter(Texture2D->Filter, Texture2D->LODGroup);
+	JsonSampler.MagFilter = FGLTFConverterUtility::ConvertMagFilter(Texture2D->Filter, Texture2D->LODGroup);
 	JsonSampler.WrapS = FGLTFConverterUtility::ConvertWrap(Texture2D->AddressX);
 	JsonSampler.WrapT = FGLTFConverterUtility::ConvertWrap(Texture2D->AddressY);
 
