[GLTFExporter] refactor(material): use separate texture-wrap settings for X & Y (S & T)

commit cc1ed41c5304a23ef9abd0fc90fbdb89e4035cb6
Date: Tue, 9 Jun 2020 08:58:53 +0200

    refactor(material): use separate texture-wrap settings for X & Y (S & T)
    (cherry picked from commit 344000619591c572adab90f1a4d530c4697330d4​)

#skipci

diff --git a/Engine/Plugins/Enterprise/GLTFExporter/MigratedFromGit/Source/GLTFExporter/Private/Converters/GLTFMaterialConverters.cpp b/Engine/Plugins/Enterprise/GLTFExporter/MigratedFromGit/Source/GLTFExporter/Private/Converters/GLTFMaterialConverters.cpp
--- a/Engine/Plugins/Enterprise/GLTFExporter/MigratedFromGit/Source/GLTFExporter/Private/Converters/GLTFMaterialConverters.cpp
+++ b/Engine/Plugins/Enterprise/GLTFExporter/MigratedFromGit/Source/GLTFExporter/Private/Converters/GLTFMaterialConverters.cpp
@@ -143,7 +143,8 @@ bool FGLTFMaterialConverter::TryGetBaseColorAndOpacity(FGLTFConvertBuilder& Buil
 	int32 TexCoord = 0;
 	FIntPoint TextureSize(512, 512);
 
-	EGLTFJsonTextureWrap TextureWrap = EGLTFJsonTextureWrap::ClampToEdge;
+	EGLTFJsonTextureWrap TextureWrapS = EGLTFJsonTextureWrap::ClampToEdge;
+	EGLTFJsonTextureWrap TextureWrapT = EGLTFJsonTextureWrap::ClampToEdge;
 	EGLTFJsonTextureFilter TextureMinFilter = EGLTFJsonTextureFilter::LinearMipmapLinear;
 	EGLTFJsonTextureFilter TextureMagFilter = EGLTFJsonTextureFilter::Linear;
 
@@ -164,7 +165,8 @@ bool FGLTFMaterialConverter::TryGetBaseColorAndOpacity(FGLTFConvertBuilder& Buil
 			FMath::Max(BaseColorTexture->GetSizeX(), OpacityTexture->GetSizeX()),
 			FMath::Max(BaseColorTexture->GetSizeY(), OpacityTexture->GetSizeY()));
 
-		TextureWrap = FGLTFConverterUtility::ConvertWrap(BaseColorTexture->AddressX);
+		TextureWrapS = FGLTFConverterUtility::ConvertWrap(BaseColorTexture->AddressX);
+		TextureWrapT = FGLTFConverterUtility::ConvertWrap(BaseColorTexture->AddressY);
 		TextureMinFilter = FGLTFConverterUtility::ConvertMinFilter(BaseColorTexture->Filter, BaseColorTexture->LODGroup);
 		TextureMagFilter = FGLTFConverterUtility::ConvertMagFilter(BaseColorTexture->Filter, BaseColorTexture->LODGroup);
 	}
@@ -172,7 +174,8 @@ bool FGLTFMaterialConverter::TryGetBaseColorAndOpacity(FGLTFConvertBuilder& Buil
 	{
 		TexCoord = BaseColorTexCoord;
 		TextureSize = FIntPoint(BaseColorTexture->GetSizeX(), BaseColorTexture->GetSizeY());
-		TextureWrap = FGLTFConverterUtility::ConvertWrap(BaseColorTexture->AddressX);
+		TextureWrapS = FGLTFConverterUtility::ConvertWrap(BaseColorTexture->AddressX);
+		TextureWrapT = FGLTFConverterUtility::ConvertWrap(BaseColorTexture->AddressY);
 		TextureMinFilter = FGLTFConverterUtility::ConvertMinFilter(BaseColorTexture->Filter, BaseColorTexture->LODGroup);
 		TextureMagFilter = FGLTFConverterUtility::ConvertMagFilter(BaseColorTexture->Filter, BaseColorTexture->LODGroup);
 	}
@@ -180,7 +183,8 @@ bool FGLTFMaterialConverter::TryGetBaseColorAndOpacity(FGLTFConvertBuilder& Buil
 	{
 		TexCoord = OpacityTexCoord;
 		TextureSize = FIntPoint(OpacityTexture->GetSizeX(), OpacityTexture->GetSizeY());
-		TextureWrap = FGLTFConverterUtility::ConvertWrap(OpacityTexture->AddressX);
+		TextureWrapS = FGLTFConverterUtility::ConvertWrap(OpacityTexture->AddressX);
+		TextureWrapT = FGLTFConverterUtility::ConvertWrap(OpacityTexture->AddressY);
 		TextureMinFilter = FGLTFConverterUtility::ConvertMinFilter(OpacityTexture->Filter, OpacityTexture->LODGroup);
 		TextureMagFilter = FGLTFConverterUtility::ConvertMagFilter(OpacityTexture->Filter, OpacityTexture->LODGroup);
 	}
@@ -215,7 +219,8 @@ bool FGLTFMaterialConverter::TryGetBaseColorAndOpacity(FGLTFConvertBuilder& Buil
 		TextureName,
 		TextureMinFilter,
 		TextureMagFilter,
-		TextureWrap);
+		TextureWrapS,
+		TextureWrapT);
 
 	OutPBRParams.BaseColorTexture.TexCoord = TexCoord;
 	OutPBRParams.BaseColorTexture.Index = TextureIndex;
@@ -269,7 +274,8 @@ bool FGLTFMaterialConverter::TryGetMetallicAndRoughness(FGLTFConvertBuilder& Bui
 	int32 TexCoord = 0;
 	FIntPoint TextureSize(512, 512);
 
-	EGLTFJsonTextureWrap TextureWrap = EGLTFJsonTextureWrap::ClampToEdge;
+	EGLTFJsonTextureWrap TextureWrapS = EGLTFJsonTextureWrap::ClampToEdge;
+	EGLTFJsonTextureWrap TextureWrapT = EGLTFJsonTextureWrap::ClampToEdge;
 	EGLTFJsonTextureFilter TextureMinFilter = EGLTFJsonTextureFilter::LinearMipmapLinear;
 	EGLTFJsonTextureFilter TextureMagFilter = EGLTFJsonTextureFilter::Linear;
 
@@ -290,7 +296,8 @@ bool FGLTFMaterialConverter::TryGetMetallicAndRoughness(FGLTFConvertBuilder& Bui
 			FMath::Max(MetallicTexture->GetSizeX(), RoughnessTexture->GetSizeX()),
 			FMath::Max(MetallicTexture->GetSizeY(), RoughnessTexture->GetSizeY()));
 
-		TextureWrap = FGLTFConverterUtility::ConvertWrap(MetallicTexture->AddressX);
+		TextureWrapS = FGLTFConverterUtility::ConvertWrap(MetallicTexture->AddressX);
+		TextureWrapT = FGLTFConverterUtility::ConvertWrap(MetallicTexture->AddressY);
 		TextureMinFilter = FGLTFConverterUtility::ConvertMinFilter(MetallicTexture->Filter, MetallicTexture->LODGroup);
 		TextureMagFilter = FGLTFConverterUtility::ConvertMagFilter(MetallicTexture->Filter, MetallicTexture->LODGroup);
 	}
@@ -298,7 +305,8 @@ bool FGLTFMaterialConverter::TryGetMetallicAndRoughness(FGLTFConvertBuilder& Bui
 	{
 		TexCoord = MetallicTexCoord;
 		TextureSize = FIntPoint(MetallicTexture->GetSizeX(), MetallicTexture->GetSizeY());
-		TextureWrap = FGLTFConverterUtility::ConvertWrap(MetallicTexture->AddressX);
+		TextureWrapS = FGLTFConverterUtility::ConvertWrap(MetallicTexture->AddressX);
+		TextureWrapT = FGLTFConverterUtility::ConvertWrap(MetallicTexture->AddressY);
 		TextureMinFilter = FGLTFConverterUtility::ConvertMinFilter(MetallicTexture->Filter, MetallicTexture->LODGroup);
 		TextureMagFilter = FGLTFConverterUtility::ConvertMagFilter(MetallicTexture->Filter, MetallicTexture->LODGroup);
 	}
@@ -306,7 +314,8 @@ bool FGLTFMaterialConverter::TryGetMetallicAndRoughness(FGLTFConvertBuilder& Bui
 	{
 		TexCoord = RoughnessTexCoord;
 		TextureSize = FIntPoint(RoughnessTexture->GetSizeX(), RoughnessTexture->GetSizeY());
-		TextureWrap = FGLTFConverterUtility::ConvertWrap(RoughnessTexture->AddressX);
+		TextureWrapS = FGLTFConverterUtility::ConvertWrap(RoughnessTexture->AddressX);
+		TextureWrapT = FGLTFConverterUtility::ConvertWrap(RoughnessTexture->AddressY);
 		TextureMinFilter = FGLTFConverterUtility::ConvertMinFilter(RoughnessTexture->Filter, RoughnessTexture->LODGroup);
 		TextureMagFilter = FGLTFConverterUtility::ConvertMagFilter(RoughnessTexture->Filter, RoughnessTexture->LODGroup);
 	}
@@ -339,7 +348,8 @@ bool FGLTFMaterialConverter::TryGetMetallicAndRoughness(FGLTFConvertBuilder& Bui
 		TextureName,
 		TextureMinFilter,
 		TextureMagFilter,
-		TextureWrap);
+		TextureWrapS,
+		TextureWrapT);
 
 	OutPBRParams.MetallicRoughnessTexture.TexCoord = TexCoord;
 	OutPBRParams.MetallicRoughnessTexture.Index = TextureIndex;
@@ -571,7 +581,8 @@ bool FGLTFMaterialConverter::TryGetBakedTexture(FGLTFConvertBuilder& Builder, FG
 	const FString TextureName = FString::Printf(TEXT("%s_%s"), *MaterialInterface->GetName(), PropertyName);
 
 	// TODO: support for other wrapping / filters?
-	const EGLTFJsonTextureWrap TextureWrap = EGLTFJsonTextureWrap::ClampToEdge;
+	const EGLTFJsonTextureWrap TextureWrapS = EGLTFJsonTextureWrap::ClampToEdge;
+	const EGLTFJsonTextureWrap TextureWrapT = EGLTFJsonTextureWrap::ClampToEdge;
 	const EGLTFJsonTextureFilter TextureMinFilter = EGLTFJsonTextureFilter::LinearMipmapLinear;
 	const EGLTFJsonTextureFilter TextureMagFilter = EGLTFJsonTextureFilter::Linear;
 
@@ -583,7 +594,8 @@ bool FGLTFMaterialConverter::TryGetBakedTexture(FGLTFConvertBuilder& Builder, FG
 		PropertyBakeOutput.PixelFormat,
 		TextureMinFilter,
 		TextureMagFilter,
-		TextureWrap);
+		TextureWrapS,
+		TextureWrapT);
 
 	OutTexInfo.TexCoord = TexCoord;
 	OutTexInfo.Index = TextureIndex;
diff --git a/Engine/Plugins/Enterprise/GLTFExporter/MigratedFromGit/Source/GLTFExporter/Private/Converters/GLTFMaterialUtility.cpp b/Engine/Plugins/Enterprise/GLTFExporter/MigratedFromGit/Source/GLTFExporter/Private/Converters/GLTFMaterialUtility.cpp
--- a/Engine/Plugins/Enterprise/GLTFExporter/MigratedFromGit/Source/GLTFExporter/Private/Converters/GLTFMaterialUtility.cpp
+++ b/Engine/Plugins/Enterprise/GLTFExporter/MigratedFromGit/Source/GLTFExporter/Private/Converters/GLTFMaterialUtility.cpp
@@ -133,7 +133,7 @@ UTexture2D* FGLTFMaterialUtility::BakeMaterialPropertyToTexture(const FIntPoint&
 	return CreateTransientTexture(PropertyBakeOutput.Pixels, PropertyBakeOutput.Size, PropertyBakeOutput.PixelFormat, bUseSRGB);
 }
 
-FGLTFJsonTextureIndex FGLTFMaterialUtility::AddCombinedTexture(FGLTFConvertBuilder& Builder, const TArray<FGLTFTextureCombineSource>& CombineSources, const FIntPoint& TextureSize, const FString& TextureName, EGLTFJsonTextureFilter MinFilter, EGLTFJsonTextureFilter MagFilter, EGLTFJsonTextureWrap Wrap)
+FGLTFJsonTextureIndex FGLTFMaterialUtility::AddCombinedTexture(FGLTFConvertBuilder& Builder, const TArray<FGLTFTextureCombineSource>& CombineSources, const FIntPoint& TextureSize, const FString& TextureName, EGLTFJsonTextureFilter MinFilter, EGLTFJsonTextureFilter MagFilter, EGLTFJsonTextureWrap WrapS, EGLTFJsonTextureWrap WrapT)
 {
 	check(CombineSources.Num() > 0);
 
@@ -145,17 +145,18 @@ FGLTFJsonTextureIndex FGLTFMaterialUtility::AddCombinedTexture(FGLTFConvertBuild
 		return FGLTFJsonTextureIndex(INDEX_NONE);
 	}
 
-	return AddTexture(Builder, Pixels, TextureSize, TextureName, PixelFormat, MinFilter, MagFilter, Wrap);
+	return AddTexture(Builder, Pixels, TextureSize, TextureName, PixelFormat, MinFilter, MagFilter, WrapS, WrapT);
 }
 
-FGLTFJsonTextureIndex FGLTFMaterialUtility::AddTexture(FGLTFConvertBuilder& Builder, const TArray<FColor>& Pixels, const FIntPoint& TextureSize, const FString& TextureName, EPixelFormat PixelFormat, EGLTFJsonTextureFilter MinFilter, EGLTFJsonTextureFilter MagFilter, EGLTFJsonTextureWrap Wrap)
+FGLTFJsonTextureIndex FGLTFMaterialUtility::AddTexture(FGLTFConvertBuilder& Builder, const TArray<FColor>& Pixels, const FIntPoint& TextureSize, const FString& TextureName, EPixelFormat PixelFormat, EGLTFJsonTextureFilter MinFilter, EGLTFJsonTextureFilter MagFilter, EGLTFJsonTextureWrap WrapS, EGLTFJsonTextureWrap WrapT)
 {
 	// TODO: maybe we should reuse existing samplers?
 	FGLTFJsonSampler JsonSampler;
 	JsonSampler.Name = TextureName;
 	JsonSampler.MinFilter = MinFilter;
 	JsonSampler.MagFilter = MagFilter;
-	JsonSampler.WrapS = JsonSampler.WrapT = Wrap;
+	JsonSampler.WrapS = WrapS;
+	JsonSampler.WrapT = WrapT;
 
 	FGLTFJsonTexture JsonTexture;
 	JsonTexture.Name = TextureName;
diff --git a/Engine/Plugins/Enterprise/GLTFExporter/MigratedFromGit/Source/GLTFExporter/Private/Converters/GLTFMaterialUtility.h b/Engine/Plugins/Enterprise/GLTFExporter/MigratedFromGit/Source/GLTFExporter/Private/Converters/GLTFMaterialUtility.h
--- a/Engine/Plugins/Enterprise/GLTFExporter/MigratedFromGit/Source/GLTFExporter/Private/Converters/GLTFMaterialUtility.h
+++ b/Engine/Plugins/Enterprise/GLTFExporter/MigratedFromGit/Source/GLTFExporter/Private/Converters/GLTFMaterialUtility.h
@@ -36,8 +36,8 @@ struct FGLTFMaterialUtility
 	static FGLTFPropertyBakeOutput BakeMaterialProperty(const FIntPoint& OutputSize, EMaterialProperty MaterialProperty, const UMaterialInterface* Material, bool bCopyAlphaFromRedChannel = false);
 	static UTexture2D* BakeMaterialPropertyToTexture(const FIntPoint& OutputSize, EMaterialProperty MaterialProperty, const UMaterialInterface* Material, bool bCopyAlphaFromRedChannel = false, bool bUseSRGB = false);
 
-	static FGLTFJsonTextureIndex AddCombinedTexture(FGLTFConvertBuilder& Builder, const TArray<FGLTFTextureCombineSource>& CombineSources, const FIntPoint& TextureSize, const FString& TextureName, EGLTFJsonTextureFilter MinFilter, EGLTFJsonTextureFilter MagFilter, EGLTFJsonTextureWrap Wrap);
-	static FGLTFJsonTextureIndex AddTexture(FGLTFConvertBuilder& Builder, const TArray<FColor>& Pixels, const FIntPoint& TextureSize, const FString& TextureName, EPixelFormat PixelFormat, EGLTFJsonTextureFilter MinFilter, EGLTFJsonTextureFilter MagFilter, EGLTFJsonTextureWrap Wrap);
+	static FGLTFJsonTextureIndex AddCombinedTexture(FGLTFConvertBuilder& Builder, const TArray<FGLTFTextureCombineSource>& CombineSources, const FIntPoint& TextureSize, const FString& TextureName, EGLTFJsonTextureFilter MinFilter, EGLTFJsonTextureFilter MagFilter, EGLTFJsonTextureWrap WrapS, EGLTFJsonTextureWrap WrapT);
+	static FGLTFJsonTextureIndex AddTexture(FGLTFConvertBuilder& Builder, const TArray<FColor>& Pixels, const FIntPoint& TextureSize, const FString& TextureName, EPixelFormat PixelFormat, EGLTFJsonTextureFilter MinFilter, EGLTFJsonTextureFilter MagFilter, EGLTFJsonTextureWrap WrapS, EGLTFJsonTextureWrap WrapT);
 
 	static FLinearColor GetMask(const FExpressionInput& ExpressionInput);
 	static uint32 GetMaskComponentCount(const FExpressionInput& ExpressionInput);
