[GLTFExporter] refactor(material): use stack-based variables for material- and mesh-settings when baking based on feedback

commit 8b57f3f9570a2b2afc30cb76c0aa9a88cb39e280
Date: Tue, 9 Jun 2020 09:36:44 +0200

    refactor(material): use stack-based variables for material- and mesh-settings when baking based on feedback
    (cherry picked from commit c6b6389a60e750cd79b835540bf2cc25752c596b​)

#skipci

diff --git a/Engine/Plugins/Enterprise/GLTFExporter/MigratedFromGit/Source/GLTFExporter/Private/Converters/GLTFMaterialUtility.cpp b/Engine/Plugins/Enterprise/GLTFExporter/MigratedFromGit/Source/GLTFExporter/Private/Converters/GLTFMaterialUtility.cpp
--- a/Engine/Plugins/Enterprise/GLTFExporter/MigratedFromGit/Source/GLTFExporter/Private/Converters/GLTFMaterialUtility.cpp
+++ b/Engine/Plugins/Enterprise/GLTFExporter/MigratedFromGit/Source/GLTFExporter/Private/Converters/GLTFMaterialUtility.cpp
@@ -76,21 +76,20 @@ FGLTFPropertyBakeOutput FGLTFMaterialUtility::BakeMaterialProperty(const FIntPoi
 {
 	TArray<FMeshData*> MeshSettings;
 
-	FMeshData* MeshSet = new FMeshData();
-
-	MeshSet->TextureCoordinateBox = FBox2D(FVector2D(0, 0), FVector2D(1, 1));
+	FMeshData MeshSet;
+	MeshSet.TextureCoordinateBox = FBox2D(FVector2D(0, 0), FVector2D(1, 1));
 
 	// TODO: Do we need to fill in any more info in MeshSet?
 
-	MeshSettings.Add(MeshSet);
+	MeshSettings.Add(&MeshSet);
 
 	TArray<FMaterialData*> MatSettings;
 
-	FMaterialData* MatSet = new FMaterialData();
-	MatSet->Material = const_cast<UMaterialInterface*>(Material);
-	MatSet->PropertySizes.Add(MaterialProperty, OutputSize);
+	FMaterialData MatSet;
+	MatSet.Material = const_cast<UMaterialInterface*>(Material);
+	MatSet.PropertySizes.Add(MaterialProperty, OutputSize);
 
-	MatSettings.Add(MatSet);
+	MatSettings.Add(&MatSet);
 
 	TArray<FBakeOutput> BakeOutputs;
 	IMaterialBakingModule& Module = FModuleManager::Get().LoadModuleChecked<IMaterialBakingModule>("MaterialBaking");
