[GLTFExporter] chore(material): add TODOs for handling discrepencies in filter-settings when merging textures

commit 5554182fd07f1ca7ec6c338d115bf17ff46847cd
Date: Tue, 9 Jun 2020 09:46:50 +0200

    chore(material): add TODOs for handling discrepencies in filter-settings when merging textures
    (cherry picked from commit 633532b6a6c52e3fdc0d98356202dfec26420d87​)

#skipci

diff --git a/Engine/Plugins/Enterprise/GLTFExporter/MigratedFromGit/Source/GLTFExporter/Private/Converters/GLTFMaterialConverters.cpp b/Engine/Plugins/Enterprise/GLTFExporter/MigratedFromGit/Source/GLTFExporter/Private/Converters/GLTFMaterialConverters.cpp
--- a/Engine/Plugins/Enterprise/GLTFExporter/MigratedFromGit/Source/GLTFExporter/Private/Converters/GLTFMaterialConverters.cpp
+++ b/Engine/Plugins/Enterprise/GLTFExporter/MigratedFromGit/Source/GLTFExporter/Private/Converters/GLTFMaterialConverters.cpp
@@ -167,6 +167,9 @@ bool FGLTFMaterialConverter::TryGetBaseColorAndOpacity(FGLTFConvertBuilder& Buil
 
 		TextureWrapS = FGLTFConverterUtility::ConvertWrap(BaseColorTexture->AddressX);
 		TextureWrapT = FGLTFConverterUtility::ConvertWrap(BaseColorTexture->AddressY);
+
+		// TODO: compare min- and mag-filter for BaseColorTexture and OpacityTexture. If they differ,
+		// we should choose one or the other and inform the user about the choice made by logging to the console.
 		TextureMinFilter = FGLTFConverterUtility::ConvertMinFilter(BaseColorTexture->Filter, BaseColorTexture->LODGroup);
 		TextureMagFilter = FGLTFConverterUtility::ConvertMagFilter(BaseColorTexture->Filter, BaseColorTexture->LODGroup);
 	}
@@ -185,7 +188,7 @@ bool FGLTFMaterialConverter::TryGetBaseColorAndOpacity(FGLTFConvertBuilder& Buil
 		TextureSize = FIntPoint(OpacityTexture->GetSizeX(), OpacityTexture->GetSizeY());
 		TextureWrapS = FGLTFConverterUtility::ConvertWrap(OpacityTexture->AddressX);
 		TextureWrapT = FGLTFConverterUtility::ConvertWrap(OpacityTexture->AddressY);
-		TextureMinFilter = FGLTFConverterUtility::ConvertMinFilter(OpacityTexture->Filter, OpacityTexture->LODGroup);
+		TextureMinFilter = FGLTFConverterUtility::ConvertMinFilter(BaseColorTexture->Filter, OpacityTexture->LODGroup);
 		TextureMagFilter = FGLTFConverterUtility::ConvertMagFilter(OpacityTexture->Filter, OpacityTexture->LODGroup);
 	}
 
@@ -298,6 +301,9 @@ bool FGLTFMaterialConverter::TryGetMetallicAndRoughness(FGLTFConvertBuilder& Bui
 
 		TextureWrapS = FGLTFConverterUtility::ConvertWrap(MetallicTexture->AddressX);
 		TextureWrapT = FGLTFConverterUtility::ConvertWrap(MetallicTexture->AddressY);
+
+		// TODO: compare min- and mag-filter for BaseColorTexture and OpacityTexture. If they differ,
+		// we should choose one or the other and inform the user about the choice made by logging to the console.
 		TextureMinFilter = FGLTFConverterUtility::ConvertMinFilter(MetallicTexture->Filter, MetallicTexture->LODGroup);
 		TextureMagFilter = FGLTFConverterUtility::ConvertMagFilter(MetallicTexture->Filter, MetallicTexture->LODGroup);
 	}
