[GLTFExporter] fix(material): remove unintentionally introduced typo when calculating min-filter

commit 45ed0b08ec7a87eae05f9181d4824ff68384095f
Date: Tue, 9 Jun 2020 10:25:41 +0200

    fix(material): remove unintentionally introduced typo when calculating min-filter
    (cherry picked from commit 829dd82cc54664e7a94c2d53c76d25263d20ee64​)

#skipci

diff --git a/Engine/Plugins/Enterprise/GLTFExporter/MigratedFromGit/Source/GLTFExporter/Private/Converters/GLTFMaterialConverters.cpp b/Engine/Plugins/Enterprise/GLTFExporter/MigratedFromGit/Source/GLTFExporter/Private/Converters/GLTFMaterialConverters.cpp
--- a/Engine/Plugins/Enterprise/GLTFExporter/MigratedFromGit/Source/GLTFExporter/Private/Converters/GLTFMaterialConverters.cpp
+++ b/Engine/Plugins/Enterprise/GLTFExporter/MigratedFromGit/Source/GLTFExporter/Private/Converters/GLTFMaterialConverters.cpp
@@ -188,7 +188,7 @@ bool FGLTFMaterialConverter::TryGetBaseColorAndOpacity(FGLTFConvertBuilder& Buil
 		TextureSize = FIntPoint(OpacityTexture->GetSizeX(), OpacityTexture->GetSizeY());
 		TextureWrapS = FGLTFConverterUtility::ConvertWrap(OpacityTexture->AddressX);
 		TextureWrapT = FGLTFConverterUtility::ConvertWrap(OpacityTexture->AddressY);
-		TextureMinFilter = FGLTFConverterUtility::ConvertMinFilter(BaseColorTexture->Filter, OpacityTexture->LODGroup);
+		TextureMinFilter = FGLTFConverterUtility::ConvertMinFilter(OpacityTexture->Filter, OpacityTexture->LODGroup);
 		TextureMagFilter = FGLTFConverterUtility::ConvertMagFilter(OpacityTexture->Filter, OpacityTexture->LODGroup);
 	}
 
