[GLTFExporter] refactor(material): replace usage of FString::Printf with standard concatenation

commit be664df5323734cbf7c3c88cb37f84caab2b9fdc
Date: Tue, 9 Jun 2020 13:20:43 +0200

    refactor(material): replace usage of FString::Printf with standard concatenation
    (cherry picked from commit 3f87a678d9813ecf1a0206fa62f351fa450f7d5f​)

#skipci

diff --git a/Engine/Plugins/Enterprise/GLTFExporter/MigratedFromGit/Source/GLTFExporter/Private/Converters/GLTFMaterialConverters.cpp b/Engine/Plugins/Enterprise/GLTFExporter/MigratedFromGit/Source/GLTFExporter/Private/Converters/GLTFMaterialConverters.cpp
--- a/Engine/Plugins/Enterprise/GLTFExporter/MigratedFromGit/Source/GLTFExporter/Private/Converters/GLTFMaterialConverters.cpp
+++ b/Engine/Plugins/Enterprise/GLTFExporter/MigratedFromGit/Source/GLTFExporter/Private/Converters/GLTFMaterialConverters.cpp
@@ -207,7 +207,7 @@ bool FGLTFMaterialConverter::TryGetBaseColorAndOpacity(FGLTFConvertBuilder& Buil
 	// TODO: handle the case where TextureSize is 1x1. In this case, both properties are constants and we should
 	// extract the value of each property from the texture and use as-is instead of exporting a combined texture.
 
-	const FString TextureName = FString::Printf(TEXT("%s_BaseColor"), *MaterialInterface->GetName());
+	const FString TextureName = MaterialInterface->GetName() + TEXT("_BaseColor");
 
 	const TArray<FGLTFTextureCombineSource> CombineSources =
 	{
@@ -339,7 +339,7 @@ bool FGLTFMaterialConverter::TryGetMetallicAndRoughness(FGLTFConvertBuilder& Bui
 	// TODO: handle the case where TextureSize is 1x1. In this case, both properties are constants and we should
 	// extract the value of each property from the texture and use as-is instead of exporting a combined texture.
 
-	const FString TextureName = FString::Printf(TEXT("%s_MetallicRoughness"), *MaterialInterface->GetName());
+	const FString TextureName = MaterialInterface->GetName() + TEXT("_MetallicRoughness");
 
 	const TArray<FGLTFTextureCombineSource> CombineSources =
 	{
@@ -584,7 +584,7 @@ bool FGLTFMaterialConverter::TryGetBakedTexture(FGLTFConvertBuilder& Builder, FG
 			return false;
 	}
 
-	const FString TextureName = FString::Printf(TEXT("%s_%s"), *MaterialInterface->GetName(), PropertyName);
+	const FString TextureName = MaterialInterface->GetName() + TEXT("_") + PropertyName;
 
 	// TODO: support for other wrapping / filters?
 	const EGLTFJsonTextureWrap TextureWrapS = EGLTFJsonTextureWrap::ClampToEdge;
