[GLTFExporter] refactor(material): use consistent format and indentation for initializers

commit ad580bde70028f2a023cae2aef3f5f22c486a137
Date: Tue, 9 Jun 2020 13:46:43 +0200

    refactor(material): use consistent format and indentation for initializers
    (cherry picked from commit 56b5f5e30daeb73b0fe5e7ff2cb74b42496b5db7​)

#skipci

diff --git a/Engine/Plugins/Enterprise/GLTFExporter/MigratedFromGit/Source/GLTFExporter/Private/Converters/GLTFMaterialConverters.cpp b/Engine/Plugins/Enterprise/GLTFExporter/MigratedFromGit/Source/GLTFExporter/Private/Converters/GLTFMaterialConverters.cpp
--- a/Engine/Plugins/Enterprise/GLTFExporter/MigratedFromGit/Source/GLTFExporter/Private/Converters/GLTFMaterialConverters.cpp
+++ b/Engine/Plugins/Enterprise/GLTFExporter/MigratedFromGit/Source/GLTFExporter/Private/Converters/GLTFMaterialConverters.cpp
@@ -43,7 +43,7 @@ FGLTFJsonMaterialIndex FGLTFMaterialConverter::Add(FGLTFConvertBuilder& Builder,
 	{
 		if (!TryGetConstantColor(JsonMaterial.PBRMetallicRoughness.BaseColorFactor, Material->BaseColor, MaterialInstance))
 		{
-			if (!TryGetSourceTexture(Builder, JsonMaterial.PBRMetallicRoughness.BaseColorTexture, Material->BaseColor, MaterialInstance, {RgbaMask, RgbMask}))
+			if (!TryGetSourceTexture(Builder, JsonMaterial.PBRMetallicRoughness.BaseColorTexture, Material->BaseColor, MaterialInstance, { RgbaMask, RgbMask }))
 			{
 				if (!TryGetBakedTexture(Builder, JsonMaterial.PBRMetallicRoughness.BaseColorTexture, MP_BaseColor, MaterialInterface))
 				{
@@ -60,7 +60,7 @@ FGLTFJsonMaterialIndex FGLTFMaterialConverter::Add(FGLTFConvertBuilder& Builder,
 
 	if (!TryGetConstantColor(JsonMaterial.EmissiveFactor, Material->EmissiveColor, MaterialInstance))
 	{
-		const bool bTextureFound = TryGetSourceTexture(Builder, JsonMaterial.EmissiveTexture, Material->EmissiveColor, MaterialInstance, {RgbaMask, RgbMask}) ||
+		const bool bTextureFound = TryGetSourceTexture(Builder, JsonMaterial.EmissiveTexture, Material->EmissiveColor, MaterialInstance, { RgbaMask, RgbMask }) ||
 			TryGetBakedTexture(Builder, JsonMaterial.EmissiveTexture, MP_EmissiveColor, MaterialInterface);
 
 		if (bTextureFound)
@@ -75,7 +75,7 @@ FGLTFJsonMaterialIndex FGLTFMaterialConverter::Add(FGLTFConvertBuilder& Builder,
 
 	if (Material->Normal.Expression != nullptr)
 	{
-		if (!TryGetSourceTexture(Builder, JsonMaterial.NormalTexture, Material->Normal, MaterialInstance, {RgbaMask, RgbMask}))
+		if (!TryGetSourceTexture(Builder, JsonMaterial.NormalTexture, Material->Normal, MaterialInstance, { RgbaMask, RgbMask }))
 		{
 			if (!TryGetBakedTexture(Builder, JsonMaterial.NormalTexture, MP_Normal, MaterialInterface))
 			{
@@ -86,7 +86,7 @@ FGLTFJsonMaterialIndex FGLTFMaterialConverter::Add(FGLTFConvertBuilder& Builder,
 
 	if (Material->AmbientOcclusion.Expression != nullptr)
 	{
-		if (!TryGetSourceTexture(Builder, JsonMaterial.OcclusionTexture, Material->AmbientOcclusion, MaterialInstance, {RMask}))
+		if (!TryGetSourceTexture(Builder, JsonMaterial.OcclusionTexture, Material->AmbientOcclusion, MaterialInstance, { RMask }))
 		{
 			if (!TryGetBakedTexture(Builder, JsonMaterial.OcclusionTexture, MP_AmbientOcclusion, MaterialInterface))
 			{
@@ -113,7 +113,7 @@ bool FGLTFMaterialConverter::TryGetBaseColorAndOpacity(FGLTFConvertBuilder& Buil
 	// NOTE: since we always bake the properties (for now) when atleast property is non-const, we need
 	// to reset the constant factors to their defaults. Otherwise the baked value of a constant property
 	// would be scaled with the factor, i.e a double scaling.
-	OutPBRParams.BaseColorFactor = {1, 1, 1, 1};
+	OutPBRParams.BaseColorFactor = { 1, 1, 1, 1 };
 
 	const UTexture2D* BaseColorTexture;
 	const UTexture2D* OpacityTexture;
@@ -123,8 +123,8 @@ bool FGLTFMaterialConverter::TryGetBaseColorAndOpacity(FGLTFConvertBuilder& Buil
 	const FLinearColor BaseColorMask(1, 1, 1, 0);
 	const FLinearColor OpacityMask(0, 0, 0, 1);
 
-	const bool bHasBaseColorSourceTexture = TryGetSourceTexture(BaseColorTexture, BaseColorTexCoord, BaseColorInput, MaterialInstance, {BaseColorMask});
-	const bool bHasOpacitySourceTexture = TryGetSourceTexture(OpacityTexture, OpacityTexCoord, OpacityInput, MaterialInstance, {OpacityMask});
+	const bool bHasBaseColorSourceTexture = TryGetSourceTexture(BaseColorTexture, BaseColorTexCoord, BaseColorInput, MaterialInstance, { BaseColorMask });
+	const bool bHasOpacitySourceTexture = TryGetSourceTexture(OpacityTexture, OpacityTexCoord, OpacityInput, MaterialInstance, { OpacityMask });
 
 	// Detect the "happy path" where both inputs share the same texture and are correctly masked.
 	if (bHasBaseColorSourceTexture &&
@@ -161,9 +161,11 @@ bool FGLTFMaterialConverter::TryGetBaseColorAndOpacity(FGLTFConvertBuilder& Buil
 		}
 
 		TexCoord = BaseColorTexCoord;
-		TextureSize = FIntPoint(
+		TextureSize =
+		{
 			FMath::Max(BaseColorTexture->GetSizeX(), OpacityTexture->GetSizeX()),
-			FMath::Max(BaseColorTexture->GetSizeY(), OpacityTexture->GetSizeY()));
+			FMath::Max(BaseColorTexture->GetSizeY(), OpacityTexture->GetSizeY())
+		};
 
 		TextureWrapS = FGLTFConverterUtility::ConvertWrap(BaseColorTexture->AddressX);
 		TextureWrapT = FGLTFConverterUtility::ConvertWrap(BaseColorTexture->AddressY);
@@ -176,7 +178,7 @@ bool FGLTFMaterialConverter::TryGetBaseColorAndOpacity(FGLTFConvertBuilder& Buil
 	else if (bHasBaseColorSourceTexture)
 	{
 		TexCoord = BaseColorTexCoord;
-		TextureSize = FIntPoint(BaseColorTexture->GetSizeX(), BaseColorTexture->GetSizeY());
+		TextureSize = { BaseColorTexture->GetSizeX(), BaseColorTexture->GetSizeY() };
 		TextureWrapS = FGLTFConverterUtility::ConvertWrap(BaseColorTexture->AddressX);
 		TextureWrapT = FGLTFConverterUtility::ConvertWrap(BaseColorTexture->AddressY);
 		TextureMinFilter = FGLTFConverterUtility::ConvertMinFilter(BaseColorTexture->Filter, BaseColorTexture->LODGroup);
@@ -185,7 +187,7 @@ bool FGLTFMaterialConverter::TryGetBaseColorAndOpacity(FGLTFConvertBuilder& Buil
 	else if (bHasOpacitySourceTexture)
 	{
 		TexCoord = OpacityTexCoord;
-		TextureSize = FIntPoint(OpacityTexture->GetSizeX(), OpacityTexture->GetSizeY());
+		TextureSize = { OpacityTexture->GetSizeX(), OpacityTexture->GetSizeY() };
 		TextureWrapS = FGLTFConverterUtility::ConvertWrap(OpacityTexture->AddressX);
 		TextureWrapT = FGLTFConverterUtility::ConvertWrap(OpacityTexture->AddressY);
 		TextureMinFilter = FGLTFConverterUtility::ConvertMinFilter(OpacityTexture->Filter, OpacityTexture->LODGroup);
@@ -199,7 +201,8 @@ bool FGLTFMaterialConverter::TryGetBaseColorAndOpacity(FGLTFConvertBuilder& Buil
 
 	// NOTE: the baked textures may have a different (smaller) size than requested, so we update the
 	// texture-size to fit both textures without wasting too much space.
-	TextureSize = {
+	TextureSize =
+	{
 		FMath::Max(BaseColorTexture->GetSizeX(), OpacityTexture->GetSizeX()),
 		FMath::Max(BaseColorTexture->GetSizeY(), OpacityTexture->GetSizeY())
 	};
@@ -211,8 +214,8 @@ bool FGLTFMaterialConverter::TryGetBaseColorAndOpacity(FGLTFConvertBuilder& Buil
 
 	const TArray<FGLTFTextureCombineSource> CombineSources =
 	{
-		{OpacityTexture, OpacityMask, SE_BLEND_Opaque},
-		{BaseColorTexture, BaseColorMask}
+		{ OpacityTexture, OpacityMask, SE_BLEND_Opaque },
+		{ BaseColorTexture, BaseColorMask }
 	};
 
 	const FGLTFJsonTextureIndex TextureIndex = FGLTFMaterialUtility::AddCombinedTexture(
@@ -257,8 +260,8 @@ bool FGLTFMaterialConverter::TryGetMetallicAndRoughness(FGLTFConvertBuilder& Bui
 	const FLinearColor MetallicMask(0, 0, 1, 0);
 	const FLinearColor RoughnessMask(0, 1, 0, 0);
 
-	const bool bHasMetallicSourceTexture = TryGetSourceTexture(MetallicTexture, MetallicTexCoord, MetallicInput, MaterialInstance, {MetallicMask});
-	const bool bHasRoughnessSourceTexture = TryGetSourceTexture(RoughnessTexture, RoughnessTexCoord, RoughnessInput, MaterialInstance, {RoughnessMask});
+	const bool bHasMetallicSourceTexture = TryGetSourceTexture(MetallicTexture, MetallicTexCoord, MetallicInput, MaterialInstance, { MetallicMask });
+	const bool bHasRoughnessSourceTexture = TryGetSourceTexture(RoughnessTexture, RoughnessTexCoord, RoughnessInput, MaterialInstance, { RoughnessMask });
 
 	// Detect the "happy path" where both inputs share the same texture and are correctly masked.
 	if (bHasMetallicSourceTexture &&
@@ -295,9 +298,11 @@ bool FGLTFMaterialConverter::TryGetMetallicAndRoughness(FGLTFConvertBuilder& Bui
 		}
 
 		TexCoord = MetallicTexCoord;
-		TextureSize = FIntPoint(
+		TextureSize =
+		{
 			FMath::Max(MetallicTexture->GetSizeX(), RoughnessTexture->GetSizeX()),
-			FMath::Max(MetallicTexture->GetSizeY(), RoughnessTexture->GetSizeY()));
+			FMath::Max(MetallicTexture->GetSizeY(), RoughnessTexture->GetSizeY())
+		};
 
 		TextureWrapS = FGLTFConverterUtility::ConvertWrap(MetallicTexture->AddressX);
 		TextureWrapT = FGLTFConverterUtility::ConvertWrap(MetallicTexture->AddressY);
@@ -310,7 +315,7 @@ bool FGLTFMaterialConverter::TryGetMetallicAndRoughness(FGLTFConvertBuilder& Bui
 	else if (bHasMetallicSourceTexture)
 	{
 		TexCoord = MetallicTexCoord;
-		TextureSize = FIntPoint(MetallicTexture->GetSizeX(), MetallicTexture->GetSizeY());
+		TextureSize = { MetallicTexture->GetSizeX(), MetallicTexture->GetSizeY() };
 		TextureWrapS = FGLTFConverterUtility::ConvertWrap(MetallicTexture->AddressX);
 		TextureWrapT = FGLTFConverterUtility::ConvertWrap(MetallicTexture->AddressY);
 		TextureMinFilter = FGLTFConverterUtility::ConvertMinFilter(MetallicTexture->Filter, MetallicTexture->LODGroup);
@@ -319,7 +324,7 @@ bool FGLTFMaterialConverter::TryGetMetallicAndRoughness(FGLTFConvertBuilder& Bui
 	else if (bHasRoughnessSourceTexture)
 	{
 		TexCoord = RoughnessTexCoord;
-		TextureSize = FIntPoint(RoughnessTexture->GetSizeX(), RoughnessTexture->GetSizeY());
+		TextureSize = { RoughnessTexture->GetSizeX(), RoughnessTexture->GetSizeY() };
 		TextureWrapS = FGLTFConverterUtility::ConvertWrap(RoughnessTexture->AddressX);
 		TextureWrapT = FGLTFConverterUtility::ConvertWrap(RoughnessTexture->AddressY);
 		TextureMinFilter = FGLTFConverterUtility::ConvertMinFilter(RoughnessTexture->Filter, RoughnessTexture->LODGroup);
@@ -331,7 +336,8 @@ bool FGLTFMaterialConverter::TryGetMetallicAndRoughness(FGLTFConvertBuilder& Bui
 
 	// NOTE: the baked textures may have a different (smaller) size than requested, so we update the
 	// texture-size to fit both textures without wasting too much space.
-	TextureSize = {
+	TextureSize =
+	{
 		FMath::Max(MetallicTexture->GetSizeX(), RoughnessTexture->GetSizeX()),
 		FMath::Max(MetallicTexture->GetSizeY(), RoughnessTexture->GetSizeY())
 	};
@@ -343,8 +349,8 @@ bool FGLTFMaterialConverter::TryGetMetallicAndRoughness(FGLTFConvertBuilder& Bui
 
 	const TArray<FGLTFTextureCombineSource> CombineSources =
 	{
-		{MetallicTexture, MetallicMask},
-		{RoughnessTexture, RoughnessMask}
+		{ MetallicTexture, MetallicMask },
+		{ RoughnessTexture, RoughnessMask }
 	};
 
 	const FGLTFJsonTextureIndex TextureIndex = FGLTFMaterialUtility::AddCombinedTexture(
@@ -395,7 +401,7 @@ bool FGLTFMaterialConverter::TryGetConstantColor(FLinearColor& OutValue, const F
 	if (Expression == nullptr)
 	{
 		// TODO: is this assumption correct?
-		OutValue = FLinearColor(MaterialInput.Constant);
+		OutValue = { MaterialInput.Constant };
 		return true;
 	}
 
diff --git a/Engine/Plugins/Enterprise/GLTFExporter/MigratedFromGit/Source/GLTFExporter/Private/Converters/GLTFMaterialUtility.cpp b/Engine/Plugins/Enterprise/GLTFExporter/MigratedFromGit/Source/GLTFExporter/Private/Converters/GLTFMaterialUtility.cpp
--- a/Engine/Plugins/Enterprise/GLTFExporter/MigratedFromGit/Source/GLTFExporter/Private/Converters/GLTFMaterialUtility.cpp
+++ b/Engine/Plugins/Enterprise/GLTFExporter/MigratedFromGit/Source/GLTFExporter/Private/Converters/GLTFMaterialUtility.cpp
@@ -34,7 +34,7 @@ bool FGLTFMaterialUtility::CombineTextures(TArray<FColor>& OutPixels, const TArr
 	UTextureRenderTarget2D* RenderTarget2D = NewObject<UTextureRenderTarget2D>();
 
 	RenderTarget2D->AddToRoot();
-	RenderTarget2D->ClearColor = FLinearColor(0.0f, 0.0f, 0.0f, 0.0f);
+	RenderTarget2D->ClearColor = { 0.0f, 0.0f, 0.0f, 0.0f };
 	RenderTarget2D->InitCustomFormat(OutputSize.X, OutputSize.Y, OutputPixelFormat, false);
 	RenderTarget2D->TargetGamma = 0;
 
@@ -42,7 +42,7 @@ bool FGLTFMaterialUtility::CombineTextures(TArray<FColor>& OutPixels, const TArr
 	FCanvas Canvas(RenderTarget, nullptr, 0, 0, 0, GMaxRHIFeatureLevel);
 
 	Canvas.SetRenderTarget_GameThread(RenderTarget);
-	Canvas.Clear(FLinearColor(0, 0, 0, 0));
+	Canvas.Clear({ 0, 0, 0, 0 });
 
 	const FVector2D TileSize(OutputSize.X, OutputSize.Y);
 	const FVector2D TilePosition(0, 0);
@@ -77,7 +77,7 @@ FGLTFPropertyBakeOutput FGLTFMaterialUtility::BakeMaterialProperty(const FIntPoi
 	TArray<FMeshData*> MeshSettings;
 
 	FMeshData MeshSet;
-	MeshSet.TextureCoordinateBox = FBox2D(FVector2D(0, 0), FVector2D(1, 1));
+	MeshSet.TextureCoordinateBox = { { 0, 0 }, { 1, 1 } };
 
 	// TODO: Do we need to fill in any more info in MeshSet?
 
@@ -117,7 +117,8 @@ FGLTFPropertyBakeOutput FGLTFMaterialUtility::BakeMaterialProperty(const FIntPoi
 		}
 	}
 
-	return {
+	return
+	{
 		MaterialProperty,
 		PF_B8G8R8A8,
 		BakedPixels,
