[GLTFExporter] refactor(material): use correct type of literals when initializing values

commit dfad92345c6f50cc4e0eaedf2c94673a9e88294d
Date: Tue, 9 Jun 2020 13:53:18 +0200

    refactor(material): use correct type of literals when initializing values
    (cherry picked from commit 6bc77a5ea5c20e3905eb0849568eaeabb9e8f90c​)

#skipci

diff --git a/Engine/Plugins/Enterprise/GLTFExporter/MigratedFromGit/Source/GLTFExporter/Private/Converters/GLTFMaterialConverters.cpp b/Engine/Plugins/Enterprise/GLTFExporter/MigratedFromGit/Source/GLTFExporter/Private/Converters/GLTFMaterialConverters.cpp
--- a/Engine/Plugins/Enterprise/GLTFExporter/MigratedFromGit/Source/GLTFExporter/Private/Converters/GLTFMaterialConverters.cpp
+++ b/Engine/Plugins/Enterprise/GLTFExporter/MigratedFromGit/Source/GLTFExporter/Private/Converters/GLTFMaterialConverters.cpp
@@ -113,15 +113,15 @@ bool FGLTFMaterialConverter::TryGetBaseColorAndOpacity(FGLTFConvertBuilder& Buil
 	// NOTE: since we always bake the properties (for now) when atleast property is non-const, we need
 	// to reset the constant factors to their defaults. Otherwise the baked value of a constant property
 	// would be scaled with the factor, i.e a double scaling.
-	OutPBRParams.BaseColorFactor = { 1, 1, 1, 1 };
+	OutPBRParams.BaseColorFactor = { 1.0f, 1.0f, 1.0f, 1.0f };
 
 	const UTexture2D* BaseColorTexture;
 	const UTexture2D* OpacityTexture;
 	int32 BaseColorTexCoord;
 	int32 OpacityTexCoord;
 
-	const FLinearColor BaseColorMask(1, 1, 1, 0);
-	const FLinearColor OpacityMask(0, 0, 0, 1);
+	const FLinearColor BaseColorMask(1.0f, 1.0f, 1.0f, 0.0f);
+	const FLinearColor OpacityMask(0.0f, 0.0f, 0.0f, 1.0f);
 
 	const bool bHasBaseColorSourceTexture = TryGetSourceTexture(BaseColorTexture, BaseColorTexCoord, BaseColorInput, MaterialInstance, { BaseColorMask });
 	const bool bHasOpacitySourceTexture = TryGetSourceTexture(OpacityTexture, OpacityTexCoord, OpacityInput, MaterialInstance, { OpacityMask });
@@ -249,16 +249,16 @@ bool FGLTFMaterialConverter::TryGetMetallicAndRoughness(FGLTFConvertBuilder& Bui
 	// NOTE: since we always bake the properties (for now) when atleast property is non-const, we need
 	// to reset the constant factors to their defaults. Otherwise the baked value of a constant property
 	// would be scaled with the factor, i.e a double scaling.
-	OutPBRParams.MetallicFactor = 1;
-	OutPBRParams.RoughnessFactor = 1;
+	OutPBRParams.MetallicFactor = 1.0f;
+	OutPBRParams.RoughnessFactor = 1.0f;
 
 	const UTexture2D* MetallicTexture;
 	const UTexture2D* RoughnessTexture;
 	int32 MetallicTexCoord;
 	int32 RoughnessTexCoord;
 
-	const FLinearColor MetallicMask(0, 0, 1, 0);
-	const FLinearColor RoughnessMask(0, 1, 0, 0);
+	const FLinearColor MetallicMask(0.0f, 0.0f, 1.0f, 0.0f);
+	const FLinearColor RoughnessMask(0.0f, 1.0f, 0.0f, 0.0f);
 
 	const bool bHasMetallicSourceTexture = TryGetSourceTexture(MetallicTexture, MetallicTexCoord, MetallicInput, MaterialInstance, { MetallicMask });
 	const bool bHasRoughnessSourceTexture = TryGetSourceTexture(RoughnessTexture, RoughnessTexCoord, RoughnessInput, MaterialInstance, { RoughnessMask });
diff --git a/Engine/Plugins/Enterprise/GLTFExporter/MigratedFromGit/Source/GLTFExporter/Private/Converters/GLTFMaterialUtility.cpp b/Engine/Plugins/Enterprise/GLTFExporter/MigratedFromGit/Source/GLTFExporter/Private/Converters/GLTFMaterialUtility.cpp
--- a/Engine/Plugins/Enterprise/GLTFExporter/MigratedFromGit/Source/GLTFExporter/Private/Converters/GLTFMaterialUtility.cpp
+++ b/Engine/Plugins/Enterprise/GLTFExporter/MigratedFromGit/Source/GLTFExporter/Private/Converters/GLTFMaterialUtility.cpp
@@ -36,16 +36,16 @@ bool FGLTFMaterialUtility::CombineTextures(TArray<FColor>& OutPixels, const TArr
 	RenderTarget2D->AddToRoot();
 	RenderTarget2D->ClearColor = { 0.0f, 0.0f, 0.0f, 0.0f };
 	RenderTarget2D->InitCustomFormat(OutputSize.X, OutputSize.Y, OutputPixelFormat, false);
-	RenderTarget2D->TargetGamma = 0;
+	RenderTarget2D->TargetGamma = 0.0f;
 
 	FRenderTarget* RenderTarget = RenderTarget2D->GameThread_GetRenderTargetResource();
-	FCanvas Canvas(RenderTarget, nullptr, 0, 0, 0, GMaxRHIFeatureLevel);
+	FCanvas Canvas(RenderTarget, nullptr, 0.0f, 0.0f, 0.0f, GMaxRHIFeatureLevel);
 
 	Canvas.SetRenderTarget_GameThread(RenderTarget);
-	Canvas.Clear({ 0, 0, 0, 0 });
+	Canvas.Clear({ 0.0f, 0.0f, 0.0f, 0.0f });
 
 	const FVector2D TileSize(OutputSize.X, OutputSize.Y);
-	const FVector2D TilePosition(0, 0);
+	const FVector2D TilePosition(0.0f, 0.0f);
 
 	for (const FGLTFTextureCombineSource& Source: Sources)
 	{
@@ -77,7 +77,7 @@ FGLTFPropertyBakeOutput FGLTFMaterialUtility::BakeMaterialProperty(const FIntPoi
 	TArray<FMeshData*> MeshSettings;
 
 	FMeshData MeshSet;
-	MeshSet.TextureCoordinateBox = { { 0, 0 }, { 1, 1 } };
+	MeshSet.TextureCoordinateBox = { { 0.0f, 0.0f }, { 1.0f, 1.0f } };
 
 	// TODO: Do we need to fill in any more info in MeshSet?
 
diff --git a/Engine/Plugins/Enterprise/GLTFExporter/MigratedFromGit/Source/GLTFExporter/Private/Converters/GLTFMaterialUtility.h b/Engine/Plugins/Enterprise/GLTFExporter/MigratedFromGit/Source/GLTFExporter/Private/Converters/GLTFMaterialUtility.h
--- a/Engine/Plugins/Enterprise/GLTFExporter/MigratedFromGit/Source/GLTFExporter/Private/Converters/GLTFMaterialUtility.h
+++ b/Engine/Plugins/Enterprise/GLTFExporter/MigratedFromGit/Source/GLTFExporter/Private/Converters/GLTFMaterialUtility.h
@@ -7,7 +7,7 @@
 
 struct FGLTFTextureCombineSource
 {
-	constexpr FORCEINLINE FGLTFTextureCombineSource(const UTexture2D* Texture, FLinearColor TintColor = FLinearColor(1, 1, 1, 1), ESimpleElementBlendMode BlendMode = SE_BLEND_Additive)
+	constexpr FORCEINLINE FGLTFTextureCombineSource(const UTexture2D* Texture, FLinearColor TintColor = { 1.0f, 1.0f, 1.0f, 1.0f }, ESimpleElementBlendMode BlendMode = SE_BLEND_Additive)
 		: Texture(Texture), TintColor(TintColor), BlendMode(BlendMode)
 	{}
 
