[GLTFExporter] refactor(texture): remove redundant enum prefixes

commit 3ace7411c5dae9ccd918e40598d0752c95e8ef60
Date: Wed, 10 Jun 2020 10:40:55 +0200

    refactor(texture): remove redundant enum prefixes
    (cherry picked from commit 96ad890eb2e6395f60bd488c1b59f0cf6cc77cbe​)

#skipci

diff --git a/Engine/Plugins/Enterprise/GLTFExporter/MigratedFromGit/Source/GLTFExporter/Private/Converters/GLTFConverterUtility.cpp b/Engine/Plugins/Enterprise/GLTFExporter/MigratedFromGit/Source/GLTFExporter/Private/Converters/GLTFConverterUtility.cpp
--- a/Engine/Plugins/Enterprise/GLTFExporter/MigratedFromGit/Source/GLTFExporter/Private/Converters/GLTFConverterUtility.cpp
+++ b/Engine/Plugins/Enterprise/GLTFExporter/MigratedFromGit/Source/GLTFExporter/Private/Converters/GLTFConverterUtility.cpp
@@ -13,12 +13,12 @@ TextureFilter FGLTFConverterUtility::GetDefaultFilter(TextureGroup LODGroup)
 
 	switch (TextureLODGroup.Filter)
 	{
-		case ETextureSamplerFilter::Point:             return TextureFilter::TF_Nearest;
-		case ETextureSamplerFilter::Bilinear:          return TextureFilter::TF_Bilinear;
-		case ETextureSamplerFilter::Trilinear:         return TextureFilter::TF_Trilinear;
-		case ETextureSamplerFilter::AnisotropicPoint:  return TextureFilter::TF_Nearest;
-		case ETextureSamplerFilter::AnisotropicLinear: return TextureFilter::TF_Trilinear;
-		default:                                       return TextureFilter::TF_Default; // fallback
+		case ETextureSamplerFilter::Point:             return TF_Nearest;
+		case ETextureSamplerFilter::Bilinear:          return TF_Bilinear;
+		case ETextureSamplerFilter::Trilinear:         return TF_Trilinear;
+		case ETextureSamplerFilter::AnisotropicPoint:  return TF_Nearest;
+		case ETextureSamplerFilter::AnisotropicLinear: return TF_Trilinear;
+		default:                                       return TF_Default; // fallback
 	}
 }
 
diff --git a/Engine/Plugins/Enterprise/GLTFExporter/MigratedFromGit/Source/GLTFExporter/Private/Converters/GLTFConverterUtility.h b/Engine/Plugins/Enterprise/GLTFExporter/MigratedFromGit/Source/GLTFExporter/Private/Converters/GLTFConverterUtility.h
--- a/Engine/Plugins/Enterprise/GLTFExporter/MigratedFromGit/Source/GLTFExporter/Private/Converters/GLTFConverterUtility.h
+++ b/Engine/Plugins/Enterprise/GLTFExporter/MigratedFromGit/Source/GLTFExporter/Private/Converters/GLTFConverterUtility.h
@@ -95,10 +95,10 @@ struct FGLTFConverterUtility
 	{
 		switch (Mode)
 		{
-			case EBlendMode::BLEND_Opaque:      return EGLTFJsonAlphaMode::Opaque;
-			case EBlendMode::BLEND_Translucent: return EGLTFJsonAlphaMode::Blend;
-			case EBlendMode::BLEND_Masked:      return EGLTFJsonAlphaMode::Mask;
-			default:                            return EGLTFJsonAlphaMode::Opaque; // fallback
+			case BLEND_Opaque:      return EGLTFJsonAlphaMode::Opaque;
+			case BLEND_Translucent: return EGLTFJsonAlphaMode::Blend;
+			case BLEND_Masked:      return EGLTFJsonAlphaMode::Mask;
+			default:                return EGLTFJsonAlphaMode::Opaque; // fallback
 		}
 	}
 
@@ -106,10 +106,10 @@ struct FGLTFConverterUtility
 	{
 		switch (Address)
 		{
-			case TextureAddress::TA_Wrap:   return EGLTFJsonTextureWrap::Repeat;
-			case TextureAddress::TA_Mirror: return EGLTFJsonTextureWrap::MirroredRepeat;
-			case TextureAddress::TA_Clamp:  return EGLTFJsonTextureWrap::ClampToEdge;
-			default:                        return EGLTFJsonTextureWrap::Repeat; // fallback
+			case TA_Wrap:   return EGLTFJsonTextureWrap::Repeat;
+			case TA_Mirror: return EGLTFJsonTextureWrap::MirroredRepeat;
+			case TA_Clamp:  return EGLTFJsonTextureWrap::ClampToEdge;
+			default:        return EGLTFJsonTextureWrap::Repeat; // fallback
 		}
 	}
 
@@ -117,10 +117,10 @@ struct FGLTFConverterUtility
 	{
 		switch (Filter)
 		{
-			case TextureFilter::TF_Nearest:   return EGLTFJsonTextureFilter::NearestMipmapNearest;
-			case TextureFilter::TF_Bilinear:  return EGLTFJsonTextureFilter::LinearMipmapNearest;
-			case TextureFilter::TF_Trilinear: return EGLTFJsonTextureFilter::LinearMipmapLinear;
-			default:                          return EGLTFJsonTextureFilter::None;
+			case TF_Nearest:   return EGLTFJsonTextureFilter::NearestMipmapNearest;
+			case TF_Bilinear:  return EGLTFJsonTextureFilter::LinearMipmapNearest;
+			case TF_Trilinear: return EGLTFJsonTextureFilter::LinearMipmapLinear;
+			default:           return EGLTFJsonTextureFilter::None;
 		}
 	}
 
@@ -128,10 +128,10 @@ struct FGLTFConverterUtility
 	{
 		switch (Filter)
 		{
-			case TextureFilter::TF_Nearest:   return EGLTFJsonTextureFilter::Nearest;
-			case TextureFilter::TF_Bilinear:  return EGLTFJsonTextureFilter::Linear;
-			case TextureFilter::TF_Trilinear: return EGLTFJsonTextureFilter::Linear;
-			default:                          return EGLTFJsonTextureFilter::None;
+			case TF_Nearest:   return EGLTFJsonTextureFilter::Nearest;
+			case TF_Bilinear:  return EGLTFJsonTextureFilter::Linear;
+			case TF_Trilinear: return EGLTFJsonTextureFilter::Linear;
+			default:           return EGLTFJsonTextureFilter::None;
 		}
 	}
 
