[GLTFExporter] refactor(material): disable export of EmissiveColor for now since it has issues

commit 6b0f88260b9d46bc241c22977ff7b53e12e88484
Date: Wed, 10 Jun 2020 10:50:51 +0200

    refactor(material): disable export of EmissiveColor for now since it has issues
    solving the issues is outside the scope of this MR and they will be solved in a dedicated MR later
    (cherry picked from commit f8857b326b2fa687913446a291d77d827dd86345​)

#skipci

diff --git a/Engine/Plugins/Enterprise/GLTFExporter/MigratedFromGit/Source/GLTFExporter/Private/Converters/GLTFMaterialConverters.cpp b/Engine/Plugins/Enterprise/GLTFExporter/MigratedFromGit/Source/GLTFExporter/Private/Converters/GLTFMaterialConverters.cpp
--- a/Engine/Plugins/Enterprise/GLTFExporter/MigratedFromGit/Source/GLTFExporter/Private/Converters/GLTFMaterialConverters.cpp
+++ b/Engine/Plugins/Enterprise/GLTFExporter/MigratedFromGit/Source/GLTFExporter/Private/Converters/GLTFMaterialConverters.cpp
@@ -58,6 +58,12 @@ FGLTFJsonMaterialIndex FGLTFMaterialConverter::Add(FGLTFConvertBuilder& Builder,
 		// TODO: handle failure?
 	}
 
+	// NOTE: export of EmissiveColor has been temporarily disabled because of visual differences that
+	// are not solvable in the scope of MR !19.
+	// The issues revolve mainly around very bright emission in the exported materials compared to
+	// how the same materials look inside of UE.
+	// TODO: solve the issues in a separate MR.
+	/*
 	if (!TryGetConstantColor(JsonMaterial.EmissiveFactor, Material->EmissiveColor, MaterialInstance))
 	{
 		const bool bTextureFound = TryGetSourceTexture(Builder, JsonMaterial.EmissiveTexture, Material->EmissiveColor, MaterialInstance, { RgbaMask, RgbMask }) ||
@@ -72,6 +78,7 @@ FGLTFJsonMaterialIndex FGLTFMaterialConverter::Add(FGLTFConvertBuilder& Builder,
 			// TODO: handle failure?
 		}
 	}
+	*/
 
 	if (Material->Normal.Expression != nullptr)
 	{
