[GLTFExporter] feat(material): store constant (baked) values as factors instead of textures for "simple" properties

commit 406cc5a3dd9a1d3e3727673e21aeb261fd72d13c
Date: Wed, 10 Jun 2020 13:27:38 +0200

    feat(material): store constant (baked) values as factors instead of textures for "simple" properties
    (cherry picked from commit b1a2b8eeb9c73225cd16a1b16acf6a80975a70c0​)

#skipci

diff --git a/Engine/Plugins/Enterprise/GLTFExporter/MigratedFromGit/Source/GLTFExporter/Private/Converters/GLTFMaterialConverters.cpp b/Engine/Plugins/Enterprise/GLTFExporter/MigratedFromGit/Source/GLTFExporter/Private/Converters/GLTFMaterialConverters.cpp
--- a/Engine/Plugins/Enterprise/GLTFExporter/MigratedFromGit/Source/GLTFExporter/Private/Converters/GLTFMaterialConverters.cpp
+++ b/Engine/Plugins/Enterprise/GLTFExporter/MigratedFromGit/Source/GLTFExporter/Private/Converters/GLTFMaterialConverters.cpp
@@ -45,7 +45,7 @@
 		{
 			if (!TryGetSourceTexture(Builder, JsonMaterial.PBRMetallicRoughness.BaseColorTexture, Material->BaseColor, MaterialInstance, { RgbaMask, RgbMask }))
 			{
-				if (!TryGetBakedTexture(Builder, JsonMaterial.PBRMetallicRoughness.BaseColorTexture, MP_BaseColor, MaterialInterface))
+				if (!TryGetBakedMaterialProperty(Builder, JsonMaterial.PBRMetallicRoughness.BaseColorTexture, JsonMaterial.PBRMetallicRoughness.BaseColorFactor, MP_BaseColor, MaterialInterface))
 				{
 					// TODO: handle failure?
 				}
@@ -66,16 +66,23 @@
 	/*
 	if (!TryGetConstantColor(JsonMaterial.EmissiveFactor, Material->EmissiveColor, MaterialInstance))
 	{
-		const bool bTextureFound = TryGetSourceTexture(Builder, JsonMaterial.EmissiveTexture, Material->EmissiveColor, MaterialInstance, { RgbaMask, RgbMask }) ||
-			TryGetBakedTexture(Builder, JsonMaterial.EmissiveTexture, MP_EmissiveColor, MaterialInterface);
-
-		if (bTextureFound)
+		if (!TryGetSourceTexture(Builder, JsonMaterial.EmissiveTexture, Material->EmissiveColor, MaterialInstance, { RgbaMask, RgbMask }))
 		{
-			JsonMaterial.EmissiveFactor = FGLTFJsonColor3::White; // make sure texture is not multiplied with black
+			if (TryGetBakedMaterialProperty(Builder, JsonMaterial.EmissiveTexture, JsonMaterial.EmissiveFactor, MP_EmissiveColor, MaterialInterface))
+			{
+				if (JsonMaterial.EmissiveTexture.Index != INDEX_NONE)
+				{
+					JsonMaterial.EmissiveFactor = FGLTFJsonColor3::White;	// make sure texture is not multiplied with black
+				}
+			}
+			else
+			{
+				// TODO: handle failure?
+			}
 		}
 		else
 		{
-			// TODO: handle failure?
+			JsonMaterial.EmissiveFactor = FGLTFJsonColor3::White;	// make sure texture is not multiplied with black
 		}
 	}
 	*/
@@ -84,7 +91,7 @@
 	{
 		if (!TryGetSourceTexture(Builder, JsonMaterial.NormalTexture, Material->Normal, MaterialInstance, { RgbaMask, RgbMask }))
 		{
-			if (!TryGetBakedTexture(Builder, JsonMaterial.NormalTexture, MP_Normal, MaterialInterface))
+			if (!TryGetBakedMaterialProperty(Builder, JsonMaterial.NormalTexture, MP_Normal, MaterialInterface))
 			{
 				// TODO: handle failure?
 			}
@@ -95,7 +102,7 @@
 	{
 		if (!TryGetSourceTexture(Builder, JsonMaterial.OcclusionTexture, Material->AmbientOcclusion, MaterialInstance, { RMask }))
 		{
-			if (!TryGetBakedTexture(Builder, JsonMaterial.OcclusionTexture, MP_AmbientOcclusion, MaterialInterface))
+			if (!TryGetBakedMaterialProperty(Builder, JsonMaterial.OcclusionTexture, MP_AmbientOcclusion, MaterialInterface))
 			{
 				// TODO: handle failure?
 			}
@@ -568,60 +575,6 @@
 	return false;
 }
 
-bool FGLTFMaterialConverter::TryGetBakedTexture(FGLTFConvertBuilder& Builder, FGLTFJsonTextureInfo& OutTexInfo, EMaterialProperty MaterialProperty, const UMaterialInterface* MaterialInterface) const
-{
-	// TODO: make default baking-resolution configurable
-	// TODO: add support for detecting the correct tex-coord for this property based on connected nodes
-	// TODO: add support for calculating the ideal resolution to use for baking based on connected (texture) nodes
-	const FIntPoint TextureSize(512, 512);
-	const uint32 TexCoord = 0;
-
-	const FGLTFPropertyBakeOutput PropertyBakeOutput = FGLTFMaterialUtility::BakeMaterialProperty(
-		TextureSize,
-		MaterialProperty,
-		MaterialInterface);
-
-	// TODO: handle cases where PropertyBakeOutput.EmissiveScale is not 1.0 (when baking EmissiveColor)
-	// TODO: handle the case where TextureSize is 1x1, which means is could be stored in a constant instead of in a texture
-
-	TCHAR* PropertyName;
-
-	switch (MaterialProperty)
-	{
-		case MP_BaseColor: PropertyName = TEXT("BaseColor"); break;
-		case MP_Normal: PropertyName = TEXT("Normal"); break;
-		case MP_EmissiveColor: PropertyName = TEXT("Emissive"); break;
-		case MP_AmbientOcclusion: PropertyName = TEXT("Occlusion"); break;
-		default:
-			// TODO: support for more properties
-			return false;
-	}
-
-	const FString TextureName = MaterialInterface->GetName() + TEXT("_") + PropertyName;
-
-	// TODO: support for other wrapping / filters?
-	const EGLTFJsonTextureWrap TextureWrapS = EGLTFJsonTextureWrap::ClampToEdge;
-	const EGLTFJsonTextureWrap TextureWrapT = EGLTFJsonTextureWrap::ClampToEdge;
-	const EGLTFJsonTextureFilter TextureMinFilter = EGLTFJsonTextureFilter::LinearMipmapLinear;
-	const EGLTFJsonTextureFilter TextureMagFilter = EGLTFJsonTextureFilter::Linear;
-
-	const FGLTFJsonTextureIndex TextureIndex = FGLTFMaterialUtility::AddTexture(
-		Builder,
-		PropertyBakeOutput.Pixels,
-		PropertyBakeOutput.Size,
-		TextureName,
-		PropertyBakeOutput.PixelFormat,
-		TextureMinFilter,
-		TextureMagFilter,
-		TextureWrapS,
-		TextureWrapT);
-
-	OutTexInfo.TexCoord = TexCoord;
-	OutTexInfo.Index = TextureIndex;
-
-	return true;
-}
-
 bool FGLTFMaterialConverter::TryGetSourceTexture(FGLTFConvertBuilder& Builder, FGLTFJsonTextureInfo& OutTexInfo, const FExpressionInput& MaterialInput, const UMaterialInstance* MaterialInstance, const TArray<FLinearColor>& AllowedMasks) const
 {
 	const UTexture2D* Texture;
@@ -709,3 +662,105 @@
 
 	return false;
 }
+
+bool FGLTFMaterialConverter::TryGetBakedMaterialProperty(FGLTFConvertBuilder& Builder, FGLTFJsonTextureInfo& OutTexInfo, FGLTFJsonColor3& OutConstant, EMaterialProperty MaterialProperty, const UMaterialInterface* MaterialInterface) const
+{
+	const FGLTFPropertyBakeOutput PropertyBakeOutput = BakeMaterialProperty(MaterialProperty, MaterialInterface);
+
+	if (PropertyBakeOutput.bIsConstant)
+	{
+		// TODO: is this conversion between FColor & FLinearColor correct with regards to sRGB etc?
+		const FLinearColor ConstantValue = PropertyBakeOutput.ConstantValue;
+
+		OutConstant = FGLTFConverterUtility::ConvertColor(ConstantValue);
+		return true;
+	}
+	else
+	{
+		return StoreBakedPropertyTexture(Builder, OutTexInfo, PropertyBakeOutput, MaterialProperty, MaterialInterface);
+	}
+}
+
+bool FGLTFMaterialConverter::TryGetBakedMaterialProperty(FGLTFConvertBuilder& Builder, FGLTFJsonTextureInfo& OutTexInfo, FGLTFJsonColor4& OutConstant, EMaterialProperty MaterialProperty, const UMaterialInterface* MaterialInterface) const
+{
+	const FGLTFPropertyBakeOutput PropertyBakeOutput = BakeMaterialProperty(MaterialProperty, MaterialInterface);
+
+	if (PropertyBakeOutput.bIsConstant)
+	{
+		// TODO: is this conversion between FColor & FLinearColor correct with regards to sRGB etc?
+		const FLinearColor ConstantValue = PropertyBakeOutput.ConstantValue;
+
+		OutConstant = FGLTFConverterUtility::ConvertColor(ConstantValue);
+		return true;
+	}
+	else
+	{
+		return StoreBakedPropertyTexture(Builder, OutTexInfo, PropertyBakeOutput, MaterialProperty, MaterialInterface);
+	}
+}
+
+bool FGLTFMaterialConverter::TryGetBakedMaterialProperty(FGLTFConvertBuilder& Builder, FGLTFJsonTextureInfo& OutTexInfo, EMaterialProperty MaterialProperty, const UMaterialInterface* MaterialInterface) const
+{
+	const FGLTFPropertyBakeOutput PropertyBakeOutput = BakeMaterialProperty(MaterialProperty, MaterialInterface);
+	return StoreBakedPropertyTexture(Builder, OutTexInfo, PropertyBakeOutput, MaterialProperty, MaterialInterface);
+}
+
+FGLTFPropertyBakeOutput FGLTFMaterialConverter::BakeMaterialProperty(EMaterialProperty MaterialProperty, const UMaterialInterface* MaterialInterface) const
+{
+	// TODO: handle cases where PropertyBakeOutput.EmissiveScale is not 1.0 (when baking EmissiveColor)
+	// TODO: make default baking-resolution configurable
+	// TODO: add support for calculating the ideal resolution to use for baking based on connected (texture) nodes
+	const FIntPoint TextureSize(512, 512);
+
+	const FGLTFPropertyBakeOutput PropertyBakeOutput = FGLTFMaterialUtility::BakeMaterialProperty(
+		TextureSize,
+		MaterialProperty,
+		MaterialInterface);
+
+	return PropertyBakeOutput;
+}
+
+bool FGLTFMaterialConverter::StoreBakedPropertyTexture(FGLTFConvertBuilder& Builder, FGLTFJsonTextureInfo& OutTexInfo, const FGLTFPropertyBakeOutput& PropertyBakeOutput, EMaterialProperty MaterialProperty, const UMaterialInterface* MaterialInterface) const
+{
+	// TODO: right now a constant property can be saved as a texture. Should we return false instead?
+
+	TCHAR* PropertyName;
+
+	switch (MaterialProperty)
+	{
+		case MP_BaseColor: PropertyName = TEXT("BaseColor"); break;
+		case MP_Normal: PropertyName = TEXT("Normal"); break;
+		case MP_EmissiveColor: PropertyName = TEXT("Emissive"); break;
+		case MP_AmbientOcclusion: PropertyName = TEXT("Occlusion"); break;
+		default:
+			// TODO: support for more properties
+			return false;
+	}
+
+	const FString TextureName = MaterialInterface->GetName() + TEXT("_") + PropertyName;
+
+	// TODO: add support for detecting the correct tex-coord for this property based on connected nodes
+	const uint32 TexCoord = 0;
+
+	// TODO: support for other wrapping / filters?
+	const EGLTFJsonTextureWrap TextureWrapS = EGLTFJsonTextureWrap::ClampToEdge;
+	const EGLTFJsonTextureWrap TextureWrapT = EGLTFJsonTextureWrap::ClampToEdge;
+	const EGLTFJsonTextureFilter TextureMinFilter = EGLTFJsonTextureFilter::LinearMipmapLinear;
+	const EGLTFJsonTextureFilter TextureMagFilter = EGLTFJsonTextureFilter::Linear;
+
+	const FGLTFJsonTextureIndex TextureIndex = FGLTFMaterialUtility::AddTexture(
+		Builder,
+		PropertyBakeOutput.Pixels,
+		PropertyBakeOutput.Size,
+		TextureName,
+		PropertyBakeOutput.PixelFormat,
+		TextureMinFilter,
+		TextureMagFilter,
+		TextureWrapS,
+		TextureWrapT);
+
+	OutTexInfo.TexCoord = TexCoord;
+	OutTexInfo.Index = TextureIndex;
+
+	return true;
+}
diff --git a/Engine/Plugins/Enterprise/GLTFExporter/MigratedFromGit/Source/GLTFExporter/Private/Converters/GLTFMaterialConverters.h b/Engine/Plugins/Enterprise/GLTFExporter/MigratedFromGit/Source/GLTFExporter/Private/Converters/GLTFMaterialConverters.h
--- a/Engine/Plugins/Enterprise/GLTFExporter/MigratedFromGit/Source/GLTFExporter/Private/Converters/GLTFMaterialConverters.h
+++ b/Engine/Plugins/Enterprise/GLTFExporter/MigratedFromGit/Source/GLTFExporter/Private/Converters/GLTFMaterialConverters.h
@@ -7,6 +7,8 @@
 #include "Converters/GLTFConverter.h"
 #include "Engine.h"
 
+struct FGLTFPropertyBakeOutput;
+
 class FGLTFMaterialConverter final : public TGLTFConverter<FGLTFJsonMaterialIndex, const UMaterialInterface*>
 {
 	FGLTFJsonMaterialIndex Add(FGLTFConvertBuilder& Builder, const FString& Name, const UMaterialInterface* Material) override;
@@ -21,5 +23,12 @@ class FGLTFMaterialConverter final : public TGLTFConverter<FGLTFJsonMaterialInde
 
 	bool TryGetSourceTexture(FGLTFConvertBuilder& Builder, FGLTFJsonTextureInfo& OutTexInfo, const FExpressionInput& MaterialInput, const UMaterialInstance* MaterialInstance, const TArray<FLinearColor>& AllowedMasks = {}) const;
 	bool TryGetSourceTexture(const UTexture2D*& OutTexture, int32& OutTexCoord, const FExpressionInput& MaterialInput, const UMaterialInstance* MaterialInstance, const TArray<FLinearColor>& AllowedMasks = {}) const;
-	bool TryGetBakedTexture(FGLTFConvertBuilder& Builder, FGLTFJsonTextureInfo& OutTexInfo, EMaterialProperty MaterialProperty, const UMaterialInterface* MaterialInterface) const;
+
+	bool TryGetBakedMaterialProperty(FGLTFConvertBuilder& Builder, FGLTFJsonTextureInfo& OutTexInfo, FGLTFJsonColor3& OutConstant, EMaterialProperty MaterialProperty, const UMaterialInterface* MaterialInterface) const;
+	bool TryGetBakedMaterialProperty(FGLTFConvertBuilder& Builder, FGLTFJsonTextureInfo& OutTexInfo, FGLTFJsonColor4& OutConstant, EMaterialProperty MaterialProperty, const UMaterialInterface* MaterialInterface) const;
+	bool TryGetBakedMaterialProperty(FGLTFConvertBuilder& Builder, FGLTFJsonTextureInfo& OutTexInfo, EMaterialProperty MaterialProperty, const UMaterialInterface* MaterialInterface) const;
+
+	FGLTFPropertyBakeOutput BakeMaterialProperty(EMaterialProperty MaterialProperty, const UMaterialInterface* MaterialInterface) const;
+
+	bool StoreBakedPropertyTexture(FGLTFConvertBuilder& Builder, FGLTFJsonTextureInfo& OutTexInfo, const FGLTFPropertyBakeOutput& PropertyBakeOutput, EMaterialProperty MaterialProperty, const UMaterialInterface* MaterialInterface) const;
 };
diff --git a/Engine/Plugins/Enterprise/GLTFExporter/MigratedFromGit/Source/GLTFExporter/Private/Converters/GLTFMaterialUtility.cpp b/Engine/Plugins/Enterprise/GLTFExporter/MigratedFromGit/Source/GLTFExporter/Private/Converters/GLTFMaterialUtility.cpp
--- a/Engine/Plugins/Enterprise/GLTFExporter/MigratedFromGit/Source/GLTFExporter/Private/Converters/GLTFMaterialUtility.cpp
+++ b/Engine/Plugins/Enterprise/GLTFExporter/MigratedFromGit/Source/GLTFExporter/Private/Converters/GLTFMaterialUtility.cpp
@@ -117,14 +117,17 @@ FGLTFPropertyBakeOutput FGLTFMaterialUtility::BakeMaterialProperty(const FIntPoi
 		}
 	}
 
-	return
+	FGLTFPropertyBakeOutput PropertyBakeOutput(MaterialProperty, PF_B8G8R8A8, BakedPixels, BakedSize, EmissiveScale);
+
+	if (BakedPixels.Num() == 1)
 	{
-		MaterialProperty,
-		PF_B8G8R8A8,
-		BakedPixels,
-		BakedSize,
-		EmissiveScale
-	};
+		const FColor& Pixel = BakedPixels[0];
+
+		PropertyBakeOutput.bIsConstant = true;
+		PropertyBakeOutput.ConstantValue = Pixel;
+	}
+
+	return PropertyBakeOutput;
 }
 
 UTexture2D* FGLTFMaterialUtility::BakeMaterialPropertyToTexture(const FIntPoint& OutputSize, EMaterialProperty MaterialProperty, const UMaterialInterface* Material, bool bCopyAlphaFromRedChannel, bool bUseSRGB)
diff --git a/Engine/Plugins/Enterprise/GLTFExporter/MigratedFromGit/Source/GLTFExporter/Private/Converters/GLTFMaterialUtility.h b/Engine/Plugins/Enterprise/GLTFExporter/MigratedFromGit/Source/GLTFExporter/Private/Converters/GLTFMaterialUtility.h
--- a/Engine/Plugins/Enterprise/GLTFExporter/MigratedFromGit/Source/GLTFExporter/Private/Converters/GLTFMaterialUtility.h
+++ b/Engine/Plugins/Enterprise/GLTFExporter/MigratedFromGit/Source/GLTFExporter/Private/Converters/GLTFMaterialUtility.h
@@ -19,7 +19,7 @@ struct FGLTFTextureCombineSource
 struct FGLTFPropertyBakeOutput
 {
 	FORCEINLINE FGLTFPropertyBakeOutput(EMaterialProperty Property, EPixelFormat PixelFormat, TArray<FColor>& Pixels, FIntPoint Size, float EmissiveScale)
-		: Property(Property), PixelFormat(PixelFormat), Pixels(Pixels), Size(Size), EmissiveScale(EmissiveScale)
+		: Property(Property), PixelFormat(PixelFormat), Pixels(Pixels), Size(Size), EmissiveScale(EmissiveScale), bIsConstant(false)
 	{}
 
 	EMaterialProperty Property;
@@ -27,6 +27,9 @@ struct FGLTFPropertyBakeOutput
 	TArray<FColor> Pixels;
 	FIntPoint Size;
 	float EmissiveScale;
+
+	bool bIsConstant;
+	FColor ConstantValue;
 };
 
 struct FGLTFMaterialUtility
