[GLTFExporter] feat(material): store baked base-color + opacity and metallic + roughness as factors when possible

commit c753f236cff9fecdda12ecda37f8ad178615fdab
Date: Wed, 10 Jun 2020 15:52:50 +0200

    feat(material): store baked base-color + opacity and metallic + roughness as factors when possible
    (cherry picked from commit c1e0f8888a0c18edea355d84d895354f9ed562ee​)

#skipci

diff --git a/Engine/Plugins/Enterprise/GLTFExporter/MigratedFromGit/Source/GLTFExporter/Private/Converters/GLTFMaterialConverters.cpp b/Engine/Plugins/Enterprise/GLTFExporter/MigratedFromGit/Source/GLTFExporter/Private/Converters/GLTFMaterialConverters.cpp
--- a/Engine/Plugins/Enterprise/GLTFExporter/MigratedFromGit/Source/GLTFExporter/Private/Converters/GLTFMaterialConverters.cpp
+++ b/Engine/Plugins/Enterprise/GLTFExporter/MigratedFromGit/Source/GLTFExporter/Private/Converters/GLTFMaterialConverters.cpp
@@ -209,20 +209,22 @@ bool FGLTFMaterialConverter::TryGetBaseColorAndOpacity(FGLTFConvertBuilder& Buil
 	}
 
 	const EMaterialProperty OpacityProperty = MaterialInterface->GetBlendMode() == BLEND_Masked ? MP_OpacityMask : MP_Opacity;
+	const FGLTFPropertyBakeOutput BaseColorBakeOutput = BakeMaterialProperty(MP_BaseColor, MaterialInterface, &TextureSize);
+	FGLTFPropertyBakeOutput OpacityBakeOutput = BakeMaterialProperty(OpacityProperty, MaterialInterface, &TextureSize, true);
 
-	BaseColorTexture = FGLTFMaterialUtility::BakeMaterialPropertyToTexture(TextureSize, MP_BaseColor, MaterialInterface);
-	OpacityTexture = FGLTFMaterialUtility::BakeMaterialPropertyToTexture(TextureSize, OpacityProperty, MaterialInterface, true);
-
-	// NOTE: the baked textures may have a different (smaller) size than requested, so we update the
-	// texture-size to fit both textures without wasting too much space.
-	TextureSize =
+	// Detect when both baked properties are constants, which means we can avoid exporting a texture
+	if (BaseColorBakeOutput.bIsConstant && OpacityBakeOutput.bIsConstant)
 	{
-		FMath::Max(BaseColorTexture->GetSizeX(), OpacityTexture->GetSizeX()),
-		FMath::Max(BaseColorTexture->GetSizeY(), OpacityTexture->GetSizeY())
-	};
+		FLinearColor BaseColorFactor(BaseColorBakeOutput.ConstantValue);
+		BaseColorFactor.A = OpacityBakeOutput.ConstantValue.A;
 
-	// TODO: handle the case where TextureSize is 1x1. In this case, both properties are constants and we should
-	// extract the value of each property from the texture and use as-is instead of exporting a combined texture.
+		OutPBRParams.BaseColorFactor = FGLTFConverterUtility::ConvertColor(BaseColorFactor);
+		return true;
+	}
+
+	TextureSize = BaseColorBakeOutput.Size.ComponentMax(OpacityBakeOutput.Size);
+	BaseColorTexture = FGLTFMaterialUtility::CreateTransientTexture(BaseColorBakeOutput);
+	OpacityTexture = FGLTFMaterialUtility::CreateTransientTexture(OpacityBakeOutput);
 
 	const FString TextureName = MaterialInterface->GetName() + TEXT("_BaseColor");
 
@@ -260,7 +262,7 @@ bool FGLTFMaterialConverter::TryGetMetallicAndRoughness(FGLTFConvertBuilder& Bui
 		return true;
 	}
 
-	// NOTE: since we always bake the properties (for now) when atleast property is non-const, we need
+	// NOTE: since we always bake the properties (for now) when atleast one property is non-const, we need
 	// to reset the constant factors to their defaults. Otherwise the baked value of a constant property
 	// would be scaled with the factor, i.e a double scaling.
 	OutPBRParams.MetallicFactor = 1.0f;
@@ -345,19 +347,20 @@ bool FGLTFMaterialConverter::TryGetMetallicAndRoughness(FGLTFConvertBuilder& Bui
 		TextureMagFilter = FGLTFConverterUtility::ConvertMagFilter(RoughnessTexture->Filter, RoughnessTexture->LODGroup);
 	}
 
-	RoughnessTexture = FGLTFMaterialUtility::BakeMaterialPropertyToTexture(TextureSize, MP_Roughness, MaterialInterface);
-	MetallicTexture = FGLTFMaterialUtility::BakeMaterialPropertyToTexture(TextureSize, MP_Metallic, MaterialInterface);
+	const FGLTFPropertyBakeOutput RoughnessBakeOutput = BakeMaterialProperty(MP_Roughness, MaterialInterface, &TextureSize);
+	const FGLTFPropertyBakeOutput MetallicBakeOutput = BakeMaterialProperty(MP_Metallic, MaterialInterface, &TextureSize);
 
-	// NOTE: the baked textures may have a different (smaller) size than requested, so we update the
-	// texture-size to fit both textures without wasting too much space.
-	TextureSize =
+	// Detect when both baked properties are constants, which means we can use factors and avoid exporting a texture
+	if (RoughnessBakeOutput.bIsConstant && MetallicBakeOutput.bIsConstant)
 	{
-		FMath::Max(MetallicTexture->GetSizeX(), RoughnessTexture->GetSizeX()),
-		FMath::Max(MetallicTexture->GetSizeY(), RoughnessTexture->GetSizeY())
-	};
+		OutPBRParams.RoughnessFactor = RoughnessBakeOutput.ConstantValue.R;
+		OutPBRParams.MetallicFactor = MetallicBakeOutput.ConstantValue.R;
+		return true;
+	}
 
-	// TODO: handle the case where TextureSize is 1x1. In this case, both properties are constants and we should
-	// extract the value of each property from the texture and use as-is instead of exporting a combined texture.
+	TextureSize = RoughnessBakeOutput.Size.ComponentMax(MetallicBakeOutput.Size);
+	RoughnessTexture = FGLTFMaterialUtility::CreateTransientTexture(RoughnessBakeOutput);
+	MetallicTexture = FGLTFMaterialUtility::CreateTransientTexture(MetallicBakeOutput);
 
 	const FString TextureName = MaterialInterface->GetName() + TEXT("_MetallicRoughness");
 
@@ -669,10 +672,7 @@ bool FGLTFMaterialConverter::TryGetBakedMaterialProperty(FGLTFConvertBuilder& Bu
 
 	if (PropertyBakeOutput.bIsConstant)
 	{
-		// TODO: is this conversion between FColor & FLinearColor correct with regards to sRGB etc?
-		const FLinearColor ConstantValue = PropertyBakeOutput.ConstantValue;
-
-		OutConstant = FGLTFConverterUtility::ConvertColor(ConstantValue);
+		OutConstant = FGLTFConverterUtility::ConvertColor(PropertyBakeOutput.ConstantValue);
 		return true;
 	}
 	else
@@ -687,10 +687,7 @@ bool FGLTFMaterialConverter::TryGetBakedMaterialProperty(FGLTFConvertBuilder& Bu
 
 	if (PropertyBakeOutput.bIsConstant)
 	{
-		// TODO: is this conversion between FColor & FLinearColor correct with regards to sRGB etc?
-		const FLinearColor ConstantValue = PropertyBakeOutput.ConstantValue;
-
-		OutConstant = FGLTFConverterUtility::ConvertColor(ConstantValue);
+		OutConstant = FGLTFConverterUtility::ConvertColor(PropertyBakeOutput.ConstantValue);
 		return true;
 	}
 	else
@@ -705,17 +702,20 @@ bool FGLTFMaterialConverter::TryGetBakedMaterialProperty(FGLTFConvertBuilder& Bu
 	return StoreBakedPropertyTexture(Builder, OutTexInfo, PropertyBakeOutput, MaterialProperty, MaterialInterface);
 }
 
-FGLTFPropertyBakeOutput FGLTFMaterialConverter::BakeMaterialProperty(EMaterialProperty MaterialProperty, const UMaterialInterface* MaterialInterface) const
+FGLTFPropertyBakeOutput FGLTFMaterialConverter::BakeMaterialProperty(EMaterialProperty MaterialProperty, const UMaterialInterface* MaterialInterface, const FIntPoint* PreferredTextureSize, bool bCopyAlphaFromRedChannel) const
 {
-	// TODO: handle cases where PropertyBakeOutput.EmissiveScale is not 1.0 (when baking EmissiveColor)
 	// TODO: make default baking-resolution configurable
+	const FIntPoint DefaultTextureSize(512, 512);
+	const FIntPoint TextureSize = PreferredTextureSize != nullptr ? *PreferredTextureSize : DefaultTextureSize;
+
+	// TODO: handle cases where PropertyBakeOutput.EmissiveScale is not 1.0 (when baking EmissiveColor)
 	// TODO: add support for calculating the ideal resolution to use for baking based on connected (texture) nodes
-	const FIntPoint TextureSize(512, 512);
 
 	const FGLTFPropertyBakeOutput PropertyBakeOutput = FGLTFMaterialUtility::BakeMaterialProperty(
 		TextureSize,
 		MaterialProperty,
-		MaterialInterface);
+		MaterialInterface,
+		bCopyAlphaFromRedChannel);
 
 	return PropertyBakeOutput;
 }
diff --git a/Engine/Plugins/Enterprise/GLTFExporter/MigratedFromGit/Source/GLTFExporter/Private/Converters/GLTFMaterialConverters.h b/Engine/Plugins/Enterprise/GLTFExporter/MigratedFromGit/Source/GLTFExporter/Private/Converters/GLTFMaterialConverters.h
--- a/Engine/Plugins/Enterprise/GLTFExporter/MigratedFromGit/Source/GLTFExporter/Private/Converters/GLTFMaterialConverters.h
+++ b/Engine/Plugins/Enterprise/GLTFExporter/MigratedFromGit/Source/GLTFExporter/Private/Converters/GLTFMaterialConverters.h
@@ -28,7 +28,7 @@ class FGLTFMaterialConverter final : public TGLTFConverter<FGLTFJsonMaterialInde
 	bool TryGetBakedMaterialProperty(FGLTFConvertBuilder& Builder, FGLTFJsonTextureInfo& OutTexInfo, FGLTFJsonColor4& OutConstant, EMaterialProperty MaterialProperty, const UMaterialInterface* MaterialInterface) const;
 	bool TryGetBakedMaterialProperty(FGLTFConvertBuilder& Builder, FGLTFJsonTextureInfo& OutTexInfo, EMaterialProperty MaterialProperty, const UMaterialInterface* MaterialInterface) const;
 
-	FGLTFPropertyBakeOutput BakeMaterialProperty(EMaterialProperty MaterialProperty, const UMaterialInterface* MaterialInterface) const;
+	FGLTFPropertyBakeOutput BakeMaterialProperty(EMaterialProperty MaterialProperty, const UMaterialInterface* MaterialInterface, const FIntPoint* PreferredTextureSize = nullptr, bool bCopyAlphaFromRedChannel = false) const;
 
 	bool StoreBakedPropertyTexture(FGLTFConvertBuilder& Builder, FGLTFJsonTextureInfo& OutTexInfo, const FGLTFPropertyBakeOutput& PropertyBakeOutput, EMaterialProperty MaterialProperty, const UMaterialInterface* MaterialInterface) const;
 };
diff --git a/Engine/Plugins/Enterprise/GLTFExporter/MigratedFromGit/Source/GLTFExporter/Private/Converters/GLTFMaterialUtility.cpp b/Engine/Plugins/Enterprise/GLTFExporter/MigratedFromGit/Source/GLTFExporter/Private/Converters/GLTFMaterialUtility.cpp
--- a/Engine/Plugins/Enterprise/GLTFExporter/MigratedFromGit/Source/GLTFExporter/Private/Converters/GLTFMaterialUtility.cpp
+++ b/Engine/Plugins/Enterprise/GLTFExporter/MigratedFromGit/Source/GLTFExporter/Private/Converters/GLTFMaterialUtility.cpp
@@ -10,6 +10,15 @@
 #include "IMaterialBakingModule.h"
 #include "MaterialBakingStructures.h"
 
+UTexture2D* FGLTFMaterialUtility::CreateTransientTexture(const FGLTFPropertyBakeOutput& PropertyBakeOutput, bool bUseSRGB)
+{
+	return CreateTransientTexture(
+		PropertyBakeOutput.Pixels,
+		PropertyBakeOutput.Size,
+		PropertyBakeOutput.PixelFormat,
+		bUseSRGB);
+}
+
 UTexture2D* FGLTFMaterialUtility::CreateTransientTexture(const TArray<FColor>& Pixels, const FIntPoint& TextureSize, EPixelFormat TextureFormat, bool bUseSRGB)
 {
 	check(TextureSize.X * TextureSize.Y == Pixels.Num());
@@ -124,18 +133,14 @@ FGLTFPropertyBakeOutput FGLTFMaterialUtility::BakeMaterialProperty(const FIntPoi
 		const FColor& Pixel = BakedPixels[0];
 
 		PropertyBakeOutput.bIsConstant = true;
+
+		// TODO: is the current conversion from sRGB => linear correct?
 		PropertyBakeOutput.ConstantValue = Pixel;
 	}
 
 	return PropertyBakeOutput;
 }
 
-UTexture2D* FGLTFMaterialUtility::BakeMaterialPropertyToTexture(const FIntPoint& OutputSize, EMaterialProperty MaterialProperty, const UMaterialInterface* Material, bool bCopyAlphaFromRedChannel, bool bUseSRGB)
-{
-	const FGLTFPropertyBakeOutput PropertyBakeOutput = BakeMaterialProperty(OutputSize, MaterialProperty, Material, bCopyAlphaFromRedChannel);
-	return CreateTransientTexture(PropertyBakeOutput.Pixels, PropertyBakeOutput.Size, PropertyBakeOutput.PixelFormat, bUseSRGB);
-}
-
 FGLTFJsonTextureIndex FGLTFMaterialUtility::AddCombinedTexture(FGLTFConvertBuilder& Builder, const TArray<FGLTFTextureCombineSource>& CombineSources, const FIntPoint& TextureSize, const FString& TextureName, EGLTFJsonTextureFilter MinFilter, EGLTFJsonTextureFilter MagFilter, EGLTFJsonTextureWrap WrapS, EGLTFJsonTextureWrap WrapT)
 {
 	check(CombineSources.Num() > 0);
diff --git a/Engine/Plugins/Enterprise/GLTFExporter/MigratedFromGit/Source/GLTFExporter/Private/Converters/GLTFMaterialUtility.h b/Engine/Plugins/Enterprise/GLTFExporter/MigratedFromGit/Source/GLTFExporter/Private/Converters/GLTFMaterialUtility.h
--- a/Engine/Plugins/Enterprise/GLTFExporter/MigratedFromGit/Source/GLTFExporter/Private/Converters/GLTFMaterialUtility.h
+++ b/Engine/Plugins/Enterprise/GLTFExporter/MigratedFromGit/Source/GLTFExporter/Private/Converters/GLTFMaterialUtility.h
@@ -29,15 +29,15 @@ struct FGLTFPropertyBakeOutput
 	float EmissiveScale;
 
 	bool bIsConstant;
-	FColor ConstantValue;
+	FLinearColor ConstantValue;
 };
 
 struct FGLTFMaterialUtility
 {
+	static UTexture2D* CreateTransientTexture(const FGLTFPropertyBakeOutput& PropertyBakeOutput, bool bUseSRGB = false);
 	static UTexture2D* CreateTransientTexture(const TArray<FColor>& Pixels, const FIntPoint& TextureSize, EPixelFormat TextureFormat, bool bUseSRGB = false);
 	static bool CombineTextures(TArray<FColor>& OutPixels, const TArray<FGLTFTextureCombineSource>& Sources, const FIntPoint& OutputSize, EPixelFormat OutputPixelFormat);
 	static FGLTFPropertyBakeOutput BakeMaterialProperty(const FIntPoint& OutputSize, EMaterialProperty MaterialProperty, const UMaterialInterface* Material, bool bCopyAlphaFromRedChannel = false);
-	static UTexture2D* BakeMaterialPropertyToTexture(const FIntPoint& OutputSize, EMaterialProperty MaterialProperty, const UMaterialInterface* Material, bool bCopyAlphaFromRedChannel = false, bool bUseSRGB = false);
 
 	static FGLTFJsonTextureIndex AddCombinedTexture(FGLTFConvertBuilder& Builder, const TArray<FGLTFTextureCombineSource>& CombineSources, const FIntPoint& TextureSize, const FString& TextureName, EGLTFJsonTextureFilter MinFilter, EGLTFJsonTextureFilter MagFilter, EGLTFJsonTextureWrap WrapS, EGLTFJsonTextureWrap WrapT);
 	static FGLTFJsonTextureIndex AddTexture(FGLTFConvertBuilder& Builder, const TArray<FColor>& Pixels, const FIntPoint& TextureSize, const FString& TextureName, EPixelFormat PixelFormat, EGLTFJsonTextureFilter MinFilter, EGLTFJsonTextureFilter MagFilter, EGLTFJsonTextureWrap WrapS, EGLTFJsonTextureWrap WrapT);
