[GLTFExporter] feat(material): use Repeat instead of ClampToEdge as default wrap for baked textures

commit eca5d8460c1041824bdc1d4eabdb383fb8426199
Date: Wed, 10 Jun 2020 16:03:47 +0200

    feat(material): use Repeat instead of ClampToEdge as default wrap for baked textures
    (cherry picked from commit f63fe24e6349eb1b6a269c803ca3da4caae5fffe​)

#skipci

diff --git a/Engine/Plugins/Enterprise/GLTFExporter/MigratedFromGit/Source/GLTFExporter/Private/Converters/GLTFMaterialConverters.cpp b/Engine/Plugins/Enterprise/GLTFExporter/MigratedFromGit/Source/GLTFExporter/Private/Converters/GLTFMaterialConverters.cpp
--- a/Engine/Plugins/Enterprise/GLTFExporter/MigratedFromGit/Source/GLTFExporter/Private/Converters/GLTFMaterialConverters.cpp
+++ b/Engine/Plugins/Enterprise/GLTFExporter/MigratedFromGit/Source/GLTFExporter/Private/Converters/GLTFMaterialConverters.cpp
@@ -157,8 +157,11 @@ bool FGLTFMaterialConverter::TryGetBaseColorAndOpacity(FGLTFConvertBuilder& Buil
 	int32 TexCoord = 0;
 	FIntPoint TextureSize(512, 512);
 
-	EGLTFJsonTextureWrap TextureWrapS = EGLTFJsonTextureWrap::ClampToEdge;
-	EGLTFJsonTextureWrap TextureWrapT = EGLTFJsonTextureWrap::ClampToEdge;
+	// TODO: should this be the default wrap-mode?
+	EGLTFJsonTextureWrap TextureWrapS = EGLTFJsonTextureWrap::Repeat;
+	EGLTFJsonTextureWrap TextureWrapT = EGLTFJsonTextureWrap::Repeat;
+
+	// TODO: should this be the default filter?
 	EGLTFJsonTextureFilter TextureMinFilter = EGLTFJsonTextureFilter::LinearMipmapLinear;
 	EGLTFJsonTextureFilter TextureMagFilter = EGLTFJsonTextureFilter::Linear;
 
@@ -296,8 +299,11 @@ bool FGLTFMaterialConverter::TryGetMetallicAndRoughness(FGLTFConvertBuilder& Bui
 	int32 TexCoord = 0;
 	FIntPoint TextureSize(512, 512);
 
-	EGLTFJsonTextureWrap TextureWrapS = EGLTFJsonTextureWrap::ClampToEdge;
-	EGLTFJsonTextureWrap TextureWrapT = EGLTFJsonTextureWrap::ClampToEdge;
+	// TODO: should this be the default wrap-mode?
+	EGLTFJsonTextureWrap TextureWrapS = EGLTFJsonTextureWrap::Repeat;
+	EGLTFJsonTextureWrap TextureWrapT = EGLTFJsonTextureWrap::Repeat;
+
+	// TODO: should this be the default filter?
 	EGLTFJsonTextureFilter TextureMinFilter = EGLTFJsonTextureFilter::LinearMipmapLinear;
 	EGLTFJsonTextureFilter TextureMagFilter = EGLTFJsonTextureFilter::Linear;
 
@@ -742,9 +748,11 @@ bool FGLTFMaterialConverter::StoreBakedPropertyTexture(FGLTFConvertBuilder& Buil
 	// TODO: add support for detecting the correct tex-coord for this property based on connected nodes
 	const uint32 TexCoord = 0;
 
-	// TODO: support for other wrapping / filters?
-	const EGLTFJsonTextureWrap TextureWrapS = EGLTFJsonTextureWrap::ClampToEdge;
-	const EGLTFJsonTextureWrap TextureWrapT = EGLTFJsonTextureWrap::ClampToEdge;
+	// TODO: should this be the default wrap-mode?
+	const EGLTFJsonTextureWrap TextureWrapS = EGLTFJsonTextureWrap::Repeat;
+	const EGLTFJsonTextureWrap TextureWrapT = EGLTFJsonTextureWrap::Repeat;
+
+	// TODO: should this be the default filter?
 	const EGLTFJsonTextureFilter TextureMinFilter = EGLTFJsonTextureFilter::LinearMipmapLinear;
 	const EGLTFJsonTextureFilter TextureMagFilter = EGLTFJsonTextureFilter::Linear;
 
