[GLTFExporter] refactor(json): use GLTFJsonUtility::WriteObjectArray in GLTFJsonMesh

commit fce08ba0eafb94623d6b30b116184bbc46086110
Date: Tue, 9 Jun 2020 01:53:10 +0200

    refactor(json): use GLTFJsonUtility::WriteObjectArray in GLTFJsonMesh
    (cherry picked from commit a26005abd68e4e95b8704b829f99cfcf81ae68c4​)

#skipci

diff --git a/Engine/Plugins/Enterprise/GLTFExporter/MigratedFromGit/Source/GLTFExporter/Private/Json/GLTFJsonMaterial.h b/Engine/Plugins/Enterprise/GLTFExporter/MigratedFromGit/Source/GLTFExporter/Private/Json/GLTFJsonMaterial.h
--- a/Engine/Plugins/Enterprise/GLTFExporter/MigratedFromGit/Source/GLTFExporter/Private/Json/GLTFJsonMaterial.h
+++ b/Engine/Plugins/Enterprise/GLTFExporter/MigratedFromGit/Source/GLTFExporter/Private/Json/GLTFJsonMaterial.h
@@ -2,10 +2,10 @@
 
 #pragma once
 
-#include "GLTFJsonUtility.h"
 #include "Json/GLTFJsonEnums.h"
 #include "Json/GLTFJsonColor3.h"
 #include "Json/GLTFJsonColor4.h"
+#include "Json/GLTFJsonUtility.h"
 #include "Serialization/JsonSerializer.h"
 
 struct FGLTFJsonTextureInfo
diff --git a/Engine/Plugins/Enterprise/GLTFExporter/MigratedFromGit/Source/GLTFExporter/Private/Json/GLTFJsonMesh.h b/Engine/Plugins/Enterprise/GLTFExporter/MigratedFromGit/Source/GLTFExporter/Private/Json/GLTFJsonMesh.h
--- a/Engine/Plugins/Enterprise/GLTFExporter/MigratedFromGit/Source/GLTFExporter/Private/Json/GLTFJsonMesh.h
+++ b/Engine/Plugins/Enterprise/GLTFExporter/MigratedFromGit/Source/GLTFExporter/Private/Json/GLTFJsonMesh.h
@@ -96,12 +96,7 @@ struct FGLTFJsonMesh
 			JsonWriter.WriteValue(TEXT("name"), Name);
 		}
 
-		JsonWriter.WriteArrayStart(TEXT("primitives"));
-		for (const FGLTFJsonPrimitive& Primitive : Primitives)
-		{
-			Primitive.WriteObject(JsonWriter);
-		}
-		JsonWriter.WriteArrayEnd();
+		FGLTFJsonUtility::WriteObjectArray(JsonWriter, TEXT("primitives"), Primitives, true);
 
 		JsonWriter.WriteObjectEnd();
 	}
diff --git a/Engine/Plugins/Enterprise/GLTFExporter/MigratedFromGit/Source/GLTFExporter/Private/Json/GLTFJsonUtility.h b/Engine/Plugins/Enterprise/GLTFExporter/MigratedFromGit/Source/GLTFExporter/Private/Json/GLTFJsonUtility.h
--- a/Engine/Plugins/Enterprise/GLTFExporter/MigratedFromGit/Source/GLTFExporter/Private/Json/GLTFJsonUtility.h
+++ b/Engine/Plugins/Enterprise/GLTFExporter/MigratedFromGit/Source/GLTFExporter/Private/Json/GLTFJsonUtility.h
@@ -75,9 +75,9 @@ struct FGLTFJsonUtility
 	}
 
 	template <class WriterType, class ElementType, class AllocatorType>
-	static void WriteObjectArray(WriterType& JsonWriter, const FString& Identifier, const TArray<ElementType, AllocatorType>& ObjectArray)
+	static void WriteObjectArray(WriterType& JsonWriter, const FString& Identifier, const TArray<ElementType, AllocatorType>& ObjectArray, bool bWriteIfEmpty = false)
 	{
-		if (ObjectArray.Num() > 0)
+		if (ObjectArray.Num() > 0 || bWriteIfEmpty)
 		{
 			JsonWriter.WriteArrayStart(Identifier);
 			for (const ElementType& Object : ObjectArray)
