[GLTFExporter] chore: update GLTFExporter.uplugin descriptor to match what PackagePlugin previously generated.

commit f2eea38a90463ea0a73144ba64835a953316b5c1
Date: Wed, 10 Jun 2020 20:38:47 +0200

    chore: update GLTFExporter.uplugin descriptor to match what PackagePlugin previously generated.
    (cherry picked from commit 67064dd15e96b05545c9efb9d0b3c27fcb7363ca​)

#skipci

diff --git a/Engine/Plugins/Enterprise/GLTFExporter/MigratedFromGit/GLTFExporter.uplugin b/Engine/Plugins/Enterprise/GLTFExporter/MigratedFromGit/GLTFExporter.uplugin
--- a/Engine/Plugins/Enterprise/GLTFExporter/MigratedFromGit/GLTFExporter.uplugin
+++ b/Engine/Plugins/Enterprise/GLTFExporter/MigratedFromGit/GLTFExporter.uplugin
@@ -1,20 +1,24 @@
 {
 	"FileVersion": 3,
-	"Version": 0.1,
+	"Version": 0,
 	"VersionName": "0.1",
 	"FriendlyName": "glTF Exporter",
 	"Description": "An exporter for Khronos glTF 3D models.",
 	"Category": "Exporters",
 	"CreatedBy": "Epic Games, Inc.",
 	"CreatedByURL": "http://epicgames.com",
-	"EnabledByDefault": false,
+	"DocsURL": "",
+	"MarketplaceURL": "",
+	"SupportURL": "",
+	"EngineVersion": "4.25.0",
 	"CanContainContent": true,
 	"IsBetaVersion": true,
 	"Modules":
 	[
 		{
 			"Name": "GLTFExporter",
-			"Type": "Editor"
+			"Type": "Editor",
+			"LoadingPhase": "Default"
 		}
 	]
 }
