[GLTFExporter] feat(material): discard constant values (1x1 textures) when trying to bake to a texture

commit 803eb3ef1b0313f2930e332992ba865e6f548990
Date: Thu, 11 Jun 2020 09:06:27 +0200

    feat(material): discard constant values (1x1 textures) when trying to bake to a texture
    (cherry picked from commit d158bcac8ae085b3fbef2124cf53e467014f9c22​)

#skipci

diff --git a/Engine/Plugins/Enterprise/GLTFExporter/MigratedFromGit/Source/GLTFExporter/Private/Converters/GLTFMaterialConverters.cpp b/Engine/Plugins/Enterprise/GLTFExporter/MigratedFromGit/Source/GLTFExporter/Private/Converters/GLTFMaterialConverters.cpp
--- a/Engine/Plugins/Enterprise/GLTFExporter/MigratedFromGit/Source/GLTFExporter/Private/Converters/GLTFMaterialConverters.cpp
+++ b/Engine/Plugins/Enterprise/GLTFExporter/MigratedFromGit/Source/GLTFExporter/Private/Converters/GLTFMaterialConverters.cpp
@@ -705,7 +705,26 @@ bool FGLTFMaterialConverter::TryGetBakedMaterialProperty(FGLTFConvertBuilder& Bu
 bool FGLTFMaterialConverter::TryGetBakedMaterialProperty(FGLTFConvertBuilder& Builder, FGLTFJsonTextureInfo& OutTexInfo, EMaterialProperty MaterialProperty, const UMaterialInterface* MaterialInterface) const
 {
 	const FGLTFPropertyBakeOutput PropertyBakeOutput = BakeMaterialProperty(MaterialProperty, MaterialInterface);
-	return StoreBakedPropertyTexture(Builder, OutTexInfo, PropertyBakeOutput, MaterialProperty, MaterialInterface);
+
+	if (!PropertyBakeOutput.bIsConstant)
+	{
+		return StoreBakedPropertyTexture(Builder, OutTexInfo, PropertyBakeOutput, MaterialProperty, MaterialInterface);
+	}
+	else
+	{
+		// NOTE: since this function is meant to bake to a texture, we assume that the property
+		// that was passed into the function is using a non-constant expression.
+		// We therefore treat a constant result when baking as a failure.
+
+		// NOTE: in some cases a constant baking result is returned for a property that is non-constant.
+		// This happens (for example) when baking AmbientOcclusion for a translucent material,
+		// even though the same material when set to opaque will properly bake AmbientOcclusion to a texture.
+		// It also happens when baking Normal in some (not yet identified) circumstances.
+		// For now, these incorrect bakes are discarded.
+
+		// TODO: investigate why non-constant properties are sometimes baked to a constant expression (see note above)
+		return false;
+	}
 }
 
 FGLTFPropertyBakeOutput FGLTFMaterialConverter::BakeMaterialProperty(EMaterialProperty MaterialProperty, const UMaterialInterface* MaterialInterface, const FIntPoint* PreferredTextureSize, bool bCopyAlphaFromRedChannel) const
@@ -728,8 +747,6 @@ FGLTFPropertyBakeOutput FGLTFMaterialConverter::BakeMaterialProperty(EMaterialPr
 
 bool FGLTFMaterialConverter::StoreBakedPropertyTexture(FGLTFConvertBuilder& Builder, FGLTFJsonTextureInfo& OutTexInfo, const FGLTFPropertyBakeOutput& PropertyBakeOutput, EMaterialProperty MaterialProperty, const UMaterialInterface* MaterialInterface) const
 {
-	// TODO: right now a constant property can be saved as a texture. Should we return false instead?
-
 	TCHAR* PropertyName;
 
 	switch (MaterialProperty)
