[GLTFExporter] feat(material): use full alpha (except for when it's used as value) in baked textures

commit f94433961d790a9e8d00d8ba02b103585004cd59
Date: Thu, 11 Jun 2020 09:47:08 +0200

    feat(material): use full alpha (except for when it's used as value) in baked textures
    makes the textures viewable without using image-manipulation software that can isolate channels
    (cherry picked from commit 7f8b8a355812f017cd0595ec784e6787d40d220b​)

#skipci

diff --git a/Engine/Plugins/Enterprise/GLTFExporter/MigratedFromGit/Source/GLTFExporter/Private/Converters/GLTFMaterialConverters.cpp b/Engine/Plugins/Enterprise/GLTFExporter/MigratedFromGit/Source/GLTFExporter/Private/Converters/GLTFMaterialConverters.cpp
--- a/Engine/Plugins/Enterprise/GLTFExporter/MigratedFromGit/Source/GLTFExporter/Private/Converters/GLTFMaterialConverters.cpp
+++ b/Engine/Plugins/Enterprise/GLTFExporter/MigratedFromGit/Source/GLTFExporter/Private/Converters/GLTFMaterialConverters.cpp
@@ -278,6 +278,7 @@ bool FGLTFMaterialConverter::TryGetMetallicAndRoughness(FGLTFConvertBuilder& Bui
 
 	const FLinearColor MetallicMask(0.0f, 0.0f, 1.0f, 0.0f);
 	const FLinearColor RoughnessMask(0.0f, 1.0f, 0.0f, 0.0f);
+	const FLinearColor AlphaMask(0.0f, 0.0f, 0.0f, 1.0f);
 
 	const bool bHasMetallicSourceTexture = TryGetSourceTexture(MetallicTexture, MetallicTexCoord, MetallicInput, MaterialInstance, { MetallicMask });
 	const bool bHasRoughnessSourceTexture = TryGetSourceTexture(RoughnessTexture, RoughnessTexCoord, RoughnessInput, MaterialInstance, { RoughnessMask });
@@ -372,7 +373,7 @@ bool FGLTFMaterialConverter::TryGetMetallicAndRoughness(FGLTFConvertBuilder& Bui
 
 	const TArray<FGLTFTextureCombineSource> CombineSources =
 	{
-		{ MetallicTexture, MetallicMask },
+		{ MetallicTexture, MetallicMask + AlphaMask, SE_BLEND_Opaque },
 		{ RoughnessTexture, RoughnessMask }
 	};
 
diff --git a/Engine/Plugins/Enterprise/GLTFExporter/MigratedFromGit/Source/GLTFExporter/Private/Converters/GLTFMaterialUtility.cpp b/Engine/Plugins/Enterprise/GLTFExporter/MigratedFromGit/Source/GLTFExporter/Private/Converters/GLTFMaterialUtility.cpp
--- a/Engine/Plugins/Enterprise/GLTFExporter/MigratedFromGit/Source/GLTFExporter/Private/Converters/GLTFMaterialUtility.cpp
+++ b/Engine/Plugins/Enterprise/GLTFExporter/MigratedFromGit/Source/GLTFExporter/Private/Converters/GLTFMaterialUtility.cpp
@@ -113,18 +113,20 @@ FGLTFPropertyBakeOutput FGLTFMaterialUtility::BakeMaterialProperty(const FIntPoi
 
 	if (bCopyAlphaFromRedChannel)
 	{
-		// NOTE: Baked property-textures only have values in RGB, not A.
-		// If someone wants to use the alpha-channel of the baked texture (for compositing),
-		// it needs to be filled with values first.
-		// Since scalar properties are baked to R + G + B, it doesn't matter which channel
-		// we use as source. For now it's set to R.
-		// TODO: Determine if we can get the baking module to somehow output values in A automatically.
-		// If not, see if the copy-operation can be performed via the canvas (using a proxy-material or similar).
 		for (FColor& Pixel: BakedPixels)
 		{
 			Pixel.A = Pixel.R;
 		}
 	}
+	else
+	{
+		// NOTE: alpha is 0 by default after baking a property, but we prefer 255 (1.0).
+		// It makes it easier to view the exported textures.
+		for (FColor& Pixel: BakedPixels)
+		{
+			Pixel.A = 255;
+		}
+	}
 
 	FGLTFPropertyBakeOutput PropertyBakeOutput(MaterialProperty, PF_B8G8R8A8, BakedPixels, BakedSize, EmissiveScale);
 
