[GLTFExporter] refactor(options): update default values to match current hardcoded settings

commit cc6f9d4cb1f806a5923adbf2985d1a3df13884f3
Date: Thu, 11 Jun 2020 10:34:34 +0200

    refactor(options): update default values to match current hardcoded settings
    (cherry picked from commit 71f754a101f2b43bbef9133a6da1c6073e2a49f5​)

#skipci

diff --git a/Engine/Plugins/Enterprise/GLTFExporter/MigratedFromGit/Source/GLTFExporter/Private/GLTFExportOptions.cpp b/Engine/Plugins/Enterprise/GLTFExporter/MigratedFromGit/Source/GLTFExporter/Private/GLTFExportOptions.cpp
--- a/Engine/Plugins/Enterprise/GLTFExporter/MigratedFromGit/Source/GLTFExporter/Private/GLTFExportOptions.cpp
+++ b/Engine/Plugins/Enterprise/GLTFExporter/MigratedFromGit/Source/GLTFExporter/Private/GLTFExportOptions.cpp
@@ -26,8 +26,8 @@ UGLTFExportOptions::UGLTFExportOptions(const FObjectInitializer& ObjectInitializ
 	bMapSkeletalMotionToRoot = false;
 	bExportUnlitMaterials = false;
 	bExportClearCoatMaterials = false;
-	bBakeMaterialInputs = false;
-	DefaultMaterialBakeSize = EGLTFExporterMaterialBakeSize::POT_1024;
+	bBakeMaterialInputs = true;
+	DefaultMaterialBakeSize = EGLTFExporterMaterialBakeSize::POT_512;
 	TextureFormat = EGLTFExporterTextureFormat::None;
 	bExportLightmaps = false;
 	TextureHDREncoding = EGLTFExporterTextureHDREncoding::None;
