[GLTFExporter] refactor: minor code cleanup

commit 9d64ab9941db2c665a0c0e4d799c4a5b60c12d22
Date: Thu, 11 Jun 2020 11:00:05 +0200

    refactor: minor code cleanup
    (cherry picked from commit 359efacefc0fcaa575c1b7fb5c5c0efd28a9e20b​)

#skipci

diff --git a/Engine/Plugins/Enterprise/GLTFExporter/MigratedFromGit/Source/GLTFExporter/Private/Builders/GLTFContainerBuilder.cpp b/Engine/Plugins/Enterprise/GLTFExporter/MigratedFromGit/Source/GLTFExporter/Private/Builders/GLTFContainerBuilder.cpp
--- a/Engine/Plugins/Enterprise/GLTFExporter/MigratedFromGit/Source/GLTFExporter/Private/Builders/GLTFContainerBuilder.cpp
+++ b/Engine/Plugins/Enterprise/GLTFExporter/MigratedFromGit/Source/GLTFExporter/Private/Builders/GLTFContainerBuilder.cpp
@@ -1,3 +1,8 @@
 // Copyright Epic Games, Inc. All Rights Reserved.
 
 #include "Builders/GLTFContainerBuilder.h"
+
+FGLTFContainerBuilder::FGLTFContainerBuilder(bool bSelectedActorsOnly)
+	: FGLTFConvertBuilder(bSelectedActorsOnly)
+{
+}
diff --git a/Engine/Plugins/Enterprise/GLTFExporter/MigratedFromGit/Source/GLTFExporter/Private/Builders/GLTFContainerBuilder.h b/Engine/Plugins/Enterprise/GLTFExporter/MigratedFromGit/Source/GLTFExporter/Private/Builders/GLTFContainerBuilder.h
--- a/Engine/Plugins/Enterprise/GLTFExporter/MigratedFromGit/Source/GLTFExporter/Private/Builders/GLTFContainerBuilder.h
+++ b/Engine/Plugins/Enterprise/GLTFExporter/MigratedFromGit/Source/GLTFExporter/Private/Builders/GLTFContainerBuilder.h
@@ -8,8 +8,5 @@ class FGLTFContainerBuilder : public FGLTFConvertBuilder
 {
 public:
 
-	FGLTFContainerBuilder(bool bSelectedActorsOnly)
-		: FGLTFConvertBuilder(bSelectedActorsOnly)
-	{
-	}
+	FGLTFContainerBuilder(bool bSelectedActorsOnly);
 };
diff --git a/Engine/Plugins/Enterprise/GLTFExporter/MigratedFromGit/Source/GLTFExporter/Private/Builders/GLTFConvertBuilder.cpp b/Engine/Plugins/Enterprise/GLTFExporter/MigratedFromGit/Source/GLTFExporter/Private/Builders/GLTFConvertBuilder.cpp
--- a/Engine/Plugins/Enterprise/GLTFExporter/MigratedFromGit/Source/GLTFExporter/Private/Builders/GLTFConvertBuilder.cpp
+++ b/Engine/Plugins/Enterprise/GLTFExporter/MigratedFromGit/Source/GLTFExporter/Private/Builders/GLTFConvertBuilder.cpp
@@ -3,6 +3,11 @@
 #include "Builders/GLTFConvertBuilder.h"
 #include "GLTFBuilderUtility.h"
 
+FGLTFConvertBuilder::FGLTFConvertBuilder(bool bSelectedActorsOnly)
+	: bSelectedActorsOnly(bSelectedActorsOnly)
+{
+}
+
 FGLTFJsonAccessorIndex FGLTFConvertBuilder::GetOrAddPositionAccessor(const FPositionVertexBuffer* VertexBuffer, const FString& DesiredName)
 {
 	if (VertexBuffer == nullptr)
diff --git a/Engine/Plugins/Enterprise/GLTFExporter/MigratedFromGit/Source/GLTFExporter/Private/Builders/GLTFConvertBuilder.h b/Engine/Plugins/Enterprise/GLTFExporter/MigratedFromGit/Source/GLTFExporter/Private/Builders/GLTFConvertBuilder.h
--- a/Engine/Plugins/Enterprise/GLTFExporter/MigratedFromGit/Source/GLTFExporter/Private/Builders/GLTFConvertBuilder.h
+++ b/Engine/Plugins/Enterprise/GLTFExporter/MigratedFromGit/Source/GLTFExporter/Private/Builders/GLTFConvertBuilder.h
@@ -15,10 +15,7 @@ public:
 
 	const bool bSelectedActorsOnly;
 
-	FGLTFConvertBuilder(bool bSelectedActorsOnly)
-		: bSelectedActorsOnly(bSelectedActorsOnly)
-	{
-	}
+	FGLTFConvertBuilder(bool bSelectedActorsOnly);
 
 	FGLTFJsonAccessorIndex GetOrAddPositionAccessor(const FPositionVertexBuffer* VertexBuffer, const FString& DesiredName = TEXT(""));
 	FGLTFJsonAccessorIndex GetOrAddColorAccessor(const FColorVertexBuffer* VertexBuffer, const FString& DesiredName = TEXT(""));
diff --git a/Engine/Plugins/Enterprise/GLTFExporter/MigratedFromGit/Source/GLTFExporter/Private/Json/GLTFJsonColor3.cpp b/Engine/Plugins/Enterprise/GLTFExporter/MigratedFromGit/Source/GLTFExporter/Private/Json/GLTFJsonColor3.cpp
--- a/Engine/Plugins/Enterprise/GLTFExporter/MigratedFromGit/Source/GLTFExporter/Private/Json/GLTFJsonColor3.cpp
+++ b/Engine/Plugins/Enterprise/GLTFExporter/MigratedFromGit/Source/GLTFExporter/Private/Json/GLTFJsonColor3.cpp
@@ -2,5 +2,5 @@
 
 #include "Json/GLTFJsonColor3.h"
 
-const FGLTFJsonColor3 FGLTFJsonColor3::Black(0, 0, 0);
-const FGLTFJsonColor3 FGLTFJsonColor3::White(1, 1, 1);
+const FGLTFJsonColor3 FGLTFJsonColor3::Black(0.0f, 0.0f, 0.0f);
+const FGLTFJsonColor3 FGLTFJsonColor3::White(1.0f, 1.0f, 1.0f);
diff --git a/Engine/Plugins/Enterprise/GLTFExporter/MigratedFromGit/Source/GLTFExporter/Private/Json/GLTFJsonColor4.cpp b/Engine/Plugins/Enterprise/GLTFExporter/MigratedFromGit/Source/GLTFExporter/Private/Json/GLTFJsonColor4.cpp
--- a/Engine/Plugins/Enterprise/GLTFExporter/MigratedFromGit/Source/GLTFExporter/Private/Json/GLTFJsonColor4.cpp
+++ b/Engine/Plugins/Enterprise/GLTFExporter/MigratedFromGit/Source/GLTFExporter/Private/Json/GLTFJsonColor4.cpp
@@ -2,5 +2,5 @@
 
 #include "Json/GLTFJsonColor4.h"
 
-const FGLTFJsonColor4 FGLTFJsonColor4::Black(0, 0, 0, 1);
-const FGLTFJsonColor4 FGLTFJsonColor4::White(1, 1, 1, 1);
+const FGLTFJsonColor4 FGLTFJsonColor4::Black(0.0f, 0.0f, 0.0f, 1.0f);
+const FGLTFJsonColor4 FGLTFJsonColor4::White(1.0f, 1.0f, 1.0f, 1.0f);
diff --git a/Engine/Plugins/Enterprise/GLTFExporter/MigratedFromGit/Source/GLTFExporter/Private/Json/GLTFJsonQuaternion.cpp b/Engine/Plugins/Enterprise/GLTFExporter/MigratedFromGit/Source/GLTFExporter/Private/Json/GLTFJsonQuaternion.cpp
--- a/Engine/Plugins/Enterprise/GLTFExporter/MigratedFromGit/Source/GLTFExporter/Private/Json/GLTFJsonQuaternion.cpp
+++ b/Engine/Plugins/Enterprise/GLTFExporter/MigratedFromGit/Source/GLTFExporter/Private/Json/GLTFJsonQuaternion.cpp
@@ -2,4 +2,4 @@
 
 #include "Json/GLTFJsonQuaternion.h"
 
-const FGLTFJsonQuaternion FGLTFJsonQuaternion::Identity(0, 0, 0, 1);
+const FGLTFJsonQuaternion FGLTFJsonQuaternion::Identity(0.0f, 0.0f, 0.0f, 1.0f);
diff --git a/Engine/Plugins/Enterprise/GLTFExporter/MigratedFromGit/Source/GLTFExporter/Private/Json/GLTFJsonVector2.cpp b/Engine/Plugins/Enterprise/GLTFExporter/MigratedFromGit/Source/GLTFExporter/Private/Json/GLTFJsonVector2.cpp
--- a/Engine/Plugins/Enterprise/GLTFExporter/MigratedFromGit/Source/GLTFExporter/Private/Json/GLTFJsonVector2.cpp
+++ b/Engine/Plugins/Enterprise/GLTFExporter/MigratedFromGit/Source/GLTFExporter/Private/Json/GLTFJsonVector2.cpp
@@ -2,5 +2,5 @@
 
 #include "Json/GLTFJsonVector2.h"
 
-const FGLTFJsonVector2 FGLTFJsonVector2::Zero(0, 0);
-const FGLTFJsonVector2 FGLTFJsonVector2::One(1, 1);
+const FGLTFJsonVector2 FGLTFJsonVector2::Zero(0.0f, 0.0f);
+const FGLTFJsonVector2 FGLTFJsonVector2::One(1.0f, 1.0f);
diff --git a/Engine/Plugins/Enterprise/GLTFExporter/MigratedFromGit/Source/GLTFExporter/Private/Json/GLTFJsonVector3.cpp b/Engine/Plugins/Enterprise/GLTFExporter/MigratedFromGit/Source/GLTFExporter/Private/Json/GLTFJsonVector3.cpp
--- a/Engine/Plugins/Enterprise/GLTFExporter/MigratedFromGit/Source/GLTFExporter/Private/Json/GLTFJsonVector3.cpp
+++ b/Engine/Plugins/Enterprise/GLTFExporter/MigratedFromGit/Source/GLTFExporter/Private/Json/GLTFJsonVector3.cpp
@@ -2,5 +2,5 @@
 
 #include "Json/GLTFJsonVector3.h"
 
-const FGLTFJsonVector3 FGLTFJsonVector3::Zero(0, 0, 0);
-const FGLTFJsonVector3 FGLTFJsonVector3::One(1, 1, 1);
+const FGLTFJsonVector3 FGLTFJsonVector3::Zero(0.0f, 0.0f, 0.0f);
+const FGLTFJsonVector3 FGLTFJsonVector3::One(1.0f, 1.0f, 1.0f);
diff --git a/Engine/Plugins/Enterprise/GLTFExporter/MigratedFromGit/Source/GLTFExporter/Private/Json/GLTFJsonVector4.cpp b/Engine/Plugins/Enterprise/GLTFExporter/MigratedFromGit/Source/GLTFExporter/Private/Json/GLTFJsonVector4.cpp
--- a/Engine/Plugins/Enterprise/GLTFExporter/MigratedFromGit/Source/GLTFExporter/Private/Json/GLTFJsonVector4.cpp
+++ b/Engine/Plugins/Enterprise/GLTFExporter/MigratedFromGit/Source/GLTFExporter/Private/Json/GLTFJsonVector4.cpp
@@ -2,5 +2,5 @@
 
 #include "Json/GLTFJsonVector4.h"
 
-const FGLTFJsonVector4 FGLTFJsonVector4::Zero(0, 0, 0, 0);
-const FGLTFJsonVector4 FGLTFJsonVector4::One(1, 1, 1, 1);
+const FGLTFJsonVector4 FGLTFJsonVector4::Zero(0.0f, 0.0f, 0.0f, 0.0f);
+const FGLTFJsonVector4 FGLTFJsonVector4::One(1.0f, 1.0f, 1.0f, 1.0f);
