[GLTFExporter] refactor(material): only cast MaterialInterface to MaterialInstance when needed.

commit 2db077f34c92676a235d7943bc6ab353b5b58adb
Date: Sun, 14 Jun 2020 02:50:44 +0200

    refactor(material): only cast MaterialInterface to MaterialInstance when needed.
    (cherry picked from commit 69f218f3d2f2177f1cff67316676f07bdc146f1b​)

#skipci

diff --git a/Engine/Plugins/Enterprise/GLTFExporter/MigratedFromGit/Source/GLTFExporter/Private/Converters/GLTFMaterialConverters.cpp b/Engine/Plugins/Enterprise/GLTFExporter/MigratedFromGit/Source/GLTFExporter/Private/Converters/GLTFMaterialConverters.cpp
--- a/Engine/Plugins/Enterprise/GLTFExporter/MigratedFromGit/Source/GLTFExporter/Private/Converters/GLTFMaterialConverters.cpp
+++ b/Engine/Plugins/Enterprise/GLTFExporter/MigratedFromGit/Source/GLTFExporter/Private/Converters/GLTFMaterialConverters.cpp
@@ -16,13 +16,12 @@
 FGLTFJsonMaterialIndex FGLTFMaterialConverter::Add(FGLTFConvertBuilder& Builder, const FString& Name, const UMaterialInterface* MaterialInterface)
 {
 	const UMaterial* Material = MaterialInterface->GetMaterial();
-	const UMaterialInstance* MaterialInstance = Cast<UMaterialInstance>(MaterialInterface);
 
 	FGLTFJsonMaterial JsonMaterial;
 	JsonMaterial.Name = Name;
 	// TODO: add support for different shading models (Default Lit, Unlit, Clear Coat)
 
-	// TODO: add support for additional blend modes (at least Additive and Modulate)
+	// TODO: add support for additional blend modes (like Additive and Modulate)?
 	JsonMaterial.AlphaMode = FGLTFConverterUtility::ConvertBlendMode(MaterialInterface->GetBlendMode());
 	JsonMaterial.AlphaCutoff = MaterialInterface->GetOpacityMaskClipValue();
 	JsonMaterial.DoubleSided = MaterialInterface->IsTwoSided();
@@ -41,9 +40,9 @@ FGLTFJsonMaterialIndex FGLTFMaterialConverter::Add(FGLTFConvertBuilder& Builder,
 	}
 	else
 	{
-		if (!TryGetConstantColor(JsonMaterial.PBRMetallicRoughness.BaseColorFactor, Material->BaseColor, MaterialInstance))
+		if (!TryGetConstantColor(JsonMaterial.PBRMetallicRoughness.BaseColorFactor, Material->BaseColor, MaterialInterface))
 		{
-			if (!TryGetSourceTexture(Builder, JsonMaterial.PBRMetallicRoughness.BaseColorTexture, Material->BaseColor, MaterialInstance, { RgbaMask, RgbMask }))
+			if (!TryGetSourceTexture(Builder, JsonMaterial.PBRMetallicRoughness.BaseColorTexture, Material->BaseColor, MaterialInterface, { RgbaMask, RgbMask }))
 			{
 				if (!TryGetBakedMaterialProperty(Builder, JsonMaterial.PBRMetallicRoughness.BaseColorTexture, JsonMaterial.PBRMetallicRoughness.BaseColorFactor, MP_BaseColor, MaterialInterface))
 				{
@@ -64,9 +63,9 @@ FGLTFJsonMaterialIndex FGLTFMaterialConverter::Add(FGLTFConvertBuilder& Builder,
 	// how the same materials look inside of UE.
 	// TODO: solve the issues in a separate MR.
 	/*
-	if (!TryGetConstantColor(JsonMaterial.EmissiveFactor, Material->EmissiveColor, MaterialInstance))
+	if (!TryGetConstantColor(JsonMaterial.EmissiveFactor, Material->EmissiveColor, MaterialInterface))
 	{
-		if (!TryGetSourceTexture(Builder, JsonMaterial.EmissiveTexture, Material->EmissiveColor, MaterialInstance, { RgbaMask, RgbMask }))
+		if (!TryGetSourceTexture(Builder, JsonMaterial.EmissiveTexture, Material->EmissiveColor, MaterialInterface, { RgbaMask, RgbMask }))
 		{
 			if (TryGetBakedMaterialProperty(Builder, JsonMaterial.EmissiveTexture, JsonMaterial.EmissiveFactor, MP_EmissiveColor, MaterialInterface))
 			{
@@ -89,7 +88,7 @@ FGLTFJsonMaterialIndex FGLTFMaterialConverter::Add(FGLTFConvertBuilder& Builder,
 
 	if (Material->Normal.Expression != nullptr)
 	{
-		if (!TryGetSourceTexture(Builder, JsonMaterial.NormalTexture, Material->Normal, MaterialInstance, { RgbaMask, RgbMask }))
+		if (!TryGetSourceTexture(Builder, JsonMaterial.NormalTexture, Material->Normal, MaterialInterface, { RgbaMask, RgbMask }))
 		{
 			if (!TryGetBakedMaterialProperty(Builder, JsonMaterial.NormalTexture, MP_Normal, MaterialInterface))
 			{
@@ -100,7 +99,7 @@ FGLTFJsonMaterialIndex FGLTFMaterialConverter::Add(FGLTFConvertBuilder& Builder,
 
 	if (Material->AmbientOcclusion.Expression != nullptr)
 	{
-		if (!TryGetSourceTexture(Builder, JsonMaterial.OcclusionTexture, Material->AmbientOcclusion, MaterialInstance, { RMask }))
+		if (!TryGetSourceTexture(Builder, JsonMaterial.OcclusionTexture, Material->AmbientOcclusion, MaterialInterface, { RMask }))
 		{
 			if (!TryGetBakedMaterialProperty(Builder, JsonMaterial.OcclusionTexture, MP_AmbientOcclusion, MaterialInterface))
 			{
@@ -114,10 +113,8 @@ FGLTFJsonMaterialIndex FGLTFMaterialConverter::Add(FGLTFConvertBuilder& Builder,
 
 bool FGLTFMaterialConverter::TryGetBaseColorAndOpacity(FGLTFConvertBuilder& Builder, FGLTFJsonPBRMetallicRoughness& OutPBRParams, const FColorMaterialInput& BaseColorInput, const FScalarMaterialInput& OpacityInput, const UMaterialInterface* MaterialInterface) const
 {
-	const UMaterialInstance* MaterialInstance = Cast<UMaterialInstance>(MaterialInterface);
-
-	const bool bIsBaseColorConstant = TryGetConstantColor(OutPBRParams.BaseColorFactor, BaseColorInput, MaterialInstance);
-	const bool bIsOpacityConstant = TryGetConstantScalar(OutPBRParams.BaseColorFactor.A, OpacityInput, MaterialInstance);
+	const bool bIsBaseColorConstant = TryGetConstantColor(OutPBRParams.BaseColorFactor, BaseColorInput, MaterialInterface);
+	const bool bIsOpacityConstant = TryGetConstantScalar(OutPBRParams.BaseColorFactor.A, OpacityInput, MaterialInterface);
 
 	if (bIsBaseColorConstant && bIsOpacityConstant)
 	{
@@ -137,8 +134,8 @@ bool FGLTFMaterialConverter::TryGetBaseColorAndOpacity(FGLTFConvertBuilder& Buil
 	const FLinearColor BaseColorMask(1.0f, 1.0f, 1.0f, 0.0f);
 	const FLinearColor OpacityMask(0.0f, 0.0f, 0.0f, 1.0f);
 
-	const bool bHasBaseColorSourceTexture = TryGetSourceTexture(BaseColorTexture, BaseColorTexCoord, BaseColorInput, MaterialInstance, { BaseColorMask });
-	const bool bHasOpacitySourceTexture = TryGetSourceTexture(OpacityTexture, OpacityTexCoord, OpacityInput, MaterialInstance, { OpacityMask });
+	const bool bHasBaseColorSourceTexture = TryGetSourceTexture(BaseColorTexture, BaseColorTexCoord, BaseColorInput, MaterialInterface, { BaseColorMask });
+	const bool bHasOpacitySourceTexture = TryGetSourceTexture(OpacityTexture, OpacityTexCoord, OpacityInput, MaterialInterface, { OpacityMask });
 
 	// Detect the "happy path" where both inputs share the same texture and are correctly masked.
 	if (bHasBaseColorSourceTexture &&
@@ -255,10 +252,8 @@ bool FGLTFMaterialConverter::TryGetBaseColorAndOpacity(FGLTFConvertBuilder& Buil
 
 bool FGLTFMaterialConverter::TryGetMetallicAndRoughness(FGLTFConvertBuilder& Builder, FGLTFJsonPBRMetallicRoughness& OutPBRParams, const FScalarMaterialInput& MetallicInput, const FScalarMaterialInput& RoughnessInput, const UMaterialInterface* MaterialInterface) const
 {
-	const UMaterialInstance* MaterialInstance = Cast<UMaterialInstance>(MaterialInterface);
-
-	const bool bIsMetallicConstant = TryGetConstantScalar(OutPBRParams.MetallicFactor, MetallicInput, MaterialInstance);
-	const bool bIsRoughnessConstant = TryGetConstantScalar(OutPBRParams.RoughnessFactor, RoughnessInput, MaterialInstance);
+	const bool bIsMetallicConstant = TryGetConstantScalar(OutPBRParams.MetallicFactor, MetallicInput, MaterialInterface);
+	const bool bIsRoughnessConstant = TryGetConstantScalar(OutPBRParams.RoughnessFactor, RoughnessInput, MaterialInterface);
 
 	if (bIsMetallicConstant && bIsRoughnessConstant)
 	{
@@ -280,8 +275,8 @@ bool FGLTFMaterialConverter::TryGetMetallicAndRoughness(FGLTFConvertBuilder& Bui
 	const FLinearColor RoughnessMask(0.0f, 1.0f, 0.0f, 0.0f);
 	const FLinearColor AlphaMask(0.0f, 0.0f, 0.0f, 1.0f);
 
-	const bool bHasMetallicSourceTexture = TryGetSourceTexture(MetallicTexture, MetallicTexCoord, MetallicInput, MaterialInstance, { MetallicMask });
-	const bool bHasRoughnessSourceTexture = TryGetSourceTexture(RoughnessTexture, RoughnessTexCoord, RoughnessInput, MaterialInstance, { RoughnessMask });
+	const bool bHasMetallicSourceTexture = TryGetSourceTexture(MetallicTexture, MetallicTexCoord, MetallicInput, MaterialInterface, { MetallicMask });
+	const bool bHasRoughnessSourceTexture = TryGetSourceTexture(RoughnessTexture, RoughnessTexCoord, RoughnessInput, MaterialInterface, { RoughnessMask });
 
 	// Detect the "happy path" where both inputs share the same texture and are correctly masked.
 	if (bHasMetallicSourceTexture &&
@@ -393,10 +388,10 @@ bool FGLTFMaterialConverter::TryGetMetallicAndRoughness(FGLTFConvertBuilder& Bui
 	return true;
 }
 
-bool FGLTFMaterialConverter::TryGetConstantColor(FGLTFJsonColor3& OutValue, const FColorMaterialInput& MaterialInput, const UMaterialInstance* MaterialInstance) const
+bool FGLTFMaterialConverter::TryGetConstantColor(FGLTFJsonColor3& OutValue, const FColorMaterialInput& MaterialInput, const UMaterialInterface* MaterialInterface) const
 {
 	FLinearColor Value;
-	if (TryGetConstantColor(Value, MaterialInput, MaterialInstance))
+	if (TryGetConstantColor(Value, MaterialInput, MaterialInterface))
 	{
 		OutValue = FGLTFConverterUtility::ConvertColor(Value);
 		return true;
@@ -405,10 +400,10 @@ bool FGLTFMaterialConverter::TryGetConstantColor(FGLTFJsonColor3& OutValue, cons
 	return false;
 }
 
-bool FGLTFMaterialConverter::TryGetConstantColor(FGLTFJsonColor4& OutValue, const FColorMaterialInput& MaterialInput, const UMaterialInstance* MaterialInstance) const
+bool FGLTFMaterialConverter::TryGetConstantColor(FGLTFJsonColor4& OutValue, const FColorMaterialInput& MaterialInput, const UMaterialInterface* MaterialInterface) const
 {
 	FLinearColor Value;
-	if (TryGetConstantColor(Value, MaterialInput, MaterialInstance))
+	if (TryGetConstantColor(Value, MaterialInput, MaterialInterface))
 	{
 		OutValue = FGLTFConverterUtility::ConvertColor(Value);
 		return true;
@@ -417,7 +412,7 @@ bool FGLTFMaterialConverter::TryGetConstantColor(FGLTFJsonColor4& OutValue, cons
 	return false;
 }
 
-bool FGLTFMaterialConverter::TryGetConstantColor(FLinearColor& OutValue, const FColorMaterialInput& MaterialInput, const UMaterialInstance* MaterialInstance) const
+bool FGLTFMaterialConverter::TryGetConstantColor(FLinearColor& OutValue, const FColorMaterialInput& MaterialInput, const UMaterialInterface* MaterialInterface) const
 {
 	// TODO: handle emissive color-values above 1.0
 
@@ -433,6 +428,7 @@ bool FGLTFMaterialConverter::TryGetConstantColor(FLinearColor& OutValue, const F
 	{
 		FLinearColor Value = VectorParameter->DefaultValue;
 
+		const UMaterialInstance* MaterialInstance = Cast<UMaterialInstance>(MaterialInterface);
 		if (MaterialInstance != nullptr)
 		{
 			const FHashedMaterialParameterInfo ParameterInfo(VectorParameter->GetParameterName());
@@ -465,6 +461,7 @@ bool FGLTFMaterialConverter::TryGetConstantColor(FLinearColor& OutValue, const F
 	{
 		float Value = ScalarParameter->DefaultValue;
 
+		const UMaterialInstance* MaterialInstance = Cast<UMaterialInstance>(MaterialInterface);
 		if (MaterialInstance != nullptr)
 		{
 			const FHashedMaterialParameterInfo ParameterInfo(ScalarParameter->GetParameterName());
@@ -505,7 +502,7 @@ bool FGLTFMaterialConverter::TryGetConstantColor(FLinearColor& OutValue, const F
 	return false;
 }
 
-bool FGLTFMaterialConverter::TryGetConstantScalar(float& OutValue, const FScalarMaterialInput& MaterialInput, const UMaterialInstance* MaterialInstance) const
+bool FGLTFMaterialConverter::TryGetConstantScalar(float& OutValue, const FScalarMaterialInput& MaterialInput, const UMaterialInterface* MaterialInterface) const
 {
 	const UMaterialExpression* Expression = MaterialInput.Expression;
 	if (Expression == nullptr)
@@ -519,6 +516,7 @@ bool FGLTFMaterialConverter::TryGetConstantScalar(float& OutValue, const FScalar
 	{
 		FLinearColor Value = VectorParameter->DefaultValue;
 
+		const UMaterialInstance* MaterialInstance = Cast<UMaterialInstance>(MaterialInterface);
 		if (MaterialInstance != nullptr)
 		{
 			const FHashedMaterialParameterInfo ParameterInfo(VectorParameter->GetParameterName());
@@ -545,6 +543,7 @@ bool FGLTFMaterialConverter::TryGetConstantScalar(float& OutValue, const FScalar
 	{
 		float Value = ScalarParameter->DefaultValue;
 
+		const UMaterialInstance* MaterialInstance = Cast<UMaterialInstance>(MaterialInterface);
 		if (MaterialInstance != nullptr)
 		{
 			const FHashedMaterialParameterInfo ParameterInfo(ScalarParameter->GetParameterName());
@@ -585,12 +584,12 @@ bool FGLTFMaterialConverter::TryGetConstantScalar(float& OutValue, const FScalar
 	return false;
 }
 
-bool FGLTFMaterialConverter::TryGetSourceTexture(FGLTFConvertBuilder& Builder, FGLTFJsonTextureInfo& OutTexInfo, const FExpressionInput& MaterialInput, const UMaterialInstance* MaterialInstance, const TArray<FLinearColor>& AllowedMasks) const
+bool FGLTFMaterialConverter::TryGetSourceTexture(FGLTFConvertBuilder& Builder, FGLTFJsonTextureInfo& OutTexInfo, const FExpressionInput& MaterialInput, const UMaterialInterface* MaterialInterface, const TArray<FLinearColor>& AllowedMasks) const
 {
 	const UTexture2D* Texture;
 	int32 TexCoord;
 
-	if (TryGetSourceTexture(Texture, TexCoord, MaterialInput, MaterialInstance, AllowedMasks))
+	if (TryGetSourceTexture(Texture, TexCoord, MaterialInput, MaterialInterface, AllowedMasks))
 	{
 		OutTexInfo.Index = Builder.GetOrAddTexture(Texture);
 		OutTexInfo.TexCoord = TexCoord;
@@ -600,7 +599,7 @@ bool FGLTFMaterialConverter::TryGetSourceTexture(FGLTFConvertBuilder& Builder, F
 	return false;
 }
 
-bool FGLTFMaterialConverter::TryGetSourceTexture(const UTexture2D*& OutTexture, int32& OutTexCoord, const FExpressionInput& MaterialInput, const UMaterialInstance* MaterialInstance, const TArray<FLinearColor>& AllowedMasks) const
+bool FGLTFMaterialConverter::TryGetSourceTexture(const UTexture2D*& OutTexture, int32& OutTexCoord, const FExpressionInput& MaterialInput, const UMaterialInterface* MaterialInterface, const TArray<FLinearColor>& AllowedMasks) const
 {
 	const FLinearColor InputMask = FGLTFMaterialUtility::GetMask(MaterialInput);
 	if (AllowedMasks.Num() > 0 && !AllowedMasks.Contains(InputMask))
@@ -616,6 +615,7 @@ bool FGLTFMaterialConverter::TryGetSourceTexture(const UTexture2D*& OutTexture,
 
 	if (const UMaterialExpressionTextureSampleParameter2D* TextureParameter = ExactCast<UMaterialExpressionTextureSampleParameter2D>(Expression))
 	{
+		const UMaterialInstance* MaterialInstance = Cast<UMaterialInstance>(MaterialInterface);
 		if (MaterialInstance != nullptr)
 		{
 			UTexture* ParameterValue;
diff --git a/Engine/Plugins/Enterprise/GLTFExporter/MigratedFromGit/Source/GLTFExporter/Private/Converters/GLTFMaterialConverters.h b/Engine/Plugins/Enterprise/GLTFExporter/MigratedFromGit/Source/GLTFExporter/Private/Converters/GLTFMaterialConverters.h
--- a/Engine/Plugins/Enterprise/GLTFExporter/MigratedFromGit/Source/GLTFExporter/Private/Converters/GLTFMaterialConverters.h
+++ b/Engine/Plugins/Enterprise/GLTFExporter/MigratedFromGit/Source/GLTFExporter/Private/Converters/GLTFMaterialConverters.h
@@ -16,13 +16,13 @@ class FGLTFMaterialConverter final : public TGLTFConverter<FGLTFJsonMaterialInde
 	bool TryGetBaseColorAndOpacity(FGLTFConvertBuilder& Builder, FGLTFJsonPBRMetallicRoughness& OutPBRParams, const FColorMaterialInput& BaseColorInput, const FScalarMaterialInput& OpacityInput, const UMaterialInterface* MaterialInterface) const;
 	bool TryGetMetallicAndRoughness(FGLTFConvertBuilder& Builder, FGLTFJsonPBRMetallicRoughness& OutPBRParams, const FScalarMaterialInput& MetallicInput, const FScalarMaterialInput& RoughnessInput, const UMaterialInterface* MaterialInterface) const;
 
-	bool TryGetConstantColor(FGLTFJsonColor3& OutValue, const FColorMaterialInput& MaterialInput, const UMaterialInstance* MaterialInstance) const;
-	bool TryGetConstantColor(FGLTFJsonColor4& OutValue, const FColorMaterialInput& MaterialInput, const UMaterialInstance* MaterialInstance) const;
-	bool TryGetConstantColor(FLinearColor& OutValue, const FColorMaterialInput& MaterialInput, const UMaterialInstance* MaterialInstance) const;
-	bool TryGetConstantScalar(float& OutValue, const FScalarMaterialInput& MaterialInput, const UMaterialInstance* MaterialInstance) const;
+	bool TryGetConstantColor(FGLTFJsonColor3& OutValue, const FColorMaterialInput& MaterialInput, const UMaterialInterface* MaterialInterface) const;
+	bool TryGetConstantColor(FGLTFJsonColor4& OutValue, const FColorMaterialInput& MaterialInput, const UMaterialInterface* MaterialInterface) const;
+	bool TryGetConstantColor(FLinearColor& OutValue, const FColorMaterialInput& MaterialInput, const UMaterialInterface* MaterialInterface) const;
+	bool TryGetConstantScalar(float& OutValue, const FScalarMaterialInput& MaterialInput, const UMaterialInterface* MaterialInterface) const;
 
-	bool TryGetSourceTexture(FGLTFConvertBuilder& Builder, FGLTFJsonTextureInfo& OutTexInfo, const FExpressionInput& MaterialInput, const UMaterialInstance* MaterialInstance, const TArray<FLinearColor>& AllowedMasks = {}) const;
-	bool TryGetSourceTexture(const UTexture2D*& OutTexture, int32& OutTexCoord, const FExpressionInput& MaterialInput, const UMaterialInstance* MaterialInstance, const TArray<FLinearColor>& AllowedMasks = {}) const;
+	bool TryGetSourceTexture(FGLTFConvertBuilder& Builder, FGLTFJsonTextureInfo& OutTexInfo, const FExpressionInput& MaterialInput, const UMaterialInterface* MaterialInterface, const TArray<FLinearColor>& AllowedMasks = {}) const;
+	bool TryGetSourceTexture(const UTexture2D*& OutTexture, int32& OutTexCoord, const FExpressionInput& MaterialInput, const UMaterialInterface* MaterialInterface, const TArray<FLinearColor>& AllowedMasks = {}) const;
 
 	bool TryGetBakedMaterialProperty(FGLTFConvertBuilder& Builder, FGLTFJsonTextureInfo& OutTexInfo, FGLTFJsonColor3& OutConstant, EMaterialProperty MaterialProperty, const UMaterialInterface* MaterialInterface) const;
 	bool TryGetBakedMaterialProperty(FGLTFConvertBuilder& Builder, FGLTFJsonTextureInfo& OutTexInfo, FGLTFJsonColor4& OutConstant, EMaterialProperty MaterialProperty, const UMaterialInterface* MaterialInterface) const;
