[GLTFExporter] fix(material): support usage of material attributes

commit 1051b1a9fc6c12603a9e7749d9bbbe5897f273bc
Date: Sun, 14 Jun 2020 14:36:02 +0200

    fix(material): support usage of material attributes
    (cherry picked from commit f47c78ec9a28ab668b943f9e6578c6ce040979e1​)

#skipci

diff --git a/Engine/Plugins/Enterprise/GLTFExporter/MigratedFromGit/Source/GLTFExporter/Private/Converters/GLTFMaterialConverters.cpp b/Engine/Plugins/Enterprise/GLTFExporter/MigratedFromGit/Source/GLTFExporter/Private/Converters/GLTFMaterialConverters.cpp
--- a/Engine/Plugins/Enterprise/GLTFExporter/MigratedFromGit/Source/GLTFExporter/Private/Converters/GLTFMaterialConverters.cpp
+++ b/Engine/Plugins/Enterprise/GLTFExporter/MigratedFromGit/Source/GLTFExporter/Private/Converters/GLTFMaterialConverters.cpp
@@ -86,7 +86,7 @@ FGLTFJsonMaterialIndex FGLTFMaterialConverter::Add(FGLTFConvertBuilder& Builder,
 	}
 	*/
 
-	if (Material->Normal.Expression != nullptr)
+	if (IsPropertyNonDefault(Material->Normal, MaterialInterface))
 	{
 		if (!TryGetSourceTexture(Builder, JsonMaterial.NormalTexture, Material->Normal, MaterialInterface, { RgbaMask, RgbMask }))
 		{
@@ -97,7 +97,7 @@ FGLTFJsonMaterialIndex FGLTFMaterialConverter::Add(FGLTFConvertBuilder& Builder,
 		}
 	}
 
-	if (Material->AmbientOcclusion.Expression != nullptr)
+	if (IsPropertyNonDefault(Material->AmbientOcclusion, MaterialInterface))
 	{
 		if (!TryGetSourceTexture(Builder, JsonMaterial.OcclusionTexture, Material->AmbientOcclusion, MaterialInterface, { RMask }))
 		{
@@ -388,6 +388,23 @@ bool FGLTFMaterialConverter::TryGetMetallicAndRoughness(FGLTFConvertBuilder& Bui
 	return true;
 }
 
+bool FGLTFMaterialConverter::IsPropertyNonDefault(const FExpressionInput& MaterialInput, const UMaterialInterface* MaterialInterface) const
+{
+	const bool bUseMaterialAttributes = MaterialInterface->GetMaterial()->bUseMaterialAttributes;
+	if (bUseMaterialAttributes)
+	{
+		return true;
+	}
+
+	const UMaterialExpression* Expression = MaterialInput.Expression;
+	if (Expression == nullptr)
+	{
+		return false;
+	}
+
+	return true;
+}
+
 bool FGLTFMaterialConverter::TryGetConstantColor(FGLTFJsonColor3& OutValue, const FColorMaterialInput& MaterialInput, const UMaterialInterface* MaterialInterface) const
 {
 	FLinearColor Value;
@@ -416,6 +433,12 @@ bool FGLTFMaterialConverter::TryGetConstantColor(FLinearColor& OutValue, const F
 {
 	// TODO: handle emissive color-values above 1.0
 
+	const bool bUseMaterialAttributes = MaterialInterface->GetMaterial()->bUseMaterialAttributes;
+	if (bUseMaterialAttributes)
+	{
+		return false;
+	}
+
 	const UMaterialExpression* Expression = MaterialInput.Expression;
 	if (Expression == nullptr)
 	{
@@ -504,6 +527,12 @@ bool FGLTFMaterialConverter::TryGetConstantColor(FLinearColor& OutValue, const F
 
 bool FGLTFMaterialConverter::TryGetConstantScalar(float& OutValue, const FScalarMaterialInput& MaterialInput, const UMaterialInterface* MaterialInterface) const
 {
+	const bool bUseMaterialAttributes = MaterialInterface->GetMaterial()->bUseMaterialAttributes;
+	if (bUseMaterialAttributes)
+	{
+		return false;
+	}
+
 	const UMaterialExpression* Expression = MaterialInput.Expression;
 	if (Expression == nullptr)
 	{
diff --git a/Engine/Plugins/Enterprise/GLTFExporter/MigratedFromGit/Source/GLTFExporter/Private/Converters/GLTFMaterialConverters.h b/Engine/Plugins/Enterprise/GLTFExporter/MigratedFromGit/Source/GLTFExporter/Private/Converters/GLTFMaterialConverters.h
--- a/Engine/Plugins/Enterprise/GLTFExporter/MigratedFromGit/Source/GLTFExporter/Private/Converters/GLTFMaterialConverters.h
+++ b/Engine/Plugins/Enterprise/GLTFExporter/MigratedFromGit/Source/GLTFExporter/Private/Converters/GLTFMaterialConverters.h
@@ -16,6 +16,7 @@ class FGLTFMaterialConverter final : public TGLTFConverter<FGLTFJsonMaterialInde
 	bool TryGetBaseColorAndOpacity(FGLTFConvertBuilder& Builder, FGLTFJsonPBRMetallicRoughness& OutPBRParams, const FColorMaterialInput& BaseColorInput, const FScalarMaterialInput& OpacityInput, const UMaterialInterface* MaterialInterface) const;
 	bool TryGetMetallicAndRoughness(FGLTFConvertBuilder& Builder, FGLTFJsonPBRMetallicRoughness& OutPBRParams, const FScalarMaterialInput& MetallicInput, const FScalarMaterialInput& RoughnessInput, const UMaterialInterface* MaterialInterface) const;
 
+	bool IsPropertyNonDefault(const FExpressionInput& MaterialInput, const UMaterialInterface* MaterialInterface) const;
 	bool TryGetConstantColor(FGLTFJsonColor3& OutValue, const FColorMaterialInput& MaterialInput, const UMaterialInterface* MaterialInterface) const;
 	bool TryGetConstantColor(FGLTFJsonColor4& OutValue, const FColorMaterialInput& MaterialInput, const UMaterialInterface* MaterialInterface) const;
 	bool TryGetConstantColor(FLinearColor& OutValue, const FColorMaterialInput& MaterialInput, const UMaterialInterface* MaterialInterface) const;
