[GLTFExporter] refactor(material): use EMaterialProperty (instead of F*MaterialInput) as parameter when trying to get constant or source texture

commit f4e3e37a11b63fba410a2cd61f293ea86c28aea6
Date: Sun, 14 Jun 2020 15:01:29 +0200

    refactor(material): use EMaterialProperty (instead of F*MaterialInput) as parameter when trying to get constant or source texture
    (cherry picked from commit 9fbec0af2c41b6de8ca66d34b415437e13e0c67c​)

#skipci

diff --git a/Engine/Plugins/Enterprise/GLTFExporter/MigratedFromGit/Source/GLTFExporter/Private/Converters/GLTFMaterialConverters.cpp b/Engine/Plugins/Enterprise/GLTFExporter/MigratedFromGit/Source/GLTFExporter/Private/Converters/GLTFMaterialConverters.cpp
--- a/Engine/Plugins/Enterprise/GLTFExporter/MigratedFromGit/Source/GLTFExporter/Private/Converters/GLTFMaterialConverters.cpp
+++ b/Engine/Plugins/Enterprise/GLTFExporter/MigratedFromGit/Source/GLTFExporter/Private/Converters/GLTFMaterialConverters.cpp
@@ -32,17 +32,16 @@ FGLTFJsonMaterialIndex FGLTFMaterialConverter::Add(FGLTFConvertBuilder& Builder,
 
 	if (JsonMaterial.AlphaMode == EGLTFJsonAlphaMode::Blend || JsonMaterial.AlphaMode == EGLTFJsonAlphaMode::Mask)
 	{
-		const FScalarMaterialInput& OpacityInput = JsonMaterial.AlphaMode == EGLTFJsonAlphaMode::Blend ? Material->Opacity : Material->OpacityMask;
-		if (!TryGetBaseColorAndOpacity(Builder, JsonMaterial.PBRMetallicRoughness, Material->BaseColor, OpacityInput, MaterialInterface))
+		if (!TryGetBaseColorAndOpacity(Builder, JsonMaterial.PBRMetallicRoughness, MaterialInterface))
 		{
 			// TODO: handle failure?
 		}
 	}
 	else
 	{
-		if (!TryGetConstantColor(JsonMaterial.PBRMetallicRoughness.BaseColorFactor, Material->BaseColor, MaterialInterface))
+		if (!TryGetConstantColor(JsonMaterial.PBRMetallicRoughness.BaseColorFactor, MP_BaseColor, MaterialInterface))
 		{
-			if (!TryGetSourceTexture(Builder, JsonMaterial.PBRMetallicRoughness.BaseColorTexture, Material->BaseColor, MaterialInterface, { RgbaMask, RgbMask }))
+			if (!TryGetSourceTexture(Builder, JsonMaterial.PBRMetallicRoughness.BaseColorTexture, MP_BaseColor, MaterialInterface, { RgbaMask, RgbMask }))
 			{
 				if (!TryGetBakedMaterialProperty(Builder, JsonMaterial.PBRMetallicRoughness.BaseColorTexture, JsonMaterial.PBRMetallicRoughness.BaseColorFactor, MP_BaseColor, MaterialInterface))
 				{
@@ -52,7 +51,7 @@ FGLTFJsonMaterialIndex FGLTFMaterialConverter::Add(FGLTFConvertBuilder& Builder,
 		}
 	}
 
-	if (!TryGetMetallicAndRoughness(Builder, JsonMaterial.PBRMetallicRoughness, Material->Metallic, Material->Roughness, MaterialInterface))
+	if (!TryGetMetallicAndRoughness(Builder, JsonMaterial.PBRMetallicRoughness, MaterialInterface))
 	{
 		// TODO: handle failure?
 	}
@@ -63,9 +62,9 @@ FGLTFJsonMaterialIndex FGLTFMaterialConverter::Add(FGLTFConvertBuilder& Builder,
 	// how the same materials look inside of UE.
 	// TODO: solve the issues in a separate MR.
 	/*
-	if (!TryGetConstantColor(JsonMaterial.EmissiveFactor, Material->EmissiveColor, MaterialInterface))
+	if (!TryGetConstantColor(JsonMaterial.EmissiveFactor, MP_EmissiveColor, MaterialInterface))
 	{
-		if (!TryGetSourceTexture(Builder, JsonMaterial.EmissiveTexture, Material->EmissiveColor, MaterialInterface, { RgbaMask, RgbMask }))
+		if (!TryGetSourceTexture(Builder, JsonMaterial.EmissiveTexture, MP_EmissiveColor, MaterialInterface, { RgbaMask, RgbMask }))
 		{
 			if (TryGetBakedMaterialProperty(Builder, JsonMaterial.EmissiveTexture, JsonMaterial.EmissiveFactor, MP_EmissiveColor, MaterialInterface))
 			{
@@ -86,9 +85,9 @@ FGLTFJsonMaterialIndex FGLTFMaterialConverter::Add(FGLTFConvertBuilder& Builder,
 	}
 	*/
 
-	if (IsPropertyNonDefault(Material->Normal, MaterialInterface))
+	if (IsPropertyNonDefault(MP_Normal, MaterialInterface))
 	{
-		if (!TryGetSourceTexture(Builder, JsonMaterial.NormalTexture, Material->Normal, MaterialInterface, { RgbaMask, RgbMask }))
+		if (!TryGetSourceTexture(Builder, JsonMaterial.NormalTexture, MP_Normal, MaterialInterface, { RgbaMask, RgbMask }))
 		{
 			if (!TryGetBakedMaterialProperty(Builder, JsonMaterial.NormalTexture, MP_Normal, MaterialInterface))
 			{
@@ -97,9 +96,9 @@ FGLTFJsonMaterialIndex FGLTFMaterialConverter::Add(FGLTFConvertBuilder& Builder,
 		}
 	}
 
-	if (IsPropertyNonDefault(Material->AmbientOcclusion, MaterialInterface))
+	if (IsPropertyNonDefault(MP_AmbientOcclusion, MaterialInterface))
 	{
-		if (!TryGetSourceTexture(Builder, JsonMaterial.OcclusionTexture, Material->AmbientOcclusion, MaterialInterface, { RMask }))
+		if (!TryGetSourceTexture(Builder, JsonMaterial.OcclusionTexture, MP_AmbientOcclusion, MaterialInterface, { RMask }))
 		{
 			if (!TryGetBakedMaterialProperty(Builder, JsonMaterial.OcclusionTexture, MP_AmbientOcclusion, MaterialInterface))
 			{
@@ -111,10 +110,11 @@ FGLTFJsonMaterialIndex FGLTFMaterialConverter::Add(FGLTFConvertBuilder& Builder,
 	return Builder.AddMaterial(JsonMaterial);
 }
 
-bool FGLTFMaterialConverter::TryGetBaseColorAndOpacity(FGLTFConvertBuilder& Builder, FGLTFJsonPBRMetallicRoughness& OutPBRParams, const FColorMaterialInput& BaseColorInput, const FScalarMaterialInput& OpacityInput, const UMaterialInterface* MaterialInterface) const
+bool FGLTFMaterialConverter::TryGetBaseColorAndOpacity(FGLTFConvertBuilder& Builder, FGLTFJsonPBRMetallicRoughness& OutPBRParams, const UMaterialInterface* MaterialInterface) const
 {
-	const bool bIsBaseColorConstant = TryGetConstantColor(OutPBRParams.BaseColorFactor, BaseColorInput, MaterialInterface);
-	const bool bIsOpacityConstant = TryGetConstantScalar(OutPBRParams.BaseColorFactor.A, OpacityInput, MaterialInterface);
+	const EMaterialProperty OpacityProperty = MaterialInterface->GetBlendMode() == BLEND_Masked ? MP_OpacityMask : MP_Opacity;
+	const bool bIsBaseColorConstant = TryGetConstantColor(OutPBRParams.BaseColorFactor, MP_BaseColor, MaterialInterface);
+	const bool bIsOpacityConstant = TryGetConstantScalar(OutPBRParams.BaseColorFactor.A, OpacityProperty, MaterialInterface);
 
 	if (bIsBaseColorConstant && bIsOpacityConstant)
 	{
@@ -134,8 +134,8 @@ bool FGLTFMaterialConverter::TryGetBaseColorAndOpacity(FGLTFConvertBuilder& Buil
 	const FLinearColor BaseColorMask(1.0f, 1.0f, 1.0f, 0.0f);
 	const FLinearColor OpacityMask(0.0f, 0.0f, 0.0f, 1.0f);
 
-	const bool bHasBaseColorSourceTexture = TryGetSourceTexture(BaseColorTexture, BaseColorTexCoord, BaseColorInput, MaterialInterface, { BaseColorMask });
-	const bool bHasOpacitySourceTexture = TryGetSourceTexture(OpacityTexture, OpacityTexCoord, OpacityInput, MaterialInterface, { OpacityMask });
+	const bool bHasBaseColorSourceTexture = TryGetSourceTexture(BaseColorTexture, BaseColorTexCoord, MP_BaseColor, MaterialInterface, { BaseColorMask });
+	const bool bHasOpacitySourceTexture = TryGetSourceTexture(OpacityTexture, OpacityTexCoord, OpacityProperty, MaterialInterface, { OpacityMask });
 
 	// Detect the "happy path" where both inputs share the same texture and are correctly masked.
 	if (bHasBaseColorSourceTexture &&
@@ -208,7 +208,6 @@ bool FGLTFMaterialConverter::TryGetBaseColorAndOpacity(FGLTFConvertBuilder& Buil
 		TextureMagFilter = FGLTFConverterUtility::ConvertMagFilter(OpacityTexture->Filter, OpacityTexture->LODGroup);
 	}
 
-	const EMaterialProperty OpacityProperty = MaterialInterface->GetBlendMode() == BLEND_Masked ? MP_OpacityMask : MP_Opacity;
 	const FGLTFPropertyBakeOutput BaseColorBakeOutput = BakeMaterialProperty(MP_BaseColor, MaterialInterface, &TextureSize);
 	FGLTFPropertyBakeOutput OpacityBakeOutput = BakeMaterialProperty(OpacityProperty, MaterialInterface, &TextureSize, true);
 
@@ -250,10 +249,10 @@ bool FGLTFMaterialConverter::TryGetBaseColorAndOpacity(FGLTFConvertBuilder& Buil
 	return true;
 }
 
-bool FGLTFMaterialConverter::TryGetMetallicAndRoughness(FGLTFConvertBuilder& Builder, FGLTFJsonPBRMetallicRoughness& OutPBRParams, const FScalarMaterialInput& MetallicInput, const FScalarMaterialInput& RoughnessInput, const UMaterialInterface* MaterialInterface) const
+bool FGLTFMaterialConverter::TryGetMetallicAndRoughness(FGLTFConvertBuilder& Builder, FGLTFJsonPBRMetallicRoughness& OutPBRParams, const UMaterialInterface* MaterialInterface) const
 {
-	const bool bIsMetallicConstant = TryGetConstantScalar(OutPBRParams.MetallicFactor, MetallicInput, MaterialInterface);
-	const bool bIsRoughnessConstant = TryGetConstantScalar(OutPBRParams.RoughnessFactor, RoughnessInput, MaterialInterface);
+	const bool bIsMetallicConstant = TryGetConstantScalar(OutPBRParams.MetallicFactor, MP_Metallic, MaterialInterface);
+	const bool bIsRoughnessConstant = TryGetConstantScalar(OutPBRParams.RoughnessFactor, MP_Roughness, MaterialInterface);
 
 	if (bIsMetallicConstant && bIsRoughnessConstant)
 	{
@@ -275,8 +274,8 @@ bool FGLTFMaterialConverter::TryGetMetallicAndRoughness(FGLTFConvertBuilder& Bui
 	const FLinearColor RoughnessMask(0.0f, 1.0f, 0.0f, 0.0f);
 	const FLinearColor AlphaMask(0.0f, 0.0f, 0.0f, 1.0f);
 
-	const bool bHasMetallicSourceTexture = TryGetSourceTexture(MetallicTexture, MetallicTexCoord, MetallicInput, MaterialInterface, { MetallicMask });
-	const bool bHasRoughnessSourceTexture = TryGetSourceTexture(RoughnessTexture, RoughnessTexCoord, RoughnessInput, MaterialInterface, { RoughnessMask });
+	const bool bHasMetallicSourceTexture = TryGetSourceTexture(MetallicTexture, MetallicTexCoord, MP_Metallic, MaterialInterface, { MetallicMask });
+	const bool bHasRoughnessSourceTexture = TryGetSourceTexture(RoughnessTexture, RoughnessTexCoord, MP_Roughness, MaterialInterface, { RoughnessMask });
 
 	// Detect the "happy path" where both inputs share the same texture and are correctly masked.
 	if (bHasMetallicSourceTexture &&
@@ -388,7 +387,7 @@ bool FGLTFMaterialConverter::TryGetMetallicAndRoughness(FGLTFConvertBuilder& Bui
 	return true;
 }
 
-bool FGLTFMaterialConverter::IsPropertyNonDefault(const FExpressionInput& MaterialInput, const UMaterialInterface* MaterialInterface) const
+bool FGLTFMaterialConverter::IsPropertyNonDefault(EMaterialProperty MaterialProperty, const UMaterialInterface* MaterialInterface) const
 {
 	const bool bUseMaterialAttributes = MaterialInterface->GetMaterial()->bUseMaterialAttributes;
 	if (bUseMaterialAttributes)
@@ -396,7 +395,14 @@ bool FGLTFMaterialConverter::IsPropertyNonDefault(const FExpressionInput& Materi
 		return true;
 	}
 
-	const UMaterialExpression* Expression = MaterialInput.Expression;
+	const FExpressionInput* MaterialInput = FGLTFMaterialUtility::GetInputFromProperty(MaterialInterface, MaterialProperty);
+	if (MaterialInput == nullptr)
+	{
+		// TODO: report error
+		return false;
+	}
+
+	const UMaterialExpression* Expression = MaterialInput->Expression;
 	if (Expression == nullptr)
 	{
 		return false;
@@ -405,10 +411,10 @@ bool FGLTFMaterialConverter::IsPropertyNonDefault(const FExpressionInput& Materi
 	return true;
 }
 
-bool FGLTFMaterialConverter::TryGetConstantColor(FGLTFJsonColor3& OutValue, const FColorMaterialInput& MaterialInput, const UMaterialInterface* MaterialInterface) const
+bool FGLTFMaterialConverter::TryGetConstantColor(FGLTFJsonColor3& OutValue, EMaterialProperty MaterialProperty, const UMaterialInterface* MaterialInterface) const
 {
 	FLinearColor Value;
-	if (TryGetConstantColor(Value, MaterialInput, MaterialInterface))
+	if (TryGetConstantColor(Value, MaterialProperty, MaterialInterface))
 	{
 		OutValue = FGLTFConverterUtility::ConvertColor(Value);
 		return true;
@@ -417,10 +423,10 @@ bool FGLTFMaterialConverter::TryGetConstantColor(FGLTFJsonColor3& OutValue, cons
 	return false;
 }
 
-bool FGLTFMaterialConverter::TryGetConstantColor(FGLTFJsonColor4& OutValue, const FColorMaterialInput& MaterialInput, const UMaterialInterface* MaterialInterface) const
+bool FGLTFMaterialConverter::TryGetConstantColor(FGLTFJsonColor4& OutValue, EMaterialProperty MaterialProperty, const UMaterialInterface* MaterialInterface) const
 {
 	FLinearColor Value;
-	if (TryGetConstantColor(Value, MaterialInput, MaterialInterface))
+	if (TryGetConstantColor(Value, MaterialProperty, MaterialInterface))
 	{
 		OutValue = FGLTFConverterUtility::ConvertColor(Value);
 		return true;
@@ -429,7 +435,7 @@ bool FGLTFMaterialConverter::TryGetConstantColor(FGLTFJsonColor4& OutValue, cons
 	return false;
 }
 
-bool FGLTFMaterialConverter::TryGetConstantColor(FLinearColor& OutValue, const FColorMaterialInput& MaterialInput, const UMaterialInterface* MaterialInterface) const
+bool FGLTFMaterialConverter::TryGetConstantColor(FLinearColor& OutValue, EMaterialProperty MaterialProperty, const UMaterialInterface* MaterialInterface) const
 {
 	// TODO: handle emissive color-values above 1.0
 
@@ -439,11 +445,18 @@ bool FGLTFMaterialConverter::TryGetConstantColor(FLinearColor& OutValue, const F
 		return false;
 	}
 
-	const UMaterialExpression* Expression = MaterialInput.Expression;
+	const FMaterialInput<FColor>* MaterialInput = FGLTFMaterialUtility::GetInputFromProperty<FColor>(MaterialInterface, MaterialProperty);
+	if (MaterialInput == nullptr)
+	{
+		// TODO: report error
+		return false;
+	}
+
+	const UMaterialExpression* Expression = MaterialInput->Expression;
 	if (Expression == nullptr)
 	{
 		// TODO: is this assumption correct?
-		OutValue = { MaterialInput.Constant };
+		OutValue = { MaterialInput->Constant };
 		return true;
 	}
 
@@ -461,11 +474,11 @@ bool FGLTFMaterialConverter::TryGetConstantColor(FLinearColor& OutValue, const F
 			}
 		}
 
-		const uint32 MaskComponentCount = FGLTFMaterialUtility::GetMaskComponentCount(MaterialInput);
+		const uint32 MaskComponentCount = FGLTFMaterialUtility::GetMaskComponentCount(*MaterialInput);
 
 		if (MaskComponentCount > 0)
 		{
-			const FLinearColor Mask = FGLTFMaterialUtility::GetMask(MaterialInput);
+			const FLinearColor Mask = FGLTFMaterialUtility::GetMask(*MaterialInput);
 
 			Value *= Mask;
 
@@ -525,7 +538,7 @@ bool FGLTFMaterialConverter::TryGetConstantColor(FLinearColor& OutValue, const F
 	return false;
 }
 
-bool FGLTFMaterialConverter::TryGetConstantScalar(float& OutValue, const FScalarMaterialInput& MaterialInput, const UMaterialInterface* MaterialInterface) const
+bool FGLTFMaterialConverter::TryGetConstantScalar(float& OutValue, EMaterialProperty MaterialProperty, const UMaterialInterface* MaterialInterface) const
 {
 	const bool bUseMaterialAttributes = MaterialInterface->GetMaterial()->bUseMaterialAttributes;
 	if (bUseMaterialAttributes)
@@ -533,11 +546,18 @@ bool FGLTFMaterialConverter::TryGetConstantScalar(float& OutValue, const FScalar
 		return false;
 	}
 
-	const UMaterialExpression* Expression = MaterialInput.Expression;
+	const FMaterialInput<float>* MaterialInput = FGLTFMaterialUtility::GetInputFromProperty<float>(MaterialInterface, MaterialProperty);
+	if (MaterialInput == nullptr)
+	{
+		// TODO: report error
+		return false;
+	}
+
+	const UMaterialExpression* Expression = MaterialInput->Expression;
 	if (Expression == nullptr)
 	{
 		// TODO: is this assumption correct?
-		OutValue = MaterialInput.Constant;
+		OutValue = MaterialInput->Constant;
 		return true;
 	}
 
@@ -555,11 +575,11 @@ bool FGLTFMaterialConverter::TryGetConstantScalar(float& OutValue, const FScalar
 			}
 		}
 
-		const uint32 MaskComponentCount = FGLTFMaterialUtility::GetMaskComponentCount(MaterialInput);
+		const uint32 MaskComponentCount = FGLTFMaterialUtility::GetMaskComponentCount(*MaterialInput);
 
 		if (MaskComponentCount > 0)
 		{
-			const FLinearColor Mask = FGLTFMaterialUtility::GetMask(MaterialInput);
+			const FLinearColor Mask = FGLTFMaterialUtility::GetMask(*MaterialInput);
 			Value *= Mask;
 		}
 
@@ -613,12 +633,12 @@ bool FGLTFMaterialConverter::TryGetConstantScalar(float& OutValue, const FScalar
 	return false;
 }
 
-bool FGLTFMaterialConverter::TryGetSourceTexture(FGLTFConvertBuilder& Builder, FGLTFJsonTextureInfo& OutTexInfo, const FExpressionInput& MaterialInput, const UMaterialInterface* MaterialInterface, const TArray<FLinearColor>& AllowedMasks) const
+bool FGLTFMaterialConverter::TryGetSourceTexture(FGLTFConvertBuilder& Builder, FGLTFJsonTextureInfo& OutTexInfo, EMaterialProperty MaterialProperty, const UMaterialInterface* MaterialInterface, const TArray<FLinearColor>& AllowedMasks) const
 {
 	const UTexture2D* Texture;
 	int32 TexCoord;
 
-	if (TryGetSourceTexture(Texture, TexCoord, MaterialInput, MaterialInterface, AllowedMasks))
+	if (TryGetSourceTexture(Texture, TexCoord, MaterialProperty, MaterialInterface, AllowedMasks))
 	{
 		OutTexInfo.Index = Builder.GetOrAddTexture(Texture);
 		OutTexInfo.TexCoord = TexCoord;
@@ -628,20 +648,27 @@ bool FGLTFMaterialConverter::TryGetSourceTexture(FGLTFConvertBuilder& Builder, F
 	return false;
 }
 
-bool FGLTFMaterialConverter::TryGetSourceTexture(const UTexture2D*& OutTexture, int32& OutTexCoord, const FExpressionInput& MaterialInput, const UMaterialInterface* MaterialInterface, const TArray<FLinearColor>& AllowedMasks) const
+bool FGLTFMaterialConverter::TryGetSourceTexture(const UTexture2D*& OutTexture, int32& OutTexCoord, EMaterialProperty MaterialProperty, const UMaterialInterface* MaterialInterface, const TArray<FLinearColor>& AllowedMasks) const
 {
-	const FLinearColor InputMask = FGLTFMaterialUtility::GetMask(MaterialInput);
-	if (AllowedMasks.Num() > 0 && !AllowedMasks.Contains(InputMask))
+	const FExpressionInput* MaterialInput = FGLTFMaterialUtility::GetInputFromProperty(MaterialInterface, MaterialProperty);
+	if (MaterialInput == nullptr)
 	{
+		// TODO: report error
 		return false;
 	}
 
-	const UMaterialExpression* Expression = MaterialInput.Expression;
+	const UMaterialExpression* Expression = MaterialInput->Expression;
 	if (Expression == nullptr)
 	{
 		return false;
 	}
 
+	const FLinearColor InputMask = FGLTFMaterialUtility::GetMask(*MaterialInput);
+	if (AllowedMasks.Num() > 0 && !AllowedMasks.Contains(InputMask))
+	{
+		return false;
+	}
+
 	if (const UMaterialExpressionTextureSampleParameter2D* TextureParameter = ExactCast<UMaterialExpressionTextureSampleParameter2D>(Expression))
 	{
 		const UMaterialInstance* MaterialInstance = Cast<UMaterialInstance>(MaterialInterface);
diff --git a/Engine/Plugins/Enterprise/GLTFExporter/MigratedFromGit/Source/GLTFExporter/Private/Converters/GLTFMaterialConverters.h b/Engine/Plugins/Enterprise/GLTFExporter/MigratedFromGit/Source/GLTFExporter/Private/Converters/GLTFMaterialConverters.h
--- a/Engine/Plugins/Enterprise/GLTFExporter/MigratedFromGit/Source/GLTFExporter/Private/Converters/GLTFMaterialConverters.h
+++ b/Engine/Plugins/Enterprise/GLTFExporter/MigratedFromGit/Source/GLTFExporter/Private/Converters/GLTFMaterialConverters.h
@@ -13,17 +13,17 @@ class FGLTFMaterialConverter final : public TGLTFConverter<FGLTFJsonMaterialInde
 {
 	FGLTFJsonMaterialIndex Add(FGLTFConvertBuilder& Builder, const FString& Name, const UMaterialInterface* Material) override;
 
-	bool TryGetBaseColorAndOpacity(FGLTFConvertBuilder& Builder, FGLTFJsonPBRMetallicRoughness& OutPBRParams, const FColorMaterialInput& BaseColorInput, const FScalarMaterialInput& OpacityInput, const UMaterialInterface* MaterialInterface) const;
-	bool TryGetMetallicAndRoughness(FGLTFConvertBuilder& Builder, FGLTFJsonPBRMetallicRoughness& OutPBRParams, const FScalarMaterialInput& MetallicInput, const FScalarMaterialInput& RoughnessInput, const UMaterialInterface* MaterialInterface) const;
+	bool TryGetBaseColorAndOpacity(FGLTFConvertBuilder& Builder, FGLTFJsonPBRMetallicRoughness& OutPBRParams, const UMaterialInterface* MaterialInterface) const;
+	bool TryGetMetallicAndRoughness(FGLTFConvertBuilder& Builder, FGLTFJsonPBRMetallicRoughness& OutPBRParams, const UMaterialInterface* MaterialInterface) const;
 
-	bool IsPropertyNonDefault(const FExpressionInput& MaterialInput, const UMaterialInterface* MaterialInterface) const;
-	bool TryGetConstantColor(FGLTFJsonColor3& OutValue, const FColorMaterialInput& MaterialInput, const UMaterialInterface* MaterialInterface) const;
-	bool TryGetConstantColor(FGLTFJsonColor4& OutValue, const FColorMaterialInput& MaterialInput, const UMaterialInterface* MaterialInterface) const;
-	bool TryGetConstantColor(FLinearColor& OutValue, const FColorMaterialInput& MaterialInput, const UMaterialInterface* MaterialInterface) const;
-	bool TryGetConstantScalar(float& OutValue, const FScalarMaterialInput& MaterialInput, const UMaterialInterface* MaterialInterface) const;
+	bool IsPropertyNonDefault(EMaterialProperty MaterialProperty, const UMaterialInterface* MaterialInterface) const;
+	bool TryGetConstantColor(FGLTFJsonColor3& OutValue, EMaterialProperty MaterialProperty, const UMaterialInterface* MaterialInterface) const;
+	bool TryGetConstantColor(FGLTFJsonColor4& OutValue, EMaterialProperty MaterialProperty, const UMaterialInterface* MaterialInterface) const;
+	bool TryGetConstantColor(FLinearColor& OutValue, EMaterialProperty MaterialProperty, const UMaterialInterface* MaterialInterface) const;
+	bool TryGetConstantScalar(float& OutValue, EMaterialProperty MaterialProperty, const UMaterialInterface* MaterialInterface) const;
 
-	bool TryGetSourceTexture(FGLTFConvertBuilder& Builder, FGLTFJsonTextureInfo& OutTexInfo, const FExpressionInput& MaterialInput, const UMaterialInterface* MaterialInterface, const TArray<FLinearColor>& AllowedMasks = {}) const;
-	bool TryGetSourceTexture(const UTexture2D*& OutTexture, int32& OutTexCoord, const FExpressionInput& MaterialInput, const UMaterialInterface* MaterialInterface, const TArray<FLinearColor>& AllowedMasks = {}) const;
+	bool TryGetSourceTexture(FGLTFConvertBuilder& Builder, FGLTFJsonTextureInfo& OutTexInfo, EMaterialProperty MaterialProperty, const UMaterialInterface* MaterialInterface, const TArray<FLinearColor>& AllowedMasks = {}) const;
+	bool TryGetSourceTexture(const UTexture2D*& OutTexture, int32& OutTexCoord, EMaterialProperty MaterialProperty, const UMaterialInterface* MaterialInterface, const TArray<FLinearColor>& AllowedMasks = {}) const;
 
 	bool TryGetBakedMaterialProperty(FGLTFConvertBuilder& Builder, FGLTFJsonTextureInfo& OutTexInfo, FGLTFJsonColor3& OutConstant, EMaterialProperty MaterialProperty, const UMaterialInterface* MaterialInterface) const;
 	bool TryGetBakedMaterialProperty(FGLTFConvertBuilder& Builder, FGLTFJsonTextureInfo& OutTexInfo, FGLTFJsonColor4& OutConstant, EMaterialProperty MaterialProperty, const UMaterialInterface* MaterialInterface) const;
diff --git a/Engine/Plugins/Enterprise/GLTFExporter/MigratedFromGit/Source/GLTFExporter/Private/Converters/GLTFMaterialUtility.h b/Engine/Plugins/Enterprise/GLTFExporter/MigratedFromGit/Source/GLTFExporter/Private/Converters/GLTFMaterialUtility.h
--- a/Engine/Plugins/Enterprise/GLTFExporter/MigratedFromGit/Source/GLTFExporter/Private/Converters/GLTFMaterialUtility.h
+++ b/Engine/Plugins/Enterprise/GLTFExporter/MigratedFromGit/Source/GLTFExporter/Private/Converters/GLTFMaterialUtility.h
@@ -34,6 +34,18 @@ struct FGLTFPropertyBakeOutput
 
 struct FGLTFMaterialUtility
 {
+	template<class InputType>
+    static const FMaterialInput<InputType>* GetInputFromProperty(const UMaterialInterface* Material, EMaterialProperty Property)
+	{
+		const FExpressionInput* ExpressionInput = const_cast<UMaterial*>(Material->GetMaterial())->GetExpressionInputForProperty(Property);
+		return static_cast<const FMaterialInput<InputType>*>(ExpressionInput);
+	}
+
+    static const FExpressionInput* GetInputFromProperty(const UMaterialInterface* Material, EMaterialProperty Property)
+	{
+		return const_cast<UMaterial*>(Material->GetMaterial())->GetExpressionInputForProperty(Property);
+	}
+
 	static UTexture2D* CreateTransientTexture(const FGLTFPropertyBakeOutput& PropertyBakeOutput, bool bUseSRGB = false);
 	static UTexture2D* CreateTransientTexture(const TArray<FColor>& Pixels, const FIntPoint& TextureSize, EPixelFormat TextureFormat, bool bUseSRGB = false);
 	static bool CombineTextures(TArray<FColor>& OutPixels, const TArray<FGLTFTextureCombineSource>& Sources, const FIntPoint& OutputSize, EPixelFormat OutputPixelFormat);
