[GLTFExporter] fix(material): use proper default values when no material input expression

commit 8beb24733f1355ec132b6bdf0e97e457af7b7065
Date: Sun, 14 Jun 2020 16:49:51 +0200

    fix(material): use proper default values when no material input expression
    (cherry picked from commit c597830e7049cf09932d153357c4dec0875fd08b​)

#skipci

diff --git a/Engine/Plugins/Enterprise/GLTFExporter/MigratedFromGit/Source/GLTFExporter/Private/Converters/GLTFMaterialConverters.cpp b/Engine/Plugins/Enterprise/GLTFExporter/MigratedFromGit/Source/GLTFExporter/Private/Converters/GLTFMaterialConverters.cpp
--- a/Engine/Plugins/Enterprise/GLTFExporter/MigratedFromGit/Source/GLTFExporter/Private/Converters/GLTFMaterialConverters.cpp
+++ b/Engine/Plugins/Enterprise/GLTFExporter/MigratedFromGit/Source/GLTFExporter/Private/Converters/GLTFMaterialConverters.cpp
@@ -452,11 +452,16 @@ bool FGLTFMaterialConverter::TryGetConstantColor(FLinearColor& OutValue, EMateri
 		return false;
 	}
 
+	if (MaterialInput->UseConstant)
+	{
+		OutValue = { MaterialInput->Constant };
+		return true;
+	}
+
 	const UMaterialExpression* Expression = MaterialInput->Expression;
 	if (Expression == nullptr)
 	{
-		// TODO: is this assumption correct?
-		OutValue = { MaterialInput->Constant };
+		OutValue = FLinearColor(FGLTFMaterialUtility::GetPropertyDefaultValue(MaterialProperty));
 		return true;
 	}
 
@@ -553,11 +558,16 @@ bool FGLTFMaterialConverter::TryGetConstantScalar(float& OutValue, EMaterialProp
 		return false;
 	}
 
+	if (MaterialInput->UseConstant)
+	{
+		OutValue = MaterialInput->Constant;
+		return true;
+	}
+
 	const UMaterialExpression* Expression = MaterialInput->Expression;
 	if (Expression == nullptr)
 	{
-		// TODO: is this assumption correct?
-		OutValue = MaterialInput->Constant;
+		OutValue = FGLTFMaterialUtility::GetPropertyDefaultValue(MaterialProperty).X;
 		return true;
 	}
 
diff --git a/Engine/Plugins/Enterprise/GLTFExporter/MigratedFromGit/Source/GLTFExporter/Private/Converters/GLTFMaterialUtility.cpp b/Engine/Plugins/Enterprise/GLTFExporter/MigratedFromGit/Source/GLTFExporter/Private/Converters/GLTFMaterialUtility.cpp
--- a/Engine/Plugins/Enterprise/GLTFExporter/MigratedFromGit/Source/GLTFExporter/Private/Converters/GLTFMaterialUtility.cpp
+++ b/Engine/Plugins/Enterprise/GLTFExporter/MigratedFromGit/Source/GLTFExporter/Private/Converters/GLTFMaterialUtility.cpp
@@ -10,6 +10,50 @@
 #include "IMaterialBakingModule.h"
 #include "MaterialBakingStructures.h"
 
+FVector4 FGLTFMaterialUtility::GetPropertyDefaultValue(EMaterialProperty Property)
+{
+	// TODO: replace with GMaterialPropertyAttributesMap lookup (when public API available)
+
+	switch (Property)
+	{
+		case MP_EmissiveColor:          return FVector4(0,0,0,0);
+		case MP_Opacity:                return FVector4(1,0,0,0);
+		case MP_OpacityMask:            return FVector4(1,0,0,0);
+		case MP_BaseColor:              return FVector4(0,0,0,0);
+		case MP_Metallic:               return FVector4(0,0,0,0);
+		case MP_Specular:               return FVector4(.5,0,0,0);
+		case MP_Roughness:              return FVector4(.5,0,0,0);
+		case MP_Anisotropy:             return FVector4(0,0,0,0);
+		case MP_Normal:                 return FVector4(0,0,1,0);
+		case MP_Tangent:                return FVector4(1,0,0,0);
+		case MP_WorldPositionOffset:    return FVector4(0,0,0,0);
+		case MP_WorldDisplacement:      return FVector4(0,0,0,0);
+		case MP_TessellationMultiplier: return FVector4(1,0,0,0);
+		case MP_SubsurfaceColor:        return FVector4(1,1,1,0);
+		case MP_CustomData0:            return FVector4(1,0,0,0);
+		case MP_CustomData1:            return FVector4(.1,0,0,0);
+		case MP_AmbientOcclusion:       return FVector4(1,0,0,0);
+		case MP_Refraction:             return FVector4(1,0,0,0);
+		case MP_PixelDepthOffset:       return FVector4(0,0,0,0);
+		case MP_ShadingModel:           return FVector4(0,0,0,0);
+		default:                        break;
+	}
+
+	if (Property >= MP_CustomizedUVs0 && Property <= MP_CustomizedUVs7)
+	{
+		return FVector4(0,0,0,0);
+	}
+
+	check(false);
+	return FVector4();
+}
+
+const FExpressionInput* FGLTFMaterialUtility::GetInputFromProperty(const UMaterialInterface* Material, EMaterialProperty Property)
+{
+	UMaterial* UnderlyingMaterial = const_cast<UMaterial*>(Material->GetMaterial());
+	return UnderlyingMaterial->GetExpressionInputForProperty(Property);
+}
+
 UTexture2D* FGLTFMaterialUtility::CreateTransientTexture(const FGLTFPropertyBakeOutput& PropertyBakeOutput, bool bUseSRGB)
 {
 	return CreateTransientTexture(
diff --git a/Engine/Plugins/Enterprise/GLTFExporter/MigratedFromGit/Source/GLTFExporter/Private/Converters/GLTFMaterialUtility.h b/Engine/Plugins/Enterprise/GLTFExporter/MigratedFromGit/Source/GLTFExporter/Private/Converters/GLTFMaterialUtility.h
--- a/Engine/Plugins/Enterprise/GLTFExporter/MigratedFromGit/Source/GLTFExporter/Private/Converters/GLTFMaterialUtility.h
+++ b/Engine/Plugins/Enterprise/GLTFExporter/MigratedFromGit/Source/GLTFExporter/Private/Converters/GLTFMaterialUtility.h
@@ -34,6 +34,10 @@ struct FGLTFPropertyBakeOutput
 
 struct FGLTFMaterialUtility
 {
+	static FVector4 GetPropertyDefaultValue(EMaterialProperty Property);
+
+	static const FExpressionInput* GetInputFromProperty(const UMaterialInterface* Material, EMaterialProperty Property);
+
 	template<class InputType>
     static const FMaterialInput<InputType>* GetInputFromProperty(const UMaterialInterface* Material, EMaterialProperty Property)
 	{
@@ -41,11 +45,6 @@ struct FGLTFMaterialUtility
 		return static_cast<const FMaterialInput<InputType>*>(ExpressionInput);
 	}
 
-    static const FExpressionInput* GetInputFromProperty(const UMaterialInterface* Material, EMaterialProperty Property)
-	{
-		return const_cast<UMaterial*>(Material->GetMaterial())->GetExpressionInputForProperty(Property);
-	}
-
 	static UTexture2D* CreateTransientTexture(const FGLTFPropertyBakeOutput& PropertyBakeOutput, bool bUseSRGB = false);
 	static UTexture2D* CreateTransientTexture(const TArray<FColor>& Pixels, const FIntPoint& TextureSize, EPixelFormat TextureFormat, bool bUseSRGB = false);
 	static bool CombineTextures(TArray<FColor>& OutPixels, const TArray<FGLTFTextureCombineSource>& Sources, const FIntPoint& OutputSize, EPixelFormat OutputPixelFormat);
