[GLTFExporter] refactor(material): add and update some TODOs

commit 8c9106b8ed739627ba7ff9eeaafd54e5c0a17dbf
Date: Sun, 14 Jun 2020 16:50:20 +0200

    refactor(material): add and update some TODOs
    (cherry picked from commit fcf188d55d78ca17b19922de43aae97ce0d912d1​)

#skipci

diff --git a/Engine/Plugins/Enterprise/GLTFExporter/MigratedFromGit/Source/GLTFExporter/Private/Converters/GLTFMaterialConverters.cpp b/Engine/Plugins/Enterprise/GLTFExporter/MigratedFromGit/Source/GLTFExporter/Private/Converters/GLTFMaterialConverters.cpp
--- a/Engine/Plugins/Enterprise/GLTFExporter/MigratedFromGit/Source/GLTFExporter/Private/Converters/GLTFMaterialConverters.cpp
+++ b/Engine/Plugins/Enterprise/GLTFExporter/MigratedFromGit/Source/GLTFExporter/Private/Converters/GLTFMaterialConverters.cpp
@@ -30,6 +30,8 @@ FGLTFJsonMaterialIndex FGLTFMaterialConverter::Add(FGLTFConvertBuilder& Builder,
 	const FLinearColor RgbMask(1.0f, 1.0f, 1.0f, 0.0f);
 	const FLinearColor RMask(1.0f, 0.0f, 0.0f, 0.0f);
 
+	// TODO: check if a property is active before trying to get it (i.e. MaterialInterface->IsPropertyActive)
+
 	if (JsonMaterial.AlphaMode == EGLTFJsonAlphaMode::Blend || JsonMaterial.AlphaMode == EGLTFJsonAlphaMode::Mask)
 	{
 		if (!TryGetBaseColorAndOpacity(Builder, JsonMaterial.PBRMetallicRoughness, MaterialInterface))
@@ -392,6 +394,7 @@ bool FGLTFMaterialConverter::IsPropertyNonDefault(EMaterialProperty MaterialProp
 	const bool bUseMaterialAttributes = MaterialInterface->GetMaterial()->bUseMaterialAttributes;
 	if (bUseMaterialAttributes)
 	{
+		// TODO: check if attribute property connected, i.e. MaterialInterface->GetMaterial()->MaterialAttributes.IsConnected(MaterialProperty)
 		return true;
 	}
 
@@ -442,6 +445,7 @@ bool FGLTFMaterialConverter::TryGetConstantColor(FLinearColor& OutValue, EMateri
 	const bool bUseMaterialAttributes = MaterialInterface->GetMaterial()->bUseMaterialAttributes;
 	if (bUseMaterialAttributes)
 	{
+		// TODO: check if attribute property connected, i.e. MaterialInterface->GetMaterial()->MaterialAttributes.IsConnected(MaterialProperty)
 		return false;
 	}
 
@@ -475,7 +479,7 @@ bool FGLTFMaterialConverter::TryGetConstantColor(FLinearColor& OutValue, EMateri
 			const FHashedMaterialParameterInfo ParameterInfo(VectorParameter->GetParameterName());
 			if (!MaterialInstance->GetVectorParameterValue(ParameterInfo, Value))
 			{
-				// TODO: can this happen?
+				// TODO: how to handle this?
 			}
 		}
 
@@ -508,7 +512,7 @@ bool FGLTFMaterialConverter::TryGetConstantColor(FLinearColor& OutValue, EMateri
 			const FHashedMaterialParameterInfo ParameterInfo(ScalarParameter->GetParameterName());
 			if (!MaterialInstance->GetScalarParameterValue(ParameterInfo, Value))
 			{
-				// TODO: can this happen?
+				// TODO: how to handle this?
 			}
 		}
 
@@ -548,6 +552,7 @@ bool FGLTFMaterialConverter::TryGetConstantScalar(float& OutValue, EMaterialProp
 	const bool bUseMaterialAttributes = MaterialInterface->GetMaterial()->bUseMaterialAttributes;
 	if (bUseMaterialAttributes)
 	{
+		// TODO: check if attribute property connected, i.e. MaterialInterface->GetMaterial()->MaterialAttributes.IsConnected(MaterialProperty)
 		return false;
 	}
 
@@ -581,7 +586,7 @@ bool FGLTFMaterialConverter::TryGetConstantScalar(float& OutValue, EMaterialProp
 			const FHashedMaterialParameterInfo ParameterInfo(VectorParameter->GetParameterName());
 			if (!MaterialInstance->GetVectorParameterValue(ParameterInfo, Value))
 			{
-				// TODO: can this happen?
+				// TODO: how to handle this?
 			}
 		}
 
@@ -608,7 +613,7 @@ bool FGLTFMaterialConverter::TryGetConstantScalar(float& OutValue, EMaterialProp
 			const FHashedMaterialParameterInfo ParameterInfo(ScalarParameter->GetParameterName());
 			if (!MaterialInstance->GetScalarParameterValue(ParameterInfo, Value))
 			{
-				// TODO: can this happen?
+				// TODO: how to handle this?
 			}
 		}
 
@@ -689,7 +694,7 @@ bool FGLTFMaterialConverter::TryGetSourceTexture(const UTexture2D*& OutTexture,
 			const FHashedMaterialParameterInfo ParameterInfo(TextureParameter->GetParameterName());
 			if (!MaterialInstance->GetTextureParameterValue(ParameterInfo, ParameterValue))
 			{
-				// TODO: can this happen?
+				// TODO: how to handle this?
 			}
 
 			OutTexture = Cast<UTexture2D>(ParameterValue);
