[GLTFExporter] refactor(material): rename MaterialInterface -> Material and MaterialProperty -> Property

commit 3763112cf80a96973acf31493f3e8d2347199130
Date: Sun, 14 Jun 2020 17:22:26 +0200

    refactor(material): rename MaterialInterface -> Material and MaterialProperty -> Property
    (cherry picked from commit 946b26de1da747f5a5eb22f90836706b5b605761​)

#skipci

diff --git a/Engine/Plugins/Enterprise/GLTFExporter/MigratedFromGit/Source/GLTFExporter/Private/Converters/GLTFMaterialConverters.cpp b/Engine/Plugins/Enterprise/GLTFExporter/MigratedFromGit/Source/GLTFExporter/Private/Converters/GLTFMaterialConverters.cpp
--- a/Engine/Plugins/Enterprise/GLTFExporter/MigratedFromGit/Source/GLTFExporter/Private/Converters/GLTFMaterialConverters.cpp
+++ b/Engine/Plugins/Enterprise/GLTFExporter/MigratedFromGit/Source/GLTFExporter/Private/Converters/GLTFMaterialConverters.cpp
@@ -13,39 +13,37 @@
 #include "Materials/MaterialExpressionTextureSample.h"
 #include "Materials/MaterialExpressionTextureSampleParameter2D.h"
 
-FGLTFJsonMaterialIndex FGLTFMaterialConverter::Add(FGLTFConvertBuilder& Builder, const FString& Name, const UMaterialInterface* MaterialInterface)
+FGLTFJsonMaterialIndex FGLTFMaterialConverter::Add(FGLTFConvertBuilder& Builder, const FString& Name, const UMaterialInterface* Material)
 {
-	const UMaterial* Material = MaterialInterface->GetMaterial();
-
 	FGLTFJsonMaterial JsonMaterial;
 	JsonMaterial.Name = Name;
 	// TODO: add support for different shading models (Default Lit, Unlit, Clear Coat)
 
 	// TODO: add support for additional blend modes (like Additive and Modulate)?
-	JsonMaterial.AlphaMode = FGLTFConverterUtility::ConvertBlendMode(MaterialInterface->GetBlendMode());
-	JsonMaterial.AlphaCutoff = MaterialInterface->GetOpacityMaskClipValue();
-	JsonMaterial.DoubleSided = MaterialInterface->IsTwoSided();
+	JsonMaterial.AlphaMode = FGLTFConverterUtility::ConvertBlendMode(Material->GetBlendMode());
+	JsonMaterial.AlphaCutoff = Material->GetOpacityMaskClipValue();
+	JsonMaterial.DoubleSided = Material->IsTwoSided();
 
 	const FLinearColor RgbaMask(1.0f, 1.0f, 1.0f, 1.0f);
 	const FLinearColor RgbMask(1.0f, 1.0f, 1.0f, 0.0f);
 	const FLinearColor RMask(1.0f, 0.0f, 0.0f, 0.0f);
 
-	// TODO: check if a property is active before trying to get it (i.e. MaterialInterface->IsPropertyActive)
+	// TODO: check if a property is active before trying to get it (i.e. Material->IsPropertyActive)
 
 	if (JsonMaterial.AlphaMode == EGLTFJsonAlphaMode::Blend || JsonMaterial.AlphaMode == EGLTFJsonAlphaMode::Mask)
 	{
-		if (!TryGetBaseColorAndOpacity(Builder, JsonMaterial.PBRMetallicRoughness, MaterialInterface))
+		if (!TryGetBaseColorAndOpacity(Builder, JsonMaterial.PBRMetallicRoughness, Material))
 		{
 			// TODO: handle failure?
 		}
 	}
 	else
 	{
-		if (!TryGetConstantColor(JsonMaterial.PBRMetallicRoughness.BaseColorFactor, MP_BaseColor, MaterialInterface))
+		if (!TryGetConstantColor(JsonMaterial.PBRMetallicRoughness.BaseColorFactor, MP_BaseColor, Material))
 		{
-			if (!TryGetSourceTexture(Builder, JsonMaterial.PBRMetallicRoughness.BaseColorTexture, MP_BaseColor, MaterialInterface, { RgbaMask, RgbMask }))
+			if (!TryGetSourceTexture(Builder, JsonMaterial.PBRMetallicRoughness.BaseColorTexture, MP_BaseColor, Material, { RgbaMask, RgbMask }))
 			{
-				if (!TryGetBakedMaterialProperty(Builder, JsonMaterial.PBRMetallicRoughness.BaseColorTexture, JsonMaterial.PBRMetallicRoughness.BaseColorFactor, MP_BaseColor, MaterialInterface))
+				if (!TryGetBakedMaterialProperty(Builder, JsonMaterial.PBRMetallicRoughness.BaseColorTexture, JsonMaterial.PBRMetallicRoughness.BaseColorFactor, MP_BaseColor, Material))
 				{
 					// TODO: handle failure?
 				}
@@ -53,7 +51,7 @@ FGLTFJsonMaterialIndex FGLTFMaterialConverter::Add(FGLTFConvertBuilder& Builder,
 		}
 	}
 
-	if (!TryGetMetallicAndRoughness(Builder, JsonMaterial.PBRMetallicRoughness, MaterialInterface))
+	if (!TryGetMetallicAndRoughness(Builder, JsonMaterial.PBRMetallicRoughness, Material))
 	{
 		// TODO: handle failure?
 	}
@@ -64,11 +62,11 @@ FGLTFJsonMaterialIndex FGLTFMaterialConverter::Add(FGLTFConvertBuilder& Builder,
 	// how the same materials look inside of UE.
 	// TODO: solve the issues in a separate MR.
 	/*
-	if (!TryGetConstantColor(JsonMaterial.EmissiveFactor, MP_EmissiveColor, MaterialInterface))
+	if (!TryGetConstantColor(JsonMaterial.EmissiveFactor, MP_EmissiveColor, Material))
 	{
-		if (!TryGetSourceTexture(Builder, JsonMaterial.EmissiveTexture, MP_EmissiveColor, MaterialInterface, { RgbaMask, RgbMask }))
+		if (!TryGetSourceTexture(Builder, JsonMaterial.EmissiveTexture, MP_EmissiveColor, Material, { RgbaMask, RgbMask }))
 		{
-			if (TryGetBakedMaterialProperty(Builder, JsonMaterial.EmissiveTexture, JsonMaterial.EmissiveFactor, MP_EmissiveColor, MaterialInterface))
+			if (TryGetBakedMaterialProperty(Builder, JsonMaterial.EmissiveTexture, JsonMaterial.EmissiveFactor, MP_EmissiveColor, Material))
 			{
 				if (JsonMaterial.EmissiveTexture.Index != INDEX_NONE)
 				{
@@ -87,22 +85,22 @@ FGLTFJsonMaterialIndex FGLTFMaterialConverter::Add(FGLTFConvertBuilder& Builder,
 	}
 	*/
 
-	if (IsPropertyNonDefault(MP_Normal, MaterialInterface))
+	if (IsPropertyNonDefault(MP_Normal, Material))
 	{
-		if (!TryGetSourceTexture(Builder, JsonMaterial.NormalTexture, MP_Normal, MaterialInterface, { RgbaMask, RgbMask }))
+		if (!TryGetSourceTexture(Builder, JsonMaterial.NormalTexture, MP_Normal, Material, { RgbaMask, RgbMask }))
 		{
-			if (!TryGetBakedMaterialProperty(Builder, JsonMaterial.NormalTexture, MP_Normal, MaterialInterface))
+			if (!TryGetBakedMaterialProperty(Builder, JsonMaterial.NormalTexture, MP_Normal, Material))
 			{
 				// TODO: handle failure?
 			}
 		}
 	}
 
-	if (IsPropertyNonDefault(MP_AmbientOcclusion, MaterialInterface))
+	if (IsPropertyNonDefault(MP_AmbientOcclusion, Material))
 	{
-		if (!TryGetSourceTexture(Builder, JsonMaterial.OcclusionTexture, MP_AmbientOcclusion, MaterialInterface, { RMask }))
+		if (!TryGetSourceTexture(Builder, JsonMaterial.OcclusionTexture, MP_AmbientOcclusion, Material, { RMask }))
 		{
-			if (!TryGetBakedMaterialProperty(Builder, JsonMaterial.OcclusionTexture, MP_AmbientOcclusion, MaterialInterface))
+			if (!TryGetBakedMaterialProperty(Builder, JsonMaterial.OcclusionTexture, MP_AmbientOcclusion, Material))
 			{
 				// TODO: handle failure?
 			}
@@ -112,11 +110,12 @@ FGLTFJsonMaterialIndex FGLTFMaterialConverter::Add(FGLTFConvertBuilder& Builder,
 	return Builder.AddMaterial(JsonMaterial);
 }
 
-bool FGLTFMaterialConverter::TryGetBaseColorAndOpacity(FGLTFConvertBuilder& Builder, FGLTFJsonPBRMetallicRoughness& OutPBRParams, const UMaterialInterface* MaterialInterface) const
+bool FGLTFMaterialConverter::TryGetBaseColorAndOpacity(FGLTFConvertBuilder& Builder, FGLTFJsonPBRMetallicRoughness& OutPBRParams, const UMaterialInterface* Material) const
 {
-	const EMaterialProperty OpacityProperty = MaterialInterface->GetBlendMode() == BLEND_Masked ? MP_OpacityMask : MP_Opacity;
-	const bool bIsBaseColorConstant = TryGetConstantColor(OutPBRParams.BaseColorFactor, MP_BaseColor, MaterialInterface);
-	const bool bIsOpacityConstant = TryGetConstantScalar(OutPBRParams.BaseColorFactor.A, OpacityProperty, MaterialInterface);
+	const EMaterialProperty OpacityProperty = Material->GetBlendMode() == BLEND_Masked ? MP_OpacityMask : MP_Opacity;
+
+	const bool bIsBaseColorConstant = TryGetConstantColor(OutPBRParams.BaseColorFactor, MP_BaseColor, Material);
+	const bool bIsOpacityConstant = TryGetConstantScalar(OutPBRParams.BaseColorFactor.A, OpacityProperty, Material);
 
 	if (bIsBaseColorConstant && bIsOpacityConstant)
 	{
@@ -136,8 +135,8 @@ bool FGLTFMaterialConverter::TryGetBaseColorAndOpacity(FGLTFConvertBuilder& Buil
 	const FLinearColor BaseColorMask(1.0f, 1.0f, 1.0f, 0.0f);
 	const FLinearColor OpacityMask(0.0f, 0.0f, 0.0f, 1.0f);
 
-	const bool bHasBaseColorSourceTexture = TryGetSourceTexture(BaseColorTexture, BaseColorTexCoord, MP_BaseColor, MaterialInterface, { BaseColorMask });
-	const bool bHasOpacitySourceTexture = TryGetSourceTexture(OpacityTexture, OpacityTexCoord, OpacityProperty, MaterialInterface, { OpacityMask });
+	const bool bHasBaseColorSourceTexture = TryGetSourceTexture(BaseColorTexture, BaseColorTexCoord, MP_BaseColor, Material, { BaseColorMask });
+	const bool bHasOpacitySourceTexture = TryGetSourceTexture(OpacityTexture, OpacityTexCoord, OpacityProperty, Material, { OpacityMask });
 
 	// Detect the "happy path" where both inputs share the same texture and are correctly masked.
 	if (bHasBaseColorSourceTexture &&
@@ -210,8 +209,8 @@ bool FGLTFMaterialConverter::TryGetBaseColorAndOpacity(FGLTFConvertBuilder& Buil
 		TextureMagFilter = FGLTFConverterUtility::ConvertMagFilter(OpacityTexture->Filter, OpacityTexture->LODGroup);
 	}
 
-	const FGLTFPropertyBakeOutput BaseColorBakeOutput = BakeMaterialProperty(MP_BaseColor, MaterialInterface, &TextureSize);
-	FGLTFPropertyBakeOutput OpacityBakeOutput = BakeMaterialProperty(OpacityProperty, MaterialInterface, &TextureSize, true);
+	const FGLTFPropertyBakeOutput BaseColorBakeOutput = BakeMaterialProperty(MP_BaseColor, Material, &TextureSize);
+	FGLTFPropertyBakeOutput OpacityBakeOutput = BakeMaterialProperty(OpacityProperty, Material, &TextureSize, true);
 
 	// Detect when both baked properties are constants, which means we can avoid exporting a texture
 	if (BaseColorBakeOutput.bIsConstant && OpacityBakeOutput.bIsConstant)
@@ -227,7 +226,7 @@ bool FGLTFMaterialConverter::TryGetBaseColorAndOpacity(FGLTFConvertBuilder& Buil
 	BaseColorTexture = FGLTFMaterialUtility::CreateTransientTexture(BaseColorBakeOutput);
 	OpacityTexture = FGLTFMaterialUtility::CreateTransientTexture(OpacityBakeOutput);
 
-	const FString TextureName = MaterialInterface->GetName() + TEXT("_BaseColor");
+	const FString TextureName = Material->GetName() + TEXT("_BaseColor");
 
 	const TArray<FGLTFTextureCombineSource> CombineSources =
 	{
@@ -251,10 +250,10 @@ bool FGLTFMaterialConverter::TryGetBaseColorAndOpacity(FGLTFConvertBuilder& Buil
 	return true;
 }
 
-bool FGLTFMaterialConverter::TryGetMetallicAndRoughness(FGLTFConvertBuilder& Builder, FGLTFJsonPBRMetallicRoughness& OutPBRParams, const UMaterialInterface* MaterialInterface) const
+bool FGLTFMaterialConverter::TryGetMetallicAndRoughness(FGLTFConvertBuilder& Builder, FGLTFJsonPBRMetallicRoughness& OutPBRParams, const UMaterialInterface* Material) const
 {
-	const bool bIsMetallicConstant = TryGetConstantScalar(OutPBRParams.MetallicFactor, MP_Metallic, MaterialInterface);
-	const bool bIsRoughnessConstant = TryGetConstantScalar(OutPBRParams.RoughnessFactor, MP_Roughness, MaterialInterface);
+	const bool bIsMetallicConstant = TryGetConstantScalar(OutPBRParams.MetallicFactor, MP_Metallic, Material);
+	const bool bIsRoughnessConstant = TryGetConstantScalar(OutPBRParams.RoughnessFactor, MP_Roughness, Material);
 
 	if (bIsMetallicConstant && bIsRoughnessConstant)
 	{
@@ -276,8 +275,8 @@ bool FGLTFMaterialConverter::TryGetMetallicAndRoughness(FGLTFConvertBuilder& Bui
 	const FLinearColor RoughnessMask(0.0f, 1.0f, 0.0f, 0.0f);
 	const FLinearColor AlphaMask(0.0f, 0.0f, 0.0f, 1.0f);
 
-	const bool bHasMetallicSourceTexture = TryGetSourceTexture(MetallicTexture, MetallicTexCoord, MP_Metallic, MaterialInterface, { MetallicMask });
-	const bool bHasRoughnessSourceTexture = TryGetSourceTexture(RoughnessTexture, RoughnessTexCoord, MP_Roughness, MaterialInterface, { RoughnessMask });
+	const bool bHasMetallicSourceTexture = TryGetSourceTexture(MetallicTexture, MetallicTexCoord, MP_Metallic, Material, { MetallicMask });
+	const bool bHasRoughnessSourceTexture = TryGetSourceTexture(RoughnessTexture, RoughnessTexCoord, MP_Roughness, Material, { RoughnessMask });
 
 	// Detect the "happy path" where both inputs share the same texture and are correctly masked.
 	if (bHasMetallicSourceTexture &&
@@ -350,8 +349,8 @@ bool FGLTFMaterialConverter::TryGetMetallicAndRoughness(FGLTFConvertBuilder& Bui
 		TextureMagFilter = FGLTFConverterUtility::ConvertMagFilter(RoughnessTexture->Filter, RoughnessTexture->LODGroup);
 	}
 
-	const FGLTFPropertyBakeOutput RoughnessBakeOutput = BakeMaterialProperty(MP_Roughness, MaterialInterface, &TextureSize);
-	const FGLTFPropertyBakeOutput MetallicBakeOutput = BakeMaterialProperty(MP_Metallic, MaterialInterface, &TextureSize);
+	const FGLTFPropertyBakeOutput RoughnessBakeOutput = BakeMaterialProperty(MP_Roughness, Material, &TextureSize);
+	const FGLTFPropertyBakeOutput MetallicBakeOutput = BakeMaterialProperty(MP_Metallic, Material, &TextureSize);
 
 	// Detect when both baked properties are constants, which means we can use factors and avoid exporting a texture
 	if (RoughnessBakeOutput.bIsConstant && MetallicBakeOutput.bIsConstant)
@@ -365,7 +364,7 @@ bool FGLTFMaterialConverter::TryGetMetallicAndRoughness(FGLTFConvertBuilder& Bui
 	RoughnessTexture = FGLTFMaterialUtility::CreateTransientTexture(RoughnessBakeOutput);
 	MetallicTexture = FGLTFMaterialUtility::CreateTransientTexture(MetallicBakeOutput);
 
-	const FString TextureName = MaterialInterface->GetName() + TEXT("_MetallicRoughness");
+	const FString TextureName = Material->GetName() + TEXT("_MetallicRoughness");
 
 	const TArray<FGLTFTextureCombineSource> CombineSources =
 	{
@@ -389,16 +388,16 @@ bool FGLTFMaterialConverter::TryGetMetallicAndRoughness(FGLTFConvertBuilder& Bui
 	return true;
 }
 
-bool FGLTFMaterialConverter::IsPropertyNonDefault(EMaterialProperty MaterialProperty, const UMaterialInterface* MaterialInterface) const
+bool FGLTFMaterialConverter::IsPropertyNonDefault(EMaterialProperty Property, const UMaterialInterface* Material) const
 {
-	const bool bUseMaterialAttributes = MaterialInterface->GetMaterial()->bUseMaterialAttributes;
+	const bool bUseMaterialAttributes = Material->GetMaterial()->bUseMaterialAttributes;
 	if (bUseMaterialAttributes)
 	{
-		// TODO: check if attribute property connected, i.e. MaterialInterface->GetMaterial()->MaterialAttributes.IsConnected(MaterialProperty)
+		// TODO: check if attribute property connected, i.e. Material->GetMaterial()->MaterialAttributes.IsConnected(Property)
 		return true;
 	}
 
-	const FExpressionInput* MaterialInput = FGLTFMaterialUtility::GetInputFromProperty(MaterialInterface, MaterialProperty);
+	const FExpressionInput* MaterialInput = FGLTFMaterialUtility::GetInputFromProperty(Material, Property);
 	if (MaterialInput == nullptr)
 	{
 		// TODO: report error
@@ -414,10 +413,10 @@ bool FGLTFMaterialConverter::IsPropertyNonDefault(EMaterialProperty MaterialProp
 	return true;
 }
 
-bool FGLTFMaterialConverter::TryGetConstantColor(FGLTFJsonColor3& OutValue, EMaterialProperty MaterialProperty, const UMaterialInterface* MaterialInterface) const
+bool FGLTFMaterialConverter::TryGetConstantColor(FGLTFJsonColor3& OutValue, EMaterialProperty Property, const UMaterialInterface* Material) const
 {
 	FLinearColor Value;
-	if (TryGetConstantColor(Value, MaterialProperty, MaterialInterface))
+	if (TryGetConstantColor(Value, Property, Material))
 	{
 		OutValue = FGLTFConverterUtility::ConvertColor(Value);
 		return true;
@@ -426,10 +425,10 @@ bool FGLTFMaterialConverter::TryGetConstantColor(FGLTFJsonColor3& OutValue, EMat
 	return false;
 }
 
-bool FGLTFMaterialConverter::TryGetConstantColor(FGLTFJsonColor4& OutValue, EMaterialProperty MaterialProperty, const UMaterialInterface* MaterialInterface) const
+bool FGLTFMaterialConverter::TryGetConstantColor(FGLTFJsonColor4& OutValue, EMaterialProperty Property, const UMaterialInterface* Material) const
 {
 	FLinearColor Value;
-	if (TryGetConstantColor(Value, MaterialProperty, MaterialInterface))
+	if (TryGetConstantColor(Value, Property, Material))
 	{
 		OutValue = FGLTFConverterUtility::ConvertColor(Value);
 		return true;
@@ -438,18 +437,18 @@ bool FGLTFMaterialConverter::TryGetConstantColor(FGLTFJsonColor4& OutValue, EMat
 	return false;
 }
 
-bool FGLTFMaterialConverter::TryGetConstantColor(FLinearColor& OutValue, EMaterialProperty MaterialProperty, const UMaterialInterface* MaterialInterface) const
+bool FGLTFMaterialConverter::TryGetConstantColor(FLinearColor& OutValue, EMaterialProperty Property, const UMaterialInterface* Material) const
 {
 	// TODO: handle emissive color-values above 1.0
 
-	const bool bUseMaterialAttributes = MaterialInterface->GetMaterial()->bUseMaterialAttributes;
+	const bool bUseMaterialAttributes = Material->GetMaterial()->bUseMaterialAttributes;
 	if (bUseMaterialAttributes)
 	{
-		// TODO: check if attribute property connected, i.e. MaterialInterface->GetMaterial()->MaterialAttributes.IsConnected(MaterialProperty)
+		// TODO: check if attribute property connected, i.e. Material->GetMaterial()->MaterialAttributes.IsConnected(Property)
 		return false;
 	}
 
-	const FMaterialInput<FColor>* MaterialInput = FGLTFMaterialUtility::GetInputFromProperty<FColor>(MaterialInterface, MaterialProperty);
+	const FMaterialInput<FColor>* MaterialInput = FGLTFMaterialUtility::GetInputFromProperty<FColor>(Material, Property);
 	if (MaterialInput == nullptr)
 	{
 		// TODO: report error
@@ -465,7 +464,7 @@ bool FGLTFMaterialConverter::TryGetConstantColor(FLinearColor& OutValue, EMateri
 	const UMaterialExpression* Expression = MaterialInput->Expression;
 	if (Expression == nullptr)
 	{
-		OutValue = FLinearColor(FGLTFMaterialUtility::GetPropertyDefaultValue(MaterialProperty));
+		OutValue = FLinearColor(FGLTFMaterialUtility::GetPropertyDefaultValue(Property));
 		return true;
 	}
 
@@ -473,7 +472,7 @@ bool FGLTFMaterialConverter::TryGetConstantColor(FLinearColor& OutValue, EMateri
 	{
 		FLinearColor Value = VectorParameter->DefaultValue;
 
-		const UMaterialInstance* MaterialInstance = Cast<UMaterialInstance>(MaterialInterface);
+		const UMaterialInstance* MaterialInstance = Cast<UMaterialInstance>(Material);
 		if (MaterialInstance != nullptr)
 		{
 			const FHashedMaterialParameterInfo ParameterInfo(VectorParameter->GetParameterName());
@@ -506,7 +505,7 @@ bool FGLTFMaterialConverter::TryGetConstantColor(FLinearColor& OutValue, EMateri
 	{
 		float Value = ScalarParameter->DefaultValue;
 
-		const UMaterialInstance* MaterialInstance = Cast<UMaterialInstance>(MaterialInterface);
+		const UMaterialInstance* MaterialInstance = Cast<UMaterialInstance>(Material);
 		if (MaterialInstance != nullptr)
 		{
 			const FHashedMaterialParameterInfo ParameterInfo(ScalarParameter->GetParameterName());
@@ -547,16 +546,16 @@ bool FGLTFMaterialConverter::TryGetConstantColor(FLinearColor& OutValue, EMateri
 	return false;
 }
 
-bool FGLTFMaterialConverter::TryGetConstantScalar(float& OutValue, EMaterialProperty MaterialProperty, const UMaterialInterface* MaterialInterface) const
+bool FGLTFMaterialConverter::TryGetConstantScalar(float& OutValue, EMaterialProperty Property, const UMaterialInterface* Material) const
 {
-	const bool bUseMaterialAttributes = MaterialInterface->GetMaterial()->bUseMaterialAttributes;
+	const bool bUseMaterialAttributes = Material->GetMaterial()->bUseMaterialAttributes;
 	if (bUseMaterialAttributes)
 	{
-		// TODO: check if attribute property connected, i.e. MaterialInterface->GetMaterial()->MaterialAttributes.IsConnected(MaterialProperty)
+		// TODO: check if attribute property connected, i.e. Material->GetMaterial()->MaterialAttributes.IsConnected(Property)
 		return false;
 	}
 
-	const FMaterialInput<float>* MaterialInput = FGLTFMaterialUtility::GetInputFromProperty<float>(MaterialInterface, MaterialProperty);
+	const FMaterialInput<float>* MaterialInput = FGLTFMaterialUtility::GetInputFromProperty<float>(Material, Property);
 	if (MaterialInput == nullptr)
 	{
 		// TODO: report error
@@ -572,7 +571,7 @@ bool FGLTFMaterialConverter::TryGetConstantScalar(float& OutValue, EMaterialProp
 	const UMaterialExpression* Expression = MaterialInput->Expression;
 	if (Expression == nullptr)
 	{
-		OutValue = FGLTFMaterialUtility::GetPropertyDefaultValue(MaterialProperty).X;
+		OutValue = FGLTFMaterialUtility::GetPropertyDefaultValue(Property).X;
 		return true;
 	}
 
@@ -580,7 +579,7 @@ bool FGLTFMaterialConverter::TryGetConstantScalar(float& OutValue, EMaterialProp
 	{
 		FLinearColor Value = VectorParameter->DefaultValue;
 
-		const UMaterialInstance* MaterialInstance = Cast<UMaterialInstance>(MaterialInterface);
+		const UMaterialInstance* MaterialInstance = Cast<UMaterialInstance>(Material);
 		if (MaterialInstance != nullptr)
 		{
 			const FHashedMaterialParameterInfo ParameterInfo(VectorParameter->GetParameterName());
@@ -607,7 +606,7 @@ bool FGLTFMaterialConverter::TryGetConstantScalar(float& OutValue, EMaterialProp
 	{
 		float Value = ScalarParameter->DefaultValue;
 
-		const UMaterialInstance* MaterialInstance = Cast<UMaterialInstance>(MaterialInterface);
+		const UMaterialInstance* MaterialInstance = Cast<UMaterialInstance>(Material);
 		if (MaterialInstance != nullptr)
 		{
 			const FHashedMaterialParameterInfo ParameterInfo(ScalarParameter->GetParameterName());
@@ -648,12 +647,12 @@ bool FGLTFMaterialConverter::TryGetConstantScalar(float& OutValue, EMaterialProp
 	return false;
 }
 
-bool FGLTFMaterialConverter::TryGetSourceTexture(FGLTFConvertBuilder& Builder, FGLTFJsonTextureInfo& OutTexInfo, EMaterialProperty MaterialProperty, const UMaterialInterface* MaterialInterface, const TArray<FLinearColor>& AllowedMasks) const
+bool FGLTFMaterialConverter::TryGetSourceTexture(FGLTFConvertBuilder& Builder, FGLTFJsonTextureInfo& OutTexInfo, EMaterialProperty Property, const UMaterialInterface* Material, const TArray<FLinearColor>& AllowedMasks) const
 {
 	const UTexture2D* Texture;
 	int32 TexCoord;
 
-	if (TryGetSourceTexture(Texture, TexCoord, MaterialProperty, MaterialInterface, AllowedMasks))
+	if (TryGetSourceTexture(Texture, TexCoord, Property, Material, AllowedMasks))
 	{
 		OutTexInfo.Index = Builder.GetOrAddTexture(Texture);
 		OutTexInfo.TexCoord = TexCoord;
@@ -663,9 +662,9 @@ bool FGLTFMaterialConverter::TryGetSourceTexture(FGLTFConvertBuilder& Builder, F
 	return false;
 }
 
-bool FGLTFMaterialConverter::TryGetSourceTexture(const UTexture2D*& OutTexture, int32& OutTexCoord, EMaterialProperty MaterialProperty, const UMaterialInterface* MaterialInterface, const TArray<FLinearColor>& AllowedMasks) const
+bool FGLTFMaterialConverter::TryGetSourceTexture(const UTexture2D*& OutTexture, int32& OutTexCoord, EMaterialProperty Property, const UMaterialInterface* Material, const TArray<FLinearColor>& AllowedMasks) const
 {
-	const FExpressionInput* MaterialInput = FGLTFMaterialUtility::GetInputFromProperty(MaterialInterface, MaterialProperty);
+	const FExpressionInput* MaterialInput = FGLTFMaterialUtility::GetInputFromProperty(Material, Property);
 	if (MaterialInput == nullptr)
 	{
 		// TODO: report error
@@ -686,7 +685,7 @@ bool FGLTFMaterialConverter::TryGetSourceTexture(const UTexture2D*& OutTexture,
 
 	if (const UMaterialExpressionTextureSampleParameter2D* TextureParameter = ExactCast<UMaterialExpressionTextureSampleParameter2D>(Expression))
 	{
-		const UMaterialInstance* MaterialInstance = Cast<UMaterialInstance>(MaterialInterface);
+		const UMaterialInstance* MaterialInstance = Cast<UMaterialInstance>(Material);
 		if (MaterialInstance != nullptr)
 		{
 			UTexture* ParameterValue;
@@ -744,9 +743,9 @@ bool FGLTFMaterialConverter::TryGetSourceTexture(const UTexture2D*& OutTexture,
 	return false;
 }
 
-bool FGLTFMaterialConverter::TryGetBakedMaterialProperty(FGLTFConvertBuilder& Builder, FGLTFJsonTextureInfo& OutTexInfo, FGLTFJsonColor3& OutConstant, EMaterialProperty MaterialProperty, const UMaterialInterface* MaterialInterface) const
+bool FGLTFMaterialConverter::TryGetBakedMaterialProperty(FGLTFConvertBuilder& Builder, FGLTFJsonTextureInfo& OutTexInfo, FGLTFJsonColor3& OutConstant, EMaterialProperty Property, const UMaterialInterface* Material) const
 {
-	const FGLTFPropertyBakeOutput PropertyBakeOutput = BakeMaterialProperty(MaterialProperty, MaterialInterface);
+	const FGLTFPropertyBakeOutput PropertyBakeOutput = BakeMaterialProperty(Property, Material);
 
 	if (PropertyBakeOutput.bIsConstant)
 	{
@@ -755,13 +754,13 @@ bool FGLTFMaterialConverter::TryGetBakedMaterialProperty(FGLTFConvertBuilder& Bu
 	}
 	else
 	{
-		return StoreBakedPropertyTexture(Builder, OutTexInfo, PropertyBakeOutput, MaterialProperty, MaterialInterface);
+		return StoreBakedPropertyTexture(Builder, OutTexInfo, PropertyBakeOutput, Property, Material);
 	}
 }
 
-bool FGLTFMaterialConverter::TryGetBakedMaterialProperty(FGLTFConvertBuilder& Builder, FGLTFJsonTextureInfo& OutTexInfo, FGLTFJsonColor4& OutConstant, EMaterialProperty MaterialProperty, const UMaterialInterface* MaterialInterface) const
+bool FGLTFMaterialConverter::TryGetBakedMaterialProperty(FGLTFConvertBuilder& Builder, FGLTFJsonTextureInfo& OutTexInfo, FGLTFJsonColor4& OutConstant, EMaterialProperty Property, const UMaterialInterface* Material) const
 {
-	const FGLTFPropertyBakeOutput PropertyBakeOutput = BakeMaterialProperty(MaterialProperty, MaterialInterface);
+	const FGLTFPropertyBakeOutput PropertyBakeOutput = BakeMaterialProperty(Property, Material);
 
 	if (PropertyBakeOutput.bIsConstant)
 	{
@@ -770,17 +769,17 @@ bool FGLTFMaterialConverter::TryGetBakedMaterialProperty(FGLTFConvertBuilder& Bu
 	}
 	else
 	{
-		return StoreBakedPropertyTexture(Builder, OutTexInfo, PropertyBakeOutput, MaterialProperty, MaterialInterface);
+		return StoreBakedPropertyTexture(Builder, OutTexInfo, PropertyBakeOutput, Property, Material);
 	}
 }
 
-bool FGLTFMaterialConverter::TryGetBakedMaterialProperty(FGLTFConvertBuilder& Builder, FGLTFJsonTextureInfo& OutTexInfo, EMaterialProperty MaterialProperty, const UMaterialInterface* MaterialInterface) const
+bool FGLTFMaterialConverter::TryGetBakedMaterialProperty(FGLTFConvertBuilder& Builder, FGLTFJsonTextureInfo& OutTexInfo, EMaterialProperty Property, const UMaterialInterface* Material) const
 {
-	const FGLTFPropertyBakeOutput PropertyBakeOutput = BakeMaterialProperty(MaterialProperty, MaterialInterface);
+	const FGLTFPropertyBakeOutput PropertyBakeOutput = BakeMaterialProperty(Property, Material);
 
 	if (!PropertyBakeOutput.bIsConstant)
 	{
-		return StoreBakedPropertyTexture(Builder, OutTexInfo, PropertyBakeOutput, MaterialProperty, MaterialInterface);
+		return StoreBakedPropertyTexture(Builder, OutTexInfo, PropertyBakeOutput, Property, Material);
 	}
 	else
 	{
@@ -799,7 +798,7 @@ bool FGLTFMaterialConverter::TryGetBakedMaterialProperty(FGLTFConvertBuilder& Bu
 	}
 }
 
-FGLTFPropertyBakeOutput FGLTFMaterialConverter::BakeMaterialProperty(EMaterialProperty MaterialProperty, const UMaterialInterface* MaterialInterface, const FIntPoint* PreferredTextureSize, bool bCopyAlphaFromRedChannel) const
+FGLTFPropertyBakeOutput FGLTFMaterialConverter::BakeMaterialProperty(EMaterialProperty Property, const UMaterialInterface* Material, const FIntPoint* PreferredTextureSize, bool bCopyAlphaFromRedChannel) const
 {
 	// TODO: make default baking-resolution configurable
 	const FIntPoint DefaultTextureSize(512, 512);
@@ -810,18 +809,18 @@ FGLTFPropertyBakeOutput FGLTFMaterialConverter::BakeMaterialProperty(EMaterialPr
 
 	const FGLTFPropertyBakeOutput PropertyBakeOutput = FGLTFMaterialUtility::BakeMaterialProperty(
 		TextureSize,
-		MaterialProperty,
-		MaterialInterface,
+		Property,
+		Material,
 		bCopyAlphaFromRedChannel);
 
 	return PropertyBakeOutput;
 }
 
-bool FGLTFMaterialConverter::StoreBakedPropertyTexture(FGLTFConvertBuilder& Builder, FGLTFJsonTextureInfo& OutTexInfo, const FGLTFPropertyBakeOutput& PropertyBakeOutput, EMaterialProperty MaterialProperty, const UMaterialInterface* MaterialInterface) const
+bool FGLTFMaterialConverter::StoreBakedPropertyTexture(FGLTFConvertBuilder& Builder, FGLTFJsonTextureInfo& OutTexInfo, const FGLTFPropertyBakeOutput& PropertyBakeOutput, EMaterialProperty Property, const UMaterialInterface* Material) const
 {
 	TCHAR* PropertyName;
 
-	switch (MaterialProperty)
+	switch (Property)
 	{
 		case MP_BaseColor: PropertyName = TEXT("BaseColor"); break;
 		case MP_Normal: PropertyName = TEXT("Normal"); break;
@@ -832,7 +831,7 @@ bool FGLTFMaterialConverter::StoreBakedPropertyTexture(FGLTFConvertBuilder& Buil
 			return false;
 	}
 
-	const FString TextureName = MaterialInterface->GetName() + TEXT("_") + PropertyName;
+	const FString TextureName = Material->GetName() + TEXT("_") + PropertyName;
 
 	// TODO: add support for detecting the correct tex-coord for this property based on connected nodes
 	const uint32 TexCoord = 0;
diff --git a/Engine/Plugins/Enterprise/GLTFExporter/MigratedFromGit/Source/GLTFExporter/Private/Converters/GLTFMaterialConverters.h b/Engine/Plugins/Enterprise/GLTFExporter/MigratedFromGit/Source/GLTFExporter/Private/Converters/GLTFMaterialConverters.h
--- a/Engine/Plugins/Enterprise/GLTFExporter/MigratedFromGit/Source/GLTFExporter/Private/Converters/GLTFMaterialConverters.h
+++ b/Engine/Plugins/Enterprise/GLTFExporter/MigratedFromGit/Source/GLTFExporter/Private/Converters/GLTFMaterialConverters.h
@@ -13,23 +13,23 @@ class FGLTFMaterialConverter final : public TGLTFConverter<FGLTFJsonMaterialInde
 {
 	FGLTFJsonMaterialIndex Add(FGLTFConvertBuilder& Builder, const FString& Name, const UMaterialInterface* Material) override;
 
-	bool TryGetBaseColorAndOpacity(FGLTFConvertBuilder& Builder, FGLTFJsonPBRMetallicRoughness& OutPBRParams, const UMaterialInterface* MaterialInterface) const;
-	bool TryGetMetallicAndRoughness(FGLTFConvertBuilder& Builder, FGLTFJsonPBRMetallicRoughness& OutPBRParams, const UMaterialInterface* MaterialInterface) const;
+	bool TryGetBaseColorAndOpacity(FGLTFConvertBuilder& Builder, FGLTFJsonPBRMetallicRoughness& OutPBRParams, const UMaterialInterface* Material) const;
+	bool TryGetMetallicAndRoughness(FGLTFConvertBuilder& Builder, FGLTFJsonPBRMetallicRoughness& OutPBRParams, const UMaterialInterface* Material) const;
 
-	bool IsPropertyNonDefault(EMaterialProperty MaterialProperty, const UMaterialInterface* MaterialInterface) const;
-	bool TryGetConstantColor(FGLTFJsonColor3& OutValue, EMaterialProperty MaterialProperty, const UMaterialInterface* MaterialInterface) const;
-	bool TryGetConstantColor(FGLTFJsonColor4& OutValue, EMaterialProperty MaterialProperty, const UMaterialInterface* MaterialInterface) const;
-	bool TryGetConstantColor(FLinearColor& OutValue, EMaterialProperty MaterialProperty, const UMaterialInterface* MaterialInterface) const;
-	bool TryGetConstantScalar(float& OutValue, EMaterialProperty MaterialProperty, const UMaterialInterface* MaterialInterface) const;
+	bool IsPropertyNonDefault(EMaterialProperty Property, const UMaterialInterface* Material) const;
+	bool TryGetConstantColor(FGLTFJsonColor3& OutValue, EMaterialProperty Property, const UMaterialInterface* Material) const;
+	bool TryGetConstantColor(FGLTFJsonColor4& OutValue, EMaterialProperty Property, const UMaterialInterface* Material) const;
+	bool TryGetConstantColor(FLinearColor& OutValue, EMaterialProperty Property, const UMaterialInterface* Material) const;
+	bool TryGetConstantScalar(float& OutValue, EMaterialProperty Property, const UMaterialInterface* Material) const;
 
-	bool TryGetSourceTexture(FGLTFConvertBuilder& Builder, FGLTFJsonTextureInfo& OutTexInfo, EMaterialProperty MaterialProperty, const UMaterialInterface* MaterialInterface, const TArray<FLinearColor>& AllowedMasks = {}) const;
-	bool TryGetSourceTexture(const UTexture2D*& OutTexture, int32& OutTexCoord, EMaterialProperty MaterialProperty, const UMaterialInterface* MaterialInterface, const TArray<FLinearColor>& AllowedMasks = {}) const;
+	bool TryGetSourceTexture(FGLTFConvertBuilder& Builder, FGLTFJsonTextureInfo& OutTexInfo, EMaterialProperty Property, const UMaterialInterface* Material, const TArray<FLinearColor>& AllowedMasks = {}) const;
+	bool TryGetSourceTexture(const UTexture2D*& OutTexture, int32& OutTexCoord, EMaterialProperty Property, const UMaterialInterface* Material, const TArray<FLinearColor>& AllowedMasks = {}) const;
 
-	bool TryGetBakedMaterialProperty(FGLTFConvertBuilder& Builder, FGLTFJsonTextureInfo& OutTexInfo, FGLTFJsonColor3& OutConstant, EMaterialProperty MaterialProperty, const UMaterialInterface* MaterialInterface) const;
-	bool TryGetBakedMaterialProperty(FGLTFConvertBuilder& Builder, FGLTFJsonTextureInfo& OutTexInfo, FGLTFJsonColor4& OutConstant, EMaterialProperty MaterialProperty, const UMaterialInterface* MaterialInterface) const;
-	bool TryGetBakedMaterialProperty(FGLTFConvertBuilder& Builder, FGLTFJsonTextureInfo& OutTexInfo, EMaterialProperty MaterialProperty, const UMaterialInterface* MaterialInterface) const;
+	bool TryGetBakedMaterialProperty(FGLTFConvertBuilder& Builder, FGLTFJsonTextureInfo& OutTexInfo, FGLTFJsonColor3& OutConstant, EMaterialProperty Property, const UMaterialInterface* Material) const;
+	bool TryGetBakedMaterialProperty(FGLTFConvertBuilder& Builder, FGLTFJsonTextureInfo& OutTexInfo, FGLTFJsonColor4& OutConstant, EMaterialProperty Property, const UMaterialInterface* Material) const;
+	bool TryGetBakedMaterialProperty(FGLTFConvertBuilder& Builder, FGLTFJsonTextureInfo& OutTexInfo, EMaterialProperty Property, const UMaterialInterface* Material) const;
 
-	FGLTFPropertyBakeOutput BakeMaterialProperty(EMaterialProperty MaterialProperty, const UMaterialInterface* MaterialInterface, const FIntPoint* PreferredTextureSize = nullptr, bool bCopyAlphaFromRedChannel = false) const;
+	FGLTFPropertyBakeOutput BakeMaterialProperty(EMaterialProperty Property, const UMaterialInterface* Material, const FIntPoint* PreferredTextureSize = nullptr, bool bCopyAlphaFromRedChannel = false) const;
 
-	bool StoreBakedPropertyTexture(FGLTFConvertBuilder& Builder, FGLTFJsonTextureInfo& OutTexInfo, const FGLTFPropertyBakeOutput& PropertyBakeOutput, EMaterialProperty MaterialProperty, const UMaterialInterface* MaterialInterface) const;
+	bool StoreBakedPropertyTexture(FGLTFConvertBuilder& Builder, FGLTFJsonTextureInfo& OutTexInfo, const FGLTFPropertyBakeOutput& PropertyBakeOutput, EMaterialProperty Property, const UMaterialInterface* Material) const;
 };
diff --git a/Engine/Plugins/Enterprise/GLTFExporter/MigratedFromGit/Source/GLTFExporter/Private/Converters/GLTFMaterialUtility.cpp b/Engine/Plugins/Enterprise/GLTFExporter/MigratedFromGit/Source/GLTFExporter/Private/Converters/GLTFMaterialUtility.cpp
--- a/Engine/Plugins/Enterprise/GLTFExporter/MigratedFromGit/Source/GLTFExporter/Private/Converters/GLTFMaterialUtility.cpp
+++ b/Engine/Plugins/Enterprise/GLTFExporter/MigratedFromGit/Source/GLTFExporter/Private/Converters/GLTFMaterialUtility.cpp
@@ -125,7 +125,7 @@ bool FGLTFMaterialUtility::CombineTextures(TArray<FColor>& OutPixels, const TArr
 	return bReadSuccessful;
 }
 
-FGLTFPropertyBakeOutput FGLTFMaterialUtility::BakeMaterialProperty(const FIntPoint& OutputSize, EMaterialProperty MaterialProperty, const UMaterialInterface* Material, bool bCopyAlphaFromRedChannel)
+FGLTFPropertyBakeOutput FGLTFMaterialUtility::BakeMaterialProperty(const FIntPoint& OutputSize, EMaterialProperty Property, const UMaterialInterface* Material, bool bCopyAlphaFromRedChannel)
 {
 	TArray<FMeshData*> MeshSettings;
 
@@ -140,7 +140,7 @@ FGLTFPropertyBakeOutput FGLTFMaterialUtility::BakeMaterialProperty(const FIntPoi
 
 	FMaterialData MatSet;
 	MatSet.Material = const_cast<UMaterialInterface*>(Material);
-	MatSet.PropertySizes.Add(MaterialProperty, OutputSize);
+	MatSet.PropertySizes.Add(Property, OutputSize);
 
 	MatSettings.Add(&MatSet);
 
@@ -151,8 +151,8 @@ FGLTFPropertyBakeOutput FGLTFMaterialUtility::BakeMaterialProperty(const FIntPoi
 
 	FBakeOutput BakeOutput = BakeOutputs[0];
 
-	TArray<FColor> BakedPixels = MoveTemp(BakeOutput.PropertyData.FindChecked(MaterialProperty));
-	const FIntPoint BakedSize = BakeOutput.PropertySizes.FindChecked(MaterialProperty);
+	TArray<FColor> BakedPixels = MoveTemp(BakeOutput.PropertyData.FindChecked(Property));
+	const FIntPoint BakedSize = BakeOutput.PropertySizes.FindChecked(Property);
 	const float EmissiveScale = BakeOutput.EmissiveScale;
 
 	if (bCopyAlphaFromRedChannel)
@@ -172,7 +172,7 @@ FGLTFPropertyBakeOutput FGLTFMaterialUtility::BakeMaterialProperty(const FIntPoi
 		}
 	}
 
-	FGLTFPropertyBakeOutput PropertyBakeOutput(MaterialProperty, PF_B8G8R8A8, BakedPixels, BakedSize, EmissiveScale);
+	FGLTFPropertyBakeOutput PropertyBakeOutput(Property, PF_B8G8R8A8, BakedPixels, BakedSize, EmissiveScale);
 
 	if (BakedPixels.Num() == 1)
 	{
diff --git a/Engine/Plugins/Enterprise/GLTFExporter/MigratedFromGit/Source/GLTFExporter/Private/Converters/GLTFMaterialUtility.h b/Engine/Plugins/Enterprise/GLTFExporter/MigratedFromGit/Source/GLTFExporter/Private/Converters/GLTFMaterialUtility.h
--- a/Engine/Plugins/Enterprise/GLTFExporter/MigratedFromGit/Source/GLTFExporter/Private/Converters/GLTFMaterialUtility.h
+++ b/Engine/Plugins/Enterprise/GLTFExporter/MigratedFromGit/Source/GLTFExporter/Private/Converters/GLTFMaterialUtility.h
@@ -48,7 +48,7 @@ struct FGLTFMaterialUtility
 	static UTexture2D* CreateTransientTexture(const FGLTFPropertyBakeOutput& PropertyBakeOutput, bool bUseSRGB = false);
 	static UTexture2D* CreateTransientTexture(const TArray<FColor>& Pixels, const FIntPoint& TextureSize, EPixelFormat TextureFormat, bool bUseSRGB = false);
 	static bool CombineTextures(TArray<FColor>& OutPixels, const TArray<FGLTFTextureCombineSource>& Sources, const FIntPoint& OutputSize, EPixelFormat OutputPixelFormat);
-	static FGLTFPropertyBakeOutput BakeMaterialProperty(const FIntPoint& OutputSize, EMaterialProperty MaterialProperty, const UMaterialInterface* Material, bool bCopyAlphaFromRedChannel = false);
+	static FGLTFPropertyBakeOutput BakeMaterialProperty(const FIntPoint& OutputSize, EMaterialProperty Property, const UMaterialInterface* Material, bool bCopyAlphaFromRedChannel = false);
 
 	static FGLTFJsonTextureIndex AddCombinedTexture(FGLTFConvertBuilder& Builder, const TArray<FGLTFTextureCombineSource>& CombineSources, const FIntPoint& TextureSize, const FString& TextureName, EGLTFJsonTextureFilter MinFilter, EGLTFJsonTextureFilter MagFilter, EGLTFJsonTextureWrap WrapS, EGLTFJsonTextureWrap WrapT);
 	static FGLTFJsonTextureIndex AddTexture(FGLTFConvertBuilder& Builder, const TArray<FColor>& Pixels, const FIntPoint& TextureSize, const FString& TextureName, EPixelFormat PixelFormat, EGLTFJsonTextureFilter MinFilter, EGLTFJsonTextureFilter MagFilter, EGLTFJsonTextureWrap WrapS, EGLTFJsonTextureWrap WrapT);
diff --git a/Engine/Plugins/Enterprise/GLTFExporter/MigratedFromGit/Source/GLTFExporter/Private/Exporters/GLTFExporterUtility.h b/Engine/Plugins/Enterprise/GLTFExporter/MigratedFromGit/Source/GLTFExporter/Private/Exporters/GLTFExporterUtility.h
--- a/Engine/Plugins/Enterprise/GLTFExporter/MigratedFromGit/Source/GLTFExporter/Private/Exporters/GLTFExporterUtility.h
+++ b/Engine/Plugins/Enterprise/GLTFExporter/MigratedFromGit/Source/GLTFExporter/Private/Exporters/GLTFExporterUtility.h
@@ -7,5 +7,5 @@ class UMaterialInterface;
 
 struct FGLTFExporterUtility
 {
-	static const UStaticMesh* GetPreviewMesh(const UMaterialInterface* MaterialInterface);
+	static const UStaticMesh* GetPreviewMesh(const UMaterialInterface* Material);
 };
