[GLTFExporter] fix(material): make sure to always set alpha = 1 when alpha mode is opaque.

commit ab27c38a40c247b9fb8a5664e924d6af684cbeba
Date: Sun, 14 Jun 2020 17:32:38 +0200

    fix(material): make sure to always set alpha = 1 when alpha mode is opaque.
    (cherry picked from commit 6d648214b4b08c7a753523f73a9ec69d4722dd81​)

#skipci

diff --git a/Engine/Plugins/Enterprise/GLTFExporter/MigratedFromGit/Source/GLTFExporter/Private/Converters/GLTFMaterialConverters.cpp b/Engine/Plugins/Enterprise/GLTFExporter/MigratedFromGit/Source/GLTFExporter/Private/Converters/GLTFMaterialConverters.cpp
--- a/Engine/Plugins/Enterprise/GLTFExporter/MigratedFromGit/Source/GLTFExporter/Private/Converters/GLTFMaterialConverters.cpp
+++ b/Engine/Plugins/Enterprise/GLTFExporter/MigratedFromGit/Source/GLTFExporter/Private/Converters/GLTFMaterialConverters.cpp
@@ -30,14 +30,7 @@ FGLTFJsonMaterialIndex FGLTFMaterialConverter::Add(FGLTFConvertBuilder& Builder,
 
 	// TODO: check if a property is active before trying to get it (i.e. Material->IsPropertyActive)
 
-	if (JsonMaterial.AlphaMode == EGLTFJsonAlphaMode::Blend || JsonMaterial.AlphaMode == EGLTFJsonAlphaMode::Mask)
-	{
-		if (!TryGetBaseColorAndOpacity(Builder, JsonMaterial.PBRMetallicRoughness, Material))
-		{
-			// TODO: handle failure?
-		}
-	}
-	else
+	if (JsonMaterial.AlphaMode == EGLTFJsonAlphaMode::Opaque)
 	{
 		if (!TryGetConstantColor(JsonMaterial.PBRMetallicRoughness.BaseColorFactor, MP_BaseColor, Material))
 		{
@@ -49,6 +42,15 @@ FGLTFJsonMaterialIndex FGLTFMaterialConverter::Add(FGLTFConvertBuilder& Builder,
 				}
 			}
 		}
+
+		JsonMaterial.PBRMetallicRoughness.BaseColorFactor.A = 1.0f; // make sure base color is opaque
+	}
+	else
+	{
+		if (!TryGetBaseColorAndOpacity(Builder, JsonMaterial.PBRMetallicRoughness, Material))
+		{
+			// TODO: handle failure?
+		}
 	}
 
 	if (!TryGetMetallicAndRoughness(Builder, JsonMaterial.PBRMetallicRoughness, Material))
