[GLTFExporter] feat(logging): add new base-class for all builders, containing logging functionality

commit a85e006f38f0e15cf433f7ac7e36c8ec7c5e0982
Date: Tue, 16 Jun 2020 15:24:19 +0200

    feat(logging): add new base-class for all builders, containing logging functionality
    (cherry picked from commit b332cee4ecf332302ff0a1d6205eb0cf185f0dca​)

#skipci

diff --git a/Engine/Plugins/Enterprise/GLTFExporter/MigratedFromGit/Source/GLTFExporter/Private/Builders/GLTFBuilder.cpp b/Engine/Plugins/Enterprise/GLTFExporter/MigratedFromGit/Source/GLTFExporter/Private/Builders/GLTFBuilder.cpp
new file mode 100644
--- /dev/null
+++ b/Engine/Plugins/Enterprise/GLTFExporter/MigratedFromGit/Source/GLTFExporter/Private/Builders/GLTFBuilder.cpp
@@ -0,0 +1,153 @@
+// Copyright Epic Games, Inc. All Rights Reserved.
+
+#include "Builders/GLTFBuilder.h"
+#include "MessageLogModule.h"
+#include "IMessageLogListing.h"
+
+DEFINE_LOG_CATEGORY(LogGLTFExporter);
+
+namespace
+{
+	// TODO: turn into config variable?
+	const TCHAR* LogListingName = TEXT("GLTFExporter");
+	const TCHAR* LogListingLabel = TEXT("GLTF Exporter");
+} // anonymous namespace
+
+FGLTFBuilder::FGLTFBuilder()
+{
+}
+
+FGLTFBuilder::~FGLTFBuilder()
+{
+}
+
+void FGLTFBuilder::ClearLogMessages()
+{
+	LogMessages.Empty();
+}
+
+void FGLTFBuilder::AddLogMessage(EGLTFMessageSeverity Severity, const FString& Message)
+{
+	LogMessages.Emplace(Severity, Message);
+}
+
+void FGLTFBuilder::AddInfoMessage(const FString& Message)
+{
+	AddLogMessage(EGLTFMessageSeverity::Info, Message);
+}
+
+void FGLTFBuilder::AddWarningMessage(const FString& Message)
+{
+	AddLogMessage(EGLTFMessageSeverity::Warning, Message);
+}
+
+void FGLTFBuilder::AddErrorMessage(const FString& Message)
+{
+	AddLogMessage(EGLTFMessageSeverity::Error, Message);
+}
+
+const TArray<FGLTFBuilder::FLogMessage>& FGLTFBuilder::GetLogMessages() const
+{
+	return LogMessages;
+}
+
+TArray<FGLTFBuilder::FLogMessage> FGLTFBuilder::GetLogMessages(EGLTFMessageSeverity Severity) const
+{
+	return LogMessages.FilterByPredicate(
+		[&Severity](const FLogMessage& LogMessage)
+		{
+			return LogMessage.Key == Severity;
+		});
+}
+
+TArray<FGLTFBuilder::FLogMessage> FGLTFBuilder::GetInfoMessages() const
+{
+	return GetLogMessages(EGLTFMessageSeverity::Info);
+}
+
+TArray<FGLTFBuilder::FLogMessage> FGLTFBuilder::GetWarningMessages() const
+{
+	return GetLogMessages(EGLTFMessageSeverity::Warning);
+}
+
+TArray<FGLTFBuilder::FLogMessage> FGLTFBuilder::GetErrorMessages() const
+{
+	return GetLogMessages(EGLTFMessageSeverity::Error);
+}
+
+int32 FGLTFBuilder::GetLogMessageCount() const
+{
+	return LogMessages.Num();
+}
+
+int32 FGLTFBuilder::GetInfoMessageCount() const
+{
+	return GetInfoMessages().Num();
+}
+
+int32 FGLTFBuilder::GetWarningMessageCount() const
+{
+	return GetWarningMessages().Num();
+}
+
+int32 FGLTFBuilder::GetErrorMessageCount() const
+{
+	return GetErrorMessages().Num();
+}
+
+void FGLTFBuilder::ShowLogMessages() const
+{
+	if (LogMessages.Num() > 0)
+	{
+		FMessageLogModule& MessageLogModule = FModuleManager::LoadModuleChecked<FMessageLogModule>("MessageLog");
+		const TSharedRef<IMessageLogListing> LogListing = MessageLogModule.GetLogListing(LogListingName);
+
+		LogListing->SetLabel(FText::FromString(LogListingLabel));
+		LogListing->ClearMessages();
+
+		for (const FLogMessage& Message : LogMessages)
+		{
+			LogListing->AddMessage(CreateTokenizedMessage(Message));
+		}
+
+		MessageLogModule.OpenMessageLog(LogListingName);
+	}
+}
+
+void FGLTFBuilder::WriteLogMessagesToConsole() const
+{
+	for (const FLogMessage& Message : LogMessages)
+	{
+		WriteLogMessageToConsole(Message);
+	}
+}
+
+void FGLTFBuilder::WriteLogMessageToConsole(const FLogMessage& LogMessage)
+{
+	switch (LogMessage.Key)
+	{
+		case EGLTFMessageSeverity::Info: UE_LOG(LogGLTFExporter, Display, TEXT("%s"), *LogMessage.Value); break;
+		case EGLTFMessageSeverity::Warning: UE_LOG(LogGLTFExporter, Warning, TEXT("%s"), *LogMessage.Value); break;
+		case EGLTFMessageSeverity::Error: UE_LOG(LogGLTFExporter, Error, TEXT("%s"), *LogMessage.Value); break;
+		default:
+			checkNoEntry();
+			break;
+	}
+}
+
+TSharedRef<FTokenizedMessage> FGLTFBuilder::CreateTokenizedMessage(const FLogMessage& LogMessage)
+{
+	EMessageSeverity::Type MessageSeverity = EMessageSeverity::Type::CriticalError;
+
+	switch (LogMessage.Key)
+	{
+		case EGLTFMessageSeverity::Info: MessageSeverity = EMessageSeverity::Type::Info; break;
+		case EGLTFMessageSeverity::Warning: MessageSeverity = EMessageSeverity::Type::Warning; break;
+		case EGLTFMessageSeverity::Error: MessageSeverity = EMessageSeverity::Type::Error; break;
+		default:
+			checkNoEntry();
+			break;
+	}
+
+	return FTokenizedMessage::Create(MessageSeverity, FText::FromString(LogMessage.Value));
+}
diff --git a/Engine/Plugins/Enterprise/GLTFExporter/MigratedFromGit/Source/GLTFExporter/Private/Builders/GLTFBuilder.h b/Engine/Plugins/Enterprise/GLTFExporter/MigratedFromGit/Source/GLTFExporter/Private/Builders/GLTFBuilder.h
new file mode 100644
--- /dev/null
+++ b/Engine/Plugins/Enterprise/GLTFExporter/MigratedFromGit/Source/GLTFExporter/Private/Builders/GLTFBuilder.h
@@ -0,0 +1,62 @@
+// Copyright Epic Games, Inc. All Rights Reserved.
+
+#pragma once
+
+DECLARE_LOG_CATEGORY_EXTERN(LogGLTFExporter, Log, All);
+
+enum class EGLTFMessageSeverity
+{
+	Info,
+	Warning,
+	Error
+};
+
+class FGLTFBuilder
+{
+public:
+
+	FGLTFBuilder();
+	virtual ~FGLTFBuilder();
+
+	typedef TTuple<EGLTFMessageSeverity, FString> FLogMessage;
+
+	void ClearLogMessages();
+
+	void AddLogMessage(EGLTFMessageSeverity Severity, const FString& Message);
+
+	void AddInfoMessage(const FString& Message);
+
+	void AddWarningMessage(const FString& Message);
+
+	void AddErrorMessage(const FString& Message);
+
+	const TArray<FLogMessage>& GetLogMessages() const;
+
+	TArray<FLogMessage> GetLogMessages(EGLTFMessageSeverity Severity) const;
+
+	TArray<FLogMessage> GetInfoMessages() const;
+
+	TArray<FLogMessage> GetWarningMessages() const;
+
+	TArray<FLogMessage> GetErrorMessages() const;
+
+	int32 GetLogMessageCount() const;
+
+	int32 GetInfoMessageCount() const;
+
+	int32 GetWarningMessageCount() const;
+
+	int32 GetErrorMessageCount() const;
+
+	void ShowLogMessages() const;
+
+	void WriteLogMessagesToConsole() const;
+
+protected:
+
+	TArray<FLogMessage> LogMessages;
+
+	static void WriteLogMessageToConsole(const FLogMessage& LogMessage);
+
+	static TSharedRef<FTokenizedMessage> CreateTokenizedMessage(const FLogMessage& LogMessage);
+};
diff --git a/Engine/Plugins/Enterprise/GLTFExporter/MigratedFromGit/Source/GLTFExporter/Private/Builders/GLTFJsonBuilder.cpp b/Engine/Plugins/Enterprise/GLTFExporter/MigratedFromGit/Source/GLTFExporter/Private/Builders/GLTFJsonBuilder.cpp
--- a/Engine/Plugins/Enterprise/GLTFExporter/MigratedFromGit/Source/GLTFExporter/Private/Builders/GLTFJsonBuilder.cpp
+++ b/Engine/Plugins/Enterprise/GLTFExporter/MigratedFromGit/Source/GLTFExporter/Private/Builders/GLTFJsonBuilder.cpp
@@ -7,10 +7,6 @@ FGLTFJsonBuilder::FGLTFJsonBuilder()
 {
 }
 
-FGLTFJsonBuilder::~FGLTFJsonBuilder()
-{
-}
-
 FGLTFJsonAccessorIndex FGLTFJsonBuilder::AddAccessor(const FGLTFJsonAccessor& JsonAccessor)
 {
 	return FGLTFJsonAccessorIndex(JsonRoot.Accessors.Add(JsonAccessor));
diff --git a/Engine/Plugins/Enterprise/GLTFExporter/MigratedFromGit/Source/GLTFExporter/Private/Builders/GLTFJsonBuilder.h b/Engine/Plugins/Enterprise/GLTFExporter/MigratedFromGit/Source/GLTFExporter/Private/Builders/GLTFJsonBuilder.h
--- a/Engine/Plugins/Enterprise/GLTFExporter/MigratedFromGit/Source/GLTFExporter/Private/Builders/GLTFJsonBuilder.h
+++ b/Engine/Plugins/Enterprise/GLTFExporter/MigratedFromGit/Source/GLTFExporter/Private/Builders/GLTFJsonBuilder.h
@@ -3,15 +3,15 @@
 #pragma once
 
 #include "Json/GLTFJsonRoot.h"
+#include "Builders/GLTFBuilder.h"
 
-class FGLTFJsonBuilder
+class FGLTFJsonBuilder : public FGLTFBuilder
 {
 public:
 
 	FGLTFJsonSceneIndex& DefaultScene;
 
 	FGLTFJsonBuilder();
-	virtual ~FGLTFJsonBuilder();
 
 	FGLTFJsonAccessorIndex AddAccessor(const FGLTFJsonAccessor& JsonAccessor);
 	FGLTFJsonBufferIndex AddBuffer(const FGLTFJsonBuffer& JsonBuffer);
diff --git a/Engine/Plugins/Enterprise/GLTFExporter/MigratedFromGit/Source/GLTFExporter/Private/GLTFExporterModule.cpp b/Engine/Plugins/Enterprise/GLTFExporter/MigratedFromGit/Source/GLTFExporter/Private/GLTFExporterModule.cpp
--- a/Engine/Plugins/Enterprise/GLTFExporter/MigratedFromGit/Source/GLTFExporter/Private/GLTFExporterModule.cpp
+++ b/Engine/Plugins/Enterprise/GLTFExporter/MigratedFromGit/Source/GLTFExporter/Private/GLTFExporterModule.cpp
@@ -5,8 +5,6 @@
 #include "CoreMinimal.h"
 #include "Modules/ModuleManager.h"
 
-DEFINE_LOG_CATEGORY(LogGLTFExporter);
-
 /**
  * glTF Exporter module implementation (private)
  */
diff --git a/Engine/Plugins/Enterprise/GLTFExporter/MigratedFromGit/Source/GLTFExporter/Public/GLTFExporterModule.h b/Engine/Plugins/Enterprise/GLTFExporter/MigratedFromGit/Source/GLTFExporter/Public/GLTFExporterModule.h
--- a/Engine/Plugins/Enterprise/GLTFExporter/MigratedFromGit/Source/GLTFExporter/Public/GLTFExporterModule.h
+++ b/Engine/Plugins/Enterprise/GLTFExporter/MigratedFromGit/Source/GLTFExporter/Public/GLTFExporterModule.h
@@ -8,8 +8,6 @@
 
 #define GLTFEXPORTER_MODULE_NAME TEXT("GLTFExporter")
 
-DECLARE_LOG_CATEGORY_EXTERN(LogGLTFExporter, Log, All);
-
 /**
  * The public interface of the GLTFExporter module
  */
