[GLTFExporter] feat(logging): display logs (or write to console) when the export is done

commit dadd692037a824cd936b6aa7c944960484b55d86
Date: Tue, 16 Jun 2020 15:25:34 +0200

    feat(logging): display logs (or write to console) when the export is done
    (cherry picked from commit 7f39779acb62fa0fbf2b7cfe7ba3234485485965​)

#skipci

diff --git a/Engine/Plugins/Enterprise/GLTFExporter/MigratedFromGit/Source/GLTFExporter/Private/Exporters/GLTFExporter.cpp b/Engine/Plugins/Enterprise/GLTFExporter/MigratedFromGit/Source/GLTFExporter/Private/Exporters/GLTFExporter.cpp
--- a/Engine/Plugins/Enterprise/GLTFExporter/MigratedFromGit/Source/GLTFExporter/Private/Exporters/GLTFExporter.cpp
+++ b/Engine/Plugins/Enterprise/GLTFExporter/MigratedFromGit/Source/GLTFExporter/Private/Exporters/GLTFExporter.cpp
@@ -32,13 +32,35 @@ bool UGLTFExporter::ExportBinary(UObject* Object, const TCHAR* Type, FArchive& A
 	}
 
 	FGLTFContainerBuilder Builder(bSelectedOnly);
-	if (!AddObject(Builder, Object))
+
+	bool bSuccess = true;
+
+	if (AddObject(Builder, Object))
 	{
-		// TODO: Report error
-		return false;
+		if (!Builder.Serialize(Archive, CurrentFilename))
+		{
+			// TODO: more descriptive error
+			Builder.AddErrorMessage(TEXT("Serialize failed"));
+			bSuccess = false;
+		}
+	}
+	else
+	{
+		// TODO: more descriptive error
+		Builder.AddErrorMessage(TEXT("AddObject failed"));
+		bSuccess = false;
+	}
+
+	if (FApp::IsUnattended())
+	{
+		Builder.WriteLogMessagesToConsole();
+	}
+	else
+	{
+		Builder.ShowLogMessages();
 	}
 
-	return Builder.Serialize(Archive, CurrentFilename);
+	return bSuccess;
 }
 
 bool UGLTFExporter::AddObject(FGLTFContainerBuilder& Builder, const UObject* Object)
