[GLTFExporter] feat: return false from UGLTFMaterialExporter on failure

commit 95ab9bd7bddbfd2c220253e923c8e563b143d693
Date: Tue, 16 Jun 2020 15:47:53 +0200

    feat: return false from UGLTFMaterialExporter on failure
    (cherry picked from commit 7396b22988d44d9eebc2467493c55c8b82640cfe​)

#skipci

diff --git a/Engine/Plugins/Enterprise/GLTFExporter/MigratedFromGit/Source/GLTFExporter/Private/Exporters/GLTFMaterialExporter.cpp b/Engine/Plugins/Enterprise/GLTFExporter/MigratedFromGit/Source/GLTFExporter/Private/Exporters/GLTFMaterialExporter.cpp
--- a/Engine/Plugins/Enterprise/GLTFExporter/MigratedFromGit/Source/GLTFExporter/Private/Exporters/GLTFMaterialExporter.cpp
+++ b/Engine/Plugins/Enterprise/GLTFExporter/MigratedFromGit/Source/GLTFExporter/Private/Exporters/GLTFMaterialExporter.cpp
@@ -16,6 +16,11 @@ bool UGLTFMaterialExporter::AddObject(FGLTFContainerBuilder& Builder, const UObj
 	const UMaterialInterface* Material = CastChecked<UMaterialInterface>(Object);
 	const FGLTFJsonMaterialIndex MaterialIndex = Builder.GetOrAddMaterial(Material);
 
+	if (MaterialIndex == INDEX_NONE)
+	{
+		return false;
+	}
+
 	const UStaticMesh* PreviewMesh = FGLTFExporterUtility::GetPreviewMesh(Material);
 	if (PreviewMesh != nullptr)
 	{
@@ -31,6 +36,10 @@ bool UGLTFMaterialExporter::AddObject(FGLTFContainerBuilder& Builder, const UObj
 
 		Builder.DefaultScene = SceneIndex;
 	}
+	else
+	{
+		// TODO: should we report an error if no preview mesh was found?
+	}
 
 	return true;
 }
