[GLTFExporter] feat: return false from UGLTFLevelExporter on failure

commit 75618d5c7dc626d3d414baec0de153bf351c9acf
Date: Tue, 16 Jun 2020 15:48:00 +0200

    feat: return false from UGLTFLevelExporter on failure
    (cherry picked from commit e68540eef99a22d13311f4e29e64bfb437cc954f​)

#skipci

diff --git a/Engine/Plugins/Enterprise/GLTFExporter/MigratedFromGit/Source/GLTFExporter/Private/Exporters/GLTFLevelExporter.cpp b/Engine/Plugins/Enterprise/GLTFExporter/MigratedFromGit/Source/GLTFExporter/Private/Exporters/GLTFLevelExporter.cpp
--- a/Engine/Plugins/Enterprise/GLTFExporter/MigratedFromGit/Source/GLTFExporter/Private/Exporters/GLTFLevelExporter.cpp
+++ b/Engine/Plugins/Enterprise/GLTFExporter/MigratedFromGit/Source/GLTFExporter/Private/Exporters/GLTFLevelExporter.cpp
@@ -15,6 +15,13 @@ bool UGLTFLevelExporter::AddObject(FGLTFContainerBuilder& Builder, const UObject
 	const UWorld* World = CastChecked<UWorld>(Object);
 	const FGLTFJsonSceneIndex SceneIndex = Builder.GetOrAddScene(World);
 
-	Builder.DefaultScene = SceneIndex;
-	return true;
+	if (SceneIndex != INDEX_NONE)
+	{
+		Builder.DefaultScene = SceneIndex;
+		return true;
+	}
+	else
+	{
+		return false;
+	}
 }
