[GLTFExporter] feat: return false from UGLTFStaticMeshExporter on failure

commit 89e922b8483479a39c68ebb5652ebd4a4adcf81e
Date: Tue, 16 Jun 2020 15:47:44 +0200

    feat: return false from UGLTFStaticMeshExporter on failure
    (cherry picked from commit ae8d50cf30cdc2d52abbf085613987eb687cf1a3​)

#skipci

diff --git a/Engine/Plugins/Enterprise/GLTFExporter/MigratedFromGit/Source/GLTFExporter/Private/Exporters/GLTFStaticMeshExporter.cpp b/Engine/Plugins/Enterprise/GLTFExporter/MigratedFromGit/Source/GLTFExporter/Private/Exporters/GLTFStaticMeshExporter.cpp
--- a/Engine/Plugins/Enterprise/GLTFExporter/MigratedFromGit/Source/GLTFExporter/Private/Exporters/GLTFStaticMeshExporter.cpp
+++ b/Engine/Plugins/Enterprise/GLTFExporter/MigratedFromGit/Source/GLTFExporter/Private/Exporters/GLTFStaticMeshExporter.cpp
@@ -15,6 +15,11 @@ bool UGLTFStaticMeshExporter::AddObject(FGLTFContainerBuilder& Builder, const UO
 	const UStaticMesh* StaticMesh = CastChecked<UStaticMesh>(Object);
 	const FGLTFJsonMeshIndex MeshIndex = Builder.GetOrAddMesh(StaticMesh);
 
+	if (MeshIndex == INDEX_NONE)
+	{
+		return false;
+	}
+
 	FGLTFJsonNode Node;
 	Node.Mesh = MeshIndex;
 	const FGLTFJsonNodeIndex NodeIndex = Builder.AddNode(Node);
