[GLTFExporter] feat(logging): add warnings for (some) failures during material conversion

commit c4a6be28aa0a3c21bacbad1d8f19d4b19d86ebc5
Date: Tue, 16 Jun 2020 15:19:07 +0200

    feat(logging): add warnings for (some) failures during material conversion
    (cherry picked from commit 28d55aaa458d1fa29e71043598f688063b7a6554​)

#skipci

diff --git a/Engine/Plugins/Enterprise/GLTFExporter/MigratedFromGit/Source/GLTFExporter/Private/Converters/GLTFMaterialConverters.cpp b/Engine/Plugins/Enterprise/GLTFExporter/MigratedFromGit/Source/GLTFExporter/Private/Converters/GLTFMaterialConverters.cpp
--- a/Engine/Plugins/Enterprise/GLTFExporter/MigratedFromGit/Source/GLTFExporter/Private/Converters/GLTFMaterialConverters.cpp
+++ b/Engine/Plugins/Enterprise/GLTFExporter/MigratedFromGit/Source/GLTFExporter/Private/Converters/GLTFMaterialConverters.cpp
@@ -38,7 +38,7 @@ FGLTFJsonMaterialIndex FGLTFMaterialConverter::Add(FGLTFConvertBuilder& Builder,
 			{
 				if (!TryGetBakedMaterialProperty(Builder, JsonMaterial.PBRMetallicRoughness.BaseColorTexture, JsonMaterial.PBRMetallicRoughness.BaseColorFactor, MP_BaseColor, Material))
 				{
-					// TODO: handle failure?
+					Builder.AddWarningMessage(TEXT("Failed to export BaseColor for material ") + Material->GetName());
 				}
 			}
 		}
@@ -49,13 +49,13 @@ FGLTFJsonMaterialIndex FGLTFMaterialConverter::Add(FGLTFConvertBuilder& Builder,
 	{
 		if (!TryGetBaseColorAndOpacity(Builder, JsonMaterial.PBRMetallicRoughness, Material))
 		{
-			// TODO: handle failure?
+			Builder.AddWarningMessage(TEXT("Failed to export BaseColor & Opacity for material ") + Material->GetName());
 		}
 	}
 
 	if (!TryGetMetallicAndRoughness(Builder, JsonMaterial.PBRMetallicRoughness, Material))
 	{
-		// TODO: handle failure?
+		Builder.AddWarningMessage(TEXT("Failed to export Metallic & Roughness for material ") + Material->GetName());
 	}
 
 	// NOTE: export of EmissiveColor has been temporarily disabled because of visual differences that
@@ -77,7 +77,7 @@ FGLTFJsonMaterialIndex FGLTFMaterialConverter::Add(FGLTFConvertBuilder& Builder,
 			}
 			else
 			{
-				// TODO: handle failure?
+				Builder.AddWarningMessage(TEXT("Failed to export EmissiveColor for material ") + Material->GetName());
 			}
 		}
 		else
@@ -93,7 +93,7 @@ FGLTFJsonMaterialIndex FGLTFMaterialConverter::Add(FGLTFConvertBuilder& Builder,
 		{
 			if (!TryGetBakedMaterialProperty(Builder, JsonMaterial.NormalTexture, MP_Normal, Material))
 			{
-				// TODO: handle failure?
+				Builder.AddWarningMessage(TEXT("Failed to export Normal for material ") + Material->GetName());
 			}
 		}
 	}
@@ -104,7 +104,7 @@ FGLTFJsonMaterialIndex FGLTFMaterialConverter::Add(FGLTFConvertBuilder& Builder,
 		{
 			if (!TryGetBakedMaterialProperty(Builder, JsonMaterial.OcclusionTexture, MP_AmbientOcclusion, Material))
 			{
-				// TODO: handle failure?
+				Builder.AddWarningMessage(TEXT("Failed to export AmbientOcclusion for material ") + Material->GetName());
 			}
 		}
 	}
@@ -174,6 +174,11 @@ bool FGLTFMaterialConverter::TryGetBaseColorAndOpacity(FGLTFConvertBuilder& Buil
 		if (!bAreTexturesCompatible)
 		{
 			// TODO: handle differences in wrapping or uv-coords
+			Builder.AddWarningMessage(
+				TEXT("BaseColor- and Opacity-textures for material ") +
+				Material->GetName() +
+				TEXT(" were not able to be combined and will be skipped"));
+
 			return false;
 		}
 
@@ -314,6 +319,11 @@ bool FGLTFMaterialConverter::TryGetMetallicAndRoughness(FGLTFConvertBuilder& Bui
 		if (!bAreTexturesCompatible)
 		{
 			// TODO: handle differences in wrapping or uv-coords
+			Builder.AddWarningMessage(
+				TEXT("Metallic- and Roughness-textures for material ") +
+				Material->GetName() +
+				TEXT(" were not able to be combined and will be skipped"));
+
 			return false;
 		}
 
