[GLTFExporter] refactor(logging): add TODO for how to handle error-messages in otherwise successful exports

commit 13c1ff704c1fa2094b431a86965ee4869ffe429d
Date: Tue, 16 Jun 2020 16:07:36 +0200

    refactor(logging): add TODO for how to handle error-messages in otherwise successful exports
    (cherry picked from commit b39a292feaefc8a7a981f9e26e00b4d3448e4b6e​)

#skipci

diff --git a/Engine/Plugins/Enterprise/GLTFExporter/MigratedFromGit/Source/GLTFExporter/Private/Exporters/GLTFExporter.cpp b/Engine/Plugins/Enterprise/GLTFExporter/MigratedFromGit/Source/GLTFExporter/Private/Exporters/GLTFExporter.cpp
--- a/Engine/Plugins/Enterprise/GLTFExporter/MigratedFromGit/Source/GLTFExporter/Private/Exporters/GLTFExporter.cpp
+++ b/Engine/Plugins/Enterprise/GLTFExporter/MigratedFromGit/Source/GLTFExporter/Private/Exporters/GLTFExporter.cpp
@@ -51,6 +51,9 @@ bool UGLTFExporter::ExportBinary(UObject* Object, const TCHAR* Type, FArchive& A
 		bSuccess = false;
 	}
 
+	// TODO: should we set bSuccess to false if Builder.GetErrorMessageCount() > 0
+	// even if both AddObject and Serialize has reported success?
+
 	if (FApp::IsUnattended())
 	{
 		Builder.WriteLogMessagesToConsole();
