[GLTFExporter] refactor(logging): capture enum by value instead of by reference in lambda

commit 9f06009aa6dcfed9808fde360bfcf33ce1c306d2
Date: Thu, 25 Jun 2020 11:19:06 +0200

    refactor(logging): capture enum by value instead of by reference in lambda
    (cherry picked from commit f1a5e5df9a6a5ed681f41196f81077f3e67661f7​)

#skipci

diff --git a/Engine/Plugins/Enterprise/GLTFExporter/MigratedFromGit/Source/GLTFExporter/Private/Builders/GLTFBuilder.cpp b/Engine/Plugins/Enterprise/GLTFExporter/MigratedFromGit/Source/GLTFExporter/Private/Builders/GLTFBuilder.cpp
--- a/Engine/Plugins/Enterprise/GLTFExporter/MigratedFromGit/Source/GLTFExporter/Private/Builders/GLTFBuilder.cpp
+++ b/Engine/Plugins/Enterprise/GLTFExporter/MigratedFromGit/Source/GLTFExporter/Private/Builders/GLTFBuilder.cpp
@@ -54,7 +54,7 @@ const TArray<FGLTFBuilder::FLogMessage>& FGLTFBuilder::GetLogMessages() const
 TArray<FGLTFBuilder::FLogMessage> FGLTFBuilder::GetLogMessages(EGLTFMessageSeverity Severity) const
 {
 	return LogMessages.FilterByPredicate(
-		[&Severity](const FLogMessage& LogMessage)
+		[Severity](const FLogMessage& LogMessage)
 		{
 			return LogMessage.Key == Severity;
 		});
