[GLTFExporter] refactor(logging): use defined name and friendly-name of the plugin for message-logs instead of hardcoded values

commit 0764c1d5d8eea7d1ea5f0c49f572904f85435949
Date: Thu, 25 Jun 2020 13:06:00 +0200

    refactor(logging): use defined name and friendly-name of the plugin for message-logs instead of hardcoded values
    (cherry picked from commit 955080b4e6f65ea15b2c14aef86bcdc9fb93f2aa​)

#skipci

diff --git a/Engine/Plugins/Enterprise/GLTFExporter/MigratedFromGit/Source/GLTFExporter/Private/Builders/GLTFBuilder.cpp b/Engine/Plugins/Enterprise/GLTFExporter/MigratedFromGit/Source/GLTFExporter/Private/Builders/GLTFBuilder.cpp
--- a/Engine/Plugins/Enterprise/GLTFExporter/MigratedFromGit/Source/GLTFExporter/Private/Builders/GLTFBuilder.cpp
+++ b/Engine/Plugins/Enterprise/GLTFExporter/MigratedFromGit/Source/GLTFExporter/Private/Builders/GLTFBuilder.cpp
@@ -1,18 +1,13 @@
 // Copyright Epic Games, Inc. All Rights Reserved.
 
 #include "Builders/GLTFBuilder.h"
+#include "GLTFExporterModule.h"
+#include "Interfaces/IPluginManager.h"
 #include "MessageLogModule.h"
 #include "IMessageLogListing.h"
 
 DEFINE_LOG_CATEGORY(LogGLTFExporter);
 
-namespace
-{
-	// TODO: turn into config variable?
-	const TCHAR* LogListingName = TEXT("GLTFExporter");
-	const TCHAR* LogListingLabel = TEXT("GLTF Exporter");
-} // anonymous namespace
-
 FGLTFBuilder::FGLTFBuilder()
 {
 }
@@ -100,9 +95,12 @@ void FGLTFBuilder::ShowLogMessages() const
 	if (LogMessages.Num() > 0)
 	{
 		FMessageLogModule& MessageLogModule = FModuleManager::LoadModuleChecked<FMessageLogModule>("MessageLog");
-		const TSharedRef<IMessageLogListing> LogListing = MessageLogModule.GetLogListing(LogListingName);
+		const TSharedRef<IMessageLogListing> LogListing = MessageLogModule.GetLogListing(GLTFEXPORTER_MODULE_NAME);
+
+		const TSharedPtr<IPlugin> Plugin = IPluginManager::Get().FindPlugin(GLTFEXPORTER_MODULE_NAME);
+		const FPluginDescriptor& PluginDescriptor = Plugin->GetDescriptor();
 
-		LogListing->SetLabel(FText::FromString(LogListingLabel));
+		LogListing->SetLabel(FText::FromString(PluginDescriptor.FriendlyName));
 		LogListing->ClearMessages();
 
 		for (const FLogMessage& Message : LogMessages)
@@ -110,7 +108,7 @@ void FGLTFBuilder::ShowLogMessages() const
 			LogListing->AddMessage(CreateTokenizedMessage(Message));
 		}
 
-		MessageLogModule.OpenMessageLog(LogListingName);
+		MessageLogModule.OpenMessageLog(GLTFEXPORTER_MODULE_NAME);
 	}
 }
 
