[GLTFExporter] refactor(logging): use FString::Printf() for formatted warning-messages

commit 05356841ed2b864d1ddcd92d9ee19e98207eee22
Date: Thu, 25 Jun 2020 13:31:48 +0200

    refactor(logging): use FString::Printf() for formatted warning-messages
    (cherry picked from commit 02a6c98e32de80b3aadb7a2cfbb47319798bcf93​)

#skipci

diff --git a/Engine/Plugins/Enterprise/GLTFExporter/MigratedFromGit/Source/GLTFExporter/Private/Converters/GLTFMaterialConverters.cpp b/Engine/Plugins/Enterprise/GLTFExporter/MigratedFromGit/Source/GLTFExporter/Private/Converters/GLTFMaterialConverters.cpp
--- a/Engine/Plugins/Enterprise/GLTFExporter/MigratedFromGit/Source/GLTFExporter/Private/Converters/GLTFMaterialConverters.cpp
+++ b/Engine/Plugins/Enterprise/GLTFExporter/MigratedFromGit/Source/GLTFExporter/Private/Converters/GLTFMaterialConverters.cpp
@@ -38,7 +38,7 @@ FGLTFJsonMaterialIndex FGLTFMaterialConverter::Add(FGLTFConvertBuilder& Builder,
 			{
 				if (!TryGetBakedMaterialProperty(Builder, JsonMaterial.PBRMetallicRoughness.BaseColorTexture, JsonMaterial.PBRMetallicRoughness.BaseColorFactor, MP_BaseColor, Material))
 				{
-					Builder.AddWarningMessage(TEXT("Failed to export BaseColor for material ") + Material->GetName());
+					Builder.AddWarningMessage(FString::Printf(TEXT("Failed to export BaseColor for material %s"), *Material->GetName()));
 				}
 			}
 		}
@@ -49,13 +49,13 @@ FGLTFJsonMaterialIndex FGLTFMaterialConverter::Add(FGLTFConvertBuilder& Builder,
 	{
 		if (!TryGetBaseColorAndOpacity(Builder, JsonMaterial.PBRMetallicRoughness, Material))
 		{
-			Builder.AddWarningMessage(TEXT("Failed to export BaseColor & Opacity for material ") + Material->GetName());
+			Builder.AddWarningMessage(FString::Printf(TEXT("Failed to export BaseColor & Opacity for material %s"), *Material->GetName()));
 		}
 	}
 
 	if (!TryGetMetallicAndRoughness(Builder, JsonMaterial.PBRMetallicRoughness, Material))
 	{
-		Builder.AddWarningMessage(TEXT("Failed to export Metallic & Roughness for material ") + Material->GetName());
+		Builder.AddWarningMessage(FString::Printf(TEXT("Failed to export Metallic & Roughness for material %s"), *Material->GetName()));
 	}
 
 	// NOTE: export of EmissiveColor has been temporarily disabled because of visual differences that
@@ -77,7 +77,7 @@ FGLTFJsonMaterialIndex FGLTFMaterialConverter::Add(FGLTFConvertBuilder& Builder,
 			}
 			else
 			{
-				Builder.AddWarningMessage(TEXT("Failed to export EmissiveColor for material ") + Material->GetName());
+				Builder.AddWarningMessage(FString::Printf(TEXT("Failed to export EmissiveColor for material %s"), *Material->GetName()));
 			}
 		}
 		else
@@ -93,7 +93,7 @@ FGLTFJsonMaterialIndex FGLTFMaterialConverter::Add(FGLTFConvertBuilder& Builder,
 		{
 			if (!TryGetBakedMaterialProperty(Builder, JsonMaterial.NormalTexture, MP_Normal, Material))
 			{
-				Builder.AddWarningMessage(TEXT("Failed to export Normal for material ") + Material->GetName());
+				Builder.AddWarningMessage(FString::Printf(TEXT("Failed to export Normal for material %s"), *Material->GetName()));
 			}
 		}
 	}
@@ -104,7 +104,7 @@ FGLTFJsonMaterialIndex FGLTFMaterialConverter::Add(FGLTFConvertBuilder& Builder,
 		{
 			if (!TryGetBakedMaterialProperty(Builder, JsonMaterial.OcclusionTexture, MP_AmbientOcclusion, Material))
 			{
-				Builder.AddWarningMessage(TEXT("Failed to export AmbientOcclusion for material ") + Material->GetName());
+				Builder.AddWarningMessage(FString::Printf(TEXT("Failed to export AmbientOcclusion for material %s"), *Material->GetName()));
 			}
 		}
 	}
@@ -174,10 +174,9 @@ bool FGLTFMaterialConverter::TryGetBaseColorAndOpacity(FGLTFConvertBuilder& Buil
 		if (!bAreTexturesCompatible)
 		{
 			// TODO: handle differences in wrapping or uv-coords
-			Builder.AddWarningMessage(
-				TEXT("BaseColor- and Opacity-textures for material ") +
-				Material->GetName() +
-				TEXT(" were not able to be combined and will be skipped"));
+			Builder.AddWarningMessage(FString::Printf(
+				TEXT("BaseColor- and Opacity-textures for material %s were not able to be combined and will be skipped"),
+				*Material->GetName()));
 
 			return false;
 		}
@@ -319,10 +318,9 @@ bool FGLTFMaterialConverter::TryGetMetallicAndRoughness(FGLTFConvertBuilder& Bui
 		if (!bAreTexturesCompatible)
 		{
 			// TODO: handle differences in wrapping or uv-coords
-			Builder.AddWarningMessage(
-				TEXT("Metallic- and Roughness-textures for material ") +
-				Material->GetName() +
-				TEXT(" were not able to be combined and will be skipped"));
+			Builder.AddWarningMessage(FString::Printf(
+				TEXT("Metallic- and Roughness-textures for material %s were not able to be combined and will be skipped"),
+				*Material->GetName()));
 
 			return false;
 		}
