[GLTFExporter] feat(logging): add error-messages for failed exports in UGLTFMaterialExporter

commit 547150015b74956d782df8ed51b4fada7d7ac721
Date: Thu, 25 Jun 2020 13:40:37 +0200

    feat(logging): add error-messages for failed exports in UGLTFMaterialExporter
    (cherry picked from commit a89a6dfe64261d1e28b06ed551118d8667cb83a4​)

#skipci

diff --git a/Engine/Plugins/Enterprise/GLTFExporter/MigratedFromGit/Source/GLTFExporter/Private/Exporters/GLTFMaterialExporter.cpp b/Engine/Plugins/Enterprise/GLTFExporter/MigratedFromGit/Source/GLTFExporter/Private/Exporters/GLTFMaterialExporter.cpp
--- a/Engine/Plugins/Enterprise/GLTFExporter/MigratedFromGit/Source/GLTFExporter/Private/Exporters/GLTFMaterialExporter.cpp
+++ b/Engine/Plugins/Enterprise/GLTFExporter/MigratedFromGit/Source/GLTFExporter/Private/Exporters/GLTFMaterialExporter.cpp
@@ -18,6 +18,7 @@ bool UGLTFMaterialExporter::AddObject(FGLTFContainerBuilder& Builder, const UObj
 
 	if (MaterialIndex == INDEX_NONE)
 	{
+		Builder.AddErrorMessage(FString::Printf(TEXT("Failed to export material %s"), *Material->GetName()));
 		return false;
 	}
 
@@ -38,7 +39,11 @@ bool UGLTFMaterialExporter::AddObject(FGLTFContainerBuilder& Builder, const UObj
 	}
 	else
 	{
-		// TODO: should we report an error if no preview mesh was found?
+		Builder.AddErrorMessage(FString::Printf(
+			TEXT("Failed to export material %s because of missing preview-mesh"),
+			*Material->GetName()));
+
+		return false;
 	}
 
 	return true;
