[GLTFExporter] feat(logging): add error-messages for failed exports in FGLTFLevelExporter

commit dfaaf5df76341d00d2f7aeeb10754cfd497d17a0
Date: Thu, 25 Jun 2020 13:45:43 +0200

    feat(logging): add error-messages for failed exports in FGLTFLevelExporter
    (cherry picked from commit 436dff58e61c3522b0b6d509320e0f6c63f195b9​)

#skipci

diff --git a/Engine/Plugins/Enterprise/GLTFExporter/MigratedFromGit/Source/GLTFExporter/Private/Exporters/GLTFLevelExporter.cpp b/Engine/Plugins/Enterprise/GLTFExporter/MigratedFromGit/Source/GLTFExporter/Private/Exporters/GLTFLevelExporter.cpp
--- a/Engine/Plugins/Enterprise/GLTFExporter/MigratedFromGit/Source/GLTFExporter/Private/Exporters/GLTFLevelExporter.cpp
+++ b/Engine/Plugins/Enterprise/GLTFExporter/MigratedFromGit/Source/GLTFExporter/Private/Exporters/GLTFLevelExporter.cpp
@@ -15,13 +15,12 @@ bool UGLTFLevelExporter::AddObject(FGLTFContainerBuilder& Builder, const UObject
 	const UWorld* World = CastChecked<UWorld>(Object);
 	const FGLTFJsonSceneIndex SceneIndex = Builder.GetOrAddScene(World);
 
-	if (SceneIndex != INDEX_NONE)
-	{
-		Builder.DefaultScene = SceneIndex;
-		return true;
-	}
-	else
+	if (SceneIndex == INDEX_NONE)
 	{
+		Builder.AddErrorMessage(FString::Printf(TEXT("Failed to export level %s"), *World->GetName()));
 		return false;
 	}
+
+	Builder.DefaultScene = SceneIndex;
+	return true;
 }
