[GLTFExporter] feat(logging): add error-messages for failed exports in FGLTFStaticMeshExporter

commit b1d66bcdb2864103303e7112494b336ec2209aec
Date: Thu, 25 Jun 2020 13:50:24 +0200

    feat(logging): add error-messages for failed exports in FGLTFStaticMeshExporter
    (cherry picked from commit 348c2f77e520cebcde6ad88149326b668445b6fa​)

#skipci

diff --git a/Engine/Plugins/Enterprise/GLTFExporter/MigratedFromGit/Source/GLTFExporter/Private/Exporters/GLTFStaticMeshExporter.cpp b/Engine/Plugins/Enterprise/GLTFExporter/MigratedFromGit/Source/GLTFExporter/Private/Exporters/GLTFStaticMeshExporter.cpp
--- a/Engine/Plugins/Enterprise/GLTFExporter/MigratedFromGit/Source/GLTFExporter/Private/Exporters/GLTFStaticMeshExporter.cpp
+++ b/Engine/Plugins/Enterprise/GLTFExporter/MigratedFromGit/Source/GLTFExporter/Private/Exporters/GLTFStaticMeshExporter.cpp
@@ -17,6 +17,7 @@ bool UGLTFStaticMeshExporter::AddObject(FGLTFContainerBuilder& Builder, const UO
 
 	if (MeshIndex == INDEX_NONE)
 	{
+		Builder.AddErrorMessage(FString::Printf(TEXT("Failed to export static mesh %s"), *StaticMesh->GetName()));
 		return false;
 	}
 
