[GLTFExporter] feat(json): add getters in GLTFJsonBuilder to change json data post-add

commit 0a2b3fbc860951bda1e8a957dfd96112b0c9ec30
Date: Wed, 17 Jun 2020 08:53:34 +0200

    feat(json): add getters in GLTFJsonBuilder to change json data post-add
    (cherry picked from commit 7532882e8ad02fd8271b28d819a921c9f07bf87b​)

#skipci

diff --git a/Engine/Plugins/Enterprise/GLTFExporter/MigratedFromGit/Source/GLTFExporter/Private/Builders/GLTFBufferBuilder.cpp b/Engine/Plugins/Enterprise/GLTFExporter/MigratedFromGit/Source/GLTFExporter/Private/Builders/GLTFBufferBuilder.cpp
--- a/Engine/Plugins/Enterprise/GLTFExporter/MigratedFromGit/Source/GLTFExporter/Private/Builders/GLTFBufferBuilder.cpp
+++ b/Engine/Plugins/Enterprise/GLTFExporter/MigratedFromGit/Source/GLTFExporter/Private/Builders/GLTFBufferBuilder.cpp
@@ -6,7 +6,7 @@
 
 FGLTFBufferBuilder::FGLTFBufferBuilder()
 {
-	BufferIndex = AddBuffer(FGLTFJsonBuffer());
+	BufferIndex = AddBuffer();
 }
 
 FGLTFJsonBufferViewIndex FGLTFBufferBuilder::AddBufferView(const void* RawData, uint64 ByteLength, const FString& Name, uint8 DataAlignment, EGLTFJsonBufferTarget BufferTarget)
@@ -35,7 +35,7 @@ FGLTFJsonBufferViewIndex FGLTFBufferBuilder::AddBufferView(const void* RawData,
 
 void FGLTFBufferBuilder::UpdateJsonBufferObject(const FString& BinaryFilePath)
 {
-	FGLTFJsonBuffer& JsonBuffer = JsonRoot.Buffers[BufferIndex];
+	FGLTFJsonBuffer& JsonBuffer = GetBuffer(BufferIndex);
 	JsonBuffer.URI = FPaths::GetCleanFilename(BinaryFilePath);
 	JsonBuffer.ByteLength = BufferData.Num();
 }
diff --git a/Engine/Plugins/Enterprise/GLTFExporter/MigratedFromGit/Source/GLTFExporter/Private/Builders/GLTFImageBuilder.cpp b/Engine/Plugins/Enterprise/GLTFExporter/MigratedFromGit/Source/GLTFExporter/Private/Builders/GLTFImageBuilder.cpp
--- a/Engine/Plugins/Enterprise/GLTFExporter/MigratedFromGit/Source/GLTFExporter/Private/Builders/GLTFImageBuilder.cpp
+++ b/Engine/Plugins/Enterprise/GLTFExporter/MigratedFromGit/Source/GLTFExporter/Private/Builders/GLTFImageBuilder.cpp
@@ -117,7 +117,7 @@ bool FGLTFImageBuilder::Serialize(FArchive& Archive, const FString& FilePath)
 	for (const auto& DataPair : ImageDataLookup)
 	{
 		const TArray64<uint8>& ImageData = DataPair.Value;
-		FGLTFJsonImage& JsonImage = JsonRoot.Images[DataPair.Key];
+		FGLTFJsonImage& JsonImage = GetImage(DataPair.Key);
 
 		const FString ImageUri = JsonImage.Name + FGLTFBuilderUtility::GetFileExtension(JsonImage.MimeType);
 		const FString ImagePath = FPaths::Combine(ImageDir, ImageUri);
diff --git a/Engine/Plugins/Enterprise/GLTFExporter/MigratedFromGit/Source/GLTFExporter/Private/Builders/GLTFJsonBuilder.cpp b/Engine/Plugins/Enterprise/GLTFExporter/MigratedFromGit/Source/GLTFExporter/Private/Builders/GLTFJsonBuilder.cpp
--- a/Engine/Plugins/Enterprise/GLTFExporter/MigratedFromGit/Source/GLTFExporter/Private/Builders/GLTFJsonBuilder.cpp
+++ b/Engine/Plugins/Enterprise/GLTFExporter/MigratedFromGit/Source/GLTFExporter/Private/Builders/GLTFJsonBuilder.cpp
@@ -57,6 +57,56 @@ FGLTFJsonTextureIndex FGLTFJsonBuilder::AddTexture(const FGLTFJsonTexture& JsonT
 	return FGLTFJsonTextureIndex(JsonRoot.Textures.Add(JsonTexture));
 }
 
+FGLTFJsonAccessor& FGLTFJsonBuilder::GetAccessor(FGLTFJsonAccessorIndex AccessorIndex)
+{
+	return JsonRoot.Accessors[AccessorIndex];
+}
+
+FGLTFJsonBuffer& FGLTFJsonBuilder::GetBuffer(FGLTFJsonBufferIndex BufferIndex)
+{
+	return JsonRoot.Buffers[BufferIndex];
+}
+
+FGLTFJsonBufferView& FGLTFJsonBuilder::GetBufferView(FGLTFJsonBufferViewIndex BufferViewIndex)
+{
+	return JsonRoot.BufferViews[BufferViewIndex];
+}
+
+FGLTFJsonImage& FGLTFJsonBuilder::GetImage(FGLTFJsonImageIndex ImageIndex)
+{
+	return JsonRoot.Images[ImageIndex];
+}
+
+FGLTFJsonMaterial& FGLTFJsonBuilder::GetMaterial(FGLTFJsonMaterialIndex MaterialIndex)
+{
+	return JsonRoot.Materials[MaterialIndex];
+}
+
+FGLTFJsonMesh& FGLTFJsonBuilder::GetMesh(FGLTFJsonMeshIndex MeshIndex)
+{
+	return JsonRoot.Meshes[MeshIndex];
+}
+
+FGLTFJsonNode& FGLTFJsonBuilder::GetNode(FGLTFJsonNodeIndex NodeIndex)
+{
+	return JsonRoot.Nodes[NodeIndex];
+}
+
+FGLTFJsonSampler& FGLTFJsonBuilder::GetSampler(FGLTFJsonSamplerIndex SamplerIndex)
+{
+	return JsonRoot.Samplers[SamplerIndex];
+}
+
+FGLTFJsonScene& FGLTFJsonBuilder::GetScene(FGLTFJsonSceneIndex SceneIndex)
+{
+	return JsonRoot.Scenes[SceneIndex];
+}
+
+FGLTFJsonTexture& FGLTFJsonBuilder::GetTexture(FGLTFJsonTextureIndex TextureIndex)
+{
+	return JsonRoot.Textures[TextureIndex];
+}
+
 bool FGLTFJsonBuilder::Serialize(FArchive& Archive, const FString& FilePath)
 {
 	JsonRoot.Serialize(&Archive, true);
diff --git a/Engine/Plugins/Enterprise/GLTFExporter/MigratedFromGit/Source/GLTFExporter/Private/Builders/GLTFJsonBuilder.h b/Engine/Plugins/Enterprise/GLTFExporter/MigratedFromGit/Source/GLTFExporter/Private/Builders/GLTFJsonBuilder.h
--- a/Engine/Plugins/Enterprise/GLTFExporter/MigratedFromGit/Source/GLTFExporter/Private/Builders/GLTFJsonBuilder.h
+++ b/Engine/Plugins/Enterprise/GLTFExporter/MigratedFromGit/Source/GLTFExporter/Private/Builders/GLTFJsonBuilder.h
@@ -13,20 +13,31 @@ public:
 
 	FGLTFJsonBuilder();
 
-	FGLTFJsonAccessorIndex AddAccessor(const FGLTFJsonAccessor& JsonAccessor);
-	FGLTFJsonBufferIndex AddBuffer(const FGLTFJsonBuffer& JsonBuffer);
-	FGLTFJsonBufferViewIndex AddBufferView(const FGLTFJsonBufferView& JsonBufferView);
-	FGLTFJsonImageIndex AddImage(const FGLTFJsonImage& JsonImage);
-	FGLTFJsonMaterialIndex AddMaterial(const FGLTFJsonMaterial& JsonMaterial);
-	FGLTFJsonMeshIndex AddMesh(const FGLTFJsonMesh& JsonMesh);
-	FGLTFJsonNodeIndex AddNode(const FGLTFJsonNode& JsonNode);
-	FGLTFJsonSamplerIndex AddSampler(const FGLTFJsonSampler& JsonSampler);
-	FGLTFJsonSceneIndex AddScene(const FGLTFJsonScene& JsonScene);
-	FGLTFJsonTextureIndex AddTexture(const FGLTFJsonTexture& JsonTexture);
+	FGLTFJsonAccessorIndex AddAccessor(const FGLTFJsonAccessor& JsonAccessor = {});
+	FGLTFJsonBufferIndex AddBuffer(const FGLTFJsonBuffer& JsonBuffer = {});
+	FGLTFJsonBufferViewIndex AddBufferView(const FGLTFJsonBufferView& JsonBufferView = {});
+	FGLTFJsonImageIndex AddImage(const FGLTFJsonImage& JsonImage = {});
+	FGLTFJsonMaterialIndex AddMaterial(const FGLTFJsonMaterial& JsonMaterial = {});
+	FGLTFJsonMeshIndex AddMesh(const FGLTFJsonMesh& JsonMesh = {});
+	FGLTFJsonNodeIndex AddNode(const FGLTFJsonNode& JsonNode = {});
+	FGLTFJsonSamplerIndex AddSampler(const FGLTFJsonSampler& JsonSampler = {});
+	FGLTFJsonSceneIndex AddScene(const FGLTFJsonScene& JsonScene = {});
+	FGLTFJsonTextureIndex AddTexture(const FGLTFJsonTexture& JsonTexture = {});
+
+	FGLTFJsonAccessor& GetAccessor(FGLTFJsonAccessorIndex AccessorIndex);
+	FGLTFJsonBuffer& GetBuffer(FGLTFJsonBufferIndex BufferIndex);
+	FGLTFJsonBufferView& GetBufferView(FGLTFJsonBufferViewIndex BufferViewIndex);
+	FGLTFJsonImage& GetImage(FGLTFJsonImageIndex ImageIndex);
+	FGLTFJsonMaterial& GetMaterial(FGLTFJsonMaterialIndex MaterialIndex);
+	FGLTFJsonMesh& GetMesh(FGLTFJsonMeshIndex MeshIndex);
+	FGLTFJsonNode& GetNode(FGLTFJsonNodeIndex NodeIndex);
+	FGLTFJsonSampler& GetSampler(FGLTFJsonSamplerIndex SamplerIndex);
+	FGLTFJsonScene& GetScene(FGLTFJsonSceneIndex SceneIndex);
+	FGLTFJsonTexture& GetTexture(FGLTFJsonTextureIndex TextureIndex);
 
 	virtual bool Serialize(FArchive& Archive, const FString& FilePath);
 
-protected:
+private:
 
 	FGLTFJsonRoot JsonRoot;
 };
