[GLTFExporter] refactor(level): tweak scene component converter to have knowledge about its node index

commit e65d8d317cabb4ec66da8e18700c063a6a3c370c
Date: Wed, 17 Jun 2020 10:08:54 +0200

    refactor(level): tweak scene component converter to have knowledge about its node index
    (cherry picked from commit e108e28ae9315c665a5712ebcbaf76062727fd11​)

#skipci

diff --git a/Engine/Plugins/Enterprise/GLTFExporter/MigratedFromGit/Source/GLTFExporter/Private/Converters/GLTFLevelConverters.cpp b/Engine/Plugins/Enterprise/GLTFExporter/MigratedFromGit/Source/GLTFExporter/Private/Converters/GLTFLevelConverters.cpp
--- a/Engine/Plugins/Enterprise/GLTFExporter/MigratedFromGit/Source/GLTFExporter/Private/Converters/GLTFLevelConverters.cpp
+++ b/Engine/Plugins/Enterprise/GLTFExporter/MigratedFromGit/Source/GLTFExporter/Private/Converters/GLTFLevelConverters.cpp
@@ -14,7 +14,8 @@ FGLTFJsonNodeIndex FGLTFSceneComponentConverter::Add(FGLTFConvertBuilder& Builde
 	const bool bIsRootComponent = Owner != nullptr && Owner->GetRootComponent() == SceneComponent;
 	const bool bRootNode = bIsRootComponent && (ParentComponent == nullptr || (bSelectedOnly && !FGLTFConverterUtility::IsSelected(ParentComponent)));
 
-	FGLTFJsonNode Node;
+	const FGLTFJsonNodeIndex NodeIndex = Builder.AddNode();
+	FGLTFJsonNode& Node = Builder.GetNode(NodeIndex);
 	Node.Name = Name.IsEmpty() ? Owner->GetName() + TEXT("_") + SceneComponent->GetName() : Name;
 
 	// TODO: add support for non-uniform scaling (Unreal doesn't treat combined transforms as simple matrix multiplication)
@@ -39,20 +40,27 @@ FGLTFJsonNodeIndex FGLTFSceneComponentConverter::Add(FGLTFConvertBuilder& Builde
 		Node.Mesh = Builder.GetOrAddMesh(StaticMeshComponent);
 	}
 
-	const TArray<USceneComponent*>& Children = SceneComponent->GetAttachChildren();
-	for (const USceneComponent* ChildComponent : Children)
+	const TArray<USceneComponent*>& ChildComponents = SceneComponent->GetAttachChildren();
+
+	TArray<FGLTFJsonNodeIndex> Children;
+	Children.Reserve(ChildComponents.Num());
+
+	for (const USceneComponent* ChildComponent : ChildComponents)
 	{
 		if (ChildComponent != nullptr && (!bSelectedOnly || FGLTFConverterUtility::IsSelected(ChildComponent)))
 		{
-			FGLTFJsonNodeIndex NodeIndex = Builder.GetOrAddNode(ChildComponent);
-			if (NodeIndex != INDEX_NONE)
+			const FGLTFJsonNodeIndex ChildNodeIndex = Builder.GetOrAddNode(ChildComponent);
+			if (ChildNodeIndex != INDEX_NONE)
 			{
-				Node.Children.Add(NodeIndex);
+				Children.Add(ChildNodeIndex);
 			}
 		}
 	}
 
-	return Builder.AddNode(Node);
+	// TODO: refactor this workaround for no longer valid reference to node (due to calling Builder.GetOrAddNode)
+	Builder.GetNode(NodeIndex).Children = Children;
+
+	return NodeIndex;
 }
 
 FGLTFJsonSceneIndex FGLTFLevelConverter::Add(FGLTFConvertBuilder& Builder, const FString& Name, const ULevel* Level)
