[GLTFExporter] feat(interaction-hotspot): add skeleton classes

commit e5324c6ad41fdd1f1efd50932c1840a20b2fef93
Date: Wed, 27 May 2020 08:42:38 +0200

    feat(interaction-hotspot): add skeleton classes
    (cherry picked from commit ac411bc5c3c3b26b836da4aa4296f56ef2c85cbf​)

#skipci

diff --git a/Engine/Plugins/Enterprise/GLTFExporter/MigratedFromGit/Source/GLTFExporter/Public/Actors/GLTFInteractionHotspotActor.cpp b/Engine/Plugins/Enterprise/GLTFExporter/MigratedFromGit/Source/GLTFExporter/Public/Actors/GLTFInteractionHotspotActor.cpp
new file mode 100644
--- /dev/null
+++ b/Engine/Plugins/Enterprise/GLTFExporter/MigratedFromGit/Source/GLTFExporter/Public/Actors/GLTFInteractionHotspotActor.cpp
@@ -0,0 +1,33 @@
+// Copyright Epic Games, Inc. All Rights Reserved.
+
+#include "Actors/GLTFInteractionHotspotActor.h"
+#include "Components/GLTFInteractionHotspotComponent.h"
+
+AGLTFInteractionHotspotActor::AGLTFInteractionHotspotActor(const FObjectInitializer& ObjectInitializer)
+	: Super(ObjectInitializer)
+{
+	SceneComponent = CreateDefaultSubobject<USceneComponent>(TEXT("SceneComponent"));
+
+	// Make the scene component the root component
+	RootComponent = SceneComponent;
+
+	InteractionHotspotComponent = CreateDefaultSubobject<UGLTFInteractionHotspotComponent>(TEXT("InteractionHotspotComponent"));
+	InteractionHotspotComponent->SetupAttachment(SceneComponent);
+
+	OnBeginCursorOver.AddDynamic(this, &AGLTFInteractionHotspotActor::BeginCursorOver);
+}
+
+void AGLTFInteractionHotspotActor::BeginPlay()
+{
+	Super::BeginPlay();
+}
+
+void AGLTFInteractionHotspotActor::Tick(float DeltaTime)
+{
+	Super::Tick(DeltaTime);
+}
+
+void AGLTFInteractionHotspotActor::BeginCursorOver(AActor* TouchedActor)
+{
+	UE_LOG(LogTemp, Warning, TEXT("AGLTFInteractionHotspotActor::BeginCursorOver()"));
+}
diff --git a/Engine/Plugins/Enterprise/GLTFExporter/MigratedFromGit/Source/GLTFExporter/Public/Actors/GLTFInteractionHotspotActor.h b/Engine/Plugins/Enterprise/GLTFExporter/MigratedFromGit/Source/GLTFExporter/Public/Actors/GLTFInteractionHotspotActor.h
new file mode 100644
--- /dev/null
+++ b/Engine/Plugins/Enterprise/GLTFExporter/MigratedFromGit/Source/GLTFExporter/Public/Actors/GLTFInteractionHotspotActor.h
@@ -0,0 +1,31 @@
+// Copyright Epic Games, Inc. All Rights Reserved.
+
+#pragma once
+
+#include "GameFramework/Actor.h"
+#include "GLTFInteractionHotspotActor.generated.h"
+
+UCLASS()
+class GLTFEXPORTER_API AGLTFInteractionHotspotActor : public AActor
+{
+	GENERATED_BODY()
+
+public:
+	AGLTFInteractionHotspotActor(const FObjectInitializer& ObjectInitializer);
+
+protected:
+	virtual void BeginPlay() override;
+
+public:
+	virtual void Tick(float DeltaTime) override;
+
+	UFUNCTION()
+	void BeginCursorOver(AActor* TouchedActor);
+
+private:
+	UPROPERTY(VisibleAnywhere, BlueprintReadOnly, meta = (AllowPrivateAccess = "true"))
+	class USceneComponent* SceneComponent;
+		
+	UPROPERTY(VisibleAnywhere, BlueprintReadOnly, meta = (AllowPrivateAccess = "true"))
+	class UGLTFInteractionHotspotComponent* InteractionHotspotComponent;
+};
diff --git a/Engine/Plugins/Enterprise/GLTFExporter/MigratedFromGit/Source/GLTFExporter/Public/Components/GLTFInteractionHotspotComponent.cpp b/Engine/Plugins/Enterprise/GLTFExporter/MigratedFromGit/Source/GLTFExporter/Public/Components/GLTFInteractionHotspotComponent.cpp
new file mode 100644
--- /dev/null
+++ b/Engine/Plugins/Enterprise/GLTFExporter/MigratedFromGit/Source/GLTFExporter/Public/Components/GLTFInteractionHotspotComponent.cpp
@@ -0,0 +1,33 @@
+// Copyright Epic Games, Inc. All Rights Reserved.
+
+#include "Components/GLTFInteractionHotspotComponent.h"
+
+UGLTFInteractionHotspotComponent::UGLTFInteractionHotspotComponent(const FObjectInitializer& ObjectInitializer)
+	: Super(ObjectInitializer)
+{
+	bHiddenInGame = false;
+
+	OnBeginCursorOver.AddDynamic(this, &UGLTFInteractionHotspotComponent::BeginCursorOver);
+	OnEndCursorOver.AddDynamic(this, &UGLTFInteractionHotspotComponent::EndCursorOver);
+	OnClicked.AddDynamic(this, &UGLTFInteractionHotspotComponent::Clicked);
+}
+
+void UGLTFInteractionHotspotComponent::BeginPlay()
+{
+	Super::BeginPlay();
+}
+
+void UGLTFInteractionHotspotComponent::BeginCursorOver(UPrimitiveComponent* TouchedComponent)
+{
+	UE_LOG(LogTemp, Warning, TEXT("UGLTFInteractionHotspotComponent::BeginCursorOver()"));
+}
+
+void UGLTFInteractionHotspotComponent::EndCursorOver(UPrimitiveComponent* TouchedComponent)
+{
+	UE_LOG(LogTemp, Warning, TEXT("UGLTFInteractionHotspotComponent::EndCursorOver()"));
+}
+
+void UGLTFInteractionHotspotComponent::Clicked(UPrimitiveComponent* TouchedComponent, FKey ButtonPressed)
+{
+	UE_LOG(LogTemp, Warning, TEXT("UGLTFInteractionHotspotComponent::Clicked()"));
+}
diff --git a/Engine/Plugins/Enterprise/GLTFExporter/MigratedFromGit/Source/GLTFExporter/Public/Components/GLTFInteractionHotspotComponent.h b/Engine/Plugins/Enterprise/GLTFExporter/MigratedFromGit/Source/GLTFExporter/Public/Components/GLTFInteractionHotspotComponent.h
new file mode 100644
--- /dev/null
+++ b/Engine/Plugins/Enterprise/GLTFExporter/MigratedFromGit/Source/GLTFExporter/Public/Components/GLTFInteractionHotspotComponent.h
@@ -0,0 +1,28 @@
+// Copyright Epic Games, Inc. All Rights Reserved.
+
+#pragma once
+
+#include "Components/BillboardComponent.h"
+#include "GLTFInteractionHotspotComponent.generated.h"
+
+/**
+ * 
+ */
+UCLASS(BlueprintType, Blueprintable, meta = (BlueprintSpawnableComponent), DisplayName = "GLTF Interaction Hotspot Component")
+class GLTFEXPORTER_API UGLTFInteractionHotspotComponent : public UBillboardComponent
+{
+	GENERATED_UCLASS_BODY()
+	
+	virtual void BeginPlay() override;
+
+	virtual bool ShouldCreatePhysicsState() const { return true; }
+
+	UFUNCTION()
+	void BeginCursorOver(UPrimitiveComponent* TouchedComponent);
+
+	UFUNCTION()
+	void EndCursorOver(UPrimitiveComponent* TouchedComponent);
+
+	UFUNCTION()
+	void Clicked(UPrimitiveComponent* TouchedComponent, FKey ButtonPressed);
+};
