[GLTFExporter] feat(interaction-hotspot): add properties for target animations

commit 2b09877d3877e3632572df779c901f46f1337f2a
Date: Thu, 28 May 2020 00:20:28 +0200

    feat(interaction-hotspot): add properties for target animations
    (cherry picked from commit 5a5697456f84aae478918e7363cb71e76b7a2883​)

#skipci

diff --git a/Engine/Plugins/Enterprise/GLTFExporter/MigratedFromGit/Source/GLTFExporter/Public/Actors/GLTFInteractionHotspotActor.cpp b/Engine/Plugins/Enterprise/GLTFExporter/MigratedFromGit/Source/GLTFExporter/Public/Actors/GLTFInteractionHotspotActor.cpp
--- a/Engine/Plugins/Enterprise/GLTFExporter/MigratedFromGit/Source/GLTFExporter/Public/Actors/GLTFInteractionHotspotActor.cpp
+++ b/Engine/Plugins/Enterprise/GLTFExporter/MigratedFromGit/Source/GLTFExporter/Public/Actors/GLTFInteractionHotspotActor.cpp
@@ -6,9 +6,8 @@
 AGLTFInteractionHotspotActor::AGLTFInteractionHotspotActor(const FObjectInitializer& ObjectInitializer)
 	: Super(ObjectInitializer)
 {
+	// A scene component with a transform in the root
 	SceneComponent = CreateDefaultSubobject<USceneComponent>(TEXT("SceneComponent"));
-
-	// Make the scene component the root component
 	RootComponent = SceneComponent;
 
 	InteractionHotspotComponent = CreateDefaultSubobject<UGLTFInteractionHotspotComponent>(TEXT("InteractionHotspotComponent"));
diff --git a/Engine/Plugins/Enterprise/GLTFExporter/MigratedFromGit/Source/GLTFExporter/Public/Actors/GLTFInteractionHotspotActor.h b/Engine/Plugins/Enterprise/GLTFExporter/MigratedFromGit/Source/GLTFExporter/Public/Actors/GLTFInteractionHotspotActor.h
--- a/Engine/Plugins/Enterprise/GLTFExporter/MigratedFromGit/Source/GLTFExporter/Public/Actors/GLTFInteractionHotspotActor.h
+++ b/Engine/Plugins/Enterprise/GLTFExporter/MigratedFromGit/Source/GLTFExporter/Public/Actors/GLTFInteractionHotspotActor.h
@@ -5,12 +5,25 @@
 #include "GameFramework/Actor.h"
 #include "GLTFInteractionHotspotActor.generated.h"
 
-UCLASS()
+class USceneComponent;
+class UGLTFInteractionHotspotComponent;
+
+/**
+ *
+ */
+UCLASS(DisplayName = "GLTF Interaction Hotspot Actor")
 class GLTFEXPORTER_API AGLTFInteractionHotspotActor : public AActor
 {
 	GENERATED_BODY()
+private:
+	UPROPERTY(VisibleAnywhere, BlueprintReadOnly, meta = (AllowPrivateAccess = "true"))
+	USceneComponent* SceneComponent;
+
+	UPROPERTY(VisibleAnywhere, BlueprintReadOnly, meta = (AllowPrivateAccess = "true"))
+	UGLTFInteractionHotspotComponent* InteractionHotspotComponent;
 
 public:
+	//~ Begin AActor Interface.
 	AGLTFInteractionHotspotActor(const FObjectInitializer& ObjectInitializer);
 
 protected:
@@ -18,14 +31,9 @@ protected:
 
 public:
 	virtual void Tick(float DeltaTime) override;
+	//~ End AActor Interface.
 
+private:
 	UFUNCTION()
 	void BeginCursorOver(AActor* TouchedActor);
-
-private:
-	UPROPERTY(VisibleAnywhere, BlueprintReadOnly, meta = (AllowPrivateAccess = "true"))
-	class USceneComponent* SceneComponent;
-		
-	UPROPERTY(VisibleAnywhere, BlueprintReadOnly, meta = (AllowPrivateAccess = "true"))
-	class UGLTFInteractionHotspotComponent* InteractionHotspotComponent;
 };
diff --git a/Engine/Plugins/Enterprise/GLTFExporter/MigratedFromGit/Source/GLTFExporter/Public/Components/GLTFInteractionHotspotComponent.cpp b/Engine/Plugins/Enterprise/GLTFExporter/MigratedFromGit/Source/GLTFExporter/Public/Components/GLTFInteractionHotspotComponent.cpp
--- a/Engine/Plugins/Enterprise/GLTFExporter/MigratedFromGit/Source/GLTFExporter/Public/Components/GLTFInteractionHotspotComponent.cpp
+++ b/Engine/Plugins/Enterprise/GLTFExporter/MigratedFromGit/Source/GLTFExporter/Public/Components/GLTFInteractionHotspotComponent.cpp
@@ -1,6 +1,8 @@
 // Copyright Epic Games, Inc. All Rights Reserved.
 
 #include "Components/GLTFInteractionHotspotComponent.h"
+#include "Animation/SkeletalMeshActor.h"
+#include "Animation/AnimSequence.h"
 
 UGLTFInteractionHotspotComponent::UGLTFInteractionHotspotComponent(const FObjectInitializer& ObjectInitializer)
 	: Super(ObjectInitializer)
diff --git a/Engine/Plugins/Enterprise/GLTFExporter/MigratedFromGit/Source/GLTFExporter/Public/Components/GLTFInteractionHotspotComponent.h b/Engine/Plugins/Enterprise/GLTFExporter/MigratedFromGit/Source/GLTFExporter/Public/Components/GLTFInteractionHotspotComponent.h
--- a/Engine/Plugins/Enterprise/GLTFExporter/MigratedFromGit/Source/GLTFExporter/Public/Components/GLTFInteractionHotspotComponent.h
+++ b/Engine/Plugins/Enterprise/GLTFExporter/MigratedFromGit/Source/GLTFExporter/Public/Components/GLTFInteractionHotspotComponent.h
@@ -5,6 +5,9 @@
 #include "Components/BillboardComponent.h"
 #include "GLTFInteractionHotspotComponent.generated.h"
 
+class ASkeletalMeshActor;
+class UAnimSequence;
+
 /**
  * 
  */
@@ -12,11 +15,18 @@ UCLASS(BlueprintType, Blueprintable, meta = (BlueprintSpawnableComponent), Displ
 class GLTFEXPORTER_API UGLTFInteractionHotspotComponent : public UBillboardComponent
 {
 	GENERATED_UCLASS_BODY()
-	
-	virtual void BeginPlay() override;
 
-	virtual bool ShouldCreatePhysicsState() const { return true; }
+	UPROPERTY(EditAnywhere, BlueprintReadWrite)
+	ASkeletalMeshActor* SkeletalMeshActor;
+
+	UPROPERTY(EditAnywhere, BlueprintReadWrite)
+	UAnimSequence* AnimationSequence;
+
+	//~ Begin UActorComponent Interface.
+	virtual void BeginPlay() override;
+	//~ End UActorComponent Interface.
 
+private:
 	UFUNCTION()
 	void BeginCursorOver(UPrimitiveComponent* TouchedComponent);
 
