[GLTFExporter] feat(interaction-hotspot): add mesh-based collision volume

commit 7482ef64bb4f407082c64868a80db96e0262fc1a
Date: Sun, 31 May 2020 23:29:36 +0200

    feat(interaction-hotspot): add mesh-based collision volume
    (cherry picked from commit 698e44fecf600a50e8ac634564ec1f4dd731217d​)

#skipci

diff --git a/Engine/Plugins/Enterprise/GLTFExporter/MigratedFromGit/Source/GLTFExporter/Public/Components/GLTFInteractionHotspotComponent.cpp b/Engine/Plugins/Enterprise/GLTFExporter/MigratedFromGit/Source/GLTFExporter/Public/Components/GLTFInteractionHotspotComponent.cpp
--- a/Engine/Plugins/Enterprise/GLTFExporter/MigratedFromGit/Source/GLTFExporter/Public/Components/GLTFInteractionHotspotComponent.cpp
+++ b/Engine/Plugins/Enterprise/GLTFExporter/MigratedFromGit/Source/GLTFExporter/Public/Components/GLTFInteractionHotspotComponent.cpp
@@ -3,33 +3,96 @@
 #include "Components/GLTFInteractionHotspotComponent.h"
 #include "Animation/SkeletalMeshActor.h"
 #include "Animation/AnimSequence.h"
+#include "Engine/CollisionProfile.h"
+#include "UObject/ConstructorHelpers.h"
+#include "PhysicsEngine/BodySetup.h"
 
 UGLTFInteractionHotspotComponent::UGLTFInteractionHotspotComponent(const FObjectInitializer& ObjectInitializer)
-	: Super(ObjectInitializer)
+	: Super(ObjectInitializer),
+	ShapeBodySetup(nullptr)
 {
+	SetSprite(DefaultSprite);
+
+	// Setup the most minimalistic collision profile for mouse input events
+	SetCollisionEnabled(ECollisionEnabled::QueryAndPhysics);
+	SetCollisionResponseToAllChannels(ECollisionResponse::ECR_Ignore);
+	SetCollisionResponseToChannel(ECollisionChannel::ECC_Visibility, ECollisionResponse::ECR_Block);
+
 	bHiddenInGame = false;
+	SetGenerateOverlapEvents(false);
 
 	OnBeginCursorOver.AddDynamic(this, &UGLTFInteractionHotspotComponent::BeginCursorOver);
 	OnEndCursorOver.AddDynamic(this, &UGLTFInteractionHotspotComponent::EndCursorOver);
 	OnClicked.AddDynamic(this, &UGLTFInteractionHotspotComponent::Clicked);
 }
 
+void UGLTFInteractionHotspotComponent::PostEditChangeProperty(FPropertyChangedEvent& PropertyChangedEvent)
+{
+	FProperty* PropertyThatChanged = PropertyChangedEvent.Property;
+
+	if (PropertyThatChanged)
+	{
+		if (PropertyThatChanged->GetName() == TEXT("DefaultSprite"))
+		{
+			SetSprite(DefaultSprite);
+		}
+
+	}
+}
+
 void UGLTFInteractionHotspotComponent::BeginPlay()
 {
 	Super::BeginPlay();
 }
 
+void UGLTFInteractionHotspotComponent::OnRegister()
+{
+#if 1
+	const UStaticMesh* SphereMesh = LoadObject<UStaticMesh>(NULL, TEXT("/Engine/BasicShapes/Sphere.Sphere"), NULL, LOAD_None, NULL);
+
+	if (SphereMesh != nullptr)
+	{
+		ShapeBodySetup = NewObject<UBodySetup>(this, NAME_None, RF_Transient);
+		FTransform BodyTransform = FTransform::Identity;// = GetComponentTransform();
+		BodyInstance.InitBody(SphereMesh->BodySetup, BodyTransform, this, GetWorld()->GetPhysicsScene());
+
+		RecreatePhysicsState();
+	}
+#else
+	// TODO: Figure out why this isn't working
+	ShapeBodySetup = NewObject<UBodySetup>(this, UBodySetup::StaticClass());
+	ShapeBodySetup->AggGeom.SphereElems.Add(FKSphereElem(50.0f));
+	BodyInstance.InitBody(ShapeBodySetup, FTransform::Identity, this, GetWorld()->GetPhysicsScene());
+
+	RecreatePhysicsState();
+#endif
+
+	Super::OnRegister();
+}
+
 void UGLTFInteractionHotspotComponent::BeginCursorOver(UPrimitiveComponent* TouchedComponent)
 {
-	UE_LOG(LogTemp, Warning, TEXT("UGLTFInteractionHotspotComponent::BeginCursorOver()"));
+	UE_LOG(LogTemp, Warning, TEXT("%s BeginCursorOver"), *GetOwner()->GetName());
+
+	if (HighlightSprite != nullptr)
+	{
+		SetSprite(HighlightSprite);
+	}
 }
 
 void UGLTFInteractionHotspotComponent::EndCursorOver(UPrimitiveComponent* TouchedComponent)
 {
-	UE_LOG(LogTemp, Warning, TEXT("UGLTFInteractionHotspotComponent::EndCursorOver()"));
+	UE_LOG(LogTemp, Warning, TEXT("%s EndCursorOver"), *GetOwner()->GetName());
+
+	SetSprite(DefaultSprite);
 }
 
 void UGLTFInteractionHotspotComponent::Clicked(UPrimitiveComponent* TouchedComponent, FKey ButtonPressed)
 {
 	UE_LOG(LogTemp, Warning, TEXT("UGLTFInteractionHotspotComponent::Clicked()"));
+
+	if (SkeletalMeshActor != nullptr && AnimationSequence != nullptr)
+	{
+		SkeletalMeshActor->GetSkeletalMeshComponent()->PlayAnimation(AnimationSequence, false);
+	}
 }
diff --git a/Engine/Plugins/Enterprise/GLTFExporter/MigratedFromGit/Source/GLTFExporter/Public/Components/GLTFInteractionHotspotComponent.h b/Engine/Plugins/Enterprise/GLTFExporter/MigratedFromGit/Source/GLTFExporter/Public/Components/GLTFInteractionHotspotComponent.h
--- a/Engine/Plugins/Enterprise/GLTFExporter/MigratedFromGit/Source/GLTFExporter/Public/Components/GLTFInteractionHotspotComponent.h
+++ b/Engine/Plugins/Enterprise/GLTFExporter/MigratedFromGit/Source/GLTFExporter/Public/Components/GLTFInteractionHotspotComponent.h
@@ -7,23 +7,47 @@
 
 class ASkeletalMeshActor;
 class UAnimSequence;
+class UBodySetup;
+class UTexture2D;
 
 /**
  * 
  */
-UCLASS(BlueprintType, Blueprintable, meta = (BlueprintSpawnableComponent), DisplayName = "GLTF Interaction Hotspot Component")
+UCLASS(BlueprintType, Blueprintable, meta = (BlueprintSpawnableComponent), hidecategories = (Sprite), DisplayName = "GLTF Interaction Hotspot Component")
 class GLTFEXPORTER_API UGLTFInteractionHotspotComponent : public UBillboardComponent
 {
 	GENERATED_UCLASS_BODY()
-
+public:
 	UPROPERTY(EditAnywhere, BlueprintReadWrite)
 	ASkeletalMeshActor* SkeletalMeshActor;
 
 	UPROPERTY(EditAnywhere, BlueprintReadWrite)
 	UAnimSequence* AnimationSequence;
 
+	UPROPERTY(EditAnywhere, BlueprintReadWrite)
+	UTexture2D* DefaultSprite;
+
+	UPROPERTY(EditAnywhere, BlueprintReadWrite)
+	UTexture2D* HighlightSprite;
+
+	UPROPERTY(EditAnywhere, BlueprintReadWrite)
+	UTexture2D* ClickSprite;
+
+	/** Description of collision */
+	UPROPERTY(transient, duplicatetransient)
+	UBodySetup* ShapeBodySetup;
+
+	//~ Begin UObject Interface.
+#if WITH_EDITOR
+	virtual void PostEditChangeProperty(FPropertyChangedEvent& PropertyChangedEvent) override;
+#endif // WITH_EDITOR
+	//~ End UObject Interface.
+
 	//~ Begin UActorComponent Interface.
 	virtual void BeginPlay() override;
+
+protected:
+	virtual void OnRegister() override;
 	//~ End UActorComponent Interface.
 
 private:
