[GLTFExporter] feat(interaction-hotspot): simplify collision volume

commit 2131fd81c4a2fe5a4222b86cff64474fd7ce1b09
Date: Tue, 2 Jun 2020 16:40:32 +0200

    feat(interaction-hotspot): simplify collision volume
    (cherry picked from commit a9c8ce192a80e4127bf4a2411461b8489a21cbea​)

#skipci

diff --git a/Engine/Plugins/Enterprise/GLTFExporter/MigratedFromGit/Source/GLTFExporter/Public/Components/GLTFInteractionHotspotComponent.cpp b/Engine/Plugins/Enterprise/GLTFExporter/MigratedFromGit/Source/GLTFExporter/Public/Components/GLTFInteractionHotspotComponent.cpp
--- a/Engine/Plugins/Enterprise/GLTFExporter/MigratedFromGit/Source/GLTFExporter/Public/Components/GLTFInteractionHotspotComponent.cpp
+++ b/Engine/Plugins/Enterprise/GLTFExporter/MigratedFromGit/Source/GLTFExporter/Public/Components/GLTFInteractionHotspotComponent.cpp
@@ -9,17 +9,20 @@
 
 UGLTFInteractionHotspotComponent::UGLTFInteractionHotspotComponent(const FObjectInitializer& ObjectInitializer)
 	: Super(ObjectInitializer),
-	ShapeBodySetup(nullptr)
+	SkeletalMeshActor(nullptr),
+	AnimationSequence(nullptr),
+	DefaultSprite(nullptr),
+	HighlightSprite(nullptr),
+	ClickSprite(nullptr),
+	Radius(50.0f)
 {
-	SetSprite(DefaultSprite);
-
 	// Setup the most minimalistic collision profile for mouse input events
 	SetCollisionEnabled(ECollisionEnabled::QueryAndPhysics);
 	SetCollisionResponseToAllChannels(ECollisionResponse::ECR_Ignore);
 	SetCollisionResponseToChannel(ECollisionChannel::ECC_Visibility, ECollisionResponse::ECR_Block);
+	SetGenerateOverlapEvents(false);
 
 	bHiddenInGame = false;
-	SetGenerateOverlapEvents(false);
 
 	OnBeginCursorOver.AddDynamic(this, &UGLTFInteractionHotspotComponent::BeginCursorOver);
 	OnEndCursorOver.AddDynamic(this, &UGLTFInteractionHotspotComponent::EndCursorOver);
@@ -28,52 +31,58 @@ UGLTFInteractionHotspotComponent::UGLTFInteractionHotspotComponent(const FObject
 
 void UGLTFInteractionHotspotComponent::PostEditChangeProperty(FPropertyChangedEvent& PropertyChangedEvent)
 {
-	FProperty* PropertyThatChanged = PropertyChangedEvent.Property;
+	const FProperty* PropertyThatChanged = PropertyChangedEvent.Property;
 
 	if (PropertyThatChanged)
 	{
-		if (PropertyThatChanged->GetName() == TEXT("DefaultSprite"))
+		const FString PropertyName = PropertyThatChanged->GetName();
+
+		if (PropertyName == TEXT("DefaultSprite"))
 		{
 			SetSprite(DefaultSprite);
 		}
-
+		else if (PropertyName == TEXT("Radius"))
+		{
+			SetRadius(Radius);
+		}
 	}
 }
 
 void UGLTFInteractionHotspotComponent::BeginPlay()
 {
 	Super::BeginPlay();
+
+	SetSprite(DefaultSprite);
 }
 
 void UGLTFInteractionHotspotComponent::OnRegister()
 {
-#if 1
-	const UStaticMesh* SphereMesh = LoadObject<UStaticMesh>(NULL, TEXT("/Engine/BasicShapes/Sphere.Sphere"), NULL, LOAD_None, NULL);
-
-	if (SphereMesh != nullptr)
-	{
-		ShapeBodySetup = NewObject<UBodySetup>(this, NAME_None, RF_Transient);
-		FTransform BodyTransform = FTransform::Identity;// = GetComponentTransform();
-		BodyInstance.InitBody(SphereMesh->BodySetup, BodyTransform, this, GetWorld()->GetPhysicsScene());
-
-		RecreatePhysicsState();
-	}
-#else
-	// TODO: Figure out why this isn't working
-	ShapeBodySetup = NewObject<UBodySetup>(this, UBodySetup::StaticClass());
-	ShapeBodySetup->AggGeom.SphereElems.Add(FKSphereElem(50.0f));
-	BodyInstance.InitBody(ShapeBodySetup, FTransform::Identity, this, GetWorld()->GetPhysicsScene());
+	ShapeBodySetup = NewObject<UBodySetup>(this);
+	ShapeBodySetup->CollisionTraceFlag = CTF_UseSimpleAsComplex;
+	ShapeBodySetup->AggGeom.SphereElems.Add(Radius);
+	
+	BodyInstance.InitBody(ShapeBodySetup, GetWorldTransform(), this, GetWorld()->GetPhysicsScene());
 
 	RecreatePhysicsState();
-#endif
+	UpdateCollisionVolume();
 
 	Super::OnRegister();
 }
 
-void UGLTFInteractionHotspotComponent::BeginCursorOver(UPrimitiveComponent* TouchedComponent)
+void UGLTFInteractionHotspotComponent::SetSprite(class UTexture2D* NewSprite)
 {
-	UE_LOG(LogTemp, Warning, TEXT("%s BeginCursorOver"), *GetOwner()->GetName());
+	Super::SetSprite(NewSprite);
+
+	UpdateCollisionVolume();
+}
+
+void UGLTFInteractionHotspotComponent::SetRadius(float NewRadius)
+{
+	// TODO: Implement
+}
 
+void UGLTFInteractionHotspotComponent::BeginCursorOver(UPrimitiveComponent* TouchedComponent)
+{
 	if (HighlightSprite != nullptr)
 	{
 		SetSprite(HighlightSprite);
@@ -82,17 +91,45 @@ void UGLTFInteractionHotspotComponent::BeginCursorOver(UPrimitiveComponent* Touc
 
 void UGLTFInteractionHotspotComponent::EndCursorOver(UPrimitiveComponent* TouchedComponent)
 {
-	UE_LOG(LogTemp, Warning, TEXT("%s EndCursorOver"), *GetOwner()->GetName());
-
 	SetSprite(DefaultSprite);
 }
 
 void UGLTFInteractionHotspotComponent::Clicked(UPrimitiveComponent* TouchedComponent, FKey ButtonPressed)
 {
-	UE_LOG(LogTemp, Warning, TEXT("UGLTFInteractionHotspotComponent::Clicked()"));
-
 	if (SkeletalMeshActor != nullptr && AnimationSequence != nullptr)
 	{
-		SkeletalMeshActor->GetSkeletalMeshComponent()->PlayAnimation(AnimationSequence, false);
+		USkeletalMeshComponent* SkeletalMeshComponent = SkeletalMeshActor->GetSkeletalMeshComponent();
+		SkeletalMeshComponent->PlayAnimation(AnimationSequence, false);
+	}
+}
+
+void UGLTFInteractionHotspotComponent::UpdateCollisionVolume()
+{
+	if (ShapeBodySetup != nullptr)
+	{
+		// TODO: Figure out where this difference stems from
+		const float Scaling = 0.15;
+		const float BillboardBoundingRadius = GetBillboardBoundingRadius() * Scaling;
+
+		check(ShapeBodySetup->AggGeom.SphereElems.Num() == 1);
+
+		if (!FMath::IsNearlyEqual(ShapeBodySetup->AggGeom.SphereElems[0].Radius, BillboardBoundingRadius) && Radius != 0.0f) {
+			BodyInstance.UpdateBodyScale(FVector(BillboardBoundingRadius / Radius), true);
+		}
 	}
 }
+
+FTransform UGLTFInteractionHotspotComponent::GetWorldTransform() const
+{
+	const AActor* Owner = GetOwner();
+
+	return Owner ? (Owner->GetTransform() * GetComponentTransform()) : GetComponentTransform();
+}
+
+float UGLTFInteractionHotspotComponent::GetBillboardBoundingRadius() const
+{
+	const FTransform WorldTransform = GetWorldTransform();
+	const FBoxSphereBounds WorldBounds = CalcBounds(WorldTransform);
+
+	return WorldBounds.SphereRadius;
+}
diff --git a/Engine/Plugins/Enterprise/GLTFExporter/MigratedFromGit/Source/GLTFExporter/Public/Components/GLTFInteractionHotspotComponent.h b/Engine/Plugins/Enterprise/GLTFExporter/MigratedFromGit/Source/GLTFExporter/Public/Components/GLTFInteractionHotspotComponent.h
--- a/Engine/Plugins/Enterprise/GLTFExporter/MigratedFromGit/Source/GLTFExporter/Public/Components/GLTFInteractionHotspotComponent.h
+++ b/Engine/Plugins/Enterprise/GLTFExporter/MigratedFromGit/Source/GLTFExporter/Public/Components/GLTFInteractionHotspotComponent.h
@@ -11,32 +11,13 @@ class UBodySetup;
 class UTexture2D;
 
 /**
- * 
+ * A component to set up hotspots which appear as billboards and allow playback of skeletal animations when cursor input is enabled.
  */
 UCLASS(BlueprintType, Blueprintable, meta = (BlueprintSpawnableComponent), hidecategories = (Sprite), DisplayName = "GLTF Interaction Hotspot Component")
 class GLTFEXPORTER_API UGLTFInteractionHotspotComponent : public UBillboardComponent
 {
 	GENERATED_UCLASS_BODY()
 public:
-	UPROPERTY(EditAnywhere, BlueprintReadWrite)
-	ASkeletalMeshActor* SkeletalMeshActor;
-
-	UPROPERTY(EditAnywhere, BlueprintReadWrite)
-	UAnimSequence* AnimationSequence;
-
-	UPROPERTY(EditAnywhere, BlueprintReadWrite)
-	UTexture2D* DefaultSprite;
-
-	UPROPERTY(EditAnywhere, BlueprintReadWrite)
-	UTexture2D* HighlightSprite;
-
-	UPROPERTY(EditAnywhere, BlueprintReadWrite)
-	UTexture2D* ClickSprite;
-
-	/** Description of collision */
-	UPROPERTY(transient, duplicatetransient)
-	UBodySetup* ShapeBodySetup;
-
 	//~ Begin UObject Interface.
 #if WITH_EDITOR
 	virtual void PostEditChangeProperty(FPropertyChangedEvent& PropertyChangedEvent) override;
@@ -50,6 +31,13 @@ protected:
 	virtual void OnRegister() override;
 	//~ End UActorComponent Interface.
 
+	//~ Begin UBillboardComponent Interface.
+public:
+	virtual void SetSprite(class UTexture2D* NewSprite) override;
+	//~ End UBillboardComponent Interface.
+
+	void SetRadius(float NewRadius);
+
 private:
 	UFUNCTION()
 	void BeginCursorOver(UPrimitiveComponent* TouchedComponent);
@@ -59,4 +47,31 @@ private:
 
 	UFUNCTION()
 	void Clicked(UPrimitiveComponent* TouchedComponent, FKey ButtonPressed);
+
+	void UpdateCollisionVolume();
+	float GetBillboardBoundingRadius() const;
+	FTransform GetWorldTransform() const;
+
+public:
+	UPROPERTY(EditAnywhere, BlueprintReadWrite, Category = "Components|GLTF Interaction Hotspot")
+	ASkeletalMeshActor* SkeletalMeshActor;
+
+	UPROPERTY(EditAnywhere, BlueprintReadWrite, Category = "Components|GLTF Interaction Hotspot")
+	UAnimSequence* AnimationSequence;
+
+	UPROPERTY(EditAnywhere, BlueprintReadWrite, Category = "Components|GLTF Interaction Hotspot")
+	UTexture2D* DefaultSprite;
+
+	UPROPERTY(EditAnywhere, BlueprintReadWrite, Category = "Components|GLTF Interaction Hotspot")
+	UTexture2D* HighlightSprite;
+
+	UPROPERTY(EditAnywhere, BlueprintReadWrite, Category = "Components|GLTF Interaction Hotspot")
+	UTexture2D* ClickSprite;
+
+	UPROPERTY(EditAnywhere, BlueprintReadWrite, Category = "Components|GLTF Interaction Hotspot")
+	float Radius;
+
+private:
+	UPROPERTY(transient, duplicatetransient)
+	UBodySetup* ShapeBodySetup;
 };
