[GLTFExporter] feat(interaction-hotspot): forward properties from actor to component

commit 5c420822f80ed85f1bcea7001cda64a3e79f6445
Date: Tue, 2 Jun 2020 16:41:16 +0200

    feat(interaction-hotspot): forward properties from actor to component
    (cherry picked from commit a631f0df1b9823aa7ab3142a8de0ce5bf642af92​)

#skipci

diff --git a/Engine/Plugins/Enterprise/GLTFExporter/MigratedFromGit/Source/GLTFExporter/Public/Actors/GLTFInteractionHotspotActor.cpp b/Engine/Plugins/Enterprise/GLTFExporter/MigratedFromGit/Source/GLTFExporter/Public/Actors/GLTFInteractionHotspotActor.cpp
--- a/Engine/Plugins/Enterprise/GLTFExporter/MigratedFromGit/Source/GLTFExporter/Public/Actors/GLTFInteractionHotspotActor.cpp
+++ b/Engine/Plugins/Enterprise/GLTFExporter/MigratedFromGit/Source/GLTFExporter/Public/Actors/GLTFInteractionHotspotActor.cpp
@@ -4,7 +4,13 @@
 #include "Components/GLTFInteractionHotspotComponent.h"
 
 AGLTFInteractionHotspotActor::AGLTFInteractionHotspotActor(const FObjectInitializer& ObjectInitializer)
-	: Super(ObjectInitializer)
+	: Super(ObjectInitializer),
+	SkeletalMeshActor(nullptr),
+	AnimationSequence(nullptr),
+	DefaultSprite(nullptr),
+	HighlightSprite(nullptr),
+	ClickSprite(nullptr),
+	Radius(50.0f)
 {
 	// A scene component with a transform in the root
 	SceneComponent = CreateDefaultSubobject<USceneComponent>(TEXT("SceneComponent"));
@@ -12,21 +18,38 @@ AGLTFInteractionHotspotActor::AGLTFInteractionHotspotActor(const FObjectInitiali
 
 	InteractionHotspotComponent = CreateDefaultSubobject<UGLTFInteractionHotspotComponent>(TEXT("InteractionHotspotComponent"));
 	InteractionHotspotComponent->SetupAttachment(SceneComponent);
-
-	OnBeginCursorOver.AddDynamic(this, &AGLTFInteractionHotspotActor::BeginCursorOver);
-}
-
-void AGLTFInteractionHotspotActor::BeginPlay()
-{
-	Super::BeginPlay();
+	ForwardPropertiesToComponent();
 }
 
-void AGLTFInteractionHotspotActor::Tick(float DeltaTime)
+void AGLTFInteractionHotspotActor::PostEditChangeProperty(FPropertyChangedEvent& PropertyChangedEvent)
 {
-	Super::Tick(DeltaTime);
+	ForwardPropertiesToComponent();
 }
 
-void AGLTFInteractionHotspotActor::BeginCursorOver(AActor* TouchedActor)
+void AGLTFInteractionHotspotActor::ForwardPropertiesToComponent()
 {
-	UE_LOG(LogTemp, Warning, TEXT("AGLTFInteractionHotspotActor::BeginCursorOver()"));
+	if (InteractionHotspotComponent->SkeletalMeshActor != SkeletalMeshActor)
+	{
+		InteractionHotspotComponent->SkeletalMeshActor = SkeletalMeshActor;
+	}
+	else if (InteractionHotspotComponent->AnimationSequence != AnimationSequence)
+	{
+		InteractionHotspotComponent->AnimationSequence = AnimationSequence;
+	}
+	else if (InteractionHotspotComponent->DefaultSprite != DefaultSprite)
+	{
+		InteractionHotspotComponent->DefaultSprite = DefaultSprite;
+	}
+	else if (InteractionHotspotComponent->HighlightSprite != HighlightSprite)
+	{
+		InteractionHotspotComponent->HighlightSprite = HighlightSprite;
+	}
+	else if (InteractionHotspotComponent->ClickSprite != ClickSprite)
+	{
+		InteractionHotspotComponent->ClickSprite = ClickSprite;
+	}
+	else if (InteractionHotspotComponent->Radius != Radius)
+	{
+		InteractionHotspotComponent->Radius = Radius;
+	}
 }
diff --git a/Engine/Plugins/Enterprise/GLTFExporter/MigratedFromGit/Source/GLTFExporter/Public/Actors/GLTFInteractionHotspotActor.h b/Engine/Plugins/Enterprise/GLTFExporter/MigratedFromGit/Source/GLTFExporter/Public/Actors/GLTFInteractionHotspotActor.h
--- a/Engine/Plugins/Enterprise/GLTFExporter/MigratedFromGit/Source/GLTFExporter/Public/Actors/GLTFInteractionHotspotActor.h
+++ b/Engine/Plugins/Enterprise/GLTFExporter/MigratedFromGit/Source/GLTFExporter/Public/Actors/GLTFInteractionHotspotActor.h
@@ -7,33 +7,52 @@
 
 class USceneComponent;
 class UGLTFInteractionHotspotComponent;
+class ASkeletalMeshActor;
+class AnimationSequence;
 
 /**
- *
+ * Actor wrapper for the GLTF hotspot component. Appears as a billboard and allows playback of skeletal animations when cursor input is enabled.
  */
 UCLASS(DisplayName = "GLTF Interaction Hotspot Actor")
 class GLTFEXPORTER_API AGLTFInteractionHotspotActor : public AActor
 {
 	GENERATED_BODY()
 private:
-	UPROPERTY(VisibleAnywhere, BlueprintReadOnly, meta = (AllowPrivateAccess = "true"))
+	UPROPERTY()
 	USceneComponent* SceneComponent;
 
-	UPROPERTY(VisibleAnywhere, BlueprintReadOnly, meta = (AllowPrivateAccess = "true"))
+	UPROPERTY()
 	UGLTFInteractionHotspotComponent* InteractionHotspotComponent;
 
+	//~ Begin UObject Interface.
 public:
-	//~ Begin AActor Interface.
 	AGLTFInteractionHotspotActor(const FObjectInitializer& ObjectInitializer);
 
-protected:
-	virtual void BeginPlay() override;
+private:
+#if WITH_EDITOR
+	virtual void PostEditChangeProperty(FPropertyChangedEvent& PropertyChangedEvent) override;
+#endif // WITH_EDITOR
+	//~ End UObject Interface.
+
+private:
+	void ForwardPropertiesToComponent();
 
 public:
-	virtual void Tick(float DeltaTime) override;
-	//~ End AActor Interface.
+	UPROPERTY(EditAnywhere, BlueprintReadWrite, Category = "Actors|GLTF Interaction Hotspot")
+	ASkeletalMeshActor* SkeletalMeshActor;
 
-private:
-	UFUNCTION()
-	void BeginCursorOver(AActor* TouchedActor);
+	UPROPERTY(EditAnywhere, BlueprintReadWrite, Category = "Actors|GLTF Interaction Hotspot")
+	UAnimSequence* AnimationSequence;
+
+	UPROPERTY(EditAnywhere, BlueprintReadWrite, Category = "Actors|GLTF Interaction Hotspot")
+	UTexture2D* DefaultSprite;
+
+	UPROPERTY(EditAnywhere, BlueprintReadWrite, Category = "Actors|GLTF Interaction Hotspot")
+	UTexture2D* HighlightSprite;
+
+	UPROPERTY(EditAnywhere, BlueprintReadWrite, Category = "Actors|GLTF Interaction Hotspot")
+	UTexture2D* ClickSprite;
+
+	UPROPERTY(EditAnywhere, BlueprintReadWrite, Category = "Actors|GLTF Interaction Hotspot")
+	float Radius;
 };
