[GLTFExporter] feat(interaction-hotspot): enable multiple of animations per hotspot

commit f228b148ddb468028e82f60cfc984a8d8b0794ab
Date: Wed, 3 Jun 2020 11:41:20 +0200

    feat(interaction-hotspot): enable multiple of animations per hotspot
    (cherry picked from commit 0f5a57f2568fce051235c6b9a58c7ee2dc6d7356​)

#skipci

diff --git a/Engine/Plugins/Enterprise/GLTFExporter/MigratedFromGit/Source/GLTFExporter/Public/Actors/GLTFInteractionHotspotActor.cpp b/Engine/Plugins/Enterprise/GLTFExporter/MigratedFromGit/Source/GLTFExporter/Public/Actors/GLTFInteractionHotspotActor.cpp
--- a/Engine/Plugins/Enterprise/GLTFExporter/MigratedFromGit/Source/GLTFExporter/Public/Actors/GLTFInteractionHotspotActor.cpp
+++ b/Engine/Plugins/Enterprise/GLTFExporter/MigratedFromGit/Source/GLTFExporter/Public/Actors/GLTFInteractionHotspotActor.cpp
@@ -1,16 +1,12 @@
 // Copyright Epic Games, Inc. All Rights Reserved.
 
 #include "Actors/GLTFInteractionHotspotActor.h"
-#include "Components/GLTFInteractionHotspotComponent.h"
 
 AGLTFInteractionHotspotActor::AGLTFInteractionHotspotActor(const FObjectInitializer& ObjectInitializer)
 	: Super(ObjectInitializer),
-	SkeletalMeshActor(nullptr),
-	AnimationSequence(nullptr),
 	DefaultSprite(nullptr),
 	HighlightSprite(nullptr),
-	ClickSprite(nullptr),
-	Radius(50.0f)
+	ClickSprite(nullptr)
 {
 	// A scene component with a transform in the root
 	SceneComponent = CreateDefaultSubobject<USceneComponent>(TEXT("SceneComponent"));
@@ -28,28 +24,44 @@ void AGLTFInteractionHotspotActor::PostEditChangeProperty(FPropertyChangedEvent&
 
 void AGLTFInteractionHotspotActor::ForwardPropertiesToComponent()
 {
-	if (InteractionHotspotComponent->SkeletalMeshActor != SkeletalMeshActor)
+	const bool UnequalAnimations = ([&]()
 	{
-		InteractionHotspotComponent->SkeletalMeshActor = SkeletalMeshActor;
-	}
-	else if (InteractionHotspotComponent->AnimationSequence != AnimationSequence)
+		if (InteractionHotspotComponent->Animations.Num() != Animations.Num())
+		{
+			return true;
+		}
+
+		for (int32 AnimationIndex = 0; AnimationIndex < Animations.Num(); ++AnimationIndex)
+		{
+			const FGLTFAnimation& ActorAnimation = Animations[AnimationIndex];
+			const FGLTFAnimation& ComponentAnimation = InteractionHotspotComponent->Animations[AnimationIndex];
+
+			if (ActorAnimation.SkeletalMeshActor != ComponentAnimation.SkeletalMeshActor || ActorAnimation.AnimationSequence != ComponentAnimation.AnimationSequence)
+			{
+				return true;
+			}
+		}
+
+		return false;
+	})();
+
+	if (UnequalAnimations)
 	{
-		InteractionHotspotComponent->AnimationSequence = AnimationSequence;
+		InteractionHotspotComponent->Animations = Animations;
 	}
-	else if (InteractionHotspotComponent->DefaultSprite != DefaultSprite)
+	
+	if (InteractionHotspotComponent->DefaultSprite != DefaultSprite)
 	{
 		InteractionHotspotComponent->DefaultSprite = DefaultSprite;
 	}
-	else if (InteractionHotspotComponent->HighlightSprite != HighlightSprite)
+	
+	if (InteractionHotspotComponent->HighlightSprite != HighlightSprite)
 	{
 		InteractionHotspotComponent->HighlightSprite = HighlightSprite;
 	}
-	else if (InteractionHotspotComponent->ClickSprite != ClickSprite)
+	
+	if (InteractionHotspotComponent->ClickSprite != ClickSprite)
 	{
 		InteractionHotspotComponent->ClickSprite = ClickSprite;
 	}
-	else if (InteractionHotspotComponent->Radius != Radius)
-	{
-		InteractionHotspotComponent->Radius = Radius;
-	}
 }
diff --git a/Engine/Plugins/Enterprise/GLTFExporter/MigratedFromGit/Source/GLTFExporter/Public/Actors/GLTFInteractionHotspotActor.h b/Engine/Plugins/Enterprise/GLTFExporter/MigratedFromGit/Source/GLTFExporter/Public/Actors/GLTFInteractionHotspotActor.h
--- a/Engine/Plugins/Enterprise/GLTFExporter/MigratedFromGit/Source/GLTFExporter/Public/Actors/GLTFInteractionHotspotActor.h
+++ b/Engine/Plugins/Enterprise/GLTFExporter/MigratedFromGit/Source/GLTFExporter/Public/Actors/GLTFInteractionHotspotActor.h
@@ -3,12 +3,10 @@
 #pragma once
 
 #include "GameFramework/Actor.h"
+#include "Components/GLTFInteractionHotspotComponent.h"
 #include "GLTFInteractionHotspotActor.generated.h"
 
 class USceneComponent;
-class UGLTFInteractionHotspotComponent;
-class ASkeletalMeshActor;
-class AnimationSequence;
 
 /**
  * Actor wrapper for the GLTF hotspot component. Appears as a billboard and allows playback of skeletal animations when cursor input is enabled.
@@ -17,13 +15,6 @@ UCLASS(DisplayName = "GLTF Interaction Hotspot Actor")
 class GLTFEXPORTER_API AGLTFInteractionHotspotActor : public AActor
 {
 	GENERATED_BODY()
-private:
-	UPROPERTY()
-	USceneComponent* SceneComponent;
-
-	UPROPERTY()
-	UGLTFInteractionHotspotComponent* InteractionHotspotComponent;
-
 	//~ Begin UObject Interface.
 public:
 	AGLTFInteractionHotspotActor(const FObjectInitializer& ObjectInitializer);
@@ -38,21 +29,22 @@ private:
 	void ForwardPropertiesToComponent();
 
 public:
-	UPROPERTY(EditAnywhere, BlueprintReadWrite, Category = "Actors|GLTF Interaction Hotspot")
-	ASkeletalMeshActor* SkeletalMeshActor;
+	UPROPERTY(EditAnywhere, BlueprintReadWrite, Category = "GLTF Interaction Hotspot")
+	TArray<FGLTFAnimation> Animations;
 
-	UPROPERTY(EditAnywhere, BlueprintReadWrite, Category = "Actors|GLTF Interaction Hotspot")
-	UAnimSequence* AnimationSequence;
-
-	UPROPERTY(EditAnywhere, BlueprintReadWrite, Category = "Actors|GLTF Interaction Hotspot")
+	UPROPERTY(EditAnywhere, BlueprintReadWrite, Category = "GLTF Interaction Hotspot")
 	UTexture2D* DefaultSprite;
 
-	UPROPERTY(EditAnywhere, BlueprintReadWrite, Category = "Actors|GLTF Interaction Hotspot")
+	UPROPERTY(EditAnywhere, BlueprintReadWrite, Category = "GLTF Interaction Hotspot")
 	UTexture2D* HighlightSprite;
 
-	UPROPERTY(EditAnywhere, BlueprintReadWrite, Category = "Actors|GLTF Interaction Hotspot")
+	UPROPERTY(EditAnywhere, BlueprintReadWrite, Category = "GLTF Interaction Hotspot")
 	UTexture2D* ClickSprite;
 
-	UPROPERTY(EditAnywhere, BlueprintReadWrite, Category = "Actors|GLTF Interaction Hotspot")
-	float Radius;
+private:
+	UPROPERTY()
+	USceneComponent* SceneComponent;
+
+	UPROPERTY()
+	UGLTFInteractionHotspotComponent* InteractionHotspotComponent;
 };
diff --git a/Engine/Plugins/Enterprise/GLTFExporter/MigratedFromGit/Source/GLTFExporter/Public/Components/GLTFInteractionHotspotComponent.cpp b/Engine/Plugins/Enterprise/GLTFExporter/MigratedFromGit/Source/GLTFExporter/Public/Components/GLTFInteractionHotspotComponent.cpp
--- a/Engine/Plugins/Enterprise/GLTFExporter/MigratedFromGit/Source/GLTFExporter/Public/Components/GLTFInteractionHotspotComponent.cpp
+++ b/Engine/Plugins/Enterprise/GLTFExporter/MigratedFromGit/Source/GLTFExporter/Public/Components/GLTFInteractionHotspotComponent.cpp
@@ -7,14 +7,17 @@
 #include "UObject/ConstructorHelpers.h"
 #include "PhysicsEngine/BodySetup.h"
 
+namespace
+{
+	static const float DefaultCollisionVolumeRadius = 50.0f;
+} // Anonymous namespace
+
 UGLTFInteractionHotspotComponent::UGLTFInteractionHotspotComponent(const FObjectInitializer& ObjectInitializer)
 	: Super(ObjectInitializer),
-	SkeletalMeshActor(nullptr),
-	AnimationSequence(nullptr),
 	DefaultSprite(nullptr),
 	HighlightSprite(nullptr),
 	ClickSprite(nullptr),
-	Radius(50.0f)
+	ShapeBodySetup(nullptr)
 {
 	// Setup the most minimalistic collision profile for mouse input events
 	SetCollisionEnabled(ECollisionEnabled::QueryAndPhysics);
@@ -41,10 +44,6 @@ void UGLTFInteractionHotspotComponent::PostEditChangeProperty(FPropertyChangedEv
 		{
 			SetSprite(DefaultSprite);
 		}
-		else if (PropertyName == TEXT("Radius"))
-		{
-			SetRadius(Radius);
-		}
 	}
 }
 
@@ -59,7 +58,7 @@ void UGLTFInteractionHotspotComponent::OnRegister()
 {
 	ShapeBodySetup = NewObject<UBodySetup>(this);
 	ShapeBodySetup->CollisionTraceFlag = CTF_UseSimpleAsComplex;
-	ShapeBodySetup->AggGeom.SphereElems.Add(Radius);
+	ShapeBodySetup->AggGeom.SphereElems.Add(DefaultCollisionVolumeRadius);
 	
 	BodyInstance.InitBody(ShapeBodySetup, GetWorldTransform(), this, GetWorld()->GetPhysicsScene());
 
@@ -76,11 +75,6 @@ void UGLTFInteractionHotspotComponent::SetSprite(class UTexture2D* NewSprite)
 	UpdateCollisionVolume();
 }
 
-void UGLTFInteractionHotspotComponent::SetRadius(float NewRadius)
-{
-	// TODO: Implement
-}
-
 void UGLTFInteractionHotspotComponent::BeginCursorOver(UPrimitiveComponent* TouchedComponent)
 {
 	if (HighlightSprite != nullptr)
@@ -96,10 +90,13 @@ void UGLTFInteractionHotspotComponent::EndCursorOver(UPrimitiveComponent* Touche
 
 void UGLTFInteractionHotspotComponent::Clicked(UPrimitiveComponent* TouchedComponent, FKey ButtonPressed)
 {
-	if (SkeletalMeshActor != nullptr && AnimationSequence != nullptr)
+	for (TArray<FGLTFAnimation>::TConstIterator Animation = Animations.CreateConstIterator(); Animation; ++Animation)
 	{
-		USkeletalMeshComponent* SkeletalMeshComponent = SkeletalMeshActor->GetSkeletalMeshComponent();
-		SkeletalMeshComponent->PlayAnimation(AnimationSequence, false);
+		if (Animation->SkeletalMeshActor != nullptr && Animation->AnimationSequence != nullptr)
+		{
+			USkeletalMeshComponent* SkeletalMeshComponent = Animation->SkeletalMeshActor->GetSkeletalMeshComponent();
+			SkeletalMeshComponent->PlayAnimation(Animation->AnimationSequence, false);
+		}
 	}
 }
 
@@ -113,8 +110,9 @@ void UGLTFInteractionHotspotComponent::UpdateCollisionVolume()
 
 		check(ShapeBodySetup->AggGeom.SphereElems.Num() == 1);
 
-		if (!FMath::IsNearlyEqual(ShapeBodySetup->AggGeom.SphereElems[0].Radius, BillboardBoundingRadius) && Radius != 0.0f) {
-			BodyInstance.UpdateBodyScale(FVector(BillboardBoundingRadius / Radius), true);
+		if (!FMath::IsNearlyEqual(ShapeBodySetup->AggGeom.SphereElems[0].Radius, BillboardBoundingRadius) && DefaultCollisionVolumeRadius != 0.0f)
+		{
+			BodyInstance.UpdateBodyScale(FVector(BillboardBoundingRadius / DefaultCollisionVolumeRadius), true);
 		}
 	}
 }
@@ -123,7 +121,7 @@ FTransform UGLTFInteractionHotspotComponent::GetWorldTransform() const
 {
 	const AActor* Owner = GetOwner();
 
-	return Owner ? (Owner->GetTransform() * GetComponentTransform()) : GetComponentTransform();
+	return (Owner != nullptr) ? (Owner->GetTransform() * GetComponentTransform()) : GetComponentTransform();
 }
 
 float UGLTFInteractionHotspotComponent::GetBillboardBoundingRadius() const
diff --git a/Engine/Plugins/Enterprise/GLTFExporter/MigratedFromGit/Source/GLTFExporter/Public/Components/GLTFInteractionHotspotComponent.h b/Engine/Plugins/Enterprise/GLTFExporter/MigratedFromGit/Source/GLTFExporter/Public/Components/GLTFInteractionHotspotComponent.h
--- a/Engine/Plugins/Enterprise/GLTFExporter/MigratedFromGit/Source/GLTFExporter/Public/Components/GLTFInteractionHotspotComponent.h
+++ b/Engine/Plugins/Enterprise/GLTFExporter/MigratedFromGit/Source/GLTFExporter/Public/Components/GLTFInteractionHotspotComponent.h
@@ -10,6 +10,18 @@ class UAnimSequence;
 class UBodySetup;
 class UTexture2D;
 
+USTRUCT(BlueprintType)
+struct FGLTFAnimation
+{
+	GENERATED_USTRUCT_BODY()
+
+	UPROPERTY(EditAnywhere, BlueprintReadWrite)
+	ASkeletalMeshActor* SkeletalMeshActor;
+
+	UPROPERTY(EditAnywhere, BlueprintReadWrite)
+	UAnimSequence* AnimationSequence;
+};
+
 /**
  * A component to set up hotspots which appear as billboards and allow playback of skeletal animations when cursor input is enabled.
  */
@@ -53,24 +65,18 @@ private:
 	FTransform GetWorldTransform() const;
 
 public:
-	UPROPERTY(EditAnywhere, BlueprintReadWrite, Category = "Components|GLTF Interaction Hotspot")
-	ASkeletalMeshActor* SkeletalMeshActor;
+	UPROPERTY(EditAnywhere, BlueprintReadWrite, Category = "GLTF Interaction Hotspot")
+	TArray<FGLTFAnimation> Animations;
 
-	UPROPERTY(EditAnywhere, BlueprintReadWrite, Category = "Components|GLTF Interaction Hotspot")
-	UAnimSequence* AnimationSequence;
-
-	UPROPERTY(EditAnywhere, BlueprintReadWrite, Category = "Components|GLTF Interaction Hotspot")
+	UPROPERTY(EditAnywhere, BlueprintReadWrite, Category = "GLTF Interaction Hotspot")
 	UTexture2D* DefaultSprite;
 
-	UPROPERTY(EditAnywhere, BlueprintReadWrite, Category = "Components|GLTF Interaction Hotspot")
+	UPROPERTY(EditAnywhere, BlueprintReadWrite, Category = "GLTF Interaction Hotspot")
 	UTexture2D* HighlightSprite;
 
-	UPROPERTY(EditAnywhere, BlueprintReadWrite, Category = "Components|GLTF Interaction Hotspot")
+	UPROPERTY(EditAnywhere, BlueprintReadWrite, Category = "GLTF Interaction Hotspot")
 	UTexture2D* ClickSprite;
 
-	UPROPERTY(EditAnywhere, BlueprintReadWrite, Category = "Components|GLTF Interaction Hotspot")
-	float Radius;
-
 private:
 	UPROPERTY(transient, duplicatetransient)
 	UBodySetup* ShapeBodySetup;
