[GLTFExporter] feat(interaction-hotspot): implement toggle behavior

commit 4e6b453f3ca5c33e1e2023bdd9ded6eabaaa4a99
Date: Thu, 4 Jun 2020 15:49:53 +0200

    feat(interaction-hotspot): implement toggle behavior
    (cherry picked from commit e26c6c8a04903c885cf13c6e9bd5e71c9981296b​)

#skipci

diff --git a/Engine/Plugins/Enterprise/GLTFExporter/MigratedFromGit/Source/GLTFExporter/Public/Actors/GLTFInteractionHotspotActor.cpp b/Engine/Plugins/Enterprise/GLTFExporter/MigratedFromGit/Source/GLTFExporter/Public/Actors/GLTFInteractionHotspotActor.cpp
--- a/Engine/Plugins/Enterprise/GLTFExporter/MigratedFromGit/Source/GLTFExporter/Public/Actors/GLTFInteractionHotspotActor.cpp
+++ b/Engine/Plugins/Enterprise/GLTFExporter/MigratedFromGit/Source/GLTFExporter/Public/Actors/GLTFInteractionHotspotActor.cpp
@@ -6,7 +6,7 @@ AGLTFInteractionHotspotActor::AGLTFInteractionHotspotActor(const FObjectInitiali
 	: Super(ObjectInitializer),
 	DefaultSprite(nullptr),
 	HighlightSprite(nullptr),
-	ClickSprite(nullptr)
+	ToggledSprite(nullptr)
 {
 	// A scene component with a transform in the root
 	SceneComponent = CreateDefaultSubobject<USceneComponent>(TEXT("SceneComponent"));
@@ -17,14 +17,16 @@ AGLTFInteractionHotspotActor::AGLTFInteractionHotspotActor(const FObjectInitiali
 	ForwardPropertiesToComponent();
 }
 
+#if WITH_EDITOR
 void AGLTFInteractionHotspotActor::PostEditChangeProperty(FPropertyChangedEvent& PropertyChangedEvent)
 {
 	ForwardPropertiesToComponent();
 }
+#endif // WITH_EDITOR
 
 void AGLTFInteractionHotspotActor::ForwardPropertiesToComponent()
 {
-	const bool UnequalAnimations = ([&]()
+	const bool IsUnequalAnimations = ([&]()
 	{
 		if (InteractionHotspotComponent->Animations.Num() != Animations.Num())
 		{
@@ -45,7 +47,7 @@ void AGLTFInteractionHotspotActor::ForwardPropertiesToComponent()
 		return false;
 	})();
 
-	if (UnequalAnimations)
+	if (IsUnequalAnimations)
 	{
 		InteractionHotspotComponent->Animations = Animations;
 	}
@@ -53,6 +55,7 @@ void AGLTFInteractionHotspotActor::ForwardPropertiesToComponent()
 	if (InteractionHotspotComponent->DefaultSprite != DefaultSprite)
 	{
 		InteractionHotspotComponent->DefaultSprite = DefaultSprite;
+		InteractionHotspotComponent->SetSprite(DefaultSprite);
 	}
 	
 	if (InteractionHotspotComponent->HighlightSprite != HighlightSprite)
@@ -60,8 +63,8 @@ void AGLTFInteractionHotspotActor::ForwardPropertiesToComponent()
 		InteractionHotspotComponent->HighlightSprite = HighlightSprite;
 	}
 	
-	if (InteractionHotspotComponent->ClickSprite != ClickSprite)
+	if (InteractionHotspotComponent->ToggledSprite != ToggledSprite)
 	{
-		InteractionHotspotComponent->ClickSprite = ClickSprite;
+		InteractionHotspotComponent->ToggledSprite = ToggledSprite;
 	}
 }
diff --git a/Engine/Plugins/Enterprise/GLTFExporter/MigratedFromGit/Source/GLTFExporter/Public/Actors/GLTFInteractionHotspotActor.h b/Engine/Plugins/Enterprise/GLTFExporter/MigratedFromGit/Source/GLTFExporter/Public/Actors/GLTFInteractionHotspotActor.h
--- a/Engine/Plugins/Enterprise/GLTFExporter/MigratedFromGit/Source/GLTFExporter/Public/Actors/GLTFInteractionHotspotActor.h
+++ b/Engine/Plugins/Enterprise/GLTFExporter/MigratedFromGit/Source/GLTFExporter/Public/Actors/GLTFInteractionHotspotActor.h
@@ -6,8 +6,6 @@
 #include "Components/GLTFInteractionHotspotComponent.h"
 #include "GLTFInteractionHotspotActor.generated.h"
 
-class USceneComponent;
-
 /**
  * Actor wrapper for the GLTF hotspot component. Appears as a billboard and allows playback of skeletal animations when cursor input is enabled.
  */
@@ -29,22 +27,27 @@ private:
 	void ForwardPropertiesToComponent();
 
 public:
+	/** List of skeletal meshes and animations to be played when the hotspot is interacted with */
 	UPROPERTY(EditAnywhere, BlueprintReadWrite, Category = "GLTF Interaction Hotspot")
 	TArray<FGLTFAnimation> Animations;
 
+	/** The billboard image that will be shown when the hotspot is in an inactive state */
 	UPROPERTY(EditAnywhere, BlueprintReadWrite, Category = "GLTF Interaction Hotspot")
 	UTexture2D* DefaultSprite;
 
+	/** The optional billboard image that will be shown when a cursor enters the hotspot */
 	UPROPERTY(EditAnywhere, BlueprintReadWrite, Category = "GLTF Interaction Hotspot")
 	UTexture2D* HighlightSprite;
 
+	/** The optional billboard image that will be shown when the hotspot is toggled by a click */
 	UPROPERTY(EditAnywhere, BlueprintReadWrite, Category = "GLTF Interaction Hotspot")
-	UTexture2D* ClickSprite;
+	UTexture2D* ToggledSprite;
 
 private:
 	UPROPERTY()
 	USceneComponent* SceneComponent;
 
-	UPROPERTY()
+public:
+	UPROPERTY(BlueprintReadOnly, Category = "GLTF Interaction Hotspot")
 	UGLTFInteractionHotspotComponent* InteractionHotspotComponent;
 };
diff --git a/Engine/Plugins/Enterprise/GLTFExporter/MigratedFromGit/Source/GLTFExporter/Public/Components/GLTFInteractionHotspotComponent.cpp b/Engine/Plugins/Enterprise/GLTFExporter/MigratedFromGit/Source/GLTFExporter/Public/Components/GLTFInteractionHotspotComponent.cpp
--- a/Engine/Plugins/Enterprise/GLTFExporter/MigratedFromGit/Source/GLTFExporter/Public/Components/GLTFInteractionHotspotComponent.cpp
+++ b/Engine/Plugins/Enterprise/GLTFExporter/MigratedFromGit/Source/GLTFExporter/Public/Components/GLTFInteractionHotspotComponent.cpp
@@ -6,6 +6,7 @@
 #include "Engine/CollisionProfile.h"
 #include "UObject/ConstructorHelpers.h"
 #include "PhysicsEngine/BodySetup.h"
+#include "Animation/AnimSingleNodeInstance.h"
 
 namespace
 {
@@ -16,8 +17,9 @@ UGLTFInteractionHotspotComponent::UGLTFInteractionHotspotComponent(const FObject
 	: Super(ObjectInitializer),
 	DefaultSprite(nullptr),
 	HighlightSprite(nullptr),
-	ClickSprite(nullptr),
-	ShapeBodySetup(nullptr)
+	ToggledSprite(nullptr),
+	ShapeBodySetup(nullptr),
+	bToggled(bToggled)
 {
 	// Setup the most minimalistic collision profile for mouse input events
 	SetCollisionEnabled(ECollisionEnabled::QueryAndPhysics);
@@ -32,6 +34,7 @@ UGLTFInteractionHotspotComponent::UGLTFInteractionHotspotComponent(const FObject
 	OnClicked.AddDynamic(this, &UGLTFInteractionHotspotComponent::Clicked);
 }
 
+#if WITH_EDITOR
 void UGLTFInteractionHotspotComponent::PostEditChangeProperty(FPropertyChangedEvent& PropertyChangedEvent)
 {
 	const FProperty* PropertyThatChanged = PropertyChangedEvent.Property;
@@ -46,6 +49,7 @@ void UGLTFInteractionHotspotComponent::PostEditChangeProperty(FPropertyChangedEv
 		}
 	}
 }
+#endif // WITH_EDITOR
 
 void UGLTFInteractionHotspotComponent::BeginPlay()
 {
@@ -77,7 +81,7 @@ void UGLTFInteractionHotspotComponent::SetSprite(class UTexture2D* NewSprite)
 
 void UGLTFInteractionHotspotComponent::BeginCursorOver(UPrimitiveComponent* TouchedComponent)
 {
-	if (HighlightSprite != nullptr)
+	if ((ToggledSprite == nullptr || !bToggled) && HighlightSprite != nullptr)
 	{
 		SetSprite(HighlightSprite);
 	}
@@ -85,19 +89,54 @@ void UGLTFInteractionHotspotComponent::BeginCursorOver(UPrimitiveComponent* Touc
 
 void UGLTFInteractionHotspotComponent::EndCursorOver(UPrimitiveComponent* TouchedComponent)
 {
-	SetSprite(DefaultSprite);
+	if (ToggledSprite == nullptr || !bToggled)
+	{
+		SetSprite(DefaultSprite);
+	}
 }
 
 void UGLTFInteractionHotspotComponent::Clicked(UPrimitiveComponent* TouchedComponent, FKey ButtonPressed)
 {
+	if (!bToggled && ToggledSprite != nullptr)
+	{
+		SetSprite(ToggledSprite);
+	}
+	else
+	{
+		SetSprite(DefaultSprite);
+	}
+
+	const bool bReverseAnimation = bToggled;
+
 	for (TArray<FGLTFAnimation>::TConstIterator Animation = Animations.CreateConstIterator(); Animation; ++Animation)
 	{
 		if (Animation->SkeletalMeshActor != nullptr && Animation->AnimationSequence != nullptr)
 		{
 			USkeletalMeshComponent* SkeletalMeshComponent = Animation->SkeletalMeshActor->GetSkeletalMeshComponent();
-			SkeletalMeshComponent->PlayAnimation(Animation->AnimationSequence, false);
+			SkeletalMeshComponent->SetAnimationMode(EAnimationMode::Type::AnimationSingleNode);
+			const float AbsolutePlayRate = FMath::Abs(SkeletalMeshComponent->GetPlayRate());
+
+			if (SkeletalMeshComponent->IsPlaying())
+			{
+				UAnimSingleNodeInstance* SingleNodeInstance = SkeletalMeshComponent->GetSingleNodeInstance();
+				
+				if (SingleNodeInstance != nullptr && SingleNodeInstance->GetAnimationAsset() == Animation->AnimationSequence)
+				{
+					// If the same animation is already playing, just reverse the play rate for a smooth transition
+					SkeletalMeshComponent->SetPlayRate(AbsolutePlayRate * -1.0f);
+					
+					continue;
+				}
+			}
+			
+			SkeletalMeshComponent->SetAnimation(Animation->AnimationSequence);
+			SkeletalMeshComponent->SetPlayRate(AbsolutePlayRate * (bReverseAnimation ? -1.0f : 1.0f));
+			SkeletalMeshComponent->SetPosition(bReverseAnimation ? Animation->AnimationSequence->GetPlayLength() : 0.0f);
+			SkeletalMeshComponent->Play(false);
 		}
 	}
+
+	bToggled = !bToggled;
 }
 
 void UGLTFInteractionHotspotComponent::UpdateCollisionVolume()
diff --git a/Engine/Plugins/Enterprise/GLTFExporter/MigratedFromGit/Source/GLTFExporter/Public/Components/GLTFInteractionHotspotComponent.h b/Engine/Plugins/Enterprise/GLTFExporter/MigratedFromGit/Source/GLTFExporter/Public/Components/GLTFInteractionHotspotComponent.h
--- a/Engine/Plugins/Enterprise/GLTFExporter/MigratedFromGit/Source/GLTFExporter/Public/Components/GLTFInteractionHotspotComponent.h
+++ b/Engine/Plugins/Enterprise/GLTFExporter/MigratedFromGit/Source/GLTFExporter/Public/Components/GLTFInteractionHotspotComponent.h
@@ -48,8 +48,6 @@ public:
 	virtual void SetSprite(class UTexture2D* NewSprite) override;
 	//~ End UBillboardComponent Interface.
 
-	void SetRadius(float NewRadius);
-
 private:
 	UFUNCTION()
 	void BeginCursorOver(UPrimitiveComponent* TouchedComponent);
@@ -65,19 +63,25 @@ private:
 	FTransform GetWorldTransform() const;
 
 public:
+	/** List of skeletal meshes and animations to be played when the hotspot is interacted with */
 	UPROPERTY(EditAnywhere, BlueprintReadWrite, Category = "GLTF Interaction Hotspot")
 	TArray<FGLTFAnimation> Animations;
 
+	/** The billboard image that will be shown when the hotspot is in an inactive state */
 	UPROPERTY(EditAnywhere, BlueprintReadWrite, Category = "GLTF Interaction Hotspot")
 	UTexture2D* DefaultSprite;
 
+	/** The optional billboard image that will be shown when a cursor enters the hotspot */
 	UPROPERTY(EditAnywhere, BlueprintReadWrite, Category = "GLTF Interaction Hotspot")
 	UTexture2D* HighlightSprite;
 
+	/** The optional billboard image that will be shown when the hotspot is toggled by a click */
 	UPROPERTY(EditAnywhere, BlueprintReadWrite, Category = "GLTF Interaction Hotspot")
-	UTexture2D* ClickSprite;
+	UTexture2D* ToggledSprite;
 
 private:
 	UPROPERTY(transient, duplicatetransient)
 	UBodySetup* ShapeBodySetup;
+
+	bool bToggled;
 };
