[GLTFExporter] refactor(interaction-hotspot): perform equality check using operator functions instead of separate logic

commit c5af465e98ee360fd3c60c1d0da4a4284b3c18ba
Date: Tue, 9 Jun 2020 13:33:33 +0200

    refactor(interaction-hotspot): perform equality check using operator functions instead of separate logic
    (cherry picked from commit eedf6906ff676cd0cb4cfc7f12584771c3682965​)

#skipci

diff --git a/Engine/Plugins/Enterprise/GLTFExporter/MigratedFromGit/Source/GLTFExporter/Public/Actors/GLTFInteractionHotspotActor.cpp b/Engine/Plugins/Enterprise/GLTFExporter/MigratedFromGit/Source/GLTFExporter/Public/Actors/GLTFInteractionHotspotActor.cpp
--- a/Engine/Plugins/Enterprise/GLTFExporter/MigratedFromGit/Source/GLTFExporter/Public/Actors/GLTFInteractionHotspotActor.cpp
+++ b/Engine/Plugins/Enterprise/GLTFExporter/MigratedFromGit/Source/GLTFExporter/Public/Actors/GLTFInteractionHotspotActor.cpp
@@ -26,28 +26,7 @@ void AGLTFInteractionHotspotActor::PostEditChangeProperty(FPropertyChangedEvent&
 
 void AGLTFInteractionHotspotActor::ForwardPropertiesToComponent()
 {
-	const bool IsUnequalAnimations = ([&]()
-	{
-		if (InteractionHotspotComponent->Animations.Num() != Animations.Num())
-		{
-			return true;
-		}
-
-		for (int32 AnimationIndex = 0; AnimationIndex < Animations.Num(); ++AnimationIndex)
-		{
-			const FGLTFAnimation& ActorAnimation = Animations[AnimationIndex];
-			const FGLTFAnimation& ComponentAnimation = InteractionHotspotComponent->Animations[AnimationIndex];
-
-			if (ActorAnimation.SkeletalMeshActor != ComponentAnimation.SkeletalMeshActor || ActorAnimation.AnimationSequence != ComponentAnimation.AnimationSequence)
-			{
-				return true;
-			}
-		}
-
-		return false;
-	})();
-
-	if (IsUnequalAnimations)
+	if (InteractionHotspotComponent->Animations != Animations)
 	{
 		InteractionHotspotComponent->Animations = Animations;
 	}
diff --git a/Engine/Plugins/Enterprise/GLTFExporter/MigratedFromGit/Source/GLTFExporter/Public/Components/GLTFInteractionHotspotComponent.h b/Engine/Plugins/Enterprise/GLTFExporter/MigratedFromGit/Source/GLTFExporter/Public/Components/GLTFInteractionHotspotComponent.h
--- a/Engine/Plugins/Enterprise/GLTFExporter/MigratedFromGit/Source/GLTFExporter/Public/Components/GLTFInteractionHotspotComponent.h
+++ b/Engine/Plugins/Enterprise/GLTFExporter/MigratedFromGit/Source/GLTFExporter/Public/Components/GLTFInteractionHotspotComponent.h
@@ -20,6 +20,17 @@ struct FGLTFAnimation
 
 	UPROPERTY(EditAnywhere, BlueprintReadWrite)
 	UAnimSequence* AnimationSequence;
+
+	bool operator == (const FGLTFAnimation& OtherAnimation) const
+	{
+		return SkeletalMeshActor == OtherAnimation.SkeletalMeshActor &&
+			AnimationSequence == OtherAnimation.AnimationSequence;
+	}
+
+	bool operator != (const FGLTFAnimation& OtherAnimation) const
+	{
+		return !(*this == OtherAnimation);
+	}
 };
 
 /**
