[GLTFExporter] refactor(interaction-hotspot): rename variable for clarity

commit 06d9aa5c7a9baf39f67710fcd8eda21045a372a7
Date: Tue, 9 Jun 2020 13:54:46 +0200

    refactor(interaction-hotspot): rename variable for clarity
    (cherry picked from commit ba45fa55338e1c13e7abb2079c74374b1c1a76e6​)

#skipci

diff --git a/Engine/Plugins/Enterprise/GLTFExporter/MigratedFromGit/Source/GLTFExporter/Public/Components/GLTFInteractionHotspotComponent.cpp b/Engine/Plugins/Enterprise/GLTFExporter/MigratedFromGit/Source/GLTFExporter/Public/Components/GLTFInteractionHotspotComponent.cpp
--- a/Engine/Plugins/Enterprise/GLTFExporter/MigratedFromGit/Source/GLTFExporter/Public/Components/GLTFInteractionHotspotComponent.cpp
+++ b/Engine/Plugins/Enterprise/GLTFExporter/MigratedFromGit/Source/GLTFExporter/Public/Components/GLTFInteractionHotspotComponent.cpp
@@ -10,7 +10,7 @@
 
 namespace
 {
-	static const float DefaultCollisionVolumeRadius = 50.0f;
+	const float UnitSphereRadius = 50.0f;
 } // Anonymous namespace
 
 UGLTFInteractionHotspotComponent::UGLTFInteractionHotspotComponent(const FObjectInitializer& ObjectInitializer)
@@ -62,7 +62,7 @@ void UGLTFInteractionHotspotComponent::OnRegister()
 {
 	ShapeBodySetup = NewObject<UBodySetup>(this);
 	ShapeBodySetup->CollisionTraceFlag = CTF_UseSimpleAsComplex;
-	ShapeBodySetup->AggGeom.SphereElems.Add(DefaultCollisionVolumeRadius);
+	ShapeBodySetup->AggGeom.SphereElems.Add(UnitSphereRadius);
 	
 	BodyInstance.InitBody(ShapeBodySetup, GetWorldTransform(), this, GetWorld()->GetPhysicsScene());
 
@@ -149,9 +149,9 @@ void UGLTFInteractionHotspotComponent::UpdateCollisionVolume()
 
 		check(ShapeBodySetup->AggGeom.SphereElems.Num() == 1);
 
-		if (!FMath::IsNearlyEqual(ShapeBodySetup->AggGeom.SphereElems[0].Radius, BillboardBoundingRadius) && DefaultCollisionVolumeRadius != 0.0f)
+		if (!FMath::IsNearlyEqual(ShapeBodySetup->AggGeom.SphereElems[0].Radius, BillboardBoundingRadius) && UnitSphereRadius != 0.0f)
 		{
-			BodyInstance.UpdateBodyScale(FVector(BillboardBoundingRadius / DefaultCollisionVolumeRadius), true);
+			BodyInstance.UpdateBodyScale(FVector(BillboardBoundingRadius / UnitSphereRadius), true);
 		}
 	}
 }
