[GLTFExporter] fix(interaction-hotspot): add missing call to super function

commit fb1ed7053b674c07a79dd43000caea617e0e44c7
Date: Mon, 15 Jun 2020 11:02:29 +0200

    fix(interaction-hotspot): add missing call to super function
    (cherry picked from commit 80e2e81135ebdc801b4e1b7006b00e6125ef2421​)

#skipci

diff --git a/Engine/Plugins/Enterprise/GLTFExporter/MigratedFromGit/Source/GLTFExporter/Public/Components/GLTFInteractionHotspotComponent.cpp b/Engine/Plugins/Enterprise/GLTFExporter/MigratedFromGit/Source/GLTFExporter/Public/Components/GLTFInteractionHotspotComponent.cpp
--- a/Engine/Plugins/Enterprise/GLTFExporter/MigratedFromGit/Source/GLTFExporter/Public/Components/GLTFInteractionHotspotComponent.cpp
+++ b/Engine/Plugins/Enterprise/GLTFExporter/MigratedFromGit/Source/GLTFExporter/Public/Components/GLTFInteractionHotspotComponent.cpp
@@ -37,6 +37,8 @@ UGLTFInteractionHotspotComponent::UGLTFInteractionHotspotComponent(const FObject
 #if WITH_EDITOR
 void UGLTFInteractionHotspotComponent::PostEditChangeProperty(FPropertyChangedEvent& PropertyChangedEvent)
 {
+	Super::PostEditChangeProperty(PropertyChangedEvent);
+
 	const FProperty* PropertyThatChanged = PropertyChangedEvent.Property;
 
 	if (PropertyThatChanged)
