[GLTFExporter] fix(interaction-hotspot): add missing call to super function

commit 70b7831cb409395cfe7d2d6e3661ce4ea3990df4
Date: Mon, 15 Jun 2020 11:19:40 +0200

    fix(interaction-hotspot): add missing call to super function
    (cherry picked from commit 2e1a71516a0ee1e83fe3d77b4431f434cfa60f2d​)

#skipci

diff --git a/Engine/Plugins/Enterprise/GLTFExporter/MigratedFromGit/Source/GLTFExporter/Public/Actors/GLTFInteractionHotspotActor.cpp b/Engine/Plugins/Enterprise/GLTFExporter/MigratedFromGit/Source/GLTFExporter/Public/Actors/GLTFInteractionHotspotActor.cpp
--- a/Engine/Plugins/Enterprise/GLTFExporter/MigratedFromGit/Source/GLTFExporter/Public/Actors/GLTFInteractionHotspotActor.cpp
+++ b/Engine/Plugins/Enterprise/GLTFExporter/MigratedFromGit/Source/GLTFExporter/Public/Actors/GLTFInteractionHotspotActor.cpp
@@ -20,6 +20,8 @@ AGLTFInteractionHotspotActor::AGLTFInteractionHotspotActor(const FObjectInitiali
 #if WITH_EDITOR
 void AGLTFInteractionHotspotActor::PostEditChangeProperty(FPropertyChangedEvent& PropertyChangedEvent)
 {
+	Super::PostEditChangeProperty(PropertyChangedEvent);
+
 	ForwardPropertiesToComponent();
 }
 #endif // WITH_EDITOR
