[GLTFExporter] fix(interaction-hotspot): set up collision volume same way as UShapeComponent to avoid crash when changing parent transform

commit 4e285b262f969e01ad3be77c7c751af182b9d916
Date: Mon, 15 Jun 2020 13:04:22 +0200

    fix(interaction-hotspot): set up collision volume same way as UShapeComponent to avoid crash when changing parent transform
    (cherry picked from commit 9ac486ca124630cd59fc36c6150c59b5090bb761​)

#skipci

diff --git a/Engine/Plugins/Enterprise/GLTFExporter/MigratedFromGit/Source/GLTFExporter/Public/Components/GLTFInteractionHotspotComponent.cpp b/Engine/Plugins/Enterprise/GLTFExporter/MigratedFromGit/Source/GLTFExporter/Public/Components/GLTFInteractionHotspotComponent.cpp
--- a/Engine/Plugins/Enterprise/GLTFExporter/MigratedFromGit/Source/GLTFExporter/Public/Components/GLTFInteractionHotspotComponent.cpp
+++ b/Engine/Plugins/Enterprise/GLTFExporter/MigratedFromGit/Source/GLTFExporter/Public/Components/GLTFInteractionHotspotComponent.cpp
@@ -62,13 +62,44 @@ void UGLTFInteractionHotspotComponent::BeginPlay()
 
 void UGLTFInteractionHotspotComponent::OnRegister()
 {
-	ShapeBodySetup = NewObject<UBodySetup>(this);
+	ShapeBodySetup = NewObject<UBodySetup>(this, NAME_None, RF_Transient);
+
+	if (GUObjectArray.IsDisregardForGC(this))
+	{
+		ShapeBodySetup->AddToRoot();
+	}
+
+	ShapeBodySetup->AddToCluster(this);
+
+	if (ShapeBodySetup->HasAnyInternalFlags(EInternalObjectFlags::Async) && GUObjectClusters.GetObjectCluster(ShapeBodySetup))
+	{
+		ShapeBodySetup->ClearInternalFlags(EInternalObjectFlags::Async);
+	}
+
 	ShapeBodySetup->CollisionTraceFlag = CTF_UseSimpleAsComplex;
-	ShapeBodySetup->AggGeom.SphereElems.Add(UnitSphereRadius);
-	
+	ShapeBodySetup->AggGeom.SphereElems.Add(FKSphereElem(UnitSphereRadius));
+
 	BodyInstance.InitBody(ShapeBodySetup, GetWorldTransform(), this, GetWorld()->GetPhysicsScene());
 
-	RecreatePhysicsState();
+	if (BodyInstance.IsValidBodyInstance())
+	{
+#if WITH_PHYSX
+		FPhysicsCommand::ExecuteWrite(BodyInstance.GetActorReferenceWithWelding(), [this](const FPhysicsActorHandle& Actor)
+		{
+			TArray<FPhysicsShapeHandle> Shapes;
+			BodyInstance.GetAllShapes_AssumesLocked(Shapes);
+
+			for (FPhysicsShapeHandle& Shape : Shapes)
+			{
+				if (BodyInstance.IsShapeBoundToBody(Shape))
+				{
+					FPhysicsInterface::SetUserData(Shape, (void*)ShapeBodySetup->AggGeom.SphereElems[0].GetUserData());
+				}
+			}
+		});
+#endif
+	}
+
 	UpdateCollisionVolume();
 
 	Super::OnRegister();
@@ -162,6 +193,8 @@ FTransform UGLTFInteractionHotspotComponent::GetWorldTransform() const
 {
 	const AActor* Owner = GetOwner();
 
+	check(Owner != nullptr);
+
 	return (Owner != nullptr) ? (Owner->GetTransform() * GetComponentTransform()) : GetComponentTransform();
 }
 
