[GLTFExporter] fix(interaction-hotspot): remove accidentally commited check

commit b9d321b5aa890927f61b81a228e29c5806f0c0a8
Date: Mon, 15 Jun 2020 13:13:52 +0200

    fix(interaction-hotspot): remove accidentally commited check
    (cherry picked from commit 188f18b056d7cb61b6f961ab178a47fd8b555a19​)

#skipci

diff --git a/Engine/Plugins/Enterprise/GLTFExporter/MigratedFromGit/Source/GLTFExporter/Public/Components/GLTFInteractionHotspotComponent.cpp b/Engine/Plugins/Enterprise/GLTFExporter/MigratedFromGit/Source/GLTFExporter/Public/Components/GLTFInteractionHotspotComponent.cpp
--- a/Engine/Plugins/Enterprise/GLTFExporter/MigratedFromGit/Source/GLTFExporter/Public/Components/GLTFInteractionHotspotComponent.cpp
+++ b/Engine/Plugins/Enterprise/GLTFExporter/MigratedFromGit/Source/GLTFExporter/Public/Components/GLTFInteractionHotspotComponent.cpp
@@ -193,8 +193,6 @@ FTransform UGLTFInteractionHotspotComponent::GetWorldTransform() const
 {
 	const AActor* Owner = GetOwner();
 
-	check(Owner != nullptr);
-
 	return (Owner != nullptr) ? (Owner->GetTransform() * GetComponentTransform()) : GetComponentTransform();
 }
 
