[GLTFExporter] refactor(interaction-hotspot): move source files to private folder

commit 038d75f2190cc0cfbe13fbb88aa8a06fd5cffc08
Date: Tue, 16 Jun 2020 10:34:16 +0200

    refactor(interaction-hotspot): move source files to private folder
    (cherry picked from commit b3be10c1369649a6d29caef1068f5084ac2d2526​)

#skipci

diff --git a/Engine/Plugins/Enterprise/GLTFExporter/MigratedFromGit/Source/GLTFExporter/Private/Actors/GLTFInteractionHotspotActor.cpp b/Engine/Plugins/Enterprise/GLTFExporter/MigratedFromGit/Source/GLTFExporter/Private/Actors/GLTFInteractionHotspotActor.cpp
new file mode 100644
--- /dev/null
+++ b/Engine/Plugins/Enterprise/GLTFExporter/MigratedFromGit/Source/GLTFExporter/Private/Actors/GLTFInteractionHotspotActor.cpp
@@ -0,0 +1,51 @@
+// Copyright Epic Games, Inc. All Rights Reserved.
+
+#include "Actors/GLTFInteractionHotspotActor.h"
+
+AGLTFInteractionHotspotActor::AGLTFInteractionHotspotActor(const FObjectInitializer& ObjectInitializer)
+	: Super(ObjectInitializer),
+	DefaultSprite(nullptr),
+	HighlightSprite(nullptr),
+	ToggledSprite(nullptr)
+{
+	// A scene component with a transform in the root
+	SceneComponent = CreateDefaultSubobject<USceneComponent>(TEXT("SceneComponent"));
+	RootComponent = SceneComponent;
+
+	InteractionHotspotComponent = CreateDefaultSubobject<UGLTFInteractionHotspotComponent>(TEXT("InteractionHotspotComponent"));
+	InteractionHotspotComponent->SetupAttachment(SceneComponent);
+	ForwardPropertiesToComponent();
+}
+
+#if WITH_EDITOR
+void AGLTFInteractionHotspotActor::PostEditChangeProperty(FPropertyChangedEvent& PropertyChangedEvent)
+{
+	Super::PostEditChangeProperty(PropertyChangedEvent);
+
+	ForwardPropertiesToComponent();
+}
+#endif // WITH_EDITOR
+
+void AGLTFInteractionHotspotActor::ForwardPropertiesToComponent()
+{
+	if (InteractionHotspotComponent->Animations != Animations)
+	{
+		InteractionHotspotComponent->Animations = Animations;
+	}
+	
+	if (InteractionHotspotComponent->DefaultSprite != DefaultSprite)
+	{
+		InteractionHotspotComponent->DefaultSprite = DefaultSprite;
+		InteractionHotspotComponent->SetSprite(DefaultSprite);
+	}
+	
+	if (InteractionHotspotComponent->HighlightSprite != HighlightSprite)
+	{
+		InteractionHotspotComponent->HighlightSprite = HighlightSprite;
+	}
+	
+	if (InteractionHotspotComponent->ToggledSprite != ToggledSprite)
+	{
+		InteractionHotspotComponent->ToggledSprite = ToggledSprite;
+	}
+}
diff --git a/Engine/Plugins/Enterprise/GLTFExporter/MigratedFromGit/Source/GLTFExporter/Private/Components/GLTFInteractionHotspotComponent.cpp b/Engine/Plugins/Enterprise/GLTFExporter/MigratedFromGit/Source/GLTFExporter/Private/Components/GLTFInteractionHotspotComponent.cpp
new file mode 100644
--- /dev/null
+++ b/Engine/Plugins/Enterprise/GLTFExporter/MigratedFromGit/Source/GLTFExporter/Private/Components/GLTFInteractionHotspotComponent.cpp
@@ -0,0 +1,205 @@
+// Copyright Epic Games, Inc. All Rights Reserved.
+
+#include "Components/GLTFInteractionHotspotComponent.h"
+#include "Animation/SkeletalMeshActor.h"
+#include "Animation/AnimSequence.h"
+#include "Engine/CollisionProfile.h"
+#include "UObject/ConstructorHelpers.h"
+#include "PhysicsEngine/BodySetup.h"
+#include "Animation/AnimSingleNodeInstance.h"
+
+namespace
+{
+	const float UnitSphereRadius = 50.0f;
+} // Anonymous namespace
+
+UGLTFInteractionHotspotComponent::UGLTFInteractionHotspotComponent(const FObjectInitializer& ObjectInitializer)
+	: Super(ObjectInitializer),
+	DefaultSprite(nullptr),
+	HighlightSprite(nullptr),
+	ToggledSprite(nullptr),
+	ShapeBodySetup(nullptr),
+	bToggled(bToggled)
+{
+	// Setup the most minimalistic collision profile for mouse input events
+	SetCollisionEnabled(ECollisionEnabled::QueryAndPhysics);
+	SetCollisionResponseToAllChannels(ECollisionResponse::ECR_Ignore);
+	SetCollisionResponseToChannel(ECollisionChannel::ECC_Visibility, ECollisionResponse::ECR_Block);
+	SetGenerateOverlapEvents(false);
+
+	bHiddenInGame = false;
+
+	OnBeginCursorOver.AddDynamic(this, &UGLTFInteractionHotspotComponent::BeginCursorOver);
+	OnEndCursorOver.AddDynamic(this, &UGLTFInteractionHotspotComponent::EndCursorOver);
+	OnClicked.AddDynamic(this, &UGLTFInteractionHotspotComponent::Clicked);
+}
+
+#if WITH_EDITOR
+void UGLTFInteractionHotspotComponent::PostEditChangeProperty(FPropertyChangedEvent& PropertyChangedEvent)
+{
+	Super::PostEditChangeProperty(PropertyChangedEvent);
+
+	const FProperty* PropertyThatChanged = PropertyChangedEvent.Property;
+
+	if (PropertyThatChanged)
+	{
+		const FString PropertyName = PropertyThatChanged->GetName();
+
+		if (PropertyName == TEXT("DefaultSprite"))
+		{
+			SetSprite(DefaultSprite);
+		}
+	}
+}
+#endif // WITH_EDITOR
+
+void UGLTFInteractionHotspotComponent::BeginPlay()
+{
+	Super::BeginPlay();
+
+	SetSprite(DefaultSprite);
+}
+
+void UGLTFInteractionHotspotComponent::OnRegister()
+{
+	ShapeBodySetup = NewObject<UBodySetup>(this, NAME_None, RF_Transient);
+
+	if (GUObjectArray.IsDisregardForGC(this))
+	{
+		ShapeBodySetup->AddToRoot();
+	}
+
+	ShapeBodySetup->AddToCluster(this);
+
+	if (ShapeBodySetup->HasAnyInternalFlags(EInternalObjectFlags::Async) && GUObjectClusters.GetObjectCluster(ShapeBodySetup))
+	{
+		ShapeBodySetup->ClearInternalFlags(EInternalObjectFlags::Async);
+	}
+
+	ShapeBodySetup->CollisionTraceFlag = CTF_UseSimpleAsComplex;
+	ShapeBodySetup->AggGeom.SphereElems.Add(FKSphereElem(UnitSphereRadius));
+
+	BodyInstance.InitBody(ShapeBodySetup, GetWorldTransform(), this, GetWorld()->GetPhysicsScene());
+
+	if (BodyInstance.IsValidBodyInstance())
+	{
+#if WITH_PHYSX
+		FPhysicsCommand::ExecuteWrite(BodyInstance.GetActorReferenceWithWelding(), [this](const FPhysicsActorHandle& Actor)
+		{
+			TArray<FPhysicsShapeHandle> Shapes;
+			BodyInstance.GetAllShapes_AssumesLocked(Shapes);
+
+			for (FPhysicsShapeHandle& Shape : Shapes)
+			{
+				if (BodyInstance.IsShapeBoundToBody(Shape))
+				{
+					FPhysicsInterface::SetUserData(Shape, (void*)ShapeBodySetup->AggGeom.SphereElems[0].GetUserData());
+				}
+			}
+		});
+#endif
+	}
+
+	UpdateCollisionVolume();
+
+	Super::OnRegister();
+}
+
+void UGLTFInteractionHotspotComponent::SetSprite(class UTexture2D* NewSprite)
+{
+	Super::SetSprite(NewSprite);
+
+	UpdateCollisionVolume();
+}
+
+void UGLTFInteractionHotspotComponent::BeginCursorOver(UPrimitiveComponent* TouchedComponent)
+{
+	if ((ToggledSprite == nullptr || !bToggled) && HighlightSprite != nullptr)
+	{
+		SetSprite(HighlightSprite);
+	}
+}
+
+void UGLTFInteractionHotspotComponent::EndCursorOver(UPrimitiveComponent* TouchedComponent)
+{
+	if (ToggledSprite == nullptr || !bToggled)
+	{
+		SetSprite(DefaultSprite);
+	}
+}
+
+void UGLTFInteractionHotspotComponent::Clicked(UPrimitiveComponent* TouchedComponent, FKey ButtonPressed)
+{
+	if (!bToggled && ToggledSprite != nullptr)
+	{
+		SetSprite(ToggledSprite);
+	}
+	else
+	{
+		SetSprite(DefaultSprite);
+	}
+
+	const bool bReverseAnimation = bToggled;
+
+	for (TArray<FGLTFAnimation>::TConstIterator Animation = Animations.CreateConstIterator(); Animation; ++Animation)
+	{
+		if (Animation->SkeletalMeshActor != nullptr && Animation->AnimationSequence != nullptr)
+		{
+			USkeletalMeshComponent* SkeletalMeshComponent = Animation->SkeletalMeshActor->GetSkeletalMeshComponent();
+			SkeletalMeshComponent->SetAnimationMode(EAnimationMode::Type::AnimationSingleNode);
+			const float AbsolutePlayRate = FMath::Abs(SkeletalMeshComponent->GetPlayRate());
+
+			if (SkeletalMeshComponent->IsPlaying())
+			{
+				UAnimSingleNodeInstance* SingleNodeInstance = SkeletalMeshComponent->GetSingleNodeInstance();
+				
+				if (SingleNodeInstance != nullptr && SingleNodeInstance->GetAnimationAsset() == Animation->AnimationSequence)
+				{
+					// If the same animation is already playing, just reverse the play rate for a smooth transition
+					SkeletalMeshComponent->SetPlayRate(AbsolutePlayRate * -1.0f);
+					
+					continue;
+				}
+			}
+			
+			SkeletalMeshComponent->SetAnimation(Animation->AnimationSequence);
+			SkeletalMeshComponent->SetPlayRate(AbsolutePlayRate * (bReverseAnimation ? -1.0f : 1.0f));
+			SkeletalMeshComponent->SetPosition(bReverseAnimation ? Animation->AnimationSequence->GetPlayLength() : 0.0f);
+			SkeletalMeshComponent->Play(false);
+		}
+	}
+
+	bToggled = !bToggled;
+}
+
+void UGLTFInteractionHotspotComponent::UpdateCollisionVolume()
+{
+	if (ShapeBodySetup != nullptr)
+	{
+		// TODO: Figure out where this difference stems from
+		const float Scaling = 0.15;
+		const float BillboardBoundingRadius = GetBillboardBoundingRadius() * Scaling;
+
+		check(ShapeBodySetup->AggGeom.SphereElems.Num() == 1);
+
+		if (!FMath::IsNearlyEqual(ShapeBodySetup->AggGeom.SphereElems[0].Radius, BillboardBoundingRadius) && UnitSphereRadius != 0.0f)
+		{
+			BodyInstance.UpdateBodyScale(FVector(BillboardBoundingRadius / UnitSphereRadius), true);
+		}
+	}
+}
+
+FTransform UGLTFInteractionHotspotComponent::GetWorldTransform() const
+{
+	const AActor* Owner = GetOwner();
+
+	return (Owner != nullptr) ? (Owner->GetTransform() * GetComponentTransform()) : GetComponentTransform();
+}
+
+float UGLTFInteractionHotspotComponent::GetBillboardBoundingRadius() const
+{
+	const FTransform WorldTransform = GetWorldTransform();
+	const FBoxSphereBounds WorldBounds = CalcBounds(WorldTransform);
+
+	return WorldBounds.SphereRadius;
+}
diff --git a/Engine/Plugins/Enterprise/GLTFExporter/MigratedFromGit/Source/GLTFExporter/Public/Actors/GLTFInteractionHotspotActor.cpp b/Engine/Plugins/Enterprise/GLTFExporter/MigratedFromGit/Source/GLTFExporter/Public/Actors/GLTFInteractionHotspotActor.cpp
deleted file mode 100644
diff --git a/Engine/Plugins/Enterprise/GLTFExporter/MigratedFromGit/Source/GLTFExporter/Public/Components/GLTFInteractionHotspotComponent.cpp b/Engine/Plugins/Enterprise/GLTFExporter/MigratedFromGit/Source/GLTFExporter/Public/Components/GLTFInteractionHotspotComponent.cpp
deleted file mode 100644
