[GLTFExporter] docs(interaction-hotspot): mark unsolved issues with todo comments

commit f4993878fdae3d46beb1d194bb1f0f5ef85d93a1
Date: Tue, 16 Jun 2020 11:12:03 +0200

    docs(interaction-hotspot): mark unsolved issues with todo comments
    (cherry picked from commit 3639ff69afcd1095e39a8800b94a5d483d3500a8​)

#skipci

diff --git a/Engine/Plugins/Enterprise/GLTFExporter/MigratedFromGit/Source/GLTFExporter/Private/Actors/GLTFInteractionHotspotActor.cpp b/Engine/Plugins/Enterprise/GLTFExporter/MigratedFromGit/Source/GLTFExporter/Private/Actors/GLTFInteractionHotspotActor.cpp
--- a/Engine/Plugins/Enterprise/GLTFExporter/MigratedFromGit/Source/GLTFExporter/Private/Actors/GLTFInteractionHotspotActor.cpp
+++ b/Engine/Plugins/Enterprise/GLTFExporter/MigratedFromGit/Source/GLTFExporter/Private/Actors/GLTFInteractionHotspotActor.cpp
@@ -8,7 +8,7 @@ AGLTFInteractionHotspotActor::AGLTFInteractionHotspotActor(const FObjectInitiali
 	HighlightSprite(nullptr),
 	ToggledSprite(nullptr)
 {
-	// A scene component with a transform in the root
+	// TODO: Figure out how to make hotspot component the root without affecting its collision volume
 	SceneComponent = CreateDefaultSubobject<USceneComponent>(TEXT("SceneComponent"));
 	RootComponent = SceneComponent;
 
@@ -28,6 +28,7 @@ void AGLTFInteractionHotspotActor::PostEditChangeProperty(FPropertyChangedEvent&
 
 void AGLTFInteractionHotspotActor::ForwardPropertiesToComponent()
 {
+	// TODO: Forward full transform to component
 	if (InteractionHotspotComponent->Animations != Animations)
 	{
 		InteractionHotspotComponent->Animations = Animations;
diff --git a/Engine/Plugins/Enterprise/GLTFExporter/MigratedFromGit/Source/GLTFExporter/Private/Components/GLTFInteractionHotspotComponent.cpp b/Engine/Plugins/Enterprise/GLTFExporter/MigratedFromGit/Source/GLTFExporter/Private/Components/GLTFInteractionHotspotComponent.cpp
--- a/Engine/Plugins/Enterprise/GLTFExporter/MigratedFromGit/Source/GLTFExporter/Private/Components/GLTFInteractionHotspotComponent.cpp
+++ b/Engine/Plugins/Enterprise/GLTFExporter/MigratedFromGit/Source/GLTFExporter/Private/Components/GLTFInteractionHotspotComponent.cpp
@@ -176,7 +176,7 @@ void UGLTFInteractionHotspotComponent::UpdateCollisionVolume()
 {
 	if (ShapeBodySetup != nullptr)
 	{
-		// TODO: Figure out where this difference stems from
+		// TODO: Figure out why the bounding radius doesn't match the size of the billboard
 		const float Scaling = 0.15;
 		const float BillboardBoundingRadius = GetBillboardBoundingRadius() * Scaling;
 
