[GLTFExporter] refactor(interaction-hotspot): Expose interaction hotspot properties in actor directly via component rather than manual forwarding

commit 70725bc65417c7826c197aa1f5f59f5730cfdba3
Date: Wed, 17 Jun 2020 09:10:50 +0200

    refactor(interaction-hotspot): Expose interaction hotspot properties in actor directly via component rather than manual forwarding
    (cherry picked from commit e0bbe7f2a501497ae9cc3d3e9b57ab923a650a5d​)

#skipci

diff --git a/Engine/Plugins/Enterprise/GLTFExporter/MigratedFromGit/Source/GLTFExporter/Private/Actors/GLTFInteractionHotspotActor.cpp b/Engine/Plugins/Enterprise/GLTFExporter/MigratedFromGit/Source/GLTFExporter/Private/Actors/GLTFInteractionHotspotActor.cpp
--- a/Engine/Plugins/Enterprise/GLTFExporter/MigratedFromGit/Source/GLTFExporter/Private/Actors/GLTFInteractionHotspotActor.cpp
+++ b/Engine/Plugins/Enterprise/GLTFExporter/MigratedFromGit/Source/GLTFExporter/Private/Actors/GLTFInteractionHotspotActor.cpp
@@ -3,50 +3,8 @@
 #include "Actors/GLTFInteractionHotspotActor.h"
 
 AGLTFInteractionHotspotActor::AGLTFInteractionHotspotActor(const FObjectInitializer& ObjectInitializer)
-	: Super(ObjectInitializer),
-	DefaultSprite(nullptr),
-	HighlightSprite(nullptr),
-	ToggledSprite(nullptr)
+	: Super(ObjectInitializer)
 {
-	// TODO: Figure out how to make hotspot component the root without affecting its collision volume
-	SceneComponent = CreateDefaultSubobject<USceneComponent>(TEXT("SceneComponent"));
-	RootComponent = SceneComponent;
-
 	InteractionHotspotComponent = CreateDefaultSubobject<UGLTFInteractionHotspotComponent>(TEXT("InteractionHotspotComponent"));
-	InteractionHotspotComponent->SetupAttachment(SceneComponent);
-	ForwardPropertiesToComponent();
-}
-
-#if WITH_EDITOR
-void AGLTFInteractionHotspotActor::PostEditChangeProperty(FPropertyChangedEvent& PropertyChangedEvent)
-{
-	Super::PostEditChangeProperty(PropertyChangedEvent);
-
-	ForwardPropertiesToComponent();
-}
-#endif // WITH_EDITOR
-
-void AGLTFInteractionHotspotActor::ForwardPropertiesToComponent()
-{
-	// TODO: Forward full transform to component
-	if (InteractionHotspotComponent->Animations != Animations)
-	{
-		InteractionHotspotComponent->Animations = Animations;
-	}
-	
-	if (InteractionHotspotComponent->DefaultSprite != DefaultSprite)
-	{
-		InteractionHotspotComponent->DefaultSprite = DefaultSprite;
-		InteractionHotspotComponent->SetSprite(DefaultSprite);
-	}
-	
-	if (InteractionHotspotComponent->HighlightSprite != HighlightSprite)
-	{
-		InteractionHotspotComponent->HighlightSprite = HighlightSprite;
-	}
-	
-	if (InteractionHotspotComponent->ToggledSprite != ToggledSprite)
-	{
-		InteractionHotspotComponent->ToggledSprite = ToggledSprite;
-	}
+	RootComponent = InteractionHotspotComponent;
 }
diff --git a/Engine/Plugins/Enterprise/GLTFExporter/MigratedFromGit/Source/GLTFExporter/Private/Components/GLTFInteractionHotspotComponent.cpp b/Engine/Plugins/Enterprise/GLTFExporter/MigratedFromGit/Source/GLTFExporter/Private/Components/GLTFInteractionHotspotComponent.cpp
--- a/Engine/Plugins/Enterprise/GLTFExporter/MigratedFromGit/Source/GLTFExporter/Private/Components/GLTFInteractionHotspotComponent.cpp
+++ b/Engine/Plugins/Enterprise/GLTFExporter/MigratedFromGit/Source/GLTFExporter/Private/Components/GLTFInteractionHotspotComponent.cpp
@@ -78,8 +78,13 @@ void UGLTFInteractionHotspotComponent::OnRegister()
 
 	ShapeBodySetup->CollisionTraceFlag = CTF_UseSimpleAsComplex;
 	ShapeBodySetup->AggGeom.SphereElems.Add(FKSphereElem(UnitSphereRadius));
+	
+	Super::OnRegister();
+}
 
-	BodyInstance.InitBody(ShapeBodySetup, GetWorldTransform(), this, GetWorld()->GetPhysicsScene());
+void UGLTFInteractionHotspotComponent::OnCreatePhysicsState()
+{
+	Super::OnCreatePhysicsState();
 
 	if (BodyInstance.IsValidBodyInstance())
 	{
@@ -101,8 +106,6 @@ void UGLTFInteractionHotspotComponent::OnRegister()
 	}
 
 	UpdateCollisionVolume();
-
-	Super::OnRegister();
 }
 
 void UGLTFInteractionHotspotComponent::SetSprite(class UTexture2D* NewSprite)
@@ -177,7 +180,7 @@ void UGLTFInteractionHotspotComponent::UpdateCollisionVolume()
 	if (ShapeBodySetup != nullptr)
 	{
 		// TODO: Figure out why the bounding radius doesn't match the size of the billboard
-		const float Scaling = 0.15;
+		const float Scaling = 0.15f;
 		const float BillboardBoundingRadius = GetBillboardBoundingRadius() * Scaling;
 
 		check(ShapeBodySetup->AggGeom.SphereElems.Num() == 1);
@@ -189,16 +192,9 @@ void UGLTFInteractionHotspotComponent::UpdateCollisionVolume()
 	}
 }
 
-FTransform UGLTFInteractionHotspotComponent::GetWorldTransform() const
-{
-	const AActor* Owner = GetOwner();
-
-	return (Owner != nullptr) ? (Owner->GetTransform() * GetComponentTransform()) : GetComponentTransform();
-}
-
 float UGLTFInteractionHotspotComponent::GetBillboardBoundingRadius() const
 {
-	const FTransform WorldTransform = GetWorldTransform();
+	const FTransform WorldTransform = GetComponentTransform();
 	const FBoxSphereBounds WorldBounds = CalcBounds(WorldTransform);
 
 	return WorldBounds.SphereRadius;
diff --git a/Engine/Plugins/Enterprise/GLTFExporter/MigratedFromGit/Source/GLTFExporter/Public/Actors/GLTFInteractionHotspotActor.h b/Engine/Plugins/Enterprise/GLTFExporter/MigratedFromGit/Source/GLTFExporter/Public/Actors/GLTFInteractionHotspotActor.h
--- a/Engine/Plugins/Enterprise/GLTFExporter/MigratedFromGit/Source/GLTFExporter/Public/Actors/GLTFInteractionHotspotActor.h
+++ b/Engine/Plugins/Enterprise/GLTFExporter/MigratedFromGit/Source/GLTFExporter/Public/Actors/GLTFInteractionHotspotActor.h
@@ -9,45 +9,16 @@
 /**
  * Actor wrapper for the GLTF hotspot component. Appears as a billboard and allows playback of skeletal animations when cursor input is enabled.
  */
-UCLASS(DisplayName = "GLTF Interaction Hotspot Actor")
+UCLASS(DisplayName = "GLTF Interaction Hotspot Actor", HideCategories = (Sprite))
 class GLTFEXPORTER_API AGLTFInteractionHotspotActor : public AActor
 {
 	GENERATED_BODY()
 	//~ Begin UObject Interface.
 public:
 	AGLTFInteractionHotspotActor(const FObjectInitializer& ObjectInitializer);
-
-private:
-#if WITH_EDITOR
-	virtual void PostEditChangeProperty(FPropertyChangedEvent& PropertyChangedEvent) override;
-#endif // WITH_EDITOR
 	//~ End UObject Interface.
 
 private:
-	void ForwardPropertiesToComponent();
-
-public:
-	/** List of skeletal meshes and animations to be played when the hotspot is interacted with */
-	UPROPERTY(EditAnywhere, BlueprintReadWrite, Category = "GLTF Interaction Hotspot")
-	TArray<FGLTFAnimation> Animations;
-
-	/** The billboard image that will be shown when the hotspot is in an inactive state */
-	UPROPERTY(EditAnywhere, BlueprintReadWrite, Category = "GLTF Interaction Hotspot")
-	UTexture2D* DefaultSprite;
-
-	/** The optional billboard image that will be shown when a cursor enters the hotspot */
-	UPROPERTY(EditAnywhere, BlueprintReadWrite, Category = "GLTF Interaction Hotspot")
-	UTexture2D* HighlightSprite;
-
-	/** The optional billboard image that will be shown when the hotspot is toggled by a click */
-	UPROPERTY(EditAnywhere, BlueprintReadWrite, Category = "GLTF Interaction Hotspot")
-	UTexture2D* ToggledSprite;
-
-private:
-	UPROPERTY()
-	USceneComponent* SceneComponent;
-
-public:
-	UPROPERTY(BlueprintReadOnly, Category = "GLTF Interaction Hotspot")
+	UPROPERTY(Category = "GLTF Interaction Hotspot", VisibleAnywhere, BlueprintReadOnly, meta = (ExposeFunctionCategories = "GLTF Interaction Hotspot", AllowPrivateAccess = "true"))
 	UGLTFInteractionHotspotComponent* InteractionHotspotComponent;
 };
diff --git a/Engine/Plugins/Enterprise/GLTFExporter/MigratedFromGit/Source/GLTFExporter/Public/Components/GLTFInteractionHotspotComponent.h b/Engine/Plugins/Enterprise/GLTFExporter/MigratedFromGit/Source/GLTFExporter/Public/Components/GLTFInteractionHotspotComponent.h
--- a/Engine/Plugins/Enterprise/GLTFExporter/MigratedFromGit/Source/GLTFExporter/Public/Components/GLTFInteractionHotspotComponent.h
+++ b/Engine/Plugins/Enterprise/GLTFExporter/MigratedFromGit/Source/GLTFExporter/Public/Components/GLTFInteractionHotspotComponent.h
@@ -52,6 +52,7 @@ public:
 
 protected:
 	virtual void OnRegister() override;
+	virtual void OnCreatePhysicsState() override;
 	//~ End UActorComponent Interface.
 
 	//~ Begin UBillboardComponent Interface.
@@ -59,6 +60,10 @@ public:
 	virtual void SetSprite(class UTexture2D* NewSprite) override;
 	//~ End UBillboardComponent Interface.
 
+	//~ Begin UPrimitiveComponent Interface.
+	virtual UBodySetup* GetBodySetup() { return ShapeBodySetup; }
+	//~ End UBillboardComponent Interface.
+
 private:
 	UFUNCTION()
 	void BeginCursorOver(UPrimitiveComponent* TouchedComponent);
@@ -71,14 +76,13 @@ private:
 
 	void UpdateCollisionVolume();
 	float GetBillboardBoundingRadius() const;
-	FTransform GetWorldTransform() const;
 
 public:
 	/** List of skeletal meshes and animations to be played when the hotspot is interacted with */
 	UPROPERTY(EditAnywhere, BlueprintReadWrite, Category = "GLTF Interaction Hotspot")
 	TArray<FGLTFAnimation> Animations;
 
-	/** The billboard image that will be shown when the hotspot is in an inactive state */
+	/** The billboard image that will be shown when the hotspot is in an inactive state or one without a specified sprite */
 	UPROPERTY(EditAnywhere, BlueprintReadWrite, Category = "GLTF Interaction Hotspot")
 	UTexture2D* DefaultSprite;
 
