[GLTFExporter] docs(interaction-hotspot): clean up class interfaces

commit 1edb3677463939a37edb7f3d87c621ecef903e51
Date: Wed, 17 Jun 2020 10:20:32 +0200

    docs(interaction-hotspot): clean up class interfaces
    (cherry picked from commit 245da7d6f1fc006e18c42410838da28b455c0499​)

#skipci

diff --git a/Engine/Plugins/Enterprise/GLTFExporter/MigratedFromGit/Source/GLTFExporter/Private/Components/GLTFInteractionHotspotComponent.cpp b/Engine/Plugins/Enterprise/GLTFExporter/MigratedFromGit/Source/GLTFExporter/Private/Components/GLTFInteractionHotspotComponent.cpp
--- a/Engine/Plugins/Enterprise/GLTFExporter/MigratedFromGit/Source/GLTFExporter/Private/Components/GLTFInteractionHotspotComponent.cpp
+++ b/Engine/Plugins/Enterprise/GLTFExporter/MigratedFromGit/Source/GLTFExporter/Private/Components/GLTFInteractionHotspotComponent.cpp
@@ -108,6 +108,11 @@ void UGLTFInteractionHotspotComponent::OnCreatePhysicsState()
 	UpdateCollisionVolume();
 }
 
+UBodySetup* UGLTFInteractionHotspotComponent::GetBodySetup()
+{
+	return ShapeBodySetup;
+}
+
 void UGLTFInteractionHotspotComponent::SetSprite(class UTexture2D* NewSprite)
 {
 	Super::SetSprite(NewSprite);
diff --git a/Engine/Plugins/Enterprise/GLTFExporter/MigratedFromGit/Source/GLTFExporter/Public/Actors/GLTFInteractionHotspotActor.h b/Engine/Plugins/Enterprise/GLTFExporter/MigratedFromGit/Source/GLTFExporter/Public/Actors/GLTFInteractionHotspotActor.h
--- a/Engine/Plugins/Enterprise/GLTFExporter/MigratedFromGit/Source/GLTFExporter/Public/Actors/GLTFInteractionHotspotActor.h
+++ b/Engine/Plugins/Enterprise/GLTFExporter/MigratedFromGit/Source/GLTFExporter/Public/Actors/GLTFInteractionHotspotActor.h
@@ -13,10 +13,10 @@ UCLASS(DisplayName = "GLTF Interaction Hotspot Actor", HideCategories = (Sprite)
 class GLTFEXPORTER_API AGLTFInteractionHotspotActor : public AActor
 {
 	GENERATED_BODY()
-	//~ Begin UObject Interface.
+	//~ Begin UObject Interface
 public:
 	AGLTFInteractionHotspotActor(const FObjectInitializer& ObjectInitializer);
-	//~ End UObject Interface.
+	//~ End UObject Interface
 
 private:
 	UPROPERTY(Category = "GLTF Interaction Hotspot", VisibleAnywhere, BlueprintReadOnly, meta = (ExposeFunctionCategories = "GLTF Interaction Hotspot", AllowPrivateAccess = "true"))
diff --git a/Engine/Plugins/Enterprise/GLTFExporter/MigratedFromGit/Source/GLTFExporter/Public/Components/GLTFInteractionHotspotComponent.h b/Engine/Plugins/Enterprise/GLTFExporter/MigratedFromGit/Source/GLTFExporter/Public/Components/GLTFInteractionHotspotComponent.h
--- a/Engine/Plugins/Enterprise/GLTFExporter/MigratedFromGit/Source/GLTFExporter/Public/Components/GLTFInteractionHotspotComponent.h
+++ b/Engine/Plugins/Enterprise/GLTFExporter/MigratedFromGit/Source/GLTFExporter/Public/Components/GLTFInteractionHotspotComponent.h
@@ -41,28 +41,28 @@ class GLTFEXPORTER_API UGLTFInteractionHotspotComponent : public UBillboardCompo
 {
 	GENERATED_UCLASS_BODY()
 public:
-	//~ Begin UObject Interface.
+	//~ Begin UObject Interface
 #if WITH_EDITOR
 	virtual void PostEditChangeProperty(FPropertyChangedEvent& PropertyChangedEvent) override;
 #endif // WITH_EDITOR
-	//~ End UObject Interface.
+	//~ End UObject Interface
 
-	//~ Begin UActorComponent Interface.
+	//~ Begin UActorComponent Interface
 	virtual void BeginPlay() override;
 
 protected:
 	virtual void OnRegister() override;
 	virtual void OnCreatePhysicsState() override;
-	//~ End UActorComponent Interface.
+	//~ End UActorComponent Interface
 
-	//~ Begin UBillboardComponent Interface.
+	//~ Begin UPrimitiveComponent Interface
 public:
-	virtual void SetSprite(class UTexture2D* NewSprite) override;
-	//~ End UBillboardComponent Interface.
+	virtual UBodySetup* GetBodySetup() override;
+	//~ End UPrimitiveComponent Interface
 
-	//~ Begin UPrimitiveComponent Interface.
-	virtual UBodySetup* GetBodySetup() { return ShapeBodySetup; }
-	//~ End UBillboardComponent Interface.
+	//~ Begin UBillboardComponent Interface
+	virtual void SetSprite(class UTexture2D* NewSprite) override;
+	//~ End UBillboardComponent Interface
 
 private:
 	UFUNCTION()
@@ -78,19 +78,19 @@ private:
 	float GetBillboardBoundingRadius() const;
 
 public:
-	/** List of skeletal meshes and animations to be played when the hotspot is interacted with */
+	/** List of skeletal meshes and animations to be played when the hotspot is interacted with. */
 	UPROPERTY(EditAnywhere, BlueprintReadWrite, Category = "GLTF Interaction Hotspot")
 	TArray<FGLTFAnimation> Animations;
 
-	/** The billboard image that will be shown when the hotspot is in an inactive state or one without a specified sprite */
+	/** The billboard image that will be shown when the hotspot is in an inactive state or one without a specified sprite. */
 	UPROPERTY(EditAnywhere, BlueprintReadWrite, Category = "GLTF Interaction Hotspot")
 	UTexture2D* DefaultSprite;
 
-	/** The optional billboard image that will be shown when a cursor enters the hotspot */
+	/** The optional billboard image that will be shown when a cursor enters the hotspot. */
 	UPROPERTY(EditAnywhere, BlueprintReadWrite, Category = "GLTF Interaction Hotspot")
 	UTexture2D* HighlightSprite;
 
-	/** The optional billboard image that will be shown when the hotspot is toggled by a click */
+	/** The optional billboard image that will be shown when the hotspot is toggled by a click. */
 	UPROPERTY(EditAnywhere, BlueprintReadWrite, Category = "GLTF Interaction Hotspot")
 	UTexture2D* ToggledSprite;
 
