[GLTFExporter] feat(interaction-hotspot): add fourth interaction state

commit afbee7cb5cd3a0ef8d82c4a01028b61d17c6109d
Date: Wed, 17 Jun 2020 16:39:49 +0200

    feat(interaction-hotspot): add fourth interaction state
    (cherry picked from commit fc23c13997c85b057008c2f5414c6e26a3c5f279​)

#skipci

diff --git a/Engine/Plugins/Enterprise/GLTFExporter/MigratedFromGit/Source/GLTFExporter/Private/Components/GLTFInteractionHotspotComponent.cpp b/Engine/Plugins/Enterprise/GLTFExporter/MigratedFromGit/Source/GLTFExporter/Private/Components/GLTFInteractionHotspotComponent.cpp
--- a/Engine/Plugins/Enterprise/GLTFExporter/MigratedFromGit/Source/GLTFExporter/Private/Components/GLTFInteractionHotspotComponent.cpp
+++ b/Engine/Plugins/Enterprise/GLTFExporter/MigratedFromGit/Source/GLTFExporter/Private/Components/GLTFInteractionHotspotComponent.cpp
@@ -15,9 +15,10 @@ namespace
 
 UGLTFInteractionHotspotComponent::UGLTFInteractionHotspotComponent(const FObjectInitializer& ObjectInitializer)
 	: Super(ObjectInitializer),
-	DefaultSprite(nullptr),
-	HighlightSprite(nullptr),
-	ToggledSprite(nullptr),
+	Image(nullptr),
+	HoveredImage(nullptr),
+	ToggledImage(nullptr),
+	ToggledHoveredImage(nullptr),
 	ShapeBodySetup(nullptr),
 	bToggled(bToggled)
 {
@@ -45,9 +46,9 @@ void UGLTFInteractionHotspotComponent::PostEditChangeProperty(FPropertyChangedEv
 	{
 		const FString PropertyName = PropertyThatChanged->GetName();
 
-		if (PropertyName == TEXT("DefaultSprite"))
+		if (PropertyName == TEXT("Image"))
 		{
-			SetSprite(DefaultSprite);
+			SetSprite(Image);
 		}
 	}
 }
@@ -57,7 +58,7 @@ void UGLTFInteractionHotspotComponent::BeginPlay()
 {
 	Super::BeginPlay();
 
-	SetSprite(DefaultSprite);
+	SetSprite(Image);
 }
 
 void UGLTFInteractionHotspotComponent::OnRegister()
@@ -122,31 +123,16 @@ void UGLTFInteractionHotspotComponent::SetSprite(class UTexture2D* NewSprite)
 
 void UGLTFInteractionHotspotComponent::BeginCursorOver(UPrimitiveComponent* TouchedComponent)
 {
-	if ((ToggledSprite == nullptr || !bToggled) && HighlightSprite != nullptr)
-	{
-		SetSprite(HighlightSprite);
-	}
+	SetSprite(GetActiveImage(true));
 }
 
 void UGLTFInteractionHotspotComponent::EndCursorOver(UPrimitiveComponent* TouchedComponent)
 {
-	if (ToggledSprite == nullptr || !bToggled)
-	{
-		SetSprite(DefaultSprite);
-	}
+	SetSprite(GetActiveImage(false));
 }
 
 void UGLTFInteractionHotspotComponent::Clicked(UPrimitiveComponent* TouchedComponent, FKey ButtonPressed)
 {
-	if (!bToggled && ToggledSprite != nullptr)
-	{
-		SetSprite(ToggledSprite);
-	}
-	else
-	{
-		SetSprite(DefaultSprite);
-	}
-
 	const bool bReverseAnimation = bToggled;
 
 	for (TArray<FGLTFAnimation>::TConstIterator Animation = Animations.CreateConstIterator(); Animation; ++Animation)
@@ -178,6 +164,8 @@ void UGLTFInteractionHotspotComponent::Clicked(UPrimitiveComponent* TouchedCompo
 	}
 
 	bToggled = !bToggled;
+
+	SetSprite(GetActiveImage(true));
 }
 
 void UGLTFInteractionHotspotComponent::UpdateCollisionVolume()
@@ -204,3 +192,40 @@ float UGLTFInteractionHotspotComponent::GetBillboardBoundingRadius() const
 
 	return WorldBounds.SphereRadius;
 }
+
+UTexture2D* UGLTFInteractionHotspotComponent::GetActiveImage(bool bCursorOver) const
+{
+	// A list of candidates ordered by descending priority
+	TArray<UTexture2D*> ImageCandidates;
+
+	if (bToggled)
+	{
+		if (bCursorOver)
+		{
+			ImageCandidates.Add(ToggledHoveredImage);
+		}
+
+		ImageCandidates.Add(ToggledImage);
+	}
+	else
+	{
+		if (bCursorOver)
+		{
+			ImageCandidates.Add(HoveredImage);
+		}
+
+		ImageCandidates.Add(Image);
+	}
+
+	ImageCandidates.Add(Image);
+
+	for (UTexture2D* ImageCandidate : ImageCandidates)
+	{
+		if (ImageCandidate != nullptr)
+		{
+			return ImageCandidate;
+		}
+	}
+
+	return nullptr;
+}
diff --git a/Engine/Plugins/Enterprise/GLTFExporter/MigratedFromGit/Source/GLTFExporter/Public/Components/GLTFInteractionHotspotComponent.h b/Engine/Plugins/Enterprise/GLTFExporter/MigratedFromGit/Source/GLTFExporter/Public/Components/GLTFInteractionHotspotComponent.h
--- a/Engine/Plugins/Enterprise/GLTFExporter/MigratedFromGit/Source/GLTFExporter/Public/Components/GLTFInteractionHotspotComponent.h
+++ b/Engine/Plugins/Enterprise/GLTFExporter/MigratedFromGit/Source/GLTFExporter/Public/Components/GLTFInteractionHotspotComponent.h
@@ -76,23 +76,28 @@ private:
 
 	void UpdateCollisionVolume();
 	float GetBillboardBoundingRadius() const;
+	UTexture2D* GetActiveImage(bool bCursorOver) const;
 
 public:
 	/** List of skeletal meshes and animations to be played when the hotspot is interacted with. */
 	UPROPERTY(EditAnywhere, BlueprintReadWrite, Category = "GLTF Interaction Hotspot")
 	TArray<FGLTFAnimation> Animations;
 
-	/** The billboard image that will be shown when the hotspot is in an inactive state or one without a specified sprite. */
+	/** The billboard image that will be shown when the hotspot is in an idle state or one without a specified image. */
 	UPROPERTY(EditAnywhere, BlueprintReadWrite, Category = "GLTF Interaction Hotspot")
-	UTexture2D* DefaultSprite;
+	UTexture2D* Image;
 
 	/** The optional billboard image that will be shown when a cursor enters the hotspot. */
 	UPROPERTY(EditAnywhere, BlueprintReadWrite, Category = "GLTF Interaction Hotspot")
-	UTexture2D* HighlightSprite;
+	UTexture2D* HoveredImage;
 
 	/** The optional billboard image that will be shown when the hotspot is toggled by a click. */
 	UPROPERTY(EditAnywhere, BlueprintReadWrite, Category = "GLTF Interaction Hotspot")
-	UTexture2D* ToggledSprite;
+	UTexture2D* ToggledImage;
+
+	/** The optional billboard image that will be shown when the hotspot is toggled by a click and a cursor enters it. */
+	UPROPERTY(EditAnywhere, BlueprintReadWrite, Category = "GLTF Interaction Hotspot")
+	UTexture2D* ToggledHoveredImage;
 
 private:
 	UPROPERTY(transient, duplicatetransient)
