[GLTFExporter] feat(interaction-hotspot): change hotspot animations from list to single target

commit 2df196a26a615cbb3d760bfdc5cd606d270662f7
Date: Mon, 22 Jun 2020 15:16:19 +0200

    feat(interaction-hotspot): change hotspot animations from list to single target
    (cherry picked from commit f7eeb7760ac05eefecd3e564d22a2ba495276a91​)

#skipci

diff --git a/Engine/Plugins/Enterprise/GLTFExporter/MigratedFromGit/Source/GLTFExporter/Private/Components/GLTFInteractionHotspotComponent.cpp b/Engine/Plugins/Enterprise/GLTFExporter/MigratedFromGit/Source/GLTFExporter/Private/Components/GLTFInteractionHotspotComponent.cpp
--- a/Engine/Plugins/Enterprise/GLTFExporter/MigratedFromGit/Source/GLTFExporter/Private/Components/GLTFInteractionHotspotComponent.cpp
+++ b/Engine/Plugins/Enterprise/GLTFExporter/MigratedFromGit/Source/GLTFExporter/Private/Components/GLTFInteractionHotspotComponent.cpp
@@ -13,6 +13,8 @@ namespace
 	const float UnitSphereRadius = 50.0f;
 } // Anonymous namespace
 
+DEFINE_LOG_CATEGORY_STATIC(LogEditorGLTFInteractionHotspot, Log, All);
+
 UGLTFInteractionHotspotComponent::UGLTFInteractionHotspotComponent(const FObjectInitializer& ObjectInitializer)
 	: Super(ObjectInitializer),
 	Image(nullptr),
@@ -50,6 +52,20 @@ void UGLTFInteractionHotspotComponent::PostEditChangeProperty(FPropertyChangedEv
 		{
 			SetSprite(Image);
 		}
+		else if (PropertyName == TEXT("SkeletalMeshActor"))
+		{
+			if (AnimationSequence != nullptr && AnimationSequence->GetSkeleton() != SkeletalMeshActor->GetSkeletalMeshComponent()->SkeletalMesh->Skeleton)
+			{
+				UE_LOG(LogEditorGLTFInteractionHotspot, Warning, TEXT("The skeleton of this actor is not compatible with the previously assigned animation sequence"));
+			}
+		}
+		else if (PropertyName == TEXT("AnimationSequence"))
+		{
+			if (SkeletalMeshActor != nullptr && SkeletalMeshActor->GetSkeletalMeshComponent()->SkeletalMesh->Skeleton != AnimationSequence->GetSkeleton())
+			{
+				UE_LOG(LogEditorGLTFInteractionHotspot, Warning, TEXT("This animation sequence is not compatible with the skeleton of the previously assigned actor"));
+			}
+		}
 	}
 }
 #endif // WITH_EDITOR
@@ -135,30 +151,23 @@ void UGLTFInteractionHotspotComponent::Clicked(UPrimitiveComponent* TouchedCompo
 {
 	const bool bReverseAnimation = bToggled;
 
-	for (TArray<FGLTFAnimation>::TConstIterator Animation = Animations.CreateConstIterator(); Animation; ++Animation)
+	if (SkeletalMeshActor != nullptr && AnimationSequence != nullptr)
 	{
-		if (Animation->SkeletalMeshActor != nullptr && Animation->AnimationSequence != nullptr)
-		{
-			USkeletalMeshComponent* SkeletalMeshComponent = Animation->SkeletalMeshActor->GetSkeletalMeshComponent();
-			SkeletalMeshComponent->SetAnimationMode(EAnimationMode::Type::AnimationSingleNode);
-			const float AbsolutePlayRate = FMath::Abs(SkeletalMeshComponent->GetPlayRate());
+		USkeletalMeshComponent* SkeletalMeshComponent = SkeletalMeshActor->GetSkeletalMeshComponent();
+		SkeletalMeshComponent->SetAnimationMode(EAnimationMode::Type::AnimationSingleNode);
+		const float AbsolutePlayRate = FMath::Abs(SkeletalMeshComponent->GetPlayRate());
+		const UAnimSingleNodeInstance* SingleNodeInstance = SkeletalMeshComponent->GetSingleNodeInstance();
 
-			if (SkeletalMeshComponent->IsPlaying())
-			{
-				UAnimSingleNodeInstance* SingleNodeInstance = SkeletalMeshComponent->GetSingleNodeInstance();
-				
-				if (SingleNodeInstance != nullptr && SingleNodeInstance->GetAnimationAsset() == Animation->AnimationSequence)
-				{
-					// If the same animation is already playing, just reverse the play rate for a smooth transition
-					SkeletalMeshComponent->SetPlayRate(AbsolutePlayRate * -1.0f);
-					
-					continue;
-				}
-			}
-			
-			SkeletalMeshComponent->SetAnimation(Animation->AnimationSequence);
+		if (SkeletalMeshComponent->IsPlaying() && SingleNodeInstance != nullptr && SingleNodeInstance->GetAnimationAsset() == AnimationSequence)
+		{
+			// If the same animation is already playing, just reverse the play rate for a smooth transition
+			SkeletalMeshComponent->SetPlayRate(AbsolutePlayRate * -1.0f);
+		}
+		else
+		{
+			SkeletalMeshComponent->SetAnimation(AnimationSequence);
 			SkeletalMeshComponent->SetPlayRate(AbsolutePlayRate * (bReverseAnimation ? -1.0f : 1.0f));
-			SkeletalMeshComponent->SetPosition(bReverseAnimation ? Animation->AnimationSequence->GetPlayLength() : 0.0f);
+			SkeletalMeshComponent->SetPosition(bReverseAnimation ? AnimationSequence->GetPlayLength() : 0.0f);
 			SkeletalMeshComponent->Play(false);
 		}
 	}
diff --git a/Engine/Plugins/Enterprise/GLTFExporter/MigratedFromGit/Source/GLTFExporter/Public/Components/GLTFInteractionHotspotComponent.h b/Engine/Plugins/Enterprise/GLTFExporter/MigratedFromGit/Source/GLTFExporter/Public/Components/GLTFInteractionHotspotComponent.h
--- a/Engine/Plugins/Enterprise/GLTFExporter/MigratedFromGit/Source/GLTFExporter/Public/Components/GLTFInteractionHotspotComponent.h
+++ b/Engine/Plugins/Enterprise/GLTFExporter/MigratedFromGit/Source/GLTFExporter/Public/Components/GLTFInteractionHotspotComponent.h
@@ -10,29 +10,6 @@ class UAnimSequence;
 class UBodySetup;
 class UTexture2D;
 
-USTRUCT(BlueprintType)
-struct FGLTFAnimation
-{
-	GENERATED_USTRUCT_BODY()
-
-	UPROPERTY(EditAnywhere, BlueprintReadWrite)
-	ASkeletalMeshActor* SkeletalMeshActor;
-
-	UPROPERTY(EditAnywhere, BlueprintReadWrite)
-	UAnimSequence* AnimationSequence;
-
-	bool operator == (const FGLTFAnimation& OtherAnimation) const
-	{
-		return SkeletalMeshActor == OtherAnimation.SkeletalMeshActor &&
-			AnimationSequence == OtherAnimation.AnimationSequence;
-	}
-
-	bool operator != (const FGLTFAnimation& OtherAnimation) const
-	{
-		return !(*this == OtherAnimation);
-	}
-};
-
 /**
  * A component to set up hotspots which appear as billboards and allow playback of skeletal animations when cursor input is enabled.
  */
@@ -79,11 +56,15 @@ private:
 	UTexture2D* GetActiveImage(bool bCursorOver) const;
 
 public:
-	/** List of skeletal meshes and animations to be played when the hotspot is interacted with. */
+	/* The skeletal mesh actor that will be animated when the hotspot is clicked. */
+	UPROPERTY(EditAnywhere, BlueprintReadWrite, Category = "GLTF Interaction Hotspot")
+	ASkeletalMeshActor* SkeletalMeshActor;
+
+	/* The animation that will be played on the skeletal mesh actor. Must be compatible with its skeletal mesh asset. */
 	UPROPERTY(EditAnywhere, BlueprintReadWrite, Category = "GLTF Interaction Hotspot")
-	TArray<FGLTFAnimation> Animations;
+	UAnimSequence* AnimationSequence;
 
-	/** The billboard image that will be shown when the hotspot is in an idle state or one without a specified image. */
+	/* The billboard image that will be shown when the hotspot is in an inactive state or one without a specified image. */
 	UPROPERTY(EditAnywhere, BlueprintReadWrite, Category = "GLTF Interaction Hotspot")
 	UTexture2D* Image;
 
