[GLTFExporter] refactor(interaction-hotspot): move code into separate runtime-compatible module

commit 8750730c82c3ce537e9c3b7ce85ed84f049d3f69
Date: Thu, 25 Jun 2020 14:39:14 +0200

    refactor(interaction-hotspot): move code into separate runtime-compatible module
    (cherry picked from commit 4e92678e79d8902938b647a7b23d5e7bf2c49308​)
    # Conflicts from cherry-pick:
    #	Plugins/GLTFExporter/GLTFExporter.uplugin

#skipci

diff --git a/Engine/Plugins/Enterprise/GLTFExporter/MigratedFromGit/GLTFExporter.uplugin b/Engine/Plugins/Enterprise/GLTFExporter/MigratedFromGit/GLTFExporter.uplugin
--- a/Engine/Plugins/Enterprise/GLTFExporter/MigratedFromGit/GLTFExporter.uplugin
+++ b/Engine/Plugins/Enterprise/GLTFExporter/MigratedFromGit/GLTFExporter.uplugin
@@ -19,6 +19,10 @@
 			"Name": "GLTFExporter",
 			"Type": "Editor",
 			"LoadingPhase": "Default"
+		},
+		{
+			"Name": "GLTFExporterRuntime",
+			"Type": "Runtime"
 		}
 	]
 }
diff --git a/Engine/Plugins/Enterprise/GLTFExporter/MigratedFromGit/Source/GLTFExporter/Private/Actors/GLTFInteractionHotspotActor.cpp b/Engine/Plugins/Enterprise/GLTFExporter/MigratedFromGit/Source/GLTFExporter/Private/Actors/GLTFInteractionHotspotActor.cpp
deleted file mode 100644
diff --git a/Engine/Plugins/Enterprise/GLTFExporter/MigratedFromGit/Source/GLTFExporter/Private/Components/GLTFInteractionHotspotComponent.cpp b/Engine/Plugins/Enterprise/GLTFExporter/MigratedFromGit/Source/GLTFExporter/Private/Components/GLTFInteractionHotspotComponent.cpp
deleted file mode 100644
diff --git a/Engine/Plugins/Enterprise/GLTFExporter/MigratedFromGit/Source/GLTFExporter/Public/Actors/GLTFInteractionHotspotActor.h b/Engine/Plugins/Enterprise/GLTFExporter/MigratedFromGit/Source/GLTFExporter/Public/Actors/GLTFInteractionHotspotActor.h
deleted file mode 100644
diff --git a/Engine/Plugins/Enterprise/GLTFExporter/MigratedFromGit/Source/GLTFExporter/Public/Components/GLTFInteractionHotspotComponent.h b/Engine/Plugins/Enterprise/GLTFExporter/MigratedFromGit/Source/GLTFExporter/Public/Components/GLTFInteractionHotspotComponent.h
deleted file mode 100644
diff --git a/Engine/Plugins/Enterprise/GLTFExporter/MigratedFromGit/Source/GLTFExporterRuntime/GLTFExporterRuntime.Build.cs b/Engine/Plugins/Enterprise/GLTFExporter/MigratedFromGit/Source/GLTFExporterRuntime/GLTFExporterRuntime.Build.cs
new file mode 100644
--- /dev/null
+++ b/Engine/Plugins/Enterprise/GLTFExporter/MigratedFromGit/Source/GLTFExporterRuntime/GLTFExporterRuntime.Build.cs
@@ -0,0 +1,42 @@
+// Copyright Epic Games, Inc. All Rights Reserved.
+
+namespace UnrealBuildTool.Rules
+{
+	public class GLTFExporterRuntime : ModuleRules
+	{
+		public GLTFExporterRuntime(ReadOnlyTargetRules Target) : base(Target)
+		{
+			PublicIncludePaths.AddRange(
+				new string[] {
+				}
+				);
+
+			PrivateIncludePaths.AddRange(
+				new string[] {
+				}
+				);
+
+			PublicDependencyModuleNames.AddRange(
+				new string[]
+				{
+				}
+				);
+
+			PrivateDependencyModuleNames.AddRange(
+				new string[]
+				{
+					"Core",
+					"CoreUObject",
+					"Engine",
+					"InputCore"
+				}
+				);
+
+			DynamicallyLoadedModuleNames.AddRange(
+				new string[]
+				{
+				}
+				);
+		}
+	}
+}
diff --git a/Engine/Plugins/Enterprise/GLTFExporter/MigratedFromGit/Source/GLTFExporterRuntime/Private/Actors/GLTFInteractionHotspotActor.cpp b/Engine/Plugins/Enterprise/GLTFExporter/MigratedFromGit/Source/GLTFExporterRuntime/Private/Actors/GLTFInteractionHotspotActor.cpp
new file mode 100644
--- /dev/null
+++ b/Engine/Plugins/Enterprise/GLTFExporter/MigratedFromGit/Source/GLTFExporterRuntime/Private/Actors/GLTFInteractionHotspotActor.cpp
@@ -0,0 +1,10 @@
+// Copyright Epic Games, Inc. All Rights Reserved.
+
+#include "Actors/GLTFInteractionHotspotActor.h"
+
+AGLTFInteractionHotspotActor::AGLTFInteractionHotspotActor(const FObjectInitializer& ObjectInitializer)
+	: Super(ObjectInitializer)
+{
+	InteractionHotspotComponent = CreateDefaultSubobject<UGLTFInteractionHotspotComponent>(TEXT("InteractionHotspotComponent"));
+	RootComponent = InteractionHotspotComponent;
+}
diff --git a/Engine/Plugins/Enterprise/GLTFExporter/MigratedFromGit/Source/GLTFExporterRuntime/Private/Components/GLTFInteractionHotspotComponent.cpp b/Engine/Plugins/Enterprise/GLTFExporter/MigratedFromGit/Source/GLTFExporterRuntime/Private/Components/GLTFInteractionHotspotComponent.cpp
new file mode 100644
--- /dev/null
+++ b/Engine/Plugins/Enterprise/GLTFExporter/MigratedFromGit/Source/GLTFExporterRuntime/Private/Components/GLTFInteractionHotspotComponent.cpp
@@ -0,0 +1,240 @@
+// Copyright Epic Games, Inc. All Rights Reserved.
+
+#include "Components/GLTFInteractionHotspotComponent.h"
+#include "Animation/SkeletalMeshActor.h"
+#include "Animation/AnimSequence.h"
+#include "Engine/CollisionProfile.h"
+#include "UObject/ConstructorHelpers.h"
+#include "PhysicsEngine/BodySetup.h"
+#include "Animation/AnimSingleNodeInstance.h"
+
+namespace
+{
+	const float UnitSphereRadius = 50.0f;
+} // Anonymous namespace
+
+DEFINE_LOG_CATEGORY_STATIC(LogEditorGLTFInteractionHotspot, Log, All);
+
+UGLTFInteractionHotspotComponent::UGLTFInteractionHotspotComponent(const FObjectInitializer& ObjectInitializer)
+	: Super(ObjectInitializer),
+	Image(nullptr),
+	HoveredImage(nullptr),
+	ToggledImage(nullptr),
+	ToggledHoveredImage(nullptr),
+	ShapeBodySetup(nullptr),
+	bToggled(bToggled)
+{
+	// Setup the most minimalistic collision profile for mouse input events
+	SetCollisionEnabled(ECollisionEnabled::QueryAndPhysics);
+	SetCollisionResponseToAllChannels(ECollisionResponse::ECR_Ignore);
+	SetCollisionResponseToChannel(ECollisionChannel::ECC_Visibility, ECollisionResponse::ECR_Block);
+	SetGenerateOverlapEvents(false);
+
+	bHiddenInGame = false;
+
+	OnBeginCursorOver.AddDynamic(this, &UGLTFInteractionHotspotComponent::BeginCursorOver);
+	OnEndCursorOver.AddDynamic(this, &UGLTFInteractionHotspotComponent::EndCursorOver);
+	OnClicked.AddDynamic(this, &UGLTFInteractionHotspotComponent::Clicked);
+}
+
+#if WITH_EDITOR
+void UGLTFInteractionHotspotComponent::PostEditChangeProperty(FPropertyChangedEvent& PropertyChangedEvent)
+{
+	Super::PostEditChangeProperty(PropertyChangedEvent);
+
+	const FProperty* PropertyThatChanged = PropertyChangedEvent.Property;
+
+	if (PropertyThatChanged)
+	{
+		const FString PropertyName = PropertyThatChanged->GetName();
+
+		if (PropertyName == TEXT("Image"))
+		{
+			SetSprite(Image);
+		}
+		else if (PropertyName == TEXT("SkeletalMeshActor"))
+		{
+			if (AnimationSequence != nullptr && AnimationSequence->GetSkeleton() != SkeletalMeshActor->GetSkeletalMeshComponent()->SkeletalMesh->Skeleton)
+			{
+				UE_LOG(LogEditorGLTFInteractionHotspot, Warning, TEXT("The skeleton of this actor is not compatible with the previously assigned animation sequence"));
+			}
+		}
+		else if (PropertyName == TEXT("AnimationSequence"))
+		{
+			if (SkeletalMeshActor != nullptr && SkeletalMeshActor->GetSkeletalMeshComponent()->SkeletalMesh->Skeleton != AnimationSequence->GetSkeleton())
+			{
+				UE_LOG(LogEditorGLTFInteractionHotspot, Warning, TEXT("This animation sequence is not compatible with the skeleton of the previously assigned actor"));
+			}
+		}
+	}
+}
+#endif // WITH_EDITOR
+
+void UGLTFInteractionHotspotComponent::BeginPlay()
+{
+	Super::BeginPlay();
+
+	SetSprite(Image);
+}
+
+void UGLTFInteractionHotspotComponent::OnRegister()
+{
+	ShapeBodySetup = NewObject<UBodySetup>(this, NAME_None, RF_Transient);
+
+	if (GUObjectArray.IsDisregardForGC(this))
+	{
+		ShapeBodySetup->AddToRoot();
+	}
+
+	ShapeBodySetup->AddToCluster(this);
+
+	if (ShapeBodySetup->HasAnyInternalFlags(EInternalObjectFlags::Async) && GUObjectClusters.GetObjectCluster(ShapeBodySetup))
+	{
+		ShapeBodySetup->ClearInternalFlags(EInternalObjectFlags::Async);
+	}
+
+	ShapeBodySetup->CollisionTraceFlag = CTF_UseSimpleAsComplex;
+	ShapeBodySetup->AggGeom.SphereElems.Add(FKSphereElem(UnitSphereRadius));
+	
+	Super::OnRegister();
+}
+
+void UGLTFInteractionHotspotComponent::OnCreatePhysicsState()
+{
+	Super::OnCreatePhysicsState();
+
+	if (BodyInstance.IsValidBodyInstance())
+	{
+#if WITH_PHYSX
+		FPhysicsCommand::ExecuteWrite(BodyInstance.GetActorReferenceWithWelding(), [this](const FPhysicsActorHandle& Actor)
+		{
+			TArray<FPhysicsShapeHandle> Shapes;
+			BodyInstance.GetAllShapes_AssumesLocked(Shapes);
+
+			for (FPhysicsShapeHandle& Shape : Shapes)
+			{
+				if (BodyInstance.IsShapeBoundToBody(Shape))
+				{
+					FPhysicsInterface::SetUserData(Shape, (void*)ShapeBodySetup->AggGeom.SphereElems[0].GetUserData());
+				}
+			}
+		});
+#endif
+	}
+
+	UpdateCollisionVolume();
+}
+
+UBodySetup* UGLTFInteractionHotspotComponent::GetBodySetup()
+{
+	return ShapeBodySetup;
+}
+
+void UGLTFInteractionHotspotComponent::SetSprite(class UTexture2D* NewSprite)
+{
+	Super::SetSprite(NewSprite);
+
+	UpdateCollisionVolume();
+}
+
+void UGLTFInteractionHotspotComponent::BeginCursorOver(UPrimitiveComponent* TouchedComponent)
+{
+	SetSprite(GetActiveImage(true));
+}
+
+void UGLTFInteractionHotspotComponent::EndCursorOver(UPrimitiveComponent* TouchedComponent)
+{
+	SetSprite(GetActiveImage(false));
+}
+
+void UGLTFInteractionHotspotComponent::Clicked(UPrimitiveComponent* TouchedComponent, FKey ButtonPressed)
+{
+	const bool bReverseAnimation = bToggled;
+
+	if (SkeletalMeshActor != nullptr && AnimationSequence != nullptr)
+	{
+		USkeletalMeshComponent* SkeletalMeshComponent = SkeletalMeshActor->GetSkeletalMeshComponent();
+		SkeletalMeshComponent->SetAnimationMode(EAnimationMode::Type::AnimationSingleNode);
+		const float AbsolutePlayRate = FMath::Abs(SkeletalMeshComponent->GetPlayRate());
+		const UAnimSingleNodeInstance* SingleNodeInstance = SkeletalMeshComponent->GetSingleNodeInstance();
+
+		if (SkeletalMeshComponent->IsPlaying() && SingleNodeInstance != nullptr && SingleNodeInstance->GetAnimationAsset() == AnimationSequence)
+		{
+			// If the same animation is already playing, just reverse the play rate for a smooth transition
+			SkeletalMeshComponent->SetPlayRate(AbsolutePlayRate * -1.0f);
+		}
+		else
+		{
+			SkeletalMeshComponent->SetAnimation(AnimationSequence);
+			SkeletalMeshComponent->SetPlayRate(AbsolutePlayRate * (bReverseAnimation ? -1.0f : 1.0f));
+			SkeletalMeshComponent->SetPosition(bReverseAnimation ? AnimationSequence->GetPlayLength() : 0.0f);
+			SkeletalMeshComponent->Play(false);
+		}
+	}
+
+	bToggled = !bToggled;
+
+	SetSprite(GetActiveImage(true));
+}
+
+void UGLTFInteractionHotspotComponent::UpdateCollisionVolume()
+{
+	if (ShapeBodySetup != nullptr)
+	{
+		// TODO: Figure out why the bounding radius doesn't match the size of the billboard
+		const float Scaling = 0.15f;
+		const float BillboardBoundingRadius = GetBillboardBoundingRadius() * Scaling;
+
+		check(ShapeBodySetup->AggGeom.SphereElems.Num() == 1);
+
+		if (!FMath::IsNearlyEqual(ShapeBodySetup->AggGeom.SphereElems[0].Radius, BillboardBoundingRadius) && UnitSphereRadius != 0.0f)
+		{
+			BodyInstance.UpdateBodyScale(FVector(BillboardBoundingRadius / UnitSphereRadius), true);
+		}
+	}
+}
+
+float UGLTFInteractionHotspotComponent::GetBillboardBoundingRadius() const
+{
+	const FTransform WorldTransform = GetComponentTransform();
+	const FBoxSphereBounds WorldBounds = CalcBounds(WorldTransform);
+
+	return WorldBounds.SphereRadius;
+}
+
+UTexture2D* UGLTFInteractionHotspotComponent::GetActiveImage(bool bCursorOver) const
+{
+	// A list of candidates ordered by descending priority
+	TArray<UTexture2D*> ImageCandidates;
+
+	if (bToggled)
+	{
+		if (bCursorOver)
+		{
+			ImageCandidates.Add(ToggledHoveredImage);
+		}
+
+		ImageCandidates.Add(ToggledImage);
+	}
+	else
+	{
+		if (bCursorOver)
+		{
+			ImageCandidates.Add(HoveredImage);
+		}
+
+		ImageCandidates.Add(Image);
+	}
+
+	ImageCandidates.Add(Image);
+
+	for (UTexture2D* ImageCandidate : ImageCandidates)
+	{
+		if (ImageCandidate != nullptr)
+		{
+			return ImageCandidate;
+		}
+	}
+
+	return nullptr;
+}
diff --git a/Engine/Plugins/Enterprise/GLTFExporter/MigratedFromGit/Source/GLTFExporterRuntime/Private/GLTFExporterRuntimeModule.cpp b/Engine/Plugins/Enterprise/GLTFExporter/MigratedFromGit/Source/GLTFExporterRuntime/Private/GLTFExporterRuntimeModule.cpp
new file mode 100644
--- /dev/null
+++ b/Engine/Plugins/Enterprise/GLTFExporter/MigratedFromGit/Source/GLTFExporterRuntime/Private/GLTFExporterRuntimeModule.cpp
@@ -0,0 +1,22 @@
+// Copyright Epic Games, Inc. All Rights Reserved.
+
+#include "GLTFExporterRuntimeModule.h"
+
+#include "CoreMinimal.h"
+#include "Modules/ModuleManager.h"
+
+DEFINE_LOG_CATEGORY(LogGLTFExporterRuntime);
+
+/**
+ * glTF Exporter runtime module implementation (private)
+ */
+class FGLTFExporterRuntimeModule final : public IGLTFExporterRuntimeModule
+{
+
+public:
+	virtual void StartupModule() override {}
+
+	virtual void ShutdownModule() override {}
+};
+
+IMPLEMENT_MODULE(FGLTFExporterRuntimeModule, GLTFExporterRuntime);
diff --git a/Engine/Plugins/Enterprise/GLTFExporter/MigratedFromGit/Source/GLTFExporterRuntime/Public/Actors/GLTFInteractionHotspotActor.h b/Engine/Plugins/Enterprise/GLTFExporter/MigratedFromGit/Source/GLTFExporterRuntime/Public/Actors/GLTFInteractionHotspotActor.h
new file mode 100644
--- /dev/null
+++ b/Engine/Plugins/Enterprise/GLTFExporter/MigratedFromGit/Source/GLTFExporterRuntime/Public/Actors/GLTFInteractionHotspotActor.h
@@ -0,0 +1,24 @@
+// Copyright Epic Games, Inc. All Rights Reserved.
+
+#pragma once
+
+#include "GameFramework/Actor.h"
+#include "Components/GLTFInteractionHotspotComponent.h"
+#include "GLTFInteractionHotspotActor.generated.h"
+
+/**
+ * Actor wrapper for the GLTF hotspot component. Appears as a billboard and allows playback of skeletal animations when cursor input is enabled.
+ */
+UCLASS(DisplayName = "GLTF Interaction Hotspot Actor", HideCategories = (Sprite))
+class GLTFEXPORTERRUNTIME_API AGLTFInteractionHotspotActor : public AActor
+{
+	GENERATED_BODY()
+	//~ Begin UObject Interface
+public:
+	AGLTFInteractionHotspotActor(const FObjectInitializer& ObjectInitializer);
+	//~ End UObject Interface
+
+private:
+	UPROPERTY(Category = "GLTF Interaction Hotspot", VisibleAnywhere, BlueprintReadOnly, meta = (ExposeFunctionCategories = "GLTF Interaction Hotspot", AllowPrivateAccess = "true"))
+	UGLTFInteractionHotspotComponent* InteractionHotspotComponent;
+};
diff --git a/Engine/Plugins/Enterprise/GLTFExporter/MigratedFromGit/Source/GLTFExporterRuntime/Public/Components/GLTFInteractionHotspotComponent.h b/Engine/Plugins/Enterprise/GLTFExporter/MigratedFromGit/Source/GLTFExporterRuntime/Public/Components/GLTFInteractionHotspotComponent.h
new file mode 100644
--- /dev/null
+++ b/Engine/Plugins/Enterprise/GLTFExporter/MigratedFromGit/Source/GLTFExporterRuntime/Public/Components/GLTFInteractionHotspotComponent.h
@@ -0,0 +1,88 @@
+// Copyright Epic Games, Inc. All Rights Reserved.
+
+#pragma once
+
+#include "Components/BillboardComponent.h"
+#include "GLTFInteractionHotspotComponent.generated.h"
+
+class ASkeletalMeshActor;
+class UAnimSequence;
+class UBodySetup;
+class UTexture2D;
+
+/**
+ * A component to set up hotspots which appear as billboards and allow playback of skeletal animations when cursor input is enabled.
+ */
+UCLASS(BlueprintType, Blueprintable, meta = (BlueprintSpawnableComponent), hidecategories = (Sprite), DisplayName = "GLTF Interaction Hotspot Component")
+class GLTFEXPORTERRUNTIME_API UGLTFInteractionHotspotComponent : public UBillboardComponent
+{
+	GENERATED_UCLASS_BODY()
+public:
+	//~ Begin UObject Interface
+#if WITH_EDITOR
+	virtual void PostEditChangeProperty(FPropertyChangedEvent& PropertyChangedEvent) override;
+#endif // WITH_EDITOR
+	//~ End UObject Interface
+
+	//~ Begin UActorComponent Interface
+	virtual void BeginPlay() override;
+
+protected:
+	virtual void OnRegister() override;
+	virtual void OnCreatePhysicsState() override;
+	//~ End UActorComponent Interface
+
+	//~ Begin UPrimitiveComponent Interface
+public:
+	virtual UBodySetup* GetBodySetup() override;
+	//~ End UPrimitiveComponent Interface
+
+	//~ Begin UBillboardComponent Interface
+	virtual void SetSprite(class UTexture2D* NewSprite) override;
+	//~ End UBillboardComponent Interface
+
+private:
+	UFUNCTION()
+	void BeginCursorOver(UPrimitiveComponent* TouchedComponent);
+
+	UFUNCTION()
+	void EndCursorOver(UPrimitiveComponent* TouchedComponent);
+
+	UFUNCTION()
+	void Clicked(UPrimitiveComponent* TouchedComponent, FKey ButtonPressed);
+
+	void UpdateCollisionVolume();
+	float GetBillboardBoundingRadius() const;
+	UTexture2D* GetActiveImage(bool bCursorOver) const;
+
+public:
+	/* The skeletal mesh actor that will be animated when the hotspot is clicked. */
+	UPROPERTY(EditAnywhere, BlueprintReadWrite, Category = "GLTF Interaction Hotspot")
+	ASkeletalMeshActor* SkeletalMeshActor;
+
+	/* The animation that will be played on the skeletal mesh actor. Must be compatible with its skeletal mesh asset. */
+	UPROPERTY(EditAnywhere, BlueprintReadWrite, Category = "GLTF Interaction Hotspot")
+	UAnimSequence* AnimationSequence;
+
+	/* The billboard image that will be shown when the hotspot is in an inactive state or one without a specified image. */
+	UPROPERTY(EditAnywhere, BlueprintReadWrite, Category = "GLTF Interaction Hotspot")
+	UTexture2D* Image;
+
+	/** The optional billboard image that will be shown when a cursor enters the hotspot. */
+	UPROPERTY(EditAnywhere, BlueprintReadWrite, Category = "GLTF Interaction Hotspot")
+	UTexture2D* HoveredImage;
+
+	/** The optional billboard image that will be shown when the hotspot is toggled by a click. */
+	UPROPERTY(EditAnywhere, BlueprintReadWrite, Category = "GLTF Interaction Hotspot")
+	UTexture2D* ToggledImage;
+
+	/** The optional billboard image that will be shown when the hotspot is toggled by a click and a cursor enters it. */
+	UPROPERTY(EditAnywhere, BlueprintReadWrite, Category = "GLTF Interaction Hotspot")
+	UTexture2D* ToggledHoveredImage;
+
+private:
+	UPROPERTY(transient, duplicatetransient)
+	UBodySetup* ShapeBodySetup;
+
+	bool bToggled;
+};
diff --git a/Engine/Plugins/Enterprise/GLTFExporter/MigratedFromGit/Source/GLTFExporterRuntime/Public/GLTFExporterRuntimeModule.h b/Engine/Plugins/Enterprise/GLTFExporter/MigratedFromGit/Source/GLTFExporterRuntime/Public/GLTFExporterRuntimeModule.h
new file mode 100644
--- /dev/null
+++ b/Engine/Plugins/Enterprise/GLTFExporter/MigratedFromGit/Source/GLTFExporterRuntime/Public/GLTFExporterRuntimeModule.h
@@ -0,0 +1,36 @@
+// Copyright Epic Games, Inc. All Rights Reserved.
+
+#pragma once
+
+#include "CoreMinimal.h"
+#include "Modules/ModuleInterface.h"
+#include "Modules/ModuleManager.h"
+
+DECLARE_LOG_CATEGORY_EXTERN(LogGLTFExporterRuntime, Log, All);
+
+/**
+ * The public interface of the GLTFExporterRuntime module
+ */
+class IGLTFExporterRuntimeModule : public IModuleInterface
+{
+public:
+	/**
+	 * Singleton-like access to IGLTFExporterRuntime
+	 *
+	 * @return Returns IGLTFExporterRuntime singleton instance, loading the module on demand if needed
+	 */
+	static IGLTFExporterRuntimeModule& Get()
+	{
+		return FModuleManager::LoadModuleChecked<IGLTFExporterRuntimeModule>("GLTFExporterRuntime");
+	}
+
+	/**
+	 * Checks to see if this module is loaded and ready.  It is only valid to call Get() if IsAvailable() returns true.
+	 *
+	 * @return True if the module is loaded and ready to use
+	 */
+	static bool IsAvailable()
+	{
+		return FModuleManager::Get().IsModuleLoaded("GLTFExporterRuntime");
+	}
+};
