[GLTFExporter] refactor(hotspot): add blueptint specifiers to UCLASS in actor and tweak HideCategories (doesn't seem to affect the actor though)

commit dfb2631605ba7c2f87d62c07e83b5157e8946462
Date: Wed, 24 Jun 2020 23:24:10 +0200

    refactor(hotspot): add blueptint specifiers to UCLASS in actor and tweak HideCategories (doesn't seem to affect the actor though)
    (cherry picked from commit a90fc40e3ab0d5cd87d84f8f9238cb90d85e25e1​)
    # Conflicts from cherry-pick:
    #	Plugins/GLTFExporter/Source/GLTFExporterRuntime/Public/Actors/GLTFInteractionHotspotActor.h
    #	Plugins/GLTFExporter/Source/GLTFExporterRuntime/Public/Components/GLTFInteractionHotspotComponent.h

#skipci

diff --git a/Engine/Plugins/Enterprise/GLTFExporter/MigratedFromGit/Source/GLTFExporterRuntime/Public/Actors/GLTFInteractionHotspotActor.h b/Engine/Plugins/Enterprise/GLTFExporter/MigratedFromGit/Source/GLTFExporterRuntime/Public/Actors/GLTFInteractionHotspotActor.h
--- a/Engine/Plugins/Enterprise/GLTFExporter/MigratedFromGit/Source/GLTFExporterRuntime/Public/Actors/GLTFInteractionHotspotActor.h
+++ b/Engine/Plugins/Enterprise/GLTFExporter/MigratedFromGit/Source/GLTFExporterRuntime/Public/Actors/GLTFInteractionHotspotActor.h
@@ -9,7 +9,7 @@
 /**
  * Actor wrapper for the GLTF hotspot component. Appears as a billboard and allows playback of skeletal animations when cursor input is enabled.
  */
-UCLASS(DisplayName = "GLTF Interaction Hotspot Actor", HideCategories = (Sprite))
+UCLASS(BlueprintType, Blueprintable, HideCategories = (Sprite, Physics, Collision, Navigation), DisplayName = "GLTF Interaction Hotspot Actor")
 class GLTFEXPORTERRUNTIME_API AGLTFInteractionHotspotActor : public AActor
 {
 	GENERATED_BODY()
diff --git a/Engine/Plugins/Enterprise/GLTFExporter/MigratedFromGit/Source/GLTFExporterRuntime/Public/Components/GLTFInteractionHotspotComponent.h b/Engine/Plugins/Enterprise/GLTFExporter/MigratedFromGit/Source/GLTFExporterRuntime/Public/Components/GLTFInteractionHotspotComponent.h
--- a/Engine/Plugins/Enterprise/GLTFExporter/MigratedFromGit/Source/GLTFExporterRuntime/Public/Components/GLTFInteractionHotspotComponent.h
+++ b/Engine/Plugins/Enterprise/GLTFExporter/MigratedFromGit/Source/GLTFExporterRuntime/Public/Components/GLTFInteractionHotspotComponent.h
@@ -13,7 +13,7 @@ class UTexture2D;
 /**
  * A component to set up hotspots which appear as billboards and allow playback of skeletal animations when cursor input is enabled.
  */
-UCLASS(BlueprintType, Blueprintable, meta = (BlueprintSpawnableComponent), hidecategories = (Sprite), DisplayName = "GLTF Interaction Hotspot Component")
+UCLASS(BlueprintType, Blueprintable, Meta = (BlueprintSpawnableComponent), HideCategories = (Sprite, Physics, Collision, Navigation), DisplayName = "GLTF Interaction Hotspot Component")
 class GLTFEXPORTERRUNTIME_API UGLTFInteractionHotspotComponent : public UBillboardComponent
 {
 	GENERATED_UCLASS_BODY()
