[GLTFExporter] docs(interaction-hotspot): mark unsolved issue with todo comment

commit 1ed4e39f6cd30a5600eab968de3e887dea615238
Date: Thu, 25 Jun 2020 15:07:18 +0200

    docs(interaction-hotspot): mark unsolved issue with todo comment
    (cherry picked from commit 570179b7e42ebe579b507ae70b3d92ac8d7d287a​)

#skipci

diff --git a/Engine/Plugins/Enterprise/GLTFExporter/MigratedFromGit/Source/GLTFExporterRuntime/Private/Actors/GLTFInteractionHotspotActor.cpp b/Engine/Plugins/Enterprise/GLTFExporter/MigratedFromGit/Source/GLTFExporterRuntime/Private/Actors/GLTFInteractionHotspotActor.cpp
--- a/Engine/Plugins/Enterprise/GLTFExporter/MigratedFromGit/Source/GLTFExporterRuntime/Private/Actors/GLTFInteractionHotspotActor.cpp
+++ b/Engine/Plugins/Enterprise/GLTFExporter/MigratedFromGit/Source/GLTFExporterRuntime/Private/Actors/GLTFInteractionHotspotActor.cpp
@@ -5,6 +5,7 @@
 AGLTFInteractionHotspotActor::AGLTFInteractionHotspotActor(const FObjectInitializer& ObjectInitializer)
 	: Super(ObjectInitializer)
 {
+	// TODO: Despite the HideCategory class specifier present in both the hotspot actor and the component, the category "Sprite" bleeds through from UBillboardComponent
 	InteractionHotspotComponent = CreateDefaultSubobject<UGLTFInteractionHotspotComponent>(TEXT("InteractionHotspotComponent"));
 	RootComponent = InteractionHotspotComponent;
 }
