[GLTFExporter] refactor: rename GLTFMaterialUtility::GetInputFromProperty -> GetInputForProperty

commit 8f5d51a4b2457b8d90610516f853ee0ee539930b
Date: Thu, 25 Jun 2020 23:23:30 +0200

    refactor: rename GLTFMaterialUtility::GetInputFromProperty -> GetInputForProperty
    (cherry picked from commit f63b91d7e3fa4b49eff620bf50a1efbbe2b335f6​)

#skipci

diff --git a/Engine/Plugins/Enterprise/GLTFExporter/MigratedFromGit/Source/GLTFExporter/Private/Converters/GLTFMaterialConverters.cpp b/Engine/Plugins/Enterprise/GLTFExporter/MigratedFromGit/Source/GLTFExporter/Private/Converters/GLTFMaterialConverters.cpp
--- a/Engine/Plugins/Enterprise/GLTFExporter/MigratedFromGit/Source/GLTFExporter/Private/Converters/GLTFMaterialConverters.cpp
+++ b/Engine/Plugins/Enterprise/GLTFExporter/MigratedFromGit/Source/GLTFExporter/Private/Converters/GLTFMaterialConverters.cpp
@@ -407,7 +407,7 @@ bool FGLTFMaterialConverter::IsPropertyNonDefault(EMaterialProperty Property, co
 		return true;
 	}
 
-	const FExpressionInput* MaterialInput = FGLTFMaterialUtility::GetInputFromProperty(Material, Property);
+	const FExpressionInput* MaterialInput = FGLTFMaterialUtility::GetInputForProperty(Material, Property);
 	if (MaterialInput == nullptr)
 	{
 		// TODO: report error
@@ -458,7 +458,7 @@ bool FGLTFMaterialConverter::TryGetConstantColor(FLinearColor& OutValue, EMateri
 		return false;
 	}
 
-	const FMaterialInput<FColor>* MaterialInput = FGLTFMaterialUtility::GetInputFromProperty<FColor>(Material, Property);
+	const FMaterialInput<FColor>* MaterialInput = FGLTFMaterialUtility::GetInputForProperty<FColor>(Material, Property);
 	if (MaterialInput == nullptr)
 	{
 		// TODO: report error
@@ -565,7 +565,7 @@ bool FGLTFMaterialConverter::TryGetConstantScalar(float& OutValue, EMaterialProp
 		return false;
 	}
 
-	const FMaterialInput<float>* MaterialInput = FGLTFMaterialUtility::GetInputFromProperty<float>(Material, Property);
+	const FMaterialInput<float>* MaterialInput = FGLTFMaterialUtility::GetInputForProperty<float>(Material, Property);
 	if (MaterialInput == nullptr)
 	{
 		// TODO: report error
@@ -674,7 +674,7 @@ bool FGLTFMaterialConverter::TryGetSourceTexture(FGLTFConvertBuilder& Builder, F
 
 bool FGLTFMaterialConverter::TryGetSourceTexture(const UTexture2D*& OutTexture, int32& OutTexCoord, EMaterialProperty Property, const UMaterialInterface* Material, const TArray<FLinearColor>& AllowedMasks) const
 {
-	const FExpressionInput* MaterialInput = FGLTFMaterialUtility::GetInputFromProperty(Material, Property);
+	const FExpressionInput* MaterialInput = FGLTFMaterialUtility::GetInputForProperty(Material, Property);
 	if (MaterialInput == nullptr)
 	{
 		// TODO: report error
diff --git a/Engine/Plugins/Enterprise/GLTFExporter/MigratedFromGit/Source/GLTFExporter/Private/Converters/GLTFMaterialUtility.cpp b/Engine/Plugins/Enterprise/GLTFExporter/MigratedFromGit/Source/GLTFExporter/Private/Converters/GLTFMaterialUtility.cpp
--- a/Engine/Plugins/Enterprise/GLTFExporter/MigratedFromGit/Source/GLTFExporter/Private/Converters/GLTFMaterialUtility.cpp
+++ b/Engine/Plugins/Enterprise/GLTFExporter/MigratedFromGit/Source/GLTFExporter/Private/Converters/GLTFMaterialUtility.cpp
@@ -6,7 +6,6 @@
 #include "CanvasItem.h"
 #include "CanvasTypes.h"
 #include "Modules/ModuleManager.h"
-
 #include "IMaterialBakingModule.h"
 #include "MaterialBakingStructures.h"
 
@@ -48,7 +47,7 @@ FVector4 FGLTFMaterialUtility::GetPropertyDefaultValue(EMaterialProperty Propert
 	return FVector4();
 }
 
-const FExpressionInput* FGLTFMaterialUtility::GetInputFromProperty(const UMaterialInterface* Material, EMaterialProperty Property)
+const FExpressionInput* FGLTFMaterialUtility::GetInputForProperty(const UMaterialInterface* Material, EMaterialProperty Property)
 {
 	UMaterial* UnderlyingMaterial = const_cast<UMaterial*>(Material->GetMaterial());
 	return UnderlyingMaterial->GetExpressionInputForProperty(Property);
diff --git a/Engine/Plugins/Enterprise/GLTFExporter/MigratedFromGit/Source/GLTFExporter/Private/Converters/GLTFMaterialUtility.h b/Engine/Plugins/Enterprise/GLTFExporter/MigratedFromGit/Source/GLTFExporter/Private/Converters/GLTFMaterialUtility.h
--- a/Engine/Plugins/Enterprise/GLTFExporter/MigratedFromGit/Source/GLTFExporter/Private/Converters/GLTFMaterialUtility.h
+++ b/Engine/Plugins/Enterprise/GLTFExporter/MigratedFromGit/Source/GLTFExporter/Private/Converters/GLTFMaterialUtility.h
@@ -36,12 +36,12 @@ struct FGLTFMaterialUtility
 {
 	static FVector4 GetPropertyDefaultValue(EMaterialProperty Property);
 
-	static const FExpressionInput* GetInputFromProperty(const UMaterialInterface* Material, EMaterialProperty Property);
+	static const FExpressionInput* GetInputForProperty(const UMaterialInterface* Material, EMaterialProperty Property);
 
 	template<class InputType>
-    static const FMaterialInput<InputType>* GetInputFromProperty(const UMaterialInterface* Material, EMaterialProperty Property)
+    static const FMaterialInput<InputType>* GetInputForProperty(const UMaterialInterface* Material, EMaterialProperty Property)
 	{
-		const FExpressionInput* ExpressionInput = const_cast<UMaterial*>(Material->GetMaterial())->GetExpressionInputForProperty(Property);
+		const FExpressionInput* ExpressionInput = GetInputForProperty(Material, Property);
 		return static_cast<const FMaterialInput<InputType>*>(ExpressionInput);
 	}
 
