[GLTFExporter] Merge remote-tracking branch 'origin/master' into feature/implement-interaction-hotspots

commit 1052d4b8aa392a4ebfba45ace8657688737e775b
Date: Fri, 26 Jun 2020 08:37:41 +0200

    Merge remote-tracking branch 'origin/master' into feature/implement-interaction-hotspots
    # Conflicts from merge:
    #	Plugins/GLTFExporter/GLTFExporter.uplugin
    (cherry picked from commit 47a6ad12fafa430cb7a781d0f84c5e5e0227a9f3​)
    # Conflicts from cherry-pick:
    #	Plugins/GLTFExporter/GLTFExporter.uplugin

#skipci

diff --git a/Engine/Plugins/Enterprise/GLTFExporter/MigratedFromGit/GLTFExporter.uplugin b/Engine/Plugins/Enterprise/GLTFExporter/MigratedFromGit/GLTFExporter.uplugin
--- a/Engine/Plugins/Enterprise/GLTFExporter/MigratedFromGit/GLTFExporter.uplugin
+++ b/Engine/Plugins/Enterprise/GLTFExporter/MigratedFromGit/GLTFExporter.uplugin
@@ -22,7 +22,8 @@
 		},
 		{
 			"Name": "GLTFExporterRuntime",
-			"Type": "Runtime"
+			"Type": "Runtime",
+			"LoadingPhase": "Default"
 		}
 	]
 }
