[GLTFExporter] feat(material): enable export of EmissiveColor and use EmissiveScale properly for baked emission

commit ad845f572bc7a5a55ba9c15a6eed2a80d5e9d69a
Date: Wed, 24 Jun 2020 16:44:55 +0200

    feat(material): enable export of EmissiveColor and use EmissiveScale properly for baked emission
    (cherry picked from commit 87166e096ef31c61796ef3be8d8bc811dc42127d​)
    # Conflicts from cherry-pick:
    #	Plugins/GLTFExporter/Source/GLTFExporter/Private/Converters/GLTFMaterialConverters.cpp

#skipci

diff --git a/Engine/Plugins/Enterprise/GLTFExporter/MigratedFromGit/Source/GLTFExporter/Private/Converters/GLTFMaterialConverters.cpp b/Engine/Plugins/Enterprise/GLTFExporter/MigratedFromGit/Source/GLTFExporter/Private/Converters/GLTFMaterialConverters.cpp
--- a/Engine/Plugins/Enterprise/GLTFExporter/MigratedFromGit/Source/GLTFExporter/Private/Converters/GLTFMaterialConverters.cpp
+++ b/Engine/Plugins/Enterprise/GLTFExporter/MigratedFromGit/Source/GLTFExporter/Private/Converters/GLTFMaterialConverters.cpp
@@ -58,34 +58,10 @@ FGLTFJsonMaterialIndex FGLTFMaterialConverter::Add(FGLTFConvertBuilder& Builder,
 		Builder.AddWarningMessage(FString::Printf(TEXT("Failed to export Metallic & Roughness for material %s"), *Material->GetName()));
 	}
 
-	// NOTE: export of EmissiveColor has been temporarily disabled because of visual differences that
-	// are not solvable in the scope of MR !19.
-	// The issues revolve mainly around very bright emission in the exported materials compared to
-	// how the same materials look inside of UE.
-	// TODO: solve the issues in a separate MR.
-	/*
-	if (!TryGetConstantColor(JsonMaterial.EmissiveFactor, MP_EmissiveColor, Material))
-	{
-		if (!TryGetSourceTexture(Builder, JsonMaterial.EmissiveTexture, MP_EmissiveColor, Material, { RgbaMask, RgbMask }))
-		{
-			if (TryGetBakedMaterialProperty(Builder, JsonMaterial.EmissiveTexture, JsonMaterial.EmissiveFactor, MP_EmissiveColor, Material))
-			{
-				if (JsonMaterial.EmissiveTexture.Index != INDEX_NONE)
-				{
-					JsonMaterial.EmissiveFactor = FGLTFJsonColor3::White;	// make sure texture is not multiplied with black
-				}
-			}
-			else
-			{
-				Builder.AddWarningMessage(FString::Printf(TEXT("Failed to export EmissiveColor for material %s"), *Material->GetName()));
-			}
-		}
-		else
-		{
-			JsonMaterial.EmissiveFactor = FGLTFJsonColor3::White;	// make sure texture is not multiplied with black
-		}
+	if (!TryGetEmissiveColor(Builder, JsonMaterial, Material))
+	{
+		// TODO: handle failure?
 	}
-	*/
 
 	if (IsPropertyNonDefault(MP_Normal, Material))
 	{
@@ -398,6 +374,46 @@ bool FGLTFMaterialConverter::TryGetMetallicAndRoughness(FGLTFConvertBuilder& Bui
 	return true;
 }
 
+bool FGLTFMaterialConverter::TryGetEmissiveColor(FGLTFConvertBuilder& Builder, FGLTFJsonMaterial& JsonMaterial, const UMaterialInterface* Material) const
+{
+	// TODO: right now we allow EmissiveFactor to be > 1.0 to support very bright emission, although it's not valid according to the glTF standard.
+	// We may want to change this behaviour and store factors above 1.0 using a custom extension instead.
+
+	if (TryGetConstantColor(JsonMaterial.EmissiveFactor, MP_EmissiveColor, Material))
+	{
+		return true;
+	}
+
+	const FLinearColor RgbaMask(1.0f, 1.0f, 1.0f, 1.0f);
+	const FLinearColor RgbMask(1.0f, 1.0f, 1.0f, 0.0f);
+
+	if (TryGetSourceTexture(Builder, JsonMaterial.EmissiveTexture, MP_EmissiveColor, Material, { RgbaMask, RgbMask }))
+	{
+		JsonMaterial.EmissiveFactor = FGLTFJsonColor3::White;	// make sure texture is not multiplied with black
+		return true;
+	}
+
+	const FGLTFPropertyBakeOutput PropertyBakeOutput = BakeMaterialProperty(MP_EmissiveColor, Material);
+	const float EmissiveScale = PropertyBakeOutput.EmissiveScale;
+
+	if (PropertyBakeOutput.bIsConstant)
+	{
+		const FLinearColor EmissiveColor = PropertyBakeOutput.ConstantValue;
+		JsonMaterial.EmissiveFactor = FGLTFConverterUtility::ConvertColor(EmissiveColor * EmissiveScale);
+	}
+	else
+	{
+		if (!StoreBakedPropertyTexture(Builder, JsonMaterial.EmissiveTexture, PropertyBakeOutput, MP_EmissiveColor, Material))
+		{
+			return false;
+		}
+
+		JsonMaterial.EmissiveFactor = { EmissiveScale, EmissiveScale, EmissiveScale };
+	}
+
+	return true;
+}
+
 bool FGLTFMaterialConverter::IsPropertyNonDefault(EMaterialProperty Property, const UMaterialInterface* Material) const
 {
 	const bool bUseMaterialAttributes = Material->GetMaterial()->bUseMaterialAttributes;
@@ -449,8 +465,6 @@ bool FGLTFMaterialConverter::TryGetConstantColor(FGLTFJsonColor4& OutValue, EMat
 
 bool FGLTFMaterialConverter::TryGetConstantColor(FLinearColor& OutValue, EMaterialProperty Property, const UMaterialInterface* Material) const
 {
-	// TODO: handle emissive color-values above 1.0
-
 	const bool bUseMaterialAttributes = Material->GetMaterial()->bUseMaterialAttributes;
 	if (bUseMaterialAttributes)
 	{
@@ -814,7 +828,6 @@ FGLTFPropertyBakeOutput FGLTFMaterialConverter::BakeMaterialProperty(EMaterialPr
 	const FIntPoint DefaultTextureSize(512, 512);
 	const FIntPoint TextureSize = PreferredTextureSize != nullptr ? *PreferredTextureSize : DefaultTextureSize;
 
-	// TODO: handle cases where PropertyBakeOutput.EmissiveScale is not 1.0 (when baking EmissiveColor)
 	// TODO: add support for calculating the ideal resolution to use for baking based on connected (texture) nodes
 
 	const FGLTFPropertyBakeOutput PropertyBakeOutput = FGLTFMaterialUtility::BakeMaterialProperty(
diff --git a/Engine/Plugins/Enterprise/GLTFExporter/MigratedFromGit/Source/GLTFExporter/Private/Converters/GLTFMaterialConverters.h b/Engine/Plugins/Enterprise/GLTFExporter/MigratedFromGit/Source/GLTFExporter/Private/Converters/GLTFMaterialConverters.h
--- a/Engine/Plugins/Enterprise/GLTFExporter/MigratedFromGit/Source/GLTFExporter/Private/Converters/GLTFMaterialConverters.h
+++ b/Engine/Plugins/Enterprise/GLTFExporter/MigratedFromGit/Source/GLTFExporter/Private/Converters/GLTFMaterialConverters.h
@@ -15,6 +15,7 @@ class FGLTFMaterialConverter final : public TGLTFConverter<FGLTFJsonMaterialInde
 
 	bool TryGetBaseColorAndOpacity(FGLTFConvertBuilder& Builder, FGLTFJsonPBRMetallicRoughness& OutPBRParams, const UMaterialInterface* Material) const;
 	bool TryGetMetallicAndRoughness(FGLTFConvertBuilder& Builder, FGLTFJsonPBRMetallicRoughness& OutPBRParams, const UMaterialInterface* Material) const;
+	bool TryGetEmissiveColor(FGLTFConvertBuilder& Builder, FGLTFJsonMaterial& JsonMaterial, const UMaterialInterface* Material) const;
 
 	bool IsPropertyNonDefault(EMaterialProperty Property, const UMaterialInterface* Material) const;
 	bool TryGetConstantColor(FGLTFJsonColor3& OutValue, EMaterialProperty Property, const UMaterialInterface* Material) const;
