[GLTFExporter] chore(material): improve TODO about sRGB=>linear conversion of baked constants

commit 604ed8e6c0464808ff6e28fb7156ac5d7c632955
Date: Wed, 24 Jun 2020 16:45:40 +0200

    chore(material): improve TODO about sRGB=>linear conversion of baked constants
    (cherry picked from commit 37f6a74eb5515f3ffaadea99a6ab1873d60d811b​)

#skipci

diff --git a/Engine/Plugins/Enterprise/GLTFExporter/MigratedFromGit/Source/GLTFExporter/Private/Converters/GLTFMaterialUtility.cpp b/Engine/Plugins/Enterprise/GLTFExporter/MigratedFromGit/Source/GLTFExporter/Private/Converters/GLTFMaterialUtility.cpp
--- a/Engine/Plugins/Enterprise/GLTFExporter/MigratedFromGit/Source/GLTFExporter/Private/Converters/GLTFMaterialUtility.cpp
+++ b/Engine/Plugins/Enterprise/GLTFExporter/MigratedFromGit/Source/GLTFExporter/Private/Converters/GLTFMaterialUtility.cpp
@@ -180,6 +180,7 @@ FGLTFPropertyBakeOutput FGLTFMaterialUtility::BakeMaterialProperty(const FIntPoi
 		PropertyBakeOutput.bIsConstant = true;
 
 		// TODO: is the current conversion from sRGB => linear correct?
+		// It seems to give correct results for some properties, but not all.
 		PropertyBakeOutput.ConstantValue = Pixel;
 	}
 
