[GLTFExporter] refactor(material): use shared component- and input-masks when converting materials

commit 570e1e8e5074f73622fb7ef24698b8fa89f8230b
Date: Thu, 25 Jun 2020 10:53:34 +0200

    refactor(material): use shared component- and input-masks when converting materials
    (cherry picked from commit 2127b3b6f6677f764c20e4c5c9e0b4f5657b17b7​)

#skipci

diff --git a/Engine/Plugins/Enterprise/GLTFExporter/MigratedFromGit/Source/GLTFExporter/Private/Converters/GLTFMaterialConverters.cpp b/Engine/Plugins/Enterprise/GLTFExporter/MigratedFromGit/Source/GLTFExporter/Private/Converters/GLTFMaterialConverters.cpp
--- a/Engine/Plugins/Enterprise/GLTFExporter/MigratedFromGit/Source/GLTFExporter/Private/Converters/GLTFMaterialConverters.cpp
+++ b/Engine/Plugins/Enterprise/GLTFExporter/MigratedFromGit/Source/GLTFExporter/Private/Converters/GLTFMaterialConverters.cpp
@@ -13,6 +13,32 @@
 #include "Materials/MaterialExpressionTextureSample.h"
 #include "Materials/MaterialExpressionTextureSampleParameter2D.h"
 
+namespace
+{
+	// Component masks
+	const FLinearColor RedMask(1.0f, 0.0f, 0.0f, 0.0f);
+	const FLinearColor GreenMask(0.0f, 1.0f, 0.0f, 0.0f);
+	const FLinearColor BlueMask(0.0f, 0.0f, 1.0f, 0.0f);
+	const FLinearColor AlphaMask(0.0f, 0.0f, 0.0f, 1.0f);
+	const FLinearColor RgbMask = RedMask + GreenMask + BlueMask;
+	const FLinearColor RgbaMask = RgbMask + AlphaMask;
+
+	// Property-specific component masks
+	const FLinearColor BaseColorMask = RgbMask;
+	const FLinearColor OpacityMask = AlphaMask;
+	const FLinearColor MetallicMask = BlueMask;
+	const FLinearColor RoughnessMask = GreenMask;
+	const FLinearColor OcclusionMask = RedMask;
+
+	// Ideal masks for texture-inputs (doesn't require baking)
+	const TArray<FLinearColor> DefaultColorInputMasks = { RgbMask, RgbaMask };
+	const TArray<FLinearColor> BaseColorInputMasks = { BaseColorMask };
+	const TArray<FLinearColor> OpacityInputMasks = { OpacityMask };
+	const TArray<FLinearColor> MetallicInputMasks = { MetallicMask };
+	const TArray<FLinearColor> RoughnessInputMasks = { RoughnessMask };
+	const TArray<FLinearColor> OcclusionInputMasks = { OcclusionMask };
+} // anonymous namespace
+
 FGLTFJsonMaterialIndex FGLTFMaterialConverter::Add(FGLTFConvertBuilder& Builder, const FString& Name, const UMaterialInterface* Material)
 {
 	FGLTFJsonMaterial JsonMaterial;
@@ -24,17 +50,13 @@ FGLTFJsonMaterialIndex FGLTFMaterialConverter::Add(FGLTFConvertBuilder& Builder,
 	JsonMaterial.AlphaCutoff = Material->GetOpacityMaskClipValue();
 	JsonMaterial.DoubleSided = Material->IsTwoSided();
 
-	const FLinearColor RgbaMask(1.0f, 1.0f, 1.0f, 1.0f);
-	const FLinearColor RgbMask(1.0f, 1.0f, 1.0f, 0.0f);
-	const FLinearColor RMask(1.0f, 0.0f, 0.0f, 0.0f);
-
 	// TODO: check if a property is active before trying to get it (i.e. Material->IsPropertyActive)
 
 	if (JsonMaterial.AlphaMode == EGLTFJsonAlphaMode::Opaque)
 	{
 		if (!TryGetConstantColor(JsonMaterial.PBRMetallicRoughness.BaseColorFactor, MP_BaseColor, Material))
 		{
-			if (!TryGetSourceTexture(Builder, JsonMaterial.PBRMetallicRoughness.BaseColorTexture, MP_BaseColor, Material, { RgbaMask, RgbMask }))
+			if (!TryGetSourceTexture(Builder, JsonMaterial.PBRMetallicRoughness.BaseColorTexture, MP_BaseColor, Material, DefaultColorInputMasks))
 			{
 				if (!TryGetBakedMaterialProperty(Builder, JsonMaterial.PBRMetallicRoughness.BaseColorTexture, JsonMaterial.PBRMetallicRoughness.BaseColorFactor, MP_BaseColor, Material))
 				{
@@ -65,7 +87,7 @@ FGLTFJsonMaterialIndex FGLTFMaterialConverter::Add(FGLTFConvertBuilder& Builder,
 
 	if (IsPropertyNonDefault(MP_Normal, Material))
 	{
-		if (!TryGetSourceTexture(Builder, JsonMaterial.NormalTexture, MP_Normal, Material, { RgbaMask, RgbMask }))
+		if (!TryGetSourceTexture(Builder, JsonMaterial.NormalTexture, MP_Normal, Material, DefaultColorInputMasks))
 		{
 			if (!TryGetBakedMaterialProperty(Builder, JsonMaterial.NormalTexture, MP_Normal, Material))
 			{
@@ -76,7 +98,7 @@ FGLTFJsonMaterialIndex FGLTFMaterialConverter::Add(FGLTFConvertBuilder& Builder,
 
 	if (IsPropertyNonDefault(MP_AmbientOcclusion, Material))
 	{
-		if (!TryGetSourceTexture(Builder, JsonMaterial.OcclusionTexture, MP_AmbientOcclusion, Material, { RMask }))
+		if (!TryGetSourceTexture(Builder, JsonMaterial.OcclusionTexture, MP_AmbientOcclusion, Material, OcclusionInputMasks))
 		{
 			if (!TryGetBakedMaterialProperty(Builder, JsonMaterial.OcclusionTexture, MP_AmbientOcclusion, Material))
 			{
@@ -110,11 +132,8 @@ bool FGLTFMaterialConverter::TryGetBaseColorAndOpacity(FGLTFConvertBuilder& Buil
 	int32 BaseColorTexCoord;
 	int32 OpacityTexCoord;
 
-	const FLinearColor BaseColorMask(1.0f, 1.0f, 1.0f, 0.0f);
-	const FLinearColor OpacityMask(0.0f, 0.0f, 0.0f, 1.0f);
-
-	const bool bHasBaseColorSourceTexture = TryGetSourceTexture(BaseColorTexture, BaseColorTexCoord, MP_BaseColor, Material, { BaseColorMask });
-	const bool bHasOpacitySourceTexture = TryGetSourceTexture(OpacityTexture, OpacityTexCoord, OpacityProperty, Material, { OpacityMask });
+	const bool bHasBaseColorSourceTexture = TryGetSourceTexture(BaseColorTexture, BaseColorTexCoord, MP_BaseColor, Material, BaseColorInputMasks);
+	const bool bHasOpacitySourceTexture = TryGetSourceTexture(OpacityTexture, OpacityTexCoord, OpacityProperty, Material, OpacityInputMasks);
 
 	// Detect the "happy path" where both inputs share the same texture and are correctly masked.
 	if (bHasBaseColorSourceTexture &&
@@ -253,12 +272,8 @@ bool FGLTFMaterialConverter::TryGetMetallicAndRoughness(FGLTFConvertBuilder& Bui
 	int32 MetallicTexCoord;
 	int32 RoughnessTexCoord;
 
-	const FLinearColor MetallicMask(0.0f, 0.0f, 1.0f, 0.0f);
-	const FLinearColor RoughnessMask(0.0f, 1.0f, 0.0f, 0.0f);
-	const FLinearColor AlphaMask(0.0f, 0.0f, 0.0f, 1.0f);
-
-	const bool bHasMetallicSourceTexture = TryGetSourceTexture(MetallicTexture, MetallicTexCoord, MP_Metallic, Material, { MetallicMask });
-	const bool bHasRoughnessSourceTexture = TryGetSourceTexture(RoughnessTexture, RoughnessTexCoord, MP_Roughness, Material, { RoughnessMask });
+	const bool bHasMetallicSourceTexture = TryGetSourceTexture(MetallicTexture, MetallicTexCoord, MP_Metallic, Material, MetallicInputMasks);
+	const bool bHasRoughnessSourceTexture = TryGetSourceTexture(RoughnessTexture, RoughnessTexCoord, MP_Roughness, Material, RoughnessInputMasks);
 
 	// Detect the "happy path" where both inputs share the same texture and are correctly masked.
 	if (bHasMetallicSourceTexture &&
@@ -384,10 +399,7 @@ bool FGLTFMaterialConverter::TryGetEmissiveColor(FGLTFConvertBuilder& Builder, F
 		return true;
 	}
 
-	const FLinearColor RgbaMask(1.0f, 1.0f, 1.0f, 1.0f);
-	const FLinearColor RgbMask(1.0f, 1.0f, 1.0f, 0.0f);
-
-	if (TryGetSourceTexture(Builder, JsonMaterial.EmissiveTexture, MP_EmissiveColor, Material, { RgbaMask, RgbMask }))
+	if (TryGetSourceTexture(Builder, JsonMaterial.EmissiveTexture, MP_EmissiveColor, Material, DefaultColorInputMasks))
 	{
 		JsonMaterial.EmissiveFactor = FGLTFJsonColor3::White;	// make sure texture is not multiplied with black
 		return true;
