[GLTFExporter] refactor(material): rename function TryGetEmissiveColor to TryGetEmissive

commit 49132cb036ad7197d827f8b966521c8c283ec054
Date: Fri, 26 Jun 2020 08:12:13 +0200

    refactor(material): rename function TryGetEmissiveColor to TryGetEmissive
    (cherry picked from commit 4304ad655bbb76264b736ed6f9a9ae8bb7b6e990​)

#skipci

diff --git a/Engine/Plugins/Enterprise/GLTFExporter/MigratedFromGit/Source/GLTFExporter/Private/Converters/GLTFMaterialConverters.cpp b/Engine/Plugins/Enterprise/GLTFExporter/MigratedFromGit/Source/GLTFExporter/Private/Converters/GLTFMaterialConverters.cpp
--- a/Engine/Plugins/Enterprise/GLTFExporter/MigratedFromGit/Source/GLTFExporter/Private/Converters/GLTFMaterialConverters.cpp
+++ b/Engine/Plugins/Enterprise/GLTFExporter/MigratedFromGit/Source/GLTFExporter/Private/Converters/GLTFMaterialConverters.cpp
@@ -80,7 +80,7 @@ FGLTFJsonMaterialIndex FGLTFMaterialConverter::Add(FGLTFConvertBuilder& Builder,
 		Builder.AddWarningMessage(FString::Printf(TEXT("Failed to export Metallic & Roughness for material %s"), *Material->GetName()));
 	}
 
-	if (!TryGetEmissiveColor(Builder, JsonMaterial, Material))
+	if (!TryGetEmissive(Builder, JsonMaterial, Material))
 	{
 		// TODO: handle failure?
 	}
@@ -389,7 +389,7 @@ bool FGLTFMaterialConverter::TryGetMetallicAndRoughness(FGLTFConvertBuilder& Bui
 	return true;
 }
 
-bool FGLTFMaterialConverter::TryGetEmissiveColor(FGLTFConvertBuilder& Builder, FGLTFJsonMaterial& JsonMaterial, const UMaterialInterface* Material) const
+bool FGLTFMaterialConverter::TryGetEmissive(FGLTFConvertBuilder& Builder, FGLTFJsonMaterial& JsonMaterial, const UMaterialInterface* Material) const
 {
 	// TODO: right now we allow EmissiveFactor to be > 1.0 to support very bright emission, although it's not valid according to the glTF standard.
 	// We may want to change this behaviour and store factors above 1.0 using a custom extension instead.
diff --git a/Engine/Plugins/Enterprise/GLTFExporter/MigratedFromGit/Source/GLTFExporter/Private/Converters/GLTFMaterialConverters.h b/Engine/Plugins/Enterprise/GLTFExporter/MigratedFromGit/Source/GLTFExporter/Private/Converters/GLTFMaterialConverters.h
--- a/Engine/Plugins/Enterprise/GLTFExporter/MigratedFromGit/Source/GLTFExporter/Private/Converters/GLTFMaterialConverters.h
+++ b/Engine/Plugins/Enterprise/GLTFExporter/MigratedFromGit/Source/GLTFExporter/Private/Converters/GLTFMaterialConverters.h
@@ -15,7 +15,7 @@ class FGLTFMaterialConverter final : public TGLTFConverter<FGLTFJsonMaterialInde
 
 	bool TryGetBaseColorAndOpacity(FGLTFConvertBuilder& Builder, FGLTFJsonPBRMetallicRoughness& OutPBRParams, const UMaterialInterface* Material) const;
 	bool TryGetMetallicAndRoughness(FGLTFConvertBuilder& Builder, FGLTFJsonPBRMetallicRoughness& OutPBRParams, const UMaterialInterface* Material) const;
-	bool TryGetEmissiveColor(FGLTFConvertBuilder& Builder, FGLTFJsonMaterial& JsonMaterial, const UMaterialInterface* Material) const;
+	bool TryGetEmissive(FGLTFConvertBuilder& Builder, FGLTFJsonMaterial& JsonMaterial, const UMaterialInterface* Material) const;
 
 	bool IsPropertyNonDefault(EMaterialProperty Property, const UMaterialInterface* Material) const;
 	bool TryGetConstantColor(FGLTFJsonColor3& OutValue, EMaterialProperty Property, const UMaterialInterface* Material) const;
