[GLTFExporter] Merge remote-tracking branch 'origin/master' into feature/add-support-for-emissive-materials

commit c39032fdfddd66ad662d8bbd6ed3a94aed708c05
Date: Fri, 26 Jun 2020 08:23:06 +0200

    Merge remote-tracking branch 'origin/master' into feature/add-support-for-emissive-materials
    # Conflicts from merge:
    #	Plugins/GLTFExporter/Source/GLTFExporter/Private/Converters/GLTFMaterialConverters.cpp
    (cherry picked from commit d384d8bb8de869f785f0bd53a5a2adbadf663171​)

#skipci

diff --git a/Engine/Plugins/Enterprise/GLTFExporter/MigratedFromGit/Source/GLTFExporter/Private/Converters/GLTFMaterialConverters.cpp b/Engine/Plugins/Enterprise/GLTFExporter/MigratedFromGit/Source/GLTFExporter/Private/Converters/GLTFMaterialConverters.cpp
--- a/Engine/Plugins/Enterprise/GLTFExporter/MigratedFromGit/Source/GLTFExporter/Private/Converters/GLTFMaterialConverters.cpp
+++ b/Engine/Plugins/Enterprise/GLTFExporter/MigratedFromGit/Source/GLTFExporter/Private/Converters/GLTFMaterialConverters.cpp
@@ -82,7 +82,7 @@ FGLTFJsonMaterialIndex FGLTFMaterialConverter::Add(FGLTFConvertBuilder& Builder,
 
 	if (!TryGetEmissive(Builder, JsonMaterial, Material))
 	{
-		// TODO: handle failure?
+		Builder.AddWarningMessage(FString::Printf(TEXT("Failed to export EmissiveColor for material %s"), *Material->GetName()));
 	}
 
 	if (IsPropertyNonDefault(MP_Normal, Material))
