[GLTFExporter] refactor(json): add GLTFJsonUtility::WriteStringArray and generalize GLTFJsonUtility::WriteObjectArray

commit feec0ba2204cf0e0b2261b46d52b31da4ac7025b
Date: Wed, 24 Jun 2020 00:26:22 +0200

    refactor(json): add GLTFJsonUtility::WriteStringArray and generalize GLTFJsonUtility::WriteObjectArray
    (cherry picked from commit 11b34b26641f16e2843c63454f36d883de6f6c7f​)

#skipci

diff --git a/Engine/Plugins/Enterprise/GLTFExporter/MigratedFromGit/Source/GLTFExporter/Private/Json/GLTFJsonUtility.h b/Engine/Plugins/Enterprise/GLTFExporter/MigratedFromGit/Source/GLTFExporter/Private/Json/GLTFJsonUtility.h
--- a/Engine/Plugins/Enterprise/GLTFExporter/MigratedFromGit/Source/GLTFExporter/Private/Json/GLTFJsonUtility.h
+++ b/Engine/Plugins/Enterprise/GLTFExporter/MigratedFromGit/Source/GLTFExporter/Private/Json/GLTFJsonUtility.h
@@ -74,15 +74,29 @@ struct FGLTFJsonUtility
 		JsonWriter.WriteRawJSONValue(Identifier, ExactStringRepresentation);
 	}
 
-	template <class WriterType, class ElementType, class AllocatorType>
-	static void WriteObjectArray(WriterType& JsonWriter, const FString& Identifier, const TArray<ElementType, AllocatorType>& ObjectArray, bool bWriteIfEmpty = false)
+	template <class WriterType, class ContainerType>
+	static void WriteObjectArray(WriterType& JsonWriter, const FString& Identifier, const ContainerType& Container, bool bWriteIfEmpty = false)
 	{
-		if (ObjectArray.Num() > 0 || bWriteIfEmpty)
+		if (Container.Num() > 0 || bWriteIfEmpty)
 		{
 			JsonWriter.WriteArrayStart(Identifier);
-			for (const ElementType& Object : ObjectArray)
+			for (const auto& Element : Container)
 			{
-				Object.WriteObject(JsonWriter);
+				Element.WriteObject(JsonWriter);
+			}
+			JsonWriter.WriteArrayEnd();
+		}
+	}
+
+	template <class WriterType, class ContainerType>
+    static void WriteStringArray(WriterType& JsonWriter, const FString& Identifier, const ContainerType& Container, bool bWriteIfEmpty = false)
+	{
+		if (Container.Num() > 0 || bWriteIfEmpty)
+		{
+			JsonWriter.WriteArrayStart(Identifier);
+			for (const auto& Element : Container)
+			{
+				JsonWriter.WriteValue(ToString(Element));
 			}
 			JsonWriter.WriteArrayEnd();
 		}
