[GLTFExporter] refactor(json): use GLTFJsonUtility::WriteStringArray for serializing extensions

commit 35ecb630b0211ef47fee4af0f5bed90e1311b4b0
Date: Wed, 24 Jun 2020 00:27:00 +0200

    refactor(json): use GLTFJsonUtility::WriteStringArray for serializing extensions
    (cherry picked from commit f19355c91c44fa26ccfc5c3a9ff07b2b1fcb80ef​)

#skipci

diff --git a/Engine/Plugins/Enterprise/GLTFExporter/MigratedFromGit/Source/GLTFExporter/Private/Json/GLTFJsonRoot.h b/Engine/Plugins/Enterprise/GLTFExporter/MigratedFromGit/Source/GLTFExporter/Private/Json/GLTFJsonRoot.h
--- a/Engine/Plugins/Enterprise/GLTFExporter/MigratedFromGit/Source/GLTFExporter/Private/Json/GLTFJsonRoot.h
+++ b/Engine/Plugins/Enterprise/GLTFExporter/MigratedFromGit/Source/GLTFExporter/Private/Json/GLTFJsonRoot.h
@@ -75,27 +75,8 @@ struct FGLTFJsonRoot
 		JsonWriter.WriteIdentifierPrefix(TEXT("asset"));
 		Asset.WriteObject(JsonWriter);
 
-		if (ExtensionsUsed.Num() > 0)
-		{
-			JsonWriter.WriteIdentifierPrefix(TEXT("extensionsUsed"));
-			JsonWriter.WriteArrayStart();
-			for (const EGLTFJsonExtension& Extension : ExtensionsUsed)
-			{
-				JsonWriter.WriteValue(FGLTFJsonUtility::ToString(Extension));
-			}
-			JsonWriter.WriteArrayEnd();
-		}
-
-		if (ExtensionsRequired.Num() > 0)
-		{
-			JsonWriter.WriteIdentifierPrefix(TEXT("extensionsRequired"));
-			JsonWriter.WriteArrayStart();
-			for (const EGLTFJsonExtension& Extension : ExtensionsRequired)
-			{
-				JsonWriter.WriteValue(FGLTFJsonUtility::ToString(Extension));
-			}
-			JsonWriter.WriteArrayEnd();
-		}
+		FGLTFJsonUtility::WriteStringArray(JsonWriter, TEXT("extensionsUsed"), ExtensionsUsed);
+		FGLTFJsonUtility::WriteStringArray(JsonWriter, TEXT("extensionsRequired"), ExtensionsRequired);
 
 		if (DefaultScene != INDEX_NONE)
 		{
