[GLTFExporter] refactor(json): add GLTFJsonExtensions

commit 195b8c549bf47b79687cdb068631673491a0ee73
Date: Wed, 24 Jun 2020 00:41:12 +0200

    refactor(json): add GLTFJsonExtensions
    (cherry picked from commit 05748de4c652c66c65dc092f0fdfbd352cdb74cb​)

#skipci

diff --git a/Engine/Plugins/Enterprise/GLTFExporter/MigratedFromGit/Source/GLTFExporter/Private/Json/GLTFJsonExtensions.h b/Engine/Plugins/Enterprise/GLTFExporter/MigratedFromGit/Source/GLTFExporter/Private/Json/GLTFJsonExtensions.h
new file mode 100644
--- /dev/null
+++ b/Engine/Plugins/Enterprise/GLTFExporter/MigratedFromGit/Source/GLTFExporter/Private/Json/GLTFJsonExtensions.h
@@ -0,0 +1,11 @@
+// Copyright Epic Games, Inc. All Rights Reserved.
+
+#pragma once
+
+#include "Json/GLTFJsonEnums.h"
+
+struct FGLTFJsonExtensions
+{
+	TSet<EGLTFJsonExtension> Used;
+	TSet<EGLTFJsonExtension> Required;
+};
diff --git a/Engine/Plugins/Enterprise/GLTFExporter/MigratedFromGit/Source/GLTFExporter/Private/Json/GLTFJsonRoot.h b/Engine/Plugins/Enterprise/GLTFExporter/MigratedFromGit/Source/GLTFExporter/Private/Json/GLTFJsonRoot.h
--- a/Engine/Plugins/Enterprise/GLTFExporter/MigratedFromGit/Source/GLTFExporter/Private/Json/GLTFJsonRoot.h
+++ b/Engine/Plugins/Enterprise/GLTFExporter/MigratedFromGit/Source/GLTFExporter/Private/Json/GLTFJsonRoot.h
@@ -3,6 +3,7 @@
 #pragma once
 
 #include "Json/GLTFJsonIndex.h"
+#include "Json/GLTFJsonExtensions.h"
 #include "Json/GLTFJsonAccessor.h"
 #include "Json/GLTFJsonBuffer.h"
 #include "Json/GLTFJsonBufferView.h"
@@ -13,8 +14,6 @@
 #include "Json/GLTFJsonSampler.h"
 #include "Json/GLTFJsonScene.h"
 #include "Json/GLTFJsonTexture.h"
-
-#include "Containers/Set.h"
 #include "Policies/CondensedJsonPrintPolicy.h"
 
 
@@ -51,8 +50,7 @@ struct FGLTFJsonRoot
 {
 	FGLTFJsonAsset Asset;
 
-	TSet<EGLTFJsonExtension> ExtensionsUsed;
-	TSet<EGLTFJsonExtension> ExtensionsRequired;
+	FGLTFJsonExtensions Extensions;
 
 	FGLTFJsonSceneIndex DefaultScene;
 
@@ -70,14 +68,13 @@ struct FGLTFJsonRoot
 	template <class CharType = TCHAR, class PrintPolicy = TPrettyJsonPrintPolicy<CharType>>
 	void WriteObject(TJsonWriter<CharType, PrintPolicy>& JsonWriter) const
 	{
+		FGLTFJsonExtensions AllExtensions(Extensions);
+
 		JsonWriter.WriteObjectStart();
 
 		JsonWriter.WriteIdentifierPrefix(TEXT("asset"));
 		Asset.WriteObject(JsonWriter);
 
-		FGLTFJsonUtility::WriteStringArray(JsonWriter, TEXT("extensionsUsed"), ExtensionsUsed);
-		FGLTFJsonUtility::WriteStringArray(JsonWriter, TEXT("extensionsRequired"), ExtensionsRequired);
-
 		if (DefaultScene != INDEX_NONE)
 		{
 			JsonWriter.WriteValue(TEXT("scene"), DefaultScene);
@@ -94,6 +91,9 @@ struct FGLTFJsonRoot
 		FGLTFJsonUtility::WriteObjectArray(JsonWriter, TEXT("scenes"), Scenes);
 		FGLTFJsonUtility::WriteObjectArray(JsonWriter, TEXT("textures"), Textures);
 
+		FGLTFJsonUtility::WriteStringArray(JsonWriter, TEXT("extensionsUsed"), AllExtensions.Used);
+		FGLTFJsonUtility::WriteStringArray(JsonWriter, TEXT("extensionsRequired"), AllExtensions.Required);
+
 		JsonWriter.WriteObjectEnd();
 	}
 
