[GLTFExporter] feat(json): add support for each json object to specify used or required extensions

commit 455fe315dea0439a18210b6363716fcfa6f84851
Date: Thu, 25 Jun 2020 08:50:16 +0200

    feat(json): add support for each json object to specify used or required extensions
    (cherry picked from commit a0599dabbff39675530afa3cd568396e4f271dfe​)

#skipci

diff --git a/Engine/Plugins/Enterprise/GLTFExporter/MigratedFromGit/Source/GLTFExporter/Private/Json/GLTFJsonAccessor.h b/Engine/Plugins/Enterprise/GLTFExporter/MigratedFromGit/Source/GLTFExporter/Private/Json/GLTFJsonAccessor.h
--- a/Engine/Plugins/Enterprise/GLTFExporter/MigratedFromGit/Source/GLTFExporter/Private/Json/GLTFJsonAccessor.h
+++ b/Engine/Plugins/Enterprise/GLTFExporter/MigratedFromGit/Source/GLTFExporter/Private/Json/GLTFJsonAccessor.h
@@ -35,7 +35,7 @@ struct FGLTFJsonAccessor
 	}
 
 	template <class CharType = TCHAR, class PrintPolicy = TPrettyJsonPrintPolicy<CharType>>
-	void WriteObject(TJsonWriter<CharType, PrintPolicy>& JsonWriter) const
+	void WriteObject(TJsonWriter<CharType, PrintPolicy>& JsonWriter, FGLTFJsonExtensions& Extensions) const
 	{
 		JsonWriter.WriteObjectStart();
 
diff --git a/Engine/Plugins/Enterprise/GLTFExporter/MigratedFromGit/Source/GLTFExporter/Private/Json/GLTFJsonBuffer.h b/Engine/Plugins/Enterprise/GLTFExporter/MigratedFromGit/Source/GLTFExporter/Private/Json/GLTFJsonBuffer.h
--- a/Engine/Plugins/Enterprise/GLTFExporter/MigratedFromGit/Source/GLTFExporter/Private/Json/GLTFJsonBuffer.h
+++ b/Engine/Plugins/Enterprise/GLTFExporter/MigratedFromGit/Source/GLTFExporter/Private/Json/GLTFJsonBuffer.h
@@ -18,7 +18,7 @@ struct FGLTFJsonBuffer
 	}
 
 	template <class CharType = TCHAR, class PrintPolicy = TPrettyJsonPrintPolicy<CharType>>
-	void WriteObject(TJsonWriter<CharType, PrintPolicy>& JsonWriter) const
+	void WriteObject(TJsonWriter<CharType, PrintPolicy>& JsonWriter, FGLTFJsonExtensions& Extensions) const
 	{
 		JsonWriter.WriteObjectStart();
 
diff --git a/Engine/Plugins/Enterprise/GLTFExporter/MigratedFromGit/Source/GLTFExporter/Private/Json/GLTFJsonBufferView.h b/Engine/Plugins/Enterprise/GLTFExporter/MigratedFromGit/Source/GLTFExporter/Private/Json/GLTFJsonBufferView.h
--- a/Engine/Plugins/Enterprise/GLTFExporter/MigratedFromGit/Source/GLTFExporter/Private/Json/GLTFJsonBufferView.h
+++ b/Engine/Plugins/Enterprise/GLTFExporter/MigratedFromGit/Source/GLTFExporter/Private/Json/GLTFJsonBufferView.h
@@ -29,7 +29,7 @@ struct FGLTFJsonBufferView
 	}
 
 	template <class CharType = TCHAR, class PrintPolicy = TPrettyJsonPrintPolicy<CharType>>
-	void WriteObject(TJsonWriter<CharType, PrintPolicy>& JsonWriter) const
+	void WriteObject(TJsonWriter<CharType, PrintPolicy>& JsonWriter, FGLTFJsonExtensions& Extensions) const
 	{
 		JsonWriter.WriteObjectStart();
 
diff --git a/Engine/Plugins/Enterprise/GLTFExporter/MigratedFromGit/Source/GLTFExporter/Private/Json/GLTFJsonImage.h b/Engine/Plugins/Enterprise/GLTFExporter/MigratedFromGit/Source/GLTFExporter/Private/Json/GLTFJsonImage.h
--- a/Engine/Plugins/Enterprise/GLTFExporter/MigratedFromGit/Source/GLTFExporter/Private/Json/GLTFJsonImage.h
+++ b/Engine/Plugins/Enterprise/GLTFExporter/MigratedFromGit/Source/GLTFExporter/Private/Json/GLTFJsonImage.h
@@ -22,7 +22,7 @@ struct FGLTFJsonImage
 	}
 
 	template <class CharType = TCHAR, class PrintPolicy = TPrettyJsonPrintPolicy<CharType>>
-	void WriteObject(TJsonWriter<CharType, PrintPolicy>& JsonWriter) const
+	void WriteObject(TJsonWriter<CharType, PrintPolicy>& JsonWriter, FGLTFJsonExtensions& Extensions) const
 	{
 		JsonWriter.WriteObjectStart();
 
diff --git a/Engine/Plugins/Enterprise/GLTFExporter/MigratedFromGit/Source/GLTFExporter/Private/Json/GLTFJsonMaterial.h b/Engine/Plugins/Enterprise/GLTFExporter/MigratedFromGit/Source/GLTFExporter/Private/Json/GLTFJsonMaterial.h
--- a/Engine/Plugins/Enterprise/GLTFExporter/MigratedFromGit/Source/GLTFExporter/Private/Json/GLTFJsonMaterial.h
+++ b/Engine/Plugins/Enterprise/GLTFExporter/MigratedFromGit/Source/GLTFExporter/Private/Json/GLTFJsonMaterial.h
@@ -19,7 +19,7 @@ struct FGLTFJsonTextureInfo
 	}
 
 	template <class CharType = TCHAR, class PrintPolicy = TPrettyJsonPrintPolicy<CharType>>
-	void WriteObject(TJsonWriter<CharType, PrintPolicy>& JsonWriter) const
+	void WriteObject(TJsonWriter<CharType, PrintPolicy>& JsonWriter, FGLTFJsonExtensions& Extensions) const
 	{
 		JsonWriter.WriteObjectStart();
 
@@ -44,7 +44,7 @@ struct FGLTFJsonNormalTextureInfo : FGLTFJsonTextureInfo
 	}
 
 	template <class CharType = TCHAR, class PrintPolicy = TPrettyJsonPrintPolicy<CharType>>
-	void WriteObject(TJsonWriter<CharType, PrintPolicy>& JsonWriter) const
+	void WriteObject(TJsonWriter<CharType, PrintPolicy>& JsonWriter, FGLTFJsonExtensions& Extensions) const
 	{
 		JsonWriter.WriteObjectStart();
 
@@ -74,7 +74,7 @@ struct FGLTFJsonOcclusionTextureInfo : FGLTFJsonTextureInfo
 	}
 
 	template <class CharType = TCHAR, class PrintPolicy = TPrettyJsonPrintPolicy<CharType>>
-	void WriteObject(TJsonWriter<CharType, PrintPolicy>& JsonWriter) const
+	void WriteObject(TJsonWriter<CharType, PrintPolicy>& JsonWriter, FGLTFJsonExtensions& Extensions) const
 	{
 		JsonWriter.WriteObjectStart();
 
@@ -111,7 +111,7 @@ struct FGLTFJsonPBRMetallicRoughness
 	}
 
 	template <class CharType = TCHAR, class PrintPolicy = TPrettyJsonPrintPolicy<CharType>>
-	void WriteObject(TJsonWriter<CharType, PrintPolicy>& JsonWriter) const
+	void WriteObject(TJsonWriter<CharType, PrintPolicy>& JsonWriter, FGLTFJsonExtensions& Extensions) const
 	{
 		JsonWriter.WriteObjectStart();
 
@@ -124,7 +124,7 @@ struct FGLTFJsonPBRMetallicRoughness
 		if (BaseColorTexture.Index != INDEX_NONE)
 		{
 			JsonWriter.WriteIdentifierPrefix(TEXT("baseColorTexture"));
-			BaseColorTexture.WriteObject(JsonWriter);
+			BaseColorTexture.WriteObject(JsonWriter, Extensions);
 		}
 
 		if (MetallicFactor != 1)
@@ -140,7 +140,7 @@ struct FGLTFJsonPBRMetallicRoughness
 		if (MetallicRoughnessTexture.Index != INDEX_NONE)
 		{
 			JsonWriter.WriteIdentifierPrefix(TEXT("metallicRoughnessTexture"));
-			MetallicRoughnessTexture.WriteObject(JsonWriter);
+			MetallicRoughnessTexture.WriteObject(JsonWriter, Extensions);
 		}
 
 		JsonWriter.WriteObjectEnd();
@@ -173,7 +173,7 @@ struct FGLTFJsonMaterial
 	}
 
 	template <class CharType = TCHAR, class PrintPolicy = TPrettyJsonPrintPolicy<CharType>>
-	void WriteObject(TJsonWriter<CharType, PrintPolicy>& JsonWriter) const
+	void WriteObject(TJsonWriter<CharType, PrintPolicy>& JsonWriter, FGLTFJsonExtensions& Extensions) const
 	{
 		JsonWriter.WriteObjectStart();
 
@@ -183,24 +183,24 @@ struct FGLTFJsonMaterial
 		}
 
 		JsonWriter.WriteIdentifierPrefix(TEXT("pbrMetallicRoughness"));
-		PBRMetallicRoughness.WriteObject(JsonWriter);
+		PBRMetallicRoughness.WriteObject(JsonWriter, Extensions);
 
 		if (NormalTexture.Index != INDEX_NONE)
 		{
 			JsonWriter.WriteIdentifierPrefix(TEXT("normalTexture"));
-			NormalTexture.WriteObject(JsonWriter);
+			NormalTexture.WriteObject(JsonWriter, Extensions);
 		}
 
 		if (OcclusionTexture.Index != INDEX_NONE)
 		{
 			JsonWriter.WriteIdentifierPrefix(TEXT("occlusionTexture"));
-			OcclusionTexture.WriteObject(JsonWriter);
+			OcclusionTexture.WriteObject(JsonWriter, Extensions);
 		}
 
 		if (EmissiveTexture.Index != INDEX_NONE)
 		{
 			JsonWriter.WriteIdentifierPrefix(TEXT("emissiveTexture"));
-			EmissiveTexture.WriteObject(JsonWriter);
+			EmissiveTexture.WriteObject(JsonWriter, Extensions);
 		}
 
 		if (EmissiveFactor != FGLTFJsonColor3::Black)
diff --git a/Engine/Plugins/Enterprise/GLTFExporter/MigratedFromGit/Source/GLTFExporter/Private/Json/GLTFJsonMesh.h b/Engine/Plugins/Enterprise/GLTFExporter/MigratedFromGit/Source/GLTFExporter/Private/Json/GLTFJsonMesh.h
--- a/Engine/Plugins/Enterprise/GLTFExporter/MigratedFromGit/Source/GLTFExporter/Private/Json/GLTFJsonMesh.h
+++ b/Engine/Plugins/Enterprise/GLTFExporter/MigratedFromGit/Source/GLTFExporter/Private/Json/GLTFJsonMesh.h
@@ -19,7 +19,7 @@ struct FGLTFJsonAttributes
 	FGLTFJsonAccessorIndex Weights0;
 
 	template <class CharType = TCHAR, class PrintPolicy = TPrettyJsonPrintPolicy<CharType>>
-	void WriteObject(TJsonWriter<CharType, PrintPolicy>& JsonWriter) const
+	void WriteObject(TJsonWriter<CharType, PrintPolicy>& JsonWriter, FGLTFJsonExtensions& Extensions) const
 	{
 		JsonWriter.WriteObjectStart();
 
@@ -54,12 +54,12 @@ struct FGLTFJsonPrimitive
 	}
 
 	template <class CharType = TCHAR, class PrintPolicy = TPrettyJsonPrintPolicy<CharType>>
-	void WriteObject(TJsonWriter<CharType, PrintPolicy>& JsonWriter) const
+	void WriteObject(TJsonWriter<CharType, PrintPolicy>& JsonWriter, FGLTFJsonExtensions& Extensions) const
 	{
 		JsonWriter.WriteObjectStart();
 
 		JsonWriter.WriteIdentifierPrefix(TEXT("attributes"));
-		Attributes.WriteObject(JsonWriter);
+		Attributes.WriteObject(JsonWriter, Extensions);
 
 		if (Indices != INDEX_NONE)
 		{
@@ -87,7 +87,7 @@ struct FGLTFJsonMesh
 	TArray<FGLTFJsonPrimitive> Primitives;
 
 	template <class CharType = TCHAR, class PrintPolicy = TPrettyJsonPrintPolicy<CharType>>
-	void WriteObject(TJsonWriter<CharType, PrintPolicy>& JsonWriter) const
+	void WriteObject(TJsonWriter<CharType, PrintPolicy>& JsonWriter, FGLTFJsonExtensions& Extensions) const
 	{
 		JsonWriter.WriteObjectStart();
 
@@ -96,7 +96,7 @@ struct FGLTFJsonMesh
 			JsonWriter.WriteValue(TEXT("name"), Name);
 		}
 
-		FGLTFJsonUtility::WriteObjectArray(JsonWriter, TEXT("primitives"), Primitives, true);
+		FGLTFJsonUtility::WriteObjectArray(JsonWriter, TEXT("primitives"), Primitives, Extensions, true);
 
 		JsonWriter.WriteObjectEnd();
 	}
diff --git a/Engine/Plugins/Enterprise/GLTFExporter/MigratedFromGit/Source/GLTFExporter/Private/Json/GLTFJsonNode.h b/Engine/Plugins/Enterprise/GLTFExporter/MigratedFromGit/Source/GLTFExporter/Private/Json/GLTFJsonNode.h
--- a/Engine/Plugins/Enterprise/GLTFExporter/MigratedFromGit/Source/GLTFExporter/Private/Json/GLTFJsonNode.h
+++ b/Engine/Plugins/Enterprise/GLTFExporter/MigratedFromGit/Source/GLTFExporter/Private/Json/GLTFJsonNode.h
@@ -29,7 +29,7 @@ struct FGLTFJsonNode
 	}
 
 	template <class CharType = TCHAR, class PrintPolicy = TPrettyJsonPrintPolicy<CharType>>
-	void WriteObject(TJsonWriter<CharType, PrintPolicy>& JsonWriter) const
+	void WriteObject(TJsonWriter<CharType, PrintPolicy>& JsonWriter, FGLTFJsonExtensions& Extensions) const
 	{
 		JsonWriter.WriteObjectStart();
 
diff --git a/Engine/Plugins/Enterprise/GLTFExporter/MigratedFromGit/Source/GLTFExporter/Private/Json/GLTFJsonRoot.h b/Engine/Plugins/Enterprise/GLTFExporter/MigratedFromGit/Source/GLTFExporter/Private/Json/GLTFJsonRoot.h
--- a/Engine/Plugins/Enterprise/GLTFExporter/MigratedFromGit/Source/GLTFExporter/Private/Json/GLTFJsonRoot.h
+++ b/Engine/Plugins/Enterprise/GLTFExporter/MigratedFromGit/Source/GLTFExporter/Private/Json/GLTFJsonRoot.h
@@ -26,7 +26,7 @@ struct FGLTFJsonAsset
 	FGLTFJsonAsset();
 
 	template <class CharType = TCHAR, class PrintPolicy = TPrettyJsonPrintPolicy<CharType>>
-	void WriteObject(TJsonWriter<CharType, PrintPolicy>& JsonWriter) const
+	void WriteObject(TJsonWriter<CharType, PrintPolicy>& JsonWriter, FGLTFJsonExtensions& Extensions) const
 	{
 		JsonWriter.WriteObjectStart();
 
@@ -73,23 +73,23 @@ struct FGLTFJsonRoot
 		JsonWriter.WriteObjectStart();
 
 		JsonWriter.WriteIdentifierPrefix(TEXT("asset"));
-		Asset.WriteObject(JsonWriter);
+		Asset.WriteObject(JsonWriter, AllExtensions);
 
 		if (DefaultScene != INDEX_NONE)
 		{
 			JsonWriter.WriteValue(TEXT("scene"), DefaultScene);
 		}
 
-		FGLTFJsonUtility::WriteObjectArray(JsonWriter, TEXT("accessors"), Accessors);
-		FGLTFJsonUtility::WriteObjectArray(JsonWriter, TEXT("buffers"), Buffers);
-		FGLTFJsonUtility::WriteObjectArray(JsonWriter, TEXT("bufferViews"), BufferViews);
-		FGLTFJsonUtility::WriteObjectArray(JsonWriter, TEXT("images"), Images);
-		FGLTFJsonUtility::WriteObjectArray(JsonWriter, TEXT("materials"), Materials);
-		FGLTFJsonUtility::WriteObjectArray(JsonWriter, TEXT("meshes"), Meshes);
-		FGLTFJsonUtility::WriteObjectArray(JsonWriter, TEXT("nodes"), Nodes);
-		FGLTFJsonUtility::WriteObjectArray(JsonWriter, TEXT("samplers"), Samplers);
-		FGLTFJsonUtility::WriteObjectArray(JsonWriter, TEXT("scenes"), Scenes);
-		FGLTFJsonUtility::WriteObjectArray(JsonWriter, TEXT("textures"), Textures);
+		FGLTFJsonUtility::WriteObjectArray(JsonWriter, TEXT("accessors"), Accessors, AllExtensions);
+		FGLTFJsonUtility::WriteObjectArray(JsonWriter, TEXT("buffers"), Buffers, AllExtensions);
+		FGLTFJsonUtility::WriteObjectArray(JsonWriter, TEXT("bufferViews"), BufferViews, AllExtensions);
+		FGLTFJsonUtility::WriteObjectArray(JsonWriter, TEXT("images"), Images, AllExtensions);
+		FGLTFJsonUtility::WriteObjectArray(JsonWriter, TEXT("materials"), Materials, AllExtensions);
+		FGLTFJsonUtility::WriteObjectArray(JsonWriter, TEXT("meshes"), Meshes, AllExtensions);
+		FGLTFJsonUtility::WriteObjectArray(JsonWriter, TEXT("nodes"), Nodes, AllExtensions);
+		FGLTFJsonUtility::WriteObjectArray(JsonWriter, TEXT("samplers"), Samplers, AllExtensions);
+		FGLTFJsonUtility::WriteObjectArray(JsonWriter, TEXT("scenes"), Scenes, AllExtensions);
+		FGLTFJsonUtility::WriteObjectArray(JsonWriter, TEXT("textures"), Textures, AllExtensions);
 
 		FGLTFJsonUtility::WriteStringArray(JsonWriter, TEXT("extensionsUsed"), AllExtensions.Used);
 		FGLTFJsonUtility::WriteStringArray(JsonWriter, TEXT("extensionsRequired"), AllExtensions.Required);
diff --git a/Engine/Plugins/Enterprise/GLTFExporter/MigratedFromGit/Source/GLTFExporter/Private/Json/GLTFJsonSampler.h b/Engine/Plugins/Enterprise/GLTFExporter/MigratedFromGit/Source/GLTFExporter/Private/Json/GLTFJsonSampler.h
--- a/Engine/Plugins/Enterprise/GLTFExporter/MigratedFromGit/Source/GLTFExporter/Private/Json/GLTFJsonSampler.h
+++ b/Engine/Plugins/Enterprise/GLTFExporter/MigratedFromGit/Source/GLTFExporter/Private/Json/GLTFJsonSampler.h
@@ -25,7 +25,7 @@ struct FGLTFJsonSampler
 	}
 
 	template <class CharType = TCHAR, class PrintPolicy = TPrettyJsonPrintPolicy<CharType>>
-	void WriteObject(TJsonWriter<CharType, PrintPolicy>& JsonWriter) const
+	void WriteObject(TJsonWriter<CharType, PrintPolicy>& JsonWriter, FGLTFJsonExtensions& Extensions) const
 	{
 		JsonWriter.WriteObjectStart();
 
diff --git a/Engine/Plugins/Enterprise/GLTFExporter/MigratedFromGit/Source/GLTFExporter/Private/Json/GLTFJsonScene.h b/Engine/Plugins/Enterprise/GLTFExporter/MigratedFromGit/Source/GLTFExporter/Private/Json/GLTFJsonScene.h
--- a/Engine/Plugins/Enterprise/GLTFExporter/MigratedFromGit/Source/GLTFExporter/Private/Json/GLTFJsonScene.h
+++ b/Engine/Plugins/Enterprise/GLTFExporter/MigratedFromGit/Source/GLTFExporter/Private/Json/GLTFJsonScene.h
@@ -12,7 +12,7 @@ struct FGLTFJsonScene
 	TArray<FGLTFJsonNodeIndex> Nodes;
 
 	template <class CharType = TCHAR, class PrintPolicy = TPrettyJsonPrintPolicy<CharType>>
-	void WriteObject(TJsonWriter<CharType, PrintPolicy>& JsonWriter) const
+	void WriteObject(TJsonWriter<CharType, PrintPolicy>& JsonWriter, FGLTFJsonExtensions& Extensions) const
 	{
 		JsonWriter.WriteObjectStart();
 
diff --git a/Engine/Plugins/Enterprise/GLTFExporter/MigratedFromGit/Source/GLTFExporter/Private/Json/GLTFJsonTexture.h b/Engine/Plugins/Enterprise/GLTFExporter/MigratedFromGit/Source/GLTFExporter/Private/Json/GLTFJsonTexture.h
--- a/Engine/Plugins/Enterprise/GLTFExporter/MigratedFromGit/Source/GLTFExporter/Private/Json/GLTFJsonTexture.h
+++ b/Engine/Plugins/Enterprise/GLTFExporter/MigratedFromGit/Source/GLTFExporter/Private/Json/GLTFJsonTexture.h
@@ -14,7 +14,7 @@ struct FGLTFJsonTexture
 	FGLTFJsonImageIndex Source;
 
 	template <class CharType = TCHAR, class PrintPolicy = TPrettyJsonPrintPolicy<CharType>>
-	void WriteObject(TJsonWriter<CharType, PrintPolicy>& JsonWriter) const
+	void WriteObject(TJsonWriter<CharType, PrintPolicy>& JsonWriter, FGLTFJsonExtensions& Extensions) const
 	{
 		JsonWriter.WriteObjectStart();
 
diff --git a/Engine/Plugins/Enterprise/GLTFExporter/MigratedFromGit/Source/GLTFExporter/Private/Json/GLTFJsonUtility.h b/Engine/Plugins/Enterprise/GLTFExporter/MigratedFromGit/Source/GLTFExporter/Private/Json/GLTFJsonUtility.h
--- a/Engine/Plugins/Enterprise/GLTFExporter/MigratedFromGit/Source/GLTFExporter/Private/Json/GLTFJsonUtility.h
+++ b/Engine/Plugins/Enterprise/GLTFExporter/MigratedFromGit/Source/GLTFExporter/Private/Json/GLTFJsonUtility.h
@@ -3,6 +3,7 @@
 #pragma once
 
 #include "Json/GLTFJsonEnums.h"
+#include "Json/GLTFJsonExtensions.h"
 
 struct FGLTFJsonUtility
 {
@@ -75,14 +76,14 @@ struct FGLTFJsonUtility
 	}
 
 	template <class WriterType, class ContainerType>
-	static void WriteObjectArray(WriterType& JsonWriter, const FString& Identifier, const ContainerType& Container, bool bWriteIfEmpty = false)
+    static void WriteObjectArray(WriterType& JsonWriter, const FString& Identifier, const ContainerType& Container, FGLTFJsonExtensions& Extensions, bool bWriteIfEmpty = false)
 	{
 		if (Container.Num() > 0 || bWriteIfEmpty)
 		{
 			JsonWriter.WriteArrayStart(Identifier);
 			for (const auto& Element : Container)
 			{
-				Element.WriteObject(JsonWriter);
+				Element.WriteObject(JsonWriter, Extensions);
 			}
 			JsonWriter.WriteArrayEnd();
 		}
