[GLTFExporter] feat(material): add support for unlit shading model

commit a3a17935d3076a49b48761b3caa998919a7890ba
Date: Thu, 25 Jun 2020 08:53:16 +0200

    feat(material): add support for unlit shading model
    (cherry picked from commit 3c75a29c7a80f4fa53bd7ed6911a1efeb4ac4e97​)
    # Conflicts from cherry-pick:
    #	Plugins/GLTFExporter/Source/GLTFExporter/Private/Converters/GLTFMaterialConverters.cpp

#skipci

diff --git a/Engine/Plugins/Enterprise/GLTFExporter/MigratedFromGit/Source/GLTFExporter/Private/Converters/GLTFConverterUtility.h b/Engine/Plugins/Enterprise/GLTFExporter/MigratedFromGit/Source/GLTFExporter/Private/Converters/GLTFConverterUtility.h
--- a/Engine/Plugins/Enterprise/GLTFExporter/MigratedFromGit/Source/GLTFExporter/Private/Converters/GLTFConverterUtility.h
+++ b/Engine/Plugins/Enterprise/GLTFExporter/MigratedFromGit/Source/GLTFExporter/Private/Converters/GLTFConverterUtility.h
@@ -91,6 +91,17 @@ struct FGLTFConverterUtility
 		return Result;
 	}
 
+	static EGLTFJsonShadingModel ConvertShadingModel(const EMaterialShadingModel ShadingModel)
+	{
+		switch (ShadingModel)
+		{
+			case MSM_Unlit:      return EGLTFJsonShadingModel::Unlit;
+			case MSM_DefaultLit: return EGLTFJsonShadingModel::Default;
+			case MSM_ClearCoat:  return EGLTFJsonShadingModel::ClearCoat;
+			default:             return EGLTFJsonShadingModel::None;
+		}
+	}
+
 	static EGLTFJsonAlphaMode ConvertBlendMode(const EBlendMode Mode)
 	{
 		switch (Mode)
diff --git a/Engine/Plugins/Enterprise/GLTFExporter/MigratedFromGit/Source/GLTFExporter/Private/Converters/GLTFMaterialConverters.cpp b/Engine/Plugins/Enterprise/GLTFExporter/MigratedFromGit/Source/GLTFExporter/Private/Converters/GLTFMaterialConverters.cpp
--- a/Engine/Plugins/Enterprise/GLTFExporter/MigratedFromGit/Source/GLTFExporter/Private/Converters/GLTFMaterialConverters.cpp
+++ b/Engine/Plugins/Enterprise/GLTFExporter/MigratedFromGit/Source/GLTFExporter/Private/Converters/GLTFMaterialConverters.cpp
@@ -50,6 +50,15 @@ FGLTFJsonMaterialIndex FGLTFMaterialConverter::Add(FGLTFConvertBuilder& Builder,
 	JsonMaterial.AlphaCutoff = Material->GetOpacityMaskClipValue();
 	JsonMaterial.DoubleSided = Material->IsTwoSided();
 
+	if (!TryGetShadingModel(JsonMaterial.ShadingModel, Material))
+	{
+		// TODO: handle failure?
+	}
+
+	const FLinearColor RgbaMask(1.0f, 1.0f, 1.0f, 1.0f);
+	const FLinearColor RgbMask(1.0f, 1.0f, 1.0f, 0.0f);
+	const FLinearColor RMask(1.0f, 0.0f, 0.0f, 0.0f);
+
 	// TODO: check if a property is active before trying to get it (i.e. Material->IsPropertyActive)
 
 	if (JsonMaterial.AlphaMode == EGLTFJsonAlphaMode::Opaque)
@@ -110,6 +119,35 @@ FGLTFJsonMaterialIndex FGLTFMaterialConverter::Add(FGLTFConvertBuilder& Builder,
 	return Builder.AddMaterial(JsonMaterial);
 }
 
+bool FGLTFMaterialConverter::TryGetShadingModel(EGLTFJsonShadingModel& OutShadingModel, const UMaterialInterface* Material) const
+{
+	const FMaterialShadingModelField ShadingModels = Material->GetShadingModels();
+	const int32 ShadingModelCount = ShadingModels.CountShadingModels();
+
+	if (ShadingModelCount <= 0)
+	{
+		// TODO: report missing shading model
+		return false;
+	}
+
+	if (ShadingModelCount > 1)
+	{
+		// TODO: report support limited to first shading model
+	}
+
+	const EMaterialShadingModel ShadingModel = ShadingModels.GetFirstShadingModel();
+	const EGLTFJsonShadingModel ConvertedShadingModel = FGLTFConverterUtility::ConvertShadingModel(ShadingModel);
+
+	if (ConvertedShadingModel == EGLTFJsonShadingModel::None)
+	{
+		// TODO: report unsupported shading model
+		return false;
+	}
+
+	OutShadingModel = ConvertedShadingModel;
+	return true;
+}
+
 bool FGLTFMaterialConverter::TryGetBaseColorAndOpacity(FGLTFConvertBuilder& Builder, FGLTFJsonPBRMetallicRoughness& OutPBRParams, const UMaterialInterface* Material) const
 {
 	const EMaterialProperty OpacityProperty = Material->GetBlendMode() == BLEND_Masked ? MP_OpacityMask : MP_Opacity;
diff --git a/Engine/Plugins/Enterprise/GLTFExporter/MigratedFromGit/Source/GLTFExporter/Private/Converters/GLTFMaterialConverters.h b/Engine/Plugins/Enterprise/GLTFExporter/MigratedFromGit/Source/GLTFExporter/Private/Converters/GLTFMaterialConverters.h
--- a/Engine/Plugins/Enterprise/GLTFExporter/MigratedFromGit/Source/GLTFExporter/Private/Converters/GLTFMaterialConverters.h
+++ b/Engine/Plugins/Enterprise/GLTFExporter/MigratedFromGit/Source/GLTFExporter/Private/Converters/GLTFMaterialConverters.h
@@ -13,6 +13,8 @@ class FGLTFMaterialConverter final : public TGLTFConverter<FGLTFJsonMaterialInde
 {
 	FGLTFJsonMaterialIndex Add(FGLTFConvertBuilder& Builder, const FString& Name, const UMaterialInterface* Material) override;
 
+	bool TryGetShadingModel(EGLTFJsonShadingModel& ShadingModel, const UMaterialInterface* Material) const;
+
 	bool TryGetBaseColorAndOpacity(FGLTFConvertBuilder& Builder, FGLTFJsonPBRMetallicRoughness& OutPBRParams, const UMaterialInterface* Material) const;
 	bool TryGetMetallicAndRoughness(FGLTFConvertBuilder& Builder, FGLTFJsonPBRMetallicRoughness& OutPBRParams, const UMaterialInterface* Material) const;
 	bool TryGetEmissive(FGLTFConvertBuilder& Builder, FGLTFJsonMaterial& JsonMaterial, const UMaterialInterface* Material) const;
diff --git a/Engine/Plugins/Enterprise/GLTFExporter/MigratedFromGit/Source/GLTFExporter/Private/Json/GLTFJsonEnums.h b/Engine/Plugins/Enterprise/GLTFExporter/MigratedFromGit/Source/GLTFExporter/Private/Json/GLTFJsonEnums.h
--- a/Engine/Plugins/Enterprise/GLTFExporter/MigratedFromGit/Source/GLTFExporter/Private/Json/GLTFJsonEnums.h
+++ b/Engine/Plugins/Enterprise/GLTFExporter/MigratedFromGit/Source/GLTFExporter/Private/Json/GLTFJsonEnums.h
@@ -12,6 +12,14 @@ enum class EGLTFJsonExtension
 	KHR_MeshQuantization
 };
 
+enum class EGLTFJsonShadingModel
+{
+	None = -1,
+	Default,
+	Unlit,
+    ClearCoat
+};
+
 enum class EGLTFJsonAccessorType
 {
 	None = -1,
diff --git a/Engine/Plugins/Enterprise/GLTFExporter/MigratedFromGit/Source/GLTFExporter/Private/Json/GLTFJsonMaterial.h b/Engine/Plugins/Enterprise/GLTFExporter/MigratedFromGit/Source/GLTFExporter/Private/Json/GLTFJsonMaterial.h
--- a/Engine/Plugins/Enterprise/GLTFExporter/MigratedFromGit/Source/GLTFExporter/Private/Json/GLTFJsonMaterial.h
+++ b/Engine/Plugins/Enterprise/GLTFExporter/MigratedFromGit/Source/GLTFExporter/Private/Json/GLTFJsonMaterial.h
@@ -151,6 +151,8 @@ struct FGLTFJsonMaterial
 {
 	FString Name;
 
+	EGLTFJsonShadingModel ShadingModel;
+
 	FGLTFJsonPBRMetallicRoughness PBRMetallicRoughness;
 
 	FGLTFJsonNormalTextureInfo NormalTexture;
@@ -165,7 +167,8 @@ struct FGLTFJsonMaterial
 	bool DoubleSided;
 
 	FGLTFJsonMaterial()
-		: EmissiveFactor(FGLTFJsonColor3::Black)
+		: ShadingModel(EGLTFJsonShadingModel::Default)
+		, EmissiveFactor(FGLTFJsonColor3::Black)
 		, AlphaMode(EGLTFJsonAlphaMode::Opaque)
 		, AlphaCutoff(0.5f)
 		, DoubleSided(false)
@@ -182,8 +185,11 @@ struct FGLTFJsonMaterial
 			JsonWriter.WriteValue(TEXT("name"), Name);
 		}
 
-		JsonWriter.WriteIdentifierPrefix(TEXT("pbrMetallicRoughness"));
-		PBRMetallicRoughness.WriteObject(JsonWriter, Extensions);
+		if (ShadingModel != EGLTFJsonShadingModel::None)
+		{
+			JsonWriter.WriteIdentifierPrefix(TEXT("pbrMetallicRoughness"));
+			PBRMetallicRoughness.WriteObject(JsonWriter, Extensions);
+		}
 
 		if (NormalTexture.Index != INDEX_NONE)
 		{
@@ -224,6 +230,34 @@ struct FGLTFJsonMaterial
 			JsonWriter.WriteValue(TEXT("doubleSided"), DoubleSided);
 		}
 
+		if (ShadingModel == EGLTFJsonShadingModel::Unlit || ShadingModel == EGLTFJsonShadingModel::ClearCoat)
+		{
+			const EGLTFJsonExtension Extension = ShadingModel == EGLTFJsonShadingModel::Unlit ?
+				EGLTFJsonExtension::KHR_MaterialsUnlit : EGLTFJsonExtension::KHR_MaterialsClearCoat;
+
+			Extensions.Used.Add(Extension);
+
+			JsonWriter.WriteObjectStart(TEXT("extensions"));
+			JsonWriter.WriteIdentifierPrefix(FGLTFJsonUtility::ToString(Extension));
+
+			switch (ShadingModel)
+			{
+				case EGLTFJsonShadingModel::Unlit:
+					JsonWriter.WriteObjectStart(); // Write empty object
+					JsonWriter.WriteObjectEnd();
+					break;
+
+				case EGLTFJsonShadingModel::ClearCoat:
+					// TODO: add support
+					break;
+
+				default:
+					checkNoEntry();
+			}
+
+			JsonWriter.WriteObjectEnd();
+		}
+
 		JsonWriter.WriteObjectEnd();
 	}
 };
