[GLTFExporter] feat(material): add proper warning messages when determining shading model

commit c6a4e9b7fc8f130bad16963473679a0c156f910a
Date: Thu, 25 Jun 2020 20:03:26 +0200

    feat(material): add proper warning messages when determining shading model
    (cherry picked from commit 06041b4be56514b6ec519ea9f25b309a3733f4a8​)

#skipci

diff --git a/Engine/Plugins/Enterprise/GLTFExporter/MigratedFromGit/Source/GLTFExporter/Private/Converters/GLTFConverterUtility.h b/Engine/Plugins/Enterprise/GLTFExporter/MigratedFromGit/Source/GLTFExporter/Private/Converters/GLTFConverterUtility.h
--- a/Engine/Plugins/Enterprise/GLTFExporter/MigratedFromGit/Source/GLTFExporter/Private/Converters/GLTFConverterUtility.h
+++ b/Engine/Plugins/Enterprise/GLTFExporter/MigratedFromGit/Source/GLTFExporter/Private/Converters/GLTFConverterUtility.h
@@ -163,4 +163,12 @@ struct FGLTFConverterUtility
 	static bool IsHDRIBackdropBlueprint(const UBlueprint* Blueprint);
 
 	static bool IsSelected(const UActorComponent* ActorComponent);
+
+	template <typename EnumType>
+    static FString GetEnumDisplayName(EnumType Value)
+	{
+		const UEnum* Enum = StaticEnum<EMaterialShadingModel>();
+		check(Enum);
+		return Enum->GetDisplayNameTextByValue(Value).ToString();
+	}
 };
diff --git a/Engine/Plugins/Enterprise/GLTFExporter/MigratedFromGit/Source/GLTFExporter/Private/Converters/GLTFMaterialConverters.cpp b/Engine/Plugins/Enterprise/GLTFExporter/MigratedFromGit/Source/GLTFExporter/Private/Converters/GLTFMaterialConverters.cpp
--- a/Engine/Plugins/Enterprise/GLTFExporter/MigratedFromGit/Source/GLTFExporter/Private/Converters/GLTFMaterialConverters.cpp
+++ b/Engine/Plugins/Enterprise/GLTFExporter/MigratedFromGit/Source/GLTFExporter/Private/Converters/GLTFMaterialConverters.cpp
@@ -43,16 +43,16 @@ FGLTFJsonMaterialIndex FGLTFMaterialConverter::Add(FGLTFConvertBuilder& Builder,
 {
 	FGLTFJsonMaterial JsonMaterial;
 	JsonMaterial.Name = Name;
-	// TODO: add support for different shading models (Default Lit, Unlit, Clear Coat)
 
 	// TODO: add support for additional blend modes (like Additive and Modulate)?
 	JsonMaterial.AlphaMode = FGLTFConverterUtility::ConvertBlendMode(Material->GetBlendMode());
 	JsonMaterial.AlphaCutoff = Material->GetOpacityMaskClipValue();
 	JsonMaterial.DoubleSided = Material->IsTwoSided();
 
-	if (!TryGetShadingModel(JsonMaterial.ShadingModel, Material))
+	if (!TryGetShadingModel(Builder, JsonMaterial.ShadingModel, Material))
 	{
-		// TODO: handle failure?
+		JsonMaterial.ShadingModel = EGLTFJsonShadingModel::Default;
+		Builder.AddWarningMessage(FString::Printf(TEXT("Material %s will be exported as shading model %s"), *Material->GetName(), *FGLTFConverterUtility::GetEnumDisplayName(MSM_DefaultLit)));
 	}
 
 	const FLinearColor RgbaMask(1.0f, 1.0f, 1.0f, 1.0f);
@@ -119,20 +119,20 @@ FGLTFJsonMaterialIndex FGLTFMaterialConverter::Add(FGLTFConvertBuilder& Builder,
 	return Builder.AddMaterial(JsonMaterial);
 }
 
-bool FGLTFMaterialConverter::TryGetShadingModel(EGLTFJsonShadingModel& OutShadingModel, const UMaterialInterface* Material) const
+bool FGLTFMaterialConverter::TryGetShadingModel(FGLTFConvertBuilder& Builder, EGLTFJsonShadingModel& OutShadingModel, const UMaterialInterface* Material) const
 {
 	const FMaterialShadingModelField ShadingModels = Material->GetShadingModels();
 	const int32 ShadingModelCount = ShadingModels.CountShadingModels();
 
 	if (ShadingModelCount <= 0)
 	{
-		// TODO: report missing shading model
+		Builder.AddWarningMessage(FString::Printf(TEXT("No shading model found for material %s"), *Material->GetName()));
 		return false;
 	}
 
 	if (ShadingModelCount > 1)
 	{
-		// TODO: report support limited to first shading model
+		Builder.AddWarningMessage(FString::Printf(TEXT("Support is limited to the first of multiple shading models found (%d) in material %s"), ShadingModelCount, *Material->GetName()));
 	}
 
 	const EMaterialShadingModel ShadingModel = ShadingModels.GetFirstShadingModel();
@@ -140,7 +140,8 @@ bool FGLTFMaterialConverter::TryGetShadingModel(EGLTFJsonShadingModel& OutShadin
 
 	if (ConvertedShadingModel == EGLTFJsonShadingModel::None)
 	{
-		// TODO: report unsupported shading model
+		const FString ShadingModelName = FGLTFConverterUtility::GetEnumDisplayName(ShadingModel);
+		Builder.AddWarningMessage(FString::Printf(TEXT("Unsupported shading model (%s) in material %s"), *ShadingModelName, *Material->GetName()));
 		return false;
 	}
 
@@ -757,6 +758,8 @@ bool FGLTFMaterialConverter::TryGetSourceTexture(const UTexture2D*& OutTexture,
 		return false;
 	}
 
+	// TODO: add support or warning for texture sampler settings that override texture asset addressing (i.e. wrap, clamp etc)?
+
 	if (const UMaterialExpressionTextureSampleParameter2D* TextureParameter = ExactCast<UMaterialExpressionTextureSampleParameter2D>(Expression))
 	{
 		const UMaterialInstance* MaterialInstance = Cast<UMaterialInstance>(Material);
diff --git a/Engine/Plugins/Enterprise/GLTFExporter/MigratedFromGit/Source/GLTFExporter/Private/Converters/GLTFMaterialConverters.h b/Engine/Plugins/Enterprise/GLTFExporter/MigratedFromGit/Source/GLTFExporter/Private/Converters/GLTFMaterialConverters.h
--- a/Engine/Plugins/Enterprise/GLTFExporter/MigratedFromGit/Source/GLTFExporter/Private/Converters/GLTFMaterialConverters.h
+++ b/Engine/Plugins/Enterprise/GLTFExporter/MigratedFromGit/Source/GLTFExporter/Private/Converters/GLTFMaterialConverters.h
@@ -13,7 +13,7 @@ class FGLTFMaterialConverter final : public TGLTFConverter<FGLTFJsonMaterialInde
 {
 	FGLTFJsonMaterialIndex Add(FGLTFConvertBuilder& Builder, const FString& Name, const UMaterialInterface* Material) override;
 
-	bool TryGetShadingModel(EGLTFJsonShadingModel& ShadingModel, const UMaterialInterface* Material) const;
+	bool TryGetShadingModel(FGLTFConvertBuilder& Builder, EGLTFJsonShadingModel& ShadingModel, const UMaterialInterface* Material) const;
 
 	bool TryGetBaseColorAndOpacity(FGLTFConvertBuilder& Builder, FGLTFJsonPBRMetallicRoughness& OutPBRParams, const UMaterialInterface* Material) const;
 	bool TryGetMetallicAndRoughness(FGLTFConvertBuilder& Builder, FGLTFJsonPBRMetallicRoughness& OutPBRParams, const UMaterialInterface* Material) const;
