[GLTFExporter] fix(material): export emissive color as base color when material unlit

commit ac7909bb5dd7ffa60ead8425219aedc7fa3363a9
Date: Thu, 25 Jun 2020 20:07:32 +0200

    fix(material): export emissive color as base color when material unlit
    (cherry picked from commit 2edd8e615c360ec4781a5f56e23bb26a2ddbaa3d​)
    # Conflicts from cherry-pick:
    #	Plugins/GLTFExporter/Source/GLTFExporter/Private/Converters/GLTFMaterialConverters.cpp
    #	Plugins/GLTFExporter/Source/GLTFExporter/Private/Converters/GLTFMaterialConverters.h

#skipci

diff --git a/Engine/Plugins/Enterprise/GLTFExporter/MigratedFromGit/Source/GLTFExporter/Private/Converters/GLTFMaterialConverters.cpp b/Engine/Plugins/Enterprise/GLTFExporter/MigratedFromGit/Source/GLTFExporter/Private/Converters/GLTFMaterialConverters.cpp
--- a/Engine/Plugins/Enterprise/GLTFExporter/MigratedFromGit/Source/GLTFExporter/Private/Converters/GLTFMaterialConverters.cpp
+++ b/Engine/Plugins/Enterprise/GLTFExporter/MigratedFromGit/Source/GLTFExporter/Private/Converters/GLTFMaterialConverters.cpp
@@ -59,59 +59,78 @@ FGLTFJsonMaterialIndex FGLTFMaterialConverter::Add(FGLTFConvertBuilder& Builder,
 	const FLinearColor RgbMask(1.0f, 1.0f, 1.0f, 0.0f);
 	const FLinearColor RMask(1.0f, 0.0f, 0.0f, 0.0f);
 
-	// TODO: check if a property is active before trying to get it (i.e. Material->IsPropertyActive)
-
-	if (JsonMaterial.AlphaMode == EGLTFJsonAlphaMode::Opaque)
+	if (JsonMaterial.ShadingModel != EGLTFJsonShadingModel::None)
 	{
-		if (!TryGetConstantColor(JsonMaterial.PBRMetallicRoughness.BaseColorFactor, MP_BaseColor, Material))
+		const EMaterialProperty BaseColorProperty = JsonMaterial.ShadingModel == EGLTFJsonShadingModel::Unlit ? MP_EmissiveColor : MP_BaseColor;
+		const EMaterialProperty OpacityProperty = JsonMaterial.AlphaMode == EGLTFJsonAlphaMode::Mask ? MP_OpacityMask : MP_Opacity;
+
+		// TODO: check if a property is active before trying to get it (i.e. Material->IsPropertyActive)
+
+		if (JsonMaterial.AlphaMode == EGLTFJsonAlphaMode::Opaque)
 		{
-			if (!TryGetSourceTexture(Builder, JsonMaterial.PBRMetallicRoughness.BaseColorTexture, MP_BaseColor, Material, DefaultColorInputMasks))
+			if (!TryGetConstantColor(JsonMaterial.PBRMetallicRoughness.BaseColorFactor, BaseColorProperty, Material))
 			{
-				if (!TryGetBakedMaterialProperty(Builder, JsonMaterial.PBRMetallicRoughness.BaseColorTexture, JsonMaterial.PBRMetallicRoughness.BaseColorFactor, MP_BaseColor, Material))
+				if (!TryGetSourceTexture(Builder, JsonMaterial.PBRMetallicRoughness.BaseColorTexture, BaseColorProperty, Material, DefaultColorInputMasks))
 				{
-					Builder.AddWarningMessage(FString::Printf(TEXT("Failed to export BaseColor for material %s"), *Material->GetName()));
+					if (!TryGetBakedMaterialProperty(Builder, JsonMaterial.PBRMetallicRoughness.BaseColorTexture, JsonMaterial.PBRMetallicRoughness.BaseColorFactor, BaseColorProperty, Material))
+					{
+						Builder.AddWarningMessage(FString::Printf(TEXT("Failed to export BaseColor for material %s"), *Material->GetName()));
+					}
 				}
 			}
-		}
 
-		JsonMaterial.PBRMetallicRoughness.BaseColorFactor.A = 1.0f; // make sure base color is opaque
-	}
-	else
-	{
-		if (!TryGetBaseColorAndOpacity(Builder, JsonMaterial.PBRMetallicRoughness, Material))
+			JsonMaterial.PBRMetallicRoughness.BaseColorFactor.A = 1.0f; // make sure base color is opaque
+		}
+		else
 		{
-			Builder.AddWarningMessage(FString::Printf(TEXT("Failed to export BaseColor & Opacity for material %s"), *Material->GetName()));
+			if (!TryGetBaseColorAndOpacity(Builder, JsonMaterial.PBRMetallicRoughness, Material, BaseColorProperty, OpacityProperty))
+			{
+				Builder.AddWarningMessage(FString::Printf(TEXT("Failed to export BaseColor & Opacity for material %s"), *Material->GetName()));
+			}
 		}
-	}
 
-	if (!TryGetMetallicAndRoughness(Builder, JsonMaterial.PBRMetallicRoughness, Material))
-	{
-		Builder.AddWarningMessage(FString::Printf(TEXT("Failed to export Metallic & Roughness for material %s"), *Material->GetName()));
-	}
+		if (JsonMaterial.ShadingModel == EGLTFJsonShadingModel::Default || JsonMaterial.ShadingModel == EGLTFJsonShadingModel::ClearCoat)
+		{
+			const EMaterialProperty MetallicProperty = MP_Metallic;
+			const EMaterialProperty RoughnessProperty = MP_Roughness;
 
-	if (!TryGetEmissive(Builder, JsonMaterial, Material))
-	{
-		Builder.AddWarningMessage(FString::Printf(TEXT("Failed to export EmissiveColor for material %s"), *Material->GetName()));
-	}
+			if (!TryGetMetallicAndRoughness(Builder, JsonMaterial.PBRMetallicRoughness, Material, MetallicProperty, RoughnessProperty))
+			{
+				Builder.AddWarningMessage(FString::Printf(TEXT("Failed to export Metallic & Roughness for material %s"), *Material->GetName()));
+			}
 
-	if (IsPropertyNonDefault(MP_Normal, Material))
-	{
-		if (!TryGetSourceTexture(Builder, JsonMaterial.NormalTexture, MP_Normal, Material, DefaultColorInputMasks))
-		{
-			if (!TryGetBakedMaterialProperty(Builder, JsonMaterial.NormalTexture, MP_Normal, Material))
+			if (!TryGetEmissive(Builder, JsonMaterial, Material))
 			{
-				Builder.AddWarningMessage(FString::Printf(TEXT("Failed to export Normal for material %s"), *Material->GetName()));
+				Builder.AddWarningMessage(FString::Printf(TEXT("Failed to export EmissiveColor for material %s"), *Material->GetName()));
 			}
-		}
-	}
 
-	if (IsPropertyNonDefault(MP_AmbientOcclusion, Material))
-	{
-		if (!TryGetSourceTexture(Builder, JsonMaterial.OcclusionTexture, MP_AmbientOcclusion, Material, OcclusionInputMasks))
-		{
-			if (!TryGetBakedMaterialProperty(Builder, JsonMaterial.OcclusionTexture, MP_AmbientOcclusion, Material))
+			const EMaterialProperty NormalProperty = MP_Normal;
+			if (IsPropertyNonDefault(NormalProperty, Material))
 			{
-				Builder.AddWarningMessage(FString::Printf(TEXT("Failed to export AmbientOcclusion for material %s"), *Material->GetName()));
+				if (!TryGetSourceTexture(Builder, JsonMaterial.NormalTexture, NormalProperty, Material, DefaultColorInputMasks))
+				{
+					if (!TryGetBakedMaterialProperty(Builder, JsonMaterial.NormalTexture, NormalProperty, Material))
+					{
+						Builder.AddWarningMessage(FString::Printf(TEXT("Failed to export Normal for material %s"), *Material->GetName()));
+					}
+				}
+			}
+
+			const EMaterialProperty AmbientOcclusionProperty = MP_AmbientOcclusion;
+			if (IsPropertyNonDefault(AmbientOcclusionProperty, Material))
+			{
+				if (!TryGetSourceTexture(Builder, JsonMaterial.OcclusionTexture, AmbientOcclusionProperty, Material, OcclusionInputMasks))
+				{
+					if (!TryGetBakedMaterialProperty(Builder, JsonMaterial.OcclusionTexture, AmbientOcclusionProperty, Material))
+					{
+						Builder.AddWarningMessage(FString::Printf(TEXT("Failed to export AmbientOcclusion for material %s"), *Material->GetName()));
+					}
+				}
+			}
+
+			if (JsonMaterial.ShadingModel == EGLTFJsonShadingModel::ClearCoat)
+			{
+				// TODO: add support for clear coat properties
 			}
 		}
 	}
@@ -149,11 +168,9 @@ bool FGLTFMaterialConverter::TryGetShadingModel(FGLTFConvertBuilder& Builder, EG
 	return true;
 }
 
-bool FGLTFMaterialConverter::TryGetBaseColorAndOpacity(FGLTFConvertBuilder& Builder, FGLTFJsonPBRMetallicRoughness& OutPBRParams, const UMaterialInterface* Material) const
+bool FGLTFMaterialConverter::TryGetBaseColorAndOpacity(FGLTFConvertBuilder& Builder, FGLTFJsonPBRMetallicRoughness& OutPBRParams, const UMaterialInterface* Material, EMaterialProperty BaseColorProperty, EMaterialProperty OpacityProperty) const
 {
-	const EMaterialProperty OpacityProperty = Material->GetBlendMode() == BLEND_Masked ? MP_OpacityMask : MP_Opacity;
-
-	const bool bIsBaseColorConstant = TryGetConstantColor(OutPBRParams.BaseColorFactor, MP_BaseColor, Material);
+	const bool bIsBaseColorConstant = TryGetConstantColor(OutPBRParams.BaseColorFactor, BaseColorProperty, Material);
 	const bool bIsOpacityConstant = TryGetConstantScalar(OutPBRParams.BaseColorFactor.A, OpacityProperty, Material);
 
 	if (bIsBaseColorConstant && bIsOpacityConstant)
@@ -171,7 +188,7 @@ bool FGLTFMaterialConverter::TryGetBaseColorAndOpacity(FGLTFConvertBuilder& Buil
 	int32 BaseColorTexCoord;
 	int32 OpacityTexCoord;
 
-	const bool bHasBaseColorSourceTexture = TryGetSourceTexture(BaseColorTexture, BaseColorTexCoord, MP_BaseColor, Material, BaseColorInputMasks);
+	const bool bHasBaseColorSourceTexture = TryGetSourceTexture(BaseColorTexture, BaseColorTexCoord, BaseColorProperty, Material, BaseColorInputMasks);
 	const bool bHasOpacitySourceTexture = TryGetSourceTexture(OpacityTexture, OpacityTexCoord, OpacityProperty, Material, OpacityInputMasks);
 
 	// Detect the "happy path" where both inputs share the same texture and are correctly masked.
@@ -249,7 +266,7 @@ bool FGLTFMaterialConverter::TryGetBaseColorAndOpacity(FGLTFConvertBuilder& Buil
 		TextureMagFilter = FGLTFConverterUtility::ConvertMagFilter(OpacityTexture->Filter, OpacityTexture->LODGroup);
 	}
 
-	const FGLTFPropertyBakeOutput BaseColorBakeOutput = BakeMaterialProperty(MP_BaseColor, Material, &TextureSize);
+	const FGLTFPropertyBakeOutput BaseColorBakeOutput = BakeMaterialProperty(BaseColorProperty, Material, &TextureSize);
 	FGLTFPropertyBakeOutput OpacityBakeOutput = BakeMaterialProperty(OpacityProperty, Material, &TextureSize, true);
 
 	// Detect when both baked properties are constants, which means we can avoid exporting a texture
@@ -290,10 +307,10 @@ bool FGLTFMaterialConverter::TryGetBaseColorAndOpacity(FGLTFConvertBuilder& Buil
 	return true;
 }
 
-bool FGLTFMaterialConverter::TryGetMetallicAndRoughness(FGLTFConvertBuilder& Builder, FGLTFJsonPBRMetallicRoughness& OutPBRParams, const UMaterialInterface* Material) const
+bool FGLTFMaterialConverter::TryGetMetallicAndRoughness(FGLTFConvertBuilder& Builder, FGLTFJsonPBRMetallicRoughness& OutPBRParams, const UMaterialInterface* Material, EMaterialProperty MetallicProperty, EMaterialProperty RoughnessProperty) const
 {
-	const bool bIsMetallicConstant = TryGetConstantScalar(OutPBRParams.MetallicFactor, MP_Metallic, Material);
-	const bool bIsRoughnessConstant = TryGetConstantScalar(OutPBRParams.RoughnessFactor, MP_Roughness, Material);
+	const bool bIsMetallicConstant = TryGetConstantScalar(OutPBRParams.MetallicFactor, MetallicProperty, Material);
+	const bool bIsRoughnessConstant = TryGetConstantScalar(OutPBRParams.RoughnessFactor, RoughnessProperty, Material);
 
 	if (bIsMetallicConstant && bIsRoughnessConstant)
 	{
@@ -311,8 +328,8 @@ bool FGLTFMaterialConverter::TryGetMetallicAndRoughness(FGLTFConvertBuilder& Bui
 	int32 MetallicTexCoord;
 	int32 RoughnessTexCoord;
 
-	const bool bHasMetallicSourceTexture = TryGetSourceTexture(MetallicTexture, MetallicTexCoord, MP_Metallic, Material, MetallicInputMasks);
-	const bool bHasRoughnessSourceTexture = TryGetSourceTexture(RoughnessTexture, RoughnessTexCoord, MP_Roughness, Material, RoughnessInputMasks);
+	const bool bHasMetallicSourceTexture = TryGetSourceTexture(MetallicTexture, MetallicTexCoord, MetallicProperty, Material, MetallicInputMasks);
+	const bool bHasRoughnessSourceTexture = TryGetSourceTexture(RoughnessTexture, RoughnessTexCoord, RoughnessProperty, Material, RoughnessInputMasks);
 
 	// Detect the "happy path" where both inputs share the same texture and are correctly masked.
 	if (bHasMetallicSourceTexture &&
@@ -389,20 +406,20 @@ bool FGLTFMaterialConverter::TryGetMetallicAndRoughness(FGLTFConvertBuilder& Bui
 		TextureMagFilter = FGLTFConverterUtility::ConvertMagFilter(RoughnessTexture->Filter, RoughnessTexture->LODGroup);
 	}
 
-	const FGLTFPropertyBakeOutput RoughnessBakeOutput = BakeMaterialProperty(MP_Roughness, Material, &TextureSize);
-	const FGLTFPropertyBakeOutput MetallicBakeOutput = BakeMaterialProperty(MP_Metallic, Material, &TextureSize);
+	const FGLTFPropertyBakeOutput MetallicBakeOutput = BakeMaterialProperty(MetallicProperty, Material, &TextureSize);
+	const FGLTFPropertyBakeOutput RoughnessBakeOutput = BakeMaterialProperty(RoughnessProperty, Material, &TextureSize);
 
 	// Detect when both baked properties are constants, which means we can use factors and avoid exporting a texture
-	if (RoughnessBakeOutput.bIsConstant && MetallicBakeOutput.bIsConstant)
+	if (MetallicBakeOutput.bIsConstant && RoughnessBakeOutput.bIsConstant)
 	{
-		OutPBRParams.RoughnessFactor = RoughnessBakeOutput.ConstantValue.R;
 		OutPBRParams.MetallicFactor = MetallicBakeOutput.ConstantValue.R;
+		OutPBRParams.RoughnessFactor = RoughnessBakeOutput.ConstantValue.R;
 		return true;
 	}
 
 	TextureSize = RoughnessBakeOutput.Size.ComponentMax(MetallicBakeOutput.Size);
-	RoughnessTexture = FGLTFMaterialUtility::CreateTransientTexture(RoughnessBakeOutput);
 	MetallicTexture = FGLTFMaterialUtility::CreateTransientTexture(MetallicBakeOutput);
+	RoughnessTexture = FGLTFMaterialUtility::CreateTransientTexture(RoughnessBakeOutput);
 
 	const FString TextureName = Material->GetName() + TEXT("_MetallicRoughness");
 
diff --git a/Engine/Plugins/Enterprise/GLTFExporter/MigratedFromGit/Source/GLTFExporter/Private/Converters/GLTFMaterialConverters.h b/Engine/Plugins/Enterprise/GLTFExporter/MigratedFromGit/Source/GLTFExporter/Private/Converters/GLTFMaterialConverters.h
--- a/Engine/Plugins/Enterprise/GLTFExporter/MigratedFromGit/Source/GLTFExporter/Private/Converters/GLTFMaterialConverters.h
+++ b/Engine/Plugins/Enterprise/GLTFExporter/MigratedFromGit/Source/GLTFExporter/Private/Converters/GLTFMaterialConverters.h
@@ -15,8 +15,8 @@ class FGLTFMaterialConverter final : public TGLTFConverter<FGLTFJsonMaterialInde
 
 	bool TryGetShadingModel(FGLTFConvertBuilder& Builder, EGLTFJsonShadingModel& ShadingModel, const UMaterialInterface* Material) const;
 
-	bool TryGetBaseColorAndOpacity(FGLTFConvertBuilder& Builder, FGLTFJsonPBRMetallicRoughness& OutPBRParams, const UMaterialInterface* Material) const;
-	bool TryGetMetallicAndRoughness(FGLTFConvertBuilder& Builder, FGLTFJsonPBRMetallicRoughness& OutPBRParams, const UMaterialInterface* Material) const;
+	bool TryGetBaseColorAndOpacity(FGLTFConvertBuilder& Builder, FGLTFJsonPBRMetallicRoughness& OutPBRParams, const UMaterialInterface* Material, EMaterialProperty BaseColorProperty, EMaterialProperty OpacityProperty) const;
+	bool TryGetMetallicAndRoughness(FGLTFConvertBuilder& Builder, FGLTFJsonPBRMetallicRoughness& OutPBRParams, const UMaterialInterface* Material, EMaterialProperty MetallicProperty, EMaterialProperty RoughnessProperty) const;
 	bool TryGetEmissive(FGLTFConvertBuilder& Builder, FGLTFJsonMaterial& JsonMaterial, const UMaterialInterface* Material) const;
 
 	bool IsPropertyNonDefault(EMaterialProperty Property, const UMaterialInterface* Material) const;
