[GLTFExporter] fix(material): decouple property from property name so that unlit texture is not incorrectly suffixed with BaseColor instead of Emissive.

commit 5c26fb1b4b91b1db7fda34dff649f13829e18d7d
Date: Thu, 18 Aug 2022 00:25:28 +0200

    fix(material): decouple property from property name so that unlit texture is not incorrectly suffixed with BaseColor instead of Emissive.
    (cherry picked from commit 890053e66943ba68b2e554cbced1da183e2de031​)

#skipci

diff --git a/Engine/Plugins/Enterprise/GLTFExporter/MigratedFromGit/Source/GLTFExporter/Private/Converters/GLTFMaterialConverters.cpp b/Engine/Plugins/Enterprise/GLTFExporter/MigratedFromGit/Source/GLTFExporter/Private/Converters/GLTFMaterialConverters.cpp
--- a/Engine/Plugins/Enterprise/GLTFExporter/MigratedFromGit/Source/GLTFExporter/Private/Converters/GLTFMaterialConverters.cpp
+++ b/Engine/Plugins/Enterprise/GLTFExporter/MigratedFromGit/Source/GLTFExporter/Private/Converters/GLTFMaterialConverters.cpp
@@ -72,7 +72,7 @@ FGLTFJsonMaterialIndex FGLTFMaterialConverter::Add(FGLTFConvertBuilder& Builder,
 			{
 				if (!TryGetSourceTexture(Builder, JsonMaterial.PBRMetallicRoughness.BaseColorTexture, BaseColorProperty, Material, DefaultColorInputMasks))
 				{
-					if (!TryGetBakedMaterialProperty(Builder, JsonMaterial.PBRMetallicRoughness.BaseColorTexture, JsonMaterial.PBRMetallicRoughness.BaseColorFactor, BaseColorProperty, Material))
+					if (!TryGetBakedMaterialProperty(Builder, JsonMaterial.PBRMetallicRoughness.BaseColorTexture, JsonMaterial.PBRMetallicRoughness.BaseColorFactor, BaseColorProperty, TEXT("BaseColor"), Material))
 					{
 						Builder.AddWarningMessage(FString::Printf(TEXT("Failed to export BaseColor for material %s"), *Material->GetName()));
 					}
@@ -109,7 +109,7 @@ FGLTFJsonMaterialIndex FGLTFMaterialConverter::Add(FGLTFConvertBuilder& Builder,
 			{
 				if (!TryGetSourceTexture(Builder, JsonMaterial.NormalTexture, NormalProperty, Material, DefaultColorInputMasks))
 				{
-					if (!TryGetBakedMaterialProperty(Builder, JsonMaterial.NormalTexture, NormalProperty, Material))
+					if (!TryGetBakedMaterialProperty(Builder, JsonMaterial.NormalTexture, NormalProperty, TEXT("Normal"), Material))
 					{
 						Builder.AddWarningMessage(FString::Printf(TEXT("Failed to export Normal for material %s"), *Material->GetName()));
 					}
@@ -121,7 +121,7 @@ FGLTFJsonMaterialIndex FGLTFMaterialConverter::Add(FGLTFConvertBuilder& Builder,
 			{
 				if (!TryGetSourceTexture(Builder, JsonMaterial.OcclusionTexture, AmbientOcclusionProperty, Material, OcclusionInputMasks))
 				{
-					if (!TryGetBakedMaterialProperty(Builder, JsonMaterial.OcclusionTexture, AmbientOcclusionProperty, Material))
+					if (!TryGetBakedMaterialProperty(Builder, JsonMaterial.OcclusionTexture, AmbientOcclusionProperty, TEXT("Occlusion"), Material))
 					{
 						Builder.AddWarningMessage(FString::Printf(TEXT("Failed to export AmbientOcclusion for material %s"), *Material->GetName()));
 					}
@@ -471,7 +471,7 @@ bool FGLTFMaterialConverter::TryGetEmissive(FGLTFConvertBuilder& Builder, FGLTFJ
 	}
 	else
 	{
-		if (!StoreBakedPropertyTexture(Builder, JsonMaterial.EmissiveTexture, PropertyBakeOutput, MP_EmissiveColor, Material))
+		if (!StoreBakedPropertyTexture(Builder, JsonMaterial.EmissiveTexture, PropertyBakeOutput, TEXT("Emissive"), Material))
 		{
 			return false;
 		}
@@ -837,7 +837,7 @@ bool FGLTFMaterialConverter::TryGetSourceTexture(const UTexture2D*& OutTexture,
 	return false;
 }
 
-bool FGLTFMaterialConverter::TryGetBakedMaterialProperty(FGLTFConvertBuilder& Builder, FGLTFJsonTextureInfo& OutTexInfo, FGLTFJsonColor3& OutConstant, EMaterialProperty Property, const UMaterialInterface* Material) const
+bool FGLTFMaterialConverter::TryGetBakedMaterialProperty(FGLTFConvertBuilder& Builder, FGLTFJsonTextureInfo& OutTexInfo, FGLTFJsonColor3& OutConstant, EMaterialProperty Property, const FString& PropertyName, const UMaterialInterface* Material) const
 {
 	const FGLTFPropertyBakeOutput PropertyBakeOutput = BakeMaterialProperty(Property, Material);
 
@@ -848,11 +848,11 @@ bool FGLTFMaterialConverter::TryGetBakedMaterialProperty(FGLTFConvertBuilder& Bu
 	}
 	else
 	{
-		return StoreBakedPropertyTexture(Builder, OutTexInfo, PropertyBakeOutput, Property, Material);
+		return StoreBakedPropertyTexture(Builder, OutTexInfo, PropertyBakeOutput, PropertyName, Material);
 	}
 }
 
-bool FGLTFMaterialConverter::TryGetBakedMaterialProperty(FGLTFConvertBuilder& Builder, FGLTFJsonTextureInfo& OutTexInfo, FGLTFJsonColor4& OutConstant, EMaterialProperty Property, const UMaterialInterface* Material) const
+bool FGLTFMaterialConverter::TryGetBakedMaterialProperty(FGLTFConvertBuilder& Builder, FGLTFJsonTextureInfo& OutTexInfo, FGLTFJsonColor4& OutConstant, EMaterialProperty Property, const FString& PropertyName, const UMaterialInterface* Material) const
 {
 	const FGLTFPropertyBakeOutput PropertyBakeOutput = BakeMaterialProperty(Property, Material);
 
@@ -863,17 +863,17 @@ bool FGLTFMaterialConverter::TryGetBakedMaterialProperty(FGLTFConvertBuilder& Bu
 	}
 	else
 	{
-		return StoreBakedPropertyTexture(Builder, OutTexInfo, PropertyBakeOutput, Property, Material);
+		return StoreBakedPropertyTexture(Builder, OutTexInfo, PropertyBakeOutput, PropertyName, Material);
 	}
 }
 
-bool FGLTFMaterialConverter::TryGetBakedMaterialProperty(FGLTFConvertBuilder& Builder, FGLTFJsonTextureInfo& OutTexInfo, EMaterialProperty Property, const UMaterialInterface* Material) const
+bool FGLTFMaterialConverter::TryGetBakedMaterialProperty(FGLTFConvertBuilder& Builder, FGLTFJsonTextureInfo& OutTexInfo, EMaterialProperty Property, const FString& PropertyName, const UMaterialInterface* Material) const
 {
 	const FGLTFPropertyBakeOutput PropertyBakeOutput = BakeMaterialProperty(Property, Material);
 
 	if (!PropertyBakeOutput.bIsConstant)
 	{
-		return StoreBakedPropertyTexture(Builder, OutTexInfo, PropertyBakeOutput, Property, Material);
+		return StoreBakedPropertyTexture(Builder, OutTexInfo, PropertyBakeOutput, PropertyName, Material);
 	}
 	else
 	{
@@ -909,21 +909,8 @@ FGLTFPropertyBakeOutput FGLTFMaterialConverter::BakeMaterialProperty(EMaterialPr
 	return PropertyBakeOutput;
 }
 
-bool FGLTFMaterialConverter::StoreBakedPropertyTexture(FGLTFConvertBuilder& Builder, FGLTFJsonTextureInfo& OutTexInfo, const FGLTFPropertyBakeOutput& PropertyBakeOutput, EMaterialProperty Property, const UMaterialInterface* Material) const
+bool FGLTFMaterialConverter::StoreBakedPropertyTexture(FGLTFConvertBuilder& Builder, FGLTFJsonTextureInfo& OutTexInfo, const FGLTFPropertyBakeOutput& PropertyBakeOutput, const FString& PropertyName, const UMaterialInterface* Material) const
 {
-	TCHAR* PropertyName;
-
-	switch (Property)
-	{
-		case MP_BaseColor: PropertyName = TEXT("BaseColor"); break;
-		case MP_Normal: PropertyName = TEXT("Normal"); break;
-		case MP_EmissiveColor: PropertyName = TEXT("Emissive"); break;
-		case MP_AmbientOcclusion: PropertyName = TEXT("Occlusion"); break;
-		default:
-			// TODO: support for more properties
-			return false;
-	}
-
 	const FString TextureName = Material->GetName() + TEXT("_") + PropertyName;
 
 	// TODO: add support for detecting the correct tex-coord for this property based on connected nodes
diff --git a/Engine/Plugins/Enterprise/GLTFExporter/MigratedFromGit/Source/GLTFExporter/Private/Converters/GLTFMaterialConverters.h b/Engine/Plugins/Enterprise/GLTFExporter/MigratedFromGit/Source/GLTFExporter/Private/Converters/GLTFMaterialConverters.h
--- a/Engine/Plugins/Enterprise/GLTFExporter/MigratedFromGit/Source/GLTFExporter/Private/Converters/GLTFMaterialConverters.h
+++ b/Engine/Plugins/Enterprise/GLTFExporter/MigratedFromGit/Source/GLTFExporter/Private/Converters/GLTFMaterialConverters.h
@@ -28,11 +28,11 @@ class FGLTFMaterialConverter final : public TGLTFConverter<FGLTFJsonMaterialInde
 	bool TryGetSourceTexture(FGLTFConvertBuilder& Builder, FGLTFJsonTextureInfo& OutTexInfo, EMaterialProperty Property, const UMaterialInterface* Material, const TArray<FLinearColor>& AllowedMasks = {}) const;
 	bool TryGetSourceTexture(const UTexture2D*& OutTexture, int32& OutTexCoord, EMaterialProperty Property, const UMaterialInterface* Material, const TArray<FLinearColor>& AllowedMasks = {}) const;
 
-	bool TryGetBakedMaterialProperty(FGLTFConvertBuilder& Builder, FGLTFJsonTextureInfo& OutTexInfo, FGLTFJsonColor3& OutConstant, EMaterialProperty Property, const UMaterialInterface* Material) const;
-	bool TryGetBakedMaterialProperty(FGLTFConvertBuilder& Builder, FGLTFJsonTextureInfo& OutTexInfo, FGLTFJsonColor4& OutConstant, EMaterialProperty Property, const UMaterialInterface* Material) const;
-	bool TryGetBakedMaterialProperty(FGLTFConvertBuilder& Builder, FGLTFJsonTextureInfo& OutTexInfo, EMaterialProperty Property, const UMaterialInterface* Material) const;
+	bool TryGetBakedMaterialProperty(FGLTFConvertBuilder& Builder, FGLTFJsonTextureInfo& OutTexInfo, FGLTFJsonColor3& OutConstant, EMaterialProperty Property, const FString& PropertyName, const UMaterialInterface* Material) const;
+	bool TryGetBakedMaterialProperty(FGLTFConvertBuilder& Builder, FGLTFJsonTextureInfo& OutTexInfo, FGLTFJsonColor4& OutConstant, EMaterialProperty Property, const FString& PropertyName, const UMaterialInterface* Material) const;
+	bool TryGetBakedMaterialProperty(FGLTFConvertBuilder& Builder, FGLTFJsonTextureInfo& OutTexInfo, EMaterialProperty Property, const FString& PropertyName, const UMaterialInterface* Material) const;
 
 	FGLTFPropertyBakeOutput BakeMaterialProperty(EMaterialProperty Property, const UMaterialInterface* Material, const FIntPoint* PreferredTextureSize = nullptr, bool bCopyAlphaFromRedChannel = false) const;
 
-	bool StoreBakedPropertyTexture(FGLTFConvertBuilder& Builder, FGLTFJsonTextureInfo& OutTexInfo, const FGLTFPropertyBakeOutput& PropertyBakeOutput, EMaterialProperty Property, const UMaterialInterface* Material) const;
+	bool StoreBakedPropertyTexture(FGLTFConvertBuilder& Builder, FGLTFJsonTextureInfo& OutTexInfo, const FGLTFPropertyBakeOutput& PropertyBakeOutput, const FString& PropertyName, const UMaterialInterface* Material) const;
 };
