[GLTFExporter] Merge branch 'master' into feature/unlit-material

commit 8456716b69240f6e0b15c9deebcb564d0fe645ad
Date: Fri, 26 Jun 2020 18:56:56 +0200

    Merge branch 'master' into feature/unlit-material
    # Conflicts from merge:
    #	Plugins/GLTFExporter/Source/GLTFExporter/Private/Converters/GLTFMaterialConverters.cpp
    #	Plugins/GLTFExporter/Source/GLTFExporter/Private/Converters/GLTFMaterialConverters.h
    (cherry picked from commit a3f1eacdf852d086780fc8cbc077a8167fecc66f​)
    # Conflicts from cherry-pick:
    #	Plugins/GLTFExporter/Source/GLTFExporter/Private/Converters/GLTFMaterialConverters.cpp
    #	Plugins/GLTFExporter/Source/GLTFExporter/Private/Converters/GLTFMaterialConverters.h

#skipci

diff --git a/Engine/Plugins/Enterprise/GLTFExporter/MigratedFromGit/Source/GLTFExporter/Private/Converters/GLTFMaterialConverters.cpp b/Engine/Plugins/Enterprise/GLTFExporter/MigratedFromGit/Source/GLTFExporter/Private/Converters/GLTFMaterialConverters.cpp
--- a/Engine/Plugins/Enterprise/GLTFExporter/MigratedFromGit/Source/GLTFExporter/Private/Converters/GLTFMaterialConverters.cpp
+++ b/Engine/Plugins/Enterprise/GLTFExporter/MigratedFromGit/Source/GLTFExporter/Private/Converters/GLTFMaterialConverters.cpp
@@ -55,10 +55,6 @@ FGLTFJsonMaterialIndex FGLTFMaterialConverter::Add(FGLTFConvertBuilder& Builder,
 		Builder.AddWarningMessage(FString::Printf(TEXT("Material %s will be exported as shading model %s"), *Material->GetName(), *FGLTFConverterUtility::GetEnumDisplayName(MSM_DefaultLit)));
 	}
 
-	const FLinearColor RgbaMask(1.0f, 1.0f, 1.0f, 1.0f);
-	const FLinearColor RgbMask(1.0f, 1.0f, 1.0f, 0.0f);
-	const FLinearColor RMask(1.0f, 0.0f, 0.0f, 0.0f);
-
 	if (JsonMaterial.ShadingModel != EGLTFJsonShadingModel::None)
 	{
 		const EMaterialProperty BaseColorProperty = JsonMaterial.ShadingModel == EGLTFJsonShadingModel::Unlit ? MP_EmissiveColor : MP_BaseColor;
@@ -99,7 +95,8 @@ FGLTFJsonMaterialIndex FGLTFMaterialConverter::Add(FGLTFConvertBuilder& Builder,
 				Builder.AddWarningMessage(FString::Printf(TEXT("Failed to export Metallic & Roughness for material %s"), *Material->GetName()));
 			}
 
-			if (!TryGetEmissive(Builder, JsonMaterial, Material))
+			const EMaterialProperty EmissiveProperty = MP_EmissiveColor;
+			if (!TryGetEmissive(Builder, JsonMaterial, EmissiveProperty, Material))
 			{
 				Builder.AddWarningMessage(FString::Printf(TEXT("Failed to export EmissiveColor for material %s"), *Material->GetName()));
 			}
@@ -267,12 +264,13 @@ bool FGLTFMaterialConverter::TryGetBaseColorAndOpacity(FGLTFConvertBuilder& Buil
 	}
 
 	const FGLTFPropertyBakeOutput BaseColorBakeOutput = BakeMaterialProperty(BaseColorProperty, Material, &TextureSize);
-	FGLTFPropertyBakeOutput OpacityBakeOutput = BakeMaterialProperty(OpacityProperty, Material, &TextureSize, true);
+	const FGLTFPropertyBakeOutput OpacityBakeOutput = BakeMaterialProperty(OpacityProperty, Material, &TextureSize, true);
+	const float BaseColorScale = BaseColorProperty == MP_EmissiveColor ? BaseColorBakeOutput.EmissiveScale : 1;
 
 	// Detect when both baked properties are constants, which means we can avoid exporting a texture
 	if (BaseColorBakeOutput.bIsConstant && OpacityBakeOutput.bIsConstant)
 	{
-		FLinearColor BaseColorFactor(BaseColorBakeOutput.ConstantValue);
+		FLinearColor BaseColorFactor(BaseColorBakeOutput.ConstantValue * BaseColorScale);
 		BaseColorFactor.A = OpacityBakeOutput.ConstantValue.A;
 
 		OutPBRParams.BaseColorFactor = FGLTFConverterUtility::ConvertColor(BaseColorFactor);
@@ -303,6 +301,7 @@ bool FGLTFMaterialConverter::TryGetBaseColorAndOpacity(FGLTFConvertBuilder& Buil
 
 	OutPBRParams.BaseColorTexture.TexCoord = TexCoord;
 	OutPBRParams.BaseColorTexture.Index = TextureIndex;
+	OutPBRParams.BaseColorFactor = { BaseColorScale, BaseColorScale, BaseColorScale };
 
 	return true;
 }
@@ -445,7 +444,7 @@ bool FGLTFMaterialConverter::TryGetMetallicAndRoughness(FGLTFConvertBuilder& Bui
 	return true;
 }
 
-bool FGLTFMaterialConverter::TryGetEmissive(FGLTFConvertBuilder& Builder, FGLTFJsonMaterial& JsonMaterial, const UMaterialInterface* Material) const
+bool FGLTFMaterialConverter::TryGetEmissive(FGLTFConvertBuilder& Builder, FGLTFJsonMaterial& JsonMaterial, EMaterialProperty EmissiveProperty, const UMaterialInterface* Material) const
 {
 	// TODO: right now we allow EmissiveFactor to be > 1.0 to support very bright emission, although it's not valid according to the glTF standard.
 	// We may want to change this behaviour and store factors above 1.0 using a custom extension instead.
@@ -455,13 +454,13 @@ bool FGLTFMaterialConverter::TryGetEmissive(FGLTFConvertBuilder& Builder, FGLTFJ
 		return true;
 	}
 
-	if (TryGetSourceTexture(Builder, JsonMaterial.EmissiveTexture, MP_EmissiveColor, Material, DefaultColorInputMasks))
+	if (TryGetSourceTexture(Builder, JsonMaterial.EmissiveTexture, EmissiveProperty, Material, DefaultColorInputMasks))
 	{
 		JsonMaterial.EmissiveFactor = FGLTFJsonColor3::White;	// make sure texture is not multiplied with black
 		return true;
 	}
 
-	const FGLTFPropertyBakeOutput PropertyBakeOutput = BakeMaterialProperty(MP_EmissiveColor, Material);
+	const FGLTFPropertyBakeOutput PropertyBakeOutput = BakeMaterialProperty(EmissiveProperty, Material);
 	const float EmissiveScale = PropertyBakeOutput.EmissiveScale;
 
 	if (PropertyBakeOutput.bIsConstant)
diff --git a/Engine/Plugins/Enterprise/GLTFExporter/MigratedFromGit/Source/GLTFExporter/Private/Converters/GLTFMaterialConverters.h b/Engine/Plugins/Enterprise/GLTFExporter/MigratedFromGit/Source/GLTFExporter/Private/Converters/GLTFMaterialConverters.h
--- a/Engine/Plugins/Enterprise/GLTFExporter/MigratedFromGit/Source/GLTFExporter/Private/Converters/GLTFMaterialConverters.h
+++ b/Engine/Plugins/Enterprise/GLTFExporter/MigratedFromGit/Source/GLTFExporter/Private/Converters/GLTFMaterialConverters.h
@@ -17,7 +17,7 @@ class FGLTFMaterialConverter final : public TGLTFConverter<FGLTFJsonMaterialInde
 
 	bool TryGetBaseColorAndOpacity(FGLTFConvertBuilder& Builder, FGLTFJsonPBRMetallicRoughness& OutPBRParams, const UMaterialInterface* Material, EMaterialProperty BaseColorProperty, EMaterialProperty OpacityProperty) const;
 	bool TryGetMetallicAndRoughness(FGLTFConvertBuilder& Builder, FGLTFJsonPBRMetallicRoughness& OutPBRParams, const UMaterialInterface* Material, EMaterialProperty MetallicProperty, EMaterialProperty RoughnessProperty) const;
-	bool TryGetEmissive(FGLTFConvertBuilder& Builder, FGLTFJsonMaterial& JsonMaterial, const UMaterialInterface* Material) const;
+	bool TryGetEmissive(FGLTFConvertBuilder& Builder, FGLTFJsonMaterial& JsonMaterial, EMaterialProperty EmissiveProperty, const UMaterialInterface* Material) const;
 
 	bool IsPropertyNonDefault(EMaterialProperty Property, const UMaterialInterface* Material) const;
 	bool TryGetConstantColor(FGLTFJsonColor3& OutValue, EMaterialProperty Property, const UMaterialInterface* Material) const;
