[GLTFExporter] refactor(level): split parts of the level conversion into GLTFActorConverter and GLTFActorUtility

commit 20fc18b9f3c6455a9eb83c23a978a060fc5c8d6b
Date: Tue, 30 Jun 2020 00:09:20 +0200

    refactor(level): split parts of the level conversion into GLTFActorConverter and GLTFActorUtility
    (cherry picked from commit a46b9d4a1842a0d43ded3598e20b16c5d4a684fb​)
    # Conflicts from cherry-pick:
    #	Plugins/GLTFExporter/Source/GLTFExporter/Private/Converters/GLTFConverterUtility.h

#skipci

diff --git a/Engine/Plugins/Enterprise/GLTFExporter/MigratedFromGit/Source/GLTFExporter/Private/Builders/GLTFConvertBuilder.cpp b/Engine/Plugins/Enterprise/GLTFExporter/MigratedFromGit/Source/GLTFExporter/Private/Builders/GLTFConvertBuilder.cpp
--- a/Engine/Plugins/Enterprise/GLTFExporter/MigratedFromGit/Source/GLTFExporter/Private/Builders/GLTFConvertBuilder.cpp
+++ b/Engine/Plugins/Enterprise/GLTFExporter/MigratedFromGit/Source/GLTFExporter/Private/Builders/GLTFConvertBuilder.cpp
@@ -133,6 +133,16 @@ FGLTFJsonNodeIndex FGLTFConvertBuilder::GetOrAddNode(const USceneComponent* Scen
 	return SceneComponentConverter.GetOrAdd(*this, DesiredName, SceneComponent);
 }
 
+FGLTFJsonNodeIndex FGLTFConvertBuilder::GetOrAddNode(const AActor* Actor, const FString& DesiredName)
+{
+	if (Actor == nullptr)
+	{
+		return FGLTFJsonNodeIndex(INDEX_NONE);
+	}
+
+	return ActorConverter.GetOrAdd(*this, DesiredName, Actor);
+}
+
 FGLTFJsonSceneIndex FGLTFConvertBuilder::GetOrAddScene(const ULevel* Level, const FString& DesiredName)
 {
 	if (Level == nullptr)
diff --git a/Engine/Plugins/Enterprise/GLTFExporter/MigratedFromGit/Source/GLTFExporter/Private/Builders/GLTFConvertBuilder.h b/Engine/Plugins/Enterprise/GLTFExporter/MigratedFromGit/Source/GLTFExporter/Private/Builders/GLTFConvertBuilder.h
--- a/Engine/Plugins/Enterprise/GLTFExporter/MigratedFromGit/Source/GLTFExporter/Private/Builders/GLTFConvertBuilder.h
+++ b/Engine/Plugins/Enterprise/GLTFExporter/MigratedFromGit/Source/GLTFExporter/Private/Builders/GLTFConvertBuilder.h
@@ -32,6 +32,7 @@ public:
 	FGLTFJsonTextureIndex GetOrAddTexture(const UTexture2D* Texture, const FString& DesiredName = TEXT(""));
 
 	FGLTFJsonNodeIndex GetOrAddNode(const USceneComponent* SceneComponent, const FString& DesiredName = TEXT(""));
+	FGLTFJsonNodeIndex GetOrAddNode(const AActor* Actor, const FString& DesiredName = TEXT(""));
 	FGLTFJsonSceneIndex GetOrAddScene(const ULevel* Level, const FString& DesiredName = TEXT(""));
 	FGLTFJsonSceneIndex GetOrAddScene(const UWorld* World, const FString& DesiredName = TEXT(""));
 
@@ -51,5 +52,6 @@ private:
 	FGLTFTexture2DConverter Texture2DConverter;
 
 	FGLTFSceneComponentConverter SceneComponentConverter;
+	FGLTFActorConverter ActorConverter;
 	FGLTFLevelConverter LevelConverter;
 };
diff --git a/Engine/Plugins/Enterprise/GLTFExporter/MigratedFromGit/Source/GLTFExporter/Private/Builders/GLTFJsonBuilder.cpp b/Engine/Plugins/Enterprise/GLTFExporter/MigratedFromGit/Source/GLTFExporter/Private/Builders/GLTFJsonBuilder.cpp
--- a/Engine/Plugins/Enterprise/GLTFExporter/MigratedFromGit/Source/GLTFExporter/Private/Builders/GLTFJsonBuilder.cpp
+++ b/Engine/Plugins/Enterprise/GLTFExporter/MigratedFromGit/Source/GLTFExporter/Private/Builders/GLTFJsonBuilder.cpp
@@ -57,6 +57,18 @@ FGLTFJsonTextureIndex FGLTFJsonBuilder::AddTexture(const FGLTFJsonTexture& JsonT
 	return FGLTFJsonTextureIndex(JsonRoot.Textures.Add(JsonTexture));
 }
 
+FGLTFJsonNodeIndex FGLTFJsonBuilder::AddChildNode(FGLTFJsonNodeIndex ParentIndex, const FGLTFJsonNode& JsonNode)
+{
+	const FGLTFJsonNodeIndex ChildIndex = AddNode(JsonNode);
+
+	if (ParentIndex != INDEX_NONE)
+	{
+		GetNode(ParentIndex).Children.Add(ChildIndex);
+	}
+
+	return ChildIndex;
+}
+
 FGLTFJsonAccessor& FGLTFJsonBuilder::GetAccessor(FGLTFJsonAccessorIndex AccessorIndex)
 {
 	return JsonRoot.Accessors[AccessorIndex];
diff --git a/Engine/Plugins/Enterprise/GLTFExporter/MigratedFromGit/Source/GLTFExporter/Private/Builders/GLTFJsonBuilder.h b/Engine/Plugins/Enterprise/GLTFExporter/MigratedFromGit/Source/GLTFExporter/Private/Builders/GLTFJsonBuilder.h
--- a/Engine/Plugins/Enterprise/GLTFExporter/MigratedFromGit/Source/GLTFExporter/Private/Builders/GLTFJsonBuilder.h
+++ b/Engine/Plugins/Enterprise/GLTFExporter/MigratedFromGit/Source/GLTFExporter/Private/Builders/GLTFJsonBuilder.h
@@ -24,6 +24,8 @@ public:
 	FGLTFJsonSceneIndex AddScene(const FGLTFJsonScene& JsonScene = {});
 	FGLTFJsonTextureIndex AddTexture(const FGLTFJsonTexture& JsonTexture = {});
 
+	FGLTFJsonNodeIndex AddChildNode(FGLTFJsonNodeIndex ParentNodeIndex, const FGLTFJsonNode& JsonNode = {});
+
 	FGLTFJsonAccessor& GetAccessor(FGLTFJsonAccessorIndex AccessorIndex);
 	FGLTFJsonBuffer& GetBuffer(FGLTFJsonBufferIndex BufferIndex);
 	FGLTFJsonBufferView& GetBufferView(FGLTFJsonBufferViewIndex BufferViewIndex);
diff --git a/Engine/Plugins/Enterprise/GLTFExporter/MigratedFromGit/Source/GLTFExporter/Private/Converters/GLTFActorUtility.cpp b/Engine/Plugins/Enterprise/GLTFExporter/MigratedFromGit/Source/GLTFExporter/Private/Converters/GLTFActorUtility.cpp
new file mode 100644
--- /dev/null
+++ b/Engine/Plugins/Enterprise/GLTFExporter/MigratedFromGit/Source/GLTFExporter/Private/Converters/GLTFActorUtility.cpp
@@ -0,0 +1,25 @@
+// Copyright Epic Games, Inc. All Rights Reserved.
+
+#include "Converters/GLTFActorUtility.h"
+#include "Engine/Blueprint.h"
+
+bool FGLTFActorUtility::IsRootActor(const AActor* Actor, bool bSelectedOnly)
+{
+	const AActor* ParentActor = Actor->GetAttachParentActor();
+	return ParentActor == nullptr || (bSelectedOnly && !ParentActor->IsSelected());
+}
+
+UBlueprint* FGLTFActorUtility::GetBlueprintFromActor(const AActor* Actor)
+{
+	return UBlueprint::GetBlueprintFromClass(Actor->GetClass());
+}
+
+bool FGLTFActorUtility::IsSkySphereBlueprint(const UBlueprint* Blueprint)
+{
+	return Blueprint != nullptr && Blueprint->GetPathName().Equals(TEXT("/Engine/EngineSky/BP_Sky_Sphere.BP_Sky_Sphere"));
+}
+
+bool FGLTFActorUtility::IsHDRIBackdropBlueprint(const UBlueprint* Blueprint)
+{
+	return Blueprint != nullptr && Blueprint->GetPathName().Equals(TEXT("/HDRIBackdrop/Blueprints/HDRIBackdrop.HDRIBackdrop"));
+}
diff --git a/Engine/Plugins/Enterprise/GLTFExporter/MigratedFromGit/Source/GLTFExporter/Private/Converters/GLTFActorUtility.h b/Engine/Plugins/Enterprise/GLTFExporter/MigratedFromGit/Source/GLTFExporter/Private/Converters/GLTFActorUtility.h
new file mode 100644
--- /dev/null
+++ b/Engine/Plugins/Enterprise/GLTFExporter/MigratedFromGit/Source/GLTFExporter/Private/Converters/GLTFActorUtility.h
@@ -0,0 +1,16 @@
+// Copyright Epic Games, Inc. All Rights Reserved.
+
+#pragma once
+
+#include "Engine/EngineTypes.h"
+
+struct FGLTFActorUtility
+{
+	static bool IsRootActor(const AActor* Actor, bool bSelectedOnly);
+
+	static UBlueprint* GetBlueprintFromActor(const AActor* Actor);
+
+	static bool IsSkySphereBlueprint(const UBlueprint* Blueprint);
+
+	static bool IsHDRIBackdropBlueprint(const UBlueprint* Blueprint);
+};
diff --git a/Engine/Plugins/Enterprise/GLTFExporter/MigratedFromGit/Source/GLTFExporter/Private/Converters/GLTFConverterUtility.cpp b/Engine/Plugins/Enterprise/GLTFExporter/MigratedFromGit/Source/GLTFExporter/Private/Converters/GLTFConverterUtility.cpp
--- a/Engine/Plugins/Enterprise/GLTFExporter/MigratedFromGit/Source/GLTFExporter/Private/Converters/GLTFConverterUtility.cpp
+++ b/Engine/Plugins/Enterprise/GLTFExporter/MigratedFromGit/Source/GLTFExporter/Private/Converters/GLTFConverterUtility.cpp
@@ -2,7 +2,6 @@
 
 #include "Converters/GLTFConverterUtility.h"
 #include "Engine/TextureLODSettings.h"
-#include "Engine/Blueprint.h"
 
 TextureFilter FGLTFConverterUtility::GetDefaultFilter(TextureGroup LODGroup)
 {
@@ -21,24 +20,3 @@ TextureFilter FGLTFConverterUtility::GetDefaultFilter(TextureGroup LODGroup)
 		default:                                       return TF_Default; // fallback
 	}
 }
-
-bool FGLTFConverterUtility::IsSkySphereBlueprint(const UBlueprint* Blueprint)
-{
-	return Blueprint != nullptr && Blueprint->GetPathName().Equals(TEXT("/Engine/EngineSky/BP_Sky_Sphere.BP_Sky_Sphere"));
-}
-
-bool FGLTFConverterUtility::IsHDRIBackdropBlueprint(const UBlueprint* Blueprint)
-{
-	return Blueprint != nullptr && Blueprint->GetPathName().Equals(TEXT("/HDRIBackdrop/Blueprints/HDRIBackdrop.HDRIBackdrop"));
-}
-
-bool FGLTFConverterUtility::IsSelected(const UActorComponent* ActorComponent)
-{
-	if (ActorComponent == nullptr)
-	{
-		return false;
-	}
-
-	const AActor* Owner = ActorComponent->GetOwner();
-	return Owner != nullptr && Owner->IsSelected();
-}
diff --git a/Engine/Plugins/Enterprise/GLTFExporter/MigratedFromGit/Source/GLTFExporter/Private/Converters/GLTFConverterUtility.h b/Engine/Plugins/Enterprise/GLTFExporter/MigratedFromGit/Source/GLTFExporter/Private/Converters/GLTFConverterUtility.h
--- a/Engine/Plugins/Enterprise/GLTFExporter/MigratedFromGit/Source/GLTFExporter/Private/Converters/GLTFConverterUtility.h
+++ b/Engine/Plugins/Enterprise/GLTFExporter/MigratedFromGit/Source/GLTFExporter/Private/Converters/GLTFConverterUtility.h
@@ -158,12 +158,6 @@ struct FGLTFConverterUtility
 
 	static TextureFilter GetDefaultFilter(TextureGroup Group);
 
-	static bool IsSkySphereBlueprint(const UBlueprint* Blueprint);
-
-	static bool IsHDRIBackdropBlueprint(const UBlueprint* Blueprint);
-
-	static bool IsSelected(const UActorComponent* ActorComponent);
-
 	template <typename EnumType>
 	static FString GetEnumDisplayName(EnumType Value)
 	{
diff --git a/Engine/Plugins/Enterprise/GLTFExporter/MigratedFromGit/Source/GLTFExporter/Private/Converters/GLTFLevelConverters.cpp b/Engine/Plugins/Enterprise/GLTFExporter/MigratedFromGit/Source/GLTFExporter/Private/Converters/GLTFLevelConverters.cpp
--- a/Engine/Plugins/Enterprise/GLTFExporter/MigratedFromGit/Source/GLTFExporter/Private/Converters/GLTFLevelConverters.cpp
+++ b/Engine/Plugins/Enterprise/GLTFExporter/MigratedFromGit/Source/GLTFExporter/Private/Converters/GLTFLevelConverters.cpp
@@ -3,93 +3,89 @@
 #include "Converters/GLTFLevelConverters.h"
 #include "Builders/GLTFContainerBuilder.h"
 #include "Converters/GLTFConverterUtility.h"
+#include "Converters/GLTFActorUtility.h"
 
 FGLTFJsonNodeIndex FGLTFSceneComponentConverter::Add(FGLTFConvertBuilder& Builder, const FString& Name, const USceneComponent* SceneComponent)
 {
 	const AActor* Owner = SceneComponent->GetOwner();
-	const USceneComponent* ParentComponent = SceneComponent->GetAttachParent();
-	const UBlueprint* Blueprint = UBlueprint::GetBlueprintFromClass(Owner != nullptr ? Owner->GetClass() : nullptr);
+	if (Owner == nullptr)
+	{
+		// TODO: report invalid scene component
+		return FGLTFJsonNodeIndex(INDEX_NONE);
+	}
 
-	const bool bSelectedOnly = Builder.bSelectedActorsOnly;
-	const bool bIsRootComponent = Owner != nullptr && Owner->GetRootComponent() == SceneComponent;
-	const bool bRootNode = bIsRootComponent && (ParentComponent == nullptr || (bSelectedOnly && !FGLTFConverterUtility::IsSelected(ParentComponent)));
+	const bool bIsRootComponent = Owner->GetRootComponent() == SceneComponent;
+	const bool bIsRootNode = bIsRootComponent && FGLTFActorUtility::IsRootActor(Owner, Builder.bSelectedActorsOnly);
 
-	const FGLTFJsonNodeIndex NodeIndex = Builder.AddNode();
+	const USceneComponent* ParentComponent = SceneComponent->GetAttachParent();
+	const FGLTFJsonNodeIndex ParentNodeIndex = Builder.GetOrAddNode(ParentComponent);
+
+	const FGLTFJsonNodeIndex NodeIndex = Builder.AddChildNode(ParentNodeIndex);
 	FGLTFJsonNode& Node = Builder.GetNode(NodeIndex);
 	Node.Name = Name.IsEmpty() ? Owner->GetName() + TEXT("_") + SceneComponent->GetName() : Name;
 
 	// TODO: add support for non-uniform scaling (Unreal doesn't treat combined transforms as simple matrix multiplication)
-	const FTransform Transform = bRootNode ? SceneComponent->GetComponentTransform() : SceneComponent->GetRelativeTransform();
+	const FTransform Transform = bIsRootNode ? SceneComponent->GetComponentTransform() : SceneComponent->GetRelativeTransform();
 	Node.Translation = FGLTFConverterUtility::ConvertPosition(Transform.GetTranslation());
 	Node.Rotation = FGLTFConverterUtility::ConvertRotation(Transform.GetRotation());
 	Node.Scale = FGLTFConverterUtility::ConvertScale(Transform.GetScale3D());
 
-	if (FGLTFConverterUtility::IsSkySphereBlueprint(Blueprint))
-	{
-		// Ignore mesh and light components that are part of BP_SkySphere
-	}
-	else if (FGLTFConverterUtility::IsHDRIBackdropBlueprint(Blueprint))
-	{
-		if (bIsRootComponent)
-		{
-			// TODO: add support for backdrop
-		}
-	}
-	else if (const UStaticMeshComponent* StaticMeshComponent = Cast<UStaticMeshComponent>(SceneComponent))
+	if (const UStaticMeshComponent* StaticMeshComponent = Cast<UStaticMeshComponent>(SceneComponent))
 	{
 		Node.Mesh = Builder.GetOrAddMesh(StaticMeshComponent);
 	}
 
-	const TArray<USceneComponent*>& ChildComponents = SceneComponent->GetAttachChildren();
+	return NodeIndex;
+}
 
-	TArray<FGLTFJsonNodeIndex> Children;
-	Children.Reserve(ChildComponents.Num());
+FGLTFJsonNodeIndex FGLTFActorConverter::Add(FGLTFConvertBuilder& Builder, const FString& Name, const AActor* Actor)
+{
+	if (Builder.bSelectedActorsOnly && !Actor->IsSelected())
+	{
+		return FGLTFJsonNodeIndex(INDEX_NONE);
+	}
 
-	for (const USceneComponent* ChildComponent : ChildComponents)
+	const UBlueprint* Blueprint = FGLTFActorUtility::GetBlueprintFromActor(Actor);
+	if (FGLTFActorUtility::IsSkySphereBlueprint(Blueprint))
+	{
+		// Ignore mesh and light components that are part of BP_SkySphere
+		// TODO: add support for BP_SkySphere later?
+	}
+	else if (FGLTFActorUtility::IsHDRIBackdropBlueprint(Blueprint))
+	{
+		// TODO: add support for backdrop
+	}
+	else
 	{
-		if (ChildComponent != nullptr && (!bSelectedOnly || FGLTFConverterUtility::IsSelected(ChildComponent)))
+		for (const UActorComponent* Component : Actor->GetComponents())
 		{
-			const FGLTFJsonNodeIndex ChildNodeIndex = Builder.GetOrAddNode(ChildComponent);
-			if (ChildNodeIndex != INDEX_NONE)
+			const USceneComponent* SceneComponent = Cast<USceneComponent>(Component);
+			if (SceneComponent != nullptr)
 			{
-				Children.Add(ChildNodeIndex);
+				Builder.GetOrAddNode(SceneComponent);
 			}
 		}
 	}
 
-	// TODO: refactor this workaround for no longer valid reference to node (due to calling Builder.GetOrAddNode)
-	Builder.GetNode(NodeIndex).Children = Children;
-
-	return NodeIndex;
+	const USceneComponent* RootComponent = Actor->GetRootComponent();
+	return Builder.GetOrAddNode(RootComponent);
 }
 
 FGLTFJsonSceneIndex FGLTFLevelConverter::Add(FGLTFConvertBuilder& Builder, const FString& Name, const ULevel* Level)
 {
-	const bool bSelectedOnly = Builder.bSelectedActorsOnly;
-
 	FGLTFJsonScene Scene;
 	Scene.Name = Name;
 
 	// TODO: export Level->Model ?
 
-	const TArray<AActor*>& Actors = Level->Actors;
-	for (const AActor* Actor : Actors)
+	for (const AActor* Actor : Level->Actors)
 	{
-		if (Actor != nullptr && (!bSelectedOnly || Actor->IsSelected()))
+		// TODO: add support for ALevelVariantSetsActor
+
+		const FGLTFJsonNodeIndex NodeIndex = Builder.GetOrAddNode(Actor);
+		if (NodeIndex != INDEX_NONE && FGLTFActorUtility::IsRootActor(Actor, Builder.bSelectedActorsOnly))
 		{
-			const USceneComponent* RootComponent = Actor->GetRootComponent();
-			if (RootComponent != nullptr)
-			{
-				const USceneComponent* ParentComponent = RootComponent->GetAttachParent();
-				if (ParentComponent == nullptr || (bSelectedOnly && !FGLTFConverterUtility::IsSelected(ParentComponent)))
-				{
-					FGLTFJsonNodeIndex NodeIndex = Builder.GetOrAddNode(RootComponent);
-					if (NodeIndex != INDEX_NONE)
-					{
-						Scene.Nodes.Add(NodeIndex);
-					}
-				}
-			}
+			Scene.Nodes.Add(NodeIndex);
 		}
 	}
 
diff --git a/Engine/Plugins/Enterprise/GLTFExporter/MigratedFromGit/Source/GLTFExporter/Private/Converters/GLTFLevelConverters.h b/Engine/Plugins/Enterprise/GLTFExporter/MigratedFromGit/Source/GLTFExporter/Private/Converters/GLTFLevelConverters.h
--- a/Engine/Plugins/Enterprise/GLTFExporter/MigratedFromGit/Source/GLTFExporter/Private/Converters/GLTFLevelConverters.h
+++ b/Engine/Plugins/Enterprise/GLTFExporter/MigratedFromGit/Source/GLTFExporter/Private/Converters/GLTFLevelConverters.h
@@ -11,6 +11,11 @@ class FGLTFSceneComponentConverter final : public TGLTFConverter<FGLTFJsonNodeIn
 	FGLTFJsonNodeIndex Add(FGLTFConvertBuilder& Builder, const FString& Name, const USceneComponent* SceneComponent) override;
 };
 
+class FGLTFActorConverter final : public TGLTFConverter<FGLTFJsonNodeIndex, const AActor*>
+{
+	FGLTFJsonNodeIndex Add(FGLTFConvertBuilder& Builder, const FString& Name, const AActor* Actor) override;
+};
+
 class FGLTFLevelConverter final : public TGLTFConverter<FGLTFJsonSceneIndex, const ULevel*>
 {
 	FGLTFJsonSceneIndex Add(FGLTFConvertBuilder& Builder, const FString& Name, const ULevel* Level) override;
