[GLTFExporter] fix(material): lookup correct EnumType in GLTFConverterUtility::GetEnumDisplayName

commit 360c4f7413c59c44762d27b6e392144e6d7d870f
Date: Tue, 30 Jun 2020 00:53:51 +0200

    fix(material): lookup correct EnumType in GLTFConverterUtility::GetEnumDisplayName
    (cherry picked from commit adc07e29581b5654f135e77faf383c654162f4ef​)

#skipci

diff --git a/Engine/Plugins/Enterprise/GLTFExporter/MigratedFromGit/Source/GLTFExporter/Private/Converters/GLTFConverterUtility.h b/Engine/Plugins/Enterprise/GLTFExporter/MigratedFromGit/Source/GLTFExporter/Private/Converters/GLTFConverterUtility.h
--- a/Engine/Plugins/Enterprise/GLTFExporter/MigratedFromGit/Source/GLTFExporter/Private/Converters/GLTFConverterUtility.h
+++ b/Engine/Plugins/Enterprise/GLTFExporter/MigratedFromGit/Source/GLTFExporter/Private/Converters/GLTFConverterUtility.h
@@ -161,8 +161,9 @@ struct FGLTFConverterUtility
 	template <typename EnumType>
 	static FString GetEnumDisplayName(EnumType Value)
 	{
-		const UEnum* Enum = StaticEnum<EMaterialShadingModel>();
-		check(Enum);
+		static_assert(TIsEnum<EnumType>::Value, "Should only call this with enum types");
+		const UEnum* Enum = StaticEnum<EnumType>();
+		check(Enum != nullptr);
 		return Enum->GetDisplayNameTextByValue(Value).ToString();
 	}
 };
