[GLTFExporter] feat(lightmaps): implement light-map export via custom glTF extension

commit b23b083c1a82869a9a07cb2dcbcab1fd50279450
Date: Thu, 2 Jul 2020 15:02:01 +0200

    feat(lightmaps): implement light-map export via custom glTF extension
    (cherry picked from commit 2d1e9c86499479d078a56ae8c01ab6fee0ac29dd​)

#skipci

diff --git a/Engine/Plugins/Enterprise/GLTFExporter/MigratedFromGit/Source/GLTFExporter/Private/Builders/GLTFConvertBuilder.cpp b/Engine/Plugins/Enterprise/GLTFExporter/MigratedFromGit/Source/GLTFExporter/Private/Builders/GLTFConvertBuilder.cpp
--- a/Engine/Plugins/Enterprise/GLTFExporter/MigratedFromGit/Source/GLTFExporter/Private/Builders/GLTFConvertBuilder.cpp
+++ b/Engine/Plugins/Enterprise/GLTFExporter/MigratedFromGit/Source/GLTFExporter/Private/Builders/GLTFConvertBuilder.cpp
@@ -123,6 +123,26 @@ FGLTFJsonTextureIndex FGLTFConvertBuilder::GetOrAddTexture(const UTexture2D* Tex
 	return Texture2DConverter.GetOrAdd(*this, DesiredName.IsEmpty() ? Texture->GetName() : DesiredName, Texture);
 }
 
+FGLTFJsonTextureIndex FGLTFConvertBuilder::GetOrAddLightMapTexture(const ULightMapTexture2D* LightMapTexture2D, const FString& DesiredName)
+{
+	if (LightMapTexture2D == nullptr)
+	{
+		return FGLTFJsonTextureIndex(INDEX_NONE);
+	}
+
+	return LightMapTexture2DConverter.GetOrAdd(*this, DesiredName.IsEmpty() ? LightMapTexture2D->GetName() : DesiredName, LightMapTexture2D);
+}
+
+FGLTFJsonLightMapIndex FGLTFConvertBuilder::GetOrAddLightMap(const UStaticMeshComponent* StaticMeshComponent, const FString& DesiredName)
+{
+	if (StaticMeshComponent == nullptr)
+	{
+		return FGLTFJsonLightMapIndex(INDEX_NONE);
+	}
+
+	return LightMapConverter.GetOrAdd(*this, DesiredName, StaticMeshComponent);
+}
+
 FGLTFJsonNodeIndex FGLTFConvertBuilder::GetOrAddNode(const USceneComponent* SceneComponent, const FString& DesiredName)
 {
 	if (SceneComponent == nullptr)
diff --git a/Engine/Plugins/Enterprise/GLTFExporter/MigratedFromGit/Source/GLTFExporter/Private/Builders/GLTFConvertBuilder.h b/Engine/Plugins/Enterprise/GLTFExporter/MigratedFromGit/Source/GLTFExporter/Private/Builders/GLTFConvertBuilder.h
--- a/Engine/Plugins/Enterprise/GLTFExporter/MigratedFromGit/Source/GLTFExporter/Private/Builders/GLTFConvertBuilder.h
+++ b/Engine/Plugins/Enterprise/GLTFExporter/MigratedFromGit/Source/GLTFExporter/Private/Builders/GLTFConvertBuilder.h
@@ -8,6 +8,7 @@
 #include "Converters/GLTFMaterialConverters.h"
 #include "Converters/GLTFTextureConverters.h"
 #include "Converters/GLTFLevelConverters.h"
+#include "Converters/GLTFLightMapConverters.h"
 
 class FGLTFConvertBuilder : public FGLTFImageBuilder
 {
@@ -30,6 +31,8 @@ public:
 
 	FGLTFJsonMaterialIndex GetOrAddMaterial(const UMaterialInterface* Material, const FString& DesiredName = TEXT(""));
 	FGLTFJsonTextureIndex GetOrAddTexture(const UTexture2D* Texture, const FString& DesiredName = TEXT(""));
+	FGLTFJsonTextureIndex GetOrAddLightMapTexture(const ULightMapTexture2D* LightMapTexture2D, const FString& DesiredName = TEXT(""));
+	FGLTFJsonLightMapIndex GetOrAddLightMap(const UStaticMeshComponent* StaticMeshComponent, const FString& DesiredName = TEXT(""));
 
 	FGLTFJsonNodeIndex GetOrAddNode(const USceneComponent* SceneComponent, const FString& DesiredName = TEXT(""));
 	FGLTFJsonNodeIndex GetOrAddNode(const AActor* Actor, const FString& DesiredName = TEXT(""));
@@ -50,6 +53,8 @@ private:
 
 	FGLTFMaterialConverter MaterialConverter;
 	FGLTFTexture2DConverter Texture2DConverter;
+	FGLTFLightMapTexture2DConverter LightMapTexture2DConverter;
+	FGLTFLightMapConverter LightMapConverter;
 
 	FGLTFSceneComponentConverter SceneComponentConverter;
 	FGLTFActorConverter ActorConverter;
diff --git a/Engine/Plugins/Enterprise/GLTFExporter/MigratedFromGit/Source/GLTFExporter/Private/Builders/GLTFJsonBuilder.cpp b/Engine/Plugins/Enterprise/GLTFExporter/MigratedFromGit/Source/GLTFExporter/Private/Builders/GLTFJsonBuilder.cpp
--- a/Engine/Plugins/Enterprise/GLTFExporter/MigratedFromGit/Source/GLTFExporter/Private/Builders/GLTFJsonBuilder.cpp
+++ b/Engine/Plugins/Enterprise/GLTFExporter/MigratedFromGit/Source/GLTFExporter/Private/Builders/GLTFJsonBuilder.cpp
@@ -57,6 +57,11 @@ FGLTFJsonTextureIndex FGLTFJsonBuilder::AddTexture(const FGLTFJsonTexture& JsonT
 	return FGLTFJsonTextureIndex(JsonRoot.Textures.Add(JsonTexture));
 }
 
+FGLTFJsonLightMapIndex FGLTFJsonBuilder::AddLightMap(const FGLTFJsonLightMap& JsonLightMap)
+{
+	return FGLTFJsonLightMapIndex(JsonRoot.LightMaps.Add(JsonLightMap));
+}
+
 FGLTFJsonNodeIndex FGLTFJsonBuilder::AddChildNode(FGLTFJsonNodeIndex ParentIndex, const FGLTFJsonNode& JsonNode)
 {
 	const FGLTFJsonNodeIndex ChildIndex = AddNode(JsonNode);
@@ -119,6 +124,11 @@ FGLTFJsonTexture& FGLTFJsonBuilder::GetTexture(FGLTFJsonTextureIndex TextureInde
 	return JsonRoot.Textures[TextureIndex];
 }
 
+FGLTFJsonLightMap& FGLTFJsonBuilder::GetLightMap(FGLTFJsonLightMapIndex LightMapIndex)
+{
+	return JsonRoot.LightMaps[LightMapIndex];
+}
+
 bool FGLTFJsonBuilder::Serialize(FArchive& Archive, const FString& FilePath)
 {
 	JsonRoot.Serialize(&Archive, true);
diff --git a/Engine/Plugins/Enterprise/GLTFExporter/MigratedFromGit/Source/GLTFExporter/Private/Builders/GLTFJsonBuilder.h b/Engine/Plugins/Enterprise/GLTFExporter/MigratedFromGit/Source/GLTFExporter/Private/Builders/GLTFJsonBuilder.h
--- a/Engine/Plugins/Enterprise/GLTFExporter/MigratedFromGit/Source/GLTFExporter/Private/Builders/GLTFJsonBuilder.h
+++ b/Engine/Plugins/Enterprise/GLTFExporter/MigratedFromGit/Source/GLTFExporter/Private/Builders/GLTFJsonBuilder.h
@@ -23,6 +23,7 @@ public:
 	FGLTFJsonSamplerIndex AddSampler(const FGLTFJsonSampler& JsonSampler = {});
 	FGLTFJsonSceneIndex AddScene(const FGLTFJsonScene& JsonScene = {});
 	FGLTFJsonTextureIndex AddTexture(const FGLTFJsonTexture& JsonTexture = {});
+	FGLTFJsonLightMapIndex AddLightMap(const FGLTFJsonLightMap& JsonLightMap = {});
 
 	FGLTFJsonNodeIndex AddChildNode(FGLTFJsonNodeIndex ParentNodeIndex, const FGLTFJsonNode& JsonNode = {});
 
@@ -36,6 +37,7 @@ public:
 	FGLTFJsonSampler& GetSampler(FGLTFJsonSamplerIndex SamplerIndex);
 	FGLTFJsonScene& GetScene(FGLTFJsonSceneIndex SceneIndex);
 	FGLTFJsonTexture& GetTexture(FGLTFJsonTextureIndex TextureIndex);
+	FGLTFJsonLightMap& GetLightMap(FGLTFJsonLightMapIndex LightMapIndex);
 
 	virtual bool Serialize(FArchive& Archive, const FString& FilePath);
 
diff --git a/Engine/Plugins/Enterprise/GLTFExporter/MigratedFromGit/Source/GLTFExporter/Private/Converters/GLTFLevelConverters.cpp b/Engine/Plugins/Enterprise/GLTFExporter/MigratedFromGit/Source/GLTFExporter/Private/Converters/GLTFLevelConverters.cpp
--- a/Engine/Plugins/Enterprise/GLTFExporter/MigratedFromGit/Source/GLTFExporter/Private/Converters/GLTFLevelConverters.cpp
+++ b/Engine/Plugins/Enterprise/GLTFExporter/MigratedFromGit/Source/GLTFExporter/Private/Converters/GLTFLevelConverters.cpp
@@ -4,6 +4,7 @@
 #include "Builders/GLTFContainerBuilder.h"
 #include "Converters/GLTFConverterUtility.h"
 #include "Converters/GLTFActorUtility.h"
+#include "Engine/MapBuildDataRegistry.h"
 
 FGLTFJsonNodeIndex FGLTFSceneComponentConverter::Add(FGLTFConvertBuilder& Builder, const FString& Name, const USceneComponent* SceneComponent)
 {
@@ -32,7 +33,10 @@ FGLTFJsonNodeIndex FGLTFSceneComponentConverter::Add(FGLTFConvertBuilder& Builde
 
 	if (const UStaticMeshComponent* StaticMeshComponent = Cast<UStaticMeshComponent>(SceneComponent))
 	{
+		// TODO: make export of light-maps optional (selectable via export options)
+
 		Node.Mesh = Builder.GetOrAddMesh(StaticMeshComponent);
+		Node.LightMap = Builder.GetOrAddLightMap(StaticMeshComponent);
 	}
 
 	return NodeIndex;
diff --git a/Engine/Plugins/Enterprise/GLTFExporter/MigratedFromGit/Source/GLTFExporter/Private/Converters/GLTFLightMapConverters.cpp b/Engine/Plugins/Enterprise/GLTFExporter/MigratedFromGit/Source/GLTFExporter/Private/Converters/GLTFLightMapConverters.cpp
new file mode 100644
--- /dev/null
+++ b/Engine/Plugins/Enterprise/GLTFExporter/MigratedFromGit/Source/GLTFExporter/Private/Converters/GLTFLightMapConverters.cpp
@@ -0,0 +1,77 @@
+// Copyright Epic Games, Inc. All Rights Reserved.
+
+#include "Converters/GLTFLightMapConverters.h"
+#include "Json/GLTFJsonLightMap.h"
+#include "Builders/GLTFContainerBuilder.h"
+#include "Engine/MapBuildDataRegistry.h"
+
+FGLTFJsonLightMapIndex FGLTFLightMapConverter::Add(FGLTFConvertBuilder& Builder, const FString& Name, const UStaticMeshComponent* StaticMeshComponent)
+{
+	const AActor* Owner = StaticMeshComponent->GetOwner();
+	if (Owner == nullptr)
+	{
+		// TODO: report invalid mesh component
+		return FGLTFJsonLightMapIndex(INDEX_NONE);
+	}
+
+	const UStaticMesh* StaticMesh = StaticMeshComponent->GetStaticMesh();
+
+	if (StaticMesh == nullptr)
+	{
+		return FGLTFJsonLightMapIndex(INDEX_NONE);
+	}
+
+	const int32 LODIndex = StaticMeshComponent->ForcedLodModel > 0 ? StaticMeshComponent->ForcedLodModel - 1 : /* auto-select */ 0;
+
+	if (LODIndex < 0 || StaticMesh->GetNumLODs() <= LODIndex || StaticMeshComponent->LODData.Num() <= LODIndex)
+	{
+		return FGLTFJsonLightMapIndex(INDEX_NONE);
+	}
+
+	const int32 CoordinateIndex = StaticMesh->LightMapCoordinateIndex;
+	const FStaticMeshLODResources& LODResources = StaticMesh->GetLODForExport(LODIndex);
+
+	if (CoordinateIndex < 0 || CoordinateIndex >= LODResources.GetNumTexCoords())
+	{
+		return FGLTFJsonLightMapIndex(INDEX_NONE);
+	}
+
+	const FStaticMeshComponentLODInfo& ComponentLODInfo = StaticMeshComponent->LODData[LODIndex];
+	const FMeshMapBuildData* MeshMapBuildData = StaticMeshComponent->GetMeshMapBuildData(ComponentLODInfo);
+
+	if (MeshMapBuildData == nullptr || MeshMapBuildData->LightMap == nullptr)
+	{
+		return FGLTFJsonLightMapIndex(INDEX_NONE);
+	}
+
+	const FLightMap2D* LightMap2D = MeshMapBuildData->LightMap->GetLightMap2D();
+	check(LightMap2D);
+
+	// TODO: is it correct to use SM5?
+	const FLightMapInteraction LightMapInteraction = LightMap2D->GetInteraction(ERHIFeatureLevel::Type::SM5);
+
+	const bool bUseHighQualityLightMap = true;
+	const ULightMapTexture2D* Texture = LightMapInteraction.GetTexture(bUseHighQualityLightMap);
+	const FGLTFJsonTextureIndex TextureIndex = Builder.GetOrAddLightMapTexture(Texture);
+
+	if (TextureIndex == INDEX_NONE)
+	{
+		return FGLTFJsonLightMapIndex(INDEX_NONE);
+	}
+
+	const FVector2D& CoordinateBias = LightMap2D->GetCoordinateBias();
+	const FVector2D& CoordinateScale = LightMap2D->GetCoordinateScale();
+	const FVector4& LightMapAdd = LightMapInteraction.GetAddArray()[0];
+	const FVector4& LightMapScale = LightMapInteraction.GetScaleArray()[0];
+
+	FGLTFJsonLightMap JsonLightMap;
+	JsonLightMap.Name = Name.IsEmpty() ? Owner->GetName() + TEXT("_") + StaticMeshComponent->GetName() : Name;
+	JsonLightMap.Texture.Index = TextureIndex;
+	JsonLightMap.Texture.TexCoord = CoordinateIndex;
+	JsonLightMap.ValueScale = { LightMapScale.X, LightMapScale.Y, LightMapScale.Z, LightMapScale.W };
+	JsonLightMap.ValueOffset = { LightMapAdd.X, LightMapAdd.Y, LightMapAdd.Z, LightMapAdd.W };
+	JsonLightMap.CoordinateScale = { CoordinateScale.X, CoordinateScale.Y };
+	JsonLightMap.CoordinateOffset = { CoordinateBias.X, CoordinateBias.Y };
+
+	return Builder.AddLightMap(JsonLightMap);
+}
diff --git a/Engine/Plugins/Enterprise/GLTFExporter/MigratedFromGit/Source/GLTFExporter/Private/Converters/GLTFLightMapConverters.h b/Engine/Plugins/Enterprise/GLTFExporter/MigratedFromGit/Source/GLTFExporter/Private/Converters/GLTFLightMapConverters.h
new file mode 100644
--- /dev/null
+++ b/Engine/Plugins/Enterprise/GLTFExporter/MigratedFromGit/Source/GLTFExporter/Private/Converters/GLTFLightMapConverters.h
@@ -0,0 +1,12 @@
+// Copyright Epic Games, Inc. All Rights Reserved.
+
+#pragma once
+
+#include "Json/GLTFJsonIndex.h"
+#include "Converters/GLTFConverter.h"
+#include "Engine.h"
+
+class FGLTFLightMapConverter final : public TGLTFConverter<FGLTFJsonLightMapIndex, const UStaticMeshComponent*>
+{
+	FGLTFJsonLightMapIndex Add(FGLTFConvertBuilder& Builder, const FString& Name, const UStaticMeshComponent* StaticMeshComponent) override;
+};
diff --git a/Engine/Plugins/Enterprise/GLTFExporter/MigratedFromGit/Source/GLTFExporter/Private/Converters/GLTFTextureConverters.cpp b/Engine/Plugins/Enterprise/GLTFExporter/MigratedFromGit/Source/GLTFExporter/Private/Converters/GLTFTextureConverters.cpp
--- a/Engine/Plugins/Enterprise/GLTFExporter/MigratedFromGit/Source/GLTFExporter/Private/Converters/GLTFTextureConverters.cpp
+++ b/Engine/Plugins/Enterprise/GLTFExporter/MigratedFromGit/Source/GLTFExporter/Private/Converters/GLTFTextureConverters.cpp
@@ -28,8 +28,23 @@ FGLTFJsonTextureIndex FGLTFTextureCubeConverter::Add(FGLTFConvertBuilder& Builde
 	return FGLTFJsonTextureIndex(INDEX_NONE);
 }
 
-FGLTFJsonTextureIndex FGLTFLightMapTexture2DConverter::Add(FGLTFConvertBuilder& Builder, const FString& Name, const ULightMapTexture2D* LightMap)
+FGLTFJsonTextureIndex FGLTFLightMapTexture2DConverter::Add(FGLTFConvertBuilder& Builder, const FString& Name, const ULightMapTexture2D* LightMapTexture2D)
 {
-	// TODO: implement support
-	return FGLTFJsonTextureIndex(INDEX_NONE);
+	// TODO: maybe we should reuse existing samplers?
+	FGLTFJsonSampler JsonSampler;
+	JsonSampler.Name = Name;
+
+	// TODO: are these filters ok to use as default? The actual texture uses "nearest"
+	JsonSampler.MinFilter = EGLTFJsonTextureFilter::LinearMipmapLinear;
+	JsonSampler.MagFilter = EGLTFJsonTextureFilter::Linear;
+
+	JsonSampler.WrapS = FGLTFConverterUtility::ConvertWrap(LightMapTexture2D->AddressX);
+	JsonSampler.WrapT = FGLTFConverterUtility::ConvertWrap(LightMapTexture2D->AddressY);
+
+	FGLTFJsonTexture JsonTexture;
+	JsonTexture.Name = Name;
+	JsonTexture.Sampler = Builder.AddSampler(JsonSampler);
+	JsonTexture.Source = Builder.AddImage(LightMapTexture2D->Source, Name);
+
+	return Builder.AddTexture(JsonTexture);
 }
diff --git a/Engine/Plugins/Enterprise/GLTFExporter/MigratedFromGit/Source/GLTFExporter/Private/Converters/GLTFTextureConverters.h b/Engine/Plugins/Enterprise/GLTFExporter/MigratedFromGit/Source/GLTFExporter/Private/Converters/GLTFTextureConverters.h
--- a/Engine/Plugins/Enterprise/GLTFExporter/MigratedFromGit/Source/GLTFExporter/Private/Converters/GLTFTextureConverters.h
+++ b/Engine/Plugins/Enterprise/GLTFExporter/MigratedFromGit/Source/GLTFExporter/Private/Converters/GLTFTextureConverters.h
@@ -18,5 +18,5 @@ class FGLTFTextureCubeConverter final : public TGLTFConverter<FGLTFJsonTextureIn
 
 class FGLTFLightMapTexture2DConverter final : public TGLTFConverter<FGLTFJsonTextureIndex, const ULightMapTexture2D*>
 {
-	FGLTFJsonTextureIndex Add(FGLTFConvertBuilder& Builder, const FString& Name, const ULightMapTexture2D* LightMap) override;
+	FGLTFJsonTextureIndex Add(FGLTFConvertBuilder& Builder, const FString& Name, const ULightMapTexture2D* LightMapTexture2D) override;
 };
diff --git a/Engine/Plugins/Enterprise/GLTFExporter/MigratedFromGit/Source/GLTFExporter/Private/Json/GLTFJsonEnums.h b/Engine/Plugins/Enterprise/GLTFExporter/MigratedFromGit/Source/GLTFExporter/Private/Json/GLTFJsonEnums.h
--- a/Engine/Plugins/Enterprise/GLTFExporter/MigratedFromGit/Source/GLTFExporter/Private/Json/GLTFJsonEnums.h
+++ b/Engine/Plugins/Enterprise/GLTFExporter/MigratedFromGit/Source/GLTFExporter/Private/Json/GLTFJsonEnums.h
@@ -9,7 +9,8 @@ enum class EGLTFJsonExtension
 	KHR_LightsPunctual,
 	KHR_MaterialsUnlit,
 	KHR_MaterialsClearCoat,
-	KHR_MeshQuantization
+	KHR_MeshQuantization,
+	EPIC_LightMapTextures
 };
 
 enum class EGLTFJsonShadingModel
diff --git a/Engine/Plugins/Enterprise/GLTFExporter/MigratedFromGit/Source/GLTFExporter/Private/Json/GLTFJsonIndex.h b/Engine/Plugins/Enterprise/GLTFExporter/MigratedFromGit/Source/GLTFExporter/Private/Json/GLTFJsonIndex.h
--- a/Engine/Plugins/Enterprise/GLTFExporter/MigratedFromGit/Source/GLTFExporter/Private/Json/GLTFJsonIndex.h
+++ b/Engine/Plugins/Enterprise/GLTFExporter/MigratedFromGit/Source/GLTFExporter/Private/Json/GLTFJsonIndex.h
@@ -64,3 +64,4 @@ struct FGLTFJsonSamplerIndex : FGLTFJsonIndex<FGLTFJsonSamplerIndex> { using FGL
 struct FGLTFJsonSceneIndex : FGLTFJsonIndex<FGLTFJsonSceneIndex> { using FGLTFJsonIndex::FGLTFJsonIndex; };
 struct FGLTFJsonSkinIndex : FGLTFJsonIndex<FGLTFJsonSkinIndex> { using FGLTFJsonIndex::FGLTFJsonIndex; };
 struct FGLTFJsonTextureIndex : FGLTFJsonIndex<FGLTFJsonTextureIndex> { using FGLTFJsonIndex::FGLTFJsonIndex; };
+struct FGLTFJsonLightMapIndex : FGLTFJsonIndex<FGLTFJsonLightMapIndex> { using FGLTFJsonIndex::FGLTFJsonIndex; };
diff --git a/Engine/Plugins/Enterprise/GLTFExporter/MigratedFromGit/Source/GLTFExporter/Private/Json/GLTFJsonLightMap.h b/Engine/Plugins/Enterprise/GLTFExporter/MigratedFromGit/Source/GLTFExporter/Private/Json/GLTFJsonLightMap.h
new file mode 100644
--- /dev/null
+++ b/Engine/Plugins/Enterprise/GLTFExporter/MigratedFromGit/Source/GLTFExporter/Private/Json/GLTFJsonLightMap.h
@@ -0,0 +1,57 @@
+// Copyright Epic Games, Inc. All Rights Reserved.
+
+#pragma once
+
+#include "Json/GLTFJsonMaterial.h"
+#include "Json/GLTFJsonVector2.h"
+#include "Json/GLTFJsonVector4.h"
+#include "Serialization/JsonSerializer.h"
+
+struct FGLTFJsonLightMap
+{
+	FString              Name;
+	FGLTFJsonVector4     ValueScale;
+	FGLTFJsonVector4     ValueOffset;
+	FGLTFJsonVector2     CoordinateScale;
+	FGLTFJsonVector2     CoordinateOffset;
+	FGLTFJsonTextureInfo Texture;
+
+	FGLTFJsonLightMap()
+		: ValueScale(FGLTFJsonVector4::One)
+		, ValueOffset(FGLTFJsonVector4::Zero)
+		, CoordinateScale(FGLTFJsonVector2::One)
+		, CoordinateOffset(FGLTFJsonVector2::Zero)
+	{
+	}
+
+	template <class CharType = TCHAR, class PrintPolicy = TPrettyJsonPrintPolicy<CharType>>
+	void WriteObject(TJsonWriter<CharType, PrintPolicy>& JsonWriter, FGLTFJsonExtensions& Extensions) const
+	{
+		JsonWriter.WriteObjectStart();
+
+		if (!Name.IsEmpty())
+		{
+			JsonWriter.WriteValue(TEXT("name"), Name);
+		}
+
+		if (Texture.Index != INDEX_NONE)
+		{
+			JsonWriter.WriteIdentifierPrefix(TEXT("texture"));
+			Texture.WriteObject(JsonWriter, Extensions);
+		}
+
+		JsonWriter.WriteIdentifierPrefix(TEXT("valueScale"));
+		ValueScale.WriteArray(JsonWriter);
+
+		JsonWriter.WriteIdentifierPrefix(TEXT("valueOffset"));
+		ValueOffset.WriteArray(JsonWriter);
+
+		JsonWriter.WriteIdentifierPrefix(TEXT("coordinateScale"));
+		CoordinateScale.WriteArray(JsonWriter);
+
+		JsonWriter.WriteIdentifierPrefix(TEXT("coordinateOffset"));
+		CoordinateOffset.WriteArray(JsonWriter);
+
+		JsonWriter.WriteObjectEnd();
+	}
+};
diff --git a/Engine/Plugins/Enterprise/GLTFExporter/MigratedFromGit/Source/GLTFExporter/Private/Json/GLTFJsonNode.h b/Engine/Plugins/Enterprise/GLTFExporter/MigratedFromGit/Source/GLTFExporter/Private/Json/GLTFJsonNode.h
--- a/Engine/Plugins/Enterprise/GLTFExporter/MigratedFromGit/Source/GLTFExporter/Private/Json/GLTFJsonNode.h
+++ b/Engine/Plugins/Enterprise/GLTFExporter/MigratedFromGit/Source/GLTFExporter/Private/Json/GLTFJsonNode.h
@@ -15,9 +15,10 @@ struct FGLTFJsonNode
 	FGLTFJsonQuaternion Rotation;
 	FGLTFJsonVector3    Scale;
 
-	FGLTFJsonCameraIndex Camera;
-	FGLTFJsonSkinIndex   Skin;
-	FGLTFJsonMeshIndex   Mesh;
+	FGLTFJsonCameraIndex    Camera;
+	FGLTFJsonSkinIndex      Skin;
+	FGLTFJsonMeshIndex      Mesh;
+	FGLTFJsonLightMapIndex  LightMap;
 
 	TArray<FGLTFJsonNodeIndex> Children;
 
@@ -71,6 +72,26 @@ struct FGLTFJsonNode
 			JsonWriter.WriteValue(TEXT("mesh"), Mesh);
 		}
 
+		const bool bWriteExtensions = LightMap != INDEX_NONE;
+
+		if (bWriteExtensions)
+		{
+			JsonWriter.WriteObjectStart(TEXT("extensions"));
+
+			if (LightMap != INDEX_NONE)
+			{
+				const EGLTFJsonExtension Extension = EGLTFJsonExtension::EPIC_LightMapTextures;
+
+				Extensions.Used.Add(Extension);
+
+				JsonWriter.WriteObjectStart(FGLTFJsonUtility::ToString(Extension));
+				JsonWriter.WriteValue(TEXT("lightmap"), LightMap);
+				JsonWriter.WriteObjectEnd();
+			}
+
+			JsonWriter.WriteObjectEnd();
+		}
+
 		if (Children.Num() > 0)
 		{
 			JsonWriter.WriteValue(TEXT("children"), Children);
diff --git a/Engine/Plugins/Enterprise/GLTFExporter/MigratedFromGit/Source/GLTFExporter/Private/Json/GLTFJsonRoot.h b/Engine/Plugins/Enterprise/GLTFExporter/MigratedFromGit/Source/GLTFExporter/Private/Json/GLTFJsonRoot.h
--- a/Engine/Plugins/Enterprise/GLTFExporter/MigratedFromGit/Source/GLTFExporter/Private/Json/GLTFJsonRoot.h
+++ b/Engine/Plugins/Enterprise/GLTFExporter/MigratedFromGit/Source/GLTFExporter/Private/Json/GLTFJsonRoot.h
@@ -14,6 +14,7 @@
 #include "Json/GLTFJsonSampler.h"
 #include "Json/GLTFJsonScene.h"
 #include "Json/GLTFJsonTexture.h"
+#include "Json/GLTFJsonLightMap.h"
 #include "Policies/CondensedJsonPrintPolicy.h"
 
 
@@ -64,6 +65,7 @@ struct FGLTFJsonRoot
 	TArray<FGLTFJsonSampler>    Samplers;
 	TArray<FGLTFJsonScene>      Scenes;
 	TArray<FGLTFJsonTexture>    Textures;
+	TArray<FGLTFJsonLightMap>   LightMaps;
 
 	template <class CharType = TCHAR, class PrintPolicy = TPrettyJsonPrintPolicy<CharType>>
 	void WriteObject(TJsonWriter<CharType, PrintPolicy>& JsonWriter) const
@@ -91,6 +93,28 @@ struct FGLTFJsonRoot
 		FGLTFJsonUtility::WriteObjectArray(JsonWriter, TEXT("scenes"), Scenes, AllExtensions);
 		FGLTFJsonUtility::WriteObjectArray(JsonWriter, TEXT("textures"), Textures, AllExtensions);
 
+		const bool bWriteExtensions = LightMaps.Num() > 0;
+
+		if (bWriteExtensions)
+		{
+			JsonWriter.WriteObjectStart(TEXT("extensions"));
+
+			if (LightMaps.Num() > 0)
+			{
+				const EGLTFJsonExtension Extension = EGLTFJsonExtension::EPIC_LightMapTextures;
+
+				AllExtensions.Used.Add(Extension);
+
+				JsonWriter.WriteObjectStart(FGLTFJsonUtility::ToString(Extension));
+
+				FGLTFJsonUtility::WriteObjectArray(JsonWriter, TEXT("lightmaps"), LightMaps, AllExtensions);
+
+				JsonWriter.WriteObjectEnd();
+			}
+
+			JsonWriter.WriteObjectEnd();
+		}
+
 		FGLTFJsonUtility::WriteStringArray(JsonWriter, TEXT("extensionsUsed"), AllExtensions.Used);
 		FGLTFJsonUtility::WriteStringArray(JsonWriter, TEXT("extensionsRequired"), AllExtensions.Required);
 
diff --git a/Engine/Plugins/Enterprise/GLTFExporter/MigratedFromGit/Source/GLTFExporter/Private/Json/GLTFJsonUtility.h b/Engine/Plugins/Enterprise/GLTFExporter/MigratedFromGit/Source/GLTFExporter/Private/Json/GLTFJsonUtility.h
--- a/Engine/Plugins/Enterprise/GLTFExporter/MigratedFromGit/Source/GLTFExporter/Private/Json/GLTFJsonUtility.h
+++ b/Engine/Plugins/Enterprise/GLTFExporter/MigratedFromGit/Source/GLTFExporter/Private/Json/GLTFJsonUtility.h
@@ -21,6 +21,7 @@ struct FGLTFJsonUtility
 			case EGLTFJsonExtension::KHR_MaterialsUnlit:     return TEXT("KHR_materials_unlit");
 			case EGLTFJsonExtension::KHR_MaterialsClearCoat: return TEXT("KHR_materials_clearcoat");
 			case EGLTFJsonExtension::KHR_MeshQuantization:   return TEXT("KHR_mesh_quantization");
+			case EGLTFJsonExtension::EPIC_LightMapTextures:  return TEXT("EPIC_lightmap_textures");
 			default:                                         return TEXT("unknown");
 		}
 	}
