[GLTFExporter] refactor(lightmaps): rename GetOrAddLightMapTexture to GetOrAddTexture

commit 2fe888ffc0324f1e57f70d06ccac31c20db760fe
Date: Fri, 3 Jul 2020 08:12:50 +0200

    refactor(lightmaps): rename GetOrAddLightMapTexture to GetOrAddTexture
    (cherry picked from commit 87e5c7f5bef1da66b7fbe104978bd9a85b5ef2cb​)

#skipci

diff --git a/Engine/Plugins/Enterprise/GLTFExporter/MigratedFromGit/Source/GLTFExporter/Private/Builders/GLTFConvertBuilder.cpp b/Engine/Plugins/Enterprise/GLTFExporter/MigratedFromGit/Source/GLTFExporter/Private/Builders/GLTFConvertBuilder.cpp
--- a/Engine/Plugins/Enterprise/GLTFExporter/MigratedFromGit/Source/GLTFExporter/Private/Builders/GLTFConvertBuilder.cpp
+++ b/Engine/Plugins/Enterprise/GLTFExporter/MigratedFromGit/Source/GLTFExporter/Private/Builders/GLTFConvertBuilder.cpp
@@ -123,7 +123,7 @@ FGLTFJsonTextureIndex FGLTFConvertBuilder::GetOrAddTexture(const UTexture2D* Tex
 	return Texture2DConverter.GetOrAdd(*this, DesiredName.IsEmpty() ? Texture->GetName() : DesiredName, Texture);
 }
 
-FGLTFJsonTextureIndex FGLTFConvertBuilder::GetOrAddLightMapTexture(const ULightMapTexture2D* LightMapTexture2D, const FString& DesiredName)
+FGLTFJsonTextureIndex FGLTFConvertBuilder::GetOrAddTexture(const ULightMapTexture2D* LightMapTexture2D, const FString& DesiredName)
 {
 	if (LightMapTexture2D == nullptr)
 	{
diff --git a/Engine/Plugins/Enterprise/GLTFExporter/MigratedFromGit/Source/GLTFExporter/Private/Builders/GLTFConvertBuilder.h b/Engine/Plugins/Enterprise/GLTFExporter/MigratedFromGit/Source/GLTFExporter/Private/Builders/GLTFConvertBuilder.h
--- a/Engine/Plugins/Enterprise/GLTFExporter/MigratedFromGit/Source/GLTFExporter/Private/Builders/GLTFConvertBuilder.h
+++ b/Engine/Plugins/Enterprise/GLTFExporter/MigratedFromGit/Source/GLTFExporter/Private/Builders/GLTFConvertBuilder.h
@@ -31,7 +31,7 @@ public:
 
 	FGLTFJsonMaterialIndex GetOrAddMaterial(const UMaterialInterface* Material, const FString& DesiredName = TEXT(""));
 	FGLTFJsonTextureIndex GetOrAddTexture(const UTexture2D* Texture, const FString& DesiredName = TEXT(""));
-	FGLTFJsonTextureIndex GetOrAddLightMapTexture(const ULightMapTexture2D* LightMapTexture2D, const FString& DesiredName = TEXT(""));
+	FGLTFJsonTextureIndex GetOrAddTexture(const ULightMapTexture2D* LightMapTexture2D, const FString& DesiredName = TEXT(""));
 	FGLTFJsonLightMapIndex GetOrAddLightMap(const UStaticMeshComponent* StaticMeshComponent, const FString& DesiredName = TEXT(""));
 
 	FGLTFJsonNodeIndex GetOrAddNode(const USceneComponent* SceneComponent, const FString& DesiredName = TEXT(""));
diff --git a/Engine/Plugins/Enterprise/GLTFExporter/MigratedFromGit/Source/GLTFExporter/Private/Converters/GLTFLightMapConverters.cpp b/Engine/Plugins/Enterprise/GLTFExporter/MigratedFromGit/Source/GLTFExporter/Private/Converters/GLTFLightMapConverters.cpp
--- a/Engine/Plugins/Enterprise/GLTFExporter/MigratedFromGit/Source/GLTFExporter/Private/Converters/GLTFLightMapConverters.cpp
+++ b/Engine/Plugins/Enterprise/GLTFExporter/MigratedFromGit/Source/GLTFExporter/Private/Converters/GLTFLightMapConverters.cpp
@@ -52,7 +52,7 @@ FGLTFJsonLightMapIndex FGLTFLightMapConverter::Add(FGLTFConvertBuilder& Builder,
 
 	const bool bUseHighQualityLightMap = true;
 	const ULightMapTexture2D* Texture = LightMapInteraction.GetTexture(bUseHighQualityLightMap);
-	const FGLTFJsonTextureIndex TextureIndex = Builder.GetOrAddLightMapTexture(Texture);
+	const FGLTFJsonTextureIndex TextureIndex = Builder.GetOrAddTexture(Texture);
 
 	if (TextureIndex == INDEX_NONE)
 	{
