[GLTFExporter] feat(json): add all custom extension enum and string values

commit 9d32fb6684fc8183a3c8a583ba137c2c716e12e7
Date: Fri, 3 Jul 2020 10:30:25 +0200

    feat(json): add all custom extension enum and string values
    (cherry picked from commit a3f1093c0db3be106f91dd6efbca13d748289230​)
    # Conflicts from cherry-pick:
    #	Plugins/GLTFExporter/Source/GLTFExporter/Private/Json/GLTFJsonEnums.h
    #	Plugins/GLTFExporter/Source/GLTFExporter/Private/Json/GLTFJsonUtility.h

#skipci

diff --git a/Engine/Plugins/Enterprise/GLTFExporter/MigratedFromGit/Source/GLTFExporter/Private/Json/GLTFJsonEnums.h b/Engine/Plugins/Enterprise/GLTFExporter/MigratedFromGit/Source/GLTFExporter/Private/Json/GLTFJsonEnums.h
--- a/Engine/Plugins/Enterprise/GLTFExporter/MigratedFromGit/Source/GLTFExporter/Private/Json/GLTFJsonEnums.h
+++ b/Engine/Plugins/Enterprise/GLTFExporter/MigratedFromGit/Source/GLTFExporter/Private/Json/GLTFJsonEnums.h
@@ -10,7 +10,13 @@ enum class EGLTFJsonExtension
 	KHR_MaterialsUnlit,
 	KHR_MaterialsClearCoat,
 	KHR_MeshQuantization,
-	EPIC_LightMapTextures
+	EPIC_CameraOrbital,
+	EPIC_HDRIBackdrops,
+	EPIC_InteractionHotspots,
+	EPIC_LightmapTextures,
+	EPIC_ReflectionCaptures,
+	EPIC_LevelVariantSets,
+	EPIC_TextureHDREncoding
 };
 
 enum class EGLTFJsonShadingModel
diff --git a/Engine/Plugins/Enterprise/GLTFExporter/MigratedFromGit/Source/GLTFExporter/Private/Json/GLTFJsonNode.h b/Engine/Plugins/Enterprise/GLTFExporter/MigratedFromGit/Source/GLTFExporter/Private/Json/GLTFJsonNode.h
--- a/Engine/Plugins/Enterprise/GLTFExporter/MigratedFromGit/Source/GLTFExporter/Private/Json/GLTFJsonNode.h
+++ b/Engine/Plugins/Enterprise/GLTFExporter/MigratedFromGit/Source/GLTFExporter/Private/Json/GLTFJsonNode.h
@@ -80,7 +80,7 @@ struct FGLTFJsonNode
 
 			if (LightMap != INDEX_NONE)
 			{
-				const EGLTFJsonExtension Extension = EGLTFJsonExtension::EPIC_LightMapTextures;
+				const EGLTFJsonExtension Extension = EGLTFJsonExtension::EPIC_LightmapTextures;
 
 				Extensions.Used.Add(Extension);
 
diff --git a/Engine/Plugins/Enterprise/GLTFExporter/MigratedFromGit/Source/GLTFExporter/Private/Json/GLTFJsonRoot.h b/Engine/Plugins/Enterprise/GLTFExporter/MigratedFromGit/Source/GLTFExporter/Private/Json/GLTFJsonRoot.h
--- a/Engine/Plugins/Enterprise/GLTFExporter/MigratedFromGit/Source/GLTFExporter/Private/Json/GLTFJsonRoot.h
+++ b/Engine/Plugins/Enterprise/GLTFExporter/MigratedFromGit/Source/GLTFExporter/Private/Json/GLTFJsonRoot.h
@@ -101,7 +101,7 @@ struct FGLTFJsonRoot
 
 			if (LightMaps.Num() > 0)
 			{
-				const EGLTFJsonExtension Extension = EGLTFJsonExtension::EPIC_LightMapTextures;
+				const EGLTFJsonExtension Extension = EGLTFJsonExtension::EPIC_LightmapTextures;
 
 				AllExtensions.Used.Add(Extension);
 
diff --git a/Engine/Plugins/Enterprise/GLTFExporter/MigratedFromGit/Source/GLTFExporter/Private/Json/GLTFJsonUtility.h b/Engine/Plugins/Enterprise/GLTFExporter/MigratedFromGit/Source/GLTFExporter/Private/Json/GLTFJsonUtility.h
--- a/Engine/Plugins/Enterprise/GLTFExporter/MigratedFromGit/Source/GLTFExporter/Private/Json/GLTFJsonUtility.h
+++ b/Engine/Plugins/Enterprise/GLTFExporter/MigratedFromGit/Source/GLTFExporter/Private/Json/GLTFJsonUtility.h
@@ -17,12 +17,18 @@ struct FGLTFJsonUtility
 	{
 		switch (Value)
 		{
-			case EGLTFJsonExtension::KHR_LightsPunctual:     return TEXT("KHR_lights_punctual");
-			case EGLTFJsonExtension::KHR_MaterialsUnlit:     return TEXT("KHR_materials_unlit");
-			case EGLTFJsonExtension::KHR_MaterialsClearCoat: return TEXT("KHR_materials_clearcoat");
-			case EGLTFJsonExtension::KHR_MeshQuantization:   return TEXT("KHR_mesh_quantization");
-			case EGLTFJsonExtension::EPIC_LightMapTextures:  return TEXT("EPIC_lightmap_textures");
-			default:                                         return TEXT("unknown");
+			case EGLTFJsonExtension::KHR_LightsPunctual:       return TEXT("KHR_lights_punctual");
+			case EGLTFJsonExtension::KHR_MaterialsUnlit:       return TEXT("KHR_materials_unlit");
+			case EGLTFJsonExtension::KHR_MaterialsClearCoat:   return TEXT("KHR_materials_clearcoat");
+			case EGLTFJsonExtension::KHR_MeshQuantization:     return TEXT("KHR_mesh_quantization");
+			case EGLTFJsonExtension::EPIC_CameraOrbital:       return TEXT("EPIC_camera_orbital");
+			case EGLTFJsonExtension::EPIC_HDRIBackdrops:       return TEXT("EPIC_hdri_backdrops");
+			case EGLTFJsonExtension::EPIC_InteractionHotspots: return TEXT("EPIC_interaction_hotspots");
+			case EGLTFJsonExtension::EPIC_LightmapTextures:    return TEXT("EPIC_lightmap_textures");
+			case EGLTFJsonExtension::EPIC_ReflectionCaptures:  return TEXT("EPIC_reflection_captures");
+			case EGLTFJsonExtension::EPIC_LevelVariantSets:    return TEXT("EPIC_level_variant_sets");
+			case EGLTFJsonExtension::EPIC_TextureHDREncoding:  return TEXT("EPIC_texture_hdrencoding");
+			default:                                           return TEXT("unknown");
 		}
 	}
 
