[GLTFExporter] fix(json): change EPIC_texture_hdrencoding to EPIC_texture_hdr_encoding

commit 4c0a1689d1485b47afa1d79684560ec810dd08a3
Date: Fri, 3 Jul 2020 16:32:48 +0200

    fix(json): change EPIC_texture_hdrencoding to EPIC_texture_hdr_encoding
    (cherry picked from commit ea0ce6c1ee919e9110548ea29c18d2913178bd0d​)

#skipci

diff --git a/Engine/Plugins/Enterprise/GLTFExporter/MigratedFromGit/Source/GLTFExporter/Private/Json/GLTFJsonUtility.h b/Engine/Plugins/Enterprise/GLTFExporter/MigratedFromGit/Source/GLTFExporter/Private/Json/GLTFJsonUtility.h
--- a/Engine/Plugins/Enterprise/GLTFExporter/MigratedFromGit/Source/GLTFExporter/Private/Json/GLTFJsonUtility.h
+++ b/Engine/Plugins/Enterprise/GLTFExporter/MigratedFromGit/Source/GLTFExporter/Private/Json/GLTFJsonUtility.h
@@ -27,7 +27,7 @@ struct FGLTFJsonUtility
 			case EGLTFJsonExtension::EPIC_LightmapTextures:    return TEXT("EPIC_lightmap_textures");
 			case EGLTFJsonExtension::EPIC_ReflectionCaptures:  return TEXT("EPIC_reflection_captures");
 			case EGLTFJsonExtension::EPIC_LevelVariantSets:    return TEXT("EPIC_level_variant_sets");
-			case EGLTFJsonExtension::EPIC_TextureHDREncoding:  return TEXT("EPIC_texture_hdrencoding");
+			case EGLTFJsonExtension::EPIC_TextureHDREncoding:  return TEXT("EPIC_texture_hdr_encoding");
 			default:                                           return TEXT("unknown");
 		}
 	}
