[GLTFExporter] feat: add export options to builder

commit 84c9fb2989444feb3ea115d388e4fce3f73549e9
Date: Fri, 3 Jul 2020 21:11:35 +0200

    feat: add export options to builder
    (cherry picked from commit 4a545a82d2d9d75748ef5b6c7a606f9447583ca3​)

#skipci

diff --git a/Engine/Plugins/Enterprise/GLTFExporter/MigratedFromGit/Source/GLTFExporter/Private/Builders/GLTFBufferBuilder.cpp b/Engine/Plugins/Enterprise/GLTFExporter/MigratedFromGit/Source/GLTFExporter/Private/Builders/GLTFBufferBuilder.cpp
--- a/Engine/Plugins/Enterprise/GLTFExporter/MigratedFromGit/Source/GLTFExporter/Private/Builders/GLTFBufferBuilder.cpp
+++ b/Engine/Plugins/Enterprise/GLTFExporter/MigratedFromGit/Source/GLTFExporter/Private/Builders/GLTFBufferBuilder.cpp
@@ -4,7 +4,8 @@
 #include "Misc/FileHelper.h"
 #include "Misc/Paths.h"
 
-FGLTFBufferBuilder::FGLTFBufferBuilder()
+FGLTFBufferBuilder::FGLTFBufferBuilder(const UGLTFExportOptions* ExportOptions)
+	: FGLTFJsonBuilder(ExportOptions)
 {
 	BufferIndex = AddBuffer();
 }
diff --git a/Engine/Plugins/Enterprise/GLTFExporter/MigratedFromGit/Source/GLTFExporter/Private/Builders/GLTFBufferBuilder.h b/Engine/Plugins/Enterprise/GLTFExporter/MigratedFromGit/Source/GLTFExporter/Private/Builders/GLTFBufferBuilder.h
--- a/Engine/Plugins/Enterprise/GLTFExporter/MigratedFromGit/Source/GLTFExporter/Private/Builders/GLTFBufferBuilder.h
+++ b/Engine/Plugins/Enterprise/GLTFExporter/MigratedFromGit/Source/GLTFExporter/Private/Builders/GLTFBufferBuilder.h
@@ -8,7 +8,7 @@ class FGLTFBufferBuilder : public FGLTFJsonBuilder
 {
 public:
 
-	FGLTFBufferBuilder();
+	FGLTFBufferBuilder(const UGLTFExportOptions* ExportOptions);
 
 	FGLTFJsonBufferViewIndex AddBufferView(const void* RawData, uint64 ByteLength, const FString& Name = TEXT(""), uint8 DataAlignment = 4, EGLTFJsonBufferTarget BufferTarget = EGLTFJsonBufferTarget::ArrayBuffer);
 
diff --git a/Engine/Plugins/Enterprise/GLTFExporter/MigratedFromGit/Source/GLTFExporter/Private/Builders/GLTFBuilder.cpp b/Engine/Plugins/Enterprise/GLTFExporter/MigratedFromGit/Source/GLTFExporter/Private/Builders/GLTFBuilder.cpp
--- a/Engine/Plugins/Enterprise/GLTFExporter/MigratedFromGit/Source/GLTFExporter/Private/Builders/GLTFBuilder.cpp
+++ b/Engine/Plugins/Enterprise/GLTFExporter/MigratedFromGit/Source/GLTFExporter/Private/Builders/GLTFBuilder.cpp
@@ -8,7 +8,8 @@
 
 DEFINE_LOG_CATEGORY(LogGLTFExporter);
 
-FGLTFBuilder::FGLTFBuilder()
+FGLTFBuilder::FGLTFBuilder(const UGLTFExportOptions* ExportOptions)
+	: ExportOptions(ExportOptions)
 {
 }
 
diff --git a/Engine/Plugins/Enterprise/GLTFExporter/MigratedFromGit/Source/GLTFExporter/Private/Builders/GLTFBuilder.h b/Engine/Plugins/Enterprise/GLTFExporter/MigratedFromGit/Source/GLTFExporter/Private/Builders/GLTFBuilder.h
--- a/Engine/Plugins/Enterprise/GLTFExporter/MigratedFromGit/Source/GLTFExporter/Private/Builders/GLTFBuilder.h
+++ b/Engine/Plugins/Enterprise/GLTFExporter/MigratedFromGit/Source/GLTFExporter/Private/Builders/GLTFBuilder.h
@@ -2,6 +2,8 @@
 
 #pragma once
 
+#include "GLTFExportOptions.h"
+
 DECLARE_LOG_CATEGORY_EXTERN(LogGLTFExporter, Log, All);
 
 enum class EGLTFMessageSeverity
@@ -15,7 +17,9 @@ class FGLTFBuilder
 {
 public:
 
-	FGLTFBuilder();
+	const UGLTFExportOptions* const ExportOptions;
+
+	FGLTFBuilder(const UGLTFExportOptions* ExportOptions);
 	virtual ~FGLTFBuilder();
 
 	typedef TTuple<EGLTFMessageSeverity, FString> FLogMessage;
diff --git a/Engine/Plugins/Enterprise/GLTFExporter/MigratedFromGit/Source/GLTFExporter/Private/Builders/GLTFContainerBuilder.cpp b/Engine/Plugins/Enterprise/GLTFExporter/MigratedFromGit/Source/GLTFExporter/Private/Builders/GLTFContainerBuilder.cpp
--- a/Engine/Plugins/Enterprise/GLTFExporter/MigratedFromGit/Source/GLTFExporter/Private/Builders/GLTFContainerBuilder.cpp
+++ b/Engine/Plugins/Enterprise/GLTFExporter/MigratedFromGit/Source/GLTFExporter/Private/Builders/GLTFContainerBuilder.cpp
@@ -2,7 +2,7 @@
 
 #include "Builders/GLTFContainerBuilder.h"
 
-FGLTFContainerBuilder::FGLTFContainerBuilder(bool bSelectedActorsOnly)
-	: FGLTFConvertBuilder(bSelectedActorsOnly)
+FGLTFContainerBuilder::FGLTFContainerBuilder(const UGLTFExportOptions* ExportOptions, bool bSelectedActorsOnly)
+	: FGLTFConvertBuilder(ExportOptions, bSelectedActorsOnly)
 {
 }
diff --git a/Engine/Plugins/Enterprise/GLTFExporter/MigratedFromGit/Source/GLTFExporter/Private/Builders/GLTFContainerBuilder.h b/Engine/Plugins/Enterprise/GLTFExporter/MigratedFromGit/Source/GLTFExporter/Private/Builders/GLTFContainerBuilder.h
--- a/Engine/Plugins/Enterprise/GLTFExporter/MigratedFromGit/Source/GLTFExporter/Private/Builders/GLTFContainerBuilder.h
+++ b/Engine/Plugins/Enterprise/GLTFExporter/MigratedFromGit/Source/GLTFExporter/Private/Builders/GLTFContainerBuilder.h
@@ -8,5 +8,5 @@ class FGLTFContainerBuilder : public FGLTFConvertBuilder
 {
 public:
 
-	FGLTFContainerBuilder(bool bSelectedActorsOnly);
+	FGLTFContainerBuilder(const UGLTFExportOptions* ExportOptions, bool bSelectedActorsOnly);
 };
diff --git a/Engine/Plugins/Enterprise/GLTFExporter/MigratedFromGit/Source/GLTFExporter/Private/Builders/GLTFConvertBuilder.cpp b/Engine/Plugins/Enterprise/GLTFExporter/MigratedFromGit/Source/GLTFExporter/Private/Builders/GLTFConvertBuilder.cpp
--- a/Engine/Plugins/Enterprise/GLTFExporter/MigratedFromGit/Source/GLTFExporter/Private/Builders/GLTFConvertBuilder.cpp
+++ b/Engine/Plugins/Enterprise/GLTFExporter/MigratedFromGit/Source/GLTFExporter/Private/Builders/GLTFConvertBuilder.cpp
@@ -3,8 +3,9 @@
 #include "Builders/GLTFConvertBuilder.h"
 #include "GLTFBuilderUtility.h"
 
-FGLTFConvertBuilder::FGLTFConvertBuilder(bool bSelectedActorsOnly)
-	: bSelectedActorsOnly(bSelectedActorsOnly)
+FGLTFConvertBuilder::FGLTFConvertBuilder(const UGLTFExportOptions* ExportOptions, bool bSelectedActorsOnly)
+	: FGLTFImageBuilder(ExportOptions)
+	, bSelectedActorsOnly(bSelectedActorsOnly)
 {
 }
 
diff --git a/Engine/Plugins/Enterprise/GLTFExporter/MigratedFromGit/Source/GLTFExporter/Private/Builders/GLTFConvertBuilder.h b/Engine/Plugins/Enterprise/GLTFExporter/MigratedFromGit/Source/GLTFExporter/Private/Builders/GLTFConvertBuilder.h
--- a/Engine/Plugins/Enterprise/GLTFExporter/MigratedFromGit/Source/GLTFExporter/Private/Builders/GLTFConvertBuilder.h
+++ b/Engine/Plugins/Enterprise/GLTFExporter/MigratedFromGit/Source/GLTFExporter/Private/Builders/GLTFConvertBuilder.h
@@ -16,7 +16,7 @@ public:
 
 	const bool bSelectedActorsOnly;
 
-	FGLTFConvertBuilder(bool bSelectedActorsOnly);
+	FGLTFConvertBuilder(const UGLTFExportOptions* ExportOptions, bool bSelectedActorsOnly);
 
 	FGLTFJsonAccessorIndex GetOrAddPositionAccessor(const FPositionVertexBuffer* VertexBuffer, const FString& DesiredName = TEXT(""));
 	FGLTFJsonAccessorIndex GetOrAddColorAccessor(const FColorVertexBuffer* VertexBuffer, const FString& DesiredName = TEXT(""));
diff --git a/Engine/Plugins/Enterprise/GLTFExporter/MigratedFromGit/Source/GLTFExporter/Private/Builders/GLTFImageBuilder.cpp b/Engine/Plugins/Enterprise/GLTFExporter/MigratedFromGit/Source/GLTFExporter/Private/Builders/GLTFImageBuilder.cpp
--- a/Engine/Plugins/Enterprise/GLTFExporter/MigratedFromGit/Source/GLTFExporter/Private/Builders/GLTFImageBuilder.cpp
+++ b/Engine/Plugins/Enterprise/GLTFExporter/MigratedFromGit/Source/GLTFExporter/Private/Builders/GLTFImageBuilder.cpp
@@ -5,6 +5,11 @@
 #include "IImageWrapper.h"
 #include "Engine/Texture.h"
 
+FGLTFImageBuilder::FGLTFImageBuilder(const UGLTFExportOptions* ExportOptions)
+	: FGLTFBufferBuilder(ExportOptions)
+{
+}
+
 FGLTFJsonImageIndex FGLTFImageBuilder::AddImage(const void* RawData, int64 ByteLength, int32 Width, int32 Height, ERGBFormat RawFormat, int32 BitDepth, bool bFloatFormat, const FString& Name, EGLTFJsonMimeType MimeType, int32 Quality)
 {
 	TArray64<uint8> ImageData;
diff --git a/Engine/Plugins/Enterprise/GLTFExporter/MigratedFromGit/Source/GLTFExporter/Private/Builders/GLTFImageBuilder.h b/Engine/Plugins/Enterprise/GLTFExporter/MigratedFromGit/Source/GLTFExporter/Private/Builders/GLTFImageBuilder.h
--- a/Engine/Plugins/Enterprise/GLTFExporter/MigratedFromGit/Source/GLTFExporter/Private/Builders/GLTFImageBuilder.h
+++ b/Engine/Plugins/Enterprise/GLTFExporter/MigratedFromGit/Source/GLTFExporter/Private/Builders/GLTFImageBuilder.h
@@ -10,6 +10,8 @@ class FGLTFImageBuilder : public FGLTFBufferBuilder
 {
 public:
 
+	FGLTFImageBuilder(const UGLTFExportOptions* ExportOptions);
+
 	FGLTFJsonImageIndex AddImage(const void* RawData, int64 ByteLength, int32 Width, int32 Height, ERGBFormat RawFormat, int32 BitDepth, bool bFloatFormat, const FString& Name = TEXT(""), EGLTFJsonMimeType MimeType = EGLTFJsonMimeType::PNG, int32 Quality = 100);
 
 	FGLTFJsonImageIndex AddImage(const void* RawData, int64 ByteLength, int32 Width, int32 Height, EPixelFormat PixelFormat = EPixelFormat::PF_B8G8R8A8, const FString& Name = TEXT(""), EGLTFJsonMimeType MimeType = EGLTFJsonMimeType::PNG, int32 Quality = 100);
diff --git a/Engine/Plugins/Enterprise/GLTFExporter/MigratedFromGit/Source/GLTFExporter/Private/Builders/GLTFJsonBuilder.cpp b/Engine/Plugins/Enterprise/GLTFExporter/MigratedFromGit/Source/GLTFExporter/Private/Builders/GLTFJsonBuilder.cpp
--- a/Engine/Plugins/Enterprise/GLTFExporter/MigratedFromGit/Source/GLTFExporter/Private/Builders/GLTFJsonBuilder.cpp
+++ b/Engine/Plugins/Enterprise/GLTFExporter/MigratedFromGit/Source/GLTFExporter/Private/Builders/GLTFJsonBuilder.cpp
@@ -2,8 +2,9 @@
 
 #include "Builders/GLTFJsonBuilder.h"
 
-FGLTFJsonBuilder::FGLTFJsonBuilder()
-	: DefaultScene(JsonRoot.DefaultScene)
+FGLTFJsonBuilder::FGLTFJsonBuilder(const UGLTFExportOptions* ExportOptions)
+	: FGLTFBuilder(ExportOptions)
+	, DefaultScene(JsonRoot.DefaultScene)
 {
 }
 
diff --git a/Engine/Plugins/Enterprise/GLTFExporter/MigratedFromGit/Source/GLTFExporter/Private/Builders/GLTFJsonBuilder.h b/Engine/Plugins/Enterprise/GLTFExporter/MigratedFromGit/Source/GLTFExporter/Private/Builders/GLTFJsonBuilder.h
--- a/Engine/Plugins/Enterprise/GLTFExporter/MigratedFromGit/Source/GLTFExporter/Private/Builders/GLTFJsonBuilder.h
+++ b/Engine/Plugins/Enterprise/GLTFExporter/MigratedFromGit/Source/GLTFExporter/Private/Builders/GLTFJsonBuilder.h
@@ -11,7 +11,7 @@ public:
 
 	FGLTFJsonSceneIndex& DefaultScene;
 
-	FGLTFJsonBuilder();
+	FGLTFJsonBuilder(const UGLTFExportOptions* ExportOptions);
 
 	FGLTFJsonAccessorIndex AddAccessor(const FGLTFJsonAccessor& JsonAccessor = {});
 	FGLTFJsonBufferIndex AddBuffer(const FGLTFJsonBuffer& JsonBuffer = {});
diff --git a/Engine/Plugins/Enterprise/GLTFExporter/MigratedFromGit/Source/GLTFExporter/Private/Exporters/GLTFExporter.cpp b/Engine/Plugins/Enterprise/GLTFExporter/MigratedFromGit/Source/GLTFExporter/Private/Exporters/GLTFExporter.cpp
--- a/Engine/Plugins/Enterprise/GLTFExporter/MigratedFromGit/Source/GLTFExporter/Private/Exporters/GLTFExporter.cpp
+++ b/Engine/Plugins/Enterprise/GLTFExporter/MigratedFromGit/Source/GLTFExporter/Private/Exporters/GLTFExporter.cpp
@@ -31,10 +31,9 @@ bool UGLTFExporter::ExportBinary(UObject* Object, const TCHAR* Type, FArchive& A
 		return false;
 	}
 
-	FGLTFContainerBuilder Builder(bSelectedOnly);
-
 	bool bSuccess = true;
 
+	FGLTFContainerBuilder Builder(Options, bSelectedOnly);
 	if (AddObject(Builder, Object))
 	{
 		if (!Builder.Serialize(Archive, CurrentFilename))
