[GLTFExporter] refactor(lightmap): minor cosmetic changes

commit 5354f780d56ed9a422f7c66fe6140d06bc99f7dc
Date: Fri, 3 Jul 2020 18:33:23 +0200

    refactor(lightmap): minor cosmetic changes
    (cherry picked from commit 7c5f933bf9ee3d1b0594152de62bbefe28ed63a0​)

#skipci

diff --git a/Engine/Plugins/Enterprise/GLTFExporter/MigratedFromGit/Source/GLTFExporter/Private/Converters/GLTFLevelConverters.cpp b/Engine/Plugins/Enterprise/GLTFExporter/MigratedFromGit/Source/GLTFExporter/Private/Converters/GLTFLevelConverters.cpp
--- a/Engine/Plugins/Enterprise/GLTFExporter/MigratedFromGit/Source/GLTFExporter/Private/Converters/GLTFLevelConverters.cpp
+++ b/Engine/Plugins/Enterprise/GLTFExporter/MigratedFromGit/Source/GLTFExporter/Private/Converters/GLTFLevelConverters.cpp
@@ -33,10 +33,12 @@ FGLTFJsonNodeIndex FGLTFSceneComponentConverter::Add(FGLTFConvertBuilder& Builde
 
 	if (const UStaticMeshComponent* StaticMeshComponent = Cast<UStaticMeshComponent>(SceneComponent))
 	{
-		// TODO: make export of light-maps optional (selectable via export options)
-
 		Node.Mesh = Builder.GetOrAddMesh(StaticMeshComponent);
-		Node.LightMap = Builder.GetOrAddLightMap(StaticMeshComponent);
+
+		if (true /* TODO: make export of light-maps optional (selectable via export options) */)
+		{
+			Node.LightMap = Builder.GetOrAddLightMap(StaticMeshComponent);
+		}
 	}
 
 	return NodeIndex;
diff --git a/Engine/Plugins/Enterprise/GLTFExporter/MigratedFromGit/Source/GLTFExporter/Private/Json/GLTFJsonRoot.h b/Engine/Plugins/Enterprise/GLTFExporter/MigratedFromGit/Source/GLTFExporter/Private/Json/GLTFJsonRoot.h
--- a/Engine/Plugins/Enterprise/GLTFExporter/MigratedFromGit/Source/GLTFExporter/Private/Json/GLTFJsonRoot.h
+++ b/Engine/Plugins/Enterprise/GLTFExporter/MigratedFromGit/Source/GLTFExporter/Private/Json/GLTFJsonRoot.h
@@ -93,22 +93,17 @@ struct FGLTFJsonRoot
 		FGLTFJsonUtility::WriteObjectArray(JsonWriter, TEXT("scenes"), Scenes, AllExtensions);
 		FGLTFJsonUtility::WriteObjectArray(JsonWriter, TEXT("textures"), Textures, AllExtensions);
 
-		const bool bWriteExtensions = LightMaps.Num() > 0;
-
-		if (bWriteExtensions)
+		if (LightMaps.Num() > 0 /* TODO: add more extension support */)
 		{
 			JsonWriter.WriteObjectStart(TEXT("extensions"));
 
 			if (LightMaps.Num() > 0)
 			{
 				const EGLTFJsonExtension Extension = EGLTFJsonExtension::EPIC_LightmapTextures;
-
 				AllExtensions.Used.Add(Extension);
 
 				JsonWriter.WriteObjectStart(FGLTFJsonUtility::ToString(Extension));
-
 				FGLTFJsonUtility::WriteObjectArray(JsonWriter, TEXT("lightmaps"), LightMaps, AllExtensions);
-
 				JsonWriter.WriteObjectEnd();
 			}
 
