[GLTFExporter] refactor(lightmap): modify data-structure of exported light-maps to match the latest specification

commit 541fe3b10d394cc869a4f23a7b5c0d536cfb0118
Date: Mon, 6 Jul 2020 16:16:22 +0200

    refactor(lightmap): modify data-structure of exported light-maps to match the latest specification
    (cherry picked from commit aa9cc1e513e5cb07d2b05adde91e49a042d36c3c​)

#skipci

diff --git a/Engine/Plugins/Enterprise/GLTFExporter/MigratedFromGit/Source/GLTFExporter/Private/Converters/GLTFLightMapConverters.cpp b/Engine/Plugins/Enterprise/GLTFExporter/MigratedFromGit/Source/GLTFExporter/Private/Converters/GLTFLightMapConverters.cpp
--- a/Engine/Plugins/Enterprise/GLTFExporter/MigratedFromGit/Source/GLTFExporter/Private/Converters/GLTFLightMapConverters.cpp
+++ b/Engine/Plugins/Enterprise/GLTFExporter/MigratedFromGit/Source/GLTFExporter/Private/Converters/GLTFLightMapConverters.cpp
@@ -68,10 +68,9 @@ FGLTFJsonLightMapIndex FGLTFLightMapConverter::Add(FGLTFConvertBuilder& Builder,
 	JsonLightMap.Name = Name.IsEmpty() ? Owner->GetName() + TEXT("_") + StaticMeshComponent->GetName() : Name;
 	JsonLightMap.Texture.Index = TextureIndex;
 	JsonLightMap.Texture.TexCoord = CoordinateIndex;
-	JsonLightMap.ValueScale = { LightMapScale.X, LightMapScale.Y, LightMapScale.Z, LightMapScale.W };
-	JsonLightMap.ValueOffset = { LightMapAdd.X, LightMapAdd.Y, LightMapAdd.Z, LightMapAdd.W };
-	JsonLightMap.CoordinateScale = { CoordinateScale.X, CoordinateScale.Y };
-	JsonLightMap.CoordinateOffset = { CoordinateBias.X, CoordinateBias.Y };
+	JsonLightMap.LightMapScale = { LightMapScale.X, LightMapScale.Y, LightMapScale.Z, LightMapScale.W };
+	JsonLightMap.LightMapAdd = { LightMapAdd.X, LightMapAdd.Y, LightMapAdd.Z, LightMapAdd.W };
+	JsonLightMap.CoordinateScaleBias = { CoordinateScale.X, CoordinateScale.Y, CoordinateBias.X, CoordinateBias.Y };
 
 	return Builder.AddLightMap(JsonLightMap);
 }
diff --git a/Engine/Plugins/Enterprise/GLTFExporter/MigratedFromGit/Source/GLTFExporter/Private/Json/GLTFJsonLightMap.h b/Engine/Plugins/Enterprise/GLTFExporter/MigratedFromGit/Source/GLTFExporter/Private/Json/GLTFJsonLightMap.h
--- a/Engine/Plugins/Enterprise/GLTFExporter/MigratedFromGit/Source/GLTFExporter/Private/Json/GLTFJsonLightMap.h
+++ b/Engine/Plugins/Enterprise/GLTFExporter/MigratedFromGit/Source/GLTFExporter/Private/Json/GLTFJsonLightMap.h
@@ -3,24 +3,20 @@
 #pragma once
 
 #include "Json/GLTFJsonMaterial.h"
-#include "Json/GLTFJsonVector2.h"
 #include "Json/GLTFJsonVector4.h"
-#include "Serialization/JsonSerializer.h"
 
 struct FGLTFJsonLightMap
 {
 	FString              Name;
 	FGLTFJsonTextureInfo Texture;
-	FGLTFJsonVector4     ValueScale;
-	FGLTFJsonVector4     ValueOffset;
-	FGLTFJsonVector2     CoordinateScale;
-	FGLTFJsonVector2     CoordinateOffset;
+	FGLTFJsonVector4     LightMapScale;
+	FGLTFJsonVector4     LightMapAdd;
+	FGLTFJsonVector4     CoordinateScaleBias;
 
 	FGLTFJsonLightMap()
-		: ValueScale(FGLTFJsonVector4::One)
-		, ValueOffset(FGLTFJsonVector4::Zero)
-		, CoordinateScale(FGLTFJsonVector2::One)
-		, CoordinateOffset(FGLTFJsonVector2::Zero)
+		: LightMapScale(FGLTFJsonVector4::One)
+		, LightMapAdd(FGLTFJsonVector4::Zero)
+		, CoordinateScaleBias({ 1, 1, 0, 0 })
 	{
 	}
 
@@ -40,17 +36,14 @@ struct FGLTFJsonLightMap
 			Texture.WriteObject(JsonWriter, Extensions);
 		}
 
-		JsonWriter.WriteIdentifierPrefix(TEXT("valueScale"));
-		ValueScale.WriteArray(JsonWriter);
+		JsonWriter.WriteIdentifierPrefix(TEXT("lightmapScale"));
+		LightMapScale.WriteArray(JsonWriter);
 
-		JsonWriter.WriteIdentifierPrefix(TEXT("valueOffset"));
-		ValueOffset.WriteArray(JsonWriter);
+		JsonWriter.WriteIdentifierPrefix(TEXT("lightmapAdd"));
+		LightMapAdd.WriteArray(JsonWriter);
 
-		JsonWriter.WriteIdentifierPrefix(TEXT("coordinateScale"));
-		CoordinateScale.WriteArray(JsonWriter);
-
-		JsonWriter.WriteIdentifierPrefix(TEXT("coordinateOffset"));
-		CoordinateOffset.WriteArray(JsonWriter);
+		JsonWriter.WriteIdentifierPrefix(TEXT("coordinateScaleBias"));
+		CoordinateScaleBias.WriteArray(JsonWriter);
 
 		JsonWriter.WriteObjectEnd();
 	}
