[GLTFExporter] chore: add dependency for VariantManagerContent to the plugin

commit 8806fe58add2721be06321e36e8f4324f1d3ae0d
Date: Sat, 4 Jul 2020 16:14:03 +0200

    chore: add dependency for VariantManagerContent to the plugin
    (cherry picked from commit aadf6ac64d707558e1fbc358f50981481c8b67b7​)

#skipci

diff --git a/Engine/Plugins/Enterprise/GLTFExporter/MigratedFromGit/GLTFExporter.uplugin b/Engine/Plugins/Enterprise/GLTFExporter/MigratedFromGit/GLTFExporter.uplugin
--- a/Engine/Plugins/Enterprise/GLTFExporter/MigratedFromGit/GLTFExporter.uplugin
+++ b/Engine/Plugins/Enterprise/GLTFExporter/MigratedFromGit/GLTFExporter.uplugin
@@ -25,5 +25,11 @@
 			"Type": "Runtime",
 			"LoadingPhase": "Default"
 		}
+	],
+	"Plugins": [
+		{
+			"Name": "VariantManagerContent",
+			"Enabled": true
+		}
 	]
 }
diff --git a/Engine/Plugins/Enterprise/GLTFExporter/MigratedFromGit/Source/GLTFExporter/GLTFExporter.Build.cs b/Engine/Plugins/Enterprise/GLTFExporter/MigratedFromGit/Source/GLTFExporter/GLTFExporter.Build.cs
--- a/Engine/Plugins/Enterprise/GLTFExporter/MigratedFromGit/Source/GLTFExporter/GLTFExporter.Build.cs
+++ b/Engine/Plugins/Enterprise/GLTFExporter/MigratedFromGit/Source/GLTFExporter/GLTFExporter.Build.cs
@@ -27,7 +27,8 @@ namespace UnrealBuildTool.Rules
 					"Projects",
 					"RenderCore",
 					"RHI",
-					"MaterialBaking"
+					"MaterialBaking",
+					"VariantManagerContent"
 				}
 				);
 		}
