[GLTFExporter] feat(variants): basic implementation of variant-export

commit ce883230486fa4a877169f8d439a7fa44c9e37cd
Date: Sat, 4 Jul 2020 16:25:57 +0200

    feat(variants): basic implementation of variant-export
    (cherry picked from commit 22b381d9f632955759ba59b334beb60e09e1f706​)
    # Conflicts from cherry-pick:
    #	Plugins/GLTFExporter/Source/GLTFExporter/Private/Builders/GLTFConvertBuilder.h
    #	Plugins/GLTFExporter/Source/GLTFExporter/Private/Json/GLTFJsonIndex.h
    #	Plugins/GLTFExporter/Source/GLTFExporter/Private/Json/GLTFJsonRoot.h

#skipci

diff --git a/Engine/Plugins/Enterprise/GLTFExporter/MigratedFromGit/Source/GLTFExporter/Private/Builders/GLTFConvertBuilder.cpp b/Engine/Plugins/Enterprise/GLTFExporter/MigratedFromGit/Source/GLTFExporter/Private/Builders/GLTFConvertBuilder.cpp
--- a/Engine/Plugins/Enterprise/GLTFExporter/MigratedFromGit/Source/GLTFExporter/Private/Builders/GLTFConvertBuilder.cpp
+++ b/Engine/Plugins/Enterprise/GLTFExporter/MigratedFromGit/Source/GLTFExporter/Private/Builders/GLTFConvertBuilder.cpp
@@ -183,3 +183,13 @@ FGLTFJsonSceneIndex FGLTFConvertBuilder::GetOrAddScene(const UWorld* World, cons
 
 	return GetOrAddScene(World->PersistentLevel, DesiredName.IsEmpty() ? World->GetName() : DesiredName);
 }
+
+FGLTFJsonLevelVariantSetsIndex FGLTFConvertBuilder::GetOrAddLevelVariantSets(const ALevelVariantSetsActor* LevelVariantSetsActor, const FString& DesiredName)
+{
+	if (LevelVariantSetsActor == nullptr)
+	{
+		return FGLTFJsonLevelVariantSetsIndex(INDEX_NONE);
+	}
+
+	return LevelVariantSetsConverter.GetOrAdd(*this, DesiredName, LevelVariantSetsActor);
+}
diff --git a/Engine/Plugins/Enterprise/GLTFExporter/MigratedFromGit/Source/GLTFExporter/Private/Builders/GLTFConvertBuilder.h b/Engine/Plugins/Enterprise/GLTFExporter/MigratedFromGit/Source/GLTFExporter/Private/Builders/GLTFConvertBuilder.h
--- a/Engine/Plugins/Enterprise/GLTFExporter/MigratedFromGit/Source/GLTFExporter/Private/Builders/GLTFConvertBuilder.h
+++ b/Engine/Plugins/Enterprise/GLTFExporter/MigratedFromGit/Source/GLTFExporter/Private/Builders/GLTFConvertBuilder.h
@@ -9,6 +9,7 @@
 #include "Converters/GLTFTextureConverters.h"
 #include "Converters/GLTFLevelConverters.h"
 #include "Converters/GLTFLightMapConverters.h"
+#include "Converters/GLTFVariantSetConverters.h"
 
 class FGLTFConvertBuilder : public FGLTFImageBuilder
 {
@@ -38,6 +39,7 @@ public:
 	FGLTFJsonNodeIndex GetOrAddNode(const AActor* Actor, const FString& DesiredName = TEXT(""));
 	FGLTFJsonSceneIndex GetOrAddScene(const ULevel* Level, const FString& DesiredName = TEXT(""));
 	FGLTFJsonSceneIndex GetOrAddScene(const UWorld* World, const FString& DesiredName = TEXT(""));
+	FGLTFJsonLevelVariantSetsIndex GetOrAddLevelVariantSets(const ALevelVariantSetsActor* LevelVariantSetsActor, const FString& DesiredName = TEXT(""));
 
 private:
 
@@ -59,4 +61,5 @@ private:
 	FGLTFSceneComponentConverter SceneComponentConverter;
 	FGLTFActorConverter ActorConverter;
 	FGLTFLevelConverter LevelConverter;
+	FGLTFLevelVariantSetsConverter LevelVariantSetsConverter;
 };
diff --git a/Engine/Plugins/Enterprise/GLTFExporter/MigratedFromGit/Source/GLTFExporter/Private/Builders/GLTFJsonBuilder.cpp b/Engine/Plugins/Enterprise/GLTFExporter/MigratedFromGit/Source/GLTFExporter/Private/Builders/GLTFJsonBuilder.cpp
--- a/Engine/Plugins/Enterprise/GLTFExporter/MigratedFromGit/Source/GLTFExporter/Private/Builders/GLTFJsonBuilder.cpp
+++ b/Engine/Plugins/Enterprise/GLTFExporter/MigratedFromGit/Source/GLTFExporter/Private/Builders/GLTFJsonBuilder.cpp
@@ -75,6 +75,11 @@ FGLTFJsonNodeIndex FGLTFJsonBuilder::AddChildNode(FGLTFJsonNodeIndex ParentIndex
 	return ChildIndex;
 }
 
+FGLTFJsonLevelVariantSetsIndex FGLTFJsonBuilder::AddLevelVariantSets(const FGLTFJsonLevelVariantSets& LevelVariantSets)
+{
+	return FGLTFJsonLevelVariantSetsIndex(JsonRoot.LevelVariantSets.Add(LevelVariantSets));
+}
+
 FGLTFJsonAccessor& FGLTFJsonBuilder::GetAccessor(FGLTFJsonAccessorIndex AccessorIndex)
 {
 	return JsonRoot.Accessors[AccessorIndex];
diff --git a/Engine/Plugins/Enterprise/GLTFExporter/MigratedFromGit/Source/GLTFExporter/Private/Builders/GLTFJsonBuilder.h b/Engine/Plugins/Enterprise/GLTFExporter/MigratedFromGit/Source/GLTFExporter/Private/Builders/GLTFJsonBuilder.h
--- a/Engine/Plugins/Enterprise/GLTFExporter/MigratedFromGit/Source/GLTFExporter/Private/Builders/GLTFJsonBuilder.h
+++ b/Engine/Plugins/Enterprise/GLTFExporter/MigratedFromGit/Source/GLTFExporter/Private/Builders/GLTFJsonBuilder.h
@@ -26,6 +26,7 @@ public:
 	FGLTFJsonLightMapIndex AddLightMap(const FGLTFJsonLightMap& JsonLightMap = {});
 
 	FGLTFJsonNodeIndex AddChildNode(FGLTFJsonNodeIndex ParentNodeIndex, const FGLTFJsonNode& JsonNode = {});
+	FGLTFJsonLevelVariantSetsIndex AddLevelVariantSets(const FGLTFJsonLevelVariantSets& LevelVariantSets = {});
 
 	FGLTFJsonAccessor& GetAccessor(FGLTFJsonAccessorIndex AccessorIndex);
 	FGLTFJsonBuffer& GetBuffer(FGLTFJsonBufferIndex BufferIndex);
diff --git a/Engine/Plugins/Enterprise/GLTFExporter/MigratedFromGit/Source/GLTFExporter/Private/Converters/GLTFLevelConverters.cpp b/Engine/Plugins/Enterprise/GLTFExporter/MigratedFromGit/Source/GLTFExporter/Private/Converters/GLTFLevelConverters.cpp
--- a/Engine/Plugins/Enterprise/GLTFExporter/MigratedFromGit/Source/GLTFExporter/Private/Converters/GLTFLevelConverters.cpp
+++ b/Engine/Plugins/Enterprise/GLTFExporter/MigratedFromGit/Source/GLTFExporter/Private/Converters/GLTFLevelConverters.cpp
@@ -86,12 +86,21 @@ FGLTFJsonSceneIndex FGLTFLevelConverter::Add(FGLTFConvertBuilder& Builder, const
 
 	for (const AActor* Actor : Level->Actors)
 	{
-		// TODO: add support for ALevelVariantSetsActor
-
-		const FGLTFJsonNodeIndex NodeIndex = Builder.GetOrAddNode(Actor);
-		if (NodeIndex != INDEX_NONE && FGLTFActorUtility::IsRootActor(Actor, Builder.bSelectedActorsOnly))
+		if (const ALevelVariantSetsActor *LevelVariantSetsActor = Cast<ALevelVariantSetsActor>(Actor))
 		{
-			Scene.Nodes.Add(NodeIndex);
+			const FGLTFJsonLevelVariantSetsIndex LevelVariantSetsIndex = Builder.GetOrAddLevelVariantSets(LevelVariantSetsActor);
+			if (LevelVariantSetsIndex != INDEX_NONE)
+			{
+				Scene.LevelVariantSets.Add(LevelVariantSetsIndex);
+			}
+		}
+		else
+		{
+			const FGLTFJsonNodeIndex NodeIndex = Builder.GetOrAddNode(Actor);
+			if (NodeIndex != INDEX_NONE && FGLTFActorUtility::IsRootActor(Actor, Builder.bSelectedActorsOnly))
+			{
+				Scene.Nodes.Add(NodeIndex);
+			}
 		}
 	}
 
diff --git a/Engine/Plugins/Enterprise/GLTFExporter/MigratedFromGit/Source/GLTFExporter/Private/Converters/GLTFVariantSetConverters.cpp b/Engine/Plugins/Enterprise/GLTFExporter/MigratedFromGit/Source/GLTFExporter/Private/Converters/GLTFVariantSetConverters.cpp
new file mode 100644
--- /dev/null
+++ b/Engine/Plugins/Enterprise/GLTFExporter/MigratedFromGit/Source/GLTFExporter/Private/Converters/GLTFVariantSetConverters.cpp
@@ -0,0 +1,154 @@
+// Copyright Epic Games, Inc. All Rights Reserved.
+
+#include "Converters/GLTFVariantSetConverters.h"
+#include "Builders/GLTFContainerBuilder.h"
+#include "LevelVariantSetsActor.h"
+#include "VariantObjectBinding.h"
+#include "LevelVariantSets.h"
+#include "PropertyValue.h"
+#include "VariantSet.h"
+#include "Variant.h"
+
+FGLTFJsonLevelVariantSetsIndex FGLTFLevelVariantSetsConverter::Add(FGLTFConvertBuilder& Builder, const FString& Name, const ALevelVariantSetsActor* LevelVariantSetsActor)
+{
+	const ULevelVariantSets* LevelVariantSets = const_cast<ALevelVariantSetsActor*>(LevelVariantSetsActor)->GetLevelVariantSets(true);
+	if (LevelVariantSets == nullptr)
+	{
+		return FGLTFJsonLevelVariantSetsIndex(INDEX_NONE);
+	}
+
+	FGLTFJsonLevelVariantSets JsonLevelVariantSets;
+	JsonLevelVariantSets.Name = Name.IsEmpty() ? LevelVariantSets->GetName() : Name;
+
+	for (const UVariantSet* VariantSet: LevelVariantSets->GetVariantSets())
+	{
+		FGLTFJsonVariantSet JsonVariantSet;
+		JsonVariantSet.Name = VariantSet->GetDisplayText().ToString();
+
+		for (const UVariant* Variant: VariantSet->GetVariants())
+		{
+			const FString ThumbnailPrefix = JsonLevelVariantSets.Name + TEXT("_") + JsonVariantSet.Name + TEXT("_");
+
+			FGLTFJsonVariant JsonVariant;
+			if (TryParseJsonVariant(Builder, JsonVariant, Variant, ThumbnailPrefix))
+			{
+				JsonVariantSet.Variants.Add(JsonVariant);
+			}
+		}
+
+		// TODO: allow sets without any variants?
+		if (JsonVariantSet.Variants.Num() > 0)
+		{
+			JsonLevelVariantSets.VariantSets.Add(JsonVariantSet);
+		}
+	}
+
+	if (JsonLevelVariantSets.VariantSets.Num() == 0)
+	{
+		// TODO: allow level variant sets without any variant sets?
+		return FGLTFJsonLevelVariantSetsIndex(INDEX_NONE);
+	}
+
+	return Builder.AddLevelVariantSets(JsonLevelVariantSets);
+}
+
+bool FGLTFLevelVariantSetsConverter::TryParseJsonVariant(FGLTFConvertBuilder& Builder, FGLTFJsonVariant& OutVariant, const UVariant* Variant, const FString& ThumbnailPrefix) const
+{
+	FGLTFJsonVariant JsonVariant;
+
+	for (const UVariantObjectBinding* Binding: Variant->GetBindings())
+	{
+		FGLTFJsonVariantNode JsonVariantNode;
+		if (TryParseJsonVariantNode(Builder, JsonVariantNode, Binding))
+		{
+			JsonVariant.Nodes.Add(JsonVariantNode);
+		}
+	}
+
+	if (JsonVariant.Nodes.Num() == 0)
+	{
+		// TODO: allow variants without nodes?
+		return false;
+	}
+
+	JsonVariant.Name = Variant->GetDisplayText().ToString();
+	JsonVariant.bIsActive = const_cast<UVariant*>(Variant)->IsActive();
+
+	if (const UTexture2D* Thumbnail = const_cast<UVariant*>(Variant)->GetThumbnail())
+	{
+		const FString ThumbnailName = ThumbnailPrefix + JsonVariant.Name + TEXT("_thumbnail");
+		JsonVariant.Thumbnail = Builder.AddImage(Thumbnail->Source, ThumbnailName);
+	}
+
+	OutVariant = JsonVariant;
+	return true;
+}
+
+bool FGLTFLevelVariantSetsConverter::TryParseJsonVariantNode(FGLTFConvertBuilder& Builder, FGLTFJsonVariantNode& OutVariantNode, const UVariantObjectBinding* Binding) const
+{
+	FGLTFJsonVariantNode JsonVariantNode;
+	bool bHasAnyProperty = false;
+
+	for (UPropertyValue* Property: Binding->GetCapturedProperties())
+	{
+		if (Property->Resolve())
+		{
+			// TODO: should we use another way to identify the properties we are interested in?
+			// Right now we use the property-catagory and property-class, but we might use the class
+			// of the parent container and the display-name of the property instead.
+			// Or something completely different...
+
+			const EPropertyValueCategory PropCategory = Property->GetPropCategory();
+			const FFieldClass* PropClass = Property->GetPropertyClass();
+
+			if (PropCategory == EPropertyValueCategory::Visibility && PropClass == FBoolProperty::StaticClass())
+			{
+				bool bIsVisible;
+				if (TryGetPropertyValue(Property, bIsVisible))
+				{
+					JsonVariantNode.bIsVisible = bIsVisible;
+					bHasAnyProperty = true;
+				}
+			}
+			else
+			{
+				// TODO: handle more properties
+			}
+		}
+		else
+		{
+			// TODO: handle failure to resolve property (or force resolve it somehow)
+		}
+	}
+
+	if (!bHasAnyProperty)
+	{
+		// TODO: allow variant nodes without any properties?
+		return false;
+	}
+
+	const AActor* Actor = Cast<AActor>(Binding->GetObject());
+	const FGLTFJsonNodeIndex Node = Builder.GetOrAddNode(Actor);
+
+	if (Node == INDEX_NONE)
+	{
+		// TODO: report warning about invalid actor / node?
+		return false;
+	}
+
+	JsonVariantNode.Node = Node;
+	OutVariantNode = JsonVariantNode;
+	return true;
+}
+
+template<typename T>
+bool FGLTFLevelVariantSetsConverter::TryGetPropertyValue(UPropertyValue* Property, T& OutValue) const
+{
+	if (Property == nullptr || !Property->HasRecordedData())
+	{
+		return false;
+	}
+
+	FMemory::Memcpy(&OutValue, Property->GetRecordedData().GetData(), sizeof(T));
+	return true;
+}
diff --git a/Engine/Plugins/Enterprise/GLTFExporter/MigratedFromGit/Source/GLTFExporter/Private/Converters/GLTFVariantSetConverters.h b/Engine/Plugins/Enterprise/GLTFExporter/MigratedFromGit/Source/GLTFExporter/Private/Converters/GLTFVariantSetConverters.h
new file mode 100644
--- /dev/null
+++ b/Engine/Plugins/Enterprise/GLTFExporter/MigratedFromGit/Source/GLTFExporter/Private/Converters/GLTFVariantSetConverters.h
@@ -0,0 +1,22 @@
+// Copyright Epic Games, Inc. All Rights Reserved.
+
+#pragma once
+
+#include "Json/GLTFJsonIndex.h"
+#include "Json/GLTFJsonLevelVariantSets.h"
+#include "Converters/GLTFConverter.h"
+#include "Engine.h"
+#include "LevelVariantSetsActor.h"
+#include "Variant.h"
+#include "PropertyValue.h"
+
+class FGLTFLevelVariantSetsConverter final : public TGLTFConverter<FGLTFJsonLevelVariantSetsIndex, const ALevelVariantSetsActor*>
+{
+	FGLTFJsonLevelVariantSetsIndex Add(FGLTFConvertBuilder& Builder, const FString& Name, const ALevelVariantSetsActor* LevelVariantSetsActor) override;
+
+	bool TryParseJsonVariant(FGLTFConvertBuilder& Builder, FGLTFJsonVariant& OutVariant, const UVariant* Variant, const FString& ThumbnailPrefix = TEXT("")) const;
+	bool TryParseJsonVariantNode(FGLTFConvertBuilder& Builder, FGLTFJsonVariantNode& OutVariantNode, const UVariantObjectBinding* Binding) const;
+
+	template<typename T>
+	bool TryGetPropertyValue(UPropertyValue* Property, T& OutValue) const;
+};
diff --git a/Engine/Plugins/Enterprise/GLTFExporter/MigratedFromGit/Source/GLTFExporter/Private/Json/GLTFJsonIndex.h b/Engine/Plugins/Enterprise/GLTFExporter/MigratedFromGit/Source/GLTFExporter/Private/Json/GLTFJsonIndex.h
--- a/Engine/Plugins/Enterprise/GLTFExporter/MigratedFromGit/Source/GLTFExporter/Private/Json/GLTFJsonIndex.h
+++ b/Engine/Plugins/Enterprise/GLTFExporter/MigratedFromGit/Source/GLTFExporter/Private/Json/GLTFJsonIndex.h
@@ -65,3 +65,4 @@ struct FGLTFJsonSceneIndex : FGLTFJsonIndex<FGLTFJsonSceneIndex> { using FGLTFJs
 struct FGLTFJsonSkinIndex : FGLTFJsonIndex<FGLTFJsonSkinIndex> { using FGLTFJsonIndex::FGLTFJsonIndex; };
 struct FGLTFJsonTextureIndex : FGLTFJsonIndex<FGLTFJsonTextureIndex> { using FGLTFJsonIndex::FGLTFJsonIndex; };
 struct FGLTFJsonLightMapIndex : FGLTFJsonIndex<FGLTFJsonLightMapIndex> { using FGLTFJsonIndex::FGLTFJsonIndex; };
+struct FGLTFJsonLevelVariantSetsIndex : FGLTFJsonIndex<FGLTFJsonLevelVariantSetsIndex> { using FGLTFJsonIndex::FGLTFJsonIndex; };
diff --git a/Engine/Plugins/Enterprise/GLTFExporter/MigratedFromGit/Source/GLTFExporter/Private/Json/GLTFJsonLevelVariantSets.h b/Engine/Plugins/Enterprise/GLTFExporter/MigratedFromGit/Source/GLTFExporter/Private/Json/GLTFJsonLevelVariantSets.h
new file mode 100644
--- /dev/null
+++ b/Engine/Plugins/Enterprise/GLTFExporter/MigratedFromGit/Source/GLTFExporter/Private/Json/GLTFJsonLevelVariantSets.h
@@ -0,0 +1,107 @@
+// Copyright Epic Games, Inc. All Rights Reserved.
+
+#pragma once
+
+#include "GLTFJsonIndex.h"
+#include "Json/GLTFJsonExtensions.h"
+#include "Json/GLTFJsonUtility.h"
+#include "Serialization/JsonSerializer.h"
+#include "Misc/Optional.h"
+
+struct FGLTFJsonVariantNode
+{
+	FGLTFJsonNodeIndex Node;
+	TOptional<bool>    bIsVisible;
+
+	template <class CharType = TCHAR, class PrintPolicy = TPrettyJsonPrintPolicy<CharType>>
+	void WriteObject(TJsonWriter<CharType, PrintPolicy>& JsonWriter, FGLTFJsonExtensions& Extensions) const
+	{
+		JsonWriter.WriteObjectStart();
+
+		if (Node != INDEX_NONE)
+		{
+			JsonWriter.WriteValue(TEXT("node"), Node);
+		}
+
+		JsonWriter.WriteObjectStart(TEXT("properties"));
+
+		// TODO: should Visible always be written, with false as default?
+		if (bIsVisible.IsSet())
+		{
+			JsonWriter.WriteValue(TEXT("visible"), bIsVisible.GetValue());
+		}
+
+		JsonWriter.WriteObjectEnd();
+		JsonWriter.WriteObjectEnd();
+	}
+};
+
+struct FGLTFJsonVariant
+{
+	FString Name;
+	bool    bIsActive;
+
+	FGLTFJsonImageIndex          Thumbnail;
+	TArray<FGLTFJsonVariantNode> Nodes;
+
+	template <class CharType = TCHAR, class PrintPolicy = TPrettyJsonPrintPolicy<CharType>>
+	void WriteObject(TJsonWriter<CharType, PrintPolicy>& JsonWriter, FGLTFJsonExtensions& Extensions) const
+	{
+		JsonWriter.WriteObjectStart();
+		JsonWriter.WriteValue(TEXT("name"), Name);
+		JsonWriter.WriteValue(TEXT("active"), bIsActive);
+
+		if (Thumbnail != INDEX_NONE)
+		{
+			JsonWriter.WriteValue(TEXT("thumbnail"), Thumbnail);
+		}
+
+		FGLTFJsonUtility::WriteObjectArray(JsonWriter, TEXT("nodes"), Nodes, Extensions);
+
+		JsonWriter.WriteObjectEnd();
+	}
+};
+
+struct FGLTFJsonVariantSet
+{
+	FString Name;
+
+	TArray<FGLTFJsonVariant> Variants;
+
+	template <class CharType = TCHAR, class PrintPolicy = TPrettyJsonPrintPolicy<CharType>>
+	void WriteObject(TJsonWriter<CharType, PrintPolicy>& JsonWriter, FGLTFJsonExtensions& Extensions) const
+	{
+		JsonWriter.WriteObjectStart();
+
+		if (!Name.IsEmpty())
+		{
+			JsonWriter.WriteValue(TEXT("name"), Name);
+		}
+
+		FGLTFJsonUtility::WriteObjectArray(JsonWriter, TEXT("variants"), Variants, Extensions);
+
+		JsonWriter.WriteObjectEnd();
+	}
+};
+
+struct FGLTFJsonLevelVariantSets
+{
+	FString Name;
+
+	TArray<FGLTFJsonVariantSet> VariantSets;
+
+	template <class CharType = TCHAR, class PrintPolicy = TPrettyJsonPrintPolicy<CharType>>
+	void WriteObject(TJsonWriter<CharType, PrintPolicy>& JsonWriter, FGLTFJsonExtensions& Extensions) const
+	{
+		JsonWriter.WriteObjectStart();
+
+		if (!Name.IsEmpty())
+		{
+			JsonWriter.WriteValue(TEXT("name"), Name);
+		}
+
+		FGLTFJsonUtility::WriteObjectArray(JsonWriter, TEXT("variantSets"), VariantSets, Extensions);
+
+		JsonWriter.WriteObjectEnd();
+	}
+};
diff --git a/Engine/Plugins/Enterprise/GLTFExporter/MigratedFromGit/Source/GLTFExporter/Private/Json/GLTFJsonRoot.h b/Engine/Plugins/Enterprise/GLTFExporter/MigratedFromGit/Source/GLTFExporter/Private/Json/GLTFJsonRoot.h
--- a/Engine/Plugins/Enterprise/GLTFExporter/MigratedFromGit/Source/GLTFExporter/Private/Json/GLTFJsonRoot.h
+++ b/Engine/Plugins/Enterprise/GLTFExporter/MigratedFromGit/Source/GLTFExporter/Private/Json/GLTFJsonRoot.h
@@ -15,6 +15,7 @@
 #include "Json/GLTFJsonScene.h"
 #include "Json/GLTFJsonTexture.h"
 #include "Json/GLTFJsonLightMap.h"
+#include "Json/GLTFJsonLevelVariantSets.h"
 #include "Policies/CondensedJsonPrintPolicy.h"
 
 
@@ -67,6 +68,8 @@ struct FGLTFJsonRoot
 	TArray<FGLTFJsonTexture>    Textures;
 	TArray<FGLTFJsonLightMap>   LightMaps;
 
+	TArray<FGLTFJsonLevelVariantSets> LevelVariantSets;
+
 	template <class CharType = TCHAR, class PrintPolicy = TPrettyJsonPrintPolicy<CharType>>
 	void WriteObject(TJsonWriter<CharType, PrintPolicy>& JsonWriter) const
 	{
@@ -93,7 +96,7 @@ struct FGLTFJsonRoot
 		FGLTFJsonUtility::WriteObjectArray(JsonWriter, TEXT("scenes"), Scenes, AllExtensions);
 		FGLTFJsonUtility::WriteObjectArray(JsonWriter, TEXT("textures"), Textures, AllExtensions);
 
-		if (LightMaps.Num() > 0 /* TODO: add more extension support */)
+		if (LightMaps.Num() > 0 || LevelVariantSets.Num() > 0 /* TODO: add more extension support */)
 		{
 			JsonWriter.WriteObjectStart(TEXT("extensions"));
 
@@ -107,6 +110,19 @@ struct FGLTFJsonRoot
 				JsonWriter.WriteObjectEnd();
 			}
 
+			if (LevelVariantSets.Num() > 0)
+			{
+				const EGLTFJsonExtension Extension = EGLTFJsonExtension::EPIC_LevelVariantSets;
+
+				AllExtensions.Used.Add(Extension);
+
+				JsonWriter.WriteObjectStart(FGLTFJsonUtility::ToString(Extension));
+
+				FGLTFJsonUtility::WriteObjectArray(JsonWriter, TEXT("levelVariantSets"), LevelVariantSets, AllExtensions);
+
+				JsonWriter.WriteObjectEnd();
+			}
+
 			JsonWriter.WriteObjectEnd();
 		}
 
diff --git a/Engine/Plugins/Enterprise/GLTFExporter/MigratedFromGit/Source/GLTFExporter/Private/Json/GLTFJsonScene.h b/Engine/Plugins/Enterprise/GLTFExporter/MigratedFromGit/Source/GLTFExporter/Private/Json/GLTFJsonScene.h
--- a/Engine/Plugins/Enterprise/GLTFExporter/MigratedFromGit/Source/GLTFExporter/Private/Json/GLTFJsonScene.h
+++ b/Engine/Plugins/Enterprise/GLTFExporter/MigratedFromGit/Source/GLTFExporter/Private/Json/GLTFJsonScene.h
@@ -9,7 +9,8 @@ struct FGLTFJsonScene
 {
 	FString Name;
 
-	TArray<FGLTFJsonNodeIndex> Nodes;
+	TArray<FGLTFJsonNodeIndex>            Nodes;
+	TArray<FGLTFJsonLevelVariantSetsIndex> LevelVariantSets;
 
 	template <class CharType = TCHAR, class PrintPolicy = TPrettyJsonPrintPolicy<CharType>>
 	void WriteObject(TJsonWriter<CharType, PrintPolicy>& JsonWriter, FGLTFJsonExtensions& Extensions) const
@@ -26,6 +27,26 @@ struct FGLTFJsonScene
 			JsonWriter.WriteValue(TEXT("nodes"), Nodes);
 		}
 
+		const bool bWriteExtensions = LevelVariantSets.Num() > 0;
+
+		if (bWriteExtensions)
+		{
+			JsonWriter.WriteObjectStart(TEXT("extensions"));
+
+			if (LevelVariantSets.Num() > 0)
+			{
+				const EGLTFJsonExtension Extension = EGLTFJsonExtension::EPIC_LevelVariantSets;
+
+				Extensions.Used.Add(Extension);
+
+				JsonWriter.WriteObjectStart(FGLTFJsonUtility::ToString(Extension));
+				FGLTFJsonUtility::WritePrimitiveArray(JsonWriter, TEXT("levelVariantSets"), LevelVariantSets);
+				JsonWriter.WriteObjectEnd();
+			}
+
+			JsonWriter.WriteObjectEnd();
+		}
+
 		JsonWriter.WriteObjectEnd();
 	}
 };
diff --git a/Engine/Plugins/Enterprise/GLTFExporter/MigratedFromGit/Source/GLTFExporter/Private/Json/GLTFJsonUtility.h b/Engine/Plugins/Enterprise/GLTFExporter/MigratedFromGit/Source/GLTFExporter/Private/Json/GLTFJsonUtility.h
--- a/Engine/Plugins/Enterprise/GLTFExporter/MigratedFromGit/Source/GLTFExporter/Private/Json/GLTFJsonUtility.h
+++ b/Engine/Plugins/Enterprise/GLTFExporter/MigratedFromGit/Source/GLTFExporter/Private/Json/GLTFJsonUtility.h
@@ -109,4 +109,18 @@ struct FGLTFJsonUtility
 			JsonWriter.WriteArrayEnd();
 		}
 	}
+
+	template <class WriterType, class ContainerType>
+    static void WritePrimitiveArray(WriterType& JsonWriter, const FString& Identifier, const ContainerType& Container, bool bWriteIfEmpty = false)
+	{
+		if (Container.Num() > 0 || bWriteIfEmpty)
+		{
+			JsonWriter.WriteArrayStart(Identifier);
+			for (const auto& Element : Container)
+			{
+				JsonWriter.WriteValue(Element);
+			}
+			JsonWriter.WriteArrayEnd();
+		}
+	}
 };
