[GLTFExporter] feat(variants): add FGLTFJsonBuilder::GetLevelVariantSets

commit 355d9df71b01dca8d904781ac98f72a31c7fe900
Date: Mon, 6 Jul 2020 08:34:37 +0200

    feat(variants): add FGLTFJsonBuilder::GetLevelVariantSets
    (cherry picked from commit 08f7b507e401d716060b001b33db978ff3e85230​)
    # Conflicts from cherry-pick:
    #	Plugins/GLTFExporter/Source/GLTFExporter/Private/Builders/GLTFJsonBuilder.cpp
    #	Plugins/GLTFExporter/Source/GLTFExporter/Private/Builders/GLTFJsonBuilder.h

#skipci

diff --git a/Engine/Plugins/Enterprise/GLTFExporter/MigratedFromGit/Source/GLTFExporter/Private/Builders/GLTFJsonBuilder.cpp b/Engine/Plugins/Enterprise/GLTFExporter/MigratedFromGit/Source/GLTFExporter/Private/Builders/GLTFJsonBuilder.cpp
--- a/Engine/Plugins/Enterprise/GLTFExporter/MigratedFromGit/Source/GLTFExporter/Private/Builders/GLTFJsonBuilder.cpp
+++ b/Engine/Plugins/Enterprise/GLTFExporter/MigratedFromGit/Source/GLTFExporter/Private/Builders/GLTFJsonBuilder.cpp
@@ -135,6 +135,11 @@ FGLTFJsonLightMap& FGLTFJsonBuilder::GetLightMap(FGLTFJsonLightMapIndex LightMap
 	return JsonRoot.LightMaps[LightMapIndex];
 }
 
+FGLTFJsonLevelVariantSets& FGLTFJsonBuilder::GetLevelVariantSets(FGLTFJsonLevelVariantSetsIndex LevelVariantSetsIndex)
+{
+	return JsonRoot.LevelVariantSets[LevelVariantSetsIndex];
+}
+
 bool FGLTFJsonBuilder::Serialize(FArchive& Archive, const FString& FilePath)
 {
 	JsonRoot.Serialize(&Archive, true);
diff --git a/Engine/Plugins/Enterprise/GLTFExporter/MigratedFromGit/Source/GLTFExporter/Private/Builders/GLTFJsonBuilder.h b/Engine/Plugins/Enterprise/GLTFExporter/MigratedFromGit/Source/GLTFExporter/Private/Builders/GLTFJsonBuilder.h
--- a/Engine/Plugins/Enterprise/GLTFExporter/MigratedFromGit/Source/GLTFExporter/Private/Builders/GLTFJsonBuilder.h
+++ b/Engine/Plugins/Enterprise/GLTFExporter/MigratedFromGit/Source/GLTFExporter/Private/Builders/GLTFJsonBuilder.h
@@ -39,6 +39,7 @@ public:
 	FGLTFJsonScene& GetScene(FGLTFJsonSceneIndex SceneIndex);
 	FGLTFJsonTexture& GetTexture(FGLTFJsonTextureIndex TextureIndex);
 	FGLTFJsonLightMap& GetLightMap(FGLTFJsonLightMapIndex LightMapIndex);
+	FGLTFJsonLevelVariantSets& GetLevelVariantSets(FGLTFJsonLevelVariantSetsIndex LevelVariantSetsIndex);
 
 	virtual bool Serialize(FArchive& Archive, const FString& FilePath);
 
