[GLTFExporter] chore(variants): remove TODO for questionable behaviour (current behaviour is correct)

commit 641bc330383f2f1e0ac6f8c26253ec65ed57e20b
Date: Mon, 6 Jul 2020 08:43:50 +0200

    chore(variants): remove TODO for questionable behaviour (current behaviour is correct)
    (cherry picked from commit 55b2e303b26b1ead106c46d569b843fea109edeb​)

#skipci

diff --git a/Engine/Plugins/Enterprise/GLTFExporter/MigratedFromGit/Source/GLTFExporter/Private/Json/GLTFJsonLevelVariantSets.h b/Engine/Plugins/Enterprise/GLTFExporter/MigratedFromGit/Source/GLTFExporter/Private/Json/GLTFJsonLevelVariantSets.h
--- a/Engine/Plugins/Enterprise/GLTFExporter/MigratedFromGit/Source/GLTFExporter/Private/Json/GLTFJsonLevelVariantSets.h
+++ b/Engine/Plugins/Enterprise/GLTFExporter/MigratedFromGit/Source/GLTFExporter/Private/Json/GLTFJsonLevelVariantSets.h
@@ -25,7 +25,6 @@ struct FGLTFJsonVariantNode
 
 		JsonWriter.WriteObjectStart(TEXT("properties"));
 
-		// TODO: should Visible always be written, with false as default?
 		if (bIsVisible.IsSet())
 		{
 			JsonWriter.WriteValue(TEXT("visible"), bIsVisible.GetValue());
