[GLTFExporter] refactor(variants): remove FGLTFJsonUtility::WritePrimitiveArray since it has no actual use-case

commit 235ea836bf044faa85913b1dd6ba4bc009265a01
Date: Mon, 6 Jul 2020 08:51:33 +0200

    refactor(variants): remove FGLTFJsonUtility::WritePrimitiveArray since it has no actual use-case
    (cherry picked from commit 01630ae25053a58e912f69db3c025546b0532bf3​)

#skipci

diff --git a/Engine/Plugins/Enterprise/GLTFExporter/MigratedFromGit/Source/GLTFExporter/Private/Json/GLTFJsonScene.h b/Engine/Plugins/Enterprise/GLTFExporter/MigratedFromGit/Source/GLTFExporter/Private/Json/GLTFJsonScene.h
--- a/Engine/Plugins/Enterprise/GLTFExporter/MigratedFromGit/Source/GLTFExporter/Private/Json/GLTFJsonScene.h
+++ b/Engine/Plugins/Enterprise/GLTFExporter/MigratedFromGit/Source/GLTFExporter/Private/Json/GLTFJsonScene.h
@@ -40,7 +40,7 @@ struct FGLTFJsonScene
 				Extensions.Used.Add(Extension);
 
 				JsonWriter.WriteObjectStart(FGLTFJsonUtility::ToString(Extension));
-				FGLTFJsonUtility::WritePrimitiveArray(JsonWriter, TEXT("levelVariantSets"), LevelVariantSets);
+				JsonWriter.WriteValue(TEXT("levelVariantSets"), LevelVariantSets);
 				JsonWriter.WriteObjectEnd();
 			}
 
diff --git a/Engine/Plugins/Enterprise/GLTFExporter/MigratedFromGit/Source/GLTFExporter/Private/Json/GLTFJsonUtility.h b/Engine/Plugins/Enterprise/GLTFExporter/MigratedFromGit/Source/GLTFExporter/Private/Json/GLTFJsonUtility.h
--- a/Engine/Plugins/Enterprise/GLTFExporter/MigratedFromGit/Source/GLTFExporter/Private/Json/GLTFJsonUtility.h
+++ b/Engine/Plugins/Enterprise/GLTFExporter/MigratedFromGit/Source/GLTFExporter/Private/Json/GLTFJsonUtility.h
@@ -109,18 +109,4 @@ struct FGLTFJsonUtility
 			JsonWriter.WriteArrayEnd();
 		}
 	}
-
-	template <class WriterType, class ContainerType>
-    static void WritePrimitiveArray(WriterType& JsonWriter, const FString& Identifier, const ContainerType& Container, bool bWriteIfEmpty = false)
-	{
-		if (Container.Num() > 0 || bWriteIfEmpty)
-		{
-			JsonWriter.WriteArrayStart(Identifier);
-			for (const auto& Element : Container)
-			{
-				JsonWriter.WriteValue(Element);
-			}
-			JsonWriter.WriteArrayEnd();
-		}
-	}
 };
