[GLTFExporter] feat(variants): add warning messages for when variant bindings fail to export their target object / node

commit 570001a3b7570c92f1d0b6c5584671c3c107ee56
Date: Mon, 6 Jul 2020 09:16:50 +0200

    feat(variants): add warning messages for when variant bindings fail to export their target object / node
    (cherry picked from commit 3fb883bff03d61036d12d430b4a8c7c5b2421dc9​)

#skipci

diff --git a/Engine/Plugins/Enterprise/GLTFExporter/MigratedFromGit/Source/GLTFExporter/Private/Converters/GLTFVariantSetConverters.cpp b/Engine/Plugins/Enterprise/GLTFExporter/MigratedFromGit/Source/GLTFExporter/Private/Converters/GLTFVariantSetConverters.cpp
--- a/Engine/Plugins/Enterprise/GLTFExporter/MigratedFromGit/Source/GLTFExporter/Private/Converters/GLTFVariantSetConverters.cpp
+++ b/Engine/Plugins/Enterprise/GLTFExporter/MigratedFromGit/Source/GLTFExporter/Private/Converters/GLTFVariantSetConverters.cpp
@@ -86,6 +86,13 @@ bool FGLTFLevelVariantSetsConverter::TryParseJsonVariant(FGLTFConvertBuilder& Bu
 
 bool FGLTFLevelVariantSetsConverter::TryParseJsonVariantNode(FGLTFConvertBuilder& Builder, FGLTFJsonVariantNode& OutVariantNode, const UVariantObjectBinding* Binding) const
 {
+	const AActor* Actor = Cast<AActor>(Binding->GetObject());
+	if (Actor == nullptr)
+	{
+		Builder.AddWarningMessage(FString::Printf(TEXT("Failed to find actor for variant binding %s"), *Binding->GetObjectPath()));
+		return false;
+	}
+
 	FGLTFJsonVariantNode JsonVariantNode;
 	bool bHasAnyProperty = false;
 
@@ -127,12 +134,11 @@ bool FGLTFLevelVariantSetsConverter::TryParseJsonVariantNode(FGLTFConvertBuilder
 		return false;
 	}
 
-	const AActor* Actor = Cast<AActor>(Binding->GetObject());
 	const FGLTFJsonNodeIndex Node = Builder.GetOrAddNode(Actor);
-
 	if (Node == INDEX_NONE)
 	{
-		// TODO: report warning about invalid actor / node?
+		// NOTE: this could occur if the user is exporting selected actors only, if the actor isn't part of the selection
+		Builder.AddWarningMessage(FString::Printf(TEXT("Failed to export node for actor used by variant binding %s"), *Binding->GetObjectPath()));
 		return false;
 	}
 
