[GLTFExporter] refactor(variants): identify supported properties by using name and class of the property

commit 7f3bf849aad31a4b743409acd6c8d9ab95cb2445
Date: Mon, 6 Jul 2020 15:21:51 +0200

    refactor(variants): identify supported properties by using name and class of the property
    (cherry picked from commit 5399e883facab377ff72a666b47c7272f1e261a7​)

#skipci

diff --git a/Engine/Plugins/Enterprise/GLTFExporter/MigratedFromGit/Source/GLTFExporter/Private/Converters/GLTFVariantSetConverters.cpp b/Engine/Plugins/Enterprise/GLTFExporter/MigratedFromGit/Source/GLTFExporter/Private/Converters/GLTFVariantSetConverters.cpp
--- a/Engine/Plugins/Enterprise/GLTFExporter/MigratedFromGit/Source/GLTFExporter/Private/Converters/GLTFVariantSetConverters.cpp
+++ b/Engine/Plugins/Enterprise/GLTFExporter/MigratedFromGit/Source/GLTFExporter/Private/Converters/GLTFVariantSetConverters.cpp
@@ -100,31 +100,25 @@ bool FGLTFLevelVariantSetsConverter::TryParseJsonVariantNode(FGLTFConvertBuilder
 	{
 		if (Property->Resolve())
 		{
-			// TODO: should we use another way to identify the properties we are interested in?
-			// Right now we use the property-catagory and property-class, but we might use the class
-			// of the parent container and the display-name of the property instead.
-			// Or something completely different...
+			// TODO: handle failure to resolve property (or force resolve it somehow)
+			continue;
+		}
 
-			const EPropertyValueCategory PropCategory = Property->GetPropCategory();
-			const FFieldClass* PropClass = Property->GetPropertyClass();
+		const FFieldClass* PropertyClass = Property->GetPropertyClass();
+		const FName& PropertyName = Property->GetPropertyName();
 
-			if (PropCategory == EPropertyValueCategory::Visibility && PropClass == FBoolProperty::StaticClass())
-			{
-				bool bIsVisible;
-				if (TryGetPropertyValue(Property, bIsVisible))
-				{
-					JsonVariantNode.bIsVisible = bIsVisible;
-					bHasAnyProperty = true;
-				}
-			}
-			else
+		if (PropertyName == TEXT("bVisible") && PropertyClass->IsChildOf(FBoolProperty::StaticClass()))
+		{
+			bool bIsVisible;
+			if (TryGetPropertyValue(Property, bIsVisible))
 			{
-				// TODO: handle more properties
+				JsonVariantNode.bIsVisible = bIsVisible;
+				bHasAnyProperty = true;
 			}
 		}
 		else
 		{
-			// TODO: handle failure to resolve property (or force resolve it somehow)
+			// TODO: handle more properties
 		}
 	}
 
