[GLTFExporter] feat(variants): add / clarify warning messages for failed variant bindings

commit 19a46c9c549529cccba30854f65ca763438b1d7a
Date: Mon, 6 Jul 2020 15:24:26 +0200

    feat(variants): add / clarify warning messages for failed variant bindings
    (cherry picked from commit 6414b0f49ce2e6798a709dbde266203e75a52db8​)

#skipci

diff --git a/Engine/Plugins/Enterprise/GLTFExporter/MigratedFromGit/Source/GLTFExporter/Private/Converters/GLTFVariantSetConverters.cpp b/Engine/Plugins/Enterprise/GLTFExporter/MigratedFromGit/Source/GLTFExporter/Private/Converters/GLTFVariantSetConverters.cpp
--- a/Engine/Plugins/Enterprise/GLTFExporter/MigratedFromGit/Source/GLTFExporter/Private/Converters/GLTFVariantSetConverters.cpp
+++ b/Engine/Plugins/Enterprise/GLTFExporter/MigratedFromGit/Source/GLTFExporter/Private/Converters/GLTFVariantSetConverters.cpp
@@ -86,10 +86,16 @@ bool FGLTFLevelVariantSetsConverter::TryParseJsonVariant(FGLTFConvertBuilder& Bu
 
 bool FGLTFLevelVariantSetsConverter::TryParseJsonVariantNode(FGLTFConvertBuilder& Builder, FGLTFJsonVariantNode& OutVariantNode, const UVariantObjectBinding* Binding) const
 {
+	const UVariant* Variant = const_cast<UVariantObjectBinding*>(Binding)->GetParent();
+
 	const AActor* Actor = Cast<AActor>(Binding->GetObject());
 	if (Actor == nullptr)
 	{
-		Builder.AddWarningMessage(FString::Printf(TEXT("Failed to find actor for variant binding %s"), *Binding->GetObjectPath()));
+		Builder.AddWarningMessage(FString::Printf(
+			TEXT("Actor '%s' for variant '%s' was not found and will be skipped. Context: %s"),
+			*Binding->GetDisplayText().ToString(),
+			*Variant->GetDisplayText().ToString(),
+			*GetLogContext(Variant)));
 		return false;
 	}
 
@@ -98,9 +104,14 @@ bool FGLTFLevelVariantSetsConverter::TryParseJsonVariantNode(FGLTFConvertBuilder
 
 	for (UPropertyValue* Property: Binding->GetCapturedProperties())
 	{
-		if (Property->Resolve())
+		if (!Property->Resolve())
 		{
-			// TODO: handle failure to resolve property (or force resolve it somehow)
+			Builder.AddWarningMessage(FString::Printf(
+				TEXT("Property '%s' for actor '%s' in variant '%s' failed to resolve and will be skipped. Context: %s"),
+				*Property->GetFullDisplayString(),
+				*Binding->GetDisplayText().ToString(),
+				*Variant->GetDisplayText().ToString(),
+				*GetLogContext(Variant)));
 			continue;
 		}
 
@@ -132,7 +143,11 @@ bool FGLTFLevelVariantSetsConverter::TryParseJsonVariantNode(FGLTFConvertBuilder
 	if (Node == INDEX_NONE)
 	{
 		// NOTE: this could occur if the user is exporting selected actors only, if the actor isn't part of the selection
-		Builder.AddWarningMessage(FString::Printf(TEXT("Failed to export node for actor used by variant binding %s"), *Binding->GetObjectPath()));
+		Builder.AddWarningMessage(FString::Printf(
+			TEXT("Actor '%s' for variant '%s' could not be exported and will be skipped. Context: %s"),
+			*Binding->GetDisplayText().ToString(),
+			*Variant->GetDisplayText().ToString(),
+			*GetLogContext(Variant)));
 		return false;
 	}
 
@@ -152,3 +167,11 @@ bool FGLTFLevelVariantSetsConverter::TryGetPropertyValue(UPropertyValue* Propert
 	FMemory::Memcpy(&OutValue, Property->GetRecordedData().GetData(), sizeof(T));
 	return true;
 }
+
+FString FGLTFLevelVariantSetsConverter::GetLogContext(const UVariant* Variant) const
+{
+	UVariantSet* VariantSet = const_cast<UVariant*>(Variant)->GetParent();
+	ULevelVariantSets* LevelVariantSets = VariantSet->GetParent();
+
+	return LevelVariantSets->GetName() + TEXT("/") + VariantSet->GetDisplayText().ToString() + TEXT("/") + Variant->GetDisplayText().ToString();
+}
diff --git a/Engine/Plugins/Enterprise/GLTFExporter/MigratedFromGit/Source/GLTFExporter/Private/Converters/GLTFVariantSetConverters.h b/Engine/Plugins/Enterprise/GLTFExporter/MigratedFromGit/Source/GLTFExporter/Private/Converters/GLTFVariantSetConverters.h
--- a/Engine/Plugins/Enterprise/GLTFExporter/MigratedFromGit/Source/GLTFExporter/Private/Converters/GLTFVariantSetConverters.h
+++ b/Engine/Plugins/Enterprise/GLTFExporter/MigratedFromGit/Source/GLTFExporter/Private/Converters/GLTFVariantSetConverters.h
@@ -19,4 +19,6 @@ class FGLTFLevelVariantSetsConverter final : public TGLTFConverter<FGLTFJsonLeve
 
 	template<typename T>
 	bool TryGetPropertyValue(UPropertyValue* Property, T& OutValue) const;
+
+	FString GetLogContext(const UVariant* Variant) const;
 };
