[GLTFExporter] feat(variants): add a specific warning message when an actor is not selected to be part of the export

commit 7cb75f5e3d4ee1d0bd777c8fe173a9f5b7feb815
Date: Mon, 6 Jul 2020 15:31:17 +0200

    feat(variants): add a specific warning message when an actor is not selected to be part of the export
    (cherry picked from commit bbf003c15df94d5de1e97c69c60ea0d22c617e5c​)

#skipci

diff --git a/Engine/Plugins/Enterprise/GLTFExporter/MigratedFromGit/Source/GLTFExporter/Private/Converters/GLTFVariantSetConverters.cpp b/Engine/Plugins/Enterprise/GLTFExporter/MigratedFromGit/Source/GLTFExporter/Private/Converters/GLTFVariantSetConverters.cpp
--- a/Engine/Plugins/Enterprise/GLTFExporter/MigratedFromGit/Source/GLTFExporter/Private/Converters/GLTFVariantSetConverters.cpp
+++ b/Engine/Plugins/Enterprise/GLTFExporter/MigratedFromGit/Source/GLTFExporter/Private/Converters/GLTFVariantSetConverters.cpp
@@ -99,6 +99,16 @@ bool FGLTFLevelVariantSetsConverter::TryParseJsonVariantNode(FGLTFConvertBuilder
 		return false;
 	}
 
+	if (Builder.bSelectedActorsOnly && !Actor->IsSelected())
+	{
+		Builder.AddWarningMessage(FString::Printf(
+			TEXT("Actor '%s' for variant '%s' is not selected and will be skipped. Context: %s"),
+			*Binding->GetDisplayText().ToString(),
+			*Variant->GetDisplayText().ToString(),
+			*GetLogContext(Variant)));
+		return false;
+	}
+
 	FGLTFJsonVariantNode JsonVariantNode;
 	bool bHasAnyProperty = false;
 
@@ -142,7 +152,6 @@ bool FGLTFLevelVariantSetsConverter::TryParseJsonVariantNode(FGLTFConvertBuilder
 	const FGLTFJsonNodeIndex Node = Builder.GetOrAddNode(Actor);
 	if (Node == INDEX_NONE)
 	{
-		// NOTE: this could occur if the user is exporting selected actors only, if the actor isn't part of the selection
 		Builder.AddWarningMessage(FString::Printf(
 			TEXT("Actor '%s' for variant '%s' could not be exported and will be skipped. Context: %s"),
 			*Binding->GetDisplayText().ToString(),
