[GLTFExporter] feat(variants): add warning when skipping actors in variants due to lack of supported properties

commit 0a3eadb28f3badbe035f872f4dc0c39a2f0b7b79
Date: Mon, 6 Jul 2020 15:39:45 +0200

    feat(variants): add warning when skipping actors in variants due to lack of supported properties
    (cherry picked from commit ee264768930a9358e345987778b4d5bf01803fc5​)

#skipci

diff --git a/Engine/Plugins/Enterprise/GLTFExporter/MigratedFromGit/Source/GLTFExporter/Private/Converters/GLTFVariantSetConverters.cpp b/Engine/Plugins/Enterprise/GLTFExporter/MigratedFromGit/Source/GLTFExporter/Private/Converters/GLTFVariantSetConverters.cpp
--- a/Engine/Plugins/Enterprise/GLTFExporter/MigratedFromGit/Source/GLTFExporter/Private/Converters/GLTFVariantSetConverters.cpp
+++ b/Engine/Plugins/Enterprise/GLTFExporter/MigratedFromGit/Source/GLTFExporter/Private/Converters/GLTFVariantSetConverters.cpp
@@ -145,7 +145,11 @@ bool FGLTFLevelVariantSetsConverter::TryParseJsonVariantNode(FGLTFConvertBuilder
 
 	if (!bHasAnyProperty)
 	{
-		// TODO: allow variant nodes without any properties?
+		Builder.AddWarningMessage(FString::Printf(
+			TEXT("Actor '%s' in variant '%s' has no properties supported by the exporter and will be skipped. Context: %s"),
+			*Binding->GetDisplayText().ToString(),
+			*Variant->GetDisplayText().ToString(),
+			*GetLogContext(Variant)));
 		return false;
 	}
 
