[GLTFExporter] refactor(variants): change wording of some warning messages

commit a731afda909bb3c60adcea51d1799d84f7e294be
Date: Mon, 6 Jul 2020 15:40:29 +0200

    refactor(variants): change wording of some warning messages
    (cherry picked from commit a80f428d01691053f251250cb0785b8eb8fbcc9c​)

#skipci

diff --git a/Engine/Plugins/Enterprise/GLTFExporter/MigratedFromGit/Source/GLTFExporter/Private/Converters/GLTFVariantSetConverters.cpp b/Engine/Plugins/Enterprise/GLTFExporter/MigratedFromGit/Source/GLTFExporter/Private/Converters/GLTFVariantSetConverters.cpp
--- a/Engine/Plugins/Enterprise/GLTFExporter/MigratedFromGit/Source/GLTFExporter/Private/Converters/GLTFVariantSetConverters.cpp
+++ b/Engine/Plugins/Enterprise/GLTFExporter/MigratedFromGit/Source/GLTFExporter/Private/Converters/GLTFVariantSetConverters.cpp
@@ -92,7 +92,7 @@ bool FGLTFLevelVariantSetsConverter::TryParseJsonVariantNode(FGLTFConvertBuilder
 	if (Actor == nullptr)
 	{
 		Builder.AddWarningMessage(FString::Printf(
-			TEXT("Actor '%s' for variant '%s' was not found and will be skipped. Context: %s"),
+			TEXT("Actor '%s' in variant '%s' was not found and will be skipped. Context: %s"),
 			*Binding->GetDisplayText().ToString(),
 			*Variant->GetDisplayText().ToString(),
 			*GetLogContext(Variant)));
@@ -102,7 +102,7 @@ bool FGLTFLevelVariantSetsConverter::TryParseJsonVariantNode(FGLTFConvertBuilder
 	if (Builder.bSelectedActorsOnly && !Actor->IsSelected())
 	{
 		Builder.AddWarningMessage(FString::Printf(
-			TEXT("Actor '%s' for variant '%s' is not selected and will be skipped. Context: %s"),
+			TEXT("Actor '%s' in variant '%s' is not selected and will be skipped. Context: %s"),
 			*Binding->GetDisplayText().ToString(),
 			*Variant->GetDisplayText().ToString(),
 			*GetLogContext(Variant)));
@@ -157,7 +157,7 @@ bool FGLTFLevelVariantSetsConverter::TryParseJsonVariantNode(FGLTFConvertBuilder
 	if (Node == INDEX_NONE)
 	{
 		Builder.AddWarningMessage(FString::Printf(
-			TEXT("Actor '%s' for variant '%s' could not be exported and will be skipped. Context: %s"),
+			TEXT("Actor '%s' in variant '%s' could not be exported and will be skipped. Context: %s"),
 			*Binding->GetDisplayText().ToString(),
 			*Variant->GetDisplayText().ToString(),
 			*GetLogContext(Variant)));
