[GLTFExporter] chore(variants): remove TODOs that are no longer relevant

commit d4bb70dc077f1072d071d6e30a4ba0d6ac393b5f
Date: Mon, 6 Jul 2020 15:41:41 +0200

    chore(variants): remove TODOs that are no longer relevant
    (cherry picked from commit 79648c9b0470d2995d09aadf4f1cf82e20169193​)

#skipci

diff --git a/Engine/Plugins/Enterprise/GLTFExporter/MigratedFromGit/Source/GLTFExporter/Private/Converters/GLTFVariantSetConverters.cpp b/Engine/Plugins/Enterprise/GLTFExporter/MigratedFromGit/Source/GLTFExporter/Private/Converters/GLTFVariantSetConverters.cpp
--- a/Engine/Plugins/Enterprise/GLTFExporter/MigratedFromGit/Source/GLTFExporter/Private/Converters/GLTFVariantSetConverters.cpp
+++ b/Engine/Plugins/Enterprise/GLTFExporter/MigratedFromGit/Source/GLTFExporter/Private/Converters/GLTFVariantSetConverters.cpp
@@ -36,7 +36,6 @@ FGLTFJsonLevelVariantSetsIndex FGLTFLevelVariantSetsConverter::Add(FGLTFConvertB
 			}
 		}
 
-		// TODO: allow sets without any variants?
 		if (JsonVariantSet.Variants.Num() > 0)
 		{
 			JsonLevelVariantSets.VariantSets.Add(JsonVariantSet);
@@ -45,7 +44,6 @@ FGLTFJsonLevelVariantSetsIndex FGLTFLevelVariantSetsConverter::Add(FGLTFConvertB
 
 	if (JsonLevelVariantSets.VariantSets.Num() == 0)
 	{
-		// TODO: allow level variant sets without any variant sets?
 		return FGLTFJsonLevelVariantSetsIndex(INDEX_NONE);
 	}
 
@@ -67,7 +65,6 @@ bool FGLTFLevelVariantSetsConverter::TryParseJsonVariant(FGLTFConvertBuilder& Bu
 
 	if (JsonVariant.Nodes.Num() == 0)
 	{
-		// TODO: allow variants without nodes?
 		return false;
 	}
 
