[GLTFExporter] feat(variants): add warning message when a property lacks recorded data

commit 24ba59687965d169a8f8b2a6492b0a1dab5af794
Date: Mon, 6 Jul 2020 15:52:28 +0200

    feat(variants): add warning message when a property lacks recorded data
    (cherry picked from commit 11a1a5dbc94ce56129fbad17a59f75f04293345f​)

#skipci

diff --git a/Engine/Plugins/Enterprise/GLTFExporter/MigratedFromGit/Source/GLTFExporter/Private/Converters/GLTFVariantSetConverters.cpp b/Engine/Plugins/Enterprise/GLTFExporter/MigratedFromGit/Source/GLTFExporter/Private/Converters/GLTFVariantSetConverters.cpp
--- a/Engine/Plugins/Enterprise/GLTFExporter/MigratedFromGit/Source/GLTFExporter/Private/Converters/GLTFVariantSetConverters.cpp
+++ b/Engine/Plugins/Enterprise/GLTFExporter/MigratedFromGit/Source/GLTFExporter/Private/Converters/GLTFVariantSetConverters.cpp
@@ -133,6 +133,15 @@ bool FGLTFLevelVariantSetsConverter::TryParseJsonVariantNode(FGLTFConvertBuilder
 				JsonVariantNode.bIsVisible = bIsVisible;
 				bHasAnyProperty = true;
 			}
+			else
+			{
+				Builder.AddWarningMessage(FString::Printf(
+					TEXT("Property '%s' for actor '%s' in variant '%s' has no recorded data and will be skipped. Context: %s"),
+					*Property->GetFullDisplayString(),
+					*Binding->GetDisplayText().ToString(),
+					*Variant->GetDisplayText().ToString(),
+					*GetLogContext(Variant)));
+			}
 		}
 		else
 		{
