[GLTFExporter] refactor(variants): store commonly used display-texts in variables instead of retrieving them when used

commit 5dc08748d6bc91aae6b61775be4196d9465049a3
Date: Mon, 6 Jul 2020 15:57:09 +0200

    refactor(variants): store commonly used display-texts in variables instead of retrieving them when used
    (cherry picked from commit 8b6afa24f0ab3488157269a9bf391f31f4d830da​)

#skipci

diff --git a/Engine/Plugins/Enterprise/GLTFExporter/MigratedFromGit/Source/GLTFExporter/Private/Converters/GLTFVariantSetConverters.cpp b/Engine/Plugins/Enterprise/GLTFExporter/MigratedFromGit/Source/GLTFExporter/Private/Converters/GLTFVariantSetConverters.cpp
--- a/Engine/Plugins/Enterprise/GLTFExporter/MigratedFromGit/Source/GLTFExporter/Private/Converters/GLTFVariantSetConverters.cpp
+++ b/Engine/Plugins/Enterprise/GLTFExporter/MigratedFromGit/Source/GLTFExporter/Private/Converters/GLTFVariantSetConverters.cpp
@@ -84,15 +84,15 @@ bool FGLTFLevelVariantSetsConverter::TryParseJsonVariant(FGLTFConvertBuilder& Bu
 bool FGLTFLevelVariantSetsConverter::TryParseJsonVariantNode(FGLTFConvertBuilder& Builder, FGLTFJsonVariantNode& OutVariantNode, const UVariantObjectBinding* Binding) const
 {
 	const UVariant* Variant = const_cast<UVariantObjectBinding*>(Binding)->GetParent();
+	const FString BindingDisplayText = Binding->GetDisplayText().ToString();
+	const FString VariantDisplayText = Variant->GetDisplayText().ToString();
 
 	const AActor* Actor = Cast<AActor>(Binding->GetObject());
 	if (Actor == nullptr)
 	{
 		Builder.AddWarningMessage(FString::Printf(
 			TEXT("Actor '%s' in variant '%s' was not found and will be skipped. Context: %s"),
-			*Binding->GetDisplayText().ToString(),
-			*Variant->GetDisplayText().ToString(),
-			*GetLogContext(Variant)));
+			*BindingDisplayText, *VariantDisplayText, *GetLogContext(Variant)));
 		return false;
 	}
 
@@ -100,9 +100,7 @@ bool FGLTFLevelVariantSetsConverter::TryParseJsonVariantNode(FGLTFConvertBuilder
 	{
 		Builder.AddWarningMessage(FString::Printf(
 			TEXT("Actor '%s' in variant '%s' is not selected and will be skipped. Context: %s"),
-			*Binding->GetDisplayText().ToString(),
-			*Variant->GetDisplayText().ToString(),
-			*GetLogContext(Variant)));
+			*BindingDisplayText, *VariantDisplayText, *GetLogContext(Variant)));
 		return false;
 	}
 
@@ -115,10 +113,7 @@ bool FGLTFLevelVariantSetsConverter::TryParseJsonVariantNode(FGLTFConvertBuilder
 		{
 			Builder.AddWarningMessage(FString::Printf(
 				TEXT("Property '%s' for actor '%s' in variant '%s' failed to resolve and will be skipped. Context: %s"),
-				*Property->GetFullDisplayString(),
-				*Binding->GetDisplayText().ToString(),
-				*Variant->GetDisplayText().ToString(),
-				*GetLogContext(Variant)));
+				*Property->GetFullDisplayString(), *BindingDisplayText, *VariantDisplayText, *GetLogContext(Variant)));
 			continue;
 		}
 
@@ -137,10 +132,7 @@ bool FGLTFLevelVariantSetsConverter::TryParseJsonVariantNode(FGLTFConvertBuilder
 			{
 				Builder.AddWarningMessage(FString::Printf(
 					TEXT("Property '%s' for actor '%s' in variant '%s' has no recorded data and will be skipped. Context: %s"),
-					*Property->GetFullDisplayString(),
-					*Binding->GetDisplayText().ToString(),
-					*Variant->GetDisplayText().ToString(),
-					*GetLogContext(Variant)));
+					*Property->GetFullDisplayString(), *BindingDisplayText, *VariantDisplayText, *GetLogContext(Variant)));
 			}
 		}
 		else
@@ -153,9 +145,7 @@ bool FGLTFLevelVariantSetsConverter::TryParseJsonVariantNode(FGLTFConvertBuilder
 	{
 		Builder.AddWarningMessage(FString::Printf(
 			TEXT("Actor '%s' in variant '%s' has no properties supported by the exporter and will be skipped. Context: %s"),
-			*Binding->GetDisplayText().ToString(),
-			*Variant->GetDisplayText().ToString(),
-			*GetLogContext(Variant)));
+			*BindingDisplayText, *VariantDisplayText, *GetLogContext(Variant)));
 		return false;
 	}
 
@@ -164,9 +154,7 @@ bool FGLTFLevelVariantSetsConverter::TryParseJsonVariantNode(FGLTFConvertBuilder
 	{
 		Builder.AddWarningMessage(FString::Printf(
 			TEXT("Actor '%s' in variant '%s' could not be exported and will be skipped. Context: %s"),
-			*Binding->GetDisplayText().ToString(),
-			*Variant->GetDisplayText().ToString(),
-			*GetLogContext(Variant)));
+			*BindingDisplayText, *VariantDisplayText, *GetLogContext(Variant)));
 		return false;
 	}
 
