[GLTFExporter] refactor(texture): move and generalize texture sampler conversion to separate converter

commit d42c5bb64009ff62907825f38804b171f99875bf
Date: Fri, 26 Jun 2020 01:37:06 +0200

    refactor(texture): move and generalize texture sampler conversion to separate converter
    (cherry picked from commit 8a415d385bc73810e9f41e9cdaccd8363697e818​)

#skipci

diff --git a/Engine/Plugins/Enterprise/GLTFExporter/MigratedFromGit/Source/GLTFExporter/Private/Builders/GLTFConvertBuilder.cpp b/Engine/Plugins/Enterprise/GLTFExporter/MigratedFromGit/Source/GLTFExporter/Private/Builders/GLTFConvertBuilder.cpp
--- a/Engine/Plugins/Enterprise/GLTFExporter/MigratedFromGit/Source/GLTFExporter/Private/Builders/GLTFConvertBuilder.cpp
+++ b/Engine/Plugins/Enterprise/GLTFExporter/MigratedFromGit/Source/GLTFExporter/Private/Builders/GLTFConvertBuilder.cpp
@@ -114,6 +114,16 @@ FGLTFJsonMaterialIndex FGLTFConvertBuilder::GetOrAddMaterial(const UMaterialInte
 	return MaterialConverter.GetOrAdd(*this, DesiredName.IsEmpty() ? Material->GetName() : DesiredName, Material);
 }
 
+FGLTFJsonSamplerIndex FGLTFConvertBuilder::GetOrAddSampler(const UTexture* Texture, const FString& DesiredName)
+{
+	if (Texture == nullptr)
+	{
+		return FGLTFJsonSamplerIndex(INDEX_NONE);
+	}
+
+	return TextureSamplerConverter.GetOrAdd(*this, DesiredName.IsEmpty() ? Texture->GetName() : DesiredName, Texture);
+}
+
 FGLTFJsonTextureIndex FGLTFConvertBuilder::GetOrAddTexture(const UTexture2D* Texture, const FString& DesiredName)
 {
 	if (Texture == nullptr)
diff --git a/Engine/Plugins/Enterprise/GLTFExporter/MigratedFromGit/Source/GLTFExporter/Private/Builders/GLTFConvertBuilder.h b/Engine/Plugins/Enterprise/GLTFExporter/MigratedFromGit/Source/GLTFExporter/Private/Builders/GLTFConvertBuilder.h
--- a/Engine/Plugins/Enterprise/GLTFExporter/MigratedFromGit/Source/GLTFExporter/Private/Builders/GLTFConvertBuilder.h
+++ b/Engine/Plugins/Enterprise/GLTFExporter/MigratedFromGit/Source/GLTFExporter/Private/Builders/GLTFConvertBuilder.h
@@ -31,6 +31,7 @@ public:
 	FGLTFJsonMeshIndex GetOrAddMesh(const UStaticMeshComponent* StaticMeshComponent, const FString& DesiredName = TEXT(""));
 
 	FGLTFJsonMaterialIndex GetOrAddMaterial(const UMaterialInterface* Material, const FString& DesiredName = TEXT(""));
+	FGLTFJsonSamplerIndex GetOrAddSampler(const UTexture* Texture, const FString& DesiredName = TEXT(""));
 	FGLTFJsonTextureIndex GetOrAddTexture(const UTexture2D* Texture, const FString& DesiredName = TEXT(""));
 	FGLTFJsonTextureIndex GetOrAddTexture(const ULightMapTexture2D* LightMapTexture2D, const FString& DesiredName = TEXT(""));
 	FGLTFJsonLightMapIndex GetOrAddLightMap(const UStaticMeshComponent* StaticMeshComponent, const FString& DesiredName = TEXT(""));
@@ -54,6 +55,7 @@ private:
 	FGLTFStaticMeshConverter StaticMeshConverter;
 
 	FGLTFMaterialConverter MaterialConverter;
+	FGLTFTextureSamplerConverter TextureSamplerConverter;
 	FGLTFTexture2DConverter Texture2DConverter;
 	FGLTFLightMapTexture2DConverter LightMapTexture2DConverter;
 	FGLTFLightMapConverter LightMapConverter;
diff --git a/Engine/Plugins/Enterprise/GLTFExporter/MigratedFromGit/Source/GLTFExporter/Private/Converters/GLTFTextureConverters.cpp b/Engine/Plugins/Enterprise/GLTFExporter/MigratedFromGit/Source/GLTFExporter/Private/Converters/GLTFTextureConverters.cpp
--- a/Engine/Plugins/Enterprise/GLTFExporter/MigratedFromGit/Source/GLTFExporter/Private/Converters/GLTFTextureConverters.cpp
+++ b/Engine/Plugins/Enterprise/GLTFExporter/MigratedFromGit/Source/GLTFExporter/Private/Converters/GLTFTextureConverters.cpp
@@ -4,19 +4,41 @@
 #include "Builders/GLTFContainerBuilder.h"
 #include "Converters/GLTFConverterUtility.h"
 
-FGLTFJsonTextureIndex FGLTFTexture2DConverter::Add(FGLTFConvertBuilder& Builder, const FString& Name, const UTexture2D* Texture2D)
+FGLTFJsonSamplerIndex FGLTFTextureSamplerConverter::Add(FGLTFConvertBuilder& Builder, const FString& Name, const UTexture* Texture)
 {
 	// TODO: maybe we should reuse existing samplers?
+
 	FGLTFJsonSampler JsonSampler;
 	JsonSampler.Name = Name;
-	JsonSampler.MinFilter = FGLTFConverterUtility::ConvertMinFilter(Texture2D->Filter, Texture2D->LODGroup);
-	JsonSampler.MagFilter = FGLTFConverterUtility::ConvertMagFilter(Texture2D->Filter, Texture2D->LODGroup);
-	JsonSampler.WrapS = FGLTFConverterUtility::ConvertWrap(Texture2D->AddressX);
-	JsonSampler.WrapT = FGLTFConverterUtility::ConvertWrap(Texture2D->AddressY);
+
+	JsonSampler.MinFilter = FGLTFConverterUtility::ConvertMinFilter(Texture->Filter, Texture->LODGroup);
+	JsonSampler.MagFilter = FGLTFConverterUtility::ConvertMagFilter(Texture->Filter, Texture->LODGroup);
+
+	if (const UTexture2D* Texture2D = Cast<UTexture2D>(Texture))
+	{
+		JsonSampler.WrapS = FGLTFConverterUtility::ConvertWrap(Texture2D->AddressX);
+		JsonSampler.WrapT = FGLTFConverterUtility::ConvertWrap(Texture2D->AddressY);
+	}
+	else if (const UTextureRenderTarget2D* RenderTarget2D = Cast<UTextureRenderTarget2D>(Texture))
+	{
+		JsonSampler.WrapS = FGLTFConverterUtility::ConvertWrap(RenderTarget2D->AddressX);
+		JsonSampler.WrapT = FGLTFConverterUtility::ConvertWrap(RenderTarget2D->AddressY);
+	}
+	else
+	{
+		// TODO: handle this?
+	}
+
+	return Builder.AddSampler(JsonSampler);
+}
+
+FGLTFJsonTextureIndex FGLTFTexture2DConverter::Add(FGLTFConvertBuilder& Builder, const FString& Name, const UTexture2D* Texture2D)
+{
+	// TODO: add RGBE encoding information for TSF_RGBE8 and TSF_BGRE8
 
 	FGLTFJsonTexture JsonTexture;
 	JsonTexture.Name = Name;
-	JsonTexture.Sampler = Builder.AddSampler(JsonSampler);
+	JsonTexture.Sampler = Builder.GetOrAddSampler(Texture2D);
 	JsonTexture.Source = Builder.AddImage(Texture2D->Source, Name);
 
 	return Builder.AddTexture(JsonTexture);
diff --git a/Engine/Plugins/Enterprise/GLTFExporter/MigratedFromGit/Source/GLTFExporter/Private/Converters/GLTFTextureConverters.h b/Engine/Plugins/Enterprise/GLTFExporter/MigratedFromGit/Source/GLTFExporter/Private/Converters/GLTFTextureConverters.h
--- a/Engine/Plugins/Enterprise/GLTFExporter/MigratedFromGit/Source/GLTFExporter/Private/Converters/GLTFTextureConverters.h
+++ b/Engine/Plugins/Enterprise/GLTFExporter/MigratedFromGit/Source/GLTFExporter/Private/Converters/GLTFTextureConverters.h
@@ -6,6 +6,11 @@
 #include "Converters/GLTFConverter.h"
 #include "Engine.h"
 
+class FGLTFTextureSamplerConverter final : public TGLTFConverter<FGLTFJsonSamplerIndex, const UTexture*>
+{
+	FGLTFJsonSamplerIndex Add(FGLTFConvertBuilder& Builder, const FString& Name, const UTexture* Texture) override;
+};
+
 class FGLTFTexture2DConverter final : public TGLTFConverter<FGLTFJsonTextureIndex, const UTexture2D*>
 {
 	FGLTFJsonTextureIndex Add(FGLTFConvertBuilder& Builder, const FString& Name, const UTexture2D* Texture2D) override;
