[GLTFExporter] refactor(texture): add GLTFTextureUtility

commit 5abf33b6b7fa3b6bb3d4392641914dad33ad84b1
Date: Fri, 26 Jun 2020 01:51:42 +0200

    refactor(texture): add GLTFTextureUtility
    (cherry picked from commit 8ffc7d3b3674fa23dfe952c4057a5010c74fd942​)

#skipci

diff --git a/Engine/Plugins/Enterprise/GLTFExporter/MigratedFromGit/Source/GLTFExporter/Private/Converters/GLTFMaterialUtility.cpp b/Engine/Plugins/Enterprise/GLTFExporter/MigratedFromGit/Source/GLTFExporter/Private/Converters/GLTFMaterialUtility.cpp
--- a/Engine/Plugins/Enterprise/GLTFExporter/MigratedFromGit/Source/GLTFExporter/Private/Converters/GLTFMaterialUtility.cpp
+++ b/Engine/Plugins/Enterprise/GLTFExporter/MigratedFromGit/Source/GLTFExporter/Private/Converters/GLTFMaterialUtility.cpp
@@ -1,7 +1,6 @@
 // Copyright Epic Games, Inc. All Rights Reserved.
 
-#include "GLTFMaterialUtility.h"
-
+#include "Converters/GLTFMaterialUtility.h"
 #include "Engine/TextureRenderTarget2D.h"
 #include "CanvasItem.h"
 #include "CanvasTypes.h"
diff --git a/Engine/Plugins/Enterprise/GLTFExporter/MigratedFromGit/Source/GLTFExporter/Private/Converters/GLTFTextureConverters.cpp b/Engine/Plugins/Enterprise/GLTFExporter/MigratedFromGit/Source/GLTFExporter/Private/Converters/GLTFTextureConverters.cpp
--- a/Engine/Plugins/Enterprise/GLTFExporter/MigratedFromGit/Source/GLTFExporter/Private/Converters/GLTFTextureConverters.cpp
+++ b/Engine/Plugins/Enterprise/GLTFExporter/MigratedFromGit/Source/GLTFExporter/Private/Converters/GLTFTextureConverters.cpp
@@ -3,6 +3,7 @@
 #include "Converters/GLTFTextureConverters.h"
 #include "Builders/GLTFContainerBuilder.h"
 #include "Converters/GLTFConverterUtility.h"
+#include "Converters/GLTFTextureUtility.h"
 
 FGLTFJsonSamplerIndex FGLTFTextureSamplerConverter::Add(FGLTFConvertBuilder& Builder, const FString& Name, const UTexture* Texture)
 {
@@ -14,20 +15,13 @@ FGLTFJsonSamplerIndex FGLTFTextureSamplerConverter::Add(FGLTFConvertBuilder& Bui
 	JsonSampler.MinFilter = FGLTFConverterUtility::ConvertMinFilter(Texture->Filter, Texture->LODGroup);
 	JsonSampler.MagFilter = FGLTFConverterUtility::ConvertMagFilter(Texture->Filter, Texture->LODGroup);
 
-	if (const UTexture2D* Texture2D = Cast<UTexture2D>(Texture))
-	{
-		JsonSampler.WrapS = FGLTFConverterUtility::ConvertWrap(Texture2D->AddressX);
-		JsonSampler.WrapT = FGLTFConverterUtility::ConvertWrap(Texture2D->AddressY);
-	}
-	else if (const UTextureRenderTarget2D* RenderTarget2D = Cast<UTextureRenderTarget2D>(Texture))
-	{
-		JsonSampler.WrapS = FGLTFConverterUtility::ConvertWrap(RenderTarget2D->AddressX);
-		JsonSampler.WrapT = FGLTFConverterUtility::ConvertWrap(RenderTarget2D->AddressY);
-	}
-	else
-	{
-		// TODO: handle this?
-	}
+	const TextureAddress AddressX = FGLTFTextureUtility::GetAddressX(Texture);
+	const TextureAddress AddressY = FGLTFTextureUtility::GetAddressY(Texture);
+
+	// TODO: report error if AddressX == TA_MAX || AddressY == TA_MAX
+
+	JsonSampler.WrapS = FGLTFConverterUtility::ConvertWrap(AddressX);
+	JsonSampler.WrapT = FGLTFConverterUtility::ConvertWrap(AddressY);
 
 	return Builder.AddSampler(JsonSampler);
 }
diff --git a/Engine/Plugins/Enterprise/GLTFExporter/MigratedFromGit/Source/GLTFExporter/Private/Converters/GLTFTextureUtility.cpp b/Engine/Plugins/Enterprise/GLTFExporter/MigratedFromGit/Source/GLTFExporter/Private/Converters/GLTFTextureUtility.cpp
new file mode 100644
--- /dev/null
+++ b/Engine/Plugins/Enterprise/GLTFExporter/MigratedFromGit/Source/GLTFExporter/Private/Converters/GLTFTextureUtility.cpp
@@ -0,0 +1,48 @@
+// Copyright Epic Games, Inc. All Rights Reserved.
+
+#include "Converters/GLTFTextureUtility.h"
+#include "Engine/Texture2DArray.h"
+
+TextureAddress FGLTFTextureUtility::GetAddressX(const UTexture* Texture)
+{
+	if (const UTexture2D* Texture2D = Cast<UTexture2D>(Texture))
+	{
+		return Texture2D->AddressX;
+	}
+
+	if (const UTexture2DArray* Texture2DArray = Cast<UTexture2DArray>(Texture))
+	{
+		return Texture2DArray->AddressX;
+	}
+
+	if (const UTextureRenderTarget2D* RenderTarget2D = Cast<UTextureRenderTarget2D>(Texture))
+	{
+		return RenderTarget2D->AddressX;
+	}
+
+	// TODO: add support for UMediaTexture?
+
+	return TextureAddress::TA_MAX;
+}
+
+TextureAddress FGLTFTextureUtility::GetAddressY(const UTexture* Texture)
+{
+	if (const UTexture2D* Texture2D = Cast<UTexture2D>(Texture))
+	{
+		return Texture2D->AddressY;
+	}
+
+	if (const UTexture2DArray* Texture2DArray = Cast<UTexture2DArray>(Texture))
+	{
+		return Texture2DArray->AddressY;
+	}
+
+	if (const UTextureRenderTarget2D* RenderTarget2D = Cast<UTextureRenderTarget2D>(Texture))
+	{
+		return RenderTarget2D->AddressY;
+	}
+
+	// TODO: add support for UMediaTexture?
+
+	return TextureAddress::TA_MAX;
+}
diff --git a/Engine/Plugins/Enterprise/GLTFExporter/MigratedFromGit/Source/GLTFExporter/Private/Converters/GLTFTextureUtility.h b/Engine/Plugins/Enterprise/GLTFExporter/MigratedFromGit/Source/GLTFExporter/Private/Converters/GLTFTextureUtility.h
new file mode 100644
--- /dev/null
+++ b/Engine/Plugins/Enterprise/GLTFExporter/MigratedFromGit/Source/GLTFExporter/Private/Converters/GLTFTextureUtility.h
@@ -0,0 +1,11 @@
+// Copyright Epic Games, Inc. All Rights Reserved.
+
+#pragma once
+
+#include "Engine.h"
+
+struct FGLTFTextureUtility
+{
+	static TextureAddress GetAddressX(const UTexture* Texture);
+	static TextureAddress GetAddressY(const UTexture* Texture);
+};
