[GLTFExporter] refactor: add support for lightmap textures in general sampler converter

commit 30b6a489782fc6126ecb595d420461379d2c86c7
Date: Tue, 7 Jul 2020 19:46:50 +0200

    refactor: add support for lightmap textures in general sampler converter
    (cherry picked from commit eafafa29dea01079d73c9293b6dba51c84386d2c​)

#skipci

diff --git a/Engine/Plugins/Enterprise/GLTFExporter/MigratedFromGit/Source/GLTFExporter/Private/Converters/GLTFTextureConverters.cpp b/Engine/Plugins/Enterprise/GLTFExporter/MigratedFromGit/Source/GLTFExporter/Private/Converters/GLTFTextureConverters.cpp
--- a/Engine/Plugins/Enterprise/GLTFExporter/MigratedFromGit/Source/GLTFExporter/Private/Converters/GLTFTextureConverters.cpp
+++ b/Engine/Plugins/Enterprise/GLTFExporter/MigratedFromGit/Source/GLTFExporter/Private/Converters/GLTFTextureConverters.cpp
@@ -12,8 +12,17 @@ FGLTFJsonSamplerIndex FGLTFTextureSamplerConverter::Add(FGLTFConvertBuilder& Bui
 	FGLTFJsonSampler JsonSampler;
 	JsonSampler.Name = Name;
 
-	JsonSampler.MinFilter = FGLTFConverterUtility::ConvertMinFilter(Texture->Filter, Texture->LODGroup);
-	JsonSampler.MagFilter = FGLTFConverterUtility::ConvertMagFilter(Texture->Filter, Texture->LODGroup);
+	if (Texture->IsA<ULightMapTexture2D>())
+	{
+		// Override default filter settings for LightMap textures (which otherwise is "nearest")
+		JsonSampler.MinFilter = EGLTFJsonTextureFilter::LinearMipmapLinear;
+		JsonSampler.MagFilter = EGLTFJsonTextureFilter::Linear;
+	}
+	else
+	{
+		JsonSampler.MinFilter = FGLTFConverterUtility::ConvertMinFilter(Texture->Filter, Texture->LODGroup);
+		JsonSampler.MagFilter = FGLTFConverterUtility::ConvertMagFilter(Texture->Filter, Texture->LODGroup);
+	}
 
 	const TextureAddress AddressX = FGLTFTextureUtility::GetAddressX(Texture);
 	const TextureAddress AddressY = FGLTFTextureUtility::GetAddressY(Texture);
@@ -46,20 +55,11 @@ FGLTFJsonTextureIndex FGLTFTextureCubeConverter::Add(FGLTFConvertBuilder& Builde
 
 FGLTFJsonTextureIndex FGLTFLightMapTexture2DConverter::Add(FGLTFConvertBuilder& Builder, const FString& Name, const ULightMapTexture2D* LightMapTexture2D)
 {
-	// TODO: maybe we should reuse existing samplers?
-	FGLTFJsonSampler JsonSampler;
-	JsonSampler.Name = Name;
-
-	// TODO: are these filters ok to use as default? The actual texture uses "nearest"
-	JsonSampler.MinFilter = EGLTFJsonTextureFilter::LinearMipmapLinear;
-	JsonSampler.MagFilter = EGLTFJsonTextureFilter::Linear;
-
-	JsonSampler.WrapS = FGLTFConverterUtility::ConvertWrap(LightMapTexture2D->AddressX);
-	JsonSampler.WrapT = FGLTFConverterUtility::ConvertWrap(LightMapTexture2D->AddressY);
+	// TODO: add RGBE encoding information for TSF_RGBE8 and TSF_BGRE8
 
 	FGLTFJsonTexture JsonTexture;
 	JsonTexture.Name = Name;
-	JsonTexture.Sampler = Builder.AddSampler(JsonSampler);
+	JsonTexture.Sampler = Builder.GetOrAddSampler(LightMapTexture2D);
 	JsonTexture.Source = Builder.AddImage(LightMapTexture2D->Source, Name);
 
 	return Builder.AddTexture(JsonTexture);
