[GLTFExporter] refactor: remove redundant texture converter for lightmaps

commit 884375cfccef6375ffded6d7f501f0314f7e6619
Date: Tue, 7 Jul 2020 19:52:26 +0200

    refactor: remove redundant texture converter for lightmaps
    (cherry picked from commit ac5cc881f7a1442a38b00f82791c9f247466a6ff​)

#skipci

diff --git a/Engine/Plugins/Enterprise/GLTFExporter/MigratedFromGit/Source/GLTFExporter/Private/Builders/GLTFConvertBuilder.cpp b/Engine/Plugins/Enterprise/GLTFExporter/MigratedFromGit/Source/GLTFExporter/Private/Builders/GLTFConvertBuilder.cpp
--- a/Engine/Plugins/Enterprise/GLTFExporter/MigratedFromGit/Source/GLTFExporter/Private/Builders/GLTFConvertBuilder.cpp
+++ b/Engine/Plugins/Enterprise/GLTFExporter/MigratedFromGit/Source/GLTFExporter/Private/Builders/GLTFConvertBuilder.cpp
@@ -134,16 +134,6 @@ FGLTFJsonTextureIndex FGLTFConvertBuilder::GetOrAddTexture(const UTexture2D* Tex
 	return Texture2DConverter.GetOrAdd(*this, DesiredName.IsEmpty() ? Texture->GetName() : DesiredName, Texture);
 }
 
-FGLTFJsonTextureIndex FGLTFConvertBuilder::GetOrAddTexture(const ULightMapTexture2D* LightMapTexture2D, const FString& DesiredName)
-{
-	if (LightMapTexture2D == nullptr)
-	{
-		return FGLTFJsonTextureIndex(INDEX_NONE);
-	}
-
-	return LightMapTexture2DConverter.GetOrAdd(*this, DesiredName.IsEmpty() ? LightMapTexture2D->GetName() : DesiredName, LightMapTexture2D);
-}
-
 FGLTFJsonLightMapIndex FGLTFConvertBuilder::GetOrAddLightMap(const UStaticMeshComponent* StaticMeshComponent, const FString& DesiredName)
 {
 	if (StaticMeshComponent == nullptr)
diff --git a/Engine/Plugins/Enterprise/GLTFExporter/MigratedFromGit/Source/GLTFExporter/Private/Builders/GLTFConvertBuilder.h b/Engine/Plugins/Enterprise/GLTFExporter/MigratedFromGit/Source/GLTFExporter/Private/Builders/GLTFConvertBuilder.h
--- a/Engine/Plugins/Enterprise/GLTFExporter/MigratedFromGit/Source/GLTFExporter/Private/Builders/GLTFConvertBuilder.h
+++ b/Engine/Plugins/Enterprise/GLTFExporter/MigratedFromGit/Source/GLTFExporter/Private/Builders/GLTFConvertBuilder.h
@@ -33,7 +33,6 @@ public:
 	FGLTFJsonMaterialIndex GetOrAddMaterial(const UMaterialInterface* Material, const FString& DesiredName = TEXT(""));
 	FGLTFJsonSamplerIndex GetOrAddSampler(const UTexture* Texture, const FString& DesiredName = TEXT(""));
 	FGLTFJsonTextureIndex GetOrAddTexture(const UTexture2D* Texture, const FString& DesiredName = TEXT(""));
-	FGLTFJsonTextureIndex GetOrAddTexture(const ULightMapTexture2D* LightMapTexture2D, const FString& DesiredName = TEXT(""));
 	FGLTFJsonLightMapIndex GetOrAddLightMap(const UStaticMeshComponent* StaticMeshComponent, const FString& DesiredName = TEXT(""));
 
 	FGLTFJsonNodeIndex GetOrAddNode(const USceneComponent* SceneComponent, const FString& DesiredName = TEXT(""));
@@ -57,7 +56,6 @@ private:
 	FGLTFMaterialConverter MaterialConverter;
 	FGLTFTextureSamplerConverter TextureSamplerConverter;
 	FGLTFTexture2DConverter Texture2DConverter;
-	FGLTFLightMapTexture2DConverter LightMapTexture2DConverter;
 	FGLTFLightMapConverter LightMapConverter;
 
 	FGLTFSceneComponentConverter SceneComponentConverter;
diff --git a/Engine/Plugins/Enterprise/GLTFExporter/MigratedFromGit/Source/GLTFExporter/Private/Converters/GLTFTextureConverters.cpp b/Engine/Plugins/Enterprise/GLTFExporter/MigratedFromGit/Source/GLTFExporter/Private/Converters/GLTFTextureConverters.cpp
--- a/Engine/Plugins/Enterprise/GLTFExporter/MigratedFromGit/Source/GLTFExporter/Private/Converters/GLTFTextureConverters.cpp
+++ b/Engine/Plugins/Enterprise/GLTFExporter/MigratedFromGit/Source/GLTFExporter/Private/Converters/GLTFTextureConverters.cpp
@@ -52,15 +52,3 @@ FGLTFJsonTextureIndex FGLTFTextureCubeConverter::Add(FGLTFConvertBuilder& Builde
 	// TODO: implement support
 	return FGLTFJsonTextureIndex(INDEX_NONE);
 }
-
-FGLTFJsonTextureIndex FGLTFLightMapTexture2DConverter::Add(FGLTFConvertBuilder& Builder, const FString& Name, const ULightMapTexture2D* LightMapTexture2D)
-{
-	// TODO: add RGBE encoding information for TSF_RGBE8 and TSF_BGRE8
-
-	FGLTFJsonTexture JsonTexture;
-	JsonTexture.Name = Name;
-	JsonTexture.Sampler = Builder.GetOrAddSampler(LightMapTexture2D);
-	JsonTexture.Source = Builder.AddImage(LightMapTexture2D->Source, Name);
-
-	return Builder.AddTexture(JsonTexture);
-}
diff --git a/Engine/Plugins/Enterprise/GLTFExporter/MigratedFromGit/Source/GLTFExporter/Private/Converters/GLTFTextureConverters.h b/Engine/Plugins/Enterprise/GLTFExporter/MigratedFromGit/Source/GLTFExporter/Private/Converters/GLTFTextureConverters.h
--- a/Engine/Plugins/Enterprise/GLTFExporter/MigratedFromGit/Source/GLTFExporter/Private/Converters/GLTFTextureConverters.h
+++ b/Engine/Plugins/Enterprise/GLTFExporter/MigratedFromGit/Source/GLTFExporter/Private/Converters/GLTFTextureConverters.h
@@ -20,8 +20,3 @@ class FGLTFTextureCubeConverter final : public TGLTFConverter<FGLTFJsonTextureIn
 {
 	FGLTFJsonTextureIndex Add(FGLTFConvertBuilder& Builder, const FString& Name, const UTextureCube* TextureCube) override;
 };
-
-class FGLTFLightMapTexture2DConverter final : public TGLTFConverter<FGLTFJsonTextureIndex, const ULightMapTexture2D*>
-{
-	FGLTFJsonTextureIndex Add(FGLTFConvertBuilder& Builder, const FString& Name, const ULightMapTexture2D* LightMapTexture2D) override;
-};
