[GLTFExporter] fix(texture): avoid overwriting textures with same filename

commit 1a49954dfe831a2cc8633a5795f176a408eea576
Date: Tue, 7 Jul 2020 21:52:59 +0200

    fix(texture): avoid overwriting textures with same filename
    (cherry picked from commit f1671ade4bbb9004fafa605662d77257932d47c6​)

#skipci

diff --git a/Engine/Plugins/Enterprise/GLTFExporter/MigratedFromGit/Source/GLTFExporter/Private/Builders/GLTFBuilderUtility.cpp b/Engine/Plugins/Enterprise/GLTFExporter/MigratedFromGit/Source/GLTFExporter/Private/Builders/GLTFBuilderUtility.cpp
--- a/Engine/Plugins/Enterprise/GLTFExporter/MigratedFromGit/Source/GLTFExporter/Private/Builders/GLTFBuilderUtility.cpp
+++ b/Engine/Plugins/Enterprise/GLTFExporter/MigratedFromGit/Source/GLTFExporter/Private/Builders/GLTFBuilderUtility.cpp
@@ -26,3 +26,40 @@ bool FGLTFBuilderUtility::CompressImage(const void* RawData, int64 ByteLength, i
 	OutCompressedData.Append(CompressedData);
 	return true;
 }
+
+FString FGLTFBuilderUtility::GetUniqueFilename(const FString& BaseFilename, const FString& FileExtension, const TSet<FString>& UniqueFilenames)
+{
+	FString Filename = BaseFilename + FileExtension;
+	if (!UniqueFilenames.Contains(Filename))
+	{
+		return Filename;
+	}
+
+	FString NewBaseFilename = BaseFilename;
+
+	// Remove potentially existing suffix numbers
+	for (int32 SuffixLen = 1; SuffixLen < NewBaseFilename.Len(); ++SuffixLen)
+	{
+		const TCHAR Char = NewBaseFilename[NewBaseFilename.Len() - SuffixLen];
+		if (!FChar::IsDigit(Char))
+		{
+			if (Char == TEXT('_') && SuffixLen > 0)
+			{
+				NewBaseFilename.LeftChopInline(SuffixLen);
+			}
+			break;
+		}
+	}
+
+	NewBaseFilename += TEXT('_');
+
+	int32 Suffix = 1;
+	do
+	{
+		Filename = NewBaseFilename + FString::FromInt(Suffix) + FileExtension;
+		Suffix++;
+	}
+	while (UniqueFilenames.Contains(Filename));
+
+	return Filename;
+}
diff --git a/Engine/Plugins/Enterprise/GLTFExporter/MigratedFromGit/Source/GLTFExporter/Private/Builders/GLTFBuilderUtility.h b/Engine/Plugins/Enterprise/GLTFExporter/MigratedFromGit/Source/GLTFExporter/Private/Builders/GLTFBuilderUtility.h
--- a/Engine/Plugins/Enterprise/GLTFExporter/MigratedFromGit/Source/GLTFExporter/Private/Builders/GLTFBuilderUtility.h
+++ b/Engine/Plugins/Enterprise/GLTFExporter/MigratedFromGit/Source/GLTFExporter/Private/Builders/GLTFBuilderUtility.h
@@ -34,4 +34,6 @@ struct FGLTFBuilderUtility
 			default:                      return TEXT(".unknown");
 		}
 	}
+
+	static FString GetUniqueFilename(const FString& BaseFilename, const FString& FileExtension, const TSet<FString>& UniqueFilenames);
 };
diff --git a/Engine/Plugins/Enterprise/GLTFExporter/MigratedFromGit/Source/GLTFExporter/Private/Builders/GLTFImageBuilder.cpp b/Engine/Plugins/Enterprise/GLTFExporter/MigratedFromGit/Source/GLTFExporter/Private/Builders/GLTFImageBuilder.cpp
--- a/Engine/Plugins/Enterprise/GLTFExporter/MigratedFromGit/Source/GLTFExporter/Private/Builders/GLTFImageBuilder.cpp
+++ b/Engine/Plugins/Enterprise/GLTFExporter/MigratedFromGit/Source/GLTFExporter/Private/Builders/GLTFImageBuilder.cpp
@@ -118,21 +118,25 @@ FGLTFJsonImageIndex FGLTFImageBuilder::AddImage(const FTextureSource& Image, con
 bool FGLTFImageBuilder::Serialize(FArchive& Archive, const FString& FilePath)
 {
 	const FString ImageDir = FPaths::GetPath(FilePath);
+	TSet<FString> UniqueImageUris;
 
 	for (const auto& DataPair : ImageDataLookup)
 	{
 		const TArray64<uint8>& ImageData = DataPair.Value;
 		FGLTFJsonImage& JsonImage = GetImage(DataPair.Key);
 
-		const FString ImageUri = JsonImage.Name + FGLTFBuilderUtility::GetFileExtension(JsonImage.MimeType);
-		const FString ImagePath = FPaths::Combine(ImageDir, ImageUri);
+		const TCHAR* Extension = FGLTFBuilderUtility::GetFileExtension(JsonImage.MimeType);
+		const FString ImageUri = FGLTFBuilderUtility::GetUniqueFilename(JsonImage.Name, Extension, UniqueImageUris);
 
+		const FString ImagePath = FPaths::Combine(ImageDir, ImageUri);
 		if (!FFileHelper::SaveArrayToFile(ImageData, *ImagePath))
 		{
 			// TODO: report error
 			continue;
 		}
 
+		UniqueImageUris.Add(ImageUri);
+
 		JsonImage.Uri = ImageUri;
 		JsonImage.Name.Empty(); // URI already contains name
 		JsonImage.MimeType = EGLTFJsonMimeType::None; // Not required if external file
