[GLTFExporter] refactor!: change thumbnail image index to texture index in variants json to avoid re-exporting the same image

commit a382ef7a8364b66f527394f9d41b7eb7d36ece52
Date: Wed, 8 Jul 2020 10:32:05 +0200

    refactor!: change thumbnail image index to texture index in variants json to avoid re-exporting the same image
    (cherry picked from commit 959a70995f50c1bcf4f7d0f165fe15804c1968d9​)

#skipci

diff --git a/Engine/Plugins/Enterprise/GLTFExporter/MigratedFromGit/Source/GLTFExporter/Private/Converters/GLTFVariantSetConverters.cpp b/Engine/Plugins/Enterprise/GLTFExporter/MigratedFromGit/Source/GLTFExporter/Private/Converters/GLTFVariantSetConverters.cpp
--- a/Engine/Plugins/Enterprise/GLTFExporter/MigratedFromGit/Source/GLTFExporter/Private/Converters/GLTFVariantSetConverters.cpp
+++ b/Engine/Plugins/Enterprise/GLTFExporter/MigratedFromGit/Source/GLTFExporter/Private/Converters/GLTFVariantSetConverters.cpp
@@ -73,8 +73,7 @@ bool FGLTFLevelVariantSetsConverter::TryParseJsonVariant(FGLTFConvertBuilder& Bu
 
 	if (const UTexture2D* Thumbnail = const_cast<UVariant*>(Variant)->GetThumbnail())
 	{
-		const FString ThumbnailName = ThumbnailPrefix + JsonVariant.Name + TEXT("_thumbnail");
-		JsonVariant.Thumbnail = Builder.AddImage(Thumbnail->Source, ThumbnailName);
+		JsonVariant.Thumbnail = Builder.GetOrAddTexture(Thumbnail);
 	}
 
 	OutVariant = JsonVariant;
diff --git a/Engine/Plugins/Enterprise/GLTFExporter/MigratedFromGit/Source/GLTFExporter/Private/Json/GLTFJsonLevelVariantSets.h b/Engine/Plugins/Enterprise/GLTFExporter/MigratedFromGit/Source/GLTFExporter/Private/Json/GLTFJsonLevelVariantSets.h
--- a/Engine/Plugins/Enterprise/GLTFExporter/MigratedFromGit/Source/GLTFExporter/Private/Json/GLTFJsonLevelVariantSets.h
+++ b/Engine/Plugins/Enterprise/GLTFExporter/MigratedFromGit/Source/GLTFExporter/Private/Json/GLTFJsonLevelVariantSets.h
@@ -40,7 +40,7 @@ struct FGLTFJsonVariant
 	FString Name;
 	bool    bIsActive;
 
-	FGLTFJsonImageIndex          Thumbnail;
+	FGLTFJsonTextureIndex        Thumbnail;
 	TArray<FGLTFJsonVariantNode> Nodes;
 
 	template <class CharType = TCHAR, class PrintPolicy = TPrettyJsonPrintPolicy<CharType>>
