[GLTFExporter] refactor: simplify GLTFBuilderUtility::GetUniqueFilename to avoid unnecessary conflicts

commit bad9418a113024e8dee9398537d40123c073897f
Date: Wed, 8 Jul 2020 10:32:50 +0200

    refactor: simplify GLTFBuilderUtility::GetUniqueFilename to avoid unnecessary conflicts
    (cherry picked from commit e5927890a69da48baecc5b3fc9cacb6ed4f79f80​)

#skipci

diff --git a/Engine/Plugins/Enterprise/GLTFExporter/MigratedFromGit/Source/GLTFExporter/Private/Builders/GLTFBuilderUtility.cpp b/Engine/Plugins/Enterprise/GLTFExporter/MigratedFromGit/Source/GLTFExporter/Private/Builders/GLTFBuilderUtility.cpp
--- a/Engine/Plugins/Enterprise/GLTFExporter/MigratedFromGit/Source/GLTFExporter/Private/Builders/GLTFBuilderUtility.cpp
+++ b/Engine/Plugins/Enterprise/GLTFExporter/MigratedFromGit/Source/GLTFExporter/Private/Builders/GLTFBuilderUtility.cpp
@@ -35,28 +35,10 @@ FString FGLTFBuilderUtility::GetUniqueFilename(const FString& BaseFilename, cons
 		return Filename;
 	}
 
-	FString NewBaseFilename = BaseFilename;
-
-	// Remove potentially existing suffix numbers
-	for (int32 SuffixLen = 1; SuffixLen < NewBaseFilename.Len(); ++SuffixLen)
-	{
-		const TCHAR Char = NewBaseFilename[NewBaseFilename.Len() - SuffixLen];
-		if (!FChar::IsDigit(Char))
-		{
-			if (Char == TEXT('_') && SuffixLen > 0)
-			{
-				NewBaseFilename.LeftChopInline(SuffixLen);
-			}
-			break;
-		}
-	}
-
-	NewBaseFilename += TEXT('_');
-
 	int32 Suffix = 1;
 	do
 	{
-		Filename = NewBaseFilename + FString::FromInt(Suffix) + FileExtension;
+		Filename = BaseFilename + TEXT('_') + FString::FromInt(Suffix) + FileExtension;
 		Suffix++;
 	}
 	while (UniqueFilenames.Contains(Filename));
diff --git a/Engine/Plugins/Enterprise/GLTFExporter/MigratedFromGit/Source/GLTFExporter/Private/Converters/GLTFVariantSetConverters.cpp b/Engine/Plugins/Enterprise/GLTFExporter/MigratedFromGit/Source/GLTFExporter/Private/Converters/GLTFVariantSetConverters.cpp
--- a/Engine/Plugins/Enterprise/GLTFExporter/MigratedFromGit/Source/GLTFExporter/Private/Converters/GLTFVariantSetConverters.cpp
+++ b/Engine/Plugins/Enterprise/GLTFExporter/MigratedFromGit/Source/GLTFExporter/Private/Converters/GLTFVariantSetConverters.cpp
@@ -73,6 +73,7 @@ bool FGLTFLevelVariantSetsConverter::TryParseJsonVariant(FGLTFConvertBuilder& Bu
 
 	if (const UTexture2D* Thumbnail = const_cast<UVariant*>(Variant)->GetThumbnail())
 	{
+		// TODO: if thumbnail is generated from viewport, give it a variant-relevant name (instead of "Texture2D_#")
 		JsonVariant.Thumbnail = Builder.GetOrAddTexture(Thumbnail);
 	}
 
