[GLTFExporter] refactor(texture): move GLTFMaterialUtility::CreateTransientTexture to GLTFTextureUtility and generalize it

commit b7c70657733d6be1a041789ee0cb0b80696aefc4
Date: Wed, 8 Jul 2020 10:43:29 +0200

    refactor(texture): move GLTFMaterialUtility::CreateTransientTexture to GLTFTextureUtility and generalize it
    (cherry picked from commit f901e5376aebe88fa2676faa424f8b8d49c977b0​)

#skipci

diff --git a/Engine/Plugins/Enterprise/GLTFExporter/MigratedFromGit/Source/GLTFExporter/Private/Converters/GLTFMaterialUtility.cpp b/Engine/Plugins/Enterprise/GLTFExporter/MigratedFromGit/Source/GLTFExporter/Private/Converters/GLTFMaterialUtility.cpp
--- a/Engine/Plugins/Enterprise/GLTFExporter/MigratedFromGit/Source/GLTFExporter/Private/Converters/GLTFMaterialUtility.cpp
+++ b/Engine/Plugins/Enterprise/GLTFExporter/MigratedFromGit/Source/GLTFExporter/Private/Converters/GLTFMaterialUtility.cpp
@@ -1,6 +1,7 @@
 // Copyright Epic Games, Inc. All Rights Reserved.
 
 #include "Converters/GLTFMaterialUtility.h"
+#include "Converters/GLTFTextureUtility.h"
 #include "Engine/TextureRenderTarget2D.h"
 #include "CanvasItem.h"
 #include "CanvasTypes.h"
@@ -54,32 +55,14 @@ const FExpressionInput* FGLTFMaterialUtility::GetInputForProperty(const UMateria
 
 UTexture2D* FGLTFMaterialUtility::CreateTransientTexture(const FGLTFPropertyBakeOutput& PropertyBakeOutput, bool bUseSRGB)
 {
-	return CreateTransientTexture(
-		PropertyBakeOutput.Pixels,
+	return FGLTFTextureUtility::CreateTransientTexture(
+		PropertyBakeOutput.Pixels.GetData(),
+		PropertyBakeOutput.Pixels.Num() * PropertyBakeOutput.Pixels.GetTypeSize(),
 		PropertyBakeOutput.Size,
 		PropertyBakeOutput.PixelFormat,
 		bUseSRGB);
 }
 
-UTexture2D* FGLTFMaterialUtility::CreateTransientTexture(const TArray<FColor>& Pixels, const FIntPoint& TextureSize, EPixelFormat TextureFormat, bool bUseSRGB)
-{
-	check(TextureSize.X * TextureSize.Y == Pixels.Num());
-
-	UTexture2D* Texture = UTexture2D::CreateTransient(TextureSize.X, TextureSize.Y, TextureFormat);
-	void* MipData = Texture->PlatformData->Mips[0].BulkData.Lock(LOCK_READ_WRITE);
-
-	FMemory::Memcpy(MipData, Pixels.GetData(), Pixels.Num() * Pixels.GetTypeSize());
-	Texture->PlatformData->Mips[0].BulkData.Unlock();
-
-	Texture->SRGB = bUseSRGB ? 1 : 0;
-	Texture->CompressionNone = true;
-	Texture->MipGenSettings = TMGS_NoMipmaps;
-	Texture->CompressionSettings = TC_Default;
-
-	Texture->UpdateResource();
-	return Texture;
-}
-
 bool FGLTFMaterialUtility::CombineTextures(TArray<FColor>& OutPixels, const TArray<FGLTFTextureCombineSource>& Sources, const FIntPoint& OutputSize, const EPixelFormat OutputPixelFormat)
 {
 	UTextureRenderTarget2D* RenderTarget2D = NewObject<UTextureRenderTarget2D>();
diff --git a/Engine/Plugins/Enterprise/GLTFExporter/MigratedFromGit/Source/GLTFExporter/Private/Converters/GLTFMaterialUtility.h b/Engine/Plugins/Enterprise/GLTFExporter/MigratedFromGit/Source/GLTFExporter/Private/Converters/GLTFMaterialUtility.h
--- a/Engine/Plugins/Enterprise/GLTFExporter/MigratedFromGit/Source/GLTFExporter/Private/Converters/GLTFMaterialUtility.h
+++ b/Engine/Plugins/Enterprise/GLTFExporter/MigratedFromGit/Source/GLTFExporter/Private/Converters/GLTFMaterialUtility.h
@@ -46,7 +46,6 @@ struct FGLTFMaterialUtility
 	}
 
 	static UTexture2D* CreateTransientTexture(const FGLTFPropertyBakeOutput& PropertyBakeOutput, bool bUseSRGB = false);
-	static UTexture2D* CreateTransientTexture(const TArray<FColor>& Pixels, const FIntPoint& TextureSize, EPixelFormat TextureFormat, bool bUseSRGB = false);
 	static bool CombineTextures(TArray<FColor>& OutPixels, const TArray<FGLTFTextureCombineSource>& Sources, const FIntPoint& OutputSize, EPixelFormat OutputPixelFormat);
 	static FGLTFPropertyBakeOutput BakeMaterialProperty(const FIntPoint& OutputSize, EMaterialProperty Property, const UMaterialInterface* Material, bool bCopyAlphaFromRedChannel = false);
 
diff --git a/Engine/Plugins/Enterprise/GLTFExporter/MigratedFromGit/Source/GLTFExporter/Private/Converters/GLTFTextureUtility.cpp b/Engine/Plugins/Enterprise/GLTFExporter/MigratedFromGit/Source/GLTFExporter/Private/Converters/GLTFTextureUtility.cpp
--- a/Engine/Plugins/Enterprise/GLTFExporter/MigratedFromGit/Source/GLTFExporter/Private/Converters/GLTFTextureUtility.cpp
+++ b/Engine/Plugins/Enterprise/GLTFExporter/MigratedFromGit/Source/GLTFExporter/Private/Converters/GLTFTextureUtility.cpp
@@ -46,3 +46,21 @@ TextureAddress FGLTFTextureUtility::GetAddressY(const UTexture* Texture)
 
 	return TextureAddress::TA_MAX;
 }
+
+UTexture2D* FGLTFTextureUtility::CreateTransientTexture(const void* RawData, int64 ByteLength, const FIntPoint& Size, EPixelFormat Format, bool bUseSRGB)
+{
+	check(CalculateImageBytes(Size.X, Size.Y, 0, Format) == ByteLength);
+
+	UTexture2D* Texture = UTexture2D::CreateTransient(Size.X, Size.Y, Format);
+	void* MipData = Texture->PlatformData->Mips[0].BulkData.Lock(LOCK_READ_WRITE);
+
+	FMemory::Memcpy(MipData, RawData, ByteLength);
+	Texture->PlatformData->Mips[0].BulkData.Unlock();
+
+	Texture->SRGB = bUseSRGB ? 1 : 0;
+	Texture->CompressionNone = true;
+	Texture->MipGenSettings = TMGS_NoMipmaps;
+
+	Texture->UpdateResource();
+	return Texture;
+}
diff --git a/Engine/Plugins/Enterprise/GLTFExporter/MigratedFromGit/Source/GLTFExporter/Private/Converters/GLTFTextureUtility.h b/Engine/Plugins/Enterprise/GLTFExporter/MigratedFromGit/Source/GLTFExporter/Private/Converters/GLTFTextureUtility.h
--- a/Engine/Plugins/Enterprise/GLTFExporter/MigratedFromGit/Source/GLTFExporter/Private/Converters/GLTFTextureUtility.h
+++ b/Engine/Plugins/Enterprise/GLTFExporter/MigratedFromGit/Source/GLTFExporter/Private/Converters/GLTFTextureUtility.h
@@ -8,4 +8,6 @@ struct FGLTFTextureUtility
 {
 	static TextureAddress GetAddressX(const UTexture* Texture);
 	static TextureAddress GetAddressY(const UTexture* Texture);
+
+	static UTexture2D* CreateTransientTexture(const void* RawData, int64 ByteLength, const FIntPoint& Size, EPixelFormat Format, bool bUseSRGB = false);
 };
