[GLTFExporter] refactor(texture): move GetDefaultFilter from GLTFConverterUtility to GLTFTextureUtility

commit 7f11e1dd1cc5634c49bc4c384dde28a122c4bf77
Date: Wed, 8 Jul 2020 10:48:51 +0200

    refactor(texture): move GetDefaultFilter from GLTFConverterUtility to GLTFTextureUtility
    (cherry picked from commit 4fe0d8564ade5122894d8c11a1571e8c27be10d1​)

#skipci

diff --git a/Engine/Plugins/Enterprise/GLTFExporter/MigratedFromGit/Source/GLTFExporter/Private/Converters/GLTFConverterUtility.cpp b/Engine/Plugins/Enterprise/GLTFExporter/MigratedFromGit/Source/GLTFExporter/Private/Converters/GLTFConverterUtility.cpp
--- a/Engine/Plugins/Enterprise/GLTFExporter/MigratedFromGit/Source/GLTFExporter/Private/Converters/GLTFConverterUtility.cpp
+++ b/Engine/Plugins/Enterprise/GLTFExporter/MigratedFromGit/Source/GLTFExporter/Private/Converters/GLTFConverterUtility.cpp
@@ -1,22 +1,58 @@
 // Copyright Epic Games, Inc. All Rights Reserved.
 
 #include "Converters/GLTFConverterUtility.h"
-#include "Engine/TextureLODSettings.h"
+#include "Converters/GLTFTextureUtility.h"
 
-TextureFilter FGLTFConverterUtility::GetDefaultFilter(TextureGroup LODGroup)
+EGLTFJsonAlphaMode FGLTFConverterUtility::ConvertBlendMode(EBlendMode Mode)
 {
-	// TODO: should this be the running platform?
-	const ITargetPlatform* RunningPlatform = GetTargetPlatformManagerRef().GetRunningTargetPlatform();
-	const UTextureLODSettings& TextureLODSettings = RunningPlatform->GetTextureLODSettings();
-	const FTextureLODGroup& TextureLODGroup = TextureLODSettings.GetTextureLODGroup(LODGroup);
+	switch (Mode)
+	{
+		case BLEND_Opaque:      return EGLTFJsonAlphaMode::Opaque;
+		case BLEND_Translucent: return EGLTFJsonAlphaMode::Blend;
+		case BLEND_Masked:      return EGLTFJsonAlphaMode::Mask;
+		default:                return EGLTFJsonAlphaMode::Opaque; // fallback
+	}
+}
+
+EGLTFJsonTextureWrap FGLTFConverterUtility::ConvertWrap(TextureAddress Address)
+{
+	switch (Address)
+	{
+		case TA_Wrap:   return EGLTFJsonTextureWrap::Repeat;
+		case TA_Mirror: return EGLTFJsonTextureWrap::MirroredRepeat;
+		case TA_Clamp:  return EGLTFJsonTextureWrap::ClampToEdge;
+		default:        return EGLTFJsonTextureWrap::Repeat; // fallback
+	}
+}
+
+EGLTFJsonTextureFilter FGLTFConverterUtility::ConvertMinFilter(TextureFilter Filter)
+{
+	switch (Filter)
+	{
+		case TF_Nearest:   return EGLTFJsonTextureFilter::NearestMipmapNearest;
+		case TF_Bilinear:  return EGLTFJsonTextureFilter::LinearMipmapNearest;
+		case TF_Trilinear: return EGLTFJsonTextureFilter::LinearMipmapLinear;
+		default:           return EGLTFJsonTextureFilter::None;
+	}
+}
 
-	switch (TextureLODGroup.Filter)
+EGLTFJsonTextureFilter FGLTFConverterUtility::ConvertMagFilter(TextureFilter Filter)
+{
+	switch (Filter)
 	{
-		case ETextureSamplerFilter::Point:             return TF_Nearest;
-		case ETextureSamplerFilter::Bilinear:          return TF_Bilinear;
-		case ETextureSamplerFilter::Trilinear:         return TF_Trilinear;
-		case ETextureSamplerFilter::AnisotropicPoint:  return TF_Nearest;
-		case ETextureSamplerFilter::AnisotropicLinear: return TF_Trilinear;
-		default:                                       return TF_Default; // fallback
+		case TF_Nearest:   return EGLTFJsonTextureFilter::Nearest;
+		case TF_Bilinear:  return EGLTFJsonTextureFilter::Linear;
+		case TF_Trilinear: return EGLTFJsonTextureFilter::Linear;
+		default:           return EGLTFJsonTextureFilter::None;
 	}
 }
+
+EGLTFJsonTextureFilter FGLTFConverterUtility::ConvertMinFilter(TextureFilter Filter, TextureGroup LODGroup)
+{
+	return ConvertMinFilter(Filter == TF_Default ? FGLTFTextureUtility::GetDefaultFilter(LODGroup) : Filter);
+}
+
+EGLTFJsonTextureFilter FGLTFConverterUtility::ConvertMagFilter(TextureFilter Filter, TextureGroup LODGroup)
+{
+	return ConvertMagFilter(Filter == TF_Default ? FGLTFTextureUtility::GetDefaultFilter(LODGroup) : Filter);
+}
diff --git a/Engine/Plugins/Enterprise/GLTFExporter/MigratedFromGit/Source/GLTFExporter/Private/Converters/GLTFConverterUtility.h b/Engine/Plugins/Enterprise/GLTFExporter/MigratedFromGit/Source/GLTFExporter/Private/Converters/GLTFConverterUtility.h
--- a/Engine/Plugins/Enterprise/GLTFExporter/MigratedFromGit/Source/GLTFExporter/Private/Converters/GLTFConverterUtility.h
+++ b/Engine/Plugins/Enterprise/GLTFExporter/MigratedFromGit/Source/GLTFExporter/Private/Converters/GLTFConverterUtility.h
@@ -91,7 +91,7 @@ struct FGLTFConverterUtility
 		return Result;
 	}
 
-	static EGLTFJsonShadingModel ConvertShadingModel(const EMaterialShadingModel ShadingModel)
+	static EGLTFJsonShadingModel ConvertShadingModel(EMaterialShadingModel ShadingModel)
 	{
 		switch (ShadingModel)
 		{
@@ -102,61 +102,15 @@ struct FGLTFConverterUtility
 		}
 	}
 
-	static EGLTFJsonAlphaMode ConvertBlendMode(const EBlendMode Mode)
-	{
-		switch (Mode)
-		{
-			case BLEND_Opaque:      return EGLTFJsonAlphaMode::Opaque;
-			case BLEND_Translucent: return EGLTFJsonAlphaMode::Blend;
-			case BLEND_Masked:      return EGLTFJsonAlphaMode::Mask;
-			default:                return EGLTFJsonAlphaMode::Opaque; // fallback
-		}
-	}
+	static EGLTFJsonAlphaMode ConvertBlendMode(EBlendMode Mode);
 
-	static EGLTFJsonTextureWrap ConvertWrap(TextureAddress Address)
-	{
-		switch (Address)
-		{
-			case TA_Wrap:   return EGLTFJsonTextureWrap::Repeat;
-			case TA_Mirror: return EGLTFJsonTextureWrap::MirroredRepeat;
-			case TA_Clamp:  return EGLTFJsonTextureWrap::ClampToEdge;
-			default:        return EGLTFJsonTextureWrap::Repeat; // fallback
-		}
-	}
+	static EGLTFJsonTextureWrap ConvertWrap(TextureAddress Address);
 
-	static EGLTFJsonTextureFilter ConvertMinFilter(TextureFilter Filter)
-	{
-		switch (Filter)
-		{
-			case TF_Nearest:   return EGLTFJsonTextureFilter::NearestMipmapNearest;
-			case TF_Bilinear:  return EGLTFJsonTextureFilter::LinearMipmapNearest;
-			case TF_Trilinear: return EGLTFJsonTextureFilter::LinearMipmapLinear;
-			default:           return EGLTFJsonTextureFilter::None;
-		}
-	}
-
-	static EGLTFJsonTextureFilter ConvertMagFilter(TextureFilter Filter)
-	{
-		switch (Filter)
-		{
-			case TF_Nearest:   return EGLTFJsonTextureFilter::Nearest;
-			case TF_Bilinear:  return EGLTFJsonTextureFilter::Linear;
-			case TF_Trilinear: return EGLTFJsonTextureFilter::Linear;
-			default:           return EGLTFJsonTextureFilter::None;
-		}
-	}
-
-	static EGLTFJsonTextureFilter ConvertMinFilter(TextureFilter Filter, TextureGroup LODGroup)
-	{
-		return ConvertMinFilter(Filter == TextureFilter::TF_Default ? GetDefaultFilter(LODGroup) : Filter);
-	}
-
-	static EGLTFJsonTextureFilter ConvertMagFilter(TextureFilter Filter, TextureGroup LODGroup)
-	{
-		return ConvertMagFilter(Filter == TextureFilter::TF_Default ? GetDefaultFilter(LODGroup) : Filter);
-	}
+	static EGLTFJsonTextureFilter ConvertMinFilter(TextureFilter Filter);
+	static EGLTFJsonTextureFilter ConvertMagFilter(TextureFilter Filter);
 
-	static TextureFilter GetDefaultFilter(TextureGroup Group);
+	static EGLTFJsonTextureFilter ConvertMinFilter(TextureFilter Filter, TextureGroup LODGroup);
+	static EGLTFJsonTextureFilter ConvertMagFilter(TextureFilter Filter, TextureGroup LODGroup);
 
 	template <typename EnumType>
 	static FString GetEnumDisplayName(EnumType Value)
diff --git a/Engine/Plugins/Enterprise/GLTFExporter/MigratedFromGit/Source/GLTFExporter/Private/Converters/GLTFTextureUtility.cpp b/Engine/Plugins/Enterprise/GLTFExporter/MigratedFromGit/Source/GLTFExporter/Private/Converters/GLTFTextureUtility.cpp
--- a/Engine/Plugins/Enterprise/GLTFExporter/MigratedFromGit/Source/GLTFExporter/Private/Converters/GLTFTextureUtility.cpp
+++ b/Engine/Plugins/Enterprise/GLTFExporter/MigratedFromGit/Source/GLTFExporter/Private/Converters/GLTFTextureUtility.cpp
@@ -3,6 +3,24 @@
 #include "Converters/GLTFTextureUtility.h"
 #include "Engine/Texture2DArray.h"
 
+TextureFilter FGLTFTextureUtility::GetDefaultFilter(TextureGroup LODGroup)
+{
+	// TODO: should this be the running platform?
+	const ITargetPlatform* RunningPlatform = GetTargetPlatformManagerRef().GetRunningTargetPlatform();
+	const UTextureLODSettings& TextureLODSettings = RunningPlatform->GetTextureLODSettings();
+	const FTextureLODGroup& TextureLODGroup = TextureLODSettings.GetTextureLODGroup(LODGroup);
+
+	switch (TextureLODGroup.Filter)
+	{
+		case ETextureSamplerFilter::Point:             return TF_Nearest;
+		case ETextureSamplerFilter::Bilinear:          return TF_Bilinear;
+		case ETextureSamplerFilter::Trilinear:         return TF_Trilinear;
+		case ETextureSamplerFilter::AnisotropicPoint:  return TF_Nearest;
+		case ETextureSamplerFilter::AnisotropicLinear: return TF_Trilinear;
+		default:                                       return TF_Default; // fallback
+	}
+}
+
 TextureAddress FGLTFTextureUtility::GetAddressX(const UTexture* Texture)
 {
 	if (const UTexture2D* Texture2D = Cast<UTexture2D>(Texture))
@@ -52,8 +70,8 @@ UTexture2D* FGLTFTextureUtility::CreateTransientTexture(const void* RawData, int
 	check(CalculateImageBytes(Size.X, Size.Y, 0, Format) == ByteLength);
 
 	UTexture2D* Texture = UTexture2D::CreateTransient(Size.X, Size.Y, Format);
-	void* MipData = Texture->PlatformData->Mips[0].BulkData.Lock(LOCK_READ_WRITE);
 
+	void* MipData = Texture->PlatformData->Mips[0].BulkData.Lock(LOCK_READ_WRITE);
 	FMemory::Memcpy(MipData, RawData, ByteLength);
 	Texture->PlatformData->Mips[0].BulkData.Unlock();
 
diff --git a/Engine/Plugins/Enterprise/GLTFExporter/MigratedFromGit/Source/GLTFExporter/Private/Converters/GLTFTextureUtility.h b/Engine/Plugins/Enterprise/GLTFExporter/MigratedFromGit/Source/GLTFExporter/Private/Converters/GLTFTextureUtility.h
--- a/Engine/Plugins/Enterprise/GLTFExporter/MigratedFromGit/Source/GLTFExporter/Private/Converters/GLTFTextureUtility.h
+++ b/Engine/Plugins/Enterprise/GLTFExporter/MigratedFromGit/Source/GLTFExporter/Private/Converters/GLTFTextureUtility.h
@@ -6,6 +6,8 @@
 
 struct FGLTFTextureUtility
 {
+	static TextureFilter GetDefaultFilter(TextureGroup Group);
+
 	static TextureAddress GetAddressX(const UTexture* Texture);
 	static TextureAddress GetAddressY(const UTexture* Texture);
 
