[GLTFExporter] fix(texture): if texture is cubemap set sampler wrapping to clamp

commit d3bf2b208aad971946106f481ced0976f8f9f517
Date: Wed, 8 Jul 2020 10:53:25 +0200

    fix(texture): if texture is cubemap set sampler wrapping to clamp
    (cherry picked from commit 962dfc4fa3ef28baa0ac747aeacf209c905394e7​)

#skipci

diff --git a/Engine/Plugins/Enterprise/GLTFExporter/MigratedFromGit/Source/GLTFExporter/Private/Converters/GLTFTextureConverters.cpp b/Engine/Plugins/Enterprise/GLTFExporter/MigratedFromGit/Source/GLTFExporter/Private/Converters/GLTFTextureConverters.cpp
--- a/Engine/Plugins/Enterprise/GLTFExporter/MigratedFromGit/Source/GLTFExporter/Private/Converters/GLTFTextureConverters.cpp
+++ b/Engine/Plugins/Enterprise/GLTFExporter/MigratedFromGit/Source/GLTFExporter/Private/Converters/GLTFTextureConverters.cpp
@@ -24,13 +24,19 @@ FGLTFJsonSamplerIndex FGLTFTextureSamplerConverter::Add(FGLTFConvertBuilder& Bui
 		JsonSampler.MagFilter = FGLTFConverterUtility::ConvertMagFilter(Texture->Filter, Texture->LODGroup);
 	}
 
-	const TextureAddress AddressX = FGLTFTextureUtility::GetAddressX(Texture);
-	const TextureAddress AddressY = FGLTFTextureUtility::GetAddressY(Texture);
-
-	// TODO: report error if AddressX == TA_MAX || AddressY == TA_MAX
+	if (FGLTFTextureUtility::IsCubemap(Texture))
+	{
+		JsonSampler.WrapS = EGLTFJsonTextureWrap::ClampToEdge;
+		JsonSampler.WrapT = EGLTFJsonTextureWrap::ClampToEdge;
+	}
+	else
+	{
+		const TTuple<TextureAddress, TextureAddress> AddressXY = FGLTFTextureUtility::GetAddressXY(Texture);
+		// TODO: report error if AddressX == TA_MAX or AddressY == TA_MAX
 
-	JsonSampler.WrapS = FGLTFConverterUtility::ConvertWrap(AddressX);
-	JsonSampler.WrapT = FGLTFConverterUtility::ConvertWrap(AddressY);
+		JsonSampler.WrapS = FGLTFConverterUtility::ConvertWrap(AddressXY.Get<0>());
+		JsonSampler.WrapT = FGLTFConverterUtility::ConvertWrap(AddressXY.Get<1>());
+	}
 
 	return Builder.AddSampler(JsonSampler);
 }
diff --git a/Engine/Plugins/Enterprise/GLTFExporter/MigratedFromGit/Source/GLTFExporter/Private/Converters/GLTFTextureUtility.cpp b/Engine/Plugins/Enterprise/GLTFExporter/MigratedFromGit/Source/GLTFExporter/Private/Converters/GLTFTextureUtility.cpp
--- a/Engine/Plugins/Enterprise/GLTFExporter/MigratedFromGit/Source/GLTFExporter/Private/Converters/GLTFTextureUtility.cpp
+++ b/Engine/Plugins/Enterprise/GLTFExporter/MigratedFromGit/Source/GLTFExporter/Private/Converters/GLTFTextureUtility.cpp
@@ -1,7 +1,11 @@
 // Copyright Epic Games, Inc. All Rights Reserved.
 
 #include "Converters/GLTFTextureUtility.h"
-#include "Engine/Texture2DArray.h"
+
+bool FGLTFTextureUtility::IsCubemap(const UTexture* Texture)
+{
+	return Texture->IsA<UTextureCube>() || Texture->IsA<UTextureRenderTargetCube>();
+}
 
 TextureFilter FGLTFTextureUtility::GetDefaultFilter(TextureGroup LODGroup)
 {
@@ -21,48 +25,28 @@ TextureFilter FGLTFTextureUtility::GetDefaultFilter(TextureGroup LODGroup)
 	}
 }
 
-TextureAddress FGLTFTextureUtility::GetAddressX(const UTexture* Texture)
+TTuple<TextureAddress, TextureAddress> FGLTFTextureUtility::GetAddressXY(const UTexture* Texture)
 {
-	if (const UTexture2D* Texture2D = Cast<UTexture2D>(Texture))
-	{
-		return Texture2D->AddressX;
-	}
+	TextureAddress AddressX;
+	TextureAddress AddressY;
 
-	if (const UTexture2DArray* Texture2DArray = Cast<UTexture2DArray>(Texture))
-	{
-		return Texture2DArray->AddressX;
-	}
-
-	if (const UTextureRenderTarget2D* RenderTarget2D = Cast<UTextureRenderTarget2D>(Texture))
-	{
-		return RenderTarget2D->AddressX;
-	}
-
-	// TODO: add support for UMediaTexture?
-
-	return TextureAddress::TA_MAX;
-}
-
-TextureAddress FGLTFTextureUtility::GetAddressY(const UTexture* Texture)
-{
 	if (const UTexture2D* Texture2D = Cast<UTexture2D>(Texture))
 	{
-		return Texture2D->AddressY;
+		AddressX = Texture2D->AddressX;
+		AddressY = Texture2D->AddressX;
 	}
-
-	if (const UTexture2DArray* Texture2DArray = Cast<UTexture2DArray>(Texture))
+	else if (const UTextureRenderTarget2D* RenderTarget2D = Cast<UTextureRenderTarget2D>(Texture))
 	{
-		return Texture2DArray->AddressY;
+		AddressX = RenderTarget2D->AddressX;
+		AddressY = RenderTarget2D->AddressX;
 	}
-
-	if (const UTextureRenderTarget2D* RenderTarget2D = Cast<UTextureRenderTarget2D>(Texture))
+	else
 	{
-		return RenderTarget2D->AddressY;
+		AddressX = TA_MAX;
+		AddressY = TA_MAX;
 	}
 
-	// TODO: add support for UMediaTexture?
-
-	return TextureAddress::TA_MAX;
+	return MakeTuple(AddressX, AddressY);
 }
 
 UTexture2D* FGLTFTextureUtility::CreateTransientTexture(const void* RawData, int64 ByteLength, const FIntPoint& Size, EPixelFormat Format, bool bUseSRGB)
diff --git a/Engine/Plugins/Enterprise/GLTFExporter/MigratedFromGit/Source/GLTFExporter/Private/Converters/GLTFTextureUtility.h b/Engine/Plugins/Enterprise/GLTFExporter/MigratedFromGit/Source/GLTFExporter/Private/Converters/GLTFTextureUtility.h
--- a/Engine/Plugins/Enterprise/GLTFExporter/MigratedFromGit/Source/GLTFExporter/Private/Converters/GLTFTextureUtility.h
+++ b/Engine/Plugins/Enterprise/GLTFExporter/MigratedFromGit/Source/GLTFExporter/Private/Converters/GLTFTextureUtility.h
@@ -6,10 +6,11 @@
 
 struct FGLTFTextureUtility
 {
+	static bool IsCubemap(const UTexture* Texture);
+
 	static TextureFilter GetDefaultFilter(TextureGroup Group);
 
-	static TextureAddress GetAddressX(const UTexture* Texture);
-	static TextureAddress GetAddressY(const UTexture* Texture);
+	static TTuple<TextureAddress, TextureAddress> GetAddressXY(const UTexture* Texture);
 
 	static UTexture2D* CreateTransientTexture(const void* RawData, int64 ByteLength, const FIntPoint& Size, EPixelFormat Format, bool bUseSRGB = false);
 };
