[GLTFExporter] feat(texture): add IsHDRFormat and CanPNGCompressFormat to GLTFTextureUtility

commit 0b6a76981f6bf2e8430d7ffc563ce40b89cb6f92
Date: Wed, 8 Jul 2020 10:58:19 +0200

    feat(texture): add IsHDRFormat and CanPNGCompressFormat to GLTFTextureUtility
    (cherry picked from commit 36cb2d6ee89bc73eab821898d8d2e0c7d57077d7​)

#skipci

diff --git a/Engine/Plugins/Enterprise/GLTFExporter/MigratedFromGit/Source/GLTFExporter/Private/Converters/GLTFTextureUtility.cpp b/Engine/Plugins/Enterprise/GLTFExporter/MigratedFromGit/Source/GLTFExporter/Private/Converters/GLTFTextureUtility.cpp
--- a/Engine/Plugins/Enterprise/GLTFExporter/MigratedFromGit/Source/GLTFExporter/Private/Converters/GLTFTextureUtility.cpp
+++ b/Engine/Plugins/Enterprise/GLTFExporter/MigratedFromGit/Source/GLTFExporter/Private/Converters/GLTFTextureUtility.cpp
@@ -1,6 +1,41 @@
 // Copyright Epic Games, Inc. All Rights Reserved.
 
 #include "Converters/GLTFTextureUtility.h"
+#include "IImageWrapper.h"
+
+bool FGLTFTextureUtility::IsHDRFormat(EPixelFormat Format)
+{
+	return CalculateImageBytes(1, 1, 0, Format) > 4;
+}
+
+bool FGLTFTextureUtility::CanPNGCompressFormat(ETextureSourceFormat InFormat, ERGBFormat& OutFormat, uint32& OutBitDepth)
+{
+	switch (InFormat)
+	{
+		case TSF_BGRA8:   OutFormat = ERGBFormat::BGRA; OutBitDepth = 8;  return true;
+		case TSF_RGBA8:   OutFormat = ERGBFormat::RGBA; OutBitDepth = 8;  return true;
+		case TSF_RGBA16:  OutFormat = ERGBFormat::RGBA; OutBitDepth = 16; return true;
+		case TSF_G8:      OutFormat = ERGBFormat::Gray; OutBitDepth = 8;  return true;
+		case TSF_G16:     OutFormat = ERGBFormat::Gray; OutBitDepth = 16; return true;
+		default:
+            return false;
+	}
+}
+
+bool FGLTFTextureUtility::CanPNGCompressFormat(EPixelFormat InFormat, ERGBFormat& OutFormat, uint32& OutBitDepth)
+{
+	switch (InFormat)
+	{
+		case PF_B8G8R8A8:           OutFormat = ERGBFormat::BGRA; OutBitDepth = 8;  return true;
+		case PF_R8G8B8A8:           OutFormat = ERGBFormat::RGBA; OutBitDepth = 8;  return true;
+		case PF_R16G16B16A16_UNORM: OutFormat = ERGBFormat::RGBA; OutBitDepth = 16; return true;
+		case PF_L8:                 OutFormat = ERGBFormat::Gray; OutBitDepth = 8;  return true;
+		case PF_G8:                 OutFormat = ERGBFormat::Gray; OutBitDepth = 8;  return true;
+		case PF_G16:                OutFormat = ERGBFormat::Gray; OutBitDepth = 16; return true;
+		default:
+            return false;
+	}
+}
 
 bool FGLTFTextureUtility::IsCubemap(const UTexture* Texture)
 {
diff --git a/Engine/Plugins/Enterprise/GLTFExporter/MigratedFromGit/Source/GLTFExporter/Private/Converters/GLTFTextureUtility.h b/Engine/Plugins/Enterprise/GLTFExporter/MigratedFromGit/Source/GLTFExporter/Private/Converters/GLTFTextureUtility.h
--- a/Engine/Plugins/Enterprise/GLTFExporter/MigratedFromGit/Source/GLTFExporter/Private/Converters/GLTFTextureUtility.h
+++ b/Engine/Plugins/Enterprise/GLTFExporter/MigratedFromGit/Source/GLTFExporter/Private/Converters/GLTFTextureUtility.h
@@ -4,8 +4,15 @@
 
 #include "Engine.h"
 
+enum class ERGBFormat : int8;
+
 struct FGLTFTextureUtility
 {
+	static bool IsHDRFormat(EPixelFormat Format);
+
+	static bool CanPNGCompressFormat(ETextureSourceFormat InFormat, ERGBFormat& OutFormat, uint32& OutBitDepth);
+	static bool CanPNGCompressFormat(EPixelFormat InFormat, ERGBFormat& OutFormat, uint32& OutBitDepth);
+
 	static bool IsCubemap(const UTexture* Texture);
 
 	static TextureFilter GetDefaultFilter(TextureGroup Group);
