[GLTFExporter] feat(texture): add GLTFTextureUtility::GetBulkData as temporary workaround to access png-compressed texture source data

commit e8d9c2568be3770a3263222fe26e6325343821ef
Date: Wed, 8 Jul 2020 11:00:53 +0200

    feat(texture): add GLTFTextureUtility::GetBulkData as temporary workaround to access png-compressed texture source data
    (cherry picked from commit c2546d8ba15815ba13689e132d6d32446fb8ec9e​)

#skipci

diff --git a/Engine/Plugins/Enterprise/GLTFExporter/MigratedFromGit/Source/GLTFExporter/Private/Converters/GLTFTextureUtility.cpp b/Engine/Plugins/Enterprise/GLTFExporter/MigratedFromGit/Source/GLTFExporter/Private/Converters/GLTFTextureUtility.cpp
--- a/Engine/Plugins/Enterprise/GLTFExporter/MigratedFromGit/Source/GLTFExporter/Private/Converters/GLTFTextureUtility.cpp
+++ b/Engine/Plugins/Enterprise/GLTFExporter/MigratedFromGit/Source/GLTFExporter/Private/Converters/GLTFTextureUtility.cpp
@@ -84,6 +84,18 @@ TTuple<TextureAddress, TextureAddress> FGLTFTextureUtility::GetAddressXY(const U
 	return MakeTuple(AddressX, AddressY);
 }
 
+struct FTextureSourceHack
+{
+	FByteBulkData BulkData;
+};
+
+const FByteBulkData& FGLTFTextureUtility::GetBulkData(const FTextureSource& TextureSource)
+{
+	// TODO: replace hack with safe solution by adding proper API access
+	auto Hack = reinterpret_cast<const FTextureSourceHack*>(&TextureSource);
+	return Hack->BulkData;
+}
+
 UTexture2D* FGLTFTextureUtility::CreateTransientTexture(const void* RawData, int64 ByteLength, const FIntPoint& Size, EPixelFormat Format, bool bUseSRGB)
 {
 	check(CalculateImageBytes(Size.X, Size.Y, 0, Format) == ByteLength);
diff --git a/Engine/Plugins/Enterprise/GLTFExporter/MigratedFromGit/Source/GLTFExporter/Private/Converters/GLTFTextureUtility.h b/Engine/Plugins/Enterprise/GLTFExporter/MigratedFromGit/Source/GLTFExporter/Private/Converters/GLTFTextureUtility.h
--- a/Engine/Plugins/Enterprise/GLTFExporter/MigratedFromGit/Source/GLTFExporter/Private/Converters/GLTFTextureUtility.h
+++ b/Engine/Plugins/Enterprise/GLTFExporter/MigratedFromGit/Source/GLTFExporter/Private/Converters/GLTFTextureUtility.h
@@ -19,5 +19,7 @@ struct FGLTFTextureUtility
 
 	static TTuple<TextureAddress, TextureAddress> GetAddressXY(const UTexture* Texture);
 
+	static const FByteBulkData& GetBulkData(const FTextureSource& TextureSource);
+
 	static UTexture2D* CreateTransientTexture(const void* RawData, int64 ByteLength, const FIntPoint& Size, EPixelFormat Format, bool bUseSRGB = false);
 };
