[GLTFExporter] refactor(texture): improve texture conversion and compression

commit 9e8aa57d090ed07c2521c1a4b96e550b32a1cecb
Date: Wed, 8 Jul 2020 11:03:20 +0200

    refactor(texture): improve texture conversion and compression
    (cherry picked from commit a343c3ac393fc2aea7b72dc90f4d3a6a55b013fb​)

#skipci

diff --git a/Engine/Plugins/Enterprise/GLTFExporter/MigratedFromGit/Source/GLTFExporter/Private/Builders/GLTFBuilderUtility.cpp b/Engine/Plugins/Enterprise/GLTFExporter/MigratedFromGit/Source/GLTFExporter/Private/Builders/GLTFBuilderUtility.cpp
--- a/Engine/Plugins/Enterprise/GLTFExporter/MigratedFromGit/Source/GLTFExporter/Private/Builders/GLTFBuilderUtility.cpp
+++ b/Engine/Plugins/Enterprise/GLTFExporter/MigratedFromGit/Source/GLTFExporter/Private/Builders/GLTFBuilderUtility.cpp
@@ -1,30 +1,28 @@
 // Copyright Epic Games, Inc. All Rights Reserved.
 
 #include "Builders/GLTFBuilderUtility.h"
-#include "IImageWrapperModule.h"
-#include "IImageWrapper.h"
 
-bool FGLTFBuilderUtility::CompressImage(const void* RawData, int64 ByteLength, int32 InWidth, int32 InHeight, ERGBFormat InRawFormat, int32 InBitDepth, TArray64<uint8>& OutCompressedData, EImageFormat OutCompressionFormat, int32 OutCompressionQuality)
+FString FGLTFBuilderUtility::GetMeshName(const UStaticMesh* StaticMesh, int32 LODIndex)
 {
-	IImageWrapperModule& ImageWrapperModule = FModuleManager::Get().LoadModuleChecked<IImageWrapperModule>(TEXT("ImageWrapper"));
+	FString Name;
 
-	TSharedPtr<IImageWrapper> ImageWrapper = ImageWrapperModule.CreateImageWrapper(OutCompressionFormat);
-	if (!ImageWrapper.IsValid())
+	if (StaticMesh != nullptr)
 	{
-		// TODO: report error
-		return false;
+		StaticMesh->GetName(Name);
+		if (LODIndex != 0) Name += TEXT("_LOD") + FString::FromInt(LODIndex);
 	}
 
-	const bool bFormatSupported = ImageWrapper->SetRaw(RawData, ByteLength, InWidth, InHeight, InRawFormat, InBitDepth);
-	if (!bFormatSupported)
+	return Name;
+}
+
+const TCHAR* FGLTFBuilderUtility::GetFileExtension(EGLTFJsonMimeType MimeType)
+{
+	switch (MimeType)
 	{
-		// TODO: report error
-		return false;
+		case EGLTFJsonMimeType::PNG:  return TEXT(".png");
+		case EGLTFJsonMimeType::JPEG: return TEXT(".jpg");
+		default:                      return TEXT(".unknown");
 	}
-
-	const TArray64<uint8>& CompressedData = ImageWrapper->GetCompressed(OutCompressionQuality);
-	OutCompressedData.Append(CompressedData);
-	return true;
 }
 
 FString FGLTFBuilderUtility::GetUniqueFilename(const FString& BaseFilename, const FString& FileExtension, const TSet<FString>& UniqueFilenames)
diff --git a/Engine/Plugins/Enterprise/GLTFExporter/MigratedFromGit/Source/GLTFExporter/Private/Builders/GLTFBuilderUtility.h b/Engine/Plugins/Enterprise/GLTFExporter/MigratedFromGit/Source/GLTFExporter/Private/Builders/GLTFBuilderUtility.h
--- a/Engine/Plugins/Enterprise/GLTFExporter/MigratedFromGit/Source/GLTFExporter/Private/Builders/GLTFBuilderUtility.h
+++ b/Engine/Plugins/Enterprise/GLTFExporter/MigratedFromGit/Source/GLTFExporter/Private/Builders/GLTFBuilderUtility.h
@@ -2,38 +2,14 @@
 
 #pragma once
 
-#include "Engine.h"
+#include "CoreMinimal.h"
 #include "Json/GLTFJsonEnums.h"
 
-enum class EImageFormat : int8;
-enum class ERGBFormat : int8;
-
 struct FGLTFBuilderUtility
 {
-	static FString GetMeshName(const UStaticMesh* StaticMesh, int32 LODIndex)
-	{
-		FString Name;
-
-		if (StaticMesh != nullptr)
-		{
-			StaticMesh->GetName(Name);
-			if (LODIndex != 0) Name += TEXT("_LOD") + FString::FromInt(LODIndex);
-		}
-
-		return Name;
-	}
-
-	static bool CompressImage(const void* RawData, int64 ByteLength, int32 InWidth, int32 InHeight, ERGBFormat InRawFormat, int32 InBitDepth, TArray64<uint8>& OutCompressedData, EImageFormat OutCompressionFormat, int32 OutCompressionQuality);
+	static FString GetMeshName(const UStaticMesh* StaticMesh, int32 LODIndex);
 
-	static const TCHAR* GetFileExtension(EGLTFJsonMimeType MimeType)
-	{
-		switch (MimeType)
-		{
-			case EGLTFJsonMimeType::PNG:  return TEXT(".png");
-			case EGLTFJsonMimeType::JPEG: return TEXT(".jpg");
-			default:                      return TEXT(".unknown");
-		}
-	}
+	static const TCHAR* GetFileExtension(EGLTFJsonMimeType MimeType);
 
 	static FString GetUniqueFilename(const FString& BaseFilename, const FString& FileExtension, const TSet<FString>& UniqueFilenames);
 };
diff --git a/Engine/Plugins/Enterprise/GLTFExporter/MigratedFromGit/Source/GLTFExporter/Private/Builders/GLTFImageBuilder.cpp b/Engine/Plugins/Enterprise/GLTFExporter/MigratedFromGit/Source/GLTFExporter/Private/Builders/GLTFImageBuilder.cpp
--- a/Engine/Plugins/Enterprise/GLTFExporter/MigratedFromGit/Source/GLTFExporter/Private/Builders/GLTFImageBuilder.cpp
+++ b/Engine/Plugins/Enterprise/GLTFExporter/MigratedFromGit/Source/GLTFExporter/Private/Builders/GLTFImageBuilder.cpp
@@ -2,117 +2,54 @@
 
 #include "Builders/GLTFImageBuilder.h"
 #include "Builders/GLTFBuilderUtility.h"
+#include "IImageWrapperModule.h"
 #include "IImageWrapper.h"
-#include "Engine/Texture.h"
+#include "Misc/FileHelper.h"
 
 FGLTFImageBuilder::FGLTFImageBuilder(const UGLTFExportOptions* ExportOptions)
 	: FGLTFBufferBuilder(ExportOptions)
 {
 }
 
-FGLTFJsonImageIndex FGLTFImageBuilder::AddImage(const void* RawData, int64 ByteLength, int32 Width, int32 Height, ERGBFormat RawFormat, int32 BitDepth, bool bFloatFormat, const FString& Name, EGLTFJsonMimeType MimeType, int32 Quality)
+FGLTFJsonImageIndex FGLTFImageBuilder::AddImage(const void* CompressedData, int64 CompressedByteLength, EGLTFJsonMimeType MimeType, const FString& Name)
 {
-	TArray64<uint8> ImageData;
-
-	if (bFloatFormat)
-	{
-		// TODO: implement support
-		return FGLTFJsonImageIndex(INDEX_NONE);
-	}
-	else
-	{
-		switch (MimeType)
-		{
-		case EGLTFJsonMimeType::PNG:
-			if (!FGLTFBuilderUtility::CompressImage(RawData, ByteLength, Width, Height, RawFormat, BitDepth, ImageData, EImageFormat::PNG, Quality))
-			{
-				// TODO: report error
-				return FGLTFJsonImageIndex(INDEX_NONE);
-			}
-			break;
-		case EGLTFJsonMimeType::JPEG:
-			if (!FGLTFBuilderUtility::CompressImage(RawData, ByteLength, Width, Height, RawFormat, BitDepth, ImageData, EImageFormat::JPEG, Quality))
-			{
-				// TODO: report error
-				return FGLTFJsonImageIndex(INDEX_NONE);
-			}
-			break;
-		default:
-			// TODO: report error
-			return FGLTFJsonImageIndex(INDEX_NONE);
-		}
-	}
-
 	FGLTFJsonImage Image;
 	Image.Name = Name;
 	Image.MimeType = MimeType;
-
 	const FGLTFJsonImageIndex ImageIndex = FGLTFJsonBuilder::AddImage(Image);
-	ImageDataLookup.Add(ImageIndex, ImageData);
+
+	TArray64<uint8>& ImageData = ImageDataLookup.Add(ImageIndex);
+	ImageData.Append(static_cast<const uint8*>(CompressedData), CompressedByteLength);
 	return ImageIndex;
 }
 
-FGLTFJsonImageIndex FGLTFImageBuilder::AddImage(const void* RawData, int64 ByteLength, int32 Width, int32 Height, EPixelFormat PixelFormat, const FString& Name, EGLTFJsonMimeType MimeType, int32 Quality)
+FGLTFJsonImageIndex FGLTFImageBuilder::AddImage(const void* RawData, int64 ByteLength, FIntPoint Size, ERGBFormat Format, int32 BitDepth, const FString& Name)
 {
-	ERGBFormat RawFormat;
-	int32 BitDepth;
-	bool bFloatFormat;
+	IImageWrapperModule& ImageWrapperModule = FModuleManager::Get().LoadModuleChecked<IImageWrapperModule>(TEXT("ImageWrapper"));
+	TSharedPtr<IImageWrapper> ImageWrapper = ImageWrapperModule.CreateImageWrapper(EImageFormat::PNG);
 
-	switch (PixelFormat)
+	if (!ImageWrapper.IsValid())
 	{
-		case PF_B8G8R8A8:           RawFormat = ERGBFormat::BGRA; BitDepth = 8;  bFloatFormat = false; break;
-		case PF_R8G8B8A8:           RawFormat = ERGBFormat::RGBA; BitDepth = 8;  bFloatFormat = false; break;
-		case PF_R16G16B16A16_UNORM: RawFormat = ERGBFormat::RGBA; BitDepth = 16; bFloatFormat = false; break;
-		case PF_FloatRGBA:          RawFormat = ERGBFormat::RGBA; BitDepth = 16; bFloatFormat = true;  break;
-		case PF_L8:                 RawFormat = ERGBFormat::Gray; BitDepth = 8;  bFloatFormat = false; break;
-		case PF_G8:                 RawFormat = ERGBFormat::Gray; BitDepth = 8;  bFloatFormat = false; break;
-		case PF_G16:                RawFormat = ERGBFormat::Gray; BitDepth = 16; bFloatFormat = false; break;
-		case PF_R16F:               RawFormat = ERGBFormat::Gray; BitDepth = 16; bFloatFormat = true;  break;
-		case PF_R16F_FILTER:        RawFormat = ERGBFormat::Gray; BitDepth = 16; bFloatFormat = true;  break;
-		case PF_R32_FLOAT:          RawFormat = ERGBFormat::Gray; BitDepth = 32; bFloatFormat = true;  break;
-		default:
-			// TODO: report error
-			return FGLTFJsonImageIndex(INDEX_NONE);
+		// TODO: report error
+		return FGLTFJsonImageIndex(INDEX_NONE);
 	}
 
-	return AddImage(RawData, ByteLength, Width, Height, RawFormat, BitDepth, bFloatFormat, Name, MimeType, Quality);
-}
+	const bool bFormatSupported = ImageWrapper->SetRaw(RawData, ByteLength, Size.X, Size.Y, Format, BitDepth);
 
-FGLTFJsonImageIndex FGLTFImageBuilder::AddImage(const FTextureSource& Image, const FString& Name, EGLTFJsonMimeType MimeType, int32 Quality)
-{
-	// TODO: are these always zero?
-	const int32 BlockIndex = 0;
-	const int32 LayerIndex = 0;
-	const int32 MipIndex = 0;
-
-	ERGBFormat RawFormat;
-	int32 BitDepth;
-	bool bFloatFormat;
-
-	const ETextureSourceFormat SourceFormat = Image.GetFormat(LayerIndex);
-	switch (SourceFormat)
-	{
-		case TSF_BGRA8:   RawFormat = ERGBFormat::BGRA; BitDepth = 8;  bFloatFormat = false; break;
-		case TSF_RGBA8:   RawFormat = ERGBFormat::RGBA; BitDepth = 8;  bFloatFormat = false; break;
-		case TSF_BGRE8:   RawFormat = ERGBFormat::BGRA; BitDepth = 8;  bFloatFormat = false; break;
-		case TSF_RGBE8:   RawFormat = ERGBFormat::RGBA; BitDepth = 8;  bFloatFormat = false; break;
-		case TSF_RGBA16:  RawFormat = ERGBFormat::RGBA; BitDepth = 16; bFloatFormat = false; break;
-		case TSF_RGBA16F: RawFormat = ERGBFormat::RGBA; BitDepth = 16; bFloatFormat = true;  break;
-		case TSF_G8:      RawFormat = ERGBFormat::Gray; BitDepth = 8;  bFloatFormat = false; break;
-		case TSF_G16:     RawFormat = ERGBFormat::Gray; BitDepth = 16; bFloatFormat = false; break;
-		default:
-			// TODO: report error
-			return FGLTFJsonImageIndex(INDEX_NONE);
-	}
-
-	TArray64<uint8> RawData;
-	if (!const_cast<FTextureSource&>(Image).GetMipData(RawData, BlockIndex, LayerIndex, MipIndex))
+	if (!bFormatSupported)
 	{
 		// TODO: report error
 		return FGLTFJsonImageIndex(INDEX_NONE);
 	}
 
-	return AddImage(RawData.GetData(), RawData.Num(), Image.GetSizeX(), Image.GetSizeY(), RawFormat, BitDepth, bFloatFormat, Name, MimeType, Quality);
+	const TArray64<uint8>& ImageData = ImageWrapper->GetCompressed();
+	return AddImage(ImageData.GetData(), ImageData.Num(), EGLTFJsonMimeType::PNG, Name);
+}
+
+FGLTFJsonImageIndex FGLTFImageBuilder::AddImage(const FColor* Pixels, FIntPoint Size, const FString& Name)
+{
+	const int64 ByteLength = Size.X * Size.Y * sizeof(FColor);
+	return AddImage(Pixels, ByteLength, Size, ERGBFormat::BGRA, 8, Name);
 }
 
 bool FGLTFImageBuilder::Serialize(FArchive& Archive, const FString& FilePath)
diff --git a/Engine/Plugins/Enterprise/GLTFExporter/MigratedFromGit/Source/GLTFExporter/Private/Builders/GLTFImageBuilder.h b/Engine/Plugins/Enterprise/GLTFExporter/MigratedFromGit/Source/GLTFExporter/Private/Builders/GLTFImageBuilder.h
--- a/Engine/Plugins/Enterprise/GLTFExporter/MigratedFromGit/Source/GLTFExporter/Private/Builders/GLTFImageBuilder.h
+++ b/Engine/Plugins/Enterprise/GLTFExporter/MigratedFromGit/Source/GLTFExporter/Private/Builders/GLTFImageBuilder.h
@@ -12,11 +12,10 @@ public:
 
 	FGLTFImageBuilder(const UGLTFExportOptions* ExportOptions);
 
-	FGLTFJsonImageIndex AddImage(const void* RawData, int64 ByteLength, int32 Width, int32 Height, ERGBFormat RawFormat, int32 BitDepth, bool bFloatFormat, const FString& Name = TEXT(""), EGLTFJsonMimeType MimeType = EGLTFJsonMimeType::PNG, int32 Quality = 100);
+	FGLTFJsonImageIndex AddImage(const void* CompressedData, int64 CompressedByteLength, EGLTFJsonMimeType MimeType, const FString& Name);
+	FGLTFJsonImageIndex AddImage(const void* RawData, int64 ByteLength, FIntPoint Size, ERGBFormat Format, int32 BitDepth, const FString& Name);
 
-	FGLTFJsonImageIndex AddImage(const void* RawData, int64 ByteLength, int32 Width, int32 Height, EPixelFormat PixelFormat = EPixelFormat::PF_B8G8R8A8, const FString& Name = TEXT(""), EGLTFJsonMimeType MimeType = EGLTFJsonMimeType::PNG, int32 Quality = 100);
-
-	FGLTFJsonImageIndex AddImage(const FTextureSource& Image, const FString& Name = TEXT(""), EGLTFJsonMimeType MimeType = EGLTFJsonMimeType::PNG, int32 Quality = 100);
+	FGLTFJsonImageIndex AddImage(const FColor* Pixels, FIntPoint Size, const FString& Name);
 
 	virtual bool Serialize(FArchive& Archive, const FString& FilePath) override;
 
diff --git a/Engine/Plugins/Enterprise/GLTFExporter/MigratedFromGit/Source/GLTFExporter/Private/Converters/GLTFMaterialUtility.cpp b/Engine/Plugins/Enterprise/GLTFExporter/MigratedFromGit/Source/GLTFExporter/Private/Converters/GLTFMaterialUtility.cpp
--- a/Engine/Plugins/Enterprise/GLTFExporter/MigratedFromGit/Source/GLTFExporter/Private/Converters/GLTFMaterialUtility.cpp
+++ b/Engine/Plugins/Enterprise/GLTFExporter/MigratedFromGit/Source/GLTFExporter/Private/Converters/GLTFMaterialUtility.cpp
@@ -197,7 +197,7 @@ FGLTFJsonTextureIndex FGLTFMaterialUtility::AddTexture(FGLTFConvertBuilder& Buil
 	FGLTFJsonTexture JsonTexture;
 	JsonTexture.Name = TextureName;
 	JsonTexture.Sampler = Builder.AddSampler(JsonSampler);
-	JsonTexture.Source = Builder.AddImage(Pixels.GetData(), Pixels.Num() * sizeof(FColor), TextureSize.X, TextureSize.Y, PixelFormat, TextureName);
+	JsonTexture.Source = Builder.AddImage(Pixels.GetData(), TextureSize, TextureName);
 
 	return Builder.AddTexture(JsonTexture);
 }
diff --git a/Engine/Plugins/Enterprise/GLTFExporter/MigratedFromGit/Source/GLTFExporter/Private/Converters/GLTFTextureConverters.cpp b/Engine/Plugins/Enterprise/GLTFExporter/MigratedFromGit/Source/GLTFExporter/Private/Converters/GLTFTextureConverters.cpp
--- a/Engine/Plugins/Enterprise/GLTFExporter/MigratedFromGit/Source/GLTFExporter/Private/Converters/GLTFTextureConverters.cpp
+++ b/Engine/Plugins/Enterprise/GLTFExporter/MigratedFromGit/Source/GLTFExporter/Private/Converters/GLTFTextureConverters.cpp
@@ -43,12 +43,30 @@ FGLTFJsonSamplerIndex FGLTFTextureSamplerConverter::Add(FGLTFConvertBuilder& Bui
 
 FGLTFJsonTextureIndex FGLTFTexture2DConverter::Add(FGLTFConvertBuilder& Builder, const FString& Name, const UTexture2D* Texture2D)
 {
-	// TODO: add RGBE encoding information for TSF_RGBE8 and TSF_BGRE8
+	FGLTFJsonImageIndex ImageIndex;
 
 	FGLTFJsonTexture JsonTexture;
 	JsonTexture.Name = Name;
+
+	if (Texture2D->Source.IsPNGCompressed())
+	{
+		// TODO: add support for non-compressed formats?
+		// TODO: add RGBE encoding information for TSF_RGBE8 and TSF_BGRE8
+
+		const FByteBulkData& BulkData = FGLTFTextureUtility::GetBulkData(Texture2D->Source);
+
+		const void* RawData = BulkData.LockReadOnly();
+		ImageIndex = Builder.AddImage(RawData, BulkData.GetBulkDataSize(), EGLTFJsonMimeType::PNG, JsonTexture.Name);
+		BulkData.Unlock();
+	}
+	else
+	{
+		// TODO: implement support
+		return FGLTFJsonTextureIndex(INDEX_NONE);
+	}
+
+	JsonTexture.Source = ImageIndex;
 	JsonTexture.Sampler = Builder.GetOrAddSampler(Texture2D);
-	JsonTexture.Source = Builder.AddImage(Texture2D->Source, Name);
 
 	return Builder.AddTexture(JsonTexture);
 }
