[GLTFExporter] feat(texture): add support for more texture formats that can be png-compressed

commit 3205af051893fd867892337a868b0fdf98624af4
Date: Wed, 8 Jul 2020 11:08:21 +0200

    feat(texture): add support for more texture formats that can be png-compressed
    (cherry picked from commit 1c91a268a55782f8feaceeda03843fdebfcb725d​)

#skipci

diff --git a/Engine/Plugins/Enterprise/GLTFExporter/MigratedFromGit/Source/GLTFExporter/Private/Builders/GLTFImageBuilder.cpp b/Engine/Plugins/Enterprise/GLTFExporter/MigratedFromGit/Source/GLTFExporter/Private/Builders/GLTFImageBuilder.cpp
--- a/Engine/Plugins/Enterprise/GLTFExporter/MigratedFromGit/Source/GLTFExporter/Private/Builders/GLTFImageBuilder.cpp
+++ b/Engine/Plugins/Enterprise/GLTFExporter/MigratedFromGit/Source/GLTFExporter/Private/Builders/GLTFImageBuilder.cpp
@@ -46,6 +46,22 @@ FGLTFJsonImageIndex FGLTFImageBuilder::AddImage(const void* RawData, int64 ByteL
 	return AddImage(ImageData.GetData(), ImageData.Num(), EGLTFJsonMimeType::PNG, Name);
 }
 
+FGLTFJsonImageIndex FGLTFImageBuilder::AddImage(const FByteBulkData& BulkData, EGLTFJsonMimeType MimeType, const FString& Name)
+{
+	const void* RawData = BulkData.LockReadOnly();
+	const FGLTFJsonImageIndex ImageIndex = AddImage(RawData, BulkData.GetBulkDataSize(), MimeType, Name);
+	BulkData.Unlock();
+	return ImageIndex;
+}
+
+FGLTFJsonImageIndex FGLTFImageBuilder::AddImage(const FByteBulkData& BulkData, FIntPoint Size, ERGBFormat Format, int32 BitDepth, const FString& Name)
+{
+	const void* RawData = BulkData.LockReadOnly();
+	const FGLTFJsonImageIndex ImageIndex = AddImage(RawData, BulkData.GetBulkDataSize(), Size, Format, BitDepth, Name);
+	BulkData.Unlock();
+	return ImageIndex;
+}
+
 FGLTFJsonImageIndex FGLTFImageBuilder::AddImage(const FColor* Pixels, FIntPoint Size, const FString& Name)
 {
 	const int64 ByteLength = Size.X * Size.Y * sizeof(FColor);
diff --git a/Engine/Plugins/Enterprise/GLTFExporter/MigratedFromGit/Source/GLTFExporter/Private/Builders/GLTFImageBuilder.h b/Engine/Plugins/Enterprise/GLTFExporter/MigratedFromGit/Source/GLTFExporter/Private/Builders/GLTFImageBuilder.h
--- a/Engine/Plugins/Enterprise/GLTFExporter/MigratedFromGit/Source/GLTFExporter/Private/Builders/GLTFImageBuilder.h
+++ b/Engine/Plugins/Enterprise/GLTFExporter/MigratedFromGit/Source/GLTFExporter/Private/Builders/GLTFImageBuilder.h
@@ -15,6 +15,9 @@ public:
 	FGLTFJsonImageIndex AddImage(const void* CompressedData, int64 CompressedByteLength, EGLTFJsonMimeType MimeType, const FString& Name);
 	FGLTFJsonImageIndex AddImage(const void* RawData, int64 ByteLength, FIntPoint Size, ERGBFormat Format, int32 BitDepth, const FString& Name);
 
+	FGLTFJsonImageIndex AddImage(const FByteBulkData& BulkData, EGLTFJsonMimeType MimeType, const FString& Name);
+	FGLTFJsonImageIndex AddImage(const FByteBulkData& BulkData, FIntPoint Size, ERGBFormat Format, int32 BitDepth, const FString& Name);
+
 	FGLTFJsonImageIndex AddImage(const FColor* Pixels, FIntPoint Size, const FString& Name);
 
 	virtual bool Serialize(FArchive& Archive, const FString& FilePath) override;
diff --git a/Engine/Plugins/Enterprise/GLTFExporter/MigratedFromGit/Source/GLTFExporter/Private/Converters/GLTFTextureConverters.cpp b/Engine/Plugins/Enterprise/GLTFExporter/MigratedFromGit/Source/GLTFExporter/Private/Converters/GLTFTextureConverters.cpp
--- a/Engine/Plugins/Enterprise/GLTFExporter/MigratedFromGit/Source/GLTFExporter/Private/Converters/GLTFTextureConverters.cpp
+++ b/Engine/Plugins/Enterprise/GLTFExporter/MigratedFromGit/Source/GLTFExporter/Private/Converters/GLTFTextureConverters.cpp
@@ -48,16 +48,26 @@ FGLTFJsonTextureIndex FGLTFTexture2DConverter::Add(FGLTFConvertBuilder& Builder,
 	FGLTFJsonTexture JsonTexture;
 	JsonTexture.Name = Name;
 
+	FIntPoint Size = { Texture2D->GetSizeX(), Texture2D->GetSizeY() };
+	ERGBFormat RGBFormat;
+	uint32 BitDepth;
+
 	if (Texture2D->Source.IsPNGCompressed())
 	{
-		// TODO: add support for non-compressed formats?
-		// TODO: add RGBE encoding information for TSF_RGBE8 and TSF_BGRE8
-
 		const FByteBulkData& BulkData = FGLTFTextureUtility::GetBulkData(Texture2D->Source);
+		ImageIndex = Builder.AddImage(BulkData, EGLTFJsonMimeType::PNG, JsonTexture.Name);
+	}
+	else if (FGLTFTextureUtility::CanPNGCompressFormat(Texture2D->Source.GetFormat(), RGBFormat, BitDepth))
+	{
+		// TODO: add support for RGBE encodings TSF_RGBE8 and TSF_BGRE8
 
-		const void* RawData = BulkData.LockReadOnly();
-		ImageIndex = Builder.AddImage(RawData, BulkData.GetBulkDataSize(), EGLTFJsonMimeType::PNG, JsonTexture.Name);
-		BulkData.Unlock();
+		const FByteBulkData& BulkData = FGLTFTextureUtility::GetBulkData(Texture2D->Source);
+		ImageIndex = Builder.AddImage(BulkData, Size, RGBFormat, BitDepth, JsonTexture.Name);
+	}
+	else if (FGLTFTextureUtility::CanPNGCompressFormat(Texture2D->GetPixelFormat(), RGBFormat, BitDepth))
+	{
+		const FByteBulkData& BulkData = Texture2D->PlatformData->Mips[0].BulkData;
+		ImageIndex = Builder.AddImage(BulkData, Size, RGBFormat, BitDepth, JsonTexture.Name);
 	}
 	else
 	{
