[GLTFExporter] feat(texture): add CreateRenderTarget and DrawTexture to GLTFTextureUtility

commit cc062264d50f14df661a4dee69be15c9c34c59e4
Date: Wed, 8 Jul 2020 11:31:31 +0200

    feat(texture): add CreateRenderTarget and DrawTexture to GLTFTextureUtility
    (cherry picked from commit e0ed3d72e49e107e43c995927ef4258a472d698c​)

#skipci

diff --git a/Engine/Plugins/Enterprise/GLTFExporter/MigratedFromGit/Source/GLTFExporter/Private/Converters/GLTFTextureUtility.cpp b/Engine/Plugins/Enterprise/GLTFExporter/MigratedFromGit/Source/GLTFExporter/Private/Converters/GLTFTextureUtility.cpp
--- a/Engine/Plugins/Enterprise/GLTFExporter/MigratedFromGit/Source/GLTFExporter/Private/Converters/GLTFTextureUtility.cpp
+++ b/Engine/Plugins/Enterprise/GLTFExporter/MigratedFromGit/Source/GLTFExporter/Private/Converters/GLTFTextureUtility.cpp
@@ -113,3 +113,31 @@ UTexture2D* FGLTFTextureUtility::CreateTransientTexture(const void* RawData, int
 	Texture->UpdateResource();
 	return Texture;
 }
+
+UTextureRenderTarget2D* FGLTFTextureUtility::CreateRenderTarget(const FIntPoint& Size, EPixelFormat Format, bool bInForceLinearGamma)
+{
+	UTextureRenderTarget2D* RenderTarget = NewObject<UTextureRenderTarget2D>();
+	RenderTarget->InitCustomFormat(Size.X, Size.Y, Format, bInForceLinearGamma);
+	return RenderTarget;
+}
+
+bool FGLTFTextureUtility::DrawTexture(UTextureRenderTarget2D* OutTarget, const UTexture2D* InSource)
+{
+	FRenderTarget* RenderTarget = OutTarget->GameThread_GetRenderTargetResource();
+	if (RenderTarget == nullptr)
+	{
+		return false;
+	}
+
+	FCanvas Canvas(RenderTarget, nullptr, 0.0f, 0.0f, 0.0f, GMaxRHIFeatureLevel);
+	FCanvasTileItem TileItem(FVector2D::ZeroVector, InSource->Resource, FLinearColor::White);
+
+	TileItem.Draw(&Canvas);
+
+	Canvas.Flush_GameThread();
+	FlushRenderingCommands();
+	Canvas.SetRenderTarget_GameThread(nullptr);
+	FlushRenderingCommands();
+
+	return true;
+}
diff --git a/Engine/Plugins/Enterprise/GLTFExporter/MigratedFromGit/Source/GLTFExporter/Private/Converters/GLTFTextureUtility.h b/Engine/Plugins/Enterprise/GLTFExporter/MigratedFromGit/Source/GLTFExporter/Private/Converters/GLTFTextureUtility.h
--- a/Engine/Plugins/Enterprise/GLTFExporter/MigratedFromGit/Source/GLTFExporter/Private/Converters/GLTFTextureUtility.h
+++ b/Engine/Plugins/Enterprise/GLTFExporter/MigratedFromGit/Source/GLTFExporter/Private/Converters/GLTFTextureUtility.h
@@ -22,4 +22,8 @@ struct FGLTFTextureUtility
 	static const FByteBulkData& GetBulkData(const FTextureSource& TextureSource);
 
 	static UTexture2D* CreateTransientTexture(const void* RawData, int64 ByteLength, const FIntPoint& Size, EPixelFormat Format, bool bUseSRGB = false);
+
+	static UTextureRenderTarget2D* CreateRenderTarget(const FIntPoint& Size, EPixelFormat Format, bool bInForceLinearGamma = false);
+
+	static bool DrawTexture(UTextureRenderTarget2D* OutTarget, const UTexture2D* InSource);
 };
