[GLTFExporter] feat(texture): add support for textures with non-png-compressable pixel formats

commit 0e764e32bad2fcf09339e1d64dee7d20c6aaf617
Date: Wed, 8 Jul 2020 11:33:05 +0200

    feat(texture): add support for textures with non-png-compressable pixel formats
    (cherry picked from commit 4b3672aef5c57f8d54b0fd79e09a1709bc0e8ce7​)

#skipci

diff --git a/Engine/Plugins/Enterprise/GLTFExporter/MigratedFromGit/Source/GLTFExporter/Private/Converters/GLTFTextureConverters.cpp b/Engine/Plugins/Enterprise/GLTFExporter/MigratedFromGit/Source/GLTFExporter/Private/Converters/GLTFTextureConverters.cpp
--- a/Engine/Plugins/Enterprise/GLTFExporter/MigratedFromGit/Source/GLTFExporter/Private/Converters/GLTFTextureConverters.cpp
+++ b/Engine/Plugins/Enterprise/GLTFExporter/MigratedFromGit/Source/GLTFExporter/Private/Converters/GLTFTextureConverters.cpp
@@ -71,8 +71,25 @@ FGLTFJsonTextureIndex FGLTFTexture2DConverter::Add(FGLTFConvertBuilder& Builder,
 	}
 	else
 	{
-		// TODO: implement support
-		return FGLTFJsonTextureIndex(INDEX_NONE);
+		UTextureRenderTarget2D* RenderTarget = FGLTFTextureUtility::CreateRenderTarget(Size, PF_B8G8R8A8, true);
+
+		FGLTFTextureUtility::DrawTexture(RenderTarget, Texture2D);
+
+		FTextureRenderTarget2DResource* Resource = static_cast<FTextureRenderTarget2DResource*>(RenderTarget->Resource);
+		if (Resource == nullptr)
+		{
+			// TODO: report error
+			return FGLTFJsonTextureIndex(INDEX_NONE);
+		}
+
+		TArray<FColor> Pixels;
+		if (!Resource->ReadPixels(Pixels))
+		{
+			// TODO: report error
+			return FGLTFJsonTextureIndex(INDEX_NONE);
+		}
+
+		ImageIndex = Builder.AddImage(Pixels.GetData(), Size, JsonTexture.Name);
 	}
 
 	JsonTexture.Source = ImageIndex;
