[GLTFExporter] fix(texture): remove ideal case when texture is already png-compressed internally, because BGRA seems to be internally stored as RGBA

commit eb4803a90f7e0fd6c9e0cee3c06922f35dda7fef
Date: Wed, 8 Jul 2020 12:12:39 +0200

    fix(texture): remove ideal case when texture is already png-compressed internally, because BGRA seems to be internally stored as RGBA
    (cherry picked from commit 838f8cfbd6503323d46ae4284e4407eb2eb2a476​)

#skipci

diff --git a/Engine/Plugins/Enterprise/GLTFExporter/MigratedFromGit/Source/GLTFExporter/Private/Builders/GLTFImageBuilder.cpp b/Engine/Plugins/Enterprise/GLTFExporter/MigratedFromGit/Source/GLTFExporter/Private/Builders/GLTFImageBuilder.cpp
--- a/Engine/Plugins/Enterprise/GLTFExporter/MigratedFromGit/Source/GLTFExporter/Private/Builders/GLTFImageBuilder.cpp
+++ b/Engine/Plugins/Enterprise/GLTFExporter/MigratedFromGit/Source/GLTFExporter/Private/Builders/GLTFImageBuilder.cpp
@@ -46,22 +46,6 @@ FGLTFJsonImageIndex FGLTFImageBuilder::AddImage(const void* RawData, int64 ByteL
 	return AddImage(ImageData.GetData(), ImageData.Num(), EGLTFJsonMimeType::PNG, Name);
 }
 
-FGLTFJsonImageIndex FGLTFImageBuilder::AddImage(const FByteBulkData& BulkData, EGLTFJsonMimeType MimeType, const FString& Name)
-{
-	const void* RawData = BulkData.LockReadOnly();
-	const FGLTFJsonImageIndex ImageIndex = AddImage(RawData, BulkData.GetBulkDataSize(), MimeType, Name);
-	BulkData.Unlock();
-	return ImageIndex;
-}
-
-FGLTFJsonImageIndex FGLTFImageBuilder::AddImage(const FByteBulkData& BulkData, FIntPoint Size, ERGBFormat Format, int32 BitDepth, const FString& Name)
-{
-	const void* RawData = BulkData.LockReadOnly();
-	const FGLTFJsonImageIndex ImageIndex = AddImage(RawData, BulkData.GetBulkDataSize(), Size, Format, BitDepth, Name);
-	BulkData.Unlock();
-	return ImageIndex;
-}
-
 FGLTFJsonImageIndex FGLTFImageBuilder::AddImage(const FColor* Pixels, FIntPoint Size, const FString& Name)
 {
 	const int64 ByteLength = Size.X * Size.Y * sizeof(FColor);
diff --git a/Engine/Plugins/Enterprise/GLTFExporter/MigratedFromGit/Source/GLTFExporter/Private/Builders/GLTFImageBuilder.h b/Engine/Plugins/Enterprise/GLTFExporter/MigratedFromGit/Source/GLTFExporter/Private/Builders/GLTFImageBuilder.h
--- a/Engine/Plugins/Enterprise/GLTFExporter/MigratedFromGit/Source/GLTFExporter/Private/Builders/GLTFImageBuilder.h
+++ b/Engine/Plugins/Enterprise/GLTFExporter/MigratedFromGit/Source/GLTFExporter/Private/Builders/GLTFImageBuilder.h
@@ -15,9 +15,6 @@ public:
 	FGLTFJsonImageIndex AddImage(const void* CompressedData, int64 CompressedByteLength, EGLTFJsonMimeType MimeType, const FString& Name);
 	FGLTFJsonImageIndex AddImage(const void* RawData, int64 ByteLength, FIntPoint Size, ERGBFormat Format, int32 BitDepth, const FString& Name);
 
-	FGLTFJsonImageIndex AddImage(const FByteBulkData& BulkData, EGLTFJsonMimeType MimeType, const FString& Name);
-	FGLTFJsonImageIndex AddImage(const FByteBulkData& BulkData, FIntPoint Size, ERGBFormat Format, int32 BitDepth, const FString& Name);
-
 	FGLTFJsonImageIndex AddImage(const FColor* Pixels, FIntPoint Size, const FString& Name);
 
 	virtual bool Serialize(FArchive& Archive, const FString& FilePath) override;
diff --git a/Engine/Plugins/Enterprise/GLTFExporter/MigratedFromGit/Source/GLTFExporter/Private/Converters/GLTFTextureConverters.cpp b/Engine/Plugins/Enterprise/GLTFExporter/MigratedFromGit/Source/GLTFExporter/Private/Converters/GLTFTextureConverters.cpp
--- a/Engine/Plugins/Enterprise/GLTFExporter/MigratedFromGit/Source/GLTFExporter/Private/Converters/GLTFTextureConverters.cpp
+++ b/Engine/Plugins/Enterprise/GLTFExporter/MigratedFromGit/Source/GLTFExporter/Private/Converters/GLTFTextureConverters.cpp
@@ -52,22 +52,21 @@ FGLTFJsonTextureIndex FGLTFTexture2DConverter::Add(FGLTFConvertBuilder& Builder,
 	ERGBFormat RGBFormat;
 	uint32 BitDepth;
 
-	if (Texture2D->Source.IsPNGCompressed())
+	if (FGLTFTextureUtility::CanPNGCompressFormat(Texture2D->Source.GetFormat(), RGBFormat, BitDepth))
 	{
-		const FByteBulkData& BulkData = FGLTFTextureUtility::GetBulkData(Texture2D->Source);
-		ImageIndex = Builder.AddImage(BulkData, EGLTFJsonMimeType::PNG, JsonTexture.Name);
-	}
-	else if (FGLTFTextureUtility::CanPNGCompressFormat(Texture2D->Source.GetFormat(), RGBFormat, BitDepth))
-	{
-		// TODO: add support for RGBE encodings TSF_RGBE8 and TSF_BGRE8
+		FTextureSource& Source = const_cast<FTextureSource&>(Texture2D->Source);
 
-		const FByteBulkData& BulkData = FGLTFTextureUtility::GetBulkData(Texture2D->Source);
-		ImageIndex = Builder.AddImage(BulkData, Size, RGBFormat, BitDepth, JsonTexture.Name);
+		const void* RawData = Source.LockMip(0);
+		ImageIndex = Builder.AddImage(RawData, Source.CalcMipSize(0), Size, RGBFormat, BitDepth, JsonTexture.Name);
+		Source.UnlockMip(0);
 	}
 	else if (FGLTFTextureUtility::CanPNGCompressFormat(Texture2D->GetPixelFormat(), RGBFormat, BitDepth))
 	{
 		const FByteBulkData& BulkData = Texture2D->PlatformData->Mips[0].BulkData;
-		ImageIndex = Builder.AddImage(BulkData, Size, RGBFormat, BitDepth, JsonTexture.Name);
+
+		const void* RawData = BulkData.LockReadOnly();
+		ImageIndex = Builder.AddImage(RawData, BulkData.GetBulkDataSize(), Size, RGBFormat, BitDepth, JsonTexture.Name);
+		BulkData.Unlock();
 	}
 	else
 	{
diff --git a/Engine/Plugins/Enterprise/GLTFExporter/MigratedFromGit/Source/GLTFExporter/Private/Converters/GLTFTextureUtility.cpp b/Engine/Plugins/Enterprise/GLTFExporter/MigratedFromGit/Source/GLTFExporter/Private/Converters/GLTFTextureUtility.cpp
--- a/Engine/Plugins/Enterprise/GLTFExporter/MigratedFromGit/Source/GLTFExporter/Private/Converters/GLTFTextureUtility.cpp
+++ b/Engine/Plugins/Enterprise/GLTFExporter/MigratedFromGit/Source/GLTFExporter/Private/Converters/GLTFTextureUtility.cpp
@@ -84,18 +84,6 @@ TTuple<TextureAddress, TextureAddress> FGLTFTextureUtility::GetAddressXY(const U
 	return MakeTuple(AddressX, AddressY);
 }
 
-struct FTextureSourceHack
-{
-	FByteBulkData BulkData;
-};
-
-const FByteBulkData& FGLTFTextureUtility::GetBulkData(const FTextureSource& TextureSource)
-{
-	// TODO: replace hack with safe solution by adding proper API access
-	auto Hack = reinterpret_cast<const FTextureSourceHack*>(&TextureSource);
-	return Hack->BulkData;
-}
-
 UTexture2D* FGLTFTextureUtility::CreateTransientTexture(const void* RawData, int64 ByteLength, const FIntPoint& Size, EPixelFormat Format, bool bUseSRGB)
 {
 	check(CalculateImageBytes(Size.X, Size.Y, 0, Format) == ByteLength);
diff --git a/Engine/Plugins/Enterprise/GLTFExporter/MigratedFromGit/Source/GLTFExporter/Private/Converters/GLTFTextureUtility.h b/Engine/Plugins/Enterprise/GLTFExporter/MigratedFromGit/Source/GLTFExporter/Private/Converters/GLTFTextureUtility.h
--- a/Engine/Plugins/Enterprise/GLTFExporter/MigratedFromGit/Source/GLTFExporter/Private/Converters/GLTFTextureUtility.h
+++ b/Engine/Plugins/Enterprise/GLTFExporter/MigratedFromGit/Source/GLTFExporter/Private/Converters/GLTFTextureUtility.h
@@ -19,8 +19,6 @@ struct FGLTFTextureUtility
 
 	static TTuple<TextureAddress, TextureAddress> GetAddressXY(const UTexture* Texture);
 
-	static const FByteBulkData& GetBulkData(const FTextureSource& TextureSource);
-
 	static UTexture2D* CreateTransientTexture(const void* RawData, int64 ByteLength, const FIntPoint& Size, EPixelFormat Format, bool bUseSRGB = false);
 
 	static UTextureRenderTarget2D* CreateRenderTarget(const FIntPoint& Size, EPixelFormat Format, bool bInForceLinearGamma = false);
