[GLTFExporter] feat(texture): add hdr encoding data to gltf texture object in json

commit 9f537f1ae7d86c8196735258872e5c278a314730
Date: Tue, 7 Jul 2020 22:24:16 +0200

    feat(texture): add hdr encoding data to gltf texture object in json
    (cherry picked from commit 3104051d8d4d73ffa68bba1ceecd480e83e4a82f​)

#skipci

diff --git a/Engine/Plugins/Enterprise/GLTFExporter/MigratedFromGit/Source/GLTFExporter/Private/Json/GLTFJsonEnums.h b/Engine/Plugins/Enterprise/GLTFExporter/MigratedFromGit/Source/GLTFExporter/Private/Json/GLTFJsonEnums.h
--- a/Engine/Plugins/Enterprise/GLTFExporter/MigratedFromGit/Source/GLTFExporter/Private/Json/GLTFJsonEnums.h
+++ b/Engine/Plugins/Enterprise/GLTFExporter/MigratedFromGit/Source/GLTFExporter/Private/Json/GLTFJsonEnums.h
@@ -27,6 +27,13 @@ enum class EGLTFJsonShadingModel
 	ClearCoat
 };
 
+enum class EGLTFJsonHDREncoding
+{
+	None = -1,
+    RGBE,
+    RGBM
+};
+
 enum class EGLTFJsonAccessorType
 {
 	None = -1,
diff --git a/Engine/Plugins/Enterprise/GLTFExporter/MigratedFromGit/Source/GLTFExporter/Private/Json/GLTFJsonTexture.h b/Engine/Plugins/Enterprise/GLTFExporter/MigratedFromGit/Source/GLTFExporter/Private/Json/GLTFJsonTexture.h
--- a/Engine/Plugins/Enterprise/GLTFExporter/MigratedFromGit/Source/GLTFExporter/Private/Json/GLTFJsonTexture.h
+++ b/Engine/Plugins/Enterprise/GLTFExporter/MigratedFromGit/Source/GLTFExporter/Private/Json/GLTFJsonTexture.h
@@ -3,6 +3,7 @@
 #pragma once
 
 #include "Json/GLTFJsonIndex.h"
+#include "Json/GLTFJsonUtility.h"
 #include "Serialization/JsonSerializer.h"
 
 struct FGLTFJsonTexture
@@ -13,6 +14,13 @@ struct FGLTFJsonTexture
 
 	FGLTFJsonImageIndex Source;
 
+	EGLTFJsonHDREncoding Encoding;
+
+	FGLTFJsonTexture()
+		: Encoding(EGLTFJsonHDREncoding::None)
+	{
+	}
+
 	template <class CharType = TCHAR, class PrintPolicy = TPrettyJsonPrintPolicy<CharType>>
 	void WriteObject(TJsonWriter<CharType, PrintPolicy>& JsonWriter, FGLTFJsonExtensions& Extensions) const
 	{
@@ -33,6 +41,20 @@ struct FGLTFJsonTexture
 			JsonWriter.WriteValue(TEXT("source"), Source);
 		}
 
+		if (Encoding != EGLTFJsonHDREncoding::None)
+		{
+			const EGLTFJsonExtension Extension = EGLTFJsonExtension::EPIC_TextureHDREncoding;
+			Extensions.Used.Add(Extension);
+
+			JsonWriter.WriteObjectStart(TEXT("extensions"));
+			JsonWriter.WriteObjectStart(FGLTFJsonUtility::ToString(Extension));
+
+			JsonWriter.WriteValue(TEXT("encoding"), FGLTFJsonUtility::ToString(Encoding));
+
+			JsonWriter.WriteObjectEnd();
+			JsonWriter.WriteObjectEnd();
+		}
+
 		JsonWriter.WriteObjectEnd();
 	}
 };
diff --git a/Engine/Plugins/Enterprise/GLTFExporter/MigratedFromGit/Source/GLTFExporter/Private/Json/GLTFJsonUtility.h b/Engine/Plugins/Enterprise/GLTFExporter/MigratedFromGit/Source/GLTFExporter/Private/Json/GLTFJsonUtility.h
--- a/Engine/Plugins/Enterprise/GLTFExporter/MigratedFromGit/Source/GLTFExporter/Private/Json/GLTFJsonUtility.h
+++ b/Engine/Plugins/Enterprise/GLTFExporter/MigratedFromGit/Source/GLTFExporter/Private/Json/GLTFJsonUtility.h
@@ -68,6 +68,16 @@ struct FGLTFJsonUtility
 		}
 	}
 
+	static const TCHAR* ToString(EGLTFJsonHDREncoding Value)
+	{
+		switch (Value)
+		{
+			case EGLTFJsonHDREncoding::RGBE: return TEXT("RGBE");
+			case EGLTFJsonHDREncoding::RGBM: return TEXT("RGBM");
+			default:                         return TEXT("UNKNOWN");
+		}
+	}
+
 	template <class CharType = TCHAR, class PrintPolicy = TPrettyJsonPrintPolicy<CharType>>
 	static void WriteExactValue(TJsonWriter<CharType, PrintPolicy>& JsonWriter, float Value)
 	{
