[GLTFExporter] feat(texture): add ReadEncodedPixels and EncodeRGBM to GLTFTextureUtility

commit bf5c2931e09aa63ff23f286922f00ebe50a6e6b6
Date: Tue, 7 Jul 2020 22:36:25 +0200

    feat(texture): add ReadEncodedPixels and EncodeRGBM to GLTFTextureUtility
    (cherry picked from commit 4f7171d3dd11d0477b482c97ecd6dd88c49d9779​)

#skipci

diff --git a/Engine/Plugins/Enterprise/GLTFExporter/MigratedFromGit/Source/GLTFExporter/Private/Converters/GLTFTextureUtility.cpp b/Engine/Plugins/Enterprise/GLTFExporter/MigratedFromGit/Source/GLTFExporter/Private/Converters/GLTFTextureUtility.cpp
--- a/Engine/Plugins/Enterprise/GLTFExporter/MigratedFromGit/Source/GLTFExporter/Private/Converters/GLTFTextureUtility.cpp
+++ b/Engine/Plugins/Enterprise/GLTFExporter/MigratedFromGit/Source/GLTFExporter/Private/Converters/GLTFTextureUtility.cpp
@@ -129,3 +129,121 @@ bool FGLTFTextureUtility::DrawTexture(UTextureRenderTarget2D* OutTarget, const U
 
 	return true;
 }
+
+UTexture2D* FGLTFTextureUtility::CreateTextureFromCubeFace(const UTextureCube* TextureCube, ECubeFace CubeFace)
+{
+	const FIntPoint Size(TextureCube->GetSizeX(), TextureCube->GetSizeY());
+	const EPixelFormat Format = TextureCube->GetPixelFormat();
+
+	if (TextureCube->PlatformData == nullptr || TextureCube->PlatformData->Mips.Num() == 0)
+	{
+		return nullptr;
+	}
+
+	if (TextureCube->PlatformData->Mips[0].BulkData.GetBulkDataSize() == 0)
+	{
+		const_cast<UTextureCube*>(TextureCube)->ForceRebuildPlatformData();
+		if (TextureCube->PlatformData->Mips[0].BulkData.GetBulkDataSize() == 0)
+		{
+			return nullptr;
+		}
+	}
+
+	const FTexture2DMipMap& Mip = TextureCube->PlatformData->Mips[0];
+	const int64 MipSize = Mip.BulkData.GetBulkDataSize() / 6;
+
+	const void* MipDataPtr = Mip.BulkData.LockReadOnly();
+	const void* FaceDataPtr =  static_cast<const uint8*>(MipDataPtr) + MipSize * CubeFace;
+	UTexture2D* FaceTexture = CreateTransientTexture(FaceDataPtr, MipSize, Size, Format, TextureCube->SRGB);
+	Mip.BulkData.Unlock();
+
+	return FaceTexture;
+}
+
+UTexture2D* FGLTFTextureUtility::CreateTextureFromCubeFace(const UTextureRenderTargetCube* RenderTargetCube, ECubeFace CubeFace)
+{
+	UTexture2D* FaceTexture;
+	const EPixelFormat Format = RenderTargetCube->GetFormat();
+	const FIntPoint Size(RenderTargetCube->SizeX, RenderTargetCube->SizeX);
+	FTextureRenderTargetCubeResource* Resource = static_cast<FTextureRenderTargetCubeResource*>(RenderTargetCube->Resource);
+
+	if (IsHDRFormat(Format))
+	{
+		TArray<FFloat16Color> Pixels;
+		if (!Resource->ReadPixels(Pixels, FReadSurfaceDataFlags(RCM_UNorm, CubeFace)))
+		{
+			return nullptr;
+		}
+
+		FaceTexture = CreateTransientTexture(Pixels.GetData(), Pixels.Num() * sizeof(FFloat16Color), Size, PF_FloatRGBA, false);
+	}
+	else
+	{
+		TArray<FColor> Pixels;
+		if (!Resource->ReadPixels(Pixels, FReadSurfaceDataFlags(RCM_UNorm, CubeFace)))
+		{
+			return nullptr;
+		}
+
+		FaceTexture = CreateTransientTexture(Pixels.GetData(), Pixels.Num() * sizeof(FColor), Size, PF_B8G8R8A8, false);
+	}
+
+	return FaceTexture;
+}
+
+bool FGLTFTextureUtility::ReadEncodedPixels(const UTextureRenderTarget2D* InRenderTarget, TArray<FColor>& OutPixels, EGLTFJsonHDREncoding& OutEncoding)
+{
+	FTextureRenderTarget2DResource* Resource = static_cast<FTextureRenderTarget2DResource*>(InRenderTarget->Resource);
+	if (Resource == nullptr)
+	{
+		return false;
+	}
+
+	if (IsHDRFormat(InRenderTarget->GetFormat()))
+	{
+		TArray<FLinearColor> HDRPixels;
+		Resource->ReadLinearColorPixels(HDRPixels);
+
+		EncodeRGBM(HDRPixels, OutPixels);
+		OutEncoding = EGLTFJsonHDREncoding::RGBM;
+	}
+	else
+	{
+		Resource->ReadPixels(OutPixels);
+		OutEncoding = EGLTFJsonHDREncoding::None;
+	}
+
+	return true;
+}
+
+FColor FGLTFTextureUtility::EncodeRGBM(const FLinearColor& Color, float MaxRange)
+{
+	FLinearColor RGBM;
+
+	RGBM.R = FMath::Sqrt(Color.R);
+	RGBM.G = FMath::Sqrt(Color.G);
+	RGBM.B = FMath::Sqrt(Color.B);
+
+	RGBM.R /= MaxRange;
+	RGBM.G /= MaxRange;
+	RGBM.B /= MaxRange;
+
+	RGBM.A = FMath::Max(FMath::Max(RGBM.R, RGBM.G), FMath::Max(RGBM.B, 1.0f / 255.0f));
+	RGBM.A = FMath::CeilToFloat(RGBM.A * 255.0f) / 255.0f;
+
+	RGBM.R /= RGBM.A;
+	RGBM.G /= RGBM.A;
+	RGBM.B /= RGBM.A;
+
+	return RGBM.ToFColor(false);
+}
+
+void FGLTFTextureUtility::EncodeRGBM(const TArray<FLinearColor>& InPixels, TArray<FColor>& OutPixels, float MaxRange)
+{
+	OutPixels.AddUninitialized(InPixels.Num());
+
+	for (int32 Index = 0; Index < InPixels.Num(); ++Index)
+	{
+		OutPixels[Index] = EncodeRGBM(InPixels[Index], MaxRange);
+	}
+}
diff --git a/Engine/Plugins/Enterprise/GLTFExporter/MigratedFromGit/Source/GLTFExporter/Private/Converters/GLTFTextureUtility.h b/Engine/Plugins/Enterprise/GLTFExporter/MigratedFromGit/Source/GLTFExporter/Private/Converters/GLTFTextureUtility.h
--- a/Engine/Plugins/Enterprise/GLTFExporter/MigratedFromGit/Source/GLTFExporter/Private/Converters/GLTFTextureUtility.h
+++ b/Engine/Plugins/Enterprise/GLTFExporter/MigratedFromGit/Source/GLTFExporter/Private/Converters/GLTFTextureUtility.h
@@ -2,6 +2,7 @@
 
 #pragma once
 
+#include "Json/GLTFJsonEnums.h"
 #include "Engine.h"
 
 enum class ERGBFormat : int8;
@@ -23,5 +24,13 @@ struct FGLTFTextureUtility
 
 	static UTextureRenderTarget2D* CreateRenderTarget(const FIntPoint& Size, EPixelFormat Format, bool bInForceLinearGamma = false);
 
-	static bool DrawTexture(UTextureRenderTarget2D* OutTarget, const UTexture2D* InSource);
+	static bool DrawTexture(UTextureRenderTarget2D* OutTarget, const UTexture2D* InSource, const FMatrix& InTransform = FMatrix::Identity);
+
+	static UTexture2D* CreateTextureFromCubeFace(const UTextureCube* TextureCube, ECubeFace CubeFace);
+	static UTexture2D* CreateTextureFromCubeFace(const UTextureRenderTargetCube* RenderTargetCube, ECubeFace CubeFace);
+
+	static bool ReadEncodedPixels(const UTextureRenderTarget2D* InRenderTarget, TArray<FColor>& OutPixels, EGLTFJsonHDREncoding& OutEncoding);
+
+	static FColor EncodeRGBM(const FLinearColor& Color, float MaxRange = 8);
+	static void EncodeRGBM(const TArray<FLinearColor>& InPixels, TArray<FColor>& OutPixels, float MaxRange = 8);
 };
