[GLTFExporter] feat(texture): add converters for UTextureRenderTarget2D and UTextureRenderTargetCube

commit beb1a5efd887007e71a4609ef666e21a296dba06
Date: Fri, 26 Jun 2020 01:52:17 +0200

    feat(texture): add converters for UTextureRenderTarget2D and UTextureRenderTargetCube
    (cherry picked from commit 131338d3cd341a1e3a3c4c6d8f1e718ce056a204​)

#skipci

diff --git a/Engine/Plugins/Enterprise/GLTFExporter/MigratedFromGit/Source/GLTFExporter/Private/Builders/GLTFConvertBuilder.cpp b/Engine/Plugins/Enterprise/GLTFExporter/MigratedFromGit/Source/GLTFExporter/Private/Builders/GLTFConvertBuilder.cpp
--- a/Engine/Plugins/Enterprise/GLTFExporter/MigratedFromGit/Source/GLTFExporter/Private/Builders/GLTFConvertBuilder.cpp
+++ b/Engine/Plugins/Enterprise/GLTFExporter/MigratedFromGit/Source/GLTFExporter/Private/Builders/GLTFConvertBuilder.cpp
@@ -134,6 +134,36 @@ FGLTFJsonTextureIndex FGLTFConvertBuilder::GetOrAddTexture(const UTexture2D* Tex
 	return Texture2DConverter.GetOrAdd(*this, DesiredName.IsEmpty() ? Texture->GetName() : DesiredName, Texture);
 }
 
+FGLTFJsonTextureIndex FGLTFConvertBuilder::GetOrAddTexture(const UTextureCube* Texture, const FString& DesiredName)
+{
+	if (Texture == nullptr)
+	{
+		return FGLTFJsonTextureIndex(INDEX_NONE);
+	}
+
+	return TextureCubeConverter.GetOrAdd(*this, DesiredName.IsEmpty() ? Texture->GetName() : DesiredName, Texture);
+}
+
+FGLTFJsonTextureIndex FGLTFConvertBuilder::GetOrAddTexture(const UTextureRenderTarget2D* Texture, const FString& DesiredName)
+{
+	if (Texture == nullptr)
+	{
+		return FGLTFJsonTextureIndex(INDEX_NONE);
+	}
+
+	return TextureRenderTarget2DConverter.GetOrAdd(*this, DesiredName.IsEmpty() ? Texture->GetName() : DesiredName, Texture);
+}
+
+FGLTFJsonTextureIndex FGLTFConvertBuilder::GetOrAddTexture(const UTextureRenderTargetCube* Texture, const FString& DesiredName)
+{
+	if (Texture == nullptr)
+	{
+		return FGLTFJsonTextureIndex(INDEX_NONE);
+	}
+
+	return TextureRenderTargetCubeConverter.GetOrAdd(*this, DesiredName.IsEmpty() ? Texture->GetName() : DesiredName, Texture);
+}
+
 FGLTFJsonLightMapIndex FGLTFConvertBuilder::GetOrAddLightMap(const UStaticMeshComponent* StaticMeshComponent, const FString& DesiredName)
 {
 	if (StaticMeshComponent == nullptr)
diff --git a/Engine/Plugins/Enterprise/GLTFExporter/MigratedFromGit/Source/GLTFExporter/Private/Builders/GLTFConvertBuilder.h b/Engine/Plugins/Enterprise/GLTFExporter/MigratedFromGit/Source/GLTFExporter/Private/Builders/GLTFConvertBuilder.h
--- a/Engine/Plugins/Enterprise/GLTFExporter/MigratedFromGit/Source/GLTFExporter/Private/Builders/GLTFConvertBuilder.h
+++ b/Engine/Plugins/Enterprise/GLTFExporter/MigratedFromGit/Source/GLTFExporter/Private/Builders/GLTFConvertBuilder.h
@@ -33,6 +33,9 @@ public:
 	FGLTFJsonMaterialIndex GetOrAddMaterial(const UMaterialInterface* Material, const FString& DesiredName = TEXT(""));
 	FGLTFJsonSamplerIndex GetOrAddSampler(const UTexture* Texture, const FString& DesiredName = TEXT(""));
 	FGLTFJsonTextureIndex GetOrAddTexture(const UTexture2D* Texture, const FString& DesiredName = TEXT(""));
+	FGLTFJsonTextureIndex GetOrAddTexture(const UTextureCube* Texture, const FString& DesiredName = TEXT(""));
+	FGLTFJsonTextureIndex GetOrAddTexture(const UTextureRenderTarget2D* Texture, const FString& DesiredName = TEXT(""));
+	FGLTFJsonTextureIndex GetOrAddTexture(const UTextureRenderTargetCube* Texture, const FString& DesiredName = TEXT(""));
 	FGLTFJsonLightMapIndex GetOrAddLightMap(const UStaticMeshComponent* StaticMeshComponent, const FString& DesiredName = TEXT(""));
 
 	FGLTFJsonNodeIndex GetOrAddNode(const USceneComponent* SceneComponent, const FString& DesiredName = TEXT(""));
@@ -56,6 +59,9 @@ private:
 	FGLTFMaterialConverter MaterialConverter;
 	FGLTFTextureSamplerConverter TextureSamplerConverter;
 	FGLTFTexture2DConverter Texture2DConverter;
+	FGLTFTextureCubeConverter TextureCubeConverter;
+	FGLTFTextureRenderTarget2DConverter TextureRenderTarget2DConverter;
+	FGLTFTextureRenderTargetCubeConverter TextureRenderTargetCubeConverter;
 	FGLTFLightMapConverter LightMapConverter;
 
 	FGLTFSceneComponentConverter SceneComponentConverter;
diff --git a/Engine/Plugins/Enterprise/GLTFExporter/MigratedFromGit/Source/GLTFExporter/Private/Converters/GLTFTextureConverters.cpp b/Engine/Plugins/Enterprise/GLTFExporter/MigratedFromGit/Source/GLTFExporter/Private/Converters/GLTFTextureConverters.cpp
--- a/Engine/Plugins/Enterprise/GLTFExporter/MigratedFromGit/Source/GLTFExporter/Private/Converters/GLTFTextureConverters.cpp
+++ b/Engine/Plugins/Enterprise/GLTFExporter/MigratedFromGit/Source/GLTFExporter/Private/Converters/GLTFTextureConverters.cpp
@@ -102,3 +102,16 @@ FGLTFJsonTextureIndex FGLTFTextureCubeConverter::Add(FGLTFConvertBuilder& Builde
 	// TODO: implement support
 	return FGLTFJsonTextureIndex(INDEX_NONE);
 }
+
+FGLTFJsonTextureIndex FGLTFTextureRenderTarget2DConverter::Add(FGLTFConvertBuilder& Builder, const FString& Name, const UTextureRenderTarget2D* RenderTarget2D)
+{
+	// TODO: implement support
+	return FGLTFJsonTextureIndex(INDEX_NONE);
+}
+
+FGLTFJsonTextureIndex FGLTFTextureRenderTargetCubeConverter::Add(FGLTFConvertBuilder& Builder, const FString& Name, const UTextureRenderTargetCube* RenderTargetCube)
+{
+	// TODO: implement support
+	return FGLTFJsonTextureIndex(INDEX_NONE);
+}
+
diff --git a/Engine/Plugins/Enterprise/GLTFExporter/MigratedFromGit/Source/GLTFExporter/Private/Converters/GLTFTextureConverters.h b/Engine/Plugins/Enterprise/GLTFExporter/MigratedFromGit/Source/GLTFExporter/Private/Converters/GLTFTextureConverters.h
--- a/Engine/Plugins/Enterprise/GLTFExporter/MigratedFromGit/Source/GLTFExporter/Private/Converters/GLTFTextureConverters.h
+++ b/Engine/Plugins/Enterprise/GLTFExporter/MigratedFromGit/Source/GLTFExporter/Private/Converters/GLTFTextureConverters.h
@@ -20,3 +20,13 @@ class FGLTFTextureCubeConverter final : public TGLTFConverter<FGLTFJsonTextureIn
 {
 	FGLTFJsonTextureIndex Add(FGLTFConvertBuilder& Builder, const FString& Name, const UTextureCube* TextureCube) override;
 };
+
+class FGLTFTextureRenderTarget2DConverter final : public TGLTFConverter<FGLTFJsonTextureIndex, const UTextureRenderTarget2D*>
+{
+	FGLTFJsonTextureIndex Add(FGLTFConvertBuilder& Builder, const FString& Name, const UTextureRenderTarget2D* RenderTarget2D) override;
+};
+
+class FGLTFTextureRenderTargetCubeConverter final : public TGLTFConverter<FGLTFJsonTextureIndex, const UTextureRenderTargetCube*>
+{
+	FGLTFJsonTextureIndex Add(FGLTFConvertBuilder& Builder, const FString& Name, const UTextureRenderTargetCube* RenderTargetCube) override;
+};
