[GLTFExporter] feat(texture): implement support for rendertargets and HDR textures

commit ef140247d13fc1c27304cefab020da85dc1b0917
Date: Tue, 7 Jul 2020 23:39:20 +0200

    feat(texture): implement support for rendertargets and HDR textures
    (cherry picked from commit 803b095b2a5be8c902d9bd4f844cd08404a61733​)

#skipci

diff --git a/Engine/Plugins/Enterprise/GLTFExporter/MigratedFromGit/Source/GLTFExporter/Private/Converters/GLTFTextureConverters.cpp b/Engine/Plugins/Enterprise/GLTFExporter/MigratedFromGit/Source/GLTFExporter/Private/Converters/GLTFTextureConverters.cpp
--- a/Engine/Plugins/Enterprise/GLTFExporter/MigratedFromGit/Source/GLTFExporter/Private/Converters/GLTFTextureConverters.cpp
+++ b/Engine/Plugins/Enterprise/GLTFExporter/MigratedFromGit/Source/GLTFExporter/Private/Converters/GLTFTextureConverters.cpp
@@ -70,19 +70,13 @@ FGLTFJsonTextureIndex FGLTFTexture2DConverter::Add(FGLTFConvertBuilder& Builder,
 	}
 	else
 	{
-		UTextureRenderTarget2D* RenderTarget = FGLTFTextureUtility::CreateRenderTarget(Size, PF_B8G8R8A8, true);
+		const EPixelFormat RenderTargetFormat = FGLTFTextureUtility::IsHDRFormat(Texture2D->GetPixelFormat()) ? PF_FloatRGBA : PF_B8G8R8A8;
+		UTextureRenderTarget2D* RenderTarget = FGLTFTextureUtility::CreateRenderTarget(Size, RenderTargetFormat, true);
 
 		FGLTFTextureUtility::DrawTexture(RenderTarget, Texture2D);
 
-		FTextureRenderTarget2DResource* Resource = static_cast<FTextureRenderTarget2DResource*>(RenderTarget->Resource);
-		if (Resource == nullptr)
-		{
-			// TODO: report error
-			return FGLTFJsonTextureIndex(INDEX_NONE);
-		}
-
 		TArray<FColor> Pixels;
-		if (!Resource->ReadPixels(Pixels))
+		if (!FGLTFTextureUtility::ReadEncodedPixels(RenderTarget, Pixels, JsonTexture.Encoding)) // TODO: use only encoding as specified by export options
 		{
 			// TODO: report error
 			return FGLTFJsonTextureIndex(INDEX_NONE);
@@ -105,8 +99,20 @@ FGLTFJsonTextureIndex FGLTFTextureCubeConverter::Add(FGLTFConvertBuilder& Builde
 
 FGLTFJsonTextureIndex FGLTFTextureRenderTarget2DConverter::Add(FGLTFConvertBuilder& Builder, const FString& Name, const UTextureRenderTarget2D* RenderTarget2D)
 {
-	// TODO: implement support
-	return FGLTFJsonTextureIndex(INDEX_NONE);
+	FGLTFJsonTexture JsonTexture;
+	JsonTexture.Name = Name;
+
+	TArray<FColor> Pixels;
+	if (!FGLTFTextureUtility::ReadEncodedPixels(RenderTarget2D, Pixels, JsonTexture.Encoding)) // TODO: use only encoding as specified by export options
+	{
+		// TODO: report error
+		return FGLTFJsonTextureIndex(INDEX_NONE);
+	}
+
+	JsonTexture.Source = Builder.AddImage(Pixels.GetData(), { RenderTarget2D->SizeX, RenderTarget2D->SizeY }, JsonTexture.Name);
+	JsonTexture.Sampler = Builder.GetOrAddSampler(RenderTarget2D);
+
+	return Builder.AddTexture(JsonTexture);
 }
 
 FGLTFJsonTextureIndex FGLTFTextureRenderTargetCubeConverter::Add(FGLTFConvertBuilder& Builder, const FString& Name, const UTextureRenderTargetCube* RenderTargetCube)
