[GLTFExporter] feat(texture): add GetCubeFaceRotation and RotateTexture to GLTFTextureUtility

commit 7682dc41b78b380640ed1e84c91ddee1f604b4b7
Date: Wed, 8 Jul 2020 11:22:10 +0200

    feat(texture): add GetCubeFaceRotation and RotateTexture to GLTFTextureUtility
    (cherry picked from commit bf03c6502ed5328f0d64a6845f3269380cbab17b​)

#skipci

diff --git a/Engine/Plugins/Enterprise/GLTFExporter/MigratedFromGit/Source/GLTFExporter/Private/Converters/GLTFTextureUtility.cpp b/Engine/Plugins/Enterprise/GLTFExporter/MigratedFromGit/Source/GLTFExporter/Private/Converters/GLTFTextureUtility.cpp
--- a/Engine/Plugins/Enterprise/GLTFExporter/MigratedFromGit/Source/GLTFExporter/Private/Converters/GLTFTextureUtility.cpp
+++ b/Engine/Plugins/Enterprise/GLTFExporter/MigratedFromGit/Source/GLTFExporter/Private/Converters/GLTFTextureUtility.cpp
@@ -42,6 +42,20 @@ bool FGLTFTextureUtility::IsCubemap(const UTexture* Texture)
 	return Texture->IsA<UTextureCube>() || Texture->IsA<UTextureRenderTargetCube>();
 }
 
+float FGLTFTextureUtility::GetCubeFaceRotation(ECubeFace CubeFace)
+{
+	switch (CubeFace)
+	{
+		case CubeFace_PosX:	return 90;
+		case CubeFace_NegX:	return -90;
+		case CubeFace_PosY:	return 180;
+		case CubeFace_NegY:	return 0;
+		case CubeFace_PosZ:	return 180;
+		case CubeFace_NegZ:	return 0;
+		default:            return 0;
+	}
+}
+
 TextureFilter FGLTFTextureUtility::GetDefaultFilter(TextureGroup LODGroup)
 {
 	// TODO: should this be the running platform?
@@ -109,7 +123,7 @@ UTextureRenderTarget2D* FGLTFTextureUtility::CreateRenderTarget(const FIntPoint&
 	return RenderTarget;
 }
 
-bool FGLTFTextureUtility::DrawTexture(UTextureRenderTarget2D* OutTarget, const UTexture2D* InSource)
+bool FGLTFTextureUtility::DrawTexture(UTextureRenderTarget2D* OutTarget, const UTexture2D* InSource, const FMatrix& InTransform)
 {
 	FRenderTarget* RenderTarget = OutTarget->GameThread_GetRenderTargetResource();
 	if (RenderTarget == nullptr)
@@ -120,7 +134,9 @@ bool FGLTFTextureUtility::DrawTexture(UTextureRenderTarget2D* OutTarget, const U
 	FCanvas Canvas(RenderTarget, nullptr, 0.0f, 0.0f, 0.0f, GMaxRHIFeatureLevel);
 	FCanvasTileItem TileItem(FVector2D::ZeroVector, InSource->Resource, FLinearColor::White);
 
+	Canvas.PushAbsoluteTransform(InTransform);
 	TileItem.Draw(&Canvas);
+	Canvas.PopTransform();
 
 	Canvas.Flush_GameThread();
 	FlushRenderingCommands();
@@ -130,6 +146,18 @@ bool FGLTFTextureUtility::DrawTexture(UTextureRenderTarget2D* OutTarget, const U
 	return true;
 }
 
+bool FGLTFTextureUtility::RotateTexture(UTextureRenderTarget2D* OutTarget, const UTexture2D* InSource, float InDegrees)
+{
+	FMatrix Transform = FMatrix::Identity;
+	if (InDegrees != 0)
+	{
+		const FVector Center = FVector(InSource->GetSizeX() / 2.0f, InSource->GetSizeY() / 2.0f, 0);
+		Transform = FTranslationMatrix(-Center) * FRotationMatrix({ 0, InDegrees, 0 }) * FTranslationMatrix(Center);
+	}
+
+	return DrawTexture(OutTarget, InSource, Transform);
+}
+
 UTexture2D* FGLTFTextureUtility::CreateTextureFromCubeFace(const UTextureCube* TextureCube, ECubeFace CubeFace)
 {
 	const FIntPoint Size(TextureCube->GetSizeX(), TextureCube->GetSizeY());
diff --git a/Engine/Plugins/Enterprise/GLTFExporter/MigratedFromGit/Source/GLTFExporter/Private/Converters/GLTFTextureUtility.h b/Engine/Plugins/Enterprise/GLTFExporter/MigratedFromGit/Source/GLTFExporter/Private/Converters/GLTFTextureUtility.h
--- a/Engine/Plugins/Enterprise/GLTFExporter/MigratedFromGit/Source/GLTFExporter/Private/Converters/GLTFTextureUtility.h
+++ b/Engine/Plugins/Enterprise/GLTFExporter/MigratedFromGit/Source/GLTFExporter/Private/Converters/GLTFTextureUtility.h
@@ -16,6 +16,8 @@ struct FGLTFTextureUtility
 
 	static bool IsCubemap(const UTexture* Texture);
 
+	static float GetCubeFaceRotation(ECubeFace CubeFace);
+
 	static TextureFilter GetDefaultFilter(TextureGroup Group);
 
 	static TTuple<TextureAddress, TextureAddress> GetAddressXY(const UTexture* Texture);
@@ -25,6 +27,7 @@ struct FGLTFTextureUtility
 	static UTextureRenderTarget2D* CreateRenderTarget(const FIntPoint& Size, EPixelFormat Format, bool bInForceLinearGamma = false);
 
 	static bool DrawTexture(UTextureRenderTarget2D* OutTarget, const UTexture2D* InSource, const FMatrix& InTransform = FMatrix::Identity);
+	static bool RotateTexture(UTextureRenderTarget2D* OutTarget, const UTexture2D* InSource, float InDegrees);
 
 	static UTexture2D* CreateTextureFromCubeFace(const UTextureCube* TextureCube, ECubeFace CubeFace);
 	static UTexture2D* CreateTextureFromCubeFace(const UTextureRenderTargetCube* RenderTargetCube, ECubeFace CubeFace);
