[GLTFExporter] refactor(texture): add cubemap face parameter to each cubemap texture converter

commit f251c66268e3fc676e509e0d81f7a0aee79cc33b
Date: Tue, 30 Jun 2020 01:39:01 +0200

    refactor(texture): add cubemap face parameter to each cubemap texture converter
    (cherry picked from commit 422efa883a6552b52dfb1dae19b5092fe4e28968​)

#skipci

diff --git a/Engine/Plugins/Enterprise/GLTFExporter/MigratedFromGit/Source/GLTFExporter/Private/Builders/GLTFConvertBuilder.cpp b/Engine/Plugins/Enterprise/GLTFExporter/MigratedFromGit/Source/GLTFExporter/Private/Builders/GLTFConvertBuilder.cpp
--- a/Engine/Plugins/Enterprise/GLTFExporter/MigratedFromGit/Source/GLTFExporter/Private/Builders/GLTFConvertBuilder.cpp
+++ b/Engine/Plugins/Enterprise/GLTFExporter/MigratedFromGit/Source/GLTFExporter/Private/Builders/GLTFConvertBuilder.cpp
@@ -134,14 +134,14 @@ FGLTFJsonTextureIndex FGLTFConvertBuilder::GetOrAddTexture(const UTexture2D* Tex
 	return Texture2DConverter.GetOrAdd(*this, DesiredName.IsEmpty() ? Texture->GetName() : DesiredName, Texture);
 }
 
-FGLTFJsonTextureIndex FGLTFConvertBuilder::GetOrAddTexture(const UTextureCube* Texture, const FString& DesiredName)
+FGLTFJsonTextureIndex FGLTFConvertBuilder::GetOrAddTexture(const UTextureCube* Texture, ECubeFace CubeFace, const FString& DesiredName)
 {
 	if (Texture == nullptr)
 	{
 		return FGLTFJsonTextureIndex(INDEX_NONE);
 	}
 
-	return TextureCubeConverter.GetOrAdd(*this, DesiredName.IsEmpty() ? Texture->GetName() : DesiredName, Texture);
+	return TextureCubeConverter.GetOrAdd(*this, DesiredName.IsEmpty() ? Texture->GetName() : DesiredName, Texture, CubeFace);
 }
 
 FGLTFJsonTextureIndex FGLTFConvertBuilder::GetOrAddTexture(const UTextureRenderTarget2D* Texture, const FString& DesiredName)
@@ -154,14 +154,14 @@ FGLTFJsonTextureIndex FGLTFConvertBuilder::GetOrAddTexture(const UTextureRenderT
 	return TextureRenderTarget2DConverter.GetOrAdd(*this, DesiredName.IsEmpty() ? Texture->GetName() : DesiredName, Texture);
 }
 
-FGLTFJsonTextureIndex FGLTFConvertBuilder::GetOrAddTexture(const UTextureRenderTargetCube* Texture, const FString& DesiredName)
+FGLTFJsonTextureIndex FGLTFConvertBuilder::GetOrAddTexture(const UTextureRenderTargetCube* Texture, ECubeFace CubeFace, const FString& DesiredName)
 {
 	if (Texture == nullptr)
 	{
 		return FGLTFJsonTextureIndex(INDEX_NONE);
 	}
 
-	return TextureRenderTargetCubeConverter.GetOrAdd(*this, DesiredName.IsEmpty() ? Texture->GetName() : DesiredName, Texture);
+	return TextureRenderTargetCubeConverter.GetOrAdd(*this, DesiredName.IsEmpty() ? Texture->GetName() : DesiredName, Texture, CubeFace);
 }
 
 FGLTFJsonLightMapIndex FGLTFConvertBuilder::GetOrAddLightMap(const UStaticMeshComponent* StaticMeshComponent, const FString& DesiredName)
diff --git a/Engine/Plugins/Enterprise/GLTFExporter/MigratedFromGit/Source/GLTFExporter/Private/Builders/GLTFConvertBuilder.h b/Engine/Plugins/Enterprise/GLTFExporter/MigratedFromGit/Source/GLTFExporter/Private/Builders/GLTFConvertBuilder.h
--- a/Engine/Plugins/Enterprise/GLTFExporter/MigratedFromGit/Source/GLTFExporter/Private/Builders/GLTFConvertBuilder.h
+++ b/Engine/Plugins/Enterprise/GLTFExporter/MigratedFromGit/Source/GLTFExporter/Private/Builders/GLTFConvertBuilder.h
@@ -33,9 +33,9 @@ public:
 	FGLTFJsonMaterialIndex GetOrAddMaterial(const UMaterialInterface* Material, const FString& DesiredName = TEXT(""));
 	FGLTFJsonSamplerIndex GetOrAddSampler(const UTexture* Texture, const FString& DesiredName = TEXT(""));
 	FGLTFJsonTextureIndex GetOrAddTexture(const UTexture2D* Texture, const FString& DesiredName = TEXT(""));
-	FGLTFJsonTextureIndex GetOrAddTexture(const UTextureCube* Texture, const FString& DesiredName = TEXT(""));
+	FGLTFJsonTextureIndex GetOrAddTexture(const UTextureCube* Texture, ECubeFace CubeFace, const FString& DesiredName = TEXT(""));
 	FGLTFJsonTextureIndex GetOrAddTexture(const UTextureRenderTarget2D* Texture, const FString& DesiredName = TEXT(""));
-	FGLTFJsonTextureIndex GetOrAddTexture(const UTextureRenderTargetCube* Texture, const FString& DesiredName = TEXT(""));
+	FGLTFJsonTextureIndex GetOrAddTexture(const UTextureRenderTargetCube* Texture, ECubeFace CubeFace, const FString& DesiredName = TEXT(""));
 	FGLTFJsonLightMapIndex GetOrAddLightMap(const UStaticMeshComponent* StaticMeshComponent, const FString& DesiredName = TEXT(""));
 
 	FGLTFJsonNodeIndex GetOrAddNode(const USceneComponent* SceneComponent, const FString& DesiredName = TEXT(""));
diff --git a/Engine/Plugins/Enterprise/GLTFExporter/MigratedFromGit/Source/GLTFExporter/Private/Converters/GLTFTextureConverters.cpp b/Engine/Plugins/Enterprise/GLTFExporter/MigratedFromGit/Source/GLTFExporter/Private/Converters/GLTFTextureConverters.cpp
--- a/Engine/Plugins/Enterprise/GLTFExporter/MigratedFromGit/Source/GLTFExporter/Private/Converters/GLTFTextureConverters.cpp
+++ b/Engine/Plugins/Enterprise/GLTFExporter/MigratedFromGit/Source/GLTFExporter/Private/Converters/GLTFTextureConverters.cpp
@@ -91,7 +91,7 @@ FGLTFJsonTextureIndex FGLTFTexture2DConverter::Add(FGLTFConvertBuilder& Builder,
 	return Builder.AddTexture(JsonTexture);
 }
 
-FGLTFJsonTextureIndex FGLTFTextureCubeConverter::Add(FGLTFConvertBuilder& Builder, const FString& Name, const UTextureCube* TextureCube)
+FGLTFJsonTextureIndex FGLTFTextureCubeConverter::Add(FGLTFConvertBuilder& Builder, const FString& Name, const UTextureCube* TextureCube, ECubeFace CubeFace)
 {
 	// TODO: implement support
 	return FGLTFJsonTextureIndex(INDEX_NONE);
@@ -115,7 +115,7 @@ FGLTFJsonTextureIndex FGLTFTextureRenderTarget2DConverter::Add(FGLTFConvertBuild
 	return Builder.AddTexture(JsonTexture);
 }
 
-FGLTFJsonTextureIndex FGLTFTextureRenderTargetCubeConverter::Add(FGLTFConvertBuilder& Builder, const FString& Name, const UTextureRenderTargetCube* RenderTargetCube)
+FGLTFJsonTextureIndex FGLTFTextureRenderTargetCubeConverter::Add(FGLTFConvertBuilder& Builder, const FString& Name, const UTextureRenderTargetCube* RenderTargetCube, ECubeFace CubeFace)
 {
 	// TODO: implement support
 	return FGLTFJsonTextureIndex(INDEX_NONE);
diff --git a/Engine/Plugins/Enterprise/GLTFExporter/MigratedFromGit/Source/GLTFExporter/Private/Converters/GLTFTextureConverters.h b/Engine/Plugins/Enterprise/GLTFExporter/MigratedFromGit/Source/GLTFExporter/Private/Converters/GLTFTextureConverters.h
--- a/Engine/Plugins/Enterprise/GLTFExporter/MigratedFromGit/Source/GLTFExporter/Private/Converters/GLTFTextureConverters.h
+++ b/Engine/Plugins/Enterprise/GLTFExporter/MigratedFromGit/Source/GLTFExporter/Private/Converters/GLTFTextureConverters.h
@@ -16,9 +16,9 @@ class FGLTFTexture2DConverter final : public TGLTFConverter<FGLTFJsonTextureInde
 	FGLTFJsonTextureIndex Add(FGLTFConvertBuilder& Builder, const FString& Name, const UTexture2D* Texture2D) override;
 };
 
-class FGLTFTextureCubeConverter final : public TGLTFConverter<FGLTFJsonTextureIndex, const UTextureCube*>
+class FGLTFTextureCubeConverter final : public TGLTFConverter<FGLTFJsonTextureIndex, const UTextureCube*, ECubeFace>
 {
-	FGLTFJsonTextureIndex Add(FGLTFConvertBuilder& Builder, const FString& Name, const UTextureCube* TextureCube) override;
+	FGLTFJsonTextureIndex Add(FGLTFConvertBuilder& Builder, const FString& Name, const UTextureCube* TextureCube, ECubeFace CubeFace) override;
 };
 
 class FGLTFTextureRenderTarget2DConverter final : public TGLTFConverter<FGLTFJsonTextureIndex, const UTextureRenderTarget2D*>
@@ -26,7 +26,7 @@ class FGLTFTextureRenderTarget2DConverter final : public TGLTFConverter<FGLTFJso
 	FGLTFJsonTextureIndex Add(FGLTFConvertBuilder& Builder, const FString& Name, const UTextureRenderTarget2D* RenderTarget2D) override;
 };
 
-class FGLTFTextureRenderTargetCubeConverter final : public TGLTFConverter<FGLTFJsonTextureIndex, const UTextureRenderTargetCube*>
+class FGLTFTextureRenderTargetCubeConverter final : public TGLTFConverter<FGLTFJsonTextureIndex, const UTextureRenderTargetCube*, ECubeFace>
 {
-	FGLTFJsonTextureIndex Add(FGLTFConvertBuilder& Builder, const FString& Name, const UTextureRenderTargetCube* RenderTargetCube) override;
+	FGLTFJsonTextureIndex Add(FGLTFConvertBuilder& Builder, const FString& Name, const UTextureRenderTargetCube* RenderTargetCube, ECubeFace CubeFace) override;
 };
