[GLTFExporter] feat(texture): add support for converting cubemap textures and rendertargets

commit 814e663b3ade62b69f9fdc1162401721c1b71371
Date: Tue, 7 Jul 2020 23:40:44 +0200

    feat(texture): add support for converting cubemap textures and rendertargets
    (cherry picked from commit 1dce7143462611ab40952537ecdcc6104c0b8e8b​)

#skipci

diff --git a/Engine/Plugins/Enterprise/GLTFExporter/MigratedFromGit/Source/GLTFExporter/Private/Converters/GLTFTextureConverters.cpp b/Engine/Plugins/Enterprise/GLTFExporter/MigratedFromGit/Source/GLTFExporter/Private/Converters/GLTFTextureConverters.cpp
--- a/Engine/Plugins/Enterprise/GLTFExporter/MigratedFromGit/Source/GLTFExporter/Private/Converters/GLTFTextureConverters.cpp
+++ b/Engine/Plugins/Enterprise/GLTFExporter/MigratedFromGit/Source/GLTFExporter/Private/Converters/GLTFTextureConverters.cpp
@@ -93,8 +93,34 @@ FGLTFJsonTextureIndex FGLTFTexture2DConverter::Add(FGLTFConvertBuilder& Builder,
 
 FGLTFJsonTextureIndex FGLTFTextureCubeConverter::Add(FGLTFConvertBuilder& Builder, const FString& Name, const UTextureCube* TextureCube, ECubeFace CubeFace)
 {
-	// TODO: implement support
-	return FGLTFJsonTextureIndex(INDEX_NONE);
+	FGLTFJsonTexture JsonTexture;
+	JsonTexture.Name = Name;
+
+	const UTexture2D* FaceTexture = FGLTFTextureUtility::CreateTextureFromCubeFace(TextureCube, CubeFace);
+	if (FaceTexture == nullptr)
+	{
+		// TODO: report error
+		return FGLTFJsonTextureIndex(INDEX_NONE);
+	}
+
+	const FIntPoint Size = { FaceTexture->GetSizeX(), FaceTexture->GetSizeY() };
+	const EPixelFormat RenderTargetFormat = FGLTFTextureUtility::IsHDRFormat(FaceTexture->GetPixelFormat()) ? PF_FloatRGBA : PF_B8G8R8A8;
+	UTextureRenderTarget2D* RenderTarget = FGLTFTextureUtility::CreateRenderTarget(Size, RenderTargetFormat, true);
+
+	const float FaceRotation = FGLTFTextureUtility::GetCubeFaceRotation(CubeFace);
+	FGLTFTextureUtility::RotateTexture(RenderTarget, FaceTexture, FaceRotation);
+
+	TArray<FColor> Pixels;
+	if (!FGLTFTextureUtility::ReadEncodedPixels(RenderTarget, Pixels, JsonTexture.Encoding)) // TODO: use only encoding as specified by export options
+	{
+		// TODO: report error
+		return FGLTFJsonTextureIndex(INDEX_NONE);
+	}
+
+	JsonTexture.Source = Builder.AddImage(Pixels.GetData(), Size, JsonTexture.Name);
+	JsonTexture.Sampler = Builder.GetOrAddSampler(TextureCube);
+
+	return Builder.AddTexture(JsonTexture);
 }
 
 FGLTFJsonTextureIndex FGLTFTextureRenderTarget2DConverter::Add(FGLTFConvertBuilder& Builder, const FString& Name, const UTextureRenderTarget2D* RenderTarget2D)
@@ -117,7 +143,33 @@ FGLTFJsonTextureIndex FGLTFTextureRenderTarget2DConverter::Add(FGLTFConvertBuild
 
 FGLTFJsonTextureIndex FGLTFTextureRenderTargetCubeConverter::Add(FGLTFConvertBuilder& Builder, const FString& Name, const UTextureRenderTargetCube* RenderTargetCube, ECubeFace CubeFace)
 {
-	// TODO: implement support
-	return FGLTFJsonTextureIndex(INDEX_NONE);
+	FGLTFJsonTexture JsonTexture;
+	JsonTexture.Name = Name;
+
+	const UTexture2D* FaceTexture = FGLTFTextureUtility::CreateTextureFromCubeFace(RenderTargetCube, CubeFace);
+	if (FaceTexture == nullptr)
+	{
+		// TODO: report error
+		return FGLTFJsonTextureIndex(INDEX_NONE);
+	}
+
+	const FIntPoint Size = { FaceTexture->GetSizeX(), FaceTexture->GetSizeY() };
+	const EPixelFormat RenderTargetFormat = FGLTFTextureUtility::IsHDRFormat(FaceTexture->GetPixelFormat()) ? PF_FloatRGBA : PF_B8G8R8A8;
+	UTextureRenderTarget2D* RenderTarget = FGLTFTextureUtility::CreateRenderTarget(Size, RenderTargetFormat, true);
+
+	const float FaceRotation = FGLTFTextureUtility::GetCubeFaceRotation(CubeFace);
+	FGLTFTextureUtility::RotateTexture(RenderTarget, FaceTexture, FaceRotation);
+
+	TArray<FColor> Pixels;
+	if (!FGLTFTextureUtility::ReadEncodedPixels(RenderTarget, Pixels, JsonTexture.Encoding)) // TODO: use only encoding as specified by export options
+	{
+		// TODO: report error
+		return FGLTFJsonTextureIndex(INDEX_NONE);
+	}
+
+	JsonTexture.Source = Builder.AddImage(Pixels.GetData(), Size, JsonTexture.Name);
+	JsonTexture.Sampler = Builder.GetOrAddSampler(RenderTargetCube);
+
+	return Builder.AddTexture(JsonTexture);
 }
 
