[GLTFExporter] refactor(texture): add TODOs for adding optimized conversion paths for certain cubemap faces

commit 24fb87e6ffae84d645963b2c74a155408068e50e
Date: Tue, 7 Jul 2020 23:46:38 +0200

    refactor(texture): add TODOs for adding optimized conversion paths for certain cubemap faces
    (cherry picked from commit 549d3f1d7d650b48c5669be998650bc143eee0cf​)

#skipci

diff --git a/Engine/Plugins/Enterprise/GLTFExporter/MigratedFromGit/Source/GLTFExporter/Private/Converters/GLTFTextureConverters.cpp b/Engine/Plugins/Enterprise/GLTFExporter/MigratedFromGit/Source/GLTFExporter/Private/Converters/GLTFTextureConverters.cpp
--- a/Engine/Plugins/Enterprise/GLTFExporter/MigratedFromGit/Source/GLTFExporter/Private/Converters/GLTFTextureConverters.cpp
+++ b/Engine/Plugins/Enterprise/GLTFExporter/MigratedFromGit/Source/GLTFExporter/Private/Converters/GLTFTextureConverters.cpp
@@ -96,6 +96,8 @@ FGLTFJsonTextureIndex FGLTFTextureCubeConverter::Add(FGLTFConvertBuilder& Builde
 	FGLTFJsonTexture JsonTexture;
 	JsonTexture.Name = Name;
 
+	// TODO: add optimized "happy path" if cube face doesn't need rotation and has suitable pixel format
+
 	const UTexture2D* FaceTexture = FGLTFTextureUtility::CreateTextureFromCubeFace(TextureCube, CubeFace);
 	if (FaceTexture == nullptr)
 	{
@@ -146,6 +148,8 @@ FGLTFJsonTextureIndex FGLTFTextureRenderTargetCubeConverter::Add(FGLTFConvertBui
 	FGLTFJsonTexture JsonTexture;
 	JsonTexture.Name = Name;
 
+	// TODO: add optimized "happy path" if cube face doesn't need rotation
+
 	const UTexture2D* FaceTexture = FGLTFTextureUtility::CreateTextureFromCubeFace(RenderTargetCube, CubeFace);
 	if (FaceTexture == nullptr)
 	{
@@ -172,4 +176,3 @@ FGLTFJsonTextureIndex FGLTFTextureRenderTargetCubeConverter::Add(FGLTFConvertBui
 
 	return Builder.AddTexture(JsonTexture);
 }
-
