[GLTFExporter] feat(json): add EGLTFJsonCubeFace and conversion from ECubeFace

commit bc3f425169e4909085bb7242d3570735da70c967
Date: Tue, 7 Jul 2020 22:39:42 +0200

    feat(json): add EGLTFJsonCubeFace and conversion from ECubeFace
    (cherry picked from commit 59c75da3ea71f96dcf0edd8953c58eb25aa9d0e0​)

#skipci

diff --git a/Engine/Plugins/Enterprise/GLTFExporter/MigratedFromGit/Source/GLTFExporter/Private/Converters/GLTFConverterUtility.cpp b/Engine/Plugins/Enterprise/GLTFExporter/MigratedFromGit/Source/GLTFExporter/Private/Converters/GLTFConverterUtility.cpp
--- a/Engine/Plugins/Enterprise/GLTFExporter/MigratedFromGit/Source/GLTFExporter/Private/Converters/GLTFConverterUtility.cpp
+++ b/Engine/Plugins/Enterprise/GLTFExporter/MigratedFromGit/Source/GLTFExporter/Private/Converters/GLTFConverterUtility.cpp
@@ -56,3 +56,17 @@ EGLTFJsonTextureFilter FGLTFConverterUtility::ConvertMagFilter(TextureFilter Fil
 {
 	return ConvertMagFilter(Filter == TF_Default ? FGLTFTextureUtility::GetDefaultFilter(LODGroup) : Filter);
 }
+
+EGLTFJsonCubeFace FGLTFConverterUtility::ConvertCubeFace(ECubeFace CubeFace)
+{
+	switch (CubeFace)
+	{
+		case CubeFace_PosX:	return EGLTFJsonCubeFace::NegX;
+		case CubeFace_NegX:	return EGLTFJsonCubeFace::PosX;
+		case CubeFace_PosY:	return EGLTFJsonCubeFace::PosZ;
+		case CubeFace_NegY:	return EGLTFJsonCubeFace::NegZ;
+		case CubeFace_PosZ:	return EGLTFJsonCubeFace::PosY;
+		case CubeFace_NegZ:	return EGLTFJsonCubeFace::NegY;
+		default:            return EGLTFJsonCubeFace::None;
+	}
+}
diff --git a/Engine/Plugins/Enterprise/GLTFExporter/MigratedFromGit/Source/GLTFExporter/Private/Converters/GLTFConverterUtility.h b/Engine/Plugins/Enterprise/GLTFExporter/MigratedFromGit/Source/GLTFExporter/Private/Converters/GLTFConverterUtility.h
--- a/Engine/Plugins/Enterprise/GLTFExporter/MigratedFromGit/Source/GLTFExporter/Private/Converters/GLTFConverterUtility.h
+++ b/Engine/Plugins/Enterprise/GLTFExporter/MigratedFromGit/Source/GLTFExporter/Private/Converters/GLTFConverterUtility.h
@@ -112,6 +112,8 @@ struct FGLTFConverterUtility
 	static EGLTFJsonTextureFilter ConvertMinFilter(TextureFilter Filter, TextureGroup LODGroup);
 	static EGLTFJsonTextureFilter ConvertMagFilter(TextureFilter Filter, TextureGroup LODGroup);
 
+	static EGLTFJsonCubeFace ConvertCubeFace(ECubeFace CubeFace);
+
 	template <typename EnumType>
 	static FString GetEnumDisplayName(EnumType Value)
 	{
diff --git a/Engine/Plugins/Enterprise/GLTFExporter/MigratedFromGit/Source/GLTFExporter/Private/Json/GLTFJsonEnums.h b/Engine/Plugins/Enterprise/GLTFExporter/MigratedFromGit/Source/GLTFExporter/Private/Json/GLTFJsonEnums.h
--- a/Engine/Plugins/Enterprise/GLTFExporter/MigratedFromGit/Source/GLTFExporter/Private/Json/GLTFJsonEnums.h
+++ b/Engine/Plugins/Enterprise/GLTFExporter/MigratedFromGit/Source/GLTFExporter/Private/Json/GLTFJsonEnums.h
@@ -34,6 +34,17 @@ enum class EGLTFJsonHDREncoding
     RGBM
 };
 
+enum class EGLTFJsonCubeFace
+{
+	None = -1,
+	PosX,
+    NegX,
+    PosY,
+    NegY,
+    PosZ,
+    NegZ
+};
+
 enum class EGLTFJsonAccessorType
 {
 	None = -1,
diff --git a/Engine/Plugins/Enterprise/GLTFExporter/MigratedFromGit/Source/GLTFExporter/Private/Json/GLTFJsonUtility.h b/Engine/Plugins/Enterprise/GLTFExporter/MigratedFromGit/Source/GLTFExporter/Private/Json/GLTFJsonUtility.h
--- a/Engine/Plugins/Enterprise/GLTFExporter/MigratedFromGit/Source/GLTFExporter/Private/Json/GLTFJsonUtility.h
+++ b/Engine/Plugins/Enterprise/GLTFExporter/MigratedFromGit/Source/GLTFExporter/Private/Json/GLTFJsonUtility.h
@@ -78,6 +78,20 @@ struct FGLTFJsonUtility
 		}
 	}
 
+	static const TCHAR* ToString(EGLTFJsonCubeFace Value)
+	{
+		switch (Value)
+		{
+			case EGLTFJsonCubeFace::PosX: return TEXT("PosX");
+			case EGLTFJsonCubeFace::NegX: return TEXT("NegX");
+			case EGLTFJsonCubeFace::PosY: return TEXT("PosY");
+			case EGLTFJsonCubeFace::NegY: return TEXT("NegY");
+			case EGLTFJsonCubeFace::PosZ: return TEXT("PosZ");
+			case EGLTFJsonCubeFace::NegZ: return TEXT("NegZ");
+			default:                      return TEXT("Unknown");
+		}
+	}
+
 	template <class CharType = TCHAR, class PrintPolicy = TPrettyJsonPrintPolicy<CharType>>
 	static void WriteExactValue(TJsonWriter<CharType, PrintPolicy>& JsonWriter, float Value)
 	{
