[GLTFExporter] feat(backdrop): add json data structures for hdri backdrops

commit 6520b97b7c7fb13c3331356e47bedb58666aa454
Date: Wed, 8 Jul 2020 00:12:31 +0200

    feat(backdrop): add json data structures for hdri backdrops
    (cherry picked from commit a9a2c7ae570499eedd825cfb356d271e56811513​)

#skipci

diff --git a/Engine/Plugins/Enterprise/GLTFExporter/MigratedFromGit/Source/GLTFExporter/Private/Builders/GLTFJsonBuilder.cpp b/Engine/Plugins/Enterprise/GLTFExporter/MigratedFromGit/Source/GLTFExporter/Private/Builders/GLTFJsonBuilder.cpp
--- a/Engine/Plugins/Enterprise/GLTFExporter/MigratedFromGit/Source/GLTFExporter/Private/Builders/GLTFJsonBuilder.cpp
+++ b/Engine/Plugins/Enterprise/GLTFExporter/MigratedFromGit/Source/GLTFExporter/Private/Builders/GLTFJsonBuilder.cpp
@@ -58,9 +58,9 @@ FGLTFJsonTextureIndex FGLTFJsonBuilder::AddTexture(const FGLTFJsonTexture& JsonT
 	return FGLTFJsonTextureIndex(JsonRoot.Textures.Add(JsonTexture));
 }
 
-FGLTFJsonLightMapIndex FGLTFJsonBuilder::AddLightMap(const FGLTFJsonLightMap& JsonLightMap)
+FGLTFJsonBackdropIndex FGLTFJsonBuilder::AddBackdrop(const FGLTFJsonBackdrop& JsonBackdrop)
 {
-	return FGLTFJsonLightMapIndex(JsonRoot.LightMaps.Add(JsonLightMap));
+	return FGLTFJsonBackdropIndex(JsonRoot.Backdrops.Add(JsonBackdrop));
 }
 
 FGLTFJsonLevelVariantSetsIndex FGLTFJsonBuilder::AddLevelVariantSets(const FGLTFJsonLevelVariantSets& LevelVariantSets)
@@ -68,6 +68,11 @@ FGLTFJsonLevelVariantSetsIndex FGLTFJsonBuilder::AddLevelVariantSets(const FGLTF
 	return FGLTFJsonLevelVariantSetsIndex(JsonRoot.LevelVariantSets.Add(LevelVariantSets));
 }
 
+FGLTFJsonLightMapIndex FGLTFJsonBuilder::AddLightMap(const FGLTFJsonLightMap& JsonLightMap)
+{
+	return FGLTFJsonLightMapIndex(JsonRoot.LightMaps.Add(JsonLightMap));
+}
+
 FGLTFJsonNodeIndex FGLTFJsonBuilder::AddChildNode(FGLTFJsonNodeIndex ParentIndex, const FGLTFJsonNode& JsonNode)
 {
 	const FGLTFJsonNodeIndex ChildIndex = AddNode(JsonNode);
@@ -130,9 +135,9 @@ FGLTFJsonTexture& FGLTFJsonBuilder::GetTexture(FGLTFJsonTextureIndex TextureInde
 	return JsonRoot.Textures[TextureIndex];
 }
 
-FGLTFJsonLightMap& FGLTFJsonBuilder::GetLightMap(FGLTFJsonLightMapIndex LightMapIndex)
+FGLTFJsonBackdrop& FGLTFJsonBuilder::GetBackdrop(FGLTFJsonBackdropIndex BackdropIndex)
 {
-	return JsonRoot.LightMaps[LightMapIndex];
+	return JsonRoot.Backdrops[BackdropIndex];
 }
 
 FGLTFJsonLevelVariantSets& FGLTFJsonBuilder::GetLevelVariantSets(FGLTFJsonLevelVariantSetsIndex LevelVariantSetsIndex)
@@ -140,6 +145,11 @@ FGLTFJsonLevelVariantSets& FGLTFJsonBuilder::GetLevelVariantSets(FGLTFJsonLevelV
 	return JsonRoot.LevelVariantSets[LevelVariantSetsIndex];
 }
 
+FGLTFJsonLightMap& FGLTFJsonBuilder::GetLightMap(FGLTFJsonLightMapIndex LightMapIndex)
+{
+	return JsonRoot.LightMaps[LightMapIndex];
+}
+
 bool FGLTFJsonBuilder::Serialize(FArchive& Archive, const FString& FilePath)
 {
 	JsonRoot.Serialize(&Archive, true);
diff --git a/Engine/Plugins/Enterprise/GLTFExporter/MigratedFromGit/Source/GLTFExporter/Private/Builders/GLTFJsonBuilder.h b/Engine/Plugins/Enterprise/GLTFExporter/MigratedFromGit/Source/GLTFExporter/Private/Builders/GLTFJsonBuilder.h
--- a/Engine/Plugins/Enterprise/GLTFExporter/MigratedFromGit/Source/GLTFExporter/Private/Builders/GLTFJsonBuilder.h
+++ b/Engine/Plugins/Enterprise/GLTFExporter/MigratedFromGit/Source/GLTFExporter/Private/Builders/GLTFJsonBuilder.h
@@ -23,8 +23,9 @@ public:
 	FGLTFJsonSamplerIndex AddSampler(const FGLTFJsonSampler& JsonSampler = {});
 	FGLTFJsonSceneIndex AddScene(const FGLTFJsonScene& JsonScene = {});
 	FGLTFJsonTextureIndex AddTexture(const FGLTFJsonTexture& JsonTexture = {});
-	FGLTFJsonLightMapIndex AddLightMap(const FGLTFJsonLightMap& JsonLightMap = {});
+	FGLTFJsonBackdropIndex AddBackdrop(const FGLTFJsonBackdrop& JsonBackdrop = {});
 	FGLTFJsonLevelVariantSetsIndex AddLevelVariantSets(const FGLTFJsonLevelVariantSets& LevelVariantSets = {});
+	FGLTFJsonLightMapIndex AddLightMap(const FGLTFJsonLightMap& JsonLightMap = {});
 
 	FGLTFJsonNodeIndex AddChildNode(FGLTFJsonNodeIndex ParentNodeIndex, const FGLTFJsonNode& JsonNode = {});
 
@@ -38,8 +39,9 @@ public:
 	FGLTFJsonSampler& GetSampler(FGLTFJsonSamplerIndex SamplerIndex);
 	FGLTFJsonScene& GetScene(FGLTFJsonSceneIndex SceneIndex);
 	FGLTFJsonTexture& GetTexture(FGLTFJsonTextureIndex TextureIndex);
-	FGLTFJsonLightMap& GetLightMap(FGLTFJsonLightMapIndex LightMapIndex);
+	FGLTFJsonBackdrop& GetBackdrop(FGLTFJsonBackdropIndex BackdropIndex);
 	FGLTFJsonLevelVariantSets& GetLevelVariantSets(FGLTFJsonLevelVariantSetsIndex LevelVariantSetsIndex);
+	FGLTFJsonLightMap& GetLightMap(FGLTFJsonLightMapIndex LightMapIndex);
 
 	virtual bool Serialize(FArchive& Archive, const FString& FilePath);
 
diff --git a/Engine/Plugins/Enterprise/GLTFExporter/MigratedFromGit/Source/GLTFExporter/Private/Json/GLTFJsonBackdrop.h b/Engine/Plugins/Enterprise/GLTFExporter/MigratedFromGit/Source/GLTFExporter/Private/Json/GLTFJsonBackdrop.h
new file mode 100644
--- /dev/null
+++ b/Engine/Plugins/Enterprise/GLTFExporter/MigratedFromGit/Source/GLTFExporter/Private/Json/GLTFJsonBackdrop.h
@@ -0,0 +1,61 @@
+// Copyright Epic Games, Inc. All Rights Reserved.
+
+#pragma once
+
+#include "Json/GLTFJsonIndex.h"
+#include "Json/GLTFJsonVector3.h"
+#include "Json/GLTFJsonUtility.h"
+
+struct FGLTFJsonBackdrop
+{
+	FString Name;
+
+	FGLTFJsonMeshIndex Mesh;
+	FGLTFJsonTextureIndex Cubemap[6];
+
+	float Intensity;
+	float Size;
+
+	FGLTFJsonVector3 ProjectionCenter;
+
+	float LightingDistanceFactor;
+	bool UseCameraProjection;
+
+	FGLTFJsonBackdrop()
+		: Intensity(1)
+		, Size(1)
+		, ProjectionCenter(FGLTFJsonVector3::Zero)
+		, LightingDistanceFactor(0)
+		, UseCameraProjection(false)
+	{
+	}
+
+	template <class CharType = TCHAR, class PrintPolicy = TPrettyJsonPrintPolicy<CharType>>
+	void WriteObject(TJsonWriter<CharType, PrintPolicy>& JsonWriter, FGLTFJsonExtensions& Extensions) const
+	{
+		JsonWriter.WriteObjectStart();
+
+		if (!Name.IsEmpty())
+		{
+			JsonWriter.WriteValue(TEXT("name"), Name);
+		}
+
+		if (Mesh != INDEX_NONE)
+		{
+			JsonWriter.WriteValue(TEXT("mesh"), Mesh);
+		}
+
+		FGLTFJsonUtility::WriteFixedArray(JsonWriter, TEXT("cubemap"), Cubemap);
+
+		FGLTFJsonUtility::WriteExactValue(JsonWriter, TEXT("intensity"), Intensity);
+		FGLTFJsonUtility::WriteExactValue(JsonWriter, TEXT("size"), Size);
+
+		JsonWriter.WriteIdentifierPrefix(TEXT("projectionCenter"));
+		ProjectionCenter.WriteArray(JsonWriter);
+
+		FGLTFJsonUtility::WriteExactValue(JsonWriter, TEXT("lightingDistanceFactor"), LightingDistanceFactor);
+		JsonWriter.WriteValue(TEXT("useCameraProjection"), UseCameraProjection);
+
+		JsonWriter.WriteObjectEnd();
+	}
+};
diff --git a/Engine/Plugins/Enterprise/GLTFExporter/MigratedFromGit/Source/GLTFExporter/Private/Json/GLTFJsonIndex.h b/Engine/Plugins/Enterprise/GLTFExporter/MigratedFromGit/Source/GLTFExporter/Private/Json/GLTFJsonIndex.h
--- a/Engine/Plugins/Enterprise/GLTFExporter/MigratedFromGit/Source/GLTFExporter/Private/Json/GLTFJsonIndex.h
+++ b/Engine/Plugins/Enterprise/GLTFExporter/MigratedFromGit/Source/GLTFExporter/Private/Json/GLTFJsonIndex.h
@@ -64,5 +64,6 @@ struct FGLTFJsonSamplerIndex : FGLTFJsonIndex<FGLTFJsonSamplerIndex> { using FGL
 struct FGLTFJsonSceneIndex : FGLTFJsonIndex<FGLTFJsonSceneIndex> { using FGLTFJsonIndex::FGLTFJsonIndex; };
 struct FGLTFJsonSkinIndex : FGLTFJsonIndex<FGLTFJsonSkinIndex> { using FGLTFJsonIndex::FGLTFJsonIndex; };
 struct FGLTFJsonTextureIndex : FGLTFJsonIndex<FGLTFJsonTextureIndex> { using FGLTFJsonIndex::FGLTFJsonIndex; };
-struct FGLTFJsonLightMapIndex : FGLTFJsonIndex<FGLTFJsonLightMapIndex> { using FGLTFJsonIndex::FGLTFJsonIndex; };
+struct FGLTFJsonBackdropIndex : FGLTFJsonIndex<FGLTFJsonBackdropIndex> { using FGLTFJsonIndex::FGLTFJsonIndex; };
 struct FGLTFJsonLevelVariantSetsIndex : FGLTFJsonIndex<FGLTFJsonLevelVariantSetsIndex> { using FGLTFJsonIndex::FGLTFJsonIndex; };
+struct FGLTFJsonLightMapIndex : FGLTFJsonIndex<FGLTFJsonLightMapIndex> { using FGLTFJsonIndex::FGLTFJsonIndex; };
diff --git a/Engine/Plugins/Enterprise/GLTFExporter/MigratedFromGit/Source/GLTFExporter/Private/Json/GLTFJsonNode.h b/Engine/Plugins/Enterprise/GLTFExporter/MigratedFromGit/Source/GLTFExporter/Private/Json/GLTFJsonNode.h
--- a/Engine/Plugins/Enterprise/GLTFExporter/MigratedFromGit/Source/GLTFExporter/Private/Json/GLTFJsonNode.h
+++ b/Engine/Plugins/Enterprise/GLTFExporter/MigratedFromGit/Source/GLTFExporter/Private/Json/GLTFJsonNode.h
@@ -18,6 +18,7 @@ struct FGLTFJsonNode
 	FGLTFJsonCameraIndex    Camera;
 	FGLTFJsonSkinIndex      Skin;
 	FGLTFJsonMeshIndex      Mesh;
+	FGLTFJsonBackdropIndex  Backdrop;
 	FGLTFJsonLightMapIndex  LightMap;
 
 	TArray<FGLTFJsonNodeIndex> Children;
@@ -72,12 +73,20 @@ struct FGLTFJsonNode
 			JsonWriter.WriteValue(TEXT("mesh"), Mesh);
 		}
 
-		const bool bWriteExtensions = LightMap != INDEX_NONE;
-
-		if (bWriteExtensions)
+		if (Backdrop != INDEX_NONE || LightMap != INDEX_NONE)
 		{
 			JsonWriter.WriteObjectStart(TEXT("extensions"));
 
+			if (Backdrop != INDEX_NONE)
+			{
+				const EGLTFJsonExtension Extension = EGLTFJsonExtension::EPIC_HDRIBackdrops;
+				Extensions.Used.Add(Extension);
+
+				JsonWriter.WriteObjectStart(FGLTFJsonUtility::ToString(Extension));
+				JsonWriter.WriteValue(TEXT("backdrop"), Backdrop);
+				JsonWriter.WriteObjectEnd();
+			}
+
 			if (LightMap != INDEX_NONE)
 			{
 				const EGLTFJsonExtension Extension = EGLTFJsonExtension::EPIC_LightmapTextures;
diff --git a/Engine/Plugins/Enterprise/GLTFExporter/MigratedFromGit/Source/GLTFExporter/Private/Json/GLTFJsonRoot.h b/Engine/Plugins/Enterprise/GLTFExporter/MigratedFromGit/Source/GLTFExporter/Private/Json/GLTFJsonRoot.h
--- a/Engine/Plugins/Enterprise/GLTFExporter/MigratedFromGit/Source/GLTFExporter/Private/Json/GLTFJsonRoot.h
+++ b/Engine/Plugins/Enterprise/GLTFExporter/MigratedFromGit/Source/GLTFExporter/Private/Json/GLTFJsonRoot.h
@@ -14,8 +14,9 @@
 #include "Json/GLTFJsonSampler.h"
 #include "Json/GLTFJsonScene.h"
 #include "Json/GLTFJsonTexture.h"
-#include "Json/GLTFJsonLightMap.h"
+#include "Json/GLTFJsonBackdrop.h"
 #include "Json/GLTFJsonLevelVariantSets.h"
+#include "Json/GLTFJsonLightMap.h"
 #include "Policies/CondensedJsonPrintPolicy.h"
 
 
@@ -66,8 +67,8 @@ struct FGLTFJsonRoot
 	TArray<FGLTFJsonSampler>    Samplers;
 	TArray<FGLTFJsonScene>      Scenes;
 	TArray<FGLTFJsonTexture>    Textures;
+	TArray<FGLTFJsonBackdrop>   Backdrops;
 	TArray<FGLTFJsonLightMap>   LightMaps;
-
 	TArray<FGLTFJsonLevelVariantSets> LevelVariantSets;
 
 	template <class CharType = TCHAR, class PrintPolicy = TPrettyJsonPrintPolicy<CharType>>
@@ -96,17 +97,17 @@ struct FGLTFJsonRoot
 		FGLTFJsonUtility::WriteObjectArray(JsonWriter, TEXT("scenes"), Scenes, AllExtensions);
 		FGLTFJsonUtility::WriteObjectArray(JsonWriter, TEXT("textures"), Textures, AllExtensions);
 
-		if (LightMaps.Num() > 0 || LevelVariantSets.Num() > 0 /* TODO: add more extension support */)
+		if (Backdrops.Num() > 0 || LevelVariantSets.Num() > 0 || LightMaps.Num() > 0 /* TODO: add more extension support */)
 		{
 			JsonWriter.WriteObjectStart(TEXT("extensions"));
 
-			if (LightMaps.Num() > 0)
+			if (Backdrops.Num() > 0)
 			{
-				const EGLTFJsonExtension Extension = EGLTFJsonExtension::EPIC_LightmapTextures;
+				const EGLTFJsonExtension Extension = EGLTFJsonExtension::EPIC_HDRIBackdrops;
 				AllExtensions.Used.Add(Extension);
 
 				JsonWriter.WriteObjectStart(FGLTFJsonUtility::ToString(Extension));
-				FGLTFJsonUtility::WriteObjectArray(JsonWriter, TEXT("lightmaps"), LightMaps, AllExtensions);
+				FGLTFJsonUtility::WriteObjectArray(JsonWriter, TEXT("backdrops"), Backdrops, AllExtensions);
 				JsonWriter.WriteObjectEnd();
 			}
 
@@ -120,6 +121,16 @@ struct FGLTFJsonRoot
 				JsonWriter.WriteObjectEnd();
 			}
 
+			if (LightMaps.Num() > 0)
+			{
+				const EGLTFJsonExtension Extension = EGLTFJsonExtension::EPIC_LightmapTextures;
+				AllExtensions.Used.Add(Extension);
+
+				JsonWriter.WriteObjectStart(FGLTFJsonUtility::ToString(Extension));
+				FGLTFJsonUtility::WriteObjectArray(JsonWriter, TEXT("lightmaps"), LightMaps, AllExtensions);
+				JsonWriter.WriteObjectEnd();
+			}
+
 			JsonWriter.WriteObjectEnd();
 		}
 
diff --git a/Engine/Plugins/Enterprise/GLTFExporter/MigratedFromGit/Source/GLTFExporter/Private/Json/GLTFJsonScene.h b/Engine/Plugins/Enterprise/GLTFExporter/MigratedFromGit/Source/GLTFExporter/Private/Json/GLTFJsonScene.h
--- a/Engine/Plugins/Enterprise/GLTFExporter/MigratedFromGit/Source/GLTFExporter/Private/Json/GLTFJsonScene.h
+++ b/Engine/Plugins/Enterprise/GLTFExporter/MigratedFromGit/Source/GLTFExporter/Private/Json/GLTFJsonScene.h
@@ -27,23 +27,18 @@ struct FGLTFJsonScene
 			JsonWriter.WriteValue(TEXT("nodes"), Nodes);
 		}
 
-		const bool bWriteExtensions = LevelVariantSets.Num() > 0;
-
-		if (bWriteExtensions)
+		if (LevelVariantSets.Num() > 0)
 		{
-			JsonWriter.WriteObjectStart(TEXT("extensions"));
+			const EGLTFJsonExtension Extension = EGLTFJsonExtension::EPIC_LevelVariantSets;
+			Extensions.Used.Add(Extension);
 
-			if (LevelVariantSets.Num() > 0)
-			{
-				const EGLTFJsonExtension Extension = EGLTFJsonExtension::EPIC_LevelVariantSets;
-				Extensions.Used.Add(Extension);
+			JsonWriter.WriteObjectStart(TEXT("extensions"));
+			JsonWriter.WriteObjectStart(FGLTFJsonUtility::ToString(Extension));
 
-				JsonWriter.WriteObjectStart(FGLTFJsonUtility::ToString(Extension));
-				JsonWriter.WriteValue(TEXT("levelVariantSets"), LevelVariantSets);
-				JsonWriter.WriteObjectEnd();
-			}
+			JsonWriter.WriteValue(TEXT("levelVariantSets"), LevelVariantSets);
 
 			JsonWriter.WriteObjectEnd();
+			JsonWriter.WriteObjectEnd();
 		}
 
 		JsonWriter.WriteObjectEnd();
diff --git a/Engine/Plugins/Enterprise/GLTFExporter/MigratedFromGit/Source/GLTFExporter/Private/Json/GLTFJsonUtility.h b/Engine/Plugins/Enterprise/GLTFExporter/MigratedFromGit/Source/GLTFExporter/Private/Json/GLTFJsonUtility.h
--- a/Engine/Plugins/Enterprise/GLTFExporter/MigratedFromGit/Source/GLTFExporter/Private/Json/GLTFJsonUtility.h
+++ b/Engine/Plugins/Enterprise/GLTFExporter/MigratedFromGit/Source/GLTFExporter/Private/Json/GLTFJsonUtility.h
@@ -133,4 +133,15 @@ struct FGLTFJsonUtility
 			JsonWriter.WriteArrayEnd();
 		}
 	}
+
+	template <class WriterType, class ElementType, SIZE_T Size>
+    static void WriteFixedArray(WriterType& JsonWriter, const FString& Identifier, const ElementType (&Array)[Size])
+	{
+		JsonWriter.WriteArrayStart(Identifier);
+		for (const ElementType& Element : Array)
+		{
+			JsonWriter.WriteValue(Element);
+		}
+		JsonWriter.WriteArrayEnd();
+	}
 };
