[GLTFExporter] feat(backdrop): add converters for hdri backdrops

commit babd2de97278b30307deffe93f60f8636cab6957
Date: Wed, 8 Jul 2020 12:52:14 +0200

    feat(backdrop): add converters for hdri backdrops
    (cherry picked from commit cdd87b98e99c625a40fa3f44a0c6a13e5b84fdad​)

#skipci

diff --git a/Engine/Plugins/Enterprise/GLTFExporter/MigratedFromGit/Source/GLTFExporter/Private/Builders/GLTFConvertBuilder.cpp b/Engine/Plugins/Enterprise/GLTFExporter/MigratedFromGit/Source/GLTFExporter/Private/Builders/GLTFConvertBuilder.cpp
--- a/Engine/Plugins/Enterprise/GLTFExporter/MigratedFromGit/Source/GLTFExporter/Private/Builders/GLTFConvertBuilder.cpp
+++ b/Engine/Plugins/Enterprise/GLTFExporter/MigratedFromGit/Source/GLTFExporter/Private/Builders/GLTFConvertBuilder.cpp
@@ -1,7 +1,7 @@
 // Copyright Epic Games, Inc. All Rights Reserved.
 
 #include "Builders/GLTFConvertBuilder.h"
-#include "GLTFBuilderUtility.h"
+#include "Builders/GLTFBuilderUtility.h"
 
 FGLTFConvertBuilder::FGLTFConvertBuilder(const UGLTFExportOptions* ExportOptions, bool bSelectedActorsOnly)
 	: FGLTFImageBuilder(ExportOptions)
@@ -164,16 +164,6 @@ FGLTFJsonTextureIndex FGLTFConvertBuilder::GetOrAddTexture(const UTextureRenderT
 	return TextureRenderTargetCubeConverter.GetOrAdd(*this, DesiredName.IsEmpty() ? Texture->GetName() : DesiredName, Texture, CubeFace);
 }
 
-FGLTFJsonLightMapIndex FGLTFConvertBuilder::GetOrAddLightMap(const UStaticMeshComponent* StaticMeshComponent, const FString& DesiredName)
-{
-	if (StaticMeshComponent == nullptr)
-	{
-		return FGLTFJsonLightMapIndex(INDEX_NONE);
-	}
-
-	return LightMapConverter.GetOrAdd(*this, DesiredName, StaticMeshComponent);
-}
-
 FGLTFJsonNodeIndex FGLTFConvertBuilder::GetOrAddNode(const USceneComponent* SceneComponent, const FString& DesiredName)
 {
 	if (SceneComponent == nullptr)
@@ -214,6 +204,16 @@ FGLTFJsonSceneIndex FGLTFConvertBuilder::GetOrAddScene(const UWorld* World, cons
 	return GetOrAddScene(World->PersistentLevel, DesiredName.IsEmpty() ? World->GetName() : DesiredName);
 }
 
+FGLTFJsonBackdropIndex FGLTFConvertBuilder::GetOrAddBackdrop(const AActor* Actor, const FString& DesiredName)
+{
+	if (Actor == nullptr)
+	{
+		return FGLTFJsonBackdropIndex(INDEX_NONE);
+	}
+
+	return BackdropConverter.GetOrAdd(*this, DesiredName, Actor);
+}
+
 FGLTFJsonLevelVariantSetsIndex FGLTFConvertBuilder::GetOrAddLevelVariantSets(const ALevelVariantSetsActor* LevelVariantSetsActor, const FString& DesiredName)
 {
 	if (LevelVariantSetsActor == nullptr)
@@ -223,3 +223,13 @@ FGLTFJsonLevelVariantSetsIndex FGLTFConvertBuilder::GetOrAddLevelVariantSets(con
 
 	return LevelVariantSetsConverter.GetOrAdd(*this, DesiredName, LevelVariantSetsActor);
 }
+
+FGLTFJsonLightMapIndex FGLTFConvertBuilder::GetOrAddLightMap(const UStaticMeshComponent* StaticMeshComponent, const FString& DesiredName)
+{
+	if (StaticMeshComponent == nullptr)
+	{
+		return FGLTFJsonLightMapIndex(INDEX_NONE);
+	}
+
+	return LightMapConverter.GetOrAdd(*this, DesiredName, StaticMeshComponent);
+}
diff --git a/Engine/Plugins/Enterprise/GLTFExporter/MigratedFromGit/Source/GLTFExporter/Private/Builders/GLTFConvertBuilder.h b/Engine/Plugins/Enterprise/GLTFExporter/MigratedFromGit/Source/GLTFExporter/Private/Builders/GLTFConvertBuilder.h
--- a/Engine/Plugins/Enterprise/GLTFExporter/MigratedFromGit/Source/GLTFExporter/Private/Builders/GLTFConvertBuilder.h
+++ b/Engine/Plugins/Enterprise/GLTFExporter/MigratedFromGit/Source/GLTFExporter/Private/Builders/GLTFConvertBuilder.h
@@ -8,8 +8,9 @@
 #include "Converters/GLTFMaterialConverters.h"
 #include "Converters/GLTFTextureConverters.h"
 #include "Converters/GLTFLevelConverters.h"
-#include "Converters/GLTFLightMapConverters.h"
+#include "Converters/GLTFBackdropConverters.h"
 #include "Converters/GLTFVariantSetConverters.h"
+#include "Converters/GLTFLightMapConverters.h"
 
 class FGLTFConvertBuilder : public FGLTFImageBuilder
 {
@@ -36,13 +37,15 @@ public:
 	FGLTFJsonTextureIndex GetOrAddTexture(const UTextureCube* Texture, ECubeFace CubeFace, const FString& DesiredName = TEXT(""));
 	FGLTFJsonTextureIndex GetOrAddTexture(const UTextureRenderTarget2D* Texture, const FString& DesiredName = TEXT(""));
 	FGLTFJsonTextureIndex GetOrAddTexture(const UTextureRenderTargetCube* Texture, ECubeFace CubeFace, const FString& DesiredName = TEXT(""));
-	FGLTFJsonLightMapIndex GetOrAddLightMap(const UStaticMeshComponent* StaticMeshComponent, const FString& DesiredName = TEXT(""));
 
 	FGLTFJsonNodeIndex GetOrAddNode(const USceneComponent* SceneComponent, const FString& DesiredName = TEXT(""));
 	FGLTFJsonNodeIndex GetOrAddNode(const AActor* Actor, const FString& DesiredName = TEXT(""));
 	FGLTFJsonSceneIndex GetOrAddScene(const ULevel* Level, const FString& DesiredName = TEXT(""));
 	FGLTFJsonSceneIndex GetOrAddScene(const UWorld* World, const FString& DesiredName = TEXT(""));
+
+	FGLTFJsonBackdropIndex GetOrAddBackdrop(const AActor* Actor, const FString& DesiredName = TEXT(""));
 	FGLTFJsonLevelVariantSetsIndex GetOrAddLevelVariantSets(const ALevelVariantSetsActor* LevelVariantSetsActor, const FString& DesiredName = TEXT(""));
+	FGLTFJsonLightMapIndex GetOrAddLightMap(const UStaticMeshComponent* StaticMeshComponent, const FString& DesiredName = TEXT(""));
 
 private:
 
@@ -62,10 +65,12 @@ private:
 	FGLTFTextureCubeConverter TextureCubeConverter;
 	FGLTFTextureRenderTarget2DConverter TextureRenderTarget2DConverter;
 	FGLTFTextureRenderTargetCubeConverter TextureRenderTargetCubeConverter;
-	FGLTFLightMapConverter LightMapConverter;
 
 	FGLTFSceneComponentConverter SceneComponentConverter;
 	FGLTFActorConverter ActorConverter;
 	FGLTFLevelConverter LevelConverter;
+
+	FGLTFBackdropConverter BackdropConverter;
 	FGLTFLevelVariantSetsConverter LevelVariantSetsConverter;
+	FGLTFLightMapConverter LightMapConverter;
 };
diff --git a/Engine/Plugins/Enterprise/GLTFExporter/MigratedFromGit/Source/GLTFExporter/Private/Converters/GLTFBackdropConverters.cpp b/Engine/Plugins/Enterprise/GLTFExporter/MigratedFromGit/Source/GLTFExporter/Private/Converters/GLTFBackdropConverters.cpp
new file mode 100644
--- /dev/null
+++ b/Engine/Plugins/Enterprise/GLTFExporter/MigratedFromGit/Source/GLTFExporter/Private/Converters/GLTFBackdropConverters.cpp
@@ -0,0 +1,10 @@
+// Copyright Epic Games, Inc. All Rights Reserved.
+
+#include "Converters/GLTFBackdropConverters.h"
+#include "Builders/GLTFContainerBuilder.h"
+
+FGLTFJsonBackdropIndex FGLTFBackdropConverter::Add(FGLTFConvertBuilder& Builder, const FString& Name, const AActor* Actor)
+{
+	// TODO: implement support
+	return FGLTFJsonBackdropIndex(INDEX_NONE);
+}
diff --git a/Engine/Plugins/Enterprise/GLTFExporter/MigratedFromGit/Source/GLTFExporter/Private/Converters/GLTFBackdropConverters.h b/Engine/Plugins/Enterprise/GLTFExporter/MigratedFromGit/Source/GLTFExporter/Private/Converters/GLTFBackdropConverters.h
new file mode 100644
--- /dev/null
+++ b/Engine/Plugins/Enterprise/GLTFExporter/MigratedFromGit/Source/GLTFExporter/Private/Converters/GLTFBackdropConverters.h
@@ -0,0 +1,12 @@
+// Copyright Epic Games, Inc. All Rights Reserved.
+
+#pragma once
+
+#include "Json/GLTFJsonIndex.h"
+#include "Converters/GLTFConverter.h"
+#include "Engine.h"
+
+class FGLTFBackdropConverter final : public TGLTFConverter<FGLTFJsonBackdropIndex, const AActor*>
+{
+	FGLTFJsonBackdropIndex Add(FGLTFConvertBuilder& Builder, const FString& Name, const AActor* Actor) override;
+};
