[GLTFExporter] feat(backdrop): implement conversion of hdri backdrops actors

commit b9ca035e00706c66660aed14d718c258cce923ee
Date: Wed, 8 Jul 2020 12:53:21 +0200

    feat(backdrop): implement conversion of hdri backdrops actors
    (cherry picked from commit 17d227db1a655729a9f78938efd5f63445933aa3​)

#skipci

diff --git a/Engine/Plugins/Enterprise/GLTFExporter/MigratedFromGit/Source/GLTFExporter/Private/Converters/GLTFActorUtility.h b/Engine/Plugins/Enterprise/GLTFExporter/MigratedFromGit/Source/GLTFExporter/Private/Converters/GLTFActorUtility.h
--- a/Engine/Plugins/Enterprise/GLTFExporter/MigratedFromGit/Source/GLTFExporter/Private/Converters/GLTFActorUtility.h
+++ b/Engine/Plugins/Enterprise/GLTFExporter/MigratedFromGit/Source/GLTFExporter/Private/Converters/GLTFActorUtility.h
@@ -13,4 +13,23 @@ struct FGLTFActorUtility
 	static bool IsSkySphereBlueprint(const UBlueprint* Blueprint);
 
 	static bool IsHDRIBackdropBlueprint(const UBlueprint* Blueprint);
+
+	template <class ValueType>
+    static bool TryGetPropertyValue(const UObject* Object, const TCHAR* PropertyName, ValueType& Value)
+	{
+		FProperty* Property = Object->GetClass()->FindPropertyByName(PropertyName);
+		if (Property == nullptr)
+		{
+			return false;
+		}
+
+		const ValueType* ValuePtr = Property->ContainerPtrToValuePtr<ValueType>(Object);
+		if (ValuePtr == nullptr)
+		{
+			return false;
+		}
+
+		Value = *ValuePtr;
+		return true;
+	}
 };
diff --git a/Engine/Plugins/Enterprise/GLTFExporter/MigratedFromGit/Source/GLTFExporter/Private/Converters/GLTFBackdropConverters.cpp b/Engine/Plugins/Enterprise/GLTFExporter/MigratedFromGit/Source/GLTFExporter/Private/Converters/GLTFBackdropConverters.cpp
--- a/Engine/Plugins/Enterprise/GLTFExporter/MigratedFromGit/Source/GLTFExporter/Private/Converters/GLTFBackdropConverters.cpp
+++ b/Engine/Plugins/Enterprise/GLTFExporter/MigratedFromGit/Source/GLTFExporter/Private/Converters/GLTFBackdropConverters.cpp
@@ -1,10 +1,100 @@
 // Copyright Epic Games, Inc. All Rights Reserved.
 
 #include "Converters/GLTFBackdropConverters.h"
+#include "Converters/GLTFConverterUtility.h"
+#include "Converters/GLTFActorUtility.h"
 #include "Builders/GLTFContainerBuilder.h"
 
 FGLTFJsonBackdropIndex FGLTFBackdropConverter::Add(FGLTFConvertBuilder& Builder, const FString& Name, const AActor* Actor)
 {
-	// TODO: implement support
-	return FGLTFJsonBackdropIndex(INDEX_NONE);
+	const UBlueprint* Blueprint = FGLTFActorUtility::GetBlueprintFromActor(Actor);
+	if (!FGLTFActorUtility::IsHDRIBackdropBlueprint(Blueprint))
+	{
+		return FGLTFJsonBackdropIndex(INDEX_NONE);
+	}
+
+	FGLTFJsonBackdrop JsonBackdrop;
+	JsonBackdrop.Name = Name;
+
+	const UStaticMesh* Mesh;
+	if (FGLTFActorUtility::TryGetPropertyValue(Actor, TEXT("Mesh"), Mesh))
+	{
+		// TODO: avoid exporting the mesh material (use gltf default material)
+		JsonBackdrop.Mesh = Builder.GetOrAddMesh(Mesh);
+	}
+	else
+	{
+		// TODO: report error
+	}
+
+	const UTextureCube* Cubemap;
+	if (FGLTFActorUtility::TryGetPropertyValue(Actor, TEXT("Cubemap"), Cubemap))
+	{
+		// TODO: add a proper custom gltf extension with its own converters for cubemaps
+
+		for (int32 CubeFaceIndex = 0; CubeFaceIndex < CubeFace_MAX; ++CubeFaceIndex)
+		{
+			const ECubeFace CubeFace = static_cast<ECubeFace>(CubeFaceIndex);
+			const EGLTFJsonCubeFace JsonCubeFace = FGLTFConverterUtility::ConvertCubeFace(CubeFace);
+			const int32 JsonCubeFaceIndex = static_cast<int32>(JsonCubeFace);
+
+			JsonBackdrop.Cubemap[JsonCubeFaceIndex] = Builder.GetOrAddTexture(Cubemap, CubeFace, Cubemap->GetName() + TEXT("_") + FGLTFJsonUtility::ToString(JsonCubeFace));
+		}
+	}
+	else
+	{
+        // TODO: report error
+    }
+
+	float Intensity;
+	if (FGLTFActorUtility::TryGetPropertyValue(Actor, TEXT("Intensity"), Intensity))
+	{
+		JsonBackdrop.Intensity = Intensity;
+	}
+	else
+	{
+		// TODO: report error
+	}
+
+	float Size;
+	if (FGLTFActorUtility::TryGetPropertyValue(Actor, TEXT("Size"), Size))
+	{
+		JsonBackdrop.Size = Size;
+	}
+	else
+	{
+		// TODO: report error
+	}
+
+	FVector ProjectionCenter;
+	if (FGLTFActorUtility::TryGetPropertyValue(Actor, TEXT("ProjectionCenter"), ProjectionCenter))
+	{
+		JsonBackdrop.ProjectionCenter = FGLTFConverterUtility::ConvertPosition(ProjectionCenter);
+	}
+	else
+	{
+		// TODO: report error
+	}
+
+	float LightingDistanceFactor;
+	if (FGLTFActorUtility::TryGetPropertyValue(Actor, TEXT("LightingDistanceFactor"), LightingDistanceFactor))
+	{
+		JsonBackdrop.LightingDistanceFactor = LightingDistanceFactor;
+	}
+	else
+	{
+		// TODO: report error
+	}
+
+	bool UseCameraProjection;
+	if (FGLTFActorUtility::TryGetPropertyValue(Actor, TEXT("UseCameraProjection"), UseCameraProjection))
+	{
+		JsonBackdrop.UseCameraProjection = UseCameraProjection;
+	}
+	else
+	{
+		// TODO: report error
+	}
+
+	return Builder.AddBackdrop(JsonBackdrop);
 }
diff --git a/Engine/Plugins/Enterprise/GLTFExporter/MigratedFromGit/Source/GLTFExporter/Private/Converters/GLTFLevelConverters.cpp b/Engine/Plugins/Enterprise/GLTFExporter/MigratedFromGit/Source/GLTFExporter/Private/Converters/GLTFLevelConverters.cpp
--- a/Engine/Plugins/Enterprise/GLTFExporter/MigratedFromGit/Source/GLTFExporter/Private/Converters/GLTFLevelConverters.cpp
+++ b/Engine/Plugins/Enterprise/GLTFExporter/MigratedFromGit/Source/GLTFExporter/Private/Converters/GLTFLevelConverters.cpp
@@ -51,6 +51,9 @@ FGLTFJsonNodeIndex FGLTFActorConverter::Add(FGLTFConvertBuilder& Builder, const
 		return FGLTFJsonNodeIndex(INDEX_NONE);
 	}
 
+	const USceneComponent* RootComponent = Actor->GetRootComponent();
+	const FGLTFJsonNodeIndex RootNodeIndex = Builder.GetOrAddNode(RootComponent);
+
 	const UBlueprint* Blueprint = FGLTFActorUtility::GetBlueprintFromActor(Actor);
 	if (FGLTFActorUtility::IsSkySphereBlueprint(Blueprint))
 	{
@@ -59,7 +62,8 @@ FGLTFJsonNodeIndex FGLTFActorConverter::Add(FGLTFConvertBuilder& Builder, const
 	}
 	else if (FGLTFActorUtility::IsHDRIBackdropBlueprint(Blueprint))
 	{
-		// TODO: add support for backdrop
+		FGLTFJsonNode& RootNode = Builder.GetNode(RootNodeIndex);
+		RootNode.Backdrop = Builder.GetOrAddBackdrop(Actor);
 	}
 	else
 	{
@@ -73,8 +77,7 @@ FGLTFJsonNodeIndex FGLTFActorConverter::Add(FGLTFConvertBuilder& Builder, const
 		}
 	}
 
-	const USceneComponent* RootComponent = Actor->GetRootComponent();
-	return Builder.GetOrAddNode(RootComponent);
+	return RootNodeIndex;
 }
 
 FGLTFJsonSceneIndex FGLTFLevelConverter::Add(FGLTFConvertBuilder& Builder, const FString& Name, const ULevel* Level)
