[GLTFExporter] refactor: add a few comments and some minor refactoring

commit d71d267307401f76880c4e86de8c6806a63c1e9a
Date: Wed, 8 Jul 2020 13:28:21 +0200

    refactor: add a few comments and some minor refactoring
    (cherry picked from commit 9e98c7a2d529f40a39383bb43817085737e57fa6​)

#skipci

diff --git a/Engine/Plugins/Enterprise/GLTFExporter/MigratedFromGit/Source/GLTFExporter/Private/Converters/GLTFTextureUtility.cpp b/Engine/Plugins/Enterprise/GLTFExporter/MigratedFromGit/Source/GLTFExporter/Private/Converters/GLTFTextureUtility.cpp
--- a/Engine/Plugins/Enterprise/GLTFExporter/MigratedFromGit/Source/GLTFExporter/Private/Converters/GLTFTextureUtility.cpp
+++ b/Engine/Plugins/Enterprise/GLTFExporter/MigratedFromGit/Source/GLTFExporter/Private/Converters/GLTFTextureUtility.cpp
@@ -170,6 +170,7 @@ UTexture2D* FGLTFTextureUtility::CreateTextureFromCubeFace(const UTextureCube* T
 
 	if (TextureCube->PlatformData->Mips[0].BulkData.GetBulkDataSize() == 0)
 	{
+		// TODO: is this correct handling?
 		const_cast<UTextureCube*>(TextureCube)->ForceRebuildPlatformData();
 		if (TextureCube->PlatformData->Mips[0].BulkData.GetBulkDataSize() == 0)
 		{
@@ -177,13 +178,13 @@ UTexture2D* FGLTFTextureUtility::CreateTextureFromCubeFace(const UTextureCube* T
 		}
 	}
 
-	const FTexture2DMipMap& Mip = TextureCube->PlatformData->Mips[0];
-	const int64 MipSize = Mip.BulkData.GetBulkDataSize() / 6;
+	const FByteBulkData& BulkData = TextureCube->PlatformData->Mips[0].BulkData;
+	const int64 MipSize = BulkData.GetBulkDataSize() / 6;
 
-	const void* MipDataPtr = Mip.BulkData.LockReadOnly();
+	const void* MipDataPtr = BulkData.LockReadOnly();
 	const void* FaceDataPtr =  static_cast<const uint8*>(MipDataPtr) + MipSize * CubeFace;
 	UTexture2D* FaceTexture = CreateTransientTexture(FaceDataPtr, MipSize, Size, Format, TextureCube->SRGB);
-	Mip.BulkData.Unlock();
+	BulkData.Unlock();
 
 	return FaceTexture;
 }
@@ -246,6 +247,8 @@ bool FGLTFTextureUtility::ReadEncodedPixels(const UTextureRenderTarget2D* InRend
 
 FColor FGLTFTextureUtility::EncodeRGBM(const FLinearColor& Color, float MaxRange)
 {
+	// Based on PlayCanvas modified RGBM encoding.
+
 	FLinearColor RGBM;
 
 	RGBM.R = FMath::Sqrt(Color.R);
