[GLTFExporter] refactor(backdrop): add warning messages if any property export fails

commit 3dc6d4f468eb7b77ff7f31ec4abe1ed8b3ae2de0
Date: Wed, 8 Jul 2020 13:45:13 +0200

    refactor(backdrop): add warning messages if any property export fails
    (cherry picked from commit f22be0c591ac7d1bc576dffa9204b36b5763daa6​)

#skipci

diff --git a/Engine/Plugins/Enterprise/GLTFExporter/MigratedFromGit/Source/GLTFExporter/Private/Converters/GLTFBackdropConverters.cpp b/Engine/Plugins/Enterprise/GLTFExporter/MigratedFromGit/Source/GLTFExporter/Private/Converters/GLTFBackdropConverters.cpp
--- a/Engine/Plugins/Enterprise/GLTFExporter/MigratedFromGit/Source/GLTFExporter/Private/Converters/GLTFBackdropConverters.cpp
+++ b/Engine/Plugins/Enterprise/GLTFExporter/MigratedFromGit/Source/GLTFExporter/Private/Converters/GLTFBackdropConverters.cpp
@@ -24,7 +24,7 @@ FGLTFJsonBackdropIndex FGLTFBackdropConverter::Add(FGLTFConvertBuilder& Builder,
 	}
 	else
 	{
-		// TODO: report error
+		Builder.AddWarningMessage(FString::Printf(TEXT("Failed to export Mesh for HDRIBackdrop %s"), *Actor->GetName()));
 	}
 
 	const UTextureCube* Cubemap;
@@ -43,7 +43,7 @@ FGLTFJsonBackdropIndex FGLTFBackdropConverter::Add(FGLTFConvertBuilder& Builder,
 	}
 	else
 	{
-		// TODO: report error
+		Builder.AddWarningMessage(FString::Printf(TEXT("Failed to export Cubemap for HDRIBackdrop %s"), *Actor->GetName()));
 	}
 
 	float Intensity;
@@ -53,7 +53,7 @@ FGLTFJsonBackdropIndex FGLTFBackdropConverter::Add(FGLTFConvertBuilder& Builder,
 	}
 	else
 	{
-		// TODO: report error
+		Builder.AddWarningMessage(FString::Printf(TEXT("Failed to export Intensity for HDRIBackdrop %s"), *Actor->GetName()));
 	}
 
 	float Size;
@@ -63,7 +63,7 @@ FGLTFJsonBackdropIndex FGLTFBackdropConverter::Add(FGLTFConvertBuilder& Builder,
 	}
 	else
 	{
-		// TODO: report error
+		Builder.AddWarningMessage(FString::Printf(TEXT("Failed to export Size for HDRIBackdrop %s"), *Actor->GetName()));
 	}
 
 	FVector ProjectionCenter;
@@ -73,7 +73,7 @@ FGLTFJsonBackdropIndex FGLTFBackdropConverter::Add(FGLTFConvertBuilder& Builder,
 	}
 	else
 	{
-		// TODO: report error
+		Builder.AddWarningMessage(FString::Printf(TEXT("Failed to export ProjectionCenter for HDRIBackdrop %s"), *Actor->GetName()));
 	}
 
 	float LightingDistanceFactor;
@@ -83,7 +83,7 @@ FGLTFJsonBackdropIndex FGLTFBackdropConverter::Add(FGLTFConvertBuilder& Builder,
 	}
 	else
 	{
-		// TODO: report error
+		Builder.AddWarningMessage(FString::Printf(TEXT("Failed to export LightingDistanceFactor for HDRIBackdrop %s"), *Actor->GetName()));
 	}
 
 	bool UseCameraProjection;
@@ -93,7 +93,7 @@ FGLTFJsonBackdropIndex FGLTFBackdropConverter::Add(FGLTFConvertBuilder& Builder,
 	}
 	else
 	{
-		// TODO: report error
+		Builder.AddWarningMessage(FString::Printf(TEXT("Failed to export UseCameraProjection for HDRIBackdrop %s"), *Actor->GetName()));
 	}
 
 	return Builder.AddBackdrop(JsonBackdrop);
diff --git a/Engine/Plugins/Enterprise/GLTFExporter/MigratedFromGit/Source/GLTFExporter/Private/Json/GLTFJsonBackdrop.h b/Engine/Plugins/Enterprise/GLTFExporter/MigratedFromGit/Source/GLTFExporter/Private/Json/GLTFJsonBackdrop.h
--- a/Engine/Plugins/Enterprise/GLTFExporter/MigratedFromGit/Source/GLTFExporter/Private/Json/GLTFJsonBackdrop.h
+++ b/Engine/Plugins/Enterprise/GLTFExporter/MigratedFromGit/Source/GLTFExporter/Private/Json/GLTFJsonBackdrop.h
@@ -22,8 +22,8 @@ struct FGLTFJsonBackdrop
 	bool UseCameraProjection;
 
 	FGLTFJsonBackdrop()
-		: Intensity(1)
-		, Size(1)
+		: Intensity(0)
+		, Size(0)
 		, ProjectionCenter(FGLTFJsonVector3::Zero)
 		, LightingDistanceFactor(0)
 		, UseCameraProjection(false)
