[GLTFExporter] feat(variants): add support for material variants and improve structure of the variant converter

commit 4f256584d46473e7ed67023e377b2b4c9d36e924
Date: Thu, 9 Jul 2020 14:49:37 +0200

    feat(variants): add support for material variants and improve structure of the variant converter
    (cherry picked from commit 706c01a26c5bf2afeb8a1bfe2f4874e9c9e09aec​)

#skipci

diff --git a/Engine/Plugins/Enterprise/GLTFExporter/MigratedFromGit/Source/GLTFExporter/Private/Converters/GLTFVariantSetConverters.cpp b/Engine/Plugins/Enterprise/GLTFExporter/MigratedFromGit/Source/GLTFExporter/Private/Converters/GLTFVariantSetConverters.cpp
--- a/Engine/Plugins/Enterprise/GLTFExporter/MigratedFromGit/Source/GLTFExporter/Private/Converters/GLTFVariantSetConverters.cpp
+++ b/Engine/Plugins/Enterprise/GLTFExporter/MigratedFromGit/Source/GLTFExporter/Private/Converters/GLTFVariantSetConverters.cpp
@@ -4,11 +4,25 @@
 #include "Builders/GLTFContainerBuilder.h"
 #include "LevelVariantSetsActor.h"
 #include "VariantObjectBinding.h"
+#include "PropertyValueMaterial.h"
 #include "LevelVariantSets.h"
 #include "PropertyValue.h"
 #include "VariantSet.h"
 #include "Variant.h"
 
+namespace
+{
+	// TODO: replace hack with safe solution by adding proper API access to UPropertyValue
+	class UPropertyValueHack : public UPropertyValue
+	{
+	public:
+		const TArray<FCapturedPropSegment>& GetCapturedPropSegments() const
+		{
+			return CapturedPropSegments;
+		}
+	};
+} // anonymous namespace
+
 FGLTFJsonLevelVariantSetsIndex FGLTFLevelVariantSetsConverter::Add(FGLTFConvertBuilder& Builder, const FString& Name, const ALevelVariantSetsActor* LevelVariantSetsActor)
 {
 	const ULevelVariantSets* LevelVariantSets = const_cast<ALevelVariantSetsActor*>(LevelVariantSetsActor)->GetLevelVariantSets(true);
@@ -27,10 +41,8 @@ FGLTFJsonLevelVariantSetsIndex FGLTFLevelVariantSetsConverter::Add(FGLTFConvertB
 
 		for (const UVariant* Variant: VariantSet->GetVariants())
 		{
-			const FString ThumbnailPrefix = JsonLevelVariantSets.Name + TEXT("_") + JsonVariantSet.Name + TEXT("_");
-
 			FGLTFJsonVariant JsonVariant;
-			if (TryParseJsonVariant(Builder, JsonVariant, Variant, ThumbnailPrefix))
+			if (TryParseVariant(Builder, JsonVariant, Variant))
 			{
 				JsonVariantSet.Variants.Add(JsonVariant);
 			}
@@ -40,6 +52,12 @@ FGLTFJsonLevelVariantSetsIndex FGLTFLevelVariantSetsConverter::Add(FGLTFConvertB
 		{
 			JsonLevelVariantSets.VariantSets.Add(JsonVariantSet);
 		}
+		else
+		{
+			Builder.AddWarningMessage(FString::Printf(
+				TEXT("Variant-set has no supported variants and will be skipped. Context: %s"),
+				*GetLogContext(VariantSet)));
+		}
 	}
 
 	if (JsonLevelVariantSets.VariantSets.Num() == 0)
@@ -50,21 +68,20 @@ FGLTFJsonLevelVariantSetsIndex FGLTFLevelVariantSetsConverter::Add(FGLTFConvertB
 	return Builder.AddLevelVariantSets(JsonLevelVariantSets);
 }
 
-bool FGLTFLevelVariantSetsConverter::TryParseJsonVariant(FGLTFConvertBuilder& Builder, FGLTFJsonVariant& OutVariant, const UVariant* Variant, const FString& ThumbnailPrefix) const
+bool FGLTFLevelVariantSetsConverter::TryParseVariant(FGLTFConvertBuilder& Builder, FGLTFJsonVariant& OutVariant, const UVariant* Variant) const
 {
 	FGLTFJsonVariant JsonVariant;
 
 	for (const UVariantObjectBinding* Binding: Variant->GetBindings())
 	{
-		FGLTFJsonVariantNode JsonVariantNode;
-		if (TryParseJsonVariantNode(Builder, JsonVariantNode, Binding))
-		{
-			JsonVariant.Nodes.Add(JsonVariantNode);
-		}
+		TryParseVariantBinding(Builder, JsonVariant, Binding);
 	}
 
 	if (JsonVariant.Nodes.Num() == 0)
 	{
+		Builder.AddWarningMessage(FString::Printf(
+			TEXT("Variant has no supported bindings and will be skipped. Context: %s"),
+			*GetLogContext(Variant)));
 		return false;
 	}
 
@@ -81,85 +98,170 @@ bool FGLTFLevelVariantSetsConverter::TryParseJsonVariant(FGLTFConvertBuilder& Bu
 	return true;
 }
 
-bool FGLTFLevelVariantSetsConverter::TryParseJsonVariantNode(FGLTFConvertBuilder& Builder, FGLTFJsonVariantNode& OutVariantNode, const UVariantObjectBinding* Binding) const
+bool FGLTFLevelVariantSetsConverter::TryParseVariantBinding(FGLTFConvertBuilder& Builder, FGLTFJsonVariant& OutVariant, const UVariantObjectBinding* Binding) const
 {
-	const UVariant* Variant = const_cast<UVariantObjectBinding*>(Binding)->GetParent();
-	const FString BindingDisplayText = Binding->GetDisplayText().ToString();
-	const FString VariantDisplayText = Variant->GetDisplayText().ToString();
-
-	const AActor* Actor = Cast<AActor>(Binding->GetObject());
-	if (Actor == nullptr)
-	{
-		Builder.AddWarningMessage(FString::Printf(
-			TEXT("Actor '%s' in variant '%s' was not found and will be skipped. Context: %s"),
-			*BindingDisplayText, *VariantDisplayText, *GetLogContext(Variant)));
-		return false;
-	}
-
-	if (Builder.bSelectedActorsOnly && !Actor->IsSelected())
-	{
-		Builder.AddWarningMessage(FString::Printf(
-			TEXT("Actor '%s' in variant '%s' is not selected and will be skipped. Context: %s"),
-			*BindingDisplayText, *VariantDisplayText, *GetLogContext(Variant)));
-		return false;
-	}
-
-	FGLTFJsonVariantNode JsonVariantNode;
-	bool bHasAnyProperty = false;
+	bool bHasParsedAnyProperty = false;
 
 	for (UPropertyValue* Property: Binding->GetCapturedProperties())
 	{
-		if (!Property->Resolve())
+		if (!const_cast<UPropertyValue*>(Property)->Resolve() || !Property->HasRecordedData())
 		{
 			Builder.AddWarningMessage(FString::Printf(
-				TEXT("Property '%s' for actor '%s' in variant '%s' failed to resolve and will be skipped. Context: %s"),
-				*Property->GetFullDisplayString(), *BindingDisplayText, *VariantDisplayText, *GetLogContext(Variant)));
+				TEXT("Property is missing recorded data, it will be skipped. Context: %s"),
+				*GetLogContext(Property)));
 			continue;
 		}
 
-		const FFieldClass* PropertyClass = Property->GetPropertyClass();
-		const FName& PropertyName = Property->GetPropertyName();
+		const bool bIsVisibilityProperty =
+			Property->GetPropertyName() == TEXT("bVisible") &&
+			Property->GetPropertyClass()->IsChildOf(FBoolProperty::StaticClass());
+
+		const bool bIsMaterialProperty = Property->IsA<UPropertyValueMaterial>();
 
-		if (PropertyName == TEXT("bVisible") && PropertyClass->IsChildOf(FBoolProperty::StaticClass()))
+		if (bIsVisibilityProperty)
 		{
-			bool bIsVisible;
-			if (TryGetPropertyValue(Property, bIsVisible))
+			if (TryParseVisibilityPropertyValue(Builder, OutVariant, Property))
 			{
-				JsonVariantNode.bIsVisible = bIsVisible;
-				bHasAnyProperty = true;
+				bHasParsedAnyProperty = true;
 			}
-			else
+		}
+		else if (bIsMaterialProperty)
+		{
+			if (TryParseMaterialPropertyValue(Builder, OutVariant, Property))
 			{
-				Builder.AddWarningMessage(FString::Printf(
-					TEXT("Property '%s' for actor '%s' in variant '%s' has no recorded data and will be skipped. Context: %s"),
-					*Property->GetFullDisplayString(), *BindingDisplayText, *VariantDisplayText, *GetLogContext(Variant)));
+				bHasParsedAnyProperty = true;
 			}
 		}
 		else
 		{
-			// TODO: handle more properties
+			// TODO: add support for more properties
+
+			Builder.AddWarningMessage(FString::Printf(
+				TEXT("Property is not supported and will be skipped. Context: %s"),
+				*GetLogContext(Property)));
 		}
 	}
 
-	if (!bHasAnyProperty)
+	if (!bHasParsedAnyProperty)
 	{
 		Builder.AddWarningMessage(FString::Printf(
-			TEXT("Actor '%s' in variant '%s' has no properties supported by the exporter and will be skipped. Context: %s"),
-			*BindingDisplayText, *VariantDisplayText, *GetLogContext(Variant)));
+			TEXT("Binding has no supported properties and will be skipped. Context: %s"),
+			*GetLogContext(Binding)));
+	}
+
+	return bHasParsedAnyProperty;
+}
+
+bool FGLTFLevelVariantSetsConverter::TryParseVisibilityPropertyValue(FGLTFConvertBuilder& Builder, FGLTFJsonVariant& OutVariant, const UPropertyValue* Property) const
+{
+	if (Property->GetPropertyName() != TEXT("bVisible") ||
+		!Property->GetPropertyClass()->IsChildOf(FBoolProperty::StaticClass()))
+	{
+		Builder.AddErrorMessage(FString::Printf(
+			TEXT("Attempted to parse visibility from an incompatible property. Context: %s"),
+			*GetLogContext(Property)));
 		return false;
 	}
 
-	const FGLTFJsonNodeIndex Node = Builder.GetOrAddNode(Actor);
-	if (Node == INDEX_NONE)
+	const USceneComponent* Target = static_cast<USceneComponent*>(Property->GetPropertyParentContainerAddress());
+	if (Target == nullptr)
+	{
+		Builder.AddWarningMessage(FString::Printf(
+			TEXT("Target object for property is invalid, the property will be skipped. Context: %s"),
+			*GetLogContext(Property)));
+		return false;
+	}
+
+	if (Builder.bSelectedActorsOnly && !Target->IsSelected())
+	{
+		Builder.AddWarningMessage(FString::Printf(
+			TEXT("Target object for property is not selected for export, the property will be skipped. Context: %s"),
+			*GetLogContext(Property)));
+		return false;
+	}
+
+	bool bIsVisible;
+
+	if (!TryGetPropertyValue(const_cast<UPropertyValue*>(Property), bIsVisible))
 	{
 		Builder.AddWarningMessage(FString::Printf(
-			TEXT("Actor '%s' in variant '%s' could not be exported and will be skipped. Context: %s"),
-			*BindingDisplayText, *VariantDisplayText, *GetLogContext(Variant)));
+			TEXT("Failed to parse recorded data for property, it will be skipped. Context: %s"),
+			*GetLogContext(Property)));
 		return false;
 	}
 
-	JsonVariantNode.Node = Node;
-	OutVariantNode = JsonVariantNode;
+	const FGLTFJsonNodeIndex NodeIndex = Builder.GetOrAddNode(Target);
+	FGLTFJsonVariantNodeProperties& NodeProperties = OutVariant.Nodes.FindOrAdd(NodeIndex);
+
+	NodeProperties.Node = NodeIndex;
+	NodeProperties.bIsVisible = bIsVisible;
+	return true;
+}
+
+bool FGLTFLevelVariantSetsConverter::TryParseMaterialPropertyValue(FGLTFConvertBuilder& Builder, FGLTFJsonVariant& OutVariant, const UPropertyValue* Property) const
+{
+	UPropertyValueMaterial* MaterialProperty = Cast<UPropertyValueMaterial>(const_cast<UPropertyValue*>(Property));
+	if (!MaterialProperty)
+	{
+		Builder.AddErrorMessage(FString::Printf(
+			TEXT("Attempted to parse material from an incompatible property. Context: %s"),
+			*GetLogContext(Property)));
+		return false;
+	}
+
+	const USceneComponent* Target = static_cast<USceneComponent*>(Property->GetPropertyParentContainerAddress());
+	if (Target == nullptr)
+	{
+		Builder.AddWarningMessage(FString::Printf(
+			TEXT("Target object for property is invalid, the property will be skipped. Context: %s"),
+			*GetLogContext(Property)));
+		return false;
+	}
+
+	if (Builder.bSelectedActorsOnly && !Target->IsSelected())
+	{
+		Builder.AddWarningMessage(FString::Printf(
+			TEXT("Target object for property is not selected for export, the property will be skipped. Context: %s"),
+			*GetLogContext(Property)));
+		return false;
+	}
+
+	// NOTE: UPropertyValueMaterial::GetMaterial does *not* ensure that the recorded data has been loaded,
+	// so we need to call UProperty::GetRecordedData first to make that happen.
+	const_cast<UPropertyValueMaterial*>(MaterialProperty)->GetRecordedData();
+
+	const UMaterialInterface* Material = MaterialProperty->GetMaterial();
+	if (Material == nullptr)
+	{
+		Builder.AddWarningMessage(FString::Printf(
+			TEXT("Failed to parse recorded data for property, it will be skipped. Context: %s"),
+			*GetLogContext(Property)));
+		return false;
+	}
+
+	const TArray<FCapturedPropSegment>& CapturedPropSegments = reinterpret_cast<UPropertyValueHack*>(MaterialProperty)->GetCapturedPropSegments();
+	const int32 NumPropSegments = CapturedPropSegments.Num();
+
+	if (NumPropSegments < 1)
+	{
+		Builder.AddWarningMessage(FString::Printf(
+			TEXT("Failed to parse element index to apply the material to, the property will be skipped. Context: %s"),
+			*GetLogContext(Property)));
+		return false;
+	}
+
+	const FGLTFJsonMaterialIndex MaterialIndex = Builder.GetOrAddMaterial(Material);
+	const int32 ElementIndex = CapturedPropSegments[NumPropSegments - 1].PropertyIndex;
+
+	FGLTFJsonVariantMaterial VariantMaterial;
+	VariantMaterial.Material = MaterialIndex;
+	VariantMaterial.Index = ElementIndex;
+
+	const FGLTFJsonNodeIndex NodeIndex = Builder.GetOrAddNode(Target);
+	FGLTFJsonVariantNodeProperties& NodeProperties = OutVariant.Nodes.FindOrAdd(NodeIndex);
+
+	NodeProperties.Node = NodeIndex;
+	NodeProperties.Materials.Add(VariantMaterial);
 	return true;
 }
 
@@ -175,10 +277,31 @@ bool FGLTFLevelVariantSetsConverter::TryGetPropertyValue(UPropertyValue* Propert
 	return true;
 }
 
+FString FGLTFLevelVariantSetsConverter::GetLogContext(const UPropertyValue* Property) const
+{
+	const UVariantObjectBinding* Parent = Property->GetParent();
+	return GetLogContext(Parent) + TEXT("/") + Property->GetFullDisplayString();
+}
+
+FString FGLTFLevelVariantSetsConverter::GetLogContext(const UVariantObjectBinding* Binding) const
+{
+	const UVariant* Parent = const_cast<UVariantObjectBinding*>(Binding)->GetParent();
+	return GetLogContext(Parent) + TEXT("/") + Binding->GetDisplayText().ToString();
+}
+
 FString FGLTFLevelVariantSetsConverter::GetLogContext(const UVariant* Variant) const
 {
-	UVariantSet* VariantSet = const_cast<UVariant*>(Variant)->GetParent();
-	ULevelVariantSets* LevelVariantSets = VariantSet->GetParent();
+	const UVariantSet* Parent = const_cast<UVariant*>(Variant)->GetParent();
+	return GetLogContext(Parent) + TEXT("/") + Variant->GetDisplayText().ToString();
+}
 
-	return LevelVariantSets->GetName() + TEXT("/") + VariantSet->GetDisplayText().ToString() + TEXT("/") + Variant->GetDisplayText().ToString();
+FString FGLTFLevelVariantSetsConverter::GetLogContext(const UVariantSet* VariantSet) const
+{
+	const ULevelVariantSets* Parent = const_cast<UVariantSet*>(VariantSet)->GetParent();
+	return GetLogContext(Parent) + TEXT("/") + VariantSet->GetDisplayText().ToString();
+}
+
+FString FGLTFLevelVariantSetsConverter::GetLogContext(const ULevelVariantSets* LevelVariantSets) const
+{
+	return LevelVariantSets->GetName();
 }
diff --git a/Engine/Plugins/Enterprise/GLTFExporter/MigratedFromGit/Source/GLTFExporter/Private/Converters/GLTFVariantSetConverters.h b/Engine/Plugins/Enterprise/GLTFExporter/MigratedFromGit/Source/GLTFExporter/Private/Converters/GLTFVariantSetConverters.h
--- a/Engine/Plugins/Enterprise/GLTFExporter/MigratedFromGit/Source/GLTFExporter/Private/Converters/GLTFVariantSetConverters.h
+++ b/Engine/Plugins/Enterprise/GLTFExporter/MigratedFromGit/Source/GLTFExporter/Private/Converters/GLTFVariantSetConverters.h
@@ -9,16 +9,25 @@
 #include "LevelVariantSetsActor.h"
 #include "Variant.h"
 #include "PropertyValue.h"
+#include "LevelVariantSets.h"
 
 class FGLTFLevelVariantSetsConverter final : public TGLTFConverter<FGLTFJsonLevelVariantSetsIndex, const ALevelVariantSetsActor*>
 {
 	FGLTFJsonLevelVariantSetsIndex Add(FGLTFConvertBuilder& Builder, const FString& Name, const ALevelVariantSetsActor* LevelVariantSetsActor) override;
 
-	bool TryParseJsonVariant(FGLTFConvertBuilder& Builder, FGLTFJsonVariant& OutVariant, const UVariant* Variant, const FString& ThumbnailPrefix = TEXT("")) const;
-	bool TryParseJsonVariantNode(FGLTFConvertBuilder& Builder, FGLTFJsonVariantNode& OutVariantNode, const UVariantObjectBinding* Binding) const;
+	bool TryParseVariant(FGLTFConvertBuilder& Builder, FGLTFJsonVariant& OutVariant, const UVariant* Variant) const;
+	bool TryParseVariantBinding(FGLTFConvertBuilder& Builder, FGLTFJsonVariant& OutVariant, const UVariantObjectBinding* Binding) const;
+	bool TryParseVisibilityPropertyValue(FGLTFConvertBuilder& Builder, FGLTFJsonVariant& OutVariant, const UPropertyValue* Property) const;
+	bool TryParseMaterialPropertyValue(FGLTFConvertBuilder& Builder, FGLTFJsonVariant& OutVariant, const UPropertyValue* Property) const;
+
+	// bool TryParseJsonVariantNode(FGLTFConvertBuilder& Builder, FGLTFJsonVariantNodeProperties& OutVariantNode, const UVariantObjectBinding* Binding) const;
 
 	template<typename T>
 	bool TryGetPropertyValue(UPropertyValue* Property, T& OutValue) const;
 
+	FString GetLogContext(const UPropertyValue* Property) const;
+	FString GetLogContext(const UVariantObjectBinding* Binding) const;
 	FString GetLogContext(const UVariant* Variant) const;
+	FString GetLogContext(const UVariantSet* VariantSet) const;
+	FString GetLogContext(const ULevelVariantSets* LevelVariantSets) const;
 };
diff --git a/Engine/Plugins/Enterprise/GLTFExporter/MigratedFromGit/Source/GLTFExporter/Private/Json/GLTFJsonLevelVariantSets.h b/Engine/Plugins/Enterprise/GLTFExporter/MigratedFromGit/Source/GLTFExporter/Private/Json/GLTFJsonLevelVariantSets.h
--- a/Engine/Plugins/Enterprise/GLTFExporter/MigratedFromGit/Source/GLTFExporter/Private/Json/GLTFJsonLevelVariantSets.h
+++ b/Engine/Plugins/Enterprise/GLTFExporter/MigratedFromGit/Source/GLTFExporter/Private/Json/GLTFJsonLevelVariantSets.h
@@ -8,11 +8,43 @@
 #include "Serialization/JsonSerializer.h"
 #include "Misc/Optional.h"
 
-struct FGLTFJsonVariantNode
+struct FGLTFJsonVariantMaterial
+{
+	FGLTFJsonMaterialIndex Material;
+	int32                  Index;
+
+	FGLTFJsonVariantMaterial()
+		: Material(INDEX_NONE)
+		, Index(INDEX_NONE)
+	{
+	}
+
+	template <class CharType = TCHAR, class PrintPolicy = TPrettyJsonPrintPolicy<CharType>>
+	void WriteObject(TJsonWriter<CharType, PrintPolicy>& JsonWriter, FGLTFJsonExtensions& Extensions) const
+	{
+		JsonWriter.WriteObjectStart();
+
+		if (Material != INDEX_NONE)
+		{
+			JsonWriter.WriteValue(TEXT("material"), Material);
+		}
+
+		if (Index != INDEX_NONE)
+		{
+			JsonWriter.WriteValue(TEXT("index"), Index);
+		}
+
+		JsonWriter.WriteObjectEnd();
+	}
+};
+
+struct FGLTFJsonVariantNodeProperties
 {
 	FGLTFJsonNodeIndex Node;
 	TOptional<bool>    bIsVisible;
 
+	TArray<FGLTFJsonVariantMaterial> Materials;
+
 	template <class CharType = TCHAR, class PrintPolicy = TPrettyJsonPrintPolicy<CharType>>
 	void WriteObject(TJsonWriter<CharType, PrintPolicy>& JsonWriter, FGLTFJsonExtensions& Extensions) const
 	{
@@ -30,6 +62,8 @@ struct FGLTFJsonVariantNode
 			JsonWriter.WriteValue(TEXT("visible"), bIsVisible.GetValue());
 		}
 
+		FGLTFJsonUtility::WriteObjectArray(JsonWriter, TEXT("materials"), Materials, Extensions);
+
 		JsonWriter.WriteObjectEnd();
 		JsonWriter.WriteObjectEnd();
 	}
@@ -37,11 +71,13 @@ struct FGLTFJsonVariantNode
 
 struct FGLTFJsonVariant
 {
+	typedef TMap<FGLTFJsonNodeIndex, FGLTFJsonVariantNodeProperties> FNodeProperties;
+
 	FString Name;
 	bool    bIsActive;
 
-	FGLTFJsonTextureIndex        Thumbnail;
-	TArray<FGLTFJsonVariantNode> Nodes;
+	FGLTFJsonTextureIndex Thumbnail;
+	FNodeProperties       Nodes;
 
 	template <class CharType = TCHAR, class PrintPolicy = TPrettyJsonPrintPolicy<CharType>>
 	void WriteObject(TJsonWriter<CharType, PrintPolicy>& JsonWriter, FGLTFJsonExtensions& Extensions) const
@@ -55,8 +91,17 @@ struct FGLTFJsonVariant
 			JsonWriter.WriteValue(TEXT("thumbnail"), Thumbnail);
 		}
 
-		FGLTFJsonUtility::WriteObjectArray(JsonWriter, TEXT("nodes"), Nodes, Extensions);
+		JsonWriter.WriteArrayStart(TEXT("nodes"));
+
+		for (const auto& DataPair: Nodes)
+		{
+			if (DataPair.Key != INDEX_NONE)
+			{
+				DataPair.Value.WriteObject(JsonWriter, Extensions);
+			}
+		}
 
+		JsonWriter.WriteArrayEnd();
 		JsonWriter.WriteObjectEnd();
 	}
 };
