[GLTFExporter] refactor(variants): remove function definition that was already commented out

commit 0e68563eaeffaf1189a545c96b6e58f132edfe9e
Date: Fri, 10 Jul 2020 15:58:44 +0200

    refactor(variants): remove function definition that was already commented out
    (cherry picked from commit 2850d37669f46bd7298f84349e0bdca805a5b700​)

#skipci

diff --git a/Engine/Plugins/Enterprise/GLTFExporter/MigratedFromGit/Source/GLTFExporter/Private/Converters/GLTFVariantSetConverters.h b/Engine/Plugins/Enterprise/GLTFExporter/MigratedFromGit/Source/GLTFExporter/Private/Converters/GLTFVariantSetConverters.h
--- a/Engine/Plugins/Enterprise/GLTFExporter/MigratedFromGit/Source/GLTFExporter/Private/Converters/GLTFVariantSetConverters.h
+++ b/Engine/Plugins/Enterprise/GLTFExporter/MigratedFromGit/Source/GLTFExporter/Private/Converters/GLTFVariantSetConverters.h
@@ -20,8 +20,6 @@ class FGLTFLevelVariantSetsConverter final : public TGLTFConverter<FGLTFJsonLeve
 	bool TryParseVisibilityPropertyValue(FGLTFConvertBuilder& Builder, FGLTFJsonVariant& OutVariant, const UPropertyValue* Property) const;
 	bool TryParseMaterialPropertyValue(FGLTFConvertBuilder& Builder, FGLTFJsonVariant& OutVariant, const UPropertyValue* Property) const;
 
-	// bool TryParseJsonVariantNode(FGLTFConvertBuilder& Builder, FGLTFJsonVariantNodeProperties& OutVariantNode, const UVariantObjectBinding* Binding) const;
-
 	template<typename T>
 	bool TryGetPropertyValue(UPropertyValue* Property, T& OutValue) const;
 
