[GLTFExporter] feat(json): add json data types for clear coat material extension

commit a23d4a3287a89c6ea2c6f2cd99b32534e733c39e
Date: Thu, 25 Jun 2020 21:00:52 +0200

    feat(json): add json data types for clear coat material extension
    (cherry picked from commit f778357dc0f48e159e713462325819cd1c60415b​)

#skipci

diff --git a/Engine/Plugins/Enterprise/GLTFExporter/MigratedFromGit/Source/GLTFExporter/Private/Json/GLTFJsonMaterial.h b/Engine/Plugins/Enterprise/GLTFExporter/MigratedFromGit/Source/GLTFExporter/Private/Json/GLTFJsonMaterial.h
--- a/Engine/Plugins/Enterprise/GLTFExporter/MigratedFromGit/Source/GLTFExporter/Private/Json/GLTFJsonMaterial.h
+++ b/Engine/Plugins/Enterprise/GLTFExporter/MigratedFromGit/Source/GLTFExporter/Private/Json/GLTFJsonMaterial.h
@@ -147,6 +147,59 @@ struct FGLTFJsonPBRMetallicRoughness
 	}
 };
 
+struct FGLTFJsonClearCoatExtension
+{
+	float ClearCoatFactor;
+	FGLTFJsonTextureInfo ClearCoatTexture;
+
+	float ClearCoatRoughnessFactor;
+	FGLTFJsonTextureInfo ClearCoatRoughnessTexture;
+
+	FGLTFJsonNormalTextureInfo ClearCoatNormalTexture;
+
+	FGLTFJsonClearCoatExtension()
+        : ClearCoatFactor(0)
+        , ClearCoatRoughnessFactor(0)
+	{
+	}
+
+	template <class CharType = TCHAR, class PrintPolicy = TPrettyJsonPrintPolicy<CharType>>
+    void WriteObject(TJsonWriter<CharType, PrintPolicy>& JsonWriter, FGLTFJsonExtensions& Extensions) const
+	{
+		JsonWriter.WriteObjectStart();
+
+		if (ClearCoatFactor != 0)
+		{
+			FGLTFJsonUtility::WriteExactValue(JsonWriter, TEXT("clearcoatFactor"), ClearCoatFactor);
+		}
+
+		if (ClearCoatTexture.Index != INDEX_NONE)
+		{
+			JsonWriter.WriteIdentifierPrefix(TEXT("clearcoatTexture"));
+			ClearCoatTexture.WriteObject(JsonWriter, Extensions);
+		}
+
+		if (ClearCoatRoughnessFactor != 0)
+		{
+			FGLTFJsonUtility::WriteExactValue(JsonWriter, TEXT("clearcoatRoughnessFactor"), ClearCoatRoughnessFactor);
+		}
+
+		if (ClearCoatRoughnessTexture.Index != INDEX_NONE)
+		{
+			JsonWriter.WriteIdentifierPrefix(TEXT("clearcoatRoughnessTexture"));
+			ClearCoatRoughnessTexture.WriteObject(JsonWriter, Extensions);
+		}
+
+		if (ClearCoatNormalTexture.Index != INDEX_NONE)
+		{
+			JsonWriter.WriteIdentifierPrefix(TEXT("clearcoatNormalTexture"));
+			ClearCoatNormalTexture.WriteObject(JsonWriter, Extensions);
+		}
+
+		JsonWriter.WriteObjectEnd();
+	}
+};
+
 struct FGLTFJsonMaterial
 {
 	FString Name;
@@ -166,6 +219,8 @@ struct FGLTFJsonMaterial
 
 	bool DoubleSided;
 
+	FGLTFJsonClearCoatExtension ClearCoat;
+
 	FGLTFJsonMaterial()
 		: ShadingModel(EGLTFJsonShadingModel::Default)
 		, EmissiveFactor(FGLTFJsonColor3::Black)
@@ -248,7 +303,7 @@ struct FGLTFJsonMaterial
 					break;
 
 				case EGLTFJsonShadingModel::ClearCoat:
-					// TODO: add support
+					ClearCoat.WriteObject(JsonWriter, Extensions);
 					break;
 
 				default:
