[GLTFExporter] feat(material): add initial conversion of clear coat factor and roughness properties

commit e66abb41319fb023b219204eaddfa17cd56d066c
Date: Thu, 25 Jun 2020 21:02:31 +0200

    feat(material): add initial conversion of clear coat factor and roughness properties
    (cherry picked from commit 660a15b78b0d63bff8b36af52c0db39d6bc84f15​)

#skipci

diff --git a/Engine/Plugins/Enterprise/GLTFExporter/MigratedFromGit/Source/GLTFExporter/Private/Converters/GLTFMaterialConverters.cpp b/Engine/Plugins/Enterprise/GLTFExporter/MigratedFromGit/Source/GLTFExporter/Private/Converters/GLTFMaterialConverters.cpp
--- a/Engine/Plugins/Enterprise/GLTFExporter/MigratedFromGit/Source/GLTFExporter/Private/Converters/GLTFMaterialConverters.cpp
+++ b/Engine/Plugins/Enterprise/GLTFExporter/MigratedFromGit/Source/GLTFExporter/Private/Converters/GLTFMaterialConverters.cpp
@@ -127,7 +127,30 @@ FGLTFJsonMaterialIndex FGLTFMaterialConverter::Add(FGLTFConvertBuilder& Builder,
 
 			if (JsonMaterial.ShadingModel == EGLTFJsonShadingModel::ClearCoat)
 			{
-				// TODO: add support for clear coat properties
+				if (JsonMaterial.ShadingModel == EGLTFJsonShadingModel::ClearCoat)
+				{
+					if (!TryGetConstantScalar(JsonMaterial.ClearCoat.ClearCoatFactor, MP_CustomData0, Material))
+					{
+						if (!TryGetSourceTexture(Builder, JsonMaterial.ClearCoat.ClearCoatTexture, MP_CustomData0, Material, { RedMask }))
+						{
+							if (!TryGetBakedMaterialProperty(Builder, JsonMaterial.ClearCoat.ClearCoatTexture, JsonMaterial.ClearCoat.ClearCoatFactor, MP_CustomData0, TEXT("ClearCoat"), Material))
+							{
+								// TODO: handle failure?
+							}
+						}
+					}
+
+					if (!TryGetConstantScalar(JsonMaterial.ClearCoat.ClearCoatRoughnessFactor, MP_CustomData1, Material))
+					{
+						if (!TryGetSourceTexture(Builder, JsonMaterial.ClearCoat.ClearCoatRoughnessTexture, MP_CustomData1, Material, { GreenMask }))
+						{
+							if (!TryGetBakedMaterialProperty(Builder, JsonMaterial.ClearCoat.ClearCoatRoughnessTexture, JsonMaterial.ClearCoat.ClearCoatRoughnessFactor, MP_CustomData1, TEXT("ClearCoatRoughness"), Material))
+							{
+								// TODO: handle failure?
+							}
+						}
+					}
+				}
 			}
 		}
 	}
@@ -866,6 +889,21 @@ bool FGLTFMaterialConverter::TryGetBakedMaterialProperty(FGLTFConvertBuilder& Bu
 	}
 }
 
+inline bool FGLTFMaterialConverter::TryGetBakedMaterialProperty(FGLTFConvertBuilder& Builder, FGLTFJsonTextureInfo& OutTexInfo, float& OutConstant, EMaterialProperty Property, const FString& PropertyName, const UMaterialInterface* Material) const
+{
+	const FGLTFPropertyBakeOutput PropertyBakeOutput = BakeMaterialProperty(Property, Material);
+
+	if (PropertyBakeOutput.bIsConstant)
+	{
+		OutConstant = PropertyBakeOutput.ConstantValue.R;
+		return true;
+	}
+	else
+	{
+		return StoreBakedPropertyTexture(Builder, OutTexInfo, PropertyBakeOutput, PropertyName, Material);
+	}
+}
+
 bool FGLTFMaterialConverter::TryGetBakedMaterialProperty(FGLTFConvertBuilder& Builder, FGLTFJsonTextureInfo& OutTexInfo, EMaterialProperty Property, const FString& PropertyName, const UMaterialInterface* Material) const
 {
 	const FGLTFPropertyBakeOutput PropertyBakeOutput = BakeMaterialProperty(Property, Material);
diff --git a/Engine/Plugins/Enterprise/GLTFExporter/MigratedFromGit/Source/GLTFExporter/Private/Converters/GLTFMaterialConverters.h b/Engine/Plugins/Enterprise/GLTFExporter/MigratedFromGit/Source/GLTFExporter/Private/Converters/GLTFMaterialConverters.h
--- a/Engine/Plugins/Enterprise/GLTFExporter/MigratedFromGit/Source/GLTFExporter/Private/Converters/GLTFMaterialConverters.h
+++ b/Engine/Plugins/Enterprise/GLTFExporter/MigratedFromGit/Source/GLTFExporter/Private/Converters/GLTFMaterialConverters.h
@@ -30,6 +30,7 @@ class FGLTFMaterialConverter final : public TGLTFConverter<FGLTFJsonMaterialInde
 
 	bool TryGetBakedMaterialProperty(FGLTFConvertBuilder& Builder, FGLTFJsonTextureInfo& OutTexInfo, FGLTFJsonColor3& OutConstant, EMaterialProperty Property, const FString& PropertyName, const UMaterialInterface* Material) const;
 	bool TryGetBakedMaterialProperty(FGLTFConvertBuilder& Builder, FGLTFJsonTextureInfo& OutTexInfo, FGLTFJsonColor4& OutConstant, EMaterialProperty Property, const FString& PropertyName, const UMaterialInterface* Material) const;
+	bool TryGetBakedMaterialProperty(FGLTFConvertBuilder& Builder, FGLTFJsonTextureInfo& OutTexInfo, float& OutConstant, EMaterialProperty Property, const FString& PropertyName, const UMaterialInterface* Material) const;
 	bool TryGetBakedMaterialProperty(FGLTFConvertBuilder& Builder, FGLTFJsonTextureInfo& OutTexInfo, EMaterialProperty Property, const FString& PropertyName, const UMaterialInterface* Material) const;
 
 	FGLTFPropertyBakeOutput BakeMaterialProperty(EMaterialProperty Property, const UMaterialInterface* Material, const FIntPoint* PreferredTextureSize = nullptr, bool bCopyAlphaFromRedChannel = false) const;
