[GLTFExporter] refactor: add FGLTFMaterialUtility::GetCustomOutputByName

commit 4aa0523ce6136d99bad2fbb487102885be5231f1
Date: Thu, 25 Jun 2020 23:11:51 +0200

    refactor: add FGLTFMaterialUtility::GetCustomOutputByName
    (cherry picked from commit 674cc46592284ac3662bb1f1c5b2e2524726fe23​)

#skipci

diff --git a/Engine/Plugins/Enterprise/GLTFExporter/MigratedFromGit/Source/GLTFExporter/Private/Converters/GLTFMaterialUtility.cpp b/Engine/Plugins/Enterprise/GLTFExporter/MigratedFromGit/Source/GLTFExporter/Private/Converters/GLTFMaterialUtility.cpp
--- a/Engine/Plugins/Enterprise/GLTFExporter/MigratedFromGit/Source/GLTFExporter/Private/Converters/GLTFMaterialUtility.cpp
+++ b/Engine/Plugins/Enterprise/GLTFExporter/MigratedFromGit/Source/GLTFExporter/Private/Converters/GLTFMaterialUtility.cpp
@@ -8,6 +8,7 @@
 #include "Modules/ModuleManager.h"
 #include "IMaterialBakingModule.h"
 #include "MaterialBakingStructures.h"
+#include "Materials/MaterialExpressionCustomOutput.h"
 
 FVector4 FGLTFMaterialUtility::GetPropertyDefaultValue(EMaterialProperty Property)
 {
@@ -53,6 +54,20 @@ const FExpressionInput* FGLTFMaterialUtility::GetInputForProperty(const UMateria
 	return UnderlyingMaterial->GetExpressionInputForProperty(Property);
 }
 
+const UMaterialExpressionCustomOutput* FGLTFMaterialUtility::GetCustomOutputByName(const UMaterialInterface* Material, const FString& Name)
+{
+	for (const UMaterialExpression* Expression : Material->GetMaterial()->Expressions)
+	{
+		const UMaterialExpressionCustomOutput* CustomOutput = Cast<UMaterialExpressionCustomOutput>(Expression);
+		if (CustomOutput != nullptr && CustomOutput->GetFunctionName() == Name)
+		{
+			return CustomOutput;
+		}
+	}
+
+	return nullptr;
+}
+
 UTexture2D* FGLTFMaterialUtility::CreateTransientTexture(const FGLTFPropertyBakeOutput& PropertyBakeOutput, bool bUseSRGB)
 {
 	return FGLTFTextureUtility::CreateTransientTexture(
@@ -108,6 +123,11 @@ bool FGLTFMaterialUtility::CombineTextures(TArray<FColor>& OutPixels, const TArr
 
 FGLTFPropertyBakeOutput FGLTFMaterialUtility::BakeMaterialProperty(const FIntPoint& OutputSize, EMaterialProperty Property, const UMaterialInterface* Material, bool bCopyAlphaFromRedChannel)
 {
+	if (Property == MP_CustomData0 || Property == MP_CustomData1)
+	{
+		// TODO: add special support
+	}
+
 	TArray<FMeshData*> MeshSettings;
 
 	FMeshData MeshSet;
diff --git a/Engine/Plugins/Enterprise/GLTFExporter/MigratedFromGit/Source/GLTFExporter/Private/Converters/GLTFMaterialUtility.h b/Engine/Plugins/Enterprise/GLTFExporter/MigratedFromGit/Source/GLTFExporter/Private/Converters/GLTFMaterialUtility.h
--- a/Engine/Plugins/Enterprise/GLTFExporter/MigratedFromGit/Source/GLTFExporter/Private/Converters/GLTFMaterialUtility.h
+++ b/Engine/Plugins/Enterprise/GLTFExporter/MigratedFromGit/Source/GLTFExporter/Private/Converters/GLTFMaterialUtility.h
@@ -45,6 +45,8 @@ struct FGLTFMaterialUtility
 		return static_cast<const FMaterialInput<InputType>*>(ExpressionInput);
 	}
 
+	static const UMaterialExpressionCustomOutput* GetCustomOutputByName(const UMaterialInterface* Material, const FString& Name);
+
 	static UTexture2D* CreateTransientTexture(const FGLTFPropertyBakeOutput& PropertyBakeOutput, bool bUseSRGB = false);
 	static bool CombineTextures(TArray<FColor>& OutPixels, const TArray<FGLTFTextureCombineSource>& Sources, const FIntPoint& OutputSize, EPixelFormat OutputPixelFormat);
 	static FGLTFPropertyBakeOutput BakeMaterialProperty(const FIntPoint& OutputSize, EMaterialProperty Property, const UMaterialInterface* Material, bool bCopyAlphaFromRedChannel = false);
