[GLTFExporter] fix(clearcoat): make sure clear coat factors are not zero if texture exported

commit 004e9a861ee25a71cfdddd1622034cac9fb21315
Date: Fri, 10 Jul 2020 23:48:22 +0200

    fix(clearcoat): make sure clear coat factors are not zero if texture exported
    (cherry picked from commit 3963d57c52caa8de5a4f081c4e7f6d7608635f8a​)

#skipci

diff --git a/Engine/Plugins/Enterprise/GLTFExporter/MigratedFromGit/Source/GLTFExporter/Private/Converters/GLTFMaterialConverters.cpp b/Engine/Plugins/Enterprise/GLTFExporter/MigratedFromGit/Source/GLTFExporter/Private/Converters/GLTFMaterialConverters.cpp
--- a/Engine/Plugins/Enterprise/GLTFExporter/MigratedFromGit/Source/GLTFExporter/Private/Converters/GLTFMaterialConverters.cpp
+++ b/Engine/Plugins/Enterprise/GLTFExporter/MigratedFromGit/Source/GLTFExporter/Private/Converters/GLTFMaterialConverters.cpp
@@ -129,26 +129,36 @@ FGLTFJsonMaterialIndex FGLTFMaterialConverter::Add(FGLTFConvertBuilder& Builder,
 			{
 				if (JsonMaterial.ShadingModel == EGLTFJsonShadingModel::ClearCoat)
 				{
-					if (!TryGetConstantScalar(JsonMaterial.ClearCoat.ClearCoatFactor, MP_CustomData0, Material))
+					const EMaterialProperty ClearCoatProperty = MP_CustomData0;
+					if (!TryGetConstantScalar(JsonMaterial.ClearCoat.ClearCoatFactor, ClearCoatProperty, Material))
 					{
-						if (!TryGetSourceTexture(Builder, JsonMaterial.ClearCoat.ClearCoatTexture, MP_CustomData0, Material, { RedMask }))
+						if (!TryGetSourceTexture(Builder, JsonMaterial.ClearCoat.ClearCoatTexture, ClearCoatProperty, Material, { RedMask }))
 						{
-							if (!TryGetBakedMaterialProperty(Builder, JsonMaterial.ClearCoat.ClearCoatTexture, JsonMaterial.ClearCoat.ClearCoatFactor, MP_CustomData0, TEXT("ClearCoat"), Material))
+							if (!TryGetBakedMaterialProperty(Builder, JsonMaterial.ClearCoat.ClearCoatTexture, JsonMaterial.ClearCoat.ClearCoatFactor, ClearCoatProperty, TEXT("ClearCoat"), Material))
 							{
 								// TODO: handle failure?
 							}
 						}
+						else
+						{
+							JsonMaterial.ClearCoat.ClearCoatFactor = 1; // make sure factor is not zero
+						}
 					}
 
-					if (!TryGetConstantScalar(JsonMaterial.ClearCoat.ClearCoatRoughnessFactor, MP_CustomData1, Material))
+					const EMaterialProperty ClearCoatRoughnessProperty = MP_CustomData1;
+					if (!TryGetConstantScalar(JsonMaterial.ClearCoat.ClearCoatRoughnessFactor, ClearCoatRoughnessProperty, Material))
 					{
-						if (!TryGetSourceTexture(Builder, JsonMaterial.ClearCoat.ClearCoatRoughnessTexture, MP_CustomData1, Material, { GreenMask }))
+						if (!TryGetSourceTexture(Builder, JsonMaterial.ClearCoat.ClearCoatRoughnessTexture, ClearCoatRoughnessProperty, Material, { GreenMask }))
 						{
-							if (!TryGetBakedMaterialProperty(Builder, JsonMaterial.ClearCoat.ClearCoatRoughnessTexture, JsonMaterial.ClearCoat.ClearCoatRoughnessFactor, MP_CustomData1, TEXT("ClearCoatRoughness"), Material))
+							if (!TryGetBakedMaterialProperty(Builder, JsonMaterial.ClearCoat.ClearCoatRoughnessTexture, JsonMaterial.ClearCoat.ClearCoatRoughnessFactor, ClearCoatRoughnessProperty, TEXT("ClearCoatRoughness"), Material))
 							{
 								// TODO: handle failure?
 							}
 						}
+						else
+						{
+							JsonMaterial.ClearCoat.ClearCoatRoughnessFactor = 1; // make sure factor is not zero
+						}
 					}
 				}
 			}
@@ -868,10 +878,14 @@ bool FGLTFMaterialConverter::TryGetBakedMaterialProperty(FGLTFConvertBuilder& Bu
 		OutConstant = FGLTFConverterUtility::ConvertColor(PropertyBakeOutput.ConstantValue);
 		return true;
 	}
-	else
+
+	if (StoreBakedPropertyTexture(Builder, OutTexInfo, PropertyBakeOutput, PropertyName, Material))
 	{
-		return StoreBakedPropertyTexture(Builder, OutTexInfo, PropertyBakeOutput, PropertyName, Material);
+		OutConstant = FGLTFJsonColor3::White; // make sure property is not zero
+		return true;
 	}
+
+	return false;
 }
 
 bool FGLTFMaterialConverter::TryGetBakedMaterialProperty(FGLTFConvertBuilder& Builder, FGLTFJsonTextureInfo& OutTexInfo, FGLTFJsonColor4& OutConstant, EMaterialProperty Property, const FString& PropertyName, const UMaterialInterface* Material) const
@@ -883,10 +897,14 @@ bool FGLTFMaterialConverter::TryGetBakedMaterialProperty(FGLTFConvertBuilder& Bu
 		OutConstant = FGLTFConverterUtility::ConvertColor(PropertyBakeOutput.ConstantValue);
 		return true;
 	}
-	else
+
+	if (StoreBakedPropertyTexture(Builder, OutTexInfo, PropertyBakeOutput, PropertyName, Material))
 	{
-		return StoreBakedPropertyTexture(Builder, OutTexInfo, PropertyBakeOutput, PropertyName, Material);
+		OutConstant = FGLTFJsonColor4::White; // make sure property is not zero
+		return true;
 	}
+
+	return false;
 }
 
 inline bool FGLTFMaterialConverter::TryGetBakedMaterialProperty(FGLTFConvertBuilder& Builder, FGLTFJsonTextureInfo& OutTexInfo, float& OutConstant, EMaterialProperty Property, const FString& PropertyName, const UMaterialInterface* Material) const
@@ -898,10 +916,14 @@ inline bool FGLTFMaterialConverter::TryGetBakedMaterialProperty(FGLTFConvertBuil
 		OutConstant = PropertyBakeOutput.ConstantValue.R;
 		return true;
 	}
-	else
+
+	if (StoreBakedPropertyTexture(Builder, OutTexInfo, PropertyBakeOutput, PropertyName, Material))
 	{
-		return StoreBakedPropertyTexture(Builder, OutTexInfo, PropertyBakeOutput, PropertyName, Material);
+		OutConstant = 1; // make sure property is not zero
+		return true;
 	}
+
+	return false;
 }
 
 bool FGLTFMaterialConverter::TryGetBakedMaterialProperty(FGLTFConvertBuilder& Builder, FGLTFJsonTextureInfo& OutTexInfo, EMaterialProperty Property, const FString& PropertyName, const UMaterialInterface* Material) const
