[GLTFExporter] fix(clearcoat): workaround lack of ClearCoat support in MaterialBaking module by using ugly hack

commit 31c7ab6ea73bcfc62af454f215f368df24e67c3f
Date: Sat, 11 Jul 2020 00:35:41 +0200

    fix(clearcoat): workaround lack of ClearCoat support in MaterialBaking module by using ugly hack
    (cherry picked from commit 9ba88d75d30dfe4bd9f3c0fbf50d7ac7c61d27eb​)

#skipci

diff --git a/Engine/Plugins/Enterprise/GLTFExporter/MigratedFromGit/Source/GLTFExporter/Private/Converters/GLTFMaterialUtility.cpp b/Engine/Plugins/Enterprise/GLTFExporter/MigratedFromGit/Source/GLTFExporter/Private/Converters/GLTFMaterialUtility.cpp
--- a/Engine/Plugins/Enterprise/GLTFExporter/MigratedFromGit/Source/GLTFExporter/Private/Converters/GLTFMaterialUtility.cpp
+++ b/Engine/Plugins/Enterprise/GLTFExporter/MigratedFromGit/Source/GLTFExporter/Private/Converters/GLTFMaterialUtility.cpp
@@ -123,9 +123,31 @@ bool FGLTFMaterialUtility::CombineTextures(TArray<FColor>& OutPixels, const TArr
 
 FGLTFPropertyBakeOutput FGLTFMaterialUtility::BakeMaterialProperty(const FIntPoint& OutputSize, EMaterialProperty Property, const UMaterialInterface* Material, bool bCopyAlphaFromRedChannel)
 {
+	EMaterialProperty SpecialProperty = MP_MAX;
+	FScalarMaterialInput OldMetallic;
+	FScalarMaterialInput OldRoughness;
+
 	if (Property == MP_CustomData0 || Property == MP_CustomData1)
 	{
-		// TODO: add special support
+		// TODO: replace this hack by adding proper support for ClearCoat properties in MaterialBaking module
+
+		UMaterial* ParentMaterial = const_cast<UMaterial*>(Material->GetMaterial());
+
+		SpecialProperty = Property;
+		if (Property == MP_CustomData0)
+		{
+			Property = MP_Metallic;
+			OldMetallic = ParentMaterial->Metallic;
+			ParentMaterial->Metallic = ParentMaterial->ClearCoat;
+		}
+		else if (Property == MP_CustomData1)
+		{
+			Property = MP_Roughness;
+			OldRoughness = ParentMaterial->Roughness;
+			ParentMaterial->Roughness = ParentMaterial->ClearCoatRoughness;
+		}
+
+		ParentMaterial->ForceRecompileForRendering();
 	}
 
 	TArray<FMeshData*> MeshSettings;
@@ -186,6 +208,24 @@ FGLTFPropertyBakeOutput FGLTFMaterialUtility::BakeMaterialProperty(const FIntPoi
 		PropertyBakeOutput.ConstantValue = Pixel;
 	}
 
+	if (SpecialProperty != MP_MAX)
+	{
+		// TODO: replace this hack by adding proper support for ClearCoat properties in MaterialBaking module
+
+		UMaterial* ParentMaterial = const_cast<UMaterial*>(Material->GetMaterial());
+
+		if (SpecialProperty == MP_CustomData0)
+		{
+			ParentMaterial->Metallic = OldMetallic;
+		}
+		else if (SpecialProperty == MP_CustomData1)
+		{
+			ParentMaterial->Roughness = OldRoughness;
+		}
+
+		ParentMaterial->ForceRecompileForRendering();
+	}
+
 	return PropertyBakeOutput;
 }
 
