[GLTFExporter] feat(clearcoat): add support and workaround for clear-coat bottom normal custom output

commit a66c74f9eabae87536b6d38bc05bed342b65bcca
Date: Sat, 11 Jul 2020 02:04:14 +0200

    feat(clearcoat): add support and workaround for clear-coat bottom normal custom output
    (cherry picked from commit 0ff430231d94f56dc4f80c5305741af0f9811e03​)

#skipci

diff --git a/Engine/Plugins/Enterprise/GLTFExporter/MigratedFromGit/Source/GLTFExporter/Private/Converters/GLTFMaterialConverters.cpp b/Engine/Plugins/Enterprise/GLTFExporter/MigratedFromGit/Source/GLTFExporter/Private/Converters/GLTFMaterialConverters.cpp
--- a/Engine/Plugins/Enterprise/GLTFExporter/MigratedFromGit/Source/GLTFExporter/Private/Converters/GLTFMaterialConverters.cpp
+++ b/Engine/Plugins/Enterprise/GLTFExporter/MigratedFromGit/Source/GLTFExporter/Private/Converters/GLTFMaterialConverters.cpp
@@ -160,6 +160,18 @@ FGLTFJsonMaterialIndex FGLTFMaterialConverter::Add(FGLTFConvertBuilder& Builder,
 							JsonMaterial.ClearCoat.ClearCoatRoughnessFactor = 1; // make sure factor is not zero
 						}
 					}
+
+					const EMaterialProperty ClearCoatNormalProperty = MP_CustomOutput;
+					if (IsPropertyNonDefault(ClearCoatNormalProperty, Material))
+					{
+						if (!TryGetSourceTexture(Builder, JsonMaterial.ClearCoat.ClearCoatNormalTexture, ClearCoatNormalProperty, Material, DefaultColorInputMasks))
+						{
+							if (!TryGetBakedMaterialProperty(Builder, JsonMaterial.ClearCoat.ClearCoatNormalTexture, ClearCoatNormalProperty, TEXT("ClearCoatNormal"), Material))
+							{
+								// TODO: handle failure?
+							}
+						}
+					}
 				}
 			}
 		}
diff --git a/Engine/Plugins/Enterprise/GLTFExporter/MigratedFromGit/Source/GLTFExporter/Private/Converters/GLTFMaterialUtility.cpp b/Engine/Plugins/Enterprise/GLTFExporter/MigratedFromGit/Source/GLTFExporter/Private/Converters/GLTFMaterialUtility.cpp
--- a/Engine/Plugins/Enterprise/GLTFExporter/MigratedFromGit/Source/GLTFExporter/Private/Converters/GLTFMaterialUtility.cpp
+++ b/Engine/Plugins/Enterprise/GLTFExporter/MigratedFromGit/Source/GLTFExporter/Private/Converters/GLTFMaterialUtility.cpp
@@ -9,6 +9,7 @@
 #include "IMaterialBakingModule.h"
 #include "MaterialBakingStructures.h"
 #include "Materials/MaterialExpressionCustomOutput.h"
+#include "Materials/MaterialExpressionClearCoatNormalCustomOutput.h"
 
 FVector4 FGLTFMaterialUtility::GetPropertyDefaultValue(EMaterialProperty Property)
 {
@@ -44,12 +45,25 @@ FVector4 FGLTFMaterialUtility::GetPropertyDefaultValue(EMaterialProperty Propert
 		return FVector4(0,0,0,0);
 	}
 
+	// TODO: replace workaround for ClearCoatBottomNormal
+	if (Property == MP_CustomOutput)
+	{
+		return FVector4(0,0,1,0);
+	}
+
 	check(false);
 	return FVector4();
 }
 
 const FExpressionInput* FGLTFMaterialUtility::GetInputForProperty(const UMaterialInterface* Material, EMaterialProperty Property)
 {
+	// TODO: replace workaround for ClearCoatBottomNormal
+	if (Property == MP_CustomOutput)
+	{
+		const UMaterialExpressionCustomOutput* CustomOutput = Cast<UMaterialExpressionClearCoatNormalCustomOutput>(GetCustomOutputByName(Material, TEXT("ClearCoatBottomNormal")));
+		return CustomOutput != nullptr ? &CastChecked<UMaterialExpressionClearCoatNormalCustomOutput>(CustomOutput)->Input : nullptr;
+	}
+
 	UMaterial* UnderlyingMaterial = const_cast<UMaterial*>(Material->GetMaterial());
 	return UnderlyingMaterial->GetExpressionInputForProperty(Property);
 }
@@ -126,11 +140,11 @@ FGLTFPropertyBakeOutput FGLTFMaterialUtility::BakeMaterialProperty(const FIntPoi
 	EMaterialProperty SpecialProperty = MP_MAX;
 	FScalarMaterialInput OldMetallic;
 	FScalarMaterialInput OldRoughness;
+	FVectorMaterialInput OldNormal;
 
-	if (Property == MP_CustomData0 || Property == MP_CustomData1)
+	// TODO: replace this hack by adding proper support for ClearCoat properties in MaterialBaking module
+	if (Property == MP_CustomData0 || Property == MP_CustomData1 || Property == MP_CustomOutput)
 	{
-		// TODO: replace this hack by adding proper support for ClearCoat properties in MaterialBaking module
-
 		UMaterial* ParentMaterial = const_cast<UMaterial*>(Material->GetMaterial());
 
 		SpecialProperty = Property;
@@ -146,6 +160,13 @@ FGLTFPropertyBakeOutput FGLTFMaterialUtility::BakeMaterialProperty(const FIntPoi
 			OldRoughness = ParentMaterial->Roughness;
 			ParentMaterial->Roughness = ParentMaterial->ClearCoatRoughness;
 		}
+		else if (Property == MP_CustomOutput)
+		{
+			Property = MP_Normal;
+			OldNormal = ParentMaterial->Normal;
+			const FExpressionInput* ClearCoatBottomNormal = GetInputForProperty(Material, MP_CustomOutput);
+			static_cast<FExpressionInput&>(ParentMaterial->Normal) = *ClearCoatBottomNormal;
+		}
 
 		ParentMaterial->ForceRecompileForRendering();
 	}
@@ -208,10 +229,9 @@ FGLTFPropertyBakeOutput FGLTFMaterialUtility::BakeMaterialProperty(const FIntPoi
 		PropertyBakeOutput.ConstantValue = Pixel;
 	}
 
+	// TODO: replace this hack by adding proper support for ClearCoat properties in MaterialBaking module
 	if (SpecialProperty != MP_MAX)
 	{
-		// TODO: replace this hack by adding proper support for ClearCoat properties in MaterialBaking module
-
 		UMaterial* ParentMaterial = const_cast<UMaterial*>(Material->GetMaterial());
 
 		if (SpecialProperty == MP_CustomData0)
@@ -222,6 +242,10 @@ FGLTFPropertyBakeOutput FGLTFMaterialUtility::BakeMaterialProperty(const FIntPoi
 		{
 			ParentMaterial->Roughness = OldRoughness;
 		}
+		else if (SpecialProperty == MP_CustomOutput)
+		{
+			ParentMaterial->Normal = OldNormal;
+		}
 
 		ParentMaterial->ForceRecompileForRendering();
 	}
