[GLTFExporter] refactor(clearcoat): add warning messages when clear coat properties failed to be exported

commit 780a987bafa008034ac03121513ace6549bf69d2
Date: Sun, 12 Jul 2020 14:35:33 +0200

    refactor(clearcoat): add warning messages when clear coat properties failed to be exported
    (cherry picked from commit d7ff0affc4218a7250427132772c97f99d06a15b​)

#skipci

diff --git a/Engine/Plugins/Enterprise/GLTFExporter/MigratedFromGit/Source/GLTFExporter/Private/Converters/GLTFMaterialConverters.cpp b/Engine/Plugins/Enterprise/GLTFExporter/MigratedFromGit/Source/GLTFExporter/Private/Converters/GLTFMaterialConverters.cpp
--- a/Engine/Plugins/Enterprise/GLTFExporter/MigratedFromGit/Source/GLTFExporter/Private/Converters/GLTFMaterialConverters.cpp
+++ b/Engine/Plugins/Enterprise/GLTFExporter/MigratedFromGit/Source/GLTFExporter/Private/Converters/GLTFMaterialConverters.cpp
@@ -136,7 +136,7 @@ FGLTFJsonMaterialIndex FGLTFMaterialConverter::Add(FGLTFConvertBuilder& Builder,
 						{
 							if (!TryGetBakedMaterialProperty(Builder, JsonMaterial.ClearCoat.ClearCoatTexture, JsonMaterial.ClearCoat.ClearCoatFactor, ClearCoatProperty, TEXT("ClearCoat"), Material))
 							{
-								// TODO: handle failure?
+								Builder.AddWarningMessage(FString::Printf(TEXT("Failed to export ClearCoat for material %s"), *Material->GetName()));
 							}
 						}
 						else
@@ -152,7 +152,7 @@ FGLTFJsonMaterialIndex FGLTFMaterialConverter::Add(FGLTFConvertBuilder& Builder,
 						{
 							if (!TryGetBakedMaterialProperty(Builder, JsonMaterial.ClearCoat.ClearCoatRoughnessTexture, JsonMaterial.ClearCoat.ClearCoatRoughnessFactor, ClearCoatRoughnessProperty, TEXT("ClearCoatRoughness"), Material))
 							{
-								// TODO: handle failure?
+								Builder.AddWarningMessage(FString::Printf(TEXT("Failed to export ClearCoatRoughness for material %s"), *Material->GetName()));
 							}
 						}
 						else
@@ -161,14 +161,14 @@ FGLTFJsonMaterialIndex FGLTFMaterialConverter::Add(FGLTFConvertBuilder& Builder,
 						}
 					}
 
-					const EMaterialProperty ClearCoatNormalProperty = MP_CustomOutput;
+					const EMaterialProperty ClearCoatNormalProperty = MP_CustomOutput; // TODO: replace workaround for ClearCoatBottomNormal with proper support for custom output
 					if (IsPropertyNonDefault(ClearCoatNormalProperty, Material))
 					{
 						if (!TryGetSourceTexture(Builder, JsonMaterial.ClearCoat.ClearCoatNormalTexture, ClearCoatNormalProperty, Material, DefaultColorInputMasks))
 						{
 							if (!TryGetBakedMaterialProperty(Builder, JsonMaterial.ClearCoat.ClearCoatNormalTexture, ClearCoatNormalProperty, TEXT("ClearCoatNormal"), Material))
 							{
-								// TODO: handle failure?
+								Builder.AddWarningMessage(FString::Printf(TEXT("Failed to export ClearCoatBottomNormal for material %s"), *Material->GetName()));
 							}
 						}
 					}
