[GLTFExporter] fix(clearcoat): fix crash and cache issues with hack for baking clear-coat properties

commit 88ba4a1aa214a7abfcf40167992f5fef6668c9ff
Date: Mon, 20 Jul 2020 16:30:19 +0200

    fix(clearcoat): fix crash and cache issues with hack for baking clear-coat properties
    (cherry picked from commit bc8c303fe4ceb29e8a587553b2e5c667e3ee735a​)

#skipci

diff --git a/Engine/Plugins/Enterprise/GLTFExporter/MigratedFromGit/Source/GLTFExporter/Private/Converters/GLTFMaterialUtility.cpp b/Engine/Plugins/Enterprise/GLTFExporter/MigratedFromGit/Source/GLTFExporter/Private/Converters/GLTFMaterialUtility.cpp
--- a/Engine/Plugins/Enterprise/GLTFExporter/MigratedFromGit/Source/GLTFExporter/Private/Converters/GLTFMaterialUtility.cpp
+++ b/Engine/Plugins/Enterprise/GLTFExporter/MigratedFromGit/Source/GLTFExporter/Private/Converters/GLTFMaterialUtility.cpp
@@ -60,7 +60,7 @@ const FExpressionInput* FGLTFMaterialUtility::GetInputForProperty(const UMateria
 	// TODO: replace workaround for ClearCoatBottomNormal
 	if (Property == MP_CustomOutput)
 	{
-		const UMaterialExpressionCustomOutput* CustomOutput = Cast<UMaterialExpressionClearCoatNormalCustomOutput>(GetCustomOutputByName(Material, TEXT("ClearCoatBottomNormal")));
+		const UMaterialExpressionCustomOutput* CustomOutput = GetCustomOutputByName(Material, TEXT("ClearCoatBottomNormal"));
 		return CustomOutput != nullptr ? &CastChecked<UMaterialExpressionClearCoatNormalCustomOutput>(CustomOutput)->Input : nullptr;
 	}
 
@@ -137,38 +137,31 @@ bool FGLTFMaterialUtility::CombineTextures(TArray<FColor>& OutPixels, const TArr
 
 FGLTFPropertyBakeOutput FGLTFMaterialUtility::BakeMaterialProperty(const FIntPoint& OutputSize, EMaterialProperty Property, const UMaterialInterface* Material, bool bCopyAlphaFromRedChannel)
 {
-	EMaterialProperty SpecialProperty = MP_MAX;
-	FScalarMaterialInput OldMetallic;
-	FScalarMaterialInput OldRoughness;
-	FVectorMaterialInput OldNormal;
+	EMaterialProperty HackProperty = MP_MAX;
+	TTuple<int32, UMaterialExpression*> HackPrevState;
 
 	// TODO: replace this hack by adding proper support for ClearCoat properties in MaterialBaking module
 	if (Property == MP_CustomData0 || Property == MP_CustomData1 || Property == MP_CustomOutput)
 	{
-		UMaterial* ParentMaterial = const_cast<UMaterial*>(Material->GetMaterial());
-
-		SpecialProperty = Property;
-		if (Property == MP_CustomData0)
-		{
-			Property = MP_Metallic;
-			OldMetallic = ParentMaterial->Metallic;
-			ParentMaterial->Metallic = ParentMaterial->ClearCoat;
-		}
-		else if (Property == MP_CustomData1)
+		switch (Property)
 		{
-			Property = MP_Roughness;
-			OldRoughness = ParentMaterial->Roughness;
-			ParentMaterial->Roughness = ParentMaterial->ClearCoatRoughness;
+			case MP_CustomData0: HackProperty = MP_Metallic; break;
+			case MP_CustomData1: HackProperty = MP_Roughness; break;
+			case MP_CustomOutput: HackProperty = MP_Normal; break;
+			default: checkNoEntry();
 		}
-		else if (Property == MP_CustomOutput)
+
+		UMaterial* ParentMaterial = const_cast<UMaterial*>(Material->GetMaterial());
+		ParentMaterial->Modify();
 		{
-			Property = MP_Normal;
-			OldNormal = ParentMaterial->Normal;
-			const FExpressionInput* ClearCoatBottomNormal = GetInputForProperty(Material, MP_CustomOutput);
-			static_cast<FExpressionInput&>(ParentMaterial->Normal) = *ClearCoatBottomNormal;
+			const FExpressionInput& Input = *GetInputForProperty(Material, Property);
+			FExpressionInput& HackInput = *const_cast<FExpressionInput*>(GetInputForProperty(Material, HackProperty));
+			HackPrevState = MakeTuple(HackInput.OutputIndex, HackInput.Expression);
+			HackInput.Connect(Input.OutputIndex, Input.Expression);
 		}
+		ParentMaterial->PostEditChange();
 
-		ParentMaterial->ForceRecompileForRendering();
+		Property = HackProperty;
 	}
 
 	TArray<FMeshData*> MeshSettings;
@@ -230,24 +223,15 @@ FGLTFPropertyBakeOutput FGLTFMaterialUtility::BakeMaterialProperty(const FIntPoi
 	}
 
 	// TODO: replace this hack by adding proper support for ClearCoat properties in MaterialBaking module
-	if (SpecialProperty != MP_MAX)
+	if (HackProperty != MP_MAX)
 	{
 		UMaterial* ParentMaterial = const_cast<UMaterial*>(Material->GetMaterial());
-
-		if (SpecialProperty == MP_CustomData0)
-		{
-			ParentMaterial->Metallic = OldMetallic;
-		}
-		else if (SpecialProperty == MP_CustomData1)
+		ParentMaterial->Modify();
 		{
-			ParentMaterial->Roughness = OldRoughness;
+			FExpressionInput& HackInput = *const_cast<FExpressionInput*>(GetInputForProperty(Material, HackProperty));
+			HackInput.Connect(HackPrevState.Get<0>(), HackPrevState.Get<1>());
 		}
-		else if (SpecialProperty == MP_CustomOutput)
-		{
-			ParentMaterial->Normal = OldNormal;
-		}
-
-		ParentMaterial->ForceRecompileForRendering();
+		ParentMaterial->PostEditChange();
 	}
 
 	return PropertyBakeOutput;
