[GLTFExporter] fix(material): if material baking results in constant default value then let gltf use default

commit 3ca4189df2bf6d38cea8afa3bf84f809d26b0135
Date: Mon, 20 Jul 2020 16:35:10 +0200

    fix(material): if material baking results in constant default value then let gltf use default
    (cherry picked from commit 886ea2b86a9afcbad45acfef917a2d69f57daea3​)

#skipci

diff --git a/Engine/Plugins/Enterprise/GLTFExporter/MigratedFromGit/Source/GLTFExporter/Private/Converters/GLTFMaterialConverters.cpp b/Engine/Plugins/Enterprise/GLTFExporter/MigratedFromGit/Source/GLTFExporter/Private/Converters/GLTFMaterialConverters.cpp
--- a/Engine/Plugins/Enterprise/GLTFExporter/MigratedFromGit/Source/GLTFExporter/Private/Converters/GLTFMaterialConverters.cpp
+++ b/Engine/Plugins/Enterprise/GLTFExporter/MigratedFromGit/Source/GLTFExporter/Private/Converters/GLTFMaterialConverters.cpp
@@ -946,21 +946,26 @@ bool FGLTFMaterialConverter::TryGetBakedMaterialProperty(FGLTFConvertBuilder& Bu
 	{
 		return StoreBakedPropertyTexture(Builder, OutTexInfo, PropertyBakeOutput, PropertyName, Material);
 	}
-	else
+
+	if (PropertyBakeOutput.ConstantValue == FLinearColor(FGLTFMaterialUtility::GetPropertyDefaultValue(Property)))
 	{
-		// NOTE: since this function is meant to bake to a texture, we assume that the property
-		// that was passed into the function is using a non-constant expression.
-		// We therefore treat a constant result when baking as a failure.
+		// Constant value is the same as the property's default so we can set gltf to default.
+		OutTexInfo.Index = FGLTFJsonTextureIndex(INDEX_NONE);
+		return true;
+	}
 
-		// NOTE: in some cases a constant baking result is returned for a property that is non-constant.
-		// This happens (for example) when baking AmbientOcclusion for a translucent material,
-		// even though the same material when set to opaque will properly bake AmbientOcclusion to a texture.
-		// It also happens when baking Normal in some (not yet identified) circumstances.
-		// For now, these incorrect bakes are discarded.
+	// NOTE: since this function is meant to bake to a texture, we assume that the property
+	// that was passed into the function is using a non-constant expression.
+	// We therefore treat a constant result with non-default value as a failure.
 
-		// TODO: investigate why non-constant properties are sometimes baked to a constant expression (see note above)
-		return false;
-	}
+	// NOTE: in some cases a constant baking result is returned for a property that is non-constant.
+	// This happens (for example) when baking AmbientOcclusion for a translucent material,
+	// even though the same material when set to opaque will properly bake AmbientOcclusion to a texture.
+	// It also happens when baking Normal in some (not yet identified) circumstances.
+	// For now, these incorrect bakes are discarded.
+
+	// TODO: investigate why non-constant properties are sometimes baked to a constant expression (see note above)
+	return false;
 }
 
 FGLTFPropertyBakeOutput FGLTFMaterialConverter::BakeMaterialProperty(EMaterialProperty Property, const UMaterialInterface* Material, const FIntPoint* PreferredTextureSize, bool bCopyAlphaFromRedChannel) const
