[GLTFExporter] refactor!(material): change hardcoded texture resolution for material baking from 512 to 1024

commit 266d39142d2455e7b74a9779dc8dcfd7092c71d4
Date: Mon, 20 Jul 2020 16:36:43 +0200

    refactor!(material): change hardcoded texture resolution for material baking from 512 to 1024
    (cherry picked from commit 1197052e1789d69bc71f121e90d305b5284e6f91​)

#skipci

diff --git a/Engine/Plugins/Enterprise/GLTFExporter/MigratedFromGit/Source/GLTFExporter/Private/Converters/GLTFMaterialConverters.cpp b/Engine/Plugins/Enterprise/GLTFExporter/MigratedFromGit/Source/GLTFExporter/Private/Converters/GLTFMaterialConverters.cpp
--- a/Engine/Plugins/Enterprise/GLTFExporter/MigratedFromGit/Source/GLTFExporter/Private/Converters/GLTFMaterialConverters.cpp
+++ b/Engine/Plugins/Enterprise/GLTFExporter/MigratedFromGit/Source/GLTFExporter/Private/Converters/GLTFMaterialConverters.cpp
@@ -248,7 +248,7 @@ bool FGLTFMaterialConverter::TryGetBaseColorAndOpacity(FGLTFConvertBuilder& Buil
 	// TODO: add support for detecting the correct tex-coord for this property based on connected nodes
 	// TODO: add support for calculating the ideal resolution to use for baking based on connected (texture) nodes
 	int32 TexCoord = 0;
-	FIntPoint TextureSize(512, 512);
+	FIntPoint TextureSize(1024, 1024);
 
 	// TODO: should this be the default wrap-mode?
 	EGLTFJsonTextureWrap TextureWrapS = EGLTFJsonTextureWrap::Repeat;
@@ -390,7 +390,7 @@ bool FGLTFMaterialConverter::TryGetMetallicAndRoughness(FGLTFConvertBuilder& Bui
 	// TODO: add support for detecting the correct tex-coord for this property based on connected nodes
 	// TODO: add support for calculating the ideal resolution to use for baking based on connected (texture) nodes
 	int32 TexCoord = 0;
-	FIntPoint TextureSize(512, 512);
+	FIntPoint TextureSize(1024, 1024);
 
 	// TODO: should this be the default wrap-mode?
 	EGLTFJsonTextureWrap TextureWrapS = EGLTFJsonTextureWrap::Repeat;
@@ -971,7 +971,7 @@ bool FGLTFMaterialConverter::TryGetBakedMaterialProperty(FGLTFConvertBuilder& Bu
 FGLTFPropertyBakeOutput FGLTFMaterialConverter::BakeMaterialProperty(EMaterialProperty Property, const UMaterialInterface* Material, const FIntPoint* PreferredTextureSize, bool bCopyAlphaFromRedChannel) const
 {
 	// TODO: make default baking-resolution configurable
-	const FIntPoint DefaultTextureSize(512, 512);
+	const FIntPoint DefaultTextureSize(1024, 1024);
 	const FIntPoint TextureSize = PreferredTextureSize != nullptr ? *PreferredTextureSize : DefaultTextureSize;
 
 	// TODO: add support for calculating the ideal resolution to use for baking based on connected (texture) nodes
diff --git a/Engine/Plugins/Enterprise/GLTFExporter/MigratedFromGit/Source/GLTFExporter/Private/GLTFExportOptions.cpp b/Engine/Plugins/Enterprise/GLTFExporter/MigratedFromGit/Source/GLTFExporter/Private/GLTFExportOptions.cpp
--- a/Engine/Plugins/Enterprise/GLTFExporter/MigratedFromGit/Source/GLTFExporter/Private/GLTFExportOptions.cpp
+++ b/Engine/Plugins/Enterprise/GLTFExporter/MigratedFromGit/Source/GLTFExporter/Private/GLTFExportOptions.cpp
@@ -27,7 +27,7 @@ UGLTFExportOptions::UGLTFExportOptions(const FObjectInitializer& ObjectInitializ
 	bExportUnlitMaterials = false;
 	bExportClearCoatMaterials = false;
 	bBakeMaterialInputs = true;
-	DefaultMaterialBakeSize = EGLTFExporterMaterialBakeSize::POT_512;
+	DefaultMaterialBakeSize = EGLTFExporterMaterialBakeSize::POT_1024;
 	TextureFormat = EGLTFExporterTextureFormat::None;
 	bExportLightmaps = false;
 	TextureHDREncoding = EGLTFExporterTextureHDREncoding::None;
