[GLTFExporter] fix(clearcoat): if material uses clear coat, then Unreal "Normal" should be glTF "clearCoatNormal" and Unreal "ClearCoatBottomNormal" is glTF "normal"

commit 30b26de419beb7b2bb9e3308c86f7e9b84ad2fd6
Date: Mon, 20 Jul 2020 17:07:56 +0200

    fix(clearcoat): if material uses clear coat, then Unreal "Normal" should be glTF "clearCoatNormal" and Unreal "ClearCoatBottomNormal" is glTF "normal"
    (cherry picked from commit 365d75d6f172a08a0369ea252546351ec706ed0a​)

#skipci

diff --git a/Engine/Plugins/Enterprise/GLTFExporter/MigratedFromGit/Source/GLTFExporter/Private/Converters/GLTFMaterialConverters.cpp b/Engine/Plugins/Enterprise/GLTFExporter/MigratedFromGit/Source/GLTFExporter/Private/Converters/GLTFMaterialConverters.cpp
--- a/Engine/Plugins/Enterprise/GLTFExporter/MigratedFromGit/Source/GLTFExporter/Private/Converters/GLTFMaterialConverters.cpp
+++ b/Engine/Plugins/Enterprise/GLTFExporter/MigratedFromGit/Source/GLTFExporter/Private/Converters/GLTFMaterialConverters.cpp
@@ -101,14 +101,15 @@ FGLTFJsonMaterialIndex FGLTFMaterialConverter::Add(FGLTFConvertBuilder& Builder,
 				Builder.AddWarningMessage(FString::Printf(TEXT("Failed to export EmissiveColor for material %s"), *Material->GetName()));
 			}
 
-			const EMaterialProperty NormalProperty = MP_Normal;
+			// TODO: replace dummy enum MP_CustomOutput workaround for ClearCoatBottomNormal with proper support for custom outputs
+			const EMaterialProperty NormalProperty = JsonMaterial.ShadingModel == EGLTFJsonShadingModel::ClearCoat ? MP_CustomOutput : MP_Normal;
 			if (IsPropertyNonDefault(NormalProperty, Material))
 			{
 				if (!TryGetSourceTexture(Builder, JsonMaterial.NormalTexture, NormalProperty, Material, DefaultColorInputMasks))
 				{
 					if (!TryGetBakedMaterialProperty(Builder, JsonMaterial.NormalTexture, NormalProperty, TEXT("Normal"), Material))
 					{
-						Builder.AddWarningMessage(FString::Printf(TEXT("Failed to export Normal for material %s"), *Material->GetName()));
+						Builder.AddWarningMessage(FString::Printf(TEXT("Failed to export Normal or ClearCoatBottomNormal for material %s"), *Material->GetName()));
 					}
 				}
 			}
@@ -161,14 +162,14 @@ FGLTFJsonMaterialIndex FGLTFMaterialConverter::Add(FGLTFConvertBuilder& Builder,
 						}
 					}
 
-					const EMaterialProperty ClearCoatNormalProperty = MP_CustomOutput; // TODO: replace workaround for ClearCoatBottomNormal with proper support for custom output
+					const EMaterialProperty ClearCoatNormalProperty = MP_Normal;
 					if (IsPropertyNonDefault(ClearCoatNormalProperty, Material))
 					{
 						if (!TryGetSourceTexture(Builder, JsonMaterial.ClearCoat.ClearCoatNormalTexture, ClearCoatNormalProperty, Material, DefaultColorInputMasks))
 						{
 							if (!TryGetBakedMaterialProperty(Builder, JsonMaterial.ClearCoat.ClearCoatNormalTexture, ClearCoatNormalProperty, TEXT("ClearCoatNormal"), Material))
 							{
-								Builder.AddWarningMessage(FString::Printf(TEXT("Failed to export ClearCoatBottomNormal for material %s"), *Material->GetName()));
+								Builder.AddWarningMessage(FString::Printf(TEXT("Failed to export Normal for material %s"), *Material->GetName()));
 							}
 						}
 					}
