[GLTFExporter] fix(material): fix some warning messages making incorrect assumptions about their material property

commit 2ca108388e5a068b9fa7acb46819efcd0e74bf47
Date: Mon, 20 Jul 2020 17:09:52 +0200

    fix(material): fix some warning messages making incorrect assumptions about their material property
    (cherry picked from commit 0c2b5d663b156dc93c0ea7a2f45dbfaaec44bc13​)

#skipci

diff --git a/Engine/Plugins/Enterprise/GLTFExporter/MigratedFromGit/Source/GLTFExporter/Private/Converters/GLTFMaterialConverters.cpp b/Engine/Plugins/Enterprise/GLTFExporter/MigratedFromGit/Source/GLTFExporter/Private/Converters/GLTFMaterialConverters.cpp
--- a/Engine/Plugins/Enterprise/GLTFExporter/MigratedFromGit/Source/GLTFExporter/Private/Converters/GLTFMaterialConverters.cpp
+++ b/Engine/Plugins/Enterprise/GLTFExporter/MigratedFromGit/Source/GLTFExporter/Private/Converters/GLTFMaterialConverters.cpp
@@ -70,7 +70,7 @@ FGLTFJsonMaterialIndex FGLTFMaterialConverter::Add(FGLTFConvertBuilder& Builder,
 				{
 					if (!TryGetBakedMaterialProperty(Builder, JsonMaterial.PBRMetallicRoughness.BaseColorTexture, JsonMaterial.PBRMetallicRoughness.BaseColorFactor, BaseColorProperty, TEXT("BaseColor"), Material))
 					{
-						Builder.AddWarningMessage(FString::Printf(TEXT("Failed to export BaseColor for material %s"), *Material->GetName()));
+						Builder.AddWarningMessage(FString::Printf(TEXT("Failed to export %s for material %s"), FGLTFMaterialUtility::GetPropertyName(BaseColorProperty), *Material->GetName()));
 					}
 				}
 			}
@@ -81,7 +81,7 @@ FGLTFJsonMaterialIndex FGLTFMaterialConverter::Add(FGLTFConvertBuilder& Builder,
 		{
 			if (!TryGetBaseColorAndOpacity(Builder, JsonMaterial.PBRMetallicRoughness, Material, BaseColorProperty, OpacityProperty))
 			{
-				Builder.AddWarningMessage(FString::Printf(TEXT("Failed to export BaseColor & Opacity for material %s"), *Material->GetName()));
+				Builder.AddWarningMessage(FString::Printf(TEXT("Failed to export %s and %s for material %s"), FGLTFMaterialUtility::GetPropertyName(BaseColorProperty), FGLTFMaterialUtility::GetPropertyName(OpacityProperty), *Material->GetName()));
 			}
 		}
 
@@ -92,13 +92,13 @@ FGLTFJsonMaterialIndex FGLTFMaterialConverter::Add(FGLTFConvertBuilder& Builder,
 
 			if (!TryGetMetallicAndRoughness(Builder, JsonMaterial.PBRMetallicRoughness, Material, MetallicProperty, RoughnessProperty))
 			{
-				Builder.AddWarningMessage(FString::Printf(TEXT("Failed to export Metallic & Roughness for material %s"), *Material->GetName()));
+				Builder.AddWarningMessage(FString::Printf(TEXT("Failed to export %s and %s for material %s"), FGLTFMaterialUtility::GetPropertyName(MetallicProperty), FGLTFMaterialUtility::GetPropertyName(RoughnessProperty), *Material->GetName()));
 			}
 
 			const EMaterialProperty EmissiveProperty = MP_EmissiveColor;
 			if (!TryGetEmissive(Builder, JsonMaterial, EmissiveProperty, Material))
 			{
-				Builder.AddWarningMessage(FString::Printf(TEXT("Failed to export EmissiveColor for material %s"), *Material->GetName()));
+				Builder.AddWarningMessage(FString::Printf(TEXT("Failed to export %s for material %s"), FGLTFMaterialUtility::GetPropertyName(EmissiveProperty), *Material->GetName()));
 			}
 
 			// TODO: replace dummy enum MP_CustomOutput workaround for ClearCoatBottomNormal with proper support for custom outputs
@@ -109,7 +109,7 @@ FGLTFJsonMaterialIndex FGLTFMaterialConverter::Add(FGLTFConvertBuilder& Builder,
 				{
 					if (!TryGetBakedMaterialProperty(Builder, JsonMaterial.NormalTexture, NormalProperty, TEXT("Normal"), Material))
 					{
-						Builder.AddWarningMessage(FString::Printf(TEXT("Failed to export Normal or ClearCoatBottomNormal for material %s"), *Material->GetName()));
+						Builder.AddWarningMessage(FString::Printf(TEXT("Failed to export %s for material %s"), FGLTFMaterialUtility::GetPropertyName(NormalProperty), *Material->GetName()));
 					}
 				}
 			}
@@ -121,7 +121,7 @@ FGLTFJsonMaterialIndex FGLTFMaterialConverter::Add(FGLTFConvertBuilder& Builder,
 				{
 					if (!TryGetBakedMaterialProperty(Builder, JsonMaterial.OcclusionTexture, AmbientOcclusionProperty, TEXT("Occlusion"), Material))
 					{
-						Builder.AddWarningMessage(FString::Printf(TEXT("Failed to export AmbientOcclusion for material %s"), *Material->GetName()));
+						Builder.AddWarningMessage(FString::Printf(TEXT("Failed to export %s for material %s"), FGLTFMaterialUtility::GetPropertyName(AmbientOcclusionProperty), *Material->GetName()));
 					}
 				}
 			}
@@ -137,7 +137,7 @@ FGLTFJsonMaterialIndex FGLTFMaterialConverter::Add(FGLTFConvertBuilder& Builder,
 						{
 							if (!TryGetBakedMaterialProperty(Builder, JsonMaterial.ClearCoat.ClearCoatTexture, JsonMaterial.ClearCoat.ClearCoatFactor, ClearCoatProperty, TEXT("ClearCoat"), Material))
 							{
-								Builder.AddWarningMessage(FString::Printf(TEXT("Failed to export ClearCoat for material %s"), *Material->GetName()));
+								Builder.AddWarningMessage(FString::Printf(TEXT("Failed to export %s for material %s"), FGLTFMaterialUtility::GetPropertyName(ClearCoatProperty), *Material->GetName()));
 							}
 						}
 						else
@@ -153,7 +153,7 @@ FGLTFJsonMaterialIndex FGLTFMaterialConverter::Add(FGLTFConvertBuilder& Builder,
 						{
 							if (!TryGetBakedMaterialProperty(Builder, JsonMaterial.ClearCoat.ClearCoatRoughnessTexture, JsonMaterial.ClearCoat.ClearCoatRoughnessFactor, ClearCoatRoughnessProperty, TEXT("ClearCoatRoughness"), Material))
 							{
-								Builder.AddWarningMessage(FString::Printf(TEXT("Failed to export ClearCoatRoughness for material %s"), *Material->GetName()));
+								Builder.AddWarningMessage(FString::Printf(TEXT("Failed to export %s for material %s"), FGLTFMaterialUtility::GetPropertyName(ClearCoatRoughnessProperty), *Material->GetName()));
 							}
 						}
 						else
@@ -169,7 +169,7 @@ FGLTFJsonMaterialIndex FGLTFMaterialConverter::Add(FGLTFConvertBuilder& Builder,
 						{
 							if (!TryGetBakedMaterialProperty(Builder, JsonMaterial.ClearCoat.ClearCoatNormalTexture, ClearCoatNormalProperty, TEXT("ClearCoatNormal"), Material))
 							{
-								Builder.AddWarningMessage(FString::Printf(TEXT("Failed to export Normal for material %s"), *Material->GetName()));
+								Builder.AddWarningMessage(FString::Printf(TEXT("Failed to export %s for material %s"), FGLTFMaterialUtility::GetPropertyName(ClearCoatNormalProperty), *Material->GetName()));
 							}
 						}
 					}
diff --git a/Engine/Plugins/Enterprise/GLTFExporter/MigratedFromGit/Source/GLTFExporter/Private/Converters/GLTFMaterialUtility.cpp b/Engine/Plugins/Enterprise/GLTFExporter/MigratedFromGit/Source/GLTFExporter/Private/Converters/GLTFMaterialUtility.cpp
--- a/Engine/Plugins/Enterprise/GLTFExporter/MigratedFromGit/Source/GLTFExporter/Private/Converters/GLTFMaterialUtility.cpp
+++ b/Engine/Plugins/Enterprise/GLTFExporter/MigratedFromGit/Source/GLTFExporter/Private/Converters/GLTFMaterialUtility.cpp
@@ -11,6 +11,48 @@
 #include "Materials/MaterialExpressionCustomOutput.h"
 #include "Materials/MaterialExpressionClearCoatNormalCustomOutput.h"
 
+const TCHAR* FGLTFMaterialUtility::GetPropertyName(EMaterialProperty Property)
+{
+	// TODO: replace this hardcoded list with some lookup function that should be in the engine.
+
+	switch (Property)
+	{
+		case MP_EmissiveColor:          return TEXT("EmissiveColor");
+		case MP_Opacity:                return TEXT("Opacity");
+		case MP_OpacityMask:            return TEXT("OpacityMask");
+		case MP_BaseColor:              return TEXT("BaseColor");
+		case MP_Metallic:               return TEXT("Metallic");
+		case MP_Specular:               return TEXT("Specular");
+		case MP_Roughness:              return TEXT("Roughness");
+		case MP_Anisotropy:             return TEXT("Anisotropy");
+		case MP_Normal:                 return TEXT("Normal");
+		case MP_Tangent:                return TEXT("Tangent");
+		case MP_WorldPositionOffset:    return TEXT("WorldPositionOffset");
+		case MP_WorldDisplacement:      return TEXT("WorldDisplacement");
+		case MP_TessellationMultiplier: return TEXT("TessellationMultiplier");
+		case MP_SubsurfaceColor:        return TEXT("SubsurfaceColor");
+		case MP_CustomData0:            return TEXT("ClearCoat");
+		case MP_CustomData1:            return TEXT("ClearCoatRoughness");
+		case MP_AmbientOcclusion:       return TEXT("AmbientOcclusion");
+		case MP_Refraction:             return TEXT("Refraction");
+		case MP_CustomizedUVs0:         return TEXT("CustomizedUV0");
+		case MP_CustomizedUVs1:         return TEXT("CustomizedUV1");
+		case MP_CustomizedUVs2:         return TEXT("CustomizedUV2");
+		case MP_CustomizedUVs3:         return TEXT("CustomizedUV3");
+		case MP_CustomizedUVs4:         return TEXT("CustomizedUV4");
+		case MP_CustomizedUVs5:         return TEXT("CustomizedUV5");
+		case MP_CustomizedUVs6:         return TEXT("CustomizedUV6");
+		case MP_CustomizedUVs7:         return TEXT("CustomizedUV7");
+		case MP_PixelDepthOffset:       return TEXT("PixelDepthOffset");
+		case MP_ShadingModel:           return TEXT("ShadingModel");
+		case MP_MaterialAttributes:     return TEXT("MaterialAttributes");
+		case MP_CustomOutput:           return TEXT("ClearCoatBottomNormal");
+		default:                        checkNoEntry();
+	}
+
+	return TEXT("");
+}
+
 FVector4 FGLTFMaterialUtility::GetPropertyDefaultValue(EMaterialProperty Property)
 {
 	// TODO: replace with GMaterialPropertyAttributesMap lookup (when public API available)
diff --git a/Engine/Plugins/Enterprise/GLTFExporter/MigratedFromGit/Source/GLTFExporter/Private/Converters/GLTFMaterialUtility.h b/Engine/Plugins/Enterprise/GLTFExporter/MigratedFromGit/Source/GLTFExporter/Private/Converters/GLTFMaterialUtility.h
--- a/Engine/Plugins/Enterprise/GLTFExporter/MigratedFromGit/Source/GLTFExporter/Private/Converters/GLTFMaterialUtility.h
+++ b/Engine/Plugins/Enterprise/GLTFExporter/MigratedFromGit/Source/GLTFExporter/Private/Converters/GLTFMaterialUtility.h
@@ -34,6 +34,7 @@ struct FGLTFPropertyBakeOutput
 
 struct FGLTFMaterialUtility
 {
+	static const TCHAR* GetPropertyName(EMaterialProperty Property);
 	static FVector4 GetPropertyDefaultValue(EMaterialProperty Property);
 
 	static const FExpressionInput* GetInputForProperty(const UMaterialInterface* Material, EMaterialProperty Property);
