[GLTFExporter] refactor(clearcoat): add clear coat input mask constants

commit d7b8ab586c6a8638dc06482ac6d969ebc6d1cab6
Date: Mon, 20 Jul 2020 17:53:05 +0200

    refactor(clearcoat): add clear coat input mask constants
    (cherry picked from commit 4e1e993f43307036c191354765c13cd211877703​)

#skipci

diff --git a/Engine/Plugins/Enterprise/GLTFExporter/MigratedFromGit/Source/GLTFExporter/Private/Converters/GLTFMaterialConverters.cpp b/Engine/Plugins/Enterprise/GLTFExporter/MigratedFromGit/Source/GLTFExporter/Private/Converters/GLTFMaterialConverters.cpp
--- a/Engine/Plugins/Enterprise/GLTFExporter/MigratedFromGit/Source/GLTFExporter/Private/Converters/GLTFMaterialConverters.cpp
+++ b/Engine/Plugins/Enterprise/GLTFExporter/MigratedFromGit/Source/GLTFExporter/Private/Converters/GLTFMaterialConverters.cpp
@@ -29,6 +29,8 @@ namespace
 	const FLinearColor MetallicMask = BlueMask;
 	const FLinearColor RoughnessMask = GreenMask;
 	const FLinearColor OcclusionMask = RedMask;
+	const FLinearColor ClearCoatMask = RedMask;
+	const FLinearColor ClearCoatRoughnessMask = GreenMask;
 
 	// Ideal masks for texture-inputs (doesn't require baking)
 	const TArray<FLinearColor> DefaultColorInputMasks = { RgbMask, RgbaMask };
@@ -37,6 +39,8 @@ namespace
 	const TArray<FLinearColor> MetallicInputMasks = { MetallicMask };
 	const TArray<FLinearColor> RoughnessInputMasks = { RoughnessMask };
 	const TArray<FLinearColor> OcclusionInputMasks = { OcclusionMask };
+	const TArray<FLinearColor> ClearCoatInputMasks = { ClearCoatMask };
+	const TArray<FLinearColor> ClearCoatRoughnessInputMasks = { ClearCoatRoughnessMask };
 } // anonymous namespace
 
 FGLTFJsonMaterialIndex FGLTFMaterialConverter::Add(FGLTFConvertBuilder& Builder, const FString& Name, const UMaterialInterface* Material)
@@ -133,7 +137,7 @@ FGLTFJsonMaterialIndex FGLTFMaterialConverter::Add(FGLTFConvertBuilder& Builder,
 					const EMaterialProperty ClearCoatProperty = MP_CustomData0;
 					if (!TryGetConstantScalar(JsonMaterial.ClearCoat.ClearCoatFactor, ClearCoatProperty, Material))
 					{
-						if (!TryGetSourceTexture(Builder, JsonMaterial.ClearCoat.ClearCoatTexture, ClearCoatProperty, Material, { RedMask }))
+						if (!TryGetSourceTexture(Builder, JsonMaterial.ClearCoat.ClearCoatTexture, ClearCoatProperty, Material, ClearCoatInputMasks))
 						{
 							if (!TryGetBakedMaterialProperty(Builder, JsonMaterial.ClearCoat.ClearCoatTexture, JsonMaterial.ClearCoat.ClearCoatFactor, ClearCoatProperty, TEXT("ClearCoat"), Material))
 							{
@@ -149,7 +153,7 @@ FGLTFJsonMaterialIndex FGLTFMaterialConverter::Add(FGLTFConvertBuilder& Builder,
 					const EMaterialProperty ClearCoatRoughnessProperty = MP_CustomData1;
 					if (!TryGetConstantScalar(JsonMaterial.ClearCoat.ClearCoatRoughnessFactor, ClearCoatRoughnessProperty, Material))
 					{
-						if (!TryGetSourceTexture(Builder, JsonMaterial.ClearCoat.ClearCoatRoughnessTexture, ClearCoatRoughnessProperty, Material, { GreenMask }))
+						if (!TryGetSourceTexture(Builder, JsonMaterial.ClearCoat.ClearCoatRoughnessTexture, ClearCoatRoughnessProperty, Material, ClearCoatRoughnessInputMasks))
 						{
 							if (!TryGetBakedMaterialProperty(Builder, JsonMaterial.ClearCoat.ClearCoatRoughnessTexture, JsonMaterial.ClearCoat.ClearCoatRoughnessFactor, ClearCoatRoughnessProperty, TEXT("ClearCoatRoughness"), Material))
 							{
