[GLTFExporter] refactor(clearcoat): combine clearcoat intensity and roughness into single texture to increase compability with some viewers that assume the properties share the same texture.

commit e4dc2470e996ee406605b1925eeebfddb9ec546a
Date: Mon, 20 Jul 2020 18:23:36 +0200

    refactor(clearcoat): combine clearcoat intensity and roughness into single texture to increase compability with some viewers that assume the properties share the same texture.
    (cherry picked from commit 26517d48cdb4995e12cff0e34868a69aa3eb96e8​)

#skipci

diff --git a/Engine/Plugins/Enterprise/GLTFExporter/MigratedFromGit/Source/GLTFExporter/Private/Converters/GLTFMaterialConverters.cpp b/Engine/Plugins/Enterprise/GLTFExporter/MigratedFromGit/Source/GLTFExporter/Private/Converters/GLTFMaterialConverters.cpp
--- a/Engine/Plugins/Enterprise/GLTFExporter/MigratedFromGit/Source/GLTFExporter/Private/Converters/GLTFMaterialConverters.cpp
+++ b/Engine/Plugins/Enterprise/GLTFExporter/MigratedFromGit/Source/GLTFExporter/Private/Converters/GLTFMaterialConverters.cpp
@@ -132,49 +132,22 @@ FGLTFJsonMaterialIndex FGLTFMaterialConverter::Add(FGLTFConvertBuilder& Builder,
 
 			if (JsonMaterial.ShadingModel == EGLTFJsonShadingModel::ClearCoat)
 			{
-				if (JsonMaterial.ShadingModel == EGLTFJsonShadingModel::ClearCoat)
-				{
-					const EMaterialProperty ClearCoatProperty = MP_CustomData0;
-					if (!TryGetConstantScalar(JsonMaterial.ClearCoat.ClearCoatFactor, ClearCoatProperty, Material))
-					{
-						if (!TryGetSourceTexture(Builder, JsonMaterial.ClearCoat.ClearCoatTexture, ClearCoatProperty, Material, ClearCoatInputMasks))
-						{
-							if (!TryGetBakedMaterialProperty(Builder, JsonMaterial.ClearCoat.ClearCoatTexture, JsonMaterial.ClearCoat.ClearCoatFactor, ClearCoatProperty, TEXT("ClearCoat"), Material))
-							{
-								Builder.AddWarningMessage(FString::Printf(TEXT("Failed to export %s for material %s"), FGLTFMaterialUtility::GetPropertyName(ClearCoatProperty), *Material->GetName()));
-							}
-						}
-						else
-						{
-							JsonMaterial.ClearCoat.ClearCoatFactor = 1; // make sure factor is not zero
-						}
-					}
+				const EMaterialProperty ClearCoatProperty = MP_CustomData0;
+				const EMaterialProperty ClearCoatRoughnessProperty = MP_CustomData1;
 
-					const EMaterialProperty ClearCoatRoughnessProperty = MP_CustomData1;
-					if (!TryGetConstantScalar(JsonMaterial.ClearCoat.ClearCoatRoughnessFactor, ClearCoatRoughnessProperty, Material))
-					{
-						if (!TryGetSourceTexture(Builder, JsonMaterial.ClearCoat.ClearCoatRoughnessTexture, ClearCoatRoughnessProperty, Material, ClearCoatRoughnessInputMasks))
-						{
-							if (!TryGetBakedMaterialProperty(Builder, JsonMaterial.ClearCoat.ClearCoatRoughnessTexture, JsonMaterial.ClearCoat.ClearCoatRoughnessFactor, ClearCoatRoughnessProperty, TEXT("ClearCoatRoughness"), Material))
-							{
-								Builder.AddWarningMessage(FString::Printf(TEXT("Failed to export %s for material %s"), FGLTFMaterialUtility::GetPropertyName(ClearCoatRoughnessProperty), *Material->GetName()));
-							}
-						}
-						else
-						{
-							JsonMaterial.ClearCoat.ClearCoatRoughnessFactor = 1; // make sure factor is not zero
-						}
-					}
+				if (!TryGetClearCoatRoughness(Builder, JsonMaterial.ClearCoat, Material, ClearCoatProperty, ClearCoatRoughnessProperty))
+				{
+					Builder.AddWarningMessage(FString::Printf(TEXT("Failed to export %s and %s for material %s"), FGLTFMaterialUtility::GetPropertyName(ClearCoatProperty), FGLTFMaterialUtility::GetPropertyName(ClearCoatRoughnessProperty), *Material->GetName()));
+				}
 
-					const EMaterialProperty ClearCoatNormalProperty = MP_Normal;
-					if (IsPropertyNonDefault(ClearCoatNormalProperty, Material))
+				const EMaterialProperty ClearCoatNormalProperty = MP_Normal;
+				if (IsPropertyNonDefault(ClearCoatNormalProperty, Material))
+				{
+					if (!TryGetSourceTexture(Builder, JsonMaterial.ClearCoat.ClearCoatNormalTexture, ClearCoatNormalProperty, Material, DefaultColorInputMasks))
 					{
-						if (!TryGetSourceTexture(Builder, JsonMaterial.ClearCoat.ClearCoatNormalTexture, ClearCoatNormalProperty, Material, DefaultColorInputMasks))
+						if (!TryGetBakedMaterialProperty(Builder, JsonMaterial.ClearCoat.ClearCoatNormalTexture, ClearCoatNormalProperty, TEXT("ClearCoatNormal"), Material))
 						{
-							if (!TryGetBakedMaterialProperty(Builder, JsonMaterial.ClearCoat.ClearCoatNormalTexture, ClearCoatNormalProperty, TEXT("ClearCoatNormal"), Material))
-							{
-								Builder.AddWarningMessage(FString::Printf(TEXT("Failed to export %s for material %s"), FGLTFMaterialUtility::GetPropertyName(ClearCoatNormalProperty), *Material->GetName()));
-							}
+							Builder.AddWarningMessage(FString::Printf(TEXT("Failed to export %s for material %s"), FGLTFMaterialUtility::GetPropertyName(ClearCoatNormalProperty), *Material->GetName()));
 						}
 					}
 				}
@@ -494,6 +467,149 @@ bool FGLTFMaterialConverter::TryGetMetallicAndRoughness(FGLTFConvertBuilder& Bui
 	return true;
 }
 
+bool FGLTFMaterialConverter::TryGetClearCoatRoughness(FGLTFConvertBuilder& Builder, FGLTFJsonClearCoatExtension& OutExtParams, const UMaterialInterface* Material, EMaterialProperty IntensityProperty, EMaterialProperty RoughnessProperty) const
+{
+	const bool bIsIntensityConstant = TryGetConstantScalar(OutExtParams.ClearCoatFactor, IntensityProperty, Material);
+	const bool bIsRoughnessConstant = TryGetConstantScalar(OutExtParams.ClearCoatRoughnessFactor, RoughnessProperty, Material);
+
+	if (bIsIntensityConstant && bIsRoughnessConstant)
+	{
+		return true;
+	}
+
+	// NOTE: since we always bake the properties (for now) when atleast one property is non-const, we need
+	// to reset the constant factors to their defaults. Otherwise the baked value of a constant property
+	// would be scaled with the factor, i.e a double scaling.
+	OutExtParams.ClearCoatFactor = 1.0f;
+	OutExtParams.ClearCoatRoughnessFactor = 1.0f;
+
+	const UTexture2D* IntensityTexture;
+	const UTexture2D* RoughnessTexture;
+	int32 IntensityTexCoord;
+	int32 RoughnessTexCoord;
+
+	const bool bHasIntensitySourceTexture = TryGetSourceTexture(IntensityTexture, IntensityTexCoord, IntensityProperty, Material, ClearCoatInputMasks);
+	const bool bHasRoughnessSourceTexture = TryGetSourceTexture(RoughnessTexture, RoughnessTexCoord, RoughnessProperty, Material, ClearCoatRoughnessInputMasks);
+
+	// Detect the "happy path" where both inputs share the same texture and are correctly masked.
+	if (bHasIntensitySourceTexture &&
+		bHasRoughnessSourceTexture &&
+		IntensityTexture == RoughnessTexture &&
+		IntensityTexCoord == RoughnessTexCoord)
+	{
+		const FGLTFJsonTextureIndex TextureIndex = Builder.GetOrAddTexture(IntensityTexture);
+		OutExtParams.ClearCoatTexture.Index = TextureIndex;
+		OutExtParams.ClearCoatTexture.TexCoord = IntensityTexCoord;
+		OutExtParams.ClearCoatRoughnessTexture.Index = TextureIndex;
+		OutExtParams.ClearCoatRoughnessTexture.TexCoord = RoughnessTexCoord;
+		return true;
+	}
+
+	// TODO: make default baking-resolution configurable
+	// TODO: add support for detecting the correct tex-coord for this property based on connected nodes
+	// TODO: add support for calculating the ideal resolution to use for baking based on connected (texture) nodes
+	int32 TexCoord = 0;
+	FIntPoint TextureSize(1024, 1024);
+
+	// TODO: should this be the default wrap-mode?
+	EGLTFJsonTextureWrap TextureWrapS = EGLTFJsonTextureWrap::Repeat;
+	EGLTFJsonTextureWrap TextureWrapT = EGLTFJsonTextureWrap::Repeat;
+
+	// TODO: should this be the default filter?
+	EGLTFJsonTextureFilter TextureMinFilter = EGLTFJsonTextureFilter::LinearMipmapLinear;
+	EGLTFJsonTextureFilter TextureMagFilter = EGLTFJsonTextureFilter::Linear;
+
+	if (bHasIntensitySourceTexture && bHasRoughnessSourceTexture)
+	{
+		const bool bAreTexturesCompatible = IntensityTexCoord == RoughnessTexCoord &&
+			IntensityTexture->AddressX == RoughnessTexture->AddressX &&
+			IntensityTexture->AddressY == RoughnessTexture->AddressY;
+
+		if (!bAreTexturesCompatible)
+		{
+			// TODO: handle differences in wrapping or uv-coords
+			Builder.AddWarningMessage(FString::Printf(
+				TEXT("Intensity- and Roughness-textures for material %s were not able to be combined and will be skipped"),
+				*Material->GetName()));
+
+			return false;
+		}
+
+		TexCoord = IntensityTexCoord;
+		TextureSize =
+		{
+			FMath::Max(IntensityTexture->GetSizeX(), RoughnessTexture->GetSizeX()),
+			FMath::Max(IntensityTexture->GetSizeY(), RoughnessTexture->GetSizeY())
+		};
+
+		TextureWrapS = FGLTFConverterUtility::ConvertWrap(IntensityTexture->AddressX);
+		TextureWrapT = FGLTFConverterUtility::ConvertWrap(IntensityTexture->AddressY);
+
+		// TODO: compare min- and mag-filter for BaseColorTexture and OpacityTexture. If they differ,
+		// we should choose one or the other and inform the user about the choice made by logging to the console.
+		TextureMinFilter = FGLTFConverterUtility::ConvertMinFilter(IntensityTexture->Filter, IntensityTexture->LODGroup);
+		TextureMagFilter = FGLTFConverterUtility::ConvertMagFilter(IntensityTexture->Filter, IntensityTexture->LODGroup);
+	}
+	else if (bHasIntensitySourceTexture)
+	{
+		TexCoord = IntensityTexCoord;
+		TextureSize = { IntensityTexture->GetSizeX(), IntensityTexture->GetSizeY() };
+		TextureWrapS = FGLTFConverterUtility::ConvertWrap(IntensityTexture->AddressX);
+		TextureWrapT = FGLTFConverterUtility::ConvertWrap(IntensityTexture->AddressY);
+		TextureMinFilter = FGLTFConverterUtility::ConvertMinFilter(IntensityTexture->Filter, IntensityTexture->LODGroup);
+		TextureMagFilter = FGLTFConverterUtility::ConvertMagFilter(IntensityTexture->Filter, IntensityTexture->LODGroup);
+	}
+	else if (bHasRoughnessSourceTexture)
+	{
+		TexCoord = RoughnessTexCoord;
+		TextureSize = { RoughnessTexture->GetSizeX(), RoughnessTexture->GetSizeY() };
+		TextureWrapS = FGLTFConverterUtility::ConvertWrap(RoughnessTexture->AddressX);
+		TextureWrapT = FGLTFConverterUtility::ConvertWrap(RoughnessTexture->AddressY);
+		TextureMinFilter = FGLTFConverterUtility::ConvertMinFilter(RoughnessTexture->Filter, RoughnessTexture->LODGroup);
+		TextureMagFilter = FGLTFConverterUtility::ConvertMagFilter(RoughnessTexture->Filter, RoughnessTexture->LODGroup);
+	}
+
+	const FGLTFPropertyBakeOutput IntensityBakeOutput = BakeMaterialProperty(IntensityProperty, Material, &TextureSize);
+	const FGLTFPropertyBakeOutput RoughnessBakeOutput = BakeMaterialProperty(RoughnessProperty, Material, &TextureSize);
+
+	// Detect when both baked properties are constants, which means we can use factors and avoid exporting a texture
+	if (IntensityBakeOutput.bIsConstant && RoughnessBakeOutput.bIsConstant)
+	{
+		OutExtParams.ClearCoatFactor = IntensityBakeOutput.ConstantValue.R;
+		OutExtParams.ClearCoatRoughnessFactor =  RoughnessBakeOutput.ConstantValue.R;
+		return true;
+	}
+
+	TextureSize = RoughnessBakeOutput.Size.ComponentMax(IntensityBakeOutput.Size);
+	IntensityTexture = FGLTFMaterialUtility::CreateTransientTexture(IntensityBakeOutput);
+	RoughnessTexture = FGLTFMaterialUtility::CreateTransientTexture(RoughnessBakeOutput);
+
+	const FString TextureName = Material->GetName() + TEXT("_ClearCoatRoughness");
+
+	const TArray<FGLTFTextureCombineSource> CombineSources =
+	{
+		{ IntensityTexture, ClearCoatMask + AlphaMask, SE_BLEND_Opaque },
+		{ RoughnessTexture, ClearCoatRoughnessMask }
+	};
+
+	const FGLTFJsonTextureIndex TextureIndex = FGLTFMaterialUtility::AddCombinedTexture(
+		Builder,
+		CombineSources,
+		TextureSize,
+		TextureName,
+		TextureMinFilter,
+		TextureMagFilter,
+		TextureWrapS,
+		TextureWrapT);
+
+	OutExtParams.ClearCoatTexture.Index = TextureIndex;
+	OutExtParams.ClearCoatTexture.TexCoord = TexCoord;
+	OutExtParams.ClearCoatRoughnessTexture.Index = TextureIndex;
+	OutExtParams.ClearCoatRoughnessTexture.TexCoord = TexCoord;
+
+	return true;
+}
+
 bool FGLTFMaterialConverter::TryGetEmissive(FGLTFConvertBuilder& Builder, FGLTFJsonMaterial& JsonMaterial, EMaterialProperty EmissiveProperty, const UMaterialInterface* Material) const
 {
 	// TODO: right now we allow EmissiveFactor to be > 1.0 to support very bright emission, although it's not valid according to the glTF standard.
diff --git a/Engine/Plugins/Enterprise/GLTFExporter/MigratedFromGit/Source/GLTFExporter/Private/Converters/GLTFMaterialConverters.h b/Engine/Plugins/Enterprise/GLTFExporter/MigratedFromGit/Source/GLTFExporter/Private/Converters/GLTFMaterialConverters.h
--- a/Engine/Plugins/Enterprise/GLTFExporter/MigratedFromGit/Source/GLTFExporter/Private/Converters/GLTFMaterialConverters.h
+++ b/Engine/Plugins/Enterprise/GLTFExporter/MigratedFromGit/Source/GLTFExporter/Private/Converters/GLTFMaterialConverters.h
@@ -17,6 +17,7 @@ class FGLTFMaterialConverter final : public TGLTFConverter<FGLTFJsonMaterialInde
 
 	bool TryGetBaseColorAndOpacity(FGLTFConvertBuilder& Builder, FGLTFJsonPBRMetallicRoughness& OutPBRParams, const UMaterialInterface* Material, EMaterialProperty BaseColorProperty, EMaterialProperty OpacityProperty) const;
 	bool TryGetMetallicAndRoughness(FGLTFConvertBuilder& Builder, FGLTFJsonPBRMetallicRoughness& OutPBRParams, const UMaterialInterface* Material, EMaterialProperty MetallicProperty, EMaterialProperty RoughnessProperty) const;
+	bool TryGetClearCoatRoughness(FGLTFConvertBuilder& Builder, FGLTFJsonClearCoatExtension& OutExtParams, const UMaterialInterface* Material, EMaterialProperty IntensityProperty, EMaterialProperty RoughnessProperty) const;
 	bool TryGetEmissive(FGLTFConvertBuilder& Builder, FGLTFJsonMaterial& JsonMaterial, EMaterialProperty EmissiveProperty, const UMaterialInterface* Material) const;
 
 	bool IsPropertyNonDefault(EMaterialProperty Property, const UMaterialInterface* Material) const;
