[GLTFExporter] fix(material): improve comparison of baked constant with default value for texture-only properties

commit 40206dca5f3f734540f18aaaaf70ee42afe2fd99
Date: Mon, 20 Jul 2020 22:25:32 +0200

    fix(material): improve comparison of baked constant with default value for texture-only properties
    (cherry picked from commit e1254b3dea3959cb3be5a0f5ec6e6ec0601c9e75​)

#skipci

diff --git a/Engine/Plugins/Enterprise/GLTFExporter/MigratedFromGit/Source/GLTFExporter/Private/Converters/GLTFMaterialConverters.cpp b/Engine/Plugins/Enterprise/GLTFExporter/MigratedFromGit/Source/GLTFExporter/Private/Converters/GLTFMaterialConverters.cpp
--- a/Engine/Plugins/Enterprise/GLTFExporter/MigratedFromGit/Source/GLTFExporter/Private/Converters/GLTFMaterialConverters.cpp
+++ b/Engine/Plugins/Enterprise/GLTFExporter/MigratedFromGit/Source/GLTFExporter/Private/Converters/GLTFMaterialConverters.cpp
@@ -1068,7 +1068,8 @@ bool FGLTFMaterialConverter::TryGetBakedMaterialProperty(FGLTFConvertBuilder& Bu
 		return StoreBakedPropertyTexture(Builder, OutTexInfo, PropertyBakeOutput, PropertyName, Material);
 	}
 
-	if (PropertyBakeOutput.ConstantValue == FLinearColor(FGLTFMaterialUtility::GetPropertyDefaultValue(Property)))
+	const FVector4 MaskedConstant = FVector4(PropertyBakeOutput.ConstantValue) * FGLTFMaterialUtility::GetPropertyMask(Property);
+	if (MaskedConstant == FGLTFMaterialUtility::GetPropertyDefaultValue(Property))
 	{
 		// Constant value is the same as the property's default so we can set gltf to default.
 		OutTexInfo.Index = FGLTFJsonTextureIndex(INDEX_NONE);
diff --git a/Engine/Plugins/Enterprise/GLTFExporter/MigratedFromGit/Source/GLTFExporter/Private/Converters/GLTFMaterialUtility.cpp b/Engine/Plugins/Enterprise/GLTFExporter/MigratedFromGit/Source/GLTFExporter/Private/Converters/GLTFMaterialUtility.cpp
--- a/Engine/Plugins/Enterprise/GLTFExporter/MigratedFromGit/Source/GLTFExporter/Private/Converters/GLTFMaterialUtility.cpp
+++ b/Engine/Plugins/Enterprise/GLTFExporter/MigratedFromGit/Source/GLTFExporter/Private/Converters/GLTFMaterialUtility.cpp
@@ -97,6 +97,50 @@ FVector4 FGLTFMaterialUtility::GetPropertyDefaultValue(EMaterialProperty Propert
 	return FVector4();
 }
 
+FVector4 FGLTFMaterialUtility::GetPropertyMask(EMaterialProperty Property)
+{
+	// TODO: replace with GMaterialPropertyAttributesMap lookup (when public API available)
+
+	switch (Property)
+	{
+		case MP_EmissiveColor:          return FVector4(1,1,1,0);
+		case MP_Opacity:                return FVector4(1,0,0,0);
+		case MP_OpacityMask:            return FVector4(1,0,0,0);
+		case MP_BaseColor:              return FVector4(1,1,1,0);
+		case MP_Metallic:               return FVector4(1,0,0,0);
+		case MP_Specular:               return FVector4(1,0,0,0);
+		case MP_Roughness:              return FVector4(1,0,0,0);
+		case MP_Anisotropy:             return FVector4(1,0,0,0);
+		case MP_Normal:                 return FVector4(1,1,1,0);
+		case MP_Tangent:                return FVector4(1,1,1,0);
+		case MP_WorldPositionOffset:    return FVector4(1,1,1,0);
+		case MP_WorldDisplacement:      return FVector4(1,1,1,0);
+		case MP_TessellationMultiplier: return FVector4(1,0,0,0);
+		case MP_SubsurfaceColor:        return FVector4(1,1,1,0);
+		case MP_CustomData0:            return FVector4(1,0,0,0);
+		case MP_CustomData1:            return FVector4(1,0,0,0);
+		case MP_AmbientOcclusion:       return FVector4(1,0,0,0);
+		case MP_Refraction:             return FVector4(1,1,0,0);
+		case MP_PixelDepthOffset:       return FVector4(1,0,0,0);
+		case MP_ShadingModel:           return FVector4(1,0,0,0);
+		default:                        break;
+	}
+
+	if (Property >= MP_CustomizedUVs0 && Property <= MP_CustomizedUVs7)
+	{
+		return FVector4(1,1,0,0);
+	}
+
+	// TODO: replace workaround for ClearCoatBottomNormal
+	if (Property == MP_CustomOutput)
+	{
+		return FVector4(1,1,1,0);
+	}
+
+	check(false);
+	return FVector4();
+}
+
 const FExpressionInput* FGLTFMaterialUtility::GetInputForProperty(const UMaterialInterface* Material, EMaterialProperty Property)
 {
 	// TODO: replace workaround for ClearCoatBottomNormal
diff --git a/Engine/Plugins/Enterprise/GLTFExporter/MigratedFromGit/Source/GLTFExporter/Private/Converters/GLTFMaterialUtility.h b/Engine/Plugins/Enterprise/GLTFExporter/MigratedFromGit/Source/GLTFExporter/Private/Converters/GLTFMaterialUtility.h
--- a/Engine/Plugins/Enterprise/GLTFExporter/MigratedFromGit/Source/GLTFExporter/Private/Converters/GLTFMaterialUtility.h
+++ b/Engine/Plugins/Enterprise/GLTFExporter/MigratedFromGit/Source/GLTFExporter/Private/Converters/GLTFMaterialUtility.h
@@ -36,6 +36,7 @@ struct FGLTFMaterialUtility
 {
 	static const TCHAR* GetPropertyName(EMaterialProperty Property);
 	static FVector4 GetPropertyDefaultValue(EMaterialProperty Property);
+	static FVector4 GetPropertyMask(EMaterialProperty Property);
 
 	static const FExpressionInput* GetInputForProperty(const UMaterialInterface* Material, EMaterialProperty Property);
 
