[GLTFExporter] fix(material): if material normal input is baked to a constant, assume it's always default for now

commit 4701cd71dd5b5e807808a6d82954b2a29c9b5390
Date: Mon, 20 Jul 2020 22:29:05 +0200

    fix(material): if material normal input is baked to a constant, assume it's always default for now
    (cherry picked from commit 413241eeba3521a81df01e79a974fe4b0547d372​)

#skipci

diff --git a/Engine/Plugins/Enterprise/GLTFExporter/MigratedFromGit/Source/GLTFExporter/Private/Converters/GLTFMaterialConverters.cpp b/Engine/Plugins/Enterprise/GLTFExporter/MigratedFromGit/Source/GLTFExporter/Private/Converters/GLTFMaterialConverters.cpp
--- a/Engine/Plugins/Enterprise/GLTFExporter/MigratedFromGit/Source/GLTFExporter/Private/Converters/GLTFMaterialConverters.cpp
+++ b/Engine/Plugins/Enterprise/GLTFExporter/MigratedFromGit/Source/GLTFExporter/Private/Converters/GLTFMaterialConverters.cpp
@@ -1076,6 +1076,15 @@ bool FGLTFMaterialConverter::TryGetBakedMaterialProperty(FGLTFConvertBuilder& Bu
 		return true;
 	}
 
+	if (Property == MP_Normal || Property == MP_CustomOutput /* ClearCoatBottomNormal */)
+	{
+		// TODO: In some cases baking normal can result in constant vector that differs slight from default (i.e 0,0,1).
+		// Yet often, when looking at such a material, it should be exactly default. Needs further investigation.
+		// Maybe because of incorrect sRGB conversion? For now, assume a constant normal is always default.
+		OutTexInfo.Index = FGLTFJsonTextureIndex(INDEX_NONE);
+		return true;
+	}
+
 	// NOTE: since this function is meant to bake to a texture, we assume that the property
 	// that was passed into the function is using a non-constant expression.
 	// We therefore treat a constant result with non-default value as a failure.
