[GLTFExporter] fix(material): if a texture-only non-normalmap material property results in a non-default constant, export it as a 1x1 image.

commit da95576645527bdda2ef6829b53672481eabf799
Date: Mon, 20 Jul 2020 22:30:38 +0200

    fix(material): if a texture-only non-normalmap material property results in a non-default constant, export it as a 1x1 image.
    (cherry picked from commit fe506fbaa34d67fe62cfb6aa08904348ac6f3d1d​)

#skipci

diff --git a/Engine/Plugins/Enterprise/GLTFExporter/MigratedFromGit/Source/GLTFExporter/Private/Converters/GLTFMaterialConverters.cpp b/Engine/Plugins/Enterprise/GLTFExporter/MigratedFromGit/Source/GLTFExporter/Private/Converters/GLTFMaterialConverters.cpp
--- a/Engine/Plugins/Enterprise/GLTFExporter/MigratedFromGit/Source/GLTFExporter/Private/Converters/GLTFMaterialConverters.cpp
+++ b/Engine/Plugins/Enterprise/GLTFExporter/MigratedFromGit/Source/GLTFExporter/Private/Converters/GLTFMaterialConverters.cpp
@@ -1085,18 +1085,25 @@ bool FGLTFMaterialConverter::TryGetBakedMaterialProperty(FGLTFConvertBuilder& Bu
 		return true;
 	}
 
-	// NOTE: since this function is meant to bake to a texture, we assume that the property
-	// that was passed into the function is using a non-constant expression.
-	// We therefore treat a constant result with non-default value as a failure.
-
-	// NOTE: in some cases a constant baking result is returned for a property that is non-constant.
-	// This happens (for example) when baking AmbientOcclusion for a translucent material,
+	// TODO: let function fail and investigate why in some cases a constant baking result is returned for a property
+	// that is non-constant. This happens (for example) when baking AmbientOcclusion for a translucent material,
 	// even though the same material when set to opaque will properly bake AmbientOcclusion to a texture.
-	// It also happens when baking Normal in some (not yet identified) circumstances.
-	// For now, these incorrect bakes are discarded.
+	// For now, create a 1x1 texture with the constant value.
 
-	// TODO: investigate why non-constant properties are sometimes baked to a constant expression (see note above)
-	return false;
+	const FString TextureName = Material->GetName() + TEXT("_") + PropertyName;
+	const FGLTFJsonTextureIndex TextureIndex = FGLTFMaterialUtility::AddTexture(
+		Builder,
+	    PropertyBakeOutput.Pixels,
+	    PropertyBakeOutput.Size,
+		Material->GetName() + TEXT("_") + PropertyName,
+		EGLTFJsonTextureFilter::Nearest,
+		EGLTFJsonTextureFilter::Nearest,
+		EGLTFJsonTextureWrap::ClampToEdge,
+        EGLTFJsonTextureWrap::ClampToEdge);
+
+	OutTexInfo.TexCoord = 0;
+	OutTexInfo.Index = TextureIndex;
+	return true;
 }
 
 FGLTFPropertyBakeOutput FGLTFMaterialConverter::BakeMaterialProperty(EMaterialProperty Property, const UMaterialInterface* Material, const FIntPoint* PreferredTextureSize, bool bCopyAlphaFromRedChannel) const
@@ -1136,7 +1143,6 @@ bool FGLTFMaterialConverter::StoreBakedPropertyTexture(FGLTFConvertBuilder& Buil
 		PropertyBakeOutput.Pixels,
 		PropertyBakeOutput.Size,
 		TextureName,
-		PropertyBakeOutput.PixelFormat,
 		TextureMinFilter,
 		TextureMagFilter,
 		TextureWrapS,
diff --git a/Engine/Plugins/Enterprise/GLTFExporter/MigratedFromGit/Source/GLTFExporter/Private/Converters/GLTFMaterialUtility.cpp b/Engine/Plugins/Enterprise/GLTFExporter/MigratedFromGit/Source/GLTFExporter/Private/Converters/GLTFMaterialUtility.cpp
--- a/Engine/Plugins/Enterprise/GLTFExporter/MigratedFromGit/Source/GLTFExporter/Private/Converters/GLTFMaterialUtility.cpp
+++ b/Engine/Plugins/Enterprise/GLTFExporter/MigratedFromGit/Source/GLTFExporter/Private/Converters/GLTFMaterialUtility.cpp
@@ -328,17 +328,17 @@ FGLTFJsonTextureIndex FGLTFMaterialUtility::AddCombinedTexture(FGLTFConvertBuild
 	check(CombineSources.Num() > 0);
 
 	TArray<FColor> Pixels;
-	const EPixelFormat PixelFormat = CombineSources[0].Texture->GetPixelFormat();
+	const EPixelFormat PixelFormat = CombineSources[0].Texture->GetPixelFormat(); // TODO: should we really assume pixel format like this?
 
 	if (!CombineTextures(Pixels, CombineSources, TextureSize, PixelFormat))
 	{
 		return FGLTFJsonTextureIndex(INDEX_NONE);
 	}
 
-	return AddTexture(Builder, Pixels, TextureSize, TextureName, PixelFormat, MinFilter, MagFilter, WrapS, WrapT);
+	return AddTexture(Builder, Pixels, TextureSize, TextureName, MinFilter, MagFilter, WrapS, WrapT);
 }
 
-FGLTFJsonTextureIndex FGLTFMaterialUtility::AddTexture(FGLTFConvertBuilder& Builder, const TArray<FColor>& Pixels, const FIntPoint& TextureSize, const FString& TextureName, EPixelFormat PixelFormat, EGLTFJsonTextureFilter MinFilter, EGLTFJsonTextureFilter MagFilter, EGLTFJsonTextureWrap WrapS, EGLTFJsonTextureWrap WrapT)
+FGLTFJsonTextureIndex FGLTFMaterialUtility::AddTexture(FGLTFConvertBuilder& Builder, const TArray<FColor>& Pixels, const FIntPoint& TextureSize, const FString& TextureName, EGLTFJsonTextureFilter MinFilter, EGLTFJsonTextureFilter MagFilter, EGLTFJsonTextureWrap WrapS, EGLTFJsonTextureWrap WrapT)
 {
 	// TODO: maybe we should reuse existing samplers?
 	FGLTFJsonSampler JsonSampler;
diff --git a/Engine/Plugins/Enterprise/GLTFExporter/MigratedFromGit/Source/GLTFExporter/Private/Converters/GLTFMaterialUtility.h b/Engine/Plugins/Enterprise/GLTFExporter/MigratedFromGit/Source/GLTFExporter/Private/Converters/GLTFMaterialUtility.h
--- a/Engine/Plugins/Enterprise/GLTFExporter/MigratedFromGit/Source/GLTFExporter/Private/Converters/GLTFMaterialUtility.h
+++ b/Engine/Plugins/Enterprise/GLTFExporter/MigratedFromGit/Source/GLTFExporter/Private/Converters/GLTFMaterialUtility.h
@@ -54,7 +54,7 @@ struct FGLTFMaterialUtility
 	static FGLTFPropertyBakeOutput BakeMaterialProperty(const FIntPoint& OutputSize, EMaterialProperty Property, const UMaterialInterface* Material, bool bCopyAlphaFromRedChannel = false);
 
 	static FGLTFJsonTextureIndex AddCombinedTexture(FGLTFConvertBuilder& Builder, const TArray<FGLTFTextureCombineSource>& CombineSources, const FIntPoint& TextureSize, const FString& TextureName, EGLTFJsonTextureFilter MinFilter, EGLTFJsonTextureFilter MagFilter, EGLTFJsonTextureWrap WrapS, EGLTFJsonTextureWrap WrapT);
-	static FGLTFJsonTextureIndex AddTexture(FGLTFConvertBuilder& Builder, const TArray<FColor>& Pixels, const FIntPoint& TextureSize, const FString& TextureName, EPixelFormat PixelFormat, EGLTFJsonTextureFilter MinFilter, EGLTFJsonTextureFilter MagFilter, EGLTFJsonTextureWrap WrapS, EGLTFJsonTextureWrap WrapT);
+	static FGLTFJsonTextureIndex AddTexture(FGLTFConvertBuilder& Builder, const TArray<FColor>& Pixels, const FIntPoint& TextureSize, const FString& TextureName, EGLTFJsonTextureFilter MinFilter, EGLTFJsonTextureFilter MagFilter, EGLTFJsonTextureWrap WrapS, EGLTFJsonTextureWrap WrapT);
 
 	static FLinearColor GetMask(const FExpressionInput& ExpressionInput);
 	static uint32 GetMaskComponentCount(const FExpressionInput& ExpressionInput);
