[GLTFExporter] refactor(texture): add TODOs for future improvements in Texture2D converter

commit de97ca8e0eb982f4c31f8591b86e73ea282e027d
Date: Tue, 21 Jul 2020 00:44:17 +0200

    refactor(texture): add TODOs for future improvements in Texture2D converter
    (cherry picked from commit d10ffc0cb88aeb057ee2373d58594b5f145826f3​)

#skipci

diff --git a/Engine/Plugins/Enterprise/GLTFExporter/MigratedFromGit/Source/GLTFExporter/Private/Converters/GLTFTextureConverters.cpp b/Engine/Plugins/Enterprise/GLTFExporter/MigratedFromGit/Source/GLTFExporter/Private/Converters/GLTFTextureConverters.cpp
--- a/Engine/Plugins/Enterprise/GLTFExporter/MigratedFromGit/Source/GLTFExporter/Private/Converters/GLTFTextureConverters.cpp
+++ b/Engine/Plugins/Enterprise/GLTFExporter/MigratedFromGit/Source/GLTFExporter/Private/Converters/GLTFTextureConverters.cpp
@@ -52,6 +52,7 @@ FGLTFJsonTextureIndex FGLTFTexture2DConverter::Add(FGLTFConvertBuilder& Builder,
 	ERGBFormat RGBFormat;
 	uint32 BitDepth;
 
+	// TODO: don't rely on Texture2D->Source, since the texture brightness etc can be tweaked for runtime
 	if (FGLTFTextureUtility::CanPNGCompressFormat(Texture2D->Source.GetFormat(), RGBFormat, BitDepth))
 	{
 		FTextureSource& Source = const_cast<FTextureSource&>(Texture2D->Source);
@@ -63,6 +64,7 @@ FGLTFJsonTextureIndex FGLTFTexture2DConverter::Add(FGLTFConvertBuilder& Builder,
 	else if (FGLTFTextureUtility::CanPNGCompressFormat(Texture2D->GetPixelFormat(), RGBFormat, BitDepth))
 	{
 		const FByteBulkData& BulkData = Texture2D->PlatformData->Mips[0].BulkData;
+		// TODO: ensure that BulkData size is not zero and if so rebuild platform data.
 
 		const void* RawData = BulkData.LockReadOnly();
 		ImageIndex = Builder.AddImage(RawData, BulkData.GetBulkDataSize(), Size, RGBFormat, BitDepth, JsonTexture.Name);
