[GLTFExporter] feat(camera): add json data structures with corresponding enums and utilities

commit 62e276f86c4110f40945c53a5031eec16fae9b6c
Date: Tue, 21 Jul 2020 07:34:05 +0200

    feat(camera): add json data structures with corresponding enums and utilities
    (cherry picked from commit fc32bbb7627ad11a47b98a316e72991be40c90c3​)

#skipci

diff --git a/Engine/Plugins/Enterprise/GLTFExporter/MigratedFromGit/Source/GLTFExporter/Private/Json/GLTFJsonCamera.h b/Engine/Plugins/Enterprise/GLTFExporter/MigratedFromGit/Source/GLTFExporter/Private/Json/GLTFJsonCamera.h
new file mode 100644
--- /dev/null
+++ b/Engine/Plugins/Enterprise/GLTFExporter/MigratedFromGit/Source/GLTFExporter/Private/Json/GLTFJsonCamera.h
@@ -0,0 +1,124 @@
+// Copyright Epic Games, Inc. All Rights Reserved.
+
+#pragma once
+
+#include "Json/GLTFJsonEnums.h"
+#include "Json/GLTFJsonUtility.h"
+#include "Serialization/JsonSerializer.h"
+
+struct GLTFEXPORTER_API FGLTFJsonOrthographic
+{
+	float XMag; // horizontal magnification of the view
+	float YMag; // vertical magnification of the view
+	float ZFar;
+	float ZNear;
+
+	FGLTFJsonOrthographic()
+		: XMag(0)
+		, YMag(0)
+		, ZFar(0)
+		, ZNear(0)
+	{
+	}
+
+	template <class CharType = TCHAR, class PrintPolicy = TPrettyJsonPrintPolicy<CharType>>
+	void WriteObject(TJsonWriter<CharType, PrintPolicy>& JsonWriter, FGLTFJsonExtensions& Extensions) const
+	{
+		JsonWriter.WriteObjectStart();
+
+		JsonWriter.WriteValue(TEXT("xmag"), XMag);
+		JsonWriter.WriteValue(TEXT("ymag"), YMag);
+		JsonWriter.WriteValue(TEXT("zfar"), ZFar);
+		JsonWriter.WriteValue(TEXT("znear"), ZNear);
+
+		JsonWriter.WriteObjectEnd();
+	}
+};
+
+struct GLTFEXPORTER_API FGLTFJsonPerspective
+{
+	float AspectRatio; // aspect ratio of the field of view
+	float YFov; // vertical field of view in radians
+	float ZFar;
+	float ZNear;
+
+	FGLTFJsonPerspective()
+		: AspectRatio(0)
+		, YFov(0)
+		, ZFar(0)
+		, ZNear(0)
+	{
+	}
+
+	template <class CharType = TCHAR, class PrintPolicy = TPrettyJsonPrintPolicy<CharType>>
+	void WriteObject(TJsonWriter<CharType, PrintPolicy>& JsonWriter, FGLTFJsonExtensions& Extensions) const
+	{
+		JsonWriter.WriteObjectStart();
+
+		if (AspectRatio != 0)
+		{
+			JsonWriter.WriteValue(TEXT("aspectRatio"), AspectRatio);
+		}
+
+		JsonWriter.WriteValue(TEXT("yfov"), YFov);
+
+		if (ZFar != 0)
+		{
+			JsonWriter.WriteValue(TEXT("zfar"), ZFar);
+		}
+
+		JsonWriter.WriteValue(TEXT("znear"), ZNear);
+
+		JsonWriter.WriteObjectEnd();
+	}
+};
+
+struct GLTFEXPORTER_API FGLTFJsonCamera
+{
+	FString Name;
+
+	EGLTFJsonCameraType Type;
+
+	union {
+		FGLTFJsonOrthographic Orthographic;
+		FGLTFJsonPerspective  Perspective;
+	};
+
+	FGLTFJsonCamera()
+		: Type(EGLTFJsonCameraType::None)
+		, Orthographic()
+		, Perspective()
+	{
+	}
+
+	template <class CharType = TCHAR, class PrintPolicy = TPrettyJsonPrintPolicy<CharType>>
+	void WriteObject(TJsonWriter<CharType, PrintPolicy>& JsonWriter, FGLTFJsonExtensions& Extensions) const
+	{
+		JsonWriter.WriteObjectStart();
+
+		if (!Name.IsEmpty())
+		{
+			JsonWriter.WriteValue(TEXT("name"), Name);
+		}
+
+		JsonWriter.WriteValue(TEXT("type"), FGLTFJsonUtility::ToString(Type));
+
+		switch (Type)
+		{
+			case EGLTFJsonCameraType::Orthographic:
+				JsonWriter.WriteIdentifierPrefix(TEXT("orthographic"));
+				Orthographic.WriteObject(JsonWriter, Extensions);
+				break;
+
+			case EGLTFJsonCameraType::Perspective:
+				JsonWriter.WriteIdentifierPrefix(TEXT("perspective"));
+				Perspective.WriteObject(JsonWriter, Extensions);
+				break;
+
+			default:
+				break;
+		}
+
+		JsonWriter.WriteObjectEnd();
+	}
+};
diff --git a/Engine/Plugins/Enterprise/GLTFExporter/MigratedFromGit/Source/GLTFExporter/Private/Json/GLTFJsonEnums.h b/Engine/Plugins/Enterprise/GLTFExporter/MigratedFromGit/Source/GLTFExporter/Private/Json/GLTFJsonEnums.h
--- a/Engine/Plugins/Enterprise/GLTFExporter/MigratedFromGit/Source/GLTFExporter/Private/Json/GLTFJsonEnums.h
+++ b/Engine/Plugins/Enterprise/GLTFExporter/MigratedFromGit/Source/GLTFExporter/Private/Json/GLTFJsonEnums.h
@@ -139,3 +139,10 @@ enum class EGLTFJsonAnimationPath
 	Scale,
 	Weights
 };
+
+enum class EGLTFJsonCameraType
+{
+	None = -1,
+	Orthographic,
+	Perspective
+};
diff --git a/Engine/Plugins/Enterprise/GLTFExporter/MigratedFromGit/Source/GLTFExporter/Private/Json/GLTFJsonRoot.h b/Engine/Plugins/Enterprise/GLTFExporter/MigratedFromGit/Source/GLTFExporter/Private/Json/GLTFJsonRoot.h
--- a/Engine/Plugins/Enterprise/GLTFExporter/MigratedFromGit/Source/GLTFExporter/Private/Json/GLTFJsonRoot.h
+++ b/Engine/Plugins/Enterprise/GLTFExporter/MigratedFromGit/Source/GLTFExporter/Private/Json/GLTFJsonRoot.h
@@ -7,6 +7,7 @@
 #include "Json/GLTFJsonAccessor.h"
 #include "Json/GLTFJsonBuffer.h"
 #include "Json/GLTFJsonBufferView.h"
+#include "Json/GLTFJsonCamera.h"
 #include "Json/GLTFJsonImage.h"
 #include "Json/GLTFJsonMaterial.h"
 #include "Json/GLTFJsonMesh.h"
@@ -60,6 +61,7 @@ struct FGLTFJsonRoot
 	TArray<FGLTFJsonAccessor>   Accessors;
 	TArray<FGLTFJsonBuffer>     Buffers;
 	TArray<FGLTFJsonBufferView> BufferViews;
+	TArray<FGLTFJsonCamera>     Cameras;
 	TArray<FGLTFJsonMaterial>   Materials;
 	TArray<FGLTFJsonMesh>       Meshes;
 	TArray<FGLTFJsonNode>       Nodes;
@@ -89,6 +91,7 @@ struct FGLTFJsonRoot
 		FGLTFJsonUtility::WriteObjectArray(JsonWriter, TEXT("accessors"), Accessors, AllExtensions);
 		FGLTFJsonUtility::WriteObjectArray(JsonWriter, TEXT("buffers"), Buffers, AllExtensions);
 		FGLTFJsonUtility::WriteObjectArray(JsonWriter, TEXT("bufferViews"), BufferViews, AllExtensions);
+		FGLTFJsonUtility::WriteObjectArray(JsonWriter, TEXT("cameras"), Cameras, AllExtensions);
 		FGLTFJsonUtility::WriteObjectArray(JsonWriter, TEXT("images"), Images, AllExtensions);
 		FGLTFJsonUtility::WriteObjectArray(JsonWriter, TEXT("materials"), Materials, AllExtensions);
 		FGLTFJsonUtility::WriteObjectArray(JsonWriter, TEXT("meshes"), Meshes, AllExtensions);
diff --git a/Engine/Plugins/Enterprise/GLTFExporter/MigratedFromGit/Source/GLTFExporter/Private/Json/GLTFJsonUtility.h b/Engine/Plugins/Enterprise/GLTFExporter/MigratedFromGit/Source/GLTFExporter/Private/Json/GLTFJsonUtility.h
--- a/Engine/Plugins/Enterprise/GLTFExporter/MigratedFromGit/Source/GLTFExporter/Private/Json/GLTFJsonUtility.h
+++ b/Engine/Plugins/Enterprise/GLTFExporter/MigratedFromGit/Source/GLTFExporter/Private/Json/GLTFJsonUtility.h
@@ -13,6 +13,8 @@ struct FGLTFJsonUtility
 		return static_cast<int32>(Value);
 	}
 
+	// TODO: replace all default cases (i.e. "unknown") with checkNoEntry() in each ToString
+
 	static const TCHAR* ToString(EGLTFJsonExtension Value)
 	{
 		switch (Value)
@@ -92,6 +94,16 @@ struct FGLTFJsonUtility
 		}
 	}
 
+	static const TCHAR* ToString(EGLTFJsonCameraType Value)
+	{
+		switch (Value)
+		{
+			case EGLTFJsonCameraType::Perspective:  return TEXT("perspective");
+			case EGLTFJsonCameraType::Orthographic: return TEXT("orthographic");
+			default:                                return TEXT("unknown");
+		}
+	}
+
 	template <class CharType = TCHAR, class PrintPolicy = TPrettyJsonPrintPolicy<CharType>>
 	static void WriteExactValue(TJsonWriter<CharType, PrintPolicy>& JsonWriter, float Value)
 	{
