[GLTFExporter] feat(camera): add converter for camera actors and components

commit 35a12d6ced70ced00ab09bdd8e61fbf54689c3f9
Date: Tue, 21 Jul 2020 07:50:23 +0200

    feat(camera): add converter for camera actors and components
    (cherry picked from commit 0929459d4defe378e2c1acdfcd5dee97f5dbf869​)

#skipci

diff --git a/Engine/Plugins/Enterprise/GLTFExporter/MigratedFromGit/Source/GLTFExporter/Private/Builders/GLTFConvertBuilder.cpp b/Engine/Plugins/Enterprise/GLTFExporter/MigratedFromGit/Source/GLTFExporter/Private/Builders/GLTFConvertBuilder.cpp
--- a/Engine/Plugins/Enterprise/GLTFExporter/MigratedFromGit/Source/GLTFExporter/Private/Builders/GLTFConvertBuilder.cpp
+++ b/Engine/Plugins/Enterprise/GLTFExporter/MigratedFromGit/Source/GLTFExporter/Private/Builders/GLTFConvertBuilder.cpp
@@ -204,6 +204,16 @@ FGLTFJsonSceneIndex FGLTFConvertBuilder::GetOrAddScene(const UWorld* World, cons
 	return GetOrAddScene(World->PersistentLevel, DesiredName.IsEmpty() ? World->GetName() : DesiredName);
 }
 
+FGLTFJsonCameraIndex FGLTFConvertBuilder::GetOrAddCamera(const UCameraComponent* CameraComponent, const FString& DesiredName)
+{
+	if (CameraComponent == nullptr)
+	{
+		return FGLTFJsonCameraIndex(INDEX_NONE);
+	}
+
+	return CameraComponentConverter.GetOrAdd(*this, DesiredName.IsEmpty() ? CameraComponent->GetName() : DesiredName, CameraComponent);
+}
+
 FGLTFJsonBackdropIndex FGLTFConvertBuilder::GetOrAddBackdrop(const AActor* Actor, const FString& DesiredName)
 {
 	if (Actor == nullptr)
diff --git a/Engine/Plugins/Enterprise/GLTFExporter/MigratedFromGit/Source/GLTFExporter/Private/Builders/GLTFConvertBuilder.h b/Engine/Plugins/Enterprise/GLTFExporter/MigratedFromGit/Source/GLTFExporter/Private/Builders/GLTFConvertBuilder.h
--- a/Engine/Plugins/Enterprise/GLTFExporter/MigratedFromGit/Source/GLTFExporter/Private/Builders/GLTFConvertBuilder.h
+++ b/Engine/Plugins/Enterprise/GLTFExporter/MigratedFromGit/Source/GLTFExporter/Private/Builders/GLTFConvertBuilder.h
@@ -43,6 +43,8 @@ public:
 	FGLTFJsonSceneIndex GetOrAddScene(const ULevel* Level, const FString& DesiredName = TEXT(""));
 	FGLTFJsonSceneIndex GetOrAddScene(const UWorld* World, const FString& DesiredName = TEXT(""));
 
+	FGLTFJsonCameraIndex GetOrAddCamera(const UCameraComponent* CameraComponent, const FString& DesiredName = TEXT(""));
+
 	FGLTFJsonBackdropIndex GetOrAddBackdrop(const AActor* Actor, const FString& DesiredName = TEXT(""));
 	FGLTFJsonLevelVariantSetsIndex GetOrAddLevelVariantSets(const ALevelVariantSetsActor* LevelVariantSetsActor, const FString& DesiredName = TEXT(""));
 	FGLTFJsonLightMapIndex GetOrAddLightMap(const UStaticMeshComponent* StaticMeshComponent, const FString& DesiredName = TEXT(""));
@@ -69,6 +71,7 @@ private:
 	FGLTFSceneComponentConverter SceneComponentConverter;
 	FGLTFActorConverter ActorConverter;
 	FGLTFLevelConverter LevelConverter;
+	FGLTFCameraComponentConverter CameraComponentConverter;
 
 	FGLTFBackdropConverter BackdropConverter;
 	FGLTFLevelVariantSetsConverter LevelVariantSetsConverter;
diff --git a/Engine/Plugins/Enterprise/GLTFExporter/MigratedFromGit/Source/GLTFExporter/Private/Builders/GLTFJsonBuilder.cpp b/Engine/Plugins/Enterprise/GLTFExporter/MigratedFromGit/Source/GLTFExporter/Private/Builders/GLTFJsonBuilder.cpp
--- a/Engine/Plugins/Enterprise/GLTFExporter/MigratedFromGit/Source/GLTFExporter/Private/Builders/GLTFJsonBuilder.cpp
+++ b/Engine/Plugins/Enterprise/GLTFExporter/MigratedFromGit/Source/GLTFExporter/Private/Builders/GLTFJsonBuilder.cpp
@@ -23,6 +23,11 @@ FGLTFJsonBufferViewIndex FGLTFJsonBuilder::AddBufferView(const FGLTFJsonBufferVi
 	return FGLTFJsonBufferViewIndex(JsonRoot.BufferViews.Add(JsonBufferView));
 }
 
+FGLTFJsonCameraIndex FGLTFJsonBuilder::AddCamera(const FGLTFJsonCamera& JsonCamera)
+{
+	return FGLTFJsonCameraIndex(JsonRoot.Cameras.Add(JsonCamera));
+}
+
 FGLTFJsonImageIndex FGLTFJsonBuilder::AddImage(const FGLTFJsonImage& JsonImage)
 {
 	return FGLTFJsonImageIndex(JsonRoot.Images.Add(JsonImage));
@@ -100,6 +105,11 @@ FGLTFJsonBufferView& FGLTFJsonBuilder::GetBufferView(FGLTFJsonBufferViewIndex Bu
 	return JsonRoot.BufferViews[BufferViewIndex];
 }
 
+FGLTFJsonCamera& FGLTFJsonBuilder::GetCamera(FGLTFJsonCameraIndex CameraIndex)
+{
+	return JsonRoot.Cameras[CameraIndex];
+}
+
 FGLTFJsonImage& FGLTFJsonBuilder::GetImage(FGLTFJsonImageIndex ImageIndex)
 {
 	return JsonRoot.Images[ImageIndex];
diff --git a/Engine/Plugins/Enterprise/GLTFExporter/MigratedFromGit/Source/GLTFExporter/Private/Builders/GLTFJsonBuilder.h b/Engine/Plugins/Enterprise/GLTFExporter/MigratedFromGit/Source/GLTFExporter/Private/Builders/GLTFJsonBuilder.h
--- a/Engine/Plugins/Enterprise/GLTFExporter/MigratedFromGit/Source/GLTFExporter/Private/Builders/GLTFJsonBuilder.h
+++ b/Engine/Plugins/Enterprise/GLTFExporter/MigratedFromGit/Source/GLTFExporter/Private/Builders/GLTFJsonBuilder.h
@@ -16,6 +16,7 @@ public:
 	FGLTFJsonAccessorIndex AddAccessor(const FGLTFJsonAccessor& JsonAccessor = {});
 	FGLTFJsonBufferIndex AddBuffer(const FGLTFJsonBuffer& JsonBuffer = {});
 	FGLTFJsonBufferViewIndex AddBufferView(const FGLTFJsonBufferView& JsonBufferView = {});
+	FGLTFJsonCameraIndex AddCamera(const FGLTFJsonCamera& JsonCamera = {});
 	FGLTFJsonImageIndex AddImage(const FGLTFJsonImage& JsonImage = {});
 	FGLTFJsonMaterialIndex AddMaterial(const FGLTFJsonMaterial& JsonMaterial = {});
 	FGLTFJsonMeshIndex AddMesh(const FGLTFJsonMesh& JsonMesh = {});
@@ -32,6 +33,7 @@ public:
 	FGLTFJsonAccessor& GetAccessor(FGLTFJsonAccessorIndex AccessorIndex);
 	FGLTFJsonBuffer& GetBuffer(FGLTFJsonBufferIndex BufferIndex);
 	FGLTFJsonBufferView& GetBufferView(FGLTFJsonBufferViewIndex BufferViewIndex);
+	FGLTFJsonCamera& GetCamera(FGLTFJsonCameraIndex CameraIndex);
 	FGLTFJsonImage& GetImage(FGLTFJsonImageIndex ImageIndex);
 	FGLTFJsonMaterial& GetMaterial(FGLTFJsonMaterialIndex MaterialIndex);
 	FGLTFJsonMesh& GetMesh(FGLTFJsonMeshIndex MeshIndex);
diff --git a/Engine/Plugins/Enterprise/GLTFExporter/MigratedFromGit/Source/GLTFExporter/Private/Converters/GLTFActorUtility.cpp b/Engine/Plugins/Enterprise/GLTFExporter/MigratedFromGit/Source/GLTFExporter/Private/Converters/GLTFActorUtility.cpp
--- a/Engine/Plugins/Enterprise/GLTFExporter/MigratedFromGit/Source/GLTFExporter/Private/Converters/GLTFActorUtility.cpp
+++ b/Engine/Plugins/Enterprise/GLTFExporter/MigratedFromGit/Source/GLTFExporter/Private/Converters/GLTFActorUtility.cpp
@@ -16,10 +16,12 @@ UBlueprint* FGLTFActorUtility::GetBlueprintFromActor(const AActor* Actor)
 
 bool FGLTFActorUtility::IsSkySphereBlueprint(const UBlueprint* Blueprint)
 {
+	// TODO: what if a blueprint inherits BP_Sky_Sphere?
 	return Blueprint != nullptr && Blueprint->GetPathName().Equals(TEXT("/Engine/EngineSky/BP_Sky_Sphere.BP_Sky_Sphere"));
 }
 
 bool FGLTFActorUtility::IsHDRIBackdropBlueprint(const UBlueprint* Blueprint)
 {
+	// TODO: what if a blueprint inherits HDRIBackdrop?
 	return Blueprint != nullptr && Blueprint->GetPathName().Equals(TEXT("/HDRIBackdrop/Blueprints/HDRIBackdrop.HDRIBackdrop"));
 }
diff --git a/Engine/Plugins/Enterprise/GLTFExporter/MigratedFromGit/Source/GLTFExporter/Private/Converters/GLTFConverterUtility.cpp b/Engine/Plugins/Enterprise/GLTFExporter/MigratedFromGit/Source/GLTFExporter/Private/Converters/GLTFConverterUtility.cpp
--- a/Engine/Plugins/Enterprise/GLTFExporter/MigratedFromGit/Source/GLTFExporter/Private/Converters/GLTFConverterUtility.cpp
+++ b/Engine/Plugins/Enterprise/GLTFExporter/MigratedFromGit/Source/GLTFExporter/Private/Converters/GLTFConverterUtility.cpp
@@ -3,6 +3,27 @@
 #include "Converters/GLTFConverterUtility.h"
 #include "Converters/GLTFTextureUtility.h"
 
+EGLTFJsonCameraType FGLTFConverterUtility::ConvertCameraType(ECameraProjectionMode::Type ProjectionMode)
+{
+	switch (ProjectionMode)
+	{
+		case ECameraProjectionMode::Perspective:  return EGLTFJsonCameraType::Perspective;
+		case ECameraProjectionMode::Orthographic: return EGLTFJsonCameraType::Orthographic;
+		default:                                  return EGLTFJsonCameraType::None;
+	}
+}
+
+EGLTFJsonShadingModel FGLTFConverterUtility::ConvertShadingModel(EMaterialShadingModel ShadingModel)
+{
+	switch (ShadingModel)
+	{
+		case MSM_Unlit:      return EGLTFJsonShadingModel::Unlit;
+		case MSM_DefaultLit: return EGLTFJsonShadingModel::Default;
+		case MSM_ClearCoat:  return EGLTFJsonShadingModel::ClearCoat;
+		default:             return EGLTFJsonShadingModel::None;
+	}
+}
+
 EGLTFJsonAlphaMode FGLTFConverterUtility::ConvertBlendMode(EBlendMode Mode)
 {
 	switch (Mode)
@@ -70,3 +91,32 @@ EGLTFJsonCubeFace FGLTFConverterUtility::ConvertCubeFace(ECubeFace CubeFace)
 		default:            return EGLTFJsonCubeFace::None;
 	}
 }
+
+FGLTFJsonOrthographic FGLTFConverterUtility::ConvertOrthographic(const FMinimalViewInfo& View)
+{
+	// TODO: do we need to convert any property from world units?
+
+	FGLTFJsonOrthographic Orthographic;
+	Orthographic.XMag = View.OrthoWidth;
+	Orthographic.YMag = View.OrthoWidth / View.AspectRatio; // TODO: is this correct?
+	Orthographic.ZFar = View.OrthoFarClipPlane;
+	Orthographic.ZNear = View.OrthoNearClipPlane;
+	return Orthographic;
+}
+
+FGLTFJsonPerspective FGLTFConverterUtility::ConvertPerspective(const FMinimalViewInfo& View)
+{
+	FGLTFJsonPerspective Perspective;
+	Perspective.AspectRatio = View.AspectRatio;
+	Perspective.YFov = ConvertFieldOfView(View);
+	Perspective.ZFar = 0; // infinite
+	Perspective.ZNear = GNearClippingPlane; // TODO: need to convert units?
+	return Perspective;
+}
+
+float FGLTFConverterUtility::ConvertFieldOfView(const FMinimalViewInfo& View)
+{
+	const float HorizontalFOV = FMath::DegreesToRadians(View.FOV);
+	const float VerticalFOV = FMath::Atan(FMath::Tan(HorizontalFOV) / View.AspectRatio);
+	return VerticalFOV;
+}
diff --git a/Engine/Plugins/Enterprise/GLTFExporter/MigratedFromGit/Source/GLTFExporter/Private/Converters/GLTFConverterUtility.h b/Engine/Plugins/Enterprise/GLTFExporter/MigratedFromGit/Source/GLTFExporter/Private/Converters/GLTFConverterUtility.h
--- a/Engine/Plugins/Enterprise/GLTFExporter/MigratedFromGit/Source/GLTFExporter/Private/Converters/GLTFConverterUtility.h
+++ b/Engine/Plugins/Enterprise/GLTFExporter/MigratedFromGit/Source/GLTFExporter/Private/Converters/GLTFConverterUtility.h
@@ -10,6 +10,7 @@
 #include "Json/GLTFJsonColor4.h"
 #include "Json/GLTFJsonQuaternion.h"
 #include "Engine/EngineTypes.h"
+#include "Json/GLTFJsonCamera.h"
 
 struct FGLTFConverterUtility
 {
@@ -91,16 +92,9 @@ struct FGLTFConverterUtility
 		return Result;
 	}
 
-	static EGLTFJsonShadingModel ConvertShadingModel(EMaterialShadingModel ShadingModel)
-	{
-		switch (ShadingModel)
-		{
-			case MSM_Unlit:      return EGLTFJsonShadingModel::Unlit;
-			case MSM_DefaultLit: return EGLTFJsonShadingModel::Default;
-			case MSM_ClearCoat:  return EGLTFJsonShadingModel::ClearCoat;
-			default:             return EGLTFJsonShadingModel::None;
-		}
-	}
+	static EGLTFJsonCameraType ConvertCameraType(ECameraProjectionMode::Type ProjectionMode);
+
+	static EGLTFJsonShadingModel ConvertShadingModel(EMaterialShadingModel ShadingModel);
 
 	static EGLTFJsonAlphaMode ConvertBlendMode(EBlendMode Mode);
 
@@ -114,6 +108,11 @@ struct FGLTFConverterUtility
 
 	static EGLTFJsonCubeFace ConvertCubeFace(ECubeFace CubeFace);
 
+	static FGLTFJsonOrthographic ConvertOrthographic(const FMinimalViewInfo& View);
+	static FGLTFJsonPerspective ConvertPerspective(const FMinimalViewInfo& View);
+
+	static float ConvertFieldOfView(const FMinimalViewInfo& View);
+
 	template <typename EnumType>
 	static FString GetEnumDisplayName(EnumType Value)
 	{
diff --git a/Engine/Plugins/Enterprise/GLTFExporter/MigratedFromGit/Source/GLTFExporter/Private/Converters/GLTFLevelConverters.cpp b/Engine/Plugins/Enterprise/GLTFExporter/MigratedFromGit/Source/GLTFExporter/Private/Converters/GLTFLevelConverters.cpp
--- a/Engine/Plugins/Enterprise/GLTFExporter/MigratedFromGit/Source/GLTFExporter/Private/Converters/GLTFLevelConverters.cpp
+++ b/Engine/Plugins/Enterprise/GLTFExporter/MigratedFromGit/Source/GLTFExporter/Private/Converters/GLTFLevelConverters.cpp
@@ -31,7 +31,11 @@ FGLTFJsonNodeIndex FGLTFSceneComponentConverter::Add(FGLTFConvertBuilder& Builde
 	Node.Rotation = FGLTFConverterUtility::ConvertRotation(Transform.GetRotation());
 	Node.Scale = FGLTFConverterUtility::ConvertScale(Transform.GetScale3D());
 
-	if (const UStaticMeshComponent* StaticMeshComponent = Cast<UStaticMeshComponent>(SceneComponent))
+	if (SceneComponent->bHiddenInGame)
+	{
+		// ignore any visible properties
+	}
+	else if (const UStaticMeshComponent* StaticMeshComponent = Cast<UStaticMeshComponent>(SceneComponent))
 	{
 		Node.Mesh = Builder.GetOrAddMesh(StaticMeshComponent);
 
@@ -40,6 +44,10 @@ FGLTFJsonNodeIndex FGLTFSceneComponentConverter::Add(FGLTFConvertBuilder& Builde
 			Node.LightMap = Builder.GetOrAddLightMap(StaticMeshComponent);
 		}
 	}
+	else if (const UCameraComponent* CameraComponent = Cast<UCameraComponent>(SceneComponent))
+	{
+		Node.Camera = Builder.GetOrAddCamera(CameraComponent);
+	}
 
 	return NodeIndex;
 }
@@ -65,6 +73,13 @@ FGLTFJsonNodeIndex FGLTFActorConverter::Add(FGLTFConvertBuilder& Builder, const
 		FGLTFJsonNode& RootNode = Builder.GetNode(RootNodeIndex);
 		RootNode.Backdrop = Builder.GetOrAddBackdrop(Actor);
 	}
+	else if (const ACameraActor* CameraActor = Cast<ACameraActor>(Actor))
+	{
+		// Ignore any other components inside camera actor (like proxy static mesh component)
+		// TODO: do we want to have to do this? Or can we rely on bHiddenInGame?
+		FGLTFJsonNode& RootNode = Builder.GetNode(RootNodeIndex);
+		RootNode.Camera = Builder.GetOrAddCamera(CameraActor->GetCameraComponent());
+	}
 	else
 	{
 		for (const UActorComponent* Component : Actor->GetComponents())
@@ -109,3 +124,29 @@ FGLTFJsonSceneIndex FGLTFLevelConverter::Add(FGLTFConvertBuilder& Builder, const
 
 	return Builder.AddScene(Scene);
 }
+
+FGLTFJsonCameraIndex FGLTFCameraComponentConverter::Add(FGLTFConvertBuilder& Builder, const FString& Name, const UCameraComponent* CameraComponent)
+{
+	FGLTFJsonCamera Camera;
+	Camera.Name = Name;
+	Camera.Type = FGLTFConverterUtility::ConvertCameraType(CameraComponent->ProjectionMode);
+
+	FMinimalViewInfo DesiredView;
+	const_cast<UCameraComponent*>(CameraComponent)->GetCameraView(0, DesiredView);
+
+	switch (Camera.Type)
+	{
+		case EGLTFJsonCameraType::Orthographic:
+			Camera.Orthographic = FGLTFConverterUtility::ConvertOrthographic(DesiredView);
+			break;
+
+		case EGLTFJsonCameraType::Perspective:
+			Camera.Perspective = FGLTFConverterUtility::ConvertPerspective(DesiredView);
+			break;
+
+		default:
+			checkNoEntry();
+	}
+
+	return Builder.AddCamera(Camera);
+}
diff --git a/Engine/Plugins/Enterprise/GLTFExporter/MigratedFromGit/Source/GLTFExporter/Private/Converters/GLTFLevelConverters.h b/Engine/Plugins/Enterprise/GLTFExporter/MigratedFromGit/Source/GLTFExporter/Private/Converters/GLTFLevelConverters.h
--- a/Engine/Plugins/Enterprise/GLTFExporter/MigratedFromGit/Source/GLTFExporter/Private/Converters/GLTFLevelConverters.h
+++ b/Engine/Plugins/Enterprise/GLTFExporter/MigratedFromGit/Source/GLTFExporter/Private/Converters/GLTFLevelConverters.h
@@ -20,3 +20,8 @@ class FGLTFLevelConverter final : public TGLTFConverter<FGLTFJsonSceneIndex, con
 {
 	FGLTFJsonSceneIndex Add(FGLTFConvertBuilder& Builder, const FString& Name, const ULevel* Level) override;
 };
+
+class FGLTFCameraComponentConverter final : public TGLTFConverter<FGLTFJsonCameraIndex, const UCameraComponent*>
+{
+	FGLTFJsonCameraIndex Add(FGLTFConvertBuilder& Builder, const FString& Name, const UCameraComponent* CameraComponent) override;
+};
diff --git a/Engine/Plugins/Enterprise/GLTFExporter/MigratedFromGit/Source/GLTFExporter/Private/Json/GLTFJsonCamera.h b/Engine/Plugins/Enterprise/GLTFExporter/MigratedFromGit/Source/GLTFExporter/Private/Json/GLTFJsonCamera.h
--- a/Engine/Plugins/Enterprise/GLTFExporter/MigratedFromGit/Source/GLTFExporter/Private/Json/GLTFJsonCamera.h
+++ b/Engine/Plugins/Enterprise/GLTFExporter/MigratedFromGit/Source/GLTFExporter/Private/Json/GLTFJsonCamera.h
@@ -6,7 +6,7 @@
 #include "Json/GLTFJsonUtility.h"
 #include "Serialization/JsonSerializer.h"
 
-struct GLTFEXPORTER_API FGLTFJsonOrthographic
+struct FGLTFJsonOrthographic
 {
 	float XMag; // horizontal magnification of the view
 	float YMag; // vertical magnification of the view
@@ -35,7 +35,7 @@ struct GLTFEXPORTER_API FGLTFJsonOrthographic
 	}
 };
 
-struct GLTFEXPORTER_API FGLTFJsonPerspective
+struct FGLTFJsonPerspective
 {
 	float AspectRatio; // aspect ratio of the field of view
 	float YFov; // vertical field of view in radians
@@ -73,7 +73,7 @@ struct GLTFEXPORTER_API FGLTFJsonPerspective
 	}
 };
 
-struct GLTFEXPORTER_API FGLTFJsonCamera
+struct FGLTFJsonCamera
 {
 	FString Name;
 
