[GLTFExporter] fix(camera): convert Unreal camera direction (+X) to glTF camera direction (-Y)

commit 81be3ef9f14fba44692373c29c00706c3431867a
Date: Tue, 21 Jul 2020 08:20:55 +0200

    fix(camera): convert Unreal camera direction (+X) to glTF camera direction (-Y)
    (cherry picked from commit 9a18c991a6bafe6dc935fe7c4abb9b6c3f497231​)

#skipci

diff --git a/Engine/Plugins/Enterprise/GLTFExporter/MigratedFromGit/Source/GLTFExporter/Private/Converters/GLTFConverterUtility.h b/Engine/Plugins/Enterprise/GLTFExporter/MigratedFromGit/Source/GLTFExporter/Private/Converters/GLTFConverterUtility.h
--- a/Engine/Plugins/Enterprise/GLTFExporter/MigratedFromGit/Source/GLTFExporter/Private/Converters/GLTFConverterUtility.h
+++ b/Engine/Plugins/Enterprise/GLTFExporter/MigratedFromGit/Source/GLTFExporter/Private/Converters/GLTFConverterUtility.h
@@ -92,6 +92,14 @@ struct FGLTFConverterUtility
 		return Result;
 	}
 
+	static FGLTFJsonQuaternion ConvertCameraDirection()
+	{
+		// Unreal uses +X axis as camera direction in Unreal coordinates.
+		// glTF uses -Y as camera direction in Unreal coordinates.
+
+		return ConvertRotation(FRotator(0, -90, 0).Quaternion());
+	}
+
 	static EGLTFJsonCameraType ConvertCameraType(ECameraProjectionMode::Type ProjectionMode);
 
 	static EGLTFJsonShadingModel ConvertShadingModel(EMaterialShadingModel ShadingModel);
diff --git a/Engine/Plugins/Enterprise/GLTFExporter/MigratedFromGit/Source/GLTFExporter/Private/Converters/GLTFLevelConverters.cpp b/Engine/Plugins/Enterprise/GLTFExporter/MigratedFromGit/Source/GLTFExporter/Private/Converters/GLTFLevelConverters.cpp
--- a/Engine/Plugins/Enterprise/GLTFExporter/MigratedFromGit/Source/GLTFExporter/Private/Converters/GLTFLevelConverters.cpp
+++ b/Engine/Plugins/Enterprise/GLTFExporter/MigratedFromGit/Source/GLTFExporter/Private/Converters/GLTFLevelConverters.cpp
@@ -46,7 +46,11 @@ FGLTFJsonNodeIndex FGLTFSceneComponentConverter::Add(FGLTFConvertBuilder& Builde
 	}
 	else if (const UCameraComponent* CameraComponent = Cast<UCameraComponent>(SceneComponent))
 	{
-		Node.Camera = Builder.GetOrAddCamera(CameraComponent);
+		// TODO: conversion of camera direction should be done in separate converter
+		FGLTFJsonNode CameraNode;
+		CameraNode.Rotation = FGLTFConverterUtility::ConvertCameraDirection();
+		CameraNode.Camera = Builder.GetOrAddCamera(CameraComponent);
+		Builder.AddChildNode(NodeIndex, CameraNode);
 	}
 
 	return NodeIndex;
@@ -73,13 +77,6 @@ FGLTFJsonNodeIndex FGLTFActorConverter::Add(FGLTFConvertBuilder& Builder, const
 		FGLTFJsonNode& RootNode = Builder.GetNode(RootNodeIndex);
 		RootNode.Backdrop = Builder.GetOrAddBackdrop(Actor);
 	}
-	else if (const ACameraActor* CameraActor = Cast<ACameraActor>(Actor))
-	{
-		// Ignore any other components inside camera actor (like proxy static mesh component)
-		// TODO: do we want to have to do this? Or can we rely on bHiddenInGame?
-		FGLTFJsonNode& RootNode = Builder.GetNode(RootNodeIndex);
-		RootNode.Camera = Builder.GetOrAddCamera(CameraActor->GetCameraComponent());
-	}
 	else
 	{
 		for (const UActorComponent* Component : Actor->GetComponents())
