[GLTFExporter] fix(camera): avoid inaccuracies when json serializing camera properties

commit 7c0a95909ee89afc847fc517900ed4e0a7fa73a6
Date: Tue, 21 Jul 2020 08:27:44 +0200

    fix(camera): avoid inaccuracies when json serializing camera properties
    (cherry picked from commit ac25f70d881008a63144c2b4445c3fc5648a76e2​)

#skipci

diff --git a/Engine/Plugins/Enterprise/GLTFExporter/MigratedFromGit/Source/GLTFExporter/Private/Json/GLTFJsonCamera.h b/Engine/Plugins/Enterprise/GLTFExporter/MigratedFromGit/Source/GLTFExporter/Private/Json/GLTFJsonCamera.h
--- a/Engine/Plugins/Enterprise/GLTFExporter/MigratedFromGit/Source/GLTFExporter/Private/Json/GLTFJsonCamera.h
+++ b/Engine/Plugins/Enterprise/GLTFExporter/MigratedFromGit/Source/GLTFExporter/Private/Json/GLTFJsonCamera.h
@@ -26,10 +26,10 @@ struct FGLTFJsonOrthographic
 	{
 		JsonWriter.WriteObjectStart();
 
-		JsonWriter.WriteValue(TEXT("xmag"), XMag);
-		JsonWriter.WriteValue(TEXT("ymag"), YMag);
-		JsonWriter.WriteValue(TEXT("zfar"), ZFar);
-		JsonWriter.WriteValue(TEXT("znear"), ZNear);
+		FGLTFJsonUtility::WriteExactValue(JsonWriter, TEXT("xmag"), XMag);
+		FGLTFJsonUtility::WriteExactValue(JsonWriter, TEXT("ymag"), YMag);
+		FGLTFJsonUtility::WriteExactValue(JsonWriter, TEXT("zfar"), ZFar);
+		FGLTFJsonUtility::WriteExactValue(JsonWriter, TEXT("znear"), ZNear);
 
 		JsonWriter.WriteObjectEnd();
 	}
@@ -57,17 +57,17 @@ struct FGLTFJsonPerspective
 
 		if (AspectRatio != 0)
 		{
-			JsonWriter.WriteValue(TEXT("aspectRatio"), AspectRatio);
+			FGLTFJsonUtility::WriteExactValue(JsonWriter, TEXT("aspectRatio"), AspectRatio);
 		}
 
-		JsonWriter.WriteValue(TEXT("yfov"), YFov);
+		FGLTFJsonUtility::WriteExactValue(JsonWriter, TEXT("yfov"), YFov);
 
 		if (ZFar != 0)
 		{
-			JsonWriter.WriteValue(TEXT("zfar"), ZFar);
+			FGLTFJsonUtility::WriteExactValue(JsonWriter, TEXT("zfar"), ZFar);
 		}
 
-		JsonWriter.WriteValue(TEXT("znear"), ZNear);
+		FGLTFJsonUtility::WriteExactValue(JsonWriter, TEXT("znear"), ZNear);
 
 		JsonWriter.WriteObjectEnd();
 	}
