[GLTFExporter] feat(light): add json data structures with corresponding enums and utilities

commit e4ec72afa50187467f95233a52fac2a4364eab30
Date: Tue, 21 Jul 2020 08:41:39 +0200

    feat(light): add json data structures with corresponding enums and utilities
    (cherry picked from commit c6ec7e471a86f69bf9e3b2d150a26a3739ef4e0a​)

#skipci

diff --git a/Engine/Plugins/Enterprise/GLTFExporter/MigratedFromGit/Source/GLTFExporter/Private/Json/GLTFJsonEnums.h b/Engine/Plugins/Enterprise/GLTFExporter/MigratedFromGit/Source/GLTFExporter/Private/Json/GLTFJsonEnums.h
--- a/Engine/Plugins/Enterprise/GLTFExporter/MigratedFromGit/Source/GLTFExporter/Private/Json/GLTFJsonEnums.h
+++ b/Engine/Plugins/Enterprise/GLTFExporter/MigratedFromGit/Source/GLTFExporter/Private/Json/GLTFJsonEnums.h
@@ -146,3 +146,11 @@ enum class EGLTFJsonCameraType
 	Orthographic,
 	Perspective
 };
+
+enum class EGLTFJsonLightType
+{
+	None = -1,
+    Directional,
+	Point,
+	Spot
+};
diff --git a/Engine/Plugins/Enterprise/GLTFExporter/MigratedFromGit/Source/GLTFExporter/Private/Json/GLTFJsonIndex.h b/Engine/Plugins/Enterprise/GLTFExporter/MigratedFromGit/Source/GLTFExporter/Private/Json/GLTFJsonIndex.h
--- a/Engine/Plugins/Enterprise/GLTFExporter/MigratedFromGit/Source/GLTFExporter/Private/Json/GLTFJsonIndex.h
+++ b/Engine/Plugins/Enterprise/GLTFExporter/MigratedFromGit/Source/GLTFExporter/Private/Json/GLTFJsonIndex.h
@@ -67,3 +67,4 @@ struct FGLTFJsonTextureIndex : FGLTFJsonIndex<FGLTFJsonTextureIndex> { using FGL
 struct FGLTFJsonBackdropIndex : FGLTFJsonIndex<FGLTFJsonBackdropIndex> { using FGLTFJsonIndex::FGLTFJsonIndex; };
 struct FGLTFJsonLevelVariantSetsIndex : FGLTFJsonIndex<FGLTFJsonLevelVariantSetsIndex> { using FGLTFJsonIndex::FGLTFJsonIndex; };
 struct FGLTFJsonLightMapIndex : FGLTFJsonIndex<FGLTFJsonLightMapIndex> { using FGLTFJsonIndex::FGLTFJsonIndex; };
+struct FGLTFJsonLightIndex : FGLTFJsonIndex<FGLTFJsonLightIndex> { using FGLTFJsonIndex::FGLTFJsonIndex; };
diff --git a/Engine/Plugins/Enterprise/GLTFExporter/MigratedFromGit/Source/GLTFExporter/Private/Json/GLTFJsonLight.h b/Engine/Plugins/Enterprise/GLTFExporter/MigratedFromGit/Source/GLTFExporter/Private/Json/GLTFJsonLight.h
new file mode 100644
--- /dev/null
+++ b/Engine/Plugins/Enterprise/GLTFExporter/MigratedFromGit/Source/GLTFExporter/Private/Json/GLTFJsonLight.h
@@ -0,0 +1,80 @@
+// Copyright Epic Games, Inc. All Rights Reserved.
+
+#pragma once
+
+#include "Json/GLTFJsonEnums.h"
+#include "Json/GLTFJsonColor3.h"
+#include "Json/GLTFJsonUtility.h"
+#include "Serialization/JsonSerializer.h"
+
+struct FGLTFJsonLight
+{
+	FString Name;
+
+	EGLTFJsonLightType Type;
+
+	FGLTFJsonColor3 Color;
+
+	float Intensity;
+	float Range;
+
+	float InnerConeAngle;
+	float OuterConeAngle;
+
+	FGLTFJsonLight()
+		: Type(EGLTFJsonLightType::None)
+		, Color(FGLTFJsonColor3::White)
+		, Intensity(1)
+		, Range(0)
+		, InnerConeAngle(0)
+		, OuterConeAngle(PI / 2.0f)
+	{
+	}
+
+	template <class CharType = TCHAR, class PrintPolicy = TPrettyJsonPrintPolicy<CharType>>
+	void WriteObject(TJsonWriter<CharType, PrintPolicy>& JsonWriter, FGLTFJsonExtensions& Extensions) const
+	{
+		JsonWriter.WriteObjectStart();
+
+		if (!Name.IsEmpty())
+		{
+			JsonWriter.WriteValue(TEXT("name"), Name);
+		}
+
+		JsonWriter.WriteValue(TEXT("type"), FGLTFJsonUtility::ToString(Type));
+
+		if (Color != FGLTFJsonColor3::White)
+		{
+			JsonWriter.WriteIdentifierPrefix(TEXT("color"));
+			Color.WriteArray(JsonWriter);
+		}
+
+		if (Intensity != 1)
+		{
+			FGLTFJsonUtility::WriteExactValue(JsonWriter, TEXT("intensity"), Intensity);
+		}
+
+		if (Type == EGLTFJsonLightType::Point || Type == EGLTFJsonLightType::Spot)
+		{
+			if (Range != 0)
+			{
+				FGLTFJsonUtility::WriteExactValue(JsonWriter, TEXT("range"), Range);
+			}
+
+			if (Type == EGLTFJsonLightType::Spot)
+			{
+				if (InnerConeAngle != 0)
+				{
+					FGLTFJsonUtility::WriteExactValue(JsonWriter, TEXT("innerConeAngle"), InnerConeAngle);
+				}
+
+				if (OuterConeAngle != PI / 2.0f)
+				{
+					FGLTFJsonUtility::WriteExactValue(JsonWriter, TEXT("outerConeAngle"), OuterConeAngle);
+				}
+			}
+		}
+
+		JsonWriter.WriteObjectEnd();
+	}
+};
diff --git a/Engine/Plugins/Enterprise/GLTFExporter/MigratedFromGit/Source/GLTFExporter/Private/Json/GLTFJsonNode.h b/Engine/Plugins/Enterprise/GLTFExporter/MigratedFromGit/Source/GLTFExporter/Private/Json/GLTFJsonNode.h
--- a/Engine/Plugins/Enterprise/GLTFExporter/MigratedFromGit/Source/GLTFExporter/Private/Json/GLTFJsonNode.h
+++ b/Engine/Plugins/Enterprise/GLTFExporter/MigratedFromGit/Source/GLTFExporter/Private/Json/GLTFJsonNode.h
@@ -20,6 +20,7 @@ struct FGLTFJsonNode
 	FGLTFJsonMeshIndex      Mesh;
 	FGLTFJsonBackdropIndex  Backdrop;
 	FGLTFJsonLightMapIndex  LightMap;
+	FGLTFJsonLightIndex     Light;
 
 	TArray<FGLTFJsonNodeIndex> Children;
 
@@ -97,6 +98,16 @@ struct FGLTFJsonNode
 				JsonWriter.WriteObjectEnd();
 			}
 
+			if (Light != INDEX_NONE)
+			{
+				const EGLTFJsonExtension Extension = EGLTFJsonExtension::KHR_LightsPunctual;
+				Extensions.Used.Add(Extension);
+
+				JsonWriter.WriteObjectStart(FGLTFJsonUtility::ToString(Extension));
+				JsonWriter.WriteValue(TEXT("light"), Light);
+				JsonWriter.WriteObjectEnd();
+			}
+
 			JsonWriter.WriteObjectEnd();
 		}
 
diff --git a/Engine/Plugins/Enterprise/GLTFExporter/MigratedFromGit/Source/GLTFExporter/Private/Json/GLTFJsonRoot.h b/Engine/Plugins/Enterprise/GLTFExporter/MigratedFromGit/Source/GLTFExporter/Private/Json/GLTFJsonRoot.h
--- a/Engine/Plugins/Enterprise/GLTFExporter/MigratedFromGit/Source/GLTFExporter/Private/Json/GLTFJsonRoot.h
+++ b/Engine/Plugins/Enterprise/GLTFExporter/MigratedFromGit/Source/GLTFExporter/Private/Json/GLTFJsonRoot.h
@@ -18,6 +18,7 @@
 #include "Json/GLTFJsonBackdrop.h"
 #include "Json/GLTFJsonLevelVariantSets.h"
 #include "Json/GLTFJsonLightMap.h"
+#include "Json/GLTFJsonLight.h"
 #include "Policies/CondensedJsonPrintPolicy.h"
 
 
@@ -71,6 +72,7 @@ struct FGLTFJsonRoot
 	TArray<FGLTFJsonTexture>    Textures;
 	TArray<FGLTFJsonBackdrop>   Backdrops;
 	TArray<FGLTFJsonLightMap>   LightMaps;
+	TArray<FGLTFJsonLight>      Lights;
 	TArray<FGLTFJsonLevelVariantSets> LevelVariantSets;
 
 	template <class CharType = TCHAR, class PrintPolicy = TPrettyJsonPrintPolicy<CharType>>
@@ -100,7 +102,7 @@ struct FGLTFJsonRoot
 		FGLTFJsonUtility::WriteObjectArray(JsonWriter, TEXT("scenes"), Scenes, AllExtensions);
 		FGLTFJsonUtility::WriteObjectArray(JsonWriter, TEXT("textures"), Textures, AllExtensions);
 
-		if (Backdrops.Num() > 0 || LevelVariantSets.Num() > 0 || LightMaps.Num() > 0 /* TODO: add more extension support */)
+		if (Backdrops.Num() > 0 || LevelVariantSets.Num() > 0 || LightMaps.Num() > 0 || Lights.Num() > 0 /* TODO: add more extension support */)
 		{
 			JsonWriter.WriteObjectStart(TEXT("extensions"));
 
@@ -134,6 +136,16 @@ struct FGLTFJsonRoot
 				JsonWriter.WriteObjectEnd();
 			}
 
+			if (Lights.Num() > 0)
+			{
+				const EGLTFJsonExtension Extension = EGLTFJsonExtension::KHR_LightsPunctual;
+				AllExtensions.Used.Add(Extension);
+
+				JsonWriter.WriteObjectStart(FGLTFJsonUtility::ToString(Extension));
+				FGLTFJsonUtility::WriteObjectArray(JsonWriter, TEXT("lights"), Lights, AllExtensions);
+				JsonWriter.WriteObjectEnd();
+			}
+
 			JsonWriter.WriteObjectEnd();
 		}
 
diff --git a/Engine/Plugins/Enterprise/GLTFExporter/MigratedFromGit/Source/GLTFExporter/Private/Json/GLTFJsonUtility.h b/Engine/Plugins/Enterprise/GLTFExporter/MigratedFromGit/Source/GLTFExporter/Private/Json/GLTFJsonUtility.h
--- a/Engine/Plugins/Enterprise/GLTFExporter/MigratedFromGit/Source/GLTFExporter/Private/Json/GLTFJsonUtility.h
+++ b/Engine/Plugins/Enterprise/GLTFExporter/MigratedFromGit/Source/GLTFExporter/Private/Json/GLTFJsonUtility.h
@@ -104,6 +104,17 @@ struct FGLTFJsonUtility
 		}
 	}
 
+	static const TCHAR* ToString(EGLTFJsonLightType Value)
+	{
+		switch (Value)
+		{
+			case EGLTFJsonLightType::Directional: return TEXT("directional");
+			case EGLTFJsonLightType::Point:       return TEXT("point");
+			case EGLTFJsonLightType::Spot:        return TEXT("spot");
+			default:                              return TEXT("unknown");
+		}
+	}
+
 	template <class CharType = TCHAR, class PrintPolicy = TPrettyJsonPrintPolicy<CharType>>
 	static void WriteExactValue(TJsonWriter<CharType, PrintPolicy>& JsonWriter, float Value)
 	{
