[GLTFExporter] feat(light): add converter for light components

commit 8d55eac8df428d2531ecc0e56f0f638799b21d1f
Date: Tue, 21 Jul 2020 09:19:54 +0200

    feat(light): add converter for light components
    (cherry picked from commit 7d8d36da1036a6cb6aec4b01bd6c683cc3e30ef9​)

#skipci

diff --git a/Engine/Plugins/Enterprise/GLTFExporter/MigratedFromGit/Source/GLTFExporter/Private/Builders/GLTFConvertBuilder.cpp b/Engine/Plugins/Enterprise/GLTFExporter/MigratedFromGit/Source/GLTFExporter/Private/Builders/GLTFConvertBuilder.cpp
--- a/Engine/Plugins/Enterprise/GLTFExporter/MigratedFromGit/Source/GLTFExporter/Private/Builders/GLTFConvertBuilder.cpp
+++ b/Engine/Plugins/Enterprise/GLTFExporter/MigratedFromGit/Source/GLTFExporter/Private/Builders/GLTFConvertBuilder.cpp
@@ -214,6 +214,16 @@ FGLTFJsonCameraIndex FGLTFConvertBuilder::GetOrAddCamera(const UCameraComponent*
 	return CameraComponentConverter.GetOrAdd(*this, DesiredName.IsEmpty() ? CameraComponent->GetName() : DesiredName, CameraComponent);
 }
 
+FGLTFJsonLightIndex FGLTFConvertBuilder::GetOrAddLight(const ULightComponent* LightComponent, const FString& DesiredName)
+{
+	if (LightComponent == nullptr)
+	{
+		return FGLTFJsonLightIndex(INDEX_NONE);
+	}
+
+	return LightComponentConverter.GetOrAdd(*this, DesiredName.IsEmpty() ? LightComponent->GetName() : DesiredName, LightComponent);
+}
+
 FGLTFJsonBackdropIndex FGLTFConvertBuilder::GetOrAddBackdrop(const AActor* Actor, const FString& DesiredName)
 {
 	if (Actor == nullptr)
diff --git a/Engine/Plugins/Enterprise/GLTFExporter/MigratedFromGit/Source/GLTFExporter/Private/Builders/GLTFConvertBuilder.h b/Engine/Plugins/Enterprise/GLTFExporter/MigratedFromGit/Source/GLTFExporter/Private/Builders/GLTFConvertBuilder.h
--- a/Engine/Plugins/Enterprise/GLTFExporter/MigratedFromGit/Source/GLTFExporter/Private/Builders/GLTFConvertBuilder.h
+++ b/Engine/Plugins/Enterprise/GLTFExporter/MigratedFromGit/Source/GLTFExporter/Private/Builders/GLTFConvertBuilder.h
@@ -44,6 +44,7 @@ public:
 	FGLTFJsonSceneIndex GetOrAddScene(const UWorld* World, const FString& DesiredName = TEXT(""));
 
 	FGLTFJsonCameraIndex GetOrAddCamera(const UCameraComponent* CameraComponent, const FString& DesiredName = TEXT(""));
+	FGLTFJsonLightIndex GetOrAddLight(const ULightComponent* LightComponent, const FString& DesiredName = TEXT(""));
 
 	FGLTFJsonBackdropIndex GetOrAddBackdrop(const AActor* Actor, const FString& DesiredName = TEXT(""));
 	FGLTFJsonLevelVariantSetsIndex GetOrAddLevelVariantSets(const ALevelVariantSetsActor* LevelVariantSetsActor, const FString& DesiredName = TEXT(""));
@@ -71,7 +72,9 @@ private:
 	FGLTFSceneComponentConverter SceneComponentConverter;
 	FGLTFActorConverter ActorConverter;
 	FGLTFLevelConverter LevelConverter;
+
 	FGLTFCameraComponentConverter CameraComponentConverter;
+	FGLTFLightComponentConverter LightComponentConverter;
 
 	FGLTFBackdropConverter BackdropConverter;
 	FGLTFLevelVariantSetsConverter LevelVariantSetsConverter;
diff --git a/Engine/Plugins/Enterprise/GLTFExporter/MigratedFromGit/Source/GLTFExporter/Private/Builders/GLTFJsonBuilder.cpp b/Engine/Plugins/Enterprise/GLTFExporter/MigratedFromGit/Source/GLTFExporter/Private/Builders/GLTFJsonBuilder.cpp
--- a/Engine/Plugins/Enterprise/GLTFExporter/MigratedFromGit/Source/GLTFExporter/Private/Builders/GLTFJsonBuilder.cpp
+++ b/Engine/Plugins/Enterprise/GLTFExporter/MigratedFromGit/Source/GLTFExporter/Private/Builders/GLTFJsonBuilder.cpp
@@ -78,6 +78,11 @@ FGLTFJsonLightMapIndex FGLTFJsonBuilder::AddLightMap(const FGLTFJsonLightMap& Js
 	return FGLTFJsonLightMapIndex(JsonRoot.LightMaps.Add(JsonLightMap));
 }
 
+FGLTFJsonLightIndex FGLTFJsonBuilder::AddLight(const FGLTFJsonLight& JsonLight)
+{
+	return FGLTFJsonLightIndex(JsonRoot.Lights.Add(JsonLight));
+}
+
 FGLTFJsonNodeIndex FGLTFJsonBuilder::AddChildNode(FGLTFJsonNodeIndex ParentIndex, const FGLTFJsonNode& JsonNode)
 {
 	const FGLTFJsonNodeIndex ChildIndex = AddNode(JsonNode);
@@ -160,6 +165,11 @@ FGLTFJsonLightMap& FGLTFJsonBuilder::GetLightMap(FGLTFJsonLightMapIndex LightMap
 	return JsonRoot.LightMaps[LightMapIndex];
 }
 
+FGLTFJsonLight& FGLTFJsonBuilder::GetLight(FGLTFJsonLightIndex LightIndex)
+{
+	return JsonRoot.Lights[LightIndex];
+}
+
 bool FGLTFJsonBuilder::Serialize(FArchive& Archive, const FString& FilePath)
 {
 	JsonRoot.Serialize(&Archive, true);
diff --git a/Engine/Plugins/Enterprise/GLTFExporter/MigratedFromGit/Source/GLTFExporter/Private/Builders/GLTFJsonBuilder.h b/Engine/Plugins/Enterprise/GLTFExporter/MigratedFromGit/Source/GLTFExporter/Private/Builders/GLTFJsonBuilder.h
--- a/Engine/Plugins/Enterprise/GLTFExporter/MigratedFromGit/Source/GLTFExporter/Private/Builders/GLTFJsonBuilder.h
+++ b/Engine/Plugins/Enterprise/GLTFExporter/MigratedFromGit/Source/GLTFExporter/Private/Builders/GLTFJsonBuilder.h
@@ -27,6 +27,7 @@ public:
 	FGLTFJsonBackdropIndex AddBackdrop(const FGLTFJsonBackdrop& JsonBackdrop = {});
 	FGLTFJsonLevelVariantSetsIndex AddLevelVariantSets(const FGLTFJsonLevelVariantSets& LevelVariantSets = {});
 	FGLTFJsonLightMapIndex AddLightMap(const FGLTFJsonLightMap& JsonLightMap = {});
+	FGLTFJsonLightIndex AddLight(const FGLTFJsonLight& JsonLight = {});
 
 	FGLTFJsonNodeIndex AddChildNode(FGLTFJsonNodeIndex ParentNodeIndex, const FGLTFJsonNode& JsonNode = {});
 
@@ -44,6 +45,7 @@ public:
 	FGLTFJsonBackdrop& GetBackdrop(FGLTFJsonBackdropIndex BackdropIndex);
 	FGLTFJsonLevelVariantSets& GetLevelVariantSets(FGLTFJsonLevelVariantSetsIndex LevelVariantSetsIndex);
 	FGLTFJsonLightMap& GetLightMap(FGLTFJsonLightMapIndex LightMapIndex);
+	FGLTFJsonLight& GetLight(FGLTFJsonLightIndex LightIndex);
 
 	virtual bool Serialize(FArchive& Archive, const FString& FilePath);
 
diff --git a/Engine/Plugins/Enterprise/GLTFExporter/MigratedFromGit/Source/GLTFExporter/Private/Converters/GLTFConverterUtility.cpp b/Engine/Plugins/Enterprise/GLTFExporter/MigratedFromGit/Source/GLTFExporter/Private/Converters/GLTFConverterUtility.cpp
--- a/Engine/Plugins/Enterprise/GLTFExporter/MigratedFromGit/Source/GLTFExporter/Private/Converters/GLTFConverterUtility.cpp
+++ b/Engine/Plugins/Enterprise/GLTFExporter/MigratedFromGit/Source/GLTFExporter/Private/Converters/GLTFConverterUtility.cpp
@@ -13,6 +13,17 @@ EGLTFJsonCameraType FGLTFConverterUtility::ConvertCameraType(ECameraProjectionMo
 	}
 }
 
+EGLTFJsonLightType FGLTFConverterUtility::ConvertLightType(ELightComponentType ComponentType)
+{
+	switch (ComponentType)
+	{
+		case LightType_Directional: return EGLTFJsonLightType::Directional;
+		case LightType_Point:       return EGLTFJsonLightType::Point;
+		case LightType_Spot:        return EGLTFJsonLightType::Spot;
+		default:                    return EGLTFJsonLightType::None;
+	}
+}
+
 EGLTFJsonShadingModel FGLTFConverterUtility::ConvertShadingModel(EMaterialShadingModel ShadingModel)
 {
 	switch (ShadingModel)
diff --git a/Engine/Plugins/Enterprise/GLTFExporter/MigratedFromGit/Source/GLTFExporter/Private/Converters/GLTFConverterUtility.h b/Engine/Plugins/Enterprise/GLTFExporter/MigratedFromGit/Source/GLTFExporter/Private/Converters/GLTFConverterUtility.h
--- a/Engine/Plugins/Enterprise/GLTFExporter/MigratedFromGit/Source/GLTFExporter/Private/Converters/GLTFConverterUtility.h
+++ b/Engine/Plugins/Enterprise/GLTFExporter/MigratedFromGit/Source/GLTFExporter/Private/Converters/GLTFConverterUtility.h
@@ -11,6 +11,7 @@
 #include "Json/GLTFJsonQuaternion.h"
 #include "Engine/EngineTypes.h"
 #include "Json/GLTFJsonCamera.h"
+#include "Json/GLTFJsonLight.h"
 
 struct FGLTFConverterUtility
 {
@@ -100,7 +101,28 @@ struct FGLTFConverterUtility
 		return ConvertRotation(FRotator(0, -90, 0).Quaternion());
 	}
 
+	static FGLTFJsonQuaternion ConvertLightDirection()
+	{
+		// Unreal uses +X axis as light direction in Unreal coordinates.
+		// glTF uses -Y as light direction in Unreal coordinates.
+
+		return ConvertRotation(FRotator(0, -90, 0).Quaternion());
+	}
+
+	static float ConvertLightAngle(const float Angle)
+	{
+		// Unreal uses degrees.
+		// glTF uses radians.
+		return FMath::DegreesToRadians(Angle);
+	}
+
+	static float ConvertLightRange(const float Radius)
+	{
+		return Radius * 0.01f; // TODO: use options export scale instead of hardcoded value
+	}
+
 	static EGLTFJsonCameraType ConvertCameraType(ECameraProjectionMode::Type ProjectionMode);
+	static EGLTFJsonLightType ConvertLightType(ELightComponentType ComponentType);
 
 	static EGLTFJsonShadingModel ConvertShadingModel(EMaterialShadingModel ShadingModel);
 
diff --git a/Engine/Plugins/Enterprise/GLTFExporter/MigratedFromGit/Source/GLTFExporter/Private/Converters/GLTFLevelConverters.cpp b/Engine/Plugins/Enterprise/GLTFExporter/MigratedFromGit/Source/GLTFExporter/Private/Converters/GLTFLevelConverters.cpp
--- a/Engine/Plugins/Enterprise/GLTFExporter/MigratedFromGit/Source/GLTFExporter/Private/Converters/GLTFLevelConverters.cpp
+++ b/Engine/Plugins/Enterprise/GLTFExporter/MigratedFromGit/Source/GLTFExporter/Private/Converters/GLTFLevelConverters.cpp
@@ -52,6 +52,14 @@ FGLTFJsonNodeIndex FGLTFSceneComponentConverter::Add(FGLTFConvertBuilder& Builde
 		CameraNode.Camera = Builder.GetOrAddCamera(CameraComponent);
 		Builder.AddChildNode(NodeIndex, CameraNode);
 	}
+	else if (const ULightComponent* LightComponent = Cast<ULightComponent>(SceneComponent))
+	{
+		// TODO: conversion of light direction should be done in separate converter
+		FGLTFJsonNode LightNode;
+		LightNode.Rotation = FGLTFConverterUtility::ConvertLightDirection();
+		LightNode.Light = Builder.GetOrAddLight(LightComponent);
+		Builder.AddChildNode(NodeIndex, LightNode);
+	}
 
 	return NodeIndex;
 }
@@ -142,8 +150,40 @@ FGLTFJsonCameraIndex FGLTFCameraComponentConverter::Add(FGLTFConvertBuilder& Bui
 			break;
 
 		default:
-			checkNoEntry();
+		    // TODO: report unsupported camera type
+		    return FGLTFJsonCameraIndex(INDEX_NONE);
 	}
 
 	return Builder.AddCamera(Camera);
 }
+
+FGLTFJsonLightIndex FGLTFLightComponentConverter::Add(FGLTFConvertBuilder& Builder, const FString& Name, const ULightComponent* LightComponent)
+{
+	FGLTFJsonLight Light;
+	Light.Name = Name;
+	Light.Type = FGLTFConverterUtility::ConvertLightType(LightComponent->GetLightType());
+
+	if (Light.Type == EGLTFJsonLightType::None)
+	{
+		// TODO: report unsupported light component type
+		return FGLTFJsonLightIndex(INDEX_NONE);
+	}
+
+	Light.Intensity = LightComponent->Intensity;
+
+	const FLinearColor LightColor = LightComponent->bUseTemperature ? FLinearColor::MakeFromColorTemperature(LightComponent->Temperature) : LightComponent->GetLightColor();
+	Light.Color = FGLTFConverterUtility::ConvertColor(LightColor);
+
+	if (const UPointLightComponent* PointLightComponent = Cast<UPointLightComponent>(LightComponent))
+	{
+		Light.Range = FGLTFConverterUtility::ConvertLightRange(PointLightComponent->AttenuationRadius);
+	}
+
+	if (const USpotLightComponent* SpotLightComponent = Cast<USpotLightComponent>(LightComponent))
+	{
+		Light.InnerConeAngle = FGLTFConverterUtility::ConvertLightAngle(SpotLightComponent->InnerConeAngle);
+		Light.OuterConeAngle = FGLTFConverterUtility::ConvertLightAngle(SpotLightComponent->OuterConeAngle);
+	}
+
+	return Builder.AddLight(Light);
+}
diff --git a/Engine/Plugins/Enterprise/GLTFExporter/MigratedFromGit/Source/GLTFExporter/Private/Converters/GLTFLevelConverters.h b/Engine/Plugins/Enterprise/GLTFExporter/MigratedFromGit/Source/GLTFExporter/Private/Converters/GLTFLevelConverters.h
--- a/Engine/Plugins/Enterprise/GLTFExporter/MigratedFromGit/Source/GLTFExporter/Private/Converters/GLTFLevelConverters.h
+++ b/Engine/Plugins/Enterprise/GLTFExporter/MigratedFromGit/Source/GLTFExporter/Private/Converters/GLTFLevelConverters.h
@@ -25,3 +25,8 @@ class FGLTFCameraComponentConverter final : public TGLTFConverter<FGLTFJsonCamer
 {
 	FGLTFJsonCameraIndex Add(FGLTFConvertBuilder& Builder, const FString& Name, const UCameraComponent* CameraComponent) override;
 };
+
+class FGLTFLightComponentConverter final : public TGLTFConverter<FGLTFJsonLightIndex, const ULightComponent*>
+{
+	FGLTFJsonLightIndex Add(FGLTFConvertBuilder& Builder, const FString& Name, const ULightComponent* LightComponent) override;
+};
diff --git a/Engine/Plugins/Enterprise/GLTFExporter/MigratedFromGit/Source/GLTFExporter/Private/Json/GLTFJsonNode.h b/Engine/Plugins/Enterprise/GLTFExporter/MigratedFromGit/Source/GLTFExporter/Private/Json/GLTFJsonNode.h
--- a/Engine/Plugins/Enterprise/GLTFExporter/MigratedFromGit/Source/GLTFExporter/Private/Json/GLTFJsonNode.h
+++ b/Engine/Plugins/Enterprise/GLTFExporter/MigratedFromGit/Source/GLTFExporter/Private/Json/GLTFJsonNode.h
@@ -74,7 +74,7 @@ struct FGLTFJsonNode
 			JsonWriter.WriteValue(TEXT("mesh"), Mesh);
 		}
 
-		if (Backdrop != INDEX_NONE || LightMap != INDEX_NONE)
+		if (Backdrop != INDEX_NONE || LightMap != INDEX_NONE || Light != INDEX_NONE)
 		{
 			JsonWriter.WriteObjectStart(TEXT("extensions"));
 
