[GLTFExporter] fix(camera): fix conversion from horizontal fov to vertical

commit a9f422fc50a8bfa195c8aaa245f57e4341f4a491
Date: Tue, 21 Jul 2020 12:46:23 +0200

    fix(camera): fix conversion from horizontal fov to vertical
    (cherry picked from commit 8bde6237b740db6e3b479121d4f411495a8daab3​)

#skipci

diff --git a/Engine/Plugins/Enterprise/GLTFExporter/MigratedFromGit/Source/GLTFExporter/Private/Converters/GLTFConverterUtility.cpp b/Engine/Plugins/Enterprise/GLTFExporter/MigratedFromGit/Source/GLTFExporter/Private/Converters/GLTFConverterUtility.cpp
--- a/Engine/Plugins/Enterprise/GLTFExporter/MigratedFromGit/Source/GLTFExporter/Private/Converters/GLTFConverterUtility.cpp
+++ b/Engine/Plugins/Enterprise/GLTFExporter/MigratedFromGit/Source/GLTFExporter/Private/Converters/GLTFConverterUtility.cpp
@@ -128,6 +128,6 @@ FGLTFJsonPerspective FGLTFConverterUtility::ConvertPerspective(const FMinimalVie
 float FGLTFConverterUtility::ConvertFieldOfView(const FMinimalViewInfo& View)
 {
 	const float HorizontalFOV = FMath::DegreesToRadians(View.FOV);
-	const float VerticalFOV = FMath::Atan(FMath::Tan(HorizontalFOV) / View.AspectRatio);
+	const float VerticalFOV = 2 * FMath::Atan(FMath::Tan(HorizontalFOV / 2) / View.AspectRatio);
 	return VerticalFOV;
 }
