[GLTFExporter] fix: fix conversion from Unreal camera and light direction to gltf

commit 50fdce32bc63f22c6dcee8fc305643c5e63330e5
Date: Tue, 21 Jul 2020 13:45:58 +0200

    fix: fix conversion from Unreal camera and light direction to gltf
    (cherry picked from commit 0bc54460ab2ff6f06da88765414f8c84733248d1​)

#skipci

diff --git a/Engine/Plugins/Enterprise/GLTFExporter/MigratedFromGit/Source/GLTFExporter/Private/Converters/GLTFConverterUtility.h b/Engine/Plugins/Enterprise/GLTFExporter/MigratedFromGit/Source/GLTFExporter/Private/Converters/GLTFConverterUtility.h
--- a/Engine/Plugins/Enterprise/GLTFExporter/MigratedFromGit/Source/GLTFExporter/Private/Converters/GLTFConverterUtility.h
+++ b/Engine/Plugins/Enterprise/GLTFExporter/MigratedFromGit/Source/GLTFExporter/Private/Converters/GLTFConverterUtility.h
@@ -96,17 +96,17 @@ struct FGLTFConverterUtility
 	static FGLTFJsonQuaternion ConvertCameraDirection()
 	{
 		// Unreal uses +X axis as camera direction in Unreal coordinates.
-		// glTF uses -Y as camera direction in Unreal coordinates.
+		// glTF uses +Y as camera direction in Unreal coordinates.
 
-		return ConvertRotation(FRotator(0, -90, 0).Quaternion());
+		return ConvertRotation(FRotator(0, 90, 0).Quaternion());
 	}
 
 	static FGLTFJsonQuaternion ConvertLightDirection()
 	{
 		// Unreal uses +X axis as light direction in Unreal coordinates.
-		// glTF uses -Y as light direction in Unreal coordinates.
+		// glTF uses +Y as light direction in Unreal coordinates.
 
-		return ConvertRotation(FRotator(0, -90, 0).Quaternion());
+		return ConvertRotation(FRotator(0, 90, 0).Quaternion());
 	}
 
 	static float ConvertLightAngle(const float Angle)
