[GLTFExporter] fix(camera): fix conversion from Unreal clipping plane and ortho size to gltf

commit f48eb06301a66307b2c21520e1f635f485a8d87e
Date: Tue, 21 Jul 2020 13:45:08 +0200

    fix(camera): fix conversion from Unreal clipping plane and ortho size to gltf
    (cherry picked from commit 91c18b5a47c0a49da768618f4096f0303db1467f​)

#skipci

diff --git a/Engine/Plugins/Enterprise/GLTFExporter/MigratedFromGit/Source/GLTFExporter/Private/Converters/GLTFConverterUtility.cpp b/Engine/Plugins/Enterprise/GLTFExporter/MigratedFromGit/Source/GLTFExporter/Private/Converters/GLTFConverterUtility.cpp
--- a/Engine/Plugins/Enterprise/GLTFExporter/MigratedFromGit/Source/GLTFExporter/Private/Converters/GLTFConverterUtility.cpp
+++ b/Engine/Plugins/Enterprise/GLTFExporter/MigratedFromGit/Source/GLTFExporter/Private/Converters/GLTFConverterUtility.cpp
@@ -105,13 +105,11 @@ EGLTFJsonCubeFace FGLTFConverterUtility::ConvertCubeFace(ECubeFace CubeFace)
 
 FGLTFJsonOrthographic FGLTFConverterUtility::ConvertOrthographic(const FMinimalViewInfo& View)
 {
-	// TODO: do we need to convert any property from world units?
-
 	FGLTFJsonOrthographic Orthographic;
-	Orthographic.XMag = View.OrthoWidth;
-	Orthographic.YMag = View.OrthoWidth / View.AspectRatio; // TODO: is this correct?
-	Orthographic.ZFar = View.OrthoFarClipPlane;
-	Orthographic.ZNear = View.OrthoNearClipPlane;
+	Orthographic.XMag = ConvertLength(View.OrthoWidth);
+	Orthographic.YMag = ConvertLength(View.OrthoWidth / View.AspectRatio); // TODO: is this correct?
+	Orthographic.ZFar = ConvertLength(View.OrthoFarClipPlane);
+	Orthographic.ZNear = ConvertLength(View.OrthoNearClipPlane);
 	return Orthographic;
 }
 
@@ -120,8 +118,8 @@ FGLTFJsonPerspective FGLTFConverterUtility::ConvertPerspective(const FMinimalVie
 	FGLTFJsonPerspective Perspective;
 	Perspective.AspectRatio = View.AspectRatio;
 	Perspective.YFov = ConvertFieldOfView(View);
-	Perspective.ZFar = 0; // infinite
-	Perspective.ZNear = GNearClippingPlane; // TODO: need to convert units?
+	Perspective.ZFar = ConvertLength(WORLD_MAX); // TODO: make sure the assumption that Unreal cameras have no draw distance is correct.
+	Perspective.ZNear = ConvertLength(GNearClippingPlane);
 	return Perspective;
 }
 
diff --git a/Engine/Plugins/Enterprise/GLTFExporter/MigratedFromGit/Source/GLTFExporter/Private/Converters/GLTFConverterUtility.h b/Engine/Plugins/Enterprise/GLTFExporter/MigratedFromGit/Source/GLTFExporter/Private/Converters/GLTFConverterUtility.h
--- a/Engine/Plugins/Enterprise/GLTFExporter/MigratedFromGit/Source/GLTFExporter/Private/Converters/GLTFConverterUtility.h
+++ b/Engine/Plugins/Enterprise/GLTFExporter/MigratedFromGit/Source/GLTFExporter/Private/Converters/GLTFConverterUtility.h
@@ -11,10 +11,14 @@
 #include "Json/GLTFJsonQuaternion.h"
 #include "Engine/EngineTypes.h"
 #include "Json/GLTFJsonCamera.h"
-#include "Json/GLTFJsonLight.h"
 
 struct FGLTFConverterUtility
 {
+	static float ConvertLength(const float Length)
+	{
+		return Length * 0.01f; // TODO: use options export scale instead of hardcoded value
+	}
+
 	static FGLTFJsonVector3 ConvertVector(const FVector& Vector)
 	{
 		// UE4 uses a left-handed coordinate system, with Z up.
@@ -116,11 +120,6 @@ struct FGLTFConverterUtility
 		return FMath::DegreesToRadians(Angle);
 	}
 
-	static float ConvertLightRange(const float Radius)
-	{
-		return Radius * 0.01f; // TODO: use options export scale instead of hardcoded value
-	}
-
 	static EGLTFJsonCameraType ConvertCameraType(ECameraProjectionMode::Type ProjectionMode);
 	static EGLTFJsonLightType ConvertLightType(ELightComponentType ComponentType);
 
diff --git a/Engine/Plugins/Enterprise/GLTFExporter/MigratedFromGit/Source/GLTFExporter/Private/Converters/GLTFLevelConverters.cpp b/Engine/Plugins/Enterprise/GLTFExporter/MigratedFromGit/Source/GLTFExporter/Private/Converters/GLTFLevelConverters.cpp
--- a/Engine/Plugins/Enterprise/GLTFExporter/MigratedFromGit/Source/GLTFExporter/Private/Converters/GLTFLevelConverters.cpp
+++ b/Engine/Plugins/Enterprise/GLTFExporter/MigratedFromGit/Source/GLTFExporter/Private/Converters/GLTFLevelConverters.cpp
@@ -176,7 +176,7 @@ FGLTFJsonLightIndex FGLTFLightComponentConverter::Add(FGLTFConvertBuilder& Build
 
 	if (const UPointLightComponent* PointLightComponent = Cast<UPointLightComponent>(LightComponent))
 	{
-		Light.Range = FGLTFConverterUtility::ConvertLightRange(PointLightComponent->AttenuationRadius);
+		Light.Range = FGLTFConverterUtility::ConvertLength(PointLightComponent->AttenuationRadius);
 	}
 
 	if (const USpotLightComponent* SpotLightComponent = Cast<USpotLightComponent>(LightComponent))
