[GLTFExporter] refactor: add names to child nodes for camera and light

commit be75f7eb3c504aadbb95686c97baeb82405125b6
Date: Tue, 21 Jul 2020 13:55:33 +0200

    refactor: add names to child nodes for camera and light
    (cherry picked from commit fbd03ae4dd40eefbee0154c2bf8f5b513add5906​)

#skipci

diff --git a/Engine/Plugins/Enterprise/GLTFExporter/MigratedFromGit/Source/GLTFExporter/Private/Converters/GLTFLevelConverters.cpp b/Engine/Plugins/Enterprise/GLTFExporter/MigratedFromGit/Source/GLTFExporter/Private/Converters/GLTFLevelConverters.cpp
--- a/Engine/Plugins/Enterprise/GLTFExporter/MigratedFromGit/Source/GLTFExporter/Private/Converters/GLTFLevelConverters.cpp
+++ b/Engine/Plugins/Enterprise/GLTFExporter/MigratedFromGit/Source/GLTFExporter/Private/Converters/GLTFLevelConverters.cpp
@@ -48,6 +48,7 @@ FGLTFJsonNodeIndex FGLTFSceneComponentConverter::Add(FGLTFConvertBuilder& Builde
 	{
 		// TODO: conversion of camera direction should be done in separate converter
 		FGLTFJsonNode CameraNode;
+		CameraNode.Name = Node.Name + TEXT("_CameraNode");
 		CameraNode.Rotation = FGLTFConverterUtility::ConvertCameraDirection();
 		CameraNode.Camera = Builder.GetOrAddCamera(CameraComponent);
 		Builder.AddChildNode(NodeIndex, CameraNode);
@@ -56,6 +57,7 @@ FGLTFJsonNodeIndex FGLTFSceneComponentConverter::Add(FGLTFConvertBuilder& Builde
 	{
 		// TODO: conversion of light direction should be done in separate converter
 		FGLTFJsonNode LightNode;
+		LightNode.Name = Node.Name + TEXT("_LightNode");
 		LightNode.Rotation = FGLTFConverterUtility::ConvertLightDirection();
 		LightNode.Light = Builder.GetOrAddLight(LightComponent);
 		Builder.AddChildNode(NodeIndex, LightNode);
