[GLTFExporter] fix(level): export LevelVariantSetActor as a node, in case some other actor is attached to it

commit a96e3d104869065a48d4d0b76db93d786932f076
Date: Tue, 21 Jul 2020 17:10:54 +0200

    fix(level): export LevelVariantSetActor as a node, in case some other actor is attached to it
    (cherry picked from commit 69c0b94e5a04b85628d82a63a579fb0debd61453​)

#skipci

diff --git a/Engine/Plugins/Enterprise/GLTFExporter/MigratedFromGit/Source/GLTFExporter/Private/Converters/GLTFLevelConverters.cpp b/Engine/Plugins/Enterprise/GLTFExporter/MigratedFromGit/Source/GLTFExporter/Private/Converters/GLTFLevelConverters.cpp
--- a/Engine/Plugins/Enterprise/GLTFExporter/MigratedFromGit/Source/GLTFExporter/Private/Converters/GLTFLevelConverters.cpp
+++ b/Engine/Plugins/Enterprise/GLTFExporter/MigratedFromGit/Source/GLTFExporter/Private/Converters/GLTFLevelConverters.cpp
@@ -119,13 +119,11 @@ FGLTFJsonSceneIndex FGLTFLevelConverter::Add(FGLTFConvertBuilder& Builder, const
 				Scene.LevelVariantSets.Add(LevelVariantSetsIndex);
 			}
 		}
-		else
+
+		const FGLTFJsonNodeIndex NodeIndex = Builder.GetOrAddNode(Actor);
+		if (NodeIndex != INDEX_NONE && FGLTFActorUtility::IsRootActor(Actor, Builder.bSelectedActorsOnly))
 		{
-			const FGLTFJsonNodeIndex NodeIndex = Builder.GetOrAddNode(Actor);
-			if (NodeIndex != INDEX_NONE && FGLTFActorUtility::IsRootActor(Actor, Builder.bSelectedActorsOnly))
-			{
-				Scene.Nodes.Add(NodeIndex);
-			}
+			Scene.Nodes.Add(NodeIndex);
 		}
 	}
 
