[GLTFExporter] refactor(level): add TODOs for future improvements

commit 7271e2d44b86fb3a31caf055bfc80f3488ca8acf
Date: Tue, 21 Jul 2020 17:30:12 +0200

    refactor(level): add TODOs for future improvements
    (cherry picked from commit 2b7d76e3ea658feb1d24408bc269974aab73e685​)

#skipci

diff --git a/Engine/Plugins/Enterprise/GLTFExporter/MigratedFromGit/Source/GLTFExporter/Private/Converters/GLTFLevelConverters.cpp b/Engine/Plugins/Enterprise/GLTFExporter/MigratedFromGit/Source/GLTFExporter/Private/Converters/GLTFLevelConverters.cpp
--- a/Engine/Plugins/Enterprise/GLTFExporter/MigratedFromGit/Source/GLTFExporter/Private/Converters/GLTFLevelConverters.cpp
+++ b/Engine/Plugins/Enterprise/GLTFExporter/MigratedFromGit/Source/GLTFExporter/Private/Converters/GLTFLevelConverters.cpp
@@ -21,6 +21,8 @@ FGLTFJsonNodeIndex FGLTFSceneComponentConverter::Add(FGLTFConvertBuilder& Builde
 	const USceneComponent* ParentComponent = SceneComponent->GetAttachParent();
 	const FGLTFJsonNodeIndex ParentNodeIndex = Builder.GetOrAddNode(ParentComponent);
 
+	// TODO: if node is root, then add it to the scene here to avoid any possible orphan nodes. this change requires level converter to support cyclic calls.
+
 	const FGLTFJsonNodeIndex NodeIndex = Builder.AddChildNode(ParentNodeIndex);
 	FGLTFJsonNode& Node = Builder.GetNode(NodeIndex);
 	Node.Name = Name.IsEmpty() ? Owner->GetName() + TEXT("_") + SceneComponent->GetName() : Name;
@@ -89,6 +91,8 @@ FGLTFJsonNodeIndex FGLTFActorConverter::Add(FGLTFConvertBuilder& Builder, const
 	}
 	else
 	{
+		// TODO: to reduce number of nodes, only export components that are of interest
+
 		for (const UActorComponent* Component : Actor->GetComponents())
 		{
 			const USceneComponent* SceneComponent = Cast<USceneComponent>(Component);
