[GLTFExporter] feat: add FGLTFConverterUtility::GetComponentType to convert integer type to json enum

commit 8deaab04ba7c5cc4a89f38a679186db972d50bca
Date: Fri, 29 May 2020 22:29:07 +0200

    feat: add FGLTFConverterUtility::GetComponentType to convert integer type to json enum
    (cherry picked from commit d70bb42721375dbb606c3fdd6689152ac3aa9548​)

#skipci

diff --git a/Engine/Plugins/Enterprise/GLTFExporter/MigratedFromGit/Source/GLTFExporter/Private/Converters/GLTFConverterUtility.h b/Engine/Plugins/Enterprise/GLTFExporter/MigratedFromGit/Source/GLTFExporter/Private/Converters/GLTFConverterUtility.h
--- a/Engine/Plugins/Enterprise/GLTFExporter/MigratedFromGit/Source/GLTFExporter/Private/Converters/GLTFConverterUtility.h
+++ b/Engine/Plugins/Enterprise/GLTFExporter/MigratedFromGit/Source/GLTFExporter/Private/Converters/GLTFConverterUtility.h
@@ -142,6 +142,19 @@ struct FGLTFConverterUtility
 
 	static float ConvertFieldOfView(const FMinimalViewInfo& View);
 
+	template <typename ComponentType>
+	static EGLTFJsonComponentType GetComponentType()
+	{
+		if (TIsSame<ComponentType, int8  >::Value) return EGLTFJsonComponentType::S8;
+		if (TIsSame<ComponentType, uint8 >::Value) return EGLTFJsonComponentType::U8;
+		if (TIsSame<ComponentType, int16 >::Value) return EGLTFJsonComponentType::S16;
+		if (TIsSame<ComponentType, uint16>::Value) return EGLTFJsonComponentType::U16;
+		if (TIsSame<ComponentType, int32 >::Value) return EGLTFJsonComponentType::S32;
+		if (TIsSame<ComponentType, uint32>::Value) return EGLTFJsonComponentType::U32;
+		if (TIsSame<ComponentType, float >::Value) return EGLTFJsonComponentType::F32;
+		return EGLTFJsonComponentType::None;
+	}
+
 	template <typename EnumType>
 	static FString GetEnumDisplayName(EnumType Value)
 	{
