[GLTFExporter] feat(skeletal-mesh): add new converters for skeletal mesh (not yet implemented)

commit 43cf3bb9c8e5af049facdfc38e6553000619fc43
Date: Mon, 25 May 2020 19:46:27 +0200

    feat(skeletal-mesh): add new converters for skeletal mesh (not yet implemented)
    (cherry picked from commit 7c3a21cbab74db14b0d33a3ef379d48180c944b7​)

#skipci

diff --git a/Engine/Plugins/Enterprise/GLTFExporter/MigratedFromGit/Source/GLTFExporter/Private/Builders/GLTFBuilderUtility.cpp b/Engine/Plugins/Enterprise/GLTFExporter/MigratedFromGit/Source/GLTFExporter/Private/Builders/GLTFBuilderUtility.cpp
--- a/Engine/Plugins/Enterprise/GLTFExporter/MigratedFromGit/Source/GLTFExporter/Private/Builders/GLTFBuilderUtility.cpp
+++ b/Engine/Plugins/Enterprise/GLTFExporter/MigratedFromGit/Source/GLTFExporter/Private/Builders/GLTFBuilderUtility.cpp
@@ -2,14 +2,18 @@
 
 #include "Builders/GLTFBuilderUtility.h"
 
-FString FGLTFBuilderUtility::GetMeshName(const UStaticMesh* StaticMesh, int32 LODIndex)
+FString FGLTFBuilderUtility::GetLODName(const UObject* Object, int32 LODIndex)
 {
 	FString Name;
 
-	if (StaticMesh != nullptr)
+	if (Object != nullptr)
 	{
-		StaticMesh->GetName(Name);
-		if (LODIndex != 0) Name += TEXT("_LOD") + FString::FromInt(LODIndex);
+		Object->GetName(Name);
+
+		if (LODIndex != 0)
+		{
+			Name += TEXT("_LOD") + FString::FromInt(LODIndex);
+		}
 	}
 
 	return Name;
diff --git a/Engine/Plugins/Enterprise/GLTFExporter/MigratedFromGit/Source/GLTFExporter/Private/Builders/GLTFBuilderUtility.h b/Engine/Plugins/Enterprise/GLTFExporter/MigratedFromGit/Source/GLTFExporter/Private/Builders/GLTFBuilderUtility.h
--- a/Engine/Plugins/Enterprise/GLTFExporter/MigratedFromGit/Source/GLTFExporter/Private/Builders/GLTFBuilderUtility.h
+++ b/Engine/Plugins/Enterprise/GLTFExporter/MigratedFromGit/Source/GLTFExporter/Private/Builders/GLTFBuilderUtility.h
@@ -7,7 +7,7 @@
 
 struct FGLTFBuilderUtility
 {
-	static FString GetMeshName(const UStaticMesh* StaticMesh, int32 LODIndex);
+	static FString GetLODName(const UObject* Object, int32 LODIndex);
 
 	static const TCHAR* GetFileExtension(EGLTFJsonMimeType MimeType);
 
diff --git a/Engine/Plugins/Enterprise/GLTFExporter/MigratedFromGit/Source/GLTFExporter/Private/Builders/GLTFConvertBuilder.cpp b/Engine/Plugins/Enterprise/GLTFExporter/MigratedFromGit/Source/GLTFExporter/Private/Builders/GLTFConvertBuilder.cpp
--- a/Engine/Plugins/Enterprise/GLTFExporter/MigratedFromGit/Source/GLTFExporter/Private/Builders/GLTFConvertBuilder.cpp
+++ b/Engine/Plugins/Enterprise/GLTFExporter/MigratedFromGit/Source/GLTFExporter/Private/Builders/GLTFConvertBuilder.cpp
@@ -59,6 +59,26 @@ FGLTFJsonAccessorIndex FGLTFConvertBuilder::GetOrAddUVAccessor(const FStaticMesh
 	return UVVertexBufferConverter.GetOrAdd(*this, DesiredName, VertexBuffer, UVIndex);
 }
 
+FGLTFJsonAccessorIndex FGLTFConvertBuilder::GetOrAddJointAccessor(const FSkinWeightVertexBuffer* VertexBuffer, const FString& DesiredName)
+{
+	if (VertexBuffer == nullptr)
+	{
+		return FGLTFJsonAccessorIndex(INDEX_NONE);
+	}
+
+	return BoneIndexVertexBufferConverter.GetOrAdd(*this, DesiredName, VertexBuffer);
+}
+
+FGLTFJsonAccessorIndex FGLTFConvertBuilder::GetOrAddWeightAccessor(const FSkinWeightVertexBuffer* VertexBuffer, const FString& DesiredName)
+{
+	if (VertexBuffer == nullptr)
+	{
+		return FGLTFJsonAccessorIndex(INDEX_NONE);
+	}
+
+	return BoneWeightVertexBufferConverter.GetOrAdd(*this, DesiredName, VertexBuffer);
+}
+
 FGLTFJsonBufferViewIndex FGLTFConvertBuilder::GetOrAddIndexBufferView(const FRawStaticIndexBuffer* IndexBuffer, const FString& DesiredName)
 {
 	if (IndexBuffer == nullptr)
@@ -86,7 +106,7 @@ FGLTFJsonMeshIndex FGLTFConvertBuilder::GetOrAddMesh(const UStaticMesh* StaticMe
 		return FGLTFJsonMeshIndex(INDEX_NONE);
 	}
 
-	return StaticMeshConverter.GetOrAdd(*this, DesiredName.IsEmpty() ? FGLTFBuilderUtility::GetMeshName(StaticMesh, LODIndex) : DesiredName, StaticMesh, LODIndex, OverrideVertexColors, OverrideMaterials);
+	return StaticMeshConverter.GetOrAdd(*this, DesiredName.IsEmpty() ? FGLTFBuilderUtility::GetLODName(StaticMesh, LODIndex) : DesiredName, StaticMesh, LODIndex, OverrideVertexColors, OverrideMaterials);
 }
 
 FGLTFJsonMeshIndex FGLTFConvertBuilder::GetOrAddMesh(const UStaticMeshComponent* StaticMeshComponent, const FString& DesiredName)
@@ -104,6 +124,52 @@ FGLTFJsonMeshIndex FGLTFConvertBuilder::GetOrAddMesh(const UStaticMeshComponent*
 	return GetOrAddMesh(StaticMesh, LODIndex, OverrideVertexColors, OverrideMaterials, DesiredName);
 }
 
+FGLTFJsonBufferViewIndex FGLTFConvertBuilder::GetOrAddIndexBufferView(const FMultiSizeIndexContainer* IndexContainer, const FString& DesiredName)
+{
+	if (IndexContainer == nullptr)
+	{
+		return FGLTFJsonBufferViewIndex(INDEX_NONE);
+	}
+
+	return IndexContainerConverter.GetOrAdd(*this, DesiredName, IndexContainer);
+}
+
+FGLTFJsonAccessorIndex FGLTFConvertBuilder::GetOrAddIndexAccessor(const FSkelMeshRenderSection* MeshSection, const FMultiSizeIndexContainer* IndexContainer, const FString& DesiredName)
+{
+	if (MeshSection == nullptr)
+	{
+		return FGLTFJsonAccessorIndex(INDEX_NONE);
+	}
+
+	return SkeletalMeshSectionConverter.GetOrAdd(*this, DesiredName, MeshSection, IndexContainer);
+}
+
+FGLTFJsonMeshIndex FGLTFConvertBuilder::GetOrAddMesh(const USkeletalMesh* SkeletalMesh, int32 LODIndex, const FColorVertexBuffer* OverrideVertexColors, const FSkinWeightVertexBuffer* OverrideSkinWeights, const FGLTFMaterialArray& OverrideMaterials, const FString& DesiredName)
+{
+	if (SkeletalMesh == nullptr)
+	{
+		return FGLTFJsonMeshIndex(INDEX_NONE);
+	}
+
+	return SkeletalMeshConverter.GetOrAdd(*this, DesiredName.IsEmpty() ? FGLTFBuilderUtility::GetLODName(SkeletalMesh, LODIndex) : DesiredName, SkeletalMesh, LODIndex, OverrideVertexColors, OverrideSkinWeights, OverrideMaterials);
+}
+
+FGLTFJsonMeshIndex FGLTFConvertBuilder::GetOrAddMesh(const USkeletalMeshComponent* SkeletalMeshComponent, const FString& DesiredName)
+{
+	if (SkeletalMeshComponent == nullptr)
+	{
+		return FGLTFJsonMeshIndex(INDEX_NONE);
+	}
+
+	const USkeletalMesh* SkeletalMesh = SkeletalMeshComponent->SkeletalMesh;
+	const int32 LODIndex = SkeletalMeshComponent->GetForcedLOD() > 0 ? SkeletalMeshComponent->GetForcedLOD() - 1 : /* auto-select */ 0;
+	const FColorVertexBuffer* OverrideVertexColors = SkeletalMeshComponent->LODInfo.IsValidIndex(LODIndex) ? SkeletalMeshComponent->LODInfo[LODIndex].OverrideVertexColors : nullptr;
+	const FSkinWeightVertexBuffer* OverrideSkinWeights = SkeletalMeshComponent->LODInfo.IsValidIndex(LODIndex) ? SkeletalMeshComponent->LODInfo[LODIndex].OverrideSkinWeights : nullptr;
+	const FGLTFMaterialArray OverrideMaterials = FGLTFMaterialArray(SkeletalMeshComponent->OverrideMaterials);
+
+	return GetOrAddMesh(SkeletalMesh, LODIndex, OverrideVertexColors, OverrideSkinWeights, OverrideMaterials, DesiredName);
+}
+
 FGLTFJsonMaterialIndex FGLTFConvertBuilder::GetOrAddMaterial(const UMaterialInterface* Material, const FString& DesiredName)
 {
 	if (Material == nullptr)
diff --git a/Engine/Plugins/Enterprise/GLTFExporter/MigratedFromGit/Source/GLTFExporter/Private/Builders/GLTFConvertBuilder.h b/Engine/Plugins/Enterprise/GLTFExporter/MigratedFromGit/Source/GLTFExporter/Private/Builders/GLTFConvertBuilder.h
--- a/Engine/Plugins/Enterprise/GLTFExporter/MigratedFromGit/Source/GLTFExporter/Private/Builders/GLTFConvertBuilder.h
+++ b/Engine/Plugins/Enterprise/GLTFExporter/MigratedFromGit/Source/GLTFExporter/Private/Builders/GLTFConvertBuilder.h
@@ -5,6 +5,7 @@
 #include "Builders/GLTFImageBuilder.h"
 #include "Converters/GLTFVertexBufferConverters.h"
 #include "Converters/GLTFStaticMeshConverters.h"
+#include "Converters/GLTFSkeletalMeshConverters.h"
 #include "Converters/GLTFMaterialConverters.h"
 #include "Converters/GLTFTextureConverters.h"
 #include "Converters/GLTFLevelConverters.h"
@@ -25,12 +26,19 @@ public:
 	FGLTFJsonAccessorIndex GetOrAddNormalAccessor(const FStaticMeshVertexBuffer* VertexBuffer, const FString& DesiredName = TEXT(""));
 	FGLTFJsonAccessorIndex GetOrAddTangentAccessor(const FStaticMeshVertexBuffer* VertexBuffer, const FString& DesiredName = TEXT(""));
 	FGLTFJsonAccessorIndex GetOrAddUVAccessor(const FStaticMeshVertexBuffer* VertexBuffer, int32 UVIndex, const FString& DesiredName = TEXT(""));
+	FGLTFJsonAccessorIndex GetOrAddJointAccessor(const FSkinWeightVertexBuffer* VertexBuffer, const FString& DesiredName = TEXT(""));
+	FGLTFJsonAccessorIndex GetOrAddWeightAccessor(const FSkinWeightVertexBuffer* VertexBuffer, const FString& DesiredName = TEXT(""));
 
 	FGLTFJsonBufferViewIndex GetOrAddIndexBufferView(const FRawStaticIndexBuffer* IndexBuffer, const FString& DesiredName = TEXT(""));
 	FGLTFJsonAccessorIndex GetOrAddIndexAccessor(const FStaticMeshSection* MeshSection, const FRawStaticIndexBuffer* IndexBuffer, const FString& DesiredName = TEXT(""));
 	FGLTFJsonMeshIndex GetOrAddMesh(const UStaticMesh* StaticMesh, int32 LODIndex = 0, const FColorVertexBuffer* OverrideVertexColors = nullptr, const FGLTFMaterialArray& OverrideMaterials = {}, const FString& DesiredName = TEXT(""));
 	FGLTFJsonMeshIndex GetOrAddMesh(const UStaticMeshComponent* StaticMeshComponent, const FString& DesiredName = TEXT(""));
 
+	FGLTFJsonBufferViewIndex GetOrAddIndexBufferView(const FMultiSizeIndexContainer* IndexContainer, const FString& DesiredName = TEXT(""));
+	FGLTFJsonAccessorIndex GetOrAddIndexAccessor(const FSkelMeshRenderSection* MeshSection, const FMultiSizeIndexContainer* IndexContainer, const FString& DesiredName = TEXT(""));
+	FGLTFJsonMeshIndex GetOrAddMesh(const USkeletalMesh* SkeletalMesh, int32 LODIndex = 0, const FColorVertexBuffer* OverrideVertexColors = nullptr, const FSkinWeightVertexBuffer* OverrideSkinWeights = nullptr, const FGLTFMaterialArray& OverrideMaterials = {}, const FString& DesiredName = TEXT(""));
+	FGLTFJsonMeshIndex GetOrAddMesh(const USkeletalMeshComponent* SkeletalMeshComponent, const FString& DesiredName = TEXT(""));
+
 	FGLTFJsonMaterialIndex GetOrAddMaterial(const UMaterialInterface* Material, const FString& DesiredName = TEXT(""));
 	FGLTFJsonSamplerIndex GetOrAddSampler(const UTexture* Texture, const FString& DesiredName = TEXT(""));
 	FGLTFJsonTextureIndex GetOrAddTexture(const UTexture2D* Texture, const FString& DesiredName = TEXT(""));
@@ -57,11 +65,17 @@ private:
 	FGLTFNormalVertexBufferConverter NormalVertexBufferConverter;
 	FGLTFTangentVertexBufferConverter TangentVertexBufferConverter;
 	FGLTFUVVertexBufferConverter UVVertexBufferConverter;
+	FGLTFBoneIndexVertexBufferConverter BoneIndexVertexBufferConverter;
+	FGLTFBoneWeightVertexBufferConverter BoneWeightVertexBufferConverter;
 
 	FGLTFIndexBufferConverter IndexBufferConverter;
 	FGLTFStaticMeshSectionConverter StaticMeshSectionConverter;
 	FGLTFStaticMeshConverter StaticMeshConverter;
 
+	FGLTFIndexContainerConverter IndexContainerConverter;
+	FGLTFSkeletalMeshSectionConverter SkeletalMeshSectionConverter;
+	FGLTFSkeletalMeshConverter SkeletalMeshConverter;
+
 	FGLTFMaterialConverter MaterialConverter;
 	FGLTFTextureSamplerConverter TextureSamplerConverter;
 	FGLTFTexture2DConverter Texture2DConverter;
diff --git a/Engine/Plugins/Enterprise/GLTFExporter/MigratedFromGit/Source/GLTFExporter/Private/Converters/GLTFSkeletalMeshConverters.cpp b/Engine/Plugins/Enterprise/GLTFExporter/MigratedFromGit/Source/GLTFExporter/Private/Converters/GLTFSkeletalMeshConverters.cpp
new file mode 100644
--- /dev/null
+++ b/Engine/Plugins/Enterprise/GLTFExporter/MigratedFromGit/Source/GLTFExporter/Private/Converters/GLTFSkeletalMeshConverters.cpp
@@ -0,0 +1,20 @@
+// Copyright Epic Games, Inc. All Rights Reserved.
+
+#include "Converters/GLTFSkeletalMeshConverters.h"
+#include "Converters/GLTFConverterUtility.h"
+#include "Builders/GLTFConvertBuilder.h"
+
+FGLTFJsonBufferViewIndex FGLTFIndexContainerConverter::Add(FGLTFConvertBuilder& Builder, const FString& Name, const FMultiSizeIndexContainer* IndexContainer)
+{
+	return FGLTFJsonBufferViewIndex(INDEX_NONE);
+}
+
+FGLTFJsonAccessorIndex FGLTFSkeletalMeshSectionConverter::Add(FGLTFConvertBuilder& Builder, const FString& Name, const FSkelMeshRenderSection* MeshSection, const FMultiSizeIndexContainer* IndexContainer)
+{
+	return FGLTFJsonAccessorIndex(INDEX_NONE);
+}
+
+FGLTFJsonMeshIndex FGLTFSkeletalMeshConverter::Add(FGLTFConvertBuilder& Builder, const FString& Name, const USkeletalMesh* SkeletalMesh, int32 LODIndex, const FColorVertexBuffer* OverrideVertexColors, const FSkinWeightVertexBuffer* OverrideSkinWeights, FGLTFMaterialArray OverrideMaterials)
+{
+	return FGLTFJsonMeshIndex(INDEX_NONE);
+}
diff --git a/Engine/Plugins/Enterprise/GLTFExporter/MigratedFromGit/Source/GLTFExporter/Private/Converters/GLTFSkeletalMeshConverters.h b/Engine/Plugins/Enterprise/GLTFExporter/MigratedFromGit/Source/GLTFExporter/Private/Converters/GLTFSkeletalMeshConverters.h
new file mode 100644
--- /dev/null
+++ b/Engine/Plugins/Enterprise/GLTFExporter/MigratedFromGit/Source/GLTFExporter/Private/Converters/GLTFSkeletalMeshConverters.h
@@ -0,0 +1,25 @@
+// Copyright Epic Games, Inc. All Rights Reserved.
+
+#pragma once
+
+#include "Json/GLTFJsonIndex.h"
+#include "Converters/GLTFConverter.h"
+#include "Converters/GLTFMaterialArray.h"
+#include "Engine.h"
+
+class FMultiSizeIndexContainer;
+
+class FGLTFIndexContainerConverter final : public TGLTFConverter<FGLTFJsonBufferViewIndex, const FMultiSizeIndexContainer*>
+{
+	FGLTFJsonBufferViewIndex Add(FGLTFConvertBuilder& Builder, const FString& Name, const FMultiSizeIndexContainer* IndexContainer) override;
+};
+
+class FGLTFSkeletalMeshSectionConverter final : public TGLTFConverter<FGLTFJsonAccessorIndex, const FSkelMeshRenderSection*, const FMultiSizeIndexContainer*>
+{
+	FGLTFJsonAccessorIndex Add(FGLTFConvertBuilder& Builder, const FString& Name, const FSkelMeshRenderSection* MeshSection, const FMultiSizeIndexContainer* IndexContainer) override;
+};
+
+class FGLTFSkeletalMeshConverter final : public TGLTFConverter<FGLTFJsonMeshIndex, const USkeletalMesh*, int32, const FColorVertexBuffer*, const FSkinWeightVertexBuffer*, FGLTFMaterialArray>
+{
+	FGLTFJsonMeshIndex Add(FGLTFConvertBuilder& Builder, const FString& Name, const USkeletalMesh* SkeletalMesh, int32 LODIndex, const FColorVertexBuffer* OverrideVertexColors, const FSkinWeightVertexBuffer* OverrideSkinWeights, FGLTFMaterialArray OverrideMaterials) override;
+};
diff --git a/Engine/Plugins/Enterprise/GLTFExporter/MigratedFromGit/Source/GLTFExporter/Private/Converters/GLTFVertexBufferConverters.cpp b/Engine/Plugins/Enterprise/GLTFExporter/MigratedFromGit/Source/GLTFExporter/Private/Converters/GLTFVertexBufferConverters.cpp
--- a/Engine/Plugins/Enterprise/GLTFExporter/MigratedFromGit/Source/GLTFExporter/Private/Converters/GLTFVertexBufferConverters.cpp
+++ b/Engine/Plugins/Enterprise/GLTFExporter/MigratedFromGit/Source/GLTFExporter/Private/Converters/GLTFVertexBufferConverters.cpp
@@ -158,3 +158,12 @@ FGLTFJsonAccessorIndex FGLTFUVVertexBufferConverter::Add(FGLTFConvertBuilder& Bu
 	return Builder.AddAccessor(JsonAccessor);
 }
 
+FGLTFJsonAccessorIndex FGLTFBoneIndexVertexBufferConverter::Add(FGLTFConvertBuilder& Builder, const FString& Name, const FSkinWeightVertexBuffer* VertexBuffer)
+{
+	return FGLTFJsonAccessorIndex(INDEX_NONE);
+}
+
+FGLTFJsonAccessorIndex FGLTFBoneWeightVertexBufferConverter::Add(FGLTFConvertBuilder& Builder, const FString& Name, const FSkinWeightVertexBuffer* VertexBuffer)
+{
+	return FGLTFJsonAccessorIndex(INDEX_NONE);
+}
diff --git a/Engine/Plugins/Enterprise/GLTFExporter/MigratedFromGit/Source/GLTFExporter/Private/Converters/GLTFVertexBufferConverters.h b/Engine/Plugins/Enterprise/GLTFExporter/MigratedFromGit/Source/GLTFExporter/Private/Converters/GLTFVertexBufferConverters.h
--- a/Engine/Plugins/Enterprise/GLTFExporter/MigratedFromGit/Source/GLTFExporter/Private/Converters/GLTFVertexBufferConverters.h
+++ b/Engine/Plugins/Enterprise/GLTFExporter/MigratedFromGit/Source/GLTFExporter/Private/Converters/GLTFVertexBufferConverters.h
@@ -30,3 +30,13 @@ class FGLTFUVVertexBufferConverter final : public TGLTFConverter<FGLTFJsonAccess
 {
 	FGLTFJsonAccessorIndex Add(FGLTFConvertBuilder& Builder, const FString& Name, const FStaticMeshVertexBuffer* VertexBuffer, int32 UVIndex) override;
 };
+
+class FGLTFBoneIndexVertexBufferConverter final : public TGLTFConverter<FGLTFJsonAccessorIndex, const FSkinWeightVertexBuffer*>
+{
+	FGLTFJsonAccessorIndex Add(FGLTFConvertBuilder& Builder, const FString& Name, const FSkinWeightVertexBuffer* VertexBuffer) override;
+};
+
+class FGLTFBoneWeightVertexBufferConverter final : public TGLTFConverter<FGLTFJsonAccessorIndex, const FSkinWeightVertexBuffer*>
+{
+	FGLTFJsonAccessorIndex Add(FGLTFConvertBuilder& Builder, const FString& Name, const FSkinWeightVertexBuffer* VertexBuffer) override;
+};
