[GLTFExporter] feat(skeletal-mesh): implement skin/weight vertexbuffer converters (has known link error)

commit 98d0b6e95d6fd94e096b8806f52c26464e54f320
Date: Mon, 25 May 2020 21:10:54 +0200

    feat(skeletal-mesh): implement skin/weight vertexbuffer converters (has known link error)
    (cherry picked from commit 8fa236df87b038898a8c9974f5c14c862c5c3f9c​)

#skipci

diff --git a/Engine/Plugins/Enterprise/GLTFExporter/MigratedFromGit/Source/GLTFExporter/Private/Builders/GLTFConvertBuilder.cpp b/Engine/Plugins/Enterprise/GLTFExporter/MigratedFromGit/Source/GLTFExporter/Private/Builders/GLTFConvertBuilder.cpp
--- a/Engine/Plugins/Enterprise/GLTFExporter/MigratedFromGit/Source/GLTFExporter/Private/Builders/GLTFConvertBuilder.cpp
+++ b/Engine/Plugins/Enterprise/GLTFExporter/MigratedFromGit/Source/GLTFExporter/Private/Builders/GLTFConvertBuilder.cpp
@@ -59,24 +59,24 @@ FGLTFJsonAccessorIndex FGLTFConvertBuilder::GetOrAddUVAccessor(const FStaticMesh
 	return UVVertexBufferConverter.GetOrAdd(*this, DesiredName, VertexBuffer, UVIndex);
 }
 
-FGLTFJsonAccessorIndex FGLTFConvertBuilder::GetOrAddJointAccessor(const FSkinWeightVertexBuffer* VertexBuffer, const FString& DesiredName)
+FGLTFJsonAccessorIndex FGLTFConvertBuilder::GetOrAddJointAccessor(const FSkinWeightVertexBuffer* VertexBuffer, int32 InfluenceOffset, const FString& DesiredName)
 {
 	if (VertexBuffer == nullptr)
 	{
 		return FGLTFJsonAccessorIndex(INDEX_NONE);
 	}
 
-	return BoneIndexVertexBufferConverter.GetOrAdd(*this, DesiredName, VertexBuffer);
+	return BoneIndexVertexBufferConverter.GetOrAdd(*this, DesiredName, VertexBuffer, InfluenceOffset);
 }
 
-FGLTFJsonAccessorIndex FGLTFConvertBuilder::GetOrAddWeightAccessor(const FSkinWeightVertexBuffer* VertexBuffer, const FString& DesiredName)
+FGLTFJsonAccessorIndex FGLTFConvertBuilder::GetOrAddWeightAccessor(const FSkinWeightVertexBuffer* VertexBuffer, int32 InfluenceOffset, const FString& DesiredName)
 {
 	if (VertexBuffer == nullptr)
 	{
 		return FGLTFJsonAccessorIndex(INDEX_NONE);
 	}
 
-	return BoneWeightVertexBufferConverter.GetOrAdd(*this, DesiredName, VertexBuffer);
+	return BoneWeightVertexBufferConverter.GetOrAdd(*this, DesiredName, VertexBuffer, InfluenceOffset);
 }
 
 FGLTFJsonBufferViewIndex FGLTFConvertBuilder::GetOrAddIndexBufferView(const FRawStaticIndexBuffer* IndexBuffer, const FString& DesiredName)
diff --git a/Engine/Plugins/Enterprise/GLTFExporter/MigratedFromGit/Source/GLTFExporter/Private/Builders/GLTFConvertBuilder.h b/Engine/Plugins/Enterprise/GLTFExporter/MigratedFromGit/Source/GLTFExporter/Private/Builders/GLTFConvertBuilder.h
--- a/Engine/Plugins/Enterprise/GLTFExporter/MigratedFromGit/Source/GLTFExporter/Private/Builders/GLTFConvertBuilder.h
+++ b/Engine/Plugins/Enterprise/GLTFExporter/MigratedFromGit/Source/GLTFExporter/Private/Builders/GLTFConvertBuilder.h
@@ -26,8 +26,8 @@ public:
 	FGLTFJsonAccessorIndex GetOrAddNormalAccessor(const FStaticMeshVertexBuffer* VertexBuffer, const FString& DesiredName = TEXT(""));
 	FGLTFJsonAccessorIndex GetOrAddTangentAccessor(const FStaticMeshVertexBuffer* VertexBuffer, const FString& DesiredName = TEXT(""));
 	FGLTFJsonAccessorIndex GetOrAddUVAccessor(const FStaticMeshVertexBuffer* VertexBuffer, int32 UVIndex, const FString& DesiredName = TEXT(""));
-	FGLTFJsonAccessorIndex GetOrAddJointAccessor(const FSkinWeightVertexBuffer* VertexBuffer, const FString& DesiredName = TEXT(""));
-	FGLTFJsonAccessorIndex GetOrAddWeightAccessor(const FSkinWeightVertexBuffer* VertexBuffer, const FString& DesiredName = TEXT(""));
+	FGLTFJsonAccessorIndex GetOrAddJointAccessor(const FSkinWeightVertexBuffer* VertexBuffer, int32 InfluenceOffset, const FString& DesiredName = TEXT(""));
+	FGLTFJsonAccessorIndex GetOrAddWeightAccessor(const FSkinWeightVertexBuffer* VertexBuffer, int32 InfluenceOffset, const FString& DesiredName = TEXT(""));
 
 	FGLTFJsonBufferViewIndex GetOrAddIndexBufferView(const FRawStaticIndexBuffer* IndexBuffer, const FString& DesiredName = TEXT(""));
 	FGLTFJsonAccessorIndex GetOrAddIndexAccessor(const FStaticMeshSection* MeshSection, const FRawStaticIndexBuffer* IndexBuffer, const FString& DesiredName = TEXT(""));
diff --git a/Engine/Plugins/Enterprise/GLTFExporter/MigratedFromGit/Source/GLTFExporter/Private/Converters/GLTFVertexBufferConverters.cpp b/Engine/Plugins/Enterprise/GLTFExporter/MigratedFromGit/Source/GLTFExporter/Private/Converters/GLTFVertexBufferConverters.cpp
--- a/Engine/Plugins/Enterprise/GLTFExporter/MigratedFromGit/Source/GLTFExporter/Private/Converters/GLTFVertexBufferConverters.cpp
+++ b/Engine/Plugins/Enterprise/GLTFExporter/MigratedFromGit/Source/GLTFExporter/Private/Converters/GLTFVertexBufferConverters.cpp
@@ -158,12 +158,93 @@ FGLTFJsonAccessorIndex FGLTFUVVertexBufferConverter::Add(FGLTFConvertBuilder& Bu
 	return Builder.AddAccessor(JsonAccessor);
 }
 
-FGLTFJsonAccessorIndex FGLTFBoneIndexVertexBufferConverter::Add(FGLTFConvertBuilder& Builder, const FString& Name, const FSkinWeightVertexBuffer* VertexBuffer)
+FGLTFJsonAccessorIndex FGLTFBoneIndexVertexBufferConverter::Add(FGLTFConvertBuilder& Builder, const FString& Name, const FSkinWeightVertexBuffer* VertexBuffer, int32 InfluenceOffset)
 {
-	return FGLTFJsonAccessorIndex(INDEX_NONE);
+	return VertexBuffer->Use16BitBoneIndex() ? Add<uint16>(Builder, Name, VertexBuffer, InfluenceOffset) : Add<uint8>(Builder, Name, VertexBuffer, InfluenceOffset);
 }
 
-FGLTFJsonAccessorIndex FGLTFBoneWeightVertexBufferConverter::Add(FGLTFConvertBuilder& Builder, const FString& Name, const FSkinWeightVertexBuffer* VertexBuffer)
+template <typename IndexType>
+FGLTFJsonAccessorIndex FGLTFBoneIndexVertexBufferConverter::Add(FGLTFConvertBuilder& Builder, const FString& Name, const FSkinWeightVertexBuffer* VertexBuffer, int32 InfluenceOffset)
 {
-	return FGLTFJsonAccessorIndex(INDEX_NONE);
+	const uint32 VertexCount = VertexBuffer->GetNumVertices();
+	const int32 MaxInfluenceCount = VertexBuffer->GetMaxBoneInfluences();
+	if (VertexCount == 0 || InfluenceOffset < 0 || MaxInfluenceCount <= InfluenceOffset)
+	{
+		return FGLTFJsonAccessorIndex(INDEX_NONE);
+	}
+
+	const int32 InfluenceCount = FMath::Min(MaxInfluenceCount - InfluenceOffset, 4);
+
+	struct VertexBoneIndices
+	{
+		IndexType Index[4];
+	};
+
+	TArray<VertexBoneIndices> BoneIndices;
+	BoneIndices.AddUninitialized(VertexCount);
+
+	for (uint32 VertexIndex = 0; VertexIndex < VertexCount; ++VertexIndex)
+	{
+		for (int32 InfluenceIndex = 0; InfluenceIndex < InfluenceCount; ++InfluenceIndex)
+		{
+			BoneIndices[VertexIndex].Index[InfluenceIndex] = static_cast<IndexType>(VertexBuffer->GetBoneIndex(VertexIndex, InfluenceOffset + InfluenceIndex));
+		}
+
+		for (int32 InfluenceIndex = InfluenceCount; InfluenceIndex < 4; ++InfluenceIndex)
+		{
+			BoneIndices[VertexIndex].Index[InfluenceIndex] = 0;
+		}
+	}
+
+	FGLTFJsonAccessor JsonAccessor;
+	JsonAccessor.Name = Name;
+	JsonAccessor.BufferView = Builder.AddBufferView(BoneIndices, Name);
+	JsonAccessor.ComponentType = FGLTFConverterUtility::GetComponentType<IndexType>();
+	JsonAccessor.Count = VertexCount;
+	JsonAccessor.Type = EGLTFJsonAccessorType::Vec4;
+
+	return Builder.AddAccessor(JsonAccessor);
+}
+
+FGLTFJsonAccessorIndex FGLTFBoneWeightVertexBufferConverter::Add(FGLTFConvertBuilder& Builder, const FString& Name, const FSkinWeightVertexBuffer* VertexBuffer, int32 InfluenceOffset)
+{
+	const uint32 VertexCount = VertexBuffer->GetNumVertices();
+	const int32 MaxInfluenceCount = VertexBuffer->GetMaxBoneInfluences();
+	if (VertexCount == 0 || InfluenceOffset < 0 || MaxInfluenceCount <= InfluenceOffset)
+	{
+		return FGLTFJsonAccessorIndex(INDEX_NONE);
+	}
+
+	const int32 InfluenceCount = FMath::Min(MaxInfluenceCount - InfluenceOffset, 4);
+
+	struct VertexBoneWeights
+	{
+		uint8 Weights[4];
+	};
+
+	TArray<VertexBoneWeights> BoneWeights;
+	BoneWeights.AddUninitialized(VertexCount);
+
+	for (uint32 VertexIndex = 0; VertexIndex < VertexCount; ++VertexIndex)
+	{
+		for (int32 InfluenceIndex = 0; InfluenceIndex < InfluenceCount; ++InfluenceIndex)
+		{
+			BoneWeights[VertexIndex].Weights[InfluenceIndex] = VertexBuffer->GetBoneWeight(VertexIndex, InfluenceOffset + InfluenceIndex);
+		}
+
+		for (int32 InfluenceIndex = InfluenceCount; InfluenceIndex < 4; ++InfluenceIndex)
+		{
+			BoneWeights[VertexIndex].Weights[InfluenceIndex] = 0;
+		}
+	}
+
+	FGLTFJsonAccessor JsonAccessor;
+	JsonAccessor.Name = Name;
+	JsonAccessor.BufferView = Builder.AddBufferView(BoneWeights, Name);
+	JsonAccessor.ComponentType = EGLTFJsonComponentType::U8;
+	JsonAccessor.Count = VertexCount;
+	JsonAccessor.Type = EGLTFJsonAccessorType::Vec4;
+	JsonAccessor.bNormalized = true;
+
+	return Builder.AddAccessor(JsonAccessor);
 }
diff --git a/Engine/Plugins/Enterprise/GLTFExporter/MigratedFromGit/Source/GLTFExporter/Private/Converters/GLTFVertexBufferConverters.h b/Engine/Plugins/Enterprise/GLTFExporter/MigratedFromGit/Source/GLTFExporter/Private/Converters/GLTFVertexBufferConverters.h
--- a/Engine/Plugins/Enterprise/GLTFExporter/MigratedFromGit/Source/GLTFExporter/Private/Converters/GLTFVertexBufferConverters.h
+++ b/Engine/Plugins/Enterprise/GLTFExporter/MigratedFromGit/Source/GLTFExporter/Private/Converters/GLTFVertexBufferConverters.h
@@ -31,12 +31,15 @@ class FGLTFUVVertexBufferConverter final : public TGLTFConverter<FGLTFJsonAccess
 	FGLTFJsonAccessorIndex Add(FGLTFConvertBuilder& Builder, const FString& Name, const FStaticMeshVertexBuffer* VertexBuffer, int32 UVIndex) override;
 };
 
-class FGLTFBoneIndexVertexBufferConverter final : public TGLTFConverter<FGLTFJsonAccessorIndex, const FSkinWeightVertexBuffer*>
+class FGLTFBoneIndexVertexBufferConverter final : public TGLTFConverter<FGLTFJsonAccessorIndex, const FSkinWeightVertexBuffer*, int32>
 {
-	FGLTFJsonAccessorIndex Add(FGLTFConvertBuilder& Builder, const FString& Name, const FSkinWeightVertexBuffer* VertexBuffer) override;
+	FGLTFJsonAccessorIndex Add(FGLTFConvertBuilder& Builder, const FString& Name, const FSkinWeightVertexBuffer* VertexBuffer, int32 InfluenceOffset) override;
+
+	template <typename IndexType>
+	static FGLTFJsonAccessorIndex Add(FGLTFConvertBuilder& Builder, const FString& Name, const FSkinWeightVertexBuffer* VertexBuffer, int32 InfluenceOffset);
 };
 
-class FGLTFBoneWeightVertexBufferConverter final : public TGLTFConverter<FGLTFJsonAccessorIndex, const FSkinWeightVertexBuffer*>
+class FGLTFBoneWeightVertexBufferConverter final : public TGLTFConverter<FGLTFJsonAccessorIndex, const FSkinWeightVertexBuffer*, int32>
 {
-	FGLTFJsonAccessorIndex Add(FGLTFConvertBuilder& Builder, const FString& Name, const FSkinWeightVertexBuffer* VertexBuffer) override;
+	FGLTFJsonAccessorIndex Add(FGLTFConvertBuilder& Builder, const FString& Name, const FSkinWeightVertexBuffer* VertexBuffer, int32 InfluenceOffset) override;
 };
