[GLTFExporter] fix(skeletal-mesh): workaround linking error using temporary hack

commit 4d0c1f4bbb8009382188f9e4deb36167c4f53aa7
Date: Fri, 29 May 2020 22:35:00 +0200

    fix(skeletal-mesh): workaround linking error using temporary hack
    (cherry picked from commit e035200d8cf955cb1d9914ec484fc7fc8b8611c5​)

#skipci

diff --git a/Engine/Plugins/Enterprise/GLTFExporter/MigratedFromGit/Source/GLTFExporter/Private/Converters/GLTFSkinWeightVertexBufferHack.h b/Engine/Plugins/Enterprise/GLTFExporter/MigratedFromGit/Source/GLTFExporter/Private/Converters/GLTFSkinWeightVertexBufferHack.h
new file mode 100644
--- /dev/null
+++ b/Engine/Plugins/Enterprise/GLTFExporter/MigratedFromGit/Source/GLTFExporter/Private/Converters/GLTFSkinWeightVertexBufferHack.h
@@ -0,0 +1,157 @@
+// Copyright Epic Games, Inc. All Rights Reserved.
+
+#pragma once
+
+#include "Rendering/SkinWeightVertexBuffer.h"
+
+class FGLTFSkinWeightDataVertexBufferHack : public FVertexBuffer
+{
+public:
+
+    FORCEINLINE uint32 GetBoneIndexByteSize() const
+    { return (Use16BitBoneIndex() ? sizeof(FBoneIndex16) : sizeof(FBoneIndex8)); }
+
+    FORCEINLINE bool Use16BitBoneIndex() const
+    { return bUse16BitBoneIndex; }
+
+	uint32 GetBoneIndex(uint32 VertexWeightOffset, uint32 VertexInfluenceCount, uint32 InfluenceIndex) const
+	{
+		if (InfluenceIndex < VertexInfluenceCount)
+		{
+			uint8* BoneData = Data + VertexWeightOffset;
+			if (Use16BitBoneIndex())
+			{
+				FBoneIndex16* BoneIndex16Ptr = (FBoneIndex16*)BoneData;
+				return BoneIndex16Ptr[InfluenceIndex];
+			}
+			else
+			{
+				return BoneData[InfluenceIndex];
+			}
+		}
+		else
+		{
+			return 0;
+		}
+	}
+
+	uint8 GetBoneWeight(uint32 VertexWeightOffset, uint32 VertexInfluenceCount, uint32 InfluenceIndex) const
+    {
+    	if (InfluenceIndex < VertexInfluenceCount)
+    	{
+    		uint8* BoneData = Data + VertexWeightOffset;
+    		uint32 BoneWeightOffset = GetBoneIndexByteSize() * VertexInfluenceCount;
+    		return BoneData[BoneWeightOffset + InfluenceIndex];
+    	}
+    	else
+    	{
+    		return 0;
+    	}
+	}
+
+	// guaranteed only to be valid if the vertex buffer is valid
+	FShaderResourceViewRHIRef SRVValue;
+
+	/** true if this vertex buffer will be used with CPU skinning. Resource arrays are set to cpu accessible if this is true */
+	bool bNeedsCPUAccess;
+
+	bool bVariableBonesPerVertex;
+
+	/** Has extra bone influences per Vertex, which means using a different TGPUSkinVertexBase */
+	uint32 MaxBoneInfluences;
+
+	/** Use 16 bit bone index instead of 8 bit */
+	bool bUse16BitBoneIndex;
+
+	/** The vertex data storage type */
+	FStaticMeshVertexDataInterface* WeightData;
+
+	/** The cached vertex data pointer. */
+	uint8* Data;
+
+	/** The cached number of vertices. */
+	uint32 NumVertices;
+
+	/** The cached number of bones. */
+	uint32 NumBones;
+};
+
+
+class FGLTFSkinWeightLookupVertexBufferHack : public FVertexBuffer
+{
+public:
+
+	void GetWeightOffsetAndInfluenceCount(uint32 VertexIndex, uint32& OutWeightOffset, uint32& OutInfluenceCount) const
+	{
+		uint32 Offset = VertexIndex * 4;
+		uint32 DataUInt32 = *((uint32*)(&Data[Offset]));
+		OutWeightOffset = DataUInt32 >> 8;
+		OutInfluenceCount = DataUInt32 & 0xff;
+	}
+
+	// guaranteed only to be valid if the vertex buffer is valid
+	FShaderResourceViewRHIRef SRVValue;
+
+	/** true if this vertex buffer will be used with CPU skinning. Resource arrays are set to cpu accessible if this is true */
+	bool bNeedsCPUAccess;
+
+	/** The vertex data storage type */
+	FStaticMeshVertexDataInterface* LookupData;
+
+	/** The cached vertex data pointer. */
+	uint8* Data;
+
+	/** The cached number of vertices. */
+	uint32 NumVertices;
+};
+
+class FGLTFSkinWeightVertexBufferHack
+{
+public:
+
+	FGLTFSkinWeightVertexBufferHack(const FSkinWeightVertexBuffer* VertexBuffer)
+		: VertexBuffer(*VertexBuffer)
+		, DataVertexBuffer(*(FGLTFSkinWeightDataVertexBufferHack*)VertexBuffer->GetDataVertexBuffer())
+		, LookupVertexBuffer(*(FGLTFSkinWeightLookupVertexBufferHack*)VertexBuffer->GetLookupVertexBuffer())
+	{
+	}
+
+	bool GetVariableBonesPerVertex() const { return VertexBuffer.GetVariableBonesPerVertex(); }
+	uint32 GetConstantInfluencesVertexStride() const { return VertexBuffer.GetConstantInfluencesVertexStride(); }
+	uint32 GetMaxBoneInfluences() const { return VertexBuffer.GetMaxBoneInfluences(); }
+
+	void GetVertexInfluenceOffsetCount(uint32 VertexIndex, uint32& VertexWeightOffset, uint32& VertexInfluenceCount) const
+	{
+		if (GetVariableBonesPerVertex())
+		{
+			LookupVertexBuffer.GetWeightOffsetAndInfluenceCount(VertexIndex, VertexWeightOffset, VertexInfluenceCount);
+		}
+		else
+		{
+			VertexWeightOffset = GetConstantInfluencesVertexStride() * VertexIndex;
+			VertexInfluenceCount = GetMaxBoneInfluences();
+		}
+	}
+
+	uint32 GetBoneIndex(uint32 VertexIndex, uint32 InfluenceIndex) const
+	{
+		uint32 VertexWeightOffset = 0;
+		uint32 VertexInfluenceCount = 0;
+		GetVertexInfluenceOffsetCount(VertexIndex, VertexWeightOffset, VertexInfluenceCount);
+		return DataVertexBuffer.GetBoneIndex(VertexWeightOffset, VertexInfluenceCount, InfluenceIndex);
+	}
+
+	uint8 GetBoneWeight(uint32 VertexIndex, uint32 InfluenceIndex) const
+	{
+		uint32 VertexWeightOffset = 0;
+		uint32 VertexInfluenceCount = 0;
+		GetVertexInfluenceOffsetCount(VertexIndex, VertexWeightOffset, VertexInfluenceCount);
+		return DataVertexBuffer.GetBoneWeight(VertexWeightOffset, VertexInfluenceCount, InfluenceIndex);
+	}
+
+private:
+
+	const FSkinWeightVertexBuffer& VertexBuffer;
+	const FGLTFSkinWeightDataVertexBufferHack& DataVertexBuffer;
+	const FGLTFSkinWeightLookupVertexBufferHack& LookupVertexBuffer;
+};
diff --git a/Engine/Plugins/Enterprise/GLTFExporter/MigratedFromGit/Source/GLTFExporter/Private/Converters/GLTFVertexBufferConverters.cpp b/Engine/Plugins/Enterprise/GLTFExporter/MigratedFromGit/Source/GLTFExporter/Private/Converters/GLTFVertexBufferConverters.cpp
--- a/Engine/Plugins/Enterprise/GLTFExporter/MigratedFromGit/Source/GLTFExporter/Private/Converters/GLTFVertexBufferConverters.cpp
+++ b/Engine/Plugins/Enterprise/GLTFExporter/MigratedFromGit/Source/GLTFExporter/Private/Converters/GLTFVertexBufferConverters.cpp
@@ -1,6 +1,7 @@
 // Copyright Epic Games, Inc. All Rights Reserved.
 
 #include "Converters/GLTFVertexBufferConverters.h"
+#include "Converters/GLTFSkinWeightVertexBufferHack.h"
 #include "Converters/GLTFConverterUtility.h"
 #include "Builders/GLTFConvertBuilder.h"
 
@@ -187,7 +188,9 @@ FGLTFJsonAccessorIndex FGLTFBoneIndexVertexBufferConverter::Add(FGLTFConvertBuil
 	{
 		for (int32 InfluenceIndex = 0; InfluenceIndex < InfluenceCount; ++InfluenceIndex)
 		{
-			BoneIndices[VertexIndex].Index[InfluenceIndex] = static_cast<IndexType>(VertexBuffer->GetBoneIndex(VertexIndex, InfluenceOffset + InfluenceIndex));
+			// TODO: remove hack
+			const uint32 BoneIndex = FGLTFSkinWeightVertexBufferHack(VertexBuffer).GetBoneIndex(VertexIndex, InfluenceOffset + InfluenceIndex);
+			BoneIndices[VertexIndex].Index[InfluenceIndex] = static_cast<IndexType>(BoneIndex);
 		}
 
 		for (int32 InfluenceIndex = InfluenceCount; InfluenceIndex < 4; ++InfluenceIndex)
@@ -229,7 +232,9 @@ FGLTFJsonAccessorIndex FGLTFBoneWeightVertexBufferConverter::Add(FGLTFConvertBui
 	{
 		for (int32 InfluenceIndex = 0; InfluenceIndex < InfluenceCount; ++InfluenceIndex)
 		{
-			BoneWeights[VertexIndex].Weights[InfluenceIndex] = VertexBuffer->GetBoneWeight(VertexIndex, InfluenceOffset + InfluenceIndex);
+			// TODO: remove hack
+			const uint8 BoneWeight = FGLTFSkinWeightVertexBufferHack(VertexBuffer).GetBoneWeight(VertexIndex, InfluenceOffset + InfluenceIndex);
+			BoneWeights[VertexIndex].Weights[InfluenceIndex] = BoneWeight;
 		}
 
 		for (int32 InfluenceIndex = InfluenceCount; InfluenceIndex < 4; ++InfluenceIndex)
