[GLTFExporter] feat(skeletal-mesh): add sanity check for bone index converter

commit 44c475874efd6a960d4acc85954643cd67dd6312
Date: Fri, 29 May 2020 22:35:41 +0200

    feat(skeletal-mesh): add sanity check for bone index converter
    (cherry picked from commit 7df3fc504f33557b5ea4f1f349060dc00bfda527​)

#skipci

diff --git a/Engine/Plugins/Enterprise/GLTFExporter/MigratedFromGit/Source/GLTFExporter/Private/Converters/GLTFVertexBufferConverters.cpp b/Engine/Plugins/Enterprise/GLTFExporter/MigratedFromGit/Source/GLTFExporter/Private/Converters/GLTFVertexBufferConverters.cpp
--- a/Engine/Plugins/Enterprise/GLTFExporter/MigratedFromGit/Source/GLTFExporter/Private/Converters/GLTFVertexBufferConverters.cpp
+++ b/Engine/Plugins/Enterprise/GLTFExporter/MigratedFromGit/Source/GLTFExporter/Private/Converters/GLTFVertexBufferConverters.cpp
@@ -4,6 +4,7 @@
 #include "Converters/GLTFSkinWeightVertexBufferHack.h"
 #include "Converters/GLTFConverterUtility.h"
 #include "Builders/GLTFConvertBuilder.h"
+#include "Math/NumericLimits.h"
 
 FGLTFJsonAccessorIndex FGLTFPositionVertexBufferConverter::Add(FGLTFConvertBuilder& Builder, const FString& Name, const FPositionVertexBuffer* VertexBuffer)
 {
@@ -190,6 +191,8 @@ FGLTFJsonAccessorIndex FGLTFBoneIndexVertexBufferConverter::Add(FGLTFConvertBuil
 		{
 			// TODO: remove hack
 			const uint32 BoneIndex = FGLTFSkinWeightVertexBufferHack(VertexBuffer).GetBoneIndex(VertexIndex, InfluenceOffset + InfluenceIndex);
+			check(BoneIndex <= TNumericLimits<IndexType>::Max());
+
 			BoneIndices[VertexIndex].Index[InfluenceIndex] = static_cast<IndexType>(BoneIndex);
 		}
 
