[GLTFExporter] refactor(skeletal-mesh): change InfluenceOffset to JointsGroupIndex/WeightsGroupIndex

commit 493ec272a5700654d900a24c711659b1ab4b6fce
Date: Sat, 30 May 2020 02:15:40 +0200

    refactor(skeletal-mesh): change InfluenceOffset to JointsGroupIndex/WeightsGroupIndex
    (cherry picked from commit aa0b33b97d4a639b1c23eb8c18a74c75a8a5b398​)

#skipci

diff --git a/Engine/Plugins/Enterprise/GLTFExporter/MigratedFromGit/Source/GLTFExporter/Private/Builders/GLTFConvertBuilder.cpp b/Engine/Plugins/Enterprise/GLTFExporter/MigratedFromGit/Source/GLTFExporter/Private/Builders/GLTFConvertBuilder.cpp
--- a/Engine/Plugins/Enterprise/GLTFExporter/MigratedFromGit/Source/GLTFExporter/Private/Builders/GLTFConvertBuilder.cpp
+++ b/Engine/Plugins/Enterprise/GLTFExporter/MigratedFromGit/Source/GLTFExporter/Private/Builders/GLTFConvertBuilder.cpp
@@ -59,24 +59,24 @@ FGLTFJsonAccessorIndex FGLTFConvertBuilder::GetOrAddUVAccessor(const FStaticMesh
 	return UVVertexBufferConverter.GetOrAdd(*this, DesiredName, VertexBuffer, UVIndex);
 }
 
-FGLTFJsonAccessorIndex FGLTFConvertBuilder::GetOrAddJointAccessor(const FSkinWeightVertexBuffer* VertexBuffer, int32 InfluenceOffset, const FString& DesiredName)
+FGLTFJsonAccessorIndex FGLTFConvertBuilder::GetOrAddJointAccessor(const FSkinWeightVertexBuffer* VertexBuffer, int32 JointsGroupIndex, const FString& DesiredName)
 {
 	if (VertexBuffer == nullptr)
 	{
 		return FGLTFJsonAccessorIndex(INDEX_NONE);
 	}
 
-	return BoneIndexVertexBufferConverter.GetOrAdd(*this, DesiredName, VertexBuffer, InfluenceOffset);
+	return BoneIndexVertexBufferConverter.GetOrAdd(*this, DesiredName, VertexBuffer, JointsGroupIndex);
 }
 
-FGLTFJsonAccessorIndex FGLTFConvertBuilder::GetOrAddWeightAccessor(const FSkinWeightVertexBuffer* VertexBuffer, int32 InfluenceOffset, const FString& DesiredName)
+FGLTFJsonAccessorIndex FGLTFConvertBuilder::GetOrAddWeightAccessor(const FSkinWeightVertexBuffer* VertexBuffer, int32 WeightsGroupIndex, const FString& DesiredName)
 {
 	if (VertexBuffer == nullptr)
 	{
 		return FGLTFJsonAccessorIndex(INDEX_NONE);
 	}
 
-	return BoneWeightVertexBufferConverter.GetOrAdd(*this, DesiredName, VertexBuffer, InfluenceOffset);
+	return BoneWeightVertexBufferConverter.GetOrAdd(*this, DesiredName, VertexBuffer, WeightsGroupIndex);
 }
 
 FGLTFJsonBufferViewIndex FGLTFConvertBuilder::GetOrAddIndexBufferView(const FRawStaticIndexBuffer* IndexBuffer, const FString& DesiredName)
diff --git a/Engine/Plugins/Enterprise/GLTFExporter/MigratedFromGit/Source/GLTFExporter/Private/Builders/GLTFConvertBuilder.h b/Engine/Plugins/Enterprise/GLTFExporter/MigratedFromGit/Source/GLTFExporter/Private/Builders/GLTFConvertBuilder.h
--- a/Engine/Plugins/Enterprise/GLTFExporter/MigratedFromGit/Source/GLTFExporter/Private/Builders/GLTFConvertBuilder.h
+++ b/Engine/Plugins/Enterprise/GLTFExporter/MigratedFromGit/Source/GLTFExporter/Private/Builders/GLTFConvertBuilder.h
@@ -26,8 +26,8 @@ public:
 	FGLTFJsonAccessorIndex GetOrAddNormalAccessor(const FStaticMeshVertexBuffer* VertexBuffer, const FString& DesiredName = TEXT(""));
 	FGLTFJsonAccessorIndex GetOrAddTangentAccessor(const FStaticMeshVertexBuffer* VertexBuffer, const FString& DesiredName = TEXT(""));
 	FGLTFJsonAccessorIndex GetOrAddUVAccessor(const FStaticMeshVertexBuffer* VertexBuffer, int32 UVIndex, const FString& DesiredName = TEXT(""));
-	FGLTFJsonAccessorIndex GetOrAddJointAccessor(const FSkinWeightVertexBuffer* VertexBuffer, int32 InfluenceOffset, const FString& DesiredName = TEXT(""));
-	FGLTFJsonAccessorIndex GetOrAddWeightAccessor(const FSkinWeightVertexBuffer* VertexBuffer, int32 InfluenceOffset, const FString& DesiredName = TEXT(""));
+	FGLTFJsonAccessorIndex GetOrAddJointAccessor(const FSkinWeightVertexBuffer* VertexBuffer, int32 JointsGroupIndex, const FString& DesiredName = TEXT(""));
+	FGLTFJsonAccessorIndex GetOrAddWeightAccessor(const FSkinWeightVertexBuffer* VertexBuffer, int32 WeightsGroupIndex, const FString& DesiredName = TEXT(""));
 
 	FGLTFJsonBufferViewIndex GetOrAddIndexBufferView(const FRawStaticIndexBuffer* IndexBuffer, const FString& DesiredName = TEXT(""));
 	FGLTFJsonAccessorIndex GetOrAddIndexAccessor(const FStaticMeshSection* MeshSection, const FRawStaticIndexBuffer* IndexBuffer, const FString& DesiredName = TEXT(""));
diff --git a/Engine/Plugins/Enterprise/GLTFExporter/MigratedFromGit/Source/GLTFExporter/Private/Converters/GLTFVertexBufferConverters.cpp b/Engine/Plugins/Enterprise/GLTFExporter/MigratedFromGit/Source/GLTFExporter/Private/Converters/GLTFVertexBufferConverters.cpp
--- a/Engine/Plugins/Enterprise/GLTFExporter/MigratedFromGit/Source/GLTFExporter/Private/Converters/GLTFVertexBufferConverters.cpp
+++ b/Engine/Plugins/Enterprise/GLTFExporter/MigratedFromGit/Source/GLTFExporter/Private/Converters/GLTFVertexBufferConverters.cpp
@@ -166,10 +166,12 @@ FGLTFJsonAccessorIndex FGLTFBoneIndexVertexBufferConverter::Add(FGLTFConvertBuil
 }
 
 template <typename IndexType>
-FGLTFJsonAccessorIndex FGLTFBoneIndexVertexBufferConverter::Add(FGLTFConvertBuilder& Builder, const FString& Name, const FSkinWeightVertexBuffer* VertexBuffer, int32 InfluenceOffset)
+FGLTFJsonAccessorIndex FGLTFBoneIndexVertexBufferConverter::Add(FGLTFConvertBuilder& Builder, const FString& Name, const FSkinWeightVertexBuffer* VertexBuffer, int32 JointsGroupIndex)
 {
 	const uint32 VertexCount = VertexBuffer->GetNumVertices();
 	const int32 MaxInfluenceCount = VertexBuffer->GetMaxBoneInfluences();
+	const int32 InfluenceOffset = JointsGroupIndex * 4;
+
 	if (VertexCount == 0 || InfluenceOffset < 0 || MaxInfluenceCount <= InfluenceOffset)
 	{
 		return FGLTFJsonAccessorIndex(INDEX_NONE);
@@ -212,10 +214,12 @@ FGLTFJsonAccessorIndex FGLTFBoneIndexVertexBufferConverter::Add(FGLTFConvertBuil
 	return Builder.AddAccessor(JsonAccessor);
 }
 
-FGLTFJsonAccessorIndex FGLTFBoneWeightVertexBufferConverter::Add(FGLTFConvertBuilder& Builder, const FString& Name, const FSkinWeightVertexBuffer* VertexBuffer, int32 InfluenceOffset)
+FGLTFJsonAccessorIndex FGLTFBoneWeightVertexBufferConverter::Add(FGLTFConvertBuilder& Builder, const FString& Name, const FSkinWeightVertexBuffer* VertexBuffer, int32 WeightsGroupIndex)
 {
 	const uint32 VertexCount = VertexBuffer->GetNumVertices();
 	const int32 MaxInfluenceCount = VertexBuffer->GetMaxBoneInfluences();
+	const int32 InfluenceOffset = WeightsGroupIndex * 4;
+
 	if (VertexCount == 0 || InfluenceOffset < 0 || MaxInfluenceCount <= InfluenceOffset)
 	{
 		return FGLTFJsonAccessorIndex(INDEX_NONE);
diff --git a/Engine/Plugins/Enterprise/GLTFExporter/MigratedFromGit/Source/GLTFExporter/Private/Converters/GLTFVertexBufferConverters.h b/Engine/Plugins/Enterprise/GLTFExporter/MigratedFromGit/Source/GLTFExporter/Private/Converters/GLTFVertexBufferConverters.h
--- a/Engine/Plugins/Enterprise/GLTFExporter/MigratedFromGit/Source/GLTFExporter/Private/Converters/GLTFVertexBufferConverters.h
+++ b/Engine/Plugins/Enterprise/GLTFExporter/MigratedFromGit/Source/GLTFExporter/Private/Converters/GLTFVertexBufferConverters.h
@@ -36,10 +36,10 @@ class FGLTFBoneIndexVertexBufferConverter final : public TGLTFConverter<FGLTFJso
 	FGLTFJsonAccessorIndex Add(FGLTFConvertBuilder& Builder, const FString& Name, const FSkinWeightVertexBuffer* VertexBuffer, int32 InfluenceOffset) override;
 
 	template <typename IndexType>
-	static FGLTFJsonAccessorIndex Add(FGLTFConvertBuilder& Builder, const FString& Name, const FSkinWeightVertexBuffer* VertexBuffer, int32 InfluenceOffset);
+	static FGLTFJsonAccessorIndex Add(FGLTFConvertBuilder& Builder, const FString& Name, const FSkinWeightVertexBuffer* VertexBuffer, int32 JointsGroupIndex);
 };
 
 class FGLTFBoneWeightVertexBufferConverter final : public TGLTFConverter<FGLTFJsonAccessorIndex, const FSkinWeightVertexBuffer*, int32>
 {
-	FGLTFJsonAccessorIndex Add(FGLTFConvertBuilder& Builder, const FString& Name, const FSkinWeightVertexBuffer* VertexBuffer, int32 InfluenceOffset) override;
+	FGLTFJsonAccessorIndex Add(FGLTFConvertBuilder& Builder, const FString& Name, const FSkinWeightVertexBuffer* VertexBuffer, int32 WeightsGroupIndex) override;
 };
