[GLTFExporter] feat(skeletal-mesh): implement skeletal mesh converters

commit ea4442e2f35a5d25eae942c2d7025bdaddc26c34
Date: Sat, 30 May 2020 02:18:24 +0200

    feat(skeletal-mesh): implement skeletal mesh converters
    (cherry picked from commit 7c6377263fa9bdd5987d2ef8c36c125b1bb798a4​)

#skipci

diff --git a/Engine/Plugins/Enterprise/GLTFExporter/MigratedFromGit/Source/GLTFExporter/Private/Converters/GLTFSkeletalMeshConverters.cpp b/Engine/Plugins/Enterprise/GLTFExporter/MigratedFromGit/Source/GLTFExporter/Private/Converters/GLTFSkeletalMeshConverters.cpp
--- a/Engine/Plugins/Enterprise/GLTFExporter/MigratedFromGit/Source/GLTFExporter/Private/Converters/GLTFSkeletalMeshConverters.cpp
+++ b/Engine/Plugins/Enterprise/GLTFExporter/MigratedFromGit/Source/GLTFExporter/Private/Converters/GLTFSkeletalMeshConverters.cpp
@@ -3,18 +3,128 @@
 #include "Converters/GLTFSkeletalMeshConverters.h"
 #include "Converters/GLTFConverterUtility.h"
 #include "Builders/GLTFConvertBuilder.h"
+#include "Rendering/MultiSizeIndexContainer.h"
+#include "Rendering/SkeletalMeshRenderData.h"
 
 FGLTFJsonBufferViewIndex FGLTFIndexContainerConverter::Add(FGLTFConvertBuilder& Builder, const FString& Name, const FMultiSizeIndexContainer* IndexContainer)
 {
-	return FGLTFJsonBufferViewIndex(INDEX_NONE);
+	const FRawStaticIndexBuffer16or32Interface* IndexBuffer = IndexContainer->GetIndexBuffer();
+
+	const int32 IndexCount = IndexBuffer->Num();
+	if (IndexCount <= 0)
+	{
+		return FGLTFJsonBufferViewIndex(INDEX_NONE);
+	}
+
+	if (IndexContainer->GetDataTypeSize() == sizeof(uint16))
+	{
+		TArray<uint16> Indices;
+		Indices.AddUninitialized(IndexCount);
+
+		for (int32 Index = 0; Index < IndexCount; ++Index)
+		{
+			Indices[Index] = static_cast<uint16>(IndexBuffer->Get(Index));
+		}
+
+		return Builder.AddBufferView(Indices, Name, sizeof(uint16), EGLTFJsonBufferTarget::ElementArrayBuffer);
+	}
+	else
+	{
+		TArray<uint32> Indices;
+		IndexContainer->GetIndexBuffer(Indices);
+		return Builder.AddBufferView(Indices, Name, sizeof(uint32), EGLTFJsonBufferTarget::ElementArrayBuffer);
+	}
 }
 
 FGLTFJsonAccessorIndex FGLTFSkeletalMeshSectionConverter::Add(FGLTFConvertBuilder& Builder, const FString& Name, const FSkelMeshRenderSection* MeshSection, const FMultiSizeIndexContainer* IndexContainer)
 {
-	return FGLTFJsonAccessorIndex(INDEX_NONE);
+	const uint32 TriangleCount = MeshSection->NumTriangles;
+	if (TriangleCount == 0)
+	{
+		return FGLTFJsonAccessorIndex(INDEX_NONE);
+	}
+
+	const uint32 FirstIndex = MeshSection->BaseIndex;
+	const bool bIs32Bit = IndexContainer->GetDataTypeSize() == sizeof(uint32);
+
+	FGLTFJsonAccessor JsonAccessor;
+	JsonAccessor.Name = Name;
+	JsonAccessor.BufferView = Builder.GetOrAddIndexBufferView(IndexContainer);
+	JsonAccessor.ByteOffset = FirstIndex * (bIs32Bit ? sizeof(uint32) : sizeof(uint16));
+	JsonAccessor.ComponentType = bIs32Bit ? EGLTFJsonComponentType::U32 : EGLTFJsonComponentType::U16;
+	JsonAccessor.Count = TriangleCount * 3;
+	JsonAccessor.Type = EGLTFJsonAccessorType::Scalar;
+
+	return Builder.AddAccessor(JsonAccessor);
 }
 
 FGLTFJsonMeshIndex FGLTFSkeletalMeshConverter::Add(FGLTFConvertBuilder& Builder, const FString& Name, const USkeletalMesh* SkeletalMesh, int32 LODIndex, const FColorVertexBuffer* OverrideVertexColors, const FSkinWeightVertexBuffer* OverrideSkinWeights, FGLTFMaterialArray OverrideMaterials)
 {
-	return FGLTFJsonMeshIndex(INDEX_NONE);
+	const FSkeletalMeshRenderData* RenderData = SkeletalMesh->GetResourceForRendering();
+
+	if (LODIndex < 0 || RenderData->LODRenderData.Num() <= LODIndex)
+	{
+		return FGLTFJsonMeshIndex(INDEX_NONE);
+	}
+
+	const FSkeletalMeshLODRenderData& MeshLOD = RenderData->LODRenderData[LODIndex];
+
+	FGLTFJsonMesh JsonMesh;
+	JsonMesh.Name = Name;
+
+	const FPositionVertexBuffer* PositionBuffer = &MeshLOD.StaticVertexBuffers.PositionVertexBuffer;
+	const FStaticMeshVertexBuffer* VertexBuffer = &MeshLOD.StaticVertexBuffers.StaticMeshVertexBuffer;
+	const FColorVertexBuffer* ColorBuffer = OverrideVertexColors != nullptr ? OverrideVertexColors : &MeshLOD.StaticVertexBuffers.ColorVertexBuffer;
+	const FSkinWeightVertexBuffer* SkinWeightBuffer = OverrideSkinWeights != nullptr ? OverrideSkinWeights : MeshLOD.GetSkinWeightVertexBuffer();
+
+	FGLTFJsonAttributes JsonAttributes;
+	JsonAttributes.Position = Builder.GetOrAddPositionAccessor(PositionBuffer, Name + TEXT("_Positions"));
+	JsonAttributes.Color0 = Builder.GetOrAddColorAccessor(ColorBuffer, Name + TEXT("_Colors"));
+	JsonAttributes.Normal = Builder.GetOrAddNormalAccessor(VertexBuffer, Name + TEXT("_Normals"));
+	JsonAttributes.Tangent = Builder.GetOrAddTangentAccessor(VertexBuffer, Name + TEXT("_Tangents"));
+
+	const uint32 UVCount = VertexBuffer->GetNumTexCoords();
+	JsonAttributes.TexCoords.AddUninitialized(UVCount);
+
+	for (uint32 UVIndex = 0; UVIndex < UVCount; ++UVIndex)
+	{
+		JsonAttributes.TexCoords[UVIndex] = Builder.GetOrAddUVAccessor(VertexBuffer, UVIndex, Name + TEXT("_UV") + FString::FromInt(UVIndex) + TEXT("s"));
+	}
+
+	const uint32 GroupCount = (SkinWeightBuffer->GetMaxBoneInfluences() + 3) / 4;
+	JsonAttributes.Joints.AddUninitialized(GroupCount);
+	JsonAttributes.Weights.AddUninitialized(GroupCount);
+
+	for (uint32 GroupIndex = 0; GroupIndex < GroupCount; ++GroupIndex)
+	{
+		JsonAttributes.Joints[GroupIndex] = Builder.GetOrAddJointAccessor(SkinWeightBuffer, GroupIndex, Name + TEXT("_Joints") + FString::FromInt(GroupIndex));
+		JsonAttributes.Weights[GroupIndex] = Builder.GetOrAddWeightAccessor(SkinWeightBuffer, GroupIndex, Name + TEXT("_Weights") + FString::FromInt(GroupIndex));
+	}
+
+	const FMultiSizeIndexContainer* IndexContainer = &MeshLOD.MultiSizeIndexContainer;
+	Builder.GetOrAddIndexBufferView(IndexContainer, Name + TEXT("_Indices"));
+
+	const int32 SectionCount = MeshLOD.RenderSections.Num();
+	JsonMesh.Primitives.AddDefaulted(SectionCount);
+
+	for (int32 SectionIndex = 0; SectionIndex < SectionCount; ++SectionIndex)
+	{
+		FGLTFJsonPrimitive& JsonPrimitive = JsonMesh.Primitives[SectionIndex];
+		JsonPrimitive.Attributes = JsonAttributes;
+
+		const FSkelMeshRenderSection& Section = MeshLOD.RenderSections[SectionIndex];
+		JsonPrimitive.Indices = Builder.GetOrAddIndexAccessor(&Section, IndexContainer,
+			Name + (SectionCount == 1 ? TEXT("_Indices") : TEXT("_Indices_Section") + FString::FromInt(SectionIndex)));
+
+		const int32 MaterialIndex = Section.MaterialIndex;
+		const UMaterialInterface* Material = OverrideMaterials.IsValidIndex(MaterialIndex) && OverrideMaterials[MaterialIndex] != nullptr ?
+			OverrideMaterials[MaterialIndex] : SkeletalMesh->Materials[MaterialIndex].MaterialInterface;
+
+		if (Material != nullptr)
+		{
+			JsonPrimitive.Material = Builder.GetOrAddMaterial(Material);
+		}
+	}
+
+	return Builder.AddMesh(JsonMesh);
 }
diff --git a/Engine/Plugins/Enterprise/GLTFExporter/MigratedFromGit/Source/GLTFExporter/Private/Converters/GLTFStaticMeshConverters.cpp b/Engine/Plugins/Enterprise/GLTFExporter/MigratedFromGit/Source/GLTFExporter/Private/Converters/GLTFStaticMeshConverters.cpp
--- a/Engine/Plugins/Enterprise/GLTFExporter/MigratedFromGit/Source/GLTFExporter/Private/Converters/GLTFStaticMeshConverters.cpp
+++ b/Engine/Plugins/Enterprise/GLTFExporter/MigratedFromGit/Source/GLTFExporter/Private/Converters/GLTFStaticMeshConverters.cpp
@@ -54,14 +54,14 @@ FGLTFJsonMeshIndex FGLTFStaticMeshConverter::Add(FGLTFConvertBuilder& Builder, c
 		return FGLTFJsonMeshIndex(INDEX_NONE);
 	}
 
-	const FStaticMeshLODResources& LODResources = StaticMesh->GetLODForExport(LODIndex);
+	const FStaticMeshLODResources& MeshLOD = StaticMesh->GetLODForExport(LODIndex);
 
 	FGLTFJsonMesh JsonMesh;
 	JsonMesh.Name = Name;
 
-	const FPositionVertexBuffer* PositionBuffer = &LODResources.VertexBuffers.PositionVertexBuffer;
-	const FStaticMeshVertexBuffer* VertexBuffer = &LODResources.VertexBuffers.StaticMeshVertexBuffer;
-	const FColorVertexBuffer* ColorBuffer = OverrideVertexColors != nullptr ? OverrideVertexColors : &LODResources.VertexBuffers.ColorVertexBuffer;
+	const FPositionVertexBuffer* PositionBuffer = &MeshLOD.VertexBuffers.PositionVertexBuffer;
+	const FStaticMeshVertexBuffer* VertexBuffer = &MeshLOD.VertexBuffers.StaticMeshVertexBuffer;
+	const FColorVertexBuffer* ColorBuffer = OverrideVertexColors != nullptr ? OverrideVertexColors : &MeshLOD.VertexBuffers.ColorVertexBuffer;
 
 	FGLTFJsonAttributes JsonAttributes;
 	JsonAttributes.Position = Builder.GetOrAddPositionAccessor(PositionBuffer, Name + TEXT("_Positions"));
@@ -77,10 +77,10 @@ FGLTFJsonMeshIndex FGLTFStaticMeshConverter::Add(FGLTFConvertBuilder& Builder, c
 		JsonAttributes.TexCoords[UVIndex] = Builder.GetOrAddUVAccessor(VertexBuffer, UVIndex, Name + TEXT("_UV") + FString::FromInt(UVIndex) + TEXT("s"));
 	}
 
-	const FRawStaticIndexBuffer* IndexBuffer = &LODResources.IndexBuffer;
+	const FRawStaticIndexBuffer* IndexBuffer = &MeshLOD.IndexBuffer;
 	Builder.GetOrAddIndexBufferView(IndexBuffer, Name + TEXT("_Indices"));
 
-	const int32 SectionCount = LODResources.Sections.Num();
+	const int32 SectionCount = MeshLOD.Sections.Num();
 	JsonMesh.Primitives.AddDefaulted(SectionCount);
 
 	for (int32 SectionIndex = 0; SectionIndex < SectionCount; ++SectionIndex)
@@ -88,7 +88,7 @@ FGLTFJsonMeshIndex FGLTFStaticMeshConverter::Add(FGLTFConvertBuilder& Builder, c
 		FGLTFJsonPrimitive& JsonPrimitive = JsonMesh.Primitives[SectionIndex];
 		JsonPrimitive.Attributes = JsonAttributes;
 
-		const FStaticMeshSection& Section = LODResources.Sections[SectionIndex];
+		const FStaticMeshSection& Section = MeshLOD.Sections[SectionIndex];
 		JsonPrimitive.Indices = Builder.GetOrAddIndexAccessor(&Section, IndexBuffer,
 			Name + (SectionCount == 1 ? TEXT("_Indices") : TEXT("_Indices_Section") + FString::FromInt(SectionIndex)));
 
diff --git a/Engine/Plugins/Enterprise/GLTFExporter/MigratedFromGit/Source/GLTFExporter/Private/Exporters/GLTFSkeletalMeshExporter.cpp b/Engine/Plugins/Enterprise/GLTFExporter/MigratedFromGit/Source/GLTFExporter/Private/Exporters/GLTFSkeletalMeshExporter.cpp
--- a/Engine/Plugins/Enterprise/GLTFExporter/MigratedFromGit/Source/GLTFExporter/Private/Exporters/GLTFSkeletalMeshExporter.cpp
+++ b/Engine/Plugins/Enterprise/GLTFExporter/MigratedFromGit/Source/GLTFExporter/Private/Exporters/GLTFSkeletalMeshExporter.cpp
@@ -1,18 +1,28 @@
 // Copyright Epic Games, Inc. All Rights Reserved.
 
 #include "Exporters/GLTFSkeletalMeshExporter.h"
+#include "Builders/GLTFContainerBuilder.h"
 #include "Engine/SkeletalMesh.h"
 
 UGLTFSkeletalMeshExporter::UGLTFSkeletalMeshExporter(const FObjectInitializer& ObjectInitializer)
 	: Super(ObjectInitializer)
 {
-	// TODO: uncomment when support is implemented
-	// SupportedClass = USkeletalMesh::StaticClass();
+	SupportedClass = USkeletalMesh::StaticClass();
 }
 
 bool UGLTFSkeletalMeshExporter::AddObject(FGLTFContainerBuilder& Builder, const UObject* Object)
 {
 	const USkeletalMesh* SkeletalMesh = CastChecked<USkeletalMesh>(Object);
-	// TODO: implement
+	const FGLTFJsonMeshIndex MeshIndex = Builder.GetOrAddMesh(SkeletalMesh);
+
+	FGLTFJsonNode Node;
+	Node.Mesh = MeshIndex;
+	const FGLTFJsonNodeIndex NodeIndex = Builder.AddNode(Node);
+
+	FGLTFJsonScene Scene;
+	Scene.Nodes.Add(NodeIndex);
+	const FGLTFJsonSceneIndex SceneIndex = Builder.AddScene(Scene);
+
+	Builder.DefaultScene = SceneIndex;
 	return true;
 }
diff --git a/Engine/Plugins/Enterprise/GLTFExporter/MigratedFromGit/Source/GLTFExporter/Private/Json/GLTFJsonMesh.h b/Engine/Plugins/Enterprise/GLTFExporter/MigratedFromGit/Source/GLTFExporter/Private/Json/GLTFJsonMesh.h
--- a/Engine/Plugins/Enterprise/GLTFExporter/MigratedFromGit/Source/GLTFExporter/Private/Json/GLTFJsonMesh.h
+++ b/Engine/Plugins/Enterprise/GLTFExporter/MigratedFromGit/Source/GLTFExporter/Private/Json/GLTFJsonMesh.h
@@ -15,8 +15,8 @@ struct FGLTFJsonAttributes
 	TArray<FGLTFJsonAccessorIndex> TexCoords;
 
 	// skeletal mesh attributes
-	FGLTFJsonAccessorIndex Joints0;
-	FGLTFJsonAccessorIndex Weights0;
+	TArray<FGLTFJsonAccessorIndex> Joints;
+	TArray<FGLTFJsonAccessorIndex> Weights;
 
 	template <class CharType = TCHAR, class PrintPolicy = TPrettyJsonPrintPolicy<CharType>>
 	void WriteObject(TJsonWriter<CharType, PrintPolicy>& JsonWriter, FGLTFJsonExtensions& Extensions) const
@@ -34,8 +34,15 @@ struct FGLTFJsonAttributes
 			JsonWriter.WriteValue(TEXT("TEXCOORD_") + FString::FromInt(Index), TexCoords[Index]);
 		}
 
-		if (Joints0 != INDEX_NONE) JsonWriter.WriteValue(TEXT("JOINTS_0"), Joints0);
-		if (Weights0 != INDEX_NONE) JsonWriter.WriteValue(TEXT("WEIGHTS_0"), Weights0);
+		for (int32 Index = 0; Index < Joints.Num(); ++Index)
+		{
+			JsonWriter.WriteValue(TEXT("JOINTS_") + FString::FromInt(Index), Joints[Index]);
+		}
+
+		for (int32 Index = 0; Index < Weights.Num(); ++Index)
+		{
+			JsonWriter.WriteValue(TEXT("WEIGHTS_") + FString::FromInt(Index), Weights[Index]);
+		}
 
 		JsonWriter.WriteObjectEnd();
 	}
