[GLTFExporter] fix(skeletal-mesh): translate vertex bone indices on each vertex using bonemap

commit 593259a390f34226895e42e8be2804416b9131c9
Date: Mon, 1 Jun 2020 22:54:37 +0200

    fix(skeletal-mesh): translate vertex bone indices on each vertex using bonemap
    (cherry picked from commit e9ec8387e013f1c5fd993596c3e382e255112b42​)

#skipci

diff --git a/Engine/Plugins/Enterprise/GLTFExporter/MigratedFromGit/Source/GLTFExporter/Private/Builders/GLTFConvertBuilder.cpp b/Engine/Plugins/Enterprise/GLTFExporter/MigratedFromGit/Source/GLTFExporter/Private/Builders/GLTFConvertBuilder.cpp
--- a/Engine/Plugins/Enterprise/GLTFExporter/MigratedFromGit/Source/GLTFExporter/Private/Builders/GLTFConvertBuilder.cpp
+++ b/Engine/Plugins/Enterprise/GLTFExporter/MigratedFromGit/Source/GLTFExporter/Private/Builders/GLTFConvertBuilder.cpp
@@ -59,14 +59,14 @@ FGLTFJsonAccessorIndex FGLTFConvertBuilder::GetOrAddUVAccessor(const FStaticMesh
 	return UVVertexBufferConverter.GetOrAdd(*this, DesiredName, VertexBuffer, UVIndex);
 }
 
-FGLTFJsonAccessorIndex FGLTFConvertBuilder::GetOrAddJointAccessor(const FSkinWeightVertexBuffer* VertexBuffer, int32 JointsGroupIndex, const FString& DesiredName)
+FGLTFJsonAccessorIndex FGLTFConvertBuilder::GetOrAddJointAccessor(const FSkinWeightVertexBuffer* VertexBuffer, int32 JointsGroupIndex, FGLTFBoneMap BoneMap, const FString& DesiredName)
 {
 	if (VertexBuffer == nullptr)
 	{
 		return FGLTFJsonAccessorIndex(INDEX_NONE);
 	}
 
-	return BoneIndexVertexBufferConverter.GetOrAdd(*this, DesiredName, VertexBuffer, JointsGroupIndex);
+	return BoneIndexVertexBufferConverter.GetOrAdd(*this, DesiredName, VertexBuffer, JointsGroupIndex, BoneMap);
 }
 
 FGLTFJsonAccessorIndex FGLTFConvertBuilder::GetOrAddWeightAccessor(const FSkinWeightVertexBuffer* VertexBuffer, int32 WeightsGroupIndex, const FString& DesiredName)
diff --git a/Engine/Plugins/Enterprise/GLTFExporter/MigratedFromGit/Source/GLTFExporter/Private/Builders/GLTFConvertBuilder.h b/Engine/Plugins/Enterprise/GLTFExporter/MigratedFromGit/Source/GLTFExporter/Private/Builders/GLTFConvertBuilder.h
--- a/Engine/Plugins/Enterprise/GLTFExporter/MigratedFromGit/Source/GLTFExporter/Private/Builders/GLTFConvertBuilder.h
+++ b/Engine/Plugins/Enterprise/GLTFExporter/MigratedFromGit/Source/GLTFExporter/Private/Builders/GLTFConvertBuilder.h
@@ -26,7 +26,7 @@ public:
 	FGLTFJsonAccessorIndex GetOrAddNormalAccessor(const FStaticMeshVertexBuffer* VertexBuffer, const FString& DesiredName = TEXT(""));
 	FGLTFJsonAccessorIndex GetOrAddTangentAccessor(const FStaticMeshVertexBuffer* VertexBuffer, const FString& DesiredName = TEXT(""));
 	FGLTFJsonAccessorIndex GetOrAddUVAccessor(const FStaticMeshVertexBuffer* VertexBuffer, int32 UVIndex, const FString& DesiredName = TEXT(""));
-	FGLTFJsonAccessorIndex GetOrAddJointAccessor(const FSkinWeightVertexBuffer* VertexBuffer, int32 JointsGroupIndex, const FString& DesiredName = TEXT(""));
+	FGLTFJsonAccessorIndex GetOrAddJointAccessor(const FSkinWeightVertexBuffer* VertexBuffer, int32 JointsGroupIndex, FGLTFBoneMap BoneMap, const FString& DesiredName = TEXT(""));
 	FGLTFJsonAccessorIndex GetOrAddWeightAccessor(const FSkinWeightVertexBuffer* VertexBuffer, int32 WeightsGroupIndex, const FString& DesiredName = TEXT(""));
 
 	FGLTFJsonBufferViewIndex GetOrAddIndexBufferView(const FRawStaticIndexBuffer* IndexBuffer, const FString& DesiredName = TEXT(""));
diff --git a/Engine/Plugins/Enterprise/GLTFExporter/MigratedFromGit/Source/GLTFExporter/Private/Converters/GLTFBoneMap.h b/Engine/Plugins/Enterprise/GLTFExporter/MigratedFromGit/Source/GLTFExporter/Private/Converters/GLTFBoneMap.h
new file mode 100644
--- /dev/null
+++ b/Engine/Plugins/Enterprise/GLTFExporter/MigratedFromGit/Source/GLTFExporter/Private/Converters/GLTFBoneMap.h
@@ -0,0 +1,75 @@
+// Copyright Epic Games, Inc. All Rights Reserved.
+
+#pragma once
+
+#include "CoreMinimal.h"
+
+// The reason behind a custom material array type are two:
+// 1. Forwarding a TArray by copy in FGLTFConverter::GetOrAdd twice will result in a empty array
+// 2. There is no GetTypeHash for generic arrays
+struct FGLTFBoneMap : TArray<FBoneIndexType>
+{
+	FGLTFBoneMap()
+	{
+	}
+
+	FGLTFBoneMap(const TArray& Other)
+		: TArray(Other)
+	{
+	}
+
+	FGLTFBoneMap(const FGLTFBoneMap& Other)
+        : TArray(Other)
+	{
+	}
+
+	template <typename ElementType, typename AllocatorType>
+    explicit FGLTFBoneMap(const TArray<ElementType, AllocatorType>& Other)
+        : TArray(Other)
+	{
+	}
+
+	FGLTFBoneMap(std::initializer_list<FBoneIndexType> InitList)
+		: TArray(InitList)
+	{
+	}
+
+	FGLTFBoneMap& operator=(std::initializer_list<FBoneIndexType> InitList)
+	{
+		static_cast<TArray&>(*this) = InitList;
+		return *this;
+	}
+
+	FGLTFBoneMap& operator=(const TArray& Other)
+	{
+		static_cast<TArray&>(*this) = Other;
+		return *this;
+	}
+
+	template <typename ElementType, typename AllocatorType>
+    FGLTFBoneMap& operator=(const TArray<ElementType, AllocatorType>& Other)
+	{
+		static_cast<TArray&>(*this) = Other;
+		return *this;
+	}
+
+	bool operator==(const TArray& Other) const
+	{
+		return static_cast<const TArray&>(*this) == Other;
+	}
+
+	bool operator!=(const TArray& Other) const
+	{
+		return static_cast<const TArray&>(*this) != Other;
+	}
+
+	friend uint32 GetTypeHash(const TArray& Array)
+	{
+		uint32 Hash = GetTypeHash(Array.Num());
+		for (FBoneIndexType Element : Array)
+		{
+			Hash = HashCombine(Hash, GetTypeHash(Element));
+		}
+		return Hash;
+	}
+};
diff --git a/Engine/Plugins/Enterprise/GLTFExporter/MigratedFromGit/Source/GLTFExporter/Private/Converters/GLTFSkeletalMeshConverters.cpp b/Engine/Plugins/Enterprise/GLTFExporter/MigratedFromGit/Source/GLTFExporter/Private/Converters/GLTFSkeletalMeshConverters.cpp
--- a/Engine/Plugins/Enterprise/GLTFExporter/MigratedFromGit/Source/GLTFExporter/Private/Converters/GLTFSkeletalMeshConverters.cpp
+++ b/Engine/Plugins/Enterprise/GLTFExporter/MigratedFromGit/Source/GLTFExporter/Private/Converters/GLTFSkeletalMeshConverters.cpp
@@ -95,9 +95,11 @@ FGLTFJsonMeshIndex FGLTFSkeletalMeshConverter::Add(FGLTFConvertBuilder& Builder,
 	JsonAttributes.Joints.AddUninitialized(GroupCount);
 	JsonAttributes.Weights.AddUninitialized(GroupCount);
 
+	const FGLTFBoneMap BoneMap = FGLTFBoneMap(MeshLOD.RenderSections[0].BoneMap); // TODO: can we make this assumption?
+
 	for (uint32 GroupIndex = 0; GroupIndex < GroupCount; ++GroupIndex)
 	{
-		JsonAttributes.Joints[GroupIndex] = Builder.GetOrAddJointAccessor(SkinWeightBuffer, GroupIndex, Name + TEXT("_Joints") + FString::FromInt(GroupIndex));
+		JsonAttributes.Joints[GroupIndex] = Builder.GetOrAddJointAccessor(SkinWeightBuffer, GroupIndex, BoneMap, Name + TEXT("_Joints") + FString::FromInt(GroupIndex));
 		JsonAttributes.Weights[GroupIndex] = Builder.GetOrAddWeightAccessor(SkinWeightBuffer, GroupIndex, Name + TEXT("_Weights") + FString::FromInt(GroupIndex));
 	}
 
diff --git a/Engine/Plugins/Enterprise/GLTFExporter/MigratedFromGit/Source/GLTFExporter/Private/Converters/GLTFVertexBufferConverters.cpp b/Engine/Plugins/Enterprise/GLTFExporter/MigratedFromGit/Source/GLTFExporter/Private/Converters/GLTFVertexBufferConverters.cpp
--- a/Engine/Plugins/Enterprise/GLTFExporter/MigratedFromGit/Source/GLTFExporter/Private/Converters/GLTFVertexBufferConverters.cpp
+++ b/Engine/Plugins/Enterprise/GLTFExporter/MigratedFromGit/Source/GLTFExporter/Private/Converters/GLTFVertexBufferConverters.cpp
@@ -160,13 +160,13 @@ FGLTFJsonAccessorIndex FGLTFUVVertexBufferConverter::Add(FGLTFConvertBuilder& Bu
 	return Builder.AddAccessor(JsonAccessor);
 }
 
-FGLTFJsonAccessorIndex FGLTFBoneIndexVertexBufferConverter::Add(FGLTFConvertBuilder& Builder, const FString& Name, const FSkinWeightVertexBuffer* VertexBuffer, int32 InfluenceOffset)
+FGLTFJsonAccessorIndex FGLTFBoneIndexVertexBufferConverter::Add(FGLTFConvertBuilder& Builder, const FString& Name, const FSkinWeightVertexBuffer* VertexBuffer, int32 InfluenceOffset, FGLTFBoneMap BoneMap)
 {
-	return VertexBuffer->Use16BitBoneIndex() ? Add<uint16>(Builder, Name, VertexBuffer, InfluenceOffset) : Add<uint8>(Builder, Name, VertexBuffer, InfluenceOffset);
+	return VertexBuffer->Use16BitBoneIndex() ? Add<uint16>(Builder, Name, VertexBuffer, InfluenceOffset, BoneMap) : Add<uint8>(Builder, Name, VertexBuffer, InfluenceOffset, BoneMap);
 }
 
 template <typename IndexType>
-FGLTFJsonAccessorIndex FGLTFBoneIndexVertexBufferConverter::Add(FGLTFConvertBuilder& Builder, const FString& Name, const FSkinWeightVertexBuffer* VertexBuffer, int32 JointsGroupIndex)
+FGLTFJsonAccessorIndex FGLTFBoneIndexVertexBufferConverter::Add(FGLTFConvertBuilder& Builder, const FString& Name, const FSkinWeightVertexBuffer* VertexBuffer, int32 JointsGroupIndex, FGLTFBoneMap BoneMap)
 {
 	const uint32 VertexCount = VertexBuffer->GetNumVertices();
 	const int32 MaxInfluenceCount = VertexBuffer->GetMaxBoneInfluences();
@@ -177,8 +177,6 @@ FGLTFJsonAccessorIndex FGLTFBoneIndexVertexBufferConverter::Add(FGLTFConvertBuil
 		return FGLTFJsonAccessorIndex(INDEX_NONE);
 	}
 
-	const int32 InfluenceCount = FMath::Min(MaxInfluenceCount - InfluenceOffset, 4);
-
 	struct VertexBoneIndices
 	{
 		IndexType Index[4];
@@ -189,19 +187,15 @@ FGLTFJsonAccessorIndex FGLTFBoneIndexVertexBufferConverter::Add(FGLTFConvertBuil
 
 	for (uint32 VertexIndex = 0; VertexIndex < VertexCount; ++VertexIndex)
 	{
-		for (int32 InfluenceIndex = 0; InfluenceIndex < InfluenceCount; ++InfluenceIndex)
+		for (int32 InfluenceIndex = 0; InfluenceIndex < 4; ++InfluenceIndex)
 		{
 			// TODO: remove hack
-			const uint32 BoneIndex = FGLTFSkinWeightVertexBufferHack(VertexBuffer).GetBoneIndex(VertexIndex, InfluenceOffset + InfluenceIndex);
+			const uint32 UnmappedBoneIndex = FGLTFSkinWeightVertexBufferHack(VertexBuffer).GetBoneIndex(VertexIndex, InfluenceOffset + InfluenceIndex);
+			const FBoneIndexType BoneIndex = BoneMap[UnmappedBoneIndex];
 			check(BoneIndex <= TNumericLimits<IndexType>::Max());
 
 			BoneIndices[VertexIndex].Index[InfluenceIndex] = static_cast<IndexType>(BoneIndex);
 		}
-
-		for (int32 InfluenceIndex = InfluenceCount; InfluenceIndex < 4; ++InfluenceIndex)
-		{
-			BoneIndices[VertexIndex].Index[InfluenceIndex] = 0;
-		}
 	}
 
 	FGLTFJsonAccessor JsonAccessor;
@@ -225,8 +219,6 @@ FGLTFJsonAccessorIndex FGLTFBoneWeightVertexBufferConverter::Add(FGLTFConvertBui
 		return FGLTFJsonAccessorIndex(INDEX_NONE);
 	}
 
-	const int32 InfluenceCount = FMath::Min(MaxInfluenceCount - InfluenceOffset, 4);
-
 	struct VertexBoneWeights
 	{
 		uint8 Weights[4];
@@ -237,17 +229,12 @@ FGLTFJsonAccessorIndex FGLTFBoneWeightVertexBufferConverter::Add(FGLTFConvertBui
 
 	for (uint32 VertexIndex = 0; VertexIndex < VertexCount; ++VertexIndex)
 	{
-		for (int32 InfluenceIndex = 0; InfluenceIndex < InfluenceCount; ++InfluenceIndex)
+		for (int32 InfluenceIndex = 0; InfluenceIndex < 4; ++InfluenceIndex)
 		{
 			// TODO: remove hack
 			const uint8 BoneWeight = FGLTFSkinWeightVertexBufferHack(VertexBuffer).GetBoneWeight(VertexIndex, InfluenceOffset + InfluenceIndex);
 			BoneWeights[VertexIndex].Weights[InfluenceIndex] = BoneWeight;
 		}
-
-		for (int32 InfluenceIndex = InfluenceCount; InfluenceIndex < 4; ++InfluenceIndex)
-		{
-			BoneWeights[VertexIndex].Weights[InfluenceIndex] = 0;
-		}
 	}
 
 	FGLTFJsonAccessor JsonAccessor;
diff --git a/Engine/Plugins/Enterprise/GLTFExporter/MigratedFromGit/Source/GLTFExporter/Private/Converters/GLTFVertexBufferConverters.h b/Engine/Plugins/Enterprise/GLTFExporter/MigratedFromGit/Source/GLTFExporter/Private/Converters/GLTFVertexBufferConverters.h
--- a/Engine/Plugins/Enterprise/GLTFExporter/MigratedFromGit/Source/GLTFExporter/Private/Converters/GLTFVertexBufferConverters.h
+++ b/Engine/Plugins/Enterprise/GLTFExporter/MigratedFromGit/Source/GLTFExporter/Private/Converters/GLTFVertexBufferConverters.h
@@ -4,6 +4,7 @@
 
 #include "Json/GLTFJsonIndex.h"
 #include "Converters/GLTFConverter.h"
+#include "Converters/GLTFBoneMap.h"
 #include "Engine.h"
 
 class FGLTFPositionVertexBufferConverter final : public TGLTFConverter<FGLTFJsonAccessorIndex, const FPositionVertexBuffer*>
@@ -31,12 +32,12 @@ class FGLTFUVVertexBufferConverter final : public TGLTFConverter<FGLTFJsonAccess
 	FGLTFJsonAccessorIndex Add(FGLTFConvertBuilder& Builder, const FString& Name, const FStaticMeshVertexBuffer* VertexBuffer, int32 UVIndex) override;
 };
 
-class FGLTFBoneIndexVertexBufferConverter final : public TGLTFConverter<FGLTFJsonAccessorIndex, const FSkinWeightVertexBuffer*, int32>
+class FGLTFBoneIndexVertexBufferConverter final : public TGLTFConverter<FGLTFJsonAccessorIndex, const FSkinWeightVertexBuffer*, int32, FGLTFBoneMap>
 {
-	FGLTFJsonAccessorIndex Add(FGLTFConvertBuilder& Builder, const FString& Name, const FSkinWeightVertexBuffer* VertexBuffer, int32 InfluenceOffset) override;
+	FGLTFJsonAccessorIndex Add(FGLTFConvertBuilder& Builder, const FString& Name, const FSkinWeightVertexBuffer* VertexBuffer, int32 InfluenceOffset, FGLTFBoneMap BoneMap) override;
 
 	template <typename IndexType>
-	static FGLTFJsonAccessorIndex Add(FGLTFConvertBuilder& Builder, const FString& Name, const FSkinWeightVertexBuffer* VertexBuffer, int32 JointsGroupIndex);
+	static FGLTFJsonAccessorIndex Add(FGLTFConvertBuilder& Builder, const FString& Name, const FSkinWeightVertexBuffer* VertexBuffer, int32 JointsGroupIndex, FGLTFBoneMap BoneMap);
 };
 
 class FGLTFBoneWeightVertexBufferConverter final : public TGLTFConverter<FGLTFJsonAccessorIndex, const FSkinWeightVertexBuffer*, int32>
