[GLTFExporter] feat(skeletal-mesh): add initial support for SkeletalMeshComponent in level export

commit 7e86cdbbe7c0ffdd1ee9bc7ecdff1bd8d5a9fefb
Date: Mon, 1 Jun 2020 23:00:46 +0200

    feat(skeletal-mesh): add initial support for SkeletalMeshComponent in level export
    (cherry picked from commit 6de6891c1bd5298be79630653174203322d16f56​)

#skipci

diff --git a/Engine/Plugins/Enterprise/GLTFExporter/MigratedFromGit/Source/GLTFExporter/Private/Converters/GLTFLevelConverters.cpp b/Engine/Plugins/Enterprise/GLTFExporter/MigratedFromGit/Source/GLTFExporter/Private/Converters/GLTFLevelConverters.cpp
--- a/Engine/Plugins/Enterprise/GLTFExporter/MigratedFromGit/Source/GLTFExporter/Private/Converters/GLTFLevelConverters.cpp
+++ b/Engine/Plugins/Enterprise/GLTFExporter/MigratedFromGit/Source/GLTFExporter/Private/Converters/GLTFLevelConverters.cpp
@@ -46,6 +46,10 @@ FGLTFJsonNodeIndex FGLTFSceneComponentConverter::Add(FGLTFConvertBuilder& Builde
 			Node.LightMap = Builder.GetOrAddLightMap(StaticMeshComponent);
 		}
 	}
+	else if (const USkeletalMeshComponent* SkeletalMeshComponent = Cast<USkeletalMeshComponent>(SceneComponent))
+	{
+		Node.Mesh = Builder.GetOrAddMesh(SkeletalMeshComponent);
+	}
 	else if (const UCameraComponent* CameraComponent = Cast<UCameraComponent>(SceneComponent))
 	{
 		// TODO: conversion of camera direction should be done in separate converter
