[GLTFExporter] fix(material): add support for detecting texcoord for each property based on input expressions

commit cf2b4f7d00cc247a9c17c3e4806bc232f7e02d50
Date: Wed, 22 Jul 2020 10:08:54 +0200

    fix(material): add support for detecting texcoord for each property based on input expressions
    (cherry picked from commit 3751a0dcbaf4f62647501b8a32fc3156888886f3​)

#skipci

diff --git a/Engine/Plugins/Enterprise/GLTFExporter/MigratedFromGit/Source/GLTFExporter/Private/Converters/GLTFMaterialConverters.cpp b/Engine/Plugins/Enterprise/GLTFExporter/MigratedFromGit/Source/GLTFExporter/Private/Converters/GLTFMaterialConverters.cpp
--- a/Engine/Plugins/Enterprise/GLTFExporter/MigratedFromGit/Source/GLTFExporter/Private/Converters/GLTFMaterialConverters.cpp
+++ b/Engine/Plugins/Enterprise/GLTFExporter/MigratedFromGit/Source/GLTFExporter/Private/Converters/GLTFMaterialConverters.cpp
@@ -223,9 +223,7 @@ bool FGLTFMaterialConverter::TryGetBaseColorAndOpacity(FGLTFConvertBuilder& Buil
 	}
 
 	// TODO: make default baking-resolution configurable
-	// TODO: add support for detecting the correct tex-coord for this property based on connected nodes
 	// TODO: add support for calculating the ideal resolution to use for baking based on connected (texture) nodes
-	int32 TexCoord = 0;
 	FIntPoint TextureSize(1024, 1024);
 
 	// TODO: should this be the default wrap-mode?
@@ -252,7 +250,6 @@ bool FGLTFMaterialConverter::TryGetBaseColorAndOpacity(FGLTFConvertBuilder& Buil
 			return false;
 		}
 
-		TexCoord = BaseColorTexCoord;
 		TextureSize =
 		{
 			FMath::Max(BaseColorTexture->GetSizeX(), OpacityTexture->GetSizeX()),
@@ -269,7 +266,6 @@ bool FGLTFMaterialConverter::TryGetBaseColorAndOpacity(FGLTFConvertBuilder& Buil
 	}
 	else if (bHasBaseColorSourceTexture)
 	{
-		TexCoord = BaseColorTexCoord;
 		TextureSize = { BaseColorTexture->GetSizeX(), BaseColorTexture->GetSizeY() };
 		TextureWrapS = FGLTFConverterUtility::ConvertWrap(BaseColorTexture->AddressX);
 		TextureWrapT = FGLTFConverterUtility::ConvertWrap(BaseColorTexture->AddressY);
@@ -278,7 +274,6 @@ bool FGLTFMaterialConverter::TryGetBaseColorAndOpacity(FGLTFConvertBuilder& Buil
 	}
 	else if (bHasOpacitySourceTexture)
 	{
-		TexCoord = OpacityTexCoord;
 		TextureSize = { OpacityTexture->GetSizeX(), OpacityTexture->GetSizeY() };
 		TextureWrapS = FGLTFConverterUtility::ConvertWrap(OpacityTexture->AddressX);
 		TextureWrapT = FGLTFConverterUtility::ConvertWrap(OpacityTexture->AddressY);
@@ -286,8 +281,8 @@ bool FGLTFMaterialConverter::TryGetBaseColorAndOpacity(FGLTFConvertBuilder& Buil
 		TextureMagFilter = FGLTFConverterUtility::ConvertMagFilter(OpacityTexture->Filter, OpacityTexture->LODGroup);
 	}
 
-	const FGLTFPropertyBakeOutput BaseColorBakeOutput = BakeMaterialProperty(BaseColorProperty, Material, &TextureSize);
-	const FGLTFPropertyBakeOutput OpacityBakeOutput = BakeMaterialProperty(OpacityProperty, Material, &TextureSize, true);
+	const FGLTFPropertyBakeOutput BaseColorBakeOutput = BakeMaterialProperty(BaseColorProperty, Material, BaseColorTexCoord, &TextureSize);
+	const FGLTFPropertyBakeOutput OpacityBakeOutput = BakeMaterialProperty(OpacityProperty, Material, OpacityTexCoord, &TextureSize, true);
 	const float BaseColorScale = BaseColorProperty == MP_EmissiveColor ? BaseColorBakeOutput.EmissiveScale : 1;
 
 	// Detect when both baked properties are constants, which means we can avoid exporting a texture
@@ -300,6 +295,25 @@ bool FGLTFMaterialConverter::TryGetBaseColorAndOpacity(FGLTFConvertBuilder& Buil
 		return true;
 	}
 
+	int32 TexCoord;
+	if (BaseColorBakeOutput.bIsConstant)
+	{
+		TexCoord = OpacityTexCoord;
+	}
+	else if (OpacityBakeOutput.bIsConstant)
+	{
+		TexCoord = BaseColorTexCoord;
+	}
+	else if (BaseColorTexCoord == OpacityTexCoord)
+	{
+		TexCoord = BaseColorTexCoord;
+	}
+	else
+	{
+		// TODO: report texture coordinate conflict
+		return false;
+	}
+
 	TextureSize = BaseColorBakeOutput.Size.ComponentMax(OpacityBakeOutput.Size);
 	BaseColorTexture = FGLTFMaterialUtility::CreateTransientTexture(BaseColorBakeOutput);
 	OpacityTexture = FGLTFMaterialUtility::CreateTransientTexture(OpacityBakeOutput);
@@ -365,9 +379,7 @@ bool FGLTFMaterialConverter::TryGetMetallicAndRoughness(FGLTFConvertBuilder& Bui
 	}
 
 	// TODO: make default baking-resolution configurable
-	// TODO: add support for detecting the correct tex-coord for this property based on connected nodes
 	// TODO: add support for calculating the ideal resolution to use for baking based on connected (texture) nodes
-	int32 TexCoord = 0;
 	FIntPoint TextureSize(1024, 1024);
 
 	// TODO: should this be the default wrap-mode?
@@ -394,7 +406,6 @@ bool FGLTFMaterialConverter::TryGetMetallicAndRoughness(FGLTFConvertBuilder& Bui
 			return false;
 		}
 
-		TexCoord = MetallicTexCoord;
 		TextureSize =
 		{
 			FMath::Max(MetallicTexture->GetSizeX(), RoughnessTexture->GetSizeX()),
@@ -411,7 +422,6 @@ bool FGLTFMaterialConverter::TryGetMetallicAndRoughness(FGLTFConvertBuilder& Bui
 	}
 	else if (bHasMetallicSourceTexture)
 	{
-		TexCoord = MetallicTexCoord;
 		TextureSize = { MetallicTexture->GetSizeX(), MetallicTexture->GetSizeY() };
 		TextureWrapS = FGLTFConverterUtility::ConvertWrap(MetallicTexture->AddressX);
 		TextureWrapT = FGLTFConverterUtility::ConvertWrap(MetallicTexture->AddressY);
@@ -420,7 +430,6 @@ bool FGLTFMaterialConverter::TryGetMetallicAndRoughness(FGLTFConvertBuilder& Bui
 	}
 	else if (bHasRoughnessSourceTexture)
 	{
-		TexCoord = RoughnessTexCoord;
 		TextureSize = { RoughnessTexture->GetSizeX(), RoughnessTexture->GetSizeY() };
 		TextureWrapS = FGLTFConverterUtility::ConvertWrap(RoughnessTexture->AddressX);
 		TextureWrapT = FGLTFConverterUtility::ConvertWrap(RoughnessTexture->AddressY);
@@ -428,8 +437,8 @@ bool FGLTFMaterialConverter::TryGetMetallicAndRoughness(FGLTFConvertBuilder& Bui
 		TextureMagFilter = FGLTFConverterUtility::ConvertMagFilter(RoughnessTexture->Filter, RoughnessTexture->LODGroup);
 	}
 
-	const FGLTFPropertyBakeOutput MetallicBakeOutput = BakeMaterialProperty(MetallicProperty, Material, &TextureSize);
-	const FGLTFPropertyBakeOutput RoughnessBakeOutput = BakeMaterialProperty(RoughnessProperty, Material, &TextureSize);
+	const FGLTFPropertyBakeOutput MetallicBakeOutput = BakeMaterialProperty(MetallicProperty, Material, MetallicTexCoord, &TextureSize);
+	const FGLTFPropertyBakeOutput RoughnessBakeOutput = BakeMaterialProperty(RoughnessProperty, Material, RoughnessTexCoord, &TextureSize);
 
 	// Detect when both baked properties are constants, which means we can use factors and avoid exporting a texture
 	if (MetallicBakeOutput.bIsConstant && RoughnessBakeOutput.bIsConstant)
@@ -439,6 +448,25 @@ bool FGLTFMaterialConverter::TryGetMetallicAndRoughness(FGLTFConvertBuilder& Bui
 		return true;
 	}
 
+	int32 TexCoord;
+	if (MetallicBakeOutput.bIsConstant)
+	{
+		TexCoord = RoughnessTexCoord;
+	}
+	else if (MetallicBakeOutput.bIsConstant)
+	{
+		TexCoord = MetallicTexCoord;
+	}
+	else if (MetallicTexCoord == RoughnessTexCoord)
+	{
+		TexCoord = MetallicTexCoord;
+	}
+	else
+	{
+		// TODO: report texture coordinate conflict
+		return false;
+	}
+
 	TextureSize = RoughnessBakeOutput.Size.ComponentMax(MetallicBakeOutput.Size);
 	MetallicTexture = FGLTFMaterialUtility::CreateTransientTexture(MetallicBakeOutput);
 	RoughnessTexture = FGLTFMaterialUtility::CreateTransientTexture(RoughnessBakeOutput);
@@ -506,9 +534,7 @@ bool FGLTFMaterialConverter::TryGetClearCoatRoughness(FGLTFConvertBuilder& Build
 	}
 
 	// TODO: make default baking-resolution configurable
-	// TODO: add support for detecting the correct tex-coord for this property based on connected nodes
 	// TODO: add support for calculating the ideal resolution to use for baking based on connected (texture) nodes
-	int32 TexCoord = 0;
 	FIntPoint TextureSize(1024, 1024);
 
 	// TODO: should this be the default wrap-mode?
@@ -535,7 +561,6 @@ bool FGLTFMaterialConverter::TryGetClearCoatRoughness(FGLTFConvertBuilder& Build
 			return false;
 		}
 
-		TexCoord = IntensityTexCoord;
 		TextureSize =
 		{
 			FMath::Max(IntensityTexture->GetSizeX(), RoughnessTexture->GetSizeX()),
@@ -552,7 +577,6 @@ bool FGLTFMaterialConverter::TryGetClearCoatRoughness(FGLTFConvertBuilder& Build
 	}
 	else if (bHasIntensitySourceTexture)
 	{
-		TexCoord = IntensityTexCoord;
 		TextureSize = { IntensityTexture->GetSizeX(), IntensityTexture->GetSizeY() };
 		TextureWrapS = FGLTFConverterUtility::ConvertWrap(IntensityTexture->AddressX);
 		TextureWrapT = FGLTFConverterUtility::ConvertWrap(IntensityTexture->AddressY);
@@ -561,7 +585,6 @@ bool FGLTFMaterialConverter::TryGetClearCoatRoughness(FGLTFConvertBuilder& Build
 	}
 	else if (bHasRoughnessSourceTexture)
 	{
-		TexCoord = RoughnessTexCoord;
 		TextureSize = { RoughnessTexture->GetSizeX(), RoughnessTexture->GetSizeY() };
 		TextureWrapS = FGLTFConverterUtility::ConvertWrap(RoughnessTexture->AddressX);
 		TextureWrapT = FGLTFConverterUtility::ConvertWrap(RoughnessTexture->AddressY);
@@ -569,8 +592,8 @@ bool FGLTFMaterialConverter::TryGetClearCoatRoughness(FGLTFConvertBuilder& Build
 		TextureMagFilter = FGLTFConverterUtility::ConvertMagFilter(RoughnessTexture->Filter, RoughnessTexture->LODGroup);
 	}
 
-	const FGLTFPropertyBakeOutput IntensityBakeOutput = BakeMaterialProperty(IntensityProperty, Material, &TextureSize);
-	const FGLTFPropertyBakeOutput RoughnessBakeOutput = BakeMaterialProperty(RoughnessProperty, Material, &TextureSize);
+	const FGLTFPropertyBakeOutput IntensityBakeOutput = BakeMaterialProperty(IntensityProperty, Material, IntensityTexCoord, &TextureSize);
+	const FGLTFPropertyBakeOutput RoughnessBakeOutput = BakeMaterialProperty(RoughnessProperty, Material, RoughnessTexCoord, &TextureSize);
 
 	// Detect when both baked properties are constants, which means we can use factors and avoid exporting a texture
 	if (IntensityBakeOutput.bIsConstant && RoughnessBakeOutput.bIsConstant)
@@ -580,6 +603,25 @@ bool FGLTFMaterialConverter::TryGetClearCoatRoughness(FGLTFConvertBuilder& Build
 		return true;
 	}
 
+	int32 TexCoord;
+	if (IntensityBakeOutput.bIsConstant)
+	{
+		TexCoord = RoughnessTexCoord;
+	}
+	else if (IntensityBakeOutput.bIsConstant)
+	{
+		TexCoord = IntensityTexCoord;
+	}
+	else if (IntensityTexCoord == RoughnessTexCoord)
+	{
+		TexCoord = IntensityTexCoord;
+	}
+	else
+	{
+		// TODO: report texture coordinate conflict
+		return false;
+	}
+
 	TextureSize = RoughnessBakeOutput.Size.ComponentMax(IntensityBakeOutput.Size);
 	IntensityTexture = FGLTFMaterialUtility::CreateTransientTexture(IntensityBakeOutput);
 	RoughnessTexture = FGLTFMaterialUtility::CreateTransientTexture(RoughnessBakeOutput);
@@ -626,7 +668,7 @@ bool FGLTFMaterialConverter::TryGetEmissive(FGLTFConvertBuilder& Builder, FGLTFJ
 		return true;
 	}
 
-	const FGLTFPropertyBakeOutput PropertyBakeOutput = BakeMaterialProperty(EmissiveProperty, Material);
+	const FGLTFPropertyBakeOutput PropertyBakeOutput = BakeMaterialProperty(EmissiveProperty, Material, JsonMaterial.EmissiveTexture.TexCoord);
 	const float EmissiveScale = PropertyBakeOutput.EmissiveScale;
 
 	if (PropertyBakeOutput.bIsConstant)
@@ -944,45 +986,50 @@ bool FGLTFMaterialConverter::TryGetSourceTexture(const UTexture2D*& OutTexture,
 
 	if (const UMaterialExpressionTextureSampleParameter2D* TextureParameter = ExactCast<UMaterialExpressionTextureSampleParameter2D>(Expression))
 	{
+		UTexture* ParameterValue = TextureParameter->Texture;
+
 		const UMaterialInstance* MaterialInstance = Cast<UMaterialInstance>(Material);
 		if (MaterialInstance != nullptr)
 		{
-			UTexture* ParameterValue;
-
 			const FHashedMaterialParameterInfo ParameterInfo(TextureParameter->GetParameterName());
 			if (!MaterialInstance->GetTextureParameterValue(ParameterInfo, ParameterValue))
 			{
 				// TODO: how to handle this?
 			}
-
-			OutTexture = Cast<UTexture2D>(ParameterValue);
 		}
 
+		OutTexture = Cast<UTexture2D>(ParameterValue);
+
 		if (OutTexture == nullptr)
 		{
-			OutTexture = Cast<UTexture2D>(TextureParameter->Texture);
+			if (ParameterValue == nullptr)
+			{
+				// TODO: report material as broken
+			}
+			else
+			{
+				// TODO: report incorrect texture type
+			}
+			return false;
 		}
 
-		if (OutTexture == nullptr)
+		if (!FGLTFMaterialUtility::TryGetTextureCoordinateIndex(TextureParameter, OutTexCoord))
 		{
-			// TODO: report material as broken
+			// TODO: report failed to identify texture coordinate index
 			return false;
 		}
 
-		// TODO: add support for texture coordinate input expression
-		OutTexCoord = TextureParameter->ConstCoordinate;
-
 		return true;
 	}
 
-	if (const UMaterialExpressionTextureSample* TextureSample = ExactCast<UMaterialExpressionTextureSample>(Expression))
+	if (const UMaterialExpressionTextureSample* TextureSampler = ExactCast<UMaterialExpressionTextureSample>(Expression))
 	{
 		// TODO: add support for texture object input expression
-		OutTexture = Cast<UTexture2D>(TextureSample->Texture);
+		OutTexture = Cast<UTexture2D>(TextureSampler->Texture);
 
 		if (OutTexture == nullptr)
 		{
-			if (TextureSample->Texture == nullptr)
+			if (TextureSampler->Texture == nullptr)
 			{
 				// TODO: report material as broken
 			}
@@ -993,8 +1040,11 @@ bool FGLTFMaterialConverter::TryGetSourceTexture(const UTexture2D*& OutTexture,
 			return false;
 		}
 
-		// TODO: add support for texture coordinate input expression
-		OutTexCoord = TextureSample->ConstCoordinate;
+		if (!FGLTFMaterialUtility::TryGetTextureCoordinateIndex(TextureSampler, OutTexCoord))
+		{
+			// TODO: report failed to identify texture coordinate index
+			return false;
+		}
 
 		return true;
 	}
@@ -1004,7 +1054,7 @@ bool FGLTFMaterialConverter::TryGetSourceTexture(const UTexture2D*& OutTexture,
 
 bool FGLTFMaterialConverter::TryGetBakedMaterialProperty(FGLTFConvertBuilder& Builder, FGLTFJsonTextureInfo& OutTexInfo, FGLTFJsonColor3& OutConstant, EMaterialProperty Property, const FString& PropertyName, const UMaterialInterface* Material) const
 {
-	const FGLTFPropertyBakeOutput PropertyBakeOutput = BakeMaterialProperty(Property, Material);
+	const FGLTFPropertyBakeOutput PropertyBakeOutput = BakeMaterialProperty(Property, Material, OutTexInfo.TexCoord);
 
 	if (PropertyBakeOutput.bIsConstant)
 	{
@@ -1023,7 +1073,7 @@ bool FGLTFMaterialConverter::TryGetBakedMaterialProperty(FGLTFConvertBuilder& Bu
 
 bool FGLTFMaterialConverter::TryGetBakedMaterialProperty(FGLTFConvertBuilder& Builder, FGLTFJsonTextureInfo& OutTexInfo, FGLTFJsonColor4& OutConstant, EMaterialProperty Property, const FString& PropertyName, const UMaterialInterface* Material) const
 {
-	const FGLTFPropertyBakeOutput PropertyBakeOutput = BakeMaterialProperty(Property, Material);
+	const FGLTFPropertyBakeOutput PropertyBakeOutput = BakeMaterialProperty(Property, Material, OutTexInfo.TexCoord);
 
 	if (PropertyBakeOutput.bIsConstant)
 	{
@@ -1042,7 +1092,7 @@ bool FGLTFMaterialConverter::TryGetBakedMaterialProperty(FGLTFConvertBuilder& Bu
 
 inline bool FGLTFMaterialConverter::TryGetBakedMaterialProperty(FGLTFConvertBuilder& Builder, FGLTFJsonTextureInfo& OutTexInfo, float& OutConstant, EMaterialProperty Property, const FString& PropertyName, const UMaterialInterface* Material) const
 {
-	const FGLTFPropertyBakeOutput PropertyBakeOutput = BakeMaterialProperty(Property, Material);
+	const FGLTFPropertyBakeOutput PropertyBakeOutput = BakeMaterialProperty(Property, Material, OutTexInfo.TexCoord);
 
 	if (PropertyBakeOutput.bIsConstant)
 	{
@@ -1061,7 +1111,7 @@ inline bool FGLTFMaterialConverter::TryGetBakedMaterialProperty(FGLTFConvertBuil
 
 bool FGLTFMaterialConverter::TryGetBakedMaterialProperty(FGLTFConvertBuilder& Builder, FGLTFJsonTextureInfo& OutTexInfo, EMaterialProperty Property, const FString& PropertyName, const UMaterialInterface* Material) const
 {
-	const FGLTFPropertyBakeOutput PropertyBakeOutput = BakeMaterialProperty(Property, Material);
+	const FGLTFPropertyBakeOutput PropertyBakeOutput = BakeMaterialProperty(Property, Material, OutTexInfo.TexCoord);
 
 	if (!PropertyBakeOutput.bIsConstant)
 	{
@@ -1106,8 +1156,31 @@ bool FGLTFMaterialConverter::TryGetBakedMaterialProperty(FGLTFConvertBuilder& Bu
 	return true;
 }
 
-FGLTFPropertyBakeOutput FGLTFMaterialConverter::BakeMaterialProperty(EMaterialProperty Property, const UMaterialInterface* Material, const FIntPoint* PreferredTextureSize, bool bCopyAlphaFromRedChannel) const
+FGLTFPropertyBakeOutput FGLTFMaterialConverter::BakeMaterialProperty(EMaterialProperty Property, const UMaterialInterface* Material, int32& OutTexCoord, const FIntPoint* PreferredTextureSize, bool bCopyAlphaFromRedChannel) const
 {
+	const FExpressionInput* PropertyInput = FGLTFMaterialUtility::GetInputForProperty(Material, Property);
+	TSet<int32> TexCoords;
+
+	FGLTFMaterialUtility::GetAllTextureCoordinateIndices(*PropertyInput, TexCoords);
+	if (TexCoords.Num() > 0)
+	{
+		// TODO: is there a nicer way to get the first element in a TSet?
+		for (int32 TexCoord : TexCoords)
+		{
+			OutTexCoord = TexCoord;
+			break;
+		}
+
+		if (TexCoords.Num() > 1)
+		{
+			// TODO: report multiple texture coordinates found, will use first
+		}
+	}
+	else
+	{
+		OutTexCoord = 0; // property is texture coordinate independent
+	}
+
 	// TODO: make default baking-resolution configurable
 	const FIntPoint DefaultTextureSize(1024, 1024);
 	const FIntPoint TextureSize = PreferredTextureSize != nullptr ? *PreferredTextureSize : DefaultTextureSize;
@@ -1118,8 +1191,14 @@ FGLTFPropertyBakeOutput FGLTFMaterialConverter::BakeMaterialProperty(EMaterialPr
 		TextureSize,
 		Property,
 		Material,
+		OutTexCoord,
 		bCopyAlphaFromRedChannel);
 
+	if (!PropertyBakeOutput.bIsConstant && TexCoords.Num() == 0)
+	{
+		// TODO: report warning about property not being constant yet texture coordinate independent
+	}
+
 	return PropertyBakeOutput;
 }
 
diff --git a/Engine/Plugins/Enterprise/GLTFExporter/MigratedFromGit/Source/GLTFExporter/Private/Converters/GLTFMaterialConverters.h b/Engine/Plugins/Enterprise/GLTFExporter/MigratedFromGit/Source/GLTFExporter/Private/Converters/GLTFMaterialConverters.h
--- a/Engine/Plugins/Enterprise/GLTFExporter/MigratedFromGit/Source/GLTFExporter/Private/Converters/GLTFMaterialConverters.h
+++ b/Engine/Plugins/Enterprise/GLTFExporter/MigratedFromGit/Source/GLTFExporter/Private/Converters/GLTFMaterialConverters.h
@@ -34,7 +34,7 @@ class FGLTFMaterialConverter final : public TGLTFConverter<FGLTFJsonMaterialInde
 	bool TryGetBakedMaterialProperty(FGLTFConvertBuilder& Builder, FGLTFJsonTextureInfo& OutTexInfo, float& OutConstant, EMaterialProperty Property, const FString& PropertyName, const UMaterialInterface* Material) const;
 	bool TryGetBakedMaterialProperty(FGLTFConvertBuilder& Builder, FGLTFJsonTextureInfo& OutTexInfo, EMaterialProperty Property, const FString& PropertyName, const UMaterialInterface* Material) const;
 
-	FGLTFPropertyBakeOutput BakeMaterialProperty(EMaterialProperty Property, const UMaterialInterface* Material, const FIntPoint* PreferredTextureSize = nullptr, bool bCopyAlphaFromRedChannel = false) const;
+	FGLTFPropertyBakeOutput BakeMaterialProperty(EMaterialProperty Property, const UMaterialInterface* Material, int32& OutTexCoord, const FIntPoint* PreferredTextureSize = nullptr, bool bCopyAlphaFromRedChannel = false) const;
 
 	bool StoreBakedPropertyTexture(FGLTFConvertBuilder& Builder, FGLTFJsonTextureInfo& OutTexInfo, const FGLTFPropertyBakeOutput& PropertyBakeOutput, const FString& PropertyName, const UMaterialInterface* Material) const;
 };
diff --git a/Engine/Plugins/Enterprise/GLTFExporter/MigratedFromGit/Source/GLTFExporter/Private/Converters/GLTFMaterialUtility.cpp b/Engine/Plugins/Enterprise/GLTFExporter/MigratedFromGit/Source/GLTFExporter/Private/Converters/GLTFMaterialUtility.cpp
--- a/Engine/Plugins/Enterprise/GLTFExporter/MigratedFromGit/Source/GLTFExporter/Private/Converters/GLTFMaterialUtility.cpp
+++ b/Engine/Plugins/Enterprise/GLTFExporter/MigratedFromGit/Source/GLTFExporter/Private/Converters/GLTFMaterialUtility.cpp
@@ -10,6 +10,7 @@
 #include "MaterialBakingStructures.h"
 #include "Materials/MaterialExpressionCustomOutput.h"
 #include "Materials/MaterialExpressionClearCoatNormalCustomOutput.h"
+#include "Materials/MaterialExpressionTextureCoordinate.h"
 
 const TCHAR* FGLTFMaterialUtility::GetPropertyName(EMaterialProperty Property)
 {
@@ -221,7 +222,7 @@ bool FGLTFMaterialUtility::CombineTextures(TArray<FColor>& OutPixels, const TArr
 	return bReadSuccessful;
 }
 
-FGLTFPropertyBakeOutput FGLTFMaterialUtility::BakeMaterialProperty(const FIntPoint& OutputSize, EMaterialProperty Property, const UMaterialInterface* Material, bool bCopyAlphaFromRedChannel)
+FGLTFPropertyBakeOutput FGLTFMaterialUtility::BakeMaterialProperty(const FIntPoint& OutputSize, EMaterialProperty Property, const UMaterialInterface* Material, int32 TexCoord, bool bCopyAlphaFromRedChannel)
 {
 	EMaterialProperty HackProperty = MP_MAX;
 	TTuple<int32, UMaterialExpression*> HackPrevState;
@@ -254,6 +255,7 @@ FGLTFPropertyBakeOutput FGLTFMaterialUtility::BakeMaterialProperty(const FIntPoi
 
 	FMeshData MeshSet;
 	MeshSet.TextureCoordinateBox = { { 0.0f, 0.0f }, { 1.0f, 1.0f } };
+	MeshSet.TextureCoordinateIndex = TexCoord;
 
 	// TODO: Do we need to fill in any more info in MeshSet?
 
@@ -365,3 +367,48 @@ uint32 FGLTFMaterialUtility::GetMaskComponentCount(const FExpressionInput& Expre
 {
 	return ExpressionInput.MaskR + ExpressionInput.MaskG + ExpressionInput.MaskB + ExpressionInput.MaskA;
 }
+
+bool FGLTFMaterialUtility::TryGetTextureCoordinateIndex(const UMaterialExpressionTextureSample* TextureSampler, int32& TexCoord)
+{
+	const UMaterialExpression* Expression = TextureSampler->Coordinates.Expression;
+	if (Expression == nullptr)
+	{
+		TexCoord = TextureSampler->ConstCoordinate;
+		return true;
+	}
+
+	if (const UMaterialExpressionTextureCoordinate* TextureCoordinate = Cast<UMaterialExpressionTextureCoordinate>(Expression))
+	{
+		TexCoord = TextureCoordinate->CoordinateIndex;
+		return true;
+	}
+
+	return false;
+}
+
+void FGLTFMaterialUtility::GetAllTextureCoordinateIndices(const FExpressionInput& ExpressionInput, TSet<int32> TexCoordIndices)
+{
+	UMaterialExpression* Expression = ExpressionInput.Expression;
+	if (Expression == nullptr)
+	{
+		return;
+	}
+
+	TArray<UMaterialExpression*> InputExpressions;
+	Expression->GetAllInputExpressions(InputExpressions);
+
+	for (const UMaterialExpression* InputExpression : InputExpressions)
+	{
+		if (const UMaterialExpressionTextureSample* TextureSampler = Cast<UMaterialExpressionTextureSample>(InputExpression))
+		{
+			if (TextureSampler->Coordinates.Expression == nullptr)
+			{
+				TexCoordIndices.Add(TextureSampler->ConstCoordinate);
+			}
+		}
+		else if (const UMaterialExpressionTextureCoordinate* TextureCoordinate = Cast<UMaterialExpressionTextureCoordinate>(InputExpression))
+		{
+			TexCoordIndices.Add(TextureCoordinate->CoordinateIndex);
+		}
+	}
+}
diff --git a/Engine/Plugins/Enterprise/GLTFExporter/MigratedFromGit/Source/GLTFExporter/Private/Converters/GLTFMaterialUtility.h b/Engine/Plugins/Enterprise/GLTFExporter/MigratedFromGit/Source/GLTFExporter/Private/Converters/GLTFMaterialUtility.h
--- a/Engine/Plugins/Enterprise/GLTFExporter/MigratedFromGit/Source/GLTFExporter/Private/Converters/GLTFMaterialUtility.h
+++ b/Engine/Plugins/Enterprise/GLTFExporter/MigratedFromGit/Source/GLTFExporter/Private/Converters/GLTFMaterialUtility.h
@@ -51,11 +51,14 @@ struct FGLTFMaterialUtility
 
 	static UTexture2D* CreateTransientTexture(const FGLTFPropertyBakeOutput& PropertyBakeOutput, bool bUseSRGB = false);
 	static bool CombineTextures(TArray<FColor>& OutPixels, const TArray<FGLTFTextureCombineSource>& Sources, const FIntPoint& OutputSize, EPixelFormat OutputPixelFormat);
-	static FGLTFPropertyBakeOutput BakeMaterialProperty(const FIntPoint& OutputSize, EMaterialProperty Property, const UMaterialInterface* Material, bool bCopyAlphaFromRedChannel = false);
+	static FGLTFPropertyBakeOutput BakeMaterialProperty(const FIntPoint& OutputSize, EMaterialProperty Property, const UMaterialInterface* Material, int32 TexCoord, bool bCopyAlphaFromRedChannel = false);
 
 	static FGLTFJsonTextureIndex AddCombinedTexture(FGLTFConvertBuilder& Builder, const TArray<FGLTFTextureCombineSource>& CombineSources, const FIntPoint& TextureSize, const FString& TextureName, EGLTFJsonTextureFilter MinFilter, EGLTFJsonTextureFilter MagFilter, EGLTFJsonTextureWrap WrapS, EGLTFJsonTextureWrap WrapT);
 	static FGLTFJsonTextureIndex AddTexture(FGLTFConvertBuilder& Builder, const TArray<FColor>& Pixels, const FIntPoint& TextureSize, const FString& TextureName, EGLTFJsonTextureFilter MinFilter, EGLTFJsonTextureFilter MagFilter, EGLTFJsonTextureWrap WrapS, EGLTFJsonTextureWrap WrapT);
 
 	static FLinearColor GetMask(const FExpressionInput& ExpressionInput);
 	static uint32 GetMaskComponentCount(const FExpressionInput& ExpressionInput);
+
+	static bool TryGetTextureCoordinateIndex(const UMaterialExpressionTextureSample* TextureSampler, int32& TexCoord);
+	static void GetAllTextureCoordinateIndices(const FExpressionInput& ExpressionInput, TSet<int32> TexCoordIndices);
 };
