[GLTFExporter] fix(skeletal-mesh): disable export of vertex skin data until crash bug fixed

commit cf4adf67e98e6e0cbd182f141679932339e51ce2
Date: Wed, 22 Jul 2020 11:10:40 +0200

    fix(skeletal-mesh): disable export of vertex skin data until crash bug fixed
    (cherry picked from commit e8e4ea89a311569dedc396542c947077bdb60f81​)

#skipci

diff --git a/Engine/Plugins/Enterprise/GLTFExporter/MigratedFromGit/Source/GLTFExporter/Private/Converters/GLTFSkeletalMeshConverters.cpp b/Engine/Plugins/Enterprise/GLTFExporter/MigratedFromGit/Source/GLTFExporter/Private/Converters/GLTFSkeletalMeshConverters.cpp
--- a/Engine/Plugins/Enterprise/GLTFExporter/MigratedFromGit/Source/GLTFExporter/Private/Converters/GLTFSkeletalMeshConverters.cpp
+++ b/Engine/Plugins/Enterprise/GLTFExporter/MigratedFromGit/Source/GLTFExporter/Private/Converters/GLTFSkeletalMeshConverters.cpp
@@ -91,6 +91,8 @@ FGLTFJsonMeshIndex FGLTFSkeletalMeshConverter::Add(FGLTFConvertBuilder& Builder,
 		JsonAttributes.TexCoords[UVIndex] = Builder.GetOrAddUVAccessor(VertexBuffer, UVIndex, Name + TEXT("_UV") + FString::FromInt(UVIndex) + TEXT("s"));
 	}
 
+	/* TODO: enable export of vertex skin data when crash bug fixed
+
 	const uint32 GroupCount = (SkinWeightBuffer->GetMaxBoneInfluences() + 3) / 4;
 	JsonAttributes.Joints.AddUninitialized(GroupCount);
 	JsonAttributes.Weights.AddUninitialized(GroupCount);
@@ -102,6 +104,7 @@ FGLTFJsonMeshIndex FGLTFSkeletalMeshConverter::Add(FGLTFConvertBuilder& Builder,
 		JsonAttributes.Joints[GroupIndex] = Builder.GetOrAddJointAccessor(SkinWeightBuffer, GroupIndex, BoneMap, Name + TEXT("_Joints") + FString::FromInt(GroupIndex));
 		JsonAttributes.Weights[GroupIndex] = Builder.GetOrAddWeightAccessor(SkinWeightBuffer, GroupIndex, Name + TEXT("_Weights") + FString::FromInt(GroupIndex));
 	}
+	*/
 
 	const FMultiSizeIndexContainer* IndexContainer = &MeshLOD.MultiSizeIndexContainer;
 	Builder.GetOrAddIndexBufferView(IndexContainer, Name + TEXT("_Indices"));
