[GLTFExporter] refactor(options): update export options to reflect current hardcoded settings

commit 5fb4f7137c837444837ca1a1fa0fe5c49a7ad196
Date: Wed, 22 Jul 2020 11:17:09 +0200

    refactor(options): update export options to reflect current hardcoded settings
    (cherry picked from commit e9c783fa0ab53c0273fdc08220bc91ea97e6de31​)

#skipci

diff --git a/Engine/Plugins/Enterprise/GLTFExporter/MigratedFromGit/Source/GLTFExporter/Private/GLTFExportOptions.cpp b/Engine/Plugins/Enterprise/GLTFExporter/MigratedFromGit/Source/GLTFExporter/Private/GLTFExportOptions.cpp
--- a/Engine/Plugins/Enterprise/GLTFExporter/MigratedFromGit/Source/GLTFExporter/Private/GLTFExportOptions.cpp
+++ b/Engine/Plugins/Enterprise/GLTFExporter/MigratedFromGit/Source/GLTFExporter/Private/GLTFExportOptions.cpp
@@ -16,7 +16,7 @@ UGLTFExportOptions::UGLTFExportOptions(const FObjectInitializer& ObjectInitializ
 	: Super(ObjectInitializer)
 {
 	bBundleWebViewer = false;
-	bExportPreviewMesh = false;
+	bExportPreviewMesh = true;
 	bAllExtensionsRequired = false;
 	bExportVertexColors = true;
 	NormalExportMethod = EGLTFExporterNormalExportMethod::NormalsAndTangents;
@@ -24,20 +24,20 @@ UGLTFExportOptions::UGLTFExportOptions(const FObjectInitializer& ObjectInitializ
 	bNormalizeUVCoordinates = EGLTFExporterNormalizeUVCoordinates::Never;
 	DefaultLevelOfDetail = 0;
 	bMapSkeletalMotionToRoot = false;
-	bExportUnlitMaterials = false;
-	bExportClearCoatMaterials = false;
+	bExportUnlitMaterials = true;
+	bExportClearCoatMaterials = true;
 	bBakeMaterialInputs = true;
 	DefaultMaterialBakeSize = EGLTFExporterMaterialBakeSize::POT_1024;
-	TextureFormat = EGLTFExporterTextureFormat::None;
-	bExportLightmaps = false;
-	TextureHDREncoding = EGLTFExporterTextureHDREncoding::None;
+	TextureFormat = EGLTFExporterTextureFormat::PNG;
+	bExportLightmaps = true;
+	TextureHDREncoding = EGLTFExporterTextureHDREncoding::RGBM;
 	ExportScale = 0.01;
-	bExportLights = false;
-	bExportCameras = false;
+	bExportLights = true;
+	bExportCameras = true;
 	bExportOrbitalCameras = false;
 	bExportReflectionCaptures = false;
-	bExportHDRIBackdrops = false;
-	bExportVariantSets = false;
+	bExportHDRIBackdrops = true;
+	bExportVariantSets = true;
 	bExportInteractionHotspots = false;
 }
 
