[GLTFExporter] feat(options): make default baking texture size configurable

commit 90419b8e84b9f87099db992d27ac3d0f7ab468a3
Date: Wed, 22 Jul 2020 11:34:30 +0200

    feat(options): make default baking texture size configurable
    (cherry picked from commit 3851cbd3bf833a8b91105b85aaada4a6eabcf2c5​)

#skipci

diff --git a/Engine/Plugins/Enterprise/GLTFExporter/MigratedFromGit/Source/GLTFExporter/Private/Converters/GLTFMaterialConverters.cpp b/Engine/Plugins/Enterprise/GLTFExporter/MigratedFromGit/Source/GLTFExporter/Private/Converters/GLTFMaterialConverters.cpp
--- a/Engine/Plugins/Enterprise/GLTFExporter/MigratedFromGit/Source/GLTFExporter/Private/Converters/GLTFMaterialConverters.cpp
+++ b/Engine/Plugins/Enterprise/GLTFExporter/MigratedFromGit/Source/GLTFExporter/Private/Converters/GLTFMaterialConverters.cpp
@@ -222,9 +222,8 @@ bool FGLTFMaterialConverter::TryGetBaseColorAndOpacity(FGLTFConvertBuilder& Buil
 		return true;
 	}
 
-	// TODO: make default baking-resolution configurable
 	// TODO: add support for calculating the ideal resolution to use for baking based on connected (texture) nodes
-	FIntPoint TextureSize(1024, 1024);
+	FIntPoint TextureSize = Builder.ExportOptions->GetDefaultMaterialBakeSize();
 
 	// TODO: should this be the default wrap-mode?
 	EGLTFJsonTextureWrap TextureWrapS = EGLTFJsonTextureWrap::Repeat;
@@ -281,8 +280,8 @@ bool FGLTFMaterialConverter::TryGetBaseColorAndOpacity(FGLTFConvertBuilder& Buil
 		TextureMagFilter = FGLTFConverterUtility::ConvertMagFilter(OpacityTexture->Filter, OpacityTexture->LODGroup);
 	}
 
-	const FGLTFPropertyBakeOutput BaseColorBakeOutput = BakeMaterialProperty(BaseColorProperty, Material, BaseColorTexCoord, &TextureSize);
-	const FGLTFPropertyBakeOutput OpacityBakeOutput = BakeMaterialProperty(OpacityProperty, Material, OpacityTexCoord, &TextureSize, true);
+	const FGLTFPropertyBakeOutput BaseColorBakeOutput = BakeMaterialProperty(Builder, BaseColorProperty, Material, BaseColorTexCoord, &TextureSize);
+	const FGLTFPropertyBakeOutput OpacityBakeOutput = BakeMaterialProperty(Builder, OpacityProperty, Material, OpacityTexCoord, &TextureSize, true);
 	const float BaseColorScale = BaseColorProperty == MP_EmissiveColor ? BaseColorBakeOutput.EmissiveScale : 1;
 
 	// Detect when both baked properties are constants, which means we can avoid exporting a texture
@@ -378,9 +377,8 @@ bool FGLTFMaterialConverter::TryGetMetallicAndRoughness(FGLTFConvertBuilder& Bui
 		return true;
 	}
 
-	// TODO: make default baking-resolution configurable
 	// TODO: add support for calculating the ideal resolution to use for baking based on connected (texture) nodes
-	FIntPoint TextureSize(1024, 1024);
+	FIntPoint TextureSize = Builder.ExportOptions->GetDefaultMaterialBakeSize();
 
 	// TODO: should this be the default wrap-mode?
 	EGLTFJsonTextureWrap TextureWrapS = EGLTFJsonTextureWrap::Repeat;
@@ -437,8 +435,8 @@ bool FGLTFMaterialConverter::TryGetMetallicAndRoughness(FGLTFConvertBuilder& Bui
 		TextureMagFilter = FGLTFConverterUtility::ConvertMagFilter(RoughnessTexture->Filter, RoughnessTexture->LODGroup);
 	}
 
-	const FGLTFPropertyBakeOutput MetallicBakeOutput = BakeMaterialProperty(MetallicProperty, Material, MetallicTexCoord, &TextureSize);
-	const FGLTFPropertyBakeOutput RoughnessBakeOutput = BakeMaterialProperty(RoughnessProperty, Material, RoughnessTexCoord, &TextureSize);
+	const FGLTFPropertyBakeOutput MetallicBakeOutput = BakeMaterialProperty(Builder, MetallicProperty, Material, MetallicTexCoord, &TextureSize);
+	const FGLTFPropertyBakeOutput RoughnessBakeOutput = BakeMaterialProperty(Builder, RoughnessProperty, Material, RoughnessTexCoord, &TextureSize);
 
 	// Detect when both baked properties are constants, which means we can use factors and avoid exporting a texture
 	if (MetallicBakeOutput.bIsConstant && RoughnessBakeOutput.bIsConstant)
@@ -533,9 +531,8 @@ bool FGLTFMaterialConverter::TryGetClearCoatRoughness(FGLTFConvertBuilder& Build
 		return true;
 	}
 
-	// TODO: make default baking-resolution configurable
 	// TODO: add support for calculating the ideal resolution to use for baking based on connected (texture) nodes
-	FIntPoint TextureSize(1024, 1024);
+	FIntPoint TextureSize = Builder.ExportOptions->GetDefaultMaterialBakeSize();
 
 	// TODO: should this be the default wrap-mode?
 	EGLTFJsonTextureWrap TextureWrapS = EGLTFJsonTextureWrap::Repeat;
@@ -592,8 +589,8 @@ bool FGLTFMaterialConverter::TryGetClearCoatRoughness(FGLTFConvertBuilder& Build
 		TextureMagFilter = FGLTFConverterUtility::ConvertMagFilter(RoughnessTexture->Filter, RoughnessTexture->LODGroup);
 	}
 
-	const FGLTFPropertyBakeOutput IntensityBakeOutput = BakeMaterialProperty(IntensityProperty, Material, IntensityTexCoord, &TextureSize);
-	const FGLTFPropertyBakeOutput RoughnessBakeOutput = BakeMaterialProperty(RoughnessProperty, Material, RoughnessTexCoord, &TextureSize);
+	const FGLTFPropertyBakeOutput IntensityBakeOutput = BakeMaterialProperty(Builder, IntensityProperty, Material, IntensityTexCoord, &TextureSize);
+	const FGLTFPropertyBakeOutput RoughnessBakeOutput = BakeMaterialProperty(Builder, RoughnessProperty, Material, RoughnessTexCoord, &TextureSize);
 
 	// Detect when both baked properties are constants, which means we can use factors and avoid exporting a texture
 	if (IntensityBakeOutput.bIsConstant && RoughnessBakeOutput.bIsConstant)
@@ -668,7 +665,7 @@ bool FGLTFMaterialConverter::TryGetEmissive(FGLTFConvertBuilder& Builder, FGLTFJ
 		return true;
 	}
 
-	const FGLTFPropertyBakeOutput PropertyBakeOutput = BakeMaterialProperty(EmissiveProperty, Material, JsonMaterial.EmissiveTexture.TexCoord);
+	const FGLTFPropertyBakeOutput PropertyBakeOutput = BakeMaterialProperty(Builder, EmissiveProperty, Material, JsonMaterial.EmissiveTexture.TexCoord);
 	const float EmissiveScale = PropertyBakeOutput.EmissiveScale;
 
 	if (PropertyBakeOutput.bIsConstant)
@@ -1054,7 +1051,7 @@ bool FGLTFMaterialConverter::TryGetSourceTexture(const UTexture2D*& OutTexture,
 
 bool FGLTFMaterialConverter::TryGetBakedMaterialProperty(FGLTFConvertBuilder& Builder, FGLTFJsonTextureInfo& OutTexInfo, FGLTFJsonColor3& OutConstant, EMaterialProperty Property, const FString& PropertyName, const UMaterialInterface* Material) const
 {
-	const FGLTFPropertyBakeOutput PropertyBakeOutput = BakeMaterialProperty(Property, Material, OutTexInfo.TexCoord);
+	const FGLTFPropertyBakeOutput PropertyBakeOutput = BakeMaterialProperty(Builder, Property, Material, OutTexInfo.TexCoord);
 
 	if (PropertyBakeOutput.bIsConstant)
 	{
@@ -1073,7 +1070,7 @@ bool FGLTFMaterialConverter::TryGetBakedMaterialProperty(FGLTFConvertBuilder& Bu
 
 bool FGLTFMaterialConverter::TryGetBakedMaterialProperty(FGLTFConvertBuilder& Builder, FGLTFJsonTextureInfo& OutTexInfo, FGLTFJsonColor4& OutConstant, EMaterialProperty Property, const FString& PropertyName, const UMaterialInterface* Material) const
 {
-	const FGLTFPropertyBakeOutput PropertyBakeOutput = BakeMaterialProperty(Property, Material, OutTexInfo.TexCoord);
+	const FGLTFPropertyBakeOutput PropertyBakeOutput = BakeMaterialProperty(Builder, Property, Material, OutTexInfo.TexCoord);
 
 	if (PropertyBakeOutput.bIsConstant)
 	{
@@ -1092,7 +1089,7 @@ bool FGLTFMaterialConverter::TryGetBakedMaterialProperty(FGLTFConvertBuilder& Bu
 
 inline bool FGLTFMaterialConverter::TryGetBakedMaterialProperty(FGLTFConvertBuilder& Builder, FGLTFJsonTextureInfo& OutTexInfo, float& OutConstant, EMaterialProperty Property, const FString& PropertyName, const UMaterialInterface* Material) const
 {
-	const FGLTFPropertyBakeOutput PropertyBakeOutput = BakeMaterialProperty(Property, Material, OutTexInfo.TexCoord);
+	const FGLTFPropertyBakeOutput PropertyBakeOutput = BakeMaterialProperty(Builder, Property, Material, OutTexInfo.TexCoord);
 
 	if (PropertyBakeOutput.bIsConstant)
 	{
@@ -1111,7 +1108,7 @@ inline bool FGLTFMaterialConverter::TryGetBakedMaterialProperty(FGLTFConvertBuil
 
 bool FGLTFMaterialConverter::TryGetBakedMaterialProperty(FGLTFConvertBuilder& Builder, FGLTFJsonTextureInfo& OutTexInfo, EMaterialProperty Property, const FString& PropertyName, const UMaterialInterface* Material) const
 {
-	const FGLTFPropertyBakeOutput PropertyBakeOutput = BakeMaterialProperty(Property, Material, OutTexInfo.TexCoord);
+	const FGLTFPropertyBakeOutput PropertyBakeOutput = BakeMaterialProperty(Builder, Property, Material, OutTexInfo.TexCoord);
 
 	if (!PropertyBakeOutput.bIsConstant)
 	{
@@ -1156,7 +1153,7 @@ bool FGLTFMaterialConverter::TryGetBakedMaterialProperty(FGLTFConvertBuilder& Bu
 	return true;
 }
 
-FGLTFPropertyBakeOutput FGLTFMaterialConverter::BakeMaterialProperty(EMaterialProperty Property, const UMaterialInterface* Material, int32& OutTexCoord, const FIntPoint* PreferredTextureSize, bool bCopyAlphaFromRedChannel) const
+FGLTFPropertyBakeOutput FGLTFMaterialConverter::BakeMaterialProperty(FGLTFConvertBuilder& Builder, EMaterialProperty Property, const UMaterialInterface* Material, int32& OutTexCoord, const FIntPoint* PreferredTextureSize, bool bCopyAlphaFromRedChannel) const
 {
 	const FExpressionInput* PropertyInput = FGLTFMaterialUtility::GetInputForProperty(Material, Property);
 	TSet<int32> TexCoords;
@@ -1181,8 +1178,7 @@ FGLTFPropertyBakeOutput FGLTFMaterialConverter::BakeMaterialProperty(EMaterialPr
 		OutTexCoord = 0; // property is texture coordinate independent
 	}
 
-	// TODO: make default baking-resolution configurable
-	const FIntPoint DefaultTextureSize(1024, 1024);
+	const FIntPoint DefaultTextureSize = Builder.ExportOptions->GetDefaultMaterialBakeSize();
 	const FIntPoint TextureSize = PreferredTextureSize != nullptr ? *PreferredTextureSize : DefaultTextureSize;
 
 	// TODO: add support for calculating the ideal resolution to use for baking based on connected (texture) nodes
diff --git a/Engine/Plugins/Enterprise/GLTFExporter/MigratedFromGit/Source/GLTFExporter/Private/Converters/GLTFMaterialConverters.h b/Engine/Plugins/Enterprise/GLTFExporter/MigratedFromGit/Source/GLTFExporter/Private/Converters/GLTFMaterialConverters.h
--- a/Engine/Plugins/Enterprise/GLTFExporter/MigratedFromGit/Source/GLTFExporter/Private/Converters/GLTFMaterialConverters.h
+++ b/Engine/Plugins/Enterprise/GLTFExporter/MigratedFromGit/Source/GLTFExporter/Private/Converters/GLTFMaterialConverters.h
@@ -34,7 +34,7 @@ class FGLTFMaterialConverter final : public TGLTFConverter<FGLTFJsonMaterialInde
 	bool TryGetBakedMaterialProperty(FGLTFConvertBuilder& Builder, FGLTFJsonTextureInfo& OutTexInfo, float& OutConstant, EMaterialProperty Property, const FString& PropertyName, const UMaterialInterface* Material) const;
 	bool TryGetBakedMaterialProperty(FGLTFConvertBuilder& Builder, FGLTFJsonTextureInfo& OutTexInfo, EMaterialProperty Property, const FString& PropertyName, const UMaterialInterface* Material) const;
 
-	FGLTFPropertyBakeOutput BakeMaterialProperty(EMaterialProperty Property, const UMaterialInterface* Material, int32& OutTexCoord, const FIntPoint* PreferredTextureSize = nullptr, bool bCopyAlphaFromRedChannel = false) const;
+	FGLTFPropertyBakeOutput BakeMaterialProperty(FGLTFConvertBuilder& Builder, EMaterialProperty Property, const UMaterialInterface* Material, int32& OutTexCoord, const FIntPoint* PreferredTextureSize = nullptr, bool bCopyAlphaFromRedChannel = false) const;
 
 	bool StoreBakedPropertyTexture(FGLTFConvertBuilder& Builder, FGLTFJsonTextureInfo& OutTexInfo, const FGLTFPropertyBakeOutput& PropertyBakeOutput, const FString& PropertyName, const UMaterialInterface* Material) const;
 };
diff --git a/Engine/Plugins/Enterprise/GLTFExporter/MigratedFromGit/Source/GLTFExporter/Private/GLTFExportOptions.cpp b/Engine/Plugins/Enterprise/GLTFExporter/MigratedFromGit/Source/GLTFExporter/Private/GLTFExportOptions.cpp
--- a/Engine/Plugins/Enterprise/GLTFExporter/MigratedFromGit/Source/GLTFExporter/Private/GLTFExportOptions.cpp
+++ b/Engine/Plugins/Enterprise/GLTFExporter/MigratedFromGit/Source/GLTFExporter/Private/GLTFExportOptions.cpp
@@ -6,7 +6,6 @@
 #include "UObject/UObjectGlobals.h"
 #include "UObject/Object.h"
 #include "UObject/UObjectBaseUtility.h"
-#include "UObject/UObjectHash.h"
 #include "UObject/Class.h"
 #include "UObject/UnrealType.h"
 #include "Misc/ConfigCacheIni.h"
@@ -15,19 +14,16 @@
 UGLTFExportOptions::UGLTFExportOptions(const FObjectInitializer& ObjectInitializer)
 	: Super(ObjectInitializer)
 {
-	bBundleWebViewer = false;
-	bExportPreviewMesh = true;
-	bAllExtensionsRequired = false;
+	bExportUnlitMaterials = true;
+	bExportClearCoatMaterials = true;
+	bBakeMaterialInputs = true;
+	DefaultMaterialBakeSize = EGLTFExporterMaterialBakeSize::POT_1024;
 	bExportVertexColors = true;
 	NormalExportMethod = EGLTFExporterNormalExportMethod::NormalsAndTangents;
 	bTangentDataQuantization = false;
 	bNormalizeUVCoordinates = EGLTFExporterNormalizeUVCoordinates::Never;
 	DefaultLevelOfDetail = 0;
 	bMapSkeletalMotionToRoot = false;
-	bExportUnlitMaterials = true;
-	bExportClearCoatMaterials = true;
-	bBakeMaterialInputs = true;
-	DefaultMaterialBakeSize = EGLTFExporterMaterialBakeSize::POT_1024;
 	TextureFormat = EGLTFExporterTextureFormat::PNG;
 	bExportLightmaps = true;
 	TextureHDREncoding = EGLTFExporterTextureHDREncoding::RGBM;
@@ -39,6 +35,9 @@ UGLTFExportOptions::UGLTFExportOptions(const FObjectInitializer& ObjectInitializ
 	bExportHDRIBackdrops = true;
 	bExportVariantSets = true;
 	bExportInteractionHotspots = false;
+	bBundleWebViewer = false;
+	bExportPreviewMesh = true;
+	bAllExtensionsRequired = false;
 }
 
 void UGLTFExportOptions::ResetToDefault()
@@ -209,5 +208,16 @@ void UGLTFExportOptions::FillOptions(bool bBatchMode, bool bShowOptionDialog, co
 
 bool UGLTFExportOptions::CanEditChange(const FProperty* InProperty) const
 {
+	if (InProperty->GetName() == TEXT("DefaultMaterialBakeSize"))
+	{
+		return true;
+	}
+
 	return false; // TODO: change to true for each option implemented and support by the object to be exported.
 }
+
+FIntPoint UGLTFExportOptions::GetDefaultMaterialBakeSize() const
+{
+	const int32 Size = 1 << static_cast<int>(DefaultMaterialBakeSize);
+	return { Size, Size };
+}
diff --git a/Engine/Plugins/Enterprise/GLTFExporter/MigratedFromGit/Source/GLTFExporter/Public/GLTFExportOptions.h b/Engine/Plugins/Enterprise/GLTFExporter/MigratedFromGit/Source/GLTFExporter/Public/GLTFExportOptions.h
--- a/Engine/Plugins/Enterprise/GLTFExporter/MigratedFromGit/Source/GLTFExporter/Public/GLTFExportOptions.h
+++ b/Engine/Plugins/Enterprise/GLTFExporter/MigratedFromGit/Source/GLTFExporter/Public/GLTFExportOptions.h
@@ -65,14 +65,17 @@ class GLTFEXPORTER_API UGLTFExportOptions : public UObject
 
 public:
 
-	UPROPERTY(EditAnywhere, BlueprintReadWrite, config, Category = Exporter)
-	uint32 bBundleWebViewer : 1;
+	UPROPERTY(EditAnywhere, BlueprintReadWrite, config, Category = Material)
+	uint32 bExportUnlitMaterials : 1;
 
-	UPROPERTY(EditAnywhere, BlueprintReadWrite, config, Category = Exporter)
-	uint32 bExportPreviewMesh : 1;
+	UPROPERTY(EditAnywhere, BlueprintReadWrite, config, Category = Material)
+	uint32 bExportClearCoatMaterials : 1;
 
-	UPROPERTY(EditAnywhere, BlueprintReadWrite, config, Category = Exporter)
-	uint32 bAllExtensionsRequired : 1;
+	UPROPERTY(EditAnywhere, BlueprintReadWrite, config, Category = Material)
+	uint32 bBakeMaterialInputs : 1;
+
+	UPROPERTY(EditAnywhere, BlueprintReadWrite, config, Category = Material)
+	EGLTFExporterMaterialBakeSize DefaultMaterialBakeSize;
 
 	UPROPERTY(EditAnywhere, BlueprintReadWrite, config, Category = Mesh)
 	uint32 bExportVertexColors : 1;
@@ -92,18 +95,6 @@ public:
 	UPROPERTY(EditAnywhere, BlueprintReadWrite, config, Category = Animation)
 	uint32 bMapSkeletalMotionToRoot : 1;
 
-	UPROPERTY(EditAnywhere, BlueprintReadWrite, config, Category = Material)
-	uint32 bExportUnlitMaterials : 1;
-
-	UPROPERTY(EditAnywhere, BlueprintReadWrite, config, Category = Material)
-	uint32 bExportClearCoatMaterials : 1;
-
-	UPROPERTY(EditAnywhere, BlueprintReadWrite, config, Category = Material)
-	uint32 bBakeMaterialInputs : 1;
-
-	UPROPERTY(EditAnywhere, BlueprintReadWrite, config, Category = Material)
-	EGLTFExporterMaterialBakeSize DefaultMaterialBakeSize;
-
 	UPROPERTY(EditAnywhere, BlueprintReadWrite, config, Category = Texture)
 	EGLTFExporterTextureFormat TextureFormat;
 
@@ -139,6 +130,15 @@ public:
 	UPROPERTY(EditAnywhere, BlueprintReadWrite, config, Category = Scene)
 	uint32 bExportInteractionHotspots : 1;
 
+	UPROPERTY(EditAnywhere, BlueprintReadWrite, config, Category = Exporter)
+    uint32 bBundleWebViewer : 1;
+
+	UPROPERTY(EditAnywhere, BlueprintReadWrite, config, Category = Exporter)
+    uint32 bExportPreviewMesh : 1;
+
+	UPROPERTY(EditAnywhere, BlueprintReadWrite, config, Category = Exporter)
+    uint32 bAllExtensionsRequired : 1;
+
 	/* Set all the each property to the CDO value */
 	void ResetToDefault();
 
@@ -159,4 +159,6 @@ public:
 	void FillOptions(bool bBatchMode, bool bShowOptionDialog, const FString& FullPath, bool& bOutOperationCanceled, bool& bOutExportAll);
 
 	virtual bool CanEditChange(const FProperty* InProperty) const override;
+
+	FIntPoint GetDefaultMaterialBakeSize() const;
 };
