[GLTFExporter] fix(material): fix issue with detecting texcoord for each property based on input expressions

commit b669ee8c1ef5dedd1c5125d41e2e67f072c4ae53
Date: Wed, 22 Jul 2020 16:20:18 +0200

    fix(material): fix issue with detecting texcoord for each property based on input expressions
    (cherry picked from commit c2305007c30ff80f6e2916e1334a2dabefe69c0b​)

#skipci

diff --git a/Engine/Plugins/Enterprise/GLTFExporter/MigratedFromGit/Source/GLTFExporter/Private/Converters/GLTFMaterialConverters.cpp b/Engine/Plugins/Enterprise/GLTFExporter/MigratedFromGit/Source/GLTFExporter/Private/Converters/GLTFMaterialConverters.cpp
--- a/Engine/Plugins/Enterprise/GLTFExporter/MigratedFromGit/Source/GLTFExporter/Private/Converters/GLTFMaterialConverters.cpp
+++ b/Engine/Plugins/Enterprise/GLTFExporter/MigratedFromGit/Source/GLTFExporter/Private/Converters/GLTFMaterialConverters.cpp
@@ -1148,7 +1148,6 @@ bool FGLTFMaterialConverter::TryGetBakedMaterialProperty(FGLTFConvertBuilder& Bu
 		EGLTFJsonTextureWrap::ClampToEdge,
         EGLTFJsonTextureWrap::ClampToEdge);
 
-	OutTexInfo.TexCoord = 0;
 	OutTexInfo.Index = TextureIndex;
 	return true;
 }
@@ -1175,7 +1174,7 @@ FGLTFPropertyBakeOutput FGLTFMaterialConverter::BakeMaterialProperty(FGLTFConver
 	}
 	else
 	{
-		OutTexCoord = 0; // property is texture coordinate independent
+		OutTexCoord = 0; // assume TexCoord0 even thought property seems to be texture coordinate independent
 	}
 
 	const FIntPoint DefaultTextureSize = Builder.ExportOptions->GetDefaultMaterialBakeSize();
@@ -1202,9 +1201,6 @@ bool FGLTFMaterialConverter::StoreBakedPropertyTexture(FGLTFConvertBuilder& Buil
 {
 	const FString TextureName = Material->GetName() + TEXT("_") + PropertyName;
 
-	// TODO: add support for detecting the correct tex-coord for this property based on connected nodes
-	const uint32 TexCoord = 0;
-
 	// TODO: should this be the default wrap-mode?
 	const EGLTFJsonTextureWrap TextureWrapS = EGLTFJsonTextureWrap::Repeat;
 	const EGLTFJsonTextureWrap TextureWrapT = EGLTFJsonTextureWrap::Repeat;
@@ -1223,8 +1219,6 @@ bool FGLTFMaterialConverter::StoreBakedPropertyTexture(FGLTFConvertBuilder& Buil
 		TextureWrapS,
 		TextureWrapT);
 
-	OutTexInfo.TexCoord = TexCoord;
 	OutTexInfo.Index = TextureIndex;
-
 	return true;
 }
diff --git a/Engine/Plugins/Enterprise/GLTFExporter/MigratedFromGit/Source/GLTFExporter/Private/Converters/GLTFMaterialUtility.cpp b/Engine/Plugins/Enterprise/GLTFExporter/MigratedFromGit/Source/GLTFExporter/Private/Converters/GLTFMaterialUtility.cpp
--- a/Engine/Plugins/Enterprise/GLTFExporter/MigratedFromGit/Source/GLTFExporter/Private/Converters/GLTFMaterialUtility.cpp
+++ b/Engine/Plugins/Enterprise/GLTFExporter/MigratedFromGit/Source/GLTFExporter/Private/Converters/GLTFMaterialUtility.cpp
@@ -386,7 +386,7 @@ bool FGLTFMaterialUtility::TryGetTextureCoordinateIndex(const UMaterialExpressio
 	return false;
 }
 
-void FGLTFMaterialUtility::GetAllTextureCoordinateIndices(const FExpressionInput& ExpressionInput, TSet<int32> TexCoordIndices)
+void FGLTFMaterialUtility::GetAllTextureCoordinateIndices(const FExpressionInput& ExpressionInput, TSet<int32>& OutTexCoords)
 {
 	UMaterialExpression* Expression = ExpressionInput.Expression;
 	if (Expression == nullptr)
@@ -403,12 +403,12 @@ void FGLTFMaterialUtility::GetAllTextureCoordinateIndices(const FExpressionInput
 		{
 			if (TextureSampler->Coordinates.Expression == nullptr)
 			{
-				TexCoordIndices.Add(TextureSampler->ConstCoordinate);
+				OutTexCoords.Add(TextureSampler->ConstCoordinate);
 			}
 		}
 		else if (const UMaterialExpressionTextureCoordinate* TextureCoordinate = Cast<UMaterialExpressionTextureCoordinate>(InputExpression))
 		{
-			TexCoordIndices.Add(TextureCoordinate->CoordinateIndex);
+			OutTexCoords.Add(TextureCoordinate->CoordinateIndex);
 		}
 	}
 }
diff --git a/Engine/Plugins/Enterprise/GLTFExporter/MigratedFromGit/Source/GLTFExporter/Private/Converters/GLTFMaterialUtility.h b/Engine/Plugins/Enterprise/GLTFExporter/MigratedFromGit/Source/GLTFExporter/Private/Converters/GLTFMaterialUtility.h
--- a/Engine/Plugins/Enterprise/GLTFExporter/MigratedFromGit/Source/GLTFExporter/Private/Converters/GLTFMaterialUtility.h
+++ b/Engine/Plugins/Enterprise/GLTFExporter/MigratedFromGit/Source/GLTFExporter/Private/Converters/GLTFMaterialUtility.h
@@ -60,5 +60,5 @@ struct FGLTFMaterialUtility
 	static uint32 GetMaskComponentCount(const FExpressionInput& ExpressionInput);
 
 	static bool TryGetTextureCoordinateIndex(const UMaterialExpressionTextureSample* TextureSampler, int32& TexCoord);
-	static void GetAllTextureCoordinateIndices(const FExpressionInput& ExpressionInput, TSet<int32> TexCoordIndices);
+	static void GetAllTextureCoordinateIndices(const FExpressionInput& ExpressionInput, TSet<int32>& OutTexCoords);
 };
