[GLTFExporter] fix(material): assume texcoord1 if multiple are detected in property AmbientOcclusion

commit f3fa7d0a51831e5d6581566b665725112ebc262c
Date: Wed, 22 Jul 2020 15:57:50 +0200

    fix(material): assume texcoord1 if multiple are detected in property AmbientOcclusion
    (cherry picked from commit 72fb1b3f283e23bab65d4206ec84ba2f17a1c00e​)

#skipci

diff --git a/Engine/Plugins/Enterprise/GLTFExporter/MigratedFromGit/Source/GLTFExporter/Private/Converters/GLTFMaterialConverters.cpp b/Engine/Plugins/Enterprise/GLTFExporter/MigratedFromGit/Source/GLTFExporter/Private/Converters/GLTFMaterialConverters.cpp
--- a/Engine/Plugins/Enterprise/GLTFExporter/MigratedFromGit/Source/GLTFExporter/Private/Converters/GLTFMaterialConverters.cpp
+++ b/Engine/Plugins/Enterprise/GLTFExporter/MigratedFromGit/Source/GLTFExporter/Private/Converters/GLTFMaterialConverters.cpp
@@ -1170,6 +1170,12 @@ FGLTFPropertyBakeOutput FGLTFMaterialConverter::BakeMaterialProperty(FGLTFConver
 		if (TexCoords.Num() > 1)
 		{
 			// TODO: report multiple texture coordinates found, will use first
+
+			// TODO: replace this hardcoded hack with something more configurable and proper
+			if (Property == MP_AmbientOcclusion && TexCoords.Contains(1))
+			{
+				OutTexCoord = 1; // assume ambient occlusion uses TexCoord1 when multiple
+			}
 		}
 	}
 	else
