[GLTFExporter] refactor(level): name lights and actors after their actor instead of component since their component name is often the same

commit 9faa5ca1d410156b7e14f92c1dbd0d37f8126a5a
Date: Wed, 22 Jul 2020 16:34:10 +0200

    refactor(level): name lights and actors after their actor instead of component since their component name is often the same
    (cherry picked from commit a104c91a1468b5074af73ef24d87550f8b2ae68c​)

#skipci

diff --git a/Engine/Plugins/Enterprise/GLTFExporter/MigratedFromGit/Source/GLTFExporter/Private/Converters/GLTFLevelConverters.cpp b/Engine/Plugins/Enterprise/GLTFExporter/MigratedFromGit/Source/GLTFExporter/Private/Converters/GLTFLevelConverters.cpp
--- a/Engine/Plugins/Enterprise/GLTFExporter/MigratedFromGit/Source/GLTFExporter/Private/Converters/GLTFLevelConverters.cpp
+++ b/Engine/Plugins/Enterprise/GLTFExporter/MigratedFromGit/Source/GLTFExporter/Private/Converters/GLTFLevelConverters.cpp
@@ -54,18 +54,18 @@ FGLTFJsonNodeIndex FGLTFSceneComponentConverter::Add(FGLTFConvertBuilder& Builde
 	{
 		// TODO: conversion of camera direction should be done in separate converter
 		FGLTFJsonNode CameraNode;
-		CameraNode.Name = Node.Name + TEXT("_CameraNode");
+		CameraNode.Name = Owner->GetName(); // TODO: choose a more unique name if owner is not ACameraActor
 		CameraNode.Rotation = FGLTFConverterUtility::ConvertCameraDirection();
-		CameraNode.Camera = Builder.GetOrAddCamera(CameraComponent);
+		CameraNode.Camera = Builder.GetOrAddCamera(CameraComponent, CameraNode.Name);
 		Builder.AddChildNode(NodeIndex, CameraNode);
 	}
 	else if (const ULightComponent* LightComponent = Cast<ULightComponent>(SceneComponent))
 	{
 		// TODO: conversion of light direction should be done in separate converter
 		FGLTFJsonNode LightNode;
-		LightNode.Name = Node.Name + TEXT("_LightNode");
+		LightNode.Name = Owner->GetName(); // TODO: choose a more unique name if owner is not ALight
 		LightNode.Rotation = FGLTFConverterUtility::ConvertLightDirection();
-		LightNode.Light = Builder.GetOrAddLight(LightComponent);
+		LightNode.Light = Builder.GetOrAddLight(LightComponent, LightNode.Name);
 		Builder.AddChildNode(NodeIndex, LightNode);
 	}
 
