[GLTFExporter] feat(camera): add support for non-aspect-ratio-constrained cameras

commit d9601dd23631ded9102e893e8ced523808441c0e
Date: Wed, 22 Jul 2020 17:56:43 +0200

    feat(camera): add support for non-aspect-ratio-constrained cameras
    (cherry picked from commit b9929350574d898b9a988890503d767bc14dcb57​)

#skipci

diff --git a/Engine/Plugins/Enterprise/GLTFExporter/MigratedFromGit/Source/GLTFExporter/Private/Converters/GLTFCameraUtility.cpp b/Engine/Plugins/Enterprise/GLTFExporter/MigratedFromGit/Source/GLTFExporter/Private/Converters/GLTFCameraUtility.cpp
new file mode 100644
--- /dev/null
+++ b/Engine/Plugins/Enterprise/GLTFExporter/MigratedFromGit/Source/GLTFExporter/Private/Converters/GLTFCameraUtility.cpp
@@ -0,0 +1,45 @@
+// Copyright Epic Games, Inc. All Rights Reserved.
+
+#include "Converters/GLTFCameraUtility.h"
+#include "Converters/GLTFConverterUtility.h"
+#include "Camera/CameraTypes.h"
+
+FGLTFJsonOrthographic FGLTFCameraUtility::ConvertOrthographic(const FMinimalViewInfo& View)
+{
+	FGLTFJsonOrthographic Orthographic;
+
+	// NOTE: it goes against the glTF standard to not define a XMag, but the gltf viewer can handle it
+	if (View.bConstrainAspectRatio)
+	{
+		Orthographic.XMag = FGLTFConverterUtility::ConvertLength(View.OrthoWidth);
+	}
+
+	Orthographic.YMag = FGLTFConverterUtility::ConvertLength(View.OrthoWidth / View.AspectRatio); // TODO: is this correct?
+	Orthographic.ZFar = FGLTFConverterUtility::ConvertLength(View.OrthoFarClipPlane);
+	Orthographic.ZNear = FGLTFConverterUtility::ConvertLength(View.OrthoNearClipPlane);
+	return Orthographic;
+}
+
+FGLTFJsonPerspective FGLTFCameraUtility::ConvertPerspective(const FMinimalViewInfo& View)
+{
+	FGLTFJsonPerspective Perspective;
+
+	if (View.bConstrainAspectRatio)
+	{
+		Perspective.AspectRatio = View.AspectRatio;
+	}
+
+	Perspective.YFov = ConvertFieldOfView(View);
+
+	// NOTE: even thought ZFar is optional, if we don't set it, then most gltf viewers won't handle it well.
+	Perspective.ZFar = FGLTFConverterUtility::ConvertLength(WORLD_MAX); // TODO: Unreal doesn't have max draw distance per view?
+	Perspective.ZNear = FGLTFConverterUtility::ConvertLength(GNearClippingPlane);
+	return Perspective;
+}
+
+float FGLTFCameraUtility::ConvertFieldOfView(const FMinimalViewInfo& View)
+{
+	const float HorizontalFOV = FMath::DegreesToRadians(View.FOV);
+	const float VerticalFOV = 2 * FMath::Atan(FMath::Tan(HorizontalFOV / 2) / View.AspectRatio);
+	return VerticalFOV;
+}
diff --git a/Engine/Plugins/Enterprise/GLTFExporter/MigratedFromGit/Source/GLTFExporter/Private/Converters/GLTFCameraUtility.h b/Engine/Plugins/Enterprise/GLTFExporter/MigratedFromGit/Source/GLTFExporter/Private/Converters/GLTFCameraUtility.h
new file mode 100644
--- /dev/null
+++ b/Engine/Plugins/Enterprise/GLTFExporter/MigratedFromGit/Source/GLTFExporter/Private/Converters/GLTFCameraUtility.h
@@ -0,0 +1,13 @@
+// Copyright Epic Games, Inc. All Rights Reserved.
+
+#pragma once
+
+#include "Json/GLTFJsonCamera.h"
+
+struct FGLTFCameraUtility
+{
+	static FGLTFJsonOrthographic ConvertOrthographic(const FMinimalViewInfo& View);
+	static FGLTFJsonPerspective ConvertPerspective(const FMinimalViewInfo& View);
+
+	static float ConvertFieldOfView(const FMinimalViewInfo& View);
+};
diff --git a/Engine/Plugins/Enterprise/GLTFExporter/MigratedFromGit/Source/GLTFExporter/Private/Converters/GLTFConverterUtility.cpp b/Engine/Plugins/Enterprise/GLTFExporter/MigratedFromGit/Source/GLTFExporter/Private/Converters/GLTFConverterUtility.cpp
--- a/Engine/Plugins/Enterprise/GLTFExporter/MigratedFromGit/Source/GLTFExporter/Private/Converters/GLTFConverterUtility.cpp
+++ b/Engine/Plugins/Enterprise/GLTFExporter/MigratedFromGit/Source/GLTFExporter/Private/Converters/GLTFConverterUtility.cpp
@@ -102,30 +102,3 @@ EGLTFJsonCubeFace FGLTFConverterUtility::ConvertCubeFace(ECubeFace CubeFace)
 		default:            return EGLTFJsonCubeFace::None;
 	}
 }
-
-FGLTFJsonOrthographic FGLTFConverterUtility::ConvertOrthographic(const FMinimalViewInfo& View)
-{
-	FGLTFJsonOrthographic Orthographic;
-	Orthographic.XMag = ConvertLength(View.OrthoWidth);
-	Orthographic.YMag = ConvertLength(View.OrthoWidth / View.AspectRatio); // TODO: is this correct?
-	Orthographic.ZFar = ConvertLength(View.OrthoFarClipPlane);
-	Orthographic.ZNear = ConvertLength(View.OrthoNearClipPlane);
-	return Orthographic;
-}
-
-FGLTFJsonPerspective FGLTFConverterUtility::ConvertPerspective(const FMinimalViewInfo& View)
-{
-	FGLTFJsonPerspective Perspective;
-	Perspective.AspectRatio = View.AspectRatio;
-	Perspective.YFov = ConvertFieldOfView(View);
-	Perspective.ZFar = ConvertLength(WORLD_MAX); // TODO: make sure the assumption that Unreal cameras have no draw distance is correct.
-	Perspective.ZNear = ConvertLength(GNearClippingPlane);
-	return Perspective;
-}
-
-float FGLTFConverterUtility::ConvertFieldOfView(const FMinimalViewInfo& View)
-{
-	const float HorizontalFOV = FMath::DegreesToRadians(View.FOV);
-	const float VerticalFOV = 2 * FMath::Atan(FMath::Tan(HorizontalFOV / 2) / View.AspectRatio);
-	return VerticalFOV;
-}
diff --git a/Engine/Plugins/Enterprise/GLTFExporter/MigratedFromGit/Source/GLTFExporter/Private/Converters/GLTFConverterUtility.h b/Engine/Plugins/Enterprise/GLTFExporter/MigratedFromGit/Source/GLTFExporter/Private/Converters/GLTFConverterUtility.h
--- a/Engine/Plugins/Enterprise/GLTFExporter/MigratedFromGit/Source/GLTFExporter/Private/Converters/GLTFConverterUtility.h
+++ b/Engine/Plugins/Enterprise/GLTFExporter/MigratedFromGit/Source/GLTFExporter/Private/Converters/GLTFConverterUtility.h
@@ -137,11 +137,6 @@ struct FGLTFConverterUtility
 
 	static EGLTFJsonCubeFace ConvertCubeFace(ECubeFace CubeFace);
 
-	static FGLTFJsonOrthographic ConvertOrthographic(const FMinimalViewInfo& View);
-	static FGLTFJsonPerspective ConvertPerspective(const FMinimalViewInfo& View);
-
-	static float ConvertFieldOfView(const FMinimalViewInfo& View);
-
 	template <typename ComponentType>
 	static EGLTFJsonComponentType GetComponentType()
 	{
diff --git a/Engine/Plugins/Enterprise/GLTFExporter/MigratedFromGit/Source/GLTFExporter/Private/Converters/GLTFLevelConverters.cpp b/Engine/Plugins/Enterprise/GLTFExporter/MigratedFromGit/Source/GLTFExporter/Private/Converters/GLTFLevelConverters.cpp
--- a/Engine/Plugins/Enterprise/GLTFExporter/MigratedFromGit/Source/GLTFExporter/Private/Converters/GLTFLevelConverters.cpp
+++ b/Engine/Plugins/Enterprise/GLTFExporter/MigratedFromGit/Source/GLTFExporter/Private/Converters/GLTFLevelConverters.cpp
@@ -4,7 +4,7 @@
 #include "Builders/GLTFContainerBuilder.h"
 #include "Converters/GLTFConverterUtility.h"
 #include "Converters/GLTFActorUtility.h"
-#include "Engine/MapBuildDataRegistry.h"
+#include "Converters/GLTFCameraUtility.h"
 
 FGLTFJsonNodeIndex FGLTFSceneComponentConverter::Add(FGLTFConvertBuilder& Builder, const FString& Name, const USceneComponent* SceneComponent)
 {
@@ -150,11 +150,11 @@ FGLTFJsonCameraIndex FGLTFCameraComponentConverter::Add(FGLTFConvertBuilder& Bui
 	switch (Camera.Type)
 	{
 		case EGLTFJsonCameraType::Orthographic:
-			Camera.Orthographic = FGLTFConverterUtility::ConvertOrthographic(DesiredView);
+			Camera.Orthographic = FGLTFCameraUtility::ConvertOrthographic(DesiredView);
 			break;
 
 		case EGLTFJsonCameraType::Perspective:
-			Camera.Perspective = FGLTFConverterUtility::ConvertPerspective(DesiredView);
+			Camera.Perspective = FGLTFCameraUtility::ConvertPerspective(DesiredView);
 			break;
 
 		default:
diff --git a/Engine/Plugins/Enterprise/GLTFExporter/MigratedFromGit/Source/GLTFExporter/Private/Json/GLTFJsonCamera.h b/Engine/Plugins/Enterprise/GLTFExporter/MigratedFromGit/Source/GLTFExporter/Private/Json/GLTFJsonCamera.h
--- a/Engine/Plugins/Enterprise/GLTFExporter/MigratedFromGit/Source/GLTFExporter/Private/Json/GLTFJsonCamera.h
+++ b/Engine/Plugins/Enterprise/GLTFExporter/MigratedFromGit/Source/GLTFExporter/Private/Json/GLTFJsonCamera.h
@@ -26,7 +26,11 @@ struct FGLTFJsonOrthographic
 	{
 		JsonWriter.WriteObjectStart();
 
-		FGLTFJsonUtility::WriteExactValue(JsonWriter, TEXT("xmag"), XMag);
+		if (XMag != 0)
+		{
+			FGLTFJsonUtility::WriteExactValue(JsonWriter, TEXT("xmag"), XMag);
+		}
+
 		FGLTFJsonUtility::WriteExactValue(JsonWriter, TEXT("ymag"), YMag);
 		FGLTFJsonUtility::WriteExactValue(JsonWriter, TEXT("zfar"), ZFar);
 		FGLTFJsonUtility::WriteExactValue(JsonWriter, TEXT("znear"), ZNear);
