[GLTFExporter] fix(light): don't export lights if static (needs to be made configurable later)

commit eb0168cd6957f98ca10fd5be172218fc13039085
Date: Wed, 22 Jul 2020 22:39:58 +0200

    fix(light): don't export lights if static (needs to be made configurable later)
    (cherry picked from commit db495c89b3d9c59b64aa980f26f7b27b2a04d5ee​)

#skipci

diff --git a/Engine/Plugins/Enterprise/GLTFExporter/MigratedFromGit/Source/GLTFExporter/Private/Converters/GLTFLevelConverters.cpp b/Engine/Plugins/Enterprise/GLTFExporter/MigratedFromGit/Source/GLTFExporter/Private/Converters/GLTFLevelConverters.cpp
--- a/Engine/Plugins/Enterprise/GLTFExporter/MigratedFromGit/Source/GLTFExporter/Private/Converters/GLTFLevelConverters.cpp
+++ b/Engine/Plugins/Enterprise/GLTFExporter/MigratedFromGit/Source/GLTFExporter/Private/Converters/GLTFLevelConverters.cpp
@@ -61,12 +61,16 @@ FGLTFJsonNodeIndex FGLTFSceneComponentConverter::Add(FGLTFConvertBuilder& Builde
 	}
 	else if (const ULightComponent* LightComponent = Cast<ULightComponent>(SceneComponent))
 	{
-		// TODO: conversion of light direction should be done in separate converter
-		FGLTFJsonNode LightNode;
-		LightNode.Name = Owner->GetName(); // TODO: choose a more unique name if owner is not ALight
-		LightNode.Rotation = FGLTFConverterUtility::ConvertLightDirection();
-		LightNode.Light = Builder.GetOrAddLight(LightComponent, LightNode.Name);
-		Builder.AddChildNode(NodeIndex, LightNode);
+		// TODO: make it configurable whether static lights should be exported or not
+		if (LightComponent->Mobility != EComponentMobility::Static)
+		{
+			// TODO: conversion of light direction should be done in separate converter
+			FGLTFJsonNode LightNode;
+			LightNode.Name = Owner->GetName(); // TODO: choose a more unique name if owner is not ALight
+			LightNode.Rotation = FGLTFConverterUtility::ConvertLightDirection();
+			LightNode.Light = Builder.GetOrAddLight(LightComponent, LightNode.Name);
+			Builder.AddChildNode(NodeIndex, LightNode);
+		}
 	}
 
 	return NodeIndex;
diff --git a/Engine/Plugins/Enterprise/GLTFExporter/MigratedFromGit/Source/GLTFExporter/Private/GLTFExportOptions.cpp b/Engine/Plugins/Enterprise/GLTFExporter/MigratedFromGit/Source/GLTFExporter/Private/GLTFExportOptions.cpp
--- a/Engine/Plugins/Enterprise/GLTFExporter/MigratedFromGit/Source/GLTFExporter/Private/GLTFExportOptions.cpp
+++ b/Engine/Plugins/Enterprise/GLTFExporter/MigratedFromGit/Source/GLTFExporter/Private/GLTFExportOptions.cpp
@@ -28,7 +28,7 @@ UGLTFExportOptions::UGLTFExportOptions(const FObjectInitializer& ObjectInitializ
 	bExportLightmaps = true;
 	TextureHDREncoding = EGLTFExporterTextureHDREncoding::RGBM;
 	ExportScale = 0.01;
-	bExportLights = true;
+	bExportLights = EGLTFExporterLightMobility::MovableAndStationary;
 	bExportCameras = true;
 	bExportOrbitalCameras = false;
 	bExportReflectionCaptures = false;
diff --git a/Engine/Plugins/Enterprise/GLTFExporter/MigratedFromGit/Source/GLTFExporter/Public/GLTFExportOptions.h b/Engine/Plugins/Enterprise/GLTFExporter/MigratedFromGit/Source/GLTFExporter/Public/GLTFExportOptions.h
--- a/Engine/Plugins/Enterprise/GLTFExporter/MigratedFromGit/Source/GLTFExporter/Public/GLTFExportOptions.h
+++ b/Engine/Plugins/Enterprise/GLTFExporter/MigratedFromGit/Source/GLTFExporter/Public/GLTFExportOptions.h
@@ -53,11 +53,19 @@ enum class EGLTFExporterTextureHDREncoding : uint8
 {
 	None,
 	RGBM,
-	RGBD,
 	RGBE,
 	HDR
 };
 
+UENUM(BlueprintType)
+enum class EGLTFExporterLightMobility : uint8
+{
+	None,
+    MovableOnly,
+	MovableAndStationary,
+    All
+};
+
 UCLASS(config = EditorPerProjectUserSettings, HideCategories = (DebugProperty))
 class GLTFEXPORTER_API UGLTFExportOptions : public UObject
 {
@@ -110,7 +118,7 @@ public:
 	float ExportScale;
 
 	UPROPERTY(EditAnywhere, BlueprintReadWrite, config, Category = Scene)
-	uint32 bExportLights : 1;
+	EGLTFExporterLightMobility bExportLights;
 
 	UPROPERTY(EditAnywhere, BlueprintReadWrite, config, Category = Scene)
 	uint32 bExportCameras : 1;
