[GLTFExporter] fix(light): fix missing json object for spot light and avoid non-negative colors

commit 06d2bca44236ef845ad388b00761e86f072f4569
Date: Thu, 23 Jul 2020 17:40:57 +0200

    fix(light): fix missing json object for spot light and avoid non-negative colors
    (cherry picked from commit 6d671471159307bd6406156fd80ef3c61f5946ea​)

#skipci

diff --git a/Engine/Plugins/Enterprise/GLTFExporter/MigratedFromGit/Source/GLTFExporter/Private/Converters/GLTFConverterUtility.h b/Engine/Plugins/Enterprise/GLTFExporter/MigratedFromGit/Source/GLTFExporter/Private/Converters/GLTFConverterUtility.h
--- a/Engine/Plugins/Enterprise/GLTFExporter/MigratedFromGit/Source/GLTFExporter/Private/Converters/GLTFConverterUtility.h
+++ b/Engine/Plugins/Enterprise/GLTFExporter/MigratedFromGit/Source/GLTFExporter/Private/Converters/GLTFConverterUtility.h
@@ -55,8 +55,13 @@ struct FGLTFConverterUtility
 
 	static FGLTFJsonColor4 ConvertColor(const FLinearColor& Color)
 	{
-		// No conversion actually needed, this is primarily for type-safety.
-		return { Color.R, Color.G, Color.B, Color.A };
+		// Just make sure its non-negative (which can happen when using MakeFromColorTemperature).
+		return {
+			FMath::Max(Color.R, 0.0f),
+			FMath::Max(Color.G, 0.0f),
+			FMath::Max(Color.B, 0.0f),
+			FMath::Max(Color.A, 0.0f)
+        };
 	}
 
 	static FGLTFPackedColor ConvertColor(const FColor& Color)
diff --git a/Engine/Plugins/Enterprise/GLTFExporter/MigratedFromGit/Source/GLTFExporter/Private/Converters/GLTFLevelConverters.cpp b/Engine/Plugins/Enterprise/GLTFExporter/MigratedFromGit/Source/GLTFExporter/Private/Converters/GLTFLevelConverters.cpp
--- a/Engine/Plugins/Enterprise/GLTFExporter/MigratedFromGit/Source/GLTFExporter/Private/Converters/GLTFLevelConverters.cpp
+++ b/Engine/Plugins/Enterprise/GLTFExporter/MigratedFromGit/Source/GLTFExporter/Private/Converters/GLTFLevelConverters.cpp
@@ -193,8 +193,8 @@ FGLTFJsonLightIndex FGLTFLightComponentConverter::Add(FGLTFConvertBuilder& Build
 
 	if (const USpotLightComponent* SpotLightComponent = Cast<USpotLightComponent>(LightComponent))
 	{
-		Light.InnerConeAngle = FGLTFConverterUtility::ConvertLightAngle(SpotLightComponent->InnerConeAngle);
-		Light.OuterConeAngle = FGLTFConverterUtility::ConvertLightAngle(SpotLightComponent->OuterConeAngle);
+		Light.Spot.InnerConeAngle = FGLTFConverterUtility::ConvertLightAngle(SpotLightComponent->InnerConeAngle);
+		Light.Spot.OuterConeAngle = FGLTFConverterUtility::ConvertLightAngle(SpotLightComponent->OuterConeAngle);
 	}
 
 	return Builder.AddLight(Light);
diff --git a/Engine/Plugins/Enterprise/GLTFExporter/MigratedFromGit/Source/GLTFExporter/Private/Json/GLTFJsonLight.h b/Engine/Plugins/Enterprise/GLTFExporter/MigratedFromGit/Source/GLTFExporter/Private/Json/GLTFJsonLight.h
--- a/Engine/Plugins/Enterprise/GLTFExporter/MigratedFromGit/Source/GLTFExporter/Private/Json/GLTFJsonLight.h
+++ b/Engine/Plugins/Enterprise/GLTFExporter/MigratedFromGit/Source/GLTFExporter/Private/Json/GLTFJsonLight.h
@@ -7,6 +7,36 @@
 #include "Json/GLTFJsonUtility.h"
 #include "Serialization/JsonSerializer.h"
 
+struct FGLTFJsonSpotLight
+{
+	float InnerConeAngle;
+	float OuterConeAngle;
+
+	FGLTFJsonSpotLight()
+        : InnerConeAngle(0)
+        , OuterConeAngle(PI / 2.0f)
+	{
+	}
+
+	template <class CharType = TCHAR, class PrintPolicy = TPrettyJsonPrintPolicy<CharType>>
+    void WriteObject(TJsonWriter<CharType, PrintPolicy>& JsonWriter, FGLTFJsonExtensions& Extensions) const
+	{
+		JsonWriter.WriteObjectStart();
+
+		if (InnerConeAngle != 0)
+		{
+			FGLTFJsonUtility::WriteExactValue(JsonWriter, TEXT("innerConeAngle"), InnerConeAngle);
+		}
+
+		if (OuterConeAngle != PI / 2.0f)
+		{
+			FGLTFJsonUtility::WriteExactValue(JsonWriter, TEXT("outerConeAngle"), OuterConeAngle);
+		}
+
+		JsonWriter.WriteObjectEnd();
+	}
+};
+
 struct FGLTFJsonLight
 {
 	FString Name;
@@ -18,16 +48,13 @@ struct FGLTFJsonLight
 	float Intensity;
 	float Range;
 
-	float InnerConeAngle;
-	float OuterConeAngle;
+	FGLTFJsonSpotLight Spot;
 
 	FGLTFJsonLight()
 		: Type(EGLTFJsonLightType::None)
 		, Color(FGLTFJsonColor3::White)
 		, Intensity(1)
 		, Range(0)
-		, InnerConeAngle(0)
-		, OuterConeAngle(PI / 2.0f)
 	{
 	}
 
@@ -63,15 +90,8 @@ struct FGLTFJsonLight
 
 			if (Type == EGLTFJsonLightType::Spot)
 			{
-				if (InnerConeAngle != 0)
-				{
-					FGLTFJsonUtility::WriteExactValue(JsonWriter, TEXT("innerConeAngle"), InnerConeAngle);
-				}
-
-				if (OuterConeAngle != PI / 2.0f)
-				{
-					FGLTFJsonUtility::WriteExactValue(JsonWriter, TEXT("outerConeAngle"), OuterConeAngle);
-				}
+				JsonWriter.WriteIdentifierPrefix(TEXT("spot"));
+				Spot.WriteObject(JsonWriter, Extensions);
 			}
 		}
 
