[GLTFExporter] refactor(options): add GLTFExportOptions::bExportHiddenInGame (default false)

commit f239a6ec99301c4fe9bbd8de4b14a3d559f774a4
Date: Thu, 23 Jul 2020 17:45:55 +0200

    refactor(options): add GLTFExportOptions::bExportHiddenInGame (default false)
    (cherry picked from commit 7f65fac2fab584990ed49fa35b70906454cc8d9e​)

#skipci

diff --git a/Engine/Plugins/Enterprise/GLTFExporter/MigratedFromGit/Source/GLTFExporter/Private/GLTFExportOptions.cpp b/Engine/Plugins/Enterprise/GLTFExporter/MigratedFromGit/Source/GLTFExporter/Private/GLTFExportOptions.cpp
--- a/Engine/Plugins/Enterprise/GLTFExporter/MigratedFromGit/Source/GLTFExporter/Private/GLTFExportOptions.cpp
+++ b/Engine/Plugins/Enterprise/GLTFExporter/MigratedFromGit/Source/GLTFExporter/Private/GLTFExportOptions.cpp
@@ -28,6 +28,7 @@ UGLTFExportOptions::UGLTFExportOptions(const FObjectInitializer& ObjectInitializ
 	bExportLightmaps = true;
 	TextureHDREncoding = EGLTFExporterTextureHDREncoding::RGBM;
 	ExportScale = 0.01;
+	bExportHiddenInGame = false;
 	bExportLights = EGLTFExporterLightMobility::MovableAndStationary;
 	bExportCameras = true;
 	bExportOrbitalCameras = false;
diff --git a/Engine/Plugins/Enterprise/GLTFExporter/MigratedFromGit/Source/GLTFExporter/Public/GLTFExportOptions.h b/Engine/Plugins/Enterprise/GLTFExporter/MigratedFromGit/Source/GLTFExporter/Public/GLTFExportOptions.h
--- a/Engine/Plugins/Enterprise/GLTFExporter/MigratedFromGit/Source/GLTFExporter/Public/GLTFExportOptions.h
+++ b/Engine/Plugins/Enterprise/GLTFExporter/MigratedFromGit/Source/GLTFExporter/Public/GLTFExportOptions.h
@@ -117,6 +117,9 @@ public:
 		 ToolTip = "Scale factor used for exporting assets, by default: 0.01, for conversion from centimeters(Unreal default) to meters(glTF)."))
 	float ExportScale;
 
+	UPROPERTY(EditAnywhere, BlueprintReadWrite, config, Category = Scene)
+	uint32 bExportHiddenInGame : 1;
+
 	UPROPERTY(EditAnywhere, BlueprintReadWrite, config, Category = Scene)
 	EGLTFExporterLightMobility bExportLights;
 
