[GLTFExporter] chore: bump plugin version number to 0.5

commit a6a4534e70e3dcab3794a1a49b7ff05e18e1fdf1
Date: Tue, 21 Jul 2020 18:43:17 +0200

    chore: bump plugin version number to 0.5
    (cherry picked from commit b0ce141f03584f2cd6ad5431f933dc285f8fa81b​)

#skipci

diff --git a/Engine/Plugins/Enterprise/GLTFExporter/MigratedFromGit/GLTFExporter.uplugin b/Engine/Plugins/Enterprise/GLTFExporter/MigratedFromGit/GLTFExporter.uplugin
--- a/Engine/Plugins/Enterprise/GLTFExporter/MigratedFromGit/GLTFExporter.uplugin
+++ b/Engine/Plugins/Enterprise/GLTFExporter/MigratedFromGit/GLTFExporter.uplugin
@@ -1,7 +1,7 @@
 {
 	"FileVersion": 3,
-	"Version": 0,
-	"VersionName": "0.1",
+	"Version": 5,
+	"VersionName": "0.5",
 	"FriendlyName": "glTF Exporter",
 	"Description": "An exporter for Khronos glTF 3D models.",
 	"Category": "Exporters",
