[GLTFExporter] chore: temporarily disable GLTFExporterRuntime module, until the build pipeline supports multiple build configurations

commit 72d8a90cdd82c9064cad7ea0a8e399f9da9abc7a
Date: Tue, 21 Jul 2020 18:39:17 +0200

    chore: temporarily disable GLTFExporterRuntime module, until the build pipeline supports multiple build configurations
    (cherry picked from commit e612213a6b4ce85c6c750c0352cda29a3345cc63​)

#skipci

diff --git a/Engine/Plugins/Enterprise/GLTFExporter/MigratedFromGit/GLTFExporter.uplugin b/Engine/Plugins/Enterprise/GLTFExporter/MigratedFromGit/GLTFExporter.uplugin
--- a/Engine/Plugins/Enterprise/GLTFExporter/MigratedFromGit/GLTFExporter.uplugin
+++ b/Engine/Plugins/Enterprise/GLTFExporter/MigratedFromGit/GLTFExporter.uplugin
@@ -19,11 +19,6 @@
 			"Name": "GLTFExporter",
 			"Type": "Editor",
 			"LoadingPhase": "Default"
-		},
-		{
-			"Name": "GLTFExporterRuntime",
-			"Type": "Runtime",
-			"LoadingPhase": "Default"
 		}
 	],
 	"Plugins": [
