[GLTFExporter] refactor: add some TODO comments

commit 25f07c0f35be5b023f6396bdbf30d977f1399296
Date: Tue, 28 Jul 2020 15:29:06 +0200

    refactor: add some TODO comments
    (cherry picked from commit b92ee7616a3dbdb1331a602eb3fd132df52acd46​)

#skipci

diff --git a/Engine/Plugins/Enterprise/GLTFExporter/MigratedFromGit/Source/GLTFExporter/Private/Converters/GLTFLevelConverters.cpp b/Engine/Plugins/Enterprise/GLTFExporter/MigratedFromGit/Source/GLTFExporter/Private/Converters/GLTFLevelConverters.cpp
--- a/Engine/Plugins/Enterprise/GLTFExporter/MigratedFromGit/Source/GLTFExporter/Private/Converters/GLTFLevelConverters.cpp
+++ b/Engine/Plugins/Enterprise/GLTFExporter/MigratedFromGit/Source/GLTFExporter/Private/Converters/GLTFLevelConverters.cpp
@@ -33,7 +33,7 @@ FGLTFJsonNodeIndex FGLTFSceneComponentConverter::Add(FGLTFConvertBuilder& Builde
 	Node.Rotation = FGLTFConverterUtility::ConvertRotation(Transform.GetRotation());
 	Node.Scale = FGLTFConverterUtility::ConvertScale(Transform.GetScale3D());
 
-	if (SceneComponent->bHiddenInGame)
+	if (SceneComponent->bHiddenInGame) // TODO: make this configurable
 	{
 		// ignore any visible properties
 	}
@@ -123,6 +123,7 @@ FGLTFJsonSceneIndex FGLTFLevelConverter::Add(FGLTFConvertBuilder& Builder, const
 
 	for (const AActor* Actor : Level->Actors)
 	{
+		// TODO: should a LevelVariantSet be exported even if not selected for export?
 		if (const ALevelVariantSetsActor *LevelVariantSetsActor = Cast<ALevelVariantSetsActor>(Actor))
 		{
 			const FGLTFJsonLevelVariantSetsIndex LevelVariantSetsIndex = Builder.GetOrAddLevelVariantSets(LevelVariantSetsActor);
diff --git a/Engine/Plugins/Enterprise/GLTFExporter/MigratedFromGit/Source/GLTFExporter/Private/Converters/GLTFVariantSetConverters.cpp b/Engine/Plugins/Enterprise/GLTFExporter/MigratedFromGit/Source/GLTFExporter/Private/Converters/GLTFVariantSetConverters.cpp
--- a/Engine/Plugins/Enterprise/GLTFExporter/MigratedFromGit/Source/GLTFExporter/Private/Converters/GLTFVariantSetConverters.cpp
+++ b/Engine/Plugins/Enterprise/GLTFExporter/MigratedFromGit/Source/GLTFExporter/Private/Converters/GLTFVariantSetConverters.cpp
@@ -190,7 +190,7 @@ bool FGLTFLevelVariantSetsConverter::TryParseVisibilityPropertyValue(FGLTFConver
 		return false;
 	}
 
-	const FGLTFJsonNodeIndex NodeIndex = Builder.GetOrAddNode(Target);
+	const FGLTFJsonNodeIndex NodeIndex = Builder.GetOrAddNode(Target); // TODO: what if the target is not selected for export?
 	FGLTFJsonVariantNodeProperties& NodeProperties = OutVariant.Nodes.FindOrAdd(NodeIndex);
 
 	NodeProperties.Node = NodeIndex;
