[GLTFExporter] refactor(lightmap): don't use actor name when exporting lightmap

commit c1e5a9b2695439b35e919d9d95adebe0603a25a6
Date: Fri, 31 Jul 2020 12:33:42 +0200

    refactor(lightmap): don't use actor name when exporting lightmap
    (cherry picked from commit a7fee3ae347c0a5a23f4e0a533c82eedea4b0df6​)

#skipci

diff --git a/Engine/Plugins/Enterprise/GLTFExporter/MigratedFromGit/Source/GLTFExporter/Private/Converters/GLTFLightMapConverters.cpp b/Engine/Plugins/Enterprise/GLTFExporter/MigratedFromGit/Source/GLTFExporter/Private/Converters/GLTFLightMapConverters.cpp
--- a/Engine/Plugins/Enterprise/GLTFExporter/MigratedFromGit/Source/GLTFExporter/Private/Converters/GLTFLightMapConverters.cpp
+++ b/Engine/Plugins/Enterprise/GLTFExporter/MigratedFromGit/Source/GLTFExporter/Private/Converters/GLTFLightMapConverters.cpp
@@ -7,13 +7,6 @@
 
 FGLTFJsonLightMapIndex FGLTFLightMapConverter::Add(FGLTFConvertBuilder& Builder, const FString& Name, const UStaticMeshComponent* StaticMeshComponent)
 {
-	const AActor* Owner = StaticMeshComponent->GetOwner();
-	if (Owner == nullptr)
-	{
-		// TODO: report invalid mesh component
-		return FGLTFJsonLightMapIndex(INDEX_NONE);
-	}
-
 	const UStaticMesh* StaticMesh = StaticMeshComponent->GetStaticMesh();
 
 	if (StaticMesh == nullptr)
@@ -65,7 +58,7 @@ FGLTFJsonLightMapIndex FGLTFLightMapConverter::Add(FGLTFConvertBuilder& Builder,
 	const FVector4& LightMapScale = LightMapInteraction.GetScaleArray()[0];
 
 	FGLTFJsonLightMap JsonLightMap;
-	JsonLightMap.Name = Name.IsEmpty() ? Owner->GetName() + TEXT("_") + StaticMeshComponent->GetName() : Name;
+	JsonLightMap.Name = Name.IsEmpty() ? StaticMeshComponent->GetName() : Name;
 	JsonLightMap.Texture.Index = TextureIndex;
 	JsonLightMap.Texture.TexCoord = CoordinateIndex;
 	JsonLightMap.LightMapScale = { LightMapScale.X, LightMapScale.Y, LightMapScale.Z, LightMapScale.W };
