[GLTFExporter] refactor: add and improve TODO comments

commit 4e5993cc4bd3bccd9f3e6db3a250b3439e4fbc6d
Date: Fri, 31 Jul 2020 12:34:14 +0200

    refactor: add and improve TODO comments
    (cherry picked from commit 7e8e18a50735847f92bb322cd0c23c0fb5fb305a​)

#skipci

diff --git a/Engine/Plugins/Enterprise/GLTFExporter/MigratedFromGit/Source/GLTFExporter/Private/Converters/GLTFLevelConverters.cpp b/Engine/Plugins/Enterprise/GLTFExporter/MigratedFromGit/Source/GLTFExporter/Private/Converters/GLTFLevelConverters.cpp
--- a/Engine/Plugins/Enterprise/GLTFExporter/MigratedFromGit/Source/GLTFExporter/Private/Converters/GLTFLevelConverters.cpp
+++ b/Engine/Plugins/Enterprise/GLTFExporter/MigratedFromGit/Source/GLTFExporter/Private/Converters/GLTFLevelConverters.cpp
@@ -11,7 +11,7 @@ FGLTFJsonNodeIndex FGLTFSceneComponentConverter::Add(FGLTFConvertBuilder& Builde
 	const AActor* Owner = SceneComponent->GetOwner();
 	if (Owner == nullptr)
 	{
-		// TODO: report invalid scene component
+		// TODO: report error (invalid scene component)
 		return FGLTFJsonNodeIndex(INDEX_NONE);
 	}
 
@@ -73,6 +73,8 @@ FGLTFJsonNodeIndex FGLTFSceneComponentConverter::Add(FGLTFConvertBuilder& Builde
 		}
 	}
 
+	// TODO: add support for SkyLight?
+
 	return NodeIndex;
 }
 
@@ -163,7 +165,7 @@ FGLTFJsonCameraIndex FGLTFCameraComponentConverter::Add(FGLTFConvertBuilder& Bui
 			break;
 
 		default:
-		    // TODO: report unsupported camera type
+		    // TODO: report error (unsupported camera type)
 		    return FGLTFJsonCameraIndex(INDEX_NONE);
 	}
 
@@ -178,7 +180,7 @@ FGLTFJsonLightIndex FGLTFLightComponentConverter::Add(FGLTFConvertBuilder& Build
 
 	if (Light.Type == EGLTFJsonLightType::None)
 	{
-		// TODO: report unsupported light component type
+		// TODO: report error (unsupported light component type)
 		return FGLTFJsonLightIndex(INDEX_NONE);
 	}
 
diff --git a/Engine/Plugins/Enterprise/GLTFExporter/MigratedFromGit/Source/GLTFExporter/Private/Converters/GLTFLightMapConverters.cpp b/Engine/Plugins/Enterprise/GLTFExporter/MigratedFromGit/Source/GLTFExporter/Private/Converters/GLTFLightMapConverters.cpp
--- a/Engine/Plugins/Enterprise/GLTFExporter/MigratedFromGit/Source/GLTFExporter/Private/Converters/GLTFLightMapConverters.cpp
+++ b/Engine/Plugins/Enterprise/GLTFExporter/MigratedFromGit/Source/GLTFExporter/Private/Converters/GLTFLightMapConverters.cpp
@@ -58,7 +58,7 @@ FGLTFJsonLightMapIndex FGLTFLightMapConverter::Add(FGLTFConvertBuilder& Builder,
 	const FVector4& LightMapScale = LightMapInteraction.GetScaleArray()[0];
 
 	FGLTFJsonLightMap JsonLightMap;
-	JsonLightMap.Name = Name.IsEmpty() ? StaticMeshComponent->GetName() : Name;
+	JsonLightMap.Name = Name.IsEmpty() ? StaticMeshComponent->GetName() : Name; // TODO: use better name (similar to light and camera)
 	JsonLightMap.Texture.Index = TextureIndex;
 	JsonLightMap.Texture.TexCoord = CoordinateIndex;
 	JsonLightMap.LightMapScale = { LightMapScale.X, LightMapScale.Y, LightMapScale.Z, LightMapScale.W };
diff --git a/Engine/Plugins/Enterprise/GLTFExporter/MigratedFromGit/Source/GLTFExporter/Private/Converters/GLTFMaterialConverters.cpp b/Engine/Plugins/Enterprise/GLTFExporter/MigratedFromGit/Source/GLTFExporter/Private/Converters/GLTFMaterialConverters.cpp
--- a/Engine/Plugins/Enterprise/GLTFExporter/MigratedFromGit/Source/GLTFExporter/Private/Converters/GLTFMaterialConverters.cpp
+++ b/Engine/Plugins/Enterprise/GLTFExporter/MigratedFromGit/Source/GLTFExporter/Private/Converters/GLTFMaterialConverters.cpp
@@ -309,7 +309,7 @@ bool FGLTFMaterialConverter::TryGetBaseColorAndOpacity(FGLTFConvertBuilder& Buil
 	}
 	else
 	{
-		// TODO: report texture coordinate conflict
+		// TODO: report error (texture coordinate conflict)
 		return false;
 	}
 
@@ -461,7 +461,7 @@ bool FGLTFMaterialConverter::TryGetMetallicAndRoughness(FGLTFConvertBuilder& Bui
 	}
 	else
 	{
-		// TODO: report texture coordinate conflict
+		// TODO: report error (texture coordinate conflict)
 		return false;
 	}
 
@@ -615,7 +615,7 @@ bool FGLTFMaterialConverter::TryGetClearCoatRoughness(FGLTFConvertBuilder& Build
 	}
 	else
 	{
-		// TODO: report texture coordinate conflict
+		// TODO: report error (texture coordinate conflict)
 		return false;
 	}
 
@@ -1001,18 +1001,18 @@ bool FGLTFMaterialConverter::TryGetSourceTexture(const UTexture2D*& OutTexture,
 		{
 			if (ParameterValue == nullptr)
 			{
-				// TODO: report material as broken
+				// TODO: report error (no texture parameter assigned)
 			}
 			else
 			{
-				// TODO: report incorrect texture type
+				// TODO: report error (incorrect texture type)
 			}
 			return false;
 		}
 
 		if (!FGLTFMaterialUtility::TryGetTextureCoordinateIndex(TextureParameter, OutTexCoord))
 		{
-			// TODO: report failed to identify texture coordinate index
+			// TODO: report error (failed to identify texture coordinate index)
 			return false;
 		}
 
@@ -1028,18 +1028,18 @@ bool FGLTFMaterialConverter::TryGetSourceTexture(const UTexture2D*& OutTexture,
 		{
 			if (TextureSampler->Texture == nullptr)
 			{
-				// TODO: report material as broken
+				// TODO: report error (no texture sample assigned)
 			}
 			else
 			{
-				// TODO: report incorrect texture type
+				// TODO: report error (incorrect texture type)
 			}
 			return false;
 		}
 
 		if (!FGLTFMaterialUtility::TryGetTextureCoordinateIndex(TextureSampler, OutTexCoord))
 		{
-			// TODO: report failed to identify texture coordinate index
+			// TODO: report error (failed to identify texture coordinate index)
 			return false;
 		}
 
@@ -1169,7 +1169,7 @@ FGLTFPropertyBakeOutput FGLTFMaterialConverter::BakeMaterialProperty(FGLTFConver
 
 		if (TexCoords.Num() > 1)
 		{
-			// TODO: report multiple texture coordinates found, will use first
+			// TODO: report warning (multiple texture coordinates found, will use first)
 
 			// TODO: replace this hardcoded hack with something more configurable and proper
 			if (Property == MP_AmbientOcclusion && TexCoords.Contains(1))
diff --git a/Engine/Plugins/Enterprise/GLTFExporter/MigratedFromGit/Source/GLTFExporter/Private/Converters/GLTFTextureConverters.cpp b/Engine/Plugins/Enterprise/GLTFExporter/MigratedFromGit/Source/GLTFExporter/Private/Converters/GLTFTextureConverters.cpp
--- a/Engine/Plugins/Enterprise/GLTFExporter/MigratedFromGit/Source/GLTFExporter/Private/Converters/GLTFTextureConverters.cpp
+++ b/Engine/Plugins/Enterprise/GLTFExporter/MigratedFromGit/Source/GLTFExporter/Private/Converters/GLTFTextureConverters.cpp
@@ -72,6 +72,8 @@ FGLTFJsonTextureIndex FGLTFTexture2DConverter::Add(FGLTFConvertBuilder& Builder,
 	}
 	else
 	{
+		// TODO: investigate and fix reason why using this fallback will result in blurry/fuzzy exported images
+
 		const EPixelFormat RenderTargetFormat = FGLTFTextureUtility::IsHDRFormat(Texture2D->GetPixelFormat()) ? PF_FloatRGBA : PF_B8G8R8A8;
 		UTextureRenderTarget2D* RenderTarget = FGLTFTextureUtility::CreateRenderTarget(Size, RenderTargetFormat, true);
 
diff --git a/Engine/Plugins/Enterprise/GLTFExporter/MigratedFromGit/Source/GLTFExporter/Private/GLTFExportOptions.cpp b/Engine/Plugins/Enterprise/GLTFExporter/MigratedFromGit/Source/GLTFExporter/Private/GLTFExportOptions.cpp
--- a/Engine/Plugins/Enterprise/GLTFExporter/MigratedFromGit/Source/GLTFExporter/Private/GLTFExportOptions.cpp
+++ b/Engine/Plugins/Enterprise/GLTFExporter/MigratedFromGit/Source/GLTFExporter/Private/GLTFExportOptions.cpp
@@ -214,7 +214,7 @@ bool UGLTFExportOptions::CanEditChange(const FProperty* InProperty) const
 		return true;
 	}
 
-	return false; // TODO: change to true for each option implemented and support by the object to be exported.
+	return false; // TODO: change to true for each option implemented and relevant to the object selected for export.
 }
 
 FIntPoint UGLTFExportOptions::GetDefaultMaterialBakeSize() const
