[GLTFExporter] refactor: remove unnecessary export option bMapSkeletalMotionToRoot

commit bbb6afd1b92fe5349e27bb8bd2f779272251a87c
Date: Fri, 31 Jul 2020 19:28:14 +0200

    refactor: remove unnecessary export option bMapSkeletalMotionToRoot
    (cherry picked from commit 2bce38ef58fdd0ae84b7fac7abba39743413872c​)

#skipci

diff --git a/Engine/Plugins/Enterprise/GLTFExporter/MigratedFromGit/Source/GLTFExporter/Private/GLTFExportOptions.cpp b/Engine/Plugins/Enterprise/GLTFExporter/MigratedFromGit/Source/GLTFExporter/Private/GLTFExportOptions.cpp
--- a/Engine/Plugins/Enterprise/GLTFExporter/MigratedFromGit/Source/GLTFExporter/Private/GLTFExportOptions.cpp
+++ b/Engine/Plugins/Enterprise/GLTFExporter/MigratedFromGit/Source/GLTFExporter/Private/GLTFExportOptions.cpp
@@ -23,7 +23,6 @@ UGLTFExportOptions::UGLTFExportOptions(const FObjectInitializer& ObjectInitializ
 	bTangentDataQuantization = false;
 	bNormalizeUVCoordinates = EGLTFExporterNormalizeUVCoordinates::Never;
 	DefaultLevelOfDetail = 0;
-	bMapSkeletalMotionToRoot = false;
 	TextureFormat = EGLTFExporterTextureFormat::PNG;
 	bExportLightmaps = true;
 	TextureHDREncoding = EGLTFExporterTextureHDREncoding::RGBM;
@@ -214,7 +213,9 @@ bool UGLTFExportOptions::CanEditChange(const FProperty* InProperty) const
 		return true;
 	}
 
-	return false; // TODO: change to true for each option implemented and relevant to the object selected for export.
+	// TODO: return true for each option implemented and supported.
+
+	return false;
 }
 
 FIntPoint UGLTFExportOptions::GetDefaultMaterialBakeSize() const
diff --git a/Engine/Plugins/Enterprise/GLTFExporter/MigratedFromGit/Source/GLTFExporter/Public/GLTFExportOptions.h b/Engine/Plugins/Enterprise/GLTFExporter/MigratedFromGit/Source/GLTFExporter/Public/GLTFExportOptions.h
--- a/Engine/Plugins/Enterprise/GLTFExporter/MigratedFromGit/Source/GLTFExporter/Public/GLTFExportOptions.h
+++ b/Engine/Plugins/Enterprise/GLTFExporter/MigratedFromGit/Source/GLTFExporter/Public/GLTFExportOptions.h
@@ -100,9 +100,6 @@ public:
 	UPROPERTY(EditAnywhere, BlueprintReadWrite, config, Category = Mesh)
 	int32 DefaultLevelOfDetail;
 
-	UPROPERTY(EditAnywhere, BlueprintReadWrite, config, Category = Animation)
-	uint32 bMapSkeletalMotionToRoot : 1;
-
 	UPROPERTY(EditAnywhere, BlueprintReadWrite, config, Category = Texture)
 	EGLTFExporterTextureFormat TextureFormat;
 
