[GLTFExporter] feat(options): implement support for bExportUnlitMaterials

commit 226f47dfc6c682cd384736af7807a0a22cc55003
Date: Mon, 17 Aug 2020 14:19:43 +0200

    feat(options): implement support for bExportUnlitMaterials
    (cherry picked from commit d750e76157b539bba329b9d966a1199caf27493b​)

#skipci

diff --git a/Engine/Plugins/Enterprise/GLTFExporter/MigratedFromGit/Source/GLTFExporter/Private/Converters/GLTFMaterialConverters.cpp b/Engine/Plugins/Enterprise/GLTFExporter/MigratedFromGit/Source/GLTFExporter/Private/Converters/GLTFMaterialConverters.cpp
--- a/Engine/Plugins/Enterprise/GLTFExporter/MigratedFromGit/Source/GLTFExporter/Private/Converters/GLTFMaterialConverters.cpp
+++ b/Engine/Plugins/Enterprise/GLTFExporter/MigratedFromGit/Source/GLTFExporter/Private/Converters/GLTFMaterialConverters.cpp
@@ -184,6 +184,13 @@ bool FGLTFMaterialConverter::TryGetShadingModel(FGLTFConvertBuilder& Builder, EG
 		return false;
 	}
 
+	if (ConvertedShadingModel == EGLTFJsonShadingModel::Unlit && !Builder.ExportOptions->bExportUnlitMaterials)
+	{
+		const FString ShadingModelName = FGLTFConverterUtility::GetEnumDisplayName(ShadingModel);
+		Builder.AddWarningMessage(FString::Printf(TEXT("Shading model (%s) in material %s disabled by export options"), *ShadingModelName, *Material->GetName()));
+		return false;
+	}
+
 	OutShadingModel = ConvertedShadingModel;
 	return true;
 }
diff --git a/Engine/Plugins/Enterprise/GLTFExporter/MigratedFromGit/Source/GLTFExporter/Private/GLTFExportOptions.cpp b/Engine/Plugins/Enterprise/GLTFExporter/MigratedFromGit/Source/GLTFExporter/Private/GLTFExportOptions.cpp
--- a/Engine/Plugins/Enterprise/GLTFExporter/MigratedFromGit/Source/GLTFExporter/Private/GLTFExportOptions.cpp
+++ b/Engine/Plugins/Enterprise/GLTFExporter/MigratedFromGit/Source/GLTFExporter/Private/GLTFExportOptions.cpp
@@ -208,7 +208,10 @@ void UGLTFExportOptions::FillOptions(bool bBatchMode, bool bShowOptionDialog, co
 
 bool UGLTFExportOptions::CanEditChange(const FProperty* InProperty) const
 {
-	if (InProperty->GetName() == TEXT("DefaultMaterialBakeSize"))
+	const FString PropertyName = InProperty->GetName();
+
+	if (PropertyName == TEXT("bExportUnlitMaterials") ||
+		PropertyName == TEXT("DefaultMaterialBakeSize"))
 	{
 		return true;
 	}
