[GLTFExporter] feat(options): implement support for bExportClearCoatMaterials

commit 3f31192e31634adb52f36cea1980e09304311a82
Date: Mon, 17 Aug 2020 14:19:52 +0200

    feat(options): implement support for bExportClearCoatMaterials
    (cherry picked from commit 6d0dd5f1734585875a0ba62c1e264fb716cb60cb​)

#skipci

diff --git a/Engine/Plugins/Enterprise/GLTFExporter/MigratedFromGit/Source/GLTFExporter/Private/Converters/GLTFMaterialConverters.cpp b/Engine/Plugins/Enterprise/GLTFExporter/MigratedFromGit/Source/GLTFExporter/Private/Converters/GLTFMaterialConverters.cpp
--- a/Engine/Plugins/Enterprise/GLTFExporter/MigratedFromGit/Source/GLTFExporter/Private/Converters/GLTFMaterialConverters.cpp
+++ b/Engine/Plugins/Enterprise/GLTFExporter/MigratedFromGit/Source/GLTFExporter/Private/Converters/GLTFMaterialConverters.cpp
@@ -191,6 +191,13 @@ bool FGLTFMaterialConverter::TryGetShadingModel(FGLTFConvertBuilder& Builder, EG
 		return false;
 	}
 
+	if (ConvertedShadingModel == EGLTFJsonShadingModel::ClearCoat && !Builder.ExportOptions->bExportClearCoatMaterials)
+	{
+		const FString ShadingModelName = FGLTFConverterUtility::GetEnumDisplayName(ShadingModel);
+		Builder.AddWarningMessage(FString::Printf(TEXT("Shading model (%s) in material %s disabled by export options"), *ShadingModelName, *Material->GetName()));
+		return false;
+	}
+
 	OutShadingModel = ConvertedShadingModel;
 	return true;
 }
diff --git a/Engine/Plugins/Enterprise/GLTFExporter/MigratedFromGit/Source/GLTFExporter/Private/GLTFExportOptions.cpp b/Engine/Plugins/Enterprise/GLTFExporter/MigratedFromGit/Source/GLTFExporter/Private/GLTFExportOptions.cpp
--- a/Engine/Plugins/Enterprise/GLTFExporter/MigratedFromGit/Source/GLTFExporter/Private/GLTFExportOptions.cpp
+++ b/Engine/Plugins/Enterprise/GLTFExporter/MigratedFromGit/Source/GLTFExporter/Private/GLTFExportOptions.cpp
@@ -211,6 +211,7 @@ bool UGLTFExportOptions::CanEditChange(const FProperty* InProperty) const
 	const FString PropertyName = InProperty->GetName();
 
 	if (PropertyName == TEXT("bExportUnlitMaterials") ||
+		PropertyName == TEXT("bExportClearCoatMaterials") ||
 		PropertyName == TEXT("DefaultMaterialBakeSize"))
 	{
 		return true;
