[GLTFExporter] feat(options): implement support for bExportVertexColors

commit 5b9a43b1210cd87fdadba5a4ba0f1ee29c1ac37f
Date: Mon, 17 Aug 2020 09:53:08 +0200

    feat(options): implement support for bExportVertexColors
    (cherry picked from commit c054437d9eb5c7d03bdbc62ca47cd09ed7c0d1c9​)

#skipci

diff --git a/Engine/Plugins/Enterprise/GLTFExporter/MigratedFromGit/Source/GLTFExporter/Private/Converters/GLTFSkeletalMeshConverters.cpp b/Engine/Plugins/Enterprise/GLTFExporter/MigratedFromGit/Source/GLTFExporter/Private/Converters/GLTFSkeletalMeshConverters.cpp
--- a/Engine/Plugins/Enterprise/GLTFExporter/MigratedFromGit/Source/GLTFExporter/Private/Converters/GLTFSkeletalMeshConverters.cpp
+++ b/Engine/Plugins/Enterprise/GLTFExporter/MigratedFromGit/Source/GLTFExporter/Private/Converters/GLTFSkeletalMeshConverters.cpp
@@ -79,7 +79,12 @@ FGLTFJsonMeshIndex FGLTFSkeletalMeshConverter::Add(FGLTFConvertBuilder& Builder,
 
 	FGLTFJsonAttributes JsonAttributes;
 	JsonAttributes.Position = Builder.GetOrAddPositionAccessor(PositionBuffer, Name + TEXT("_Positions"));
-	JsonAttributes.Color0 = Builder.GetOrAddColorAccessor(ColorBuffer, Name + TEXT("_Colors"));
+
+	if (Builder.ExportOptions->bExportVertexColors)
+	{
+		JsonAttributes.Color0 = Builder.GetOrAddColorAccessor(ColorBuffer, Name + TEXT("_Colors"));
+	}
+
 	JsonAttributes.Normal = Builder.GetOrAddNormalAccessor(VertexBuffer, Name + TEXT("_Normals"));
 	JsonAttributes.Tangent = Builder.GetOrAddTangentAccessor(VertexBuffer, Name + TEXT("_Tangents"));
 
diff --git a/Engine/Plugins/Enterprise/GLTFExporter/MigratedFromGit/Source/GLTFExporter/Private/Converters/GLTFStaticMeshConverters.cpp b/Engine/Plugins/Enterprise/GLTFExporter/MigratedFromGit/Source/GLTFExporter/Private/Converters/GLTFStaticMeshConverters.cpp
--- a/Engine/Plugins/Enterprise/GLTFExporter/MigratedFromGit/Source/GLTFExporter/Private/Converters/GLTFStaticMeshConverters.cpp
+++ b/Engine/Plugins/Enterprise/GLTFExporter/MigratedFromGit/Source/GLTFExporter/Private/Converters/GLTFStaticMeshConverters.cpp
@@ -65,7 +65,12 @@ FGLTFJsonMeshIndex FGLTFStaticMeshConverter::Add(FGLTFConvertBuilder& Builder, c
 
 	FGLTFJsonAttributes JsonAttributes;
 	JsonAttributes.Position = Builder.GetOrAddPositionAccessor(PositionBuffer, Name + TEXT("_Positions"));
-	JsonAttributes.Color0 = Builder.GetOrAddColorAccessor(ColorBuffer, Name + TEXT("_Colors"));
+
+	if (Builder.ExportOptions->bExportVertexColors)
+	{
+		JsonAttributes.Color0 = Builder.GetOrAddColorAccessor(ColorBuffer, Name + TEXT("_Colors"));
+	}
+
 	JsonAttributes.Normal = Builder.GetOrAddNormalAccessor(VertexBuffer, Name + TEXT("_Normals"));
 	JsonAttributes.Tangent = Builder.GetOrAddTangentAccessor(VertexBuffer, Name + TEXT("_Tangents"));
 
diff --git a/Engine/Plugins/Enterprise/GLTFExporter/MigratedFromGit/Source/GLTFExporter/Private/GLTFExportOptions.cpp b/Engine/Plugins/Enterprise/GLTFExporter/MigratedFromGit/Source/GLTFExporter/Private/GLTFExportOptions.cpp
--- a/Engine/Plugins/Enterprise/GLTFExporter/MigratedFromGit/Source/GLTFExporter/Private/GLTFExportOptions.cpp
+++ b/Engine/Plugins/Enterprise/GLTFExporter/MigratedFromGit/Source/GLTFExporter/Private/GLTFExportOptions.cpp
@@ -212,7 +212,8 @@ bool UGLTFExportOptions::CanEditChange(const FProperty* InProperty) const
 
 	if (PropertyName == TEXT("bExportUnlitMaterials") ||
 		PropertyName == TEXT("bExportClearCoatMaterials") ||
-		PropertyName == TEXT("DefaultMaterialBakeSize"))
+		PropertyName == TEXT("DefaultMaterialBakeSize") ||
+		PropertyName == TEXT("bExportVertexColors"))
 	{
 		return true;
 	}
