[GLTFExporter] feat(options): implement support for bBakeMaterialInputs

commit c3f66e7e6fca5b5c3c77fbe2e0d88acd4b317a9d
Date: Mon, 17 Aug 2020 11:49:32 +0200

    feat(options): implement support for bBakeMaterialInputs
    (cherry picked from commit 9cee18dd8dc1e39651f113c4ca85a139285d7430​)

#skipci

diff --git a/Engine/Plugins/Enterprise/GLTFExporter/MigratedFromGit/Source/GLTFExporter/Private/Converters/GLTFMaterialConverters.cpp b/Engine/Plugins/Enterprise/GLTFExporter/MigratedFromGit/Source/GLTFExporter/Private/Converters/GLTFMaterialConverters.cpp
--- a/Engine/Plugins/Enterprise/GLTFExporter/MigratedFromGit/Source/GLTFExporter/Private/Converters/GLTFMaterialConverters.cpp
+++ b/Engine/Plugins/Enterprise/GLTFExporter/MigratedFromGit/Source/GLTFExporter/Private/Converters/GLTFMaterialConverters.cpp
@@ -236,6 +236,16 @@ bool FGLTFMaterialConverter::TryGetBaseColorAndOpacity(FGLTFConvertBuilder& Buil
 		return true;
 	}
 
+	if (!Builder.ExportOptions->bBakeMaterialInputs)
+	{
+		Builder.AddWarningMessage(FString::Printf(
+			TEXT("Properties %s and %s for material %s will use defaults since baking is disabled by export options"),
+			FGLTFMaterialUtility::GetPropertyName(BaseColorProperty),
+			FGLTFMaterialUtility::GetPropertyName(OpacityProperty),
+			*Material->GetName()));
+		return false;
+	}
+
 	// TODO: add support for calculating the ideal resolution to use for baking based on connected (texture) nodes
 	FIntPoint TextureSize = Builder.ExportOptions->GetDefaultMaterialBakeSize();
 
@@ -391,6 +401,16 @@ bool FGLTFMaterialConverter::TryGetMetallicAndRoughness(FGLTFConvertBuilder& Bui
 		return true;
 	}
 
+	if (!Builder.ExportOptions->bBakeMaterialInputs)
+	{
+		Builder.AddWarningMessage(FString::Printf(
+			TEXT("Properties %s and %s for material %s will use defaults since baking is disabled by export options"),
+			FGLTFMaterialUtility::GetPropertyName(MetallicProperty),
+			FGLTFMaterialUtility::GetPropertyName(RoughnessProperty),
+			*Material->GetName()));
+		return false;
+	}
+
 	// TODO: add support for calculating the ideal resolution to use for baking based on connected (texture) nodes
 	FIntPoint TextureSize = Builder.ExportOptions->GetDefaultMaterialBakeSize();
 
@@ -545,6 +565,16 @@ bool FGLTFMaterialConverter::TryGetClearCoatRoughness(FGLTFConvertBuilder& Build
 		return true;
 	}
 
+	if (!Builder.ExportOptions->bBakeMaterialInputs)
+	{
+		Builder.AddWarningMessage(FString::Printf(
+			TEXT("Properties %s and %s for material %s will use defaults since baking is disabled by export options"),
+			FGLTFMaterialUtility::GetPropertyName(IntensityProperty),
+			FGLTFMaterialUtility::GetPropertyName(RoughnessProperty),
+			*Material->GetName()));
+		return false;
+	}
+
 	// TODO: add support for calculating the ideal resolution to use for baking based on connected (texture) nodes
 	FIntPoint TextureSize = Builder.ExportOptions->GetDefaultMaterialBakeSize();
 
@@ -679,6 +709,15 @@ bool FGLTFMaterialConverter::TryGetEmissive(FGLTFConvertBuilder& Builder, FGLTFJ
 		return true;
 	}
 
+	if (!Builder.ExportOptions->bBakeMaterialInputs)
+	{
+		Builder.AddWarningMessage(FString::Printf(
+			TEXT("Property %s for material %s will use defaults since baking is disabled by export options"),
+			FGLTFMaterialUtility::GetPropertyName(EmissiveProperty),
+			*Material->GetName()));
+		return false;
+	}
+
 	const FGLTFPropertyBakeOutput PropertyBakeOutput = BakeMaterialProperty(Builder, EmissiveProperty, Material, JsonMaterial.EmissiveTexture.TexCoord);
 	const float EmissiveScale = PropertyBakeOutput.EmissiveScale;
 
@@ -1065,6 +1104,15 @@ bool FGLTFMaterialConverter::TryGetSourceTexture(const UTexture2D*& OutTexture,
 
 bool FGLTFMaterialConverter::TryGetBakedMaterialProperty(FGLTFConvertBuilder& Builder, FGLTFJsonTextureInfo& OutTexInfo, FGLTFJsonColor3& OutConstant, EMaterialProperty Property, const FString& PropertyName, const UMaterialInterface* Material) const
 {
+	if (!Builder.ExportOptions->bBakeMaterialInputs)
+	{
+		Builder.AddWarningMessage(FString::Printf(
+			TEXT("Property %s for material %s will use defaults since baking is disabled by export options"),
+			FGLTFMaterialUtility::GetPropertyName(Property),
+			*Material->GetName()));
+		return false;
+	}
+
 	const FGLTFPropertyBakeOutput PropertyBakeOutput = BakeMaterialProperty(Builder, Property, Material, OutTexInfo.TexCoord);
 
 	if (PropertyBakeOutput.bIsConstant)
@@ -1084,6 +1132,15 @@ bool FGLTFMaterialConverter::TryGetBakedMaterialProperty(FGLTFConvertBuilder& Bu
 
 bool FGLTFMaterialConverter::TryGetBakedMaterialProperty(FGLTFConvertBuilder& Builder, FGLTFJsonTextureInfo& OutTexInfo, FGLTFJsonColor4& OutConstant, EMaterialProperty Property, const FString& PropertyName, const UMaterialInterface* Material) const
 {
+	if (!Builder.ExportOptions->bBakeMaterialInputs)
+	{
+		Builder.AddWarningMessage(FString::Printf(
+			TEXT("Property %s for material %s will use defaults since baking is disabled by export options"),
+			FGLTFMaterialUtility::GetPropertyName(Property),
+			*Material->GetName()));
+		return false;
+	}
+
 	const FGLTFPropertyBakeOutput PropertyBakeOutput = BakeMaterialProperty(Builder, Property, Material, OutTexInfo.TexCoord);
 
 	if (PropertyBakeOutput.bIsConstant)
@@ -1103,6 +1160,15 @@ bool FGLTFMaterialConverter::TryGetBakedMaterialProperty(FGLTFConvertBuilder& Bu
 
 inline bool FGLTFMaterialConverter::TryGetBakedMaterialProperty(FGLTFConvertBuilder& Builder, FGLTFJsonTextureInfo& OutTexInfo, float& OutConstant, EMaterialProperty Property, const FString& PropertyName, const UMaterialInterface* Material) const
 {
+	if (!Builder.ExportOptions->bBakeMaterialInputs)
+	{
+		Builder.AddWarningMessage(FString::Printf(
+			TEXT("Property %s for material %s will use defaults since baking is disabled by export options"),
+			FGLTFMaterialUtility::GetPropertyName(Property),
+			*Material->GetName()));
+		return false;
+	}
+
 	const FGLTFPropertyBakeOutput PropertyBakeOutput = BakeMaterialProperty(Builder, Property, Material, OutTexInfo.TexCoord);
 
 	if (PropertyBakeOutput.bIsConstant)
@@ -1122,6 +1188,15 @@ inline bool FGLTFMaterialConverter::TryGetBakedMaterialProperty(FGLTFConvertBuil
 
 bool FGLTFMaterialConverter::TryGetBakedMaterialProperty(FGLTFConvertBuilder& Builder, FGLTFJsonTextureInfo& OutTexInfo, EMaterialProperty Property, const FString& PropertyName, const UMaterialInterface* Material) const
 {
+	if (!Builder.ExportOptions->bBakeMaterialInputs)
+	{
+		Builder.AddWarningMessage(FString::Printf(
+			TEXT("Property %s for material %s will use defaults since baking is disabled by export options"),
+			FGLTFMaterialUtility::GetPropertyName(Property),
+			*Material->GetName()));
+		return false;
+	}
+
 	const FGLTFPropertyBakeOutput PropertyBakeOutput = BakeMaterialProperty(Builder, Property, Material, OutTexInfo.TexCoord);
 
 	if (!PropertyBakeOutput.bIsConstant)
diff --git a/Engine/Plugins/Enterprise/GLTFExporter/MigratedFromGit/Source/GLTFExporter/Private/GLTFExportOptions.cpp b/Engine/Plugins/Enterprise/GLTFExporter/MigratedFromGit/Source/GLTFExporter/Private/GLTFExportOptions.cpp
--- a/Engine/Plugins/Enterprise/GLTFExporter/MigratedFromGit/Source/GLTFExporter/Private/GLTFExportOptions.cpp
+++ b/Engine/Plugins/Enterprise/GLTFExporter/MigratedFromGit/Source/GLTFExporter/Private/GLTFExportOptions.cpp
@@ -212,6 +212,7 @@ bool UGLTFExportOptions::CanEditChange(const FProperty* InProperty) const
 
 	if (PropertyName == TEXT("bExportUnlitMaterials") ||
 		PropertyName == TEXT("bExportClearCoatMaterials") ||
+		PropertyName == TEXT("bBakeMaterialInputs") ||
 		PropertyName == TEXT("DefaultMaterialBakeSize") ||
 		PropertyName == TEXT("bExportVertexColors"))
 	{
