[GLTFExporter] feat(options): implement support for bExportLightmaps

commit 857aaefd4315e5ba84609827f912a5e91d5a419e
Date: Mon, 17 Aug 2020 14:18:31 +0200

    feat(options): implement support for bExportLightmaps
    (cherry picked from commit 939772161da24773a38c144324a837050deb0415​)

#skipci

diff --git a/Engine/Plugins/Enterprise/GLTFExporter/MigratedFromGit/Source/GLTFExporter/Private/Converters/GLTFLevelConverters.cpp b/Engine/Plugins/Enterprise/GLTFExporter/MigratedFromGit/Source/GLTFExporter/Private/Converters/GLTFLevelConverters.cpp
--- a/Engine/Plugins/Enterprise/GLTFExporter/MigratedFromGit/Source/GLTFExporter/Private/Converters/GLTFLevelConverters.cpp
+++ b/Engine/Plugins/Enterprise/GLTFExporter/MigratedFromGit/Source/GLTFExporter/Private/Converters/GLTFLevelConverters.cpp
@@ -40,11 +40,7 @@ FGLTFJsonNodeIndex FGLTFSceneComponentConverter::Add(FGLTFConvertBuilder& Builde
 	else if (const UStaticMeshComponent* StaticMeshComponent = Cast<UStaticMeshComponent>(SceneComponent))
 	{
 		Node.Mesh = Builder.GetOrAddMesh(StaticMeshComponent);
-
-		if (true /* TODO: make export of light-maps optional (selectable via export options) */)
-		{
-			Node.LightMap = Builder.GetOrAddLightMap(StaticMeshComponent);
-		}
+		Node.LightMap = Builder.GetOrAddLightMap(StaticMeshComponent);
 	}
 	else if (const USkeletalMeshComponent* SkeletalMeshComponent = Cast<USkeletalMeshComponent>(SceneComponent))
 	{
diff --git a/Engine/Plugins/Enterprise/GLTFExporter/MigratedFromGit/Source/GLTFExporter/Private/Converters/GLTFLightMapConverters.cpp b/Engine/Plugins/Enterprise/GLTFExporter/MigratedFromGit/Source/GLTFExporter/Private/Converters/GLTFLightMapConverters.cpp
--- a/Engine/Plugins/Enterprise/GLTFExporter/MigratedFromGit/Source/GLTFExporter/Private/Converters/GLTFLightMapConverters.cpp
+++ b/Engine/Plugins/Enterprise/GLTFExporter/MigratedFromGit/Source/GLTFExporter/Private/Converters/GLTFLightMapConverters.cpp
@@ -37,6 +37,15 @@ FGLTFJsonLightMapIndex FGLTFLightMapConverter::Add(FGLTFConvertBuilder& Builder,
 		return FGLTFJsonLightMapIndex(INDEX_NONE);
 	}
 
+	// TODO: this check should ideally be placed earlier (or outside this function), but we need to check
+	// if a light-map is in use before displaying warnings about skipped light-maps.
+	// Perhaps we should extract the code needed for checking for the presence of light-maps into
+	// a utility-function to allow the check to be performed elsewhere?
+	if (!Builder.ExportOptions->bExportLightmaps)
+	{
+		Builder.AddWarningMessage(FString::Printf(TEXT("Lightmap for %s disabled by export options"), *StaticMeshComponent->GetOwner()->GetName()));
+	}
+
 	const FLightMap2D* LightMap2D = MeshMapBuildData->LightMap->GetLightMap2D();
 	check(LightMap2D);
 
diff --git a/Engine/Plugins/Enterprise/GLTFExporter/MigratedFromGit/Source/GLTFExporter/Private/GLTFExportOptions.cpp b/Engine/Plugins/Enterprise/GLTFExporter/MigratedFromGit/Source/GLTFExporter/Private/GLTFExportOptions.cpp
--- a/Engine/Plugins/Enterprise/GLTFExporter/MigratedFromGit/Source/GLTFExporter/Private/GLTFExportOptions.cpp
+++ b/Engine/Plugins/Enterprise/GLTFExporter/MigratedFromGit/Source/GLTFExporter/Private/GLTFExportOptions.cpp
@@ -214,7 +214,8 @@ bool UGLTFExportOptions::CanEditChange(const FProperty* InProperty) const
 		PropertyName == TEXT("bExportClearCoatMaterials") ||
 		PropertyName == TEXT("bBakeMaterialInputs") ||
 		PropertyName == TEXT("DefaultMaterialBakeSize") ||
-		PropertyName == TEXT("bExportVertexColors"))
+		PropertyName == TEXT("bExportVertexColors") ||
+		PropertyName == TEXT("bExportLightmaps"))
 	{
 		return true;
 	}
