[GLTFExporter] feat(options): implement support for DefaultLevelOfDetail

commit a1538b732eeb66ea1c432f64062952db4147a90e
Date: Mon, 17 Aug 2020 14:50:00 +0200

    feat(options): implement support for DefaultLevelOfDetail
    (cherry picked from commit 9cad17aafd60e9be6c7e89d30806b61eec897f51​)

#skipci

diff --git a/Engine/Plugins/Enterprise/GLTFExporter/MigratedFromGit/Source/GLTFExporter/Private/Builders/GLTFConvertBuilder.cpp b/Engine/Plugins/Enterprise/GLTFExporter/MigratedFromGit/Source/GLTFExporter/Private/Builders/GLTFConvertBuilder.cpp
--- a/Engine/Plugins/Enterprise/GLTFExporter/MigratedFromGit/Source/GLTFExporter/Private/Builders/GLTFConvertBuilder.cpp
+++ b/Engine/Plugins/Enterprise/GLTFExporter/MigratedFromGit/Source/GLTFExporter/Private/Builders/GLTFConvertBuilder.cpp
@@ -106,6 +106,11 @@ FGLTFJsonMeshIndex FGLTFConvertBuilder::GetOrAddMesh(const UStaticMesh* StaticMe
 		return FGLTFJsonMeshIndex(INDEX_NONE);
 	}
 
+	if (LODIndex < 0)
+	{
+		LODIndex = ExportOptions->DefaultLevelOfDetail;
+	}
+
 	return StaticMeshConverter.GetOrAdd(*this, DesiredName.IsEmpty() ? FGLTFBuilderUtility::GetLODName(StaticMesh, LODIndex) : DesiredName, StaticMesh, LODIndex, OverrideVertexColors, OverrideMaterials);
 }
 
@@ -117,7 +122,7 @@ FGLTFJsonMeshIndex FGLTFConvertBuilder::GetOrAddMesh(const UStaticMeshComponent*
 	}
 
 	const UStaticMesh* StaticMesh = StaticMeshComponent->GetStaticMesh();
-	const int32 LODIndex = StaticMeshComponent->ForcedLodModel > 0 ? StaticMeshComponent->ForcedLodModel - 1 : /* auto-select */ 0;
+	const int32 LODIndex = StaticMeshComponent->ForcedLodModel > 0 ? StaticMeshComponent->ForcedLodModel - 1 : ExportOptions->DefaultLevelOfDetail;
 	const FColorVertexBuffer* OverrideVertexColors = StaticMeshComponent->LODData.IsValidIndex(LODIndex) ? StaticMeshComponent->LODData[LODIndex].OverrideVertexColors : nullptr;
 	const FGLTFMaterialArray OverrideMaterials = FGLTFMaterialArray(StaticMeshComponent->OverrideMaterials);
 
@@ -151,6 +156,11 @@ FGLTFJsonMeshIndex FGLTFConvertBuilder::GetOrAddMesh(const USkeletalMesh* Skelet
 		return FGLTFJsonMeshIndex(INDEX_NONE);
 	}
 
+	if (LODIndex < 0)
+	{
+		LODIndex = ExportOptions->DefaultLevelOfDetail;
+	}
+
 	return SkeletalMeshConverter.GetOrAdd(*this, DesiredName.IsEmpty() ? FGLTFBuilderUtility::GetLODName(SkeletalMesh, LODIndex) : DesiredName, SkeletalMesh, LODIndex, OverrideVertexColors, OverrideSkinWeights, OverrideMaterials);
 }
 
@@ -162,7 +172,7 @@ FGLTFJsonMeshIndex FGLTFConvertBuilder::GetOrAddMesh(const USkeletalMeshComponen
 	}
 
 	const USkeletalMesh* SkeletalMesh = SkeletalMeshComponent->SkeletalMesh;
-	const int32 LODIndex = SkeletalMeshComponent->GetForcedLOD() > 0 ? SkeletalMeshComponent->GetForcedLOD() - 1 : /* auto-select */ 0;
+	const int32 LODIndex = SkeletalMeshComponent->GetForcedLOD() > 0 ? SkeletalMeshComponent->GetForcedLOD() - 1 : ExportOptions->DefaultLevelOfDetail;
 	const FColorVertexBuffer* OverrideVertexColors = SkeletalMeshComponent->LODInfo.IsValidIndex(LODIndex) ? SkeletalMeshComponent->LODInfo[LODIndex].OverrideVertexColors : nullptr;
 	const FSkinWeightVertexBuffer* OverrideSkinWeights = SkeletalMeshComponent->LODInfo.IsValidIndex(LODIndex) ? SkeletalMeshComponent->LODInfo[LODIndex].OverrideSkinWeights : nullptr;
 	const FGLTFMaterialArray OverrideMaterials = FGLTFMaterialArray(SkeletalMeshComponent->OverrideMaterials);
diff --git a/Engine/Plugins/Enterprise/GLTFExporter/MigratedFromGit/Source/GLTFExporter/Private/Builders/GLTFConvertBuilder.h b/Engine/Plugins/Enterprise/GLTFExporter/MigratedFromGit/Source/GLTFExporter/Private/Builders/GLTFConvertBuilder.h
--- a/Engine/Plugins/Enterprise/GLTFExporter/MigratedFromGit/Source/GLTFExporter/Private/Builders/GLTFConvertBuilder.h
+++ b/Engine/Plugins/Enterprise/GLTFExporter/MigratedFromGit/Source/GLTFExporter/Private/Builders/GLTFConvertBuilder.h
@@ -31,12 +31,12 @@ public:
 
 	FGLTFJsonBufferViewIndex GetOrAddIndexBufferView(const FRawStaticIndexBuffer* IndexBuffer, const FString& DesiredName = TEXT(""));
 	FGLTFJsonAccessorIndex GetOrAddIndexAccessor(const FStaticMeshSection* MeshSection, const FRawStaticIndexBuffer* IndexBuffer, const FString& DesiredName = TEXT(""));
-	FGLTFJsonMeshIndex GetOrAddMesh(const UStaticMesh* StaticMesh, int32 LODIndex = 0, const FColorVertexBuffer* OverrideVertexColors = nullptr, const FGLTFMaterialArray& OverrideMaterials = {}, const FString& DesiredName = TEXT(""));
+	FGLTFJsonMeshIndex GetOrAddMesh(const UStaticMesh* StaticMesh, int32 LODIndex = -1, const FColorVertexBuffer* OverrideVertexColors = nullptr, const FGLTFMaterialArray& OverrideMaterials = {}, const FString& DesiredName = TEXT(""));
 	FGLTFJsonMeshIndex GetOrAddMesh(const UStaticMeshComponent* StaticMeshComponent, const FString& DesiredName = TEXT(""));
 
 	FGLTFJsonBufferViewIndex GetOrAddIndexBufferView(const FMultiSizeIndexContainer* IndexContainer, const FString& DesiredName = TEXT(""));
 	FGLTFJsonAccessorIndex GetOrAddIndexAccessor(const FSkelMeshRenderSection* MeshSection, const FMultiSizeIndexContainer* IndexContainer, const FString& DesiredName = TEXT(""));
-	FGLTFJsonMeshIndex GetOrAddMesh(const USkeletalMesh* SkeletalMesh, int32 LODIndex = 0, const FColorVertexBuffer* OverrideVertexColors = nullptr, const FSkinWeightVertexBuffer* OverrideSkinWeights = nullptr, const FGLTFMaterialArray& OverrideMaterials = {}, const FString& DesiredName = TEXT(""));
+	FGLTFJsonMeshIndex GetOrAddMesh(const USkeletalMesh* SkeletalMesh, int32 LODIndex = -1, const FColorVertexBuffer* OverrideVertexColors = nullptr, const FSkinWeightVertexBuffer* OverrideSkinWeights = nullptr, const FGLTFMaterialArray& OverrideMaterials = {}, const FString& DesiredName = TEXT(""));
 	FGLTFJsonMeshIndex GetOrAddMesh(const USkeletalMeshComponent* SkeletalMeshComponent, const FString& DesiredName = TEXT(""));
 
 	FGLTFJsonMaterialIndex GetOrAddMaterial(const UMaterialInterface* Material, const FString& DesiredName = TEXT(""));
diff --git a/Engine/Plugins/Enterprise/GLTFExporter/MigratedFromGit/Source/GLTFExporter/Private/Converters/GLTFLightMapConverters.cpp b/Engine/Plugins/Enterprise/GLTFExporter/MigratedFromGit/Source/GLTFExporter/Private/Converters/GLTFLightMapConverters.cpp
--- a/Engine/Plugins/Enterprise/GLTFExporter/MigratedFromGit/Source/GLTFExporter/Private/Converters/GLTFLightMapConverters.cpp
+++ b/Engine/Plugins/Enterprise/GLTFExporter/MigratedFromGit/Source/GLTFExporter/Private/Converters/GLTFLightMapConverters.cpp
@@ -14,7 +14,7 @@ FGLTFJsonLightMapIndex FGLTFLightMapConverter::Add(FGLTFConvertBuilder& Builder,
 		return FGLTFJsonLightMapIndex(INDEX_NONE);
 	}
 
-	const int32 LODIndex = StaticMeshComponent->ForcedLodModel > 0 ? StaticMeshComponent->ForcedLodModel - 1 : /* auto-select */ 0;
+	const int32 LODIndex = StaticMeshComponent->ForcedLodModel > 0 ? StaticMeshComponent->ForcedLodModel - 1 : Builder.ExportOptions->DefaultLevelOfDetail;
 
 	if (LODIndex < 0 || StaticMesh->GetNumLODs() <= LODIndex || StaticMeshComponent->LODData.Num() <= LODIndex)
 	{
diff --git a/Engine/Plugins/Enterprise/GLTFExporter/MigratedFromGit/Source/GLTFExporter/Private/Exporters/GLTFMaterialExporter.cpp b/Engine/Plugins/Enterprise/GLTFExporter/MigratedFromGit/Source/GLTFExporter/Private/Exporters/GLTFMaterialExporter.cpp
--- a/Engine/Plugins/Enterprise/GLTFExporter/MigratedFromGit/Source/GLTFExporter/Private/Exporters/GLTFMaterialExporter.cpp
+++ b/Engine/Plugins/Enterprise/GLTFExporter/MigratedFromGit/Source/GLTFExporter/Private/Exporters/GLTFMaterialExporter.cpp
@@ -25,7 +25,7 @@ bool UGLTFMaterialExporter::AddObject(FGLTFContainerBuilder& Builder, const UObj
 	const UStaticMesh* PreviewMesh = FGLTFExporterUtility::GetPreviewMesh(Material);
 	if (PreviewMesh != nullptr)
 	{
-		const FGLTFJsonMeshIndex MeshIndex = Builder.GetOrAddMesh(PreviewMesh, 0, nullptr, { Material }, "PreviewMesh");
+		const FGLTFJsonMeshIndex MeshIndex = Builder.GetOrAddMesh(PreviewMesh, -1, nullptr, { Material }, "PreviewMesh");
 
 		FGLTFJsonNode Node;
 		Node.Mesh = MeshIndex;
diff --git a/Engine/Plugins/Enterprise/GLTFExporter/MigratedFromGit/Source/GLTFExporter/Private/GLTFExportOptions.cpp b/Engine/Plugins/Enterprise/GLTFExporter/MigratedFromGit/Source/GLTFExporter/Private/GLTFExportOptions.cpp
--- a/Engine/Plugins/Enterprise/GLTFExporter/MigratedFromGit/Source/GLTFExporter/Private/GLTFExportOptions.cpp
+++ b/Engine/Plugins/Enterprise/GLTFExporter/MigratedFromGit/Source/GLTFExporter/Private/GLTFExportOptions.cpp
@@ -215,6 +215,7 @@ bool UGLTFExportOptions::CanEditChange(const FProperty* InProperty) const
 		PropertyName == TEXT("bBakeMaterialInputs") ||
 		PropertyName == TEXT("DefaultMaterialBakeSize") ||
 		PropertyName == TEXT("bExportVertexColors") ||
+		PropertyName == TEXT("DefaultLevelOfDetail") ||
 		PropertyName == TEXT("bExportLightmaps"))
 	{
 		return true;
