[GLTFExporter] feat(options): implement support for bExportLights

commit 91181655ffbeb02233503228a013900f72470a7e
Date: Tue, 18 Aug 2020 08:17:44 +0200

    feat(options): implement support for bExportLights
    (cherry picked from commit 786787fde8e1547322b06e7a85f69133bee6bd56​)

#skipci

diff --git a/Engine/Plugins/Enterprise/GLTFExporter/MigratedFromGit/Source/GLTFExporter/Private/Converters/GLTFLevelConverters.cpp b/Engine/Plugins/Enterprise/GLTFExporter/MigratedFromGit/Source/GLTFExporter/Private/Converters/GLTFLevelConverters.cpp
--- a/Engine/Plugins/Enterprise/GLTFExporter/MigratedFromGit/Source/GLTFExporter/Private/Converters/GLTFLevelConverters.cpp
+++ b/Engine/Plugins/Enterprise/GLTFExporter/MigratedFromGit/Source/GLTFExporter/Private/Converters/GLTFLevelConverters.cpp
@@ -57,8 +57,7 @@ FGLTFJsonNodeIndex FGLTFSceneComponentConverter::Add(FGLTFConvertBuilder& Builde
 	}
 	else if (const ULightComponent* LightComponent = Cast<ULightComponent>(SceneComponent))
 	{
-		// TODO: make it configurable whether static lights should be exported or not
-		if (LightComponent->Mobility != EComponentMobility::Static)
+		if (ShouldExportLight(Builder.ExportOptions->bExportLights, LightComponent->Mobility))
 		{
 			// TODO: conversion of light direction should be done in separate converter
 			FGLTFJsonNode LightNode;
@@ -67,6 +66,12 @@ FGLTFJsonNodeIndex FGLTFSceneComponentConverter::Add(FGLTFConvertBuilder& Builde
 			LightNode.Light = Builder.GetOrAddLight(LightComponent, LightNode.Name);
 			Builder.AddChildNode(NodeIndex, LightNode);
 		}
+		else
+		{
+			Builder.AddWarningMessage(FString::Printf(
+				TEXT("Light %s disabled by export options"),
+				*Owner->GetName()));	// TODO: choose a more unique name if owner is not ALight
+		}
 	}
 
 	// TODO: add support for SkyLight?
@@ -74,6 +79,21 @@ FGLTFJsonNodeIndex FGLTFSceneComponentConverter::Add(FGLTFConvertBuilder& Builde
 	return NodeIndex;
 }
 
+bool FGLTFSceneComponentConverter::ShouldExportLight(EGLTFExporterLightMobility Options, EComponentMobility::Type Mobility)
+{
+	switch (Options)
+	{
+		case EGLTFExporterLightMobility::All:
+			return true;
+		case EGLTFExporterLightMobility::MovableAndStationary:
+			return Mobility == EComponentMobility::Movable || Mobility == EComponentMobility::Stationary;
+		case EGLTFExporterLightMobility::MovableOnly:
+			return Mobility == EComponentMobility::Movable;
+		default:
+			return false;
+	}
+}
+
 FGLTFJsonNodeIndex FGLTFActorConverter::Add(FGLTFConvertBuilder& Builder, const FString& Name, const AActor* Actor)
 {
 	if (Builder.bSelectedActorsOnly && !Actor->IsSelected())
diff --git a/Engine/Plugins/Enterprise/GLTFExporter/MigratedFromGit/Source/GLTFExporter/Private/Converters/GLTFLevelConverters.h b/Engine/Plugins/Enterprise/GLTFExporter/MigratedFromGit/Source/GLTFExporter/Private/Converters/GLTFLevelConverters.h
--- a/Engine/Plugins/Enterprise/GLTFExporter/MigratedFromGit/Source/GLTFExporter/Private/Converters/GLTFLevelConverters.h
+++ b/Engine/Plugins/Enterprise/GLTFExporter/MigratedFromGit/Source/GLTFExporter/Private/Converters/GLTFLevelConverters.h
@@ -5,10 +5,13 @@
 #include "Json/GLTFJsonIndex.h"
 #include "Converters/GLTFConverter.h"
 #include "Engine.h"
+#include "GLTFExportOptions.h"
 
 class FGLTFSceneComponentConverter final : public TGLTFConverter<FGLTFJsonNodeIndex, const USceneComponent*>
 {
 	FGLTFJsonNodeIndex Add(FGLTFConvertBuilder& Builder, const FString& Name, const USceneComponent* SceneComponent) override;
+
+	static bool ShouldExportLight(EGLTFExporterLightMobility Options, EComponentMobility::Type Mobility);
 };
 
 class FGLTFActorConverter final : public TGLTFConverter<FGLTFJsonNodeIndex, const AActor*>
diff --git a/Engine/Plugins/Enterprise/GLTFExporter/MigratedFromGit/Source/GLTFExporter/Private/GLTFExportOptions.cpp b/Engine/Plugins/Enterprise/GLTFExporter/MigratedFromGit/Source/GLTFExporter/Private/GLTFExportOptions.cpp
--- a/Engine/Plugins/Enterprise/GLTFExporter/MigratedFromGit/Source/GLTFExporter/Private/GLTFExportOptions.cpp
+++ b/Engine/Plugins/Enterprise/GLTFExporter/MigratedFromGit/Source/GLTFExporter/Private/GLTFExportOptions.cpp
@@ -216,7 +216,8 @@ bool UGLTFExportOptions::CanEditChange(const FProperty* InProperty) const
 		PropertyName == TEXT("DefaultMaterialBakeSize") ||
 		PropertyName == TEXT("bExportVertexColors") ||
 		PropertyName == TEXT("DefaultLevelOfDetail") ||
-		PropertyName == TEXT("bExportLightmaps"))
+		PropertyName == TEXT("bExportLightmaps") ||
+		PropertyName == TEXT("bExportLights"))
 	{
 		return true;
 	}
