[GLTFExporter] feat(options): implement support for ExportScale

commit 2cc96d9cf31669bc62b2b3d7f63cee322c0ecd46
Date: Tue, 18 Aug 2020 08:31:47 +0200

    feat(options): implement support for ExportScale
    (cherry picked from commit 8ab65eaab2b4b6241a9acfb36c3da2303b33a6ad​)

#skipci

diff --git a/Engine/Plugins/Enterprise/GLTFExporter/MigratedFromGit/Source/GLTFExporter/Private/Converters/GLTFBackdropConverters.cpp b/Engine/Plugins/Enterprise/GLTFExporter/MigratedFromGit/Source/GLTFExporter/Private/Converters/GLTFBackdropConverters.cpp
--- a/Engine/Plugins/Enterprise/GLTFExporter/MigratedFromGit/Source/GLTFExporter/Private/Converters/GLTFBackdropConverters.cpp
+++ b/Engine/Plugins/Enterprise/GLTFExporter/MigratedFromGit/Source/GLTFExporter/Private/Converters/GLTFBackdropConverters.cpp
@@ -69,7 +69,7 @@ FGLTFJsonBackdropIndex FGLTFBackdropConverter::Add(FGLTFConvertBuilder& Builder,
 	FVector ProjectionCenter;
 	if (FGLTFActorUtility::TryGetPropertyValue(Actor, TEXT("ProjectionCenter"), ProjectionCenter))
 	{
-		JsonBackdrop.ProjectionCenter = FGLTFConverterUtility::ConvertPosition(ProjectionCenter);
+		JsonBackdrop.ProjectionCenter = FGLTFConverterUtility::ConvertPosition(ProjectionCenter, Builder.ExportOptions->ExportScale);
 	}
 	else
 	{
diff --git a/Engine/Plugins/Enterprise/GLTFExporter/MigratedFromGit/Source/GLTFExporter/Private/Converters/GLTFCameraUtility.cpp b/Engine/Plugins/Enterprise/GLTFExporter/MigratedFromGit/Source/GLTFExporter/Private/Converters/GLTFCameraUtility.cpp
--- a/Engine/Plugins/Enterprise/GLTFExporter/MigratedFromGit/Source/GLTFExporter/Private/Converters/GLTFCameraUtility.cpp
+++ b/Engine/Plugins/Enterprise/GLTFExporter/MigratedFromGit/Source/GLTFExporter/Private/Converters/GLTFCameraUtility.cpp
@@ -4,23 +4,23 @@
 #include "Converters/GLTFConverterUtility.h"
 #include "Camera/CameraTypes.h"
 
-FGLTFJsonOrthographic FGLTFCameraUtility::ConvertOrthographic(const FMinimalViewInfo& View)
+FGLTFJsonOrthographic FGLTFCameraUtility::ConvertOrthographic(const FMinimalViewInfo& View, const float ExportScale)
 {
 	FGLTFJsonOrthographic Orthographic;
 
 	// NOTE: it goes against the glTF standard to not define a XMag, but the gltf viewer can handle it
 	if (View.bConstrainAspectRatio)
 	{
-		Orthographic.XMag = FGLTFConverterUtility::ConvertLength(View.OrthoWidth);
+		Orthographic.XMag = FGLTFConverterUtility::ConvertLength(View.OrthoWidth, ExportScale);
 	}
 
-	Orthographic.YMag = FGLTFConverterUtility::ConvertLength(View.OrthoWidth / View.AspectRatio); // TODO: is this correct?
-	Orthographic.ZFar = FGLTFConverterUtility::ConvertLength(View.OrthoFarClipPlane);
-	Orthographic.ZNear = FGLTFConverterUtility::ConvertLength(View.OrthoNearClipPlane);
+	Orthographic.YMag = FGLTFConverterUtility::ConvertLength(View.OrthoWidth / View.AspectRatio, ExportScale); // TODO: is this correct?
+	Orthographic.ZFar = FGLTFConverterUtility::ConvertLength(View.OrthoFarClipPlane, ExportScale);
+	Orthographic.ZNear = FGLTFConverterUtility::ConvertLength(View.OrthoNearClipPlane, ExportScale);
 	return Orthographic;
 }
 
-FGLTFJsonPerspective FGLTFCameraUtility::ConvertPerspective(const FMinimalViewInfo& View)
+FGLTFJsonPerspective FGLTFCameraUtility::ConvertPerspective(const FMinimalViewInfo& View, const float ExportScale)
 {
 	FGLTFJsonPerspective Perspective;
 
@@ -32,8 +32,8 @@ FGLTFJsonPerspective FGLTFCameraUtility::ConvertPerspective(const FMinimalViewIn
 	Perspective.YFov = ConvertFieldOfView(View);
 
 	// NOTE: even thought ZFar is optional, if we don't set it, then most gltf viewers won't handle it well.
-	Perspective.ZFar = FGLTFConverterUtility::ConvertLength(WORLD_MAX); // TODO: Unreal doesn't have max draw distance per view?
-	Perspective.ZNear = FGLTFConverterUtility::ConvertLength(GNearClippingPlane);
+	Perspective.ZFar = FGLTFConverterUtility::ConvertLength(WORLD_MAX, ExportScale); // TODO: Unreal doesn't have max draw distance per view?
+	Perspective.ZNear = FGLTFConverterUtility::ConvertLength(GNearClippingPlane, ExportScale);
 	return Perspective;
 }
 
diff --git a/Engine/Plugins/Enterprise/GLTFExporter/MigratedFromGit/Source/GLTFExporter/Private/Converters/GLTFCameraUtility.h b/Engine/Plugins/Enterprise/GLTFExporter/MigratedFromGit/Source/GLTFExporter/Private/Converters/GLTFCameraUtility.h
--- a/Engine/Plugins/Enterprise/GLTFExporter/MigratedFromGit/Source/GLTFExporter/Private/Converters/GLTFCameraUtility.h
+++ b/Engine/Plugins/Enterprise/GLTFExporter/MigratedFromGit/Source/GLTFExporter/Private/Converters/GLTFCameraUtility.h
@@ -6,8 +6,8 @@
 
 struct FGLTFCameraUtility
 {
-	static FGLTFJsonOrthographic ConvertOrthographic(const FMinimalViewInfo& View);
-	static FGLTFJsonPerspective ConvertPerspective(const FMinimalViewInfo& View);
+	static FGLTFJsonOrthographic ConvertOrthographic(const FMinimalViewInfo& View, const float ExportScale);
+	static FGLTFJsonPerspective ConvertPerspective(const FMinimalViewInfo& View, const float ExportScale);
 
 	static float ConvertFieldOfView(const FMinimalViewInfo& View);
 };
diff --git a/Engine/Plugins/Enterprise/GLTFExporter/MigratedFromGit/Source/GLTFExporter/Private/Converters/GLTFConverterUtility.h b/Engine/Plugins/Enterprise/GLTFExporter/MigratedFromGit/Source/GLTFExporter/Private/Converters/GLTFConverterUtility.h
--- a/Engine/Plugins/Enterprise/GLTFExporter/MigratedFromGit/Source/GLTFExporter/Private/Converters/GLTFConverterUtility.h
+++ b/Engine/Plugins/Enterprise/GLTFExporter/MigratedFromGit/Source/GLTFExporter/Private/Converters/GLTFConverterUtility.h
@@ -14,9 +14,9 @@
 
 struct FGLTFConverterUtility
 {
-	static float ConvertLength(const float Length)
+	static float ConvertLength(const float Length, const float ExportScale)
 	{
-		return Length * 0.01f; // TODO: use options export scale instead of hardcoded value
+		return Length * ExportScale;
 	}
 
 	static FGLTFJsonVector3 ConvertVector(const FVector& Vector)
@@ -26,9 +26,9 @@ struct FGLTFConverterUtility
 		return { Vector.X, Vector.Z, Vector.Y };
 	}
 
-	static FGLTFJsonVector3 ConvertPosition(const FVector& Position)
+	static FGLTFJsonVector3 ConvertPosition(const FVector& Position, const float ExportScale)
 	{
-		return ConvertVector(Position * 0.01f); // TODO: use options export scale instead of hardcoded value
+		return ConvertVector(Position * ExportScale);
 	}
 
 	static FGLTFJsonVector3 ConvertScale(const FVector& Scale)
diff --git a/Engine/Plugins/Enterprise/GLTFExporter/MigratedFromGit/Source/GLTFExporter/Private/Converters/GLTFLevelConverters.cpp b/Engine/Plugins/Enterprise/GLTFExporter/MigratedFromGit/Source/GLTFExporter/Private/Converters/GLTFLevelConverters.cpp
--- a/Engine/Plugins/Enterprise/GLTFExporter/MigratedFromGit/Source/GLTFExporter/Private/Converters/GLTFLevelConverters.cpp
+++ b/Engine/Plugins/Enterprise/GLTFExporter/MigratedFromGit/Source/GLTFExporter/Private/Converters/GLTFLevelConverters.cpp
@@ -29,7 +29,7 @@ FGLTFJsonNodeIndex FGLTFSceneComponentConverter::Add(FGLTFConvertBuilder& Builde
 
 	// TODO: add support for non-uniform scaling (Unreal doesn't treat combined transforms as simple matrix multiplication)
 	const FTransform Transform = bIsRootNode ? SceneComponent->GetComponentTransform() : SceneComponent->GetRelativeTransform();
-	Node.Translation = FGLTFConverterUtility::ConvertPosition(Transform.GetTranslation());
+	Node.Translation = FGLTFConverterUtility::ConvertPosition(Transform.GetTranslation(), Builder.ExportOptions->ExportScale);
 	Node.Rotation = FGLTFConverterUtility::ConvertRotation(Transform.GetRotation());
 	Node.Scale = FGLTFConverterUtility::ConvertScale(Transform.GetScale3D());
 
@@ -173,11 +173,11 @@ FGLTFJsonCameraIndex FGLTFCameraComponentConverter::Add(FGLTFConvertBuilder& Bui
 	switch (Camera.Type)
 	{
 		case EGLTFJsonCameraType::Orthographic:
-			Camera.Orthographic = FGLTFCameraUtility::ConvertOrthographic(DesiredView);
+			Camera.Orthographic = FGLTFCameraUtility::ConvertOrthographic(DesiredView, Builder.ExportOptions->ExportScale);
 			break;
 
 		case EGLTFJsonCameraType::Perspective:
-			Camera.Perspective = FGLTFCameraUtility::ConvertPerspective(DesiredView);
+			Camera.Perspective = FGLTFCameraUtility::ConvertPerspective(DesiredView, Builder.ExportOptions->ExportScale);
 			break;
 
 		default:
@@ -207,7 +207,7 @@ FGLTFJsonLightIndex FGLTFLightComponentConverter::Add(FGLTFConvertBuilder& Build
 
 	if (const UPointLightComponent* PointLightComponent = Cast<UPointLightComponent>(LightComponent))
 	{
-		Light.Range = FGLTFConverterUtility::ConvertLength(PointLightComponent->AttenuationRadius);
+		Light.Range = FGLTFConverterUtility::ConvertLength(PointLightComponent->AttenuationRadius, Builder.ExportOptions->ExportScale);
 	}
 
 	if (const USpotLightComponent* SpotLightComponent = Cast<USpotLightComponent>(LightComponent))
diff --git a/Engine/Plugins/Enterprise/GLTFExporter/MigratedFromGit/Source/GLTFExporter/Private/Converters/GLTFVertexBufferConverters.cpp b/Engine/Plugins/Enterprise/GLTFExporter/MigratedFromGit/Source/GLTFExporter/Private/Converters/GLTFVertexBufferConverters.cpp
--- a/Engine/Plugins/Enterprise/GLTFExporter/MigratedFromGit/Source/GLTFExporter/Private/Converters/GLTFVertexBufferConverters.cpp
+++ b/Engine/Plugins/Enterprise/GLTFExporter/MigratedFromGit/Source/GLTFExporter/Private/Converters/GLTFVertexBufferConverters.cpp
@@ -19,7 +19,7 @@ FGLTFJsonAccessorIndex FGLTFPositionVertexBufferConverter::Add(FGLTFConvertBuild
 
 	for (uint32 VertexIndex = 0; VertexIndex < VertexCount; ++VertexIndex)
 	{
-		Positions[VertexIndex] = FGLTFConverterUtility::ConvertPosition(VertexBuffer->VertexPosition(VertexIndex));
+		Positions[VertexIndex] = FGLTFConverterUtility::ConvertPosition(VertexBuffer->VertexPosition(VertexIndex), Builder.ExportOptions->ExportScale);
 	}
 
 	// More accurate bounding box if based on raw vertex values
diff --git a/Engine/Plugins/Enterprise/GLTFExporter/MigratedFromGit/Source/GLTFExporter/Private/GLTFExportOptions.cpp b/Engine/Plugins/Enterprise/GLTFExporter/MigratedFromGit/Source/GLTFExporter/Private/GLTFExportOptions.cpp
--- a/Engine/Plugins/Enterprise/GLTFExporter/MigratedFromGit/Source/GLTFExporter/Private/GLTFExportOptions.cpp
+++ b/Engine/Plugins/Enterprise/GLTFExporter/MigratedFromGit/Source/GLTFExporter/Private/GLTFExportOptions.cpp
@@ -217,6 +217,7 @@ bool UGLTFExportOptions::CanEditChange(const FProperty* InProperty) const
 		PropertyName == TEXT("bExportVertexColors") ||
 		PropertyName == TEXT("DefaultLevelOfDetail") ||
 		PropertyName == TEXT("bExportLightmaps") ||
+		PropertyName == TEXT("ExportScale") ||
 		PropertyName == TEXT("bExportLights"))
 	{
 		return true;
