[GLTFExporter] feat(options): implement support for bExportCameras

commit ffe3edada63a6b8623b7de60e64172c33daefbdf
Date: Tue, 18 Aug 2020 08:36:00 +0200

    feat(options): implement support for bExportCameras
    (cherry picked from commit bcfb4161a5a04970bdf4d882e8647babe5560830​)

#skipci

diff --git a/Engine/Plugins/Enterprise/GLTFExporter/MigratedFromGit/Source/GLTFExporter/Private/Converters/GLTFLevelConverters.cpp b/Engine/Plugins/Enterprise/GLTFExporter/MigratedFromGit/Source/GLTFExporter/Private/Converters/GLTFLevelConverters.cpp
--- a/Engine/Plugins/Enterprise/GLTFExporter/MigratedFromGit/Source/GLTFExporter/Private/Converters/GLTFLevelConverters.cpp
+++ b/Engine/Plugins/Enterprise/GLTFExporter/MigratedFromGit/Source/GLTFExporter/Private/Converters/GLTFLevelConverters.cpp
@@ -48,12 +48,21 @@ FGLTFJsonNodeIndex FGLTFSceneComponentConverter::Add(FGLTFConvertBuilder& Builde
 	}
 	else if (const UCameraComponent* CameraComponent = Cast<UCameraComponent>(SceneComponent))
 	{
-		// TODO: conversion of camera direction should be done in separate converter
-		FGLTFJsonNode CameraNode;
-		CameraNode.Name = Owner->GetName(); // TODO: choose a more unique name if owner is not ACameraActor
-		CameraNode.Rotation = FGLTFConverterUtility::ConvertCameraDirection();
-		CameraNode.Camera = Builder.GetOrAddCamera(CameraComponent, CameraNode.Name);
-		Builder.AddChildNode(NodeIndex, CameraNode);
+		if (Builder.ExportOptions->bExportCameras)
+		{
+			// TODO: conversion of camera direction should be done in separate converter
+			FGLTFJsonNode CameraNode;
+			CameraNode.Name = Owner->GetName(); // TODO: choose a more unique name if owner is not ACameraActor
+			CameraNode.Rotation = FGLTFConverterUtility::ConvertCameraDirection();
+			CameraNode.Camera = Builder.GetOrAddCamera(CameraComponent, CameraNode.Name);
+			Builder.AddChildNode(NodeIndex, CameraNode);
+		}
+		else
+		{
+			Builder.AddWarningMessage(FString::Printf(
+				TEXT("Camera %s disabled by export options"),
+				*Owner->GetName()));	// TODO: choose a more unique name if owner is not ACameraActor
+		}
 	}
 	else if (const ULightComponent* LightComponent = Cast<ULightComponent>(SceneComponent))
 	{
diff --git a/Engine/Plugins/Enterprise/GLTFExporter/MigratedFromGit/Source/GLTFExporter/Private/GLTFExportOptions.cpp b/Engine/Plugins/Enterprise/GLTFExporter/MigratedFromGit/Source/GLTFExporter/Private/GLTFExportOptions.cpp
--- a/Engine/Plugins/Enterprise/GLTFExporter/MigratedFromGit/Source/GLTFExporter/Private/GLTFExportOptions.cpp
+++ b/Engine/Plugins/Enterprise/GLTFExporter/MigratedFromGit/Source/GLTFExporter/Private/GLTFExportOptions.cpp
@@ -218,7 +218,8 @@ bool UGLTFExportOptions::CanEditChange(const FProperty* InProperty) const
 		PropertyName == TEXT("DefaultLevelOfDetail") ||
 		PropertyName == TEXT("bExportLightmaps") ||
 		PropertyName == TEXT("ExportScale") ||
-		PropertyName == TEXT("bExportLights"))
+		PropertyName == TEXT("bExportLights") ||
+		PropertyName == TEXT("bExportCameras"))
 	{
 		return true;
 	}
