[GLTFExporter] feat(options): implement support for bExportHDRIBackdrops

commit 155808c3b5aebec86d1e57c0f0555c4262c41a53
Date: Tue, 18 Aug 2020 08:42:19 +0200

    feat(options): implement support for bExportHDRIBackdrops
    (cherry picked from commit a94b32cd1125d712ecbf59f9495182cc73be82b2​)

#skipci

diff --git a/Engine/Plugins/Enterprise/GLTFExporter/MigratedFromGit/Source/GLTFExporter/Private/Converters/GLTFLevelConverters.cpp b/Engine/Plugins/Enterprise/GLTFExporter/MigratedFromGit/Source/GLTFExporter/Private/Converters/GLTFLevelConverters.cpp
--- a/Engine/Plugins/Enterprise/GLTFExporter/MigratedFromGit/Source/GLTFExporter/Private/Converters/GLTFLevelConverters.cpp
+++ b/Engine/Plugins/Enterprise/GLTFExporter/MigratedFromGit/Source/GLTFExporter/Private/Converters/GLTFLevelConverters.cpp
@@ -121,8 +121,17 @@ FGLTFJsonNodeIndex FGLTFActorConverter::Add(FGLTFConvertBuilder& Builder, const
 	}
 	else if (FGLTFActorUtility::IsHDRIBackdropBlueprint(Blueprint))
 	{
-		FGLTFJsonNode& RootNode = Builder.GetNode(RootNodeIndex);
-		RootNode.Backdrop = Builder.GetOrAddBackdrop(Actor);
+		if (Builder.ExportOptions->bExportHDRIBackdrops)
+		{
+			FGLTFJsonNode& RootNode = Builder.GetNode(RootNodeIndex);
+			RootNode.Backdrop = Builder.GetOrAddBackdrop(Actor);
+		}
+		else
+		{
+			Builder.AddWarningMessage(FString::Printf(
+				TEXT("HDRIBackdrop %s disabled by export options"),
+				*Actor->GetName()));
+		}
 	}
 	else
 	{
diff --git a/Engine/Plugins/Enterprise/GLTFExporter/MigratedFromGit/Source/GLTFExporter/Private/GLTFExportOptions.cpp b/Engine/Plugins/Enterprise/GLTFExporter/MigratedFromGit/Source/GLTFExporter/Private/GLTFExportOptions.cpp
--- a/Engine/Plugins/Enterprise/GLTFExporter/MigratedFromGit/Source/GLTFExporter/Private/GLTFExportOptions.cpp
+++ b/Engine/Plugins/Enterprise/GLTFExporter/MigratedFromGit/Source/GLTFExporter/Private/GLTFExportOptions.cpp
@@ -219,7 +219,8 @@ bool UGLTFExportOptions::CanEditChange(const FProperty* InProperty) const
 		PropertyName == TEXT("bExportLightmaps") ||
 		PropertyName == TEXT("ExportScale") ||
 		PropertyName == TEXT("bExportLights") ||
-		PropertyName == TEXT("bExportCameras"))
+		PropertyName == TEXT("bExportCameras") ||
+		PropertyName == TEXT("bExportHDRIBackdrops"))
 	{
 		return true;
 	}
