[GLTFExporter] feat(options): implement support for bExportVariantSets

commit 387dd3da72c8864bab555c30469669e82d87dea1
Date: Tue, 18 Aug 2020 08:47:56 +0200

    feat(options): implement support for bExportVariantSets
    (cherry picked from commit 2edde956d1d2d060f958d7a74912fc418fe198f3​)

#skipci

diff --git a/Engine/Plugins/Enterprise/GLTFExporter/MigratedFromGit/Source/GLTFExporter/Private/Converters/GLTFLevelConverters.cpp b/Engine/Plugins/Enterprise/GLTFExporter/MigratedFromGit/Source/GLTFExporter/Private/Converters/GLTFLevelConverters.cpp
--- a/Engine/Plugins/Enterprise/GLTFExporter/MigratedFromGit/Source/GLTFExporter/Private/Converters/GLTFLevelConverters.cpp
+++ b/Engine/Plugins/Enterprise/GLTFExporter/MigratedFromGit/Source/GLTFExporter/Private/Converters/GLTFLevelConverters.cpp
@@ -162,10 +162,19 @@ FGLTFJsonSceneIndex FGLTFLevelConverter::Add(FGLTFConvertBuilder& Builder, const
 		// TODO: should a LevelVariantSet be exported even if not selected for export?
 		if (const ALevelVariantSetsActor *LevelVariantSetsActor = Cast<ALevelVariantSetsActor>(Actor))
 		{
-			const FGLTFJsonLevelVariantSetsIndex LevelVariantSetsIndex = Builder.GetOrAddLevelVariantSets(LevelVariantSetsActor);
-			if (LevelVariantSetsIndex != INDEX_NONE)
+			if (Builder.ExportOptions->bExportVariantSets)
 			{
-				Scene.LevelVariantSets.Add(LevelVariantSetsIndex);
+				const FGLTFJsonLevelVariantSetsIndex LevelVariantSetsIndex = Builder.GetOrAddLevelVariantSets(LevelVariantSetsActor);
+				if (LevelVariantSetsIndex != INDEX_NONE)
+				{
+					Scene.LevelVariantSets.Add(LevelVariantSetsIndex);
+				}
+			}
+			else
+			{
+				Builder.AddWarningMessage(FString::Printf(
+					TEXT("Level Variant Set %s disabled by export options"),
+					*LevelVariantSetsActor->GetName()));
 			}
 		}
 
diff --git a/Engine/Plugins/Enterprise/GLTFExporter/MigratedFromGit/Source/GLTFExporter/Private/GLTFExportOptions.cpp b/Engine/Plugins/Enterprise/GLTFExporter/MigratedFromGit/Source/GLTFExporter/Private/GLTFExportOptions.cpp
--- a/Engine/Plugins/Enterprise/GLTFExporter/MigratedFromGit/Source/GLTFExporter/Private/GLTFExportOptions.cpp
+++ b/Engine/Plugins/Enterprise/GLTFExporter/MigratedFromGit/Source/GLTFExporter/Private/GLTFExportOptions.cpp
@@ -220,7 +220,8 @@ bool UGLTFExportOptions::CanEditChange(const FProperty* InProperty) const
 		PropertyName == TEXT("ExportScale") ||
 		PropertyName == TEXT("bExportLights") ||
 		PropertyName == TEXT("bExportCameras") ||
-		PropertyName == TEXT("bExportHDRIBackdrops"))
+		PropertyName == TEXT("bExportHDRIBackdrops") ||
+		PropertyName == TEXT("bExportVariantSets"))
 	{
 		return true;
 	}
