[GLTFExporter] feat(options): implement support for bExportPreviewMesh

commit 974118acba4eb0573aa681de5807e7abadbb46f5
Date: Tue, 18 Aug 2020 08:54:39 +0200

    feat(options): implement support for bExportPreviewMesh
    (cherry picked from commit 59ec4bee8cd810e40b8adfcf852c8e99af92f7cb​)

#skipci

diff --git a/Engine/Plugins/Enterprise/GLTFExporter/MigratedFromGit/Source/GLTFExporter/Private/Exporters/GLTFMaterialExporter.cpp b/Engine/Plugins/Enterprise/GLTFExporter/MigratedFromGit/Source/GLTFExporter/Private/Exporters/GLTFMaterialExporter.cpp
--- a/Engine/Plugins/Enterprise/GLTFExporter/MigratedFromGit/Source/GLTFExporter/Private/Exporters/GLTFMaterialExporter.cpp
+++ b/Engine/Plugins/Enterprise/GLTFExporter/MigratedFromGit/Source/GLTFExporter/Private/Exporters/GLTFMaterialExporter.cpp
@@ -22,28 +22,31 @@ bool UGLTFMaterialExporter::AddObject(FGLTFContainerBuilder& Builder, const UObj
 		return false;
 	}
 
-	const UStaticMesh* PreviewMesh = FGLTFExporterUtility::GetPreviewMesh(Material);
-	if (PreviewMesh != nullptr)
+	if (Builder.ExportOptions->bExportPreviewMesh)
 	{
-		const FGLTFJsonMeshIndex MeshIndex = Builder.GetOrAddMesh(PreviewMesh, -1, nullptr, { Material }, "PreviewMesh");
-
-		FGLTFJsonNode Node;
-		Node.Mesh = MeshIndex;
-		const FGLTFJsonNodeIndex NodeIndex = Builder.AddNode(Node);
-
-		FGLTFJsonScene Scene;
-		Scene.Nodes.Add(NodeIndex);
-		const FGLTFJsonSceneIndex SceneIndex = Builder.AddScene(Scene);
-
-		Builder.DefaultScene = SceneIndex;
-	}
-	else
-	{
-		Builder.AddErrorMessage(FString::Printf(
-			TEXT("Failed to export material %s because of missing preview-mesh"),
-			*Material->GetName()));
-
-		return false;
+		const UStaticMesh* PreviewMesh = FGLTFExporterUtility::GetPreviewMesh(Material);
+		if (PreviewMesh != nullptr)
+		{
+			const FGLTFJsonMeshIndex MeshIndex = Builder.GetOrAddMesh(PreviewMesh, -1, nullptr, { Material }, "PreviewMesh");
+
+			FGLTFJsonNode Node;
+			Node.Mesh = MeshIndex;
+			const FGLTFJsonNodeIndex NodeIndex = Builder.AddNode(Node);
+
+			FGLTFJsonScene Scene;
+			Scene.Nodes.Add(NodeIndex);
+			const FGLTFJsonSceneIndex SceneIndex = Builder.AddScene(Scene);
+
+			Builder.DefaultScene = SceneIndex;
+		}
+		else
+		{
+			Builder.AddErrorMessage(FString::Printf(
+                TEXT("Failed to export material %s because of missing preview-mesh"),
+                *Material->GetName()));
+
+			return false;
+		}
 	}
 
 	return true;
diff --git a/Engine/Plugins/Enterprise/GLTFExporter/MigratedFromGit/Source/GLTFExporter/Private/GLTFExportOptions.cpp b/Engine/Plugins/Enterprise/GLTFExporter/MigratedFromGit/Source/GLTFExporter/Private/GLTFExportOptions.cpp
--- a/Engine/Plugins/Enterprise/GLTFExporter/MigratedFromGit/Source/GLTFExporter/Private/GLTFExportOptions.cpp
+++ b/Engine/Plugins/Enterprise/GLTFExporter/MigratedFromGit/Source/GLTFExporter/Private/GLTFExportOptions.cpp
@@ -221,7 +221,8 @@ bool UGLTFExportOptions::CanEditChange(const FProperty* InProperty) const
 		PropertyName == TEXT("bExportLights") ||
 		PropertyName == TEXT("bExportCameras") ||
 		PropertyName == TEXT("bExportHDRIBackdrops") ||
-		PropertyName == TEXT("bExportVariantSets"))
+		PropertyName == TEXT("bExportVariantSets") ||
+		PropertyName == TEXT("bExportPreviewMesh"))
 	{
 		return true;
 	}
