[GLTFExporter] refactor(options): use GET_MEMBER_NAME_CHECKED() instead of raw strings when checking property-names

commit 80bee3d1bf2d15d76a445fa4f1020a5579bc66b3
Date: Wed, 19 Aug 2020 08:13:52 +0200

    refactor(options): use GET_MEMBER_NAME_CHECKED() instead of raw strings when checking property-names
    (cherry picked from commit a9e8a63b87a4ff027fecbe96dabb8992cbc98769​)

#skipci

diff --git a/Engine/Plugins/Enterprise/GLTFExporter/MigratedFromGit/Source/GLTFExporter/Private/GLTFExportOptions.cpp b/Engine/Plugins/Enterprise/GLTFExporter/MigratedFromGit/Source/GLTFExporter/Private/GLTFExportOptions.cpp
--- a/Engine/Plugins/Enterprise/GLTFExporter/MigratedFromGit/Source/GLTFExporter/Private/GLTFExportOptions.cpp
+++ b/Engine/Plugins/Enterprise/GLTFExporter/MigratedFromGit/Source/GLTFExporter/Private/GLTFExportOptions.cpp
@@ -208,21 +208,21 @@ void UGLTFExportOptions::FillOptions(bool bBatchMode, bool bShowOptionDialog, co
 
 bool UGLTFExportOptions::CanEditChange(const FProperty* InProperty) const
 {
-	const FString PropertyName = InProperty->GetName();
-
-	if (PropertyName == TEXT("bExportUnlitMaterials") ||
-		PropertyName == TEXT("bExportClearCoatMaterials") ||
-		PropertyName == TEXT("bBakeMaterialInputs") ||
-		PropertyName == TEXT("DefaultMaterialBakeSize") ||
-		PropertyName == TEXT("bExportVertexColors") ||
-		PropertyName == TEXT("DefaultLevelOfDetail") ||
-		PropertyName == TEXT("bExportLightmaps") ||
-		PropertyName == TEXT("ExportScale") ||
-		PropertyName == TEXT("bExportLights") ||
-		PropertyName == TEXT("bExportCameras") ||
-		PropertyName == TEXT("bExportHDRIBackdrops") ||
-		PropertyName == TEXT("bExportVariantSets") ||
-		PropertyName == TEXT("bExportPreviewMesh"))
+	const FName PropertyFName = InProperty->GetFName();
+
+	if (PropertyFName == GET_MEMBER_NAME_CHECKED(ThisClass, bExportUnlitMaterials) ||
+		PropertyFName == GET_MEMBER_NAME_CHECKED(ThisClass, bExportClearCoatMaterials) ||
+		PropertyFName == GET_MEMBER_NAME_CHECKED(ThisClass, bBakeMaterialInputs) ||
+		PropertyFName == GET_MEMBER_NAME_CHECKED(ThisClass, DefaultMaterialBakeSize) ||
+		PropertyFName == GET_MEMBER_NAME_CHECKED(ThisClass, bExportVertexColors) ||
+		PropertyFName == GET_MEMBER_NAME_CHECKED(ThisClass, DefaultLevelOfDetail) ||
+		PropertyFName == GET_MEMBER_NAME_CHECKED(ThisClass, bExportLightmaps) ||
+		PropertyFName == GET_MEMBER_NAME_CHECKED(ThisClass, ExportScale) ||
+		PropertyFName == GET_MEMBER_NAME_CHECKED(ThisClass, bExportLights) ||
+		PropertyFName == GET_MEMBER_NAME_CHECKED(ThisClass, bExportCameras) ||
+		PropertyFName == GET_MEMBER_NAME_CHECKED(ThisClass, bExportHDRIBackdrops) ||
+		PropertyFName == GET_MEMBER_NAME_CHECKED(ThisClass, bExportVariantSets) ||
+		PropertyFName == GET_MEMBER_NAME_CHECKED(ThisClass, bExportPreviewMesh))
 	{
 		return true;
 	}
