[GLTFExporter] refactor(options): change wording of warning-messages for when baking is disabled

commit 71696b76aa723fce5ac313f96ccd349af1a711f1
Date: Wed, 19 Aug 2020 08:25:30 +0200

    refactor(options): change wording of warning-messages for when baking is disabled
    (cherry picked from commit d47c8e5511f929701087c9504f946b887bae221d​)

#skipci

diff --git a/Engine/Plugins/Enterprise/GLTFExporter/MigratedFromGit/Source/GLTFExporter/Private/Converters/GLTFMaterialConverters.cpp b/Engine/Plugins/Enterprise/GLTFExporter/MigratedFromGit/Source/GLTFExporter/Private/Converters/GLTFMaterialConverters.cpp
--- a/Engine/Plugins/Enterprise/GLTFExporter/MigratedFromGit/Source/GLTFExporter/Private/Converters/GLTFMaterialConverters.cpp
+++ b/Engine/Plugins/Enterprise/GLTFExporter/MigratedFromGit/Source/GLTFExporter/Private/Converters/GLTFMaterialConverters.cpp
@@ -239,7 +239,7 @@ bool FGLTFMaterialConverter::TryGetBaseColorAndOpacity(FGLTFConvertBuilder& Buil
 	if (!Builder.ExportOptions->bBakeMaterialInputs)
 	{
 		Builder.AddWarningMessage(FString::Printf(
-			TEXT("Properties %s and %s for material %s will use defaults since baking is disabled by export options"),
+			TEXT("%s and %s for material %s need to baked, but material baking is disabled by export options"),
 			FGLTFMaterialUtility::GetPropertyName(BaseColorProperty),
 			FGLTFMaterialUtility::GetPropertyName(OpacityProperty),
 			*Material->GetName()));
@@ -404,7 +404,7 @@ bool FGLTFMaterialConverter::TryGetMetallicAndRoughness(FGLTFConvertBuilder& Bui
 	if (!Builder.ExportOptions->bBakeMaterialInputs)
 	{
 		Builder.AddWarningMessage(FString::Printf(
-			TEXT("Properties %s and %s for material %s will use defaults since baking is disabled by export options"),
+			TEXT("%s and %s for material %s need to baked, but material baking is disabled by export options"),
 			FGLTFMaterialUtility::GetPropertyName(MetallicProperty),
 			FGLTFMaterialUtility::GetPropertyName(RoughnessProperty),
 			*Material->GetName()));
@@ -568,7 +568,7 @@ bool FGLTFMaterialConverter::TryGetClearCoatRoughness(FGLTFConvertBuilder& Build
 	if (!Builder.ExportOptions->bBakeMaterialInputs)
 	{
 		Builder.AddWarningMessage(FString::Printf(
-			TEXT("Properties %s and %s for material %s will use defaults since baking is disabled by export options"),
+			TEXT("%s and %s for material %s need to baked, but material baking is disabled by export options"),
 			FGLTFMaterialUtility::GetPropertyName(IntensityProperty),
 			FGLTFMaterialUtility::GetPropertyName(RoughnessProperty),
 			*Material->GetName()));
@@ -712,7 +712,7 @@ bool FGLTFMaterialConverter::TryGetEmissive(FGLTFConvertBuilder& Builder, FGLTFJ
 	if (!Builder.ExportOptions->bBakeMaterialInputs)
 	{
 		Builder.AddWarningMessage(FString::Printf(
-			TEXT("Property %s for material %s will use defaults since baking is disabled by export options"),
+			TEXT("%s for material %s needs to baked, but material baking is disabled by export options"),
 			FGLTFMaterialUtility::GetPropertyName(EmissiveProperty),
 			*Material->GetName()));
 		return false;
@@ -1107,7 +1107,7 @@ bool FGLTFMaterialConverter::TryGetBakedMaterialProperty(FGLTFConvertBuilder& Bu
 	if (!Builder.ExportOptions->bBakeMaterialInputs)
 	{
 		Builder.AddWarningMessage(FString::Printf(
-			TEXT("Property %s for material %s will use defaults since baking is disabled by export options"),
+			TEXT("%s for material %s needs to baked, but material baking is disabled by export options"),
 			FGLTFMaterialUtility::GetPropertyName(Property),
 			*Material->GetName()));
 		return false;
@@ -1135,7 +1135,7 @@ bool FGLTFMaterialConverter::TryGetBakedMaterialProperty(FGLTFConvertBuilder& Bu
 	if (!Builder.ExportOptions->bBakeMaterialInputs)
 	{
 		Builder.AddWarningMessage(FString::Printf(
-			TEXT("Property %s for material %s will use defaults since baking is disabled by export options"),
+			TEXT("%s for material %s needs to baked, but material baking is disabled by export options"),
 			FGLTFMaterialUtility::GetPropertyName(Property),
 			*Material->GetName()));
 		return false;
@@ -1163,7 +1163,7 @@ inline bool FGLTFMaterialConverter::TryGetBakedMaterialProperty(FGLTFConvertBuil
 	if (!Builder.ExportOptions->bBakeMaterialInputs)
 	{
 		Builder.AddWarningMessage(FString::Printf(
-			TEXT("Property %s for material %s will use defaults since baking is disabled by export options"),
+			TEXT("%s for material %s needs to baked, but material baking is disabled by export options"),
 			FGLTFMaterialUtility::GetPropertyName(Property),
 			*Material->GetName()));
 		return false;
@@ -1191,7 +1191,7 @@ bool FGLTFMaterialConverter::TryGetBakedMaterialProperty(FGLTFConvertBuilder& Bu
 	if (!Builder.ExportOptions->bBakeMaterialInputs)
 	{
 		Builder.AddWarningMessage(FString::Printf(
-			TEXT("Property %s for material %s will use defaults since baking is disabled by export options"),
+			TEXT("%s for material %s needs to baked, but material baking is disabled by export options"),
 			FGLTFMaterialUtility::GetPropertyName(Property),
 			*Material->GetName()));
 		return false;
