[GLTFExporter] refactor(options): move ShouldExportLight into UGLTFExportOptions

commit d1df76a6e82a680ef7f8867ebcbcccbbf14b2f11
Date: Wed, 19 Aug 2020 09:45:10 +0200

    refactor(options): move ShouldExportLight into UGLTFExportOptions
    (cherry picked from commit 3c7e8400bd300f4ef9bda5e96110ec6744afe2f0​)

#skipci

diff --git a/Engine/Plugins/Enterprise/GLTFExporter/MigratedFromGit/Source/GLTFExporter/Private/Converters/GLTFLevelConverters.cpp b/Engine/Plugins/Enterprise/GLTFExporter/MigratedFromGit/Source/GLTFExporter/Private/Converters/GLTFLevelConverters.cpp
--- a/Engine/Plugins/Enterprise/GLTFExporter/MigratedFromGit/Source/GLTFExporter/Private/Converters/GLTFLevelConverters.cpp
+++ b/Engine/Plugins/Enterprise/GLTFExporter/MigratedFromGit/Source/GLTFExporter/Private/Converters/GLTFLevelConverters.cpp
@@ -66,7 +66,7 @@ FGLTFJsonNodeIndex FGLTFSceneComponentConverter::Add(FGLTFConvertBuilder& Builde
 	}
 	else if (const ULightComponent* LightComponent = Cast<ULightComponent>(SceneComponent))
 	{
-		if (ShouldExportLight(Builder.ExportOptions->bExportLights, LightComponent->Mobility))
+		if (Builder.ExportOptions->ShouldExportLight(LightComponent->Mobility))
 		{
 			// TODO: conversion of light direction should be done in separate converter
 			FGLTFJsonNode LightNode;
@@ -88,21 +88,6 @@ FGLTFJsonNodeIndex FGLTFSceneComponentConverter::Add(FGLTFConvertBuilder& Builde
 	return NodeIndex;
 }
 
-bool FGLTFSceneComponentConverter::ShouldExportLight(EGLTFExporterLightMobility Options, EComponentMobility::Type Mobility)
-{
-	switch (Options)
-	{
-		case EGLTFExporterLightMobility::All:
-			return true;
-		case EGLTFExporterLightMobility::MovableAndStationary:
-			return Mobility == EComponentMobility::Movable || Mobility == EComponentMobility::Stationary;
-		case EGLTFExporterLightMobility::MovableOnly:
-			return Mobility == EComponentMobility::Movable;
-		default:
-			return false;
-	}
-}
-
 FGLTFJsonNodeIndex FGLTFActorConverter::Add(FGLTFConvertBuilder& Builder, const FString& Name, const AActor* Actor)
 {
 	if (Builder.bSelectedActorsOnly && !Actor->IsSelected())
diff --git a/Engine/Plugins/Enterprise/GLTFExporter/MigratedFromGit/Source/GLTFExporter/Private/Converters/GLTFLevelConverters.h b/Engine/Plugins/Enterprise/GLTFExporter/MigratedFromGit/Source/GLTFExporter/Private/Converters/GLTFLevelConverters.h
--- a/Engine/Plugins/Enterprise/GLTFExporter/MigratedFromGit/Source/GLTFExporter/Private/Converters/GLTFLevelConverters.h
+++ b/Engine/Plugins/Enterprise/GLTFExporter/MigratedFromGit/Source/GLTFExporter/Private/Converters/GLTFLevelConverters.h
@@ -10,8 +10,6 @@
 class FGLTFSceneComponentConverter final : public TGLTFConverter<FGLTFJsonNodeIndex, const USceneComponent*>
 {
 	FGLTFJsonNodeIndex Add(FGLTFConvertBuilder& Builder, const FString& Name, const USceneComponent* SceneComponent) override;
-
-	static bool ShouldExportLight(EGLTFExporterLightMobility Options, EComponentMobility::Type Mobility);
 };
 
 class FGLTFActorConverter final : public TGLTFConverter<FGLTFJsonNodeIndex, const AActor*>
diff --git a/Engine/Plugins/Enterprise/GLTFExporter/MigratedFromGit/Source/GLTFExporter/Private/GLTFExportOptions.cpp b/Engine/Plugins/Enterprise/GLTFExporter/MigratedFromGit/Source/GLTFExporter/Private/GLTFExportOptions.cpp
--- a/Engine/Plugins/Enterprise/GLTFExporter/MigratedFromGit/Source/GLTFExporter/Private/GLTFExportOptions.cpp
+++ b/Engine/Plugins/Enterprise/GLTFExporter/MigratedFromGit/Source/GLTFExporter/Private/GLTFExportOptions.cpp
@@ -237,3 +237,18 @@ FIntPoint UGLTFExportOptions::GetDefaultMaterialBakeSize() const
 	const int32 Size = 1 << static_cast<int>(DefaultMaterialBakeSize);
 	return { Size, Size };
 }
+
+bool UGLTFExportOptions::ShouldExportLight(EComponentMobility::Type LightMobility) const
+{
+	switch (bExportLights)
+	{
+		case EGLTFExporterLightMobility::All:
+			return true;
+		case EGLTFExporterLightMobility::MovableAndStationary:
+			return LightMobility == EComponentMobility::Movable || LightMobility == EComponentMobility::Stationary;
+		case EGLTFExporterLightMobility::MovableOnly:
+			return LightMobility == EComponentMobility::Movable;
+		default:
+			return false;
+	}
+}
diff --git a/Engine/Plugins/Enterprise/GLTFExporter/MigratedFromGit/Source/GLTFExporter/Public/GLTFExportOptions.h b/Engine/Plugins/Enterprise/GLTFExporter/MigratedFromGit/Source/GLTFExporter/Public/GLTFExportOptions.h
--- a/Engine/Plugins/Enterprise/GLTFExporter/MigratedFromGit/Source/GLTFExporter/Public/GLTFExportOptions.h
+++ b/Engine/Plugins/Enterprise/GLTFExporter/MigratedFromGit/Source/GLTFExporter/Public/GLTFExportOptions.h
@@ -169,4 +169,6 @@ public:
 	virtual bool CanEditChange(const FProperty* InProperty) const override;
 
 	FIntPoint GetDefaultMaterialBakeSize() const;
+
+	bool ShouldExportLight(EComponentMobility::Type LightMobility) const;
 };
