[GLTFExporter] refactor(options): remove unnecessary TODOs and line-breaks for option-related warnings

commit 85058816b99ad2425ce3694f0860c26727323d29
Date: Wed, 19 Aug 2020 09:49:33 +0200

    refactor(options): remove unnecessary TODOs and line-breaks for option-related warnings
    (cherry picked from commit d0c200f6bb790b37984441ea10216d0b813f78e7​)

#skipci

diff --git a/Engine/Plugins/Enterprise/GLTFExporter/MigratedFromGit/Source/GLTFExporter/Private/Converters/GLTFLevelConverters.cpp b/Engine/Plugins/Enterprise/GLTFExporter/MigratedFromGit/Source/GLTFExporter/Private/Converters/GLTFLevelConverters.cpp
--- a/Engine/Plugins/Enterprise/GLTFExporter/MigratedFromGit/Source/GLTFExporter/Private/Converters/GLTFLevelConverters.cpp
+++ b/Engine/Plugins/Enterprise/GLTFExporter/MigratedFromGit/Source/GLTFExporter/Private/Converters/GLTFLevelConverters.cpp
@@ -59,9 +59,7 @@ FGLTFJsonNodeIndex FGLTFSceneComponentConverter::Add(FGLTFConvertBuilder& Builde
 		}
 		else
 		{
-			Builder.AddWarningMessage(FString::Printf(
-				TEXT("Camera %s disabled by export options"),
-				*Owner->GetName()));	// TODO: choose a more unique name if owner is not ACameraActor
+			Builder.AddWarningMessage(FString::Printf(TEXT("Camera %s disabled by export options"), *Owner->GetName()));
 		}
 	}
 	else if (const ULightComponent* LightComponent = Cast<ULightComponent>(SceneComponent))
@@ -77,9 +75,7 @@ FGLTFJsonNodeIndex FGLTFSceneComponentConverter::Add(FGLTFConvertBuilder& Builde
 		}
 		else
 		{
-			Builder.AddWarningMessage(FString::Printf(
-				TEXT("Light %s disabled by export options"),
-				*Owner->GetName()));	// TODO: choose a more unique name if owner is not ALight
+			Builder.AddWarningMessage(FString::Printf(TEXT("Light %s disabled by export options"), *Owner->GetName()));
 		}
 	}
 
