[GLTFExporter] refactor: rename ExportScale to (more general) ConversionScale in functions that convert units

commit ebaa78f9681fa952c991114df373afa7b905072d
Date: Wed, 19 Aug 2020 09:54:10 +0200

    refactor: rename ExportScale to (more general) ConversionScale in functions that convert units
    (cherry picked from commit e5da288f409f364213d9c66725418925ef7a8648​)

#skipci

diff --git a/Engine/Plugins/Enterprise/GLTFExporter/MigratedFromGit/Source/GLTFExporter/Private/Converters/GLTFCameraUtility.cpp b/Engine/Plugins/Enterprise/GLTFExporter/MigratedFromGit/Source/GLTFExporter/Private/Converters/GLTFCameraUtility.cpp
--- a/Engine/Plugins/Enterprise/GLTFExporter/MigratedFromGit/Source/GLTFExporter/Private/Converters/GLTFCameraUtility.cpp
+++ b/Engine/Plugins/Enterprise/GLTFExporter/MigratedFromGit/Source/GLTFExporter/Private/Converters/GLTFCameraUtility.cpp
@@ -4,23 +4,23 @@
 #include "Converters/GLTFConverterUtility.h"
 #include "Camera/CameraTypes.h"
 
-FGLTFJsonOrthographic FGLTFCameraUtility::ConvertOrthographic(const FMinimalViewInfo& View, const float ExportScale)
+FGLTFJsonOrthographic FGLTFCameraUtility::ConvertOrthographic(const FMinimalViewInfo& View, const float ConversionScale)
 {
 	FGLTFJsonOrthographic Orthographic;
 
 	// NOTE: it goes against the glTF standard to not define a XMag, but the gltf viewer can handle it
 	if (View.bConstrainAspectRatio)
 	{
-		Orthographic.XMag = FGLTFConverterUtility::ConvertLength(View.OrthoWidth, ExportScale);
+		Orthographic.XMag = FGLTFConverterUtility::ConvertLength(View.OrthoWidth, ConversionScale);
 	}
 
-	Orthographic.YMag = FGLTFConverterUtility::ConvertLength(View.OrthoWidth / View.AspectRatio, ExportScale); // TODO: is this correct?
-	Orthographic.ZFar = FGLTFConverterUtility::ConvertLength(View.OrthoFarClipPlane, ExportScale);
-	Orthographic.ZNear = FGLTFConverterUtility::ConvertLength(View.OrthoNearClipPlane, ExportScale);
+	Orthographic.YMag = FGLTFConverterUtility::ConvertLength(View.OrthoWidth / View.AspectRatio, ConversionScale); // TODO: is this correct?
+	Orthographic.ZFar = FGLTFConverterUtility::ConvertLength(View.OrthoFarClipPlane, ConversionScale);
+	Orthographic.ZNear = FGLTFConverterUtility::ConvertLength(View.OrthoNearClipPlane, ConversionScale);
 	return Orthographic;
 }
 
-FGLTFJsonPerspective FGLTFCameraUtility::ConvertPerspective(const FMinimalViewInfo& View, const float ExportScale)
+FGLTFJsonPerspective FGLTFCameraUtility::ConvertPerspective(const FMinimalViewInfo& View, const float ConversionScale)
 {
 	FGLTFJsonPerspective Perspective;
 
@@ -32,8 +32,8 @@ FGLTFJsonPerspective FGLTFCameraUtility::ConvertPerspective(const FMinimalViewIn
 	Perspective.YFov = ConvertFieldOfView(View);
 
 	// NOTE: even thought ZFar is optional, if we don't set it, then most gltf viewers won't handle it well.
-	Perspective.ZFar = FGLTFConverterUtility::ConvertLength(WORLD_MAX, ExportScale); // TODO: Unreal doesn't have max draw distance per view?
-	Perspective.ZNear = FGLTFConverterUtility::ConvertLength(GNearClippingPlane, ExportScale);
+	Perspective.ZFar = FGLTFConverterUtility::ConvertLength(WORLD_MAX, ConversionScale); // TODO: Unreal doesn't have max draw distance per view?
+	Perspective.ZNear = FGLTFConverterUtility::ConvertLength(GNearClippingPlane, ConversionScale);
 	return Perspective;
 }
 
diff --git a/Engine/Plugins/Enterprise/GLTFExporter/MigratedFromGit/Source/GLTFExporter/Private/Converters/GLTFCameraUtility.h b/Engine/Plugins/Enterprise/GLTFExporter/MigratedFromGit/Source/GLTFExporter/Private/Converters/GLTFCameraUtility.h
--- a/Engine/Plugins/Enterprise/GLTFExporter/MigratedFromGit/Source/GLTFExporter/Private/Converters/GLTFCameraUtility.h
+++ b/Engine/Plugins/Enterprise/GLTFExporter/MigratedFromGit/Source/GLTFExporter/Private/Converters/GLTFCameraUtility.h
@@ -6,8 +6,8 @@
 
 struct FGLTFCameraUtility
 {
-	static FGLTFJsonOrthographic ConvertOrthographic(const FMinimalViewInfo& View, const float ExportScale);
-	static FGLTFJsonPerspective ConvertPerspective(const FMinimalViewInfo& View, const float ExportScale);
+	static FGLTFJsonOrthographic ConvertOrthographic(const FMinimalViewInfo& View, const float ConversionScale);
+	static FGLTFJsonPerspective ConvertPerspective(const FMinimalViewInfo& View, const float ConversionScale);
 
 	static float ConvertFieldOfView(const FMinimalViewInfo& View);
 };
diff --git a/Engine/Plugins/Enterprise/GLTFExporter/MigratedFromGit/Source/GLTFExporter/Private/Converters/GLTFConverterUtility.h b/Engine/Plugins/Enterprise/GLTFExporter/MigratedFromGit/Source/GLTFExporter/Private/Converters/GLTFConverterUtility.h
--- a/Engine/Plugins/Enterprise/GLTFExporter/MigratedFromGit/Source/GLTFExporter/Private/Converters/GLTFConverterUtility.h
+++ b/Engine/Plugins/Enterprise/GLTFExporter/MigratedFromGit/Source/GLTFExporter/Private/Converters/GLTFConverterUtility.h
@@ -14,9 +14,9 @@
 
 struct FGLTFConverterUtility
 {
-	static float ConvertLength(const float Length, const float ExportScale)
+	static float ConvertLength(const float Length, const float ConversionScale)
 	{
-		return Length * ExportScale;
+		return Length * ConversionScale;
 	}
 
 	static FGLTFJsonVector3 ConvertVector(const FVector& Vector)
@@ -26,9 +26,9 @@ struct FGLTFConverterUtility
 		return { Vector.X, Vector.Z, Vector.Y };
 	}
 
-	static FGLTFJsonVector3 ConvertPosition(const FVector& Position, const float ExportScale)
+	static FGLTFJsonVector3 ConvertPosition(const FVector& Position, const float ConversionScale)
 	{
-		return ConvertVector(Position * ExportScale);
+		return ConvertVector(Position * ConversionScale);
 	}
 
 	static FGLTFJsonVector3 ConvertScale(const FVector& Scale)
