[GLTFExporter] fix(options): make hash of exported mesh unaffected by color-overrides when not exporting vertex-colors

commit 6ab0c964bfcfb4a05cb2970ee0eeadcfe349d343
Date: Thu, 20 Aug 2020 10:01:16 +0200

    fix(options): make hash of exported mesh unaffected by color-overrides when not exporting vertex-colors
    (cherry picked from commit c733cec6623890171577681d437ff5f55aa5c33f​)

#skipci

diff --git a/Engine/Plugins/Enterprise/GLTFExporter/MigratedFromGit/Source/GLTFExporter/Private/Builders/GLTFConvertBuilder.cpp b/Engine/Plugins/Enterprise/GLTFExporter/MigratedFromGit/Source/GLTFExporter/Private/Builders/GLTFConvertBuilder.cpp
--- a/Engine/Plugins/Enterprise/GLTFExporter/MigratedFromGit/Source/GLTFExporter/Private/Builders/GLTFConvertBuilder.cpp
+++ b/Engine/Plugins/Enterprise/GLTFExporter/MigratedFromGit/Source/GLTFExporter/Private/Builders/GLTFConvertBuilder.cpp
@@ -111,6 +111,11 @@ FGLTFJsonMeshIndex FGLTFConvertBuilder::GetOrAddMesh(const UStaticMesh* StaticMe
 		LODIndex = ExportOptions->DefaultLevelOfDetail;
 	}
 
+	if (!ExportOptions->bExportVertexColors)
+	{
+		OverrideVertexColors = nullptr;
+	}
+
 	return StaticMeshConverter.GetOrAdd(*this, DesiredName.IsEmpty() ? FGLTFBuilderUtility::GetLODName(StaticMesh, LODIndex) : DesiredName, StaticMesh, LODIndex, OverrideVertexColors, OverrideMaterials);
 }
 
@@ -161,6 +166,11 @@ FGLTFJsonMeshIndex FGLTFConvertBuilder::GetOrAddMesh(const USkeletalMesh* Skelet
 		LODIndex = ExportOptions->DefaultLevelOfDetail;
 	}
 
+	if (!ExportOptions->bExportVertexColors)
+	{
+		OverrideVertexColors = nullptr;
+	}
+
 	return SkeletalMeshConverter.GetOrAdd(*this, DesiredName.IsEmpty() ? FGLTFBuilderUtility::GetLODName(SkeletalMesh, LODIndex) : DesiredName, SkeletalMesh, LODIndex, OverrideVertexColors, OverrideSkinWeights, OverrideMaterials);
 }
 
