[GLTFExporter] refactor: add TODOs for warnings in material converter

commit 33f5bfd6ae8e35537d47b4875a14b7e726827228
Date: Sun, 23 Aug 2020 15:10:24 +0200

    refactor: add TODOs for warnings in material converter
    (cherry picked from commit 46df6398072558ecf567958ccfea69d277b8cd30​)

#skipci

diff --git a/Engine/Plugins/Enterprise/GLTFExporter/MigratedFromGit/Source/GLTFExporter/Private/Converters/GLTFMaterialConverters.cpp b/Engine/Plugins/Enterprise/GLTFExporter/MigratedFromGit/Source/GLTFExporter/Private/Converters/GLTFMaterialConverters.cpp
--- a/Engine/Plugins/Enterprise/GLTFExporter/MigratedFromGit/Source/GLTFExporter/Private/Converters/GLTFMaterialConverters.cpp
+++ b/Engine/Plugins/Enterprise/GLTFExporter/MigratedFromGit/Source/GLTFExporter/Private/Converters/GLTFMaterialConverters.cpp
@@ -45,6 +45,20 @@ namespace
 
 FGLTFJsonMaterialIndex FGLTFMaterialConverter::Add(FGLTFConvertBuilder& Builder, const FString& Name, const UMaterialInterface* Material)
 {
+	{
+		const UMaterial* ParentMaterial = Material->GetMaterial();
+
+		if (ParentMaterial->MaterialDomain != MD_Surface)
+		{
+			// TODO: report warning (non-surface materials not supported, will be treated as surface)
+		}
+
+		if (!ParentMaterial->bTangentSpaceNormal)
+		{
+			// TODO: report warning (world-space normals not supported, will be treated as tangent-space)
+		}
+	}
+
 	FGLTFJsonMaterial JsonMaterial;
 	JsonMaterial.Name = Name;
 
