[GLTFExporter] fix: use the same hash for meshes with default override-materials as when no overrides are specified

commit 5fda73593fe35f2a8c783f9c9419f6dcdac87e80
Date: Wed, 12 Aug 2020 14:51:03 +0200

    fix: use the same hash for meshes with default override-materials as when no overrides are specified
    (cherry picked from commit cb9c58ddc23c36083a5ce9f47426994180566f0f​)

#skipci

diff --git a/Engine/Plugins/Enterprise/GLTFExporter/MigratedFromGit/Source/GLTFExporter/Private/Builders/GLTFConvertBuilder.cpp b/Engine/Plugins/Enterprise/GLTFExporter/MigratedFromGit/Source/GLTFExporter/Private/Builders/GLTFConvertBuilder.cpp
--- a/Engine/Plugins/Enterprise/GLTFExporter/MigratedFromGit/Source/GLTFExporter/Private/Builders/GLTFConvertBuilder.cpp
+++ b/Engine/Plugins/Enterprise/GLTFExporter/MigratedFromGit/Source/GLTFExporter/Private/Builders/GLTFConvertBuilder.cpp
@@ -129,7 +129,14 @@ FGLTFJsonMeshIndex FGLTFConvertBuilder::GetOrAddMesh(const UStaticMeshComponent*
 	const UStaticMesh* StaticMesh = StaticMeshComponent->GetStaticMesh();
 	const int32 LODIndex = StaticMeshComponent->ForcedLodModel > 0 ? StaticMeshComponent->ForcedLodModel - 1 : ExportOptions->DefaultLevelOfDetail;
 	const FColorVertexBuffer* OverrideVertexColors = StaticMeshComponent->LODData.IsValidIndex(LODIndex) ? StaticMeshComponent->LODData[LODIndex].OverrideVertexColors : nullptr;
-	const FGLTFMaterialArray OverrideMaterials = FGLTFMaterialArray(StaticMeshComponent->OverrideMaterials);
+	FGLTFMaterialArray OverrideMaterials = FGLTFMaterialArray(StaticMeshComponent->OverrideMaterials);
+
+	// Clear unnecessary overrides in order to keep the hash identical to when no overrides are specified.
+	// Otherwise the exact same mesh can be exported twice, once with default overrides and once with no overrides.
+	if (OverrideMaterials == StaticMesh->StaticMaterials)
+	{
+		OverrideMaterials.Reset();
+	}
 
 	return GetOrAddMesh(StaticMesh, LODIndex, OverrideVertexColors, OverrideMaterials, DesiredName);
 }
diff --git a/Engine/Plugins/Enterprise/GLTFExporter/MigratedFromGit/Source/GLTFExporter/Private/Converters/GLTFMaterialArray.h b/Engine/Plugins/Enterprise/GLTFExporter/MigratedFromGit/Source/GLTFExporter/Private/Converters/GLTFMaterialArray.h
--- a/Engine/Plugins/Enterprise/GLTFExporter/MigratedFromGit/Source/GLTFExporter/Private/Converters/GLTFMaterialArray.h
+++ b/Engine/Plugins/Enterprise/GLTFExporter/MigratedFromGit/Source/GLTFExporter/Private/Converters/GLTFMaterialArray.h
@@ -55,6 +55,24 @@ struct FGLTFMaterialArray : TArray<const UMaterialInterface*>
 		return *this;
 	}
 
+	bool operator==(const TArray<FStaticMaterial>& Other) const
+	{
+		if (Other.Num() != Num())
+		{
+			return false;
+		}
+
+		for (int32 Index = 0; Index < Num(); ++Index)
+		{
+			if (Other[Index].MaterialInterface != (*this)[Index])
+			{
+				return false;
+			}
+		}
+
+		return true;
+	}
+
 	bool operator==(const TArray& Other) const
 	{
 		return static_cast<const TArray&>(*this) == Other;
