[GLTFExporter] feat(variants): add support for mesh (static / skeletal) variants

commit 46ad587aa75d131b573ac49ac2dbb9173f6760e9
Date: Wed, 12 Aug 2020 14:55:14 +0200

    feat(variants): add support for mesh (static / skeletal) variants
    (cherry picked from commit 4acfed0a6e3e611ac1b7422e83ed7300ec89d938​)

#skipci

diff --git a/Engine/Plugins/Enterprise/GLTFExporter/MigratedFromGit/Source/GLTFExporter/Private/Converters/GLTFVariantSetConverters.cpp b/Engine/Plugins/Enterprise/GLTFExporter/MigratedFromGit/Source/GLTFExporter/Private/Converters/GLTFVariantSetConverters.cpp
--- a/Engine/Plugins/Enterprise/GLTFExporter/MigratedFromGit/Source/GLTFExporter/Private/Converters/GLTFVariantSetConverters.cpp
+++ b/Engine/Plugins/Enterprise/GLTFExporter/MigratedFromGit/Source/GLTFExporter/Private/Converters/GLTFVariantSetConverters.cpp
@@ -117,6 +117,8 @@ bool FGLTFLevelVariantSetsConverter::TryParseVariantBinding(FGLTFConvertBuilder&
 			Property->GetPropertyClass()->IsChildOf(FBoolProperty::StaticClass());
 
 		const bool bIsMaterialProperty = Property->IsA<UPropertyValueMaterial>();
+		const bool bIsStaticMeshProperty = Property->GetPropertyName() == TEXT("StaticMesh");
+		const bool bIsSkeletalMeshProperty = Property->GetPropertyName() == TEXT("SkeletalMesh");
 
 		if (bIsVisibilityProperty)
 		{
@@ -132,6 +134,20 @@ bool FGLTFLevelVariantSetsConverter::TryParseVariantBinding(FGLTFConvertBuilder&
 				bHasParsedAnyProperty = true;
 			}
 		}
+		else if (bIsStaticMeshProperty)
+		{
+			if (TryParseStaticMeshPropertyValue(Builder, OutVariant, Property))
+			{
+				bHasParsedAnyProperty = true;
+			}
+		}
+		else if (bIsSkeletalMeshProperty)
+		{
+			if (TryParseSkeletalMeshPropertyValue(Builder, OutVariant, Property))
+			{
+				bHasParsedAnyProperty = true;
+			}
+		}
 		else
 		{
 			// TODO: add support for more properties
@@ -265,6 +281,98 @@ bool FGLTFLevelVariantSetsConverter::TryParseMaterialPropertyValue(FGLTFConvertB
 	return true;
 }
 
+bool FGLTFLevelVariantSetsConverter::TryParseStaticMeshPropertyValue(FGLTFConvertBuilder& Builder, FGLTFJsonVariant& OutVariant, const UPropertyValue* Property) const
+{
+	if (Property->GetPropertyName() != TEXT("StaticMesh"))
+	{
+		Builder.AddErrorMessage(FString::Printf(
+			TEXT("Attempted to parse static mesh from an incompatible property. Context: %s"),
+			*GetLogContext(Property)));
+		return false;
+	}
+
+	const USceneComponent* Target = static_cast<USceneComponent*>(Property->GetPropertyParentContainerAddress());
+	if (Target == nullptr)
+	{
+		Builder.AddWarningMessage(FString::Printf(
+			TEXT("Target object for property is invalid, the property will be skipped. Context: %s"),
+			*GetLogContext(Property)));
+		return false;
+	}
+
+	if (Builder.bSelectedActorsOnly && !Target->IsSelected())
+	{
+		Builder.AddWarningMessage(FString::Printf(
+			TEXT("Target object for property is not selected for export, the property will be skipped. Context: %s"),
+			*GetLogContext(Property)));
+		return false;
+	}
+
+	const UStaticMesh* StaticMesh;
+
+	if (!TryGetPropertyValue(const_cast<UPropertyValue*>(Property), StaticMesh) || StaticMesh == nullptr)
+	{
+		Builder.AddWarningMessage(FString::Printf(
+			TEXT("Failed to parse recorded data for property, it will be skipped. Context: %s"),
+			*GetLogContext(Property)));
+		return false;
+	}
+
+	const FGLTFJsonNodeIndex NodeIndex = Builder.GetOrAddNode(Target);
+	const FGLTFJsonMeshIndex MeshIndex = Builder.GetOrAddMesh(StaticMesh);
+	FGLTFJsonVariantNodeProperties& NodeProperties = OutVariant.Nodes.FindOrAdd(NodeIndex);
+
+	NodeProperties.Node = NodeIndex;
+	NodeProperties.Mesh = MeshIndex;
+	return true;
+}
+
+bool FGLTFLevelVariantSetsConverter::TryParseSkeletalMeshPropertyValue(FGLTFConvertBuilder& Builder, FGLTFJsonVariant& OutVariant, const UPropertyValue* Property) const
+{
+	if (Property->GetPropertyName() != TEXT("SkeletalMesh"))
+	{
+		Builder.AddErrorMessage(FString::Printf(
+			TEXT("Attempted to parse skeletal mesh from an incompatible property. Context: %s"),
+			*GetLogContext(Property)));
+		return false;
+	}
+
+	const USceneComponent* Target = static_cast<USceneComponent*>(Property->GetPropertyParentContainerAddress());
+	if (Target == nullptr)
+	{
+		Builder.AddWarningMessage(FString::Printf(
+			TEXT("Target object for property is invalid, the property will be skipped. Context: %s"),
+			*GetLogContext(Property)));
+		return false;
+	}
+
+	if (Builder.bSelectedActorsOnly && !Target->IsSelected())
+	{
+		Builder.AddWarningMessage(FString::Printf(
+			TEXT("Target object for property is not selected for export, the property will be skipped. Context: %s"),
+			*GetLogContext(Property)));
+		return false;
+	}
+
+	const USkeletalMesh* SkeletalMesh;
+
+	if (!TryGetPropertyValue(const_cast<UPropertyValue*>(Property), SkeletalMesh) || SkeletalMesh == nullptr)
+	{
+		Builder.AddWarningMessage(FString::Printf(
+			TEXT("Failed to parse recorded data for property, it will be skipped. Context: %s"),
+			*GetLogContext(Property)));
+		return false;
+	}
+
+	const FGLTFJsonNodeIndex NodeIndex = Builder.GetOrAddNode(Target);
+	const FGLTFJsonMeshIndex MeshIndex = Builder.GetOrAddMesh(SkeletalMesh);
+	FGLTFJsonVariantNodeProperties& NodeProperties = OutVariant.Nodes.FindOrAdd(NodeIndex);
+
+	NodeProperties.Node = NodeIndex;
+	NodeProperties.Mesh = MeshIndex;
+	return true;
+}
+
 template<typename T>
 bool FGLTFLevelVariantSetsConverter::TryGetPropertyValue(UPropertyValue* Property, T& OutValue) const
 {
diff --git a/Engine/Plugins/Enterprise/GLTFExporter/MigratedFromGit/Source/GLTFExporter/Private/Converters/GLTFVariantSetConverters.h b/Engine/Plugins/Enterprise/GLTFExporter/MigratedFromGit/Source/GLTFExporter/Private/Converters/GLTFVariantSetConverters.h
--- a/Engine/Plugins/Enterprise/GLTFExporter/MigratedFromGit/Source/GLTFExporter/Private/Converters/GLTFVariantSetConverters.h
+++ b/Engine/Plugins/Enterprise/GLTFExporter/MigratedFromGit/Source/GLTFExporter/Private/Converters/GLTFVariantSetConverters.h
@@ -19,6 +19,8 @@ class FGLTFLevelVariantSetsConverter final : public TGLTFConverter<FGLTFJsonLeve
 	bool TryParseVariantBinding(FGLTFConvertBuilder& Builder, FGLTFJsonVariant& OutVariant, const UVariantObjectBinding* Binding) const;
 	bool TryParseVisibilityPropertyValue(FGLTFConvertBuilder& Builder, FGLTFJsonVariant& OutVariant, const UPropertyValue* Property) const;
 	bool TryParseMaterialPropertyValue(FGLTFConvertBuilder& Builder, FGLTFJsonVariant& OutVariant, const UPropertyValue* Property) const;
+	bool TryParseStaticMeshPropertyValue(FGLTFConvertBuilder& Builder, FGLTFJsonVariant& OutVariant, const UPropertyValue* Property) const;
+	bool TryParseSkeletalMeshPropertyValue(FGLTFConvertBuilder& Builder, FGLTFJsonVariant& OutVariant, const UPropertyValue* Property) const;
 
 	template<typename T>
 	bool TryGetPropertyValue(UPropertyValue* Property, T& OutValue) const;
diff --git a/Engine/Plugins/Enterprise/GLTFExporter/MigratedFromGit/Source/GLTFExporter/Private/Json/GLTFJsonLevelVariantSets.h b/Engine/Plugins/Enterprise/GLTFExporter/MigratedFromGit/Source/GLTFExporter/Private/Json/GLTFJsonLevelVariantSets.h
--- a/Engine/Plugins/Enterprise/GLTFExporter/MigratedFromGit/Source/GLTFExporter/Private/Json/GLTFJsonLevelVariantSets.h
+++ b/Engine/Plugins/Enterprise/GLTFExporter/MigratedFromGit/Source/GLTFExporter/Private/Json/GLTFJsonLevelVariantSets.h
@@ -41,6 +41,7 @@ struct FGLTFJsonVariantMaterial
 struct FGLTFJsonVariantNodeProperties
 {
 	FGLTFJsonNodeIndex Node;
+	FGLTFJsonMeshIndex Mesh;
 	TOptional<bool>    bIsVisible;
 
 	TArray<FGLTFJsonVariantMaterial> Materials;
@@ -62,6 +63,11 @@ struct FGLTFJsonVariantNodeProperties
 			JsonWriter.WriteValue(TEXT("visible"), bIsVisible.GetValue());
 		}
 
+		if (Mesh != INDEX_NONE)
+		{
+			JsonWriter.WriteValue(TEXT("mesh"), Mesh);
+		}
+
 		FGLTFJsonUtility::WriteObjectArray(JsonWriter, TEXT("materials"), Materials, Extensions);
 
 		JsonWriter.WriteObjectEnd();
