[GLTFExporter] fix(variants): perform the check for unnecessary override-materials when exporting *all* meshes

commit 66bee4f0bb78feab83b7eb29c0b1bf93572cea61
Date: Thu, 13 Aug 2020 13:07:54 +0200

    fix(variants): perform the check for unnecessary override-materials when exporting *all* meshes
    previously the check was only performed when a mesh was exported via a mesh-component
    (cherry picked from commit 1e3c1a7ae6e0e7dbf073976a313d1f25e74da46c​)
    # Conflicts from cherry-pick:
    #	Plugins/GLTFExporter/Source/GLTFExporter/Private/Builders/GLTFConvertBuilder.cpp

#skipci

diff --git a/Engine/Plugins/Enterprise/GLTFExporter/MigratedFromGit/Source/GLTFExporter/Private/Builders/GLTFConvertBuilder.cpp b/Engine/Plugins/Enterprise/GLTFExporter/MigratedFromGit/Source/GLTFExporter/Private/Builders/GLTFConvertBuilder.cpp
--- a/Engine/Plugins/Enterprise/GLTFExporter/MigratedFromGit/Source/GLTFExporter/Private/Builders/GLTFConvertBuilder.cpp
+++ b/Engine/Plugins/Enterprise/GLTFExporter/MigratedFromGit/Source/GLTFExporter/Private/Builders/GLTFConvertBuilder.cpp
@@ -115,8 +115,13 @@ FGLTFJsonMeshIndex FGLTFConvertBuilder::GetOrAddMesh(const UStaticMesh* StaticMe
 	{
 		OverrideVertexColors = nullptr;
 	}
-
-	return StaticMeshConverter.GetOrAdd(*this, DesiredName.IsEmpty() ? FGLTFBuilderUtility::GetLODName(StaticMesh, LODIndex) : DesiredName, StaticMesh, LODIndex, OverrideVertexColors, OverrideMaterials);
+	
+	return StaticMeshConverter.GetOrAdd(
+		*this, DesiredName.IsEmpty() ? FGLTFBuilderUtility::GetLODName(StaticMesh, LODIndex) : DesiredName,
+		StaticMesh,
+		LODIndex,
+		OverrideVertexColors,
+		OverrideMaterials == StaticMesh->StaticMaterials ? FGLTFMaterialArray(): OverrideMaterials);
 }
 
 FGLTFJsonMeshIndex FGLTFConvertBuilder::GetOrAddMesh(const UStaticMeshComponent* StaticMeshComponent, const FString& DesiredName)
@@ -129,14 +134,7 @@ FGLTFJsonMeshIndex FGLTFConvertBuilder::GetOrAddMesh(const UStaticMeshComponent*
 	const UStaticMesh* StaticMesh = StaticMeshComponent->GetStaticMesh();
 	const int32 LODIndex = StaticMeshComponent->ForcedLodModel > 0 ? StaticMeshComponent->ForcedLodModel - 1 : ExportOptions->DefaultLevelOfDetail;
 	const FColorVertexBuffer* OverrideVertexColors = StaticMeshComponent->LODData.IsValidIndex(LODIndex) ? StaticMeshComponent->LODData[LODIndex].OverrideVertexColors : nullptr;
-	FGLTFMaterialArray OverrideMaterials = FGLTFMaterialArray(StaticMeshComponent->OverrideMaterials);
-
-	// Clear unnecessary overrides in order to keep the hash identical to when no overrides are specified.
-	// Otherwise the exact same mesh can be exported twice, once with default overrides and once with no overrides.
-	if (OverrideMaterials == StaticMesh->StaticMaterials)
-	{
-		OverrideMaterials.Reset();
-	}
+	const FGLTFMaterialArray OverrideMaterials = FGLTFMaterialArray(StaticMeshComponent->OverrideMaterials);
 
 	return GetOrAddMesh(StaticMesh, LODIndex, OverrideVertexColors, OverrideMaterials, DesiredName);
 }
