[GLTFExporter] feat(variants): use current override-materials from the mesh-component when exporting meshes for variants

commit d249d3be639ca0b4c7aa210332d80313dd50f101
Date: Thu, 13 Aug 2020 13:09:25 +0200

    feat(variants): use current override-materials from the mesh-component when exporting meshes for variants
    necessary to get correct behaviour for mesh variants in the viewer
    (cherry picked from commit 6459309b4fb6c2cc0f4faf7b347e0d43540129d0​)

#skipci

diff --git a/Engine/Plugins/Enterprise/GLTFExporter/MigratedFromGit/Source/GLTFExporter/Private/Converters/GLTFVariantSetConverters.cpp b/Engine/Plugins/Enterprise/GLTFExporter/MigratedFromGit/Source/GLTFExporter/Private/Converters/GLTFVariantSetConverters.cpp
--- a/Engine/Plugins/Enterprise/GLTFExporter/MigratedFromGit/Source/GLTFExporter/Private/Converters/GLTFVariantSetConverters.cpp
+++ b/Engine/Plugins/Enterprise/GLTFExporter/MigratedFromGit/Source/GLTFExporter/Private/Converters/GLTFVariantSetConverters.cpp
@@ -300,6 +300,15 @@ bool FGLTFLevelVariantSetsConverter::TryParseStaticMeshPropertyValue(FGLTFConver
 		return false;
 	}
 
+	const UMeshComponent* MeshComponent = Cast<UMeshComponent>(Target);
+	if (MeshComponent == nullptr)
+	{
+		Builder.AddWarningMessage(FString::Printf(
+			TEXT("Target object for property has no mesh-component, the property will be skipped. Context: %s"),
+			*GetLogContext(Property)));
+		return false;
+	}
+
 	if (Builder.bSelectedActorsOnly && !Target->IsSelected())
 	{
 		Builder.AddWarningMessage(FString::Printf(
@@ -318,8 +327,9 @@ bool FGLTFLevelVariantSetsConverter::TryParseStaticMeshPropertyValue(FGLTFConver
 		return false;
 	}
 
+	const FGLTFMaterialArray OverrideMaterials(MeshComponent->OverrideMaterials);
 	const FGLTFJsonNodeIndex NodeIndex = Builder.GetOrAddNode(Target);
-	const FGLTFJsonMeshIndex MeshIndex = Builder.GetOrAddMesh(StaticMesh);
+	const FGLTFJsonMeshIndex MeshIndex = Builder.GetOrAddMesh(StaticMesh, 0, nullptr, OverrideMaterials);
 	FGLTFJsonVariantNodeProperties& NodeProperties = OutVariant.Nodes.FindOrAdd(NodeIndex);
 
 	NodeProperties.Node = NodeIndex;
@@ -346,6 +356,15 @@ bool FGLTFLevelVariantSetsConverter::TryParseSkeletalMeshPropertyValue(FGLTFConv
 		return false;
 	}
 
+	const UMeshComponent* MeshComponent = Cast<UMeshComponent>(Target);
+	if (MeshComponent == nullptr)
+	{
+		Builder.AddWarningMessage(FString::Printf(
+			TEXT("Target object for property has no mesh-component, the property will be skipped. Context: %s"),
+			*GetLogContext(Property)));
+		return false;
+	}
+
 	if (Builder.bSelectedActorsOnly && !Target->IsSelected())
 	{
 		Builder.AddWarningMessage(FString::Printf(
@@ -364,8 +383,9 @@ bool FGLTFLevelVariantSetsConverter::TryParseSkeletalMeshPropertyValue(FGLTFConv
 		return false;
 	}
 
+	const FGLTFMaterialArray OverrideMaterials(MeshComponent->OverrideMaterials);
 	const FGLTFJsonNodeIndex NodeIndex = Builder.GetOrAddNode(Target);
-	const FGLTFJsonMeshIndex MeshIndex = Builder.GetOrAddMesh(SkeletalMesh);
+	const FGLTFJsonMeshIndex MeshIndex = Builder.GetOrAddMesh(SkeletalMesh, 0, nullptr, nullptr, OverrideMaterials);
 	FGLTFJsonVariantNodeProperties& NodeProperties = OutVariant.Nodes.FindOrAdd(NodeIndex);
 
 	NodeProperties.Node = NodeIndex;
