[GLTFExporter] fix(variants): use the updated default-value of -1 for LODIndex in GetOrAddMesh() calls

commit 50b5d61c293a2d83a62c61f5c7c16164ddcd7ddc
Date: Fri, 21 Aug 2020 08:45:32 +0200

    fix(variants): use the updated default-value of -1 for LODIndex in GetOrAddMesh() calls
    (cherry picked from commit 09aab1a2005248bb8fe7a742c7ce978c4f59d2c3​)

#skipci

diff --git a/Engine/Plugins/Enterprise/GLTFExporter/MigratedFromGit/Source/GLTFExporter/Private/Converters/GLTFVariantSetConverters.cpp b/Engine/Plugins/Enterprise/GLTFExporter/MigratedFromGit/Source/GLTFExporter/Private/Converters/GLTFVariantSetConverters.cpp
--- a/Engine/Plugins/Enterprise/GLTFExporter/MigratedFromGit/Source/GLTFExporter/Private/Converters/GLTFVariantSetConverters.cpp
+++ b/Engine/Plugins/Enterprise/GLTFExporter/MigratedFromGit/Source/GLTFExporter/Private/Converters/GLTFVariantSetConverters.cpp
@@ -329,7 +329,7 @@ bool FGLTFLevelVariantSetsConverter::TryParseStaticMeshPropertyValue(FGLTFConver
 
 	const FGLTFMaterialArray OverrideMaterials(MeshComponent->OverrideMaterials);
 	const FGLTFJsonNodeIndex NodeIndex = Builder.GetOrAddNode(Target);
-	const FGLTFJsonMeshIndex MeshIndex = Builder.GetOrAddMesh(StaticMesh, 0, nullptr, OverrideMaterials);
+	const FGLTFJsonMeshIndex MeshIndex = Builder.GetOrAddMesh(StaticMesh, -1, nullptr, OverrideMaterials);
 	FGLTFJsonVariantNodeProperties& NodeProperties = OutVariant.Nodes.FindOrAdd(NodeIndex);
 
 	NodeProperties.Node = NodeIndex;
@@ -385,7 +385,7 @@ bool FGLTFLevelVariantSetsConverter::TryParseSkeletalMeshPropertyValue(FGLTFConv
 
 	const FGLTFMaterialArray OverrideMaterials(MeshComponent->OverrideMaterials);
 	const FGLTFJsonNodeIndex NodeIndex = Builder.GetOrAddNode(Target);
-	const FGLTFJsonMeshIndex MeshIndex = Builder.GetOrAddMesh(SkeletalMesh, 0, nullptr, nullptr, OverrideMaterials);
+	const FGLTFJsonMeshIndex MeshIndex = Builder.GetOrAddMesh(SkeletalMesh, -1, nullptr, nullptr, OverrideMaterials);
 	FGLTFJsonVariantNodeProperties& NodeProperties = OutVariant.Nodes.FindOrAdd(NodeIndex);
 
 	NodeProperties.Node = NodeIndex;
