[GLTFExporter] fix(variants): check selection-status of the owning AActor instead of the USceneComponent

commit 1118b5d3a134b7d2f76e3b1cf9ebd40d07bd1bcd
Date: Fri, 21 Aug 2020 09:02:59 +0200

    fix(variants): check selection-status of the owning AActor instead of the USceneComponent
    (cherry picked from commit 9510a6afbf7b666454b0866c5da5eb5011414512​)

#skipci

diff --git a/Engine/Plugins/Enterprise/GLTFExporter/MigratedFromGit/Source/GLTFExporter/Private/Converters/GLTFVariantSetConverters.cpp b/Engine/Plugins/Enterprise/GLTFExporter/MigratedFromGit/Source/GLTFExporter/Private/Converters/GLTFVariantSetConverters.cpp
--- a/Engine/Plugins/Enterprise/GLTFExporter/MigratedFromGit/Source/GLTFExporter/Private/Converters/GLTFVariantSetConverters.cpp
+++ b/Engine/Plugins/Enterprise/GLTFExporter/MigratedFromGit/Source/GLTFExporter/Private/Converters/GLTFVariantSetConverters.cpp
@@ -188,10 +188,18 @@ bool FGLTFLevelVariantSetsConverter::TryParseVisibilityPropertyValue(FGLTFConver
 		return false;
 	}
 
-	if (Builder.bSelectedActorsOnly && !Target->IsSelected())
+	const AActor* Owner = Target->GetOwner();
+	if (Owner == nullptr)
+	{
+		Builder.AddWarningMessage(FString::Printf(TEXT("Invalid scene component, the property will be skipped. Context: %s"),
+			*GetLogContext(Property)));
+		return false;
+	}
+
+	if (Builder.bSelectedActorsOnly && !Owner->IsSelected())
 	{
 		Builder.AddWarningMessage(FString::Printf(
-			TEXT("Target object for property is not selected for export, the property will be skipped. Context: %s"),
+			TEXT("Target actor for property is not selected for export, the property will be skipped. Context: %s"),
 			*GetLogContext(Property)));
 		return false;
 	}
@@ -234,10 +242,18 @@ bool FGLTFLevelVariantSetsConverter::TryParseMaterialPropertyValue(FGLTFConvertB
 		return false;
 	}
 
-	if (Builder.bSelectedActorsOnly && !Target->IsSelected())
+	const AActor* Owner = Target->GetOwner();
+	if (Owner == nullptr)
+	{
+		Builder.AddWarningMessage(FString::Printf(TEXT("Invalid scene component, the property will be skipped. Context: %s"),
+			*GetLogContext(Property)));
+		return false;
+	}
+
+	if (Builder.bSelectedActorsOnly && !Owner->IsSelected())
 	{
 		Builder.AddWarningMessage(FString::Printf(
-			TEXT("Target object for property is not selected for export, the property will be skipped. Context: %s"),
+			TEXT("Target actor for property is not selected for export, the property will be skipped. Context: %s"),
 			*GetLogContext(Property)));
 		return false;
 	}
@@ -300,19 +316,27 @@ bool FGLTFLevelVariantSetsConverter::TryParseStaticMeshPropertyValue(FGLTFConver
 		return false;
 	}
 
-	const UMeshComponent* MeshComponent = Cast<UMeshComponent>(Target);
-	if (MeshComponent == nullptr)
+	const AActor* Owner = Target->GetOwner();
+	if (Owner == nullptr)
+	{
+		Builder.AddWarningMessage(FString::Printf(TEXT("Invalid scene component, the property will be skipped. Context: %s"),
+			*GetLogContext(Property)));
+		return false;
+	}
+
+	if (Builder.bSelectedActorsOnly && !Owner->IsSelected())
 	{
 		Builder.AddWarningMessage(FString::Printf(
-			TEXT("Target object for property has no mesh-component, the property will be skipped. Context: %s"),
+			TEXT("Target actor for property is not selected for export, the property will be skipped. Context: %s"),
 			*GetLogContext(Property)));
 		return false;
 	}
 
-	if (Builder.bSelectedActorsOnly && !Target->IsSelected())
+	const UMeshComponent* MeshComponent = Cast<UMeshComponent>(Target);
+	if (MeshComponent == nullptr)
 	{
 		Builder.AddWarningMessage(FString::Printf(
-			TEXT("Target object for property is not selected for export, the property will be skipped. Context: %s"),
+			TEXT("Target object for property has no mesh-component, the property will be skipped. Context: %s"),
 			*GetLogContext(Property)));
 		return false;
 	}
@@ -356,19 +380,27 @@ bool FGLTFLevelVariantSetsConverter::TryParseSkeletalMeshPropertyValue(FGLTFConv
 		return false;
 	}
 
-	const UMeshComponent* MeshComponent = Cast<UMeshComponent>(Target);
-	if (MeshComponent == nullptr)
+	const AActor* Owner = Target->GetOwner();
+	if (Owner == nullptr)
+	{
+		Builder.AddWarningMessage(FString::Printf(TEXT("Invalid scene component, the property will be skipped. Context: %s"),
+			*GetLogContext(Property)));
+		return false;
+	}
+
+	if (Builder.bSelectedActorsOnly && !Owner->IsSelected())
 	{
 		Builder.AddWarningMessage(FString::Printf(
-			TEXT("Target object for property has no mesh-component, the property will be skipped. Context: %s"),
+			TEXT("Target actor for property is not selected for export, the property will be skipped. Context: %s"),
 			*GetLogContext(Property)));
 		return false;
 	}
 
-	if (Builder.bSelectedActorsOnly && !Target->IsSelected())
+	const UMeshComponent* MeshComponent = Cast<UMeshComponent>(Target);
+	if (MeshComponent == nullptr)
 	{
 		Builder.AddWarningMessage(FString::Printf(
-			TEXT("Target object for property is not selected for export, the property will be skipped. Context: %s"),
+			TEXT("Target object for property has no mesh-component, the property will be skipped. Context: %s"),
 			*GetLogContext(Property)));
 		return false;
 	}
