[GLTFExporter] feat(variants): support mesh-variants that set the mesh to none / null

commit c40de799b9fa7938c51dd07ac340a7f21b81bc0a
Date: Fri, 21 Aug 2020 10:22:01 +0200

    feat(variants): support mesh-variants that set the mesh to none / null
    (cherry picked from commit 19f0a9fe0ec42c05a984c65dc021e10e07e77378​)

#skipci

diff --git a/Engine/Plugins/Enterprise/GLTFExporter/MigratedFromGit/Source/GLTFExporter/Private/Converters/GLTFVariantSetConverters.cpp b/Engine/Plugins/Enterprise/GLTFExporter/MigratedFromGit/Source/GLTFExporter/Private/Converters/GLTFVariantSetConverters.cpp
--- a/Engine/Plugins/Enterprise/GLTFExporter/MigratedFromGit/Source/GLTFExporter/Private/Converters/GLTFVariantSetConverters.cpp
+++ b/Engine/Plugins/Enterprise/GLTFExporter/MigratedFromGit/Source/GLTFExporter/Private/Converters/GLTFVariantSetConverters.cpp
@@ -343,7 +343,7 @@ bool FGLTFLevelVariantSetsConverter::TryParseStaticMeshPropertyValue(FGLTFConver
 
 	const UStaticMesh* StaticMesh;
 
-	if (!TryGetPropertyValue(const_cast<UPropertyValue*>(Property), StaticMesh) || StaticMesh == nullptr)
+	if (!TryGetPropertyValue(const_cast<UPropertyValue*>(Property), StaticMesh))
 	{
 		Builder.AddWarningMessage(FString::Printf(
 			TEXT("Failed to parse recorded data for property, it will be skipped. Context: %s"),
diff --git a/Engine/Plugins/Enterprise/GLTFExporter/MigratedFromGit/Source/GLTFExporter/Private/Json/GLTFJsonLevelVariantSets.h b/Engine/Plugins/Enterprise/GLTFExporter/MigratedFromGit/Source/GLTFExporter/Private/Json/GLTFJsonLevelVariantSets.h
--- a/Engine/Plugins/Enterprise/GLTFExporter/MigratedFromGit/Source/GLTFExporter/Private/Json/GLTFJsonLevelVariantSets.h
+++ b/Engine/Plugins/Enterprise/GLTFExporter/MigratedFromGit/Source/GLTFExporter/Private/Json/GLTFJsonLevelVariantSets.h
@@ -41,9 +41,9 @@ struct FGLTFJsonVariantMaterial
 struct FGLTFJsonVariantNodeProperties
 {
 	FGLTFJsonNodeIndex Node;
-	FGLTFJsonMeshIndex Mesh;
 	TOptional<bool>    bIsVisible;
 
+	TOptional<FGLTFJsonMeshIndex>    Mesh;
 	TArray<FGLTFJsonVariantMaterial> Materials;
 
 	template <class CharType = TCHAR, class PrintPolicy = TPrettyJsonPrintPolicy<CharType>>
@@ -63,9 +63,16 @@ struct FGLTFJsonVariantNodeProperties
 			JsonWriter.WriteValue(TEXT("visible"), bIsVisible.GetValue());
 		}
 
-		if (Mesh != INDEX_NONE)
+		if (Mesh.IsSet())
 		{
-			JsonWriter.WriteValue(TEXT("mesh"), Mesh);
+			if (Mesh.GetValue() != INDEX_NONE)
+			{
+				JsonWriter.WriteValue(TEXT("mesh"), Mesh.GetValue());
+			}
+			else
+			{
+				JsonWriter.WriteValue(TEXT("mesh"), nullptr);
+			}
 		}
 
 		FGLTFJsonUtility::WriteObjectArray(JsonWriter, TEXT("materials"), Materials, Extensions);
