[GLTFExporter] refactor(variants): remove redundant sanity-checks in variant parsers

commit 380c67f9c668e0d3a7fdd23c0a075f4f3552b0ab
Date: Fri, 21 Aug 2020 11:02:09 +0200

    refactor(variants): remove redundant sanity-checks in variant parsers
    (cherry picked from commit 3567d68c7b40db4049a07b8bcfd6f8c20f0e0535​)

#skipci

diff --git a/Engine/Plugins/Enterprise/GLTFExporter/MigratedFromGit/Source/GLTFExporter/Private/Converters/GLTFVariantSetConverters.cpp b/Engine/Plugins/Enterprise/GLTFExporter/MigratedFromGit/Source/GLTFExporter/Private/Converters/GLTFVariantSetConverters.cpp
--- a/Engine/Plugins/Enterprise/GLTFExporter/MigratedFromGit/Source/GLTFExporter/Private/Converters/GLTFVariantSetConverters.cpp
+++ b/Engine/Plugins/Enterprise/GLTFExporter/MigratedFromGit/Source/GLTFExporter/Private/Converters/GLTFVariantSetConverters.cpp
@@ -170,15 +170,6 @@ bool FGLTFLevelVariantSetsConverter::TryParseVariantBinding(FGLTFConvertBuilder&
 
 bool FGLTFLevelVariantSetsConverter::TryParseVisibilityPropertyValue(FGLTFConvertBuilder& Builder, FGLTFJsonVariant& OutVariant, const UPropertyValue* Property) const
 {
-	if (Property->GetPropertyName() != TEXT("bVisible") ||
-		!Property->GetPropertyClass()->IsChildOf(FBoolProperty::StaticClass()))
-	{
-		Builder.AddErrorMessage(FString::Printf(
-			TEXT("Attempted to parse visibility from an incompatible property. Context: %s"),
-			*GetLogContext(Property)));
-		return false;
-	}
-
 	const USceneComponent* Target = static_cast<USceneComponent*>(Property->GetPropertyParentContainerAddress());
 	if (Target == nullptr)
 	{
@@ -225,10 +216,10 @@ bool FGLTFLevelVariantSetsConverter::TryParseVisibilityPropertyValue(FGLTFConver
 bool FGLTFLevelVariantSetsConverter::TryParseMaterialPropertyValue(FGLTFConvertBuilder& Builder, FGLTFJsonVariant& OutVariant, const UPropertyValue* Property) const
 {
 	UPropertyValueMaterial* MaterialProperty = Cast<UPropertyValueMaterial>(const_cast<UPropertyValue*>(Property));
-	if (!MaterialProperty)
+	if (MaterialProperty == nullptr)
 	{
-		Builder.AddErrorMessage(FString::Printf(
-			TEXT("Attempted to parse material from an incompatible property. Context: %s"),
+		Builder.AddWarningMessage(FString::Printf(
+			TEXT("Material property is invalid, the property will be skipped. Context: %s"),
 			*GetLogContext(Property)));
 		return false;
 	}
@@ -260,7 +251,7 @@ bool FGLTFLevelVariantSetsConverter::TryParseMaterialPropertyValue(FGLTFConvertB
 
 	// NOTE: UPropertyValueMaterial::GetMaterial does *not* ensure that the recorded data has been loaded,
 	// so we need to call UProperty::GetRecordedData first to make that happen.
-	const_cast<UPropertyValueMaterial*>(MaterialProperty)->GetRecordedData();
+	MaterialProperty->GetRecordedData();
 
 	const UMaterialInterface* Material = MaterialProperty->GetMaterial();
 	if (Material == nullptr)
@@ -299,14 +290,6 @@ bool FGLTFLevelVariantSetsConverter::TryParseMaterialPropertyValue(FGLTFConvertB
 
 bool FGLTFLevelVariantSetsConverter::TryParseStaticMeshPropertyValue(FGLTFConvertBuilder& Builder, FGLTFJsonVariant& OutVariant, const UPropertyValue* Property) const
 {
-	if (Property->GetPropertyName() != TEXT("StaticMesh"))
-	{
-		Builder.AddErrorMessage(FString::Printf(
-			TEXT("Attempted to parse static mesh from an incompatible property. Context: %s"),
-			*GetLogContext(Property)));
-		return false;
-	}
-
 	const USceneComponent* Target = static_cast<USceneComponent*>(Property->GetPropertyParentContainerAddress());
 	if (Target == nullptr)
 	{
@@ -363,14 +346,6 @@ bool FGLTFLevelVariantSetsConverter::TryParseStaticMeshPropertyValue(FGLTFConver
 
 bool FGLTFLevelVariantSetsConverter::TryParseSkeletalMeshPropertyValue(FGLTFConvertBuilder& Builder, FGLTFJsonVariant& OutVariant, const UPropertyValue* Property) const
 {
-	if (Property->GetPropertyName() != TEXT("SkeletalMesh"))
-	{
-		Builder.AddErrorMessage(FString::Printf(
-			TEXT("Attempted to parse skeletal mesh from an incompatible property. Context: %s"),
-			*GetLogContext(Property)));
-		return false;
-	}
-
 	const USceneComponent* Target = static_cast<USceneComponent*>(Property->GetPropertyParentContainerAddress());
 	if (Target == nullptr)
 	{
