[GLTFExporter] feat(variants): use -1 instead of null for indicating that a mesh-variant should use "None"

commit f7ab1580a14f72c64638f94ea711f765a18ebb88
Date: Mon, 24 Aug 2020 08:53:43 +0200

    feat(variants): use -1 instead of null for indicating that a mesh-variant should use "None"
    (cherry picked from commit 4e07e86519952a2e38a61f4d21206634c01cb2e3​)

#skipci

diff --git a/Engine/Plugins/Enterprise/GLTFExporter/MigratedFromGit/Source/GLTFExporter/Private/Json/GLTFJsonLevelVariantSets.h b/Engine/Plugins/Enterprise/GLTFExporter/MigratedFromGit/Source/GLTFExporter/Private/Json/GLTFJsonLevelVariantSets.h
--- a/Engine/Plugins/Enterprise/GLTFExporter/MigratedFromGit/Source/GLTFExporter/Private/Json/GLTFJsonLevelVariantSets.h
+++ b/Engine/Plugins/Enterprise/GLTFExporter/MigratedFromGit/Source/GLTFExporter/Private/Json/GLTFJsonLevelVariantSets.h
@@ -65,14 +65,7 @@ struct FGLTFJsonVariantNodeProperties
 
 		if (Mesh.IsSet())
 		{
-			if (Mesh.GetValue() != INDEX_NONE)
-			{
-				JsonWriter.WriteValue(TEXT("mesh"), Mesh.GetValue());
-			}
-			else
-			{
-				JsonWriter.WriteValue(TEXT("mesh"), nullptr);
-			}
+			JsonWriter.WriteValue(TEXT("mesh"), Mesh.GetValue());
 		}
 
 		FGLTFJsonUtility::WriteObjectArray(JsonWriter, TEXT("materials"), Materials, Extensions);
