[GLTFExporter] feat(variants): support mesh-variants that set the material to "None"

commit ba25e76ff32417b307b3c0ff81986213ae82402c
Date: Mon, 24 Aug 2020 09:14:19 +0200

    feat(variants): support mesh-variants that set the material to "None"
    (cherry picked from commit a2a57c1d0f41ccdd0d548a281ec4d383db0dc1e1​)

#skipci

diff --git a/Engine/Plugins/Enterprise/GLTFExporter/MigratedFromGit/Source/GLTFExporter/Private/Converters/GLTFVariantSetConverters.cpp b/Engine/Plugins/Enterprise/GLTFExporter/MigratedFromGit/Source/GLTFExporter/Private/Converters/GLTFVariantSetConverters.cpp
--- a/Engine/Plugins/Enterprise/GLTFExporter/MigratedFromGit/Source/GLTFExporter/Private/Converters/GLTFVariantSetConverters.cpp
+++ b/Engine/Plugins/Enterprise/GLTFExporter/MigratedFromGit/Source/GLTFExporter/Private/Converters/GLTFVariantSetConverters.cpp
@@ -249,19 +249,6 @@ bool FGLTFLevelVariantSetsConverter::TryParseMaterialPropertyValue(FGLTFConvertB
 		return false;
 	}
 
-	// NOTE: UPropertyValueMaterial::GetMaterial does *not* ensure that the recorded data has been loaded,
-	// so we need to call UProperty::GetRecordedData first to make that happen.
-	MaterialProperty->GetRecordedData();
-
-	const UMaterialInterface* Material = MaterialProperty->GetMaterial();
-	if (Material == nullptr)
-	{
-		Builder.AddWarningMessage(FString::Printf(
-			TEXT("Failed to parse recorded data for property, it will be skipped. Context: %s"),
-			*GetLogContext(Property)));
-		return false;
-	}
-
 	const TArray<FCapturedPropSegment>& CapturedPropSegments = reinterpret_cast<UPropertyValueHack*>(MaterialProperty)->GetCapturedPropSegments();
 	const int32 NumPropSegments = CapturedPropSegments.Num();
 
@@ -273,6 +260,12 @@ bool FGLTFLevelVariantSetsConverter::TryParseMaterialPropertyValue(FGLTFConvertB
 		return false;
 	}
 
+	// NOTE: UPropertyValueMaterial::GetMaterial does *not* ensure that the recorded data has been loaded,
+	// so we need to call UProperty::GetRecordedData first to make that happen.
+	MaterialProperty->GetRecordedData();
+
+	// TODO: find way to determine whether the material is null because "None" was selected, or because it failed to resolve
+	const UMaterialInterface* Material = MaterialProperty->GetMaterial();
 	const FGLTFJsonMaterialIndex MaterialIndex = Builder.GetOrAddMaterial(Material);
 	const int32 ElementIndex = CapturedPropSegments[NumPropSegments - 1].PropertyIndex;
 
diff --git a/Engine/Plugins/Enterprise/GLTFExporter/MigratedFromGit/Source/GLTFExporter/Private/Json/GLTFJsonLevelVariantSets.h b/Engine/Plugins/Enterprise/GLTFExporter/MigratedFromGit/Source/GLTFExporter/Private/Json/GLTFJsonLevelVariantSets.h
--- a/Engine/Plugins/Enterprise/GLTFExporter/MigratedFromGit/Source/GLTFExporter/Private/Json/GLTFJsonLevelVariantSets.h
+++ b/Engine/Plugins/Enterprise/GLTFExporter/MigratedFromGit/Source/GLTFExporter/Private/Json/GLTFJsonLevelVariantSets.h
@@ -24,10 +24,7 @@ struct FGLTFJsonVariantMaterial
 	{
 		JsonWriter.WriteObjectStart();
 
-		if (Material != INDEX_NONE)
-		{
-			JsonWriter.WriteValue(TEXT("material"), Material);
-		}
+		JsonWriter.WriteValue(TEXT("material"), Material);
 
 		if (Index != INDEX_NONE)
 		{
