[GLTFExporter] chore(orbit-camera): restore runtime module

commit 9abdfb9c23edd16a97f8c1713449a7848b5f6fa0
Date: Tue, 4 Aug 2020 12:12:15 +0200

    chore(orbit-camera): restore runtime module
    (cherry picked from commit 6ac78a06177777d2a7fb4b6b647188a5dd08c8a1​)

#skipci

diff --git a/Engine/Plugins/Enterprise/GLTFExporter/MigratedFromGit/GLTFExporter.uplugin b/Engine/Plugins/Enterprise/GLTFExporter/MigratedFromGit/GLTFExporter.uplugin
--- a/Engine/Plugins/Enterprise/GLTFExporter/MigratedFromGit/GLTFExporter.uplugin
+++ b/Engine/Plugins/Enterprise/GLTFExporter/MigratedFromGit/GLTFExporter.uplugin
@@ -13,12 +13,16 @@
 	"EngineVersion": "4.25.0",
 	"CanContainContent": true,
 	"IsBetaVersion": true,
-	"Modules":
-	[
+	"Modules": [
 		{
 			"Name": "GLTFExporter",
 			"Type": "Editor",
 			"LoadingPhase": "Default"
+		},
+		{
+			"Name": "GLTFExporterRuntime",
+			"Type": "Editor",
+			"LoadingPhase": "Default"
 		}
 	],
 	"Plugins": [
