[GLTFExporter] feat(orbit-camera): add class for actor version of the orbit camera

commit 5c2fda8da486a63639bb62eed5145e43cd1626bc
Date: Tue, 4 Aug 2020 12:17:15 +0200

    feat(orbit-camera): add class for actor version of the orbit camera
    (cherry picked from commit db76000babd2032a459adcc52244c0ef8716b937​)

#skipci

diff --git a/Engine/Plugins/Enterprise/GLTFExporter/MigratedFromGit/Source/GLTFExporterRuntime/Private/Actors/GLTFOrbitCameraActor.cpp b/Engine/Plugins/Enterprise/GLTFExporter/MigratedFromGit/Source/GLTFExporterRuntime/Private/Actors/GLTFOrbitCameraActor.cpp
new file mode 100644
--- /dev/null
+++ b/Engine/Plugins/Enterprise/GLTFExporter/MigratedFromGit/Source/GLTFExporterRuntime/Private/Actors/GLTFOrbitCameraActor.cpp
@@ -0,0 +1,167 @@
+// Copyright Epic Games, Inc. All Rights Reserved.
+
+#include "Actors/GLTFOrbitCameraActor.h"
+//#include "Kismet/GameplayStatics.h"
+
+namespace
+{
+	const float DollyDuration = 0.2f;
+
+	float EaseInQuad(float T, float B, float C, float D)
+	{
+		return C * (T /= D) * T + B;
+	}
+} // Anonymous namespace
+
+#define DEBUGGLTFORBITCAMERA 0
+
+#if DEBUGGLTFORBITCAMERA
+DEFINE_LOG_CATEGORY_STATIC(LogEditorGLTFOrbitCamera, Log, All);
+#endif
+
+AGLTFOrbitCameraActor::AGLTFOrbitCameraActor(const FObjectInitializer& ObjectInitializer)
+	: Super(ObjectInitializer)
+	UPROPERTY(EditAnywhere, BlueprintReadWrite, Category = "GLTF Orbit Camera Actor"),
+	Focus(nullptr),
+	DistanceMin(0.0f),
+	DistanceMax(0.0f),
+	PitchAngleMin(-90.0f),
+	PitchAngleMax(90.0f),
+	OrbitInertia(0.07f),
+	OrbitSensitivity(30.0f),
+	DistanceSensitivity(50.0f),
+	Distance(0.0f),
+	Pitch(0.0f),
+	Yaw(0.0f),
+	TargetDistance(0.0f),
+	TargetPitch(0.0f),
+	TargetYaw(0.0f),
+	DollyTime(0.0f),
+	DollyStartDistance(0.0f)
+{
+	PrimaryActorTick.bCanEverTick = true;
+	// TODO: Figure out how to auto-enable both rendering and input for the camera
+	//AutoReceiveInput = static_cast<EAutoReceiveInput::Type>(GetAutoActivatePlayerIndex() + 1);
+	AutoReceiveInput = EAutoReceiveInput::Player0;
+}
+
+void AGLTFOrbitCameraActor::BeginPlay()
+{
+	Super::BeginPlay();
+	
+	Distance = ClampDistance(Distance);
+	TargetDistance = Distance;
+
+	if (InputComponent)
+	{
+		InputComponent->BindAxisKey(EKeys::MouseX, this, &AGLTFOrbitCameraActor::OnMouseX);
+		InputComponent->BindAxisKey(EKeys::MouseY, this, &AGLTFOrbitCameraActor::OnMouseY);
+		InputComponent->BindAxisKey(EKeys::MouseWheelAxis, this, &AGLTFOrbitCameraActor::OnMouseWheelAxis);
+	}
+}
+
+void AGLTFOrbitCameraActor::Tick(float DeltaSeconds)
+{
+	Super::Tick(DeltaSeconds);
+
+	if (DollyTime != 0.0f)
+	{
+		Distance = EaseInQuad(DollyDuration - DollyTime, DollyStartDistance, TargetDistance - DollyStartDistance, DollyDuration);
+		DollyTime = FMath::Max(DollyTime - DeltaSeconds, 0.0f);
+	}
+
+	const float Alpha = (OrbitInertia == 0.0f) ? 1.0f : FMath::Min(DeltaSeconds / OrbitInertia, 1.0f);
+	Yaw = FMath::Lerp(Yaw, TargetYaw, Alpha);
+	Pitch = FMath::Lerp(Pitch, TargetPitch, Alpha);
+	//Distance = FMath::Lerp(Distance, TargetDistance, Alpha);
+
+	const FVector FocusPosition = (this->Focus != nullptr) ? this->Focus->GetActorLocation() : FVector::ZeroVector;
+	const FTransform FocusTransform = FTransform(FocusPosition);
+	const FTransform DollyTransform = FTransform(-FVector::ForwardVector * Distance);
+	const FTransform RotationTransform = FTransform(FQuat::MakeFromEuler(FVector(0.0f, Pitch, Yaw)));
+	const FTransform ResultTransform = DollyTransform * RotationTransform * FocusTransform;
+
+	SetActorTransform(ResultTransform);
+}
+
+void AGLTFOrbitCameraActor::OnMouseX(float AxisValue)
+{
+	if (AxisValue == 0.0f)
+	{
+		return;
+	}
+
+	const float Delta = (TargetYaw + AxisValue * OrbitSensitivity) - Yaw;
+	const float Remainder = FMath::Fmod(Delta, 360.0f);
+
+	if (Remainder > 180.0f)
+	{
+		TargetYaw = Yaw - (360.0f - Remainder);
+}
+	else if (Remainder < -180.0f)
+	{
+		TargetYaw = Yaw + (360.0f + Remainder);
+	}
+	else
+	{
+		TargetYaw = Yaw + Remainder;
+	}
+
+	//TargetYaw = ClampYaw(TargetYaw + AxisValue * OrbitSensitivity);
+
+#if DEBUGGLTFORBITCAMERA
+	UE_LOG(LogEditorGLTFOrbitCamera, Warning, TEXT("AGLTFOrbitCameraActor::OnMouseX(), %f"), AxisValue);
+#endif
+}
+
+void AGLTFOrbitCameraActor::OnMouseY(float AxisValue)
+{
+	if (AxisValue == 0.0f)
+	{
+		return;
+	}
+
+	TargetPitch = ClampPitch(TargetPitch + AxisValue * OrbitSensitivity);
+
+#if DEBUGGLTFORBITCAMERA
+	UE_LOG(LogEditorGLTFOrbitCamera, Warning, TEXT("AGLTFOrbitCameraActor::OnMouseY(), %f"), AxisValue);
+#endif
+}
+
+void AGLTFOrbitCameraActor::OnMouseWheelAxis(float AxisValue)
+{
+	if (AxisValue == 0.0f)
+	{
+		return;
+	}
+
+	DollyTime = DollyDuration;
+	TargetDistance = ClampDistance(TargetDistance + -AxisValue * DistanceSensitivity);
+	DollyStartDistance = Distance;
+
+#if DEBUGGLTFORBITCAMERA
+	UE_LOG(LogEditorGLTFOrbitCamera, Warning, TEXT("AGLTFOrbitCameraActor::OnMouseWheelAxis(), %f"), AxisValue);
+#endif
+}
+
+float AGLTFOrbitCameraActor::ClampDistance(float Value) const
+{
+	return FMath::Clamp(Value, DistanceMin, DistanceMax);
+}
+
+float AGLTFOrbitCameraActor::ClampPitch(float Value) const
+{
+	return FMath::Clamp(Value, PitchAngleMin, PitchAngleMax);
+}
+
+float AGLTFOrbitCameraActor::ClampYaw(float Value) const
+{
+	return FMath::Fmod(Value + 360.0f, 360.0f);
+}
+
+void AGLTFOrbitCameraActor::RemoveInertia()
+{
+	Yaw = TargetYaw;
+	Pitch = TargetPitch;
+	Distance = TargetDistance;
+}
diff --git a/Engine/Plugins/Enterprise/GLTFExporter/MigratedFromGit/Source/GLTFExporterRuntime/Public/Actors/GLTFOrbitCameraActor.h b/Engine/Plugins/Enterprise/GLTFExporter/MigratedFromGit/Source/GLTFExporterRuntime/Public/Actors/GLTFOrbitCameraActor.h
new file mode 100644
--- /dev/null
+++ b/Engine/Plugins/Enterprise/GLTFExporter/MigratedFromGit/Source/GLTFExporterRuntime/Public/Actors/GLTFOrbitCameraActor.h
@@ -0,0 +1,85 @@
+// Copyright Epic Games, Inc. All Rights Reserved.
+
+#pragma once
+
+#include "Camera/CameraActor.h"
+#include "GLTFOrbitCameraActor.generated.h"
+
+/**
+ * GLTF-compatible camera that will carry over settings and simulate the behavior in the resulting viewer.
+ * Focuses one actor in the scene and orbits it through mouse control.
+ */
+UCLASS(BlueprintType, Blueprintable, DisplayName = "GLTF Orbit Camera Actor")
+class GLTFEXPORTERRUNTIME_API AGLTFOrbitCameraActor : public ACameraActor
+{
+	GENERATED_BODY()
+	//~ Begin UObject Interface
+public:
+	AGLTFOrbitCameraActor(const FObjectInitializer& ObjectInitializer);
+	//~ End UObject Interface
+
+	//~ Begin AActor Interface
+protected:
+	virtual void BeginPlay() override;
+public:
+	virtual void Tick(float DeltaSeconds) override;
+	//~ End AActor Interface
+
+private:
+	UFUNCTION()
+	void OnMouseX(float AxisValue);
+
+	UFUNCTION()
+	void OnMouseY(float AxisValue);
+
+	UFUNCTION()
+	void OnMouseWheelAxis(float AxisValue);
+
+	float ClampDistance(float Value) const;
+	float ClampPitch(float Value) const;
+	float ClampYaw(float Value) const;
+	void RemoveInertia();
+
+public:
+	/* Actor which the camera will focus on and subsequently orbit. */
+	UPROPERTY(EditAnywhere, BlueprintReadWrite, Category = "GLTF Orbit Camera Actor")
+	AActor* Focus;
+
+	/* Closest distance the camera can approach the focused actor. */
+	UPROPERTY(EditAnywhere, BlueprintReadWrite, Category = "GLTF Orbit Camera Actor")
+	float DistanceMin;
+
+	/* Farthest distance the camera can recede from the focused actor. */
+	UPROPERTY(EditAnywhere, BlueprintReadWrite, Category = "GLTF Orbit Camera Actor")
+	float DistanceMax;
+
+	/* Minimum angle (in degrees) that the camera can pitch relative to the focused actor. */
+	UPROPERTY(EditAnywhere, BlueprintReadWrite, Category = "GLTF Orbit Camera Actor")
+	float PitchAngleMin;
+
+	/* Maximum angle (in degrees) that the camera can pitch relative to the focused actor. */
+	UPROPERTY(VisibleAnywhere, BlueprintReadOnly, Category = "GLTF Orbit Camera Actor")
+	float PitchAngleMax;
+
+	/* Deceleration that occurs after orbital movement. The higher the value, the longer it takes to settle. */
+	UPROPERTY(VisibleAnywhere, BlueprintReadOnly, Category = "GLTF Orbit Camera Actor")
+	float OrbitInertia;
+
+	/* Size of the orbital movement relative to user input. The higher the value, the faster it moves. */
+	UPROPERTY(VisibleAnywhere, BlueprintReadOnly, Category = "GLTF Orbit Camera Actor")
+	float OrbitSensitivity;
+	
+	/* Size of the dolly movement relative to user input. The higher the value, the faster it moves. */
+	UPROPERTY(VisibleAnywhere, BlueprintReadOnly, Category = "GLTF Orbit Camera Actor")
+	float DistanceSensitivity;
+
+private:
+	float Distance;
+	float Pitch;
+	float Yaw;
+	float TargetDistance;
+	float TargetPitch;
+	float TargetYaw;
+	float DollyTime;
+	float DollyStartDistance;
+};
\ No newline at end of file
