[GLTFExporter] fix(orbit-camera): correct property specifiers

commit 328d93a65c4c438091576a4b943a1ef195af449c
Date: Tue, 4 Aug 2020 12:31:52 +0200

    fix(orbit-camera): correct property specifiers
    (cherry picked from commit 5dcb577b3c8806af3e8d6ca11748d2ee6624042a​)

#skipci

diff --git a/Engine/Plugins/Enterprise/GLTFExporter/MigratedFromGit/Source/GLTFExporterRuntime/Public/Actors/GLTFOrbitCameraActor.h b/Engine/Plugins/Enterprise/GLTFExporter/MigratedFromGit/Source/GLTFExporterRuntime/Public/Actors/GLTFOrbitCameraActor.h
--- a/Engine/Plugins/Enterprise/GLTFExporter/MigratedFromGit/Source/GLTFExporterRuntime/Public/Actors/GLTFOrbitCameraActor.h
+++ b/Engine/Plugins/Enterprise/GLTFExporter/MigratedFromGit/Source/GLTFExporterRuntime/Public/Actors/GLTFOrbitCameraActor.h
@@ -58,7 +58,7 @@ public:
 	float PitchAngleMin;
 
 	/* Maximum angle (in degrees) that the camera can pitch relative to the focused actor. */
-	UPROPERTY(VisibleAnywhere, BlueprintReadOnly, Category = "GLTF Orbit Camera Actor")
+	UPROPERTY(EditAnywhere, BlueprintReadWrite, Category = "GLTF Orbit Camera Actor")
 	float PitchAngleMax;
 
 	/* Deceleration that occurs after orbital movement. The higher the value, the longer it takes to settle. */
