[GLTFExporter] fix(orbit-camera): initialize input dynamically

commit f77af7b2edcf145c376ad691f24217b48dbd865c
Date: Tue, 4 Aug 2020 13:23:35 +0200

    fix(orbit-camera): initialize input dynamically
    (cherry picked from commit 61418e3a2ccb0850edd257d661ee6790b3997ed6​)

#skipci

diff --git a/Engine/Plugins/Enterprise/GLTFExporter/MigratedFromGit/Source/GLTFExporterRuntime/Private/Actors/GLTFOrbitCameraActor.cpp b/Engine/Plugins/Enterprise/GLTFExporter/MigratedFromGit/Source/GLTFExporterRuntime/Private/Actors/GLTFOrbitCameraActor.cpp
--- a/Engine/Plugins/Enterprise/GLTFExporter/MigratedFromGit/Source/GLTFExporterRuntime/Private/Actors/GLTFOrbitCameraActor.cpp
+++ b/Engine/Plugins/Enterprise/GLTFExporter/MigratedFromGit/Source/GLTFExporterRuntime/Private/Actors/GLTFOrbitCameraActor.cpp
@@ -40,9 +40,6 @@ AGLTFOrbitCameraActor::AGLTFOrbitCameraActor(const FObjectInitializer& ObjectIni
 	DollyStartDistance(0.0f)
 {
 	PrimaryActorTick.bCanEverTick = true;
-	// TODO: Figure out how to auto-enable both rendering and input for the camera
-	//AutoReceiveInput = static_cast<EAutoReceiveInput::Type>(GetAutoActivatePlayerIndex() + 1);
-	AutoReceiveInput = EAutoReceiveInput::Player0;
 }
 
 void AGLTFOrbitCameraActor::BeginPlay()
@@ -50,7 +47,11 @@ void AGLTFOrbitCameraActor::BeginPlay()
 	Super::BeginPlay();
 	
 	Distance = ClampDistance(Distance);
+	Pitch = ClampPitch(Pitch);
+	Yaw = ClampYaw(Yaw);
 	TargetDistance = Distance;
+	TargetPitch = Pitch;
+	TargetYaw = Yaw;
 
 	if (InputComponent)
 	{
@@ -84,6 +85,13 @@ void AGLTFOrbitCameraActor::Tick(float DeltaSeconds)
 	SetActorTransform(ResultTransform);
 }
 
+void AGLTFOrbitCameraActor::PreInitializeComponents()
+{
+	AutoReceiveInput = static_cast<EAutoReceiveInput::Type>(GetAutoActivatePlayerIndex() + 1);
+
+	Super::PreInitializeComponents();
+}
+
 void AGLTFOrbitCameraActor::OnMouseX(float AxisValue)
 {
 	if (AxisValue == 0.0f)
@@ -107,8 +115,6 @@ void AGLTFOrbitCameraActor::OnMouseX(float AxisValue)
 		TargetYaw = Yaw + Remainder;
 	}
 
-	//TargetYaw = ClampYaw(TargetYaw + AxisValue * OrbitSensitivity);
-
 #if DEBUGGLTFORBITCAMERA
 	UE_LOG(LogEditorGLTFOrbitCamera, Warning, TEXT("AGLTFOrbitCameraActor::OnMouseX(), %f"), AxisValue);
 #endif
diff --git a/Engine/Plugins/Enterprise/GLTFExporter/MigratedFromGit/Source/GLTFExporterRuntime/Public/Actors/GLTFOrbitCameraActor.h b/Engine/Plugins/Enterprise/GLTFExporter/MigratedFromGit/Source/GLTFExporterRuntime/Public/Actors/GLTFOrbitCameraActor.h
--- a/Engine/Plugins/Enterprise/GLTFExporter/MigratedFromGit/Source/GLTFExporterRuntime/Public/Actors/GLTFOrbitCameraActor.h
+++ b/Engine/Plugins/Enterprise/GLTFExporter/MigratedFromGit/Source/GLTFExporterRuntime/Public/Actors/GLTFOrbitCameraActor.h
@@ -23,6 +23,7 @@ protected:
 	virtual void BeginPlay() override;
 public:
 	virtual void Tick(float DeltaSeconds) override;
+	virtual void PreInitializeComponents() override;
 	//~ End AActor Interface
 
 private:
