[GLTFExporter] feat(orbit-camera): add warning message for invalid focus setting

commit cce9174f2489741f3cb379f3fed44e14f56e0d1f
Date: Tue, 4 Aug 2020 13:52:42 +0200

    feat(orbit-camera): add warning message for invalid focus setting
    (cherry picked from commit 530cca516504c96edf9fc4ff716871b5c1a60c47​)

#skipci

diff --git a/Engine/Plugins/Enterprise/GLTFExporter/MigratedFromGit/Source/GLTFExporterRuntime/Private/Actors/GLTFOrbitCameraActor.cpp b/Engine/Plugins/Enterprise/GLTFExporter/MigratedFromGit/Source/GLTFExporterRuntime/Private/Actors/GLTFOrbitCameraActor.cpp
--- a/Engine/Plugins/Enterprise/GLTFExporter/MigratedFromGit/Source/GLTFExporterRuntime/Private/Actors/GLTFOrbitCameraActor.cpp
+++ b/Engine/Plugins/Enterprise/GLTFExporter/MigratedFromGit/Source/GLTFExporterRuntime/Private/Actors/GLTFOrbitCameraActor.cpp
@@ -15,16 +15,14 @@ namespace
 
 #define DEBUGGLTFORBITCAMERA 0
 
-#if DEBUGGLTFORBITCAMERA
 DEFINE_LOG_CATEGORY_STATIC(LogEditorGLTFOrbitCamera, Log, All);
-#endif
 
 AGLTFOrbitCameraActor::AGLTFOrbitCameraActor(const FObjectInitializer& ObjectInitializer)
 	: Super(ObjectInitializer)
 	UPROPERTY(EditAnywhere, BlueprintReadWrite, Category = "GLTF Orbit Camera Actor"),
 	Focus(nullptr),
-	DistanceMin(0.0f),
-	DistanceMax(0.0f),
+	DistanceMin(100.0f),
+	DistanceMax(1000.0f),
 	PitchAngleMin(-90.0f),
 	PitchAngleMax(90.0f),
 	OrbitInertia(0.07f),
@@ -42,6 +40,28 @@ AGLTFOrbitCameraActor::AGLTFOrbitCameraActor(const FObjectInitializer& ObjectIni
 	PrimaryActorTick.bCanEverTick = true;
 }
 
+#if WITH_EDITOR
+void AGLTFOrbitCameraActor::PostEditChangeProperty(FPropertyChangedEvent& PropertyChangedEvent)
+{
+	Super::PostEditChangeProperty(PropertyChangedEvent);
+
+	const FProperty* PropertyThatChanged = PropertyChangedEvent.Property;
+
+	if (PropertyThatChanged)
+	{
+		const FString PropertyName = PropertyThatChanged->GetName();
+
+		if (PropertyName == TEXT("Focus"))
+		{
+			if (Focus == this)
+			{
+				UE_LOG(LogEditorGLTFOrbitCamera, Warning, TEXT("The camera cannot focus itself."));
+			}
+		}
+	}
+}
+#endif // WITH_EDITOR
+
 void AGLTFOrbitCameraActor::BeginPlay()
 {
 	Super::BeginPlay();
diff --git a/Engine/Plugins/Enterprise/GLTFExporter/MigratedFromGit/Source/GLTFExporterRuntime/Public/Actors/GLTFOrbitCameraActor.h b/Engine/Plugins/Enterprise/GLTFExporter/MigratedFromGit/Source/GLTFExporterRuntime/Public/Actors/GLTFOrbitCameraActor.h
--- a/Engine/Plugins/Enterprise/GLTFExporter/MigratedFromGit/Source/GLTFExporterRuntime/Public/Actors/GLTFOrbitCameraActor.h
+++ b/Engine/Plugins/Enterprise/GLTFExporter/MigratedFromGit/Source/GLTFExporterRuntime/Public/Actors/GLTFOrbitCameraActor.h
@@ -16,6 +16,9 @@ class GLTFEXPORTERRUNTIME_API AGLTFOrbitCameraActor : public ACameraActor
 	//~ Begin UObject Interface
 public:
 	AGLTFOrbitCameraActor(const FObjectInitializer& ObjectInitializer);
+#if WITH_EDITOR
+	virtual void PostEditChangeProperty(FPropertyChangedEvent& PropertyChangedEvent) override;
+#endif // WITH_EDITOR
 	//~ End UObject Interface
 
 	//~ Begin AActor Interface
