[GLTFExporter] feat(orbit-camera): add runtime validation of focus target

commit 94999d8914d813fc7fb0d1af3b39da543f580dc8
Date: Tue, 4 Aug 2020 14:31:23 +0200

    feat(orbit-camera): add runtime validation of focus target
    (cherry picked from commit 56cbd21353616dcfd3185c11c625dbdc54f3f1b0​)

#skipci

diff --git a/Engine/Plugins/Enterprise/GLTFExporter/MigratedFromGit/Source/GLTFExporterRuntime/Private/Actors/GLTFOrbitCameraActor.cpp b/Engine/Plugins/Enterprise/GLTFExporter/MigratedFromGit/Source/GLTFExporterRuntime/Private/Actors/GLTFOrbitCameraActor.cpp
--- a/Engine/Plugins/Enterprise/GLTFExporter/MigratedFromGit/Source/GLTFExporterRuntime/Private/Actors/GLTFOrbitCameraActor.cpp
+++ b/Engine/Plugins/Enterprise/GLTFExporter/MigratedFromGit/Source/GLTFExporterRuntime/Private/Actors/GLTFOrbitCameraActor.cpp
@@ -65,6 +65,11 @@ void AGLTFOrbitCameraActor::PostEditChangeProperty(FPropertyChangedEvent& Proper
 void AGLTFOrbitCameraActor::BeginPlay()
 {
 	Super::BeginPlay();
+
+	if (this->Focus == nullptr || this->Focus == this)
+	{
+		UE_LOG(LogEditorGLTFOrbitCamera, Warning, TEXT("The camera focus must not be null or the camera itself."));
+	}
 	
 	Distance = ClampDistance(Distance);
 	Pitch = ClampPitch(Pitch);
@@ -96,7 +101,7 @@ void AGLTFOrbitCameraActor::Tick(float DeltaSeconds)
 	Pitch = FMath::Lerp(Pitch, TargetPitch, Alpha);
 	//Distance = FMath::Lerp(Distance, TargetDistance, Alpha);
 
-	const FVector FocusPosition = (this->Focus != nullptr) ? this->Focus->GetActorLocation() : FVector::ZeroVector;
+	const FVector FocusPosition = (this->Focus != nullptr && this->Focus != this) ? this->Focus->GetActorLocation() : FVector::ZeroVector;
 	const FTransform FocusTransform = FTransform(FocusPosition);
 	const FTransform DollyTransform = FTransform(-FVector::ForwardVector * Distance);
 	const FTransform RotationTransform = FTransform(FQuat::MakeFromEuler(FVector(0.0f, Pitch, Yaw)));
