[GLTFExporter] fix(orbit-camera): fix stuttering orbit motion

commit b881ef0a1dadc22b10e5cb5fc969b38e4989593d
Date: Wed, 5 Aug 2020 13:29:00 +0200

    fix(orbit-camera): fix stuttering orbit motion
    (cherry picked from commit e41bcdeeb32f6a69e393f93187d51ce89a1dd165​)

#skipci

diff --git a/Engine/Plugins/Enterprise/GLTFExporter/MigratedFromGit/Source/GLTFExporterRuntime/Private/Actors/GLTFOrbitCameraActor.cpp b/Engine/Plugins/Enterprise/GLTFExporter/MigratedFromGit/Source/GLTFExporterRuntime/Private/Actors/GLTFOrbitCameraActor.cpp
--- a/Engine/Plugins/Enterprise/GLTFExporter/MigratedFromGit/Source/GLTFExporterRuntime/Private/Actors/GLTFOrbitCameraActor.cpp
+++ b/Engine/Plugins/Enterprise/GLTFExporter/MigratedFromGit/Source/GLTFExporterRuntime/Private/Actors/GLTFOrbitCameraActor.cpp
@@ -11,6 +11,12 @@ namespace
 	{
 		return C * (T /= D) * T + B;
 	}
+
+	// Positive angles are indexed from 0 and negative angles from -1
+	int32 AngleCycleIndex(float Angle)
+	{
+		return (static_cast<int32>(Angle) / 360) + ((Angle < 0.0f) ? -1 : 0);
+	}
 } // Anonymous namespace
 
 #define DEBUGGLTFORBITCAMERA 0
@@ -47,7 +53,7 @@ void AGLTFOrbitCameraActor::PostEditChangeProperty(FPropertyChangedEvent& Proper
 
 	const FProperty* PropertyThatChanged = PropertyChangedEvent.Property;
 
-	if (PropertyThatChanged)
+	if (PropertyThatChanged != nullptr)
 	{
 		const FString PropertyName = PropertyThatChanged->GetName();
 
@@ -101,6 +107,16 @@ void AGLTFOrbitCameraActor::Tick(float DeltaSeconds)
 	Pitch = FMath::Lerp(Pitch, TargetPitch, Alpha);
 	//Distance = FMath::Lerp(Distance, TargetDistance, Alpha);
 
+	const int32 YawCycleIndex = AngleCycleIndex(Yaw);
+	const int32 TargetYawCycleIndex = AngleCycleIndex(TargetYaw);
+
+	// Clamp the angles to a positive 360 degrees while they are within the same cycle
+	if (YawCycleIndex == TargetYawCycleIndex && YawCycleIndex != 0)
+	{
+		Yaw = ClampYaw(Yaw);
+		TargetYaw = ClampYaw(TargetYaw);
+	}
+
 	const FVector FocusPosition = (this->Focus != nullptr && this->Focus != this) ? this->Focus->GetActorLocation() : FVector::ZeroVector;
 	const FTransform FocusTransform = FTransform(FocusPosition);
 	const FTransform DollyTransform = FTransform(-FVector::ForwardVector * Distance);
@@ -124,21 +140,7 @@ void AGLTFOrbitCameraActor::OnMouseX(float AxisValue)
 		return;
 	}
 
-	const float Delta = (TargetYaw + AxisValue * OrbitSensitivity) - Yaw;
-	const float Remainder = FMath::Fmod(Delta, 360.0f);
-
-	if (Remainder > 180.0f)
-	{
-		TargetYaw = Yaw - (360.0f - Remainder);
-}
-	else if (Remainder < -180.0f)
-	{
-		TargetYaw = Yaw + (360.0f + Remainder);
-	}
-	else
-	{
-		TargetYaw = Yaw + Remainder;
-	}
+	TargetYaw += AxisValue * OrbitSensitivity;
 
 #if DEBUGGLTFORBITCAMERA
 	UE_LOG(LogEditorGLTFOrbitCamera, Warning, TEXT("AGLTFOrbitCameraActor::OnMouseX(), %f"), AxisValue);
@@ -187,7 +189,10 @@ float AGLTFOrbitCameraActor::ClampPitch(float Value) const
 
 float AGLTFOrbitCameraActor::ClampYaw(float Value) const
 {
-	return FMath::Fmod(Value + 360.0f, 360.0f);
+	const int32 CycleIndex = AngleCycleIndex(Value);
+	const float Offset = 360.0f * static_cast<float>(FMath::Abs(FMath::Min(CycleIndex, 0)));
+
+	return FMath::Fmod(Value + Offset, 360.0f);
 }
 
 void AGLTFOrbitCameraActor::RemoveInertia()
