[GLTFExporter] refactor(orbit-camera): remove obsolete code

commit ceb628716a5af1aab79dd89c8976870aa88fd4ba
Date: Tue, 18 Aug 2020 13:39:52 +0200

    refactor(orbit-camera): remove obsolete code
    (cherry picked from commit b232c8faf64fb038bca64f8a8ce96555d368810b​)

#skipci

diff --git a/Engine/Plugins/Enterprise/GLTFExporter/MigratedFromGit/Source/GLTFExporterRuntime/Private/Actors/GLTFOrbitCameraActor.cpp b/Engine/Plugins/Enterprise/GLTFExporter/MigratedFromGit/Source/GLTFExporterRuntime/Private/Actors/GLTFOrbitCameraActor.cpp
--- a/Engine/Plugins/Enterprise/GLTFExporter/MigratedFromGit/Source/GLTFExporterRuntime/Private/Actors/GLTFOrbitCameraActor.cpp
+++ b/Engine/Plugins/Enterprise/GLTFExporter/MigratedFromGit/Source/GLTFExporterRuntime/Private/Actors/GLTFOrbitCameraActor.cpp
@@ -1,7 +1,6 @@
 // Copyright Epic Games, Inc. All Rights Reserved.
 
 #include "Actors/GLTFOrbitCameraActor.h"
-//#include "Kismet/GameplayStatics.h"
 
 namespace
 {
@@ -19,8 +18,6 @@ namespace
 	}
 } // Anonymous namespace
 
-#define DEBUGGLTFORBITCAMERA 0
-
 DEFINE_LOG_CATEGORY_STATIC(LogEditorGLTFOrbitCamera, Log, All);
 
 AGLTFOrbitCameraActor::AGLTFOrbitCameraActor(const FObjectInitializer& ObjectInitializer)
@@ -105,7 +102,6 @@ void AGLTFOrbitCameraActor::Tick(float DeltaSeconds)
 	const float Alpha = (OrbitInertia == 0.0f) ? 1.0f : FMath::Min(DeltaSeconds / OrbitInertia, 1.0f);
 	Yaw = FMath::Lerp(Yaw, TargetYaw, Alpha);
 	Pitch = FMath::Lerp(Pitch, TargetPitch, Alpha);
-	//Distance = FMath::Lerp(Distance, TargetDistance, Alpha);
 
 	const int32 YawCycleIndex = AngleCycleIndex(Yaw);
 	const int32 TargetYawCycleIndex = AngleCycleIndex(TargetYaw);
@@ -141,10 +137,6 @@ void AGLTFOrbitCameraActor::OnMouseX(float AxisValue)
 	}
 
 	TargetYaw += AxisValue * OrbitSensitivity;
-
-#if DEBUGGLTFORBITCAMERA
-	UE_LOG(LogEditorGLTFOrbitCamera, Warning, TEXT("AGLTFOrbitCameraActor::OnMouseX(), %f"), AxisValue);
-#endif
 }
 
 void AGLTFOrbitCameraActor::OnMouseY(float AxisValue)
@@ -155,10 +147,6 @@ void AGLTFOrbitCameraActor::OnMouseY(float AxisValue)
 	}
 
 	TargetPitch = ClampPitch(TargetPitch + AxisValue * OrbitSensitivity);
-
-#if DEBUGGLTFORBITCAMERA
-	UE_LOG(LogEditorGLTFOrbitCamera, Warning, TEXT("AGLTFOrbitCameraActor::OnMouseY(), %f"), AxisValue);
-#endif
 }
 
 void AGLTFOrbitCameraActor::OnMouseWheelAxis(float AxisValue)
@@ -171,10 +159,6 @@ void AGLTFOrbitCameraActor::OnMouseWheelAxis(float AxisValue)
 	DollyTime = DollyDuration;
 	TargetDistance = ClampDistance(TargetDistance + -AxisValue * DistanceSensitivity);
 	DollyStartDistance = Distance;
-
-#if DEBUGGLTFORBITCAMERA
-	UE_LOG(LogEditorGLTFOrbitCamera, Warning, TEXT("AGLTFOrbitCameraActor::OnMouseWheelAxis(), %f"), AxisValue);
-#endif
 }
 
 float AGLTFOrbitCameraActor::ClampDistance(float Value) const
