[GLTFExporter] refactor(orbit-camera): convert anonymous namespace constant to class property

commit 2f53f6bf0c9cefec6e841577ac5bb06623ab0cd4
Date: Tue, 18 Aug 2020 13:59:21 +0200

    refactor(orbit-camera): convert anonymous namespace constant to class property
    (cherry picked from commit cd4ae8c3dda067e627f0a42e4ce159dd6bfa557e​)

#skipci

diff --git a/Engine/Plugins/Enterprise/GLTFExporter/MigratedFromGit/Source/GLTFExporterRuntime/Private/Actors/GLTFOrbitCameraActor.cpp b/Engine/Plugins/Enterprise/GLTFExporter/MigratedFromGit/Source/GLTFExporterRuntime/Private/Actors/GLTFOrbitCameraActor.cpp
--- a/Engine/Plugins/Enterprise/GLTFExporter/MigratedFromGit/Source/GLTFExporterRuntime/Private/Actors/GLTFOrbitCameraActor.cpp
+++ b/Engine/Plugins/Enterprise/GLTFExporter/MigratedFromGit/Source/GLTFExporterRuntime/Private/Actors/GLTFOrbitCameraActor.cpp
@@ -4,8 +4,6 @@
 
 namespace
 {
-	const float DollyDuration = 0.2f;
-
 	float EaseInQuad(float T, float B, float C, float D)
 	{
 		return C * (T /= D) * T + B;
@@ -28,6 +26,7 @@ AGLTFOrbitCameraActor::AGLTFOrbitCameraActor(const FObjectInitializer& ObjectIni
 	DistanceMax(1000.0f),
 	PitchAngleMin(-90.0f),
 	PitchAngleMax(90.0f),
+	DollyDuration(0.2f),
 	OrbitInertia(0.07f),
 	OrbitSensitivity(30.0f),
 	DistanceSensitivity(50.0f),
diff --git a/Engine/Plugins/Enterprise/GLTFExporter/MigratedFromGit/Source/GLTFExporterRuntime/Public/Actors/GLTFOrbitCameraActor.h b/Engine/Plugins/Enterprise/GLTFExporter/MigratedFromGit/Source/GLTFExporterRuntime/Public/Actors/GLTFOrbitCameraActor.h
--- a/Engine/Plugins/Enterprise/GLTFExporter/MigratedFromGit/Source/GLTFExporterRuntime/Public/Actors/GLTFOrbitCameraActor.h
+++ b/Engine/Plugins/Enterprise/GLTFExporter/MigratedFromGit/Source/GLTFExporterRuntime/Public/Actors/GLTFOrbitCameraActor.h
@@ -65,6 +65,10 @@ public:
 	UPROPERTY(EditAnywhere, BlueprintReadWrite, Category = "GLTF Orbit Camera Actor")
 	float PitchAngleMax;
 
+	/* Duration (in seconds) that it takes the camera to complete a change in distance. */
+	UPROPERTY(VisibleAnywhere, BlueprintReadOnly, Category = "GLTF Orbit Camera Actor")
+	float DollyDuration;
+
 	/* Deceleration that occurs after orbital movement. The higher the value, the longer it takes to settle. */
 	UPROPERTY(VisibleAnywhere, BlueprintReadOnly, Category = "GLTF Orbit Camera Actor")
 	float OrbitInertia;
