[GLTFExporter] feat(orbit-camera): change default camera behavior to more closely match that in viewer

commit 87bce030684380fc34530033a241127760ab20a3
Date: Tue, 18 Aug 2020 14:52:10 +0200

    feat(orbit-camera): change default camera behavior to more closely match that in viewer
    (cherry picked from commit 1da7e282db3f50b9d19d7d79207af73dabde67f5​)

#skipci

diff --git a/Engine/Plugins/Enterprise/GLTFExporter/MigratedFromGit/Source/GLTFExporterRuntime/Private/Actors/GLTFOrbitCameraActor.cpp b/Engine/Plugins/Enterprise/GLTFExporter/MigratedFromGit/Source/GLTFExporterRuntime/Private/Actors/GLTFOrbitCameraActor.cpp
--- a/Engine/Plugins/Enterprise/GLTFExporter/MigratedFromGit/Source/GLTFExporterRuntime/Private/Actors/GLTFOrbitCameraActor.cpp
+++ b/Engine/Plugins/Enterprise/GLTFExporter/MigratedFromGit/Source/GLTFExporterRuntime/Private/Actors/GLTFOrbitCameraActor.cpp
@@ -4,11 +4,6 @@
 
 namespace
 {
-	float EaseInQuad(float T, float B, float C, float D)
-	{
-		return C * (T /= D) * T + B;
-	}
-
 	// Positive angles are indexed from 0 and negative angles from -1
 	int32 AngleCycleIndex(float Angle)
 	{
@@ -26,9 +21,9 @@ AGLTFOrbitCameraActor::AGLTFOrbitCameraActor(const FObjectInitializer& ObjectIni
 	DistanceMax(1000.0f),
 	PitchAngleMin(-90.0f),
 	PitchAngleMax(90.0f),
-	DollyDuration(0.2f),
-	OrbitInertia(0.07f),
-	OrbitSensitivity(30.0f),
+	DollyDuration(0.4f),
+	OrbitInertia(0.1f),
+	OrbitSensitivity(5.0f),
 	DistanceSensitivity(50.0f),
 	Distance(0.0f),
 	Pitch(0.0f),
@@ -94,8 +89,8 @@ void AGLTFOrbitCameraActor::Tick(float DeltaSeconds)
 
 	if (DollyTime != 0.0f)
 	{
-		Distance = EaseInQuad(DollyDuration - DollyTime, DollyStartDistance, TargetDistance - DollyStartDistance, DollyDuration);
 		DollyTime = FMath::Max(DollyTime - DeltaSeconds, 0.0f);
+		Distance = FMath::InterpEaseInOut(DollyStartDistance, TargetDistance, DollyTime != 0.0f ? ((DollyDuration - DollyTime) / DollyDuration) : 1.0f, 2.0f);
 	}
 
 	const float Alpha = (OrbitInertia == 0.0f) ? 1.0f : FMath::Min(DeltaSeconds / OrbitInertia, 1.0f);
