[GLTFExporter] docs(orbit-camera): add comments to clarify intent

commit 493d810a3ce03ce61465184a187559806498466c
Date: Wed, 19 Aug 2020 10:43:06 +0200

    docs(orbit-camera): add comments to clarify intent
    (cherry picked from commit baa8fe4e4ca3b54f364ec6d7ad1d3d31fb293ff0​)

#skipci

diff --git a/Engine/Plugins/Enterprise/GLTFExporter/MigratedFromGit/Source/GLTFExporterRuntime/Private/Actors/GLTFOrbitCameraActor.cpp b/Engine/Plugins/Enterprise/GLTFExporter/MigratedFromGit/Source/GLTFExporterRuntime/Private/Actors/GLTFOrbitCameraActor.cpp
--- a/Engine/Plugins/Enterprise/GLTFExporter/MigratedFromGit/Source/GLTFExporterRuntime/Private/Actors/GLTFOrbitCameraActor.cpp
+++ b/Engine/Plugins/Enterprise/GLTFExporter/MigratedFromGit/Source/GLTFExporterRuntime/Private/Actors/GLTFOrbitCameraActor.cpp
@@ -63,6 +63,7 @@ void AGLTFOrbitCameraActor::BeginPlay()
 {
 	Super::BeginPlay();
 
+	// TODO: Focus the scene AABB center if all else fails
 	if (this->Focus == nullptr || this->Focus == this)
 	{
 		UE_LOG(LogEditorGLTFOrbitCamera, Warning, TEXT("The camera focus must not be null or the camera itself."));
@@ -100,7 +101,8 @@ void AGLTFOrbitCameraActor::Tick(float DeltaSeconds)
 	const int32 YawCycleIndex = AngleCycleIndex(Yaw);
 	const int32 TargetYawCycleIndex = AngleCycleIndex(TargetYaw);
 
-	// Clamp the angles to a positive 360 degrees while they are within the same cycle
+	// Clamp the angles to a positive 360 degrees while they are within the same cycle.
+	// It is important that this doesn't happen during a transition as it will otherwise be skewed.
 	if (YawCycleIndex == TargetYawCycleIndex && YawCycleIndex != 0)
 	{
 		Yaw = ClampYaw(Yaw);
