[GLTFExporter] docs(orbit-camera): add comments to clarify intent

commit 18ddd0700f2284fde9f725f88ca3ad473de7ae23
Date: Fri, 21 Aug 2020 09:33:17 +0200

    docs(orbit-camera): add comments to clarify intent
    (cherry picked from commit 6ba36b70505ebc384be9da017e88fcd1d6fbe297​)

#skipci

diff --git a/Engine/Plugins/Enterprise/GLTFExporter/MigratedFromGit/Source/GLTFExporterRuntime/Private/Actors/GLTFOrbitCameraActor.cpp b/Engine/Plugins/Enterprise/GLTFExporter/MigratedFromGit/Source/GLTFExporterRuntime/Private/Actors/GLTFOrbitCameraActor.cpp
--- a/Engine/Plugins/Enterprise/GLTFExporter/MigratedFromGit/Source/GLTFExporterRuntime/Private/Actors/GLTFOrbitCameraActor.cpp
+++ b/Engine/Plugins/Enterprise/GLTFExporter/MigratedFromGit/Source/GLTFExporterRuntime/Private/Actors/GLTFOrbitCameraActor.cpp
@@ -4,7 +4,8 @@
 
 namespace
 {
-	// Positive angles are indexed from 0 and negative angles from -1
+	// Uniquely identifies an angle by how many times it has crossed the 0-360 degree range.
+	// Positive angles are indexed from 0 and negative angles from -1.
 	int32 AngleCycleIndex(float Angle)
 	{
 		return (static_cast<int32>(Angle) / 360) + ((Angle < 0.0f) ? -1 : 0);
