[GLTFExporter] refactor(orbit-camera): remove redundancy guards

commit 01c938b759c34af0629b58a7a2ea4815aa139e5e
Date: Fri, 21 Aug 2020 09:56:06 +0200

    refactor(orbit-camera): remove redundancy guards
    (cherry picked from commit 83e1a62d234d9e7b565565b35f1e8a5384cc139f​)

#skipci

diff --git a/Engine/Plugins/Enterprise/GLTFExporter/MigratedFromGit/Source/GLTFExporterRuntime/Private/Actors/GLTFOrbitCameraActor.cpp b/Engine/Plugins/Enterprise/GLTFExporter/MigratedFromGit/Source/GLTFExporterRuntime/Private/Actors/GLTFOrbitCameraActor.cpp
--- a/Engine/Plugins/Enterprise/GLTFExporter/MigratedFromGit/Source/GLTFExporterRuntime/Private/Actors/GLTFOrbitCameraActor.cpp
+++ b/Engine/Plugins/Enterprise/GLTFExporter/MigratedFromGit/Source/GLTFExporterRuntime/Private/Actors/GLTFOrbitCameraActor.cpp
@@ -128,31 +128,16 @@ void AGLTFOrbitCameraActor::PreInitializeComponents()
 
 void AGLTFOrbitCameraActor::OnMouseX(float AxisValue)
 {
-	if (AxisValue == 0.0f)
-	{
-		return;
-	}
-
 	TargetYaw += AxisValue * OrbitSensitivity;
 }
 
 void AGLTFOrbitCameraActor::OnMouseY(float AxisValue)
 {
-	if (AxisValue == 0.0f)
-	{
-		return;
-	}
-
 	TargetPitch = ClampPitch(TargetPitch + AxisValue * OrbitSensitivity);
 }
 
 void AGLTFOrbitCameraActor::OnMouseWheelAxis(float AxisValue)
 {
-	if (AxisValue == 0.0f)
-	{
-		return;
-	}
-
 	DollyTime = DollyDuration;
 	TargetDistance = ClampDistance(TargetDistance + -AxisValue * DistanceSensitivity);
 	DollyStartDistance = Distance;
