[GLTFExporter] refactor: fix initializer list variable style

commit b9898bd0fd3bb0be89f981f55244bcc01a6dcd28
Date: Mon, 24 Aug 2020 20:22:22 +0200

    refactor: fix initializer list variable style
    (cherry picked from commit 8bc2e40ab357fd23abedc94a7cf35a769f2b2197​)

#skipci

diff --git a/Engine/Plugins/Enterprise/GLTFExporter/MigratedFromGit/Source/GLTFExporterRuntime/Private/Actors/GLTFOrbitCameraActor.cpp b/Engine/Plugins/Enterprise/GLTFExporter/MigratedFromGit/Source/GLTFExporterRuntime/Private/Actors/GLTFOrbitCameraActor.cpp
--- a/Engine/Plugins/Enterprise/GLTFExporter/MigratedFromGit/Source/GLTFExporterRuntime/Private/Actors/GLTFOrbitCameraActor.cpp
+++ b/Engine/Plugins/Enterprise/GLTFExporter/MigratedFromGit/Source/GLTFExporterRuntime/Private/Actors/GLTFOrbitCameraActor.cpp
@@ -16,24 +16,23 @@ DEFINE_LOG_CATEGORY_STATIC(LogEditorGLTFOrbitCamera, Log, All);
 
 AGLTFOrbitCameraActor::AGLTFOrbitCameraActor(const FObjectInitializer& ObjectInitializer)
 	: Super(ObjectInitializer)
-	UPROPERTY(EditAnywhere, BlueprintReadWrite, Category = "GLTF Orbit Camera Actor"),
-	Focus(nullptr),
-	DistanceMin(100.0f),
-	DistanceMax(1000.0f),
-	PitchAngleMin(-90.0f),
-	PitchAngleMax(90.0f),
-	DollyDuration(0.4f),
-	OrbitInertia(0.1f),
-	OrbitSensitivity(5.0f),
-	DistanceSensitivity(50.0f),
-	Distance(0.0f),
-	Pitch(0.0f),
-	Yaw(0.0f),
-	TargetDistance(0.0f),
-	TargetPitch(0.0f),
-	TargetYaw(0.0f),
-	DollyTime(0.0f),
-	DollyStartDistance(0.0f)
+	, Focus(nullptr)
+	, DistanceMin(100.0f)
+	, DistanceMax(1000.0f)
+	, PitchAngleMin(-90.0f)
+	, PitchAngleMax(90.0f)
+	, DollyDuration(0.4f)
+	, OrbitInertia(0.1f)
+	, OrbitSensitivity(5.0f)
+	, DistanceSensitivity(50.0f)
+	, Distance(0.0f)
+	, Pitch(0.0f)
+	, Yaw(0.0f)
+	, TargetDistance(0.0f)
+	, TargetPitch(0.0f)
+	, TargetYaw(0.0f)
+	, DollyTime(0.0f)
+	, DollyStartDistance(0.0f)
 {
 	PrimaryActorTick.bCanEverTick = true;
 }
