[GLTFExporter] fix: change UPROPERTY VisibleAnywhere to EditAnywhere in orbit camera

commit 4859e5d880d1811e420b4d42ff379bab6c920e60
Date: Mon, 24 Aug 2020 20:23:26 +0200

    fix: change UPROPERTY VisibleAnywhere to EditAnywhere in orbit camera
    (cherry picked from commit 47528b94e7432bc63d1572d39bb9d84edd5e24ea​)

#skipci

diff --git a/Engine/Plugins/Enterprise/GLTFExporter/MigratedFromGit/Source/GLTFExporterRuntime/Public/Actors/GLTFOrbitCameraActor.h b/Engine/Plugins/Enterprise/GLTFExporter/MigratedFromGit/Source/GLTFExporterRuntime/Public/Actors/GLTFOrbitCameraActor.h
--- a/Engine/Plugins/Enterprise/GLTFExporter/MigratedFromGit/Source/GLTFExporterRuntime/Public/Actors/GLTFOrbitCameraActor.h
+++ b/Engine/Plugins/Enterprise/GLTFExporter/MigratedFromGit/Source/GLTFExporterRuntime/Public/Actors/GLTFOrbitCameraActor.h
@@ -66,19 +66,19 @@ public:
 	float PitchAngleMax;
 
 	/* Duration (in seconds) that it takes the camera to complete a change in distance. */
-	UPROPERTY(VisibleAnywhere, BlueprintReadOnly, Category = "GLTF Orbit Camera Actor")
+	UPROPERTY(EditAnywhere, BlueprintReadOnly, Category = "GLTF Orbit Camera Actor")
 	float DollyDuration;
 
 	/* Deceleration that occurs after orbital movement. The higher the value, the longer it takes to settle. */
-	UPROPERTY(VisibleAnywhere, BlueprintReadOnly, Category = "GLTF Orbit Camera Actor")
+	UPROPERTY(EditAnywhere, BlueprintReadOnly, Category = "GLTF Orbit Camera Actor")
 	float OrbitInertia;
 
 	/* Size of the orbital movement relative to user input. The higher the value, the faster it moves. */
-	UPROPERTY(VisibleAnywhere, BlueprintReadOnly, Category = "GLTF Orbit Camera Actor")
+	UPROPERTY(EditAnywhere, BlueprintReadOnly, Category = "GLTF Orbit Camera Actor")
 	float OrbitSensitivity;
-	
+
 	/* Size of the dolly movement relative to user input. The higher the value, the faster it moves. */
-	UPROPERTY(VisibleAnywhere, BlueprintReadOnly, Category = "GLTF Orbit Camera Actor")
+	UPROPERTY(EditAnywhere, BlueprintReadOnly, Category = "GLTF Orbit Camera Actor")
 	float DistanceSensitivity;
 
 private:
@@ -90,4 +90,4 @@ private:
 	float TargetYaw;
 	float DollyTime;
 	float DollyStartDistance;
-};
\ No newline at end of file
+};
