[GLTFExporter] refactor: remove unnecessary check and simplify logging

commit f30612f0b818a0d2d0230e40bfd2e718696fd732
Date: Mon, 24 Aug 2020 20:24:12 +0200

    refactor: remove unnecessary check and simplify logging
    (cherry picked from commit a667e6417832f294278c4a8fe8240820159adfb8​)

#skipci

diff --git a/Engine/Plugins/Enterprise/GLTFExporter/MigratedFromGit/Source/GLTFExporterRuntime/Private/Actors/GLTFOrbitCameraActor.cpp b/Engine/Plugins/Enterprise/GLTFExporter/MigratedFromGit/Source/GLTFExporterRuntime/Private/Actors/GLTFOrbitCameraActor.cpp
--- a/Engine/Plugins/Enterprise/GLTFExporter/MigratedFromGit/Source/GLTFExporterRuntime/Private/Actors/GLTFOrbitCameraActor.cpp
+++ b/Engine/Plugins/Enterprise/GLTFExporter/MigratedFromGit/Source/GLTFExporterRuntime/Private/Actors/GLTFOrbitCameraActor.cpp
@@ -8,12 +8,11 @@ namespace
 	// Positive angles are indexed from 0 and negative angles from -1.
 	int32 AngleCycleIndex(float Angle)
 	{
+		// TODO: is this similar to viewer camera?
 		return (static_cast<int32>(Angle) / 360) + ((Angle < 0.0f) ? -1 : 0);
 	}
 } // Anonymous namespace
 
-DEFINE_LOG_CATEGORY_STATIC(LogEditorGLTFOrbitCamera, Log, All);
-
 AGLTFOrbitCameraActor::AGLTFOrbitCameraActor(const FObjectInitializer& ObjectInitializer)
 	: Super(ObjectInitializer)
 	, Focus(nullptr)
@@ -37,38 +36,17 @@ AGLTFOrbitCameraActor::AGLTFOrbitCameraActor(const FObjectInitializer& ObjectIni
 	PrimaryActorTick.bCanEverTick = true;
 }
 
-#if WITH_EDITOR
-void AGLTFOrbitCameraActor::PostEditChangeProperty(FPropertyChangedEvent& PropertyChangedEvent)
-{
-	Super::PostEditChangeProperty(PropertyChangedEvent);
-
-	const FProperty* PropertyThatChanged = PropertyChangedEvent.Property;
-
-	if (PropertyThatChanged != nullptr)
-	{
-		const FString PropertyName = PropertyThatChanged->GetName();
-
-		if (PropertyName == TEXT("Focus"))
-		{
-			if (Focus == this)
-			{
-				UE_LOG(LogEditorGLTFOrbitCamera, Warning, TEXT("The camera cannot focus itself."));
-			}
-		}
-	}
-}
-#endif // WITH_EDITOR
-
 void AGLTFOrbitCameraActor::BeginPlay()
 {
 	Super::BeginPlay();
 
-	// TODO: Focus the scene AABB center if all else fails
-	if (this->Focus == nullptr || this->Focus == this)
+	if (this->Focus == nullptr)
 	{
-		UE_LOG(LogEditorGLTFOrbitCamera, Warning, TEXT("The camera focus must not be null or the camera itself."));
+		// TODO: if focus is null use the scene center (similar to viewer camera)
+		UE_LOG(LogTemp, Warning, TEXT("The camera focus must not be null"));
 	}
-	
+
+	// TODO: calculate distance and direction based on camera offset to focus (similar to viewer camera)
 	Distance = ClampDistance(Distance);
 	Pitch = ClampPitch(Pitch);
 	Yaw = ClampYaw(Yaw);
diff --git a/Engine/Plugins/Enterprise/GLTFExporter/MigratedFromGit/Source/GLTFExporterRuntime/Public/Actors/GLTFOrbitCameraActor.h b/Engine/Plugins/Enterprise/GLTFExporter/MigratedFromGit/Source/GLTFExporterRuntime/Public/Actors/GLTFOrbitCameraActor.h
--- a/Engine/Plugins/Enterprise/GLTFExporter/MigratedFromGit/Source/GLTFExporterRuntime/Public/Actors/GLTFOrbitCameraActor.h
+++ b/Engine/Plugins/Enterprise/GLTFExporter/MigratedFromGit/Source/GLTFExporterRuntime/Public/Actors/GLTFOrbitCameraActor.h
@@ -16,9 +16,6 @@ class GLTFEXPORTERRUNTIME_API AGLTFOrbitCameraActor : public ACameraActor
 	//~ Begin UObject Interface
 public:
 	AGLTFOrbitCameraActor(const FObjectInitializer& ObjectInitializer);
-#if WITH_EDITOR
-	virtual void PostEditChangeProperty(FPropertyChangedEvent& PropertyChangedEvent) override;
-#endif // WITH_EDITOR
 	//~ End UObject Interface
 
 	//~ Begin AActor Interface
