[GLTFExporter] fix: fix incorrect prefix in export options and remove no longer planned format option

commit d784463284349fcf9b972430c73188a0da5c0614
Date: Tue, 25 Aug 2020 21:09:32 +0200

    fix: fix incorrect prefix in export options and remove no longer planned format option
    (cherry picked from commit c8b04e751425a9bb5a17837ee46417f5423f2fbd​)

#skipci

diff --git a/Engine/Plugins/Enterprise/GLTFExporter/MigratedFromGit/Source/GLTFExporter/Private/GLTFExportOptions.cpp b/Engine/Plugins/Enterprise/GLTFExporter/MigratedFromGit/Source/GLTFExporter/Private/GLTFExportOptions.cpp
--- a/Engine/Plugins/Enterprise/GLTFExporter/MigratedFromGit/Source/GLTFExporter/Private/GLTFExportOptions.cpp
+++ b/Engine/Plugins/Enterprise/GLTFExporter/MigratedFromGit/Source/GLTFExporter/Private/GLTFExportOptions.cpp
@@ -21,14 +21,14 @@ UGLTFExportOptions::UGLTFExportOptions(const FObjectInitializer& ObjectInitializ
 	bExportVertexColors = true;
 	NormalExportMethod = EGLTFExporterNormalExportMethod::NormalsAndTangents;
 	bTangentDataQuantization = false;
-	bNormalizeUVCoordinates = EGLTFExporterNormalizeUVCoordinates::Never;
+	NormalizeUVCoordinates = EGLTFExporterNormalizeUVCoordinates::Never;
 	DefaultLevelOfDetail = 0;
 	TextureFormat = EGLTFExporterTextureFormat::PNG;
 	bExportLightmaps = true;
 	TextureHDREncoding = EGLTFExporterTextureHDREncoding::RGBM;
 	ExportScale = 0.01;
 	bExportHiddenInGame = false;
-	bExportLights = EGLTFExporterLightMobility::MovableAndStationary;
+	ExportLights = EGLTFExporterLightMobility::MovableAndStationary;
 	bExportCameras = true;
 	bExportOrbitalCameras = false;
 	bExportReflectionCaptures = false;
@@ -218,7 +218,7 @@ bool UGLTFExportOptions::CanEditChange(const FProperty* InProperty) const
 		PropertyFName == GET_MEMBER_NAME_CHECKED(ThisClass, DefaultLevelOfDetail) ||
 		PropertyFName == GET_MEMBER_NAME_CHECKED(ThisClass, bExportLightmaps) ||
 		PropertyFName == GET_MEMBER_NAME_CHECKED(ThisClass, ExportScale) ||
-		PropertyFName == GET_MEMBER_NAME_CHECKED(ThisClass, bExportLights) ||
+		PropertyFName == GET_MEMBER_NAME_CHECKED(ThisClass, ExportLights) ||
 		PropertyFName == GET_MEMBER_NAME_CHECKED(ThisClass, bExportCameras) ||
 		PropertyFName == GET_MEMBER_NAME_CHECKED(ThisClass, bExportHDRIBackdrops) ||
 		PropertyFName == GET_MEMBER_NAME_CHECKED(ThisClass, bExportVariantSets) ||
@@ -240,7 +240,7 @@ FIntPoint UGLTFExportOptions::GetDefaultMaterialBakeSize() const
 
 bool UGLTFExportOptions::ShouldExportLight(EComponentMobility::Type LightMobility) const
 {
-	switch (bExportLights)
+	switch (ExportLights)
 	{
 		case EGLTFExporterLightMobility::All:
 			return true;
diff --git a/Engine/Plugins/Enterprise/GLTFExporter/MigratedFromGit/Source/GLTFExporter/Public/GLTFExportOptions.h b/Engine/Plugins/Enterprise/GLTFExporter/MigratedFromGit/Source/GLTFExporter/Public/GLTFExportOptions.h
--- a/Engine/Plugins/Enterprise/GLTFExporter/MigratedFromGit/Source/GLTFExporter/Public/GLTFExportOptions.h
+++ b/Engine/Plugins/Enterprise/GLTFExporter/MigratedFromGit/Source/GLTFExporter/Public/GLTFExportOptions.h
@@ -44,8 +44,8 @@ UENUM(BlueprintType)
 enum class EGLTFExporterTextureFormat : uint8
 {
 	None,
-	PNG,
-	JPEG
+	PNG
+	// TODO: add more formats
 };
 
 UENUM(BlueprintType)
@@ -95,7 +95,7 @@ public:
 	uint32 bTangentDataQuantization : 1;
 
 	UPROPERTY(EditAnywhere, BlueprintReadWrite, config, Category = Mesh)
-	EGLTFExporterNormalizeUVCoordinates bNormalizeUVCoordinates;
+	EGLTFExporterNormalizeUVCoordinates NormalizeUVCoordinates;
 
 	UPROPERTY(EditAnywhere, BlueprintReadWrite, config, Category = Mesh)
 	int32 DefaultLevelOfDetail;
@@ -118,7 +118,7 @@ public:
 	uint32 bExportHiddenInGame : 1;
 
 	UPROPERTY(EditAnywhere, BlueprintReadWrite, config, Category = Scene)
-	EGLTFExporterLightMobility bExportLights;
+	EGLTFExporterLightMobility ExportLights;
 
 	UPROPERTY(EditAnywhere, BlueprintReadWrite, config, Category = Scene)
 	uint32 bExportCameras : 1;
