[GLTFExporter] feat: only add scale-vectors to (new) leaf-nodes for meshes

commit e1a258a8d35261b92bd965a5957e6f1b9b91b619
Date: Wed, 26 Aug 2020 09:56:03 +0200

    feat: only add scale-vectors to (new) leaf-nodes for meshes
    avoids "skewing" stemming from UE handling scale differently than glTF
    (cherry picked from commit 999d2fce50c6f7478b561569733b3ab8712a1d2c​)

#skipci

diff --git a/Engine/Plugins/Enterprise/GLTFExporter/MigratedFromGit/Source/GLTFExporter/Private/Converters/GLTFLevelConverters.cpp b/Engine/Plugins/Enterprise/GLTFExporter/MigratedFromGit/Source/GLTFExporter/Private/Converters/GLTFLevelConverters.cpp
--- a/Engine/Plugins/Enterprise/GLTFExporter/MigratedFromGit/Source/GLTFExporter/Private/Converters/GLTFLevelConverters.cpp
+++ b/Engine/Plugins/Enterprise/GLTFExporter/MigratedFromGit/Source/GLTFExporter/Private/Converters/GLTFLevelConverters.cpp
@@ -23,15 +23,17 @@ FGLTFJsonNodeIndex FGLTFSceneComponentConverter::Add(FGLTFConvertBuilder& Builde
 
 	// TODO: if node is root, then add it to the scene here to avoid any possible orphan nodes. this change requires level converter to support cyclic calls.
 
+	const FTransform Transform = SceneComponent->GetComponentTransform();
+	const FTransform RelativeTransform = bIsRootNode ? Transform : SceneComponent->GetRelativeTransform();
+	const FVector ParentScale = ParentComponent != nullptr ? ParentComponent->GetComponentTransform().GetScale3D() : FVector::OneVector;
+
 	const FGLTFJsonNodeIndex NodeIndex = Builder.AddChildNode(ParentNodeIndex);
 	FGLTFJsonNode& Node = Builder.GetNode(NodeIndex);
 	Node.Name = Name.IsEmpty() ? Owner->GetName() + TEXT("_") + SceneComponent->GetName() : Name;
 
-	// TODO: add support for non-uniform scaling (Unreal doesn't treat combined transforms as simple matrix multiplication)
-	const FTransform Transform = bIsRootNode ? SceneComponent->GetComponentTransform() : SceneComponent->GetRelativeTransform();
-	Node.Translation = FGLTFConverterUtility::ConvertPosition(Transform.GetTranslation(), Builder.ExportOptions->ExportScale);
-	Node.Rotation = FGLTFConverterUtility::ConvertRotation(Transform.GetRotation());
-	Node.Scale = FGLTFConverterUtility::ConvertScale(Transform.GetScale3D());
+	Node.Translation = FGLTFConverterUtility::ConvertPosition(RelativeTransform.GetTranslation() * ParentScale, Builder.ExportOptions->ExportScale);
+	Node.Rotation = FGLTFConverterUtility::ConvertRotation(RelativeTransform.GetRotation());
+	Node.Scale = FGLTFJsonVector3::One;
 
 	if (SceneComponent->bHiddenInGame) // TODO: make this configurable
 	{
@@ -39,17 +41,28 @@ FGLTFJsonNodeIndex FGLTFSceneComponentConverter::Add(FGLTFConvertBuilder& Builde
 	}
 	else if (const UStaticMeshComponent* StaticMeshComponent = Cast<UStaticMeshComponent>(SceneComponent))
 	{
-		Node.Mesh = Builder.GetOrAddMesh(StaticMeshComponent);
-		Node.LightMap = Builder.GetOrAddLightMap(StaticMeshComponent);
+		const FGLTFJsonNodeIndex LeafNodeIndex = Builder.AddChildNode(NodeIndex);
+		FGLTFJsonNode& LeafNode = Builder.GetNode(LeafNodeIndex);
+		LeafNode.Name = Node.Name + TEXT("_STATICMESH");
+
+		LeafNode.Scale = FGLTFConverterUtility::ConvertScale(Transform.GetScale3D());
+		LeafNode.Mesh = Builder.GetOrAddMesh(StaticMeshComponent);
+		LeafNode.LightMap = Builder.GetOrAddLightMap(StaticMeshComponent);
 	}
 	else if (const USkeletalMeshComponent* SkeletalMeshComponent = Cast<USkeletalMeshComponent>(SceneComponent))
 	{
-		Node.Mesh = Builder.GetOrAddMesh(SkeletalMeshComponent);
+		const FGLTFJsonNodeIndex LeafNodeIndex = Builder.AddChildNode(NodeIndex);
+		FGLTFJsonNode& LeafNode = Builder.GetNode(LeafNodeIndex);
+		LeafNode.Name = Node.Name + TEXT("_SKELETALMESH");
+
+		LeafNode.Scale = FGLTFConverterUtility::ConvertScale(Transform.GetScale3D());
+		LeafNode.Mesh = Builder.GetOrAddMesh(SkeletalMeshComponent);
 	}
 	else if (const UCameraComponent* CameraComponent = Cast<UCameraComponent>(SceneComponent))
 	{
 		if (Builder.ExportOptions->bExportCameras)
 		{
+			// TODO: do we need to support scale for cameras, i.e should we add a leaf-node?
 			// TODO: conversion of camera direction should be done in separate converter
 			FGLTFJsonNode CameraNode;
 			CameraNode.Name = Owner->GetName(); // TODO: choose a more unique name if owner is not ACameraActor
@@ -66,6 +79,7 @@ FGLTFJsonNodeIndex FGLTFSceneComponentConverter::Add(FGLTFConvertBuilder& Builde
 	{
 		if (Builder.ExportOptions->ShouldExportLight(LightComponent->Mobility))
 		{
+			// TODO: do we need to support scale for lights, i.e should we add a leaf-node?
 			// TODO: conversion of light direction should be done in separate converter
 			FGLTFJsonNode LightNode;
 			LightNode.Name = Owner->GetName(); // TODO: choose a more unique name if owner is not ALight
