[GLTFExporter] feat: handle mixed absolute / relative transforms by calculating a relative diff from two world transforms

commit be5265ed3611acf95dde4e5811bc8ea26dcb286e
Date: Wed, 26 Aug 2020 14:31:07 +0200

    feat: handle mixed absolute / relative transforms by calculating a relative diff from two world transforms
    (cherry picked from commit 7b460ec78c7365c971adccd4d8df3042786c11cb​)

#skipci

diff --git a/Engine/Plugins/Enterprise/GLTFExporter/MigratedFromGit/Source/GLTFExporter/Private/Converters/GLTFLevelConverters.cpp b/Engine/Plugins/Enterprise/GLTFExporter/MigratedFromGit/Source/GLTFExporter/Private/Converters/GLTFLevelConverters.cpp
--- a/Engine/Plugins/Enterprise/GLTFExporter/MigratedFromGit/Source/GLTFExporter/Private/Converters/GLTFLevelConverters.cpp
+++ b/Engine/Plugins/Enterprise/GLTFExporter/MigratedFromGit/Source/GLTFExporter/Private/Converters/GLTFLevelConverters.cpp
@@ -24,7 +24,7 @@ FGLTFJsonNodeIndex FGLTFSceneComponentConverter::Add(FGLTFConvertBuilder& Builde
 	// TODO: if node is root, then add it to the scene here to avoid any possible orphan nodes. this change requires level converter to support cyclic calls.
 
 	const FTransform Transform = SceneComponent->GetComponentTransform();
-	const FTransform RelativeTransform = bIsRootNode ? Transform : SceneComponent->GetRelativeTransform();
+	const FTransform RelativeTransform = bIsRootNode ? Transform : Transform.GetRelativeTransform(ParentComponent->GetComponentTransform());
 	const FVector ParentScale = ParentComponent != nullptr ? ParentComponent->GetComponentTransform().GetScale3D() : FVector::OneVector;
 
 	const FGLTFJsonNodeIndex NodeIndex = Builder.AddChildNode(ParentNodeIndex);
