[GLTFExporter] feat: avoid writing rotation quaternions that are within tolerance of the identity quaternion

commit 52c2a97d6a7d92997410dd55bac5893a6d52d42c
Date: Wed, 26 Aug 2020 14:31:56 +0200

    feat: avoid writing rotation quaternions that are within tolerance of the identity quaternion
    (cherry picked from commit 620180c5d161fee93cc8ab2b95025f06cdbcafc7​)

#skipci

diff --git a/Engine/Plugins/Enterprise/GLTFExporter/MigratedFromGit/Source/GLTFExporter/Private/Converters/GLTFConverterUtility.h b/Engine/Plugins/Enterprise/GLTFExporter/MigratedFromGit/Source/GLTFExporter/Private/Converters/GLTFConverterUtility.h
--- a/Engine/Plugins/Enterprise/GLTFExporter/MigratedFromGit/Source/GLTFExporter/Private/Converters/GLTFConverterUtility.h
+++ b/Engine/Plugins/Enterprise/GLTFExporter/MigratedFromGit/Source/GLTFExporter/Private/Converters/GLTFConverterUtility.h
@@ -83,6 +83,13 @@ struct FGLTFConverterUtility
 
 		// Not checking if quaternion is normalized
 		// e.g. some sources use non-unit Quats for rotation tangents
+
+		// Return the identity quaternion when possible (depending on tolerance)
+		if (Rotation.Equals(FQuat::Identity))
+		{
+			return FGLTFJsonQuaternion::Identity;
+		}
+
 		return { -Rotation.X, -Rotation.Z, -Rotation.Y, Rotation.W };
 	}
 
