[GLTFExporter] chore: clarify TODOs for light- and camera-scale

commit 206ce60d58b9f2bcc5655ee851f466b5144c1466
Date: Wed, 26 Aug 2020 14:37:32 +0200

    chore: clarify TODOs for light- and camera-scale
    (cherry picked from commit c15de6e26b2af2fde3075f0a1f2578695a1aca6c​)

#skipci

diff --git a/Engine/Plugins/Enterprise/GLTFExporter/MigratedFromGit/Source/GLTFExporter/Private/Converters/GLTFLevelConverters.cpp b/Engine/Plugins/Enterprise/GLTFExporter/MigratedFromGit/Source/GLTFExporter/Private/Converters/GLTFLevelConverters.cpp
--- a/Engine/Plugins/Enterprise/GLTFExporter/MigratedFromGit/Source/GLTFExporter/Private/Converters/GLTFLevelConverters.cpp
+++ b/Engine/Plugins/Enterprise/GLTFExporter/MigratedFromGit/Source/GLTFExporter/Private/Converters/GLTFLevelConverters.cpp
@@ -62,7 +62,7 @@ FGLTFJsonNodeIndex FGLTFSceneComponentConverter::Add(FGLTFConvertBuilder& Builde
 	{
 		if (Builder.ExportOptions->bExportCameras)
 		{
-			// TODO: do we need to support scale for cameras, i.e should we add a leaf-node?
+			// TODO: do we need to support scale for camera-nodes?
 			// TODO: conversion of camera direction should be done in separate converter
 			FGLTFJsonNode CameraNode;
 			CameraNode.Name = Owner->GetName(); // TODO: choose a more unique name if owner is not ACameraActor
@@ -79,7 +79,7 @@ FGLTFJsonNodeIndex FGLTFSceneComponentConverter::Add(FGLTFConvertBuilder& Builde
 	{
 		if (Builder.ExportOptions->ShouldExportLight(LightComponent->Mobility))
 		{
-			// TODO: do we need to support scale for lights, i.e should we add a leaf-node?
+			// TODO: do we need to support scale for light-nodes?
 			// TODO: conversion of light direction should be done in separate converter
 			FGLTFJsonNode LightNode;
 			LightNode.Name = Owner->GetName(); // TODO: choose a more unique name if owner is not ALight
