[GLTFExporter] refactor: modify mesh-nodes to follow the same naming- and code-standard as for light- and camera-nodes

commit e1353ec7bc3eb4a754f4844939024af3590b1b6c
Date: Wed, 26 Aug 2020 14:42:57 +0200

    refactor: modify mesh-nodes to follow the same naming- and code-standard as for light- and camera-nodes
    (cherry picked from commit 95870957a9fc99239eeb140ac50780d2b9e220a8​)

#skipci

diff --git a/Engine/Plugins/Enterprise/GLTFExporter/MigratedFromGit/Source/GLTFExporter/Private/Converters/GLTFLevelConverters.cpp b/Engine/Plugins/Enterprise/GLTFExporter/MigratedFromGit/Source/GLTFExporter/Private/Converters/GLTFLevelConverters.cpp
--- a/Engine/Plugins/Enterprise/GLTFExporter/MigratedFromGit/Source/GLTFExporter/Private/Converters/GLTFLevelConverters.cpp
+++ b/Engine/Plugins/Enterprise/GLTFExporter/MigratedFromGit/Source/GLTFExporter/Private/Converters/GLTFLevelConverters.cpp
@@ -41,22 +41,20 @@ FGLTFJsonNodeIndex FGLTFSceneComponentConverter::Add(FGLTFConvertBuilder& Builde
 	}
 	else if (const UStaticMeshComponent* StaticMeshComponent = Cast<UStaticMeshComponent>(SceneComponent))
 	{
-		const FGLTFJsonNodeIndex LeafNodeIndex = Builder.AddChildNode(NodeIndex);
-		FGLTFJsonNode& LeafNode = Builder.GetNode(LeafNodeIndex);
-		LeafNode.Name = Node.Name + TEXT("_STATICMESH");
-
-		LeafNode.Scale = FGLTFConverterUtility::ConvertScale(Transform.GetScale3D());
-		LeafNode.Mesh = Builder.GetOrAddMesh(StaticMeshComponent);
-		LeafNode.LightMap = Builder.GetOrAddLightMap(StaticMeshComponent);
+		FGLTFJsonNode MeshNode;
+		MeshNode.Name = Node.Name + TEXT("_STATICMESH");
+		MeshNode.Scale = FGLTFConverterUtility::ConvertScale(Transform.GetScale3D());
+		MeshNode.Mesh = Builder.GetOrAddMesh(StaticMeshComponent);
+		MeshNode.LightMap = Builder.GetOrAddLightMap(StaticMeshComponent);
+		Builder.AddChildNode(NodeIndex, MeshNode);
 	}
 	else if (const USkeletalMeshComponent* SkeletalMeshComponent = Cast<USkeletalMeshComponent>(SceneComponent))
 	{
-		const FGLTFJsonNodeIndex LeafNodeIndex = Builder.AddChildNode(NodeIndex);
-		FGLTFJsonNode& LeafNode = Builder.GetNode(LeafNodeIndex);
-		LeafNode.Name = Node.Name + TEXT("_SKELETALMESH");
-
-		LeafNode.Scale = FGLTFConverterUtility::ConvertScale(Transform.GetScale3D());
-		LeafNode.Mesh = Builder.GetOrAddMesh(SkeletalMeshComponent);
+		FGLTFJsonNode MeshNode;
+		MeshNode.Name = Node.Name + TEXT("_SKELETALMESH");
+		MeshNode.Scale = FGLTFConverterUtility::ConvertScale(Transform.GetScale3D());
+		MeshNode.Mesh = Builder.GetOrAddMesh(SkeletalMeshComponent);
+		Builder.AddChildNode(NodeIndex, MeshNode);
 	}
 	else if (const UCameraComponent* CameraComponent = Cast<UCameraComponent>(SceneComponent))
 	{
