[GLTFExporter] chore: add TODOs for incorrect node-mapping in extensions due to the extra nodes used for meshes

commit eafdccd55f96c2ea0193213082fcae96ddb898eb
Date: Wed, 26 Aug 2020 14:48:09 +0200

    chore: add TODOs for incorrect node-mapping in extensions due to the extra nodes used for meshes
    (cherry picked from commit 04d86c7e29945b86aa375721a91a404642a9d966​)

#skipci

diff --git a/Engine/Plugins/Enterprise/GLTFExporter/MigratedFromGit/Source/GLTFExporter/Private/Converters/GLTFLevelConverters.cpp b/Engine/Plugins/Enterprise/GLTFExporter/MigratedFromGit/Source/GLTFExporter/Private/Converters/GLTFLevelConverters.cpp
--- a/Engine/Plugins/Enterprise/GLTFExporter/MigratedFromGit/Source/GLTFExporter/Private/Converters/GLTFLevelConverters.cpp
+++ b/Engine/Plugins/Enterprise/GLTFExporter/MigratedFromGit/Source/GLTFExporter/Private/Converters/GLTFLevelConverters.cpp
@@ -41,6 +41,12 @@ FGLTFJsonNodeIndex FGLTFSceneComponentConverter::Add(FGLTFConvertBuilder& Builde
 	}
 	else if (const UStaticMeshComponent* StaticMeshComponent = Cast<UStaticMeshComponent>(SceneComponent))
 	{
+		// TODO: using a separate child-node for the mesh breaks variants and possibly other
+		// extensions that count on the mesh being attached to the original node.
+		// We need to either find a way to solve the scale-issues without introducing these
+		// extra nodes, or add some form of mapping that allows the extensions to find the
+		// exact node that the mesh is attached to.
+
 		FGLTFJsonNode MeshNode;
 		MeshNode.Name = Node.Name + TEXT("_STATICMESH");
 		MeshNode.Scale = FGLTFConverterUtility::ConvertScale(Transform.GetScale3D());
@@ -50,6 +56,12 @@ FGLTFJsonNodeIndex FGLTFSceneComponentConverter::Add(FGLTFConvertBuilder& Builde
 	}
 	else if (const USkeletalMeshComponent* SkeletalMeshComponent = Cast<USkeletalMeshComponent>(SceneComponent))
 	{
+		// TODO: using a separate child-node for the mesh breaks variants and possibly other
+		// extensions that count on the mesh being attached to the original node.
+		// We need to either find a way to solve the scale-issues without introducing these
+		// extra nodes, or add some form of mapping that allows the extensions to find the
+		// exact node that the mesh is attached to.
+
 		FGLTFJsonNode MeshNode;
 		MeshNode.Name = Node.Name + TEXT("_SKELETALMESH");
 		MeshNode.Scale = FGLTFConverterUtility::ConvertScale(Transform.GetScale3D());
