[GLTFExporter] feat: add and implement export-option for supporting non-uniform scale

commit a10d7163b132d37e40aae7f600dca0d93c11d9fb
Date: Tue, 1 Sep 2020 08:41:41 +0200

    feat: add and implement export-option for supporting non-uniform scale
    (cherry picked from commit 0ee83eb70c57b5a1b676585e0679435476c482fc​)

#skipci

diff --git a/Engine/Plugins/Enterprise/GLTFExporter/MigratedFromGit/Source/GLTFExporter/Private/Converters/GLTFLevelConverters.cpp b/Engine/Plugins/Enterprise/GLTFExporter/MigratedFromGit/Source/GLTFExporter/Private/Converters/GLTFLevelConverters.cpp
--- a/Engine/Plugins/Enterprise/GLTFExporter/MigratedFromGit/Source/GLTFExporter/Private/Converters/GLTFLevelConverters.cpp
+++ b/Engine/Plugins/Enterprise/GLTFExporter/MigratedFromGit/Source/GLTFExporter/Private/Converters/GLTFLevelConverters.cpp
@@ -17,6 +17,7 @@ FGLTFJsonNodeIndex FGLTFSceneComponentConverter::Add(FGLTFConvertBuilder& Builde
 
 	const bool bIsRootComponent = Owner->GetRootComponent() == SceneComponent;
 	const bool bIsRootNode = bIsRootComponent && FGLTFActorUtility::IsRootActor(Owner, Builder.bSelectedActorsOnly);
+	const bool bSupportNonUniformScale = Builder.ExportOptions->bSupportNonUniformScale;
 
 	const USceneComponent* ParentComponent = SceneComponent->GetAttachParent();
 	const FGLTFJsonNodeIndex ParentNodeIndex = Builder.GetOrAddNode(ParentComponent);
@@ -25,15 +26,17 @@ FGLTFJsonNodeIndex FGLTFSceneComponentConverter::Add(FGLTFConvertBuilder& Builde
 
 	const FTransform Transform = SceneComponent->GetComponentTransform();
 	const FTransform RelativeTransform = bIsRootNode ? Transform : Transform.GetRelativeTransform(ParentComponent->GetComponentTransform());
+	
 	const FVector ParentScale = ParentComponent != nullptr ? ParentComponent->GetComponentTransform().GetScale3D() : FVector::OneVector;
+	const FVector Translation = bSupportNonUniformScale ? RelativeTransform.GetTranslation() * ParentScale : RelativeTransform.GetTranslation();
+	const FVector Scale = bSupportNonUniformScale ? FVector::OneVector : RelativeTransform.GetScale3D();
 
 	const FGLTFJsonNodeIndex NodeIndex = Builder.AddChildNode(ParentNodeIndex);
 	FGLTFJsonNode& Node = Builder.GetNode(NodeIndex);
 	Node.Name = Name.IsEmpty() ? Owner->GetName() + TEXT("_") + SceneComponent->GetName() : Name;
-
-	Node.Translation = FGLTFConverterUtility::ConvertPosition(RelativeTransform.GetTranslation() * ParentScale, Builder.ExportOptions->ExportScale);
+	Node.Translation = FGLTFConverterUtility::ConvertPosition(Translation, Builder.ExportOptions->ExportScale);
 	Node.Rotation = FGLTFConverterUtility::ConvertRotation(RelativeTransform.GetRotation());
-	Node.Scale = FGLTFJsonVector3::One;
+	Node.Scale = FGLTFConverterUtility::ConvertScale(Scale);
 
 	if (SceneComponent->bHiddenInGame) // TODO: make this configurable
 	{
@@ -47,12 +50,20 @@ FGLTFJsonNodeIndex FGLTFSceneComponentConverter::Add(FGLTFConvertBuilder& Builde
 		// extra nodes, or add some form of mapping that allows the extensions to find the
 		// exact node that the mesh is attached to.
 
-		FGLTFJsonNode MeshNode;
-		MeshNode.Name = Node.Name + TEXT("_STATICMESH");
-		MeshNode.Scale = FGLTFConverterUtility::ConvertScale(Transform.GetScale3D());
-		MeshNode.Mesh = Builder.GetOrAddMesh(StaticMeshComponent);
-		MeshNode.LightMap = Builder.GetOrAddLightMap(StaticMeshComponent);
-		Builder.AddChildNode(NodeIndex, MeshNode);
+		if (bSupportNonUniformScale)
+		{
+			FGLTFJsonNode MeshNode;
+			MeshNode.Name = Node.Name + TEXT("_Mesh");
+			MeshNode.Scale = FGLTFConverterUtility::ConvertScale(Transform.GetScale3D());
+			MeshNode.Mesh = Builder.GetOrAddMesh(StaticMeshComponent);
+			MeshNode.LightMap = Builder.GetOrAddLightMap(StaticMeshComponent);
+			Builder.AddChildNode(NodeIndex, MeshNode);
+		}
+		else
+		{
+			Node.Mesh = Builder.GetOrAddMesh(StaticMeshComponent);
+			Node.LightMap = Builder.GetOrAddLightMap(StaticMeshComponent);
+		}
 	}
 	else if (const USkeletalMeshComponent* SkeletalMeshComponent = Cast<USkeletalMeshComponent>(SceneComponent))
 	{
@@ -62,21 +73,28 @@ FGLTFJsonNodeIndex FGLTFSceneComponentConverter::Add(FGLTFConvertBuilder& Builde
 		// extra nodes, or add some form of mapping that allows the extensions to find the
 		// exact node that the mesh is attached to.
 
-		FGLTFJsonNode MeshNode;
-		MeshNode.Name = Node.Name + TEXT("_SKELETALMESH");
-		MeshNode.Scale = FGLTFConverterUtility::ConvertScale(Transform.GetScale3D());
-		MeshNode.Mesh = Builder.GetOrAddMesh(SkeletalMeshComponent);
-		Builder.AddChildNode(NodeIndex, MeshNode);
+		if (bSupportNonUniformScale)
+		{
+			FGLTFJsonNode MeshNode;
+			MeshNode.Name = Node.Name + TEXT("_Mesh");
+			MeshNode.Scale = FGLTFConverterUtility::ConvertScale(Transform.GetScale3D());
+			MeshNode.Mesh = Builder.GetOrAddMesh(SkeletalMeshComponent);
+			Builder.AddChildNode(NodeIndex, MeshNode);
+		}
+		else
+		{
+			Node.Mesh = Builder.GetOrAddMesh(SkeletalMeshComponent);
+		}
 	}
 	else if (const UCameraComponent* CameraComponent = Cast<UCameraComponent>(SceneComponent))
 	{
 		if (Builder.ExportOptions->bExportCameras)
 		{
-			// TODO: do we need to support scale for camera-nodes?
 			// TODO: conversion of camera direction should be done in separate converter
 			FGLTFJsonNode CameraNode;
 			CameraNode.Name = Owner->GetName(); // TODO: choose a more unique name if owner is not ACameraActor
 			CameraNode.Rotation = FGLTFConverterUtility::ConvertCameraDirection();
+			CameraNode.Scale = FGLTFConverterUtility::ConvertScale(bSupportNonUniformScale ? Transform.GetScale3D() : FVector::OneVector);
 			CameraNode.Camera = Builder.GetOrAddCamera(CameraComponent, CameraNode.Name);
 			Builder.AddChildNode(NodeIndex, CameraNode);
 		}
@@ -89,11 +107,11 @@ FGLTFJsonNodeIndex FGLTFSceneComponentConverter::Add(FGLTFConvertBuilder& Builde
 	{
 		if (Builder.ExportOptions->ShouldExportLight(LightComponent->Mobility))
 		{
-			// TODO: do we need to support scale for light-nodes?
 			// TODO: conversion of light direction should be done in separate converter
 			FGLTFJsonNode LightNode;
 			LightNode.Name = Owner->GetName(); // TODO: choose a more unique name if owner is not ALight
 			LightNode.Rotation = FGLTFConverterUtility::ConvertLightDirection();
+			LightNode.Scale = FGLTFConverterUtility::ConvertScale(bSupportNonUniformScale ? Transform.GetScale3D() : FVector::OneVector);
 			LightNode.Light = Builder.GetOrAddLight(LightComponent, LightNode.Name);
 			Builder.AddChildNode(NodeIndex, LightNode);
 		}
diff --git a/Engine/Plugins/Enterprise/GLTFExporter/MigratedFromGit/Source/GLTFExporter/Private/GLTFExportOptions.cpp b/Engine/Plugins/Enterprise/GLTFExporter/MigratedFromGit/Source/GLTFExporter/Private/GLTFExportOptions.cpp
--- a/Engine/Plugins/Enterprise/GLTFExporter/MigratedFromGit/Source/GLTFExporter/Private/GLTFExportOptions.cpp
+++ b/Engine/Plugins/Enterprise/GLTFExporter/MigratedFromGit/Source/GLTFExporter/Private/GLTFExportOptions.cpp
@@ -35,6 +35,7 @@ UGLTFExportOptions::UGLTFExportOptions(const FObjectInitializer& ObjectInitializ
 	bExportHDRIBackdrops = true;
 	bExportVariantSets = true;
 	bExportInteractionHotspots = false;
+	bSupportNonUniformScale = false;
 	bBundleWebViewer = false;
 	bExportPreviewMesh = true;
 	bAllExtensionsRequired = false;
@@ -222,6 +223,7 @@ bool UGLTFExportOptions::CanEditChange(const FProperty* InProperty) const
 		PropertyFName == GET_MEMBER_NAME_CHECKED(ThisClass, bExportCameras) ||
 		PropertyFName == GET_MEMBER_NAME_CHECKED(ThisClass, bExportHDRIBackdrops) ||
 		PropertyFName == GET_MEMBER_NAME_CHECKED(ThisClass, bExportVariantSets) ||
+		PropertyFName == GET_MEMBER_NAME_CHECKED(ThisClass, bSupportNonUniformScale) ||
 		PropertyFName == GET_MEMBER_NAME_CHECKED(ThisClass, bExportPreviewMesh))
 	{
 		return true;
diff --git a/Engine/Plugins/Enterprise/GLTFExporter/MigratedFromGit/Source/GLTFExporter/Public/GLTFExportOptions.h b/Engine/Plugins/Enterprise/GLTFExporter/MigratedFromGit/Source/GLTFExporter/Public/GLTFExportOptions.h
--- a/Engine/Plugins/Enterprise/GLTFExporter/MigratedFromGit/Source/GLTFExporter/Public/GLTFExportOptions.h
+++ b/Engine/Plugins/Enterprise/GLTFExporter/MigratedFromGit/Source/GLTFExporter/Public/GLTFExportOptions.h
@@ -138,6 +138,9 @@ public:
 	UPROPERTY(EditAnywhere, BlueprintReadWrite, config, Category = Scene)
 	uint32 bExportInteractionHotspots : 1;
 
+	UPROPERTY(EditAnywhere, BlueprintReadWrite, config, Category = Scene)
+	uint32 bSupportNonUniformScale : 1;
+
 	UPROPERTY(EditAnywhere, BlueprintReadWrite, config, Category = Exporter)
     uint32 bBundleWebViewer : 1;
 
