[GLTFExporter] feat: add matrix in json nodes and conversion of transform

commit 0e124a9101bdf454c4ba53f522c10456d7b18a6d
Date: Wed, 26 Aug 2020 15:40:46 +0200

    feat: add matrix in json nodes and conversion of transform
    (cherry picked from commit dca333d7ecf3f47becdf31af1c8da938dcc71250​)
    # Conflicts from cherry-pick:
    #	Plugins/GLTFExporter/Source/GLTFExporter/Private/Converters/GLTFLevelConverters.cpp

#skipci

diff --git a/Engine/Plugins/Enterprise/GLTFExporter/MigratedFromGit/Source/GLTFExporter/Private/Converters/GLTFConverterUtility.h b/Engine/Plugins/Enterprise/GLTFExporter/MigratedFromGit/Source/GLTFExporter/Private/Converters/GLTFConverterUtility.h
--- a/Engine/Plugins/Enterprise/GLTFExporter/MigratedFromGit/Source/GLTFExporter/Private/Converters/GLTFConverterUtility.h
+++ b/Engine/Plugins/Enterprise/GLTFExporter/MigratedFromGit/Source/GLTFExporter/Private/Converters/GLTFConverterUtility.h
@@ -8,6 +8,7 @@
 #include "Json/GLTFJsonVector3.h"
 #include "Json/GLTFJsonVector4.h"
 #include "Json/GLTFJsonColor4.h"
+#include "Json/GLTFJsonMatrix4.h"
 #include "Json/GLTFJsonQuaternion.h"
 #include "Engine/EngineTypes.h"
 #include "Json/GLTFJsonCamera.h"
@@ -93,22 +94,37 @@ struct FGLTFConverterUtility
 		return { -Rotation.X, -Rotation.Z, -Rotation.Y, Rotation.W };
 	}
 
-	static FMatrix ConvertMatrix(const FMatrix& Matrix)
+	static FGLTFJsonMatrix4 ConvertMatrix(const FMatrix& Matrix)
 	{
 		// Unreal stores matrix elements in row major order.
 		// glTF stores matrix elements in column major order.
 
-		FMatrix Result;
+		FGLTFJsonMatrix4 Result;
 		for (int32 Row = 0; Row < 4; ++Row)
 		{
 			for (int32 Col = 0; Col < 4; ++Col)
 			{
-				Result.M[Col][Row] = Matrix.M[Row][Col];
+				Result(Row, Col) = Matrix.M[Row][Col];
 			}
 		}
 		return Result;
 	}
 
+	static FGLTFJsonMatrix4 ConvertTransform(const FTransform& Transform, const float ConversionScale)
+	{
+		const FQuat Rotation = Transform.GetRotation();
+		const FVector Translation = Transform.GetTranslation();
+		const FVector Scale = Transform.GetScale3D();
+
+		const FQuat ConvertedRotation(-Rotation.X, -Rotation.Z, -Rotation.Y, Rotation.W);
+		const FVector ConvertedTranslation = FVector(Translation.X, Translation.Z, Translation.Y) * ConversionScale;
+		const FVector ConvertedScale(Scale.X, Scale.Z, Scale.Y);
+
+		const FTransform ConvertedTransform(ConvertedRotation, ConvertedTranslation, ConvertedScale);
+		const FMatrix Matrix = ConvertedTransform.ToMatrixWithScale(); // TODO: should it be ToMatrixNoScale()?
+		return ConvertMatrix(Matrix);
+	}
+
 	static FGLTFJsonQuaternion ConvertCameraDirection()
 	{
 		// Unreal uses +X axis as camera direction in Unreal coordinates.
diff --git a/Engine/Plugins/Enterprise/GLTFExporter/MigratedFromGit/Source/GLTFExporter/Private/Converters/GLTFLevelConverters.cpp b/Engine/Plugins/Enterprise/GLTFExporter/MigratedFromGit/Source/GLTFExporter/Private/Converters/GLTFLevelConverters.cpp
--- a/Engine/Plugins/Enterprise/GLTFExporter/MigratedFromGit/Source/GLTFExporter/Private/Converters/GLTFLevelConverters.cpp
+++ b/Engine/Plugins/Enterprise/GLTFExporter/MigratedFromGit/Source/GLTFExporter/Private/Converters/GLTFLevelConverters.cpp
@@ -34,9 +34,11 @@ FGLTFJsonNodeIndex FGLTFSceneComponentConverter::Add(FGLTFConvertBuilder& Builde
 	const FGLTFJsonNodeIndex NodeIndex = Builder.AddChildNode(ParentNodeIndex);
 	FGLTFJsonNode& Node = Builder.GetNode(NodeIndex);
 	Node.Name = Name.IsEmpty() ? Owner->GetName() + TEXT("_") + SceneComponent->GetName() : Name;
-	Node.Translation = FGLTFConverterUtility::ConvertPosition(Translation, Builder.ExportOptions->ExportScale);
-	Node.Rotation = FGLTFConverterUtility::ConvertRotation(RelativeTransform.GetRotation());
-	Node.Scale = FGLTFConverterUtility::ConvertScale(Scale);
+
+	// TODO: add support for non-uniform scaling (Unreal doesn't treat combined transforms as simple matrix multiplication)
+	const FTransform Transform = bIsRootNode ? SceneComponent->GetComponentTransform() : SceneComponent->GetRelativeTransform();
+	Node.Matrix = FGLTFConverterUtility::ConvertTransform(Transform, Builder.ExportOptions->ExportScale);
+	Node.bUseMatrix = true;
 
 	if (SceneComponent->bHiddenInGame) // TODO: make this configurable
 	{
diff --git a/Engine/Plugins/Enterprise/GLTFExporter/MigratedFromGit/Source/GLTFExporter/Private/Json/GLTFJsonMatrix4.cpp b/Engine/Plugins/Enterprise/GLTFExporter/MigratedFromGit/Source/GLTFExporter/Private/Json/GLTFJsonMatrix4.cpp
new file mode 100644
--- /dev/null
+++ b/Engine/Plugins/Enterprise/GLTFExporter/MigratedFromGit/Source/GLTFExporter/Private/Json/GLTFJsonMatrix4.cpp
@@ -0,0 +1,11 @@
+// Copyright Epic Games, Inc. All Rights Reserved.
+
+#include "Json/GLTFJsonMatrix4.h"
+
+const FGLTFJsonMatrix4 FGLTFJsonMatrix4::Identity =
+{
+	1, 0 ,0, 0,
+	0, 1, 0, 0,
+	0, 0, 1, 0,
+	0, 0, 0, 1
+};
diff --git a/Engine/Plugins/Enterprise/GLTFExporter/MigratedFromGit/Source/GLTFExporter/Private/Json/GLTFJsonMatrix4.h b/Engine/Plugins/Enterprise/GLTFExporter/MigratedFromGit/Source/GLTFExporter/Private/Json/GLTFJsonMatrix4.h
new file mode 100644
--- /dev/null
+++ b/Engine/Plugins/Enterprise/GLTFExporter/MigratedFromGit/Source/GLTFExporter/Private/Json/GLTFJsonMatrix4.h
@@ -0,0 +1,49 @@
+// Copyright Epic Games, Inc. All Rights Reserved.
+
+#pragma once
+
+#include "Json/GLTFJsonUtility.h"
+#include "Serialization/JsonSerializer.h"
+
+struct FGLTFJsonMatrix4
+{
+	float M[16];
+
+	static const FGLTFJsonMatrix4 Identity;
+
+	template <class CharType = TCHAR, class PrintPolicy = TPrettyJsonPrintPolicy<CharType>>
+	void WriteArray(TJsonWriter<CharType, PrintPolicy>& JsonWriter) const
+	{
+		JsonWriter.WriteArrayStart();
+
+		for(int32 I = 0; I < 16; ++I)
+		{
+			FGLTFJsonUtility::WriteExactValue(JsonWriter, M[I]);
+		}
+
+		JsonWriter.WriteArrayEnd();
+	}
+
+	float& operator()(int32 Row, int32 Col)
+	{
+		return M[Row * 4 + Col];
+	}
+
+	bool operator==(const FGLTFJsonMatrix4& Other) const
+	{
+		for(int32 I = 0; I < 16; ++I)
+		{
+			if (M[I] != Other.M[I])
+			{
+				return false;
+			}
+		}
+
+		return true;
+	}
+
+	bool operator!=(const FGLTFJsonMatrix4& Other) const
+	{
+		return !(*this == Other);
+	}
+};
diff --git a/Engine/Plugins/Enterprise/GLTFExporter/MigratedFromGit/Source/GLTFExporter/Private/Json/GLTFJsonNode.h b/Engine/Plugins/Enterprise/GLTFExporter/MigratedFromGit/Source/GLTFExporter/Private/Json/GLTFJsonNode.h
--- a/Engine/Plugins/Enterprise/GLTFExporter/MigratedFromGit/Source/GLTFExporter/Private/Json/GLTFJsonNode.h
+++ b/Engine/Plugins/Enterprise/GLTFExporter/MigratedFromGit/Source/GLTFExporter/Private/Json/GLTFJsonNode.h
@@ -3,6 +3,7 @@
 #pragma once
 
 #include "Json/GLTFJsonIndex.h"
+#include "Json/GLTFJsonMatrix4.h"
 #include "Json/GLTFJsonVector3.h"
 #include "Json/GLTFJsonQuaternion.h"
 #include "Serialization/JsonSerializer.h"
@@ -11,9 +12,19 @@ struct FGLTFJsonNode
 {
 	FString Name;
 
-	FGLTFJsonVector3    Translation;
-	FGLTFJsonQuaternion Rotation;
-	FGLTFJsonVector3    Scale;
+	bool bUseMatrix;
+
+	union
+	{
+		FGLTFJsonMatrix4 Matrix;
+
+		struct
+		{
+			FGLTFJsonVector3    Translation;
+			FGLTFJsonQuaternion Rotation;
+			FGLTFJsonVector3    Scale;
+		};
+	};
 
 	FGLTFJsonCameraIndex    Camera;
 	FGLTFJsonSkinIndex      Skin;
@@ -25,7 +36,8 @@ struct FGLTFJsonNode
 	TArray<FGLTFJsonNodeIndex> Children;
 
 	FGLTFJsonNode()
-		: Translation(FGLTFJsonVector3::Zero)
+		: bUseMatrix(false)
+		, Translation(FGLTFJsonVector3::Zero)
 		, Rotation(FGLTFJsonQuaternion::Identity)
 		, Scale(FGLTFJsonVector3::One)
 	{
@@ -41,22 +53,33 @@ struct FGLTFJsonNode
 			JsonWriter.WriteValue(TEXT("name"), Name);
 		}
 
-		if (Translation != FGLTFJsonVector3::Zero)
+		if (bUseMatrix)
 		{
-			JsonWriter.WriteIdentifierPrefix(TEXT("translation"));
-			Translation.WriteArray(JsonWriter);
+			if (Matrix != FGLTFJsonMatrix4::Identity)
+			{
+				JsonWriter.WriteIdentifierPrefix(TEXT("matrix"));
+				Matrix.WriteArray(JsonWriter);
+			}
 		}
-
-		if (Rotation != FGLTFJsonQuaternion::Identity)
+		else
 		{
-			JsonWriter.WriteIdentifierPrefix(TEXT("rotation"));
-			Rotation.WriteArray(JsonWriter);
-		}
+			if (Translation != FGLTFJsonVector3::Zero)
+			{
+				JsonWriter.WriteIdentifierPrefix(TEXT("translation"));
+				Translation.WriteArray(JsonWriter);
+			}
 
-		if (Scale != FGLTFJsonVector3::One)
-		{
-			JsonWriter.WriteIdentifierPrefix(TEXT("scale"));
-			Scale.WriteArray(JsonWriter);
+			if (Rotation != FGLTFJsonQuaternion::Identity)
+			{
+				JsonWriter.WriteIdentifierPrefix(TEXT("rotation"));
+				Rotation.WriteArray(JsonWriter);
+			}
+
+			if (Scale != FGLTFJsonVector3::One)
+			{
+				JsonWriter.WriteIdentifierPrefix(TEXT("scale"));
+				Scale.WriteArray(JsonWriter);
+			}
 		}
 
 		if (Camera != INDEX_NONE)
