[GLTFExporter] refactor: avoid unnecessary copy of material baking output

commit 75b629892f16c9fcbd1e7c44ed6c8c315a69b98a
Date: Mon, 31 Aug 2020 18:05:26 +0200

    refactor: avoid unnecessary copy of material baking output
    (cherry picked from commit b1f06cc17dd90d838f21fcdbe8ee927b5a583db9​)

#skipci

diff --git a/Engine/Plugins/Enterprise/GLTFExporter/MigratedFromGit/Source/GLTFExporter/Private/Converters/GLTFMaterialUtility.cpp b/Engine/Plugins/Enterprise/GLTFExporter/MigratedFromGit/Source/GLTFExporter/Private/Converters/GLTFMaterialUtility.cpp
--- a/Engine/Plugins/Enterprise/GLTFExporter/MigratedFromGit/Source/GLTFExporter/Private/Converters/GLTFMaterialUtility.cpp
+++ b/Engine/Plugins/Enterprise/GLTFExporter/MigratedFromGit/Source/GLTFExporter/Private/Converters/GLTFMaterialUtility.cpp
@@ -274,7 +274,7 @@ FGLTFPropertyBakeOutput FGLTFMaterialUtility::BakeMaterialProperty(const FIntPoi
 
 	Module.BakeMaterials(MatSettings, MeshSettings, BakeOutputs);
 
-	FBakeOutput BakeOutput = BakeOutputs[0];
+	FBakeOutput& BakeOutput = BakeOutputs[0];
 
 	TArray<FColor> BakedPixels = MoveTemp(BakeOutput.PropertyData.FindChecked(Property));
 	const FIntPoint BakedSize = BakeOutput.PropertySizes.FindChecked(Property);
