[GLTFExporter] Merge remote-tracking branch 'origin/master' into feature/improve-handling-of-scaled-transforms

commit 794629ecf1d50d6623dadcf8479bfc131b56e29c
Date: Tue, 1 Sep 2020 08:47:01 +0200

    Merge remote-tracking branch 'origin/master' into feature/improve-handling-of-scaled-transforms
    # Conflicts from merge:
    #	Plugins/GLTFExporter/Source/GLTFExporter/Private/Converters/GLTFLevelConverters.cpp
    #	Tests/BasicLevelVariants/BasicLevelVariants.gltf
    #	Tests/InteractionHotspots/InteractionHotspots.gltf
    #	Tests/LightsAndCameras/LightsAndCameras.gltf
    #	Tests/MeshVariants/MeshVariants.gltf
    #	Tests/ModelPark/ModelPark.gltf
    #	Tests/NonUniformScale/NonUniformScale.gltf
    #	Tests/OrbitCameras/OrbitCameras.gltf
    #	Tests/PrimitiveBeach/PrimitiveBeach.gltf
    #	Tests/PrimitiveOrchard/PrimitiveOrchard.gltf
    #	Tests/SimpleBaked/SimpleBaked.gltf
    (cherry picked from commit 3dd18025b638a2de7a40257cacf1241421a46197​)

#skipci

diff --git a/Engine/Plugins/Enterprise/GLTFExporter/MigratedFromGit/Source/GLTFExporter/Private/Converters/GLTFLevelConverters.cpp b/Engine/Plugins/Enterprise/GLTFExporter/MigratedFromGit/Source/GLTFExporter/Private/Converters/GLTFLevelConverters.cpp
--- a/Engine/Plugins/Enterprise/GLTFExporter/MigratedFromGit/Source/GLTFExporter/Private/Converters/GLTFLevelConverters.cpp
+++ b/Engine/Plugins/Enterprise/GLTFExporter/MigratedFromGit/Source/GLTFExporter/Private/Converters/GLTFLevelConverters.cpp
@@ -34,11 +34,9 @@ FGLTFJsonNodeIndex FGLTFSceneComponentConverter::Add(FGLTFConvertBuilder& Builde
 	const FGLTFJsonNodeIndex NodeIndex = Builder.AddChildNode(ParentNodeIndex);
 	FGLTFJsonNode& Node = Builder.GetNode(NodeIndex);
 	Node.Name = Name.IsEmpty() ? Owner->GetName() + TEXT("_") + SceneComponent->GetName() : Name;
-
-	// TODO: add support for non-uniform scaling (Unreal doesn't treat combined transforms as simple matrix multiplication)
-	const FTransform Transform = bIsRootNode ? SceneComponent->GetComponentTransform() : SceneComponent->GetRelativeTransform();
-	Node.Matrix = FGLTFConverterUtility::ConvertTransform(Transform, Builder.ExportOptions->ExportScale);
-	Node.bUseMatrix = true;
+	Node.Translation = FGLTFConverterUtility::ConvertPosition(Translation, Builder.ExportOptions->ExportScale);
+	Node.Rotation = FGLTFConverterUtility::ConvertRotation(RelativeTransform.GetRotation());
+	Node.Scale = FGLTFConverterUtility::ConvertScale(Scale);
 
 	if (SceneComponent->bHiddenInGame) // TODO: make this configurable
 	{
