[GLTFExporter] fix: load all texture mips before drawing the texture to prevent (sometimes) blurry results

commit b6fa79e9d4d7de38143597f733d5049d6033487d
Date: Fri, 21 Aug 2020 16:15:33 +0200

    fix: load all texture mips before drawing the texture to prevent (sometimes) blurry results
    (cherry picked from commit ff8487f00c7f237d7749fb580ee77d0bbe1ec0b3​)

#skipci

diff --git a/Engine/Plugins/Enterprise/GLTFExporter/MigratedFromGit/Source/GLTFExporter/Private/Converters/GLTFTextureUtility.cpp b/Engine/Plugins/Enterprise/GLTFExporter/MigratedFromGit/Source/GLTFExporter/Private/Converters/GLTFTextureUtility.cpp
--- a/Engine/Plugins/Enterprise/GLTFExporter/MigratedFromGit/Source/GLTFExporter/Private/Converters/GLTFTextureUtility.cpp
+++ b/Engine/Plugins/Enterprise/GLTFExporter/MigratedFromGit/Source/GLTFExporter/Private/Converters/GLTFTextureUtility.cpp
@@ -131,6 +131,10 @@ bool FGLTFTextureUtility::DrawTexture(UTextureRenderTarget2D* OutTarget, const U
 		return false;
 	}
 
+	// Fully stream in the texture before drawing it.
+	const_cast<UTexture2D*>(InSource)->SetForceMipLevelsToBeResident(30.0f, true);
+	const_cast<UTexture2D*>(InSource)->WaitForStreaming();
+
 	FCanvas Canvas(RenderTarget, nullptr, 0.0f, 0.0f, 0.0f, GMaxRHIFeatureLevel);
 	FCanvasTileItem TileItem(FVector2D::ZeroVector, InSource->Resource, FLinearColor::White);
 
