[GLTFExporter] feat: remove the old code-path for exporting textures via their sources

commit 4af5e12ecef3eed8be0eefcbaf07f9fea412ab88
Date: Fri, 21 Aug 2020 16:28:24 +0200

    feat: remove the old code-path for exporting textures via their sources
    (cherry picked from commit d43dff81393a72fd8987e49bfbeb43edec9a8f68​)

#skipci

diff --git a/Engine/Plugins/Enterprise/GLTFExporter/MigratedFromGit/Source/GLTFExporter/Private/Converters/GLTFTextureConverters.cpp b/Engine/Plugins/Enterprise/GLTFExporter/MigratedFromGit/Source/GLTFExporter/Private/Converters/GLTFTextureConverters.cpp
--- a/Engine/Plugins/Enterprise/GLTFExporter/MigratedFromGit/Source/GLTFExporter/Private/Converters/GLTFTextureConverters.cpp
+++ b/Engine/Plugins/Enterprise/GLTFExporter/MigratedFromGit/Source/GLTFExporter/Private/Converters/GLTFTextureConverters.cpp
@@ -52,16 +52,7 @@ FGLTFJsonTextureIndex FGLTFTexture2DConverter::Add(FGLTFConvertBuilder& Builder,
 	ERGBFormat RGBFormat;
 	uint32 BitDepth;
 
-	// TODO: don't rely on Texture2D->Source, since the texture brightness etc can be tweaked for runtime
-	if (FGLTFTextureUtility::CanPNGCompressFormat(Texture2D->Source.GetFormat(), RGBFormat, BitDepth))
-	{
-		FTextureSource& Source = const_cast<FTextureSource&>(Texture2D->Source);
-
-		const void* RawData = Source.LockMip(0);
-		ImageIndex = Builder.AddImage(RawData, Source.CalcMipSize(0), Size, RGBFormat, BitDepth, JsonTexture.Name);
-		Source.UnlockMip(0);
-	}
-	else if (FGLTFTextureUtility::CanPNGCompressFormat(Texture2D->GetPixelFormat(), RGBFormat, BitDepth))
+	if (FGLTFTextureUtility::CanPNGCompressFormat(Texture2D->GetPixelFormat(), RGBFormat, BitDepth))
 	{
 		const FByteBulkData& BulkData = Texture2D->PlatformData->Mips[0].BulkData;
 		// TODO: ensure that BulkData size is not zero and if so rebuild platform data.
