[GLTFExporter] fix: fix random crashes when exporting textures via platform-data by loading inline mip data first

commit 62ba6b317a95ccfa6facce7d6edb191072fd0dfe
Date: Fri, 21 Aug 2020 16:29:24 +0200

    fix: fix random crashes when exporting textures via platform-data by loading inline mip data first
    (cherry picked from commit 25f8d1d345de8a04c6d71f7f33db1d3f05822c9e​)

#skipci

diff --git a/Engine/Plugins/Enterprise/GLTFExporter/MigratedFromGit/Source/GLTFExporter/Private/Converters/GLTFTextureConverters.cpp b/Engine/Plugins/Enterprise/GLTFExporter/MigratedFromGit/Source/GLTFExporter/Private/Converters/GLTFTextureConverters.cpp
--- a/Engine/Plugins/Enterprise/GLTFExporter/MigratedFromGit/Source/GLTFExporter/Private/Converters/GLTFTextureConverters.cpp
+++ b/Engine/Plugins/Enterprise/GLTFExporter/MigratedFromGit/Source/GLTFExporter/Private/Converters/GLTFTextureConverters.cpp
@@ -48,23 +48,28 @@ FGLTFJsonTextureIndex FGLTFTexture2DConverter::Add(FGLTFConvertBuilder& Builder,
 	FGLTFJsonTexture JsonTexture;
 	JsonTexture.Name = Name;
 
-	FIntPoint Size = { Texture2D->GetSizeX(), Texture2D->GetSizeY() };
+	const FIntPoint Size = { Texture2D->GetSizeX(), Texture2D->GetSizeY() };
 	ERGBFormat RGBFormat;
 	uint32 BitDepth;
 
 	if (FGLTFTextureUtility::CanPNGCompressFormat(Texture2D->GetPixelFormat(), RGBFormat, BitDepth))
 	{
-		const FByteBulkData& BulkData = Texture2D->PlatformData->Mips[0].BulkData;
-		// TODO: ensure that BulkData size is not zero and if so rebuild platform data.
+		if (Texture2D->PlatformData->TryInlineMipData(0, const_cast<UTexture2D*>(Texture2D)))
+		{
+			const FByteBulkData& BulkData = Texture2D->PlatformData->Mips[0].BulkData;
+			const void* RawData = BulkData.LockReadOnly();
 
-		const void* RawData = BulkData.LockReadOnly();
-		ImageIndex = Builder.AddImage(RawData, BulkData.GetBulkDataSize(), Size, RGBFormat, BitDepth, JsonTexture.Name);
-		BulkData.Unlock();
+			ImageIndex = Builder.AddImage(RawData, BulkData.GetBulkDataSize(), Size, RGBFormat, BitDepth, JsonTexture.Name);
+			BulkData.Unlock();
+		}
+		else
+		{
+			// TODO: report error
+			return FGLTFJsonTextureIndex(INDEX_NONE);
+		}
 	}
 	else
 	{
-		// TODO: investigate and fix reason why using this fallback will result in blurry/fuzzy exported images
-
 		const EPixelFormat RenderTargetFormat = FGLTFTextureUtility::IsHDRFormat(Texture2D->GetPixelFormat()) ? PF_FloatRGBA : PF_B8G8R8A8;
 		UTextureRenderTarget2D* RenderTarget = FGLTFTextureUtility::CreateRenderTarget(Size, RenderTargetFormat, true);
 
