[GLTFExporter] feat: use correct(?) gamma of 2.2 when converting unsupported texture formats (to be exported)

commit 3c9643ab84b4b97a3732284d982423e8fdf8dd72
Date: Fri, 21 Aug 2020 16:32:11 +0200

    feat: use correct(?) gamma of 2.2 when converting unsupported texture formats (to be exported)
    (cherry picked from commit 19ae78e3a6108230a54ab23b181d352fb864de37​)

#skipci

diff --git a/Engine/Plugins/Enterprise/GLTFExporter/MigratedFromGit/Source/GLTFExporter/Private/Converters/GLTFTextureConverters.cpp b/Engine/Plugins/Enterprise/GLTFExporter/MigratedFromGit/Source/GLTFExporter/Private/Converters/GLTFTextureConverters.cpp
--- a/Engine/Plugins/Enterprise/GLTFExporter/MigratedFromGit/Source/GLTFExporter/Private/Converters/GLTFTextureConverters.cpp
+++ b/Engine/Plugins/Enterprise/GLTFExporter/MigratedFromGit/Source/GLTFExporter/Private/Converters/GLTFTextureConverters.cpp
@@ -71,7 +71,15 @@ FGLTFJsonTextureIndex FGLTFTexture2DConverter::Add(FGLTFConvertBuilder& Builder,
 	else
 	{
 		const EPixelFormat RenderTargetFormat = FGLTFTextureUtility::IsHDRFormat(Texture2D->GetPixelFormat()) ? PF_FloatRGBA : PF_B8G8R8A8;
+
+		// TODO: both bInForceLinearGamma and TargetGamma=2.2f seem to be necessary for the exported images to match the results
+		// from exporting using the texture's source or its platform-data.
+		// It's not entirely clear why gamma must be set to 2.2 instead of 0.0 (which should use the correct gamma anyway),
+		// and why bInForceLinearGamma must also be true.
+		// Also, some pixel-formats such as PF_BC5 are incorrectly rendered.
+		// We need to investigate this more in-depth later.
 		UTextureRenderTarget2D* RenderTarget = FGLTFTextureUtility::CreateRenderTarget(Size, RenderTargetFormat, true);
+		RenderTarget->TargetGamma = 2.2f;
 
 		FGLTFTextureUtility::DrawTexture(RenderTarget, Texture2D);
 
