[GLTFExporter] chore: clarify TODO for incorrect rendering / conversion of normal-maps

commit d32447f1d2398e014465635437347e873471abda
Date: Tue, 25 Aug 2020 11:02:22 +0200

    chore: clarify TODO for incorrect rendering / conversion of normal-maps
    (cherry picked from commit 2d8b5f6ed7cf3dd26808e5ed47307b99ca449573​)

#skipci

diff --git a/Engine/Plugins/Enterprise/GLTFExporter/MigratedFromGit/Source/GLTFExporter/Private/Converters/GLTFTextureConverters.cpp b/Engine/Plugins/Enterprise/GLTFExporter/MigratedFromGit/Source/GLTFExporter/Private/Converters/GLTFTextureConverters.cpp
--- a/Engine/Plugins/Enterprise/GLTFExporter/MigratedFromGit/Source/GLTFExporter/Private/Converters/GLTFTextureConverters.cpp
+++ b/Engine/Plugins/Enterprise/GLTFExporter/MigratedFromGit/Source/GLTFExporter/Private/Converters/GLTFTextureConverters.cpp
@@ -76,11 +76,13 @@ FGLTFJsonTextureIndex FGLTFTexture2DConverter::Add(FGLTFConvertBuilder& Builder,
 		// from exporting using the texture's source or its platform-data.
 		// It's not entirely clear why gamma must be set to 2.2 instead of 0.0 (which should use the correct gamma anyway),
 		// and why bInForceLinearGamma must also be true.
-		// Also, some pixel-formats such as PF_BC5 are incorrectly rendered.
-		// We need to investigate this more in-depth later.
 		UTextureRenderTarget2D* RenderTarget = FGLTFTextureUtility::CreateRenderTarget(Size, RenderTargetFormat, true);
 		RenderTarget->TargetGamma = 2.2f;
 
+		// TODO: normal-maps are not correctly rendered to the render-target, the resulting pixels have no data in the blue channel.
+		// This may be due to the textures in question using PF_BC5 as pixel-format, or something else entirely.
+		// There may be similar issues with other texture-types that we haven't had the chance to test yet.
+
 		FGLTFTextureUtility::DrawTexture(RenderTarget, Texture2D);
 
 		TArray<FColor> Pixels;
