[GLTFExporter] fix: try to rebuild platform-data before using TryInlineMipData to ensure platform-data exists

commit c16d7fe13b6872a2b813f856e1b5388a4bd0c166
Date: Tue, 25 Aug 2020 13:43:57 +0200

    fix: try to rebuild platform-data before using TryInlineMipData to ensure platform-data exists
    (cherry picked from commit 92160512eb4a6515af4aaccc09bfd363e1d9e7d5​)

#skipci

diff --git a/Engine/Plugins/Enterprise/GLTFExporter/MigratedFromGit/Source/GLTFExporter/Private/Converters/GLTFTextureConverters.cpp b/Engine/Plugins/Enterprise/GLTFExporter/MigratedFromGit/Source/GLTFExporter/Private/Converters/GLTFTextureConverters.cpp
--- a/Engine/Plugins/Enterprise/GLTFExporter/MigratedFromGit/Source/GLTFExporter/Private/Converters/GLTFTextureConverters.cpp
+++ b/Engine/Plugins/Enterprise/GLTFExporter/MigratedFromGit/Source/GLTFExporter/Private/Converters/GLTFTextureConverters.cpp
@@ -54,6 +54,11 @@ FGLTFJsonTextureIndex FGLTFTexture2DConverter::Add(FGLTFConvertBuilder& Builder,
 
 	if (FGLTFTextureUtility::CanPNGCompressFormat(Texture2D->GetPixelFormat(), RGBFormat, BitDepth))
 	{
+		if (Texture2D->PlatformData->Mips[0].BulkData.GetBulkDataSize() == 0)
+		{
+			const_cast<UTexture2D*>(Texture2D)->ForceRebuildPlatformData();
+		}
+
 		if (Texture2D->PlatformData->TryInlineMipData(0, const_cast<UTexture2D*>(Texture2D)))
 		{
 			const FByteBulkData& BulkData = Texture2D->PlatformData->Mips[0].BulkData;
