[GLTFExporter] fix: try using a special code-path when exporting light-maps to overcome issues

commit 88732657c80ab131d6c1f6e9afa7bdb583ad692c
Date: Tue, 25 Aug 2020 14:58:59 +0200

    fix: try using a special code-path when exporting light-maps to overcome issues
    (cherry picked from commit 3bb6de30d08be90349562f575c0654d456d64a39​)

#skipci

diff --git a/Engine/Plugins/Enterprise/GLTFExporter/MigratedFromGit/Source/GLTFExporter/Private/Converters/GLTFTextureConverters.cpp b/Engine/Plugins/Enterprise/GLTFExporter/MigratedFromGit/Source/GLTFExporter/Private/Converters/GLTFTextureConverters.cpp
--- a/Engine/Plugins/Enterprise/GLTFExporter/MigratedFromGit/Source/GLTFExporter/Private/Converters/GLTFTextureConverters.cpp
+++ b/Engine/Plugins/Enterprise/GLTFExporter/MigratedFromGit/Source/GLTFExporter/Private/Converters/GLTFTextureConverters.cpp
@@ -52,20 +52,18 @@ FGLTFJsonTextureIndex FGLTFTexture2DConverter::Add(FGLTFConvertBuilder& Builder,
 	ERGBFormat RGBFormat;
 	uint32 BitDepth;
 
-	if (FGLTFTextureUtility::CanPNGCompressFormat(Texture2D->GetPixelFormat(), RGBFormat, BitDepth))
+	// TODO: the special code-path for exporting light-maps by source is temporary, and is only used to overcome issues
+	// when trying to access platform-data for the light-maps. They would (intermittently) fail. We need to find the
+	// root cause and then remove this temporary fix.
+	if (Texture2D->IsA<ULightMapTexture2D>())
 	{
-		if (Texture2D->PlatformData->Mips[0].BulkData.GetBulkDataSize() == 0)
+		if (FGLTFTextureUtility::CanPNGCompressFormat(Texture2D->Source.GetFormat(), RGBFormat, BitDepth))
 		{
-			const_cast<UTexture2D*>(Texture2D)->ForceRebuildPlatformData();
-		}
+			FTextureSource& Source = const_cast<FTextureSource&>(Texture2D->Source);
 
-		if (Texture2D->PlatformData->TryInlineMipData(0, const_cast<UTexture2D*>(Texture2D)))
-		{
-			const FByteBulkData& BulkData = Texture2D->PlatformData->Mips[0].BulkData;
-			const void* RawData = BulkData.LockReadOnly();
-
-			ImageIndex = Builder.AddImage(RawData, BulkData.GetBulkDataSize(), Size, RGBFormat, BitDepth, JsonTexture.Name);
-			BulkData.Unlock();
+			const void* RawData = Source.LockMip(0);
+			ImageIndex = Builder.AddImage(RawData, Source.CalcMipSize(0), Size, RGBFormat, BitDepth, JsonTexture.Name);
+			Source.UnlockMip(0);
 		}
 		else
 		{
@@ -73,6 +71,25 @@ FGLTFJsonTextureIndex FGLTFTexture2DConverter::Add(FGLTFConvertBuilder& Builder,
 			return FGLTFJsonTextureIndex(INDEX_NONE);
 		}
 	}
+	else if (FGLTFTextureUtility::CanPNGCompressFormat(Texture2D->GetPixelFormat(), RGBFormat, BitDepth))
+	{
+		if (Texture2D->PlatformData->Mips[0].BulkData.GetBulkDataSize() == 0)
+		{
+			// TODO: is this correct handling?
+			const_cast<UTexture2D*>(Texture2D)->ForceRebuildPlatformData();
+			if (Texture2D->PlatformData->Mips[0].BulkData.GetBulkDataSize() == 0)
+			{
+				// TODO: report error
+				return FGLTFJsonTextureIndex(INDEX_NONE);
+			}
+		}
+
+		const FByteBulkData& BulkData = Texture2D->PlatformData->Mips[0].BulkData;
+		const void* RawData = BulkData.LockReadOnly();
+
+		ImageIndex = Builder.AddImage(RawData, BulkData.GetBulkDataSize(), Size, RGBFormat, BitDepth, JsonTexture.Name);
+		BulkData.Unlock();
+	}
 	else
 	{
 		const EPixelFormat RenderTargetFormat = FGLTFTextureUtility::IsHDRFormat(Texture2D->GetPixelFormat()) ? PF_FloatRGBA : PF_B8G8R8A8;
