[GLTFExporter] fix: export normal-maps via source just like for light-maps to overcome issues (see TODO)

commit 9bc9b1849b42beda30cdfffda0aa23812e0aeca3
Date: Mon, 31 Aug 2020 14:44:09 +0200

    fix: export normal-maps via source just like for light-maps to overcome issues (see TODO)
    (cherry picked from commit e96c7467d441e4ceedfb59c52f104dd62f9f4636​)

#skipci

diff --git a/Engine/Plugins/Enterprise/GLTFExporter/MigratedFromGit/Source/GLTFExporter/Private/Converters/GLTFTextureConverters.cpp b/Engine/Plugins/Enterprise/GLTFExporter/MigratedFromGit/Source/GLTFExporter/Private/Converters/GLTFTextureConverters.cpp
--- a/Engine/Plugins/Enterprise/GLTFExporter/MigratedFromGit/Source/GLTFExporter/Private/Converters/GLTFTextureConverters.cpp
+++ b/Engine/Plugins/Enterprise/GLTFExporter/MigratedFromGit/Source/GLTFExporter/Private/Converters/GLTFTextureConverters.cpp
@@ -52,11 +52,12 @@ FGLTFJsonTextureIndex FGLTFTexture2DConverter::Add(FGLTFConvertBuilder& Builder,
 	ERGBFormat RGBFormat;
 	uint32 BitDepth;
 
-	// TODO: the special code-path for exporting light-maps by source is temporary, and is only used to overcome issues
-	// when trying to access platform-data for the light-maps. They would (intermittently) fail. We need to find the
-	// root cause and then remove this temporary fix.
-	if (Texture2D->IsA<ULightMapTexture2D>())
+	if (Texture2D->IsA<ULightMapTexture2D>() || Texture2D->IsNormalMap())
 	{
+		// NOTE: export of light- and normal-maps via source is done to overcome issues with accessing
+		// pixel-data (for light-maps) and compressed pixel-data (for normal-maps).
+		// TODO: by exporting via source, we lose the ability to export texture-adjustments. We may want to support this later.
+
 		if (FGLTFTextureUtility::CanPNGCompressFormat(Texture2D->Source.GetFormat(), RGBFormat, BitDepth))
 		{
 			FTextureSource& Source = const_cast<FTextureSource&>(Texture2D->Source);
