[GLTFExporter] chore: remove leftover TODO which is already covered by a NOTE higher up

commit 090c3a1cf99510960b66e23aff4d9cc18a73f436
Date: Wed, 2 Sep 2020 08:33:42 +0200

    chore: remove leftover TODO which is already covered by a NOTE higher up
    (cherry picked from commit b84f06b7706c6dce1214815ea0d775deaedc3a35​)

#skipci

diff --git a/Engine/Plugins/Enterprise/GLTFExporter/MigratedFromGit/Source/GLTFExporter/Private/Converters/GLTFTextureConverters.cpp b/Engine/Plugins/Enterprise/GLTFExporter/MigratedFromGit/Source/GLTFExporter/Private/Converters/GLTFTextureConverters.cpp
--- a/Engine/Plugins/Enterprise/GLTFExporter/MigratedFromGit/Source/GLTFExporter/Private/Converters/GLTFTextureConverters.cpp
+++ b/Engine/Plugins/Enterprise/GLTFExporter/MigratedFromGit/Source/GLTFExporter/Private/Converters/GLTFTextureConverters.cpp
@@ -102,10 +102,6 @@ FGLTFJsonTextureIndex FGLTFTexture2DConverter::Add(FGLTFConvertBuilder& Builder,
 		UTextureRenderTarget2D* RenderTarget = FGLTFTextureUtility::CreateRenderTarget(Size, RenderTargetFormat, true);
 		RenderTarget->TargetGamma = 2.2f;
 
-		// TODO: normal-maps are not correctly rendered to the render-target, the resulting pixels have no data in the blue channel.
-		// This may be due to the textures in question using PF_BC5 as pixel-format, or something else entirely.
-		// There may be similar issues with other texture-types that we haven't had the chance to test yet.
-
 		FGLTFTextureUtility::DrawTexture(RenderTarget, Texture2D);
 
 		TArray<FColor> Pixels;
