[GLTFExporter] chore: add general TODO regarding special handling of certain texture-types during conversion

commit 17ef386bc6d75b33259bd22dcc1f7e2a4c865105
Date: Wed, 2 Sep 2020 08:38:10 +0200

    chore: add general TODO regarding special handling of certain texture-types during conversion
    (cherry picked from commit 4a1422a93f5be2c08a26ac3e311e8db089984d41​)

#skipci

diff --git a/Engine/Plugins/Enterprise/GLTFExporter/MigratedFromGit/Source/GLTFExporter/Private/Converters/GLTFTextureConverters.cpp b/Engine/Plugins/Enterprise/GLTFExporter/MigratedFromGit/Source/GLTFExporter/Private/Converters/GLTFTextureConverters.cpp
--- a/Engine/Plugins/Enterprise/GLTFExporter/MigratedFromGit/Source/GLTFExporter/Private/Converters/GLTFTextureConverters.cpp
+++ b/Engine/Plugins/Enterprise/GLTFExporter/MigratedFromGit/Source/GLTFExporter/Private/Converters/GLTFTextureConverters.cpp
@@ -52,6 +52,7 @@ FGLTFJsonTextureIndex FGLTFTexture2DConverter::Add(FGLTFConvertBuilder& Builder,
 	ERGBFormat RGBFormat;
 	uint32 BitDepth;
 
+	// TODO: make sure there are no other special cases (like ex normal-maps and light-maps) that need specialized handling.
 	if (Texture2D->IsA<ULightMapTexture2D>() || Texture2D->IsNormalMap())
 	{
 		// NOTE: export of light- and normal-maps via source is done to overcome issues with accessing
