[GLTFExporter] refactor: extract TargetGamma and bForceLinearGamma into variables to clarify their importance in the TODO above

commit 2ac845c3bdc3ab65ba8576a6f62819f2c4ca301a
Date: Thu, 3 Sep 2020 15:25:34 +0200

    refactor: extract TargetGamma and bForceLinearGamma into variables to clarify their importance in the TODO above
    (cherry picked from commit 12ff1318091abf945f9c909bb34810528327b462​)

#skipci

diff --git a/Engine/Plugins/Enterprise/GLTFExporter/MigratedFromGit/Source/GLTFExporter/Private/Converters/GLTFTextureConverters.cpp b/Engine/Plugins/Enterprise/GLTFExporter/MigratedFromGit/Source/GLTFExporter/Private/Converters/GLTFTextureConverters.cpp
--- a/Engine/Plugins/Enterprise/GLTFExporter/MigratedFromGit/Source/GLTFExporter/Private/Converters/GLTFTextureConverters.cpp
+++ b/Engine/Plugins/Enterprise/GLTFExporter/MigratedFromGit/Source/GLTFExporter/Private/Converters/GLTFTextureConverters.cpp
@@ -97,12 +97,15 @@ FGLTFJsonTextureIndex FGLTFTexture2DConverter::Add(FGLTFConvertBuilder& Builder,
 	{
 		const EPixelFormat RenderTargetFormat = FGLTFTextureUtility::IsHDRFormat(Texture2D->GetPixelFormat()) ? PF_FloatRGBA : PF_B8G8R8A8;
 
-		// TODO: both bInForceLinearGamma and TargetGamma=2.2f seem to be necessary for the exported images to match the results
+		// TODO: both bForceLinearGamma and TargetGamma=2.2f seem to be necessary for the exported images to match the results
 		// from exporting using the texture's source or its platform-data.
 		// It's not entirely clear why gamma must be set to 2.2 instead of 0.0 (which should use the correct gamma anyway),
 		// and why bInForceLinearGamma must also be true.
-		UTextureRenderTarget2D* RenderTarget = FGLTFTextureUtility::CreateRenderTarget(Size, RenderTargetFormat, true);
-		RenderTarget->TargetGamma = 2.2f;
+		const bool bForceLinearGamma = true;
+		const float TargetGamma = 2.2f;
+
+		UTextureRenderTarget2D* RenderTarget = FGLTFTextureUtility::CreateRenderTarget(Size, RenderTargetFormat, bForceLinearGamma);
+		RenderTarget->TargetGamma = TargetGamma;
 
 		FGLTFTextureUtility::DrawTexture(RenderTarget, Texture2D);
 
