[GLTFExporter] refactor: split logic for loading platform-data for textures into utility-function(s)

commit d4a00091bd846d0b1bb791eb61b630e0c6dc7ab1
Date: Thu, 3 Sep 2020 15:39:22 +0200

    refactor: split logic for loading platform-data for textures into utility-function(s)
    (cherry picked from commit 071272ee2cf9cf6ea0bc9922302eaf10e60532ed​)

#skipci

diff --git a/Engine/Plugins/Enterprise/GLTFExporter/MigratedFromGit/Source/GLTFExporter/Private/Converters/GLTFTextureConverters.cpp b/Engine/Plugins/Enterprise/GLTFExporter/MigratedFromGit/Source/GLTFExporter/Private/Converters/GLTFTextureConverters.cpp
--- a/Engine/Plugins/Enterprise/GLTFExporter/MigratedFromGit/Source/GLTFExporter/Private/Converters/GLTFTextureConverters.cpp
+++ b/Engine/Plugins/Enterprise/GLTFExporter/MigratedFromGit/Source/GLTFExporter/Private/Converters/GLTFTextureConverters.cpp
@@ -76,15 +76,10 @@ FGLTFJsonTextureIndex FGLTFTexture2DConverter::Add(FGLTFConvertBuilder& Builder,
 	}
 	else if (FGLTFTextureUtility::CanPNGCompressFormat(Texture2D->GetPixelFormat(), RGBFormat, BitDepth))
 	{
-		if (Texture2D->PlatformData->Mips[0].BulkData.GetBulkDataSize() == 0)
+		if (!FGLTFTextureUtility::LoadPlatformData(const_cast<UTexture2D*>(Texture2D)))
 		{
-			// TODO: is this correct handling?
-			const_cast<UTexture2D*>(Texture2D)->ForceRebuildPlatformData();
-			if (Texture2D->PlatformData->Mips[0].BulkData.GetBulkDataSize() == 0)
-			{
-				// TODO: report error
-				return FGLTFJsonTextureIndex(INDEX_NONE);
-			}
+			// TODO: report error
+			return FGLTFJsonTextureIndex(INDEX_NONE);
 		}
 
 		const FByteBulkData& BulkData = Texture2D->PlatformData->Mips[0].BulkData;
diff --git a/Engine/Plugins/Enterprise/GLTFExporter/MigratedFromGit/Source/GLTFExporter/Private/Converters/GLTFTextureUtility.cpp b/Engine/Plugins/Enterprise/GLTFExporter/MigratedFromGit/Source/GLTFExporter/Private/Converters/GLTFTextureUtility.cpp
--- a/Engine/Plugins/Enterprise/GLTFExporter/MigratedFromGit/Source/GLTFExporter/Private/Converters/GLTFTextureUtility.cpp
+++ b/Engine/Plugins/Enterprise/GLTFExporter/MigratedFromGit/Source/GLTFExporter/Private/Converters/GLTFTextureUtility.cpp
@@ -167,21 +167,11 @@ UTexture2D* FGLTFTextureUtility::CreateTextureFromCubeFace(const UTextureCube* T
 	const FIntPoint Size(TextureCube->GetSizeX(), TextureCube->GetSizeY());
 	const EPixelFormat Format = TextureCube->GetPixelFormat();
 
-	if (TextureCube->PlatformData == nullptr || TextureCube->PlatformData->Mips.Num() == 0)
+	if (!LoadPlatformData(const_cast<UTextureCube*>(TextureCube)))
 	{
 		return nullptr;
 	}
 
-	if (TextureCube->PlatformData->Mips[0].BulkData.GetBulkDataSize() == 0)
-	{
-		// TODO: is this correct handling?
-		const_cast<UTextureCube*>(TextureCube)->ForceRebuildPlatformData();
-		if (TextureCube->PlatformData->Mips[0].BulkData.GetBulkDataSize() == 0)
-		{
-			return nullptr;
-		}
-	}
-
 	const FByteBulkData& BulkData = TextureCube->PlatformData->Mips[0].BulkData;
 	const int64 MipSize = BulkData.GetBulkDataSize() / 6;
 
@@ -282,3 +272,43 @@ void FGLTFTextureUtility::EncodeRGBM(const TArray<FLinearColor>& InPixels, TArra
 		OutPixels[Index] = EncodeRGBM(InPixels[Index], MaxRange);
 	}
 }
+
+bool FGLTFTextureUtility::LoadPlatformData(UTexture2D* Texture)
+{
+	if (Texture->PlatformData == nullptr || Texture->PlatformData->Mips.Num() == 0)
+	{
+		return false;
+	}
+
+	if (Texture->PlatformData->Mips[0].BulkData.GetBulkDataSize() == 0)
+	{
+		// TODO: is this correct handling?
+		Texture->ForceRebuildPlatformData();
+		if (Texture->PlatformData->Mips[0].BulkData.GetBulkDataSize() == 0)
+		{
+			return false;
+		}
+	}
+
+	return true;
+}
+
+bool FGLTFTextureUtility::LoadPlatformData(UTextureCube* TextureCube)
+{
+	if (TextureCube->PlatformData == nullptr || TextureCube->PlatformData->Mips.Num() == 0)
+	{
+		return false;
+	}
+
+	if (TextureCube->PlatformData->Mips[0].BulkData.GetBulkDataSize() == 0)
+	{
+		// TODO: is this correct handling?
+		TextureCube->ForceRebuildPlatformData();
+		if (TextureCube->PlatformData->Mips[0].BulkData.GetBulkDataSize() == 0)
+		{
+			return false;
+		}
+	}
+
+	return true;
+}
diff --git a/Engine/Plugins/Enterprise/GLTFExporter/MigratedFromGit/Source/GLTFExporter/Private/Converters/GLTFTextureUtility.h b/Engine/Plugins/Enterprise/GLTFExporter/MigratedFromGit/Source/GLTFExporter/Private/Converters/GLTFTextureUtility.h
--- a/Engine/Plugins/Enterprise/GLTFExporter/MigratedFromGit/Source/GLTFExporter/Private/Converters/GLTFTextureUtility.h
+++ b/Engine/Plugins/Enterprise/GLTFExporter/MigratedFromGit/Source/GLTFExporter/Private/Converters/GLTFTextureUtility.h
@@ -36,4 +36,8 @@ struct FGLTFTextureUtility
 
 	static FColor EncodeRGBM(const FLinearColor& Color, float MaxRange = 8);
 	static void EncodeRGBM(const TArray<FLinearColor>& InPixels, TArray<FColor>& OutPixels, float MaxRange = 8);
+
+	// TODO: maybe use template specialization to avoid the need for duplicated functions
+	static bool LoadPlatformData(UTexture2D* Texture);
+	static bool LoadPlatformData(UTextureCube* TextureCube);
 };
