[GLTFExporter] refactor: add and update some TODO comments

commit b85366ee4b494b71fa6ef5d0a29e74aa31ae4957
Date: Fri, 4 Sep 2020 23:16:15 +0200

    refactor: add and update some TODO comments
    (cherry picked from commit 64b5575a5095d178b9480571cf6fb158f8463994​)

#skipci

diff --git a/Engine/Plugins/Enterprise/GLTFExporter/MigratedFromGit/Source/GLTFExporter/Private/Converters/GLTFTextureUtility.cpp b/Engine/Plugins/Enterprise/GLTFExporter/MigratedFromGit/Source/GLTFExporter/Private/Converters/GLTFTextureUtility.cpp
--- a/Engine/Plugins/Enterprise/GLTFExporter/MigratedFromGit/Source/GLTFExporter/Private/Converters/GLTFTextureUtility.cpp
+++ b/Engine/Plugins/Enterprise/GLTFExporter/MigratedFromGit/Source/GLTFExporter/Private/Converters/GLTFTextureUtility.cpp
@@ -102,6 +102,7 @@ UTexture2D* FGLTFTextureUtility::CreateTransientTexture(const void* RawData, int
 {
 	check(CalculateImageBytes(Size.X, Size.Y, 0, Format) == ByteLength);
 
+	// TODO: do these temp textures need to be part of the root set to avoid garbage collection?
 	UTexture2D* Texture = UTexture2D::CreateTransient(Size.X, Size.Y, Format);
 
 	void* MipData = Texture->PlatformData->Mips[0].BulkData.Lock(LOCK_READ_WRITE);
diff --git a/Engine/Plugins/Enterprise/GLTFExporter/MigratedFromGit/Source/GLTFExporter/Private/Converters/GLTFTextureUtility.h b/Engine/Plugins/Enterprise/GLTFExporter/MigratedFromGit/Source/GLTFExporter/Private/Converters/GLTFTextureUtility.h
--- a/Engine/Plugins/Enterprise/GLTFExporter/MigratedFromGit/Source/GLTFExporter/Private/Converters/GLTFTextureUtility.h
+++ b/Engine/Plugins/Enterprise/GLTFExporter/MigratedFromGit/Source/GLTFExporter/Private/Converters/GLTFTextureUtility.h
@@ -29,6 +29,7 @@ struct FGLTFTextureUtility
 	static bool DrawTexture(UTextureRenderTarget2D* OutTarget, const UTexture2D* InSource, const FMatrix& InTransform = FMatrix::Identity);
 	static bool RotateTexture(UTextureRenderTarget2D* OutTarget, const UTexture2D* InSource, float InDegrees);
 
+	// TODO: maybe use template specialization to avoid the need for duplicated functions
 	static UTexture2D* CreateTextureFromCubeFace(const UTextureCube* TextureCube, ECubeFace CubeFace);
 	static UTexture2D* CreateTextureFromCubeFace(const UTextureRenderTargetCube* RenderTargetCube, ECubeFace CubeFace);
 
diff --git a/Engine/Plugins/Enterprise/GLTFExporter/MigratedFromGit/Source/GLTFExporter/Private/Exporters/GLTFExporter.cpp b/Engine/Plugins/Enterprise/GLTFExporter/MigratedFromGit/Source/GLTFExporter/Private/Exporters/GLTFExporter.cpp
--- a/Engine/Plugins/Enterprise/GLTFExporter/MigratedFromGit/Source/GLTFExporter/Private/Exporters/GLTFExporter.cpp
+++ b/Engine/Plugins/Enterprise/GLTFExporter/MigratedFromGit/Source/GLTFExporter/Private/Exporters/GLTFExporter.cpp
@@ -33,6 +33,8 @@ bool UGLTFExporter::ExportBinary(UObject* Object, const TCHAR* Type, FArchive& A
 
 	bool bSuccess = true;
 
+	// TODO: add support for UAssetExportTask::IgnoreObjectList?
+
 	FGLTFContainerBuilder Builder(Options, bSelectedOnly);
 	if (AddObject(Builder, Object))
 	{
@@ -50,8 +52,7 @@ bool UGLTFExporter::ExportBinary(UObject* Object, const TCHAR* Type, FArchive& A
 		bSuccess = false;
 	}
 
-	// TODO: should we set bSuccess to false if Builder.GetErrorMessageCount() > 0
-	// even if both AddObject and Serialize has reported success?
+	// TODO: should we copy messages to UAssetExportTask::Errors?
 
 	if (FApp::IsUnattended())
 	{
