[GLTFExporter] feat: support material functions and attribute layers when searching for texture coordinate indices

commit 7b3f193997fc61335cc0a87568781e2acb5554d7
Date: Sun, 6 Sep 2020 20:40:20 +0200

    feat: support material functions and attribute layers when searching for texture coordinate indices
    (cherry picked from commit 0767bec74f6d5198caf4f422fd2b3812ea9b1c51​)

#skipci

diff --git a/Engine/Plugins/Enterprise/GLTFExporter/MigratedFromGit/Source/GLTFExporter/Private/Converters/GLTFMaterialUtility.cpp b/Engine/Plugins/Enterprise/GLTFExporter/MigratedFromGit/Source/GLTFExporter/Private/Converters/GLTFMaterialUtility.cpp
--- a/Engine/Plugins/Enterprise/GLTFExporter/MigratedFromGit/Source/GLTFExporter/Private/Converters/GLTFMaterialUtility.cpp
+++ b/Engine/Plugins/Enterprise/GLTFExporter/MigratedFromGit/Source/GLTFExporter/Private/Converters/GLTFMaterialUtility.cpp
@@ -397,18 +397,48 @@ void FGLTFMaterialUtility::GetAllTextureCoordinateIndices(const FExpressionInput
 	TArray<UMaterialExpression*> InputExpressions;
 	Expression->GetAllInputExpressions(InputExpressions);
 
-	for (const UMaterialExpression* InputExpression : InputExpressions)
+	GetAllTextureCoordinateIndices(InputExpressions, OutTexCoords);
+}
+
+void FGLTFMaterialUtility::GetAllTextureCoordinateIndices(const TArray<UMaterialExpression*>& Expressions, TSet<int32>& OutTexCoords)
+{
+	for (const UMaterialExpression* Expression : Expressions)
 	{
-		if (const UMaterialExpressionTextureSample* TextureSampler = Cast<UMaterialExpressionTextureSample>(InputExpression))
+		if (const UMaterialExpressionTextureSample* TextureSample = Cast<UMaterialExpressionTextureSample>(Expression))
 		{
-			if (TextureSampler->Coordinates.Expression == nullptr)
+			if (TextureSample->Coordinates.Expression == nullptr)
 			{
-				OutTexCoords.Add(TextureSampler->ConstCoordinate);
+				OutTexCoords.Add(TextureSample->ConstCoordinate);
 			}
 		}
-		else if (const UMaterialExpressionTextureCoordinate* TextureCoordinate = Cast<UMaterialExpressionTextureCoordinate>(InputExpression))
+		else if (const UMaterialExpressionTextureCoordinate* TextureCoordinate = Cast<UMaterialExpressionTextureCoordinate>(Expression))
 		{
 			OutTexCoords.Add(TextureCoordinate->CoordinateIndex);
 		}
+		else if (const UMaterialExpressionMaterialFunctionCall* FunctionCall = Cast<UMaterialExpressionMaterialFunctionCall>(Expression))
+		{
+			GetAllTextureCoordinateIndices(FunctionCall->Function->FunctionExpressions, OutTexCoords);
+        }
+		else if (const UMaterialExpressionMaterialAttributeLayers* AttributeLayers = Cast<UMaterialExpressionMaterialAttributeLayers>(Expression))
+		{
+			const TArray<UMaterialFunctionInterface*>& Layers = AttributeLayers->GetLayers();
+			const TArray<UMaterialFunctionInterface*>& Blends = AttributeLayers->GetBlends();
+
+			for (UMaterialFunctionInterface* Layer : AttributeLayers->GetLayers())
+			{
+				if (Layer != nullptr && Layer->GetFunctionExpressions() != nullptr)
+				{
+					GetAllTextureCoordinateIndices(*Layer->GetFunctionExpressions(), OutTexCoords);
+				}
+			}
+
+			for (UMaterialFunctionInterface* Blend : AttributeLayers->GetBlends())
+			{
+				if (Blend != nullptr && Blend->GetFunctionExpressions() != nullptr)
+				{
+					GetAllTextureCoordinateIndices(*Blend->GetFunctionExpressions(), OutTexCoords);
+				}
+			}
+		}
 	}
 }
diff --git a/Engine/Plugins/Enterprise/GLTFExporter/MigratedFromGit/Source/GLTFExporter/Private/Converters/GLTFMaterialUtility.h b/Engine/Plugins/Enterprise/GLTFExporter/MigratedFromGit/Source/GLTFExporter/Private/Converters/GLTFMaterialUtility.h
--- a/Engine/Plugins/Enterprise/GLTFExporter/MigratedFromGit/Source/GLTFExporter/Private/Converters/GLTFMaterialUtility.h
+++ b/Engine/Plugins/Enterprise/GLTFExporter/MigratedFromGit/Source/GLTFExporter/Private/Converters/GLTFMaterialUtility.h
@@ -61,4 +61,5 @@ struct FGLTFMaterialUtility
 
 	static bool TryGetTextureCoordinateIndex(const UMaterialExpressionTextureSample* TextureSampler, int32& TexCoord);
 	static void GetAllTextureCoordinateIndices(const FExpressionInput& ExpressionInput, TSet<int32>& OutTexCoords);
+	static void GetAllTextureCoordinateIndices(const TArray<UMaterialExpression*>& Expressions, TSet<int32>& OutTexCoords);
 };
