[GLTFExporter] refactor: improve expansion of function expressions to avoid infinite recursion

commit 3f85fb3f095263af9f92360e2755a753bc73f8ee
Date: Sun, 6 Sep 2020 21:35:52 +0200

    refactor: improve expansion of function expressions to avoid infinite recursion
    (cherry picked from commit ad15d29ebd8ef53cd18314e5c0c338c2338f5f70​)

#skipci

diff --git a/Engine/Plugins/Enterprise/GLTFExporter/MigratedFromGit/Source/GLTFExporter/Private/Converters/GLTFMaterialUtility.cpp b/Engine/Plugins/Enterprise/GLTFExporter/MigratedFromGit/Source/GLTFExporter/Private/Converters/GLTFMaterialUtility.cpp
--- a/Engine/Plugins/Enterprise/GLTFExporter/MigratedFromGit/Source/GLTFExporter/Private/Converters/GLTFMaterialUtility.cpp
+++ b/Engine/Plugins/Enterprise/GLTFExporter/MigratedFromGit/Source/GLTFExporter/Private/Converters/GLTFMaterialUtility.cpp
@@ -157,6 +157,8 @@ const FExpressionInput* FGLTFMaterialUtility::GetInputForProperty(const UMateria
 
 const UMaterialExpressionCustomOutput* FGLTFMaterialUtility::GetCustomOutputByName(const UMaterialInterface* Material, const FString& Name)
 {
+	// TODO: should we also search inside material functions and attribute layers?
+
 	for (const UMaterialExpression* Expression : Material->GetMaterial()->Expressions)
 	{
 		const UMaterialExpressionCustomOutput* CustomOutput = Cast<UMaterialExpressionCustomOutput>(Expression);
@@ -386,22 +388,19 @@ bool FGLTFMaterialUtility::TryGetTextureCoordinateIndex(const UMaterialExpressio
 	return false;
 }
 
-void FGLTFMaterialUtility::GetAllTextureCoordinateIndices(const FExpressionInput& ExpressionInput, TSet<int32>& OutTexCoords)
+void FGLTFMaterialUtility::GetAllTextureCoordinateIndices(const FExpressionInput& Input, TSet<int32>& OutTexCoords)
 {
-	UMaterialExpression* Expression = ExpressionInput.Expression;
-	if (Expression == nullptr)
+	UMaterialExpression* InputExpression = Input.Expression;
+	if (InputExpression == nullptr)
 	{
 		return;
 	}
 
-	TArray<UMaterialExpression*> InputExpressions;
-	Expression->GetAllInputExpressions(InputExpressions);
+	TArray<UMaterialExpression*> Expressions;
+	InputExpression->GetAllInputExpressions(Expressions);
 
-	GetAllTextureCoordinateIndices(InputExpressions, OutTexCoords);
-}
+	ExpandAllFunctionExpressions(Expressions);
 
-void FGLTFMaterialUtility::GetAllTextureCoordinateIndices(const TArray<UMaterialExpression*>& Expressions, TSet<int32>& OutTexCoords)
-{
 	for (const UMaterialExpression* Expression : Expressions)
 	{
 		if (const UMaterialExpressionTextureSample* TextureSample = Cast<UMaterialExpressionTextureSample>(Expression))
@@ -415,30 +414,43 @@ void FGLTFMaterialUtility::GetAllTextureCoordinateIndices(const TArray<UMaterial
 		{
 			OutTexCoords.Add(TextureCoordinate->CoordinateIndex);
 		}
-		else if (const UMaterialExpressionMaterialFunctionCall* FunctionCall = Cast<UMaterialExpressionMaterialFunctionCall>(Expression))
+	}
+}
+
+void FGLTFMaterialUtility::ExpandAllFunctionExpressions(TArray<UMaterialExpression*>& InOutExpressions)
+{
+	TArray<UMaterialExpression*> UnexpandedExpressions(InOutExpressions);
+
+	for (const UMaterialExpression* UnexpandedExpression : UnexpandedExpressions)
+	{
+		if (const UMaterialExpressionMaterialFunctionCall* FunctionCall = Cast<UMaterialExpressionMaterialFunctionCall>(UnexpandedExpression))
 		{
-			GetAllTextureCoordinateIndices(FunctionCall->Function->FunctionExpressions, OutTexCoords);
-        }
-		else if (const UMaterialExpressionMaterialAttributeLayers* AttributeLayers = Cast<UMaterialExpressionMaterialAttributeLayers>(Expression))
+			if (UMaterialFunctionInterface* Function = FunctionCall->MaterialFunction)
+			{
+				Function->GetAllExpressionsOfType<UMaterialExpression>(InOutExpressions);
+			}
+		}
+		else if (const UMaterialExpressionMaterialAttributeLayers* AttributeLayers = Cast<UMaterialExpressionMaterialAttributeLayers>(UnexpandedExpression))
 		{
-			const TArray<UMaterialFunctionInterface*>& Layers = AttributeLayers->GetLayers();
-			const TArray<UMaterialFunctionInterface*>& Blends = AttributeLayers->GetBlends();
-
 			for (UMaterialFunctionInterface* Layer : AttributeLayers->GetLayers())
 			{
-				if (Layer != nullptr && Layer->GetFunctionExpressions() != nullptr)
+				if (Layer != nullptr)
 				{
-					GetAllTextureCoordinateIndices(*Layer->GetFunctionExpressions(), OutTexCoords);
+					Layer->GetAllExpressionsOfType<UMaterialExpression>(InOutExpressions);
 				}
 			}
 
 			for (UMaterialFunctionInterface* Blend : AttributeLayers->GetBlends())
 			{
-				if (Blend != nullptr && Blend->GetFunctionExpressions() != nullptr)
+				if (Blend != nullptr)
 				{
-					GetAllTextureCoordinateIndices(*Blend->GetFunctionExpressions(), OutTexCoords);
+					Blend->GetAllExpressionsOfType<UMaterialExpression>(InOutExpressions);
 				}
 			}
 		}
 	}
+
+	InOutExpressions.RemoveAll([](const UMaterialExpression* Expression) {
+        return Expression->IsA<UMaterialExpressionMaterialFunctionCall>() || Expression->IsA<UMaterialExpressionMaterialAttributeLayers>();
+    });
 }
diff --git a/Engine/Plugins/Enterprise/GLTFExporter/MigratedFromGit/Source/GLTFExporter/Private/Converters/GLTFMaterialUtility.h b/Engine/Plugins/Enterprise/GLTFExporter/MigratedFromGit/Source/GLTFExporter/Private/Converters/GLTFMaterialUtility.h
--- a/Engine/Plugins/Enterprise/GLTFExporter/MigratedFromGit/Source/GLTFExporter/Private/Converters/GLTFMaterialUtility.h
+++ b/Engine/Plugins/Enterprise/GLTFExporter/MigratedFromGit/Source/GLTFExporter/Private/Converters/GLTFMaterialUtility.h
@@ -60,6 +60,7 @@ struct FGLTFMaterialUtility
 	static uint32 GetMaskComponentCount(const FExpressionInput& ExpressionInput);
 
 	static bool TryGetTextureCoordinateIndex(const UMaterialExpressionTextureSample* TextureSampler, int32& TexCoord);
-	static void GetAllTextureCoordinateIndices(const FExpressionInput& ExpressionInput, TSet<int32>& OutTexCoords);
-	static void GetAllTextureCoordinateIndices(const TArray<UMaterialExpression*>& Expressions, TSet<int32>& OutTexCoords);
+	static void GetAllTextureCoordinateIndices(const FExpressionInput& Input, TSet<int32>& OutTexCoords);
+
+	static void ExpandAllFunctionExpressions(TArray<UMaterialExpression*>& InOutExpressions);
 };
