[GLTFExporter] feat: add LeafNode to FGLTFJsonNode, to resolve ambiguity if multiple nodes are created for 1 object

commit 0fc4c97b584f94d81f989a845e3bf3b91eceb12d
Date: Tue, 1 Sep 2020 13:20:03 +0200

    feat: add LeafNode to FGLTFJsonNode, to resolve ambiguity if multiple nodes are created for 1 object
    (cherry picked from commit 9ebf066d224b0e6a3c56415750dafa60e5d7dd83​)

#skipci

diff --git a/Engine/Plugins/Enterprise/GLTFExporter/MigratedFromGit/Source/GLTFExporter/Private/Builders/GLTFJsonBuilder.cpp b/Engine/Plugins/Enterprise/GLTFExporter/MigratedFromGit/Source/GLTFExporter/Private/Builders/GLTFJsonBuilder.cpp
--- a/Engine/Plugins/Enterprise/GLTFExporter/MigratedFromGit/Source/GLTFExporter/Private/Builders/GLTFJsonBuilder.cpp
+++ b/Engine/Plugins/Enterprise/GLTFExporter/MigratedFromGit/Source/GLTFExporter/Private/Builders/GLTFJsonBuilder.cpp
@@ -170,6 +170,17 @@ FGLTFJsonLight& FGLTFJsonBuilder::GetLight(FGLTFJsonLightIndex LightIndex)
 	return JsonRoot.Lights[LightIndex];
 }
 
+FGLTFJsonNodeIndex FGLTFJsonBuilder::GetLeafNodeIndex(FGLTFJsonNodeIndex NodeIndex)
+{
+	if (NodeIndex == INDEX_NONE)
+	{
+		return FGLTFJsonNodeIndex(INDEX_NONE);
+	}
+
+	const FGLTFJsonNode& Node = GetNode(NodeIndex);
+	return Node.LeafNode != INDEX_NONE ? Node.LeafNode : NodeIndex;
+}
+
 bool FGLTFJsonBuilder::Serialize(FArchive& Archive, const FString& FilePath)
 {
 	JsonRoot.Serialize(&Archive, true);
diff --git a/Engine/Plugins/Enterprise/GLTFExporter/MigratedFromGit/Source/GLTFExporter/Private/Builders/GLTFJsonBuilder.h b/Engine/Plugins/Enterprise/GLTFExporter/MigratedFromGit/Source/GLTFExporter/Private/Builders/GLTFJsonBuilder.h
--- a/Engine/Plugins/Enterprise/GLTFExporter/MigratedFromGit/Source/GLTFExporter/Private/Builders/GLTFJsonBuilder.h
+++ b/Engine/Plugins/Enterprise/GLTFExporter/MigratedFromGit/Source/GLTFExporter/Private/Builders/GLTFJsonBuilder.h
@@ -47,6 +47,8 @@ public:
 	FGLTFJsonLightMap& GetLightMap(FGLTFJsonLightMapIndex LightMapIndex);
 	FGLTFJsonLight& GetLight(FGLTFJsonLightIndex LightIndex);
 
+	FGLTFJsonNodeIndex GetLeafNodeIndex(FGLTFJsonNodeIndex NodeIndex);
+
 	virtual bool Serialize(FArchive& Archive, const FString& FilePath);
 
 private:
diff --git a/Engine/Plugins/Enterprise/GLTFExporter/MigratedFromGit/Source/GLTFExporter/Private/Json/GLTFJsonNode.h b/Engine/Plugins/Enterprise/GLTFExporter/MigratedFromGit/Source/GLTFExporter/Private/Json/GLTFJsonNode.h
--- a/Engine/Plugins/Enterprise/GLTFExporter/MigratedFromGit/Source/GLTFExporter/Private/Json/GLTFJsonNode.h
+++ b/Engine/Plugins/Enterprise/GLTFExporter/MigratedFromGit/Source/GLTFExporter/Private/Json/GLTFJsonNode.h
@@ -33,6 +33,8 @@ struct FGLTFJsonNode
 	FGLTFJsonLightMapIndex  LightMap;
 	FGLTFJsonLightIndex     Light;
 
+	FGLTFJsonNodeIndex LeafNode;
+
 	TArray<FGLTFJsonNodeIndex> Children;
 
 	FGLTFJsonNode()
