[GLTFExporter] feat: specify LeafNode on the original node when creating child-nodes for the same source-object

commit 6382802f7550ade072f678e0f697b6423b62edba
Date: Tue, 1 Sep 2020 13:26:08 +0200

    feat: specify LeafNode on the original node when creating child-nodes for the same source-object
    (cherry picked from commit 1f51fb9a238d7c6bae0b4b98d2c07cfce96ffe08​)

#skipci

diff --git a/Engine/Plugins/Enterprise/GLTFExporter/MigratedFromGit/Source/GLTFExporter/Private/Converters/GLTFLevelConverters.cpp b/Engine/Plugins/Enterprise/GLTFExporter/MigratedFromGit/Source/GLTFExporter/Private/Converters/GLTFLevelConverters.cpp
--- a/Engine/Plugins/Enterprise/GLTFExporter/MigratedFromGit/Source/GLTFExporter/Private/Converters/GLTFLevelConverters.cpp
+++ b/Engine/Plugins/Enterprise/GLTFExporter/MigratedFromGit/Source/GLTFExporter/Private/Converters/GLTFLevelConverters.cpp
@@ -57,7 +57,7 @@ FGLTFJsonNodeIndex FGLTFSceneComponentConverter::Add(FGLTFConvertBuilder& Builde
 			MeshNode.Scale = FGLTFConverterUtility::ConvertScale(Transform.GetScale3D());
 			MeshNode.Mesh = Builder.GetOrAddMesh(StaticMeshComponent);
 			MeshNode.LightMap = Builder.GetOrAddLightMap(StaticMeshComponent);
-			Builder.AddChildNode(NodeIndex, MeshNode);
+			Builder.GetNode(NodeIndex).LeafNode = Builder.AddChildNode(NodeIndex, MeshNode);
 		}
 		else
 		{
@@ -79,7 +79,7 @@ FGLTFJsonNodeIndex FGLTFSceneComponentConverter::Add(FGLTFConvertBuilder& Builde
 			MeshNode.Name = Node.Name + TEXT("_Mesh");
 			MeshNode.Scale = FGLTFConverterUtility::ConvertScale(Transform.GetScale3D());
 			MeshNode.Mesh = Builder.GetOrAddMesh(SkeletalMeshComponent);
-			Builder.AddChildNode(NodeIndex, MeshNode);
+			Builder.GetNode(NodeIndex).LeafNode = Builder.AddChildNode(NodeIndex, MeshNode);
 		}
 		else
 		{
@@ -96,7 +96,7 @@ FGLTFJsonNodeIndex FGLTFSceneComponentConverter::Add(FGLTFConvertBuilder& Builde
 			CameraNode.Rotation = FGLTFConverterUtility::ConvertCameraDirection();
 			CameraNode.Scale = FGLTFConverterUtility::ConvertScale(bSupportNonUniformScale ? Transform.GetScale3D() : FVector::OneVector);
 			CameraNode.Camera = Builder.GetOrAddCamera(CameraComponent, CameraNode.Name);
-			Builder.AddChildNode(NodeIndex, CameraNode);
+			Builder.GetNode(NodeIndex).LeafNode = Builder.AddChildNode(NodeIndex, CameraNode);
 		}
 		else
 		{
@@ -113,7 +113,7 @@ FGLTFJsonNodeIndex FGLTFSceneComponentConverter::Add(FGLTFConvertBuilder& Builde
 			LightNode.Rotation = FGLTFConverterUtility::ConvertLightDirection();
 			LightNode.Scale = FGLTFConverterUtility::ConvertScale(bSupportNonUniformScale ? Transform.GetScale3D() : FVector::OneVector);
 			LightNode.Light = Builder.GetOrAddLight(LightComponent, LightNode.Name);
-			Builder.AddChildNode(NodeIndex, LightNode);
+			Builder.GetNode(NodeIndex).LeafNode = Builder.AddChildNode(NodeIndex, LightNode);
 		}
 		else
 		{
