[GLTFExporter] fix(variants): reference the correct target for variants by finding the leaf-node for each node

commit 58a83a964a6bf16a0a497632fc915f10b9dc7367
Date: Tue, 1 Sep 2020 13:28:27 +0200

    fix(variants): reference the correct target for variants by finding the leaf-node for each node
    (cherry picked from commit 193fddc17d5505036a9bb834bb7dedb6fdecfb4d​)

#skipci

diff --git a/Engine/Plugins/Enterprise/GLTFExporter/MigratedFromGit/Source/GLTFExporter/Private/Converters/GLTFVariantSetConverters.cpp b/Engine/Plugins/Enterprise/GLTFExporter/MigratedFromGit/Source/GLTFExporter/Private/Converters/GLTFVariantSetConverters.cpp
--- a/Engine/Plugins/Enterprise/GLTFExporter/MigratedFromGit/Source/GLTFExporter/Private/Converters/GLTFVariantSetConverters.cpp
+++ b/Engine/Plugins/Enterprise/GLTFExporter/MigratedFromGit/Source/GLTFExporter/Private/Converters/GLTFVariantSetConverters.cpp
@@ -206,9 +206,11 @@ bool FGLTFLevelVariantSetsConverter::TryParseVisibilityPropertyValue(FGLTFConver
 	}
 
 	const FGLTFJsonNodeIndex NodeIndex = Builder.GetOrAddNode(Target); // TODO: what if the target is not selected for export?
-	FGLTFJsonVariantNodeProperties& NodeProperties = OutVariant.Nodes.FindOrAdd(NodeIndex);
+	const FGLTFJsonNodeIndex LeafNodeIndex = Builder.GetLeafNodeIndex(NodeIndex);
 
-	NodeProperties.Node = NodeIndex;
+	FGLTFJsonVariantNodeProperties& NodeProperties = OutVariant.Nodes.FindOrAdd(LeafNodeIndex);
+
+	NodeProperties.Node = LeafNodeIndex;
 	NodeProperties.bIsVisible = bIsVisible;
 	return true;
 }
@@ -274,9 +276,10 @@ bool FGLTFLevelVariantSetsConverter::TryParseMaterialPropertyValue(FGLTFConvertB
 	VariantMaterial.Index = ElementIndex;
 
 	const FGLTFJsonNodeIndex NodeIndex = Builder.GetOrAddNode(Target);
-	FGLTFJsonVariantNodeProperties& NodeProperties = OutVariant.Nodes.FindOrAdd(NodeIndex);
+	const FGLTFJsonNodeIndex LeafNodeIndex = Builder.GetLeafNodeIndex(NodeIndex);
+	FGLTFJsonVariantNodeProperties& NodeProperties = OutVariant.Nodes.FindOrAdd(LeafNodeIndex);
 
-	NodeProperties.Node = NodeIndex;
+	NodeProperties.Node = LeafNodeIndex;
 	NodeProperties.Materials.Add(VariantMaterial);
 	return true;
 }
@@ -329,10 +332,11 @@ bool FGLTFLevelVariantSetsConverter::TryParseStaticMeshPropertyValue(FGLTFConver
 
 	const FGLTFMaterialArray OverrideMaterials(MeshComponent->OverrideMaterials);
 	const FGLTFJsonNodeIndex NodeIndex = Builder.GetOrAddNode(Target);
+	const FGLTFJsonNodeIndex LeafNodeIndex = Builder.GetLeafNodeIndex(NodeIndex);
 	const FGLTFJsonMeshIndex MeshIndex = Builder.GetOrAddMesh(StaticMesh, -1, nullptr, OverrideMaterials);
-	FGLTFJsonVariantNodeProperties& NodeProperties = OutVariant.Nodes.FindOrAdd(NodeIndex);
+	FGLTFJsonVariantNodeProperties& NodeProperties = OutVariant.Nodes.FindOrAdd(LeafNodeIndex);
 
-	NodeProperties.Node = NodeIndex;
+	NodeProperties.Node = LeafNodeIndex;
 	NodeProperties.Mesh = MeshIndex;
 	return true;
 }
@@ -386,9 +390,10 @@ bool FGLTFLevelVariantSetsConverter::TryParseSkeletalMeshPropertyValue(FGLTFConv
 	const FGLTFMaterialArray OverrideMaterials(MeshComponent->OverrideMaterials);
 	const FGLTFJsonNodeIndex NodeIndex = Builder.GetOrAddNode(Target);
 	const FGLTFJsonMeshIndex MeshIndex = Builder.GetOrAddMesh(SkeletalMesh, -1, nullptr, nullptr, OverrideMaterials);
-	FGLTFJsonVariantNodeProperties& NodeProperties = OutVariant.Nodes.FindOrAdd(NodeIndex);
+	const FGLTFJsonNodeIndex LeafNodeIndex = Builder.GetLeafNodeIndex(NodeIndex);
+	FGLTFJsonVariantNodeProperties& NodeProperties = OutVariant.Nodes.FindOrAdd(LeafNodeIndex);
 
-	NodeProperties.Node = NodeIndex;
+	NodeProperties.Node = LeafNodeIndex;
 	NodeProperties.Mesh = MeshIndex;
 	return true;
 }
