[GLTFExporter] refactor: change LeafNode to ComponentNode and add helper-function for adding & registering component-nodes

commit 64693d9fd99f0d1680c7303350527cc77b48ff39
Date: Tue, 1 Sep 2020 16:46:00 +0200

    refactor: change LeafNode to ComponentNode and add helper-function for adding & registering component-nodes
    (cherry picked from commit abaa288dc658cab91fe499d029ef5cbd4b77ec23​)

#skipci

diff --git a/Engine/Plugins/Enterprise/GLTFExporter/MigratedFromGit/Source/GLTFExporter/Private/Builders/GLTFJsonBuilder.cpp b/Engine/Plugins/Enterprise/GLTFExporter/MigratedFromGit/Source/GLTFExporter/Private/Builders/GLTFJsonBuilder.cpp
--- a/Engine/Plugins/Enterprise/GLTFExporter/MigratedFromGit/Source/GLTFExporter/Private/Builders/GLTFJsonBuilder.cpp
+++ b/Engine/Plugins/Enterprise/GLTFExporter/MigratedFromGit/Source/GLTFExporter/Private/Builders/GLTFJsonBuilder.cpp
@@ -95,6 +95,18 @@ FGLTFJsonNodeIndex FGLTFJsonBuilder::AddChildNode(FGLTFJsonNodeIndex ParentIndex
 	return ChildIndex;
 }
 
+FGLTFJsonNodeIndex FGLTFJsonBuilder::AddChildComponentNode(FGLTFJsonNodeIndex ParentIndex, const FGLTFJsonNode& JsonNode)
+{
+	const FGLTFJsonNodeIndex ChildIndex = AddChildNode(ParentIndex, JsonNode);
+
+	if (ParentIndex != INDEX_NONE)
+	{
+		GetNode(ParentIndex).ComponentNode = ChildIndex;
+	}
+
+	return ChildIndex;
+}
+
 FGLTFJsonAccessor& FGLTFJsonBuilder::GetAccessor(FGLTFJsonAccessorIndex AccessorIndex)
 {
 	return JsonRoot.Accessors[AccessorIndex];
@@ -170,7 +182,7 @@ FGLTFJsonLight& FGLTFJsonBuilder::GetLight(FGLTFJsonLightIndex LightIndex)
 	return JsonRoot.Lights[LightIndex];
 }
 
-FGLTFJsonNodeIndex FGLTFJsonBuilder::GetLeafNodeIndex(FGLTFJsonNodeIndex NodeIndex)
+FGLTFJsonNodeIndex FGLTFJsonBuilder::GetComponentNodeIndex(FGLTFJsonNodeIndex NodeIndex)
 {
 	if (NodeIndex == INDEX_NONE)
 	{
@@ -178,7 +190,7 @@ FGLTFJsonNodeIndex FGLTFJsonBuilder::GetLeafNodeIndex(FGLTFJsonNodeIndex NodeInd
 	}
 
 	const FGLTFJsonNode& Node = GetNode(NodeIndex);
-	return Node.LeafNode != INDEX_NONE ? Node.LeafNode : NodeIndex;
+	return Node.ComponentNode != INDEX_NONE ? Node.ComponentNode : NodeIndex;
 }
 
 bool FGLTFJsonBuilder::Serialize(FArchive& Archive, const FString& FilePath)
diff --git a/Engine/Plugins/Enterprise/GLTFExporter/MigratedFromGit/Source/GLTFExporter/Private/Builders/GLTFJsonBuilder.h b/Engine/Plugins/Enterprise/GLTFExporter/MigratedFromGit/Source/GLTFExporter/Private/Builders/GLTFJsonBuilder.h
--- a/Engine/Plugins/Enterprise/GLTFExporter/MigratedFromGit/Source/GLTFExporter/Private/Builders/GLTFJsonBuilder.h
+++ b/Engine/Plugins/Enterprise/GLTFExporter/MigratedFromGit/Source/GLTFExporter/Private/Builders/GLTFJsonBuilder.h
@@ -30,6 +30,7 @@ public:
 	FGLTFJsonLightIndex AddLight(const FGLTFJsonLight& JsonLight = {});
 
 	FGLTFJsonNodeIndex AddChildNode(FGLTFJsonNodeIndex ParentNodeIndex, const FGLTFJsonNode& JsonNode = {});
+	FGLTFJsonNodeIndex AddChildComponentNode(FGLTFJsonNodeIndex ParentNodeIndex, const FGLTFJsonNode& JsonNode = {});
 
 	FGLTFJsonAccessor& GetAccessor(FGLTFJsonAccessorIndex AccessorIndex);
 	FGLTFJsonBuffer& GetBuffer(FGLTFJsonBufferIndex BufferIndex);
@@ -47,7 +48,7 @@ public:
 	FGLTFJsonLightMap& GetLightMap(FGLTFJsonLightMapIndex LightMapIndex);
 	FGLTFJsonLight& GetLight(FGLTFJsonLightIndex LightIndex);
 
-	FGLTFJsonNodeIndex GetLeafNodeIndex(FGLTFJsonNodeIndex NodeIndex);
+	FGLTFJsonNodeIndex GetComponentNodeIndex(FGLTFJsonNodeIndex NodeIndex);
 
 	virtual bool Serialize(FArchive& Archive, const FString& FilePath);
 
diff --git a/Engine/Plugins/Enterprise/GLTFExporter/MigratedFromGit/Source/GLTFExporter/Private/Converters/GLTFLevelConverters.cpp b/Engine/Plugins/Enterprise/GLTFExporter/MigratedFromGit/Source/GLTFExporter/Private/Converters/GLTFLevelConverters.cpp
--- a/Engine/Plugins/Enterprise/GLTFExporter/MigratedFromGit/Source/GLTFExporter/Private/Converters/GLTFLevelConverters.cpp
+++ b/Engine/Plugins/Enterprise/GLTFExporter/MigratedFromGit/Source/GLTFExporter/Private/Converters/GLTFLevelConverters.cpp
@@ -57,7 +57,7 @@ FGLTFJsonNodeIndex FGLTFSceneComponentConverter::Add(FGLTFConvertBuilder& Builde
 			MeshNode.Scale = FGLTFConverterUtility::ConvertScale(Transform.GetScale3D());
 			MeshNode.Mesh = Builder.GetOrAddMesh(StaticMeshComponent);
 			MeshNode.LightMap = Builder.GetOrAddLightMap(StaticMeshComponent);
-			Builder.GetNode(NodeIndex).LeafNode = Builder.AddChildNode(NodeIndex, MeshNode);
+			Builder.AddChildComponentNode(NodeIndex, MeshNode);
 		}
 		else
 		{
@@ -79,7 +79,7 @@ FGLTFJsonNodeIndex FGLTFSceneComponentConverter::Add(FGLTFConvertBuilder& Builde
 			MeshNode.Name = Node.Name + TEXT("_Mesh");
 			MeshNode.Scale = FGLTFConverterUtility::ConvertScale(Transform.GetScale3D());
 			MeshNode.Mesh = Builder.GetOrAddMesh(SkeletalMeshComponent);
-			Builder.GetNode(NodeIndex).LeafNode = Builder.AddChildNode(NodeIndex, MeshNode);
+			Builder.AddChildComponentNode(NodeIndex, MeshNode);
 		}
 		else
 		{
@@ -96,7 +96,7 @@ FGLTFJsonNodeIndex FGLTFSceneComponentConverter::Add(FGLTFConvertBuilder& Builde
 			CameraNode.Rotation = FGLTFConverterUtility::ConvertCameraDirection();
 			CameraNode.Scale = FGLTFConverterUtility::ConvertScale(bSupportNonUniformScale ? Transform.GetScale3D() : FVector::OneVector);
 			CameraNode.Camera = Builder.GetOrAddCamera(CameraComponent, CameraNode.Name);
-			Builder.GetNode(NodeIndex).LeafNode = Builder.AddChildNode(NodeIndex, CameraNode);
+			Builder.AddChildComponentNode(NodeIndex, CameraNode);
 		}
 		else
 		{
@@ -113,7 +113,7 @@ FGLTFJsonNodeIndex FGLTFSceneComponentConverter::Add(FGLTFConvertBuilder& Builde
 			LightNode.Rotation = FGLTFConverterUtility::ConvertLightDirection();
 			LightNode.Scale = FGLTFConverterUtility::ConvertScale(bSupportNonUniformScale ? Transform.GetScale3D() : FVector::OneVector);
 			LightNode.Light = Builder.GetOrAddLight(LightComponent, LightNode.Name);
-			Builder.GetNode(NodeIndex).LeafNode = Builder.AddChildNode(NodeIndex, LightNode);
+			Builder.AddChildComponentNode(NodeIndex, LightNode);
 		}
 		else
 		{
diff --git a/Engine/Plugins/Enterprise/GLTFExporter/MigratedFromGit/Source/GLTFExporter/Private/Converters/GLTFVariantSetConverters.cpp b/Engine/Plugins/Enterprise/GLTFExporter/MigratedFromGit/Source/GLTFExporter/Private/Converters/GLTFVariantSetConverters.cpp
--- a/Engine/Plugins/Enterprise/GLTFExporter/MigratedFromGit/Source/GLTFExporter/Private/Converters/GLTFVariantSetConverters.cpp
+++ b/Engine/Plugins/Enterprise/GLTFExporter/MigratedFromGit/Source/GLTFExporter/Private/Converters/GLTFVariantSetConverters.cpp
@@ -206,11 +206,11 @@ bool FGLTFLevelVariantSetsConverter::TryParseVisibilityPropertyValue(FGLTFConver
 	}
 
 	const FGLTFJsonNodeIndex NodeIndex = Builder.GetOrAddNode(Target); // TODO: what if the target is not selected for export?
-	const FGLTFJsonNodeIndex LeafNodeIndex = Builder.GetLeafNodeIndex(NodeIndex);
+	const FGLTFJsonNodeIndex ComponentNodeIndex = Builder.GetComponentNodeIndex(NodeIndex);
 
-	FGLTFJsonVariantNodeProperties& NodeProperties = OutVariant.Nodes.FindOrAdd(LeafNodeIndex);
+	FGLTFJsonVariantNodeProperties& NodeProperties = OutVariant.Nodes.FindOrAdd(ComponentNodeIndex);
 
-	NodeProperties.Node = LeafNodeIndex;
+	NodeProperties.Node = ComponentNodeIndex;
 	NodeProperties.bIsVisible = bIsVisible;
 	return true;
 }
@@ -276,10 +276,10 @@ bool FGLTFLevelVariantSetsConverter::TryParseMaterialPropertyValue(FGLTFConvertB
 	VariantMaterial.Index = ElementIndex;
 
 	const FGLTFJsonNodeIndex NodeIndex = Builder.GetOrAddNode(Target);
-	const FGLTFJsonNodeIndex LeafNodeIndex = Builder.GetLeafNodeIndex(NodeIndex);
-	FGLTFJsonVariantNodeProperties& NodeProperties = OutVariant.Nodes.FindOrAdd(LeafNodeIndex);
+	const FGLTFJsonNodeIndex ComponentNodeIndex = Builder.GetComponentNodeIndex(NodeIndex);
+	FGLTFJsonVariantNodeProperties& NodeProperties = OutVariant.Nodes.FindOrAdd(ComponentNodeIndex);
 
-	NodeProperties.Node = LeafNodeIndex;
+	NodeProperties.Node = ComponentNodeIndex;
 	NodeProperties.Materials.Add(VariantMaterial);
 	return true;
 }
@@ -332,11 +332,11 @@ bool FGLTFLevelVariantSetsConverter::TryParseStaticMeshPropertyValue(FGLTFConver
 
 	const FGLTFMaterialArray OverrideMaterials(MeshComponent->OverrideMaterials);
 	const FGLTFJsonNodeIndex NodeIndex = Builder.GetOrAddNode(Target);
-	const FGLTFJsonNodeIndex LeafNodeIndex = Builder.GetLeafNodeIndex(NodeIndex);
+	const FGLTFJsonNodeIndex ComponentNodeIndex = Builder.GetComponentNodeIndex(NodeIndex);
 	const FGLTFJsonMeshIndex MeshIndex = Builder.GetOrAddMesh(StaticMesh, -1, nullptr, OverrideMaterials);
-	FGLTFJsonVariantNodeProperties& NodeProperties = OutVariant.Nodes.FindOrAdd(LeafNodeIndex);
+	FGLTFJsonVariantNodeProperties& NodeProperties = OutVariant.Nodes.FindOrAdd(ComponentNodeIndex);
 
-	NodeProperties.Node = LeafNodeIndex;
+	NodeProperties.Node = ComponentNodeIndex;
 	NodeProperties.Mesh = MeshIndex;
 	return true;
 }
@@ -390,10 +390,10 @@ bool FGLTFLevelVariantSetsConverter::TryParseSkeletalMeshPropertyValue(FGLTFConv
 	const FGLTFMaterialArray OverrideMaterials(MeshComponent->OverrideMaterials);
 	const FGLTFJsonNodeIndex NodeIndex = Builder.GetOrAddNode(Target);
 	const FGLTFJsonMeshIndex MeshIndex = Builder.GetOrAddMesh(SkeletalMesh, -1, nullptr, nullptr, OverrideMaterials);
-	const FGLTFJsonNodeIndex LeafNodeIndex = Builder.GetLeafNodeIndex(NodeIndex);
-	FGLTFJsonVariantNodeProperties& NodeProperties = OutVariant.Nodes.FindOrAdd(LeafNodeIndex);
+	const FGLTFJsonNodeIndex ComponentNodeIndex = Builder.GetComponentNodeIndex(NodeIndex);
+	FGLTFJsonVariantNodeProperties& NodeProperties = OutVariant.Nodes.FindOrAdd(ComponentNodeIndex);
 
-	NodeProperties.Node = LeafNodeIndex;
+	NodeProperties.Node = ComponentNodeIndex;
 	NodeProperties.Mesh = MeshIndex;
 	return true;
 }
diff --git a/Engine/Plugins/Enterprise/GLTFExporter/MigratedFromGit/Source/GLTFExporter/Private/Json/GLTFJsonNode.h b/Engine/Plugins/Enterprise/GLTFExporter/MigratedFromGit/Source/GLTFExporter/Private/Json/GLTFJsonNode.h
--- a/Engine/Plugins/Enterprise/GLTFExporter/MigratedFromGit/Source/GLTFExporter/Private/Json/GLTFJsonNode.h
+++ b/Engine/Plugins/Enterprise/GLTFExporter/MigratedFromGit/Source/GLTFExporter/Private/Json/GLTFJsonNode.h
@@ -33,7 +33,7 @@ struct FGLTFJsonNode
 	FGLTFJsonLightMapIndex  LightMap;
 	FGLTFJsonLightIndex     Light;
 
-	FGLTFJsonNodeIndex LeafNode;
+	FGLTFJsonNodeIndex ComponentNode;
 
 	TArray<FGLTFJsonNodeIndex> Children;
 
