[GLTFExporter] refactor: use the convenience-function GetComponentScale() when possible

commit 31f7e44a32bc6aa258e397b4054197f7ab1a65dd
Date: Wed, 2 Sep 2020 08:15:16 +0200

    refactor: use the convenience-function GetComponentScale() when possible
    (cherry picked from commit 3fe8ad7d8e0af223d1ac04b3d69fb94a962ca747​)

#skipci

diff --git a/Engine/Plugins/Enterprise/GLTFExporter/MigratedFromGit/Source/GLTFExporter/Private/Converters/GLTFLevelConverters.cpp b/Engine/Plugins/Enterprise/GLTFExporter/MigratedFromGit/Source/GLTFExporter/Private/Converters/GLTFLevelConverters.cpp
--- a/Engine/Plugins/Enterprise/GLTFExporter/MigratedFromGit/Source/GLTFExporter/Private/Converters/GLTFLevelConverters.cpp
+++ b/Engine/Plugins/Enterprise/GLTFExporter/MigratedFromGit/Source/GLTFExporter/Private/Converters/GLTFLevelConverters.cpp
@@ -26,8 +26,8 @@ FGLTFJsonNodeIndex FGLTFSceneComponentConverter::Add(FGLTFConvertBuilder& Builde
 
 	const FTransform Transform = SceneComponent->GetComponentTransform();
 	const FTransform RelativeTransform = bIsRootNode ? Transform : Transform.GetRelativeTransform(ParentComponent->GetComponentTransform());
-	
-	const FVector ParentScale = ParentComponent != nullptr ? ParentComponent->GetComponentTransform().GetScale3D() : FVector::OneVector;
+
+	const FVector ParentScale = ParentComponent != nullptr ? ParentComponent->GetComponentScale() : FVector::OneVector;
 	const FVector Translation = bSupportNonUniformScale ? RelativeTransform.GetTranslation() * ParentScale : RelativeTransform.GetTranslation();
 	const FVector Scale = bSupportNonUniformScale ? FVector::OneVector : RelativeTransform.GetScale3D();
 
