[GLTFExporter] refactor: rename export-option bSupportNonUniformScale to bExportNonUniformScale

commit 4225e9e6d3a03dd7f48c9cd598d1e95622e75049
Date: Wed, 2 Sep 2020 08:22:41 +0200

    refactor: rename export-option bSupportNonUniformScale to bExportNonUniformScale
    (cherry picked from commit 9e31e95aa469295ec9717452ba957a11628dae25​)

#skipci

diff --git a/Engine/Plugins/Enterprise/GLTFExporter/MigratedFromGit/Source/GLTFExporter/Private/Converters/GLTFLevelConverters.cpp b/Engine/Plugins/Enterprise/GLTFExporter/MigratedFromGit/Source/GLTFExporter/Private/Converters/GLTFLevelConverters.cpp
--- a/Engine/Plugins/Enterprise/GLTFExporter/MigratedFromGit/Source/GLTFExporter/Private/Converters/GLTFLevelConverters.cpp
+++ b/Engine/Plugins/Enterprise/GLTFExporter/MigratedFromGit/Source/GLTFExporter/Private/Converters/GLTFLevelConverters.cpp
@@ -17,7 +17,7 @@ FGLTFJsonNodeIndex FGLTFSceneComponentConverter::Add(FGLTFConvertBuilder& Builde
 
 	const bool bIsRootComponent = Owner->GetRootComponent() == SceneComponent;
 	const bool bIsRootNode = bIsRootComponent && FGLTFActorUtility::IsRootActor(Owner, Builder.bSelectedActorsOnly);
-	const bool bSupportNonUniformScale = Builder.ExportOptions->bSupportNonUniformScale;
+	const bool bExportNonUniformScale = Builder.ExportOptions->bExportNonUniformScale;
 
 	const USceneComponent* ParentComponent = SceneComponent->GetAttachParent();
 	const FGLTFJsonNodeIndex ParentNodeIndex = Builder.GetOrAddNode(ParentComponent);
@@ -28,8 +28,8 @@ FGLTFJsonNodeIndex FGLTFSceneComponentConverter::Add(FGLTFConvertBuilder& Builde
 	const FTransform RelativeTransform = bIsRootNode ? Transform : Transform.GetRelativeTransform(ParentComponent->GetComponentTransform());
 
 	const FVector ParentScale = ParentComponent != nullptr ? ParentComponent->GetComponentScale() : FVector::OneVector;
-	const FVector Translation = bSupportNonUniformScale ? RelativeTransform.GetTranslation() * ParentScale : RelativeTransform.GetTranslation();
-	const FVector Scale = bSupportNonUniformScale ? FVector::OneVector : RelativeTransform.GetScale3D();
+	const FVector Translation = bExportNonUniformScale ? RelativeTransform.GetTranslation() * ParentScale : RelativeTransform.GetTranslation();
+	const FVector Scale = bExportNonUniformScale ? FVector::OneVector : RelativeTransform.GetScale3D();
 
 	const FGLTFJsonNodeIndex NodeIndex = Builder.AddChildNode(ParentNodeIndex);
 	FGLTFJsonNode& Node = Builder.GetNode(NodeIndex);
@@ -50,7 +50,7 @@ FGLTFJsonNodeIndex FGLTFSceneComponentConverter::Add(FGLTFConvertBuilder& Builde
 		// extra nodes, or add some form of mapping that allows the extensions to find the
 		// exact node that the mesh is attached to.
 
-		if (bSupportNonUniformScale)
+		if (bExportNonUniformScale)
 		{
 			FGLTFJsonNode MeshNode;
 			MeshNode.Name = Node.Name + TEXT("_Mesh");
@@ -73,7 +73,7 @@ FGLTFJsonNodeIndex FGLTFSceneComponentConverter::Add(FGLTFConvertBuilder& Builde
 		// extra nodes, or add some form of mapping that allows the extensions to find the
 		// exact node that the mesh is attached to.
 
-		if (bSupportNonUniformScale)
+		if (bExportNonUniformScale)
 		{
 			FGLTFJsonNode MeshNode;
 			MeshNode.Name = Node.Name + TEXT("_Mesh");
@@ -94,7 +94,7 @@ FGLTFJsonNodeIndex FGLTFSceneComponentConverter::Add(FGLTFConvertBuilder& Builde
 			FGLTFJsonNode CameraNode;
 			CameraNode.Name = Owner->GetName(); // TODO: choose a more unique name if owner is not ACameraActor
 			CameraNode.Rotation = FGLTFConverterUtility::ConvertCameraDirection();
-			CameraNode.Scale = FGLTFConverterUtility::ConvertScale(bSupportNonUniformScale ? Transform.GetScale3D() : FVector::OneVector);
+			CameraNode.Scale = FGLTFConverterUtility::ConvertScale(bExportNonUniformScale ? Transform.GetScale3D() : FVector::OneVector);
 			CameraNode.Camera = Builder.GetOrAddCamera(CameraComponent, CameraNode.Name);
 			Builder.AddChildComponentNode(NodeIndex, CameraNode);
 		}
@@ -111,7 +111,7 @@ FGLTFJsonNodeIndex FGLTFSceneComponentConverter::Add(FGLTFConvertBuilder& Builde
 			FGLTFJsonNode LightNode;
 			LightNode.Name = Owner->GetName(); // TODO: choose a more unique name if owner is not ALight
 			LightNode.Rotation = FGLTFConverterUtility::ConvertLightDirection();
-			LightNode.Scale = FGLTFConverterUtility::ConvertScale(bSupportNonUniformScale ? Transform.GetScale3D() : FVector::OneVector);
+			LightNode.Scale = FGLTFConverterUtility::ConvertScale(bExportNonUniformScale ? Transform.GetScale3D() : FVector::OneVector);
 			LightNode.Light = Builder.GetOrAddLight(LightComponent, LightNode.Name);
 			Builder.AddChildComponentNode(NodeIndex, LightNode);
 		}
diff --git a/Engine/Plugins/Enterprise/GLTFExporter/MigratedFromGit/Source/GLTFExporter/Private/GLTFExportOptions.cpp b/Engine/Plugins/Enterprise/GLTFExporter/MigratedFromGit/Source/GLTFExporter/Private/GLTFExportOptions.cpp
--- a/Engine/Plugins/Enterprise/GLTFExporter/MigratedFromGit/Source/GLTFExporter/Private/GLTFExportOptions.cpp
+++ b/Engine/Plugins/Enterprise/GLTFExporter/MigratedFromGit/Source/GLTFExporter/Private/GLTFExportOptions.cpp
@@ -35,7 +35,7 @@ UGLTFExportOptions::UGLTFExportOptions(const FObjectInitializer& ObjectInitializ
 	bExportHDRIBackdrops = true;
 	bExportVariantSets = true;
 	bExportInteractionHotspots = false;
-	bSupportNonUniformScale = false;
+	bExportNonUniformScale = false;
 	bBundleWebViewer = false;
 	bExportPreviewMesh = true;
 	bAllExtensionsRequired = false;
@@ -223,7 +223,7 @@ bool UGLTFExportOptions::CanEditChange(const FProperty* InProperty) const
 		PropertyFName == GET_MEMBER_NAME_CHECKED(ThisClass, bExportCameras) ||
 		PropertyFName == GET_MEMBER_NAME_CHECKED(ThisClass, bExportHDRIBackdrops) ||
 		PropertyFName == GET_MEMBER_NAME_CHECKED(ThisClass, bExportVariantSets) ||
-		PropertyFName == GET_MEMBER_NAME_CHECKED(ThisClass, bSupportNonUniformScale) ||
+		PropertyFName == GET_MEMBER_NAME_CHECKED(ThisClass, bExportNonUniformScale) ||
 		PropertyFName == GET_MEMBER_NAME_CHECKED(ThisClass, bExportPreviewMesh))
 	{
 		return true;
diff --git a/Engine/Plugins/Enterprise/GLTFExporter/MigratedFromGit/Source/GLTFExporter/Public/GLTFExportOptions.h b/Engine/Plugins/Enterprise/GLTFExporter/MigratedFromGit/Source/GLTFExporter/Public/GLTFExportOptions.h
--- a/Engine/Plugins/Enterprise/GLTFExporter/MigratedFromGit/Source/GLTFExporter/Public/GLTFExportOptions.h
+++ b/Engine/Plugins/Enterprise/GLTFExporter/MigratedFromGit/Source/GLTFExporter/Public/GLTFExportOptions.h
@@ -139,7 +139,7 @@ public:
 	uint32 bExportInteractionHotspots : 1;
 
 	UPROPERTY(EditAnywhere, BlueprintReadWrite, config, Category = Scene)
-	uint32 bSupportNonUniformScale : 1;
+	uint32 bExportNonUniformScale : 1;
 
 	UPROPERTY(EditAnywhere, BlueprintReadWrite, config, Category = Exporter)
     uint32 bBundleWebViewer : 1;
