[GLTFExporter] refactor: unify scale calculation for child-component nodes to reduce code duplication

commit 8cfc008394fce570ddbaa45186fe6acad2ba7215
Date: Thu, 3 Sep 2020 15:51:23 +0200

    refactor: unify scale calculation for child-component nodes to reduce code duplication
    (cherry picked from commit 6af75bff4339b30fc2c06562dcd3242cd1b1bea3​)

#skipci

diff --git a/Engine/Plugins/Enterprise/GLTFExporter/MigratedFromGit/Source/GLTFExporter/Private/Converters/GLTFLevelConverters.cpp b/Engine/Plugins/Enterprise/GLTFExporter/MigratedFromGit/Source/GLTFExporter/Private/Converters/GLTFLevelConverters.cpp
--- a/Engine/Plugins/Enterprise/GLTFExporter/MigratedFromGit/Source/GLTFExporter/Private/Converters/GLTFLevelConverters.cpp
+++ b/Engine/Plugins/Enterprise/GLTFExporter/MigratedFromGit/Source/GLTFExporter/Private/Converters/GLTFLevelConverters.cpp
@@ -38,6 +38,8 @@ FGLTFJsonNodeIndex FGLTFSceneComponentConverter::Add(FGLTFConvertBuilder& Builde
 	Node.Rotation = FGLTFConverterUtility::ConvertRotation(RelativeTransform.GetRotation());
 	Node.Scale = FGLTFConverterUtility::ConvertScale(Scale);
 
+	const FGLTFJsonVector3 ConvertedScale = FGLTFConverterUtility::ConvertScale(bExportNonUniformScale ? Transform.GetScale3D() : FVector::OneVector);
+
 	if (SceneComponent->bHiddenInGame) // TODO: make this configurable
 	{
 		// ignore any visible properties
@@ -54,7 +56,7 @@ FGLTFJsonNodeIndex FGLTFSceneComponentConverter::Add(FGLTFConvertBuilder& Builde
 		{
 			FGLTFJsonNode MeshNode;
 			MeshNode.Name = Node.Name + TEXT("_Mesh");
-			MeshNode.Scale = FGLTFConverterUtility::ConvertScale(Transform.GetScale3D());
+			MeshNode.Scale = ConvertedScale;
 			MeshNode.Mesh = Builder.GetOrAddMesh(StaticMeshComponent);
 			MeshNode.LightMap = Builder.GetOrAddLightMap(StaticMeshComponent);
 			Builder.AddChildComponentNode(NodeIndex, MeshNode);
@@ -77,7 +79,7 @@ FGLTFJsonNodeIndex FGLTFSceneComponentConverter::Add(FGLTFConvertBuilder& Builde
 		{
 			FGLTFJsonNode MeshNode;
 			MeshNode.Name = Node.Name + TEXT("_Mesh");
-			MeshNode.Scale = FGLTFConverterUtility::ConvertScale(Transform.GetScale3D());
+			MeshNode.Scale = ConvertedScale;
 			MeshNode.Mesh = Builder.GetOrAddMesh(SkeletalMeshComponent);
 			Builder.AddChildComponentNode(NodeIndex, MeshNode);
 		}
@@ -94,7 +96,7 @@ FGLTFJsonNodeIndex FGLTFSceneComponentConverter::Add(FGLTFConvertBuilder& Builde
 			FGLTFJsonNode CameraNode;
 			CameraNode.Name = Owner->GetName(); // TODO: choose a more unique name if owner is not ACameraActor
 			CameraNode.Rotation = FGLTFConverterUtility::ConvertCameraDirection();
-			CameraNode.Scale = FGLTFConverterUtility::ConvertScale(bExportNonUniformScale ? Transform.GetScale3D() : FVector::OneVector);
+			CameraNode.Scale = ConvertedScale;
 			CameraNode.Camera = Builder.GetOrAddCamera(CameraComponent, CameraNode.Name);
 			Builder.AddChildComponentNode(NodeIndex, CameraNode);
 		}
@@ -111,7 +113,7 @@ FGLTFJsonNodeIndex FGLTFSceneComponentConverter::Add(FGLTFConvertBuilder& Builde
 			FGLTFJsonNode LightNode;
 			LightNode.Name = Owner->GetName(); // TODO: choose a more unique name if owner is not ALight
 			LightNode.Rotation = FGLTFConverterUtility::ConvertLightDirection();
-			LightNode.Scale = FGLTFConverterUtility::ConvertScale(bExportNonUniformScale ? Transform.GetScale3D() : FVector::OneVector);
+			LightNode.Scale = ConvertedScale;
 			LightNode.Light = Builder.GetOrAddLight(LightComponent, LightNode.Name);
 			Builder.AddChildComponentNode(NodeIndex, LightNode);
 		}
