[GLTFExporter] refactor: change name of ConvertedScale to ComponentNodeScale (less ambiguous)

commit 3deb66445859f0bddff9c189e6e283c26e6f7eb0
Date: Mon, 7 Sep 2020 10:10:03 +0200

    refactor: change name of ConvertedScale to ComponentNodeScale (less ambiguous)
    (cherry picked from commit 72a211b428c94fccc9f7e3a34ed177db05a89be9​)

#skipci

diff --git a/Engine/Plugins/Enterprise/GLTFExporter/MigratedFromGit/Source/GLTFExporter/Private/Converters/GLTFLevelConverters.cpp b/Engine/Plugins/Enterprise/GLTFExporter/MigratedFromGit/Source/GLTFExporter/Private/Converters/GLTFLevelConverters.cpp
--- a/Engine/Plugins/Enterprise/GLTFExporter/MigratedFromGit/Source/GLTFExporter/Private/Converters/GLTFLevelConverters.cpp
+++ b/Engine/Plugins/Enterprise/GLTFExporter/MigratedFromGit/Source/GLTFExporter/Private/Converters/GLTFLevelConverters.cpp
@@ -38,7 +38,7 @@ FGLTFJsonNodeIndex FGLTFSceneComponentConverter::Add(FGLTFConvertBuilder& Builde
 	Node.Rotation = FGLTFConverterUtility::ConvertRotation(RelativeTransform.GetRotation());
 	Node.Scale = FGLTFConverterUtility::ConvertScale(Scale);
 
-	const FGLTFJsonVector3 ConvertedScale = FGLTFConverterUtility::ConvertScale(bExportNonUniformScale ? Transform.GetScale3D() : FVector::OneVector);
+	const FGLTFJsonVector3 ComponentNodeScale = FGLTFConverterUtility::ConvertScale(bExportNonUniformScale ? Transform.GetScale3D() : FVector::OneVector);
 
 	if (SceneComponent->bHiddenInGame) // TODO: make this configurable
 	{
@@ -56,7 +56,7 @@ FGLTFJsonNodeIndex FGLTFSceneComponentConverter::Add(FGLTFConvertBuilder& Builde
 		{
 			FGLTFJsonNode MeshNode;
 			MeshNode.Name = Node.Name + TEXT("_Mesh");
-			MeshNode.Scale = ConvertedScale;
+			MeshNode.Scale = ComponentNodeScale;
 			MeshNode.Mesh = Builder.GetOrAddMesh(StaticMeshComponent);
 			MeshNode.LightMap = Builder.GetOrAddLightMap(StaticMeshComponent);
 			Builder.AddChildComponentNode(NodeIndex, MeshNode);
@@ -79,7 +79,7 @@ FGLTFJsonNodeIndex FGLTFSceneComponentConverter::Add(FGLTFConvertBuilder& Builde
 		{
 			FGLTFJsonNode MeshNode;
 			MeshNode.Name = Node.Name + TEXT("_Mesh");
-			MeshNode.Scale = ConvertedScale;
+			MeshNode.Scale = ComponentNodeScale;
 			MeshNode.Mesh = Builder.GetOrAddMesh(SkeletalMeshComponent);
 			Builder.AddChildComponentNode(NodeIndex, MeshNode);
 		}
@@ -96,7 +96,7 @@ FGLTFJsonNodeIndex FGLTFSceneComponentConverter::Add(FGLTFConvertBuilder& Builde
 			FGLTFJsonNode CameraNode;
 			CameraNode.Name = Owner->GetName(); // TODO: choose a more unique name if owner is not ACameraActor
 			CameraNode.Rotation = FGLTFConverterUtility::ConvertCameraDirection();
-			CameraNode.Scale = ConvertedScale;
+			CameraNode.Scale = ComponentNodeScale;
 			CameraNode.Camera = Builder.GetOrAddCamera(CameraComponent, CameraNode.Name);
 			Builder.AddChildComponentNode(NodeIndex, CameraNode);
 		}
@@ -113,7 +113,7 @@ FGLTFJsonNodeIndex FGLTFSceneComponentConverter::Add(FGLTFConvertBuilder& Builde
 			FGLTFJsonNode LightNode;
 			LightNode.Name = Owner->GetName(); // TODO: choose a more unique name if owner is not ALight
 			LightNode.Rotation = FGLTFConverterUtility::ConvertLightDirection();
-			LightNode.Scale = ConvertedScale;
+			LightNode.Scale = ComponentNodeScale;
 			LightNode.Light = Builder.GetOrAddLight(LightComponent, LightNode.Name);
 			Builder.AddChildComponentNode(NodeIndex, LightNode);
 		}
