[GLTFExporter] chore: remove TODOs for functionality which has since been implemented

commit 486fed29503a895997c9e5873cfa5f8ac8b554fd
Date: Mon, 7 Sep 2020 10:10:39 +0200

    chore: remove TODOs for functionality which has since been implemented
    (cherry picked from commit db490ba62cb5912a19a357266eae665a4566182b​)

#skipci

diff --git a/Engine/Plugins/Enterprise/GLTFExporter/MigratedFromGit/Source/GLTFExporter/Private/Converters/GLTFLevelConverters.cpp b/Engine/Plugins/Enterprise/GLTFExporter/MigratedFromGit/Source/GLTFExporter/Private/Converters/GLTFLevelConverters.cpp
--- a/Engine/Plugins/Enterprise/GLTFExporter/MigratedFromGit/Source/GLTFExporter/Private/Converters/GLTFLevelConverters.cpp
+++ b/Engine/Plugins/Enterprise/GLTFExporter/MigratedFromGit/Source/GLTFExporter/Private/Converters/GLTFLevelConverters.cpp
@@ -46,12 +46,6 @@ FGLTFJsonNodeIndex FGLTFSceneComponentConverter::Add(FGLTFConvertBuilder& Builde
 	}
 	else if (const UStaticMeshComponent* StaticMeshComponent = Cast<UStaticMeshComponent>(SceneComponent))
 	{
-		// TODO: using a separate child-node for the mesh breaks variants and possibly other
-		// extensions that count on the mesh being attached to the original node.
-		// We need to either find a way to solve the scale-issues without introducing these
-		// extra nodes, or add some form of mapping that allows the extensions to find the
-		// exact node that the mesh is attached to.
-
 		if (bExportNonUniformScale)
 		{
 			FGLTFJsonNode MeshNode;
@@ -69,12 +63,6 @@ FGLTFJsonNodeIndex FGLTFSceneComponentConverter::Add(FGLTFConvertBuilder& Builde
 	}
 	else if (const USkeletalMeshComponent* SkeletalMeshComponent = Cast<USkeletalMeshComponent>(SceneComponent))
 	{
-		// TODO: using a separate child-node for the mesh breaks variants and possibly other
-		// extensions that count on the mesh being attached to the original node.
-		// We need to either find a way to solve the scale-issues without introducing these
-		// extra nodes, or add some form of mapping that allows the extensions to find the
-		// exact node that the mesh is attached to.
-
 		if (bExportNonUniformScale)
 		{
 			FGLTFJsonNode MeshNode;
