[GLTFExporter] feat: modify exporter to accept custom export-options via an automated export-task

commit e2d380f750b882009d6dfcc751952abbb7ecace4
Date: Thu, 3 Sep 2020 13:04:17 +0200

    feat: modify exporter to accept custom export-options via an automated export-task
    (cherry picked from commit 9fb8a2dc3f42d9e5faa8b2f2b87cb37b4b87534e​)

#skipci

diff --git a/Engine/Plugins/Enterprise/GLTFExporter/MigratedFromGit/Source/GLTFExporter/Private/Exporters/GLTFExporter.cpp b/Engine/Plugins/Enterprise/GLTFExporter/MigratedFromGit/Source/GLTFExporter/Private/Exporters/GLTFExporter.cpp
--- a/Engine/Plugins/Enterprise/GLTFExporter/MigratedFromGit/Source/GLTFExporter/Private/Exporters/GLTFExporter.cpp
+++ b/Engine/Plugins/Enterprise/GLTFExporter/MigratedFromGit/Source/GLTFExporter/Private/Exporters/GLTFExporter.cpp
@@ -2,6 +2,7 @@
 
 #include "Exporters/GLTFExporter.h"
 #include "UObject/GCObjectScopeGuard.h"
+#include "AssetExportTask.h"
 #include "Builders/GLTFContainerBuilder.h"
 #include "GLTFExportOptions.h"
 
@@ -22,8 +23,12 @@ UGLTFExporter::UGLTFExporter(const FObjectInitializer& ObjectInitializer)
 
 bool UGLTFExporter::ExportBinary(UObject* Object, const TCHAR* Type, FArchive& Archive, FFeedbackContext* Warn, int32 FileIndex, uint32 PortFlags)
 {
-	UGLTFExportOptions* Options = NewObject<UGLTFExportOptions>();
-	FGCObjectScopeGuard OptionsGuard(Options);
+	UGLTFExportOptions* Options = GetAutomatedExportOptions();
+	if (Options == nullptr)
+	{
+		Options = NewObject<UGLTFExportOptions>();
+		FGCObjectScopeGuard OptionsGuard(Options);
+	}
 
 	if (!FillExportOptions(Options))
 	{
@@ -86,3 +91,13 @@ bool UGLTFExporter::FillExportOptions(UGLTFExportOptions* ExportOptions)
 	SetShowExportOption(!bExportAll);
 	return true;
 }
+
+UGLTFExportOptions* UGLTFExporter::GetAutomatedExportOptions() const
+{
+	if (ExportTask && ExportTask->bAutomated)
+	{
+		return Cast<UGLTFExportOptions>(ExportTask->Options);
+	}
+
+	return nullptr;
+}
diff --git a/Engine/Plugins/Enterprise/GLTFExporter/MigratedFromGit/Source/GLTFExporter/Public/Exporters/GLTFExporter.h b/Engine/Plugins/Enterprise/GLTFExporter/MigratedFromGit/Source/GLTFExporter/Public/Exporters/GLTFExporter.h
--- a/Engine/Plugins/Enterprise/GLTFExporter/MigratedFromGit/Source/GLTFExporter/Public/Exporters/GLTFExporter.h
+++ b/Engine/Plugins/Enterprise/GLTFExporter/MigratedFromGit/Source/GLTFExporter/Public/Exporters/GLTFExporter.h
@@ -27,4 +27,6 @@ protected:
 private:
 
 	bool FillExportOptions(UGLTFExportOptions* ExportOptions);
+
+	UGLTFExportOptions* GetAutomatedExportOptions() const;
 };
