[GLTFExporter] fix: explicitly compare ExportTask against nullptr

commit 3b2dc219e39651c7a3b426321ae1509519f2d6db
Date: Fri, 4 Sep 2020 08:13:40 +0200

    fix: explicitly compare ExportTask against nullptr
    (cherry picked from commit c68961b03b930eb6794f30fcf685e0c362e5606d​)

#skipci

diff --git a/Engine/Plugins/Enterprise/GLTFExporter/MigratedFromGit/Source/GLTFExporter/Private/Exporters/GLTFExporter.cpp b/Engine/Plugins/Enterprise/GLTFExporter/MigratedFromGit/Source/GLTFExporter/Private/Exporters/GLTFExporter.cpp
--- a/Engine/Plugins/Enterprise/GLTFExporter/MigratedFromGit/Source/GLTFExporter/Private/Exporters/GLTFExporter.cpp
+++ b/Engine/Plugins/Enterprise/GLTFExporter/MigratedFromGit/Source/GLTFExporter/Private/Exporters/GLTFExporter.cpp
@@ -94,7 +94,7 @@ bool UGLTFExporter::FillExportOptions(UGLTFExportOptions* ExportOptions)
 
 UGLTFExportOptions* UGLTFExporter::GetAutomatedExportOptions() const
 {
-	if (ExportTask && ExportTask->bAutomated)
+	if (ExportTask != nullptr && ExportTask->bAutomated)
 	{
 		return Cast<UGLTFExportOptions>(ExportTask->Options);
 	}
