[GLTFExporter] feat: add support for loading and saving export-options and to treat automated export-options correctly

commit 5d68aa1aa161b6036e11c53fbdaf7e450e27786e
Date: Mon, 7 Sep 2020 10:39:07 +0200

    feat: add support for loading and saving export-options and to treat automated export-options correctly
    (cherry picked from commit 05011e7d10672f4ca29bc0cba1d6b7a0307e2168​)

#skipci

diff --git a/Engine/Plugins/Enterprise/GLTFExporter/MigratedFromGit/Source/GLTFExporter/Private/Exporters/GLTFExporter.cpp b/Engine/Plugins/Enterprise/GLTFExporter/MigratedFromGit/Source/GLTFExporter/Private/Exporters/GLTFExporter.cpp
--- a/Engine/Plugins/Enterprise/GLTFExporter/MigratedFromGit/Source/GLTFExporter/Private/Exporters/GLTFExporter.cpp
+++ b/Engine/Plugins/Enterprise/GLTFExporter/MigratedFromGit/Source/GLTFExporter/Private/Exporters/GLTFExporter.cpp
@@ -1,8 +1,9 @@
 // Copyright Epic Games, Inc. All Rights Reserved.
 
 #include "Exporters/GLTFExporter.h"
-#include "Builders/GLTFContainerBuilder.h"
 #include "GLTFExportOptions.h"
+#include "UI/GLTFExportOptionsWindow.h"
+#include "Builders/GLTFContainerBuilder.h"
 #include "UObject/GCObjectScopeGuard.h"
 #include "AssetExportTask.h"
 
@@ -23,20 +24,14 @@ UGLTFExporter::UGLTFExporter(const FObjectInitializer& ObjectInitializer)
 
 bool UGLTFExporter::ExportBinary(UObject* Object, const TCHAR* Type, FArchive& Archive, FFeedbackContext* Warn, int32 FileIndex, uint32 PortFlags)
 {
-	UGLTFExportOptions* Options = GetAutomatedExportOptions();
+	UGLTFExportOptions* Options = GetExportOptions();
 	if (Options == nullptr)
-	{
-		Options = NewObject<UGLTFExportOptions>();
-	}
-
-	FGCObjectScopeGuard OptionsGuard(Options);
-
-	if (!FillExportOptions(Options))
 	{
 		// User cancelled the export
 		return false;
 	}
 
+	FGCObjectScopeGuard OptionsGuard(Options);
 	bool bSuccess = true;
 
 	// TODO: add support for UAssetExportTask::IgnoreObjectList?
@@ -77,28 +72,44 @@ bool UGLTFExporter::AddObject(FGLTFContainerBuilder& Builder, const UObject* Obj
 	return false;
 }
 
-bool UGLTFExporter::FillExportOptions(UGLTFExportOptions* ExportOptions)
+UGLTFExportOptions* UGLTFExporter::GetExportOptions()
 {
-	bool bExportAll = GetBatchMode() && !GetShowExportOption();
-	bool bExportCancel = false;
+	UGLTFExportOptions* Options = nullptr;
+	bool bAutomatedTask = GIsAutomationTesting || FApp::IsUnattended();
 
-	ExportOptions->FillOptions(GetBatchMode(), GetShowExportOption(), CurrentFilename, bExportCancel, bExportAll);
-	if (bExportCancel)
+	if (ExportTask != nullptr)
 	{
-		SetCancelBatch(GetBatchMode());
-		return false;
+		Options = Cast<UGLTFExportOptions>(ExportTask->Options);
+
+		if (ExportTask->bAutomated)
+		{
+			bAutomatedTask = true;
+		}
 	}
 
-	SetShowExportOption(!bExportAll);
-	return true;
-}
+	if (Options == nullptr)
+	{
+		Options = NewObject<UGLTFExportOptions>();
+		Options->LoadOptions();
+	}
 
-UGLTFExportOptions* UGLTFExporter::GetAutomatedExportOptions() const
-{
-	if (ExportTask != nullptr && ExportTask->bAutomated)
+	if (GetShowExportOption() && !bAutomatedTask)
 	{
-		return Cast<UGLTFExportOptions>(ExportTask->Options);
+		bool bExportAll = GetBatchMode();
+		bool bOperationCanceled = false;
+
+		FGCObjectScopeGuard OptionsGuard(Options);
+		SGLTFExportOptionsWindow::ShowDialog(Options, CurrentFilename, GetBatchMode(), bOperationCanceled, bExportAll);
+
+		if (bOperationCanceled)
+		{
+			SetCancelBatch(GetBatchMode());
+			return nullptr;
+		}
+
+		SetShowExportOption(!bExportAll);
+		Options->SaveOptions();
 	}
 
-	return nullptr;
+	return Options;
 }
diff --git a/Engine/Plugins/Enterprise/GLTFExporter/MigratedFromGit/Source/GLTFExporter/Private/GLTFExportOptions.cpp b/Engine/Plugins/Enterprise/GLTFExporter/MigratedFromGit/Source/GLTFExporter/Private/GLTFExportOptions.cpp
--- a/Engine/Plugins/Enterprise/GLTFExporter/MigratedFromGit/Source/GLTFExporter/Private/GLTFExportOptions.cpp
+++ b/Engine/Plugins/Enterprise/GLTFExporter/MigratedFromGit/Source/GLTFExporter/Private/GLTFExportOptions.cpp
@@ -1,7 +1,6 @@
 // Copyright Epic Games, Inc. All Rights Reserved.
 
 #include "GLTFExportOptions.h"
-#include "UI/GLTFExportOptionsWindow.h"
 #include "UObject/ObjectMacros.h"
 #include "UObject/UObjectGlobals.h"
 #include "UObject/Object.h"
@@ -189,24 +188,6 @@ void UGLTFExportOptions::SaveOptions()
 	GConfig->Flush(false);
 }
 
-void UGLTFExportOptions::FillOptions(bool bBatchMode, bool bShowOptionDialog, const FString& FullPath, bool& bOutOperationCanceled, bool& bOutExportAll)
-{
-	bOutOperationCanceled = false;
-
-	// TODO: call LoadOptions();
-
-	//Return if we do not show the export options or we are running automation test or we are unattended
-	if (!bShowOptionDialog || GIsAutomationTesting || FApp::IsUnattended())
-	{
-		return;
-	}
-
-	bOutExportAll = false;
-
-	SGLTFExportOptionsWindow::ShowDialog(this, FullPath, bBatchMode, bOutOperationCanceled, bOutExportAll);
-	// TODO: call SaveOptions();
-}
-
 bool UGLTFExportOptions::CanEditChange(const FProperty* InProperty) const
 {
 	const FName PropertyFName = InProperty->GetFName();
diff --git a/Engine/Plugins/Enterprise/GLTFExporter/MigratedFromGit/Source/GLTFExporter/Public/Exporters/GLTFExporter.h b/Engine/Plugins/Enterprise/GLTFExporter/MigratedFromGit/Source/GLTFExporter/Public/Exporters/GLTFExporter.h
--- a/Engine/Plugins/Enterprise/GLTFExporter/MigratedFromGit/Source/GLTFExporter/Public/Exporters/GLTFExporter.h
+++ b/Engine/Plugins/Enterprise/GLTFExporter/MigratedFromGit/Source/GLTFExporter/Public/Exporters/GLTFExporter.h
@@ -26,7 +26,5 @@ protected:
 
 private:
 
-	bool FillExportOptions(UGLTFExportOptions* ExportOptions);
-
-	UGLTFExportOptions* GetAutomatedExportOptions() const;
+	UGLTFExportOptions* GetExportOptions();
 };
diff --git a/Engine/Plugins/Enterprise/GLTFExporter/MigratedFromGit/Source/GLTFExporter/Public/GLTFExportOptions.h b/Engine/Plugins/Enterprise/GLTFExporter/MigratedFromGit/Source/GLTFExporter/Public/GLTFExportOptions.h
--- a/Engine/Plugins/Enterprise/GLTFExporter/MigratedFromGit/Source/GLTFExporter/Public/GLTFExportOptions.h
+++ b/Engine/Plugins/Enterprise/GLTFExporter/MigratedFromGit/Source/GLTFExporter/Public/GLTFExportOptions.h
@@ -159,16 +159,6 @@ public:
 	/* Load the each property from a local ini which the value was store by the function SaveOptions() */
 	void LoadOptions();
 
-	/**
-	* Load the export option from the last save state and show the dialog if bShowOptionDialog is true.
-	* FullPath is the export file path we display it in the dialog
-	* If user cancel the dialog, the bOutOperationCanceled will be true
-	* bOutExportAll will be true if the user want to use the same option for all other asset he want to export
-	*
-	* The function is saving the dialog state in a user ini file and reload it from there. It is not changing the CDO.
-	*/
-	void FillOptions(bool bBatchMode, bool bShowOptionDialog, const FString& FullPath, bool& bOutOperationCanceled, bool& bOutExportAll);
-
 	virtual bool CanEditChange(const FProperty* InProperty) const override;
 
 	FIntPoint GetDefaultMaterialBakeSize() const;
