[GLTFExporter] feat: add dependency to GLTFExporterRuntime since it's needed for hotspots

commit 3944631eea1c746e7f2634fd43d3b7580218d789
Date: Fri, 11 Sep 2020 14:22:58 +0200

    feat: add dependency to GLTFExporterRuntime since it's needed for hotspots
    (cherry picked from commit 9f6693e4a8cbb7386cc7ff3eb427b50dbc640a1a​)

#skipci

diff --git a/Engine/Plugins/Enterprise/GLTFExporter/MigratedFromGit/Source/GLTFExporter/GLTFExporter.Build.cs b/Engine/Plugins/Enterprise/GLTFExporter/MigratedFromGit/Source/GLTFExporter/GLTFExporter.Build.cs
--- a/Engine/Plugins/Enterprise/GLTFExporter/MigratedFromGit/Source/GLTFExporter/GLTFExporter.Build.cs
+++ b/Engine/Plugins/Enterprise/GLTFExporter/MigratedFromGit/Source/GLTFExporter/GLTFExporter.Build.cs
@@ -28,7 +28,8 @@ namespace UnrealBuildTool.Rules
 					"RenderCore",
 					"RHI",
 					"MaterialBaking",
-					"VariantManagerContent"
+					"VariantManagerContent",
+					"GLTFExporterRuntime"
 				}
 				);
 		}
