[GLTFExporter] feat(interaction-hotspot): add FGLTFJsonInteractionHotspot

commit 70c5ce7f5546ce405114dd63b955a9e58ca8e3af
Date: Fri, 11 Sep 2020 14:24:08 +0200

    feat(interaction-hotspot): add FGLTFJsonInteractionHotspot
    (cherry picked from commit 85ef8262748de3fcd6be3148db49fa7ca4ea1388​)

#skipci

diff --git a/Engine/Plugins/Enterprise/GLTFExporter/MigratedFromGit/Source/GLTFExporter/Private/Json/FGLTFJsonInteractionHotspot.h b/Engine/Plugins/Enterprise/GLTFExporter/MigratedFromGit/Source/GLTFExporter/Private/Json/FGLTFJsonInteractionHotspot.h
new file mode 100644
--- /dev/null
+++ b/Engine/Plugins/Enterprise/GLTFExporter/MigratedFromGit/Source/GLTFExporter/Private/Json/FGLTFJsonInteractionHotspot.h
@@ -0,0 +1,49 @@
+// Copyright Epic Games, Inc. All Rights Reserved.
+
+#pragma once
+
+#include "GLTFJsonIndex.h"
+#include "Json/GLTFJsonExtensions.h"
+#include "Serialization/JsonSerializer.h"
+
+struct FGLTFJsonInteractionHotspot
+{
+	FString Name;
+
+	FGLTFJsonAnimationIndex Animation;
+	FGLTFJsonTextureIndex   Image;
+	FGLTFJsonTextureIndex   HoveredImage;
+	FGLTFJsonTextureIndex   ToggledImage;
+	FGLTFJsonTextureIndex   ToggledHoveredImage;
+
+	template <class CharType = TCHAR, class PrintPolicy = TPrettyJsonPrintPolicy<CharType>>
+	void WriteObject(TJsonWriter<CharType, PrintPolicy>& JsonWriter, FGLTFJsonExtensions& Extensions) const
+	{
+		JsonWriter.WriteObjectStart();
+
+		if (!Name.IsEmpty())
+		{
+			JsonWriter.WriteValue(TEXT("name"), Name);
+		}
+
+		JsonWriter.WriteValue(TEXT("animation"), Animation);
+		JsonWriter.WriteValue(TEXT("image"), Image);
+
+		if (HoveredImage != INDEX_NONE)
+		{
+			JsonWriter.WriteValue(TEXT("hoveredImage"), HoveredImage);
+		}
+
+		if (ToggledImage != INDEX_NONE)
+		{
+			JsonWriter.WriteValue(TEXT("toggledImage"), ToggledImage);
+		}
+
+		if (ToggledHoveredImage != INDEX_NONE)
+		{
+			JsonWriter.WriteValue(TEXT("toggledHoveredImage"), ToggledHoveredImage);
+		}
+
+		JsonWriter.WriteObjectEnd();
+	}
+};
diff --git a/Engine/Plugins/Enterprise/GLTFExporter/MigratedFromGit/Source/GLTFExporter/Private/Json/GLTFJsonIndex.h b/Engine/Plugins/Enterprise/GLTFExporter/MigratedFromGit/Source/GLTFExporter/Private/Json/GLTFJsonIndex.h
--- a/Engine/Plugins/Enterprise/GLTFExporter/MigratedFromGit/Source/GLTFExporter/Private/Json/GLTFJsonIndex.h
+++ b/Engine/Plugins/Enterprise/GLTFExporter/MigratedFromGit/Source/GLTFExporter/Private/Json/GLTFJsonIndex.h
@@ -68,3 +68,5 @@ struct FGLTFJsonBackdropIndex : FGLTFJsonIndex<FGLTFJsonBackdropIndex> { using F
 struct FGLTFJsonLevelVariantSetsIndex : FGLTFJsonIndex<FGLTFJsonLevelVariantSetsIndex> { using FGLTFJsonIndex::FGLTFJsonIndex; };
 struct FGLTFJsonLightMapIndex : FGLTFJsonIndex<FGLTFJsonLightMapIndex> { using FGLTFJsonIndex::FGLTFJsonIndex; };
 struct FGLTFJsonLightIndex : FGLTFJsonIndex<FGLTFJsonLightIndex> { using FGLTFJsonIndex::FGLTFJsonIndex; };
+struct FGLTFJsonInteractionHotspotIndex : FGLTFJsonIndex<FGLTFJsonInteractionHotspotIndex> { using FGLTFJsonIndex::FGLTFJsonIndex; };
+struct FGLTFJsonAnimationIndex : FGLTFJsonIndex<FGLTFJsonAnimationIndex> { using FGLTFJsonIndex::FGLTFJsonIndex; };
