[GLTFExporter] feat(interaction-hotspot): add InteractionHotspots to FGLTFJsonRoot (including serialization)

commit 73e8d41525edd767baaebb80904fac43bc5afedf
Date: Fri, 11 Sep 2020 14:25:39 +0200

    feat(interaction-hotspot): add InteractionHotspots to FGLTFJsonRoot (including serialization)
    (cherry picked from commit 43620d0d87280d317e57e931aff3656e607485b6​)

#skipci

diff --git a/Engine/Plugins/Enterprise/GLTFExporter/MigratedFromGit/Source/GLTFExporter/Private/Json/GLTFJsonRoot.h b/Engine/Plugins/Enterprise/GLTFExporter/MigratedFromGit/Source/GLTFExporter/Private/Json/GLTFJsonRoot.h
--- a/Engine/Plugins/Enterprise/GLTFExporter/MigratedFromGit/Source/GLTFExporter/Private/Json/GLTFJsonRoot.h
+++ b/Engine/Plugins/Enterprise/GLTFExporter/MigratedFromGit/Source/GLTFExporter/Private/Json/GLTFJsonRoot.h
@@ -19,6 +19,7 @@
 #include "Json/GLTFJsonLevelVariantSets.h"
 #include "Json/GLTFJsonLightMap.h"
 #include "Json/GLTFJsonLight.h"
+#include "Json/FGLTFJsonInteractionHotspot.h"
 #include "Policies/CondensedJsonPrintPolicy.h"
 
 
@@ -73,7 +74,8 @@ struct FGLTFJsonRoot
 	TArray<FGLTFJsonBackdrop>   Backdrops;
 	TArray<FGLTFJsonLightMap>   LightMaps;
 	TArray<FGLTFJsonLight>      Lights;
-	TArray<FGLTFJsonLevelVariantSets> LevelVariantSets;
+	TArray<FGLTFJsonLevelVariantSets>   LevelVariantSets;
+	TArray<FGLTFJsonInteractionHotspot> InteractionHotspots;
 
 	template <class CharType = TCHAR, class PrintPolicy = TPrettyJsonPrintPolicy<CharType>>
 	void WriteObject(TJsonWriter<CharType, PrintPolicy>& JsonWriter) const
@@ -102,7 +104,7 @@ struct FGLTFJsonRoot
 		FGLTFJsonUtility::WriteObjectArray(JsonWriter, TEXT("scenes"), Scenes, AllExtensions);
 		FGLTFJsonUtility::WriteObjectArray(JsonWriter, TEXT("textures"), Textures, AllExtensions);
 
-		if (Backdrops.Num() > 0 || LevelVariantSets.Num() > 0 || LightMaps.Num() > 0 || Lights.Num() > 0 /* TODO: add more extension support */)
+		if (Backdrops.Num() > 0 || LevelVariantSets.Num() > 0 || LightMaps.Num() > 0 || Lights.Num() > 0 || InteractionHotspots.Num() > 0 /* TODO: add more extension support */)
 		{
 			JsonWriter.WriteObjectStart(TEXT("extensions"));
 
@@ -146,6 +148,16 @@ struct FGLTFJsonRoot
 				JsonWriter.WriteObjectEnd();
 			}
 
+			if (InteractionHotspots.Num() > 0)
+			{
+				const EGLTFJsonExtension Extension = EGLTFJsonExtension::EPIC_InteractionHotspots;
+				AllExtensions.Used.Add(Extension);
+
+				JsonWriter.WriteObjectStart(FGLTFJsonUtility::ToString(Extension));
+				FGLTFJsonUtility::WriteObjectArray(JsonWriter, TEXT("interactions"), InteractionHotspots, AllExtensions);
+				JsonWriter.WriteObjectEnd();
+			}
+
 			JsonWriter.WriteObjectEnd();
 		}
 
