[GLTFExporter] feat(interaction-hotspot): add InteractionHotspot to FGLTFJsonNode (including serialization)

commit 9e47a88e9783940f1b07527a512e6a4ddf84ff57
Date: Fri, 11 Sep 2020 14:27:29 +0200

    feat(interaction-hotspot): add InteractionHotspot to FGLTFJsonNode (including serialization)
    (cherry picked from commit d5ef7930ef5e3bfe73bf13dd492d01957d5e6ea8​)

#skipci

diff --git a/Engine/Plugins/Enterprise/GLTFExporter/MigratedFromGit/Source/GLTFExporter/Private/Json/GLTFJsonNode.h b/Engine/Plugins/Enterprise/GLTFExporter/MigratedFromGit/Source/GLTFExporter/Private/Json/GLTFJsonNode.h
--- a/Engine/Plugins/Enterprise/GLTFExporter/MigratedFromGit/Source/GLTFExporter/Private/Json/GLTFJsonNode.h
+++ b/Engine/Plugins/Enterprise/GLTFExporter/MigratedFromGit/Source/GLTFExporter/Private/Json/GLTFJsonNode.h
@@ -26,12 +26,13 @@ struct FGLTFJsonNode
 		};
 	};
 
-	FGLTFJsonCameraIndex    Camera;
-	FGLTFJsonSkinIndex      Skin;
-	FGLTFJsonMeshIndex      Mesh;
-	FGLTFJsonBackdropIndex  Backdrop;
-	FGLTFJsonLightMapIndex  LightMap;
-	FGLTFJsonLightIndex     Light;
+	FGLTFJsonCameraIndex             Camera;
+	FGLTFJsonSkinIndex               Skin;
+	FGLTFJsonMeshIndex               Mesh;
+	FGLTFJsonBackdropIndex           Backdrop;
+	FGLTFJsonLightMapIndex           LightMap;
+	FGLTFJsonLightIndex              Light;
+	FGLTFJsonInteractionHotspotIndex InteractionHotspot;
 
 	FGLTFJsonNodeIndex ComponentNode;
 
@@ -99,7 +100,7 @@ struct FGLTFJsonNode
 			JsonWriter.WriteValue(TEXT("mesh"), Mesh);
 		}
 
-		if (Backdrop != INDEX_NONE || LightMap != INDEX_NONE || Light != INDEX_NONE)
+		if (Backdrop != INDEX_NONE || LightMap != INDEX_NONE || Light != INDEX_NONE || InteractionHotspot != INDEX_NONE)
 		{
 			JsonWriter.WriteObjectStart(TEXT("extensions"));
 
@@ -133,6 +134,16 @@ struct FGLTFJsonNode
 				JsonWriter.WriteObjectEnd();
 			}
 
+			if (InteractionHotspot != INDEX_NONE)
+			{
+				const EGLTFJsonExtension Extension = EGLTFJsonExtension::EPIC_InteractionHotspots;
+				Extensions.Used.Add(Extension);
+
+				JsonWriter.WriteObjectStart(FGLTFJsonUtility::ToString(Extension));
+				JsonWriter.WriteValue(TEXT("interaction"), InteractionHotspot);
+				JsonWriter.WriteObjectEnd();
+			}
+
 			JsonWriter.WriteObjectEnd();
 		}
 
