[GLTFExporter] feat(interaction-hotspot): implement Add- and Get-functions for hotspots to FGLTFJsonBuilder

commit 5d0632c91911722fd3bdcf37ec3f409c2073a011
Date: Fri, 11 Sep 2020 14:29:14 +0200

    feat(interaction-hotspot): implement Add- and Get-functions for hotspots to FGLTFJsonBuilder
    (cherry picked from commit 43deedc6cfd11ce61e332ac7c832a5d907298cb2​)

#skipci

diff --git a/Engine/Plugins/Enterprise/GLTFExporter/MigratedFromGit/Source/GLTFExporter/Private/Builders/GLTFJsonBuilder.cpp b/Engine/Plugins/Enterprise/GLTFExporter/MigratedFromGit/Source/GLTFExporter/Private/Builders/GLTFJsonBuilder.cpp
--- a/Engine/Plugins/Enterprise/GLTFExporter/MigratedFromGit/Source/GLTFExporter/Private/Builders/GLTFJsonBuilder.cpp
+++ b/Engine/Plugins/Enterprise/GLTFExporter/MigratedFromGit/Source/GLTFExporter/Private/Builders/GLTFJsonBuilder.cpp
@@ -83,6 +83,11 @@ FGLTFJsonLightIndex FGLTFJsonBuilder::AddLight(const FGLTFJsonLight& JsonLight)
 	return FGLTFJsonLightIndex(JsonRoot.Lights.Add(JsonLight));
 }
 
+FGLTFJsonInteractionHotspotIndex FGLTFJsonBuilder::AddInteractionHotspot(const FGLTFJsonInteractionHotspot& JsonInteractionHotspot)
+{
+	return FGLTFJsonInteractionHotspotIndex(JsonRoot.InteractionHotspots.Add(JsonInteractionHotspot));
+}
+
 FGLTFJsonNodeIndex FGLTFJsonBuilder::AddChildNode(FGLTFJsonNodeIndex ParentIndex, const FGLTFJsonNode& JsonNode)
 {
 	const FGLTFJsonNodeIndex ChildIndex = AddNode(JsonNode);
@@ -182,6 +187,11 @@ FGLTFJsonLight& FGLTFJsonBuilder::GetLight(FGLTFJsonLightIndex LightIndex)
 	return JsonRoot.Lights[LightIndex];
 }
 
+FGLTFJsonInteractionHotspot& FGLTFJsonBuilder::GetInteractionHotspot(FGLTFJsonInteractionHotspotIndex InteractionHotspotIndex)
+{
+	return JsonRoot.InteractionHotspots[InteractionHotspotIndex];
+}
+
 FGLTFJsonNodeIndex FGLTFJsonBuilder::GetComponentNodeIndex(FGLTFJsonNodeIndex NodeIndex)
 {
 	if (NodeIndex == INDEX_NONE)
diff --git a/Engine/Plugins/Enterprise/GLTFExporter/MigratedFromGit/Source/GLTFExporter/Private/Builders/GLTFJsonBuilder.h b/Engine/Plugins/Enterprise/GLTFExporter/MigratedFromGit/Source/GLTFExporter/Private/Builders/GLTFJsonBuilder.h
--- a/Engine/Plugins/Enterprise/GLTFExporter/MigratedFromGit/Source/GLTFExporter/Private/Builders/GLTFJsonBuilder.h
+++ b/Engine/Plugins/Enterprise/GLTFExporter/MigratedFromGit/Source/GLTFExporter/Private/Builders/GLTFJsonBuilder.h
@@ -28,6 +28,7 @@ public:
 	FGLTFJsonLevelVariantSetsIndex AddLevelVariantSets(const FGLTFJsonLevelVariantSets& LevelVariantSets = {});
 	FGLTFJsonLightMapIndex AddLightMap(const FGLTFJsonLightMap& JsonLightMap = {});
 	FGLTFJsonLightIndex AddLight(const FGLTFJsonLight& JsonLight = {});
+	FGLTFJsonInteractionHotspotIndex AddInteractionHotspot(const FGLTFJsonInteractionHotspot& JsonInteractionHotspot = {});
 
 	FGLTFJsonNodeIndex AddChildNode(FGLTFJsonNodeIndex ParentNodeIndex, const FGLTFJsonNode& JsonNode = {});
 	FGLTFJsonNodeIndex AddChildComponentNode(FGLTFJsonNodeIndex ParentNodeIndex, const FGLTFJsonNode& JsonNode = {});
@@ -47,6 +48,7 @@ public:
 	FGLTFJsonLevelVariantSets& GetLevelVariantSets(FGLTFJsonLevelVariantSetsIndex LevelVariantSetsIndex);
 	FGLTFJsonLightMap& GetLightMap(FGLTFJsonLightMapIndex LightMapIndex);
 	FGLTFJsonLight& GetLight(FGLTFJsonLightIndex LightIndex);
+	FGLTFJsonInteractionHotspot& GetInteractionHotspot(FGLTFJsonInteractionHotspotIndex InteractionHotspotIndex);
 
 	FGLTFJsonNodeIndex GetComponentNodeIndex(FGLTFJsonNodeIndex NodeIndex);
 
