[GLTFExporter] feat(interaction-hotspot): add converter for hotspots

commit 6425e0e989249b2495a83be3af1a38803c0d4952
Date: Fri, 11 Sep 2020 14:29:46 +0200

    feat(interaction-hotspot): add converter for hotspots
    (cherry picked from commit 52fc496b4235ca8776c0c2e0e8495976079da415​)

#skipci

diff --git a/Engine/Plugins/Enterprise/GLTFExporter/MigratedFromGit/Source/GLTFExporter/Private/Builders/GLTFConvertBuilder.cpp b/Engine/Plugins/Enterprise/GLTFExporter/MigratedFromGit/Source/GLTFExporter/Private/Builders/GLTFConvertBuilder.cpp
--- a/Engine/Plugins/Enterprise/GLTFExporter/MigratedFromGit/Source/GLTFExporter/Private/Builders/GLTFConvertBuilder.cpp
+++ b/Engine/Plugins/Enterprise/GLTFExporter/MigratedFromGit/Source/GLTFExporter/Private/Builders/GLTFConvertBuilder.cpp
@@ -344,3 +344,13 @@ FGLTFJsonLightMapIndex FGLTFConvertBuilder::GetOrAddLightMap(const UStaticMeshCo
 
 	return LightMapConverter.GetOrAdd(*this, DesiredName, StaticMeshComponent);
 }
+
+FGLTFJsonInteractionHotspotIndex FGLTFConvertBuilder::GetOrAddInteractionHotspot(const UGLTFInteractionHotspotComponent* HotspotComponent, const FString& DesiredName)
+{
+	if (HotspotComponent == nullptr)
+	{
+		return FGLTFJsonInteractionHotspotIndex(INDEX_NONE);
+	}
+
+	return InteractionHotspotComponentConverter.GetOrAdd(*this, DesiredName.IsEmpty() ? HotspotComponent->GetName() : DesiredName, HotspotComponent);
+}
diff --git a/Engine/Plugins/Enterprise/GLTFExporter/MigratedFromGit/Source/GLTFExporter/Private/Builders/GLTFConvertBuilder.h b/Engine/Plugins/Enterprise/GLTFExporter/MigratedFromGit/Source/GLTFExporter/Private/Builders/GLTFConvertBuilder.h
--- a/Engine/Plugins/Enterprise/GLTFExporter/MigratedFromGit/Source/GLTFExporter/Private/Builders/GLTFConvertBuilder.h
+++ b/Engine/Plugins/Enterprise/GLTFExporter/MigratedFromGit/Source/GLTFExporter/Private/Builders/GLTFConvertBuilder.h
@@ -12,6 +12,7 @@
 #include "Converters/GLTFBackdropConverters.h"
 #include "Converters/GLTFVariantSetConverters.h"
 #include "Converters/GLTFLightMapConverters.h"
+#include "Converters/GLTFInteractionHotspotConverters.h"
 
 class FGLTFConvertBuilder : public FGLTFImageBuilder
 {
@@ -57,6 +58,7 @@ public:
 	FGLTFJsonBackdropIndex GetOrAddBackdrop(const AActor* Actor, const FString& DesiredName = TEXT(""));
 	FGLTFJsonLevelVariantSetsIndex GetOrAddLevelVariantSets(const ALevelVariantSetsActor* LevelVariantSetsActor, const FString& DesiredName = TEXT(""));
 	FGLTFJsonLightMapIndex GetOrAddLightMap(const UStaticMeshComponent* StaticMeshComponent, const FString& DesiredName = TEXT(""));
+	FGLTFJsonInteractionHotspotIndex GetOrAddInteractionHotspot(const UGLTFInteractionHotspotComponent* HotspotComponent, const FString& DesiredName = TEXT(""));
 
 private:
 
@@ -93,4 +95,5 @@ private:
 	FGLTFBackdropConverter BackdropConverter;
 	FGLTFLevelVariantSetsConverter LevelVariantSetsConverter;
 	FGLTFLightMapConverter LightMapConverter;
+	FGLTFInteractionHotspotComponentConverter InteractionHotspotComponentConverter;
 };
diff --git a/Engine/Plugins/Enterprise/GLTFExporter/MigratedFromGit/Source/GLTFExporter/Private/Converters/GLTFInteractionHotspotConverters.cpp b/Engine/Plugins/Enterprise/GLTFExporter/MigratedFromGit/Source/GLTFExporter/Private/Converters/GLTFInteractionHotspotConverters.cpp
new file mode 100644
--- /dev/null
+++ b/Engine/Plugins/Enterprise/GLTFExporter/MigratedFromGit/Source/GLTFExporter/Private/Converters/GLTFInteractionHotspotConverters.cpp
@@ -0,0 +1,20 @@
+// Copyright Epic Games, Inc. All Rights Reserved.
+
+#include "GLTFInteractionHotspotConverters.h"
+#include "Json/FGLTFJsonInteractionHotspot.h"
+#include "Builders/GLTFContainerBuilder.h"
+
+FGLTFJsonInteractionHotspotIndex FGLTFInteractionHotspotComponentConverter::Add(FGLTFConvertBuilder& Builder, const FString& Name, const UGLTFInteractionHotspotComponent* HotspotComponent)
+{
+	// TODO: should we warn and / or return INDEX_NONE if the hotspot has no valid Image assigned?
+
+	FGLTFJsonInteractionHotspot JsonHotspot;
+	JsonHotspot.Name = Name;
+	JsonHotspot.Animation = FGLTFJsonAnimationIndex(INDEX_NONE);	// TODO: assign the animation property once we have support for exporting animations
+	JsonHotspot.Image = Builder.GetOrAddTexture(HotspotComponent->Image);
+	JsonHotspot.HoveredImage = Builder.GetOrAddTexture(HotspotComponent->HoveredImage);
+	JsonHotspot.ToggledImage = Builder.GetOrAddTexture(HotspotComponent->ToggledImage);
+	JsonHotspot.ToggledHoveredImage = Builder.GetOrAddTexture(HotspotComponent->ToggledHoveredImage);
+
+	return Builder.AddInteractionHotspot(JsonHotspot);
+}
diff --git a/Engine/Plugins/Enterprise/GLTFExporter/MigratedFromGit/Source/GLTFExporter/Private/Converters/GLTFInteractionHotspotConverters.h b/Engine/Plugins/Enterprise/GLTFExporter/MigratedFromGit/Source/GLTFExporter/Private/Converters/GLTFInteractionHotspotConverters.h
new file mode 100644
--- /dev/null
+++ b/Engine/Plugins/Enterprise/GLTFExporter/MigratedFromGit/Source/GLTFExporter/Private/Converters/GLTFInteractionHotspotConverters.h
@@ -0,0 +1,13 @@
+// Copyright Epic Games, Inc. All Rights Reserved.
+
+#pragma once
+
+#include "Json/GLTFJsonIndex.h"
+#include "Converters/GLTFConverter.h"
+#include "Components/GLTFInteractionHotspotComponent.h"
+#include "Engine.h"
+
+class FGLTFInteractionHotspotComponentConverter final : public TGLTFConverter<FGLTFJsonInteractionHotspotIndex, const UGLTFInteractionHotspotComponent*>
+{
+	FGLTFJsonInteractionHotspotIndex Add(FGLTFConvertBuilder& Builder, const FString& Name, const UGLTFInteractionHotspotComponent* HotspotComponent) override;
+};
