[GLTFExporter] feat(interaction-hotspot): add support for hotspots to FGLTFSceneComponentConverter

commit 12908e688e475fb6b13e434cd1af75f855d3f21e
Date: Fri, 11 Sep 2020 14:30:45 +0200

    feat(interaction-hotspot): add support for hotspots to FGLTFSceneComponentConverter
    (cherry picked from commit a04ebdf1a7ed13aee3834b0d36855c141e089ae5​)

#skipci

diff --git a/Engine/Plugins/Enterprise/GLTFExporter/MigratedFromGit/Source/GLTFExporter/Private/Converters/GLTFLevelConverters.cpp b/Engine/Plugins/Enterprise/GLTFExporter/MigratedFromGit/Source/GLTFExporter/Private/Converters/GLTFLevelConverters.cpp
--- a/Engine/Plugins/Enterprise/GLTFExporter/MigratedFromGit/Source/GLTFExporter/Private/Converters/GLTFLevelConverters.cpp
+++ b/Engine/Plugins/Enterprise/GLTFExporter/MigratedFromGit/Source/GLTFExporter/Private/Converters/GLTFLevelConverters.cpp
@@ -5,6 +5,7 @@
 #include "Converters/GLTFConverterUtility.h"
 #include "Converters/GLTFActorUtility.h"
 #include "Converters/GLTFCameraUtility.h"
+#include "Components/GLTFInteractionHotspotComponent.h"
 
 FGLTFJsonNodeIndex FGLTFSceneComponentConverter::Add(FGLTFConvertBuilder& Builder, const FString& Name, const USceneComponent* SceneComponent)
 {
@@ -76,6 +77,31 @@ FGLTFJsonNodeIndex FGLTFSceneComponentConverter::Add(FGLTFConvertBuilder& Builde
 			Node.Mesh = Builder.GetOrAddMesh(SkeletalMeshComponent);
 		}
 	}
+	else if (const UGLTFInteractionHotspotComponent* HotspotComponent = Cast<UGLTFInteractionHotspotComponent>(SceneComponent))
+	{
+		if (Builder.ExportOptions->bExportInteractionHotspots)
+		{
+			if (bExportNonUniformScale)
+			{
+				FGLTFJsonNode HotspotNode;
+				HotspotNode.Name = Node.Name + TEXT("_Hotspot");
+				HotspotNode.Scale = ComponentNodeScale;
+
+				// TODO: use better / more unique name for the hotspot?
+				HotspotNode.InteractionHotspot = Builder.GetOrAddInteractionHotspot(HotspotComponent, Owner->GetName());
+				Builder.AddChildComponentNode(NodeIndex, HotspotNode);
+			}
+			else
+			{
+				// TODO: use better / more unique name for the hotspot?
+				Node.InteractionHotspot = Builder.GetOrAddInteractionHotspot(HotspotComponent, Owner->GetName());
+			}
+		}
+		else
+		{
+			Builder.AddWarningMessage(FString::Printf(TEXT("Interaction-hotspot %s disabled by export options"), *Owner->GetName()));
+		}
+	}
 	else if (const UCameraComponent* CameraComponent = Cast<UCameraComponent>(SceneComponent))
 	{
 		if (Builder.ExportOptions->bExportCameras)
