[GLTFExporter] feat(interaction-hotspot): make export-option bExportInteractionHotspots editable

commit 47d6e9eab93c42ba4769e4909fc9b8631386a91e
Date: Fri, 11 Sep 2020 14:31:09 +0200

    feat(interaction-hotspot): make export-option bExportInteractionHotspots editable
    (cherry picked from commit 1b5dfc102b06cc05a1cda4b6da874c337f41a15b​)

#skipci

diff --git a/Engine/Plugins/Enterprise/GLTFExporter/MigratedFromGit/Source/GLTFExporter/Private/GLTFExportOptions.cpp b/Engine/Plugins/Enterprise/GLTFExporter/MigratedFromGit/Source/GLTFExporter/Private/GLTFExportOptions.cpp
--- a/Engine/Plugins/Enterprise/GLTFExporter/MigratedFromGit/Source/GLTFExporter/Private/GLTFExportOptions.cpp
+++ b/Engine/Plugins/Enterprise/GLTFExporter/MigratedFromGit/Source/GLTFExporter/Private/GLTFExportOptions.cpp
@@ -204,6 +204,7 @@ bool UGLTFExportOptions::CanEditChange(const FProperty* InProperty) const
 		PropertyFName == GET_MEMBER_NAME_CHECKED(ThisClass, bExportCameras) ||
 		PropertyFName == GET_MEMBER_NAME_CHECKED(ThisClass, bExportHDRIBackdrops) ||
 		PropertyFName == GET_MEMBER_NAME_CHECKED(ThisClass, bExportVariantSets) ||
+		PropertyFName == GET_MEMBER_NAME_CHECKED(ThisClass, bExportInteractionHotspots) ||
 		PropertyFName == GET_MEMBER_NAME_CHECKED(ThisClass, bExportNonUniformScale) ||
 		PropertyFName == GET_MEMBER_NAME_CHECKED(ThisClass, bExportPreviewMesh))
 	{
