[GLTFExporter] feat: make the bExportInteractionHotspot option true by default

commit 837c70db8f246bbf8a6defa1276528b61af58c6d
Date: Tue, 15 Sep 2020 08:07:27 +0200

    feat: make the bExportInteractionHotspot option true by default
    (cherry picked from commit f92c784726d8f0bbd39109ebe1b34aaafdbcbcfc​)

#skipci

diff --git a/Engine/Plugins/Enterprise/GLTFExporter/MigratedFromGit/Source/GLTFExporter/Private/GLTFExportOptions.cpp b/Engine/Plugins/Enterprise/GLTFExporter/MigratedFromGit/Source/GLTFExporter/Private/GLTFExportOptions.cpp
--- a/Engine/Plugins/Enterprise/GLTFExporter/MigratedFromGit/Source/GLTFExporter/Private/GLTFExportOptions.cpp
+++ b/Engine/Plugins/Enterprise/GLTFExporter/MigratedFromGit/Source/GLTFExporter/Private/GLTFExportOptions.cpp
@@ -33,7 +33,7 @@ UGLTFExportOptions::UGLTFExportOptions(const FObjectInitializer& ObjectInitializ
 	bExportReflectionCaptures = false;
 	bExportHDRIBackdrops = true;
 	bExportVariantSets = true;
-	bExportInteractionHotspots = false;
+	bExportInteractionHotspots = true;
 	bExportNonUniformScale = false;
 	bBundleWebViewer = false;
 	bExportPreviewMesh = true;
