[GLTFExporter] refactor(interaction-hotspot): avoid using "interaction" as prefix where possible

commit ec2e7b559842448b1c391ffa067911b42acddac2
Date: Tue, 15 Sep 2020 08:36:33 +0200

    refactor(interaction-hotspot): avoid using "interaction" as prefix where possible
    (cherry picked from commit 2bf038c1f6175abf5906d998a9cd53d08a1fec88​)

#skipci

diff --git a/Engine/Plugins/Enterprise/GLTFExporter/MigratedFromGit/Source/GLTFExporter/Private/Builders/GLTFConvertBuilder.cpp b/Engine/Plugins/Enterprise/GLTFExporter/MigratedFromGit/Source/GLTFExporter/Private/Builders/GLTFConvertBuilder.cpp
--- a/Engine/Plugins/Enterprise/GLTFExporter/MigratedFromGit/Source/GLTFExporter/Private/Builders/GLTFConvertBuilder.cpp
+++ b/Engine/Plugins/Enterprise/GLTFExporter/MigratedFromGit/Source/GLTFExporter/Private/Builders/GLTFConvertBuilder.cpp
@@ -115,7 +115,7 @@ FGLTFJsonMeshIndex FGLTFConvertBuilder::GetOrAddMesh(const UStaticMesh* StaticMe
 	{
 		OverrideVertexColors = nullptr;
 	}
-	
+
 	return StaticMeshConverter.GetOrAdd(
 		*this, DesiredName.IsEmpty() ? FGLTFBuilderUtility::GetLODName(StaticMesh, LODIndex) : DesiredName,
 		StaticMesh,
@@ -345,12 +345,12 @@ FGLTFJsonLightMapIndex FGLTFConvertBuilder::GetOrAddLightMap(const UStaticMeshCo
 	return LightMapConverter.GetOrAdd(*this, DesiredName, StaticMeshComponent);
 }
 
-FGLTFJsonInteractionHotspotIndex FGLTFConvertBuilder::GetOrAddInteractionHotspot(const UGLTFInteractionHotspotComponent* HotspotComponent, const FString& DesiredName)
+FGLTFJsonHotspotIndex FGLTFConvertBuilder::GetOrAddHotspot(const UGLTFInteractionHotspotComponent* HotspotComponent, const FString& DesiredName)
 {
 	if (HotspotComponent == nullptr)
 	{
-		return FGLTFJsonInteractionHotspotIndex(INDEX_NONE);
+		return FGLTFJsonHotspotIndex(INDEX_NONE);
 	}
 
-	return InteractionHotspotComponentConverter.GetOrAdd(*this, DesiredName.IsEmpty() ? HotspotComponent->GetName() : DesiredName, HotspotComponent);
+	return HotspotComponentConverter.GetOrAdd(*this, DesiredName.IsEmpty() ? HotspotComponent->GetName() : DesiredName, HotspotComponent);
 }
diff --git a/Engine/Plugins/Enterprise/GLTFExporter/MigratedFromGit/Source/GLTFExporter/Private/Builders/GLTFConvertBuilder.h b/Engine/Plugins/Enterprise/GLTFExporter/MigratedFromGit/Source/GLTFExporter/Private/Builders/GLTFConvertBuilder.h
--- a/Engine/Plugins/Enterprise/GLTFExporter/MigratedFromGit/Source/GLTFExporter/Private/Builders/GLTFConvertBuilder.h
+++ b/Engine/Plugins/Enterprise/GLTFExporter/MigratedFromGit/Source/GLTFExporter/Private/Builders/GLTFConvertBuilder.h
@@ -12,7 +12,7 @@
 #include "Converters/GLTFBackdropConverters.h"
 #include "Converters/GLTFVariantSetConverters.h"
 #include "Converters/GLTFLightMapConverters.h"
-#include "Converters/GLTFInteractionHotspotConverters.h"
+#include "Converters/GLTFHotspotConverters.h"
 
 class FGLTFConvertBuilder : public FGLTFImageBuilder
 {
@@ -58,7 +58,7 @@ public:
 	FGLTFJsonBackdropIndex GetOrAddBackdrop(const AActor* Actor, const FString& DesiredName = TEXT(""));
 	FGLTFJsonLevelVariantSetsIndex GetOrAddLevelVariantSets(const ALevelVariantSetsActor* LevelVariantSetsActor, const FString& DesiredName = TEXT(""));
 	FGLTFJsonLightMapIndex GetOrAddLightMap(const UStaticMeshComponent* StaticMeshComponent, const FString& DesiredName = TEXT(""));
-	FGLTFJsonInteractionHotspotIndex GetOrAddInteractionHotspot(const UGLTFInteractionHotspotComponent* HotspotComponent, const FString& DesiredName = TEXT(""));
+	FGLTFJsonHotspotIndex GetOrAddHotspot(const UGLTFInteractionHotspotComponent* HotspotComponent, const FString& DesiredName = TEXT(""));
 
 private:
 
@@ -95,5 +95,5 @@ private:
 	FGLTFBackdropConverter BackdropConverter;
 	FGLTFLevelVariantSetsConverter LevelVariantSetsConverter;
 	FGLTFLightMapConverter LightMapConverter;
-	FGLTFInteractionHotspotComponentConverter InteractionHotspotComponentConverter;
+	FGLTFHotspotComponentConverter HotspotComponentConverter;
 };
diff --git a/Engine/Plugins/Enterprise/GLTFExporter/MigratedFromGit/Source/GLTFExporter/Private/Builders/GLTFJsonBuilder.cpp b/Engine/Plugins/Enterprise/GLTFExporter/MigratedFromGit/Source/GLTFExporter/Private/Builders/GLTFJsonBuilder.cpp
--- a/Engine/Plugins/Enterprise/GLTFExporter/MigratedFromGit/Source/GLTFExporter/Private/Builders/GLTFJsonBuilder.cpp
+++ b/Engine/Plugins/Enterprise/GLTFExporter/MigratedFromGit/Source/GLTFExporter/Private/Builders/GLTFJsonBuilder.cpp
@@ -83,9 +83,9 @@ FGLTFJsonLightIndex FGLTFJsonBuilder::AddLight(const FGLTFJsonLight& JsonLight)
 	return FGLTFJsonLightIndex(JsonRoot.Lights.Add(JsonLight));
 }
 
-FGLTFJsonInteractionHotspotIndex FGLTFJsonBuilder::AddInteractionHotspot(const FGLTFJsonInteractionHotspot& JsonInteractionHotspot)
+FGLTFJsonHotspotIndex FGLTFJsonBuilder::AddHotspot(const FGLTFJsonHotspot& JsonHotspot)
 {
-	return FGLTFJsonInteractionHotspotIndex(JsonRoot.InteractionHotspots.Add(JsonInteractionHotspot));
+	return FGLTFJsonHotspotIndex(JsonRoot.Hotspots.Add(JsonHotspot));
 }
 
 FGLTFJsonNodeIndex FGLTFJsonBuilder::AddChildNode(FGLTFJsonNodeIndex ParentIndex, const FGLTFJsonNode& JsonNode)
@@ -187,9 +187,9 @@ FGLTFJsonLight& FGLTFJsonBuilder::GetLight(FGLTFJsonLightIndex LightIndex)
 	return JsonRoot.Lights[LightIndex];
 }
 
-FGLTFJsonInteractionHotspot& FGLTFJsonBuilder::GetInteractionHotspot(FGLTFJsonInteractionHotspotIndex InteractionHotspotIndex)
+FGLTFJsonHotspot& FGLTFJsonBuilder::GetHotspot(FGLTFJsonHotspotIndex HotspotIndex)
 {
-	return JsonRoot.InteractionHotspots[InteractionHotspotIndex];
+	return JsonRoot.Hotspots[HotspotIndex];
 }
 
 FGLTFJsonNodeIndex FGLTFJsonBuilder::GetComponentNodeIndex(FGLTFJsonNodeIndex NodeIndex)
diff --git a/Engine/Plugins/Enterprise/GLTFExporter/MigratedFromGit/Source/GLTFExporter/Private/Builders/GLTFJsonBuilder.h b/Engine/Plugins/Enterprise/GLTFExporter/MigratedFromGit/Source/GLTFExporter/Private/Builders/GLTFJsonBuilder.h
--- a/Engine/Plugins/Enterprise/GLTFExporter/MigratedFromGit/Source/GLTFExporter/Private/Builders/GLTFJsonBuilder.h
+++ b/Engine/Plugins/Enterprise/GLTFExporter/MigratedFromGit/Source/GLTFExporter/Private/Builders/GLTFJsonBuilder.h
@@ -28,7 +28,7 @@ public:
 	FGLTFJsonLevelVariantSetsIndex AddLevelVariantSets(const FGLTFJsonLevelVariantSets& LevelVariantSets = {});
 	FGLTFJsonLightMapIndex AddLightMap(const FGLTFJsonLightMap& JsonLightMap = {});
 	FGLTFJsonLightIndex AddLight(const FGLTFJsonLight& JsonLight = {});
-	FGLTFJsonInteractionHotspotIndex AddInteractionHotspot(const FGLTFJsonInteractionHotspot& JsonInteractionHotspot = {});
+	FGLTFJsonHotspotIndex AddHotspot(const FGLTFJsonHotspot& JsonHotspot = {});
 
 	FGLTFJsonNodeIndex AddChildNode(FGLTFJsonNodeIndex ParentNodeIndex, const FGLTFJsonNode& JsonNode = {});
 	FGLTFJsonNodeIndex AddChildComponentNode(FGLTFJsonNodeIndex ParentNodeIndex, const FGLTFJsonNode& JsonNode = {});
@@ -48,7 +48,7 @@ public:
 	FGLTFJsonLevelVariantSets& GetLevelVariantSets(FGLTFJsonLevelVariantSetsIndex LevelVariantSetsIndex);
 	FGLTFJsonLightMap& GetLightMap(FGLTFJsonLightMapIndex LightMapIndex);
 	FGLTFJsonLight& GetLight(FGLTFJsonLightIndex LightIndex);
-	FGLTFJsonInteractionHotspot& GetInteractionHotspot(FGLTFJsonInteractionHotspotIndex InteractionHotspotIndex);
+	FGLTFJsonHotspot& GetHotspot(FGLTFJsonHotspotIndex HotspotIndex);
 
 	FGLTFJsonNodeIndex GetComponentNodeIndex(FGLTFJsonNodeIndex NodeIndex);
 
diff --git a/Engine/Plugins/Enterprise/GLTFExporter/MigratedFromGit/Source/GLTFExporter/Private/Converters/GLTFHotspotConverters.cpp b/Engine/Plugins/Enterprise/GLTFExporter/MigratedFromGit/Source/GLTFExporter/Private/Converters/GLTFHotspotConverters.cpp
new file mode 100644
--- /dev/null
+++ b/Engine/Plugins/Enterprise/GLTFExporter/MigratedFromGit/Source/GLTFExporter/Private/Converters/GLTFHotspotConverters.cpp
@@ -0,0 +1,20 @@
+// Copyright Epic Games, Inc. All Rights Reserved.
+
+#include "GLTFHotspotConverters.h"
+#include "Json/FGLTFJsonHotspot.h"
+#include "Builders/GLTFContainerBuilder.h"
+
+FGLTFJsonHotspotIndex FGLTFHotspotComponentConverter::Add(FGLTFConvertBuilder& Builder, const FString& Name, const UGLTFInteractionHotspotComponent* HotspotComponent)
+{
+	// TODO: should we warn and / or return INDEX_NONE if the hotspot has no valid Image assigned?
+
+	FGLTFJsonHotspot JsonHotspot;
+	JsonHotspot.Name = Name;
+	JsonHotspot.Animation = FGLTFJsonAnimationIndex(INDEX_NONE);	// TODO: assign the animation property once we have support for exporting animations
+	JsonHotspot.Image = Builder.GetOrAddTexture(HotspotComponent->Image);
+	JsonHotspot.HoveredImage = Builder.GetOrAddTexture(HotspotComponent->HoveredImage);
+	JsonHotspot.ToggledImage = Builder.GetOrAddTexture(HotspotComponent->ToggledImage);
+	JsonHotspot.ToggledHoveredImage = Builder.GetOrAddTexture(HotspotComponent->ToggledHoveredImage);
+
+	return Builder.AddHotspot(JsonHotspot);
+}
diff --git a/Engine/Plugins/Enterprise/GLTFExporter/MigratedFromGit/Source/GLTFExporter/Private/Converters/GLTFHotspotConverters.h b/Engine/Plugins/Enterprise/GLTFExporter/MigratedFromGit/Source/GLTFExporter/Private/Converters/GLTFHotspotConverters.h
new file mode 100644
--- /dev/null
+++ b/Engine/Plugins/Enterprise/GLTFExporter/MigratedFromGit/Source/GLTFExporter/Private/Converters/GLTFHotspotConverters.h
@@ -0,0 +1,13 @@
+// Copyright Epic Games, Inc. All Rights Reserved.
+
+#pragma once
+
+#include "Json/GLTFJsonIndex.h"
+#include "Converters/GLTFConverter.h"
+#include "Components/GLTFInteractionHotspotComponent.h"
+#include "Engine.h"
+
+class FGLTFHotspotComponentConverter final : public TGLTFConverter<FGLTFJsonHotspotIndex, const UGLTFInteractionHotspotComponent*>
+{
+	FGLTFJsonHotspotIndex Add(FGLTFConvertBuilder& Builder, const FString& Name, const UGLTFInteractionHotspotComponent* HotspotComponent) override;
+};
diff --git a/Engine/Plugins/Enterprise/GLTFExporter/MigratedFromGit/Source/GLTFExporter/Private/Converters/GLTFInteractionHotspotConverters.cpp b/Engine/Plugins/Enterprise/GLTFExporter/MigratedFromGit/Source/GLTFExporter/Private/Converters/GLTFInteractionHotspotConverters.cpp
deleted file mode 100644
diff --git a/Engine/Plugins/Enterprise/GLTFExporter/MigratedFromGit/Source/GLTFExporter/Private/Converters/GLTFInteractionHotspotConverters.h b/Engine/Plugins/Enterprise/GLTFExporter/MigratedFromGit/Source/GLTFExporter/Private/Converters/GLTFInteractionHotspotConverters.h
deleted file mode 100644
diff --git a/Engine/Plugins/Enterprise/GLTFExporter/MigratedFromGit/Source/GLTFExporter/Private/Converters/GLTFLevelConverters.cpp b/Engine/Plugins/Enterprise/GLTFExporter/MigratedFromGit/Source/GLTFExporter/Private/Converters/GLTFLevelConverters.cpp
--- a/Engine/Plugins/Enterprise/GLTFExporter/MigratedFromGit/Source/GLTFExporter/Private/Converters/GLTFLevelConverters.cpp
+++ b/Engine/Plugins/Enterprise/GLTFExporter/MigratedFromGit/Source/GLTFExporter/Private/Converters/GLTFLevelConverters.cpp
@@ -88,13 +88,13 @@ FGLTFJsonNodeIndex FGLTFSceneComponentConverter::Add(FGLTFConvertBuilder& Builde
 				HotspotNode.Scale = ComponentNodeScale;
 
 				// TODO: use better / more unique name for the hotspot?
-				HotspotNode.InteractionHotspot = Builder.GetOrAddInteractionHotspot(HotspotComponent, Owner->GetName());
+				HotspotNode.Hotspot = Builder.GetOrAddHotspot(HotspotComponent, Owner->GetName());
 				Builder.AddChildComponentNode(NodeIndex, HotspotNode);
 			}
 			else
 			{
 				// TODO: use better / more unique name for the hotspot?
-				Node.InteractionHotspot = Builder.GetOrAddInteractionHotspot(HotspotComponent, Owner->GetName());
+				Node.Hotspot = Builder.GetOrAddHotspot(HotspotComponent, Owner->GetName());
 			}
 		}
 		else
diff --git a/Engine/Plugins/Enterprise/GLTFExporter/MigratedFromGit/Source/GLTFExporter/Private/Json/FGLTFJsonHotspot.h b/Engine/Plugins/Enterprise/GLTFExporter/MigratedFromGit/Source/GLTFExporter/Private/Json/FGLTFJsonHotspot.h
new file mode 100644
--- /dev/null
+++ b/Engine/Plugins/Enterprise/GLTFExporter/MigratedFromGit/Source/GLTFExporter/Private/Json/FGLTFJsonHotspot.h
@@ -0,0 +1,49 @@
+// Copyright Epic Games, Inc. All Rights Reserved.
+
+#pragma once
+
+#include "GLTFJsonIndex.h"
+#include "Json/GLTFJsonExtensions.h"
+#include "Serialization/JsonSerializer.h"
+
+struct FGLTFJsonHotspot
+{
+	FString Name;
+
+	FGLTFJsonAnimationIndex Animation;
+	FGLTFJsonTextureIndex   Image;
+	FGLTFJsonTextureIndex   HoveredImage;
+	FGLTFJsonTextureIndex   ToggledImage;
+	FGLTFJsonTextureIndex   ToggledHoveredImage;
+
+	template <class CharType = TCHAR, class PrintPolicy = TPrettyJsonPrintPolicy<CharType>>
+	void WriteObject(TJsonWriter<CharType, PrintPolicy>& JsonWriter, FGLTFJsonExtensions& Extensions) const
+	{
+		JsonWriter.WriteObjectStart();
+
+		if (!Name.IsEmpty())
+		{
+			JsonWriter.WriteValue(TEXT("name"), Name);
+		}
+
+		JsonWriter.WriteValue(TEXT("animation"), Animation);
+		JsonWriter.WriteValue(TEXT("image"), Image);
+
+		if (HoveredImage != INDEX_NONE)
+		{
+			JsonWriter.WriteValue(TEXT("hoveredImage"), HoveredImage);
+		}
+
+		if (ToggledImage != INDEX_NONE)
+		{
+			JsonWriter.WriteValue(TEXT("toggledImage"), ToggledImage);
+		}
+
+		if (ToggledHoveredImage != INDEX_NONE)
+		{
+			JsonWriter.WriteValue(TEXT("toggledHoveredImage"), ToggledHoveredImage);
+		}
+
+		JsonWriter.WriteObjectEnd();
+	}
+};
diff --git a/Engine/Plugins/Enterprise/GLTFExporter/MigratedFromGit/Source/GLTFExporter/Private/Json/FGLTFJsonInteractionHotspot.h b/Engine/Plugins/Enterprise/GLTFExporter/MigratedFromGit/Source/GLTFExporter/Private/Json/FGLTFJsonInteractionHotspot.h
deleted file mode 100644
diff --git a/Engine/Plugins/Enterprise/GLTFExporter/MigratedFromGit/Source/GLTFExporter/Private/Json/GLTFJsonIndex.h b/Engine/Plugins/Enterprise/GLTFExporter/MigratedFromGit/Source/GLTFExporter/Private/Json/GLTFJsonIndex.h
--- a/Engine/Plugins/Enterprise/GLTFExporter/MigratedFromGit/Source/GLTFExporter/Private/Json/GLTFJsonIndex.h
+++ b/Engine/Plugins/Enterprise/GLTFExporter/MigratedFromGit/Source/GLTFExporter/Private/Json/GLTFJsonIndex.h
@@ -68,5 +68,5 @@ struct FGLTFJsonBackdropIndex : FGLTFJsonIndex<FGLTFJsonBackdropIndex> { using F
 struct FGLTFJsonLevelVariantSetsIndex : FGLTFJsonIndex<FGLTFJsonLevelVariantSetsIndex> { using FGLTFJsonIndex::FGLTFJsonIndex; };
 struct FGLTFJsonLightMapIndex : FGLTFJsonIndex<FGLTFJsonLightMapIndex> { using FGLTFJsonIndex::FGLTFJsonIndex; };
 struct FGLTFJsonLightIndex : FGLTFJsonIndex<FGLTFJsonLightIndex> { using FGLTFJsonIndex::FGLTFJsonIndex; };
-struct FGLTFJsonInteractionHotspotIndex : FGLTFJsonIndex<FGLTFJsonInteractionHotspotIndex> { using FGLTFJsonIndex::FGLTFJsonIndex; };
+struct FGLTFJsonHotspotIndex : FGLTFJsonIndex<FGLTFJsonHotspotIndex> { using FGLTFJsonIndex::FGLTFJsonIndex; };
 struct FGLTFJsonAnimationIndex : FGLTFJsonIndex<FGLTFJsonAnimationIndex> { using FGLTFJsonIndex::FGLTFJsonIndex; };
diff --git a/Engine/Plugins/Enterprise/GLTFExporter/MigratedFromGit/Source/GLTFExporter/Private/Json/GLTFJsonNode.h b/Engine/Plugins/Enterprise/GLTFExporter/MigratedFromGit/Source/GLTFExporter/Private/Json/GLTFJsonNode.h
--- a/Engine/Plugins/Enterprise/GLTFExporter/MigratedFromGit/Source/GLTFExporter/Private/Json/GLTFJsonNode.h
+++ b/Engine/Plugins/Enterprise/GLTFExporter/MigratedFromGit/Source/GLTFExporter/Private/Json/GLTFJsonNode.h
@@ -26,13 +26,13 @@ struct FGLTFJsonNode
 		};
 	};
 
-	FGLTFJsonCameraIndex             Camera;
-	FGLTFJsonSkinIndex               Skin;
-	FGLTFJsonMeshIndex               Mesh;
-	FGLTFJsonBackdropIndex           Backdrop;
-	FGLTFJsonLightMapIndex           LightMap;
-	FGLTFJsonLightIndex              Light;
-	FGLTFJsonInteractionHotspotIndex InteractionHotspot;
+	FGLTFJsonCameraIndex   Camera;
+	FGLTFJsonSkinIndex     Skin;
+	FGLTFJsonMeshIndex     Mesh;
+	FGLTFJsonBackdropIndex Backdrop;
+	FGLTFJsonLightMapIndex LightMap;
+	FGLTFJsonLightIndex    Light;
+	FGLTFJsonHotspotIndex  Hotspot;
 
 	FGLTFJsonNodeIndex ComponentNode;
 
@@ -100,7 +100,7 @@ struct FGLTFJsonNode
 			JsonWriter.WriteValue(TEXT("mesh"), Mesh);
 		}
 
-		if (Backdrop != INDEX_NONE || LightMap != INDEX_NONE || Light != INDEX_NONE || InteractionHotspot != INDEX_NONE)
+		if (Backdrop != INDEX_NONE || LightMap != INDEX_NONE || Light != INDEX_NONE || Hotspot != INDEX_NONE)
 		{
 			JsonWriter.WriteObjectStart(TEXT("extensions"));
 
@@ -134,13 +134,13 @@ struct FGLTFJsonNode
 				JsonWriter.WriteObjectEnd();
 			}
 
-			if (InteractionHotspot != INDEX_NONE)
+			if (Hotspot != INDEX_NONE)
 			{
 				const EGLTFJsonExtension Extension = EGLTFJsonExtension::EPIC_InteractionHotspots;
 				Extensions.Used.Add(Extension);
 
 				JsonWriter.WriteObjectStart(FGLTFJsonUtility::ToString(Extension));
-				JsonWriter.WriteValue(TEXT("interaction"), InteractionHotspot);
+				JsonWriter.WriteValue(TEXT("interaction"), Hotspot);
 				JsonWriter.WriteObjectEnd();
 			}
 
diff --git a/Engine/Plugins/Enterprise/GLTFExporter/MigratedFromGit/Source/GLTFExporter/Private/Json/GLTFJsonRoot.h b/Engine/Plugins/Enterprise/GLTFExporter/MigratedFromGit/Source/GLTFExporter/Private/Json/GLTFJsonRoot.h
--- a/Engine/Plugins/Enterprise/GLTFExporter/MigratedFromGit/Source/GLTFExporter/Private/Json/GLTFJsonRoot.h
+++ b/Engine/Plugins/Enterprise/GLTFExporter/MigratedFromGit/Source/GLTFExporter/Private/Json/GLTFJsonRoot.h
@@ -19,7 +19,7 @@
 #include "Json/GLTFJsonLevelVariantSets.h"
 #include "Json/GLTFJsonLightMap.h"
 #include "Json/GLTFJsonLight.h"
-#include "Json/FGLTFJsonInteractionHotspot.h"
+#include "Json/FGLTFJsonHotspot.h"
 #include "Policies/CondensedJsonPrintPolicy.h"
 
 
@@ -74,8 +74,8 @@ struct FGLTFJsonRoot
 	TArray<FGLTFJsonBackdrop>   Backdrops;
 	TArray<FGLTFJsonLightMap>   LightMaps;
 	TArray<FGLTFJsonLight>      Lights;
-	TArray<FGLTFJsonLevelVariantSets>   LevelVariantSets;
-	TArray<FGLTFJsonInteractionHotspot> InteractionHotspots;
+	TArray<FGLTFJsonLevelVariantSets> LevelVariantSets;
+	TArray<FGLTFJsonHotspot>          Hotspots;
 
 	template <class CharType = TCHAR, class PrintPolicy = TPrettyJsonPrintPolicy<CharType>>
 	void WriteObject(TJsonWriter<CharType, PrintPolicy>& JsonWriter) const
@@ -104,7 +104,7 @@ struct FGLTFJsonRoot
 		FGLTFJsonUtility::WriteObjectArray(JsonWriter, TEXT("scenes"), Scenes, AllExtensions);
 		FGLTFJsonUtility::WriteObjectArray(JsonWriter, TEXT("textures"), Textures, AllExtensions);
 
-		if (Backdrops.Num() > 0 || LevelVariantSets.Num() > 0 || LightMaps.Num() > 0 || Lights.Num() > 0 || InteractionHotspots.Num() > 0 /* TODO: add more extension support */)
+		if (Backdrops.Num() > 0 || LevelVariantSets.Num() > 0 || LightMaps.Num() > 0 || Lights.Num() > 0 || Hotspots.Num() > 0 /* TODO: add more extension support */)
 		{
 			JsonWriter.WriteObjectStart(TEXT("extensions"));
 
@@ -148,13 +148,13 @@ struct FGLTFJsonRoot
 				JsonWriter.WriteObjectEnd();
 			}
 
-			if (InteractionHotspots.Num() > 0)
+			if (Hotspots.Num() > 0)
 			{
 				const EGLTFJsonExtension Extension = EGLTFJsonExtension::EPIC_InteractionHotspots;
 				AllExtensions.Used.Add(Extension);
 
 				JsonWriter.WriteObjectStart(FGLTFJsonUtility::ToString(Extension));
-				FGLTFJsonUtility::WriteObjectArray(JsonWriter, TEXT("interactions"), InteractionHotspots, AllExtensions);
+				FGLTFJsonUtility::WriteObjectArray(JsonWriter, TEXT("interactions"), Hotspots, AllExtensions);
 				JsonWriter.WriteObjectEnd();
 			}
 
