[GLTFExporter] feat: add export-option bExportSourceTextures

commit fcc59981cfc6f1eccf6f548adb8a4a9daa339c52
Date: Thu, 10 Sep 2020 15:41:33 +0200

    feat: add export-option bExportSourceTextures
    (cherry picked from commit f8b8898d7847bb2448b142be7bfd969e3d53975e​)

#skipci

diff --git a/Engine/Plugins/Enterprise/GLTFExporter/MigratedFromGit/Source/GLTFExporter/Private/Converters/GLTFTextureConverters.cpp b/Engine/Plugins/Enterprise/GLTFExporter/MigratedFromGit/Source/GLTFExporter/Private/Converters/GLTFTextureConverters.cpp
--- a/Engine/Plugins/Enterprise/GLTFExporter/MigratedFromGit/Source/GLTFExporter/Private/Converters/GLTFTextureConverters.cpp
+++ b/Engine/Plugins/Enterprise/GLTFExporter/MigratedFromGit/Source/GLTFExporter/Private/Converters/GLTFTextureConverters.cpp
@@ -52,6 +52,8 @@ FGLTFJsonTextureIndex FGLTFTexture2DConverter::Add(FGLTFConvertBuilder& Builder,
 	ERGBFormat RGBFormat;
 	uint32 BitDepth;
 
+	// TODO: implement support for bExportSourceTextures
+
 	// TODO: make sure there are no other special cases (like ex normal-maps and light-maps) that need specialized handling.
 	if (Texture2D->IsA<ULightMapTexture2D>() || Texture2D->IsNormalMap())
 	{
diff --git a/Engine/Plugins/Enterprise/GLTFExporter/MigratedFromGit/Source/GLTFExporter/Private/GLTFExportOptions.cpp b/Engine/Plugins/Enterprise/GLTFExporter/MigratedFromGit/Source/GLTFExporter/Private/GLTFExportOptions.cpp
--- a/Engine/Plugins/Enterprise/GLTFExporter/MigratedFromGit/Source/GLTFExporter/Private/GLTFExportOptions.cpp
+++ b/Engine/Plugins/Enterprise/GLTFExporter/MigratedFromGit/Source/GLTFExporter/Private/GLTFExportOptions.cpp
@@ -23,6 +23,7 @@ UGLTFExportOptions::UGLTFExportOptions(const FObjectInitializer& ObjectInitializ
 	NormalizeUVCoordinates = EGLTFExporterNormalizeUVCoordinates::Never;
 	DefaultLevelOfDetail = 0;
 	TextureFormat = EGLTFExporterTextureFormat::PNG;
+	bExportSourceTextures = true;
 	bExportLightmaps = true;
 	TextureHDREncoding = EGLTFExporterTextureHDREncoding::RGBM;
 	ExportScale = 0.01;
@@ -198,6 +199,7 @@ bool UGLTFExportOptions::CanEditChange(const FProperty* InProperty) const
 		PropertyFName == GET_MEMBER_NAME_CHECKED(ThisClass, DefaultMaterialBakeSize) ||
 		PropertyFName == GET_MEMBER_NAME_CHECKED(ThisClass, bExportVertexColors) ||
 		PropertyFName == GET_MEMBER_NAME_CHECKED(ThisClass, DefaultLevelOfDetail) ||
+		PropertyFName == GET_MEMBER_NAME_CHECKED(ThisClass, bExportSourceTextures) ||
 		PropertyFName == GET_MEMBER_NAME_CHECKED(ThisClass, bExportLightmaps) ||
 		PropertyFName == GET_MEMBER_NAME_CHECKED(ThisClass, ExportScale) ||
 		PropertyFName == GET_MEMBER_NAME_CHECKED(ThisClass, ExportLights) ||
diff --git a/Engine/Plugins/Enterprise/GLTFExporter/MigratedFromGit/Source/GLTFExporter/Public/GLTFExportOptions.h b/Engine/Plugins/Enterprise/GLTFExporter/MigratedFromGit/Source/GLTFExporter/Public/GLTFExportOptions.h
--- a/Engine/Plugins/Enterprise/GLTFExporter/MigratedFromGit/Source/GLTFExporter/Public/GLTFExportOptions.h
+++ b/Engine/Plugins/Enterprise/GLTFExporter/MigratedFromGit/Source/GLTFExporter/Public/GLTFExportOptions.h
@@ -103,6 +103,9 @@ public:
 	UPROPERTY(EditAnywhere, BlueprintReadWrite, config, Category = Texture)
 	EGLTFExporterTextureFormat TextureFormat;
 
+	UPROPERTY(EditAnywhere, BlueprintReadWrite, config, Category = Texture)
+	uint32 bExportSourceTextures : 1;
+
 	UPROPERTY(EditAnywhere, BlueprintReadWrite, config, Category = Texture)
 	uint32 bExportLightmaps : 1;
 
