[GLTFExporter] refactor: remove code-path for exporting textures via platform-data

commit 0c8c8467f767e9c874723c61a93793cac832f9cc
Date: Thu, 10 Sep 2020 15:44:36 +0200

    refactor: remove code-path for exporting textures via platform-data
    (cherry picked from commit a87a40bbea2eab1725bb27f4bcf9b4be2dd8aadc​)

#skipci

diff --git a/Engine/Plugins/Enterprise/GLTFExporter/MigratedFromGit/Source/GLTFExporter/Private/Converters/GLTFTextureConverters.cpp b/Engine/Plugins/Enterprise/GLTFExporter/MigratedFromGit/Source/GLTFExporter/Private/Converters/GLTFTextureConverters.cpp
--- a/Engine/Plugins/Enterprise/GLTFExporter/MigratedFromGit/Source/GLTFExporter/Private/Converters/GLTFTextureConverters.cpp
+++ b/Engine/Plugins/Enterprise/GLTFExporter/MigratedFromGit/Source/GLTFExporter/Private/Converters/GLTFTextureConverters.cpp
@@ -76,20 +76,6 @@ FGLTFJsonTextureIndex FGLTFTexture2DConverter::Add(FGLTFConvertBuilder& Builder,
 			return FGLTFJsonTextureIndex(INDEX_NONE);
 		}
 	}
-	else if (FGLTFTextureUtility::CanPNGCompressFormat(Texture2D->GetPixelFormat(), RGBFormat, BitDepth))
-	{
-		if (!FGLTFTextureUtility::LoadPlatformData(const_cast<UTexture2D*>(Texture2D)))
-		{
-			// TODO: report error
-			return FGLTFJsonTextureIndex(INDEX_NONE);
-		}
-
-		const FByteBulkData& BulkData = Texture2D->PlatformData->Mips[0].BulkData;
-		const void* RawData = BulkData.LockReadOnly();
-
-		ImageIndex = Builder.AddImage(RawData, BulkData.GetBulkDataSize(), Size, RGBFormat, BitDepth, JsonTexture.Name);
-		BulkData.Unlock();
-	}
 	else
 	{
 		const EPixelFormat RenderTargetFormat = FGLTFTextureUtility::IsHDRFormat(Texture2D->GetPixelFormat()) ? PF_FloatRGBA : PF_B8G8R8A8;
