[GLTFExporter] feat: implement export-option bExportSourceTextures, but use render-data on failure

commit 8053ea7cdd50ce7c6e0241339674c057e7ec9117
Date: Fri, 11 Sep 2020 09:05:33 +0200

    feat: implement export-option bExportSourceTextures, but use render-data on failure
    (cherry picked from commit fcce978ff2900945113af2ef609be381147df1ca​)

#skipci

diff --git a/Engine/Plugins/Enterprise/GLTFExporter/MigratedFromGit/Source/GLTFExporter/Private/Converters/GLTFTextureConverters.cpp b/Engine/Plugins/Enterprise/GLTFExporter/MigratedFromGit/Source/GLTFExporter/Private/Converters/GLTFTextureConverters.cpp
--- a/Engine/Plugins/Enterprise/GLTFExporter/MigratedFromGit/Source/GLTFExporter/Private/Converters/GLTFTextureConverters.cpp
+++ b/Engine/Plugins/Enterprise/GLTFExporter/MigratedFromGit/Source/GLTFExporter/Private/Converters/GLTFTextureConverters.cpp
@@ -49,37 +49,42 @@ FGLTFJsonTextureIndex FGLTFTexture2DConverter::Add(FGLTFConvertBuilder& Builder,
 	JsonTexture.Name = Name;
 
 	const FIntPoint Size = { Texture2D->GetSizeX(), Texture2D->GetSizeY() };
-	ERGBFormat RGBFormat;
-	uint32 BitDepth;
+	const bool bPreferSourceExport = Builder.ExportOptions->bExportSourceTextures;
 
-	// TODO: implement support for bExportSourceTextures
+	// NOTE: export of light- and normal-maps via source is done to overcome issues with accessing
+	// pixel-data (for light-maps) and compressed pixel-data (for normal-maps).
 
-	// TODO: make sure there are no other special cases (like ex normal-maps and light-maps) that need specialized handling.
-	if (Texture2D->IsA<ULightMapTexture2D>() || Texture2D->IsNormalMap())
-	{
-		// NOTE: export of light- and normal-maps via source is done to overcome issues with accessing
-		// pixel-data (for light-maps) and compressed pixel-data (for normal-maps).
+	// TODO: make sure there are no other special cases than normal-maps and light-maps that require source-export
+	const bool bRequireSourceExport = Texture2D->IsA<ULightMapTexture2D>() || Texture2D->IsNormalMap();
 
-		// TODO: by exporting via source, we lose the ability to export texture-adjustments. We may want to support this later.
+	if (bRequireSourceExport || bPreferSourceExport)
+	{
+		ERGBFormat RGBFormat;
+		uint32 BitDepth;
 
-		if (FGLTFTextureUtility::CanPNGCompressFormat(Texture2D->Source.GetFormat(), RGBFormat, BitDepth))
+		if (Texture2D->Source.IsValid() && FGLTFTextureUtility::CanPNGCompressFormat(Texture2D->Source.GetFormat(), RGBFormat, BitDepth))
 		{
 			FTextureSource& Source = const_cast<FTextureSource&>(Texture2D->Source);
 
 			const void* RawData = Source.LockMip(0);
 			ImageIndex = Builder.AddImage(RawData, Source.CalcMipSize(0), Size, RGBFormat, BitDepth, JsonTexture.Name);
 			Source.UnlockMip(0);
+
+			if (ImageIndex != INDEX_NONE)
+			{
+				// TODO: add warning if the texture has one or more adjustments,
+				// since adjustments are not applied when exporting via source.
+			}
 		}
-		else
+
+		if (ImageIndex == INDEX_NONE && bPreferSourceExport)
 		{
-			// TODO: report error
-			return FGLTFJsonTextureIndex(INDEX_NONE);
+			Builder.AddWarningMessage(FString::Printf(TEXT("Unable to export source for texture %s, render data will be used as fallback"), *JsonTexture.Name));
 		}
 	}
-	else
-	{
-		const EPixelFormat RenderTargetFormat = FGLTFTextureUtility::IsHDRFormat(Texture2D->GetPixelFormat()) ? PF_FloatRGBA : PF_B8G8R8A8;
 
+	if (ImageIndex == INDEX_NONE && !bRequireSourceExport)
+	{
 		// TODO: both bForceLinearGamma and TargetGamma=2.2f seem to be necessary for the exported images to match the results
 		// from exporting using the texture's source or its platform-data.
 		// It's not entirely clear why gamma must be set to 2.2 instead of 0.0 (which should use the correct gamma anyway),
@@ -87,19 +92,23 @@ FGLTFJsonTextureIndex FGLTFTexture2DConverter::Add(FGLTFConvertBuilder& Builder,
 		const bool bForceLinearGamma = true;
 		const float TargetGamma = 2.2f;
 
+		const EPixelFormat RenderTargetFormat = FGLTFTextureUtility::IsHDRFormat(Texture2D->GetPixelFormat()) ? PF_FloatRGBA : PF_B8G8R8A8;
 		UTextureRenderTarget2D* RenderTarget = FGLTFTextureUtility::CreateRenderTarget(Size, RenderTargetFormat, bForceLinearGamma);
 		RenderTarget->TargetGamma = TargetGamma;
 
 		FGLTFTextureUtility::DrawTexture(RenderTarget, Texture2D);
 
 		TArray<FColor> Pixels;
-		if (!FGLTFTextureUtility::ReadEncodedPixels(RenderTarget, Pixels, JsonTexture.Encoding)) // TODO: use only encoding as specified by export options
+		if (FGLTFTextureUtility::ReadEncodedPixels(RenderTarget, Pixels, JsonTexture.Encoding)) // TODO: use only encoding as specified by export options
 		{
-			// TODO: report error
-			return FGLTFJsonTextureIndex(INDEX_NONE);
+			ImageIndex = Builder.AddImage(Pixels.GetData(), Size, JsonTexture.Name);
 		}
+	}
 
-		ImageIndex = Builder.AddImage(Pixels.GetData(), Size, JsonTexture.Name);
+	if (ImageIndex == INDEX_NONE)
+	{
+		Builder.AddWarningMessage(FString::Printf(TEXT("Failed to export texture %s"), *JsonTexture.Name));
+		return FGLTFJsonTextureIndex(INDEX_NONE);
 	}
 
 	JsonTexture.Source = ImageIndex;
