[GLTFExporter] feat: add utility for detecting if any adjustment is used on a texture

commit db8ee4e6baf25cb62f87c7a1435b35d8134fe0f7
Date: Fri, 11 Sep 2020 09:40:25 +0200

    feat: add utility for detecting if any adjustment is used on a texture
    (cherry picked from commit 76a3c72321ed3e642fbb7280c681fec66d7c852b​)

#skipci

diff --git a/Engine/Plugins/Enterprise/GLTFExporter/MigratedFromGit/Source/GLTFExporter/Private/Converters/GLTFTextureUtility.cpp b/Engine/Plugins/Enterprise/GLTFExporter/MigratedFromGit/Source/GLTFExporter/Private/Converters/GLTFTextureUtility.cpp
--- a/Engine/Plugins/Enterprise/GLTFExporter/MigratedFromGit/Source/GLTFExporter/Private/Converters/GLTFTextureUtility.cpp
+++ b/Engine/Plugins/Enterprise/GLTFExporter/MigratedFromGit/Source/GLTFExporter/Private/Converters/GLTFTextureUtility.cpp
@@ -42,6 +42,26 @@ bool FGLTFTextureUtility::IsCubemap(const UTexture* Texture)
 	return Texture->IsA<UTextureCube>() || Texture->IsA<UTextureRenderTargetCube>();
 }
 
+bool FGLTFTextureUtility::HasAnyAdjustment(const UTexture* Texture)
+{
+	const float ErrorTolerance = KINDA_SMALL_NUMBER;
+
+	if (!FMath::IsNearlyEqual(Texture->AdjustBrightness, 1.0f, ErrorTolerance) ||
+		!FMath::IsNearlyEqual(Texture->AdjustBrightnessCurve, 1.0f, ErrorTolerance) ||
+		!FMath::IsNearlyEqual(Texture->AdjustSaturation, 1.0f, ErrorTolerance) ||
+		!FMath::IsNearlyEqual(Texture->AdjustVibrance, 0.0f, ErrorTolerance) ||
+		!FMath::IsNearlyEqual(Texture->AdjustRGBCurve, 1.0f, ErrorTolerance) ||
+		!FMath::IsNearlyEqual(Texture->AdjustHue, 0.0f, ErrorTolerance) ||
+		!FMath::IsNearlyEqual(Texture->AdjustMinAlpha, 0.0f, ErrorTolerance) ||
+		!FMath::IsNearlyEqual(Texture->AdjustMaxAlpha, 1.0f, ErrorTolerance) ||
+		Texture->bChromaKeyTexture)
+	{
+		return true;
+	}
+
+	return false;
+}
+
 float FGLTFTextureUtility::GetCubeFaceRotation(ECubeFace CubeFace)
 {
 	switch (CubeFace)
diff --git a/Engine/Plugins/Enterprise/GLTFExporter/MigratedFromGit/Source/GLTFExporter/Private/Converters/GLTFTextureUtility.h b/Engine/Plugins/Enterprise/GLTFExporter/MigratedFromGit/Source/GLTFExporter/Private/Converters/GLTFTextureUtility.h
--- a/Engine/Plugins/Enterprise/GLTFExporter/MigratedFromGit/Source/GLTFExporter/Private/Converters/GLTFTextureUtility.h
+++ b/Engine/Plugins/Enterprise/GLTFExporter/MigratedFromGit/Source/GLTFExporter/Private/Converters/GLTFTextureUtility.h
@@ -16,6 +16,8 @@ struct FGLTFTextureUtility
 
 	static bool IsCubemap(const UTexture* Texture);
 
+	static bool HasAnyAdjustment(const UTexture* Texture);
+
 	static float GetCubeFaceRotation(ECubeFace CubeFace);
 
 	static TextureFilter GetDefaultFilter(TextureGroup Group);
