[GLTFExporter] refactor: add warning message when exporting source-textures that have adjustments

commit 839f892da43b177b38a94745b29774848a291d34
Date: Fri, 11 Sep 2020 09:43:02 +0200

    refactor: add warning message when exporting source-textures that have adjustments
    (cherry picked from commit 3c3ca6613b0bbffa3187999f65af1fec44453eae​)

#skipci

diff --git a/Engine/Plugins/Enterprise/GLTFExporter/MigratedFromGit/Source/GLTFExporter/Private/Converters/GLTFTextureConverters.cpp b/Engine/Plugins/Enterprise/GLTFExporter/MigratedFromGit/Source/GLTFExporter/Private/Converters/GLTFTextureConverters.cpp
--- a/Engine/Plugins/Enterprise/GLTFExporter/MigratedFromGit/Source/GLTFExporter/Private/Converters/GLTFTextureConverters.cpp
+++ b/Engine/Plugins/Enterprise/GLTFExporter/MigratedFromGit/Source/GLTFExporter/Private/Converters/GLTFTextureConverters.cpp
@@ -69,17 +69,21 @@ FGLTFJsonTextureIndex FGLTFTexture2DConverter::Add(FGLTFConvertBuilder& Builder,
 			const void* RawData = Source.LockMip(0);
 			ImageIndex = Builder.AddImage(RawData, Source.CalcMipSize(0), Size, RGBFormat, BitDepth, JsonTexture.Name);
 			Source.UnlockMip(0);
+		}
 
-			if (ImageIndex != INDEX_NONE)
+		if (ImageIndex != INDEX_NONE)
+		{
+			if (FGLTFTextureUtility::HasAnyAdjustment(Texture2D))
 			{
-				// TODO: add warning if the texture has one or more adjustments,
-				// since adjustments are not applied when exporting via source.
+				Builder.AddWarningMessage(FString::Printf(TEXT("Adjustments for texture %s are unsupported because source is used for export"), *JsonTexture.Name));
 			}
 		}
-
-		if (ImageIndex == INDEX_NONE && bPreferSourceExport)
+		else
 		{
-			Builder.AddWarningMessage(FString::Printf(TEXT("Unable to export source for texture %s, render data will be used as fallback"), *JsonTexture.Name));
+			if (bPreferSourceExport)
+			{
+				Builder.AddWarningMessage(FString::Printf(TEXT("Unable to export source for texture %s, render data will be used as fallback"), *JsonTexture.Name));
+			}
 		}
 	}
 
