[GLTFExporter] refactor: remove unnecessary control-flow in the HasAnyAdjustment function

commit b601c7e30ead893aaf1b37289b24c92f6a4191d6
Date: Tue, 15 Sep 2020 08:56:13 +0200

    refactor: remove unnecessary control-flow in the HasAnyAdjustment function
    (cherry picked from commit 8feb871376f44268f877425f9f5eb24effd5bccd​)

#skipci

diff --git a/Engine/Plugins/Enterprise/GLTFExporter/MigratedFromGit/Source/GLTFExporter/Private/Converters/GLTFTextureUtility.cpp b/Engine/Plugins/Enterprise/GLTFExporter/MigratedFromGit/Source/GLTFExporter/Private/Converters/GLTFTextureUtility.cpp
--- a/Engine/Plugins/Enterprise/GLTFExporter/MigratedFromGit/Source/GLTFExporter/Private/Converters/GLTFTextureUtility.cpp
+++ b/Engine/Plugins/Enterprise/GLTFExporter/MigratedFromGit/Source/GLTFExporter/Private/Converters/GLTFTextureUtility.cpp
@@ -46,7 +46,7 @@ bool FGLTFTextureUtility::HasAnyAdjustment(const UTexture* Texture)
 {
 	const float ErrorTolerance = KINDA_SMALL_NUMBER;
 
-	if (!FMath::IsNearlyEqual(Texture->AdjustBrightness, 1.0f, ErrorTolerance) ||
+	return !FMath::IsNearlyEqual(Texture->AdjustBrightness, 1.0f, ErrorTolerance) ||
 		!FMath::IsNearlyEqual(Texture->AdjustBrightnessCurve, 1.0f, ErrorTolerance) ||
 		!FMath::IsNearlyEqual(Texture->AdjustSaturation, 1.0f, ErrorTolerance) ||
 		!FMath::IsNearlyEqual(Texture->AdjustVibrance, 0.0f, ErrorTolerance) ||
@@ -54,12 +54,7 @@ bool FGLTFTextureUtility::HasAnyAdjustment(const UTexture* Texture)
 		!FMath::IsNearlyEqual(Texture->AdjustHue, 0.0f, ErrorTolerance) ||
 		!FMath::IsNearlyEqual(Texture->AdjustMinAlpha, 0.0f, ErrorTolerance) ||
 		!FMath::IsNearlyEqual(Texture->AdjustMaxAlpha, 1.0f, ErrorTolerance) ||
-		Texture->bChromaKeyTexture)
-	{
-		return true;
-	}
-
-	return false;
+		Texture->bChromaKeyTexture;
 }
 
 float FGLTFTextureUtility::GetCubeFaceRotation(ECubeFace CubeFace)
