[GLTFExporter] refactor: re-use existing reference to texture source by moving it to the containing scope

commit 93b2349046da1209e64ba2e2d033cbac7c503258
Date: Tue, 15 Sep 2020 09:05:24 +0200

    refactor: re-use existing reference to texture source by moving it to the containing scope
    (cherry picked from commit 27b1d57fec8c2d4d2bb81f6faf3852a2c0001438​)

#skipci

diff --git a/Engine/Plugins/Enterprise/GLTFExporter/MigratedFromGit/Source/GLTFExporter/Private/Converters/GLTFTextureConverters.cpp b/Engine/Plugins/Enterprise/GLTFExporter/MigratedFromGit/Source/GLTFExporter/Private/Converters/GLTFTextureConverters.cpp
--- a/Engine/Plugins/Enterprise/GLTFExporter/MigratedFromGit/Source/GLTFExporter/Private/Converters/GLTFTextureConverters.cpp
+++ b/Engine/Plugins/Enterprise/GLTFExporter/MigratedFromGit/Source/GLTFExporter/Private/Converters/GLTFTextureConverters.cpp
@@ -62,10 +62,10 @@ FGLTFJsonTextureIndex FGLTFTexture2DConverter::Add(FGLTFConvertBuilder& Builder,
 		ERGBFormat RGBFormat;
 		uint32 BitDepth;
 
-		if (Texture2D->Source.IsValid() && FGLTFTextureUtility::CanPNGCompressFormat(Texture2D->Source.GetFormat(), RGBFormat, BitDepth))
-		{
-			FTextureSource& Source = const_cast<FTextureSource&>(Texture2D->Source);
+		FTextureSource& Source = const_cast<FTextureSource&>(Texture2D->Source);
 
+		if (Source.IsValid() && FGLTFTextureUtility::CanPNGCompressFormat(Source.GetFormat(), RGBFormat, BitDepth))
+		{
 			const void* RawData = Source.LockMip(0);
 			ImageIndex = Builder.AddImage(RawData, Source.CalcMipSize(0), Size, RGBFormat, BitDepth, JsonTexture.Name);
 			Source.UnlockMip(0);
