[GLTFExporter] refactor: clarify warning-message for unsupported texture adjustments

commit 68943d2d4836576704f911d6644465e04e12ee06
Date: Tue, 15 Sep 2020 09:08:04 +0200

    refactor: clarify warning-message for unsupported texture adjustments
    (cherry picked from commit e9b0b453aee45131b4d778ad0abd2997cead2b94​)

#skipci

diff --git a/Engine/Plugins/Enterprise/GLTFExporter/MigratedFromGit/Source/GLTFExporter/Private/Converters/GLTFTextureConverters.cpp b/Engine/Plugins/Enterprise/GLTFExporter/MigratedFromGit/Source/GLTFExporter/Private/Converters/GLTFTextureConverters.cpp
--- a/Engine/Plugins/Enterprise/GLTFExporter/MigratedFromGit/Source/GLTFExporter/Private/Converters/GLTFTextureConverters.cpp
+++ b/Engine/Plugins/Enterprise/GLTFExporter/MigratedFromGit/Source/GLTFExporter/Private/Converters/GLTFTextureConverters.cpp
@@ -75,7 +75,7 @@ FGLTFJsonTextureIndex FGLTFTexture2DConverter::Add(FGLTFConvertBuilder& Builder,
 		{
 			if (FGLTFTextureUtility::HasAnyAdjustment(Texture2D))
 			{
-				Builder.AddWarningMessage(FString::Printf(TEXT("Adjustments for texture %s are unsupported because source is used for export"), *JsonTexture.Name));
+				Builder.AddWarningMessage(FString::Printf(TEXT("Adjustments for texture %s are not supported when exporting source data"), *JsonTexture.Name));
 			}
 		}
 		else
