[GLTFExporter] chore: improve wording in note regarding why some textures require source-export

commit 7ae061954b5c0e3e8733484e34ef4bbf7e44b9c2
Date: Tue, 15 Sep 2020 15:50:38 +0200

    chore: improve wording in note regarding why some textures require source-export
    (cherry picked from commit bd15e61daa378fb2d551fb450e8f5b2004aaf8a6​)

#skipci

diff --git a/Engine/Plugins/Enterprise/GLTFExporter/MigratedFromGit/Source/GLTFExporter/Private/Converters/GLTFTextureConverters.cpp b/Engine/Plugins/Enterprise/GLTFExporter/MigratedFromGit/Source/GLTFExporter/Private/Converters/GLTFTextureConverters.cpp
--- a/Engine/Plugins/Enterprise/GLTFExporter/MigratedFromGit/Source/GLTFExporter/Private/Converters/GLTFTextureConverters.cpp
+++ b/Engine/Plugins/Enterprise/GLTFExporter/MigratedFromGit/Source/GLTFExporter/Private/Converters/GLTFTextureConverters.cpp
@@ -51,10 +51,11 @@ FGLTFJsonTextureIndex FGLTFTexture2DConverter::Add(FGLTFConvertBuilder& Builder,
 	const FIntPoint Size = { Texture2D->GetSizeX(), Texture2D->GetSizeY() };
 	const bool bPreferSourceExport = Builder.ExportOptions->bExportSourceTextures;
 
-	// NOTE: export of light- and normal-maps via source is done to overcome issues with accessing
-	// pixel-data (for light-maps) and compressed pixel-data (for normal-maps).
+	// NOTE: export of lightmaps and normalmaps via source data is done to workaround issues
+	// with using render data for normalmaps (which are compressed to red & green channels) and
+	// lightmaps (which can otherwise suffer quality-loss due to incorrect gamma transformation).
 
-	// TODO: make sure there are no other special cases than normal-maps and light-maps that require source-export
+	// TODO: make sure there are no other special cases than normalmaps and lightmaps that require source-export
 	const bool bRequireSourceExport = Texture2D->IsA<ULightMapTexture2D>() || Texture2D->IsNormalMap();
 
 	if (bRequireSourceExport || bPreferSourceExport)
