[GLTFExporter] fix: avoid trying to re-convert assets that fail to be converted

commit 605d8973998ded1c4d7c2e7c31ccd6144869b03a
Date: Thu, 17 Sep 2020 16:39:23 +0200

    fix: avoid trying to re-convert assets that fail to be converted
    (cherry picked from commit aae58158d0842f3714a45360df7af1143581fe25​)

#skipci

diff --git a/Engine/Plugins/Enterprise/GLTFExporter/MigratedFromGit/Source/GLTFExporter/Private/Converters/GLTFConverter.h b/Engine/Plugins/Enterprise/GLTFExporter/MigratedFromGit/Source/GLTFExporter/Private/Converters/GLTFConverter.h
--- a/Engine/Plugins/Enterprise/GLTFExporter/MigratedFromGit/Source/GLTFExporter/Private/Converters/GLTFConverter.h
+++ b/Engine/Plugins/Enterprise/GLTFExporter/MigratedFromGit/Source/GLTFExporter/Private/Converters/GLTFConverter.h
@@ -24,15 +24,15 @@ public:
 	IndexType GetOrAdd(FGLTFConvertBuilder& Builder, const FString& DesiredName, ArgTypes... Args)
 	{
 		const KeyType Key(Forward<ArgTypes>(Args)...);
-
-		IndexType Index = IndexLookup.FindRef(Key);
-		if (Index == INDEX_NONE)
+		if (IndexType* FoundIndex = IndexLookup.Find(Key))
 		{
-			Index = Add(Builder, DesiredName, Forward<ArgTypes>(Args)...);
-			IndexLookup.Add(Key, Index);
+			return *FoundIndex;
 		}
 
-		return Index;
+		IndexType NewIndex = Add(Builder, DesiredName, Forward<ArgTypes>(Args)...);
+
+		IndexLookup.Add(Key, NewIndex);
+		return NewIndex;
 	}
 
 private:
