[GLTFExporter] refactor: simplify converters by removing need to pass builder as argument every call

commit 95b13f18cb4763b9de06b88ea50a0c717fe56816
Date: Sun, 20 Sep 2020 18:54:18 +0200

    refactor: simplify converters by removing need to pass builder as argument every call
    (cherry picked from commit 7c9bf06fae0be0a86ffc07b3c2db5c3dfa6a25a6​)

#skipci

diff --git a/Engine/Plugins/Enterprise/GLTFExporter/MigratedFromGit/Source/GLTFExporter/Private/Builders/GLTFConvertBuilder.cpp b/Engine/Plugins/Enterprise/GLTFExporter/MigratedFromGit/Source/GLTFExporter/Private/Builders/GLTFConvertBuilder.cpp
--- a/Engine/Plugins/Enterprise/GLTFExporter/MigratedFromGit/Source/GLTFExporter/Private/Builders/GLTFConvertBuilder.cpp
+++ b/Engine/Plugins/Enterprise/GLTFExporter/MigratedFromGit/Source/GLTFExporter/Private/Builders/GLTFConvertBuilder.cpp
@@ -16,7 +16,7 @@ FGLTFJsonAccessorIndex FGLTFConvertBuilder::GetOrAddPositionAccessor(const FPosi
 		return FGLTFJsonAccessorIndex(INDEX_NONE);
 	}
 
-	return PositionVertexBufferConverter.GetOrAdd(*this, DesiredName, VertexBuffer);
+	return PositionVertexBufferConverter.GetOrAdd(DesiredName, VertexBuffer);
 }
 
 FGLTFJsonAccessorIndex FGLTFConvertBuilder::GetOrAddColorAccessor(const FColorVertexBuffer* VertexBuffer, const FString& DesiredName)
@@ -26,7 +26,7 @@ FGLTFJsonAccessorIndex FGLTFConvertBuilder::GetOrAddColorAccessor(const FColorVe
 		return FGLTFJsonAccessorIndex(INDEX_NONE);
 	}
 
-	return ColorVertexBufferConverter.GetOrAdd(*this, DesiredName, VertexBuffer);
+	return ColorVertexBufferConverter.GetOrAdd(DesiredName, VertexBuffer);
 }
 
 FGLTFJsonAccessorIndex FGLTFConvertBuilder::GetOrAddNormalAccessor(const FStaticMeshVertexBuffer* VertexBuffer, const FString& DesiredName)
@@ -36,7 +36,7 @@ FGLTFJsonAccessorIndex FGLTFConvertBuilder::GetOrAddNormalAccessor(const FStatic
 		return FGLTFJsonAccessorIndex(INDEX_NONE);
 	}
 
-	return NormalVertexBufferConverter.GetOrAdd(*this, DesiredName, VertexBuffer);
+	return NormalVertexBufferConverter.GetOrAdd(DesiredName, VertexBuffer);
 }
 
 FGLTFJsonAccessorIndex FGLTFConvertBuilder::GetOrAddTangentAccessor(const FStaticMeshVertexBuffer* VertexBuffer, const FString& DesiredName)
@@ -46,7 +46,7 @@ FGLTFJsonAccessorIndex FGLTFConvertBuilder::GetOrAddTangentAccessor(const FStati
 		return FGLTFJsonAccessorIndex(INDEX_NONE);
 	}
 
-	return TangentVertexBufferConverter.GetOrAdd(*this, DesiredName, VertexBuffer);
+	return TangentVertexBufferConverter.GetOrAdd(DesiredName, VertexBuffer);
 }
 
 FGLTFJsonAccessorIndex FGLTFConvertBuilder::GetOrAddUVAccessor(const FStaticMeshVertexBuffer* VertexBuffer, int32 UVIndex, const FString& DesiredName)
@@ -56,7 +56,7 @@ FGLTFJsonAccessorIndex FGLTFConvertBuilder::GetOrAddUVAccessor(const FStaticMesh
 		return FGLTFJsonAccessorIndex(INDEX_NONE);
 	}
 
-	return UVVertexBufferConverter.GetOrAdd(*this, DesiredName, VertexBuffer, UVIndex);
+	return UVVertexBufferConverter.GetOrAdd(DesiredName, VertexBuffer, UVIndex);
 }
 
 FGLTFJsonAccessorIndex FGLTFConvertBuilder::GetOrAddJointAccessor(const FSkinWeightVertexBuffer* VertexBuffer, int32 JointsGroupIndex, FGLTFBoneMap BoneMap, const FString& DesiredName)
@@ -66,7 +66,7 @@ FGLTFJsonAccessorIndex FGLTFConvertBuilder::GetOrAddJointAccessor(const FSkinWei
 		return FGLTFJsonAccessorIndex(INDEX_NONE);
 	}
 
-	return BoneIndexVertexBufferConverter.GetOrAdd(*this, DesiredName, VertexBuffer, JointsGroupIndex, BoneMap);
+	return BoneIndexVertexBufferConverter.GetOrAdd(DesiredName, VertexBuffer, JointsGroupIndex, BoneMap);
 }
 
 FGLTFJsonAccessorIndex FGLTFConvertBuilder::GetOrAddWeightAccessor(const FSkinWeightVertexBuffer* VertexBuffer, int32 WeightsGroupIndex, const FString& DesiredName)
@@ -76,7 +76,7 @@ FGLTFJsonAccessorIndex FGLTFConvertBuilder::GetOrAddWeightAccessor(const FSkinWe
 		return FGLTFJsonAccessorIndex(INDEX_NONE);
 	}
 
-	return BoneWeightVertexBufferConverter.GetOrAdd(*this, DesiredName, VertexBuffer, WeightsGroupIndex);
+	return BoneWeightVertexBufferConverter.GetOrAdd(DesiredName, VertexBuffer, WeightsGroupIndex);
 }
 
 FGLTFJsonBufferViewIndex FGLTFConvertBuilder::GetOrAddIndexBufferView(const FRawStaticIndexBuffer* IndexBuffer, const FString& DesiredName)
@@ -86,7 +86,7 @@ FGLTFJsonBufferViewIndex FGLTFConvertBuilder::GetOrAddIndexBufferView(const FRaw
 		return FGLTFJsonBufferViewIndex(INDEX_NONE);
 	}
 
-	return IndexBufferConverter.GetOrAdd(*this, DesiredName, IndexBuffer);
+	return IndexBufferConverter.GetOrAdd(DesiredName, IndexBuffer);
 }
 
 FGLTFJsonAccessorIndex FGLTFConvertBuilder::GetOrAddIndexAccessor(const FStaticMeshSection* MeshSection, const FRawStaticIndexBuffer* IndexBuffer, const FString& DesiredName)
@@ -96,7 +96,7 @@ FGLTFJsonAccessorIndex FGLTFConvertBuilder::GetOrAddIndexAccessor(const FStaticM
 		return FGLTFJsonAccessorIndex(INDEX_NONE);
 	}
 
-	return StaticMeshSectionConverter.GetOrAdd(*this, DesiredName, MeshSection, IndexBuffer);
+	return StaticMeshSectionConverter.GetOrAdd(DesiredName, MeshSection, IndexBuffer);
 }
 
 FGLTFJsonMeshIndex FGLTFConvertBuilder::GetOrAddMesh(const UStaticMesh* StaticMesh, int32 LODIndex, const FColorVertexBuffer* OverrideVertexColors, const FGLTFMaterialArray& OverrideMaterials, const FString& DesiredName)
@@ -117,7 +117,7 @@ FGLTFJsonMeshIndex FGLTFConvertBuilder::GetOrAddMesh(const UStaticMesh* StaticMe
 	}
 
 	return StaticMeshConverter.GetOrAdd(
-		*this, DesiredName.IsEmpty() ? FGLTFBuilderUtility::GetLODName(StaticMesh, LODIndex) : DesiredName,
+		DesiredName.IsEmpty() ? FGLTFBuilderUtility::GetLODName(StaticMesh, LODIndex) : DesiredName,
 		StaticMesh,
 		LODIndex,
 		OverrideVertexColors,
@@ -146,7 +146,7 @@ FGLTFJsonBufferViewIndex FGLTFConvertBuilder::GetOrAddIndexBufferView(const FMul
 		return FGLTFJsonBufferViewIndex(INDEX_NONE);
 	}
 
-	return IndexContainerConverter.GetOrAdd(*this, DesiredName, IndexContainer);
+	return IndexContainerConverter.GetOrAdd(DesiredName, IndexContainer);
 }
 
 FGLTFJsonAccessorIndex FGLTFConvertBuilder::GetOrAddIndexAccessor(const FSkelMeshRenderSection* MeshSection, const FMultiSizeIndexContainer* IndexContainer, const FString& DesiredName)
@@ -156,7 +156,7 @@ FGLTFJsonAccessorIndex FGLTFConvertBuilder::GetOrAddIndexAccessor(const FSkelMes
 		return FGLTFJsonAccessorIndex(INDEX_NONE);
 	}
 
-	return SkeletalMeshSectionConverter.GetOrAdd(*this, DesiredName, MeshSection, IndexContainer);
+	return SkeletalMeshSectionConverter.GetOrAdd(DesiredName, MeshSection, IndexContainer);
 }
 
 FGLTFJsonMeshIndex FGLTFConvertBuilder::GetOrAddMesh(const USkeletalMesh* SkeletalMesh, int32 LODIndex, const FColorVertexBuffer* OverrideVertexColors, const FSkinWeightVertexBuffer* OverrideSkinWeights, const FGLTFMaterialArray& OverrideMaterials, const FString& DesiredName)
@@ -176,7 +176,7 @@ FGLTFJsonMeshIndex FGLTFConvertBuilder::GetOrAddMesh(const USkeletalMesh* Skelet
 		OverrideVertexColors = nullptr;
 	}
 
-	return SkeletalMeshConverter.GetOrAdd(*this, DesiredName.IsEmpty() ? FGLTFBuilderUtility::GetLODName(SkeletalMesh, LODIndex) : DesiredName, SkeletalMesh, LODIndex, OverrideVertexColors, OverrideSkinWeights, OverrideMaterials);
+	return SkeletalMeshConverter.GetOrAdd(DesiredName.IsEmpty() ? FGLTFBuilderUtility::GetLODName(SkeletalMesh, LODIndex) : DesiredName, SkeletalMesh, LODIndex, OverrideVertexColors, OverrideSkinWeights, OverrideMaterials);
 }
 
 FGLTFJsonMeshIndex FGLTFConvertBuilder::GetOrAddMesh(const USkeletalMeshComponent* SkeletalMeshComponent, const FString& DesiredName)
@@ -202,7 +202,7 @@ FGLTFJsonMaterialIndex FGLTFConvertBuilder::GetOrAddMaterial(const UMaterialInte
 		return FGLTFJsonMaterialIndex(INDEX_NONE);
 	}
 
-	return MaterialConverter.GetOrAdd(*this, DesiredName.IsEmpty() ? Material->GetName() : DesiredName, Material);
+	return MaterialConverter.GetOrAdd(DesiredName.IsEmpty() ? Material->GetName() : DesiredName, Material);
 }
 
 FGLTFJsonSamplerIndex FGLTFConvertBuilder::GetOrAddSampler(const UTexture* Texture, const FString& DesiredName)
@@ -212,7 +212,7 @@ FGLTFJsonSamplerIndex FGLTFConvertBuilder::GetOrAddSampler(const UTexture* Textu
 		return FGLTFJsonSamplerIndex(INDEX_NONE);
 	}
 
-	return TextureSamplerConverter.GetOrAdd(*this, DesiredName.IsEmpty() ? Texture->GetName() : DesiredName, Texture);
+	return TextureSamplerConverter.GetOrAdd(DesiredName.IsEmpty() ? Texture->GetName() : DesiredName, Texture);
 }
 
 FGLTFJsonTextureIndex FGLTFConvertBuilder::GetOrAddTexture(const UTexture2D* Texture, const FString& DesiredName)
@@ -222,7 +222,7 @@ FGLTFJsonTextureIndex FGLTFConvertBuilder::GetOrAddTexture(const UTexture2D* Tex
 		return FGLTFJsonTextureIndex(INDEX_NONE);
 	}
 
-	return Texture2DConverter.GetOrAdd(*this, DesiredName.IsEmpty() ? Texture->GetName() : DesiredName, Texture);
+	return Texture2DConverter.GetOrAdd(DesiredName.IsEmpty() ? Texture->GetName() : DesiredName, Texture);
 }
 
 FGLTFJsonTextureIndex FGLTFConvertBuilder::GetOrAddTexture(const UTextureCube* Texture, ECubeFace CubeFace, const FString& DesiredName)
@@ -232,7 +232,7 @@ FGLTFJsonTextureIndex FGLTFConvertBuilder::GetOrAddTexture(const UTextureCube* T
 		return FGLTFJsonTextureIndex(INDEX_NONE);
 	}
 
-	return TextureCubeConverter.GetOrAdd(*this, DesiredName.IsEmpty() ? Texture->GetName() : DesiredName, Texture, CubeFace);
+	return TextureCubeConverter.GetOrAdd(DesiredName.IsEmpty() ? Texture->GetName() : DesiredName, Texture, CubeFace);
 }
 
 FGLTFJsonTextureIndex FGLTFConvertBuilder::GetOrAddTexture(const UTextureRenderTarget2D* Texture, const FString& DesiredName)
@@ -242,7 +242,7 @@ FGLTFJsonTextureIndex FGLTFConvertBuilder::GetOrAddTexture(const UTextureRenderT
 		return FGLTFJsonTextureIndex(INDEX_NONE);
 	}
 
-	return TextureRenderTarget2DConverter.GetOrAdd(*this, DesiredName.IsEmpty() ? Texture->GetName() : DesiredName, Texture);
+	return TextureRenderTarget2DConverter.GetOrAdd(DesiredName.IsEmpty() ? Texture->GetName() : DesiredName, Texture);
 }
 
 FGLTFJsonTextureIndex FGLTFConvertBuilder::GetOrAddTexture(const UTextureRenderTargetCube* Texture, ECubeFace CubeFace, const FString& DesiredName)
@@ -252,7 +252,7 @@ FGLTFJsonTextureIndex FGLTFConvertBuilder::GetOrAddTexture(const UTextureRenderT
 		return FGLTFJsonTextureIndex(INDEX_NONE);
 	}
 
-	return TextureRenderTargetCubeConverter.GetOrAdd(*this, DesiredName.IsEmpty() ? Texture->GetName() : DesiredName, Texture, CubeFace);
+	return TextureRenderTargetCubeConverter.GetOrAdd(DesiredName.IsEmpty() ? Texture->GetName() : DesiredName, Texture, CubeFace);
 }
 
 FGLTFJsonNodeIndex FGLTFConvertBuilder::GetOrAddNode(const USceneComponent* SceneComponent, const FString& DesiredName)
@@ -262,7 +262,7 @@ FGLTFJsonNodeIndex FGLTFConvertBuilder::GetOrAddNode(const USceneComponent* Scen
 		return FGLTFJsonNodeIndex(INDEX_NONE);
 	}
 
-	return SceneComponentConverter.GetOrAdd(*this, DesiredName, SceneComponent);
+	return SceneComponentConverter.GetOrAdd(DesiredName, SceneComponent);
 }
 
 FGLTFJsonNodeIndex FGLTFConvertBuilder::GetOrAddNode(const AActor* Actor, const FString& DesiredName)
@@ -272,7 +272,7 @@ FGLTFJsonNodeIndex FGLTFConvertBuilder::GetOrAddNode(const AActor* Actor, const
 		return FGLTFJsonNodeIndex(INDEX_NONE);
 	}
 
-	return ActorConverter.GetOrAdd(*this, DesiredName, Actor);
+	return ActorConverter.GetOrAdd(DesiredName, Actor);
 }
 
 FGLTFJsonSceneIndex FGLTFConvertBuilder::GetOrAddScene(const ULevel* Level, const FString& DesiredName)
@@ -282,7 +282,7 @@ FGLTFJsonSceneIndex FGLTFConvertBuilder::GetOrAddScene(const ULevel* Level, cons
 		return FGLTFJsonSceneIndex(INDEX_NONE);
 	}
 
-	return LevelConverter.GetOrAdd(*this, DesiredName, Level);
+	return LevelConverter.GetOrAdd(DesiredName, Level);
 }
 
 FGLTFJsonSceneIndex FGLTFConvertBuilder::GetOrAddScene(const UWorld* World, const FString& DesiredName)
@@ -302,7 +302,7 @@ FGLTFJsonCameraIndex FGLTFConvertBuilder::GetOrAddCamera(const UCameraComponent*
 		return FGLTFJsonCameraIndex(INDEX_NONE);
 	}
 
-	return CameraComponentConverter.GetOrAdd(*this, DesiredName.IsEmpty() ? CameraComponent->GetName() : DesiredName, CameraComponent);
+	return CameraComponentConverter.GetOrAdd(DesiredName.IsEmpty() ? CameraComponent->GetName() : DesiredName, CameraComponent);
 }
 
 FGLTFJsonLightIndex FGLTFConvertBuilder::GetOrAddLight(const ULightComponent* LightComponent, const FString& DesiredName)
@@ -312,7 +312,7 @@ FGLTFJsonLightIndex FGLTFConvertBuilder::GetOrAddLight(const ULightComponent* Li
 		return FGLTFJsonLightIndex(INDEX_NONE);
 	}
 
-	return LightComponentConverter.GetOrAdd(*this, DesiredName.IsEmpty() ? LightComponent->GetName() : DesiredName, LightComponent);
+	return LightComponentConverter.GetOrAdd(DesiredName.IsEmpty() ? LightComponent->GetName() : DesiredName, LightComponent);
 }
 
 FGLTFJsonBackdropIndex FGLTFConvertBuilder::GetOrAddBackdrop(const AActor* Actor, const FString& DesiredName)
@@ -322,7 +322,7 @@ FGLTFJsonBackdropIndex FGLTFConvertBuilder::GetOrAddBackdrop(const AActor* Actor
 		return FGLTFJsonBackdropIndex(INDEX_NONE);
 	}
 
-	return BackdropConverter.GetOrAdd(*this, DesiredName, Actor);
+	return BackdropConverter.GetOrAdd(DesiredName, Actor);
 }
 
 FGLTFJsonLevelVariantSetsIndex FGLTFConvertBuilder::GetOrAddLevelVariantSets(const ALevelVariantSetsActor* LevelVariantSetsActor, const FString& DesiredName)
@@ -332,7 +332,7 @@ FGLTFJsonLevelVariantSetsIndex FGLTFConvertBuilder::GetOrAddLevelVariantSets(con
 		return FGLTFJsonLevelVariantSetsIndex(INDEX_NONE);
 	}
 
-	return LevelVariantSetsConverter.GetOrAdd(*this, DesiredName, LevelVariantSetsActor);
+	return LevelVariantSetsConverter.GetOrAdd(DesiredName, LevelVariantSetsActor);
 }
 
 FGLTFJsonLightMapIndex FGLTFConvertBuilder::GetOrAddLightMap(const UStaticMeshComponent* StaticMeshComponent, const FString& DesiredName)
@@ -342,7 +342,7 @@ FGLTFJsonLightMapIndex FGLTFConvertBuilder::GetOrAddLightMap(const UStaticMeshCo
 		return FGLTFJsonLightMapIndex(INDEX_NONE);
 	}
 
-	return LightMapConverter.GetOrAdd(*this, DesiredName, StaticMeshComponent);
+	return LightMapConverter.GetOrAdd(DesiredName, StaticMeshComponent);
 }
 
 FGLTFJsonHotspotIndex FGLTFConvertBuilder::GetOrAddHotspot(const UGLTFInteractionHotspotComponent* HotspotComponent, const FString& DesiredName)
@@ -352,5 +352,5 @@ FGLTFJsonHotspotIndex FGLTFConvertBuilder::GetOrAddHotspot(const UGLTFInteractio
 		return FGLTFJsonHotspotIndex(INDEX_NONE);
 	}
 
-	return HotspotComponentConverter.GetOrAdd(*this, DesiredName.IsEmpty() ? HotspotComponent->GetName() : DesiredName, HotspotComponent);
+	return HotspotComponentConverter.GetOrAdd(DesiredName.IsEmpty() ? HotspotComponent->GetName() : DesiredName, HotspotComponent);
 }
diff --git a/Engine/Plugins/Enterprise/GLTFExporter/MigratedFromGit/Source/GLTFExporter/Private/Builders/GLTFConvertBuilder.h b/Engine/Plugins/Enterprise/GLTFExporter/MigratedFromGit/Source/GLTFExporter/Private/Builders/GLTFConvertBuilder.h
--- a/Engine/Plugins/Enterprise/GLTFExporter/MigratedFromGit/Source/GLTFExporter/Private/Builders/GLTFConvertBuilder.h
+++ b/Engine/Plugins/Enterprise/GLTFExporter/MigratedFromGit/Source/GLTFExporter/Private/Builders/GLTFConvertBuilder.h
@@ -62,38 +62,38 @@ public:
 
 private:
 
-	FGLTFPositionVertexBufferConverter PositionVertexBufferConverter;
-	FGLTFColorVertexBufferConverter ColorVertexBufferConverter;
-	FGLTFNormalVertexBufferConverter NormalVertexBufferConverter;
-	FGLTFTangentVertexBufferConverter TangentVertexBufferConverter;
-	FGLTFUVVertexBufferConverter UVVertexBufferConverter;
-	FGLTFBoneIndexVertexBufferConverter BoneIndexVertexBufferConverter;
-	FGLTFBoneWeightVertexBufferConverter BoneWeightVertexBufferConverter;
-
-	FGLTFIndexBufferConverter IndexBufferConverter;
-	FGLTFStaticMeshSectionConverter StaticMeshSectionConverter;
-	FGLTFStaticMeshConverter StaticMeshConverter;
-
-	FGLTFIndexContainerConverter IndexContainerConverter;
-	FGLTFSkeletalMeshSectionConverter SkeletalMeshSectionConverter;
-	FGLTFSkeletalMeshConverter SkeletalMeshConverter;
-
-	FGLTFMaterialConverter MaterialConverter;
-	FGLTFTextureSamplerConverter TextureSamplerConverter;
-	FGLTFTexture2DConverter Texture2DConverter;
-	FGLTFTextureCubeConverter TextureCubeConverter;
-	FGLTFTextureRenderTarget2DConverter TextureRenderTarget2DConverter;
-	FGLTFTextureRenderTargetCubeConverter TextureRenderTargetCubeConverter;
-
-	FGLTFSceneComponentConverter SceneComponentConverter;
-	FGLTFActorConverter ActorConverter;
-	FGLTFLevelConverter LevelConverter;
-
-	FGLTFCameraComponentConverter CameraComponentConverter;
-	FGLTFLightComponentConverter LightComponentConverter;
-
-	FGLTFBackdropConverter BackdropConverter;
-	FGLTFLevelVariantSetsConverter LevelVariantSetsConverter;
-	FGLTFLightMapConverter LightMapConverter;
-	FGLTFHotspotComponentConverter HotspotComponentConverter;
+	FGLTFPositionVertexBufferConverter PositionVertexBufferConverter = *this;
+	FGLTFColorVertexBufferConverter ColorVertexBufferConverter = *this;
+	FGLTFNormalVertexBufferConverter NormalVertexBufferConverter = *this;
+	FGLTFTangentVertexBufferConverter TangentVertexBufferConverter = *this;
+	FGLTFUVVertexBufferConverter UVVertexBufferConverter = *this;
+	FGLTFBoneIndexVertexBufferConverter BoneIndexVertexBufferConverter = *this;
+	FGLTFBoneWeightVertexBufferConverter BoneWeightVertexBufferConverter = *this;
+
+	FGLTFIndexBufferConverter IndexBufferConverter = *this;
+	FGLTFStaticMeshSectionConverter StaticMeshSectionConverter = *this;
+	FGLTFStaticMeshConverter StaticMeshConverter = *this;
+
+	FGLTFIndexContainerConverter IndexContainerConverter = *this;
+	FGLTFSkeletalMeshSectionConverter SkeletalMeshSectionConverter = *this;
+	FGLTFSkeletalMeshConverter SkeletalMeshConverter = *this;
+
+	FGLTFMaterialConverter MaterialConverter = *this;
+	FGLTFTextureSamplerConverter TextureSamplerConverter = *this;
+	FGLTFTexture2DConverter Texture2DConverter = *this;
+	FGLTFTextureCubeConverter TextureCubeConverter = *this;
+	FGLTFTextureRenderTarget2DConverter TextureRenderTarget2DConverter = *this;
+	FGLTFTextureRenderTargetCubeConverter TextureRenderTargetCubeConverter = *this;
+
+	FGLTFSceneComponentConverter SceneComponentConverter = *this;
+	FGLTFActorConverter ActorConverter = *this;
+	FGLTFLevelConverter LevelConverter = *this;
+
+	FGLTFCameraComponentConverter CameraComponentConverter = *this;
+	FGLTFLightComponentConverter LightComponentConverter = *this;
+
+	FGLTFBackdropConverter BackdropConverter = *this;
+	FGLTFLevelVariantSetsConverter LevelVariantSetsConverter = *this;
+	FGLTFLightMapConverter LightMapConverter = *this;
+	FGLTFHotspotComponentConverter HotspotComponentConverter = *this;
 };
diff --git a/Engine/Plugins/Enterprise/GLTFExporter/MigratedFromGit/Source/GLTFExporter/Private/Converters/GLTFBackdropConverters.cpp b/Engine/Plugins/Enterprise/GLTFExporter/MigratedFromGit/Source/GLTFExporter/Private/Converters/GLTFBackdropConverters.cpp
--- a/Engine/Plugins/Enterprise/GLTFExporter/MigratedFromGit/Source/GLTFExporter/Private/Converters/GLTFBackdropConverters.cpp
+++ b/Engine/Plugins/Enterprise/GLTFExporter/MigratedFromGit/Source/GLTFExporter/Private/Converters/GLTFBackdropConverters.cpp
@@ -5,7 +5,7 @@
 #include "Converters/GLTFActorUtility.h"
 #include "Builders/GLTFContainerBuilder.h"
 
-FGLTFJsonBackdropIndex FGLTFBackdropConverter::Add(FGLTFConvertBuilder& Builder, const FString& Name, const AActor* Actor)
+FGLTFJsonBackdropIndex FGLTFBackdropConverter::Convert(const FString& Name, const AActor* Actor)
 {
 	const UBlueprint* Blueprint = FGLTFActorUtility::GetBlueprintFromActor(Actor);
 	if (!FGLTFActorUtility::IsHDRIBackdropBlueprint(Blueprint))
diff --git a/Engine/Plugins/Enterprise/GLTFExporter/MigratedFromGit/Source/GLTFExporter/Private/Converters/GLTFBackdropConverters.h b/Engine/Plugins/Enterprise/GLTFExporter/MigratedFromGit/Source/GLTFExporter/Private/Converters/GLTFBackdropConverters.h
--- a/Engine/Plugins/Enterprise/GLTFExporter/MigratedFromGit/Source/GLTFExporter/Private/Converters/GLTFBackdropConverters.h
+++ b/Engine/Plugins/Enterprise/GLTFExporter/MigratedFromGit/Source/GLTFExporter/Private/Converters/GLTFBackdropConverters.h
@@ -8,5 +8,7 @@
 
 class FGLTFBackdropConverter final : public TGLTFConverter<FGLTFJsonBackdropIndex, const AActor*>
 {
-	FGLTFJsonBackdropIndex Add(FGLTFConvertBuilder& Builder, const FString& Name, const AActor* Actor) override;
+	using TGLTFConverter::TGLTFConverter;
+
+	FGLTFJsonBackdropIndex Convert(const FString& Name, const AActor* Actor) override;
 };
diff --git a/Engine/Plugins/Enterprise/GLTFExporter/MigratedFromGit/Source/GLTFExporter/Private/Converters/GLTFConverter.h b/Engine/Plugins/Enterprise/GLTFExporter/MigratedFromGit/Source/GLTFExporter/Private/Converters/GLTFConverter.h
--- a/Engine/Plugins/Enterprise/GLTFExporter/MigratedFromGit/Source/GLTFExporter/Private/Converters/GLTFConverter.h
+++ b/Engine/Plugins/Enterprise/GLTFExporter/MigratedFromGit/Source/GLTFExporter/Private/Converters/GLTFConverter.h
@@ -6,38 +6,61 @@
 
 class FGLTFConvertBuilder;
 
-template <typename IndexType, typename... ArgTypes>
+template <typename OutputType, typename... InputTypes>
 class TGLTFConverter
 {
-	typedef TTuple<ArgTypes...> KeyType;
+	typedef TTuple<InputTypes...> InputKeyType;
 
 public:
 
+	TGLTFConverter(FGLTFConvertBuilder& Builder)
+		: Builder(Builder)
+	{
+	}
+
 	virtual ~TGLTFConverter() = default;
 
-	IndexType Get(ArgTypes&&... Args) const
+	OutputType Get(InputTypes&&... Inputs) const
+	{
+		const InputKeyType InputKey(Forward<InputTypes>(Inputs)...);
+		if (OutputType* SavedOutput = SavedOutputs.Find(InputKey))
+		{
+			return *SavedOutput;
+		}
+
+		return {};
+	}
+
+	OutputType Add(const FString& DesiredName, InputTypes... Inputs)
 	{
-		const KeyType Key(Forward<ArgTypes>(Args)...);
-		return IndexLookup.FindRef(Key);
+		const InputKeyType InputKey(Forward<InputTypes>(Inputs)...);
+		OutputType NewOutput = Convert(DesiredName, Forward<InputTypes>(Inputs)...);
+
+		SavedOutputs.Add(InputKey, NewOutput);
+		return NewOutput;
 	}
 
-	IndexType GetOrAdd(FGLTFConvertBuilder& Builder, const FString& DesiredName, ArgTypes... Args)
+	OutputType GetOrAdd(const FString& DesiredName, InputTypes... Inputs)
 	{
-		const KeyType Key(Forward<ArgTypes>(Args)...);
-		if (IndexType* FoundIndex = IndexLookup.Find(Key))
+		const InputKeyType InputKey(Forward<InputTypes>(Inputs)...);
+		if (OutputType* SavedOutput = SavedOutputs.Find(InputKey))
 		{
-			return *FoundIndex;
+			return *SavedOutput;
 		}
 
-		IndexType NewIndex = Add(Builder, DesiredName, Forward<ArgTypes>(Args)...);
+		OutputType NewOutput = Convert(DesiredName, Forward<InputTypes>(Inputs)...);
 
-		IndexLookup.Add(Key, NewIndex);
-		return NewIndex;
+		SavedOutputs.Add(InputKey, NewOutput);
+		return NewOutput;
 	}
 
-private:
+protected:
 
-	virtual IndexType Add(FGLTFConvertBuilder& Builder, const FString& Name, ArgTypes... Args) = 0;
+	virtual OutputType Convert(const FString& Name, InputTypes... Inputs) = 0;
+
+	FGLTFConvertBuilder& Builder;
+
+private:
 
-	TMap<KeyType, IndexType> IndexLookup;
+	TMap<InputKeyType, OutputType> SavedOutputs;
 };
diff --git a/Engine/Plugins/Enterprise/GLTFExporter/MigratedFromGit/Source/GLTFExporter/Private/Converters/GLTFHotspotConverters.cpp b/Engine/Plugins/Enterprise/GLTFExporter/MigratedFromGit/Source/GLTFExporter/Private/Converters/GLTFHotspotConverters.cpp
--- a/Engine/Plugins/Enterprise/GLTFExporter/MigratedFromGit/Source/GLTFExporter/Private/Converters/GLTFHotspotConverters.cpp
+++ b/Engine/Plugins/Enterprise/GLTFExporter/MigratedFromGit/Source/GLTFExporter/Private/Converters/GLTFHotspotConverters.cpp
@@ -4,7 +4,7 @@
 #include "Json/FGLTFJsonHotspot.h"
 #include "Builders/GLTFContainerBuilder.h"
 
-FGLTFJsonHotspotIndex FGLTFHotspotComponentConverter::Add(FGLTFConvertBuilder& Builder, const FString& Name, const UGLTFInteractionHotspotComponent* HotspotComponent)
+FGLTFJsonHotspotIndex FGLTFHotspotComponentConverter::Convert(const FString& Name, const UGLTFInteractionHotspotComponent* HotspotComponent)
 {
 	// TODO: should we warn and / or return INDEX_NONE if the hotspot has no valid Image assigned?
 
diff --git a/Engine/Plugins/Enterprise/GLTFExporter/MigratedFromGit/Source/GLTFExporter/Private/Converters/GLTFHotspotConverters.h b/Engine/Plugins/Enterprise/GLTFExporter/MigratedFromGit/Source/GLTFExporter/Private/Converters/GLTFHotspotConverters.h
--- a/Engine/Plugins/Enterprise/GLTFExporter/MigratedFromGit/Source/GLTFExporter/Private/Converters/GLTFHotspotConverters.h
+++ b/Engine/Plugins/Enterprise/GLTFExporter/MigratedFromGit/Source/GLTFExporter/Private/Converters/GLTFHotspotConverters.h
@@ -9,5 +9,7 @@
 
 class FGLTFHotspotComponentConverter final : public TGLTFConverter<FGLTFJsonHotspotIndex, const UGLTFInteractionHotspotComponent*>
 {
-	FGLTFJsonHotspotIndex Add(FGLTFConvertBuilder& Builder, const FString& Name, const UGLTFInteractionHotspotComponent* HotspotComponent) override;
+	using TGLTFConverter::TGLTFConverter;
+
+	FGLTFJsonHotspotIndex Convert(const FString& Name, const UGLTFInteractionHotspotComponent* HotspotComponent) override;
 };
diff --git a/Engine/Plugins/Enterprise/GLTFExporter/MigratedFromGit/Source/GLTFExporter/Private/Converters/GLTFLevelConverters.cpp b/Engine/Plugins/Enterprise/GLTFExporter/MigratedFromGit/Source/GLTFExporter/Private/Converters/GLTFLevelConverters.cpp
--- a/Engine/Plugins/Enterprise/GLTFExporter/MigratedFromGit/Source/GLTFExporter/Private/Converters/GLTFLevelConverters.cpp
+++ b/Engine/Plugins/Enterprise/GLTFExporter/MigratedFromGit/Source/GLTFExporter/Private/Converters/GLTFLevelConverters.cpp
@@ -7,7 +7,7 @@
 #include "Converters/GLTFCameraUtility.h"
 #include "Components/GLTFInteractionHotspotComponent.h"
 
-FGLTFJsonNodeIndex FGLTFSceneComponentConverter::Add(FGLTFConvertBuilder& Builder, const FString& Name, const USceneComponent* SceneComponent)
+FGLTFJsonNodeIndex FGLTFSceneComponentConverter::Convert(const FString& Name, const USceneComponent* SceneComponent)
 {
 	const AActor* Owner = SceneComponent->GetOwner();
 	if (Owner == nullptr)
@@ -142,7 +142,7 @@ FGLTFJsonNodeIndex FGLTFSceneComponentConverter::Add(FGLTFConvertBuilder& Builde
 	return NodeIndex;
 }
 
-FGLTFJsonNodeIndex FGLTFActorConverter::Add(FGLTFConvertBuilder& Builder, const FString& Name, const AActor* Actor)
+FGLTFJsonNodeIndex FGLTFActorConverter::Convert(const FString& Name, const AActor* Actor)
 {
 	if (Builder.bSelectedActorsOnly && !Actor->IsSelected())
 	{
@@ -189,7 +189,7 @@ FGLTFJsonNodeIndex FGLTFActorConverter::Add(FGLTFConvertBuilder& Builder, const
 	return RootNodeIndex;
 }
 
-FGLTFJsonSceneIndex FGLTFLevelConverter::Add(FGLTFConvertBuilder& Builder, const FString& Name, const ULevel* Level)
+FGLTFJsonSceneIndex FGLTFLevelConverter::Convert(const FString& Name, const ULevel* Level)
 {
 	FGLTFJsonScene Scene;
 	Scene.Name = Name;
@@ -227,7 +227,7 @@ FGLTFJsonSceneIndex FGLTFLevelConverter::Add(FGLTFConvertBuilder& Builder, const
 	return Builder.AddScene(Scene);
 }
 
-FGLTFJsonCameraIndex FGLTFCameraComponentConverter::Add(FGLTFConvertBuilder& Builder, const FString& Name, const UCameraComponent* CameraComponent)
+FGLTFJsonCameraIndex FGLTFCameraComponentConverter::Convert(const FString& Name, const UCameraComponent* CameraComponent)
 {
 	FGLTFJsonCamera Camera;
 	Camera.Name = Name;
@@ -254,7 +254,7 @@ FGLTFJsonCameraIndex FGLTFCameraComponentConverter::Add(FGLTFConvertBuilder& Bui
 	return Builder.AddCamera(Camera);
 }
 
-FGLTFJsonLightIndex FGLTFLightComponentConverter::Add(FGLTFConvertBuilder& Builder, const FString& Name, const ULightComponent* LightComponent)
+FGLTFJsonLightIndex FGLTFLightComponentConverter::Convert(const FString& Name, const ULightComponent* LightComponent)
 {
 	FGLTFJsonLight Light;
 	Light.Name = Name;
diff --git a/Engine/Plugins/Enterprise/GLTFExporter/MigratedFromGit/Source/GLTFExporter/Private/Converters/GLTFLevelConverters.h b/Engine/Plugins/Enterprise/GLTFExporter/MigratedFromGit/Source/GLTFExporter/Private/Converters/GLTFLevelConverters.h
--- a/Engine/Plugins/Enterprise/GLTFExporter/MigratedFromGit/Source/GLTFExporter/Private/Converters/GLTFLevelConverters.h
+++ b/Engine/Plugins/Enterprise/GLTFExporter/MigratedFromGit/Source/GLTFExporter/Private/Converters/GLTFLevelConverters.h
@@ -9,25 +9,35 @@
 
 class FGLTFSceneComponentConverter final : public TGLTFConverter<FGLTFJsonNodeIndex, const USceneComponent*>
 {
-	FGLTFJsonNodeIndex Add(FGLTFConvertBuilder& Builder, const FString& Name, const USceneComponent* SceneComponent) override;
+	using TGLTFConverter::TGLTFConverter;
+
+	FGLTFJsonNodeIndex Convert(const FString& Name, const USceneComponent* SceneComponent) override;
 };
 
 class FGLTFActorConverter final : public TGLTFConverter<FGLTFJsonNodeIndex, const AActor*>
 {
-	FGLTFJsonNodeIndex Add(FGLTFConvertBuilder& Builder, const FString& Name, const AActor* Actor) override;
+	using TGLTFConverter::TGLTFConverter;
+
+	FGLTFJsonNodeIndex Convert(const FString& Name, const AActor* Actor) override;
 };
 
 class FGLTFLevelConverter final : public TGLTFConverter<FGLTFJsonSceneIndex, const ULevel*>
 {
-	FGLTFJsonSceneIndex Add(FGLTFConvertBuilder& Builder, const FString& Name, const ULevel* Level) override;
+	using TGLTFConverter::TGLTFConverter;
+
+	FGLTFJsonSceneIndex Convert(const FString& Name, const ULevel* Level) override;
 };
 
 class FGLTFCameraComponentConverter final : public TGLTFConverter<FGLTFJsonCameraIndex, const UCameraComponent*>
 {
-	FGLTFJsonCameraIndex Add(FGLTFConvertBuilder& Builder, const FString& Name, const UCameraComponent* CameraComponent) override;
+	using TGLTFConverter::TGLTFConverter;
+
+	FGLTFJsonCameraIndex Convert(const FString& Name, const UCameraComponent* CameraComponent) override;
 };
 
 class FGLTFLightComponentConverter final : public TGLTFConverter<FGLTFJsonLightIndex, const ULightComponent*>
 {
-	FGLTFJsonLightIndex Add(FGLTFConvertBuilder& Builder, const FString& Name, const ULightComponent* LightComponent) override;
+	using TGLTFConverter::TGLTFConverter;
+
+	FGLTFJsonLightIndex Convert(const FString& Name, const ULightComponent* LightComponent) override;
 };
diff --git a/Engine/Plugins/Enterprise/GLTFExporter/MigratedFromGit/Source/GLTFExporter/Private/Converters/GLTFLightMapConverters.cpp b/Engine/Plugins/Enterprise/GLTFExporter/MigratedFromGit/Source/GLTFExporter/Private/Converters/GLTFLightMapConverters.cpp
--- a/Engine/Plugins/Enterprise/GLTFExporter/MigratedFromGit/Source/GLTFExporter/Private/Converters/GLTFLightMapConverters.cpp
+++ b/Engine/Plugins/Enterprise/GLTFExporter/MigratedFromGit/Source/GLTFExporter/Private/Converters/GLTFLightMapConverters.cpp
@@ -5,7 +5,7 @@
 #include "Builders/GLTFContainerBuilder.h"
 #include "Engine/MapBuildDataRegistry.h"
 
-FGLTFJsonLightMapIndex FGLTFLightMapConverter::Add(FGLTFConvertBuilder& Builder, const FString& Name, const UStaticMeshComponent* StaticMeshComponent)
+FGLTFJsonLightMapIndex FGLTFLightMapConverter::Convert(const FString& Name, const UStaticMeshComponent* StaticMeshComponent)
 {
 	const UStaticMesh* StaticMesh = StaticMeshComponent->GetStaticMesh();
 
diff --git a/Engine/Plugins/Enterprise/GLTFExporter/MigratedFromGit/Source/GLTFExporter/Private/Converters/GLTFLightMapConverters.h b/Engine/Plugins/Enterprise/GLTFExporter/MigratedFromGit/Source/GLTFExporter/Private/Converters/GLTFLightMapConverters.h
--- a/Engine/Plugins/Enterprise/GLTFExporter/MigratedFromGit/Source/GLTFExporter/Private/Converters/GLTFLightMapConverters.h
+++ b/Engine/Plugins/Enterprise/GLTFExporter/MigratedFromGit/Source/GLTFExporter/Private/Converters/GLTFLightMapConverters.h
@@ -8,5 +8,7 @@
 
 class FGLTFLightMapConverter final : public TGLTFConverter<FGLTFJsonLightMapIndex, const UStaticMeshComponent*>
 {
-	FGLTFJsonLightMapIndex Add(FGLTFConvertBuilder& Builder, const FString& Name, const UStaticMeshComponent* StaticMeshComponent) override;
+	using TGLTFConverter::TGLTFConverter;
+
+	FGLTFJsonLightMapIndex Convert(const FString& Name, const UStaticMeshComponent* StaticMeshComponent) override;
 };
diff --git a/Engine/Plugins/Enterprise/GLTFExporter/MigratedFromGit/Source/GLTFExporter/Private/Converters/GLTFMaterialConverters.cpp b/Engine/Plugins/Enterprise/GLTFExporter/MigratedFromGit/Source/GLTFExporter/Private/Converters/GLTFMaterialConverters.cpp
--- a/Engine/Plugins/Enterprise/GLTFExporter/MigratedFromGit/Source/GLTFExporter/Private/Converters/GLTFMaterialConverters.cpp
+++ b/Engine/Plugins/Enterprise/GLTFExporter/MigratedFromGit/Source/GLTFExporter/Private/Converters/GLTFMaterialConverters.cpp
@@ -43,7 +43,7 @@ namespace
 	const TArray<FLinearColor> ClearCoatRoughnessInputMasks = { ClearCoatRoughnessMask };
 } // anonymous namespace
 
-FGLTFJsonMaterialIndex FGLTFMaterialConverter::Add(FGLTFConvertBuilder& Builder, const FString& Name, const UMaterialInterface* Material)
+FGLTFJsonMaterialIndex FGLTFMaterialConverter::Convert(const FString& Name, const UMaterialInterface* Material)
 {
 	{
 		const UMaterial* ParentMaterial = Material->GetMaterial();
@@ -67,7 +67,7 @@ FGLTFJsonMaterialIndex FGLTFMaterialConverter::Add(FGLTFConvertBuilder& Builder,
 	JsonMaterial.AlphaCutoff = Material->GetOpacityMaskClipValue();
 	JsonMaterial.DoubleSided = Material->IsTwoSided();
 
-	if (!TryGetShadingModel(Builder, JsonMaterial.ShadingModel, Material))
+	if (!TryGetShadingModel(JsonMaterial.ShadingModel, Material))
 	{
 		JsonMaterial.ShadingModel = EGLTFJsonShadingModel::Default;
 		Builder.AddWarningMessage(FString::Printf(TEXT("Material %s will be exported as shading model %s"), *Material->GetName(), *FGLTFConverterUtility::GetEnumDisplayName(MSM_DefaultLit)));
@@ -84,9 +84,9 @@ FGLTFJsonMaterialIndex FGLTFMaterialConverter::Add(FGLTFConvertBuilder& Builder,
 		{
 			if (!TryGetConstantColor(JsonMaterial.PBRMetallicRoughness.BaseColorFactor, BaseColorProperty, Material))
 			{
-				if (!TryGetSourceTexture(Builder, JsonMaterial.PBRMetallicRoughness.BaseColorTexture, BaseColorProperty, Material, DefaultColorInputMasks))
+				if (!TryGetSourceTexture(JsonMaterial.PBRMetallicRoughness.BaseColorTexture, BaseColorProperty, Material, DefaultColorInputMasks))
 				{
-					if (!TryGetBakedMaterialProperty(Builder, JsonMaterial.PBRMetallicRoughness.BaseColorTexture, JsonMaterial.PBRMetallicRoughness.BaseColorFactor, BaseColorProperty, TEXT("BaseColor"), Material))
+					if (!TryGetBakedMaterialProperty(JsonMaterial.PBRMetallicRoughness.BaseColorTexture, JsonMaterial.PBRMetallicRoughness.BaseColorFactor, BaseColorProperty, TEXT("BaseColor"), Material))
 					{
 						Builder.AddWarningMessage(FString::Printf(TEXT("Failed to export %s for material %s"), FGLTFMaterialUtility::GetPropertyName(BaseColorProperty), *Material->GetName()));
 					}
@@ -97,7 +97,7 @@ FGLTFJsonMaterialIndex FGLTFMaterialConverter::Add(FGLTFConvertBuilder& Builder,
 		}
 		else
 		{
-			if (!TryGetBaseColorAndOpacity(Builder, JsonMaterial.PBRMetallicRoughness, Material, BaseColorProperty, OpacityProperty))
+			if (!TryGetBaseColorAndOpacity(JsonMaterial.PBRMetallicRoughness, Material, BaseColorProperty, OpacityProperty))
 			{
 				Builder.AddWarningMessage(FString::Printf(TEXT("Failed to export %s and %s for material %s"), FGLTFMaterialUtility::GetPropertyName(BaseColorProperty), FGLTFMaterialUtility::GetPropertyName(OpacityProperty), *Material->GetName()));
 			}
@@ -108,13 +108,13 @@ FGLTFJsonMaterialIndex FGLTFMaterialConverter::Add(FGLTFConvertBuilder& Builder,
 			const EMaterialProperty MetallicProperty = MP_Metallic;
 			const EMaterialProperty RoughnessProperty = MP_Roughness;
 
-			if (!TryGetMetallicAndRoughness(Builder, JsonMaterial.PBRMetallicRoughness, Material, MetallicProperty, RoughnessProperty))
+			if (!TryGetMetallicAndRoughness(JsonMaterial.PBRMetallicRoughness, Material, MetallicProperty, RoughnessProperty))
 			{
 				Builder.AddWarningMessage(FString::Printf(TEXT("Failed to export %s and %s for material %s"), FGLTFMaterialUtility::GetPropertyName(MetallicProperty), FGLTFMaterialUtility::GetPropertyName(RoughnessProperty), *Material->GetName()));
 			}
 
 			const EMaterialProperty EmissiveProperty = MP_EmissiveColor;
-			if (!TryGetEmissive(Builder, JsonMaterial, EmissiveProperty, Material))
+			if (!TryGetEmissive(JsonMaterial, EmissiveProperty, Material))
 			{
 				Builder.AddWarningMessage(FString::Printf(TEXT("Failed to export %s for material %s"), FGLTFMaterialUtility::GetPropertyName(EmissiveProperty), *Material->GetName()));
 			}
@@ -123,9 +123,9 @@ FGLTFJsonMaterialIndex FGLTFMaterialConverter::Add(FGLTFConvertBuilder& Builder,
 			const EMaterialProperty NormalProperty = JsonMaterial.ShadingModel == EGLTFJsonShadingModel::ClearCoat ? MP_CustomOutput : MP_Normal;
 			if (IsPropertyNonDefault(NormalProperty, Material))
 			{
-				if (!TryGetSourceTexture(Builder, JsonMaterial.NormalTexture, NormalProperty, Material, DefaultColorInputMasks))
+				if (!TryGetSourceTexture(JsonMaterial.NormalTexture, NormalProperty, Material, DefaultColorInputMasks))
 				{
-					if (!TryGetBakedMaterialProperty(Builder, JsonMaterial.NormalTexture, NormalProperty, TEXT("Normal"), Material))
+					if (!TryGetBakedMaterialProperty(JsonMaterial.NormalTexture, NormalProperty, TEXT("Normal"), Material))
 					{
 						Builder.AddWarningMessage(FString::Printf(TEXT("Failed to export %s for material %s"), FGLTFMaterialUtility::GetPropertyName(NormalProperty), *Material->GetName()));
 					}
@@ -135,9 +135,9 @@ FGLTFJsonMaterialIndex FGLTFMaterialConverter::Add(FGLTFConvertBuilder& Builder,
 			const EMaterialProperty AmbientOcclusionProperty = MP_AmbientOcclusion;
 			if (IsPropertyNonDefault(AmbientOcclusionProperty, Material))
 			{
-				if (!TryGetSourceTexture(Builder, JsonMaterial.OcclusionTexture, AmbientOcclusionProperty, Material, OcclusionInputMasks))
+				if (!TryGetSourceTexture(JsonMaterial.OcclusionTexture, AmbientOcclusionProperty, Material, OcclusionInputMasks))
 				{
-					if (!TryGetBakedMaterialProperty(Builder, JsonMaterial.OcclusionTexture, AmbientOcclusionProperty, TEXT("Occlusion"), Material))
+					if (!TryGetBakedMaterialProperty(JsonMaterial.OcclusionTexture, AmbientOcclusionProperty, TEXT("Occlusion"), Material))
 					{
 						Builder.AddWarningMessage(FString::Printf(TEXT("Failed to export %s for material %s"), FGLTFMaterialUtility::GetPropertyName(AmbientOcclusionProperty), *Material->GetName()));
 					}
@@ -149,7 +149,7 @@ FGLTFJsonMaterialIndex FGLTFMaterialConverter::Add(FGLTFConvertBuilder& Builder,
 				const EMaterialProperty ClearCoatProperty = MP_CustomData0;
 				const EMaterialProperty ClearCoatRoughnessProperty = MP_CustomData1;
 
-				if (!TryGetClearCoatRoughness(Builder, JsonMaterial.ClearCoat, Material, ClearCoatProperty, ClearCoatRoughnessProperty))
+				if (!TryGetClearCoatRoughness(JsonMaterial.ClearCoat, Material, ClearCoatProperty, ClearCoatRoughnessProperty))
 				{
 					Builder.AddWarningMessage(FString::Printf(TEXT("Failed to export %s and %s for material %s"), FGLTFMaterialUtility::GetPropertyName(ClearCoatProperty), FGLTFMaterialUtility::GetPropertyName(ClearCoatRoughnessProperty), *Material->GetName()));
 				}
@@ -157,9 +157,9 @@ FGLTFJsonMaterialIndex FGLTFMaterialConverter::Add(FGLTFConvertBuilder& Builder,
 				const EMaterialProperty ClearCoatNormalProperty = MP_Normal;
 				if (IsPropertyNonDefault(ClearCoatNormalProperty, Material))
 				{
-					if (!TryGetSourceTexture(Builder, JsonMaterial.ClearCoat.ClearCoatNormalTexture, ClearCoatNormalProperty, Material, DefaultColorInputMasks))
+					if (!TryGetSourceTexture(JsonMaterial.ClearCoat.ClearCoatNormalTexture, ClearCoatNormalProperty, Material, DefaultColorInputMasks))
 					{
-						if (!TryGetBakedMaterialProperty(Builder, JsonMaterial.ClearCoat.ClearCoatNormalTexture, ClearCoatNormalProperty, TEXT("ClearCoatNormal"), Material))
+						if (!TryGetBakedMaterialProperty(JsonMaterial.ClearCoat.ClearCoatNormalTexture, ClearCoatNormalProperty, TEXT("ClearCoatNormal"), Material))
 						{
 							Builder.AddWarningMessage(FString::Printf(TEXT("Failed to export %s for material %s"), FGLTFMaterialUtility::GetPropertyName(ClearCoatNormalProperty), *Material->GetName()));
 						}
@@ -172,7 +172,7 @@ FGLTFJsonMaterialIndex FGLTFMaterialConverter::Add(FGLTFConvertBuilder& Builder,
 	return Builder.AddMaterial(JsonMaterial);
 }
 
-bool FGLTFMaterialConverter::TryGetShadingModel(FGLTFConvertBuilder& Builder, EGLTFJsonShadingModel& OutShadingModel, const UMaterialInterface* Material) const
+bool FGLTFMaterialConverter::TryGetShadingModel(EGLTFJsonShadingModel& OutShadingModel, const UMaterialInterface* Material) const
 {
 	const FMaterialShadingModelField ShadingModels = Material->GetShadingModels();
 	const int32 ShadingModelCount = ShadingModels.CountShadingModels();
@@ -216,7 +216,7 @@ bool FGLTFMaterialConverter::TryGetShadingModel(FGLTFConvertBuilder& Builder, EG
 	return true;
 }
 
-bool FGLTFMaterialConverter::TryGetBaseColorAndOpacity(FGLTFConvertBuilder& Builder, FGLTFJsonPBRMetallicRoughness& OutPBRParams, const UMaterialInterface* Material, EMaterialProperty BaseColorProperty, EMaterialProperty OpacityProperty) const
+bool FGLTFMaterialConverter::TryGetBaseColorAndOpacity(FGLTFJsonPBRMetallicRoughness& OutPBRParams, const UMaterialInterface* Material, EMaterialProperty BaseColorProperty, EMaterialProperty OpacityProperty) const
 {
 	const bool bIsBaseColorConstant = TryGetConstantColor(OutPBRParams.BaseColorFactor, BaseColorProperty, Material);
 	const bool bIsOpacityConstant = TryGetConstantScalar(OutPBRParams.BaseColorFactor.A, OpacityProperty, Material);
@@ -318,8 +318,8 @@ bool FGLTFMaterialConverter::TryGetBaseColorAndOpacity(FGLTFConvertBuilder& Buil
 		TextureMagFilter = FGLTFConverterUtility::ConvertMagFilter(OpacityTexture->Filter, OpacityTexture->LODGroup);
 	}
 
-	const FGLTFPropertyBakeOutput BaseColorBakeOutput = BakeMaterialProperty(Builder, BaseColorProperty, Material, BaseColorTexCoord, &TextureSize);
-	const FGLTFPropertyBakeOutput OpacityBakeOutput = BakeMaterialProperty(Builder, OpacityProperty, Material, OpacityTexCoord, &TextureSize, true);
+	const FGLTFPropertyBakeOutput BaseColorBakeOutput = BakeMaterialProperty(BaseColorProperty, Material, BaseColorTexCoord, &TextureSize);
+	const FGLTFPropertyBakeOutput OpacityBakeOutput = BakeMaterialProperty(OpacityProperty, Material, OpacityTexCoord, &TextureSize, true);
 	const float BaseColorScale = BaseColorProperty == MP_EmissiveColor ? BaseColorBakeOutput.EmissiveScale : 1;
 
 	// Detect when both baked properties are constants, which means we can avoid exporting a texture
@@ -380,7 +380,7 @@ bool FGLTFMaterialConverter::TryGetBaseColorAndOpacity(FGLTFConvertBuilder& Buil
 	return true;
 }
 
-bool FGLTFMaterialConverter::TryGetMetallicAndRoughness(FGLTFConvertBuilder& Builder, FGLTFJsonPBRMetallicRoughness& OutPBRParams, const UMaterialInterface* Material, EMaterialProperty MetallicProperty, EMaterialProperty RoughnessProperty) const
+bool FGLTFMaterialConverter::TryGetMetallicAndRoughness(FGLTFJsonPBRMetallicRoughness& OutPBRParams, const UMaterialInterface* Material, EMaterialProperty MetallicProperty, EMaterialProperty RoughnessProperty) const
 {
 	const bool bIsMetallicConstant = TryGetConstantScalar(OutPBRParams.MetallicFactor, MetallicProperty, Material);
 	const bool bIsRoughnessConstant = TryGetConstantScalar(OutPBRParams.RoughnessFactor, RoughnessProperty, Material);
@@ -483,8 +483,8 @@ bool FGLTFMaterialConverter::TryGetMetallicAndRoughness(FGLTFConvertBuilder& Bui
 		TextureMagFilter = FGLTFConverterUtility::ConvertMagFilter(RoughnessTexture->Filter, RoughnessTexture->LODGroup);
 	}
 
-	const FGLTFPropertyBakeOutput MetallicBakeOutput = BakeMaterialProperty(Builder, MetallicProperty, Material, MetallicTexCoord, &TextureSize);
-	const FGLTFPropertyBakeOutput RoughnessBakeOutput = BakeMaterialProperty(Builder, RoughnessProperty, Material, RoughnessTexCoord, &TextureSize);
+	const FGLTFPropertyBakeOutput MetallicBakeOutput = BakeMaterialProperty(MetallicProperty, Material, MetallicTexCoord, &TextureSize);
+	const FGLTFPropertyBakeOutput RoughnessBakeOutput = BakeMaterialProperty(RoughnessProperty, Material, RoughnessTexCoord, &TextureSize);
 
 	// Detect when both baked properties are constants, which means we can use factors and avoid exporting a texture
 	if (MetallicBakeOutput.bIsConstant && RoughnessBakeOutput.bIsConstant)
@@ -541,7 +541,7 @@ bool FGLTFMaterialConverter::TryGetMetallicAndRoughness(FGLTFConvertBuilder& Bui
 	return true;
 }
 
-bool FGLTFMaterialConverter::TryGetClearCoatRoughness(FGLTFConvertBuilder& Builder, FGLTFJsonClearCoatExtension& OutExtParams, const UMaterialInterface* Material, EMaterialProperty IntensityProperty, EMaterialProperty RoughnessProperty) const
+bool FGLTFMaterialConverter::TryGetClearCoatRoughness(FGLTFJsonClearCoatExtension& OutExtParams, const UMaterialInterface* Material, EMaterialProperty IntensityProperty, EMaterialProperty RoughnessProperty) const
 {
 	const bool bIsIntensityConstant = TryGetConstantScalar(OutExtParams.ClearCoatFactor, IntensityProperty, Material);
 	const bool bIsRoughnessConstant = TryGetConstantScalar(OutExtParams.ClearCoatRoughnessFactor, RoughnessProperty, Material);
@@ -647,8 +647,8 @@ bool FGLTFMaterialConverter::TryGetClearCoatRoughness(FGLTFConvertBuilder& Build
 		TextureMagFilter = FGLTFConverterUtility::ConvertMagFilter(RoughnessTexture->Filter, RoughnessTexture->LODGroup);
 	}
 
-	const FGLTFPropertyBakeOutput IntensityBakeOutput = BakeMaterialProperty(Builder, IntensityProperty, Material, IntensityTexCoord, &TextureSize);
-	const FGLTFPropertyBakeOutput RoughnessBakeOutput = BakeMaterialProperty(Builder, RoughnessProperty, Material, RoughnessTexCoord, &TextureSize);
+	const FGLTFPropertyBakeOutput IntensityBakeOutput = BakeMaterialProperty(IntensityProperty, Material, IntensityTexCoord, &TextureSize);
+	const FGLTFPropertyBakeOutput RoughnessBakeOutput = BakeMaterialProperty(RoughnessProperty, Material, RoughnessTexCoord, &TextureSize);
 
 	// Detect when both baked properties are constants, which means we can use factors and avoid exporting a texture
 	if (IntensityBakeOutput.bIsConstant && RoughnessBakeOutput.bIsConstant)
@@ -707,7 +707,7 @@ bool FGLTFMaterialConverter::TryGetClearCoatRoughness(FGLTFConvertBuilder& Build
 	return true;
 }
 
-bool FGLTFMaterialConverter::TryGetEmissive(FGLTFConvertBuilder& Builder, FGLTFJsonMaterial& JsonMaterial, EMaterialProperty EmissiveProperty, const UMaterialInterface* Material) const
+bool FGLTFMaterialConverter::TryGetEmissive(FGLTFJsonMaterial& JsonMaterial, EMaterialProperty EmissiveProperty, const UMaterialInterface* Material) const
 {
 	// TODO: right now we allow EmissiveFactor to be > 1.0 to support very bright emission, although it's not valid according to the glTF standard.
 	// We may want to change this behaviour and store factors above 1.0 using a custom extension instead.
@@ -717,7 +717,7 @@ bool FGLTFMaterialConverter::TryGetEmissive(FGLTFConvertBuilder& Builder, FGLTFJ
 		return true;
 	}
 
-	if (TryGetSourceTexture(Builder, JsonMaterial.EmissiveTexture, EmissiveProperty, Material, DefaultColorInputMasks))
+	if (TryGetSourceTexture(JsonMaterial.EmissiveTexture, EmissiveProperty, Material, DefaultColorInputMasks))
 	{
 		JsonMaterial.EmissiveFactor = FGLTFJsonColor3::White;	// make sure texture is not multiplied with black
 		return true;
@@ -732,7 +732,7 @@ bool FGLTFMaterialConverter::TryGetEmissive(FGLTFConvertBuilder& Builder, FGLTFJ
 		return false;
 	}
 
-	const FGLTFPropertyBakeOutput PropertyBakeOutput = BakeMaterialProperty(Builder, EmissiveProperty, Material, JsonMaterial.EmissiveTexture.TexCoord);
+	const FGLTFPropertyBakeOutput PropertyBakeOutput = BakeMaterialProperty(EmissiveProperty, Material, JsonMaterial.EmissiveTexture.TexCoord);
 	const float EmissiveScale = PropertyBakeOutput.EmissiveScale;
 
 	if (PropertyBakeOutput.bIsConstant)
@@ -742,7 +742,7 @@ bool FGLTFMaterialConverter::TryGetEmissive(FGLTFConvertBuilder& Builder, FGLTFJ
 	}
 	else
 	{
-		if (!StoreBakedPropertyTexture(Builder, JsonMaterial.EmissiveTexture, PropertyBakeOutput, TEXT("Emissive"), Material))
+		if (!StoreBakedPropertyTexture(JsonMaterial.EmissiveTexture, PropertyBakeOutput, TEXT("Emissive"), Material))
 		{
 			return false;
 		}
@@ -1010,7 +1010,7 @@ bool FGLTFMaterialConverter::TryGetConstantScalar(float& OutValue, EMaterialProp
 	return false;
 }
 
-bool FGLTFMaterialConverter::TryGetSourceTexture(FGLTFConvertBuilder& Builder, FGLTFJsonTextureInfo& OutTexInfo, EMaterialProperty Property, const UMaterialInterface* Material, const TArray<FLinearColor>& AllowedMasks) const
+bool FGLTFMaterialConverter::TryGetSourceTexture(FGLTFJsonTextureInfo& OutTexInfo, EMaterialProperty Property, const UMaterialInterface* Material, const TArray<FLinearColor>& AllowedMasks) const
 {
 	const UTexture2D* Texture;
 	int32 TexCoord;
@@ -1116,7 +1116,7 @@ bool FGLTFMaterialConverter::TryGetSourceTexture(const UTexture2D*& OutTexture,
 	return false;
 }
 
-bool FGLTFMaterialConverter::TryGetBakedMaterialProperty(FGLTFConvertBuilder& Builder, FGLTFJsonTextureInfo& OutTexInfo, FGLTFJsonColor3& OutConstant, EMaterialProperty Property, const FString& PropertyName, const UMaterialInterface* Material) const
+bool FGLTFMaterialConverter::TryGetBakedMaterialProperty(FGLTFJsonTextureInfo& OutTexInfo, FGLTFJsonColor3& OutConstant, EMaterialProperty Property, const FString& PropertyName, const UMaterialInterface* Material) const
 {
 	if (!Builder.ExportOptions->bBakeMaterialInputs)
 	{
@@ -1127,7 +1127,7 @@ bool FGLTFMaterialConverter::TryGetBakedMaterialProperty(FGLTFConvertBuilder& Bu
 		return false;
 	}
 
-	const FGLTFPropertyBakeOutput PropertyBakeOutput = BakeMaterialProperty(Builder, Property, Material, OutTexInfo.TexCoord);
+	const FGLTFPropertyBakeOutput PropertyBakeOutput = BakeMaterialProperty(Property, Material, OutTexInfo.TexCoord);
 
 	if (PropertyBakeOutput.bIsConstant)
 	{
@@ -1135,7 +1135,7 @@ bool FGLTFMaterialConverter::TryGetBakedMaterialProperty(FGLTFConvertBuilder& Bu
 		return true;
 	}
 
-	if (StoreBakedPropertyTexture(Builder, OutTexInfo, PropertyBakeOutput, PropertyName, Material))
+	if (StoreBakedPropertyTexture(OutTexInfo, PropertyBakeOutput, PropertyName, Material))
 	{
 		OutConstant = FGLTFJsonColor3::White; // make sure property is not zero
 		return true;
@@ -1144,7 +1144,7 @@ bool FGLTFMaterialConverter::TryGetBakedMaterialProperty(FGLTFConvertBuilder& Bu
 	return false;
 }
 
-bool FGLTFMaterialConverter::TryGetBakedMaterialProperty(FGLTFConvertBuilder& Builder, FGLTFJsonTextureInfo& OutTexInfo, FGLTFJsonColor4& OutConstant, EMaterialProperty Property, const FString& PropertyName, const UMaterialInterface* Material) const
+bool FGLTFMaterialConverter::TryGetBakedMaterialProperty(FGLTFJsonTextureInfo& OutTexInfo, FGLTFJsonColor4& OutConstant, EMaterialProperty Property, const FString& PropertyName, const UMaterialInterface* Material) const
 {
 	if (!Builder.ExportOptions->bBakeMaterialInputs)
 	{
@@ -1155,7 +1155,7 @@ bool FGLTFMaterialConverter::TryGetBakedMaterialProperty(FGLTFConvertBuilder& Bu
 		return false;
 	}
 
-	const FGLTFPropertyBakeOutput PropertyBakeOutput = BakeMaterialProperty(Builder, Property, Material, OutTexInfo.TexCoord);
+	const FGLTFPropertyBakeOutput PropertyBakeOutput = BakeMaterialProperty(Property, Material, OutTexInfo.TexCoord);
 
 	if (PropertyBakeOutput.bIsConstant)
 	{
@@ -1163,7 +1163,7 @@ bool FGLTFMaterialConverter::TryGetBakedMaterialProperty(FGLTFConvertBuilder& Bu
 		return true;
 	}
 
-	if (StoreBakedPropertyTexture(Builder, OutTexInfo, PropertyBakeOutput, PropertyName, Material))
+	if (StoreBakedPropertyTexture(OutTexInfo, PropertyBakeOutput, PropertyName, Material))
 	{
 		OutConstant = FGLTFJsonColor4::White; // make sure property is not zero
 		return true;
@@ -1172,7 +1172,7 @@ bool FGLTFMaterialConverter::TryGetBakedMaterialProperty(FGLTFConvertBuilder& Bu
 	return false;
 }
 
-inline bool FGLTFMaterialConverter::TryGetBakedMaterialProperty(FGLTFConvertBuilder& Builder, FGLTFJsonTextureInfo& OutTexInfo, float& OutConstant, EMaterialProperty Property, const FString& PropertyName, const UMaterialInterface* Material) const
+inline bool FGLTFMaterialConverter::TryGetBakedMaterialProperty(FGLTFJsonTextureInfo& OutTexInfo, float& OutConstant, EMaterialProperty Property, const FString& PropertyName, const UMaterialInterface* Material) const
 {
 	if (!Builder.ExportOptions->bBakeMaterialInputs)
 	{
@@ -1183,7 +1183,7 @@ inline bool FGLTFMaterialConverter::TryGetBakedMaterialProperty(FGLTFConvertBuil
 		return false;
 	}
 
-	const FGLTFPropertyBakeOutput PropertyBakeOutput = BakeMaterialProperty(Builder, Property, Material, OutTexInfo.TexCoord);
+	const FGLTFPropertyBakeOutput PropertyBakeOutput = BakeMaterialProperty(Property, Material, OutTexInfo.TexCoord);
 
 	if (PropertyBakeOutput.bIsConstant)
 	{
@@ -1191,7 +1191,7 @@ inline bool FGLTFMaterialConverter::TryGetBakedMaterialProperty(FGLTFConvertBuil
 		return true;
 	}
 
-	if (StoreBakedPropertyTexture(Builder, OutTexInfo, PropertyBakeOutput, PropertyName, Material))
+	if (StoreBakedPropertyTexture(OutTexInfo, PropertyBakeOutput, PropertyName, Material))
 	{
 		OutConstant = 1; // make sure property is not zero
 		return true;
@@ -1200,7 +1200,7 @@ inline bool FGLTFMaterialConverter::TryGetBakedMaterialProperty(FGLTFConvertBuil
 	return false;
 }
 
-bool FGLTFMaterialConverter::TryGetBakedMaterialProperty(FGLTFConvertBuilder& Builder, FGLTFJsonTextureInfo& OutTexInfo, EMaterialProperty Property, const FString& PropertyName, const UMaterialInterface* Material) const
+bool FGLTFMaterialConverter::TryGetBakedMaterialProperty(FGLTFJsonTextureInfo& OutTexInfo, EMaterialProperty Property, const FString& PropertyName, const UMaterialInterface* Material) const
 {
 	if (!Builder.ExportOptions->bBakeMaterialInputs)
 	{
@@ -1211,11 +1211,11 @@ bool FGLTFMaterialConverter::TryGetBakedMaterialProperty(FGLTFConvertBuilder& Bu
 		return false;
 	}
 
-	const FGLTFPropertyBakeOutput PropertyBakeOutput = BakeMaterialProperty(Builder, Property, Material, OutTexInfo.TexCoord);
+	const FGLTFPropertyBakeOutput PropertyBakeOutput = BakeMaterialProperty(Property, Material, OutTexInfo.TexCoord);
 
 	if (!PropertyBakeOutput.bIsConstant)
 	{
-		return StoreBakedPropertyTexture(Builder, OutTexInfo, PropertyBakeOutput, PropertyName, Material);
+		return StoreBakedPropertyTexture(OutTexInfo, PropertyBakeOutput, PropertyName, Material);
 	}
 
 	const FVector4 MaskedConstant = FVector4(PropertyBakeOutput.ConstantValue) * FGLTFMaterialUtility::GetPropertyMask(Property);
@@ -1255,7 +1255,7 @@ bool FGLTFMaterialConverter::TryGetBakedMaterialProperty(FGLTFConvertBuilder& Bu
 	return true;
 }
 
-FGLTFPropertyBakeOutput FGLTFMaterialConverter::BakeMaterialProperty(FGLTFConvertBuilder& Builder, EMaterialProperty Property, const UMaterialInterface* Material, int32& OutTexCoord, const FIntPoint* PreferredTextureSize, bool bCopyAlphaFromRedChannel) const
+FGLTFPropertyBakeOutput FGLTFMaterialConverter::BakeMaterialProperty(EMaterialProperty Property, const UMaterialInterface* Material, int32& OutTexCoord, const FIntPoint* PreferredTextureSize, bool bCopyAlphaFromRedChannel) const
 {
 	const FExpressionInput* PropertyInput = FGLTFMaterialUtility::GetInputForProperty(Material, Property);
 	TSet<int32> TexCoords;
@@ -1306,7 +1306,7 @@ FGLTFPropertyBakeOutput FGLTFMaterialConverter::BakeMaterialProperty(FGLTFConver
 	return PropertyBakeOutput;
 }
 
-bool FGLTFMaterialConverter::StoreBakedPropertyTexture(FGLTFConvertBuilder& Builder, FGLTFJsonTextureInfo& OutTexInfo, const FGLTFPropertyBakeOutput& PropertyBakeOutput, const FString& PropertyName, const UMaterialInterface* Material) const
+bool FGLTFMaterialConverter::StoreBakedPropertyTexture(FGLTFJsonTextureInfo& OutTexInfo, const FGLTFPropertyBakeOutput& PropertyBakeOutput, const FString& PropertyName, const UMaterialInterface* Material) const
 {
 	const FString TextureName = Material->GetName() + TEXT("_") + PropertyName;
 
diff --git a/Engine/Plugins/Enterprise/GLTFExporter/MigratedFromGit/Source/GLTFExporter/Private/Converters/GLTFMaterialConverters.h b/Engine/Plugins/Enterprise/GLTFExporter/MigratedFromGit/Source/GLTFExporter/Private/Converters/GLTFMaterialConverters.h
--- a/Engine/Plugins/Enterprise/GLTFExporter/MigratedFromGit/Source/GLTFExporter/Private/Converters/GLTFMaterialConverters.h
+++ b/Engine/Plugins/Enterprise/GLTFExporter/MigratedFromGit/Source/GLTFExporter/Private/Converters/GLTFMaterialConverters.h
@@ -11,14 +11,16 @@ struct FGLTFPropertyBakeOutput;
 
 class FGLTFMaterialConverter final : public TGLTFConverter<FGLTFJsonMaterialIndex, const UMaterialInterface*>
 {
-	FGLTFJsonMaterialIndex Add(FGLTFConvertBuilder& Builder, const FString& Name, const UMaterialInterface* Material) override;
+	using TGLTFConverter::TGLTFConverter;
 
-	bool TryGetShadingModel(FGLTFConvertBuilder& Builder, EGLTFJsonShadingModel& ShadingModel, const UMaterialInterface* Material) const;
+	FGLTFJsonMaterialIndex Convert(const FString& Name, const UMaterialInterface* Material) override;
 
-	bool TryGetBaseColorAndOpacity(FGLTFConvertBuilder& Builder, FGLTFJsonPBRMetallicRoughness& OutPBRParams, const UMaterialInterface* Material, EMaterialProperty BaseColorProperty, EMaterialProperty OpacityProperty) const;
-	bool TryGetMetallicAndRoughness(FGLTFConvertBuilder& Builder, FGLTFJsonPBRMetallicRoughness& OutPBRParams, const UMaterialInterface* Material, EMaterialProperty MetallicProperty, EMaterialProperty RoughnessProperty) const;
-	bool TryGetClearCoatRoughness(FGLTFConvertBuilder& Builder, FGLTFJsonClearCoatExtension& OutExtParams, const UMaterialInterface* Material, EMaterialProperty IntensityProperty, EMaterialProperty RoughnessProperty) const;
-	bool TryGetEmissive(FGLTFConvertBuilder& Builder, FGLTFJsonMaterial& JsonMaterial, EMaterialProperty EmissiveProperty, const UMaterialInterface* Material) const;
+	bool TryGetShadingModel(EGLTFJsonShadingModel& ShadingModel, const UMaterialInterface* Material) const;
+
+	bool TryGetBaseColorAndOpacity(FGLTFJsonPBRMetallicRoughness& OutPBRParams, const UMaterialInterface* Material, EMaterialProperty BaseColorProperty, EMaterialProperty OpacityProperty) const;
+	bool TryGetMetallicAndRoughness(FGLTFJsonPBRMetallicRoughness& OutPBRParams, const UMaterialInterface* Material, EMaterialProperty MetallicProperty, EMaterialProperty RoughnessProperty) const;
+	bool TryGetClearCoatRoughness(FGLTFJsonClearCoatExtension& OutExtParams, const UMaterialInterface* Material, EMaterialProperty IntensityProperty, EMaterialProperty RoughnessProperty) const;
+	bool TryGetEmissive(FGLTFJsonMaterial& JsonMaterial, EMaterialProperty EmissiveProperty, const UMaterialInterface* Material) const;
 
 	bool IsPropertyNonDefault(EMaterialProperty Property, const UMaterialInterface* Material) const;
 	bool TryGetConstantColor(FGLTFJsonColor3& OutValue, EMaterialProperty Property, const UMaterialInterface* Material) const;
@@ -26,15 +28,15 @@ class FGLTFMaterialConverter final : public TGLTFConverter<FGLTFJsonMaterialInde
 	bool TryGetConstantColor(FLinearColor& OutValue, EMaterialProperty Property, const UMaterialInterface* Material) const;
 	bool TryGetConstantScalar(float& OutValue, EMaterialProperty Property, const UMaterialInterface* Material) const;
 
-	bool TryGetSourceTexture(FGLTFConvertBuilder& Builder, FGLTFJsonTextureInfo& OutTexInfo, EMaterialProperty Property, const UMaterialInterface* Material, const TArray<FLinearColor>& AllowedMasks = {}) const;
+	bool TryGetSourceTexture(FGLTFJsonTextureInfo& OutTexInfo, EMaterialProperty Property, const UMaterialInterface* Material, const TArray<FLinearColor>& AllowedMasks = {}) const;
 	bool TryGetSourceTexture(const UTexture2D*& OutTexture, int32& OutTexCoord, EMaterialProperty Property, const UMaterialInterface* Material, const TArray<FLinearColor>& AllowedMasks = {}) const;
 
-	bool TryGetBakedMaterialProperty(FGLTFConvertBuilder& Builder, FGLTFJsonTextureInfo& OutTexInfo, FGLTFJsonColor3& OutConstant, EMaterialProperty Property, const FString& PropertyName, const UMaterialInterface* Material) const;
-	bool TryGetBakedMaterialProperty(FGLTFConvertBuilder& Builder, FGLTFJsonTextureInfo& OutTexInfo, FGLTFJsonColor4& OutConstant, EMaterialProperty Property, const FString& PropertyName, const UMaterialInterface* Material) const;
-	bool TryGetBakedMaterialProperty(FGLTFConvertBuilder& Builder, FGLTFJsonTextureInfo& OutTexInfo, float& OutConstant, EMaterialProperty Property, const FString& PropertyName, const UMaterialInterface* Material) const;
-	bool TryGetBakedMaterialProperty(FGLTFConvertBuilder& Builder, FGLTFJsonTextureInfo& OutTexInfo, EMaterialProperty Property, const FString& PropertyName, const UMaterialInterface* Material) const;
+	bool TryGetBakedMaterialProperty(FGLTFJsonTextureInfo& OutTexInfo, FGLTFJsonColor3& OutConstant, EMaterialProperty Property, const FString& PropertyName, const UMaterialInterface* Material) const;
+	bool TryGetBakedMaterialProperty(FGLTFJsonTextureInfo& OutTexInfo, FGLTFJsonColor4& OutConstant, EMaterialProperty Property, const FString& PropertyName, const UMaterialInterface* Material) const;
+	bool TryGetBakedMaterialProperty(FGLTFJsonTextureInfo& OutTexInfo, float& OutConstant, EMaterialProperty Property, const FString& PropertyName, const UMaterialInterface* Material) const;
+	bool TryGetBakedMaterialProperty(FGLTFJsonTextureInfo& OutTexInfo, EMaterialProperty Property, const FString& PropertyName, const UMaterialInterface* Material) const;
 
-	FGLTFPropertyBakeOutput BakeMaterialProperty(FGLTFConvertBuilder& Builder, EMaterialProperty Property, const UMaterialInterface* Material, int32& OutTexCoord, const FIntPoint* PreferredTextureSize = nullptr, bool bCopyAlphaFromRedChannel = false) const;
+	FGLTFPropertyBakeOutput BakeMaterialProperty(EMaterialProperty Property, const UMaterialInterface* Material, int32& OutTexCoord, const FIntPoint* PreferredTextureSize = nullptr, bool bCopyAlphaFromRedChannel = false) const;
 
-	bool StoreBakedPropertyTexture(FGLTFConvertBuilder& Builder, FGLTFJsonTextureInfo& OutTexInfo, const FGLTFPropertyBakeOutput& PropertyBakeOutput, const FString& PropertyName, const UMaterialInterface* Material) const;
+	bool StoreBakedPropertyTexture(FGLTFJsonTextureInfo& OutTexInfo, const FGLTFPropertyBakeOutput& PropertyBakeOutput, const FString& PropertyName, const UMaterialInterface* Material) const;
 };
diff --git a/Engine/Plugins/Enterprise/GLTFExporter/MigratedFromGit/Source/GLTFExporter/Private/Converters/GLTFSkeletalMeshConverters.cpp b/Engine/Plugins/Enterprise/GLTFExporter/MigratedFromGit/Source/GLTFExporter/Private/Converters/GLTFSkeletalMeshConverters.cpp
--- a/Engine/Plugins/Enterprise/GLTFExporter/MigratedFromGit/Source/GLTFExporter/Private/Converters/GLTFSkeletalMeshConverters.cpp
+++ b/Engine/Plugins/Enterprise/GLTFExporter/MigratedFromGit/Source/GLTFExporter/Private/Converters/GLTFSkeletalMeshConverters.cpp
@@ -6,7 +6,7 @@
 #include "Rendering/MultiSizeIndexContainer.h"
 #include "Rendering/SkeletalMeshRenderData.h"
 
-FGLTFJsonBufferViewIndex FGLTFIndexContainerConverter::Add(FGLTFConvertBuilder& Builder, const FString& Name, const FMultiSizeIndexContainer* IndexContainer)
+FGLTFJsonBufferViewIndex FGLTFIndexContainerConverter::Convert(const FString& Name, const FMultiSizeIndexContainer* IndexContainer)
 {
 	const FRawStaticIndexBuffer16or32Interface* IndexBuffer = IndexContainer->GetIndexBuffer();
 
@@ -36,7 +36,7 @@ FGLTFJsonBufferViewIndex FGLTFIndexContainerConverter::Add(FGLTFConvertBuilder&
 	}
 }
 
-FGLTFJsonAccessorIndex FGLTFSkeletalMeshSectionConverter::Add(FGLTFConvertBuilder& Builder, const FString& Name, const FSkelMeshRenderSection* MeshSection, const FMultiSizeIndexContainer* IndexContainer)
+FGLTFJsonAccessorIndex FGLTFSkeletalMeshSectionConverter::Convert(const FString& Name, const FSkelMeshRenderSection* MeshSection, const FMultiSizeIndexContainer* IndexContainer)
 {
 	const uint32 TriangleCount = MeshSection->NumTriangles;
 	if (TriangleCount == 0)
@@ -58,7 +58,7 @@ FGLTFJsonAccessorIndex FGLTFSkeletalMeshSectionConverter::Add(FGLTFConvertBuilde
 	return Builder.AddAccessor(JsonAccessor);
 }
 
-FGLTFJsonMeshIndex FGLTFSkeletalMeshConverter::Add(FGLTFConvertBuilder& Builder, const FString& Name, const USkeletalMesh* SkeletalMesh, int32 LODIndex, const FColorVertexBuffer* OverrideVertexColors, const FSkinWeightVertexBuffer* OverrideSkinWeights, FGLTFMaterialArray OverrideMaterials)
+FGLTFJsonMeshIndex FGLTFSkeletalMeshConverter::Convert(const FString& Name, const USkeletalMesh* SkeletalMesh, int32 LODIndex, const FColorVertexBuffer* OverrideVertexColors, const FSkinWeightVertexBuffer* OverrideSkinWeights, FGLTFMaterialArray OverrideMaterials)
 {
 	const FSkeletalMeshRenderData* RenderData = SkeletalMesh->GetResourceForRendering();
 
diff --git a/Engine/Plugins/Enterprise/GLTFExporter/MigratedFromGit/Source/GLTFExporter/Private/Converters/GLTFSkeletalMeshConverters.h b/Engine/Plugins/Enterprise/GLTFExporter/MigratedFromGit/Source/GLTFExporter/Private/Converters/GLTFSkeletalMeshConverters.h
--- a/Engine/Plugins/Enterprise/GLTFExporter/MigratedFromGit/Source/GLTFExporter/Private/Converters/GLTFSkeletalMeshConverters.h
+++ b/Engine/Plugins/Enterprise/GLTFExporter/MigratedFromGit/Source/GLTFExporter/Private/Converters/GLTFSkeletalMeshConverters.h
@@ -11,15 +11,21 @@ class FMultiSizeIndexContainer;
 
 class FGLTFIndexContainerConverter final : public TGLTFConverter<FGLTFJsonBufferViewIndex, const FMultiSizeIndexContainer*>
 {
-	FGLTFJsonBufferViewIndex Add(FGLTFConvertBuilder& Builder, const FString& Name, const FMultiSizeIndexContainer* IndexContainer) override;
+	using TGLTFConverter::TGLTFConverter;
+
+	FGLTFJsonBufferViewIndex Convert(const FString& Name, const FMultiSizeIndexContainer* IndexContainer) override;
 };
 
 class FGLTFSkeletalMeshSectionConverter final : public TGLTFConverter<FGLTFJsonAccessorIndex, const FSkelMeshRenderSection*, const FMultiSizeIndexContainer*>
 {
-	FGLTFJsonAccessorIndex Add(FGLTFConvertBuilder& Builder, const FString& Name, const FSkelMeshRenderSection* MeshSection, const FMultiSizeIndexContainer* IndexContainer) override;
+	using TGLTFConverter::TGLTFConverter;
+
+	FGLTFJsonAccessorIndex Convert(const FString& Name, const FSkelMeshRenderSection* MeshSection, const FMultiSizeIndexContainer* IndexContainer) override;
 };
 
 class FGLTFSkeletalMeshConverter final : public TGLTFConverter<FGLTFJsonMeshIndex, const USkeletalMesh*, int32, const FColorVertexBuffer*, const FSkinWeightVertexBuffer*, FGLTFMaterialArray>
 {
-	FGLTFJsonMeshIndex Add(FGLTFConvertBuilder& Builder, const FString& Name, const USkeletalMesh* SkeletalMesh, int32 LODIndex, const FColorVertexBuffer* OverrideVertexColors, const FSkinWeightVertexBuffer* OverrideSkinWeights, FGLTFMaterialArray OverrideMaterials) override;
+	using TGLTFConverter::TGLTFConverter;
+
+	FGLTFJsonMeshIndex Convert(const FString& Name, const USkeletalMesh* SkeletalMesh, int32 LODIndex, const FColorVertexBuffer* OverrideVertexColors, const FSkinWeightVertexBuffer* OverrideSkinWeights, FGLTFMaterialArray OverrideMaterials) override;
 };
diff --git a/Engine/Plugins/Enterprise/GLTFExporter/MigratedFromGit/Source/GLTFExporter/Private/Converters/GLTFStaticMeshConverters.cpp b/Engine/Plugins/Enterprise/GLTFExporter/MigratedFromGit/Source/GLTFExporter/Private/Converters/GLTFStaticMeshConverters.cpp
--- a/Engine/Plugins/Enterprise/GLTFExporter/MigratedFromGit/Source/GLTFExporter/Private/Converters/GLTFStaticMeshConverters.cpp
+++ b/Engine/Plugins/Enterprise/GLTFExporter/MigratedFromGit/Source/GLTFExporter/Private/Converters/GLTFStaticMeshConverters.cpp
@@ -4,7 +4,7 @@
 #include "Converters/GLTFConverterUtility.h"
 #include "Builders/GLTFConvertBuilder.h"
 
-FGLTFJsonBufferViewIndex FGLTFIndexBufferConverter::Add(FGLTFConvertBuilder& Builder, const FString& Name, const FRawStaticIndexBuffer* IndexBuffer)
+FGLTFJsonBufferViewIndex FGLTFIndexBufferConverter::Convert(const FString& Name, const FRawStaticIndexBuffer* IndexBuffer)
 {
 	if (IndexBuffer->GetNumIndices() == 0)
 	{
@@ -25,7 +25,7 @@ FGLTFJsonBufferViewIndex FGLTFIndexBufferConverter::Add(FGLTFConvertBuilder& Bui
 	}
 }
 
-FGLTFJsonAccessorIndex FGLTFStaticMeshSectionConverter::Add(FGLTFConvertBuilder& Builder, const FString& Name, const FStaticMeshSection* MeshSection, const FRawStaticIndexBuffer* IndexBuffer)
+FGLTFJsonAccessorIndex FGLTFStaticMeshSectionConverter::Convert(const FString& Name, const FStaticMeshSection* MeshSection, const FRawStaticIndexBuffer* IndexBuffer)
 {
 	const uint32 TriangleCount = MeshSection->NumTriangles;
 	if (TriangleCount == 0)
@@ -47,7 +47,7 @@ FGLTFJsonAccessorIndex FGLTFStaticMeshSectionConverter::Add(FGLTFConvertBuilder&
 	return Builder.AddAccessor(JsonAccessor);
 }
 
-FGLTFJsonMeshIndex FGLTFStaticMeshConverter::Add(FGLTFConvertBuilder& Builder, const FString& Name, const UStaticMesh* StaticMesh, int32 LODIndex, const FColorVertexBuffer* OverrideVertexColors, FGLTFMaterialArray OverrideMaterials)
+FGLTFJsonMeshIndex FGLTFStaticMeshConverter::Convert(const FString& Name, const UStaticMesh* StaticMesh, int32 LODIndex, const FColorVertexBuffer* OverrideVertexColors, FGLTFMaterialArray OverrideMaterials)
 {
 	if (LODIndex < 0 || StaticMesh->GetNumLODs() <= LODIndex)
 	{
diff --git a/Engine/Plugins/Enterprise/GLTFExporter/MigratedFromGit/Source/GLTFExporter/Private/Converters/GLTFStaticMeshConverters.h b/Engine/Plugins/Enterprise/GLTFExporter/MigratedFromGit/Source/GLTFExporter/Private/Converters/GLTFStaticMeshConverters.h
--- a/Engine/Plugins/Enterprise/GLTFExporter/MigratedFromGit/Source/GLTFExporter/Private/Converters/GLTFStaticMeshConverters.h
+++ b/Engine/Plugins/Enterprise/GLTFExporter/MigratedFromGit/Source/GLTFExporter/Private/Converters/GLTFStaticMeshConverters.h
@@ -9,15 +9,21 @@
 
 class FGLTFIndexBufferConverter final : public TGLTFConverter<FGLTFJsonBufferViewIndex, const FRawStaticIndexBuffer*>
 {
-	FGLTFJsonBufferViewIndex Add(FGLTFConvertBuilder& Builder, const FString& Name, const FRawStaticIndexBuffer* IndexBuffer) override;
+	using TGLTFConverter::TGLTFConverter;
+
+	FGLTFJsonBufferViewIndex Convert(const FString& Name, const FRawStaticIndexBuffer* IndexBuffer) override;
 };
 
 class FGLTFStaticMeshSectionConverter final : public TGLTFConverter<FGLTFJsonAccessorIndex, const FStaticMeshSection*, const FRawStaticIndexBuffer*>
 {
-	FGLTFJsonAccessorIndex Add(FGLTFConvertBuilder& Builder, const FString& Name, const FStaticMeshSection* MeshSection, const FRawStaticIndexBuffer* IndexBuffer) override;
+	using TGLTFConverter::TGLTFConverter;
+
+	FGLTFJsonAccessorIndex Convert(const FString& Name, const FStaticMeshSection* MeshSection, const FRawStaticIndexBuffer* IndexBuffer) override;
 };
 
 class FGLTFStaticMeshConverter final : public TGLTFConverter<FGLTFJsonMeshIndex, const UStaticMesh*, int32, const FColorVertexBuffer*, FGLTFMaterialArray>
 {
-	FGLTFJsonMeshIndex Add(FGLTFConvertBuilder& Builder, const FString& Name, const UStaticMesh* StaticMesh, int32 LODIndex, const FColorVertexBuffer* OverrideVertexColors, FGLTFMaterialArray OverrideMaterials) override;
+	using TGLTFConverter::TGLTFConverter;
+
+	FGLTFJsonMeshIndex Convert(const FString& Name, const UStaticMesh* StaticMesh, int32 LODIndex, const FColorVertexBuffer* OverrideVertexColors, FGLTFMaterialArray OverrideMaterials) override;
 };
diff --git a/Engine/Plugins/Enterprise/GLTFExporter/MigratedFromGit/Source/GLTFExporter/Private/Converters/GLTFTextureConverters.cpp b/Engine/Plugins/Enterprise/GLTFExporter/MigratedFromGit/Source/GLTFExporter/Private/Converters/GLTFTextureConverters.cpp
--- a/Engine/Plugins/Enterprise/GLTFExporter/MigratedFromGit/Source/GLTFExporter/Private/Converters/GLTFTextureConverters.cpp
+++ b/Engine/Plugins/Enterprise/GLTFExporter/MigratedFromGit/Source/GLTFExporter/Private/Converters/GLTFTextureConverters.cpp
@@ -5,7 +5,7 @@
 #include "Converters/GLTFConverterUtility.h"
 #include "Converters/GLTFTextureUtility.h"
 
-FGLTFJsonSamplerIndex FGLTFTextureSamplerConverter::Add(FGLTFConvertBuilder& Builder, const FString& Name, const UTexture* Texture)
+FGLTFJsonSamplerIndex FGLTFTextureSamplerConverter::Convert(const FString& Name, const UTexture* Texture)
 {
 	// TODO: maybe we should reuse existing samplers?
 
@@ -41,7 +41,7 @@ FGLTFJsonSamplerIndex FGLTFTextureSamplerConverter::Add(FGLTFConvertBuilder& Bui
 	return Builder.AddSampler(JsonSampler);
 }
 
-FGLTFJsonTextureIndex FGLTFTexture2DConverter::Add(FGLTFConvertBuilder& Builder, const FString& Name, const UTexture2D* Texture2D)
+FGLTFJsonTextureIndex FGLTFTexture2DConverter::Convert(const FString& Name, const UTexture2D* Texture2D)
 {
 	FGLTFJsonImageIndex ImageIndex;
 
@@ -122,7 +122,7 @@ FGLTFJsonTextureIndex FGLTFTexture2DConverter::Add(FGLTFConvertBuilder& Builder,
 	return Builder.AddTexture(JsonTexture);
 }
 
-FGLTFJsonTextureIndex FGLTFTextureCubeConverter::Add(FGLTFConvertBuilder& Builder, const FString& Name, const UTextureCube* TextureCube, ECubeFace CubeFace)
+FGLTFJsonTextureIndex FGLTFTextureCubeConverter::Convert(const FString& Name, const UTextureCube* TextureCube, ECubeFace CubeFace)
 {
 	FGLTFJsonTexture JsonTexture;
 	JsonTexture.Name = Name;
@@ -156,7 +156,7 @@ FGLTFJsonTextureIndex FGLTFTextureCubeConverter::Add(FGLTFConvertBuilder& Builde
 	return Builder.AddTexture(JsonTexture);
 }
 
-FGLTFJsonTextureIndex FGLTFTextureRenderTarget2DConverter::Add(FGLTFConvertBuilder& Builder, const FString& Name, const UTextureRenderTarget2D* RenderTarget2D)
+FGLTFJsonTextureIndex FGLTFTextureRenderTarget2DConverter::Convert(const FString& Name, const UTextureRenderTarget2D* RenderTarget2D)
 {
 	FGLTFJsonTexture JsonTexture;
 	JsonTexture.Name = Name;
@@ -174,7 +174,7 @@ FGLTFJsonTextureIndex FGLTFTextureRenderTarget2DConverter::Add(FGLTFConvertBuild
 	return Builder.AddTexture(JsonTexture);
 }
 
-FGLTFJsonTextureIndex FGLTFTextureRenderTargetCubeConverter::Add(FGLTFConvertBuilder& Builder, const FString& Name, const UTextureRenderTargetCube* RenderTargetCube, ECubeFace CubeFace)
+FGLTFJsonTextureIndex FGLTFTextureRenderTargetCubeConverter::Convert(const FString& Name, const UTextureRenderTargetCube* RenderTargetCube, ECubeFace CubeFace)
 {
 	FGLTFJsonTexture JsonTexture;
 	JsonTexture.Name = Name;
diff --git a/Engine/Plugins/Enterprise/GLTFExporter/MigratedFromGit/Source/GLTFExporter/Private/Converters/GLTFTextureConverters.h b/Engine/Plugins/Enterprise/GLTFExporter/MigratedFromGit/Source/GLTFExporter/Private/Converters/GLTFTextureConverters.h
--- a/Engine/Plugins/Enterprise/GLTFExporter/MigratedFromGit/Source/GLTFExporter/Private/Converters/GLTFTextureConverters.h
+++ b/Engine/Plugins/Enterprise/GLTFExporter/MigratedFromGit/Source/GLTFExporter/Private/Converters/GLTFTextureConverters.h
@@ -8,25 +8,35 @@
 
 class FGLTFTextureSamplerConverter final : public TGLTFConverter<FGLTFJsonSamplerIndex, const UTexture*>
 {
-	FGLTFJsonSamplerIndex Add(FGLTFConvertBuilder& Builder, const FString& Name, const UTexture* Texture) override;
+	using TGLTFConverter::TGLTFConverter;
+
+	FGLTFJsonSamplerIndex Convert(const FString& Name, const UTexture* Texture) override;
 };
 
 class FGLTFTexture2DConverter final : public TGLTFConverter<FGLTFJsonTextureIndex, const UTexture2D*>
 {
-	FGLTFJsonTextureIndex Add(FGLTFConvertBuilder& Builder, const FString& Name, const UTexture2D* Texture2D) override;
+	using TGLTFConverter::TGLTFConverter;
+
+	FGLTFJsonTextureIndex Convert(const FString& Name, const UTexture2D* Texture2D) override;
 };
 
 class FGLTFTextureCubeConverter final : public TGLTFConverter<FGLTFJsonTextureIndex, const UTextureCube*, ECubeFace>
 {
-	FGLTFJsonTextureIndex Add(FGLTFConvertBuilder& Builder, const FString& Name, const UTextureCube* TextureCube, ECubeFace CubeFace) override;
+	using TGLTFConverter::TGLTFConverter;
+
+	FGLTFJsonTextureIndex Convert(const FString& Name, const UTextureCube* TextureCube, ECubeFace CubeFace) override;
 };
 
 class FGLTFTextureRenderTarget2DConverter final : public TGLTFConverter<FGLTFJsonTextureIndex, const UTextureRenderTarget2D*>
 {
-	FGLTFJsonTextureIndex Add(FGLTFConvertBuilder& Builder, const FString& Name, const UTextureRenderTarget2D* RenderTarget2D) override;
+	using TGLTFConverter::TGLTFConverter;
+
+	FGLTFJsonTextureIndex Convert(const FString& Name, const UTextureRenderTarget2D* RenderTarget2D) override;
 };
 
 class FGLTFTextureRenderTargetCubeConverter final : public TGLTFConverter<FGLTFJsonTextureIndex, const UTextureRenderTargetCube*, ECubeFace>
 {
-	FGLTFJsonTextureIndex Add(FGLTFConvertBuilder& Builder, const FString& Name, const UTextureRenderTargetCube* RenderTargetCube, ECubeFace CubeFace) override;
+	using TGLTFConverter::TGLTFConverter;
+
+	FGLTFJsonTextureIndex Convert(const FString& Name, const UTextureRenderTargetCube* RenderTargetCube, ECubeFace CubeFace) override;
 };
diff --git a/Engine/Plugins/Enterprise/GLTFExporter/MigratedFromGit/Source/GLTFExporter/Private/Converters/GLTFVariantSetConverters.cpp b/Engine/Plugins/Enterprise/GLTFExporter/MigratedFromGit/Source/GLTFExporter/Private/Converters/GLTFVariantSetConverters.cpp
--- a/Engine/Plugins/Enterprise/GLTFExporter/MigratedFromGit/Source/GLTFExporter/Private/Converters/GLTFVariantSetConverters.cpp
+++ b/Engine/Plugins/Enterprise/GLTFExporter/MigratedFromGit/Source/GLTFExporter/Private/Converters/GLTFVariantSetConverters.cpp
@@ -23,7 +23,7 @@ namespace
 	};
 } // anonymous namespace
 
-FGLTFJsonLevelVariantSetsIndex FGLTFLevelVariantSetsConverter::Add(FGLTFConvertBuilder& Builder, const FString& Name, const ALevelVariantSetsActor* LevelVariantSetsActor)
+FGLTFJsonLevelVariantSetsIndex FGLTFLevelVariantSetsConverter::Convert(const FString& Name, const ALevelVariantSetsActor* LevelVariantSetsActor)
 {
 	const ULevelVariantSets* LevelVariantSets = const_cast<ALevelVariantSetsActor*>(LevelVariantSetsActor)->GetLevelVariantSets(true);
 	if (LevelVariantSets == nullptr)
@@ -42,7 +42,7 @@ FGLTFJsonLevelVariantSetsIndex FGLTFLevelVariantSetsConverter::Add(FGLTFConvertB
 		for (const UVariant* Variant: VariantSet->GetVariants())
 		{
 			FGLTFJsonVariant JsonVariant;
-			if (TryParseVariant(Builder, JsonVariant, Variant))
+			if (TryParseVariant(JsonVariant, Variant))
 			{
 				JsonVariantSet.Variants.Add(JsonVariant);
 			}
@@ -68,13 +68,13 @@ FGLTFJsonLevelVariantSetsIndex FGLTFLevelVariantSetsConverter::Add(FGLTFConvertB
 	return Builder.AddLevelVariantSets(JsonLevelVariantSets);
 }
 
-bool FGLTFLevelVariantSetsConverter::TryParseVariant(FGLTFConvertBuilder& Builder, FGLTFJsonVariant& OutVariant, const UVariant* Variant) const
+bool FGLTFLevelVariantSetsConverter::TryParseVariant(FGLTFJsonVariant& OutVariant, const UVariant* Variant) const
 {
 	FGLTFJsonVariant JsonVariant;
 
 	for (const UVariantObjectBinding* Binding: Variant->GetBindings())
 	{
-		TryParseVariantBinding(Builder, JsonVariant, Binding);
+		TryParseVariantBinding(JsonVariant, Binding);
 	}
 
 	if (JsonVariant.Nodes.Num() == 0)
@@ -98,7 +98,7 @@ bool FGLTFLevelVariantSetsConverter::TryParseVariant(FGLTFConvertBuilder& Builde
 	return true;
 }
 
-bool FGLTFLevelVariantSetsConverter::TryParseVariantBinding(FGLTFConvertBuilder& Builder, FGLTFJsonVariant& OutVariant, const UVariantObjectBinding* Binding) const
+bool FGLTFLevelVariantSetsConverter::TryParseVariantBinding(FGLTFJsonVariant& OutVariant, const UVariantObjectBinding* Binding) const
 {
 	bool bHasParsedAnyProperty = false;
 
@@ -122,28 +122,28 @@ bool FGLTFLevelVariantSetsConverter::TryParseVariantBinding(FGLTFConvertBuilder&
 
 		if (bIsVisibilityProperty)
 		{
-			if (TryParseVisibilityPropertyValue(Builder, OutVariant, Property))
+			if (TryParseVisibilityPropertyValue(OutVariant, Property))
 			{
 				bHasParsedAnyProperty = true;
 			}
 		}
 		else if (bIsMaterialProperty)
 		{
-			if (TryParseMaterialPropertyValue(Builder, OutVariant, Property))
+			if (TryParseMaterialPropertyValue(OutVariant, Property))
 			{
 				bHasParsedAnyProperty = true;
 			}
 		}
 		else if (bIsStaticMeshProperty)
 		{
-			if (TryParseStaticMeshPropertyValue(Builder, OutVariant, Property))
+			if (TryParseStaticMeshPropertyValue(OutVariant, Property))
 			{
 				bHasParsedAnyProperty = true;
 			}
 		}
 		else if (bIsSkeletalMeshProperty)
 		{
-			if (TryParseSkeletalMeshPropertyValue(Builder, OutVariant, Property))
+			if (TryParseSkeletalMeshPropertyValue(OutVariant, Property))
 			{
 				bHasParsedAnyProperty = true;
 			}
@@ -168,7 +168,7 @@ bool FGLTFLevelVariantSetsConverter::TryParseVariantBinding(FGLTFConvertBuilder&
 	return bHasParsedAnyProperty;
 }
 
-bool FGLTFLevelVariantSetsConverter::TryParseVisibilityPropertyValue(FGLTFConvertBuilder& Builder, FGLTFJsonVariant& OutVariant, const UPropertyValue* Property) const
+bool FGLTFLevelVariantSetsConverter::TryParseVisibilityPropertyValue(FGLTFJsonVariant& OutVariant, const UPropertyValue* Property) const
 {
 	const USceneComponent* Target = static_cast<USceneComponent*>(Property->GetPropertyParentContainerAddress());
 	if (Target == nullptr)
@@ -215,7 +215,7 @@ bool FGLTFLevelVariantSetsConverter::TryParseVisibilityPropertyValue(FGLTFConver
 	return true;
 }
 
-bool FGLTFLevelVariantSetsConverter::TryParseMaterialPropertyValue(FGLTFConvertBuilder& Builder, FGLTFJsonVariant& OutVariant, const UPropertyValue* Property) const
+bool FGLTFLevelVariantSetsConverter::TryParseMaterialPropertyValue(FGLTFJsonVariant& OutVariant, const UPropertyValue* Property) const
 {
 	UPropertyValueMaterial* MaterialProperty = Cast<UPropertyValueMaterial>(const_cast<UPropertyValue*>(Property));
 	if (MaterialProperty == nullptr)
@@ -284,7 +284,7 @@ bool FGLTFLevelVariantSetsConverter::TryParseMaterialPropertyValue(FGLTFConvertB
 	return true;
 }
 
-bool FGLTFLevelVariantSetsConverter::TryParseStaticMeshPropertyValue(FGLTFConvertBuilder& Builder, FGLTFJsonVariant& OutVariant, const UPropertyValue* Property) const
+bool FGLTFLevelVariantSetsConverter::TryParseStaticMeshPropertyValue(FGLTFJsonVariant& OutVariant, const UPropertyValue* Property) const
 {
 	const USceneComponent* Target = static_cast<USceneComponent*>(Property->GetPropertyParentContainerAddress());
 	if (Target == nullptr)
@@ -341,7 +341,7 @@ bool FGLTFLevelVariantSetsConverter::TryParseStaticMeshPropertyValue(FGLTFConver
 	return true;
 }
 
-bool FGLTFLevelVariantSetsConverter::TryParseSkeletalMeshPropertyValue(FGLTFConvertBuilder& Builder, FGLTFJsonVariant& OutVariant, const UPropertyValue* Property) const
+bool FGLTFLevelVariantSetsConverter::TryParseSkeletalMeshPropertyValue(FGLTFJsonVariant& OutVariant, const UPropertyValue* Property) const
 {
 	const USceneComponent* Target = static_cast<USceneComponent*>(Property->GetPropertyParentContainerAddress());
 	if (Target == nullptr)
diff --git a/Engine/Plugins/Enterprise/GLTFExporter/MigratedFromGit/Source/GLTFExporter/Private/Converters/GLTFVariantSetConverters.h b/Engine/Plugins/Enterprise/GLTFExporter/MigratedFromGit/Source/GLTFExporter/Private/Converters/GLTFVariantSetConverters.h
--- a/Engine/Plugins/Enterprise/GLTFExporter/MigratedFromGit/Source/GLTFExporter/Private/Converters/GLTFVariantSetConverters.h
+++ b/Engine/Plugins/Enterprise/GLTFExporter/MigratedFromGit/Source/GLTFExporter/Private/Converters/GLTFVariantSetConverters.h
@@ -13,14 +13,16 @@
 
 class FGLTFLevelVariantSetsConverter final : public TGLTFConverter<FGLTFJsonLevelVariantSetsIndex, const ALevelVariantSetsActor*>
 {
-	FGLTFJsonLevelVariantSetsIndex Add(FGLTFConvertBuilder& Builder, const FString& Name, const ALevelVariantSetsActor* LevelVariantSetsActor) override;
+	using TGLTFConverter::TGLTFConverter;
 
-	bool TryParseVariant(FGLTFConvertBuilder& Builder, FGLTFJsonVariant& OutVariant, const UVariant* Variant) const;
-	bool TryParseVariantBinding(FGLTFConvertBuilder& Builder, FGLTFJsonVariant& OutVariant, const UVariantObjectBinding* Binding) const;
-	bool TryParseVisibilityPropertyValue(FGLTFConvertBuilder& Builder, FGLTFJsonVariant& OutVariant, const UPropertyValue* Property) const;
-	bool TryParseMaterialPropertyValue(FGLTFConvertBuilder& Builder, FGLTFJsonVariant& OutVariant, const UPropertyValue* Property) const;
-	bool TryParseStaticMeshPropertyValue(FGLTFConvertBuilder& Builder, FGLTFJsonVariant& OutVariant, const UPropertyValue* Property) const;
-	bool TryParseSkeletalMeshPropertyValue(FGLTFConvertBuilder& Builder, FGLTFJsonVariant& OutVariant, const UPropertyValue* Property) const;
+	FGLTFJsonLevelVariantSetsIndex Convert(const FString& Name, const ALevelVariantSetsActor* LevelVariantSetsActor) override;
+
+	bool TryParseVariant(FGLTFJsonVariant& OutVariant, const UVariant* Variant) const;
+	bool TryParseVariantBinding(FGLTFJsonVariant& OutVariant, const UVariantObjectBinding* Binding) const;
+	bool TryParseVisibilityPropertyValue(FGLTFJsonVariant& OutVariant, const UPropertyValue* Property) const;
+	bool TryParseMaterialPropertyValue(FGLTFJsonVariant& OutVariant, const UPropertyValue* Property) const;
+	bool TryParseStaticMeshPropertyValue(FGLTFJsonVariant& OutVariant, const UPropertyValue* Property) const;
+	bool TryParseSkeletalMeshPropertyValue(FGLTFJsonVariant& OutVariant, const UPropertyValue* Property) const;
 
 	template<typename T>
 	bool TryGetPropertyValue(UPropertyValue* Property, T& OutValue) const;
diff --git a/Engine/Plugins/Enterprise/GLTFExporter/MigratedFromGit/Source/GLTFExporter/Private/Converters/GLTFVertexBufferConverters.cpp b/Engine/Plugins/Enterprise/GLTFExporter/MigratedFromGit/Source/GLTFExporter/Private/Converters/GLTFVertexBufferConverters.cpp
--- a/Engine/Plugins/Enterprise/GLTFExporter/MigratedFromGit/Source/GLTFExporter/Private/Converters/GLTFVertexBufferConverters.cpp
+++ b/Engine/Plugins/Enterprise/GLTFExporter/MigratedFromGit/Source/GLTFExporter/Private/Converters/GLTFVertexBufferConverters.cpp
@@ -6,7 +6,7 @@
 #include "Builders/GLTFConvertBuilder.h"
 #include "Math/NumericLimits.h"
 
-FGLTFJsonAccessorIndex FGLTFPositionVertexBufferConverter::Add(FGLTFConvertBuilder& Builder, const FString& Name, const FPositionVertexBuffer* VertexBuffer)
+FGLTFJsonAccessorIndex FGLTFPositionVertexBufferConverter::Convert(const FString& Name, const FPositionVertexBuffer* VertexBuffer)
 {
 	const uint32 VertexCount = VertexBuffer->GetNumVertices();
 	if (VertexCount == 0)
@@ -55,7 +55,7 @@ FGLTFJsonAccessorIndex FGLTFPositionVertexBufferConverter::Add(FGLTFConvertBuild
 	return Builder.AddAccessor(JsonAccessor);
 }
 
-FGLTFJsonAccessorIndex FGLTFColorVertexBufferConverter::Add(FGLTFConvertBuilder& Builder, const FString& Name, const FColorVertexBuffer* VertexBuffer)
+FGLTFJsonAccessorIndex FGLTFColorVertexBufferConverter::Convert(const FString& Name, const FColorVertexBuffer* VertexBuffer)
 {
 	const uint32 VertexCount = VertexBuffer->GetNumVertices();
 	if (VertexCount == 0)
@@ -82,7 +82,7 @@ FGLTFJsonAccessorIndex FGLTFColorVertexBufferConverter::Add(FGLTFConvertBuilder&
 	return Builder.AddAccessor(JsonAccessor);
 }
 
-FGLTFJsonAccessorIndex FGLTFNormalVertexBufferConverter::Add(FGLTFConvertBuilder& Builder, const FString& Name, const FStaticMeshVertexBuffer* VertexBuffer)
+FGLTFJsonAccessorIndex FGLTFNormalVertexBufferConverter::Convert(const FString& Name, const FStaticMeshVertexBuffer* VertexBuffer)
 {
 	const uint32 VertexCount = VertexBuffer->GetNumVertices();
 	if (VertexCount == 0)
@@ -108,7 +108,7 @@ FGLTFJsonAccessorIndex FGLTFNormalVertexBufferConverter::Add(FGLTFConvertBuilder
 	return Builder.AddAccessor(JsonAccessor);
 }
 
-FGLTFJsonAccessorIndex FGLTFTangentVertexBufferConverter::Add(FGLTFConvertBuilder& Builder, const FString& Name, const FStaticMeshVertexBuffer* VertexBuffer)
+FGLTFJsonAccessorIndex FGLTFTangentVertexBufferConverter::Convert(const FString& Name, const FStaticMeshVertexBuffer* VertexBuffer)
 {
 	const uint32 VertexCount = VertexBuffer->GetNumVertices();
 	if (VertexCount == 0)
@@ -134,7 +134,7 @@ FGLTFJsonAccessorIndex FGLTFTangentVertexBufferConverter::Add(FGLTFConvertBuilde
 	return Builder.AddAccessor(JsonAccessor);
 }
 
-FGLTFJsonAccessorIndex FGLTFUVVertexBufferConverter::Add(FGLTFConvertBuilder& Builder, const FString& Name, const FStaticMeshVertexBuffer* VertexBuffer, int32 UVIndex)
+FGLTFJsonAccessorIndex FGLTFUVVertexBufferConverter::Convert(const FString& Name, const FStaticMeshVertexBuffer* VertexBuffer, int32 UVIndex)
 {
 	const uint32 VertexCount = VertexBuffer->GetNumVertices();
 	if (VertexCount == 0 || UVIndex < 0 || VertexBuffer->GetNumTexCoords() <= static_cast<uint32>(UVIndex))
@@ -160,13 +160,13 @@ FGLTFJsonAccessorIndex FGLTFUVVertexBufferConverter::Add(FGLTFConvertBuilder& Bu
 	return Builder.AddAccessor(JsonAccessor);
 }
 
-FGLTFJsonAccessorIndex FGLTFBoneIndexVertexBufferConverter::Add(FGLTFConvertBuilder& Builder, const FString& Name, const FSkinWeightVertexBuffer* VertexBuffer, int32 InfluenceOffset, FGLTFBoneMap BoneMap)
+FGLTFJsonAccessorIndex FGLTFBoneIndexVertexBufferConverter::Convert(const FString& Name, const FSkinWeightVertexBuffer* VertexBuffer, int32 InfluenceOffset, FGLTFBoneMap BoneMap)
 {
-	return VertexBuffer->Use16BitBoneIndex() ? Add<uint16>(Builder, Name, VertexBuffer, InfluenceOffset, BoneMap) : Add<uint8>(Builder, Name, VertexBuffer, InfluenceOffset, BoneMap);
+	return VertexBuffer->Use16BitBoneIndex() ? Convert<uint16>(Name, VertexBuffer, InfluenceOffset, BoneMap) : Convert<uint8>(Name, VertexBuffer, InfluenceOffset, BoneMap);
 }
 
 template <typename IndexType>
-FGLTFJsonAccessorIndex FGLTFBoneIndexVertexBufferConverter::Add(FGLTFConvertBuilder& Builder, const FString& Name, const FSkinWeightVertexBuffer* VertexBuffer, int32 JointsGroupIndex, FGLTFBoneMap BoneMap)
+FGLTFJsonAccessorIndex FGLTFBoneIndexVertexBufferConverter::Convert(const FString& Name, const FSkinWeightVertexBuffer* VertexBuffer, int32 JointsGroupIndex, FGLTFBoneMap BoneMap)
 {
 	const uint32 VertexCount = VertexBuffer->GetNumVertices();
 	const int32 MaxInfluenceCount = VertexBuffer->GetMaxBoneInfluences();
@@ -208,7 +208,7 @@ FGLTFJsonAccessorIndex FGLTFBoneIndexVertexBufferConverter::Add(FGLTFConvertBuil
 	return Builder.AddAccessor(JsonAccessor);
 }
 
-FGLTFJsonAccessorIndex FGLTFBoneWeightVertexBufferConverter::Add(FGLTFConvertBuilder& Builder, const FString& Name, const FSkinWeightVertexBuffer* VertexBuffer, int32 WeightsGroupIndex)
+FGLTFJsonAccessorIndex FGLTFBoneWeightVertexBufferConverter::Convert(const FString& Name, const FSkinWeightVertexBuffer* VertexBuffer, int32 WeightsGroupIndex)
 {
 	const uint32 VertexCount = VertexBuffer->GetNumVertices();
 	const int32 MaxInfluenceCount = VertexBuffer->GetMaxBoneInfluences();
diff --git a/Engine/Plugins/Enterprise/GLTFExporter/MigratedFromGit/Source/GLTFExporter/Private/Converters/GLTFVertexBufferConverters.h b/Engine/Plugins/Enterprise/GLTFExporter/MigratedFromGit/Source/GLTFExporter/Private/Converters/GLTFVertexBufferConverters.h
--- a/Engine/Plugins/Enterprise/GLTFExporter/MigratedFromGit/Source/GLTFExporter/Private/Converters/GLTFVertexBufferConverters.h
+++ b/Engine/Plugins/Enterprise/GLTFExporter/MigratedFromGit/Source/GLTFExporter/Private/Converters/GLTFVertexBufferConverters.h
@@ -9,38 +9,52 @@
 
 class FGLTFPositionVertexBufferConverter final : public TGLTFConverter<FGLTFJsonAccessorIndex, const FPositionVertexBuffer*>
 {
-	FGLTFJsonAccessorIndex Add(FGLTFConvertBuilder& Builder, const FString& Name, const FPositionVertexBuffer* VertexBuffer) override;
+	using TGLTFConverter::TGLTFConverter;
+
+	FGLTFJsonAccessorIndex Convert(const FString& Name, const FPositionVertexBuffer* VertexBuffer) override;
 };
 
 class FGLTFColorVertexBufferConverter final : public TGLTFConverter<FGLTFJsonAccessorIndex, const FColorVertexBuffer*>
 {
-	FGLTFJsonAccessorIndex Add(FGLTFConvertBuilder& Builder, const FString& Name, const FColorVertexBuffer* VertexBuffer) override;
+	using TGLTFConverter::TGLTFConverter;
+
+	FGLTFJsonAccessorIndex Convert(const FString& Name, const FColorVertexBuffer* VertexBuffer) override;
 };
 
 class FGLTFNormalVertexBufferConverter final : public TGLTFConverter<FGLTFJsonAccessorIndex, const FStaticMeshVertexBuffer*>
 {
-	FGLTFJsonAccessorIndex Add(FGLTFConvertBuilder& Builder, const FString& Name, const FStaticMeshVertexBuffer* VertexBuffer) override;
+	using TGLTFConverter::TGLTFConverter;
+
+	FGLTFJsonAccessorIndex Convert(const FString& Name, const FStaticMeshVertexBuffer* VertexBuffer) override;
 };
 
 class FGLTFTangentVertexBufferConverter final : public TGLTFConverter<FGLTFJsonAccessorIndex, const FStaticMeshVertexBuffer*>
 {
-	FGLTFJsonAccessorIndex Add(FGLTFConvertBuilder& Builder, const FString& Name, const FStaticMeshVertexBuffer* VertexBuffer) override;
+	using TGLTFConverter::TGLTFConverter;
+
+	FGLTFJsonAccessorIndex Convert(const FString& Name, const FStaticMeshVertexBuffer* VertexBuffer) override;
 };
 
 class FGLTFUVVertexBufferConverter final : public TGLTFConverter<FGLTFJsonAccessorIndex, const FStaticMeshVertexBuffer*, int32>
 {
-	FGLTFJsonAccessorIndex Add(FGLTFConvertBuilder& Builder, const FString& Name, const FStaticMeshVertexBuffer* VertexBuffer, int32 UVIndex) override;
+	using TGLTFConverter::TGLTFConverter;
+
+	FGLTFJsonAccessorIndex Convert(const FString& Name, const FStaticMeshVertexBuffer* VertexBuffer, int32 UVIndex) override;
 };
 
 class FGLTFBoneIndexVertexBufferConverter final : public TGLTFConverter<FGLTFJsonAccessorIndex, const FSkinWeightVertexBuffer*, int32, FGLTFBoneMap>
 {
-	FGLTFJsonAccessorIndex Add(FGLTFConvertBuilder& Builder, const FString& Name, const FSkinWeightVertexBuffer* VertexBuffer, int32 InfluenceOffset, FGLTFBoneMap BoneMap) override;
+	using TGLTFConverter::TGLTFConverter;
+
+	FGLTFJsonAccessorIndex Convert(const FString& Name, const FSkinWeightVertexBuffer* VertexBuffer, int32 InfluenceOffset, FGLTFBoneMap BoneMap) override;
 
 	template <typename IndexType>
-	static FGLTFJsonAccessorIndex Add(FGLTFConvertBuilder& Builder, const FString& Name, const FSkinWeightVertexBuffer* VertexBuffer, int32 JointsGroupIndex, FGLTFBoneMap BoneMap);
+	FGLTFJsonAccessorIndex Convert(const FString& Name, const FSkinWeightVertexBuffer* VertexBuffer, int32 JointsGroupIndex, FGLTFBoneMap BoneMap);
 };
 
 class FGLTFBoneWeightVertexBufferConverter final : public TGLTFConverter<FGLTFJsonAccessorIndex, const FSkinWeightVertexBuffer*, int32>
 {
-	FGLTFJsonAccessorIndex Add(FGLTFConvertBuilder& Builder, const FString& Name, const FSkinWeightVertexBuffer* VertexBuffer, int32 WeightsGroupIndex) override;
+	using TGLTFConverter::TGLTFConverter;
+
+	FGLTFJsonAccessorIndex Convert(const FString& Name, const FSkinWeightVertexBuffer* VertexBuffer, int32 WeightsGroupIndex) override;
 };
