[GLTFExporter] refactor: simplify gltf object naming by making it completely deterministic

commit 7c6f695a40685b238b0c3f3ba4d4edd7109ae803
Date: Sun, 20 Sep 2020 20:50:42 +0200

    refactor: simplify gltf object naming by making it completely deterministic
    (cherry picked from commit 2aafc1381dcf7db3e27d4410ccd67b0c712248c0​)

#skipci

diff --git a/Engine/Plugins/Enterprise/GLTFExporter/MigratedFromGit/Source/GLTFExporter/Private/Builders/GLTFBufferBuilder.cpp b/Engine/Plugins/Enterprise/GLTFExporter/MigratedFromGit/Source/GLTFExporter/Private/Builders/GLTFBufferBuilder.cpp
--- a/Engine/Plugins/Enterprise/GLTFExporter/MigratedFromGit/Source/GLTFExporter/Private/Builders/GLTFBufferBuilder.cpp
+++ b/Engine/Plugins/Enterprise/GLTFExporter/MigratedFromGit/Source/GLTFExporter/Private/Builders/GLTFBufferBuilder.cpp
@@ -10,7 +10,7 @@ FGLTFBufferBuilder::FGLTFBufferBuilder(const UGLTFExportOptions* ExportOptions)
 	BufferIndex = AddBuffer();
 }
 
-FGLTFJsonBufferViewIndex FGLTFBufferBuilder::AddBufferView(const void* RawData, uint64 ByteLength, const FString& Name, uint8 DataAlignment, EGLTFJsonBufferTarget BufferTarget)
+FGLTFJsonBufferViewIndex FGLTFBufferBuilder::AddBufferView(const void* RawData, uint64 ByteLength, uint8 DataAlignment, EGLTFJsonBufferTarget BufferTarget)
 {
 	uint64 ByteOffset = BufferData.Num();
 
@@ -25,7 +25,6 @@ FGLTFJsonBufferViewIndex FGLTFBufferBuilder::AddBufferView(const void* RawData,
 	BufferData.Append(static_cast<const uint8*>(RawData), ByteLength);
 
 	FGLTFJsonBufferView JsonBufferView;
-	JsonBufferView.Name = Name;
 	JsonBufferView.Buffer = BufferIndex;
 	JsonBufferView.ByteOffset = ByteOffset;
 	JsonBufferView.ByteLength = ByteLength;
diff --git a/Engine/Plugins/Enterprise/GLTFExporter/MigratedFromGit/Source/GLTFExporter/Private/Builders/GLTFBufferBuilder.h b/Engine/Plugins/Enterprise/GLTFExporter/MigratedFromGit/Source/GLTFExporter/Private/Builders/GLTFBufferBuilder.h
--- a/Engine/Plugins/Enterprise/GLTFExporter/MigratedFromGit/Source/GLTFExporter/Private/Builders/GLTFBufferBuilder.h
+++ b/Engine/Plugins/Enterprise/GLTFExporter/MigratedFromGit/Source/GLTFExporter/Private/Builders/GLTFBufferBuilder.h
@@ -10,12 +10,12 @@ public:
 
 	FGLTFBufferBuilder(const UGLTFExportOptions* ExportOptions);
 
-	FGLTFJsonBufferViewIndex AddBufferView(const void* RawData, uint64 ByteLength, const FString& Name = TEXT(""), uint8 DataAlignment = 4, EGLTFJsonBufferTarget BufferTarget = EGLTFJsonBufferTarget::ArrayBuffer);
+	FGLTFJsonBufferViewIndex AddBufferView(const void* RawData, uint64 ByteLength, uint8 DataAlignment = 4, EGLTFJsonBufferTarget BufferTarget = EGLTFJsonBufferTarget::ArrayBuffer);
 
 	template <class ElementType>
-	FGLTFJsonBufferViewIndex AddBufferView(const TArray<ElementType>& Array, const FString& Name = TEXT(""), uint8 DataAlignment = 4, EGLTFJsonBufferTarget BufferTarget = EGLTFJsonBufferTarget::ArrayBuffer)
+	FGLTFJsonBufferViewIndex AddBufferView(const TArray<ElementType>& Array, uint8 DataAlignment = 4, EGLTFJsonBufferTarget BufferTarget = EGLTFJsonBufferTarget::ArrayBuffer)
 	{
-		return AddBufferView(Array.GetData(), Array.Num() * sizeof(ElementType), Name, DataAlignment, BufferTarget);
+		return AddBufferView(Array.GetData(), Array.Num() * sizeof(ElementType), DataAlignment, BufferTarget);
 	}
 
 	void UpdateJsonBufferObject(const FString& BinaryFilePath);
diff --git a/Engine/Plugins/Enterprise/GLTFExporter/MigratedFromGit/Source/GLTFExporter/Private/Builders/GLTFBuilderUtility.cpp b/Engine/Plugins/Enterprise/GLTFExporter/MigratedFromGit/Source/GLTFExporter/Private/Builders/GLTFBuilderUtility.cpp
--- a/Engine/Plugins/Enterprise/GLTFExporter/MigratedFromGit/Source/GLTFExporter/Private/Builders/GLTFBuilderUtility.cpp
+++ b/Engine/Plugins/Enterprise/GLTFExporter/MigratedFromGit/Source/GLTFExporter/Private/Builders/GLTFBuilderUtility.cpp
@@ -2,23 +2,6 @@
 
 #include "Builders/GLTFBuilderUtility.h"
 
-FString FGLTFBuilderUtility::GetLODName(const UObject* Object, int32 LODIndex)
-{
-	FString Name;
-
-	if (Object != nullptr)
-	{
-		Object->GetName(Name);
-
-		if (LODIndex != 0)
-		{
-			Name += TEXT("_LOD") + FString::FromInt(LODIndex);
-		}
-	}
-
-	return Name;
-}
-
 const TCHAR* FGLTFBuilderUtility::GetFileExtension(EGLTFJsonMimeType MimeType)
 {
 	switch (MimeType)
diff --git a/Engine/Plugins/Enterprise/GLTFExporter/MigratedFromGit/Source/GLTFExporter/Private/Builders/GLTFBuilderUtility.h b/Engine/Plugins/Enterprise/GLTFExporter/MigratedFromGit/Source/GLTFExporter/Private/Builders/GLTFBuilderUtility.h
--- a/Engine/Plugins/Enterprise/GLTFExporter/MigratedFromGit/Source/GLTFExporter/Private/Builders/GLTFBuilderUtility.h
+++ b/Engine/Plugins/Enterprise/GLTFExporter/MigratedFromGit/Source/GLTFExporter/Private/Builders/GLTFBuilderUtility.h
@@ -7,8 +7,6 @@
 
 struct FGLTFBuilderUtility
 {
-	static FString GetLODName(const UObject* Object, int32 LODIndex);
-
 	static const TCHAR* GetFileExtension(EGLTFJsonMimeType MimeType);
 
 	static FString GetUniqueFilename(const FString& BaseFilename, const FString& FileExtension, const TSet<FString>& UniqueFilenames);
diff --git a/Engine/Plugins/Enterprise/GLTFExporter/MigratedFromGit/Source/GLTFExporter/Private/Builders/GLTFConvertBuilder.cpp b/Engine/Plugins/Enterprise/GLTFExporter/MigratedFromGit/Source/GLTFExporter/Private/Builders/GLTFConvertBuilder.cpp
--- a/Engine/Plugins/Enterprise/GLTFExporter/MigratedFromGit/Source/GLTFExporter/Private/Builders/GLTFConvertBuilder.cpp
+++ b/Engine/Plugins/Enterprise/GLTFExporter/MigratedFromGit/Source/GLTFExporter/Private/Builders/GLTFConvertBuilder.cpp
@@ -9,97 +9,97 @@ FGLTFConvertBuilder::FGLTFConvertBuilder(const UGLTFExportOptions* ExportOptions
 {
 }
 
-FGLTFJsonAccessorIndex FGLTFConvertBuilder::GetOrAddPositionAccessor(const FPositionVertexBuffer* VertexBuffer, const FString& DesiredName)
+FGLTFJsonAccessorIndex FGLTFConvertBuilder::GetOrAddPositionAccessor(const FPositionVertexBuffer* VertexBuffer)
 {
 	if (VertexBuffer == nullptr)
 	{
 		return FGLTFJsonAccessorIndex(INDEX_NONE);
 	}
 
-	return PositionVertexBufferConverter.GetOrAdd(DesiredName, VertexBuffer);
+	return PositionVertexBufferConverter.GetOrAdd(VertexBuffer);
 }
 
-FGLTFJsonAccessorIndex FGLTFConvertBuilder::GetOrAddColorAccessor(const FColorVertexBuffer* VertexBuffer, const FString& DesiredName)
+FGLTFJsonAccessorIndex FGLTFConvertBuilder::GetOrAddColorAccessor(const FColorVertexBuffer* VertexBuffer)
 {
 	if (VertexBuffer == nullptr)
 	{
 		return FGLTFJsonAccessorIndex(INDEX_NONE);
 	}
 
-	return ColorVertexBufferConverter.GetOrAdd(DesiredName, VertexBuffer);
+	return ColorVertexBufferConverter.GetOrAdd(VertexBuffer);
 }
 
-FGLTFJsonAccessorIndex FGLTFConvertBuilder::GetOrAddNormalAccessor(const FStaticMeshVertexBuffer* VertexBuffer, const FString& DesiredName)
+FGLTFJsonAccessorIndex FGLTFConvertBuilder::GetOrAddNormalAccessor(const FStaticMeshVertexBuffer* VertexBuffer)
 {
 	if (VertexBuffer == nullptr)
 	{
 		return FGLTFJsonAccessorIndex(INDEX_NONE);
 	}
 
-	return NormalVertexBufferConverter.GetOrAdd(DesiredName, VertexBuffer);
+	return NormalVertexBufferConverter.GetOrAdd(VertexBuffer);
 }
 
-FGLTFJsonAccessorIndex FGLTFConvertBuilder::GetOrAddTangentAccessor(const FStaticMeshVertexBuffer* VertexBuffer, const FString& DesiredName)
+FGLTFJsonAccessorIndex FGLTFConvertBuilder::GetOrAddTangentAccessor(const FStaticMeshVertexBuffer* VertexBuffer)
 {
 	if (VertexBuffer == nullptr)
 	{
 		return FGLTFJsonAccessorIndex(INDEX_NONE);
 	}
 
-	return TangentVertexBufferConverter.GetOrAdd(DesiredName, VertexBuffer);
+	return TangentVertexBufferConverter.GetOrAdd(VertexBuffer);
 }
 
-FGLTFJsonAccessorIndex FGLTFConvertBuilder::GetOrAddUVAccessor(const FStaticMeshVertexBuffer* VertexBuffer, int32 UVIndex, const FString& DesiredName)
+FGLTFJsonAccessorIndex FGLTFConvertBuilder::GetOrAddUVAccessor(const FStaticMeshVertexBuffer* VertexBuffer, int32 UVIndex)
 {
 	if (VertexBuffer == nullptr)
 	{
 		return FGLTFJsonAccessorIndex(INDEX_NONE);
 	}
 
-	return UVVertexBufferConverter.GetOrAdd(DesiredName, VertexBuffer, UVIndex);
+	return UVVertexBufferConverter.GetOrAdd(VertexBuffer, UVIndex);
 }
 
-FGLTFJsonAccessorIndex FGLTFConvertBuilder::GetOrAddJointAccessor(const FSkinWeightVertexBuffer* VertexBuffer, int32 JointsGroupIndex, FGLTFBoneMap BoneMap, const FString& DesiredName)
+FGLTFJsonAccessorIndex FGLTFConvertBuilder::GetOrAddJointAccessor(const FSkinWeightVertexBuffer* VertexBuffer, int32 JointsGroupIndex, FGLTFBoneMap BoneMap)
 {
 	if (VertexBuffer == nullptr)
 	{
 		return FGLTFJsonAccessorIndex(INDEX_NONE);
 	}
 
-	return BoneIndexVertexBufferConverter.GetOrAdd(DesiredName, VertexBuffer, JointsGroupIndex, BoneMap);
+	return BoneIndexVertexBufferConverter.GetOrAdd(VertexBuffer, JointsGroupIndex, BoneMap);
 }
 
-FGLTFJsonAccessorIndex FGLTFConvertBuilder::GetOrAddWeightAccessor(const FSkinWeightVertexBuffer* VertexBuffer, int32 WeightsGroupIndex, const FString& DesiredName)
+FGLTFJsonAccessorIndex FGLTFConvertBuilder::GetOrAddWeightAccessor(const FSkinWeightVertexBuffer* VertexBuffer, int32 WeightsGroupIndex)
 {
 	if (VertexBuffer == nullptr)
 	{
 		return FGLTFJsonAccessorIndex(INDEX_NONE);
 	}
 
-	return BoneWeightVertexBufferConverter.GetOrAdd(DesiredName, VertexBuffer, WeightsGroupIndex);
+	return BoneWeightVertexBufferConverter.GetOrAdd(VertexBuffer, WeightsGroupIndex);
 }
 
-FGLTFJsonBufferViewIndex FGLTFConvertBuilder::GetOrAddIndexBufferView(const FRawStaticIndexBuffer* IndexBuffer, const FString& DesiredName)
+FGLTFJsonBufferViewIndex FGLTFConvertBuilder::GetOrAddIndexBufferView(const FRawStaticIndexBuffer* IndexBuffer)
 {
 	if (IndexBuffer == nullptr)
 	{
 		return FGLTFJsonBufferViewIndex(INDEX_NONE);
 	}
 
-	return IndexBufferConverter.GetOrAdd(DesiredName, IndexBuffer);
+	return IndexBufferConverter.GetOrAdd(IndexBuffer);
 }
 
-FGLTFJsonAccessorIndex FGLTFConvertBuilder::GetOrAddIndexAccessor(const FStaticMeshSection* MeshSection, const FRawStaticIndexBuffer* IndexBuffer, const FString& DesiredName)
+FGLTFJsonAccessorIndex FGLTFConvertBuilder::GetOrAddIndexAccessor(const FStaticMeshSection* MeshSection, const FRawStaticIndexBuffer* IndexBuffer)
 {
 	if (MeshSection == nullptr)
 	{
 		return FGLTFJsonAccessorIndex(INDEX_NONE);
 	}
 
-	return StaticMeshSectionConverter.GetOrAdd(DesiredName, MeshSection, IndexBuffer);
+	return StaticMeshSectionConverter.GetOrAdd(MeshSection, IndexBuffer);
 }
 
-FGLTFJsonMeshIndex FGLTFConvertBuilder::GetOrAddMesh(const UStaticMesh* StaticMesh, int32 LODIndex, const FColorVertexBuffer* OverrideVertexColors, const FGLTFMaterialArray& OverrideMaterials, const FString& DesiredName)
+FGLTFJsonMeshIndex FGLTFConvertBuilder::GetOrAddMesh(const UStaticMesh* StaticMesh, int32 LODIndex, const FColorVertexBuffer* OverrideVertexColors, const FGLTFMaterialArray& OverrideMaterials)
 {
 	if (StaticMesh == nullptr)
 	{
@@ -117,14 +117,13 @@ FGLTFJsonMeshIndex FGLTFConvertBuilder::GetOrAddMesh(const UStaticMesh* StaticMe
 	}
 
 	return StaticMeshConverter.GetOrAdd(
-		DesiredName.IsEmpty() ? FGLTFBuilderUtility::GetLODName(StaticMesh, LODIndex) : DesiredName,
 		StaticMesh,
 		LODIndex,
 		OverrideVertexColors,
 		OverrideMaterials == StaticMesh->StaticMaterials ? FGLTFMaterialArray(): OverrideMaterials);
 }
 
-FGLTFJsonMeshIndex FGLTFConvertBuilder::GetOrAddMesh(const UStaticMeshComponent* StaticMeshComponent, const FString& DesiredName)
+FGLTFJsonMeshIndex FGLTFConvertBuilder::GetOrAddMesh(const UStaticMeshComponent* StaticMeshComponent)
 {
 	if (StaticMeshComponent == nullptr)
 	{
@@ -136,30 +135,30 @@ FGLTFJsonMeshIndex FGLTFConvertBuilder::GetOrAddMesh(const UStaticMeshComponent*
 	const FColorVertexBuffer* OverrideVertexColors = StaticMeshComponent->LODData.IsValidIndex(LODIndex) ? StaticMeshComponent->LODData[LODIndex].OverrideVertexColors : nullptr;
 	const FGLTFMaterialArray OverrideMaterials = FGLTFMaterialArray(StaticMeshComponent->OverrideMaterials);
 
-	return GetOrAddMesh(StaticMesh, LODIndex, OverrideVertexColors, OverrideMaterials, DesiredName);
+	return GetOrAddMesh(StaticMesh, LODIndex, OverrideVertexColors, OverrideMaterials);
 }
 
-FGLTFJsonBufferViewIndex FGLTFConvertBuilder::GetOrAddIndexBufferView(const FMultiSizeIndexContainer* IndexContainer, const FString& DesiredName)
+FGLTFJsonBufferViewIndex FGLTFConvertBuilder::GetOrAddIndexBufferView(const FMultiSizeIndexContainer* IndexContainer)
 {
 	if (IndexContainer == nullptr)
 	{
 		return FGLTFJsonBufferViewIndex(INDEX_NONE);
 	}
 
-	return IndexContainerConverter.GetOrAdd(DesiredName, IndexContainer);
+	return IndexContainerConverter.GetOrAdd(IndexContainer);
 }
 
-FGLTFJsonAccessorIndex FGLTFConvertBuilder::GetOrAddIndexAccessor(const FSkelMeshRenderSection* MeshSection, const FMultiSizeIndexContainer* IndexContainer, const FString& DesiredName)
+FGLTFJsonAccessorIndex FGLTFConvertBuilder::GetOrAddIndexAccessor(const FSkelMeshRenderSection* MeshSection, const FMultiSizeIndexContainer* IndexContainer)
 {
 	if (MeshSection == nullptr)
 	{
 		return FGLTFJsonAccessorIndex(INDEX_NONE);
 	}
 
-	return SkeletalMeshSectionConverter.GetOrAdd(DesiredName, MeshSection, IndexContainer);
+	return SkeletalMeshSectionConverter.GetOrAdd(MeshSection, IndexContainer);
 }
 
-FGLTFJsonMeshIndex FGLTFConvertBuilder::GetOrAddMesh(const USkeletalMesh* SkeletalMesh, int32 LODIndex, const FColorVertexBuffer* OverrideVertexColors, const FSkinWeightVertexBuffer* OverrideSkinWeights, const FGLTFMaterialArray& OverrideMaterials, const FString& DesiredName)
+FGLTFJsonMeshIndex FGLTFConvertBuilder::GetOrAddMesh(const USkeletalMesh* SkeletalMesh, int32 LODIndex, const FColorVertexBuffer* OverrideVertexColors, const FSkinWeightVertexBuffer* OverrideSkinWeights, const FGLTFMaterialArray& OverrideMaterials)
 {
 	if (SkeletalMesh == nullptr)
 	{
@@ -176,10 +175,10 @@ FGLTFJsonMeshIndex FGLTFConvertBuilder::GetOrAddMesh(const USkeletalMesh* Skelet
 		OverrideVertexColors = nullptr;
 	}
 
-	return SkeletalMeshConverter.GetOrAdd(DesiredName.IsEmpty() ? FGLTFBuilderUtility::GetLODName(SkeletalMesh, LODIndex) : DesiredName, SkeletalMesh, LODIndex, OverrideVertexColors, OverrideSkinWeights, OverrideMaterials);
+	return SkeletalMeshConverter.GetOrAdd(SkeletalMesh, LODIndex, OverrideVertexColors, OverrideSkinWeights, OverrideMaterials);
 }
 
-FGLTFJsonMeshIndex FGLTFConvertBuilder::GetOrAddMesh(const USkeletalMeshComponent* SkeletalMeshComponent, const FString& DesiredName)
+FGLTFJsonMeshIndex FGLTFConvertBuilder::GetOrAddMesh(const USkeletalMeshComponent* SkeletalMeshComponent)
 {
 	if (SkeletalMeshComponent == nullptr)
 	{
@@ -192,165 +191,165 @@ FGLTFJsonMeshIndex FGLTFConvertBuilder::GetOrAddMesh(const USkeletalMeshComponen
 	const FSkinWeightVertexBuffer* OverrideSkinWeights = SkeletalMeshComponent->LODInfo.IsValidIndex(LODIndex) ? SkeletalMeshComponent->LODInfo[LODIndex].OverrideSkinWeights : nullptr;
 	const FGLTFMaterialArray OverrideMaterials = FGLTFMaterialArray(SkeletalMeshComponent->OverrideMaterials);
 
-	return GetOrAddMesh(SkeletalMesh, LODIndex, OverrideVertexColors, OverrideSkinWeights, OverrideMaterials, DesiredName);
+	return GetOrAddMesh(SkeletalMesh, LODIndex, OverrideVertexColors, OverrideSkinWeights, OverrideMaterials);
 }
 
-FGLTFJsonMaterialIndex FGLTFConvertBuilder::GetOrAddMaterial(const UMaterialInterface* Material, const FString& DesiredName)
+FGLTFJsonMaterialIndex FGLTFConvertBuilder::GetOrAddMaterial(const UMaterialInterface* Material)
 {
 	if (Material == nullptr)
 	{
 		return FGLTFJsonMaterialIndex(INDEX_NONE);
 	}
 
-	return MaterialConverter.GetOrAdd(DesiredName.IsEmpty() ? Material->GetName() : DesiredName, Material);
+	return MaterialConverter.GetOrAdd(Material);
 }
 
-FGLTFJsonSamplerIndex FGLTFConvertBuilder::GetOrAddSampler(const UTexture* Texture, const FString& DesiredName)
+FGLTFJsonSamplerIndex FGLTFConvertBuilder::GetOrAddSampler(const UTexture* Texture)
 {
 	if (Texture == nullptr)
 	{
 		return FGLTFJsonSamplerIndex(INDEX_NONE);
 	}
 
-	return TextureSamplerConverter.GetOrAdd(DesiredName.IsEmpty() ? Texture->GetName() : DesiredName, Texture);
+	return TextureSamplerConverter.GetOrAdd(Texture);
 }
 
-FGLTFJsonTextureIndex FGLTFConvertBuilder::GetOrAddTexture(const UTexture2D* Texture, const FString& DesiredName)
+FGLTFJsonTextureIndex FGLTFConvertBuilder::GetOrAddTexture(const UTexture2D* Texture)
 {
 	if (Texture == nullptr)
 	{
 		return FGLTFJsonTextureIndex(INDEX_NONE);
 	}
 
-	return Texture2DConverter.GetOrAdd(DesiredName.IsEmpty() ? Texture->GetName() : DesiredName, Texture);
+	return Texture2DConverter.GetOrAdd(Texture);
 }
 
-FGLTFJsonTextureIndex FGLTFConvertBuilder::GetOrAddTexture(const UTextureCube* Texture, ECubeFace CubeFace, const FString& DesiredName)
+FGLTFJsonTextureIndex FGLTFConvertBuilder::GetOrAddTexture(const UTextureCube* Texture, ECubeFace CubeFace)
 {
 	if (Texture == nullptr)
 	{
 		return FGLTFJsonTextureIndex(INDEX_NONE);
 	}
 
-	return TextureCubeConverter.GetOrAdd(DesiredName.IsEmpty() ? Texture->GetName() : DesiredName, Texture, CubeFace);
+	return TextureCubeConverter.GetOrAdd(Texture, CubeFace);
 }
 
-FGLTFJsonTextureIndex FGLTFConvertBuilder::GetOrAddTexture(const UTextureRenderTarget2D* Texture, const FString& DesiredName)
+FGLTFJsonTextureIndex FGLTFConvertBuilder::GetOrAddTexture(const UTextureRenderTarget2D* Texture)
 {
 	if (Texture == nullptr)
 	{
 		return FGLTFJsonTextureIndex(INDEX_NONE);
 	}
 
-	return TextureRenderTarget2DConverter.GetOrAdd(DesiredName.IsEmpty() ? Texture->GetName() : DesiredName, Texture);
+	return TextureRenderTarget2DConverter.GetOrAdd(Texture);
 }
 
-FGLTFJsonTextureIndex FGLTFConvertBuilder::GetOrAddTexture(const UTextureRenderTargetCube* Texture, ECubeFace CubeFace, const FString& DesiredName)
+FGLTFJsonTextureIndex FGLTFConvertBuilder::GetOrAddTexture(const UTextureRenderTargetCube* Texture, ECubeFace CubeFace)
 {
 	if (Texture == nullptr)
 	{
 		return FGLTFJsonTextureIndex(INDEX_NONE);
 	}
 
-	return TextureRenderTargetCubeConverter.GetOrAdd(DesiredName.IsEmpty() ? Texture->GetName() : DesiredName, Texture, CubeFace);
+	return TextureRenderTargetCubeConverter.GetOrAdd(Texture, CubeFace);
 }
 
-FGLTFJsonNodeIndex FGLTFConvertBuilder::GetOrAddNode(const USceneComponent* SceneComponent, const FString& DesiredName)
+FGLTFJsonNodeIndex FGLTFConvertBuilder::GetOrAddNode(const USceneComponent* SceneComponent)
 {
 	if (SceneComponent == nullptr)
 	{
 		return FGLTFJsonNodeIndex(INDEX_NONE);
 	}
 
-	return SceneComponentConverter.GetOrAdd(DesiredName, SceneComponent);
+	return SceneComponentConverter.GetOrAdd(SceneComponent);
 }
 
-FGLTFJsonNodeIndex FGLTFConvertBuilder::GetOrAddNode(const AActor* Actor, const FString& DesiredName)
+FGLTFJsonNodeIndex FGLTFConvertBuilder::GetOrAddNode(const AActor* Actor)
 {
 	if (Actor == nullptr)
 	{
 		return FGLTFJsonNodeIndex(INDEX_NONE);
 	}
 
-	return ActorConverter.GetOrAdd(DesiredName, Actor);
+	return ActorConverter.GetOrAdd(Actor);
 }
 
-FGLTFJsonSceneIndex FGLTFConvertBuilder::GetOrAddScene(const ULevel* Level, const FString& DesiredName)
+FGLTFJsonSceneIndex FGLTFConvertBuilder::GetOrAddScene(const ULevel* Level)
 {
 	if (Level == nullptr)
 	{
 		return FGLTFJsonSceneIndex(INDEX_NONE);
 	}
 
-	return LevelConverter.GetOrAdd(DesiredName, Level);
+	return LevelConverter.GetOrAdd(Level);
 }
 
-FGLTFJsonSceneIndex FGLTFConvertBuilder::GetOrAddScene(const UWorld* World, const FString& DesiredName)
+FGLTFJsonSceneIndex FGLTFConvertBuilder::GetOrAddScene(const UWorld* World)
 {
 	if (World == nullptr)
 	{
 		return FGLTFJsonSceneIndex(INDEX_NONE);
 	}
 
-	return GetOrAddScene(World->PersistentLevel, DesiredName.IsEmpty() ? World->GetName() : DesiredName);
+	return GetOrAddScene(World->PersistentLevel);
 }
 
-FGLTFJsonCameraIndex FGLTFConvertBuilder::GetOrAddCamera(const UCameraComponent* CameraComponent, const FString& DesiredName)
+FGLTFJsonCameraIndex FGLTFConvertBuilder::GetOrAddCamera(const UCameraComponent* CameraComponent)
 {
 	if (CameraComponent == nullptr)
 	{
 		return FGLTFJsonCameraIndex(INDEX_NONE);
 	}
 
-	return CameraComponentConverter.GetOrAdd(DesiredName.IsEmpty() ? CameraComponent->GetName() : DesiredName, CameraComponent);
+	return CameraComponentConverter.GetOrAdd(CameraComponent);
 }
 
-FGLTFJsonLightIndex FGLTFConvertBuilder::GetOrAddLight(const ULightComponent* LightComponent, const FString& DesiredName)
+FGLTFJsonLightIndex FGLTFConvertBuilder::GetOrAddLight(const ULightComponent* LightComponent)
 {
 	if (LightComponent == nullptr)
 	{
 		return FGLTFJsonLightIndex(INDEX_NONE);
 	}
 
-	return LightComponentConverter.GetOrAdd(DesiredName.IsEmpty() ? LightComponent->GetName() : DesiredName, LightComponent);
+	return LightComponentConverter.GetOrAdd(LightComponent);
 }
 
-FGLTFJsonBackdropIndex FGLTFConvertBuilder::GetOrAddBackdrop(const AActor* Actor, const FString& DesiredName)
+FGLTFJsonBackdropIndex FGLTFConvertBuilder::GetOrAddBackdrop(const AActor* Actor)
 {
 	if (Actor == nullptr)
 	{
 		return FGLTFJsonBackdropIndex(INDEX_NONE);
 	}
 
-	return BackdropConverter.GetOrAdd(DesiredName, Actor);
+	return BackdropConverter.GetOrAdd(Actor);
 }
 
-FGLTFJsonLevelVariantSetsIndex FGLTFConvertBuilder::GetOrAddLevelVariantSets(const ALevelVariantSetsActor* LevelVariantSetsActor, const FString& DesiredName)
+FGLTFJsonLevelVariantSetsIndex FGLTFConvertBuilder::GetOrAddLevelVariantSets(const ALevelVariantSetsActor* LevelVariantSetsActor)
 {
 	if (LevelVariantSetsActor == nullptr)
 	{
 		return FGLTFJsonLevelVariantSetsIndex(INDEX_NONE);
 	}
 
-	return LevelVariantSetsConverter.GetOrAdd(DesiredName, LevelVariantSetsActor);
+	return LevelVariantSetsConverter.GetOrAdd(LevelVariantSetsActor);
 }
 
-FGLTFJsonLightMapIndex FGLTFConvertBuilder::GetOrAddLightMap(const UStaticMeshComponent* StaticMeshComponent, const FString& DesiredName)
+FGLTFJsonLightMapIndex FGLTFConvertBuilder::GetOrAddLightMap(const UStaticMeshComponent* StaticMeshComponent)
 {
 	if (StaticMeshComponent == nullptr)
 	{
 		return FGLTFJsonLightMapIndex(INDEX_NONE);
 	}
 
-	return LightMapConverter.GetOrAdd(DesiredName, StaticMeshComponent);
+	return LightMapConverter.GetOrAdd(StaticMeshComponent);
 }
 
-FGLTFJsonHotspotIndex FGLTFConvertBuilder::GetOrAddHotspot(const UGLTFInteractionHotspotComponent* HotspotComponent, const FString& DesiredName)
+FGLTFJsonHotspotIndex FGLTFConvertBuilder::GetOrAddHotspot(const UGLTFInteractionHotspotComponent* HotspotComponent)
 {
 	if (HotspotComponent == nullptr)
 	{
 		return FGLTFJsonHotspotIndex(INDEX_NONE);
 	}
 
-	return HotspotComponentConverter.GetOrAdd(DesiredName.IsEmpty() ? HotspotComponent->GetName() : DesiredName, HotspotComponent);
+	return HotspotComponentConverter.GetOrAdd(HotspotComponent);
 }
diff --git a/Engine/Plugins/Enterprise/GLTFExporter/MigratedFromGit/Source/GLTFExporter/Private/Builders/GLTFConvertBuilder.h b/Engine/Plugins/Enterprise/GLTFExporter/MigratedFromGit/Source/GLTFExporter/Private/Builders/GLTFConvertBuilder.h
--- a/Engine/Plugins/Enterprise/GLTFExporter/MigratedFromGit/Source/GLTFExporter/Private/Builders/GLTFConvertBuilder.h
+++ b/Engine/Plugins/Enterprise/GLTFExporter/MigratedFromGit/Source/GLTFExporter/Private/Builders/GLTFConvertBuilder.h
@@ -22,43 +22,43 @@ public:
 
 	FGLTFConvertBuilder(const UGLTFExportOptions* ExportOptions, bool bSelectedActorsOnly);
 
-	FGLTFJsonAccessorIndex GetOrAddPositionAccessor(const FPositionVertexBuffer* VertexBuffer, const FString& DesiredName = TEXT(""));
-	FGLTFJsonAccessorIndex GetOrAddColorAccessor(const FColorVertexBuffer* VertexBuffer, const FString& DesiredName = TEXT(""));
-	FGLTFJsonAccessorIndex GetOrAddNormalAccessor(const FStaticMeshVertexBuffer* VertexBuffer, const FString& DesiredName = TEXT(""));
-	FGLTFJsonAccessorIndex GetOrAddTangentAccessor(const FStaticMeshVertexBuffer* VertexBuffer, const FString& DesiredName = TEXT(""));
-	FGLTFJsonAccessorIndex GetOrAddUVAccessor(const FStaticMeshVertexBuffer* VertexBuffer, int32 UVIndex, const FString& DesiredName = TEXT(""));
-	FGLTFJsonAccessorIndex GetOrAddJointAccessor(const FSkinWeightVertexBuffer* VertexBuffer, int32 JointsGroupIndex, FGLTFBoneMap BoneMap, const FString& DesiredName = TEXT(""));
-	FGLTFJsonAccessorIndex GetOrAddWeightAccessor(const FSkinWeightVertexBuffer* VertexBuffer, int32 WeightsGroupIndex, const FString& DesiredName = TEXT(""));
-
-	FGLTFJsonBufferViewIndex GetOrAddIndexBufferView(const FRawStaticIndexBuffer* IndexBuffer, const FString& DesiredName = TEXT(""));
-	FGLTFJsonAccessorIndex GetOrAddIndexAccessor(const FStaticMeshSection* MeshSection, const FRawStaticIndexBuffer* IndexBuffer, const FString& DesiredName = TEXT(""));
-	FGLTFJsonMeshIndex GetOrAddMesh(const UStaticMesh* StaticMesh, int32 LODIndex = -1, const FColorVertexBuffer* OverrideVertexColors = nullptr, const FGLTFMaterialArray& OverrideMaterials = {}, const FString& DesiredName = TEXT(""));
-	FGLTFJsonMeshIndex GetOrAddMesh(const UStaticMeshComponent* StaticMeshComponent, const FString& DesiredName = TEXT(""));
-
-	FGLTFJsonBufferViewIndex GetOrAddIndexBufferView(const FMultiSizeIndexContainer* IndexContainer, const FString& DesiredName = TEXT(""));
-	FGLTFJsonAccessorIndex GetOrAddIndexAccessor(const FSkelMeshRenderSection* MeshSection, const FMultiSizeIndexContainer* IndexContainer, const FString& DesiredName = TEXT(""));
-	FGLTFJsonMeshIndex GetOrAddMesh(const USkeletalMesh* SkeletalMesh, int32 LODIndex = -1, const FColorVertexBuffer* OverrideVertexColors = nullptr, const FSkinWeightVertexBuffer* OverrideSkinWeights = nullptr, const FGLTFMaterialArray& OverrideMaterials = {}, const FString& DesiredName = TEXT(""));
-	FGLTFJsonMeshIndex GetOrAddMesh(const USkeletalMeshComponent* SkeletalMeshComponent, const FString& DesiredName = TEXT(""));
-
-	FGLTFJsonMaterialIndex GetOrAddMaterial(const UMaterialInterface* Material, const FString& DesiredName = TEXT(""));
-	FGLTFJsonSamplerIndex GetOrAddSampler(const UTexture* Texture, const FString& DesiredName = TEXT(""));
-	FGLTFJsonTextureIndex GetOrAddTexture(const UTexture2D* Texture, const FString& DesiredName = TEXT(""));
-	FGLTFJsonTextureIndex GetOrAddTexture(const UTextureCube* Texture, ECubeFace CubeFace, const FString& DesiredName = TEXT(""));
-	FGLTFJsonTextureIndex GetOrAddTexture(const UTextureRenderTarget2D* Texture, const FString& DesiredName = TEXT(""));
-	FGLTFJsonTextureIndex GetOrAddTexture(const UTextureRenderTargetCube* Texture, ECubeFace CubeFace, const FString& DesiredName = TEXT(""));
-
-	FGLTFJsonNodeIndex GetOrAddNode(const USceneComponent* SceneComponent, const FString& DesiredName = TEXT(""));
-	FGLTFJsonNodeIndex GetOrAddNode(const AActor* Actor, const FString& DesiredName = TEXT(""));
-	FGLTFJsonSceneIndex GetOrAddScene(const ULevel* Level, const FString& DesiredName = TEXT(""));
-	FGLTFJsonSceneIndex GetOrAddScene(const UWorld* World, const FString& DesiredName = TEXT(""));
-
-	FGLTFJsonCameraIndex GetOrAddCamera(const UCameraComponent* CameraComponent, const FString& DesiredName = TEXT(""));
-	FGLTFJsonLightIndex GetOrAddLight(const ULightComponent* LightComponent, const FString& DesiredName = TEXT(""));
-
-	FGLTFJsonBackdropIndex GetOrAddBackdrop(const AActor* Actor, const FString& DesiredName = TEXT(""));
-	FGLTFJsonLevelVariantSetsIndex GetOrAddLevelVariantSets(const ALevelVariantSetsActor* LevelVariantSetsActor, const FString& DesiredName = TEXT(""));
-	FGLTFJsonLightMapIndex GetOrAddLightMap(const UStaticMeshComponent* StaticMeshComponent, const FString& DesiredName = TEXT(""));
-	FGLTFJsonHotspotIndex GetOrAddHotspot(const UGLTFInteractionHotspotComponent* HotspotComponent, const FString& DesiredName = TEXT(""));
+	FGLTFJsonAccessorIndex GetOrAddPositionAccessor(const FPositionVertexBuffer* VertexBuffer);
+	FGLTFJsonAccessorIndex GetOrAddColorAccessor(const FColorVertexBuffer* VertexBuffer);
+	FGLTFJsonAccessorIndex GetOrAddNormalAccessor(const FStaticMeshVertexBuffer* VertexBuffer);
+	FGLTFJsonAccessorIndex GetOrAddTangentAccessor(const FStaticMeshVertexBuffer* VertexBuffer);
+	FGLTFJsonAccessorIndex GetOrAddUVAccessor(const FStaticMeshVertexBuffer* VertexBuffer, int32 UVIndex);
+	FGLTFJsonAccessorIndex GetOrAddJointAccessor(const FSkinWeightVertexBuffer* VertexBuffer, int32 JointsGroupIndex, FGLTFBoneMap BoneMap);
+	FGLTFJsonAccessorIndex GetOrAddWeightAccessor(const FSkinWeightVertexBuffer* VertexBuffer, int32 WeightsGroupIndex);
+
+	FGLTFJsonBufferViewIndex GetOrAddIndexBufferView(const FRawStaticIndexBuffer* IndexBuffer);
+	FGLTFJsonAccessorIndex GetOrAddIndexAccessor(const FStaticMeshSection* MeshSection, const FRawStaticIndexBuffer* IndexBuffer);
+	FGLTFJsonMeshIndex GetOrAddMesh(const UStaticMesh* StaticMesh, int32 LODIndex = -1, const FColorVertexBuffer* OverrideVertexColors = nullptr, const FGLTFMaterialArray& OverrideMaterials = {});
+	FGLTFJsonMeshIndex GetOrAddMesh(const UStaticMeshComponent* StaticMeshComponent);
+
+	FGLTFJsonBufferViewIndex GetOrAddIndexBufferView(const FMultiSizeIndexContainer* IndexContainer);
+	FGLTFJsonAccessorIndex GetOrAddIndexAccessor(const FSkelMeshRenderSection* MeshSection, const FMultiSizeIndexContainer* IndexContainer);
+	FGLTFJsonMeshIndex GetOrAddMesh(const USkeletalMesh* SkeletalMesh, int32 LODIndex = -1, const FColorVertexBuffer* OverrideVertexColors = nullptr, const FSkinWeightVertexBuffer* OverrideSkinWeights = nullptr, const FGLTFMaterialArray& OverrideMaterials = {});
+	FGLTFJsonMeshIndex GetOrAddMesh(const USkeletalMeshComponent* SkeletalMeshComponent);
+
+	FGLTFJsonMaterialIndex GetOrAddMaterial(const UMaterialInterface* Material);
+	FGLTFJsonSamplerIndex GetOrAddSampler(const UTexture* Texture);
+	FGLTFJsonTextureIndex GetOrAddTexture(const UTexture2D* Texture);
+	FGLTFJsonTextureIndex GetOrAddTexture(const UTextureCube* Texture, ECubeFace CubeFace);
+	FGLTFJsonTextureIndex GetOrAddTexture(const UTextureRenderTarget2D* Texture);
+	FGLTFJsonTextureIndex GetOrAddTexture(const UTextureRenderTargetCube* Texture, ECubeFace CubeFace);
+
+	FGLTFJsonNodeIndex GetOrAddNode(const USceneComponent* SceneComponent);
+	FGLTFJsonNodeIndex GetOrAddNode(const AActor* Actor);
+	FGLTFJsonSceneIndex GetOrAddScene(const ULevel* Level);
+	FGLTFJsonSceneIndex GetOrAddScene(const UWorld* World);
+
+	FGLTFJsonCameraIndex GetOrAddCamera(const UCameraComponent* CameraComponent);
+	FGLTFJsonLightIndex GetOrAddLight(const ULightComponent* LightComponent);
+
+	FGLTFJsonBackdropIndex GetOrAddBackdrop(const AActor* Actor);
+	FGLTFJsonLevelVariantSetsIndex GetOrAddLevelVariantSets(const ALevelVariantSetsActor* LevelVariantSetsActor);
+	FGLTFJsonLightMapIndex GetOrAddLightMap(const UStaticMeshComponent* StaticMeshComponent);
+	FGLTFJsonHotspotIndex GetOrAddHotspot(const UGLTFInteractionHotspotComponent* HotspotComponent);
 
 private:
 
diff --git a/Engine/Plugins/Enterprise/GLTFExporter/MigratedFromGit/Source/GLTFExporter/Private/Converters/GLTFBackdropConverters.cpp b/Engine/Plugins/Enterprise/GLTFExporter/MigratedFromGit/Source/GLTFExporter/Private/Converters/GLTFBackdropConverters.cpp
--- a/Engine/Plugins/Enterprise/GLTFExporter/MigratedFromGit/Source/GLTFExporter/Private/Converters/GLTFBackdropConverters.cpp
+++ b/Engine/Plugins/Enterprise/GLTFExporter/MigratedFromGit/Source/GLTFExporter/Private/Converters/GLTFBackdropConverters.cpp
@@ -5,7 +5,7 @@
 #include "Converters/GLTFActorUtility.h"
 #include "Builders/GLTFContainerBuilder.h"
 
-FGLTFJsonBackdropIndex FGLTFBackdropConverter::Convert(const FString& Name, const AActor* Actor)
+FGLTFJsonBackdropIndex FGLTFBackdropConverter::Convert(const AActor* Actor)
 {
 	const UBlueprint* Blueprint = FGLTFActorUtility::GetBlueprintFromActor(Actor);
 	if (!FGLTFActorUtility::IsHDRIBackdropBlueprint(Blueprint))
@@ -14,7 +14,7 @@ FGLTFJsonBackdropIndex FGLTFBackdropConverter::Convert(const FString& Name, cons
 	}
 
 	FGLTFJsonBackdrop JsonBackdrop;
-	JsonBackdrop.Name = Name;
+	Actor->GetName(JsonBackdrop.Name);
 
 	const UStaticMesh* Mesh;
 	if (FGLTFActorUtility::TryGetPropertyValue(Actor, TEXT("Mesh"), Mesh))
@@ -36,9 +36,7 @@ FGLTFJsonBackdropIndex FGLTFBackdropConverter::Convert(const FString& Name, cons
 		{
 			const ECubeFace CubeFace = static_cast<ECubeFace>(CubeFaceIndex);
 			const EGLTFJsonCubeFace JsonCubeFace = FGLTFConverterUtility::ConvertCubeFace(CubeFace);
-			const int32 JsonCubeFaceIndex = static_cast<int32>(JsonCubeFace);
-
-			JsonBackdrop.Cubemap[JsonCubeFaceIndex] = Builder.GetOrAddTexture(Cubemap, CubeFace, Cubemap->GetName() + TEXT("_") + FGLTFJsonUtility::ToString(JsonCubeFace));
+			JsonBackdrop.Cubemap[static_cast<int32>(JsonCubeFace)] = Builder.GetOrAddTexture(Cubemap, CubeFace);
 		}
 	}
 	else
diff --git a/Engine/Plugins/Enterprise/GLTFExporter/MigratedFromGit/Source/GLTFExporter/Private/Converters/GLTFBackdropConverters.h b/Engine/Plugins/Enterprise/GLTFExporter/MigratedFromGit/Source/GLTFExporter/Private/Converters/GLTFBackdropConverters.h
--- a/Engine/Plugins/Enterprise/GLTFExporter/MigratedFromGit/Source/GLTFExporter/Private/Converters/GLTFBackdropConverters.h
+++ b/Engine/Plugins/Enterprise/GLTFExporter/MigratedFromGit/Source/GLTFExporter/Private/Converters/GLTFBackdropConverters.h
@@ -10,5 +10,5 @@ class FGLTFBackdropConverter final : public TGLTFConverter<FGLTFJsonBackdropInde
 {
 	using TGLTFConverter::TGLTFConverter;
 
-	FGLTFJsonBackdropIndex Convert(const FString& Name, const AActor* Actor) override;
+	FGLTFJsonBackdropIndex Convert(const AActor* Actor) override;
 };
diff --git a/Engine/Plugins/Enterprise/GLTFExporter/MigratedFromGit/Source/GLTFExporter/Private/Converters/GLTFConverter.h b/Engine/Plugins/Enterprise/GLTFExporter/MigratedFromGit/Source/GLTFExporter/Private/Converters/GLTFConverter.h
--- a/Engine/Plugins/Enterprise/GLTFExporter/MigratedFromGit/Source/GLTFExporter/Private/Converters/GLTFConverter.h
+++ b/Engine/Plugins/Enterprise/GLTFExporter/MigratedFromGit/Source/GLTFExporter/Private/Converters/GLTFConverter.h
@@ -31,16 +31,16 @@ public:
 		return {};
 	}
 
-	OutputType Add(const FString& DesiredName, InputTypes... Inputs)
+	OutputType Add(InputTypes... Inputs)
 	{
 		const InputKeyType InputKey(Forward<InputTypes>(Inputs)...);
-		OutputType NewOutput = Convert(DesiredName, Forward<InputTypes>(Inputs)...);
+		OutputType NewOutput = Convert(Forward<InputTypes>(Inputs)...);
 
 		SavedOutputs.Add(InputKey, NewOutput);
 		return NewOutput;
 	}
 
-	OutputType GetOrAdd(const FString& DesiredName, InputTypes... Inputs)
+	OutputType GetOrAdd(InputTypes... Inputs)
 	{
 		const InputKeyType InputKey(Forward<InputTypes>(Inputs)...);
 		if (OutputType* SavedOutput = SavedOutputs.Find(InputKey))
@@ -48,7 +48,7 @@ public:
 			return *SavedOutput;
 		}
 
-		OutputType NewOutput = Convert(DesiredName, Forward<InputTypes>(Inputs)...);
+		OutputType NewOutput = Convert(Forward<InputTypes>(Inputs)...);
 
 		SavedOutputs.Add(InputKey, NewOutput);
 		return NewOutput;
@@ -56,7 +56,7 @@ public:
 
 protected:
 
-	virtual OutputType Convert(const FString& Name, InputTypes... Inputs) = 0;
+	virtual OutputType Convert(InputTypes... Inputs) = 0;
 
 	FGLTFConvertBuilder& Builder;
 
diff --git a/Engine/Plugins/Enterprise/GLTFExporter/MigratedFromGit/Source/GLTFExporter/Private/Converters/GLTFConverterUtility.h b/Engine/Plugins/Enterprise/GLTFExporter/MigratedFromGit/Source/GLTFExporter/Private/Converters/GLTFConverterUtility.h
--- a/Engine/Plugins/Enterprise/GLTFExporter/MigratedFromGit/Source/GLTFExporter/Private/Converters/GLTFConverterUtility.h
+++ b/Engine/Plugins/Enterprise/GLTFExporter/MigratedFromGit/Source/GLTFExporter/Private/Converters/GLTFConverterUtility.h
@@ -11,7 +11,6 @@
 #include "Json/GLTFJsonMatrix4.h"
 #include "Json/GLTFJsonQuaternion.h"
 #include "Engine/EngineTypes.h"
-#include "Json/GLTFJsonCamera.h"
 
 struct FGLTFConverterUtility
 {
@@ -177,13 +176,4 @@ struct FGLTFConverterUtility
 		if (TIsSame<ComponentType, float >::Value) return EGLTFJsonComponentType::F32;
 		return EGLTFJsonComponentType::None;
 	}
-
-	template <typename EnumType>
-	static FString GetEnumDisplayName(EnumType Value)
-	{
-		static_assert(TIsEnum<EnumType>::Value, "Should only call this with enum types");
-		const UEnum* Enum = StaticEnum<EnumType>();
-		check(Enum != nullptr);
-		return Enum->GetDisplayNameTextByValue(Value).ToString();
-	}
 };
diff --git a/Engine/Plugins/Enterprise/GLTFExporter/MigratedFromGit/Source/GLTFExporter/Private/Converters/GLTFHotspotConverters.cpp b/Engine/Plugins/Enterprise/GLTFExporter/MigratedFromGit/Source/GLTFExporter/Private/Converters/GLTFHotspotConverters.cpp
--- a/Engine/Plugins/Enterprise/GLTFExporter/MigratedFromGit/Source/GLTFExporter/Private/Converters/GLTFHotspotConverters.cpp
+++ b/Engine/Plugins/Enterprise/GLTFExporter/MigratedFromGit/Source/GLTFExporter/Private/Converters/GLTFHotspotConverters.cpp
@@ -4,12 +4,12 @@
 #include "Json/FGLTFJsonHotspot.h"
 #include "Builders/GLTFContainerBuilder.h"
 
-FGLTFJsonHotspotIndex FGLTFHotspotComponentConverter::Convert(const FString& Name, const UGLTFInteractionHotspotComponent* HotspotComponent)
+FGLTFJsonHotspotIndex FGLTFHotspotComponentConverter::Convert(const UGLTFInteractionHotspotComponent* HotspotComponent)
 {
 	// TODO: should we warn and / or return INDEX_NONE if the hotspot has no valid Image assigned?
 
 	FGLTFJsonHotspot JsonHotspot;
-	JsonHotspot.Name = Name;
+	HotspotComponent->GetName(JsonHotspot.Name);
 	JsonHotspot.Animation = FGLTFJsonAnimationIndex(INDEX_NONE);	// TODO: assign the animation property once we have support for exporting animations
 	JsonHotspot.Image = Builder.GetOrAddTexture(HotspotComponent->Image);
 	JsonHotspot.HoveredImage = Builder.GetOrAddTexture(HotspotComponent->HoveredImage);
diff --git a/Engine/Plugins/Enterprise/GLTFExporter/MigratedFromGit/Source/GLTFExporter/Private/Converters/GLTFHotspotConverters.h b/Engine/Plugins/Enterprise/GLTFExporter/MigratedFromGit/Source/GLTFExporter/Private/Converters/GLTFHotspotConverters.h
--- a/Engine/Plugins/Enterprise/GLTFExporter/MigratedFromGit/Source/GLTFExporter/Private/Converters/GLTFHotspotConverters.h
+++ b/Engine/Plugins/Enterprise/GLTFExporter/MigratedFromGit/Source/GLTFExporter/Private/Converters/GLTFHotspotConverters.h
@@ -11,5 +11,5 @@ class FGLTFHotspotComponentConverter final : public TGLTFConverter<FGLTFJsonHots
 {
 	using TGLTFConverter::TGLTFConverter;
 
-	FGLTFJsonHotspotIndex Convert(const FString& Name, const UGLTFInteractionHotspotComponent* HotspotComponent) override;
+	FGLTFJsonHotspotIndex Convert(const UGLTFInteractionHotspotComponent* HotspotComponent) override;
 };
diff --git a/Engine/Plugins/Enterprise/GLTFExporter/MigratedFromGit/Source/GLTFExporter/Private/Converters/GLTFLevelConverters.cpp b/Engine/Plugins/Enterprise/GLTFExporter/MigratedFromGit/Source/GLTFExporter/Private/Converters/GLTFLevelConverters.cpp
--- a/Engine/Plugins/Enterprise/GLTFExporter/MigratedFromGit/Source/GLTFExporter/Private/Converters/GLTFLevelConverters.cpp
+++ b/Engine/Plugins/Enterprise/GLTFExporter/MigratedFromGit/Source/GLTFExporter/Private/Converters/GLTFLevelConverters.cpp
@@ -4,10 +4,11 @@
 #include "Builders/GLTFContainerBuilder.h"
 #include "Converters/GLTFConverterUtility.h"
 #include "Converters/GLTFActorUtility.h"
+#include "Converters/GLTFNameUtility.h"
 #include "Converters/GLTFCameraUtility.h"
 #include "Components/GLTFInteractionHotspotComponent.h"
 
-FGLTFJsonNodeIndex FGLTFSceneComponentConverter::Convert(const FString& Name, const USceneComponent* SceneComponent)
+FGLTFJsonNodeIndex FGLTFSceneComponentConverter::Convert(const USceneComponent* SceneComponent)
 {
 	const AActor* Owner = SceneComponent->GetOwner();
 	if (Owner == nullptr)
@@ -34,7 +35,7 @@ FGLTFJsonNodeIndex FGLTFSceneComponentConverter::Convert(const FString& Name, co
 
 	const FGLTFJsonNodeIndex NodeIndex = Builder.AddChildNode(ParentNodeIndex);
 	FGLTFJsonNode& Node = Builder.GetNode(NodeIndex);
-	Node.Name = Name.IsEmpty() ? Owner->GetName() + TEXT("_") + SceneComponent->GetName() : Name;
+	Node.Name = FGLTFNameUtility::GetName(SceneComponent);
 	Node.Translation = FGLTFConverterUtility::ConvertPosition(Translation, Builder.ExportOptions->ExportScale);
 	Node.Rotation = FGLTFConverterUtility::ConvertRotation(RelativeTransform.GetRotation());
 	Node.Scale = FGLTFConverterUtility::ConvertScale(Scale);
@@ -50,7 +51,7 @@ FGLTFJsonNodeIndex FGLTFSceneComponentConverter::Convert(const FString& Name, co
 		if (bExportNonUniformScale)
 		{
 			FGLTFJsonNode MeshNode;
-			MeshNode.Name = Node.Name + TEXT("_Mesh");
+			MeshNode.Name = FGLTFNameUtility::GetName(StaticMeshComponent);
 			MeshNode.Scale = ComponentNodeScale;
 			MeshNode.Mesh = Builder.GetOrAddMesh(StaticMeshComponent);
 			MeshNode.LightMap = Builder.GetOrAddLightMap(StaticMeshComponent);
@@ -67,7 +68,7 @@ FGLTFJsonNodeIndex FGLTFSceneComponentConverter::Convert(const FString& Name, co
 		if (bExportNonUniformScale)
 		{
 			FGLTFJsonNode MeshNode;
-			MeshNode.Name = Node.Name + TEXT("_Mesh");
+			MeshNode.Name = FGLTFNameUtility::GetName(SkeletalMeshComponent);
 			MeshNode.Scale = ComponentNodeScale;
 			MeshNode.Mesh = Builder.GetOrAddMesh(SkeletalMeshComponent);
 			Builder.AddChildComponentNode(NodeIndex, MeshNode);
@@ -84,17 +85,15 @@ FGLTFJsonNodeIndex FGLTFSceneComponentConverter::Convert(const FString& Name, co
 			if (bExportNonUniformScale)
 			{
 				FGLTFJsonNode HotspotNode;
-				HotspotNode.Name = Node.Name + TEXT("_Hotspot");
+				HotspotNode.Name = FGLTFNameUtility::GetName(HotspotComponent);
 				HotspotNode.Scale = ComponentNodeScale;
 
-				// TODO: use better / more unique name for the hotspot?
-				HotspotNode.Hotspot = Builder.GetOrAddHotspot(HotspotComponent, Owner->GetName());
+				HotspotNode.Hotspot = Builder.GetOrAddHotspot(HotspotComponent);
 				Builder.AddChildComponentNode(NodeIndex, HotspotNode);
 			}
 			else
 			{
-				// TODO: use better / more unique name for the hotspot?
-				Node.Hotspot = Builder.GetOrAddHotspot(HotspotComponent, Owner->GetName());
+				Node.Hotspot = Builder.GetOrAddHotspot(HotspotComponent);
 			}
 		}
 		else
@@ -108,10 +107,10 @@ FGLTFJsonNodeIndex FGLTFSceneComponentConverter::Convert(const FString& Name, co
 		{
 			// TODO: conversion of camera direction should be done in separate converter
 			FGLTFJsonNode CameraNode;
-			CameraNode.Name = Owner->GetName(); // TODO: choose a more unique name if owner is not ACameraActor
+			CameraNode.Name = FGLTFNameUtility::GetName(CameraComponent);
 			CameraNode.Rotation = FGLTFConverterUtility::ConvertCameraDirection();
 			CameraNode.Scale = ComponentNodeScale;
-			CameraNode.Camera = Builder.GetOrAddCamera(CameraComponent, CameraNode.Name);
+			CameraNode.Camera = Builder.GetOrAddCamera(CameraComponent);
 			Builder.AddChildComponentNode(NodeIndex, CameraNode);
 		}
 		else
@@ -125,10 +124,10 @@ FGLTFJsonNodeIndex FGLTFSceneComponentConverter::Convert(const FString& Name, co
 		{
 			// TODO: conversion of light direction should be done in separate converter
 			FGLTFJsonNode LightNode;
-			LightNode.Name = Owner->GetName(); // TODO: choose a more unique name if owner is not ALight
+			LightNode.Name = FGLTFNameUtility::GetName(LightComponent);
 			LightNode.Rotation = FGLTFConverterUtility::ConvertLightDirection();
 			LightNode.Scale = ComponentNodeScale;
-			LightNode.Light = Builder.GetOrAddLight(LightComponent, LightNode.Name);
+			LightNode.Light = Builder.GetOrAddLight(LightComponent);
 			Builder.AddChildComponentNode(NodeIndex, LightNode);
 		}
 		else
@@ -142,7 +141,7 @@ FGLTFJsonNodeIndex FGLTFSceneComponentConverter::Convert(const FString& Name, co
 	return NodeIndex;
 }
 
-FGLTFJsonNodeIndex FGLTFActorConverter::Convert(const FString& Name, const AActor* Actor)
+FGLTFJsonNodeIndex FGLTFActorConverter::Convert(const AActor* Actor)
 {
 	if (Builder.bSelectedActorsOnly && !Actor->IsSelected())
 	{
@@ -189,10 +188,18 @@ FGLTFJsonNodeIndex FGLTFActorConverter::Convert(const FString& Name, const AActo
 	return RootNodeIndex;
 }
 
-FGLTFJsonSceneIndex FGLTFLevelConverter::Convert(const FString& Name, const ULevel* Level)
+FGLTFJsonSceneIndex FGLTFLevelConverter::Convert(const ULevel* Level)
 {
 	FGLTFJsonScene Scene;
-	Scene.Name = Name;
+
+	if (const UWorld* World = Level->GetWorld())
+	{
+		World->GetName(Scene.Name);
+	}
+	else
+	{
+		Level->GetName(Scene.Name);
+	}
 
 	// TODO: export Level->Model ?
 
@@ -227,10 +234,10 @@ FGLTFJsonSceneIndex FGLTFLevelConverter::Convert(const FString& Name, const ULev
 	return Builder.AddScene(Scene);
 }
 
-FGLTFJsonCameraIndex FGLTFCameraComponentConverter::Convert(const FString& Name, const UCameraComponent* CameraComponent)
+FGLTFJsonCameraIndex FGLTFCameraComponentConverter::Convert(const UCameraComponent* CameraComponent)
 {
 	FGLTFJsonCamera Camera;
-	Camera.Name = Name;
+	Camera.Name = FGLTFNameUtility::GetName(CameraComponent);
 	Camera.Type = FGLTFConverterUtility::ConvertCameraType(CameraComponent->ProjectionMode);
 
 	FMinimalViewInfo DesiredView;
@@ -254,10 +261,11 @@ FGLTFJsonCameraIndex FGLTFCameraComponentConverter::Convert(const FString& Name,
 	return Builder.AddCamera(Camera);
 }
 
-FGLTFJsonLightIndex FGLTFLightComponentConverter::Convert(const FString& Name, const ULightComponent* LightComponent)
+FGLTFJsonLightIndex FGLTFLightComponentConverter::Convert(const ULightComponent* LightComponent)
 {
 	FGLTFJsonLight Light;
-	Light.Name = Name;
+
+	Light.Name = FGLTFNameUtility::GetName(LightComponent);
 	Light.Type = FGLTFConverterUtility::ConvertLightType(LightComponent->GetLightType());
 
 	if (Light.Type == EGLTFJsonLightType::None)
diff --git a/Engine/Plugins/Enterprise/GLTFExporter/MigratedFromGit/Source/GLTFExporter/Private/Converters/GLTFLevelConverters.h b/Engine/Plugins/Enterprise/GLTFExporter/MigratedFromGit/Source/GLTFExporter/Private/Converters/GLTFLevelConverters.h
--- a/Engine/Plugins/Enterprise/GLTFExporter/MigratedFromGit/Source/GLTFExporter/Private/Converters/GLTFLevelConverters.h
+++ b/Engine/Plugins/Enterprise/GLTFExporter/MigratedFromGit/Source/GLTFExporter/Private/Converters/GLTFLevelConverters.h
@@ -4,40 +4,39 @@
 
 #include "Json/GLTFJsonIndex.h"
 #include "Converters/GLTFConverter.h"
-#include "GLTFExportOptions.h"
 #include "Engine.h"
 
 class FGLTFSceneComponentConverter final : public TGLTFConverter<FGLTFJsonNodeIndex, const USceneComponent*>
 {
 	using TGLTFConverter::TGLTFConverter;
 
-	FGLTFJsonNodeIndex Convert(const FString& Name, const USceneComponent* SceneComponent) override;
+	FGLTFJsonNodeIndex Convert(const USceneComponent* SceneComponent) override;
 };
 
 class FGLTFActorConverter final : public TGLTFConverter<FGLTFJsonNodeIndex, const AActor*>
 {
 	using TGLTFConverter::TGLTFConverter;
 
-	FGLTFJsonNodeIndex Convert(const FString& Name, const AActor* Actor) override;
+	FGLTFJsonNodeIndex Convert(const AActor* Actor) override;
 };
 
 class FGLTFLevelConverter final : public TGLTFConverter<FGLTFJsonSceneIndex, const ULevel*>
 {
 	using TGLTFConverter::TGLTFConverter;
 
-	FGLTFJsonSceneIndex Convert(const FString& Name, const ULevel* Level) override;
+	FGLTFJsonSceneIndex Convert(const ULevel* Level) override;
 };
 
 class FGLTFCameraComponentConverter final : public TGLTFConverter<FGLTFJsonCameraIndex, const UCameraComponent*>
 {
 	using TGLTFConverter::TGLTFConverter;
 
-	FGLTFJsonCameraIndex Convert(const FString& Name, const UCameraComponent* CameraComponent) override;
+	FGLTFJsonCameraIndex Convert(const UCameraComponent* CameraComponent) override;
 };
 
 class FGLTFLightComponentConverter final : public TGLTFConverter<FGLTFJsonLightIndex, const ULightComponent*>
 {
 	using TGLTFConverter::TGLTFConverter;
 
-	FGLTFJsonLightIndex Convert(const FString& Name, const ULightComponent* LightComponent) override;
+	FGLTFJsonLightIndex Convert(const ULightComponent* LightComponent) override;
 };
diff --git a/Engine/Plugins/Enterprise/GLTFExporter/MigratedFromGit/Source/GLTFExporter/Private/Converters/GLTFLightMapConverters.cpp b/Engine/Plugins/Enterprise/GLTFExporter/MigratedFromGit/Source/GLTFExporter/Private/Converters/GLTFLightMapConverters.cpp
--- a/Engine/Plugins/Enterprise/GLTFExporter/MigratedFromGit/Source/GLTFExporter/Private/Converters/GLTFLightMapConverters.cpp
+++ b/Engine/Plugins/Enterprise/GLTFExporter/MigratedFromGit/Source/GLTFExporter/Private/Converters/GLTFLightMapConverters.cpp
@@ -5,7 +5,7 @@
 #include "Builders/GLTFContainerBuilder.h"
 #include "Engine/MapBuildDataRegistry.h"
 
-FGLTFJsonLightMapIndex FGLTFLightMapConverter::Convert(const FString& Name, const UStaticMeshComponent* StaticMeshComponent)
+FGLTFJsonLightMapIndex FGLTFLightMapConverter::Convert(const UStaticMeshComponent* StaticMeshComponent)
 {
 	const UStaticMesh* StaticMesh = StaticMeshComponent->GetStaticMesh();
 
@@ -67,7 +67,7 @@ FGLTFJsonLightMapIndex FGLTFLightMapConverter::Convert(const FString& Name, cons
 	const FVector4& LightMapScale = LightMapInteraction.GetScaleArray()[0];
 
 	FGLTFJsonLightMap JsonLightMap;
-	JsonLightMap.Name = Name.IsEmpty() ? StaticMeshComponent->GetName() : Name; // TODO: use better name (similar to light and camera)
+	StaticMeshComponent->GetName(JsonLightMap.Name); // TODO: use better name (similar to light and camera)
 	JsonLightMap.Texture.Index = TextureIndex;
 	JsonLightMap.Texture.TexCoord = CoordinateIndex;
 	JsonLightMap.LightMapScale = { LightMapScale.X, LightMapScale.Y, LightMapScale.Z, LightMapScale.W };
diff --git a/Engine/Plugins/Enterprise/GLTFExporter/MigratedFromGit/Source/GLTFExporter/Private/Converters/GLTFLightMapConverters.h b/Engine/Plugins/Enterprise/GLTFExporter/MigratedFromGit/Source/GLTFExporter/Private/Converters/GLTFLightMapConverters.h
--- a/Engine/Plugins/Enterprise/GLTFExporter/MigratedFromGit/Source/GLTFExporter/Private/Converters/GLTFLightMapConverters.h
+++ b/Engine/Plugins/Enterprise/GLTFExporter/MigratedFromGit/Source/GLTFExporter/Private/Converters/GLTFLightMapConverters.h
@@ -10,5 +10,5 @@ class FGLTFLightMapConverter final : public TGLTFConverter<FGLTFJsonLightMapInde
 {
 	using TGLTFConverter::TGLTFConverter;
 
-	FGLTFJsonLightMapIndex Convert(const FString& Name, const UStaticMeshComponent* StaticMeshComponent) override;
+	FGLTFJsonLightMapIndex Convert(const UStaticMeshComponent* StaticMeshComponent) override;
 };
diff --git a/Engine/Plugins/Enterprise/GLTFExporter/MigratedFromGit/Source/GLTFExporter/Private/Converters/GLTFMaterialConverters.cpp b/Engine/Plugins/Enterprise/GLTFExporter/MigratedFromGit/Source/GLTFExporter/Private/Converters/GLTFMaterialConverters.cpp
--- a/Engine/Plugins/Enterprise/GLTFExporter/MigratedFromGit/Source/GLTFExporter/Private/Converters/GLTFMaterialConverters.cpp
+++ b/Engine/Plugins/Enterprise/GLTFExporter/MigratedFromGit/Source/GLTFExporter/Private/Converters/GLTFMaterialConverters.cpp
@@ -2,6 +2,7 @@
 
 #include "Converters/GLTFMaterialConverters.h"
 #include "Converters/GLTFConverterUtility.h"
+#include "Converters/GLTFNameUtility.h"
 #include "Converters/GLTFMaterialUtility.h"
 #include "Builders/GLTFContainerBuilder.h"
 #include "Materials/MaterialExpressionConstant.h"
@@ -43,7 +44,7 @@ namespace
 	const TArray<FLinearColor> ClearCoatRoughnessInputMasks = { ClearCoatRoughnessMask };
 } // anonymous namespace
 
-FGLTFJsonMaterialIndex FGLTFMaterialConverter::Convert(const FString& Name, const UMaterialInterface* Material)
+FGLTFJsonMaterialIndex FGLTFMaterialConverter::Convert(const UMaterialInterface* Material)
 {
 	{
 		const UMaterial* ParentMaterial = Material->GetMaterial();
@@ -60,7 +61,7 @@ FGLTFJsonMaterialIndex FGLTFMaterialConverter::Convert(const FString& Name, cons
 	}
 
 	FGLTFJsonMaterial JsonMaterial;
-	JsonMaterial.Name = Name;
+	Material->GetName(JsonMaterial.Name);
 
 	// TODO: add support for additional blend modes (like Additive and Modulate)?
 	JsonMaterial.AlphaMode = FGLTFConverterUtility::ConvertBlendMode(Material->GetBlendMode());
@@ -70,7 +71,7 @@ FGLTFJsonMaterialIndex FGLTFMaterialConverter::Convert(const FString& Name, cons
 	if (!TryGetShadingModel(JsonMaterial.ShadingModel, Material))
 	{
 		JsonMaterial.ShadingModel = EGLTFJsonShadingModel::Default;
-		Builder.AddWarningMessage(FString::Printf(TEXT("Material %s will be exported as shading model %s"), *Material->GetName(), *FGLTFConverterUtility::GetEnumDisplayName(MSM_DefaultLit)));
+		Builder.AddWarningMessage(FString::Printf(TEXT("Material %s will be exported as shading model %s"), *Material->GetName(), *FGLTFNameUtility::GetName(MSM_DefaultLit)));
 	}
 
 	if (JsonMaterial.ShadingModel != EGLTFJsonShadingModel::None)
@@ -193,21 +194,21 @@ bool FGLTFMaterialConverter::TryGetShadingModel(EGLTFJsonShadingModel& OutShadin
 
 	if (ConvertedShadingModel == EGLTFJsonShadingModel::None)
 	{
-		const FString ShadingModelName = FGLTFConverterUtility::GetEnumDisplayName(ShadingModel);
+		const FString ShadingModelName = FGLTFNameUtility::GetName(ShadingModel);
 		Builder.AddWarningMessage(FString::Printf(TEXT("Unsupported shading model (%s) in material %s"), *ShadingModelName, *Material->GetName()));
 		return false;
 	}
 
 	if (ConvertedShadingModel == EGLTFJsonShadingModel::Unlit && !Builder.ExportOptions->bExportUnlitMaterials)
 	{
-		const FString ShadingModelName = FGLTFConverterUtility::GetEnumDisplayName(ShadingModel);
+		const FString ShadingModelName = FGLTFNameUtility::GetName(ShadingModel);
 		Builder.AddWarningMessage(FString::Printf(TEXT("Shading model (%s) in material %s disabled by export options"), *ShadingModelName, *Material->GetName()));
 		return false;
 	}
 
 	if (ConvertedShadingModel == EGLTFJsonShadingModel::ClearCoat && !Builder.ExportOptions->bExportClearCoatMaterials)
 	{
-		const FString ShadingModelName = FGLTFConverterUtility::GetEnumDisplayName(ShadingModel);
+		const FString ShadingModelName = FGLTFNameUtility::GetName(ShadingModel);
 		Builder.AddWarningMessage(FString::Printf(TEXT("Shading model (%s) in material %s disabled by export options"), *ShadingModelName, *Material->GetName()));
 		return false;
 	}
diff --git a/Engine/Plugins/Enterprise/GLTFExporter/MigratedFromGit/Source/GLTFExporter/Private/Converters/GLTFMaterialConverters.h b/Engine/Plugins/Enterprise/GLTFExporter/MigratedFromGit/Source/GLTFExporter/Private/Converters/GLTFMaterialConverters.h
--- a/Engine/Plugins/Enterprise/GLTFExporter/MigratedFromGit/Source/GLTFExporter/Private/Converters/GLTFMaterialConverters.h
+++ b/Engine/Plugins/Enterprise/GLTFExporter/MigratedFromGit/Source/GLTFExporter/Private/Converters/GLTFMaterialConverters.h
@@ -13,7 +13,7 @@ class FGLTFMaterialConverter final : public TGLTFConverter<FGLTFJsonMaterialInde
 {
 	using TGLTFConverter::TGLTFConverter;
 
-	FGLTFJsonMaterialIndex Convert(const FString& Name, const UMaterialInterface* Material) override;
+	FGLTFJsonMaterialIndex Convert(const UMaterialInterface* Material) override;
 
 	bool TryGetShadingModel(EGLTFJsonShadingModel& ShadingModel, const UMaterialInterface* Material) const;
 
diff --git a/Engine/Plugins/Enterprise/GLTFExporter/MigratedFromGit/Source/GLTFExporter/Private/Converters/GLTFNameUtility.cpp b/Engine/Plugins/Enterprise/GLTFExporter/MigratedFromGit/Source/GLTFExporter/Private/Converters/GLTFNameUtility.cpp
new file mode 100644
--- /dev/null
+++ b/Engine/Plugins/Enterprise/GLTFExporter/MigratedFromGit/Source/GLTFExporter/Private/Converters/GLTFNameUtility.cpp
@@ -0,0 +1,78 @@
+// Copyright Epic Games, Inc. All Rights Reserved.
+
+#include "Converters/GLTFNameUtility.h"
+#include "Actors/GLTFInteractionHotspotActor.h"
+#include "Animation/SkeletalMeshActor.h"
+#include "Engine.h"
+
+namespace
+{
+	FString GetNameWithLODSuffix(const UObject* Object, int32 LODIndex)
+	{
+		FString Name;
+		Object->GetName(Name);
+
+		if (LODIndex != 0)
+		{
+			Name += TEXT("_LOD") + FString::FromInt(LODIndex);
+		}
+
+		return Name;
+	}
+
+	template <typename ActorType>
+    FString GetActorNameIfOwnerOfType(const UActorComponent* Component)
+	{
+		if (const ActorType* Owner = Cast<ActorType>(Component->GetOwner()))
+		{
+			return Owner->GetName();
+		}
+
+		return Component->GetName();
+	}
+}
+
+FString FGLTFNameUtility::GetName(const UStaticMesh* StaticMesh, int32 LODIndex)
+{
+	return GetNameWithLODSuffix(StaticMesh, LODIndex);
+}
+
+FString FGLTFNameUtility::GetName(const USkeletalMesh* SkeletalMesh, int32 LODIndex)
+{
+	return GetNameWithLODSuffix(SkeletalMesh, LODIndex);
+}
+
+FString FGLTFNameUtility::GetName(const USceneComponent* Component)
+{
+	if (const AActor* Owner = Component->GetOwner())
+	{
+		return Owner->GetName() + TEXT("_") + Component->GetName();
+	}
+
+	return Component->GetName();
+}
+
+FString FGLTFNameUtility::GetName(const UStaticMeshComponent* Component)
+{
+	return GetActorNameIfOwnerOfType<AStaticMeshActor>(Component);
+}
+
+FString FGLTFNameUtility::GetName(const USkeletalMeshComponent* Component)
+{
+	return GetActorNameIfOwnerOfType<ASkeletalMeshActor>(Component);
+}
+
+FString FGLTFNameUtility::GetName(const ULightComponent* Component)
+{
+	return GetActorNameIfOwnerOfType<ALight>(Component);
+}
+
+FString FGLTFNameUtility::GetName(const UCameraComponent* Component)
+{
+	return GetActorNameIfOwnerOfType<ACameraActor>(Component);
+}
+
+FString FGLTFNameUtility::GetName(const UGLTFInteractionHotspotComponent* Component)
+{
+	return GetActorNameIfOwnerOfType<AGLTFInteractionHotspotActor>(Component);
+}
diff --git a/Engine/Plugins/Enterprise/GLTFExporter/MigratedFromGit/Source/GLTFExporter/Private/Converters/GLTFNameUtility.h b/Engine/Plugins/Enterprise/GLTFExporter/MigratedFromGit/Source/GLTFExporter/Private/Converters/GLTFNameUtility.h
new file mode 100644
--- /dev/null
+++ b/Engine/Plugins/Enterprise/GLTFExporter/MigratedFromGit/Source/GLTFExporter/Private/Converters/GLTFNameUtility.h
@@ -0,0 +1,31 @@
+// Copyright Epic Games, Inc. All Rights Reserved.
+
+#pragma once
+
+#include "CoreMinimal.h"
+
+class ULightComponent;
+class UCameraComponent;
+class UGLTFInteractionHotspotComponent;
+
+struct FGLTFNameUtility
+{
+	template <typename EnumType, typename = typename TEnableIf<TIsEnum<EnumType>::Value>::Type>
+	static FString GetName(EnumType Value)
+	{
+		const UEnum* Enum = StaticEnum<EnumType>();
+		check(Enum != nullptr);
+		return Enum->GetDisplayNameTextByValue(Value).ToString();
+	}
+
+	static FString GetName(const UStaticMesh* StaticMesh, int32 LODIndex);
+	static FString GetName(const USkeletalMesh* SkeletalMesh, int32 LODIndex);
+
+	static FString GetName(const USceneComponent* Component);
+
+	static FString GetName(const UStaticMeshComponent* Component);
+	static FString GetName(const USkeletalMeshComponent* Component);
+	static FString GetName(const ULightComponent* Component);
+	static FString GetName(const UCameraComponent* Component);
+	static FString GetName(const UGLTFInteractionHotspotComponent* Component);
+};
diff --git a/Engine/Plugins/Enterprise/GLTFExporter/MigratedFromGit/Source/GLTFExporter/Private/Converters/GLTFSkeletalMeshConverters.cpp b/Engine/Plugins/Enterprise/GLTFExporter/MigratedFromGit/Source/GLTFExporter/Private/Converters/GLTFSkeletalMeshConverters.cpp
--- a/Engine/Plugins/Enterprise/GLTFExporter/MigratedFromGit/Source/GLTFExporter/Private/Converters/GLTFSkeletalMeshConverters.cpp
+++ b/Engine/Plugins/Enterprise/GLTFExporter/MigratedFromGit/Source/GLTFExporter/Private/Converters/GLTFSkeletalMeshConverters.cpp
@@ -2,11 +2,12 @@
 
 #include "Converters/GLTFSkeletalMeshConverters.h"
 #include "Converters/GLTFConverterUtility.h"
+#include "Converters/GLTFNameUtility.h"
 #include "Builders/GLTFConvertBuilder.h"
 #include "Rendering/MultiSizeIndexContainer.h"
 #include "Rendering/SkeletalMeshRenderData.h"
 
-FGLTFJsonBufferViewIndex FGLTFIndexContainerConverter::Convert(const FString& Name, const FMultiSizeIndexContainer* IndexContainer)
+FGLTFJsonBufferViewIndex FGLTFIndexContainerConverter::Convert(const FMultiSizeIndexContainer* IndexContainer)
 {
 	const FRawStaticIndexBuffer16or32Interface* IndexBuffer = IndexContainer->GetIndexBuffer();
 
@@ -26,17 +27,17 @@ FGLTFJsonBufferViewIndex FGLTFIndexContainerConverter::Convert(const FString& Na
 			Indices[Index] = static_cast<uint16>(IndexBuffer->Get(Index));
 		}
 
-		return Builder.AddBufferView(Indices, Name, sizeof(uint16), EGLTFJsonBufferTarget::ElementArrayBuffer);
+		return Builder.AddBufferView(Indices, sizeof(uint16), EGLTFJsonBufferTarget::ElementArrayBuffer);
 	}
 	else
 	{
 		TArray<uint32> Indices;
 		IndexContainer->GetIndexBuffer(Indices);
-		return Builder.AddBufferView(Indices, Name, sizeof(uint32), EGLTFJsonBufferTarget::ElementArrayBuffer);
+		return Builder.AddBufferView(Indices, sizeof(uint32), EGLTFJsonBufferTarget::ElementArrayBuffer);
 	}
 }
 
-FGLTFJsonAccessorIndex FGLTFSkeletalMeshSectionConverter::Convert(const FString& Name, const FSkelMeshRenderSection* MeshSection, const FMultiSizeIndexContainer* IndexContainer)
+FGLTFJsonAccessorIndex FGLTFSkeletalMeshSectionConverter::Convert(const FSkelMeshRenderSection* MeshSection, const FMultiSizeIndexContainer* IndexContainer)
 {
 	const uint32 TriangleCount = MeshSection->NumTriangles;
 	if (TriangleCount == 0)
@@ -48,7 +49,6 @@ FGLTFJsonAccessorIndex FGLTFSkeletalMeshSectionConverter::Convert(const FString&
 	const bool bIs32Bit = IndexContainer->GetDataTypeSize() == sizeof(uint32);
 
 	FGLTFJsonAccessor JsonAccessor;
-	JsonAccessor.Name = Name;
 	JsonAccessor.BufferView = Builder.GetOrAddIndexBufferView(IndexContainer);
 	JsonAccessor.ByteOffset = FirstIndex * (bIs32Bit ? sizeof(uint32) : sizeof(uint16));
 	JsonAccessor.ComponentType = bIs32Bit ? EGLTFJsonComponentType::U32 : EGLTFJsonComponentType::U16;
@@ -58,7 +58,7 @@ FGLTFJsonAccessorIndex FGLTFSkeletalMeshSectionConverter::Convert(const FString&
 	return Builder.AddAccessor(JsonAccessor);
 }
 
-FGLTFJsonMeshIndex FGLTFSkeletalMeshConverter::Convert(const FString& Name, const USkeletalMesh* SkeletalMesh, int32 LODIndex, const FColorVertexBuffer* OverrideVertexColors, const FSkinWeightVertexBuffer* OverrideSkinWeights, FGLTFMaterialArray OverrideMaterials)
+FGLTFJsonMeshIndex FGLTFSkeletalMeshConverter::Convert(const USkeletalMesh* SkeletalMesh, int32 LODIndex, const FColorVertexBuffer* OverrideVertexColors, const FSkinWeightVertexBuffer* OverrideSkinWeights, FGLTFMaterialArray OverrideMaterials)
 {
 	const FSkeletalMeshRenderData* RenderData = SkeletalMesh->GetResourceForRendering();
 
@@ -70,7 +70,7 @@ FGLTFJsonMeshIndex FGLTFSkeletalMeshConverter::Convert(const FString& Name, cons
 	const FSkeletalMeshLODRenderData& MeshLOD = RenderData->LODRenderData[LODIndex];
 
 	FGLTFJsonMesh JsonMesh;
-	JsonMesh.Name = Name;
+	JsonMesh.Name = FGLTFNameUtility::GetName(SkeletalMesh, LODIndex);
 
 	const FPositionVertexBuffer* PositionBuffer = &MeshLOD.StaticVertexBuffers.PositionVertexBuffer;
 	const FStaticMeshVertexBuffer* VertexBuffer = &MeshLOD.StaticVertexBuffers.StaticMeshVertexBuffer;
@@ -78,22 +78,22 @@ FGLTFJsonMeshIndex FGLTFSkeletalMeshConverter::Convert(const FString& Name, cons
 	const FSkinWeightVertexBuffer* SkinWeightBuffer = OverrideSkinWeights != nullptr ? OverrideSkinWeights : MeshLOD.GetSkinWeightVertexBuffer();
 
 	FGLTFJsonAttributes JsonAttributes;
-	JsonAttributes.Position = Builder.GetOrAddPositionAccessor(PositionBuffer, Name + TEXT("_Positions"));
+	JsonAttributes.Position = Builder.GetOrAddPositionAccessor(PositionBuffer);
 
 	if (Builder.ExportOptions->bExportVertexColors)
 	{
-		JsonAttributes.Color0 = Builder.GetOrAddColorAccessor(ColorBuffer, Name + TEXT("_Colors"));
+		JsonAttributes.Color0 = Builder.GetOrAddColorAccessor(ColorBuffer);
 	}
 
-	JsonAttributes.Normal = Builder.GetOrAddNormalAccessor(VertexBuffer, Name + TEXT("_Normals"));
-	JsonAttributes.Tangent = Builder.GetOrAddTangentAccessor(VertexBuffer, Name + TEXT("_Tangents"));
+	JsonAttributes.Normal = Builder.GetOrAddNormalAccessor(VertexBuffer);
+	JsonAttributes.Tangent = Builder.GetOrAddTangentAccessor(VertexBuffer);
 
 	const uint32 UVCount = VertexBuffer->GetNumTexCoords();
 	JsonAttributes.TexCoords.AddUninitialized(UVCount);
 
 	for (uint32 UVIndex = 0; UVIndex < UVCount; ++UVIndex)
 	{
-		JsonAttributes.TexCoords[UVIndex] = Builder.GetOrAddUVAccessor(VertexBuffer, UVIndex, Name + TEXT("_UV") + FString::FromInt(UVIndex) + TEXT("s"));
+		JsonAttributes.TexCoords[UVIndex] = Builder.GetOrAddUVAccessor(VertexBuffer, UVIndex);
 	}
 
 	/* TODO: enable export of vertex skin data when crash bug fixed
@@ -112,7 +112,7 @@ FGLTFJsonMeshIndex FGLTFSkeletalMeshConverter::Convert(const FString& Name, cons
 	*/
 
 	const FMultiSizeIndexContainer* IndexContainer = &MeshLOD.MultiSizeIndexContainer;
-	Builder.GetOrAddIndexBufferView(IndexContainer, Name + TEXT("_Indices"));
+	Builder.GetOrAddIndexBufferView(IndexContainer);
 
 	const int32 SectionCount = MeshLOD.RenderSections.Num();
 	JsonMesh.Primitives.AddDefaulted(SectionCount);
@@ -123,8 +123,7 @@ FGLTFJsonMeshIndex FGLTFSkeletalMeshConverter::Convert(const FString& Name, cons
 		JsonPrimitive.Attributes = JsonAttributes;
 
 		const FSkelMeshRenderSection& Section = MeshLOD.RenderSections[SectionIndex];
-		JsonPrimitive.Indices = Builder.GetOrAddIndexAccessor(&Section, IndexContainer,
-			Name + (SectionCount == 1 ? TEXT("_Indices") : TEXT("_Indices_Section") + FString::FromInt(SectionIndex)));
+		JsonPrimitive.Indices = Builder.GetOrAddIndexAccessor(&Section, IndexContainer);
 
 		const int32 MaterialIndex = Section.MaterialIndex;
 		const UMaterialInterface* Material = OverrideMaterials.IsValidIndex(MaterialIndex) && OverrideMaterials[MaterialIndex] != nullptr ?
diff --git a/Engine/Plugins/Enterprise/GLTFExporter/MigratedFromGit/Source/GLTFExporter/Private/Converters/GLTFSkeletalMeshConverters.h b/Engine/Plugins/Enterprise/GLTFExporter/MigratedFromGit/Source/GLTFExporter/Private/Converters/GLTFSkeletalMeshConverters.h
--- a/Engine/Plugins/Enterprise/GLTFExporter/MigratedFromGit/Source/GLTFExporter/Private/Converters/GLTFSkeletalMeshConverters.h
+++ b/Engine/Plugins/Enterprise/GLTFExporter/MigratedFromGit/Source/GLTFExporter/Private/Converters/GLTFSkeletalMeshConverters.h
@@ -13,19 +13,19 @@ class FGLTFIndexContainerConverter final : public TGLTFConverter<FGLTFJsonBuffer
 {
 	using TGLTFConverter::TGLTFConverter;
 
-	FGLTFJsonBufferViewIndex Convert(const FString& Name, const FMultiSizeIndexContainer* IndexContainer) override;
+	FGLTFJsonBufferViewIndex Convert(const FMultiSizeIndexContainer* IndexContainer) override;
 };
 
 class FGLTFSkeletalMeshSectionConverter final : public TGLTFConverter<FGLTFJsonAccessorIndex, const FSkelMeshRenderSection*, const FMultiSizeIndexContainer*>
 {
 	using TGLTFConverter::TGLTFConverter;
 
-	FGLTFJsonAccessorIndex Convert(const FString& Name, const FSkelMeshRenderSection* MeshSection, const FMultiSizeIndexContainer* IndexContainer) override;
+	FGLTFJsonAccessorIndex Convert(const FSkelMeshRenderSection* MeshSection, const FMultiSizeIndexContainer* IndexContainer) override;
 };
 
 class FGLTFSkeletalMeshConverter final : public TGLTFConverter<FGLTFJsonMeshIndex, const USkeletalMesh*, int32, const FColorVertexBuffer*, const FSkinWeightVertexBuffer*, FGLTFMaterialArray>
 {
 	using TGLTFConverter::TGLTFConverter;
 
-	FGLTFJsonMeshIndex Convert(const FString& Name, const USkeletalMesh* SkeletalMesh, int32 LODIndex, const FColorVertexBuffer* OverrideVertexColors, const FSkinWeightVertexBuffer* OverrideSkinWeights, FGLTFMaterialArray OverrideMaterials) override;
+	FGLTFJsonMeshIndex Convert(const USkeletalMesh* SkeletalMesh, int32 LODIndex, const FColorVertexBuffer* OverrideVertexColors, const FSkinWeightVertexBuffer* OverrideSkinWeights, FGLTFMaterialArray OverrideMaterials) override;
 };
diff --git a/Engine/Plugins/Enterprise/GLTFExporter/MigratedFromGit/Source/GLTFExporter/Private/Converters/GLTFStaticMeshConverters.cpp b/Engine/Plugins/Enterprise/GLTFExporter/MigratedFromGit/Source/GLTFExporter/Private/Converters/GLTFStaticMeshConverters.cpp
--- a/Engine/Plugins/Enterprise/GLTFExporter/MigratedFromGit/Source/GLTFExporter/Private/Converters/GLTFStaticMeshConverters.cpp
+++ b/Engine/Plugins/Enterprise/GLTFExporter/MigratedFromGit/Source/GLTFExporter/Private/Converters/GLTFStaticMeshConverters.cpp
@@ -2,9 +2,10 @@
 
 #include "Converters/GLTFStaticMeshConverters.h"
 #include "Converters/GLTFConverterUtility.h"
+#include "Converters/GLTFNameUtility.h"
 #include "Builders/GLTFConvertBuilder.h"
 
-FGLTFJsonBufferViewIndex FGLTFIndexBufferConverter::Convert(const FString& Name, const FRawStaticIndexBuffer* IndexBuffer)
+FGLTFJsonBufferViewIndex FGLTFIndexBufferConverter::Convert(const FRawStaticIndexBuffer* IndexBuffer)
 {
 	if (IndexBuffer->GetNumIndices() == 0)
 	{
@@ -15,17 +16,17 @@ FGLTFJsonBufferViewIndex FGLTFIndexBufferConverter::Convert(const FString& Name,
 	{
 		TArray<uint32> Indices;
 		IndexBuffer->GetCopy(Indices);
-		return Builder.AddBufferView(Indices, Name, sizeof(uint32), EGLTFJsonBufferTarget::ElementArrayBuffer);
+		return Builder.AddBufferView(Indices, sizeof(uint32), EGLTFJsonBufferTarget::ElementArrayBuffer);
 	}
 	else
 	{
 		const uint16* IndexData = IndexBuffer->AccessStream16();
 		const int32 IndexDataSize = IndexBuffer->GetIndexDataSize();
-		return Builder.AddBufferView(IndexData, IndexDataSize, Name, sizeof(uint16), EGLTFJsonBufferTarget::ElementArrayBuffer);
+		return Builder.AddBufferView(IndexData, IndexDataSize, sizeof(uint16), EGLTFJsonBufferTarget::ElementArrayBuffer);
 	}
 }
 
-FGLTFJsonAccessorIndex FGLTFStaticMeshSectionConverter::Convert(const FString& Name, const FStaticMeshSection* MeshSection, const FRawStaticIndexBuffer* IndexBuffer)
+FGLTFJsonAccessorIndex FGLTFStaticMeshSectionConverter::Convert(const FStaticMeshSection* MeshSection, const FRawStaticIndexBuffer* IndexBuffer)
 {
 	const uint32 TriangleCount = MeshSection->NumTriangles;
 	if (TriangleCount == 0)
@@ -37,7 +38,6 @@ FGLTFJsonAccessorIndex FGLTFStaticMeshSectionConverter::Convert(const FString& N
 	const bool bIs32Bit = IndexBuffer->Is32Bit();
 
 	FGLTFJsonAccessor JsonAccessor;
-	JsonAccessor.Name = Name;
 	JsonAccessor.BufferView = Builder.GetOrAddIndexBufferView(IndexBuffer);
 	JsonAccessor.ByteOffset = FirstIndex * (bIs32Bit ? sizeof(uint32) : sizeof(uint16));
 	JsonAccessor.ComponentType = bIs32Bit ? EGLTFJsonComponentType::U32 : EGLTFJsonComponentType::U16;
@@ -47,7 +47,7 @@ FGLTFJsonAccessorIndex FGLTFStaticMeshSectionConverter::Convert(const FString& N
 	return Builder.AddAccessor(JsonAccessor);
 }
 
-FGLTFJsonMeshIndex FGLTFStaticMeshConverter::Convert(const FString& Name, const UStaticMesh* StaticMesh, int32 LODIndex, const FColorVertexBuffer* OverrideVertexColors, FGLTFMaterialArray OverrideMaterials)
+FGLTFJsonMeshIndex FGLTFStaticMeshConverter::Convert(const UStaticMesh* StaticMesh, int32 LODIndex, const FColorVertexBuffer* OverrideVertexColors, FGLTFMaterialArray OverrideMaterials)
 {
 	if (LODIndex < 0 || StaticMesh->GetNumLODs() <= LODIndex)
 	{
@@ -57,33 +57,33 @@ FGLTFJsonMeshIndex FGLTFStaticMeshConverter::Convert(const FString& Name, const
 	const FStaticMeshLODResources& MeshLOD = StaticMesh->GetLODForExport(LODIndex);
 
 	FGLTFJsonMesh JsonMesh;
-	JsonMesh.Name = Name;
+	JsonMesh.Name = FGLTFNameUtility::GetName(StaticMesh, LODIndex);
 
 	const FPositionVertexBuffer* PositionBuffer = &MeshLOD.VertexBuffers.PositionVertexBuffer;
 	const FStaticMeshVertexBuffer* VertexBuffer = &MeshLOD.VertexBuffers.StaticMeshVertexBuffer;
 	const FColorVertexBuffer* ColorBuffer = OverrideVertexColors != nullptr ? OverrideVertexColors : &MeshLOD.VertexBuffers.ColorVertexBuffer;
 
 	FGLTFJsonAttributes JsonAttributes;
-	JsonAttributes.Position = Builder.GetOrAddPositionAccessor(PositionBuffer, Name + TEXT("_Positions"));
+	JsonAttributes.Position = Builder.GetOrAddPositionAccessor(PositionBuffer);
 
 	if (Builder.ExportOptions->bExportVertexColors)
 	{
-		JsonAttributes.Color0 = Builder.GetOrAddColorAccessor(ColorBuffer, Name + TEXT("_Colors"));
+		JsonAttributes.Color0 = Builder.GetOrAddColorAccessor(ColorBuffer);
 	}
 
-	JsonAttributes.Normal = Builder.GetOrAddNormalAccessor(VertexBuffer, Name + TEXT("_Normals"));
-	JsonAttributes.Tangent = Builder.GetOrAddTangentAccessor(VertexBuffer, Name + TEXT("_Tangents"));
+	JsonAttributes.Normal = Builder.GetOrAddNormalAccessor(VertexBuffer);
+	JsonAttributes.Tangent = Builder.GetOrAddTangentAccessor(VertexBuffer);
 
 	const uint32 UVCount = VertexBuffer->GetNumTexCoords();
 	JsonAttributes.TexCoords.AddUninitialized(UVCount);
 
 	for (uint32 UVIndex = 0; UVIndex < UVCount; ++UVIndex)
 	{
-		JsonAttributes.TexCoords[UVIndex] = Builder.GetOrAddUVAccessor(VertexBuffer, UVIndex, Name + TEXT("_UV") + FString::FromInt(UVIndex) + TEXT("s"));
+		JsonAttributes.TexCoords[UVIndex] = Builder.GetOrAddUVAccessor(VertexBuffer, UVIndex);
 	}
 
 	const FRawStaticIndexBuffer* IndexBuffer = &MeshLOD.IndexBuffer;
-	Builder.GetOrAddIndexBufferView(IndexBuffer, Name + TEXT("_Indices"));
+	Builder.GetOrAddIndexBufferView(IndexBuffer);
 
 	const int32 SectionCount = MeshLOD.Sections.Num();
 	JsonMesh.Primitives.AddDefaulted(SectionCount);
@@ -94,8 +94,7 @@ FGLTFJsonMeshIndex FGLTFStaticMeshConverter::Convert(const FString& Name, const
 		JsonPrimitive.Attributes = JsonAttributes;
 
 		const FStaticMeshSection& Section = MeshLOD.Sections[SectionIndex];
-		JsonPrimitive.Indices = Builder.GetOrAddIndexAccessor(&Section, IndexBuffer,
-			Name + (SectionCount == 1 ? TEXT("_Indices") : TEXT("_Indices_Section") + FString::FromInt(SectionIndex)));
+		JsonPrimitive.Indices = Builder.GetOrAddIndexAccessor(&Section, IndexBuffer);
 
 		const int32 MaterialIndex = Section.MaterialIndex;
 		const UMaterialInterface* Material = OverrideMaterials.IsValidIndex(MaterialIndex) && OverrideMaterials[MaterialIndex] != nullptr ?
diff --git a/Engine/Plugins/Enterprise/GLTFExporter/MigratedFromGit/Source/GLTFExporter/Private/Converters/GLTFStaticMeshConverters.h b/Engine/Plugins/Enterprise/GLTFExporter/MigratedFromGit/Source/GLTFExporter/Private/Converters/GLTFStaticMeshConverters.h
--- a/Engine/Plugins/Enterprise/GLTFExporter/MigratedFromGit/Source/GLTFExporter/Private/Converters/GLTFStaticMeshConverters.h
+++ b/Engine/Plugins/Enterprise/GLTFExporter/MigratedFromGit/Source/GLTFExporter/Private/Converters/GLTFStaticMeshConverters.h
@@ -11,19 +11,19 @@ class FGLTFIndexBufferConverter final : public TGLTFConverter<FGLTFJsonBufferVie
 {
 	using TGLTFConverter::TGLTFConverter;
 
-	FGLTFJsonBufferViewIndex Convert(const FString& Name, const FRawStaticIndexBuffer* IndexBuffer) override;
+	FGLTFJsonBufferViewIndex Convert(const FRawStaticIndexBuffer* IndexBuffer) override;
 };
 
 class FGLTFStaticMeshSectionConverter final : public TGLTFConverter<FGLTFJsonAccessorIndex, const FStaticMeshSection*, const FRawStaticIndexBuffer*>
 {
 	using TGLTFConverter::TGLTFConverter;
 
-	FGLTFJsonAccessorIndex Convert(const FString& Name, const FStaticMeshSection* MeshSection, const FRawStaticIndexBuffer* IndexBuffer) override;
+	FGLTFJsonAccessorIndex Convert(const FStaticMeshSection* MeshSection, const FRawStaticIndexBuffer* IndexBuffer) override;
 };
 
 class FGLTFStaticMeshConverter final : public TGLTFConverter<FGLTFJsonMeshIndex, const UStaticMesh*, int32, const FColorVertexBuffer*, FGLTFMaterialArray>
 {
 	using TGLTFConverter::TGLTFConverter;
 
-	FGLTFJsonMeshIndex Convert(const FString& Name, const UStaticMesh* StaticMesh, int32 LODIndex, const FColorVertexBuffer* OverrideVertexColors, FGLTFMaterialArray OverrideMaterials) override;
+	FGLTFJsonMeshIndex Convert(const UStaticMesh* StaticMesh, int32 LODIndex, const FColorVertexBuffer* OverrideVertexColors, FGLTFMaterialArray OverrideMaterials) override;
 };
diff --git a/Engine/Plugins/Enterprise/GLTFExporter/MigratedFromGit/Source/GLTFExporter/Private/Converters/GLTFTextureConverters.cpp b/Engine/Plugins/Enterprise/GLTFExporter/MigratedFromGit/Source/GLTFExporter/Private/Converters/GLTFTextureConverters.cpp
--- a/Engine/Plugins/Enterprise/GLTFExporter/MigratedFromGit/Source/GLTFExporter/Private/Converters/GLTFTextureConverters.cpp
+++ b/Engine/Plugins/Enterprise/GLTFExporter/MigratedFromGit/Source/GLTFExporter/Private/Converters/GLTFTextureConverters.cpp
@@ -5,12 +5,12 @@
 #include "Converters/GLTFConverterUtility.h"
 #include "Converters/GLTFTextureUtility.h"
 
-FGLTFJsonSamplerIndex FGLTFTextureSamplerConverter::Convert(const FString& Name, const UTexture* Texture)
+FGLTFJsonSamplerIndex FGLTFTextureSamplerConverter::Convert(const UTexture* Texture)
 {
 	// TODO: maybe we should reuse existing samplers?
 
 	FGLTFJsonSampler JsonSampler;
-	JsonSampler.Name = Name;
+	Texture->GetName(JsonSampler.Name);
 
 	if (Texture->IsA<ULightMapTexture2D>())
 	{
@@ -41,12 +41,12 @@ FGLTFJsonSamplerIndex FGLTFTextureSamplerConverter::Convert(const FString& Name,
 	return Builder.AddSampler(JsonSampler);
 }
 
-FGLTFJsonTextureIndex FGLTFTexture2DConverter::Convert(const FString& Name, const UTexture2D* Texture2D)
+FGLTFJsonTextureIndex FGLTFTexture2DConverter::Convert(const UTexture2D* Texture2D)
 {
 	FGLTFJsonImageIndex ImageIndex;
 
 	FGLTFJsonTexture JsonTexture;
-	JsonTexture.Name = Name;
+	Texture2D->GetName(JsonTexture.Name);
 
 	const FIntPoint Size = { Texture2D->GetSizeX(), Texture2D->GetSizeY() };
 	const bool bPreferSourceExport = Builder.ExportOptions->bExportSourceTextures;
@@ -122,10 +122,10 @@ FGLTFJsonTextureIndex FGLTFTexture2DConverter::Convert(const FString& Name, cons
 	return Builder.AddTexture(JsonTexture);
 }
 
-FGLTFJsonTextureIndex FGLTFTextureCubeConverter::Convert(const FString& Name, const UTextureCube* TextureCube, ECubeFace CubeFace)
+FGLTFJsonTextureIndex FGLTFTextureCubeConverter::Convert(const UTextureCube* TextureCube, ECubeFace CubeFace)
 {
 	FGLTFJsonTexture JsonTexture;
-	JsonTexture.Name = Name;
+	JsonTexture.Name = TextureCube->GetName() + TEXT("_") + FGLTFJsonUtility::ToString(FGLTFConverterUtility::ConvertCubeFace(CubeFace));
 
 	// TODO: add optimized "happy path" if cube face doesn't need rotation and has suitable pixel format
 
@@ -156,10 +156,10 @@ FGLTFJsonTextureIndex FGLTFTextureCubeConverter::Convert(const FString& Name, co
 	return Builder.AddTexture(JsonTexture);
 }
 
-FGLTFJsonTextureIndex FGLTFTextureRenderTarget2DConverter::Convert(const FString& Name, const UTextureRenderTarget2D* RenderTarget2D)
+FGLTFJsonTextureIndex FGLTFTextureRenderTarget2DConverter::Convert(const UTextureRenderTarget2D* RenderTarget2D)
 {
 	FGLTFJsonTexture JsonTexture;
-	JsonTexture.Name = Name;
+	RenderTarget2D->GetName(JsonTexture.Name);
 
 	TArray<FColor> Pixels;
 	if (!FGLTFTextureUtility::ReadEncodedPixels(RenderTarget2D, Pixels, JsonTexture.Encoding)) // TODO: use only encoding as specified by export options
@@ -174,10 +174,10 @@ FGLTFJsonTextureIndex FGLTFTextureRenderTarget2DConverter::Convert(const FString
 	return Builder.AddTexture(JsonTexture);
 }
 
-FGLTFJsonTextureIndex FGLTFTextureRenderTargetCubeConverter::Convert(const FString& Name, const UTextureRenderTargetCube* RenderTargetCube, ECubeFace CubeFace)
+FGLTFJsonTextureIndex FGLTFTextureRenderTargetCubeConverter::Convert(const UTextureRenderTargetCube* RenderTargetCube, ECubeFace CubeFace)
 {
 	FGLTFJsonTexture JsonTexture;
-	JsonTexture.Name = Name;
+	JsonTexture.Name = RenderTargetCube->GetName() + TEXT("_") + FGLTFJsonUtility::ToString(FGLTFConverterUtility::ConvertCubeFace(CubeFace));
 
 	// TODO: add optimized "happy path" if cube face doesn't need rotation
 
diff --git a/Engine/Plugins/Enterprise/GLTFExporter/MigratedFromGit/Source/GLTFExporter/Private/Converters/GLTFTextureConverters.h b/Engine/Plugins/Enterprise/GLTFExporter/MigratedFromGit/Source/GLTFExporter/Private/Converters/GLTFTextureConverters.h
--- a/Engine/Plugins/Enterprise/GLTFExporter/MigratedFromGit/Source/GLTFExporter/Private/Converters/GLTFTextureConverters.h
+++ b/Engine/Plugins/Enterprise/GLTFExporter/MigratedFromGit/Source/GLTFExporter/Private/Converters/GLTFTextureConverters.h
@@ -10,33 +10,33 @@ class FGLTFTextureSamplerConverter final : public TGLTFConverter<FGLTFJsonSample
 {
 	using TGLTFConverter::TGLTFConverter;
 
-	FGLTFJsonSamplerIndex Convert(const FString& Name, const UTexture* Texture) override;
+	FGLTFJsonSamplerIndex Convert(const UTexture* Texture) override;
 };
 
 class FGLTFTexture2DConverter final : public TGLTFConverter<FGLTFJsonTextureIndex, const UTexture2D*>
 {
 	using TGLTFConverter::TGLTFConverter;
 
-	FGLTFJsonTextureIndex Convert(const FString& Name, const UTexture2D* Texture2D) override;
+	FGLTFJsonTextureIndex Convert(const UTexture2D* Texture2D) override;
 };
 
 class FGLTFTextureCubeConverter final : public TGLTFConverter<FGLTFJsonTextureIndex, const UTextureCube*, ECubeFace>
 {
 	using TGLTFConverter::TGLTFConverter;
 
-	FGLTFJsonTextureIndex Convert(const FString& Name, const UTextureCube* TextureCube, ECubeFace CubeFace) override;
+	FGLTFJsonTextureIndex Convert(const UTextureCube* TextureCube, ECubeFace CubeFace) override;
 };
 
 class FGLTFTextureRenderTarget2DConverter final : public TGLTFConverter<FGLTFJsonTextureIndex, const UTextureRenderTarget2D*>
 {
 	using TGLTFConverter::TGLTFConverter;
 
-	FGLTFJsonTextureIndex Convert(const FString& Name, const UTextureRenderTarget2D* RenderTarget2D) override;
+	FGLTFJsonTextureIndex Convert(const UTextureRenderTarget2D* RenderTarget2D) override;
 };
 
 class FGLTFTextureRenderTargetCubeConverter final : public TGLTFConverter<FGLTFJsonTextureIndex, const UTextureRenderTargetCube*, ECubeFace>
 {
 	using TGLTFConverter::TGLTFConverter;
 
-	FGLTFJsonTextureIndex Convert(const FString& Name, const UTextureRenderTargetCube* RenderTargetCube, ECubeFace CubeFace) override;
+	FGLTFJsonTextureIndex Convert(const UTextureRenderTargetCube* RenderTargetCube, ECubeFace CubeFace) override;
 };
diff --git a/Engine/Plugins/Enterprise/GLTFExporter/MigratedFromGit/Source/GLTFExporter/Private/Converters/GLTFVariantSetConverters.cpp b/Engine/Plugins/Enterprise/GLTFExporter/MigratedFromGit/Source/GLTFExporter/Private/Converters/GLTFVariantSetConverters.cpp
--- a/Engine/Plugins/Enterprise/GLTFExporter/MigratedFromGit/Source/GLTFExporter/Private/Converters/GLTFVariantSetConverters.cpp
+++ b/Engine/Plugins/Enterprise/GLTFExporter/MigratedFromGit/Source/GLTFExporter/Private/Converters/GLTFVariantSetConverters.cpp
@@ -23,7 +23,7 @@ namespace
 	};
 } // anonymous namespace
 
-FGLTFJsonLevelVariantSetsIndex FGLTFLevelVariantSetsConverter::Convert(const FString& Name, const ALevelVariantSetsActor* LevelVariantSetsActor)
+FGLTFJsonLevelVariantSetsIndex FGLTFLevelVariantSetsConverter::Convert(const ALevelVariantSetsActor* LevelVariantSetsActor)
 {
 	const ULevelVariantSets* LevelVariantSets = const_cast<ALevelVariantSetsActor*>(LevelVariantSetsActor)->GetLevelVariantSets(true);
 	if (LevelVariantSets == nullptr)
@@ -32,7 +32,7 @@ FGLTFJsonLevelVariantSetsIndex FGLTFLevelVariantSetsConverter::Convert(const FSt
 	}
 
 	FGLTFJsonLevelVariantSets JsonLevelVariantSets;
-	JsonLevelVariantSets.Name = Name.IsEmpty() ? LevelVariantSets->GetName() : Name;
+	LevelVariantSets->GetName(JsonLevelVariantSets.Name);
 
 	for (const UVariantSet* VariantSet: LevelVariantSets->GetVariantSets())
 	{
diff --git a/Engine/Plugins/Enterprise/GLTFExporter/MigratedFromGit/Source/GLTFExporter/Private/Converters/GLTFVariantSetConverters.h b/Engine/Plugins/Enterprise/GLTFExporter/MigratedFromGit/Source/GLTFExporter/Private/Converters/GLTFVariantSetConverters.h
--- a/Engine/Plugins/Enterprise/GLTFExporter/MigratedFromGit/Source/GLTFExporter/Private/Converters/GLTFVariantSetConverters.h
+++ b/Engine/Plugins/Enterprise/GLTFExporter/MigratedFromGit/Source/GLTFExporter/Private/Converters/GLTFVariantSetConverters.h
@@ -15,7 +15,7 @@ class FGLTFLevelVariantSetsConverter final : public TGLTFConverter<FGLTFJsonLeve
 {
 	using TGLTFConverter::TGLTFConverter;
 
-	FGLTFJsonLevelVariantSetsIndex Convert(const FString& Name, const ALevelVariantSetsActor* LevelVariantSetsActor) override;
+	FGLTFJsonLevelVariantSetsIndex Convert(const ALevelVariantSetsActor* LevelVariantSetsActor) override;
 
 	bool TryParseVariant(FGLTFJsonVariant& OutVariant, const UVariant* Variant) const;
 	bool TryParseVariantBinding(FGLTFJsonVariant& OutVariant, const UVariantObjectBinding* Binding) const;
diff --git a/Engine/Plugins/Enterprise/GLTFExporter/MigratedFromGit/Source/GLTFExporter/Private/Converters/GLTFVertexBufferConverters.cpp b/Engine/Plugins/Enterprise/GLTFExporter/MigratedFromGit/Source/GLTFExporter/Private/Converters/GLTFVertexBufferConverters.cpp
--- a/Engine/Plugins/Enterprise/GLTFExporter/MigratedFromGit/Source/GLTFExporter/Private/Converters/GLTFVertexBufferConverters.cpp
+++ b/Engine/Plugins/Enterprise/GLTFExporter/MigratedFromGit/Source/GLTFExporter/Private/Converters/GLTFVertexBufferConverters.cpp
@@ -6,7 +6,7 @@
 #include "Builders/GLTFConvertBuilder.h"
 #include "Math/NumericLimits.h"
 
-FGLTFJsonAccessorIndex FGLTFPositionVertexBufferConverter::Convert(const FString& Name, const FPositionVertexBuffer* VertexBuffer)
+FGLTFJsonAccessorIndex FGLTFPositionVertexBufferConverter::Convert(const FPositionVertexBuffer* VertexBuffer)
 {
 	const uint32 VertexCount = VertexBuffer->GetNumVertices();
 	if (VertexCount == 0)
@@ -38,8 +38,7 @@ FGLTFJsonAccessorIndex FGLTFPositionVertexBufferConverter::Convert(const FString
 	}
 
 	FGLTFJsonAccessor JsonAccessor;
-	JsonAccessor.Name = Name;
-	JsonAccessor.BufferView = Builder.AddBufferView(Positions, Name);
+	JsonAccessor.BufferView = Builder.AddBufferView(Positions);
 	JsonAccessor.ComponentType = EGLTFJsonComponentType::F32;
 	JsonAccessor.Count = VertexCount;
 	JsonAccessor.Type = EGLTFJsonAccessorType::Vec3;
@@ -55,7 +54,7 @@ FGLTFJsonAccessorIndex FGLTFPositionVertexBufferConverter::Convert(const FString
 	return Builder.AddAccessor(JsonAccessor);
 }
 
-FGLTFJsonAccessorIndex FGLTFColorVertexBufferConverter::Convert(const FString& Name, const FColorVertexBuffer* VertexBuffer)
+FGLTFJsonAccessorIndex FGLTFColorVertexBufferConverter::Convert(const FColorVertexBuffer* VertexBuffer)
 {
 	const uint32 VertexCount = VertexBuffer->GetNumVertices();
 	if (VertexCount == 0)
@@ -72,8 +71,7 @@ FGLTFJsonAccessorIndex FGLTFColorVertexBufferConverter::Convert(const FString& N
 	}
 
 	FGLTFJsonAccessor JsonAccessor;
-	JsonAccessor.Name = Name;
-	JsonAccessor.BufferView = Builder.AddBufferView(Colors, Name);
+	JsonAccessor.BufferView = Builder.AddBufferView(Colors);
 	JsonAccessor.ComponentType = EGLTFJsonComponentType::U8;
 	JsonAccessor.Count = VertexCount;
 	JsonAccessor.Type = EGLTFJsonAccessorType::Vec4;
@@ -82,7 +80,7 @@ FGLTFJsonAccessorIndex FGLTFColorVertexBufferConverter::Convert(const FString& N
 	return Builder.AddAccessor(JsonAccessor);
 }
 
-FGLTFJsonAccessorIndex FGLTFNormalVertexBufferConverter::Convert(const FString& Name, const FStaticMeshVertexBuffer* VertexBuffer)
+FGLTFJsonAccessorIndex FGLTFNormalVertexBufferConverter::Convert(const FStaticMeshVertexBuffer* VertexBuffer)
 {
 	const uint32 VertexCount = VertexBuffer->GetNumVertices();
 	if (VertexCount == 0)
@@ -99,8 +97,7 @@ FGLTFJsonAccessorIndex FGLTFNormalVertexBufferConverter::Convert(const FString&
 	}
 
 	FGLTFJsonAccessor JsonAccessor;
-	JsonAccessor.Name = Name;
-	JsonAccessor.BufferView = Builder.AddBufferView(Normals, Name);
+	JsonAccessor.BufferView = Builder.AddBufferView(Normals);
 	JsonAccessor.ComponentType = EGLTFJsonComponentType::F32;
 	JsonAccessor.Count = VertexCount;
 	JsonAccessor.Type = EGLTFJsonAccessorType::Vec3;
@@ -108,7 +105,7 @@ FGLTFJsonAccessorIndex FGLTFNormalVertexBufferConverter::Convert(const FString&
 	return Builder.AddAccessor(JsonAccessor);
 }
 
-FGLTFJsonAccessorIndex FGLTFTangentVertexBufferConverter::Convert(const FString& Name, const FStaticMeshVertexBuffer* VertexBuffer)
+FGLTFJsonAccessorIndex FGLTFTangentVertexBufferConverter::Convert(const FStaticMeshVertexBuffer* VertexBuffer)
 {
 	const uint32 VertexCount = VertexBuffer->GetNumVertices();
 	if (VertexCount == 0)
@@ -125,8 +122,7 @@ FGLTFJsonAccessorIndex FGLTFTangentVertexBufferConverter::Convert(const FString&
 	}
 
 	FGLTFJsonAccessor JsonAccessor;
-	JsonAccessor.Name = Name;
-	JsonAccessor.BufferView = Builder.AddBufferView(Tangents, Name);
+	JsonAccessor.BufferView = Builder.AddBufferView(Tangents);
 	JsonAccessor.ComponentType = EGLTFJsonComponentType::F32;
 	JsonAccessor.Count = VertexCount;
 	JsonAccessor.Type = EGLTFJsonAccessorType::Vec4;
@@ -134,7 +130,7 @@ FGLTFJsonAccessorIndex FGLTFTangentVertexBufferConverter::Convert(const FString&
 	return Builder.AddAccessor(JsonAccessor);
 }
 
-FGLTFJsonAccessorIndex FGLTFUVVertexBufferConverter::Convert(const FString& Name, const FStaticMeshVertexBuffer* VertexBuffer, int32 UVIndex)
+FGLTFJsonAccessorIndex FGLTFUVVertexBufferConverter::Convert(const FStaticMeshVertexBuffer* VertexBuffer, int32 UVIndex)
 {
 	const uint32 VertexCount = VertexBuffer->GetNumVertices();
 	if (VertexCount == 0 || UVIndex < 0 || VertexBuffer->GetNumTexCoords() <= static_cast<uint32>(UVIndex))
@@ -151,8 +147,7 @@ FGLTFJsonAccessorIndex FGLTFUVVertexBufferConverter::Convert(const FString& Name
 	}
 
 	FGLTFJsonAccessor JsonAccessor;
-	JsonAccessor.Name = Name;
-	JsonAccessor.BufferView = Builder.AddBufferView(UVs, Name);
+	JsonAccessor.BufferView = Builder.AddBufferView(UVs);
 	JsonAccessor.ComponentType = EGLTFJsonComponentType::F32;
 	JsonAccessor.Count = VertexCount;
 	JsonAccessor.Type = EGLTFJsonAccessorType::Vec2;
@@ -160,13 +155,13 @@ FGLTFJsonAccessorIndex FGLTFUVVertexBufferConverter::Convert(const FString& Name
 	return Builder.AddAccessor(JsonAccessor);
 }
 
-FGLTFJsonAccessorIndex FGLTFBoneIndexVertexBufferConverter::Convert(const FString& Name, const FSkinWeightVertexBuffer* VertexBuffer, int32 InfluenceOffset, FGLTFBoneMap BoneMap)
+FGLTFJsonAccessorIndex FGLTFBoneIndexVertexBufferConverter::Convert(const FSkinWeightVertexBuffer* VertexBuffer, int32 InfluenceOffset, FGLTFBoneMap BoneMap)
 {
-	return VertexBuffer->Use16BitBoneIndex() ? Convert<uint16>(Name, VertexBuffer, InfluenceOffset, BoneMap) : Convert<uint8>(Name, VertexBuffer, InfluenceOffset, BoneMap);
+	return VertexBuffer->Use16BitBoneIndex() ? Convert<uint16>(VertexBuffer, InfluenceOffset, BoneMap) : Convert<uint8>(VertexBuffer, InfluenceOffset, BoneMap);
 }
 
 template <typename IndexType>
-FGLTFJsonAccessorIndex FGLTFBoneIndexVertexBufferConverter::Convert(const FString& Name, const FSkinWeightVertexBuffer* VertexBuffer, int32 JointsGroupIndex, FGLTFBoneMap BoneMap)
+FGLTFJsonAccessorIndex FGLTFBoneIndexVertexBufferConverter::Convert(const FSkinWeightVertexBuffer* VertexBuffer, int32 JointsGroupIndex, FGLTFBoneMap BoneMap)
 {
 	const uint32 VertexCount = VertexBuffer->GetNumVertices();
 	const int32 MaxInfluenceCount = VertexBuffer->GetMaxBoneInfluences();
@@ -199,8 +194,7 @@ FGLTFJsonAccessorIndex FGLTFBoneIndexVertexBufferConverter::Convert(const FStrin
 	}
 
 	FGLTFJsonAccessor JsonAccessor;
-	JsonAccessor.Name = Name;
-	JsonAccessor.BufferView = Builder.AddBufferView(BoneIndices, Name);
+	JsonAccessor.BufferView = Builder.AddBufferView(BoneIndices);
 	JsonAccessor.ComponentType = FGLTFConverterUtility::GetComponentType<IndexType>();
 	JsonAccessor.Count = VertexCount;
 	JsonAccessor.Type = EGLTFJsonAccessorType::Vec4;
@@ -208,7 +202,7 @@ FGLTFJsonAccessorIndex FGLTFBoneIndexVertexBufferConverter::Convert(const FStrin
 	return Builder.AddAccessor(JsonAccessor);
 }
 
-FGLTFJsonAccessorIndex FGLTFBoneWeightVertexBufferConverter::Convert(const FString& Name, const FSkinWeightVertexBuffer* VertexBuffer, int32 WeightsGroupIndex)
+FGLTFJsonAccessorIndex FGLTFBoneWeightVertexBufferConverter::Convert(const FSkinWeightVertexBuffer* VertexBuffer, int32 WeightsGroupIndex)
 {
 	const uint32 VertexCount = VertexBuffer->GetNumVertices();
 	const int32 MaxInfluenceCount = VertexBuffer->GetMaxBoneInfluences();
@@ -238,8 +232,7 @@ FGLTFJsonAccessorIndex FGLTFBoneWeightVertexBufferConverter::Convert(const FStri
 	}
 
 	FGLTFJsonAccessor JsonAccessor;
-	JsonAccessor.Name = Name;
-	JsonAccessor.BufferView = Builder.AddBufferView(BoneWeights, Name);
+	JsonAccessor.BufferView = Builder.AddBufferView(BoneWeights);
 	JsonAccessor.ComponentType = EGLTFJsonComponentType::U8;
 	JsonAccessor.Count = VertexCount;
 	JsonAccessor.Type = EGLTFJsonAccessorType::Vec4;
diff --git a/Engine/Plugins/Enterprise/GLTFExporter/MigratedFromGit/Source/GLTFExporter/Private/Converters/GLTFVertexBufferConverters.h b/Engine/Plugins/Enterprise/GLTFExporter/MigratedFromGit/Source/GLTFExporter/Private/Converters/GLTFVertexBufferConverters.h
--- a/Engine/Plugins/Enterprise/GLTFExporter/MigratedFromGit/Source/GLTFExporter/Private/Converters/GLTFVertexBufferConverters.h
+++ b/Engine/Plugins/Enterprise/GLTFExporter/MigratedFromGit/Source/GLTFExporter/Private/Converters/GLTFVertexBufferConverters.h
@@ -11,50 +11,50 @@ class FGLTFPositionVertexBufferConverter final : public TGLTFConverter<FGLTFJson
 {
 	using TGLTFConverter::TGLTFConverter;
 
-	FGLTFJsonAccessorIndex Convert(const FString& Name, const FPositionVertexBuffer* VertexBuffer) override;
+	FGLTFJsonAccessorIndex Convert(const FPositionVertexBuffer* VertexBuffer) override;
 };
 
 class FGLTFColorVertexBufferConverter final : public TGLTFConverter<FGLTFJsonAccessorIndex, const FColorVertexBuffer*>
 {
 	using TGLTFConverter::TGLTFConverter;
 
-	FGLTFJsonAccessorIndex Convert(const FString& Name, const FColorVertexBuffer* VertexBuffer) override;
+	FGLTFJsonAccessorIndex Convert(const FColorVertexBuffer* VertexBuffer) override;
 };
 
 class FGLTFNormalVertexBufferConverter final : public TGLTFConverter<FGLTFJsonAccessorIndex, const FStaticMeshVertexBuffer*>
 {
 	using TGLTFConverter::TGLTFConverter;
 
-	FGLTFJsonAccessorIndex Convert(const FString& Name, const FStaticMeshVertexBuffer* VertexBuffer) override;
+	FGLTFJsonAccessorIndex Convert(const FStaticMeshVertexBuffer* VertexBuffer) override;
 };
 
 class FGLTFTangentVertexBufferConverter final : public TGLTFConverter<FGLTFJsonAccessorIndex, const FStaticMeshVertexBuffer*>
 {
 	using TGLTFConverter::TGLTFConverter;
 
-	FGLTFJsonAccessorIndex Convert(const FString& Name, const FStaticMeshVertexBuffer* VertexBuffer) override;
+	FGLTFJsonAccessorIndex Convert(const FStaticMeshVertexBuffer* VertexBuffer) override;
 };
 
 class FGLTFUVVertexBufferConverter final : public TGLTFConverter<FGLTFJsonAccessorIndex, const FStaticMeshVertexBuffer*, int32>
 {
 	using TGLTFConverter::TGLTFConverter;
 
-	FGLTFJsonAccessorIndex Convert(const FString& Name, const FStaticMeshVertexBuffer* VertexBuffer, int32 UVIndex) override;
+	FGLTFJsonAccessorIndex Convert(const FStaticMeshVertexBuffer* VertexBuffer, int32 UVIndex) override;
 };
 
 class FGLTFBoneIndexVertexBufferConverter final : public TGLTFConverter<FGLTFJsonAccessorIndex, const FSkinWeightVertexBuffer*, int32, FGLTFBoneMap>
 {
 	using TGLTFConverter::TGLTFConverter;
 
-	FGLTFJsonAccessorIndex Convert(const FString& Name, const FSkinWeightVertexBuffer* VertexBuffer, int32 InfluenceOffset, FGLTFBoneMap BoneMap) override;
+	FGLTFJsonAccessorIndex Convert(const FSkinWeightVertexBuffer* VertexBuffer, int32 InfluenceOffset, FGLTFBoneMap BoneMap) override;
 
 	template <typename IndexType>
-	FGLTFJsonAccessorIndex Convert(const FString& Name, const FSkinWeightVertexBuffer* VertexBuffer, int32 JointsGroupIndex, FGLTFBoneMap BoneMap);
+	FGLTFJsonAccessorIndex Convert(const FSkinWeightVertexBuffer* VertexBuffer, int32 JointsGroupIndex, FGLTFBoneMap BoneMap);
 };
 
 class FGLTFBoneWeightVertexBufferConverter final : public TGLTFConverter<FGLTFJsonAccessorIndex, const FSkinWeightVertexBuffer*, int32>
 {
 	using TGLTFConverter::TGLTFConverter;
 
-	FGLTFJsonAccessorIndex Convert(const FString& Name, const FSkinWeightVertexBuffer* VertexBuffer, int32 WeightsGroupIndex) override;
+	FGLTFJsonAccessorIndex Convert(const FSkinWeightVertexBuffer* VertexBuffer, int32 WeightsGroupIndex) override;
 };
diff --git a/Engine/Plugins/Enterprise/GLTFExporter/MigratedFromGit/Source/GLTFExporter/Private/Exporters/GLTFMaterialExporter.cpp b/Engine/Plugins/Enterprise/GLTFExporter/MigratedFromGit/Source/GLTFExporter/Private/Exporters/GLTFMaterialExporter.cpp
--- a/Engine/Plugins/Enterprise/GLTFExporter/MigratedFromGit/Source/GLTFExporter/Private/Exporters/GLTFMaterialExporter.cpp
+++ b/Engine/Plugins/Enterprise/GLTFExporter/MigratedFromGit/Source/GLTFExporter/Private/Exporters/GLTFMaterialExporter.cpp
@@ -27,7 +27,7 @@ bool UGLTFMaterialExporter::AddObject(FGLTFContainerBuilder& Builder, const UObj
 		const UStaticMesh* PreviewMesh = FGLTFExporterUtility::GetPreviewMesh(Material);
 		if (PreviewMesh != nullptr)
 		{
-			const FGLTFJsonMeshIndex MeshIndex = Builder.GetOrAddMesh(PreviewMesh, -1, nullptr, { Material }, "PreviewMesh");
+			const FGLTFJsonMeshIndex MeshIndex = Builder.GetOrAddMesh(PreviewMesh, -1, nullptr, { Material });
 
 			FGLTFJsonNode Node;
 			Node.Mesh = MeshIndex;
