[GLTFExporter] refactor: move builder reference in base converter to separate class GLTFBuilderContext

commit e844e6abb8fb31de42eb03267ba24f97aab4576f
Date: Sun, 20 Sep 2020 21:11:07 +0200

    refactor: move builder reference in base converter to separate class GLTFBuilderContext
    (cherry picked from commit 687a38ceb31a723b4a98c4a7c5813dc1d6fb0bb7​)

#skipci

diff --git a/Engine/Plugins/Enterprise/GLTFExporter/MigratedFromGit/Source/GLTFExporter/Private/Converters/GLTFBackdropConverters.h b/Engine/Plugins/Enterprise/GLTFExporter/MigratedFromGit/Source/GLTFExporter/Private/Converters/GLTFBackdropConverters.h
--- a/Engine/Plugins/Enterprise/GLTFExporter/MigratedFromGit/Source/GLTFExporter/Private/Converters/GLTFBackdropConverters.h
+++ b/Engine/Plugins/Enterprise/GLTFExporter/MigratedFromGit/Source/GLTFExporter/Private/Converters/GLTFBackdropConverters.h
@@ -4,11 +4,12 @@
 
 #include "Json/GLTFJsonIndex.h"
 #include "Converters/GLTFConverter.h"
+#include "Converters/GLTFBuilderContext.h"
 #include "Engine.h"
 
-class FGLTFBackdropConverter final : public TGLTFConverter<FGLTFJsonBackdropIndex, const AActor*>
+class FGLTFBackdropConverter : public FGLTFBuilderContext, public TGLTFConverter<FGLTFJsonBackdropIndex, const AActor*>
 {
-	using TGLTFConverter::TGLTFConverter;
+	using FGLTFBuilderContext::FGLTFBuilderContext;
 
-	FGLTFJsonBackdropIndex Convert(const AActor* Actor) override;
+	FGLTFJsonBackdropIndex Convert(const AActor* Actor) override final;
 };
diff --git a/Engine/Plugins/Enterprise/GLTFExporter/MigratedFromGit/Source/GLTFExporter/Private/Converters/GLTFBuilderContext.h b/Engine/Plugins/Enterprise/GLTFExporter/MigratedFromGit/Source/GLTFExporter/Private/Converters/GLTFBuilderContext.h
new file mode 100644
--- /dev/null
+++ b/Engine/Plugins/Enterprise/GLTFExporter/MigratedFromGit/Source/GLTFExporter/Private/Converters/GLTFBuilderContext.h
@@ -0,0 +1,19 @@
+// Copyright Epic Games, Inc. All Rights Reserved.
+
+#pragma once
+
+class FGLTFConvertBuilder;
+
+class FGLTFBuilderContext
+{
+public:
+
+	FGLTFBuilderContext(FGLTFConvertBuilder& Builder)
+        : Builder(Builder)
+	{
+	}
+
+protected:
+
+	FGLTFConvertBuilder& Builder;
+};
diff --git a/Engine/Plugins/Enterprise/GLTFExporter/MigratedFromGit/Source/GLTFExporter/Private/Converters/GLTFConverter.h b/Engine/Plugins/Enterprise/GLTFExporter/MigratedFromGit/Source/GLTFExporter/Private/Converters/GLTFConverter.h
--- a/Engine/Plugins/Enterprise/GLTFExporter/MigratedFromGit/Source/GLTFExporter/Private/Converters/GLTFConverter.h
+++ b/Engine/Plugins/Enterprise/GLTFExporter/MigratedFromGit/Source/GLTFExporter/Private/Converters/GLTFConverter.h
@@ -4,8 +4,6 @@
 
 #include "CoreMinimal.h"
 
-class FGLTFConvertBuilder;
-
 template <typename OutputType, typename... InputTypes>
 class TGLTFConverter
 {
@@ -13,11 +11,6 @@ class TGLTFConverter
 
 public:
 
-	TGLTFConverter(FGLTFConvertBuilder& Builder)
-		: Builder(Builder)
-	{
-	}
-
 	virtual ~TGLTFConverter() = default;
 
 	OutputType Get(InputTypes&&... Inputs) const
@@ -58,8 +51,6 @@ protected:
 
 	virtual OutputType Convert(InputTypes... Inputs) = 0;
 
-	FGLTFConvertBuilder& Builder;
-
 private:
 
 	TMap<InputKeyType, OutputType> SavedOutputs;
diff --git a/Engine/Plugins/Enterprise/GLTFExporter/MigratedFromGit/Source/GLTFExporter/Private/Converters/GLTFHotspotConverters.h b/Engine/Plugins/Enterprise/GLTFExporter/MigratedFromGit/Source/GLTFExporter/Private/Converters/GLTFHotspotConverters.h
--- a/Engine/Plugins/Enterprise/GLTFExporter/MigratedFromGit/Source/GLTFExporter/Private/Converters/GLTFHotspotConverters.h
+++ b/Engine/Plugins/Enterprise/GLTFExporter/MigratedFromGit/Source/GLTFExporter/Private/Converters/GLTFHotspotConverters.h
@@ -4,12 +4,13 @@
 
 #include "Json/GLTFJsonIndex.h"
 #include "Converters/GLTFConverter.h"
+#include "Converters/GLTFBuilderContext.h"
 #include "Components/GLTFInteractionHotspotComponent.h"
 #include "Engine.h"
 
-class FGLTFHotspotComponentConverter final : public TGLTFConverter<FGLTFJsonHotspotIndex, const UGLTFInteractionHotspotComponent*>
+class FGLTFHotspotComponentConverter : public FGLTFBuilderContext, public TGLTFConverter<FGLTFJsonHotspotIndex, const UGLTFInteractionHotspotComponent*>
 {
-	using TGLTFConverter::TGLTFConverter;
+	using FGLTFBuilderContext::FGLTFBuilderContext;
 
-	FGLTFJsonHotspotIndex Convert(const UGLTFInteractionHotspotComponent* HotspotComponent) override;
+	FGLTFJsonHotspotIndex Convert(const UGLTFInteractionHotspotComponent* HotspotComponent) override final;
 };
diff --git a/Engine/Plugins/Enterprise/GLTFExporter/MigratedFromGit/Source/GLTFExporter/Private/Converters/GLTFLevelConverters.h b/Engine/Plugins/Enterprise/GLTFExporter/MigratedFromGit/Source/GLTFExporter/Private/Converters/GLTFLevelConverters.h
--- a/Engine/Plugins/Enterprise/GLTFExporter/MigratedFromGit/Source/GLTFExporter/Private/Converters/GLTFLevelConverters.h
+++ b/Engine/Plugins/Enterprise/GLTFExporter/MigratedFromGit/Source/GLTFExporter/Private/Converters/GLTFLevelConverters.h
@@ -4,39 +4,40 @@
 
 #include "Json/GLTFJsonIndex.h"
 #include "Converters/GLTFConverter.h"
+#include "Converters/GLTFBuilderContext.h"
 #include "Engine.h"
 
-class FGLTFSceneComponentConverter final : public TGLTFConverter<FGLTFJsonNodeIndex, const USceneComponent*>
+class FGLTFSceneComponentConverter : public FGLTFBuilderContext, public TGLTFConverter<FGLTFJsonNodeIndex, const USceneComponent*>
 {
-	using TGLTFConverter::TGLTFConverter;
+	using FGLTFBuilderContext::FGLTFBuilderContext;
 
-	FGLTFJsonNodeIndex Convert(const USceneComponent* SceneComponent) override;
+	FGLTFJsonNodeIndex Convert(const USceneComponent* SceneComponent) override final;
 };
 
-class FGLTFActorConverter final : public TGLTFConverter<FGLTFJsonNodeIndex, const AActor*>
+class FGLTFActorConverter : public FGLTFBuilderContext, public TGLTFConverter<FGLTFJsonNodeIndex, const AActor*>
 {
-	using TGLTFConverter::TGLTFConverter;
+	using FGLTFBuilderContext::FGLTFBuilderContext;
 
-	FGLTFJsonNodeIndex Convert(const AActor* Actor) override;
+	FGLTFJsonNodeIndex Convert(const AActor* Actor) override final;
 };
 
-class FGLTFLevelConverter final : public TGLTFConverter<FGLTFJsonSceneIndex, const ULevel*>
+class FGLTFLevelConverter : public FGLTFBuilderContext, public TGLTFConverter<FGLTFJsonSceneIndex, const ULevel*>
 {
-	using TGLTFConverter::TGLTFConverter;
+	using FGLTFBuilderContext::FGLTFBuilderContext;
 
-	FGLTFJsonSceneIndex Convert(const ULevel* Level) override;
+	FGLTFJsonSceneIndex Convert(const ULevel* Level) override final;
 };
 
-class FGLTFCameraComponentConverter final : public TGLTFConverter<FGLTFJsonCameraIndex, const UCameraComponent*>
+class FGLTFCameraComponentConverter : public FGLTFBuilderContext, public TGLTFConverter<FGLTFJsonCameraIndex, const UCameraComponent*>
 {
-	using TGLTFConverter::TGLTFConverter;
+	using FGLTFBuilderContext::FGLTFBuilderContext;
 
-	FGLTFJsonCameraIndex Convert(const UCameraComponent* CameraComponent) override;
+	FGLTFJsonCameraIndex Convert(const UCameraComponent* CameraComponent) override final;
 };
 
-class FGLTFLightComponentConverter final : public TGLTFConverter<FGLTFJsonLightIndex, const ULightComponent*>
+class FGLTFLightComponentConverter : public FGLTFBuilderContext, public TGLTFConverter<FGLTFJsonLightIndex, const ULightComponent*>
 {
-	using TGLTFConverter::TGLTFConverter;
+	using FGLTFBuilderContext::FGLTFBuilderContext;
 
-	FGLTFJsonLightIndex Convert(const ULightComponent* LightComponent) override;
+	FGLTFJsonLightIndex Convert(const ULightComponent* LightComponent) override final;
 };
diff --git a/Engine/Plugins/Enterprise/GLTFExporter/MigratedFromGit/Source/GLTFExporter/Private/Converters/GLTFLightMapConverters.h b/Engine/Plugins/Enterprise/GLTFExporter/MigratedFromGit/Source/GLTFExporter/Private/Converters/GLTFLightMapConverters.h
--- a/Engine/Plugins/Enterprise/GLTFExporter/MigratedFromGit/Source/GLTFExporter/Private/Converters/GLTFLightMapConverters.h
+++ b/Engine/Plugins/Enterprise/GLTFExporter/MigratedFromGit/Source/GLTFExporter/Private/Converters/GLTFLightMapConverters.h
@@ -4,11 +4,12 @@
 
 #include "Json/GLTFJsonIndex.h"
 #include "Converters/GLTFConverter.h"
+#include "Converters/GLTFBuilderContext.h"
 #include "Engine.h"
 
-class FGLTFLightMapConverter final : public TGLTFConverter<FGLTFJsonLightMapIndex, const UStaticMeshComponent*>
+class FGLTFLightMapConverter : public FGLTFBuilderContext, public TGLTFConverter<FGLTFJsonLightMapIndex, const UStaticMeshComponent*>
 {
-	using TGLTFConverter::TGLTFConverter;
+	using FGLTFBuilderContext::FGLTFBuilderContext;
 
-	FGLTFJsonLightMapIndex Convert(const UStaticMeshComponent* StaticMeshComponent) override;
+	FGLTFJsonLightMapIndex Convert(const UStaticMeshComponent* StaticMeshComponent) override final;
 };
diff --git a/Engine/Plugins/Enterprise/GLTFExporter/MigratedFromGit/Source/GLTFExporter/Private/Converters/GLTFMaterialConverters.h b/Engine/Plugins/Enterprise/GLTFExporter/MigratedFromGit/Source/GLTFExporter/Private/Converters/GLTFMaterialConverters.h
--- a/Engine/Plugins/Enterprise/GLTFExporter/MigratedFromGit/Source/GLTFExporter/Private/Converters/GLTFMaterialConverters.h
+++ b/Engine/Plugins/Enterprise/GLTFExporter/MigratedFromGit/Source/GLTFExporter/Private/Converters/GLTFMaterialConverters.h
@@ -5,15 +5,16 @@
 #include "Json/GLTFJsonIndex.h"
 #include "Json/GLTFJsonMaterial.h"
 #include "Converters/GLTFConverter.h"
+#include "Converters/GLTFBuilderContext.h"
 #include "Engine.h"
 
 struct FGLTFPropertyBakeOutput;
 
-class FGLTFMaterialConverter final : public TGLTFConverter<FGLTFJsonMaterialIndex, const UMaterialInterface*>
+class FGLTFMaterialConverter : public FGLTFBuilderContext, public TGLTFConverter<FGLTFJsonMaterialIndex, const UMaterialInterface*>
 {
-	using TGLTFConverter::TGLTFConverter;
+	using FGLTFBuilderContext::FGLTFBuilderContext;
 
-	FGLTFJsonMaterialIndex Convert(const UMaterialInterface* Material) override;
+	FGLTFJsonMaterialIndex Convert(const UMaterialInterface* Material) override final;
 
 	bool TryGetShadingModel(EGLTFJsonShadingModel& ShadingModel, const UMaterialInterface* Material) const;
 
diff --git a/Engine/Plugins/Enterprise/GLTFExporter/MigratedFromGit/Source/GLTFExporter/Private/Converters/GLTFSkeletalMeshConverters.h b/Engine/Plugins/Enterprise/GLTFExporter/MigratedFromGit/Source/GLTFExporter/Private/Converters/GLTFSkeletalMeshConverters.h
--- a/Engine/Plugins/Enterprise/GLTFExporter/MigratedFromGit/Source/GLTFExporter/Private/Converters/GLTFSkeletalMeshConverters.h
+++ b/Engine/Plugins/Enterprise/GLTFExporter/MigratedFromGit/Source/GLTFExporter/Private/Converters/GLTFSkeletalMeshConverters.h
@@ -4,28 +4,29 @@
 
 #include "Json/GLTFJsonIndex.h"
 #include "Converters/GLTFConverter.h"
+#include "Converters/GLTFBuilderContext.h"
 #include "Converters/GLTFMaterialArray.h"
 #include "Engine.h"
 
 class FMultiSizeIndexContainer;
 
-class FGLTFIndexContainerConverter final : public TGLTFConverter<FGLTFJsonBufferViewIndex, const FMultiSizeIndexContainer*>
+class FGLTFIndexContainerConverter : public FGLTFBuilderContext, public TGLTFConverter<FGLTFJsonBufferViewIndex, const FMultiSizeIndexContainer*>
 {
-	using TGLTFConverter::TGLTFConverter;
+	using FGLTFBuilderContext::FGLTFBuilderContext;
 
-	FGLTFJsonBufferViewIndex Convert(const FMultiSizeIndexContainer* IndexContainer) override;
+	FGLTFJsonBufferViewIndex Convert(const FMultiSizeIndexContainer* IndexContainer) override final;
 };
 
-class FGLTFSkeletalMeshSectionConverter final : public TGLTFConverter<FGLTFJsonAccessorIndex, const FSkelMeshRenderSection*, const FMultiSizeIndexContainer*>
+class FGLTFSkeletalMeshSectionConverter : public FGLTFBuilderContext, public TGLTFConverter<FGLTFJsonAccessorIndex, const FSkelMeshRenderSection*, const FMultiSizeIndexContainer*>
 {
-	using TGLTFConverter::TGLTFConverter;
+	using FGLTFBuilderContext::FGLTFBuilderContext;
 
-	FGLTFJsonAccessorIndex Convert(const FSkelMeshRenderSection* MeshSection, const FMultiSizeIndexContainer* IndexContainer) override;
+	FGLTFJsonAccessorIndex Convert(const FSkelMeshRenderSection* MeshSection, const FMultiSizeIndexContainer* IndexContainer) override final;
 };
 
-class FGLTFSkeletalMeshConverter final : public TGLTFConverter<FGLTFJsonMeshIndex, const USkeletalMesh*, int32, const FColorVertexBuffer*, const FSkinWeightVertexBuffer*, FGLTFMaterialArray>
+class FGLTFSkeletalMeshConverter : public FGLTFBuilderContext, public TGLTFConverter<FGLTFJsonMeshIndex, const USkeletalMesh*, int32, const FColorVertexBuffer*, const FSkinWeightVertexBuffer*, FGLTFMaterialArray>
 {
-	using TGLTFConverter::TGLTFConverter;
+	using FGLTFBuilderContext::FGLTFBuilderContext;
 
-	FGLTFJsonMeshIndex Convert(const USkeletalMesh* SkeletalMesh, int32 LODIndex, const FColorVertexBuffer* OverrideVertexColors, const FSkinWeightVertexBuffer* OverrideSkinWeights, FGLTFMaterialArray OverrideMaterials) override;
+	FGLTFJsonMeshIndex Convert(const USkeletalMesh* SkeletalMesh, int32 LODIndex, const FColorVertexBuffer* OverrideVertexColors, const FSkinWeightVertexBuffer* OverrideSkinWeights, FGLTFMaterialArray OverrideMaterials) override final;
 };
diff --git a/Engine/Plugins/Enterprise/GLTFExporter/MigratedFromGit/Source/GLTFExporter/Private/Converters/GLTFStaticMeshConverters.h b/Engine/Plugins/Enterprise/GLTFExporter/MigratedFromGit/Source/GLTFExporter/Private/Converters/GLTFStaticMeshConverters.h
--- a/Engine/Plugins/Enterprise/GLTFExporter/MigratedFromGit/Source/GLTFExporter/Private/Converters/GLTFStaticMeshConverters.h
+++ b/Engine/Plugins/Enterprise/GLTFExporter/MigratedFromGit/Source/GLTFExporter/Private/Converters/GLTFStaticMeshConverters.h
@@ -4,26 +4,27 @@
 
 #include "Json/GLTFJsonIndex.h"
 #include "Converters/GLTFConverter.h"
+#include "Converters/GLTFBuilderContext.h"
 #include "Converters/GLTFMaterialArray.h"
 #include "Engine.h"
 
-class FGLTFIndexBufferConverter final : public TGLTFConverter<FGLTFJsonBufferViewIndex, const FRawStaticIndexBuffer*>
+class FGLTFIndexBufferConverter : public FGLTFBuilderContext, public TGLTFConverter<FGLTFJsonBufferViewIndex, const FRawStaticIndexBuffer*>
 {
-	using TGLTFConverter::TGLTFConverter;
+	using FGLTFBuilderContext::FGLTFBuilderContext;
 
-	FGLTFJsonBufferViewIndex Convert(const FRawStaticIndexBuffer* IndexBuffer) override;
+	FGLTFJsonBufferViewIndex Convert(const FRawStaticIndexBuffer* IndexBuffer) override final;
 };
 
-class FGLTFStaticMeshSectionConverter final : public TGLTFConverter<FGLTFJsonAccessorIndex, const FStaticMeshSection*, const FRawStaticIndexBuffer*>
+class FGLTFStaticMeshSectionConverter : public FGLTFBuilderContext, public TGLTFConverter<FGLTFJsonAccessorIndex, const FStaticMeshSection*, const FRawStaticIndexBuffer*>
 {
-	using TGLTFConverter::TGLTFConverter;
+	using FGLTFBuilderContext::FGLTFBuilderContext;
 
-	FGLTFJsonAccessorIndex Convert(const FStaticMeshSection* MeshSection, const FRawStaticIndexBuffer* IndexBuffer) override;
+	FGLTFJsonAccessorIndex Convert(const FStaticMeshSection* MeshSection, const FRawStaticIndexBuffer* IndexBuffer) override final;
 };
 
-class FGLTFStaticMeshConverter final : public TGLTFConverter<FGLTFJsonMeshIndex, const UStaticMesh*, int32, const FColorVertexBuffer*, FGLTFMaterialArray>
+class FGLTFStaticMeshConverter : public FGLTFBuilderContext, public TGLTFConverter<FGLTFJsonMeshIndex, const UStaticMesh*, int32, const FColorVertexBuffer*, FGLTFMaterialArray>
 {
-	using TGLTFConverter::TGLTFConverter;
+	using FGLTFBuilderContext::FGLTFBuilderContext;
 
-	FGLTFJsonMeshIndex Convert(const UStaticMesh* StaticMesh, int32 LODIndex, const FColorVertexBuffer* OverrideVertexColors, FGLTFMaterialArray OverrideMaterials) override;
+	FGLTFJsonMeshIndex Convert(const UStaticMesh* StaticMesh, int32 LODIndex, const FColorVertexBuffer* OverrideVertexColors, FGLTFMaterialArray OverrideMaterials) override final;
 };
diff --git a/Engine/Plugins/Enterprise/GLTFExporter/MigratedFromGit/Source/GLTFExporter/Private/Converters/GLTFTextureConverters.h b/Engine/Plugins/Enterprise/GLTFExporter/MigratedFromGit/Source/GLTFExporter/Private/Converters/GLTFTextureConverters.h
--- a/Engine/Plugins/Enterprise/GLTFExporter/MigratedFromGit/Source/GLTFExporter/Private/Converters/GLTFTextureConverters.h
+++ b/Engine/Plugins/Enterprise/GLTFExporter/MigratedFromGit/Source/GLTFExporter/Private/Converters/GLTFTextureConverters.h
@@ -4,39 +4,40 @@
 
 #include "Json/GLTFJsonIndex.h"
 #include "Converters/GLTFConverter.h"
+#include "Converters/GLTFBuilderContext.h"
 #include "Engine.h"
 
-class FGLTFTextureSamplerConverter final : public TGLTFConverter<FGLTFJsonSamplerIndex, const UTexture*>
+class FGLTFTextureSamplerConverter : public FGLTFBuilderContext, public TGLTFConverter<FGLTFJsonSamplerIndex, const UTexture*>
 {
-	using TGLTFConverter::TGLTFConverter;
+	using FGLTFBuilderContext::FGLTFBuilderContext;
 
-	FGLTFJsonSamplerIndex Convert(const UTexture* Texture) override;
+	FGLTFJsonSamplerIndex Convert(const UTexture* Texture) override final;
 };
 
-class FGLTFTexture2DConverter final : public TGLTFConverter<FGLTFJsonTextureIndex, const UTexture2D*>
+class FGLTFTexture2DConverter : public FGLTFBuilderContext, public TGLTFConverter<FGLTFJsonTextureIndex, const UTexture2D*>
 {
-	using TGLTFConverter::TGLTFConverter;
+	using FGLTFBuilderContext::FGLTFBuilderContext;
 
-	FGLTFJsonTextureIndex Convert(const UTexture2D* Texture2D) override;
+	FGLTFJsonTextureIndex Convert(const UTexture2D* Texture2D) override final;
 };
 
-class FGLTFTextureCubeConverter final : public TGLTFConverter<FGLTFJsonTextureIndex, const UTextureCube*, ECubeFace>
+class FGLTFTextureCubeConverter : public FGLTFBuilderContext, public TGLTFConverter<FGLTFJsonTextureIndex, const UTextureCube*, ECubeFace>
 {
-	using TGLTFConverter::TGLTFConverter;
+	using FGLTFBuilderContext::FGLTFBuilderContext;
 
-	FGLTFJsonTextureIndex Convert(const UTextureCube* TextureCube, ECubeFace CubeFace) override;
+	FGLTFJsonTextureIndex Convert(const UTextureCube* TextureCube, ECubeFace CubeFace) override final;
 };
 
-class FGLTFTextureRenderTarget2DConverter final : public TGLTFConverter<FGLTFJsonTextureIndex, const UTextureRenderTarget2D*>
+class FGLTFTextureRenderTarget2DConverter : public FGLTFBuilderContext, public TGLTFConverter<FGLTFJsonTextureIndex, const UTextureRenderTarget2D*>
 {
-	using TGLTFConverter::TGLTFConverter;
+	using FGLTFBuilderContext::FGLTFBuilderContext;
 
-	FGLTFJsonTextureIndex Convert(const UTextureRenderTarget2D* RenderTarget2D) override;
+	FGLTFJsonTextureIndex Convert(const UTextureRenderTarget2D* RenderTarget2D) override final;
 };
 
-class FGLTFTextureRenderTargetCubeConverter final : public TGLTFConverter<FGLTFJsonTextureIndex, const UTextureRenderTargetCube*, ECubeFace>
+class FGLTFTextureRenderTargetCubeConverter : public FGLTFBuilderContext, public TGLTFConverter<FGLTFJsonTextureIndex, const UTextureRenderTargetCube*, ECubeFace>
 {
-	using TGLTFConverter::TGLTFConverter;
+	using FGLTFBuilderContext::FGLTFBuilderContext;
 
-	FGLTFJsonTextureIndex Convert(const UTextureRenderTargetCube* RenderTargetCube, ECubeFace CubeFace) override;
+	FGLTFJsonTextureIndex Convert(const UTextureRenderTargetCube* RenderTargetCube, ECubeFace CubeFace) override final;
 };
diff --git a/Engine/Plugins/Enterprise/GLTFExporter/MigratedFromGit/Source/GLTFExporter/Private/Converters/GLTFVariantSetConverters.h b/Engine/Plugins/Enterprise/GLTFExporter/MigratedFromGit/Source/GLTFExporter/Private/Converters/GLTFVariantSetConverters.h
--- a/Engine/Plugins/Enterprise/GLTFExporter/MigratedFromGit/Source/GLTFExporter/Private/Converters/GLTFVariantSetConverters.h
+++ b/Engine/Plugins/Enterprise/GLTFExporter/MigratedFromGit/Source/GLTFExporter/Private/Converters/GLTFVariantSetConverters.h
@@ -5,17 +5,18 @@
 #include "Json/GLTFJsonIndex.h"
 #include "Json/GLTFJsonLevelVariantSets.h"
 #include "Converters/GLTFConverter.h"
+#include "Converters/GLTFBuilderContext.h"
 #include "Engine.h"
 #include "LevelVariantSetsActor.h"
 #include "Variant.h"
 #include "PropertyValue.h"
 #include "LevelVariantSets.h"
 
-class FGLTFLevelVariantSetsConverter final : public TGLTFConverter<FGLTFJsonLevelVariantSetsIndex, const ALevelVariantSetsActor*>
+class FGLTFLevelVariantSetsConverter : public FGLTFBuilderContext, public TGLTFConverter<FGLTFJsonLevelVariantSetsIndex, const ALevelVariantSetsActor*>
 {
-	using TGLTFConverter::TGLTFConverter;
+	using FGLTFBuilderContext::FGLTFBuilderContext;
 
-	FGLTFJsonLevelVariantSetsIndex Convert(const ALevelVariantSetsActor* LevelVariantSetsActor) override;
+	FGLTFJsonLevelVariantSetsIndex Convert(const ALevelVariantSetsActor* LevelVariantSetsActor) override final;
 
 	bool TryParseVariant(FGLTFJsonVariant& OutVariant, const UVariant* Variant) const;
 	bool TryParseVariantBinding(FGLTFJsonVariant& OutVariant, const UVariantObjectBinding* Binding) const;
diff --git a/Engine/Plugins/Enterprise/GLTFExporter/MigratedFromGit/Source/GLTFExporter/Private/Converters/GLTFVertexBufferConverters.h b/Engine/Plugins/Enterprise/GLTFExporter/MigratedFromGit/Source/GLTFExporter/Private/Converters/GLTFVertexBufferConverters.h
--- a/Engine/Plugins/Enterprise/GLTFExporter/MigratedFromGit/Source/GLTFExporter/Private/Converters/GLTFVertexBufferConverters.h
+++ b/Engine/Plugins/Enterprise/GLTFExporter/MigratedFromGit/Source/GLTFExporter/Private/Converters/GLTFVertexBufferConverters.h
@@ -4,57 +4,58 @@
 
 #include "Json/GLTFJsonIndex.h"
 #include "Converters/GLTFConverter.h"
+#include "Converters/GLTFBuilderContext.h"
 #include "Converters/GLTFBoneMap.h"
 #include "Engine.h"
 
-class FGLTFPositionVertexBufferConverter final : public TGLTFConverter<FGLTFJsonAccessorIndex, const FPositionVertexBuffer*>
+class FGLTFPositionVertexBufferConverter : public FGLTFBuilderContext, public TGLTFConverter<FGLTFJsonAccessorIndex, const FPositionVertexBuffer*>
 {
-	using TGLTFConverter::TGLTFConverter;
+	using FGLTFBuilderContext::FGLTFBuilderContext;
 
-	FGLTFJsonAccessorIndex Convert(const FPositionVertexBuffer* VertexBuffer) override;
+	FGLTFJsonAccessorIndex Convert(const FPositionVertexBuffer* VertexBuffer) override final;
 };
 
-class FGLTFColorVertexBufferConverter final : public TGLTFConverter<FGLTFJsonAccessorIndex, const FColorVertexBuffer*>
+class FGLTFColorVertexBufferConverter : public FGLTFBuilderContext, public TGLTFConverter<FGLTFJsonAccessorIndex, const FColorVertexBuffer*>
 {
-	using TGLTFConverter::TGLTFConverter;
+	using FGLTFBuilderContext::FGLTFBuilderContext;
 
-	FGLTFJsonAccessorIndex Convert(const FColorVertexBuffer* VertexBuffer) override;
+	FGLTFJsonAccessorIndex Convert(const FColorVertexBuffer* VertexBuffer) override final;
 };
 
-class FGLTFNormalVertexBufferConverter final : public TGLTFConverter<FGLTFJsonAccessorIndex, const FStaticMeshVertexBuffer*>
+class FGLTFNormalVertexBufferConverter : public FGLTFBuilderContext, public TGLTFConverter<FGLTFJsonAccessorIndex, const FStaticMeshVertexBuffer*>
 {
-	using TGLTFConverter::TGLTFConverter;
+	using FGLTFBuilderContext::FGLTFBuilderContext;
 
-	FGLTFJsonAccessorIndex Convert(const FStaticMeshVertexBuffer* VertexBuffer) override;
+	FGLTFJsonAccessorIndex Convert(const FStaticMeshVertexBuffer* VertexBuffer) override final;
 };
 
-class FGLTFTangentVertexBufferConverter final : public TGLTFConverter<FGLTFJsonAccessorIndex, const FStaticMeshVertexBuffer*>
+class FGLTFTangentVertexBufferConverter : public FGLTFBuilderContext, public TGLTFConverter<FGLTFJsonAccessorIndex, const FStaticMeshVertexBuffer*>
 {
-	using TGLTFConverter::TGLTFConverter;
+	using FGLTFBuilderContext::FGLTFBuilderContext;
 
-	FGLTFJsonAccessorIndex Convert(const FStaticMeshVertexBuffer* VertexBuffer) override;
+	FGLTFJsonAccessorIndex Convert(const FStaticMeshVertexBuffer* VertexBuffer) override final;
 };
 
-class FGLTFUVVertexBufferConverter final : public TGLTFConverter<FGLTFJsonAccessorIndex, const FStaticMeshVertexBuffer*, int32>
+class FGLTFUVVertexBufferConverter : public FGLTFBuilderContext, public TGLTFConverter<FGLTFJsonAccessorIndex, const FStaticMeshVertexBuffer*, int32>
 {
-	using TGLTFConverter::TGLTFConverter;
+	using FGLTFBuilderContext::FGLTFBuilderContext;
 
-	FGLTFJsonAccessorIndex Convert(const FStaticMeshVertexBuffer* VertexBuffer, int32 UVIndex) override;
+	FGLTFJsonAccessorIndex Convert(const FStaticMeshVertexBuffer* VertexBuffer, int32 UVIndex) override final;
 };
 
-class FGLTFBoneIndexVertexBufferConverter final : public TGLTFConverter<FGLTFJsonAccessorIndex, const FSkinWeightVertexBuffer*, int32, FGLTFBoneMap>
+class FGLTFBoneIndexVertexBufferConverter : public FGLTFBuilderContext, public TGLTFConverter<FGLTFJsonAccessorIndex, const FSkinWeightVertexBuffer*, int32, FGLTFBoneMap>
 {
-	using TGLTFConverter::TGLTFConverter;
+	using FGLTFBuilderContext::FGLTFBuilderContext;
 
-	FGLTFJsonAccessorIndex Convert(const FSkinWeightVertexBuffer* VertexBuffer, int32 InfluenceOffset, FGLTFBoneMap BoneMap) override;
+	FGLTFJsonAccessorIndex Convert(const FSkinWeightVertexBuffer* VertexBuffer, int32 InfluenceOffset, FGLTFBoneMap BoneMap) override final;
 
 	template <typename IndexType>
 	FGLTFJsonAccessorIndex Convert(const FSkinWeightVertexBuffer* VertexBuffer, int32 JointsGroupIndex, FGLTFBoneMap BoneMap);
 };
 
-class FGLTFBoneWeightVertexBufferConverter final : public TGLTFConverter<FGLTFJsonAccessorIndex, const FSkinWeightVertexBuffer*, int32>
+class FGLTFBoneWeightVertexBufferConverter : public FGLTFBuilderContext, public TGLTFConverter<FGLTFJsonAccessorIndex, const FSkinWeightVertexBuffer*, int32>
 {
-	using TGLTFConverter::TGLTFConverter;
+	using FGLTFBuilderContext::FGLTFBuilderContext;
 
-	FGLTFJsonAccessorIndex Convert(const FSkinWeightVertexBuffer* VertexBuffer, int32 WeightsGroupIndex) override;
+	FGLTFJsonAccessorIndex Convert(const FSkinWeightVertexBuffer* VertexBuffer, int32 WeightsGroupIndex) override final;
 };
