[GLTFExporter] refactor: split converters into seperate files by gltf category

commit 3f034449bb53c977aa523655632532291976a028
Date: Sun, 20 Sep 2020 21:27:45 +0200

    refactor: split converters into seperate files by gltf category
    (cherry picked from commit 903f87abc753e77a62b001cab5c556ca2009f9ff​)

#skipci

diff --git a/Engine/Plugins/Enterprise/GLTFExporter/MigratedFromGit/Source/GLTFExporter/Private/Builders/GLTFConvertBuilder.h b/Engine/Plugins/Enterprise/GLTFExporter/MigratedFromGit/Source/GLTFExporter/Private/Builders/GLTFConvertBuilder.h
--- a/Engine/Plugins/Enterprise/GLTFExporter/MigratedFromGit/Source/GLTFExporter/Private/Builders/GLTFConvertBuilder.h
+++ b/Engine/Plugins/Enterprise/GLTFExporter/MigratedFromGit/Source/GLTFExporter/Private/Builders/GLTFConvertBuilder.h
@@ -3,12 +3,15 @@
 #pragma once
 
 #include "Builders/GLTFImageBuilder.h"
-#include "Converters/GLTFVertexBufferConverters.h"
-#include "Converters/GLTFStaticMeshConverters.h"
-#include "Converters/GLTFSkeletalMeshConverters.h"
+#include "Converters/GLTFAccessorConverters.h"
+#include "Converters/GLTFMeshConverters.h"
 #include "Converters/GLTFMaterialConverters.h"
+#include "Converters/GLTFSamplerConverters.h"
 #include "Converters/GLTFTextureConverters.h"
-#include "Converters/GLTFLevelConverters.h"
+#include "Converters/GLTFNodeConverters.h"
+#include "Converters/GLTFSceneConverters.h"
+#include "Converters/GLTFCameraConverters.h"
+#include "Converters/GLTFLightConverters.h"
 #include "Converters/GLTFBackdropConverters.h"
 #include "Converters/GLTFVariantSetConverters.h"
 #include "Converters/GLTFLightMapConverters.h"
diff --git a/Engine/Plugins/Enterprise/GLTFExporter/MigratedFromGit/Source/GLTFExporter/Private/Converters/GLTFAccessorConverters.cpp b/Engine/Plugins/Enterprise/GLTFExporter/MigratedFromGit/Source/GLTFExporter/Private/Converters/GLTFAccessorConverters.cpp
new file mode 100644
--- /dev/null
+++ b/Engine/Plugins/Enterprise/GLTFExporter/MigratedFromGit/Source/GLTFExporter/Private/Converters/GLTFAccessorConverters.cpp
@@ -0,0 +1,337 @@
+// Copyright Epic Games, Inc. All Rights Reserved.
+
+#include "Converters/GLTFAccessorConverters.h"
+#include "Converters/GLTFConverterUtility.h"
+#include "Converters/GLTFSkinWeightVertexBufferHack.h"
+#include "Builders/GLTFConvertBuilder.h"
+#include "Math/NumericLimits.h"
+#include "Rendering/MultiSizeIndexContainer.h"
+#include "Rendering/SkeletalMeshRenderData.h"
+
+FGLTFJsonAccessorIndex FGLTFPositionVertexBufferConverter::Convert(const FPositionVertexBuffer* VertexBuffer)
+{
+	const uint32 VertexCount = VertexBuffer->GetNumVertices();
+	if (VertexCount == 0)
+	{
+		return FGLTFJsonAccessorIndex(INDEX_NONE);
+	}
+
+	TArray<FGLTFJsonVector3> Positions;
+	Positions.AddUninitialized(VertexCount);
+
+	for (uint32 VertexIndex = 0; VertexIndex < VertexCount; ++VertexIndex)
+	{
+		Positions[VertexIndex] = FGLTFConverterUtility::ConvertPosition(VertexBuffer->VertexPosition(VertexIndex), Builder.ExportOptions->ExportScale);
+	}
+
+	// More accurate bounding box if based on raw vertex values
+	FGLTFJsonVector3 MinPosition = Positions[0];
+	FGLTFJsonVector3 MaxPosition = Positions[0];
+
+	for (uint32 VertexIndex = 1; VertexIndex < VertexCount; ++VertexIndex)
+	{
+		const FGLTFJsonVector3& Position = Positions[VertexIndex];
+		MinPosition.X = FMath::Min(MinPosition.X, Position.X);
+		MinPosition.Y = FMath::Min(MinPosition.Y, Position.Y);
+		MinPosition.Z = FMath::Min(MinPosition.Z, Position.Z);
+		MaxPosition.X = FMath::Max(MaxPosition.X, Position.X);
+		MaxPosition.Y = FMath::Max(MaxPosition.Y, Position.Y);
+		MaxPosition.Z = FMath::Max(MaxPosition.Z, Position.Z);
+	}
+
+	FGLTFJsonAccessor JsonAccessor;
+	JsonAccessor.BufferView = Builder.AddBufferView(Positions);
+	JsonAccessor.ComponentType = EGLTFJsonComponentType::F32;
+	JsonAccessor.Count = VertexCount;
+	JsonAccessor.Type = EGLTFJsonAccessorType::Vec3;
+
+	JsonAccessor.MinMaxLength = 3;
+	JsonAccessor.Min[0] = MinPosition.X;
+	JsonAccessor.Min[1] = MinPosition.Y;
+	JsonAccessor.Min[2] = MinPosition.Z;
+	JsonAccessor.Max[0] = MaxPosition.X;
+	JsonAccessor.Max[1] = MaxPosition.Y;
+	JsonAccessor.Max[2] = MaxPosition.Z;
+
+	return Builder.AddAccessor(JsonAccessor);
+}
+
+FGLTFJsonAccessorIndex FGLTFColorVertexBufferConverter::Convert(const FColorVertexBuffer* VertexBuffer)
+{
+	const uint32 VertexCount = VertexBuffer->GetNumVertices();
+	if (VertexCount == 0)
+	{
+		return FGLTFJsonAccessorIndex(INDEX_NONE);
+	}
+
+	TArray<FGLTFPackedColor> Colors;
+	Colors.AddUninitialized(VertexCount);
+
+	for (uint32 VertexIndex = 0; VertexIndex < VertexCount; ++VertexIndex)
+	{
+		Colors[VertexIndex] = FGLTFConverterUtility::ConvertColor(VertexBuffer->VertexColor(VertexIndex));
+	}
+
+	FGLTFJsonAccessor JsonAccessor;
+	JsonAccessor.BufferView = Builder.AddBufferView(Colors);
+	JsonAccessor.ComponentType = EGLTFJsonComponentType::U8;
+	JsonAccessor.Count = VertexCount;
+	JsonAccessor.Type = EGLTFJsonAccessorType::Vec4;
+	JsonAccessor.bNormalized = true;
+
+	return Builder.AddAccessor(JsonAccessor);
+}
+
+FGLTFJsonAccessorIndex FGLTFNormalVertexBufferConverter::Convert(const FStaticMeshVertexBuffer* VertexBuffer)
+{
+	const uint32 VertexCount = VertexBuffer->GetNumVertices();
+	if (VertexCount == 0)
+	{
+		return FGLTFJsonAccessorIndex(INDEX_NONE);
+	}
+
+	TArray<FGLTFJsonVector3> Normals;
+	Normals.AddUninitialized(VertexCount);
+
+	for (uint32 VertexIndex = 0; VertexIndex < VertexCount; ++VertexIndex)
+	{
+		Normals[VertexIndex] = FGLTFConverterUtility::ConvertNormal(VertexBuffer->VertexTangentZ(VertexIndex));
+	}
+
+	FGLTFJsonAccessor JsonAccessor;
+	JsonAccessor.BufferView = Builder.AddBufferView(Normals);
+	JsonAccessor.ComponentType = EGLTFJsonComponentType::F32;
+	JsonAccessor.Count = VertexCount;
+	JsonAccessor.Type = EGLTFJsonAccessorType::Vec3;
+
+	return Builder.AddAccessor(JsonAccessor);
+}
+
+FGLTFJsonAccessorIndex FGLTFTangentVertexBufferConverter::Convert(const FStaticMeshVertexBuffer* VertexBuffer)
+{
+	const uint32 VertexCount = VertexBuffer->GetNumVertices();
+	if (VertexCount == 0)
+	{
+		return FGLTFJsonAccessorIndex(INDEX_NONE);
+	}
+
+	TArray<FGLTFJsonVector4> Tangents;
+	Tangents.AddUninitialized(VertexCount);
+
+	for (uint32 VertexIndex = 0; VertexIndex < VertexCount; ++VertexIndex)
+	{
+		Tangents[VertexIndex] = FGLTFConverterUtility::ConvertTangent(VertexBuffer->VertexTangentX(VertexIndex));
+	}
+
+	FGLTFJsonAccessor JsonAccessor;
+	JsonAccessor.BufferView = Builder.AddBufferView(Tangents);
+	JsonAccessor.ComponentType = EGLTFJsonComponentType::F32;
+	JsonAccessor.Count = VertexCount;
+	JsonAccessor.Type = EGLTFJsonAccessorType::Vec4;
+
+	return Builder.AddAccessor(JsonAccessor);
+}
+
+FGLTFJsonAccessorIndex FGLTFUVVertexBufferConverter::Convert(const FStaticMeshVertexBuffer* VertexBuffer, int32 UVIndex)
+{
+	const uint32 VertexCount = VertexBuffer->GetNumVertices();
+	if (VertexCount == 0 || UVIndex < 0 || VertexBuffer->GetNumTexCoords() <= static_cast<uint32>(UVIndex))
+	{
+		return FGLTFJsonAccessorIndex(INDEX_NONE);
+	}
+
+	TArray<FGLTFJsonVector2> UVs;
+	UVs.AddUninitialized(VertexCount);
+
+	for (uint32 VertexIndex = 0; VertexIndex < VertexCount; ++VertexIndex)
+	{
+		UVs[VertexIndex] = FGLTFConverterUtility::ConvertUV(VertexBuffer->GetVertexUV(VertexIndex, UVIndex));
+	}
+
+	FGLTFJsonAccessor JsonAccessor;
+	JsonAccessor.BufferView = Builder.AddBufferView(UVs);
+	JsonAccessor.ComponentType = EGLTFJsonComponentType::F32;
+	JsonAccessor.Count = VertexCount;
+	JsonAccessor.Type = EGLTFJsonAccessorType::Vec2;
+
+	return Builder.AddAccessor(JsonAccessor);
+}
+
+FGLTFJsonAccessorIndex FGLTFBoneIndexVertexBufferConverter::Convert(const FSkinWeightVertexBuffer* VertexBuffer, int32 InfluenceOffset, FGLTFBoneMap BoneMap)
+{
+	return VertexBuffer->Use16BitBoneIndex() ? Convert<uint16>(VertexBuffer, InfluenceOffset, BoneMap) : Convert<uint8>(VertexBuffer, InfluenceOffset, BoneMap);
+}
+
+template <typename IndexType>
+FGLTFJsonAccessorIndex FGLTFBoneIndexVertexBufferConverter::Convert(const FSkinWeightVertexBuffer* VertexBuffer, int32 JointsGroupIndex, FGLTFBoneMap BoneMap)
+{
+	const uint32 VertexCount = VertexBuffer->GetNumVertices();
+	const int32 MaxInfluenceCount = VertexBuffer->GetMaxBoneInfluences();
+	const int32 InfluenceOffset = JointsGroupIndex * 4;
+
+	if (VertexCount == 0 || InfluenceOffset < 0 || MaxInfluenceCount <= InfluenceOffset)
+	{
+		return FGLTFJsonAccessorIndex(INDEX_NONE);
+	}
+
+	struct VertexBoneIndices
+	{
+		IndexType Index[4];
+	};
+
+	TArray<VertexBoneIndices> BoneIndices;
+	BoneIndices.AddUninitialized(VertexCount);
+
+	for (uint32 VertexIndex = 0; VertexIndex < VertexCount; ++VertexIndex)
+	{
+		for (int32 InfluenceIndex = 0; InfluenceIndex < 4; ++InfluenceIndex)
+		{
+			// TODO: remove hack
+			const uint32 UnmappedBoneIndex = FGLTFSkinWeightVertexBufferHack(VertexBuffer).GetBoneIndex(VertexIndex, InfluenceOffset + InfluenceIndex);
+			const FBoneIndexType BoneIndex = BoneMap[UnmappedBoneIndex];
+			check(BoneIndex <= TNumericLimits<IndexType>::Max());
+
+			BoneIndices[VertexIndex].Index[InfluenceIndex] = static_cast<IndexType>(BoneIndex);
+		}
+	}
+
+	FGLTFJsonAccessor JsonAccessor;
+	JsonAccessor.BufferView = Builder.AddBufferView(BoneIndices);
+	JsonAccessor.ComponentType = FGLTFConverterUtility::GetComponentType<IndexType>();
+	JsonAccessor.Count = VertexCount;
+	JsonAccessor.Type = EGLTFJsonAccessorType::Vec4;
+
+	return Builder.AddAccessor(JsonAccessor);
+}
+
+FGLTFJsonAccessorIndex FGLTFBoneWeightVertexBufferConverter::Convert(const FSkinWeightVertexBuffer* VertexBuffer, int32 WeightsGroupIndex)
+{
+	const uint32 VertexCount = VertexBuffer->GetNumVertices();
+	const int32 MaxInfluenceCount = VertexBuffer->GetMaxBoneInfluences();
+	const int32 InfluenceOffset = WeightsGroupIndex * 4;
+
+	if (VertexCount == 0 || InfluenceOffset < 0 || MaxInfluenceCount <= InfluenceOffset)
+	{
+		return FGLTFJsonAccessorIndex(INDEX_NONE);
+	}
+
+	struct VertexBoneWeights
+	{
+		uint8 Weights[4];
+	};
+
+	TArray<VertexBoneWeights> BoneWeights;
+	BoneWeights.AddUninitialized(VertexCount);
+
+	for (uint32 VertexIndex = 0; VertexIndex < VertexCount; ++VertexIndex)
+	{
+		for (int32 InfluenceIndex = 0; InfluenceIndex < 4; ++InfluenceIndex)
+		{
+			// TODO: remove hack
+			const uint8 BoneWeight = FGLTFSkinWeightVertexBufferHack(VertexBuffer).GetBoneWeight(VertexIndex, InfluenceOffset + InfluenceIndex);
+			BoneWeights[VertexIndex].Weights[InfluenceIndex] = BoneWeight;
+		}
+	}
+
+	FGLTFJsonAccessor JsonAccessor;
+	JsonAccessor.BufferView = Builder.AddBufferView(BoneWeights);
+	JsonAccessor.ComponentType = EGLTFJsonComponentType::U8;
+	JsonAccessor.Count = VertexCount;
+	JsonAccessor.Type = EGLTFJsonAccessorType::Vec4;
+	JsonAccessor.bNormalized = true;
+
+	return Builder.AddAccessor(JsonAccessor);
+}
+
+FGLTFJsonBufferViewIndex FGLTFIndexBufferConverter::Convert(const FRawStaticIndexBuffer* IndexBuffer)
+{
+	if (IndexBuffer->GetNumIndices() == 0)
+	{
+		return FGLTFJsonBufferViewIndex(INDEX_NONE);
+	}
+
+	if (IndexBuffer->Is32Bit())
+	{
+		TArray<uint32> Indices;
+		IndexBuffer->GetCopy(Indices);
+		return Builder.AddBufferView(Indices, sizeof(uint32), EGLTFJsonBufferTarget::ElementArrayBuffer);
+	}
+	else
+	{
+		const uint16* IndexData = IndexBuffer->AccessStream16();
+		const int32 IndexDataSize = IndexBuffer->GetIndexDataSize();
+		return Builder.AddBufferView(IndexData, IndexDataSize, sizeof(uint16), EGLTFJsonBufferTarget::ElementArrayBuffer);
+	}
+}
+
+FGLTFJsonAccessorIndex FGLTFStaticMeshSectionConverter::Convert(const FStaticMeshSection* MeshSection, const FRawStaticIndexBuffer* IndexBuffer)
+{
+	const uint32 TriangleCount = MeshSection->NumTriangles;
+	if (TriangleCount == 0)
+	{
+		return FGLTFJsonAccessorIndex(INDEX_NONE);
+	}
+
+	const uint32 FirstIndex = MeshSection->FirstIndex;
+	const bool bIs32Bit = IndexBuffer->Is32Bit();
+
+	FGLTFJsonAccessor JsonAccessor;
+	JsonAccessor.BufferView = Builder.GetOrAddIndexBufferView(IndexBuffer);
+	JsonAccessor.ByteOffset = FirstIndex * (bIs32Bit ? sizeof(uint32) : sizeof(uint16));
+	JsonAccessor.ComponentType = bIs32Bit ? EGLTFJsonComponentType::U32 : EGLTFJsonComponentType::U16;
+	JsonAccessor.Count = TriangleCount * 3;
+	JsonAccessor.Type = EGLTFJsonAccessorType::Scalar;
+
+	return Builder.AddAccessor(JsonAccessor);
+}
+
+FGLTFJsonBufferViewIndex FGLTFIndexContainerConverter::Convert(const FMultiSizeIndexContainer* IndexContainer)
+{
+	const FRawStaticIndexBuffer16or32Interface* IndexBuffer = IndexContainer->GetIndexBuffer();
+
+	const int32 IndexCount = IndexBuffer->Num();
+	if (IndexCount <= 0)
+	{
+		return FGLTFJsonBufferViewIndex(INDEX_NONE);
+	}
+
+	if (IndexContainer->GetDataTypeSize() == sizeof(uint16))
+	{
+		TArray<uint16> Indices;
+		Indices.AddUninitialized(IndexCount);
+
+		for (int32 Index = 0; Index < IndexCount; ++Index)
+		{
+			Indices[Index] = static_cast<uint16>(IndexBuffer->Get(Index));
+		}
+
+		return Builder.AddBufferView(Indices, sizeof(uint16), EGLTFJsonBufferTarget::ElementArrayBuffer);
+	}
+	else
+	{
+		TArray<uint32> Indices;
+		IndexContainer->GetIndexBuffer(Indices);
+		return Builder.AddBufferView(Indices, sizeof(uint32), EGLTFJsonBufferTarget::ElementArrayBuffer);
+	}
+}
+
+FGLTFJsonAccessorIndex FGLTFSkeletalMeshSectionConverter::Convert(const FSkelMeshRenderSection* MeshSection, const FMultiSizeIndexContainer* IndexContainer)
+{
+	const uint32 TriangleCount = MeshSection->NumTriangles;
+	if (TriangleCount == 0)
+	{
+		return FGLTFJsonAccessorIndex(INDEX_NONE);
+	}
+
+	const uint32 FirstIndex = MeshSection->BaseIndex;
+	const bool bIs32Bit = IndexContainer->GetDataTypeSize() == sizeof(uint32);
+
+	FGLTFJsonAccessor JsonAccessor;
+	JsonAccessor.BufferView = Builder.GetOrAddIndexBufferView(IndexContainer);
+	JsonAccessor.ByteOffset = FirstIndex * (bIs32Bit ? sizeof(uint32) : sizeof(uint16));
+	JsonAccessor.ComponentType = bIs32Bit ? EGLTFJsonComponentType::U32 : EGLTFJsonComponentType::U16;
+	JsonAccessor.Count = TriangleCount * 3;
+	JsonAccessor.Type = EGLTFJsonAccessorType::Scalar;
+
+	return Builder.AddAccessor(JsonAccessor);
+}
diff --git a/Engine/Plugins/Enterprise/GLTFExporter/MigratedFromGit/Source/GLTFExporter/Private/Converters/GLTFAccessorConverters.h b/Engine/Plugins/Enterprise/GLTFExporter/MigratedFromGit/Source/GLTFExporter/Private/Converters/GLTFAccessorConverters.h
new file mode 100644
--- /dev/null
+++ b/Engine/Plugins/Enterprise/GLTFExporter/MigratedFromGit/Source/GLTFExporter/Private/Converters/GLTFAccessorConverters.h
@@ -0,0 +1,93 @@
+// Copyright Epic Games, Inc. All Rights Reserved.
+
+#pragma once
+
+#include "Json/GLTFJsonIndex.h"
+#include "Converters/GLTFConverter.h"
+#include "Converters/GLTFBuilderContext.h"
+#include "Converters/GLTFBoneMap.h"
+#include "Engine.h"
+
+class FMultiSizeIndexContainer;
+
+class FGLTFPositionVertexBufferConverter : public FGLTFBuilderContext, public TGLTFConverter<FGLTFJsonAccessorIndex, const FPositionVertexBuffer*>
+{
+	using FGLTFBuilderContext::FGLTFBuilderContext;
+
+	FGLTFJsonAccessorIndex Convert(const FPositionVertexBuffer* VertexBuffer) override final;
+};
+
+class FGLTFColorVertexBufferConverter : public FGLTFBuilderContext, public TGLTFConverter<FGLTFJsonAccessorIndex, const FColorVertexBuffer*>
+{
+	using FGLTFBuilderContext::FGLTFBuilderContext;
+
+	FGLTFJsonAccessorIndex Convert(const FColorVertexBuffer* VertexBuffer) override final;
+};
+
+class FGLTFNormalVertexBufferConverter : public FGLTFBuilderContext, public TGLTFConverter<FGLTFJsonAccessorIndex, const FStaticMeshVertexBuffer*>
+{
+	using FGLTFBuilderContext::FGLTFBuilderContext;
+
+	FGLTFJsonAccessorIndex Convert(const FStaticMeshVertexBuffer* VertexBuffer) override final;
+};
+
+class FGLTFTangentVertexBufferConverter : public FGLTFBuilderContext, public TGLTFConverter<FGLTFJsonAccessorIndex, const FStaticMeshVertexBuffer*>
+{
+	using FGLTFBuilderContext::FGLTFBuilderContext;
+
+	FGLTFJsonAccessorIndex Convert(const FStaticMeshVertexBuffer* VertexBuffer) override final;
+};
+
+class FGLTFUVVertexBufferConverter : public FGLTFBuilderContext, public TGLTFConverter<FGLTFJsonAccessorIndex, const FStaticMeshVertexBuffer*, int32>
+{
+	using FGLTFBuilderContext::FGLTFBuilderContext;
+
+	FGLTFJsonAccessorIndex Convert(const FStaticMeshVertexBuffer* VertexBuffer, int32 UVIndex) override final;
+};
+
+class FGLTFBoneIndexVertexBufferConverter : public FGLTFBuilderContext, public TGLTFConverter<FGLTFJsonAccessorIndex, const FSkinWeightVertexBuffer*, int32, FGLTFBoneMap>
+{
+	using FGLTFBuilderContext::FGLTFBuilderContext;
+
+	FGLTFJsonAccessorIndex Convert(const FSkinWeightVertexBuffer* VertexBuffer, int32 InfluenceOffset, FGLTFBoneMap BoneMap) override final;
+
+	template <typename IndexType>
+	FGLTFJsonAccessorIndex Convert(const FSkinWeightVertexBuffer* VertexBuffer, int32 JointsGroupIndex, FGLTFBoneMap BoneMap);
+};
+
+class FGLTFBoneWeightVertexBufferConverter : public FGLTFBuilderContext, public TGLTFConverter<FGLTFJsonAccessorIndex, const FSkinWeightVertexBuffer*, int32>
+{
+	using FGLTFBuilderContext::FGLTFBuilderContext;
+
+	FGLTFJsonAccessorIndex Convert(const FSkinWeightVertexBuffer* VertexBuffer, int32 WeightsGroupIndex) override final;
+};
+
+// TODO: merge into FGLTFStaticMeshSectionConverter
+class FGLTFIndexBufferConverter : public FGLTFBuilderContext, public TGLTFConverter<FGLTFJsonBufferViewIndex, const FRawStaticIndexBuffer*>
+{
+	using FGLTFBuilderContext::FGLTFBuilderContext;
+
+	FGLTFJsonBufferViewIndex Convert(const FRawStaticIndexBuffer* IndexBuffer) override final;
+};
+
+class FGLTFStaticMeshSectionConverter : public FGLTFBuilderContext, public TGLTFConverter<FGLTFJsonAccessorIndex, const FStaticMeshSection*, const FRawStaticIndexBuffer*>
+{
+	using FGLTFBuilderContext::FGLTFBuilderContext;
+
+	FGLTFJsonAccessorIndex Convert(const FStaticMeshSection* MeshSection, const FRawStaticIndexBuffer* IndexBuffer) override final;
+};
+
+// TODO: merge into FGLTFSkeletalMeshSectionConverter
+class FGLTFIndexContainerConverter : public FGLTFBuilderContext, public TGLTFConverter<FGLTFJsonBufferViewIndex, const FMultiSizeIndexContainer*>
+{
+	using FGLTFBuilderContext::FGLTFBuilderContext;
+
+	FGLTFJsonBufferViewIndex Convert(const FMultiSizeIndexContainer* IndexContainer) override final;
+};
+
+class FGLTFSkeletalMeshSectionConverter : public FGLTFBuilderContext, public TGLTFConverter<FGLTFJsonAccessorIndex, const FSkelMeshRenderSection*, const FMultiSizeIndexContainer*>
+{
+	using FGLTFBuilderContext::FGLTFBuilderContext;
+
+	FGLTFJsonAccessorIndex Convert(const FSkelMeshRenderSection* MeshSection, const FMultiSizeIndexContainer* IndexContainer) override final;
+};
diff --git a/Engine/Plugins/Enterprise/GLTFExporter/MigratedFromGit/Source/GLTFExporter/Private/Converters/GLTFCameraConverters.cpp b/Engine/Plugins/Enterprise/GLTFExporter/MigratedFromGit/Source/GLTFExporter/Private/Converters/GLTFCameraConverters.cpp
new file mode 100644
--- /dev/null
+++ b/Engine/Plugins/Enterprise/GLTFExporter/MigratedFromGit/Source/GLTFExporter/Private/Converters/GLTFCameraConverters.cpp
@@ -0,0 +1,34 @@
+// Copyright Epic Games, Inc. All Rights Reserved.
+
+#include "Converters/GLTFCameraConverters.h"
+#include "Builders/GLTFContainerBuilder.h"
+#include "Converters/GLTFConverterUtility.h"
+#include "Converters/GLTFNameUtility.h"
+#include "Converters/GLTFCameraUtility.h"
+
+FGLTFJsonCameraIndex FGLTFCameraComponentConverter::Convert(const UCameraComponent* CameraComponent)
+{
+	FGLTFJsonCamera Camera;
+	Camera.Name = FGLTFNameUtility::GetName(CameraComponent);
+	Camera.Type = FGLTFConverterUtility::ConvertCameraType(CameraComponent->ProjectionMode);
+
+	FMinimalViewInfo DesiredView;
+	const_cast<UCameraComponent*>(CameraComponent)->GetCameraView(0, DesiredView);
+
+	switch (Camera.Type)
+	{
+		case EGLTFJsonCameraType::Orthographic:
+			Camera.Orthographic = FGLTFCameraUtility::ConvertOrthographic(DesiredView, Builder.ExportOptions->ExportScale);
+			break;
+
+		case EGLTFJsonCameraType::Perspective:
+			Camera.Perspective = FGLTFCameraUtility::ConvertPerspective(DesiredView, Builder.ExportOptions->ExportScale);
+			break;
+
+		default:
+		    // TODO: report error (unsupported camera type)
+		    return FGLTFJsonCameraIndex(INDEX_NONE);
+	}
+
+	return Builder.AddCamera(Camera);
+}
diff --git a/Engine/Plugins/Enterprise/GLTFExporter/MigratedFromGit/Source/GLTFExporter/Private/Converters/GLTFCameraConverters.h b/Engine/Plugins/Enterprise/GLTFExporter/MigratedFromGit/Source/GLTFExporter/Private/Converters/GLTFCameraConverters.h
new file mode 100644
--- /dev/null
+++ b/Engine/Plugins/Enterprise/GLTFExporter/MigratedFromGit/Source/GLTFExporter/Private/Converters/GLTFCameraConverters.h
@@ -0,0 +1,15 @@
+// Copyright Epic Games, Inc. All Rights Reserved.
+
+#pragma once
+
+#include "Json/GLTFJsonIndex.h"
+#include "Converters/GLTFConverter.h"
+#include "Converters/GLTFBuilderContext.h"
+#include "Engine.h"
+
+class FGLTFCameraComponentConverter : public FGLTFBuilderContext, public TGLTFConverter<FGLTFJsonCameraIndex, const UCameraComponent*>
+{
+	using FGLTFBuilderContext::FGLTFBuilderContext;
+
+	FGLTFJsonCameraIndex Convert(const UCameraComponent* CameraComponent) override final;
+};
diff --git a/Engine/Plugins/Enterprise/GLTFExporter/MigratedFromGit/Source/GLTFExporter/Private/Converters/GLTFLevelConverters.cpp b/Engine/Plugins/Enterprise/GLTFExporter/MigratedFromGit/Source/GLTFExporter/Private/Converters/GLTFLevelConverters.cpp
deleted file mode 100644
diff --git a/Engine/Plugins/Enterprise/GLTFExporter/MigratedFromGit/Source/GLTFExporter/Private/Converters/GLTFLevelConverters.h b/Engine/Plugins/Enterprise/GLTFExporter/MigratedFromGit/Source/GLTFExporter/Private/Converters/GLTFLevelConverters.h
deleted file mode 100644
diff --git a/Engine/Plugins/Enterprise/GLTFExporter/MigratedFromGit/Source/GLTFExporter/Private/Converters/GLTFLightConverters.cpp b/Engine/Plugins/Enterprise/GLTFExporter/MigratedFromGit/Source/GLTFExporter/Private/Converters/GLTFLightConverters.cpp
new file mode 100644
--- /dev/null
+++ b/Engine/Plugins/Enterprise/GLTFExporter/MigratedFromGit/Source/GLTFExporter/Private/Converters/GLTFLightConverters.cpp
@@ -0,0 +1,38 @@
+// Copyright Epic Games, Inc. All Rights Reserved.
+
+#include "Converters/GLTFLightConverters.h"
+#include "Builders/GLTFContainerBuilder.h"
+#include "Converters/GLTFConverterUtility.h"
+#include "Converters/GLTFNameUtility.h"
+
+FGLTFJsonLightIndex FGLTFLightComponentConverter::Convert(const ULightComponent* LightComponent)
+{
+	FGLTFJsonLight Light;
+
+	Light.Name = FGLTFNameUtility::GetName(LightComponent);
+	Light.Type = FGLTFConverterUtility::ConvertLightType(LightComponent->GetLightType());
+
+	if (Light.Type == EGLTFJsonLightType::None)
+	{
+		// TODO: report error (unsupported light component type)
+		return FGLTFJsonLightIndex(INDEX_NONE);
+	}
+
+	Light.Intensity = LightComponent->Intensity;
+
+	const FLinearColor LightColor = LightComponent->bUseTemperature ? FLinearColor::MakeFromColorTemperature(LightComponent->Temperature) : LightComponent->GetLightColor();
+	Light.Color = FGLTFConverterUtility::ConvertColor(LightColor);
+
+	if (const UPointLightComponent* PointLightComponent = Cast<UPointLightComponent>(LightComponent))
+	{
+		Light.Range = FGLTFConverterUtility::ConvertLength(PointLightComponent->AttenuationRadius, Builder.ExportOptions->ExportScale);
+	}
+
+	if (const USpotLightComponent* SpotLightComponent = Cast<USpotLightComponent>(LightComponent))
+	{
+		Light.Spot.InnerConeAngle = FGLTFConverterUtility::ConvertLightAngle(SpotLightComponent->InnerConeAngle);
+		Light.Spot.OuterConeAngle = FGLTFConverterUtility::ConvertLightAngle(SpotLightComponent->OuterConeAngle);
+	}
+
+	return Builder.AddLight(Light);
+}
diff --git a/Engine/Plugins/Enterprise/GLTFExporter/MigratedFromGit/Source/GLTFExporter/Private/Converters/GLTFLightConverters.h b/Engine/Plugins/Enterprise/GLTFExporter/MigratedFromGit/Source/GLTFExporter/Private/Converters/GLTFLightConverters.h
new file mode 100644
--- /dev/null
+++ b/Engine/Plugins/Enterprise/GLTFExporter/MigratedFromGit/Source/GLTFExporter/Private/Converters/GLTFLightConverters.h
@@ -0,0 +1,15 @@
+// Copyright Epic Games, Inc. All Rights Reserved.
+
+#pragma once
+
+#include "Json/GLTFJsonIndex.h"
+#include "Converters/GLTFConverter.h"
+#include "Converters/GLTFBuilderContext.h"
+#include "Engine.h"
+
+class FGLTFLightComponentConverter : public FGLTFBuilderContext, public TGLTFConverter<FGLTFJsonLightIndex, const ULightComponent*>
+{
+	using FGLTFBuilderContext::FGLTFBuilderContext;
+
+	FGLTFJsonLightIndex Convert(const ULightComponent* LightComponent) override final;
+};
diff --git a/Engine/Plugins/Enterprise/GLTFExporter/MigratedFromGit/Source/GLTFExporter/Private/Converters/GLTFMeshConverters.cpp b/Engine/Plugins/Enterprise/GLTFExporter/MigratedFromGit/Source/GLTFExporter/Private/Converters/GLTFMeshConverters.cpp
new file mode 100644
--- /dev/null
+++ b/Engine/Plugins/Enterprise/GLTFExporter/MigratedFromGit/Source/GLTFExporter/Private/Converters/GLTFMeshConverters.cpp
@@ -0,0 +1,150 @@
+// Copyright Epic Games, Inc. All Rights Reserved.
+
+#include "Converters/GLTFMeshConverters.h"
+#include "Converters/GLTFConverterUtility.h"
+#include "Converters/GLTFNameUtility.h"
+#include "Builders/GLTFConvertBuilder.h"
+#include "Rendering/MultiSizeIndexContainer.h"
+#include "Rendering/SkeletalMeshRenderData.h"
+
+FGLTFJsonMeshIndex FGLTFStaticMeshConverter::Convert(const UStaticMesh* StaticMesh, int32 LODIndex, const FColorVertexBuffer* OverrideVertexColors, FGLTFMaterialArray OverrideMaterials)
+{
+	if (LODIndex < 0 || StaticMesh->GetNumLODs() <= LODIndex)
+	{
+		return FGLTFJsonMeshIndex(INDEX_NONE);
+	}
+
+	const FStaticMeshLODResources& MeshLOD = StaticMesh->GetLODForExport(LODIndex);
+
+	FGLTFJsonMesh JsonMesh;
+	JsonMesh.Name = FGLTFNameUtility::GetName(StaticMesh, LODIndex);
+
+	const FPositionVertexBuffer* PositionBuffer = &MeshLOD.VertexBuffers.PositionVertexBuffer;
+	const FStaticMeshVertexBuffer* VertexBuffer = &MeshLOD.VertexBuffers.StaticMeshVertexBuffer;
+	const FColorVertexBuffer* ColorBuffer = OverrideVertexColors != nullptr ? OverrideVertexColors : &MeshLOD.VertexBuffers.ColorVertexBuffer;
+
+	FGLTFJsonAttributes JsonAttributes;
+	JsonAttributes.Position = Builder.GetOrAddPositionAccessor(PositionBuffer);
+
+	if (Builder.ExportOptions->bExportVertexColors)
+	{
+		JsonAttributes.Color0 = Builder.GetOrAddColorAccessor(ColorBuffer);
+	}
+
+	JsonAttributes.Normal = Builder.GetOrAddNormalAccessor(VertexBuffer);
+	JsonAttributes.Tangent = Builder.GetOrAddTangentAccessor(VertexBuffer);
+
+	const uint32 UVCount = VertexBuffer->GetNumTexCoords();
+	JsonAttributes.TexCoords.AddUninitialized(UVCount);
+
+	for (uint32 UVIndex = 0; UVIndex < UVCount; ++UVIndex)
+	{
+		JsonAttributes.TexCoords[UVIndex] = Builder.GetOrAddUVAccessor(VertexBuffer, UVIndex);
+	}
+
+	const FRawStaticIndexBuffer* IndexBuffer = &MeshLOD.IndexBuffer;
+	Builder.GetOrAddIndexBufferView(IndexBuffer);
+
+	const int32 SectionCount = MeshLOD.Sections.Num();
+	JsonMesh.Primitives.AddDefaulted(SectionCount);
+
+	for (int32 SectionIndex = 0; SectionIndex < SectionCount; ++SectionIndex)
+	{
+		FGLTFJsonPrimitive& JsonPrimitive = JsonMesh.Primitives[SectionIndex];
+		JsonPrimitive.Attributes = JsonAttributes;
+
+		const FStaticMeshSection& Section = MeshLOD.Sections[SectionIndex];
+		JsonPrimitive.Indices = Builder.GetOrAddIndexAccessor(&Section, IndexBuffer);
+
+		const int32 MaterialIndex = Section.MaterialIndex;
+		const UMaterialInterface* Material = OverrideMaterials.IsValidIndex(MaterialIndex) && OverrideMaterials[MaterialIndex] != nullptr ?
+			OverrideMaterials[MaterialIndex] : StaticMesh->GetMaterial(MaterialIndex);
+
+		if (Material != nullptr)
+		{
+			JsonPrimitive.Material = Builder.GetOrAddMaterial(Material);
+		}
+	}
+
+	return Builder.AddMesh(JsonMesh);
+}
+
+FGLTFJsonMeshIndex FGLTFSkeletalMeshConverter::Convert(const USkeletalMesh* SkeletalMesh, int32 LODIndex, const FColorVertexBuffer* OverrideVertexColors, const FSkinWeightVertexBuffer* OverrideSkinWeights, FGLTFMaterialArray OverrideMaterials)
+{
+	const FSkeletalMeshRenderData* RenderData = SkeletalMesh->GetResourceForRendering();
+
+	if (LODIndex < 0 || RenderData->LODRenderData.Num() <= LODIndex)
+	{
+		return FGLTFJsonMeshIndex(INDEX_NONE);
+	}
+
+	const FSkeletalMeshLODRenderData& MeshLOD = RenderData->LODRenderData[LODIndex];
+
+	FGLTFJsonMesh JsonMesh;
+	JsonMesh.Name = FGLTFNameUtility::GetName(SkeletalMesh, LODIndex);
+
+	const FPositionVertexBuffer* PositionBuffer = &MeshLOD.StaticVertexBuffers.PositionVertexBuffer;
+	const FStaticMeshVertexBuffer* VertexBuffer = &MeshLOD.StaticVertexBuffers.StaticMeshVertexBuffer;
+	const FColorVertexBuffer* ColorBuffer = OverrideVertexColors != nullptr ? OverrideVertexColors : &MeshLOD.StaticVertexBuffers.ColorVertexBuffer;
+	const FSkinWeightVertexBuffer* SkinWeightBuffer = OverrideSkinWeights != nullptr ? OverrideSkinWeights : MeshLOD.GetSkinWeightVertexBuffer();
+
+	FGLTFJsonAttributes JsonAttributes;
+	JsonAttributes.Position = Builder.GetOrAddPositionAccessor(PositionBuffer);
+
+	if (Builder.ExportOptions->bExportVertexColors)
+	{
+		JsonAttributes.Color0 = Builder.GetOrAddColorAccessor(ColorBuffer);
+	}
+
+	JsonAttributes.Normal = Builder.GetOrAddNormalAccessor(VertexBuffer);
+	JsonAttributes.Tangent = Builder.GetOrAddTangentAccessor(VertexBuffer);
+
+	const uint32 UVCount = VertexBuffer->GetNumTexCoords();
+	JsonAttributes.TexCoords.AddUninitialized(UVCount);
+
+	for (uint32 UVIndex = 0; UVIndex < UVCount; ++UVIndex)
+	{
+		JsonAttributes.TexCoords[UVIndex] = Builder.GetOrAddUVAccessor(VertexBuffer, UVIndex);
+	}
+
+	/* TODO: enable export of vertex skin data when crash bug fixed
+
+	const uint32 GroupCount = (SkinWeightBuffer->GetMaxBoneInfluences() + 3) / 4;
+	JsonAttributes.Joints.AddUninitialized(GroupCount);
+	JsonAttributes.Weights.AddUninitialized(GroupCount);
+
+	const FGLTFBoneMap BoneMap = FGLTFBoneMap(MeshLOD.RenderSections[0].BoneMap); // TODO: can we make this assumption?
+
+	for (uint32 GroupIndex = 0; GroupIndex < GroupCount; ++GroupIndex)
+	{
+		JsonAttributes.Joints[GroupIndex] = Builder.GetOrAddJointAccessor(SkinWeightBuffer, GroupIndex, BoneMap, Name + TEXT("_Joints") + FString::FromInt(GroupIndex));
+		JsonAttributes.Weights[GroupIndex] = Builder.GetOrAddWeightAccessor(SkinWeightBuffer, GroupIndex, Name + TEXT("_Weights") + FString::FromInt(GroupIndex));
+	}
+	*/
+
+	const FMultiSizeIndexContainer* IndexContainer = &MeshLOD.MultiSizeIndexContainer;
+	Builder.GetOrAddIndexBufferView(IndexContainer);
+
+	const int32 SectionCount = MeshLOD.RenderSections.Num();
+	JsonMesh.Primitives.AddDefaulted(SectionCount);
+
+	for (int32 SectionIndex = 0; SectionIndex < SectionCount; ++SectionIndex)
+	{
+		FGLTFJsonPrimitive& JsonPrimitive = JsonMesh.Primitives[SectionIndex];
+		JsonPrimitive.Attributes = JsonAttributes;
+
+		const FSkelMeshRenderSection& Section = MeshLOD.RenderSections[SectionIndex];
+		JsonPrimitive.Indices = Builder.GetOrAddIndexAccessor(&Section, IndexContainer);
+
+		const int32 MaterialIndex = Section.MaterialIndex;
+		const UMaterialInterface* Material = OverrideMaterials.IsValidIndex(MaterialIndex) && OverrideMaterials[MaterialIndex] != nullptr ?
+			OverrideMaterials[MaterialIndex] : SkeletalMesh->Materials[MaterialIndex].MaterialInterface;
+
+		if (Material != nullptr)
+		{
+			JsonPrimitive.Material = Builder.GetOrAddMaterial(Material);
+		}
+	}
+
+	return Builder.AddMesh(JsonMesh);
+}
diff --git a/Engine/Plugins/Enterprise/GLTFExporter/MigratedFromGit/Source/GLTFExporter/Private/Converters/GLTFMeshConverters.h b/Engine/Plugins/Enterprise/GLTFExporter/MigratedFromGit/Source/GLTFExporter/Private/Converters/GLTFMeshConverters.h
new file mode 100644
--- /dev/null
+++ b/Engine/Plugins/Enterprise/GLTFExporter/MigratedFromGit/Source/GLTFExporter/Private/Converters/GLTFMeshConverters.h
@@ -0,0 +1,23 @@
+// Copyright Epic Games, Inc. All Rights Reserved.
+
+#pragma once
+
+#include "Json/GLTFJsonIndex.h"
+#include "Converters/GLTFConverter.h"
+#include "Converters/GLTFBuilderContext.h"
+#include "Converters/GLTFMaterialArray.h"
+#include "Engine.h"
+
+class FGLTFStaticMeshConverter : public FGLTFBuilderContext, public TGLTFConverter<FGLTFJsonMeshIndex, const UStaticMesh*, int32, const FColorVertexBuffer*, FGLTFMaterialArray>
+{
+	using FGLTFBuilderContext::FGLTFBuilderContext;
+
+	FGLTFJsonMeshIndex Convert(const UStaticMesh* StaticMesh, int32 LODIndex, const FColorVertexBuffer* OverrideVertexColors, FGLTFMaterialArray OverrideMaterials) override final;
+};
+
+class FGLTFSkeletalMeshConverter : public FGLTFBuilderContext, public TGLTFConverter<FGLTFJsonMeshIndex, const USkeletalMesh*, int32, const FColorVertexBuffer*, const FSkinWeightVertexBuffer*, FGLTFMaterialArray>
+{
+	using FGLTFBuilderContext::FGLTFBuilderContext;
+
+	FGLTFJsonMeshIndex Convert(const USkeletalMesh* SkeletalMesh, int32 LODIndex, const FColorVertexBuffer* OverrideVertexColors, const FSkinWeightVertexBuffer* OverrideSkinWeights, FGLTFMaterialArray OverrideMaterials) override final;
+};
diff --git a/Engine/Plugins/Enterprise/GLTFExporter/MigratedFromGit/Source/GLTFExporter/Private/Converters/GLTFNodeConverters.cpp b/Engine/Plugins/Enterprise/GLTFExporter/MigratedFromGit/Source/GLTFExporter/Private/Converters/GLTFNodeConverters.cpp
new file mode 100644
--- /dev/null
+++ b/Engine/Plugins/Enterprise/GLTFExporter/MigratedFromGit/Source/GLTFExporter/Private/Converters/GLTFNodeConverters.cpp
@@ -0,0 +1,188 @@
+// Copyright Epic Games, Inc. All Rights Reserved.
+
+#include "Converters/GLTFNodeConverters.h"
+#include "Builders/GLTFContainerBuilder.h"
+#include "Converters/GLTFConverterUtility.h"
+#include "Converters/GLTFActorUtility.h"
+#include "Converters/GLTFNameUtility.h"
+#include "Components/GLTFInteractionHotspotComponent.h"
+
+FGLTFJsonNodeIndex FGLTFSceneComponentConverter::Convert(const USceneComponent* SceneComponent)
+{
+	const AActor* Owner = SceneComponent->GetOwner();
+	if (Owner == nullptr)
+	{
+		// TODO: report error (invalid scene component)
+		return FGLTFJsonNodeIndex(INDEX_NONE);
+	}
+
+	const bool bIsRootComponent = Owner->GetRootComponent() == SceneComponent;
+	const bool bIsRootNode = bIsRootComponent && FGLTFActorUtility::IsRootActor(Owner, Builder.bSelectedActorsOnly);
+	const bool bExportNonUniformScale = Builder.ExportOptions->bExportNonUniformScale;
+
+	const USceneComponent* ParentComponent = SceneComponent->GetAttachParent();
+	const FGLTFJsonNodeIndex ParentNodeIndex = Builder.GetOrAddNode(ParentComponent);
+
+	// TODO: if node is root, then add it to the scene here to avoid any possible orphan nodes. this change requires level converter to support cyclic calls.
+
+	const FTransform Transform = SceneComponent->GetComponentTransform();
+	const FTransform RelativeTransform = bIsRootNode ? Transform : Transform.GetRelativeTransform(ParentComponent->GetComponentTransform());
+
+	const FVector ParentScale = ParentComponent != nullptr ? ParentComponent->GetComponentScale() : FVector::OneVector;
+	const FVector Translation = bExportNonUniformScale ? RelativeTransform.GetTranslation() * ParentScale : RelativeTransform.GetTranslation();
+	const FVector Scale = bExportNonUniformScale ? FVector::OneVector : RelativeTransform.GetScale3D();
+
+	const FGLTFJsonNodeIndex NodeIndex = Builder.AddChildNode(ParentNodeIndex);
+	FGLTFJsonNode& Node = Builder.GetNode(NodeIndex);
+	Node.Name = FGLTFNameUtility::GetName(SceneComponent);
+	Node.Translation = FGLTFConverterUtility::ConvertPosition(Translation, Builder.ExportOptions->ExportScale);
+	Node.Rotation = FGLTFConverterUtility::ConvertRotation(RelativeTransform.GetRotation());
+	Node.Scale = FGLTFConverterUtility::ConvertScale(Scale);
+
+	const FGLTFJsonVector3 ComponentNodeScale = FGLTFConverterUtility::ConvertScale(bExportNonUniformScale ? Transform.GetScale3D() : FVector::OneVector);
+
+	if (SceneComponent->bHiddenInGame) // TODO: make this configurable
+	{
+		// ignore any visible properties
+	}
+	else if (const UStaticMeshComponent* StaticMeshComponent = Cast<UStaticMeshComponent>(SceneComponent))
+	{
+		if (bExportNonUniformScale)
+		{
+			FGLTFJsonNode MeshNode;
+			MeshNode.Name = FGLTFNameUtility::GetName(StaticMeshComponent);
+			MeshNode.Scale = ComponentNodeScale;
+			MeshNode.Mesh = Builder.GetOrAddMesh(StaticMeshComponent);
+			MeshNode.LightMap = Builder.GetOrAddLightMap(StaticMeshComponent);
+			Builder.AddChildComponentNode(NodeIndex, MeshNode);
+		}
+		else
+		{
+			Node.Mesh = Builder.GetOrAddMesh(StaticMeshComponent);
+			Node.LightMap = Builder.GetOrAddLightMap(StaticMeshComponent);
+		}
+	}
+	else if (const USkeletalMeshComponent* SkeletalMeshComponent = Cast<USkeletalMeshComponent>(SceneComponent))
+	{
+		if (bExportNonUniformScale)
+		{
+			FGLTFJsonNode MeshNode;
+			MeshNode.Name = FGLTFNameUtility::GetName(SkeletalMeshComponent);
+			MeshNode.Scale = ComponentNodeScale;
+			MeshNode.Mesh = Builder.GetOrAddMesh(SkeletalMeshComponent);
+			Builder.AddChildComponentNode(NodeIndex, MeshNode);
+		}
+		else
+		{
+			Node.Mesh = Builder.GetOrAddMesh(SkeletalMeshComponent);
+		}
+	}
+	else if (const UGLTFInteractionHotspotComponent* HotspotComponent = Cast<UGLTFInteractionHotspotComponent>(SceneComponent))
+	{
+		if (Builder.ExportOptions->bExportInteractionHotspots)
+		{
+			if (bExportNonUniformScale)
+			{
+				FGLTFJsonNode HotspotNode;
+				HotspotNode.Name = FGLTFNameUtility::GetName(HotspotComponent);
+				HotspotNode.Scale = ComponentNodeScale;
+
+				HotspotNode.Hotspot = Builder.GetOrAddHotspot(HotspotComponent);
+				Builder.AddChildComponentNode(NodeIndex, HotspotNode);
+			}
+			else
+			{
+				Node.Hotspot = Builder.GetOrAddHotspot(HotspotComponent);
+			}
+		}
+		else
+		{
+			Builder.AddWarningMessage(FString::Printf(TEXT("Interaction-hotspot %s disabled by export options"), *Owner->GetName()));
+		}
+	}
+	else if (const UCameraComponent* CameraComponent = Cast<UCameraComponent>(SceneComponent))
+	{
+		if (Builder.ExportOptions->bExportCameras)
+		{
+			// TODO: conversion of camera direction should be done in separate converter
+			FGLTFJsonNode CameraNode;
+			CameraNode.Name = FGLTFNameUtility::GetName(CameraComponent);
+			CameraNode.Rotation = FGLTFConverterUtility::ConvertCameraDirection();
+			CameraNode.Scale = ComponentNodeScale;
+			CameraNode.Camera = Builder.GetOrAddCamera(CameraComponent);
+			Builder.AddChildComponentNode(NodeIndex, CameraNode);
+		}
+		else
+		{
+			Builder.AddWarningMessage(FString::Printf(TEXT("Camera %s disabled by export options"), *Owner->GetName()));
+		}
+	}
+	else if (const ULightComponent* LightComponent = Cast<ULightComponent>(SceneComponent))
+	{
+		if (Builder.ExportOptions->ShouldExportLight(LightComponent->Mobility))
+		{
+			// TODO: conversion of light direction should be done in separate converter
+			FGLTFJsonNode LightNode;
+			LightNode.Name = FGLTFNameUtility::GetName(LightComponent);
+			LightNode.Rotation = FGLTFConverterUtility::ConvertLightDirection();
+			LightNode.Scale = ComponentNodeScale;
+			LightNode.Light = Builder.GetOrAddLight(LightComponent);
+			Builder.AddChildComponentNode(NodeIndex, LightNode);
+		}
+		else
+		{
+			Builder.AddWarningMessage(FString::Printf(TEXT("Light %s disabled by export options"), *Owner->GetName()));
+		}
+	}
+
+	// TODO: add support for SkyLight?
+
+	return NodeIndex;
+}
+
+FGLTFJsonNodeIndex FGLTFActorConverter::Convert(const AActor* Actor)
+{
+	if (Builder.bSelectedActorsOnly && !Actor->IsSelected())
+	{
+		return FGLTFJsonNodeIndex(INDEX_NONE);
+	}
+
+	const USceneComponent* RootComponent = Actor->GetRootComponent();
+	const FGLTFJsonNodeIndex RootNodeIndex = Builder.GetOrAddNode(RootComponent);
+
+	const UBlueprint* Blueprint = FGLTFActorUtility::GetBlueprintFromActor(Actor);
+	if (FGLTFActorUtility::IsSkySphereBlueprint(Blueprint))
+	{
+		// Ignore mesh and light components that are part of BP_SkySphere
+		// TODO: add support for BP_SkySphere later?
+	}
+	else if (FGLTFActorUtility::IsHDRIBackdropBlueprint(Blueprint))
+	{
+		if (Builder.ExportOptions->bExportHDRIBackdrops)
+		{
+			FGLTFJsonNode& RootNode = Builder.GetNode(RootNodeIndex);
+			RootNode.Backdrop = Builder.GetOrAddBackdrop(Actor);
+		}
+		else
+		{
+			Builder.AddWarningMessage(FString::Printf(
+				TEXT("HDRIBackdrop %s disabled by export options"),
+				*Actor->GetName()));
+		}
+	}
+	else
+	{
+		// TODO: to reduce number of nodes, only export components that are of interest
+
+		for (const UActorComponent* Component : Actor->GetComponents())
+		{
+			const USceneComponent* SceneComponent = Cast<USceneComponent>(Component);
+			if (SceneComponent != nullptr)
+			{
+				Builder.GetOrAddNode(SceneComponent);
+			}
+		}
+	}
+
+	return RootNodeIndex;
+}
diff --git a/Engine/Plugins/Enterprise/GLTFExporter/MigratedFromGit/Source/GLTFExporter/Private/Converters/GLTFNodeConverters.h b/Engine/Plugins/Enterprise/GLTFExporter/MigratedFromGit/Source/GLTFExporter/Private/Converters/GLTFNodeConverters.h
new file mode 100644
--- /dev/null
+++ b/Engine/Plugins/Enterprise/GLTFExporter/MigratedFromGit/Source/GLTFExporter/Private/Converters/GLTFNodeConverters.h
@@ -0,0 +1,22 @@
+// Copyright Epic Games, Inc. All Rights Reserved.
+
+#pragma once
+
+#include "Json/GLTFJsonIndex.h"
+#include "Converters/GLTFConverter.h"
+#include "Converters/GLTFBuilderContext.h"
+#include "Engine.h"
+
+class FGLTFSceneComponentConverter : public FGLTFBuilderContext, public TGLTFConverter<FGLTFJsonNodeIndex, const USceneComponent*>
+{
+	using FGLTFBuilderContext::FGLTFBuilderContext;
+
+	FGLTFJsonNodeIndex Convert(const USceneComponent* SceneComponent) override final;
+};
+
+class FGLTFActorConverter : public FGLTFBuilderContext, public TGLTFConverter<FGLTFJsonNodeIndex, const AActor*>
+{
+	using FGLTFBuilderContext::FGLTFBuilderContext;
+
+	FGLTFJsonNodeIndex Convert(const AActor* Actor) override final;
+};
diff --git a/Engine/Plugins/Enterprise/GLTFExporter/MigratedFromGit/Source/GLTFExporter/Private/Converters/GLTFSamplerConverters.cpp b/Engine/Plugins/Enterprise/GLTFExporter/MigratedFromGit/Source/GLTFExporter/Private/Converters/GLTFSamplerConverters.cpp
new file mode 100644
--- /dev/null
+++ b/Engine/Plugins/Enterprise/GLTFExporter/MigratedFromGit/Source/GLTFExporter/Private/Converters/GLTFSamplerConverters.cpp
@@ -0,0 +1,42 @@
+// Copyright Epic Games, Inc. All Rights Reserved.
+
+#include "Converters/GLTFSamplerConverters.h"
+#include "Builders/GLTFContainerBuilder.h"
+#include "Converters/GLTFConverterUtility.h"
+#include "Converters/GLTFTextureUtility.h"
+
+FGLTFJsonSamplerIndex FGLTFTextureSamplerConverter::Convert(const UTexture* Texture)
+{
+	// TODO: maybe we should reuse existing samplers?
+
+	FGLTFJsonSampler JsonSampler;
+	Texture->GetName(JsonSampler.Name);
+
+	if (Texture->IsA<ULightMapTexture2D>())
+	{
+		// Override default filter settings for LightMap textures (which otherwise is "nearest")
+		JsonSampler.MinFilter = EGLTFJsonTextureFilter::LinearMipmapLinear;
+		JsonSampler.MagFilter = EGLTFJsonTextureFilter::Linear;
+	}
+	else
+	{
+		JsonSampler.MinFilter = FGLTFConverterUtility::ConvertMinFilter(Texture->Filter, Texture->LODGroup);
+		JsonSampler.MagFilter = FGLTFConverterUtility::ConvertMagFilter(Texture->Filter, Texture->LODGroup);
+	}
+
+	if (FGLTFTextureUtility::IsCubemap(Texture))
+	{
+		JsonSampler.WrapS = EGLTFJsonTextureWrap::ClampToEdge;
+		JsonSampler.WrapT = EGLTFJsonTextureWrap::ClampToEdge;
+	}
+	else
+	{
+		const TTuple<TextureAddress, TextureAddress> AddressXY = FGLTFTextureUtility::GetAddressXY(Texture);
+		// TODO: report error if AddressX == TA_MAX or AddressY == TA_MAX
+
+		JsonSampler.WrapS = FGLTFConverterUtility::ConvertWrap(AddressXY.Get<0>());
+		JsonSampler.WrapT = FGLTFConverterUtility::ConvertWrap(AddressXY.Get<1>());
+	}
+
+	return Builder.AddSampler(JsonSampler);
+}
diff --git a/Engine/Plugins/Enterprise/GLTFExporter/MigratedFromGit/Source/GLTFExporter/Private/Converters/GLTFSamplerConverters.h b/Engine/Plugins/Enterprise/GLTFExporter/MigratedFromGit/Source/GLTFExporter/Private/Converters/GLTFSamplerConverters.h
new file mode 100644
--- /dev/null
+++ b/Engine/Plugins/Enterprise/GLTFExporter/MigratedFromGit/Source/GLTFExporter/Private/Converters/GLTFSamplerConverters.h
@@ -0,0 +1,15 @@
+// Copyright Epic Games, Inc. All Rights Reserved.
+
+#pragma once
+
+#include "Json/GLTFJsonIndex.h"
+#include "Converters/GLTFConverter.h"
+#include "Converters/GLTFBuilderContext.h"
+#include "Engine.h"
+
+class FGLTFTextureSamplerConverter : public FGLTFBuilderContext, public TGLTFConverter<FGLTFJsonSamplerIndex, const UTexture*>
+{
+	using FGLTFBuilderContext::FGLTFBuilderContext;
+
+	FGLTFJsonSamplerIndex Convert(const UTexture* Texture) override final;
+};
diff --git a/Engine/Plugins/Enterprise/GLTFExporter/MigratedFromGit/Source/GLTFExporter/Private/Converters/GLTFSceneConverters.cpp b/Engine/Plugins/Enterprise/GLTFExporter/MigratedFromGit/Source/GLTFExporter/Private/Converters/GLTFSceneConverters.cpp
new file mode 100644
--- /dev/null
+++ b/Engine/Plugins/Enterprise/GLTFExporter/MigratedFromGit/Source/GLTFExporter/Private/Converters/GLTFSceneConverters.cpp
@@ -0,0 +1,52 @@
+// Copyright Epic Games, Inc. All Rights Reserved.
+
+#include "Converters/GLTFSceneConverters.h"
+#include "Builders/GLTFContainerBuilder.h"
+#include "Converters/GLTFConverterUtility.h"
+#include "Converters/GLTFActorUtility.h"
+
+FGLTFJsonSceneIndex FGLTFLevelConverter::Convert(const ULevel* Level)
+{
+	FGLTFJsonScene Scene;
+
+	if (const UWorld* World = Level->GetWorld())
+	{
+		World->GetName(Scene.Name);
+	}
+	else
+	{
+		Level->GetName(Scene.Name);
+	}
+
+	// TODO: export Level->Model ?
+
+	for (const AActor* Actor : Level->Actors)
+	{
+		// TODO: should a LevelVariantSet be exported even if not selected for export?
+		if (const ALevelVariantSetsActor *LevelVariantSetsActor = Cast<ALevelVariantSetsActor>(Actor))
+		{
+			if (Builder.ExportOptions->bExportVariantSets)
+			{
+				const FGLTFJsonLevelVariantSetsIndex LevelVariantSetsIndex = Builder.GetOrAddLevelVariantSets(LevelVariantSetsActor);
+				if (LevelVariantSetsIndex != INDEX_NONE)
+				{
+					Scene.LevelVariantSets.Add(LevelVariantSetsIndex);
+				}
+			}
+			else
+			{
+				Builder.AddWarningMessage(FString::Printf(
+					TEXT("Level Variant Set %s disabled by export options"),
+					*LevelVariantSetsActor->GetName()));
+			}
+		}
+
+		const FGLTFJsonNodeIndex NodeIndex = Builder.GetOrAddNode(Actor);
+		if (NodeIndex != INDEX_NONE && FGLTFActorUtility::IsRootActor(Actor, Builder.bSelectedActorsOnly))
+		{
+			Scene.Nodes.Add(NodeIndex);
+		}
+	}
+
+	return Builder.AddScene(Scene);
+}
diff --git a/Engine/Plugins/Enterprise/GLTFExporter/MigratedFromGit/Source/GLTFExporter/Private/Converters/GLTFSceneConverters.h b/Engine/Plugins/Enterprise/GLTFExporter/MigratedFromGit/Source/GLTFExporter/Private/Converters/GLTFSceneConverters.h
new file mode 100644
--- /dev/null
+++ b/Engine/Plugins/Enterprise/GLTFExporter/MigratedFromGit/Source/GLTFExporter/Private/Converters/GLTFSceneConverters.h
@@ -0,0 +1,15 @@
+// Copyright Epic Games, Inc. All Rights Reserved.
+
+#pragma once
+
+#include "Json/GLTFJsonIndex.h"
+#include "Converters/GLTFConverter.h"
+#include "Converters/GLTFBuilderContext.h"
+#include "Engine.h"
+
+class FGLTFLevelConverter : public FGLTFBuilderContext, public TGLTFConverter<FGLTFJsonSceneIndex, const ULevel*>
+{
+	using FGLTFBuilderContext::FGLTFBuilderContext;
+
+	FGLTFJsonSceneIndex Convert(const ULevel* Level) override final;
+};
diff --git a/Engine/Plugins/Enterprise/GLTFExporter/MigratedFromGit/Source/GLTFExporter/Private/Converters/GLTFSkeletalMeshConverters.cpp b/Engine/Plugins/Enterprise/GLTFExporter/MigratedFromGit/Source/GLTFExporter/Private/Converters/GLTFSkeletalMeshConverters.cpp
deleted file mode 100644
diff --git a/Engine/Plugins/Enterprise/GLTFExporter/MigratedFromGit/Source/GLTFExporter/Private/Converters/GLTFSkeletalMeshConverters.h b/Engine/Plugins/Enterprise/GLTFExporter/MigratedFromGit/Source/GLTFExporter/Private/Converters/GLTFSkeletalMeshConverters.h
deleted file mode 100644
diff --git a/Engine/Plugins/Enterprise/GLTFExporter/MigratedFromGit/Source/GLTFExporter/Private/Converters/GLTFStaticMeshConverters.cpp b/Engine/Plugins/Enterprise/GLTFExporter/MigratedFromGit/Source/GLTFExporter/Private/Converters/GLTFStaticMeshConverters.cpp
deleted file mode 100644
diff --git a/Engine/Plugins/Enterprise/GLTFExporter/MigratedFromGit/Source/GLTFExporter/Private/Converters/GLTFStaticMeshConverters.h b/Engine/Plugins/Enterprise/GLTFExporter/MigratedFromGit/Source/GLTFExporter/Private/Converters/GLTFStaticMeshConverters.h
deleted file mode 100644
diff --git a/Engine/Plugins/Enterprise/GLTFExporter/MigratedFromGit/Source/GLTFExporter/Private/Converters/GLTFTextureConverters.cpp b/Engine/Plugins/Enterprise/GLTFExporter/MigratedFromGit/Source/GLTFExporter/Private/Converters/GLTFTextureConverters.cpp
--- a/Engine/Plugins/Enterprise/GLTFExporter/MigratedFromGit/Source/GLTFExporter/Private/Converters/GLTFTextureConverters.cpp
+++ b/Engine/Plugins/Enterprise/GLTFExporter/MigratedFromGit/Source/GLTFExporter/Private/Converters/GLTFTextureConverters.cpp
@@ -5,42 +5,6 @@
 #include "Converters/GLTFConverterUtility.h"
 #include "Converters/GLTFTextureUtility.h"
 
-FGLTFJsonSamplerIndex FGLTFTextureSamplerConverter::Convert(const UTexture* Texture)
-{
-	// TODO: maybe we should reuse existing samplers?
-
-	FGLTFJsonSampler JsonSampler;
-	Texture->GetName(JsonSampler.Name);
-
-	if (Texture->IsA<ULightMapTexture2D>())
-	{
-		// Override default filter settings for LightMap textures (which otherwise is "nearest")
-		JsonSampler.MinFilter = EGLTFJsonTextureFilter::LinearMipmapLinear;
-		JsonSampler.MagFilter = EGLTFJsonTextureFilter::Linear;
-	}
-	else
-	{
-		JsonSampler.MinFilter = FGLTFConverterUtility::ConvertMinFilter(Texture->Filter, Texture->LODGroup);
-		JsonSampler.MagFilter = FGLTFConverterUtility::ConvertMagFilter(Texture->Filter, Texture->LODGroup);
-	}
-
-	if (FGLTFTextureUtility::IsCubemap(Texture))
-	{
-		JsonSampler.WrapS = EGLTFJsonTextureWrap::ClampToEdge;
-		JsonSampler.WrapT = EGLTFJsonTextureWrap::ClampToEdge;
-	}
-	else
-	{
-		const TTuple<TextureAddress, TextureAddress> AddressXY = FGLTFTextureUtility::GetAddressXY(Texture);
-		// TODO: report error if AddressX == TA_MAX or AddressY == TA_MAX
-
-		JsonSampler.WrapS = FGLTFConverterUtility::ConvertWrap(AddressXY.Get<0>());
-		JsonSampler.WrapT = FGLTFConverterUtility::ConvertWrap(AddressXY.Get<1>());
-	}
-
-	return Builder.AddSampler(JsonSampler);
-}
-
 FGLTFJsonTextureIndex FGLTFTexture2DConverter::Convert(const UTexture2D* Texture2D)
 {
 	FGLTFJsonImageIndex ImageIndex;
diff --git a/Engine/Plugins/Enterprise/GLTFExporter/MigratedFromGit/Source/GLTFExporter/Private/Converters/GLTFTextureConverters.h b/Engine/Plugins/Enterprise/GLTFExporter/MigratedFromGit/Source/GLTFExporter/Private/Converters/GLTFTextureConverters.h
--- a/Engine/Plugins/Enterprise/GLTFExporter/MigratedFromGit/Source/GLTFExporter/Private/Converters/GLTFTextureConverters.h
+++ b/Engine/Plugins/Enterprise/GLTFExporter/MigratedFromGit/Source/GLTFExporter/Private/Converters/GLTFTextureConverters.h
@@ -7,13 +7,6 @@
 #include "Converters/GLTFBuilderContext.h"
 #include "Engine.h"
 
-class FGLTFTextureSamplerConverter : public FGLTFBuilderContext, public TGLTFConverter<FGLTFJsonSamplerIndex, const UTexture*>
-{
-	using FGLTFBuilderContext::FGLTFBuilderContext;
-
-	FGLTFJsonSamplerIndex Convert(const UTexture* Texture) override final;
-};
-
 class FGLTFTexture2DConverter : public FGLTFBuilderContext, public TGLTFConverter<FGLTFJsonTextureIndex, const UTexture2D*>
 {
 	using FGLTFBuilderContext::FGLTFBuilderContext;
diff --git a/Engine/Plugins/Enterprise/GLTFExporter/MigratedFromGit/Source/GLTFExporter/Private/Converters/GLTFVertexBufferConverters.cpp b/Engine/Plugins/Enterprise/GLTFExporter/MigratedFromGit/Source/GLTFExporter/Private/Converters/GLTFVertexBufferConverters.cpp
deleted file mode 100644
diff --git a/Engine/Plugins/Enterprise/GLTFExporter/MigratedFromGit/Source/GLTFExporter/Private/Converters/GLTFVertexBufferConverters.h b/Engine/Plugins/Enterprise/GLTFExporter/MigratedFromGit/Source/GLTFExporter/Private/Converters/GLTFVertexBufferConverters.h
deleted file mode 100644
