[GLTFExporter] refactor: fix inconsistent naming between VariantSets and LevelVariantSets

commit 298c58828c32ea8faccbdc03066712e37b416fb9
Date: Sun, 20 Sep 2020 22:06:40 +0200

    refactor: fix inconsistent naming between VariantSets and LevelVariantSets
    (cherry picked from commit 207f44b5698c67eb82a021b87340276d6e053971​)

#skipci

diff --git a/Engine/Plugins/Enterprise/GLTFExporter/MigratedFromGit/Source/GLTFExporter/Private/Builders/GLTFConvertBuilder.cpp b/Engine/Plugins/Enterprise/GLTFExporter/MigratedFromGit/Source/GLTFExporter/Private/Builders/GLTFConvertBuilder.cpp
--- a/Engine/Plugins/Enterprise/GLTFExporter/MigratedFromGit/Source/GLTFExporter/Private/Builders/GLTFConvertBuilder.cpp
+++ b/Engine/Plugins/Enterprise/GLTFExporter/MigratedFromGit/Source/GLTFExporter/Private/Builders/GLTFConvertBuilder.cpp
@@ -324,14 +324,14 @@ FGLTFJsonBackdropIndex FGLTFConvertBuilder::GetOrAddBackdrop(const AActor* Actor
 	return BackdropConverter.GetOrAdd(Actor);
 }
 
-FGLTFJsonLevelVariantSetsIndex FGLTFConvertBuilder::GetOrAddLevelVariantSets(const ALevelVariantSetsActor* LevelVariantSetsActor)
+FGLTFJsonVariationIndex FGLTFConvertBuilder::GetOrAddVariation(const ALevelVariantSetsActor* LevelVariantSetsActor)
 {
 	if (LevelVariantSetsActor == nullptr)
 	{
-		return FGLTFJsonLevelVariantSetsIndex(INDEX_NONE);
+		return FGLTFJsonVariationIndex(INDEX_NONE);
 	}
 
-	return LevelVariantSetsConverter.GetOrAdd(LevelVariantSetsActor);
+	return VariationConverter.GetOrAdd(LevelVariantSetsActor);
 }
 
 FGLTFJsonLightMapIndex FGLTFConvertBuilder::GetOrAddLightMap(const UStaticMeshComponent* StaticMeshComponent)
diff --git a/Engine/Plugins/Enterprise/GLTFExporter/MigratedFromGit/Source/GLTFExporter/Private/Builders/GLTFConvertBuilder.h b/Engine/Plugins/Enterprise/GLTFExporter/MigratedFromGit/Source/GLTFExporter/Private/Builders/GLTFConvertBuilder.h
--- a/Engine/Plugins/Enterprise/GLTFExporter/MigratedFromGit/Source/GLTFExporter/Private/Builders/GLTFConvertBuilder.h
+++ b/Engine/Plugins/Enterprise/GLTFExporter/MigratedFromGit/Source/GLTFExporter/Private/Builders/GLTFConvertBuilder.h
@@ -13,7 +13,7 @@
 #include "Converters/GLTFCameraConverters.h"
 #include "Converters/GLTFLightConverters.h"
 #include "Converters/GLTFBackdropConverters.h"
-#include "Converters/GLTFVariantSetConverters.h"
+#include "Converters/GLTFVarationConverters.h"
 #include "Converters/GLTFLightMapConverters.h"
 #include "Converters/GLTFHotspotConverters.h"
 
@@ -59,7 +59,7 @@ public:
 	FGLTFJsonLightIndex GetOrAddLight(const ULightComponent* LightComponent);
 
 	FGLTFJsonBackdropIndex GetOrAddBackdrop(const AActor* Actor);
-	FGLTFJsonLevelVariantSetsIndex GetOrAddLevelVariantSets(const ALevelVariantSetsActor* LevelVariantSetsActor);
+	FGLTFJsonVariationIndex GetOrAddVariation(const ALevelVariantSetsActor* LevelVariantSetsActor);
 	FGLTFJsonLightMapIndex GetOrAddLightMap(const UStaticMeshComponent* StaticMeshComponent);
 	FGLTFJsonHotspotIndex GetOrAddHotspot(const UGLTFInteractionHotspotComponent* HotspotComponent);
 
@@ -96,7 +96,7 @@ private:
 	FGLTFLightComponentConverter LightComponentConverter = *this;
 
 	FGLTFBackdropConverter BackdropConverter = *this;
-	FGLTFLevelVariantSetsConverter LevelVariantSetsConverter = *this;
+	FGLTFVariationConverter VariationConverter = *this;
 	FGLTFLightMapConverter LightMapConverter = *this;
 	FGLTFHotspotComponentConverter HotspotComponentConverter = *this;
 };
diff --git a/Engine/Plugins/Enterprise/GLTFExporter/MigratedFromGit/Source/GLTFExporter/Private/Builders/GLTFJsonBuilder.cpp b/Engine/Plugins/Enterprise/GLTFExporter/MigratedFromGit/Source/GLTFExporter/Private/Builders/GLTFJsonBuilder.cpp
--- a/Engine/Plugins/Enterprise/GLTFExporter/MigratedFromGit/Source/GLTFExporter/Private/Builders/GLTFJsonBuilder.cpp
+++ b/Engine/Plugins/Enterprise/GLTFExporter/MigratedFromGit/Source/GLTFExporter/Private/Builders/GLTFJsonBuilder.cpp
@@ -68,9 +68,9 @@ FGLTFJsonBackdropIndex FGLTFJsonBuilder::AddBackdrop(const FGLTFJsonBackdrop& Js
 	return FGLTFJsonBackdropIndex(JsonRoot.Backdrops.Add(JsonBackdrop));
 }
 
-FGLTFJsonLevelVariantSetsIndex FGLTFJsonBuilder::AddLevelVariantSets(const FGLTFJsonLevelVariantSets& LevelVariantSets)
+FGLTFJsonVariationIndex FGLTFJsonBuilder::AddVariation(const FGLTFJsonVariation& JsonVariation)
 {
-	return FGLTFJsonLevelVariantSetsIndex(JsonRoot.LevelVariantSets.Add(LevelVariantSets));
+	return FGLTFJsonVariationIndex(JsonRoot.Variations.Add(JsonVariation));
 }
 
 FGLTFJsonLightMapIndex FGLTFJsonBuilder::AddLightMap(const FGLTFJsonLightMap& JsonLightMap)
@@ -172,9 +172,9 @@ FGLTFJsonBackdrop& FGLTFJsonBuilder::GetBackdrop(FGLTFJsonBackdropIndex Backdrop
 	return JsonRoot.Backdrops[BackdropIndex];
 }
 
-FGLTFJsonLevelVariantSets& FGLTFJsonBuilder::GetLevelVariantSets(FGLTFJsonLevelVariantSetsIndex LevelVariantSetsIndex)
+FGLTFJsonVariation& FGLTFJsonBuilder::GetVariation(FGLTFJsonVariationIndex VariationIndex)
 {
-	return JsonRoot.LevelVariantSets[LevelVariantSetsIndex];
+	return JsonRoot.Variations[VariationIndex];
 }
 
 FGLTFJsonLightMap& FGLTFJsonBuilder::GetLightMap(FGLTFJsonLightMapIndex LightMapIndex)
diff --git a/Engine/Plugins/Enterprise/GLTFExporter/MigratedFromGit/Source/GLTFExporter/Private/Builders/GLTFJsonBuilder.h b/Engine/Plugins/Enterprise/GLTFExporter/MigratedFromGit/Source/GLTFExporter/Private/Builders/GLTFJsonBuilder.h
--- a/Engine/Plugins/Enterprise/GLTFExporter/MigratedFromGit/Source/GLTFExporter/Private/Builders/GLTFJsonBuilder.h
+++ b/Engine/Plugins/Enterprise/GLTFExporter/MigratedFromGit/Source/GLTFExporter/Private/Builders/GLTFJsonBuilder.h
@@ -25,7 +25,7 @@ public:
 	FGLTFJsonSceneIndex AddScene(const FGLTFJsonScene& JsonScene = {});
 	FGLTFJsonTextureIndex AddTexture(const FGLTFJsonTexture& JsonTexture = {});
 	FGLTFJsonBackdropIndex AddBackdrop(const FGLTFJsonBackdrop& JsonBackdrop = {});
-	FGLTFJsonLevelVariantSetsIndex AddLevelVariantSets(const FGLTFJsonLevelVariantSets& LevelVariantSets = {});
+	FGLTFJsonVariationIndex AddVariation(const FGLTFJsonVariation& JsonVariation = {});
 	FGLTFJsonLightMapIndex AddLightMap(const FGLTFJsonLightMap& JsonLightMap = {});
 	FGLTFJsonLightIndex AddLight(const FGLTFJsonLight& JsonLight = {});
 	FGLTFJsonHotspotIndex AddHotspot(const FGLTFJsonHotspot& JsonHotspot = {});
@@ -45,7 +45,7 @@ public:
 	FGLTFJsonScene& GetScene(FGLTFJsonSceneIndex SceneIndex);
 	FGLTFJsonTexture& GetTexture(FGLTFJsonTextureIndex TextureIndex);
 	FGLTFJsonBackdrop& GetBackdrop(FGLTFJsonBackdropIndex BackdropIndex);
-	FGLTFJsonLevelVariantSets& GetLevelVariantSets(FGLTFJsonLevelVariantSetsIndex LevelVariantSetsIndex);
+	FGLTFJsonVariation& GetVariation(FGLTFJsonVariationIndex VariationIndex);
 	FGLTFJsonLightMap& GetLightMap(FGLTFJsonLightMapIndex LightMapIndex);
 	FGLTFJsonLight& GetLight(FGLTFJsonLightIndex LightIndex);
 	FGLTFJsonHotspot& GetHotspot(FGLTFJsonHotspotIndex HotspotIndex);
diff --git a/Engine/Plugins/Enterprise/GLTFExporter/MigratedFromGit/Source/GLTFExporter/Private/Converters/GLTFSceneConverters.cpp b/Engine/Plugins/Enterprise/GLTFExporter/MigratedFromGit/Source/GLTFExporter/Private/Converters/GLTFSceneConverters.cpp
--- a/Engine/Plugins/Enterprise/GLTFExporter/MigratedFromGit/Source/GLTFExporter/Private/Converters/GLTFSceneConverters.cpp
+++ b/Engine/Plugins/Enterprise/GLTFExporter/MigratedFromGit/Source/GLTFExporter/Private/Converters/GLTFSceneConverters.cpp
@@ -27,10 +27,10 @@ FGLTFJsonSceneIndex FGLTFLevelConverter::Convert(const ULevel* Level)
 		{
 			if (Builder.ExportOptions->bExportVariantSets)
 			{
-				const FGLTFJsonLevelVariantSetsIndex LevelVariantSetsIndex = Builder.GetOrAddLevelVariantSets(LevelVariantSetsActor);
-				if (LevelVariantSetsIndex != INDEX_NONE)
+				const FGLTFJsonVariationIndex VariationIndex = Builder.GetOrAddVariation(LevelVariantSetsActor);
+				if (VariationIndex != INDEX_NONE)
 				{
-					Scene.LevelVariantSets.Add(LevelVariantSetsIndex);
+					Scene.Variations.Add(VariationIndex);
 				}
 			}
 			else
diff --git a/Engine/Plugins/Enterprise/GLTFExporter/MigratedFromGit/Source/GLTFExporter/Private/Converters/GLTFVarationConverters.cpp b/Engine/Plugins/Enterprise/GLTFExporter/MigratedFromGit/Source/GLTFExporter/Private/Converters/GLTFVarationConverters.cpp
new file mode 100644
--- /dev/null
+++ b/Engine/Plugins/Enterprise/GLTFExporter/MigratedFromGit/Source/GLTFExporter/Private/Converters/GLTFVarationConverters.cpp
@@ -0,0 +1,440 @@
+// Copyright Epic Games, Inc. All Rights Reserved.
+
+#include "Converters/GLTFVarationConverters.h"
+#include "Builders/GLTFContainerBuilder.h"
+#include "LevelVariantSetsActor.h"
+#include "VariantObjectBinding.h"
+#include "PropertyValueMaterial.h"
+#include "LevelVariantSets.h"
+#include "PropertyValue.h"
+#include "VariantSet.h"
+#include "Variant.h"
+
+namespace
+{
+	// TODO: replace hack with safe solution by adding proper API access to UPropertyValue
+	class UPropertyValueHack : public UPropertyValue
+	{
+	public:
+		const TArray<FCapturedPropSegment>& GetCapturedPropSegments() const
+		{
+			return CapturedPropSegments;
+		}
+	};
+} // anonymous namespace
+
+FGLTFJsonVariationIndex FGLTFVariationConverter::Convert(const ALevelVariantSetsActor* LevelVariantSetsActor)
+{
+	const ULevelVariantSets* LevelVariantSets = const_cast<ALevelVariantSetsActor*>(LevelVariantSetsActor)->GetLevelVariantSets(true);
+	if (LevelVariantSets == nullptr)
+	{
+		return FGLTFJsonVariationIndex(INDEX_NONE);
+	}
+
+	FGLTFJsonVariation JsonVariation;
+	LevelVariantSets->GetName(JsonVariation.Name);
+
+	for (const UVariantSet* VariantSet: LevelVariantSets->GetVariantSets())
+	{
+		FGLTFJsonVariantSet JsonVariantSet;
+		JsonVariantSet.Name = VariantSet->GetDisplayText().ToString();
+
+		for (const UVariant* Variant: VariantSet->GetVariants())
+		{
+			FGLTFJsonVariant JsonVariant;
+			if (TryParseVariant(JsonVariant, Variant))
+			{
+				JsonVariantSet.Variants.Add(JsonVariant);
+			}
+		}
+
+		if (JsonVariantSet.Variants.Num() > 0)
+		{
+			JsonVariation.VariantSets.Add(JsonVariantSet);
+		}
+		else
+		{
+			Builder.AddWarningMessage(FString::Printf(
+				TEXT("Variant-set has no supported variants and will be skipped. Context: %s"),
+				*GetLogContext(VariantSet)));
+		}
+	}
+
+	if (JsonVariation.VariantSets.Num() == 0)
+	{
+		return FGLTFJsonVariationIndex(INDEX_NONE);
+	}
+
+	return Builder.AddVariation(JsonVariation);
+}
+
+bool FGLTFVariationConverter::TryParseVariant(FGLTFJsonVariant& OutVariant, const UVariant* Variant) const
+{
+	FGLTFJsonVariant JsonVariant;
+
+	for (const UVariantObjectBinding* Binding: Variant->GetBindings())
+	{
+		TryParseVariantBinding(JsonVariant, Binding);
+	}
+
+	if (JsonVariant.Nodes.Num() == 0)
+	{
+		Builder.AddWarningMessage(FString::Printf(
+			TEXT("Variant has no supported bindings and will be skipped. Context: %s"),
+			*GetLogContext(Variant)));
+		return false;
+	}
+
+	JsonVariant.Name = Variant->GetDisplayText().ToString();
+	JsonVariant.bIsActive = const_cast<UVariant*>(Variant)->IsActive();
+
+	if (const UTexture2D* Thumbnail = const_cast<UVariant*>(Variant)->GetThumbnail())
+	{
+		// TODO: if thumbnail is generated from viewport, give it a variant-relevant name (instead of "Texture2D_#")
+		JsonVariant.Thumbnail = Builder.GetOrAddTexture(Thumbnail);
+	}
+
+	OutVariant = JsonVariant;
+	return true;
+}
+
+bool FGLTFVariationConverter::TryParseVariantBinding(FGLTFJsonVariant& OutVariant, const UVariantObjectBinding* Binding) const
+{
+	bool bHasParsedAnyProperty = false;
+
+	for (UPropertyValue* Property: Binding->GetCapturedProperties())
+	{
+		if (!const_cast<UPropertyValue*>(Property)->Resolve() || !Property->HasRecordedData())
+		{
+			Builder.AddWarningMessage(FString::Printf(
+				TEXT("Property is missing recorded data, it will be skipped. Context: %s"),
+				*GetLogContext(Property)));
+			continue;
+		}
+
+		const bool bIsVisibilityProperty =
+			Property->GetPropertyName() == TEXT("bVisible") &&
+			Property->GetPropertyClass()->IsChildOf(FBoolProperty::StaticClass());
+
+		const bool bIsMaterialProperty = Property->IsA<UPropertyValueMaterial>();
+		const bool bIsStaticMeshProperty = Property->GetPropertyName() == TEXT("StaticMesh");
+		const bool bIsSkeletalMeshProperty = Property->GetPropertyName() == TEXT("SkeletalMesh");
+
+		if (bIsVisibilityProperty)
+		{
+			if (TryParseVisibilityPropertyValue(OutVariant, Property))
+			{
+				bHasParsedAnyProperty = true;
+			}
+		}
+		else if (bIsMaterialProperty)
+		{
+			if (TryParseMaterialPropertyValue(OutVariant, Property))
+			{
+				bHasParsedAnyProperty = true;
+			}
+		}
+		else if (bIsStaticMeshProperty)
+		{
+			if (TryParseStaticMeshPropertyValue(OutVariant, Property))
+			{
+				bHasParsedAnyProperty = true;
+			}
+		}
+		else if (bIsSkeletalMeshProperty)
+		{
+			if (TryParseSkeletalMeshPropertyValue(OutVariant, Property))
+			{
+				bHasParsedAnyProperty = true;
+			}
+		}
+		else
+		{
+			// TODO: add support for more properties
+
+			Builder.AddWarningMessage(FString::Printf(
+				TEXT("Property is not supported and will be skipped. Context: %s"),
+				*GetLogContext(Property)));
+		}
+	}
+
+	if (!bHasParsedAnyProperty)
+	{
+		Builder.AddWarningMessage(FString::Printf(
+			TEXT("Binding has no supported properties and will be skipped. Context: %s"),
+			*GetLogContext(Binding)));
+	}
+
+	return bHasParsedAnyProperty;
+}
+
+bool FGLTFVariationConverter::TryParseVisibilityPropertyValue(FGLTFJsonVariant& OutVariant, const UPropertyValue* Property) const
+{
+	const USceneComponent* Target = static_cast<USceneComponent*>(Property->GetPropertyParentContainerAddress());
+	if (Target == nullptr)
+	{
+		Builder.AddWarningMessage(FString::Printf(
+			TEXT("Target object for property is invalid, the property will be skipped. Context: %s"),
+			*GetLogContext(Property)));
+		return false;
+	}
+
+	const AActor* Owner = Target->GetOwner();
+	if (Owner == nullptr)
+	{
+		Builder.AddWarningMessage(FString::Printf(TEXT("Invalid scene component, the property will be skipped. Context: %s"),
+			*GetLogContext(Property)));
+		return false;
+	}
+
+	if (Builder.bSelectedActorsOnly && !Owner->IsSelected())
+	{
+		Builder.AddWarningMessage(FString::Printf(
+			TEXT("Target actor for property is not selected for export, the property will be skipped. Context: %s"),
+			*GetLogContext(Property)));
+		return false;
+	}
+
+	bool bIsVisible;
+
+	if (!TryGetPropertyValue(const_cast<UPropertyValue*>(Property), bIsVisible))
+	{
+		Builder.AddWarningMessage(FString::Printf(
+			TEXT("Failed to parse recorded data for property, it will be skipped. Context: %s"),
+			*GetLogContext(Property)));
+		return false;
+	}
+
+	const FGLTFJsonNodeIndex NodeIndex = Builder.GetOrAddNode(Target); // TODO: what if the target is not selected for export?
+	const FGLTFJsonNodeIndex ComponentNodeIndex = Builder.GetComponentNodeIndex(NodeIndex);
+
+	FGLTFJsonVariantNodeProperties& NodeProperties = OutVariant.Nodes.FindOrAdd(ComponentNodeIndex);
+
+	NodeProperties.Node = ComponentNodeIndex;
+	NodeProperties.bIsVisible = bIsVisible;
+	return true;
+}
+
+bool FGLTFVariationConverter::TryParseMaterialPropertyValue(FGLTFJsonVariant& OutVariant, const UPropertyValue* Property) const
+{
+	UPropertyValueMaterial* MaterialProperty = Cast<UPropertyValueMaterial>(const_cast<UPropertyValue*>(Property));
+	if (MaterialProperty == nullptr)
+	{
+		Builder.AddWarningMessage(FString::Printf(
+			TEXT("Material property is invalid, the property will be skipped. Context: %s"),
+			*GetLogContext(Property)));
+		return false;
+	}
+
+	const USceneComponent* Target = static_cast<USceneComponent*>(Property->GetPropertyParentContainerAddress());
+	if (Target == nullptr)
+	{
+		Builder.AddWarningMessage(FString::Printf(
+			TEXT("Target object for property is invalid, the property will be skipped. Context: %s"),
+			*GetLogContext(Property)));
+		return false;
+	}
+
+	const AActor* Owner = Target->GetOwner();
+	if (Owner == nullptr)
+	{
+		Builder.AddWarningMessage(FString::Printf(TEXT("Invalid scene component, the property will be skipped. Context: %s"),
+			*GetLogContext(Property)));
+		return false;
+	}
+
+	if (Builder.bSelectedActorsOnly && !Owner->IsSelected())
+	{
+		Builder.AddWarningMessage(FString::Printf(
+			TEXT("Target actor for property is not selected for export, the property will be skipped. Context: %s"),
+			*GetLogContext(Property)));
+		return false;
+	}
+
+	const TArray<FCapturedPropSegment>& CapturedPropSegments = reinterpret_cast<UPropertyValueHack*>(MaterialProperty)->GetCapturedPropSegments();
+	const int32 NumPropSegments = CapturedPropSegments.Num();
+
+	if (NumPropSegments < 1)
+	{
+		Builder.AddWarningMessage(FString::Printf(
+			TEXT("Failed to parse element index to apply the material to, the property will be skipped. Context: %s"),
+			*GetLogContext(Property)));
+		return false;
+	}
+
+	// NOTE: UPropertyValueMaterial::GetMaterial does *not* ensure that the recorded data has been loaded,
+	// so we need to call UProperty::GetRecordedData first to make that happen.
+	MaterialProperty->GetRecordedData();
+
+	// TODO: find way to determine whether the material is null because "None" was selected, or because it failed to resolve
+	const UMaterialInterface* Material = MaterialProperty->GetMaterial();
+	const FGLTFJsonMaterialIndex MaterialIndex = Builder.GetOrAddMaterial(Material);
+	const int32 ElementIndex = CapturedPropSegments[NumPropSegments - 1].PropertyIndex;
+
+	FGLTFJsonVariantMaterial VariantMaterial;
+	VariantMaterial.Material = MaterialIndex;
+	VariantMaterial.Index = ElementIndex;
+
+	const FGLTFJsonNodeIndex NodeIndex = Builder.GetOrAddNode(Target);
+	const FGLTFJsonNodeIndex ComponentNodeIndex = Builder.GetComponentNodeIndex(NodeIndex);
+	FGLTFJsonVariantNodeProperties& NodeProperties = OutVariant.Nodes.FindOrAdd(ComponentNodeIndex);
+
+	NodeProperties.Node = ComponentNodeIndex;
+	NodeProperties.Materials.Add(VariantMaterial);
+	return true;
+}
+
+bool FGLTFVariationConverter::TryParseStaticMeshPropertyValue(FGLTFJsonVariant& OutVariant, const UPropertyValue* Property) const
+{
+	const USceneComponent* Target = static_cast<USceneComponent*>(Property->GetPropertyParentContainerAddress());
+	if (Target == nullptr)
+	{
+		Builder.AddWarningMessage(FString::Printf(
+			TEXT("Target object for property is invalid, the property will be skipped. Context: %s"),
+			*GetLogContext(Property)));
+		return false;
+	}
+
+	const AActor* Owner = Target->GetOwner();
+	if (Owner == nullptr)
+	{
+		Builder.AddWarningMessage(FString::Printf(TEXT("Invalid scene component, the property will be skipped. Context: %s"),
+			*GetLogContext(Property)));
+		return false;
+	}
+
+	if (Builder.bSelectedActorsOnly && !Owner->IsSelected())
+	{
+		Builder.AddWarningMessage(FString::Printf(
+			TEXT("Target actor for property is not selected for export, the property will be skipped. Context: %s"),
+			*GetLogContext(Property)));
+		return false;
+	}
+
+	const UMeshComponent* MeshComponent = Cast<UMeshComponent>(Target);
+	if (MeshComponent == nullptr)
+	{
+		Builder.AddWarningMessage(FString::Printf(
+			TEXT("Target object for property has no mesh-component, the property will be skipped. Context: %s"),
+			*GetLogContext(Property)));
+		return false;
+	}
+
+	const UStaticMesh* StaticMesh;
+
+	if (!TryGetPropertyValue(const_cast<UPropertyValue*>(Property), StaticMesh))
+	{
+		Builder.AddWarningMessage(FString::Printf(
+			TEXT("Failed to parse recorded data for property, it will be skipped. Context: %s"),
+			*GetLogContext(Property)));
+		return false;
+	}
+
+	const FGLTFMaterialArray OverrideMaterials(MeshComponent->OverrideMaterials);
+	const FGLTFJsonNodeIndex NodeIndex = Builder.GetOrAddNode(Target);
+	const FGLTFJsonNodeIndex ComponentNodeIndex = Builder.GetComponentNodeIndex(NodeIndex);
+	const FGLTFJsonMeshIndex MeshIndex = Builder.GetOrAddMesh(StaticMesh, -1, nullptr, OverrideMaterials);
+	FGLTFJsonVariantNodeProperties& NodeProperties = OutVariant.Nodes.FindOrAdd(ComponentNodeIndex);
+
+	NodeProperties.Node = ComponentNodeIndex;
+	NodeProperties.Mesh = MeshIndex;
+	return true;
+}
+
+bool FGLTFVariationConverter::TryParseSkeletalMeshPropertyValue(FGLTFJsonVariant& OutVariant, const UPropertyValue* Property) const
+{
+	const USceneComponent* Target = static_cast<USceneComponent*>(Property->GetPropertyParentContainerAddress());
+	if (Target == nullptr)
+	{
+		Builder.AddWarningMessage(FString::Printf(
+			TEXT("Target object for property is invalid, the property will be skipped. Context: %s"),
+			*GetLogContext(Property)));
+		return false;
+	}
+
+	const AActor* Owner = Target->GetOwner();
+	if (Owner == nullptr)
+	{
+		Builder.AddWarningMessage(FString::Printf(TEXT("Invalid scene component, the property will be skipped. Context: %s"),
+			*GetLogContext(Property)));
+		return false;
+	}
+
+	if (Builder.bSelectedActorsOnly && !Owner->IsSelected())
+	{
+		Builder.AddWarningMessage(FString::Printf(
+			TEXT("Target actor for property is not selected for export, the property will be skipped. Context: %s"),
+			*GetLogContext(Property)));
+		return false;
+	}
+
+	const UMeshComponent* MeshComponent = Cast<UMeshComponent>(Target);
+	if (MeshComponent == nullptr)
+	{
+		Builder.AddWarningMessage(FString::Printf(
+			TEXT("Target object for property has no mesh-component, the property will be skipped. Context: %s"),
+			*GetLogContext(Property)));
+		return false;
+	}
+
+	const USkeletalMesh* SkeletalMesh;
+
+	if (!TryGetPropertyValue(const_cast<UPropertyValue*>(Property), SkeletalMesh) || SkeletalMesh == nullptr)
+	{
+		Builder.AddWarningMessage(FString::Printf(
+			TEXT("Failed to parse recorded data for property, it will be skipped. Context: %s"),
+			*GetLogContext(Property)));
+		return false;
+	}
+
+	const FGLTFMaterialArray OverrideMaterials(MeshComponent->OverrideMaterials);
+	const FGLTFJsonNodeIndex NodeIndex = Builder.GetOrAddNode(Target);
+	const FGLTFJsonMeshIndex MeshIndex = Builder.GetOrAddMesh(SkeletalMesh, -1, nullptr, nullptr, OverrideMaterials);
+	const FGLTFJsonNodeIndex ComponentNodeIndex = Builder.GetComponentNodeIndex(NodeIndex);
+	FGLTFJsonVariantNodeProperties& NodeProperties = OutVariant.Nodes.FindOrAdd(ComponentNodeIndex);
+
+	NodeProperties.Node = ComponentNodeIndex;
+	NodeProperties.Mesh = MeshIndex;
+	return true;
+}
+
+template<typename T>
+bool FGLTFVariationConverter::TryGetPropertyValue(UPropertyValue* Property, T& OutValue) const
+{
+	if (Property == nullptr || !Property->HasRecordedData())
+	{
+		return false;
+	}
+
+	FMemory::Memcpy(&OutValue, Property->GetRecordedData().GetData(), sizeof(T));
+	return true;
+}
+
+FString FGLTFVariationConverter::GetLogContext(const UPropertyValue* Property) const
+{
+	const UVariantObjectBinding* Parent = Property->GetParent();
+	return GetLogContext(Parent) + TEXT("/") + Property->GetFullDisplayString();
+}
+
+FString FGLTFVariationConverter::GetLogContext(const UVariantObjectBinding* Binding) const
+{
+	const UVariant* Parent = const_cast<UVariantObjectBinding*>(Binding)->GetParent();
+	return GetLogContext(Parent) + TEXT("/") + Binding->GetDisplayText().ToString();
+}
+
+FString FGLTFVariationConverter::GetLogContext(const UVariant* Variant) const
+{
+	const UVariantSet* Parent = const_cast<UVariant*>(Variant)->GetParent();
+	return GetLogContext(Parent) + TEXT("/") + Variant->GetDisplayText().ToString();
+}
+
+FString FGLTFVariationConverter::GetLogContext(const UVariantSet* VariantSet) const
+{
+	const ULevelVariantSets* Parent = const_cast<UVariantSet*>(VariantSet)->GetParent();
+	return GetLogContext(Parent) + TEXT("/") + VariantSet->GetDisplayText().ToString();
+}
+
+FString FGLTFVariationConverter::GetLogContext(const ULevelVariantSets* LevelVariantSets) const
+{
+	return LevelVariantSets->GetName();
+}
diff --git a/Engine/Plugins/Enterprise/GLTFExporter/MigratedFromGit/Source/GLTFExporter/Private/Converters/GLTFVarationConverters.h b/Engine/Plugins/Enterprise/GLTFExporter/MigratedFromGit/Source/GLTFExporter/Private/Converters/GLTFVarationConverters.h
new file mode 100644
--- /dev/null
+++ b/Engine/Plugins/Enterprise/GLTFExporter/MigratedFromGit/Source/GLTFExporter/Private/Converters/GLTFVarationConverters.h
@@ -0,0 +1,36 @@
+// Copyright Epic Games, Inc. All Rights Reserved.
+
+#pragma once
+
+#include "Json/GLTFJsonIndex.h"
+#include "Json/GLTFJsonVariation.h"
+#include "Converters/GLTFConverter.h"
+#include "Converters/GLTFBuilderContext.h"
+#include "Engine.h"
+#include "LevelVariantSetsActor.h"
+#include "Variant.h"
+#include "PropertyValue.h"
+#include "LevelVariantSets.h"
+
+class FGLTFVariationConverter : public FGLTFBuilderContext, public TGLTFConverter<FGLTFJsonVariationIndex, const ALevelVariantSetsActor*>
+{
+	using FGLTFBuilderContext::FGLTFBuilderContext;
+
+	FGLTFJsonVariationIndex Convert(const ALevelVariantSetsActor* LevelVariantSetsActor) override final;
+
+	bool TryParseVariant(FGLTFJsonVariant& OutVariant, const UVariant* Variant) const;
+	bool TryParseVariantBinding(FGLTFJsonVariant& OutVariant, const UVariantObjectBinding* Binding) const;
+	bool TryParseVisibilityPropertyValue(FGLTFJsonVariant& OutVariant, const UPropertyValue* Property) const;
+	bool TryParseMaterialPropertyValue(FGLTFJsonVariant& OutVariant, const UPropertyValue* Property) const;
+	bool TryParseStaticMeshPropertyValue(FGLTFJsonVariant& OutVariant, const UPropertyValue* Property) const;
+	bool TryParseSkeletalMeshPropertyValue(FGLTFJsonVariant& OutVariant, const UPropertyValue* Property) const;
+
+	template<typename T>
+	bool TryGetPropertyValue(UPropertyValue* Property, T& OutValue) const;
+
+	FString GetLogContext(const UPropertyValue* Property) const;
+	FString GetLogContext(const UVariantObjectBinding* Binding) const;
+	FString GetLogContext(const UVariant* Variant) const;
+	FString GetLogContext(const UVariantSet* VariantSet) const;
+	FString GetLogContext(const ULevelVariantSets* LevelVariantSets) const;
+};
diff --git a/Engine/Plugins/Enterprise/GLTFExporter/MigratedFromGit/Source/GLTFExporter/Private/Converters/GLTFVariantSetConverters.cpp b/Engine/Plugins/Enterprise/GLTFExporter/MigratedFromGit/Source/GLTFExporter/Private/Converters/GLTFVariantSetConverters.cpp
deleted file mode 100644
diff --git a/Engine/Plugins/Enterprise/GLTFExporter/MigratedFromGit/Source/GLTFExporter/Private/Converters/GLTFVariantSetConverters.h b/Engine/Plugins/Enterprise/GLTFExporter/MigratedFromGit/Source/GLTFExporter/Private/Converters/GLTFVariantSetConverters.h
deleted file mode 100644
diff --git a/Engine/Plugins/Enterprise/GLTFExporter/MigratedFromGit/Source/GLTFExporter/Private/Json/GLTFJsonIndex.h b/Engine/Plugins/Enterprise/GLTFExporter/MigratedFromGit/Source/GLTFExporter/Private/Json/GLTFJsonIndex.h
--- a/Engine/Plugins/Enterprise/GLTFExporter/MigratedFromGit/Source/GLTFExporter/Private/Json/GLTFJsonIndex.h
+++ b/Engine/Plugins/Enterprise/GLTFExporter/MigratedFromGit/Source/GLTFExporter/Private/Json/GLTFJsonIndex.h
@@ -65,7 +65,7 @@ struct FGLTFJsonSceneIndex : FGLTFJsonIndex<FGLTFJsonSceneIndex> { using FGLTFJs
 struct FGLTFJsonSkinIndex : FGLTFJsonIndex<FGLTFJsonSkinIndex> { using FGLTFJsonIndex::FGLTFJsonIndex; };
 struct FGLTFJsonTextureIndex : FGLTFJsonIndex<FGLTFJsonTextureIndex> { using FGLTFJsonIndex::FGLTFJsonIndex; };
 struct FGLTFJsonBackdropIndex : FGLTFJsonIndex<FGLTFJsonBackdropIndex> { using FGLTFJsonIndex::FGLTFJsonIndex; };
-struct FGLTFJsonLevelVariantSetsIndex : FGLTFJsonIndex<FGLTFJsonLevelVariantSetsIndex> { using FGLTFJsonIndex::FGLTFJsonIndex; };
+struct FGLTFJsonVariationIndex : FGLTFJsonIndex<FGLTFJsonVariationIndex> { using FGLTFJsonIndex::FGLTFJsonIndex; };
 struct FGLTFJsonLightMapIndex : FGLTFJsonIndex<FGLTFJsonLightMapIndex> { using FGLTFJsonIndex::FGLTFJsonIndex; };
 struct FGLTFJsonLightIndex : FGLTFJsonIndex<FGLTFJsonLightIndex> { using FGLTFJsonIndex::FGLTFJsonIndex; };
 struct FGLTFJsonHotspotIndex : FGLTFJsonIndex<FGLTFJsonHotspotIndex> { using FGLTFJsonIndex::FGLTFJsonIndex; };
diff --git a/Engine/Plugins/Enterprise/GLTFExporter/MigratedFromGit/Source/GLTFExporter/Private/Json/GLTFJsonLevelVariantSets.h b/Engine/Plugins/Enterprise/GLTFExporter/MigratedFromGit/Source/GLTFExporter/Private/Json/GLTFJsonLevelVariantSets.h
deleted file mode 100644
diff --git a/Engine/Plugins/Enterprise/GLTFExporter/MigratedFromGit/Source/GLTFExporter/Private/Json/GLTFJsonRoot.h b/Engine/Plugins/Enterprise/GLTFExporter/MigratedFromGit/Source/GLTFExporter/Private/Json/GLTFJsonRoot.h
--- a/Engine/Plugins/Enterprise/GLTFExporter/MigratedFromGit/Source/GLTFExporter/Private/Json/GLTFJsonRoot.h
+++ b/Engine/Plugins/Enterprise/GLTFExporter/MigratedFromGit/Source/GLTFExporter/Private/Json/GLTFJsonRoot.h
@@ -16,7 +16,7 @@
 #include "Json/GLTFJsonScene.h"
 #include "Json/GLTFJsonTexture.h"
 #include "Json/GLTFJsonBackdrop.h"
-#include "Json/GLTFJsonLevelVariantSets.h"
+#include "Json/GLTFJsonVariation.h"
 #include "Json/GLTFJsonLightMap.h"
 #include "Json/GLTFJsonLight.h"
 #include "Json/FGLTFJsonHotspot.h"
@@ -74,8 +74,8 @@ struct FGLTFJsonRoot
 	TArray<FGLTFJsonBackdrop>   Backdrops;
 	TArray<FGLTFJsonLightMap>   LightMaps;
 	TArray<FGLTFJsonLight>      Lights;
-	TArray<FGLTFJsonLevelVariantSets> LevelVariantSets;
-	TArray<FGLTFJsonHotspot>          Hotspots;
+	TArray<FGLTFJsonVariation>        Variations;
+	TArray<FGLTFJsonHotspot>    Hotspots;
 
 	template <class CharType = TCHAR, class PrintPolicy = TPrettyJsonPrintPolicy<CharType>>
 	void WriteObject(TJsonWriter<CharType, PrintPolicy>& JsonWriter) const
@@ -104,7 +104,7 @@ struct FGLTFJsonRoot
 		FGLTFJsonUtility::WriteObjectArray(JsonWriter, TEXT("scenes"), Scenes, AllExtensions);
 		FGLTFJsonUtility::WriteObjectArray(JsonWriter, TEXT("textures"), Textures, AllExtensions);
 
-		if (Backdrops.Num() > 0 || LevelVariantSets.Num() > 0 || LightMaps.Num() > 0 || Lights.Num() > 0 || Hotspots.Num() > 0 /* TODO: add more extension support */)
+		if (Backdrops.Num() > 0 || Variations.Num() > 0 || LightMaps.Num() > 0 || Lights.Num() > 0 || Hotspots.Num() > 0 /* TODO: add more extension support */)
 		{
 			JsonWriter.WriteObjectStart(TEXT("extensions"));
 
@@ -118,13 +118,13 @@ struct FGLTFJsonRoot
 				JsonWriter.WriteObjectEnd();
 			}
 
-			if (LevelVariantSets.Num() > 0)
+			if (Variations.Num() > 0)
 			{
 				const EGLTFJsonExtension Extension = EGLTFJsonExtension::EPIC_LevelVariantSets;
 				AllExtensions.Used.Add(Extension);
 
 				JsonWriter.WriteObjectStart(FGLTFJsonUtility::ToString(Extension));
-				FGLTFJsonUtility::WriteObjectArray(JsonWriter, TEXT("levelVariantSets"), LevelVariantSets, AllExtensions);
+				FGLTFJsonUtility::WriteObjectArray(JsonWriter, TEXT("levelVariantSets"), Variations, AllExtensions);
 				JsonWriter.WriteObjectEnd();
 			}
 
diff --git a/Engine/Plugins/Enterprise/GLTFExporter/MigratedFromGit/Source/GLTFExporter/Private/Json/GLTFJsonScene.h b/Engine/Plugins/Enterprise/GLTFExporter/MigratedFromGit/Source/GLTFExporter/Private/Json/GLTFJsonScene.h
--- a/Engine/Plugins/Enterprise/GLTFExporter/MigratedFromGit/Source/GLTFExporter/Private/Json/GLTFJsonScene.h
+++ b/Engine/Plugins/Enterprise/GLTFExporter/MigratedFromGit/Source/GLTFExporter/Private/Json/GLTFJsonScene.h
@@ -9,8 +9,8 @@ struct FGLTFJsonScene
 {
 	FString Name;
 
-	TArray<FGLTFJsonNodeIndex>            Nodes;
-	TArray<FGLTFJsonLevelVariantSetsIndex> LevelVariantSets;
+	TArray<FGLTFJsonNodeIndex> Nodes;
+	TArray<FGLTFJsonVariationIndex>  Variations;
 
 	template <class CharType = TCHAR, class PrintPolicy = TPrettyJsonPrintPolicy<CharType>>
 	void WriteObject(TJsonWriter<CharType, PrintPolicy>& JsonWriter, FGLTFJsonExtensions& Extensions) const
@@ -27,7 +27,7 @@ struct FGLTFJsonScene
 			JsonWriter.WriteValue(TEXT("nodes"), Nodes);
 		}
 
-		if (LevelVariantSets.Num() > 0)
+		if (Variations.Num() > 0)
 		{
 			const EGLTFJsonExtension Extension = EGLTFJsonExtension::EPIC_LevelVariantSets;
 			Extensions.Used.Add(Extension);
@@ -35,7 +35,7 @@ struct FGLTFJsonScene
 			JsonWriter.WriteObjectStart(TEXT("extensions"));
 			JsonWriter.WriteObjectStart(FGLTFJsonUtility::ToString(Extension));
 
-			JsonWriter.WriteValue(TEXT("levelVariantSets"), LevelVariantSets);
+			JsonWriter.WriteValue(TEXT("levelVariantSets"), Variations);
 
 			JsonWriter.WriteObjectEnd();
 			JsonWriter.WriteObjectEnd();
diff --git a/Engine/Plugins/Enterprise/GLTFExporter/MigratedFromGit/Source/GLTFExporter/Private/Json/GLTFJsonVariation.h b/Engine/Plugins/Enterprise/GLTFExporter/MigratedFromGit/Source/GLTFExporter/Private/Json/GLTFJsonVariation.h
new file mode 100644
--- /dev/null
+++ b/Engine/Plugins/Enterprise/GLTFExporter/MigratedFromGit/Source/GLTFExporter/Private/Json/GLTFJsonVariation.h
@@ -0,0 +1,154 @@
+// Copyright Epic Games, Inc. All Rights Reserved.
+
+#pragma once
+
+#include "GLTFJsonIndex.h"
+#include "Json/GLTFJsonExtensions.h"
+#include "Json/GLTFJsonUtility.h"
+#include "Serialization/JsonSerializer.h"
+#include "Misc/Optional.h"
+
+struct FGLTFJsonVariantMaterial
+{
+	FGLTFJsonMaterialIndex Material;
+	int32                  Index;
+
+	FGLTFJsonVariantMaterial()
+		: Material(INDEX_NONE)
+		, Index(INDEX_NONE)
+	{
+	}
+
+	template <class CharType = TCHAR, class PrintPolicy = TPrettyJsonPrintPolicy<CharType>>
+	void WriteObject(TJsonWriter<CharType, PrintPolicy>& JsonWriter, FGLTFJsonExtensions& Extensions) const
+	{
+		JsonWriter.WriteObjectStart();
+
+		JsonWriter.WriteValue(TEXT("material"), Material);
+
+		if (Index != INDEX_NONE)
+		{
+			JsonWriter.WriteValue(TEXT("index"), Index);
+		}
+
+		JsonWriter.WriteObjectEnd();
+	}
+};
+
+struct FGLTFJsonVariantNodeProperties
+{
+	FGLTFJsonNodeIndex Node;
+	TOptional<bool>    bIsVisible;
+
+	TOptional<FGLTFJsonMeshIndex>    Mesh;
+	TArray<FGLTFJsonVariantMaterial> Materials;
+
+	template <class CharType = TCHAR, class PrintPolicy = TPrettyJsonPrintPolicy<CharType>>
+	void WriteObject(TJsonWriter<CharType, PrintPolicy>& JsonWriter, FGLTFJsonExtensions& Extensions) const
+	{
+		JsonWriter.WriteObjectStart();
+
+		if (Node != INDEX_NONE)
+		{
+			JsonWriter.WriteValue(TEXT("node"), Node);
+		}
+
+		JsonWriter.WriteObjectStart(TEXT("properties"));
+
+		if (bIsVisible.IsSet())
+		{
+			JsonWriter.WriteValue(TEXT("visible"), bIsVisible.GetValue());
+		}
+
+		if (Mesh.IsSet())
+		{
+			JsonWriter.WriteValue(TEXT("mesh"), Mesh.GetValue());
+		}
+
+		FGLTFJsonUtility::WriteObjectArray(JsonWriter, TEXT("materials"), Materials, Extensions);
+
+		JsonWriter.WriteObjectEnd();
+		JsonWriter.WriteObjectEnd();
+	}
+};
+
+struct FGLTFJsonVariant
+{
+	typedef TMap<FGLTFJsonNodeIndex, FGLTFJsonVariantNodeProperties> FNodeProperties;
+
+	FString Name;
+	bool    bIsActive;
+
+	FGLTFJsonTextureIndex Thumbnail;
+	FNodeProperties       Nodes;
+
+	template <class CharType = TCHAR, class PrintPolicy = TPrettyJsonPrintPolicy<CharType>>
+	void WriteObject(TJsonWriter<CharType, PrintPolicy>& JsonWriter, FGLTFJsonExtensions& Extensions) const
+	{
+		JsonWriter.WriteObjectStart();
+		JsonWriter.WriteValue(TEXT("name"), Name);
+		JsonWriter.WriteValue(TEXT("active"), bIsActive);
+
+		if (Thumbnail != INDEX_NONE)
+		{
+			JsonWriter.WriteValue(TEXT("thumbnail"), Thumbnail);
+		}
+
+		JsonWriter.WriteArrayStart(TEXT("nodes"));
+
+		for (const auto& DataPair: Nodes)
+		{
+			if (DataPair.Key != INDEX_NONE)
+			{
+				DataPair.Value.WriteObject(JsonWriter, Extensions);
+			}
+		}
+
+		JsonWriter.WriteArrayEnd();
+		JsonWriter.WriteObjectEnd();
+	}
+};
+
+struct FGLTFJsonVariantSet
+{
+	FString Name;
+
+	TArray<FGLTFJsonVariant> Variants;
+
+	template <class CharType = TCHAR, class PrintPolicy = TPrettyJsonPrintPolicy<CharType>>
+	void WriteObject(TJsonWriter<CharType, PrintPolicy>& JsonWriter, FGLTFJsonExtensions& Extensions) const
+	{
+		JsonWriter.WriteObjectStart();
+
+		if (!Name.IsEmpty())
+		{
+			JsonWriter.WriteValue(TEXT("name"), Name);
+		}
+
+		FGLTFJsonUtility::WriteObjectArray(JsonWriter, TEXT("variants"), Variants, Extensions);
+
+		JsonWriter.WriteObjectEnd();
+	}
+};
+
+struct FGLTFJsonVariation
+{
+	FString Name;
+
+	TArray<FGLTFJsonVariantSet> VariantSets;
+
+	template <class CharType = TCHAR, class PrintPolicy = TPrettyJsonPrintPolicy<CharType>>
+	void WriteObject(TJsonWriter<CharType, PrintPolicy>& JsonWriter, FGLTFJsonExtensions& Extensions) const
+	{
+		JsonWriter.WriteObjectStart();
+
+		if (!Name.IsEmpty())
+		{
+			JsonWriter.WriteValue(TEXT("name"), Name);
+		}
+
+		FGLTFJsonUtility::WriteObjectArray(JsonWriter, TEXT("variantSets"), VariantSets, Extensions);
+
+		JsonWriter.WriteObjectEnd();
+	}
+};
