[GLTFExporter] refactor: remove redundant comments about anonymous namespace

commit e97c5df8b4b6cf4b22cd5c447464e855498fce34
Date: Sun, 20 Sep 2020 22:15:16 +0200

    refactor: remove redundant comments about anonymous namespace
    (cherry picked from commit 5ff79bdca98c1fb6ad4d6de460d0495e8b427932​)

#skipci

diff --git a/Engine/Plugins/Enterprise/GLTFExporter/MigratedFromGit/Source/GLTFExporter/Private/Converters/GLTFMaterialConverters.cpp b/Engine/Plugins/Enterprise/GLTFExporter/MigratedFromGit/Source/GLTFExporter/Private/Converters/GLTFMaterialConverters.cpp
--- a/Engine/Plugins/Enterprise/GLTFExporter/MigratedFromGit/Source/GLTFExporter/Private/Converters/GLTFMaterialConverters.cpp
+++ b/Engine/Plugins/Enterprise/GLTFExporter/MigratedFromGit/Source/GLTFExporter/Private/Converters/GLTFMaterialConverters.cpp
@@ -42,7 +42,7 @@ namespace
 	const TArray<FLinearColor> OcclusionInputMasks = { OcclusionMask };
 	const TArray<FLinearColor> ClearCoatInputMasks = { ClearCoatMask };
 	const TArray<FLinearColor> ClearCoatRoughnessInputMasks = { ClearCoatRoughnessMask };
-} // anonymous namespace
+}
 
 FGLTFJsonMaterialIndex FGLTFMaterialConverter::Convert(const UMaterialInterface* Material)
 {
diff --git a/Engine/Plugins/Enterprise/GLTFExporter/MigratedFromGit/Source/GLTFExporter/Private/Converters/GLTFVarationConverters.cpp b/Engine/Plugins/Enterprise/GLTFExporter/MigratedFromGit/Source/GLTFExporter/Private/Converters/GLTFVarationConverters.cpp
--- a/Engine/Plugins/Enterprise/GLTFExporter/MigratedFromGit/Source/GLTFExporter/Private/Converters/GLTFVarationConverters.cpp
+++ b/Engine/Plugins/Enterprise/GLTFExporter/MigratedFromGit/Source/GLTFExporter/Private/Converters/GLTFVarationConverters.cpp
@@ -21,7 +21,7 @@ namespace
 			return CapturedPropSegments;
 		}
 	};
-} // anonymous namespace
+}
 
 FGLTFJsonVariationIndex FGLTFVariationConverter::Convert(const ALevelVariantSetsActor* LevelVariantSetsActor)
 {
diff --git a/Engine/Plugins/Enterprise/GLTFExporter/MigratedFromGit/Source/GLTFExporterRuntime/Private/Actors/GLTFOrbitCameraActor.cpp b/Engine/Plugins/Enterprise/GLTFExporter/MigratedFromGit/Source/GLTFExporterRuntime/Private/Actors/GLTFOrbitCameraActor.cpp
--- a/Engine/Plugins/Enterprise/GLTFExporter/MigratedFromGit/Source/GLTFExporterRuntime/Private/Actors/GLTFOrbitCameraActor.cpp
+++ b/Engine/Plugins/Enterprise/GLTFExporter/MigratedFromGit/Source/GLTFExporterRuntime/Private/Actors/GLTFOrbitCameraActor.cpp
@@ -11,7 +11,7 @@ namespace
 		// TODO: is this similar to viewer camera?
 		return (static_cast<int32>(Angle) / 360) + ((Angle < 0.0f) ? -1 : 0);
 	}
-} // Anonymous namespace
+}
 
 AGLTFOrbitCameraActor::AGLTFOrbitCameraActor(const FObjectInitializer& ObjectInitializer)
 	: Super(ObjectInitializer)
diff --git a/Engine/Plugins/Enterprise/GLTFExporter/MigratedFromGit/Source/GLTFExporterRuntime/Private/Components/GLTFInteractionHotspotComponent.cpp b/Engine/Plugins/Enterprise/GLTFExporter/MigratedFromGit/Source/GLTFExporterRuntime/Private/Components/GLTFInteractionHotspotComponent.cpp
--- a/Engine/Plugins/Enterprise/GLTFExporter/MigratedFromGit/Source/GLTFExporterRuntime/Private/Components/GLTFInteractionHotspotComponent.cpp
+++ b/Engine/Plugins/Enterprise/GLTFExporter/MigratedFromGit/Source/GLTFExporterRuntime/Private/Components/GLTFInteractionHotspotComponent.cpp
@@ -11,7 +11,7 @@
 namespace
 {
 	const float UnitSphereRadius = 50.0f;
-} // Anonymous namespace
+}
 
 DEFINE_LOG_CATEGORY_STATIC(LogEditorGLTFInteractionHotspot, Log, All);
 
@@ -95,7 +95,7 @@ void UGLTFInteractionHotspotComponent::OnRegister()
 
 	ShapeBodySetup->CollisionTraceFlag = CTF_UseSimpleAsComplex;
 	ShapeBodySetup->AggGeom.SphereElems.Add(FKSphereElem(UnitSphereRadius));
-	
+
 	Super::OnRegister();
 }
 
