[GLTFExporter] refactor: merge GLTFIndexBufferConverter and GLTFIndexContainerConverter into their respective *MeshSectionConverter

commit 4a9772a4ea217799f12df5350ac5f528487463a2
Date: Sun, 20 Sep 2020 23:55:36 +0200

    refactor: merge GLTFIndexBufferConverter and GLTFIndexContainerConverter into their respective *MeshSectionConverter
    (cherry picked from commit c4c3a3bed39c34ef65d725b0f830f4cb4cbf3198​)

#skipci

diff --git a/Engine/Plugins/Enterprise/GLTFExporter/MigratedFromGit/Source/GLTFExporter/Private/Builders/GLTFConvertBuilder.cpp b/Engine/Plugins/Enterprise/GLTFExporter/MigratedFromGit/Source/GLTFExporter/Private/Builders/GLTFConvertBuilder.cpp
--- a/Engine/Plugins/Enterprise/GLTFExporter/MigratedFromGit/Source/GLTFExporter/Private/Builders/GLTFConvertBuilder.cpp
+++ b/Engine/Plugins/Enterprise/GLTFExporter/MigratedFromGit/Source/GLTFExporter/Private/Builders/GLTFConvertBuilder.cpp
@@ -79,16 +79,6 @@ FGLTFJsonAccessorIndex FGLTFConvertBuilder::GetOrAddWeightAccessor(const FSkinWe
 	return BoneWeightVertexBufferConverter.GetOrAdd(VertexBuffer, WeightsGroupIndex);
 }
 
-FGLTFJsonBufferViewIndex FGLTFConvertBuilder::GetOrAddIndexBufferView(const FRawStaticIndexBuffer* IndexBuffer)
-{
-	if (IndexBuffer == nullptr)
-	{
-		return FGLTFJsonBufferViewIndex(INDEX_NONE);
-	}
-
-	return IndexBufferConverter.GetOrAdd(IndexBuffer);
-}
-
 FGLTFJsonAccessorIndex FGLTFConvertBuilder::GetOrAddIndexAccessor(const FStaticMeshSection* MeshSection, const FRawStaticIndexBuffer* IndexBuffer)
 {
 	if (MeshSection == nullptr)
@@ -138,16 +128,6 @@ FGLTFJsonMeshIndex FGLTFConvertBuilder::GetOrAddMesh(const UStaticMeshComponent*
 	return GetOrAddMesh(StaticMesh, LODIndex, OverrideVertexColors, OverrideMaterials);
 }
 
-FGLTFJsonBufferViewIndex FGLTFConvertBuilder::GetOrAddIndexBufferView(const FMultiSizeIndexContainer* IndexContainer)
-{
-	if (IndexContainer == nullptr)
-	{
-		return FGLTFJsonBufferViewIndex(INDEX_NONE);
-	}
-
-	return IndexContainerConverter.GetOrAdd(IndexContainer);
-}
-
 FGLTFJsonAccessorIndex FGLTFConvertBuilder::GetOrAddIndexAccessor(const FSkelMeshRenderSection* MeshSection, const FMultiSizeIndexContainer* IndexContainer)
 {
 	if (MeshSection == nullptr)
diff --git a/Engine/Plugins/Enterprise/GLTFExporter/MigratedFromGit/Source/GLTFExporter/Private/Builders/GLTFConvertBuilder.h b/Engine/Plugins/Enterprise/GLTFExporter/MigratedFromGit/Source/GLTFExporter/Private/Builders/GLTFConvertBuilder.h
--- a/Engine/Plugins/Enterprise/GLTFExporter/MigratedFromGit/Source/GLTFExporter/Private/Builders/GLTFConvertBuilder.h
+++ b/Engine/Plugins/Enterprise/GLTFExporter/MigratedFromGit/Source/GLTFExporter/Private/Builders/GLTFConvertBuilder.h
@@ -33,12 +33,10 @@ public:
 	FGLTFJsonAccessorIndex GetOrAddJointAccessor(const FSkinWeightVertexBuffer* VertexBuffer, int32 JointsGroupIndex, FGLTFBoneMap BoneMap);
 	FGLTFJsonAccessorIndex GetOrAddWeightAccessor(const FSkinWeightVertexBuffer* VertexBuffer, int32 WeightsGroupIndex);
 
-	FGLTFJsonBufferViewIndex GetOrAddIndexBufferView(const FRawStaticIndexBuffer* IndexBuffer);
 	FGLTFJsonAccessorIndex GetOrAddIndexAccessor(const FStaticMeshSection* MeshSection, const FRawStaticIndexBuffer* IndexBuffer);
 	FGLTFJsonMeshIndex GetOrAddMesh(const UStaticMesh* StaticMesh, int32 LODIndex = -1, const FColorVertexBuffer* OverrideVertexColors = nullptr, const FGLTFMaterialArray& OverrideMaterials = {});
 	FGLTFJsonMeshIndex GetOrAddMesh(const UStaticMeshComponent* StaticMeshComponent);
 
-	FGLTFJsonBufferViewIndex GetOrAddIndexBufferView(const FMultiSizeIndexContainer* IndexContainer);
 	FGLTFJsonAccessorIndex GetOrAddIndexAccessor(const FSkelMeshRenderSection* MeshSection, const FMultiSizeIndexContainer* IndexContainer);
 	FGLTFJsonMeshIndex GetOrAddMesh(const USkeletalMesh* SkeletalMesh, int32 LODIndex = -1, const FColorVertexBuffer* OverrideVertexColors = nullptr, const FSkinWeightVertexBuffer* OverrideSkinWeights = nullptr, const FGLTFMaterialArray& OverrideMaterials = {});
 	FGLTFJsonMeshIndex GetOrAddMesh(const USkeletalMeshComponent* SkeletalMeshComponent);
@@ -73,11 +71,9 @@ private:
 	FGLTFBoneIndexVertexBufferConverter BoneIndexVertexBufferConverter = *this;
 	FGLTFBoneWeightVertexBufferConverter BoneWeightVertexBufferConverter = *this;
 
-	FGLTFIndexBufferConverter IndexBufferConverter = *this;
 	FGLTFStaticMeshSectionConverter StaticMeshSectionConverter = *this;
 	FGLTFStaticMeshConverter StaticMeshConverter = *this;
 
-	FGLTFIndexContainerConverter IndexContainerConverter = *this;
 	FGLTFSkeletalMeshSectionConverter SkeletalMeshSectionConverter = *this;
 	FGLTFSkeletalMeshConverter SkeletalMeshConverter = *this;
 
diff --git a/Engine/Plugins/Enterprise/GLTFExporter/MigratedFromGit/Source/GLTFExporter/Private/Converters/GLTFAccessorConverters.cpp b/Engine/Plugins/Enterprise/GLTFExporter/MigratedFromGit/Source/GLTFExporter/Private/Converters/GLTFAccessorConverters.cpp
--- a/Engine/Plugins/Enterprise/GLTFExporter/MigratedFromGit/Source/GLTFExporter/Private/Converters/GLTFAccessorConverters.cpp
+++ b/Engine/Plugins/Enterprise/GLTFExporter/MigratedFromGit/Source/GLTFExporter/Private/Converters/GLTFAccessorConverters.cpp
@@ -243,94 +243,51 @@ FGLTFJsonAccessorIndex FGLTFBoneWeightVertexBufferConverter::Convert(const FSkin
 	return Builder.AddAccessor(JsonAccessor);
 }
 
-FGLTFJsonBufferViewIndex FGLTFIndexBufferConverter::Convert(const FRawStaticIndexBuffer* IndexBuffer)
-{
-	if (IndexBuffer->GetNumIndices() == 0)
-	{
-		return FGLTFJsonBufferViewIndex(INDEX_NONE);
-	}
-
-	if (IndexBuffer->Is32Bit())
-	{
-		TArray<uint32> Indices;
-		IndexBuffer->GetCopy(Indices);
-		return Builder.AddBufferView(Indices, sizeof(uint32), EGLTFJsonBufferTarget::ElementArrayBuffer);
-	}
-	else
-	{
-		const uint16* IndexData = IndexBuffer->AccessStream16();
-		const int32 IndexDataSize = IndexBuffer->GetIndexDataSize();
-		return Builder.AddBufferView(IndexData, IndexDataSize, sizeof(uint16), EGLTFJsonBufferTarget::ElementArrayBuffer);
-	}
-}
-
 FGLTFJsonAccessorIndex FGLTFStaticMeshSectionConverter::Convert(const FStaticMeshSection* MeshSection, const FRawStaticIndexBuffer* IndexBuffer)
 {
-	const uint32 TriangleCount = MeshSection->NumTriangles;
-	if (TriangleCount == 0)
+	const uint32 IndexCount = MeshSection->NumTriangles * 3;
+	if (IndexCount == 0)
 	{
 		return FGLTFJsonAccessorIndex(INDEX_NONE);
 	}
 
-	const uint32 FirstIndex = MeshSection->FirstIndex;
 	const bool bIs32Bit = IndexBuffer->Is32Bit();
+	const uint8 IndexTypeSize = bIs32Bit ? sizeof(uint32) : sizeof(uint16);
+	const uint32 FirstIndex = MeshSection->FirstIndex;
+
+	const void* IndexDataRootPtr = bIs32Bit ? static_cast<const void*>(IndexBuffer->AccessStream32()) : static_cast<const void*>(IndexBuffer->AccessStream16());
+	const void* IndexDataPtr = static_cast<const uint8*>(IndexDataRootPtr) + FirstIndex * IndexTypeSize;
+	const int32 IndexDataSize = IndexCount * IndexTypeSize;
 
 	FGLTFJsonAccessor JsonAccessor;
-	JsonAccessor.BufferView = Builder.GetOrAddIndexBufferView(IndexBuffer);
-	JsonAccessor.ByteOffset = FirstIndex * (bIs32Bit ? sizeof(uint32) : sizeof(uint16));
+	JsonAccessor.BufferView = Builder.AddBufferView(IndexDataPtr, IndexDataSize, IndexTypeSize, EGLTFJsonBufferTarget::ElementArrayBuffer);
 	JsonAccessor.ComponentType = bIs32Bit ? EGLTFJsonComponentType::U32 : EGLTFJsonComponentType::U16;
-	JsonAccessor.Count = TriangleCount * 3;
+	JsonAccessor.Count = IndexCount;
 	JsonAccessor.Type = EGLTFJsonAccessorType::Scalar;
 
 	return Builder.AddAccessor(JsonAccessor);
 }
 
-FGLTFJsonBufferViewIndex FGLTFIndexContainerConverter::Convert(const FMultiSizeIndexContainer* IndexContainer)
-{
-	const FRawStaticIndexBuffer16or32Interface* IndexBuffer = IndexContainer->GetIndexBuffer();
-
-	const int32 IndexCount = IndexBuffer->Num();
-	if (IndexCount <= 0)
-	{
-		return FGLTFJsonBufferViewIndex(INDEX_NONE);
-	}
-
-	if (IndexContainer->GetDataTypeSize() == sizeof(uint16))
-	{
-		TArray<uint16> Indices;
-		Indices.AddUninitialized(IndexCount);
-
-		for (int32 Index = 0; Index < IndexCount; ++Index)
-		{
-			Indices[Index] = static_cast<uint16>(IndexBuffer->Get(Index));
-		}
-
-		return Builder.AddBufferView(Indices, sizeof(uint16), EGLTFJsonBufferTarget::ElementArrayBuffer);
-	}
-	else
-	{
-		TArray<uint32> Indices;
-		IndexContainer->GetIndexBuffer(Indices);
-		return Builder.AddBufferView(Indices, sizeof(uint32), EGLTFJsonBufferTarget::ElementArrayBuffer);
-	}
-}
-
 FGLTFJsonAccessorIndex FGLTFSkeletalMeshSectionConverter::Convert(const FSkelMeshRenderSection* MeshSection, const FMultiSizeIndexContainer* IndexContainer)
 {
-	const uint32 TriangleCount = MeshSection->NumTriangles;
-	if (TriangleCount == 0)
+	const uint32 IndexCount = MeshSection->NumTriangles * 3;
+	if (IndexCount == 0)
 	{
 		return FGLTFJsonAccessorIndex(INDEX_NONE);
 	}
 
+	const uint8 IndexTypeSize = IndexContainer->GetDataTypeSize();
+	const bool bIs32Bit = IndexTypeSize == sizeof(uint32);
 	const uint32 FirstIndex = MeshSection->BaseIndex;
-	const bool bIs32Bit = IndexContainer->GetDataTypeSize() == sizeof(uint32);
+
+	const FRawStaticIndexBuffer16or32Interface* IndexBuffer = IndexContainer->GetIndexBuffer();
+	const void* IndexDataPtr = const_cast<FRawStaticIndexBuffer16or32Interface*>(IndexBuffer)->GetPointerTo(FirstIndex);
+	const int32 IndexDataSize = IndexCount * IndexTypeSize;
 
 	FGLTFJsonAccessor JsonAccessor;
-	JsonAccessor.BufferView = Builder.GetOrAddIndexBufferView(IndexContainer);
-	JsonAccessor.ByteOffset = FirstIndex * (bIs32Bit ? sizeof(uint32) : sizeof(uint16));
+	JsonAccessor.BufferView = Builder.AddBufferView(IndexDataPtr, IndexDataSize, IndexTypeSize, EGLTFJsonBufferTarget::ElementArrayBuffer);
 	JsonAccessor.ComponentType = bIs32Bit ? EGLTFJsonComponentType::U32 : EGLTFJsonComponentType::U16;
-	JsonAccessor.Count = TriangleCount * 3;
+	JsonAccessor.Count = IndexCount;
 	JsonAccessor.Type = EGLTFJsonAccessorType::Scalar;
 
 	return Builder.AddAccessor(JsonAccessor);
diff --git a/Engine/Plugins/Enterprise/GLTFExporter/MigratedFromGit/Source/GLTFExporter/Private/Converters/GLTFAccessorConverters.h b/Engine/Plugins/Enterprise/GLTFExporter/MigratedFromGit/Source/GLTFExporter/Private/Converters/GLTFAccessorConverters.h
--- a/Engine/Plugins/Enterprise/GLTFExporter/MigratedFromGit/Source/GLTFExporter/Private/Converters/GLTFAccessorConverters.h
+++ b/Engine/Plugins/Enterprise/GLTFExporter/MigratedFromGit/Source/GLTFExporter/Private/Converters/GLTFAccessorConverters.h
@@ -62,14 +62,6 @@ class FGLTFBoneWeightVertexBufferConverter : public FGLTFBuilderContext, public
 	FGLTFJsonAccessorIndex Convert(const FSkinWeightVertexBuffer* VertexBuffer, int32 WeightsGroupIndex) override final;
 };
 
-// TODO: merge into FGLTFStaticMeshSectionConverter
-class FGLTFIndexBufferConverter : public FGLTFBuilderContext, public TGLTFConverter<FGLTFJsonBufferViewIndex, const FRawStaticIndexBuffer*>
-{
-	using FGLTFBuilderContext::FGLTFBuilderContext;
-
-	FGLTFJsonBufferViewIndex Convert(const FRawStaticIndexBuffer* IndexBuffer) override final;
-};
-
 class FGLTFStaticMeshSectionConverter : public FGLTFBuilderContext, public TGLTFConverter<FGLTFJsonAccessorIndex, const FStaticMeshSection*, const FRawStaticIndexBuffer*>
 {
 	using FGLTFBuilderContext::FGLTFBuilderContext;
@@ -77,14 +69,6 @@ class FGLTFStaticMeshSectionConverter : public FGLTFBuilderContext, public TGLTF
 	FGLTFJsonAccessorIndex Convert(const FStaticMeshSection* MeshSection, const FRawStaticIndexBuffer* IndexBuffer) override final;
 };
 
-// TODO: merge into FGLTFSkeletalMeshSectionConverter
-class FGLTFIndexContainerConverter : public FGLTFBuilderContext, public TGLTFConverter<FGLTFJsonBufferViewIndex, const FMultiSizeIndexContainer*>
-{
-	using FGLTFBuilderContext::FGLTFBuilderContext;
-
-	FGLTFJsonBufferViewIndex Convert(const FMultiSizeIndexContainer* IndexContainer) override final;
-};
-
 class FGLTFSkeletalMeshSectionConverter : public FGLTFBuilderContext, public TGLTFConverter<FGLTFJsonAccessorIndex, const FSkelMeshRenderSection*, const FMultiSizeIndexContainer*>
 {
 	using FGLTFBuilderContext::FGLTFBuilderContext;
diff --git a/Engine/Plugins/Enterprise/GLTFExporter/MigratedFromGit/Source/GLTFExporter/Private/Converters/GLTFMeshConverters.cpp b/Engine/Plugins/Enterprise/GLTFExporter/MigratedFromGit/Source/GLTFExporter/Private/Converters/GLTFMeshConverters.cpp
--- a/Engine/Plugins/Enterprise/GLTFExporter/MigratedFromGit/Source/GLTFExporter/Private/Converters/GLTFMeshConverters.cpp
+++ b/Engine/Plugins/Enterprise/GLTFExporter/MigratedFromGit/Source/GLTFExporter/Private/Converters/GLTFMeshConverters.cpp
@@ -43,7 +43,6 @@ FGLTFJsonMeshIndex FGLTFStaticMeshConverter::Convert(const UStaticMesh* StaticMe
 	}
 
 	const FRawStaticIndexBuffer* IndexBuffer = &MeshLOD.IndexBuffer;
-	Builder.GetOrAddIndexBufferView(IndexBuffer);
 
 	const int32 SectionCount = MeshLOD.Sections.Num();
 	JsonMesh.Primitives.AddDefaulted(SectionCount);
@@ -123,7 +122,6 @@ FGLTFJsonMeshIndex FGLTFSkeletalMeshConverter::Convert(const USkeletalMesh* Skel
 	*/
 
 	const FMultiSizeIndexContainer* IndexContainer = &MeshLOD.MultiSizeIndexContainer;
-	Builder.GetOrAddIndexBufferView(IndexContainer);
 
 	const int32 SectionCount = MeshLOD.RenderSections.Num();
 	JsonMesh.Primitives.AddDefaulted(SectionCount);
