[GLTFExporter] refactor: add missing virtual specifier and common base class for converters of same type

commit 518a1e4fb7ed69136df28094d9beabc4de3f3634
Date: Mon, 21 Sep 2020 00:24:00 +0200

    refactor: add missing virtual specifier and common base class for converters of same type
    (cherry picked from commit c2e1bd6d48d9738099f6bb41d6c2ee253b0b932a​)

#skipci

diff --git a/Engine/Plugins/Enterprise/GLTFExporter/MigratedFromGit/Source/GLTFExporter/Private/Converters/GLTFAccessorConverters.h b/Engine/Plugins/Enterprise/GLTFExporter/MigratedFromGit/Source/GLTFExporter/Private/Converters/GLTFAccessorConverters.h
--- a/Engine/Plugins/Enterprise/GLTFExporter/MigratedFromGit/Source/GLTFExporter/Private/Converters/GLTFAccessorConverters.h
+++ b/Engine/Plugins/Enterprise/GLTFExporter/MigratedFromGit/Source/GLTFExporter/Private/Converters/GLTFAccessorConverters.h
@@ -10,68 +10,74 @@
 
 class FMultiSizeIndexContainer;
 
-class FGLTFPositionVertexBufferConverter : public FGLTFBuilderContext, public TGLTFConverter<FGLTFJsonAccessorIndex, const FPositionVertexBuffer*>
+template <typename... InputTypes>
+class FGLTFAccessorConverter : public FGLTFBuilderContext, public TGLTFConverter<FGLTFJsonAccessorIndex, InputTypes...>
 {
 	using FGLTFBuilderContext::FGLTFBuilderContext;
+};
+
+class FGLTFPositionVertexBufferConverter : public FGLTFAccessorConverter<const FPositionVertexBuffer*>
+{
+	using FGLTFAccessorConverter::FGLTFAccessorConverter;
 
-	FGLTFJsonAccessorIndex Convert(const FPositionVertexBuffer* VertexBuffer) override final;
+	virtual FGLTFJsonAccessorIndex Convert(const FPositionVertexBuffer* VertexBuffer) override final;
 };
 
-class FGLTFColorVertexBufferConverter : public FGLTFBuilderContext, public TGLTFConverter<FGLTFJsonAccessorIndex, const FColorVertexBuffer*>
+class FGLTFColorVertexBufferConverter : public FGLTFAccessorConverter<const FColorVertexBuffer*>
 {
-	using FGLTFBuilderContext::FGLTFBuilderContext;
+	using FGLTFAccessorConverter::FGLTFAccessorConverter;
 
-	FGLTFJsonAccessorIndex Convert(const FColorVertexBuffer* VertexBuffer) override final;
+	virtual FGLTFJsonAccessorIndex Convert(const FColorVertexBuffer* VertexBuffer) override final;
 };
 
-class FGLTFNormalVertexBufferConverter : public FGLTFBuilderContext, public TGLTFConverter<FGLTFJsonAccessorIndex, const FStaticMeshVertexBuffer*>
+class FGLTFNormalVertexBufferConverter : public FGLTFAccessorConverter<const FStaticMeshVertexBuffer*>
 {
-	using FGLTFBuilderContext::FGLTFBuilderContext;
+	using FGLTFAccessorConverter::FGLTFAccessorConverter;
 
-	FGLTFJsonAccessorIndex Convert(const FStaticMeshVertexBuffer* VertexBuffer) override final;
+	virtual FGLTFJsonAccessorIndex Convert(const FStaticMeshVertexBuffer* VertexBuffer) override final;
 };
 
-class FGLTFTangentVertexBufferConverter : public FGLTFBuilderContext, public TGLTFConverter<FGLTFJsonAccessorIndex, const FStaticMeshVertexBuffer*>
+class FGLTFTangentVertexBufferConverter : public FGLTFAccessorConverter<const FStaticMeshVertexBuffer*>
 {
-	using FGLTFBuilderContext::FGLTFBuilderContext;
+	using FGLTFAccessorConverter::FGLTFAccessorConverter;
 
-	FGLTFJsonAccessorIndex Convert(const FStaticMeshVertexBuffer* VertexBuffer) override final;
+	virtual FGLTFJsonAccessorIndex Convert(const FStaticMeshVertexBuffer* VertexBuffer) override final;
 };
 
-class FGLTFUVVertexBufferConverter : public FGLTFBuilderContext, public TGLTFConverter<FGLTFJsonAccessorIndex, const FStaticMeshVertexBuffer*, int32>
+class FGLTFUVVertexBufferConverter : public FGLTFAccessorConverter<const FStaticMeshVertexBuffer*, int32>
 {
-	using FGLTFBuilderContext::FGLTFBuilderContext;
+	using FGLTFAccessorConverter::FGLTFAccessorConverter;
 
-	FGLTFJsonAccessorIndex Convert(const FStaticMeshVertexBuffer* VertexBuffer, int32 UVIndex) override final;
+	virtual FGLTFJsonAccessorIndex Convert(const FStaticMeshVertexBuffer* VertexBuffer, int32 UVIndex) override final;
 };
 
-class FGLTFBoneIndexVertexBufferConverter : public FGLTFBuilderContext, public TGLTFConverter<FGLTFJsonAccessorIndex, const FSkinWeightVertexBuffer*, int32, FGLTFBoneMap>
+class FGLTFBoneIndexVertexBufferConverter : public FGLTFAccessorConverter<const FSkinWeightVertexBuffer*, int32, FGLTFBoneMap>
 {
-	using FGLTFBuilderContext::FGLTFBuilderContext;
+	using FGLTFAccessorConverter::FGLTFAccessorConverter;
 
-	FGLTFJsonAccessorIndex Convert(const FSkinWeightVertexBuffer* VertexBuffer, int32 InfluenceOffset, FGLTFBoneMap BoneMap) override final;
+	virtual FGLTFJsonAccessorIndex Convert(const FSkinWeightVertexBuffer* VertexBuffer, int32 InfluenceOffset, FGLTFBoneMap BoneMap) override final;
 
 	template <typename IndexType>
 	FGLTFJsonAccessorIndex Convert(const FSkinWeightVertexBuffer* VertexBuffer, int32 JointsGroupIndex, FGLTFBoneMap BoneMap);
 };
 
-class FGLTFBoneWeightVertexBufferConverter : public FGLTFBuilderContext, public TGLTFConverter<FGLTFJsonAccessorIndex, const FSkinWeightVertexBuffer*, int32>
+class FGLTFBoneWeightVertexBufferConverter : public FGLTFAccessorConverter<const FSkinWeightVertexBuffer*, int32>
 {
-	using FGLTFBuilderContext::FGLTFBuilderContext;
+	using FGLTFAccessorConverter::FGLTFAccessorConverter;
 
-	FGLTFJsonAccessorIndex Convert(const FSkinWeightVertexBuffer* VertexBuffer, int32 WeightsGroupIndex) override final;
+	virtual FGLTFJsonAccessorIndex Convert(const FSkinWeightVertexBuffer* VertexBuffer, int32 WeightsGroupIndex) override final;
 };
 
-class FGLTFStaticMeshSectionConverter : public FGLTFBuilderContext, public TGLTFConverter<FGLTFJsonAccessorIndex, const FStaticMeshSection*, const FRawStaticIndexBuffer*>
+class FGLTFStaticMeshSectionConverter : public FGLTFAccessorConverter<const FStaticMeshSection*, const FRawStaticIndexBuffer*>
 {
-	using FGLTFBuilderContext::FGLTFBuilderContext;
+	using FGLTFAccessorConverter::FGLTFAccessorConverter;
 
-	FGLTFJsonAccessorIndex Convert(const FStaticMeshSection* MeshSection, const FRawStaticIndexBuffer* IndexBuffer) override final;
+	virtual FGLTFJsonAccessorIndex Convert(const FStaticMeshSection* MeshSection, const FRawStaticIndexBuffer* IndexBuffer) override final;
 };
 
-class FGLTFSkeletalMeshSectionConverter : public FGLTFBuilderContext, public TGLTFConverter<FGLTFJsonAccessorIndex, const FSkelMeshRenderSection*, const FMultiSizeIndexContainer*>
+class FGLTFSkeletalMeshSectionConverter : public FGLTFAccessorConverter<const FSkelMeshRenderSection*, const FMultiSizeIndexContainer*>
 {
-	using FGLTFBuilderContext::FGLTFBuilderContext;
+	using FGLTFAccessorConverter::FGLTFAccessorConverter;
 
-	FGLTFJsonAccessorIndex Convert(const FSkelMeshRenderSection* MeshSection, const FMultiSizeIndexContainer* IndexContainer) override final;
+	virtual FGLTFJsonAccessorIndex Convert(const FSkelMeshRenderSection* MeshSection, const FMultiSizeIndexContainer* IndexContainer) override final;
 };
diff --git a/Engine/Plugins/Enterprise/GLTFExporter/MigratedFromGit/Source/GLTFExporter/Private/Converters/GLTFBackdropConverters.h b/Engine/Plugins/Enterprise/GLTFExporter/MigratedFromGit/Source/GLTFExporter/Private/Converters/GLTFBackdropConverters.h
--- a/Engine/Plugins/Enterprise/GLTFExporter/MigratedFromGit/Source/GLTFExporter/Private/Converters/GLTFBackdropConverters.h
+++ b/Engine/Plugins/Enterprise/GLTFExporter/MigratedFromGit/Source/GLTFExporter/Private/Converters/GLTFBackdropConverters.h
@@ -11,5 +11,5 @@ class FGLTFBackdropConverter : public FGLTFBuilderContext, public TGLTFConverter
 {
 	using FGLTFBuilderContext::FGLTFBuilderContext;
 
-	FGLTFJsonBackdropIndex Convert(const AActor* Actor) override final;
+	virtual FGLTFJsonBackdropIndex Convert(const AActor* Actor) override final;
 };
diff --git a/Engine/Plugins/Enterprise/GLTFExporter/MigratedFromGit/Source/GLTFExporter/Private/Converters/GLTFCameraConverters.h b/Engine/Plugins/Enterprise/GLTFExporter/MigratedFromGit/Source/GLTFExporter/Private/Converters/GLTFCameraConverters.h
--- a/Engine/Plugins/Enterprise/GLTFExporter/MigratedFromGit/Source/GLTFExporter/Private/Converters/GLTFCameraConverters.h
+++ b/Engine/Plugins/Enterprise/GLTFExporter/MigratedFromGit/Source/GLTFExporter/Private/Converters/GLTFCameraConverters.h
@@ -11,5 +11,5 @@ class FGLTFCameraComponentConverter : public FGLTFBuilderContext, public TGLTFCo
 {
 	using FGLTFBuilderContext::FGLTFBuilderContext;
 
-	FGLTFJsonCameraIndex Convert(const UCameraComponent* CameraComponent) override final;
+	virtual FGLTFJsonCameraIndex Convert(const UCameraComponent* CameraComponent) override final;
 };
diff --git a/Engine/Plugins/Enterprise/GLTFExporter/MigratedFromGit/Source/GLTFExporter/Private/Converters/GLTFHotspotConverters.h b/Engine/Plugins/Enterprise/GLTFExporter/MigratedFromGit/Source/GLTFExporter/Private/Converters/GLTFHotspotConverters.h
--- a/Engine/Plugins/Enterprise/GLTFExporter/MigratedFromGit/Source/GLTFExporter/Private/Converters/GLTFHotspotConverters.h
+++ b/Engine/Plugins/Enterprise/GLTFExporter/MigratedFromGit/Source/GLTFExporter/Private/Converters/GLTFHotspotConverters.h
@@ -12,5 +12,5 @@ class FGLTFHotspotComponentConverter : public FGLTFBuilderContext, public TGLTFC
 {
 	using FGLTFBuilderContext::FGLTFBuilderContext;
 
-	FGLTFJsonHotspotIndex Convert(const UGLTFInteractionHotspotComponent* HotspotComponent) override final;
+	virtual FGLTFJsonHotspotIndex Convert(const UGLTFInteractionHotspotComponent* HotspotComponent) override final;
 };
diff --git a/Engine/Plugins/Enterprise/GLTFExporter/MigratedFromGit/Source/GLTFExporter/Private/Converters/GLTFLightConverters.h b/Engine/Plugins/Enterprise/GLTFExporter/MigratedFromGit/Source/GLTFExporter/Private/Converters/GLTFLightConverters.h
--- a/Engine/Plugins/Enterprise/GLTFExporter/MigratedFromGit/Source/GLTFExporter/Private/Converters/GLTFLightConverters.h
+++ b/Engine/Plugins/Enterprise/GLTFExporter/MigratedFromGit/Source/GLTFExporter/Private/Converters/GLTFLightConverters.h
@@ -11,5 +11,5 @@ class FGLTFLightComponentConverter : public FGLTFBuilderContext, public TGLTFCon
 {
 	using FGLTFBuilderContext::FGLTFBuilderContext;
 
-	FGLTFJsonLightIndex Convert(const ULightComponent* LightComponent) override final;
+	virtual FGLTFJsonLightIndex Convert(const ULightComponent* LightComponent) override final;
 };
diff --git a/Engine/Plugins/Enterprise/GLTFExporter/MigratedFromGit/Source/GLTFExporter/Private/Converters/GLTFLightMapConverters.h b/Engine/Plugins/Enterprise/GLTFExporter/MigratedFromGit/Source/GLTFExporter/Private/Converters/GLTFLightMapConverters.h
--- a/Engine/Plugins/Enterprise/GLTFExporter/MigratedFromGit/Source/GLTFExporter/Private/Converters/GLTFLightMapConverters.h
+++ b/Engine/Plugins/Enterprise/GLTFExporter/MigratedFromGit/Source/GLTFExporter/Private/Converters/GLTFLightMapConverters.h
@@ -11,5 +11,5 @@ class FGLTFLightMapConverter : public FGLTFBuilderContext, public TGLTFConverter
 {
 	using FGLTFBuilderContext::FGLTFBuilderContext;
 
-	FGLTFJsonLightMapIndex Convert(const UStaticMeshComponent* StaticMeshComponent) override final;
+	virtual FGLTFJsonLightMapIndex Convert(const UStaticMeshComponent* StaticMeshComponent) override final;
 };
diff --git a/Engine/Plugins/Enterprise/GLTFExporter/MigratedFromGit/Source/GLTFExporter/Private/Converters/GLTFMaterialConverters.h b/Engine/Plugins/Enterprise/GLTFExporter/MigratedFromGit/Source/GLTFExporter/Private/Converters/GLTFMaterialConverters.h
--- a/Engine/Plugins/Enterprise/GLTFExporter/MigratedFromGit/Source/GLTFExporter/Private/Converters/GLTFMaterialConverters.h
+++ b/Engine/Plugins/Enterprise/GLTFExporter/MigratedFromGit/Source/GLTFExporter/Private/Converters/GLTFMaterialConverters.h
@@ -14,7 +14,7 @@ class FGLTFMaterialConverter : public FGLTFBuilderContext, public TGLTFConverter
 {
 	using FGLTFBuilderContext::FGLTFBuilderContext;
 
-	FGLTFJsonMaterialIndex Convert(const UMaterialInterface* Material) override final;
+	virtual FGLTFJsonMaterialIndex Convert(const UMaterialInterface* Material) override final;
 
 	bool TryGetShadingModel(EGLTFJsonShadingModel& ShadingModel, const UMaterialInterface* Material) const;
 
diff --git a/Engine/Plugins/Enterprise/GLTFExporter/MigratedFromGit/Source/GLTFExporter/Private/Converters/GLTFMeshConverters.h b/Engine/Plugins/Enterprise/GLTFExporter/MigratedFromGit/Source/GLTFExporter/Private/Converters/GLTFMeshConverters.h
--- a/Engine/Plugins/Enterprise/GLTFExporter/MigratedFromGit/Source/GLTFExporter/Private/Converters/GLTFMeshConverters.h
+++ b/Engine/Plugins/Enterprise/GLTFExporter/MigratedFromGit/Source/GLTFExporter/Private/Converters/GLTFMeshConverters.h
@@ -8,16 +8,22 @@
 #include "Converters/GLTFMaterialArray.h"
 #include "Engine.h"
 
-class FGLTFStaticMeshConverter : public FGLTFBuilderContext, public TGLTFConverter<FGLTFJsonMeshIndex, const UStaticMesh*, int32, const FColorVertexBuffer*, FGLTFMaterialArray>
+template <typename... InputTypes>
+class FGLTFMeshConverter : public FGLTFBuilderContext, public TGLTFConverter<FGLTFJsonMeshIndex, InputTypes...>
 {
 	using FGLTFBuilderContext::FGLTFBuilderContext;
+};
+
+class FGLTFStaticMeshConverter : public FGLTFMeshConverter<const UStaticMesh*, int32, const FColorVertexBuffer*, FGLTFMaterialArray>
+{
+	using FGLTFMeshConverter::FGLTFMeshConverter;
 
-	FGLTFJsonMeshIndex Convert(const UStaticMesh* StaticMesh, int32 LODIndex, const FColorVertexBuffer* OverrideVertexColors, FGLTFMaterialArray OverrideMaterials) override final;
+	virtual FGLTFJsonMeshIndex Convert(const UStaticMesh* StaticMesh, int32 LODIndex, const FColorVertexBuffer* OverrideVertexColors, FGLTFMaterialArray OverrideMaterials) override final;
 };
 
-class FGLTFSkeletalMeshConverter : public FGLTFBuilderContext, public TGLTFConverter<FGLTFJsonMeshIndex, const USkeletalMesh*, int32, const FColorVertexBuffer*, const FSkinWeightVertexBuffer*, FGLTFMaterialArray>
+class FGLTFSkeletalMeshConverter : public FGLTFMeshConverter<const USkeletalMesh*, int32, const FColorVertexBuffer*, const FSkinWeightVertexBuffer*, FGLTFMaterialArray>
 {
-	using FGLTFBuilderContext::FGLTFBuilderContext;
+	using FGLTFMeshConverter::FGLTFMeshConverter;
 
-	FGLTFJsonMeshIndex Convert(const USkeletalMesh* SkeletalMesh, int32 LODIndex, const FColorVertexBuffer* OverrideVertexColors, const FSkinWeightVertexBuffer* OverrideSkinWeights, FGLTFMaterialArray OverrideMaterials) override final;
+	virtual FGLTFJsonMeshIndex Convert(const USkeletalMesh* SkeletalMesh, int32 LODIndex, const FColorVertexBuffer* OverrideVertexColors, const FSkinWeightVertexBuffer* OverrideSkinWeights, FGLTFMaterialArray OverrideMaterials) override final;
 };
diff --git a/Engine/Plugins/Enterprise/GLTFExporter/MigratedFromGit/Source/GLTFExporter/Private/Converters/GLTFNodeConverters.h b/Engine/Plugins/Enterprise/GLTFExporter/MigratedFromGit/Source/GLTFExporter/Private/Converters/GLTFNodeConverters.h
--- a/Engine/Plugins/Enterprise/GLTFExporter/MigratedFromGit/Source/GLTFExporter/Private/Converters/GLTFNodeConverters.h
+++ b/Engine/Plugins/Enterprise/GLTFExporter/MigratedFromGit/Source/GLTFExporter/Private/Converters/GLTFNodeConverters.h
@@ -7,16 +7,22 @@
 #include "Converters/GLTFBuilderContext.h"
 #include "Engine.h"
 
-class FGLTFSceneComponentConverter : public FGLTFBuilderContext, public TGLTFConverter<FGLTFJsonNodeIndex, const USceneComponent*>
+template <typename... InputTypes>
+class FGLTFNodeConverter : public FGLTFBuilderContext, public TGLTFConverter<FGLTFJsonNodeIndex, InputTypes...>
 {
 	using FGLTFBuilderContext::FGLTFBuilderContext;
+};
+
+class FGLTFSceneComponentConverter : public FGLTFNodeConverter<const USceneComponent*>
+{
+	using FGLTFNodeConverter::FGLTFNodeConverter;
 
-	FGLTFJsonNodeIndex Convert(const USceneComponent* SceneComponent) override final;
+	virtual FGLTFJsonNodeIndex Convert(const USceneComponent* SceneComponent) override final;
 };
 
-class FGLTFActorConverter : public FGLTFBuilderContext, public TGLTFConverter<FGLTFJsonNodeIndex, const AActor*>
+class FGLTFActorConverter : public FGLTFNodeConverter<const AActor*>
 {
-	using FGLTFBuilderContext::FGLTFBuilderContext;
+	using FGLTFNodeConverter::FGLTFNodeConverter;
 
-	FGLTFJsonNodeIndex Convert(const AActor* Actor) override final;
+	virtual FGLTFJsonNodeIndex Convert(const AActor* Actor) override final;
 };
diff --git a/Engine/Plugins/Enterprise/GLTFExporter/MigratedFromGit/Source/GLTFExporter/Private/Converters/GLTFSamplerConverters.h b/Engine/Plugins/Enterprise/GLTFExporter/MigratedFromGit/Source/GLTFExporter/Private/Converters/GLTFSamplerConverters.h
--- a/Engine/Plugins/Enterprise/GLTFExporter/MigratedFromGit/Source/GLTFExporter/Private/Converters/GLTFSamplerConverters.h
+++ b/Engine/Plugins/Enterprise/GLTFExporter/MigratedFromGit/Source/GLTFExporter/Private/Converters/GLTFSamplerConverters.h
@@ -11,5 +11,5 @@ class FGLTFTextureSamplerConverter : public FGLTFBuilderContext, public TGLTFCon
 {
 	using FGLTFBuilderContext::FGLTFBuilderContext;
 
-	FGLTFJsonSamplerIndex Convert(const UTexture* Texture) override final;
+	virtual FGLTFJsonSamplerIndex Convert(const UTexture* Texture) override final;
 };
diff --git a/Engine/Plugins/Enterprise/GLTFExporter/MigratedFromGit/Source/GLTFExporter/Private/Converters/GLTFSceneConverters.h b/Engine/Plugins/Enterprise/GLTFExporter/MigratedFromGit/Source/GLTFExporter/Private/Converters/GLTFSceneConverters.h
--- a/Engine/Plugins/Enterprise/GLTFExporter/MigratedFromGit/Source/GLTFExporter/Private/Converters/GLTFSceneConverters.h
+++ b/Engine/Plugins/Enterprise/GLTFExporter/MigratedFromGit/Source/GLTFExporter/Private/Converters/GLTFSceneConverters.h
@@ -11,5 +11,5 @@ class FGLTFLevelConverter : public FGLTFBuilderContext, public TGLTFConverter<FG
 {
 	using FGLTFBuilderContext::FGLTFBuilderContext;
 
-	FGLTFJsonSceneIndex Convert(const ULevel* Level) override final;
+	virtual FGLTFJsonSceneIndex Convert(const ULevel* Level) override final;
 };
diff --git a/Engine/Plugins/Enterprise/GLTFExporter/MigratedFromGit/Source/GLTFExporter/Private/Converters/GLTFTextureConverters.h b/Engine/Plugins/Enterprise/GLTFExporter/MigratedFromGit/Source/GLTFExporter/Private/Converters/GLTFTextureConverters.h
--- a/Engine/Plugins/Enterprise/GLTFExporter/MigratedFromGit/Source/GLTFExporter/Private/Converters/GLTFTextureConverters.h
+++ b/Engine/Plugins/Enterprise/GLTFExporter/MigratedFromGit/Source/GLTFExporter/Private/Converters/GLTFTextureConverters.h
@@ -7,30 +7,36 @@
 #include "Converters/GLTFBuilderContext.h"
 #include "Engine.h"
 
-class FGLTFTexture2DConverter : public FGLTFBuilderContext, public TGLTFConverter<FGLTFJsonTextureIndex, const UTexture2D*>
+template <typename... InputTypes>
+class FGLTFTextureConverter : public FGLTFBuilderContext, public TGLTFConverter<FGLTFJsonTextureIndex, InputTypes...>
 {
 	using FGLTFBuilderContext::FGLTFBuilderContext;
+};
+
+class FGLTFTexture2DConverter : public FGLTFTextureConverter<const UTexture2D*>
+{
+	using FGLTFTextureConverter::FGLTFTextureConverter;
 
-	FGLTFJsonTextureIndex Convert(const UTexture2D* Texture2D) override final;
+	virtual FGLTFJsonTextureIndex Convert(const UTexture2D* Texture2D) override final;
 };
 
-class FGLTFTextureCubeConverter : public FGLTFBuilderContext, public TGLTFConverter<FGLTFJsonTextureIndex, const UTextureCube*, ECubeFace>
+class FGLTFTextureCubeConverter : public FGLTFTextureConverter<const UTextureCube*, ECubeFace>
 {
-	using FGLTFBuilderContext::FGLTFBuilderContext;
+	using FGLTFTextureConverter::FGLTFTextureConverter;
 
-	FGLTFJsonTextureIndex Convert(const UTextureCube* TextureCube, ECubeFace CubeFace) override final;
+	virtual FGLTFJsonTextureIndex Convert(const UTextureCube* TextureCube, ECubeFace CubeFace) override final;
 };
 
-class FGLTFTextureRenderTarget2DConverter : public FGLTFBuilderContext, public TGLTFConverter<FGLTFJsonTextureIndex, const UTextureRenderTarget2D*>
+class FGLTFTextureRenderTarget2DConverter : public FGLTFTextureConverter<const UTextureRenderTarget2D*>
 {
-	using FGLTFBuilderContext::FGLTFBuilderContext;
+	using FGLTFTextureConverter::FGLTFTextureConverter;
 
-	FGLTFJsonTextureIndex Convert(const UTextureRenderTarget2D* RenderTarget2D) override final;
+	virtual FGLTFJsonTextureIndex Convert(const UTextureRenderTarget2D* RenderTarget2D) override final;
 };
 
-class FGLTFTextureRenderTargetCubeConverter : public FGLTFBuilderContext, public TGLTFConverter<FGLTFJsonTextureIndex, const UTextureRenderTargetCube*, ECubeFace>
+class FGLTFTextureRenderTargetCubeConverter : public FGLTFTextureConverter<const UTextureRenderTargetCube*, ECubeFace>
 {
-	using FGLTFBuilderContext::FGLTFBuilderContext;
+	using FGLTFTextureConverter::FGLTFTextureConverter;
 
-	FGLTFJsonTextureIndex Convert(const UTextureRenderTargetCube* RenderTargetCube, ECubeFace CubeFace) override final;
+	virtual FGLTFJsonTextureIndex Convert(const UTextureRenderTargetCube* RenderTargetCube, ECubeFace CubeFace) override final;
 };
diff --git a/Engine/Plugins/Enterprise/GLTFExporter/MigratedFromGit/Source/GLTFExporter/Private/Converters/GLTFVarationConverters.h b/Engine/Plugins/Enterprise/GLTFExporter/MigratedFromGit/Source/GLTFExporter/Private/Converters/GLTFVarationConverters.h
--- a/Engine/Plugins/Enterprise/GLTFExporter/MigratedFromGit/Source/GLTFExporter/Private/Converters/GLTFVarationConverters.h
+++ b/Engine/Plugins/Enterprise/GLTFExporter/MigratedFromGit/Source/GLTFExporter/Private/Converters/GLTFVarationConverters.h
@@ -16,7 +16,7 @@ class FGLTFVariationConverter : public FGLTFBuilderContext, public TGLTFConverte
 {
 	using FGLTFBuilderContext::FGLTFBuilderContext;
 
-	FGLTFJsonVariationIndex Convert(const ALevelVariantSetsActor* LevelVariantSetsActor) override final;
+	virtual FGLTFJsonVariationIndex Convert(const ALevelVariantSetsActor* LevelVariantSetsActor) override final;
 
 	bool TryParseVariant(FGLTFJsonVariant& OutVariant, const UVariant* Variant) const;
 	bool TryParseVariantBinding(FGLTFJsonVariant& OutVariant, const UVariantObjectBinding* Binding) const;
