[GLTFExporter] refactor: rename some converters and set final specifier for each entire class

commit 72f26322e9184cea5baa69d726866c3492a736b8
Date: Mon, 21 Sep 2020 01:13:56 +0200

    refactor: rename some converters and set final specifier for each entire class
    (cherry picked from commit b7c141377fe852cd40778c1f023f36080320786e​)

#skipci

diff --git a/Engine/Plugins/Enterprise/GLTFExporter/MigratedFromGit/Source/GLTFExporter/Private/Builders/GLTFConvertBuilder.cpp b/Engine/Plugins/Enterprise/GLTFExporter/MigratedFromGit/Source/GLTFExporter/Private/Builders/GLTFConvertBuilder.cpp
--- a/Engine/Plugins/Enterprise/GLTFExporter/MigratedFromGit/Source/GLTFExporter/Private/Builders/GLTFConvertBuilder.cpp
+++ b/Engine/Plugins/Enterprise/GLTFExporter/MigratedFromGit/Source/GLTFExporter/Private/Builders/GLTFConvertBuilder.cpp
@@ -1,7 +1,6 @@
 // Copyright Epic Games, Inc. All Rights Reserved.
 
 #include "Builders/GLTFConvertBuilder.h"
-#include "Builders/GLTFBuilderUtility.h"
 
 FGLTFConvertBuilder::FGLTFConvertBuilder(const UGLTFExportOptions* ExportOptions, bool bSelectedActorsOnly)
 	: FGLTFImageBuilder(ExportOptions)
@@ -16,7 +15,7 @@ FGLTFJsonAccessorIndex FGLTFConvertBuilder::GetOrAddPositionAccessor(const FPosi
 		return FGLTFJsonAccessorIndex(INDEX_NONE);
 	}
 
-	return PositionVertexBufferConverter.GetOrAdd(VertexBuffer);
+	return PositionBufferConverter.GetOrAdd(VertexBuffer);
 }
 
 FGLTFJsonAccessorIndex FGLTFConvertBuilder::GetOrAddColorAccessor(const FColorVertexBuffer* VertexBuffer)
@@ -26,7 +25,7 @@ FGLTFJsonAccessorIndex FGLTFConvertBuilder::GetOrAddColorAccessor(const FColorVe
 		return FGLTFJsonAccessorIndex(INDEX_NONE);
 	}
 
-	return ColorVertexBufferConverter.GetOrAdd(VertexBuffer);
+	return ColorBufferConverter.GetOrAdd(VertexBuffer);
 }
 
 FGLTFJsonAccessorIndex FGLTFConvertBuilder::GetOrAddNormalAccessor(const FStaticMeshVertexBuffer* VertexBuffer)
@@ -36,7 +35,7 @@ FGLTFJsonAccessorIndex FGLTFConvertBuilder::GetOrAddNormalAccessor(const FStatic
 		return FGLTFJsonAccessorIndex(INDEX_NONE);
 	}
 
-	return NormalVertexBufferConverter.GetOrAdd(VertexBuffer);
+	return NormalBufferConverter.GetOrAdd(VertexBuffer);
 }
 
 FGLTFJsonAccessorIndex FGLTFConvertBuilder::GetOrAddTangentAccessor(const FStaticMeshVertexBuffer* VertexBuffer)
@@ -46,7 +45,7 @@ FGLTFJsonAccessorIndex FGLTFConvertBuilder::GetOrAddTangentAccessor(const FStati
 		return FGLTFJsonAccessorIndex(INDEX_NONE);
 	}
 
-	return TangentVertexBufferConverter.GetOrAdd(VertexBuffer);
+	return TangentBufferConverter.GetOrAdd(VertexBuffer);
 }
 
 FGLTFJsonAccessorIndex FGLTFConvertBuilder::GetOrAddUVAccessor(const FStaticMeshVertexBuffer* VertexBuffer, int32 UVIndex)
@@ -56,7 +55,7 @@ FGLTFJsonAccessorIndex FGLTFConvertBuilder::GetOrAddUVAccessor(const FStaticMesh
 		return FGLTFJsonAccessorIndex(INDEX_NONE);
 	}
 
-	return UVVertexBufferConverter.GetOrAdd(VertexBuffer, UVIndex);
+	return UVBufferConverter.GetOrAdd(VertexBuffer, UVIndex);
 }
 
 FGLTFJsonAccessorIndex FGLTFConvertBuilder::GetOrAddJointAccessor(const FSkinWeightVertexBuffer* VertexBuffer, int32 JointsGroupIndex, FGLTFBoneMap BoneMap)
@@ -66,7 +65,7 @@ FGLTFJsonAccessorIndex FGLTFConvertBuilder::GetOrAddJointAccessor(const FSkinWei
 		return FGLTFJsonAccessorIndex(INDEX_NONE);
 	}
 
-	return BoneIndexVertexBufferConverter.GetOrAdd(VertexBuffer, JointsGroupIndex, BoneMap);
+	return BoneIndexBufferConverter.GetOrAdd(VertexBuffer, JointsGroupIndex, BoneMap);
 }
 
 FGLTFJsonAccessorIndex FGLTFConvertBuilder::GetOrAddWeightAccessor(const FSkinWeightVertexBuffer* VertexBuffer, int32 WeightsGroupIndex)
@@ -76,7 +75,7 @@ FGLTFJsonAccessorIndex FGLTFConvertBuilder::GetOrAddWeightAccessor(const FSkinWe
 		return FGLTFJsonAccessorIndex(INDEX_NONE);
 	}
 
-	return BoneWeightVertexBufferConverter.GetOrAdd(VertexBuffer, WeightsGroupIndex);
+	return BoneWeightBufferConverter.GetOrAdd(VertexBuffer, WeightsGroupIndex);
 }
 
 FGLTFJsonAccessorIndex FGLTFConvertBuilder::GetOrAddIndexAccessor(const FStaticMeshSection* MeshSection, const FRawStaticIndexBuffer* IndexBuffer)
@@ -191,7 +190,7 @@ FGLTFJsonSamplerIndex FGLTFConvertBuilder::GetOrAddSampler(const UTexture* Textu
 		return FGLTFJsonSamplerIndex(INDEX_NONE);
 	}
 
-	return TextureSamplerConverter.GetOrAdd(Texture);
+	return SamplerConverter.GetOrAdd(Texture);
 }
 
 FGLTFJsonTextureIndex FGLTFConvertBuilder::GetOrAddTexture(const UTexture2D* Texture)
@@ -241,7 +240,7 @@ FGLTFJsonNodeIndex FGLTFConvertBuilder::GetOrAddNode(const USceneComponent* Scen
 		return FGLTFJsonNodeIndex(INDEX_NONE);
 	}
 
-	return SceneComponentConverter.GetOrAdd(SceneComponent);
+	return ComponentConverter.GetOrAdd(SceneComponent);
 }
 
 FGLTFJsonNodeIndex FGLTFConvertBuilder::GetOrAddNode(const AActor* Actor)
@@ -261,7 +260,7 @@ FGLTFJsonSceneIndex FGLTFConvertBuilder::GetOrAddScene(const ULevel* Level)
 		return FGLTFJsonSceneIndex(INDEX_NONE);
 	}
 
-	return LevelConverter.GetOrAdd(Level);
+	return SceneConverter.GetOrAdd(Level);
 }
 
 FGLTFJsonSceneIndex FGLTFConvertBuilder::GetOrAddScene(const UWorld* World)
@@ -281,7 +280,7 @@ FGLTFJsonCameraIndex FGLTFConvertBuilder::GetOrAddCamera(const UCameraComponent*
 		return FGLTFJsonCameraIndex(INDEX_NONE);
 	}
 
-	return CameraComponentConverter.GetOrAdd(CameraComponent);
+	return CameraConverter.GetOrAdd(CameraComponent);
 }
 
 FGLTFJsonLightIndex FGLTFConvertBuilder::GetOrAddLight(const ULightComponent* LightComponent)
@@ -291,17 +290,17 @@ FGLTFJsonLightIndex FGLTFConvertBuilder::GetOrAddLight(const ULightComponent* Li
 		return FGLTFJsonLightIndex(INDEX_NONE);
 	}
 
-	return LightComponentConverter.GetOrAdd(LightComponent);
+	return LightConverter.GetOrAdd(LightComponent);
 }
 
-FGLTFJsonBackdropIndex FGLTFConvertBuilder::GetOrAddBackdrop(const AActor* Actor)
+FGLTFJsonBackdropIndex FGLTFConvertBuilder::GetOrAddBackdrop(const AActor* BackdropActor)
 {
-	if (Actor == nullptr)
+	if (BackdropActor == nullptr)
 	{
 		return FGLTFJsonBackdropIndex(INDEX_NONE);
 	}
 
-	return BackdropConverter.GetOrAdd(Actor);
+	return BackdropConverter.GetOrAdd(BackdropActor);
 }
 
 FGLTFJsonVariationIndex FGLTFConvertBuilder::GetOrAddVariation(const ALevelVariantSetsActor* LevelVariantSetsActor)
@@ -331,5 +330,5 @@ FGLTFJsonHotspotIndex FGLTFConvertBuilder::GetOrAddHotspot(const UGLTFInteractio
 		return FGLTFJsonHotspotIndex(INDEX_NONE);
 	}
 
-	return HotspotComponentConverter.GetOrAdd(HotspotComponent);
+	return HotspotConverter.GetOrAdd(HotspotComponent);
 }
diff --git a/Engine/Plugins/Enterprise/GLTFExporter/MigratedFromGit/Source/GLTFExporter/Private/Builders/GLTFConvertBuilder.h b/Engine/Plugins/Enterprise/GLTFExporter/MigratedFromGit/Source/GLTFExporter/Private/Builders/GLTFConvertBuilder.h
--- a/Engine/Plugins/Enterprise/GLTFExporter/MigratedFromGit/Source/GLTFExporter/Private/Builders/GLTFConvertBuilder.h
+++ b/Engine/Plugins/Enterprise/GLTFExporter/MigratedFromGit/Source/GLTFExporter/Private/Builders/GLTFConvertBuilder.h
@@ -32,12 +32,11 @@ public:
 	FGLTFJsonAccessorIndex GetOrAddUVAccessor(const FStaticMeshVertexBuffer* VertexBuffer, int32 UVIndex);
 	FGLTFJsonAccessorIndex GetOrAddJointAccessor(const FSkinWeightVertexBuffer* VertexBuffer, int32 JointsGroupIndex, FGLTFBoneMap BoneMap);
 	FGLTFJsonAccessorIndex GetOrAddWeightAccessor(const FSkinWeightVertexBuffer* VertexBuffer, int32 WeightsGroupIndex);
-
 	FGLTFJsonAccessorIndex GetOrAddIndexAccessor(const FStaticMeshSection* MeshSection, const FRawStaticIndexBuffer* IndexBuffer);
+	FGLTFJsonAccessorIndex GetOrAddIndexAccessor(const FSkelMeshRenderSection* MeshSection, const FMultiSizeIndexContainer* IndexContainer);
+
 	FGLTFJsonMeshIndex GetOrAddMesh(const UStaticMesh* StaticMesh, int32 LODIndex = -1, const FColorVertexBuffer* OverrideVertexColors = nullptr, const FGLTFMaterialArray& OverrideMaterials = {});
 	FGLTFJsonMeshIndex GetOrAddMesh(const UStaticMeshComponent* StaticMeshComponent);
-
-	FGLTFJsonAccessorIndex GetOrAddIndexAccessor(const FSkelMeshRenderSection* MeshSection, const FMultiSizeIndexContainer* IndexContainer);
 	FGLTFJsonMeshIndex GetOrAddMesh(const USkeletalMesh* SkeletalMesh, int32 LODIndex = -1, const FColorVertexBuffer* OverrideVertexColors = nullptr, const FSkinWeightVertexBuffer* OverrideSkinWeights = nullptr, const FGLTFMaterialArray& OverrideMaterials = {});
 	FGLTFJsonMeshIndex GetOrAddMesh(const USkeletalMeshComponent* SkeletalMeshComponent);
 
@@ -55,44 +54,41 @@ public:
 
 	FGLTFJsonCameraIndex GetOrAddCamera(const UCameraComponent* CameraComponent);
 	FGLTFJsonLightIndex GetOrAddLight(const ULightComponent* LightComponent);
-
-	FGLTFJsonBackdropIndex GetOrAddBackdrop(const AActor* Actor);
+	FGLTFJsonBackdropIndex GetOrAddBackdrop(const AActor* BackdropActor);
 	FGLTFJsonVariationIndex GetOrAddVariation(const ALevelVariantSetsActor* LevelVariantSetsActor);
 	FGLTFJsonLightMapIndex GetOrAddLightMap(const UStaticMeshComponent* StaticMeshComponent);
 	FGLTFJsonHotspotIndex GetOrAddHotspot(const UGLTFInteractionHotspotComponent* HotspotComponent);
 
 private:
 
-	FGLTFPositionVertexBufferConverter PositionVertexBufferConverter = *this;
-	FGLTFColorVertexBufferConverter ColorVertexBufferConverter = *this;
-	FGLTFNormalVertexBufferConverter NormalVertexBufferConverter = *this;
-	FGLTFTangentVertexBufferConverter TangentVertexBufferConverter = *this;
-	FGLTFUVVertexBufferConverter UVVertexBufferConverter = *this;
-	FGLTFBoneIndexVertexBufferConverter BoneIndexVertexBufferConverter = *this;
-	FGLTFBoneWeightVertexBufferConverter BoneWeightVertexBufferConverter = *this;
-
-	FGLTFStaticMeshSectionConverter StaticMeshSectionConverter = *this;
-	FGLTFStaticMeshConverter StaticMeshConverter = *this;
-
-	FGLTFSkeletalMeshSectionConverter SkeletalMeshSectionConverter = *this;
-	FGLTFSkeletalMeshConverter SkeletalMeshConverter = *this;
-
-	FGLTFMaterialConverter MaterialConverter = *this;
-	FGLTFTextureSamplerConverter TextureSamplerConverter = *this;
-	FGLTFTexture2DConverter Texture2DConverter = *this;
-	FGLTFTextureCubeConverter TextureCubeConverter = *this;
-	FGLTFTextureRenderTarget2DConverter TextureRenderTarget2DConverter = *this;
-	FGLTFTextureRenderTargetCubeConverter TextureRenderTargetCubeConverter = *this;
-
-	FGLTFSceneComponentConverter SceneComponentConverter = *this;
-	FGLTFActorConverter ActorConverter = *this;
-	FGLTFLevelConverter LevelConverter = *this;
-
-	FGLTFCameraComponentConverter CameraComponentConverter = *this;
-	FGLTFLightComponentConverter LightComponentConverter = *this;
-
-	FGLTFBackdropConverter BackdropConverter = *this;
-	FGLTFVariationConverter VariationConverter = *this;
-	FGLTFLightMapConverter LightMapConverter = *this;
-	FGLTFHotspotComponentConverter HotspotComponentConverter = *this;
+	FGLTFPositionBufferConverter PositionBufferConverter{ *this };
+	FGLTFColorBufferConverter ColorBufferConverter{ *this };
+	FGLTFNormalBufferConverter NormalBufferConverter{ *this };
+	FGLTFTangentBufferConverter TangentBufferConverter{ *this };
+	FGLTFUVBufferConverter UVBufferConverter{ *this };
+	FGLTFBoneIndexBufferConverter BoneIndexBufferConverter{ *this };
+	FGLTFBoneWeightBufferConverter BoneWeightBufferConverter{ *this };
+	FGLTFStaticMeshSectionConverter StaticMeshSectionConverter{ *this };
+	FGLTFSkeletalMeshSectionConverter SkeletalMeshSectionConverter{ *this };
+
+	FGLTFStaticMeshConverter StaticMeshConverter{ *this };
+	FGLTFSkeletalMeshConverter SkeletalMeshConverter{ *this };
+
+	FGLTFMaterialConverter MaterialConverter{ *this };
+	FGLTFSamplerConverter SamplerConverter{ *this };
+	FGLTFTexture2DConverter Texture2DConverter{ *this };
+	FGLTFTextureCubeConverter TextureCubeConverter{ *this };
+	FGLTFTextureRenderTarget2DConverter TextureRenderTarget2DConverter{ *this };
+	FGLTFTextureRenderTargetCubeConverter TextureRenderTargetCubeConverter{ *this };
+
+	FGLTFComponentConverter ComponentConverter{ *this };
+	FGLTFActorConverter ActorConverter{ *this };
+	FGLTFSceneConverter SceneConverter{ *this };
+
+	FGLTFCameraConverter CameraConverter{ *this };
+	FGLTFLightConverter LightConverter{ *this };
+	FGLTFBackdropConverter BackdropConverter{ *this };
+	FGLTFVariationConverter VariationConverter{ *this };
+	FGLTFLightMapConverter LightMapConverter{ *this };
+	FGLTFHotspotConverter HotspotConverter{ *this };
 };
diff --git a/Engine/Plugins/Enterprise/GLTFExporter/MigratedFromGit/Source/GLTFExporter/Private/Converters/GLTFAccessorConverters.cpp b/Engine/Plugins/Enterprise/GLTFExporter/MigratedFromGit/Source/GLTFExporter/Private/Converters/GLTFAccessorConverters.cpp
--- a/Engine/Plugins/Enterprise/GLTFExporter/MigratedFromGit/Source/GLTFExporter/Private/Converters/GLTFAccessorConverters.cpp
+++ b/Engine/Plugins/Enterprise/GLTFExporter/MigratedFromGit/Source/GLTFExporter/Private/Converters/GLTFAccessorConverters.cpp
@@ -8,7 +8,7 @@
 #include "Rendering/MultiSizeIndexContainer.h"
 #include "Rendering/SkeletalMeshRenderData.h"
 
-FGLTFJsonAccessorIndex FGLTFPositionVertexBufferConverter::Convert(const FPositionVertexBuffer* VertexBuffer)
+FGLTFJsonAccessorIndex FGLTFPositionBufferConverter::Convert(const FPositionVertexBuffer* VertexBuffer)
 {
 	const uint32 VertexCount = VertexBuffer->GetNumVertices();
 	if (VertexCount == 0)
@@ -56,7 +56,7 @@ FGLTFJsonAccessorIndex FGLTFPositionVertexBufferConverter::Convert(const FPositi
 	return Builder.AddAccessor(JsonAccessor);
 }
 
-FGLTFJsonAccessorIndex FGLTFColorVertexBufferConverter::Convert(const FColorVertexBuffer* VertexBuffer)
+FGLTFJsonAccessorIndex FGLTFColorBufferConverter::Convert(const FColorVertexBuffer* VertexBuffer)
 {
 	const uint32 VertexCount = VertexBuffer->GetNumVertices();
 	if (VertexCount == 0)
@@ -82,7 +82,7 @@ FGLTFJsonAccessorIndex FGLTFColorVertexBufferConverter::Convert(const FColorVert
 	return Builder.AddAccessor(JsonAccessor);
 }
 
-FGLTFJsonAccessorIndex FGLTFNormalVertexBufferConverter::Convert(const FStaticMeshVertexBuffer* VertexBuffer)
+FGLTFJsonAccessorIndex FGLTFNormalBufferConverter::Convert(const FStaticMeshVertexBuffer* VertexBuffer)
 {
 	const uint32 VertexCount = VertexBuffer->GetNumVertices();
 	if (VertexCount == 0)
@@ -107,7 +107,7 @@ FGLTFJsonAccessorIndex FGLTFNormalVertexBufferConverter::Convert(const FStaticMe
 	return Builder.AddAccessor(JsonAccessor);
 }
 
-FGLTFJsonAccessorIndex FGLTFTangentVertexBufferConverter::Convert(const FStaticMeshVertexBuffer* VertexBuffer)
+FGLTFJsonAccessorIndex FGLTFTangentBufferConverter::Convert(const FStaticMeshVertexBuffer* VertexBuffer)
 {
 	const uint32 VertexCount = VertexBuffer->GetNumVertices();
 	if (VertexCount == 0)
@@ -132,7 +132,7 @@ FGLTFJsonAccessorIndex FGLTFTangentVertexBufferConverter::Convert(const FStaticM
 	return Builder.AddAccessor(JsonAccessor);
 }
 
-FGLTFJsonAccessorIndex FGLTFUVVertexBufferConverter::Convert(const FStaticMeshVertexBuffer* VertexBuffer, int32 UVIndex)
+FGLTFJsonAccessorIndex FGLTFUVBufferConverter::Convert(const FStaticMeshVertexBuffer* VertexBuffer, int32 UVIndex)
 {
 	const uint32 VertexCount = VertexBuffer->GetNumVertices();
 	if (VertexCount == 0 || UVIndex < 0 || VertexBuffer->GetNumTexCoords() <= static_cast<uint32>(UVIndex))
@@ -157,13 +157,13 @@ FGLTFJsonAccessorIndex FGLTFUVVertexBufferConverter::Convert(const FStaticMeshVe
 	return Builder.AddAccessor(JsonAccessor);
 }
 
-FGLTFJsonAccessorIndex FGLTFBoneIndexVertexBufferConverter::Convert(const FSkinWeightVertexBuffer* VertexBuffer, int32 InfluenceOffset, FGLTFBoneMap BoneMap)
+FGLTFJsonAccessorIndex FGLTFBoneIndexBufferConverter::Convert(const FSkinWeightVertexBuffer* VertexBuffer, int32 InfluenceOffset, FGLTFBoneMap BoneMap)
 {
 	return VertexBuffer->Use16BitBoneIndex() ? Convert<uint16>(VertexBuffer, InfluenceOffset, BoneMap) : Convert<uint8>(VertexBuffer, InfluenceOffset, BoneMap);
 }
 
 template <typename IndexType>
-FGLTFJsonAccessorIndex FGLTFBoneIndexVertexBufferConverter::Convert(const FSkinWeightVertexBuffer* VertexBuffer, int32 JointsGroupIndex, FGLTFBoneMap BoneMap)
+FGLTFJsonAccessorIndex FGLTFBoneIndexBufferConverter::Convert(const FSkinWeightVertexBuffer* VertexBuffer, int32 JointsGroupIndex, FGLTFBoneMap BoneMap)
 {
 	const uint32 VertexCount = VertexBuffer->GetNumVertices();
 	const int32 MaxInfluenceCount = VertexBuffer->GetMaxBoneInfluences();
@@ -204,7 +204,7 @@ FGLTFJsonAccessorIndex FGLTFBoneIndexVertexBufferConverter::Convert(const FSkinW
 	return Builder.AddAccessor(JsonAccessor);
 }
 
-FGLTFJsonAccessorIndex FGLTFBoneWeightVertexBufferConverter::Convert(const FSkinWeightVertexBuffer* VertexBuffer, int32 WeightsGroupIndex)
+FGLTFJsonAccessorIndex FGLTFBoneWeightBufferConverter::Convert(const FSkinWeightVertexBuffer* VertexBuffer, int32 WeightsGroupIndex)
 {
 	const uint32 VertexCount = VertexBuffer->GetNumVertices();
 	const int32 MaxInfluenceCount = VertexBuffer->GetMaxBoneInfluences();
diff --git a/Engine/Plugins/Enterprise/GLTFExporter/MigratedFromGit/Source/GLTFExporter/Private/Converters/GLTFAccessorConverters.h b/Engine/Plugins/Enterprise/GLTFExporter/MigratedFromGit/Source/GLTFExporter/Private/Converters/GLTFAccessorConverters.h
--- a/Engine/Plugins/Enterprise/GLTFExporter/MigratedFromGit/Source/GLTFExporter/Private/Converters/GLTFAccessorConverters.h
+++ b/Engine/Plugins/Enterprise/GLTFExporter/MigratedFromGit/Source/GLTFExporter/Private/Converters/GLTFAccessorConverters.h
@@ -16,68 +16,68 @@ class FGLTFAccessorConverter : public FGLTFBuilderContext, public TGLTFConverter
 	using FGLTFBuilderContext::FGLTFBuilderContext;
 };
 
-class FGLTFPositionVertexBufferConverter : public FGLTFAccessorConverter<const FPositionVertexBuffer*>
+class FGLTFPositionBufferConverter final : public FGLTFAccessorConverter<const FPositionVertexBuffer*>
 {
 	using FGLTFAccessorConverter::FGLTFAccessorConverter;
 
-	virtual FGLTFJsonAccessorIndex Convert(const FPositionVertexBuffer* VertexBuffer) override final;
+	virtual FGLTFJsonAccessorIndex Convert(const FPositionVertexBuffer* VertexBuffer) override;
 };
 
-class FGLTFColorVertexBufferConverter : public FGLTFAccessorConverter<const FColorVertexBuffer*>
+class FGLTFColorBufferConverter final : public FGLTFAccessorConverter<const FColorVertexBuffer*>
 {
 	using FGLTFAccessorConverter::FGLTFAccessorConverter;
 
-	virtual FGLTFJsonAccessorIndex Convert(const FColorVertexBuffer* VertexBuffer) override final;
+	virtual FGLTFJsonAccessorIndex Convert(const FColorVertexBuffer* VertexBuffer) override;
 };
 
-class FGLTFNormalVertexBufferConverter : public FGLTFAccessorConverter<const FStaticMeshVertexBuffer*>
+class FGLTFNormalBufferConverter final : public FGLTFAccessorConverter<const FStaticMeshVertexBuffer*>
 {
 	using FGLTFAccessorConverter::FGLTFAccessorConverter;
 
-	virtual FGLTFJsonAccessorIndex Convert(const FStaticMeshVertexBuffer* VertexBuffer) override final;
+	virtual FGLTFJsonAccessorIndex Convert(const FStaticMeshVertexBuffer* VertexBuffer) override;
 };
 
-class FGLTFTangentVertexBufferConverter : public FGLTFAccessorConverter<const FStaticMeshVertexBuffer*>
+class FGLTFTangentBufferConverter final : public FGLTFAccessorConverter<const FStaticMeshVertexBuffer*>
 {
 	using FGLTFAccessorConverter::FGLTFAccessorConverter;
 
-	virtual FGLTFJsonAccessorIndex Convert(const FStaticMeshVertexBuffer* VertexBuffer) override final;
+	virtual FGLTFJsonAccessorIndex Convert(const FStaticMeshVertexBuffer* VertexBuffer) override;
 };
 
-class FGLTFUVVertexBufferConverter : public FGLTFAccessorConverter<const FStaticMeshVertexBuffer*, int32>
+class FGLTFUVBufferConverter final : public FGLTFAccessorConverter<const FStaticMeshVertexBuffer*, int32>
 {
 	using FGLTFAccessorConverter::FGLTFAccessorConverter;
 
-	virtual FGLTFJsonAccessorIndex Convert(const FStaticMeshVertexBuffer* VertexBuffer, int32 UVIndex) override final;
+	virtual FGLTFJsonAccessorIndex Convert(const FStaticMeshVertexBuffer* VertexBuffer, int32 UVIndex) override;
 };
 
-class FGLTFBoneIndexVertexBufferConverter : public FGLTFAccessorConverter<const FSkinWeightVertexBuffer*, int32, FGLTFBoneMap>
+class FGLTFBoneIndexBufferConverter final : public FGLTFAccessorConverter<const FSkinWeightVertexBuffer*, int32, FGLTFBoneMap>
 {
 	using FGLTFAccessorConverter::FGLTFAccessorConverter;
 
-	virtual FGLTFJsonAccessorIndex Convert(const FSkinWeightVertexBuffer* VertexBuffer, int32 InfluenceOffset, FGLTFBoneMap BoneMap) override final;
+	virtual FGLTFJsonAccessorIndex Convert(const FSkinWeightVertexBuffer* VertexBuffer, int32 InfluenceOffset, FGLTFBoneMap BoneMap) override;
 
 	template <typename IndexType>
 	FGLTFJsonAccessorIndex Convert(const FSkinWeightVertexBuffer* VertexBuffer, int32 JointsGroupIndex, FGLTFBoneMap BoneMap);
 };
 
-class FGLTFBoneWeightVertexBufferConverter : public FGLTFAccessorConverter<const FSkinWeightVertexBuffer*, int32>
+class FGLTFBoneWeightBufferConverter final : public FGLTFAccessorConverter<const FSkinWeightVertexBuffer*, int32>
 {
 	using FGLTFAccessorConverter::FGLTFAccessorConverter;
 
-	virtual FGLTFJsonAccessorIndex Convert(const FSkinWeightVertexBuffer* VertexBuffer, int32 WeightsGroupIndex) override final;
+	virtual FGLTFJsonAccessorIndex Convert(const FSkinWeightVertexBuffer* VertexBuffer, int32 WeightsGroupIndex) override;
 };
 
-class FGLTFStaticMeshSectionConverter : public FGLTFAccessorConverter<const FStaticMeshSection*, const FRawStaticIndexBuffer*>
+class FGLTFStaticMeshSectionConverter final : public FGLTFAccessorConverter<const FStaticMeshSection*, const FRawStaticIndexBuffer*>
 {
 	using FGLTFAccessorConverter::FGLTFAccessorConverter;
 
-	virtual FGLTFJsonAccessorIndex Convert(const FStaticMeshSection* MeshSection, const FRawStaticIndexBuffer* IndexBuffer) override final;
+	virtual FGLTFJsonAccessorIndex Convert(const FStaticMeshSection* MeshSection, const FRawStaticIndexBuffer* IndexBuffer) override;
 };
 
-class FGLTFSkeletalMeshSectionConverter : public FGLTFAccessorConverter<const FSkelMeshRenderSection*, const FMultiSizeIndexContainer*>
+class FGLTFSkeletalMeshSectionConverter final : public FGLTFAccessorConverter<const FSkelMeshRenderSection*, const FMultiSizeIndexContainer*>
 {
 	using FGLTFAccessorConverter::FGLTFAccessorConverter;
 
-	virtual FGLTFJsonAccessorIndex Convert(const FSkelMeshRenderSection* MeshSection, const FMultiSizeIndexContainer* IndexContainer) override final;
+	virtual FGLTFJsonAccessorIndex Convert(const FSkelMeshRenderSection* MeshSection, const FMultiSizeIndexContainer* IndexContainer) override;
 };
diff --git a/Engine/Plugins/Enterprise/GLTFExporter/MigratedFromGit/Source/GLTFExporter/Private/Converters/GLTFBackdropConverters.cpp b/Engine/Plugins/Enterprise/GLTFExporter/MigratedFromGit/Source/GLTFExporter/Private/Converters/GLTFBackdropConverters.cpp
--- a/Engine/Plugins/Enterprise/GLTFExporter/MigratedFromGit/Source/GLTFExporter/Private/Converters/GLTFBackdropConverters.cpp
+++ b/Engine/Plugins/Enterprise/GLTFExporter/MigratedFromGit/Source/GLTFExporter/Private/Converters/GLTFBackdropConverters.cpp
@@ -5,30 +5,30 @@
 #include "Converters/GLTFActorUtility.h"
 #include "Builders/GLTFContainerBuilder.h"
 
-FGLTFJsonBackdropIndex FGLTFBackdropConverter::Convert(const AActor* Actor)
+FGLTFJsonBackdropIndex FGLTFBackdropConverter::Convert(const AActor* BackdropActor)
 {
-	const UBlueprint* Blueprint = FGLTFActorUtility::GetBlueprintFromActor(Actor);
+	const UBlueprint* Blueprint = FGLTFActorUtility::GetBlueprintFromActor(BackdropActor);
 	if (!FGLTFActorUtility::IsHDRIBackdropBlueprint(Blueprint))
 	{
 		return FGLTFJsonBackdropIndex(INDEX_NONE);
 	}
 
 	FGLTFJsonBackdrop JsonBackdrop;
-	Actor->GetName(JsonBackdrop.Name);
+	BackdropActor->GetName(JsonBackdrop.Name);
 
 	const UStaticMesh* Mesh;
-	if (FGLTFActorUtility::TryGetPropertyValue(Actor, TEXT("Mesh"), Mesh))
+	if (FGLTFActorUtility::TryGetPropertyValue(BackdropActor, TEXT("Mesh"), Mesh))
 	{
 		// TODO: avoid exporting the mesh material (use gltf default material)
 		JsonBackdrop.Mesh = Builder.GetOrAddMesh(Mesh);
 	}
 	else
 	{
-		Builder.AddWarningMessage(FString::Printf(TEXT("Failed to export Mesh for HDRIBackdrop %s"), *Actor->GetName()));
+		Builder.AddWarningMessage(FString::Printf(TEXT("Failed to export Mesh for HDRIBackdrop %s"), *BackdropActor->GetName()));
 	}
 
 	const UTextureCube* Cubemap;
-	if (FGLTFActorUtility::TryGetPropertyValue(Actor, TEXT("Cubemap"), Cubemap))
+	if (FGLTFActorUtility::TryGetPropertyValue(BackdropActor, TEXT("Cubemap"), Cubemap))
 	{
 		// TODO: add a proper custom gltf extension with its own converters for cubemaps
 
@@ -41,57 +41,57 @@ FGLTFJsonBackdropIndex FGLTFBackdropConverter::Convert(const AActor* Actor)
 	}
 	else
 	{
-		Builder.AddWarningMessage(FString::Printf(TEXT("Failed to export Cubemap for HDRIBackdrop %s"), *Actor->GetName()));
+		Builder.AddWarningMessage(FString::Printf(TEXT("Failed to export Cubemap for HDRIBackdrop %s"), *BackdropActor->GetName()));
 	}
 
 	float Intensity;
-	if (FGLTFActorUtility::TryGetPropertyValue(Actor, TEXT("Intensity"), Intensity))
+	if (FGLTFActorUtility::TryGetPropertyValue(BackdropActor, TEXT("Intensity"), Intensity))
 	{
 		JsonBackdrop.Intensity = Intensity;
 	}
 	else
 	{
-		Builder.AddWarningMessage(FString::Printf(TEXT("Failed to export Intensity for HDRIBackdrop %s"), *Actor->GetName()));
+		Builder.AddWarningMessage(FString::Printf(TEXT("Failed to export Intensity for HDRIBackdrop %s"), *BackdropActor->GetName()));
 	}
 
 	float Size;
-	if (FGLTFActorUtility::TryGetPropertyValue(Actor, TEXT("Size"), Size))
+	if (FGLTFActorUtility::TryGetPropertyValue(BackdropActor, TEXT("Size"), Size))
 	{
 		JsonBackdrop.Size = Size;
 	}
 	else
 	{
-		Builder.AddWarningMessage(FString::Printf(TEXT("Failed to export Size for HDRIBackdrop %s"), *Actor->GetName()));
+		Builder.AddWarningMessage(FString::Printf(TEXT("Failed to export Size for HDRIBackdrop %s"), *BackdropActor->GetName()));
 	}
 
 	FVector ProjectionCenter;
-	if (FGLTFActorUtility::TryGetPropertyValue(Actor, TEXT("ProjectionCenter"), ProjectionCenter))
+	if (FGLTFActorUtility::TryGetPropertyValue(BackdropActor, TEXT("ProjectionCenter"), ProjectionCenter))
 	{
 		JsonBackdrop.ProjectionCenter = FGLTFConverterUtility::ConvertPosition(ProjectionCenter, Builder.ExportOptions->ExportScale);
 	}
 	else
 	{
-		Builder.AddWarningMessage(FString::Printf(TEXT("Failed to export ProjectionCenter for HDRIBackdrop %s"), *Actor->GetName()));
+		Builder.AddWarningMessage(FString::Printf(TEXT("Failed to export ProjectionCenter for HDRIBackdrop %s"), *BackdropActor->GetName()));
 	}
 
 	float LightingDistanceFactor;
-	if (FGLTFActorUtility::TryGetPropertyValue(Actor, TEXT("LightingDistanceFactor"), LightingDistanceFactor))
+	if (FGLTFActorUtility::TryGetPropertyValue(BackdropActor, TEXT("LightingDistanceFactor"), LightingDistanceFactor))
 	{
 		JsonBackdrop.LightingDistanceFactor = LightingDistanceFactor;
 	}
 	else
 	{
-		Builder.AddWarningMessage(FString::Printf(TEXT("Failed to export LightingDistanceFactor for HDRIBackdrop %s"), *Actor->GetName()));
+		Builder.AddWarningMessage(FString::Printf(TEXT("Failed to export LightingDistanceFactor for HDRIBackdrop %s"), *BackdropActor->GetName()));
 	}
 
 	bool UseCameraProjection;
-	if (FGLTFActorUtility::TryGetPropertyValue(Actor, TEXT("UseCameraProjection"), UseCameraProjection))
+	if (FGLTFActorUtility::TryGetPropertyValue(BackdropActor, TEXT("UseCameraProjection"), UseCameraProjection))
 	{
 		JsonBackdrop.UseCameraProjection = UseCameraProjection;
 	}
 	else
 	{
-		Builder.AddWarningMessage(FString::Printf(TEXT("Failed to export UseCameraProjection for HDRIBackdrop %s"), *Actor->GetName()));
+		Builder.AddWarningMessage(FString::Printf(TEXT("Failed to export UseCameraProjection for HDRIBackdrop %s"), *BackdropActor->GetName()));
 	}
 
 	return Builder.AddBackdrop(JsonBackdrop);
diff --git a/Engine/Plugins/Enterprise/GLTFExporter/MigratedFromGit/Source/GLTFExporter/Private/Converters/GLTFBackdropConverters.h b/Engine/Plugins/Enterprise/GLTFExporter/MigratedFromGit/Source/GLTFExporter/Private/Converters/GLTFBackdropConverters.h
--- a/Engine/Plugins/Enterprise/GLTFExporter/MigratedFromGit/Source/GLTFExporter/Private/Converters/GLTFBackdropConverters.h
+++ b/Engine/Plugins/Enterprise/GLTFExporter/MigratedFromGit/Source/GLTFExporter/Private/Converters/GLTFBackdropConverters.h
@@ -7,9 +7,9 @@
 #include "Converters/GLTFBuilderContext.h"
 #include "Engine.h"
 
-class FGLTFBackdropConverter : public FGLTFBuilderContext, public TGLTFConverter<FGLTFJsonBackdropIndex, const AActor*>
+class FGLTFBackdropConverter final : public FGLTFBuilderContext, public TGLTFConverter<FGLTFJsonBackdropIndex, const AActor*>
 {
 	using FGLTFBuilderContext::FGLTFBuilderContext;
 
-	virtual FGLTFJsonBackdropIndex Convert(const AActor* Actor) override final;
+	virtual FGLTFJsonBackdropIndex Convert(const AActor* BackdropActor) override;
 };
diff --git a/Engine/Plugins/Enterprise/GLTFExporter/MigratedFromGit/Source/GLTFExporter/Private/Converters/GLTFCameraConverters.cpp b/Engine/Plugins/Enterprise/GLTFExporter/MigratedFromGit/Source/GLTFExporter/Private/Converters/GLTFCameraConverters.cpp
--- a/Engine/Plugins/Enterprise/GLTFExporter/MigratedFromGit/Source/GLTFExporter/Private/Converters/GLTFCameraConverters.cpp
+++ b/Engine/Plugins/Enterprise/GLTFExporter/MigratedFromGit/Source/GLTFExporter/Private/Converters/GLTFCameraConverters.cpp
@@ -6,7 +6,7 @@
 #include "Converters/GLTFNameUtility.h"
 #include "Converters/GLTFCameraUtility.h"
 
-FGLTFJsonCameraIndex FGLTFCameraComponentConverter::Convert(const UCameraComponent* CameraComponent)
+FGLTFJsonCameraIndex FGLTFCameraConverter::Convert(const UCameraComponent* CameraComponent)
 {
 	FGLTFJsonCamera Camera;
 	Camera.Name = FGLTFNameUtility::GetName(CameraComponent);
diff --git a/Engine/Plugins/Enterprise/GLTFExporter/MigratedFromGit/Source/GLTFExporter/Private/Converters/GLTFCameraConverters.h b/Engine/Plugins/Enterprise/GLTFExporter/MigratedFromGit/Source/GLTFExporter/Private/Converters/GLTFCameraConverters.h
--- a/Engine/Plugins/Enterprise/GLTFExporter/MigratedFromGit/Source/GLTFExporter/Private/Converters/GLTFCameraConverters.h
+++ b/Engine/Plugins/Enterprise/GLTFExporter/MigratedFromGit/Source/GLTFExporter/Private/Converters/GLTFCameraConverters.h
@@ -7,9 +7,9 @@
 #include "Converters/GLTFBuilderContext.h"
 #include "Engine.h"
 
-class FGLTFCameraComponentConverter : public FGLTFBuilderContext, public TGLTFConverter<FGLTFJsonCameraIndex, const UCameraComponent*>
+class FGLTFCameraConverter final : public FGLTFBuilderContext, public TGLTFConverter<FGLTFJsonCameraIndex, const UCameraComponent*>
 {
 	using FGLTFBuilderContext::FGLTFBuilderContext;
 
-	virtual FGLTFJsonCameraIndex Convert(const UCameraComponent* CameraComponent) override final;
+	virtual FGLTFJsonCameraIndex Convert(const UCameraComponent* CameraComponent) override;
 };
diff --git a/Engine/Plugins/Enterprise/GLTFExporter/MigratedFromGit/Source/GLTFExporter/Private/Converters/GLTFHotspotConverters.cpp b/Engine/Plugins/Enterprise/GLTFExporter/MigratedFromGit/Source/GLTFExporter/Private/Converters/GLTFHotspotConverters.cpp
--- a/Engine/Plugins/Enterprise/GLTFExporter/MigratedFromGit/Source/GLTFExporter/Private/Converters/GLTFHotspotConverters.cpp
+++ b/Engine/Plugins/Enterprise/GLTFExporter/MigratedFromGit/Source/GLTFExporter/Private/Converters/GLTFHotspotConverters.cpp
@@ -4,7 +4,7 @@
 #include "Json/FGLTFJsonHotspot.h"
 #include "Builders/GLTFContainerBuilder.h"
 
-FGLTFJsonHotspotIndex FGLTFHotspotComponentConverter::Convert(const UGLTFInteractionHotspotComponent* HotspotComponent)
+FGLTFJsonHotspotIndex FGLTFHotspotConverter::Convert(const UGLTFInteractionHotspotComponent* HotspotComponent)
 {
 	// TODO: should we warn and / or return INDEX_NONE if the hotspot has no valid Image assigned?
 
diff --git a/Engine/Plugins/Enterprise/GLTFExporter/MigratedFromGit/Source/GLTFExporter/Private/Converters/GLTFHotspotConverters.h b/Engine/Plugins/Enterprise/GLTFExporter/MigratedFromGit/Source/GLTFExporter/Private/Converters/GLTFHotspotConverters.h
--- a/Engine/Plugins/Enterprise/GLTFExporter/MigratedFromGit/Source/GLTFExporter/Private/Converters/GLTFHotspotConverters.h
+++ b/Engine/Plugins/Enterprise/GLTFExporter/MigratedFromGit/Source/GLTFExporter/Private/Converters/GLTFHotspotConverters.h
@@ -6,11 +6,10 @@
 #include "Converters/GLTFConverter.h"
 #include "Converters/GLTFBuilderContext.h"
 #include "Components/GLTFInteractionHotspotComponent.h"
-#include "Engine.h"
 
-class FGLTFHotspotComponentConverter : public FGLTFBuilderContext, public TGLTFConverter<FGLTFJsonHotspotIndex, const UGLTFInteractionHotspotComponent*>
+class FGLTFHotspotConverter final : public FGLTFBuilderContext, public TGLTFConverter<FGLTFJsonHotspotIndex, const UGLTFInteractionHotspotComponent*>
 {
 	using FGLTFBuilderContext::FGLTFBuilderContext;
 
-	virtual FGLTFJsonHotspotIndex Convert(const UGLTFInteractionHotspotComponent* HotspotComponent) override final;
+	virtual FGLTFJsonHotspotIndex Convert(const UGLTFInteractionHotspotComponent* HotspotComponent) override;
 };
diff --git a/Engine/Plugins/Enterprise/GLTFExporter/MigratedFromGit/Source/GLTFExporter/Private/Converters/GLTFLightConverters.cpp b/Engine/Plugins/Enterprise/GLTFExporter/MigratedFromGit/Source/GLTFExporter/Private/Converters/GLTFLightConverters.cpp
--- a/Engine/Plugins/Enterprise/GLTFExporter/MigratedFromGit/Source/GLTFExporter/Private/Converters/GLTFLightConverters.cpp
+++ b/Engine/Plugins/Enterprise/GLTFExporter/MigratedFromGit/Source/GLTFExporter/Private/Converters/GLTFLightConverters.cpp
@@ -5,7 +5,7 @@
 #include "Converters/GLTFConverterUtility.h"
 #include "Converters/GLTFNameUtility.h"
 
-FGLTFJsonLightIndex FGLTFLightComponentConverter::Convert(const ULightComponent* LightComponent)
+FGLTFJsonLightIndex FGLTFLightConverter::Convert(const ULightComponent* LightComponent)
 {
 	FGLTFJsonLight Light;
 
diff --git a/Engine/Plugins/Enterprise/GLTFExporter/MigratedFromGit/Source/GLTFExporter/Private/Converters/GLTFLightConverters.h b/Engine/Plugins/Enterprise/GLTFExporter/MigratedFromGit/Source/GLTFExporter/Private/Converters/GLTFLightConverters.h
--- a/Engine/Plugins/Enterprise/GLTFExporter/MigratedFromGit/Source/GLTFExporter/Private/Converters/GLTFLightConverters.h
+++ b/Engine/Plugins/Enterprise/GLTFExporter/MigratedFromGit/Source/GLTFExporter/Private/Converters/GLTFLightConverters.h
@@ -7,9 +7,9 @@
 #include "Converters/GLTFBuilderContext.h"
 #include "Engine.h"
 
-class FGLTFLightComponentConverter : public FGLTFBuilderContext, public TGLTFConverter<FGLTFJsonLightIndex, const ULightComponent*>
+class FGLTFLightConverter final : public FGLTFBuilderContext, public TGLTFConverter<FGLTFJsonLightIndex, const ULightComponent*>
 {
 	using FGLTFBuilderContext::FGLTFBuilderContext;
 
-	virtual FGLTFJsonLightIndex Convert(const ULightComponent* LightComponent) override final;
+	virtual FGLTFJsonLightIndex Convert(const ULightComponent* LightComponent) override;
 };
diff --git a/Engine/Plugins/Enterprise/GLTFExporter/MigratedFromGit/Source/GLTFExporter/Private/Converters/GLTFLightMapConverters.h b/Engine/Plugins/Enterprise/GLTFExporter/MigratedFromGit/Source/GLTFExporter/Private/Converters/GLTFLightMapConverters.h
--- a/Engine/Plugins/Enterprise/GLTFExporter/MigratedFromGit/Source/GLTFExporter/Private/Converters/GLTFLightMapConverters.h
+++ b/Engine/Plugins/Enterprise/GLTFExporter/MigratedFromGit/Source/GLTFExporter/Private/Converters/GLTFLightMapConverters.h
@@ -7,9 +7,9 @@
 #include "Converters/GLTFBuilderContext.h"
 #include "Engine.h"
 
-class FGLTFLightMapConverter : public FGLTFBuilderContext, public TGLTFConverter<FGLTFJsonLightMapIndex, const UStaticMeshComponent*>
+class FGLTFLightMapConverter final : public FGLTFBuilderContext, public TGLTFConverter<FGLTFJsonLightMapIndex, const UStaticMeshComponent*>
 {
 	using FGLTFBuilderContext::FGLTFBuilderContext;
 
-	virtual FGLTFJsonLightMapIndex Convert(const UStaticMeshComponent* StaticMeshComponent) override final;
+	virtual FGLTFJsonLightMapIndex Convert(const UStaticMeshComponent* StaticMeshComponent) override;
 };
diff --git a/Engine/Plugins/Enterprise/GLTFExporter/MigratedFromGit/Source/GLTFExporter/Private/Converters/GLTFMaterialConverters.h b/Engine/Plugins/Enterprise/GLTFExporter/MigratedFromGit/Source/GLTFExporter/Private/Converters/GLTFMaterialConverters.h
--- a/Engine/Plugins/Enterprise/GLTFExporter/MigratedFromGit/Source/GLTFExporter/Private/Converters/GLTFMaterialConverters.h
+++ b/Engine/Plugins/Enterprise/GLTFExporter/MigratedFromGit/Source/GLTFExporter/Private/Converters/GLTFMaterialConverters.h
@@ -10,11 +10,11 @@
 
 struct FGLTFPropertyBakeOutput;
 
-class FGLTFMaterialConverter : public FGLTFBuilderContext, public TGLTFConverter<FGLTFJsonMaterialIndex, const UMaterialInterface*>
+class FGLTFMaterialConverter final : public FGLTFBuilderContext, public TGLTFConverter<FGLTFJsonMaterialIndex, const UMaterialInterface*>
 {
 	using FGLTFBuilderContext::FGLTFBuilderContext;
 
-	virtual FGLTFJsonMaterialIndex Convert(const UMaterialInterface* Material) override final;
+	virtual FGLTFJsonMaterialIndex Convert(const UMaterialInterface* Material) override;
 
 	bool TryGetShadingModel(EGLTFJsonShadingModel& ShadingModel, const UMaterialInterface* Material) const;
 
diff --git a/Engine/Plugins/Enterprise/GLTFExporter/MigratedFromGit/Source/GLTFExporter/Private/Converters/GLTFMeshConverters.h b/Engine/Plugins/Enterprise/GLTFExporter/MigratedFromGit/Source/GLTFExporter/Private/Converters/GLTFMeshConverters.h
--- a/Engine/Plugins/Enterprise/GLTFExporter/MigratedFromGit/Source/GLTFExporter/Private/Converters/GLTFMeshConverters.h
+++ b/Engine/Plugins/Enterprise/GLTFExporter/MigratedFromGit/Source/GLTFExporter/Private/Converters/GLTFMeshConverters.h
@@ -14,16 +14,16 @@ class FGLTFMeshConverter : public FGLTFBuilderContext, public TGLTFConverter<FGL
 	using FGLTFBuilderContext::FGLTFBuilderContext;
 };
 
-class FGLTFStaticMeshConverter : public FGLTFMeshConverter<const UStaticMesh*, int32, const FColorVertexBuffer*, FGLTFMaterialArray>
+class FGLTFStaticMeshConverter final : public FGLTFMeshConverter<const UStaticMesh*, int32, const FColorVertexBuffer*, FGLTFMaterialArray>
 {
 	using FGLTFMeshConverter::FGLTFMeshConverter;
 
-	virtual FGLTFJsonMeshIndex Convert(const UStaticMesh* StaticMesh, int32 LODIndex, const FColorVertexBuffer* OverrideVertexColors, FGLTFMaterialArray OverrideMaterials) override final;
+	virtual FGLTFJsonMeshIndex Convert(const UStaticMesh* StaticMesh, int32 LODIndex, const FColorVertexBuffer* OverrideVertexColors, FGLTFMaterialArray OverrideMaterials) override;
 };
 
-class FGLTFSkeletalMeshConverter : public FGLTFMeshConverter<const USkeletalMesh*, int32, const FColorVertexBuffer*, const FSkinWeightVertexBuffer*, FGLTFMaterialArray>
+class FGLTFSkeletalMeshConverter final : public FGLTFMeshConverter<const USkeletalMesh*, int32, const FColorVertexBuffer*, const FSkinWeightVertexBuffer*, FGLTFMaterialArray>
 {
 	using FGLTFMeshConverter::FGLTFMeshConverter;
 
-	virtual FGLTFJsonMeshIndex Convert(const USkeletalMesh* SkeletalMesh, int32 LODIndex, const FColorVertexBuffer* OverrideVertexColors, const FSkinWeightVertexBuffer* OverrideSkinWeights, FGLTFMaterialArray OverrideMaterials) override final;
+	virtual FGLTFJsonMeshIndex Convert(const USkeletalMesh* SkeletalMesh, int32 LODIndex, const FColorVertexBuffer* OverrideVertexColors, const FSkinWeightVertexBuffer* OverrideSkinWeights, FGLTFMaterialArray OverrideMaterials) override;
 };
diff --git a/Engine/Plugins/Enterprise/GLTFExporter/MigratedFromGit/Source/GLTFExporter/Private/Converters/GLTFNodeConverters.cpp b/Engine/Plugins/Enterprise/GLTFExporter/MigratedFromGit/Source/GLTFExporter/Private/Converters/GLTFNodeConverters.cpp
--- a/Engine/Plugins/Enterprise/GLTFExporter/MigratedFromGit/Source/GLTFExporter/Private/Converters/GLTFNodeConverters.cpp
+++ b/Engine/Plugins/Enterprise/GLTFExporter/MigratedFromGit/Source/GLTFExporter/Private/Converters/GLTFNodeConverters.cpp
@@ -7,7 +7,7 @@
 #include "Converters/GLTFNameUtility.h"
 #include "Components/GLTFInteractionHotspotComponent.h"
 
-FGLTFJsonNodeIndex FGLTFSceneComponentConverter::Convert(const USceneComponent* SceneComponent)
+FGLTFJsonNodeIndex FGLTFComponentConverter::Convert(const USceneComponent* SceneComponent)
 {
 	const AActor* Owner = SceneComponent->GetOwner();
 	if (Owner == nullptr)
diff --git a/Engine/Plugins/Enterprise/GLTFExporter/MigratedFromGit/Source/GLTFExporter/Private/Converters/GLTFNodeConverters.h b/Engine/Plugins/Enterprise/GLTFExporter/MigratedFromGit/Source/GLTFExporter/Private/Converters/GLTFNodeConverters.h
--- a/Engine/Plugins/Enterprise/GLTFExporter/MigratedFromGit/Source/GLTFExporter/Private/Converters/GLTFNodeConverters.h
+++ b/Engine/Plugins/Enterprise/GLTFExporter/MigratedFromGit/Source/GLTFExporter/Private/Converters/GLTFNodeConverters.h
@@ -13,16 +13,16 @@ class FGLTFNodeConverter : public FGLTFBuilderContext, public TGLTFConverter<FGL
 	using FGLTFBuilderContext::FGLTFBuilderContext;
 };
 
-class FGLTFSceneComponentConverter : public FGLTFNodeConverter<const USceneComponent*>
+class FGLTFComponentConverter final : public FGLTFNodeConverter<const USceneComponent*>
 {
 	using FGLTFNodeConverter::FGLTFNodeConverter;
 
-	virtual FGLTFJsonNodeIndex Convert(const USceneComponent* SceneComponent) override final;
+	virtual FGLTFJsonNodeIndex Convert(const USceneComponent* SceneComponent) override;
 };
 
-class FGLTFActorConverter : public FGLTFNodeConverter<const AActor*>
+class FGLTFActorConverter final : public FGLTFNodeConverter<const AActor*>
 {
 	using FGLTFNodeConverter::FGLTFNodeConverter;
 
-	virtual FGLTFJsonNodeIndex Convert(const AActor* Actor) override final;
+	virtual FGLTFJsonNodeIndex Convert(const AActor* Actor) override;
 };
diff --git a/Engine/Plugins/Enterprise/GLTFExporter/MigratedFromGit/Source/GLTFExporter/Private/Converters/GLTFSamplerConverters.cpp b/Engine/Plugins/Enterprise/GLTFExporter/MigratedFromGit/Source/GLTFExporter/Private/Converters/GLTFSamplerConverters.cpp
--- a/Engine/Plugins/Enterprise/GLTFExporter/MigratedFromGit/Source/GLTFExporter/Private/Converters/GLTFSamplerConverters.cpp
+++ b/Engine/Plugins/Enterprise/GLTFExporter/MigratedFromGit/Source/GLTFExporter/Private/Converters/GLTFSamplerConverters.cpp
@@ -5,7 +5,7 @@
 #include "Converters/GLTFConverterUtility.h"
 #include "Converters/GLTFTextureUtility.h"
 
-FGLTFJsonSamplerIndex FGLTFTextureSamplerConverter::Convert(const UTexture* Texture)
+FGLTFJsonSamplerIndex FGLTFSamplerConverter::Convert(const UTexture* Texture)
 {
 	// TODO: maybe we should reuse existing samplers?
 
diff --git a/Engine/Plugins/Enterprise/GLTFExporter/MigratedFromGit/Source/GLTFExporter/Private/Converters/GLTFSamplerConverters.h b/Engine/Plugins/Enterprise/GLTFExporter/MigratedFromGit/Source/GLTFExporter/Private/Converters/GLTFSamplerConverters.h
--- a/Engine/Plugins/Enterprise/GLTFExporter/MigratedFromGit/Source/GLTFExporter/Private/Converters/GLTFSamplerConverters.h
+++ b/Engine/Plugins/Enterprise/GLTFExporter/MigratedFromGit/Source/GLTFExporter/Private/Converters/GLTFSamplerConverters.h
@@ -7,9 +7,9 @@
 #include "Converters/GLTFBuilderContext.h"
 #include "Engine.h"
 
-class FGLTFTextureSamplerConverter : public FGLTFBuilderContext, public TGLTFConverter<FGLTFJsonSamplerIndex, const UTexture*>
+class FGLTFSamplerConverter final : public FGLTFBuilderContext, public TGLTFConverter<FGLTFJsonSamplerIndex, const UTexture*>
 {
 	using FGLTFBuilderContext::FGLTFBuilderContext;
 
-	virtual FGLTFJsonSamplerIndex Convert(const UTexture* Texture) override final;
+	virtual FGLTFJsonSamplerIndex Convert(const UTexture* Texture) override;
 };
diff --git a/Engine/Plugins/Enterprise/GLTFExporter/MigratedFromGit/Source/GLTFExporter/Private/Converters/GLTFSceneConverters.cpp b/Engine/Plugins/Enterprise/GLTFExporter/MigratedFromGit/Source/GLTFExporter/Private/Converters/GLTFSceneConverters.cpp
--- a/Engine/Plugins/Enterprise/GLTFExporter/MigratedFromGit/Source/GLTFExporter/Private/Converters/GLTFSceneConverters.cpp
+++ b/Engine/Plugins/Enterprise/GLTFExporter/MigratedFromGit/Source/GLTFExporter/Private/Converters/GLTFSceneConverters.cpp
@@ -5,7 +5,7 @@
 #include "Converters/GLTFConverterUtility.h"
 #include "Converters/GLTFActorUtility.h"
 
-FGLTFJsonSceneIndex FGLTFLevelConverter::Convert(const ULevel* Level)
+FGLTFJsonSceneIndex FGLTFSceneConverter::Convert(const ULevel* Level)
 {
 	FGLTFJsonScene Scene;
 
diff --git a/Engine/Plugins/Enterprise/GLTFExporter/MigratedFromGit/Source/GLTFExporter/Private/Converters/GLTFSceneConverters.h b/Engine/Plugins/Enterprise/GLTFExporter/MigratedFromGit/Source/GLTFExporter/Private/Converters/GLTFSceneConverters.h
--- a/Engine/Plugins/Enterprise/GLTFExporter/MigratedFromGit/Source/GLTFExporter/Private/Converters/GLTFSceneConverters.h
+++ b/Engine/Plugins/Enterprise/GLTFExporter/MigratedFromGit/Source/GLTFExporter/Private/Converters/GLTFSceneConverters.h
@@ -7,9 +7,9 @@
 #include "Converters/GLTFBuilderContext.h"
 #include "Engine.h"
 
-class FGLTFLevelConverter : public FGLTFBuilderContext, public TGLTFConverter<FGLTFJsonSceneIndex, const ULevel*>
+class FGLTFSceneConverter final : public FGLTFBuilderContext, public TGLTFConverter<FGLTFJsonSceneIndex, const ULevel*>
 {
 	using FGLTFBuilderContext::FGLTFBuilderContext;
 
-	virtual FGLTFJsonSceneIndex Convert(const ULevel* Level) override final;
+	virtual FGLTFJsonSceneIndex Convert(const ULevel* Level) override;
 };
diff --git a/Engine/Plugins/Enterprise/GLTFExporter/MigratedFromGit/Source/GLTFExporter/Private/Converters/GLTFTextureConverters.h b/Engine/Plugins/Enterprise/GLTFExporter/MigratedFromGit/Source/GLTFExporter/Private/Converters/GLTFTextureConverters.h
--- a/Engine/Plugins/Enterprise/GLTFExporter/MigratedFromGit/Source/GLTFExporter/Private/Converters/GLTFTextureConverters.h
+++ b/Engine/Plugins/Enterprise/GLTFExporter/MigratedFromGit/Source/GLTFExporter/Private/Converters/GLTFTextureConverters.h
@@ -13,30 +13,30 @@ class FGLTFTextureConverter : public FGLTFBuilderContext, public TGLTFConverter<
 	using FGLTFBuilderContext::FGLTFBuilderContext;
 };
 
-class FGLTFTexture2DConverter : public FGLTFTextureConverter<const UTexture2D*>
+class FGLTFTexture2DConverter final : public FGLTFTextureConverter<const UTexture2D*>
 {
 	using FGLTFTextureConverter::FGLTFTextureConverter;
 
-	virtual FGLTFJsonTextureIndex Convert(const UTexture2D* Texture2D) override final;
+	virtual FGLTFJsonTextureIndex Convert(const UTexture2D* Texture2D) override;
 };
 
-class FGLTFTextureCubeConverter : public FGLTFTextureConverter<const UTextureCube*, ECubeFace>
+class FGLTFTextureCubeConverter final : public FGLTFTextureConverter<const UTextureCube*, ECubeFace>
 {
 	using FGLTFTextureConverter::FGLTFTextureConverter;
 
-	virtual FGLTFJsonTextureIndex Convert(const UTextureCube* TextureCube, ECubeFace CubeFace) override final;
+	virtual FGLTFJsonTextureIndex Convert(const UTextureCube* TextureCube, ECubeFace CubeFace) override;
 };
 
-class FGLTFTextureRenderTarget2DConverter : public FGLTFTextureConverter<const UTextureRenderTarget2D*>
+class FGLTFTextureRenderTarget2DConverter final : public FGLTFTextureConverter<const UTextureRenderTarget2D*>
 {
 	using FGLTFTextureConverter::FGLTFTextureConverter;
 
-	virtual FGLTFJsonTextureIndex Convert(const UTextureRenderTarget2D* RenderTarget2D) override final;
+	virtual FGLTFJsonTextureIndex Convert(const UTextureRenderTarget2D* RenderTarget2D) override;
 };
 
-class FGLTFTextureRenderTargetCubeConverter : public FGLTFTextureConverter<const UTextureRenderTargetCube*, ECubeFace>
+class FGLTFTextureRenderTargetCubeConverter final : public FGLTFTextureConverter<const UTextureRenderTargetCube*, ECubeFace>
 {
 	using FGLTFTextureConverter::FGLTFTextureConverter;
 
-	virtual FGLTFJsonTextureIndex Convert(const UTextureRenderTargetCube* RenderTargetCube, ECubeFace CubeFace) override final;
+	virtual FGLTFJsonTextureIndex Convert(const UTextureRenderTargetCube* RenderTargetCube, ECubeFace CubeFace) override;
 };
diff --git a/Engine/Plugins/Enterprise/GLTFExporter/MigratedFromGit/Source/GLTFExporter/Private/Converters/GLTFVarationConverters.h b/Engine/Plugins/Enterprise/GLTFExporter/MigratedFromGit/Source/GLTFExporter/Private/Converters/GLTFVarationConverters.h
--- a/Engine/Plugins/Enterprise/GLTFExporter/MigratedFromGit/Source/GLTFExporter/Private/Converters/GLTFVarationConverters.h
+++ b/Engine/Plugins/Enterprise/GLTFExporter/MigratedFromGit/Source/GLTFExporter/Private/Converters/GLTFVarationConverters.h
@@ -12,11 +12,11 @@
 #include "PropertyValue.h"
 #include "LevelVariantSets.h"
 
-class FGLTFVariationConverter : public FGLTFBuilderContext, public TGLTFConverter<FGLTFJsonVariationIndex, const ALevelVariantSetsActor*>
+class FGLTFVariationConverter final : public FGLTFBuilderContext, public TGLTFConverter<FGLTFJsonVariationIndex, const ALevelVariantSetsActor*>
 {
 	using FGLTFBuilderContext::FGLTFBuilderContext;
 
-	virtual FGLTFJsonVariationIndex Convert(const ALevelVariantSetsActor* LevelVariantSetsActor) override final;
+	virtual FGLTFJsonVariationIndex Convert(const ALevelVariantSetsActor* LevelVariantSetsActor) override;
 
 	bool TryParseVariant(FGLTFJsonVariant& OutVariant, const UVariant* Variant) const;
 	bool TryParseVariantBinding(FGLTFJsonVariant& OutVariant, const UVariantObjectBinding* Binding) const;
