[GLTFExporter] feat(orbit-camera): add JSON-structure for orbital camera extension

commit 0c199f59706b94f6738bbb46094989ff7362bce3
Date: Fri, 18 Sep 2020 15:31:17 +0200

    feat(orbit-camera): add JSON-structure for orbital camera extension
    (cherry picked from commit f3bca59c6a30a71c3228811335b4e457c42e8117​)

#skipci

diff --git a/Engine/Plugins/Enterprise/GLTFExporter/MigratedFromGit/Source/GLTFExporter/Private/Json/GLTFJsonEnums.h b/Engine/Plugins/Enterprise/GLTFExporter/MigratedFromGit/Source/GLTFExporter/Private/Json/GLTFJsonEnums.h
--- a/Engine/Plugins/Enterprise/GLTFExporter/MigratedFromGit/Source/GLTFExporter/Private/Json/GLTFJsonEnums.h
+++ b/Engine/Plugins/Enterprise/GLTFExporter/MigratedFromGit/Source/GLTFExporter/Private/Json/GLTFJsonEnums.h
@@ -10,7 +10,7 @@ enum class EGLTFJsonExtension
 	KHR_MaterialsUnlit,
 	KHR_MaterialsClearCoat,
 	KHR_MeshQuantization,
-	EPIC_CameraOrbital,
+	EPIC_OrbitalCameras,
 	EPIC_HDRIBackdrops,
 	EPIC_InteractionHotspots,
 	EPIC_LightmapTextures,
diff --git a/Engine/Plugins/Enterprise/GLTFExporter/MigratedFromGit/Source/GLTFExporter/Private/Json/GLTFJsonOrbitCamera.h b/Engine/Plugins/Enterprise/GLTFExporter/MigratedFromGit/Source/GLTFExporter/Private/Json/GLTFJsonOrbitCamera.h
new file mode 100644
--- /dev/null
+++ b/Engine/Plugins/Enterprise/GLTFExporter/MigratedFromGit/Source/GLTFExporter/Private/Json/GLTFJsonOrbitCamera.h
@@ -0,0 +1,29 @@
+// Copyright Epic Games, Inc. All Rights Reserved.
+
+#pragma once
+
+#include "Json/GLTFJsonIndex.h"
+#include "Json/GLTFJsonVector4.h"
+
+struct FGLTFJsonOrbitCamera
+{
+	FGLTFJsonNodeIndex Focus;
+	float              MaxDistance;
+	float              MinDistance;
+	float              MaxAngle;
+	float              MinAngle;
+
+	template <class CharType = TCHAR, class PrintPolicy = TPrettyJsonPrintPolicy<CharType>>
+	void WriteObject(TJsonWriter<CharType, PrintPolicy>& JsonWriter, FGLTFJsonExtensions& Extensions) const
+	{
+		JsonWriter.WriteObjectStart();
+
+		JsonWriter.WriteValue(TEXT("focus"), Focus);
+		FGLTFJsonUtility::WriteExactValue(JsonWriter, TEXT("maxDistance"), MaxDistance);
+		FGLTFJsonUtility::WriteExactValue(JsonWriter, TEXT("minDistance"), MinDistance);
+		FGLTFJsonUtility::WriteExactValue(JsonWriter, TEXT("maxAngle"), MaxAngle);
+		FGLTFJsonUtility::WriteExactValue(JsonWriter, TEXT("minAngle"), MinAngle);
+
+		JsonWriter.WriteObjectEnd();
+	}
+};
diff --git a/Engine/Plugins/Enterprise/GLTFExporter/MigratedFromGit/Source/GLTFExporter/Private/Json/GLTFJsonUtility.h b/Engine/Plugins/Enterprise/GLTFExporter/MigratedFromGit/Source/GLTFExporter/Private/Json/GLTFJsonUtility.h
--- a/Engine/Plugins/Enterprise/GLTFExporter/MigratedFromGit/Source/GLTFExporter/Private/Json/GLTFJsonUtility.h
+++ b/Engine/Plugins/Enterprise/GLTFExporter/MigratedFromGit/Source/GLTFExporter/Private/Json/GLTFJsonUtility.h
@@ -23,7 +23,7 @@ struct FGLTFJsonUtility
 			case EGLTFJsonExtension::KHR_MaterialsUnlit:       return TEXT("KHR_materials_unlit");
 			case EGLTFJsonExtension::KHR_MaterialsClearCoat:   return TEXT("KHR_materials_clearcoat");
 			case EGLTFJsonExtension::KHR_MeshQuantization:     return TEXT("KHR_mesh_quantization");
-			case EGLTFJsonExtension::EPIC_CameraOrbital:       return TEXT("EPIC_camera_orbital");
+			case EGLTFJsonExtension::EPIC_OrbitalCameras:      return TEXT("EPIC_orbital_cameras");
 			case EGLTFJsonExtension::EPIC_HDRIBackdrops:       return TEXT("EPIC_hdri_backdrops");
 			case EGLTFJsonExtension::EPIC_InteractionHotspots: return TEXT("EPIC_interaction_hotspots");
 			case EGLTFJsonExtension::EPIC_LightmapTextures:    return TEXT("EPIC_lightmap_textures");
