[GLTFExporter] feat(orbit-camera): add optional orbit camera extension to FGLTFJsonCamera

commit a73aa144be4c615c664b80d3b837bb9853d271cb
Date: Fri, 18 Sep 2020 15:31:44 +0200

    feat(orbit-camera): add optional orbit camera extension to FGLTFJsonCamera
    (cherry picked from commit ccdd8712bea8cdf0c9cc076e0e1df456a3f82268​)

#skipci

diff --git a/Engine/Plugins/Enterprise/GLTFExporter/MigratedFromGit/Source/GLTFExporter/Private/Json/GLTFJsonCamera.h b/Engine/Plugins/Enterprise/GLTFExporter/MigratedFromGit/Source/GLTFExporter/Private/Json/GLTFJsonCamera.h
--- a/Engine/Plugins/Enterprise/GLTFExporter/MigratedFromGit/Source/GLTFExporter/Private/Json/GLTFJsonCamera.h
+++ b/Engine/Plugins/Enterprise/GLTFExporter/MigratedFromGit/Source/GLTFExporter/Private/Json/GLTFJsonCamera.h
@@ -3,6 +3,7 @@
 #pragma once
 
 #include "Json/GLTFJsonEnums.h"
+#include "Json/GLTFJsonOrbitCamera.h"
 #include "Json/GLTFJsonUtility.h"
 #include "Serialization/JsonSerializer.h"
 
@@ -81,7 +82,8 @@ struct FGLTFJsonCamera
 {
 	FString Name;
 
-	EGLTFJsonCameraType Type;
+	EGLTFJsonCameraType             Type;
+	TOptional<FGLTFJsonOrbitCamera> OrbitCamera;
 
 	union {
 		FGLTFJsonOrthographic Orthographic;
@@ -123,6 +125,19 @@ struct FGLTFJsonCamera
 				break;
 		}
 
+		if (OrbitCamera.IsSet())
+		{
+			JsonWriter.WriteObjectStart(TEXT("extensions"));
+
+			const EGLTFJsonExtension Extension = EGLTFJsonExtension::EPIC_OrbitalCameras;
+			Extensions.Used.Add(Extension);
+
+			JsonWriter.WriteIdentifierPrefix(FGLTFJsonUtility::ToString(Extension));
+			OrbitCamera.GetValue().WriteObject(JsonWriter, Extensions);
+
+			JsonWriter.WriteObjectEnd();
+		}
+
 		JsonWriter.WriteObjectEnd();
 	}
 };
