[GLTFExporter] feat(orbit-camera): add conversion of orbital cameras

commit 3a72b70333b96eda78e7ea2c110a582de7a0ad65
Date: Fri, 18 Sep 2020 15:32:37 +0200

    feat(orbit-camera): add conversion of orbital cameras
    (cherry picked from commit aa7f1af5b33c7df769bf0e701d2a261c3ab2e50e​)

#skipci

diff --git a/Engine/Plugins/Enterprise/GLTFExporter/MigratedFromGit/Source/GLTFExporter/Private/Converters/GLTFCameraConverters.cpp b/Engine/Plugins/Enterprise/GLTFExporter/MigratedFromGit/Source/GLTFExporter/Private/Converters/GLTFCameraConverters.cpp
--- a/Engine/Plugins/Enterprise/GLTFExporter/MigratedFromGit/Source/GLTFExporter/Private/Converters/GLTFCameraConverters.cpp
+++ b/Engine/Plugins/Enterprise/GLTFExporter/MigratedFromGit/Source/GLTFExporter/Private/Converters/GLTFCameraConverters.cpp
@@ -5,6 +5,7 @@
 #include "Converters/GLTFConverterUtility.h"
 #include "Converters/GLTFNameUtility.h"
 #include "Converters/GLTFCameraUtility.h"
+#include "Actors/GLTFOrbitCameraActor.h"
 
 FGLTFJsonCameraIndex FGLTFCameraConverter::Convert(const UCameraComponent* CameraComponent)
 {
@@ -30,5 +31,41 @@ FGLTFJsonCameraIndex FGLTFCameraConverter::Convert(const UCameraComponent* Camer
 		    return FGLTFJsonCameraIndex(INDEX_NONE);
 	}
 
+	const AActor* Owner = CameraComponent->GetOwner();
+	const AGLTFOrbitCameraActor* OrbitCameraActor = Owner != nullptr ? Cast<AGLTFOrbitCameraActor>(Owner) : nullptr;
+
+	if (OrbitCameraActor != nullptr)
+	{
+		if (Builder.ExportOptions->bExportOrbitalCameras)
+		{
+			if (OrbitCameraActor->Focus != nullptr)
+			{
+				FGLTFJsonOrbitCamera OrbitCamera;
+				OrbitCamera.Focus = Builder.GetOrAddNode(OrbitCameraActor->Focus);
+				OrbitCamera.MaxDistance = FGLTFConverterUtility::ConvertLength(OrbitCameraActor->DistanceMax, Builder.ExportOptions->ExportScale);
+				OrbitCamera.MinDistance = FGLTFConverterUtility::ConvertLength(OrbitCameraActor->DistanceMin, Builder.ExportOptions->ExportScale);
+
+				// TODO: we need to decide what values to use for looking straight ahead, up and down.
+				// Right now the actor and the viewer have different opinions.
+				OrbitCamera.MaxAngle = OrbitCameraActor->PitchAngleMax + 90.0f;
+				OrbitCamera.MinAngle = OrbitCameraActor->PitchAngleMin + 90.0f;
+
+				Camera.OrbitCamera = OrbitCamera;
+			}
+			else
+			{
+				Builder.AddWarningMessage(FString::Printf(
+					TEXT("OrbitalCamera %s has no focus set and will be skipped"),
+					*Owner->GetName()));
+			}
+		}
+		else
+		{
+			Builder.AddWarningMessage(FString::Printf(
+				TEXT("OrbitalCamera %s disabled by export options"),
+				*Owner->GetName()));
+		}
+	}
+
 	return Builder.AddCamera(Camera);
 }
