[GLTFExporter] feat: make the bExportOrbitalCameras option true by default and editable

commit 20a318dbf2eec5eccfcc35fdb9ec46b00129772b
Date: Fri, 18 Sep 2020 15:33:11 +0200

    feat: make the bExportOrbitalCameras option true by default and editable
    (cherry picked from commit 14765b8ac79d47ab4a8e7a5db67bca1350343455​)

#skipci

diff --git a/Engine/Plugins/Enterprise/GLTFExporter/MigratedFromGit/Source/GLTFExporter/Private/GLTFExportOptions.cpp b/Engine/Plugins/Enterprise/GLTFExporter/MigratedFromGit/Source/GLTFExporter/Private/GLTFExportOptions.cpp
--- a/Engine/Plugins/Enterprise/GLTFExporter/MigratedFromGit/Source/GLTFExporter/Private/GLTFExportOptions.cpp
+++ b/Engine/Plugins/Enterprise/GLTFExporter/MigratedFromGit/Source/GLTFExporter/Private/GLTFExportOptions.cpp
@@ -30,7 +30,7 @@ UGLTFExportOptions::UGLTFExportOptions(const FObjectInitializer& ObjectInitializ
 	bExportHiddenInGame = false;
 	ExportLights = EGLTFExporterLightMobility::MovableAndStationary;
 	bExportCameras = true;
-	bExportOrbitalCameras = false;
+	bExportOrbitalCameras = true;
 	bExportReflectionCaptures = false;
 	bExportHDRIBackdrops = true;
 	bExportVariantSets = true;
@@ -204,6 +204,7 @@ bool UGLTFExportOptions::CanEditChange(const FProperty* InProperty) const
 		PropertyFName == GET_MEMBER_NAME_CHECKED(ThisClass, ExportScale) ||
 		PropertyFName == GET_MEMBER_NAME_CHECKED(ThisClass, ExportLights) ||
 		PropertyFName == GET_MEMBER_NAME_CHECKED(ThisClass, bExportCameras) ||
+		PropertyFName == GET_MEMBER_NAME_CHECKED(ThisClass, bExportOrbitalCameras) ||
 		PropertyFName == GET_MEMBER_NAME_CHECKED(ThisClass, bExportHDRIBackdrops) ||
 		PropertyFName == GET_MEMBER_NAME_CHECKED(ThisClass, bExportVariantSets) ||
 		PropertyFName == GET_MEMBER_NAME_CHECKED(ThisClass, bExportInteractionHotspots) ||
