[GLTFExporter] chore(orbit-camera): update TODO for differences in pitch compared to the viewer

commit 7d7fecd8bf80802e4684d685b9a6085e033a4b6a
Date: Mon, 21 Sep 2020 09:41:27 +0200

    chore(orbit-camera): update TODO for differences in pitch compared to the viewer
    (cherry picked from commit ac9cbb9ea3aa7185daca6deeff4d490862106e88​)

#skipci

diff --git a/Engine/Plugins/Enterprise/GLTFExporter/MigratedFromGit/Source/GLTFExporter/Private/Converters/GLTFCameraConverters.cpp b/Engine/Plugins/Enterprise/GLTFExporter/MigratedFromGit/Source/GLTFExporter/Private/Converters/GLTFCameraConverters.cpp
--- a/Engine/Plugins/Enterprise/GLTFExporter/MigratedFromGit/Source/GLTFExporter/Private/Converters/GLTFCameraConverters.cpp
+++ b/Engine/Plugins/Enterprise/GLTFExporter/MigratedFromGit/Source/GLTFExporter/Private/Converters/GLTFCameraConverters.cpp
@@ -45,8 +45,8 @@ FGLTFJsonCameraIndex FGLTFCameraConverter::Convert(const UCameraComponent* Camer
 				OrbitCamera.MaxDistance = FGLTFConverterUtility::ConvertLength(OrbitCameraActor->DistanceMax, Builder.ExportOptions->ExportScale);
 				OrbitCamera.MinDistance = FGLTFConverterUtility::ConvertLength(OrbitCameraActor->DistanceMin, Builder.ExportOptions->ExportScale);
 
-				// TODO: we need to decide what values to use for looking straight ahead, up and down.
-				// Right now the actor and the viewer have different opinions.
+				// TODO: should we modify AGLTFOrbitCameraActor to calculate pitch the same way
+				// as the viewer, to avoid having to manually add 90 degrees at export?
 				OrbitCamera.MaxAngle = OrbitCameraActor->PitchAngleMax + 90.0f;
 				OrbitCamera.MinAngle = OrbitCameraActor->PitchAngleMin + 90.0f;
 
