[GLTFExporter] feat(orbit-camera): use calculated center of the focused actor as focus-position

commit 35be140f2d4e6799bac659db60da0f0d5ee778ab
Date: Tue, 22 Sep 2020 15:47:06 +0200

    feat(orbit-camera): use calculated center of the focused actor as focus-position
    (cherry picked from commit 36baca868aa5db0e64aa3e4bbdb63893db429666​)

#skipci

diff --git a/Engine/Plugins/Enterprise/GLTFExporter/MigratedFromGit/Source/GLTFExporterRuntime/Private/Actors/GLTFOrbitCameraActor.cpp b/Engine/Plugins/Enterprise/GLTFExporter/MigratedFromGit/Source/GLTFExporterRuntime/Private/Actors/GLTFOrbitCameraActor.cpp
--- a/Engine/Plugins/Enterprise/GLTFExporter/MigratedFromGit/Source/GLTFExporterRuntime/Private/Actors/GLTFOrbitCameraActor.cpp
+++ b/Engine/Plugins/Enterprise/GLTFExporter/MigratedFromGit/Source/GLTFExporterRuntime/Private/Actors/GLTFOrbitCameraActor.cpp
@@ -24,6 +24,7 @@ AGLTFOrbitCameraActor::AGLTFOrbitCameraActor(const FObjectInitializer& ObjectIni
 	, OrbitInertia(0.1f)
 	, OrbitSensitivity(5.0f)
 	, DistanceSensitivity(50.0f)
+	, FocusPosition(0.0f, 0.0f, 0.0f)
 	, Distance(0.0f)
 	, Pitch(0.0f)
 	, Yaw(0.0f)
@@ -40,10 +41,17 @@ void AGLTFOrbitCameraActor::BeginPlay()
 {
 	Super::BeginPlay();
 
-	if (this->Focus == nullptr)
+	if (this->Focus != nullptr && this->Focus != this)
 	{
-		// TODO: if focus is null use the scene center (similar to viewer camera)
-		UE_LOG(LogTemp, Warning, TEXT("The camera focus must not be null"));
+		const FBox BoundingBox = this->Focus->GetComponentsBoundingBox(true, true);
+		FocusPosition = BoundingBox.IsValid ? BoundingBox.GetCenter() : this->Focus->GetActorLocation();
+	}
+	else
+	{
+		UE_LOG(LogTemp, Warning, TEXT("The camera focus must not be null, and must not be the camera's own actor"));
+
+		// TODO: use the scene center (similar to viewer camera) instead of using the zero-vector.
+		// It may however prove difficult, since we would need to exclude sky-spheres, backdrops etc when calculating the center.
 	}
 
 	// TODO: calculate distance and direction based on camera offset to focus (similar to viewer camera)
@@ -87,7 +95,6 @@ void AGLTFOrbitCameraActor::Tick(float DeltaSeconds)
 		TargetYaw = ClampYaw(TargetYaw);
 	}
 
-	const FVector FocusPosition = (this->Focus != nullptr && this->Focus != this) ? this->Focus->GetActorLocation() : FVector::ZeroVector;
 	const FTransform FocusTransform = FTransform(FocusPosition);
 	const FTransform DollyTransform = FTransform(-FVector::ForwardVector * Distance);
 	const FTransform RotationTransform = FTransform(FQuat::MakeFromEuler(FVector(0.0f, Pitch, Yaw)));
diff --git a/Engine/Plugins/Enterprise/GLTFExporter/MigratedFromGit/Source/GLTFExporterRuntime/Public/Actors/GLTFOrbitCameraActor.h b/Engine/Plugins/Enterprise/GLTFExporter/MigratedFromGit/Source/GLTFExporterRuntime/Public/Actors/GLTFOrbitCameraActor.h
--- a/Engine/Plugins/Enterprise/GLTFExporter/MigratedFromGit/Source/GLTFExporterRuntime/Public/Actors/GLTFOrbitCameraActor.h
+++ b/Engine/Plugins/Enterprise/GLTFExporter/MigratedFromGit/Source/GLTFExporterRuntime/Public/Actors/GLTFOrbitCameraActor.h
@@ -82,6 +82,7 @@ public:
 	float DistanceSensitivity;
 
 private:
+	FVector FocusPosition;
 	float Distance;
 	float Pitch;
 	float Yaw;
