[GLTFExporter] feat(orbit-camera): aim the camera at the focused actor from the start

commit 8a41d39994fe19f344e36cd92246acc4729eb38a
Date: Tue, 22 Sep 2020 15:48:31 +0200

    feat(orbit-camera): aim the camera at the focused actor from the start
    (cherry picked from commit 548c7fb59b16debeaa096d68757851c0c191a0fa​)

#skipci

diff --git a/Engine/Plugins/Enterprise/GLTFExporter/MigratedFromGit/Source/GLTFExporterRuntime/Private/Actors/GLTFOrbitCameraActor.cpp b/Engine/Plugins/Enterprise/GLTFExporter/MigratedFromGit/Source/GLTFExporterRuntime/Private/Actors/GLTFOrbitCameraActor.cpp
--- a/Engine/Plugins/Enterprise/GLTFExporter/MigratedFromGit/Source/GLTFExporterRuntime/Private/Actors/GLTFOrbitCameraActor.cpp
+++ b/Engine/Plugins/Enterprise/GLTFExporter/MigratedFromGit/Source/GLTFExporterRuntime/Private/Actors/GLTFOrbitCameraActor.cpp
@@ -54,10 +54,16 @@ void AGLTFOrbitCameraActor::BeginPlay()
 		// It may however prove difficult, since we would need to exclude sky-spheres, backdrops etc when calculating the center.
 	}
 
-	// TODO: calculate distance and direction based on camera offset to focus (similar to viewer camera)
-	Distance = ClampDistance(Distance);
-	Pitch = ClampPitch(Pitch);
-	Yaw = ClampYaw(Yaw);
+	// Ensure that the camera is initially aimed at the focus-position
+	SetActorRotation(GetLookAtRotation(FocusPosition));
+
+	const FVector Position = GetActorLocation();
+	const FRotator Rotation = GetActorRotation();
+
+	// Calculate values based on the current location and orientation
+	Distance = ClampDistance((FocusPosition - Position).Size());
+	Pitch = ClampPitch(Rotation.Pitch);
+	Yaw = ClampYaw(Rotation.Yaw);
 	TargetDistance = Distance;
 	TargetPitch = Pitch;
 	TargetYaw = Yaw;
@@ -151,3 +157,10 @@ void AGLTFOrbitCameraActor::RemoveInertia()
 	Pitch = TargetPitch;
 	Distance = TargetDistance;
 }
+
+FRotator AGLTFOrbitCameraActor::GetLookAtRotation(const FVector TargetPosition) const
+{
+	const FVector EyePosition = GetActorLocation();
+
+	return FRotationMatrix::MakeFromXZ(TargetPosition - EyePosition, FVector::UpVector).Rotator();
+}
diff --git a/Engine/Plugins/Enterprise/GLTFExporter/MigratedFromGit/Source/GLTFExporterRuntime/Public/Actors/GLTFOrbitCameraActor.h b/Engine/Plugins/Enterprise/GLTFExporter/MigratedFromGit/Source/GLTFExporterRuntime/Public/Actors/GLTFOrbitCameraActor.h
--- a/Engine/Plugins/Enterprise/GLTFExporter/MigratedFromGit/Source/GLTFExporterRuntime/Public/Actors/GLTFOrbitCameraActor.h
+++ b/Engine/Plugins/Enterprise/GLTFExporter/MigratedFromGit/Source/GLTFExporterRuntime/Public/Actors/GLTFOrbitCameraActor.h
@@ -44,6 +44,8 @@ private:
 
 	void RemoveInertia();
 
+	FRotator GetLookAtRotation(const FVector TargetPosition) const;
+
 public:
 	/* Actor which the camera will focus on and subsequently orbit. */
 	UPROPERTY(EditAnywhere, BlueprintReadWrite, Category = "GLTF Orbit Camera Actor")
