[GLTFExporter] fix(orbit-camera): prevent wheel-axis from resetting dolly-values every frame

commit 3cd3f2ed212f2b62dc452913343d899d25691aa0
Date: Wed, 23 Sep 2020 08:41:06 +0200

    fix(orbit-camera): prevent wheel-axis from resetting dolly-values every frame
    OnMouseWheelAxis is called every frame regardless of whether the axis has changed
    (cherry picked from commit 47dad3b36572039613afd1ad14989e619954e593​)

#skipci

diff --git a/Engine/Plugins/Enterprise/GLTFExporter/MigratedFromGit/Source/GLTFExporterRuntime/Private/Actors/GLTFOrbitCameraActor.cpp b/Engine/Plugins/Enterprise/GLTFExporter/MigratedFromGit/Source/GLTFExporterRuntime/Private/Actors/GLTFOrbitCameraActor.cpp
--- a/Engine/Plugins/Enterprise/GLTFExporter/MigratedFromGit/Source/GLTFExporterRuntime/Private/Actors/GLTFOrbitCameraActor.cpp
+++ b/Engine/Plugins/Enterprise/GLTFExporter/MigratedFromGit/Source/GLTFExporterRuntime/Private/Actors/GLTFOrbitCameraActor.cpp
@@ -128,9 +128,12 @@ void AGLTFOrbitCameraActor::OnMouseY(float AxisValue)
 
 void AGLTFOrbitCameraActor::OnMouseWheelAxis(float AxisValue)
 {
-	DollyTime = DollyDuration;
-	TargetDistance = ClampDistance(TargetDistance + -AxisValue * DistanceSensitivity);
-	DollyStartDistance = Distance;
+	if (!FMath::IsNearlyZero(AxisValue))
+	{
+		DollyTime = DollyDuration;
+		TargetDistance = ClampDistance(TargetDistance + -AxisValue * DistanceSensitivity);
+		DollyStartDistance = Distance;
+	}
 }
 
 float AGLTFOrbitCameraActor::ClampDistance(float Value) const
