[GLTFExporter] feat(orbit-camera): make dolly-zoom feel more like the viewer with the same settings

commit f7b20559dda41054ad2a90034d7f31261353993d
Date: Wed, 23 Sep 2020 10:45:07 +0200

    feat(orbit-camera): make dolly-zoom feel more like the viewer with the same settings
    (cherry picked from commit 7d652c3246b8f6b51aca7342677a07055cf2ffb9​)

#skipci

diff --git a/Engine/Plugins/Enterprise/GLTFExporter/MigratedFromGit/Source/GLTFExporterRuntime/Private/Actors/GLTFOrbitCameraActor.cpp b/Engine/Plugins/Enterprise/GLTFExporter/MigratedFromGit/Source/GLTFExporterRuntime/Private/Actors/GLTFOrbitCameraActor.cpp
--- a/Engine/Plugins/Enterprise/GLTFExporter/MigratedFromGit/Source/GLTFExporterRuntime/Private/Actors/GLTFOrbitCameraActor.cpp
+++ b/Engine/Plugins/Enterprise/GLTFExporter/MigratedFromGit/Source/GLTFExporterRuntime/Private/Actors/GLTFOrbitCameraActor.cpp
@@ -20,10 +20,10 @@ AGLTFOrbitCameraActor::AGLTFOrbitCameraActor(const FObjectInitializer& ObjectIni
 	, DistanceMax(1000.0f)
 	, PitchAngleMin(-90.0f)
 	, PitchAngleMax(90.0f)
-	, DollyDuration(0.4f)
+	, DollyDuration(0.2f)
 	, OrbitInertia(0.1f)
 	, OrbitSensitivity(5.0f)
-	, DistanceSensitivity(50.0f)
+	, DistanceSensitivity(0.5f)
 	, FocusPosition(0.0f, 0.0f, 0.0f)
 	, Distance(0.0f)
 	, Pitch(0.0f)
@@ -83,7 +83,7 @@ void AGLTFOrbitCameraActor::Tick(float DeltaSeconds)
 	if (DollyTime != 0.0f)
 	{
 		DollyTime = FMath::Max(DollyTime - DeltaSeconds, 0.0f);
-		Distance = FMath::InterpEaseInOut(DollyStartDistance, TargetDistance, DollyTime != 0.0f ? ((DollyDuration - DollyTime) / DollyDuration) : 1.0f, 2.0f);
+		Distance = FMath::InterpEaseInOut(DollyStartDistance, TargetDistance, (DollyDuration - DollyTime) / DollyDuration, 1.2f);
 	}
 
 	const float Alpha = (OrbitInertia == 0.0f) ? 1.0f : FMath::Min(DeltaSeconds / OrbitInertia, 1.0f);
@@ -130,8 +130,10 @@ void AGLTFOrbitCameraActor::OnMouseWheelAxis(float AxisValue)
 {
 	if (!FMath::IsNearlyZero(AxisValue))
 	{
+		const float DeltaDistance = -AxisValue * (TargetDistance * DistanceSensitivity * 0.1f);
+
 		DollyTime = DollyDuration;
-		TargetDistance = ClampDistance(TargetDistance + -AxisValue * DistanceSensitivity);
+		TargetDistance = ClampDistance(TargetDistance + DeltaDistance);
 		DollyStartDistance = Distance;
 	}
 }
