[GLTFExporter] feat(orbit-camera): make orbit feel the same as the viewer with identical OrbitSensitivity

commit 68a1e19d92f453230e9c3721409b3d90e46b47b6
Date: Wed, 23 Sep 2020 11:03:11 +0200

    feat(orbit-camera): make orbit feel the same as the viewer with identical OrbitSensitivity
    (cherry picked from commit f48d6aebad8a8acbef611fa464953fffbc5c70be​)

#skipci

diff --git a/Engine/Plugins/Enterprise/GLTFExporter/MigratedFromGit/Source/GLTFExporterRuntime/Private/Actors/GLTFOrbitCameraActor.cpp b/Engine/Plugins/Enterprise/GLTFExporter/MigratedFromGit/Source/GLTFExporterRuntime/Private/Actors/GLTFOrbitCameraActor.cpp
--- a/Engine/Plugins/Enterprise/GLTFExporter/MigratedFromGit/Source/GLTFExporterRuntime/Private/Actors/GLTFOrbitCameraActor.cpp
+++ b/Engine/Plugins/Enterprise/GLTFExporter/MigratedFromGit/Source/GLTFExporterRuntime/Private/Actors/GLTFOrbitCameraActor.cpp
@@ -11,6 +11,11 @@ namespace
 		// TODO: is this similar to viewer camera?
 		return (static_cast<int32>(Angle) / 360) + ((Angle < 0.0f) ? -1 : 0);
 	}
+
+	// Scales to convert from the export-friendly sensitivity-values stored in our properties
+	// to values that we can use when processing axis-values (to get similar results as in the viewer).
+	const float OrbitSensitivityScale = 16.667f;
+	const float DistanceSensitivityScale = 0.1f;
 }
 
 AGLTFOrbitCameraActor::AGLTFOrbitCameraActor(const FObjectInitializer& ObjectInitializer)
@@ -22,7 +27,7 @@ AGLTFOrbitCameraActor::AGLTFOrbitCameraActor(const FObjectInitializer& ObjectIni
 	, PitchAngleMax(90.0f)
 	, DollyDuration(0.2f)
 	, OrbitInertia(0.1f)
-	, OrbitSensitivity(5.0f)
+	, OrbitSensitivity(0.3f)
 	, DistanceSensitivity(0.5f)
 	, FocusPosition(0.0f, 0.0f, 0.0f)
 	, Distance(0.0f)
@@ -118,19 +123,19 @@ void AGLTFOrbitCameraActor::PreInitializeComponents()
 
 void AGLTFOrbitCameraActor::OnMouseX(float AxisValue)
 {
-	TargetYaw += AxisValue * OrbitSensitivity;
+	TargetYaw += AxisValue * OrbitSensitivity * OrbitSensitivityScale;
 }
 
 void AGLTFOrbitCameraActor::OnMouseY(float AxisValue)
 {
-	TargetPitch = ClampPitch(TargetPitch + AxisValue * OrbitSensitivity);
+	TargetPitch = ClampPitch(TargetPitch + AxisValue * OrbitSensitivity * OrbitSensitivityScale);
 }
 
 void AGLTFOrbitCameraActor::OnMouseWheelAxis(float AxisValue)
 {
 	if (!FMath::IsNearlyZero(AxisValue))
 	{
-		const float DeltaDistance = -AxisValue * (TargetDistance * DistanceSensitivity * 0.1f);
+		const float DeltaDistance = -AxisValue * (TargetDistance * DistanceSensitivity * DistanceSensitivityScale);
 
 		DollyTime = DollyDuration;
 		TargetDistance = ClampDistance(TargetDistance + DeltaDistance);
