[GLTFExporter] feat(orbit-camera): avoid using calculated center for focus, and track the focused actor

commit 9d0ab0117e40374a0df7277e58226de2924446f0
Date: Wed, 23 Sep 2020 13:42:28 +0200

    feat(orbit-camera): avoid using calculated center for focus, and track the focused actor
    (cherry picked from commit b9040793cbc85ebfa7de74b092435c6950e8dcb2​)

#skipci

diff --git a/Engine/Plugins/Enterprise/GLTFExporter/MigratedFromGit/Source/GLTFExporterRuntime/Private/Actors/GLTFOrbitCameraActor.cpp b/Engine/Plugins/Enterprise/GLTFExporter/MigratedFromGit/Source/GLTFExporterRuntime/Private/Actors/GLTFOrbitCameraActor.cpp
--- a/Engine/Plugins/Enterprise/GLTFExporter/MigratedFromGit/Source/GLTFExporterRuntime/Private/Actors/GLTFOrbitCameraActor.cpp
+++ b/Engine/Plugins/Enterprise/GLTFExporter/MigratedFromGit/Source/GLTFExporterRuntime/Private/Actors/GLTFOrbitCameraActor.cpp
@@ -29,7 +29,6 @@ AGLTFOrbitCameraActor::AGLTFOrbitCameraActor(const FObjectInitializer& ObjectIni
 	, OrbitInertia(0.1f)
 	, OrbitSensitivity(0.3f)
 	, DistanceSensitivity(0.5f)
-	, FocusPosition(0.0f, 0.0f, 0.0f)
 	, Distance(0.0f)
 	, Pitch(0.0f)
 	, Yaw(0.0f)
@@ -46,17 +45,12 @@ void AGLTFOrbitCameraActor::BeginPlay()
 {
 	Super::BeginPlay();
 
-	if (this->Focus != nullptr && this->Focus != this)
-	{
-		const FBox BoundingBox = this->Focus->GetComponentsBoundingBox(true, true);
-		FocusPosition = BoundingBox.IsValid ? BoundingBox.GetCenter() : this->Focus->GetActorLocation();
-	}
-	else
+	bool bHasValidFocusActor;
+	const FVector FocusPosition = GetFocusPosition(&bHasValidFocusActor);
+
+	if (!bHasValidFocusActor)
 	{
 		UE_LOG(LogTemp, Warning, TEXT("The camera focus must not be null, and must not be the camera's own actor"));
-
-		// TODO: use the scene center (similar to viewer camera) instead of using the zero-vector.
-		// It may however prove difficult, since we would need to exclude sky-spheres, backdrops etc when calculating the center.
 	}
 
 	// Ensure that the camera is initially aimed at the focus-position
@@ -106,7 +100,7 @@ void AGLTFOrbitCameraActor::Tick(float DeltaSeconds)
 		TargetYaw = ClampYaw(TargetYaw);
 	}
 
-	const FTransform FocusTransform = FTransform(FocusPosition);
+	const FTransform FocusTransform = FTransform(GetFocusPosition());
 	const FTransform DollyTransform = FTransform(-FVector::ForwardVector * Distance);
 	const FTransform RotationTransform = FTransform(FQuat::MakeFromEuler(FVector(0.0f, Pitch, Yaw)));
 	const FTransform ResultTransform = DollyTransform * RotationTransform * FocusTransform;
@@ -174,3 +168,14 @@ FRotator AGLTFOrbitCameraActor::GetLookAtRotation(const FVector TargetPosition)
 
 	return FRotationMatrix::MakeFromXZ(TargetPosition - EyePosition, FVector::UpVector).Rotator();
 }
+
+FVector AGLTFOrbitCameraActor::GetFocusPosition(bool* bOutHasValidFocusActor) const
+{
+	const bool bHasValidFocusActor = this->Focus != nullptr && this->Focus != this;
+	if (bOutHasValidFocusActor != nullptr)
+	{
+		*bOutHasValidFocusActor = bHasValidFocusActor;
+	}
+
+	return bHasValidFocusActor ? this->Focus->GetActorLocation() : FVector::ZeroVector;
+}
diff --git a/Engine/Plugins/Enterprise/GLTFExporter/MigratedFromGit/Source/GLTFExporterRuntime/Public/Actors/GLTFOrbitCameraActor.h b/Engine/Plugins/Enterprise/GLTFExporter/MigratedFromGit/Source/GLTFExporterRuntime/Public/Actors/GLTFOrbitCameraActor.h
--- a/Engine/Plugins/Enterprise/GLTFExporter/MigratedFromGit/Source/GLTFExporterRuntime/Public/Actors/GLTFOrbitCameraActor.h
+++ b/Engine/Plugins/Enterprise/GLTFExporter/MigratedFromGit/Source/GLTFExporterRuntime/Public/Actors/GLTFOrbitCameraActor.h
@@ -46,6 +46,8 @@ private:
 
 	FRotator GetLookAtRotation(const FVector TargetPosition) const;
 
+	FVector GetFocusPosition(bool* bOutHasValidFocusActor = nullptr) const;
+
 public:
 	/* Actor which the camera will focus on and subsequently orbit. */
 	UPROPERTY(EditAnywhere, BlueprintReadWrite, Category = "GLTF Orbit Camera Actor")
@@ -84,7 +86,6 @@ public:
 	float DistanceSensitivity;
 
 private:
-	FVector FocusPosition;
 	float Distance;
 	float Pitch;
 	float Yaw;
