[GLTFExporter] feat(orbit-camera): make min- and max-pitch use similar values as the viewer

commit 43a8f1c474c6850e5cb7dadece640a15d9b0e351
Date: Wed, 23 Sep 2020 14:04:17 +0200

    feat(orbit-camera): make min- and max-pitch use similar values as the viewer
    (cherry picked from commit 0dbb044762d86890b7fe1943d83605ad2dd01378​)

#skipci

diff --git a/Engine/Plugins/Enterprise/GLTFExporter/MigratedFromGit/Source/GLTFExporter/Private/Converters/GLTFCameraConverters.cpp b/Engine/Plugins/Enterprise/GLTFExporter/MigratedFromGit/Source/GLTFExporter/Private/Converters/GLTFCameraConverters.cpp
--- a/Engine/Plugins/Enterprise/GLTFExporter/MigratedFromGit/Source/GLTFExporter/Private/Converters/GLTFCameraConverters.cpp
+++ b/Engine/Plugins/Enterprise/GLTFExporter/MigratedFromGit/Source/GLTFExporter/Private/Converters/GLTFCameraConverters.cpp
@@ -44,11 +44,8 @@ FGLTFJsonCameraIndex FGLTFCameraConverter::Convert(const UCameraComponent* Camer
 				OrbitCamera.Focus = Builder.GetOrAddNode(OrbitCameraActor->Focus);
 				OrbitCamera.MaxDistance = FGLTFConverterUtility::ConvertLength(OrbitCameraActor->DistanceMax, Builder.ExportOptions->ExportScale);
 				OrbitCamera.MinDistance = FGLTFConverterUtility::ConvertLength(OrbitCameraActor->DistanceMin, Builder.ExportOptions->ExportScale);
-
-				// TODO: should we modify AGLTFOrbitCameraActor to calculate pitch the same way
-				// as the viewer, to avoid having to manually add 90 degrees at export?
-				OrbitCamera.MaxAngle = OrbitCameraActor->PitchAngleMax + 90.0f;
-				OrbitCamera.MinAngle = OrbitCameraActor->PitchAngleMin + 90.0f;
+				OrbitCamera.MaxAngle = OrbitCameraActor->PitchAngleMax;
+				OrbitCamera.MinAngle = OrbitCameraActor->PitchAngleMin;
 
 				Camera.OrbitCamera = OrbitCamera;
 			}
diff --git a/Engine/Plugins/Enterprise/GLTFExporter/MigratedFromGit/Source/GLTFExporterRuntime/Private/Actors/GLTFOrbitCameraActor.cpp b/Engine/Plugins/Enterprise/GLTFExporter/MigratedFromGit/Source/GLTFExporterRuntime/Private/Actors/GLTFOrbitCameraActor.cpp
--- a/Engine/Plugins/Enterprise/GLTFExporter/MigratedFromGit/Source/GLTFExporterRuntime/Private/Actors/GLTFOrbitCameraActor.cpp
+++ b/Engine/Plugins/Enterprise/GLTFExporter/MigratedFromGit/Source/GLTFExporterRuntime/Private/Actors/GLTFOrbitCameraActor.cpp
@@ -102,7 +102,7 @@ void AGLTFOrbitCameraActor::Tick(float DeltaSeconds)
 
 	const FTransform FocusTransform = FTransform(GetFocusPosition());
 	const FTransform DollyTransform = FTransform(-FVector::ForwardVector * Distance);
-	const FTransform RotationTransform = FTransform(FQuat::MakeFromEuler(FVector(0.0f, Pitch, Yaw)));
+	const FTransform RotationTransform = FTransform(FQuat::MakeFromEuler(FVector(0.0f, -Pitch, Yaw)));
 	const FTransform ResultTransform = DollyTransform * RotationTransform * FocusTransform;
 
 	SetActorTransform(ResultTransform);
@@ -122,7 +122,7 @@ void AGLTFOrbitCameraActor::OnMouseX(float AxisValue)
 
 void AGLTFOrbitCameraActor::OnMouseY(float AxisValue)
 {
-	TargetPitch = ClampPitch(TargetPitch + AxisValue * OrbitSensitivity * OrbitSensitivityScale);
+	TargetPitch = ClampPitch(TargetPitch - AxisValue * OrbitSensitivity * OrbitSensitivityScale);
 }
 
 void AGLTFOrbitCameraActor::OnMouseWheelAxis(float AxisValue)
