[GLTFExporter] feat(orbit-camera): add conversion of more camera-settings

commit 7d1a76161e642a95927f2cc5f20f4e517dbd0851
Date: Wed, 23 Sep 2020 14:54:57 +0200

    feat(orbit-camera): add conversion of more camera-settings
    (cherry picked from commit f8bfb3059ba75997a971d9c05abb311563fb5b68​)

#skipci

diff --git a/Engine/Plugins/Enterprise/GLTFExporter/MigratedFromGit/Source/GLTFExporter/Private/Converters/GLTFCameraConverters.cpp b/Engine/Plugins/Enterprise/GLTFExporter/MigratedFromGit/Source/GLTFExporter/Private/Converters/GLTFCameraConverters.cpp
--- a/Engine/Plugins/Enterprise/GLTFExporter/MigratedFromGit/Source/GLTFExporter/Private/Converters/GLTFCameraConverters.cpp
+++ b/Engine/Plugins/Enterprise/GLTFExporter/MigratedFromGit/Source/GLTFExporter/Private/Converters/GLTFCameraConverters.cpp
@@ -46,6 +46,10 @@ FGLTFJsonCameraIndex FGLTFCameraConverter::Convert(const UCameraComponent* Camer
 				OrbitCamera.MinDistance = FGLTFConverterUtility::ConvertLength(OrbitCameraActor->DistanceMin, Builder.ExportOptions->ExportScale);
 				OrbitCamera.MaxAngle = OrbitCameraActor->PitchAngleMax;
 				OrbitCamera.MinAngle = OrbitCameraActor->PitchAngleMin;
+				OrbitCamera.DistanceSensitivity = OrbitCameraActor->DistanceSensitivity;
+				OrbitCamera.OrbitSensitivity = OrbitCameraActor->OrbitSensitivity;
+				OrbitCamera.OrbitInertia = OrbitCameraActor->OrbitInertia;
+				OrbitCamera.DollyDuration = OrbitCameraActor->DollyDuration;
 
 				Camera.OrbitCamera = OrbitCamera;
 			}
diff --git a/Engine/Plugins/Enterprise/GLTFExporter/MigratedFromGit/Source/GLTFExporter/Private/Json/GLTFJsonOrbitCamera.h b/Engine/Plugins/Enterprise/GLTFExporter/MigratedFromGit/Source/GLTFExporter/Private/Json/GLTFJsonOrbitCamera.h
--- a/Engine/Plugins/Enterprise/GLTFExporter/MigratedFromGit/Source/GLTFExporter/Private/Json/GLTFJsonOrbitCamera.h
+++ b/Engine/Plugins/Enterprise/GLTFExporter/MigratedFromGit/Source/GLTFExporter/Private/Json/GLTFJsonOrbitCamera.h
@@ -12,6 +12,10 @@ struct FGLTFJsonOrbitCamera
 	float              MinDistance;
 	float              MaxAngle;
 	float              MinAngle;
+	float              DistanceSensitivity;
+	float              OrbitSensitivity;
+	float              OrbitInertia;
+	float              DollyDuration;
 
 	template <class CharType = TCHAR, class PrintPolicy = TPrettyJsonPrintPolicy<CharType>>
 	void WriteObject(TJsonWriter<CharType, PrintPolicy>& JsonWriter, FGLTFJsonExtensions& Extensions) const
@@ -23,6 +27,10 @@ struct FGLTFJsonOrbitCamera
 		FGLTFJsonUtility::WriteExactValue(JsonWriter, TEXT("minDistance"), MinDistance);
 		FGLTFJsonUtility::WriteExactValue(JsonWriter, TEXT("maxAngle"), MaxAngle);
 		FGLTFJsonUtility::WriteExactValue(JsonWriter, TEXT("minAngle"), MinAngle);
+		FGLTFJsonUtility::WriteExactValue(JsonWriter, TEXT("distanceSensitivity"), DistanceSensitivity);
+		FGLTFJsonUtility::WriteExactValue(JsonWriter, TEXT("orbitSensitivity"), OrbitSensitivity);
+		FGLTFJsonUtility::WriteExactValue(JsonWriter, TEXT("orbitInertia"), OrbitInertia);
+		FGLTFJsonUtility::WriteExactValue(JsonWriter, TEXT("dollyDuration"), DollyDuration);
 
 		JsonWriter.WriteObjectEnd();
 	}
