[GLTFExporter] fix(orbit-camera): fix flipped pitch-value when aiming towards the focused actor

commit dc0f6748388a8243e96a58bc6cdf34cd4521f515
Date: Wed, 23 Sep 2020 16:18:02 +0200

    fix(orbit-camera): fix flipped pitch-value when aiming towards the focused actor
    (cherry picked from commit 3c347ad05ead1659ea6d47cfc11e090e4ace41bb​)

#skipci

diff --git a/Engine/Plugins/Enterprise/GLTFExporter/MigratedFromGit/Source/GLTFExporterRuntime/Private/Actors/GLTFOrbitCameraActor.cpp b/Engine/Plugins/Enterprise/GLTFExporter/MigratedFromGit/Source/GLTFExporterRuntime/Private/Actors/GLTFOrbitCameraActor.cpp
--- a/Engine/Plugins/Enterprise/GLTFExporter/MigratedFromGit/Source/GLTFExporterRuntime/Private/Actors/GLTFOrbitCameraActor.cpp
+++ b/Engine/Plugins/Enterprise/GLTFExporter/MigratedFromGit/Source/GLTFExporterRuntime/Private/Actors/GLTFOrbitCameraActor.cpp
@@ -61,7 +61,7 @@ void AGLTFOrbitCameraActor::BeginPlay()
 
 	// Calculate values based on the current location and orientation
 	Distance = ClampDistance((FocusPosition - Position).Size());
-	Pitch = ClampPitch(Rotation.Pitch);
+	Pitch = ClampPitch(-Rotation.Pitch);
 	Yaw = ClampYaw(Rotation.Yaw);
 	TargetDistance = Distance;
 	TargetPitch = Pitch;
