[GLTFExporter] feat(orbit-camera): set AutoActivateForPlayer to Player0 when the actor is created

commit 8faa8555484e1724e0184ddcf6cc5d51d7c593b8
Date: Wed, 23 Sep 2020 16:19:21 +0200

    feat(orbit-camera): set AutoActivateForPlayer to Player0 when the actor is created
    (cherry picked from commit 9fe8e09bff249c5197bfcb6362ea69837d6f04c7​)

#skipci

diff --git a/Engine/Plugins/Enterprise/GLTFExporter/MigratedFromGit/Source/GLTFExporterRuntime/Private/Actors/GLTFOrbitCameraActor.cpp b/Engine/Plugins/Enterprise/GLTFExporter/MigratedFromGit/Source/GLTFExporterRuntime/Private/Actors/GLTFOrbitCameraActor.cpp
--- a/Engine/Plugins/Enterprise/GLTFExporter/MigratedFromGit/Source/GLTFExporterRuntime/Private/Actors/GLTFOrbitCameraActor.cpp
+++ b/Engine/Plugins/Enterprise/GLTFExporter/MigratedFromGit/Source/GLTFExporterRuntime/Private/Actors/GLTFOrbitCameraActor.cpp
@@ -115,6 +115,11 @@ void AGLTFOrbitCameraActor::PreInitializeComponents()
 	Super::PreInitializeComponents();
 }
 
+void AGLTFOrbitCameraActor::PostActorCreated()
+{
+	SetAutoActivateForPlayer(EAutoReceiveInput::Player0);
+}
+
 void AGLTFOrbitCameraActor::OnMouseX(float AxisValue)
 {
 	TargetYaw += AxisValue * OrbitSensitivity * OrbitSensitivityScale;
@@ -179,3 +184,22 @@ FVector AGLTFOrbitCameraActor::GetFocusPosition(bool* bOutHasValidFocusActor) co
 
 	return bHasValidFocusActor ? this->Focus->GetActorLocation() : FVector::ZeroVector;
 }
+
+bool AGLTFOrbitCameraActor::SetAutoActivateForPlayer(const EAutoReceiveInput::Type Player)
+{
+	// TODO: remove hack by adding proper API access to ACameraActor (or by other means)
+	FProperty* Property = GetClass()->FindPropertyByName(TEXT("AutoActivateForPlayer"));
+	if (Property == nullptr)
+	{
+		return false;
+	}
+
+	TEnumAsByte<EAutoReceiveInput::Type>* ValuePtr = Property->ContainerPtrToValuePtr<TEnumAsByte<EAutoReceiveInput::Type>>(this);
+	if (ValuePtr == nullptr)
+	{
+		return false;
+	}
+
+	*ValuePtr = Player;
+	return true;
+}
diff --git a/Engine/Plugins/Enterprise/GLTFExporter/MigratedFromGit/Source/GLTFExporterRuntime/Public/Actors/GLTFOrbitCameraActor.h b/Engine/Plugins/Enterprise/GLTFExporter/MigratedFromGit/Source/GLTFExporterRuntime/Public/Actors/GLTFOrbitCameraActor.h
--- a/Engine/Plugins/Enterprise/GLTFExporter/MigratedFromGit/Source/GLTFExporterRuntime/Public/Actors/GLTFOrbitCameraActor.h
+++ b/Engine/Plugins/Enterprise/GLTFExporter/MigratedFromGit/Source/GLTFExporterRuntime/Public/Actors/GLTFOrbitCameraActor.h
@@ -24,6 +24,7 @@ protected:
 public:
 	virtual void Tick(float DeltaSeconds) override;
 	virtual void PreInitializeComponents() override;
+	virtual void PostActorCreated() override;
 	//~ End AActor Interface
 
 private:
@@ -48,6 +49,8 @@ private:
 
 	FVector GetFocusPosition(bool* bOutHasValidFocusActor = nullptr) const;
 
+	bool SetAutoActivateForPlayer(const EAutoReceiveInput::Type Player);
+
 public:
 	/* Actor which the camera will focus on and subsequently orbit. */
 	UPROPERTY(EditAnywhere, BlueprintReadWrite, Category = "GLTF Orbit Camera Actor")
