[GLTFExporter] refactor: rename prefix ("F" to "T") for template classes FGLTFTextureConverter, FGLTFMeshConverter, FGLTFAccessorConverter, and FGLTFNodeConverter

commit bde5d1e057284c6b564e3cf7ba73e2bcbeca1df2
Date: Tue, 22 Sep 2020 15:21:52 +0200

    refactor: rename prefix ("F" to "T") for template classes FGLTFTextureConverter, FGLTFMeshConverter, FGLTFAccessorConverter, and FGLTFNodeConverter
    (cherry picked from commit 8d4938677f5870d6b50a15927ee8a6e1c346c01e​)

#skipci

diff --git a/Engine/Plugins/Enterprise/GLTFExporter/MigratedFromGit/Source/GLTFExporter/Private/Converters/GLTFAccessorConverters.h b/Engine/Plugins/Enterprise/GLTFExporter/MigratedFromGit/Source/GLTFExporter/Private/Converters/GLTFAccessorConverters.h
--- a/Engine/Plugins/Enterprise/GLTFExporter/MigratedFromGit/Source/GLTFExporter/Private/Converters/GLTFAccessorConverters.h
+++ b/Engine/Plugins/Enterprise/GLTFExporter/MigratedFromGit/Source/GLTFExporter/Private/Converters/GLTFAccessorConverters.h
@@ -11,49 +11,49 @@
 class FMultiSizeIndexContainer;
 
 template <typename... InputTypes>
-class FGLTFAccessorConverter : public FGLTFBuilderContext, public TGLTFConverter<FGLTFJsonAccessorIndex, InputTypes...>
+class TGLTFAccessorConverter : public FGLTFBuilderContext, public TGLTFConverter<FGLTFJsonAccessorIndex, InputTypes...>
 {
 	using FGLTFBuilderContext::FGLTFBuilderContext;
 };
 
-class FGLTFPositionBufferConverter final : public FGLTFAccessorConverter<const FPositionVertexBuffer*>
+class FGLTFPositionBufferConverter final : public TGLTFAccessorConverter<const FPositionVertexBuffer*>
 {
-	using FGLTFAccessorConverter::FGLTFAccessorConverter;
+	using TGLTFAccessorConverter::TGLTFAccessorConverter;
 
 	virtual FGLTFJsonAccessorIndex Convert(const FPositionVertexBuffer* VertexBuffer) override;
 };
 
-class FGLTFColorBufferConverter final : public FGLTFAccessorConverter<const FColorVertexBuffer*>
+class FGLTFColorBufferConverter final : public TGLTFAccessorConverter<const FColorVertexBuffer*>
 {
-	using FGLTFAccessorConverter::FGLTFAccessorConverter;
+	using TGLTFAccessorConverter::TGLTFAccessorConverter;
 
 	virtual FGLTFJsonAccessorIndex Convert(const FColorVertexBuffer* VertexBuffer) override;
 };
 
-class FGLTFNormalBufferConverter final : public FGLTFAccessorConverter<const FStaticMeshVertexBuffer*>
+class FGLTFNormalBufferConverter final : public TGLTFAccessorConverter<const FStaticMeshVertexBuffer*>
 {
-	using FGLTFAccessorConverter::FGLTFAccessorConverter;
+	using TGLTFAccessorConverter::TGLTFAccessorConverter;
 
 	virtual FGLTFJsonAccessorIndex Convert(const FStaticMeshVertexBuffer* VertexBuffer) override;
 };
 
-class FGLTFTangentBufferConverter final : public FGLTFAccessorConverter<const FStaticMeshVertexBuffer*>
+class FGLTFTangentBufferConverter final : public TGLTFAccessorConverter<const FStaticMeshVertexBuffer*>
 {
-	using FGLTFAccessorConverter::FGLTFAccessorConverter;
+	using TGLTFAccessorConverter::TGLTFAccessorConverter;
 
 	virtual FGLTFJsonAccessorIndex Convert(const FStaticMeshVertexBuffer* VertexBuffer) override;
 };
 
-class FGLTFUVBufferConverter final : public FGLTFAccessorConverter<const FStaticMeshVertexBuffer*, int32>
+class FGLTFUVBufferConverter final : public TGLTFAccessorConverter<const FStaticMeshVertexBuffer*, int32>
 {
-	using FGLTFAccessorConverter::FGLTFAccessorConverter;
+	using TGLTFAccessorConverter::TGLTFAccessorConverter;
 
 	virtual FGLTFJsonAccessorIndex Convert(const FStaticMeshVertexBuffer* VertexBuffer, int32 UVIndex) override;
 };
 
-class FGLTFBoneIndexBufferConverter final : public FGLTFAccessorConverter<const FSkinWeightVertexBuffer*, int32, FGLTFBoneMap>
+class FGLTFBoneIndexBufferConverter final : public TGLTFAccessorConverter<const FSkinWeightVertexBuffer*, int32, FGLTFBoneMap>
 {
-	using FGLTFAccessorConverter::FGLTFAccessorConverter;
+	using TGLTFAccessorConverter::TGLTFAccessorConverter;
 
 	virtual FGLTFJsonAccessorIndex Convert(const FSkinWeightVertexBuffer* VertexBuffer, int32 InfluenceOffset, FGLTFBoneMap BoneMap) override;
 
@@ -61,23 +61,23 @@ class FGLTFBoneIndexBufferConverter final : public FGLTFAccessorConverter<const
 	FGLTFJsonAccessorIndex Convert(const FSkinWeightVertexBuffer* VertexBuffer, int32 JointsGroupIndex, FGLTFBoneMap BoneMap);
 };
 
-class FGLTFBoneWeightBufferConverter final : public FGLTFAccessorConverter<const FSkinWeightVertexBuffer*, int32>
+class FGLTFBoneWeightBufferConverter final : public TGLTFAccessorConverter<const FSkinWeightVertexBuffer*, int32>
 {
-	using FGLTFAccessorConverter::FGLTFAccessorConverter;
+	using TGLTFAccessorConverter::TGLTFAccessorConverter;
 
 	virtual FGLTFJsonAccessorIndex Convert(const FSkinWeightVertexBuffer* VertexBuffer, int32 WeightsGroupIndex) override;
 };
 
-class FGLTFStaticMeshSectionConverter final : public FGLTFAccessorConverter<const FStaticMeshSection*, const FRawStaticIndexBuffer*>
+class FGLTFStaticMeshSectionConverter final : public TGLTFAccessorConverter<const FStaticMeshSection*, const FRawStaticIndexBuffer*>
 {
-	using FGLTFAccessorConverter::FGLTFAccessorConverter;
+	using TGLTFAccessorConverter::TGLTFAccessorConverter;
 
 	virtual FGLTFJsonAccessorIndex Convert(const FStaticMeshSection* MeshSection, const FRawStaticIndexBuffer* IndexBuffer) override;
 };
 
-class FGLTFSkeletalMeshSectionConverter final : public FGLTFAccessorConverter<const FSkelMeshRenderSection*, const FMultiSizeIndexContainer*>
+class FGLTFSkeletalMeshSectionConverter final : public TGLTFAccessorConverter<const FSkelMeshRenderSection*, const FMultiSizeIndexContainer*>
 {
-	using FGLTFAccessorConverter::FGLTFAccessorConverter;
+	using TGLTFAccessorConverter::TGLTFAccessorConverter;
 
 	virtual FGLTFJsonAccessorIndex Convert(const FSkelMeshRenderSection* MeshSection, const FMultiSizeIndexContainer* IndexContainer) override;
 };
diff --git a/Engine/Plugins/Enterprise/GLTFExporter/MigratedFromGit/Source/GLTFExporter/Private/Converters/GLTFMeshConverters.h b/Engine/Plugins/Enterprise/GLTFExporter/MigratedFromGit/Source/GLTFExporter/Private/Converters/GLTFMeshConverters.h
--- a/Engine/Plugins/Enterprise/GLTFExporter/MigratedFromGit/Source/GLTFExporter/Private/Converters/GLTFMeshConverters.h
+++ b/Engine/Plugins/Enterprise/GLTFExporter/MigratedFromGit/Source/GLTFExporter/Private/Converters/GLTFMeshConverters.h
@@ -9,21 +9,21 @@
 #include "Engine.h"
 
 template <typename... InputTypes>
-class FGLTFMeshConverter : public FGLTFBuilderContext, public TGLTFConverter<FGLTFJsonMeshIndex, InputTypes...>
+class TGLTFMeshConverter : public FGLTFBuilderContext, public TGLTFConverter<FGLTFJsonMeshIndex, InputTypes...>
 {
 	using FGLTFBuilderContext::FGLTFBuilderContext;
 };
 
-class FGLTFStaticMeshConverter final : public FGLTFMeshConverter<const UStaticMesh*, int32, const FColorVertexBuffer*, FGLTFMaterialArray>
+class FGLTFStaticMeshConverter final : public TGLTFMeshConverter<const UStaticMesh*, int32, const FColorVertexBuffer*, FGLTFMaterialArray>
 {
-	using FGLTFMeshConverter::FGLTFMeshConverter;
+	using TGLTFMeshConverter::TGLTFMeshConverter;
 
 	virtual FGLTFJsonMeshIndex Convert(const UStaticMesh* StaticMesh, int32 LODIndex, const FColorVertexBuffer* OverrideVertexColors, FGLTFMaterialArray OverrideMaterials) override;
 };
 
-class FGLTFSkeletalMeshConverter final : public FGLTFMeshConverter<const USkeletalMesh*, int32, const FColorVertexBuffer*, const FSkinWeightVertexBuffer*, FGLTFMaterialArray>
+class FGLTFSkeletalMeshConverter final : public TGLTFMeshConverter<const USkeletalMesh*, int32, const FColorVertexBuffer*, const FSkinWeightVertexBuffer*, FGLTFMaterialArray>
 {
-	using FGLTFMeshConverter::FGLTFMeshConverter;
+	using TGLTFMeshConverter::TGLTFMeshConverter;
 
 	virtual FGLTFJsonMeshIndex Convert(const USkeletalMesh* SkeletalMesh, int32 LODIndex, const FColorVertexBuffer* OverrideVertexColors, const FSkinWeightVertexBuffer* OverrideSkinWeights, FGLTFMaterialArray OverrideMaterials) override;
 };
diff --git a/Engine/Plugins/Enterprise/GLTFExporter/MigratedFromGit/Source/GLTFExporter/Private/Converters/GLTFNodeConverters.h b/Engine/Plugins/Enterprise/GLTFExporter/MigratedFromGit/Source/GLTFExporter/Private/Converters/GLTFNodeConverters.h
--- a/Engine/Plugins/Enterprise/GLTFExporter/MigratedFromGit/Source/GLTFExporter/Private/Converters/GLTFNodeConverters.h
+++ b/Engine/Plugins/Enterprise/GLTFExporter/MigratedFromGit/Source/GLTFExporter/Private/Converters/GLTFNodeConverters.h
@@ -8,21 +8,21 @@
 #include "Engine.h"
 
 template <typename... InputTypes>
-class FGLTFNodeConverter : public FGLTFBuilderContext, public TGLTFConverter<FGLTFJsonNodeIndex, InputTypes...>
+class TGLTFNodeConverter : public FGLTFBuilderContext, public TGLTFConverter<FGLTFJsonNodeIndex, InputTypes...>
 {
 	using FGLTFBuilderContext::FGLTFBuilderContext;
 };
 
-class FGLTFComponentConverter final : public FGLTFNodeConverter<const USceneComponent*>
+class FGLTFComponentConverter final : public TGLTFNodeConverter<const USceneComponent*>
 {
-	using FGLTFNodeConverter::FGLTFNodeConverter;
+	using TGLTFNodeConverter::TGLTFNodeConverter;
 
 	virtual FGLTFJsonNodeIndex Convert(const USceneComponent* SceneComponent) override;
 };
 
-class FGLTFActorConverter final : public FGLTFNodeConverter<const AActor*>
+class FGLTFActorConverter final : public TGLTFNodeConverter<const AActor*>
 {
-	using FGLTFNodeConverter::FGLTFNodeConverter;
+	using TGLTFNodeConverter::TGLTFNodeConverter;
 
 	virtual FGLTFJsonNodeIndex Convert(const AActor* Actor) override;
 };
diff --git a/Engine/Plugins/Enterprise/GLTFExporter/MigratedFromGit/Source/GLTFExporter/Private/Converters/GLTFTextureConverters.h b/Engine/Plugins/Enterprise/GLTFExporter/MigratedFromGit/Source/GLTFExporter/Private/Converters/GLTFTextureConverters.h
--- a/Engine/Plugins/Enterprise/GLTFExporter/MigratedFromGit/Source/GLTFExporter/Private/Converters/GLTFTextureConverters.h
+++ b/Engine/Plugins/Enterprise/GLTFExporter/MigratedFromGit/Source/GLTFExporter/Private/Converters/GLTFTextureConverters.h
@@ -8,35 +8,35 @@
 #include "Engine.h"
 
 template <typename... InputTypes>
-class FGLTFTextureConverter : public FGLTFBuilderContext, public TGLTFConverter<FGLTFJsonTextureIndex, InputTypes...>
+class TGLTFTextureConverter : public FGLTFBuilderContext, public TGLTFConverter<FGLTFJsonTextureIndex, InputTypes...>
 {
 	using FGLTFBuilderContext::FGLTFBuilderContext;
 };
 
-class FGLTFTexture2DConverter final : public FGLTFTextureConverter<const UTexture2D*>
+class FGLTFTexture2DConverter final : public TGLTFTextureConverter<const UTexture2D*>
 {
-	using FGLTFTextureConverter::FGLTFTextureConverter;
+	using TGLTFTextureConverter::TGLTFTextureConverter;
 
 	virtual FGLTFJsonTextureIndex Convert(const UTexture2D* Texture2D) override;
 };
 
-class FGLTFTextureCubeConverter final : public FGLTFTextureConverter<const UTextureCube*, ECubeFace>
+class FGLTFTextureCubeConverter final : public TGLTFTextureConverter<const UTextureCube*, ECubeFace>
 {
-	using FGLTFTextureConverter::FGLTFTextureConverter;
+	using TGLTFTextureConverter::TGLTFTextureConverter;
 
 	virtual FGLTFJsonTextureIndex Convert(const UTextureCube* TextureCube, ECubeFace CubeFace) override;
 };
 
-class FGLTFTextureRenderTarget2DConverter final : public FGLTFTextureConverter<const UTextureRenderTarget2D*>
+class FGLTFTextureRenderTarget2DConverter final : public TGLTFTextureConverter<const UTextureRenderTarget2D*>
 {
-	using FGLTFTextureConverter::FGLTFTextureConverter;
+	using TGLTFTextureConverter::TGLTFTextureConverter;
 
 	virtual FGLTFJsonTextureIndex Convert(const UTextureRenderTarget2D* RenderTarget2D) override;
 };
 
-class FGLTFTextureRenderTargetCubeConverter final : public FGLTFTextureConverter<const UTextureRenderTargetCube*, ECubeFace>
+class FGLTFTextureRenderTargetCubeConverter final : public TGLTFTextureConverter<const UTextureRenderTargetCube*, ECubeFace>
 {
-	using FGLTFTextureConverter::FGLTFTextureConverter;
+	using TGLTFTextureConverter::TGLTFTextureConverter;
 
 	virtual FGLTFJsonTextureIndex Convert(const UTextureRenderTargetCube* RenderTargetCube, ECubeFace CubeFace) override;
 };
