[GLTFExporter] refactor: change converter input FMultiSizeIndexContainer to FRawStaticIndexBuffer16or32Interface

commit 9b17d0f30db049f6beb999ee37115d1aaf883140
Date: Tue, 22 Sep 2020 15:37:45 +0200

    refactor: change converter input FMultiSizeIndexContainer to FRawStaticIndexBuffer16or32Interface
    (cherry picked from commit 14561944626b32bc950c0871eb267165ad80fea7​)

#skipci

diff --git a/Engine/Plugins/Enterprise/GLTFExporter/MigratedFromGit/Source/GLTFExporter/Private/Builders/GLTFConvertBuilder.cpp b/Engine/Plugins/Enterprise/GLTFExporter/MigratedFromGit/Source/GLTFExporter/Private/Builders/GLTFConvertBuilder.cpp
--- a/Engine/Plugins/Enterprise/GLTFExporter/MigratedFromGit/Source/GLTFExporter/Private/Builders/GLTFConvertBuilder.cpp
+++ b/Engine/Plugins/Enterprise/GLTFExporter/MigratedFromGit/Source/GLTFExporter/Private/Builders/GLTFConvertBuilder.cpp
@@ -127,14 +127,14 @@ FGLTFJsonMeshIndex FGLTFConvertBuilder::GetOrAddMesh(const UStaticMeshComponent*
 	return GetOrAddMesh(StaticMesh, LODIndex, OverrideVertexColors, OverrideMaterials);
 }
 
-FGLTFJsonAccessorIndex FGLTFConvertBuilder::GetOrAddIndexAccessor(const FSkelMeshRenderSection* MeshSection, const FMultiSizeIndexContainer* IndexContainer)
+FGLTFJsonAccessorIndex FGLTFConvertBuilder::GetOrAddIndexAccessor(const FSkelMeshRenderSection* MeshSection, const FRawStaticIndexBuffer16or32Interface* IndexBuffer)
 {
 	if (MeshSection == nullptr)
 	{
 		return FGLTFJsonAccessorIndex(INDEX_NONE);
 	}
 
-	return SkeletalMeshSectionConverter.GetOrAdd(MeshSection, IndexContainer);
+	return SkeletalMeshSectionConverter.GetOrAdd(MeshSection, IndexBuffer);
 }
 
 FGLTFJsonMeshIndex FGLTFConvertBuilder::GetOrAddMesh(const USkeletalMesh* SkeletalMesh, int32 LODIndex, const FColorVertexBuffer* OverrideVertexColors, const FSkinWeightVertexBuffer* OverrideSkinWeights, const FGLTFMaterialArray& OverrideMaterials)
diff --git a/Engine/Plugins/Enterprise/GLTFExporter/MigratedFromGit/Source/GLTFExporter/Private/Builders/GLTFConvertBuilder.h b/Engine/Plugins/Enterprise/GLTFExporter/MigratedFromGit/Source/GLTFExporter/Private/Builders/GLTFConvertBuilder.h
--- a/Engine/Plugins/Enterprise/GLTFExporter/MigratedFromGit/Source/GLTFExporter/Private/Builders/GLTFConvertBuilder.h
+++ b/Engine/Plugins/Enterprise/GLTFExporter/MigratedFromGit/Source/GLTFExporter/Private/Builders/GLTFConvertBuilder.h
@@ -33,7 +33,7 @@ public:
 	FGLTFJsonAccessorIndex GetOrAddJointAccessor(const FSkinWeightVertexBuffer* VertexBuffer, int32 JointsGroupIndex, FGLTFBoneMap BoneMap);
 	FGLTFJsonAccessorIndex GetOrAddWeightAccessor(const FSkinWeightVertexBuffer* VertexBuffer, int32 WeightsGroupIndex);
 	FGLTFJsonAccessorIndex GetOrAddIndexAccessor(const FStaticMeshSection* MeshSection, const FRawStaticIndexBuffer* IndexBuffer);
-	FGLTFJsonAccessorIndex GetOrAddIndexAccessor(const FSkelMeshRenderSection* MeshSection, const FMultiSizeIndexContainer* IndexContainer);
+	FGLTFJsonAccessorIndex GetOrAddIndexAccessor(const FSkelMeshRenderSection* MeshSection, const FRawStaticIndexBuffer16or32Interface* IndexBuffer);
 
 	FGLTFJsonMeshIndex GetOrAddMesh(const UStaticMesh* StaticMesh, int32 LODIndex = -1, const FColorVertexBuffer* OverrideVertexColors = nullptr, const FGLTFMaterialArray& OverrideMaterials = {});
 	FGLTFJsonMeshIndex GetOrAddMesh(const UStaticMeshComponent* StaticMeshComponent);
diff --git a/Engine/Plugins/Enterprise/GLTFExporter/MigratedFromGit/Source/GLTFExporter/Private/Converters/GLTFAccessorConverters.cpp b/Engine/Plugins/Enterprise/GLTFExporter/MigratedFromGit/Source/GLTFExporter/Private/Converters/GLTFAccessorConverters.cpp
--- a/Engine/Plugins/Enterprise/GLTFExporter/MigratedFromGit/Source/GLTFExporter/Private/Converters/GLTFAccessorConverters.cpp
+++ b/Engine/Plugins/Enterprise/GLTFExporter/MigratedFromGit/Source/GLTFExporter/Private/Converters/GLTFAccessorConverters.cpp
@@ -268,7 +268,7 @@ FGLTFJsonAccessorIndex FGLTFStaticMeshSectionConverter::Convert(const FStaticMes
 	return Builder.AddAccessor(JsonAccessor);
 }
 
-FGLTFJsonAccessorIndex FGLTFSkeletalMeshSectionConverter::Convert(const FSkelMeshRenderSection* MeshSection, const FMultiSizeIndexContainer* IndexContainer)
+FGLTFJsonAccessorIndex FGLTFSkeletalMeshSectionConverter::Convert(const FSkelMeshRenderSection* MeshSection, const FRawStaticIndexBuffer16or32Interface* IndexBuffer)
 {
 	const uint32 IndexCount = MeshSection->NumTriangles * 3;
 	if (IndexCount == 0)
@@ -276,11 +276,10 @@ FGLTFJsonAccessorIndex FGLTFSkeletalMeshSectionConverter::Convert(const FSkelMes
 		return FGLTFJsonAccessorIndex(INDEX_NONE);
 	}
 
-	const uint8 IndexTypeSize = IndexContainer->GetDataTypeSize();
-	const bool bIs32Bit = IndexTypeSize == sizeof(uint32);
+	const bool bIs32Bit = IndexBuffer->GetResourceDataSize() == IndexBuffer->Num() * sizeof(uint32);
+	const uint8 IndexTypeSize = bIs32Bit ? sizeof(uint32) : sizeof(uint16);
 	const uint32 FirstIndex = MeshSection->BaseIndex;
 
-	const FRawStaticIndexBuffer16or32Interface* IndexBuffer = IndexContainer->GetIndexBuffer();
 	const void* IndexDataPtr = const_cast<FRawStaticIndexBuffer16or32Interface*>(IndexBuffer)->GetPointerTo(FirstIndex);
 	const int32 IndexDataSize = IndexCount * IndexTypeSize;
 
diff --git a/Engine/Plugins/Enterprise/GLTFExporter/MigratedFromGit/Source/GLTFExporter/Private/Converters/GLTFAccessorConverters.h b/Engine/Plugins/Enterprise/GLTFExporter/MigratedFromGit/Source/GLTFExporter/Private/Converters/GLTFAccessorConverters.h
--- a/Engine/Plugins/Enterprise/GLTFExporter/MigratedFromGit/Source/GLTFExporter/Private/Converters/GLTFAccessorConverters.h
+++ b/Engine/Plugins/Enterprise/GLTFExporter/MigratedFromGit/Source/GLTFExporter/Private/Converters/GLTFAccessorConverters.h
@@ -8,8 +8,6 @@
 #include "Converters/GLTFBoneMap.h"
 #include "Engine.h"
 
-class FMultiSizeIndexContainer;
-
 template <typename... InputTypes>
 class TGLTFAccessorConverter : public FGLTFBuilderContext, public TGLTFConverter<FGLTFJsonAccessorIndex, InputTypes...>
 {
@@ -75,9 +73,9 @@ class FGLTFStaticMeshSectionConverter final : public TGLTFAccessorConverter<cons
 	virtual FGLTFJsonAccessorIndex Convert(const FStaticMeshSection* MeshSection, const FRawStaticIndexBuffer* IndexBuffer) override;
 };
 
-class FGLTFSkeletalMeshSectionConverter final : public TGLTFAccessorConverter<const FSkelMeshRenderSection*, const FMultiSizeIndexContainer*>
+class FGLTFSkeletalMeshSectionConverter final : public TGLTFAccessorConverter<const FSkelMeshRenderSection*, const FRawStaticIndexBuffer16or32Interface*>
 {
 	using TGLTFAccessorConverter::TGLTFAccessorConverter;
 
-	virtual FGLTFJsonAccessorIndex Convert(const FSkelMeshRenderSection* MeshSection, const FMultiSizeIndexContainer* IndexContainer) override;
+	virtual FGLTFJsonAccessorIndex Convert(const FSkelMeshRenderSection* MeshSection, const FRawStaticIndexBuffer16or32Interface* IndexBuffer) override;
 };
diff --git a/Engine/Plugins/Enterprise/GLTFExporter/MigratedFromGit/Source/GLTFExporter/Private/Converters/GLTFMeshConverters.cpp b/Engine/Plugins/Enterprise/GLTFExporter/MigratedFromGit/Source/GLTFExporter/Private/Converters/GLTFMeshConverters.cpp
--- a/Engine/Plugins/Enterprise/GLTFExporter/MigratedFromGit/Source/GLTFExporter/Private/Converters/GLTFMeshConverters.cpp
+++ b/Engine/Plugins/Enterprise/GLTFExporter/MigratedFromGit/Source/GLTFExporter/Private/Converters/GLTFMeshConverters.cpp
@@ -121,7 +121,7 @@ FGLTFJsonMeshIndex FGLTFSkeletalMeshConverter::Convert(const USkeletalMesh* Skel
 	}
 	*/
 
-	const FMultiSizeIndexContainer* IndexContainer = &MeshLOD.MultiSizeIndexContainer;
+	const FRawStaticIndexBuffer16or32Interface* IndexBuffer = MeshLOD.MultiSizeIndexContainer.GetIndexBuffer();
 
 	const int32 SectionCount = MeshLOD.RenderSections.Num();
 	JsonMesh.Primitives.AddDefaulted(SectionCount);
@@ -132,7 +132,7 @@ FGLTFJsonMeshIndex FGLTFSkeletalMeshConverter::Convert(const USkeletalMesh* Skel
 		JsonPrimitive.Attributes = JsonAttributes;
 
 		const FSkelMeshRenderSection& Section = MeshLOD.RenderSections[SectionIndex];
-		JsonPrimitive.Indices = Builder.GetOrAddIndexAccessor(&Section, IndexContainer);
+		JsonPrimitive.Indices = Builder.GetOrAddIndexAccessor(&Section, IndexBuffer);
 
 		const int32 MaterialIndex = Section.MaterialIndex;
 		const UMaterialInterface* Material = OverrideMaterials.IsValidIndex(MaterialIndex) && OverrideMaterials[MaterialIndex] != nullptr ?
