[GLTFExporter] refactor: improve GLTFMaterialArray and fix missing override check for SkeletalMesh

commit 7bf181d7075cb1d116eed986b886e8aefadb5c5b
Date: Tue, 22 Sep 2020 23:16:46 +0200

    refactor: improve GLTFMaterialArray and fix missing override check for SkeletalMesh
    (cherry picked from commit 7392ff43d551c240dc070cc5e9ecd8c0692c3da8​)

#skipci

diff --git a/Engine/Plugins/Enterprise/GLTFExporter/MigratedFromGit/Source/GLTFExporter/Private/Builders/GLTFConvertBuilder.cpp b/Engine/Plugins/Enterprise/GLTFExporter/MigratedFromGit/Source/GLTFExporter/Private/Builders/GLTFConvertBuilder.cpp
--- a/Engine/Plugins/Enterprise/GLTFExporter/MigratedFromGit/Source/GLTFExporter/Private/Builders/GLTFConvertBuilder.cpp
+++ b/Engine/Plugins/Enterprise/GLTFExporter/MigratedFromGit/Source/GLTFExporter/Private/Builders/GLTFConvertBuilder.cpp
@@ -105,11 +105,7 @@ FGLTFJsonMeshIndex FGLTFConvertBuilder::GetOrAddMesh(const UStaticMesh* StaticMe
 		OverrideVertexColors = nullptr;
 	}
 
-	return StaticMeshConverter.GetOrAdd(
-		StaticMesh,
-		LODIndex,
-		OverrideVertexColors,
-		OverrideMaterials == StaticMesh->StaticMaterials ? FGLTFMaterialArray(): OverrideMaterials);
+	return StaticMeshConverter.GetOrAdd(StaticMesh, LODIndex, OverrideVertexColors, OverrideMaterials == StaticMesh->StaticMaterials ? FGLTFMaterialArray() : OverrideMaterials);
 }
 
 FGLTFJsonMeshIndex FGLTFConvertBuilder::GetOrAddMesh(const UStaticMeshComponent* StaticMeshComponent)
@@ -154,7 +150,7 @@ FGLTFJsonMeshIndex FGLTFConvertBuilder::GetOrAddMesh(const USkeletalMesh* Skelet
 		OverrideVertexColors = nullptr;
 	}
 
-	return SkeletalMeshConverter.GetOrAdd(SkeletalMesh, LODIndex, OverrideVertexColors, OverrideSkinWeights, OverrideMaterials);
+	return SkeletalMeshConverter.GetOrAdd(SkeletalMesh, LODIndex, OverrideVertexColors, OverrideSkinWeights, OverrideMaterials == SkeletalMesh->Materials ? FGLTFMaterialArray() : OverrideMaterials);
 }
 
 FGLTFJsonMeshIndex FGLTFConvertBuilder::GetOrAddMesh(const USkeletalMeshComponent* SkeletalMeshComponent)
diff --git a/Engine/Plugins/Enterprise/GLTFExporter/MigratedFromGit/Source/GLTFExporter/Private/Converters/GLTFMaterialArray.h b/Engine/Plugins/Enterprise/GLTFExporter/MigratedFromGit/Source/GLTFExporter/Private/Converters/GLTFMaterialArray.h
--- a/Engine/Plugins/Enterprise/GLTFExporter/MigratedFromGit/Source/GLTFExporter/Private/Converters/GLTFMaterialArray.h
+++ b/Engine/Plugins/Enterprise/GLTFExporter/MigratedFromGit/Source/GLTFExporter/Private/Converters/GLTFMaterialArray.h
@@ -36,51 +36,38 @@ struct FGLTFMaterialArray : TArray<const UMaterialInterface*>
 	{
 	}
 
-	FGLTFMaterialArray& operator=(std::initializer_list<const UMaterialInterface*> InitList)
+	const UMaterialInterface* GetOverride(const TArray<FStaticMaterial>& Originals, int32 Index) const
 	{
-		static_cast<TArray&>(*this) = InitList;
-		return *this;
+		return GetOverride<FStaticMaterial>(Originals, Index);
 	}
 
-	FGLTFMaterialArray& operator=(const TArray& Other)
+	const UMaterialInterface* GetOverride(const TArray<FSkeletalMaterial>& Originals, int32 Index) const
 	{
-		static_cast<TArray&>(*this) = Other;
-		return *this;
+		return GetOverride<FSkeletalMaterial>(Originals, Index);
 	}
 
-	template <typename ElementType, typename AllocatorType>
-	FGLTFMaterialArray& operator=(const TArray<ElementType, AllocatorType>& Other)
-	{
-		static_cast<TArray&>(*this) = Other;
-		return *this;
-	}
+	using TArray::operator=;
+	using TArray::operator==;
+	using TArray::operator!=;
 
 	bool operator==(const TArray<FStaticMaterial>& Other) const
 	{
-		if (Other.Num() != Num())
-		{
-			return false;
-		}
-
-		for (int32 Index = 0; Index < Num(); ++Index)
-		{
-			if (Other[Index].MaterialInterface != (*this)[Index])
-			{
-				return false;
-			}
-		}
+		return Equals(Other);
+	}
 
-		return true;
+	bool operator!=(const TArray<FStaticMaterial>& Other) const
+	{
+		return !Equals(Other);
 	}
 
-	bool operator==(const TArray& Other) const
+	bool operator==(const TArray<FSkeletalMaterial>& Other) const
 	{
-		return static_cast<const TArray&>(*this) == Other;
+		return Equals(Other);
 	}
 
-	bool operator!=(const TArray& Other) const
+	bool operator!=(const TArray<FSkeletalMaterial>& Other) const
 	{
-		return static_cast<const TArray&>(*this) != Other;
+		return !Equals(Other);
 	}
 
 	friend uint32 GetTypeHash(const TArray& Array)
@@ -92,4 +79,41 @@ struct FGLTFMaterialArray : TArray<const UMaterialInterface*>
 		}
 		return Hash;
 	}
+
+private:
+
+	template <typename MeshMaterialType>
+    const UMaterialInterface* GetOverride(const TArray<MeshMaterialType>& Originals, int32 Index) const
+	{
+		if (IsValidIndex(Index) && (*this)[Index] != nullptr)
+		{
+			return (*this)[Index];
+		}
+
+		if (Originals.IsValidIndex(Index))
+		{
+			return Originals[Index].MaterialInterface;
+		}
+
+		return nullptr;
+	}
+
+	template <typename MeshMaterialType>
+	bool Equals(const TArray<MeshMaterialType>& Other) const
+	{
+		if (Other.Num() != Num())
+		{
+			return false;
+		}
+
+		for (int32 Index = 0; Index < Num(); ++Index)
+		{
+			if (Other[Index].MaterialInterface != (*this)[Index])
+			{
+				return false;
+			}
+		}
+
+		return true;
+	}
 };
diff --git a/Engine/Plugins/Enterprise/GLTFExporter/MigratedFromGit/Source/GLTFExporter/Private/Converters/GLTFMeshConverters.cpp b/Engine/Plugins/Enterprise/GLTFExporter/MigratedFromGit/Source/GLTFExporter/Private/Converters/GLTFMeshConverters.cpp
--- a/Engine/Plugins/Enterprise/GLTFExporter/MigratedFromGit/Source/GLTFExporter/Private/Converters/GLTFMeshConverters.cpp
+++ b/Engine/Plugins/Enterprise/GLTFExporter/MigratedFromGit/Source/GLTFExporter/Private/Converters/GLTFMeshConverters.cpp
@@ -55,10 +55,7 @@ FGLTFJsonMeshIndex FGLTFStaticMeshConverter::Convert(const UStaticMesh* StaticMe
 		const FStaticMeshSection& Section = MeshLOD.Sections[SectionIndex];
 		JsonPrimitive.Indices = Builder.GetOrAddIndexAccessor(&Section, IndexBuffer);
 
-		const int32 MaterialIndex = Section.MaterialIndex;
-		const UMaterialInterface* Material = OverrideMaterials.IsValidIndex(MaterialIndex) && OverrideMaterials[MaterialIndex] != nullptr ?
-			OverrideMaterials[MaterialIndex] : StaticMesh->GetMaterial(MaterialIndex);
-
+		const UMaterialInterface* Material = OverrideMaterials.GetOverride(StaticMesh->StaticMaterials, Section.MaterialIndex);
 		if (Material != nullptr)
 		{
 			JsonPrimitive.Material = Builder.GetOrAddMaterial(Material);
@@ -134,10 +131,7 @@ FGLTFJsonMeshIndex FGLTFSkeletalMeshConverter::Convert(const USkeletalMesh* Skel
 		const FSkelMeshRenderSection& Section = MeshLOD.RenderSections[SectionIndex];
 		JsonPrimitive.Indices = Builder.GetOrAddIndexAccessor(&Section, IndexBuffer);
 
-		const int32 MaterialIndex = Section.MaterialIndex;
-		const UMaterialInterface* Material = OverrideMaterials.IsValidIndex(MaterialIndex) && OverrideMaterials[MaterialIndex] != nullptr ?
-			OverrideMaterials[MaterialIndex] : SkeletalMesh->Materials[MaterialIndex].MaterialInterface;
-
+		const UMaterialInterface* Material = OverrideMaterials.GetOverride(SkeletalMesh->Materials, Section.MaterialIndex);
 		if (Material != nullptr)
 		{
 			JsonPrimitive.Material = Builder.GetOrAddMaterial(Material);
