[GLTFExporter] refactor: don't show warning messages for disabled options that are intuitively understandable (for ex related to scene components)

commit 7f2bf8d00e960da918cb1d3103cfbc05b533c131
Date: Sat, 26 Sep 2020 21:37:30 +0200

    refactor: don't show warning messages for disabled options that are intuitively understandable (for ex related to scene components)
    (cherry picked from commit 194d18852b09446120300758da8fe453624d145c​)

#skipci

diff --git a/Engine/Plugins/Enterprise/GLTFExporter/MigratedFromGit/Source/GLTFExporter/Private/Converters/GLTFLightMapConverters.cpp b/Engine/Plugins/Enterprise/GLTFExporter/MigratedFromGit/Source/GLTFExporter/Private/Converters/GLTFLightMapConverters.cpp
--- a/Engine/Plugins/Enterprise/GLTFExporter/MigratedFromGit/Source/GLTFExporter/Private/Converters/GLTFLightMapConverters.cpp
+++ b/Engine/Plugins/Enterprise/GLTFExporter/MigratedFromGit/Source/GLTFExporter/Private/Converters/GLTFLightMapConverters.cpp
@@ -37,15 +37,6 @@ FGLTFJsonLightMapIndex FGLTFLightMapConverter::Convert(const UStaticMeshComponen
 		return FGLTFJsonLightMapIndex(INDEX_NONE);
 	}
 
-	// TODO: this check should ideally be placed earlier (or outside this function), but we need to check
-	// if a light-map is in use before displaying warnings about skipped light-maps.
-	// Perhaps we should extract the code needed for checking for the presence of light-maps into
-	// a utility-function to allow the check to be performed elsewhere?
-	if (!Builder.ExportOptions->bExportLightmaps)
-	{
-		Builder.AddWarningMessage(FString::Printf(TEXT("Lightmap for %s disabled by export options"), *StaticMeshComponent->GetOwner()->GetName()));
-	}
-
 	const FLightMap2D* LightMap2D = MeshMapBuildData->LightMap->GetLightMap2D();
 	check(LightMap2D);
 
diff --git a/Engine/Plugins/Enterprise/GLTFExporter/MigratedFromGit/Source/GLTFExporter/Private/Converters/GLTFNodeConverters.cpp b/Engine/Plugins/Enterprise/GLTFExporter/MigratedFromGit/Source/GLTFExporter/Private/Converters/GLTFNodeConverters.cpp
--- a/Engine/Plugins/Enterprise/GLTFExporter/MigratedFromGit/Source/GLTFExporter/Private/Converters/GLTFNodeConverters.cpp
+++ b/Engine/Plugins/Enterprise/GLTFExporter/MigratedFromGit/Source/GLTFExporter/Private/Converters/GLTFNodeConverters.cpp
@@ -53,13 +53,19 @@ FGLTFJsonNodeIndex FGLTFComponentConverter::Convert(const USceneComponent* Scene
 			MeshNode.Name = FGLTFNameUtility::GetName(StaticMeshComponent);
 			MeshNode.Scale = ComponentNodeScale;
 			MeshNode.Mesh = Builder.GetOrAddMesh(StaticMeshComponent);
-			MeshNode.LightMap = Builder.GetOrAddLightMap(StaticMeshComponent);
+			if (Builder.ExportOptions->bExportLightmaps)
+			{
+				MeshNode.LightMap = Builder.GetOrAddLightMap(StaticMeshComponent);
+			}
 			Builder.AddChildComponentNode(NodeIndex, MeshNode);
 		}
 		else
 		{
 			Node.Mesh = Builder.GetOrAddMesh(StaticMeshComponent);
-			Node.LightMap = Builder.GetOrAddLightMap(StaticMeshComponent);
+			if (Builder.ExportOptions->bExportLightmaps)
+			{
+				Node.LightMap = Builder.GetOrAddLightMap(StaticMeshComponent);
+			}
 		}
 	}
 	else if (const USkeletalMeshComponent* SkeletalMeshComponent = Cast<USkeletalMeshComponent>(SceneComponent))
@@ -95,10 +101,6 @@ FGLTFJsonNodeIndex FGLTFComponentConverter::Convert(const USceneComponent* Scene
 				Node.Hotspot = Builder.GetOrAddHotspot(HotspotComponent);
 			}
 		}
-		else
-		{
-			Builder.AddWarningMessage(FString::Printf(TEXT("Interaction-hotspot %s disabled by export options"), *Owner->GetName()));
-		}
 	}
 	else if (const UCameraComponent* CameraComponent = Cast<UCameraComponent>(SceneComponent))
 	{
@@ -112,10 +114,6 @@ FGLTFJsonNodeIndex FGLTFComponentConverter::Convert(const USceneComponent* Scene
 			CameraNode.Camera = Builder.GetOrAddCamera(CameraComponent);
 			Builder.AddChildComponentNode(NodeIndex, CameraNode);
 		}
-		else
-		{
-			Builder.AddWarningMessage(FString::Printf(TEXT("Camera %s disabled by export options"), *Owner->GetName()));
-		}
 	}
 	else if (const ULightComponent* LightComponent = Cast<ULightComponent>(SceneComponent))
 	{
@@ -129,10 +127,6 @@ FGLTFJsonNodeIndex FGLTFComponentConverter::Convert(const USceneComponent* Scene
 			LightNode.Light = Builder.GetOrAddLight(LightComponent);
 			Builder.AddChildComponentNode(NodeIndex, LightNode);
 		}
-		else
-		{
-			Builder.AddWarningMessage(FString::Printf(TEXT("Light %s disabled by export options"), *Owner->GetName()));
-		}
 	}
 
 	// TODO: add support for SkyLight?
@@ -163,12 +157,6 @@ FGLTFJsonNodeIndex FGLTFActorConverter::Convert(const AActor* Actor)
 			FGLTFJsonNode& RootNode = Builder.GetNode(RootNodeIndex);
 			RootNode.Backdrop = Builder.GetOrAddBackdrop(Actor);
 		}
-		else
-		{
-			Builder.AddWarningMessage(FString::Printf(
-				TEXT("HDRIBackdrop %s disabled by export options"),
-				*Actor->GetName()));
-		}
 	}
 	else
 	{
diff --git a/Engine/Plugins/Enterprise/GLTFExporter/MigratedFromGit/Source/GLTFExporter/Private/Converters/GLTFSceneConverters.cpp b/Engine/Plugins/Enterprise/GLTFExporter/MigratedFromGit/Source/GLTFExporter/Private/Converters/GLTFSceneConverters.cpp
--- a/Engine/Plugins/Enterprise/GLTFExporter/MigratedFromGit/Source/GLTFExporter/Private/Converters/GLTFSceneConverters.cpp
+++ b/Engine/Plugins/Enterprise/GLTFExporter/MigratedFromGit/Source/GLTFExporter/Private/Converters/GLTFSceneConverters.cpp
@@ -33,12 +33,6 @@ FGLTFJsonSceneIndex FGLTFSceneConverter::Convert(const ULevel* Level)
 					Scene.Variations.Add(VariationIndex);
 				}
 			}
-			else
-			{
-				Builder.AddWarningMessage(FString::Printf(
-					TEXT("Level Variant Set %s disabled by export options"),
-					*LevelVariantSetsActor->GetName()));
-			}
 		}
 
 		const FGLTFJsonNodeIndex NodeIndex = Builder.GetOrAddNode(Actor);
