[GLTFExporter] refactor(material): improve warning messages related to disabled options

commit d7afbc9a2b9e734911f82d4531e2173bf177995b
Date: Sat, 26 Sep 2020 21:37:58 +0200

    refactor(material): improve warning messages related to disabled options
    (cherry picked from commit 4644b8c4cadaf542bb5dc54a7b329c3f2a72e2bf​)

#skipci

diff --git a/Engine/Plugins/Enterprise/GLTFExporter/MigratedFromGit/Source/GLTFExporter/Private/Converters/GLTFMaterialConverters.cpp b/Engine/Plugins/Enterprise/GLTFExporter/MigratedFromGit/Source/GLTFExporter/Private/Converters/GLTFMaterialConverters.cpp
--- a/Engine/Plugins/Enterprise/GLTFExporter/MigratedFromGit/Source/GLTFExporter/Private/Converters/GLTFMaterialConverters.cpp
+++ b/Engine/Plugins/Enterprise/GLTFExporter/MigratedFromGit/Source/GLTFExporter/Private/Converters/GLTFMaterialConverters.cpp
@@ -254,7 +254,7 @@ bool FGLTFMaterialConverter::TryGetBaseColorAndOpacity(FGLTFJsonPBRMetallicRough
 	if (!Builder.ExportOptions->bBakeMaterialInputs)
 	{
 		Builder.AddWarningMessage(FString::Printf(
-			TEXT("%s and %s for material %s need to baked, but material baking is disabled by export options"),
+			TEXT("%s and %s for material %s needs to bake, but material baking disabled by export options"),
 			FGLTFMaterialUtility::GetPropertyName(BaseColorProperty),
 			FGLTFMaterialUtility::GetPropertyName(OpacityProperty),
 			*Material->GetName()));
@@ -419,7 +419,7 @@ bool FGLTFMaterialConverter::TryGetMetallicAndRoughness(FGLTFJsonPBRMetallicRoug
 	if (!Builder.ExportOptions->bBakeMaterialInputs)
 	{
 		Builder.AddWarningMessage(FString::Printf(
-			TEXT("%s and %s for material %s need to baked, but material baking is disabled by export options"),
+			TEXT("%s and %s for material %s needs to bake, but material baking disabled by export options"),
 			FGLTFMaterialUtility::GetPropertyName(MetallicProperty),
 			FGLTFMaterialUtility::GetPropertyName(RoughnessProperty),
 			*Material->GetName()));
@@ -583,7 +583,7 @@ bool FGLTFMaterialConverter::TryGetClearCoatRoughness(FGLTFJsonClearCoatExtensio
 	if (!Builder.ExportOptions->bBakeMaterialInputs)
 	{
 		Builder.AddWarningMessage(FString::Printf(
-			TEXT("%s and %s for material %s need to baked, but material baking is disabled by export options"),
+			TEXT("%s and %s for material %s needs to bake, but material baking disabled by export options"),
 			FGLTFMaterialUtility::GetPropertyName(IntensityProperty),
 			FGLTFMaterialUtility::GetPropertyName(RoughnessProperty),
 			*Material->GetName()));
@@ -727,7 +727,7 @@ bool FGLTFMaterialConverter::TryGetEmissive(FGLTFJsonMaterial& JsonMaterial, EMa
 	if (!Builder.ExportOptions->bBakeMaterialInputs)
 	{
 		Builder.AddWarningMessage(FString::Printf(
-			TEXT("%s for material %s needs to baked, but material baking is disabled by export options"),
+			TEXT("%s for material %s needs to bake, but material baking disabled by export options"),
 			FGLTFMaterialUtility::GetPropertyName(EmissiveProperty),
 			*Material->GetName()));
 		return false;
@@ -1122,7 +1122,7 @@ bool FGLTFMaterialConverter::TryGetBakedMaterialProperty(FGLTFJsonTextureInfo& O
 	if (!Builder.ExportOptions->bBakeMaterialInputs)
 	{
 		Builder.AddWarningMessage(FString::Printf(
-			TEXT("%s for material %s needs to baked, but material baking is disabled by export options"),
+			TEXT("%s for material %s needs to bake, but material baking disabled by export options"),
 			FGLTFMaterialUtility::GetPropertyName(Property),
 			*Material->GetName()));
 		return false;
@@ -1150,7 +1150,7 @@ bool FGLTFMaterialConverter::TryGetBakedMaterialProperty(FGLTFJsonTextureInfo& O
 	if (!Builder.ExportOptions->bBakeMaterialInputs)
 	{
 		Builder.AddWarningMessage(FString::Printf(
-			TEXT("%s for material %s needs to baked, but material baking is disabled by export options"),
+			TEXT("%s for material %s needs to bake, but material baking disabled by export options"),
 			FGLTFMaterialUtility::GetPropertyName(Property),
 			*Material->GetName()));
 		return false;
@@ -1178,7 +1178,7 @@ inline bool FGLTFMaterialConverter::TryGetBakedMaterialProperty(FGLTFJsonTexture
 	if (!Builder.ExportOptions->bBakeMaterialInputs)
 	{
 		Builder.AddWarningMessage(FString::Printf(
-			TEXT("%s for material %s needs to baked, but material baking is disabled by export options"),
+			TEXT("%s for material %s needs to bake, but material baking disabled by export options"),
 			FGLTFMaterialUtility::GetPropertyName(Property),
 			*Material->GetName()));
 		return false;
@@ -1206,7 +1206,7 @@ bool FGLTFMaterialConverter::TryGetBakedMaterialProperty(FGLTFJsonTextureInfo& O
 	if (!Builder.ExportOptions->bBakeMaterialInputs)
 	{
 		Builder.AddWarningMessage(FString::Printf(
-			TEXT("%s for material %s needs to baked, but material baking is disabled by export options"),
+			TEXT("%s for material %s needs to bake, but material baking disabled by export options"),
 			FGLTFMaterialUtility::GetPropertyName(Property),
 			*Material->GetName()));
 		return false;
