[GLTFExporter] refactor: instead of enabling implemented export options, disable those that are not yet supported

commit cc0db76656d3c36aacc33da131ebe4276b5151cf
Date: Sat, 26 Sep 2020 21:40:11 +0200

    refactor: instead of enabling implemented export options, disable those that are not yet supported
    (cherry picked from commit b45b7d2f0d4a93543d95c07dae2526d8034070e5​)
    # Conflicts from cherry-pick:
    #	Plugins/GLTFExporter/Source/GLTFExporter/Private/GLTFExportOptions.cpp

#skipci

diff --git a/Engine/Plugins/Enterprise/GLTFExporter/MigratedFromGit/Source/GLTFExporter/Private/GLTFExportOptions.cpp b/Engine/Plugins/Enterprise/GLTFExporter/MigratedFromGit/Source/GLTFExporter/Private/GLTFExportOptions.cpp
--- a/Engine/Plugins/Enterprise/GLTFExporter/MigratedFromGit/Source/GLTFExporter/Private/GLTFExportOptions.cpp
+++ b/Engine/Plugins/Enterprise/GLTFExporter/MigratedFromGit/Source/GLTFExporter/Private/GLTFExportOptions.cpp
@@ -193,30 +193,21 @@ bool UGLTFExportOptions::CanEditChange(const FProperty* InProperty) const
 {
 	const FName PropertyFName = InProperty->GetFName();
 
-	if (PropertyFName == GET_MEMBER_NAME_CHECKED(ThisClass, bExportUnlitMaterials) ||
-		PropertyFName == GET_MEMBER_NAME_CHECKED(ThisClass, bExportClearCoatMaterials) ||
-		PropertyFName == GET_MEMBER_NAME_CHECKED(ThisClass, bBakeMaterialInputs) ||
-		PropertyFName == GET_MEMBER_NAME_CHECKED(ThisClass, DefaultMaterialBakeSize) ||
-		PropertyFName == GET_MEMBER_NAME_CHECKED(ThisClass, bExportVertexColors) ||
-		PropertyFName == GET_MEMBER_NAME_CHECKED(ThisClass, DefaultLevelOfDetail) ||
-		PropertyFName == GET_MEMBER_NAME_CHECKED(ThisClass, bExportSourceTextures) ||
-		PropertyFName == GET_MEMBER_NAME_CHECKED(ThisClass, bExportLightmaps) ||
-		PropertyFName == GET_MEMBER_NAME_CHECKED(ThisClass, ExportScale) ||
-		PropertyFName == GET_MEMBER_NAME_CHECKED(ThisClass, ExportLights) ||
-		PropertyFName == GET_MEMBER_NAME_CHECKED(ThisClass, bExportCameras) ||
-		PropertyFName == GET_MEMBER_NAME_CHECKED(ThisClass, bExportOrbitalCameras) ||
-		PropertyFName == GET_MEMBER_NAME_CHECKED(ThisClass, bExportHDRIBackdrops) ||
-		PropertyFName == GET_MEMBER_NAME_CHECKED(ThisClass, bExportVariantSets) ||
-		PropertyFName == GET_MEMBER_NAME_CHECKED(ThisClass, bExportInteractionHotspots) ||
-		PropertyFName == GET_MEMBER_NAME_CHECKED(ThisClass, bExportNonUniformScale) ||
-		PropertyFName == GET_MEMBER_NAME_CHECKED(ThisClass, bExportPreviewMesh))
+	// TODO: remove options that have been implemented in exporters/converters.
+	if (PropertyFName == GET_MEMBER_NAME_CHECKED(ThisClass, NormalExportMethod) ||
+		PropertyFName == GET_MEMBER_NAME_CHECKED(ThisClass, bTangentDataQuantization) ||
+		PropertyFName == GET_MEMBER_NAME_CHECKED(ThisClass, NormalizeUVCoordinates) ||
+		PropertyFName == GET_MEMBER_NAME_CHECKED(ThisClass, TextureFormat) ||
+		PropertyFName == GET_MEMBER_NAME_CHECKED(ThisClass, TextureHDREncoding) ||
+		PropertyFName == GET_MEMBER_NAME_CHECKED(ThisClass, bExportHiddenInGame) ||
+		PropertyFName == GET_MEMBER_NAME_CHECKED(ThisClass, bExportReflectionCaptures) ||
+		PropertyFName == GET_MEMBER_NAME_CHECKED(ThisClass, bBundleWebViewer) ||
+		PropertyFName == GET_MEMBER_NAME_CHECKED(ThisClass, bAllExtensionsRequired))
 	{
-		return true;
+		return false;
 	}
 
-	// TODO: return true for each option implemented and supported.
-
-	return false;
+	return true;
 }
 
 FIntPoint UGLTFExportOptions::GetDefaultMaterialBakeSize() const
