[GLTFExporter] feat(orbit-camera): disable export option for orbital camera if cameras in general are disabled

commit dd841f6543856f0dab3999d7c417a78c0f27301f
Date: Sat, 26 Sep 2020 21:46:40 +0200

    feat(orbit-camera): disable export option for orbital camera if cameras in general are disabled
    (cherry picked from commit 5730ae370c7f4054af61c3b40a468590f764577e​)

#skipci

diff --git a/Engine/Plugins/Enterprise/GLTFExporter/MigratedFromGit/Source/GLTFExporter/Private/Converters/GLTFCameraConverters.cpp b/Engine/Plugins/Enterprise/GLTFExporter/MigratedFromGit/Source/GLTFExporter/Private/Converters/GLTFCameraConverters.cpp
--- a/Engine/Plugins/Enterprise/GLTFExporter/MigratedFromGit/Source/GLTFExporter/Private/Converters/GLTFCameraConverters.cpp
+++ b/Engine/Plugins/Enterprise/GLTFExporter/MigratedFromGit/Source/GLTFExporter/Private/Converters/GLTFCameraConverters.cpp
@@ -60,12 +60,6 @@ FGLTFJsonCameraIndex FGLTFCameraConverter::Convert(const UCameraComponent* Camer
 					*Owner->GetName()));
 			}
 		}
-		else
-		{
-			Builder.AddWarningMessage(FString::Printf(
-				TEXT("OrbitalCamera %s disabled by export options"),
-				*Owner->GetName()));
-		}
 	}
 
 	return Builder.AddCamera(Camera);
diff --git a/Engine/Plugins/Enterprise/GLTFExporter/MigratedFromGit/Source/GLTFExporter/Private/GLTFExportOptions.cpp b/Engine/Plugins/Enterprise/GLTFExporter/MigratedFromGit/Source/GLTFExporter/Private/GLTFExportOptions.cpp
--- a/Engine/Plugins/Enterprise/GLTFExporter/MigratedFromGit/Source/GLTFExporter/Private/GLTFExportOptions.cpp
+++ b/Engine/Plugins/Enterprise/GLTFExporter/MigratedFromGit/Source/GLTFExporter/Private/GLTFExportOptions.cpp
@@ -207,6 +207,11 @@ bool UGLTFExportOptions::CanEditChange(const FProperty* InProperty) const
 		return false;
 	}
 
+	if (PropertyFName == GET_MEMBER_NAME_CHECKED(ThisClass, bExportOrbitalCameras))
+	{
+		return bExportCameras;
+	}
+
 	return true;
 }
 
