[GLTFExporter] feat: disable export option for material bake size if baking in general is disabled

commit 3344a08fa074620174496da71ecbcb0fead3bf24
Date: Sat, 26 Sep 2020 21:49:58 +0200

    feat: disable export option for material bake size if baking in general is disabled
    (cherry picked from commit 8a39454f930862ff88ebd5748e4c6d45ae7708cd​)
    # Conflicts from cherry-pick:
    #	Plugins/GLTFExporter/Source/GLTFExporter/Private/GLTFExportOptions.cpp

#skipci

diff --git a/Engine/Plugins/Enterprise/GLTFExporter/MigratedFromGit/Source/GLTFExporter/Private/GLTFExportOptions.cpp b/Engine/Plugins/Enterprise/GLTFExporter/MigratedFromGit/Source/GLTFExporter/Private/GLTFExportOptions.cpp
--- a/Engine/Plugins/Enterprise/GLTFExporter/MigratedFromGit/Source/GLTFExporter/Private/GLTFExportOptions.cpp
+++ b/Engine/Plugins/Enterprise/GLTFExporter/MigratedFromGit/Source/GLTFExporter/Private/GLTFExportOptions.cpp
@@ -207,6 +207,11 @@ bool UGLTFExportOptions::CanEditChange(const FProperty* InProperty) const
 		return false;
 	}
 
+	if (PropertyFName == GET_MEMBER_NAME_CHECKED(ThisClass, DefaultMaterialBakeSize))
+	{
+		return bBakeMaterialInputs;
+	}
+
 	if (PropertyFName == GET_MEMBER_NAME_CHECKED(ThisClass, bExportOrbitalCameras))
 	{
 		return bExportCameras;
