[GLTFExporter] refactor: improve error messages in exporters

commit 79eaf6156a63cb66e589dfc63ac64bcdb0823b84
Date: Sat, 26 Sep 2020 22:17:48 +0200

    refactor: improve error messages in exporters
    (cherry picked from commit 4b2b90435e83fcd6dd0da9f3d945c08308e732b1​)

#skipci

diff --git a/Engine/Plugins/Enterprise/GLTFExporter/MigratedFromGit/Source/GLTFExporter/Private/Exporters/GLTFMaterialExporter.cpp b/Engine/Plugins/Enterprise/GLTFExporter/MigratedFromGit/Source/GLTFExporter/Private/Exporters/GLTFMaterialExporter.cpp
--- a/Engine/Plugins/Enterprise/GLTFExporter/MigratedFromGit/Source/GLTFExporter/Private/Exporters/GLTFMaterialExporter.cpp
+++ b/Engine/Plugins/Enterprise/GLTFExporter/MigratedFromGit/Source/GLTFExporter/Private/Exporters/GLTFMaterialExporter.cpp
@@ -14,11 +14,13 @@ UGLTFMaterialExporter::UGLTFMaterialExporter(const FObjectInitializer& ObjectIni
 bool UGLTFMaterialExporter::AddObject(FGLTFContainerBuilder& Builder, const UObject* Object)
 {
 	const UMaterialInterface* Material = CastChecked<UMaterialInterface>(Object);
-	const FGLTFJsonMaterialIndex MaterialIndex = Builder.GetOrAddMaterial(Material);
 
+	const FGLTFJsonMaterialIndex MaterialIndex = Builder.GetOrAddMaterial(Material);
 	if (MaterialIndex == INDEX_NONE)
 	{
-		Builder.AddErrorMessage(FString::Printf(TEXT("Failed to export material %s"), *Material->GetName()));
+		Builder.AddErrorMessage(
+			FString::Printf(TEXT("Failed to export material %s"),
+			*Material->GetName()));
 		return false;
 	}
 
@@ -28,6 +30,14 @@ bool UGLTFMaterialExporter::AddObject(FGLTFContainerBuilder& Builder, const UObj
 		if (PreviewMesh != nullptr)
 		{
 			const FGLTFJsonMeshIndex MeshIndex = Builder.GetOrAddMesh(PreviewMesh, -1, nullptr, { Material });
+			if (MeshIndex == INDEX_NONE)
+			{
+				Builder.AddErrorMessage(
+					FString::Printf(TEXT("Failed to export preview mesh %s for material %s"),
+					*Material->GetName(),
+					*PreviewMesh->GetName()));
+				return false;
+			}
 
 			FGLTFJsonNode Node;
 			Node.Mesh = MeshIndex;
@@ -41,10 +51,9 @@ bool UGLTFMaterialExporter::AddObject(FGLTFContainerBuilder& Builder, const UObj
 		}
 		else
 		{
-			Builder.AddErrorMessage(FString::Printf(
-				TEXT("Failed to export material %s because of missing preview-mesh"),
+			Builder.AddErrorMessage(
+				FString::Printf(TEXT("Failed to export material %s because of missing preview mesh"),
 				*Material->GetName()));
-
 			return false;
 		}
 	}
