[GLTFExporter] feat: add GLTFMeshSection generate new unique vertex indices for each section and standardize static and skeletal mesh section data

commit 14beeb36c57736b5a0da6c14776fcf0417b0db38
Date: Tue, 22 Sep 2020 18:18:33 +0200

    feat: add GLTFMeshSection generate new unique vertex indices for each section and standardize static and skeletal mesh section data
    (cherry picked from commit c83c8e5e4987b66e2b598d79876439a9aeac53e2​)

#skipci

diff --git a/Engine/Plugins/Enterprise/GLTFExporter/MigratedFromGit/Source/GLTFExporter/Private/Converters/GLTFMeshSection.cpp b/Engine/Plugins/Enterprise/GLTFExporter/MigratedFromGit/Source/GLTFExporter/Private/Converters/GLTFMeshSection.cpp
new file mode 100644
--- /dev/null
+++ b/Engine/Plugins/Enterprise/GLTFExporter/MigratedFromGit/Source/GLTFExporter/Private/Converters/GLTFMeshSection.cpp
@@ -0,0 +1,67 @@
+// Copyright Epic Games, Inc. All Rights Reserved.
+
+#include "Converters/GLTFMeshSection.h"
+#include "Rendering/MultiSizeIndexContainer.h"
+#include "Rendering/SkeletalMeshRenderData.h"
+
+FGLTFMeshSection::FGLTFMeshSection(const FStaticMeshSection* MeshSection, const FRawStaticIndexBuffer* OldIndexBuffer)
+{
+	const uint32 IndexOffset = MeshSection->FirstIndex;
+	const uint32 IndexCount = MeshSection->NumTriangles * 3;
+
+	IndexMap.Reserve(IndexCount);
+	IndexBuffer.AddUninitialized(IndexCount);
+
+	TMap<uint32, uint32> IndexLookup;
+
+	for (uint32 Index = 0; Index < IndexCount; Index++)
+	{
+		const uint32 OldIndex = OldIndexBuffer->GetIndex(IndexOffset + Index);
+		uint32 NewIndex;
+
+		if (uint32* FoundIndex = IndexLookup.Find(OldIndex))
+		{
+			NewIndex = *FoundIndex;
+		}
+		else
+		{
+			NewIndex = IndexMap.Num();
+			IndexLookup.Add(OldIndex, NewIndex);
+			IndexMap.Add(OldIndex);
+		}
+
+		IndexBuffer[Index] = NewIndex;
+	}
+}
+
+FGLTFMeshSection::FGLTFMeshSection(const FSkelMeshRenderSection* MeshSection, const FRawStaticIndexBuffer16or32Interface* OldIndexBuffer)
+{
+	const uint32 IndexOffset = MeshSection->BaseIndex;
+	const uint32 IndexCount = MeshSection->NumTriangles * 3;
+
+	IndexMap.Reserve(IndexCount);
+	IndexBuffer.AddUninitialized(IndexCount);
+
+	TMap<uint32, uint32> IndexLookup;
+
+	for (uint32 Index = 0; Index < IndexCount; Index++)
+	{
+		const uint32 OldIndex = OldIndexBuffer->Get(IndexOffset + Index);
+		uint32 NewIndex;
+
+		if (uint32* FoundIndex = IndexLookup.Find(OldIndex))
+		{
+			NewIndex = *FoundIndex;
+		}
+		else
+		{
+			NewIndex = IndexMap.Num();
+			IndexLookup.Add(OldIndex, NewIndex);
+			IndexMap.Add(OldIndex);
+		}
+
+		IndexBuffer[Index] = NewIndex;
+	}
+
+	BoneMap = MeshSection->BoneMap;
+}
diff --git a/Engine/Plugins/Enterprise/GLTFExporter/MigratedFromGit/Source/GLTFExporter/Private/Converters/GLTFMeshSection.h b/Engine/Plugins/Enterprise/GLTFExporter/MigratedFromGit/Source/GLTFExporter/Private/Converters/GLTFMeshSection.h
new file mode 100644
--- /dev/null
+++ b/Engine/Plugins/Enterprise/GLTFExporter/MigratedFromGit/Source/GLTFExporter/Private/Converters/GLTFMeshSection.h
@@ -0,0 +1,15 @@
+// Copyright Epic Games, Inc. All Rights Reserved.
+
+#pragma once
+
+#include "CoreMinimal.h"
+
+struct FGLTFMeshSection
+{
+	FGLTFMeshSection(const FStaticMeshSection* MeshSection, const FRawStaticIndexBuffer* OldIndexBuffer);
+	FGLTFMeshSection(const FSkelMeshRenderSection* MeshSection, const FRawStaticIndexBuffer16or32Interface* OldIndexBuffer);
+
+	TArray<uint32> IndexMap;
+	TArray<uint32> IndexBuffer;
+	TArray<FBoneIndexType> BoneMap;
+};
