[GLTFExporter] refactor: reorder conversion of vertex data to facilitate unique vertex data per section.

commit 3c2019deb0f3cd57d97467e5b61a53a79434171b
Date: Tue, 22 Sep 2020 18:19:58 +0200

    refactor: reorder conversion of vertex data to facilitate unique vertex data per section.
    (cherry picked from commit de88b64af5a2f02ebc7f0fb7e85eb4ef69a5c3d7​)

#skipci

diff --git a/Engine/Plugins/Enterprise/GLTFExporter/MigratedFromGit/Source/GLTFExporter/Private/Converters/GLTFMeshConverters.cpp b/Engine/Plugins/Enterprise/GLTFExporter/MigratedFromGit/Source/GLTFExporter/Private/Converters/GLTFMeshConverters.cpp
--- a/Engine/Plugins/Enterprise/GLTFExporter/MigratedFromGit/Source/GLTFExporter/Private/Converters/GLTFMeshConverters.cpp
+++ b/Engine/Plugins/Enterprise/GLTFExporter/MigratedFromGit/Source/GLTFExporter/Private/Converters/GLTFMeshConverters.cpp
@@ -22,26 +22,6 @@ FGLTFJsonMeshIndex FGLTFStaticMeshConverter::Convert(const UStaticMesh* StaticMe
 	const FPositionVertexBuffer* PositionBuffer = &MeshLOD.VertexBuffers.PositionVertexBuffer;
 	const FStaticMeshVertexBuffer* VertexBuffer = &MeshLOD.VertexBuffers.StaticMeshVertexBuffer;
 	const FColorVertexBuffer* ColorBuffer = OverrideVertexColors != nullptr ? OverrideVertexColors : &MeshLOD.VertexBuffers.ColorVertexBuffer;
-
-	FGLTFJsonAttributes JsonAttributes;
-	JsonAttributes.Position = Builder.GetOrAddPositionAccessor(PositionBuffer);
-
-	if (Builder.ExportOptions->bExportVertexColors)
-	{
-		JsonAttributes.Color0 = Builder.GetOrAddColorAccessor(ColorBuffer);
-	}
-
-	JsonAttributes.Normal = Builder.GetOrAddNormalAccessor(VertexBuffer);
-	JsonAttributes.Tangent = Builder.GetOrAddTangentAccessor(VertexBuffer);
-
-	const uint32 UVCount = VertexBuffer->GetNumTexCoords();
-	JsonAttributes.TexCoords.AddUninitialized(UVCount);
-
-	for (uint32 UVIndex = 0; UVIndex < UVCount; ++UVIndex)
-	{
-		JsonAttributes.TexCoords[UVIndex] = Builder.GetOrAddUVAccessor(VertexBuffer, UVIndex);
-	}
-
 	const FRawStaticIndexBuffer* IndexBuffer = &MeshLOD.IndexBuffer;
 
 	const int32 SectionCount = MeshLOD.Sections.Num();
@@ -50,10 +30,26 @@ FGLTFJsonMeshIndex FGLTFStaticMeshConverter::Convert(const UStaticMesh* StaticMe
 	for (int32 SectionIndex = 0; SectionIndex < SectionCount; ++SectionIndex)
 	{
 		FGLTFJsonPrimitive& JsonPrimitive = JsonMesh.Primitives[SectionIndex];
-		JsonPrimitive.Attributes = JsonAttributes;
-
 		const FStaticMeshSection& Section = MeshLOD.Sections[SectionIndex];
+
 		JsonPrimitive.Indices = Builder.GetOrAddIndexAccessor(&Section, IndexBuffer);
+		JsonPrimitive.Attributes.Position = Builder.GetOrAddPositionAccessor(PositionBuffer);
+
+		if (Builder.ExportOptions->bExportVertexColors)
+		{
+			JsonPrimitive.Attributes.Color0 = Builder.GetOrAddColorAccessor(ColorBuffer);
+		}
+
+		JsonPrimitive.Attributes.Normal = Builder.GetOrAddNormalAccessor(VertexBuffer);
+		JsonPrimitive.Attributes.Tangent = Builder.GetOrAddTangentAccessor(VertexBuffer);
+
+		const uint32 UVCount = VertexBuffer->GetNumTexCoords();
+		JsonPrimitive.Attributes.TexCoords.AddUninitialized(UVCount);
+
+		for (uint32 UVIndex = 0; UVIndex < UVCount; ++UVIndex)
+		{
+			JsonPrimitive.Attributes.TexCoords[UVIndex] = Builder.GetOrAddUVAccessor(VertexBuffer, UVIndex);
+		}
 
 		const UMaterialInterface* Material = OverrideMaterials.GetOverride(StaticMesh->StaticMaterials, Section.MaterialIndex);
 		if (Material != nullptr)
@@ -84,52 +80,49 @@ FGLTFJsonMeshIndex FGLTFSkeletalMeshConverter::Convert(const USkeletalMesh* Skel
 	const FColorVertexBuffer* ColorBuffer = OverrideVertexColors != nullptr ? OverrideVertexColors : &MeshLOD.StaticVertexBuffers.ColorVertexBuffer;
 	const FSkinWeightVertexBuffer* SkinWeightBuffer = OverrideSkinWeights != nullptr ? OverrideSkinWeights : MeshLOD.GetSkinWeightVertexBuffer();
 
-	FGLTFJsonAttributes JsonAttributes;
-	JsonAttributes.Position = Builder.GetOrAddPositionAccessor(PositionBuffer);
-
-	if (Builder.ExportOptions->bExportVertexColors)
-	{
-		JsonAttributes.Color0 = Builder.GetOrAddColorAccessor(ColorBuffer);
-	}
-
-	JsonAttributes.Normal = Builder.GetOrAddNormalAccessor(VertexBuffer);
-	JsonAttributes.Tangent = Builder.GetOrAddTangentAccessor(VertexBuffer);
+	const FRawStaticIndexBuffer16or32Interface* IndexBuffer = MeshLOD.MultiSizeIndexContainer.GetIndexBuffer();
 
-	const uint32 UVCount = VertexBuffer->GetNumTexCoords();
-	JsonAttributes.TexCoords.AddUninitialized(UVCount);
+	const int32 SectionCount = MeshLOD.RenderSections.Num();
+	JsonMesh.Primitives.AddDefaulted(SectionCount);
 
-	for (uint32 UVIndex = 0; UVIndex < UVCount; ++UVIndex)
+	for (int32 SectionIndex = 0; SectionIndex < SectionCount; ++SectionIndex)
 	{
-		JsonAttributes.TexCoords[UVIndex] = Builder.GetOrAddUVAccessor(VertexBuffer, UVIndex);
-	}
+		FGLTFJsonPrimitive& JsonPrimitive = JsonMesh.Primitives[SectionIndex];
+		const FSkelMeshRenderSection& Section = MeshLOD.RenderSections[SectionIndex];
 
-	/* TODO: enable export of vertex skin data when crash bug fixed
+		JsonPrimitive.Indices = Builder.GetOrAddIndexAccessor(&Section, IndexBuffer);
+		JsonPrimitive.Attributes.Position = Builder.GetOrAddPositionAccessor(PositionBuffer);
 
-	const uint32 GroupCount = (SkinWeightBuffer->GetMaxBoneInfluences() + 3) / 4;
-	JsonAttributes.Joints.AddUninitialized(GroupCount);
-	JsonAttributes.Weights.AddUninitialized(GroupCount);
+		if (Builder.ExportOptions->bExportVertexColors)
+		{
+			JsonPrimitive.Attributes.Color0 = Builder.GetOrAddColorAccessor(ColorBuffer);
+		}
 
-	const FGLTFBoneMap BoneMap = FGLTFBoneMap(MeshLOD.RenderSections[0].BoneMap); // TODO: can we make this assumption?
+		JsonPrimitive.Attributes.Normal = Builder.GetOrAddNormalAccessor(VertexBuffer);
+		JsonPrimitive.Attributes.Tangent = Builder.GetOrAddTangentAccessor(VertexBuffer);
 
-	for (uint32 GroupIndex = 0; GroupIndex < GroupCount; ++GroupIndex)
-	{
-		JsonAttributes.Joints[GroupIndex] = Builder.GetOrAddJointAccessor(SkinWeightBuffer, GroupIndex, BoneMap, Name + TEXT("_Joints") + FString::FromInt(GroupIndex));
-		JsonAttributes.Weights[GroupIndex] = Builder.GetOrAddWeightAccessor(SkinWeightBuffer, GroupIndex, Name + TEXT("_Weights") + FString::FromInt(GroupIndex));
-	}
-	*/
+		const uint32 UVCount = VertexBuffer->GetNumTexCoords();
+		JsonPrimitive.Attributes.TexCoords.AddUninitialized(UVCount);
 
-	const FRawStaticIndexBuffer16or32Interface* IndexBuffer = MeshLOD.MultiSizeIndexContainer.GetIndexBuffer();
+		for (uint32 UVIndex = 0; UVIndex < UVCount; ++UVIndex)
+		{
+			JsonPrimitive.Attributes.TexCoords[UVIndex] = Builder.GetOrAddUVAccessor(VertexBuffer, UVIndex);
+		}
 
-	const int32 SectionCount = MeshLOD.RenderSections.Num();
-	JsonMesh.Primitives.AddDefaulted(SectionCount);
+		/* TODO: enable export of vertex skin data when crash bug fixed
 
-	for (int32 SectionIndex = 0; SectionIndex < SectionCount; ++SectionIndex)
-	{
-		FGLTFJsonPrimitive& JsonPrimitive = JsonMesh.Primitives[SectionIndex];
-		JsonPrimitive.Attributes = JsonAttributes;
+		const uint32 GroupCount = (SkinWeightBuffer->GetMaxBoneInfluences() + 3) / 4;
+		JsonPrimitive.Attributes.Joints.AddUninitialized(GroupCount);
+		JsonPrimitive.Attributes.Weights.AddUninitialized(GroupCount);
 
-		const FSkelMeshRenderSection& Section = MeshLOD.RenderSections[SectionIndex];
-		JsonPrimitive.Indices = Builder.GetOrAddIndexAccessor(&Section, IndexBuffer);
+		const FGLTFBoneMap BoneMap = FGLTFBoneMap(Section.BoneMap);
+
+		for (uint32 GroupIndex = 0; GroupIndex < GroupCount; ++GroupIndex)
+		{
+			JsonPrimitive.Attributes.Joints[GroupIndex] = Builder.GetOrAddJointAccessor(SkinWeightBuffer, GroupIndex, BoneMap, Name + TEXT("_Joints") + FString::FromInt(GroupIndex));
+			JsonPrimitive.Attributes.Weights[GroupIndex] = Builder.GetOrAddWeightAccessor(SkinWeightBuffer, GroupIndex, Name + TEXT("_Weights") + FString::FromInt(GroupIndex));
+		}
+		*/
 
 		const UMaterialInterface* Material = OverrideMaterials.GetOverride(SkeletalMesh->Materials, Section.MaterialIndex);
 		if (Material != nullptr)
