[GLTFExporter] feat: add GLTFMeshSection::MaxBoneIndex

commit 21683c5310978996012aa35a8e981f085aa0f29e
Date: Wed, 23 Sep 2020 15:15:27 +0200

    feat: add GLTFMeshSection::MaxBoneIndex
    (cherry picked from commit 4a65713aba859c0071ca39401a85ba6305b49a8e​)

#skipci

diff --git a/Engine/Plugins/Enterprise/GLTFExporter/MigratedFromGit/Source/GLTFExporter/Private/Converters/GLTFMeshSection.cpp b/Engine/Plugins/Enterprise/GLTFExporter/MigratedFromGit/Source/GLTFExporter/Private/Converters/GLTFMeshSection.cpp
--- a/Engine/Plugins/Enterprise/GLTFExporter/MigratedFromGit/Source/GLTFExporter/Private/Converters/GLTFMeshSection.cpp
+++ b/Engine/Plugins/Enterprise/GLTFExporter/MigratedFromGit/Source/GLTFExporter/Private/Converters/GLTFMeshSection.cpp
@@ -3,6 +3,7 @@
 #include "Converters/GLTFMeshSection.h"
 #include "Rendering/MultiSizeIndexContainer.h"
 #include "Rendering/SkeletalMeshRenderData.h"
+#include "Algo/MaxElement.h"
 
 FGLTFMeshSection::FGLTFMeshSection(const FStaticMeshSection* MeshSection, const FRawStaticIndexBuffer* OldIndexBuffer)
 {
@@ -32,6 +33,8 @@ FGLTFMeshSection::FGLTFMeshSection(const FStaticMeshSection* MeshSection, const
 
 		IndexBuffer[Index] = NewIndex;
 	}
+
+	MaxBoneIndex = 0;
 }
 
 FGLTFMeshSection::FGLTFMeshSection(const FSkelMeshRenderSection* MeshSection, const FRawStaticIndexBuffer16or32Interface* OldIndexBuffer)
@@ -64,4 +67,5 @@ FGLTFMeshSection::FGLTFMeshSection(const FSkelMeshRenderSection* MeshSection, co
 	}
 
 	BoneMap = MeshSection->BoneMap;
+	MaxBoneIndex = BoneMap.Num() > 0 ? *Algo::MaxElement(MeshSection->BoneMap) : 0;
 }
diff --git a/Engine/Plugins/Enterprise/GLTFExporter/MigratedFromGit/Source/GLTFExporter/Private/Converters/GLTFMeshSection.h b/Engine/Plugins/Enterprise/GLTFExporter/MigratedFromGit/Source/GLTFExporter/Private/Converters/GLTFMeshSection.h
--- a/Engine/Plugins/Enterprise/GLTFExporter/MigratedFromGit/Source/GLTFExporter/Private/Converters/GLTFMeshSection.h
+++ b/Engine/Plugins/Enterprise/GLTFExporter/MigratedFromGit/Source/GLTFExporter/Private/Converters/GLTFMeshSection.h
@@ -11,5 +11,7 @@ struct FGLTFMeshSection
 
 	TArray<uint32> IndexMap;
 	TArray<uint32> IndexBuffer;
+
 	TArray<FBoneIndexType> BoneMap;
+	uint32 MaxBoneIndex;
 };
