[GLTFExporter] refactor: use new GLTFMeshSection in accessor and mesh converters

commit 58c43e3cd0974c972426b3e7623f828326bb6165
Date: Wed, 23 Sep 2020 15:16:58 +0200

    refactor: use new GLTFMeshSection in accessor and mesh converters
    (cherry picked from commit 22a3b804c10de1fdee3a3e4de0085f58b3be4bfb​)

#skipci

diff --git a/Engine/Plugins/Enterprise/GLTFExporter/MigratedFromGit/Source/GLTFExporter/Private/Builders/GLTFConvertBuilder.cpp b/Engine/Plugins/Enterprise/GLTFExporter/MigratedFromGit/Source/GLTFExporter/Private/Builders/GLTFConvertBuilder.cpp
--- a/Engine/Plugins/Enterprise/GLTFExporter/MigratedFromGit/Source/GLTFExporter/Private/Builders/GLTFConvertBuilder.cpp
+++ b/Engine/Plugins/Enterprise/GLTFExporter/MigratedFromGit/Source/GLTFExporter/Private/Builders/GLTFConvertBuilder.cpp
@@ -8,84 +8,84 @@ FGLTFConvertBuilder::FGLTFConvertBuilder(const UGLTFExportOptions* ExportOptions
 {
 }
 
-FGLTFJsonAccessorIndex FGLTFConvertBuilder::GetOrAddPositionAccessor(const FPositionVertexBuffer* VertexBuffer)
+FGLTFJsonAccessorIndex FGLTFConvertBuilder::GetOrAddPositionAccessor(const FGLTFMeshSection* MeshSection, const FPositionVertexBuffer* VertexBuffer)
 {
 	if (VertexBuffer == nullptr)
 	{
 		return FGLTFJsonAccessorIndex(INDEX_NONE);
 	}
 
-	return PositionBufferConverter.GetOrAdd(VertexBuffer);
+	return PositionBufferConverter.GetOrAdd(MeshSection, VertexBuffer);
 }
 
-FGLTFJsonAccessorIndex FGLTFConvertBuilder::GetOrAddColorAccessor(const FColorVertexBuffer* VertexBuffer)
+FGLTFJsonAccessorIndex FGLTFConvertBuilder::GetOrAddColorAccessor(const FGLTFMeshSection* MeshSection, const FColorVertexBuffer* VertexBuffer)
 {
 	if (VertexBuffer == nullptr)
 	{
 		return FGLTFJsonAccessorIndex(INDEX_NONE);
 	}
 
-	return ColorBufferConverter.GetOrAdd(VertexBuffer);
+	return ColorBufferConverter.GetOrAdd(MeshSection, VertexBuffer);
 }
 
-FGLTFJsonAccessorIndex FGLTFConvertBuilder::GetOrAddNormalAccessor(const FStaticMeshVertexBuffer* VertexBuffer)
+FGLTFJsonAccessorIndex FGLTFConvertBuilder::GetOrAddNormalAccessor(const FGLTFMeshSection* MeshSection, const FStaticMeshVertexBuffer* VertexBuffer)
 {
 	if (VertexBuffer == nullptr)
 	{
 		return FGLTFJsonAccessorIndex(INDEX_NONE);
 	}
 
-	return NormalBufferConverter.GetOrAdd(VertexBuffer);
+	return NormalBufferConverter.GetOrAdd(MeshSection, VertexBuffer);
 }
 
-FGLTFJsonAccessorIndex FGLTFConvertBuilder::GetOrAddTangentAccessor(const FStaticMeshVertexBuffer* VertexBuffer)
+FGLTFJsonAccessorIndex FGLTFConvertBuilder::GetOrAddTangentAccessor(const FGLTFMeshSection* MeshSection, const FStaticMeshVertexBuffer* VertexBuffer)
 {
 	if (VertexBuffer == nullptr)
 	{
 		return FGLTFJsonAccessorIndex(INDEX_NONE);
 	}
 
-	return TangentBufferConverter.GetOrAdd(VertexBuffer);
+	return TangentBufferConverter.GetOrAdd(MeshSection, VertexBuffer);
 }
 
-FGLTFJsonAccessorIndex FGLTFConvertBuilder::GetOrAddUVAccessor(const FStaticMeshVertexBuffer* VertexBuffer, int32 UVIndex)
+FGLTFJsonAccessorIndex FGLTFConvertBuilder::GetOrAddUVAccessor(const FGLTFMeshSection* MeshSection, const FStaticMeshVertexBuffer* VertexBuffer, int32 UVIndex)
 {
 	if (VertexBuffer == nullptr)
 	{
 		return FGLTFJsonAccessorIndex(INDEX_NONE);
 	}
 
-	return UVBufferConverter.GetOrAdd(VertexBuffer, UVIndex);
+	return UVBufferConverter.GetOrAdd(MeshSection, VertexBuffer, UVIndex);
 }
 
-FGLTFJsonAccessorIndex FGLTFConvertBuilder::GetOrAddJointAccessor(const FSkinWeightVertexBuffer* VertexBuffer, int32 JointsGroupIndex, FGLTFBoneMap BoneMap)
+FGLTFJsonAccessorIndex FGLTFConvertBuilder::GetOrAddJointAccessor(const FGLTFMeshSection* MeshSection, const FSkinWeightVertexBuffer* VertexBuffer, int32 InfluenceOffset)
 {
 	if (VertexBuffer == nullptr)
 	{
 		return FGLTFJsonAccessorIndex(INDEX_NONE);
 	}
 
-	return BoneIndexBufferConverter.GetOrAdd(VertexBuffer, JointsGroupIndex, BoneMap);
+	return BoneIndexBufferConverter.GetOrAdd(MeshSection, VertexBuffer, InfluenceOffset);
 }
 
-FGLTFJsonAccessorIndex FGLTFConvertBuilder::GetOrAddWeightAccessor(const FSkinWeightVertexBuffer* VertexBuffer, int32 WeightsGroupIndex)
+FGLTFJsonAccessorIndex FGLTFConvertBuilder::GetOrAddWeightAccessor(const FGLTFMeshSection* MeshSection, const FSkinWeightVertexBuffer* VertexBuffer, int32 InfluenceOffset)
 {
 	if (VertexBuffer == nullptr)
 	{
 		return FGLTFJsonAccessorIndex(INDEX_NONE);
 	}
 
-	return BoneWeightBufferConverter.GetOrAdd(VertexBuffer, WeightsGroupIndex);
+	return BoneWeightBufferConverter.GetOrAdd(MeshSection, VertexBuffer, InfluenceOffset);
 }
 
-FGLTFJsonAccessorIndex FGLTFConvertBuilder::GetOrAddIndexAccessor(const FStaticMeshSection* MeshSection, const FRawStaticIndexBuffer* IndexBuffer)
+FGLTFJsonAccessorIndex FGLTFConvertBuilder::GetOrAddIndexAccessor(const FGLTFMeshSection* MeshSection)
 {
 	if (MeshSection == nullptr)
 	{
 		return FGLTFJsonAccessorIndex(INDEX_NONE);
 	}
 
-	return StaticMeshSectionConverter.GetOrAdd(MeshSection, IndexBuffer);
+	return IndexBufferConverter.GetOrAdd(MeshSection);
 }
 
 FGLTFJsonMeshIndex FGLTFConvertBuilder::GetOrAddMesh(const UStaticMesh* StaticMesh, int32 LODIndex, const FColorVertexBuffer* OverrideVertexColors, const FGLTFMaterialArray& OverrideMaterials)
@@ -123,16 +123,6 @@ FGLTFJsonMeshIndex FGLTFConvertBuilder::GetOrAddMesh(const UStaticMeshComponent*
 	return GetOrAddMesh(StaticMesh, LODIndex, OverrideVertexColors, OverrideMaterials);
 }
 
-FGLTFJsonAccessorIndex FGLTFConvertBuilder::GetOrAddIndexAccessor(const FSkelMeshRenderSection* MeshSection, const FRawStaticIndexBuffer16or32Interface* IndexBuffer)
-{
-	if (MeshSection == nullptr)
-	{
-		return FGLTFJsonAccessorIndex(INDEX_NONE);
-	}
-
-	return SkeletalMeshSectionConverter.GetOrAdd(MeshSection, IndexBuffer);
-}
-
 FGLTFJsonMeshIndex FGLTFConvertBuilder::GetOrAddMesh(const USkeletalMesh* SkeletalMesh, int32 LODIndex, const FColorVertexBuffer* OverrideVertexColors, const FSkinWeightVertexBuffer* OverrideSkinWeights, const FGLTFMaterialArray& OverrideMaterials)
 {
 	if (SkeletalMesh == nullptr)
diff --git a/Engine/Plugins/Enterprise/GLTFExporter/MigratedFromGit/Source/GLTFExporter/Private/Builders/GLTFConvertBuilder.h b/Engine/Plugins/Enterprise/GLTFExporter/MigratedFromGit/Source/GLTFExporter/Private/Builders/GLTFConvertBuilder.h
--- a/Engine/Plugins/Enterprise/GLTFExporter/MigratedFromGit/Source/GLTFExporter/Private/Builders/GLTFConvertBuilder.h
+++ b/Engine/Plugins/Enterprise/GLTFExporter/MigratedFromGit/Source/GLTFExporter/Private/Builders/GLTFConvertBuilder.h
@@ -25,15 +25,14 @@ public:
 
 	FGLTFConvertBuilder(const UGLTFExportOptions* ExportOptions, bool bSelectedActorsOnly);
 
-	FGLTFJsonAccessorIndex GetOrAddPositionAccessor(const FPositionVertexBuffer* VertexBuffer);
-	FGLTFJsonAccessorIndex GetOrAddColorAccessor(const FColorVertexBuffer* VertexBuffer);
-	FGLTFJsonAccessorIndex GetOrAddNormalAccessor(const FStaticMeshVertexBuffer* VertexBuffer);
-	FGLTFJsonAccessorIndex GetOrAddTangentAccessor(const FStaticMeshVertexBuffer* VertexBuffer);
-	FGLTFJsonAccessorIndex GetOrAddUVAccessor(const FStaticMeshVertexBuffer* VertexBuffer, int32 UVIndex);
-	FGLTFJsonAccessorIndex GetOrAddJointAccessor(const FSkinWeightVertexBuffer* VertexBuffer, int32 JointsGroupIndex, FGLTFBoneMap BoneMap);
-	FGLTFJsonAccessorIndex GetOrAddWeightAccessor(const FSkinWeightVertexBuffer* VertexBuffer, int32 WeightsGroupIndex);
-	FGLTFJsonAccessorIndex GetOrAddIndexAccessor(const FStaticMeshSection* MeshSection, const FRawStaticIndexBuffer* IndexBuffer);
-	FGLTFJsonAccessorIndex GetOrAddIndexAccessor(const FSkelMeshRenderSection* MeshSection, const FRawStaticIndexBuffer16or32Interface* IndexBuffer);
+	FGLTFJsonAccessorIndex GetOrAddPositionAccessor(const FGLTFMeshSection* MeshSection, const FPositionVertexBuffer* VertexBuffer);
+	FGLTFJsonAccessorIndex GetOrAddColorAccessor(const FGLTFMeshSection* MeshSection, const FColorVertexBuffer* VertexBuffer);
+	FGLTFJsonAccessorIndex GetOrAddNormalAccessor(const FGLTFMeshSection* MeshSection, const FStaticMeshVertexBuffer* VertexBuffer);
+	FGLTFJsonAccessorIndex GetOrAddTangentAccessor(const FGLTFMeshSection* MeshSection, const FStaticMeshVertexBuffer* VertexBuffer);
+	FGLTFJsonAccessorIndex GetOrAddUVAccessor(const FGLTFMeshSection* MeshSection, const FStaticMeshVertexBuffer* VertexBuffer, int32 UVIndex);
+	FGLTFJsonAccessorIndex GetOrAddJointAccessor(const FGLTFMeshSection* MeshSection, const FSkinWeightVertexBuffer* VertexBuffer, int32 InfluenceOffset);
+	FGLTFJsonAccessorIndex GetOrAddWeightAccessor(const FGLTFMeshSection* MeshSection, const FSkinWeightVertexBuffer* VertexBuffer, int32 InfluenceOffset);
+	FGLTFJsonAccessorIndex GetOrAddIndexAccessor(const FGLTFMeshSection* MeshSection);
 
 	FGLTFJsonMeshIndex GetOrAddMesh(const UStaticMesh* StaticMesh, int32 LODIndex = -1, const FColorVertexBuffer* OverrideVertexColors = nullptr, const FGLTFMaterialArray& OverrideMaterials = {});
 	FGLTFJsonMeshIndex GetOrAddMesh(const UStaticMeshComponent* StaticMeshComponent);
@@ -68,8 +67,7 @@ private:
 	FGLTFUVBufferConverter UVBufferConverter{ *this };
 	FGLTFBoneIndexBufferConverter BoneIndexBufferConverter{ *this };
 	FGLTFBoneWeightBufferConverter BoneWeightBufferConverter{ *this };
-	FGLTFStaticMeshSectionConverter StaticMeshSectionConverter{ *this };
-	FGLTFSkeletalMeshSectionConverter SkeletalMeshSectionConverter{ *this };
+	FGLTFIndexBufferConverter IndexBufferConverter{ *this };
 
 	FGLTFStaticMeshConverter StaticMeshConverter{ *this };
 	FGLTFSkeletalMeshConverter SkeletalMeshConverter{ *this };
diff --git a/Engine/Plugins/Enterprise/GLTFExporter/MigratedFromGit/Source/GLTFExporter/Private/Converters/GLTFAccessorConverters.cpp b/Engine/Plugins/Enterprise/GLTFExporter/MigratedFromGit/Source/GLTFExporter/Private/Converters/GLTFAccessorConverters.cpp
--- a/Engine/Plugins/Enterprise/GLTFExporter/MigratedFromGit/Source/GLTFExporter/Private/Converters/GLTFAccessorConverters.cpp
+++ b/Engine/Plugins/Enterprise/GLTFExporter/MigratedFromGit/Source/GLTFExporter/Private/Converters/GLTFAccessorConverters.cpp
@@ -4,24 +4,24 @@
 #include "Converters/GLTFConverterUtility.h"
 #include "Converters/GLTFSkinWeightVertexBufferHack.h"
 #include "Builders/GLTFConvertBuilder.h"
-#include "Math/NumericLimits.h"
-#include "Rendering/MultiSizeIndexContainer.h"
-#include "Rendering/SkeletalMeshRenderData.h"
 
-FGLTFJsonAccessorIndex FGLTFPositionBufferConverter::Convert(const FPositionVertexBuffer* VertexBuffer)
+FGLTFJsonAccessorIndex FGLTFPositionBufferConverter::Convert(const FGLTFMeshSection* MeshSection, const FPositionVertexBuffer* VertexBuffer)
 {
-	const uint32 VertexCount = VertexBuffer->GetNumVertices();
-	if (VertexCount == 0)
+	if (VertexBuffer == nullptr || VertexBuffer->GetNumVertices() == 0)
 	{
 		return FGLTFJsonAccessorIndex(INDEX_NONE);
 	}
 
+	const TArray<uint32>& IndexMap = MeshSection->IndexMap;
+	const uint32 VertexCount = IndexMap.Num();
+
 	TArray<FGLTFJsonVector3> Positions;
 	Positions.AddUninitialized(VertexCount);
 
 	for (uint32 VertexIndex = 0; VertexIndex < VertexCount; ++VertexIndex)
 	{
-		Positions[VertexIndex] = FGLTFConverterUtility::ConvertPosition(VertexBuffer->VertexPosition(VertexIndex), Builder.ExportOptions->ExportScale);
+		const uint32 MappedVertexIndex = IndexMap[VertexIndex];
+		Positions[VertexIndex] = FGLTFConverterUtility::ConvertPosition(VertexBuffer->VertexPosition(MappedVertexIndex), Builder.ExportOptions->ExportScale);
 	}
 
 	// More accurate bounding box if based on raw vertex values
@@ -56,20 +56,23 @@ FGLTFJsonAccessorIndex FGLTFPositionBufferConverter::Convert(const FPositionVert
 	return Builder.AddAccessor(JsonAccessor);
 }
 
-FGLTFJsonAccessorIndex FGLTFColorBufferConverter::Convert(const FColorVertexBuffer* VertexBuffer)
+FGLTFJsonAccessorIndex FGLTFColorBufferConverter::Convert(const FGLTFMeshSection* MeshSection, const FColorVertexBuffer* VertexBuffer)
 {
-	const uint32 VertexCount = VertexBuffer->GetNumVertices();
-	if (VertexCount == 0)
+	if (VertexBuffer == nullptr || VertexBuffer->GetNumVertices() == 0)
 	{
 		return FGLTFJsonAccessorIndex(INDEX_NONE);
 	}
 
+	const TArray<uint32>& IndexMap = MeshSection->IndexMap;
+	const uint32 VertexCount = IndexMap.Num();
+
 	TArray<FGLTFPackedColor> Colors;
 	Colors.AddUninitialized(VertexCount);
 
 	for (uint32 VertexIndex = 0; VertexIndex < VertexCount; ++VertexIndex)
 	{
-		Colors[VertexIndex] = FGLTFConverterUtility::ConvertColor(VertexBuffer->VertexColor(VertexIndex));
+		const uint32 MappedVertexIndex = IndexMap[VertexIndex];
+		Colors[VertexIndex] = FGLTFConverterUtility::ConvertColor(VertexBuffer->VertexColor(MappedVertexIndex));
 	}
 
 	FGLTFJsonAccessor JsonAccessor;
@@ -82,20 +85,23 @@ FGLTFJsonAccessorIndex FGLTFColorBufferConverter::Convert(const FColorVertexBuff
 	return Builder.AddAccessor(JsonAccessor);
 }
 
-FGLTFJsonAccessorIndex FGLTFNormalBufferConverter::Convert(const FStaticMeshVertexBuffer* VertexBuffer)
+FGLTFJsonAccessorIndex FGLTFNormalBufferConverter::Convert(const FGLTFMeshSection* MeshSection, const FStaticMeshVertexBuffer* VertexBuffer)
 {
-	const uint32 VertexCount = VertexBuffer->GetNumVertices();
-	if (VertexCount == 0)
+	if (VertexBuffer == nullptr || VertexBuffer->GetNumVertices() == 0)
 	{
 		return FGLTFJsonAccessorIndex(INDEX_NONE);
 	}
 
+	const TArray<uint32>& IndexMap = MeshSection->IndexMap;
+	const uint32 VertexCount = IndexMap.Num();
+
 	TArray<FGLTFJsonVector3> Normals;
 	Normals.AddUninitialized(VertexCount);
 
 	for (uint32 VertexIndex = 0; VertexIndex < VertexCount; ++VertexIndex)
 	{
-		Normals[VertexIndex] = FGLTFConverterUtility::ConvertNormal(VertexBuffer->VertexTangentZ(VertexIndex));
+		const uint32 MappedVertexIndex = IndexMap[VertexIndex];
+		Normals[VertexIndex] = FGLTFConverterUtility::ConvertNormal(VertexBuffer->VertexTangentZ(MappedVertexIndex));
 	}
 
 	FGLTFJsonAccessor JsonAccessor;
@@ -107,20 +113,23 @@ FGLTFJsonAccessorIndex FGLTFNormalBufferConverter::Convert(const FStaticMeshVert
 	return Builder.AddAccessor(JsonAccessor);
 }
 
-FGLTFJsonAccessorIndex FGLTFTangentBufferConverter::Convert(const FStaticMeshVertexBuffer* VertexBuffer)
+FGLTFJsonAccessorIndex FGLTFTangentBufferConverter::Convert(const FGLTFMeshSection* MeshSection, const FStaticMeshVertexBuffer* VertexBuffer)
 {
-	const uint32 VertexCount = VertexBuffer->GetNumVertices();
-	if (VertexCount == 0)
+	if (VertexBuffer == nullptr || VertexBuffer->GetNumVertices() == 0)
 	{
 		return FGLTFJsonAccessorIndex(INDEX_NONE);
 	}
 
+	const TArray<uint32>& IndexMap = MeshSection->IndexMap;
+	const uint32 VertexCount = IndexMap.Num();
+
 	TArray<FGLTFJsonVector4> Tangents;
 	Tangents.AddUninitialized(VertexCount);
 
 	for (uint32 VertexIndex = 0; VertexIndex < VertexCount; ++VertexIndex)
 	{
-		Tangents[VertexIndex] = FGLTFConverterUtility::ConvertTangent(VertexBuffer->VertexTangentX(VertexIndex));
+		const uint32 MappedVertexIndex = IndexMap[VertexIndex];
+		Tangents[VertexIndex] = FGLTFConverterUtility::ConvertTangent(VertexBuffer->VertexTangentX(MappedVertexIndex));
 	}
 
 	FGLTFJsonAccessor JsonAccessor;
@@ -132,11 +141,20 @@ FGLTFJsonAccessorIndex FGLTFTangentBufferConverter::Convert(const FStaticMeshVer
 	return Builder.AddAccessor(JsonAccessor);
 }
 
-FGLTFJsonAccessorIndex FGLTFUVBufferConverter::Convert(const FStaticMeshVertexBuffer* VertexBuffer, int32 UVIndex)
+FGLTFJsonAccessorIndex FGLTFUVBufferConverter::Convert(const FGLTFMeshSection* MeshSection, const FStaticMeshVertexBuffer* VertexBuffer, uint32 UVIndex)
 {
-	const uint32 VertexCount = VertexBuffer->GetNumVertices();
-	if (VertexCount == 0 || UVIndex < 0 || VertexBuffer->GetNumTexCoords() <= static_cast<uint32>(UVIndex))
+	if (VertexBuffer == nullptr || VertexBuffer->GetNumVertices() == 0)
+	{
+		return FGLTFJsonAccessorIndex(INDEX_NONE);
+	}
+
+	const TArray<uint32>& IndexMap = MeshSection->IndexMap;
+	const uint32 VertexCount = IndexMap.Num();
+
+	const uint32 UVCount = VertexBuffer->GetNumTexCoords();
+	if (UVIndex >= UVCount)
 	{
+		// TODO: report warning
 		return FGLTFJsonAccessorIndex(INDEX_NONE);
 	}
 
@@ -145,7 +163,8 @@ FGLTFJsonAccessorIndex FGLTFUVBufferConverter::Convert(const FStaticMeshVertexBu
 
 	for (uint32 VertexIndex = 0; VertexIndex < VertexCount; ++VertexIndex)
 	{
-		UVs[VertexIndex] = FGLTFConverterUtility::ConvertUV(VertexBuffer->GetVertexUV(VertexIndex, UVIndex));
+		const uint32 MappedVertexIndex = IndexMap[VertexIndex];
+		UVs[VertexIndex] = FGLTFConverterUtility::ConvertUV(VertexBuffer->GetVertexUV(MappedVertexIndex, UVIndex));
 	}
 
 	FGLTFJsonAccessor JsonAccessor;
@@ -157,23 +176,31 @@ FGLTFJsonAccessorIndex FGLTFUVBufferConverter::Convert(const FStaticMeshVertexBu
 	return Builder.AddAccessor(JsonAccessor);
 }
 
-FGLTFJsonAccessorIndex FGLTFBoneIndexBufferConverter::Convert(const FSkinWeightVertexBuffer* VertexBuffer, int32 InfluenceOffset, FGLTFBoneMap BoneMap)
+FGLTFJsonAccessorIndex FGLTFBoneIndexBufferConverter::Convert(const FGLTFMeshSection* MeshSection, const FSkinWeightVertexBuffer* VertexBuffer, uint32 InfluenceOffset)
 {
-	return VertexBuffer->Use16BitBoneIndex() ? Convert<uint16>(VertexBuffer, InfluenceOffset, BoneMap) : Convert<uint8>(VertexBuffer, InfluenceOffset, BoneMap);
+	return MeshSection->MaxBoneIndex <= UINT8_MAX
+		? Convert<uint8>(MeshSection, VertexBuffer, InfluenceOffset)
+		: Convert<uint16>(MeshSection, VertexBuffer, InfluenceOffset);
 }
 
 template <typename IndexType>
-FGLTFJsonAccessorIndex FGLTFBoneIndexBufferConverter::Convert(const FSkinWeightVertexBuffer* VertexBuffer, int32 JointsGroupIndex, FGLTFBoneMap BoneMap)
+FGLTFJsonAccessorIndex FGLTFBoneIndexBufferConverter::Convert(const FGLTFMeshSection* MeshSection, const FSkinWeightVertexBuffer* VertexBuffer, uint32 InfluenceOffset) const
 {
-	const uint32 VertexCount = VertexBuffer->GetNumVertices();
-	const int32 MaxInfluenceCount = VertexBuffer->GetMaxBoneInfluences();
-	const int32 InfluenceOffset = JointsGroupIndex * 4;
+	if (VertexBuffer == nullptr || VertexBuffer->GetNumVertices() == 0)
+	{
+		return FGLTFJsonAccessorIndex(INDEX_NONE);
+	}
 
-	if (VertexCount == 0 || InfluenceOffset < 0 || MaxInfluenceCount <= InfluenceOffset)
+	const uint32 InfluenceCount = VertexBuffer->GetMaxBoneInfluences();
+	if (InfluenceOffset >= InfluenceCount)
 	{
+		// TODO: report warning
 		return FGLTFJsonAccessorIndex(INDEX_NONE);
 	}
 
+	const TArray<uint32>& IndexMap = MeshSection->IndexMap;
+	const uint32 VertexCount = IndexMap.Num();
+
 	struct VertexBoneIndices
 	{
 		IndexType Index[4];
@@ -182,16 +209,18 @@ FGLTFJsonAccessorIndex FGLTFBoneIndexBufferConverter::Convert(const FSkinWeightV
 	TArray<VertexBoneIndices> BoneIndices;
 	BoneIndices.AddUninitialized(VertexCount);
 
+	const TArray<FBoneIndexType>& BoneMap = MeshSection->BoneMap;
+
 	for (uint32 VertexIndex = 0; VertexIndex < VertexCount; ++VertexIndex)
 	{
+		const uint32 MappedVertexIndex = IndexMap[VertexIndex];
+		VertexBoneIndices& VertexBones = BoneIndices[VertexIndex];
+
 		for (int32 InfluenceIndex = 0; InfluenceIndex < 4; ++InfluenceIndex)
 		{
 			// TODO: remove hack
-			const uint32 UnmappedBoneIndex = FGLTFSkinWeightVertexBufferHack(VertexBuffer).GetBoneIndex(VertexIndex, InfluenceOffset + InfluenceIndex);
-			const FBoneIndexType BoneIndex = BoneMap[UnmappedBoneIndex];
-			check(BoneIndex <= TNumericLimits<IndexType>::Max());
-
-			BoneIndices[VertexIndex].Index[InfluenceIndex] = static_cast<IndexType>(BoneIndex);
+			const uint32 UnmappedBoneIndex = FGLTFSkinWeightVertexBufferHack(VertexBuffer).GetBoneIndex(MappedVertexIndex, InfluenceOffset + InfluenceIndex);
+			VertexBones.Index[InfluenceIndex] = static_cast<IndexType>(BoneMap[UnmappedBoneIndex]);
 		}
 	}
 
@@ -204,17 +233,23 @@ FGLTFJsonAccessorIndex FGLTFBoneIndexBufferConverter::Convert(const FSkinWeightV
 	return Builder.AddAccessor(JsonAccessor);
 }
 
-FGLTFJsonAccessorIndex FGLTFBoneWeightBufferConverter::Convert(const FSkinWeightVertexBuffer* VertexBuffer, int32 WeightsGroupIndex)
+FGLTFJsonAccessorIndex FGLTFBoneWeightBufferConverter::Convert(const FGLTFMeshSection* MeshSection, const FSkinWeightVertexBuffer* VertexBuffer, uint32 InfluenceOffset)
 {
-	const uint32 VertexCount = VertexBuffer->GetNumVertices();
-	const int32 MaxInfluenceCount = VertexBuffer->GetMaxBoneInfluences();
-	const int32 InfluenceOffset = WeightsGroupIndex * 4;
+	if (VertexBuffer == nullptr || VertexBuffer->GetNumVertices() == 0)
+	{
+		return FGLTFJsonAccessorIndex(INDEX_NONE);
+	}
 
-	if (VertexCount == 0 || InfluenceOffset < 0 || MaxInfluenceCount <= InfluenceOffset)
+	const uint32 InfluenceCount = VertexBuffer->GetMaxBoneInfluences();
+	if (InfluenceOffset >= InfluenceCount)
 	{
+		// TODO: report warning
 		return FGLTFJsonAccessorIndex(INDEX_NONE);
 	}
 
+	const TArray<uint32>& IndexMap = MeshSection->IndexMap;
+	const uint32 VertexCount = IndexMap.Num();
+
 	struct VertexBoneWeights
 	{
 		uint8 Weights[4];
@@ -225,11 +260,14 @@ FGLTFJsonAccessorIndex FGLTFBoneWeightBufferConverter::Convert(const FSkinWeight
 
 	for (uint32 VertexIndex = 0; VertexIndex < VertexCount; ++VertexIndex)
 	{
-		for (int32 InfluenceIndex = 0; InfluenceIndex < 4; ++InfluenceIndex)
+		const uint32 MappedVertexIndex = IndexMap[VertexIndex];
+		VertexBoneWeights& VertexWeights = BoneWeights[VertexIndex];
+
+		for (uint32 InfluenceIndex = 0; InfluenceIndex < 4; ++InfluenceIndex)
 		{
 			// TODO: remove hack
-			const uint8 BoneWeight = FGLTFSkinWeightVertexBufferHack(VertexBuffer).GetBoneWeight(VertexIndex, InfluenceOffset + InfluenceIndex);
-			BoneWeights[VertexIndex].Weights[InfluenceIndex] = BoneWeight;
+			const uint8 BoneWeight = FGLTFSkinWeightVertexBufferHack(VertexBuffer).GetBoneWeight(MappedVertexIndex, InfluenceOffset + InfluenceIndex);
+			VertexWeights.Weights[InfluenceIndex] = BoneWeight;
 		}
 	}
 
@@ -243,49 +281,35 @@ FGLTFJsonAccessorIndex FGLTFBoneWeightBufferConverter::Convert(const FSkinWeight
 	return Builder.AddAccessor(JsonAccessor);
 }
 
-FGLTFJsonAccessorIndex FGLTFStaticMeshSectionConverter::Convert(const FStaticMeshSection* MeshSection, const FRawStaticIndexBuffer* IndexBuffer)
+FGLTFJsonAccessorIndex FGLTFIndexBufferConverter::Convert(const FGLTFMeshSection* MeshSection)
 {
-	const uint32 IndexCount = MeshSection->NumTriangles * 3;
-	if (IndexCount == 0)
-	{
-		return FGLTFJsonAccessorIndex(INDEX_NONE);
-	}
-
-	const bool bIs32Bit = IndexBuffer->Is32Bit();
-	const uint8 IndexTypeSize = bIs32Bit ? sizeof(uint32) : sizeof(uint16);
-	const uint32 FirstIndex = MeshSection->FirstIndex;
-
-	const void* IndexDataRootPtr = bIs32Bit ? static_cast<const void*>(IndexBuffer->AccessStream32()) : static_cast<const void*>(IndexBuffer->AccessStream16());
-	const void* IndexDataPtr = static_cast<const uint8*>(IndexDataRootPtr) + FirstIndex * IndexTypeSize;
-	const int32 IndexDataSize = IndexCount * IndexTypeSize;
-
-	FGLTFJsonAccessor JsonAccessor;
-	JsonAccessor.BufferView = Builder.AddBufferView(IndexDataPtr, IndexDataSize, IndexTypeSize, EGLTFJsonBufferTarget::ElementArrayBuffer);
-	JsonAccessor.ComponentType = bIs32Bit ? EGLTFJsonComponentType::U32 : EGLTFJsonComponentType::U16;
-	JsonAccessor.Count = IndexCount;
-	JsonAccessor.Type = EGLTFJsonAccessorType::Scalar;
-
-	return Builder.AddAccessor(JsonAccessor);
+	const uint32 MaxVertexIndex = MeshSection->IndexMap.Num() - 1;
+	if (MaxVertexIndex <= UINT8_MAX) return Convert<uint8>(MeshSection);
+	if (MaxVertexIndex <= UINT16_MAX) return Convert<uint16>(MeshSection);
+	return Convert<uint32>(MeshSection);
 }
 
-FGLTFJsonAccessorIndex FGLTFSkeletalMeshSectionConverter::Convert(const FSkelMeshRenderSection* MeshSection, const FRawStaticIndexBuffer16or32Interface* IndexBuffer)
+template <typename IndexType>
+FGLTFJsonAccessorIndex FGLTFIndexBufferConverter::Convert(const FGLTFMeshSection* MeshSection) const
 {
-	const uint32 IndexCount = MeshSection->NumTriangles * 3;
+	const TArray<uint32>& IndexBuffer = MeshSection->IndexBuffer;
+	const uint32 IndexCount = IndexBuffer.Num();
 	if (IndexCount == 0)
 	{
 		return FGLTFJsonAccessorIndex(INDEX_NONE);
 	}
 
-	const bool bIs32Bit = IndexBuffer->GetResourceDataSize() == IndexBuffer->Num() * sizeof(uint32);
-	const uint8 IndexTypeSize = bIs32Bit ? sizeof(uint32) : sizeof(uint16);
-	const uint32 FirstIndex = MeshSection->BaseIndex;
+	TArray<IndexType> Indices;
+	Indices.AddUninitialized(IndexCount);
 
-	const void* IndexDataPtr = const_cast<FRawStaticIndexBuffer16or32Interface*>(IndexBuffer)->GetPointerTo(FirstIndex);
-	const int32 IndexDataSize = IndexCount * IndexTypeSize;
+	for (uint32 Index = 0; Index < IndexCount; ++Index)
+	{
+		Indices[Index] = static_cast<IndexType>(IndexBuffer[Index]);
+	}
 
 	FGLTFJsonAccessor JsonAccessor;
-	JsonAccessor.BufferView = Builder.AddBufferView(IndexDataPtr, IndexDataSize, IndexTypeSize, EGLTFJsonBufferTarget::ElementArrayBuffer);
-	JsonAccessor.ComponentType = bIs32Bit ? EGLTFJsonComponentType::U32 : EGLTFJsonComponentType::U16;
+	JsonAccessor.BufferView = Builder.AddBufferView(Indices, sizeof(IndexType), EGLTFJsonBufferTarget::ElementArrayBuffer);
+	JsonAccessor.ComponentType = FGLTFConverterUtility::GetComponentType<IndexType>();
 	JsonAccessor.Count = IndexCount;
 	JsonAccessor.Type = EGLTFJsonAccessorType::Scalar;
 
diff --git a/Engine/Plugins/Enterprise/GLTFExporter/MigratedFromGit/Source/GLTFExporter/Private/Converters/GLTFAccessorConverters.h b/Engine/Plugins/Enterprise/GLTFExporter/MigratedFromGit/Source/GLTFExporter/Private/Converters/GLTFAccessorConverters.h
--- a/Engine/Plugins/Enterprise/GLTFExporter/MigratedFromGit/Source/GLTFExporter/Private/Converters/GLTFAccessorConverters.h
+++ b/Engine/Plugins/Enterprise/GLTFExporter/MigratedFromGit/Source/GLTFExporter/Private/Converters/GLTFAccessorConverters.h
@@ -5,7 +5,7 @@
 #include "Json/GLTFJsonIndex.h"
 #include "Converters/GLTFConverter.h"
 #include "Converters/GLTFBuilderContext.h"
-#include "Converters/GLTFBoneMap.h"
+#include "Converters/GLTFMeshSection.h"
 #include "Engine.h"
 
 template <typename... InputTypes>
@@ -14,68 +14,64 @@ class TGLTFAccessorConverter : public FGLTFBuilderContext, public TGLTFConverter
 	using FGLTFBuilderContext::FGLTFBuilderContext;
 };
 
-class FGLTFPositionBufferConverter final : public TGLTFAccessorConverter<const FPositionVertexBuffer*>
+class FGLTFPositionBufferConverter final : public TGLTFAccessorConverter<const FGLTFMeshSection*, const FPositionVertexBuffer*>
 {
 	using TGLTFAccessorConverter::TGLTFAccessorConverter;
 
-	virtual FGLTFJsonAccessorIndex Convert(const FPositionVertexBuffer* VertexBuffer) override;
+	virtual FGLTFJsonAccessorIndex Convert(const FGLTFMeshSection* MeshSection, const FPositionVertexBuffer* VertexBuffer) override;
 };
 
-class FGLTFColorBufferConverter final : public TGLTFAccessorConverter<const FColorVertexBuffer*>
+class FGLTFColorBufferConverter final : public TGLTFAccessorConverter<const FGLTFMeshSection*, const FColorVertexBuffer*>
 {
 	using TGLTFAccessorConverter::TGLTFAccessorConverter;
 
-	virtual FGLTFJsonAccessorIndex Convert(const FColorVertexBuffer* VertexBuffer) override;
+	virtual FGLTFJsonAccessorIndex Convert(const FGLTFMeshSection* MeshSection, const FColorVertexBuffer* VertexBuffer) override;
 };
 
-class FGLTFNormalBufferConverter final : public TGLTFAccessorConverter<const FStaticMeshVertexBuffer*>
+class FGLTFNormalBufferConverter final : public TGLTFAccessorConverter<const FGLTFMeshSection*, const FStaticMeshVertexBuffer*>
 {
 	using TGLTFAccessorConverter::TGLTFAccessorConverter;
 
-	virtual FGLTFJsonAccessorIndex Convert(const FStaticMeshVertexBuffer* VertexBuffer) override;
+	virtual FGLTFJsonAccessorIndex Convert(const FGLTFMeshSection* MeshSection, const FStaticMeshVertexBuffer* VertexBuffer) override;
 };
 
-class FGLTFTangentBufferConverter final : public TGLTFAccessorConverter<const FStaticMeshVertexBuffer*>
+class FGLTFTangentBufferConverter final : public TGLTFAccessorConverter<const FGLTFMeshSection*, const FStaticMeshVertexBuffer*>
 {
 	using TGLTFAccessorConverter::TGLTFAccessorConverter;
 
-	virtual FGLTFJsonAccessorIndex Convert(const FStaticMeshVertexBuffer* VertexBuffer) override;
+	virtual FGLTFJsonAccessorIndex Convert(const FGLTFMeshSection* MeshSection, const FStaticMeshVertexBuffer* VertexBuffer) override;
 };
 
-class FGLTFUVBufferConverter final : public TGLTFAccessorConverter<const FStaticMeshVertexBuffer*, int32>
+class FGLTFUVBufferConverter final : public TGLTFAccessorConverter<const FGLTFMeshSection*, const FStaticMeshVertexBuffer*, uint32>
 {
 	using TGLTFAccessorConverter::TGLTFAccessorConverter;
 
-	virtual FGLTFJsonAccessorIndex Convert(const FStaticMeshVertexBuffer* VertexBuffer, int32 UVIndex) override;
+	virtual FGLTFJsonAccessorIndex Convert(const FGLTFMeshSection* MeshSection, const FStaticMeshVertexBuffer* VertexBuffer, uint32 UVIndex) override;
 };
 
-class FGLTFBoneIndexBufferConverter final : public TGLTFAccessorConverter<const FSkinWeightVertexBuffer*, int32, FGLTFBoneMap>
+class FGLTFBoneIndexBufferConverter final : public TGLTFAccessorConverter<const FGLTFMeshSection*, const FSkinWeightVertexBuffer*, uint32>
 {
 	using TGLTFAccessorConverter::TGLTFAccessorConverter;
 
-	virtual FGLTFJsonAccessorIndex Convert(const FSkinWeightVertexBuffer* VertexBuffer, int32 InfluenceOffset, FGLTFBoneMap BoneMap) override;
+	virtual FGLTFJsonAccessorIndex Convert(const FGLTFMeshSection* MeshSection, const FSkinWeightVertexBuffer* VertexBuffer, uint32 InfluenceOffset) override;
 
 	template <typename IndexType>
-	FGLTFJsonAccessorIndex Convert(const FSkinWeightVertexBuffer* VertexBuffer, int32 JointsGroupIndex, FGLTFBoneMap BoneMap);
+	FGLTFJsonAccessorIndex Convert(const FGLTFMeshSection* MeshSection, const FSkinWeightVertexBuffer* VertexBuffer, uint32 InfluenceOffset) const;
 };
 
-class FGLTFBoneWeightBufferConverter final : public TGLTFAccessorConverter<const FSkinWeightVertexBuffer*, int32>
+class FGLTFBoneWeightBufferConverter final : public TGLTFAccessorConverter<const FGLTFMeshSection*, const FSkinWeightVertexBuffer*, uint32>
 {
 	using TGLTFAccessorConverter::TGLTFAccessorConverter;
 
-	virtual FGLTFJsonAccessorIndex Convert(const FSkinWeightVertexBuffer* VertexBuffer, int32 WeightsGroupIndex) override;
+	virtual FGLTFJsonAccessorIndex Convert(const FGLTFMeshSection* MeshSection, const FSkinWeightVertexBuffer* VertexBuffer, uint32 InfluenceOffset) override;
 };
 
-class FGLTFStaticMeshSectionConverter final : public TGLTFAccessorConverter<const FStaticMeshSection*, const FRawStaticIndexBuffer*>
+class FGLTFIndexBufferConverter final : public TGLTFAccessorConverter<const FGLTFMeshSection*>
 {
 	using TGLTFAccessorConverter::TGLTFAccessorConverter;
 
-	virtual FGLTFJsonAccessorIndex Convert(const FStaticMeshSection* MeshSection, const FRawStaticIndexBuffer* IndexBuffer) override;
-};
+	virtual FGLTFJsonAccessorIndex Convert(const FGLTFMeshSection* MeshSection) override;
 
-class FGLTFSkeletalMeshSectionConverter final : public TGLTFAccessorConverter<const FSkelMeshRenderSection*, const FRawStaticIndexBuffer16or32Interface*>
-{
-	using TGLTFAccessorConverter::TGLTFAccessorConverter;
-
-	virtual FGLTFJsonAccessorIndex Convert(const FSkelMeshRenderSection* MeshSection, const FRawStaticIndexBuffer16or32Interface* IndexBuffer) override;
+	template <typename IndexType>
+    FGLTFJsonAccessorIndex Convert(const FGLTFMeshSection* MeshSection) const;
 };
diff --git a/Engine/Plugins/Enterprise/GLTFExporter/MigratedFromGit/Source/GLTFExporter/Private/Converters/GLTFMeshConverters.cpp b/Engine/Plugins/Enterprise/GLTFExporter/MigratedFromGit/Source/GLTFExporter/Private/Converters/GLTFMeshConverters.cpp
--- a/Engine/Plugins/Enterprise/GLTFExporter/MigratedFromGit/Source/GLTFExporter/Private/Converters/GLTFMeshConverters.cpp
+++ b/Engine/Plugins/Enterprise/GLTFExporter/MigratedFromGit/Source/GLTFExporter/Private/Converters/GLTFMeshConverters.cpp
@@ -19,39 +19,45 @@ FGLTFJsonMeshIndex FGLTFStaticMeshConverter::Convert(const UStaticMesh* StaticMe
 	FGLTFJsonMesh JsonMesh;
 	JsonMesh.Name = FGLTFNameUtility::GetName(StaticMesh, LODIndex);
 
+	const FRawStaticIndexBuffer* IndexBuffer = &MeshLOD.IndexBuffer;
 	const FPositionVertexBuffer* PositionBuffer = &MeshLOD.VertexBuffers.PositionVertexBuffer;
 	const FStaticMeshVertexBuffer* VertexBuffer = &MeshLOD.VertexBuffers.StaticMeshVertexBuffer;
 	const FColorVertexBuffer* ColorBuffer = OverrideVertexColors != nullptr ? OverrideVertexColors : &MeshLOD.VertexBuffers.ColorVertexBuffer;
-	const FRawStaticIndexBuffer* IndexBuffer = &MeshLOD.IndexBuffer;
 
 	const int32 SectionCount = MeshLOD.Sections.Num();
 	JsonMesh.Primitives.AddDefaulted(SectionCount);
 
 	for (int32 SectionIndex = 0; SectionIndex < SectionCount; ++SectionIndex)
 	{
+		const FStaticMeshSection* Section = &MeshLOD.Sections[SectionIndex];
+		const FGLTFMeshSection* ConvertedSection = MeshSectionConverter.GetOrAdd(Section, IndexBuffer);
+
 		FGLTFJsonPrimitive& JsonPrimitive = JsonMesh.Primitives[SectionIndex];
-		const FStaticMeshSection& Section = MeshLOD.Sections[SectionIndex];
+		JsonPrimitive.Indices = Builder.GetOrAddIndexAccessor(ConvertedSection);
 
-		JsonPrimitive.Indices = Builder.GetOrAddIndexAccessor(&Section, IndexBuffer);
-		JsonPrimitive.Attributes.Position = Builder.GetOrAddPositionAccessor(PositionBuffer);
+		JsonPrimitive.Attributes.Position = Builder.GetOrAddPositionAccessor(ConvertedSection, PositionBuffer);
+		if (JsonPrimitive.Attributes.Position == INDEX_NONE)
+		{
+			// TODO: report warning
+		}
 
 		if (Builder.ExportOptions->bExportVertexColors)
 		{
-			JsonPrimitive.Attributes.Color0 = Builder.GetOrAddColorAccessor(ColorBuffer);
+			JsonPrimitive.Attributes.Color0 = Builder.GetOrAddColorAccessor(ConvertedSection, ColorBuffer);
 		}
 
-		JsonPrimitive.Attributes.Normal = Builder.GetOrAddNormalAccessor(VertexBuffer);
-		JsonPrimitive.Attributes.Tangent = Builder.GetOrAddTangentAccessor(VertexBuffer);
+		JsonPrimitive.Attributes.Normal = Builder.GetOrAddNormalAccessor(ConvertedSection, VertexBuffer);
+		JsonPrimitive.Attributes.Tangent = Builder.GetOrAddTangentAccessor(ConvertedSection, VertexBuffer);
 
 		const uint32 UVCount = VertexBuffer->GetNumTexCoords();
 		JsonPrimitive.Attributes.TexCoords.AddUninitialized(UVCount);
 
 		for (uint32 UVIndex = 0; UVIndex < UVCount; ++UVIndex)
 		{
-			JsonPrimitive.Attributes.TexCoords[UVIndex] = Builder.GetOrAddUVAccessor(VertexBuffer, UVIndex);
+			JsonPrimitive.Attributes.TexCoords[UVIndex] = Builder.GetOrAddUVAccessor(ConvertedSection, VertexBuffer, UVIndex);
 		}
 
-		const UMaterialInterface* Material = OverrideMaterials.GetOverride(StaticMesh->StaticMaterials, Section.MaterialIndex);
+		const UMaterialInterface* Material = OverrideMaterials.GetOverride(StaticMesh->StaticMaterials, Section->MaterialIndex);
 		if (Material != nullptr)
 		{
 			JsonPrimitive.Material = Builder.GetOrAddMaterial(Material);
@@ -75,38 +81,43 @@ FGLTFJsonMeshIndex FGLTFSkeletalMeshConverter::Convert(const USkeletalMesh* Skel
 	FGLTFJsonMesh JsonMesh;
 	JsonMesh.Name = FGLTFNameUtility::GetName(SkeletalMesh, LODIndex);
 
+	const FRawStaticIndexBuffer16or32Interface* IndexBuffer = MeshLOD.MultiSizeIndexContainer.GetIndexBuffer();
 	const FPositionVertexBuffer* PositionBuffer = &MeshLOD.StaticVertexBuffers.PositionVertexBuffer;
 	const FStaticMeshVertexBuffer* VertexBuffer = &MeshLOD.StaticVertexBuffers.StaticMeshVertexBuffer;
 	const FColorVertexBuffer* ColorBuffer = OverrideVertexColors != nullptr ? OverrideVertexColors : &MeshLOD.StaticVertexBuffers.ColorVertexBuffer;
 	const FSkinWeightVertexBuffer* SkinWeightBuffer = OverrideSkinWeights != nullptr ? OverrideSkinWeights : MeshLOD.GetSkinWeightVertexBuffer();
 
-	const FRawStaticIndexBuffer16or32Interface* IndexBuffer = MeshLOD.MultiSizeIndexContainer.GetIndexBuffer();
-
 	const int32 SectionCount = MeshLOD.RenderSections.Num();
 	JsonMesh.Primitives.AddDefaulted(SectionCount);
 
 	for (int32 SectionIndex = 0; SectionIndex < SectionCount; ++SectionIndex)
 	{
+		const FSkelMeshRenderSection* Section = &MeshLOD.RenderSections[SectionIndex];
+		const FGLTFMeshSection* ConvertedSection = MeshSectionConverter.GetOrAdd(Section, IndexBuffer);
+
 		FGLTFJsonPrimitive& JsonPrimitive = JsonMesh.Primitives[SectionIndex];
-		const FSkelMeshRenderSection& Section = MeshLOD.RenderSections[SectionIndex];
+		JsonPrimitive.Indices = Builder.GetOrAddIndexAccessor(ConvertedSection);
 
-		JsonPrimitive.Indices = Builder.GetOrAddIndexAccessor(&Section, IndexBuffer);
-		JsonPrimitive.Attributes.Position = Builder.GetOrAddPositionAccessor(PositionBuffer);
+		JsonPrimitive.Attributes.Position = Builder.GetOrAddPositionAccessor(ConvertedSection, PositionBuffer);
+		if (JsonPrimitive.Attributes.Position == INDEX_NONE)
+		{
+			// TODO: report warning
+		}
 
 		if (Builder.ExportOptions->bExportVertexColors)
 		{
-			JsonPrimitive.Attributes.Color0 = Builder.GetOrAddColorAccessor(ColorBuffer);
+			JsonPrimitive.Attributes.Color0 = Builder.GetOrAddColorAccessor(ConvertedSection, ColorBuffer);
 		}
 
-		JsonPrimitive.Attributes.Normal = Builder.GetOrAddNormalAccessor(VertexBuffer);
-		JsonPrimitive.Attributes.Tangent = Builder.GetOrAddTangentAccessor(VertexBuffer);
+		JsonPrimitive.Attributes.Normal = Builder.GetOrAddNormalAccessor(ConvertedSection, VertexBuffer);
+		JsonPrimitive.Attributes.Tangent = Builder.GetOrAddTangentAccessor(ConvertedSection, VertexBuffer);
 
 		const uint32 UVCount = VertexBuffer->GetNumTexCoords();
 		JsonPrimitive.Attributes.TexCoords.AddUninitialized(UVCount);
 
 		for (uint32 UVIndex = 0; UVIndex < UVCount; ++UVIndex)
 		{
-			JsonPrimitive.Attributes.TexCoords[UVIndex] = Builder.GetOrAddUVAccessor(VertexBuffer, UVIndex);
+			JsonPrimitive.Attributes.TexCoords[UVIndex] = Builder.GetOrAddUVAccessor(ConvertedSection, VertexBuffer, UVIndex);
 		}
 
 		/* TODO: enable export of vertex skin data when crash bug fixed
@@ -115,16 +126,14 @@ FGLTFJsonMeshIndex FGLTFSkeletalMeshConverter::Convert(const USkeletalMesh* Skel
 		JsonPrimitive.Attributes.Joints.AddUninitialized(GroupCount);
 		JsonPrimitive.Attributes.Weights.AddUninitialized(GroupCount);
 
-		const FGLTFBoneMap BoneMap = FGLTFBoneMap(Section.BoneMap);
-
 		for (uint32 GroupIndex = 0; GroupIndex < GroupCount; ++GroupIndex)
 		{
-			JsonPrimitive.Attributes.Joints[GroupIndex] = Builder.GetOrAddJointAccessor(SkinWeightBuffer, GroupIndex, BoneMap, Name + TEXT("_Joints") + FString::FromInt(GroupIndex));
-			JsonPrimitive.Attributes.Weights[GroupIndex] = Builder.GetOrAddWeightAccessor(SkinWeightBuffer, GroupIndex, Name + TEXT("_Weights") + FString::FromInt(GroupIndex));
+			JsonPrimitive.Attributes.Joints[GroupIndex] = Builder.GetOrAddJointAccessor(ConvertedSection, SkinWeightBuffer, GroupIndex * 4);
+			JsonPrimitive.Attributes.Weights[GroupIndex] = Builder.GetOrAddWeightAccessor(ConvertedSection, SkinWeightBuffer, GroupIndex * 4);
 		}
 		*/
 
-		const UMaterialInterface* Material = OverrideMaterials.GetOverride(SkeletalMesh->Materials, Section.MaterialIndex);
+		const UMaterialInterface* Material = OverrideMaterials.GetOverride(SkeletalMesh->Materials, Section->MaterialIndex);
 		if (Material != nullptr)
 		{
 			JsonPrimitive.Material = Builder.GetOrAddMaterial(Material);
diff --git a/Engine/Plugins/Enterprise/GLTFExporter/MigratedFromGit/Source/GLTFExporter/Private/Converters/GLTFMeshConverters.h b/Engine/Plugins/Enterprise/GLTFExporter/MigratedFromGit/Source/GLTFExporter/Private/Converters/GLTFMeshConverters.h
--- a/Engine/Plugins/Enterprise/GLTFExporter/MigratedFromGit/Source/GLTFExporter/Private/Converters/GLTFMeshConverters.h
+++ b/Engine/Plugins/Enterprise/GLTFExporter/MigratedFromGit/Source/GLTFExporter/Private/Converters/GLTFMeshConverters.h
@@ -5,6 +5,7 @@
 #include "Json/GLTFJsonIndex.h"
 #include "Converters/GLTFConverter.h"
 #include "Converters/GLTFBuilderContext.h"
+#include "Converters/GLTFMeshSection.h"
 #include "Converters/GLTFMaterialArray.h"
 #include "Engine.h"
 
@@ -14,11 +15,24 @@ class TGLTFMeshConverter : public FGLTFBuilderContext, public TGLTFConverter<FGL
 	using FGLTFBuilderContext::FGLTFBuilderContext;
 };
 
+template <typename MeshSectionType, typename IndexBufferType>
+class TGLTFMeshSectionConverter final : public TGLTFConverter<const FGLTFMeshSection*, const MeshSectionType*, const IndexBufferType*>
+{
+	TArray<TUniquePtr<FGLTFMeshSection>> Outputs;
+
+	const FGLTFMeshSection* Convert(const MeshSectionType* MeshSection, const IndexBufferType* IndexBuffer)
+	{
+		return Outputs.Add_GetRef(MakeUnique<FGLTFMeshSection>(MeshSection, IndexBuffer)).Get();
+	}
+};
+
 class FGLTFStaticMeshConverter final : public TGLTFMeshConverter<const UStaticMesh*, int32, const FColorVertexBuffer*, FGLTFMaterialArray>
 {
 	using TGLTFMeshConverter::TGLTFMeshConverter;
 
 	virtual FGLTFJsonMeshIndex Convert(const UStaticMesh* StaticMesh, int32 LODIndex, const FColorVertexBuffer* OverrideVertexColors, FGLTFMaterialArray OverrideMaterials) override;
+
+	TGLTFMeshSectionConverter<FStaticMeshSection, FRawStaticIndexBuffer> MeshSectionConverter;
 };
 
 class FGLTFSkeletalMeshConverter final : public TGLTFMeshConverter<const USkeletalMesh*, int32, const FColorVertexBuffer*, const FSkinWeightVertexBuffer*, FGLTFMaterialArray>
@@ -26,4 +40,6 @@ class FGLTFSkeletalMeshConverter final : public TGLTFMeshConverter<const USkelet
 	using TGLTFMeshConverter::TGLTFMeshConverter;
 
 	virtual FGLTFJsonMeshIndex Convert(const USkeletalMesh* SkeletalMesh, int32 LODIndex, const FColorVertexBuffer* OverrideVertexColors, const FSkinWeightVertexBuffer* OverrideSkinWeights, FGLTFMaterialArray OverrideMaterials) override;
+
+	TGLTFMeshSectionConverter<FSkelMeshRenderSection, FRawStaticIndexBuffer16or32Interface> MeshSectionConverter;
 };
