[GLTFExporter] refactor: rename GLTFBuilder to GLTFMessageBuilder and reorganize some members

commit ffa090f14f20e531db7e43efd6e3a07209fa311b
Date: Sun, 27 Sep 2020 11:45:38 +0200

    refactor: rename GLTFBuilder to GLTFMessageBuilder and reorganize some members
    (cherry picked from commit 4553f7d188c41c853714c2b07c7c2d4a56f44d7e​)

#skipci

diff --git a/Engine/Plugins/Enterprise/GLTFExporter/MigratedFromGit/Source/GLTFExporter/Private/Builders/GLTFBufferBuilder.cpp b/Engine/Plugins/Enterprise/GLTFExporter/MigratedFromGit/Source/GLTFExporter/Private/Builders/GLTFBufferBuilder.cpp
--- a/Engine/Plugins/Enterprise/GLTFExporter/MigratedFromGit/Source/GLTFExporter/Private/Builders/GLTFBufferBuilder.cpp
+++ b/Engine/Plugins/Enterprise/GLTFExporter/MigratedFromGit/Source/GLTFExporter/Private/Builders/GLTFBufferBuilder.cpp
@@ -4,8 +4,7 @@
 #include "Misc/FileHelper.h"
 #include "Misc/Paths.h"
 
-FGLTFBufferBuilder::FGLTFBufferBuilder(const UGLTFExportOptions* ExportOptions)
-	: FGLTFJsonBuilder(ExportOptions)
+FGLTFBufferBuilder::FGLTFBufferBuilder()
 {
 	BufferIndex = AddBuffer();
 }
diff --git a/Engine/Plugins/Enterprise/GLTFExporter/MigratedFromGit/Source/GLTFExporter/Private/Builders/GLTFBufferBuilder.h b/Engine/Plugins/Enterprise/GLTFExporter/MigratedFromGit/Source/GLTFExporter/Private/Builders/GLTFBufferBuilder.h
--- a/Engine/Plugins/Enterprise/GLTFExporter/MigratedFromGit/Source/GLTFExporter/Private/Builders/GLTFBufferBuilder.h
+++ b/Engine/Plugins/Enterprise/GLTFExporter/MigratedFromGit/Source/GLTFExporter/Private/Builders/GLTFBufferBuilder.h
@@ -6,9 +6,11 @@
 
 class FGLTFBufferBuilder : public FGLTFJsonBuilder
 {
-public:
+protected:
+
+	FGLTFBufferBuilder();
 
-	FGLTFBufferBuilder(const UGLTFExportOptions* ExportOptions);
+public:
 
 	FGLTFJsonBufferViewIndex AddBufferView(const void* RawData, uint64 ByteLength, uint8 DataAlignment = 4, EGLTFJsonBufferTarget BufferTarget = EGLTFJsonBufferTarget::ArrayBuffer);
 
@@ -18,11 +20,12 @@ public:
 		return AddBufferView(Array.GetData(), Array.Num() * sizeof(ElementType), DataAlignment, BufferTarget);
 	}
 
-	void UpdateJsonBufferObject(const FString& BinaryFilePath);
 	virtual bool Serialize(FArchive& Archive, const FString& FilePath) override;
 
 private:
 
+	void UpdateJsonBufferObject(const FString& BinaryFilePath);
+
 	FGLTFJsonBufferIndex BufferIndex;
 	TArray64<uint8> BufferData;
 };
diff --git a/Engine/Plugins/Enterprise/GLTFExporter/MigratedFromGit/Source/GLTFExporter/Private/Builders/GLTFBuilder.cpp b/Engine/Plugins/Enterprise/GLTFExporter/MigratedFromGit/Source/GLTFExporter/Private/Builders/GLTFBuilder.cpp
deleted file mode 100644
diff --git a/Engine/Plugins/Enterprise/GLTFExporter/MigratedFromGit/Source/GLTFExporter/Private/Builders/GLTFBuilder.h b/Engine/Plugins/Enterprise/GLTFExporter/MigratedFromGit/Source/GLTFExporter/Private/Builders/GLTFBuilder.h
deleted file mode 100644
diff --git a/Engine/Plugins/Enterprise/GLTFExporter/MigratedFromGit/Source/GLTFExporter/Private/Builders/GLTFConvertBuilder.cpp b/Engine/Plugins/Enterprise/GLTFExporter/MigratedFromGit/Source/GLTFExporter/Private/Builders/GLTFConvertBuilder.cpp
--- a/Engine/Plugins/Enterprise/GLTFExporter/MigratedFromGit/Source/GLTFExporter/Private/Builders/GLTFConvertBuilder.cpp
+++ b/Engine/Plugins/Enterprise/GLTFExporter/MigratedFromGit/Source/GLTFExporter/Private/Builders/GLTFConvertBuilder.cpp
@@ -3,7 +3,7 @@
 #include "Builders/GLTFConvertBuilder.h"
 
 FGLTFConvertBuilder::FGLTFConvertBuilder(const UGLTFExportOptions* ExportOptions, bool bSelectedActorsOnly)
-	: FGLTFImageBuilder(ExportOptions)
+	: ExportOptions(ExportOptions)
 	, bSelectedActorsOnly(bSelectedActorsOnly)
 {
 }
diff --git a/Engine/Plugins/Enterprise/GLTFExporter/MigratedFromGit/Source/GLTFExporter/Private/Builders/GLTFConvertBuilder.h b/Engine/Plugins/Enterprise/GLTFExporter/MigratedFromGit/Source/GLTFExporter/Private/Builders/GLTFConvertBuilder.h
--- a/Engine/Plugins/Enterprise/GLTFExporter/MigratedFromGit/Source/GLTFExporter/Private/Builders/GLTFConvertBuilder.h
+++ b/Engine/Plugins/Enterprise/GLTFExporter/MigratedFromGit/Source/GLTFExporter/Private/Builders/GLTFConvertBuilder.h
@@ -2,6 +2,7 @@
 
 #pragma once
 
+#include "GLTFExportOptions.h"
 #include "Builders/GLTFImageBuilder.h"
 #include "Converters/GLTFAccessorConverters.h"
 #include "Converters/GLTFMeshConverters.h"
@@ -19,12 +20,15 @@
 
 class FGLTFConvertBuilder : public FGLTFImageBuilder
 {
+protected:
+
+	FGLTFConvertBuilder(const UGLTFExportOptions* ExportOptions, bool bSelectedActorsOnly);
+
 public:
 
+	const UGLTFExportOptions* const ExportOptions;
 	const bool bSelectedActorsOnly;
 
-	FGLTFConvertBuilder(const UGLTFExportOptions* ExportOptions, bool bSelectedActorsOnly);
-
 	FGLTFJsonAccessorIndex GetOrAddPositionAccessor(const FGLTFMeshSection* MeshSection, const FPositionVertexBuffer* VertexBuffer);
 	FGLTFJsonAccessorIndex GetOrAddColorAccessor(const FGLTFMeshSection* MeshSection, const FColorVertexBuffer* VertexBuffer);
 	FGLTFJsonAccessorIndex GetOrAddNormalAccessor(const FGLTFMeshSection* MeshSection, const FStaticMeshVertexBuffer* VertexBuffer);
diff --git a/Engine/Plugins/Enterprise/GLTFExporter/MigratedFromGit/Source/GLTFExporter/Private/Builders/GLTFImageBuilder.cpp b/Engine/Plugins/Enterprise/GLTFExporter/MigratedFromGit/Source/GLTFExporter/Private/Builders/GLTFImageBuilder.cpp
--- a/Engine/Plugins/Enterprise/GLTFExporter/MigratedFromGit/Source/GLTFExporter/Private/Builders/GLTFImageBuilder.cpp
+++ b/Engine/Plugins/Enterprise/GLTFExporter/MigratedFromGit/Source/GLTFExporter/Private/Builders/GLTFImageBuilder.cpp
@@ -6,8 +6,7 @@
 #include "IImageWrapper.h"
 #include "Misc/FileHelper.h"
 
-FGLTFImageBuilder::FGLTFImageBuilder(const UGLTFExportOptions* ExportOptions)
-	: FGLTFBufferBuilder(ExportOptions)
+FGLTFImageBuilder::FGLTFImageBuilder()
 {
 }
 
diff --git a/Engine/Plugins/Enterprise/GLTFExporter/MigratedFromGit/Source/GLTFExporter/Private/Builders/GLTFImageBuilder.h b/Engine/Plugins/Enterprise/GLTFExporter/MigratedFromGit/Source/GLTFExporter/Private/Builders/GLTFImageBuilder.h
--- a/Engine/Plugins/Enterprise/GLTFExporter/MigratedFromGit/Source/GLTFExporter/Private/Builders/GLTFImageBuilder.h
+++ b/Engine/Plugins/Enterprise/GLTFExporter/MigratedFromGit/Source/GLTFExporter/Private/Builders/GLTFImageBuilder.h
@@ -8,9 +8,11 @@ enum class ERGBFormat : int8;
 
 class FGLTFImageBuilder : public FGLTFBufferBuilder
 {
-public:
+protected:
+
+	FGLTFImageBuilder();
 
-	FGLTFImageBuilder(const UGLTFExportOptions* ExportOptions);
+public:
 
 	FGLTFJsonImageIndex AddImage(const void* CompressedData, int64 CompressedByteLength, EGLTFJsonMimeType MimeType, const FString& Name);
 	FGLTFJsonImageIndex AddImage(const void* RawData, int64 ByteLength, FIntPoint Size, ERGBFormat Format, int32 BitDepth, const FString& Name);
diff --git a/Engine/Plugins/Enterprise/GLTFExporter/MigratedFromGit/Source/GLTFExporter/Private/Builders/GLTFJsonBuilder.cpp b/Engine/Plugins/Enterprise/GLTFExporter/MigratedFromGit/Source/GLTFExporter/Private/Builders/GLTFJsonBuilder.cpp
--- a/Engine/Plugins/Enterprise/GLTFExporter/MigratedFromGit/Source/GLTFExporter/Private/Builders/GLTFJsonBuilder.cpp
+++ b/Engine/Plugins/Enterprise/GLTFExporter/MigratedFromGit/Source/GLTFExporter/Private/Builders/GLTFJsonBuilder.cpp
@@ -1,10 +1,10 @@
 // Copyright Epic Games, Inc. All Rights Reserved.
 
 #include "Builders/GLTFJsonBuilder.h"
+#include "Builders/GLTFBuilderUtility.h"
 
-FGLTFJsonBuilder::FGLTFJsonBuilder(const UGLTFExportOptions* ExportOptions)
-	: FGLTFBuilder(ExportOptions)
-	, DefaultScene(JsonRoot.DefaultScene)
+FGLTFJsonBuilder::FGLTFJsonBuilder()
+	: DefaultScene(JsonRoot.DefaultScene)
 {
 }
 
diff --git a/Engine/Plugins/Enterprise/GLTFExporter/MigratedFromGit/Source/GLTFExporter/Private/Builders/GLTFJsonBuilder.h b/Engine/Plugins/Enterprise/GLTFExporter/MigratedFromGit/Source/GLTFExporter/Private/Builders/GLTFJsonBuilder.h
--- a/Engine/Plugins/Enterprise/GLTFExporter/MigratedFromGit/Source/GLTFExporter/Private/Builders/GLTFJsonBuilder.h
+++ b/Engine/Plugins/Enterprise/GLTFExporter/MigratedFromGit/Source/GLTFExporter/Private/Builders/GLTFJsonBuilder.h
@@ -3,16 +3,18 @@
 #pragma once
 
 #include "Json/GLTFJsonRoot.h"
-#include "Builders/GLTFBuilder.h"
+#include "Builders/GLTFMessageBuilder.h"
 
-class FGLTFJsonBuilder : public FGLTFBuilder
+class FGLTFJsonBuilder : public FGLTFMessageBuilder
 {
+protected:
+
+	FGLTFJsonBuilder();
+
 public:
 
 	FGLTFJsonSceneIndex& DefaultScene;
 
-	FGLTFJsonBuilder(const UGLTFExportOptions* ExportOptions);
-
 	FGLTFJsonAccessorIndex AddAccessor(const FGLTFJsonAccessor& JsonAccessor = {});
 	FGLTFJsonBufferIndex AddBuffer(const FGLTFJsonBuffer& JsonBuffer = {});
 	FGLTFJsonBufferViewIndex AddBufferView(const FGLTFJsonBufferView& JsonBufferView = {});
diff --git a/Engine/Plugins/Enterprise/GLTFExporter/MigratedFromGit/Source/GLTFExporter/Private/Builders/GLTFMessageBuilder.cpp b/Engine/Plugins/Enterprise/GLTFExporter/MigratedFromGit/Source/GLTFExporter/Private/Builders/GLTFMessageBuilder.cpp
new file mode 100644
--- /dev/null
+++ b/Engine/Plugins/Enterprise/GLTFExporter/MigratedFromGit/Source/GLTFExporter/Private/Builders/GLTFMessageBuilder.cpp
@@ -0,0 +1,149 @@
+// Copyright Epic Games, Inc. All Rights Reserved.
+
+#include "Builders/GLTFMessageBuilder.h"
+#include "GLTFExporterModule.h"
+#include "Interfaces/IPluginManager.h"
+#include "MessageLogModule.h"
+#include "IMessageLogListing.h"
+
+FGLTFMessageBuilder::FGLTFMessageBuilder()
+{
+}
+
+FGLTFMessageBuilder::~FGLTFMessageBuilder()
+{
+}
+
+void FGLTFMessageBuilder::ClearMessages()
+{
+	Messages.Empty();
+}
+
+void FGLTFMessageBuilder::AddMessage(EGLTFMessageSeverity Severity, const FString& Message)
+{
+	Messages.Emplace(Severity, Message);
+}
+
+void FGLTFMessageBuilder::AddInfoMessage(const FString& Message)
+{
+	AddMessage(EGLTFMessageSeverity::Info, Message);
+}
+
+void FGLTFMessageBuilder::AddWarningMessage(const FString& Message)
+{
+	AddMessage(EGLTFMessageSeverity::Warning, Message);
+}
+
+void FGLTFMessageBuilder::AddErrorMessage(const FString& Message)
+{
+	AddMessage(EGLTFMessageSeverity::Error, Message);
+}
+
+const TArray<FGLTFMessageBuilder::FLogMessage>& FGLTFMessageBuilder::GetMessages() const
+{
+	return Messages;
+}
+
+TArray<FGLTFMessageBuilder::FLogMessage> FGLTFMessageBuilder::GetMessages(EGLTFMessageSeverity Severity) const
+{
+	return Messages.FilterByPredicate(
+		[Severity](const FLogMessage& LogMessage)
+		{
+			return LogMessage.Key == Severity;
+		});
+}
+
+TArray<FGLTFMessageBuilder::FLogMessage> FGLTFMessageBuilder::GetInfoMessages() const
+{
+	return GetMessages(EGLTFMessageSeverity::Info);
+}
+
+TArray<FGLTFMessageBuilder::FLogMessage> FGLTFMessageBuilder::GetWarningMessages() const
+{
+	return GetMessages(EGLTFMessageSeverity::Warning);
+}
+
+TArray<FGLTFMessageBuilder::FLogMessage> FGLTFMessageBuilder::GetErrorMessages() const
+{
+	return GetMessages(EGLTFMessageSeverity::Error);
+}
+
+int32 FGLTFMessageBuilder::GetMessageCount() const
+{
+	return Messages.Num();
+}
+
+int32 FGLTFMessageBuilder::GetInfoMessageCount() const
+{
+	return GetInfoMessages().Num();
+}
+
+int32 FGLTFMessageBuilder::GetWarningMessageCount() const
+{
+	return GetWarningMessages().Num();
+}
+
+int32 FGLTFMessageBuilder::GetErrorMessageCount() const
+{
+	return GetErrorMessages().Num();
+}
+
+void FGLTFMessageBuilder::ShowMessages() const
+{
+	if (Messages.Num() > 0)
+	{
+		FMessageLogModule& MessageLogModule = FModuleManager::LoadModuleChecked<FMessageLogModule>("MessageLog");
+		const TSharedRef<IMessageLogListing> LogListing = MessageLogModule.GetLogListing(GLTFEXPORTER_MODULE_NAME);
+
+		const TSharedPtr<IPlugin> Plugin = IPluginManager::Get().FindPlugin(GLTFEXPORTER_MODULE_NAME);
+		const FPluginDescriptor& PluginDescriptor = Plugin->GetDescriptor();
+
+		LogListing->SetLabel(FText::FromString(PluginDescriptor.FriendlyName));
+		LogListing->ClearMessages();
+
+		for (const FLogMessage& Message : Messages)
+		{
+			LogListing->AddMessage(CreateTokenizedMessage(Message));
+		}
+
+		MessageLogModule.OpenMessageLog(GLTFEXPORTER_MODULE_NAME);
+	}
+}
+
+void FGLTFMessageBuilder::WriteMessagesToConsole() const
+{
+	for (const FLogMessage& Message : Messages)
+	{
+		WriteMessageToConsole(Message);
+	}
+}
+
+void FGLTFMessageBuilder::WriteMessageToConsole(const FLogMessage& LogMessage)
+{
+	switch (LogMessage.Key)
+	{
+		case EGLTFMessageSeverity::Info:    UE_LOG(LogGLTFExporter, Display, TEXT("%s"), *LogMessage.Value); break;
+		case EGLTFMessageSeverity::Warning: UE_LOG(LogGLTFExporter, Warning, TEXT("%s"), *LogMessage.Value); break;
+		case EGLTFMessageSeverity::Error:   UE_LOG(LogGLTFExporter, Error, TEXT("%s"), *LogMessage.Value); break;
+		default:
+			checkNoEntry();
+			break;
+	}
+}
+
+TSharedRef<FTokenizedMessage> FGLTFMessageBuilder::CreateTokenizedMessage(const FLogMessage& LogMessage)
+{
+	EMessageSeverity::Type MessageSeverity = EMessageSeverity::Type::CriticalError;
+
+	switch (LogMessage.Key)
+	{
+		case EGLTFMessageSeverity::Info:    MessageSeverity = EMessageSeverity::Type::Info; break;
+		case EGLTFMessageSeverity::Warning: MessageSeverity = EMessageSeverity::Type::Warning; break;
+		case EGLTFMessageSeverity::Error:   MessageSeverity = EMessageSeverity::Type::Error; break;
+		default:
+			checkNoEntry();
+			break;
+	}
+
+	return FTokenizedMessage::Create(MessageSeverity, FText::FromString(LogMessage.Value));
+}
diff --git a/Engine/Plugins/Enterprise/GLTFExporter/MigratedFromGit/Source/GLTFExporter/Private/Builders/GLTFMessageBuilder.h b/Engine/Plugins/Enterprise/GLTFExporter/MigratedFromGit/Source/GLTFExporter/Private/Builders/GLTFMessageBuilder.h
new file mode 100644
--- /dev/null
+++ b/Engine/Plugins/Enterprise/GLTFExporter/MigratedFromGit/Source/GLTFExporter/Private/Builders/GLTFMessageBuilder.h
@@ -0,0 +1,62 @@
+// Copyright Epic Games, Inc. All Rights Reserved.
+
+#pragma once
+
+enum class EGLTFMessageSeverity
+{
+	Info,
+	Warning,
+	Error
+};
+
+class FGLTFMessageBuilder
+{
+protected:
+
+	FGLTFMessageBuilder();
+	virtual ~FGLTFMessageBuilder();
+
+public:
+
+	typedef TTuple<EGLTFMessageSeverity, FString> FLogMessage;
+
+	void ClearMessages();
+
+	void AddMessage(EGLTFMessageSeverity Severity, const FString& Message);
+
+	void AddInfoMessage(const FString& Message);
+
+	void AddWarningMessage(const FString& Message);
+
+	void AddErrorMessage(const FString& Message);
+
+	const TArray<FLogMessage>& GetMessages() const;
+
+	TArray<FLogMessage> GetMessages(EGLTFMessageSeverity Severity) const;
+
+	TArray<FLogMessage> GetInfoMessages() const;
+
+	TArray<FLogMessage> GetWarningMessages() const;
+
+	TArray<FLogMessage> GetErrorMessages() const;
+
+	int32 GetMessageCount() const;
+
+	int32 GetInfoMessageCount() const;
+
+	int32 GetWarningMessageCount() const;
+
+	int32 GetErrorMessageCount() const;
+
+	void ShowMessages() const;
+
+	void WriteMessagesToConsole() const;
+
+private:
+
+	static void WriteMessageToConsole(const FLogMessage& LogMessage);
+
+	static TSharedRef<FTokenizedMessage> CreateTokenizedMessage(const FLogMessage& LogMessage);
+
+	TArray<FLogMessage> Messages;
+};
diff --git a/Engine/Plugins/Enterprise/GLTFExporter/MigratedFromGit/Source/GLTFExporter/Private/Exporters/GLTFExporter.cpp b/Engine/Plugins/Enterprise/GLTFExporter/MigratedFromGit/Source/GLTFExporter/Private/Exporters/GLTFExporter.cpp
--- a/Engine/Plugins/Enterprise/GLTFExporter/MigratedFromGit/Source/GLTFExporter/Private/Exporters/GLTFExporter.cpp
+++ b/Engine/Plugins/Enterprise/GLTFExporter/MigratedFromGit/Source/GLTFExporter/Private/Exporters/GLTFExporter.cpp
@@ -57,11 +57,11 @@ bool UGLTFExporter::ExportBinary(UObject* Object, const TCHAR* Type, FArchive& A
 
 	if (FApp::IsUnattended())
 	{
-		Builder.WriteLogMessagesToConsole();
+		Builder.WriteMessagesToConsole();
 	}
 	else
 	{
-		Builder.ShowLogMessages();
+		Builder.ShowMessages();
 	}
 
 	return bSuccess;
diff --git a/Engine/Plugins/Enterprise/GLTFExporter/MigratedFromGit/Source/GLTFExporter/Private/GLTFExporterModule.cpp b/Engine/Plugins/Enterprise/GLTFExporter/MigratedFromGit/Source/GLTFExporter/Private/GLTFExporterModule.cpp
--- a/Engine/Plugins/Enterprise/GLTFExporter/MigratedFromGit/Source/GLTFExporter/Private/GLTFExporterModule.cpp
+++ b/Engine/Plugins/Enterprise/GLTFExporter/MigratedFromGit/Source/GLTFExporter/Private/GLTFExporterModule.cpp
@@ -5,6 +5,8 @@
 #include "CoreMinimal.h"
 #include "Modules/ModuleManager.h"
 
+DEFINE_LOG_CATEGORY(LogGLTFExporter);
+
 /**
  * glTF Exporter module implementation (private)
  */
diff --git a/Engine/Plugins/Enterprise/GLTFExporter/MigratedFromGit/Source/GLTFExporter/Public/GLTFExporterModule.h b/Engine/Plugins/Enterprise/GLTFExporter/MigratedFromGit/Source/GLTFExporter/Public/GLTFExporterModule.h
--- a/Engine/Plugins/Enterprise/GLTFExporter/MigratedFromGit/Source/GLTFExporter/Public/GLTFExporterModule.h
+++ b/Engine/Plugins/Enterprise/GLTFExporter/MigratedFromGit/Source/GLTFExporter/Public/GLTFExporterModule.h
@@ -8,6 +8,8 @@
 
 #define GLTFEXPORTER_MODULE_NAME TEXT("GLTFExporter")
 
+DECLARE_LOG_CATEGORY_EXTERN(LogGLTFExporter, Log, All);
+
 /**
  * The public interface of the GLTFExporter module
  */
