[GLTFExporter] refactor: change default parameter values for AddBufferView

commit 9c76e66f3e1eecc993c4ccbe76ac23010250b739
Date: Sun, 27 Sep 2020 20:35:51 +0200

    refactor: change default parameter values for AddBufferView
    (cherry picked from commit 6cd2f22c282273b8593c845332953f4abbe0f5c0​)

#skipci

diff --git a/Engine/Plugins/Enterprise/GLTFExporter/MigratedFromGit/Source/GLTFExporter/Private/Builders/GLTFBufferBuilder.cpp b/Engine/Plugins/Enterprise/GLTFExporter/MigratedFromGit/Source/GLTFExporter/Private/Builders/GLTFBufferBuilder.cpp
--- a/Engine/Plugins/Enterprise/GLTFExporter/MigratedFromGit/Source/GLTFExporter/Private/Builders/GLTFBufferBuilder.cpp
+++ b/Engine/Plugins/Enterprise/GLTFExporter/MigratedFromGit/Source/GLTFExporter/Private/Builders/GLTFBufferBuilder.cpp
@@ -9,7 +9,7 @@ FGLTFBufferBuilder::FGLTFBufferBuilder()
 	BufferIndex = AddBuffer();
 }
 
-FGLTFJsonBufferViewIndex FGLTFBufferBuilder::AddBufferView(const void* RawData, uint64 ByteLength, uint8 DataAlignment, EGLTFJsonBufferTarget BufferTarget)
+FGLTFJsonBufferViewIndex FGLTFBufferBuilder::AddBufferView(const void* RawData, uint64 ByteLength, EGLTFJsonBufferTarget BufferTarget, uint8 DataAlignment)
 {
 	uint64 ByteOffset = BufferData.Num();
 
diff --git a/Engine/Plugins/Enterprise/GLTFExporter/MigratedFromGit/Source/GLTFExporter/Private/Builders/GLTFBufferBuilder.h b/Engine/Plugins/Enterprise/GLTFExporter/MigratedFromGit/Source/GLTFExporter/Private/Builders/GLTFBufferBuilder.h
--- a/Engine/Plugins/Enterprise/GLTFExporter/MigratedFromGit/Source/GLTFExporter/Private/Builders/GLTFBufferBuilder.h
+++ b/Engine/Plugins/Enterprise/GLTFExporter/MigratedFromGit/Source/GLTFExporter/Private/Builders/GLTFBufferBuilder.h
@@ -12,12 +12,12 @@ protected:
 
 public:
 
-	FGLTFJsonBufferViewIndex AddBufferView(const void* RawData, uint64 ByteLength, uint8 DataAlignment = 4, EGLTFJsonBufferTarget BufferTarget = EGLTFJsonBufferTarget::ArrayBuffer);
+	FGLTFJsonBufferViewIndex AddBufferView(const void* RawData, uint64 ByteLength, EGLTFJsonBufferTarget BufferTarget = EGLTFJsonBufferTarget::None, uint8 DataAlignment = 4);
 
 	template <class ElementType>
-	FGLTFJsonBufferViewIndex AddBufferView(const TArray<ElementType>& Array, uint8 DataAlignment = 4, EGLTFJsonBufferTarget BufferTarget = EGLTFJsonBufferTarget::ArrayBuffer)
+	FGLTFJsonBufferViewIndex AddBufferView(const TArray<ElementType>& Array, EGLTFJsonBufferTarget BufferTarget = EGLTFJsonBufferTarget::None, uint8 DataAlignment = 4)
 	{
-		return AddBufferView(Array.GetData(), Array.Num() * sizeof(ElementType), DataAlignment, BufferTarget);
+		return AddBufferView(Array.GetData(), Array.Num() * sizeof(ElementType), BufferTarget, DataAlignment);
 	}
 
 	virtual bool Serialize(FArchive& Archive, const FString& FilePath) override;
diff --git a/Engine/Plugins/Enterprise/GLTFExporter/MigratedFromGit/Source/GLTFExporter/Private/Converters/GLTFAccessorConverters.cpp b/Engine/Plugins/Enterprise/GLTFExporter/MigratedFromGit/Source/GLTFExporter/Private/Converters/GLTFAccessorConverters.cpp
--- a/Engine/Plugins/Enterprise/GLTFExporter/MigratedFromGit/Source/GLTFExporter/Private/Converters/GLTFAccessorConverters.cpp
+++ b/Engine/Plugins/Enterprise/GLTFExporter/MigratedFromGit/Source/GLTFExporter/Private/Converters/GLTFAccessorConverters.cpp
@@ -40,7 +40,7 @@ FGLTFJsonAccessorIndex FGLTFPositionBufferConverter::Convert(const FGLTFMeshSect
 	}
 
 	FGLTFJsonAccessor JsonAccessor;
-	JsonAccessor.BufferView = Builder.AddBufferView(Positions);
+	JsonAccessor.BufferView = Builder.AddBufferView(Positions, EGLTFJsonBufferTarget::ArrayBuffer);
 	JsonAccessor.ComponentType = EGLTFJsonComponentType::F32;
 	JsonAccessor.Count = VertexCount;
 	JsonAccessor.Type = EGLTFJsonAccessorType::Vec3;
@@ -76,7 +76,7 @@ FGLTFJsonAccessorIndex FGLTFColorBufferConverter::Convert(const FGLTFMeshSection
 	}
 
 	FGLTFJsonAccessor JsonAccessor;
-	JsonAccessor.BufferView = Builder.AddBufferView(Colors);
+	JsonAccessor.BufferView = Builder.AddBufferView(Colors, EGLTFJsonBufferTarget::ArrayBuffer);
 	JsonAccessor.ComponentType = EGLTFJsonComponentType::U8;
 	JsonAccessor.Count = VertexCount;
 	JsonAccessor.Type = EGLTFJsonAccessorType::Vec4;
@@ -105,7 +105,7 @@ FGLTFJsonAccessorIndex FGLTFNormalBufferConverter::Convert(const FGLTFMeshSectio
 	}
 
 	FGLTFJsonAccessor JsonAccessor;
-	JsonAccessor.BufferView = Builder.AddBufferView(Normals);
+	JsonAccessor.BufferView = Builder.AddBufferView(Normals, EGLTFJsonBufferTarget::ArrayBuffer);
 	JsonAccessor.ComponentType = EGLTFJsonComponentType::F32;
 	JsonAccessor.Count = VertexCount;
 	JsonAccessor.Type = EGLTFJsonAccessorType::Vec3;
@@ -133,7 +133,7 @@ FGLTFJsonAccessorIndex FGLTFTangentBufferConverter::Convert(const FGLTFMeshSecti
 	}
 
 	FGLTFJsonAccessor JsonAccessor;
-	JsonAccessor.BufferView = Builder.AddBufferView(Tangents);
+	JsonAccessor.BufferView = Builder.AddBufferView(Tangents, EGLTFJsonBufferTarget::ArrayBuffer);
 	JsonAccessor.ComponentType = EGLTFJsonComponentType::F32;
 	JsonAccessor.Count = VertexCount;
 	JsonAccessor.Type = EGLTFJsonAccessorType::Vec4;
@@ -168,7 +168,7 @@ FGLTFJsonAccessorIndex FGLTFUVBufferConverter::Convert(const FGLTFMeshSection* M
 	}
 
 	FGLTFJsonAccessor JsonAccessor;
-	JsonAccessor.BufferView = Builder.AddBufferView(UVs);
+	JsonAccessor.BufferView = Builder.AddBufferView(UVs, EGLTFJsonBufferTarget::ArrayBuffer);
 	JsonAccessor.ComponentType = EGLTFJsonComponentType::F32;
 	JsonAccessor.Count = VertexCount;
 	JsonAccessor.Type = EGLTFJsonAccessorType::Vec2;
@@ -225,7 +225,7 @@ FGLTFJsonAccessorIndex FGLTFBoneIndexBufferConverter::Convert(const FGLTFMeshSec
 	}
 
 	FGLTFJsonAccessor JsonAccessor;
-	JsonAccessor.BufferView = Builder.AddBufferView(BoneIndices);
+	JsonAccessor.BufferView = Builder.AddBufferView(BoneIndices, EGLTFJsonBufferTarget::ArrayBuffer);
 	JsonAccessor.ComponentType = FGLTFConverterUtility::GetComponentType<IndexType>();
 	JsonAccessor.Count = VertexCount;
 	JsonAccessor.Type = EGLTFJsonAccessorType::Vec4;
@@ -272,7 +272,7 @@ FGLTFJsonAccessorIndex FGLTFBoneWeightBufferConverter::Convert(const FGLTFMeshSe
 	}
 
 	FGLTFJsonAccessor JsonAccessor;
-	JsonAccessor.BufferView = Builder.AddBufferView(BoneWeights);
+	JsonAccessor.BufferView = Builder.AddBufferView(BoneWeights, EGLTFJsonBufferTarget::ArrayBuffer);
 	JsonAccessor.ComponentType = EGLTFJsonComponentType::U8;
 	JsonAccessor.Count = VertexCount;
 	JsonAccessor.Type = EGLTFJsonAccessorType::Vec4;
@@ -308,7 +308,7 @@ FGLTFJsonAccessorIndex FGLTFIndexBufferConverter::Convert(const FGLTFMeshSection
 	}
 
 	FGLTFJsonAccessor JsonAccessor;
-	JsonAccessor.BufferView = Builder.AddBufferView(Indices, sizeof(IndexType), EGLTFJsonBufferTarget::ElementArrayBuffer);
+	JsonAccessor.BufferView = Builder.AddBufferView(Indices, EGLTFJsonBufferTarget::ElementArrayBuffer, sizeof(IndexType));
 	JsonAccessor.ComponentType = FGLTFConverterUtility::GetComponentType<IndexType>();
 	JsonAccessor.Count = IndexCount;
 	JsonAccessor.Type = EGLTFJsonAccessorType::Scalar;
