[GLTFExporter] feat(options): prohibit user from inputting negative UGLTFExportOptions::DefaultLevelOfDetail

commit f7db49fbfd2163d9560b324849563678392da74d
Date: Mon, 28 Sep 2020 08:16:49 +0200

    feat(options): prohibit user from inputting negative UGLTFExportOptions::DefaultLevelOfDetail
    (cherry picked from commit 7a6d73c0066246d26521e1850ba498e197ac4e13​)

#skipci

diff --git a/Engine/Plugins/Enterprise/GLTFExporter/MigratedFromGit/Source/GLTFExporter/Private/GLTFExportOptions.cpp b/Engine/Plugins/Enterprise/GLTFExporter/MigratedFromGit/Source/GLTFExporter/Private/GLTFExportOptions.cpp
--- a/Engine/Plugins/Enterprise/GLTFExporter/MigratedFromGit/Source/GLTFExporter/Private/GLTFExportOptions.cpp
+++ b/Engine/Plugins/Enterprise/GLTFExporter/MigratedFromGit/Source/GLTFExporter/Private/GLTFExportOptions.cpp
@@ -220,6 +220,19 @@ bool UGLTFExportOptions::CanEditChange(const FProperty* InProperty) const
 	return true;
 }
 
+void UGLTFExportOptions::PostEditChangeProperty(FPropertyChangedEvent& PropertyChangedEvent)
+{
+	const FName PropertyFName = PropertyChangedEvent.GetPropertyName();
+
+	if (PropertyFName == GET_MEMBER_NAME_CHECKED(ThisClass, DefaultLevelOfDetail))
+	{
+		if (DefaultLevelOfDetail < 0)
+		{
+			DefaultLevelOfDetail = 0;
+		}
+	}
+}
+
 FIntPoint UGLTFExportOptions::GetDefaultMaterialBakeSize() const
 {
 	const int32 Size = 1 << static_cast<int>(DefaultMaterialBakeSize);
diff --git a/Engine/Plugins/Enterprise/GLTFExporter/MigratedFromGit/Source/GLTFExporter/Public/GLTFExportOptions.h b/Engine/Plugins/Enterprise/GLTFExporter/MigratedFromGit/Source/GLTFExporter/Public/GLTFExportOptions.h
--- a/Engine/Plugins/Enterprise/GLTFExporter/MigratedFromGit/Source/GLTFExporter/Public/GLTFExportOptions.h
+++ b/Engine/Plugins/Enterprise/GLTFExporter/MigratedFromGit/Source/GLTFExporter/Public/GLTFExportOptions.h
@@ -164,6 +164,8 @@ public:
 
 	virtual bool CanEditChange(const FProperty* InProperty) const override;
 
+	virtual void PostEditChangeProperty(FPropertyChangedEvent& PropertyChangedEvent) override;
+
 	FIntPoint GetDefaultMaterialBakeSize() const;
 
 	bool ShouldExportLight(EComponentMobility::Type LightMobility) const;
