[GLTFExporter] refactor: add GLTFConverter::Sanitize and move some sanitation logic of mesh parameters

commit 022ea6d6ea658c8dd69f649a72e8273298954ce1
Date: Mon, 28 Sep 2020 13:12:30 +0200

    refactor: add GLTFConverter::Sanitize and move some sanitation logic of mesh parameters
    (cherry picked from commit 113e17a119b9a2816647d82a8b696b2cf0f39d0e​)

#skipci

diff --git a/Engine/Plugins/Enterprise/GLTFExporter/MigratedFromGit/Source/GLTFExporter/Private/Builders/GLTFConvertBuilder.cpp b/Engine/Plugins/Enterprise/GLTFExporter/MigratedFromGit/Source/GLTFExporter/Private/Builders/GLTFConvertBuilder.cpp
--- a/Engine/Plugins/Enterprise/GLTFExporter/MigratedFromGit/Source/GLTFExporter/Private/Builders/GLTFConvertBuilder.cpp
+++ b/Engine/Plugins/Enterprise/GLTFExporter/MigratedFromGit/Source/GLTFExporter/Private/Builders/GLTFConvertBuilder.cpp
@@ -95,17 +95,7 @@ FGLTFJsonMeshIndex FGLTFConvertBuilder::GetOrAddMesh(const UStaticMesh* StaticMe
 		return FGLTFJsonMeshIndex(INDEX_NONE);
 	}
 
-	if (LODIndex < 0)
-	{
-		LODIndex = ExportOptions->DefaultLevelOfDetail;
-	}
-
-	if (!ExportOptions->bExportVertexColors)
-	{
-		OverrideVertexColors = nullptr;
-	}
-
-	return StaticMeshConverter.GetOrAdd(StaticMesh, LODIndex, OverrideVertexColors, OverrideMaterials == StaticMesh->StaticMaterials ? FGLTFMaterialArray() : OverrideMaterials);
+	return StaticMeshConverter.GetOrAdd(StaticMesh, LODIndex, OverrideVertexColors, OverrideMaterials);
 }
 
 FGLTFJsonMeshIndex FGLTFConvertBuilder::GetOrAddMesh(const UStaticMeshComponent* StaticMeshComponent)
@@ -130,17 +120,7 @@ FGLTFJsonMeshIndex FGLTFConvertBuilder::GetOrAddMesh(const USkeletalMesh* Skelet
 		return FGLTFJsonMeshIndex(INDEX_NONE);
 	}
 
-	if (LODIndex < 0)
-	{
-		LODIndex = ExportOptions->DefaultLevelOfDetail;
-	}
-
-	if (!ExportOptions->bExportVertexColors)
-	{
-		OverrideVertexColors = nullptr;
-	}
-
-	return SkeletalMeshConverter.GetOrAdd(SkeletalMesh, LODIndex, OverrideVertexColors, OverrideSkinWeights, OverrideMaterials == SkeletalMesh->Materials ? FGLTFMaterialArray() : OverrideMaterials);
+	return SkeletalMeshConverter.GetOrAdd(SkeletalMesh, LODIndex, OverrideVertexColors, OverrideSkinWeights, OverrideMaterials);
 }
 
 FGLTFJsonMeshIndex FGLTFConvertBuilder::GetOrAddMesh(const USkeletalMeshComponent* SkeletalMeshComponent)
diff --git a/Engine/Plugins/Enterprise/GLTFExporter/MigratedFromGit/Source/GLTFExporter/Private/Converters/GLTFConverter.h b/Engine/Plugins/Enterprise/GLTFExporter/MigratedFromGit/Source/GLTFExporter/Private/Converters/GLTFConverter.h
--- a/Engine/Plugins/Enterprise/GLTFExporter/MigratedFromGit/Source/GLTFExporter/Private/Converters/GLTFConverter.h
+++ b/Engine/Plugins/Enterprise/GLTFExporter/MigratedFromGit/Source/GLTFExporter/Private/Converters/GLTFConverter.h
@@ -15,6 +15,7 @@ public:
 
 	OutputType Get(InputTypes&&... Inputs) const
 	{
+		Sanitize(Inputs...);
 		const InputKeyType InputKey(Forward<InputTypes>(Inputs)...);
 		if (OutputType* SavedOutput = SavedOutputs.Find(InputKey))
 		{
@@ -26,6 +27,7 @@ public:
 
 	OutputType Add(InputTypes... Inputs)
 	{
+		Sanitize(Inputs...);
 		const InputKeyType InputKey(Forward<InputTypes>(Inputs)...);
 		OutputType NewOutput = Convert(Forward<InputTypes>(Inputs)...);
 
@@ -35,6 +37,7 @@ public:
 
 	OutputType GetOrAdd(InputTypes... Inputs)
 	{
+		Sanitize(Inputs...);
 		const InputKeyType InputKey(Forward<InputTypes>(Inputs)...);
 		if (OutputType* SavedOutput = SavedOutputs.Find(InputKey))
 		{
@@ -49,6 +52,8 @@ public:
 
 protected:
 
+	virtual void Sanitize(InputTypes&... Inputs) { }
+
 	virtual OutputType Convert(InputTypes... Inputs) = 0;
 
 private:
diff --git a/Engine/Plugins/Enterprise/GLTFExporter/MigratedFromGit/Source/GLTFExporter/Private/Converters/GLTFMeshConverters.cpp b/Engine/Plugins/Enterprise/GLTFExporter/MigratedFromGit/Source/GLTFExporter/Private/Converters/GLTFMeshConverters.cpp
--- a/Engine/Plugins/Enterprise/GLTFExporter/MigratedFromGit/Source/GLTFExporter/Private/Converters/GLTFMeshConverters.cpp
+++ b/Engine/Plugins/Enterprise/GLTFExporter/MigratedFromGit/Source/GLTFExporter/Private/Converters/GLTFMeshConverters.cpp
@@ -7,6 +7,24 @@
 #include "Rendering/MultiSizeIndexContainer.h"
 #include "Rendering/SkeletalMeshRenderData.h"
 
+void FGLTFStaticMeshConverter::Sanitize(const UStaticMesh*& StaticMesh, int32& LODIndex, const FColorVertexBuffer*& OverrideVertexColors, FGLTFMaterialArray& OverrideMaterials)
+{
+	if (LODIndex < 0)
+	{
+		LODIndex = Builder.ExportOptions->DefaultLevelOfDetail;
+	}
+
+	if (!Builder.ExportOptions->bExportVertexColors)
+	{
+		OverrideVertexColors = nullptr;
+	}
+
+	if (OverrideMaterials == StaticMesh->StaticMaterials)
+	{
+		OverrideMaterials.Empty();
+	}
+}
+
 FGLTFJsonMeshIndex FGLTFStaticMeshConverter::Convert(const UStaticMesh* StaticMesh, int32 LODIndex, const FColorVertexBuffer* OverrideVertexColors, FGLTFMaterialArray OverrideMaterials)
 {
 	if (LODIndex < 0 || StaticMesh->GetNumLODs() <= LODIndex)
@@ -67,6 +85,24 @@ FGLTFJsonMeshIndex FGLTFStaticMeshConverter::Convert(const UStaticMesh* StaticMe
 	return Builder.AddMesh(JsonMesh);
 }
 
+void FGLTFSkeletalMeshConverter::Sanitize(const USkeletalMesh*& SkeletalMesh, int32& LODIndex, const FColorVertexBuffer*& OverrideVertexColors, const FSkinWeightVertexBuffer*& OverrideSkinWeights, FGLTFMaterialArray& OverrideMaterials)
+{
+	if (LODIndex < 0)
+	{
+		LODIndex = Builder.ExportOptions->DefaultLevelOfDetail;
+	}
+
+	if (!Builder.ExportOptions->bExportVertexColors)
+	{
+		OverrideVertexColors = nullptr;
+	}
+
+	if (OverrideMaterials == SkeletalMesh->Materials)
+	{
+		OverrideMaterials.Empty();
+	}
+}
+
 FGLTFJsonMeshIndex FGLTFSkeletalMeshConverter::Convert(const USkeletalMesh* SkeletalMesh, int32 LODIndex, const FColorVertexBuffer* OverrideVertexColors, const FSkinWeightVertexBuffer* OverrideSkinWeights, FGLTFMaterialArray OverrideMaterials)
 {
 	const FSkeletalMeshRenderData* RenderData = SkeletalMesh->GetResourceForRendering();
diff --git a/Engine/Plugins/Enterprise/GLTFExporter/MigratedFromGit/Source/GLTFExporter/Private/Converters/GLTFMeshConverters.h b/Engine/Plugins/Enterprise/GLTFExporter/MigratedFromGit/Source/GLTFExporter/Private/Converters/GLTFMeshConverters.h
--- a/Engine/Plugins/Enterprise/GLTFExporter/MigratedFromGit/Source/GLTFExporter/Private/Converters/GLTFMeshConverters.h
+++ b/Engine/Plugins/Enterprise/GLTFExporter/MigratedFromGit/Source/GLTFExporter/Private/Converters/GLTFMeshConverters.h
@@ -30,6 +30,8 @@ class FGLTFStaticMeshConverter final : public TGLTFMeshConverter<const UStaticMe
 {
 	using TGLTFMeshConverter::TGLTFMeshConverter;
 
+	virtual void Sanitize(const UStaticMesh*& StaticMesh, int32& LODIndex, const FColorVertexBuffer*& OverrideVertexColors, FGLTFMaterialArray& OverrideMaterials) override;
+
 	virtual FGLTFJsonMeshIndex Convert(const UStaticMesh* StaticMesh, int32 LODIndex, const FColorVertexBuffer* OverrideVertexColors, FGLTFMaterialArray OverrideMaterials) override;
 
 	TGLTFMeshSectionConverter<FStaticMeshSection, FRawStaticIndexBuffer> MeshSectionConverter;
@@ -39,6 +41,8 @@ class FGLTFSkeletalMeshConverter final : public TGLTFMeshConverter<const USkelet
 {
 	using TGLTFMeshConverter::TGLTFMeshConverter;
 
+	virtual void Sanitize(const USkeletalMesh*& SkeletalMesh, int32& LODIndex, const FColorVertexBuffer*& OverrideVertexColors, const FSkinWeightVertexBuffer*& OverrideSkinWeights, FGLTFMaterialArray& OverrideMaterials) override;
+
 	virtual FGLTFJsonMeshIndex Convert(const USkeletalMesh* SkeletalMesh, int32 LODIndex, const FColorVertexBuffer* OverrideVertexColors, const FSkinWeightVertexBuffer* OverrideSkinWeights, FGLTFMaterialArray OverrideMaterials) override;
 
 	TGLTFMeshSectionConverter<FSkelMeshRenderSection, FRawStaticIndexBuffer16or32Interface> MeshSectionConverter;
