[GLTFExporter] feat(orbit-camera): allow null as Focus, but prevent focusing the actor itself

commit 03e8c43197cbd499845fba72b34945e71c42ab84
Date: Mon, 28 Sep 2020 08:47:27 +0200

    feat(orbit-camera): allow null as Focus, but prevent focusing the actor itself
    (cherry picked from commit 9f6d6ca8bb03cbeb45c518d733d871abaf9fca98​)

#skipci

diff --git a/Engine/Plugins/Enterprise/GLTFExporter/MigratedFromGit/Source/GLTFExporter/Private/Converters/GLTFCameraConverters.cpp b/Engine/Plugins/Enterprise/GLTFExporter/MigratedFromGit/Source/GLTFExporter/Private/Converters/GLTFCameraConverters.cpp
--- a/Engine/Plugins/Enterprise/GLTFExporter/MigratedFromGit/Source/GLTFExporter/Private/Converters/GLTFCameraConverters.cpp
+++ b/Engine/Plugins/Enterprise/GLTFExporter/MigratedFromGit/Source/GLTFExporter/Private/Converters/GLTFCameraConverters.cpp
@@ -38,27 +38,18 @@ FGLTFJsonCameraIndex FGLTFCameraConverter::Convert(const UCameraComponent* Camer
 	{
 		if (Builder.ExportOptions->bExportOrbitalCameras)
 		{
-			if (OrbitCameraActor->Focus != nullptr)
-			{
-				FGLTFJsonOrbitCamera OrbitCamera;
-				OrbitCamera.Focus = Builder.GetOrAddNode(OrbitCameraActor->Focus);
-				OrbitCamera.MaxDistance = FGLTFConverterUtility::ConvertLength(OrbitCameraActor->DistanceMax, Builder.ExportOptions->ExportScale);
-				OrbitCamera.MinDistance = FGLTFConverterUtility::ConvertLength(OrbitCameraActor->DistanceMin, Builder.ExportOptions->ExportScale);
-				OrbitCamera.MaxAngle = OrbitCameraActor->PitchAngleMax;
-				OrbitCamera.MinAngle = OrbitCameraActor->PitchAngleMin;
-				OrbitCamera.DistanceSensitivity = OrbitCameraActor->DistanceSensitivity;
-				OrbitCamera.OrbitSensitivity = OrbitCameraActor->OrbitSensitivity;
-				OrbitCamera.OrbitInertia = OrbitCameraActor->OrbitInertia;
-				OrbitCamera.DollyDuration = OrbitCameraActor->DollyDuration;
-
-				Camera.OrbitCamera = OrbitCamera;
-			}
-			else
-			{
-				Builder.AddWarningMessage(FString::Printf(
-					TEXT("OrbitalCamera %s has no focus set and will be skipped"),
-					*Owner->GetName()));
-			}
+			FGLTFJsonOrbitCamera OrbitCamera;
+			OrbitCamera.Focus = Builder.GetOrAddNode(OrbitCameraActor->Focus);
+			OrbitCamera.MaxDistance = FGLTFConverterUtility::ConvertLength(OrbitCameraActor->DistanceMax, Builder.ExportOptions->ExportScale);
+			OrbitCamera.MinDistance = FGLTFConverterUtility::ConvertLength(OrbitCameraActor->DistanceMin, Builder.ExportOptions->ExportScale);
+			OrbitCamera.MaxAngle = OrbitCameraActor->PitchAngleMax;
+			OrbitCamera.MinAngle = OrbitCameraActor->PitchAngleMin;
+			OrbitCamera.DistanceSensitivity = OrbitCameraActor->DistanceSensitivity;
+			OrbitCamera.OrbitSensitivity = OrbitCameraActor->OrbitSensitivity;
+			OrbitCamera.OrbitInertia = OrbitCameraActor->OrbitInertia;
+			OrbitCamera.DollyDuration = OrbitCameraActor->DollyDuration;
+
+			Camera.OrbitCamera = OrbitCamera;
 		}
 	}
 
diff --git a/Engine/Plugins/Enterprise/GLTFExporter/MigratedFromGit/Source/GLTFExporterRuntime/Private/Actors/GLTFOrbitCameraActor.cpp b/Engine/Plugins/Enterprise/GLTFExporter/MigratedFromGit/Source/GLTFExporterRuntime/Private/Actors/GLTFOrbitCameraActor.cpp
--- a/Engine/Plugins/Enterprise/GLTFExporter/MigratedFromGit/Source/GLTFExporterRuntime/Private/Actors/GLTFOrbitCameraActor.cpp
+++ b/Engine/Plugins/Enterprise/GLTFExporter/MigratedFromGit/Source/GLTFExporterRuntime/Private/Actors/GLTFOrbitCameraActor.cpp
@@ -41,17 +41,34 @@ AGLTFOrbitCameraActor::AGLTFOrbitCameraActor(const FObjectInitializer& ObjectIni
 	PrimaryActorTick.bCanEverTick = true;
 }
 
-void AGLTFOrbitCameraActor::BeginPlay()
+#if WITH_EDITOR
+void AGLTFOrbitCameraActor::PostEditChangeProperty(FPropertyChangedEvent& PropertyChangedEvent)
 {
-	Super::BeginPlay();
+	Super::PostEditChangeProperty(PropertyChangedEvent);
 
-	bool bHasValidFocusActor;
-	const FVector FocusPosition = GetFocusPosition(&bHasValidFocusActor);
+	const FProperty* PropertyThatChanged = PropertyChangedEvent.Property;
 
-	if (!bHasValidFocusActor)
+	if (PropertyThatChanged)
 	{
-		UE_LOG(LogTemp, Warning, TEXT("The camera focus must not be null, and must not be the camera's own actor"));
+		const FName PropertyName = PropertyThatChanged->GetFName();
+
+		if (PropertyName == GET_MEMBER_NAME_CHECKED(ThisClass, Focus))
+		{
+			if (Focus == this)
+			{
+				Focus = nullptr;
+				UE_LOG(LogTemp, Warning, TEXT("The camera focus must not be the camera's own actor"));
+			}
+		}
 	}
+}
+#endif // WITH_EDITOR
+
+void AGLTFOrbitCameraActor::BeginPlay()
+{
+	Super::BeginPlay();
+
+	const FVector FocusPosition = GetFocusPosition();
 
 	// Ensure that the camera is initially aimed at the focus-position
 	SetActorRotation(GetLookAtRotation(FocusPosition));
@@ -174,15 +191,9 @@ FRotator AGLTFOrbitCameraActor::GetLookAtRotation(const FVector TargetPosition)
 	return FRotationMatrix::MakeFromXZ(TargetPosition - EyePosition, FVector::UpVector).Rotator();
 }
 
-FVector AGLTFOrbitCameraActor::GetFocusPosition(bool* bOutHasValidFocusActor) const
+FVector AGLTFOrbitCameraActor::GetFocusPosition() const
 {
-	const bool bHasValidFocusActor = this->Focus != nullptr && this->Focus != this;
-	if (bOutHasValidFocusActor != nullptr)
-	{
-		*bOutHasValidFocusActor = bHasValidFocusActor;
-	}
-
-	return bHasValidFocusActor ? this->Focus->GetActorLocation() : FVector::ZeroVector;
+	return this->Focus != nullptr ? this->Focus->GetActorLocation() : FVector::ZeroVector;
 }
 
 bool AGLTFOrbitCameraActor::SetAutoActivateForPlayer(const EAutoReceiveInput::Type Player)
diff --git a/Engine/Plugins/Enterprise/GLTFExporter/MigratedFromGit/Source/GLTFExporterRuntime/Public/Actors/GLTFOrbitCameraActor.h b/Engine/Plugins/Enterprise/GLTFExporter/MigratedFromGit/Source/GLTFExporterRuntime/Public/Actors/GLTFOrbitCameraActor.h
--- a/Engine/Plugins/Enterprise/GLTFExporter/MigratedFromGit/Source/GLTFExporterRuntime/Public/Actors/GLTFOrbitCameraActor.h
+++ b/Engine/Plugins/Enterprise/GLTFExporter/MigratedFromGit/Source/GLTFExporterRuntime/Public/Actors/GLTFOrbitCameraActor.h
@@ -16,6 +16,7 @@ class GLTFEXPORTERRUNTIME_API AGLTFOrbitCameraActor : public ACameraActor
 	//~ Begin UObject Interface
 public:
 	AGLTFOrbitCameraActor(const FObjectInitializer& ObjectInitializer);
+	virtual void PostEditChangeProperty(FPropertyChangedEvent& PropertyChangedEvent) override;
 	//~ End UObject Interface
 
 	//~ Begin AActor Interface
@@ -47,7 +48,7 @@ private:
 
 	FRotator GetLookAtRotation(const FVector TargetPosition) const;
 
-	FVector GetFocusPosition(bool* bOutHasValidFocusActor = nullptr) const;
+	FVector GetFocusPosition() const;
 
 	bool SetAutoActivateForPlayer(const EAutoReceiveInput::Type Player);
 
