[GLTFExporter] feat: ignore color- and weight-overrides if they are identical to the original

commit b27f3de2cd03b1a20917f47e9b70dbb20cd8286f
Date: Mon, 28 Sep 2020 14:10:04 +0200

    feat: ignore color- and weight-overrides if they are identical to the original
    (cherry picked from commit 3c91e43aa970efef44e4490c0011650c98b09623​)

#skipci

diff --git a/Engine/Plugins/Enterprise/GLTFExporter/MigratedFromGit/Source/GLTFExporter/Private/Converters/GLTFMeshConverters.cpp b/Engine/Plugins/Enterprise/GLTFExporter/MigratedFromGit/Source/GLTFExporter/Private/Converters/GLTFMeshConverters.cpp
--- a/Engine/Plugins/Enterprise/GLTFExporter/MigratedFromGit/Source/GLTFExporter/Private/Converters/GLTFMeshConverters.cpp
+++ b/Engine/Plugins/Enterprise/GLTFExporter/MigratedFromGit/Source/GLTFExporter/Private/Converters/GLTFMeshConverters.cpp
@@ -23,6 +23,16 @@ void FGLTFStaticMeshConverter::Sanitize(const UStaticMesh*& StaticMesh, int32& L
 	{
 		OverrideMaterials.Empty();
 	}
+
+	if (LODIndex >= 0 && StaticMesh->GetNumLODs() > LODIndex)
+	{
+		const FStaticMeshLODResources& MeshLOD = StaticMesh->GetLODForExport(LODIndex);
+
+		if (OverrideVertexColors != nullptr && OverrideVertexColors == &MeshLOD.VertexBuffers.ColorVertexBuffer)
+		{
+			OverrideVertexColors = nullptr;
+		}
+	}
 }
 
 FGLTFJsonMeshIndex FGLTFStaticMeshConverter::Convert(const UStaticMesh* StaticMesh, int32 LODIndex, const FColorVertexBuffer* OverrideVertexColors, FGLTFMaterialArray OverrideMaterials)
@@ -101,6 +111,23 @@ void FGLTFSkeletalMeshConverter::Sanitize(const USkeletalMesh*& SkeletalMesh, in
 	{
 		OverrideMaterials.Empty();
 	}
+
+	const FSkeletalMeshRenderData* RenderData = SkeletalMesh->GetResourceForRendering();
+
+	if (LODIndex >= 0 && RenderData->LODRenderData.Num() > LODIndex)
+	{
+		const FSkeletalMeshLODRenderData& MeshLOD = RenderData->LODRenderData[LODIndex];
+
+		if (OverrideVertexColors != nullptr && OverrideVertexColors == &MeshLOD.StaticVertexBuffers.ColorVertexBuffer)
+		{
+			OverrideVertexColors = nullptr;
+		}
+
+		if (OverrideSkinWeights != nullptr && OverrideSkinWeights == MeshLOD.GetSkinWeightVertexBuffer())
+		{
+			OverrideSkinWeights = nullptr;
+		}
+	}
 }
 
 FGLTFJsonMeshIndex FGLTFSkeletalMeshConverter::Convert(const USkeletalMesh* SkeletalMesh, int32 LODIndex, const FColorVertexBuffer* OverrideVertexColors, const FSkinWeightVertexBuffer* OverrideSkinWeights, FGLTFMaterialArray OverrideMaterials)
