[GLTFExporter] refactor: make sure LODIndex is non-negative and simplify some checks

commit e99e064123703456c7b2bbbde5231fb85055d066
Date: Mon, 28 Sep 2020 18:19:15 +0200

    refactor: make sure LODIndex is non-negative and simplify some checks
    (cherry picked from commit 179ece0fd0c293d43fdc55457218ea655c5a1d75​)

#skipci

diff --git a/Engine/Plugins/Enterprise/GLTFExporter/MigratedFromGit/Source/GLTFExporter/Private/Converters/GLTFMeshConverters.cpp b/Engine/Plugins/Enterprise/GLTFExporter/MigratedFromGit/Source/GLTFExporter/Private/Converters/GLTFMeshConverters.cpp
--- a/Engine/Plugins/Enterprise/GLTFExporter/MigratedFromGit/Source/GLTFExporter/Private/Converters/GLTFMeshConverters.cpp
+++ b/Engine/Plugins/Enterprise/GLTFExporter/MigratedFromGit/Source/GLTFExporter/Private/Converters/GLTFMeshConverters.cpp
@@ -11,7 +11,7 @@ void FGLTFStaticMeshConverter::Sanitize(const UStaticMesh*& StaticMesh, int32& L
 {
 	if (LODIndex < 0)
 	{
-		LODIndex = Builder.ExportOptions->DefaultLevelOfDetail;
+		LODIndex = FMath::Max(Builder.ExportOptions->DefaultLevelOfDetail, 0);
 	}
 
 	if (!Builder.ExportOptions->bExportVertexColors)
@@ -24,11 +24,10 @@ void FGLTFStaticMeshConverter::Sanitize(const UStaticMesh*& StaticMesh, int32& L
 		OverrideMaterials.Empty();
 	}
 
-	if (LODIndex >= 0 && StaticMesh->GetNumLODs() > LODIndex)
+	if (StaticMesh->GetNumLODs() > LODIndex)
 	{
 		const FStaticMeshLODResources& MeshLOD = StaticMesh->GetLODForExport(LODIndex);
-
-		if (OverrideVertexColors != nullptr && OverrideVertexColors == &MeshLOD.VertexBuffers.ColorVertexBuffer)
+		if (OverrideVertexColors == &MeshLOD.VertexBuffers.ColorVertexBuffer)
 		{
 			OverrideVertexColors = nullptr;
 		}
@@ -37,7 +36,7 @@ void FGLTFStaticMeshConverter::Sanitize(const UStaticMesh*& StaticMesh, int32& L
 
 FGLTFJsonMeshIndex FGLTFStaticMeshConverter::Convert(const UStaticMesh* StaticMesh, int32 LODIndex, const FColorVertexBuffer* OverrideVertexColors, FGLTFMaterialArray OverrideMaterials)
 {
-	if (LODIndex < 0 || StaticMesh->GetNumLODs() <= LODIndex)
+	if (StaticMesh->GetNumLODs() <= LODIndex)
 	{
 		return FGLTFJsonMeshIndex(INDEX_NONE);
 	}
@@ -99,7 +98,7 @@ void FGLTFSkeletalMeshConverter::Sanitize(const USkeletalMesh*& SkeletalMesh, in
 {
 	if (LODIndex < 0)
 	{
-		LODIndex = Builder.ExportOptions->DefaultLevelOfDetail;
+		LODIndex = FMath::Max(Builder.ExportOptions->DefaultLevelOfDetail, 0);
 	}
 
 	if (!Builder.ExportOptions->bExportVertexColors)
@@ -113,17 +112,16 @@ void FGLTFSkeletalMeshConverter::Sanitize(const USkeletalMesh*& SkeletalMesh, in
 	}
 
 	const FSkeletalMeshRenderData* RenderData = SkeletalMesh->GetResourceForRendering();
-
-	if (LODIndex >= 0 && RenderData->LODRenderData.Num() > LODIndex)
+	if (RenderData->LODRenderData.Num() > LODIndex)
 	{
 		const FSkeletalMeshLODRenderData& MeshLOD = RenderData->LODRenderData[LODIndex];
 
-		if (OverrideVertexColors != nullptr && OverrideVertexColors == &MeshLOD.StaticVertexBuffers.ColorVertexBuffer)
+		if (OverrideVertexColors == &MeshLOD.StaticVertexBuffers.ColorVertexBuffer)
 		{
 			OverrideVertexColors = nullptr;
 		}
 
-		if (OverrideSkinWeights != nullptr && OverrideSkinWeights == MeshLOD.GetSkinWeightVertexBuffer())
+		if (OverrideSkinWeights == MeshLOD.GetSkinWeightVertexBuffer())
 		{
 			OverrideSkinWeights = nullptr;
 		}
@@ -134,7 +132,7 @@ FGLTFJsonMeshIndex FGLTFSkeletalMeshConverter::Convert(const USkeletalMesh* Skel
 {
 	const FSkeletalMeshRenderData* RenderData = SkeletalMesh->GetResourceForRendering();
 
-	if (LODIndex < 0 || RenderData->LODRenderData.Num() <= LODIndex)
+	if (RenderData->LODRenderData.Num() <= LODIndex)
 	{
 		return FGLTFJsonMeshIndex(INDEX_NONE);
 	}
