[GLTFExporter] fix(orbit-camera): remove PostEditChangeProperty-declaration from runtime builds

commit 176184d8e23d70c3c48a03a06a0240a024ded923
Date: Tue, 29 Sep 2020 08:07:25 +0200

    fix(orbit-camera): remove PostEditChangeProperty-declaration from runtime builds
    (cherry picked from commit 34214fb9bf5625788e4050012301addbaf120544​)

#skipci

diff --git a/Engine/Plugins/Enterprise/GLTFExporter/MigratedFromGit/Source/GLTFExporterRuntime/Public/Actors/GLTFOrbitCameraActor.h b/Engine/Plugins/Enterprise/GLTFExporter/MigratedFromGit/Source/GLTFExporterRuntime/Public/Actors/GLTFOrbitCameraActor.h
--- a/Engine/Plugins/Enterprise/GLTFExporter/MigratedFromGit/Source/GLTFExporterRuntime/Public/Actors/GLTFOrbitCameraActor.h
+++ b/Engine/Plugins/Enterprise/GLTFExporter/MigratedFromGit/Source/GLTFExporterRuntime/Public/Actors/GLTFOrbitCameraActor.h
@@ -16,7 +16,9 @@ class GLTFEXPORTERRUNTIME_API AGLTFOrbitCameraActor : public ACameraActor
 	//~ Begin UObject Interface
 public:
 	AGLTFOrbitCameraActor(const FObjectInitializer& ObjectInitializer);
+#if WITH_EDITOR
 	virtual void PostEditChangeProperty(FPropertyChangedEvent& PropertyChangedEvent) override;
+#endif // WITH_EDITOR
 	//~ End UObject Interface
 
 	//~ Begin AActor Interface
