[GLTFExporter] feat: re-use existing image instead of adding a new one if they're identical

commit 43bc5c8124b4b0933c2e7827c30e65cb6dc701d1
Date: Tue, 29 Sep 2020 08:45:10 +0200

    feat: re-use existing image instead of adding a new one if they're identical
    (cherry picked from commit 1d2703420dfc8f4e53f12710940aadb5315b8ac2​)

#skipci

diff --git a/Engine/Plugins/Enterprise/GLTFExporter/MigratedFromGit/Source/GLTFExporter/Private/Builders/GLTFImageBuilder.cpp b/Engine/Plugins/Enterprise/GLTFExporter/MigratedFromGit/Source/GLTFExporter/Private/Builders/GLTFImageBuilder.cpp
--- a/Engine/Plugins/Enterprise/GLTFExporter/MigratedFromGit/Source/GLTFExporter/Private/Builders/GLTFImageBuilder.cpp
+++ b/Engine/Plugins/Enterprise/GLTFExporter/MigratedFromGit/Source/GLTFExporter/Private/Builders/GLTFImageBuilder.cpp
@@ -12,13 +12,18 @@ FGLTFImageBuilder::FGLTFImageBuilder()
 
 FGLTFJsonImageIndex FGLTFImageBuilder::AddImage(const void* CompressedData, int64 CompressedByteLength, EGLTFJsonMimeType MimeType, const FString& Name)
 {
-	FGLTFJsonImage Image;
-	Image.Name = Name;
-	Image.MimeType = MimeType;
-	const FGLTFJsonImageIndex ImageIndex = FGLTFJsonBuilder::AddImage(Image);
+	FGLTFJsonImageIndex ImageIndex = FindImage(CompressedData, CompressedByteLength, MimeType);
+	if (ImageIndex == INDEX_NONE)
+	{
+		FGLTFJsonImage Image;
+		Image.Name = Name;
+		Image.MimeType = MimeType;
+		ImageIndex = FGLTFJsonBuilder::AddImage(Image);
+
+		TArray64<uint8>& ImageData = ImageDataLookup.Add(ImageIndex);
+		ImageData.Append(static_cast<const uint8*>(CompressedData), CompressedByteLength);
+	}
 
-	TArray64<uint8>& ImageData = ImageDataLookup.Add(ImageIndex);
-	ImageData.Append(static_cast<const uint8*>(CompressedData), CompressedByteLength);
 	return ImageIndex;
 }
 
@@ -80,3 +85,25 @@ bool FGLTFImageBuilder::Serialize(FArchive& Archive, const FString& FilePath)
 
 	return FGLTFBufferBuilder::Serialize(Archive, FilePath);
 }
+
+FGLTFJsonImageIndex FGLTFImageBuilder::FindImage(const void* CompressedData, int64 CompressedByteLength, EGLTFJsonMimeType MimeType) const
+{
+	for (const auto& DataPair : ImageDataLookup)
+	{
+		const FGLTFJsonImage& JsonImage = const_cast<FGLTFImageBuilder*>(this)->GetImage(DataPair.Key);
+		if (JsonImage.MimeType == MimeType)
+		{
+			const TArray64<uint8>& ImageData = DataPair.Value;
+			if (ImageData.Num() == CompressedByteLength)
+			{
+				if (FMemory::Memcmp(ImageData.GetData(), CompressedData, CompressedByteLength) == 0)
+				{
+					const FGLTFJsonImageIndex ImageIndex = DataPair.Key;
+					return ImageIndex;
+				}
+			}
+		}
+	}
+
+	return FGLTFJsonImageIndex(INDEX_NONE);
+}
diff --git a/Engine/Plugins/Enterprise/GLTFExporter/MigratedFromGit/Source/GLTFExporter/Private/Builders/GLTFImageBuilder.h b/Engine/Plugins/Enterprise/GLTFExporter/MigratedFromGit/Source/GLTFExporter/Private/Builders/GLTFImageBuilder.h
--- a/Engine/Plugins/Enterprise/GLTFExporter/MigratedFromGit/Source/GLTFExporter/Private/Builders/GLTFImageBuilder.h
+++ b/Engine/Plugins/Enterprise/GLTFExporter/MigratedFromGit/Source/GLTFExporter/Private/Builders/GLTFImageBuilder.h
@@ -24,4 +24,6 @@ public:
 private:
 
 	TMap<FGLTFJsonImageIndex, TArray64<uint8>> ImageDataLookup;
+
+	FGLTFJsonImageIndex FindImage(const void* CompressedData, int64 CompressedByteLength, EGLTFJsonMimeType MimeType) const;
 };
