[GLTFExporter] chore: add TODO for more robust detection of identical images

commit 6676840aaaceffe617996859df032f70ea01486b
Date: Tue, 29 Sep 2020 14:15:27 +0200

    chore: add TODO for more robust detection of identical images
    (cherry picked from commit a7cbfc97892d7817577b57d0b9e8f9eb749c6ce2​)

#skipci

diff --git a/Engine/Plugins/Enterprise/GLTFExporter/MigratedFromGit/Source/GLTFExporter/Private/Builders/GLTFImageBuilder.cpp b/Engine/Plugins/Enterprise/GLTFExporter/MigratedFromGit/Source/GLTFExporter/Private/Builders/GLTFImageBuilder.cpp
--- a/Engine/Plugins/Enterprise/GLTFExporter/MigratedFromGit/Source/GLTFExporter/Private/Builders/GLTFImageBuilder.cpp
+++ b/Engine/Plugins/Enterprise/GLTFExporter/MigratedFromGit/Source/GLTFExporter/Private/Builders/GLTFImageBuilder.cpp
@@ -88,6 +88,10 @@ bool FGLTFImageBuilder::Serialize(FArchive& Archive, const FString& FilePath)
 
 FGLTFJsonImageIndex FGLTFImageBuilder::FindImage(const void* CompressedData, int64 CompressedByteLength, EGLTFJsonMimeType MimeType) const
 {
+	// TODO: maybe use size and / or other properties when looking for an identical image, just in case
+	// images with different properties happen to have the exact same compressed data.
+	// It would however limit the call-sites from where this function can be used, since we need to know these properties.
+
 	for (const auto& DataPair : ImageDataLookup)
 	{
 		const FGLTFJsonImage& JsonImage = const_cast<FGLTFImageBuilder*>(this)->GetImage(DataPair.Key);
