[GLTFExporter] refactor: add TODO for resuing texture indicies as well

commit 370681e1a28fc1db6e5a09f0cb86f4e79d9a0177
Date: Tue, 29 Sep 2020 15:40:30 +0200

    refactor: add TODO for resuing texture indicies as well
    (cherry picked from commit 7bd3f243cf3a5bfd2acb6a59435535599e5dcf95​)

#skipci

diff --git a/Engine/Plugins/Enterprise/GLTFExporter/MigratedFromGit/Source/GLTFExporter/Private/Converters/GLTFMaterialUtility.cpp b/Engine/Plugins/Enterprise/GLTFExporter/MigratedFromGit/Source/GLTFExporter/Private/Converters/GLTFMaterialUtility.cpp
--- a/Engine/Plugins/Enterprise/GLTFExporter/MigratedFromGit/Source/GLTFExporter/Private/Converters/GLTFMaterialUtility.cpp
+++ b/Engine/Plugins/Enterprise/GLTFExporter/MigratedFromGit/Source/GLTFExporter/Private/Converters/GLTFMaterialUtility.cpp
@@ -352,6 +352,7 @@ FGLTFJsonTextureIndex FGLTFMaterialUtility::AddTexture(FGLTFConvertBuilder& Buil
 	JsonSampler.WrapS = WrapS;
 	JsonSampler.WrapT = WrapT;
 
+	// TODO: reuse same texture index when image is the same
 	FGLTFJsonTexture JsonTexture;
 	JsonTexture.Name = TextureName;
 	JsonTexture.Sampler = Builder.AddSampler(JsonSampler);
