[GLTFExporter] feat: use default material (WorldGridMaterial) for unassigned slots in meshes

commit 5192031d1b3768814323feb811994b985ee744eb
Date: Mon, 28 Sep 2020 10:11:14 +0200

    feat: use default material (WorldGridMaterial) for unassigned slots in meshes
    (cherry picked from commit ec700cabd7856ae0752522ffc8d7eb7437b1371c​)

#skipci

diff --git a/Engine/Plugins/Enterprise/GLTFExporter/MigratedFromGit/Source/GLTFExporter/Private/Converters/GLTFMeshConverters.cpp b/Engine/Plugins/Enterprise/GLTFExporter/MigratedFromGit/Source/GLTFExporter/Private/Converters/GLTFMeshConverters.cpp
--- a/Engine/Plugins/Enterprise/GLTFExporter/MigratedFromGit/Source/GLTFExporter/Private/Converters/GLTFMeshConverters.cpp
+++ b/Engine/Plugins/Enterprise/GLTFExporter/MigratedFromGit/Source/GLTFExporter/Private/Converters/GLTFMeshConverters.cpp
@@ -85,10 +85,12 @@ FGLTFJsonMeshIndex FGLTFStaticMeshConverter::Convert(const UStaticMesh* StaticMe
 		}
 
 		const UMaterialInterface* Material = OverrideMaterials.GetOverride(StaticMesh->StaticMaterials, Section->MaterialIndex);
-		if (Material != nullptr)
+		if (Material == nullptr)
 		{
-			JsonPrimitive.Material = Builder.GetOrAddMaterial(Material);
+			Material = UMaterial::GetDefaultMaterial(MD_Surface);
 		}
+
+		JsonPrimitive.Material = Builder.GetOrAddMaterial(Material);
 	}
 
 	return Builder.AddMesh(JsonMesh);
@@ -195,10 +197,12 @@ FGLTFJsonMeshIndex FGLTFSkeletalMeshConverter::Convert(const USkeletalMesh* Skel
 		*/
 
 		const UMaterialInterface* Material = OverrideMaterials.GetOverride(SkeletalMesh->Materials, Section->MaterialIndex);
-		if (Material != nullptr)
+		if (Material == nullptr)
 		{
-			JsonPrimitive.Material = Builder.GetOrAddMaterial(Material);
+			Material = UMaterial::GetDefaultMaterial(MD_Surface);
 		}
+
+		JsonPrimitive.Material = Builder.GetOrAddMaterial(Material);
 	}
 
 	return Builder.AddMesh(JsonMesh);
