[GLTFExporter] feat(interaction-hotspot): inherit from UMaterialBillboardComponent to allow us to use a custom material

commit 1946d6dd51055e8afbdf841b61ccfd74b94694b7
Date: Mon, 7 Sep 2020 10:51:11 +0200

    feat(interaction-hotspot): inherit from UMaterialBillboardComponent to allow us to use a custom material
    (cherry picked from commit f6de4e7b2fcbfaa7fbb27d4618c383a132951417​)

#skipci

diff --git a/Engine/Plugins/Enterprise/GLTFExporter/MigratedFromGit/Source/GLTFExporterRuntime/Private/Components/GLTFInteractionHotspotComponent.cpp b/Engine/Plugins/Enterprise/GLTFExporter/MigratedFromGit/Source/GLTFExporterRuntime/Private/Components/GLTFInteractionHotspotComponent.cpp
--- a/Engine/Plugins/Enterprise/GLTFExporter/MigratedFromGit/Source/GLTFExporterRuntime/Private/Components/GLTFInteractionHotspotComponent.cpp
+++ b/Engine/Plugins/Enterprise/GLTFExporter/MigratedFromGit/Source/GLTFExporterRuntime/Private/Components/GLTFInteractionHotspotComponent.cpp
@@ -132,7 +132,9 @@ UBodySetup* UGLTFInteractionHotspotComponent::GetBodySetup()
 
 void UGLTFInteractionHotspotComponent::SetSprite(class UTexture2D* NewSprite)
 {
-	Super::SetSprite(NewSprite);
+	// TODO: implement equivalent functionality using UMaterialBillboardComponent
+
+	// Super::SetSprite(NewSprite);
 
 	UpdateCollisionVolume();
 }
diff --git a/Engine/Plugins/Enterprise/GLTFExporter/MigratedFromGit/Source/GLTFExporterRuntime/Public/Components/GLTFInteractionHotspotComponent.h b/Engine/Plugins/Enterprise/GLTFExporter/MigratedFromGit/Source/GLTFExporterRuntime/Public/Components/GLTFInteractionHotspotComponent.h
--- a/Engine/Plugins/Enterprise/GLTFExporter/MigratedFromGit/Source/GLTFExporterRuntime/Public/Components/GLTFInteractionHotspotComponent.h
+++ b/Engine/Plugins/Enterprise/GLTFExporter/MigratedFromGit/Source/GLTFExporterRuntime/Public/Components/GLTFInteractionHotspotComponent.h
@@ -2,7 +2,7 @@
 
 #pragma once
 
-#include "Components/BillboardComponent.h"
+#include "Components/MaterialBillboardComponent.h"
 #include "GLTFInteractionHotspotComponent.generated.h"
 
 class ASkeletalMeshActor;
@@ -14,7 +14,7 @@ class UTexture2D;
  * A component to set up hotspots which appear as billboards and allow playback of skeletal animations when cursor input is enabled.
  */
 UCLASS(BlueprintType, Blueprintable, Meta = (BlueprintSpawnableComponent), HideCategories = (Sprite, Physics, Collision, Navigation), DisplayName = "GLTF Interaction Hotspot Component")
-class GLTFEXPORTERRUNTIME_API UGLTFInteractionHotspotComponent : public UBillboardComponent
+class GLTFEXPORTERRUNTIME_API UGLTFInteractionHotspotComponent : public UMaterialBillboardComponent
 {
 	GENERATED_UCLASS_BODY()
 public:
@@ -37,9 +37,7 @@ public:
 	virtual UBodySetup* GetBodySetup() override;
 	//~ End UPrimitiveComponent Interface
 
-	//~ Begin UBillboardComponent Interface
-	virtual void SetSprite(class UTexture2D* NewSprite) override;
-	//~ End UBillboardComponent Interface
+	virtual void SetSprite(class UTexture2D* NewSprite);
 
 private:
 	UFUNCTION()
