[GLTFExporter] feat(interaction-hotspot): update sprite-element to match current sprite

commit 5a3b8d2d8480051f96f546fe8121e9857e9e3f9f
Date: Mon, 7 Sep 2020 15:41:02 +0200

    feat(interaction-hotspot): update sprite-element to match current sprite
    (cherry picked from commit 6e3a31f5eb9b17163cf9736eafae90dec5cedc5d​)

#skipci

diff --git a/Engine/Plugins/Enterprise/GLTFExporter/MigratedFromGit/Source/GLTFExporterRuntime/Private/Components/GLTFInteractionHotspotComponent.cpp b/Engine/Plugins/Enterprise/GLTFExporter/MigratedFromGit/Source/GLTFExporterRuntime/Private/Components/GLTFInteractionHotspotComponent.cpp
--- a/Engine/Plugins/Enterprise/GLTFExporter/MigratedFromGit/Source/GLTFExporterRuntime/Private/Components/GLTFInteractionHotspotComponent.cpp
+++ b/Engine/Plugins/Enterprise/GLTFExporter/MigratedFromGit/Source/GLTFExporterRuntime/Private/Components/GLTFInteractionHotspotComponent.cpp
@@ -22,8 +22,19 @@ UGLTFInteractionHotspotComponent::UGLTFInteractionHotspotComponent(const FObject
 	ToggledImage(nullptr),
 	ToggledHoveredImage(nullptr),
 	ShapeBodySetup(nullptr),
+	DefaultMaterial(nullptr),
 	bToggled(bToggled)
 {
+	struct FConstructorStatics
+	{
+		ConstructorHelpers::FObjectFinder<UMaterial> Material;
+		FConstructorStatics()
+			: Material(TEXT("/GLTFExporter/Materials/Hotspot"))
+		{
+		}
+	};
+	static FConstructorStatics ConstructorStatics;
+
 	// Setup the most minimalistic collision profile for mouse input events
 	SetCollisionEnabled(ECollisionEnabled::QueryAndPhysics);
 	SetCollisionResponseToAllChannels(ECollisionResponse::ECR_Ignore);
@@ -31,6 +42,9 @@ UGLTFInteractionHotspotComponent::UGLTFInteractionHotspotComponent(const FObject
 	SetGenerateOverlapEvents(false);
 
 	bHiddenInGame = false;
+	DefaultMaterial = UMaterialInstanceDynamic::Create(ConstructorStatics.Material.Object, GetTransientPackage());
+
+	CreateDefaultSpriteElement();
 
 	OnBeginCursorOver.AddDynamic(this, &UGLTFInteractionHotspotComponent::BeginCursorOver);
 	OnEndCursorOver.AddDynamic(this, &UGLTFInteractionHotspotComponent::EndCursorOver);
@@ -70,6 +84,14 @@ void UGLTFInteractionHotspotComponent::PostEditChangeProperty(FPropertyChangedEv
 }
 #endif // WITH_EDITOR
 
+void UGLTFInteractionHotspotComponent::PostLoad()
+{
+	Super::PostLoad();
+
+	CreateDefaultSpriteElement();	// NOTE: needed in order to overwrite any persisted element
+	SetSprite(Image);
+}
+
 void UGLTFInteractionHotspotComponent::BeginPlay()
 {
 	Super::BeginPlay();
@@ -132,9 +154,10 @@ UBodySetup* UGLTFInteractionHotspotComponent::GetBodySetup()
 
 void UGLTFInteractionHotspotComponent::SetSprite(class UTexture2D* NewSprite)
 {
-	// TODO: implement equivalent functionality using UMaterialBillboardComponent
+	GetSpriteMaterial()->SetTextureParameterValue("Sprite", NewSprite);
 
-	// Super::SetSprite(NewSprite);
+	const FVector2D PixelSize = NewSprite != nullptr ? FVector2D(NewSprite->GetSurfaceWidth(), NewSprite->GetSurfaceHeight()) : FVector2D(32.0f, 32.0f);
+	UpdateSpriteSize(PixelSize);
 
 	UpdateCollisionVolume();
 }
@@ -181,6 +204,8 @@ void UGLTFInteractionHotspotComponent::Clicked(UPrimitiveComponent* TouchedCompo
 
 void UGLTFInteractionHotspotComponent::UpdateCollisionVolume()
 {
+	// TODO: update collider-size dynamically to always match screen-size
+
 	if (ShapeBodySetup != nullptr)
 	{
 		// TODO: Figure out why the bounding radius doesn't match the size of the billboard
@@ -240,3 +265,49 @@ UTexture2D* UGLTFInteractionHotspotComponent::GetActiveImage(bool bCursorOver) c
 
 	return nullptr;
 }
+
+void UGLTFInteractionHotspotComponent::CreateDefaultSpriteElement()
+{
+	FMaterialSpriteElement Element;
+	Element.Material = DefaultMaterial;
+	Element.bSizeIsInScreenSpace = true;
+	Element.BaseSizeX = 0.1f;
+	Element.BaseSizeY = 0.1f;
+
+	SetElements({ Element });
+}
+
+FMaterialSpriteElement& UGLTFInteractionHotspotComponent::GetSpriteElement()
+{
+	return Elements[0];
+}
+
+UMaterialInstanceDynamic* UGLTFInteractionHotspotComponent::GetSpriteMaterial() const
+{
+	return static_cast<UMaterialInstanceDynamic*>(GetMaterial(0));
+}
+
+void UGLTFInteractionHotspotComponent::UpdateSpriteSize(const FVector2D& PixelSize)
+{
+	// TODO: keep screen-size in sync with any changes to screen-resolution
+
+	FMaterialSpriteElement& Element = GetSpriteElement();
+
+	if (UGameViewportClient* GameViewportClient = GetWorld()->GetGameViewport())
+	{
+		FVector2D ViewportSize;
+		GameViewportClient->GetViewportSize(ViewportSize);
+
+		Element.BaseSizeX = PixelSize.X / ViewportSize.X;
+		Element.BaseSizeY = PixelSize.X / ViewportSize.Y;
+	}
+	else
+	{
+		// TODO: if running in the editor, find and use the size of the viewport
+
+		Element.BaseSizeX = 0.1f;
+		Element.BaseSizeY = 0.1f;
+	}
+
+	MarkRenderStateDirty();
+}
diff --git a/Engine/Plugins/Enterprise/GLTFExporter/MigratedFromGit/Source/GLTFExporterRuntime/Public/Components/GLTFInteractionHotspotComponent.h b/Engine/Plugins/Enterprise/GLTFExporter/MigratedFromGit/Source/GLTFExporterRuntime/Public/Components/GLTFInteractionHotspotComponent.h
--- a/Engine/Plugins/Enterprise/GLTFExporter/MigratedFromGit/Source/GLTFExporterRuntime/Public/Components/GLTFInteractionHotspotComponent.h
+++ b/Engine/Plugins/Enterprise/GLTFExporter/MigratedFromGit/Source/GLTFExporterRuntime/Public/Components/GLTFInteractionHotspotComponent.h
@@ -22,6 +22,7 @@ public:
 #if WITH_EDITOR
 	virtual void PostEditChangeProperty(FPropertyChangedEvent& PropertyChangedEvent) override;
 #endif // WITH_EDITOR
+	virtual void PostLoad() override;
 	//~ End UObject Interface
 
 	//~ Begin UActorComponent Interface
@@ -82,5 +83,16 @@ private:
 	UPROPERTY(transient, duplicatetransient)
 	UBodySetup* ShapeBodySetup;
 
+	UPROPERTY(transient, duplicatetransient)
+	UMaterialInstanceDynamic* DefaultMaterial;
+
 	bool bToggled;
+
+	void CreateDefaultSpriteElement();
+
+	FMaterialSpriteElement& GetSpriteElement();
+
+	UMaterialInstanceDynamic* GetSpriteMaterial() const;
+
+	void UpdateSpriteSize(const FVector2D& PixelSize);
 };
