[GLTFExporter] feat(interaction-hotspot): simplify collision-handling by using a sub-object

commit 77919b767c3ff52e1d98b738273baecd96858b58
Date: Tue, 8 Sep 2020 08:33:26 +0200

    feat(interaction-hotspot): simplify collision-handling by using a sub-object
    (cherry picked from commit e45d77816102f8240996fed7151e615cb540e4c0​)

#skipci

diff --git a/Engine/Plugins/Enterprise/GLTFExporter/MigratedFromGit/Source/GLTFExporterRuntime/Private/Components/GLTFInteractionHotspotComponent.cpp b/Engine/Plugins/Enterprise/GLTFExporter/MigratedFromGit/Source/GLTFExporterRuntime/Private/Components/GLTFInteractionHotspotComponent.cpp
--- a/Engine/Plugins/Enterprise/GLTFExporter/MigratedFromGit/Source/GLTFExporterRuntime/Private/Components/GLTFInteractionHotspotComponent.cpp
+++ b/Engine/Plugins/Enterprise/GLTFExporter/MigratedFromGit/Source/GLTFExporterRuntime/Private/Components/GLTFInteractionHotspotComponent.cpp
@@ -3,16 +3,10 @@
 #include "Components/GLTFInteractionHotspotComponent.h"
 #include "Animation/SkeletalMeshActor.h"
 #include "Animation/AnimSequence.h"
-#include "Engine/CollisionProfile.h"
+#include "Components/SphereComponent.h"
 #include "UObject/ConstructorHelpers.h"
-#include "PhysicsEngine/BodySetup.h"
 #include "Animation/AnimSingleNodeInstance.h"
 
-namespace
-{
-	const float UnitSphereRadius = 50.0f;
-}
-
 DEFINE_LOG_CATEGORY_STATIC(LogEditorGLTFInteractionHotspot, Log, All);
 
 UGLTFInteractionHotspotComponent::UGLTFInteractionHotspotComponent(const FObjectInitializer& ObjectInitializer)
@@ -21,7 +15,7 @@ UGLTFInteractionHotspotComponent::UGLTFInteractionHotspotComponent(const FObject
 	HoveredImage(nullptr),
 	ToggledImage(nullptr),
 	ToggledHoveredImage(nullptr),
-	ShapeBodySetup(nullptr),
+	SphereComponent(nullptr),
 	DefaultMaterial(nullptr),
 	bToggled(bToggled)
 {
@@ -35,20 +29,24 @@ UGLTFInteractionHotspotComponent::UGLTFInteractionHotspotComponent(const FObject
 	};
 	static FConstructorStatics ConstructorStatics;
 
-	// Setup the most minimalistic collision profile for mouse input events
-	SetCollisionEnabled(ECollisionEnabled::QueryAndPhysics);
-	SetCollisionResponseToAllChannels(ECollisionResponse::ECR_Ignore);
-	SetCollisionResponseToChannel(ECollisionChannel::ECC_Visibility, ECollisionResponse::ECR_Block);
-	SetGenerateOverlapEvents(false);
-
 	bHiddenInGame = false;
 	DefaultMaterial = UMaterialInstanceDynamic::Create(ConstructorStatics.Material.Object, GetTransientPackage());
 
 	CreateDefaultSpriteElement();
 
-	OnBeginCursorOver.AddDynamic(this, &UGLTFInteractionHotspotComponent::BeginCursorOver);
-	OnEndCursorOver.AddDynamic(this, &UGLTFInteractionHotspotComponent::EndCursorOver);
-	OnClicked.AddDynamic(this, &UGLTFInteractionHotspotComponent::Clicked);
+	SphereComponent = CreateDefaultSubobject<USphereComponent>(TEXT("Collider"), true);
+	SphereComponent->InitSphereRadius(100.0f);
+	SphereComponent->SetupAttachment(this);
+
+	// Setup the most minimalistic collision profile for mouse input events
+	SphereComponent->SetCollisionEnabled(ECollisionEnabled::QueryAndPhysics);
+	SphereComponent->SetCollisionResponseToAllChannels(ECollisionResponse::ECR_Ignore);
+	SphereComponent->SetCollisionResponseToChannel(ECollisionChannel::ECC_Visibility, ECollisionResponse::ECR_Block);
+	SphereComponent->SetGenerateOverlapEvents(false);
+
+	SphereComponent->OnBeginCursorOver.AddDynamic(this, &UGLTFInteractionHotspotComponent::BeginCursorOver);
+	SphereComponent->OnEndCursorOver.AddDynamic(this, &UGLTFInteractionHotspotComponent::EndCursorOver);
+	SphereComponent->OnClicked.AddDynamic(this, &UGLTFInteractionHotspotComponent::Clicked);
 }
 
 #if WITH_EDITOR
@@ -101,57 +99,9 @@ void UGLTFInteractionHotspotComponent::BeginPlay()
 
 void UGLTFInteractionHotspotComponent::OnRegister()
 {
-	ShapeBodySetup = NewObject<UBodySetup>(this, NAME_None, RF_Transient);
-
-	if (GUObjectArray.IsDisregardForGC(this))
-	{
-		ShapeBodySetup->AddToRoot();
-	}
-
-	ShapeBodySetup->AddToCluster(this);
-
-	if (ShapeBodySetup->HasAnyInternalFlags(EInternalObjectFlags::Async) && GUObjectClusters.GetObjectCluster(ShapeBodySetup))
-	{
-		ShapeBodySetup->ClearInternalFlags(EInternalObjectFlags::Async);
-	}
-
-	ShapeBodySetup->CollisionTraceFlag = CTF_UseSimpleAsComplex;
-	ShapeBodySetup->AggGeom.SphereElems.Add(FKSphereElem(UnitSphereRadius));
-
 	Super::OnRegister();
 }
 
-void UGLTFInteractionHotspotComponent::OnCreatePhysicsState()
-{
-	Super::OnCreatePhysicsState();
-
-	if (BodyInstance.IsValidBodyInstance())
-	{
-#if WITH_PHYSX
-		FPhysicsCommand::ExecuteWrite(BodyInstance.GetActorReferenceWithWelding(), [this](const FPhysicsActorHandle& Actor)
-		{
-			TArray<FPhysicsShapeHandle> Shapes;
-			BodyInstance.GetAllShapes_AssumesLocked(Shapes);
-
-			for (FPhysicsShapeHandle& Shape : Shapes)
-			{
-				if (BodyInstance.IsShapeBoundToBody(Shape))
-				{
-					FPhysicsInterface::SetUserData(Shape, (void*)ShapeBodySetup->AggGeom.SphereElems[0].GetUserData());
-				}
-			}
-		});
-#endif
-	}
-
-	UpdateCollisionVolume();
-}
-
-UBodySetup* UGLTFInteractionHotspotComponent::GetBodySetup()
-{
-	return ShapeBodySetup;
-}
-
 void UGLTFInteractionHotspotComponent::SetSprite(class UTexture2D* NewSprite)
 {
 	GetSpriteMaterial()->SetTextureParameterValue("Sprite", NewSprite);
@@ -205,20 +155,6 @@ void UGLTFInteractionHotspotComponent::Clicked(UPrimitiveComponent* TouchedCompo
 void UGLTFInteractionHotspotComponent::UpdateCollisionVolume()
 {
 	// TODO: update collider-size dynamically to always match screen-size
-
-	if (ShapeBodySetup != nullptr)
-	{
-		// TODO: Figure out why the bounding radius doesn't match the size of the billboard
-		const float Scaling = 0.15f;
-		const float BillboardBoundingRadius = GetBillboardBoundingRadius() * Scaling;
-
-		check(ShapeBodySetup->AggGeom.SphereElems.Num() == 1);
-
-		if (!FMath::IsNearlyEqual(ShapeBodySetup->AggGeom.SphereElems[0].Radius, BillboardBoundingRadius) && UnitSphereRadius != 0.0f)
-		{
-			BodyInstance.UpdateBodyScale(FVector(BillboardBoundingRadius / UnitSphereRadius), true);
-		}
-	}
 }
 
 float UGLTFInteractionHotspotComponent::GetBillboardBoundingRadius() const
diff --git a/Engine/Plugins/Enterprise/GLTFExporter/MigratedFromGit/Source/GLTFExporterRuntime/Public/Components/GLTFInteractionHotspotComponent.h b/Engine/Plugins/Enterprise/GLTFExporter/MigratedFromGit/Source/GLTFExporterRuntime/Public/Components/GLTFInteractionHotspotComponent.h
--- a/Engine/Plugins/Enterprise/GLTFExporter/MigratedFromGit/Source/GLTFExporterRuntime/Public/Components/GLTFInteractionHotspotComponent.h
+++ b/Engine/Plugins/Enterprise/GLTFExporter/MigratedFromGit/Source/GLTFExporterRuntime/Public/Components/GLTFInteractionHotspotComponent.h
@@ -7,8 +7,8 @@
 
 class ASkeletalMeshActor;
 class UAnimSequence;
-class UBodySetup;
 class UTexture2D;
+class USphereComponent;
 
 /**
  * A component to set up hotspots which appear as billboards and allow playback of skeletal animations when cursor input is enabled.
@@ -30,13 +30,9 @@ public:
 
 protected:
 	virtual void OnRegister() override;
-	virtual void OnCreatePhysicsState() override;
 	//~ End UActorComponent Interface
 
-	//~ Begin UPrimitiveComponent Interface
 public:
-	virtual UBodySetup* GetBodySetup() override;
-	//~ End UPrimitiveComponent Interface
 
 	virtual void SetSprite(class UTexture2D* NewSprite);
 
@@ -80,8 +76,9 @@ public:
 	UTexture2D* ToggledHoveredImage;
 
 private:
+
 	UPROPERTY(transient, duplicatetransient)
-	UBodySetup* ShapeBodySetup;
+	USphereComponent* SphereComponent;
 
 	UPROPERTY(transient, duplicatetransient)
 	UMaterialInstanceDynamic* DefaultMaterial;
