[GLTFExporter] feat(interaction-hotspot): keep state for active image and -image-size (needed later)

commit 4f0fa9c176e5a3e73d665564d175212a0dca6347
Date: Tue, 8 Sep 2020 09:00:42 +0200

    feat(interaction-hotspot): keep state for active image and -image-size (needed later)
    (cherry picked from commit 08f664ec2db383e58a8674e44b84dc13b76230de​)

#skipci

diff --git a/Engine/Plugins/Enterprise/GLTFExporter/MigratedFromGit/Source/GLTFExporterRuntime/Private/Components/GLTFInteractionHotspotComponent.cpp b/Engine/Plugins/Enterprise/GLTFExporter/MigratedFromGit/Source/GLTFExporterRuntime/Private/Components/GLTFInteractionHotspotComponent.cpp
--- a/Engine/Plugins/Enterprise/GLTFExporter/MigratedFromGit/Source/GLTFExporterRuntime/Private/Components/GLTFInteractionHotspotComponent.cpp
+++ b/Engine/Plugins/Enterprise/GLTFExporter/MigratedFromGit/Source/GLTFExporterRuntime/Private/Components/GLTFInteractionHotspotComponent.cpp
@@ -17,6 +17,8 @@ UGLTFInteractionHotspotComponent::UGLTFInteractionHotspotComponent(const FObject
 	ToggledHoveredImage(nullptr),
 	SphereComponent(nullptr),
 	DefaultMaterial(nullptr),
+	ActiveImage(nullptr),
+	ActiveImageSize(0.0f, 0.0f),
 	bToggled(bToggled)
 {
 	struct FConstructorStatics
@@ -62,7 +64,7 @@ void UGLTFInteractionHotspotComponent::PostEditChangeProperty(FPropertyChangedEv
 
 		if (PropertyName == TEXT("Image"))
 		{
-			SetSprite(Image);
+			SetActiveImage(Image);
 		}
 		else if (PropertyName == TEXT("SkeletalMeshActor"))
 		{
@@ -87,14 +89,14 @@ void UGLTFInteractionHotspotComponent::PostLoad()
 	Super::PostLoad();
 
 	CreateDefaultSpriteElement();	// NOTE: needed in order to overwrite any persisted element
-	SetSprite(Image);
+	SetActiveImage(Image);
 }
 
 void UGLTFInteractionHotspotComponent::BeginPlay()
 {
 	Super::BeginPlay();
 
-	SetSprite(Image);
+	SetActiveImage(Image);
 }
 
 void UGLTFInteractionHotspotComponent::OnRegister()
@@ -102,24 +104,53 @@ void UGLTFInteractionHotspotComponent::OnRegister()
 	Super::OnRegister();
 }
 
-void UGLTFInteractionHotspotComponent::SetSprite(class UTexture2D* NewSprite)
+void UGLTFInteractionHotspotComponent::SetActiveImage(class UTexture2D* NewImage)
 {
-	GetSpriteMaterial()->SetTextureParameterValue("Sprite", NewSprite);
+	if (NewImage != ActiveImage)
+	{
+		ActiveImage = NewImage;
+		GetSpriteMaterial()->SetTextureParameterValue("Sprite", ActiveImage);
+
+		const FVector2D NewImageSize(
+			ActiveImage != nullptr ? ActiveImage->GetSurfaceWidth() : 32.0f,
+			ActiveImage != nullptr ? ActiveImage->GetSurfaceHeight() : 32.0f);
+
+		if (NewImageSize != ActiveImageSize)
+		{
+			ActiveImageSize = NewImageSize;
 
-	const FVector2D PixelSize = NewSprite != nullptr ? FVector2D(NewSprite->GetSurfaceWidth(), NewSprite->GetSurfaceHeight()) : FVector2D(32.0f, 32.0f);
-	UpdateSpriteSize(PixelSize);
+			FMaterialSpriteElement& Element = GetSpriteElement();
 
-	UpdateCollisionVolume();
+			if (UGameViewportClient* GameViewportClient = GetWorld()->GetGameViewport())
+			{
+				FVector2D ViewportSize;
+				GameViewportClient->GetViewportSize(ViewportSize);
+
+				Element.BaseSizeX = ActiveImageSize.X / ViewportSize.X;
+				Element.BaseSizeY = ActiveImageSize.X / ViewportSize.Y;
+			}
+			else
+			{
+				// TODO: if running in the editor, find and use the size of the viewport
+
+				Element.BaseSizeX = 0.1f;
+				Element.BaseSizeY = 0.1f;
+			}
+
+			MarkRenderStateDirty();
+			UpdateCollisionVolume();
+		}
+	}
 }
 
 void UGLTFInteractionHotspotComponent::BeginCursorOver(UPrimitiveComponent* TouchedComponent)
 {
-	SetSprite(GetActiveImage(true));
+	SetActiveImage(CalculateActiveImage(true));
 }
 
 void UGLTFInteractionHotspotComponent::EndCursorOver(UPrimitiveComponent* TouchedComponent)
 {
-	SetSprite(GetActiveImage(false));
+	SetActiveImage(CalculateActiveImage(false));
 }
 
 void UGLTFInteractionHotspotComponent::Clicked(UPrimitiveComponent* TouchedComponent, FKey ButtonPressed)
@@ -149,7 +180,7 @@ void UGLTFInteractionHotspotComponent::Clicked(UPrimitiveComponent* TouchedCompo
 
 	bToggled = !bToggled;
 
-	SetSprite(GetActiveImage(true));
+	SetActiveImage(CalculateActiveImage(true));
 }
 
 void UGLTFInteractionHotspotComponent::UpdateCollisionVolume()
@@ -165,7 +196,7 @@ float UGLTFInteractionHotspotComponent::GetBillboardBoundingRadius() const
 	return WorldBounds.SphereRadius;
 }
 
-UTexture2D* UGLTFInteractionHotspotComponent::GetActiveImage(bool bCursorOver) const
+UTexture2D* UGLTFInteractionHotspotComponent::CalculateActiveImage(bool bCursorOver) const
 {
 	// A list of candidates ordered by descending priority
 	TArray<UTexture2D*> ImageCandidates;
@@ -222,28 +253,3 @@ UMaterialInstanceDynamic* UGLTFInteractionHotspotComponent::GetSpriteMaterial()
 {
 	return static_cast<UMaterialInstanceDynamic*>(GetMaterial(0));
 }
-
-void UGLTFInteractionHotspotComponent::UpdateSpriteSize(const FVector2D& PixelSize)
-{
-	// TODO: keep screen-size in sync with any changes to screen-resolution
-
-	FMaterialSpriteElement& Element = GetSpriteElement();
-
-	if (UGameViewportClient* GameViewportClient = GetWorld()->GetGameViewport())
-	{
-		FVector2D ViewportSize;
-		GameViewportClient->GetViewportSize(ViewportSize);
-
-		Element.BaseSizeX = PixelSize.X / ViewportSize.X;
-		Element.BaseSizeY = PixelSize.X / ViewportSize.Y;
-	}
-	else
-	{
-		// TODO: if running in the editor, find and use the size of the viewport
-
-		Element.BaseSizeX = 0.1f;
-		Element.BaseSizeY = 0.1f;
-	}
-
-	MarkRenderStateDirty();
-}
diff --git a/Engine/Plugins/Enterprise/GLTFExporter/MigratedFromGit/Source/GLTFExporterRuntime/Public/Components/GLTFInteractionHotspotComponent.h b/Engine/Plugins/Enterprise/GLTFExporter/MigratedFromGit/Source/GLTFExporterRuntime/Public/Components/GLTFInteractionHotspotComponent.h
--- a/Engine/Plugins/Enterprise/GLTFExporter/MigratedFromGit/Source/GLTFExporterRuntime/Public/Components/GLTFInteractionHotspotComponent.h
+++ b/Engine/Plugins/Enterprise/GLTFExporter/MigratedFromGit/Source/GLTFExporterRuntime/Public/Components/GLTFInteractionHotspotComponent.h
@@ -32,10 +32,6 @@ protected:
 	virtual void OnRegister() override;
 	//~ End UActorComponent Interface
 
-public:
-
-	virtual void SetSprite(class UTexture2D* NewSprite);
-
 private:
 	UFUNCTION()
 	void BeginCursorOver(UPrimitiveComponent* TouchedComponent);
@@ -48,7 +44,11 @@ private:
 
 	void UpdateCollisionVolume();
 	float GetBillboardBoundingRadius() const;
-	UTexture2D* GetActiveImage(bool bCursorOver) const;
+	UTexture2D* CalculateActiveImage(bool bCursorOver) const;
+
+protected:
+
+	virtual void SetActiveImage(class UTexture2D* NewImage);
 
 public:
 	/* The skeletal mesh actor that will be animated when the hotspot is clicked. */
@@ -83,6 +83,11 @@ private:
 	UPROPERTY(transient, duplicatetransient)
 	UMaterialInstanceDynamic* DefaultMaterial;
 
+	UPROPERTY(transient, duplicatetransient)
+	UTexture* ActiveImage;
+
+	FVector2D ActiveImageSize;
+
 	bool bToggled;
 
 	void CreateDefaultSpriteElement();
@@ -90,6 +95,4 @@ private:
 	FMaterialSpriteElement& GetSpriteElement();
 
 	UMaterialInstanceDynamic* GetSpriteMaterial() const;
-
-	void UpdateSpriteSize(const FVector2D& PixelSize);
 };
