[GLTFExporter] fix(interaction-hotspot): fix incorrect scaling of sprite

commit 10fe355d0866e9dac0469e53f31b8f10b1041513
Date: Tue, 8 Sep 2020 10:33:36 +0200

    fix(interaction-hotspot): fix incorrect scaling of sprite
    (cherry picked from commit de8a8dde90f60c6391ce27e32fcc58c48c3f6fa2​)

#skipci

diff --git a/Engine/Plugins/Enterprise/GLTFExporter/MigratedFromGit/Source/GLTFExporterRuntime/Private/Components/GLTFInteractionHotspotComponent.cpp b/Engine/Plugins/Enterprise/GLTFExporter/MigratedFromGit/Source/GLTFExporterRuntime/Private/Components/GLTFInteractionHotspotComponent.cpp
--- a/Engine/Plugins/Enterprise/GLTFExporter/MigratedFromGit/Source/GLTFExporterRuntime/Private/Components/GLTFInteractionHotspotComponent.cpp
+++ b/Engine/Plugins/Enterprise/GLTFExporter/MigratedFromGit/Source/GLTFExporterRuntime/Private/Components/GLTFInteractionHotspotComponent.cpp
@@ -108,39 +108,21 @@ void UGLTFInteractionHotspotComponent::SetActiveImage(class UTexture2D* NewImage
 {
 	if (NewImage != ActiveImage)
 	{
+		GetSpriteMaterial()->SetTextureParameterValue("Sprite", NewImage);
 		ActiveImage = NewImage;
-		GetSpriteMaterial()->SetTextureParameterValue("Sprite", ActiveImage);
-
-		const FVector2D NewImageSize(
-			ActiveImage != nullptr ? ActiveImage->GetSurfaceWidth() : 32.0f,
-			ActiveImage != nullptr ? ActiveImage->GetSurfaceHeight() : 32.0f);
-
-		if (NewImageSize != ActiveImageSize)
-		{
-			ActiveImageSize = NewImageSize;
-
-			FMaterialSpriteElement& Element = GetSpriteElement();
-
-			if (UGameViewportClient* GameViewportClient = GetWorld()->GetGameViewport())
-			{
-				FVector2D ViewportSize;
-				GameViewportClient->GetViewportSize(ViewportSize);
-
-				Element.BaseSizeX = ActiveImageSize.X / ViewportSize.X;
-				Element.BaseSizeY = ActiveImageSize.X / ViewportSize.Y;
-			}
-			else
-			{
-				// TODO: if running in the editor, find and use the size of the viewport
+	}
 
-				Element.BaseSizeX = 0.1f;
-				Element.BaseSizeY = 0.1f;
-			}
+	const FVector2D NewImageSize(
+		NewImage != nullptr ? NewImage->GetSurfaceWidth() : 32.0f,
+		NewImage != nullptr ? NewImage->GetSurfaceHeight() : 32.0f);
 
-			MarkRenderStateDirty();
-			UpdateCollisionVolume();
-		}
+	if (NewImageSize != ActiveImageSize)
+	{
+		ActiveImageSize = NewImageSize;
 	}
+
+	// NOTE: we do this even if size is unchanged since the last update may have failed
+	UpdateSpriteSize();
 }
 
 void UGLTFInteractionHotspotComponent::BeginCursorOver(UPrimitiveComponent* TouchedComponent)
@@ -253,3 +235,26 @@ UMaterialInstanceDynamic* UGLTFInteractionHotspotComponent::GetSpriteMaterial()
 {
 	return static_cast<UMaterialInstanceDynamic*>(GetMaterial(0));
 }
+
+void UGLTFInteractionHotspotComponent::UpdateSpriteSize()
+{
+	if (UGameViewportClient* GameViewportClient = GetWorld()->GetGameViewport())
+	{
+		FVector2D ViewportSize;
+		GameViewportClient->GetViewportSize(ViewportSize);
+
+		const float BaseSizeX = ActiveImageSize.X / ViewportSize.X;
+		const float BaseSizeY = ActiveImageSize.Y / ViewportSize.Y;
+
+		FMaterialSpriteElement& Element = GetSpriteElement();
+
+		if (BaseSizeX != Element.BaseSizeX || BaseSizeY != Element.BaseSizeY)	// TODO: use epsilon for comparison?
+		{
+			Element.BaseSizeX = BaseSizeX;
+			Element.BaseSizeY = BaseSizeY;
+
+			MarkRenderStateDirty();
+			UpdateCollisionVolume();
+		}
+	}
+}
diff --git a/Engine/Plugins/Enterprise/GLTFExporter/MigratedFromGit/Source/GLTFExporterRuntime/Public/Components/GLTFInteractionHotspotComponent.h b/Engine/Plugins/Enterprise/GLTFExporter/MigratedFromGit/Source/GLTFExporterRuntime/Public/Components/GLTFInteractionHotspotComponent.h
--- a/Engine/Plugins/Enterprise/GLTFExporter/MigratedFromGit/Source/GLTFExporterRuntime/Public/Components/GLTFInteractionHotspotComponent.h
+++ b/Engine/Plugins/Enterprise/GLTFExporter/MigratedFromGit/Source/GLTFExporterRuntime/Public/Components/GLTFInteractionHotspotComponent.h
@@ -95,4 +95,6 @@ private:
 	FMaterialSpriteElement& GetSpriteElement();
 
 	UMaterialInstanceDynamic* GetSpriteMaterial() const;
+
+	void UpdateSpriteSize();
 };
