[GLTFExporter] feat(interaction-hotspot): sync collider-size to screen-size in play-mode

commit bc6c6a2bcaf5379b8033e0d9277a1e37009a5ff7
Date: Tue, 8 Sep 2020 10:52:27 +0200

    feat(interaction-hotspot): sync collider-size to screen-size in play-mode
    (cherry picked from commit e65dd806cb4ef90b6cadba70911bb8216a6ab6e1​)

#skipci

diff --git a/Engine/Plugins/Enterprise/GLTFExporter/MigratedFromGit/Source/GLTFExporterRuntime/Private/Components/GLTFInteractionHotspotComponent.cpp b/Engine/Plugins/Enterprise/GLTFExporter/MigratedFromGit/Source/GLTFExporterRuntime/Private/Components/GLTFInteractionHotspotComponent.cpp
--- a/Engine/Plugins/Enterprise/GLTFExporter/MigratedFromGit/Source/GLTFExporterRuntime/Private/Components/GLTFInteractionHotspotComponent.cpp
+++ b/Engine/Plugins/Enterprise/GLTFExporter/MigratedFromGit/Source/GLTFExporterRuntime/Private/Components/GLTFInteractionHotspotComponent.cpp
@@ -4,6 +4,7 @@
 #include "Animation/SkeletalMeshActor.h"
 #include "Animation/AnimSequence.h"
 #include "Components/SphereComponent.h"
+#include "Kismet/GameplayStatics.h"
 #include "UObject/ConstructorHelpers.h"
 #include "Animation/AnimSingleNodeInstance.h"
 
@@ -21,6 +22,10 @@ UGLTFInteractionHotspotComponent::UGLTFInteractionHotspotComponent(const FObject
 	ActiveImageSize(0.0f, 0.0f),
 	bToggled(bToggled)
 {
+	bHiddenInGame = false;
+	PrimaryComponentTick.bCanEverTick = true;
+	PrimaryComponentTick.bStartWithTickEnabled = true;
+
 	struct FConstructorStatics
 	{
 		ConstructorHelpers::FObjectFinder<UMaterial> Material;
@@ -31,13 +36,13 @@ UGLTFInteractionHotspotComponent::UGLTFInteractionHotspotComponent(const FObject
 	};
 	static FConstructorStatics ConstructorStatics;
 
-	bHiddenInGame = false;
 	DefaultMaterial = UMaterialInstanceDynamic::Create(ConstructorStatics.Material.Object, GetTransientPackage());
 
 	CreateDefaultSpriteElement();
 
 	SphereComponent = CreateDefaultSubobject<USphereComponent>(TEXT("Collider"), true);
 	SphereComponent->InitSphereRadius(100.0f);
+	SphereComponent->SetVisibility(false);
 	SphereComponent->SetupAttachment(this);
 
 	// Setup the most minimalistic collision profile for mouse input events
@@ -99,6 +104,47 @@ void UGLTFInteractionHotspotComponent::BeginPlay()
 	SetActiveImage(Image);
 }
 
+void UGLTFInteractionHotspotComponent::TickComponent(float DeltaTime, ELevelTick TickType, FActorComponentTickFunction* ThisTickFunction)
+{
+	UWorld* World = GetWorld();
+	if (World == nullptr)
+	{
+		return;
+	}
+
+	// Update scale of the sphere-component to match the screen-size of the active image
+	if (APlayerController* PlayerController = World->GetFirstPlayerController())	// TODO: is is safe to assume that we want the first controller for projections?
+	{
+		const FVector ColliderLocation = SphereComponent->GetComponentLocation();
+		FVector ColliderScreenLocation;
+
+		if (PlayerController->ProjectWorldLocationToScreenWithDistance(ColliderLocation, ColliderScreenLocation))
+		{
+			const FVector2D CornerScreenLocation = FVector2D(ColliderScreenLocation) + ActiveImageSize * 0.5f;
+			FVector RayLocation;
+			FVector RayDirection;
+
+			if (PlayerController->DeprojectScreenPositionToWorld(CornerScreenLocation.X, CornerScreenLocation.Y, RayLocation, RayDirection))
+			{
+				const FVector ExtentLocation = RayLocation + RayDirection * ColliderScreenLocation.Z;
+				const float NewSphereRadius = (ExtentLocation - ColliderLocation).Size() / SphereComponent->GetShapeScale();
+				const float OldSphereRadius = SphereComponent->GetUnscaledSphereRadius();
+
+				if (FMath::Abs(NewSphereRadius - OldSphereRadius) > 0.1f)	// TODO: better epsilon?
+				{
+					SphereComponent->SetSphereRadius(NewSphereRadius);
+				}
+			}
+		}
+	}
+
+	// Fade out / hide the component when it's hidden behind other objects
+	if (World != nullptr)
+	{
+		// TODO: implement
+	}
+}
+
 void UGLTFInteractionHotspotComponent::OnRegister()
 {
 	Super::OnRegister();
@@ -165,19 +211,6 @@ void UGLTFInteractionHotspotComponent::Clicked(UPrimitiveComponent* TouchedCompo
 	SetActiveImage(CalculateActiveImage(true));
 }
 
-void UGLTFInteractionHotspotComponent::UpdateCollisionVolume()
-{
-	// TODO: update collider-size dynamically to always match screen-size
-}
-
-float UGLTFInteractionHotspotComponent::GetBillboardBoundingRadius() const
-{
-	const FTransform WorldTransform = GetComponentTransform();
-	const FBoxSphereBounds WorldBounds = CalcBounds(WorldTransform);
-
-	return WorldBounds.SphereRadius;
-}
-
 UTexture2D* UGLTFInteractionHotspotComponent::CalculateActiveImage(bool bCursorOver) const
 {
 	// A list of candidates ordered by descending priority
@@ -254,7 +287,6 @@ void UGLTFInteractionHotspotComponent::UpdateSpriteSize()
 			Element.BaseSizeY = BaseSizeY;
 
 			MarkRenderStateDirty();
-			UpdateCollisionVolume();
 		}
 	}
 }
diff --git a/Engine/Plugins/Enterprise/GLTFExporter/MigratedFromGit/Source/GLTFExporterRuntime/Public/Components/GLTFInteractionHotspotComponent.h b/Engine/Plugins/Enterprise/GLTFExporter/MigratedFromGit/Source/GLTFExporterRuntime/Public/Components/GLTFInteractionHotspotComponent.h
--- a/Engine/Plugins/Enterprise/GLTFExporter/MigratedFromGit/Source/GLTFExporterRuntime/Public/Components/GLTFInteractionHotspotComponent.h
+++ b/Engine/Plugins/Enterprise/GLTFExporter/MigratedFromGit/Source/GLTFExporterRuntime/Public/Components/GLTFInteractionHotspotComponent.h
@@ -27,6 +27,7 @@ public:
 
 	//~ Begin UActorComponent Interface
 	virtual void BeginPlay() override;
+	virtual void TickComponent(float DeltaTime, ELevelTick TickType, FActorComponentTickFunction* ThisTickFunction) override;
 
 protected:
 	virtual void OnRegister() override;
@@ -42,8 +43,6 @@ private:
 	UFUNCTION()
 	void Clicked(UPrimitiveComponent* TouchedComponent, FKey ButtonPressed);
 
-	void UpdateCollisionVolume();
-	float GetBillboardBoundingRadius() const;
 	UTexture2D* CalculateActiveImage(bool bCursorOver) const;
 
 protected:
