[GLTFExporter] feat(interaction-hotspot): implement fading of hotspot when occluded

commit a6a7f82b9efe2ac9fefdfb08be21bb8ac5926e51
Date: Tue, 8 Sep 2020 17:05:49 +0200

    feat(interaction-hotspot): implement fading of hotspot when occluded
    (cherry picked from commit e9ee682f29df89b89807ccfafb0deec684288734​)

#skipci

diff --git a/Engine/Plugins/Enterprise/GLTFExporter/MigratedFromGit/Source/GLTFExporterRuntime/Private/Components/GLTFInteractionHotspotComponent.cpp b/Engine/Plugins/Enterprise/GLTFExporter/MigratedFromGit/Source/GLTFExporterRuntime/Private/Components/GLTFInteractionHotspotComponent.cpp
--- a/Engine/Plugins/Enterprise/GLTFExporter/MigratedFromGit/Source/GLTFExporterRuntime/Private/Components/GLTFInteractionHotspotComponent.cpp
+++ b/Engine/Plugins/Enterprise/GLTFExporter/MigratedFromGit/Source/GLTFExporterRuntime/Private/Components/GLTFInteractionHotspotComponent.cpp
@@ -20,7 +20,9 @@ UGLTFInteractionHotspotComponent::UGLTFInteractionHotspotComponent(const FObject
 	DefaultMaterial(nullptr),
 	ActiveImage(nullptr),
 	ActiveImageSize(0.0f, 0.0f),
-	bToggled(bToggled)
+	bToggled(bToggled),
+	RealtimeSecondsWhenLastInSight(0.0f),
+	RealtimeSecondsWhenLastHidden(0.0f)
 {
 	bHiddenInGame = false;
 	PrimaryComponentTick.bCanEverTick = true;
@@ -107,41 +109,78 @@ void UGLTFInteractionHotspotComponent::BeginPlay()
 void UGLTFInteractionHotspotComponent::TickComponent(float DeltaTime, ELevelTick TickType, FActorComponentTickFunction* ThisTickFunction)
 {
 	UWorld* World = GetWorld();
-	if (World == nullptr)
+
+	// TODO: is is safe to assume that we want the first controller for projections?
+	APlayerController* PlayerController = World != nullptr ? World->GetFirstPlayerController() : nullptr;
+	if (PlayerController == nullptr)
+	{
+		return;
+	}
+
+	const FVector ColliderLocation = SphereComponent->GetComponentLocation();
+
+	FVector ColliderScreenLocation;
+	if (!PlayerController->ProjectWorldLocationToScreenWithDistance(ColliderLocation, ColliderScreenLocation))
 	{
 		return;
 	}
 
 	// Update scale of the sphere-component to match the screen-size of the active image
-	if (APlayerController* PlayerController = World->GetFirstPlayerController())	// TODO: is is safe to assume that we want the first controller for projections?
 	{
-		const FVector ColliderLocation = SphereComponent->GetComponentLocation();
-		FVector ColliderScreenLocation;
+		const FVector2D CornerScreenLocation = FVector2D(ColliderScreenLocation) + ActiveImageSize * 0.5f;
+		FVector RayLocation;
+		FVector RayDirection;
 
-		if (PlayerController->ProjectWorldLocationToScreenWithDistance(ColliderLocation, ColliderScreenLocation))
+		if (PlayerController->DeprojectScreenPositionToWorld(CornerScreenLocation.X, CornerScreenLocation.Y, RayLocation, RayDirection))
 		{
-			const FVector2D CornerScreenLocation = FVector2D(ColliderScreenLocation) + ActiveImageSize * 0.5f;
-			FVector RayLocation;
-			FVector RayDirection;
+			const FVector ExtentLocation = RayLocation + RayDirection * ColliderScreenLocation.Z;
+			const float NewSphereRadius = (ExtentLocation - ColliderLocation).Size() / SphereComponent->GetShapeScale();
+			const float OldSphereRadius = SphereComponent->GetUnscaledSphereRadius();
 
-			if (PlayerController->DeprojectScreenPositionToWorld(CornerScreenLocation.X, CornerScreenLocation.Y, RayLocation, RayDirection))
+			if (FMath::Abs(NewSphereRadius - OldSphereRadius) > 0.1f)	// TODO: better epsilon?
 			{
-				const FVector ExtentLocation = RayLocation + RayDirection * ColliderScreenLocation.Z;
-				const float NewSphereRadius = (ExtentLocation - ColliderLocation).Size() / SphereComponent->GetShapeScale();
-				const float OldSphereRadius = SphereComponent->GetUnscaledSphereRadius();
-
-				if (FMath::Abs(NewSphereRadius - OldSphereRadius) > 0.1f)	// TODO: better epsilon?
-				{
-					SphereComponent->SetSphereRadius(NewSphereRadius);
-				}
+				SphereComponent->SetSphereRadius(NewSphereRadius);
 			}
 		}
 	}
 
 	// Fade out / hide the component when it's hidden behind other objects
-	if (World != nullptr)
 	{
-		// TODO: implement
+		FHitResult HitResult;
+		bool bIsHotspotOccluded = false;
+
+		if (PlayerController->GetHitResultAtScreenPosition(FVector2D(ColliderScreenLocation), ECC_Visibility, false, HitResult))
+		{
+			const UPrimitiveComponent* HitComponent = HitResult.GetComponent();
+			if (HitComponent != nullptr && HitComponent != SphereComponent)
+			{
+				bIsHotspotOccluded = HitComponent->GetAttachParent() == nullptr || !HitComponent->GetAttachParent()->IsA<UGLTFInteractionHotspotComponent>();
+			}
+		}
+
+		const float CurrentRealtimeSeconds = UGameplayStatics::GetRealTimeSeconds(World);
+		float Opacity;
+
+		if (bIsHotspotOccluded)
+		{
+			RealtimeSecondsWhenLastHidden = CurrentRealtimeSeconds;
+
+			const float HiddenDuration = FMath::Max(RealtimeSecondsWhenLastHidden - RealtimeSecondsWhenLastInSight, 0.0f);
+			const float FadeOutDuration = 0.5f;
+
+			Opacity = 1.0f - FMath::Clamp(HiddenDuration / FadeOutDuration, 0.0f, 1.0f);
+		}
+		else
+		{
+			RealtimeSecondsWhenLastInSight = CurrentRealtimeSeconds;
+
+			const float VisibleDuration = FMath::Max(RealtimeSecondsWhenLastInSight - RealtimeSecondsWhenLastHidden, 0.0f);
+			const float FadeInDuration = 0.25f;
+
+			Opacity = FMath::Clamp(VisibleDuration / FadeInDuration, 0.0f, 1.0f);
+		}
+
+		SetSpriteOpacity(Opacity);
 	}
 }
 
@@ -290,3 +329,8 @@ void UGLTFInteractionHotspotComponent::UpdateSpriteSize()
 		}
 	}
 }
+
+void UGLTFInteractionHotspotComponent::SetSpriteOpacity(const float Opacity) const
+{
+	GetSpriteMaterial()->SetScalarParameterValue("Opacity", Opacity);
+}
diff --git a/Engine/Plugins/Enterprise/GLTFExporter/MigratedFromGit/Source/GLTFExporterRuntime/Public/Components/GLTFInteractionHotspotComponent.h b/Engine/Plugins/Enterprise/GLTFExporter/MigratedFromGit/Source/GLTFExporterRuntime/Public/Components/GLTFInteractionHotspotComponent.h
--- a/Engine/Plugins/Enterprise/GLTFExporter/MigratedFromGit/Source/GLTFExporterRuntime/Public/Components/GLTFInteractionHotspotComponent.h
+++ b/Engine/Plugins/Enterprise/GLTFExporter/MigratedFromGit/Source/GLTFExporterRuntime/Public/Components/GLTFInteractionHotspotComponent.h
@@ -89,6 +89,10 @@ private:
 
 	bool bToggled;
 
+	float RealtimeSecondsWhenLastInSight;
+
+	float RealtimeSecondsWhenLastHidden;
+
 	void CreateDefaultSpriteElement();
 
 	FMaterialSpriteElement& GetSpriteElement();
@@ -96,4 +100,6 @@ private:
 	UMaterialInstanceDynamic* GetSpriteMaterial() const;
 
 	void UpdateSpriteSize();
+
+	void SetSpriteOpacity(const float Opacity) const;
 };
