[GLTFExporter] feat(interaction-hotspot): enable / disable interaction based on opacity

commit d55002b22c75ca5def6a26e2ca51bc67bc6eb28c
Date: Wed, 9 Sep 2020 09:06:33 +0200

    feat(interaction-hotspot): enable / disable interaction based on opacity
    (cherry picked from commit 12f635cecbff3bb51879f0062f96bbe98bb9efcc​)

#skipci

diff --git a/Engine/Plugins/Enterprise/GLTFExporter/MigratedFromGit/Source/GLTFExporterRuntime/Private/Components/GLTFInteractionHotspotComponent.cpp b/Engine/Plugins/Enterprise/GLTFExporter/MigratedFromGit/Source/GLTFExporterRuntime/Private/Components/GLTFInteractionHotspotComponent.cpp
--- a/Engine/Plugins/Enterprise/GLTFExporter/MigratedFromGit/Source/GLTFExporterRuntime/Private/Components/GLTFInteractionHotspotComponent.cpp
+++ b/Engine/Plugins/Enterprise/GLTFExporter/MigratedFromGit/Source/GLTFExporterRuntime/Private/Components/GLTFInteractionHotspotComponent.cpp
@@ -144,7 +144,7 @@ void UGLTFInteractionHotspotComponent::TickComponent(float DeltaTime, ELevelTick
 		}
 	}
 
-	// Fade out / hide the component when it's hidden behind other objects
+	// Update opacity and interactivity of the hotspot based on if it its occluded by other objects or not
 	{
 		FHitResult HitResult;
 		bool bIsHotspotOccluded = false;
@@ -181,6 +181,9 @@ void UGLTFInteractionHotspotComponent::TickComponent(float DeltaTime, ELevelTick
 		}
 
 		SetSpriteOpacity(Opacity);
+
+		const ECollisionEnabled::Type CollisionEnabled = Opacity >= 0.5f ? ECollisionEnabled::QueryAndPhysics :  ECollisionEnabled::NoCollision;
+		SphereComponent->SetCollisionEnabled(CollisionEnabled);
 	}
 }
 
