[GLTFExporter] chore(interaction-hotspot): add TODO for allowing interaction when partially blocked

commit b1aa91ac27e8c2409b7e3c0c205cc2fc4ec11782
Date: Wed, 9 Sep 2020 09:07:36 +0200

    chore(interaction-hotspot): add TODO for allowing interaction when partially blocked
    (cherry picked from commit 604e6fcb9fdf843a82c86d1dd054093658fc0ca7​)

#skipci

diff --git a/Engine/Plugins/Enterprise/GLTFExporter/MigratedFromGit/Source/GLTFExporterRuntime/Private/Components/GLTFInteractionHotspotComponent.cpp b/Engine/Plugins/Enterprise/GLTFExporter/MigratedFromGit/Source/GLTFExporterRuntime/Private/Components/GLTFInteractionHotspotComponent.cpp
--- a/Engine/Plugins/Enterprise/GLTFExporter/MigratedFromGit/Source/GLTFExporterRuntime/Private/Components/GLTFInteractionHotspotComponent.cpp
+++ b/Engine/Plugins/Enterprise/GLTFExporter/MigratedFromGit/Source/GLTFExporterRuntime/Private/Components/GLTFInteractionHotspotComponent.cpp
@@ -53,6 +53,10 @@ UGLTFInteractionHotspotComponent::UGLTFInteractionHotspotComponent(const FObject
 	SphereComponent->SetCollisionResponseToChannel(ECollisionChannel::ECC_Visibility, ECollisionResponse::ECR_Block);
 	SphereComponent->SetGenerateOverlapEvents(false);
 
+	// TODO: ensure that hover and click works even if the hotspot is partially blocked.
+	// The sprite is rendered on top of other objects, and it would be confusing for the user
+	// if it didn't respond to interaction under these circumstances.
+
 	SphereComponent->OnBeginCursorOver.AddDynamic(this, &UGLTFInteractionHotspotComponent::BeginCursorOver);
 	SphereComponent->OnEndCursorOver.AddDynamic(this, &UGLTFInteractionHotspotComponent::EndCursorOver);
 	SphereComponent->OnClicked.AddDynamic(this, &UGLTFInteractionHotspotComponent::Clicked);
