[GLTFExporter] feat(interaction-hotspot): maintain a pixel-perfect size in the editor & when playing

commit 082147906f59b5c6741cee0d3994cdc95ad85cd1
Date: Wed, 9 Sep 2020 11:34:45 +0200

    feat(interaction-hotspot): maintain a pixel-perfect size in the editor & when playing
    (cherry picked from commit fc3645347bcaad302b4ff593ab45b348ff474d66​)

#skipci

diff --git a/Engine/Plugins/Enterprise/GLTFExporter/MigratedFromGit/Source/GLTFExporterRuntime/Private/Components/GLTFInteractionHotspotComponent.cpp b/Engine/Plugins/Enterprise/GLTFExporter/MigratedFromGit/Source/GLTFExporterRuntime/Private/Components/GLTFInteractionHotspotComponent.cpp
--- a/Engine/Plugins/Enterprise/GLTFExporter/MigratedFromGit/Source/GLTFExporterRuntime/Private/Components/GLTFInteractionHotspotComponent.cpp
+++ b/Engine/Plugins/Enterprise/GLTFExporter/MigratedFromGit/Source/GLTFExporterRuntime/Private/Components/GLTFInteractionHotspotComponent.cpp
@@ -8,6 +8,13 @@
 #include "UObject/ConstructorHelpers.h"
 #include "Animation/AnimSingleNodeInstance.h"
 
+#if WITH_EDITOR
+#include "LevelEditor.h"
+#include "Modules/ModuleManager.h"
+#include "SLevelViewport.h"
+#include "Slate/SceneViewport.h"
+#endif
+
 DEFINE_LOG_CATEGORY_STATIC(LogEditorGLTFInteractionHotspot, Log, All);
 
 UGLTFInteractionHotspotComponent::UGLTFInteractionHotspotComponent(const FObjectInitializer& ObjectInitializer)
@@ -317,27 +324,82 @@ UMaterialInstanceDynamic* UGLTFInteractionHotspotComponent::GetSpriteMaterial()
 
 void UGLTFInteractionHotspotComponent::UpdateSpriteSize()
 {
-	if (UGameViewportClient* GameViewportClient = GetWorld()->GetGameViewport())
+	const FIntPoint ViewportSize = GetCurrentViewportSize();
+
+	float BaseSizeX = 0.1f;
+	float BaseSizeY = 0.1f;
+
+	if (ViewportSize.X > 0 && ViewportSize.Y > 0)
+	{
+		BaseSizeX = ActiveImageSize.X / static_cast<float>(ViewportSize.X);
+		BaseSizeY = ActiveImageSize.Y / static_cast<float>(ViewportSize.Y);
+	}
+
+	FMaterialSpriteElement& Element = GetSpriteElement();
+
+	if (BaseSizeX != Element.BaseSizeX || BaseSizeY != Element.BaseSizeY)	// TODO: use epsilon for comparison?
 	{
-		FVector2D ViewportSize;
-		GameViewportClient->GetViewportSize(ViewportSize);
+		Element.BaseSizeX = BaseSizeX;
+		Element.BaseSizeY = BaseSizeY;
+
+		MarkRenderStateDirty();
+	}
+}
+
+void UGLTFInteractionHotspotComponent::SetSpriteOpacity(const float Opacity) const
+{
+	GetSpriteMaterial()->SetScalarParameterValue("Opacity", Opacity);
+}
 
-		const float BaseSizeX = ActiveImageSize.X / ViewportSize.X;
-		const float BaseSizeY = ActiveImageSize.Y / ViewportSize.Y;
+FIntPoint UGLTFInteractionHotspotComponent::GetCurrentViewportSize()
+{
+	// TODO: verify that correct size is calculated in the various play-modes and in the editor
 
-		FMaterialSpriteElement& Element = GetSpriteElement();
+	const FViewport* Viewport = nullptr;
 
-		if (BaseSizeX != Element.BaseSizeX || BaseSizeY != Element.BaseSizeY)	// TODO: use epsilon for comparison?
+	if (UWorld* World = GetWorld())
+	{
+		if (World->IsGameWorld())
 		{
-			Element.BaseSizeX = BaseSizeX;
-			Element.BaseSizeY = BaseSizeY;
+			if (UGameViewportClient* GameViewportClient = World->GetGameViewport())
+			{
+				Viewport = GameViewportClient->Viewport;
+			}
+		}
+#if WITH_EDITOR
+		else
+		{
+			static const FName LevelEditorName = "LevelEditor";
+
+			if (FModuleManager::Get().IsModuleLoaded(LevelEditorName))
+			{
+				FLevelEditorModule& LevelEditorModule = FModuleManager::GetModuleChecked<FLevelEditorModule>(LevelEditorName);
 
-			MarkRenderStateDirty();
+				if (const TSharedPtr<SLevelViewport> ActiveLevelViewport = LevelEditorModule.GetFirstActiveLevelViewport())
+				{
+					Viewport = ActiveLevelViewport->GetActiveViewport();
+				}
+			}
 		}
+#endif
+	}
+
+	if (Viewport != nullptr)
+	{
+		if (!Viewport->ViewportResizedEvent.IsBoundToObject(this))
+		{
+			Viewport->ViewportResizedEvent.AddUObject(this, &UGLTFInteractionHotspotComponent::ViewportResized);
+		}
+
+		return Viewport->GetSizeXY();
+	}
+	else
+	{
+		return FIntPoint(0, 0);
 	}
 }
 
-void UGLTFInteractionHotspotComponent::SetSpriteOpacity(const float Opacity) const
+void UGLTFInteractionHotspotComponent::ViewportResized(FViewport*, uint32)
 {
-	GetSpriteMaterial()->SetScalarParameterValue("Opacity", Opacity);
+	UpdateSpriteSize();
 }
diff --git a/Engine/Plugins/Enterprise/GLTFExporter/MigratedFromGit/Source/GLTFExporterRuntime/Public/Components/GLTFInteractionHotspotComponent.h b/Engine/Plugins/Enterprise/GLTFExporter/MigratedFromGit/Source/GLTFExporterRuntime/Public/Components/GLTFInteractionHotspotComponent.h
--- a/Engine/Plugins/Enterprise/GLTFExporter/MigratedFromGit/Source/GLTFExporterRuntime/Public/Components/GLTFInteractionHotspotComponent.h
+++ b/Engine/Plugins/Enterprise/GLTFExporter/MigratedFromGit/Source/GLTFExporterRuntime/Public/Components/GLTFInteractionHotspotComponent.h
@@ -9,6 +9,7 @@ class ASkeletalMeshActor;
 class UAnimSequence;
 class UTexture2D;
 class USphereComponent;
+class FViewport;
 
 /**
  * A component to set up hotspots which appear as billboards and allow playback of skeletal animations when cursor input is enabled.
@@ -102,4 +103,8 @@ private:
 	void UpdateSpriteSize();
 
 	void SetSpriteOpacity(const float Opacity) const;
+
+	FIntPoint GetCurrentViewportSize();
+
+	void ViewportResized(FViewport*, uint32);
 };
