[GLTFExporter] refactor(interaction-hotspot): move module-name constant into anonymous namespace

commit 7b358a77d19b83cf431f0097763598371643aa3b
Date: Wed, 9 Sep 2020 13:57:22 +0200

    refactor(interaction-hotspot): move module-name constant into anonymous namespace
    (cherry picked from commit 493bf5473fa7328ddec70556589b37cdb754ce3a​)

#skipci

diff --git a/Engine/Plugins/Enterprise/GLTFExporter/MigratedFromGit/Source/GLTFExporterRuntime/Private/Components/GLTFInteractionHotspotComponent.cpp b/Engine/Plugins/Enterprise/GLTFExporter/MigratedFromGit/Source/GLTFExporterRuntime/Private/Components/GLTFInteractionHotspotComponent.cpp
--- a/Engine/Plugins/Enterprise/GLTFExporter/MigratedFromGit/Source/GLTFExporterRuntime/Private/Components/GLTFInteractionHotspotComponent.cpp
+++ b/Engine/Plugins/Enterprise/GLTFExporter/MigratedFromGit/Source/GLTFExporterRuntime/Private/Components/GLTFInteractionHotspotComponent.cpp
@@ -17,6 +17,11 @@
 
 DEFINE_LOG_CATEGORY_STATIC(LogEditorGLTFInteractionHotspot, Log, All);
 
+namespace
+{
+	const FName NAME_LevelEditorModule = TEXT("LevelEditor");
+}
+
 UGLTFInteractionHotspotComponent::UGLTFInteractionHotspotComponent(const FObjectInitializer& ObjectInitializer)
 	: Super(ObjectInitializer),
 	Image(nullptr),
@@ -369,11 +374,9 @@ FIntPoint UGLTFInteractionHotspotComponent::GetCurrentViewportSize()
 #if WITH_EDITOR
 		else
 		{
-			static const FName LevelEditorName = "LevelEditor";
-
-			if (FModuleManager::Get().IsModuleLoaded(LevelEditorName))
+			if (FModuleManager::Get().IsModuleLoaded(NAME_LevelEditorModule))
 			{
-				FLevelEditorModule& LevelEditorModule = FModuleManager::GetModuleChecked<FLevelEditorModule>(LevelEditorName);
+				FLevelEditorModule& LevelEditorModule = FModuleManager::GetModuleChecked<FLevelEditorModule>(NAME_LevelEditorModule);
 
 				if (const TSharedPtr<SLevelViewport> ActiveLevelViewport = LevelEditorModule.GetFirstActiveLevelViewport())
 				{
