[GLTFExporter] feat(interaction-hotspot): use tags for ignoring hits with other hotspots

commit 7b3d45dbb8081a5085c13007e899f150170179ed
Date: Wed, 9 Sep 2020 13:58:25 +0200

    feat(interaction-hotspot): use tags for ignoring hits with other hotspots
    should be safer than checking type of the attach-parent, since the parent may have more child-colliders that are unrelated to hotspots
    (cherry picked from commit fbe6f9d53796dcf2d29645ff3fe3f77197b5dd58​)

#skipci

diff --git a/Engine/Plugins/Enterprise/GLTFExporter/MigratedFromGit/Source/GLTFExporterRuntime/Private/Components/GLTFInteractionHotspotComponent.cpp b/Engine/Plugins/Enterprise/GLTFExporter/MigratedFromGit/Source/GLTFExporterRuntime/Private/Components/GLTFInteractionHotspotComponent.cpp
--- a/Engine/Plugins/Enterprise/GLTFExporter/MigratedFromGit/Source/GLTFExporterRuntime/Private/Components/GLTFInteractionHotspotComponent.cpp
+++ b/Engine/Plugins/Enterprise/GLTFExporter/MigratedFromGit/Source/GLTFExporterRuntime/Private/Components/GLTFInteractionHotspotComponent.cpp
@@ -19,6 +19,7 @@ DEFINE_LOG_CATEGORY_STATIC(LogEditorGLTFInteractionHotspot, Log, All);
 
 namespace
 {
+	const FName NAME_InteractionHotspotTag = TEXT("InteractionHotspot");
 	const FName NAME_LevelEditorModule = TEXT("LevelEditor");
 }
 
@@ -58,6 +59,7 @@ UGLTFInteractionHotspotComponent::UGLTFInteractionHotspotComponent(const FObject
 	SphereComponent->InitSphereRadius(100.0f);
 	SphereComponent->SetVisibility(false);
 	SphereComponent->SetupAttachment(this);
+	SphereComponent->ComponentTags.Add(NAME_InteractionHotspotTag);
 
 	// Setup the most minimalistic collision profile for mouse input events
 	SphereComponent->SetCollisionEnabled(ECollisionEnabled::QueryAndPhysics);
@@ -167,10 +169,9 @@ void UGLTFInteractionHotspotComponent::TickComponent(float DeltaTime, ELevelTick
 
 		if (PlayerController->GetHitResultAtScreenPosition(FVector2D(ColliderScreenLocation), ECC_Visibility, false, HitResult))
 		{
-			const UPrimitiveComponent* HitComponent = HitResult.GetComponent();
-			if (HitComponent != nullptr && HitComponent != SphereComponent)
+			if (const UPrimitiveComponent* HitComponent = HitResult.GetComponent())
 			{
-				bIsHotspotOccluded = HitComponent->GetAttachParent() == nullptr || !HitComponent->GetAttachParent()->IsA<UGLTFInteractionHotspotComponent>();
+				bIsHotspotOccluded = !HitComponent->ComponentTags.Contains(NAME_InteractionHotspotTag);
 			}
 		}
 
