[GLTFExporter] fix(interaction-hotspot): ensure that the default-texture is shown when no image has been selected

commit 456b94f19d93f21ff30d7f85e8f5ceb0bb9e45bf
Date: Wed, 9 Sep 2020 14:43:36 +0200

    fix(interaction-hotspot): ensure that the default-texture is shown when no image has been selected
    (cherry picked from commit f48c0782c90816df1d97757ef45aa51ecbdef754​)

#skipci

diff --git a/Engine/Plugins/Enterprise/GLTFExporter/MigratedFromGit/Source/GLTFExporterRuntime/Private/Components/GLTFInteractionHotspotComponent.cpp b/Engine/Plugins/Enterprise/GLTFExporter/MigratedFromGit/Source/GLTFExporterRuntime/Private/Components/GLTFInteractionHotspotComponent.cpp
--- a/Engine/Plugins/Enterprise/GLTFExporter/MigratedFromGit/Source/GLTFExporterRuntime/Private/Components/GLTFInteractionHotspotComponent.cpp
+++ b/Engine/Plugins/Enterprise/GLTFExporter/MigratedFromGit/Source/GLTFExporterRuntime/Private/Components/GLTFInteractionHotspotComponent.cpp
@@ -209,11 +209,17 @@ void UGLTFInteractionHotspotComponent::OnRegister()
 	Super::OnRegister();
 }
 
-void UGLTFInteractionHotspotComponent::SetActiveImage(class UTexture2D* NewImage)
+void UGLTFInteractionHotspotComponent::SetActiveImage(UTexture2D* NewImage)
 {
 	if (NewImage != ActiveImage)
 	{
-		GetSpriteMaterial()->SetTextureParameterValue("Sprite", NewImage);
+		UTexture* RenderedImage = NewImage;
+		if (RenderedImage == nullptr)
+		{
+			GetSpriteMaterial()->GetTextureParameterDefaultValue(TEXT("Sprite"), RenderedImage);
+		}
+
+		GetSpriteMaterial()->SetTextureParameterValue(TEXT("Sprite"), RenderedImage);
 		ActiveImage = NewImage;
 	}
 
diff --git a/Engine/Plugins/Enterprise/GLTFExporter/MigratedFromGit/Source/GLTFExporterRuntime/Public/Components/GLTFInteractionHotspotComponent.h b/Engine/Plugins/Enterprise/GLTFExporter/MigratedFromGit/Source/GLTFExporterRuntime/Public/Components/GLTFInteractionHotspotComponent.h
--- a/Engine/Plugins/Enterprise/GLTFExporter/MigratedFromGit/Source/GLTFExporterRuntime/Public/Components/GLTFInteractionHotspotComponent.h
+++ b/Engine/Plugins/Enterprise/GLTFExporter/MigratedFromGit/Source/GLTFExporterRuntime/Public/Components/GLTFInteractionHotspotComponent.h
@@ -48,7 +48,7 @@ private:
 
 protected:
 
-	virtual void SetActiveImage(class UTexture2D* NewImage);
+	virtual void SetActiveImage(UTexture2D* NewImage);
 
 public:
 	/* The skeletal mesh actor that will be animated when the hotspot is clicked. */
