[GLTFExporter] refactor(interaction-hotspot): use constants for material-parameter names

commit f2e497850cca91be554e39626cebdffca893bb26
Date: Wed, 9 Sep 2020 14:44:22 +0200

    refactor(interaction-hotspot): use constants for material-parameter names
    (cherry picked from commit 34ed1d3a572ff91d186d058793c445f4b15593ad​)

#skipci

diff --git a/Engine/Plugins/Enterprise/GLTFExporter/MigratedFromGit/Source/GLTFExporterRuntime/Private/Components/GLTFInteractionHotspotComponent.cpp b/Engine/Plugins/Enterprise/GLTFExporter/MigratedFromGit/Source/GLTFExporterRuntime/Private/Components/GLTFInteractionHotspotComponent.cpp
--- a/Engine/Plugins/Enterprise/GLTFExporter/MigratedFromGit/Source/GLTFExporterRuntime/Private/Components/GLTFInteractionHotspotComponent.cpp
+++ b/Engine/Plugins/Enterprise/GLTFExporter/MigratedFromGit/Source/GLTFExporterRuntime/Private/Components/GLTFInteractionHotspotComponent.cpp
@@ -21,6 +21,8 @@ namespace
 {
 	const FName NAME_InteractionHotspotTag = TEXT("InteractionHotspot");
 	const FName NAME_LevelEditorModule = TEXT("LevelEditor");
+	const FName NAME_SpriteParameter = TEXT("Sprite");
+	const FName NAME_OpacityParameter = TEXT("Opacity");
 }
 
 UGLTFInteractionHotspotComponent::UGLTFInteractionHotspotComponent(const FObjectInitializer& ObjectInitializer)
@@ -216,10 +218,10 @@ void UGLTFInteractionHotspotComponent::SetActiveImage(UTexture2D* NewImage)
 		UTexture* RenderedImage = NewImage;
 		if (RenderedImage == nullptr)
 		{
-			GetSpriteMaterial()->GetTextureParameterDefaultValue(TEXT("Sprite"), RenderedImage);
+			GetSpriteMaterial()->GetTextureParameterDefaultValue(NAME_SpriteParameter, RenderedImage);
 		}
 
-		GetSpriteMaterial()->SetTextureParameterValue(TEXT("Sprite"), RenderedImage);
+		GetSpriteMaterial()->SetTextureParameterValue(NAME_SpriteParameter, RenderedImage);
 		ActiveImage = NewImage;
 	}
 
@@ -360,7 +362,7 @@ void UGLTFInteractionHotspotComponent::UpdateSpriteSize()
 
 void UGLTFInteractionHotspotComponent::SetSpriteOpacity(const float Opacity) const
 {
-	GetSpriteMaterial()->SetScalarParameterValue("Opacity", Opacity);
+	GetSpriteMaterial()->SetScalarParameterValue(NAME_OpacityParameter, Opacity);
 }
 
 FIntPoint UGLTFInteractionHotspotComponent::GetCurrentViewportSize()
