[GLTFExporter] fix(interaction-hotspot): calculate correct screen-size in the constructor

commit 5664ff9b8bd7409ccd36143caf84a7a4964316de
Date: Tue, 15 Sep 2020 11:11:53 +0200

    fix(interaction-hotspot): calculate correct screen-size in the constructor
    in some cases BeginPlay / PostLoad aren't triggered, so we need a correct initial size
    (cherry picked from commit 130223608d73ea4c69cae8e3c9d888db77b8e708​)

#skipci

diff --git a/Engine/Plugins/Enterprise/GLTFExporter/MigratedFromGit/Source/GLTFExporterRuntime/Private/Components/GLTFInteractionHotspotComponent.cpp b/Engine/Plugins/Enterprise/GLTFExporter/MigratedFromGit/Source/GLTFExporterRuntime/Private/Components/GLTFInteractionHotspotComponent.cpp
--- a/Engine/Plugins/Enterprise/GLTFExporter/MigratedFromGit/Source/GLTFExporterRuntime/Private/Components/GLTFInteractionHotspotComponent.cpp
+++ b/Engine/Plugins/Enterprise/GLTFExporter/MigratedFromGit/Source/GLTFExporterRuntime/Private/Components/GLTFInteractionHotspotComponent.cpp
@@ -56,6 +56,7 @@ UGLTFInteractionHotspotComponent::UGLTFInteractionHotspotComponent(const FObject
 	DefaultMaterial = UMaterialInstanceDynamic::Create(ConstructorStatics.Material.Object, GetTransientPackage());
 
 	CreateDefaultSpriteElement();
+	SetActiveImage(Image);	// Ensures correct initial screen-size, even if Image == nullptr
 
 	SphereComponent = CreateDefaultSubobject<USphereComponent>(TEXT("Collider"), true);
 	SphereComponent->InitSphereRadius(100.0f);
@@ -211,33 +212,6 @@ void UGLTFInteractionHotspotComponent::OnRegister()
 	Super::OnRegister();
 }
 
-void UGLTFInteractionHotspotComponent::SetActiveImage(UTexture2D* NewImage)
-{
-	if (NewImage != ActiveImage)
-	{
-		UTexture* RenderedImage = NewImage;
-		if (RenderedImage == nullptr)
-		{
-			GetSpriteMaterial()->GetTextureParameterDefaultValue(NAME_SpriteParameter, RenderedImage);
-		}
-
-		GetSpriteMaterial()->SetTextureParameterValue(NAME_SpriteParameter, RenderedImage);
-		ActiveImage = NewImage;
-	}
-
-	const FVector2D NewImageSize(
-		NewImage != nullptr ? NewImage->GetSurfaceWidth() : 32.0f,
-		NewImage != nullptr ? NewImage->GetSurfaceHeight() : 32.0f);
-
-	if (NewImageSize != ActiveImageSize)
-	{
-		ActiveImageSize = NewImageSize;
-	}
-
-	// NOTE: we do this even if size is unchanged since the last update may have failed
-	UpdateSpriteSize();
-}
-
 void UGLTFInteractionHotspotComponent::BeginCursorOver(UPrimitiveComponent* TouchedComponent)
 {
 	SetActiveImage(CalculateActiveImage(true));
@@ -278,6 +252,21 @@ void UGLTFInteractionHotspotComponent::Clicked(UPrimitiveComponent* TouchedCompo
 	SetActiveImage(CalculateActiveImage(true));
 }
 
+void UGLTFInteractionHotspotComponent::SetActiveImage(UTexture2D* NewImage)
+{
+	UTexture* DefaultTexture;
+	GetSpriteMaterial()->GetTextureParameterDefaultValue(NAME_SpriteParameter, DefaultTexture);
+
+	UTexture* SpriteTexture = NewImage != nullptr ? NewImage : DefaultTexture;
+	GetSpriteMaterial()->SetTextureParameterValue(NAME_SpriteParameter, SpriteTexture);
+
+	ActiveImage = NewImage;
+	ActiveImageSize = { SpriteTexture->GetSurfaceWidth(), SpriteTexture->GetSurfaceHeight() };
+
+	// NOTE: we do this even if size is unchanged since the last update may have failed
+	UpdateSpriteSize();
+}
+
 UTexture2D* UGLTFInteractionHotspotComponent::CalculateActiveImage(bool bCursorOver) const
 {
 	// A list of candidates ordered by descending priority
diff --git a/Engine/Plugins/Enterprise/GLTFExporter/MigratedFromGit/Source/GLTFExporterRuntime/Public/Components/GLTFInteractionHotspotComponent.h b/Engine/Plugins/Enterprise/GLTFExporter/MigratedFromGit/Source/GLTFExporterRuntime/Public/Components/GLTFInteractionHotspotComponent.h
--- a/Engine/Plugins/Enterprise/GLTFExporter/MigratedFromGit/Source/GLTFExporterRuntime/Public/Components/GLTFInteractionHotspotComponent.h
+++ b/Engine/Plugins/Enterprise/GLTFExporter/MigratedFromGit/Source/GLTFExporterRuntime/Public/Components/GLTFInteractionHotspotComponent.h
@@ -44,11 +44,9 @@ private:
 	UFUNCTION()
 	void Clicked(UPrimitiveComponent* TouchedComponent, FKey ButtonPressed);
 
-	UTexture2D* CalculateActiveImage(bool bCursorOver) const;
-
-protected:
+	void SetActiveImage(UTexture2D* NewImage);
 
-	virtual void SetActiveImage(UTexture2D* NewImage);
+	UTexture2D* CalculateActiveImage(bool bCursorOver) const;
 
 public:
 	/* The skeletal mesh actor that will be animated when the hotspot is clicked. */
