[GLTFExporter] feat(interaction-hotspot): use a flag to determine interactivity instead of changing state for the collider

commit c0f63f627dbd567b89273a389fad6aa9a615cce9
Date: Tue, 15 Sep 2020 11:13:57 +0200

    feat(interaction-hotspot): use a flag to determine interactivity instead of changing state for the collider
    disabling collisions prevents line-traces for visibility from working as intended
    (cherry picked from commit 9b1d0c2ae1d8b68fa68f4b0dd6be288fb7ed3c7b​)

#skipci

diff --git a/Engine/Plugins/Enterprise/GLTFExporter/MigratedFromGit/Source/GLTFExporterRuntime/Private/Components/GLTFInteractionHotspotComponent.cpp b/Engine/Plugins/Enterprise/GLTFExporter/MigratedFromGit/Source/GLTFExporterRuntime/Private/Components/GLTFInteractionHotspotComponent.cpp
--- a/Engine/Plugins/Enterprise/GLTFExporter/MigratedFromGit/Source/GLTFExporterRuntime/Private/Components/GLTFInteractionHotspotComponent.cpp
+++ b/Engine/Plugins/Enterprise/GLTFExporter/MigratedFromGit/Source/GLTFExporterRuntime/Private/Components/GLTFInteractionHotspotComponent.cpp
@@ -36,6 +36,7 @@ UGLTFInteractionHotspotComponent::UGLTFInteractionHotspotComponent(const FObject
 	ActiveImage(nullptr),
 	ActiveImageSize(0.0f, 0.0f),
 	bToggled(bToggled),
+	bIsInteractable(true),
 	RealtimeSecondsWhenLastInSight(0.0f),
 	RealtimeSecondsWhenLastHidden(0.0f)
 {
@@ -201,9 +202,7 @@ void UGLTFInteractionHotspotComponent::TickComponent(float DeltaTime, ELevelTick
 		}
 
 		SetSpriteOpacity(Opacity);
-
-		const ECollisionEnabled::Type CollisionEnabled = Opacity >= 0.5f ? ECollisionEnabled::QueryAndPhysics :  ECollisionEnabled::NoCollision;
-		SphereComponent->SetCollisionEnabled(CollisionEnabled);
+		bIsInteractable = Opacity >= 0.5f;
 	}
 }
 
@@ -214,7 +213,10 @@ void UGLTFInteractionHotspotComponent::OnRegister()
 
 void UGLTFInteractionHotspotComponent::BeginCursorOver(UPrimitiveComponent* TouchedComponent)
 {
-	SetActiveImage(CalculateActiveImage(true));
+	if (bIsInteractable)
+	{
+		SetActiveImage(CalculateActiveImage(true));
+	}
 }
 
 void UGLTFInteractionHotspotComponent::EndCursorOver(UPrimitiveComponent* TouchedComponent)
@@ -224,6 +226,11 @@ void UGLTFInteractionHotspotComponent::EndCursorOver(UPrimitiveComponent* Touche
 
 void UGLTFInteractionHotspotComponent::Clicked(UPrimitiveComponent* TouchedComponent, FKey ButtonPressed)
 {
+	if (!bIsInteractable)
+	{
+		return;
+	}
+
 	const bool bReverseAnimation = bToggled;
 
 	if (SkeletalMeshActor != nullptr && AnimationSequence != nullptr)
diff --git a/Engine/Plugins/Enterprise/GLTFExporter/MigratedFromGit/Source/GLTFExporterRuntime/Public/Components/GLTFInteractionHotspotComponent.h b/Engine/Plugins/Enterprise/GLTFExporter/MigratedFromGit/Source/GLTFExporterRuntime/Public/Components/GLTFInteractionHotspotComponent.h
--- a/Engine/Plugins/Enterprise/GLTFExporter/MigratedFromGit/Source/GLTFExporterRuntime/Public/Components/GLTFInteractionHotspotComponent.h
+++ b/Engine/Plugins/Enterprise/GLTFExporter/MigratedFromGit/Source/GLTFExporterRuntime/Public/Components/GLTFInteractionHotspotComponent.h
@@ -88,6 +88,8 @@ private:
 
 	bool bToggled;
 
+	bool bIsInteractable;
+
 	float RealtimeSecondsWhenLastInSight;
 
 	float RealtimeSecondsWhenLastHidden;
