[GLTFExporter] chore(interaction-hotspot): add TODOs for fixing incorrect behavior

commit c8a6f4e0579f0ba2e4c34bc0f4bb2ea0d7d36af4
Date: Tue, 15 Sep 2020 16:12:27 +0200

    chore(interaction-hotspot): add TODOs for fixing incorrect behavior
    (cherry picked from commit 8b10793366db454811a15d397519eed2a2f6d2e8​)

#skipci

diff --git a/Engine/Plugins/Enterprise/GLTFExporter/MigratedFromGit/Source/GLTFExporterRuntime/Private/Components/GLTFInteractionHotspotComponent.cpp b/Engine/Plugins/Enterprise/GLTFExporter/MigratedFromGit/Source/GLTFExporterRuntime/Private/Components/GLTFInteractionHotspotComponent.cpp
--- a/Engine/Plugins/Enterprise/GLTFExporter/MigratedFromGit/Source/GLTFExporterRuntime/Private/Components/GLTFInteractionHotspotComponent.cpp
+++ b/Engine/Plugins/Enterprise/GLTFExporter/MigratedFromGit/Source/GLTFExporterRuntime/Private/Components/GLTFInteractionHotspotComponent.cpp
@@ -54,9 +54,17 @@ UGLTFInteractionHotspotComponent::UGLTFInteractionHotspotComponent(const FObject
 	};
 	static FConstructorStatics ConstructorStatics;
 
+	// TODO: wait until initialization of the component before instantiating the dynamic material instance.
+	// Right now crashes occur sometimes when running in a standalone window. Crash-logs indicate that the
+	// instantiation is happening outside of the game-thread.
+	// We should therefore use a suitable entry-point where the material can be safely instantiated.
+
 	DefaultMaterial = UMaterialInstanceDynamic::Create(ConstructorStatics.Material.Object, GetTransientPackage());
 
 	CreateDefaultSpriteElement();
+
+	// TODO: unify and move all initial calls to SetActiveImage into a suitable entry-point that is called
+	// for all platforms, and that is called late enough to give us access to a valid viewport with a non-zero size.
 	SetActiveImage(Image);	// Ensures correct initial screen-size, even if Image == nullptr
 
 	SphereComponent = CreateDefaultSubobject<USphereComponent>(TEXT("Collider"), true);
@@ -117,6 +125,9 @@ void UGLTFInteractionHotspotComponent::PostLoad()
 {
 	Super::PostLoad();
 
+	// TODO: unify and move all initial calls to SetActiveImage into a suitable entry-point that is called
+	// for all platforms, and that is called late enough to give us access to a valid viewport with a non-zero size.
+
 	CreateDefaultSpriteElement();	// NOTE: needed in order to overwrite any persisted element
 	SetActiveImage(Image);
 }
@@ -125,6 +136,9 @@ void UGLTFInteractionHotspotComponent::BeginPlay()
 {
 	Super::BeginPlay();
 
+	// TODO: unify and move all initial calls to SetActiveImage into a suitable entry-point that is called
+	// for all platforms, and that is called late enough to give us access to a valid viewport with a non-zero size.
+
 	SetActiveImage(Image);
 }
 
