[GLTFExporter] fix(interaction-hotspot): fix inability to click a hotspot in the editor to select it

commit 654d08706a389ce3a195164c539208ba189fd401
Date: Tue, 22 Sep 2020 10:34:46 +0200

    fix(interaction-hotspot): fix inability to click a hotspot in the editor to select it
    UE 4.25.2 removes ability to click transparent scene-proxies that have no valid BathHitProxyID. Se we circumvent the issue by using a masked material instead
    (cherry picked from commit fe3e9536af33b3b3796015ed42aa9692be03edbd​)

#skipci

diff --git a/Engine/Plugins/Enterprise/GLTFExporter/MigratedFromGit/Source/GLTFExporterRuntime/Private/Components/GLTFInteractionHotspotComponent.cpp b/Engine/Plugins/Enterprise/GLTFExporter/MigratedFromGit/Source/GLTFExporterRuntime/Private/Components/GLTFInteractionHotspotComponent.cpp
--- a/Engine/Plugins/Enterprise/GLTFExporter/MigratedFromGit/Source/GLTFExporterRuntime/Private/Components/GLTFInteractionHotspotComponent.cpp
+++ b/Engine/Plugins/Enterprise/GLTFExporter/MigratedFromGit/Source/GLTFExporterRuntime/Private/Components/GLTFInteractionHotspotComponent.cpp
@@ -22,7 +22,7 @@ namespace
 	const FName NAME_InteractionHotspotTag = TEXT("InteractionHotspot");
 	const FName NAME_LevelEditorModule = TEXT("LevelEditor");
 	const FName NAME_SpriteParameter = TEXT("Sprite");
-	const FName NAME_OpacityParameter = TEXT("Opacity");
+	const FName NAME_OpacityParameter = TEXT("OpacityMask");
 }
 
 UGLTFInteractionHotspotComponent::UGLTFInteractionHotspotComponent(const FObjectInitializer& ObjectInitializer)
