[GLTFExporter] refactor(interaction-hotspot): move logic from the component to the actor, and remove component

commit c8467c1a3e661efd3bccda7169920e74ab137100
Date: Thu, 24 Sep 2020 11:15:25 +0200

    refactor(interaction-hotspot): move logic from the component to the actor, and remove component
    (cherry picked from commit de24997cfdcba65b749e2200f7c5648337d588bf​)

#skipci

diff --git a/Engine/Plugins/Enterprise/GLTFExporter/MigratedFromGit/Source/GLTFExporterRuntime/Private/Actors/GLTFInteractionHotspotActor.cpp b/Engine/Plugins/Enterprise/GLTFExporter/MigratedFromGit/Source/GLTFExporterRuntime/Private/Actors/GLTFInteractionHotspotActor.cpp
--- a/Engine/Plugins/Enterprise/GLTFExporter/MigratedFromGit/Source/GLTFExporterRuntime/Private/Actors/GLTFInteractionHotspotActor.cpp
+++ b/Engine/Plugins/Enterprise/GLTFExporter/MigratedFromGit/Source/GLTFExporterRuntime/Private/Actors/GLTFInteractionHotspotActor.cpp
@@ -1,11 +1,396 @@
 // Copyright Epic Games, Inc. All Rights Reserved.
 
 #include "Actors/GLTFInteractionHotspotActor.h"
+#include "Animation/SkeletalMeshActor.h"
+#include "Animation/AnimSequence.h"
+#include "Components/SphereComponent.h"
+#include "Components/MaterialBillboardComponent.h"
+#include "Kismet/GameplayStatics.h"
+#include "UObject/ConstructorHelpers.h"
+#include "Animation/AnimSingleNodeInstance.h"
+
+#if WITH_EDITOR
+#include "LevelEditor.h"
+#include "Modules/ModuleManager.h"
+#include "SLevelViewport.h"
+#include "Slate/SceneViewport.h"
+#endif
+
+DEFINE_LOG_CATEGORY_STATIC(LogEditorGLTFInteractionHotspot, Log, All);
+
+namespace
+{
+	const FName NAME_InteractionHotspotTag = TEXT("InteractionHotspot");
+	const FName NAME_LevelEditorModule = TEXT("LevelEditor");
+	const FName NAME_SpriteParameter = TEXT("Sprite");
+	const FName NAME_OpacityParameter = TEXT("OpacityMask");
+}
 
 AGLTFInteractionHotspotActor::AGLTFInteractionHotspotActor(const FObjectInitializer& ObjectInitializer)
-	: Super(ObjectInitializer)
+	: Super(ObjectInitializer),
+	Image(nullptr),
+	HoveredImage(nullptr),
+	ToggledImage(nullptr),
+	ToggledHoveredImage(nullptr),
+	SphereComponent(nullptr),
+	DefaultMaterial(nullptr),
+	ActiveImage(nullptr),
+	ActiveImageSize(0.0f, 0.0f),
+	bToggled(bToggled),
+	bIsInteractable(true),
+	RealtimeSecondsWhenLastInSight(0.0f),
+	RealtimeSecondsWhenLastHidden(0.0f)
+{
+	USceneComponent* SceneRootComponent = CreateDefaultSubobject<USceneComponent>(TEXT("Root"));
+	SetRootComponent(SceneRootComponent);
+	AddInstanceComponent(SceneRootComponent);
+	SceneRootComponent->SetMobility(EComponentMobility::Movable);
+
+	BillboardComponent = CreateDefaultSubobject<UMaterialBillboardComponent>(TEXT("MaterialBillboardComponent"));
+	AddInstanceComponent(BillboardComponent);
+	BillboardComponent->SetupAttachment(RootComponent);
+	BillboardComponent->SetMobility(EComponentMobility::Movable);
+
+	SphereComponent = CreateDefaultSubobject<USphereComponent>(TEXT("SphereComponent"));
+	AddInstanceComponent(BillboardComponent);
+	SphereComponent->SetupAttachment(RootComponent);
+	SphereComponent->SetMobility(EComponentMobility::Movable);
+	SphereComponent->ComponentTags.Add(NAME_InteractionHotspotTag);
+	SphereComponent->InitSphereRadius(100.0f);
+	SphereComponent->SetVisibility(false);
+
+	// Setup the most minimalistic collision profile for mouse input events
+	SphereComponent->SetCollisionEnabled(ECollisionEnabled::QueryAndPhysics);
+	SphereComponent->SetCollisionResponseToAllChannels(ECR_Ignore);
+	SphereComponent->SetCollisionResponseToChannel(ECC_Visibility, ECR_Block);
+	SphereComponent->SetGenerateOverlapEvents(false);
+
+	// Respond to interactions with the sphere-component
+	SphereComponent->OnBeginCursorOver.AddDynamic(this, &AGLTFInteractionHotspotActor::BeginCursorOver);
+	SphereComponent->OnEndCursorOver.AddDynamic(this, &AGLTFInteractionHotspotActor::EndCursorOver);
+	SphereComponent->OnClicked.AddDynamic(this, &AGLTFInteractionHotspotActor::Clicked);
+
+	const ConstructorHelpers::FObjectFinder<UMaterial> MaterialRef(TEXT("/GLTFExporter/Materials/Hotspot"));
+	DefaultMaterial = MaterialRef.Object;
+
+	PrimaryActorTick.bCanEverTick = true;
+	PrimaryActorTick.bStartWithTickEnabled = true;
+}
+
+
+#if WITH_EDITOR
+void AGLTFInteractionHotspotActor::PostEditChangeProperty(FPropertyChangedEvent& PropertyChangedEvent)
+{
+	Super::PostEditChangeProperty(PropertyChangedEvent);
+
+	const FProperty* PropertyThatChanged = PropertyChangedEvent.Property;
+
+	if (PropertyThatChanged)
+	{
+		const FString PropertyName = PropertyThatChanged->GetName();
+
+		if (PropertyName == TEXT("Image"))
+		{
+			SetActiveImage(Image);
+		}
+		else if (PropertyName == TEXT("SkeletalMeshActor"))
+		{
+			if (AnimationSequence != nullptr && AnimationSequence->GetSkeleton() != SkeletalMeshActor->GetSkeletalMeshComponent()->SkeletalMesh->Skeleton)
+			{
+				UE_LOG(LogEditorGLTFInteractionHotspot, Warning, TEXT("The skeleton of this actor is not compatible with the previously assigned animation sequence"));
+			}
+		}
+		else if (PropertyName == TEXT("AnimationSequence"))
+		{
+			if (SkeletalMeshActor != nullptr && SkeletalMeshActor->GetSkeletalMeshComponent()->SkeletalMesh->Skeleton != AnimationSequence->GetSkeleton())
+			{
+				UE_LOG(LogEditorGLTFInteractionHotspot, Warning, TEXT("This animation sequence is not compatible with the skeleton of the previously assigned actor"));
+			}
+		}
+	}
+}
+#endif // WITH_EDITOR
+
+void AGLTFInteractionHotspotActor::PostRegisterAllComponents()
+{
+	Super::PostRegisterAllComponents();
+
+	SetupSpriteElement();
+	SetActiveImage(Image);
+}
+
+void AGLTFInteractionHotspotActor::Tick(float DeltaTime)
+{
+	UWorld* World = GetWorld();
+
+	// TODO: is is safe to assume that we want the first controller for projections?
+	APlayerController* PlayerController = World != nullptr ? World->GetFirstPlayerController() : nullptr;
+	if (PlayerController == nullptr)
+	{
+		return;
+	}
+
+	const FVector ColliderLocation = SphereComponent->GetComponentLocation();
+
+	FVector ColliderScreenLocation;
+	if (!PlayerController->ProjectWorldLocationToScreenWithDistance(ColliderLocation, ColliderScreenLocation))
+	{
+		return;
+	}
+
+	// Update scale of the sphere-component to match the screen-size of the active image
+	{
+		const FVector2D CornerScreenLocation = FVector2D(ColliderScreenLocation) + ActiveImageSize * 0.5f;
+		FVector RayLocation;
+		FVector RayDirection;
+
+		if (PlayerController->DeprojectScreenPositionToWorld(CornerScreenLocation.X, CornerScreenLocation.Y, RayLocation, RayDirection))
+		{
+			const FVector ExtentLocation = RayLocation + RayDirection * ColliderScreenLocation.Z;
+			const float NewSphereRadius = (ExtentLocation - ColliderLocation).Size() / SphereComponent->GetShapeScale();
+			const float OldSphereRadius = SphereComponent->GetUnscaledSphereRadius();
+
+			if (FMath::Abs(NewSphereRadius - OldSphereRadius) > 0.1f)	// TODO: better epsilon?
+			{
+				SphereComponent->SetSphereRadius(NewSphereRadius);
+			}
+		}
+	}
+
+	// Update opacity and interactivity of the hotspot based on if it its occluded by other objects or not
+	{
+		FHitResult HitResult;
+		bool bIsHotspotOccluded = false;
+
+		if (PlayerController->GetHitResultAtScreenPosition(FVector2D(ColliderScreenLocation), ECC_Visibility, false, HitResult))
+		{
+			if (const UPrimitiveComponent* HitComponent = HitResult.GetComponent())
+			{
+				bIsHotspotOccluded = !HitComponent->ComponentTags.Contains(NAME_InteractionHotspotTag);
+			}
+		}
+
+		const float CurrentRealtimeSeconds = UGameplayStatics::GetRealTimeSeconds(World);
+		float Opacity;
+
+		if (bIsHotspotOccluded)
+		{
+			RealtimeSecondsWhenLastHidden = CurrentRealtimeSeconds;
+
+			const float HiddenDuration = FMath::Max(RealtimeSecondsWhenLastHidden - RealtimeSecondsWhenLastInSight, 0.0f);
+			const float FadeOutDuration = 0.5f;
+
+			Opacity = 1.0f - FMath::Clamp(HiddenDuration / FadeOutDuration, 0.0f, 1.0f);
+		}
+		else
+		{
+			RealtimeSecondsWhenLastInSight = CurrentRealtimeSeconds;
+
+			const float VisibleDuration = FMath::Max(RealtimeSecondsWhenLastInSight - RealtimeSecondsWhenLastHidden, 0.0f);
+			const float FadeInDuration = 0.25f;
+
+			Opacity = FMath::Clamp(VisibleDuration / FadeInDuration, 0.0f, 1.0f);
+		}
+
+		SetSpriteOpacity(Opacity);
+		bIsInteractable = Opacity >= 0.5f;
+	}
+}
+
+void AGLTFInteractionHotspotActor::BeginCursorOver(UPrimitiveComponent* TouchedComponent)
+{
+	if (bIsInteractable)
+	{
+		SetActiveImage(CalculateActiveImage(true));
+	}
+}
+
+void AGLTFInteractionHotspotActor::EndCursorOver(UPrimitiveComponent* TouchedComponent)
+{
+	SetActiveImage(CalculateActiveImage(false));
+}
+
+void AGLTFInteractionHotspotActor::Clicked(UPrimitiveComponent* TouchedComponent, FKey ButtonPressed)
+{
+	if (!bIsInteractable)
+	{
+		return;
+	}
+
+	const bool bReverseAnimation = bToggled;
+
+	if (SkeletalMeshActor != nullptr && AnimationSequence != nullptr)
+	{
+		USkeletalMeshComponent* SkeletalMeshComponent = SkeletalMeshActor->GetSkeletalMeshComponent();
+		SkeletalMeshComponent->SetAnimationMode(EAnimationMode::Type::AnimationSingleNode);
+		const float AbsolutePlayRate = FMath::Abs(SkeletalMeshComponent->GetPlayRate());
+		const UAnimSingleNodeInstance* SingleNodeInstance = SkeletalMeshComponent->GetSingleNodeInstance();
+
+		if (SkeletalMeshComponent->IsPlaying() && SingleNodeInstance != nullptr && SingleNodeInstance->GetAnimationAsset() == AnimationSequence)
+		{
+			// If the same animation is already playing, just reverse the play rate for a smooth transition
+			SkeletalMeshComponent->SetPlayRate(AbsolutePlayRate * -1.0f);
+		}
+		else
+		{
+			SkeletalMeshComponent->SetAnimation(AnimationSequence);
+			SkeletalMeshComponent->SetPlayRate(AbsolutePlayRate * (bReverseAnimation ? -1.0f : 1.0f));
+			SkeletalMeshComponent->SetPosition(bReverseAnimation ? AnimationSequence->GetPlayLength() : 0.0f);
+			SkeletalMeshComponent->Play(false);
+		}
+	}
+
+	bToggled = !bToggled;
+
+	SetActiveImage(CalculateActiveImage(true));
+}
+
+void AGLTFInteractionHotspotActor::SetActiveImage(UTexture2D* NewImage)
+{
+	UTexture* DefaultTexture;
+	GetSpriteMaterial()->GetTextureParameterDefaultValue(NAME_SpriteParameter, DefaultTexture);
+
+	UTexture* SpriteTexture = NewImage != nullptr ? NewImage : DefaultTexture;
+	GetSpriteMaterial()->SetTextureParameterValue(NAME_SpriteParameter, SpriteTexture);
+
+	ActiveImage = NewImage;
+	ActiveImageSize = { SpriteTexture->GetSurfaceWidth(), SpriteTexture->GetSurfaceHeight() };
+
+	// NOTE: we do this even if size is unchanged since the last update may have failed
+	UpdateSpriteSize();
+}
+
+UTexture2D* AGLTFInteractionHotspotActor::CalculateActiveImage(bool bCursorOver) const
+{
+	// A list of candidates ordered by descending priority
+	TArray<UTexture2D*> ImageCandidates;
+
+	if (bToggled)
+	{
+		if (bCursorOver)
+		{
+			ImageCandidates.Add(ToggledHoveredImage);
+		}
+
+		ImageCandidates.Add(ToggledImage);
+	}
+	else
+	{
+		if (bCursorOver)
+		{
+			ImageCandidates.Add(HoveredImage);
+		}
+
+		ImageCandidates.Add(Image);
+	}
+
+	ImageCandidates.Add(Image);
+
+	for (UTexture2D* ImageCandidate : ImageCandidates)
+	{
+		if (ImageCandidate != nullptr)
+		{
+			return ImageCandidate;
+		}
+	}
+
+	return nullptr;
+}
+
+void AGLTFInteractionHotspotActor::SetupSpriteElement() const
+{
+	UMaterialInstanceDynamic* MaterialInstance = UMaterialInstanceDynamic::Create(DefaultMaterial, GetTransientPackage());
+
+	FMaterialSpriteElement Element;
+	Element.Material = MaterialInstance;
+	Element.bSizeIsInScreenSpace = true;
+	Element.BaseSizeX = 0.1f;
+	Element.BaseSizeY = 0.1f;
+
+	BillboardComponent->SetElements({ Element });
+}
+
+UMaterialInstanceDynamic* AGLTFInteractionHotspotActor::GetSpriteMaterial() const
+{
+	return static_cast<UMaterialInstanceDynamic*>(BillboardComponent->GetMaterial(0));
+}
+
+void AGLTFInteractionHotspotActor::UpdateSpriteSize()
+{
+	const FIntPoint ViewportSize = GetCurrentViewportSize();
+
+	float BaseSizeX = 0.1f;
+	float BaseSizeY = 0.1f;
+
+	if (ViewportSize.X > 0 && ViewportSize.Y > 0)
+	{
+		BaseSizeX = ActiveImageSize.X / static_cast<float>(ViewportSize.X);
+		BaseSizeY = ActiveImageSize.Y / static_cast<float>(ViewportSize.Y);
+	}
+
+	FMaterialSpriteElement& Element = BillboardComponent->Elements[0];
+
+	if (BaseSizeX != Element.BaseSizeX || BaseSizeY != Element.BaseSizeY)	// TODO: use epsilon for comparison?
+	{
+		Element.BaseSizeX = BaseSizeX;
+		Element.BaseSizeY = BaseSizeY;
+
+		BillboardComponent->MarkRenderStateDirty();
+	}
+}
+
+void AGLTFInteractionHotspotActor::SetSpriteOpacity(const float Opacity) const
+{
+	GetSpriteMaterial()->SetScalarParameterValue(NAME_OpacityParameter, Opacity);
+}
+
+FIntPoint AGLTFInteractionHotspotActor::GetCurrentViewportSize()
+{
+	// TODO: verify that correct size is calculated in the various play-modes and in the editor
+
+	const FViewport* Viewport = nullptr;
+
+	if (UWorld* World = GetWorld())
+	{
+		if (World->IsGameWorld())
+		{
+			if (UGameViewportClient* GameViewportClient = World->GetGameViewport())
+			{
+				Viewport = GameViewportClient->Viewport;
+			}
+		}
+#if WITH_EDITOR
+		else
+		{
+			if (FModuleManager::Get().IsModuleLoaded(NAME_LevelEditorModule))
+			{
+				FLevelEditorModule& LevelEditorModule = FModuleManager::GetModuleChecked<FLevelEditorModule>(NAME_LevelEditorModule);
+
+				if (const TSharedPtr<SLevelViewport> ActiveLevelViewport = LevelEditorModule.GetFirstActiveLevelViewport())
+				{
+					Viewport = ActiveLevelViewport->GetActiveViewport();
+				}
+			}
+		}
+#endif
+	}
+
+	if (Viewport != nullptr)
+	{
+		if (!Viewport->ViewportResizedEvent.IsBoundToObject(this))
+		{
+			Viewport->ViewportResizedEvent.AddUObject(this, &AGLTFInteractionHotspotActor::ViewportResized);
+		}
+
+		return Viewport->GetSizeXY();
+	}
+	else
+	{
+		return FIntPoint(0, 0);
+	}
+}
+
+void AGLTFInteractionHotspotActor::ViewportResized(FViewport*, uint32)
 {
-	// TODO: Despite the HideCategory class specifier present in both the hotspot actor and the component, the category "Sprite" bleeds through from UBillboardComponent
-	InteractionHotspotComponent = CreateDefaultSubobject<UGLTFInteractionHotspotComponent>(TEXT("InteractionHotspotComponent"));
-	RootComponent = InteractionHotspotComponent;
+	UpdateSpriteSize();
 }
diff --git a/Engine/Plugins/Enterprise/GLTFExporter/MigratedFromGit/Source/GLTFExporterRuntime/Private/Components/GLTFInteractionHotspotComponent.cpp b/Engine/Plugins/Enterprise/GLTFExporter/MigratedFromGit/Source/GLTFExporterRuntime/Private/Components/GLTFInteractionHotspotComponent.cpp
deleted file mode 100644
diff --git a/Engine/Plugins/Enterprise/GLTFExporter/MigratedFromGit/Source/GLTFExporterRuntime/Public/Actors/GLTFInteractionHotspotActor.h b/Engine/Plugins/Enterprise/GLTFExporter/MigratedFromGit/Source/GLTFExporterRuntime/Public/Actors/GLTFInteractionHotspotActor.h
--- a/Engine/Plugins/Enterprise/GLTFExporter/MigratedFromGit/Source/GLTFExporterRuntime/Public/Actors/GLTFInteractionHotspotActor.h
+++ b/Engine/Plugins/Enterprise/GLTFExporter/MigratedFromGit/Source/GLTFExporterRuntime/Public/Actors/GLTFInteractionHotspotActor.h
@@ -2,10 +2,12 @@
 
 #pragma once
 
-#include "GameFramework/Actor.h"
-#include "Components/GLTFInteractionHotspotComponent.h"
 #include "GLTFInteractionHotspotActor.generated.h"
 
+class ASkeletalMeshActor;
+class UMaterialBillboardComponent;
+class USphereComponent;
+
 /**
  * Actor wrapper for the GLTF hotspot component. Appears as a billboard and allows playback of skeletal animations when cursor input is enabled.
  */
@@ -16,9 +18,89 @@ class GLTFEXPORTERRUNTIME_API AGLTFInteractionHotspotActor : public AActor
 	//~ Begin UObject Interface
 public:
 	AGLTFInteractionHotspotActor(const FObjectInitializer& ObjectInitializer);
+#if WITH_EDITOR
+	virtual void PostEditChangeProperty(FPropertyChangedEvent& PropertyChangedEvent) override;
+#endif // WITH_EDITOR
 	//~ End UObject Interface
 
+	//~ Begin AActor Interface
+	virtual void PostRegisterAllComponents() override;
+	virtual void Tick(float DeltaSeconds) override;
+	//~ End AActor Interface
+
+private:
+	UFUNCTION()
+	void BeginCursorOver(UPrimitiveComponent* TouchedComponent);
+
+	UFUNCTION()
+	void EndCursorOver(UPrimitiveComponent* TouchedComponent);
+
+	UFUNCTION()
+	void Clicked(UPrimitiveComponent* TouchedComponent, FKey ButtonPressed);
+
+	void SetActiveImage(UTexture2D* NewImage);
+
+	UTexture2D* CalculateActiveImage(bool bCursorOver) const;
+
+public:
+
+	/* The skeletal mesh actor that will be animated when the hotspot is clicked. */
+	UPROPERTY(EditAnywhere, BlueprintReadWrite, Category = "GLTF Interaction Hotspot")
+	ASkeletalMeshActor* SkeletalMeshActor;
+
+	/* The animation that will be played on the skeletal mesh actor. Must be compatible with its skeletal mesh asset. */
+	UPROPERTY(EditAnywhere, BlueprintReadWrite, Category = "GLTF Interaction Hotspot")
+	UAnimSequence* AnimationSequence;
+
+	/* The billboard image that will be shown when the hotspot is in an inactive state or one without a specified image. */
+	UPROPERTY(EditAnywhere, BlueprintReadWrite, Category = "GLTF Interaction Hotspot")
+	UTexture2D* Image;
+
+	/** The optional billboard image that will be shown when a cursor enters the hotspot. */
+	UPROPERTY(EditAnywhere, BlueprintReadWrite, Category = "GLTF Interaction Hotspot")
+	UTexture2D* HoveredImage;
+
+	/** The optional billboard image that will be shown when the hotspot is toggled by a click. */
+	UPROPERTY(EditAnywhere, BlueprintReadWrite, Category = "GLTF Interaction Hotspot")
+	UTexture2D* ToggledImage;
+
+	/** The optional billboard image that will be shown when the hotspot is toggled by a click and a cursor enters it. */
+	UPROPERTY(EditAnywhere, BlueprintReadWrite, Category = "GLTF Interaction Hotspot")
+	UTexture2D* ToggledHoveredImage;
+
 private:
-	UPROPERTY(Category = "GLTF Interaction Hotspot", VisibleAnywhere, BlueprintReadOnly, meta = (ExposeFunctionCategories = "GLTF Interaction Hotspot", AllowPrivateAccess = "true"))
-	UGLTFInteractionHotspotComponent* InteractionHotspotComponent;
+
+	UPROPERTY(transient, DuplicateTransient)
+	UMaterialBillboardComponent* BillboardComponent;
+
+	UPROPERTY(transient, duplicatetransient)
+	USphereComponent* SphereComponent;
+
+	UPROPERTY(transient, duplicatetransient)
+	UMaterialInterface* DefaultMaterial;
+
+	UPROPERTY(transient, duplicatetransient)
+	UTexture* ActiveImage;
+
+	FVector2D ActiveImageSize;
+
+	bool bToggled;
+
+	bool bIsInteractable;
+
+	float RealtimeSecondsWhenLastInSight;
+
+	float RealtimeSecondsWhenLastHidden;
+
+	void SetupSpriteElement() const;
+
+	UMaterialInstanceDynamic* GetSpriteMaterial() const;
+
+	void UpdateSpriteSize();
+
+	void SetSpriteOpacity(const float Opacity) const;
+
+	FIntPoint GetCurrentViewportSize();
+
+	void ViewportResized(FViewport*, uint32);
 };
diff --git a/Engine/Plugins/Enterprise/GLTFExporter/MigratedFromGit/Source/GLTFExporterRuntime/Public/Components/GLTFInteractionHotspotComponent.h b/Engine/Plugins/Enterprise/GLTFExporter/MigratedFromGit/Source/GLTFExporterRuntime/Public/Components/GLTFInteractionHotspotComponent.h
deleted file mode 100644
