[GLTFExporter] refactor(interaction-hotspot): use the hotspot-actor for conversion

commit 4138c4c348f8714814191692fc324a626bc1299e
Date: Thu, 24 Sep 2020 11:17:24 +0200

    refactor(interaction-hotspot): use the hotspot-actor for conversion
    (cherry picked from commit 923f68cab1a5db08d2e5694330910db35b9ccbf8​)

#skipci

diff --git a/Engine/Plugins/Enterprise/GLTFExporter/MigratedFromGit/Source/GLTFExporter/Private/Builders/GLTFConvertBuilder.cpp b/Engine/Plugins/Enterprise/GLTFExporter/MigratedFromGit/Source/GLTFExporter/Private/Builders/GLTFConvertBuilder.cpp
--- a/Engine/Plugins/Enterprise/GLTFExporter/MigratedFromGit/Source/GLTFExporter/Private/Builders/GLTFConvertBuilder.cpp
+++ b/Engine/Plugins/Enterprise/GLTFExporter/MigratedFromGit/Source/GLTFExporter/Private/Builders/GLTFConvertBuilder.cpp
@@ -289,12 +289,12 @@ FGLTFJsonLightMapIndex FGLTFConvertBuilder::GetOrAddLightMap(const UStaticMeshCo
 	return LightMapConverter.GetOrAdd(StaticMeshComponent);
 }
 
-FGLTFJsonHotspotIndex FGLTFConvertBuilder::GetOrAddHotspot(const UGLTFInteractionHotspotComponent* HotspotComponent)
+FGLTFJsonHotspotIndex FGLTFConvertBuilder::GetOrAddHotspot(const AGLTFInteractionHotspotActor* HotspotActor)
 {
-	if (HotspotComponent == nullptr)
+	if (HotspotActor == nullptr)
 	{
 		return FGLTFJsonHotspotIndex(INDEX_NONE);
 	}
 
-	return HotspotConverter.GetOrAdd(HotspotComponent);
+	return HotspotConverter.GetOrAdd(HotspotActor);
 }
diff --git a/Engine/Plugins/Enterprise/GLTFExporter/MigratedFromGit/Source/GLTFExporter/Private/Builders/GLTFConvertBuilder.h b/Engine/Plugins/Enterprise/GLTFExporter/MigratedFromGit/Source/GLTFExporter/Private/Builders/GLTFConvertBuilder.h
--- a/Engine/Plugins/Enterprise/GLTFExporter/MigratedFromGit/Source/GLTFExporter/Private/Builders/GLTFConvertBuilder.h
+++ b/Engine/Plugins/Enterprise/GLTFExporter/MigratedFromGit/Source/GLTFExporter/Private/Builders/GLTFConvertBuilder.h
@@ -60,7 +60,7 @@ public:
 	FGLTFJsonBackdropIndex GetOrAddBackdrop(const AActor* BackdropActor);
 	FGLTFJsonVariationIndex GetOrAddVariation(const ALevelVariantSetsActor* LevelVariantSetsActor);
 	FGLTFJsonLightMapIndex GetOrAddLightMap(const UStaticMeshComponent* StaticMeshComponent);
-	FGLTFJsonHotspotIndex GetOrAddHotspot(const UGLTFInteractionHotspotComponent* HotspotComponent);
+	FGLTFJsonHotspotIndex GetOrAddHotspot(const AGLTFInteractionHotspotActor* HotspotActor);
 
 private:
 
diff --git a/Engine/Plugins/Enterprise/GLTFExporter/MigratedFromGit/Source/GLTFExporter/Private/Converters/GLTFHotspotConverters.cpp b/Engine/Plugins/Enterprise/GLTFExporter/MigratedFromGit/Source/GLTFExporter/Private/Converters/GLTFHotspotConverters.cpp
--- a/Engine/Plugins/Enterprise/GLTFExporter/MigratedFromGit/Source/GLTFExporter/Private/Converters/GLTFHotspotConverters.cpp
+++ b/Engine/Plugins/Enterprise/GLTFExporter/MigratedFromGit/Source/GLTFExporter/Private/Converters/GLTFHotspotConverters.cpp
@@ -4,17 +4,17 @@
 #include "Json/FGLTFJsonHotspot.h"
 #include "Builders/GLTFContainerBuilder.h"
 
-FGLTFJsonHotspotIndex FGLTFHotspotConverter::Convert(const UGLTFInteractionHotspotComponent* HotspotComponent)
+FGLTFJsonHotspotIndex FGLTFHotspotConverter::Convert(const AGLTFInteractionHotspotActor* HotspotActor)
 {
 	// TODO: should we warn and / or return INDEX_NONE if the hotspot has no valid Image assigned?
 
 	FGLTFJsonHotspot JsonHotspot;
-	HotspotComponent->GetName(JsonHotspot.Name);
+	HotspotActor->GetName(JsonHotspot.Name);
 	JsonHotspot.Animation = FGLTFJsonAnimationIndex(INDEX_NONE);	// TODO: assign the animation property once we have support for exporting animations
-	JsonHotspot.Image = Builder.GetOrAddTexture(HotspotComponent->Image);
-	JsonHotspot.HoveredImage = Builder.GetOrAddTexture(HotspotComponent->HoveredImage);
-	JsonHotspot.ToggledImage = Builder.GetOrAddTexture(HotspotComponent->ToggledImage);
-	JsonHotspot.ToggledHoveredImage = Builder.GetOrAddTexture(HotspotComponent->ToggledHoveredImage);
+	JsonHotspot.Image = Builder.GetOrAddTexture(HotspotActor->Image);
+	JsonHotspot.HoveredImage = Builder.GetOrAddTexture(HotspotActor->HoveredImage);
+	JsonHotspot.ToggledImage = Builder.GetOrAddTexture(HotspotActor->ToggledImage);
+	JsonHotspot.ToggledHoveredImage = Builder.GetOrAddTexture(HotspotActor->ToggledHoveredImage);
 
 	return Builder.AddHotspot(JsonHotspot);
 }
diff --git a/Engine/Plugins/Enterprise/GLTFExporter/MigratedFromGit/Source/GLTFExporter/Private/Converters/GLTFHotspotConverters.h b/Engine/Plugins/Enterprise/GLTFExporter/MigratedFromGit/Source/GLTFExporter/Private/Converters/GLTFHotspotConverters.h
--- a/Engine/Plugins/Enterprise/GLTFExporter/MigratedFromGit/Source/GLTFExporter/Private/Converters/GLTFHotspotConverters.h
+++ b/Engine/Plugins/Enterprise/GLTFExporter/MigratedFromGit/Source/GLTFExporter/Private/Converters/GLTFHotspotConverters.h
@@ -5,11 +5,11 @@
 #include "Json/GLTFJsonIndex.h"
 #include "Converters/GLTFConverter.h"
 #include "Converters/GLTFBuilderContext.h"
-#include "Components/GLTFInteractionHotspotComponent.h"
+#include "Actors/GLTFInteractionHotspotActor.h"
 
-class FGLTFHotspotConverter final : public FGLTFBuilderContext, public TGLTFConverter<FGLTFJsonHotspotIndex, const UGLTFInteractionHotspotComponent*>
+class FGLTFHotspotConverter final : public FGLTFBuilderContext, public TGLTFConverter<FGLTFJsonHotspotIndex, const AGLTFInteractionHotspotActor*>
 {
 	using FGLTFBuilderContext::FGLTFBuilderContext;
 
-	virtual FGLTFJsonHotspotIndex Convert(const UGLTFInteractionHotspotComponent* HotspotComponent) override;
+	virtual FGLTFJsonHotspotIndex Convert(const AGLTFInteractionHotspotActor* HotspotActor) override;
 };
diff --git a/Engine/Plugins/Enterprise/GLTFExporter/MigratedFromGit/Source/GLTFExporter/Private/Converters/GLTFNameUtility.cpp b/Engine/Plugins/Enterprise/GLTFExporter/MigratedFromGit/Source/GLTFExporter/Private/Converters/GLTFNameUtility.cpp
--- a/Engine/Plugins/Enterprise/GLTFExporter/MigratedFromGit/Source/GLTFExporter/Private/Converters/GLTFNameUtility.cpp
+++ b/Engine/Plugins/Enterprise/GLTFExporter/MigratedFromGit/Source/GLTFExporter/Private/Converters/GLTFNameUtility.cpp
@@ -1,7 +1,6 @@
 // Copyright Epic Games, Inc. All Rights Reserved.
 
 #include "Converters/GLTFNameUtility.h"
-#include "Actors/GLTFInteractionHotspotActor.h"
 #include "Animation/SkeletalMeshActor.h"
 #include "Engine.h"
 
@@ -71,8 +70,3 @@ FString FGLTFNameUtility::GetName(const UCameraComponent* Component)
 {
 	return GetActorNameIfOwnerOfType<ACameraActor>(Component);
 }
-
-FString FGLTFNameUtility::GetName(const UGLTFInteractionHotspotComponent* Component)
-{
-	return GetActorNameIfOwnerOfType<AGLTFInteractionHotspotActor>(Component);
-}
diff --git a/Engine/Plugins/Enterprise/GLTFExporter/MigratedFromGit/Source/GLTFExporter/Private/Converters/GLTFNameUtility.h b/Engine/Plugins/Enterprise/GLTFExporter/MigratedFromGit/Source/GLTFExporter/Private/Converters/GLTFNameUtility.h
--- a/Engine/Plugins/Enterprise/GLTFExporter/MigratedFromGit/Source/GLTFExporter/Private/Converters/GLTFNameUtility.h
+++ b/Engine/Plugins/Enterprise/GLTFExporter/MigratedFromGit/Source/GLTFExporter/Private/Converters/GLTFNameUtility.h
@@ -6,7 +6,7 @@
 
 class ULightComponent;
 class UCameraComponent;
-class UGLTFInteractionHotspotComponent;
+class AGLTFInteractionHotspotActor;
 
 struct FGLTFNameUtility
 {
@@ -27,5 +27,4 @@ struct FGLTFNameUtility
 	static FString GetName(const USkeletalMeshComponent* Component);
 	static FString GetName(const ULightComponent* Component);
 	static FString GetName(const UCameraComponent* Component);
-	static FString GetName(const UGLTFInteractionHotspotComponent* Component);
 };
diff --git a/Engine/Plugins/Enterprise/GLTFExporter/MigratedFromGit/Source/GLTFExporter/Private/Converters/GLTFNodeConverters.cpp b/Engine/Plugins/Enterprise/GLTFExporter/MigratedFromGit/Source/GLTFExporter/Private/Converters/GLTFNodeConverters.cpp
--- a/Engine/Plugins/Enterprise/GLTFExporter/MigratedFromGit/Source/GLTFExporter/Private/Converters/GLTFNodeConverters.cpp
+++ b/Engine/Plugins/Enterprise/GLTFExporter/MigratedFromGit/Source/GLTFExporter/Private/Converters/GLTFNodeConverters.cpp
@@ -5,7 +5,7 @@
 #include "Converters/GLTFConverterUtility.h"
 #include "Converters/GLTFActorUtility.h"
 #include "Converters/GLTFNameUtility.h"
-#include "Components/GLTFInteractionHotspotComponent.h"
+#include "Actors/GLTFInteractionHotspotActor.h"
 
 FGLTFJsonNodeIndex FGLTFComponentConverter::Convert(const USceneComponent* SceneComponent)
 {
@@ -83,22 +83,23 @@ FGLTFJsonNodeIndex FGLTFComponentConverter::Convert(const USceneComponent* Scene
 			Node.Mesh = Builder.GetOrAddMesh(SkeletalMeshComponent);
 		}
 	}
-	else if (const UGLTFInteractionHotspotComponent* HotspotComponent = Cast<UGLTFInteractionHotspotComponent>(SceneComponent))
+	else if (Owner->IsA<AGLTFInteractionHotspotActor>() && bIsRootComponent)
 	{
 		if (Builder.ExportOptions->bExportInteractionHotspots)
 		{
+			const AGLTFInteractionHotspotActor* HotspotActor = Cast<AGLTFInteractionHotspotActor>(Owner);
+
 			if (bExportNonUniformScale)
 			{
 				FGLTFJsonNode HotspotNode;
-				HotspotNode.Name = FGLTFNameUtility::GetName(HotspotComponent);
+				HotspotNode.Name = Owner->GetName();	// TODO: should we add a custom GetName-function in FGLTFNameUtility for this?
 				HotspotNode.Scale = ComponentNodeScale;
-
-				HotspotNode.Hotspot = Builder.GetOrAddHotspot(HotspotComponent);
+				HotspotNode.Hotspot = Builder.GetOrAddHotspot(HotspotActor);
 				Builder.AddChildComponentNode(NodeIndex, HotspotNode);
 			}
 			else
 			{
-				Node.Hotspot = Builder.GetOrAddHotspot(HotspotComponent);
+				Node.Hotspot = Builder.GetOrAddHotspot(HotspotActor);
 			}
 		}
 	}
