[GLTFExporter] feat: set GLTFExporterRuntime module to be "Runtime" to allow it to be used in packaged games

commit 2bf712b05132c381d1d70083340393b469dddf36
Date: Thu, 24 Sep 2020 11:18:49 +0200

    feat: set GLTFExporterRuntime module to be "Runtime" to allow it to be used in packaged games
    (cherry picked from commit 3832580ea45b75d3ee8a4c32d9aa68334b9baf04​)

#skipci

diff --git a/Engine/Plugins/Enterprise/GLTFExporter/MigratedFromGit/GLTFExporter.uplugin b/Engine/Plugins/Enterprise/GLTFExporter/MigratedFromGit/GLTFExporter.uplugin
--- a/Engine/Plugins/Enterprise/GLTFExporter/MigratedFromGit/GLTFExporter.uplugin
+++ b/Engine/Plugins/Enterprise/GLTFExporter/MigratedFromGit/GLTFExporter.uplugin
@@ -21,7 +21,7 @@
 		},
 		{
 			"Name": "GLTFExporterRuntime",
-			"Type": "Editor",
+			"Type": "Runtime",
 			"LoadingPhase": "Default"
 		}
 	],
