[GLTFExporter] feat(interaction-hotspot): use a masked material while in the level-editor

commit 2cc4aa136840bf354596d7ab92c7029ff3cf765d
Date: Fri, 25 Sep 2020 09:42:39 +0200

    feat(interaction-hotspot): use a masked material while in the level-editor
    fixes issues with clicking hotspots in the editor, while making them look good in runtime
    (cherry picked from commit 4747a225a92f649d5ec7e83e76c53a16cb1033cc​)

#skipci

diff --git a/Engine/Plugins/Enterprise/GLTFExporter/MigratedFromGit/Source/GLTFExporterRuntime/Private/Actors/GLTFInteractionHotspotActor.cpp b/Engine/Plugins/Enterprise/GLTFExporter/MigratedFromGit/Source/GLTFExporterRuntime/Private/Actors/GLTFInteractionHotspotActor.cpp
--- a/Engine/Plugins/Enterprise/GLTFExporter/MigratedFromGit/Source/GLTFExporterRuntime/Private/Actors/GLTFInteractionHotspotActor.cpp
+++ b/Engine/Plugins/Enterprise/GLTFExporter/MigratedFromGit/Source/GLTFExporterRuntime/Private/Actors/GLTFInteractionHotspotActor.cpp
@@ -73,6 +73,11 @@ AGLTFInteractionHotspotActor::AGLTFInteractionHotspotActor(const FObjectInitiali
 	const ConstructorHelpers::FObjectFinder<UMaterial> MaterialRef(TEXT("/GLTFExporter/Materials/Hotspot"));
 	DefaultMaterial = MaterialRef.Object;
 
+#if WITH_EDITORONLY_DATA
+	const ConstructorHelpers::FObjectFinder<UMaterial> IconMaterialRef(TEXT("/GLTFExporter/Materials/HotspotIcon"));
+	DefaultIconMaterial = IconMaterialRef.Object;
+#endif // WITH_EDITORONLY_DATA
+
 	PrimaryActorTick.bCanEverTick = true;
 	PrimaryActorTick.bStartWithTickEnabled = true;
 }
@@ -299,7 +304,18 @@ UTexture2D* AGLTFInteractionHotspotActor::CalculateActiveImage(bool bCursorOver)
 
 void AGLTFInteractionHotspotActor::SetupSpriteElement() const
 {
-	UMaterialInstanceDynamic* MaterialInstance = UMaterialInstanceDynamic::Create(DefaultMaterial, GetTransientPackage());
+	UMaterialInstanceDynamic* MaterialInstance;
+
+#if WITH_EDITORONLY_DATA
+	if (GetWorld() != nullptr && GetWorld()->WorldType == EWorldType::Editor)
+	{
+		MaterialInstance = UMaterialInstanceDynamic::Create(DefaultIconMaterial, GetTransientPackage());
+	}
+	else
+#endif // WITH_EDITORONLY_DATA
+	{
+		MaterialInstance = UMaterialInstanceDynamic::Create(DefaultMaterial, GetTransientPackage());
+	}
 
 	FMaterialSpriteElement Element;
 	Element.Material = MaterialInstance;
diff --git a/Engine/Plugins/Enterprise/GLTFExporter/MigratedFromGit/Source/GLTFExporterRuntime/Public/Actors/GLTFInteractionHotspotActor.h b/Engine/Plugins/Enterprise/GLTFExporter/MigratedFromGit/Source/GLTFExporterRuntime/Public/Actors/GLTFInteractionHotspotActor.h
--- a/Engine/Plugins/Enterprise/GLTFExporter/MigratedFromGit/Source/GLTFExporterRuntime/Public/Actors/GLTFInteractionHotspotActor.h
+++ b/Engine/Plugins/Enterprise/GLTFExporter/MigratedFromGit/Source/GLTFExporterRuntime/Public/Actors/GLTFInteractionHotspotActor.h
@@ -79,6 +79,11 @@ private:
 	UPROPERTY(transient, duplicatetransient)
 	UMaterialInterface* DefaultMaterial;
 
+#if WITH_EDITORONLY_DATA
+	UPROPERTY(transient, duplicatetransient)
+	UMaterialInterface* DefaultIconMaterial;
+#endif // WITH_EDITORONLY_DATA
+
 	UPROPERTY(transient, duplicatetransient)
 	UTexture* ActiveImage;
 
