[GLTFExporter] refactor: remove support for non-uniform scale since it can't be properly support with animations and future transform variant changes

commit 1d651df29658f0a3ea7ce186bf59b7c430b9b8cc
Date: Tue, 29 Sep 2020 20:16:35 +0200

    refactor: remove support for non-uniform scale since it can't be properly support with animations and future transform variant changes
    (cherry picked from commit ec764f74a0be31a519bcd96dc4d8976a4b1553a6​)

#skipci

diff --git a/Engine/Plugins/Enterprise/GLTFExporter/MigratedFromGit/Source/GLTFExporter/Private/Converters/GLTFNodeConverters.cpp b/Engine/Plugins/Enterprise/GLTFExporter/MigratedFromGit/Source/GLTFExporter/Private/Converters/GLTFNodeConverters.cpp
--- a/Engine/Plugins/Enterprise/GLTFExporter/MigratedFromGit/Source/GLTFExporter/Private/Converters/GLTFNodeConverters.cpp
+++ b/Engine/Plugins/Enterprise/GLTFExporter/MigratedFromGit/Source/GLTFExporter/Private/Converters/GLTFNodeConverters.cpp
@@ -18,28 +18,34 @@ FGLTFJsonNodeIndex FGLTFComponentConverter::Convert(const USceneComponent* Scene
 
 	const bool bIsRootComponent = Owner->GetRootComponent() == SceneComponent;
 	const bool bIsRootNode = bIsRootComponent && FGLTFActorUtility::IsRootActor(Owner, Builder.bSelectedActorsOnly);
-	const bool bExportNonUniformScale = Builder.ExportOptions->bExportNonUniformScale;
+
+	// TODO: if node is root, then add it to the scene here to avoid any possible orphan nodes. this change requires node converter to support cyclic calls.
 
 	const USceneComponent* ParentComponent = SceneComponent->GetAttachParent();
 	const FGLTFJsonNodeIndex ParentNodeIndex = Builder.GetOrAddNode(ParentComponent);
 
-	// TODO: if node is root, then add it to the scene here to avoid any possible orphan nodes. this change requires level converter to support cyclic calls.
+	if (ParentComponent != nullptr && !SceneComponent->IsUsingAbsoluteScale())
+	{
+		const FVector ParentScale = ParentComponent->GetComponentScale();
+		if (ParentScale.X != ParentScale.Y || ParentScale.Y != ParentScale.Z)
+		{
+			Builder.AddWarningMessage(
+			FString::Printf(TEXT("Non-uniform parent scale (%s) for component %s (in actor %s) may be represented differently in glTF"),
+				*ParentScale.ToString(),
+				*SceneComponent->GetName(),
+				*Owner->GetName()));
+		}
+	}
 
 	const FTransform Transform = SceneComponent->GetComponentTransform();
 	const FTransform RelativeTransform = bIsRootNode ? Transform : Transform.GetRelativeTransform(ParentComponent->GetComponentTransform());
 
-	const FVector ParentScale = ParentComponent != nullptr ? ParentComponent->GetComponentScale() : FVector::OneVector;
-	const FVector Translation = bExportNonUniformScale ? RelativeTransform.GetTranslation() * ParentScale : RelativeTransform.GetTranslation();
-	const FVector Scale = bExportNonUniformScale ? FVector::OneVector : RelativeTransform.GetScale3D();
-
 	const FGLTFJsonNodeIndex NodeIndex = Builder.AddChildNode(ParentNodeIndex);
 	FGLTFJsonNode& Node = Builder.GetNode(NodeIndex);
 	Node.Name = FGLTFNameUtility::GetName(SceneComponent);
-	Node.Translation = FGLTFConverterUtility::ConvertPosition(Translation, Builder.ExportOptions->ExportScale);
+	Node.Translation = FGLTFConverterUtility::ConvertPosition(RelativeTransform.GetTranslation(), Builder.ExportOptions->ExportScale);
 	Node.Rotation = FGLTFConverterUtility::ConvertRotation(RelativeTransform.GetRotation());
-	Node.Scale = FGLTFConverterUtility::ConvertScale(Scale);
-
-	const FGLTFJsonVector3 ComponentNodeScale = FGLTFConverterUtility::ConvertScale(bExportNonUniformScale ? Transform.GetScale3D() : FVector::OneVector);
+	Node.Scale = FGLTFConverterUtility::ConvertScale(RelativeTransform.GetScale3D());
 
 	if (SceneComponent->bHiddenInGame) // TODO: make this configurable
 	{
@@ -47,60 +53,23 @@ FGLTFJsonNodeIndex FGLTFComponentConverter::Convert(const USceneComponent* Scene
 	}
 	else if (const UStaticMeshComponent* StaticMeshComponent = Cast<UStaticMeshComponent>(SceneComponent))
 	{
-		if (bExportNonUniformScale)
+		Node.Mesh = Builder.GetOrAddMesh(StaticMeshComponent);
+
+		if (Builder.ExportOptions->bExportLightmaps)
 		{
-			FGLTFJsonNode MeshNode;
-			MeshNode.Name = FGLTFNameUtility::GetName(StaticMeshComponent);
-			MeshNode.Scale = ComponentNodeScale;
-			MeshNode.Mesh = Builder.GetOrAddMesh(StaticMeshComponent);
-			if (Builder.ExportOptions->bExportLightmaps)
-			{
-				MeshNode.LightMap = Builder.GetOrAddLightMap(StaticMeshComponent);
-			}
-			Builder.AddChildComponentNode(NodeIndex, MeshNode);
-		}
-		else
-		{
-			Node.Mesh = Builder.GetOrAddMesh(StaticMeshComponent);
-			if (Builder.ExportOptions->bExportLightmaps)
-			{
-				Node.LightMap = Builder.GetOrAddLightMap(StaticMeshComponent);
-			}
+			Node.LightMap = Builder.GetOrAddLightMap(StaticMeshComponent);
 		}
 	}
 	else if (const USkeletalMeshComponent* SkeletalMeshComponent = Cast<USkeletalMeshComponent>(SceneComponent))
 	{
-		if (bExportNonUniformScale)
-		{
-			FGLTFJsonNode MeshNode;
-			MeshNode.Name = FGLTFNameUtility::GetName(SkeletalMeshComponent);
-			MeshNode.Scale = ComponentNodeScale;
-			MeshNode.Mesh = Builder.GetOrAddMesh(SkeletalMeshComponent);
-			Builder.AddChildComponentNode(NodeIndex, MeshNode);
-		}
-		else
-		{
-			Node.Mesh = Builder.GetOrAddMesh(SkeletalMeshComponent);
-		}
+		Node.Mesh = Builder.GetOrAddMesh(SkeletalMeshComponent);
 	}
 	else if (Owner->IsA<AGLTFInteractionHotspotActor>() && bIsRootComponent)
 	{
 		if (Builder.ExportOptions->bExportInteractionHotspots)
 		{
 			const AGLTFInteractionHotspotActor* HotspotActor = Cast<AGLTFInteractionHotspotActor>(Owner);
-
-			if (bExportNonUniformScale)
-			{
-				FGLTFJsonNode HotspotNode;
-				HotspotNode.Name = Owner->GetName();	// TODO: should we add a custom GetName-function in FGLTFNameUtility for this?
-				HotspotNode.Scale = ComponentNodeScale;
-				HotspotNode.Hotspot = Builder.GetOrAddHotspot(HotspotActor);
-				Builder.AddChildComponentNode(NodeIndex, HotspotNode);
-			}
-			else
-			{
-				Node.Hotspot = Builder.GetOrAddHotspot(HotspotActor);
-			}
+			Node.Hotspot = Builder.GetOrAddHotspot(HotspotActor);
 		}
 	}
 	else if (const UCameraComponent* CameraComponent = Cast<UCameraComponent>(SceneComponent))
@@ -111,7 +80,6 @@ FGLTFJsonNodeIndex FGLTFComponentConverter::Convert(const USceneComponent* Scene
 			FGLTFJsonNode CameraNode;
 			CameraNode.Name = FGLTFNameUtility::GetName(CameraComponent);
 			CameraNode.Rotation = FGLTFConverterUtility::ConvertCameraDirection();
-			CameraNode.Scale = ComponentNodeScale;
 			CameraNode.Camera = Builder.GetOrAddCamera(CameraComponent);
 			Builder.AddChildComponentNode(NodeIndex, CameraNode);
 		}
@@ -124,7 +92,6 @@ FGLTFJsonNodeIndex FGLTFComponentConverter::Convert(const USceneComponent* Scene
 			FGLTFJsonNode LightNode;
 			LightNode.Name = FGLTFNameUtility::GetName(LightComponent);
 			LightNode.Rotation = FGLTFConverterUtility::ConvertLightDirection();
-			LightNode.Scale = ComponentNodeScale;
 			LightNode.Light = Builder.GetOrAddLight(LightComponent);
 			Builder.AddChildComponentNode(NodeIndex, LightNode);
 		}
diff --git a/Engine/Plugins/Enterprise/GLTFExporter/MigratedFromGit/Source/GLTFExporter/Private/GLTFExportOptions.cpp b/Engine/Plugins/Enterprise/GLTFExporter/MigratedFromGit/Source/GLTFExporter/Private/GLTFExportOptions.cpp
--- a/Engine/Plugins/Enterprise/GLTFExporter/MigratedFromGit/Source/GLTFExporter/Private/GLTFExportOptions.cpp
+++ b/Engine/Plugins/Enterprise/GLTFExporter/MigratedFromGit/Source/GLTFExporter/Private/GLTFExportOptions.cpp
@@ -35,7 +35,6 @@ UGLTFExportOptions::UGLTFExportOptions(const FObjectInitializer& ObjectInitializ
 	bExportHDRIBackdrops = true;
 	bExportVariantSets = true;
 	bExportInteractionHotspots = true;
-	bExportNonUniformScale = false;
 	bBundleWebViewer = false;
 	bExportPreviewMesh = true;
 	bAllExtensionsRequired = false;
diff --git a/Engine/Plugins/Enterprise/GLTFExporter/MigratedFromGit/Source/GLTFExporter/Public/GLTFExportOptions.h b/Engine/Plugins/Enterprise/GLTFExporter/MigratedFromGit/Source/GLTFExporter/Public/GLTFExportOptions.h
--- a/Engine/Plugins/Enterprise/GLTFExporter/MigratedFromGit/Source/GLTFExporter/Public/GLTFExportOptions.h
+++ b/Engine/Plugins/Enterprise/GLTFExporter/MigratedFromGit/Source/GLTFExporter/Public/GLTFExportOptions.h
@@ -141,9 +141,6 @@ public:
 	UPROPERTY(EditAnywhere, BlueprintReadWrite, config, Category = Scene)
 	uint32 bExportInteractionHotspots : 1;
 
-	UPROPERTY(EditAnywhere, BlueprintReadWrite, config, Category = Scene)
-	uint32 bExportNonUniformScale : 1;
-
 	UPROPERTY(EditAnywhere, BlueprintReadWrite, config, Category = Exporter)
 	uint32 bBundleWebViewer : 1;
 
