[GLTFExporter] refactor: detect hotspot actors in actor converter instead of component converter

commit 5a27ee7109ab69fb26d02cb6847bb06e10c34bd2
Date: Tue, 29 Sep 2020 20:51:01 +0200

    refactor: detect hotspot actors in actor converter instead of component converter
    (cherry picked from commit 2c692923573f22c8a9748f2de10dd299493080d9​)

#skipci

diff --git a/Engine/Plugins/Enterprise/GLTFExporter/MigratedFromGit/Source/GLTFExporter/Private/Converters/GLTFNodeConverters.cpp b/Engine/Plugins/Enterprise/GLTFExporter/MigratedFromGit/Source/GLTFExporter/Private/Converters/GLTFNodeConverters.cpp
--- a/Engine/Plugins/Enterprise/GLTFExporter/MigratedFromGit/Source/GLTFExporter/Private/Converters/GLTFNodeConverters.cpp
+++ b/Engine/Plugins/Enterprise/GLTFExporter/MigratedFromGit/Source/GLTFExporter/Private/Converters/GLTFNodeConverters.cpp
@@ -64,14 +64,6 @@ FGLTFJsonNodeIndex FGLTFComponentConverter::Convert(const USceneComponent* Scene
 	{
 		Node.Mesh = Builder.GetOrAddMesh(SkeletalMeshComponent);
 	}
-	else if (Owner->IsA<AGLTFInteractionHotspotActor>() && bIsRootComponent)
-	{
-		if (Builder.ExportOptions->bExportInteractionHotspots)
-		{
-			const AGLTFInteractionHotspotActor* HotspotActor = Cast<AGLTFInteractionHotspotActor>(Owner);
-			Node.Hotspot = Builder.GetOrAddHotspot(HotspotActor);
-		}
-	}
 	else if (const UCameraComponent* CameraComponent = Cast<UCameraComponent>(SceneComponent))
 	{
 		if (Builder.ExportOptions->bExportCameras)
@@ -126,6 +118,14 @@ FGLTFJsonNodeIndex FGLTFActorConverter::Convert(const AActor* Actor)
 			RootNode.Backdrop = Builder.GetOrAddBackdrop(Actor);
 		}
 	}
+	else if (const AGLTFInteractionHotspotActor* HotspotActor = Cast<AGLTFInteractionHotspotActor>(Actor))
+	{
+		if (Builder.ExportOptions->bExportInteractionHotspots)
+		{
+			FGLTFJsonNode& RootNode = Builder.GetNode(RootNodeIndex);
+			RootNode.Hotspot = Builder.GetOrAddHotspot(HotspotActor);
+		}
+	}
 	else
 	{
 		// TODO: to reduce number of nodes, only export components that are of interest
