[GLTFExporter] feat(skeletal-mesh): enable export of vertex bone indices and weights

commit 7fd061d0c17a6bd509c59f16c330622adda071de
Date: Thu, 24 Sep 2020 00:21:41 +0200

    feat(skeletal-mesh): enable export of vertex bone indices and weights
    (cherry picked from commit e8b3a0fe5c31c708a791ebe185d3ac1652890b0e​)

#skipci

diff --git a/Engine/Plugins/Enterprise/GLTFExporter/MigratedFromGit/Source/GLTFExporter/Private/Converters/GLTFMeshConverters.cpp b/Engine/Plugins/Enterprise/GLTFExporter/MigratedFromGit/Source/GLTFExporter/Private/Converters/GLTFMeshConverters.cpp
--- a/Engine/Plugins/Enterprise/GLTFExporter/MigratedFromGit/Source/GLTFExporter/Private/Converters/GLTFMeshConverters.cpp
+++ b/Engine/Plugins/Enterprise/GLTFExporter/MigratedFromGit/Source/GLTFExporter/Private/Converters/GLTFMeshConverters.cpp
@@ -183,8 +183,6 @@ FGLTFJsonMeshIndex FGLTFSkeletalMeshConverter::Convert(const USkeletalMesh* Skel
 			JsonPrimitive.Attributes.TexCoords[UVIndex] = Builder.GetOrAddUVAccessor(ConvertedSection, VertexBuffer, UVIndex);
 		}
 
-		/* TODO: enable export of vertex skin data when crash bug fixed
-
 		const uint32 GroupCount = (SkinWeightBuffer->GetMaxBoneInfluences() + 3) / 4;
 		JsonPrimitive.Attributes.Joints.AddUninitialized(GroupCount);
 		JsonPrimitive.Attributes.Weights.AddUninitialized(GroupCount);
@@ -194,7 +192,6 @@ FGLTFJsonMeshIndex FGLTFSkeletalMeshConverter::Convert(const USkeletalMesh* Skel
 			JsonPrimitive.Attributes.Joints[GroupIndex] = Builder.GetOrAddJointAccessor(ConvertedSection, SkinWeightBuffer, GroupIndex * 4);
 			JsonPrimitive.Attributes.Weights[GroupIndex] = Builder.GetOrAddWeightAccessor(ConvertedSection, SkinWeightBuffer, GroupIndex * 4);
 		}
-		*/
 
 		const UMaterialInterface* Material = OverrideMaterials.GetOverride(SkeletalMesh->Materials, Section->MaterialIndex);
 		if (Material == nullptr)
