[GLTFExporter] feat(skeletal-mesh): add skeletal skin, bone, and socket converters

commit ff9ae6035bfca8533190b9566d12cf12a82d7920
Date: Thu, 24 Sep 2020 00:21:47 +0200

    feat(skeletal-mesh): add skeletal skin, bone, and socket converters
    (cherry picked from commit 7ef9f63cc0e26c19729bdfcdcc546d906ee0a779​)

#skipci

diff --git a/Engine/Plugins/Enterprise/GLTFExporter/MigratedFromGit/Source/GLTFExporter/Private/Builders/GLTFConvertBuilder.cpp b/Engine/Plugins/Enterprise/GLTFExporter/MigratedFromGit/Source/GLTFExporter/Private/Builders/GLTFConvertBuilder.cpp
--- a/Engine/Plugins/Enterprise/GLTFExporter/MigratedFromGit/Source/GLTFExporter/Private/Builders/GLTFConvertBuilder.cpp
+++ b/Engine/Plugins/Enterprise/GLTFExporter/MigratedFromGit/Source/GLTFExporter/Private/Builders/GLTFConvertBuilder.cpp
@@ -199,6 +199,36 @@ FGLTFJsonTextureIndex FGLTFConvertBuilder::GetOrAddTexture(const UTextureRenderT
 	return TextureRenderTargetCubeConverter.GetOrAdd(Texture, CubeFace);
 }
 
+FGLTFJsonNodeIndex FGLTFConvertBuilder::GetOrAddNode(FGLTFJsonNodeIndex RootNode, const USkeletalMesh* SkeletalMesh, FName SocketName)
+{
+	if (RootNode == INDEX_NONE || SkeletalMesh == nullptr || SocketName == NAME_None)
+	{
+		return FGLTFJsonNodeIndex(INDEX_NONE);
+	}
+
+	return SkeletalSocketConverter.GetOrAdd(RootNode, SkeletalMesh, SocketName);
+}
+
+FGLTFJsonNodeIndex FGLTFConvertBuilder::GetOrAddNode(FGLTFJsonNodeIndex RootNode, const USkeletalMesh* SkeletalMesh, int32 BoneIndex)
+{
+	if (RootNode == INDEX_NONE || SkeletalMesh == nullptr || BoneIndex == INDEX_NONE)
+	{
+		return FGLTFJsonNodeIndex(INDEX_NONE);
+	}
+
+	return SkeletalBoneConverter.GetOrAdd(RootNode, SkeletalMesh, BoneIndex);
+}
+
+FGLTFJsonSkinIndex FGLTFConvertBuilder::GetOrAddSkin(FGLTFJsonNodeIndex RootNode, const USkeletalMesh* SkeletalMesh)
+{
+	if (RootNode == INDEX_NONE || SkeletalMesh == nullptr)
+	{
+		return FGLTFJsonSkinIndex(INDEX_NONE);
+	}
+
+	return SkinConverter.GetOrAdd(RootNode, SkeletalMesh);
+}
+
 FGLTFJsonNodeIndex FGLTFConvertBuilder::GetOrAddNode(const USceneComponent* SceneComponent)
 {
 	if (SceneComponent == nullptr)
diff --git a/Engine/Plugins/Enterprise/GLTFExporter/MigratedFromGit/Source/GLTFExporter/Private/Builders/GLTFConvertBuilder.h b/Engine/Plugins/Enterprise/GLTFExporter/MigratedFromGit/Source/GLTFExporter/Private/Builders/GLTFConvertBuilder.h
--- a/Engine/Plugins/Enterprise/GLTFExporter/MigratedFromGit/Source/GLTFExporter/Private/Builders/GLTFConvertBuilder.h
+++ b/Engine/Plugins/Enterprise/GLTFExporter/MigratedFromGit/Source/GLTFExporter/Private/Builders/GLTFConvertBuilder.h
@@ -10,6 +10,7 @@
 #include "Converters/GLTFSamplerConverters.h"
 #include "Converters/GLTFTextureConverters.h"
 #include "Converters/GLTFNodeConverters.h"
+#include "Converters/GLTFSkinConverters.h"
 #include "Converters/GLTFSceneConverters.h"
 #include "Converters/GLTFCameraConverters.h"
 #include "Converters/GLTFLightConverters.h"
@@ -50,6 +51,10 @@ public:
 	FGLTFJsonTextureIndex GetOrAddTexture(const UTextureRenderTarget2D* Texture);
 	FGLTFJsonTextureIndex GetOrAddTexture(const UTextureRenderTargetCube* Texture, ECubeFace CubeFace);
 
+	FGLTFJsonNodeIndex GetOrAddNode(FGLTFJsonNodeIndex RootNode, const USkeletalMesh* SkeletalMesh, FName SocketName);
+	FGLTFJsonNodeIndex GetOrAddNode(FGLTFJsonNodeIndex RootNode, const USkeletalMesh* SkeletalMesh, int32 BoneIndex);
+	FGLTFJsonSkinIndex GetOrAddSkin(FGLTFJsonNodeIndex RootNode, const USkeletalMesh* SkeletalMesh);
+
 	FGLTFJsonNodeIndex GetOrAddNode(const USceneComponent* SceneComponent);
 	FGLTFJsonNodeIndex GetOrAddNode(const AActor* Actor);
 	FGLTFJsonSceneIndex GetOrAddScene(const ULevel* Level);
@@ -83,6 +88,10 @@ private:
 	FGLTFTextureRenderTarget2DConverter TextureRenderTarget2DConverter{ *this };
 	FGLTFTextureRenderTargetCubeConverter TextureRenderTargetCubeConverter{ *this };
 
+	FGLTFSkeletalSocketConverter SkeletalSocketConverter{ *this };
+	FGLTFSkeletalBoneConverter SkeletalBoneConverter{ *this };
+	FGLTFSkinConverter SkinConverter{ *this };
+
 	FGLTFComponentConverter ComponentConverter{ *this };
 	FGLTFActorConverter ActorConverter{ *this };
 	FGLTFSceneConverter SceneConverter{ *this };
diff --git a/Engine/Plugins/Enterprise/GLTFExporter/MigratedFromGit/Source/GLTFExporter/Private/Builders/GLTFJsonBuilder.cpp b/Engine/Plugins/Enterprise/GLTFExporter/MigratedFromGit/Source/GLTFExporter/Private/Builders/GLTFJsonBuilder.cpp
--- a/Engine/Plugins/Enterprise/GLTFExporter/MigratedFromGit/Source/GLTFExporter/Private/Builders/GLTFJsonBuilder.cpp
+++ b/Engine/Plugins/Enterprise/GLTFExporter/MigratedFromGit/Source/GLTFExporter/Private/Builders/GLTFJsonBuilder.cpp
@@ -58,6 +58,11 @@ FGLTFJsonSceneIndex FGLTFJsonBuilder::AddScene(const FGLTFJsonScene& JsonScene)
 	return FGLTFJsonSceneIndex(JsonRoot.Scenes.Add(JsonScene));
 }
 
+FGLTFJsonSkinIndex FGLTFJsonBuilder::AddSkin(const FGLTFJsonSkin& JsonSkin)
+{
+	return FGLTFJsonSkinIndex(JsonRoot.Skins.Add(JsonSkin));
+}
+
 FGLTFJsonTextureIndex FGLTFJsonBuilder::AddTexture(const FGLTFJsonTexture& JsonTexture)
 {
 	return FGLTFJsonTextureIndex(JsonRoot.Textures.Add(JsonTexture));
@@ -162,6 +167,11 @@ FGLTFJsonScene& FGLTFJsonBuilder::GetScene(FGLTFJsonSceneIndex SceneIndex)
 	return JsonRoot.Scenes[SceneIndex];
 }
 
+FGLTFJsonSkin& FGLTFJsonBuilder::GetSkin(FGLTFJsonSkinIndex SkinIndex)
+{
+	return JsonRoot.Skins[SkinIndex];
+}
+
 FGLTFJsonTexture& FGLTFJsonBuilder::GetTexture(FGLTFJsonTextureIndex TextureIndex)
 {
 	return JsonRoot.Textures[TextureIndex];
diff --git a/Engine/Plugins/Enterprise/GLTFExporter/MigratedFromGit/Source/GLTFExporter/Private/Builders/GLTFJsonBuilder.h b/Engine/Plugins/Enterprise/GLTFExporter/MigratedFromGit/Source/GLTFExporter/Private/Builders/GLTFJsonBuilder.h
--- a/Engine/Plugins/Enterprise/GLTFExporter/MigratedFromGit/Source/GLTFExporter/Private/Builders/GLTFJsonBuilder.h
+++ b/Engine/Plugins/Enterprise/GLTFExporter/MigratedFromGit/Source/GLTFExporter/Private/Builders/GLTFJsonBuilder.h
@@ -25,6 +25,7 @@ public:
 	FGLTFJsonNodeIndex AddNode(const FGLTFJsonNode& JsonNode = {});
 	FGLTFJsonSamplerIndex AddSampler(const FGLTFJsonSampler& JsonSampler = {});
 	FGLTFJsonSceneIndex AddScene(const FGLTFJsonScene& JsonScene = {});
+	FGLTFJsonSkinIndex AddSkin(const FGLTFJsonSkin& JsonSkin = {});
 	FGLTFJsonTextureIndex AddTexture(const FGLTFJsonTexture& JsonTexture = {});
 	FGLTFJsonBackdropIndex AddBackdrop(const FGLTFJsonBackdrop& JsonBackdrop = {});
 	FGLTFJsonVariationIndex AddVariation(const FGLTFJsonVariation& JsonVariation = {});
@@ -45,6 +46,7 @@ public:
 	FGLTFJsonNode& GetNode(FGLTFJsonNodeIndex NodeIndex);
 	FGLTFJsonSampler& GetSampler(FGLTFJsonSamplerIndex SamplerIndex);
 	FGLTFJsonScene& GetScene(FGLTFJsonSceneIndex SceneIndex);
+	FGLTFJsonSkin& GetSkin(FGLTFJsonSkinIndex SkinIndex);
 	FGLTFJsonTexture& GetTexture(FGLTFJsonTextureIndex TextureIndex);
 	FGLTFJsonBackdrop& GetBackdrop(FGLTFJsonBackdropIndex BackdropIndex);
 	FGLTFJsonVariation& GetVariation(FGLTFJsonVariationIndex VariationIndex);
diff --git a/Engine/Plugins/Enterprise/GLTFExporter/MigratedFromGit/Source/GLTFExporter/Private/Converters/GLTFNodeConverters.cpp b/Engine/Plugins/Enterprise/GLTFExporter/MigratedFromGit/Source/GLTFExporter/Private/Converters/GLTFNodeConverters.cpp
--- a/Engine/Plugins/Enterprise/GLTFExporter/MigratedFromGit/Source/GLTFExporter/Private/Converters/GLTFNodeConverters.cpp
+++ b/Engine/Plugins/Enterprise/GLTFExporter/MigratedFromGit/Source/GLTFExporter/Private/Converters/GLTFNodeConverters.cpp
@@ -142,3 +142,66 @@ FGLTFJsonNodeIndex FGLTFActorConverter::Convert(const AActor* Actor)
 
 	return RootNodeIndex;
 }
+
+FGLTFJsonNodeIndex FGLTFSkeletalSocketConverter::Convert(FGLTFJsonNodeIndex RootNode, const USkeletalMesh* SkeletalMesh, FName SocketName)
+{
+	const USkeletalMeshSocket* Socket = SkeletalMesh->FindSocket(SocketName);
+	if (Socket != nullptr)
+	{
+		FGLTFJsonNode Node;
+		Node.Name = SocketName.ToString();
+
+		// TODO: add support for non-uniform scaling
+		const FTransform Transform = Socket->GetSocketLocalTransform();
+		Node.Translation = FGLTFConverterUtility::ConvertPosition(Transform.GetTranslation(), Builder.ExportOptions->ExportScale);
+		Node.Rotation = FGLTFConverterUtility::ConvertRotation(Transform.GetRotation());
+		Node.Scale = FGLTFConverterUtility::ConvertScale(Transform.GetScale3D());
+
+		const int32 ParentBone = SkeletalMesh->RefSkeleton.FindBoneIndex(Socket->BoneName);
+		const FGLTFJsonNodeIndex ParentNode = ParentBone != INDEX_NONE ? Builder.GetOrAddNode(RootNode, SkeletalMesh, ParentBone) : RootNode;
+		return Builder.AddChildNode(ParentNode, Node);
+	}
+
+	const int32 BoneIndex = SkeletalMesh->RefSkeleton.FindBoneIndex(SocketName);
+	if (BoneIndex != INDEX_NONE)
+	{
+		return Builder.GetOrAddNode(RootNode, SkeletalMesh, BoneIndex);
+	}
+
+	// TODO: report error
+	return FGLTFJsonNodeIndex(INDEX_NONE);
+}
+
+FGLTFJsonNodeIndex FGLTFSkeletalBoneConverter::Convert(FGLTFJsonNodeIndex RootNode, const USkeletalMesh* SkeletalMesh, int32 BoneIndex)
+{
+	// TODO: add support for MasterPoseComponent?
+
+	const TArray<FMeshBoneInfo>& BoneInfos = SkeletalMesh->RefSkeleton.GetRefBoneInfo();
+	if (!BoneInfos.IsValidIndex(BoneIndex))
+	{
+		// TODO: report error
+		return FGLTFJsonNodeIndex(INDEX_NONE);
+	}
+
+	const FMeshBoneInfo& BoneInfo = BoneInfos[BoneIndex];
+
+	FGLTFJsonNode Node;
+	Node.Name = BoneInfo.Name.ToString();
+
+	const TArray<FTransform>& BonePoses = SkeletalMesh->RefSkeleton.GetRefBonePose();
+	if (BonePoses.IsValidIndex(BoneIndex))
+	{
+		// TODO: add support for non-uniform scaling
+		const FTransform& BonePose = BonePoses[BoneIndex];
+		Node.Translation = FGLTFConverterUtility::ConvertPosition(BonePose.GetTranslation(), Builder.ExportOptions->ExportScale);
+		Node.Rotation = FGLTFConverterUtility::ConvertRotation(BonePose.GetRotation());
+		Node.Scale = FGLTFConverterUtility::ConvertScale(BonePose.GetScale3D());
+	}
+	else
+	{
+		// TODO: report error
+	}
+
+	const FGLTFJsonNodeIndex ParentNode = BoneInfo.ParentIndex != INDEX_NONE ? Builder.GetOrAddNode(RootNode, SkeletalMesh, BoneInfo.ParentIndex) : RootNode;
+	return Builder.AddChildNode(ParentNode, Node);
+}
diff --git a/Engine/Plugins/Enterprise/GLTFExporter/MigratedFromGit/Source/GLTFExporter/Private/Converters/GLTFNodeConverters.h b/Engine/Plugins/Enterprise/GLTFExporter/MigratedFromGit/Source/GLTFExporter/Private/Converters/GLTFNodeConverters.h
--- a/Engine/Plugins/Enterprise/GLTFExporter/MigratedFromGit/Source/GLTFExporter/Private/Converters/GLTFNodeConverters.h
+++ b/Engine/Plugins/Enterprise/GLTFExporter/MigratedFromGit/Source/GLTFExporter/Private/Converters/GLTFNodeConverters.h
@@ -26,3 +26,17 @@ class FGLTFActorConverter final : public TGLTFNodeConverter<const AActor*>
 
 	virtual FGLTFJsonNodeIndex Convert(const AActor* Actor) override;
 };
+
+class FGLTFSkeletalSocketConverter final : public TGLTFNodeConverter<FGLTFJsonNodeIndex, const USkeletalMesh*, FName>
+{
+	using TGLTFNodeConverter::TGLTFNodeConverter;
+
+	virtual FGLTFJsonNodeIndex Convert(FGLTFJsonNodeIndex RootNode, const USkeletalMesh* SkeletalMesh, FName SocketName) override;
+};
+
+class FGLTFSkeletalBoneConverter final : public TGLTFNodeConverter<FGLTFJsonNodeIndex, const USkeletalMesh*, int32>
+{
+	using TGLTFNodeConverter::TGLTFNodeConverter;
+
+	virtual FGLTFJsonNodeIndex Convert(FGLTFJsonNodeIndex RootNode, const USkeletalMesh* SkeletalMesh, int32 BoneIndex) override;
+};
diff --git a/Engine/Plugins/Enterprise/GLTFExporter/MigratedFromGit/Source/GLTFExporter/Private/Converters/GLTFSkinConverters.cpp b/Engine/Plugins/Enterprise/GLTFExporter/MigratedFromGit/Source/GLTFExporter/Private/Converters/GLTFSkinConverters.cpp
new file mode 100644
--- /dev/null
+++ b/Engine/Plugins/Enterprise/GLTFExporter/MigratedFromGit/Source/GLTFExporter/Private/Converters/GLTFSkinConverters.cpp
@@ -0,0 +1,22 @@
+// Copyright Epic Games, Inc. All Rights Reserved.
+
+#include "Converters/GLTFSkinConverters.h"
+#include "Converters/GLTFConverterUtility.h"
+#include "Builders/GLTFConvertBuilder.h"
+
+FGLTFJsonSkinIndex FGLTFSkinConverter::Convert(FGLTFJsonNodeIndex RootNode, const USkeletalMesh* SkeletalMesh)
+{
+	FGLTFJsonSkin Skin;
+	Skin.Name = SkeletalMesh->Skeleton != nullptr ? SkeletalMesh->Skeleton->GetName() : SkeletalMesh->GetName();
+	Skin.Skeleton = RootNode;
+
+	const int32 BoneCount = SkeletalMesh->RefSkeleton.GetNum();
+	Skin.Joints.AddUninitialized(BoneCount);
+
+	for (int32 BoneIndex = 0; BoneIndex < BoneCount; ++BoneIndex)
+	{
+		Skin.Joints[BoneIndex] = Builder.GetOrAddNode(RootNode, SkeletalMesh, BoneIndex);
+	}
+
+	return Builder.AddSkin(Skin);
+}
diff --git a/Engine/Plugins/Enterprise/GLTFExporter/MigratedFromGit/Source/GLTFExporter/Private/Converters/GLTFSkinConverters.h b/Engine/Plugins/Enterprise/GLTFExporter/MigratedFromGit/Source/GLTFExporter/Private/Converters/GLTFSkinConverters.h
new file mode 100644
--- /dev/null
+++ b/Engine/Plugins/Enterprise/GLTFExporter/MigratedFromGit/Source/GLTFExporter/Private/Converters/GLTFSkinConverters.h
@@ -0,0 +1,15 @@
+// Copyright Epic Games, Inc. All Rights Reserved.
+
+#pragma once
+
+#include "Json/GLTFJsonIndex.h"
+#include "Converters/GLTFConverter.h"
+#include "Converters/GLTFBuilderContext.h"
+#include "Engine.h"
+
+class FGLTFSkinConverter final : public FGLTFBuilderContext, public TGLTFConverter<FGLTFJsonSkinIndex, FGLTFJsonNodeIndex, const USkeletalMesh*>
+{
+	using FGLTFBuilderContext::FGLTFBuilderContext;
+
+	virtual FGLTFJsonSkinIndex Convert(FGLTFJsonNodeIndex RootNode, const USkeletalMesh* SkeletalMesh) override;
+};
diff --git a/Engine/Plugins/Enterprise/GLTFExporter/MigratedFromGit/Source/GLTFExporter/Private/Json/GLTFJsonRoot.h b/Engine/Plugins/Enterprise/GLTFExporter/MigratedFromGit/Source/GLTFExporter/Private/Json/GLTFJsonRoot.h
--- a/Engine/Plugins/Enterprise/GLTFExporter/MigratedFromGit/Source/GLTFExporter/Private/Json/GLTFJsonRoot.h
+++ b/Engine/Plugins/Enterprise/GLTFExporter/MigratedFromGit/Source/GLTFExporter/Private/Json/GLTFJsonRoot.h
@@ -14,6 +14,7 @@
 #include "Json/GLTFJsonNode.h"
 #include "Json/GLTFJsonSampler.h"
 #include "Json/GLTFJsonScene.h"
+#include "Json/GLTFJsonSkin.h"
 #include "Json/GLTFJsonTexture.h"
 #include "Json/GLTFJsonBackdrop.h"
 #include "Json/GLTFJsonVariation.h"
@@ -70,6 +71,7 @@ struct FGLTFJsonRoot
 	TArray<FGLTFJsonImage>      Images;
 	TArray<FGLTFJsonSampler>    Samplers;
 	TArray<FGLTFJsonScene>      Scenes;
+	TArray<FGLTFJsonSkin>       Skins;
 	TArray<FGLTFJsonTexture>    Textures;
 	TArray<FGLTFJsonBackdrop>   Backdrops;
 	TArray<FGLTFJsonLightMap>   LightMaps;
@@ -102,6 +104,7 @@ struct FGLTFJsonRoot
 		FGLTFJsonUtility::WriteObjectArray(JsonWriter, TEXT("nodes"), Nodes, AllExtensions);
 		FGLTFJsonUtility::WriteObjectArray(JsonWriter, TEXT("samplers"), Samplers, AllExtensions);
 		FGLTFJsonUtility::WriteObjectArray(JsonWriter, TEXT("scenes"), Scenes, AllExtensions);
+		FGLTFJsonUtility::WriteObjectArray(JsonWriter, TEXT("skins"), Skins, AllExtensions);
 		FGLTFJsonUtility::WriteObjectArray(JsonWriter, TEXT("textures"), Textures, AllExtensions);
 
 		if (Backdrops.Num() > 0 || Variations.Num() > 0 || LightMaps.Num() > 0 || Lights.Num() > 0 || Hotspots.Num() > 0 /* TODO: add more extension support */)
diff --git a/Engine/Plugins/Enterprise/GLTFExporter/MigratedFromGit/Source/GLTFExporter/Private/Json/GLTFJsonSkin.h b/Engine/Plugins/Enterprise/GLTFExporter/MigratedFromGit/Source/GLTFExporter/Private/Json/GLTFJsonSkin.h
new file mode 100644
--- /dev/null
+++ b/Engine/Plugins/Enterprise/GLTFExporter/MigratedFromGit/Source/GLTFExporter/Private/Json/GLTFJsonSkin.h
@@ -0,0 +1,44 @@
+// Copyright Epic Games, Inc. All Rights Reserved.
+
+#pragma once
+
+#include "Json/GLTFJsonIndex.h"
+#include "Serialization/JsonSerializer.h"
+
+struct FGLTFJsonSkin
+{
+	FString Name;
+
+	FGLTFJsonAccessorIndex InverseBindMatrices;
+	FGLTFJsonNodeIndex Skeleton;
+
+	TArray<FGLTFJsonNodeIndex> Joints;
+
+	template <class CharType = TCHAR, class PrintPolicy = TPrettyJsonPrintPolicy<CharType>>
+	void WriteObject(TJsonWriter<CharType, PrintPolicy>& JsonWriter, FGLTFJsonExtensions& Extensions) const
+	{
+		JsonWriter.WriteObjectStart();
+
+		if (!Name.IsEmpty())
+		{
+			JsonWriter.WriteValue(TEXT("name"), Name);
+		}
+
+		if (InverseBindMatrices != INDEX_NONE)
+		{
+			JsonWriter.WriteValue(TEXT("inverseBindMatrices"), InverseBindMatrices);
+		}
+
+		if (Skeleton != INDEX_NONE)
+		{
+			JsonWriter.WriteValue(TEXT("skeleton"), Skeleton);
+		}
+
+		if (Joints.Num() > 0)
+		{
+			JsonWriter.WriteValue(TEXT("joints"), Joints);
+		}
+
+		JsonWriter.WriteObjectEnd();
+	}
+};
