[GLTFExporter] feat(skeletal-mesh): include skeleton bone tree when exporting skeletal-mesh

commit e9ca19c4be1090d49f9165bdba98079e72059d8a
Date: Mon, 1 Jun 2020 01:35:10 +0200

    feat(skeletal-mesh): include skeleton bone tree when exporting skeletal-mesh
    (cherry picked from commit 6f28ba2d67b62a311210bc63880eca36fb775f99​)

#skipci

diff --git a/Engine/Plugins/Enterprise/GLTFExporter/MigratedFromGit/Source/GLTFExporter/Private/Converters/GLTFNodeConverters.cpp b/Engine/Plugins/Enterprise/GLTFExporter/MigratedFromGit/Source/GLTFExporter/Private/Converters/GLTFNodeConverters.cpp
--- a/Engine/Plugins/Enterprise/GLTFExporter/MigratedFromGit/Source/GLTFExporter/Private/Converters/GLTFNodeConverters.cpp
+++ b/Engine/Plugins/Enterprise/GLTFExporter/MigratedFromGit/Source/GLTFExporter/Private/Converters/GLTFNodeConverters.cpp
@@ -63,6 +63,13 @@ FGLTFJsonNodeIndex FGLTFComponentConverter::Convert(const USceneComponent* Scene
 	else if (const USkeletalMeshComponent* SkeletalMeshComponent = Cast<USkeletalMeshComponent>(SceneComponent))
 	{
 		Node.Mesh = Builder.GetOrAddMesh(SkeletalMeshComponent);
+
+		const USkeletalMesh* SkeletalMesh = SkeletalMeshComponent->SkeletalMesh;
+		if (SkeletalMesh != nullptr)
+		{
+			const FGLTFJsonSkinIndex SkinIndex = Builder.GetOrAddSkin(NodeIndex, SkeletalMesh);
+			Builder.GetNode(NodeIndex).Skin = SkinIndex; // GetOrAddSkin may have invalidated Node reference
+		}
 	}
 	else if (const UCameraComponent* CameraComponent = Cast<UCameraComponent>(SceneComponent))
 	{
diff --git a/Engine/Plugins/Enterprise/GLTFExporter/MigratedFromGit/Source/GLTFExporter/Private/Exporters/GLTFSkeletalMeshExporter.cpp b/Engine/Plugins/Enterprise/GLTFExporter/MigratedFromGit/Source/GLTFExporter/Private/Exporters/GLTFSkeletalMeshExporter.cpp
--- a/Engine/Plugins/Enterprise/GLTFExporter/MigratedFromGit/Source/GLTFExporter/Private/Exporters/GLTFSkeletalMeshExporter.cpp
+++ b/Engine/Plugins/Enterprise/GLTFExporter/MigratedFromGit/Source/GLTFExporter/Private/Exporters/GLTFSkeletalMeshExporter.cpp
@@ -18,6 +18,7 @@ bool UGLTFSkeletalMeshExporter::AddObject(FGLTFContainerBuilder& Builder, const
 	FGLTFJsonNode Node;
 	Node.Mesh = MeshIndex;
 	const FGLTFJsonNodeIndex NodeIndex = Builder.AddNode(Node);
+	Builder.GetNode(NodeIndex).Skin = Builder.GetOrAddSkin(NodeIndex, SkeletalMesh);
 
 	FGLTFJsonScene Scene;
 	Scene.Nodes.Add(NodeIndex);
