[GLTFExporter] feat(skeletal-mesh): add GLTFSceneComponentSocketConverter

commit 8bff197ae01df40cb3d1880eff2502d63bf9e33c
Date: Mon, 1 Jun 2020 23:00:38 +0200

    feat(skeletal-mesh): add GLTFSceneComponentSocketConverter
    (cherry picked from commit 016b3abaaf14f180307170d84a89ca52973f5f0e​)

#skipci

diff --git a/Engine/Plugins/Enterprise/GLTFExporter/MigratedFromGit/Source/GLTFExporter/Private/Builders/GLTFConvertBuilder.cpp b/Engine/Plugins/Enterprise/GLTFExporter/MigratedFromGit/Source/GLTFExporter/Private/Builders/GLTFConvertBuilder.cpp
--- a/Engine/Plugins/Enterprise/GLTFExporter/MigratedFromGit/Source/GLTFExporter/Private/Builders/GLTFConvertBuilder.cpp
+++ b/Engine/Plugins/Enterprise/GLTFExporter/MigratedFromGit/Source/GLTFExporter/Private/Builders/GLTFConvertBuilder.cpp
@@ -229,6 +229,16 @@ FGLTFJsonSkinIndex FGLTFConvertBuilder::GetOrAddSkin(FGLTFJsonNodeIndex RootNode
 	return SkinConverter.GetOrAdd(RootNode, SkeletalMesh);
 }
 
+FGLTFJsonNodeIndex FGLTFConvertBuilder::GetOrAddNode(const USceneComponent* SceneComponent, FName SocketName)
+{
+	if (SceneComponent == nullptr || SocketName == NAME_None)
+	{
+		return FGLTFJsonNodeIndex(INDEX_NONE);
+	}
+
+	return ComponentSocketConverter.GetOrAdd(SceneComponent, SocketName);
+}
+
 FGLTFJsonNodeIndex FGLTFConvertBuilder::GetOrAddNode(const USceneComponent* SceneComponent)
 {
 	if (SceneComponent == nullptr)
diff --git a/Engine/Plugins/Enterprise/GLTFExporter/MigratedFromGit/Source/GLTFExporter/Private/Builders/GLTFConvertBuilder.h b/Engine/Plugins/Enterprise/GLTFExporter/MigratedFromGit/Source/GLTFExporter/Private/Builders/GLTFConvertBuilder.h
--- a/Engine/Plugins/Enterprise/GLTFExporter/MigratedFromGit/Source/GLTFExporter/Private/Builders/GLTFConvertBuilder.h
+++ b/Engine/Plugins/Enterprise/GLTFExporter/MigratedFromGit/Source/GLTFExporter/Private/Builders/GLTFConvertBuilder.h
@@ -55,6 +55,7 @@ public:
 	FGLTFJsonNodeIndex GetOrAddNode(FGLTFJsonNodeIndex RootNode, const USkeletalMesh* SkeletalMesh, int32 BoneIndex);
 	FGLTFJsonSkinIndex GetOrAddSkin(FGLTFJsonNodeIndex RootNode, const USkeletalMesh* SkeletalMesh);
 
+	FGLTFJsonNodeIndex GetOrAddNode(const USceneComponent* SceneComponent, FName SocketName);
 	FGLTFJsonNodeIndex GetOrAddNode(const USceneComponent* SceneComponent);
 	FGLTFJsonNodeIndex GetOrAddNode(const AActor* Actor);
 	FGLTFJsonSceneIndex GetOrAddScene(const ULevel* Level);
@@ -92,6 +93,7 @@ private:
 	FGLTFSkeletalBoneConverter SkeletalBoneConverter{ *this };
 	FGLTFSkinConverter SkinConverter{ *this };
 
+	FGLTFComponentSocketConverter ComponentSocketConverter{ *this };
 	FGLTFComponentConverter ComponentConverter{ *this };
 	FGLTFActorConverter ActorConverter{ *this };
 	FGLTFSceneConverter SceneConverter{ *this };
diff --git a/Engine/Plugins/Enterprise/GLTFExporter/MigratedFromGit/Source/GLTFExporter/Private/Converters/GLTFNodeConverters.cpp b/Engine/Plugins/Enterprise/GLTFExporter/MigratedFromGit/Source/GLTFExporter/Private/Converters/GLTFNodeConverters.cpp
--- a/Engine/Plugins/Enterprise/GLTFExporter/MigratedFromGit/Source/GLTFExporter/Private/Converters/GLTFNodeConverters.cpp
+++ b/Engine/Plugins/Enterprise/GLTFExporter/MigratedFromGit/Source/GLTFExporter/Private/Converters/GLTFNodeConverters.cpp
@@ -7,6 +7,22 @@
 #include "Converters/GLTFNameUtility.h"
 #include "Actors/GLTFInteractionHotspotActor.h"
 
+FGLTFJsonNodeIndex FGLTFComponentSocketConverter::Convert(const USceneComponent* SceneComponent, FName SocketName)
+{
+	const FGLTFJsonNodeIndex NodeIndex = Builder.GetOrAddNode(SceneComponent);
+
+	if (const USkinnedMeshComponent* SkinnedMeshComponent = Cast<USkinnedMeshComponent>(SceneComponent))
+	{
+		// TODO: add support for SocketOverrideLookup?
+		const USkeletalMesh* SkeletalMesh = SkinnedMeshComponent->SkeletalMesh;
+		return Builder.GetOrAddNode(NodeIndex, SkeletalMesh, SocketName);
+	}
+
+	// TODO: add support for non-skeletalmesh sockets
+
+	return FGLTFJsonNodeIndex(INDEX_NONE);
+}
+
 FGLTFJsonNodeIndex FGLTFComponentConverter::Convert(const USceneComponent* SceneComponent)
 {
 	const AActor* Owner = SceneComponent->GetOwner();
diff --git a/Engine/Plugins/Enterprise/GLTFExporter/MigratedFromGit/Source/GLTFExporter/Private/Converters/GLTFNodeConverters.h b/Engine/Plugins/Enterprise/GLTFExporter/MigratedFromGit/Source/GLTFExporter/Private/Converters/GLTFNodeConverters.h
--- a/Engine/Plugins/Enterprise/GLTFExporter/MigratedFromGit/Source/GLTFExporter/Private/Converters/GLTFNodeConverters.h
+++ b/Engine/Plugins/Enterprise/GLTFExporter/MigratedFromGit/Source/GLTFExporter/Private/Converters/GLTFNodeConverters.h
@@ -13,6 +13,13 @@ class TGLTFNodeConverter : public FGLTFBuilderContext, public TGLTFConverter<FGL
 	using FGLTFBuilderContext::FGLTFBuilderContext;
 };
 
+class FGLTFComponentSocketConverter final : public TGLTFNodeConverter<const USceneComponent*, FName>
+{
+	using TGLTFNodeConverter::TGLTFNodeConverter;
+
+	virtual FGLTFJsonNodeIndex Convert(const USceneComponent* SceneComponent, FName SocketName) override;
+};
+
 class FGLTFComponentConverter final : public TGLTFNodeConverter<const USceneComponent*>
 {
 	using TGLTFNodeConverter::TGLTFNodeConverter;
