[GLTFExporter] fix(skeletal-mesh): calculate and export InverseBindMatrices to fix bone transform issue with skeletal-mesh export

commit 724b5bd6186b1276a1ea21ed950fac15a56fa465
Date: Wed, 10 Jun 2020 20:49:14 +0200

    fix(skeletal-mesh): calculate and export InverseBindMatrices to fix bone transform issue with skeletal-mesh export
    (cherry picked from commit dd36fcb178c084ea6098ee2e7023335eb13f8619​)

#skipci

diff --git a/Engine/Plugins/Enterprise/GLTFExporter/MigratedFromGit/Source/GLTFExporter/Private/Converters/GLTFBoneUtility.cpp b/Engine/Plugins/Enterprise/GLTFExporter/MigratedFromGit/Source/GLTFExporter/Private/Converters/GLTFBoneUtility.cpp
new file mode 100644
--- /dev/null
+++ b/Engine/Plugins/Enterprise/GLTFExporter/MigratedFromGit/Source/GLTFExporter/Private/Converters/GLTFBoneUtility.cpp
@@ -0,0 +1,21 @@
+// Copyright Epic Games, Inc. All Rights Reserved.
+
+#include "Converters/GLTFBoneUtility.h"
+#include "Engine.h"
+
+FTransform FGLTFBoneUtility::GetBindTransform(const FReferenceSkeleton& RefSkeleton, const int32 BoneIndex)
+{
+	const TArray<FMeshBoneInfo>& BoneInfos = RefSkeleton.GetRefBoneInfo();
+	const TArray<FTransform>& BonePoses = RefSkeleton.GetRefBonePose();
+
+	int32 CurBoneIndex = BoneIndex;
+	FTransform BindTransform = FTransform::Identity;
+
+	do
+	{
+		BindTransform = BindTransform * BonePoses[CurBoneIndex];
+		CurBoneIndex = BoneInfos[CurBoneIndex].ParentIndex;
+	} while (CurBoneIndex != INDEX_NONE);
+
+	return BindTransform;
+}
diff --git a/Engine/Plugins/Enterprise/GLTFExporter/MigratedFromGit/Source/GLTFExporter/Private/Converters/GLTFBoneUtility.h b/Engine/Plugins/Enterprise/GLTFExporter/MigratedFromGit/Source/GLTFExporter/Private/Converters/GLTFBoneUtility.h
new file mode 100644
--- /dev/null
+++ b/Engine/Plugins/Enterprise/GLTFExporter/MigratedFromGit/Source/GLTFExporter/Private/Converters/GLTFBoneUtility.h
@@ -0,0 +1,10 @@
+// Copyright Epic Games, Inc. All Rights Reserved.
+
+#pragma once
+
+#include "CoreMinimal.h"
+
+struct FGLTFBoneUtility
+{
+	static FTransform GetBindTransform(const FReferenceSkeleton& RefSkeleton, const int32 BoneIndex);
+};
diff --git a/Engine/Plugins/Enterprise/GLTFExporter/MigratedFromGit/Source/GLTFExporter/Private/Converters/GLTFSkinConverters.cpp b/Engine/Plugins/Enterprise/GLTFExporter/MigratedFromGit/Source/GLTFExporter/Private/Converters/GLTFSkinConverters.cpp
--- a/Engine/Plugins/Enterprise/GLTFExporter/MigratedFromGit/Source/GLTFExporter/Private/Converters/GLTFSkinConverters.cpp
+++ b/Engine/Plugins/Enterprise/GLTFExporter/MigratedFromGit/Source/GLTFExporter/Private/Converters/GLTFSkinConverters.cpp
@@ -2,6 +2,7 @@
 
 #include "Converters/GLTFSkinConverters.h"
 #include "Converters/GLTFConverterUtility.h"
+#include "Converters/GLTFBoneUtility.h"
 #include "Builders/GLTFConvertBuilder.h"
 
 FGLTFJsonSkinIndex FGLTFSkinConverter::Convert(FGLTFJsonNodeIndex RootNode, const USkeletalMesh* SkeletalMesh)
@@ -18,5 +19,21 @@ FGLTFJsonSkinIndex FGLTFSkinConverter::Convert(FGLTFJsonNodeIndex RootNode, cons
 		Skin.Joints[BoneIndex] = Builder.GetOrAddNode(RootNode, SkeletalMesh, BoneIndex);
 	}
 
+	TArray<FGLTFJsonMatrix4> InverseBindMatrices;
+	InverseBindMatrices.AddUninitialized(BoneCount);
+
+	for (int32 BoneIndex = 0; BoneIndex < BoneCount; ++BoneIndex)
+	{
+		const FTransform InverseTransform = FGLTFBoneUtility::GetBindTransform(SkeletalMesh->RefSkeleton, BoneIndex).Inverse();
+		InverseBindMatrices[BoneIndex] = FGLTFConverterUtility::ConvertTransform(InverseTransform, Builder.ExportOptions->ExportScale);
+	}
+
+	FGLTFJsonAccessor JsonAccessor;
+	JsonAccessor.BufferView = Builder.AddBufferView(InverseBindMatrices);
+	JsonAccessor.ComponentType = EGLTFJsonComponentType::F32;
+	JsonAccessor.Count = BoneCount;
+	JsonAccessor.Type = EGLTFJsonAccessorType::Mat4;
+
+	Skin.InverseBindMatrices = Builder.AddAccessor(JsonAccessor);
 	return Builder.AddSkin(Skin);
 }
