[GLTFExporter] feat(json): add json data structures for animations

commit 4ea6189625ee0bcbd08eab0bfe03057681e8c8c0
Date: Tue, 9 Jun 2020 01:51:38 +0200

    feat(json): add json data structures for animations
    (cherry picked from commit 9762c8b0d446ba5736e9b28ee426e131c32f8425​)

#skipci

diff --git a/Engine/Plugins/Enterprise/GLTFExporter/MigratedFromGit/Source/GLTFExporter/Private/Builders/GLTFJsonBuilder.cpp b/Engine/Plugins/Enterprise/GLTFExporter/MigratedFromGit/Source/GLTFExporter/Private/Builders/GLTFJsonBuilder.cpp
--- a/Engine/Plugins/Enterprise/GLTFExporter/MigratedFromGit/Source/GLTFExporter/Private/Builders/GLTFJsonBuilder.cpp
+++ b/Engine/Plugins/Enterprise/GLTFExporter/MigratedFromGit/Source/GLTFExporter/Private/Builders/GLTFJsonBuilder.cpp
@@ -13,6 +13,11 @@ FGLTFJsonAccessorIndex FGLTFJsonBuilder::AddAccessor(const FGLTFJsonAccessor& Js
 	return FGLTFJsonAccessorIndex(JsonRoot.Accessors.Add(JsonAccessor));
 }
 
+FGLTFJsonAnimationIndex FGLTFJsonBuilder::AddAnimation(const FGLTFJsonAnimation& JsonAnimation)
+{
+	return FGLTFJsonAnimationIndex(JsonRoot.Animations.Add(JsonAnimation));
+}
+
 FGLTFJsonBufferIndex FGLTFJsonBuilder::AddBuffer(const FGLTFJsonBuffer& JsonBuffer)
 {
 	return FGLTFJsonBufferIndex(JsonRoot.Buffers.Add(JsonBuffer));
@@ -122,6 +127,11 @@ FGLTFJsonAccessor& FGLTFJsonBuilder::GetAccessor(FGLTFJsonAccessorIndex Accessor
 	return JsonRoot.Accessors[AccessorIndex];
 }
 
+FGLTFJsonAnimation& FGLTFJsonBuilder::GetAnimation(FGLTFJsonAnimationIndex AnimationIndex)
+{
+	return JsonRoot.Animations[AnimationIndex];
+}
+
 FGLTFJsonBuffer& FGLTFJsonBuilder::GetBuffer(FGLTFJsonBufferIndex BufferIndex)
 {
 	return JsonRoot.Buffers[BufferIndex];
diff --git a/Engine/Plugins/Enterprise/GLTFExporter/MigratedFromGit/Source/GLTFExporter/Private/Builders/GLTFJsonBuilder.h b/Engine/Plugins/Enterprise/GLTFExporter/MigratedFromGit/Source/GLTFExporter/Private/Builders/GLTFJsonBuilder.h
--- a/Engine/Plugins/Enterprise/GLTFExporter/MigratedFromGit/Source/GLTFExporter/Private/Builders/GLTFJsonBuilder.h
+++ b/Engine/Plugins/Enterprise/GLTFExporter/MigratedFromGit/Source/GLTFExporter/Private/Builders/GLTFJsonBuilder.h
@@ -16,6 +16,7 @@ public:
 	FGLTFJsonSceneIndex& DefaultScene;
 
 	FGLTFJsonAccessorIndex AddAccessor(const FGLTFJsonAccessor& JsonAccessor = {});
+	FGLTFJsonAnimationIndex AddAnimation(const FGLTFJsonAnimation& JsonAnimation = {});
 	FGLTFJsonBufferIndex AddBuffer(const FGLTFJsonBuffer& JsonBuffer = {});
 	FGLTFJsonBufferViewIndex AddBufferView(const FGLTFJsonBufferView& JsonBufferView = {});
 	FGLTFJsonCameraIndex AddCamera(const FGLTFJsonCamera& JsonCamera = {});
@@ -37,6 +38,7 @@ public:
 	FGLTFJsonNodeIndex AddChildComponentNode(FGLTFJsonNodeIndex ParentNodeIndex, const FGLTFJsonNode& JsonNode = {});
 
 	FGLTFJsonAccessor& GetAccessor(FGLTFJsonAccessorIndex AccessorIndex);
+	FGLTFJsonAnimation& GetAnimation(FGLTFJsonAnimationIndex AnimationIndex);
 	FGLTFJsonBuffer& GetBuffer(FGLTFJsonBufferIndex BufferIndex);
 	FGLTFJsonBufferView& GetBufferView(FGLTFJsonBufferViewIndex BufferViewIndex);
 	FGLTFJsonCamera& GetCamera(FGLTFJsonCameraIndex CameraIndex);
diff --git a/Engine/Plugins/Enterprise/GLTFExporter/MigratedFromGit/Source/GLTFExporter/Private/Converters/GLTFConverterUtility.cpp b/Engine/Plugins/Enterprise/GLTFExporter/MigratedFromGit/Source/GLTFExporter/Private/Converters/GLTFConverterUtility.cpp
--- a/Engine/Plugins/Enterprise/GLTFExporter/MigratedFromGit/Source/GLTFExporter/Private/Converters/GLTFConverterUtility.cpp
+++ b/Engine/Plugins/Enterprise/GLTFExporter/MigratedFromGit/Source/GLTFExporter/Private/Converters/GLTFConverterUtility.cpp
@@ -24,6 +24,16 @@ EGLTFJsonLightType FGLTFConverterUtility::ConvertLightType(ELightComponentType C
 	}
 }
 
+EGLTFJsonInterpolation FGLTFConverterUtility::ConvertInterpolation(const EAnimInterpolationType Type)
+{
+	switch (Type)
+	{
+		case EAnimInterpolationType::Linear: return EGLTFJsonInterpolation::Linear;
+		case EAnimInterpolationType::Step:   return EGLTFJsonInterpolation::Step;
+		default:                             return EGLTFJsonInterpolation::Linear; // fallback
+	}
+}
+
 EGLTFJsonShadingModel FGLTFConverterUtility::ConvertShadingModel(EMaterialShadingModel ShadingModel)
 {
 	switch (ShadingModel)
diff --git a/Engine/Plugins/Enterprise/GLTFExporter/MigratedFromGit/Source/GLTFExporter/Private/Converters/GLTFConverterUtility.h b/Engine/Plugins/Enterprise/GLTFExporter/MigratedFromGit/Source/GLTFExporter/Private/Converters/GLTFConverterUtility.h
--- a/Engine/Plugins/Enterprise/GLTFExporter/MigratedFromGit/Source/GLTFExporter/Private/Converters/GLTFConverterUtility.h
+++ b/Engine/Plugins/Enterprise/GLTFExporter/MigratedFromGit/Source/GLTFExporter/Private/Converters/GLTFConverterUtility.h
@@ -150,6 +150,8 @@ struct FGLTFConverterUtility
 	static EGLTFJsonCameraType ConvertCameraType(ECameraProjectionMode::Type ProjectionMode);
 	static EGLTFJsonLightType ConvertLightType(ELightComponentType ComponentType);
 
+	static EGLTFJsonInterpolation ConvertInterpolation(const EAnimInterpolationType Type);
+
 	static EGLTFJsonShadingModel ConvertShadingModel(EMaterialShadingModel ShadingModel);
 
 	static EGLTFJsonAlphaMode ConvertBlendMode(EBlendMode Mode);
diff --git a/Engine/Plugins/Enterprise/GLTFExporter/MigratedFromGit/Source/GLTFExporter/Private/Json/GLTFJsonAnimation.h b/Engine/Plugins/Enterprise/GLTFExporter/MigratedFromGit/Source/GLTFExporter/Private/Json/GLTFJsonAnimation.h
new file mode 100644
--- /dev/null
+++ b/Engine/Plugins/Enterprise/GLTFExporter/MigratedFromGit/Source/GLTFExporter/Private/Json/GLTFJsonAnimation.h
@@ -0,0 +1,102 @@
+// Copyright Epic Games, Inc. All Rights Reserved.
+
+#pragma once
+
+#include "Json/GLTFJsonEnums.h"
+#include "Json/GLTFJsonIndex.h"
+#include "Json/GLTFJsonUtility.h"
+#include "Serialization/JsonSerializer.h"
+
+struct FGLTFJsonAnimationChannelTarget
+{
+	FGLTFJsonNodeIndex Node;
+	EGLTFJsonTargetPath Path;
+
+	FGLTFJsonAnimationChannelTarget()
+        : Path(EGLTFJsonTargetPath::None)
+	{
+	}
+
+	template <class CharType = TCHAR, class PrintPolicy = TPrettyJsonPrintPolicy<CharType>>
+    void WriteObject(TJsonWriter<CharType, PrintPolicy>& JsonWriter, FGLTFJsonExtensions& Extensions) const
+	{
+		JsonWriter.WriteObjectStart();
+
+		JsonWriter.WriteValue(TEXT("node"), Node);
+		JsonWriter.WriteValue(TEXT("path"), FGLTFJsonUtility::ToString(Path));
+
+		JsonWriter.WriteObjectEnd();
+	}
+};
+
+struct FGLTFJsonAnimationChannel
+{
+	FGLTFJsonAnimationSamplerIndex Sampler;
+	FGLTFJsonAnimationChannelTarget Target;
+
+	template <class CharType = TCHAR, class PrintPolicy = TPrettyJsonPrintPolicy<CharType>>
+    void WriteObject(TJsonWriter<CharType, PrintPolicy>& JsonWriter, FGLTFJsonExtensions& Extensions) const
+	{
+		JsonWriter.WriteObjectStart();
+
+		JsonWriter.WriteValue(TEXT("sampler"), Sampler);
+
+		JsonWriter.WriteIdentifierPrefix(TEXT("target"));
+		Target.WriteObject(JsonWriter, Extensions);
+
+		JsonWriter.WriteObjectEnd();
+	}
+};
+
+struct FGLTFJsonAnimationSampler
+{
+	FGLTFJsonAccessorIndex Input;
+	FGLTFJsonAccessorIndex Output;
+
+	EGLTFJsonInterpolation Interpolation;
+
+	FGLTFJsonAnimationSampler()
+        : Interpolation(EGLTFJsonInterpolation::Linear)
+	{
+	}
+
+	template <class CharType = TCHAR, class PrintPolicy = TPrettyJsonPrintPolicy<CharType>>
+    void WriteObject(TJsonWriter<CharType, PrintPolicy>& JsonWriter, FGLTFJsonExtensions& Extensions) const
+	{
+		JsonWriter.WriteObjectStart();
+
+		JsonWriter.WriteValue(TEXT("input"), Input);
+		JsonWriter.WriteValue(TEXT("output"), Output);
+
+		if (Interpolation != EGLTFJsonInterpolation::Linear)
+		{
+			JsonWriter.WriteValue(TEXT("interpolation"), FGLTFJsonUtility::ToString(Interpolation));
+		}
+
+		JsonWriter.WriteObjectEnd();
+	}
+};
+
+struct FGLTFJsonAnimation
+{
+	FString Name;
+
+	TArray<FGLTFJsonAnimationChannel> Channels;
+	TArray<FGLTFJsonAnimationSampler> Samplers;
+
+	template <class CharType = TCHAR, class PrintPolicy = TPrettyJsonPrintPolicy<CharType>>
+	void WriteObject(TJsonWriter<CharType, PrintPolicy>& JsonWriter, FGLTFJsonExtensions& Extensions) const
+	{
+		JsonWriter.WriteObjectStart();
+
+		if (!Name.IsEmpty())
+		{
+			JsonWriter.WriteValue(TEXT("name"), Name);
+		}
+
+		FGLTFJsonUtility::WriteObjectArray(JsonWriter, TEXT("channels"), Channels, Extensions, true);
+		FGLTFJsonUtility::WriteObjectArray(JsonWriter, TEXT("samplers"), Samplers, Extensions, true);
+
+		JsonWriter.WriteObjectEnd();
+	}
+};
diff --git a/Engine/Plugins/Enterprise/GLTFExporter/MigratedFromGit/Source/GLTFExporter/Private/Json/GLTFJsonEnums.h b/Engine/Plugins/Enterprise/GLTFExporter/MigratedFromGit/Source/GLTFExporter/Private/Json/GLTFJsonEnums.h
--- a/Engine/Plugins/Enterprise/GLTFExporter/MigratedFromGit/Source/GLTFExporter/Private/Json/GLTFJsonEnums.h
+++ b/Engine/Plugins/Enterprise/GLTFExporter/MigratedFromGit/Source/GLTFExporter/Private/Json/GLTFJsonEnums.h
@@ -124,14 +124,14 @@ enum class EGLTFJsonTextureWrap
 	ClampToEdge = 33071
 };
 
-enum class EGLTFJsonAnimationInterpolation
+enum class EGLTFJsonInterpolation
 {
 	Linear,
 	Step,
 	CubicSpline,
 };
 
-enum class EGLTFJsonAnimationPath
+enum class EGLTFJsonTargetPath
 {
 	None = -1,
 	Translation,
diff --git a/Engine/Plugins/Enterprise/GLTFExporter/MigratedFromGit/Source/GLTFExporter/Private/Json/GLTFJsonIndex.h b/Engine/Plugins/Enterprise/GLTFExporter/MigratedFromGit/Source/GLTFExporter/Private/Json/GLTFJsonIndex.h
--- a/Engine/Plugins/Enterprise/GLTFExporter/MigratedFromGit/Source/GLTFExporter/Private/Json/GLTFJsonIndex.h
+++ b/Engine/Plugins/Enterprise/GLTFExporter/MigratedFromGit/Source/GLTFExporter/Private/Json/GLTFJsonIndex.h
@@ -53,6 +53,8 @@ private:
 };
 
 struct FGLTFJsonAccessorIndex : FGLTFJsonIndex<FGLTFJsonAccessorIndex> { using FGLTFJsonIndex::FGLTFJsonIndex; };
+struct FGLTFJsonAnimationIndex : FGLTFJsonIndex<FGLTFJsonAnimationIndex> { using FGLTFJsonIndex::FGLTFJsonIndex; };
+struct FGLTFJsonAnimationSamplerIndex : FGLTFJsonIndex<FGLTFJsonAnimationSamplerIndex> { using FGLTFJsonIndex::FGLTFJsonIndex; };
 struct FGLTFJsonBufferIndex : FGLTFJsonIndex<FGLTFJsonBufferIndex> { using FGLTFJsonIndex::FGLTFJsonIndex; };
 struct FGLTFJsonBufferViewIndex : FGLTFJsonIndex<FGLTFJsonBufferViewIndex> { using FGLTFJsonIndex::FGLTFJsonIndex; };
 struct FGLTFJsonCameraIndex : FGLTFJsonIndex<FGLTFJsonCameraIndex> { using FGLTFJsonIndex::FGLTFJsonIndex; };
@@ -69,4 +71,3 @@ struct FGLTFJsonVariationIndex : FGLTFJsonIndex<FGLTFJsonVariationIndex> { using
 struct FGLTFJsonLightMapIndex : FGLTFJsonIndex<FGLTFJsonLightMapIndex> { using FGLTFJsonIndex::FGLTFJsonIndex; };
 struct FGLTFJsonLightIndex : FGLTFJsonIndex<FGLTFJsonLightIndex> { using FGLTFJsonIndex::FGLTFJsonIndex; };
 struct FGLTFJsonHotspotIndex : FGLTFJsonIndex<FGLTFJsonHotspotIndex> { using FGLTFJsonIndex::FGLTFJsonIndex; };
-struct FGLTFJsonAnimationIndex : FGLTFJsonIndex<FGLTFJsonAnimationIndex> { using FGLTFJsonIndex::FGLTFJsonIndex; };
diff --git a/Engine/Plugins/Enterprise/GLTFExporter/MigratedFromGit/Source/GLTFExporter/Private/Json/GLTFJsonRoot.h b/Engine/Plugins/Enterprise/GLTFExporter/MigratedFromGit/Source/GLTFExporter/Private/Json/GLTFJsonRoot.h
--- a/Engine/Plugins/Enterprise/GLTFExporter/MigratedFromGit/Source/GLTFExporter/Private/Json/GLTFJsonRoot.h
+++ b/Engine/Plugins/Enterprise/GLTFExporter/MigratedFromGit/Source/GLTFExporter/Private/Json/GLTFJsonRoot.h
@@ -5,6 +5,7 @@
 #include "Json/GLTFJsonIndex.h"
 #include "Json/GLTFJsonExtensions.h"
 #include "Json/GLTFJsonAccessor.h"
+#include "Json/GLTFJsonAnimation.h"
 #include "Json/GLTFJsonBuffer.h"
 #include "Json/GLTFJsonBufferView.h"
 #include "Json/GLTFJsonCamera.h"
@@ -62,6 +63,7 @@ struct FGLTFJsonRoot
 	FGLTFJsonSceneIndex DefaultScene;
 
 	TArray<FGLTFJsonAccessor>   Accessors;
+	TArray<FGLTFJsonAnimation>  Animations;
 	TArray<FGLTFJsonBuffer>     Buffers;
 	TArray<FGLTFJsonBufferView> BufferViews;
 	TArray<FGLTFJsonCamera>     Cameras;
@@ -95,6 +97,7 @@ struct FGLTFJsonRoot
 		}
 
 		FGLTFJsonUtility::WriteObjectArray(JsonWriter, TEXT("accessors"), Accessors, AllExtensions);
+		FGLTFJsonUtility::WriteObjectArray(JsonWriter, TEXT("animations"), Animations, AllExtensions);
 		FGLTFJsonUtility::WriteObjectArray(JsonWriter, TEXT("buffers"), Buffers, AllExtensions);
 		FGLTFJsonUtility::WriteObjectArray(JsonWriter, TEXT("bufferViews"), BufferViews, AllExtensions);
 		FGLTFJsonUtility::WriteObjectArray(JsonWriter, TEXT("cameras"), Cameras, AllExtensions);
diff --git a/Engine/Plugins/Enterprise/GLTFExporter/MigratedFromGit/Source/GLTFExporter/Private/Json/GLTFJsonUtility.h b/Engine/Plugins/Enterprise/GLTFExporter/MigratedFromGit/Source/GLTFExporter/Private/Json/GLTFJsonUtility.h
--- a/Engine/Plugins/Enterprise/GLTFExporter/MigratedFromGit/Source/GLTFExporter/Private/Json/GLTFJsonUtility.h
+++ b/Engine/Plugins/Enterprise/GLTFExporter/MigratedFromGit/Source/GLTFExporter/Private/Json/GLTFJsonUtility.h
@@ -115,6 +115,29 @@ struct FGLTFJsonUtility
 		}
 	}
 
+	static const TCHAR* ToString(EGLTFJsonInterpolation Value)
+	{
+		switch (Value)
+		{
+			case EGLTFJsonInterpolation::Linear:      return TEXT("LINEAR");
+			case EGLTFJsonInterpolation::Step:        return TEXT("STEP");
+			case EGLTFJsonInterpolation::CubicSpline: return TEXT("CUBICSPLINE");
+			default:                                  return TEXT("UNKNOWN");
+		}
+	}
+
+	static const TCHAR* ToString(EGLTFJsonTargetPath Value)
+	{
+		switch (Value)
+		{
+			case EGLTFJsonTargetPath::Translation: return TEXT("translation");
+			case EGLTFJsonTargetPath::Rotation:    return TEXT("rotation");
+			case EGLTFJsonTargetPath::Scale:       return TEXT("scale");
+			case EGLTFJsonTargetPath::Weights:     return TEXT("weights");
+			default:                               return TEXT("unknown");
+		}
+	}
+
 	template <class CharType = TCHAR, class PrintPolicy = TPrettyJsonPrintPolicy<CharType>>
 	static void WriteExactValue(TJsonWriter<CharType, PrintPolicy>& JsonWriter, float Value)
 	{
