[GLTFExporter] feat(skeletal-animations): add initial converter for animation sequence

commit 2e81655447a3d11caa411469cd598da1f7ad3b53
Date: Thu, 11 Jun 2020 10:19:56 +0200

    feat(skeletal-animations): add initial converter for animation sequence
    (cherry picked from commit afdb73ae044f883328afa82cf90678c4be2abe47​)

#skipci

diff --git a/Engine/Plugins/Enterprise/GLTFExporter/MigratedFromGit/Source/GLTFExporter/Private/Converters/GLTFAnimationConverters.cpp b/Engine/Plugins/Enterprise/GLTFExporter/MigratedFromGit/Source/GLTFExporter/Private/Converters/GLTFAnimationConverters.cpp
new file mode 100644
--- /dev/null
+++ b/Engine/Plugins/Enterprise/GLTFExporter/MigratedFromGit/Source/GLTFExporter/Private/Converters/GLTFAnimationConverters.cpp
@@ -0,0 +1,162 @@
+// Copyright Epic Games, Inc. All Rights Reserved.
+
+#include "Converters/GLTFAnimationConverters.h"
+#include "Converters/GLTFConverterUtility.h"
+#include "Builders/GLTFConvertBuilder.h"
+#include "Animation/AnimSequence.h"
+
+FGLTFJsonAnimationIndex FGLTFAnimationConverter::Convert(FGLTFJsonNodeIndex RootNode, const USkeletalMesh* SkeletalMesh, const UAnimSequence* AnimSequence)
+{
+	const int32 FrameCount = AnimSequence->GetRawNumberOfFrames();
+	if (FrameCount < 0)
+	{
+		return FGLTFJsonAnimationIndex(INDEX_NONE);
+	}
+
+	FGLTFJsonAnimation JsonAnimation;
+	AnimSequence->GetName(JsonAnimation.Name);
+
+	TArray<float> Timestamps;
+	Timestamps.AddUninitialized(FrameCount);
+
+	for (int32 Frame = 0; Frame < FrameCount; ++Frame)
+	{
+		Timestamps[Frame] = AnimSequence->GetTimeAtFrame(Frame);
+	}
+
+	FGLTFJsonAccessor TimestampAccessor;
+	TimestampAccessor.BufferView = Builder.AddBufferView(Timestamps);
+	TimestampAccessor.ComponentType = EGLTFJsonComponentType::F32;
+	TimestampAccessor.Count = Timestamps.Num();
+	TimestampAccessor.Type = EGLTFJsonAccessorType::Scalar;
+	TimestampAccessor.MinMaxLength = 1;
+	TimestampAccessor.Min[0] = 0;
+	TimestampAccessor.Max[0] = AnimSequence->SequenceLength;
+	const FGLTFJsonAccessorIndex TimestampAccessorIndex = Builder.AddAccessor(TimestampAccessor);
+
+	const EGLTFJsonInterpolation Interpolation = FGLTFConverterUtility::ConvertInterpolation(AnimSequence->Interpolation);
+	const TArray<FName>& TrackNames = AnimSequence->GetAnimationTrackNames();
+	const int32 TrackCount = TrackNames.Num();
+
+	for (int32 TrackIndex = 0; TrackIndex < TrackCount; ++TrackIndex)
+	{
+		const FRawAnimSequenceTrack& Track = AnimSequence->GetRawAnimationTrack(TrackIndex);
+
+		const int32 BoneTreeIndex = AnimSequence->GetSkeletonIndexFromRawDataTrackIndex(TrackIndex);
+		const int32 BoneIndex = AnimSequence->GetSkeleton()->GetMeshBoneIndexFromSkeletonBoneIndex(SkeletalMesh, BoneTreeIndex);
+		const FGLTFJsonNodeIndex NodeIndex = Builder.GetOrAddNode(RootNode, SkeletalMesh, BoneIndex);
+
+		const TArray<FVector>& PosKeys = Track.PosKeys;
+		if (PosKeys.Num() > 0)
+		{
+			TArray<FGLTFJsonVector3> Translations;
+			Translations.AddUninitialized(FrameCount); // TODO: PosKeys.Num()
+
+			for (int32 Index = 0; Index < PosKeys.Num(); ++Index)
+			{
+				Translations[Index] = FGLTFConverterUtility::ConvertPosition(PosKeys[Index], Builder.ExportOptions->ExportScale);
+			}
+
+			// TODO: remove later
+			const FGLTFJsonVector3 Last = FGLTFConverterUtility::ConvertPosition(PosKeys.Last(), Builder.ExportOptions->ExportScale);
+			for (int32 Index = PosKeys.Num(); Index < FrameCount; ++Index)
+			{
+				Translations[Index] = Last;
+			}
+
+			FGLTFJsonAccessor JsonAccessor;
+			JsonAccessor.BufferView = Builder.AddBufferView(Translations);
+			JsonAccessor.ComponentType = EGLTFJsonComponentType::F32;
+			JsonAccessor.Count = Translations.Num();
+			JsonAccessor.Type = EGLTFJsonAccessorType::Vec3;
+
+			FGLTFJsonAnimationSampler JsonSampler;
+			JsonSampler.Input = TimestampAccessorIndex;
+			JsonSampler.Output = Builder.AddAccessor(JsonAccessor);
+			JsonSampler.Interpolation = Interpolation;
+
+			FGLTFJsonAnimationChannel JsonChannel;
+			JsonChannel.Sampler = FGLTFJsonAnimationSamplerIndex(JsonAnimation.Samplers.Add(JsonSampler));
+			JsonChannel.Target.Path = EGLTFJsonTargetPath::Translation;
+			JsonChannel.Target.Node = NodeIndex;
+
+			JsonAnimation.Channels.Add(JsonChannel);
+		}
+
+		const TArray<FQuat>& RotKeys = Track.RotKeys;
+		if (RotKeys.Num() > 0)
+		{
+			TArray<FGLTFJsonQuaternion> Rotations;
+			Rotations.AddUninitialized(FrameCount); // TODO: RotKeys.Num()
+
+			for (int32 Index = 0; Index < RotKeys.Num(); ++Index)
+			{
+				Rotations[Index] = FGLTFConverterUtility::ConvertRotation(RotKeys[Index]);
+			}
+
+			// TODO: remove later
+			const FGLTFJsonQuaternion Last = FGLTFConverterUtility::ConvertRotation(RotKeys.Last());
+			for (int32 Index = RotKeys.Num(); Index < FrameCount; ++Index)
+			{
+				Rotations[Index] = Last;
+			}
+
+			FGLTFJsonAccessor JsonAccessor;
+			JsonAccessor.BufferView = Builder.AddBufferView(Rotations);
+			JsonAccessor.ComponentType = EGLTFJsonComponentType::F32;
+			JsonAccessor.Count = Rotations.Num();
+			JsonAccessor.Type = EGLTFJsonAccessorType::Vec4;
+
+			FGLTFJsonAnimationSampler JsonSampler;
+			JsonSampler.Input = TimestampAccessorIndex;
+			JsonSampler.Output = Builder.AddAccessor(JsonAccessor);
+			JsonSampler.Interpolation = Interpolation;
+
+			FGLTFJsonAnimationChannel JsonChannel;
+			JsonChannel.Sampler = FGLTFJsonAnimationSamplerIndex(JsonAnimation.Samplers.Add(JsonSampler));
+			JsonChannel.Target.Path = EGLTFJsonTargetPath::Rotation;
+			JsonChannel.Target.Node = NodeIndex;
+
+			JsonAnimation.Channels.Add(JsonChannel);
+		}
+
+		const TArray<FVector>& ScaleKeys = Track.ScaleKeys;
+		if (ScaleKeys.Num() > 0)
+		{
+			TArray<FGLTFJsonVector3> Scales;
+			Scales.AddUninitialized(FrameCount); // TODO: ScaleKeys.Num()
+
+			for (int32 Index = 0; Index < ScaleKeys.Num(); ++Index)
+			{
+				Scales[Index] = FGLTFConverterUtility::ConvertScale(ScaleKeys[Index]);
+			}
+
+			// TODO: remove later
+			const FGLTFJsonVector3 Last = FGLTFConverterUtility::ConvertScale(ScaleKeys.Last());
+			for (int32 Index = ScaleKeys.Num(); Index < FrameCount; ++Index)
+			{
+				Scales[Index] = Last;
+			}
+
+			FGLTFJsonAccessor JsonAccessor;
+			JsonAccessor.BufferView = Builder.AddBufferView(Scales);
+			JsonAccessor.ComponentType = EGLTFJsonComponentType::F32;
+			JsonAccessor.Count = Scales.Num();
+			JsonAccessor.Type = EGLTFJsonAccessorType::Vec3;
+
+			FGLTFJsonAnimationSampler JsonSampler;
+			JsonSampler.Input = TimestampAccessorIndex;
+			JsonSampler.Output = Builder.AddAccessor(JsonAccessor);
+			JsonSampler.Interpolation = Interpolation;
+
+			FGLTFJsonAnimationChannel JsonChannel;
+			JsonChannel.Sampler = FGLTFJsonAnimationSamplerIndex(JsonAnimation.Samplers.Add(JsonSampler));
+			JsonChannel.Target.Path = EGLTFJsonTargetPath::Scale;
+			JsonChannel.Target.Node = NodeIndex;
+
+			JsonAnimation.Channels.Add(JsonChannel);
+		}
+	}
+
+	return Builder.AddAnimation(JsonAnimation);
+}
diff --git a/Engine/Plugins/Enterprise/GLTFExporter/MigratedFromGit/Source/GLTFExporter/Private/Converters/GLTFAnimationConverters.h b/Engine/Plugins/Enterprise/GLTFExporter/MigratedFromGit/Source/GLTFExporter/Private/Converters/GLTFAnimationConverters.h
new file mode 100644
--- /dev/null
+++ b/Engine/Plugins/Enterprise/GLTFExporter/MigratedFromGit/Source/GLTFExporter/Private/Converters/GLTFAnimationConverters.h
@@ -0,0 +1,17 @@
+// Copyright Epic Games, Inc. All Rights Reserved.
+
+#pragma once
+
+#include "Json/GLTFJsonIndex.h"
+#include "Converters/GLTFConverter.h"
+#include "Converters/GLTFBuilderContext.h"
+#include "Engine.h"
+
+class UAnimSequence;
+
+class FGLTFAnimationConverter final : public FGLTFBuilderContext, public TGLTFConverter<FGLTFJsonAnimationIndex, FGLTFJsonNodeIndex, const USkeletalMesh*, const UAnimSequence*>
+{
+	using FGLTFBuilderContext::FGLTFBuilderContext;
+
+	virtual FGLTFJsonAnimationIndex Convert(FGLTFJsonNodeIndex RootNode, const USkeletalMesh* SkeletalMesh, const UAnimSequence* AnimSequence) override;
+};
