[GLTFExporter] feat(skeletal-animations): implement animation sequence exporter

commit 4fbe5ad1d35ab5516c01c29f1a35e5ae1713085d
Date: Thu, 11 Jun 2020 10:26:51 +0200

    feat(skeletal-animations): implement animation sequence exporter
    (cherry picked from commit 8595a1a86eedf8a4663760a03d8221353aaf5a79​)

#skipci

diff --git a/Engine/Plugins/Enterprise/GLTFExporter/MigratedFromGit/Source/GLTFExporter/Private/Builders/GLTFConvertBuilder.cpp b/Engine/Plugins/Enterprise/GLTFExporter/MigratedFromGit/Source/GLTFExporter/Private/Builders/GLTFConvertBuilder.cpp
--- a/Engine/Plugins/Enterprise/GLTFExporter/MigratedFromGit/Source/GLTFExporter/Private/Builders/GLTFConvertBuilder.cpp
+++ b/Engine/Plugins/Enterprise/GLTFExporter/MigratedFromGit/Source/GLTFExporter/Private/Builders/GLTFConvertBuilder.cpp
@@ -229,6 +229,16 @@ FGLTFJsonSkinIndex FGLTFConvertBuilder::GetOrAddSkin(FGLTFJsonNodeIndex RootNode
 	return SkinConverter.GetOrAdd(RootNode, SkeletalMesh);
 }
 
+FGLTFJsonAnimationIndex FGLTFConvertBuilder::GetOrAddAnimation(FGLTFJsonNodeIndex RootNode, const USkeletalMesh* SkeletalMesh, const UAnimSequence* AnimSequence)
+{
+	if (RootNode == INDEX_NONE || SkeletalMesh == nullptr || AnimSequence == nullptr)
+	{
+		return FGLTFJsonAnimationIndex(INDEX_NONE);
+	}
+
+	return AnimationConverter.GetOrAdd(RootNode, SkeletalMesh, AnimSequence);
+}
+
 FGLTFJsonNodeIndex FGLTFConvertBuilder::GetOrAddNode(const USceneComponent* SceneComponent, FName SocketName)
 {
 	if (SceneComponent == nullptr || SocketName == NAME_None)
diff --git a/Engine/Plugins/Enterprise/GLTFExporter/MigratedFromGit/Source/GLTFExporter/Private/Builders/GLTFConvertBuilder.h b/Engine/Plugins/Enterprise/GLTFExporter/MigratedFromGit/Source/GLTFExporter/Private/Builders/GLTFConvertBuilder.h
--- a/Engine/Plugins/Enterprise/GLTFExporter/MigratedFromGit/Source/GLTFExporter/Private/Builders/GLTFConvertBuilder.h
+++ b/Engine/Plugins/Enterprise/GLTFExporter/MigratedFromGit/Source/GLTFExporter/Private/Builders/GLTFConvertBuilder.h
@@ -11,6 +11,7 @@
 #include "Converters/GLTFTextureConverters.h"
 #include "Converters/GLTFNodeConverters.h"
 #include "Converters/GLTFSkinConverters.h"
+#include "Converters/GLTFAnimationConverters.h"
 #include "Converters/GLTFSceneConverters.h"
 #include "Converters/GLTFCameraConverters.h"
 #include "Converters/GLTFLightConverters.h"
@@ -54,6 +55,7 @@ public:
 	FGLTFJsonNodeIndex GetOrAddNode(FGLTFJsonNodeIndex RootNode, const USkeletalMesh* SkeletalMesh, FName SocketName);
 	FGLTFJsonNodeIndex GetOrAddNode(FGLTFJsonNodeIndex RootNode, const USkeletalMesh* SkeletalMesh, int32 BoneIndex);
 	FGLTFJsonSkinIndex GetOrAddSkin(FGLTFJsonNodeIndex RootNode, const USkeletalMesh* SkeletalMesh);
+	FGLTFJsonAnimationIndex GetOrAddAnimation(FGLTFJsonNodeIndex RootNode, const USkeletalMesh* SkeletalMesh, const UAnimSequence* AnimSequence);
 
 	FGLTFJsonNodeIndex GetOrAddNode(const USceneComponent* SceneComponent, FName SocketName);
 	FGLTFJsonNodeIndex GetOrAddNode(const USceneComponent* SceneComponent);
@@ -92,6 +94,7 @@ private:
 	FGLTFSkeletalSocketConverter SkeletalSocketConverter{ *this };
 	FGLTFSkeletalBoneConverter SkeletalBoneConverter{ *this };
 	FGLTFSkinConverter SkinConverter{ *this };
+	FGLTFAnimationConverter AnimationConverter{ *this };
 
 	FGLTFComponentSocketConverter ComponentSocketConverter{ *this };
 	FGLTFComponentConverter ComponentConverter{ *this };
diff --git a/Engine/Plugins/Enterprise/GLTFExporter/MigratedFromGit/Source/GLTFExporter/Private/Exporters/GLTFAnimSequenceExporter.cpp b/Engine/Plugins/Enterprise/GLTFExporter/MigratedFromGit/Source/GLTFExporter/Private/Exporters/GLTFAnimSequenceExporter.cpp
--- a/Engine/Plugins/Enterprise/GLTFExporter/MigratedFromGit/Source/GLTFExporter/Private/Exporters/GLTFAnimSequenceExporter.cpp
+++ b/Engine/Plugins/Enterprise/GLTFExporter/MigratedFromGit/Source/GLTFExporter/Private/Exporters/GLTFAnimSequenceExporter.cpp
@@ -1,18 +1,40 @@
 // Copyright Epic Games, Inc. All Rights Reserved.
 
 #include "Exporters/GLTFAnimSequenceExporter.h"
+#include "Exporters/GLTFExporterUtility.h"
+#include "Builders/GLTFContainerBuilder.h"
 #include "Animation/AnimSequence.h"
 
 UGLTFAnimSequenceExporter::UGLTFAnimSequenceExporter(const FObjectInitializer& ObjectInitializer)
 	: Super(ObjectInitializer)
 {
-	// TODO: uncomment when support is implemented
-	// SupportedClass = UAnimSequence::StaticClass();
+	SupportedClass = UAnimSequence::StaticClass();
 }
 
 bool UGLTFAnimSequenceExporter::AddObject(FGLTFContainerBuilder& Builder, const UObject* Object)
 {
 	const UAnimSequence* AnimSequence = CastChecked<UAnimSequence>(Object);
-	// TODO: implement
+
+	const USkeletalMesh* SkeletalMesh = FGLTFExporterUtility::GetPreviewMesh(AnimSequence);
+	if (SkeletalMesh == nullptr)
+	{
+		// report error
+		return false;
+	}
+
+	const FGLTFJsonMeshIndex MeshIndex = Builder.GetOrAddMesh(SkeletalMesh);
+
+	FGLTFJsonNode Node;
+	Node.Mesh = MeshIndex;
+	const FGLTFJsonNodeIndex NodeIndex = Builder.AddNode(Node);
+	Builder.GetNode(NodeIndex).Skin = Builder.GetOrAddSkin(NodeIndex, SkeletalMesh);
+
+	Builder.GetOrAddAnimation(NodeIndex, SkeletalMesh, AnimSequence);
+
+	FGLTFJsonScene Scene;
+	Scene.Nodes.Add(NodeIndex);
+	const FGLTFJsonSceneIndex SceneIndex = Builder.AddScene(Scene);
+
+	Builder.DefaultScene = SceneIndex;
 	return true;
 }
diff --git a/Engine/Plugins/Enterprise/GLTFExporter/MigratedFromGit/Source/GLTFExporter/Private/Exporters/GLTFExporterUtility.cpp b/Engine/Plugins/Enterprise/GLTFExporter/MigratedFromGit/Source/GLTFExporter/Private/Exporters/GLTFExporterUtility.cpp
--- a/Engine/Plugins/Enterprise/GLTFExporter/MigratedFromGit/Source/GLTFExporter/Private/Exporters/GLTFExporterUtility.cpp
+++ b/Engine/Plugins/Enterprise/GLTFExporter/MigratedFromGit/Source/GLTFExporter/Private/Exporters/GLTFExporterUtility.cpp
@@ -30,3 +30,18 @@ const UStaticMesh* FGLTFExporterUtility::GetPreviewMesh(const UMaterialInterface
 
 	return PreviewMesh;
 }
+
+const USkeletalMesh* FGLTFExporterUtility::GetPreviewMesh(const UAnimSequence* AnimSequence)
+{
+	const USkeletalMesh* PreviewMesh = AnimSequence->GetPreviewMesh();
+	if (PreviewMesh == nullptr)
+	{
+		const USkeleton* Skeleton = AnimSequence->GetSkeleton();
+		if (Skeleton != nullptr)
+		{
+			PreviewMesh = Skeleton->GetPreviewMesh();
+		}
+	}
+
+	return PreviewMesh;
+}
diff --git a/Engine/Plugins/Enterprise/GLTFExporter/MigratedFromGit/Source/GLTFExporter/Private/Exporters/GLTFExporterUtility.h b/Engine/Plugins/Enterprise/GLTFExporter/MigratedFromGit/Source/GLTFExporter/Private/Exporters/GLTFExporterUtility.h
--- a/Engine/Plugins/Enterprise/GLTFExporter/MigratedFromGit/Source/GLTFExporter/Private/Exporters/GLTFExporterUtility.h
+++ b/Engine/Plugins/Enterprise/GLTFExporter/MigratedFromGit/Source/GLTFExporter/Private/Exporters/GLTFExporterUtility.h
@@ -8,4 +8,5 @@ class UMaterialInterface;
 struct FGLTFExporterUtility
 {
 	static const UStaticMesh* GetPreviewMesh(const UMaterialInterface* Material);
+	static const USkeletalMesh* GetPreviewMesh(const UAnimSequence* AnimSequence);
 };
