[GLTFExporter] fix(skeletal-animations): fix issue with gltf skin index not being applied

commit cc1459488df018d7d7f6a3f055d62eb8d985f86e
Date: Sat, 26 Sep 2020 02:56:31 +0200

    fix(skeletal-animations): fix issue with gltf skin index not being applied
    (cherry picked from commit 947d3dceb40159a8e30a42054ac015417853ef07​)

#skipci

diff --git a/Engine/Plugins/Enterprise/GLTFExporter/MigratedFromGit/Source/GLTFExporter/Private/Exporters/GLTFAnimSequenceExporter.cpp b/Engine/Plugins/Enterprise/GLTFExporter/MigratedFromGit/Source/GLTFExporter/Private/Exporters/GLTFAnimSequenceExporter.cpp
--- a/Engine/Plugins/Enterprise/GLTFExporter/MigratedFromGit/Source/GLTFExporter/Private/Exporters/GLTFAnimSequenceExporter.cpp
+++ b/Engine/Plugins/Enterprise/GLTFExporter/MigratedFromGit/Source/GLTFExporter/Private/Exporters/GLTFAnimSequenceExporter.cpp
@@ -27,7 +27,9 @@ bool UGLTFAnimSequenceExporter::AddObject(FGLTFContainerBuilder& Builder, const
 	FGLTFJsonNode Node;
 	Node.Mesh = MeshIndex;
 	const FGLTFJsonNodeIndex NodeIndex = Builder.AddNode(Node);
-	Builder.GetNode(NodeIndex).Skin = Builder.GetOrAddSkin(NodeIndex, SkeletalMesh);
+
+	const FGLTFJsonSkinIndex SkinIndex = Builder.GetOrAddSkin(NodeIndex, SkeletalMesh);
+	Builder.GetNode(NodeIndex).Skin = SkinIndex;
 
 	Builder.GetOrAddAnimation(NodeIndex, SkeletalMesh, AnimSequence);
 
