[GLTFExporter] feat(skeletal-animations): optimize export size by not padding animation key frames with last transform if lower than frame count

commit 758e3df170ee874dc4849a3424ae32b5d706b5a6
Date: Sat, 26 Sep 2020 02:59:22 +0200

    feat(skeletal-animations): optimize export size by not padding animation key frames with last transform if lower than frame count
    (cherry picked from commit ac3aba446d207f0b25f6fbb706bf22e6cfd34bf3​)

#skipci

diff --git a/Engine/Plugins/Enterprise/GLTFExporter/MigratedFromGit/Source/GLTFExporter/Private/Converters/GLTFAnimationConverters.cpp b/Engine/Plugins/Enterprise/GLTFExporter/MigratedFromGit/Source/GLTFExporter/Private/Converters/GLTFAnimationConverters.cpp
--- a/Engine/Plugins/Enterprise/GLTFExporter/MigratedFromGit/Source/GLTFExporter/Private/Converters/GLTFAnimationConverters.cpp
+++ b/Engine/Plugins/Enterprise/GLTFExporter/MigratedFromGit/Source/GLTFExporter/Private/Converters/GLTFAnimationConverters.cpp
@@ -10,6 +10,14 @@ FGLTFJsonAnimationIndex FGLTFAnimationConverter::Convert(FGLTFJsonNodeIndex Root
 	const int32 FrameCount = AnimSequence->GetRawNumberOfFrames();
 	if (FrameCount < 0)
 	{
+		// TODO: report warning
+		return FGLTFJsonAnimationIndex(INDEX_NONE);
+	}
+
+	USkeleton* Skeleton = AnimSequence->GetSkeleton();
+	if (Skeleton == nullptr)
+	{
+		// TODO: report error
 		return FGLTFJsonAnimationIndex(INDEX_NONE);
 	}
 
@@ -24,15 +32,12 @@ FGLTFJsonAnimationIndex FGLTFAnimationConverter::Convert(FGLTFJsonNodeIndex Root
 		Timestamps[Frame] = AnimSequence->GetTimeAtFrame(Frame);
 	}
 
-	FGLTFJsonAccessor TimestampAccessor;
-	TimestampAccessor.BufferView = Builder.AddBufferView(Timestamps);
-	TimestampAccessor.ComponentType = EGLTFJsonComponentType::F32;
-	TimestampAccessor.Count = Timestamps.Num();
-	TimestampAccessor.Type = EGLTFJsonAccessorType::Scalar;
-	TimestampAccessor.MinMaxLength = 1;
-	TimestampAccessor.Min[0] = 0;
-	TimestampAccessor.Max[0] = AnimSequence->SequenceLength;
-	const FGLTFJsonAccessorIndex TimestampAccessorIndex = Builder.AddAccessor(TimestampAccessor);
+	FGLTFJsonAccessor JsonInputAccessor;
+	JsonInputAccessor.BufferView = Builder.AddBufferView(Timestamps);
+	JsonInputAccessor.ComponentType = EGLTFJsonComponentType::F32;
+	JsonInputAccessor.Type = EGLTFJsonAccessorType::Scalar;
+	JsonInputAccessor.MinMaxLength = 1;
+	JsonInputAccessor.Min[0] = 0;
 
 	const EGLTFJsonInterpolation Interpolation = FGLTFConverterUtility::ConvertInterpolation(AnimSequence->Interpolation);
 	const TArray<FName>& TrackNames = AnimSequence->GetAnimationTrackNames();
@@ -42,37 +47,33 @@ FGLTFJsonAnimationIndex FGLTFAnimationConverter::Convert(FGLTFJsonNodeIndex Root
 	{
 		const FRawAnimSequenceTrack& Track = AnimSequence->GetRawAnimationTrack(TrackIndex);
 
-		const int32 BoneTreeIndex = AnimSequence->GetSkeletonIndexFromRawDataTrackIndex(TrackIndex);
-		const int32 BoneIndex = AnimSequence->GetSkeleton()->GetMeshBoneIndexFromSkeletonBoneIndex(SkeletalMesh, BoneTreeIndex);
+		const int32 SkeletonBoneIndex = AnimSequence->GetSkeletonIndexFromRawDataTrackIndex(TrackIndex);
+		const int32 BoneIndex = Skeleton->GetMeshBoneIndexFromSkeletonBoneIndex(SkeletalMesh, SkeletonBoneIndex);
 		const FGLTFJsonNodeIndex NodeIndex = Builder.GetOrAddNode(RootNode, SkeletalMesh, BoneIndex);
 
 		const TArray<FVector>& PosKeys = Track.PosKeys;
 		if (PosKeys.Num() > 0)
 		{
 			TArray<FGLTFJsonVector3> Translations;
-			Translations.AddUninitialized(FrameCount); // TODO: PosKeys.Num()
+			Translations.AddUninitialized(PosKeys.Num());
 
 			for (int32 Index = 0; Index < PosKeys.Num(); ++Index)
 			{
 				Translations[Index] = FGLTFConverterUtility::ConvertPosition(PosKeys[Index], Builder.ExportOptions->ExportScale);
 			}
 
-			// TODO: remove later
-			const FGLTFJsonVector3 Last = FGLTFConverterUtility::ConvertPosition(PosKeys.Last(), Builder.ExportOptions->ExportScale);
-			for (int32 Index = PosKeys.Num(); Index < FrameCount; ++Index)
-			{
-				Translations[Index] = Last;
-			}
+			JsonInputAccessor.Count = Translations.Num();
+			JsonInputAccessor.Max[0] = Timestamps[Translations.Num() - 1];
 
-			FGLTFJsonAccessor JsonAccessor;
-			JsonAccessor.BufferView = Builder.AddBufferView(Translations);
-			JsonAccessor.ComponentType = EGLTFJsonComponentType::F32;
-			JsonAccessor.Count = Translations.Num();
-			JsonAccessor.Type = EGLTFJsonAccessorType::Vec3;
+			FGLTFJsonAccessor JsonOutputAccessor;
+			JsonOutputAccessor.BufferView = Builder.AddBufferView(Translations);
+			JsonOutputAccessor.ComponentType = EGLTFJsonComponentType::F32;
+			JsonOutputAccessor.Count = Translations.Num();
+			JsonOutputAccessor.Type = EGLTFJsonAccessorType::Vec3;
 
 			FGLTFJsonAnimationSampler JsonSampler;
-			JsonSampler.Input = TimestampAccessorIndex;
-			JsonSampler.Output = Builder.AddAccessor(JsonAccessor);
+			JsonSampler.Input = Builder.AddAccessor(JsonInputAccessor);
+			JsonSampler.Output = Builder.AddAccessor(JsonOutputAccessor);
 			JsonSampler.Interpolation = Interpolation;
 
 			FGLTFJsonAnimationChannel JsonChannel;
@@ -87,29 +88,25 @@ FGLTFJsonAnimationIndex FGLTFAnimationConverter::Convert(FGLTFJsonNodeIndex Root
 		if (RotKeys.Num() > 0)
 		{
 			TArray<FGLTFJsonQuaternion> Rotations;
-			Rotations.AddUninitialized(FrameCount); // TODO: RotKeys.Num()
+			Rotations.AddUninitialized(RotKeys.Num());
 
 			for (int32 Index = 0; Index < RotKeys.Num(); ++Index)
 			{
 				Rotations[Index] = FGLTFConverterUtility::ConvertRotation(RotKeys[Index]);
 			}
 
-			// TODO: remove later
-			const FGLTFJsonQuaternion Last = FGLTFConverterUtility::ConvertRotation(RotKeys.Last());
-			for (int32 Index = RotKeys.Num(); Index < FrameCount; ++Index)
-			{
-				Rotations[Index] = Last;
-			}
+			JsonInputAccessor.Count = Rotations.Num();
+			JsonInputAccessor.Max[0] = Timestamps[Rotations.Num() - 1];
 
-			FGLTFJsonAccessor JsonAccessor;
-			JsonAccessor.BufferView = Builder.AddBufferView(Rotations);
-			JsonAccessor.ComponentType = EGLTFJsonComponentType::F32;
-			JsonAccessor.Count = Rotations.Num();
-			JsonAccessor.Type = EGLTFJsonAccessorType::Vec4;
+			FGLTFJsonAccessor JsonOutputAccessor;
+			JsonOutputAccessor.BufferView = Builder.AddBufferView(Rotations);
+			JsonOutputAccessor.ComponentType = EGLTFJsonComponentType::F32;
+			JsonOutputAccessor.Count = Rotations.Num();
+			JsonOutputAccessor.Type = EGLTFJsonAccessorType::Vec4;
 
 			FGLTFJsonAnimationSampler JsonSampler;
-			JsonSampler.Input = TimestampAccessorIndex;
-			JsonSampler.Output = Builder.AddAccessor(JsonAccessor);
+			JsonSampler.Input = Builder.AddAccessor(JsonInputAccessor);;
+			JsonSampler.Output = Builder.AddAccessor(JsonOutputAccessor);
 			JsonSampler.Interpolation = Interpolation;
 
 			FGLTFJsonAnimationChannel JsonChannel;
@@ -124,29 +121,25 @@ FGLTFJsonAnimationIndex FGLTFAnimationConverter::Convert(FGLTFJsonNodeIndex Root
 		if (ScaleKeys.Num() > 0)
 		{
 			TArray<FGLTFJsonVector3> Scales;
-			Scales.AddUninitialized(FrameCount); // TODO: ScaleKeys.Num()
+			Scales.AddUninitialized(ScaleKeys.Num());
 
 			for (int32 Index = 0; Index < ScaleKeys.Num(); ++Index)
 			{
 				Scales[Index] = FGLTFConverterUtility::ConvertScale(ScaleKeys[Index]);
 			}
 
-			// TODO: remove later
-			const FGLTFJsonVector3 Last = FGLTFConverterUtility::ConvertScale(ScaleKeys.Last());
-			for (int32 Index = ScaleKeys.Num(); Index < FrameCount; ++Index)
-			{
-				Scales[Index] = Last;
-			}
+			JsonInputAccessor.Count = Scales.Num();
+			JsonInputAccessor.Max[0] = Timestamps[Scales.Num() - 1];
 
-			FGLTFJsonAccessor JsonAccessor;
-			JsonAccessor.BufferView = Builder.AddBufferView(Scales);
-			JsonAccessor.ComponentType = EGLTFJsonComponentType::F32;
-			JsonAccessor.Count = Scales.Num();
-			JsonAccessor.Type = EGLTFJsonAccessorType::Vec3;
+			FGLTFJsonAccessor JsonOutputAccessor;
+			JsonOutputAccessor.BufferView = Builder.AddBufferView(Scales);
+			JsonOutputAccessor.ComponentType = EGLTFJsonComponentType::F32;
+			JsonOutputAccessor.Count = Scales.Num();
+			JsonOutputAccessor.Type = EGLTFJsonAccessorType::Vec3;
 
 			FGLTFJsonAnimationSampler JsonSampler;
-			JsonSampler.Input = TimestampAccessorIndex;
-			JsonSampler.Output = Builder.AddAccessor(JsonAccessor);
+			JsonSampler.Input = Builder.AddAccessor(JsonInputAccessor);
+			JsonSampler.Output = Builder.AddAccessor(JsonOutputAccessor);
 			JsonSampler.Interpolation = Interpolation;
 
 			FGLTFJsonAnimationChannel JsonChannel;
