[GLTFExporter] feat(skeletal-animations): add support for retargeting bone transforms according to each mode

commit 43e6aa67a989da4e21ea11a8b661080fd02c2d21
Date: Sat, 26 Sep 2020 03:04:39 +0200

    feat(skeletal-animations): add support for retargeting bone transforms according to each mode
    (cherry picked from commit 4f8574466c39ecebb480d37cd35c65bd45952568​)

#skipci

diff --git a/Engine/Plugins/Enterprise/GLTFExporter/MigratedFromGit/Source/GLTFExporter/Private/Converters/GLTFAnimationConverters.cpp b/Engine/Plugins/Enterprise/GLTFExporter/MigratedFromGit/Source/GLTFExporter/Private/Converters/GLTFAnimationConverters.cpp
--- a/Engine/Plugins/Enterprise/GLTFExporter/MigratedFromGit/Source/GLTFExporter/Private/Converters/GLTFAnimationConverters.cpp
+++ b/Engine/Plugins/Enterprise/GLTFExporter/MigratedFromGit/Source/GLTFExporter/Private/Converters/GLTFAnimationConverters.cpp
@@ -2,6 +2,7 @@
 
 #include "Converters/GLTFAnimationConverters.h"
 #include "Converters/GLTFConverterUtility.h"
+#include "Converters/GLTFBoneUtility.h"
 #include "Builders/GLTFConvertBuilder.h"
 #include "Animation/AnimSequence.h"
 
@@ -39,6 +40,9 @@ FGLTFJsonAnimationIndex FGLTFAnimationConverter::Convert(FGLTFJsonNodeIndex Root
 	JsonInputAccessor.MinMaxLength = 1;
 	JsonInputAccessor.Min[0] = 0;
 
+	FBoneContainer BoneContainer;
+	FGLTFBoneUtility::InitializeToSkeleton(BoneContainer, Skeleton);
+
 	const EGLTFJsonInterpolation Interpolation = FGLTFConverterUtility::ConvertInterpolation(AnimSequence->Interpolation);
 	const TArray<FName>& TrackNames = AnimSequence->GetAnimationTrackNames();
 	const int32 TrackCount = TrackNames.Num();
@@ -46,20 +50,43 @@ FGLTFJsonAnimationIndex FGLTFAnimationConverter::Convert(FGLTFJsonNodeIndex Root
 	for (int32 TrackIndex = 0; TrackIndex < TrackCount; ++TrackIndex)
 	{
 		const FRawAnimSequenceTrack& Track = AnimSequence->GetRawAnimationTrack(TrackIndex);
+		const TArray<FVector>& KeyPositions = Track.PosKeys;
+		const TArray<FQuat>& KeyRotations = Track.RotKeys;
+		const TArray<FVector>& KeyScales = Track.ScaleKeys;
+
+		const int32 MaxKeys = FMath::Max3(KeyPositions.Num(), KeyRotations.Num(), KeyScales.Num());
+		TArray<FTransform> KeyTransforms;
+		KeyTransforms.AddUninitialized(MaxKeys);
+
+		for (int32 Key = 0; Key < KeyTransforms.Num(); ++Key)
+		{
+			const FVector& Position = KeyPositions.IsValidIndex(Key) ? KeyPositions[Key] : FVector::ZeroVector;
+			const FQuat& Rotation = KeyRotations.IsValidIndex(Key) ? KeyRotations[Key] : FQuat::Identity;
+			const FVector& Scale = KeyScales.IsValidIndex(Key) ? KeyScales[Key] : FVector::OneVector;
+			KeyTransforms[Key] = { Rotation, Position, Scale };
+		}
 
 		const int32 SkeletonBoneIndex = AnimSequence->GetSkeletonIndexFromRawDataTrackIndex(TrackIndex);
 		const int32 BoneIndex = Skeleton->GetMeshBoneIndexFromSkeletonBoneIndex(SkeletalMesh, SkeletonBoneIndex);
 		const FGLTFJsonNodeIndex NodeIndex = Builder.GetOrAddNode(RootNode, SkeletalMesh, BoneIndex);
 
-		const TArray<FVector>& PosKeys = Track.PosKeys;
-		if (PosKeys.Num() > 0)
+		if (Skeleton->GetBoneTranslationRetargetingMode(SkeletonBoneIndex) != EBoneTranslationRetargetingMode::Animation)
+		{
+			for (int32 Key = 0; Key < KeyTransforms.Num(); ++Key)
+			{
+				FGLTFBoneUtility::RetargetTransform(AnimSequence, KeyTransforms[Key], SkeletonBoneIndex, BoneIndex, BoneContainer);
+			}
+		}
+
+		if (KeyPositions.Num() > 0)
 		{
 			TArray<FGLTFJsonVector3> Translations;
-			Translations.AddUninitialized(PosKeys.Num());
+			Translations.AddUninitialized(KeyPositions.Num());
 
-			for (int32 Index = 0; Index < PosKeys.Num(); ++Index)
+			for (int32 Key = 0; Key < KeyPositions.Num(); ++Key)
 			{
-				Translations[Index] = FGLTFConverterUtility::ConvertPosition(PosKeys[Index], Builder.ExportOptions->ExportScale);
+				const FVector Position = KeyTransforms[Key].GetTranslation();
+				Translations[Key] = FGLTFConverterUtility::ConvertPosition(Position, Builder.ExportOptions->ExportScale);
 			}
 
 			JsonInputAccessor.Count = Translations.Num();
@@ -80,19 +107,18 @@ FGLTFJsonAnimationIndex FGLTFAnimationConverter::Convert(FGLTFJsonNodeIndex Root
 			JsonChannel.Sampler = FGLTFJsonAnimationSamplerIndex(JsonAnimation.Samplers.Add(JsonSampler));
 			JsonChannel.Target.Path = EGLTFJsonTargetPath::Translation;
 			JsonChannel.Target.Node = NodeIndex;
-
 			JsonAnimation.Channels.Add(JsonChannel);
 		}
 
-		const TArray<FQuat>& RotKeys = Track.RotKeys;
-		if (RotKeys.Num() > 0)
+		if (KeyRotations.Num() > 0)
 		{
 			TArray<FGLTFJsonQuaternion> Rotations;
-			Rotations.AddUninitialized(RotKeys.Num());
+			Rotations.AddUninitialized(KeyRotations.Num());
 
-			for (int32 Index = 0; Index < RotKeys.Num(); ++Index)
+			for (int32 Key = 0; Key < KeyRotations.Num(); ++Key)
 			{
-				Rotations[Index] = FGLTFConverterUtility::ConvertRotation(RotKeys[Index]);
+				const FQuat Rotation = KeyTransforms[Key].GetRotation();
+				Rotations[Key] = FGLTFConverterUtility::ConvertRotation(Rotation);
 			}
 
 			JsonInputAccessor.Count = Rotations.Num();
@@ -113,19 +139,18 @@ FGLTFJsonAnimationIndex FGLTFAnimationConverter::Convert(FGLTFJsonNodeIndex Root
 			JsonChannel.Sampler = FGLTFJsonAnimationSamplerIndex(JsonAnimation.Samplers.Add(JsonSampler));
 			JsonChannel.Target.Path = EGLTFJsonTargetPath::Rotation;
 			JsonChannel.Target.Node = NodeIndex;
-
 			JsonAnimation.Channels.Add(JsonChannel);
 		}
 
-		const TArray<FVector>& ScaleKeys = Track.ScaleKeys;
-		if (ScaleKeys.Num() > 0)
+		if (KeyScales.Num() > 0)
 		{
 			TArray<FGLTFJsonVector3> Scales;
-			Scales.AddUninitialized(ScaleKeys.Num());
+			Scales.AddUninitialized(KeyScales.Num());
 
-			for (int32 Index = 0; Index < ScaleKeys.Num(); ++Index)
+			for (int32 Key = 0; Key < KeyScales.Num(); ++Key)
 			{
-				Scales[Index] = FGLTFConverterUtility::ConvertScale(ScaleKeys[Index]);
+				const FVector Scale = KeyTransforms[Key].GetScale3D();
+				Scales[Key] = FGLTFConverterUtility::ConvertScale(Scale);
 			}
 
 			JsonInputAccessor.Count = Scales.Num();
@@ -146,7 +171,6 @@ FGLTFJsonAnimationIndex FGLTFAnimationConverter::Convert(FGLTFJsonNodeIndex Root
 			JsonChannel.Sampler = FGLTFJsonAnimationSamplerIndex(JsonAnimation.Samplers.Add(JsonSampler));
 			JsonChannel.Target.Path = EGLTFJsonTargetPath::Scale;
 			JsonChannel.Target.Node = NodeIndex;
-
 			JsonAnimation.Channels.Add(JsonChannel);
 		}
 	}
diff --git a/Engine/Plugins/Enterprise/GLTFExporter/MigratedFromGit/Source/GLTFExporter/Private/Converters/GLTFBoneUtility.cpp b/Engine/Plugins/Enterprise/GLTFExporter/MigratedFromGit/Source/GLTFExporter/Private/Converters/GLTFBoneUtility.cpp
--- a/Engine/Plugins/Enterprise/GLTFExporter/MigratedFromGit/Source/GLTFExporter/Private/Converters/GLTFBoneUtility.cpp
+++ b/Engine/Plugins/Enterprise/GLTFExporter/MigratedFromGit/Source/GLTFExporter/Private/Converters/GLTFBoneUtility.cpp
@@ -3,7 +3,7 @@
 #include "Converters/GLTFBoneUtility.h"
 #include "Engine.h"
 
-FTransform FGLTFBoneUtility::GetBindTransform(const FReferenceSkeleton& RefSkeleton, const int32 BoneIndex)
+FTransform FGLTFBoneUtility::GetBindTransform(const FReferenceSkeleton& RefSkeleton, int32 BoneIndex)
 {
 	const TArray<FMeshBoneInfo>& BoneInfos = RefSkeleton.GetRefBoneInfo();
 	const TArray<FTransform>& BonePoses = RefSkeleton.GetRefBonePose();
@@ -19,3 +19,24 @@ FTransform FGLTFBoneUtility::GetBindTransform(const FReferenceSkeleton& RefSkele
 
 	return BindTransform;
 }
+
+void FGLTFBoneUtility::InitializeToSkeleton(FBoneContainer& BoneContainer, const USkeleton* Skeleton)
+{
+	TArray<FBoneIndexType> RequiredBoneIndices;
+	RequiredBoneIndices.AddUninitialized(Skeleton->GetReferenceSkeleton().GetNum());
+
+	for (int32 BoneIndex = 0; BoneIndex < RequiredBoneIndices.Num(); ++BoneIndex)
+	{
+		RequiredBoneIndices[BoneIndex] = BoneIndex;
+	}
+
+	BoneContainer.SetUseRAWData(true);
+	BoneContainer.InitializeTo(RequiredBoneIndices, FCurveEvaluationOption(true), *const_cast<USkeleton*>(Skeleton));
+}
+
+void FGLTFBoneUtility::RetargetTransform(const UAnimSequence* AnimSequence, FTransform& BoneTransform, int32 SkeletonBoneIndex, int32 BoneIndex, const FBoneContainer& RequiredBones)
+{
+	const USkeleton* Skeleton = AnimSequence->GetSkeleton();
+	const FName& RetargetSource = AnimSequence->RetargetSource;
+	FAnimationRuntime::RetargetBoneTransform(Skeleton, RetargetSource, BoneTransform, SkeletonBoneIndex, FCompactPoseBoneIndex(BoneIndex), RequiredBones, false);
+}
diff --git a/Engine/Plugins/Enterprise/GLTFExporter/MigratedFromGit/Source/GLTFExporter/Private/Converters/GLTFBoneUtility.h b/Engine/Plugins/Enterprise/GLTFExporter/MigratedFromGit/Source/GLTFExporter/Private/Converters/GLTFBoneUtility.h
--- a/Engine/Plugins/Enterprise/GLTFExporter/MigratedFromGit/Source/GLTFExporter/Private/Converters/GLTFBoneUtility.h
+++ b/Engine/Plugins/Enterprise/GLTFExporter/MigratedFromGit/Source/GLTFExporter/Private/Converters/GLTFBoneUtility.h
@@ -6,5 +6,9 @@
 
 struct FGLTFBoneUtility
 {
-	static FTransform GetBindTransform(const FReferenceSkeleton& RefSkeleton, const int32 BoneIndex);
+	static FTransform GetBindTransform(const FReferenceSkeleton& RefSkeleton, int32 BoneIndex);
+
+	static void InitializeToSkeleton(FBoneContainer& BoneContainer, const USkeleton* Skeleton);
+
+	static void RetargetTransform(const UAnimSequence* AnimSequence, FTransform& BoneTransform, int32 SkeletonBoneIndex, int32 BoneIndex, const FBoneContainer& RequiredBones);
 };
