[GLTFExporter] feat(skeletal-animations): add support for exporting animations in SkeletalMeshComponent

commit 9b065e8230f15e5bf6dd7f8b77991c58461b59a5
Date: Sat, 26 Sep 2020 20:05:26 +0200

    feat(skeletal-animations): add support for exporting animations in SkeletalMeshComponent
    (cherry picked from commit 21fe06f774db102f8553dd0b4f245d61b8207801​)

#skipci

diff --git a/Engine/Plugins/Enterprise/GLTFExporter/MigratedFromGit/Source/GLTFExporter/Private/Converters/GLTFAnimationConverters.cpp b/Engine/Plugins/Enterprise/GLTFExporter/MigratedFromGit/Source/GLTFExporter/Private/Converters/GLTFAnimationConverters.cpp
--- a/Engine/Plugins/Enterprise/GLTFExporter/MigratedFromGit/Source/GLTFExporter/Private/Converters/GLTFAnimationConverters.cpp
+++ b/Engine/Plugins/Enterprise/GLTFExporter/MigratedFromGit/Source/GLTFExporter/Private/Converters/GLTFAnimationConverters.cpp
@@ -15,7 +15,7 @@ FGLTFJsonAnimationIndex FGLTFAnimationConverter::Convert(FGLTFJsonNodeIndex Root
 		return FGLTFJsonAnimationIndex(INDEX_NONE);
 	}
 
-	USkeleton* Skeleton = AnimSequence->GetSkeleton();
+	const USkeleton* Skeleton = AnimSequence->GetSkeleton();
 	if (Skeleton == nullptr)
 	{
 		// TODO: report error
@@ -67,7 +67,7 @@ FGLTFJsonAnimationIndex FGLTFAnimationConverter::Convert(FGLTFJsonNodeIndex Root
 		}
 
 		const int32 SkeletonBoneIndex = AnimSequence->GetSkeletonIndexFromRawDataTrackIndex(TrackIndex);
-		const int32 BoneIndex = Skeleton->GetMeshBoneIndexFromSkeletonBoneIndex(SkeletalMesh, SkeletonBoneIndex);
+		const int32 BoneIndex = const_cast<USkeleton*>(Skeleton)->GetMeshBoneIndexFromSkeletonBoneIndex(SkeletalMesh, SkeletonBoneIndex);
 		const FGLTFJsonNodeIndex NodeIndex = Builder.GetOrAddNode(RootNode, SkeletalMesh, BoneIndex);
 
 		if (Skeleton->GetBoneTranslationRetargetingMode(SkeletonBoneIndex) != EBoneTranslationRetargetingMode::Animation)
diff --git a/Engine/Plugins/Enterprise/GLTFExporter/MigratedFromGit/Source/GLTFExporter/Private/Converters/GLTFNodeConverters.cpp b/Engine/Plugins/Enterprise/GLTFExporter/MigratedFromGit/Source/GLTFExporter/Private/Converters/GLTFNodeConverters.cpp
--- a/Engine/Plugins/Enterprise/GLTFExporter/MigratedFromGit/Source/GLTFExporter/Private/Converters/GLTFNodeConverters.cpp
+++ b/Engine/Plugins/Enterprise/GLTFExporter/MigratedFromGit/Source/GLTFExporter/Private/Converters/GLTFNodeConverters.cpp
@@ -80,11 +80,18 @@ FGLTFJsonNodeIndex FGLTFComponentConverter::Convert(const USceneComponent* Scene
 	{
 		Node.Mesh = Builder.GetOrAddMesh(SkeletalMeshComponent);
 
-		const USkeletalMesh* SkeletalMesh = SkeletalMeshComponent->SkeletalMesh;
-		if (SkeletalMesh != nullptr)
+		if (const USkeletalMesh* SkeletalMesh = SkeletalMeshComponent->SkeletalMesh)
 		{
 			const FGLTFJsonSkinIndex SkinIndex = Builder.GetOrAddSkin(NodeIndex, SkeletalMesh);
-			Builder.GetNode(NodeIndex).Skin = SkinIndex; // GetOrAddSkin may have invalidated Node reference
+			Builder.GetNode(NodeIndex).Skin = SkinIndex;
+
+			if (SkeletalMeshComponent->GetAnimationMode() == EAnimationMode::AnimationSingleNode)
+			{
+				if (const UAnimSequence* AnimSequence = Cast<UAnimSequence>(SkeletalMeshComponent->AnimationData.AnimToPlay))
+				{
+					Builder.GetOrAddAnimation(NodeIndex, SkeletalMesh, AnimSequence);
+				}
+			}
 		}
 	}
 	else if (const UCameraComponent* CameraComponent = Cast<UCameraComponent>(SceneComponent))
