[GLTFExporter] refactor(skeletal-animations): add error messages and minor refactoring

commit 00f1f40a813aa8dcf5d1d682e5072771e2d1d1ef
Date: Sat, 26 Sep 2020 18:09:49 +0200

    refactor(skeletal-animations): add error messages and minor refactoring
    (cherry picked from commit 9e431048a5c70d26f444d218cf2b3e88d529c93e​)

#skipci

diff --git a/Engine/Plugins/Enterprise/GLTFExporter/MigratedFromGit/Source/GLTFExporter/Private/Converters/GLTFAnimationConverters.cpp b/Engine/Plugins/Enterprise/GLTFExporter/MigratedFromGit/Source/GLTFExporter/Private/Converters/GLTFAnimationConverters.cpp
--- a/Engine/Plugins/Enterprise/GLTFExporter/MigratedFromGit/Source/GLTFExporter/Private/Converters/GLTFAnimationConverters.cpp
+++ b/Engine/Plugins/Enterprise/GLTFExporter/MigratedFromGit/Source/GLTFExporter/Private/Converters/GLTFAnimationConverters.cpp
@@ -55,15 +55,20 @@ FGLTFJsonAnimationIndex FGLTFAnimationConverter::Convert(FGLTFJsonNodeIndex Root
 		const TArray<FVector>& KeyScales = Track.ScaleKeys;
 
 		const int32 MaxKeys = FMath::Max3(KeyPositions.Num(), KeyRotations.Num(), KeyScales.Num());
+		if (MaxKeys == 0)
+		{
+			continue;
+		}
+
 		TArray<FTransform> KeyTransforms;
 		KeyTransforms.AddUninitialized(MaxKeys);
 
 		for (int32 Key = 0; Key < KeyTransforms.Num(); ++Key)
 		{
-			const FVector& Position = KeyPositions.IsValidIndex(Key) ? KeyPositions[Key] : FVector::ZeroVector;
-			const FQuat& Rotation = KeyRotations.IsValidIndex(Key) ? KeyRotations[Key] : FQuat::Identity;
-			const FVector& Scale = KeyScales.IsValidIndex(Key) ? KeyScales[Key] : FVector::OneVector;
-			KeyTransforms[Key] = { Rotation, Position, Scale };
+			const FVector& KeyPosition = KeyPositions.IsValidIndex(Key) ? KeyPositions[Key] : FVector::ZeroVector;
+			const FQuat& KeyRotation = KeyRotations.IsValidIndex(Key) ? KeyRotations[Key] : FQuat::Identity;
+			const FVector& KeyScale = KeyScales.IsValidIndex(Key) ? KeyScales[Key] : FVector::OneVector;
+			KeyTransforms[Key] = { KeyRotation, KeyPosition, KeyScale };
 		}
 
 		const int32 SkeletonBoneIndex = AnimSequence->GetSkeletonIndexFromRawDataTrackIndex(TrackIndex);
@@ -85,8 +90,8 @@ FGLTFJsonAnimationIndex FGLTFAnimationConverter::Convert(FGLTFJsonNodeIndex Root
 
 			for (int32 Key = 0; Key < KeyPositions.Num(); ++Key)
 			{
-				const FVector Position = KeyTransforms[Key].GetTranslation();
-				Translations[Key] = FGLTFConverterUtility::ConvertPosition(Position, Builder.ExportOptions->ExportScale);
+				const FVector KeyPosition = KeyTransforms[Key].GetTranslation();
+				Translations[Key] = FGLTFConverterUtility::ConvertPosition(KeyPosition, Builder.ExportOptions->ExportScale);
 			}
 
 			JsonInputAccessor.Count = Translations.Num();
@@ -117,8 +122,8 @@ FGLTFJsonAnimationIndex FGLTFAnimationConverter::Convert(FGLTFJsonNodeIndex Root
 
 			for (int32 Key = 0; Key < KeyRotations.Num(); ++Key)
 			{
-				const FQuat Rotation = KeyTransforms[Key].GetRotation();
-				Rotations[Key] = FGLTFConverterUtility::ConvertRotation(Rotation);
+				const FQuat KeyRotation = KeyTransforms[Key].GetRotation();
+				Rotations[Key] = FGLTFConverterUtility::ConvertRotation(KeyRotation);
 			}
 
 			JsonInputAccessor.Count = Rotations.Num();
@@ -149,8 +154,8 @@ FGLTFJsonAnimationIndex FGLTFAnimationConverter::Convert(FGLTFJsonNodeIndex Root
 
 			for (int32 Key = 0; Key < KeyScales.Num(); ++Key)
 			{
-				const FVector Scale = KeyTransforms[Key].GetScale3D();
-				Scales[Key] = FGLTFConverterUtility::ConvertScale(Scale);
+				const FVector KeyScale = KeyTransforms[Key].GetScale3D();
+				Scales[Key] = FGLTFConverterUtility::ConvertScale(KeyScale);
 			}
 
 			JsonInputAccessor.Count = Scales.Num();
diff --git a/Engine/Plugins/Enterprise/GLTFExporter/MigratedFromGit/Source/GLTFExporter/Private/Converters/GLTFSkinConverters.cpp b/Engine/Plugins/Enterprise/GLTFExporter/MigratedFromGit/Source/GLTFExporter/Private/Converters/GLTFSkinConverters.cpp
--- a/Engine/Plugins/Enterprise/GLTFExporter/MigratedFromGit/Source/GLTFExporter/Private/Converters/GLTFSkinConverters.cpp
+++ b/Engine/Plugins/Enterprise/GLTFExporter/MigratedFromGit/Source/GLTFExporter/Private/Converters/GLTFSkinConverters.cpp
@@ -12,28 +12,32 @@ FGLTFJsonSkinIndex FGLTFSkinConverter::Convert(FGLTFJsonNodeIndex RootNode, cons
 	Skin.Skeleton = RootNode;
 
 	const int32 BoneCount = SkeletalMesh->RefSkeleton.GetNum();
-	Skin.Joints.AddUninitialized(BoneCount);
-
-	for (int32 BoneIndex = 0; BoneIndex < BoneCount; ++BoneIndex)
+	if (BoneCount > 0)
 	{
-		Skin.Joints[BoneIndex] = Builder.GetOrAddNode(RootNode, SkeletalMesh, BoneIndex);
-	}
+		Skin.Joints.AddUninitialized(BoneCount);
 
-	TArray<FGLTFJsonMatrix4> InverseBindMatrices;
-	InverseBindMatrices.AddUninitialized(BoneCount);
+		for (int32 BoneIndex = 0; BoneIndex < BoneCount; ++BoneIndex)
+		{
+			Skin.Joints[BoneIndex] = Builder.GetOrAddNode(RootNode, SkeletalMesh, BoneIndex);
+		}
 
-	for (int32 BoneIndex = 0; BoneIndex < BoneCount; ++BoneIndex)
-	{
-		const FTransform InverseTransform = FGLTFBoneUtility::GetBindTransform(SkeletalMesh->RefSkeleton, BoneIndex).Inverse();
-		InverseBindMatrices[BoneIndex] = FGLTFConverterUtility::ConvertTransform(InverseTransform, Builder.ExportOptions->ExportScale);
-	}
+		TArray<FGLTFJsonMatrix4> InverseBindMatrices;
+		InverseBindMatrices.AddUninitialized(BoneCount);
 
-	FGLTFJsonAccessor JsonAccessor;
-	JsonAccessor.BufferView = Builder.AddBufferView(InverseBindMatrices);
-	JsonAccessor.ComponentType = EGLTFJsonComponentType::F32;
-	JsonAccessor.Count = BoneCount;
-	JsonAccessor.Type = EGLTFJsonAccessorType::Mat4;
+		for (int32 BoneIndex = 0; BoneIndex < BoneCount; ++BoneIndex)
+		{
+			const FTransform InverseBindTransform = FGLTFBoneUtility::GetBindTransform(SkeletalMesh->RefSkeleton, BoneIndex).Inverse();
+			InverseBindMatrices[BoneIndex] = FGLTFConverterUtility::ConvertTransform(InverseBindTransform, Builder.ExportOptions->ExportScale);
+		}
+
+		FGLTFJsonAccessor JsonAccessor;
+		JsonAccessor.BufferView = Builder.AddBufferView(InverseBindMatrices);
+		JsonAccessor.ComponentType = EGLTFJsonComponentType::F32;
+		JsonAccessor.Count = BoneCount;
+		JsonAccessor.Type = EGLTFJsonAccessorType::Mat4;
+
+		Skin.InverseBindMatrices = Builder.AddAccessor(JsonAccessor);
+	}
 
-	Skin.InverseBindMatrices = Builder.AddAccessor(JsonAccessor);
 	return Builder.AddSkin(Skin);
 }
diff --git a/Engine/Plugins/Enterprise/GLTFExporter/MigratedFromGit/Source/GLTFExporter/Private/Exporters/GLTFAnimSequenceExporter.cpp b/Engine/Plugins/Enterprise/GLTFExporter/MigratedFromGit/Source/GLTFExporter/Private/Exporters/GLTFAnimSequenceExporter.cpp
--- a/Engine/Plugins/Enterprise/GLTFExporter/MigratedFromGit/Source/GLTFExporter/Private/Exporters/GLTFAnimSequenceExporter.cpp
+++ b/Engine/Plugins/Enterprise/GLTFExporter/MigratedFromGit/Source/GLTFExporter/Private/Exporters/GLTFAnimSequenceExporter.cpp
@@ -23,15 +23,31 @@ bool UGLTFAnimSequenceExporter::AddObject(FGLTFContainerBuilder& Builder, const
 	}
 
 	const FGLTFJsonMeshIndex MeshIndex = Builder.GetOrAddMesh(SkeletalMesh);
+	if (MeshIndex == INDEX_NONE)
+	{
+		Builder.AddErrorMessage(FString::Printf(TEXT("Failed to export skeletal mesh %s"), *SkeletalMesh->GetName()));
+		return false;
+	}
 
 	FGLTFJsonNode Node;
 	Node.Mesh = MeshIndex;
 	const FGLTFJsonNodeIndex NodeIndex = Builder.AddNode(Node);
 
 	const FGLTFJsonSkinIndex SkinIndex = Builder.GetOrAddSkin(NodeIndex, SkeletalMesh);
+	if (SkinIndex == INDEX_NONE)
+	{
+		Builder.AddErrorMessage(FString::Printf(TEXT("Failed to export skeletal mesh skin %s"), *SkeletalMesh->GetName()));
+		return false;
+	}
+
 	Builder.GetNode(NodeIndex).Skin = SkinIndex;
 
-	Builder.GetOrAddAnimation(NodeIndex, SkeletalMesh, AnimSequence);
+	const FGLTFJsonAnimationIndex AnimationIndex =  Builder.GetOrAddAnimation(NodeIndex, SkeletalMesh, AnimSequence);
+	if (AnimationIndex == INDEX_NONE)
+	{
+		Builder.AddErrorMessage(FString::Printf(TEXT("Failed to export skeletal animation %s"), *SkeletalMesh->GetName()));
+		return false;
+	}
 
 	FGLTFJsonScene Scene;
 	Scene.Nodes.Add(NodeIndex);
