[GLTFExporter] refactor(skeletal-animations): add SkeletalMeshComponent variants for GetOrAddSkin and -Animation

commit 7bed974eac0a75b11c4571f5cdb3c270c0604477
Date: Sat, 26 Sep 2020 19:27:16 +0200

    refactor(skeletal-animations): add SkeletalMeshComponent variants for GetOrAddSkin and -Animation
    (cherry picked from commit 2e680154420cce05af49bff1a885f189f1099253​)

#skipci

diff --git a/Engine/Plugins/Enterprise/GLTFExporter/MigratedFromGit/Source/GLTFExporter/Private/Builders/GLTFConvertBuilder.cpp b/Engine/Plugins/Enterprise/GLTFExporter/MigratedFromGit/Source/GLTFExporter/Private/Builders/GLTFConvertBuilder.cpp
--- a/Engine/Plugins/Enterprise/GLTFExporter/MigratedFromGit/Source/GLTFExporter/Private/Builders/GLTFConvertBuilder.cpp
+++ b/Engine/Plugins/Enterprise/GLTFExporter/MigratedFromGit/Source/GLTFExporter/Private/Builders/GLTFConvertBuilder.cpp
@@ -229,6 +229,23 @@ FGLTFJsonSkinIndex FGLTFConvertBuilder::GetOrAddSkin(FGLTFJsonNodeIndex RootNode
 	return SkinConverter.GetOrAdd(RootNode, SkeletalMesh);
 }
 
+FGLTFJsonSkinIndex FGLTFConvertBuilder::GetOrAddSkin(FGLTFJsonNodeIndex RootNode, const USkeletalMeshComponent* SkeletalMeshComponent)
+{
+	if (RootNode == INDEX_NONE || SkeletalMeshComponent == nullptr)
+	{
+		return FGLTFJsonSkinIndex(INDEX_NONE);
+	}
+
+	const USkeletalMesh* SkeletalMesh = SkeletalMeshComponent->SkeletalMesh;
+
+	if (SkeletalMesh == nullptr)
+	{
+		return FGLTFJsonSkinIndex(INDEX_NONE);
+	}
+
+	return GetOrAddSkin(RootNode, SkeletalMesh);
+}
+
 FGLTFJsonAnimationIndex FGLTFConvertBuilder::GetOrAddAnimation(FGLTFJsonNodeIndex RootNode, const USkeletalMesh* SkeletalMesh, const UAnimSequence* AnimSequence)
 {
 	if (RootNode == INDEX_NONE || SkeletalMesh == nullptr || AnimSequence == nullptr)
@@ -239,6 +256,24 @@ FGLTFJsonAnimationIndex FGLTFConvertBuilder::GetOrAddAnimation(FGLTFJsonNodeInde
 	return AnimationConverter.GetOrAdd(RootNode, SkeletalMesh, AnimSequence);
 }
 
+FGLTFJsonAnimationIndex FGLTFConvertBuilder::GetOrAddAnimation(FGLTFJsonNodeIndex RootNode, const USkeletalMeshComponent* SkeletalMeshComponent)
+{
+	if (RootNode == INDEX_NONE || SkeletalMeshComponent == nullptr || SkeletalMeshComponent->GetAnimationMode() != EAnimationMode::AnimationSingleNode)
+	{
+		return FGLTFJsonAnimationIndex(INDEX_NONE);
+	}
+
+	const USkeletalMesh* SkeletalMesh = SkeletalMeshComponent->SkeletalMesh;
+	const UAnimSequence* AnimSequence = Cast<UAnimSequence>(SkeletalMeshComponent->AnimationData.AnimToPlay);
+
+	if (SkeletalMesh == nullptr || AnimSequence == nullptr)
+	{
+		return FGLTFJsonAnimationIndex(INDEX_NONE);
+	}
+
+	return GetOrAddAnimation(RootNode, SkeletalMesh, AnimSequence);
+}
+
 FGLTFJsonNodeIndex FGLTFConvertBuilder::GetOrAddNode(const USceneComponent* SceneComponent, FName SocketName)
 {
 	if (SceneComponent == nullptr || SocketName == NAME_None)
diff --git a/Engine/Plugins/Enterprise/GLTFExporter/MigratedFromGit/Source/GLTFExporter/Private/Builders/GLTFConvertBuilder.h b/Engine/Plugins/Enterprise/GLTFExporter/MigratedFromGit/Source/GLTFExporter/Private/Builders/GLTFConvertBuilder.h
--- a/Engine/Plugins/Enterprise/GLTFExporter/MigratedFromGit/Source/GLTFExporter/Private/Builders/GLTFConvertBuilder.h
+++ b/Engine/Plugins/Enterprise/GLTFExporter/MigratedFromGit/Source/GLTFExporter/Private/Builders/GLTFConvertBuilder.h
@@ -55,7 +55,9 @@ public:
 	FGLTFJsonNodeIndex GetOrAddNode(FGLTFJsonNodeIndex RootNode, const USkeletalMesh* SkeletalMesh, FName SocketName);
 	FGLTFJsonNodeIndex GetOrAddNode(FGLTFJsonNodeIndex RootNode, const USkeletalMesh* SkeletalMesh, int32 BoneIndex);
 	FGLTFJsonSkinIndex GetOrAddSkin(FGLTFJsonNodeIndex RootNode, const USkeletalMesh* SkeletalMesh);
+	FGLTFJsonSkinIndex GetOrAddSkin(FGLTFJsonNodeIndex RootNode, const USkeletalMeshComponent* SkeletalMeshComponent);
 	FGLTFJsonAnimationIndex GetOrAddAnimation(FGLTFJsonNodeIndex RootNode, const USkeletalMesh* SkeletalMesh, const UAnimSequence* AnimSequence);
+	FGLTFJsonAnimationIndex GetOrAddAnimation(FGLTFJsonNodeIndex RootNode, const USkeletalMeshComponent* SkeletalMeshComponent);
 
 	FGLTFJsonNodeIndex GetOrAddNode(const USceneComponent* SceneComponent, FName SocketName);
 	FGLTFJsonNodeIndex GetOrAddNode(const USceneComponent* SceneComponent);
diff --git a/Engine/Plugins/Enterprise/GLTFExporter/MigratedFromGit/Source/GLTFExporter/Private/Converters/GLTFNodeConverters.cpp b/Engine/Plugins/Enterprise/GLTFExporter/MigratedFromGit/Source/GLTFExporter/Private/Converters/GLTFNodeConverters.cpp
--- a/Engine/Plugins/Enterprise/GLTFExporter/MigratedFromGit/Source/GLTFExporter/Private/Converters/GLTFNodeConverters.cpp
+++ b/Engine/Plugins/Enterprise/GLTFExporter/MigratedFromGit/Source/GLTFExporter/Private/Converters/GLTFNodeConverters.cpp
@@ -80,18 +80,12 @@ FGLTFJsonNodeIndex FGLTFComponentConverter::Convert(const USceneComponent* Scene
 	{
 		Node.Mesh = Builder.GetOrAddMesh(SkeletalMeshComponent);
 
-		if (const USkeletalMesh* SkeletalMesh = SkeletalMeshComponent->SkeletalMesh)
+		// TODO: remove need for NodeIndex by adding support for cyclic calls in converter
+		const FGLTFJsonSkinIndex SkinIndex = Builder.GetOrAddSkin(NodeIndex, SkeletalMeshComponent);
+		if (SkinIndex != INDEX_NONE)
 		{
-			const FGLTFJsonSkinIndex SkinIndex = Builder.GetOrAddSkin(NodeIndex, SkeletalMesh);
 			Builder.GetNode(NodeIndex).Skin = SkinIndex;
-
-			if (SkeletalMeshComponent->GetAnimationMode() == EAnimationMode::AnimationSingleNode)
-			{
-				if (const UAnimSequence* AnimSequence = Cast<UAnimSequence>(SkeletalMeshComponent->AnimationData.AnimToPlay))
-				{
-					Builder.GetOrAddAnimation(NodeIndex, SkeletalMesh, AnimSequence);
-				}
-			}
+			Builder.GetOrAddAnimation(NodeIndex, SkeletalMeshComponent);
 		}
 	}
 	else if (const UCameraComponent* CameraComponent = Cast<UCameraComponent>(SceneComponent))
