[GLTFExporter] feat(skeletal-animations): add and implement GLTFExportOptions::bExportAnimationSequences

commit 5dfb85650ed4358a681f7c5128a440ce1b4a73c8
Date: Sat, 26 Sep 2020 19:50:55 +0200

    feat(skeletal-animations): add and implement GLTFExportOptions::bExportAnimationSequences
    (cherry picked from commit fe0d93f0fbbdaafbc291f062910b9e775f1e2a45​)

#skipci

diff --git a/Engine/Plugins/Enterprise/GLTFExporter/MigratedFromGit/Source/GLTFExporter/Private/Converters/GLTFNodeConverters.cpp b/Engine/Plugins/Enterprise/GLTFExporter/MigratedFromGit/Source/GLTFExporter/Private/Converters/GLTFNodeConverters.cpp
--- a/Engine/Plugins/Enterprise/GLTFExporter/MigratedFromGit/Source/GLTFExporter/Private/Converters/GLTFNodeConverters.cpp
+++ b/Engine/Plugins/Enterprise/GLTFExporter/MigratedFromGit/Source/GLTFExporter/Private/Converters/GLTFNodeConverters.cpp
@@ -85,7 +85,10 @@ FGLTFJsonNodeIndex FGLTFComponentConverter::Convert(const USceneComponent* Scene
 		if (SkinIndex != INDEX_NONE)
 		{
 			Builder.GetNode(NodeIndex).Skin = SkinIndex;
-			Builder.GetOrAddAnimation(NodeIndex, SkeletalMeshComponent);
+			if (Builder.ExportOptions->bExportAnimationSequences)
+			{
+				Builder.GetOrAddAnimation(NodeIndex, SkeletalMeshComponent);
+			}
 		}
 	}
 	else if (const UCameraComponent* CameraComponent = Cast<UCameraComponent>(SceneComponent))
diff --git a/Engine/Plugins/Enterprise/GLTFExporter/MigratedFromGit/Source/GLTFExporter/Private/Converters/GLTFSkinConverters.cpp b/Engine/Plugins/Enterprise/GLTFExporter/MigratedFromGit/Source/GLTFExporter/Private/Converters/GLTFSkinConverters.cpp
--- a/Engine/Plugins/Enterprise/GLTFExporter/MigratedFromGit/Source/GLTFExporter/Private/Converters/GLTFSkinConverters.cpp
+++ b/Engine/Plugins/Enterprise/GLTFExporter/MigratedFromGit/Source/GLTFExporter/Private/Converters/GLTFSkinConverters.cpp
@@ -12,32 +12,35 @@ FGLTFJsonSkinIndex FGLTFSkinConverter::Convert(FGLTFJsonNodeIndex RootNode, cons
 	Skin.Skeleton = RootNode;
 
 	const int32 BoneCount = SkeletalMesh->RefSkeleton.GetNum();
-	if (BoneCount > 0)
+	if (BoneCount == 0)
 	{
-		Skin.Joints.AddUninitialized(BoneCount);
-
-		for (int32 BoneIndex = 0; BoneIndex < BoneCount; ++BoneIndex)
-		{
-			Skin.Joints[BoneIndex] = Builder.GetOrAddNode(RootNode, SkeletalMesh, BoneIndex);
-		}
+		// TODO: report warning
+		return FGLTFJsonSkinIndex(INDEX_NONE);
+	}
 
-		TArray<FGLTFJsonMatrix4> InverseBindMatrices;
-		InverseBindMatrices.AddUninitialized(BoneCount);
+	Skin.Joints.AddUninitialized(BoneCount);
 
-		for (int32 BoneIndex = 0; BoneIndex < BoneCount; ++BoneIndex)
-		{
-			const FTransform InverseBindTransform = FGLTFBoneUtility::GetBindTransform(SkeletalMesh->RefSkeleton, BoneIndex).Inverse();
-			InverseBindMatrices[BoneIndex] = FGLTFConverterUtility::ConvertTransform(InverseBindTransform, Builder.ExportOptions->ExportScale);
-		}
+	for (int32 BoneIndex = 0; BoneIndex < BoneCount; ++BoneIndex)
+	{
+		Skin.Joints[BoneIndex] = Builder.GetOrAddNode(RootNode, SkeletalMesh, BoneIndex);
+	}
 
-		FGLTFJsonAccessor JsonAccessor;
-		JsonAccessor.BufferView = Builder.AddBufferView(InverseBindMatrices);
-		JsonAccessor.ComponentType = EGLTFJsonComponentType::F32;
-		JsonAccessor.Count = BoneCount;
-		JsonAccessor.Type = EGLTFJsonAccessorType::Mat4;
+	TArray<FGLTFJsonMatrix4> InverseBindMatrices;
+	InverseBindMatrices.AddUninitialized(BoneCount);
 
-		Skin.InverseBindMatrices = Builder.AddAccessor(JsonAccessor);
+	for (int32 BoneIndex = 0; BoneIndex < BoneCount; ++BoneIndex)
+	{
+		const FTransform InverseBindTransform = FGLTFBoneUtility::GetBindTransform(SkeletalMesh->RefSkeleton, BoneIndex).Inverse();
+		InverseBindMatrices[BoneIndex] = FGLTFConverterUtility::ConvertTransform(InverseBindTransform, Builder.ExportOptions->ExportScale);
 	}
 
+	FGLTFJsonAccessor JsonAccessor;
+	JsonAccessor.BufferView = Builder.AddBufferView(InverseBindMatrices);
+	JsonAccessor.ComponentType = EGLTFJsonComponentType::F32;
+	JsonAccessor.Count = BoneCount;
+	JsonAccessor.Type = EGLTFJsonAccessorType::Mat4;
+
+	Skin.InverseBindMatrices = Builder.AddAccessor(JsonAccessor);
+
 	return Builder.AddSkin(Skin);
 }
diff --git a/Engine/Plugins/Enterprise/GLTFExporter/MigratedFromGit/Source/GLTFExporter/Private/Exporters/GLTFAnimSequenceExporter.cpp b/Engine/Plugins/Enterprise/GLTFExporter/MigratedFromGit/Source/GLTFExporter/Private/Exporters/GLTFAnimSequenceExporter.cpp
--- a/Engine/Plugins/Enterprise/GLTFExporter/MigratedFromGit/Source/GLTFExporter/Private/Exporters/GLTFAnimSequenceExporter.cpp
+++ b/Engine/Plugins/Enterprise/GLTFExporter/MigratedFromGit/Source/GLTFExporter/Private/Exporters/GLTFAnimSequenceExporter.cpp
@@ -36,17 +36,24 @@ bool UGLTFAnimSequenceExporter::AddObject(FGLTFContainerBuilder& Builder, const
 	const FGLTFJsonSkinIndex SkinIndex = Builder.GetOrAddSkin(NodeIndex, SkeletalMesh);
 	if (SkinIndex == INDEX_NONE)
 	{
-		Builder.AddErrorMessage(FString::Printf(TEXT("Failed to export skeletal mesh skin %s"), *SkeletalMesh->GetName()));
+		Builder.AddErrorMessage(FString::Printf(TEXT("Failed to export bones in skeletal mesh %s"), *SkeletalMesh->GetName()));
 		return false;
 	}
 
 	Builder.GetNode(NodeIndex).Skin = SkinIndex;
 
-	const FGLTFJsonAnimationIndex AnimationIndex =  Builder.GetOrAddAnimation(NodeIndex, SkeletalMesh, AnimSequence);
-	if (AnimationIndex == INDEX_NONE)
+	if (Builder.ExportOptions->bExportAnimationSequences)
 	{
-		Builder.AddErrorMessage(FString::Printf(TEXT("Failed to export skeletal animation %s"), *SkeletalMesh->GetName()));
-		return false;
+		const FGLTFJsonAnimationIndex AnimationIndex =  Builder.GetOrAddAnimation(NodeIndex, SkeletalMesh, AnimSequence);
+		if (AnimationIndex == INDEX_NONE)
+		{
+			Builder.AddErrorMessage(FString::Printf(TEXT("Failed to export animation sequence %s"), *AnimSequence->GetName()));
+			return false;
+		}
+	}
+	else
+	{
+		Builder.AddWarningMessage(FString::Printf(TEXT("Export of animation sequences disabled by export options")));
 	}
 
 	FGLTFJsonScene Scene;
diff --git a/Engine/Plugins/Enterprise/GLTFExporter/MigratedFromGit/Source/GLTFExporter/Private/GLTFExportOptions.cpp b/Engine/Plugins/Enterprise/GLTFExporter/MigratedFromGit/Source/GLTFExporter/Private/GLTFExportOptions.cpp
--- a/Engine/Plugins/Enterprise/GLTFExporter/MigratedFromGit/Source/GLTFExporter/Private/GLTFExportOptions.cpp
+++ b/Engine/Plugins/Enterprise/GLTFExporter/MigratedFromGit/Source/GLTFExporter/Private/GLTFExportOptions.cpp
@@ -22,6 +22,7 @@ UGLTFExportOptions::UGLTFExportOptions(const FObjectInitializer& ObjectInitializ
 	bTangentDataQuantization = false;
 	NormalizeUVCoordinates = EGLTFExporterNormalizeUVCoordinates::Never;
 	DefaultLevelOfDetail = 0;
+	bExportAnimationSequences = true;
 	TextureFormat = EGLTFExporterTextureFormat::PNG;
 	bExportSourceTextures = true;
 	bExportLightmaps = true;
diff --git a/Engine/Plugins/Enterprise/GLTFExporter/MigratedFromGit/Source/GLTFExporter/Public/GLTFExportOptions.h b/Engine/Plugins/Enterprise/GLTFExporter/MigratedFromGit/Source/GLTFExporter/Public/GLTFExportOptions.h
--- a/Engine/Plugins/Enterprise/GLTFExporter/MigratedFromGit/Source/GLTFExporter/Public/GLTFExportOptions.h
+++ b/Engine/Plugins/Enterprise/GLTFExporter/MigratedFromGit/Source/GLTFExporter/Public/GLTFExportOptions.h
@@ -100,6 +100,9 @@ public:
 	UPROPERTY(EditAnywhere, BlueprintReadWrite, config, Category = Mesh)
 	int32 DefaultLevelOfDetail;
 
+	UPROPERTY(EditAnywhere, BlueprintReadWrite, config, Category = Animation)
+	uint32 bExportAnimationSequences : 1;
+
 	UPROPERTY(EditAnywhere, BlueprintReadWrite, config, Category = Texture)
 	EGLTFExporterTextureFormat TextureFormat;
 
