[GLTFExporter] feat(skeletal-animations): add and implement GLTFExportOptions::bExportVertexSkinWeights

commit 67b4c1d4270fbe9b468e764853b75d31398229b4
Date: Sat, 26 Sep 2020 19:56:21 +0200

    feat(skeletal-animations): add and implement GLTFExportOptions::bExportVertexSkinWeights
    (cherry picked from commit 9a9423a1a5776404093d43164ed986ce3440f347​)

#skipci

diff --git a/Engine/Plugins/Enterprise/GLTFExporter/MigratedFromGit/Source/GLTFExporter/Private/Converters/GLTFMeshConverters.cpp b/Engine/Plugins/Enterprise/GLTFExporter/MigratedFromGit/Source/GLTFExporter/Private/Converters/GLTFMeshConverters.cpp
--- a/Engine/Plugins/Enterprise/GLTFExporter/MigratedFromGit/Source/GLTFExporter/Private/Converters/GLTFMeshConverters.cpp
+++ b/Engine/Plugins/Enterprise/GLTFExporter/MigratedFromGit/Source/GLTFExporter/Private/Converters/GLTFMeshConverters.cpp
@@ -183,14 +183,17 @@ FGLTFJsonMeshIndex FGLTFSkeletalMeshConverter::Convert(const USkeletalMesh* Skel
 			JsonPrimitive.Attributes.TexCoords[UVIndex] = Builder.GetOrAddUVAccessor(ConvertedSection, VertexBuffer, UVIndex);
 		}
 
-		const uint32 GroupCount = (SkinWeightBuffer->GetMaxBoneInfluences() + 3) / 4;
-		JsonPrimitive.Attributes.Joints.AddUninitialized(GroupCount);
-		JsonPrimitive.Attributes.Weights.AddUninitialized(GroupCount);
-
-		for (uint32 GroupIndex = 0; GroupIndex < GroupCount; ++GroupIndex)
+		if (Builder.ExportOptions->bExportVertexSkinWeights)
 		{
-			JsonPrimitive.Attributes.Joints[GroupIndex] = Builder.GetOrAddJointAccessor(ConvertedSection, SkinWeightBuffer, GroupIndex * 4);
-			JsonPrimitive.Attributes.Weights[GroupIndex] = Builder.GetOrAddWeightAccessor(ConvertedSection, SkinWeightBuffer, GroupIndex * 4);
+			const uint32 GroupCount = (SkinWeightBuffer->GetMaxBoneInfluences() + 3) / 4;
+			JsonPrimitive.Attributes.Joints.AddUninitialized(GroupCount);
+			JsonPrimitive.Attributes.Weights.AddUninitialized(GroupCount);
+
+			for (uint32 GroupIndex = 0; GroupIndex < GroupCount; ++GroupIndex)
+			{
+				JsonPrimitive.Attributes.Joints[GroupIndex] = Builder.GetOrAddJointAccessor(ConvertedSection, SkinWeightBuffer, GroupIndex * 4);
+				JsonPrimitive.Attributes.Weights[GroupIndex] = Builder.GetOrAddWeightAccessor(ConvertedSection, SkinWeightBuffer, GroupIndex * 4);
+			}
 		}
 
 		const UMaterialInterface* Material = OverrideMaterials.GetOverride(SkeletalMesh->Materials, Section->MaterialIndex);
diff --git a/Engine/Plugins/Enterprise/GLTFExporter/MigratedFromGit/Source/GLTFExporter/Private/Converters/GLTFNodeConverters.cpp b/Engine/Plugins/Enterprise/GLTFExporter/MigratedFromGit/Source/GLTFExporter/Private/Converters/GLTFNodeConverters.cpp
--- a/Engine/Plugins/Enterprise/GLTFExporter/MigratedFromGit/Source/GLTFExporter/Private/Converters/GLTFNodeConverters.cpp
+++ b/Engine/Plugins/Enterprise/GLTFExporter/MigratedFromGit/Source/GLTFExporter/Private/Converters/GLTFNodeConverters.cpp
@@ -80,14 +80,17 @@ FGLTFJsonNodeIndex FGLTFComponentConverter::Convert(const USceneComponent* Scene
 	{
 		Node.Mesh = Builder.GetOrAddMesh(SkeletalMeshComponent);
 
-		// TODO: remove need for NodeIndex by adding support for cyclic calls in converter
-		const FGLTFJsonSkinIndex SkinIndex = Builder.GetOrAddSkin(NodeIndex, SkeletalMeshComponent);
-		if (SkinIndex != INDEX_NONE)
+		if (Builder.ExportOptions->bExportVertexSkinWeights)
 		{
-			Builder.GetNode(NodeIndex).Skin = SkinIndex;
-			if (Builder.ExportOptions->bExportAnimationSequences)
+			// TODO: remove need for NodeIndex by adding support for cyclic calls in converter
+			const FGLTFJsonSkinIndex SkinIndex = Builder.GetOrAddSkin(NodeIndex, SkeletalMeshComponent);
+			if (SkinIndex != INDEX_NONE)
 			{
-				Builder.GetOrAddAnimation(NodeIndex, SkeletalMeshComponent);
+				Builder.GetNode(NodeIndex).Skin = SkinIndex;
+				if (Builder.ExportOptions->bExportAnimationSequences)
+				{
+					Builder.GetOrAddAnimation(NodeIndex, SkeletalMeshComponent);
+				}
 			}
 		}
 	}
diff --git a/Engine/Plugins/Enterprise/GLTFExporter/MigratedFromGit/Source/GLTFExporter/Private/Exporters/GLTFAnimSequenceExporter.cpp b/Engine/Plugins/Enterprise/GLTFExporter/MigratedFromGit/Source/GLTFExporter/Private/Exporters/GLTFAnimSequenceExporter.cpp
--- a/Engine/Plugins/Enterprise/GLTFExporter/MigratedFromGit/Source/GLTFExporter/Private/Exporters/GLTFAnimSequenceExporter.cpp
+++ b/Engine/Plugins/Enterprise/GLTFExporter/MigratedFromGit/Source/GLTFExporter/Private/Exporters/GLTFAnimSequenceExporter.cpp
@@ -33,27 +33,34 @@ bool UGLTFAnimSequenceExporter::AddObject(FGLTFContainerBuilder& Builder, const
 	Node.Mesh = MeshIndex;
 	const FGLTFJsonNodeIndex NodeIndex = Builder.AddNode(Node);
 
-	const FGLTFJsonSkinIndex SkinIndex = Builder.GetOrAddSkin(NodeIndex, SkeletalMesh);
-	if (SkinIndex == INDEX_NONE)
+	if (Builder.ExportOptions->bExportVertexSkinWeights)
 	{
-		Builder.AddErrorMessage(FString::Printf(TEXT("Failed to export bones in skeletal mesh %s"), *SkeletalMesh->GetName()));
-		return false;
-	}
+		const FGLTFJsonSkinIndex SkinIndex = Builder.GetOrAddSkin(NodeIndex, SkeletalMesh);
+		if (SkinIndex == INDEX_NONE)
+		{
+			Builder.AddErrorMessage(FString::Printf(TEXT("Failed to export bones in skeletal mesh %s"), *SkeletalMesh->GetName()));
+			return false;
+		}
 
-	Builder.GetNode(NodeIndex).Skin = SkinIndex;
+		Builder.GetNode(NodeIndex).Skin = SkinIndex;
 
-	if (Builder.ExportOptions->bExportAnimationSequences)
-	{
-		const FGLTFJsonAnimationIndex AnimationIndex =  Builder.GetOrAddAnimation(NodeIndex, SkeletalMesh, AnimSequence);
-		if (AnimationIndex == INDEX_NONE)
+		if (Builder.ExportOptions->bExportAnimationSequences)
 		{
-			Builder.AddErrorMessage(FString::Printf(TEXT("Failed to export animation sequence %s"), *AnimSequence->GetName()));
-			return false;
+			const FGLTFJsonAnimationIndex AnimationIndex =  Builder.GetOrAddAnimation(NodeIndex, SkeletalMesh, AnimSequence);
+			if (AnimationIndex == INDEX_NONE)
+			{
+				Builder.AddErrorMessage(FString::Printf(TEXT("Failed to export animation sequence %s"), *AnimSequence->GetName()));
+				return false;
+			}
+		}
+		else
+		{
+			Builder.AddWarningMessage(FString::Printf(TEXT("Export of animation sequences disabled by export options")));
 		}
 	}
 	else
 	{
-		Builder.AddWarningMessage(FString::Printf(TEXT("Export of animation sequences disabled by export options")));
+		Builder.AddWarningMessage(FString::Printf(TEXT("Export of vertex skin weights disabled by export options")));
 	}
 
 	FGLTFJsonScene Scene;
diff --git a/Engine/Plugins/Enterprise/GLTFExporter/MigratedFromGit/Source/GLTFExporter/Private/GLTFExportOptions.cpp b/Engine/Plugins/Enterprise/GLTFExporter/MigratedFromGit/Source/GLTFExporter/Private/GLTFExportOptions.cpp
--- a/Engine/Plugins/Enterprise/GLTFExporter/MigratedFromGit/Source/GLTFExporter/Private/GLTFExportOptions.cpp
+++ b/Engine/Plugins/Enterprise/GLTFExporter/MigratedFromGit/Source/GLTFExporter/Private/GLTFExportOptions.cpp
@@ -22,6 +22,7 @@ UGLTFExportOptions::UGLTFExportOptions(const FObjectInitializer& ObjectInitializ
 	bTangentDataQuantization = false;
 	NormalizeUVCoordinates = EGLTFExporterNormalizeUVCoordinates::Never;
 	DefaultLevelOfDetail = 0;
+	bExportVertexSkinWeights = true;
 	bExportAnimationSequences = true;
 	TextureFormat = EGLTFExporterTextureFormat::PNG;
 	bExportSourceTextures = true;
diff --git a/Engine/Plugins/Enterprise/GLTFExporter/MigratedFromGit/Source/GLTFExporter/Public/GLTFExportOptions.h b/Engine/Plugins/Enterprise/GLTFExporter/MigratedFromGit/Source/GLTFExporter/Public/GLTFExportOptions.h
--- a/Engine/Plugins/Enterprise/GLTFExporter/MigratedFromGit/Source/GLTFExporter/Public/GLTFExportOptions.h
+++ b/Engine/Plugins/Enterprise/GLTFExporter/MigratedFromGit/Source/GLTFExporter/Public/GLTFExportOptions.h
@@ -100,6 +100,9 @@ public:
 	UPROPERTY(EditAnywhere, BlueprintReadWrite, config, Category = Mesh)
 	int32 DefaultLevelOfDetail;
 
+	UPROPERTY(EditAnywhere, BlueprintReadWrite, config, Category = Animation)
+	uint32 bExportVertexSkinWeights : 1;
+
 	UPROPERTY(EditAnywhere, BlueprintReadWrite, config, Category = Animation)
 	uint32 bExportAnimationSequences : 1;
 
