[GLTFExporter] feat(skeletal-animations): add and implement GLTFExportOptions::bRetargetBoneTransforms

commit 0fd4d99c677f32e063af7de7af1c64e06e3b21f6
Date: Sat, 26 Sep 2020 20:02:22 +0200

    feat(skeletal-animations): add and implement GLTFExportOptions::bRetargetBoneTransforms
    (cherry picked from commit 978a6e3aecd31094d1cdbf12bcb54c9d7763d218​)

#skipci

diff --git a/Engine/Plugins/Enterprise/GLTFExporter/MigratedFromGit/Source/GLTFExporter/Private/Converters/GLTFAnimationConverters.cpp b/Engine/Plugins/Enterprise/GLTFExporter/MigratedFromGit/Source/GLTFExporter/Private/Converters/GLTFAnimationConverters.cpp
--- a/Engine/Plugins/Enterprise/GLTFExporter/MigratedFromGit/Source/GLTFExporter/Private/Converters/GLTFAnimationConverters.cpp
+++ b/Engine/Plugins/Enterprise/GLTFExporter/MigratedFromGit/Source/GLTFExporter/Private/Converters/GLTFAnimationConverters.cpp
@@ -41,7 +41,10 @@ FGLTFJsonAnimationIndex FGLTFAnimationConverter::Convert(FGLTFJsonNodeIndex Root
 	JsonInputAccessor.Min[0] = 0;
 
 	FBoneContainer BoneContainer;
-	FGLTFBoneUtility::InitializeToSkeleton(BoneContainer, Skeleton);
+	if (Builder.ExportOptions->bRetargetBoneTransforms)
+	{
+		FGLTFBoneUtility::InitializeToSkeleton(BoneContainer, Skeleton);
+	}
 
 	const EGLTFJsonInterpolation Interpolation = FGLTFConverterUtility::ConvertInterpolation(AnimSequence->Interpolation);
 	const TArray<FName>& TrackNames = AnimSequence->GetAnimationTrackNames();
@@ -75,7 +78,7 @@ FGLTFJsonAnimationIndex FGLTFAnimationConverter::Convert(FGLTFJsonNodeIndex Root
 		const int32 BoneIndex = const_cast<USkeleton*>(Skeleton)->GetMeshBoneIndexFromSkeletonBoneIndex(SkeletalMesh, SkeletonBoneIndex);
 		const FGLTFJsonNodeIndex NodeIndex = Builder.GetOrAddNode(RootNode, SkeletalMesh, BoneIndex);
 
-		if (Skeleton->GetBoneTranslationRetargetingMode(SkeletonBoneIndex) != EBoneTranslationRetargetingMode::Animation)
+		if (Builder.ExportOptions->bRetargetBoneTransforms && Skeleton->GetBoneTranslationRetargetingMode(SkeletonBoneIndex) != EBoneTranslationRetargetingMode::Animation)
 		{
 			for (int32 Key = 0; Key < KeyTransforms.Num(); ++Key)
 			{
diff --git a/Engine/Plugins/Enterprise/GLTFExporter/MigratedFromGit/Source/GLTFExporter/Private/GLTFExportOptions.cpp b/Engine/Plugins/Enterprise/GLTFExporter/MigratedFromGit/Source/GLTFExporter/Private/GLTFExportOptions.cpp
--- a/Engine/Plugins/Enterprise/GLTFExporter/MigratedFromGit/Source/GLTFExporter/Private/GLTFExportOptions.cpp
+++ b/Engine/Plugins/Enterprise/GLTFExporter/MigratedFromGit/Source/GLTFExporter/Private/GLTFExportOptions.cpp
@@ -24,6 +24,7 @@ UGLTFExportOptions::UGLTFExportOptions(const FObjectInitializer& ObjectInitializ
 	DefaultLevelOfDetail = 0;
 	bExportVertexSkinWeights = true;
 	bExportAnimationSequences = true;
+	bRetargetBoneTransforms = true;
 	TextureFormat = EGLTFExporterTextureFormat::PNG;
 	bExportSourceTextures = true;
 	bExportLightmaps = true;
diff --git a/Engine/Plugins/Enterprise/GLTFExporter/MigratedFromGit/Source/GLTFExporter/Public/GLTFExportOptions.h b/Engine/Plugins/Enterprise/GLTFExporter/MigratedFromGit/Source/GLTFExporter/Public/GLTFExportOptions.h
--- a/Engine/Plugins/Enterprise/GLTFExporter/MigratedFromGit/Source/GLTFExporter/Public/GLTFExportOptions.h
+++ b/Engine/Plugins/Enterprise/GLTFExporter/MigratedFromGit/Source/GLTFExporter/Public/GLTFExportOptions.h
@@ -106,6 +106,9 @@ public:
 	UPROPERTY(EditAnywhere, BlueprintReadWrite, config, Category = Animation)
 	uint32 bExportAnimationSequences : 1;
 
+	UPROPERTY(EditAnywhere, BlueprintReadWrite, config, Category = Animation)
+	uint32 bRetargetBoneTransforms : 1;
+
 	UPROPERTY(EditAnywhere, BlueprintReadWrite, config, Category = Texture)
 	EGLTFExporterTextureFormat TextureFormat;
 
