[GLTFExporter] feat: disable some export options for animations depending on other options

commit 9316b6b63d3ac1c34c69ad5740e41bc4d48eb6d6
Date: Sat, 26 Sep 2020 21:52:44 +0200

    feat: disable some export options for animations depending on other options
    (cherry picked from commit 7490360d4adf325ad12b019849d977bf094c12ba​)

#skipci

diff --git a/Engine/Plugins/Enterprise/GLTFExporter/MigratedFromGit/Source/GLTFExporter/Private/GLTFExportOptions.cpp b/Engine/Plugins/Enterprise/GLTFExporter/MigratedFromGit/Source/GLTFExporter/Private/GLTFExportOptions.cpp
--- a/Engine/Plugins/Enterprise/GLTFExporter/MigratedFromGit/Source/GLTFExporter/Private/GLTFExportOptions.cpp
+++ b/Engine/Plugins/Enterprise/GLTFExporter/MigratedFromGit/Source/GLTFExporter/Private/GLTFExportOptions.cpp
@@ -214,6 +214,16 @@ bool UGLTFExportOptions::CanEditChange(const FProperty* InProperty) const
 		return bBakeMaterialInputs;
 	}
 
+	if (PropertyFName == GET_MEMBER_NAME_CHECKED(ThisClass, bExportAnimationSequences))
+	{
+		return bExportVertexSkinWeights;
+	}
+
+	if (PropertyFName == GET_MEMBER_NAME_CHECKED(ThisClass, bRetargetBoneTransforms))
+	{
+		return bExportVertexSkinWeights && bExportAnimationSequences;
+	}
+
 	if (PropertyFName == GET_MEMBER_NAME_CHECKED(ThisClass, bExportOrbitalCameras))
 	{
 		return bExportCameras;
