[GLTFExporter] feat(skeletal-animations): improve error messages

commit d78acc0c803ab6089e589ab04c0c74afd52d358a
Date: Sat, 26 Sep 2020 22:17:28 +0200

    feat(skeletal-animations): improve error messages
    (cherry picked from commit 3c1c8d7aa1537fad1ce2cf01e0b0c1044bd659fc​)

#skipci

diff --git a/Engine/Plugins/Enterprise/GLTFExporter/MigratedFromGit/Source/GLTFExporter/Private/Exporters/GLTFAnimSequenceExporter.cpp b/Engine/Plugins/Enterprise/GLTFExporter/MigratedFromGit/Source/GLTFExporter/Private/Exporters/GLTFAnimSequenceExporter.cpp
--- a/Engine/Plugins/Enterprise/GLTFExporter/MigratedFromGit/Source/GLTFExporter/Private/Exporters/GLTFAnimSequenceExporter.cpp
+++ b/Engine/Plugins/Enterprise/GLTFExporter/MigratedFromGit/Source/GLTFExporter/Private/Exporters/GLTFAnimSequenceExporter.cpp
@@ -18,19 +18,29 @@ bool UGLTFAnimSequenceExporter::AddObject(FGLTFContainerBuilder& Builder, const
 	const USkeletalMesh* SkeletalMesh = FGLTFExporterUtility::GetPreviewMesh(AnimSequence);
 	if (SkeletalMesh == nullptr)
 	{
-		// report error
+		Builder.AddErrorMessage(
+			FString::Printf(TEXT("Failed to export animation sequence %s because of missing preview mesh"),
+			*AnimSequence->GetName()));
 		return false;
 	}
 
-	const FGLTFJsonMeshIndex MeshIndex = Builder.GetOrAddMesh(SkeletalMesh);
-	if (MeshIndex == INDEX_NONE)
+	FGLTFJsonNode Node;
+
+	if (Builder.ExportOptions->bExportPreviewMesh)
 	{
-		Builder.AddErrorMessage(FString::Printf(TEXT("Failed to export skeletal mesh %s"), *SkeletalMesh->GetName()));
-		return false;
+		const FGLTFJsonMeshIndex MeshIndex = Builder.GetOrAddMesh(SkeletalMesh);
+		if (MeshIndex == INDEX_NONE)
+		{
+			Builder.AddErrorMessage(
+				FString::Printf(TEXT("Failed to export skeletal mesh %s for animation sequence %s"),
+				*SkeletalMesh->GetName(),
+				*AnimSequence->GetName()));
+			return false;
+		}
+
+		Node.Mesh = MeshIndex;
 	}
 
-	FGLTFJsonNode Node;
-	Node.Mesh = MeshIndex;
 	const FGLTFJsonNodeIndex NodeIndex = Builder.AddNode(Node);
 
 	if (Builder.ExportOptions->bExportVertexSkinWeights)
@@ -38,7 +48,10 @@ bool UGLTFAnimSequenceExporter::AddObject(FGLTFContainerBuilder& Builder, const
 		const FGLTFJsonSkinIndex SkinIndex = Builder.GetOrAddSkin(NodeIndex, SkeletalMesh);
 		if (SkinIndex == INDEX_NONE)
 		{
-			Builder.AddErrorMessage(FString::Printf(TEXT("Failed to export bones in skeletal mesh %s"), *SkeletalMesh->GetName()));
+			Builder.AddErrorMessage(
+				FString::Printf(TEXT("Failed to export bones in skeletal mesh %s for animation sequence %s"),
+				*SkeletalMesh->GetName(),
+				*AnimSequence->GetName()));
 			return false;
 		}
 
@@ -49,18 +62,24 @@ bool UGLTFAnimSequenceExporter::AddObject(FGLTFContainerBuilder& Builder, const
 			const FGLTFJsonAnimationIndex AnimationIndex =  Builder.GetOrAddAnimation(NodeIndex, SkeletalMesh, AnimSequence);
 			if (AnimationIndex == INDEX_NONE)
 			{
-				Builder.AddErrorMessage(FString::Printf(TEXT("Failed to export animation sequence %s"), *AnimSequence->GetName()));
+				Builder.AddErrorMessage(
+					FString::Printf(TEXT("Failed to export animation sequence %s"),
+					*AnimSequence->GetName()));
 				return false;
 			}
 		}
 		else
 		{
-			Builder.AddWarningMessage(FString::Printf(TEXT("Export of animation sequences disabled by export options")));
+			Builder.AddWarningMessage(
+				FString::Printf(TEXT("Failed to export animation sequence %s because animation sequences are disabled by export options"),
+				*AnimSequence->GetName()));
 		}
 	}
 	else
 	{
-		Builder.AddWarningMessage(FString::Printf(TEXT("Export of vertex skin weights disabled by export options")));
+		Builder.AddWarningMessage(
+			FString::Printf(TEXT("Failed to export animation sequence %s because vertex skin weights are disabled by export options"),
+			*AnimSequence->GetName()));
 	}
 
 	FGLTFJsonScene Scene;
diff --git a/Engine/Plugins/Enterprise/GLTFExporter/MigratedFromGit/Source/GLTFExporter/Private/Exporters/GLTFSkeletalMeshExporter.cpp b/Engine/Plugins/Enterprise/GLTFExporter/MigratedFromGit/Source/GLTFExporter/Private/Exporters/GLTFSkeletalMeshExporter.cpp
--- a/Engine/Plugins/Enterprise/GLTFExporter/MigratedFromGit/Source/GLTFExporter/Private/Exporters/GLTFSkeletalMeshExporter.cpp
+++ b/Engine/Plugins/Enterprise/GLTFExporter/MigratedFromGit/Source/GLTFExporter/Private/Exporters/GLTFSkeletalMeshExporter.cpp
@@ -13,12 +13,29 @@ UGLTFSkeletalMeshExporter::UGLTFSkeletalMeshExporter(const FObjectInitializer& O
 bool UGLTFSkeletalMeshExporter::AddObject(FGLTFContainerBuilder& Builder, const UObject* Object)
 {
 	const USkeletalMesh* SkeletalMesh = CastChecked<USkeletalMesh>(Object);
+
 	const FGLTFJsonMeshIndex MeshIndex = Builder.GetOrAddMesh(SkeletalMesh);
+	if (MeshIndex == INDEX_NONE)
+	{
+		Builder.AddErrorMessage(FString::Printf(TEXT("Failed to export skeletal mesh %s"), *SkeletalMesh->GetName()));
+		return false;
+	}
 
 	FGLTFJsonNode Node;
 	Node.Mesh = MeshIndex;
 	const FGLTFJsonNodeIndex NodeIndex = Builder.AddNode(Node);
-	Builder.GetNode(NodeIndex).Skin = Builder.GetOrAddSkin(NodeIndex, SkeletalMesh);
+
+	if (Builder.ExportOptions->bExportVertexSkinWeights)
+	{
+		const FGLTFJsonSkinIndex SkinIndex = Builder.GetOrAddSkin(NodeIndex, SkeletalMesh);
+		if (SkinIndex == INDEX_NONE)
+		{
+			Builder.AddErrorMessage(FString::Printf(TEXT("Failed to export bones in skeletal mesh %s"), *SkeletalMesh->GetName()));
+			return false;
+		}
+
+		Builder.GetNode(NodeIndex).Skin = SkinIndex;
+	}
 
 	FGLTFJsonScene Scene;
 	Scene.Nodes.Add(NodeIndex);
