[GLTFExporter] feat(skeletal-animations): add extension EPIC_animation_play_data to convert FSingleAnimationPlayData

commit 242a5563f14af361848db7793aada0013683927b
Date: Sat, 26 Sep 2020 23:58:56 +0200

    feat(skeletal-animations): add extension EPIC_animation_play_data to convert FSingleAnimationPlayData
    (cherry picked from commit 010da938888517a40c0a388cfe87772a5cb10543​)

#skipci

diff --git a/Engine/Plugins/Enterprise/GLTFExporter/MigratedFromGit/Source/GLTFExporter/Private/Json/GLTFJsonAnimation.h b/Engine/Plugins/Enterprise/GLTFExporter/MigratedFromGit/Source/GLTFExporter/Private/Json/GLTFJsonAnimation.h
--- a/Engine/Plugins/Enterprise/GLTFExporter/MigratedFromGit/Source/GLTFExporter/Private/Json/GLTFJsonAnimation.h
+++ b/Engine/Plugins/Enterprise/GLTFExporter/MigratedFromGit/Source/GLTFExporter/Private/Json/GLTFJsonAnimation.h
@@ -4,6 +4,7 @@
 
 #include "Json/GLTFJsonEnums.h"
 #include "Json/GLTFJsonIndex.h"
+#include "Json/GLTFJsonPlayData.h"
 #include "Json/GLTFJsonUtility.h"
 #include "Serialization/JsonSerializer.h"
 
@@ -84,6 +85,8 @@ struct FGLTFJsonAnimation
 	TArray<FGLTFJsonAnimationChannel> Channels;
 	TArray<FGLTFJsonAnimationSampler> Samplers;
 
+	FGLTFJsonPlayData PlayData;
+
 	template <class CharType = TCHAR, class PrintPolicy = TPrettyJsonPrintPolicy<CharType>>
 	void WriteObject(TJsonWriter<CharType, PrintPolicy>& JsonWriter, FGLTFJsonExtensions& Extensions) const
 	{
@@ -97,6 +100,17 @@ struct FGLTFJsonAnimation
 		FGLTFJsonUtility::WriteObjectArray(JsonWriter, TEXT("channels"), Channels, Extensions, true);
 		FGLTFJsonUtility::WriteObjectArray(JsonWriter, TEXT("samplers"), Samplers, Extensions, true);
 
+		if (PlayData != FGLTFJsonPlayData())
+		{
+			const EGLTFJsonExtension Extension = EGLTFJsonExtension::EPIC_AnimationPlayData;
+			Extensions.Used.Add(Extension);
+
+			JsonWriter.WriteObjectStart(TEXT("extensions"));
+			JsonWriter.WriteIdentifierPrefix(FGLTFJsonUtility::ToString(Extension));
+			PlayData.WriteObject(JsonWriter, Extensions);
+			JsonWriter.WriteObjectEnd();
+		}
+
 		JsonWriter.WriteObjectEnd();
 	}
 };
diff --git a/Engine/Plugins/Enterprise/GLTFExporter/MigratedFromGit/Source/GLTFExporter/Private/Json/GLTFJsonEnums.h b/Engine/Plugins/Enterprise/GLTFExporter/MigratedFromGit/Source/GLTFExporter/Private/Json/GLTFJsonEnums.h
--- a/Engine/Plugins/Enterprise/GLTFExporter/MigratedFromGit/Source/GLTFExporter/Private/Json/GLTFJsonEnums.h
+++ b/Engine/Plugins/Enterprise/GLTFExporter/MigratedFromGit/Source/GLTFExporter/Private/Json/GLTFJsonEnums.h
@@ -16,7 +16,8 @@ enum class EGLTFJsonExtension
 	EPIC_LightmapTextures,
 	EPIC_ReflectionCaptures,
 	EPIC_LevelVariantSets,
-	EPIC_TextureHDREncoding
+	EPIC_TextureHDREncoding,
+    EPIC_AnimationPlayData
 };
 
 enum class EGLTFJsonShadingModel
diff --git a/Engine/Plugins/Enterprise/GLTFExporter/MigratedFromGit/Source/GLTFExporter/Private/Json/GLTFJsonPlayData.h b/Engine/Plugins/Enterprise/GLTFExporter/MigratedFromGit/Source/GLTFExporter/Private/Json/GLTFJsonPlayData.h
new file mode 100644
--- /dev/null
+++ b/Engine/Plugins/Enterprise/GLTFExporter/MigratedFromGit/Source/GLTFExporter/Private/Json/GLTFJsonPlayData.h
@@ -0,0 +1,71 @@
+// Copyright Epic Games, Inc. All Rights Reserved.
+
+#pragma once
+
+#include "Json/GLTFJsonUtility.h"
+#include "Serialization/JsonSerializer.h"
+
+struct FGLTFJsonPlayData
+{
+	FString Name;
+
+	bool Looping;
+	bool Playing;
+
+	float PlayRate;
+	float Position;
+
+	FGLTFJsonPlayData()
+        : Looping(true)
+        , Playing(true)
+        , PlayRate(1)
+        , Position(0)
+	{
+	}
+
+	template <class CharType = TCHAR, class PrintPolicy = TPrettyJsonPrintPolicy<CharType>>
+	void WriteObject(TJsonWriter<CharType, PrintPolicy>& JsonWriter, FGLTFJsonExtensions& Extensions) const
+	{
+		JsonWriter.WriteObjectStart();
+
+		if (!Name.IsEmpty())
+		{
+			JsonWriter.WriteValue(TEXT("name"), Name);
+		}
+
+		if (Looping != true)
+		{
+			JsonWriter.WriteValue(TEXT("looping"), Looping);
+		}
+
+		if (Playing != true)
+		{
+			JsonWriter.WriteValue(TEXT("playing"), Playing);
+		}
+
+		if (PlayRate != 1)
+		{
+			FGLTFJsonUtility::WriteExactValue(JsonWriter, TEXT("playRate"), PlayRate);
+		}
+
+		if (Position != 0)
+		{
+			FGLTFJsonUtility::WriteExactValue(JsonWriter, TEXT("position"), Position);
+		}
+
+		JsonWriter.WriteObjectEnd();
+	}
+
+	bool operator==(const FGLTFJsonPlayData& Other) const
+	{
+		return Looping == Other.Looping
+            && Playing == Other.Playing
+            && PlayRate == Other.PlayRate
+            && Position == Other.Position;
+	}
+
+	bool operator!=(const FGLTFJsonPlayData& Other) const
+	{
+		return !(*this == Other);
+	}
+};
diff --git a/Engine/Plugins/Enterprise/GLTFExporter/MigratedFromGit/Source/GLTFExporter/Private/Json/GLTFJsonUtility.h b/Engine/Plugins/Enterprise/GLTFExporter/MigratedFromGit/Source/GLTFExporter/Private/Json/GLTFJsonUtility.h
--- a/Engine/Plugins/Enterprise/GLTFExporter/MigratedFromGit/Source/GLTFExporter/Private/Json/GLTFJsonUtility.h
+++ b/Engine/Plugins/Enterprise/GLTFExporter/MigratedFromGit/Source/GLTFExporter/Private/Json/GLTFJsonUtility.h
@@ -30,6 +30,7 @@ struct FGLTFJsonUtility
 			case EGLTFJsonExtension::EPIC_ReflectionCaptures:  return TEXT("EPIC_reflection_captures");
 			case EGLTFJsonExtension::EPIC_LevelVariantSets:    return TEXT("EPIC_level_variant_sets");
 			case EGLTFJsonExtension::EPIC_TextureHDREncoding:  return TEXT("EPIC_texture_hdr_encoding");
+			case EGLTFJsonExtension::EPIC_AnimationPlayData:   return TEXT("EPIC_animation_play_data");
 			default:                                           return TEXT("unknown");
 		}
 	}
