[GLTFExporter] feat(skeletal-animations): implement extension EPIC_animation_play_data

commit 550afacbb5b669b714512c299ae855970f9ca49b
Date: Sat, 26 Sep 2020 23:59:49 +0200

    feat(skeletal-animations): implement extension EPIC_animation_play_data
    (cherry picked from commit 44fe3f4b18f84f8d212f3c3f2be5839038afce8b​)

#skipci

diff --git a/Engine/Plugins/Enterprise/GLTFExporter/MigratedFromGit/Source/GLTFExporter/Private/Builders/GLTFConvertBuilder.cpp b/Engine/Plugins/Enterprise/GLTFExporter/MigratedFromGit/Source/GLTFExporter/Private/Builders/GLTFConvertBuilder.cpp
--- a/Engine/Plugins/Enterprise/GLTFExporter/MigratedFromGit/Source/GLTFExporter/Private/Builders/GLTFConvertBuilder.cpp
+++ b/Engine/Plugins/Enterprise/GLTFExporter/MigratedFromGit/Source/GLTFExporter/Private/Builders/GLTFConvertBuilder.cpp
@@ -258,20 +258,12 @@ FGLTFJsonAnimationIndex FGLTFConvertBuilder::GetOrAddAnimation(FGLTFJsonNodeInde
 
 FGLTFJsonAnimationIndex FGLTFConvertBuilder::GetOrAddAnimation(FGLTFJsonNodeIndex RootNode, const USkeletalMeshComponent* SkeletalMeshComponent)
 {
-	if (RootNode == INDEX_NONE || SkeletalMeshComponent == nullptr || SkeletalMeshComponent->GetAnimationMode() != EAnimationMode::AnimationSingleNode)
-	{
-		return FGLTFJsonAnimationIndex(INDEX_NONE);
-	}
-
-	const USkeletalMesh* SkeletalMesh = SkeletalMeshComponent->SkeletalMesh;
-	const UAnimSequence* AnimSequence = Cast<UAnimSequence>(SkeletalMeshComponent->AnimationData.AnimToPlay);
-
-	if (SkeletalMesh == nullptr || AnimSequence == nullptr)
+	if (RootNode == INDEX_NONE || SkeletalMeshComponent == nullptr)
 	{
 		return FGLTFJsonAnimationIndex(INDEX_NONE);
 	}
 
-	return GetOrAddAnimation(RootNode, SkeletalMesh, AnimSequence);
+	return AnimationDataConverter.GetOrAdd(RootNode, SkeletalMeshComponent);
 }
 
 FGLTFJsonNodeIndex FGLTFConvertBuilder::GetOrAddNode(const USceneComponent* SceneComponent, FName SocketName)
diff --git a/Engine/Plugins/Enterprise/GLTFExporter/MigratedFromGit/Source/GLTFExporter/Private/Builders/GLTFConvertBuilder.h b/Engine/Plugins/Enterprise/GLTFExporter/MigratedFromGit/Source/GLTFExporter/Private/Builders/GLTFConvertBuilder.h
--- a/Engine/Plugins/Enterprise/GLTFExporter/MigratedFromGit/Source/GLTFExporter/Private/Builders/GLTFConvertBuilder.h
+++ b/Engine/Plugins/Enterprise/GLTFExporter/MigratedFromGit/Source/GLTFExporter/Private/Builders/GLTFConvertBuilder.h
@@ -97,6 +97,7 @@ private:
 	FGLTFSkeletalBoneConverter SkeletalBoneConverter{ *this };
 	FGLTFSkinConverter SkinConverter{ *this };
 	FGLTFAnimationConverter AnimationConverter{ *this };
+	FGLTFAnimationDataConverter AnimationDataConverter{ *this };
 
 	FGLTFComponentSocketConverter ComponentSocketConverter{ *this };
 	FGLTFComponentConverter ComponentConverter{ *this };
diff --git a/Engine/Plugins/Enterprise/GLTFExporter/MigratedFromGit/Source/GLTFExporter/Private/Converters/GLTFAnimationConverters.cpp b/Engine/Plugins/Enterprise/GLTFExporter/MigratedFromGit/Source/GLTFExporter/Private/Converters/GLTFAnimationConverters.cpp
--- a/Engine/Plugins/Enterprise/GLTFExporter/MigratedFromGit/Source/GLTFExporter/Private/Converters/GLTFAnimationConverters.cpp
+++ b/Engine/Plugins/Enterprise/GLTFExporter/MigratedFromGit/Source/GLTFExporter/Private/Converters/GLTFAnimationConverters.cpp
@@ -2,6 +2,7 @@
 
 #include "Converters/GLTFAnimationConverters.h"
 #include "Converters/GLTFConverterUtility.h"
+#include "Converters/GLTFNameUtility.h"
 #include "Converters/GLTFBoneUtility.h"
 #include "Builders/GLTFConvertBuilder.h"
 #include "Animation/AnimSequence.h"
@@ -185,3 +186,29 @@ FGLTFJsonAnimationIndex FGLTFAnimationConverter::Convert(FGLTFJsonNodeIndex Root
 
 	return Builder.AddAnimation(JsonAnimation);
 }
+
+FGLTFJsonAnimationIndex FGLTFAnimationDataConverter::Convert(FGLTFJsonNodeIndex RootNode, const USkeletalMeshComponent* SkeletalMeshComponent)
+{
+	const USkeletalMesh* SkeletalMesh = SkeletalMeshComponent->SkeletalMesh;
+	const UAnimSequence* AnimSequence = Cast<UAnimSequence>(SkeletalMeshComponent->AnimationData.AnimToPlay);
+
+	if (SkeletalMesh == nullptr || AnimSequence == nullptr || SkeletalMeshComponent->GetAnimationMode() != EAnimationMode::AnimationSingleNode)
+	{
+		return FGLTFJsonAnimationIndex(INDEX_NONE);
+	}
+
+	const FGLTFJsonAnimationIndex AnimationIndex = Builder.GetOrAddAnimation(RootNode, SkeletalMesh, AnimSequence);
+	if (AnimationIndex != INDEX_NONE)
+	{
+		FGLTFJsonAnimation& JsonAnimation = Builder.GetAnimation(AnimationIndex);
+		JsonAnimation.Name = FGLTFNameUtility::GetName(SkeletalMeshComponent);
+
+		FGLTFJsonPlayData& JsonPlayData = JsonAnimation.PlayData;
+		JsonPlayData.Looping = SkeletalMeshComponent->AnimationData.bSavedLooping;
+		JsonPlayData.Playing = SkeletalMeshComponent->AnimationData.bSavedPlaying;
+		JsonPlayData.PlayRate = SkeletalMeshComponent->AnimationData.SavedPlayRate;
+		JsonPlayData.Position = SkeletalMeshComponent->AnimationData.SavedPosition;
+	}
+
+	return AnimationIndex;
+}
diff --git a/Engine/Plugins/Enterprise/GLTFExporter/MigratedFromGit/Source/GLTFExporter/Private/Converters/GLTFAnimationConverters.h b/Engine/Plugins/Enterprise/GLTFExporter/MigratedFromGit/Source/GLTFExporter/Private/Converters/GLTFAnimationConverters.h
--- a/Engine/Plugins/Enterprise/GLTFExporter/MigratedFromGit/Source/GLTFExporter/Private/Converters/GLTFAnimationConverters.h
+++ b/Engine/Plugins/Enterprise/GLTFExporter/MigratedFromGit/Source/GLTFExporter/Private/Converters/GLTFAnimationConverters.h
@@ -7,11 +7,22 @@
 #include "Converters/GLTFBuilderContext.h"
 #include "Engine.h"
 
-class UAnimSequence;
-
-class FGLTFAnimationConverter final : public FGLTFBuilderContext, public TGLTFConverter<FGLTFJsonAnimationIndex, FGLTFJsonNodeIndex, const USkeletalMesh*, const UAnimSequence*>
+template <typename... InputTypes>
+class TGLTFAnimationConverter : public FGLTFBuilderContext, public TGLTFConverter<FGLTFJsonAnimationIndex, InputTypes...>
 {
 	using FGLTFBuilderContext::FGLTFBuilderContext;
+};
+
+class FGLTFAnimationConverter final : public TGLTFAnimationConverter<FGLTFJsonNodeIndex, const USkeletalMesh*, const UAnimSequence*>
+{
+	using TGLTFAnimationConverter::TGLTFAnimationConverter;
 
 	virtual FGLTFJsonAnimationIndex Convert(FGLTFJsonNodeIndex RootNode, const USkeletalMesh* SkeletalMesh, const UAnimSequence* AnimSequence) override;
 };
+
+class FGLTFAnimationDataConverter final : public TGLTFAnimationConverter<FGLTFJsonNodeIndex, const USkeletalMeshComponent*>
+{
+	using TGLTFAnimationConverter::TGLTFAnimationConverter;
+
+	virtual FGLTFJsonAnimationIndex Convert(FGLTFJsonNodeIndex RootNode, const USkeletalMeshComponent* SkeletalMeshComponent) override;
+};
