[GLTFExporter] feat(skeletal-animations): add support for animations in hotspot converter

commit e05a425e525a87604bebbfa2f705df9db91fab05
Date: Sun, 27 Sep 2020 00:32:35 +0200

    feat(skeletal-animations): add support for animations in hotspot converter
    (cherry picked from commit 495315e92bef3ac32aacf1c5b65d3d0380edd2ee​)

#skipci

diff --git a/Engine/Plugins/Enterprise/GLTFExporter/MigratedFromGit/Source/GLTFExporter/Private/Converters/GLTFHotspotConverters.cpp b/Engine/Plugins/Enterprise/GLTFExporter/MigratedFromGit/Source/GLTFExporter/Private/Converters/GLTFHotspotConverters.cpp
--- a/Engine/Plugins/Enterprise/GLTFExporter/MigratedFromGit/Source/GLTFExporter/Private/Converters/GLTFHotspotConverters.cpp
+++ b/Engine/Plugins/Enterprise/GLTFExporter/MigratedFromGit/Source/GLTFExporter/Private/Converters/GLTFHotspotConverters.cpp
@@ -3,14 +3,44 @@
 #include "GLTFHotspotConverters.h"
 #include "Json/FGLTFJsonHotspot.h"
 #include "Builders/GLTFContainerBuilder.h"
+#include "Animation/SkeletalMeshActor.h"
 
 FGLTFJsonHotspotIndex FGLTFHotspotConverter::Convert(const AGLTFInteractionHotspotActor* HotspotActor)
 {
-	// TODO: should we warn and / or return INDEX_NONE if the hotspot has no valid Image assigned?
-
 	FGLTFJsonHotspot JsonHotspot;
 	HotspotActor->GetName(JsonHotspot.Name);
-	JsonHotspot.Animation = FGLTFJsonAnimationIndex(INDEX_NONE);	// TODO: assign the animation property once we have support for exporting animations
+
+	if (const ASkeletalMeshActor* SkeletalMeshActor = HotspotActor->SkeletalMeshActor)
+	{
+		const FGLTFJsonNodeIndex RootNode = Builder.GetOrAddNode(SkeletalMeshActor);
+		if (const USkeletalMesh* SkeletalMesh = SkeletalMeshActor->GetSkeletalMeshComponent()->SkeletalMesh)
+		{
+			if (const UAnimSequence* AnimSequence = HotspotActor->AnimationSequence)
+			{
+				if (SkeletalMesh->Skeleton == AnimSequence->GetSkeleton())
+				{
+					JsonHotspot.Animation = Builder.GetOrAddAnimation(RootNode, SkeletalMesh, AnimSequence);
+				}
+				else
+				{
+					// TODO: report warning
+				}
+			}
+			else
+			{
+				// TODO: report warning
+			}
+		}
+		else
+		{
+			// TODO: report warning
+		}
+	}
+	else
+	{
+		// TODO: report warning
+	}
+
 	JsonHotspot.Image = Builder.GetOrAddTexture(HotspotActor->Image);
 	JsonHotspot.HoveredImage = Builder.GetOrAddTexture(HotspotActor->HoveredImage);
 	JsonHotspot.ToggledImage = Builder.GetOrAddTexture(HotspotActor->ToggledImage);
