[GLTFExporter] refactor(skeletal-animations): add warning messages when exporting hotspot if some export options are disabled

commit adbf4eee6c7ed77112a0b3ff9351573de14ccb24
Date: Sun, 27 Sep 2020 00:44:02 +0200

    refactor(skeletal-animations): add warning messages when exporting hotspot if some export options are disabled
    (cherry picked from commit 3da9545244ec98495807adbc898ac0eae534e307​)

#skipci

diff --git a/Engine/Plugins/Enterprise/GLTFExporter/MigratedFromGit/Source/GLTFExporter/Private/Converters/GLTFHotspotConverters.cpp b/Engine/Plugins/Enterprise/GLTFExporter/MigratedFromGit/Source/GLTFExporter/Private/Converters/GLTFHotspotConverters.cpp
--- a/Engine/Plugins/Enterprise/GLTFExporter/MigratedFromGit/Source/GLTFExporter/Private/Converters/GLTFHotspotConverters.cpp
+++ b/Engine/Plugins/Enterprise/GLTFExporter/MigratedFromGit/Source/GLTFExporter/Private/Converters/GLTFHotspotConverters.cpp
@@ -10,16 +10,35 @@ FGLTFJsonHotspotIndex FGLTFHotspotConverter::Convert(const AGLTFInteractionHotsp
 	FGLTFJsonHotspot JsonHotspot;
 	HotspotActor->GetName(JsonHotspot.Name);
 
-	if (const ASkeletalMeshActor* SkeletalMeshActor = HotspotActor->SkeletalMeshActor)
+	if (!Builder.ExportOptions->bExportVertexSkinWeights)
 	{
-		const FGLTFJsonNodeIndex RootNode = Builder.GetOrAddNode(SkeletalMeshActor);
-		if (const USkeletalMesh* SkeletalMesh = SkeletalMeshActor->GetSkeletalMeshComponent()->SkeletalMesh)
+		Builder.AddWarningMessage(
+			FString::Printf(TEXT("Can't export animation in hotspot %s because vertex skin weights are disabled by export options"),
+			*JsonHotspot.Name));
+	}
+	else if (!Builder.ExportOptions->bExportAnimationSequences)
+	{
+		Builder.AddWarningMessage(
+			FString::Printf(TEXT("Can't export animation in hotspot %s because animation sequences are disabled by export options"),
+			*JsonHotspot.Name));
+	}
+	else
+	{
+		if (const ASkeletalMeshActor* SkeletalMeshActor = HotspotActor->SkeletalMeshActor)
 		{
-			if (const UAnimSequence* AnimSequence = HotspotActor->AnimationSequence)
+			const FGLTFJsonNodeIndex RootNode = Builder.GetOrAddNode(SkeletalMeshActor);
+			if (const USkeletalMesh* SkeletalMesh = SkeletalMeshActor->GetSkeletalMeshComponent()->SkeletalMesh)
 			{
-				if (SkeletalMesh->Skeleton == AnimSequence->GetSkeleton())
+				if (const UAnimSequence* AnimSequence = HotspotActor->AnimationSequence)
 				{
-					JsonHotspot.Animation = Builder.GetOrAddAnimation(RootNode, SkeletalMesh, AnimSequence);
+					if (SkeletalMesh->Skeleton == AnimSequence->GetSkeleton())
+					{
+						JsonHotspot.Animation = Builder.GetOrAddAnimation(RootNode, SkeletalMesh, AnimSequence);
+					}
+					else
+					{
+						// TODO: report warning
+					}
 				}
 				else
 				{
@@ -36,10 +55,6 @@ FGLTFJsonHotspotIndex FGLTFHotspotConverter::Convert(const AGLTFInteractionHotsp
 			// TODO: report warning
 		}
 	}
-	else
-	{
-		// TODO: report warning
-	}
 
 	JsonHotspot.Image = Builder.GetOrAddTexture(HotspotActor->Image);
 	JsonHotspot.HoveredImage = Builder.GetOrAddTexture(HotspotActor->HoveredImage);
