[GLTFExporter] refactor(skeletal-mesh): add some sanitation of vertex data for skin weights

commit b9834a8f25ca06cd402732826fba6abdc344b006
Date: Mon, 28 Sep 2020 13:15:59 +0200

    refactor(skeletal-mesh): add some sanitation of vertex data for skin weights
    (cherry picked from commit c35649c744e240f2eab85eab3ad3da44d04f832a​)

#skipci

diff --git a/Engine/Plugins/Enterprise/GLTFExporter/MigratedFromGit/Source/GLTFExporter/Private/Converters/GLTFMeshConverters.cpp b/Engine/Plugins/Enterprise/GLTFExporter/MigratedFromGit/Source/GLTFExporter/Private/Converters/GLTFMeshConverters.cpp
--- a/Engine/Plugins/Enterprise/GLTFExporter/MigratedFromGit/Source/GLTFExporter/Private/Converters/GLTFMeshConverters.cpp
+++ b/Engine/Plugins/Enterprise/GLTFExporter/MigratedFromGit/Source/GLTFExporter/Private/Converters/GLTFMeshConverters.cpp
@@ -108,6 +108,11 @@ void FGLTFSkeletalMeshConverter::Sanitize(const USkeletalMesh*& SkeletalMesh, in
 		OverrideVertexColors = nullptr;
 	}
 
+	if (!Builder.ExportOptions->bExportVertexSkinWeights)
+	{
+		OverrideSkinWeights = nullptr;
+	}
+
 	if (OverrideMaterials == SkeletalMesh->Materials)
 	{
 		OverrideMaterials.Empty();
