[GLTFExporter] refactor: rename GLTFExportOptions::bTangentDataQuantization to bUseMeshQuantization and enable it by default

commit ec5e435a86fbb7a7c493b0f98a50ae5d31ff811b
Date: Sun, 27 Sep 2020 21:55:49 +0200

    refactor: rename GLTFExportOptions::bTangentDataQuantization to bUseMeshQuantization and enable it by default
    (cherry picked from commit ab86619633a181f3f1cd98afc6ea629b8b23260d​)

#skipci

diff --git a/Engine/Plugins/Enterprise/GLTFExporter/MigratedFromGit/Source/GLTFExporter/Private/GLTFExportOptions.cpp b/Engine/Plugins/Enterprise/GLTFExporter/MigratedFromGit/Source/GLTFExporter/Private/GLTFExportOptions.cpp
--- a/Engine/Plugins/Enterprise/GLTFExporter/MigratedFromGit/Source/GLTFExporter/Private/GLTFExportOptions.cpp
+++ b/Engine/Plugins/Enterprise/GLTFExporter/MigratedFromGit/Source/GLTFExporter/Private/GLTFExportOptions.cpp
@@ -19,7 +19,7 @@ UGLTFExportOptions::UGLTFExportOptions(const FObjectInitializer& ObjectInitializ
 	DefaultMaterialBakeSize = EGLTFExporterMaterialBakeSize::POT_1024;
 	bExportVertexColors = true;
 	NormalExportMethod = EGLTFExporterNormalExportMethod::NormalsAndTangents;
-	bTangentDataQuantization = false;
+	bUseMeshQuantization = true;
 	NormalizeUVCoordinates = EGLTFExporterNormalizeUVCoordinates::Never;
 	DefaultLevelOfDetail = 0;
 	bExportVertexSkinWeights = true;
@@ -197,7 +197,6 @@ bool UGLTFExportOptions::CanEditChange(const FProperty* InProperty) const
 
 	// TODO: remove options that have been implemented in exporters/converters.
 	if (PropertyFName == GET_MEMBER_NAME_CHECKED(ThisClass, NormalExportMethod) ||
-		PropertyFName == GET_MEMBER_NAME_CHECKED(ThisClass, bTangentDataQuantization) ||
 		PropertyFName == GET_MEMBER_NAME_CHECKED(ThisClass, NormalizeUVCoordinates) ||
 		PropertyFName == GET_MEMBER_NAME_CHECKED(ThisClass, TextureFormat) ||
 		PropertyFName == GET_MEMBER_NAME_CHECKED(ThisClass, TextureHDREncoding) ||
diff --git a/Engine/Plugins/Enterprise/GLTFExporter/MigratedFromGit/Source/GLTFExporter/Public/GLTFExportOptions.h b/Engine/Plugins/Enterprise/GLTFExporter/MigratedFromGit/Source/GLTFExporter/Public/GLTFExportOptions.h
--- a/Engine/Plugins/Enterprise/GLTFExporter/MigratedFromGit/Source/GLTFExporter/Public/GLTFExportOptions.h
+++ b/Engine/Plugins/Enterprise/GLTFExporter/MigratedFromGit/Source/GLTFExporter/Public/GLTFExportOptions.h
@@ -92,7 +92,7 @@ public:
 	EGLTFExporterNormalExportMethod NormalExportMethod;
 
 	UPROPERTY(EditAnywhere, BlueprintReadWrite, config, Category = Mesh)
-	uint32 bTangentDataQuantization : 1;
+	uint32 bUseMeshQuantization : 1;
 
 	UPROPERTY(EditAnywhere, BlueprintReadWrite, config, Category = Mesh)
 	EGLTFExporterNormalizeUVCoordinates NormalizeUVCoordinates;
