[GLTFExporter] refactor: add GLTFJsonBuilder::AddExtension

commit 8097a64833bbc6a05a63feddd33631ff405b5cdc
Date: Sun, 27 Sep 2020 21:56:09 +0200

    refactor: add GLTFJsonBuilder::AddExtension
    (cherry picked from commit 8d195934452232c48a10ff739f6586a795e5c81d​)

#skipci

diff --git a/Engine/Plugins/Enterprise/GLTFExporter/MigratedFromGit/Source/GLTFExporter/Private/Builders/GLTFJsonBuilder.cpp b/Engine/Plugins/Enterprise/GLTFExporter/MigratedFromGit/Source/GLTFExporter/Private/Builders/GLTFJsonBuilder.cpp
--- a/Engine/Plugins/Enterprise/GLTFExporter/MigratedFromGit/Source/GLTFExporter/Private/Builders/GLTFJsonBuilder.cpp
+++ b/Engine/Plugins/Enterprise/GLTFExporter/MigratedFromGit/Source/GLTFExporter/Private/Builders/GLTFJsonBuilder.cpp
@@ -8,6 +8,15 @@ FGLTFJsonBuilder::FGLTFJsonBuilder()
 {
 }
 
+void FGLTFJsonBuilder::AddExtension(EGLTFJsonExtension Extension, bool bIsRequired)
+{
+	JsonRoot.Extensions.Used.Add(Extension);
+	if (bIsRequired)
+	{
+		JsonRoot.Extensions.Required.Add(Extension);
+	}
+}
+
 FGLTFJsonAccessorIndex FGLTFJsonBuilder::AddAccessor(const FGLTFJsonAccessor& JsonAccessor)
 {
 	return FGLTFJsonAccessorIndex(JsonRoot.Accessors.Add(JsonAccessor));
diff --git a/Engine/Plugins/Enterprise/GLTFExporter/MigratedFromGit/Source/GLTFExporter/Private/Builders/GLTFJsonBuilder.h b/Engine/Plugins/Enterprise/GLTFExporter/MigratedFromGit/Source/GLTFExporter/Private/Builders/GLTFJsonBuilder.h
--- a/Engine/Plugins/Enterprise/GLTFExporter/MigratedFromGit/Source/GLTFExporter/Private/Builders/GLTFJsonBuilder.h
+++ b/Engine/Plugins/Enterprise/GLTFExporter/MigratedFromGit/Source/GLTFExporter/Private/Builders/GLTFJsonBuilder.h
@@ -15,6 +15,8 @@ public:
 
 	FGLTFJsonSceneIndex& DefaultScene;
 
+	void AddExtension(EGLTFJsonExtension Extension, bool bIsRequired = false);
+
 	FGLTFJsonAccessorIndex AddAccessor(const FGLTFJsonAccessor& JsonAccessor = {});
 	FGLTFJsonAnimationIndex AddAnimation(const FGLTFJsonAnimation& JsonAnimation = {});
 	FGLTFJsonBufferIndex AddBuffer(const FGLTFJsonBuffer& JsonBuffer = {});
