[GLTFExporter] refactor: add FGLTFPackedVector8 and FGLTFPackedVector16

commit 701e4e07e0889bc20793bc64de0b3600bebdac67
Date: Sun, 27 Sep 2020 22:38:02 +0200

    refactor: add FGLTFPackedVector8 and FGLTFPackedVector16
    (cherry picked from commit c4793dfde04c55e2dc9b639aef13b98d717a0a97​)

#skipci

diff --git a/Engine/Plugins/Enterprise/GLTFExporter/MigratedFromGit/Source/GLTFExporter/Private/Converters/GLTFPackedColor.h b/Engine/Plugins/Enterprise/GLTFExporter/MigratedFromGit/Source/GLTFExporter/Private/Converters/GLTFPackedColor.h
deleted file mode 100644
diff --git a/Engine/Plugins/Enterprise/GLTFExporter/MigratedFromGit/Source/GLTFExporter/Private/Converters/GLTFPackedTypes.h b/Engine/Plugins/Enterprise/GLTFExporter/MigratedFromGit/Source/GLTFExporter/Private/Converters/GLTFPackedTypes.h
new file mode 100644
--- /dev/null
+++ b/Engine/Plugins/Enterprise/GLTFExporter/MigratedFromGit/Source/GLTFExporter/Private/Converters/GLTFPackedTypes.h
@@ -0,0 +1,39 @@
+// Copyright Epic Games, Inc. All Rights Reserved.
+
+#pragma once
+
+#include "CoreTypes.h"
+
+// UE4 uses an endianness dependent memory layout (ABGR for little endian, ARGB for big endian)
+// glTF uses RGBA independent of endianness, hence why we need this dedicated struct
+struct FGLTFPackedColor
+{
+	uint8 R, G, B, A;
+
+	FGLTFPackedColor(uint8 R, uint8 G, uint8 B, uint8 A = 255)
+		: R(R), G(G), B(B), A(A)
+	{
+	}
+};
+
+// Primarily for type-safety
+struct FGLTFPackedVector8
+{
+	int8 X, Y, Z, W;
+
+	FGLTFPackedVector8(int8 X, int8 Y, int8 Z, int8 W = 0)
+        : X(X), Y(Y), Z(Z), W(W)
+	{
+	}
+};
+
+// Primarily for type-safety
+struct FGLTFPackedVector16
+{
+	int16 X, Y, Z, W;
+
+	FGLTFPackedVector16(int16 X, int16 Y, int16 Z, int16 W = 0)
+        : X(X), Y(Y), Z(Z), W(W)
+	{
+	}
+};
