[GLTFExporter] refactor: add GLTFConverterUtility::ConvertNormal and ConvertTangent for packed types

commit b4d37ee19511e774a6a95a89f783f70f82004b4c
Date: Sun, 27 Sep 2020 22:38:46 +0200

    refactor: add GLTFConverterUtility::ConvertNormal and ConvertTangent for packed types
    (cherry picked from commit ef57b9c6bebb2565e21968bf464e5ac63b8a5bd3​)

#skipci

diff --git a/Engine/Plugins/Enterprise/GLTFExporter/MigratedFromGit/Source/GLTFExporter/Private/Converters/GLTFConverterUtility.h b/Engine/Plugins/Enterprise/GLTFExporter/MigratedFromGit/Source/GLTFExporter/Private/Converters/GLTFConverterUtility.h
--- a/Engine/Plugins/Enterprise/GLTFExporter/MigratedFromGit/Source/GLTFExporter/Private/Converters/GLTFConverterUtility.h
+++ b/Engine/Plugins/Enterprise/GLTFExporter/MigratedFromGit/Source/GLTFExporter/Private/Converters/GLTFConverterUtility.h
@@ -2,7 +2,7 @@
 
 #pragma once
 
-#include "Converters/GLTFPackedColor.h"
+#include "Converters/GLTFPackedTypes.h"
 #include "Json/GLTFJsonEnums.h"
 #include "Json/GLTFJsonVector2.h"
 #include "Json/GLTFJsonVector3.h"
@@ -41,12 +41,32 @@ struct FGLTFConverterUtility
 		return ConvertVector(Normal);
 	}
 
+	static FGLTFPackedVector16 ConvertNormal(const FPackedRGBA16N& Normal)
+	{
+		return { Normal.X, Normal.Z, Normal.Y };
+	}
+
+	static FGLTFPackedVector8 ConvertNormal(const FPackedNormal& Normal)
+	{
+		return { Normal.Vector.X, Normal.Vector.Z, Normal.Vector.Y };
+	}
+
 	static FGLTFJsonVector4 ConvertTangent(const FVector& Tangent)
 	{
 		// glTF stores tangent as Vec4, with W component indicating handedness of tangent basis.
 		return { ConvertVector(Tangent), 1.0f };
 	}
 
+	static FGLTFPackedVector16 ConvertTangent(const FPackedRGBA16N& Tangent)
+	{
+		return { Tangent.X, Tangent.Z, Tangent.Y, MAX_int16 /* = 1.0 */ };
+	}
+
+	static FGLTFPackedVector8 ConvertTangent(const FPackedNormal& Tangent)
+	{
+		return { Tangent.Vector.X, Tangent.Vector.Z, Tangent.Vector.Y, MAX_int8 /* = 1.0 */ };
+	}
+
 	static FGLTFJsonVector2 ConvertUV(const FVector2D& UV)
 	{
 		// No conversion actually needed, this is primarily for type-safety.
