[GLTFExporter] feat: add support for tangent data quantization

commit e03ec580bd29f476e495a1bacba0acd3e0b34c2c
Date: Sun, 17 May 2020 02:05:40 +0200

    feat: add support for tangent data quantization
    (cherry picked from commit 6216f947f15e133c8c4e3af04ed64cc7daa5a1f9​)

#skipci

diff --git a/Engine/Plugins/Enterprise/GLTFExporter/MigratedFromGit/Source/GLTFExporter/Private/Converters/GLTFAccessorConverters.cpp b/Engine/Plugins/Enterprise/GLTFExporter/MigratedFromGit/Source/GLTFExporter/Private/Converters/GLTFAccessorConverters.cpp
--- a/Engine/Plugins/Enterprise/GLTFExporter/MigratedFromGit/Source/GLTFExporter/Private/Converters/GLTFAccessorConverters.cpp
+++ b/Engine/Plugins/Enterprise/GLTFExporter/MigratedFromGit/Source/GLTFExporter/Private/Converters/GLTFAccessorConverters.cpp
@@ -5,6 +5,67 @@
 #include "Converters/GLTFSkinWeightVertexBufferHack.h"
 #include "Builders/GLTFConvertBuilder.h"
 
+namespace
+{
+    struct FGLTFNormalConverter
+	{
+		template <typename PackedType>
+		static auto Convert(TStaticMeshVertexTangentDatum<PackedType> TangentVectors)
+		{
+			return FGLTFConverterUtility::ConvertNormal(FVector(TangentVectors.GetTangentZ()));
+		}
+	};
+
+	struct FGLTFPackedNormalConverter
+	{
+		template <typename PackedType>
+        static auto Convert(TStaticMeshVertexTangentDatum<PackedType> TangentVectors)
+		{
+			return FGLTFConverterUtility::ConvertNormal(TangentVectors.TangentZ);
+		}
+	};
+
+	struct FGLTFTangentConverter
+	{
+		template <typename PackedType>
+        static auto Convert(TStaticMeshVertexTangentDatum<PackedType> TangentVectors)
+		{
+			return FGLTFConverterUtility::ConvertTangent(TangentVectors.GetTangentX());
+		}
+	};
+
+    struct FGLTFPackedTangentConverter
+    {
+    	template <typename PackedType>
+        static auto Convert(TStaticMeshVertexTangentDatum<PackedType> TangentVectors)
+    	{
+    		return FGLTFConverterUtility::ConvertTangent(TangentVectors.TangentX);
+    	}
+	};
+
+	template <typename VectorType, typename PackedType, typename ConverterType>
+    FGLTFJsonBufferViewIndex AddTangentDataBufferView(FGLTFConvertBuilder& Builder, const FGLTFMeshSection* MeshSection, const FStaticMeshVertexBuffer* VertexBuffer)
+	{
+		const TArray<uint32>& IndexMap = MeshSection->IndexMap;
+		const uint32 VertexCount = IndexMap.Num();
+
+		typedef TStaticMeshVertexTangentDatum<PackedType> TangentType;
+		const void* TangentData = const_cast<FStaticMeshVertexBuffer*>(VertexBuffer)->GetTangentData();
+		const TangentType* TangentVectors = static_cast<const TangentType*>(TangentData);
+
+		TArray<VectorType> Vectors;
+		Vectors.AddUninitialized(VertexCount);
+
+		for (uint32 VertexIndex = 0; VertexIndex < VertexCount; ++VertexIndex)
+		{
+			const uint32 MappedVertexIndex = IndexMap[VertexIndex];
+			Vectors[VertexIndex] = ConverterType::Convert(TangentVectors[MappedVertexIndex]);
+		}
+
+		return Builder.AddBufferView(Vectors);
+	}
+}
+
 FGLTFJsonAccessorIndex FGLTFPositionBufferConverter::Convert(const FGLTFMeshSection* MeshSection, const FPositionVertexBuffer* VertexBuffer)
 {
 	if (VertexBuffer == nullptr || VertexBuffer->GetNumVertices() == 0)
@@ -95,20 +156,42 @@ FGLTFJsonAccessorIndex FGLTFNormalBufferConverter::Convert(const FGLTFMeshSectio
 	const TArray<uint32>& IndexMap = MeshSection->IndexMap;
 	const uint32 VertexCount = IndexMap.Num();
 
-	TArray<FGLTFJsonVector3> Normals;
-	Normals.AddUninitialized(VertexCount);
+	FGLTFJsonBufferViewIndex BufferViewIndex;
+	EGLTFJsonComponentType ComponentType;
 
-	for (uint32 VertexIndex = 0; VertexIndex < VertexCount; ++VertexIndex)
+	const bool bQuantization = Builder.ExportOptions->bUseMeshQuantization;
+	const bool bHighPrecision = VertexBuffer->GetUseHighPrecisionTangentBasis();
+
+	if (bQuantization)
 	{
-		const uint32 MappedVertexIndex = IndexMap[VertexIndex];
-		Normals[VertexIndex] = FGLTFConverterUtility::ConvertNormal(VertexBuffer->VertexTangentZ(MappedVertexIndex));
+		Builder.AddExtension(EGLTFJsonExtension::KHR_MeshQuantization);
+
+		if (bHighPrecision)
+		{
+			ComponentType = EGLTFJsonComponentType::S16;
+			BufferViewIndex = AddTangentDataBufferView<FGLTFPackedVector16, FPackedRGBA16N, FGLTFPackedNormalConverter>(Builder, MeshSection, VertexBuffer);
+		}
+		else
+		{
+			ComponentType = EGLTFJsonComponentType::S8;
+			BufferViewIndex = AddTangentDataBufferView<FGLTFPackedVector8, FPackedNormal, FGLTFPackedNormalConverter>(Builder, MeshSection, VertexBuffer);
+			Builder.GetBufferView(BufferViewIndex).ByteStride = 4;
+		}
+	}
+	else
+	{
+		ComponentType = EGLTFJsonComponentType::F32;
+		BufferViewIndex = bHighPrecision
+            ? AddTangentDataBufferView<FGLTFJsonVector3, FPackedRGBA16N, FGLTFNormalConverter>(Builder, MeshSection, VertexBuffer)
+            : AddTangentDataBufferView<FGLTFJsonVector3, FPackedNormal, FGLTFNormalConverter>(Builder, MeshSection, VertexBuffer);
 	}
 
 	FGLTFJsonAccessor JsonAccessor;
-	JsonAccessor.BufferView = Builder.AddBufferView(Normals, EGLTFJsonBufferTarget::ArrayBuffer);
-	JsonAccessor.ComponentType = EGLTFJsonComponentType::F32;
+	JsonAccessor.BufferView = BufferViewIndex;
+	JsonAccessor.ComponentType = ComponentType;
 	JsonAccessor.Count = VertexCount;
 	JsonAccessor.Type = EGLTFJsonAccessorType::Vec3;
+	JsonAccessor.bNormalized = bQuantization;
 
 	return Builder.AddAccessor(JsonAccessor);
 }
@@ -123,20 +206,41 @@ FGLTFJsonAccessorIndex FGLTFTangentBufferConverter::Convert(const FGLTFMeshSecti
 	const TArray<uint32>& IndexMap = MeshSection->IndexMap;
 	const uint32 VertexCount = IndexMap.Num();
 
-	TArray<FGLTFJsonVector4> Tangents;
-	Tangents.AddUninitialized(VertexCount);
+	FGLTFJsonBufferViewIndex BufferViewIndex;
+	EGLTFJsonComponentType ComponentType;
 
-	for (uint32 VertexIndex = 0; VertexIndex < VertexCount; ++VertexIndex)
+	const bool bQuantization = Builder.ExportOptions->bUseMeshQuantization;
+	const bool bHighPrecision = VertexBuffer->GetUseHighPrecisionTangentBasis();
+
+	if (bQuantization)
 	{
-		const uint32 MappedVertexIndex = IndexMap[VertexIndex];
-		Tangents[VertexIndex] = FGLTFConverterUtility::ConvertTangent(VertexBuffer->VertexTangentX(MappedVertexIndex));
+		Builder.AddExtension(EGLTFJsonExtension::KHR_MeshQuantization);
+
+		if (bHighPrecision)
+		{
+			ComponentType = EGLTFJsonComponentType::S16;
+			BufferViewIndex = AddTangentDataBufferView<FGLTFPackedVector16, FPackedRGBA16N, FGLTFPackedTangentConverter>(Builder, MeshSection, VertexBuffer);
+		}
+		else
+		{
+			ComponentType = EGLTFJsonComponentType::S8;
+			BufferViewIndex = AddTangentDataBufferView<FGLTFPackedVector8, FPackedNormal, FGLTFPackedTangentConverter>(Builder, MeshSection, VertexBuffer);
+		}
+	}
+	else
+	{
+		ComponentType = EGLTFJsonComponentType::F32;
+		BufferViewIndex = bHighPrecision
+            ? AddTangentDataBufferView<FGLTFJsonVector4, FPackedRGBA16N, FGLTFTangentConverter>(Builder, MeshSection, VertexBuffer)
+            : AddTangentDataBufferView<FGLTFJsonVector4, FPackedNormal, FGLTFTangentConverter>(Builder, MeshSection, VertexBuffer);
 	}
 
 	FGLTFJsonAccessor JsonAccessor;
-	JsonAccessor.BufferView = Builder.AddBufferView(Tangents, EGLTFJsonBufferTarget::ArrayBuffer);
-	JsonAccessor.ComponentType = EGLTFJsonComponentType::F32;
+	JsonAccessor.BufferView = BufferViewIndex;
+	JsonAccessor.ComponentType = ComponentType;
 	JsonAccessor.Count = VertexCount;
 	JsonAccessor.Type = EGLTFJsonAccessorType::Vec4;
+	JsonAccessor.bNormalized = bQuantization;
 
 	return Builder.AddAccessor(JsonAccessor);
 }
