[GLTFExporter] refactor: reorganize template functions for converting vertex tangent data

commit 9a5108515a1da9f370cc29a3721f569a0acae693
Date: Mon, 28 Sep 2020 17:59:14 +0200

    refactor: reorganize template functions for converting vertex tangent data
    (cherry picked from commit d801077295952192102552c0e416d9110889fa87​)

#skipci

diff --git a/Engine/Plugins/Enterprise/GLTFExporter/MigratedFromGit/Source/GLTFExporter/Private/Converters/GLTFAccessorConverters.cpp b/Engine/Plugins/Enterprise/GLTFExporter/MigratedFromGit/Source/GLTFExporter/Private/Converters/GLTFAccessorConverters.cpp
--- a/Engine/Plugins/Enterprise/GLTFExporter/MigratedFromGit/Source/GLTFExporter/Private/Converters/GLTFAccessorConverters.cpp
+++ b/Engine/Plugins/Enterprise/GLTFExporter/MigratedFromGit/Source/GLTFExporter/Private/Converters/GLTFAccessorConverters.cpp
@@ -7,63 +7,41 @@
 
 namespace
 {
-    struct FGLTFNormalConverter
+	template <typename DestinationType, typename SourceType>
+	struct TGLTFVertexNormalUtility
 	{
-		template <typename PackedType>
-		static auto Convert(TStaticMeshVertexTangentDatum<PackedType> TangentVectors)
+		static DestinationType Convert(TStaticMeshVertexTangentDatum<SourceType> Vertex)
 		{
-			return FGLTFConverterUtility::ConvertNormal(FVector(TangentVectors.GetTangentZ()));
+			return FGLTFConverterUtility::ConvertNormal(Vertex.TangentZ);
 		}
 	};
 
-	struct FGLTFPackedNormalConverter
+	template <typename SourceType>
+	struct TGLTFVertexNormalUtility<FGLTFJsonVector3, SourceType>
 	{
-		template <typename PackedType>
-        static auto Convert(TStaticMeshVertexTangentDatum<PackedType> TangentVectors)
+		static FGLTFJsonVector3 Convert(TStaticMeshVertexTangentDatum<SourceType> Vertex)
 		{
-			return FGLTFConverterUtility::ConvertNormal(TangentVectors.TangentZ);
+			return FGLTFConverterUtility::ConvertNormal(Vertex.TangentZ.ToFVector());
 		}
 	};
 
-	struct FGLTFTangentConverter
+	template <typename DestinationType, typename SourceType>
+	struct TGLTFVertexTangentUtility
 	{
-		template <typename PackedType>
-        static auto Convert(TStaticMeshVertexTangentDatum<PackedType> TangentVectors)
+		static DestinationType Convert(TStaticMeshVertexTangentDatum<SourceType> Vertex)
 		{
-			return FGLTFConverterUtility::ConvertTangent(TangentVectors.GetTangentX());
+			return FGLTFConverterUtility::ConvertTangent(Vertex.TangentX);
 		}
 	};
 
-    struct FGLTFPackedTangentConverter
-    {
-    	template <typename PackedType>
-        static auto Convert(TStaticMeshVertexTangentDatum<PackedType> TangentVectors)
-    	{
-    		return FGLTFConverterUtility::ConvertTangent(TangentVectors.TangentX);
-    	}
-	};
-
-	template <typename VectorType, typename PackedType, typename ConverterType>
-    FGLTFJsonBufferViewIndex AddTangentDataBufferView(FGLTFConvertBuilder& Builder, const FGLTFMeshSection* MeshSection, const FStaticMeshVertexBuffer* VertexBuffer)
+	template <typename SourceType>
+	struct TGLTFVertexTangentUtility<FGLTFJsonVector4, SourceType>
 	{
-		const TArray<uint32>& IndexMap = MeshSection->IndexMap;
-		const uint32 VertexCount = IndexMap.Num();
-
-		typedef TStaticMeshVertexTangentDatum<PackedType> TangentType;
-		const void* TangentData = const_cast<FStaticMeshVertexBuffer*>(VertexBuffer)->GetTangentData();
-		const TangentType* TangentVectors = static_cast<const TangentType*>(TangentData);
-
-		TArray<VectorType> Vectors;
-		Vectors.AddUninitialized(VertexCount);
-
-		for (uint32 VertexIndex = 0; VertexIndex < VertexCount; ++VertexIndex)
+		static FGLTFJsonVector4 Convert(TStaticMeshVertexTangentDatum<SourceType> Vertex)
 		{
-			const uint32 MappedVertexIndex = IndexMap[VertexIndex];
-			Vectors[VertexIndex] = ConverterType::Convert(TangentVectors[MappedVertexIndex]);
+			return FGLTFConverterUtility::ConvertTangent(Vertex.TangentX.ToFVector());
 		}
-
-		return Builder.AddBufferView(Vectors);
-	}
+	};
 }
 
 FGLTFJsonAccessorIndex FGLTFPositionBufferConverter::Convert(const FGLTFMeshSection* MeshSection, const FPositionVertexBuffer* VertexBuffer)
@@ -153,28 +131,25 @@ FGLTFJsonAccessorIndex FGLTFNormalBufferConverter::Convert(const FGLTFMeshSectio
 		return FGLTFJsonAccessorIndex(INDEX_NONE);
 	}
 
-	const TArray<uint32>& IndexMap = MeshSection->IndexMap;
-	const uint32 VertexCount = IndexMap.Num();
-
 	FGLTFJsonBufferViewIndex BufferViewIndex;
 	EGLTFJsonComponentType ComponentType;
 
-	const bool bQuantization = Builder.ExportOptions->bUseMeshQuantization;
+	const bool bMeshQuantization = Builder.ExportOptions->bUseMeshQuantization;
 	const bool bHighPrecision = VertexBuffer->GetUseHighPrecisionTangentBasis();
 
-	if (bQuantization)
+	if (bMeshQuantization)
 	{
-		Builder.AddExtension(EGLTFJsonExtension::KHR_MeshQuantization);
+		Builder.AddExtension(EGLTFJsonExtension::KHR_MeshQuantization, true);
 
 		if (bHighPrecision)
 		{
 			ComponentType = EGLTFJsonComponentType::S16;
-			BufferViewIndex = AddTangentDataBufferView<FGLTFPackedVector16, FPackedRGBA16N, FGLTFPackedNormalConverter>(Builder, MeshSection, VertexBuffer);
+			BufferViewIndex = ConvertBufferView<FGLTFPackedVector16, FPackedRGBA16N>(MeshSection, VertexBuffer);
 		}
 		else
 		{
 			ComponentType = EGLTFJsonComponentType::S8;
-			BufferViewIndex = AddTangentDataBufferView<FGLTFPackedVector8, FPackedNormal, FGLTFPackedNormalConverter>(Builder, MeshSection, VertexBuffer);
+			BufferViewIndex = ConvertBufferView<FGLTFPackedVector8, FPackedNormal>(MeshSection, VertexBuffer);
 			Builder.GetBufferView(BufferViewIndex).ByteStride = 4;
 		}
 	}
@@ -182,20 +157,48 @@ FGLTFJsonAccessorIndex FGLTFNormalBufferConverter::Convert(const FGLTFMeshSectio
 	{
 		ComponentType = EGLTFJsonComponentType::F32;
 		BufferViewIndex = bHighPrecision
-            ? AddTangentDataBufferView<FGLTFJsonVector3, FPackedRGBA16N, FGLTFNormalConverter>(Builder, MeshSection, VertexBuffer)
-            : AddTangentDataBufferView<FGLTFJsonVector3, FPackedNormal, FGLTFNormalConverter>(Builder, MeshSection, VertexBuffer);
+			? ConvertBufferView<FGLTFJsonVector3, FPackedRGBA16N>(MeshSection, VertexBuffer)
+			: ConvertBufferView<FGLTFJsonVector3, FPackedNormal>(MeshSection, VertexBuffer);
 	}
 
 	FGLTFJsonAccessor JsonAccessor;
 	JsonAccessor.BufferView = BufferViewIndex;
 	JsonAccessor.ComponentType = ComponentType;
-	JsonAccessor.Count = VertexCount;
+	JsonAccessor.Count = MeshSection->IndexMap.Num();
 	JsonAccessor.Type = EGLTFJsonAccessorType::Vec3;
-	JsonAccessor.bNormalized = bQuantization;
+	JsonAccessor.bNormalized = bMeshQuantization;
 
 	return Builder.AddAccessor(JsonAccessor);
 }
 
+template <typename DestinationType, typename SourceType>
+FGLTFJsonBufferViewIndex FGLTFNormalBufferConverter::ConvertBufferView(const FGLTFMeshSection* MeshSection, const FStaticMeshVertexBuffer* VertexBuffer)
+{
+	const TArray<uint32>& IndexMap = MeshSection->IndexMap;
+	const uint32 VertexCount = IndexMap.Num();
+
+	const void* TangentData = const_cast<FStaticMeshVertexBuffer*>(VertexBuffer)->GetTangentData();
+	if (TangentData == nullptr)
+	{
+		// TODO: report error
+		return FGLTFJsonBufferViewIndex(INDEX_NONE);
+	}
+
+	typedef TStaticMeshVertexTangentDatum<SourceType> VertexTangentType;
+	const VertexTangentType* VertexTangents = static_cast<const VertexTangentType*>(TangentData);
+
+	TArray<DestinationType> Normals;
+	Normals.AddUninitialized(VertexCount);
+
+	for (uint32 VertexIndex = 0; VertexIndex < VertexCount; ++VertexIndex)
+	{
+		const uint32 MappedVertexIndex = IndexMap[VertexIndex];
+		Normals[VertexIndex] = TGLTFVertexNormalUtility<DestinationType, SourceType>::Convert(VertexTangents[MappedVertexIndex]);
+	}
+
+	return Builder.AddBufferView(Normals, EGLTFJsonBufferTarget::ArrayBuffer);
+}
+
 FGLTFJsonAccessorIndex FGLTFTangentBufferConverter::Convert(const FGLTFMeshSection* MeshSection, const FStaticMeshVertexBuffer* VertexBuffer)
 {
 	if (VertexBuffer == nullptr || VertexBuffer->GetNumVertices() == 0)
@@ -203,48 +206,73 @@ FGLTFJsonAccessorIndex FGLTFTangentBufferConverter::Convert(const FGLTFMeshSecti
 		return FGLTFJsonAccessorIndex(INDEX_NONE);
 	}
 
-	const TArray<uint32>& IndexMap = MeshSection->IndexMap;
-	const uint32 VertexCount = IndexMap.Num();
-
 	FGLTFJsonBufferViewIndex BufferViewIndex;
 	EGLTFJsonComponentType ComponentType;
 
-	const bool bQuantization = Builder.ExportOptions->bUseMeshQuantization;
+	const bool bMeshQuantization = Builder.ExportOptions->bUseMeshQuantization;
 	const bool bHighPrecision = VertexBuffer->GetUseHighPrecisionTangentBasis();
 
-	if (bQuantization)
+	if (bMeshQuantization)
 	{
-		Builder.AddExtension(EGLTFJsonExtension::KHR_MeshQuantization);
+		Builder.AddExtension(EGLTFJsonExtension::KHR_MeshQuantization, true);
 
 		if (bHighPrecision)
 		{
 			ComponentType = EGLTFJsonComponentType::S16;
-			BufferViewIndex = AddTangentDataBufferView<FGLTFPackedVector16, FPackedRGBA16N, FGLTFPackedTangentConverter>(Builder, MeshSection, VertexBuffer);
+			BufferViewIndex = ConvertBufferView<FGLTFPackedVector16, FPackedRGBA16N>(MeshSection, VertexBuffer);
 		}
 		else
 		{
 			ComponentType = EGLTFJsonComponentType::S8;
-			BufferViewIndex = AddTangentDataBufferView<FGLTFPackedVector8, FPackedNormal, FGLTFPackedTangentConverter>(Builder, MeshSection, VertexBuffer);
+			BufferViewIndex = ConvertBufferView<FGLTFPackedVector8, FPackedNormal>(MeshSection, VertexBuffer);
 		}
 	}
 	else
 	{
 		ComponentType = EGLTFJsonComponentType::F32;
 		BufferViewIndex = bHighPrecision
-            ? AddTangentDataBufferView<FGLTFJsonVector4, FPackedRGBA16N, FGLTFTangentConverter>(Builder, MeshSection, VertexBuffer)
-            : AddTangentDataBufferView<FGLTFJsonVector4, FPackedNormal, FGLTFTangentConverter>(Builder, MeshSection, VertexBuffer);
+			? ConvertBufferView<FGLTFJsonVector4, FPackedRGBA16N>(MeshSection, VertexBuffer)
+			: ConvertBufferView<FGLTFJsonVector4, FPackedNormal>(MeshSection, VertexBuffer);
 	}
 
 	FGLTFJsonAccessor JsonAccessor;
 	JsonAccessor.BufferView = BufferViewIndex;
 	JsonAccessor.ComponentType = ComponentType;
-	JsonAccessor.Count = VertexCount;
+	JsonAccessor.Count = MeshSection->IndexMap.Num();
 	JsonAccessor.Type = EGLTFJsonAccessorType::Vec4;
-	JsonAccessor.bNormalized = bQuantization;
+	JsonAccessor.bNormalized = bMeshQuantization;
 
 	return Builder.AddAccessor(JsonAccessor);
 }
 
+template <typename DestinationType, typename SourceType>
+FGLTFJsonBufferViewIndex FGLTFTangentBufferConverter::ConvertBufferView(const FGLTFMeshSection* MeshSection, const FStaticMeshVertexBuffer* VertexBuffer)
+{
+	const TArray<uint32>& IndexMap = MeshSection->IndexMap;
+	const uint32 VertexCount = IndexMap.Num();
+
+	const void* TangentData = const_cast<FStaticMeshVertexBuffer*>(VertexBuffer)->GetTangentData();
+	if (TangentData == nullptr)
+	{
+		// TODO: report error
+		return FGLTFJsonBufferViewIndex(INDEX_NONE);
+	}
+
+	typedef TStaticMeshVertexTangentDatum<SourceType> VertexTangentType;
+	const VertexTangentType* VertexTangents = static_cast<const VertexTangentType*>(TangentData);
+
+	TArray<DestinationType> Tangents;
+	Tangents.AddUninitialized(VertexCount);
+
+	for (uint32 VertexIndex = 0; VertexIndex < VertexCount; ++VertexIndex)
+	{
+		const uint32 MappedVertexIndex = IndexMap[VertexIndex];
+		Tangents[VertexIndex] = TGLTFVertexTangentUtility<DestinationType, SourceType>::Convert(VertexTangents[MappedVertexIndex]);
+	}
+
+	return Builder.AddBufferView(Tangents, EGLTFJsonBufferTarget::ArrayBuffer);
+}
+
 FGLTFJsonAccessorIndex FGLTFUVBufferConverter::Convert(const FGLTFMeshSection* MeshSection, const FStaticMeshVertexBuffer* VertexBuffer, uint32 UVIndex)
 {
 	if (VertexBuffer == nullptr || VertexBuffer->GetNumVertices() == 0)
diff --git a/Engine/Plugins/Enterprise/GLTFExporter/MigratedFromGit/Source/GLTFExporter/Private/Converters/GLTFAccessorConverters.h b/Engine/Plugins/Enterprise/GLTFExporter/MigratedFromGit/Source/GLTFExporter/Private/Converters/GLTFAccessorConverters.h
--- a/Engine/Plugins/Enterprise/GLTFExporter/MigratedFromGit/Source/GLTFExporter/Private/Converters/GLTFAccessorConverters.h
+++ b/Engine/Plugins/Enterprise/GLTFExporter/MigratedFromGit/Source/GLTFExporter/Private/Converters/GLTFAccessorConverters.h
@@ -33,6 +33,9 @@ class FGLTFNormalBufferConverter final : public TGLTFAccessorConverter<const FGL
 	using TGLTFAccessorConverter::TGLTFAccessorConverter;
 
 	virtual FGLTFJsonAccessorIndex Convert(const FGLTFMeshSection* MeshSection, const FStaticMeshVertexBuffer* VertexBuffer) override;
+
+	template <typename DestinationType, typename SourceType>
+	FGLTFJsonBufferViewIndex ConvertBufferView(const FGLTFMeshSection* MeshSection, const FStaticMeshVertexBuffer* VertexBuffer);
 };
 
 class FGLTFTangentBufferConverter final : public TGLTFAccessorConverter<const FGLTFMeshSection*, const FStaticMeshVertexBuffer*>
@@ -40,6 +43,9 @@ class FGLTFTangentBufferConverter final : public TGLTFAccessorConverter<const FG
 	using TGLTFAccessorConverter::TGLTFAccessorConverter;
 
 	virtual FGLTFJsonAccessorIndex Convert(const FGLTFMeshSection* MeshSection, const FStaticMeshVertexBuffer* VertexBuffer) override;
+
+	template <typename DestinationType, typename SourceType>
+    FGLTFJsonBufferViewIndex ConvertBufferView(const FGLTFMeshSection* MeshSection, const FStaticMeshVertexBuffer* VertexBuffer);
 };
 
 class FGLTFUVBufferConverter final : public TGLTFAccessorConverter<const FGLTFMeshSection*, const FStaticMeshVertexBuffer*, uint32>
