[GLTFExporter] refactor: split GLTFExportOptions::bUseMeshQuantization into bQuantizeVertexNormals and bQuantizeVertexTangents

commit 6c7cbf4205f97fd2a4b4ec807c64cf566e06dce3
Date: Mon, 28 Sep 2020 18:02:06 +0200

    refactor: split GLTFExportOptions::bUseMeshQuantization into bQuantizeVertexNormals and bQuantizeVertexTangents
    (cherry picked from commit a3667d21e503acb0dfa394e86cabf22dfcda012f​)

#skipci

diff --git a/Engine/Plugins/Enterprise/GLTFExporter/MigratedFromGit/Source/GLTFExporter/Private/Converters/GLTFAccessorConverters.cpp b/Engine/Plugins/Enterprise/GLTFExporter/MigratedFromGit/Source/GLTFExporter/Private/Converters/GLTFAccessorConverters.cpp
--- a/Engine/Plugins/Enterprise/GLTFExporter/MigratedFromGit/Source/GLTFExporter/Private/Converters/GLTFAccessorConverters.cpp
+++ b/Engine/Plugins/Enterprise/GLTFExporter/MigratedFromGit/Source/GLTFExporter/Private/Converters/GLTFAccessorConverters.cpp
@@ -134,7 +134,7 @@ FGLTFJsonAccessorIndex FGLTFNormalBufferConverter::Convert(const FGLTFMeshSectio
 	FGLTFJsonBufferViewIndex BufferViewIndex;
 	EGLTFJsonComponentType ComponentType;
 
-	const bool bMeshQuantization = Builder.ExportOptions->bUseMeshQuantization;
+	const bool bMeshQuantization = Builder.ExportOptions->bQuantizeVertexNormals;
 	const bool bHighPrecision = VertexBuffer->GetUseHighPrecisionTangentBasis();
 
 	if (bMeshQuantization)
@@ -209,7 +209,7 @@ FGLTFJsonAccessorIndex FGLTFTangentBufferConverter::Convert(const FGLTFMeshSecti
 	FGLTFJsonBufferViewIndex BufferViewIndex;
 	EGLTFJsonComponentType ComponentType;
 
-	const bool bMeshQuantization = Builder.ExportOptions->bUseMeshQuantization;
+	const bool bMeshQuantization = Builder.ExportOptions->bQuantizeVertexTangents;
 	const bool bHighPrecision = VertexBuffer->GetUseHighPrecisionTangentBasis();
 
 	if (bMeshQuantization)
diff --git a/Engine/Plugins/Enterprise/GLTFExporter/MigratedFromGit/Source/GLTFExporter/Private/GLTFExportOptions.cpp b/Engine/Plugins/Enterprise/GLTFExporter/MigratedFromGit/Source/GLTFExporter/Private/GLTFExportOptions.cpp
--- a/Engine/Plugins/Enterprise/GLTFExporter/MigratedFromGit/Source/GLTFExporter/Private/GLTFExportOptions.cpp
+++ b/Engine/Plugins/Enterprise/GLTFExporter/MigratedFromGit/Source/GLTFExporter/Private/GLTFExportOptions.cpp
@@ -19,7 +19,8 @@ UGLTFExportOptions::UGLTFExportOptions(const FObjectInitializer& ObjectInitializ
 	DefaultMaterialBakeSize = EGLTFExporterMaterialBakeSize::POT_1024;
 	bExportVertexColors = true;
 	NormalExportMethod = EGLTFExporterNormalExportMethod::NormalsAndTangents;
-	bUseMeshQuantization = true;
+	bQuantizeVertexNormals = true;
+	bQuantizeVertexTangents = true;
 	NormalizeUVCoordinates = EGLTFExporterNormalizeUVCoordinates::Never;
 	DefaultLevelOfDetail = 0;
 	bExportVertexSkinWeights = true;
diff --git a/Engine/Plugins/Enterprise/GLTFExporter/MigratedFromGit/Source/GLTFExporter/Public/GLTFExportOptions.h b/Engine/Plugins/Enterprise/GLTFExporter/MigratedFromGit/Source/GLTFExporter/Public/GLTFExportOptions.h
--- a/Engine/Plugins/Enterprise/GLTFExporter/MigratedFromGit/Source/GLTFExporter/Public/GLTFExportOptions.h
+++ b/Engine/Plugins/Enterprise/GLTFExporter/MigratedFromGit/Source/GLTFExporter/Public/GLTFExportOptions.h
@@ -92,7 +92,10 @@ public:
 	EGLTFExporterNormalExportMethod NormalExportMethod;
 
 	UPROPERTY(EditAnywhere, BlueprintReadWrite, config, Category = Mesh)
-	uint32 bUseMeshQuantization : 1;
+	uint32 bQuantizeVertexNormals : 1;
+
+	UPROPERTY(EditAnywhere, BlueprintReadWrite, config, Category = Mesh)
+	uint32 bQuantizeVertexTangents : 1;
 
 	UPROPERTY(EditAnywhere, BlueprintReadWrite, config, Category = Mesh)
 	EGLTFExporterNormalizeUVCoordinates NormalizeUVCoordinates;
