[GLTFExporter] refactor: remove unnecessary mimetype comparison in image search and TODO regarding prop comparison since the compressed data should contain all this

commit 6404e10bb3e5cdd9469b24f711b5906cdab794a8
Date: Tue, 29 Sep 2020 21:48:34 +0200

    refactor: remove unnecessary mimetype comparison in image search and TODO regarding prop comparison since the compressed data should contain all this
    (cherry picked from commit 31f305b0dc4586cc85dec54d7edbd6dcf6051db7​)

#skipci

diff --git a/Engine/Plugins/Enterprise/GLTFExporter/MigratedFromGit/Source/GLTFExporter/Private/Builders/GLTFImageBuilder.cpp b/Engine/Plugins/Enterprise/GLTFExporter/MigratedFromGit/Source/GLTFExporter/Private/Builders/GLTFImageBuilder.cpp
--- a/Engine/Plugins/Enterprise/GLTFExporter/MigratedFromGit/Source/GLTFExporter/Private/Builders/GLTFImageBuilder.cpp
+++ b/Engine/Plugins/Enterprise/GLTFExporter/MigratedFromGit/Source/GLTFExporter/Private/Builders/GLTFImageBuilder.cpp
@@ -12,7 +12,7 @@ FGLTFImageBuilder::FGLTFImageBuilder()
 
 FGLTFJsonImageIndex FGLTFImageBuilder::AddImage(const void* CompressedData, int64 CompressedByteLength, EGLTFJsonMimeType MimeType, const FString& Name)
 {
-	FGLTFJsonImageIndex ImageIndex = FindImage(CompressedData, CompressedByteLength, MimeType);
+	FGLTFJsonImageIndex ImageIndex = FindImage(CompressedData, CompressedByteLength);
 	if (ImageIndex == INDEX_NONE)
 	{
 		FGLTFJsonImage Image;
@@ -86,25 +86,18 @@ bool FGLTFImageBuilder::Serialize(FArchive& Archive, const FString& FilePath)
 	return FGLTFBufferBuilder::Serialize(Archive, FilePath);
 }
 
-FGLTFJsonImageIndex FGLTFImageBuilder::FindImage(const void* CompressedData, int64 CompressedByteLength, EGLTFJsonMimeType MimeType) const
+FGLTFJsonImageIndex FGLTFImageBuilder::FindImage(const void* CompressedData, int64 CompressedByteLength) const
 {
-	// TODO: maybe use size and / or other properties when looking for an identical image, just in case
-	// images with different properties happen to have the exact same compressed data.
-	// It would however limit the call-sites from where this function can be used, since we need to know these properties.
-
+	// TODO: instead of looping through all images, use a hash table to minimize search space
 	for (const auto& DataPair : ImageDataLookup)
 	{
-		const FGLTFJsonImage& JsonImage = const_cast<FGLTFImageBuilder*>(this)->GetImage(DataPair.Key);
-		if (JsonImage.MimeType == MimeType)
+		const TArray64<uint8>& ImageData = DataPair.Value;
+		if (ImageData.Num() == CompressedByteLength)
 		{
-			const TArray64<uint8>& ImageData = DataPair.Value;
-			if (ImageData.Num() == CompressedByteLength)
+			if (FMemory::Memcmp(ImageData.GetData(), CompressedData, CompressedByteLength) == 0)
 			{
-				if (FMemory::Memcmp(ImageData.GetData(), CompressedData, CompressedByteLength) == 0)
-				{
-					const FGLTFJsonImageIndex ImageIndex = DataPair.Key;
-					return ImageIndex;
-				}
+				const FGLTFJsonImageIndex ImageIndex = DataPair.Key;
+				return ImageIndex;
 			}
 		}
 	}
diff --git a/Engine/Plugins/Enterprise/GLTFExporter/MigratedFromGit/Source/GLTFExporter/Private/Builders/GLTFImageBuilder.h b/Engine/Plugins/Enterprise/GLTFExporter/MigratedFromGit/Source/GLTFExporter/Private/Builders/GLTFImageBuilder.h
--- a/Engine/Plugins/Enterprise/GLTFExporter/MigratedFromGit/Source/GLTFExporter/Private/Builders/GLTFImageBuilder.h
+++ b/Engine/Plugins/Enterprise/GLTFExporter/MigratedFromGit/Source/GLTFExporter/Private/Builders/GLTFImageBuilder.h
@@ -25,5 +25,5 @@ private:
 
 	TMap<FGLTFJsonImageIndex, TArray64<uint8>> ImageDataLookup;
 
-	FGLTFJsonImageIndex FindImage(const void* CompressedData, int64 CompressedByteLength, EGLTFJsonMimeType MimeType) const;
+	FGLTFJsonImageIndex FindImage(const void* CompressedData, int64 CompressedByteLength) const;
 };
