[GLTFExporter] refactor: make each builder's serialize function a regular functions, instead of virtual overrides

commit 18578fc386e46cdc4d000a0a0db9f150020099b7
Date: Tue, 29 Sep 2020 22:04:11 +0200

    refactor: make each builder's serialize function a regular functions, instead of virtual overrides
    (cherry picked from commit d063eb2861f3932a52b2970c02849ec025f90c01​)

#skipci

diff --git a/Engine/Plugins/Enterprise/GLTFExporter/MigratedFromGit/Source/GLTFExporter/Private/Builders/GLTFBufferBuilder.cpp b/Engine/Plugins/Enterprise/GLTFExporter/MigratedFromGit/Source/GLTFExporter/Private/Builders/GLTFBufferBuilder.cpp
--- a/Engine/Plugins/Enterprise/GLTFExporter/MigratedFromGit/Source/GLTFExporter/Private/Builders/GLTFBufferBuilder.cpp
+++ b/Engine/Plugins/Enterprise/GLTFExporter/MigratedFromGit/Source/GLTFExporter/Private/Builders/GLTFBufferBuilder.cpp
@@ -9,6 +9,20 @@ FGLTFBufferBuilder::FGLTFBufferBuilder()
 	BufferIndex = AddBuffer();
 }
 
+bool FGLTFBufferBuilder::Serialize(FArchive& Archive, const FString& FilePath)
+{
+	const FString BinaryFilePath = FPaths::ChangeExtension(FilePath, TEXT(".bin"));
+
+	if(!FFileHelper::SaveArrayToFile(BufferData, *BinaryFilePath))
+	{
+		// TODO: report error
+		return false;
+	}
+
+	UpdateJsonBufferObject(BinaryFilePath);
+	return true;
+}
+
 FGLTFJsonBufferViewIndex FGLTFBufferBuilder::AddBufferView(const void* RawData, uint64 ByteLength, EGLTFJsonBufferTarget BufferTarget, uint8 DataAlignment)
 {
 	uint64 ByteOffset = BufferData.Num();
@@ -38,17 +52,3 @@ void FGLTFBufferBuilder::UpdateJsonBufferObject(const FString& BinaryFilePath)
 	JsonBuffer.URI = FPaths::GetCleanFilename(BinaryFilePath);
 	JsonBuffer.ByteLength = BufferData.Num();
 }
-
-bool FGLTFBufferBuilder::Serialize(FArchive& Archive, const FString& FilePath)
-{
-	const FString BinaryFilePath = FPaths::ChangeExtension(FilePath, TEXT(".bin"));
-
-	if(!FFileHelper::SaveArrayToFile(BufferData, *BinaryFilePath))
-	{
-		// TODO: report error
-		return false;
-	}
-
-	UpdateJsonBufferObject(BinaryFilePath);
-	return FGLTFJsonBuilder::Serialize(Archive, FilePath);
-}
diff --git a/Engine/Plugins/Enterprise/GLTFExporter/MigratedFromGit/Source/GLTFExporter/Private/Builders/GLTFBufferBuilder.h b/Engine/Plugins/Enterprise/GLTFExporter/MigratedFromGit/Source/GLTFExporter/Private/Builders/GLTFBufferBuilder.h
--- a/Engine/Plugins/Enterprise/GLTFExporter/MigratedFromGit/Source/GLTFExporter/Private/Builders/GLTFBufferBuilder.h
+++ b/Engine/Plugins/Enterprise/GLTFExporter/MigratedFromGit/Source/GLTFExporter/Private/Builders/GLTFBufferBuilder.h
@@ -10,6 +10,8 @@ protected:
 
 	FGLTFBufferBuilder();
 
+	bool Serialize(FArchive& Archive, const FString& FilePath);
+
 public:
 
 	FGLTFJsonBufferViewIndex AddBufferView(const void* RawData, uint64 ByteLength, EGLTFJsonBufferTarget BufferTarget = EGLTFJsonBufferTarget::None, uint8 DataAlignment = 4);
@@ -20,8 +22,6 @@ public:
 		return AddBufferView(Array.GetData(), Array.Num() * sizeof(ElementType), BufferTarget, DataAlignment);
 	}
 
-	virtual bool Serialize(FArchive& Archive, const FString& FilePath) override;
-
 private:
 
 	void UpdateJsonBufferObject(const FString& BinaryFilePath);
diff --git a/Engine/Plugins/Enterprise/GLTFExporter/MigratedFromGit/Source/GLTFExporter/Private/Builders/GLTFContainerBuilder.cpp b/Engine/Plugins/Enterprise/GLTFExporter/MigratedFromGit/Source/GLTFExporter/Private/Builders/GLTFContainerBuilder.cpp
--- a/Engine/Plugins/Enterprise/GLTFExporter/MigratedFromGit/Source/GLTFExporter/Private/Builders/GLTFContainerBuilder.cpp
+++ b/Engine/Plugins/Enterprise/GLTFExporter/MigratedFromGit/Source/GLTFExporter/Private/Builders/GLTFContainerBuilder.cpp
@@ -6,3 +6,10 @@ FGLTFContainerBuilder::FGLTFContainerBuilder(const UGLTFExportOptions* ExportOpt
 	: FGLTFConvertBuilder(ExportOptions, bSelectedActorsOnly)
 {
 }
+
+bool FGLTFContainerBuilder::Serialize(FArchive& Archive, const FString& FilePath)
+{
+	return FGLTFImageBuilder::Serialize(FilePath)
+		&& FGLTFBufferBuilder::Serialize(Archive, FilePath)
+		&& FGLTFJsonBuilder::Serialize(Archive);
+}
diff --git a/Engine/Plugins/Enterprise/GLTFExporter/MigratedFromGit/Source/GLTFExporter/Private/Builders/GLTFContainerBuilder.h b/Engine/Plugins/Enterprise/GLTFExporter/MigratedFromGit/Source/GLTFExporter/Private/Builders/GLTFContainerBuilder.h
--- a/Engine/Plugins/Enterprise/GLTFExporter/MigratedFromGit/Source/GLTFExporter/Private/Builders/GLTFContainerBuilder.h
+++ b/Engine/Plugins/Enterprise/GLTFExporter/MigratedFromGit/Source/GLTFExporter/Private/Builders/GLTFContainerBuilder.h
@@ -9,4 +9,6 @@ class FGLTFContainerBuilder : public FGLTFConvertBuilder
 public:
 
 	FGLTFContainerBuilder(const UGLTFExportOptions* ExportOptions, bool bSelectedActorsOnly);
+
+	bool Serialize(FArchive& Archive, const FString& FilePath);
 };
diff --git a/Engine/Plugins/Enterprise/GLTFExporter/MigratedFromGit/Source/GLTFExporter/Private/Builders/GLTFImageBuilder.cpp b/Engine/Plugins/Enterprise/GLTFExporter/MigratedFromGit/Source/GLTFExporter/Private/Builders/GLTFImageBuilder.cpp
--- a/Engine/Plugins/Enterprise/GLTFExporter/MigratedFromGit/Source/GLTFExporter/Private/Builders/GLTFImageBuilder.cpp
+++ b/Engine/Plugins/Enterprise/GLTFExporter/MigratedFromGit/Source/GLTFExporter/Private/Builders/GLTFImageBuilder.cpp
@@ -10,6 +10,36 @@ FGLTFImageBuilder::FGLTFImageBuilder()
 {
 }
 
+bool FGLTFImageBuilder::Serialize(const FString& FilePath)
+{
+	const FString ImageDir = FPaths::GetPath(FilePath);
+	TSet<FString> UniqueImageUris;
+
+	for (const auto& DataPair : ImageDataLookup)
+	{
+		const TArray64<uint8>& ImageData = DataPair.Value;
+		FGLTFJsonImage& JsonImage = GetImage(DataPair.Key);
+
+		const TCHAR* Extension = FGLTFBuilderUtility::GetFileExtension(JsonImage.MimeType);
+		const FString ImageUri = FGLTFBuilderUtility::GetUniqueFilename(JsonImage.Name, Extension, UniqueImageUris);
+
+		const FString ImagePath = FPaths::Combine(ImageDir, ImageUri);
+		if (!FFileHelper::SaveArrayToFile(ImageData, *ImagePath))
+		{
+			// TODO: report error
+			continue;
+		}
+
+		UniqueImageUris.Add(ImageUri);
+
+		JsonImage.Uri = ImageUri;
+		JsonImage.Name.Empty(); // URI already contains name
+		JsonImage.MimeType = EGLTFJsonMimeType::None; // Not required if external file
+	}
+
+	return true;
+}
+
 FGLTFJsonImageIndex FGLTFImageBuilder::AddImage(const void* CompressedData, int64 CompressedByteLength, EGLTFJsonMimeType MimeType, const FString& Name)
 {
 	FGLTFJsonImageIndex ImageIndex = FindImage(CompressedData, CompressedByteLength);
@@ -56,36 +86,6 @@ FGLTFJsonImageIndex FGLTFImageBuilder::AddImage(const FColor* Pixels, FIntPoint
 	return AddImage(Pixels, ByteLength, Size, ERGBFormat::BGRA, 8, Name);
 }
 
-bool FGLTFImageBuilder::Serialize(FArchive& Archive, const FString& FilePath)
-{
-	const FString ImageDir = FPaths::GetPath(FilePath);
-	TSet<FString> UniqueImageUris;
-
-	for (const auto& DataPair : ImageDataLookup)
-	{
-		const TArray64<uint8>& ImageData = DataPair.Value;
-		FGLTFJsonImage& JsonImage = GetImage(DataPair.Key);
-
-		const TCHAR* Extension = FGLTFBuilderUtility::GetFileExtension(JsonImage.MimeType);
-		const FString ImageUri = FGLTFBuilderUtility::GetUniqueFilename(JsonImage.Name, Extension, UniqueImageUris);
-
-		const FString ImagePath = FPaths::Combine(ImageDir, ImageUri);
-		if (!FFileHelper::SaveArrayToFile(ImageData, *ImagePath))
-		{
-			// TODO: report error
-			continue;
-		}
-
-		UniqueImageUris.Add(ImageUri);
-
-		JsonImage.Uri = ImageUri;
-		JsonImage.Name.Empty(); // URI already contains name
-		JsonImage.MimeType = EGLTFJsonMimeType::None; // Not required if external file
-	}
-
-	return FGLTFBufferBuilder::Serialize(Archive, FilePath);
-}
-
 FGLTFJsonImageIndex FGLTFImageBuilder::FindImage(const void* CompressedData, int64 CompressedByteLength) const
 {
 	// TODO: instead of looping through all images, use a hash table to minimize search space
diff --git a/Engine/Plugins/Enterprise/GLTFExporter/MigratedFromGit/Source/GLTFExporter/Private/Builders/GLTFImageBuilder.h b/Engine/Plugins/Enterprise/GLTFExporter/MigratedFromGit/Source/GLTFExporter/Private/Builders/GLTFImageBuilder.h
--- a/Engine/Plugins/Enterprise/GLTFExporter/MigratedFromGit/Source/GLTFExporter/Private/Builders/GLTFImageBuilder.h
+++ b/Engine/Plugins/Enterprise/GLTFExporter/MigratedFromGit/Source/GLTFExporter/Private/Builders/GLTFImageBuilder.h
@@ -12,6 +12,8 @@ protected:
 
 	FGLTFImageBuilder();
 
+	bool Serialize(const FString& FilePath);
+
 public:
 
 	FGLTFJsonImageIndex AddImage(const void* CompressedData, int64 CompressedByteLength, EGLTFJsonMimeType MimeType, const FString& Name);
@@ -19,8 +21,6 @@ public:
 
 	FGLTFJsonImageIndex AddImage(const FColor* Pixels, FIntPoint Size, const FString& Name);
 
-	virtual bool Serialize(FArchive& Archive, const FString& FilePath) override;
-
 private:
 
 	TMap<FGLTFJsonImageIndex, TArray64<uint8>> ImageDataLookup;
diff --git a/Engine/Plugins/Enterprise/GLTFExporter/MigratedFromGit/Source/GLTFExporter/Private/Builders/GLTFJsonBuilder.cpp b/Engine/Plugins/Enterprise/GLTFExporter/MigratedFromGit/Source/GLTFExporter/Private/Builders/GLTFJsonBuilder.cpp
--- a/Engine/Plugins/Enterprise/GLTFExporter/MigratedFromGit/Source/GLTFExporter/Private/Builders/GLTFJsonBuilder.cpp
+++ b/Engine/Plugins/Enterprise/GLTFExporter/MigratedFromGit/Source/GLTFExporter/Private/Builders/GLTFJsonBuilder.cpp
@@ -1,13 +1,18 @@
 // Copyright Epic Games, Inc. All Rights Reserved.
 
 #include "Builders/GLTFJsonBuilder.h"
-#include "Builders/GLTFBuilderUtility.h"
 
 FGLTFJsonBuilder::FGLTFJsonBuilder()
 	: DefaultScene(JsonRoot.DefaultScene)
 {
 }
 
+bool FGLTFJsonBuilder::Serialize(FArchive& Archive)
+{
+	JsonRoot.Serialize(&Archive, true);
+	return true;
+}
+
 void FGLTFJsonBuilder::AddExtension(EGLTFJsonExtension Extension, bool bIsRequired)
 {
 	JsonRoot.Extensions.Used.Add(Extension);
@@ -231,9 +236,3 @@ FGLTFJsonNodeIndex FGLTFJsonBuilder::GetComponentNodeIndex(FGLTFJsonNodeIndex No
 	const FGLTFJsonNode& Node = GetNode(NodeIndex);
 	return Node.ComponentNode != INDEX_NONE ? Node.ComponentNode : NodeIndex;
 }
-
-bool FGLTFJsonBuilder::Serialize(FArchive& Archive, const FString& FilePath)
-{
-	JsonRoot.Serialize(&Archive, true);
-	return true;
-}
diff --git a/Engine/Plugins/Enterprise/GLTFExporter/MigratedFromGit/Source/GLTFExporter/Private/Builders/GLTFJsonBuilder.h b/Engine/Plugins/Enterprise/GLTFExporter/MigratedFromGit/Source/GLTFExporter/Private/Builders/GLTFJsonBuilder.h
--- a/Engine/Plugins/Enterprise/GLTFExporter/MigratedFromGit/Source/GLTFExporter/Private/Builders/GLTFJsonBuilder.h
+++ b/Engine/Plugins/Enterprise/GLTFExporter/MigratedFromGit/Source/GLTFExporter/Private/Builders/GLTFJsonBuilder.h
@@ -11,6 +11,8 @@ protected:
 
 	FGLTFJsonBuilder();
 
+	bool Serialize(FArchive& Archive);
+
 public:
 
 	FGLTFJsonSceneIndex& DefaultScene;
@@ -60,8 +62,6 @@ public:
 
 	FGLTFJsonNodeIndex GetComponentNodeIndex(FGLTFJsonNodeIndex NodeIndex);
 
-	virtual bool Serialize(FArchive& Archive, const FString& FilePath);
-
 private:
 
 	FGLTFJsonRoot JsonRoot;
diff --git a/Engine/Plugins/Enterprise/GLTFExporter/MigratedFromGit/Source/GLTFExporter/Private/Builders/GLTFMessageBuilder.cpp b/Engine/Plugins/Enterprise/GLTFExporter/MigratedFromGit/Source/GLTFExporter/Private/Builders/GLTFMessageBuilder.cpp
--- a/Engine/Plugins/Enterprise/GLTFExporter/MigratedFromGit/Source/GLTFExporter/Private/Builders/GLTFMessageBuilder.cpp
+++ b/Engine/Plugins/Enterprise/GLTFExporter/MigratedFromGit/Source/GLTFExporter/Private/Builders/GLTFMessageBuilder.cpp
@@ -10,10 +10,6 @@ FGLTFMessageBuilder::FGLTFMessageBuilder()
 {
 }
 
-FGLTFMessageBuilder::~FGLTFMessageBuilder()
-{
-}
-
 void FGLTFMessageBuilder::ClearMessages()
 {
 	Messages.Empty();
diff --git a/Engine/Plugins/Enterprise/GLTFExporter/MigratedFromGit/Source/GLTFExporter/Private/Builders/GLTFMessageBuilder.h b/Engine/Plugins/Enterprise/GLTFExporter/MigratedFromGit/Source/GLTFExporter/Private/Builders/GLTFMessageBuilder.h
--- a/Engine/Plugins/Enterprise/GLTFExporter/MigratedFromGit/Source/GLTFExporter/Private/Builders/GLTFMessageBuilder.h
+++ b/Engine/Plugins/Enterprise/GLTFExporter/MigratedFromGit/Source/GLTFExporter/Private/Builders/GLTFMessageBuilder.h
@@ -14,7 +14,6 @@ class FGLTFMessageBuilder
 protected:
 
 	FGLTFMessageBuilder();
-	virtual ~FGLTFMessageBuilder();
 
 public:
 
