[GLTFExporter] refactor: add GLTFBuilder and reorganize FilePath and ExportOptions parameters

commit b85f58626de02e6efce52d1d8d473b4425761015
Date: Tue, 29 Sep 2020 22:05:58 +0200

    refactor: add GLTFBuilder and reorganize FilePath and ExportOptions parameters
    (cherry picked from commit 3f0926a5d12a3cb09f5b64202491033a3da8eb65​)

#skipci

diff --git a/Engine/Plugins/Enterprise/GLTFExporter/MigratedFromGit/Source/GLTFExporter/Private/Builders/GLTFBufferBuilder.cpp b/Engine/Plugins/Enterprise/GLTFExporter/MigratedFromGit/Source/GLTFExporter/Private/Builders/GLTFBufferBuilder.cpp
--- a/Engine/Plugins/Enterprise/GLTFExporter/MigratedFromGit/Source/GLTFExporter/Private/Builders/GLTFBufferBuilder.cpp
+++ b/Engine/Plugins/Enterprise/GLTFExporter/MigratedFromGit/Source/GLTFExporter/Private/Builders/GLTFBufferBuilder.cpp
@@ -4,12 +4,13 @@
 #include "Misc/FileHelper.h"
 #include "Misc/Paths.h"
 
-FGLTFBufferBuilder::FGLTFBufferBuilder()
+FGLTFBufferBuilder::FGLTFBufferBuilder(const FString& FilePath, const UGLTFExportOptions* ExportOptions)
+	: FGLTFJsonBuilder(FilePath, ExportOptions)
 {
 	BufferIndex = AddBuffer();
 }
 
-bool FGLTFBufferBuilder::Serialize(FArchive& Archive, const FString& FilePath)
+bool FGLTFBufferBuilder::Serialize(FArchive& Archive)
 {
 	const FString BinaryFilePath = FPaths::ChangeExtension(FilePath, TEXT(".bin"));
 
diff --git a/Engine/Plugins/Enterprise/GLTFExporter/MigratedFromGit/Source/GLTFExporter/Private/Builders/GLTFBufferBuilder.h b/Engine/Plugins/Enterprise/GLTFExporter/MigratedFromGit/Source/GLTFExporter/Private/Builders/GLTFBufferBuilder.h
--- a/Engine/Plugins/Enterprise/GLTFExporter/MigratedFromGit/Source/GLTFExporter/Private/Builders/GLTFBufferBuilder.h
+++ b/Engine/Plugins/Enterprise/GLTFExporter/MigratedFromGit/Source/GLTFExporter/Private/Builders/GLTFBufferBuilder.h
@@ -8,9 +8,9 @@ class FGLTFBufferBuilder : public FGLTFJsonBuilder
 {
 protected:
 
-	FGLTFBufferBuilder();
+	FGLTFBufferBuilder(const FString& FilePath, const UGLTFExportOptions* ExportOptions);
 
-	bool Serialize(FArchive& Archive, const FString& FilePath);
+	bool Serialize(FArchive& Archive);
 
 public:
 
diff --git a/Engine/Plugins/Enterprise/GLTFExporter/MigratedFromGit/Source/GLTFExporter/Private/Builders/GLTFBuilder.cpp b/Engine/Plugins/Enterprise/GLTFExporter/MigratedFromGit/Source/GLTFExporter/Private/Builders/GLTFBuilder.cpp
new file mode 100644
--- /dev/null
+++ b/Engine/Plugins/Enterprise/GLTFExporter/MigratedFromGit/Source/GLTFExporter/Private/Builders/GLTFBuilder.cpp
@@ -0,0 +1,9 @@
+// Copyright Epic Games, Inc. All Rights Reserved.
+
+#include "Builders/GLTFBuilder.h"
+
+FGLTFBuilder::FGLTFBuilder(const FString& FilePath, const UGLTFExportOptions* ExportOptions)
+	: FilePath(FilePath)
+	, ExportOptions(ExportOptions)
+{
+}
diff --git a/Engine/Plugins/Enterprise/GLTFExporter/MigratedFromGit/Source/GLTFExporter/Private/Builders/GLTFBuilder.h b/Engine/Plugins/Enterprise/GLTFExporter/MigratedFromGit/Source/GLTFExporter/Private/Builders/GLTFBuilder.h
new file mode 100644
--- /dev/null
+++ b/Engine/Plugins/Enterprise/GLTFExporter/MigratedFromGit/Source/GLTFExporter/Private/Builders/GLTFBuilder.h
@@ -0,0 +1,18 @@
+// Copyright Epic Games, Inc. All Rights Reserved.
+
+#pragma once
+
+#include "GLTFExportOptions.h"
+
+class FGLTFBuilder
+{
+protected:
+
+	FGLTFBuilder(const FString& FilePath, const UGLTFExportOptions* ExportOptions);
+
+public:
+
+	const FString FilePath;
+
+	const UGLTFExportOptions* const ExportOptions;
+};
diff --git a/Engine/Plugins/Enterprise/GLTFExporter/MigratedFromGit/Source/GLTFExporter/Private/Builders/GLTFContainerBuilder.cpp b/Engine/Plugins/Enterprise/GLTFExporter/MigratedFromGit/Source/GLTFExporter/Private/Builders/GLTFContainerBuilder.cpp
--- a/Engine/Plugins/Enterprise/GLTFExporter/MigratedFromGit/Source/GLTFExporter/Private/Builders/GLTFContainerBuilder.cpp
+++ b/Engine/Plugins/Enterprise/GLTFExporter/MigratedFromGit/Source/GLTFExporter/Private/Builders/GLTFContainerBuilder.cpp
@@ -2,14 +2,14 @@
 
 #include "Builders/GLTFContainerBuilder.h"
 
-FGLTFContainerBuilder::FGLTFContainerBuilder(const UGLTFExportOptions* ExportOptions, bool bSelectedActorsOnly)
-	: FGLTFConvertBuilder(ExportOptions, bSelectedActorsOnly)
+FGLTFContainerBuilder::FGLTFContainerBuilder(const FString& FilePath, const UGLTFExportOptions* ExportOptions, bool bSelectedActorsOnly)
+	: FGLTFConvertBuilder(FilePath, ExportOptions, bSelectedActorsOnly)
 {
 }
 
-bool FGLTFContainerBuilder::Serialize(FArchive& Archive, const FString& FilePath)
+bool FGLTFContainerBuilder::Serialize(FArchive& Archive)
 {
-	return FGLTFImageBuilder::Serialize(FilePath)
-		&& FGLTFBufferBuilder::Serialize(Archive, FilePath)
+	return FGLTFImageBuilder::Serialize()
+		&& FGLTFBufferBuilder::Serialize(Archive)
 		&& FGLTFJsonBuilder::Serialize(Archive);
 }
diff --git a/Engine/Plugins/Enterprise/GLTFExporter/MigratedFromGit/Source/GLTFExporter/Private/Builders/GLTFContainerBuilder.h b/Engine/Plugins/Enterprise/GLTFExporter/MigratedFromGit/Source/GLTFExporter/Private/Builders/GLTFContainerBuilder.h
--- a/Engine/Plugins/Enterprise/GLTFExporter/MigratedFromGit/Source/GLTFExporter/Private/Builders/GLTFContainerBuilder.h
+++ b/Engine/Plugins/Enterprise/GLTFExporter/MigratedFromGit/Source/GLTFExporter/Private/Builders/GLTFContainerBuilder.h
@@ -8,7 +8,7 @@ class FGLTFContainerBuilder : public FGLTFConvertBuilder
 {
 public:
 
-	FGLTFContainerBuilder(const UGLTFExportOptions* ExportOptions, bool bSelectedActorsOnly);
+	FGLTFContainerBuilder(const FString& FilePath, const UGLTFExportOptions* ExportOptions, bool bSelectedActorsOnly);
 
-	bool Serialize(FArchive& Archive, const FString& FilePath);
+	bool Serialize(FArchive& Archive);
 };
diff --git a/Engine/Plugins/Enterprise/GLTFExporter/MigratedFromGit/Source/GLTFExporter/Private/Builders/GLTFConvertBuilder.cpp b/Engine/Plugins/Enterprise/GLTFExporter/MigratedFromGit/Source/GLTFExporter/Private/Builders/GLTFConvertBuilder.cpp
--- a/Engine/Plugins/Enterprise/GLTFExporter/MigratedFromGit/Source/GLTFExporter/Private/Builders/GLTFConvertBuilder.cpp
+++ b/Engine/Plugins/Enterprise/GLTFExporter/MigratedFromGit/Source/GLTFExporter/Private/Builders/GLTFConvertBuilder.cpp
@@ -2,8 +2,8 @@
 
 #include "Builders/GLTFConvertBuilder.h"
 
-FGLTFConvertBuilder::FGLTFConvertBuilder(const UGLTFExportOptions* ExportOptions, bool bSelectedActorsOnly)
-	: ExportOptions(ExportOptions)
+FGLTFConvertBuilder::FGLTFConvertBuilder(const FString& FilePath, const UGLTFExportOptions* ExportOptions, bool bSelectedActorsOnly)
+	: FGLTFImageBuilder(FilePath, ExportOptions)
 	, bSelectedActorsOnly(bSelectedActorsOnly)
 {
 }
diff --git a/Engine/Plugins/Enterprise/GLTFExporter/MigratedFromGit/Source/GLTFExporter/Private/Builders/GLTFConvertBuilder.h b/Engine/Plugins/Enterprise/GLTFExporter/MigratedFromGit/Source/GLTFExporter/Private/Builders/GLTFConvertBuilder.h
--- a/Engine/Plugins/Enterprise/GLTFExporter/MigratedFromGit/Source/GLTFExporter/Private/Builders/GLTFConvertBuilder.h
+++ b/Engine/Plugins/Enterprise/GLTFExporter/MigratedFromGit/Source/GLTFExporter/Private/Builders/GLTFConvertBuilder.h
@@ -2,7 +2,6 @@
 
 #pragma once
 
-#include "GLTFExportOptions.h"
 #include "Builders/GLTFImageBuilder.h"
 #include "Converters/GLTFAccessorConverters.h"
 #include "Converters/GLTFMeshConverters.h"
@@ -24,11 +23,10 @@ class FGLTFConvertBuilder : public FGLTFImageBuilder
 {
 protected:
 
-	FGLTFConvertBuilder(const UGLTFExportOptions* ExportOptions, bool bSelectedActorsOnly);
+	FGLTFConvertBuilder(const FString& FilePath, const UGLTFExportOptions* ExportOptions, bool bSelectedActorsOnly);
 
 public:
 
-	const UGLTFExportOptions* const ExportOptions;
 	const bool bSelectedActorsOnly;
 
 	FGLTFJsonAccessorIndex GetOrAddPositionAccessor(const FGLTFMeshSection* MeshSection, const FPositionVertexBuffer* VertexBuffer);
diff --git a/Engine/Plugins/Enterprise/GLTFExporter/MigratedFromGit/Source/GLTFExporter/Private/Builders/GLTFImageBuilder.cpp b/Engine/Plugins/Enterprise/GLTFExporter/MigratedFromGit/Source/GLTFExporter/Private/Builders/GLTFImageBuilder.cpp
--- a/Engine/Plugins/Enterprise/GLTFExporter/MigratedFromGit/Source/GLTFExporter/Private/Builders/GLTFImageBuilder.cpp
+++ b/Engine/Plugins/Enterprise/GLTFExporter/MigratedFromGit/Source/GLTFExporter/Private/Builders/GLTFImageBuilder.cpp
@@ -6,11 +6,12 @@
 #include "IImageWrapper.h"
 #include "Misc/FileHelper.h"
 
-FGLTFImageBuilder::FGLTFImageBuilder()
+FGLTFImageBuilder::FGLTFImageBuilder(const FString& FilePath, const UGLTFExportOptions* ExportOptions)
+    : FGLTFBufferBuilder(FilePath, ExportOptions)
 {
 }
 
-bool FGLTFImageBuilder::Serialize(const FString& FilePath)
+bool FGLTFImageBuilder::Serialize()
 {
 	const FString ImageDir = FPaths::GetPath(FilePath);
 	TSet<FString> UniqueImageUris;
diff --git a/Engine/Plugins/Enterprise/GLTFExporter/MigratedFromGit/Source/GLTFExporter/Private/Builders/GLTFImageBuilder.h b/Engine/Plugins/Enterprise/GLTFExporter/MigratedFromGit/Source/GLTFExporter/Private/Builders/GLTFImageBuilder.h
--- a/Engine/Plugins/Enterprise/GLTFExporter/MigratedFromGit/Source/GLTFExporter/Private/Builders/GLTFImageBuilder.h
+++ b/Engine/Plugins/Enterprise/GLTFExporter/MigratedFromGit/Source/GLTFExporter/Private/Builders/GLTFImageBuilder.h
@@ -10,9 +10,9 @@ class FGLTFImageBuilder : public FGLTFBufferBuilder
 {
 protected:
 
-	FGLTFImageBuilder();
+	FGLTFImageBuilder(const FString& FilePath, const UGLTFExportOptions* ExportOptions);
 
-	bool Serialize(const FString& FilePath);
+	bool Serialize();
 
 public:
 
diff --git a/Engine/Plugins/Enterprise/GLTFExporter/MigratedFromGit/Source/GLTFExporter/Private/Builders/GLTFJsonBuilder.cpp b/Engine/Plugins/Enterprise/GLTFExporter/MigratedFromGit/Source/GLTFExporter/Private/Builders/GLTFJsonBuilder.cpp
--- a/Engine/Plugins/Enterprise/GLTFExporter/MigratedFromGit/Source/GLTFExporter/Private/Builders/GLTFJsonBuilder.cpp
+++ b/Engine/Plugins/Enterprise/GLTFExporter/MigratedFromGit/Source/GLTFExporter/Private/Builders/GLTFJsonBuilder.cpp
@@ -2,8 +2,9 @@
 
 #include "Builders/GLTFJsonBuilder.h"
 
-FGLTFJsonBuilder::FGLTFJsonBuilder()
-	: DefaultScene(JsonRoot.DefaultScene)
+FGLTFJsonBuilder::FGLTFJsonBuilder(const FString& FilePath, const UGLTFExportOptions* ExportOptions)
+	: FGLTFMessageBuilder(FilePath, ExportOptions)
+	, DefaultScene(JsonRoot.DefaultScene)
 {
 }
 
diff --git a/Engine/Plugins/Enterprise/GLTFExporter/MigratedFromGit/Source/GLTFExporter/Private/Builders/GLTFJsonBuilder.h b/Engine/Plugins/Enterprise/GLTFExporter/MigratedFromGit/Source/GLTFExporter/Private/Builders/GLTFJsonBuilder.h
--- a/Engine/Plugins/Enterprise/GLTFExporter/MigratedFromGit/Source/GLTFExporter/Private/Builders/GLTFJsonBuilder.h
+++ b/Engine/Plugins/Enterprise/GLTFExporter/MigratedFromGit/Source/GLTFExporter/Private/Builders/GLTFJsonBuilder.h
@@ -9,7 +9,7 @@ class FGLTFJsonBuilder : public FGLTFMessageBuilder
 {
 protected:
 
-	FGLTFJsonBuilder();
+	FGLTFJsonBuilder(const FString& FilePath, const UGLTFExportOptions* ExportOptions);
 
 	bool Serialize(FArchive& Archive);
 
diff --git a/Engine/Plugins/Enterprise/GLTFExporter/MigratedFromGit/Source/GLTFExporter/Private/Builders/GLTFMessageBuilder.cpp b/Engine/Plugins/Enterprise/GLTFExporter/MigratedFromGit/Source/GLTFExporter/Private/Builders/GLTFMessageBuilder.cpp
--- a/Engine/Plugins/Enterprise/GLTFExporter/MigratedFromGit/Source/GLTFExporter/Private/Builders/GLTFMessageBuilder.cpp
+++ b/Engine/Plugins/Enterprise/GLTFExporter/MigratedFromGit/Source/GLTFExporter/Private/Builders/GLTFMessageBuilder.cpp
@@ -6,7 +6,8 @@
 #include "MessageLogModule.h"
 #include "IMessageLogListing.h"
 
-FGLTFMessageBuilder::FGLTFMessageBuilder()
+FGLTFMessageBuilder::FGLTFMessageBuilder(const FString& FilePath, const UGLTFExportOptions* ExportOptions)
+	: FGLTFBuilder(FilePath, ExportOptions)
 {
 }
 
diff --git a/Engine/Plugins/Enterprise/GLTFExporter/MigratedFromGit/Source/GLTFExporter/Private/Builders/GLTFMessageBuilder.h b/Engine/Plugins/Enterprise/GLTFExporter/MigratedFromGit/Source/GLTFExporter/Private/Builders/GLTFMessageBuilder.h
--- a/Engine/Plugins/Enterprise/GLTFExporter/MigratedFromGit/Source/GLTFExporter/Private/Builders/GLTFMessageBuilder.h
+++ b/Engine/Plugins/Enterprise/GLTFExporter/MigratedFromGit/Source/GLTFExporter/Private/Builders/GLTFMessageBuilder.h
@@ -2,6 +2,8 @@
 
 #pragma once
 
+#include "Builders/GLTFBuilder.h"
+
 enum class EGLTFMessageSeverity
 {
 	Info,
@@ -9,11 +11,11 @@ enum class EGLTFMessageSeverity
 	Error
 };
 
-class FGLTFMessageBuilder
+class FGLTFMessageBuilder : public FGLTFBuilder
 {
 protected:
 
-	FGLTFMessageBuilder();
+	FGLTFMessageBuilder(const FString& FilePath, const UGLTFExportOptions* ExportOptions);
 
 public:
 
diff --git a/Engine/Plugins/Enterprise/GLTFExporter/MigratedFromGit/Source/GLTFExporter/Private/Exporters/GLTFExporter.cpp b/Engine/Plugins/Enterprise/GLTFExporter/MigratedFromGit/Source/GLTFExporter/Private/Exporters/GLTFExporter.cpp
--- a/Engine/Plugins/Enterprise/GLTFExporter/MigratedFromGit/Source/GLTFExporter/Private/Exporters/GLTFExporter.cpp
+++ b/Engine/Plugins/Enterprise/GLTFExporter/MigratedFromGit/Source/GLTFExporter/Private/Exporters/GLTFExporter.cpp
@@ -36,10 +36,10 @@ bool UGLTFExporter::ExportBinary(UObject* Object, const TCHAR* Type, FArchive& A
 
 	// TODO: add support for UAssetExportTask::IgnoreObjectList?
 
-	FGLTFContainerBuilder Builder(Options, bSelectedOnly);
+	FGLTFContainerBuilder Builder(CurrentFilename, Options, bSelectedOnly);
 	if (AddObject(Builder, Object))
 	{
-		if (!Builder.Serialize(Archive, CurrentFilename))
+		if (!Builder.Serialize(Archive))
 		{
 			// TODO: more descriptive error
 			Builder.AddErrorMessage(TEXT("Serialize failed"));
