[GLTFExporter] refactor: save each converted image to file directly to catch issues sooner rather than later

commit 0536749f1f556393c5b5c8ae819922a1c2a9e097
Date: Tue, 29 Sep 2020 22:15:46 +0200

    refactor: save each converted image to file directly to catch issues sooner rather than later
    (cherry picked from commit a7652455b0197673a746244a3a3f448164e7b19a​)

#skipci

diff --git a/Engine/Plugins/Enterprise/GLTFExporter/MigratedFromGit/Source/GLTFExporter/Private/Builders/GLTFContainerBuilder.cpp b/Engine/Plugins/Enterprise/GLTFExporter/MigratedFromGit/Source/GLTFExporter/Private/Builders/GLTFContainerBuilder.cpp
--- a/Engine/Plugins/Enterprise/GLTFExporter/MigratedFromGit/Source/GLTFExporter/Private/Builders/GLTFContainerBuilder.cpp
+++ b/Engine/Plugins/Enterprise/GLTFExporter/MigratedFromGit/Source/GLTFExporter/Private/Builders/GLTFContainerBuilder.cpp
@@ -9,7 +9,6 @@ FGLTFContainerBuilder::FGLTFContainerBuilder(const FString& FilePath, const UGLT
 
 bool FGLTFContainerBuilder::Serialize(FArchive& Archive)
 {
-	return FGLTFImageBuilder::Serialize()
-		&& FGLTFBufferBuilder::Serialize(Archive)
+	return FGLTFBufferBuilder::Serialize(Archive)
 		&& FGLTFJsonBuilder::Serialize(Archive);
 }
diff --git a/Engine/Plugins/Enterprise/GLTFExporter/MigratedFromGit/Source/GLTFExporter/Private/Builders/GLTFImageBuilder.cpp b/Engine/Plugins/Enterprise/GLTFExporter/MigratedFromGit/Source/GLTFExporter/Private/Builders/GLTFImageBuilder.cpp
--- a/Engine/Plugins/Enterprise/GLTFExporter/MigratedFromGit/Source/GLTFExporter/Private/Builders/GLTFImageBuilder.cpp
+++ b/Engine/Plugins/Enterprise/GLTFExporter/MigratedFromGit/Source/GLTFExporter/Private/Builders/GLTFImageBuilder.cpp
@@ -8,47 +8,17 @@
 
 FGLTFImageBuilder::FGLTFImageBuilder(const FString& FilePath, const UGLTFExportOptions* ExportOptions)
     : FGLTFBufferBuilder(FilePath, ExportOptions)
+	, ImageDirPath(FPaths::GetPath(FilePath))
 {
 }
 
-bool FGLTFImageBuilder::Serialize()
-{
-	const FString ImageDir = FPaths::GetPath(FilePath);
-	TSet<FString> UniqueImageUris;
-
-	for (const auto& DataPair : ImageDataLookup)
-	{
-		const TArray64<uint8>& ImageData = DataPair.Value;
-		FGLTFJsonImage& JsonImage = GetImage(DataPair.Key);
-
-		const TCHAR* Extension = FGLTFBuilderUtility::GetFileExtension(JsonImage.MimeType);
-		const FString ImageUri = FGLTFBuilderUtility::GetUniqueFilename(JsonImage.Name, Extension, UniqueImageUris);
-
-		const FString ImagePath = FPaths::Combine(ImageDir, ImageUri);
-		if (!FFileHelper::SaveArrayToFile(ImageData, *ImagePath))
-		{
-			// TODO: report error
-			continue;
-		}
-
-		UniqueImageUris.Add(ImageUri);
-
-		JsonImage.Uri = ImageUri;
-		JsonImage.Name.Empty(); // URI already contains name
-		JsonImage.MimeType = EGLTFJsonMimeType::None; // Not required if external file
-	}
-
-	return true;
-}
-
 FGLTFJsonImageIndex FGLTFImageBuilder::AddImage(const void* CompressedData, int64 CompressedByteLength, EGLTFJsonMimeType MimeType, const FString& Name)
 {
 	FGLTFJsonImageIndex ImageIndex = FindImage(CompressedData, CompressedByteLength);
 	if (ImageIndex == INDEX_NONE)
 	{
 		FGLTFJsonImage Image;
-		Image.Name = Name;
-		Image.MimeType = MimeType;
+		Image.Uri = SaveImageToFile(CompressedData, CompressedByteLength, MimeType, Name);
 		ImageIndex = FGLTFJsonBuilder::AddImage(Image);
 
 		TArray64<uint8>& ImageData = ImageDataLookup.Add(ImageIndex);
@@ -87,6 +57,25 @@ FGLTFJsonImageIndex FGLTFImageBuilder::AddImage(const FColor* Pixels, FIntPoint
 	return AddImage(Pixels, ByteLength, Size, ERGBFormat::BGRA, 8, Name);
 }
 
+FString FGLTFImageBuilder::SaveImageToFile(const void* CompressedData, int64 CompressedByteLength, EGLTFJsonMimeType MimeType, const FString& Name)
+{
+	const TCHAR* Extension = FGLTFBuilderUtility::GetFileExtension(MimeType);
+	const FString ImageUri = FGLTFBuilderUtility::GetUniqueFilename(Name, Extension, UniqueImageUris);
+
+	const TArrayView<const uint8> ImageData(static_cast<const uint8*>(CompressedData), CompressedByteLength);
+	const FString ImagePath = FPaths::Combine(ImageDirPath, ImageUri);
+
+	if (!FFileHelper::SaveArrayToFile(ImageData, *ImagePath))
+	{
+		// TODO: report error
+		return TEXT("");
+	}
+
+	UniqueImageUris.Add(ImageUri);
+	return ImageUri;
+}
+
+
 FGLTFJsonImageIndex FGLTFImageBuilder::FindImage(const void* CompressedData, int64 CompressedByteLength) const
 {
 	// TODO: instead of looping through all images, use a hash table to minimize search space
diff --git a/Engine/Plugins/Enterprise/GLTFExporter/MigratedFromGit/Source/GLTFExporter/Private/Builders/GLTFImageBuilder.h b/Engine/Plugins/Enterprise/GLTFExporter/MigratedFromGit/Source/GLTFExporter/Private/Builders/GLTFImageBuilder.h
--- a/Engine/Plugins/Enterprise/GLTFExporter/MigratedFromGit/Source/GLTFExporter/Private/Builders/GLTFImageBuilder.h
+++ b/Engine/Plugins/Enterprise/GLTFExporter/MigratedFromGit/Source/GLTFExporter/Private/Builders/GLTFImageBuilder.h
@@ -12,8 +12,6 @@ protected:
 
 	FGLTFImageBuilder(const FString& FilePath, const UGLTFExportOptions* ExportOptions);
 
-	bool Serialize();
-
 public:
 
 	FGLTFJsonImageIndex AddImage(const void* CompressedData, int64 CompressedByteLength, EGLTFJsonMimeType MimeType, const FString& Name);
@@ -23,7 +21,12 @@ public:
 
 private:
 
-	TMap<FGLTFJsonImageIndex, TArray64<uint8>> ImageDataLookup;
+	FString SaveImageToFile(const void* CompressedData, int64 CompressedByteLength, EGLTFJsonMimeType MimeType, const FString& Name);
 
 	FGLTFJsonImageIndex FindImage(const void* CompressedData, int64 CompressedByteLength) const;
+
+	TMap<FGLTFJsonImageIndex, TArray64<uint8>> ImageDataLookup;
+
+	const FString ImageDirPath;
+	TSet<FString> UniqueImageUris;
 };
