[GLTFExporter] refactor: optimize image reuse search using hash key lookup

commit 3ebec2b2fee5ac36fa863da317e155f251033cb9
Date: Tue, 29 Sep 2020 22:54:12 +0200

    refactor: optimize image reuse search using hash key lookup
    (cherry picked from commit 1bb90baa3a7360195d1580e5bcaba32cba2fcab3​)

#skipci

diff --git a/Engine/Plugins/Enterprise/GLTFExporter/MigratedFromGit/Source/GLTFExporter/Private/Builders/GLTFBinaryHashKey.h b/Engine/Plugins/Enterprise/GLTFExporter/MigratedFromGit/Source/GLTFExporter/Private/Builders/GLTFBinaryHashKey.h
new file mode 100644
--- /dev/null
+++ b/Engine/Plugins/Enterprise/GLTFExporter/MigratedFromGit/Source/GLTFExporter/Private/Builders/GLTFBinaryHashKey.h
@@ -0,0 +1,43 @@
+// Copyright Epic Games, Inc. All Rights Reserved.
+
+#pragma once
+
+#include "CoreMinimal.h"
+#include "Misc/Crc.h"
+
+class FGLTFBinaryHashKey
+{
+public:
+
+	FGLTFBinaryHashKey(const void* RawData, int64 ByteLength)
+        : Hash(FCrc::MemCrc32(RawData, ByteLength))
+        , Bytes(static_cast<const uint8*>(RawData), ByteLength)
+	{
+	}
+
+	bool operator==(const FGLTFBinaryHashKey& Other) const
+	{
+		return CompareBytes(Other);
+	}
+
+	bool operator!=(const FGLTFBinaryHashKey& Other) const
+	{
+		return !CompareBytes(Other);
+	}
+
+	friend uint32 GetTypeHash(const FGLTFBinaryHashKey& Other)
+	{
+		return Other.Hash;
+	}
+
+private:
+
+	bool CompareBytes(const FGLTFBinaryHashKey& Other) const
+	{
+		return Bytes.Num() == Other.Bytes.Num()
+            && FMemory::Memcmp(Bytes.GetData(), Other.Bytes.GetData(), Bytes.Num()) == 0;
+	}
+
+	const uint32 Hash;
+	const TArray64<uint8> Bytes;
+};
diff --git a/Engine/Plugins/Enterprise/GLTFExporter/MigratedFromGit/Source/GLTFExporter/Private/Builders/GLTFImageBuilder.cpp b/Engine/Plugins/Enterprise/GLTFExporter/MigratedFromGit/Source/GLTFExporter/Private/Builders/GLTFImageBuilder.cpp
--- a/Engine/Plugins/Enterprise/GLTFExporter/MigratedFromGit/Source/GLTFExporter/Private/Builders/GLTFImageBuilder.cpp
+++ b/Engine/Plugins/Enterprise/GLTFExporter/MigratedFromGit/Source/GLTFExporter/Private/Builders/GLTFImageBuilder.cpp
@@ -14,15 +14,16 @@ FGLTFImageBuilder::FGLTFImageBuilder(const FString& FilePath, const UGLTFExportO
 
 FGLTFJsonImageIndex FGLTFImageBuilder::AddImage(const void* CompressedData, int64 CompressedByteLength, EGLTFJsonMimeType MimeType, const FString& Name)
 {
-	FGLTFJsonImageIndex ImageIndex = FindImage(CompressedData, CompressedByteLength);
+	// TODO: should this function be renamed to GetOrAddImage?
+
+	const FGLTFBinaryHashKey HashKey(CompressedData, CompressedByteLength);
+	FGLTFJsonImageIndex& ImageIndex = UniqueImageIndices.FindOrAdd(HashKey);
+
 	if (ImageIndex == INDEX_NONE)
 	{
 		FGLTFJsonImage Image;
 		Image.Uri = SaveImageToFile(CompressedData, CompressedByteLength, MimeType, Name);
 		ImageIndex = FGLTFJsonBuilder::AddImage(Image);
-
-		TArray64<uint8>& ImageData = ImageDataLookup.Add(ImageIndex);
-		ImageData.Append(static_cast<const uint8*>(CompressedData), CompressedByteLength);
 	}
 
 	return ImageIndex;
@@ -74,23 +75,3 @@ FString FGLTFImageBuilder::SaveImageToFile(const void* CompressedData, int64 Com
 	UniqueImageUris.Add(ImageUri);
 	return ImageUri;
 }
-
-
-FGLTFJsonImageIndex FGLTFImageBuilder::FindImage(const void* CompressedData, int64 CompressedByteLength) const
-{
-	// TODO: instead of looping through all images, use a hash table to minimize search space
-	for (const auto& DataPair : ImageDataLookup)
-	{
-		const TArray64<uint8>& ImageData = DataPair.Value;
-		if (ImageData.Num() == CompressedByteLength)
-		{
-			if (FMemory::Memcmp(ImageData.GetData(), CompressedData, CompressedByteLength) == 0)
-			{
-				const FGLTFJsonImageIndex ImageIndex = DataPair.Key;
-				return ImageIndex;
-			}
-		}
-	}
-
-	return FGLTFJsonImageIndex(INDEX_NONE);
-}
diff --git a/Engine/Plugins/Enterprise/GLTFExporter/MigratedFromGit/Source/GLTFExporter/Private/Builders/GLTFImageBuilder.h b/Engine/Plugins/Enterprise/GLTFExporter/MigratedFromGit/Source/GLTFExporter/Private/Builders/GLTFImageBuilder.h
--- a/Engine/Plugins/Enterprise/GLTFExporter/MigratedFromGit/Source/GLTFExporter/Private/Builders/GLTFImageBuilder.h
+++ b/Engine/Plugins/Enterprise/GLTFExporter/MigratedFromGit/Source/GLTFExporter/Private/Builders/GLTFImageBuilder.h
@@ -3,6 +3,7 @@
 #pragma once
 
 #include "Builders/GLTFBufferBuilder.h"
+#include "Builders/GLTFBinaryHashKey.h"
 
 enum class ERGBFormat : int8;
 
@@ -23,10 +24,8 @@ private:
 
 	FString SaveImageToFile(const void* CompressedData, int64 CompressedByteLength, EGLTFJsonMimeType MimeType, const FString& Name);
 
-	FGLTFJsonImageIndex FindImage(const void* CompressedData, int64 CompressedByteLength) const;
-
-	TMap<FGLTFJsonImageIndex, TArray64<uint8>> ImageDataLookup;
-
 	const FString ImageDirPath;
+
 	TSet<FString> UniqueImageUris;
+	TMap<FGLTFBinaryHashKey, FGLTFJsonImageIndex> UniqueImageIndices;
 };
