[GLTFExporter] refactor: don't copy buffer data, save continuously to file

commit 8c8de10e4455526d5505e364c09645bed83edd0e
Date: Sun, 27 Sep 2020 16:18:14 +0200

    refactor: don't copy buffer data, save continuously to file
    (cherry picked from commit 25436257f818474544d2b0084e2923b5cfcff904​)

#skipci

diff --git a/Engine/Plugins/Enterprise/GLTFExporter/MigratedFromGit/Source/GLTFExporter/Private/Builders/GLTFBufferBuilder.cpp b/Engine/Plugins/Enterprise/GLTFExporter/MigratedFromGit/Source/GLTFExporter/Private/Builders/GLTFBufferBuilder.cpp
--- a/Engine/Plugins/Enterprise/GLTFExporter/MigratedFromGit/Source/GLTFExporter/Private/Builders/GLTFBufferBuilder.cpp
+++ b/Engine/Plugins/Enterprise/GLTFExporter/MigratedFromGit/Source/GLTFExporter/Private/Builders/GLTFBufferBuilder.cpp
@@ -7,36 +7,48 @@
 FGLTFBufferBuilder::FGLTFBufferBuilder(const FString& FilePath, const UGLTFExportOptions* ExportOptions)
 	: FGLTFJsonBuilder(FilePath, ExportOptions)
 {
-	BufferIndex = AddBuffer();
+	const FString ExternalBinaryPath = FPaths::ChangeExtension(FilePath, TEXT(".bin"));
+	BufferArchive.Reset(IFileManager::Get().CreateFileWriter(*ExternalBinaryPath));
+	if (BufferArchive == nullptr)
+	{
+		// TODO: report error
+	}
+
+	FGLTFJsonBuffer JsonBuffer;
+	JsonBuffer.URI = FPaths::GetCleanFilename(ExternalBinaryPath);
+	BufferIndex = AddBuffer(JsonBuffer);
 }
 
-bool FGLTFBufferBuilder::Serialize(FArchive& Archive)
+FGLTFBufferBuilder::~FGLTFBufferBuilder()
 {
-	const FString BinaryFilePath = FPaths::ChangeExtension(FilePath, TEXT(".bin"));
-
-	if(!FFileHelper::SaveArrayToFile(BufferData, *BinaryFilePath))
+	if (BufferArchive != nullptr)
 	{
-		// TODO: report error
-		return false;
+		BufferArchive->Close();
 	}
-
-	UpdateJsonBufferObject(BinaryFilePath);
-	return true;
 }
 
 FGLTFJsonBufferViewIndex FGLTFBufferBuilder::AddBufferView(const void* RawData, uint64 ByteLength, EGLTFJsonBufferTarget BufferTarget, uint8 DataAlignment)
 {
-	uint64 ByteOffset = BufferData.Num();
+	if (BufferArchive == nullptr)
+	{
+		// TODO: report error
+		return FGLTFJsonBufferViewIndex(INDEX_NONE);
+	}
+
+	uint64 ByteOffset = BufferArchive->Tell();
 
 	// Data offset must be a multiple of the size of the glTF component type (given by ByteAlignment).
-	const uint8 Padding = (DataAlignment - (ByteOffset % DataAlignment)) % DataAlignment;
+	const int64 Padding = (DataAlignment - (ByteOffset % DataAlignment)) % DataAlignment;
 	if (Padding > 0)
 	{
 		ByteOffset += Padding;
-		BufferData.AddZeroed(Padding);
+		BufferArchive->Seek(ByteOffset);
 	}
 
-	BufferData.Append(static_cast<const uint8*>(RawData), ByteLength);
+	BufferArchive->Serialize(const_cast<void*>(RawData), ByteLength);
+
+	FGLTFJsonBuffer& JsonBuffer = GetBuffer(BufferIndex);
+	JsonBuffer.ByteLength = BufferArchive->Tell();
 
 	FGLTFJsonBufferView JsonBufferView;
 	JsonBufferView.Buffer = BufferIndex;
@@ -46,10 +58,3 @@ FGLTFJsonBufferViewIndex FGLTFBufferBuilder::AddBufferView(const void* RawData,
 
 	return FGLTFJsonBuilder::AddBufferView(JsonBufferView);
 }
-
-void FGLTFBufferBuilder::UpdateJsonBufferObject(const FString& BinaryFilePath)
-{
-	FGLTFJsonBuffer& JsonBuffer = GetBuffer(BufferIndex);
-	JsonBuffer.URI = FPaths::GetCleanFilename(BinaryFilePath);
-	JsonBuffer.ByteLength = BufferData.Num();
-}
diff --git a/Engine/Plugins/Enterprise/GLTFExporter/MigratedFromGit/Source/GLTFExporter/Private/Builders/GLTFBufferBuilder.h b/Engine/Plugins/Enterprise/GLTFExporter/MigratedFromGit/Source/GLTFExporter/Private/Builders/GLTFBufferBuilder.h
--- a/Engine/Plugins/Enterprise/GLTFExporter/MigratedFromGit/Source/GLTFExporter/Private/Builders/GLTFBufferBuilder.h
+++ b/Engine/Plugins/Enterprise/GLTFExporter/MigratedFromGit/Source/GLTFExporter/Private/Builders/GLTFBufferBuilder.h
@@ -9,8 +9,7 @@ class FGLTFBufferBuilder : public FGLTFJsonBuilder
 protected:
 
 	FGLTFBufferBuilder(const FString& FilePath, const UGLTFExportOptions* ExportOptions);
-
-	bool Serialize(FArchive& Archive);
+	~FGLTFBufferBuilder();
 
 public:
 
@@ -24,8 +23,6 @@ public:
 
 private:
 
-	void UpdateJsonBufferObject(const FString& BinaryFilePath);
-
 	FGLTFJsonBufferIndex BufferIndex;
-	TArray64<uint8> BufferData;
+	TUniquePtr<FArchive> BufferArchive;
 };
diff --git a/Engine/Plugins/Enterprise/GLTFExporter/MigratedFromGit/Source/GLTFExporter/Private/Builders/GLTFContainerBuilder.cpp b/Engine/Plugins/Enterprise/GLTFExporter/MigratedFromGit/Source/GLTFExporter/Private/Builders/GLTFContainerBuilder.cpp
--- a/Engine/Plugins/Enterprise/GLTFExporter/MigratedFromGit/Source/GLTFExporter/Private/Builders/GLTFContainerBuilder.cpp
+++ b/Engine/Plugins/Enterprise/GLTFExporter/MigratedFromGit/Source/GLTFExporter/Private/Builders/GLTFContainerBuilder.cpp
@@ -9,6 +9,5 @@ FGLTFContainerBuilder::FGLTFContainerBuilder(const FString& FilePath, const UGLT
 
 bool FGLTFContainerBuilder::Serialize(FArchive& Archive)
 {
-	return FGLTFBufferBuilder::Serialize(Archive)
-		&& FGLTFJsonBuilder::Serialize(Archive);
+	return FGLTFJsonBuilder::Serialize(Archive);
 }
