[GLTFExporter] refactor: rename Serialize to Write* in builders and remove bool retval

commit def223c71354c16bec9f1987b828d386c74e2b52
Date: Sun, 27 Sep 2020 20:46:47 +0200

    refactor: rename Serialize to Write* in builders and remove bool retval
    (cherry picked from commit fd3688d65fa5abe99ff61539c5ef0040c2dfb9cf​)

#skipci

diff --git a/Engine/Plugins/Enterprise/GLTFExporter/MigratedFromGit/Source/GLTFExporter/Private/Builders/GLTFContainerBuilder.cpp b/Engine/Plugins/Enterprise/GLTFExporter/MigratedFromGit/Source/GLTFExporter/Private/Builders/GLTFContainerBuilder.cpp
--- a/Engine/Plugins/Enterprise/GLTFExporter/MigratedFromGit/Source/GLTFExporter/Private/Builders/GLTFContainerBuilder.cpp
+++ b/Engine/Plugins/Enterprise/GLTFExporter/MigratedFromGit/Source/GLTFExporter/Private/Builders/GLTFContainerBuilder.cpp
@@ -7,7 +7,7 @@ FGLTFContainerBuilder::FGLTFContainerBuilder(const FString& FilePath, const UGLT
 {
 }
 
-bool FGLTFContainerBuilder::Serialize(FArchive& Archive)
+void FGLTFContainerBuilder::Write(FArchive& Archive) const
 {
-	return FGLTFJsonBuilder::Serialize(Archive);
+	WriteJson(Archive);
 }
diff --git a/Engine/Plugins/Enterprise/GLTFExporter/MigratedFromGit/Source/GLTFExporter/Private/Builders/GLTFContainerBuilder.h b/Engine/Plugins/Enterprise/GLTFExporter/MigratedFromGit/Source/GLTFExporter/Private/Builders/GLTFContainerBuilder.h
--- a/Engine/Plugins/Enterprise/GLTFExporter/MigratedFromGit/Source/GLTFExporter/Private/Builders/GLTFContainerBuilder.h
+++ b/Engine/Plugins/Enterprise/GLTFExporter/MigratedFromGit/Source/GLTFExporter/Private/Builders/GLTFContainerBuilder.h
@@ -10,5 +10,5 @@ public:
 
 	FGLTFContainerBuilder(const FString& FilePath, const UGLTFExportOptions* ExportOptions, bool bSelectedActorsOnly);
 
-	bool Serialize(FArchive& Archive);
+	void Write(FArchive& Archive) const;
 };
diff --git a/Engine/Plugins/Enterprise/GLTFExporter/MigratedFromGit/Source/GLTFExporter/Private/Builders/GLTFJsonBuilder.cpp b/Engine/Plugins/Enterprise/GLTFExporter/MigratedFromGit/Source/GLTFExporter/Private/Builders/GLTFJsonBuilder.cpp
--- a/Engine/Plugins/Enterprise/GLTFExporter/MigratedFromGit/Source/GLTFExporter/Private/Builders/GLTFJsonBuilder.cpp
+++ b/Engine/Plugins/Enterprise/GLTFExporter/MigratedFromGit/Source/GLTFExporter/Private/Builders/GLTFJsonBuilder.cpp
@@ -8,10 +8,9 @@ FGLTFJsonBuilder::FGLTFJsonBuilder(const FString& FilePath, const UGLTFExportOpt
 {
 }
 
-bool FGLTFJsonBuilder::Serialize(FArchive& Archive)
+void FGLTFJsonBuilder::WriteJson(FArchive& Archive) const
 {
-	JsonRoot.Serialize(&Archive, true);
-	return true;
+	JsonRoot.ToJson(&Archive, true);
 }
 
 void FGLTFJsonBuilder::AddExtension(EGLTFJsonExtension Extension, bool bIsRequired)
diff --git a/Engine/Plugins/Enterprise/GLTFExporter/MigratedFromGit/Source/GLTFExporter/Private/Builders/GLTFJsonBuilder.h b/Engine/Plugins/Enterprise/GLTFExporter/MigratedFromGit/Source/GLTFExporter/Private/Builders/GLTFJsonBuilder.h
--- a/Engine/Plugins/Enterprise/GLTFExporter/MigratedFromGit/Source/GLTFExporter/Private/Builders/GLTFJsonBuilder.h
+++ b/Engine/Plugins/Enterprise/GLTFExporter/MigratedFromGit/Source/GLTFExporter/Private/Builders/GLTFJsonBuilder.h
@@ -11,7 +11,7 @@ protected:
 
 	FGLTFJsonBuilder(const FString& FilePath, const UGLTFExportOptions* ExportOptions);
 
-	bool Serialize(FArchive& Archive);
+	void WriteJson(FArchive& Archive) const;
 
 public:
 
diff --git a/Engine/Plugins/Enterprise/GLTFExporter/MigratedFromGit/Source/GLTFExporter/Private/Exporters/GLTFExporter.cpp b/Engine/Plugins/Enterprise/GLTFExporter/MigratedFromGit/Source/GLTFExporter/Private/Exporters/GLTFExporter.cpp
--- a/Engine/Plugins/Enterprise/GLTFExporter/MigratedFromGit/Source/GLTFExporter/Private/Exporters/GLTFExporter.cpp
+++ b/Engine/Plugins/Enterprise/GLTFExporter/MigratedFromGit/Source/GLTFExporter/Private/Exporters/GLTFExporter.cpp
@@ -31,26 +31,20 @@ bool UGLTFExporter::ExportBinary(UObject* Object, const TCHAR* Type, FArchive& A
 		return false;
 	}
 
-	FGCObjectScopeGuard OptionsGuard(Options);
-	bool bSuccess = true;
-
 	// TODO: add support for UAssetExportTask::IgnoreObjectList?
 
+	FGCObjectScopeGuard OptionsGuard(Options);
 	FGLTFContainerBuilder Builder(CurrentFilename, Options, bSelectedOnly);
-	if (AddObject(Builder, Object))
+
+	const bool bSuccess = AddObject(Builder, Object);
+	if (bSuccess)
 	{
-		if (!Builder.Serialize(Archive))
-		{
-			// TODO: more descriptive error
-			Builder.AddErrorMessage(TEXT("Serialize failed"));
-			bSuccess = false;
-		}
+		Builder.Write(Archive);
 	}
 	else
 	{
 		// TODO: more descriptive error
 		Builder.AddErrorMessage(TEXT("AddObject failed"));
-		bSuccess = false;
 	}
 
 	// TODO: should we copy messages to UAssetExportTask::Errors?
diff --git a/Engine/Plugins/Enterprise/GLTFExporter/MigratedFromGit/Source/GLTFExporter/Private/Json/GLTFJsonRoot.h b/Engine/Plugins/Enterprise/GLTFExporter/MigratedFromGit/Source/GLTFExporter/Private/Json/GLTFJsonRoot.h
--- a/Engine/Plugins/Enterprise/GLTFExporter/MigratedFromGit/Source/GLTFExporter/Private/Json/GLTFJsonRoot.h
+++ b/Engine/Plugins/Enterprise/GLTFExporter/MigratedFromGit/Source/GLTFExporter/Private/Json/GLTFJsonRoot.h
@@ -174,22 +174,22 @@ struct FGLTFJsonRoot
 	}
 
 	template <class CharType = TCHAR, class PrintPolicy = TPrettyJsonPrintPolicy<CharType>>
-	void Serialize(FArchive* const Archive) const
+	void ToJson(FArchive* const Archive) const
 	{
 		TSharedRef<TJsonWriter<CharType, PrintPolicy>> JsonWriter = TJsonWriterFactory<CharType, PrintPolicy>::Create(Archive);
 		WriteObject(*JsonWriter);
 		JsonWriter->Close();
 	}
 
-	void Serialize(FArchive* const Archive, bool bPrettyPrint = true) const
+	void ToJson(FArchive* const Archive, bool bPrettyPrint = true) const
 	{
 		if (bPrettyPrint)
 		{
-			Serialize<UTF8CHAR, TPrettyJsonPrintPolicy<UTF8CHAR>>(Archive);
+			ToJson<UTF8CHAR, TPrettyJsonPrintPolicy<UTF8CHAR>>(Archive);
 		}
 		else
 		{
-			Serialize<UTF8CHAR, TCondensedJsonPrintPolicy<UTF8CHAR>>(Archive);
+			ToJson<UTF8CHAR, TCondensedJsonPrintPolicy<UTF8CHAR>>(Archive);
 		}
 	}
 };
