[GLTFExporter] feat: register .glb file extension in exporter

commit eb067d0e7b46418219754100e79d9022333bfe71
Date: Sun, 27 Sep 2020 20:49:53 +0200

    feat: register .glb file extension in exporter
    (cherry picked from commit 298b52964243af754aaf49a0adfaf790ae5224aa​)

#skipci

diff --git a/Engine/Plugins/Enterprise/GLTFExporter/MigratedFromGit/Source/GLTFExporter/Private/Exporters/GLTFExporter.cpp b/Engine/Plugins/Enterprise/GLTFExporter/MigratedFromGit/Source/GLTFExporter/Private/Exporters/GLTFExporter.cpp
--- a/Engine/Plugins/Enterprise/GLTFExporter/MigratedFromGit/Source/GLTFExporter/Private/Exporters/GLTFExporter.cpp
+++ b/Engine/Plugins/Enterprise/GLTFExporter/MigratedFromGit/Source/GLTFExporter/Private/Exporters/GLTFExporter.cpp
@@ -17,9 +17,8 @@ UGLTFExporter::UGLTFExporter(const FObjectInitializer& ObjectInitializer)
 	FormatExtension.Add(TEXT("gltf"));
 	FormatDescription.Add(TEXT("GL Transmission Format"));
 
-	// TODO: uncomment when support for .glb is implemented
-	// FormatExtension.Add(TEXT("glb"));
-	// FormatDescription.Add(TEXT("GL Transmission Format (Binary)"));
+	FormatExtension.Add(TEXT("glb"));
+	FormatDescription.Add(TEXT("GL Transmission Format (Binary)"));
 }
 
 bool UGLTFExporter::ExportBinary(UObject* Object, const TCHAR* Type, FArchive& Archive, FFeedbackContext* Warn, int32 FileIndex, uint32 PortFlags)
