[GLTFExporter] refactor: add GLTFBuilder::bIsGlbFile and GLTFBuilderUtility::IsGlbFile

commit ee35a89543d0e563c583813b797fce54c691c692
Date: Sun, 27 Sep 2020 20:51:23 +0200

    refactor: add GLTFBuilder::bIsGlbFile and GLTFBuilderUtility::IsGlbFile
    (cherry picked from commit 109cee450ebee425e1114a8bc26674a967634050​)

#skipci

diff --git a/Engine/Plugins/Enterprise/GLTFExporter/MigratedFromGit/Source/GLTFExporter/Private/Builders/GLTFBuilder.cpp b/Engine/Plugins/Enterprise/GLTFExporter/MigratedFromGit/Source/GLTFExporter/Private/Builders/GLTFBuilder.cpp
--- a/Engine/Plugins/Enterprise/GLTFExporter/MigratedFromGit/Source/GLTFExporter/Private/Builders/GLTFBuilder.cpp
+++ b/Engine/Plugins/Enterprise/GLTFExporter/MigratedFromGit/Source/GLTFExporter/Private/Builders/GLTFBuilder.cpp
@@ -1,9 +1,11 @@
 // Copyright Epic Games, Inc. All Rights Reserved.
 
 #include "Builders/GLTFBuilder.h"
+#include "Builders/GLTFBuilderUtility.h"
 
 FGLTFBuilder::FGLTFBuilder(const FString& FilePath, const UGLTFExportOptions* ExportOptions)
-	: FilePath(FilePath)
+	: bIsGlbFile(FGLTFBuilderUtility::IsGlbFile(FilePath))
+	, FilePath(FilePath)
 	, ExportOptions(ExportOptions)
 {
 }
diff --git a/Engine/Plugins/Enterprise/GLTFExporter/MigratedFromGit/Source/GLTFExporter/Private/Builders/GLTFBuilder.h b/Engine/Plugins/Enterprise/GLTFExporter/MigratedFromGit/Source/GLTFExporter/Private/Builders/GLTFBuilder.h
--- a/Engine/Plugins/Enterprise/GLTFExporter/MigratedFromGit/Source/GLTFExporter/Private/Builders/GLTFBuilder.h
+++ b/Engine/Plugins/Enterprise/GLTFExporter/MigratedFromGit/Source/GLTFExporter/Private/Builders/GLTFBuilder.h
@@ -12,6 +12,8 @@ protected:
 
 public:
 
+	const bool bIsGlbFile;
+
 	const FString FilePath;
 
 	const UGLTFExportOptions* const ExportOptions;
diff --git a/Engine/Plugins/Enterprise/GLTFExporter/MigratedFromGit/Source/GLTFExporter/Private/Builders/GLTFBuilderUtility.cpp b/Engine/Plugins/Enterprise/GLTFExporter/MigratedFromGit/Source/GLTFExporter/Private/Builders/GLTFBuilderUtility.cpp
--- a/Engine/Plugins/Enterprise/GLTFExporter/MigratedFromGit/Source/GLTFExporter/Private/Builders/GLTFBuilderUtility.cpp
+++ b/Engine/Plugins/Enterprise/GLTFExporter/MigratedFromGit/Source/GLTFExporter/Private/Builders/GLTFBuilderUtility.cpp
@@ -30,3 +30,8 @@ FString FGLTFBuilderUtility::GetUniqueFilename(const FString& BaseFilename, cons
 
 	return Filename;
 }
+
+bool FGLTFBuilderUtility::IsGlbFile(const FString& Filename)
+{
+	return FPaths::GetExtension(Filename).Equals(TEXT("glb"), ESearchCase::IgnoreCase);
+}
diff --git a/Engine/Plugins/Enterprise/GLTFExporter/MigratedFromGit/Source/GLTFExporter/Private/Builders/GLTFBuilderUtility.h b/Engine/Plugins/Enterprise/GLTFExporter/MigratedFromGit/Source/GLTFExporter/Private/Builders/GLTFBuilderUtility.h
--- a/Engine/Plugins/Enterprise/GLTFExporter/MigratedFromGit/Source/GLTFExporter/Private/Builders/GLTFBuilderUtility.h
+++ b/Engine/Plugins/Enterprise/GLTFExporter/MigratedFromGit/Source/GLTFExporter/Private/Builders/GLTFBuilderUtility.h
@@ -10,4 +10,6 @@ struct FGLTFBuilderUtility
 	static const TCHAR* GetFileExtension(EGLTFJsonMimeType MimeType);
 
 	static FString GetUniqueFilename(const FString& BaseFilename, const FString& FileExtension, const TSet<FString>& UniqueFilenames);
+
+	static bool IsGlbFile(const FString& Filename);
 };
