[GLTFExporter] refactor: add support for GLB in GLTFBufferBuilder

commit 83318dc9cce465bcbae00ecab92c8eacc1c9dcd8
Date: Sun, 27 Sep 2020 20:53:01 +0200

    refactor: add support for GLB in GLTFBufferBuilder
    (cherry picked from commit f33e42c5e196b3cff55c600df813c8ee07f3088f​)

#skipci

diff --git a/Engine/Plugins/Enterprise/GLTFExporter/MigratedFromGit/Source/GLTFExporter/Private/Builders/GLTFBufferBuilder.cpp b/Engine/Plugins/Enterprise/GLTFExporter/MigratedFromGit/Source/GLTFExporter/Private/Builders/GLTFBufferBuilder.cpp
--- a/Engine/Plugins/Enterprise/GLTFExporter/MigratedFromGit/Source/GLTFExporter/Private/Builders/GLTFBufferBuilder.cpp
+++ b/Engine/Plugins/Enterprise/GLTFExporter/MigratedFromGit/Source/GLTFExporter/Private/Builders/GLTFBufferBuilder.cpp
@@ -1,21 +1,31 @@
 // Copyright Epic Games, Inc. All Rights Reserved.
 
 #include "Builders/GLTFBufferBuilder.h"
+#include "Serialization/BufferArchive.h"
 #include "Misc/FileHelper.h"
 #include "Misc/Paths.h"
 
 FGLTFBufferBuilder::FGLTFBufferBuilder(const FString& FilePath, const UGLTFExportOptions* ExportOptions)
 	: FGLTFJsonBuilder(FilePath, ExportOptions)
 {
-	const FString ExternalBinaryPath = FPaths::ChangeExtension(FilePath, TEXT(".bin"));
-	BufferArchive.Reset(IFileManager::Get().CreateFileWriter(*ExternalBinaryPath));
-	if (BufferArchive == nullptr)
+	FGLTFJsonBuffer JsonBuffer;
+
+	if (bIsGlbFile)
 	{
-		// TODO: report error
+		BufferArchive = MakeUnique<FBufferArchive>();
+	}
+	else
+	{
+		const FString ExternalBinaryPath = FPaths::ChangeExtension(FilePath, TEXT(".bin"));
+		JsonBuffer.URI = FPaths::GetCleanFilename(ExternalBinaryPath);
+
+		BufferArchive.Reset(IFileManager::Get().CreateFileWriter(*ExternalBinaryPath));
+		if (BufferArchive == nullptr)
+		{
+			// TODO: report error
+		}
 	}
 
-	FGLTFJsonBuffer JsonBuffer;
-	JsonBuffer.URI = FPaths::GetCleanFilename(ExternalBinaryPath);
 	BufferIndex = AddBuffer(JsonBuffer);
 }
 
