[GLTFExporter] refactor: add support for GLB in GLTFImageBuilder

commit 739d59fd489f1bbb9498697bc28061240227af4d
Date: Sun, 27 Sep 2020 20:53:19 +0200

    refactor: add support for GLB in GLTFImageBuilder
    (cherry picked from commit c281aa4077f91a92aa44b8164ccc6affc0ae00ae​)

#skipci

diff --git a/Engine/Plugins/Enterprise/GLTFExporter/MigratedFromGit/Source/GLTFExporter/Private/Builders/GLTFImageBuilder.cpp b/Engine/Plugins/Enterprise/GLTFExporter/MigratedFromGit/Source/GLTFExporter/Private/Builders/GLTFImageBuilder.cpp
--- a/Engine/Plugins/Enterprise/GLTFExporter/MigratedFromGit/Source/GLTFExporter/Private/Builders/GLTFImageBuilder.cpp
+++ b/Engine/Plugins/Enterprise/GLTFExporter/MigratedFromGit/Source/GLTFExporter/Private/Builders/GLTFImageBuilder.cpp
@@ -21,9 +21,20 @@ FGLTFJsonImageIndex FGLTFImageBuilder::AddImage(const void* CompressedData, int6
 
 	if (ImageIndex == INDEX_NONE)
 	{
-		FGLTFJsonImage Image;
-		Image.Uri = SaveImageToFile(CompressedData, CompressedByteLength, MimeType, Name);
-		ImageIndex = FGLTFJsonBuilder::AddImage(Image);
+		FGLTFJsonImage JsonImage;
+
+		if (bIsGlbFile)
+		{
+			JsonImage.Name = Name;
+			JsonImage.MimeType = MimeType;
+			JsonImage.BufferView = AddBufferView(CompressedData, CompressedByteLength);
+		}
+		else
+		{
+			JsonImage.Uri = SaveImageToFile(CompressedData, CompressedByteLength, MimeType, Name);
+		}
+
+		ImageIndex = FGLTFJsonBuilder::AddImage(JsonImage);
 	}
 
 	return ImageIndex;
