[GLTFExporter] refactor: add GLTFBufferBuilder::GetBufferData

commit 69a4e7e31515af71f07dae240bd0aad1bef6bd0b
Date: Sun, 27 Sep 2020 20:53:55 +0200

    refactor: add GLTFBufferBuilder::GetBufferData
    (cherry picked from commit d60b141a446d5cce8a3a7fb644a2347251cfad78​)

#skipci

diff --git a/Engine/Plugins/Enterprise/GLTFExporter/MigratedFromGit/Source/GLTFExporter/Private/Builders/GLTFBufferBuilder.cpp b/Engine/Plugins/Enterprise/GLTFExporter/MigratedFromGit/Source/GLTFExporter/Private/Builders/GLTFBufferBuilder.cpp
--- a/Engine/Plugins/Enterprise/GLTFExporter/MigratedFromGit/Source/GLTFExporter/Private/Builders/GLTFBufferBuilder.cpp
+++ b/Engine/Plugins/Enterprise/GLTFExporter/MigratedFromGit/Source/GLTFExporter/Private/Builders/GLTFBufferBuilder.cpp
@@ -37,6 +37,12 @@ FGLTFBufferBuilder::~FGLTFBufferBuilder()
 	}
 }
 
+const TArray<uint8>& FGLTFBufferBuilder::GetBufferData() const
+{
+	check(bIsGlbFile);
+	return *static_cast<FBufferArchive*>(BufferArchive.Get());
+}
+
 FGLTFJsonBufferViewIndex FGLTFBufferBuilder::AddBufferView(const void* RawData, uint64 ByteLength, EGLTFJsonBufferTarget BufferTarget, uint8 DataAlignment)
 {
 	if (BufferArchive == nullptr)
diff --git a/Engine/Plugins/Enterprise/GLTFExporter/MigratedFromGit/Source/GLTFExporter/Private/Builders/GLTFBufferBuilder.h b/Engine/Plugins/Enterprise/GLTFExporter/MigratedFromGit/Source/GLTFExporter/Private/Builders/GLTFBufferBuilder.h
--- a/Engine/Plugins/Enterprise/GLTFExporter/MigratedFromGit/Source/GLTFExporter/Private/Builders/GLTFBufferBuilder.h
+++ b/Engine/Plugins/Enterprise/GLTFExporter/MigratedFromGit/Source/GLTFExporter/Private/Builders/GLTFBufferBuilder.h
@@ -11,6 +11,8 @@ protected:
 	FGLTFBufferBuilder(const FString& FilePath, const UGLTFExportOptions* ExportOptions);
 	~FGLTFBufferBuilder();
 
+	const TArray<uint8>& GetBufferData() const;
+
 public:
 
 	FGLTFJsonBufferViewIndex AddBufferView(const void* RawData, uint64 ByteLength, EGLTFJsonBufferTarget BufferTarget = EGLTFJsonBufferTarget::None, uint8 DataAlignment = 4);
