[GLTFExporter] refactor: add GLTFContainerUtility

commit cf717ac96b6ed50ddc23e7d57882dc0642ecf23a
Date: Sun, 27 Sep 2020 20:54:09 +0200

    refactor: add GLTFContainerUtility
    (cherry picked from commit cd5c14a83f3bb8d71379171c2322118c7f341453​)

#skipci

diff --git a/Engine/Plugins/Enterprise/GLTFExporter/MigratedFromGit/Source/GLTFExporter/Private/Builders/GLTFContainerUtility.cpp b/Engine/Plugins/Enterprise/GLTFExporter/MigratedFromGit/Source/GLTFExporter/Private/Builders/GLTFContainerUtility.cpp
new file mode 100644
--- /dev/null
+++ b/Engine/Plugins/Enterprise/GLTFExporter/MigratedFromGit/Source/GLTFExporter/Private/Builders/GLTFContainerUtility.cpp
@@ -0,0 +1,66 @@
+// Copyright Epic Games, Inc. All Rights Reserved.
+
+#include "Builders/GLTFContainerUtility.h"
+
+void FGLTFContainerUtility::WriteGlb(FArchive& Archive, const TArray<uint8>& JsonData, const TArray<uint8>& BinaryData)
+{
+	const uint32 JsonChunkType = 0x4E4F534A; // "JSON" in ASCII
+	const uint32 BinaryChunkType = 0x004E4942; // "BIN" in ASCII
+	const uint32 FileSize =
+		3 * sizeof(uint32) +
+		2 * sizeof(uint32) + GetPaddedChunkSize(JsonData.Num()) +
+		2 * sizeof(uint32) + GetPaddedChunkSize(BinaryData.Num());
+
+	WriteHeader(Archive, FileSize);
+	WriteChunk(Archive, JsonChunkType, JsonData, 0x20);
+	WriteChunk(Archive, BinaryChunkType, BinaryData, 0x0);
+}
+
+void FGLTFContainerUtility::WriteHeader(FArchive& Archive, uint32 FileSize)
+{
+	const uint32 FileSignature = 0x46546C67; // "glTF" in ASCII
+	const uint32 FileVersion = 2;
+
+	WriteInt(Archive, FileSignature);
+	WriteInt(Archive, FileVersion);
+	WriteInt(Archive, FileSize);
+}
+
+void FGLTFContainerUtility::WriteChunk(FArchive& Archive, uint32 ChunkType, const TArray<uint8>& ChunkData, uint8 ChunkTrailingByte)
+{
+	const uint32 ChunkLength = GetPaddedChunkSize(ChunkData.Num());
+	const uint32 ChunkTrailing = GetTrailingChunkSize(ChunkData.Num());
+
+	WriteInt(Archive, ChunkLength);
+	WriteInt(Archive, ChunkType);
+	WriteData(Archive, ChunkData);
+	WriteFill(Archive, ChunkTrailing, ChunkTrailingByte);
+}
+
+void FGLTFContainerUtility::WriteInt(FArchive& Archive, uint32 Value)
+{
+	Archive.SerializeInt(Value, MAX_uint32);
+}
+
+void FGLTFContainerUtility::WriteData(FArchive& Archive, const TArray<uint8>& Data)
+{
+	Archive.Serialize(const_cast<uint8*>(Data.GetData()), Data.Num());
+}
+
+void FGLTFContainerUtility::WriteFill(FArchive& Archive, int32 Size, uint8 Value)
+{
+	while (--Size >= 0)
+	{
+		Archive.Serialize(&Value, sizeof(Value));
+	}
+}
+
+int32 FGLTFContainerUtility::GetPaddedChunkSize(int32 Size)
+{
+	return (Size + 3) & ~3;
+}
+
+int32 FGLTFContainerUtility::GetTrailingChunkSize(int32 Size)
+{
+	return (4 - (Size & 3)) & 3;
+}
diff --git a/Engine/Plugins/Enterprise/GLTFExporter/MigratedFromGit/Source/GLTFExporter/Private/Builders/GLTFContainerUtility.h b/Engine/Plugins/Enterprise/GLTFExporter/MigratedFromGit/Source/GLTFExporter/Private/Builders/GLTFContainerUtility.h
new file mode 100644
--- /dev/null
+++ b/Engine/Plugins/Enterprise/GLTFExporter/MigratedFromGit/Source/GLTFExporter/Private/Builders/GLTFContainerUtility.h
@@ -0,0 +1,24 @@
+// Copyright Epic Games, Inc. All Rights Reserved.
+
+#pragma once
+
+#include "CoreMinimal.h"
+
+class FGLTFContainerUtility
+{
+public:
+
+	static void WriteGlb(FArchive& Archive, const TArray<uint8>& JsonData, const TArray<uint8>& BinaryData);
+
+private:
+
+	static void WriteHeader(FArchive& Archive, uint32 FileSize);
+	static void WriteChunk(FArchive& Archive, uint32 ChunkType, const TArray<uint8>& ChunkData, uint8 ChunkTrailingByte);
+
+	static void WriteInt(FArchive& Archive, uint32 Value);
+	static void WriteData(FArchive& Archive, const TArray<uint8>& Data);
+	static void WriteFill(FArchive& Archive, int32 Size, uint8 Value);
+
+	static int32 GetPaddedChunkSize(int32 Size);
+	static int32 GetTrailingChunkSize(int32 Size);
+};
