[GLTFExporter] refactor: use GLTFContainerUtility in GLTFContainerBuilder

commit 95c50f3c707b797445feb944e7da1d17b851d2f7
Date: Sun, 27 Sep 2020 20:54:33 +0200

    refactor: use GLTFContainerUtility in GLTFContainerBuilder
    (cherry picked from commit 044912de7d5fc44e4a6e914960ca029a4a841a6c​)

#skipci

diff --git a/Engine/Plugins/Enterprise/GLTFExporter/MigratedFromGit/Source/GLTFExporter/Private/Builders/GLTFContainerBuilder.cpp b/Engine/Plugins/Enterprise/GLTFExporter/MigratedFromGit/Source/GLTFExporter/Private/Builders/GLTFContainerBuilder.cpp
--- a/Engine/Plugins/Enterprise/GLTFExporter/MigratedFromGit/Source/GLTFExporter/Private/Builders/GLTFContainerBuilder.cpp
+++ b/Engine/Plugins/Enterprise/GLTFExporter/MigratedFromGit/Source/GLTFExporter/Private/Builders/GLTFContainerBuilder.cpp
@@ -1,13 +1,31 @@
 // Copyright Epic Games, Inc. All Rights Reserved.
 
 #include "Builders/GLTFContainerBuilder.h"
+#include "Builders/GLTFContainerUtility.h"
 
 FGLTFContainerBuilder::FGLTFContainerBuilder(const FString& FilePath, const UGLTFExportOptions* ExportOptions, bool bSelectedActorsOnly)
 	: FGLTFConvertBuilder(FilePath, ExportOptions, bSelectedActorsOnly)
 {
 }
 
+void FGLTFContainerBuilder::WriteGlb(FArchive& Archive) const
+{
+	FBufferArchive JsonData;
+	WriteJson(JsonData);
+
+	const TArray<uint8>& BufferData = GetBufferData();
+
+	FGLTFContainerUtility::WriteGlb(Archive, JsonData, BufferData);
+}
+
 void FGLTFContainerBuilder::Write(FArchive& Archive) const
 {
-	WriteJson(Archive);
+	if (bIsGlbFile)
+	{
+		WriteGlb(Archive);
+	}
+	else
+	{
+		WriteJson(Archive);
+	}
 }
diff --git a/Engine/Plugins/Enterprise/GLTFExporter/MigratedFromGit/Source/GLTFExporter/Private/Builders/GLTFContainerBuilder.h b/Engine/Plugins/Enterprise/GLTFExporter/MigratedFromGit/Source/GLTFExporter/Private/Builders/GLTFContainerBuilder.h
--- a/Engine/Plugins/Enterprise/GLTFExporter/MigratedFromGit/Source/GLTFExporter/Private/Builders/GLTFContainerBuilder.h
+++ b/Engine/Plugins/Enterprise/GLTFExporter/MigratedFromGit/Source/GLTFExporter/Private/Builders/GLTFContainerBuilder.h
@@ -11,4 +11,8 @@ public:
 	FGLTFContainerBuilder(const FString& FilePath, const UGLTFExportOptions* ExportOptions, bool bSelectedActorsOnly);
 
 	void Write(FArchive& Archive) const;
+
+private:
+
+	void WriteGlb(FArchive& Archive) const;
 };
