[GLTFExporter] fix: add missing headers and forward declarations to compile using RunUAT BuildPlugin

commit 390e65d10dd481475be02ffa6dd015cf8abbef90
Date: Wed, 30 Sep 2020 19:27:15 +0200

    fix: add missing headers and forward declarations to compile using RunUAT BuildPlugin
    (cherry picked from commit 0b554a79bba54ad5a4c71dc45008178d3de5dedf​)

#skipci

diff --git a/Engine/Plugins/Enterprise/GLTFExporter/MigratedFromGit/Source/GLTFExporterRuntime/Private/Actors/GLTFInteractionHotspotActor.cpp b/Engine/Plugins/Enterprise/GLTFExporter/MigratedFromGit/Source/GLTFExporterRuntime/Private/Actors/GLTFInteractionHotspotActor.cpp
--- a/Engine/Plugins/Enterprise/GLTFExporter/MigratedFromGit/Source/GLTFExporterRuntime/Private/Actors/GLTFInteractionHotspotActor.cpp
+++ b/Engine/Plugins/Enterprise/GLTFExporter/MigratedFromGit/Source/GLTFExporterRuntime/Private/Actors/GLTFInteractionHotspotActor.cpp
@@ -8,6 +8,7 @@
 #include "Kismet/GameplayStatics.h"
 #include "UObject/ConstructorHelpers.h"
 #include "Animation/AnimSingleNodeInstance.h"
+#include "Materials/MaterialInstanceDynamic.h"
 
 #if WITH_EDITOR
 #include "LevelEditor.h"
diff --git a/Engine/Plugins/Enterprise/GLTFExporter/MigratedFromGit/Source/GLTFExporterRuntime/Public/Actors/GLTFInteractionHotspotActor.h b/Engine/Plugins/Enterprise/GLTFExporter/MigratedFromGit/Source/GLTFExporterRuntime/Public/Actors/GLTFInteractionHotspotActor.h
--- a/Engine/Plugins/Enterprise/GLTFExporter/MigratedFromGit/Source/GLTFExporterRuntime/Public/Actors/GLTFInteractionHotspotActor.h
+++ b/Engine/Plugins/Enterprise/GLTFExporter/MigratedFromGit/Source/GLTFExporterRuntime/Public/Actors/GLTFInteractionHotspotActor.h
@@ -2,11 +2,18 @@
 
 #pragma once
 
+#include "GameFramework/Actor.h"
 #include "GLTFInteractionHotspotActor.generated.h"
 
 class ASkeletalMeshActor;
+class UAnimSequence;
+class UTexture2D;
+class UTexture;
+class UMaterialInterface;
 class UMaterialBillboardComponent;
 class USphereComponent;
+class UMaterialInstanceDynamic;
+class FViewport;
 
 /**
  * Actor wrapper for the GLTF hotspot component. Appears as a billboard and allows playback of skeletal animations when cursor input is enabled.
