[GLTFExporter] fix: add missing byte stride to 16-bit quantized vertex normals

commit 8366fd929815edf9a928003f17ee7ff070bbc2b9
Date: Wed, 30 Sep 2020 19:55:16 +0200

    fix: add missing byte stride to 16-bit quantized vertex normals
    (cherry picked from commit 0ea87065f5a8261f94285a14d2d793d2653077aa​)

#skipci

diff --git a/Engine/Plugins/Enterprise/GLTFExporter/MigratedFromGit/Source/GLTFExporter/Private/Converters/GLTFAccessorConverters.cpp b/Engine/Plugins/Enterprise/GLTFExporter/MigratedFromGit/Source/GLTFExporter/Private/Converters/GLTFAccessorConverters.cpp
--- a/Engine/Plugins/Enterprise/GLTFExporter/MigratedFromGit/Source/GLTFExporter/Private/Converters/GLTFAccessorConverters.cpp
+++ b/Engine/Plugins/Enterprise/GLTFExporter/MigratedFromGit/Source/GLTFExporter/Private/Converters/GLTFAccessorConverters.cpp
@@ -145,12 +145,13 @@ FGLTFJsonAccessorIndex FGLTFNormalBufferConverter::Convert(const FGLTFMeshSectio
 		{
 			ComponentType = EGLTFJsonComponentType::S16;
 			BufferViewIndex = ConvertBufferView<FGLTFPackedVector16, FPackedRGBA16N>(MeshSection, VertexBuffer);
+			Builder.GetBufferView(BufferViewIndex).ByteStride = sizeof(FGLTFPackedVector16);
 		}
 		else
 		{
 			ComponentType = EGLTFJsonComponentType::S8;
 			BufferViewIndex = ConvertBufferView<FGLTFPackedVector8, FPackedNormal>(MeshSection, VertexBuffer);
-			Builder.GetBufferView(BufferViewIndex).ByteStride = 4;
+			Builder.GetBufferView(BufferViewIndex).ByteStride = sizeof(FGLTFPackedVector8);
 		}
 	}
 	else
