[GLTFExporter] feat: precompile runtime-module when built using UnrealBuildTool / Build.bat

commit 2c87892744839dc6607c8f2fd77cce9fea86fe4b
Date: Thu, 1 Oct 2020 08:56:32 +0200

    feat: precompile runtime-module when built using UnrealBuildTool / Build.bat
    (cherry picked from commit 81014455a8c5be6609f6bbdb8121c65434233838​)

#skipci

diff --git a/Engine/Plugins/Enterprise/GLTFExporter/MigratedFromGit/Source/GLTFExporterRuntime/GLTFExporterRuntime.Build.cs b/Engine/Plugins/Enterprise/GLTFExporter/MigratedFromGit/Source/GLTFExporterRuntime/GLTFExporterRuntime.Build.cs
--- a/Engine/Plugins/Enterprise/GLTFExporter/MigratedFromGit/Source/GLTFExporterRuntime/GLTFExporterRuntime.Build.cs
+++ b/Engine/Plugins/Enterprise/GLTFExporter/MigratedFromGit/Source/GLTFExporterRuntime/GLTFExporterRuntime.Build.cs
@@ -6,6 +6,9 @@ namespace UnrealBuildTool.Rules
 	{
 		public GLTFExporterRuntime(ReadOnlyTargetRules Target) : base(Target)
 		{
+			PCHUsage = PCHUsageMode.NoPCHs;
+			bPrecompile = true;
+
 			PublicIncludePaths.AddRange(
 				new string[] {
 				}
