[GLTFExporter] feat: modify IsRootActor to handle bSelectedOnly correctly

commit 2fa7871d5239fb09243309374d100609c716d153
Date: Thu, 1 Oct 2020 11:14:03 +0200

    feat: modify IsRootActor to handle bSelectedOnly correctly
    we do not care if the parent actor is selected, all selected actors should be considered root actors
    (cherry picked from commit 02579e050a870d59baf54d7a8c38b3864936b2c1​)

#skipci

diff --git a/Engine/Plugins/Enterprise/GLTFExporter/MigratedFromGit/Source/GLTFExporter/Private/Converters/GLTFActorUtility.cpp b/Engine/Plugins/Enterprise/GLTFExporter/MigratedFromGit/Source/GLTFExporter/Private/Converters/GLTFActorUtility.cpp
--- a/Engine/Plugins/Enterprise/GLTFExporter/MigratedFromGit/Source/GLTFExporter/Private/Converters/GLTFActorUtility.cpp
+++ b/Engine/Plugins/Enterprise/GLTFExporter/MigratedFromGit/Source/GLTFExporter/Private/Converters/GLTFActorUtility.cpp
@@ -5,8 +5,7 @@
 
 bool FGLTFActorUtility::IsRootActor(const AActor* Actor, bool bSelectedOnly)
 {
-	const AActor* ParentActor = Actor->GetAttachParentActor();
-	return ParentActor == nullptr || (bSelectedOnly && !ParentActor->IsSelected());
+	return bSelectedOnly ? Actor->IsSelected() : Actor->GetAttachParentActor() == nullptr;
 }
 
 UBlueprint* FGLTFActorUtility::GetBlueprintFromActor(const AActor* Actor)
