[GLTFExporter] feat: skip components for unselected actors (when only exporting selected actors)

commit f6e3efeacc054d29a8fa16265bacb02b0b5523a1
Date: Thu, 1 Oct 2020 11:15:49 +0200

    feat: skip components for unselected actors (when only exporting selected actors)
    we need this test since scene-components may be added without going through the actor-converter
    (cherry picked from commit d184b05dd589bf161fbee1eb188c322603b2bc82​)

#skipci

diff --git a/Engine/Plugins/Enterprise/GLTFExporter/MigratedFromGit/Source/GLTFExporter/Private/Converters/GLTFNodeConverters.cpp b/Engine/Plugins/Enterprise/GLTFExporter/MigratedFromGit/Source/GLTFExporter/Private/Converters/GLTFNodeConverters.cpp
--- a/Engine/Plugins/Enterprise/GLTFExporter/MigratedFromGit/Source/GLTFExporter/Private/Converters/GLTFNodeConverters.cpp
+++ b/Engine/Plugins/Enterprise/GLTFExporter/MigratedFromGit/Source/GLTFExporter/Private/Converters/GLTFNodeConverters.cpp
@@ -32,6 +32,11 @@ FGLTFJsonNodeIndex FGLTFComponentConverter::Convert(const USceneComponent* Scene
 		return FGLTFJsonNodeIndex(INDEX_NONE);
 	}
 
+	if (Builder.bSelectedActorsOnly && !Owner->IsSelected())
+	{
+		return FGLTFJsonNodeIndex(INDEX_NONE);
+	}
+
 	const bool bIsRootComponent = Owner->GetRootComponent() == SceneComponent;
 	const bool bIsRootNode = bIsRootComponent && FGLTFActorUtility::IsRootActor(Owner, Builder.bSelectedActorsOnly);
 
