[GLTFExporter] feat: avoid exporting parent component for root-nodes

commit 2688fb243b1ba5f3ebe72e0bfec5586e7f1af840
Date: Thu, 1 Oct 2020 11:17:20 +0200

    feat: avoid exporting parent component for root-nodes
    (cherry picked from commit 2e8e1248682fe7e2ad6d7671e0170f8cd6517ec7​)

#skipci

diff --git a/Engine/Plugins/Enterprise/GLTFExporter/MigratedFromGit/Source/GLTFExporter/Private/Converters/GLTFNodeConverters.cpp b/Engine/Plugins/Enterprise/GLTFExporter/MigratedFromGit/Source/GLTFExporter/Private/Converters/GLTFNodeConverters.cpp
--- a/Engine/Plugins/Enterprise/GLTFExporter/MigratedFromGit/Source/GLTFExporter/Private/Converters/GLTFNodeConverters.cpp
+++ b/Engine/Plugins/Enterprise/GLTFExporter/MigratedFromGit/Source/GLTFExporter/Private/Converters/GLTFNodeConverters.cpp
@@ -42,7 +42,7 @@ FGLTFJsonNodeIndex FGLTFComponentConverter::Convert(const USceneComponent* Scene
 
 	// TODO: if node is root, then add it to the scene here to avoid any possible orphan nodes. this change requires node converter to support cyclic calls.
 
-	const USceneComponent* ParentComponent = SceneComponent->GetAttachParent();
+	const USceneComponent* ParentComponent = !bIsRootNode ? SceneComponent->GetAttachParent() : nullptr;
 	const FGLTFJsonNodeIndex ParentNodeIndex = Builder.GetOrAddNode(ParentComponent);
 
 	if (ParentComponent != nullptr && !SceneComponent->IsUsingAbsoluteScale())
