[GLTFExporter] chore: remove TODO for orphan nodes since they should now be handled correctly

commit 1787f17831b7f673cb2d932faf4c39c454153ca4
Date: Thu, 1 Oct 2020 11:18:28 +0200

    chore: remove TODO for orphan nodes since they should now be handled correctly
    (cherry picked from commit 39948b883359d2e9c79cad462d9400e07202eaa8​)

#skipci

diff --git a/Engine/Plugins/Enterprise/GLTFExporter/MigratedFromGit/Source/GLTFExporter/Private/Converters/GLTFNodeConverters.cpp b/Engine/Plugins/Enterprise/GLTFExporter/MigratedFromGit/Source/GLTFExporter/Private/Converters/GLTFNodeConverters.cpp
--- a/Engine/Plugins/Enterprise/GLTFExporter/MigratedFromGit/Source/GLTFExporter/Private/Converters/GLTFNodeConverters.cpp
+++ b/Engine/Plugins/Enterprise/GLTFExporter/MigratedFromGit/Source/GLTFExporter/Private/Converters/GLTFNodeConverters.cpp
@@ -40,8 +40,6 @@ FGLTFJsonNodeIndex FGLTFComponentConverter::Convert(const USceneComponent* Scene
 	const bool bIsRootComponent = Owner->GetRootComponent() == SceneComponent;
 	const bool bIsRootNode = bIsRootComponent && FGLTFActorUtility::IsRootActor(Owner, Builder.bSelectedActorsOnly);
 
-	// TODO: if node is root, then add it to the scene here to avoid any possible orphan nodes. this change requires node converter to support cyclic calls.
-
 	const USceneComponent* ParentComponent = !bIsRootNode ? SceneComponent->GetAttachParent() : nullptr;
 	const FGLTFJsonNodeIndex ParentNodeIndex = Builder.GetOrAddNode(ParentComponent);
 
