[GLTFExporter] chore: remove TODO for functionality that is already implemented

commit 900b87e1676cc08901827bd9fb9ec189e49c6baf
Date: Thu, 1 Oct 2020 11:18:59 +0200

    chore: remove TODO for functionality that is already implemented
    (cherry picked from commit 3678750d9e90cb26968dc19e0eb705a0f136c2a1​)

#skipci

diff --git a/Engine/Plugins/Enterprise/GLTFExporter/MigratedFromGit/Source/GLTFExporter/Private/Converters/GLTFVarationConverters.cpp b/Engine/Plugins/Enterprise/GLTFExporter/MigratedFromGit/Source/GLTFExporter/Private/Converters/GLTFVarationConverters.cpp
--- a/Engine/Plugins/Enterprise/GLTFExporter/MigratedFromGit/Source/GLTFExporter/Private/Converters/GLTFVarationConverters.cpp
+++ b/Engine/Plugins/Enterprise/GLTFExporter/MigratedFromGit/Source/GLTFExporter/Private/Converters/GLTFVarationConverters.cpp
@@ -205,7 +205,7 @@ bool FGLTFVariationConverter::TryParseVisibilityPropertyValue(FGLTFJsonVariant&
 		return false;
 	}
 
-	const FGLTFJsonNodeIndex NodeIndex = Builder.GetOrAddNode(Target); // TODO: what if the target is not selected for export?
+	const FGLTFJsonNodeIndex NodeIndex = Builder.GetOrAddNode(Target);
 	const FGLTFJsonNodeIndex ComponentNodeIndex = Builder.GetComponentNodeIndex(NodeIndex);
 
 	FGLTFJsonVariantNodeProperties& NodeProperties = OutVariant.Nodes.FindOrAdd(ComponentNodeIndex);
