[GLTFExporter] fix: update IsRootActor to support preservation of hierarchy when the entire chain is selected

commit 26319c32e863ac406e078b184ad1bebea4c4d113
Date: Thu, 1 Oct 2020 16:37:02 +0200

    fix: update IsRootActor to support preservation of hierarchy when the entire chain is selected
    (cherry picked from commit 53e597473181a6a5f9d157519b869d6372abedc9​)

#skipci

diff --git a/Engine/Plugins/Enterprise/GLTFExporter/MigratedFromGit/Source/GLTFExporter/Private/Converters/GLTFActorUtility.cpp b/Engine/Plugins/Enterprise/GLTFExporter/MigratedFromGit/Source/GLTFExporter/Private/Converters/GLTFActorUtility.cpp
--- a/Engine/Plugins/Enterprise/GLTFExporter/MigratedFromGit/Source/GLTFExporter/Private/Converters/GLTFActorUtility.cpp
+++ b/Engine/Plugins/Enterprise/GLTFExporter/MigratedFromGit/Source/GLTFExporter/Private/Converters/GLTFActorUtility.cpp
@@ -5,7 +5,8 @@
 
 bool FGLTFActorUtility::IsRootActor(const AActor* Actor, bool bSelectedOnly)
 {
-	return bSelectedOnly ? Actor->IsSelected() : Actor->GetAttachParentActor() == nullptr;
+	const AActor* ParentActor = Actor->GetAttachParentActor();
+	return bSelectedOnly ? Actor->IsSelected() && (ParentActor == nullptr || !ParentActor->IsSelected()) : ParentActor == nullptr;
 }
 
 UBlueprint* FGLTFActorUtility::GetBlueprintFromActor(const AActor* Actor)
