[GLTFExporter] feat: add support for actors/components attached to sockets/bones

commit 650efcf137b7fbf6a97845f5a5d11e1c42add0c4
Date: Thu, 1 Oct 2020 20:42:17 +0200

    feat: add support for actors/components attached to sockets/bones
    (cherry picked from commit d8e372d155831f5ec0ae4bb71389b0ef8c55b820​)
    # Conflicts from cherry-pick:
    #	Plugins/GLTFExporter/Source/GLTFExporter/Private/Converters/GLTFNodeConverters.cpp

#skipci

diff --git a/Engine/Plugins/Enterprise/GLTFExporter/MigratedFromGit/Source/GLTFExporter/Private/Builders/GLTFConvertBuilder.cpp b/Engine/Plugins/Enterprise/GLTFExporter/MigratedFromGit/Source/GLTFExporter/Private/Builders/GLTFConvertBuilder.cpp
--- a/Engine/Plugins/Enterprise/GLTFExporter/MigratedFromGit/Source/GLTFExporter/Private/Builders/GLTFConvertBuilder.cpp
+++ b/Engine/Plugins/Enterprise/GLTFExporter/MigratedFromGit/Source/GLTFExporter/Private/Builders/GLTFConvertBuilder.cpp
@@ -268,7 +268,7 @@ FGLTFJsonAnimationIndex FGLTFConvertBuilder::GetOrAddAnimation(FGLTFJsonNodeInde
 
 FGLTFJsonNodeIndex FGLTFConvertBuilder::GetOrAddNode(const USceneComponent* SceneComponent, FName SocketName)
 {
-	if (SceneComponent == nullptr || SocketName == NAME_None)
+	if (SceneComponent == nullptr)
 	{
 		return FGLTFJsonNodeIndex(INDEX_NONE);
 	}
diff --git a/Engine/Plugins/Enterprise/GLTFExporter/MigratedFromGit/Source/GLTFExporter/Private/Converters/GLTFNodeConverters.cpp b/Engine/Plugins/Enterprise/GLTFExporter/MigratedFromGit/Source/GLTFExporter/Private/Converters/GLTFNodeConverters.cpp
--- a/Engine/Plugins/Enterprise/GLTFExporter/MigratedFromGit/Source/GLTFExporter/Private/Converters/GLTFNodeConverters.cpp
+++ b/Engine/Plugins/Enterprise/GLTFExporter/MigratedFromGit/Source/GLTFExporter/Private/Converters/GLTFNodeConverters.cpp
@@ -11,16 +11,24 @@ FGLTFJsonNodeIndex FGLTFComponentSocketConverter::Convert(const USceneComponent*
 {
 	const FGLTFJsonNodeIndex NodeIndex = Builder.GetOrAddNode(SceneComponent);
 
-	if (const USkinnedMeshComponent* SkinnedMeshComponent = Cast<USkinnedMeshComponent>(SceneComponent))
+	if (SocketName != NAME_None)
 	{
-		// TODO: add support for SocketOverrideLookup?
-		const USkeletalMesh* SkeletalMesh = SkinnedMeshComponent->SkeletalMesh;
-		return Builder.GetOrAddNode(NodeIndex, SkeletalMesh, SocketName);
-	}
+		if (const USkinnedMeshComponent* SkinnedMeshComponent = Cast<USkinnedMeshComponent>(SceneComponent))
+		{
+			// TODO: add support for SocketOverrideLookup?
+			const USkeletalMesh* SkeletalMesh = SkinnedMeshComponent->SkeletalMesh;
+			return Builder.GetOrAddNode(NodeIndex, SkeletalMesh, SocketName);
+		}
 
-	// TODO: add support for non-skeletalmesh sockets
+		// TODO: add support for non-skeletalmesh sockets
 
-	return FGLTFJsonNodeIndex(INDEX_NONE);
+		Builder.AddWarningMessage(
+			FString::Printf(TEXT("Can't export socket %s (in component %s) because it doesn't belong to a skeletal mesh"),
+			*SocketName.ToString(),
+			*SceneComponent->GetName()));
+	}
+
+	return NodeIndex;
 }
 
 FGLTFJsonNodeIndex FGLTFComponentConverter::Convert(const USceneComponent* SceneComponent)
@@ -41,7 +49,8 @@ FGLTFJsonNodeIndex FGLTFComponentConverter::Convert(const USceneComponent* Scene
 	const bool bIsRootNode = bIsRootComponent && FGLTFActorUtility::IsRootActor(Owner, Builder.bSelectedActorsOnly);
 
 	const USceneComponent* ParentComponent = !bIsRootNode ? SceneComponent->GetAttachParent() : nullptr;
-	const FGLTFJsonNodeIndex ParentNodeIndex = Builder.GetOrAddNode(ParentComponent);
+	const FName SocketName = SceneComponent->GetAttachSocketName();
+	const FGLTFJsonNodeIndex ParentNodeIndex = Builder.GetOrAddNode(ParentComponent, SocketName);
 
 	if (ParentComponent != nullptr && !SceneComponent->IsUsingAbsoluteScale())
 	{
@@ -49,7 +58,7 @@ FGLTFJsonNodeIndex FGLTFComponentConverter::Convert(const USceneComponent* Scene
 		if (ParentScale.X != ParentScale.Y || ParentScale.Y != ParentScale.Z)
 		{
 			Builder.AddWarningMessage(
-			FString::Printf(TEXT("Non-uniform parent scale (%s) for component %s (in actor %s) may be represented differently in glTF"),
+				FString::Printf(TEXT("Non-uniform parent scale (%s) for component %s (in actor %s) may be represented differently in glTF"),
 				*ParentScale.ToString(),
 				*SceneComponent->GetName(),
 				*Owner->GetName()));
@@ -57,7 +66,8 @@ FGLTFJsonNodeIndex FGLTFComponentConverter::Convert(const USceneComponent* Scene
 	}
 
 	const FTransform Transform = SceneComponent->GetComponentTransform();
-	const FTransform RelativeTransform = bIsRootNode ? Transform : Transform.GetRelativeTransform(ParentComponent->GetComponentTransform());
+	const FTransform ParentTransform = ParentComponent != nullptr ? ParentComponent->GetSocketTransform(SocketName) : FTransform::Identity;
+	const FTransform RelativeTransform = bIsRootNode ? Transform : Transform.GetRelativeTransform(ParentTransform);
 
 	const FGLTFJsonNodeIndex NodeIndex = Builder.AddChildNode(ParentNodeIndex);
 	FGLTFJsonNode& Node = Builder.GetNode(NodeIndex);
@@ -185,7 +195,7 @@ FGLTFJsonNodeIndex FGLTFSkeletalSocketConverter::Convert(FGLTFJsonNodeIndex Root
 		FGLTFJsonNode Node;
 		Node.Name = SocketName.ToString();
 
-		// TODO: add support for non-uniform scaling
+		// TODO: add warning check for non-uniform scaling
 		const FTransform Transform = Socket->GetSocketLocalTransform();
 		Node.Translation = FGLTFConverterUtility::ConvertPosition(Transform.GetTranslation(), Builder.ExportOptions->ExportScale);
 		Node.Rotation = FGLTFConverterUtility::ConvertRotation(Transform.GetRotation());
@@ -225,7 +235,7 @@ FGLTFJsonNodeIndex FGLTFSkeletalBoneConverter::Convert(FGLTFJsonNodeIndex RootNo
 	const TArray<FTransform>& BonePoses = SkeletalMesh->RefSkeleton.GetRefBonePose();
 	if (BonePoses.IsValidIndex(BoneIndex))
 	{
-		// TODO: add support for non-uniform scaling
+		// TODO: add warning check for non-uniform scaling
 		const FTransform& BonePose = BonePoses[BoneIndex];
 		Node.Translation = FGLTFConverterUtility::ConvertPosition(BonePose.GetTranslation(), Builder.ExportOptions->ExportScale);
 		Node.Rotation = FGLTFConverterUtility::ConvertRotation(BonePose.GetRotation());
