[GLTFExporter] Revert "fix: update IsRootActor to support preservation of hierarchy when the entire chain is selected"

commit 6835def771ec70a933d3f9783157110fe8ab6747
Date: Fri, 2 Oct 2020 08:32:38 +0200

    Revert "fix: update IsRootActor to support preservation of hierarchy when the entire chain is selected"
    This reverts commit 53e597473181a6a5f9d157519b869d6372abedc9.
    (cherry picked from commit 8b4d369928b8e530f3fde191d74924a8930256e7​)

#skipci

diff --git a/Engine/Plugins/Enterprise/GLTFExporter/MigratedFromGit/Source/GLTFExporter/Private/Converters/GLTFActorUtility.cpp b/Engine/Plugins/Enterprise/GLTFExporter/MigratedFromGit/Source/GLTFExporter/Private/Converters/GLTFActorUtility.cpp
--- a/Engine/Plugins/Enterprise/GLTFExporter/MigratedFromGit/Source/GLTFExporter/Private/Converters/GLTFActorUtility.cpp
+++ b/Engine/Plugins/Enterprise/GLTFExporter/MigratedFromGit/Source/GLTFExporter/Private/Converters/GLTFActorUtility.cpp
@@ -5,8 +5,7 @@
 
 bool FGLTFActorUtility::IsRootActor(const AActor* Actor, bool bSelectedOnly)
 {
-	const AActor* ParentActor = Actor->GetAttachParentActor();
-	return bSelectedOnly ? Actor->IsSelected() && (ParentActor == nullptr || !ParentActor->IsSelected()) : ParentActor == nullptr;
+	return bSelectedOnly ? Actor->IsSelected() : Actor->GetAttachParentActor() == nullptr;
 }
 
 UBlueprint* FGLTFActorUtility::GetBlueprintFromActor(const AActor* Actor)
