[GLTFExporter] Revert "feat: modify IsRootActor to handle bSelectedOnly correctly"

commit 947936c7b0ba6b7056d0f3e245b6a45d8821de01
Date: Fri, 2 Oct 2020 08:32:39 +0200

    Revert "feat: modify IsRootActor to handle bSelectedOnly correctly"
    This reverts commit 02579e050a870d59baf54d7a8c38b3864936b2c1.
    (cherry picked from commit 103ed40927b4db14b532012a951a0c9ddc756073​)

#skipci

diff --git a/Engine/Plugins/Enterprise/GLTFExporter/MigratedFromGit/Source/GLTFExporter/Private/Converters/GLTFActorUtility.cpp b/Engine/Plugins/Enterprise/GLTFExporter/MigratedFromGit/Source/GLTFExporter/Private/Converters/GLTFActorUtility.cpp
--- a/Engine/Plugins/Enterprise/GLTFExporter/MigratedFromGit/Source/GLTFExporter/Private/Converters/GLTFActorUtility.cpp
+++ b/Engine/Plugins/Enterprise/GLTFExporter/MigratedFromGit/Source/GLTFExporter/Private/Converters/GLTFActorUtility.cpp
@@ -5,7 +5,8 @@
 
 bool FGLTFActorUtility::IsRootActor(const AActor* Actor, bool bSelectedOnly)
 {
-	return bSelectedOnly ? Actor->IsSelected() : Actor->GetAttachParentActor() == nullptr;
+	const AActor* ParentActor = Actor->GetAttachParentActor();
+	return ParentActor == nullptr || (bSelectedOnly && !ParentActor->IsSelected());
 }
 
 UBlueprint* FGLTFActorUtility::GetBlueprintFromActor(const AActor* Actor)
