[GLTFExporter] refactor: add TODO in scene converter for future size optimizations

commit 1602d2907d12083b582ff36489e204c85695ad36
Date: Fri, 2 Oct 2020 14:06:09 +0200

    refactor: add TODO in scene converter for future size optimizations
    (cherry picked from commit 7bb0ddf1207183c3cd55ff624a90fe8cba26f92e​)

#skipci

diff --git a/Engine/Plugins/Enterprise/GLTFExporter/MigratedFromGit/Source/GLTFExporter/Private/Converters/GLTFSceneConverters.cpp b/Engine/Plugins/Enterprise/GLTFExporter/MigratedFromGit/Source/GLTFExporter/Private/Converters/GLTFSceneConverters.cpp
--- a/Engine/Plugins/Enterprise/GLTFExporter/MigratedFromGit/Source/GLTFExporter/Private/Converters/GLTFSceneConverters.cpp
+++ b/Engine/Plugins/Enterprise/GLTFExporter/MigratedFromGit/Source/GLTFExporter/Private/Converters/GLTFSceneConverters.cpp
@@ -38,6 +38,7 @@ FGLTFJsonSceneIndex FGLTFSceneConverter::Convert(const ULevel* Level)
 		const FGLTFJsonNodeIndex NodeIndex = Builder.GetOrAddNode(Actor);
 		if (NodeIndex != INDEX_NONE && FGLTFActorUtility::IsRootActor(Actor, Builder.bSelectedActorsOnly))
 		{
+			// TODO: to avoid having to add irrelevant actors/components let GLTFComponentConverter decide and add root nodes to scene. This change may require node converters to support cyclic calls.
 			Scene.Nodes.Add(NodeIndex);
 		}
 	}
