[GLTFExporter] feat(interaction-hotspot): add and initialize properties for default images

commit 56633a358afc9fbbd556fadfa106ef8513ad56a1
Date: Fri, 2 Oct 2020 12:50:02 +0200

    feat(interaction-hotspot): add and initialize properties for default images
    (cherry picked from commit e374a41bc40e846f277e76698e6c6bdfe6f149e9​)

#skipci

diff --git a/Engine/Plugins/Enterprise/GLTFExporter/MigratedFromGit/Source/GLTFExporterRuntime/Private/Actors/GLTFInteractionHotspotActor.cpp b/Engine/Plugins/Enterprise/GLTFExporter/MigratedFromGit/Source/GLTFExporterRuntime/Private/Actors/GLTFInteractionHotspotActor.cpp
--- a/Engine/Plugins/Enterprise/GLTFExporter/MigratedFromGit/Source/GLTFExporterRuntime/Private/Actors/GLTFInteractionHotspotActor.cpp
+++ b/Engine/Plugins/Enterprise/GLTFExporter/MigratedFromGit/Source/GLTFExporterRuntime/Private/Actors/GLTFInteractionHotspotActor.cpp
@@ -33,8 +33,13 @@ AGLTFInteractionHotspotActor::AGLTFInteractionHotspotActor(const FObjectInitiali
 	HoveredImage(nullptr),
 	ToggledImage(nullptr),
 	ToggledHoveredImage(nullptr),
+	BillboardComponent(nullptr),
 	SphereComponent(nullptr),
 	DefaultMaterial(nullptr),
+	DefaultImage(nullptr),
+	DefaultHoveredImage(nullptr),
+	DefaultToggledImage(nullptr),
+	DefaultToggledHoveredImage(nullptr),
 	ActiveImage(nullptr),
 	ActiveImageSize(0.0f, 0.0f),
 	bToggled(bToggled),
@@ -71,13 +76,33 @@ AGLTFInteractionHotspotActor::AGLTFInteractionHotspotActor(const FObjectInitiali
 	SphereComponent->OnEndCursorOver.AddDynamic(this, &AGLTFInteractionHotspotActor::EndCursorOver);
 	SphereComponent->OnClicked.AddDynamic(this, &AGLTFInteractionHotspotActor::Clicked);
 
-	const ConstructorHelpers::FObjectFinder<UMaterial> MaterialRef(TEXT("/GLTFExporter/Materials/Hotspot"));
-	DefaultMaterial = MaterialRef.Object;
-
-#if WITH_EDITORONLY_DATA
-	const ConstructorHelpers::FObjectFinder<UMaterial> IconMaterialRef(TEXT("/GLTFExporter/Materials/HotspotIcon"));
-	DefaultIconMaterial = IconMaterialRef.Object;
-#endif // WITH_EDITORONLY_DATA
+	struct FConstructorStatics
+	{
+		ConstructorHelpers::FObjectFinder<UMaterial> DefaultMaterial;
+		ConstructorHelpers::FObjectFinder<UMaterial> DefaultIconMaterial;
+		ConstructorHelpers::FObjectFinder<UTexture2D> DefaultImage;
+		ConstructorHelpers::FObjectFinder<UTexture2D> DefaultHoveredImage;
+		ConstructorHelpers::FObjectFinder<UTexture2D> DefaultToggledImage;
+		ConstructorHelpers::FObjectFinder<UTexture2D> DefaultToggledHoveredImage;
+
+		FConstructorStatics()
+			: DefaultMaterial(TEXT("/GLTFExporter/Materials/Hotspot"))
+			, DefaultIconMaterial(TEXT("/GLTFExporter/Materials/HotspotIcon"))
+			, DefaultImage(TEXT("/GLTFExporter/Textures/Hotspots/Image"))
+			, DefaultHoveredImage(TEXT("/GLTFExporter/Textures/Hotspots/HoveredImage"))
+			, DefaultToggledImage(TEXT("/GLTFExporter/Textures/Hotspots/ToggledImage"))
+			, DefaultToggledHoveredImage(TEXT("/GLTFExporter/Textures/Hotspots/ToggledHoveredImage"))
+		{
+		}
+	};
+	static FConstructorStatics ConstructorStatics;
+
+	DefaultMaterial = ConstructorStatics.DefaultMaterial.Object;
+	DefaultIconMaterial = ConstructorStatics.DefaultIconMaterial.Object;
+	DefaultImage = ConstructorStatics.DefaultImage.Object;
+	DefaultHoveredImage = ConstructorStatics.DefaultHoveredImage.Object;
+	DefaultToggledImage = ConstructorStatics.DefaultToggledImage.Object;
+	DefaultToggledHoveredImage = ConstructorStatics.DefaultToggledHoveredImage.Object;
 
 	PrimaryActorTick.bCanEverTick = true;
 	PrimaryActorTick.bStartWithTickEnabled = true;
diff --git a/Engine/Plugins/Enterprise/GLTFExporter/MigratedFromGit/Source/GLTFExporterRuntime/Public/Actors/GLTFInteractionHotspotActor.h b/Engine/Plugins/Enterprise/GLTFExporter/MigratedFromGit/Source/GLTFExporterRuntime/Public/Actors/GLTFInteractionHotspotActor.h
--- a/Engine/Plugins/Enterprise/GLTFExporter/MigratedFromGit/Source/GLTFExporterRuntime/Public/Actors/GLTFInteractionHotspotActor.h
+++ b/Engine/Plugins/Enterprise/GLTFExporter/MigratedFromGit/Source/GLTFExporterRuntime/Public/Actors/GLTFInteractionHotspotActor.h
@@ -86,10 +86,20 @@ private:
 	UPROPERTY(transient, duplicatetransient)
 	UMaterialInterface* DefaultMaterial;
 
-#if WITH_EDITORONLY_DATA
+	UPROPERTY(transient, duplicatetransient)
+	UTexture2D* DefaultImage;
+
+	UPROPERTY(transient, duplicatetransient)
+	UTexture2D* DefaultHoveredImage;
+
+	UPROPERTY(transient, duplicatetransient)
+	UTexture2D* DefaultToggledImage;
+
+	UPROPERTY(transient, duplicatetransient)
+	UTexture2D* DefaultToggledHoveredImage;
+
 	UPROPERTY(transient, duplicatetransient)
 	UMaterialInterface* DefaultIconMaterial;
-#endif // WITH_EDITORONLY_DATA
 
 	UPROPERTY(transient, duplicatetransient)
 	UTexture* ActiveImage;
