[GLTFExporter] fix(orbit-camera): add missing include

commit ddc9e7716973a86d715cb40473facecef50993d6
Date: Fri, 2 Oct 2020 11:02:52 +0200

    fix(orbit-camera): add missing include
    (cherry picked from commit 561955883f5f7ffb3fffecde352ae1fcbb49a7ad​)

#skipci

diff --git a/Engine/Plugins/Enterprise/GLTFExporter/MigratedFromGit/Source/GLTFExporterRuntime/Private/Actors/GLTFOrbitCameraActor.cpp b/Engine/Plugins/Enterprise/GLTFExporter/MigratedFromGit/Source/GLTFExporterRuntime/Private/Actors/GLTFOrbitCameraActor.cpp
--- a/Engine/Plugins/Enterprise/GLTFExporter/MigratedFromGit/Source/GLTFExporterRuntime/Private/Actors/GLTFOrbitCameraActor.cpp
+++ b/Engine/Plugins/Enterprise/GLTFExporter/MigratedFromGit/Source/GLTFExporterRuntime/Private/Actors/GLTFOrbitCameraActor.cpp
@@ -1,6 +1,7 @@
 // Copyright Epic Games, Inc. All Rights Reserved.
 
 #include "Actors/GLTFOrbitCameraActor.h"
+#include "Components/InputComponent.h"
 
 namespace
 {
