[GLTFExporter] feat(interaction-hotspot): modify logic for selecting the active image to display

commit 76d3f7de003ea0b53ec966f0adcb4c9b364ab11f
Date: Fri, 2 Oct 2020 11:12:23 +0200

    feat(interaction-hotspot): modify logic for selecting the active image to display
    (cherry picked from commit f8f4857c45d9cb22493697fdaaa7c29dd4822a8d​)

#skipci

diff --git a/Engine/Plugins/Enterprise/GLTFExporter/MigratedFromGit/Source/GLTFExporterRuntime/Private/Actors/GLTFInteractionHotspotActor.cpp b/Engine/Plugins/Enterprise/GLTFExporter/MigratedFromGit/Source/GLTFExporterRuntime/Private/Actors/GLTFInteractionHotspotActor.cpp
--- a/Engine/Plugins/Enterprise/GLTFExporter/MigratedFromGit/Source/GLTFExporterRuntime/Private/Actors/GLTFInteractionHotspotActor.cpp
+++ b/Engine/Plugins/Enterprise/GLTFExporter/MigratedFromGit/Source/GLTFExporterRuntime/Private/Actors/GLTFInteractionHotspotActor.cpp
@@ -122,7 +122,7 @@ void AGLTFInteractionHotspotActor::PostEditChangeProperty(FPropertyChangedEvent&
 
 		if (PropertyName == TEXT("Image"))
 		{
-			SetActiveImage(Image);
+			UpdateActiveImageFromState(EGLTFHotspotState::Default);
 		}
 		else if (PropertyName == TEXT("SkeletalMeshActor"))
 		{
@@ -147,7 +147,7 @@ void AGLTFInteractionHotspotActor::PostRegisterAllComponents()
 	Super::PostRegisterAllComponents();
 
 	SetupSpriteElement();
-	SetActiveImage(Image);
+	UpdateActiveImageFromState(EGLTFHotspotState::Default);
 }
 
 void AGLTFInteractionHotspotActor::Tick(float DeltaTime)
@@ -232,13 +232,13 @@ void AGLTFInteractionHotspotActor::BeginCursorOver(UPrimitiveComponent* TouchedC
 {
 	if (bIsInteractable)
 	{
-		SetActiveImage(CalculateActiveImage(true));
+		UpdateActiveImageFromState(bToggled ? EGLTFHotspotState::ToggledHovered :  EGLTFHotspotState::Hovered);
 	}
 }
 
 void AGLTFInteractionHotspotActor::EndCursorOver(UPrimitiveComponent* TouchedComponent)
 {
-	SetActiveImage(CalculateActiveImage(false));
+	UpdateActiveImageFromState(bToggled ? EGLTFHotspotState::Toggled :  EGLTFHotspotState::Default);
 }
 
 void AGLTFInteractionHotspotActor::Clicked(UPrimitiveComponent* TouchedComponent, FKey ButtonPressed)
@@ -273,11 +273,13 @@ void AGLTFInteractionHotspotActor::Clicked(UPrimitiveComponent* TouchedComponent
 
 	bToggled = !bToggled;
 
-	SetActiveImage(CalculateActiveImage(true));
+	UpdateActiveImageFromState(bToggled ? EGLTFHotspotState::ToggledHovered :  EGLTFHotspotState::Hovered);
 }
 
-void AGLTFInteractionHotspotActor::SetActiveImage(UTexture2D* NewImage)
+void AGLTFInteractionHotspotActor::UpdateActiveImageFromState(EGLTFHotspotState State)
 {
+	UTexture2D* NewImage = const_cast<UTexture2D*>(GetImageForState(State));
+
 	UTexture* DefaultTexture;
 	GetSpriteMaterial()->GetTextureParameterDefaultValue(NAME_SpriteParameter, DefaultTexture);
 
@@ -291,43 +293,6 @@ void AGLTFInteractionHotspotActor::SetActiveImage(UTexture2D* NewImage)
 	UpdateSpriteSize();
 }
 
-UTexture2D* AGLTFInteractionHotspotActor::CalculateActiveImage(bool bCursorOver) const
-{
-	// A list of candidates ordered by descending priority
-	TArray<UTexture2D*> ImageCandidates;
-
-	if (bToggled)
-	{
-		if (bCursorOver)
-		{
-			ImageCandidates.Add(ToggledHoveredImage);
-		}
-
-		ImageCandidates.Add(ToggledImage);
-	}
-	else
-	{
-		if (bCursorOver)
-		{
-			ImageCandidates.Add(HoveredImage);
-		}
-
-		ImageCandidates.Add(Image);
-	}
-
-	ImageCandidates.Add(Image);
-
-	for (UTexture2D* ImageCandidate : ImageCandidates)
-	{
-		if (ImageCandidate != nullptr)
-		{
-			return ImageCandidate;
-		}
-	}
-
-	return nullptr;
-}
-
 void AGLTFInteractionHotspotActor::SetupSpriteElement() const
 {
 	UMaterialInstanceDynamic* MaterialInstance;
@@ -436,3 +401,46 @@ void AGLTFInteractionHotspotActor::ViewportResized(FViewport*, uint32)
 {
 	UpdateSpriteSize();
 }
+
+const UTexture2D* AGLTFInteractionHotspotActor::GetImageForState(EGLTFHotspotState State) const
+{
+	const UTexture2D* CurrentImage = DefaultImage;
+	const UTexture2D* CurrentHoveredImage = DefaultHoveredImage;
+	const UTexture2D* CurrentToggledImage = DefaultToggledImage;
+	const UTexture2D* CurrentToggledHoveredImage = DefaultToggledHoveredImage;
+
+	if (Image != nullptr)
+	{
+		CurrentImage = Image;
+		CurrentHoveredImage = HoveredImage != nullptr ? HoveredImage : CurrentImage;
+		CurrentToggledImage = ToggledImage != nullptr ? ToggledImage : CurrentImage;
+
+		if (ToggledHoveredImage != nullptr)
+		{
+			CurrentToggledHoveredImage = ToggledHoveredImage;
+		}
+		else if (ToggledImage != nullptr)
+		{
+			CurrentToggledHoveredImage = ToggledImage;
+		}
+		else if (HoveredImage != nullptr)
+		{
+			CurrentToggledHoveredImage = HoveredImage;
+		}
+		else
+		{
+			CurrentToggledHoveredImage = Image;
+		}
+	}
+
+	switch (State)
+	{
+		case EGLTFHotspotState::Default:        return CurrentImage;
+		case EGLTFHotspotState::Hovered:        return CurrentHoveredImage;
+		case EGLTFHotspotState::Toggled:        return CurrentToggledImage;
+		case EGLTFHotspotState::ToggledHovered: return CurrentToggledHoveredImage;
+		default:                                checkNoEntry();
+	}
+
+	return nullptr;
+}
diff --git a/Engine/Plugins/Enterprise/GLTFExporter/MigratedFromGit/Source/GLTFExporterRuntime/Public/Actors/GLTFInteractionHotspotActor.h b/Engine/Plugins/Enterprise/GLTFExporter/MigratedFromGit/Source/GLTFExporterRuntime/Public/Actors/GLTFInteractionHotspotActor.h
--- a/Engine/Plugins/Enterprise/GLTFExporter/MigratedFromGit/Source/GLTFExporterRuntime/Public/Actors/GLTFInteractionHotspotActor.h
+++ b/Engine/Plugins/Enterprise/GLTFExporter/MigratedFromGit/Source/GLTFExporterRuntime/Public/Actors/GLTFInteractionHotspotActor.h
@@ -15,6 +15,14 @@ class USphereComponent;
 class UMaterialInstanceDynamic;
 class FViewport;
 
+enum class EGLTFHotspotState : uint8
+{
+	Default,
+	Hovered,
+	Toggled,
+	ToggledHovered
+};
+
 /**
  * Actor wrapper for the GLTF hotspot component. Appears as a billboard and allows playback of skeletal animations when cursor input is enabled.
  */
@@ -45,9 +53,7 @@ private:
 	UFUNCTION()
 	void Clicked(UPrimitiveComponent* TouchedComponent, FKey ButtonPressed);
 
-	void SetActiveImage(UTexture2D* NewImage);
-
-	UTexture2D* CalculateActiveImage(bool bCursorOver) const;
+	void UpdateActiveImageFromState(EGLTFHotspotState State);
 
 public:
 
@@ -102,7 +108,7 @@ private:
 	UMaterialInterface* DefaultIconMaterial;
 
 	UPROPERTY(transient, duplicatetransient)
-	UTexture* ActiveImage;
+	const UTexture* ActiveImage;
 
 	FVector2D ActiveImageSize;
 
@@ -125,4 +131,8 @@ private:
 	FIntPoint GetCurrentViewportSize();
 
 	void ViewportResized(FViewport*, uint32);
+
+public:
+
+	const UTexture2D* GetImageForState(EGLTFHotspotState State) const;
 };
