[GLTFExporter] feat(interaction-hotspot): prevent editing of other images when the default image is not set

commit 4962837e748d2c39d24039922813c5e1b8a82174
Date: Fri, 2 Oct 2020 11:13:01 +0200

    feat(interaction-hotspot): prevent editing of other images when the default image is not set
    (cherry picked from commit 2b272ec7262d373b99d44636785697dcb19e0122​)

#skipci

diff --git a/Engine/Plugins/Enterprise/GLTFExporter/MigratedFromGit/Source/GLTFExporterRuntime/Public/Actors/GLTFInteractionHotspotActor.h b/Engine/Plugins/Enterprise/GLTFExporter/MigratedFromGit/Source/GLTFExporterRuntime/Public/Actors/GLTFInteractionHotspotActor.h
--- a/Engine/Plugins/Enterprise/GLTFExporter/MigratedFromGit/Source/GLTFExporterRuntime/Public/Actors/GLTFInteractionHotspotActor.h
+++ b/Engine/Plugins/Enterprise/GLTFExporter/MigratedFromGit/Source/GLTFExporterRuntime/Public/Actors/GLTFInteractionHotspotActor.h
@@ -70,15 +70,15 @@ public:
 	UTexture2D* Image;
 
 	/** The optional billboard image that will be shown when a cursor enters the hotspot. */
-	UPROPERTY(EditAnywhere, BlueprintReadWrite, Category = "GLTF Interaction Hotspot")
+	UPROPERTY(EditAnywhere, BlueprintReadWrite, Category = "GLTF Interaction Hotspot", meta=(EditCondition="Image != nullptr"))
 	UTexture2D* HoveredImage;
 
 	/** The optional billboard image that will be shown when the hotspot is toggled by a click. */
-	UPROPERTY(EditAnywhere, BlueprintReadWrite, Category = "GLTF Interaction Hotspot")
+	UPROPERTY(EditAnywhere, BlueprintReadWrite, Category = "GLTF Interaction Hotspot", meta=(EditCondition="Image != nullptr"))
 	UTexture2D* ToggledImage;
 
 	/** The optional billboard image that will be shown when the hotspot is toggled by a click and a cursor enters it. */
-	UPROPERTY(EditAnywhere, BlueprintReadWrite, Category = "GLTF Interaction Hotspot")
+	UPROPERTY(EditAnywhere, BlueprintReadWrite, Category = "GLTF Interaction Hotspot", meta=(EditCondition="Image != nullptr"))
 	UTexture2D* ToggledHoveredImage;
 
 private:
