[GLTFExporter] feat(interaction-hotspot): calculate correct image to export for each state

commit dd86f6996c2e692bab9d5ab309f1a96b8dc7e9c1
Date: Fri, 2 Oct 2020 11:22:30 +0200

    feat(interaction-hotspot): calculate correct image to export for each state
    (cherry picked from commit 87ea84087163d49d3686ca4e3fa01679dada8e64​)

#skipci

diff --git a/Engine/Plugins/Enterprise/GLTFExporter/MigratedFromGit/Source/GLTFExporter/Private/Converters/GLTFHotspotConverters.cpp b/Engine/Plugins/Enterprise/GLTFExporter/MigratedFromGit/Source/GLTFExporter/Private/Converters/GLTFHotspotConverters.cpp
--- a/Engine/Plugins/Enterprise/GLTFExporter/MigratedFromGit/Source/GLTFExporter/Private/Converters/GLTFHotspotConverters.cpp
+++ b/Engine/Plugins/Enterprise/GLTFExporter/MigratedFromGit/Source/GLTFExporter/Private/Converters/GLTFHotspotConverters.cpp
@@ -56,10 +56,10 @@ FGLTFJsonHotspotIndex FGLTFHotspotConverter::Convert(const AGLTFInteractionHotsp
 		}
 	}
 
-	JsonHotspot.Image = Builder.GetOrAddTexture(HotspotActor->Image);
-	JsonHotspot.HoveredImage = Builder.GetOrAddTexture(HotspotActor->HoveredImage);
-	JsonHotspot.ToggledImage = Builder.GetOrAddTexture(HotspotActor->ToggledImage);
-	JsonHotspot.ToggledHoveredImage = Builder.GetOrAddTexture(HotspotActor->ToggledHoveredImage);
+	JsonHotspot.Image = Builder.GetOrAddTexture(HotspotActor->GetImageForState(EGLTFHotspotState::Default));
+	JsonHotspot.HoveredImage = Builder.GetOrAddTexture(HotspotActor->GetImageForState(EGLTFHotspotState::Hovered));
+	JsonHotspot.ToggledImage = Builder.GetOrAddTexture(HotspotActor->GetImageForState(EGLTFHotspotState::Toggled));
+	JsonHotspot.ToggledHoveredImage = Builder.GetOrAddTexture(HotspotActor->GetImageForState(EGLTFHotspotState::ToggledHovered));
 
 	return Builder.AddHotspot(JsonHotspot);
 }
