[GLTFExporter] feat: combine multiple mesh-sections into one FGTLFMeshSection if they share material-index

commit 3405b9dd248bd7ab2d4b0958ee01868bf38386f0
Date: Thu, 1 Oct 2020 15:31:19 +0200

    feat: combine multiple mesh-sections into one FGTLFMeshSection if they share material-index
    this way we get the same type of mapping between materials and primitives as in glTF
    (cherry picked from commit 3f68b409a422760d0093b997fe5e06849a4ff277​)

#skipci

diff --git a/Engine/Plugins/Enterprise/GLTFExporter/MigratedFromGit/Source/GLTFExporter/Private/Converters/GLTFAccessorConverters.cpp b/Engine/Plugins/Enterprise/GLTFExporter/MigratedFromGit/Source/GLTFExporter/Private/Converters/GLTFAccessorConverters.cpp
--- a/Engine/Plugins/Enterprise/GLTFExporter/MigratedFromGit/Source/GLTFExporter/Private/Converters/GLTFAccessorConverters.cpp
+++ b/Engine/Plugins/Enterprise/GLTFExporter/MigratedFromGit/Source/GLTFExporter/Private/Converters/GLTFAccessorConverters.cpp
@@ -342,10 +342,9 @@ FGLTFJsonAccessorIndex FGLTFBoneIndexBufferConverter::Convert(const FGLTFMeshSec
 	TArray<VertexBoneIndices> BoneIndices;
 	BoneIndices.AddUninitialized(VertexCount);
 
-	const TArray<FBoneIndexType>& BoneMap = MeshSection->BoneMap;
-
 	for (uint32 VertexIndex = 0; VertexIndex < VertexCount; ++VertexIndex)
 	{
+		const TArray<FBoneIndexType>& BoneMap = MeshSection->BoneMaps[MeshSection->BoneMapLookup[VertexIndex]];
 		const uint32 MappedVertexIndex = IndexMap[VertexIndex];
 		VertexBoneIndices& VertexBones = BoneIndices[VertexIndex];
 
diff --git a/Engine/Plugins/Enterprise/GLTFExporter/MigratedFromGit/Source/GLTFExporter/Private/Converters/GLTFMeshConverters.cpp b/Engine/Plugins/Enterprise/GLTFExporter/MigratedFromGit/Source/GLTFExporter/Private/Converters/GLTFMeshConverters.cpp
--- a/Engine/Plugins/Enterprise/GLTFExporter/MigratedFromGit/Source/GLTFExporter/Private/Converters/GLTFMeshConverters.cpp
+++ b/Engine/Plugins/Enterprise/GLTFExporter/MigratedFromGit/Source/GLTFExporter/Private/Converters/GLTFMeshConverters.cpp
@@ -4,7 +4,6 @@
 #include "Converters/GLTFConverterUtility.h"
 #include "Converters/GLTFNameUtility.h"
 #include "Builders/GLTFConvertBuilder.h"
-#include "Rendering/MultiSizeIndexContainer.h"
 #include "Rendering/SkeletalMeshRenderData.h"
 
 void FGLTFStaticMeshConverter::Sanitize(const UStaticMesh*& StaticMesh, int32& LODIndex, const FColorVertexBuffer*& OverrideVertexColors, FGLTFMaterialArray& OverrideMaterials)
@@ -46,20 +45,18 @@ FGLTFJsonMeshIndex FGLTFStaticMeshConverter::Convert(const UStaticMesh* StaticMe
 	FGLTFJsonMesh JsonMesh;
 	JsonMesh.Name = FGLTFNameUtility::GetName(StaticMesh, LODIndex);
 
-	const FRawStaticIndexBuffer* IndexBuffer = &MeshLOD.IndexBuffer;
 	const FPositionVertexBuffer* PositionBuffer = &MeshLOD.VertexBuffers.PositionVertexBuffer;
 	const FStaticMeshVertexBuffer* VertexBuffer = &MeshLOD.VertexBuffers.StaticMeshVertexBuffer;
 	const FColorVertexBuffer* ColorBuffer = OverrideVertexColors != nullptr ? OverrideVertexColors : &MeshLOD.VertexBuffers.ColorVertexBuffer;
 
-	const int32 SectionCount = MeshLOD.Sections.Num();
-	JsonMesh.Primitives.AddDefaulted(SectionCount);
+	const int32 MaterialCount = StaticMesh->StaticMaterials.Num();
+	JsonMesh.Primitives.AddDefaulted(MaterialCount);
 
-	for (int32 SectionIndex = 0; SectionIndex < SectionCount; ++SectionIndex)
+	for (int32 MaterialIndex = 0; MaterialIndex < MaterialCount; ++MaterialIndex)
 	{
-		const FStaticMeshSection* Section = &MeshLOD.Sections[SectionIndex];
-		const FGLTFMeshSection* ConvertedSection = MeshSectionConverter.GetOrAdd(Section, IndexBuffer);
+		const FGLTFMeshSection* ConvertedSection = MeshSectionConverter.GetOrAdd(&MeshLOD, MaterialIndex);
 
-		FGLTFJsonPrimitive& JsonPrimitive = JsonMesh.Primitives[SectionIndex];
+		FGLTFJsonPrimitive& JsonPrimitive = JsonMesh.Primitives[MaterialIndex];
 		JsonPrimitive.Indices = Builder.GetOrAddIndexAccessor(ConvertedSection);
 
 		JsonPrimitive.Attributes.Position = Builder.GetOrAddPositionAccessor(ConvertedSection, PositionBuffer);
@@ -84,7 +81,7 @@ FGLTFJsonMeshIndex FGLTFStaticMeshConverter::Convert(const UStaticMesh* StaticMe
 			JsonPrimitive.Attributes.TexCoords[UVIndex] = Builder.GetOrAddUVAccessor(ConvertedSection, VertexBuffer, UVIndex);
 		}
 
-		const UMaterialInterface* Material = OverrideMaterials.GetOverride(StaticMesh->StaticMaterials, Section->MaterialIndex);
+		const UMaterialInterface* Material = OverrideMaterials.GetOverride(StaticMesh->StaticMaterials, MaterialIndex);
 		if (Material == nullptr)
 		{
 			Material = UMaterial::GetDefaultMaterial(MD_Surface);
@@ -149,21 +146,19 @@ FGLTFJsonMeshIndex FGLTFSkeletalMeshConverter::Convert(const USkeletalMesh* Skel
 	FGLTFJsonMesh JsonMesh;
 	JsonMesh.Name = FGLTFNameUtility::GetName(SkeletalMesh, LODIndex);
 
-	const FRawStaticIndexBuffer16or32Interface* IndexBuffer = MeshLOD.MultiSizeIndexContainer.GetIndexBuffer();
 	const FPositionVertexBuffer* PositionBuffer = &MeshLOD.StaticVertexBuffers.PositionVertexBuffer;
 	const FStaticMeshVertexBuffer* VertexBuffer = &MeshLOD.StaticVertexBuffers.StaticMeshVertexBuffer;
 	const FColorVertexBuffer* ColorBuffer = OverrideVertexColors != nullptr ? OverrideVertexColors : &MeshLOD.StaticVertexBuffers.ColorVertexBuffer;
 	const FSkinWeightVertexBuffer* SkinWeightBuffer = OverrideSkinWeights != nullptr ? OverrideSkinWeights : MeshLOD.GetSkinWeightVertexBuffer();
 
-	const int32 SectionCount = MeshLOD.RenderSections.Num();
-	JsonMesh.Primitives.AddDefaulted(SectionCount);
+	const uint16 MaterialCount = SkeletalMesh->Materials.Num();
+	JsonMesh.Primitives.AddDefaulted(MaterialCount);
 
-	for (int32 SectionIndex = 0; SectionIndex < SectionCount; ++SectionIndex)
+	for (uint16 MaterialIndex = 0; MaterialIndex < MaterialCount; ++MaterialIndex)
 	{
-		const FSkelMeshRenderSection* Section = &MeshLOD.RenderSections[SectionIndex];
-		const FGLTFMeshSection* ConvertedSection = MeshSectionConverter.GetOrAdd(Section, IndexBuffer);
+		const FGLTFMeshSection* ConvertedSection = MeshSectionConverter.GetOrAdd(&MeshLOD, MaterialIndex);
 
-		FGLTFJsonPrimitive& JsonPrimitive = JsonMesh.Primitives[SectionIndex];
+		FGLTFJsonPrimitive& JsonPrimitive = JsonMesh.Primitives[MaterialIndex];
 		JsonPrimitive.Indices = Builder.GetOrAddIndexAccessor(ConvertedSection);
 
 		JsonPrimitive.Attributes.Position = Builder.GetOrAddPositionAccessor(ConvertedSection, PositionBuffer);
@@ -201,7 +196,7 @@ FGLTFJsonMeshIndex FGLTFSkeletalMeshConverter::Convert(const USkeletalMesh* Skel
 			}
 		}
 
-		const UMaterialInterface* Material = OverrideMaterials.GetOverride(SkeletalMesh->Materials, Section->MaterialIndex);
+		const UMaterialInterface* Material = OverrideMaterials.GetOverride(SkeletalMesh->Materials, MaterialIndex);
 		if (Material == nullptr)
 		{
 			Material = UMaterial::GetDefaultMaterial(MD_Surface);
diff --git a/Engine/Plugins/Enterprise/GLTFExporter/MigratedFromGit/Source/GLTFExporter/Private/Converters/GLTFMeshConverters.h b/Engine/Plugins/Enterprise/GLTFExporter/MigratedFromGit/Source/GLTFExporter/Private/Converters/GLTFMeshConverters.h
--- a/Engine/Plugins/Enterprise/GLTFExporter/MigratedFromGit/Source/GLTFExporter/Private/Converters/GLTFMeshConverters.h
+++ b/Engine/Plugins/Enterprise/GLTFExporter/MigratedFromGit/Source/GLTFExporter/Private/Converters/GLTFMeshConverters.h
@@ -15,14 +15,14 @@ class TGLTFMeshConverter : public FGLTFBuilderContext, public TGLTFConverter<FGL
 	using FGLTFBuilderContext::FGLTFBuilderContext;
 };
 
-template <typename MeshSectionType, typename IndexBufferType>
-class TGLTFMeshSectionConverter final : public TGLTFConverter<const FGLTFMeshSection*, const MeshSectionType*, const IndexBufferType*>
+template <typename MeshLODType, typename MaterialIndexType>
+class TGLTFMeshSectionConverter final : public TGLTFConverter<const FGLTFMeshSection*, const MeshLODType*, const MaterialIndexType>
 {
 	TArray<TUniquePtr<FGLTFMeshSection>> Outputs;
 
-	const FGLTFMeshSection* Convert(const MeshSectionType* MeshSection, const IndexBufferType* IndexBuffer)
+	const FGLTFMeshSection* Convert(const MeshLODType* MeshLOD, const MaterialIndexType MaterialIndex)
 	{
-		return Outputs.Add_GetRef(MakeUnique<FGLTFMeshSection>(MeshSection, IndexBuffer)).Get();
+		return Outputs.Add_GetRef(MakeUnique<FGLTFMeshSection>(MeshLOD, MaterialIndex)).Get();
 	}
 };
 
@@ -34,7 +34,7 @@ class FGLTFStaticMeshConverter final : public TGLTFMeshConverter<const UStaticMe
 
 	virtual FGLTFJsonMeshIndex Convert(const UStaticMesh* StaticMesh, int32 LODIndex, const FColorVertexBuffer* OverrideVertexColors, FGLTFMaterialArray OverrideMaterials) override;
 
-	TGLTFMeshSectionConverter<FStaticMeshSection, FRawStaticIndexBuffer> MeshSectionConverter;
+	TGLTFMeshSectionConverter<FStaticMeshLODResources, int32> MeshSectionConverter;
 };
 
 class FGLTFSkeletalMeshConverter final : public TGLTFMeshConverter<const USkeletalMesh*, int32, const FColorVertexBuffer*, const FSkinWeightVertexBuffer*, FGLTFMaterialArray>
@@ -45,5 +45,5 @@ class FGLTFSkeletalMeshConverter final : public TGLTFMeshConverter<const USkelet
 
 	virtual FGLTFJsonMeshIndex Convert(const USkeletalMesh* SkeletalMesh, int32 LODIndex, const FColorVertexBuffer* OverrideVertexColors, const FSkinWeightVertexBuffer* OverrideSkinWeights, FGLTFMaterialArray OverrideMaterials) override;
 
-	TGLTFMeshSectionConverter<FSkelMeshRenderSection, FRawStaticIndexBuffer16or32Interface> MeshSectionConverter;
+	TGLTFMeshSectionConverter<FSkeletalMeshLODRenderData, uint16> MeshSectionConverter;
 };
diff --git a/Engine/Plugins/Enterprise/GLTFExporter/MigratedFromGit/Source/GLTFExporter/Private/Converters/GLTFMeshSection.cpp b/Engine/Plugins/Enterprise/GLTFExporter/MigratedFromGit/Source/GLTFExporter/Private/Converters/GLTFMeshSection.cpp
--- a/Engine/Plugins/Enterprise/GLTFExporter/MigratedFromGit/Source/GLTFExporter/Private/Converters/GLTFMeshSection.cpp
+++ b/Engine/Plugins/Enterprise/GLTFExporter/MigratedFromGit/Source/GLTFExporter/Private/Converters/GLTFMeshSection.cpp
@@ -5,67 +5,105 @@
 #include "Rendering/SkeletalMeshRenderData.h"
 #include "Algo/MaxElement.h"
 
-FGLTFMeshSection::FGLTFMeshSection(const FStaticMeshSection* MeshSection, const FRawStaticIndexBuffer* OldIndexBuffer)
+FGLTFMeshSection::FGLTFMeshSection(const FStaticMeshLODResources* MeshLOD, const int32 MaterialIndex)
 {
-	const uint32 IndexOffset = MeshSection->FirstIndex;
-	const uint32 IndexCount = MeshSection->NumTriangles * 3;
+	uint32 TotalIndexCount = 0;
 
-	IndexMap.Reserve(IndexCount);
-	IndexBuffer.AddUninitialized(IndexCount);
+	TArray<const FStaticMeshSection*> Sections;
+	for (const FStaticMeshSection& Section : MeshLOD->Sections)
+	{
+		if (Section.MaterialIndex == MaterialIndex)
+		{
+			Sections.Add(&Section);
+			TotalIndexCount += Section.NumTriangles * 3;
+		}
+	}
+
+	IndexMap.Reserve(TotalIndexCount);
+	IndexBuffer.AddUninitialized(TotalIndexCount);
+	BoneMapLookup.AddUninitialized(TotalIndexCount);
 
 	TMap<uint32, uint32> IndexLookup;
 
-	for (uint32 Index = 0; Index < IndexCount; Index++)
+	for (const FStaticMeshSection* MeshSection : Sections)
 	{
-		const uint32 OldIndex = OldIndexBuffer->GetIndex(IndexOffset + Index);
-		uint32 NewIndex;
+		const uint32 IndexOffset = MeshSection->FirstIndex;
+		const uint32 IndexCount = MeshSection->NumTriangles * 3;
 
-		if (uint32* FoundIndex = IndexLookup.Find(OldIndex))
+		for (uint32 Index = 0; Index < IndexCount; Index++)
 		{
-			NewIndex = *FoundIndex;
-		}
-		else
-		{
-			NewIndex = IndexMap.Num();
-			IndexLookup.Add(OldIndex, NewIndex);
-			IndexMap.Add(OldIndex);
-		}
+			const uint32 OldIndex = MeshLOD->IndexBuffer.GetIndex(IndexOffset + Index);
+			uint32 NewIndex;
+
+			if (uint32* FoundIndex = IndexLookup.Find(OldIndex))
+			{
+				NewIndex = *FoundIndex;
+			}
+			else
+			{
+				NewIndex = IndexMap.Num();
+				IndexLookup.Add(OldIndex, NewIndex);
+				IndexMap.Add(OldIndex);
+			}
 
-		IndexBuffer[Index] = NewIndex;
+			IndexBuffer[Index] = NewIndex;
+			BoneMapLookup[Index] = 0;
+		}
 	}
 
+	BoneMaps.Add({});
 	MaxBoneIndex = 0;
 }
 
-FGLTFMeshSection::FGLTFMeshSection(const FSkelMeshRenderSection* MeshSection, const FRawStaticIndexBuffer16or32Interface* OldIndexBuffer)
+FGLTFMeshSection::FGLTFMeshSection(const FSkeletalMeshLODRenderData* MeshLOD, const uint16 MaterialIndex)
 {
-	const uint32 IndexOffset = MeshSection->BaseIndex;
-	const uint32 IndexCount = MeshSection->NumTriangles * 3;
+	uint32 TotalIndexCount = 0;
 
-	IndexMap.Reserve(IndexCount);
-	IndexBuffer.AddUninitialized(IndexCount);
+	TArray<const FSkelMeshRenderSection*> Sections;
+	for (const FSkelMeshRenderSection& Section : MeshLOD->RenderSections)
+	{
+		if (Section.MaterialIndex == MaterialIndex)
+		{
+			Sections.Add(&Section);
+			TotalIndexCount += Section.NumTriangles * 3;
+		}
+	}
+
+	IndexMap.Reserve(TotalIndexCount);
+	IndexBuffer.AddUninitialized(TotalIndexCount);
+	BoneMapLookup.AddUninitialized(TotalIndexCount);
 
 	TMap<uint32, uint32> IndexLookup;
 
-	for (uint32 Index = 0; Index < IndexCount; Index++)
+	const FRawStaticIndexBuffer16or32Interface* OldIndexBuffer = MeshLOD->MultiSizeIndexContainer.GetIndexBuffer();
+
+	for (const FSkelMeshRenderSection* MeshSection : Sections)
 	{
-		const uint32 OldIndex = OldIndexBuffer->Get(IndexOffset + Index);
-		uint32 NewIndex;
+		const uint32 IndexOffset = MeshSection->BaseIndex;
+		const uint32 IndexCount = MeshSection->NumTriangles * 3;
+		const uint32 BoneMapIndex = BoneMaps.Num();
 
-		if (uint32* FoundIndex = IndexLookup.Find(OldIndex))
+		for (uint32 Index = 0; Index < IndexCount; Index++)
 		{
-			NewIndex = *FoundIndex;
-		}
-		else
-		{
-			NewIndex = IndexMap.Num();
-			IndexLookup.Add(OldIndex, NewIndex);
-			IndexMap.Add(OldIndex);
+			const uint32 OldIndex = OldIndexBuffer->Get(IndexOffset + Index);
+			uint32 NewIndex;
+
+			if (uint32* FoundIndex = IndexLookup.Find(OldIndex))
+			{
+				NewIndex = *FoundIndex;
+			}
+			else
+			{
+				NewIndex = IndexMap.Num();
+				IndexLookup.Add(OldIndex, NewIndex);
+				IndexMap.Add(OldIndex);
+			}
+
+			IndexBuffer[Index] = NewIndex;
+			BoneMapLookup[Index] = BoneMapIndex;
 		}
 
-		IndexBuffer[Index] = NewIndex;
+		BoneMaps.Add(MeshSection->BoneMap);
+		MaxBoneIndex = MeshSection->BoneMap.Num() > 0 ? *Algo::MaxElement(MeshSection->BoneMap) : 0;
 	}
-
-	BoneMap = MeshSection->BoneMap;
-	MaxBoneIndex = BoneMap.Num() > 0 ? *Algo::MaxElement(MeshSection->BoneMap) : 0;
 }
diff --git a/Engine/Plugins/Enterprise/GLTFExporter/MigratedFromGit/Source/GLTFExporter/Private/Converters/GLTFMeshSection.h b/Engine/Plugins/Enterprise/GLTFExporter/MigratedFromGit/Source/GLTFExporter/Private/Converters/GLTFMeshSection.h
--- a/Engine/Plugins/Enterprise/GLTFExporter/MigratedFromGit/Source/GLTFExporter/Private/Converters/GLTFMeshSection.h
+++ b/Engine/Plugins/Enterprise/GLTFExporter/MigratedFromGit/Source/GLTFExporter/Private/Converters/GLTFMeshSection.h
@@ -6,12 +6,13 @@
 
 struct FGLTFMeshSection
 {
-	FGLTFMeshSection(const FStaticMeshSection* MeshSection, const FRawStaticIndexBuffer* OldIndexBuffer);
-	FGLTFMeshSection(const FSkelMeshRenderSection* MeshSection, const FRawStaticIndexBuffer16or32Interface* OldIndexBuffer);
+	FGLTFMeshSection(const FStaticMeshLODResources* MeshLOD, const int32 MaterialIndex);
+	FGLTFMeshSection(const FSkeletalMeshLODRenderData* MeshLOD, const uint16 MaterialIndex);
 
 	TArray<uint32> IndexMap;
 	TArray<uint32> IndexBuffer;
 
-	TArray<FBoneIndexType> BoneMap;
+	TArray<TArray<FBoneIndexType>> BoneMaps;
+	TArray<uint32> BoneMapLookup;
 	uint32 MaxBoneIndex;
 };
