[GLTFExporter] fix: make BoneMapLookup be equal in size to IndexMap

commit 63a5f8125a527751903769a46fa35de62f1c441c
Date: Thu, 1 Oct 2020 15:53:02 +0200

    fix: make BoneMapLookup be equal in size to IndexMap
    (cherry picked from commit f2ae756dd2e445d2be3ca14287f6588592c82b86​)

#skipci

diff --git a/Engine/Plugins/Enterprise/GLTFExporter/MigratedFromGit/Source/GLTFExporter/Private/Converters/GLTFMeshSection.cpp b/Engine/Plugins/Enterprise/GLTFExporter/MigratedFromGit/Source/GLTFExporter/Private/Converters/GLTFMeshSection.cpp
--- a/Engine/Plugins/Enterprise/GLTFExporter/MigratedFromGit/Source/GLTFExporter/Private/Converters/GLTFMeshSection.cpp
+++ b/Engine/Plugins/Enterprise/GLTFExporter/MigratedFromGit/Source/GLTFExporter/Private/Converters/GLTFMeshSection.cpp
@@ -21,7 +21,7 @@ FGLTFMeshSection::FGLTFMeshSection(const FStaticMeshLODResources* MeshLOD, const
 
 	IndexMap.Reserve(TotalIndexCount);
 	IndexBuffer.AddUninitialized(TotalIndexCount);
-	BoneMapLookup.AddUninitialized(TotalIndexCount);
+	BoneMapLookup.Reserve(TotalIndexCount);
 
 	TMap<uint32, uint32> IndexLookup;
 
@@ -44,10 +44,10 @@ FGLTFMeshSection::FGLTFMeshSection(const FStaticMeshLODResources* MeshLOD, const
 				NewIndex = IndexMap.Num();
 				IndexLookup.Add(OldIndex, NewIndex);
 				IndexMap.Add(OldIndex);
+				BoneMapLookup.Add(0);
 			}
 
 			IndexBuffer[Index] = NewIndex;
-			BoneMapLookup[Index] = 0;
 		}
 	}
 
@@ -71,7 +71,7 @@ FGLTFMeshSection::FGLTFMeshSection(const FSkeletalMeshLODRenderData* MeshLOD, co
 
 	IndexMap.Reserve(TotalIndexCount);
 	IndexBuffer.AddUninitialized(TotalIndexCount);
-	BoneMapLookup.AddUninitialized(TotalIndexCount);
+	BoneMapLookup.Reserve(TotalIndexCount);
 
 	TMap<uint32, uint32> IndexLookup;
 
@@ -97,10 +97,10 @@ FGLTFMeshSection::FGLTFMeshSection(const FSkeletalMeshLODRenderData* MeshLOD, co
 				NewIndex = IndexMap.Num();
 				IndexLookup.Add(OldIndex, NewIndex);
 				IndexMap.Add(OldIndex);
+				BoneMapLookup.Add(BoneMapIndex);
 			}
 
 			IndexBuffer[Index] = NewIndex;
-			BoneMapLookup[Index] = BoneMapIndex;
 		}
 
 		BoneMaps.Add(MeshSection->BoneMap);
