[GLTFExporter] fix: avoid crash when calculating bounding-boxes for empty sections

commit 709b42b2570e41517a31fd211dd699b865b7cd2f
Date: Thu, 1 Oct 2020 16:51:45 +0200

    fix: avoid crash when calculating bounding-boxes for empty sections
    (cherry picked from commit a56bf4fda59f6188ae256a9ed78d46e0ea203948​)

#skipci

diff --git a/Engine/Plugins/Enterprise/GLTFExporter/MigratedFromGit/Source/GLTFExporter/Private/Converters/GLTFAccessorConverters.cpp b/Engine/Plugins/Enterprise/GLTFExporter/MigratedFromGit/Source/GLTFExporter/Private/Converters/GLTFAccessorConverters.cpp
--- a/Engine/Plugins/Enterprise/GLTFExporter/MigratedFromGit/Source/GLTFExporter/Private/Converters/GLTFAccessorConverters.cpp
+++ b/Engine/Plugins/Enterprise/GLTFExporter/MigratedFromGit/Source/GLTFExporter/Private/Converters/GLTFAccessorConverters.cpp
@@ -64,8 +64,8 @@ FGLTFJsonAccessorIndex FGLTFPositionBufferConverter::Convert(const FGLTFMeshSect
 	}
 
 	// More accurate bounding box if based on raw vertex values
-	FGLTFJsonVector3 MinPosition = Positions[0];
-	FGLTFJsonVector3 MaxPosition = Positions[0];
+	FGLTFJsonVector3 MinPosition = VertexCount > 0 ? Positions[0] : FGLTFJsonVector3::Zero;
+	FGLTFJsonVector3 MaxPosition = VertexCount > 0 ? Positions[0] : FGLTFJsonVector3::Zero;
 
 	for (uint32 VertexIndex = 1; VertexIndex < VertexCount; ++VertexIndex)
 	{
