[GLTFExporter] refactor: use correct type for MaxBoneIndex

commit a73445571ecf2b8d8836ec38dcb9a2fce60f0cf7
Date: Fri, 2 Oct 2020 08:53:15 +0200

    refactor: use correct type for MaxBoneIndex
    (cherry picked from commit 443eb9544127ecf5e60b67bb167d1b6e2bfb684d​)

#skipci

diff --git a/Engine/Plugins/Enterprise/GLTFExporter/MigratedFromGit/Source/GLTFExporter/Private/Converters/GLTFMeshSection.h b/Engine/Plugins/Enterprise/GLTFExporter/MigratedFromGit/Source/GLTFExporter/Private/Converters/GLTFMeshSection.h
--- a/Engine/Plugins/Enterprise/GLTFExporter/MigratedFromGit/Source/GLTFExporter/Private/Converters/GLTFMeshSection.h
+++ b/Engine/Plugins/Enterprise/GLTFExporter/MigratedFromGit/Source/GLTFExporter/Private/Converters/GLTFMeshSection.h
@@ -14,5 +14,5 @@ struct FGLTFMeshSection
 
 	TArray<TArray<FBoneIndexType>> BoneMaps;
 	TArray<uint32> BoneMapLookup;
-	uint32 MaxBoneIndex;
+	FBoneIndexType MaxBoneIndex;
 };
