[GLTFExporter] fix: use all mesh-sections when calculating MaxBoneIndex instead of just the last one

commit 5a8c67053c9a1d7775bc0be03088381642346b55
Date: Fri, 2 Oct 2020 08:54:08 +0200

    fix: use all mesh-sections when calculating MaxBoneIndex instead of just the last one
    (cherry picked from commit c39672b4c07aa6bf2bb3ee464e87251e1de11e45​)

#skipci

diff --git a/Engine/Plugins/Enterprise/GLTFExporter/MigratedFromGit/Source/GLTFExporter/Private/Converters/GLTFMeshSection.cpp b/Engine/Plugins/Enterprise/GLTFExporter/MigratedFromGit/Source/GLTFExporter/Private/Converters/GLTFMeshSection.cpp
--- a/Engine/Plugins/Enterprise/GLTFExporter/MigratedFromGit/Source/GLTFExporter/Private/Converters/GLTFMeshSection.cpp
+++ b/Engine/Plugins/Enterprise/GLTFExporter/MigratedFromGit/Source/GLTFExporter/Private/Converters/GLTFMeshSection.cpp
@@ -56,6 +56,7 @@ FGLTFMeshSection::FGLTFMeshSection(const FStaticMeshLODResources* MeshLOD, const
 }
 
 FGLTFMeshSection::FGLTFMeshSection(const FSkeletalMeshLODRenderData* MeshLOD, const uint16 MaterialIndex)
+	: MaxBoneIndex(0)
 {
 	uint32 TotalIndexCount = 0;
 
@@ -104,6 +105,6 @@ FGLTFMeshSection::FGLTFMeshSection(const FSkeletalMeshLODRenderData* MeshLOD, co
 		}
 
 		BoneMaps.Add(MeshSection->BoneMap);
-		MaxBoneIndex = MeshSection->BoneMap.Num() > 0 ? *Algo::MaxElement(MeshSection->BoneMap) : 0;
+		MaxBoneIndex = MeshSection->BoneMap.Num() > 0 ? FMath::Max(*Algo::MaxElement(MeshSection->BoneMap), MaxBoneIndex) : MaxBoneIndex;
 	}
 }
