[GLTFExporter] refactor: move calculation for MaxBoneIndex into if-clause for clarity

commit 71e9037220433c7ddd81c1768fa93c47e3c1100c
Date: Fri, 2 Oct 2020 16:36:20 +0200

    refactor: move calculation for MaxBoneIndex into if-clause for clarity
    (cherry picked from commit 433bb0ad107109788e7560cdba615b3cbc5d68c3​)

#skipci

diff --git a/Engine/Plugins/Enterprise/GLTFExporter/MigratedFromGit/Source/GLTFExporter/Private/Converters/GLTFMeshSection.cpp b/Engine/Plugins/Enterprise/GLTFExporter/MigratedFromGit/Source/GLTFExporter/Private/Converters/GLTFMeshSection.cpp
--- a/Engine/Plugins/Enterprise/GLTFExporter/MigratedFromGit/Source/GLTFExporter/Private/Converters/GLTFMeshSection.cpp
+++ b/Engine/Plugins/Enterprise/GLTFExporter/MigratedFromGit/Source/GLTFExporter/Private/Converters/GLTFMeshSection.cpp
@@ -105,6 +105,10 @@ FGLTFMeshSection::FGLTFMeshSection(const FSkeletalMeshLODRenderData* MeshLOD, co
 		}
 
 		BoneMaps.Add(MeshSection->BoneMap);
-		MaxBoneIndex = MeshSection->BoneMap.Num() > 0 ? FMath::Max(*Algo::MaxElement(MeshSection->BoneMap), MaxBoneIndex) : MaxBoneIndex;
+
+		if (const FBoneIndexType* MaxSectionBoneIndex = Algo::MaxElement(MeshSection->BoneMap))
+		{
+		    MaxBoneIndex = FMath::Max(*MaxSectionBoneIndex, MaxBoneIndex);
+		}
 	}
 }
