[GLTFExporter] feat: disable property AnimationSequence if no SkeletalMeshActor is assigned

commit d519e6908521b16bfd329e19a5fb85a0da8f8fa7
Date: Fri, 2 Oct 2020 18:12:13 +0200

    feat: disable property AnimationSequence if no SkeletalMeshActor is assigned
    (cherry picked from commit 5981455ace0babfab895f45be76d144cebcdd40e​)

#skipci

diff --git a/Engine/Plugins/Enterprise/GLTFExporter/MigratedFromGit/Source/GLTFExporterRuntime/Public/Actors/GLTFInteractionHotspotActor.h b/Engine/Plugins/Enterprise/GLTFExporter/MigratedFromGit/Source/GLTFExporterRuntime/Public/Actors/GLTFInteractionHotspotActor.h
--- a/Engine/Plugins/Enterprise/GLTFExporter/MigratedFromGit/Source/GLTFExporterRuntime/Public/Actors/GLTFInteractionHotspotActor.h
+++ b/Engine/Plugins/Enterprise/GLTFExporter/MigratedFromGit/Source/GLTFExporterRuntime/Public/Actors/GLTFInteractionHotspotActor.h
@@ -62,7 +62,7 @@ public:
 	ASkeletalMeshActor* SkeletalMeshActor;
 
 	/* The animation that will be played on the skeletal mesh actor. Must be compatible with its skeletal mesh asset. */
-	UPROPERTY(EditAnywhere, BlueprintReadWrite, Category = "GLTF Interaction Hotspot")
+	UPROPERTY(EditAnywhere, BlueprintReadWrite, Category = "GLTF Interaction Hotspot", meta=(EditCondition="SkeletalMeshActor != nullptr"))
 	UAnimSequence* AnimationSequence;
 
 	/* The billboard image that will be shown when the hotspot is in an inactive state or one without a specified image. */
