[GLTFExporter] refactor: use proper compile-time safe property name checks instead of string

commit 6664c47f074699748a5a56de235e9d5b19b938a8
Date: Fri, 2 Oct 2020 18:14:02 +0200

    refactor: use proper compile-time safe property name checks instead of string
    (cherry picked from commit 378a98e2819c51f210780bce348438f9832c528a​)

#skipci

diff --git a/Engine/Plugins/Enterprise/GLTFExporter/MigratedFromGit/Source/GLTFExporterRuntime/Private/Actors/GLTFInteractionHotspotActor.cpp b/Engine/Plugins/Enterprise/GLTFExporter/MigratedFromGit/Source/GLTFExporterRuntime/Private/Actors/GLTFInteractionHotspotActor.cpp
--- a/Engine/Plugins/Enterprise/GLTFExporter/MigratedFromGit/Source/GLTFExporterRuntime/Private/Actors/GLTFInteractionHotspotActor.cpp
+++ b/Engine/Plugins/Enterprise/GLTFExporter/MigratedFromGit/Source/GLTFExporterRuntime/Private/Actors/GLTFInteractionHotspotActor.cpp
@@ -118,20 +118,20 @@ void AGLTFInteractionHotspotActor::PostEditChangeProperty(FPropertyChangedEvent&
 
 	if (PropertyThatChanged)
 	{
-		const FString PropertyName = PropertyThatChanged->GetName();
+		const FName PropertyFName = PropertyThatChanged->GetFName();
 
-		if (PropertyName == TEXT("Image"))
+		if (PropertyFName == GET_MEMBER_NAME_CHECKED(ThisClass, Image))
 		{
 			UpdateActiveImageFromState(EGLTFHotspotState::Default);
 		}
-		else if (PropertyName == TEXT("SkeletalMeshActor"))
+		else if (PropertyFName == GET_MEMBER_NAME_CHECKED(ThisClass, SkeletalMeshActor))
 		{
 			if (AnimationSequence != nullptr && AnimationSequence->GetSkeleton() != SkeletalMeshActor->GetSkeletalMeshComponent()->SkeletalMesh->Skeleton)
 			{
 				UE_LOG(LogEditorGLTFInteractionHotspot, Warning, TEXT("The skeleton of this actor is not compatible with the previously assigned animation sequence"));
 			}
 		}
-		else if (PropertyName == TEXT("AnimationSequence"))
+		else if (PropertyFName == GET_MEMBER_NAME_CHECKED(ThisClass, AnimationSequence))
 		{
 			if (SkeletalMeshActor != nullptr && SkeletalMeshActor->GetSkeletalMeshComponent()->SkeletalMesh->Skeleton != AnimationSequence->GetSkeleton())
 			{
diff --git a/Engine/Plugins/Enterprise/GLTFExporter/MigratedFromGit/Source/GLTFExporterRuntime/Public/Actors/GLTFInteractionHotspotActor.h b/Engine/Plugins/Enterprise/GLTFExporter/MigratedFromGit/Source/GLTFExporterRuntime/Public/Actors/GLTFInteractionHotspotActor.h
--- a/Engine/Plugins/Enterprise/GLTFExporter/MigratedFromGit/Source/GLTFExporterRuntime/Public/Actors/GLTFInteractionHotspotActor.h
+++ b/Engine/Plugins/Enterprise/GLTFExporter/MigratedFromGit/Source/GLTFExporterRuntime/Public/Actors/GLTFInteractionHotspotActor.h
@@ -83,31 +83,31 @@ public:
 
 private:
 
-	UPROPERTY(transient, DuplicateTransient)
+	UPROPERTY(Transient, DuplicateTransient)
 	UMaterialBillboardComponent* BillboardComponent;
 
-	UPROPERTY(transient, duplicatetransient)
+	UPROPERTY(Transient, DuplicateTransient)
 	USphereComponent* SphereComponent;
 
-	UPROPERTY(transient, duplicatetransient)
+	UPROPERTY(Transient, DuplicateTransient)
 	UMaterialInterface* DefaultMaterial;
 
-	UPROPERTY(transient, duplicatetransient)
+	UPROPERTY(Transient, DuplicateTransient)
 	UTexture2D* DefaultImage;
 
-	UPROPERTY(transient, duplicatetransient)
+	UPROPERTY(Transient, DuplicateTransient)
 	UTexture2D* DefaultHoveredImage;
 
-	UPROPERTY(transient, duplicatetransient)
+	UPROPERTY(Transient, DuplicateTransient)
 	UTexture2D* DefaultToggledImage;
 
-	UPROPERTY(transient, duplicatetransient)
+	UPROPERTY(Transient, DuplicateTransient)
 	UTexture2D* DefaultToggledHoveredImage;
 
-	UPROPERTY(transient, duplicatetransient)
+	UPROPERTY(Transient, DuplicateTransient)
 	UMaterialInterface* DefaultIconMaterial;
 
-	UPROPERTY(transient, duplicatetransient)
+	UPROPERTY(Transient, DuplicateTransient)
 	const UTexture* ActiveImage;
 
 	FVector2D ActiveImageSize;
