[GLTFExporter] fix: refactor validation checks in hotspot actor to avoid nullptr bug crash

commit 3fd72a6cb89a64f28f7432a99d5e088d1530fe18
Date: Fri, 2 Oct 2020 18:49:33 +0200

    fix: refactor validation checks in hotspot actor to avoid nullptr bug crash
    (cherry picked from commit 3ff73e10b787dd9d3b14a6b92b76f1e1de4bcc03​)

#skipci

diff --git a/Engine/Plugins/Enterprise/GLTFExporter/MigratedFromGit/Source/GLTFExporterRuntime/Private/Actors/GLTFInteractionHotspotActor.cpp b/Engine/Plugins/Enterprise/GLTFExporter/MigratedFromGit/Source/GLTFExporterRuntime/Private/Actors/GLTFInteractionHotspotActor.cpp
--- a/Engine/Plugins/Enterprise/GLTFExporter/MigratedFromGit/Source/GLTFExporterRuntime/Private/Actors/GLTFInteractionHotspotActor.cpp
+++ b/Engine/Plugins/Enterprise/GLTFExporter/MigratedFromGit/Source/GLTFExporterRuntime/Private/Actors/GLTFInteractionHotspotActor.cpp
@@ -126,17 +126,15 @@ void AGLTFInteractionHotspotActor::PostEditChangeProperty(FPropertyChangedEvent&
 		}
 		else if (PropertyFName == GET_MEMBER_NAME_CHECKED(ThisClass, SkeletalMeshActor))
 		{
-			if (AnimationSequence != nullptr && AnimationSequence->GetSkeleton() != SkeletalMeshActor->GetSkeletalMeshComponent()->SkeletalMesh->Skeleton)
+			if (SkeletalMeshActor != nullptr && AnimationSequence == nullptr)
 			{
-				UE_LOG(LogEditorGLTFInteractionHotspot, Warning, TEXT("The skeleton of this actor is not compatible with the previously assigned animation sequence"));
+				AnimationSequence = Cast<UAnimSequence>(SkeletalMeshActor->GetSkeletalMeshComponent()->AnimationData.AnimToPlay);
 			}
+			ValidateAnimation();
 		}
 		else if (PropertyFName == GET_MEMBER_NAME_CHECKED(ThisClass, AnimationSequence))
 		{
-			if (SkeletalMeshActor != nullptr && SkeletalMeshActor->GetSkeletalMeshComponent()->SkeletalMesh->Skeleton != AnimationSequence->GetSkeleton())
-			{
-				UE_LOG(LogEditorGLTFInteractionHotspot, Warning, TEXT("This animation sequence is not compatible with the skeleton of the previously assigned actor"));
-			}
+			ValidateAnimation();
 		}
 	}
 }
@@ -402,6 +400,29 @@ void AGLTFInteractionHotspotActor::ViewportResized(FViewport*, uint32)
 	UpdateSpriteSize();
 }
 
+void AGLTFInteractionHotspotActor::ValidateAnimation()
+{
+	if (SkeletalMeshActor != nullptr && AnimationSequence != nullptr)
+	{
+		if (USkeletalMesh* SkeletalMesh = SkeletalMeshActor->GetSkeletalMeshComponent()->SkeletalMesh)
+		{
+			if (SkeletalMesh->Skeleton != AnimationSequence->GetSkeleton())
+			{
+				if (SkeletalMesh->Skeleton != nullptr)
+				{
+					UE_LOG(LogEditorGLTFInteractionHotspot, Warning, TEXT("Animation %s is incompatible with skeleton %s, removing animation from actor."), *AnimationSequence->GetName(), *SkeletalMesh->Skeleton->GetName());
+				}
+				else
+				{
+					UE_LOG(LogEditorGLTFInteractionHotspot, Warning, TEXT("Animation %s is incompatible because mesh %s has no skeleton, removing animation from actor."), *AnimationSequence->GetName(), *SkeletalMesh->GetName());
+				}
+
+				AnimationSequence = nullptr;
+			}
+		}
+	}
+}
+
 const UTexture2D* AGLTFInteractionHotspotActor::GetImageForState(EGLTFHotspotState State) const
 {
 	const UTexture2D* CurrentImage = DefaultImage;
diff --git a/Engine/Plugins/Enterprise/GLTFExporter/MigratedFromGit/Source/GLTFExporterRuntime/Public/Actors/GLTFInteractionHotspotActor.h b/Engine/Plugins/Enterprise/GLTFExporter/MigratedFromGit/Source/GLTFExporterRuntime/Public/Actors/GLTFInteractionHotspotActor.h
--- a/Engine/Plugins/Enterprise/GLTFExporter/MigratedFromGit/Source/GLTFExporterRuntime/Public/Actors/GLTFInteractionHotspotActor.h
+++ b/Engine/Plugins/Enterprise/GLTFExporter/MigratedFromGit/Source/GLTFExporterRuntime/Public/Actors/GLTFInteractionHotspotActor.h
@@ -132,6 +132,8 @@ private:
 
 	void ViewportResized(FViewport*, uint32);
 
+	void ValidateAnimation();
+
 public:
 
 	const UTexture2D* GetImageForState(EGLTFHotspotState State) const;
