[GLTFExporter] refactor: move skeleton compatibility check from hotspot converter to animation converter

commit dea91eef9b5493d7fa0040e4687999c32455725d
Date: Fri, 2 Oct 2020 19:04:54 +0200

    refactor: move skeleton compatibility check from hotspot converter to animation converter
    (cherry picked from commit 726c6d992f7fe104520e49fc369694ac3c2da5f8​)

#skipci

diff --git a/Engine/Plugins/Enterprise/GLTFExporter/MigratedFromGit/Source/GLTFExporter/Private/Converters/GLTFAnimationConverters.cpp b/Engine/Plugins/Enterprise/GLTFExporter/MigratedFromGit/Source/GLTFExporter/Private/Converters/GLTFAnimationConverters.cpp
--- a/Engine/Plugins/Enterprise/GLTFExporter/MigratedFromGit/Source/GLTFExporter/Private/Converters/GLTFAnimationConverters.cpp
+++ b/Engine/Plugins/Enterprise/GLTFExporter/MigratedFromGit/Source/GLTFExporter/Private/Converters/GLTFAnimationConverters.cpp
@@ -23,6 +23,12 @@ FGLTFJsonAnimationIndex FGLTFAnimationConverter::Convert(FGLTFJsonNodeIndex Root
 		return FGLTFJsonAnimationIndex(INDEX_NONE);
 	}
 
+	if (Skeleton != SkeletalMesh->Skeleton)
+	{
+		// TODO: report error
+		return FGLTFJsonAnimationIndex(INDEX_NONE);
+	}
+
 	FGLTFJsonAnimation JsonAnimation;
 	AnimSequence->GetName(JsonAnimation.Name);
 
diff --git a/Engine/Plugins/Enterprise/GLTFExporter/MigratedFromGit/Source/GLTFExporter/Private/Converters/GLTFHotspotConverters.cpp b/Engine/Plugins/Enterprise/GLTFExporter/MigratedFromGit/Source/GLTFExporter/Private/Converters/GLTFHotspotConverters.cpp
--- a/Engine/Plugins/Enterprise/GLTFExporter/MigratedFromGit/Source/GLTFExporter/Private/Converters/GLTFHotspotConverters.cpp
+++ b/Engine/Plugins/Enterprise/GLTFExporter/MigratedFromGit/Source/GLTFExporter/Private/Converters/GLTFHotspotConverters.cpp
@@ -27,18 +27,13 @@ FGLTFJsonHotspotIndex FGLTFHotspotConverter::Convert(const AGLTFInteractionHotsp
 		if (const ASkeletalMeshActor* SkeletalMeshActor = HotspotActor->SkeletalMeshActor)
 		{
 			const FGLTFJsonNodeIndex RootNode = Builder.GetOrAddNode(SkeletalMeshActor);
-			if (const USkeletalMesh* SkeletalMesh = SkeletalMeshActor->GetSkeletalMeshComponent()->SkeletalMesh)
+
+			const USkeletalMeshComponent* SkeletalMeshComponent = SkeletalMeshActor->GetSkeletalMeshComponent();
+			if (const USkeletalMesh* SkeletalMesh = SkeletalMeshComponent->SkeletalMesh)
 			{
 				if (const UAnimSequence* AnimSequence = HotspotActor->AnimationSequence)
 				{
-					if (SkeletalMesh->Skeleton == AnimSequence->GetSkeleton())
-					{
-						JsonHotspot.Animation = Builder.GetOrAddAnimation(RootNode, SkeletalMesh, AnimSequence);
-					}
-					else
-					{
-						// TODO: report warning
-					}
+					JsonHotspot.Animation = Builder.GetOrAddAnimation(RootNode, SkeletalMesh, AnimSequence);
 				}
 				else
 				{
