[GLTFExporter] fix: A lot of engine code interprets AnisotropicPoint as linear instead of nearest, which also give more correct results.

commit 7a5db6d3b6e3e16696949fc5320a88b9e241454b
Date: Fri, 2 Oct 2020 21:18:43 +0200

    fix: A lot of engine code interprets AnisotropicPoint as linear instead of nearest, which also give more correct results.
    (cherry picked from commit f48a93caa94fda87ac8016130321e8e9c9222752​)

#skipci

diff --git a/Engine/Plugins/Enterprise/GLTFExporter/MigratedFromGit/Source/GLTFExporter/Private/Converters/GLTFTextureUtility.cpp b/Engine/Plugins/Enterprise/GLTFExporter/MigratedFromGit/Source/GLTFExporter/Private/Converters/GLTFTextureUtility.cpp
--- a/Engine/Plugins/Enterprise/GLTFExporter/MigratedFromGit/Source/GLTFExporter/Private/Converters/GLTFTextureUtility.cpp
+++ b/Engine/Plugins/Enterprise/GLTFExporter/MigratedFromGit/Source/GLTFExporter/Private/Converters/GLTFTextureUtility.cpp
@@ -83,9 +83,9 @@ TextureFilter FGLTFTextureUtility::GetDefaultFilter(TextureGroup LODGroup)
 		case ETextureSamplerFilter::Point:             return TF_Nearest;
 		case ETextureSamplerFilter::Bilinear:          return TF_Bilinear;
 		case ETextureSamplerFilter::Trilinear:         return TF_Trilinear;
-		case ETextureSamplerFilter::AnisotropicPoint:  return TF_Nearest;
+		case ETextureSamplerFilter::AnisotropicPoint:  return TF_Trilinear; // A lot of engine code doesn't result in nearest
 		case ETextureSamplerFilter::AnisotropicLinear: return TF_Trilinear;
-		default:                                       return TF_Default; // fallback
+		default:                                       return TF_Default; // Let caller decide fallback
 	}
 }
 
