[GLTFExporter] refactor: rearrange node converters group socket conversion consistently

commit 47a75de16b24bdea68d057c410214be309b30c21
Date: Fri, 2 Oct 2020 21:30:56 +0200

    refactor: rearrange node converters group socket conversion consistently
    (cherry picked from commit 38849f810f57430d62fe88f3c7a46bc578681984​)

#skipci

diff --git a/Engine/Plugins/Enterprise/GLTFExporter/MigratedFromGit/Source/GLTFExporter/Private/Builders/GLTFConvertBuilder.cpp b/Engine/Plugins/Enterprise/GLTFExporter/MigratedFromGit/Source/GLTFExporter/Private/Builders/GLTFConvertBuilder.cpp
--- a/Engine/Plugins/Enterprise/GLTFExporter/MigratedFromGit/Source/GLTFExporter/Private/Builders/GLTFConvertBuilder.cpp
+++ b/Engine/Plugins/Enterprise/GLTFExporter/MigratedFromGit/Source/GLTFExporter/Private/Builders/GLTFConvertBuilder.cpp
@@ -199,26 +199,6 @@ FGLTFJsonTextureIndex FGLTFConvertBuilder::GetOrAddTexture(const UTextureRenderT
 	return TextureRenderTargetCubeConverter.GetOrAdd(Texture, CubeFace);
 }
 
-FGLTFJsonNodeIndex FGLTFConvertBuilder::GetOrAddNode(FGLTFJsonNodeIndex RootNode, const USkeletalMesh* SkeletalMesh, FName SocketName)
-{
-	if (RootNode == INDEX_NONE || SkeletalMesh == nullptr || SocketName == NAME_None)
-	{
-		return FGLTFJsonNodeIndex(INDEX_NONE);
-	}
-
-	return SkeletalSocketConverter.GetOrAdd(RootNode, SkeletalMesh, SocketName);
-}
-
-FGLTFJsonNodeIndex FGLTFConvertBuilder::GetOrAddNode(FGLTFJsonNodeIndex RootNode, const USkeletalMesh* SkeletalMesh, int32 BoneIndex)
-{
-	if (RootNode == INDEX_NONE || SkeletalMesh == nullptr || BoneIndex == INDEX_NONE)
-	{
-		return FGLTFJsonNodeIndex(INDEX_NONE);
-	}
-
-	return SkeletalBoneConverter.GetOrAdd(RootNode, SkeletalMesh, BoneIndex);
-}
-
 FGLTFJsonSkinIndex FGLTFConvertBuilder::GetOrAddSkin(FGLTFJsonNodeIndex RootNode, const USkeletalMesh* SkeletalMesh)
 {
 	if (RootNode == INDEX_NONE || SkeletalMesh == nullptr)
@@ -266,14 +246,14 @@ FGLTFJsonAnimationIndex FGLTFConvertBuilder::GetOrAddAnimation(FGLTFJsonNodeInde
 	return AnimationDataConverter.GetOrAdd(RootNode, SkeletalMeshComponent);
 }
 
-FGLTFJsonNodeIndex FGLTFConvertBuilder::GetOrAddNode(const USceneComponent* SceneComponent, FName SocketName)
+FGLTFJsonNodeIndex FGLTFConvertBuilder::GetOrAddNode(const AActor* Actor)
 {
-	if (SceneComponent == nullptr)
+	if (Actor == nullptr)
 	{
 		return FGLTFJsonNodeIndex(INDEX_NONE);
 	}
 
-	return ComponentSocketConverter.GetOrAdd(SceneComponent, SocketName);
+	return ActorConverter.GetOrAdd(Actor);
 }
 
 FGLTFJsonNodeIndex FGLTFConvertBuilder::GetOrAddNode(const USceneComponent* SceneComponent)
@@ -286,14 +266,34 @@ FGLTFJsonNodeIndex FGLTFConvertBuilder::GetOrAddNode(const USceneComponent* Scen
 	return ComponentConverter.GetOrAdd(SceneComponent);
 }
 
-FGLTFJsonNodeIndex FGLTFConvertBuilder::GetOrAddNode(const AActor* Actor)
+FGLTFJsonNodeIndex FGLTFConvertBuilder::GetOrAddNode(const USceneComponent* SceneComponent, FName SocketName)
 {
-	if (Actor == nullptr)
+	if (SceneComponent == nullptr)
 	{
 		return FGLTFJsonNodeIndex(INDEX_NONE);
 	}
 
-	return ActorConverter.GetOrAdd(Actor);
+	return ComponentSocketConverter.GetOrAdd(SceneComponent, SocketName);
+}
+
+FGLTFJsonNodeIndex FGLTFConvertBuilder::GetOrAddNode(FGLTFJsonNodeIndex RootNode, const USkeletalMesh* SkeletalMesh, FName SocketName)
+{
+	if (RootNode == INDEX_NONE || SkeletalMesh == nullptr || SocketName == NAME_None)
+	{
+		return FGLTFJsonNodeIndex(INDEX_NONE);
+	}
+
+	return SkeletalSocketConverter.GetOrAdd(RootNode, SkeletalMesh, SocketName);
+}
+
+FGLTFJsonNodeIndex FGLTFConvertBuilder::GetOrAddNode(FGLTFJsonNodeIndex RootNode, const USkeletalMesh* SkeletalMesh, int32 BoneIndex)
+{
+	if (RootNode == INDEX_NONE || SkeletalMesh == nullptr || BoneIndex == INDEX_NONE)
+	{
+		return FGLTFJsonNodeIndex(INDEX_NONE);
+	}
+
+	return SkeletalBoneConverter.GetOrAdd(RootNode, SkeletalMesh, BoneIndex);
 }
 
 FGLTFJsonSceneIndex FGLTFConvertBuilder::GetOrAddScene(const ULevel* Level)
diff --git a/Engine/Plugins/Enterprise/GLTFExporter/MigratedFromGit/Source/GLTFExporter/Private/Builders/GLTFConvertBuilder.h b/Engine/Plugins/Enterprise/GLTFExporter/MigratedFromGit/Source/GLTFExporter/Private/Builders/GLTFConvertBuilder.h
--- a/Engine/Plugins/Enterprise/GLTFExporter/MigratedFromGit/Source/GLTFExporter/Private/Builders/GLTFConvertBuilder.h
+++ b/Engine/Plugins/Enterprise/GLTFExporter/MigratedFromGit/Source/GLTFExporter/Private/Builders/GLTFConvertBuilder.h
@@ -50,16 +50,16 @@ public:
 	FGLTFJsonTextureIndex GetOrAddTexture(const UTextureRenderTarget2D* Texture);
 	FGLTFJsonTextureIndex GetOrAddTexture(const UTextureRenderTargetCube* Texture, ECubeFace CubeFace);
 
-	FGLTFJsonNodeIndex GetOrAddNode(FGLTFJsonNodeIndex RootNode, const USkeletalMesh* SkeletalMesh, FName SocketName);
-	FGLTFJsonNodeIndex GetOrAddNode(FGLTFJsonNodeIndex RootNode, const USkeletalMesh* SkeletalMesh, int32 BoneIndex);
 	FGLTFJsonSkinIndex GetOrAddSkin(FGLTFJsonNodeIndex RootNode, const USkeletalMesh* SkeletalMesh);
 	FGLTFJsonSkinIndex GetOrAddSkin(FGLTFJsonNodeIndex RootNode, const USkeletalMeshComponent* SkeletalMeshComponent);
 	FGLTFJsonAnimationIndex GetOrAddAnimation(FGLTFJsonNodeIndex RootNode, const USkeletalMesh* SkeletalMesh, const UAnimSequence* AnimSequence);
 	FGLTFJsonAnimationIndex GetOrAddAnimation(FGLTFJsonNodeIndex RootNode, const USkeletalMeshComponent* SkeletalMeshComponent);
 
-	FGLTFJsonNodeIndex GetOrAddNode(const USceneComponent* SceneComponent, FName SocketName);
-	FGLTFJsonNodeIndex GetOrAddNode(const USceneComponent* SceneComponent);
 	FGLTFJsonNodeIndex GetOrAddNode(const AActor* Actor);
+	FGLTFJsonNodeIndex GetOrAddNode(const USceneComponent* SceneComponent);
+	FGLTFJsonNodeIndex GetOrAddNode(const USceneComponent* SceneComponent, FName SocketName);
+	FGLTFJsonNodeIndex GetOrAddNode(FGLTFJsonNodeIndex RootNode, const USkeletalMesh* SkeletalMesh, FName SocketName);
+	FGLTFJsonNodeIndex GetOrAddNode(FGLTFJsonNodeIndex RootNode, const USkeletalMesh* SkeletalMesh, int32 BoneIndex);
 	FGLTFJsonSceneIndex GetOrAddScene(const ULevel* Level);
 	FGLTFJsonSceneIndex GetOrAddScene(const UWorld* World);
 
diff --git a/Engine/Plugins/Enterprise/GLTFExporter/MigratedFromGit/Source/GLTFExporter/Private/Converters/GLTFNodeConverters.cpp b/Engine/Plugins/Enterprise/GLTFExporter/MigratedFromGit/Source/GLTFExporter/Private/Converters/GLTFNodeConverters.cpp
--- a/Engine/Plugins/Enterprise/GLTFExporter/MigratedFromGit/Source/GLTFExporter/Private/Converters/GLTFNodeConverters.cpp
+++ b/Engine/Plugins/Enterprise/GLTFExporter/MigratedFromGit/Source/GLTFExporter/Private/Converters/GLTFNodeConverters.cpp
@@ -7,28 +7,53 @@
 #include "Converters/GLTFNameUtility.h"
 #include "Actors/GLTFInteractionHotspotActor.h"
 
-FGLTFJsonNodeIndex FGLTFComponentSocketConverter::Convert(const USceneComponent* SceneComponent, FName SocketName)
+FGLTFJsonNodeIndex FGLTFActorConverter::Convert(const AActor* Actor)
 {
-	const FGLTFJsonNodeIndex NodeIndex = Builder.GetOrAddNode(SceneComponent);
-
-	if (SocketName != NAME_None)
+	if (Builder.bSelectedActorsOnly && !Actor->IsSelected())
 	{
-		if (const USkinnedMeshComponent* SkinnedMeshComponent = Cast<USkinnedMeshComponent>(SceneComponent))
-		{
-			// TODO: add support for SocketOverrideLookup?
-			const USkeletalMesh* SkeletalMesh = SkinnedMeshComponent->SkeletalMesh;
-			return Builder.GetOrAddNode(NodeIndex, SkeletalMesh, SocketName);
-		}
-
-		// TODO: add support for non-skeletalmesh sockets
-
-		Builder.AddWarningMessage(
-			FString::Printf(TEXT("Can't export socket %s (in component %s) because it doesn't belong to a skeletal mesh"),
-			*SocketName.ToString(),
-			*SceneComponent->GetName()));
+		return FGLTFJsonNodeIndex(INDEX_NONE);
 	}
 
-	return NodeIndex;
+	const USceneComponent* RootComponent = Actor->GetRootComponent();
+	const FGLTFJsonNodeIndex RootNodeIndex = Builder.GetOrAddNode(RootComponent);
+
+	const UBlueprint* Blueprint = FGLTFActorUtility::GetBlueprintFromActor(Actor);
+	if (FGLTFActorUtility::IsSkySphereBlueprint(Blueprint))
+	{
+		// Ignore mesh and light components that are part of BP_SkySphere
+		// TODO: add support for BP_SkySphere later?
+	}
+	else if (FGLTFActorUtility::IsHDRIBackdropBlueprint(Blueprint))
+	{
+		if (Builder.ExportOptions->bExportHDRIBackdrops)
+		{
+			FGLTFJsonNode& RootNode = Builder.GetNode(RootNodeIndex);
+			RootNode.Backdrop = Builder.GetOrAddBackdrop(Actor);
+		}
+	}
+	else if (const AGLTFInteractionHotspotActor* HotspotActor = Cast<AGLTFInteractionHotspotActor>(Actor))
+	{
+		if (Builder.ExportOptions->bExportInteractionHotspots)
+		{
+			FGLTFJsonNode& RootNode = Builder.GetNode(RootNodeIndex);
+			RootNode.Hotspot = Builder.GetOrAddHotspot(HotspotActor);
+		}
+	}
+	else
+	{
+		// TODO: to reduce number of nodes, only export components that are of interest
+
+		for (const UActorComponent* Component : Actor->GetComponents())
+		{
+			const USceneComponent* SceneComponent = Cast<USceneComponent>(Component);
+			if (SceneComponent != nullptr)
+			{
+				Builder.GetOrAddNode(SceneComponent);
+			}
+		}
+	}
+
+	return RootNodeIndex;
 }
 
 FGLTFJsonNodeIndex FGLTFComponentConverter::Convert(const USceneComponent* SceneComponent)
@@ -138,53 +163,28 @@
 	return NodeIndex;
 }
 
-FGLTFJsonNodeIndex FGLTFActorConverter::Convert(const AActor* Actor)
+FGLTFJsonNodeIndex FGLTFComponentSocketConverter::Convert(const USceneComponent* SceneComponent, FName SocketName)
 {
-	if (Builder.bSelectedActorsOnly && !Actor->IsSelected())
+	const FGLTFJsonNodeIndex NodeIndex = Builder.GetOrAddNode(SceneComponent);
+
+	if (SocketName != NAME_None)
 	{
-		return FGLTFJsonNodeIndex(INDEX_NONE);
+		if (const USkinnedMeshComponent* SkinnedMeshComponent = Cast<USkinnedMeshComponent>(SceneComponent))
+		{
+			// TODO: add support for SocketOverrideLookup?
+			const USkeletalMesh* SkeletalMesh = SkinnedMeshComponent->SkeletalMesh;
+			return Builder.GetOrAddNode(NodeIndex, SkeletalMesh, SocketName);
+		}
+
+		// TODO: add support for non-skeletalmesh sockets
+
+		Builder.AddWarningMessage(
+            FString::Printf(TEXT("Can't export socket %s (in component %s) because it doesn't belong to a skeletal mesh"),
+            *SocketName.ToString(),
+            *SceneComponent->GetName()));
 	}
 
-	const USceneComponent* RootComponent = Actor->GetRootComponent();
-	const FGLTFJsonNodeIndex RootNodeIndex = Builder.GetOrAddNode(RootComponent);
-
-	const UBlueprint* Blueprint = FGLTFActorUtility::GetBlueprintFromActor(Actor);
-	if (FGLTFActorUtility::IsSkySphereBlueprint(Blueprint))
-	{
-		// Ignore mesh and light components that are part of BP_SkySphere
-		// TODO: add support for BP_SkySphere later?
-	}
-	else if (FGLTFActorUtility::IsHDRIBackdropBlueprint(Blueprint))
-	{
-		if (Builder.ExportOptions->bExportHDRIBackdrops)
-		{
-			FGLTFJsonNode& RootNode = Builder.GetNode(RootNodeIndex);
-			RootNode.Backdrop = Builder.GetOrAddBackdrop(Actor);
-		}
-	}
-	else if (const AGLTFInteractionHotspotActor* HotspotActor = Cast<AGLTFInteractionHotspotActor>(Actor))
-	{
-		if (Builder.ExportOptions->bExportInteractionHotspots)
-		{
-			FGLTFJsonNode& RootNode = Builder.GetNode(RootNodeIndex);
-			RootNode.Hotspot = Builder.GetOrAddHotspot(HotspotActor);
-		}
-	}
-	else
-	{
-		// TODO: to reduce number of nodes, only export components that are of interest
-
-		for (const UActorComponent* Component : Actor->GetComponents())
-		{
-			const USceneComponent* SceneComponent = Cast<USceneComponent>(Component);
-			if (SceneComponent != nullptr)
-			{
-				Builder.GetOrAddNode(SceneComponent);
-			}
-		}
-	}
-
-	return RootNodeIndex;
+	return NodeIndex;
 }
 
 FGLTFJsonNodeIndex FGLTFSkeletalSocketConverter::Convert(FGLTFJsonNodeIndex RootNode, const USkeletalMesh* SkeletalMesh, FName SocketName)
diff --git a/Engine/Plugins/Enterprise/GLTFExporter/MigratedFromGit/Source/GLTFExporter/Private/Converters/GLTFNodeConverters.h b/Engine/Plugins/Enterprise/GLTFExporter/MigratedFromGit/Source/GLTFExporter/Private/Converters/GLTFNodeConverters.h
--- a/Engine/Plugins/Enterprise/GLTFExporter/MigratedFromGit/Source/GLTFExporter/Private/Converters/GLTFNodeConverters.h
+++ b/Engine/Plugins/Enterprise/GLTFExporter/MigratedFromGit/Source/GLTFExporter/Private/Converters/GLTFNodeConverters.h
@@ -13,11 +13,11 @@ class TGLTFNodeConverter : public FGLTFBuilderContext, public TGLTFConverter<FGL
 	using FGLTFBuilderContext::FGLTFBuilderContext;
 };
 
-class FGLTFComponentSocketConverter final : public TGLTFNodeConverter<const USceneComponent*, FName>
+class FGLTFActorConverter final : public TGLTFNodeConverter<const AActor*>
 {
 	using TGLTFNodeConverter::TGLTFNodeConverter;
 
-	virtual FGLTFJsonNodeIndex Convert(const USceneComponent* SceneComponent, FName SocketName) override;
+	virtual FGLTFJsonNodeIndex Convert(const AActor* Actor) override;
 };
 
 class FGLTFComponentConverter final : public TGLTFNodeConverter<const USceneComponent*>
@@ -27,11 +27,11 @@ class FGLTFComponentConverter final : public TGLTFNodeConverter<const USceneComp
 	virtual FGLTFJsonNodeIndex Convert(const USceneComponent* SceneComponent) override;
 };
 
-class FGLTFActorConverter final : public TGLTFNodeConverter<const AActor*>
+class FGLTFComponentSocketConverter final : public TGLTFNodeConverter<const USceneComponent*, FName>
 {
 	using TGLTFNodeConverter::TGLTFNodeConverter;
 
-	virtual FGLTFJsonNodeIndex Convert(const AActor* Actor) override;
+	virtual FGLTFJsonNodeIndex Convert(const USceneComponent* SceneComponent, FName SocketName) override;
 };
 
 class FGLTFSkeletalSocketConverter final : public TGLTFNodeConverter<FGLTFJsonNodeIndex, const USkeletalMesh*, FName>
