[GLTFExporter] feat: add support for static mesh sockets

commit ccc3c54e1fecc5ffc4f236dec97392170d44f2f6
Date: Fri, 2 Oct 2020 21:52:17 +0200

    feat: add support for static mesh sockets
    (cherry picked from commit 75e287e28bf3dc8faee24c9d74db8d0e9c442877​)

#skipci

diff --git a/Engine/Plugins/Enterprise/GLTFExporter/MigratedFromGit/Source/GLTFExporter/Private/Builders/GLTFConvertBuilder.cpp b/Engine/Plugins/Enterprise/GLTFExporter/MigratedFromGit/Source/GLTFExporter/Private/Builders/GLTFConvertBuilder.cpp
--- a/Engine/Plugins/Enterprise/GLTFExporter/MigratedFromGit/Source/GLTFExporter/Private/Builders/GLTFConvertBuilder.cpp
+++ b/Engine/Plugins/Enterprise/GLTFExporter/MigratedFromGit/Source/GLTFExporter/Private/Builders/GLTFConvertBuilder.cpp
@@ -276,6 +276,16 @@ FGLTFJsonNodeIndex FGLTFConvertBuilder::GetOrAddNode(const USceneComponent* Scen
 	return ComponentSocketConverter.GetOrAdd(SceneComponent, SocketName);
 }
 
+FGLTFJsonNodeIndex FGLTFConvertBuilder::GetOrAddNode(FGLTFJsonNodeIndex RootNode, const UStaticMesh* StaticMesh, FName SocketName)
+{
+	if (RootNode == INDEX_NONE || StaticMesh == nullptr || SocketName == NAME_None)
+	{
+		return FGLTFJsonNodeIndex(INDEX_NONE);
+	}
+
+	return StaticSocketConverter.GetOrAdd(RootNode, StaticMesh, SocketName);
+}
+
 FGLTFJsonNodeIndex FGLTFConvertBuilder::GetOrAddNode(FGLTFJsonNodeIndex RootNode, const USkeletalMesh* SkeletalMesh, FName SocketName)
 {
 	if (RootNode == INDEX_NONE || SkeletalMesh == nullptr || SocketName == NAME_None)
diff --git a/Engine/Plugins/Enterprise/GLTFExporter/MigratedFromGit/Source/GLTFExporter/Private/Builders/GLTFConvertBuilder.h b/Engine/Plugins/Enterprise/GLTFExporter/MigratedFromGit/Source/GLTFExporter/Private/Builders/GLTFConvertBuilder.h
--- a/Engine/Plugins/Enterprise/GLTFExporter/MigratedFromGit/Source/GLTFExporter/Private/Builders/GLTFConvertBuilder.h
+++ b/Engine/Plugins/Enterprise/GLTFExporter/MigratedFromGit/Source/GLTFExporter/Private/Builders/GLTFConvertBuilder.h
@@ -58,6 +58,7 @@ public:
 	FGLTFJsonNodeIndex GetOrAddNode(const AActor* Actor);
 	FGLTFJsonNodeIndex GetOrAddNode(const USceneComponent* SceneComponent);
 	FGLTFJsonNodeIndex GetOrAddNode(const USceneComponent* SceneComponent, FName SocketName);
+	FGLTFJsonNodeIndex GetOrAddNode(FGLTFJsonNodeIndex RootNode, const UStaticMesh* StaticMesh, FName SocketName);
 	FGLTFJsonNodeIndex GetOrAddNode(FGLTFJsonNodeIndex RootNode, const USkeletalMesh* SkeletalMesh, FName SocketName);
 	FGLTFJsonNodeIndex GetOrAddNode(FGLTFJsonNodeIndex RootNode, const USkeletalMesh* SkeletalMesh, int32 BoneIndex);
 	FGLTFJsonSceneIndex GetOrAddScene(const ULevel* Level);
@@ -91,15 +92,16 @@ private:
 	FGLTFTextureRenderTarget2DConverter TextureRenderTarget2DConverter{ *this };
 	FGLTFTextureRenderTargetCubeConverter TextureRenderTargetCubeConverter{ *this };
 
-	FGLTFSkeletalSocketConverter SkeletalSocketConverter{ *this };
-	FGLTFSkeletalBoneConverter SkeletalBoneConverter{ *this };
 	FGLTFSkinConverter SkinConverter{ *this };
 	FGLTFAnimationConverter AnimationConverter{ *this };
 	FGLTFAnimationDataConverter AnimationDataConverter{ *this };
 
-	FGLTFComponentSocketConverter ComponentSocketConverter{ *this };
-	FGLTFComponentConverter ComponentConverter{ *this };
 	FGLTFActorConverter ActorConverter{ *this };
+	FGLTFComponentConverter ComponentConverter{ *this };
+	FGLTFComponentSocketConverter ComponentSocketConverter{ *this };
+	FGLTFStaticSocketConverter StaticSocketConverter{ *this };
+	FGLTFSkeletalSocketConverter SkeletalSocketConverter{ *this };
+	FGLTFSkeletalBoneConverter SkeletalBoneConverter{ *this };
 	FGLTFSceneConverter SceneConverter{ *this };
 
 	FGLTFCameraConverter CameraConverter{ *this };
diff --git a/Engine/Plugins/Enterprise/GLTFExporter/MigratedFromGit/Source/GLTFExporter/Private/Converters/GLTFNodeConverters.cpp b/Engine/Plugins/Enterprise/GLTFExporter/MigratedFromGit/Source/GLTFExporter/Private/Converters/GLTFNodeConverters.cpp
--- a/Engine/Plugins/Enterprise/GLTFExporter/MigratedFromGit/Source/GLTFExporter/Private/Converters/GLTFNodeConverters.cpp
+++ b/Engine/Plugins/Enterprise/GLTFExporter/MigratedFromGit/Source/GLTFExporter/Private/Converters/GLTFNodeConverters.cpp
@@ -169,6 +169,12 @@ FGLTFJsonNodeIndex FGLTFComponentSocketConverter::Convert(const USceneComponent*
 
 	if (SocketName != NAME_None)
 	{
+		if (const UStaticMeshComponent* StaticMeshComponent = Cast<UStaticMeshComponent>(SceneComponent))
+		{
+			const UStaticMesh* StaticMesh = StaticMeshComponent->GetStaticMesh();
+			return Builder.GetOrAddNode(NodeIndex, StaticMesh, SocketName);
+		}
+
 		if (const USkinnedMeshComponent* SkinnedMeshComponent = Cast<USkinnedMeshComponent>(SceneComponent))
 		{
 			// TODO: add support for SocketOverrideLookup?
@@ -176,10 +182,10 @@ FGLTFJsonNodeIndex FGLTFComponentSocketConverter::Convert(const USceneComponent*
 			return Builder.GetOrAddNode(NodeIndex, SkeletalMesh, SocketName);
 		}
 
-		// TODO: add support for non-skeletalmesh sockets
+		// TODO: add support for more socket types
 
 		Builder.AddWarningMessage(
-            FString::Printf(TEXT("Can't export socket %s (in component %s) because it doesn't belong to a skeletal mesh"),
+            FString::Printf(TEXT("Can't export socket %s because it belongs to an unsupported mesh component %s"),
             *SocketName.ToString(),
             *SceneComponent->GetName()));
 	}
@@ -187,6 +193,26 @@ FGLTFJsonNodeIndex FGLTFComponentSocketConverter::Convert(const USceneComponent*
 	return NodeIndex;
 }
 
+FGLTFJsonNodeIndex FGLTFStaticSocketConverter::Convert(FGLTFJsonNodeIndex RootNode, const UStaticMesh* StaticMesh, FName SocketName)
+{
+	const UStaticMeshSocket* Socket = StaticMesh->FindSocket(SocketName);
+	if (Socket == nullptr)
+	{
+		// TODO: report error
+		return FGLTFJsonNodeIndex(INDEX_NONE);
+	}
+
+	FGLTFJsonNode Node;
+	Node.Name = SocketName.ToString();
+
+	// TODO: add warning check for non-uniform scaling
+	Node.Translation = FGLTFConverterUtility::ConvertPosition(Socket->RelativeLocation, Builder.ExportOptions->ExportScale);
+	Node.Rotation = FGLTFConverterUtility::ConvertRotation(Socket->RelativeRotation.Quaternion());
+	Node.Scale = FGLTFConverterUtility::ConvertScale(Socket->RelativeScale);
+
+	return Builder.AddChildNode(RootNode, Node);
+}
+
 FGLTFJsonNodeIndex FGLTFSkeletalSocketConverter::Convert(FGLTFJsonNodeIndex RootNode, const USkeletalMesh* SkeletalMesh, FName SocketName)
 {
 	const USkeletalMeshSocket* Socket = SkeletalMesh->FindSocket(SocketName);
@@ -196,10 +222,9 @@ FGLTFJsonNodeIndex FGLTFSkeletalSocketConverter::Convert(FGLTFJsonNodeIndex Root
 		Node.Name = SocketName.ToString();
 
 		// TODO: add warning check for non-uniform scaling
-		const FTransform Transform = Socket->GetSocketLocalTransform();
-		Node.Translation = FGLTFConverterUtility::ConvertPosition(Transform.GetTranslation(), Builder.ExportOptions->ExportScale);
-		Node.Rotation = FGLTFConverterUtility::ConvertRotation(Transform.GetRotation());
-		Node.Scale = FGLTFConverterUtility::ConvertScale(Transform.GetScale3D());
+		Node.Translation = FGLTFConverterUtility::ConvertPosition(Socket->RelativeLocation, Builder.ExportOptions->ExportScale);
+		Node.Rotation = FGLTFConverterUtility::ConvertRotation(Socket->RelativeRotation.Quaternion());
+		Node.Scale = FGLTFConverterUtility::ConvertScale(Socket->RelativeScale);
 
 		const int32 ParentBone = SkeletalMesh->RefSkeleton.FindBoneIndex(Socket->BoneName);
 		const FGLTFJsonNodeIndex ParentNode = ParentBone != INDEX_NONE ? Builder.GetOrAddNode(RootNode, SkeletalMesh, ParentBone) : RootNode;
diff --git a/Engine/Plugins/Enterprise/GLTFExporter/MigratedFromGit/Source/GLTFExporter/Private/Converters/GLTFNodeConverters.h b/Engine/Plugins/Enterprise/GLTFExporter/MigratedFromGit/Source/GLTFExporter/Private/Converters/GLTFNodeConverters.h
--- a/Engine/Plugins/Enterprise/GLTFExporter/MigratedFromGit/Source/GLTFExporter/Private/Converters/GLTFNodeConverters.h
+++ b/Engine/Plugins/Enterprise/GLTFExporter/MigratedFromGit/Source/GLTFExporter/Private/Converters/GLTFNodeConverters.h
@@ -34,6 +34,13 @@ class FGLTFComponentSocketConverter final : public TGLTFNodeConverter<const USce
 	virtual FGLTFJsonNodeIndex Convert(const USceneComponent* SceneComponent, FName SocketName) override;
 };
 
+class FGLTFStaticSocketConverter final : public TGLTFNodeConverter<FGLTFJsonNodeIndex, const UStaticMesh*, FName>
+{
+	using TGLTFNodeConverter::TGLTFNodeConverter;
+
+	virtual FGLTFJsonNodeIndex Convert(FGLTFJsonNodeIndex RootNode, const UStaticMesh* StaticMesh, FName SocketName) override;
+};
+
 class FGLTFSkeletalSocketConverter final : public TGLTFNodeConverter<FGLTFJsonNodeIndex, const USkeletalMesh*, FName>
 {
 	using TGLTFNodeConverter::TGLTFNodeConverter;
