[GLTFExporter] refactor: use unique ptrs for top-level json objects to avoid reference invalidation

commit 89893839d0fdbc53769fb819310beff354a96ec2
Date: Sat, 3 Oct 2020 20:00:16 +0200

    refactor: use unique ptrs for top-level json objects to avoid reference invalidation
    (cherry picked from commit 186ffce01accb74bfbf931da8ace6c09598f74ca​)

#skipci

diff --git a/Engine/Plugins/Enterprise/GLTFExporter/MigratedFromGit/Source/GLTFExporter/Private/Builders/GLTFJsonBuilder.cpp b/Engine/Plugins/Enterprise/GLTFExporter/MigratedFromGit/Source/GLTFExporter/Private/Builders/GLTFJsonBuilder.cpp
--- a/Engine/Plugins/Enterprise/GLTFExporter/MigratedFromGit/Source/GLTFExporter/Private/Builders/GLTFJsonBuilder.cpp
+++ b/Engine/Plugins/Enterprise/GLTFExporter/MigratedFromGit/Source/GLTFExporter/Private/Builders/GLTFJsonBuilder.cpp
@@ -24,92 +24,92 @@ void FGLTFJsonBuilder::AddExtension(EGLTFJsonExtension Extension, bool bIsRequir
 
 FGLTFJsonAccessorIndex FGLTFJsonBuilder::AddAccessor(const FGLTFJsonAccessor& JsonAccessor)
 {
-	return FGLTFJsonAccessorIndex(JsonRoot.Accessors.Add(JsonAccessor));
+	return FGLTFJsonAccessorIndex(JsonRoot.Accessors.Add(MakeUnique<FGLTFJsonAccessor>(JsonAccessor)));
 }
 
 FGLTFJsonAnimationIndex FGLTFJsonBuilder::AddAnimation(const FGLTFJsonAnimation& JsonAnimation)
 {
-	return FGLTFJsonAnimationIndex(JsonRoot.Animations.Add(JsonAnimation));
+	return FGLTFJsonAnimationIndex(JsonRoot.Animations.Add(MakeUnique<FGLTFJsonAnimation>(JsonAnimation)));
 }
 
 FGLTFJsonBufferIndex FGLTFJsonBuilder::AddBuffer(const FGLTFJsonBuffer& JsonBuffer)
 {
-	return FGLTFJsonBufferIndex(JsonRoot.Buffers.Add(JsonBuffer));
+	return FGLTFJsonBufferIndex(JsonRoot.Buffers.Add(MakeUnique<FGLTFJsonBuffer>(JsonBuffer)));
 }
 
 FGLTFJsonBufferViewIndex FGLTFJsonBuilder::AddBufferView(const FGLTFJsonBufferView& JsonBufferView)
 {
-	return FGLTFJsonBufferViewIndex(JsonRoot.BufferViews.Add(JsonBufferView));
+	return FGLTFJsonBufferViewIndex(JsonRoot.BufferViews.Add(MakeUnique<FGLTFJsonBufferView>(JsonBufferView)));
 }
 
 FGLTFJsonCameraIndex FGLTFJsonBuilder::AddCamera(const FGLTFJsonCamera& JsonCamera)
 {
-	return FGLTFJsonCameraIndex(JsonRoot.Cameras.Add(JsonCamera));
+	return FGLTFJsonCameraIndex(JsonRoot.Cameras.Add(MakeUnique<FGLTFJsonCamera>(JsonCamera)));
 }
 
 FGLTFJsonImageIndex FGLTFJsonBuilder::AddImage(const FGLTFJsonImage& JsonImage)
 {
-	return FGLTFJsonImageIndex(JsonRoot.Images.Add(JsonImage));
+	return FGLTFJsonImageIndex(JsonRoot.Images.Add(MakeUnique<FGLTFJsonImage>(JsonImage)));
 }
 
 FGLTFJsonMaterialIndex FGLTFJsonBuilder::AddMaterial(const FGLTFJsonMaterial& JsonMaterial)
 {
-	return FGLTFJsonMaterialIndex(JsonRoot.Materials.Add(JsonMaterial));
+	return FGLTFJsonMaterialIndex(JsonRoot.Materials.Add(MakeUnique<FGLTFJsonMaterial>(JsonMaterial)));
 }
 
 FGLTFJsonMeshIndex FGLTFJsonBuilder::AddMesh(const FGLTFJsonMesh& JsonMesh)
 {
-	return FGLTFJsonMeshIndex(JsonRoot.Meshes.Add(JsonMesh));
+	return FGLTFJsonMeshIndex(JsonRoot.Meshes.Add(MakeUnique<FGLTFJsonMesh>(JsonMesh)));
 }
 
 FGLTFJsonNodeIndex FGLTFJsonBuilder::AddNode(const FGLTFJsonNode& JsonNode)
 {
-	return FGLTFJsonNodeIndex(JsonRoot.Nodes.Add(JsonNode));
+	return FGLTFJsonNodeIndex(JsonRoot.Nodes.Add(MakeUnique<FGLTFJsonNode>(JsonNode)));
 }
 
 FGLTFJsonSamplerIndex FGLTFJsonBuilder::AddSampler(const FGLTFJsonSampler& JsonSampler)
 {
-	return FGLTFJsonSamplerIndex(JsonRoot.Samplers.Add(JsonSampler));
+	return FGLTFJsonSamplerIndex(JsonRoot.Samplers.Add(MakeUnique<FGLTFJsonSampler>(JsonSampler)));
 }
 
 FGLTFJsonSceneIndex FGLTFJsonBuilder::AddScene(const FGLTFJsonScene& JsonScene)
 {
-	return FGLTFJsonSceneIndex(JsonRoot.Scenes.Add(JsonScene));
+	return FGLTFJsonSceneIndex(JsonRoot.Scenes.Add(MakeUnique<FGLTFJsonScene>(JsonScene)));
 }
 
 FGLTFJsonSkinIndex FGLTFJsonBuilder::AddSkin(const FGLTFJsonSkin& JsonSkin)
 {
-	return FGLTFJsonSkinIndex(JsonRoot.Skins.Add(JsonSkin));
+	return FGLTFJsonSkinIndex(JsonRoot.Skins.Add(MakeUnique<FGLTFJsonSkin>(JsonSkin)));
 }
 
 FGLTFJsonTextureIndex FGLTFJsonBuilder::AddTexture(const FGLTFJsonTexture& JsonTexture)
 {
-	return FGLTFJsonTextureIndex(JsonRoot.Textures.Add(JsonTexture));
+	return FGLTFJsonTextureIndex(JsonRoot.Textures.Add(MakeUnique<FGLTFJsonTexture>(JsonTexture)));
 }
 
 FGLTFJsonBackdropIndex FGLTFJsonBuilder::AddBackdrop(const FGLTFJsonBackdrop& JsonBackdrop)
 {
-	return FGLTFJsonBackdropIndex(JsonRoot.Backdrops.Add(JsonBackdrop));
+	return FGLTFJsonBackdropIndex(JsonRoot.Backdrops.Add(MakeUnique<FGLTFJsonBackdrop>(JsonBackdrop)));
 }
 
 FGLTFJsonVariationIndex FGLTFJsonBuilder::AddVariation(const FGLTFJsonVariation& JsonVariation)
 {
-	return FGLTFJsonVariationIndex(JsonRoot.Variations.Add(JsonVariation));
+	return FGLTFJsonVariationIndex(JsonRoot.Variations.Add(MakeUnique<FGLTFJsonVariation>(JsonVariation)));
 }
 
 FGLTFJsonLightMapIndex FGLTFJsonBuilder::AddLightMap(const FGLTFJsonLightMap& JsonLightMap)
 {
-	return FGLTFJsonLightMapIndex(JsonRoot.LightMaps.Add(JsonLightMap));
+	return FGLTFJsonLightMapIndex(JsonRoot.LightMaps.Add(MakeUnique<FGLTFJsonLightMap>(JsonLightMap)));
 }
 
 FGLTFJsonLightIndex FGLTFJsonBuilder::AddLight(const FGLTFJsonLight& JsonLight)
 {
-	return FGLTFJsonLightIndex(JsonRoot.Lights.Add(JsonLight));
+	return FGLTFJsonLightIndex(JsonRoot.Lights.Add(MakeUnique<FGLTFJsonLight>(JsonLight)));
 }
 
 FGLTFJsonHotspotIndex FGLTFJsonBuilder::AddHotspot(const FGLTFJsonHotspot& JsonHotspot)
 {
-	return FGLTFJsonHotspotIndex(JsonRoot.Hotspots.Add(JsonHotspot));
+	return FGLTFJsonHotspotIndex(JsonRoot.Hotspots.Add(MakeUnique<FGLTFJsonHotspot>(JsonHotspot)));
 }
 
 FGLTFJsonNodeIndex FGLTFJsonBuilder::AddChildNode(FGLTFJsonNodeIndex ParentIndex, const FGLTFJsonNode& JsonNode)
@@ -138,92 +138,92 @@ FGLTFJsonNodeIndex FGLTFJsonBuilder::AddChildComponentNode(FGLTFJsonNodeIndex Pa
 
 FGLTFJsonAccessor& FGLTFJsonBuilder::GetAccessor(FGLTFJsonAccessorIndex AccessorIndex)
 {
-	return JsonRoot.Accessors[AccessorIndex];
+	return *JsonRoot.Accessors[AccessorIndex];
 }
 
 FGLTFJsonAnimation& FGLTFJsonBuilder::GetAnimation(FGLTFJsonAnimationIndex AnimationIndex)
 {
-	return JsonRoot.Animations[AnimationIndex];
+	return *JsonRoot.Animations[AnimationIndex];
 }
 
 FGLTFJsonBuffer& FGLTFJsonBuilder::GetBuffer(FGLTFJsonBufferIndex BufferIndex)
 {
-	return JsonRoot.Buffers[BufferIndex];
+	return *JsonRoot.Buffers[BufferIndex];
 }
 
 FGLTFJsonBufferView& FGLTFJsonBuilder::GetBufferView(FGLTFJsonBufferViewIndex BufferViewIndex)
 {
-	return JsonRoot.BufferViews[BufferViewIndex];
+	return *JsonRoot.BufferViews[BufferViewIndex];
 }
 
 FGLTFJsonCamera& FGLTFJsonBuilder::GetCamera(FGLTFJsonCameraIndex CameraIndex)
 {
-	return JsonRoot.Cameras[CameraIndex];
+	return *JsonRoot.Cameras[CameraIndex];
 }
 
 FGLTFJsonImage& FGLTFJsonBuilder::GetImage(FGLTFJsonImageIndex ImageIndex)
 {
-	return JsonRoot.Images[ImageIndex];
+	return *JsonRoot.Images[ImageIndex];
 }
 
 FGLTFJsonMaterial& FGLTFJsonBuilder::GetMaterial(FGLTFJsonMaterialIndex MaterialIndex)
 {
-	return JsonRoot.Materials[MaterialIndex];
+	return *JsonRoot.Materials[MaterialIndex];
 }
 
 FGLTFJsonMesh& FGLTFJsonBuilder::GetMesh(FGLTFJsonMeshIndex MeshIndex)
 {
-	return JsonRoot.Meshes[MeshIndex];
+	return *JsonRoot.Meshes[MeshIndex];
 }
 
 FGLTFJsonNode& FGLTFJsonBuilder::GetNode(FGLTFJsonNodeIndex NodeIndex)
 {
-	return JsonRoot.Nodes[NodeIndex];
+	return *JsonRoot.Nodes[NodeIndex];
 }
 
 FGLTFJsonSampler& FGLTFJsonBuilder::GetSampler(FGLTFJsonSamplerIndex SamplerIndex)
 {
-	return JsonRoot.Samplers[SamplerIndex];
+	return *JsonRoot.Samplers[SamplerIndex];
 }
 
 FGLTFJsonScene& FGLTFJsonBuilder::GetScene(FGLTFJsonSceneIndex SceneIndex)
 {
-	return JsonRoot.Scenes[SceneIndex];
+	return *JsonRoot.Scenes[SceneIndex];
 }
 
 FGLTFJsonSkin& FGLTFJsonBuilder::GetSkin(FGLTFJsonSkinIndex SkinIndex)
 {
-	return JsonRoot.Skins[SkinIndex];
+	return *JsonRoot.Skins[SkinIndex];
 }
 
 FGLTFJsonTexture& FGLTFJsonBuilder::GetTexture(FGLTFJsonTextureIndex TextureIndex)
 {
-	return JsonRoot.Textures[TextureIndex];
+	return *JsonRoot.Textures[TextureIndex];
 }
 
 FGLTFJsonBackdrop& FGLTFJsonBuilder::GetBackdrop(FGLTFJsonBackdropIndex BackdropIndex)
 {
-	return JsonRoot.Backdrops[BackdropIndex];
+	return *JsonRoot.Backdrops[BackdropIndex];
 }
 
 FGLTFJsonVariation& FGLTFJsonBuilder::GetVariation(FGLTFJsonVariationIndex VariationIndex)
 {
-	return JsonRoot.Variations[VariationIndex];
+	return *JsonRoot.Variations[VariationIndex];
 }
 
 FGLTFJsonLightMap& FGLTFJsonBuilder::GetLightMap(FGLTFJsonLightMapIndex LightMapIndex)
 {
-	return JsonRoot.LightMaps[LightMapIndex];
+	return *JsonRoot.LightMaps[LightMapIndex];
 }
 
 FGLTFJsonLight& FGLTFJsonBuilder::GetLight(FGLTFJsonLightIndex LightIndex)
 {
-	return JsonRoot.Lights[LightIndex];
+	return *JsonRoot.Lights[LightIndex];
 }
 
 FGLTFJsonHotspot& FGLTFJsonBuilder::GetHotspot(FGLTFJsonHotspotIndex HotspotIndex)
 {
-	return JsonRoot.Hotspots[HotspotIndex];
+	return *JsonRoot.Hotspots[HotspotIndex];
 }
 
 FGLTFJsonNodeIndex FGLTFJsonBuilder::GetComponentNodeIndex(FGLTFJsonNodeIndex NodeIndex)
diff --git a/Engine/Plugins/Enterprise/GLTFExporter/MigratedFromGit/Source/GLTFExporter/Private/Json/GLTFJsonRoot.h b/Engine/Plugins/Enterprise/GLTFExporter/MigratedFromGit/Source/GLTFExporter/Private/Json/GLTFJsonRoot.h
--- a/Engine/Plugins/Enterprise/GLTFExporter/MigratedFromGit/Source/GLTFExporter/Private/Json/GLTFJsonRoot.h
+++ b/Engine/Plugins/Enterprise/GLTFExporter/MigratedFromGit/Source/GLTFExporter/Private/Json/GLTFJsonRoot.h
@@ -62,24 +62,24 @@ struct FGLTFJsonRoot
 
 	FGLTFJsonSceneIndex DefaultScene;
 
-	TArray<FGLTFJsonAccessor>   Accessors;
-	TArray<FGLTFJsonAnimation>  Animations;
-	TArray<FGLTFJsonBuffer>     Buffers;
-	TArray<FGLTFJsonBufferView> BufferViews;
-	TArray<FGLTFJsonCamera>     Cameras;
-	TArray<FGLTFJsonMaterial>   Materials;
-	TArray<FGLTFJsonMesh>       Meshes;
-	TArray<FGLTFJsonNode>       Nodes;
-	TArray<FGLTFJsonImage>      Images;
-	TArray<FGLTFJsonSampler>    Samplers;
-	TArray<FGLTFJsonScene>      Scenes;
-	TArray<FGLTFJsonSkin>       Skins;
-	TArray<FGLTFJsonTexture>    Textures;
-	TArray<FGLTFJsonBackdrop>   Backdrops;
-	TArray<FGLTFJsonLightMap>   LightMaps;
-	TArray<FGLTFJsonLight>      Lights;
-	TArray<FGLTFJsonVariation>        Variations;
-	TArray<FGLTFJsonHotspot>    Hotspots;
+	TArray<TUniquePtr<FGLTFJsonAccessor>>   Accessors;
+	TArray<TUniquePtr<FGLTFJsonAnimation>>  Animations;
+	TArray<TUniquePtr<FGLTFJsonBuffer>>     Buffers;
+	TArray<TUniquePtr<FGLTFJsonBufferView>> BufferViews;
+	TArray<TUniquePtr<FGLTFJsonCamera>>     Cameras;
+	TArray<TUniquePtr<FGLTFJsonMaterial>>   Materials;
+	TArray<TUniquePtr<FGLTFJsonMesh>>       Meshes;
+	TArray<TUniquePtr<FGLTFJsonNode>>       Nodes;
+	TArray<TUniquePtr<FGLTFJsonImage>>      Images;
+	TArray<TUniquePtr<FGLTFJsonSampler>>    Samplers;
+	TArray<TUniquePtr<FGLTFJsonScene>>      Scenes;
+	TArray<TUniquePtr<FGLTFJsonSkin>>       Skins;
+	TArray<TUniquePtr<FGLTFJsonTexture>>    Textures;
+	TArray<TUniquePtr<FGLTFJsonBackdrop>>   Backdrops;
+	TArray<TUniquePtr<FGLTFJsonLightMap>>   LightMaps;
+	TArray<TUniquePtr<FGLTFJsonLight>>      Lights;
+	TArray<TUniquePtr<FGLTFJsonVariation>>  Variations;
+	TArray<TUniquePtr<FGLTFJsonHotspot>>    Hotspots;
 
 	template <class CharType = TCHAR, class PrintPolicy = TPrettyJsonPrintPolicy<CharType>>
 	void WriteObject(TJsonWriter<CharType, PrintPolicy>& JsonWriter) const
@@ -96,19 +96,19 @@ struct FGLTFJsonRoot
 			JsonWriter.WriteValue(TEXT("scene"), DefaultScene);
 		}
 
-		FGLTFJsonUtility::WriteObjectArray(JsonWriter, TEXT("accessors"), Accessors, AllExtensions);
-		FGLTFJsonUtility::WriteObjectArray(JsonWriter, TEXT("animations"), Animations, AllExtensions);
-		FGLTFJsonUtility::WriteObjectArray(JsonWriter, TEXT("buffers"), Buffers, AllExtensions);
-		FGLTFJsonUtility::WriteObjectArray(JsonWriter, TEXT("bufferViews"), BufferViews, AllExtensions);
-		FGLTFJsonUtility::WriteObjectArray(JsonWriter, TEXT("cameras"), Cameras, AllExtensions);
-		FGLTFJsonUtility::WriteObjectArray(JsonWriter, TEXT("images"), Images, AllExtensions);
-		FGLTFJsonUtility::WriteObjectArray(JsonWriter, TEXT("materials"), Materials, AllExtensions);
-		FGLTFJsonUtility::WriteObjectArray(JsonWriter, TEXT("meshes"), Meshes, AllExtensions);
-		FGLTFJsonUtility::WriteObjectArray(JsonWriter, TEXT("nodes"), Nodes, AllExtensions);
-		FGLTFJsonUtility::WriteObjectArray(JsonWriter, TEXT("samplers"), Samplers, AllExtensions);
-		FGLTFJsonUtility::WriteObjectArray(JsonWriter, TEXT("scenes"), Scenes, AllExtensions);
-		FGLTFJsonUtility::WriteObjectArray(JsonWriter, TEXT("skins"), Skins, AllExtensions);
-		FGLTFJsonUtility::WriteObjectArray(JsonWriter, TEXT("textures"), Textures, AllExtensions);
+		FGLTFJsonUtility::WriteObjectPtrArray(JsonWriter, TEXT("accessors"), Accessors, AllExtensions);
+		FGLTFJsonUtility::WriteObjectPtrArray(JsonWriter, TEXT("animations"), Animations, AllExtensions);
+		FGLTFJsonUtility::WriteObjectPtrArray(JsonWriter, TEXT("buffers"), Buffers, AllExtensions);
+		FGLTFJsonUtility::WriteObjectPtrArray(JsonWriter, TEXT("bufferViews"), BufferViews, AllExtensions);
+		FGLTFJsonUtility::WriteObjectPtrArray(JsonWriter, TEXT("cameras"), Cameras, AllExtensions);
+		FGLTFJsonUtility::WriteObjectPtrArray(JsonWriter, TEXT("images"), Images, AllExtensions);
+		FGLTFJsonUtility::WriteObjectPtrArray(JsonWriter, TEXT("materials"), Materials, AllExtensions);
+		FGLTFJsonUtility::WriteObjectPtrArray(JsonWriter, TEXT("meshes"), Meshes, AllExtensions);
+		FGLTFJsonUtility::WriteObjectPtrArray(JsonWriter, TEXT("nodes"), Nodes, AllExtensions);
+		FGLTFJsonUtility::WriteObjectPtrArray(JsonWriter, TEXT("samplers"), Samplers, AllExtensions);
+		FGLTFJsonUtility::WriteObjectPtrArray(JsonWriter, TEXT("scenes"), Scenes, AllExtensions);
+		FGLTFJsonUtility::WriteObjectPtrArray(JsonWriter, TEXT("skins"), Skins, AllExtensions);
+		FGLTFJsonUtility::WriteObjectPtrArray(JsonWriter, TEXT("textures"), Textures, AllExtensions);
 
 		if (Backdrops.Num() > 0 || Variations.Num() > 0 || LightMaps.Num() > 0 || Lights.Num() > 0 || Hotspots.Num() > 0 /* TODO: add more extension support */)
 		{
@@ -120,7 +120,7 @@ struct FGLTFJsonRoot
 				AllExtensions.Used.Add(Extension);
 
 				JsonWriter.WriteObjectStart(FGLTFJsonUtility::ToString(Extension));
-				FGLTFJsonUtility::WriteObjectArray(JsonWriter, TEXT("backdrops"), Backdrops, AllExtensions);
+				FGLTFJsonUtility::WriteObjectPtrArray(JsonWriter, TEXT("backdrops"), Backdrops, AllExtensions);
 				JsonWriter.WriteObjectEnd();
 			}
 
@@ -130,7 +130,7 @@ struct FGLTFJsonRoot
 				AllExtensions.Used.Add(Extension);
 
 				JsonWriter.WriteObjectStart(FGLTFJsonUtility::ToString(Extension));
-				FGLTFJsonUtility::WriteObjectArray(JsonWriter, TEXT("levelVariantSets"), Variations, AllExtensions);
+				FGLTFJsonUtility::WriteObjectPtrArray(JsonWriter, TEXT("levelVariantSets"), Variations, AllExtensions);
 				JsonWriter.WriteObjectEnd();
 			}
 
@@ -140,7 +140,7 @@ struct FGLTFJsonRoot
 				AllExtensions.Used.Add(Extension);
 
 				JsonWriter.WriteObjectStart(FGLTFJsonUtility::ToString(Extension));
-				FGLTFJsonUtility::WriteObjectArray(JsonWriter, TEXT("lightmaps"), LightMaps, AllExtensions);
+				FGLTFJsonUtility::WriteObjectPtrArray(JsonWriter, TEXT("lightmaps"), LightMaps, AllExtensions);
 				JsonWriter.WriteObjectEnd();
 			}
 
@@ -150,7 +150,7 @@ struct FGLTFJsonRoot
 				AllExtensions.Used.Add(Extension);
 
 				JsonWriter.WriteObjectStart(FGLTFJsonUtility::ToString(Extension));
-				FGLTFJsonUtility::WriteObjectArray(JsonWriter, TEXT("lights"), Lights, AllExtensions);
+				FGLTFJsonUtility::WriteObjectPtrArray(JsonWriter, TEXT("lights"), Lights, AllExtensions);
 				JsonWriter.WriteObjectEnd();
 			}
 
@@ -160,7 +160,7 @@ struct FGLTFJsonRoot
 				AllExtensions.Used.Add(Extension);
 
 				JsonWriter.WriteObjectStart(FGLTFJsonUtility::ToString(Extension));
-				FGLTFJsonUtility::WriteObjectArray(JsonWriter, TEXT("interactions"), Hotspots, AllExtensions);
+				FGLTFJsonUtility::WriteObjectPtrArray(JsonWriter, TEXT("interactions"), Hotspots, AllExtensions);
 				JsonWriter.WriteObjectEnd();
 			}
 
diff --git a/Engine/Plugins/Enterprise/GLTFExporter/MigratedFromGit/Source/GLTFExporter/Private/Json/GLTFJsonUtility.h b/Engine/Plugins/Enterprise/GLTFExporter/MigratedFromGit/Source/GLTFExporter/Private/Json/GLTFJsonUtility.h
--- a/Engine/Plugins/Enterprise/GLTFExporter/MigratedFromGit/Source/GLTFExporter/Private/Json/GLTFJsonUtility.h
+++ b/Engine/Plugins/Enterprise/GLTFExporter/MigratedFromGit/Source/GLTFExporter/Private/Json/GLTFJsonUtility.h
@@ -167,6 +167,20 @@ struct FGLTFJsonUtility
 		}
 	}
 
+	template <class WriterType, class ContainerType>
+    static void WriteObjectPtrArray(WriterType& JsonWriter, const FString& Identifier, const ContainerType& Container, FGLTFJsonExtensions& Extensions, bool bWriteIfEmpty = false)
+	{
+		if (Container.Num() > 0 || bWriteIfEmpty)
+		{
+			JsonWriter.WriteArrayStart(Identifier);
+			for (const auto& Element : Container)
+			{
+				Element->WriteObject(JsonWriter, Extensions);
+			}
+			JsonWriter.WriteArrayEnd();
+		}
+	}
+
 	template <class WriterType, class ContainerType>
 	static void WriteStringArray(WriterType& JsonWriter, const FString& Identifier, const ContainerType& Container, bool bWriteIfEmpty = false)
 	{
