[GLTFExporter] refactor: rename FGLTFJsonHotspot.h to GLTFJsonHotspot.h

commit 7445e02dfdd79220cd122829fc0bcfe310b3ff36
Date: Sat, 3 Oct 2020 20:29:06 +0200

    refactor: rename FGLTFJsonHotspot.h to GLTFJsonHotspot.h
    (cherry picked from commit 9a291d6e9af810ee329f67a1a8c3075339276e94​)

#skipci

diff --git a/Engine/Plugins/Enterprise/GLTFExporter/MigratedFromGit/Source/GLTFExporter/Private/Converters/GLTFHotspotConverters.cpp b/Engine/Plugins/Enterprise/GLTFExporter/MigratedFromGit/Source/GLTFExporter/Private/Converters/GLTFHotspotConverters.cpp
--- a/Engine/Plugins/Enterprise/GLTFExporter/MigratedFromGit/Source/GLTFExporter/Private/Converters/GLTFHotspotConverters.cpp
+++ b/Engine/Plugins/Enterprise/GLTFExporter/MigratedFromGit/Source/GLTFExporter/Private/Converters/GLTFHotspotConverters.cpp
@@ -1,7 +1,6 @@
 // Copyright Epic Games, Inc. All Rights Reserved.
 
-#include "GLTFHotspotConverters.h"
-#include "Json/FGLTFJsonHotspot.h"
+#include "Converters/GLTFHotspotConverters.h"
 #include "Builders/GLTFContainerBuilder.h"
 #include "Animation/SkeletalMeshActor.h"
 
diff --git a/Engine/Plugins/Enterprise/GLTFExporter/MigratedFromGit/Source/GLTFExporter/Private/Json/FGLTFJsonHotspot.h b/Engine/Plugins/Enterprise/GLTFExporter/MigratedFromGit/Source/GLTFExporter/Private/Json/FGLTFJsonHotspot.h
deleted file mode 100644
diff --git a/Engine/Plugins/Enterprise/GLTFExporter/MigratedFromGit/Source/GLTFExporter/Private/Json/GLTFJsonHotspot.h b/Engine/Plugins/Enterprise/GLTFExporter/MigratedFromGit/Source/GLTFExporter/Private/Json/GLTFJsonHotspot.h
new file mode 100644
--- /dev/null
+++ b/Engine/Plugins/Enterprise/GLTFExporter/MigratedFromGit/Source/GLTFExporter/Private/Json/GLTFJsonHotspot.h
@@ -0,0 +1,49 @@
+// Copyright Epic Games, Inc. All Rights Reserved.
+
+#pragma once
+
+#include "Json/GLTFJsonIndex.h"
+#include "Json/GLTFJsonExtensions.h"
+#include "Serialization/JsonSerializer.h"
+
+struct FGLTFJsonHotspot
+{
+	FString Name;
+
+	FGLTFJsonAnimationIndex Animation;
+	FGLTFJsonTextureIndex   Image;
+	FGLTFJsonTextureIndex   HoveredImage;
+	FGLTFJsonTextureIndex   ToggledImage;
+	FGLTFJsonTextureIndex   ToggledHoveredImage;
+
+	template <class CharType = TCHAR, class PrintPolicy = TPrettyJsonPrintPolicy<CharType>>
+	void WriteObject(TJsonWriter<CharType, PrintPolicy>& JsonWriter, FGLTFJsonExtensions& Extensions) const
+	{
+		JsonWriter.WriteObjectStart();
+
+		if (!Name.IsEmpty())
+		{
+			JsonWriter.WriteValue(TEXT("name"), Name);
+		}
+
+		JsonWriter.WriteValue(TEXT("animation"), Animation);
+		JsonWriter.WriteValue(TEXT("image"), Image);
+
+		if (HoveredImage != INDEX_NONE)
+		{
+			JsonWriter.WriteValue(TEXT("hoveredImage"), HoveredImage);
+		}
+
+		if (ToggledImage != INDEX_NONE)
+		{
+			JsonWriter.WriteValue(TEXT("toggledImage"), ToggledImage);
+		}
+
+		if (ToggledHoveredImage != INDEX_NONE)
+		{
+			JsonWriter.WriteValue(TEXT("toggledHoveredImage"), ToggledHoveredImage);
+		}
+
+		JsonWriter.WriteObjectEnd();
+	}
+};
diff --git a/Engine/Plugins/Enterprise/GLTFExporter/MigratedFromGit/Source/GLTFExporter/Private/Json/GLTFJsonRoot.h b/Engine/Plugins/Enterprise/GLTFExporter/MigratedFromGit/Source/GLTFExporter/Private/Json/GLTFJsonRoot.h
--- a/Engine/Plugins/Enterprise/GLTFExporter/MigratedFromGit/Source/GLTFExporter/Private/Json/GLTFJsonRoot.h
+++ b/Engine/Plugins/Enterprise/GLTFExporter/MigratedFromGit/Source/GLTFExporter/Private/Json/GLTFJsonRoot.h
@@ -21,7 +21,7 @@
 #include "Json/GLTFJsonVariation.h"
 #include "Json/GLTFJsonLightMap.h"
 #include "Json/GLTFJsonLight.h"
-#include "Json/FGLTFJsonHotspot.h"
+#include "Json/GLTFJsonHotspot.h"
 #include "Policies/CondensedJsonPrintPolicy.h"
 
 
