[GLTFExporter] feat: add GLTFTask and GLTFTaskBuilder for general task management

commit 3ac447690b979adb1c78f25fd81ab3cd15dafc7c
Date: Sun, 4 Oct 2020 02:18:34 +0200

    feat: add GLTFTask and GLTFTaskBuilder for general task management
    (cherry picked from commit 77be6eabcd50416f03c4dd6f4a79fcae47b07a29​)

#skipci

diff --git a/Engine/Plugins/Enterprise/GLTFExporter/MigratedFromGit/Source/GLTFExporter/Private/Builders/GLTFContainerBuilder.cpp b/Engine/Plugins/Enterprise/GLTFExporter/MigratedFromGit/Source/GLTFExporter/Private/Builders/GLTFContainerBuilder.cpp
--- a/Engine/Plugins/Enterprise/GLTFExporter/MigratedFromGit/Source/GLTFExporter/Private/Builders/GLTFContainerBuilder.cpp
+++ b/Engine/Plugins/Enterprise/GLTFExporter/MigratedFromGit/Source/GLTFExporter/Private/Builders/GLTFContainerBuilder.cpp
@@ -18,8 +18,10 @@ void FGLTFContainerBuilder::WriteGlb(FArchive& Archive) const
 	FGLTFContainerUtility::WriteGlb(Archive, JsonData, BufferData);
 }
 
-void FGLTFContainerBuilder::Write(FArchive& Archive) const
+void FGLTFContainerBuilder::Write(FArchive& Archive, FFeedbackContext* Context)
 {
+	CompleteAllTasks(Context);
+
 	if (bIsGlbFile)
 	{
 		WriteGlb(Archive);
diff --git a/Engine/Plugins/Enterprise/GLTFExporter/MigratedFromGit/Source/GLTFExporter/Private/Builders/GLTFContainerBuilder.h b/Engine/Plugins/Enterprise/GLTFExporter/MigratedFromGit/Source/GLTFExporter/Private/Builders/GLTFContainerBuilder.h
--- a/Engine/Plugins/Enterprise/GLTFExporter/MigratedFromGit/Source/GLTFExporter/Private/Builders/GLTFContainerBuilder.h
+++ b/Engine/Plugins/Enterprise/GLTFExporter/MigratedFromGit/Source/GLTFExporter/Private/Builders/GLTFContainerBuilder.h
@@ -10,7 +10,7 @@ public:
 
 	FGLTFContainerBuilder(const FString& FilePath, const UGLTFExportOptions* ExportOptions, bool bSelectedActorsOnly);
 
-	void Write(FArchive& Archive) const;
+	void Write(FArchive& Archive, FFeedbackContext* Context = GWarn);
 
 private:
 
diff --git a/Engine/Plugins/Enterprise/GLTFExporter/MigratedFromGit/Source/GLTFExporter/Private/Builders/GLTFMessageBuilder.cpp b/Engine/Plugins/Enterprise/GLTFExporter/MigratedFromGit/Source/GLTFExporter/Private/Builders/GLTFMessageBuilder.cpp
--- a/Engine/Plugins/Enterprise/GLTFExporter/MigratedFromGit/Source/GLTFExporter/Private/Builders/GLTFMessageBuilder.cpp
+++ b/Engine/Plugins/Enterprise/GLTFExporter/MigratedFromGit/Source/GLTFExporter/Private/Builders/GLTFMessageBuilder.cpp
@@ -7,7 +7,7 @@
 #include "IMessageLogListing.h"
 
 FGLTFMessageBuilder::FGLTFMessageBuilder(const FString& FilePath, const UGLTFExportOptions* ExportOptions)
-	: FGLTFBuilder(FilePath, ExportOptions)
+	: FGLTFTaskBuilder(FilePath, ExportOptions)
 {
 }
 
diff --git a/Engine/Plugins/Enterprise/GLTFExporter/MigratedFromGit/Source/GLTFExporter/Private/Builders/GLTFMessageBuilder.h b/Engine/Plugins/Enterprise/GLTFExporter/MigratedFromGit/Source/GLTFExporter/Private/Builders/GLTFMessageBuilder.h
--- a/Engine/Plugins/Enterprise/GLTFExporter/MigratedFromGit/Source/GLTFExporter/Private/Builders/GLTFMessageBuilder.h
+++ b/Engine/Plugins/Enterprise/GLTFExporter/MigratedFromGit/Source/GLTFExporter/Private/Builders/GLTFMessageBuilder.h
@@ -2,7 +2,7 @@
 
 #pragma once
 
-#include "Builders/GLTFBuilder.h"
+#include "Builders/GLTFTaskBuilder.h"
 
 enum class EGLTFMessageSeverity
 {
@@ -11,7 +11,7 @@ enum class EGLTFMessageSeverity
 	Error
 };
 
-class FGLTFMessageBuilder : public FGLTFBuilder
+class FGLTFMessageBuilder : public FGLTFTaskBuilder
 {
 protected:
 
diff --git a/Engine/Plugins/Enterprise/GLTFExporter/MigratedFromGit/Source/GLTFExporter/Private/Builders/GLTFTask.h b/Engine/Plugins/Enterprise/GLTFExporter/MigratedFromGit/Source/GLTFExporter/Private/Builders/GLTFTask.h
new file mode 100644
--- /dev/null
+++ b/Engine/Plugins/Enterprise/GLTFExporter/MigratedFromGit/Source/GLTFExporter/Private/Builders/GLTFTask.h
@@ -0,0 +1,49 @@
+// Copyright Epic Games, Inc. All Rights Reserved.
+
+#pragma once
+
+#include "CoreMinimal.h"
+
+enum class EGLTFTaskCategory
+{
+	Actor,
+    Mesh,
+    Animation,
+    Material,
+    Texture,
+    MAX
+};
+
+class FGLTFTask
+{
+public:
+
+	const EGLTFTaskCategory Category;
+
+	const FString Name;
+
+	FGLTFTask(EGLTFTaskCategory Category, const FString& Name)
+		: Category(Category)
+		, Name(Name)
+	{
+	}
+
+	virtual ~FGLTFTask() = default;
+
+	virtual void Run() = 0;
+};
+
+class FGLTFSleepTask : public FGLTFTask
+{
+public:
+
+	FGLTFSleepTask(EGLTFTaskCategory Category, const FString& Name)
+        : FGLTFTask(Category, Name)
+	{
+	}
+
+	virtual void Run() override
+	{
+		FPlatformProcess::Sleep(0.5);
+	}
+};
diff --git a/Engine/Plugins/Enterprise/GLTFExporter/MigratedFromGit/Source/GLTFExporter/Private/Builders/GLTFTaskBuilder.cpp b/Engine/Plugins/Enterprise/GLTFExporter/MigratedFromGit/Source/GLTFExporter/Private/Builders/GLTFTaskBuilder.cpp
new file mode 100644
--- /dev/null
+++ b/Engine/Plugins/Enterprise/GLTFExporter/MigratedFromGit/Source/GLTFExporter/Private/Builders/GLTFTaskBuilder.cpp
@@ -0,0 +1,61 @@
+// Copyright Epic Games, Inc. All Rights Reserved.
+
+#include "Builders/GLTFTaskBuilder.h"
+#include "Misc/FeedbackContext.h"
+#include "Misc/ScopedSlowTask.h"
+
+FGLTFTaskBuilder::FGLTFTaskBuilder(const FString& FilePath, const UGLTFExportOptions* ExportOptions)
+	: FGLTFBuilder(FilePath, ExportOptions)
+{
+}
+
+void FGLTFTaskBuilder::CompleteAllTasks(FFeedbackContext* Context)
+{
+	const int32 CategoryCount = static_cast<int32>(EGLTFTaskCategory::MAX);
+	for (int32 CategoryIndex = 0; CategoryIndex < CategoryCount; CategoryIndex++)
+	{
+		const EGLTFTaskCategory Category = static_cast<EGLTFTaskCategory>(CategoryIndex);
+
+		TArray<TUniquePtr<FGLTFTask>>* Tasks = CategorizedTasks.Find(Category);
+		if (Tasks == nullptr)
+		{
+			continue;
+		}
+
+		const FText FormatMessage = GetCategoryFormatMessage(Category);
+		FScopedSlowTask Progress(Tasks->Num(), FormatMessage, true, *Context);
+		Progress.MakeDialog();
+
+		for (TUniquePtr<FGLTFTask>& Task : *Tasks)
+		{
+			const FText Name = FText::FromString(Task->Name);
+			const FText Message = FText::Format(FormatMessage, Name);
+			Progress.EnterProgressFrame(1, Message);
+
+			Task->Run();
+		}
+	}
+
+	CategorizedTasks.Empty();
+}
+
+void FGLTFTaskBuilder::SetupTask(TUniquePtr<FGLTFTask> Task)
+{
+	const EGLTFTaskCategory Category = Task->Category;
+	CategorizedTasks.FindOrAdd(Category).Add(MoveTemp(Task));
+}
+
+FText FGLTFTaskBuilder::GetCategoryFormatMessage(EGLTFTaskCategory Category)
+{
+	switch (Category)
+	{
+		case EGLTFTaskCategory::Actor:     return NSLOCTEXT("GLTFExporter", "ActorTaskMessage", "Actors... {0}");
+		case EGLTFTaskCategory::Mesh:      return NSLOCTEXT("GLTFExporter", "ActorTaskMessage", "Meshes... {0}");
+		case EGLTFTaskCategory::Animation: return NSLOCTEXT("GLTFExporter", "ActorTaskMessage", "Animations... {0}");
+		case EGLTFTaskCategory::Material:  return NSLOCTEXT("GLTFExporter", "ActorTaskMessage", "Materials... {0}");
+		case EGLTFTaskCategory::Texture:   return NSLOCTEXT("GLTFExporter", "ActorTaskMessage", "Textures... {0}");
+		default:                           checkNoEntry();
+	}
+
+	return {};
+}
diff --git a/Engine/Plugins/Enterprise/GLTFExporter/MigratedFromGit/Source/GLTFExporter/Private/Builders/GLTFTaskBuilder.h b/Engine/Plugins/Enterprise/GLTFExporter/MigratedFromGit/Source/GLTFExporter/Private/Builders/GLTFTaskBuilder.h
new file mode 100644
--- /dev/null
+++ b/Engine/Plugins/Enterprise/GLTFExporter/MigratedFromGit/Source/GLTFExporter/Private/Builders/GLTFTaskBuilder.h
@@ -0,0 +1,32 @@
+// Copyright Epic Games, Inc. All Rights Reserved.
+
+#pragma once
+
+#include "Builders/GLTFBuilder.h"
+#include "Builders/GLTFTask.h"
+#include "GLTFExportOptions.h"
+
+class FGLTFTaskBuilder : public FGLTFBuilder
+{
+protected:
+
+	FGLTFTaskBuilder(const FString& FilePath, const UGLTFExportOptions* ExportOptions);
+
+public:
+
+	template <typename TaskType, typename = typename TEnableIf<TIsDerivedFrom<TaskType, FGLTFTask>::Value>::Type>
+	void SetupTask(TUniquePtr<TaskType> Task)
+	{
+		SetupTask(TUniquePtr<FGLTFTask>(Task.Release()));
+	}
+
+	void SetupTask(TUniquePtr<FGLTFTask> Task);
+
+	void CompleteAllTasks(FFeedbackContext* Context = GWarn);
+
+private:
+
+	static FText GetCategoryFormatMessage(EGLTFTaskCategory Category);
+
+	TMap<EGLTFTaskCategory, TArray<TUniquePtr<FGLTFTask>>> CategorizedTasks;
+};
diff --git a/Engine/Plugins/Enterprise/GLTFExporter/MigratedFromGit/Source/GLTFExporter/Private/Exporters/GLTFExporter.cpp b/Engine/Plugins/Enterprise/GLTFExporter/MigratedFromGit/Source/GLTFExporter/Private/Exporters/GLTFExporter.cpp
--- a/Engine/Plugins/Enterprise/GLTFExporter/MigratedFromGit/Source/GLTFExporter/Private/Exporters/GLTFExporter.cpp
+++ b/Engine/Plugins/Enterprise/GLTFExporter/MigratedFromGit/Source/GLTFExporter/Private/Exporters/GLTFExporter.cpp
@@ -38,7 +38,7 @@ bool UGLTFExporter::ExportBinary(UObject* Object, const TCHAR* Type, FArchive& A
 	const bool bSuccess = AddObject(Builder, Object);
 	if (bSuccess)
 	{
-		Builder.Write(Archive);
+		Builder.Write(Archive, Warn);
 	}
 	else
 	{
