[GLTFExporter] feat: add locking mechanism by task priority to avoid lingering tasks

commit ece148d48f6facbae91d74c110112b8842f57e5d
Date: Sun, 4 Oct 2020 12:41:20 +0200

    feat: add locking mechanism by task priority to avoid lingering tasks
    (cherry picked from commit 06a0e72538e2ef4d9007637938ade13a68072561​)

#skipci

diff --git a/Engine/Plugins/Enterprise/GLTFExporter/MigratedFromGit/Source/GLTFExporter/Private/Builders/GLTFTask.h b/Engine/Plugins/Enterprise/GLTFExporter/MigratedFromGit/Source/GLTFExporter/Private/Builders/GLTFTask.h
--- a/Engine/Plugins/Enterprise/GLTFExporter/MigratedFromGit/Source/GLTFExporter/Private/Builders/GLTFTask.h
+++ b/Engine/Plugins/Enterprise/GLTFExporter/MigratedFromGit/Source/GLTFExporter/Private/Builders/GLTFTask.h
@@ -4,7 +4,7 @@
 
 #include "CoreMinimal.h"
 
-enum class EGLTFTaskCategory
+enum class EGLTFTaskPriority : uint8
 {
 	Actor,
     Mesh,
@@ -18,32 +18,16 @@ class FGLTFTask
 {
 public:
 
-	const EGLTFTaskCategory Category;
+	const EGLTFTaskPriority Priority;
 
-	const FString Name;
-
-	FGLTFTask(EGLTFTaskCategory Category, const FString& Name)
-		: Category(Category)
-		, Name(Name)
+	FGLTFTask(EGLTFTaskPriority Priority)
+        : Priority(Priority)
 	{
 	}
 
 	virtual ~FGLTFTask() = default;
 
-	virtual void Run() = 0;
-};
-
-class FGLTFSleepTask : public FGLTFTask
-{
-public:
-
-	FGLTFSleepTask(EGLTFTaskCategory Category, const FString& Name)
-        : FGLTFTask(Category, Name)
-	{
-	}
+	virtual FString GetName() = 0;
 
-	virtual void Run() override
-	{
-		FPlatformProcess::Sleep(0.5);
-	}
+	virtual void Complete() = 0;
 };
diff --git a/Engine/Plugins/Enterprise/GLTFExporter/MigratedFromGit/Source/GLTFExporter/Private/Builders/GLTFTaskBuilder.cpp b/Engine/Plugins/Enterprise/GLTFExporter/MigratedFromGit/Source/GLTFExporter/Private/Builders/GLTFTaskBuilder.cpp
--- a/Engine/Plugins/Enterprise/GLTFExporter/MigratedFromGit/Source/GLTFExporter/Private/Builders/GLTFTaskBuilder.cpp
+++ b/Engine/Plugins/Enterprise/GLTFExporter/MigratedFromGit/Source/GLTFExporter/Private/Builders/GLTFTaskBuilder.cpp
@@ -6,54 +6,70 @@
 
 FGLTFTaskBuilder::FGLTFTaskBuilder(const FString& FilePath, const UGLTFExportOptions* ExportOptions)
 	: FGLTFBuilder(FilePath, ExportOptions)
+	, PriorityIndexLock(INDEX_NONE)
 {
 }
 
 void FGLTFTaskBuilder::CompleteAllTasks(FFeedbackContext* Context)
 {
-	const int32 CategoryCount = static_cast<int32>(EGLTFTaskCategory::MAX);
-	for (int32 CategoryIndex = 0; CategoryIndex < CategoryCount; CategoryIndex++)
+	const int32 PriorityCount = static_cast<int32>(EGLTFTaskPriority::MAX);
+	for (int32 PriorityIndex = 0; PriorityIndex < PriorityCount; PriorityIndex++)
 	{
-		const EGLTFTaskCategory Category = static_cast<EGLTFTaskCategory>(CategoryIndex);
+		PriorityIndexLock = PriorityIndex;
+		const EGLTFTaskPriority Priority = static_cast<EGLTFTaskPriority>(PriorityIndex);
 
-		TArray<TUniquePtr<FGLTFTask>>* Tasks = CategorizedTasks.Find(Category);
+		TArray<TUniquePtr<FGLTFTask>>* Tasks = TasksByPriority.Find(Priority);
 		if (Tasks == nullptr)
 		{
 			continue;
 		}
 
-		const FText FormatMessage = GetCategoryFormatMessage(Category);
-		FScopedSlowTask Progress(Tasks->Num(), FormatMessage, true, *Context);
+		const FText MessageFormat = GetPriorityMessageFormat(Priority);
+		FScopedSlowTask Progress(Tasks->Num(), FText::Format(MessageFormat, FText()), true, *Context);
 		Progress.MakeDialog();
 
 		for (TUniquePtr<FGLTFTask>& Task : *Tasks)
 		{
-			const FText Name = FText::FromString(Task->Name);
-			const FText Message = FText::Format(FormatMessage, Name);
+			const FText Name = FText::FromString(Task->GetName());
+			const FText Message = FText::Format(MessageFormat, Name);
 			Progress.EnterProgressFrame(1, Message);
 
-			Task->Run();
+			Task->Complete();
 		}
 	}
 
-	CategorizedTasks.Empty();
+	TasksByPriority.Empty();
+	PriorityIndexLock = INDEX_NONE;
 }
 
-void FGLTFTaskBuilder::SetupTask(TUniquePtr<FGLTFTask> Task)
+bool FGLTFTaskBuilder::SetupTask(TUniquePtr<FGLTFTask> Task)
 {
-	const EGLTFTaskCategory Category = Task->Category;
-	CategorizedTasks.FindOrAdd(Category).Add(MoveTemp(Task));
+	const EGLTFTaskPriority Priority = Task->Priority;
+	if (static_cast<int32>(Priority) < static_cast<int32>(EGLTFTaskPriority::MAX))
+	{
+		// TODO: report error
+		return false;
+	}
+
+	if (static_cast<int32>(Priority) <= PriorityIndexLock)
+	{
+		// TODO: report error
+		return false;
+	}
+
+	TasksByPriority.FindOrAdd(Priority).Add(MoveTemp(Task));
+	return true;
 }
 
-FText FGLTFTaskBuilder::GetCategoryFormatMessage(EGLTFTaskCategory Category)
+FText FGLTFTaskBuilder::GetPriorityMessageFormat(EGLTFTaskPriority Priority)
 {
-	switch (Category)
+	switch (Priority)
 	{
-		case EGLTFTaskCategory::Actor:     return NSLOCTEXT("GLTFExporter", "ActorTaskMessage", "Actors... {0}");
-		case EGLTFTaskCategory::Mesh:      return NSLOCTEXT("GLTFExporter", "ActorTaskMessage", "Meshes... {0}");
-		case EGLTFTaskCategory::Animation: return NSLOCTEXT("GLTFExporter", "ActorTaskMessage", "Animations... {0}");
-		case EGLTFTaskCategory::Material:  return NSLOCTEXT("GLTFExporter", "ActorTaskMessage", "Materials... {0}");
-		case EGLTFTaskCategory::Texture:   return NSLOCTEXT("GLTFExporter", "ActorTaskMessage", "Textures... {0}");
+		case EGLTFTaskPriority::Actor:     return NSLOCTEXT("GLTFExporter", "ActorTaskMessage", "Actors... {0}");
+		case EGLTFTaskPriority::Mesh:      return NSLOCTEXT("GLTFExporter", "MeshTaskMessage", "Meshes... {0}");
+		case EGLTFTaskPriority::Animation: return NSLOCTEXT("GLTFExporter", "AnimationTaskMessage", "Animations... {0}");
+		case EGLTFTaskPriority::Material:  return NSLOCTEXT("GLTFExporter", "MaterialTaskMessage", "Materials... {0}");
+		case EGLTFTaskPriority::Texture:   return NSLOCTEXT("GLTFExporter", "TextureTaskMessage", "Textures... {0}");
 		default:                           checkNoEntry();
 	}
 
diff --git a/Engine/Plugins/Enterprise/GLTFExporter/MigratedFromGit/Source/GLTFExporter/Private/Builders/GLTFTaskBuilder.h b/Engine/Plugins/Enterprise/GLTFExporter/MigratedFromGit/Source/GLTFExporter/Private/Builders/GLTFTaskBuilder.h
--- a/Engine/Plugins/Enterprise/GLTFExporter/MigratedFromGit/Source/GLTFExporter/Private/Builders/GLTFTaskBuilder.h
+++ b/Engine/Plugins/Enterprise/GLTFExporter/MigratedFromGit/Source/GLTFExporter/Private/Builders/GLTFTaskBuilder.h
@@ -4,7 +4,6 @@
 
 #include "Builders/GLTFBuilder.h"
 #include "Builders/GLTFTask.h"
-#include "GLTFExportOptions.h"
 
 class FGLTFTaskBuilder : public FGLTFBuilder
 {
@@ -15,18 +14,19 @@ protected:
 public:
 
 	template <typename TaskType, typename = typename TEnableIf<TIsDerivedFrom<TaskType, FGLTFTask>::Value>::Type>
-	void SetupTask(TUniquePtr<TaskType> Task)
+	bool SetupTask(TUniquePtr<TaskType> Task)
 	{
-		SetupTask(TUniquePtr<FGLTFTask>(Task.Release()));
+		return SetupTask(TUniquePtr<FGLTFTask>(Task.Release()));
 	}
 
-	void SetupTask(TUniquePtr<FGLTFTask> Task);
+	bool SetupTask(TUniquePtr<FGLTFTask> Task);
 
 	void CompleteAllTasks(FFeedbackContext* Context = GWarn);
 
 private:
 
-	static FText GetCategoryFormatMessage(EGLTFTaskCategory Category);
+	static FText GetPriorityMessageFormat(EGLTFTaskPriority Priority);
 
-	TMap<EGLTFTaskCategory, TArray<TUniquePtr<FGLTFTask>>> CategorizedTasks;
+	int32 PriorityIndexLock;
+	TMap<EGLTFTaskPriority, TArray<TUniquePtr<FGLTFTask>>> TasksByPriority;
 };
