[GLTFExporter] refactor: move GLTFTask and add helper function to GLTFTaskBuilder

commit 2ec23b49b821daf15d2fd8af088415d2246b01ed
Date: Sun, 4 Oct 2020 13:04:05 +0200

    refactor: move GLTFTask and add helper function to GLTFTaskBuilder
    (cherry picked from commit 4e6723f181799984626405b00dcb23e8ad9383cd​)

#skipci

diff --git a/Engine/Plugins/Enterprise/GLTFExporter/MigratedFromGit/Source/GLTFExporter/Private/Builders/GLTFTask.h b/Engine/Plugins/Enterprise/GLTFExporter/MigratedFromGit/Source/GLTFExporter/Private/Builders/GLTFTask.h
deleted file mode 100644
diff --git a/Engine/Plugins/Enterprise/GLTFExporter/MigratedFromGit/Source/GLTFExporter/Private/Builders/GLTFTaskBuilder.cpp b/Engine/Plugins/Enterprise/GLTFExporter/MigratedFromGit/Source/GLTFExporter/Private/Builders/GLTFTaskBuilder.cpp
--- a/Engine/Plugins/Enterprise/GLTFExporter/MigratedFromGit/Source/GLTFExporter/Private/Builders/GLTFTaskBuilder.cpp
+++ b/Engine/Plugins/Enterprise/GLTFExporter/MigratedFromGit/Source/GLTFExporter/Private/Builders/GLTFTaskBuilder.cpp
@@ -45,7 +45,7 @@ void FGLTFTaskBuilder::CompleteAllTasks(FFeedbackContext* Context)
 bool FGLTFTaskBuilder::SetupTask(TUniquePtr<FGLTFTask> Task)
 {
 	const EGLTFTaskPriority Priority = Task->Priority;
-	if (static_cast<int32>(Priority) < static_cast<int32>(EGLTFTaskPriority::MAX))
+	if (static_cast<int32>(Priority) >= static_cast<int32>(EGLTFTaskPriority::MAX))
 	{
 		// TODO: report error
 		return false;
diff --git a/Engine/Plugins/Enterprise/GLTFExporter/MigratedFromGit/Source/GLTFExporter/Private/Builders/GLTFTaskBuilder.h b/Engine/Plugins/Enterprise/GLTFExporter/MigratedFromGit/Source/GLTFExporter/Private/Builders/GLTFTaskBuilder.h
--- a/Engine/Plugins/Enterprise/GLTFExporter/MigratedFromGit/Source/GLTFExporter/Private/Builders/GLTFTaskBuilder.h
+++ b/Engine/Plugins/Enterprise/GLTFExporter/MigratedFromGit/Source/GLTFExporter/Private/Builders/GLTFTaskBuilder.h
@@ -3,7 +3,7 @@
 #pragma once
 
 #include "Builders/GLTFBuilder.h"
-#include "Builders/GLTFTask.h"
+#include "Tasks/GLTFTask.h"
 
 class FGLTFTaskBuilder : public FGLTFBuilder
 {
@@ -13,6 +13,12 @@ protected:
 
 public:
 
+	template <typename TaskType, typename... TaskArgTypes, typename = typename TEnableIf<TIsDerivedFrom<TaskType, FGLTFTask>::Value>::Type>
+    bool SetupTask(TaskArgTypes&&... Args)
+	{
+		return SetupTask(MakeUnique<TaskType>(Forward<TaskArgTypes>(Args)...));
+	}
+
 	template <typename TaskType, typename = typename TEnableIf<TIsDerivedFrom<TaskType, FGLTFTask>::Value>::Type>
 	bool SetupTask(TUniquePtr<TaskType> Task)
 	{
diff --git a/Engine/Plugins/Enterprise/GLTFExporter/MigratedFromGit/Source/GLTFExporter/Private/Tasks/GLTFTask.h b/Engine/Plugins/Enterprise/GLTFExporter/MigratedFromGit/Source/GLTFExporter/Private/Tasks/GLTFTask.h
new file mode 100644
--- /dev/null
+++ b/Engine/Plugins/Enterprise/GLTFExporter/MigratedFromGit/Source/GLTFExporter/Private/Tasks/GLTFTask.h
@@ -0,0 +1,33 @@
+// Copyright Epic Games, Inc. All Rights Reserved.
+
+#pragma once
+
+#include "CoreMinimal.h"
+
+enum class EGLTFTaskPriority : uint8
+{
+	Actor,
+    Mesh,
+    Animation,
+    Material,
+    Texture,
+    MAX
+};
+
+class FGLTFTask
+{
+public:
+
+	const EGLTFTaskPriority Priority;
+
+	FGLTFTask(EGLTFTaskPriority Priority)
+		: Priority(Priority)
+	{
+	}
+
+	virtual ~FGLTFTask() = default;
+
+	virtual FString GetName() = 0;
+
+	virtual void Complete() = 0;
+};
