[GLTFExporter] feat: add GLTFMaterialTask and change material converter to use it

commit 568fa49e60dfc564e72937cf8cb3519c5ec1eb3a
Date: Sun, 4 Oct 2020 13:05:03 +0200

    feat: add GLTFMaterialTask and change material converter to use it
    (cherry picked from commit ef6e5db345a5a483828aa95e31a76ede02a76233​)

#skipci

diff --git a/Engine/Plugins/Enterprise/GLTFExporter/MigratedFromGit/Source/GLTFExporter/Private/Converters/GLTFMaterialConverters.cpp b/Engine/Plugins/Enterprise/GLTFExporter/MigratedFromGit/Source/GLTFExporter/Private/Converters/GLTFMaterialConverters.cpp
--- a/Engine/Plugins/Enterprise/GLTFExporter/MigratedFromGit/Source/GLTFExporter/Private/Converters/GLTFMaterialConverters.cpp
+++ b/Engine/Plugins/Enterprise/GLTFExporter/MigratedFromGit/Source/GLTFExporter/Private/Converters/GLTFMaterialConverters.cpp
@@ -1,1334 +1,12 @@
 // Copyright Epic Games, Inc. All Rights Reserved.
 
 #include "Converters/GLTFMaterialConverters.h"
-#include "Converters/GLTFConverterUtility.h"
-#include "Converters/GLTFNameUtility.h"
-#include "Converters/GLTFMaterialUtility.h"
-#include "Builders/GLTFContainerBuilder.h"
-#include "Materials/MaterialExpressionConstant.h"
-#include "Materials/MaterialExpressionConstant2Vector.h"
-#include "Materials/MaterialExpressionConstant3Vector.h"
-#include "Materials/MaterialExpressionConstant4Vector.h"
-#include "Materials/MaterialExpressionScalarParameter.h"
-#include "Materials/MaterialExpressionVectorParameter.h"
-#include "Materials/MaterialExpressionTextureSample.h"
-#include "Materials/MaterialExpressionTextureSampleParameter2D.h"
-
-namespace
-{
-	// Component masks
-	const FLinearColor RedMask(1.0f, 0.0f, 0.0f, 0.0f);
-	const FLinearColor GreenMask(0.0f, 1.0f, 0.0f, 0.0f);
-	const FLinearColor BlueMask(0.0f, 0.0f, 1.0f, 0.0f);
-	const FLinearColor AlphaMask(0.0f, 0.0f, 0.0f, 1.0f);
-	const FLinearColor RgbMask = RedMask + GreenMask + BlueMask;
-	const FLinearColor RgbaMask = RgbMask + AlphaMask;
-
-	// Property-specific component masks
-	const FLinearColor BaseColorMask = RgbMask;
-	const FLinearColor OpacityMask = AlphaMask;
-	const FLinearColor MetallicMask = BlueMask;
-	const FLinearColor RoughnessMask = GreenMask;
-	const FLinearColor OcclusionMask = RedMask;
-	const FLinearColor ClearCoatMask = RedMask;
-	const FLinearColor ClearCoatRoughnessMask = GreenMask;
-
-	// Ideal masks for texture-inputs (doesn't require baking)
-	const TArray<FLinearColor> DefaultColorInputMasks = { RgbMask, RgbaMask };
-	const TArray<FLinearColor> BaseColorInputMasks = { BaseColorMask };
-	const TArray<FLinearColor> OpacityInputMasks = { OpacityMask };
-	const TArray<FLinearColor> MetallicInputMasks = { MetallicMask };
-	const TArray<FLinearColor> RoughnessInputMasks = { RoughnessMask };
-	const TArray<FLinearColor> OcclusionInputMasks = { OcclusionMask };
-	const TArray<FLinearColor> ClearCoatInputMasks = { ClearCoatMask };
-	const TArray<FLinearColor> ClearCoatRoughnessInputMasks = { ClearCoatRoughnessMask };
-}
+#include "Builders/GLTFConvertBuilder.h"
+#include "Tasks/GLTFMaterialTasks.h"
 
 FGLTFJsonMaterialIndex FGLTFMaterialConverter::Convert(const UMaterialInterface* Material)
 {
-	{
-		const UMaterial* ParentMaterial = Material->GetMaterial();
-
-		if (ParentMaterial->MaterialDomain != MD_Surface)
-		{
-			// TODO: report warning (non-surface materials not supported, will be treated as surface)
-		}
-
-		if (!ParentMaterial->bTangentSpaceNormal)
-		{
-			// TODO: report warning (world-space normals not supported, will be treated as tangent-space)
-		}
-	}
-
-	FGLTFJsonMaterial JsonMaterial;
-	Material->GetName(JsonMaterial.Name);
-
-	// TODO: add support for additional blend modes (like Additive and Modulate)?
-	JsonMaterial.AlphaMode = FGLTFConverterUtility::ConvertBlendMode(Material->GetBlendMode());
-	JsonMaterial.AlphaCutoff = Material->GetOpacityMaskClipValue();
-	JsonMaterial.DoubleSided = Material->IsTwoSided();
-
-	if (!TryGetShadingModel(JsonMaterial.ShadingModel, Material))
-	{
-		JsonMaterial.ShadingModel = EGLTFJsonShadingModel::Default;
-		Builder.AddWarningMessage(FString::Printf(TEXT("Material %s will be exported as shading model %s"), *Material->GetName(), *FGLTFNameUtility::GetName(MSM_DefaultLit)));
-	}
-
-	if (JsonMaterial.ShadingModel != EGLTFJsonShadingModel::None)
-	{
-		const EMaterialProperty BaseColorProperty = JsonMaterial.ShadingModel == EGLTFJsonShadingModel::Unlit ? MP_EmissiveColor : MP_BaseColor;
-		const EMaterialProperty OpacityProperty = JsonMaterial.AlphaMode == EGLTFJsonAlphaMode::Mask ? MP_OpacityMask : MP_Opacity;
-
-		// TODO: check if a property is active before trying to get it (i.e. Material->IsPropertyActive)
-
-		if (JsonMaterial.AlphaMode == EGLTFJsonAlphaMode::Opaque)
-		{
-			if (!TryGetConstantColor(JsonMaterial.PBRMetallicRoughness.BaseColorFactor, BaseColorProperty, Material))
-			{
-				if (!TryGetSourceTexture(JsonMaterial.PBRMetallicRoughness.BaseColorTexture, BaseColorProperty, Material, DefaultColorInputMasks))
-				{
-					if (!TryGetBakedMaterialProperty(JsonMaterial.PBRMetallicRoughness.BaseColorTexture, JsonMaterial.PBRMetallicRoughness.BaseColorFactor, BaseColorProperty, TEXT("BaseColor"), Material))
-					{
-						Builder.AddWarningMessage(FString::Printf(TEXT("Failed to export %s for material %s"), FGLTFMaterialUtility::GetPropertyName(BaseColorProperty), *Material->GetName()));
-					}
-				}
-			}
-
-			JsonMaterial.PBRMetallicRoughness.BaseColorFactor.A = 1.0f; // make sure base color is opaque
-		}
-		else
-		{
-			if (!TryGetBaseColorAndOpacity(JsonMaterial.PBRMetallicRoughness, Material, BaseColorProperty, OpacityProperty))
-			{
-				Builder.AddWarningMessage(FString::Printf(TEXT("Failed to export %s and %s for material %s"), FGLTFMaterialUtility::GetPropertyName(BaseColorProperty), FGLTFMaterialUtility::GetPropertyName(OpacityProperty), *Material->GetName()));
-			}
-		}
-
-		if (JsonMaterial.ShadingModel == EGLTFJsonShadingModel::Default || JsonMaterial.ShadingModel == EGLTFJsonShadingModel::ClearCoat)
-		{
-			const EMaterialProperty MetallicProperty = MP_Metallic;
-			const EMaterialProperty RoughnessProperty = MP_Roughness;
-
-			if (!TryGetMetallicAndRoughness(JsonMaterial.PBRMetallicRoughness, Material, MetallicProperty, RoughnessProperty))
-			{
-				Builder.AddWarningMessage(FString::Printf(TEXT("Failed to export %s and %s for material %s"), FGLTFMaterialUtility::GetPropertyName(MetallicProperty), FGLTFMaterialUtility::GetPropertyName(RoughnessProperty), *Material->GetName()));
-			}
-
-			const EMaterialProperty EmissiveProperty = MP_EmissiveColor;
-			if (!TryGetEmissive(JsonMaterial, EmissiveProperty, Material))
-			{
-				Builder.AddWarningMessage(FString::Printf(TEXT("Failed to export %s for material %s"), FGLTFMaterialUtility::GetPropertyName(EmissiveProperty), *Material->GetName()));
-			}
-
-			// TODO: replace dummy enum MP_CustomOutput workaround for ClearCoatBottomNormal with proper support for custom outputs
-			const EMaterialProperty NormalProperty = JsonMaterial.ShadingModel == EGLTFJsonShadingModel::ClearCoat ? MP_CustomOutput : MP_Normal;
-			if (IsPropertyNonDefault(NormalProperty, Material))
-			{
-				if (!TryGetSourceTexture(JsonMaterial.NormalTexture, NormalProperty, Material, DefaultColorInputMasks))
-				{
-					if (!TryGetBakedMaterialProperty(JsonMaterial.NormalTexture, NormalProperty, TEXT("Normal"), Material))
-					{
-						Builder.AddWarningMessage(FString::Printf(TEXT("Failed to export %s for material %s"), FGLTFMaterialUtility::GetPropertyName(NormalProperty), *Material->GetName()));
-					}
-				}
-			}
-
-			const EMaterialProperty AmbientOcclusionProperty = MP_AmbientOcclusion;
-			if (IsPropertyNonDefault(AmbientOcclusionProperty, Material))
-			{
-				if (!TryGetSourceTexture(JsonMaterial.OcclusionTexture, AmbientOcclusionProperty, Material, OcclusionInputMasks))
-				{
-					if (!TryGetBakedMaterialProperty(JsonMaterial.OcclusionTexture, AmbientOcclusionProperty, TEXT("Occlusion"), Material))
-					{
-						Builder.AddWarningMessage(FString::Printf(TEXT("Failed to export %s for material %s"), FGLTFMaterialUtility::GetPropertyName(AmbientOcclusionProperty), *Material->GetName()));
-					}
-				}
-			}
-
-			if (JsonMaterial.ShadingModel == EGLTFJsonShadingModel::ClearCoat)
-			{
-				const EMaterialProperty ClearCoatProperty = MP_CustomData0;
-				const EMaterialProperty ClearCoatRoughnessProperty = MP_CustomData1;
-
-				if (!TryGetClearCoatRoughness(JsonMaterial.ClearCoat, Material, ClearCoatProperty, ClearCoatRoughnessProperty))
-				{
-					Builder.AddWarningMessage(FString::Printf(TEXT("Failed to export %s and %s for material %s"), FGLTFMaterialUtility::GetPropertyName(ClearCoatProperty), FGLTFMaterialUtility::GetPropertyName(ClearCoatRoughnessProperty), *Material->GetName()));
-				}
-
-				const EMaterialProperty ClearCoatNormalProperty = MP_Normal;
-				if (IsPropertyNonDefault(ClearCoatNormalProperty, Material))
-				{
-					if (!TryGetSourceTexture(JsonMaterial.ClearCoat.ClearCoatNormalTexture, ClearCoatNormalProperty, Material, DefaultColorInputMasks))
-					{
-						if (!TryGetBakedMaterialProperty(JsonMaterial.ClearCoat.ClearCoatNormalTexture, ClearCoatNormalProperty, TEXT("ClearCoatNormal"), Material))
-						{
-							Builder.AddWarningMessage(FString::Printf(TEXT("Failed to export %s for material %s"), FGLTFMaterialUtility::GetPropertyName(ClearCoatNormalProperty), *Material->GetName()));
-						}
-					}
-				}
-			}
-		}
-	}
-
-	return Builder.AddMaterial(JsonMaterial);
-}
-
-bool FGLTFMaterialConverter::TryGetShadingModel(EGLTFJsonShadingModel& OutShadingModel, const UMaterialInterface* Material) const
-{
-	const FMaterialShadingModelField ShadingModels = Material->GetShadingModels();
-	const int32 ShadingModelCount = ShadingModels.CountShadingModels();
-
-	if (ShadingModelCount <= 0)
-	{
-		Builder.AddWarningMessage(FString::Printf(TEXT("No shading model found for material %s"), *Material->GetName()));
-		return false;
-	}
-
-	if (ShadingModelCount > 1)
-	{
-		Builder.AddWarningMessage(FString::Printf(TEXT("Support is limited to the first of multiple shading models found (%d) in material %s"), ShadingModelCount, *Material->GetName()));
-	}
-
-	const EMaterialShadingModel ShadingModel = ShadingModels.GetFirstShadingModel();
-	const EGLTFJsonShadingModel ConvertedShadingModel = FGLTFConverterUtility::ConvertShadingModel(ShadingModel);
-
-	if (ConvertedShadingModel == EGLTFJsonShadingModel::None)
-	{
-		const FString ShadingModelName = FGLTFNameUtility::GetName(ShadingModel);
-		Builder.AddWarningMessage(FString::Printf(TEXT("Unsupported shading model (%s) in material %s"), *ShadingModelName, *Material->GetName()));
-		return false;
-	}
-
-	if (ConvertedShadingModel == EGLTFJsonShadingModel::Unlit && !Builder.ExportOptions->bExportUnlitMaterials)
-	{
-		const FString ShadingModelName = FGLTFNameUtility::GetName(ShadingModel);
-		Builder.AddWarningMessage(FString::Printf(TEXT("Shading model (%s) in material %s disabled by export options"), *ShadingModelName, *Material->GetName()));
-		return false;
-	}
-
-	if (ConvertedShadingModel == EGLTFJsonShadingModel::ClearCoat && !Builder.ExportOptions->bExportClearCoatMaterials)
-	{
-		const FString ShadingModelName = FGLTFNameUtility::GetName(ShadingModel);
-		Builder.AddWarningMessage(FString::Printf(TEXT("Shading model (%s) in material %s disabled by export options"), *ShadingModelName, *Material->GetName()));
-		return false;
-	}
-
-	OutShadingModel = ConvertedShadingModel;
-	return true;
-}
-
-bool FGLTFMaterialConverter::TryGetBaseColorAndOpacity(FGLTFJsonPBRMetallicRoughness& OutPBRParams, const UMaterialInterface* Material, EMaterialProperty BaseColorProperty, EMaterialProperty OpacityProperty) const
-{
-	const bool bIsBaseColorConstant = TryGetConstantColor(OutPBRParams.BaseColorFactor, BaseColorProperty, Material);
-	const bool bIsOpacityConstant = TryGetConstantScalar(OutPBRParams.BaseColorFactor.A, OpacityProperty, Material);
-
-	if (bIsBaseColorConstant && bIsOpacityConstant)
-	{
-		return true;
-	}
-
-	// NOTE: since we always bake the properties (for now) when atleast property is non-const, we need
-	// to reset the constant factors to their defaults. Otherwise the baked value of a constant property
-	// would be scaled with the factor, i.e a double scaling.
-	OutPBRParams.BaseColorFactor = { 1.0f, 1.0f, 1.0f, 1.0f };
-
-	const UTexture2D* BaseColorTexture;
-	const UTexture2D* OpacityTexture;
-	int32 BaseColorTexCoord;
-	int32 OpacityTexCoord;
-
-	const bool bHasBaseColorSourceTexture = TryGetSourceTexture(BaseColorTexture, BaseColorTexCoord, BaseColorProperty, Material, BaseColorInputMasks);
-	const bool bHasOpacitySourceTexture = TryGetSourceTexture(OpacityTexture, OpacityTexCoord, OpacityProperty, Material, OpacityInputMasks);
-
-	// Detect the "happy path" where both inputs share the same texture and are correctly masked.
-	if (bHasBaseColorSourceTexture &&
-		bHasOpacitySourceTexture &&
-		BaseColorTexture == OpacityTexture &&
-		BaseColorTexCoord == OpacityTexCoord)
-	{
-		OutPBRParams.BaseColorTexture.Index = Builder.GetOrAddTexture(BaseColorTexture);
-		OutPBRParams.BaseColorTexture.TexCoord = BaseColorTexCoord;
-		return true;
-	}
-
-	if (!Builder.ExportOptions->bBakeMaterialInputs)
-	{
-		Builder.AddWarningMessage(FString::Printf(
-			TEXT("%s and %s for material %s needs to bake, but material baking disabled by export options"),
-			FGLTFMaterialUtility::GetPropertyName(BaseColorProperty),
-			FGLTFMaterialUtility::GetPropertyName(OpacityProperty),
-			*Material->GetName()));
-		return false;
-	}
-
-	// TODO: add support for calculating the ideal resolution to use for baking based on connected (texture) nodes
-	FIntPoint TextureSize = Builder.ExportOptions->GetDefaultMaterialBakeSize();
-
-	// TODO: should this be the default wrap-mode?
-	EGLTFJsonTextureWrap TextureWrapS = EGLTFJsonTextureWrap::Repeat;
-	EGLTFJsonTextureWrap TextureWrapT = EGLTFJsonTextureWrap::Repeat;
-
-	// TODO: should this be the default filter?
-	EGLTFJsonTextureFilter TextureMinFilter = EGLTFJsonTextureFilter::LinearMipmapLinear;
-	EGLTFJsonTextureFilter TextureMagFilter = EGLTFJsonTextureFilter::Linear;
-
-	if (bHasBaseColorSourceTexture && bHasOpacitySourceTexture)
-	{
-		const bool bAreTexturesCompatible = BaseColorTexCoord == OpacityTexCoord &&
-			BaseColorTexture->AddressX == OpacityTexture->AddressX &&
-			BaseColorTexture->AddressY == OpacityTexture->AddressY;
-
-		if (!bAreTexturesCompatible)
-		{
-			// TODO: handle differences in wrapping or uv-coords
-			Builder.AddWarningMessage(FString::Printf(
-				TEXT("BaseColor- and Opacity-textures for material %s were not able to be combined and will be skipped"),
-				*Material->GetName()));
-
-			return false;
-		}
-
-		TextureSize =
-		{
-			FMath::Max(BaseColorTexture->GetSizeX(), OpacityTexture->GetSizeX()),
-			FMath::Max(BaseColorTexture->GetSizeY(), OpacityTexture->GetSizeY())
-		};
-
-		TextureWrapS = FGLTFConverterUtility::ConvertWrap(BaseColorTexture->AddressX);
-		TextureWrapT = FGLTFConverterUtility::ConvertWrap(BaseColorTexture->AddressY);
-
-		// TODO: compare min- and mag-filter for BaseColorTexture and OpacityTexture. If they differ,
-		// we should choose one or the other and inform the user about the choice made by logging to the console.
-		TextureMinFilter = FGLTFConverterUtility::ConvertMinFilter(BaseColorTexture->Filter, BaseColorTexture->LODGroup);
-		TextureMagFilter = FGLTFConverterUtility::ConvertMagFilter(BaseColorTexture->Filter, BaseColorTexture->LODGroup);
-	}
-	else if (bHasBaseColorSourceTexture)
-	{
-		TextureSize = { BaseColorTexture->GetSizeX(), BaseColorTexture->GetSizeY() };
-		TextureWrapS = FGLTFConverterUtility::ConvertWrap(BaseColorTexture->AddressX);
-		TextureWrapT = FGLTFConverterUtility::ConvertWrap(BaseColorTexture->AddressY);
-		TextureMinFilter = FGLTFConverterUtility::ConvertMinFilter(BaseColorTexture->Filter, BaseColorTexture->LODGroup);
-		TextureMagFilter = FGLTFConverterUtility::ConvertMagFilter(BaseColorTexture->Filter, BaseColorTexture->LODGroup);
-	}
-	else if (bHasOpacitySourceTexture)
-	{
-		TextureSize = { OpacityTexture->GetSizeX(), OpacityTexture->GetSizeY() };
-		TextureWrapS = FGLTFConverterUtility::ConvertWrap(OpacityTexture->AddressX);
-		TextureWrapT = FGLTFConverterUtility::ConvertWrap(OpacityTexture->AddressY);
-		TextureMinFilter = FGLTFConverterUtility::ConvertMinFilter(OpacityTexture->Filter, OpacityTexture->LODGroup);
-		TextureMagFilter = FGLTFConverterUtility::ConvertMagFilter(OpacityTexture->Filter, OpacityTexture->LODGroup);
-	}
-
-	const FGLTFPropertyBakeOutput BaseColorBakeOutput = BakeMaterialProperty(BaseColorProperty, Material, BaseColorTexCoord, &TextureSize);
-	const FGLTFPropertyBakeOutput OpacityBakeOutput = BakeMaterialProperty(OpacityProperty, Material, OpacityTexCoord, &TextureSize, true);
-	const float BaseColorScale = BaseColorProperty == MP_EmissiveColor ? BaseColorBakeOutput.EmissiveScale : 1;
-
-	// Detect when both baked properties are constants, which means we can avoid exporting a texture
-	if (BaseColorBakeOutput.bIsConstant && OpacityBakeOutput.bIsConstant)
-	{
-		FLinearColor BaseColorFactor(BaseColorBakeOutput.ConstantValue * BaseColorScale);
-		BaseColorFactor.A = OpacityBakeOutput.ConstantValue.A;
-
-		OutPBRParams.BaseColorFactor = FGLTFConverterUtility::ConvertColor(BaseColorFactor);
-		return true;
-	}
-
-	int32 TexCoord;
-	if (BaseColorBakeOutput.bIsConstant)
-	{
-		TexCoord = OpacityTexCoord;
-	}
-	else if (OpacityBakeOutput.bIsConstant)
-	{
-		TexCoord = BaseColorTexCoord;
-	}
-	else if (BaseColorTexCoord == OpacityTexCoord)
-	{
-		TexCoord = BaseColorTexCoord;
-	}
-	else
-	{
-		// TODO: report error (texture coordinate conflict)
-		return false;
-	}
-
-	TextureSize = BaseColorBakeOutput.Size.ComponentMax(OpacityBakeOutput.Size);
-	BaseColorTexture = FGLTFMaterialUtility::CreateTransientTexture(BaseColorBakeOutput);
-	OpacityTexture = FGLTFMaterialUtility::CreateTransientTexture(OpacityBakeOutput);
-
-	const FString TextureName = Material->GetName() + TEXT("_BaseColor");
-
-	const TArray<FGLTFTextureCombineSource> CombineSources =
-	{
-		{ OpacityTexture, OpacityMask, SE_BLEND_Opaque },
-		{ BaseColorTexture, BaseColorMask }
-	};
-
-	const FGLTFJsonTextureIndex TextureIndex = FGLTFMaterialUtility::AddCombinedTexture(
-		Builder,
-		CombineSources,
-		TextureSize,
-		TextureName,
-		TextureMinFilter,
-		TextureMagFilter,
-		TextureWrapS,
-		TextureWrapT);
-
-	OutPBRParams.BaseColorTexture.TexCoord = TexCoord;
-	OutPBRParams.BaseColorTexture.Index = TextureIndex;
-	OutPBRParams.BaseColorFactor = { BaseColorScale, BaseColorScale, BaseColorScale };
-
-	return true;
-}
-
-bool FGLTFMaterialConverter::TryGetMetallicAndRoughness(FGLTFJsonPBRMetallicRoughness& OutPBRParams, const UMaterialInterface* Material, EMaterialProperty MetallicProperty, EMaterialProperty RoughnessProperty) const
-{
-	const bool bIsMetallicConstant = TryGetConstantScalar(OutPBRParams.MetallicFactor, MetallicProperty, Material);
-	const bool bIsRoughnessConstant = TryGetConstantScalar(OutPBRParams.RoughnessFactor, RoughnessProperty, Material);
-
-	if (bIsMetallicConstant && bIsRoughnessConstant)
-	{
-		return true;
-	}
-
-	// NOTE: since we always bake the properties (for now) when atleast one property is non-const, we need
-	// to reset the constant factors to their defaults. Otherwise the baked value of a constant property
-	// would be scaled with the factor, i.e a double scaling.
-	OutPBRParams.MetallicFactor = 1.0f;
-	OutPBRParams.RoughnessFactor = 1.0f;
-
-	const UTexture2D* MetallicTexture;
-	const UTexture2D* RoughnessTexture;
-	int32 MetallicTexCoord;
-	int32 RoughnessTexCoord;
-
-	const bool bHasMetallicSourceTexture = TryGetSourceTexture(MetallicTexture, MetallicTexCoord, MetallicProperty, Material, MetallicInputMasks);
-	const bool bHasRoughnessSourceTexture = TryGetSourceTexture(RoughnessTexture, RoughnessTexCoord, RoughnessProperty, Material, RoughnessInputMasks);
-
-	// Detect the "happy path" where both inputs share the same texture and are correctly masked.
-	if (bHasMetallicSourceTexture &&
-		bHasRoughnessSourceTexture &&
-		MetallicTexture == RoughnessTexture &&
-		MetallicTexCoord == RoughnessTexCoord)
-	{
-		OutPBRParams.MetallicRoughnessTexture.Index = Builder.GetOrAddTexture(MetallicTexture);
-		OutPBRParams.MetallicRoughnessTexture.TexCoord = MetallicTexCoord;
-		return true;
-	}
-
-	if (!Builder.ExportOptions->bBakeMaterialInputs)
-	{
-		Builder.AddWarningMessage(FString::Printf(
-			TEXT("%s and %s for material %s needs to bake, but material baking disabled by export options"),
-			FGLTFMaterialUtility::GetPropertyName(MetallicProperty),
-			FGLTFMaterialUtility::GetPropertyName(RoughnessProperty),
-			*Material->GetName()));
-		return false;
-	}
-
-	// TODO: add support for calculating the ideal resolution to use for baking based on connected (texture) nodes
-	FIntPoint TextureSize = Builder.ExportOptions->GetDefaultMaterialBakeSize();
-
-	// TODO: should this be the default wrap-mode?
-	EGLTFJsonTextureWrap TextureWrapS = EGLTFJsonTextureWrap::Repeat;
-	EGLTFJsonTextureWrap TextureWrapT = EGLTFJsonTextureWrap::Repeat;
-
-	// TODO: should this be the default filter?
-	EGLTFJsonTextureFilter TextureMinFilter = EGLTFJsonTextureFilter::LinearMipmapLinear;
-	EGLTFJsonTextureFilter TextureMagFilter = EGLTFJsonTextureFilter::Linear;
-
-	if (bHasMetallicSourceTexture && bHasRoughnessSourceTexture)
-	{
-		const bool bAreTexturesCompatible = MetallicTexCoord == RoughnessTexCoord &&
-			MetallicTexture->AddressX == RoughnessTexture->AddressX &&
-			MetallicTexture->AddressY == RoughnessTexture->AddressY;
-
-		if (!bAreTexturesCompatible)
-		{
-			// TODO: handle differences in wrapping or uv-coords
-			Builder.AddWarningMessage(FString::Printf(
-				TEXT("Metallic- and Roughness-textures for material %s were not able to be combined and will be skipped"),
-				*Material->GetName()));
-
-			return false;
-		}
-
-		TextureSize =
-		{
-			FMath::Max(MetallicTexture->GetSizeX(), RoughnessTexture->GetSizeX()),
-			FMath::Max(MetallicTexture->GetSizeY(), RoughnessTexture->GetSizeY())
-		};
-
-		TextureWrapS = FGLTFConverterUtility::ConvertWrap(MetallicTexture->AddressX);
-		TextureWrapT = FGLTFConverterUtility::ConvertWrap(MetallicTexture->AddressY);
-
-		// TODO: compare min- and mag-filter for BaseColorTexture and OpacityTexture. If they differ,
-		// we should choose one or the other and inform the user about the choice made by logging to the console.
-		TextureMinFilter = FGLTFConverterUtility::ConvertMinFilter(MetallicTexture->Filter, MetallicTexture->LODGroup);
-		TextureMagFilter = FGLTFConverterUtility::ConvertMagFilter(MetallicTexture->Filter, MetallicTexture->LODGroup);
-	}
-	else if (bHasMetallicSourceTexture)
-	{
-		TextureSize = { MetallicTexture->GetSizeX(), MetallicTexture->GetSizeY() };
-		TextureWrapS = FGLTFConverterUtility::ConvertWrap(MetallicTexture->AddressX);
-		TextureWrapT = FGLTFConverterUtility::ConvertWrap(MetallicTexture->AddressY);
-		TextureMinFilter = FGLTFConverterUtility::ConvertMinFilter(MetallicTexture->Filter, MetallicTexture->LODGroup);
-		TextureMagFilter = FGLTFConverterUtility::ConvertMagFilter(MetallicTexture->Filter, MetallicTexture->LODGroup);
-	}
-	else if (bHasRoughnessSourceTexture)
-	{
-		TextureSize = { RoughnessTexture->GetSizeX(), RoughnessTexture->GetSizeY() };
-		TextureWrapS = FGLTFConverterUtility::ConvertWrap(RoughnessTexture->AddressX);
-		TextureWrapT = FGLTFConverterUtility::ConvertWrap(RoughnessTexture->AddressY);
-		TextureMinFilter = FGLTFConverterUtility::ConvertMinFilter(RoughnessTexture->Filter, RoughnessTexture->LODGroup);
-		TextureMagFilter = FGLTFConverterUtility::ConvertMagFilter(RoughnessTexture->Filter, RoughnessTexture->LODGroup);
-	}
-
-	const FGLTFPropertyBakeOutput MetallicBakeOutput = BakeMaterialProperty(MetallicProperty, Material, MetallicTexCoord, &TextureSize);
-	const FGLTFPropertyBakeOutput RoughnessBakeOutput = BakeMaterialProperty(RoughnessProperty, Material, RoughnessTexCoord, &TextureSize);
-
-	// Detect when both baked properties are constants, which means we can use factors and avoid exporting a texture
-	if (MetallicBakeOutput.bIsConstant && RoughnessBakeOutput.bIsConstant)
-	{
-		OutPBRParams.MetallicFactor = MetallicBakeOutput.ConstantValue.R;
-		OutPBRParams.RoughnessFactor = RoughnessBakeOutput.ConstantValue.R;
-		return true;
-	}
-
-	int32 TexCoord;
-	if (MetallicBakeOutput.bIsConstant)
-	{
-		TexCoord = RoughnessTexCoord;
-	}
-	else if (MetallicBakeOutput.bIsConstant)
-	{
-		TexCoord = MetallicTexCoord;
-	}
-	else if (MetallicTexCoord == RoughnessTexCoord)
-	{
-		TexCoord = MetallicTexCoord;
-	}
-	else
-	{
-		// TODO: report error (texture coordinate conflict)
-		return false;
-	}
-
-	TextureSize = RoughnessBakeOutput.Size.ComponentMax(MetallicBakeOutput.Size);
-	MetallicTexture = FGLTFMaterialUtility::CreateTransientTexture(MetallicBakeOutput);
-	RoughnessTexture = FGLTFMaterialUtility::CreateTransientTexture(RoughnessBakeOutput);
-
-	const FString TextureName = Material->GetName() + TEXT("_MetallicRoughness");
-
-	const TArray<FGLTFTextureCombineSource> CombineSources =
-	{
-		{ MetallicTexture, MetallicMask + AlphaMask, SE_BLEND_Opaque },
-		{ RoughnessTexture, RoughnessMask }
-	};
-
-	const FGLTFJsonTextureIndex TextureIndex = FGLTFMaterialUtility::AddCombinedTexture(
-		Builder,
-		CombineSources,
-		TextureSize,
-		TextureName,
-		TextureMinFilter,
-		TextureMagFilter,
-		TextureWrapS,
-		TextureWrapT);
-
-	OutPBRParams.MetallicRoughnessTexture.TexCoord = TexCoord;
-	OutPBRParams.MetallicRoughnessTexture.Index = TextureIndex;
-
-	return true;
-}
-
-bool FGLTFMaterialConverter::TryGetClearCoatRoughness(FGLTFJsonClearCoatExtension& OutExtParams, const UMaterialInterface* Material, EMaterialProperty IntensityProperty, EMaterialProperty RoughnessProperty) const
-{
-	const bool bIsIntensityConstant = TryGetConstantScalar(OutExtParams.ClearCoatFactor, IntensityProperty, Material);
-	const bool bIsRoughnessConstant = TryGetConstantScalar(OutExtParams.ClearCoatRoughnessFactor, RoughnessProperty, Material);
-
-	if (bIsIntensityConstant && bIsRoughnessConstant)
-	{
-		return true;
-	}
-
-	// NOTE: since we always bake the properties (for now) when atleast one property is non-const, we need
-	// to reset the constant factors to their defaults. Otherwise the baked value of a constant property
-	// would be scaled with the factor, i.e a double scaling.
-	OutExtParams.ClearCoatFactor = 1.0f;
-	OutExtParams.ClearCoatRoughnessFactor = 1.0f;
-
-	const UTexture2D* IntensityTexture;
-	const UTexture2D* RoughnessTexture;
-	int32 IntensityTexCoord;
-	int32 RoughnessTexCoord;
-
-	const bool bHasIntensitySourceTexture = TryGetSourceTexture(IntensityTexture, IntensityTexCoord, IntensityProperty, Material, ClearCoatInputMasks);
-	const bool bHasRoughnessSourceTexture = TryGetSourceTexture(RoughnessTexture, RoughnessTexCoord, RoughnessProperty, Material, ClearCoatRoughnessInputMasks);
-
-	// Detect the "happy path" where both inputs share the same texture and are correctly masked.
-	if (bHasIntensitySourceTexture &&
-		bHasRoughnessSourceTexture &&
-		IntensityTexture == RoughnessTexture &&
-		IntensityTexCoord == RoughnessTexCoord)
-	{
-		const FGLTFJsonTextureIndex TextureIndex = Builder.GetOrAddTexture(IntensityTexture);
-		OutExtParams.ClearCoatTexture.Index = TextureIndex;
-		OutExtParams.ClearCoatTexture.TexCoord = IntensityTexCoord;
-		OutExtParams.ClearCoatRoughnessTexture.Index = TextureIndex;
-		OutExtParams.ClearCoatRoughnessTexture.TexCoord = RoughnessTexCoord;
-		return true;
-	}
-
-	if (!Builder.ExportOptions->bBakeMaterialInputs)
-	{
-		Builder.AddWarningMessage(FString::Printf(
-			TEXT("%s and %s for material %s needs to bake, but material baking disabled by export options"),
-			FGLTFMaterialUtility::GetPropertyName(IntensityProperty),
-			FGLTFMaterialUtility::GetPropertyName(RoughnessProperty),
-			*Material->GetName()));
-		return false;
-	}
-
-	// TODO: add support for calculating the ideal resolution to use for baking based on connected (texture) nodes
-	FIntPoint TextureSize = Builder.ExportOptions->GetDefaultMaterialBakeSize();
-
-	// TODO: should this be the default wrap-mode?
-	EGLTFJsonTextureWrap TextureWrapS = EGLTFJsonTextureWrap::Repeat;
-	EGLTFJsonTextureWrap TextureWrapT = EGLTFJsonTextureWrap::Repeat;
-
-	// TODO: should this be the default filter?
-	EGLTFJsonTextureFilter TextureMinFilter = EGLTFJsonTextureFilter::LinearMipmapLinear;
-	EGLTFJsonTextureFilter TextureMagFilter = EGLTFJsonTextureFilter::Linear;
-
-	if (bHasIntensitySourceTexture && bHasRoughnessSourceTexture)
-	{
-		const bool bAreTexturesCompatible = IntensityTexCoord == RoughnessTexCoord &&
-			IntensityTexture->AddressX == RoughnessTexture->AddressX &&
-			IntensityTexture->AddressY == RoughnessTexture->AddressY;
-
-		if (!bAreTexturesCompatible)
-		{
-			// TODO: handle differences in wrapping or uv-coords
-			Builder.AddWarningMessage(FString::Printf(
-				TEXT("Intensity- and Roughness-textures for material %s were not able to be combined and will be skipped"),
-				*Material->GetName()));
-
-			return false;
-		}
-
-		TextureSize =
-		{
-			FMath::Max(IntensityTexture->GetSizeX(), RoughnessTexture->GetSizeX()),
-			FMath::Max(IntensityTexture->GetSizeY(), RoughnessTexture->GetSizeY())
-		};
-
-		TextureWrapS = FGLTFConverterUtility::ConvertWrap(IntensityTexture->AddressX);
-		TextureWrapT = FGLTFConverterUtility::ConvertWrap(IntensityTexture->AddressY);
-
-		// TODO: compare min- and mag-filter for BaseColorTexture and OpacityTexture. If they differ,
-		// we should choose one or the other and inform the user about the choice made by logging to the console.
-		TextureMinFilter = FGLTFConverterUtility::ConvertMinFilter(IntensityTexture->Filter, IntensityTexture->LODGroup);
-		TextureMagFilter = FGLTFConverterUtility::ConvertMagFilter(IntensityTexture->Filter, IntensityTexture->LODGroup);
-	}
-	else if (bHasIntensitySourceTexture)
-	{
-		TextureSize = { IntensityTexture->GetSizeX(), IntensityTexture->GetSizeY() };
-		TextureWrapS = FGLTFConverterUtility::ConvertWrap(IntensityTexture->AddressX);
-		TextureWrapT = FGLTFConverterUtility::ConvertWrap(IntensityTexture->AddressY);
-		TextureMinFilter = FGLTFConverterUtility::ConvertMinFilter(IntensityTexture->Filter, IntensityTexture->LODGroup);
-		TextureMagFilter = FGLTFConverterUtility::ConvertMagFilter(IntensityTexture->Filter, IntensityTexture->LODGroup);
-	}
-	else if (bHasRoughnessSourceTexture)
-	{
-		TextureSize = { RoughnessTexture->GetSizeX(), RoughnessTexture->GetSizeY() };
-		TextureWrapS = FGLTFConverterUtility::ConvertWrap(RoughnessTexture->AddressX);
-		TextureWrapT = FGLTFConverterUtility::ConvertWrap(RoughnessTexture->AddressY);
-		TextureMinFilter = FGLTFConverterUtility::ConvertMinFilter(RoughnessTexture->Filter, RoughnessTexture->LODGroup);
-		TextureMagFilter = FGLTFConverterUtility::ConvertMagFilter(RoughnessTexture->Filter, RoughnessTexture->LODGroup);
-	}
-
-	const FGLTFPropertyBakeOutput IntensityBakeOutput = BakeMaterialProperty(IntensityProperty, Material, IntensityTexCoord, &TextureSize);
-	const FGLTFPropertyBakeOutput RoughnessBakeOutput = BakeMaterialProperty(RoughnessProperty, Material, RoughnessTexCoord, &TextureSize);
-
-	// Detect when both baked properties are constants, which means we can use factors and avoid exporting a texture
-	if (IntensityBakeOutput.bIsConstant && RoughnessBakeOutput.bIsConstant)
-	{
-		OutExtParams.ClearCoatFactor = IntensityBakeOutput.ConstantValue.R;
-		OutExtParams.ClearCoatRoughnessFactor =  RoughnessBakeOutput.ConstantValue.R;
-		return true;
-	}
-
-	int32 TexCoord;
-	if (IntensityBakeOutput.bIsConstant)
-	{
-		TexCoord = RoughnessTexCoord;
-	}
-	else if (IntensityBakeOutput.bIsConstant)
-	{
-		TexCoord = IntensityTexCoord;
-	}
-	else if (IntensityTexCoord == RoughnessTexCoord)
-	{
-		TexCoord = IntensityTexCoord;
-	}
-	else
-	{
-		// TODO: report error (texture coordinate conflict)
-		return false;
-	}
-
-	TextureSize = RoughnessBakeOutput.Size.ComponentMax(IntensityBakeOutput.Size);
-	IntensityTexture = FGLTFMaterialUtility::CreateTransientTexture(IntensityBakeOutput);
-	RoughnessTexture = FGLTFMaterialUtility::CreateTransientTexture(RoughnessBakeOutput);
-
-	const FString TextureName = Material->GetName() + TEXT("_ClearCoatRoughness");
-
-	const TArray<FGLTFTextureCombineSource> CombineSources =
-	{
-		{ IntensityTexture, ClearCoatMask + AlphaMask, SE_BLEND_Opaque },
-		{ RoughnessTexture, ClearCoatRoughnessMask }
-	};
-
-	const FGLTFJsonTextureIndex TextureIndex = FGLTFMaterialUtility::AddCombinedTexture(
-		Builder,
-		CombineSources,
-		TextureSize,
-		TextureName,
-		TextureMinFilter,
-		TextureMagFilter,
-		TextureWrapS,
-		TextureWrapT);
-
-	OutExtParams.ClearCoatTexture.Index = TextureIndex;
-	OutExtParams.ClearCoatTexture.TexCoord = TexCoord;
-	OutExtParams.ClearCoatRoughnessTexture.Index = TextureIndex;
-	OutExtParams.ClearCoatRoughnessTexture.TexCoord = TexCoord;
-
-	return true;
-}
-
-bool FGLTFMaterialConverter::TryGetEmissive(FGLTFJsonMaterial& JsonMaterial, EMaterialProperty EmissiveProperty, const UMaterialInterface* Material) const
-{
-	// TODO: right now we allow EmissiveFactor to be > 1.0 to support very bright emission, although it's not valid according to the glTF standard.
-	// We may want to change this behaviour and store factors above 1.0 using a custom extension instead.
-
-	if (TryGetConstantColor(JsonMaterial.EmissiveFactor, MP_EmissiveColor, Material))
-	{
-		return true;
-	}
-
-	if (TryGetSourceTexture(JsonMaterial.EmissiveTexture, EmissiveProperty, Material, DefaultColorInputMasks))
-	{
-		JsonMaterial.EmissiveFactor = FGLTFJsonColor3::White;	// make sure texture is not multiplied with black
-		return true;
-	}
-
-	if (!Builder.ExportOptions->bBakeMaterialInputs)
-	{
-		Builder.AddWarningMessage(FString::Printf(
-			TEXT("%s for material %s needs to bake, but material baking disabled by export options"),
-			FGLTFMaterialUtility::GetPropertyName(EmissiveProperty),
-			*Material->GetName()));
-		return false;
-	}
-
-	const FGLTFPropertyBakeOutput PropertyBakeOutput = BakeMaterialProperty(EmissiveProperty, Material, JsonMaterial.EmissiveTexture.TexCoord);
-	const float EmissiveScale = PropertyBakeOutput.EmissiveScale;
-
-	if (PropertyBakeOutput.bIsConstant)
-	{
-		const FLinearColor EmissiveColor = PropertyBakeOutput.ConstantValue;
-		JsonMaterial.EmissiveFactor = FGLTFConverterUtility::ConvertColor(EmissiveColor * EmissiveScale);
-	}
-	else
-	{
-		if (!StoreBakedPropertyTexture(JsonMaterial.EmissiveTexture, PropertyBakeOutput, TEXT("Emissive"), Material))
-		{
-			return false;
-		}
-
-		JsonMaterial.EmissiveFactor = { EmissiveScale, EmissiveScale, EmissiveScale };
-	}
-
-	return true;
-}
-
-bool FGLTFMaterialConverter::IsPropertyNonDefault(EMaterialProperty Property, const UMaterialInterface* Material) const
-{
-	const bool bUseMaterialAttributes = Material->GetMaterial()->bUseMaterialAttributes;
-	if (bUseMaterialAttributes)
-	{
-		// TODO: check if attribute property connected, i.e. Material->GetMaterial()->MaterialAttributes.IsConnected(Property)
-		return true;
-	}
-
-	const FExpressionInput* MaterialInput = FGLTFMaterialUtility::GetInputForProperty(Material, Property);
-	if (MaterialInput == nullptr)
-	{
-		// TODO: report error
-		return false;
-	}
-
-	const UMaterialExpression* Expression = MaterialInput->Expression;
-	if (Expression == nullptr)
-	{
-		return false;
-	}
-
-	return true;
-}
-
-bool FGLTFMaterialConverter::TryGetConstantColor(FGLTFJsonColor3& OutValue, EMaterialProperty Property, const UMaterialInterface* Material) const
-{
-	FLinearColor Value;
-	if (TryGetConstantColor(Value, Property, Material))
-	{
-		OutValue = FGLTFConverterUtility::ConvertColor(Value);
-		return true;
-	}
-
-	return false;
-}
-
-bool FGLTFMaterialConverter::TryGetConstantColor(FGLTFJsonColor4& OutValue, EMaterialProperty Property, const UMaterialInterface* Material) const
-{
-	FLinearColor Value;
-	if (TryGetConstantColor(Value, Property, Material))
-	{
-		OutValue = FGLTFConverterUtility::ConvertColor(Value);
-		return true;
-	}
-
-	return false;
-}
-
-bool FGLTFMaterialConverter::TryGetConstantColor(FLinearColor& OutValue, EMaterialProperty Property, const UMaterialInterface* Material) const
-{
-	const bool bUseMaterialAttributes = Material->GetMaterial()->bUseMaterialAttributes;
-	if (bUseMaterialAttributes)
-	{
-		// TODO: check if attribute property connected, i.e. Material->GetMaterial()->MaterialAttributes.IsConnected(Property)
-		return false;
-	}
-
-	const FMaterialInput<FColor>* MaterialInput = FGLTFMaterialUtility::GetInputForProperty<FColor>(Material, Property);
-	if (MaterialInput == nullptr)
-	{
-		// TODO: report error
-		return false;
-	}
-
-	if (MaterialInput->UseConstant)
-	{
-		OutValue = { MaterialInput->Constant };
-		return true;
-	}
-
-	const UMaterialExpression* Expression = MaterialInput->Expression;
-	if (Expression == nullptr)
-	{
-		OutValue = FLinearColor(FGLTFMaterialUtility::GetPropertyDefaultValue(Property));
-		return true;
-	}
-
-	if (const UMaterialExpressionVectorParameter* VectorParameter = ExactCast<UMaterialExpressionVectorParameter>(Expression))
-	{
-		FLinearColor Value = VectorParameter->DefaultValue;
-
-		const UMaterialInstance* MaterialInstance = Cast<UMaterialInstance>(Material);
-		if (MaterialInstance != nullptr)
-		{
-			const FHashedMaterialParameterInfo ParameterInfo(VectorParameter->GetParameterName());
-			if (!MaterialInstance->GetVectorParameterValue(ParameterInfo, Value))
-			{
-				// TODO: how to handle this?
-			}
-		}
-
-		const uint32 MaskComponentCount = FGLTFMaterialUtility::GetMaskComponentCount(*MaterialInput);
-
-		if (MaskComponentCount > 0)
-		{
-			const FLinearColor Mask = FGLTFMaterialUtility::GetMask(*MaterialInput);
-
-			Value *= Mask;
-
-			if (MaskComponentCount == 1)
-			{
-				const float ComponentValue = Value.R + Value.G + Value.B + Value.A;
-				Value = { ComponentValue, ComponentValue, ComponentValue, ComponentValue };
-			}
-		}
-
-		OutValue = Value;
-		return true;
-	}
-
-	if (const UMaterialExpressionScalarParameter* ScalarParameter = ExactCast<UMaterialExpressionScalarParameter>(Expression))
-	{
-		float Value = ScalarParameter->DefaultValue;
-
-		const UMaterialInstance* MaterialInstance = Cast<UMaterialInstance>(Material);
-		if (MaterialInstance != nullptr)
-		{
-			const FHashedMaterialParameterInfo ParameterInfo(ScalarParameter->GetParameterName());
-			if (!MaterialInstance->GetScalarParameterValue(ParameterInfo, Value))
-			{
-				// TODO: how to handle this?
-			}
-		}
-
-		OutValue = { Value, Value, Value, Value };
-		return true;
-	}
-
-	if (const UMaterialExpressionConstant4Vector* Constant4Vector = ExactCast<UMaterialExpressionConstant4Vector>(Expression))
-	{
-		OutValue = Constant4Vector->Constant;
-		return true;
-	}
-
-	if (const UMaterialExpressionConstant3Vector* Constant3Vector = ExactCast<UMaterialExpressionConstant3Vector>(Expression))
-	{
-		OutValue = Constant3Vector->Constant;
-		return true;
-	}
-
-	if (const UMaterialExpressionConstant2Vector* Constant2Vector = ExactCast<UMaterialExpressionConstant2Vector>(Expression))
-	{
-		OutValue = { Constant2Vector->R, Constant2Vector->G, 0, 0 };
-		return true;
-	}
-
-	if (const UMaterialExpressionConstant* Constant = ExactCast<UMaterialExpressionConstant>(Expression))
-	{
-		OutValue = { Constant->R, Constant->R, Constant->R, Constant->R };
-		return true;
-	}
-
-	return false;
-}
-
-bool FGLTFMaterialConverter::TryGetConstantScalar(float& OutValue, EMaterialProperty Property, const UMaterialInterface* Material) const
-{
-	const bool bUseMaterialAttributes = Material->GetMaterial()->bUseMaterialAttributes;
-	if (bUseMaterialAttributes)
-	{
-		// TODO: check if attribute property connected, i.e. Material->GetMaterial()->MaterialAttributes.IsConnected(Property)
-		return false;
-	}
-
-	const FMaterialInput<float>* MaterialInput = FGLTFMaterialUtility::GetInputForProperty<float>(Material, Property);
-	if (MaterialInput == nullptr)
-	{
-		// TODO: report error
-		return false;
-	}
-
-	if (MaterialInput->UseConstant)
-	{
-		OutValue = MaterialInput->Constant;
-		return true;
-	}
-
-	const UMaterialExpression* Expression = MaterialInput->Expression;
-	if (Expression == nullptr)
-	{
-		OutValue = FGLTFMaterialUtility::GetPropertyDefaultValue(Property).X;
-		return true;
-	}
-
-	if (const UMaterialExpressionVectorParameter* VectorParameter = ExactCast<UMaterialExpressionVectorParameter>(Expression))
-	{
-		FLinearColor Value = VectorParameter->DefaultValue;
-
-		const UMaterialInstance* MaterialInstance = Cast<UMaterialInstance>(Material);
-		if (MaterialInstance != nullptr)
-		{
-			const FHashedMaterialParameterInfo ParameterInfo(VectorParameter->GetParameterName());
-			if (!MaterialInstance->GetVectorParameterValue(ParameterInfo, Value))
-			{
-				// TODO: how to handle this?
-			}
-		}
-
-		const uint32 MaskComponentCount = FGLTFMaterialUtility::GetMaskComponentCount(*MaterialInput);
-
-		if (MaskComponentCount > 0)
-		{
-			const FLinearColor Mask = FGLTFMaterialUtility::GetMask(*MaterialInput);
-			Value *= Mask;
-		}
-
-		// TODO: is this a correct assumption, that the max component should be used as value?
-		OutValue = Value.GetMax();
-		return true;
-	}
-
-	if (const UMaterialExpressionScalarParameter* ScalarParameter = ExactCast<UMaterialExpressionScalarParameter>(Expression))
-	{
-		float Value = ScalarParameter->DefaultValue;
-
-		const UMaterialInstance* MaterialInstance = Cast<UMaterialInstance>(Material);
-		if (MaterialInstance != nullptr)
-		{
-			const FHashedMaterialParameterInfo ParameterInfo(ScalarParameter->GetParameterName());
-			if (!MaterialInstance->GetScalarParameterValue(ParameterInfo, Value))
-			{
-				// TODO: how to handle this?
-			}
-		}
-
-		OutValue = Value;
-		return true;
-	}
-
-	if (const UMaterialExpressionConstant4Vector* Constant4Vector = ExactCast<UMaterialExpressionConstant4Vector>(Expression))
-	{
-		OutValue = Constant4Vector->Constant.R;
-		return true;
-	}
-
-	if (const UMaterialExpressionConstant3Vector* Constant3Vector = ExactCast<UMaterialExpressionConstant3Vector>(Expression))
-	{
-		OutValue = Constant3Vector->Constant.R;
-		return true;
-	}
-
-	if (const UMaterialExpressionConstant2Vector* Constant2Vector = ExactCast<UMaterialExpressionConstant2Vector>(Expression))
-	{
-		OutValue = Constant2Vector->R;
-		return true;
-	}
-
-	if (const UMaterialExpressionConstant* Constant = ExactCast<UMaterialExpressionConstant>(Expression))
-	{
-		OutValue = Constant->R;
-		return true;
-	}
-
-	return false;
-}
-
-bool FGLTFMaterialConverter::TryGetSourceTexture(FGLTFJsonTextureInfo& OutTexInfo, EMaterialProperty Property, const UMaterialInterface* Material, const TArray<FLinearColor>& AllowedMasks) const
-{
-	const UTexture2D* Texture;
-	int32 TexCoord;
-
-	if (TryGetSourceTexture(Texture, TexCoord, Property, Material, AllowedMasks))
-	{
-		OutTexInfo.Index = Builder.GetOrAddTexture(Texture);
-		OutTexInfo.TexCoord = TexCoord;
-		return true;
-	}
-
-	return false;
-}
-
-bool FGLTFMaterialConverter::TryGetSourceTexture(const UTexture2D*& OutTexture, int32& OutTexCoord, EMaterialProperty Property, const UMaterialInterface* Material, const TArray<FLinearColor>& AllowedMasks) const
-{
-	const FExpressionInput* MaterialInput = FGLTFMaterialUtility::GetInputForProperty(Material, Property);
-	if (MaterialInput == nullptr)
-	{
-		// TODO: report error
-		return false;
-	}
-
-	const UMaterialExpression* Expression = MaterialInput->Expression;
-	if (Expression == nullptr)
-	{
-		return false;
-	}
-
-	const FLinearColor InputMask = FGLTFMaterialUtility::GetMask(*MaterialInput);
-	if (AllowedMasks.Num() > 0 && !AllowedMasks.Contains(InputMask))
-	{
-		return false;
-	}
-
-	// TODO: add support or warning for texture sampler settings that override texture asset addressing (i.e. wrap, clamp etc)?
-
-	if (const UMaterialExpressionTextureSampleParameter2D* TextureParameter = ExactCast<UMaterialExpressionTextureSampleParameter2D>(Expression))
-	{
-		UTexture* ParameterValue = TextureParameter->Texture;
-
-		const UMaterialInstance* MaterialInstance = Cast<UMaterialInstance>(Material);
-		if (MaterialInstance != nullptr)
-		{
-			const FHashedMaterialParameterInfo ParameterInfo(TextureParameter->GetParameterName());
-			if (!MaterialInstance->GetTextureParameterValue(ParameterInfo, ParameterValue))
-			{
-				// TODO: how to handle this?
-			}
-		}
-
-		OutTexture = Cast<UTexture2D>(ParameterValue);
-
-		if (OutTexture == nullptr)
-		{
-			if (ParameterValue == nullptr)
-			{
-				// TODO: report error (no texture parameter assigned)
-			}
-			else
-			{
-				// TODO: report error (incorrect texture type)
-			}
-			return false;
-		}
-
-		if (!FGLTFMaterialUtility::TryGetTextureCoordinateIndex(TextureParameter, OutTexCoord))
-		{
-			// TODO: report error (failed to identify texture coordinate index)
-			return false;
-		}
-
-		return true;
-	}
-
-	if (const UMaterialExpressionTextureSample* TextureSampler = ExactCast<UMaterialExpressionTextureSample>(Expression))
-	{
-		// TODO: add support for texture object input expression
-		OutTexture = Cast<UTexture2D>(TextureSampler->Texture);
-
-		if (OutTexture == nullptr)
-		{
-			if (TextureSampler->Texture == nullptr)
-			{
-				// TODO: report error (no texture sample assigned)
-			}
-			else
-			{
-				// TODO: report error (incorrect texture type)
-			}
-			return false;
-		}
-
-		if (!FGLTFMaterialUtility::TryGetTextureCoordinateIndex(TextureSampler, OutTexCoord))
-		{
-			// TODO: report error (failed to identify texture coordinate index)
-			return false;
-		}
-
-		return true;
-	}
-
-	return false;
-}
-
-bool FGLTFMaterialConverter::TryGetBakedMaterialProperty(FGLTFJsonTextureInfo& OutTexInfo, FGLTFJsonColor3& OutConstant, EMaterialProperty Property, const FString& PropertyName, const UMaterialInterface* Material) const
-{
-	if (!Builder.ExportOptions->bBakeMaterialInputs)
-	{
-		Builder.AddWarningMessage(FString::Printf(
-			TEXT("%s for material %s needs to bake, but material baking disabled by export options"),
-			FGLTFMaterialUtility::GetPropertyName(Property),
-			*Material->GetName()));
-		return false;
-	}
-
-	const FGLTFPropertyBakeOutput PropertyBakeOutput = BakeMaterialProperty(Property, Material, OutTexInfo.TexCoord);
-
-	if (PropertyBakeOutput.bIsConstant)
-	{
-		OutConstant = FGLTFConverterUtility::ConvertColor(PropertyBakeOutput.ConstantValue);
-		return true;
-	}
-
-	if (StoreBakedPropertyTexture(OutTexInfo, PropertyBakeOutput, PropertyName, Material))
-	{
-		OutConstant = FGLTFJsonColor3::White; // make sure property is not zero
-		return true;
-	}
-
-	return false;
-}
-
-bool FGLTFMaterialConverter::TryGetBakedMaterialProperty(FGLTFJsonTextureInfo& OutTexInfo, FGLTFJsonColor4& OutConstant, EMaterialProperty Property, const FString& PropertyName, const UMaterialInterface* Material) const
-{
-	if (!Builder.ExportOptions->bBakeMaterialInputs)
-	{
-		Builder.AddWarningMessage(FString::Printf(
-			TEXT("%s for material %s needs to bake, but material baking disabled by export options"),
-			FGLTFMaterialUtility::GetPropertyName(Property),
-			*Material->GetName()));
-		return false;
-	}
-
-	const FGLTFPropertyBakeOutput PropertyBakeOutput = BakeMaterialProperty(Property, Material, OutTexInfo.TexCoord);
-
-	if (PropertyBakeOutput.bIsConstant)
-	{
-		OutConstant = FGLTFConverterUtility::ConvertColor(PropertyBakeOutput.ConstantValue);
-		return true;
-	}
-
-	if (StoreBakedPropertyTexture(OutTexInfo, PropertyBakeOutput, PropertyName, Material))
-	{
-		OutConstant = FGLTFJsonColor4::White; // make sure property is not zero
-		return true;
-	}
-
-	return false;
-}
-
-inline bool FGLTFMaterialConverter::TryGetBakedMaterialProperty(FGLTFJsonTextureInfo& OutTexInfo, float& OutConstant, EMaterialProperty Property, const FString& PropertyName, const UMaterialInterface* Material) const
-{
-	if (!Builder.ExportOptions->bBakeMaterialInputs)
-	{
-		Builder.AddWarningMessage(FString::Printf(
-			TEXT("%s for material %s needs to bake, but material baking disabled by export options"),
-			FGLTFMaterialUtility::GetPropertyName(Property),
-			*Material->GetName()));
-		return false;
-	}
-
-	const FGLTFPropertyBakeOutput PropertyBakeOutput = BakeMaterialProperty(Property, Material, OutTexInfo.TexCoord);
-
-	if (PropertyBakeOutput.bIsConstant)
-	{
-		OutConstant = PropertyBakeOutput.ConstantValue.R;
-		return true;
-	}
-
-	if (StoreBakedPropertyTexture(OutTexInfo, PropertyBakeOutput, PropertyName, Material))
-	{
-		OutConstant = 1; // make sure property is not zero
-		return true;
-	}
-
-	return false;
-}
-
-bool FGLTFMaterialConverter::TryGetBakedMaterialProperty(FGLTFJsonTextureInfo& OutTexInfo, EMaterialProperty Property, const FString& PropertyName, const UMaterialInterface* Material) const
-{
-	if (!Builder.ExportOptions->bBakeMaterialInputs)
-	{
-		Builder.AddWarningMessage(FString::Printf(
-			TEXT("%s for material %s needs to bake, but material baking disabled by export options"),
-			FGLTFMaterialUtility::GetPropertyName(Property),
-			*Material->GetName()));
-		return false;
-	}
-
-	const FGLTFPropertyBakeOutput PropertyBakeOutput = BakeMaterialProperty(Property, Material, OutTexInfo.TexCoord);
-
-	if (!PropertyBakeOutput.bIsConstant)
-	{
-		return StoreBakedPropertyTexture(OutTexInfo, PropertyBakeOutput, PropertyName, Material);
-	}
-
-	const FVector4 MaskedConstant = FVector4(PropertyBakeOutput.ConstantValue) * FGLTFMaterialUtility::GetPropertyMask(Property);
-	if (MaskedConstant == FGLTFMaterialUtility::GetPropertyDefaultValue(Property))
-	{
-		// Constant value is the same as the property's default so we can set gltf to default.
-		OutTexInfo.Index = FGLTFJsonTextureIndex(INDEX_NONE);
-		return true;
-	}
-
-	if (Property == MP_Normal || Property == MP_CustomOutput /* ClearCoatBottomNormal */)
-	{
-		// TODO: In some cases baking normal can result in constant vector that differs slight from default (i.e 0,0,1).
-		// Yet often, when looking at such a material, it should be exactly default. Needs further investigation.
-		// Maybe because of incorrect sRGB conversion? For now, assume a constant normal is always default.
-		OutTexInfo.Index = FGLTFJsonTextureIndex(INDEX_NONE);
-		return true;
-	}
-
-	// TODO: let function fail and investigate why in some cases a constant baking result is returned for a property
-	// that is non-constant. This happens (for example) when baking AmbientOcclusion for a translucent material,
-	// even though the same material when set to opaque will properly bake AmbientOcclusion to a texture.
-	// For now, create a 1x1 texture with the constant value.
-
-	const FString TextureName = Material->GetName() + TEXT("_") + PropertyName;
-	const FGLTFJsonTextureIndex TextureIndex = FGLTFMaterialUtility::AddTexture(
-		Builder,
-		PropertyBakeOutput.Pixels,
-		PropertyBakeOutput.Size,
-		Material->GetName() + TEXT("_") + PropertyName,
-		EGLTFJsonTextureFilter::Nearest,
-		EGLTFJsonTextureFilter::Nearest,
-		EGLTFJsonTextureWrap::ClampToEdge,
-		EGLTFJsonTextureWrap::ClampToEdge);
-
-	OutTexInfo.Index = TextureIndex;
-	return true;
-}
-
-FGLTFPropertyBakeOutput FGLTFMaterialConverter::BakeMaterialProperty(EMaterialProperty Property, const UMaterialInterface* Material, int32& OutTexCoord, const FIntPoint* PreferredTextureSize, bool bCopyAlphaFromRedChannel) const
-{
-	const FExpressionInput* PropertyInput = FGLTFMaterialUtility::GetInputForProperty(Material, Property);
-	TSet<int32> TexCoords;
-
-	FGLTFMaterialUtility::GetAllTextureCoordinateIndices(*PropertyInput, TexCoords);
-	if (TexCoords.Num() > 0)
-	{
-		// TODO: is there a nicer way to get the first element in a TSet?
-		for (int32 TexCoord : TexCoords)
-		{
-			OutTexCoord = TexCoord;
-			break;
-		}
-
-		if (TexCoords.Num() > 1)
-		{
-			// TODO: report warning (multiple texture coordinates found, will use first)
-
-			// TODO: replace this hardcoded hack with something more configurable and proper
-			if (Property == MP_AmbientOcclusion && TexCoords.Contains(1))
-			{
-				OutTexCoord = 1; // assume ambient occlusion uses TexCoord1 when multiple
-			}
-		}
-	}
-	else
-	{
-		OutTexCoord = 0; // assume TexCoord0 even thought property seems to be texture coordinate independent
-	}
-
-	const FIntPoint DefaultTextureSize = Builder.ExportOptions->GetDefaultMaterialBakeSize();
-	const FIntPoint TextureSize = PreferredTextureSize != nullptr ? *PreferredTextureSize : DefaultTextureSize;
-
-	// TODO: add support for calculating the ideal resolution to use for baking based on connected (texture) nodes
-
-	const FGLTFPropertyBakeOutput PropertyBakeOutput = FGLTFMaterialUtility::BakeMaterialProperty(
-		TextureSize,
-		Property,
-		Material,
-		OutTexCoord,
-		bCopyAlphaFromRedChannel);
-
-	if (!PropertyBakeOutput.bIsConstant && TexCoords.Num() == 0)
-	{
-		// TODO: report warning about property not being constant yet texture coordinate independent
-	}
-
-	return PropertyBakeOutput;
-}
-
-bool FGLTFMaterialConverter::StoreBakedPropertyTexture(FGLTFJsonTextureInfo& OutTexInfo, const FGLTFPropertyBakeOutput& PropertyBakeOutput, const FString& PropertyName, const UMaterialInterface* Material) const
-{
-	const FString TextureName = Material->GetName() + TEXT("_") + PropertyName;
-
-	// TODO: should this be the default wrap-mode?
-	const EGLTFJsonTextureWrap TextureWrapS = EGLTFJsonTextureWrap::Repeat;
-	const EGLTFJsonTextureWrap TextureWrapT = EGLTFJsonTextureWrap::Repeat;
-
-	// TODO: should this be the default filter?
-	const EGLTFJsonTextureFilter TextureMinFilter = EGLTFJsonTextureFilter::LinearMipmapLinear;
-	const EGLTFJsonTextureFilter TextureMagFilter = EGLTFJsonTextureFilter::Linear;
-
-	const FGLTFJsonTextureIndex TextureIndex = FGLTFMaterialUtility::AddTexture(
-		Builder,
-		PropertyBakeOutput.Pixels,
-		PropertyBakeOutput.Size,
-		TextureName,
-		TextureMinFilter,
-		TextureMagFilter,
-		TextureWrapS,
-		TextureWrapT);
-
-	OutTexInfo.Index = TextureIndex;
-	return true;
+	const FGLTFJsonMaterialIndex MaterialIndex = Builder.AddMaterial();
+	Builder.SetupTask<FGLTFMaterialTask>(Builder, Material, MaterialIndex);
+	return MaterialIndex;
 }
diff --git a/Engine/Plugins/Enterprise/GLTFExporter/MigratedFromGit/Source/GLTFExporter/Private/Converters/GLTFMaterialConverters.h b/Engine/Plugins/Enterprise/GLTFExporter/MigratedFromGit/Source/GLTFExporter/Private/Converters/GLTFMaterialConverters.h
--- a/Engine/Plugins/Enterprise/GLTFExporter/MigratedFromGit/Source/GLTFExporter/Private/Converters/GLTFMaterialConverters.h
+++ b/Engine/Plugins/Enterprise/GLTFExporter/MigratedFromGit/Source/GLTFExporter/Private/Converters/GLTFMaterialConverters.h
@@ -3,41 +3,13 @@
 #pragma once
 
 #include "Json/GLTFJsonIndex.h"
-#include "Json/GLTFJsonMaterial.h"
 #include "Converters/GLTFConverter.h"
 #include "Converters/GLTFBuilderContext.h"
 #include "Engine.h"
 
-struct FGLTFPropertyBakeOutput;
-
 class FGLTFMaterialConverter final : public FGLTFBuilderContext, public TGLTFConverter<FGLTFJsonMaterialIndex, const UMaterialInterface*>
 {
 	using FGLTFBuilderContext::FGLTFBuilderContext;
 
 	virtual FGLTFJsonMaterialIndex Convert(const UMaterialInterface* Material) override;
-
-	bool TryGetShadingModel(EGLTFJsonShadingModel& ShadingModel, const UMaterialInterface* Material) const;
-
-	bool TryGetBaseColorAndOpacity(FGLTFJsonPBRMetallicRoughness& OutPBRParams, const UMaterialInterface* Material, EMaterialProperty BaseColorProperty, EMaterialProperty OpacityProperty) const;
-	bool TryGetMetallicAndRoughness(FGLTFJsonPBRMetallicRoughness& OutPBRParams, const UMaterialInterface* Material, EMaterialProperty MetallicProperty, EMaterialProperty RoughnessProperty) const;
-	bool TryGetClearCoatRoughness(FGLTFJsonClearCoatExtension& OutExtParams, const UMaterialInterface* Material, EMaterialProperty IntensityProperty, EMaterialProperty RoughnessProperty) const;
-	bool TryGetEmissive(FGLTFJsonMaterial& JsonMaterial, EMaterialProperty EmissiveProperty, const UMaterialInterface* Material) const;
-
-	bool IsPropertyNonDefault(EMaterialProperty Property, const UMaterialInterface* Material) const;
-	bool TryGetConstantColor(FGLTFJsonColor3& OutValue, EMaterialProperty Property, const UMaterialInterface* Material) const;
-	bool TryGetConstantColor(FGLTFJsonColor4& OutValue, EMaterialProperty Property, const UMaterialInterface* Material) const;
-	bool TryGetConstantColor(FLinearColor& OutValue, EMaterialProperty Property, const UMaterialInterface* Material) const;
-	bool TryGetConstantScalar(float& OutValue, EMaterialProperty Property, const UMaterialInterface* Material) const;
-
-	bool TryGetSourceTexture(FGLTFJsonTextureInfo& OutTexInfo, EMaterialProperty Property, const UMaterialInterface* Material, const TArray<FLinearColor>& AllowedMasks = {}) const;
-	bool TryGetSourceTexture(const UTexture2D*& OutTexture, int32& OutTexCoord, EMaterialProperty Property, const UMaterialInterface* Material, const TArray<FLinearColor>& AllowedMasks = {}) const;
-
-	bool TryGetBakedMaterialProperty(FGLTFJsonTextureInfo& OutTexInfo, FGLTFJsonColor3& OutConstant, EMaterialProperty Property, const FString& PropertyName, const UMaterialInterface* Material) const;
-	bool TryGetBakedMaterialProperty(FGLTFJsonTextureInfo& OutTexInfo, FGLTFJsonColor4& OutConstant, EMaterialProperty Property, const FString& PropertyName, const UMaterialInterface* Material) const;
-	bool TryGetBakedMaterialProperty(FGLTFJsonTextureInfo& OutTexInfo, float& OutConstant, EMaterialProperty Property, const FString& PropertyName, const UMaterialInterface* Material) const;
-	bool TryGetBakedMaterialProperty(FGLTFJsonTextureInfo& OutTexInfo, EMaterialProperty Property, const FString& PropertyName, const UMaterialInterface* Material) const;
-
-	FGLTFPropertyBakeOutput BakeMaterialProperty(EMaterialProperty Property, const UMaterialInterface* Material, int32& OutTexCoord, const FIntPoint* PreferredTextureSize = nullptr, bool bCopyAlphaFromRedChannel = false) const;
-
-	bool StoreBakedPropertyTexture(FGLTFJsonTextureInfo& OutTexInfo, const FGLTFPropertyBakeOutput& PropertyBakeOutput, const FString& PropertyName, const UMaterialInterface* Material) const;
 };
diff --git a/Engine/Plugins/Enterprise/GLTFExporter/MigratedFromGit/Source/GLTFExporter/Private/Tasks/GLTFMaterialTasks.cpp b/Engine/Plugins/Enterprise/GLTFExporter/MigratedFromGit/Source/GLTFExporter/Private/Tasks/GLTFMaterialTasks.cpp
new file mode 100644
--- /dev/null
+++ b/Engine/Plugins/Enterprise/GLTFExporter/MigratedFromGit/Source/GLTFExporter/Private/Tasks/GLTFMaterialTasks.cpp
@@ -0,0 +1,1332 @@
+// Copyright Epic Games, Inc. All Rights Reserved.
+
+#include "Tasks/GLTFMaterialTasks.h"
+#include "Converters/GLTFConverterUtility.h"
+#include "Converters/GLTFNameUtility.h"
+#include "Converters/GLTFMaterialUtility.h"
+#include "Builders/GLTFContainerBuilder.h"
+#include "Materials/MaterialExpressionConstant.h"
+#include "Materials/MaterialExpressionConstant2Vector.h"
+#include "Materials/MaterialExpressionConstant3Vector.h"
+#include "Materials/MaterialExpressionConstant4Vector.h"
+#include "Materials/MaterialExpressionScalarParameter.h"
+#include "Materials/MaterialExpressionVectorParameter.h"
+#include "Materials/MaterialExpressionTextureSample.h"
+#include "Materials/MaterialExpressionTextureSampleParameter2D.h"
+
+namespace
+{
+	// Component masks
+	const FLinearColor RedMask(1.0f, 0.0f, 0.0f, 0.0f);
+	const FLinearColor GreenMask(0.0f, 1.0f, 0.0f, 0.0f);
+	const FLinearColor BlueMask(0.0f, 0.0f, 1.0f, 0.0f);
+	const FLinearColor AlphaMask(0.0f, 0.0f, 0.0f, 1.0f);
+	const FLinearColor RgbMask = RedMask + GreenMask + BlueMask;
+	const FLinearColor RgbaMask = RgbMask + AlphaMask;
+
+	// Property-specific component masks
+	const FLinearColor BaseColorMask = RgbMask;
+	const FLinearColor OpacityMask = AlphaMask;
+	const FLinearColor MetallicMask = BlueMask;
+	const FLinearColor RoughnessMask = GreenMask;
+	const FLinearColor OcclusionMask = RedMask;
+	const FLinearColor ClearCoatMask = RedMask;
+	const FLinearColor ClearCoatRoughnessMask = GreenMask;
+
+	// Ideal masks for texture-inputs (doesn't require baking)
+	const TArray<FLinearColor> DefaultColorInputMasks = { RgbMask, RgbaMask };
+	const TArray<FLinearColor> BaseColorInputMasks = { BaseColorMask };
+	const TArray<FLinearColor> OpacityInputMasks = { OpacityMask };
+	const TArray<FLinearColor> MetallicInputMasks = { MetallicMask };
+	const TArray<FLinearColor> RoughnessInputMasks = { RoughnessMask };
+	const TArray<FLinearColor> OcclusionInputMasks = { OcclusionMask };
+	const TArray<FLinearColor> ClearCoatInputMasks = { ClearCoatMask };
+	const TArray<FLinearColor> ClearCoatRoughnessInputMasks = { ClearCoatRoughnessMask };
+}
+
+void FGLTFMaterialTask::Complete()
+{
+	{
+		const UMaterial* ParentMaterial = Material->GetMaterial();
+
+		if (ParentMaterial->MaterialDomain != MD_Surface)
+		{
+			// TODO: report warning (non-surface materials not supported, will be treated as surface)
+		}
+
+		if (!ParentMaterial->bTangentSpaceNormal)
+		{
+			// TODO: report warning (world-space normals not supported, will be treated as tangent-space)
+		}
+	}
+
+	FGLTFJsonMaterial& JsonMaterial = Builder.GetMaterial(MaterialIndex);
+	Material->GetName(JsonMaterial.Name);
+
+	// TODO: add support for additional blend modes (like Additive and Modulate)?
+	JsonMaterial.AlphaMode = FGLTFConverterUtility::ConvertBlendMode(Material->GetBlendMode());
+	JsonMaterial.AlphaCutoff = Material->GetOpacityMaskClipValue();
+	JsonMaterial.DoubleSided = Material->IsTwoSided();
+
+	if (!TryGetShadingModel(JsonMaterial.ShadingModel))
+	{
+		JsonMaterial.ShadingModel = EGLTFJsonShadingModel::Default;
+		Builder.AddWarningMessage(FString::Printf(TEXT("Material %s will be exported as shading model %s"), *Material->GetName(), *FGLTFNameUtility::GetName(MSM_DefaultLit)));
+	}
+
+	if (JsonMaterial.ShadingModel != EGLTFJsonShadingModel::None)
+	{
+		const EMaterialProperty BaseColorProperty = JsonMaterial.ShadingModel == EGLTFJsonShadingModel::Unlit ? MP_EmissiveColor : MP_BaseColor;
+		const EMaterialProperty OpacityProperty = JsonMaterial.AlphaMode == EGLTFJsonAlphaMode::Mask ? MP_OpacityMask : MP_Opacity;
+
+		// TODO: check if a property is active before trying to get it (i.e. Material->IsPropertyActive)
+
+		if (JsonMaterial.AlphaMode == EGLTFJsonAlphaMode::Opaque)
+		{
+			if (!TryGetConstantColor(JsonMaterial.PBRMetallicRoughness.BaseColorFactor, BaseColorProperty))
+			{
+				if (!TryGetSourceTexture(JsonMaterial.PBRMetallicRoughness.BaseColorTexture, BaseColorProperty, DefaultColorInputMasks))
+				{
+					if (!TryGetBakedMaterialProperty(JsonMaterial.PBRMetallicRoughness.BaseColorTexture, JsonMaterial.PBRMetallicRoughness.BaseColorFactor, BaseColorProperty, TEXT("BaseColor")))
+					{
+						Builder.AddWarningMessage(FString::Printf(TEXT("Failed to export %s for material %s"), FGLTFMaterialUtility::GetPropertyName(BaseColorProperty), *Material->GetName()));
+					}
+				}
+			}
+
+			JsonMaterial.PBRMetallicRoughness.BaseColorFactor.A = 1.0f; // make sure base color is opaque
+		}
+		else
+		{
+			if (!TryGetBaseColorAndOpacity(JsonMaterial.PBRMetallicRoughness, BaseColorProperty, OpacityProperty))
+			{
+				Builder.AddWarningMessage(FString::Printf(TEXT("Failed to export %s and %s for material %s"), FGLTFMaterialUtility::GetPropertyName(BaseColorProperty), FGLTFMaterialUtility::GetPropertyName(OpacityProperty), *Material->GetName()));
+			}
+		}
+
+		if (JsonMaterial.ShadingModel == EGLTFJsonShadingModel::Default || JsonMaterial.ShadingModel == EGLTFJsonShadingModel::ClearCoat)
+		{
+			const EMaterialProperty MetallicProperty = MP_Metallic;
+			const EMaterialProperty RoughnessProperty = MP_Roughness;
+
+			if (!TryGetMetallicAndRoughness(JsonMaterial.PBRMetallicRoughness, MetallicProperty, RoughnessProperty))
+			{
+				Builder.AddWarningMessage(FString::Printf(TEXT("Failed to export %s and %s for material %s"), FGLTFMaterialUtility::GetPropertyName(MetallicProperty), FGLTFMaterialUtility::GetPropertyName(RoughnessProperty), *Material->GetName()));
+			}
+
+			const EMaterialProperty EmissiveProperty = MP_EmissiveColor;
+			if (!TryGetEmissive(JsonMaterial, EmissiveProperty))
+			{
+				Builder.AddWarningMessage(FString::Printf(TEXT("Failed to export %s for material %s"), FGLTFMaterialUtility::GetPropertyName(EmissiveProperty), *Material->GetName()));
+			}
+
+			// TODO: replace dummy enum MP_CustomOutput workaround for ClearCoatBottomNormal with proper support for custom outputs
+			const EMaterialProperty NormalProperty = JsonMaterial.ShadingModel == EGLTFJsonShadingModel::ClearCoat ? MP_CustomOutput : MP_Normal;
+			if (IsPropertyNonDefault(NormalProperty))
+			{
+				if (!TryGetSourceTexture(JsonMaterial.NormalTexture, NormalProperty, DefaultColorInputMasks))
+				{
+					if (!TryGetBakedMaterialProperty(JsonMaterial.NormalTexture, NormalProperty, TEXT("Normal")))
+					{
+						Builder.AddWarningMessage(FString::Printf(TEXT("Failed to export %s for material %s"), FGLTFMaterialUtility::GetPropertyName(NormalProperty), *Material->GetName()));
+					}
+				}
+			}
+
+			const EMaterialProperty AmbientOcclusionProperty = MP_AmbientOcclusion;
+			if (IsPropertyNonDefault(AmbientOcclusionProperty))
+			{
+				if (!TryGetSourceTexture(JsonMaterial.OcclusionTexture, AmbientOcclusionProperty, OcclusionInputMasks))
+				{
+					if (!TryGetBakedMaterialProperty(JsonMaterial.OcclusionTexture, AmbientOcclusionProperty, TEXT("Occlusion")))
+					{
+						Builder.AddWarningMessage(FString::Printf(TEXT("Failed to export %s for material %s"), FGLTFMaterialUtility::GetPropertyName(AmbientOcclusionProperty), *Material->GetName()));
+					}
+				}
+			}
+
+			if (JsonMaterial.ShadingModel == EGLTFJsonShadingModel::ClearCoat)
+			{
+				const EMaterialProperty ClearCoatProperty = MP_CustomData0;
+				const EMaterialProperty ClearCoatRoughnessProperty = MP_CustomData1;
+
+				if (!TryGetClearCoatRoughness(JsonMaterial.ClearCoat, ClearCoatProperty, ClearCoatRoughnessProperty))
+				{
+					Builder.AddWarningMessage(FString::Printf(TEXT("Failed to export %s and %s for material %s"), FGLTFMaterialUtility::GetPropertyName(ClearCoatProperty), FGLTFMaterialUtility::GetPropertyName(ClearCoatRoughnessProperty), *Material->GetName()));
+				}
+
+				const EMaterialProperty ClearCoatNormalProperty = MP_Normal;
+				if (IsPropertyNonDefault(ClearCoatNormalProperty))
+				{
+					if (!TryGetSourceTexture(JsonMaterial.ClearCoat.ClearCoatNormalTexture, ClearCoatNormalProperty, DefaultColorInputMasks))
+					{
+						if (!TryGetBakedMaterialProperty(JsonMaterial.ClearCoat.ClearCoatNormalTexture, ClearCoatNormalProperty, TEXT("ClearCoatNormal")))
+						{
+							Builder.AddWarningMessage(FString::Printf(TEXT("Failed to export %s for material %s"), FGLTFMaterialUtility::GetPropertyName(ClearCoatNormalProperty), *Material->GetName()));
+						}
+					}
+				}
+			}
+		}
+	}
+}
+
+bool FGLTFMaterialTask::TryGetShadingModel(EGLTFJsonShadingModel& OutShadingModel) const
+{
+	const FMaterialShadingModelField ShadingModels = Material->GetShadingModels();
+	const int32 ShadingModelCount = ShadingModels.CountShadingModels();
+
+	if (ShadingModelCount <= 0)
+	{
+		Builder.AddWarningMessage(FString::Printf(TEXT("No shading model found for material %s"), *Material->GetName()));
+		return false;
+	}
+
+	if (ShadingModelCount > 1)
+	{
+		Builder.AddWarningMessage(FString::Printf(TEXT("Support is limited to the first of multiple shading models found (%d) in material %s"), ShadingModelCount, *Material->GetName()));
+	}
+
+	const EMaterialShadingModel ShadingModel = ShadingModels.GetFirstShadingModel();
+	const EGLTFJsonShadingModel ConvertedShadingModel = FGLTFConverterUtility::ConvertShadingModel(ShadingModel);
+
+	if (ConvertedShadingModel == EGLTFJsonShadingModel::None)
+	{
+		const FString ShadingModelName = FGLTFNameUtility::GetName(ShadingModel);
+		Builder.AddWarningMessage(FString::Printf(TEXT("Unsupported shading model (%s) in material %s"), *ShadingModelName, *Material->GetName()));
+		return false;
+	}
+
+	if (ConvertedShadingModel == EGLTFJsonShadingModel::Unlit && !Builder.ExportOptions->bExportUnlitMaterials)
+	{
+		const FString ShadingModelName = FGLTFNameUtility::GetName(ShadingModel);
+		Builder.AddWarningMessage(FString::Printf(TEXT("Shading model (%s) in material %s disabled by export options"), *ShadingModelName, *Material->GetName()));
+		return false;
+	}
+
+	if (ConvertedShadingModel == EGLTFJsonShadingModel::ClearCoat && !Builder.ExportOptions->bExportClearCoatMaterials)
+	{
+		const FString ShadingModelName = FGLTFNameUtility::GetName(ShadingModel);
+		Builder.AddWarningMessage(FString::Printf(TEXT("Shading model (%s) in material %s disabled by export options"), *ShadingModelName, *Material->GetName()));
+		return false;
+	}
+
+	OutShadingModel = ConvertedShadingModel;
+	return true;
+}
+
+bool FGLTFMaterialTask::TryGetBaseColorAndOpacity(FGLTFJsonPBRMetallicRoughness& OutPBRParams, EMaterialProperty BaseColorProperty, EMaterialProperty OpacityProperty) const
+{
+	const bool bIsBaseColorConstant = TryGetConstantColor(OutPBRParams.BaseColorFactor, BaseColorProperty);
+	const bool bIsOpacityConstant = TryGetConstantScalar(OutPBRParams.BaseColorFactor.A, OpacityProperty);
+
+	if (bIsBaseColorConstant && bIsOpacityConstant)
+	{
+		return true;
+	}
+
+	// NOTE: since we always bake the properties (for now) when atleast property is non-const, we need
+	// to reset the constant factors to their defaults. Otherwise the baked value of a constant property
+	// would be scaled with the factor, i.e a double scaling.
+	OutPBRParams.BaseColorFactor = { 1.0f, 1.0f, 1.0f, 1.0f };
+
+	const UTexture2D* BaseColorTexture;
+	const UTexture2D* OpacityTexture;
+	int32 BaseColorTexCoord;
+	int32 OpacityTexCoord;
+
+	const bool bHasBaseColorSourceTexture = TryGetSourceTexture(BaseColorTexture, BaseColorTexCoord, BaseColorProperty, BaseColorInputMasks);
+	const bool bHasOpacitySourceTexture = TryGetSourceTexture(OpacityTexture, OpacityTexCoord, OpacityProperty, OpacityInputMasks);
+
+	// Detect the "happy path" where both inputs share the same texture and are correctly masked.
+	if (bHasBaseColorSourceTexture &&
+		bHasOpacitySourceTexture &&
+		BaseColorTexture == OpacityTexture &&
+		BaseColorTexCoord == OpacityTexCoord)
+	{
+		OutPBRParams.BaseColorTexture.Index = Builder.GetOrAddTexture(BaseColorTexture);
+		OutPBRParams.BaseColorTexture.TexCoord = BaseColorTexCoord;
+		return true;
+	}
+
+	if (!Builder.ExportOptions->bBakeMaterialInputs)
+	{
+		Builder.AddWarningMessage(FString::Printf(
+			TEXT("%s and %s for material %s needs to bake, but material baking disabled by export options"),
+			FGLTFMaterialUtility::GetPropertyName(BaseColorProperty),
+			FGLTFMaterialUtility::GetPropertyName(OpacityProperty),
+			*Material->GetName()));
+		return false;
+	}
+
+	// TODO: add support for calculating the ideal resolution to use for baking based on connected (texture) nodes
+	FIntPoint TextureSize = Builder.ExportOptions->GetDefaultMaterialBakeSize();
+
+	// TODO: should this be the default wrap-mode?
+	EGLTFJsonTextureWrap TextureWrapS = EGLTFJsonTextureWrap::Repeat;
+	EGLTFJsonTextureWrap TextureWrapT = EGLTFJsonTextureWrap::Repeat;
+
+	// TODO: should this be the default filter?
+	EGLTFJsonTextureFilter TextureMinFilter = EGLTFJsonTextureFilter::LinearMipmapLinear;
+	EGLTFJsonTextureFilter TextureMagFilter = EGLTFJsonTextureFilter::Linear;
+
+	if (bHasBaseColorSourceTexture && bHasOpacitySourceTexture)
+	{
+		const bool bAreTexturesCompatible = BaseColorTexCoord == OpacityTexCoord &&
+			BaseColorTexture->AddressX == OpacityTexture->AddressX &&
+			BaseColorTexture->AddressY == OpacityTexture->AddressY;
+
+		if (!bAreTexturesCompatible)
+		{
+			// TODO: handle differences in wrapping or uv-coords
+			Builder.AddWarningMessage(FString::Printf(
+				TEXT("BaseColor- and Opacity-textures for material %s were not able to be combined and will be skipped"),
+				*Material->GetName()));
+
+			return false;
+		}
+
+		TextureSize =
+		{
+			FMath::Max(BaseColorTexture->GetSizeX(), OpacityTexture->GetSizeX()),
+			FMath::Max(BaseColorTexture->GetSizeY(), OpacityTexture->GetSizeY())
+		};
+
+		TextureWrapS = FGLTFConverterUtility::ConvertWrap(BaseColorTexture->AddressX);
+		TextureWrapT = FGLTFConverterUtility::ConvertWrap(BaseColorTexture->AddressY);
+
+		// TODO: compare min- and mag-filter for BaseColorTexture and OpacityTexture. If they differ,
+		// we should choose one or the other and inform the user about the choice made by logging to the console.
+		TextureMinFilter = FGLTFConverterUtility::ConvertMinFilter(BaseColorTexture->Filter, BaseColorTexture->LODGroup);
+		TextureMagFilter = FGLTFConverterUtility::ConvertMagFilter(BaseColorTexture->Filter, BaseColorTexture->LODGroup);
+	}
+	else if (bHasBaseColorSourceTexture)
+	{
+		TextureSize = { BaseColorTexture->GetSizeX(), BaseColorTexture->GetSizeY() };
+		TextureWrapS = FGLTFConverterUtility::ConvertWrap(BaseColorTexture->AddressX);
+		TextureWrapT = FGLTFConverterUtility::ConvertWrap(BaseColorTexture->AddressY);
+		TextureMinFilter = FGLTFConverterUtility::ConvertMinFilter(BaseColorTexture->Filter, BaseColorTexture->LODGroup);
+		TextureMagFilter = FGLTFConverterUtility::ConvertMagFilter(BaseColorTexture->Filter, BaseColorTexture->LODGroup);
+	}
+	else if (bHasOpacitySourceTexture)
+	{
+		TextureSize = { OpacityTexture->GetSizeX(), OpacityTexture->GetSizeY() };
+		TextureWrapS = FGLTFConverterUtility::ConvertWrap(OpacityTexture->AddressX);
+		TextureWrapT = FGLTFConverterUtility::ConvertWrap(OpacityTexture->AddressY);
+		TextureMinFilter = FGLTFConverterUtility::ConvertMinFilter(OpacityTexture->Filter, OpacityTexture->LODGroup);
+		TextureMagFilter = FGLTFConverterUtility::ConvertMagFilter(OpacityTexture->Filter, OpacityTexture->LODGroup);
+	}
+
+	const FGLTFPropertyBakeOutput BaseColorBakeOutput = BakeMaterialProperty(BaseColorProperty, BaseColorTexCoord, &TextureSize);
+	const FGLTFPropertyBakeOutput OpacityBakeOutput = BakeMaterialProperty(OpacityProperty, OpacityTexCoord, &TextureSize, true);
+	const float BaseColorScale = BaseColorProperty == MP_EmissiveColor ? BaseColorBakeOutput.EmissiveScale : 1;
+
+	// Detect when both baked properties are constants, which means we can avoid exporting a texture
+	if (BaseColorBakeOutput.bIsConstant && OpacityBakeOutput.bIsConstant)
+	{
+		FLinearColor BaseColorFactor(BaseColorBakeOutput.ConstantValue * BaseColorScale);
+		BaseColorFactor.A = OpacityBakeOutput.ConstantValue.A;
+
+		OutPBRParams.BaseColorFactor = FGLTFConverterUtility::ConvertColor(BaseColorFactor);
+		return true;
+	}
+
+	int32 TexCoord;
+	if (BaseColorBakeOutput.bIsConstant)
+	{
+		TexCoord = OpacityTexCoord;
+	}
+	else if (OpacityBakeOutput.bIsConstant)
+	{
+		TexCoord = BaseColorTexCoord;
+	}
+	else if (BaseColorTexCoord == OpacityTexCoord)
+	{
+		TexCoord = BaseColorTexCoord;
+	}
+	else
+	{
+		// TODO: report error (texture coordinate conflict)
+		return false;
+	}
+
+	TextureSize = BaseColorBakeOutput.Size.ComponentMax(OpacityBakeOutput.Size);
+	BaseColorTexture = FGLTFMaterialUtility::CreateTransientTexture(BaseColorBakeOutput);
+	OpacityTexture = FGLTFMaterialUtility::CreateTransientTexture(OpacityBakeOutput);
+
+	const FString TextureName = Material->GetName() + TEXT("_BaseColor");
+
+	const TArray<FGLTFTextureCombineSource> CombineSources =
+	{
+		{ OpacityTexture, OpacityMask, SE_BLEND_Opaque },
+		{ BaseColorTexture, BaseColorMask }
+	};
+
+	const FGLTFJsonTextureIndex TextureIndex = FGLTFMaterialUtility::AddCombinedTexture(
+		Builder,
+		CombineSources,
+		TextureSize,
+		TextureName,
+		TextureMinFilter,
+		TextureMagFilter,
+		TextureWrapS,
+		TextureWrapT);
+
+	OutPBRParams.BaseColorTexture.TexCoord = TexCoord;
+	OutPBRParams.BaseColorTexture.Index = TextureIndex;
+	OutPBRParams.BaseColorFactor = { BaseColorScale, BaseColorScale, BaseColorScale };
+
+	return true;
+}
+
+bool FGLTFMaterialTask::TryGetMetallicAndRoughness(FGLTFJsonPBRMetallicRoughness& OutPBRParams, EMaterialProperty MetallicProperty, EMaterialProperty RoughnessProperty) const
+{
+	const bool bIsMetallicConstant = TryGetConstantScalar(OutPBRParams.MetallicFactor, MetallicProperty);
+	const bool bIsRoughnessConstant = TryGetConstantScalar(OutPBRParams.RoughnessFactor, RoughnessProperty);
+
+	if (bIsMetallicConstant && bIsRoughnessConstant)
+	{
+		return true;
+	}
+
+	// NOTE: since we always bake the properties (for now) when atleast one property is non-const, we need
+	// to reset the constant factors to their defaults. Otherwise the baked value of a constant property
+	// would be scaled with the factor, i.e a double scaling.
+	OutPBRParams.MetallicFactor = 1.0f;
+	OutPBRParams.RoughnessFactor = 1.0f;
+
+	const UTexture2D* MetallicTexture;
+	const UTexture2D* RoughnessTexture;
+	int32 MetallicTexCoord;
+	int32 RoughnessTexCoord;
+
+	const bool bHasMetallicSourceTexture = TryGetSourceTexture(MetallicTexture, MetallicTexCoord, MetallicProperty, MetallicInputMasks);
+	const bool bHasRoughnessSourceTexture = TryGetSourceTexture(RoughnessTexture, RoughnessTexCoord, RoughnessProperty, RoughnessInputMasks);
+
+	// Detect the "happy path" where both inputs share the same texture and are correctly masked.
+	if (bHasMetallicSourceTexture &&
+		bHasRoughnessSourceTexture &&
+		MetallicTexture == RoughnessTexture &&
+		MetallicTexCoord == RoughnessTexCoord)
+	{
+		OutPBRParams.MetallicRoughnessTexture.Index = Builder.GetOrAddTexture(MetallicTexture);
+		OutPBRParams.MetallicRoughnessTexture.TexCoord = MetallicTexCoord;
+		return true;
+	}
+
+	if (!Builder.ExportOptions->bBakeMaterialInputs)
+	{
+		Builder.AddWarningMessage(FString::Printf(
+			TEXT("%s and %s for material %s needs to bake, but material baking disabled by export options"),
+			FGLTFMaterialUtility::GetPropertyName(MetallicProperty),
+			FGLTFMaterialUtility::GetPropertyName(RoughnessProperty),
+			*Material->GetName()));
+		return false;
+	}
+
+	// TODO: add support for calculating the ideal resolution to use for baking based on connected (texture) nodes
+	FIntPoint TextureSize = Builder.ExportOptions->GetDefaultMaterialBakeSize();
+
+	// TODO: should this be the default wrap-mode?
+	EGLTFJsonTextureWrap TextureWrapS = EGLTFJsonTextureWrap::Repeat;
+	EGLTFJsonTextureWrap TextureWrapT = EGLTFJsonTextureWrap::Repeat;
+
+	// TODO: should this be the default filter?
+	EGLTFJsonTextureFilter TextureMinFilter = EGLTFJsonTextureFilter::LinearMipmapLinear;
+	EGLTFJsonTextureFilter TextureMagFilter = EGLTFJsonTextureFilter::Linear;
+
+	if (bHasMetallicSourceTexture && bHasRoughnessSourceTexture)
+	{
+		const bool bAreTexturesCompatible = MetallicTexCoord == RoughnessTexCoord &&
+			MetallicTexture->AddressX == RoughnessTexture->AddressX &&
+			MetallicTexture->AddressY == RoughnessTexture->AddressY;
+
+		if (!bAreTexturesCompatible)
+		{
+			// TODO: handle differences in wrapping or uv-coords
+			Builder.AddWarningMessage(FString::Printf(
+				TEXT("Metallic- and Roughness-textures for material %s were not able to be combined and will be skipped"),
+				*Material->GetName()));
+
+			return false;
+		}
+
+		TextureSize =
+		{
+			FMath::Max(MetallicTexture->GetSizeX(), RoughnessTexture->GetSizeX()),
+			FMath::Max(MetallicTexture->GetSizeY(), RoughnessTexture->GetSizeY())
+		};
+
+		TextureWrapS = FGLTFConverterUtility::ConvertWrap(MetallicTexture->AddressX);
+		TextureWrapT = FGLTFConverterUtility::ConvertWrap(MetallicTexture->AddressY);
+
+		// TODO: compare min- and mag-filter for BaseColorTexture and OpacityTexture. If they differ,
+		// we should choose one or the other and inform the user about the choice made by logging to the console.
+		TextureMinFilter = FGLTFConverterUtility::ConvertMinFilter(MetallicTexture->Filter, MetallicTexture->LODGroup);
+		TextureMagFilter = FGLTFConverterUtility::ConvertMagFilter(MetallicTexture->Filter, MetallicTexture->LODGroup);
+	}
+	else if (bHasMetallicSourceTexture)
+	{
+		TextureSize = { MetallicTexture->GetSizeX(), MetallicTexture->GetSizeY() };
+		TextureWrapS = FGLTFConverterUtility::ConvertWrap(MetallicTexture->AddressX);
+		TextureWrapT = FGLTFConverterUtility::ConvertWrap(MetallicTexture->AddressY);
+		TextureMinFilter = FGLTFConverterUtility::ConvertMinFilter(MetallicTexture->Filter, MetallicTexture->LODGroup);
+		TextureMagFilter = FGLTFConverterUtility::ConvertMagFilter(MetallicTexture->Filter, MetallicTexture->LODGroup);
+	}
+	else if (bHasRoughnessSourceTexture)
+	{
+		TextureSize = { RoughnessTexture->GetSizeX(), RoughnessTexture->GetSizeY() };
+		TextureWrapS = FGLTFConverterUtility::ConvertWrap(RoughnessTexture->AddressX);
+		TextureWrapT = FGLTFConverterUtility::ConvertWrap(RoughnessTexture->AddressY);
+		TextureMinFilter = FGLTFConverterUtility::ConvertMinFilter(RoughnessTexture->Filter, RoughnessTexture->LODGroup);
+		TextureMagFilter = FGLTFConverterUtility::ConvertMagFilter(RoughnessTexture->Filter, RoughnessTexture->LODGroup);
+	}
+
+	const FGLTFPropertyBakeOutput MetallicBakeOutput = BakeMaterialProperty(MetallicProperty, MetallicTexCoord, &TextureSize);
+	const FGLTFPropertyBakeOutput RoughnessBakeOutput = BakeMaterialProperty(RoughnessProperty, RoughnessTexCoord, &TextureSize);
+
+	// Detect when both baked properties are constants, which means we can use factors and avoid exporting a texture
+	if (MetallicBakeOutput.bIsConstant && RoughnessBakeOutput.bIsConstant)
+	{
+		OutPBRParams.MetallicFactor = MetallicBakeOutput.ConstantValue.R;
+		OutPBRParams.RoughnessFactor = RoughnessBakeOutput.ConstantValue.R;
+		return true;
+	}
+
+	int32 TexCoord;
+	if (MetallicBakeOutput.bIsConstant)
+	{
+		TexCoord = RoughnessTexCoord;
+	}
+	else if (MetallicBakeOutput.bIsConstant)
+	{
+		TexCoord = MetallicTexCoord;
+	}
+	else if (MetallicTexCoord == RoughnessTexCoord)
+	{
+		TexCoord = MetallicTexCoord;
+	}
+	else
+	{
+		// TODO: report error (texture coordinate conflict)
+		return false;
+	}
+
+	TextureSize = RoughnessBakeOutput.Size.ComponentMax(MetallicBakeOutput.Size);
+	MetallicTexture = FGLTFMaterialUtility::CreateTransientTexture(MetallicBakeOutput);
+	RoughnessTexture = FGLTFMaterialUtility::CreateTransientTexture(RoughnessBakeOutput);
+
+	const FString TextureName = Material->GetName() + TEXT("_MetallicRoughness");
+
+	const TArray<FGLTFTextureCombineSource> CombineSources =
+	{
+		{ MetallicTexture, MetallicMask + AlphaMask, SE_BLEND_Opaque },
+		{ RoughnessTexture, RoughnessMask }
+	};
+
+	const FGLTFJsonTextureIndex TextureIndex = FGLTFMaterialUtility::AddCombinedTexture(
+		Builder,
+		CombineSources,
+		TextureSize,
+		TextureName,
+		TextureMinFilter,
+		TextureMagFilter,
+		TextureWrapS,
+		TextureWrapT);
+
+	OutPBRParams.MetallicRoughnessTexture.TexCoord = TexCoord;
+	OutPBRParams.MetallicRoughnessTexture.Index = TextureIndex;
+
+	return true;
+}
+
+bool FGLTFMaterialTask::TryGetClearCoatRoughness(FGLTFJsonClearCoatExtension& OutExtParams, EMaterialProperty IntensityProperty, EMaterialProperty RoughnessProperty) const
+{
+	const bool bIsIntensityConstant = TryGetConstantScalar(OutExtParams.ClearCoatFactor, IntensityProperty);
+	const bool bIsRoughnessConstant = TryGetConstantScalar(OutExtParams.ClearCoatRoughnessFactor, RoughnessProperty);
+
+	if (bIsIntensityConstant && bIsRoughnessConstant)
+	{
+		return true;
+	}
+
+	// NOTE: since we always bake the properties (for now) when atleast one property is non-const, we need
+	// to reset the constant factors to their defaults. Otherwise the baked value of a constant property
+	// would be scaled with the factor, i.e a double scaling.
+	OutExtParams.ClearCoatFactor = 1.0f;
+	OutExtParams.ClearCoatRoughnessFactor = 1.0f;
+
+	const UTexture2D* IntensityTexture;
+	const UTexture2D* RoughnessTexture;
+	int32 IntensityTexCoord;
+	int32 RoughnessTexCoord;
+
+	const bool bHasIntensitySourceTexture = TryGetSourceTexture(IntensityTexture, IntensityTexCoord, IntensityProperty, ClearCoatInputMasks);
+	const bool bHasRoughnessSourceTexture = TryGetSourceTexture(RoughnessTexture, RoughnessTexCoord, RoughnessProperty, ClearCoatRoughnessInputMasks);
+
+	// Detect the "happy path" where both inputs share the same texture and are correctly masked.
+	if (bHasIntensitySourceTexture &&
+		bHasRoughnessSourceTexture &&
+		IntensityTexture == RoughnessTexture &&
+		IntensityTexCoord == RoughnessTexCoord)
+	{
+		const FGLTFJsonTextureIndex TextureIndex = Builder.GetOrAddTexture(IntensityTexture);
+		OutExtParams.ClearCoatTexture.Index = TextureIndex;
+		OutExtParams.ClearCoatTexture.TexCoord = IntensityTexCoord;
+		OutExtParams.ClearCoatRoughnessTexture.Index = TextureIndex;
+		OutExtParams.ClearCoatRoughnessTexture.TexCoord = RoughnessTexCoord;
+		return true;
+	}
+
+	if (!Builder.ExportOptions->bBakeMaterialInputs)
+	{
+		Builder.AddWarningMessage(FString::Printf(
+			TEXT("%s and %s for material %s needs to bake, but material baking disabled by export options"),
+			FGLTFMaterialUtility::GetPropertyName(IntensityProperty),
+			FGLTFMaterialUtility::GetPropertyName(RoughnessProperty),
+			*Material->GetName()));
+		return false;
+	}
+
+	// TODO: add support for calculating the ideal resolution to use for baking based on connected (texture) nodes
+	FIntPoint TextureSize = Builder.ExportOptions->GetDefaultMaterialBakeSize();
+
+	// TODO: should this be the default wrap-mode?
+	EGLTFJsonTextureWrap TextureWrapS = EGLTFJsonTextureWrap::Repeat;
+	EGLTFJsonTextureWrap TextureWrapT = EGLTFJsonTextureWrap::Repeat;
+
+	// TODO: should this be the default filter?
+	EGLTFJsonTextureFilter TextureMinFilter = EGLTFJsonTextureFilter::LinearMipmapLinear;
+	EGLTFJsonTextureFilter TextureMagFilter = EGLTFJsonTextureFilter::Linear;
+
+	if (bHasIntensitySourceTexture && bHasRoughnessSourceTexture)
+	{
+		const bool bAreTexturesCompatible = IntensityTexCoord == RoughnessTexCoord &&
+			IntensityTexture->AddressX == RoughnessTexture->AddressX &&
+			IntensityTexture->AddressY == RoughnessTexture->AddressY;
+
+		if (!bAreTexturesCompatible)
+		{
+			// TODO: handle differences in wrapping or uv-coords
+			Builder.AddWarningMessage(FString::Printf(
+				TEXT("Intensity- and Roughness-textures for material %s were not able to be combined and will be skipped"),
+				*Material->GetName()));
+
+			return false;
+		}
+
+		TextureSize =
+		{
+			FMath::Max(IntensityTexture->GetSizeX(), RoughnessTexture->GetSizeX()),
+			FMath::Max(IntensityTexture->GetSizeY(), RoughnessTexture->GetSizeY())
+		};
+
+		TextureWrapS = FGLTFConverterUtility::ConvertWrap(IntensityTexture->AddressX);
+		TextureWrapT = FGLTFConverterUtility::ConvertWrap(IntensityTexture->AddressY);
+
+		// TODO: compare min- and mag-filter for BaseColorTexture and OpacityTexture. If they differ,
+		// we should choose one or the other and inform the user about the choice made by logging to the console.
+		TextureMinFilter = FGLTFConverterUtility::ConvertMinFilter(IntensityTexture->Filter, IntensityTexture->LODGroup);
+		TextureMagFilter = FGLTFConverterUtility::ConvertMagFilter(IntensityTexture->Filter, IntensityTexture->LODGroup);
+	}
+	else if (bHasIntensitySourceTexture)
+	{
+		TextureSize = { IntensityTexture->GetSizeX(), IntensityTexture->GetSizeY() };
+		TextureWrapS = FGLTFConverterUtility::ConvertWrap(IntensityTexture->AddressX);
+		TextureWrapT = FGLTFConverterUtility::ConvertWrap(IntensityTexture->AddressY);
+		TextureMinFilter = FGLTFConverterUtility::ConvertMinFilter(IntensityTexture->Filter, IntensityTexture->LODGroup);
+		TextureMagFilter = FGLTFConverterUtility::ConvertMagFilter(IntensityTexture->Filter, IntensityTexture->LODGroup);
+	}
+	else if (bHasRoughnessSourceTexture)
+	{
+		TextureSize = { RoughnessTexture->GetSizeX(), RoughnessTexture->GetSizeY() };
+		TextureWrapS = FGLTFConverterUtility::ConvertWrap(RoughnessTexture->AddressX);
+		TextureWrapT = FGLTFConverterUtility::ConvertWrap(RoughnessTexture->AddressY);
+		TextureMinFilter = FGLTFConverterUtility::ConvertMinFilter(RoughnessTexture->Filter, RoughnessTexture->LODGroup);
+		TextureMagFilter = FGLTFConverterUtility::ConvertMagFilter(RoughnessTexture->Filter, RoughnessTexture->LODGroup);
+	}
+
+	const FGLTFPropertyBakeOutput IntensityBakeOutput = BakeMaterialProperty(IntensityProperty, IntensityTexCoord, &TextureSize);
+	const FGLTFPropertyBakeOutput RoughnessBakeOutput = BakeMaterialProperty(RoughnessProperty, RoughnessTexCoord, &TextureSize);
+
+	// Detect when both baked properties are constants, which means we can use factors and avoid exporting a texture
+	if (IntensityBakeOutput.bIsConstant && RoughnessBakeOutput.bIsConstant)
+	{
+		OutExtParams.ClearCoatFactor = IntensityBakeOutput.ConstantValue.R;
+		OutExtParams.ClearCoatRoughnessFactor =  RoughnessBakeOutput.ConstantValue.R;
+		return true;
+	}
+
+	int32 TexCoord;
+	if (IntensityBakeOutput.bIsConstant)
+	{
+		TexCoord = RoughnessTexCoord;
+	}
+	else if (IntensityBakeOutput.bIsConstant)
+	{
+		TexCoord = IntensityTexCoord;
+	}
+	else if (IntensityTexCoord == RoughnessTexCoord)
+	{
+		TexCoord = IntensityTexCoord;
+	}
+	else
+	{
+		// TODO: report error (texture coordinate conflict)
+		return false;
+	}
+
+	TextureSize = RoughnessBakeOutput.Size.ComponentMax(IntensityBakeOutput.Size);
+	IntensityTexture = FGLTFMaterialUtility::CreateTransientTexture(IntensityBakeOutput);
+	RoughnessTexture = FGLTFMaterialUtility::CreateTransientTexture(RoughnessBakeOutput);
+
+	const FString TextureName = Material->GetName() + TEXT("_ClearCoatRoughness");
+
+	const TArray<FGLTFTextureCombineSource> CombineSources =
+	{
+		{ IntensityTexture, ClearCoatMask + AlphaMask, SE_BLEND_Opaque },
+		{ RoughnessTexture, ClearCoatRoughnessMask }
+	};
+
+	const FGLTFJsonTextureIndex TextureIndex = FGLTFMaterialUtility::AddCombinedTexture(
+		Builder,
+		CombineSources,
+		TextureSize,
+		TextureName,
+		TextureMinFilter,
+		TextureMagFilter,
+		TextureWrapS,
+		TextureWrapT);
+
+	OutExtParams.ClearCoatTexture.Index = TextureIndex;
+	OutExtParams.ClearCoatTexture.TexCoord = TexCoord;
+	OutExtParams.ClearCoatRoughnessTexture.Index = TextureIndex;
+	OutExtParams.ClearCoatRoughnessTexture.TexCoord = TexCoord;
+
+	return true;
+}
+
+bool FGLTFMaterialTask::TryGetEmissive(FGLTFJsonMaterial& JsonMaterial, EMaterialProperty EmissiveProperty) const
+{
+	// TODO: right now we allow EmissiveFactor to be > 1.0 to support very bright emission, although it's not valid according to the glTF standard.
+	// We may want to change this behaviour and store factors above 1.0 using a custom extension instead.
+
+	if (TryGetConstantColor(JsonMaterial.EmissiveFactor, MP_EmissiveColor))
+	{
+		return true;
+	}
+
+	if (TryGetSourceTexture(JsonMaterial.EmissiveTexture, EmissiveProperty, DefaultColorInputMasks))
+	{
+		JsonMaterial.EmissiveFactor = FGLTFJsonColor3::White;	// make sure texture is not multiplied with black
+		return true;
+	}
+
+	if (!Builder.ExportOptions->bBakeMaterialInputs)
+	{
+		Builder.AddWarningMessage(FString::Printf(
+			TEXT("%s for material %s needs to bake, but material baking disabled by export options"),
+			FGLTFMaterialUtility::GetPropertyName(EmissiveProperty),
+			*Material->GetName()));
+		return false;
+	}
+
+	const FGLTFPropertyBakeOutput PropertyBakeOutput = BakeMaterialProperty(EmissiveProperty, JsonMaterial.EmissiveTexture.TexCoord);
+	const float EmissiveScale = PropertyBakeOutput.EmissiveScale;
+
+	if (PropertyBakeOutput.bIsConstant)
+	{
+		const FLinearColor EmissiveColor = PropertyBakeOutput.ConstantValue;
+		JsonMaterial.EmissiveFactor = FGLTFConverterUtility::ConvertColor(EmissiveColor * EmissiveScale);
+	}
+	else
+	{
+		if (!StoreBakedPropertyTexture(JsonMaterial.EmissiveTexture, PropertyBakeOutput, TEXT("Emissive")))
+		{
+			return false;
+		}
+
+		JsonMaterial.EmissiveFactor = { EmissiveScale, EmissiveScale, EmissiveScale };
+	}
+
+	return true;
+}
+
+bool FGLTFMaterialTask::IsPropertyNonDefault(EMaterialProperty Property) const
+{
+	const bool bUseMaterialAttributes = Material->GetMaterial()->bUseMaterialAttributes;
+	if (bUseMaterialAttributes)
+	{
+		// TODO: check if attribute property connected, i.e. Material->GetMaterial()->MaterialAttributes.IsConnected(Property)
+		return true;
+	}
+
+	const FExpressionInput* MaterialInput = FGLTFMaterialUtility::GetInputForProperty(Material, Property);
+	if (MaterialInput == nullptr)
+	{
+		// TODO: report error
+		return false;
+	}
+
+	const UMaterialExpression* Expression = MaterialInput->Expression;
+	if (Expression == nullptr)
+	{
+		return false;
+	}
+
+	return true;
+}
+
+bool FGLTFMaterialTask::TryGetConstantColor(FGLTFJsonColor3& OutValue, EMaterialProperty Property) const
+{
+	FLinearColor Value;
+	if (TryGetConstantColor(Value, Property))
+	{
+		OutValue = FGLTFConverterUtility::ConvertColor(Value);
+		return true;
+	}
+
+	return false;
+}
+
+bool FGLTFMaterialTask::TryGetConstantColor(FGLTFJsonColor4& OutValue, EMaterialProperty Property) const
+{
+	FLinearColor Value;
+	if (TryGetConstantColor(Value, Property))
+	{
+		OutValue = FGLTFConverterUtility::ConvertColor(Value);
+		return true;
+	}
+
+	return false;
+}
+
+bool FGLTFMaterialTask::TryGetConstantColor(FLinearColor& OutValue, EMaterialProperty Property) const
+{
+	const bool bUseMaterialAttributes = Material->GetMaterial()->bUseMaterialAttributes;
+	if (bUseMaterialAttributes)
+	{
+		// TODO: check if attribute property connected, i.e. Material->GetMaterial()->MaterialAttributes.IsConnected(Property)
+		return false;
+	}
+
+	const FMaterialInput<FColor>* MaterialInput = FGLTFMaterialUtility::GetInputForProperty<FColor>(Material, Property);
+	if (MaterialInput == nullptr)
+	{
+		// TODO: report error
+		return false;
+	}
+
+	if (MaterialInput->UseConstant)
+	{
+		OutValue = { MaterialInput->Constant };
+		return true;
+	}
+
+	const UMaterialExpression* Expression = MaterialInput->Expression;
+	if (Expression == nullptr)
+	{
+		OutValue = FLinearColor(FGLTFMaterialUtility::GetPropertyDefaultValue(Property));
+		return true;
+	}
+
+	if (const UMaterialExpressionVectorParameter* VectorParameter = ExactCast<UMaterialExpressionVectorParameter>(Expression))
+	{
+		FLinearColor Value = VectorParameter->DefaultValue;
+
+		const UMaterialInstance* MaterialInstance = Cast<UMaterialInstance>(Material);
+		if (MaterialInstance != nullptr)
+		{
+			const FHashedMaterialParameterInfo ParameterInfo(VectorParameter->GetParameterName());
+			if (!MaterialInstance->GetVectorParameterValue(ParameterInfo, Value))
+			{
+				// TODO: how to handle this?
+			}
+		}
+
+		const uint32 MaskComponentCount = FGLTFMaterialUtility::GetMaskComponentCount(*MaterialInput);
+
+		if (MaskComponentCount > 0)
+		{
+			const FLinearColor Mask = FGLTFMaterialUtility::GetMask(*MaterialInput);
+
+			Value *= Mask;
+
+			if (MaskComponentCount == 1)
+			{
+				const float ComponentValue = Value.R + Value.G + Value.B + Value.A;
+				Value = { ComponentValue, ComponentValue, ComponentValue, ComponentValue };
+			}
+		}
+
+		OutValue = Value;
+		return true;
+	}
+
+	if (const UMaterialExpressionScalarParameter* ScalarParameter = ExactCast<UMaterialExpressionScalarParameter>(Expression))
+	{
+		float Value = ScalarParameter->DefaultValue;
+
+		const UMaterialInstance* MaterialInstance = Cast<UMaterialInstance>(Material);
+		if (MaterialInstance != nullptr)
+		{
+			const FHashedMaterialParameterInfo ParameterInfo(ScalarParameter->GetParameterName());
+			if (!MaterialInstance->GetScalarParameterValue(ParameterInfo, Value))
+			{
+				// TODO: how to handle this?
+			}
+		}
+
+		OutValue = { Value, Value, Value, Value };
+		return true;
+	}
+
+	if (const UMaterialExpressionConstant4Vector* Constant4Vector = ExactCast<UMaterialExpressionConstant4Vector>(Expression))
+	{
+		OutValue = Constant4Vector->Constant;
+		return true;
+	}
+
+	if (const UMaterialExpressionConstant3Vector* Constant3Vector = ExactCast<UMaterialExpressionConstant3Vector>(Expression))
+	{
+		OutValue = Constant3Vector->Constant;
+		return true;
+	}
+
+	if (const UMaterialExpressionConstant2Vector* Constant2Vector = ExactCast<UMaterialExpressionConstant2Vector>(Expression))
+	{
+		OutValue = { Constant2Vector->R, Constant2Vector->G, 0, 0 };
+		return true;
+	}
+
+	if (const UMaterialExpressionConstant* Constant = ExactCast<UMaterialExpressionConstant>(Expression))
+	{
+		OutValue = { Constant->R, Constant->R, Constant->R, Constant->R };
+		return true;
+	}
+
+	return false;
+}
+
+bool FGLTFMaterialTask::TryGetConstantScalar(float& OutValue, EMaterialProperty Property) const
+{
+	const bool bUseMaterialAttributes = Material->GetMaterial()->bUseMaterialAttributes;
+	if (bUseMaterialAttributes)
+	{
+		// TODO: check if attribute property connected, i.e. Material->GetMaterial()->MaterialAttributes.IsConnected(Property)
+		return false;
+	}
+
+	const FMaterialInput<float>* MaterialInput = FGLTFMaterialUtility::GetInputForProperty<float>(Material, Property);
+	if (MaterialInput == nullptr)
+	{
+		// TODO: report error
+		return false;
+	}
+
+	if (MaterialInput->UseConstant)
+	{
+		OutValue = MaterialInput->Constant;
+		return true;
+	}
+
+	const UMaterialExpression* Expression = MaterialInput->Expression;
+	if (Expression == nullptr)
+	{
+		OutValue = FGLTFMaterialUtility::GetPropertyDefaultValue(Property).X;
+		return true;
+	}
+
+	if (const UMaterialExpressionVectorParameter* VectorParameter = ExactCast<UMaterialExpressionVectorParameter>(Expression))
+	{
+		FLinearColor Value = VectorParameter->DefaultValue;
+
+		const UMaterialInstance* MaterialInstance = Cast<UMaterialInstance>(Material);
+		if (MaterialInstance != nullptr)
+		{
+			const FHashedMaterialParameterInfo ParameterInfo(VectorParameter->GetParameterName());
+			if (!MaterialInstance->GetVectorParameterValue(ParameterInfo, Value))
+			{
+				// TODO: how to handle this?
+			}
+		}
+
+		const uint32 MaskComponentCount = FGLTFMaterialUtility::GetMaskComponentCount(*MaterialInput);
+
+		if (MaskComponentCount > 0)
+		{
+			const FLinearColor Mask = FGLTFMaterialUtility::GetMask(*MaterialInput);
+			Value *= Mask;
+		}
+
+		// TODO: is this a correct assumption, that the max component should be used as value?
+		OutValue = Value.GetMax();
+		return true;
+	}
+
+	if (const UMaterialExpressionScalarParameter* ScalarParameter = ExactCast<UMaterialExpressionScalarParameter>(Expression))
+	{
+		float Value = ScalarParameter->DefaultValue;
+
+		const UMaterialInstance* MaterialInstance = Cast<UMaterialInstance>(Material);
+		if (MaterialInstance != nullptr)
+		{
+			const FHashedMaterialParameterInfo ParameterInfo(ScalarParameter->GetParameterName());
+			if (!MaterialInstance->GetScalarParameterValue(ParameterInfo, Value))
+			{
+				// TODO: how to handle this?
+			}
+		}
+
+		OutValue = Value;
+		return true;
+	}
+
+	if (const UMaterialExpressionConstant4Vector* Constant4Vector = ExactCast<UMaterialExpressionConstant4Vector>(Expression))
+	{
+		OutValue = Constant4Vector->Constant.R;
+		return true;
+	}
+
+	if (const UMaterialExpressionConstant3Vector* Constant3Vector = ExactCast<UMaterialExpressionConstant3Vector>(Expression))
+	{
+		OutValue = Constant3Vector->Constant.R;
+		return true;
+	}
+
+	if (const UMaterialExpressionConstant2Vector* Constant2Vector = ExactCast<UMaterialExpressionConstant2Vector>(Expression))
+	{
+		OutValue = Constant2Vector->R;
+		return true;
+	}
+
+	if (const UMaterialExpressionConstant* Constant = ExactCast<UMaterialExpressionConstant>(Expression))
+	{
+		OutValue = Constant->R;
+		return true;
+	}
+
+	return false;
+}
+
+bool FGLTFMaterialTask::TryGetSourceTexture(FGLTFJsonTextureInfo& OutTexInfo, EMaterialProperty Property, const TArray<FLinearColor>& AllowedMasks) const
+{
+	const UTexture2D* Texture;
+	int32 TexCoord;
+
+	if (TryGetSourceTexture(Texture, TexCoord, Property, AllowedMasks))
+	{
+		OutTexInfo.Index = Builder.GetOrAddTexture(Texture);
+		OutTexInfo.TexCoord = TexCoord;
+		return true;
+	}
+
+	return false;
+}
+
+bool FGLTFMaterialTask::TryGetSourceTexture(const UTexture2D*& OutTexture, int32& OutTexCoord, EMaterialProperty Property, const TArray<FLinearColor>& AllowedMasks) const
+{
+	const FExpressionInput* MaterialInput = FGLTFMaterialUtility::GetInputForProperty(Material, Property);
+	if (MaterialInput == nullptr)
+	{
+		// TODO: report error
+		return false;
+	}
+
+	const UMaterialExpression* Expression = MaterialInput->Expression;
+	if (Expression == nullptr)
+	{
+		return false;
+	}
+
+	const FLinearColor InputMask = FGLTFMaterialUtility::GetMask(*MaterialInput);
+	if (AllowedMasks.Num() > 0 && !AllowedMasks.Contains(InputMask))
+	{
+		return false;
+	}
+
+	// TODO: add support or warning for texture sampler settings that override texture asset addressing (i.e. wrap, clamp etc)?
+
+	if (const UMaterialExpressionTextureSampleParameter2D* TextureParameter = ExactCast<UMaterialExpressionTextureSampleParameter2D>(Expression))
+	{
+		UTexture* ParameterValue = TextureParameter->Texture;
+
+		const UMaterialInstance* MaterialInstance = Cast<UMaterialInstance>(Material);
+		if (MaterialInstance != nullptr)
+		{
+			const FHashedMaterialParameterInfo ParameterInfo(TextureParameter->GetParameterName());
+			if (!MaterialInstance->GetTextureParameterValue(ParameterInfo, ParameterValue))
+			{
+				// TODO: how to handle this?
+			}
+		}
+
+		OutTexture = Cast<UTexture2D>(ParameterValue);
+
+		if (OutTexture == nullptr)
+		{
+			if (ParameterValue == nullptr)
+			{
+				// TODO: report error (no texture parameter assigned)
+			}
+			else
+			{
+				// TODO: report error (incorrect texture type)
+			}
+			return false;
+		}
+
+		if (!FGLTFMaterialUtility::TryGetTextureCoordinateIndex(TextureParameter, OutTexCoord))
+		{
+			// TODO: report error (failed to identify texture coordinate index)
+			return false;
+		}
+
+		return true;
+	}
+
+	if (const UMaterialExpressionTextureSample* TextureSampler = ExactCast<UMaterialExpressionTextureSample>(Expression))
+	{
+		// TODO: add support for texture object input expression
+		OutTexture = Cast<UTexture2D>(TextureSampler->Texture);
+
+		if (OutTexture == nullptr)
+		{
+			if (TextureSampler->Texture == nullptr)
+			{
+				// TODO: report error (no texture sample assigned)
+			}
+			else
+			{
+				// TODO: report error (incorrect texture type)
+			}
+			return false;
+		}
+
+		if (!FGLTFMaterialUtility::TryGetTextureCoordinateIndex(TextureSampler, OutTexCoord))
+		{
+			// TODO: report error (failed to identify texture coordinate index)
+			return false;
+		}
+
+		return true;
+	}
+
+	return false;
+}
+
+bool FGLTFMaterialTask::TryGetBakedMaterialProperty(FGLTFJsonTextureInfo& OutTexInfo, FGLTFJsonColor3& OutConstant, EMaterialProperty Property, const FString& PropertyName) const
+{
+	if (!Builder.ExportOptions->bBakeMaterialInputs)
+	{
+		Builder.AddWarningMessage(FString::Printf(
+			TEXT("%s for material %s needs to bake, but material baking disabled by export options"),
+			FGLTFMaterialUtility::GetPropertyName(Property),
+			*Material->GetName()));
+		return false;
+	}
+
+	const FGLTFPropertyBakeOutput PropertyBakeOutput = BakeMaterialProperty(Property, OutTexInfo.TexCoord);
+
+	if (PropertyBakeOutput.bIsConstant)
+	{
+		OutConstant = FGLTFConverterUtility::ConvertColor(PropertyBakeOutput.ConstantValue);
+		return true;
+	}
+
+	if (StoreBakedPropertyTexture(OutTexInfo, PropertyBakeOutput, PropertyName))
+	{
+		OutConstant = FGLTFJsonColor3::White; // make sure property is not zero
+		return true;
+	}
+
+	return false;
+}
+
+bool FGLTFMaterialTask::TryGetBakedMaterialProperty(FGLTFJsonTextureInfo& OutTexInfo, FGLTFJsonColor4& OutConstant, EMaterialProperty Property, const FString& PropertyName) const
+{
+	if (!Builder.ExportOptions->bBakeMaterialInputs)
+	{
+		Builder.AddWarningMessage(FString::Printf(
+			TEXT("%s for material %s needs to bake, but material baking disabled by export options"),
+			FGLTFMaterialUtility::GetPropertyName(Property),
+			*Material->GetName()));
+		return false;
+	}
+
+	const FGLTFPropertyBakeOutput PropertyBakeOutput = BakeMaterialProperty(Property, OutTexInfo.TexCoord);
+
+	if (PropertyBakeOutput.bIsConstant)
+	{
+		OutConstant = FGLTFConverterUtility::ConvertColor(PropertyBakeOutput.ConstantValue);
+		return true;
+	}
+
+	if (StoreBakedPropertyTexture(OutTexInfo, PropertyBakeOutput, PropertyName))
+	{
+		OutConstant = FGLTFJsonColor4::White; // make sure property is not zero
+		return true;
+	}
+
+	return false;
+}
+
+inline bool FGLTFMaterialTask::TryGetBakedMaterialProperty(FGLTFJsonTextureInfo& OutTexInfo, float& OutConstant, EMaterialProperty Property, const FString& PropertyName) const
+{
+	if (!Builder.ExportOptions->bBakeMaterialInputs)
+	{
+		Builder.AddWarningMessage(FString::Printf(
+			TEXT("%s for material %s needs to bake, but material baking disabled by export options"),
+			FGLTFMaterialUtility::GetPropertyName(Property),
+			*Material->GetName()));
+		return false;
+	}
+
+	const FGLTFPropertyBakeOutput PropertyBakeOutput = BakeMaterialProperty(Property, OutTexInfo.TexCoord);
+
+	if (PropertyBakeOutput.bIsConstant)
+	{
+		OutConstant = PropertyBakeOutput.ConstantValue.R;
+		return true;
+	}
+
+	if (StoreBakedPropertyTexture(OutTexInfo, PropertyBakeOutput, PropertyName))
+	{
+		OutConstant = 1; // make sure property is not zero
+		return true;
+	}
+
+	return false;
+}
+
+bool FGLTFMaterialTask::TryGetBakedMaterialProperty(FGLTFJsonTextureInfo& OutTexInfo, EMaterialProperty Property, const FString& PropertyName) const
+{
+	if (!Builder.ExportOptions->bBakeMaterialInputs)
+	{
+		Builder.AddWarningMessage(FString::Printf(
+			TEXT("%s for material %s needs to bake, but material baking disabled by export options"),
+			FGLTFMaterialUtility::GetPropertyName(Property),
+			*Material->GetName()));
+		return false;
+	}
+
+	const FGLTFPropertyBakeOutput PropertyBakeOutput = BakeMaterialProperty(Property, OutTexInfo.TexCoord);
+
+	if (!PropertyBakeOutput.bIsConstant)
+	{
+		return StoreBakedPropertyTexture(OutTexInfo, PropertyBakeOutput, PropertyName);
+	}
+
+	const FVector4 MaskedConstant = FVector4(PropertyBakeOutput.ConstantValue) * FGLTFMaterialUtility::GetPropertyMask(Property);
+	if (MaskedConstant == FGLTFMaterialUtility::GetPropertyDefaultValue(Property))
+	{
+		// Constant value is the same as the property's default so we can set gltf to default.
+		OutTexInfo.Index = FGLTFJsonTextureIndex(INDEX_NONE);
+		return true;
+	}
+
+	if (Property == MP_Normal || Property == MP_CustomOutput /* ClearCoatBottomNormal */)
+	{
+		// TODO: In some cases baking normal can result in constant vector that differs slight from default (i.e 0,0,1).
+		// Yet often, when looking at such a material, it should be exactly default. Needs further investigation.
+		// Maybe because of incorrect sRGB conversion? For now, assume a constant normal is always default.
+		OutTexInfo.Index = FGLTFJsonTextureIndex(INDEX_NONE);
+		return true;
+	}
+
+	// TODO: let function fail and investigate why in some cases a constant baking result is returned for a property
+	// that is non-constant. This happens (for example) when baking AmbientOcclusion for a translucent material,
+	// even though the same material when set to opaque will properly bake AmbientOcclusion to a texture.
+	// For now, create a 1x1 texture with the constant value.
+
+	const FString TextureName = Material->GetName() + TEXT("_") + PropertyName;
+	const FGLTFJsonTextureIndex TextureIndex = FGLTFMaterialUtility::AddTexture(
+		Builder,
+		PropertyBakeOutput.Pixels,
+		PropertyBakeOutput.Size,
+		Material->GetName() + TEXT("_") + PropertyName,
+		EGLTFJsonTextureFilter::Nearest,
+		EGLTFJsonTextureFilter::Nearest,
+		EGLTFJsonTextureWrap::ClampToEdge,
+		EGLTFJsonTextureWrap::ClampToEdge);
+
+	OutTexInfo.Index = TextureIndex;
+	return true;
+}
+
+FGLTFPropertyBakeOutput FGLTFMaterialTask::BakeMaterialProperty(EMaterialProperty Property, int32& OutTexCoord, const FIntPoint* PreferredTextureSize, bool bCopyAlphaFromRedChannel) const
+{
+	const FExpressionInput* PropertyInput = FGLTFMaterialUtility::GetInputForProperty(Material, Property);
+	TSet<int32> TexCoords;
+
+	FGLTFMaterialUtility::GetAllTextureCoordinateIndices(*PropertyInput, TexCoords);
+	if (TexCoords.Num() > 0)
+	{
+		// TODO: is there a nicer way to get the first element in a TSet?
+		for (int32 TexCoord : TexCoords)
+		{
+			OutTexCoord = TexCoord;
+			break;
+		}
+
+		if (TexCoords.Num() > 1)
+		{
+			// TODO: report warning (multiple texture coordinates found, will use first)
+
+			// TODO: replace this hardcoded hack with something more configurable and proper
+			if (Property == MP_AmbientOcclusion && TexCoords.Contains(1))
+			{
+				OutTexCoord = 1; // assume ambient occlusion uses TexCoord1 when multiple
+			}
+		}
+	}
+	else
+	{
+		OutTexCoord = 0; // assume TexCoord0 even thought property seems to be texture coordinate independent
+	}
+
+	const FIntPoint DefaultTextureSize = Builder.ExportOptions->GetDefaultMaterialBakeSize();
+	const FIntPoint TextureSize = PreferredTextureSize != nullptr ? *PreferredTextureSize : DefaultTextureSize;
+
+	// TODO: add support for calculating the ideal resolution to use for baking based on connected (texture) nodes
+
+	const FGLTFPropertyBakeOutput PropertyBakeOutput = FGLTFMaterialUtility::BakeMaterialProperty(
+		TextureSize,
+		Property,
+		Material,
+		OutTexCoord,
+		bCopyAlphaFromRedChannel);
+
+	if (!PropertyBakeOutput.bIsConstant && TexCoords.Num() == 0)
+	{
+		// TODO: report warning about property not being constant yet texture coordinate independent
+	}
+
+	return PropertyBakeOutput;
+}
+
+bool FGLTFMaterialTask::StoreBakedPropertyTexture(FGLTFJsonTextureInfo& OutTexInfo, const FGLTFPropertyBakeOutput& PropertyBakeOutput, const FString& PropertyName) const
+{
+	const FString TextureName = Material->GetName() + TEXT("_") + PropertyName;
+
+	// TODO: should this be the default wrap-mode?
+	const EGLTFJsonTextureWrap TextureWrapS = EGLTFJsonTextureWrap::Repeat;
+	const EGLTFJsonTextureWrap TextureWrapT = EGLTFJsonTextureWrap::Repeat;
+
+	// TODO: should this be the default filter?
+	const EGLTFJsonTextureFilter TextureMinFilter = EGLTFJsonTextureFilter::LinearMipmapLinear;
+	const EGLTFJsonTextureFilter TextureMagFilter = EGLTFJsonTextureFilter::Linear;
+
+	const FGLTFJsonTextureIndex TextureIndex = FGLTFMaterialUtility::AddTexture(
+		Builder,
+		PropertyBakeOutput.Pixels,
+		PropertyBakeOutput.Size,
+		TextureName,
+		TextureMinFilter,
+		TextureMagFilter,
+		TextureWrapS,
+		TextureWrapT);
+
+	OutTexInfo.Index = TextureIndex;
+	return true;
+}
diff --git a/Engine/Plugins/Enterprise/GLTFExporter/MigratedFromGit/Source/GLTFExporter/Private/Tasks/GLTFMaterialTasks.h b/Engine/Plugins/Enterprise/GLTFExporter/MigratedFromGit/Source/GLTFExporter/Private/Tasks/GLTFMaterialTasks.h
new file mode 100644
--- /dev/null
+++ b/Engine/Plugins/Enterprise/GLTFExporter/MigratedFromGit/Source/GLTFExporter/Private/Tasks/GLTFMaterialTasks.h
@@ -0,0 +1,60 @@
+// Copyright Epic Games, Inc. All Rights Reserved.
+
+#pragma once
+
+#include "Tasks/GLTFTask.h"
+#include "Builders/GLTFConvertBuilder.h"
+#include "Engine.h"
+
+struct FGLTFPropertyBakeOutput;
+
+class FGLTFMaterialTask : public FGLTFTask
+{
+public:
+
+	FGLTFMaterialTask(FGLTFConvertBuilder& Builder, const UMaterialInterface* Material, FGLTFJsonMaterialIndex MaterialIndex)
+        : FGLTFTask(EGLTFTaskPriority::Material)
+		, Builder(Builder)
+		, Material(Material)
+		, MaterialIndex(MaterialIndex)
+	{
+	}
+
+	virtual FString GetName() override
+	{
+		return Material->GetName();
+	}
+
+	virtual void Complete() override;
+
+private:
+
+	FGLTFConvertBuilder& Builder;
+	const UMaterialInterface* Material;
+	const FGLTFJsonMaterialIndex MaterialIndex;
+
+	bool TryGetShadingModel(EGLTFJsonShadingModel& ShadingModel) const;
+
+	bool TryGetBaseColorAndOpacity(FGLTFJsonPBRMetallicRoughness& OutPBRParams, EMaterialProperty BaseColorProperty, EMaterialProperty OpacityProperty) const;
+	bool TryGetMetallicAndRoughness(FGLTFJsonPBRMetallicRoughness& OutPBRParams, EMaterialProperty MetallicProperty, EMaterialProperty RoughnessProperty) const;
+	bool TryGetClearCoatRoughness(FGLTFJsonClearCoatExtension& OutExtParams, EMaterialProperty IntensityProperty, EMaterialProperty RoughnessProperty) const;
+	bool TryGetEmissive(FGLTFJsonMaterial& JsonMaterial, EMaterialProperty EmissiveProperty) const;
+
+	bool IsPropertyNonDefault(EMaterialProperty Property) const;
+	bool TryGetConstantColor(FGLTFJsonColor3& OutValue, EMaterialProperty Property) const;
+	bool TryGetConstantColor(FGLTFJsonColor4& OutValue, EMaterialProperty Property) const;
+	bool TryGetConstantColor(FLinearColor& OutValue, EMaterialProperty Property) const;
+	bool TryGetConstantScalar(float& OutValue, EMaterialProperty Property) const;
+
+	bool TryGetSourceTexture(FGLTFJsonTextureInfo& OutTexInfo, EMaterialProperty Property, const TArray<FLinearColor>& AllowedMasks = {}) const;
+	bool TryGetSourceTexture(const UTexture2D*& OutTexture, int32& OutTexCoord, EMaterialProperty Property, const TArray<FLinearColor>& AllowedMasks = {}) const;
+
+	bool TryGetBakedMaterialProperty(FGLTFJsonTextureInfo& OutTexInfo, FGLTFJsonColor3& OutConstant, EMaterialProperty Property, const FString& PropertyName) const;
+	bool TryGetBakedMaterialProperty(FGLTFJsonTextureInfo& OutTexInfo, FGLTFJsonColor4& OutConstant, EMaterialProperty Property, const FString& PropertyName) const;
+	bool TryGetBakedMaterialProperty(FGLTFJsonTextureInfo& OutTexInfo, float& OutConstant, EMaterialProperty Property, const FString& PropertyName) const;
+	bool TryGetBakedMaterialProperty(FGLTFJsonTextureInfo& OutTexInfo, EMaterialProperty Property, const FString& PropertyName) const;
+
+	FGLTFPropertyBakeOutput BakeMaterialProperty(EMaterialProperty Property, int32& OutTexCoord, const FIntPoint* PreferredTextureSize = nullptr, bool bCopyAlphaFromRedChannel = false) const;
+
+	bool StoreBakedPropertyTexture(FGLTFJsonTextureInfo& OutTexInfo, const FGLTFPropertyBakeOutput& PropertyBakeOutput, const FString& PropertyName) const;
+};
