[GLTFExporter] feat: add texture tasks and change texture converters to use them

commit 6360524b18a4e2b8c451f74af276250fe19cccbf
Date: Sun, 4 Oct 2020 13:50:15 +0200

    feat: add texture tasks and change texture converters to use them
    (cherry picked from commit d4c3e3111f96b0f3d6721c0d2f75ad99405b96a1​)

#skipci

diff --git a/Engine/Plugins/Enterprise/GLTFExporter/MigratedFromGit/Source/GLTFExporter/Private/Converters/GLTFTextureConverters.cpp b/Engine/Plugins/Enterprise/GLTFExporter/MigratedFromGit/Source/GLTFExporter/Private/Converters/GLTFTextureConverters.cpp
--- a/Engine/Plugins/Enterprise/GLTFExporter/MigratedFromGit/Source/GLTFExporter/Private/Converters/GLTFTextureConverters.cpp
+++ b/Engine/Plugins/Enterprise/GLTFExporter/MigratedFromGit/Source/GLTFExporter/Private/Converters/GLTFTextureConverters.cpp
@@ -2,172 +2,32 @@
 
 #include "Converters/GLTFTextureConverters.h"
 #include "Builders/GLTFContainerBuilder.h"
-#include "Converters/GLTFConverterUtility.h"
-#include "Converters/GLTFTextureUtility.h"
+#include "Tasks/GLTFTextureTasks.h"
 
 FGLTFJsonTextureIndex FGLTFTexture2DConverter::Convert(const UTexture2D* Texture2D)
 {
-	FGLTFJsonImageIndex ImageIndex;
-
-	FGLTFJsonTexture JsonTexture;
-	Texture2D->GetName(JsonTexture.Name);
-
-	const FIntPoint Size = { Texture2D->GetSizeX(), Texture2D->GetSizeY() };
-	const bool bPreferSourceExport = Builder.ExportOptions->bExportSourceTextures;
-
-	// NOTE: export of lightmaps and normalmaps via source data is done to workaround issues
-	// with using render data for normalmaps (which are compressed to red & green channels) and
-	// lightmaps (which can otherwise suffer quality-loss due to incorrect gamma transformation).
-
-	// TODO: make sure there are no other special cases than normalmaps and lightmaps that require source-export
-	const bool bRequireSourceExport = Texture2D->IsA<ULightMapTexture2D>() || Texture2D->IsNormalMap();
-
-	if (bRequireSourceExport || bPreferSourceExport)
-	{
-		ERGBFormat RGBFormat;
-		uint32 BitDepth;
-
-		FTextureSource& Source = const_cast<FTextureSource&>(Texture2D->Source);
-
-		if (Source.IsValid() && FGLTFTextureUtility::CanPNGCompressFormat(Source.GetFormat(), RGBFormat, BitDepth))
-		{
-			const void* RawData = Source.LockMip(0);
-			ImageIndex = Builder.AddImage(RawData, Source.CalcMipSize(0), Size, RGBFormat, BitDepth, JsonTexture.Name);
-			Source.UnlockMip(0);
-		}
-
-		if (ImageIndex != INDEX_NONE)
-		{
-			if (FGLTFTextureUtility::HasAnyAdjustment(Texture2D))
-			{
-				Builder.AddWarningMessage(FString::Printf(TEXT("Adjustments for texture %s are not supported when exporting source data"), *JsonTexture.Name));
-			}
-		}
-		else
-		{
-			if (bPreferSourceExport)
-			{
-				Builder.AddWarningMessage(FString::Printf(TEXT("Unable to export source for texture %s, render data will be used as fallback"), *JsonTexture.Name));
-			}
-		}
-	}
-
-	if (ImageIndex == INDEX_NONE && !bRequireSourceExport)
-	{
-		// TODO: both bForceLinearGamma and TargetGamma=2.2f seem to be necessary for the exported images to match the results
-		// from exporting using the texture's source or its platform-data.
-		// It's not entirely clear why gamma must be set to 2.2 instead of 0.0 (which should use the correct gamma anyway),
-		// and why bInForceLinearGamma must also be true.
-		const bool bForceLinearGamma = true;
-		const float TargetGamma = 2.2f;
-
-		const EPixelFormat RenderTargetFormat = FGLTFTextureUtility::IsHDRFormat(Texture2D->GetPixelFormat()) ? PF_FloatRGBA : PF_B8G8R8A8;
-		UTextureRenderTarget2D* RenderTarget = FGLTFTextureUtility::CreateRenderTarget(Size, RenderTargetFormat, bForceLinearGamma);
-		RenderTarget->TargetGamma = TargetGamma;
-
-		FGLTFTextureUtility::DrawTexture(RenderTarget, Texture2D);
-
-		TArray<FColor> Pixels;
-		if (FGLTFTextureUtility::ReadEncodedPixels(RenderTarget, Pixels, JsonTexture.Encoding)) // TODO: use only encoding as specified by export options
-		{
-			ImageIndex = Builder.AddImage(Pixels.GetData(), Size, JsonTexture.Name);
-		}
-	}
-
-	if (ImageIndex == INDEX_NONE)
-	{
-		Builder.AddWarningMessage(FString::Printf(TEXT("Failed to export texture %s"), *JsonTexture.Name));
-		return FGLTFJsonTextureIndex(INDEX_NONE);
-	}
-
-	JsonTexture.Source = ImageIndex;
-	JsonTexture.Sampler = Builder.GetOrAddSampler(Texture2D);
-
-	return Builder.AddTexture(JsonTexture);
+	const FGLTFJsonTextureIndex TextureIndex = Builder.AddTexture();
+	Builder.SetupTask<FGLTFTexture2DTask>(Builder, Texture2D, TextureIndex);
+	return TextureIndex;
 }
 
 FGLTFJsonTextureIndex FGLTFTextureCubeConverter::Convert(const UTextureCube* TextureCube, ECubeFace CubeFace)
 {
-	FGLTFJsonTexture JsonTexture;
-	JsonTexture.Name = TextureCube->GetName() + TEXT("_") + FGLTFJsonUtility::ToString(FGLTFConverterUtility::ConvertCubeFace(CubeFace));
-
-	// TODO: add optimized "happy path" if cube face doesn't need rotation and has suitable pixel format
-
-	const UTexture2D* FaceTexture = FGLTFTextureUtility::CreateTextureFromCubeFace(TextureCube, CubeFace);
-	if (FaceTexture == nullptr)
-	{
-		// TODO: report error
-		return FGLTFJsonTextureIndex(INDEX_NONE);
-	}
-
-	const FIntPoint Size = { FaceTexture->GetSizeX(), FaceTexture->GetSizeY() };
-	const EPixelFormat RenderTargetFormat = FGLTFTextureUtility::IsHDRFormat(FaceTexture->GetPixelFormat()) ? PF_FloatRGBA : PF_B8G8R8A8;
-	UTextureRenderTarget2D* RenderTarget = FGLTFTextureUtility::CreateRenderTarget(Size, RenderTargetFormat, true);
-
-	const float FaceRotation = FGLTFTextureUtility::GetCubeFaceRotation(CubeFace);
-	FGLTFTextureUtility::RotateTexture(RenderTarget, FaceTexture, FaceRotation);
-
-	TArray<FColor> Pixels;
-	if (!FGLTFTextureUtility::ReadEncodedPixels(RenderTarget, Pixels, JsonTexture.Encoding)) // TODO: use only encoding as specified by export options
-	{
-		// TODO: report error
-		return FGLTFJsonTextureIndex(INDEX_NONE);
-	}
-
-	JsonTexture.Source = Builder.AddImage(Pixels.GetData(), Size, JsonTexture.Name);
-	JsonTexture.Sampler = Builder.GetOrAddSampler(TextureCube);
-
-	return Builder.AddTexture(JsonTexture);
+	const FGLTFJsonTextureIndex TextureIndex = Builder.AddTexture();
+	Builder.SetupTask<FGLTFTextureCubeTask>(Builder, TextureCube, CubeFace, TextureIndex);
+	return TextureIndex;
 }
 
 FGLTFJsonTextureIndex FGLTFTextureRenderTarget2DConverter::Convert(const UTextureRenderTarget2D* RenderTarget2D)
 {
-	FGLTFJsonTexture JsonTexture;
-	RenderTarget2D->GetName(JsonTexture.Name);
-
-	TArray<FColor> Pixels;
-	if (!FGLTFTextureUtility::ReadEncodedPixels(RenderTarget2D, Pixels, JsonTexture.Encoding)) // TODO: use only encoding as specified by export options
-	{
-		// TODO: report error
-		return FGLTFJsonTextureIndex(INDEX_NONE);
-	}
-
-	JsonTexture.Source = Builder.AddImage(Pixels.GetData(), { RenderTarget2D->SizeX, RenderTarget2D->SizeY }, JsonTexture.Name);
-	JsonTexture.Sampler = Builder.GetOrAddSampler(RenderTarget2D);
-
-	return Builder.AddTexture(JsonTexture);
+	const FGLTFJsonTextureIndex TextureIndex = Builder.AddTexture();
+	Builder.SetupTask<FGLTFTextureRenderTarget2DTask>(Builder, RenderTarget2D, TextureIndex);
+	return TextureIndex;
 }
 
 FGLTFJsonTextureIndex FGLTFTextureRenderTargetCubeConverter::Convert(const UTextureRenderTargetCube* RenderTargetCube, ECubeFace CubeFace)
 {
-	FGLTFJsonTexture JsonTexture;
-	JsonTexture.Name = RenderTargetCube->GetName() + TEXT("_") + FGLTFJsonUtility::ToString(FGLTFConverterUtility::ConvertCubeFace(CubeFace));
-
-	// TODO: add optimized "happy path" if cube face doesn't need rotation
-
-	const UTexture2D* FaceTexture = FGLTFTextureUtility::CreateTextureFromCubeFace(RenderTargetCube, CubeFace);
-	if (FaceTexture == nullptr)
-	{
-		// TODO: report error
-		return FGLTFJsonTextureIndex(INDEX_NONE);
-	}
-
-	const FIntPoint Size = { FaceTexture->GetSizeX(), FaceTexture->GetSizeY() };
-	const EPixelFormat RenderTargetFormat = FGLTFTextureUtility::IsHDRFormat(FaceTexture->GetPixelFormat()) ? PF_FloatRGBA : PF_B8G8R8A8;
-	UTextureRenderTarget2D* RenderTarget = FGLTFTextureUtility::CreateRenderTarget(Size, RenderTargetFormat, true);
-
-	const float FaceRotation = FGLTFTextureUtility::GetCubeFaceRotation(CubeFace);
-	FGLTFTextureUtility::RotateTexture(RenderTarget, FaceTexture, FaceRotation);
-
-	TArray<FColor> Pixels;
-	if (!FGLTFTextureUtility::ReadEncodedPixels(RenderTarget, Pixels, JsonTexture.Encoding)) // TODO: use only encoding as specified by export options
-	{
-		// TODO: report error
-		return FGLTFJsonTextureIndex(INDEX_NONE);
-	}
-
-	JsonTexture.Source = Builder.AddImage(Pixels.GetData(), Size, JsonTexture.Name);
-	JsonTexture.Sampler = Builder.GetOrAddSampler(RenderTargetCube);
-
-	return Builder.AddTexture(JsonTexture);
+	const FGLTFJsonTextureIndex TextureIndex = Builder.AddTexture();
+	Builder.SetupTask<FGLTFTextureRenderTargetCubeTask>(Builder, RenderTargetCube, CubeFace, TextureIndex);
+	return TextureIndex;
 }
diff --git a/Engine/Plugins/Enterprise/GLTFExporter/MigratedFromGit/Source/GLTFExporter/Private/Tasks/GLTFTextureTasks.cpp b/Engine/Plugins/Enterprise/GLTFExporter/MigratedFromGit/Source/GLTFExporter/Private/Tasks/GLTFTextureTasks.cpp
new file mode 100644
--- /dev/null
+++ b/Engine/Plugins/Enterprise/GLTFExporter/MigratedFromGit/Source/GLTFExporter/Private/Tasks/GLTFTextureTasks.cpp
@@ -0,0 +1,165 @@
+// Copyright Epic Games, Inc. All Rights Reserved.
+
+#include "Tasks/GLTFTextureTasks.h"
+#include "Builders/GLTFContainerBuilder.h"
+#include "Converters/GLTFConverterUtility.h"
+#include "Converters/GLTFTextureUtility.h"
+
+void FGLTFTexture2DTask::Complete()
+{
+	FGLTFJsonImageIndex ImageIndex;
+
+	FGLTFJsonTexture& JsonTexture = Builder.GetTexture(TextureIndex);
+	Texture2D->GetName(JsonTexture.Name);
+
+	const FIntPoint Size = { Texture2D->GetSizeX(), Texture2D->GetSizeY() };
+	const bool bPreferSourceExport = Builder.ExportOptions->bExportSourceTextures;
+
+	// NOTE: export of lightmaps and normalmaps via source data is done to workaround issues
+	// with using render data for normalmaps (which are compressed to red & green channels) and
+	// lightmaps (which can otherwise suffer quality-loss due to incorrect gamma transformation).
+
+	// TODO: make sure there are no other special cases than normalmaps and lightmaps that require source-export
+	const bool bRequireSourceExport = Texture2D->IsA<ULightMapTexture2D>() || Texture2D->IsNormalMap();
+
+	if (bRequireSourceExport || bPreferSourceExport)
+	{
+		ERGBFormat RGBFormat;
+		uint32 BitDepth;
+
+		FTextureSource& Source = const_cast<FTextureSource&>(Texture2D->Source);
+
+		if (Source.IsValid() && FGLTFTextureUtility::CanPNGCompressFormat(Source.GetFormat(), RGBFormat, BitDepth))
+		{
+			const void* RawData = Source.LockMip(0);
+			ImageIndex = Builder.AddImage(RawData, Source.CalcMipSize(0), Size, RGBFormat, BitDepth, JsonTexture.Name);
+			Source.UnlockMip(0);
+		}
+
+		if (ImageIndex != INDEX_NONE)
+		{
+			if (FGLTFTextureUtility::HasAnyAdjustment(Texture2D))
+			{
+				Builder.AddWarningMessage(FString::Printf(TEXT("Adjustments for texture %s are not supported when exporting source data"), *JsonTexture.Name));
+			}
+		}
+		else
+		{
+			if (bPreferSourceExport)
+			{
+				Builder.AddWarningMessage(FString::Printf(TEXT("Unable to export source for texture %s, render data will be used as fallback"), *JsonTexture.Name));
+			}
+		}
+	}
+
+	if (ImageIndex == INDEX_NONE && !bRequireSourceExport)
+	{
+		// TODO: both bForceLinearGamma and TargetGamma=2.2f seem to be necessary for the exported images to match the results
+		// from exporting using the texture's source or its platform-data.
+		// It's not entirely clear why gamma must be set to 2.2 instead of 0.0 (which should use the correct gamma anyway),
+		// and why bInForceLinearGamma must also be true.
+		const bool bForceLinearGamma = true;
+		const float TargetGamma = 2.2f;
+
+		const EPixelFormat RenderTargetFormat = FGLTFTextureUtility::IsHDRFormat(Texture2D->GetPixelFormat()) ? PF_FloatRGBA : PF_B8G8R8A8;
+		UTextureRenderTarget2D* RenderTarget = FGLTFTextureUtility::CreateRenderTarget(Size, RenderTargetFormat, bForceLinearGamma);
+		RenderTarget->TargetGamma = TargetGamma;
+
+		FGLTFTextureUtility::DrawTexture(RenderTarget, Texture2D);
+
+		TArray<FColor> Pixels;
+		if (FGLTFTextureUtility::ReadEncodedPixels(RenderTarget, Pixels, JsonTexture.Encoding)) // TODO: use only encoding as specified by export options
+		{
+			ImageIndex = Builder.AddImage(Pixels.GetData(), Size, JsonTexture.Name);
+		}
+	}
+
+	if (ImageIndex == INDEX_NONE)
+	{
+		Builder.AddWarningMessage(FString::Printf(TEXT("Failed to export texture %s"), *JsonTexture.Name));
+		return;
+	}
+
+	JsonTexture.Source = ImageIndex;
+	JsonTexture.Sampler = Builder.GetOrAddSampler(Texture2D);
+}
+
+void FGLTFTextureCubeTask::Complete()
+{
+	FGLTFJsonTexture& JsonTexture = Builder.GetTexture(TextureIndex);
+	JsonTexture.Name = TextureCube->GetName() + TEXT("_") + FGLTFJsonUtility::ToString(FGLTFConverterUtility::ConvertCubeFace(CubeFace));
+
+	// TODO: add optimized "happy path" if cube face doesn't need rotation and has suitable pixel format
+
+	const UTexture2D* FaceTexture = FGLTFTextureUtility::CreateTextureFromCubeFace(TextureCube, CubeFace);
+	if (FaceTexture == nullptr)
+	{
+		// TODO: report error
+		return;
+	}
+
+	const FIntPoint Size = { FaceTexture->GetSizeX(), FaceTexture->GetSizeY() };
+	const EPixelFormat RenderTargetFormat = FGLTFTextureUtility::IsHDRFormat(FaceTexture->GetPixelFormat()) ? PF_FloatRGBA : PF_B8G8R8A8;
+	UTextureRenderTarget2D* RenderTarget = FGLTFTextureUtility::CreateRenderTarget(Size, RenderTargetFormat, true);
+
+	const float FaceRotation = FGLTFTextureUtility::GetCubeFaceRotation(CubeFace);
+	FGLTFTextureUtility::RotateTexture(RenderTarget, FaceTexture, FaceRotation);
+
+	TArray<FColor> Pixels;
+	if (!FGLTFTextureUtility::ReadEncodedPixels(RenderTarget, Pixels, JsonTexture.Encoding)) // TODO: use only encoding as specified by export options
+	{
+		// TODO: report error
+		return;
+	}
+
+	JsonTexture.Source = Builder.AddImage(Pixels.GetData(), Size, JsonTexture.Name);
+	JsonTexture.Sampler = Builder.GetOrAddSampler(TextureCube);
+}
+
+void FGLTFTextureRenderTarget2DTask::Complete()
+{
+	FGLTFJsonTexture& JsonTexture = Builder.GetTexture(TextureIndex);
+	RenderTarget2D->GetName(JsonTexture.Name);
+
+	TArray<FColor> Pixels;
+	if (!FGLTFTextureUtility::ReadEncodedPixels(RenderTarget2D, Pixels, JsonTexture.Encoding)) // TODO: use only encoding as specified by export options
+	{
+		// TODO: report error
+		return;
+	}
+
+	JsonTexture.Source = Builder.AddImage(Pixels.GetData(), { RenderTarget2D->SizeX, RenderTarget2D->SizeY }, JsonTexture.Name);
+	JsonTexture.Sampler = Builder.GetOrAddSampler(RenderTarget2D);
+}
+
+void FGLTFTextureRenderTargetCubeTask::Complete()
+{
+	FGLTFJsonTexture& JsonTexture = Builder.GetTexture(TextureIndex);
+	JsonTexture.Name = RenderTargetCube->GetName() + TEXT("_") + FGLTFJsonUtility::ToString(FGLTFConverterUtility::ConvertCubeFace(CubeFace));
+
+	// TODO: add optimized "happy path" if cube face doesn't need rotation
+
+	const UTexture2D* FaceTexture = FGLTFTextureUtility::CreateTextureFromCubeFace(RenderTargetCube, CubeFace);
+	if (FaceTexture == nullptr)
+	{
+		// TODO: report error
+		return;
+	}
+
+	const FIntPoint Size = { FaceTexture->GetSizeX(), FaceTexture->GetSizeY() };
+	const EPixelFormat RenderTargetFormat = FGLTFTextureUtility::IsHDRFormat(FaceTexture->GetPixelFormat()) ? PF_FloatRGBA : PF_B8G8R8A8;
+	UTextureRenderTarget2D* RenderTarget = FGLTFTextureUtility::CreateRenderTarget(Size, RenderTargetFormat, true);
+
+	const float FaceRotation = FGLTFTextureUtility::GetCubeFaceRotation(CubeFace);
+	FGLTFTextureUtility::RotateTexture(RenderTarget, FaceTexture, FaceRotation);
+
+	TArray<FColor> Pixels;
+	if (!FGLTFTextureUtility::ReadEncodedPixels(RenderTarget, Pixels, JsonTexture.Encoding)) // TODO: use only encoding as specified by export options
+	{
+		// TODO: report error
+		return;
+	}
+
+	JsonTexture.Source = Builder.AddImage(Pixels.GetData(), Size, JsonTexture.Name);
+	JsonTexture.Sampler = Builder.GetOrAddSampler(RenderTargetCube);
+}
diff --git a/Engine/Plugins/Enterprise/GLTFExporter/MigratedFromGit/Source/GLTFExporter/Private/Tasks/GLTFTextureTasks.h b/Engine/Plugins/Enterprise/GLTFExporter/MigratedFromGit/Source/GLTFExporter/Private/Tasks/GLTFTextureTasks.h
new file mode 100644
--- /dev/null
+++ b/Engine/Plugins/Enterprise/GLTFExporter/MigratedFromGit/Source/GLTFExporter/Private/Tasks/GLTFTextureTasks.h
@@ -0,0 +1,115 @@
+// Copyright Epic Games, Inc. All Rights Reserved.
+
+#pragma once
+
+#include "Tasks/GLTFTask.h"
+#include "Builders/GLTFConvertBuilder.h"
+#include "Engine.h"
+
+class FGLTFTexture2DTask : public FGLTFTask
+{
+public:
+
+	FGLTFTexture2DTask(FGLTFConvertBuilder& Builder, const UTexture2D* Texture2D, FGLTFJsonTextureIndex TextureIndex)
+        : FGLTFTask(EGLTFTaskPriority::Texture)
+		, Builder(Builder)
+		, Texture2D(Texture2D)
+		, TextureIndex(TextureIndex)
+	{
+	}
+
+	virtual FString GetName() override
+	{
+		return Texture2D->GetName();
+	}
+
+	virtual void Complete() override;
+
+private:
+
+	FGLTFConvertBuilder& Builder;
+	const UTexture2D* Texture2D;
+	const FGLTFJsonTextureIndex TextureIndex;
+};
+
+class FGLTFTextureCubeTask : public FGLTFTask
+{
+public:
+
+	FGLTFTextureCubeTask(FGLTFConvertBuilder& Builder, const UTextureCube* TextureCube, ECubeFace CubeFace, FGLTFJsonTextureIndex TextureIndex)
+		: FGLTFTask(EGLTFTaskPriority::Texture)
+		, Builder(Builder)
+		, TextureCube(TextureCube)
+		, CubeFace(CubeFace)
+		, TextureIndex(TextureIndex)
+	{
+	}
+
+	virtual FString GetName() override
+	{
+		return TextureCube->GetName();
+	}
+
+	virtual void Complete() override;
+
+private:
+
+	FGLTFConvertBuilder& Builder;
+	const UTextureCube* TextureCube;
+	ECubeFace CubeFace;
+	const FGLTFJsonTextureIndex TextureIndex;
+};
+
+class FGLTFTextureRenderTarget2DTask : public FGLTFTask
+{
+public:
+
+	FGLTFTextureRenderTarget2DTask(FGLTFConvertBuilder& Builder, const UTextureRenderTarget2D* RenderTarget2D, FGLTFJsonTextureIndex TextureIndex)
+		: FGLTFTask(EGLTFTaskPriority::Texture)
+		, Builder(Builder)
+		, RenderTarget2D(RenderTarget2D)
+		, TextureIndex(TextureIndex)
+	{
+	}
+
+	virtual FString GetName() override
+	{
+		return RenderTarget2D->GetName();
+	}
+
+	virtual void Complete() override;
+
+private:
+
+	FGLTFConvertBuilder& Builder;
+	const UTextureRenderTarget2D* RenderTarget2D;
+	const FGLTFJsonTextureIndex TextureIndex;
+};
+
+class FGLTFTextureRenderTargetCubeTask : public FGLTFTask
+{
+public:
+
+	FGLTFTextureRenderTargetCubeTask(FGLTFConvertBuilder& Builder, const UTextureRenderTargetCube* RenderTargetCube, ECubeFace CubeFace, FGLTFJsonTextureIndex TextureIndex)
+		: FGLTFTask(EGLTFTaskPriority::Texture)
+		, Builder(Builder)
+		, RenderTargetCube(RenderTargetCube)
+		, CubeFace(CubeFace)
+		, TextureIndex(TextureIndex)
+	{
+	}
+
+	virtual FString GetName() override
+	{
+		return RenderTargetCube->GetName();
+	}
+
+	virtual void Complete() override;
+
+private:
+
+	FGLTFConvertBuilder& Builder;
+	const UTextureRenderTargetCube* RenderTargetCube;
+	ECubeFace CubeFace;
+	const FGLTFJsonTextureIndex TextureIndex;
+};
