[GLTFExporter] refactor: move mesh section converters to separate file

commit 2ff33e5d2a8bfdffe01a31b5738d41fd20b1ce93
Date: Sun, 4 Oct 2020 14:07:54 +0200

    refactor: move mesh section converters to separate file
    (cherry picked from commit 7e78164a0de82cefa2f78b6796254798fbaf6c85​)

#skipci

diff --git a/Engine/Plugins/Enterprise/GLTFExporter/MigratedFromGit/Source/GLTFExporter/Private/Converters/GLTFMeshConverters.h b/Engine/Plugins/Enterprise/GLTFExporter/MigratedFromGit/Source/GLTFExporter/Private/Converters/GLTFMeshConverters.h
--- a/Engine/Plugins/Enterprise/GLTFExporter/MigratedFromGit/Source/GLTFExporter/Private/Converters/GLTFMeshConverters.h
+++ b/Engine/Plugins/Enterprise/GLTFExporter/MigratedFromGit/Source/GLTFExporter/Private/Converters/GLTFMeshConverters.h
@@ -5,7 +5,7 @@
 #include "Json/GLTFJsonIndex.h"
 #include "Converters/GLTFConverter.h"
 #include "Converters/GLTFBuilderContext.h"
-#include "Converters/GLTFMeshSection.h"
+#include "Converters/GLTFMeshSectionConverters.h"
 #include "Converters/GLTFMaterialArray.h"
 #include "Engine.h"
 
@@ -15,17 +15,6 @@ class TGLTFMeshConverter : public FGLTFBuilderContext, public TGLTFConverter<FGL
 	using FGLTFBuilderContext::FGLTFBuilderContext;
 };
 
-template <typename MeshLODType, typename MaterialIndexType>
-class TGLTFMeshSectionConverter final : public TGLTFConverter<const FGLTFMeshSection*, const MeshLODType*, const MaterialIndexType>
-{
-	TArray<TUniquePtr<FGLTFMeshSection>> Outputs;
-
-	const FGLTFMeshSection* Convert(const MeshLODType* MeshLOD, const MaterialIndexType MaterialIndex)
-	{
-		return Outputs.Add_GetRef(MakeUnique<FGLTFMeshSection>(MeshLOD, MaterialIndex)).Get();
-	}
-};
-
 class FGLTFStaticMeshConverter final : public TGLTFMeshConverter<const UStaticMesh*, int32, const FColorVertexBuffer*, FGLTFMaterialArray>
 {
 	using TGLTFMeshConverter::TGLTFMeshConverter;
@@ -34,7 +23,7 @@ class FGLTFStaticMeshConverter final : public TGLTFMeshConverter<const UStaticMe
 
 	virtual FGLTFJsonMeshIndex Convert(const UStaticMesh* StaticMesh, int32 LODIndex, const FColorVertexBuffer* OverrideVertexColors, FGLTFMaterialArray OverrideMaterials) override;
 
-	TGLTFMeshSectionConverter<FStaticMeshLODResources, int32> MeshSectionConverter;
+	FGLTFStaticMeshSectionConverter MeshSectionConverter;
 };
 
 class FGLTFSkeletalMeshConverter final : public TGLTFMeshConverter<const USkeletalMesh*, int32, const FColorVertexBuffer*, const FSkinWeightVertexBuffer*, FGLTFMaterialArray>
@@ -45,5 +34,5 @@ class FGLTFSkeletalMeshConverter final : public TGLTFMeshConverter<const USkelet
 
 	virtual FGLTFJsonMeshIndex Convert(const USkeletalMesh* SkeletalMesh, int32 LODIndex, const FColorVertexBuffer* OverrideVertexColors, const FSkinWeightVertexBuffer* OverrideSkinWeights, FGLTFMaterialArray OverrideMaterials) override;
 
-	TGLTFMeshSectionConverter<FSkeletalMeshLODRenderData, uint16> MeshSectionConverter;
+	FGLTFSkeletalMeshSectionConverter MeshSectionConverter;
 };
diff --git a/Engine/Plugins/Enterprise/GLTFExporter/MigratedFromGit/Source/GLTFExporter/Private/Converters/GLTFMeshSectionConverters.h b/Engine/Plugins/Enterprise/GLTFExporter/MigratedFromGit/Source/GLTFExporter/Private/Converters/GLTFMeshSectionConverters.h
new file mode 100644
--- /dev/null
+++ b/Engine/Plugins/Enterprise/GLTFExporter/MigratedFromGit/Source/GLTFExporter/Private/Converters/GLTFMeshSectionConverters.h
@@ -0,0 +1,21 @@
+// Copyright Epic Games, Inc. All Rights Reserved.
+
+#pragma once
+
+#include "Converters/GLTFConverter.h"
+#include "Converters/GLTFMeshSection.h"
+#include "Engine.h"
+
+template <typename MeshLODType, typename MaterialIndexType>
+class TGLTFMeshSectionConverter final : public TGLTFConverter<const FGLTFMeshSection*, const MeshLODType*, const MaterialIndexType>
+{
+	TArray<TUniquePtr<FGLTFMeshSection>> Outputs;
+
+	const FGLTFMeshSection* Convert(const MeshLODType* MeshLOD, const MaterialIndexType MaterialIndex)
+	{
+		return Outputs.Add_GetRef(MakeUnique<FGLTFMeshSection>(MeshLOD, MaterialIndex)).Get();
+	}
+};
+
+typedef TGLTFMeshSectionConverter<FStaticMeshLODResources, int32> FGLTFStaticMeshSectionConverter;
+typedef TGLTFMeshSectionConverter<FSkeletalMeshLODRenderData, uint16> FGLTFSkeletalMeshSectionConverter;
