[GLTFExporter] feat: add mesh tasks and change mesh converters to use them

commit aefe84fea62704117e317ee4d3e7269701e0f55e
Date: Sun, 4 Oct 2020 14:22:41 +0200

    feat: add mesh tasks and change mesh converters to use them
    (cherry picked from commit 79eebf0295053d2c3b9732ba93cd856d65a8f209​)

#skipci

diff --git a/Engine/Plugins/Enterprise/GLTFExporter/MigratedFromGit/Source/GLTFExporter/Private/Converters/GLTFMeshConverters.cpp b/Engine/Plugins/Enterprise/GLTFExporter/MigratedFromGit/Source/GLTFExporter/Private/Converters/GLTFMeshConverters.cpp
--- a/Engine/Plugins/Enterprise/GLTFExporter/MigratedFromGit/Source/GLTFExporter/Private/Converters/GLTFMeshConverters.cpp
+++ b/Engine/Plugins/Enterprise/GLTFExporter/MigratedFromGit/Source/GLTFExporter/Private/Converters/GLTFMeshConverters.cpp
@@ -1,9 +1,8 @@
 // Copyright Epic Games, Inc. All Rights Reserved.
 
 #include "Converters/GLTFMeshConverters.h"
-#include "Converters/GLTFConverterUtility.h"
-#include "Converters/GLTFNameUtility.h"
 #include "Builders/GLTFConvertBuilder.h"
+#include "Tasks/GLTFMeshTasks.h"
 #include "Rendering/SkeletalMeshRenderData.h"
 
 void FGLTFStaticMeshConverter::Sanitize(const UStaticMesh*& StaticMesh, int32& LODIndex, const FColorVertexBuffer*& OverrideVertexColors, FGLTFMaterialArray& OverrideMaterials)
@@ -40,57 +39,9 @@ FGLTFJsonMeshIndex FGLTFStaticMeshConverter::Convert(const UStaticMesh* StaticMe
 		return FGLTFJsonMeshIndex(INDEX_NONE);
 	}
 
-	const FStaticMeshLODResources& MeshLOD = StaticMesh->GetLODForExport(LODIndex);
-
-	FGLTFJsonMesh JsonMesh;
-	JsonMesh.Name = FGLTFNameUtility::GetName(StaticMesh, LODIndex);
-
-	const FPositionVertexBuffer* PositionBuffer = &MeshLOD.VertexBuffers.PositionVertexBuffer;
-	const FStaticMeshVertexBuffer* VertexBuffer = &MeshLOD.VertexBuffers.StaticMeshVertexBuffer;
-	const FColorVertexBuffer* ColorBuffer = OverrideVertexColors != nullptr ? OverrideVertexColors : &MeshLOD.VertexBuffers.ColorVertexBuffer;
-
-	const int32 MaterialCount = StaticMesh->StaticMaterials.Num();
-	JsonMesh.Primitives.AddDefaulted(MaterialCount);
-
-	for (int32 MaterialIndex = 0; MaterialIndex < MaterialCount; ++MaterialIndex)
-	{
-		const FGLTFMeshSection* ConvertedSection = MeshSectionConverter.GetOrAdd(&MeshLOD, MaterialIndex);
-
-		FGLTFJsonPrimitive& JsonPrimitive = JsonMesh.Primitives[MaterialIndex];
-		JsonPrimitive.Indices = Builder.GetOrAddIndexAccessor(ConvertedSection);
-
-		JsonPrimitive.Attributes.Position = Builder.GetOrAddPositionAccessor(ConvertedSection, PositionBuffer);
-		if (JsonPrimitive.Attributes.Position == INDEX_NONE)
-		{
-			// TODO: report warning
-		}
-
-		if (Builder.ExportOptions->bExportVertexColors)
-		{
-			JsonPrimitive.Attributes.Color0 = Builder.GetOrAddColorAccessor(ConvertedSection, ColorBuffer);
-		}
-
-		JsonPrimitive.Attributes.Normal = Builder.GetOrAddNormalAccessor(ConvertedSection, VertexBuffer);
-		JsonPrimitive.Attributes.Tangent = Builder.GetOrAddTangentAccessor(ConvertedSection, VertexBuffer);
-
-		const uint32 UVCount = VertexBuffer->GetNumTexCoords();
-		JsonPrimitive.Attributes.TexCoords.AddUninitialized(UVCount);
-
-		for (uint32 UVIndex = 0; UVIndex < UVCount; ++UVIndex)
-		{
-			JsonPrimitive.Attributes.TexCoords[UVIndex] = Builder.GetOrAddUVAccessor(ConvertedSection, VertexBuffer, UVIndex);
-		}
-
-		const UMaterialInterface* Material = OverrideMaterials.GetOverride(StaticMesh->StaticMaterials, MaterialIndex);
-		if (Material == nullptr)
-		{
-			Material = UMaterial::GetDefaultMaterial(MD_Surface);
-		}
-
-		JsonPrimitive.Material = Builder.GetOrAddMaterial(Material);
-	}
-
-	return Builder.AddMesh(JsonMesh);
+	const FGLTFJsonMeshIndex MeshIndex = Builder.AddMesh();
+	Builder.SetupTask<FGLTFStaticMeshTask>(Builder, MeshSectionConverter, StaticMesh, LODIndex, OverrideVertexColors, OverrideMaterials, MeshIndex);
+	return MeshIndex;
 }
 
 void FGLTFSkeletalMeshConverter::Sanitize(const USkeletalMesh*& SkeletalMesh, int32& LODIndex, const FColorVertexBuffer*& OverrideVertexColors, const FSkinWeightVertexBuffer*& OverrideSkinWeights, FGLTFMaterialArray& OverrideMaterials)
@@ -141,69 +92,7 @@ FGLTFJsonMeshIndex FGLTFSkeletalMeshConverter::Convert(const USkeletalMesh* Skel
 		return FGLTFJsonMeshIndex(INDEX_NONE);
 	}
 
-	const FSkeletalMeshLODRenderData& MeshLOD = RenderData->LODRenderData[LODIndex];
-
-	FGLTFJsonMesh JsonMesh;
-	JsonMesh.Name = FGLTFNameUtility::GetName(SkeletalMesh, LODIndex);
-
-	const FPositionVertexBuffer* PositionBuffer = &MeshLOD.StaticVertexBuffers.PositionVertexBuffer;
-	const FStaticMeshVertexBuffer* VertexBuffer = &MeshLOD.StaticVertexBuffers.StaticMeshVertexBuffer;
-	const FColorVertexBuffer* ColorBuffer = OverrideVertexColors != nullptr ? OverrideVertexColors : &MeshLOD.StaticVertexBuffers.ColorVertexBuffer;
-	const FSkinWeightVertexBuffer* SkinWeightBuffer = OverrideSkinWeights != nullptr ? OverrideSkinWeights : MeshLOD.GetSkinWeightVertexBuffer();
-
-	const uint16 MaterialCount = SkeletalMesh->Materials.Num();
-	JsonMesh.Primitives.AddDefaulted(MaterialCount);
-
-	for (uint16 MaterialIndex = 0; MaterialIndex < MaterialCount; ++MaterialIndex)
-	{
-		const FGLTFMeshSection* ConvertedSection = MeshSectionConverter.GetOrAdd(&MeshLOD, MaterialIndex);
-
-		FGLTFJsonPrimitive& JsonPrimitive = JsonMesh.Primitives[MaterialIndex];
-		JsonPrimitive.Indices = Builder.GetOrAddIndexAccessor(ConvertedSection);
-
-		JsonPrimitive.Attributes.Position = Builder.GetOrAddPositionAccessor(ConvertedSection, PositionBuffer);
-		if (JsonPrimitive.Attributes.Position == INDEX_NONE)
-		{
-			// TODO: report warning
-		}
-
-		if (Builder.ExportOptions->bExportVertexColors)
-		{
-			JsonPrimitive.Attributes.Color0 = Builder.GetOrAddColorAccessor(ConvertedSection, ColorBuffer);
-		}
-
-		JsonPrimitive.Attributes.Normal = Builder.GetOrAddNormalAccessor(ConvertedSection, VertexBuffer);
-		JsonPrimitive.Attributes.Tangent = Builder.GetOrAddTangentAccessor(ConvertedSection, VertexBuffer);
-
-		const uint32 UVCount = VertexBuffer->GetNumTexCoords();
-		JsonPrimitive.Attributes.TexCoords.AddUninitialized(UVCount);
-
-		for (uint32 UVIndex = 0; UVIndex < UVCount; ++UVIndex)
-		{
-			JsonPrimitive.Attributes.TexCoords[UVIndex] = Builder.GetOrAddUVAccessor(ConvertedSection, VertexBuffer, UVIndex);
-		}
-
-		if (Builder.ExportOptions->bExportVertexSkinWeights)
-		{
-			const uint32 GroupCount = (SkinWeightBuffer->GetMaxBoneInfluences() + 3) / 4;
-			JsonPrimitive.Attributes.Joints.AddUninitialized(GroupCount);
-			JsonPrimitive.Attributes.Weights.AddUninitialized(GroupCount);
-
-			for (uint32 GroupIndex = 0; GroupIndex < GroupCount; ++GroupIndex)
-			{
-				JsonPrimitive.Attributes.Joints[GroupIndex] = Builder.GetOrAddJointAccessor(ConvertedSection, SkinWeightBuffer, GroupIndex * 4);
-				JsonPrimitive.Attributes.Weights[GroupIndex] = Builder.GetOrAddWeightAccessor(ConvertedSection, SkinWeightBuffer, GroupIndex * 4);
-			}
-		}
-
-		const UMaterialInterface* Material = OverrideMaterials.GetOverride(SkeletalMesh->Materials, MaterialIndex);
-		if (Material == nullptr)
-		{
-			Material = UMaterial::GetDefaultMaterial(MD_Surface);
-		}
-
-		JsonPrimitive.Material = Builder.GetOrAddMaterial(Material);
-	}
-
-	return Builder.AddMesh(JsonMesh);
+	const FGLTFJsonMeshIndex MeshIndex = Builder.AddMesh();
+	Builder.SetupTask<FGLTFSkeletalMeshTask>(Builder, MeshSectionConverter, SkeletalMesh, LODIndex, OverrideVertexColors, OverrideSkinWeights, OverrideMaterials, MeshIndex);
+	return MeshIndex;
 }
diff --git a/Engine/Plugins/Enterprise/GLTFExporter/MigratedFromGit/Source/GLTFExporter/Private/Tasks/GLTFMeshTasks.cpp b/Engine/Plugins/Enterprise/GLTFExporter/MigratedFromGit/Source/GLTFExporter/Private/Tasks/GLTFMeshTasks.cpp
new file mode 100644
--- /dev/null
+++ b/Engine/Plugins/Enterprise/GLTFExporter/MigratedFromGit/Source/GLTFExporter/Private/Tasks/GLTFMeshTasks.cpp
@@ -0,0 +1,128 @@
+// Copyright Epic Games, Inc. All Rights Reserved.
+
+#include "Tasks/GLTFMeshTasks.h"
+#include "Converters/GLTFConverterUtility.h"
+#include "Converters/GLTFNameUtility.h"
+#include "Builders/GLTFConvertBuilder.h"
+#include "Rendering/SkeletalMeshRenderData.h"
+
+void FGLTFStaticMeshTask::Complete()
+{
+	const FStaticMeshLODResources& MeshLOD = StaticMesh->GetLODForExport(LODIndex);
+
+	FGLTFJsonMesh& JsonMesh = Builder.GetMesh(MeshIndex);
+	JsonMesh.Name = FGLTFNameUtility::GetName(StaticMesh, LODIndex);
+
+	const FPositionVertexBuffer* PositionBuffer = &MeshLOD.VertexBuffers.PositionVertexBuffer;
+	const FStaticMeshVertexBuffer* VertexBuffer = &MeshLOD.VertexBuffers.StaticMeshVertexBuffer;
+	const FColorVertexBuffer* ColorBuffer = OverrideVertexColors != nullptr ? OverrideVertexColors : &MeshLOD.VertexBuffers.ColorVertexBuffer;
+
+	const int32 MaterialCount = StaticMesh->StaticMaterials.Num();
+	JsonMesh.Primitives.AddDefaulted(MaterialCount);
+
+	for (int32 MaterialIndex = 0; MaterialIndex < MaterialCount; ++MaterialIndex)
+	{
+		const FGLTFMeshSection* ConvertedSection = MeshSectionConverter.GetOrAdd(&MeshLOD, MaterialIndex);
+
+		FGLTFJsonPrimitive& JsonPrimitive = JsonMesh.Primitives[MaterialIndex];
+		JsonPrimitive.Indices = Builder.GetOrAddIndexAccessor(ConvertedSection);
+
+		JsonPrimitive.Attributes.Position = Builder.GetOrAddPositionAccessor(ConvertedSection, PositionBuffer);
+		if (JsonPrimitive.Attributes.Position == INDEX_NONE)
+		{
+			// TODO: report warning
+		}
+
+		if (Builder.ExportOptions->bExportVertexColors)
+		{
+			JsonPrimitive.Attributes.Color0 = Builder.GetOrAddColorAccessor(ConvertedSection, ColorBuffer);
+		}
+
+		JsonPrimitive.Attributes.Normal = Builder.GetOrAddNormalAccessor(ConvertedSection, VertexBuffer);
+		JsonPrimitive.Attributes.Tangent = Builder.GetOrAddTangentAccessor(ConvertedSection, VertexBuffer);
+
+		const uint32 UVCount = VertexBuffer->GetNumTexCoords();
+		JsonPrimitive.Attributes.TexCoords.AddUninitialized(UVCount);
+
+		for (uint32 UVIndex = 0; UVIndex < UVCount; ++UVIndex)
+		{
+			JsonPrimitive.Attributes.TexCoords[UVIndex] = Builder.GetOrAddUVAccessor(ConvertedSection, VertexBuffer, UVIndex);
+		}
+
+		const UMaterialInterface* Material = OverrideMaterials.GetOverride(StaticMesh->StaticMaterials, MaterialIndex);
+		if (Material == nullptr)
+		{
+			Material = UMaterial::GetDefaultMaterial(MD_Surface);
+		}
+
+		JsonPrimitive.Material = Builder.GetOrAddMaterial(Material);
+	}
+}
+
+void FGLTFSkeletalMeshTask::Complete()
+{
+	const FSkeletalMeshRenderData* RenderData = SkeletalMesh->GetResourceForRendering();
+	const FSkeletalMeshLODRenderData& MeshLOD = RenderData->LODRenderData[LODIndex];
+
+	FGLTFJsonMesh& JsonMesh = Builder.GetMesh(MeshIndex);
+	JsonMesh.Name = FGLTFNameUtility::GetName(SkeletalMesh, LODIndex);
+
+	const FPositionVertexBuffer* PositionBuffer = &MeshLOD.StaticVertexBuffers.PositionVertexBuffer;
+	const FStaticMeshVertexBuffer* VertexBuffer = &MeshLOD.StaticVertexBuffers.StaticMeshVertexBuffer;
+	const FColorVertexBuffer* ColorBuffer = OverrideVertexColors != nullptr ? OverrideVertexColors : &MeshLOD.StaticVertexBuffers.ColorVertexBuffer;
+	const FSkinWeightVertexBuffer* SkinWeightBuffer = OverrideSkinWeights != nullptr ? OverrideSkinWeights : MeshLOD.GetSkinWeightVertexBuffer();
+
+	const uint16 MaterialCount = SkeletalMesh->Materials.Num();
+	JsonMesh.Primitives.AddDefaulted(MaterialCount);
+
+	for (uint16 MaterialIndex = 0; MaterialIndex < MaterialCount; ++MaterialIndex)
+	{
+		const FGLTFMeshSection* ConvertedSection = MeshSectionConverter.GetOrAdd(&MeshLOD, MaterialIndex);
+
+		FGLTFJsonPrimitive& JsonPrimitive = JsonMesh.Primitives[MaterialIndex];
+		JsonPrimitive.Indices = Builder.GetOrAddIndexAccessor(ConvertedSection);
+
+		JsonPrimitive.Attributes.Position = Builder.GetOrAddPositionAccessor(ConvertedSection, PositionBuffer);
+		if (JsonPrimitive.Attributes.Position == INDEX_NONE)
+		{
+			// TODO: report warning
+		}
+
+		if (Builder.ExportOptions->bExportVertexColors)
+		{
+			JsonPrimitive.Attributes.Color0 = Builder.GetOrAddColorAccessor(ConvertedSection, ColorBuffer);
+		}
+
+		JsonPrimitive.Attributes.Normal = Builder.GetOrAddNormalAccessor(ConvertedSection, VertexBuffer);
+		JsonPrimitive.Attributes.Tangent = Builder.GetOrAddTangentAccessor(ConvertedSection, VertexBuffer);
+
+		const uint32 UVCount = VertexBuffer->GetNumTexCoords();
+		JsonPrimitive.Attributes.TexCoords.AddUninitialized(UVCount);
+
+		for (uint32 UVIndex = 0; UVIndex < UVCount; ++UVIndex)
+		{
+			JsonPrimitive.Attributes.TexCoords[UVIndex] = Builder.GetOrAddUVAccessor(ConvertedSection, VertexBuffer, UVIndex);
+		}
+
+		if (Builder.ExportOptions->bExportVertexSkinWeights)
+		{
+			const uint32 GroupCount = (SkinWeightBuffer->GetMaxBoneInfluences() + 3) / 4;
+			JsonPrimitive.Attributes.Joints.AddUninitialized(GroupCount);
+			JsonPrimitive.Attributes.Weights.AddUninitialized(GroupCount);
+
+			for (uint32 GroupIndex = 0; GroupIndex < GroupCount; ++GroupIndex)
+			{
+				JsonPrimitive.Attributes.Joints[GroupIndex] = Builder.GetOrAddJointAccessor(ConvertedSection, SkinWeightBuffer, GroupIndex * 4);
+				JsonPrimitive.Attributes.Weights[GroupIndex] = Builder.GetOrAddWeightAccessor(ConvertedSection, SkinWeightBuffer, GroupIndex * 4);
+			}
+		}
+
+		const UMaterialInterface* Material = OverrideMaterials.GetOverride(SkeletalMesh->Materials, MaterialIndex);
+		if (Material == nullptr)
+		{
+			Material = UMaterial::GetDefaultMaterial(MD_Surface);
+		}
+
+		JsonPrimitive.Material = Builder.GetOrAddMaterial(Material);
+	}
+}
diff --git a/Engine/Plugins/Enterprise/GLTFExporter/MigratedFromGit/Source/GLTFExporter/Private/Tasks/GLTFMeshTasks.h b/Engine/Plugins/Enterprise/GLTFExporter/MigratedFromGit/Source/GLTFExporter/Private/Tasks/GLTFMeshTasks.h
new file mode 100644
--- /dev/null
+++ b/Engine/Plugins/Enterprise/GLTFExporter/MigratedFromGit/Source/GLTFExporter/Private/Tasks/GLTFMeshTasks.h
@@ -0,0 +1,79 @@
+// Copyright Epic Games, Inc. All Rights Reserved.
+
+#pragma once
+
+#include "Tasks/GLTFTask.h"
+#include "Builders/GLTFConvertBuilder.h"
+#include "Converters/GLTFMeshSectionConverters.h"
+#include "Converters/GLTFMaterialArray.h"
+#include "Engine.h"
+
+class FGLTFStaticMeshTask : public FGLTFTask
+{
+public:
+
+	FGLTFStaticMeshTask(FGLTFConvertBuilder& Builder, FGLTFStaticMeshSectionConverter& MeshSectionConverter, const UStaticMesh* StaticMesh, int32 LODIndex, const FColorVertexBuffer* OverrideVertexColors, FGLTFMaterialArray OverrideMaterials, FGLTFJsonMeshIndex MeshIndex)
+		: FGLTFTask(EGLTFTaskPriority::Mesh)
+		, Builder(Builder)
+		, MeshSectionConverter(MeshSectionConverter)
+		, StaticMesh(StaticMesh)
+		, LODIndex(LODIndex)
+		, OverrideVertexColors(OverrideVertexColors)
+		, OverrideMaterials(OverrideMaterials)
+        , MeshIndex(MeshIndex)
+	{
+	}
+
+	virtual FString GetName() override
+	{
+		return StaticMesh->GetName();
+	}
+
+	virtual void Complete() override;
+
+private:
+
+	FGLTFConvertBuilder& Builder;
+	FGLTFStaticMeshSectionConverter& MeshSectionConverter;
+	const UStaticMesh* StaticMesh;
+	int32 LODIndex;
+	const FColorVertexBuffer* OverrideVertexColors;
+	FGLTFMaterialArray OverrideMaterials;
+	const FGLTFJsonMeshIndex MeshIndex;
+};
+
+class FGLTFSkeletalMeshTask : public FGLTFTask
+{
+public:
+
+	FGLTFSkeletalMeshTask(FGLTFConvertBuilder& Builder, FGLTFSkeletalMeshSectionConverter& MeshSectionConverter, const USkeletalMesh* SkeletalMesh, int32 LODIndex, const FColorVertexBuffer* OverrideVertexColors, const FSkinWeightVertexBuffer* OverrideSkinWeights, FGLTFMaterialArray OverrideMaterials, FGLTFJsonMeshIndex MeshIndex)
+		: FGLTFTask(EGLTFTaskPriority::Mesh)
+		, Builder(Builder)
+		, MeshSectionConverter(MeshSectionConverter)
+		, SkeletalMesh(SkeletalMesh)
+		, LODIndex(LODIndex)
+		, OverrideVertexColors(OverrideVertexColors)
+		, OverrideSkinWeights(OverrideSkinWeights)
+		, OverrideMaterials(OverrideMaterials)
+		, MeshIndex(MeshIndex)
+	{
+	}
+
+	virtual FString GetName() override
+	{
+		return SkeletalMesh->GetName();
+	}
+
+	virtual void Complete() override;
+
+private:
+
+	FGLTFConvertBuilder& Builder;
+	FGLTFSkeletalMeshSectionConverter& MeshSectionConverter;
+	const USkeletalMesh* SkeletalMesh;
+	int32 LODIndex;
+	const FColorVertexBuffer* OverrideVertexColors;
+	const FSkinWeightVertexBuffer* OverrideSkinWeights;
+	FGLTFMaterialArray OverrideMaterials;
+	const FGLTFJsonMeshIndex MeshIndex;
+};
