[GLTFExporter] feat: add animation task and change animation converter to use it

commit 365f24a67ede491a0e9090f02e203aa4f0fa7656
Date: Sun, 4 Oct 2020 14:49:07 +0200

    feat: add animation task and change animation converter to use it
    (cherry picked from commit d8ac2cc579ca51b2bebb96154964eb890fe4ca79​)

#skipci

diff --git a/Engine/Plugins/Enterprise/GLTFExporter/MigratedFromGit/Source/GLTFExporter/Private/Converters/GLTFAnimationConverters.cpp b/Engine/Plugins/Enterprise/GLTFExporter/MigratedFromGit/Source/GLTFExporter/Private/Converters/GLTFAnimationConverters.cpp
--- a/Engine/Plugins/Enterprise/GLTFExporter/MigratedFromGit/Source/GLTFExporter/Private/Converters/GLTFAnimationConverters.cpp
+++ b/Engine/Plugins/Enterprise/GLTFExporter/MigratedFromGit/Source/GLTFExporter/Private/Converters/GLTFAnimationConverters.cpp
@@ -2,15 +2,13 @@
 
 #include "Converters/GLTFAnimationConverters.h"
 #include "Converters/GLTFConverterUtility.h"
-#include "Converters/GLTFNameUtility.h"
-#include "Converters/GLTFBoneUtility.h"
 #include "Builders/GLTFConvertBuilder.h"
+#include "Tasks/GLTFAnimationTasks.h"
 #include "Animation/AnimSequence.h"
 
 FGLTFJsonAnimationIndex FGLTFAnimationConverter::Convert(FGLTFJsonNodeIndex RootNode, const USkeletalMesh* SkeletalMesh, const UAnimSequence* AnimSequence)
 {
-	const int32 FrameCount = AnimSequence->GetRawNumberOfFrames();
-	if (FrameCount < 0)
+	if (AnimSequence->GetRawNumberOfFrames() < 0)
 	{
 		// TODO: report warning
 		return FGLTFJsonAnimationIndex(INDEX_NONE);
@@ -29,168 +27,9 @@ FGLTFJsonAnimationIndex FGLTFAnimationConverter::Convert(FGLTFJsonNodeIndex Root
 		return FGLTFJsonAnimationIndex(INDEX_NONE);
 	}
 
-	FGLTFJsonAnimation JsonAnimation;
-	AnimSequence->GetName(JsonAnimation.Name);
-
-	TArray<float> Timestamps;
-	Timestamps.AddUninitialized(FrameCount);
-
-	for (int32 Frame = 0; Frame < FrameCount; ++Frame)
-	{
-		Timestamps[Frame] = AnimSequence->GetTimeAtFrame(Frame);
-	}
-
-	FGLTFJsonAccessor JsonInputAccessor;
-	JsonInputAccessor.BufferView = Builder.AddBufferView(Timestamps);
-	JsonInputAccessor.ComponentType = EGLTFJsonComponentType::F32;
-	JsonInputAccessor.Type = EGLTFJsonAccessorType::Scalar;
-	JsonInputAccessor.MinMaxLength = 1;
-	JsonInputAccessor.Min[0] = 0;
-
-	FBoneContainer BoneContainer;
-	if (Builder.ExportOptions->bRetargetBoneTransforms)
-	{
-		FGLTFBoneUtility::InitializeToSkeleton(BoneContainer, Skeleton);
-	}
-
-	const EGLTFJsonInterpolation Interpolation = FGLTFConverterUtility::ConvertInterpolation(AnimSequence->Interpolation);
-	const TArray<FName>& TrackNames = AnimSequence->GetAnimationTrackNames();
-	const int32 TrackCount = TrackNames.Num();
-
-	for (int32 TrackIndex = 0; TrackIndex < TrackCount; ++TrackIndex)
-	{
-		const FRawAnimSequenceTrack& Track = AnimSequence->GetRawAnimationTrack(TrackIndex);
-		const TArray<FVector>& KeyPositions = Track.PosKeys;
-		const TArray<FQuat>& KeyRotations = Track.RotKeys;
-		const TArray<FVector>& KeyScales = Track.ScaleKeys;
-
-		const int32 MaxKeys = FMath::Max3(KeyPositions.Num(), KeyRotations.Num(), KeyScales.Num());
-		if (MaxKeys == 0)
-		{
-			continue;
-		}
-
-		TArray<FTransform> KeyTransforms;
-		KeyTransforms.AddUninitialized(MaxKeys);
-
-		for (int32 Key = 0; Key < KeyTransforms.Num(); ++Key)
-		{
-			const FVector& KeyPosition = KeyPositions.IsValidIndex(Key) ? KeyPositions[Key] : FVector::ZeroVector;
-			const FQuat& KeyRotation = KeyRotations.IsValidIndex(Key) ? KeyRotations[Key] : FQuat::Identity;
-			const FVector& KeyScale = KeyScales.IsValidIndex(Key) ? KeyScales[Key] : FVector::OneVector;
-			KeyTransforms[Key] = { KeyRotation, KeyPosition, KeyScale };
-		}
-
-		const int32 SkeletonBoneIndex = AnimSequence->GetSkeletonIndexFromRawDataTrackIndex(TrackIndex);
-		const int32 BoneIndex = const_cast<USkeleton*>(Skeleton)->GetMeshBoneIndexFromSkeletonBoneIndex(SkeletalMesh, SkeletonBoneIndex);
-		const FGLTFJsonNodeIndex NodeIndex = Builder.GetOrAddNode(RootNode, SkeletalMesh, BoneIndex);
-
-		if (Builder.ExportOptions->bRetargetBoneTransforms && Skeleton->GetBoneTranslationRetargetingMode(SkeletonBoneIndex) != EBoneTranslationRetargetingMode::Animation)
-		{
-			for (int32 Key = 0; Key < KeyTransforms.Num(); ++Key)
-			{
-				FGLTFBoneUtility::RetargetTransform(AnimSequence, KeyTransforms[Key], SkeletonBoneIndex, BoneIndex, BoneContainer);
-			}
-		}
-
-		if (KeyPositions.Num() > 0)
-		{
-			TArray<FGLTFJsonVector3> Translations;
-			Translations.AddUninitialized(KeyPositions.Num());
-
-			for (int32 Key = 0; Key < KeyPositions.Num(); ++Key)
-			{
-				const FVector KeyPosition = KeyTransforms[Key].GetTranslation();
-				Translations[Key] = FGLTFConverterUtility::ConvertPosition(KeyPosition, Builder.ExportOptions->ExportScale);
-			}
-
-			JsonInputAccessor.Count = Translations.Num();
-			JsonInputAccessor.Max[0] = Timestamps[Translations.Num() - 1];
-
-			FGLTFJsonAccessor JsonOutputAccessor;
-			JsonOutputAccessor.BufferView = Builder.AddBufferView(Translations);
-			JsonOutputAccessor.ComponentType = EGLTFJsonComponentType::F32;
-			JsonOutputAccessor.Count = Translations.Num();
-			JsonOutputAccessor.Type = EGLTFJsonAccessorType::Vec3;
-
-			FGLTFJsonAnimationSampler JsonSampler;
-			JsonSampler.Input = Builder.AddAccessor(JsonInputAccessor);
-			JsonSampler.Output = Builder.AddAccessor(JsonOutputAccessor);
-			JsonSampler.Interpolation = Interpolation;
-
-			FGLTFJsonAnimationChannel JsonChannel;
-			JsonChannel.Sampler = FGLTFJsonAnimationSamplerIndex(JsonAnimation.Samplers.Add(JsonSampler));
-			JsonChannel.Target.Path = EGLTFJsonTargetPath::Translation;
-			JsonChannel.Target.Node = NodeIndex;
-			JsonAnimation.Channels.Add(JsonChannel);
-		}
-
-		if (KeyRotations.Num() > 0)
-		{
-			TArray<FGLTFJsonQuaternion> Rotations;
-			Rotations.AddUninitialized(KeyRotations.Num());
-
-			for (int32 Key = 0; Key < KeyRotations.Num(); ++Key)
-			{
-				const FQuat KeyRotation = KeyTransforms[Key].GetRotation();
-				Rotations[Key] = FGLTFConverterUtility::ConvertRotation(KeyRotation);
-			}
-
-			JsonInputAccessor.Count = Rotations.Num();
-			JsonInputAccessor.Max[0] = Timestamps[Rotations.Num() - 1];
-
-			FGLTFJsonAccessor JsonOutputAccessor;
-			JsonOutputAccessor.BufferView = Builder.AddBufferView(Rotations);
-			JsonOutputAccessor.ComponentType = EGLTFJsonComponentType::F32;
-			JsonOutputAccessor.Count = Rotations.Num();
-			JsonOutputAccessor.Type = EGLTFJsonAccessorType::Vec4;
-
-			FGLTFJsonAnimationSampler JsonSampler;
-			JsonSampler.Input = Builder.AddAccessor(JsonInputAccessor);;
-			JsonSampler.Output = Builder.AddAccessor(JsonOutputAccessor);
-			JsonSampler.Interpolation = Interpolation;
-
-			FGLTFJsonAnimationChannel JsonChannel;
-			JsonChannel.Sampler = FGLTFJsonAnimationSamplerIndex(JsonAnimation.Samplers.Add(JsonSampler));
-			JsonChannel.Target.Path = EGLTFJsonTargetPath::Rotation;
-			JsonChannel.Target.Node = NodeIndex;
-			JsonAnimation.Channels.Add(JsonChannel);
-		}
-
-		if (KeyScales.Num() > 0)
-		{
-			TArray<FGLTFJsonVector3> Scales;
-			Scales.AddUninitialized(KeyScales.Num());
-
-			for (int32 Key = 0; Key < KeyScales.Num(); ++Key)
-			{
-				const FVector KeyScale = KeyTransforms[Key].GetScale3D();
-				Scales[Key] = FGLTFConverterUtility::ConvertScale(KeyScale);
-			}
-
-			JsonInputAccessor.Count = Scales.Num();
-			JsonInputAccessor.Max[0] = Timestamps[Scales.Num() - 1];
-
-			FGLTFJsonAccessor JsonOutputAccessor;
-			JsonOutputAccessor.BufferView = Builder.AddBufferView(Scales);
-			JsonOutputAccessor.ComponentType = EGLTFJsonComponentType::F32;
-			JsonOutputAccessor.Count = Scales.Num();
-			JsonOutputAccessor.Type = EGLTFJsonAccessorType::Vec3;
-
-			FGLTFJsonAnimationSampler JsonSampler;
-			JsonSampler.Input = Builder.AddAccessor(JsonInputAccessor);
-			JsonSampler.Output = Builder.AddAccessor(JsonOutputAccessor);
-			JsonSampler.Interpolation = Interpolation;
-
-			FGLTFJsonAnimationChannel JsonChannel;
-			JsonChannel.Sampler = FGLTFJsonAnimationSamplerIndex(JsonAnimation.Samplers.Add(JsonSampler));
-			JsonChannel.Target.Path = EGLTFJsonTargetPath::Scale;
-			JsonChannel.Target.Node = NodeIndex;
-			JsonAnimation.Channels.Add(JsonChannel);
-		}
-	}
-
-	return Builder.AddAnimation(JsonAnimation);
+	const FGLTFJsonAnimationIndex AnimationIndex = Builder.AddAnimation();
+	Builder.SetupTask<FGLTFAnimationTask>(Builder, RootNode, SkeletalMesh, AnimSequence, AnimationIndex);
+	return AnimationIndex;
 }
 
 FGLTFJsonAnimationIndex FGLTFAnimationDataConverter::Convert(FGLTFJsonNodeIndex RootNode, const USkeletalMeshComponent* SkeletalMeshComponent)
@@ -207,9 +46,8 @@ FGLTFJsonAnimationIndex FGLTFAnimationDataConverter::Convert(FGLTFJsonNodeIndex
 	if (AnimationIndex != INDEX_NONE)
 	{
 		FGLTFJsonAnimation& JsonAnimation = Builder.GetAnimation(AnimationIndex);
-		JsonAnimation.Name = FGLTFNameUtility::GetName(SkeletalMeshComponent);
-
 		FGLTFJsonPlayData& JsonPlayData = JsonAnimation.PlayData;
+
 		JsonPlayData.Looping = SkeletalMeshComponent->AnimationData.bSavedLooping;
 		JsonPlayData.Playing = SkeletalMeshComponent->AnimationData.bSavedPlaying;
 		JsonPlayData.PlayRate = SkeletalMeshComponent->AnimationData.SavedPlayRate;
diff --git a/Engine/Plugins/Enterprise/GLTFExporter/MigratedFromGit/Source/GLTFExporter/Private/Tasks/GLTFAnimationTasks.cpp b/Engine/Plugins/Enterprise/GLTFExporter/MigratedFromGit/Source/GLTFExporter/Private/Tasks/GLTFAnimationTasks.cpp
new file mode 100644
--- /dev/null
+++ b/Engine/Plugins/Enterprise/GLTFExporter/MigratedFromGit/Source/GLTFExporter/Private/Tasks/GLTFAnimationTasks.cpp
@@ -0,0 +1,173 @@
+// Copyright Epic Games, Inc. All Rights Reserved.
+
+#include "Tasks/GLTFAnimationTasks.h"
+#include "Builders/GLTFContainerBuilder.h"
+#include "Converters/GLTFConverterUtility.h"
+#include "Converters/GLTFBoneUtility.h"
+
+void FGLTFAnimationTask::Complete()
+{
+	const int32 FrameCount = AnimSequence->GetRawNumberOfFrames();
+	const USkeleton* Skeleton = AnimSequence->GetSkeleton();
+
+	FGLTFJsonAnimation& JsonAnimation = Builder.GetAnimation(AnimationIndex);
+	AnimSequence->GetName(JsonAnimation.Name);
+
+	TArray<float> Timestamps;
+	Timestamps.AddUninitialized(FrameCount);
+
+	for (int32 Frame = 0; Frame < FrameCount; ++Frame)
+	{
+		Timestamps[Frame] = AnimSequence->GetTimeAtFrame(Frame);
+	}
+
+	FGLTFJsonAccessor JsonInputAccessor;
+	JsonInputAccessor.BufferView = Builder.AddBufferView(Timestamps);
+	JsonInputAccessor.ComponentType = EGLTFJsonComponentType::F32;
+	JsonInputAccessor.Type = EGLTFJsonAccessorType::Scalar;
+	JsonInputAccessor.MinMaxLength = 1;
+	JsonInputAccessor.Min[0] = 0;
+
+	FBoneContainer BoneContainer;
+	if (Builder.ExportOptions->bRetargetBoneTransforms)
+	{
+		FGLTFBoneUtility::InitializeToSkeleton(BoneContainer, Skeleton);
+	}
+
+	const EGLTFJsonInterpolation Interpolation = FGLTFConverterUtility::ConvertInterpolation(AnimSequence->Interpolation);
+	const TArray<FName>& TrackNames = AnimSequence->GetAnimationTrackNames();
+	const int32 TrackCount = TrackNames.Num();
+
+	for (int32 TrackIndex = 0; TrackIndex < TrackCount; ++TrackIndex)
+	{
+		const FRawAnimSequenceTrack& Track = AnimSequence->GetRawAnimationTrack(TrackIndex);
+		const TArray<FVector>& KeyPositions = Track.PosKeys;
+		const TArray<FQuat>& KeyRotations = Track.RotKeys;
+		const TArray<FVector>& KeyScales = Track.ScaleKeys;
+
+		const int32 MaxKeys = FMath::Max3(KeyPositions.Num(), KeyRotations.Num(), KeyScales.Num());
+		if (MaxKeys == 0)
+		{
+			continue;
+		}
+
+		TArray<FTransform> KeyTransforms;
+		KeyTransforms.AddUninitialized(MaxKeys);
+
+		for (int32 Key = 0; Key < KeyTransforms.Num(); ++Key)
+		{
+			const FVector& KeyPosition = KeyPositions.IsValidIndex(Key) ? KeyPositions[Key] : FVector::ZeroVector;
+			const FQuat& KeyRotation = KeyRotations.IsValidIndex(Key) ? KeyRotations[Key] : FQuat::Identity;
+			const FVector& KeyScale = KeyScales.IsValidIndex(Key) ? KeyScales[Key] : FVector::OneVector;
+			KeyTransforms[Key] = { KeyRotation, KeyPosition, KeyScale };
+		}
+
+		const int32 SkeletonBoneIndex = AnimSequence->GetSkeletonIndexFromRawDataTrackIndex(TrackIndex);
+		const int32 BoneIndex = const_cast<USkeleton*>(Skeleton)->GetMeshBoneIndexFromSkeletonBoneIndex(SkeletalMesh, SkeletonBoneIndex);
+		const FGLTFJsonNodeIndex NodeIndex = Builder.GetOrAddNode(RootNode, SkeletalMesh, BoneIndex);
+
+		if (Builder.ExportOptions->bRetargetBoneTransforms && Skeleton->GetBoneTranslationRetargetingMode(SkeletonBoneIndex) != EBoneTranslationRetargetingMode::Animation)
+		{
+			for (int32 Key = 0; Key < KeyTransforms.Num(); ++Key)
+			{
+				FGLTFBoneUtility::RetargetTransform(AnimSequence, KeyTransforms[Key], SkeletonBoneIndex, BoneIndex, BoneContainer);
+			}
+		}
+
+		if (KeyPositions.Num() > 0)
+		{
+			TArray<FGLTFJsonVector3> Translations;
+			Translations.AddUninitialized(KeyPositions.Num());
+
+			for (int32 Key = 0; Key < KeyPositions.Num(); ++Key)
+			{
+				const FVector KeyPosition = KeyTransforms[Key].GetTranslation();
+				Translations[Key] = FGLTFConverterUtility::ConvertPosition(KeyPosition, Builder.ExportOptions->ExportScale);
+			}
+
+			JsonInputAccessor.Count = Translations.Num();
+			JsonInputAccessor.Max[0] = Timestamps[Translations.Num() - 1];
+
+			FGLTFJsonAccessor JsonOutputAccessor;
+			JsonOutputAccessor.BufferView = Builder.AddBufferView(Translations);
+			JsonOutputAccessor.ComponentType = EGLTFJsonComponentType::F32;
+			JsonOutputAccessor.Count = Translations.Num();
+			JsonOutputAccessor.Type = EGLTFJsonAccessorType::Vec3;
+
+			FGLTFJsonAnimationSampler JsonSampler;
+			JsonSampler.Input = Builder.AddAccessor(JsonInputAccessor);
+			JsonSampler.Output = Builder.AddAccessor(JsonOutputAccessor);
+			JsonSampler.Interpolation = Interpolation;
+
+			FGLTFJsonAnimationChannel JsonChannel;
+			JsonChannel.Sampler = FGLTFJsonAnimationSamplerIndex(JsonAnimation.Samplers.Add(JsonSampler));
+			JsonChannel.Target.Path = EGLTFJsonTargetPath::Translation;
+			JsonChannel.Target.Node = NodeIndex;
+			JsonAnimation.Channels.Add(JsonChannel);
+		}
+
+		if (KeyRotations.Num() > 0)
+		{
+			TArray<FGLTFJsonQuaternion> Rotations;
+			Rotations.AddUninitialized(KeyRotations.Num());
+
+			for (int32 Key = 0; Key < KeyRotations.Num(); ++Key)
+			{
+				const FQuat KeyRotation = KeyTransforms[Key].GetRotation();
+				Rotations[Key] = FGLTFConverterUtility::ConvertRotation(KeyRotation);
+			}
+
+			JsonInputAccessor.Count = Rotations.Num();
+			JsonInputAccessor.Max[0] = Timestamps[Rotations.Num() - 1];
+
+			FGLTFJsonAccessor JsonOutputAccessor;
+			JsonOutputAccessor.BufferView = Builder.AddBufferView(Rotations);
+			JsonOutputAccessor.ComponentType = EGLTFJsonComponentType::F32;
+			JsonOutputAccessor.Count = Rotations.Num();
+			JsonOutputAccessor.Type = EGLTFJsonAccessorType::Vec4;
+
+			FGLTFJsonAnimationSampler JsonSampler;
+			JsonSampler.Input = Builder.AddAccessor(JsonInputAccessor);;
+			JsonSampler.Output = Builder.AddAccessor(JsonOutputAccessor);
+			JsonSampler.Interpolation = Interpolation;
+
+			FGLTFJsonAnimationChannel JsonChannel;
+			JsonChannel.Sampler = FGLTFJsonAnimationSamplerIndex(JsonAnimation.Samplers.Add(JsonSampler));
+			JsonChannel.Target.Path = EGLTFJsonTargetPath::Rotation;
+			JsonChannel.Target.Node = NodeIndex;
+			JsonAnimation.Channels.Add(JsonChannel);
+		}
+
+		if (KeyScales.Num() > 0)
+		{
+			TArray<FGLTFJsonVector3> Scales;
+			Scales.AddUninitialized(KeyScales.Num());
+
+			for (int32 Key = 0; Key < KeyScales.Num(); ++Key)
+			{
+				const FVector KeyScale = KeyTransforms[Key].GetScale3D();
+				Scales[Key] = FGLTFConverterUtility::ConvertScale(KeyScale);
+			}
+
+			JsonInputAccessor.Count = Scales.Num();
+			JsonInputAccessor.Max[0] = Timestamps[Scales.Num() - 1];
+
+			FGLTFJsonAccessor JsonOutputAccessor;
+			JsonOutputAccessor.BufferView = Builder.AddBufferView(Scales);
+			JsonOutputAccessor.ComponentType = EGLTFJsonComponentType::F32;
+			JsonOutputAccessor.Count = Scales.Num();
+			JsonOutputAccessor.Type = EGLTFJsonAccessorType::Vec3;
+
+			FGLTFJsonAnimationSampler JsonSampler;
+			JsonSampler.Input = Builder.AddAccessor(JsonInputAccessor);
+			JsonSampler.Output = Builder.AddAccessor(JsonOutputAccessor);
+			JsonSampler.Interpolation = Interpolation;
+
+			FGLTFJsonAnimationChannel JsonChannel;
+			JsonChannel.Sampler = FGLTFJsonAnimationSamplerIndex(JsonAnimation.Samplers.Add(JsonSampler));
+			JsonChannel.Target.Path = EGLTFJsonTargetPath::Scale;
+			JsonChannel.Target.Node = NodeIndex;
+			JsonAnimation.Channels.Add(JsonChannel);
+		}
+	}
+}
diff --git a/Engine/Plugins/Enterprise/GLTFExporter/MigratedFromGit/Source/GLTFExporter/Private/Tasks/GLTFAnimationTasks.h b/Engine/Plugins/Enterprise/GLTFExporter/MigratedFromGit/Source/GLTFExporter/Private/Tasks/GLTFAnimationTasks.h
new file mode 100644
--- /dev/null
+++ b/Engine/Plugins/Enterprise/GLTFExporter/MigratedFromGit/Source/GLTFExporter/Private/Tasks/GLTFAnimationTasks.h
@@ -0,0 +1,37 @@
+// Copyright Epic Games, Inc. All Rights Reserved.
+
+#pragma once
+
+#include "Tasks/GLTFTask.h"
+#include "Builders/GLTFConvertBuilder.h"
+#include "Engine.h"
+
+class FGLTFAnimationTask : public FGLTFTask
+{
+public:
+
+	FGLTFAnimationTask(FGLTFConvertBuilder& Builder, FGLTFJsonNodeIndex RootNode, const USkeletalMesh* SkeletalMesh, const UAnimSequence* AnimSequence, FGLTFJsonAnimationIndex AnimationIndex)
+        : FGLTFTask(EGLTFTaskPriority::Animation)
+		, Builder(Builder)
+		, RootNode(RootNode)
+		, SkeletalMesh(SkeletalMesh)
+		, AnimSequence(AnimSequence)
+		, AnimationIndex(AnimationIndex)
+	{
+	}
+
+	virtual FString GetName() override
+	{
+		return AnimSequence->GetName();
+	}
+
+	virtual void Complete() override;
+
+private:
+
+	FGLTFConvertBuilder& Builder;
+	FGLTFJsonNodeIndex RootNode;
+	const USkeletalMesh* SkeletalMesh;
+	const UAnimSequence* AnimSequence;
+	const FGLTFJsonAnimationIndex AnimationIndex;
+};
