[GLTFExporter] refactor: change task messages to be more consistent with engine messages

commit 7fb7dabab58a36294e89166b1622467091d38e51
Date: Sun, 4 Oct 2020 14:58:51 +0200

    refactor: change task messages to be more consistent with engine messages
    (cherry picked from commit e9fc534406ad4e99d23eeb3fa57c7c95c9f223ef​)

#skipci

diff --git a/Engine/Plugins/Enterprise/GLTFExporter/MigratedFromGit/Source/GLTFExporter/Private/Builders/GLTFTaskBuilder.cpp b/Engine/Plugins/Enterprise/GLTFExporter/MigratedFromGit/Source/GLTFExporter/Private/Builders/GLTFTaskBuilder.cpp
--- a/Engine/Plugins/Enterprise/GLTFExporter/MigratedFromGit/Source/GLTFExporter/Private/Builders/GLTFTaskBuilder.cpp
+++ b/Engine/Plugins/Enterprise/GLTFExporter/MigratedFromGit/Source/GLTFExporter/Private/Builders/GLTFTaskBuilder.cpp
@@ -65,11 +65,11 @@ FText FGLTFTaskBuilder::GetPriorityMessageFormat(EGLTFTaskPriority Priority)
 {
 	switch (Priority)
 	{
-		case EGLTFTaskPriority::Actor:     return NSLOCTEXT("GLTFExporter", "ActorTaskMessage", "Actors... {0}");
-		case EGLTFTaskPriority::Mesh:      return NSLOCTEXT("GLTFExporter", "MeshTaskMessage", "Meshes... {0}");
-		case EGLTFTaskPriority::Animation: return NSLOCTEXT("GLTFExporter", "AnimationTaskMessage", "Animations... {0}");
-		case EGLTFTaskPriority::Material:  return NSLOCTEXT("GLTFExporter", "MaterialTaskMessage", "Materials... {0}");
-		case EGLTFTaskPriority::Texture:   return NSLOCTEXT("GLTFExporter", "TextureTaskMessage", "Textures... {0}");
+		case EGLTFTaskPriority::Actor:     return NSLOCTEXT("GLTFExporter", "ActorTaskMessage", "Converting actors: {0}");
+		case EGLTFTaskPriority::Mesh:      return NSLOCTEXT("GLTFExporter", "MeshTaskMessage", "Converting meshes: {0}");
+		case EGLTFTaskPriority::Animation: return NSLOCTEXT("GLTFExporter", "AnimationTaskMessage", "Converting animations: {0}");
+		case EGLTFTaskPriority::Material:  return NSLOCTEXT("GLTFExporter", "MaterialTaskMessage", "Converting materials: {0}");
+		case EGLTFTaskPriority::Texture:   return NSLOCTEXT("GLTFExporter", "TextureTaskMessage", "Converting textures: {0}");
 		default:                           checkNoEntry();
 	}
 
diff --git a/Engine/Plugins/Enterprise/GLTFExporter/MigratedFromGit/Source/GLTFExporter/Private/Tasks/GLTFTask.h b/Engine/Plugins/Enterprise/GLTFExporter/MigratedFromGit/Source/GLTFExporter/Private/Tasks/GLTFTask.h
--- a/Engine/Plugins/Enterprise/GLTFExporter/MigratedFromGit/Source/GLTFExporter/Private/Tasks/GLTFTask.h
+++ b/Engine/Plugins/Enterprise/GLTFExporter/MigratedFromGit/Source/GLTFExporter/Private/Tasks/GLTFTask.h
@@ -6,7 +6,7 @@
 
 enum class EGLTFTaskPriority : uint8
 {
-	Actor,
+	Actor, // TODO: add and implement tasks for actors
     Mesh,
     Animation,
     Material,
