[GLTFExporter] refactor: move GLTFImageBuilder::ImageDirPath to GLTFBuilder::DirPath

commit b5fbd5823f6425c272e459d76e3e03ed5fe34a5d
Date: Sun, 4 Oct 2020 21:13:25 +0200

    refactor: move GLTFImageBuilder::ImageDirPath to GLTFBuilder::DirPath
    (cherry picked from commit 4cfa616fbaebd3567a43dff55387c0ac9e6126c1​)

#skipci

diff --git a/Engine/Plugins/Enterprise/GLTFExporter/MigratedFromGit/Source/GLTFExporter/Private/Builders/GLTFBuilder.cpp b/Engine/Plugins/Enterprise/GLTFExporter/MigratedFromGit/Source/GLTFExporter/Private/Builders/GLTFBuilder.cpp
--- a/Engine/Plugins/Enterprise/GLTFExporter/MigratedFromGit/Source/GLTFExporter/Private/Builders/GLTFBuilder.cpp
+++ b/Engine/Plugins/Enterprise/GLTFExporter/MigratedFromGit/Source/GLTFExporter/Private/Builders/GLTFBuilder.cpp
@@ -6,6 +6,7 @@
 FGLTFBuilder::FGLTFBuilder(const FString& FilePath, const UGLTFExportOptions* ExportOptions)
 	: bIsGlbFile(FGLTFBuilderUtility::IsGlbFile(FilePath))
 	, FilePath(FilePath)
+	, DirPath(FPaths::GetPath(FilePath))
 	, ExportOptions(ExportOptions)
 {
 }
diff --git a/Engine/Plugins/Enterprise/GLTFExporter/MigratedFromGit/Source/GLTFExporter/Private/Builders/GLTFBuilder.h b/Engine/Plugins/Enterprise/GLTFExporter/MigratedFromGit/Source/GLTFExporter/Private/Builders/GLTFBuilder.h
--- a/Engine/Plugins/Enterprise/GLTFExporter/MigratedFromGit/Source/GLTFExporter/Private/Builders/GLTFBuilder.h
+++ b/Engine/Plugins/Enterprise/GLTFExporter/MigratedFromGit/Source/GLTFExporter/Private/Builders/GLTFBuilder.h
@@ -16,5 +16,7 @@ public:
 
 	const FString FilePath;
 
+	const FString DirPath;
+
 	const UGLTFExportOptions* const ExportOptions;
 };
diff --git a/Engine/Plugins/Enterprise/GLTFExporter/MigratedFromGit/Source/GLTFExporter/Private/Builders/GLTFImageBuilder.cpp b/Engine/Plugins/Enterprise/GLTFExporter/MigratedFromGit/Source/GLTFExporter/Private/Builders/GLTFImageBuilder.cpp
--- a/Engine/Plugins/Enterprise/GLTFExporter/MigratedFromGit/Source/GLTFExporter/Private/Builders/GLTFImageBuilder.cpp
+++ b/Engine/Plugins/Enterprise/GLTFExporter/MigratedFromGit/Source/GLTFExporter/Private/Builders/GLTFImageBuilder.cpp
@@ -8,7 +8,6 @@
 
 FGLTFImageBuilder::FGLTFImageBuilder(const FString& FilePath, const UGLTFExportOptions* ExportOptions)
     : FGLTFBufferBuilder(FilePath, ExportOptions)
-	, ImageDirPath(FPaths::GetPath(FilePath))
 {
 }
 
@@ -75,7 +74,7 @@ FString FGLTFImageBuilder::SaveImageToFile(const void* CompressedData, int64 Com
 	const FString ImageUri = FGLTFBuilderUtility::GetUniqueFilename(Name, Extension, UniqueImageUris);
 
 	const TArrayView<const uint8> ImageData(static_cast<const uint8*>(CompressedData), CompressedByteLength);
-	const FString ImagePath = FPaths::Combine(ImageDirPath, ImageUri);
+	const FString ImagePath = FPaths::Combine(DirPath, ImageUri);
 
 	if (!FFileHelper::SaveArrayToFile(ImageData, *ImagePath))
 	{
diff --git a/Engine/Plugins/Enterprise/GLTFExporter/MigratedFromGit/Source/GLTFExporter/Private/Builders/GLTFImageBuilder.h b/Engine/Plugins/Enterprise/GLTFExporter/MigratedFromGit/Source/GLTFExporter/Private/Builders/GLTFImageBuilder.h
--- a/Engine/Plugins/Enterprise/GLTFExporter/MigratedFromGit/Source/GLTFExporter/Private/Builders/GLTFImageBuilder.h
+++ b/Engine/Plugins/Enterprise/GLTFExporter/MigratedFromGit/Source/GLTFExporter/Private/Builders/GLTFImageBuilder.h
@@ -24,8 +24,6 @@ private:
 
 	FString SaveImageToFile(const void* CompressedData, int64 CompressedByteLength, EGLTFJsonMimeType MimeType, const FString& Name);
 
-	const FString ImageDirPath;
-
 	TSet<FString> UniqueImageUris;
 	TMap<FGLTFBinaryHashKey, FGLTFJsonImageIndex> UniqueImageIndices;
 };
