[GLTFExporter] feat: add support for bundling web viewer files

commit a1bb1b2803db9276971433b12c884a771564f4b2
Date: Sun, 4 Oct 2020 21:44:22 +0200

    feat: add support for bundling web viewer files
    (cherry picked from commit c65655b6a26ecd133e49b2cb77cf7f8749fc071a​)

#skipci

diff --git a/Engine/Plugins/Enterprise/GLTFExporter/MigratedFromGit/Source/GLTFExporter/Private/Builders/GLTFContainerBuilder.cpp b/Engine/Plugins/Enterprise/GLTFExporter/MigratedFromGit/Source/GLTFExporter/Private/Builders/GLTFContainerBuilder.cpp
--- a/Engine/Plugins/Enterprise/GLTFExporter/MigratedFromGit/Source/GLTFExporter/Private/Builders/GLTFContainerBuilder.cpp
+++ b/Engine/Plugins/Enterprise/GLTFExporter/MigratedFromGit/Source/GLTFExporter/Private/Builders/GLTFContainerBuilder.cpp
@@ -2,6 +2,8 @@
 
 #include "Builders/GLTFContainerBuilder.h"
 #include "Builders/GLTFContainerUtility.h"
+#include "GLTFExporterModule.h"
+#include "Interfaces/IPluginManager.h"
 
 FGLTFContainerBuilder::FGLTFContainerBuilder(const FString& FilePath, const UGLTFExportOptions* ExportOptions, bool bSelectedActorsOnly)
 	: FGLTFConvertBuilder(FilePath, ExportOptions, bSelectedActorsOnly)
@@ -18,6 +20,32 @@ void FGLTFContainerBuilder::WriteGlb(FArchive& Archive) const
 	FGLTFContainerUtility::WriteGlb(Archive, JsonData, BufferData);
 }
 
+void FGLTFContainerBuilder::BundleWebViewer()
+{
+	// TODO: instead of raw files, store web viewer as a zip and unpack it when needed
+
+	const FString WebViewerPath = IPluginManager::Get().FindPlugin(GLTFEXPORTER_MODULE_NAME)->GetBaseDir() / TEXT("Resources") / TEXT("WebViewer");
+	TArray<FString> SourceFiles;
+
+	IFileManager& FileManager = IFileManager::Get();
+	FileManager.FindFilesRecursive(SourceFiles, *WebViewerPath, TEXT("*"), true, false);
+
+	if (SourceFiles.Num() == 0)
+	{
+		AddErrorMessage(FString::Printf(TEXT("No source files for web viewer found at %s"), *WebViewerPath));
+		return;
+	}
+
+	for (const FString& SourceFile : SourceFiles)
+	{
+		FString DestinationFile = DirPath / SourceFile.RightChop(WebViewerPath.Len());
+		if (FileManager.Copy(*DestinationFile, *SourceFile) != COPY_OK)
+		{
+			AddErrorMessage(FString::Printf(TEXT("Failed to copy web viewer file from %s to %s"), *SourceFile, *DestinationFile));
+		}
+	}
+}
+
 void FGLTFContainerBuilder::Write(FArchive& Archive, FFeedbackContext* Context)
 {
 	CompleteAllTasks(Context);
@@ -30,4 +58,9 @@ void FGLTFContainerBuilder::Write(FArchive& Archive, FFeedbackContext* Context)
 	{
 		WriteJson(Archive);
 	}
+
+	if (ExportOptions->bBundleWebViewer)
+	{
+		BundleWebViewer();
+	}
 }
diff --git a/Engine/Plugins/Enterprise/GLTFExporter/MigratedFromGit/Source/GLTFExporter/Private/Builders/GLTFContainerBuilder.h b/Engine/Plugins/Enterprise/GLTFExporter/MigratedFromGit/Source/GLTFExporter/Private/Builders/GLTFContainerBuilder.h
--- a/Engine/Plugins/Enterprise/GLTFExporter/MigratedFromGit/Source/GLTFExporter/Private/Builders/GLTFContainerBuilder.h
+++ b/Engine/Plugins/Enterprise/GLTFExporter/MigratedFromGit/Source/GLTFExporter/Private/Builders/GLTFContainerBuilder.h
@@ -15,4 +15,6 @@ public:
 private:
 
 	void WriteGlb(FArchive& Archive) const;
+
+	void BundleWebViewer();
 };
diff --git a/Engine/Plugins/Enterprise/GLTFExporter/MigratedFromGit/Source/GLTFExporter/Private/GLTFExportOptions.cpp b/Engine/Plugins/Enterprise/GLTFExporter/MigratedFromGit/Source/GLTFExporter/Private/GLTFExportOptions.cpp
--- a/Engine/Plugins/Enterprise/GLTFExporter/MigratedFromGit/Source/GLTFExporter/Private/GLTFExportOptions.cpp
+++ b/Engine/Plugins/Enterprise/GLTFExporter/MigratedFromGit/Source/GLTFExporter/Private/GLTFExportOptions.cpp
@@ -39,7 +39,7 @@ UGLTFExportOptions::UGLTFExportOptions(const FObjectInitializer& ObjectInitializ
 	bExportHDRIBackdrops = true;
 	bExportVariantSets = true;
 	bExportInteractionHotspots = true;
-	bBundleWebViewer = false;
+	bBundleWebViewer = true;
 	bExportPreviewMesh = true;
 	bAllExtensionsRequired = false;
 }
@@ -203,7 +203,6 @@ bool UGLTFExportOptions::CanEditChange(const FProperty* InProperty) const
 		PropertyFName == GET_MEMBER_NAME_CHECKED(ThisClass, TextureHDREncoding) ||
 		PropertyFName == GET_MEMBER_NAME_CHECKED(ThisClass, bExportHiddenInGame) ||
 		PropertyFName == GET_MEMBER_NAME_CHECKED(ThisClass, bExportReflectionCaptures) ||
-		PropertyFName == GET_MEMBER_NAME_CHECKED(ThisClass, bBundleWebViewer) ||
 		PropertyFName == GET_MEMBER_NAME_CHECKED(ThisClass, bAllExtensionsRequired))
 	{
 		return false;
