[GLTFExporter] chore: add initial web viewer build

commit 8ba8ff2e6da298e3473257814aa27a3c039ae868

Date: Sun, 4 Oct 2020 21:53:02 +0200

    chore: add initial web viewer build
    (cherry picked from commit 579de970d04ed0ebeb041c213de4ce0d6d759d05​)

#skipci

diff --git a/Engine/Plugins/Enterprise/GLTFExporter/MigratedFromGit/Resources/WebViewer/index.html b/Engine/Plugins/Enterprise/GLTFExporter/MigratedFromGit/Resources/WebViewer/index.html
new file mode 100644
--- /dev/null
+++ b/Engine/Plugins/Enterprise/GLTFExporter/MigratedFromGit/Resources/WebViewer/index.html
@@ -0,0 +1,37 @@
+<!doctype html>
+<html lang="en">
+
+<head>
+  <meta charset="utf-8" />
+  <meta name="viewport"
+    content="width=device-width,initial-scale=1,maximum-scale=1,minimum-scale=1,shrink-to-fit=no,user-scalable=no,minimal-ui" />
+  <meta name="theme-color" content="##ffffff" />
+  <meta name="description" content="Epic glTF Viewer" />
+  <link rel="icon" media="(prefers-color-scheme:dark)" type="image/png" sizes="32x32"
+    href="./viewer/favicon-dark-32x32.png" />
+  <link rel="icon" media="(prefers-color-scheme:dark)" type="image/png" sizes="16x16"
+    href="./viewer/favicon-dark-16x16.png" />
+  <link rel="icon" media="(prefers-color-scheme:light)" type="image/png" sizes="32x32"
+    href="./viewer/favicon-32x32.png" />
+  <link rel="icon" media="(prefers-color-scheme:light)" type="image/png" sizes="16x16"
+    href="./viewer/favicon-16x16.png" />
+  <link rel="apple-touch-icon" sizes="180x180" href="./viewer/apple-touch-icon.png" />
+  <link rel="mask-icon" href="./viewer/safari-pinned-tab.svg" color="#0aaff1" />
+  <link rel="manifest" href="./viewer/manifest.json" />
+  <link rel="stylesheet" href="https://fonts.googleapis.com/css2?family=Open+Sans:wght@300;400;600;700&display=swap" />
+  <link rel="stylesheet" href="https://fonts.googleapis.com/icon?family=Material+Icons" />
+  <title>Epic glTF Viewer</title>
+</head>
+
+<body><noscript>You need to enable JavaScript to run this app.</noscript>
+  <div id="root"></div>
+  <script src="https:[REDACTED_EMAIL]/umd/react.production.min.js"></script>
+  <script src="https:[REDACTED_EMAIL]/umd/react-dom.production.min.js"></script>
+  <script src="https:[REDACTED_EMAIL]/lib/mobx.umd.min.js"></script>
+  <script src="https:[REDACTED_EMAIL]/prop-types.min.js"></script>
+  <script src="./viewer/static/js/runtime-main.c850d8be.js"></script>
+  <script src="./viewer/static/js/2.fd937dd3.chunk.js"></script>
+  <script src="./viewer/static/js/main.76e9e68f.chunk.js"></script>
+</body>
+
+</html>
\ No newline at end of file
diff --git a/Engine/Plugins/Enterprise/GLTFExporter/MigratedFromGit/Resources/WebViewer/viewer/manifest.json b/Engine/Plugins/Enterprise/GLTFExporter/MigratedFromGit/Resources/WebViewer/viewer/manifest.json
new file mode 100644
--- /dev/null
+++ b/Engine/Plugins/Enterprise/GLTFExporter/MigratedFromGit/Resources/WebViewer/viewer/manifest.json
@@ -0,0 +1,25 @@
+{
+  "short_name": "Epic glTF Viewer",
+  "name": "Epic glTF Viewer",
+  "icons": [
+    {
+      "src": "android-chrome-192x192.png",
+      "sizes": "192x192",
+      "type": "image/png"
+    },
+    {
+      "src": "android-chrome-256x256.png",
+      "sizes": "256x256",
+      "type": "image/png"
+    },
+    {
+      "src": "android-chrome-512x512.png",
+      "sizes": "512x512",
+      "type": "image/png"
+    }
+  ],
+  "start_url": ".",
+  "display": "standalone",
+  "theme_color": "#ffffff",
+  "background_color": "#ffffff"
+}
\ No newline at end of file
diff --git a/Engine/Plugins/Enterprise/GLTFExporter/MigratedFromGit/Resources/WebViewer/viewer/playcanvas/1002267.json b/Engine/Plugins/Enterprise/GLTFExporter/MigratedFromGit/Resources/WebViewer/viewer/playcanvas/1002267.json
new file mode 100644
--- /dev/null
+++ b/Engine/Plugins/Enterprise/GLTFExporter/MigratedFromGit/Resources/WebViewer/viewer/playcanvas/1002267.json
@@ -0,0 +1,65 @@
+{
+  "name": "ClearSky",
+  "settings": {
+    "physics": {
+      "gravity": [
+        0,
+        -9.8,
+        0
+      ]
+    },
+    "render": {
+      "fog_end": 1000,
+      "tonemapping": 3,
+      "skybox": 36180622,
+      "fog_density": 0.01,
+      "gamma_correction": 1,
+      "exposure": 1,
+      "fog_start": 1,
+      "global_ambient": [
+        0.2,
+        0.2,
+        0.2
+      ],
+      "skyboxIntensity": 1,
+      "fog_color": [
+        0,
+        0,
+        0
+      ],
+      "lightmapMode": 1,
+      "fog": "none",
+      "lightmapMaxResolution": 2048,
+      "skyboxMip": 2,
+      "lightmapSizeMultiplier": 16
+    }
+  },
+  "entities": {
+    "aed82cab-703d-4336-8ce9-dfc2d23a3b99": {
+      "position": [
+        0,
+        0,
+        0
+      ],
+      "scale": [
+        1,
+        1,
+        1
+      ],
+      "name": "Root",
+      "parent": null,
+      "resource_id": "aed82cab-703d-4336-8ce9-dfc2d23a3b99",
+      "components": {},
+      "rotation": [
+        0,
+        0,
+        0
+      ],
+      "tags": [],
+      "enabled": true,
+      "children": []
+    }
+  },
+  "created": "2020-10-04T13:56:13.260Z",
+  "id": 1002267
+}
\ No newline at end of file
diff --git a/Engine/Plugins/Enterprise/GLTFExporter/MigratedFromGit/Resources/WebViewer/viewer/playcanvas/config.json b/Engine/Plugins/Enterprise/GLTFExporter/MigratedFromGit/Resources/WebViewer/viewer/playcanvas/config.json
new file mode 100644
--- /dev/null
+++ b/Engine/Plugins/Enterprise/GLTFExporter/MigratedFromGit/Resources/WebViewer/viewer/playcanvas/config.json
@@ -0,0 +1,123 @@
+{
+  "application_properties": {
+    "i18nAssets": [],
+    "useTouch": true,
+    "layerOrder": [
+      {
+        "layer": 0,
+        "enabled": true,
+        "transparent": false
+      },
+      {
+        "layer": 1,
+        "enabled": true,
+        "transparent": false
+      },
+      {
+        "layer": 2,
+        "enabled": true,
+        "transparent": false
+      },
+      {
+        "layer": 0,
+        "enabled": true,
+        "transparent": true
+      },
+      {
+        "layer": 3,
+        "enabled": true,
+        "transparent": false
+      },
+      {
+        "layer": 3,
+        "enabled": true,
+        "transparent": true
+      },
+      {
+        "layer": 4,
+        "enabled": true,
+        "transparent": true
+      }
+    ],
+    "externalScripts": [],
+    "height": 720,
+    "vr": false,
+    "useModelV2": false,
+    "antiAlias": true,
+    "layers": {
+      "0": {
+        "transparentSortMode": 3,
+        "opaqueSortMode": 2,
+        "name": "World"
+      },
+      "1": {
+        "transparentSortMode": 3,
+        "opaqueSortMode": 2,
+        "name": "Depth"
+      },
+      "2": {
+        "transparentSortMode": 3,
+        "opaqueSortMode": 0,
+        "name": "Skybox"
+      },
+      "3": {
+        "transparentSortMode": 3,
+        "opaqueSortMode": 0,
+        "name": "Immediate"
+      },
+      "4": {
+        "transparentSortMode": 1,
+        "opaqueSortMode": 1,
+        "name": "UI"
+      }
+    },
+    "width": 1280,
+    "useDevicePixelRatio": false,
+    "useKeyboard": true,
+    "batchGroups": [],
+    "preserveDrawingBuffer": false,
+    "useLegacyScripts": false,
+    "fillMode": "FILL_WINDOW",
+    "scripts": [],
+    "useMouse": true,
+    "use3dPhysics": false,
+    "transparentCanvas": false,
+    "resolutionMode": "AUTO",
+    "loadingScreenScript": null,
+    "preferWebGl2": true,
+    "useGamepads": false,
+    "libraries": []
+  },
+  "scenes": [
+    {
+      "name": "ClearSky",
+      "url": "1002267.json"
+    }
+  ],
+  "assets": {
+    "36180622": {
+      "tags": [],
+      "name": "clearsky",
+      "revision": 1,
+      "preload": true,
+      "meta": null,
+      "data": {
+        "name": "clearsky",
+        "minFilter": 5,
+        "magFilter": 1,
+        "anisotropy": 1,
+        "rgbm": true
+      },
+      "type": "cubemap",
+      "file": {
+        "filename": "clearsky.dds",
+        "size": 524456,
+        "hash": "562ce97935ad1f80ef9afd8363a85782",
+        "url": "files/assets/36180622/1/clearsky.dds"
+      },
+      "region": "eu-west-1",
+      "i18n": {},
+      "id": "36180622"
+    }
+  }
+}
\ No newline at end of file
diff --git a/Engine/Plugins/Enterprise/GLTFExporter/MigratedFromGit/Resources/WebViewer/viewer/static/js/2.fd937dd3.chunk.js.LICENSE.txt b/Engine/Plugins/Enterprise/GLTFExporter/MigratedFromGit/Resources/WebViewer/viewer/static/js/2.fd937dd3.chunk.js.LICENSE.txt
new file mode 100644
--- /dev/null
+++ b/Engine/Plugins/Enterprise/GLTFExporter/MigratedFromGit/Resources/WebViewer/viewer/static/js/2.fd937dd3.chunk.js.LICENSE.txt
@@ -0,0 +1,50 @@
+/*!
+ * The buffer module from node.js, for the browser.
+ *
+ * @author   Feross Aboukhadijeh <http://feross.org>
+ * @license  MIT
+ */
+
+/*! *****************************************************************************
+Copyright (c) Microsoft Corporation. All rights reserved.
+Licensed under the Apache License, Version 2.0 (the "License"); you may not use
+this file except in compliance with the License. You may obtain a copy of the
+License at http://www.apache.org/licenses/LICENSE-2.0
+
+THIS CODE IS PROVIDED ON AN *AS IS* BASIS, WITHOUT WARRANTIES OR CONDITIONS OF ANY
+KIND, EITHER EXPRESS OR IMPLIED, INCLUDING WITHOUT LIMITATION ANY IMPLIED
+WARRANTIES OR CONDITIONS OF TITLE, FITNESS FOR A PARTICULAR PURPOSE,
+MERCHANTABLITY OR NON-INFRINGEMENT.
+
+See the Apache Version 2.0 License for specific language governing permissions
+and limitations under the License.
+***************************************************************************** */
+
+/*! Hammer.JS - v2.0.7 - 2016-04-22
+ * http://hammerjs.github.io/
+ *
+ * Copyright (c) 2016 Jorik Tangelder;
+ * Licensed under the MIT license */
+
+/**
+ * @license
+ * PlayCanvas Engine v1.36.0-rc.2 revision 798f16872
+ * Copyright 2011-2020 PlayCanvas Ltd. All rights reserved.
+ */
+
+/**
+ * A better abstraction over CSS.
+ *
+ * @copyright Oleg Isonen (Slobodskoi) / Isonen 2014-present
+ * @website [REDACTED_PROFILE_URL]/jss
+ * @license MIT
+ */
+
+/** @license React v16.13.0
+ * react-is.production.min.js
+ *
+ * Copyright (c) Facebook, Inc. and its affiliates.
+ *
+ * This source code is licensed under the MIT license found in the
+ * LICENSE file in the root directory of this source tree.
+ */
diff --git a/Engine/Plugins/Enterprise/GLTFExporter/MigratedFromGit/Resources/WebViewer/viewer/static/js/runtime-main.c850d8be.js b/Engine/Plugins/Enterprise/GLTFExporter/MigratedFromGit/Resources/WebViewer/viewer/static/js/runtime-main.c850d8be.js
new file mode 100644
--- /dev/null
+++ b/Engine/Plugins/Enterprise/GLTFExporter/MigratedFromGit/Resources/WebViewer/viewer/static/js/runtime-main.c850d8be.js
@@ -0,0 +1,2 @@
+!function(e){function r(r){for(var n,i,l=r[0],f=r[1],p=r[2],c=0,s=[];c<l.length;c++)i=l[c],Object.prototype.hasOwnProperty.call(o,i)&&o[i]&&s.push(o[i][0]),o[i]=0;for(n in f)Object.prototype.hasOwnProperty.call(f,n)&&(e[n]=f[n]);for(a&&a(r);s.length;)s.shift()();return u.push.apply(u,p||[]),t()}function t(){for(var e,r=0;r<u.length;r++){for(var t=u[r],n=!0,l=1;l<t.length;l++){var f=t[l];0!==o[f]&&(n=!1)}n&&(u.splice(r--,1),e=i(i.s=t[0]))}return e}var n={},o={1:0},u=[];function i(r){if(n[r])return n[r].exports;var t=n[r]={i:r,l:!1,exports:{}};return e[r].call(t.exports,t,t.exports,i),t.l=!0,t.exports}i.m=e,i.c=n,i.d=function(e,r,t){i.o(e,r)||Object.defineProperty(e,r,{enumerable:!0,get:t})},i.r=function(e){"undefined"!==typeof Symbol&&Symbol.toStringTag&&Object.defineProperty(e,Symbol.toStringTag,{value:"Module"}),Object.defineProperty(e,"__esModule",{value:!0})},i.t=function(e,r){if(1&r&&(e=i(e)),8&r)return e;if(4&r&&"object"===typeof e&&e&&e.__esModule)return e;var t=Object.create(null);if(i.r(t),Object.defineProperty(t,"default",{enumerable:!0,value:e}),2&r&&"string"!=typeof e)for(var n in e)i.d(t,n,function(r){return e[r]}.bind(null,n));return t},i.n=function(e){var r=e&&e.__esModule?function(){return e.default}:function(){return e};return i.d(r,"a",r),r},i.o=function(e,r){return Object.prototype.hasOwnProperty.call(e,r)},i.p="./";var l=this["webpackJsonp@epic/gltf-web-viewer"]=this["webpackJsonp@epic/gltf-web-viewer"]||[],f=l.push.bind(l);l.push=r,l=l.slice();for(var p=0;p<l.length;p++)r(l[p]);var a=f;t()}([]);
+//# sourceMappingURL=runtime-main.c850d8be.js.map
\ No newline at end of file
