[GLTFExporter] refactor: remove export options TextureFormat NormalExportMethod NormalizeUVCoordinates bExportReflectionCaptures

commit 51c978dcd614a7ba8b2a84a7703a987def588304
Date: Sun, 4 Oct 2020 22:29:05 +0200

    refactor: remove export options TextureFormat NormalExportMethod NormalizeUVCoordinates bExportReflectionCaptures
    (cherry picked from commit 56e600a4af3dd89a8ff18d8c5ea84673d1c9196c​)

#skipci

diff --git a/Engine/Plugins/Enterprise/GLTFExporter/MigratedFromGit/Source/GLTFExporter/Private/GLTFExportOptions.cpp b/Engine/Plugins/Enterprise/GLTFExporter/MigratedFromGit/Source/GLTFExporter/Private/GLTFExportOptions.cpp
--- a/Engine/Plugins/Enterprise/GLTFExporter/MigratedFromGit/Source/GLTFExporter/Private/GLTFExportOptions.cpp
+++ b/Engine/Plugins/Enterprise/GLTFExporter/MigratedFromGit/Source/GLTFExporter/Private/GLTFExportOptions.cpp
@@ -18,15 +18,12 @@ UGLTFExportOptions::UGLTFExportOptions(const FObjectInitializer& ObjectInitializ
 	bBakeMaterialInputs = true;
 	DefaultMaterialBakeSize = EGLTFExporterMaterialBakeSize::POT_1024;
 	bExportVertexColors = true;
-	NormalExportMethod = EGLTFExporterNormalExportMethod::NormalsAndTangents;
 	bQuantizeVertexNormals = true;
 	bQuantizeVertexTangents = true;
-	NormalizeUVCoordinates = EGLTFExporterNormalizeUVCoordinates::Never;
 	DefaultLevelOfDetail = 0;
 	bExportVertexSkinWeights = true;
 	bExportAnimationSequences = true;
 	bRetargetBoneTransforms = true;
-	TextureFormat = EGLTFExporterTextureFormat::PNG;
 	bExportSourceTextures = true;
 	bExportLightmaps = true;
 	TextureHDREncoding = EGLTFExporterTextureHDREncoding::RGBM;
@@ -35,7 +32,6 @@ UGLTFExportOptions::UGLTFExportOptions(const FObjectInitializer& ObjectInitializ
 	ExportLights = EGLTFExporterLightMobility::MovableAndStationary;
 	bExportCameras = true;
 	bExportOrbitalCameras = true;
-	bExportReflectionCaptures = false;
 	bExportHDRIBackdrops = true;
 	bExportVariantSets = true;
 	bExportInteractionHotspots = true;
@@ -197,12 +193,8 @@ bool UGLTFExportOptions::CanEditChange(const FProperty* InProperty) const
 	const FName PropertyFName = InProperty->GetFName();
 
 	// TODO: remove options that have been implemented in exporters/converters.
-	if (PropertyFName == GET_MEMBER_NAME_CHECKED(ThisClass, NormalExportMethod) ||
-		PropertyFName == GET_MEMBER_NAME_CHECKED(ThisClass, NormalizeUVCoordinates) ||
-		PropertyFName == GET_MEMBER_NAME_CHECKED(ThisClass, TextureFormat) ||
-		PropertyFName == GET_MEMBER_NAME_CHECKED(ThisClass, TextureHDREncoding) ||
+	if (PropertyFName == GET_MEMBER_NAME_CHECKED(ThisClass, TextureHDREncoding) ||
 		PropertyFName == GET_MEMBER_NAME_CHECKED(ThisClass, bExportHiddenInGame) ||
-		PropertyFName == GET_MEMBER_NAME_CHECKED(ThisClass, bExportReflectionCaptures) ||
 		PropertyFName == GET_MEMBER_NAME_CHECKED(ThisClass, bAllExtensionsRequired))
 	{
 		return false;
diff --git a/Engine/Plugins/Enterprise/GLTFExporter/MigratedFromGit/Source/GLTFExporter/Public/GLTFExportOptions.h b/Engine/Plugins/Enterprise/GLTFExporter/MigratedFromGit/Source/GLTFExporter/Public/GLTFExportOptions.h
--- a/Engine/Plugins/Enterprise/GLTFExporter/MigratedFromGit/Source/GLTFExporter/Public/GLTFExportOptions.h
+++ b/Engine/Plugins/Enterprise/GLTFExporter/MigratedFromGit/Source/GLTFExporter/Public/GLTFExportOptions.h
@@ -5,22 +5,6 @@
 #include "CoreMinimal.h"
 #include "GLTFExportOptions.generated.h"
 
-UENUM(BlueprintType)
-enum class EGLTFExporterNormalizeUVCoordinates : uint8
-{
-	Always,
-	Auto,
-	Never
-};
-
-UENUM(BlueprintType)
-enum class EGLTFExporterNormalExportMethod : uint8
-{
-	None UMETA(DisplayName = "None"),
-	Normals UMETA(DisplayName = "Normals"),
-	NormalsAndTangents UMETA(DisplayName = "Normals and Tangents")
-};
-
 UENUM(BlueprintType)
 enum class EGLTFExporterMaterialBakeSize : uint8
 {
@@ -40,14 +24,6 @@ enum class EGLTFExporterMaterialBakeSize : uint8
 	POT_8192 UMETA(DisplayName = "8192 x 8192")
 };
 
-UENUM(BlueprintType)
-enum class EGLTFExporterTextureFormat : uint8
-{
-	None,
-	PNG
-	// TODO: add more formats
-};
-
 UENUM(BlueprintType)
 enum class EGLTFExporterTextureHDREncoding : uint8
 {
@@ -88,18 +64,12 @@ public:
 	UPROPERTY(EditAnywhere, BlueprintReadWrite, config, Category = Mesh)
 	uint32 bExportVertexColors : 1;
 
-	UPROPERTY(EditAnywhere, BlueprintReadWrite, config, Category = Mesh)
-	EGLTFExporterNormalExportMethod NormalExportMethod;
-
 	UPROPERTY(EditAnywhere, BlueprintReadWrite, config, Category = Mesh)
 	uint32 bQuantizeVertexNormals : 1;
 
 	UPROPERTY(EditAnywhere, BlueprintReadWrite, config, Category = Mesh)
 	uint32 bQuantizeVertexTangents : 1;
 
-	UPROPERTY(EditAnywhere, BlueprintReadWrite, config, Category = Mesh)
-	EGLTFExporterNormalizeUVCoordinates NormalizeUVCoordinates;
-
 	UPROPERTY(EditAnywhere, BlueprintReadWrite, config, Category = Mesh)
 	int32 DefaultLevelOfDetail;
 
@@ -112,9 +82,6 @@ public:
 	UPROPERTY(EditAnywhere, BlueprintReadWrite, config, Category = Animation)
 	uint32 bRetargetBoneTransforms : 1;
 
-	UPROPERTY(EditAnywhere, BlueprintReadWrite, config, Category = Texture)
-	EGLTFExporterTextureFormat TextureFormat;
-
 	UPROPERTY(EditAnywhere, BlueprintReadWrite, config, Category = Texture)
 	uint32 bExportSourceTextures : 1;
 
@@ -141,9 +108,6 @@ public:
 	UPROPERTY(EditAnywhere, BlueprintReadWrite, config, Category = Scene)
 	uint32 bExportOrbitalCameras : 1;
 
-	UPROPERTY(EditAnywhere, BlueprintReadWrite, config, Category = Scene)
-	uint32 bExportReflectionCaptures : 1;
-
 	UPROPERTY(EditAnywhere, BlueprintReadWrite, config, Category = Scene)
 	uint32 bExportHDRIBackdrops : 1;
 
