[GLTFExporter] feat: implement export option bAllExtensionsRequired

commit e95c0f79b95a1b0509ac6c9d5a47bcc306ca8816
Date: Sun, 4 Oct 2020 22:34:07 +0200

    feat: implement export option bAllExtensionsRequired
    (cherry picked from commit 2bad3f18fe42787fcb0ad3d4a69e666af5ac76e3​)

#skipci

diff --git a/Engine/Plugins/Enterprise/GLTFExporter/MigratedFromGit/Source/GLTFExporter/Private/Builders/GLTFJsonBuilder.cpp b/Engine/Plugins/Enterprise/GLTFExporter/MigratedFromGit/Source/GLTFExporter/Private/Builders/GLTFJsonBuilder.cpp
--- a/Engine/Plugins/Enterprise/GLTFExporter/MigratedFromGit/Source/GLTFExporter/Private/Builders/GLTFJsonBuilder.cpp
+++ b/Engine/Plugins/Enterprise/GLTFExporter/MigratedFromGit/Source/GLTFExporter/Private/Builders/GLTFJsonBuilder.cpp
@@ -10,7 +10,7 @@ FGLTFJsonBuilder::FGLTFJsonBuilder(const FString& FilePath, const UGLTFExportOpt
 
 void FGLTFJsonBuilder::WriteJson(FArchive& Archive) const
 {
-	JsonRoot.ToJson(&Archive, true);
+	JsonRoot.ToJson(&Archive, true, ExportOptions->bAllExtensionsRequired);
 }
 
 void FGLTFJsonBuilder::AddExtension(EGLTFJsonExtension Extension, bool bIsRequired)
diff --git a/Engine/Plugins/Enterprise/GLTFExporter/MigratedFromGit/Source/GLTFExporter/Private/GLTFExportOptions.cpp b/Engine/Plugins/Enterprise/GLTFExporter/MigratedFromGit/Source/GLTFExporter/Private/GLTFExportOptions.cpp
--- a/Engine/Plugins/Enterprise/GLTFExporter/MigratedFromGit/Source/GLTFExporter/Private/GLTFExportOptions.cpp
+++ b/Engine/Plugins/Enterprise/GLTFExporter/MigratedFromGit/Source/GLTFExporter/Private/GLTFExportOptions.cpp
@@ -194,8 +194,7 @@ bool UGLTFExportOptions::CanEditChange(const FProperty* InProperty) const
 
 	// TODO: remove options that have been implemented in exporters/converters.
 	if (PropertyFName == GET_MEMBER_NAME_CHECKED(ThisClass, TextureHDREncoding) ||
-		PropertyFName == GET_MEMBER_NAME_CHECKED(ThisClass, bExportHiddenInGame) ||
-		PropertyFName == GET_MEMBER_NAME_CHECKED(ThisClass, bAllExtensionsRequired))
+		PropertyFName == GET_MEMBER_NAME_CHECKED(ThisClass, bExportHiddenInGame))
 	{
 		return false;
 	}
diff --git a/Engine/Plugins/Enterprise/GLTFExporter/MigratedFromGit/Source/GLTFExporter/Private/Json/GLTFJsonRoot.h b/Engine/Plugins/Enterprise/GLTFExporter/MigratedFromGit/Source/GLTFExporter/Private/Json/GLTFJsonRoot.h
--- a/Engine/Plugins/Enterprise/GLTFExporter/MigratedFromGit/Source/GLTFExporter/Private/Json/GLTFJsonRoot.h
+++ b/Engine/Plugins/Enterprise/GLTFExporter/MigratedFromGit/Source/GLTFExporter/Private/Json/GLTFJsonRoot.h
@@ -82,7 +82,7 @@ struct FGLTFJsonRoot
 	TArray<TUniquePtr<FGLTFJsonHotspot>>    Hotspots;
 
 	template <class CharType = TCHAR, class PrintPolicy = TPrettyJsonPrintPolicy<CharType>>
-	void WriteObject(TJsonWriter<CharType, PrintPolicy>& JsonWriter) const
+	void WriteObject(TJsonWriter<CharType, PrintPolicy>& JsonWriter, bool bAllExtensionsRequired) const
 	{
 		FGLTFJsonExtensions AllExtensions(Extensions);
 
@@ -168,28 +168,28 @@ struct FGLTFJsonRoot
 		}
 
 		FGLTFJsonUtility::WriteStringArray(JsonWriter, TEXT("extensionsUsed"), AllExtensions.Used);
-		FGLTFJsonUtility::WriteStringArray(JsonWriter, TEXT("extensionsRequired"), AllExtensions.Required);
+		FGLTFJsonUtility::WriteStringArray(JsonWriter, TEXT("extensionsRequired"), bAllExtensionsRequired ? AllExtensions.Used : AllExtensions.Required);
 
 		JsonWriter.WriteObjectEnd();
 	}
 
 	template <class CharType = TCHAR, class PrintPolicy = TPrettyJsonPrintPolicy<CharType>>
-	void ToJson(FArchive* const Archive) const
+	void ToJson(FArchive* const Archive, bool bAllExtensionsRequired) const
 	{
 		TSharedRef<TJsonWriter<CharType, PrintPolicy>> JsonWriter = TJsonWriterFactory<CharType, PrintPolicy>::Create(Archive);
-		WriteObject(*JsonWriter);
+		WriteObject(*JsonWriter, bAllExtensionsRequired);
 		JsonWriter->Close();
 	}
 
-	void ToJson(FArchive* const Archive, bool bPrettyPrint = true) const
+	void ToJson(FArchive* const Archive, bool bPrettyPrint, bool bAllExtensionsRequired) const
 	{
 		if (bPrettyPrint)
 		{
-			ToJson<UTF8CHAR, TPrettyJsonPrintPolicy<UTF8CHAR>>(Archive);
+			ToJson<UTF8CHAR, TPrettyJsonPrintPolicy<UTF8CHAR>>(Archive, bAllExtensionsRequired);
 		}
 		else
 		{
-			ToJson<UTF8CHAR, TCondensedJsonPrintPolicy<UTF8CHAR>>(Archive);
+			ToJson<UTF8CHAR, TCondensedJsonPrintPolicy<UTF8CHAR>>(Archive, bAllExtensionsRequired);
 		}
 	}
 };
