[GLTFExporter] feat: write condensed json (instead of pretty) when exporting glb

commit d7e1ba26fd645f08a1367485a8fb4d3ae7a680e3
Date: Sun, 4 Oct 2020 22:37:33 +0200

    feat: write condensed json (instead of pretty) when exporting glb
    (cherry picked from commit 3896b5ce052b92e2ab13ff3c64f06661fbcebee7​)

#skipci

diff --git a/Engine/Plugins/Enterprise/GLTFExporter/MigratedFromGit/Source/GLTFExporter/Private/Builders/GLTFJsonBuilder.cpp b/Engine/Plugins/Enterprise/GLTFExporter/MigratedFromGit/Source/GLTFExporter/Private/Builders/GLTFJsonBuilder.cpp
--- a/Engine/Plugins/Enterprise/GLTFExporter/MigratedFromGit/Source/GLTFExporter/Private/Builders/GLTFJsonBuilder.cpp
+++ b/Engine/Plugins/Enterprise/GLTFExporter/MigratedFromGit/Source/GLTFExporter/Private/Builders/GLTFJsonBuilder.cpp
@@ -10,7 +10,7 @@ FGLTFJsonBuilder::FGLTFJsonBuilder(const FString& FilePath, const UGLTFExportOpt
 
 void FGLTFJsonBuilder::WriteJson(FArchive& Archive) const
 {
-	JsonRoot.ToJson(&Archive, true, ExportOptions->bAllExtensionsRequired);
+	JsonRoot.ToJson(&Archive, bIsGlbFile, ExportOptions->bAllExtensionsRequired);
 }
 
 void FGLTFJsonBuilder::AddExtension(EGLTFJsonExtension Extension, bool bIsRequired)
diff --git a/Engine/Plugins/Enterprise/GLTFExporter/MigratedFromGit/Source/GLTFExporter/Private/Json/GLTFJsonRoot.h b/Engine/Plugins/Enterprise/GLTFExporter/MigratedFromGit/Source/GLTFExporter/Private/Json/GLTFJsonRoot.h
--- a/Engine/Plugins/Enterprise/GLTFExporter/MigratedFromGit/Source/GLTFExporter/Private/Json/GLTFJsonRoot.h
+++ b/Engine/Plugins/Enterprise/GLTFExporter/MigratedFromGit/Source/GLTFExporter/Private/Json/GLTFJsonRoot.h
@@ -181,15 +181,15 @@ struct FGLTFJsonRoot
 		JsonWriter->Close();
 	}
 
-	void ToJson(FArchive* const Archive, bool bPrettyPrint, bool bAllExtensionsRequired) const
+	void ToJson(FArchive* const Archive, bool bCondensedPrint, bool bAllExtensionsRequired) const
 	{
-		if (bPrettyPrint)
+		if (bCondensedPrint)
 		{
-			ToJson<UTF8CHAR, TPrettyJsonPrintPolicy<UTF8CHAR>>(Archive, bAllExtensionsRequired);
+			ToJson<UTF8CHAR, TCondensedJsonPrintPolicy<UTF8CHAR>>(Archive, bAllExtensionsRequired);
 		}
 		else
 		{
-			ToJson<UTF8CHAR, TCondensedJsonPrintPolicy<UTF8CHAR>>(Archive, bAllExtensionsRequired);
+			ToJson<UTF8CHAR, TPrettyJsonPrintPolicy<UTF8CHAR>>(Archive, bAllExtensionsRequired);
 		}
 	}
 };
