[GLTFExporter] refactor: rename GLTFInteractionHotspotActor to GLTFHotspotActor and GLTFOrbitCameraActor to GLTFCameraActor

commit e8bfb7d66ffdd0156f3b2749ec1a5fbe69e9a665
Date: Mon, 5 Oct 2020 00:23:17 +0200

    refactor: rename GLTFInteractionHotspotActor to GLTFHotspotActor and GLTFOrbitCameraActor to GLTFCameraActor
    (cherry picked from commit 48675cf318758bac2ea4b1e9d6a1c9616d99059d​)

#skipci

diff --git a/Engine/Plugins/Enterprise/GLTFExporter/MigratedFromGit/Source/GLTFExporter/Private/Builders/GLTFConvertBuilder.cpp b/Engine/Plugins/Enterprise/GLTFExporter/MigratedFromGit/Source/GLTFExporter/Private/Builders/GLTFConvertBuilder.cpp
--- a/Engine/Plugins/Enterprise/GLTFExporter/MigratedFromGit/Source/GLTFExporter/Private/Builders/GLTFConvertBuilder.cpp
+++ b/Engine/Plugins/Enterprise/GLTFExporter/MigratedFromGit/Source/GLTFExporter/Private/Builders/GLTFConvertBuilder.cpp
@@ -376,7 +376,7 @@ FGLTFJsonLightMapIndex FGLTFConvertBuilder::GetOrAddLightMap(const UStaticMeshCo
 	return LightMapConverter.GetOrAdd(StaticMeshComponent);
 }
 
-FGLTFJsonHotspotIndex FGLTFConvertBuilder::GetOrAddHotspot(const AGLTFInteractionHotspotActor* HotspotActor)
+FGLTFJsonHotspotIndex FGLTFConvertBuilder::GetOrAddHotspot(const AGLTFHotspotActor* HotspotActor)
 {
 	if (HotspotActor == nullptr)
 	{
diff --git a/Engine/Plugins/Enterprise/GLTFExporter/MigratedFromGit/Source/GLTFExporter/Private/Builders/GLTFConvertBuilder.h b/Engine/Plugins/Enterprise/GLTFExporter/MigratedFromGit/Source/GLTFExporter/Private/Builders/GLTFConvertBuilder.h
--- a/Engine/Plugins/Enterprise/GLTFExporter/MigratedFromGit/Source/GLTFExporter/Private/Builders/GLTFConvertBuilder.h
+++ b/Engine/Plugins/Enterprise/GLTFExporter/MigratedFromGit/Source/GLTFExporter/Private/Builders/GLTFConvertBuilder.h
@@ -69,7 +69,7 @@ public:
 	FGLTFJsonBackdropIndex GetOrAddBackdrop(const AActor* BackdropActor);
 	FGLTFJsonVariationIndex GetOrAddVariation(const ALevelVariantSetsActor* LevelVariantSetsActor);
 	FGLTFJsonLightMapIndex GetOrAddLightMap(const UStaticMeshComponent* StaticMeshComponent);
-	FGLTFJsonHotspotIndex GetOrAddHotspot(const AGLTFInteractionHotspotActor* HotspotActor);
+	FGLTFJsonHotspotIndex GetOrAddHotspot(const AGLTFHotspotActor* HotspotActor);
 
 private:
 
diff --git a/Engine/Plugins/Enterprise/GLTFExporter/MigratedFromGit/Source/GLTFExporter/Private/Converters/GLTFCameraConverters.cpp b/Engine/Plugins/Enterprise/GLTFExporter/MigratedFromGit/Source/GLTFExporter/Private/Converters/GLTFCameraConverters.cpp
--- a/Engine/Plugins/Enterprise/GLTFExporter/MigratedFromGit/Source/GLTFExporter/Private/Converters/GLTFCameraConverters.cpp
+++ b/Engine/Plugins/Enterprise/GLTFExporter/MigratedFromGit/Source/GLTFExporter/Private/Converters/GLTFCameraConverters.cpp
@@ -5,7 +5,7 @@
 #include "Converters/GLTFConverterUtility.h"
 #include "Converters/GLTFNameUtility.h"
 #include "Converters/GLTFCameraUtility.h"
-#include "Actors/GLTFOrbitCameraActor.h"
+#include "Actors/GLTFCameraActor.h"
 
 FGLTFJsonCameraIndex FGLTFCameraConverter::Convert(const UCameraComponent* CameraComponent)
 {
@@ -32,22 +32,22 @@ FGLTFJsonCameraIndex FGLTFCameraConverter::Convert(const UCameraComponent* Camer
 	}
 
 	const AActor* Owner = CameraComponent->GetOwner();
-	const AGLTFOrbitCameraActor* OrbitCameraActor = Owner != nullptr ? Cast<AGLTFOrbitCameraActor>(Owner) : nullptr;
+	const AGLTFCameraActor* CameraActor = Owner != nullptr ? Cast<AGLTFCameraActor>(Owner) : nullptr;
 
-	if (OrbitCameraActor != nullptr)
+	if (CameraActor != nullptr)
 	{
 		if (Builder.ExportOptions->bExportOrbitalCameras)
 		{
 			FGLTFJsonOrbitCamera OrbitCamera;
-			OrbitCamera.Focus = Builder.GetOrAddNode(OrbitCameraActor->Focus);
-			OrbitCamera.MaxDistance = FGLTFConverterUtility::ConvertLength(OrbitCameraActor->DistanceMax, Builder.ExportOptions->ExportScale);
-			OrbitCamera.MinDistance = FGLTFConverterUtility::ConvertLength(OrbitCameraActor->DistanceMin, Builder.ExportOptions->ExportScale);
-			OrbitCamera.MaxAngle = OrbitCameraActor->PitchAngleMax;
-			OrbitCamera.MinAngle = OrbitCameraActor->PitchAngleMin;
-			OrbitCamera.DistanceSensitivity = OrbitCameraActor->DistanceSensitivity;
-			OrbitCamera.OrbitSensitivity = OrbitCameraActor->OrbitSensitivity;
-			OrbitCamera.OrbitInertia = OrbitCameraActor->OrbitInertia;
-			OrbitCamera.DollyDuration = OrbitCameraActor->DollyDuration;
+			OrbitCamera.Focus = Builder.GetOrAddNode(CameraActor->Focus);
+			OrbitCamera.MaxDistance = FGLTFConverterUtility::ConvertLength(CameraActor->DistanceMax, Builder.ExportOptions->ExportScale);
+			OrbitCamera.MinDistance = FGLTFConverterUtility::ConvertLength(CameraActor->DistanceMin, Builder.ExportOptions->ExportScale);
+			OrbitCamera.MaxAngle = CameraActor->PitchAngleMax;
+			OrbitCamera.MinAngle = CameraActor->PitchAngleMin;
+			OrbitCamera.DistanceSensitivity = CameraActor->DistanceSensitivity;
+			OrbitCamera.OrbitSensitivity = CameraActor->OrbitSensitivity;
+			OrbitCamera.OrbitInertia = CameraActor->OrbitInertia;
+			OrbitCamera.DollyDuration = CameraActor->DollyDuration;
 
 			Camera.OrbitCamera = OrbitCamera;
 		}
diff --git a/Engine/Plugins/Enterprise/GLTFExporter/MigratedFromGit/Source/GLTFExporter/Private/Converters/GLTFHotspotConverters.cpp b/Engine/Plugins/Enterprise/GLTFExporter/MigratedFromGit/Source/GLTFExporter/Private/Converters/GLTFHotspotConverters.cpp
--- a/Engine/Plugins/Enterprise/GLTFExporter/MigratedFromGit/Source/GLTFExporter/Private/Converters/GLTFHotspotConverters.cpp
+++ b/Engine/Plugins/Enterprise/GLTFExporter/MigratedFromGit/Source/GLTFExporter/Private/Converters/GLTFHotspotConverters.cpp
@@ -4,7 +4,7 @@
 #include "Builders/GLTFContainerBuilder.h"
 #include "Animation/SkeletalMeshActor.h"
 
-FGLTFJsonHotspotIndex FGLTFHotspotConverter::Convert(const AGLTFInteractionHotspotActor* HotspotActor)
+FGLTFJsonHotspotIndex FGLTFHotspotConverter::Convert(const AGLTFHotspotActor* HotspotActor)
 {
 	FGLTFJsonHotspot JsonHotspot;
 	HotspotActor->GetName(JsonHotspot.Name);
diff --git a/Engine/Plugins/Enterprise/GLTFExporter/MigratedFromGit/Source/GLTFExporter/Private/Converters/GLTFHotspotConverters.h b/Engine/Plugins/Enterprise/GLTFExporter/MigratedFromGit/Source/GLTFExporter/Private/Converters/GLTFHotspotConverters.h
--- a/Engine/Plugins/Enterprise/GLTFExporter/MigratedFromGit/Source/GLTFExporter/Private/Converters/GLTFHotspotConverters.h
+++ b/Engine/Plugins/Enterprise/GLTFExporter/MigratedFromGit/Source/GLTFExporter/Private/Converters/GLTFHotspotConverters.h
@@ -5,11 +5,11 @@
 #include "Json/GLTFJsonIndex.h"
 #include "Converters/GLTFConverter.h"
 #include "Converters/GLTFBuilderContext.h"
-#include "Actors/GLTFInteractionHotspotActor.h"
+#include "Actors/GLTFHotspotActor.h"
 
-class FGLTFHotspotConverter final : public FGLTFBuilderContext, public TGLTFConverter<FGLTFJsonHotspotIndex, const AGLTFInteractionHotspotActor*>
+class FGLTFHotspotConverter final : public FGLTFBuilderContext, public TGLTFConverter<FGLTFJsonHotspotIndex, const AGLTFHotspotActor*>
 {
 	using FGLTFBuilderContext::FGLTFBuilderContext;
 
-	virtual FGLTFJsonHotspotIndex Convert(const AGLTFInteractionHotspotActor* HotspotActor) override;
+	virtual FGLTFJsonHotspotIndex Convert(const AGLTFHotspotActor* HotspotActor) override;
 };
diff --git a/Engine/Plugins/Enterprise/GLTFExporter/MigratedFromGit/Source/GLTFExporter/Private/Converters/GLTFNameUtility.h b/Engine/Plugins/Enterprise/GLTFExporter/MigratedFromGit/Source/GLTFExporter/Private/Converters/GLTFNameUtility.h
--- a/Engine/Plugins/Enterprise/GLTFExporter/MigratedFromGit/Source/GLTFExporter/Private/Converters/GLTFNameUtility.h
+++ b/Engine/Plugins/Enterprise/GLTFExporter/MigratedFromGit/Source/GLTFExporter/Private/Converters/GLTFNameUtility.h
@@ -6,7 +6,7 @@
 
 class ULightComponent;
 class UCameraComponent;
-class AGLTFInteractionHotspotActor;
+class AGLTFHotspotActor;
 
 struct FGLTFNameUtility
 {
diff --git a/Engine/Plugins/Enterprise/GLTFExporter/MigratedFromGit/Source/GLTFExporter/Private/Converters/GLTFNodeConverters.cpp b/Engine/Plugins/Enterprise/GLTFExporter/MigratedFromGit/Source/GLTFExporter/Private/Converters/GLTFNodeConverters.cpp
--- a/Engine/Plugins/Enterprise/GLTFExporter/MigratedFromGit/Source/GLTFExporter/Private/Converters/GLTFNodeConverters.cpp
+++ b/Engine/Plugins/Enterprise/GLTFExporter/MigratedFromGit/Source/GLTFExporter/Private/Converters/GLTFNodeConverters.cpp
@@ -5,7 +5,7 @@
 #include "Converters/GLTFConverterUtility.h"
 #include "Converters/GLTFActorUtility.h"
 #include "Converters/GLTFNameUtility.h"
-#include "Actors/GLTFInteractionHotspotActor.h"
+#include "Actors/GLTFHotspotActor.h"
 
 FGLTFJsonNodeIndex FGLTFActorConverter::Convert(const AActor* Actor)
 {
@@ -31,7 +31,7 @@ FGLTFJsonNodeIndex FGLTFActorConverter::Convert(const AActor* Actor)
 			RootNode.Backdrop = Builder.GetOrAddBackdrop(Actor);
 		}
 	}
-	else if (const AGLTFInteractionHotspotActor* HotspotActor = Cast<AGLTFInteractionHotspotActor>(Actor))
+	else if (const AGLTFHotspotActor* HotspotActor = Cast<AGLTFHotspotActor>(Actor))
 	{
 		if (Builder.ExportOptions->bExportInteractionHotspots)
 		{
diff --git a/Engine/Plugins/Enterprise/GLTFExporter/MigratedFromGit/Source/GLTFExporterRuntime/Private/Actors/GLTFCameraActor.cpp b/Engine/Plugins/Enterprise/GLTFExporter/MigratedFromGit/Source/GLTFExporterRuntime/Private/Actors/GLTFCameraActor.cpp
new file mode 100644
--- /dev/null
+++ b/Engine/Plugins/Enterprise/GLTFExporter/MigratedFromGit/Source/GLTFExporterRuntime/Private/Actors/GLTFCameraActor.cpp
@@ -0,0 +1,217 @@
+// Copyright Epic Games, Inc. All Rights Reserved.
+
+#include "Actors/GLTFCameraActor.h"
+#include "Components/InputComponent.h"
+
+namespace
+{
+	// Uniquely identifies an angle by how many times it has crossed the 0-360 degree range.
+	// Positive angles are indexed from 0 and negative angles from -1.
+	int32 AngleCycleIndex(float Angle)
+	{
+		// TODO: is this similar to viewer camera?
+		return (static_cast<int32>(Angle) / 360) + ((Angle < 0.0f) ? -1 : 0);
+	}
+
+	// Scales to convert from the export-friendly sensitivity-values stored in our properties
+	// to values that we can use when processing axis-values (to get similar results as in the viewer).
+	const float OrbitSensitivityScale = 16.667f;
+	const float DistanceSensitivityScale = 0.1f;
+}
+
+AGLTFCameraActor::AGLTFCameraActor(const FObjectInitializer& ObjectInitializer)
+	: Super(ObjectInitializer)
+	, Focus(nullptr)
+	, DistanceMin(100.0f)
+	, DistanceMax(1000.0f)
+	, PitchAngleMin(-90.0f)
+	, PitchAngleMax(90.0f)
+	, DollyDuration(0.2f)
+	, OrbitInertia(0.1f)
+	, OrbitSensitivity(0.3f)
+	, DistanceSensitivity(0.5f)
+	, Distance(0.0f)
+	, Pitch(0.0f)
+	, Yaw(0.0f)
+	, TargetDistance(0.0f)
+	, TargetPitch(0.0f)
+	, TargetYaw(0.0f)
+	, DollyTime(0.0f)
+	, DollyStartDistance(0.0f)
+{
+	PrimaryActorTick.bCanEverTick = true;
+}
+
+#if WITH_EDITOR
+void AGLTFCameraActor::PostEditChangeProperty(FPropertyChangedEvent& PropertyChangedEvent)
+{
+	Super::PostEditChangeProperty(PropertyChangedEvent);
+
+	const FProperty* PropertyThatChanged = PropertyChangedEvent.Property;
+
+	if (PropertyThatChanged)
+	{
+		const FName PropertyName = PropertyThatChanged->GetFName();
+
+		if (PropertyName == GET_MEMBER_NAME_CHECKED(ThisClass, Focus))
+		{
+			if (Focus == this)
+			{
+				Focus = nullptr;
+				UE_LOG(LogTemp, Warning, TEXT("The camera focus must not be the camera's own actor"));
+			}
+		}
+	}
+}
+#endif // WITH_EDITOR
+
+void AGLTFCameraActor::BeginPlay()
+{
+	Super::BeginPlay();
+
+	const FVector FocusPosition = GetFocusPosition();
+
+	// Ensure that the camera is initially aimed at the focus-position
+	SetActorRotation(GetLookAtRotation(FocusPosition));
+
+	const FVector Position = GetActorLocation();
+	const FRotator Rotation = GetActorRotation();
+
+	// Calculate values based on the current location and orientation
+	Distance = ClampDistance((FocusPosition - Position).Size());
+	Pitch = ClampPitch(-Rotation.Pitch);
+	Yaw = ClampYaw(Rotation.Yaw);
+	TargetDistance = Distance;
+	TargetPitch = Pitch;
+	TargetYaw = Yaw;
+
+	if (InputComponent)
+	{
+		InputComponent->BindAxisKey(EKeys::MouseX, this, &AGLTFCameraActor::OnMouseX);
+		InputComponent->BindAxisKey(EKeys::MouseY, this, &AGLTFCameraActor::OnMouseY);
+		InputComponent->BindAxisKey(EKeys::MouseWheelAxis, this, &AGLTFCameraActor::OnMouseWheelAxis);
+	}
+}
+
+void AGLTFCameraActor::Tick(float DeltaSeconds)
+{
+	Super::Tick(DeltaSeconds);
+
+	if (DollyTime != 0.0f)
+	{
+		DollyTime = FMath::Max(DollyTime - DeltaSeconds, 0.0f);
+		Distance = FMath::InterpEaseInOut(DollyStartDistance, TargetDistance, (DollyDuration - DollyTime) / DollyDuration, 1.2f);
+	}
+
+	const float Alpha = (OrbitInertia == 0.0f) ? 1.0f : FMath::Min(DeltaSeconds / OrbitInertia, 1.0f);
+	Yaw = FMath::Lerp(Yaw, TargetYaw, Alpha);
+	Pitch = FMath::Lerp(Pitch, TargetPitch, Alpha);
+
+	const int32 YawCycleIndex = AngleCycleIndex(Yaw);
+	const int32 TargetYawCycleIndex = AngleCycleIndex(TargetYaw);
+
+	// Clamp the angles to a positive 360 degrees while they are within the same cycle.
+	// It is important that this doesn't happen during a transition as it will otherwise be skewed.
+	if (YawCycleIndex == TargetYawCycleIndex && YawCycleIndex != 0)
+	{
+		Yaw = ClampYaw(Yaw);
+		TargetYaw = ClampYaw(TargetYaw);
+	}
+
+	const FTransform FocusTransform = FTransform(GetFocusPosition());
+	const FTransform DollyTransform = FTransform(-FVector::ForwardVector * Distance);
+	const FTransform RotationTransform = FTransform(FQuat::MakeFromEuler(FVector(0.0f, -Pitch, Yaw)));
+	const FTransform ResultTransform = DollyTransform * RotationTransform * FocusTransform;
+
+	SetActorTransform(ResultTransform);
+}
+
+void AGLTFCameraActor::PreInitializeComponents()
+{
+	AutoReceiveInput = static_cast<EAutoReceiveInput::Type>(GetAutoActivatePlayerIndex() + 1);
+
+	Super::PreInitializeComponents();
+}
+
+void AGLTFCameraActor::PostActorCreated()
+{
+	SetAutoActivateForPlayer(EAutoReceiveInput::Player0);
+}
+
+void AGLTFCameraActor::OnMouseX(float AxisValue)
+{
+	TargetYaw += AxisValue * OrbitSensitivity * OrbitSensitivityScale;
+}
+
+void AGLTFCameraActor::OnMouseY(float AxisValue)
+{
+	TargetPitch = ClampPitch(TargetPitch - AxisValue * OrbitSensitivity * OrbitSensitivityScale);
+}
+
+void AGLTFCameraActor::OnMouseWheelAxis(float AxisValue)
+{
+	if (!FMath::IsNearlyZero(AxisValue))
+	{
+		const float DeltaDistance = -AxisValue * (TargetDistance * DistanceSensitivity * DistanceSensitivityScale);
+
+		DollyTime = DollyDuration;
+		TargetDistance = ClampDistance(TargetDistance + DeltaDistance);
+		DollyStartDistance = Distance;
+	}
+}
+
+float AGLTFCameraActor::ClampDistance(float Value) const
+{
+	return FMath::Clamp(Value, DistanceMin, DistanceMax);
+}
+
+float AGLTFCameraActor::ClampPitch(float Value) const
+{
+	return FMath::Clamp(Value, PitchAngleMin, PitchAngleMax);
+}
+
+float AGLTFCameraActor::ClampYaw(float Value) const
+{
+	const int32 CycleIndex = AngleCycleIndex(Value);
+	const float Offset = 360.0f * static_cast<float>(FMath::Abs(FMath::Min(CycleIndex, 0)));
+
+	return FMath::Fmod(Value + Offset, 360.0f);
+}
+
+void AGLTFCameraActor::RemoveInertia()
+{
+	Yaw = TargetYaw;
+	Pitch = TargetPitch;
+	Distance = TargetDistance;
+}
+
+FRotator AGLTFCameraActor::GetLookAtRotation(const FVector TargetPosition) const
+{
+	const FVector EyePosition = GetActorLocation();
+
+	return FRotationMatrix::MakeFromXZ(TargetPosition - EyePosition, FVector::UpVector).Rotator();
+}
+
+FVector AGLTFCameraActor::GetFocusPosition() const
+{
+	return this->Focus != nullptr ? this->Focus->GetActorLocation() : FVector::ZeroVector;
+}
+
+bool AGLTFCameraActor::SetAutoActivateForPlayer(const EAutoReceiveInput::Type Player)
+{
+	// TODO: remove hack by adding proper API access to ACameraActor (or by other means)
+	FProperty* Property = GetClass()->FindPropertyByName(TEXT("AutoActivateForPlayer"));
+	if (Property == nullptr)
+	{
+		return false;
+	}
+
+	TEnumAsByte<EAutoReceiveInput::Type>* ValuePtr = Property->ContainerPtrToValuePtr<TEnumAsByte<EAutoReceiveInput::Type>>(this);
+	if (ValuePtr == nullptr)
+	{
+		return false;
+	}
+
+	*ValuePtr = Player;
+	return true;
+}
diff --git a/Engine/Plugins/Enterprise/GLTFExporter/MigratedFromGit/Source/GLTFExporterRuntime/Private/Actors/GLTFHotspotActor.cpp b/Engine/Plugins/Enterprise/GLTFExporter/MigratedFromGit/Source/GLTFExporterRuntime/Private/Actors/GLTFHotspotActor.cpp
new file mode 100644
--- /dev/null
+++ b/Engine/Plugins/Enterprise/GLTFExporter/MigratedFromGit/Source/GLTFExporterRuntime/Private/Actors/GLTFHotspotActor.cpp
@@ -0,0 +1,467 @@
+// Copyright Epic Games, Inc. All Rights Reserved.
+
+#include "Actors/GLTFHotspotActor.h"
+#include "Animation/SkeletalMeshActor.h"
+#include "Animation/AnimSequence.h"
+#include "Components/SphereComponent.h"
+#include "Components/MaterialBillboardComponent.h"
+#include "Kismet/GameplayStatics.h"
+#include "UObject/ConstructorHelpers.h"
+#include "Animation/AnimSingleNodeInstance.h"
+#include "Materials/MaterialInstanceDynamic.h"
+
+#if WITH_EDITOR
+#include "LevelEditor.h"
+#include "Modules/ModuleManager.h"
+#include "SLevelViewport.h"
+#include "Slate/SceneViewport.h"
+#endif
+
+DEFINE_LOG_CATEGORY_STATIC(LogEditorGLTFInteractionHotspot, Log, All);
+
+namespace
+{
+	const FName NAME_InteractionHotspotTag = TEXT("InteractionHotspot");
+	const FName NAME_LevelEditorModule = TEXT("LevelEditor");
+	const FName NAME_SpriteParameter = TEXT("Sprite");
+	const FName NAME_OpacityParameter = TEXT("OpacityMask");
+}
+
+AGLTFHotspotActor::AGLTFHotspotActor(const FObjectInitializer& ObjectInitializer)
+	: Super(ObjectInitializer),
+	Image(nullptr),
+	HoveredImage(nullptr),
+	ToggledImage(nullptr),
+	ToggledHoveredImage(nullptr),
+	BillboardComponent(nullptr),
+	SphereComponent(nullptr),
+	DefaultMaterial(nullptr),
+	DefaultImage(nullptr),
+	DefaultHoveredImage(nullptr),
+	DefaultToggledImage(nullptr),
+	DefaultToggledHoveredImage(nullptr),
+	ActiveImage(nullptr),
+	ActiveImageSize(0.0f, 0.0f),
+	bToggled(bToggled),
+	bIsInteractable(true),
+	RealtimeSecondsWhenLastInSight(0.0f),
+	RealtimeSecondsWhenLastHidden(0.0f)
+{
+	USceneComponent* SceneRootComponent = CreateDefaultSubobject<USceneComponent>(TEXT("Root"));
+	SetRootComponent(SceneRootComponent);
+	AddInstanceComponent(SceneRootComponent);
+	SceneRootComponent->SetMobility(EComponentMobility::Movable);
+
+	BillboardComponent = CreateDefaultSubobject<UMaterialBillboardComponent>(TEXT("MaterialBillboardComponent"));
+	AddInstanceComponent(BillboardComponent);
+	BillboardComponent->SetupAttachment(RootComponent);
+	BillboardComponent->SetMobility(EComponentMobility::Movable);
+
+	SphereComponent = CreateDefaultSubobject<USphereComponent>(TEXT("SphereComponent"));
+	AddInstanceComponent(BillboardComponent);
+	SphereComponent->SetupAttachment(RootComponent);
+	SphereComponent->SetMobility(EComponentMobility::Movable);
+	SphereComponent->ComponentTags.Add(NAME_InteractionHotspotTag);
+	SphereComponent->InitSphereRadius(100.0f);
+	SphereComponent->SetVisibility(false);
+
+	// Setup the most minimalistic collision profile for mouse input events
+	SphereComponent->SetCollisionEnabled(ECollisionEnabled::QueryAndPhysics);
+	SphereComponent->SetCollisionResponseToAllChannels(ECR_Ignore);
+	SphereComponent->SetCollisionResponseToChannel(ECC_Visibility, ECR_Block);
+	SphereComponent->SetGenerateOverlapEvents(false);
+
+	// Respond to interactions with the sphere-component
+	SphereComponent->OnBeginCursorOver.AddDynamic(this, &AGLTFHotspotActor::BeginCursorOver);
+	SphereComponent->OnEndCursorOver.AddDynamic(this, &AGLTFHotspotActor::EndCursorOver);
+	SphereComponent->OnClicked.AddDynamic(this, &AGLTFHotspotActor::Clicked);
+
+	struct FConstructorStatics
+	{
+		ConstructorHelpers::FObjectFinder<UMaterial> DefaultMaterial;
+		ConstructorHelpers::FObjectFinder<UMaterial> DefaultIconMaterial;
+		ConstructorHelpers::FObjectFinder<UTexture2D> DefaultImage;
+		ConstructorHelpers::FObjectFinder<UTexture2D> DefaultHoveredImage;
+		ConstructorHelpers::FObjectFinder<UTexture2D> DefaultToggledImage;
+		ConstructorHelpers::FObjectFinder<UTexture2D> DefaultToggledHoveredImage;
+
+		FConstructorStatics()
+			: DefaultMaterial(TEXT("/GLTFExporter/Materials/Hotspot"))
+			, DefaultIconMaterial(TEXT("/GLTFExporter/Materials/HotspotIcon"))
+			, DefaultImage(TEXT("/GLTFExporter/Textures/Hotspots/Image"))
+			, DefaultHoveredImage(TEXT("/GLTFExporter/Textures/Hotspots/HoveredImage"))
+			, DefaultToggledImage(TEXT("/GLTFExporter/Textures/Hotspots/ToggledImage"))
+			, DefaultToggledHoveredImage(TEXT("/GLTFExporter/Textures/Hotspots/ToggledHoveredImage"))
+		{
+		}
+	};
+	static FConstructorStatics ConstructorStatics;
+
+	DefaultMaterial = ConstructorStatics.DefaultMaterial.Object;
+	DefaultIconMaterial = ConstructorStatics.DefaultIconMaterial.Object;
+	DefaultImage = ConstructorStatics.DefaultImage.Object;
+	DefaultHoveredImage = ConstructorStatics.DefaultHoveredImage.Object;
+	DefaultToggledImage = ConstructorStatics.DefaultToggledImage.Object;
+	DefaultToggledHoveredImage = ConstructorStatics.DefaultToggledHoveredImage.Object;
+
+	PrimaryActorTick.bCanEverTick = true;
+	PrimaryActorTick.bStartWithTickEnabled = true;
+}
+
+
+#if WITH_EDITOR
+void AGLTFHotspotActor::PostEditChangeProperty(FPropertyChangedEvent& PropertyChangedEvent)
+{
+	Super::PostEditChangeProperty(PropertyChangedEvent);
+
+	const FProperty* PropertyThatChanged = PropertyChangedEvent.Property;
+
+	if (PropertyThatChanged)
+	{
+		const FName PropertyFName = PropertyThatChanged->GetFName();
+
+		if (PropertyFName == GET_MEMBER_NAME_CHECKED(ThisClass, Image))
+		{
+			UpdateActiveImageFromState(EGLTFHotspotState::Default);
+		}
+		else if (PropertyFName == GET_MEMBER_NAME_CHECKED(ThisClass, SkeletalMeshActor))
+		{
+			if (SkeletalMeshActor != nullptr && AnimationSequence == nullptr)
+			{
+				AnimationSequence = Cast<UAnimSequence>(SkeletalMeshActor->GetSkeletalMeshComponent()->AnimationData.AnimToPlay);
+			}
+			ValidateAnimation();
+		}
+		else if (PropertyFName == GET_MEMBER_NAME_CHECKED(ThisClass, AnimationSequence))
+		{
+			ValidateAnimation();
+		}
+	}
+}
+#endif // WITH_EDITOR
+
+void AGLTFHotspotActor::PostRegisterAllComponents()
+{
+	Super::PostRegisterAllComponents();
+
+	SetupSpriteElement();
+	UpdateActiveImageFromState(EGLTFHotspotState::Default);
+}
+
+void AGLTFHotspotActor::Tick(float DeltaTime)
+{
+	UWorld* World = GetWorld();
+
+	// TODO: is is safe to assume that we want the first controller for projections?
+	APlayerController* PlayerController = World != nullptr ? World->GetFirstPlayerController() : nullptr;
+	if (PlayerController == nullptr)
+	{
+		return;
+	}
+
+	const FVector ColliderLocation = SphereComponent->GetComponentLocation();
+
+	FVector ColliderScreenLocation;
+	if (!PlayerController->ProjectWorldLocationToScreenWithDistance(ColliderLocation, ColliderScreenLocation))
+	{
+		return;
+	}
+
+	// Update scale of the sphere-component to match the screen-size of the active image
+	{
+		const FVector2D CornerScreenLocation = FVector2D(ColliderScreenLocation) + ActiveImageSize * 0.5f;
+		FVector RayLocation;
+		FVector RayDirection;
+
+		if (PlayerController->DeprojectScreenPositionToWorld(CornerScreenLocation.X, CornerScreenLocation.Y, RayLocation, RayDirection))
+		{
+			const FVector ExtentLocation = RayLocation + RayDirection * ColliderScreenLocation.Z;
+			const float NewSphereRadius = (ExtentLocation - ColliderLocation).Size() / SphereComponent->GetShapeScale();
+			const float OldSphereRadius = SphereComponent->GetUnscaledSphereRadius();
+
+			if (FMath::Abs(NewSphereRadius - OldSphereRadius) > 0.1f)	// TODO: better epsilon?
+			{
+				SphereComponent->SetSphereRadius(NewSphereRadius);
+			}
+		}
+	}
+
+	// Update opacity and interactivity of the hotspot based on if it its occluded by other objects or not
+	{
+		FHitResult HitResult;
+		bool bIsHotspotOccluded = false;
+
+		if (PlayerController->GetHitResultAtScreenPosition(FVector2D(ColliderScreenLocation), ECC_Visibility, false, HitResult))
+		{
+			if (const UPrimitiveComponent* HitComponent = HitResult.GetComponent())
+			{
+				bIsHotspotOccluded = !HitComponent->ComponentTags.Contains(NAME_InteractionHotspotTag);
+			}
+		}
+
+		const float CurrentRealtimeSeconds = UGameplayStatics::GetRealTimeSeconds(World);
+		float Opacity;
+
+		if (bIsHotspotOccluded)
+		{
+			RealtimeSecondsWhenLastHidden = CurrentRealtimeSeconds;
+
+			const float HiddenDuration = FMath::Max(RealtimeSecondsWhenLastHidden - RealtimeSecondsWhenLastInSight, 0.0f);
+			const float FadeOutDuration = 0.5f;
+
+			Opacity = 1.0f - FMath::Clamp(HiddenDuration / FadeOutDuration, 0.0f, 1.0f);
+		}
+		else
+		{
+			RealtimeSecondsWhenLastInSight = CurrentRealtimeSeconds;
+
+			const float VisibleDuration = FMath::Max(RealtimeSecondsWhenLastInSight - RealtimeSecondsWhenLastHidden, 0.0f);
+			const float FadeInDuration = 0.25f;
+
+			Opacity = FMath::Clamp(VisibleDuration / FadeInDuration, 0.0f, 1.0f);
+		}
+
+		SetSpriteOpacity(Opacity);
+		bIsInteractable = Opacity >= 0.5f;
+	}
+}
+
+void AGLTFHotspotActor::BeginCursorOver(UPrimitiveComponent* TouchedComponent)
+{
+	if (bIsInteractable)
+	{
+		UpdateActiveImageFromState(bToggled ? EGLTFHotspotState::ToggledHovered :  EGLTFHotspotState::Hovered);
+	}
+}
+
+void AGLTFHotspotActor::EndCursorOver(UPrimitiveComponent* TouchedComponent)
+{
+	UpdateActiveImageFromState(bToggled ? EGLTFHotspotState::Toggled :  EGLTFHotspotState::Default);
+}
+
+void AGLTFHotspotActor::Clicked(UPrimitiveComponent* TouchedComponent, FKey ButtonPressed)
+{
+	if (!bIsInteractable)
+	{
+		return;
+	}
+
+	const bool bReverseAnimation = bToggled;
+
+	if (SkeletalMeshActor != nullptr && AnimationSequence != nullptr)
+	{
+		USkeletalMeshComponent* SkeletalMeshComponent = SkeletalMeshActor->GetSkeletalMeshComponent();
+		SkeletalMeshComponent->SetAnimationMode(EAnimationMode::Type::AnimationSingleNode);
+		const float AbsolutePlayRate = FMath::Abs(SkeletalMeshComponent->GetPlayRate());
+		const UAnimSingleNodeInstance* SingleNodeInstance = SkeletalMeshComponent->GetSingleNodeInstance();
+
+		if (SkeletalMeshComponent->IsPlaying() && SingleNodeInstance != nullptr && SingleNodeInstance->GetAnimationAsset() == AnimationSequence)
+		{
+			// If the same animation is already playing, just reverse the play rate for a smooth transition
+			SkeletalMeshComponent->SetPlayRate(AbsolutePlayRate * -1.0f);
+		}
+		else
+		{
+			SkeletalMeshComponent->SetAnimation(AnimationSequence);
+			SkeletalMeshComponent->SetPlayRate(AbsolutePlayRate * (bReverseAnimation ? -1.0f : 1.0f));
+			SkeletalMeshComponent->SetPosition(bReverseAnimation ? AnimationSequence->GetPlayLength() : 0.0f);
+			SkeletalMeshComponent->Play(false);
+		}
+	}
+
+	bToggled = !bToggled;
+
+	UpdateActiveImageFromState(bToggled ? EGLTFHotspotState::ToggledHovered :  EGLTFHotspotState::Hovered);
+}
+
+void AGLTFHotspotActor::UpdateActiveImageFromState(EGLTFHotspotState State)
+{
+	UTexture2D* NewImage = const_cast<UTexture2D*>(GetImageForState(State));
+
+	UTexture* DefaultTexture;
+	GetSpriteMaterial()->GetTextureParameterDefaultValue(NAME_SpriteParameter, DefaultTexture);
+
+	UTexture* SpriteTexture = NewImage != nullptr ? NewImage : DefaultTexture;
+	GetSpriteMaterial()->SetTextureParameterValue(NAME_SpriteParameter, SpriteTexture);
+
+	ActiveImage = NewImage;
+	ActiveImageSize = { SpriteTexture->GetSurfaceWidth(), SpriteTexture->GetSurfaceHeight() };
+
+	// NOTE: we do this even if size is unchanged since the last update may have failed
+	UpdateSpriteSize();
+}
+
+void AGLTFHotspotActor::SetupSpriteElement() const
+{
+	UMaterialInstanceDynamic* MaterialInstance;
+
+#if WITH_EDITORONLY_DATA
+	if (GetWorld() != nullptr && GetWorld()->WorldType == EWorldType::Editor)
+	{
+		MaterialInstance = UMaterialInstanceDynamic::Create(DefaultIconMaterial, GetTransientPackage());
+	}
+	else
+#endif // WITH_EDITORONLY_DATA
+	{
+		MaterialInstance = UMaterialInstanceDynamic::Create(DefaultMaterial, GetTransientPackage());
+	}
+
+	FMaterialSpriteElement Element;
+	Element.Material = MaterialInstance;
+	Element.bSizeIsInScreenSpace = true;
+	Element.BaseSizeX = 0.1f;
+	Element.BaseSizeY = 0.1f;
+
+	BillboardComponent->SetElements({ Element });
+}
+
+UMaterialInstanceDynamic* AGLTFHotspotActor::GetSpriteMaterial() const
+{
+	return static_cast<UMaterialInstanceDynamic*>(BillboardComponent->GetMaterial(0));
+}
+
+void AGLTFHotspotActor::UpdateSpriteSize()
+{
+	const FIntPoint ViewportSize = GetCurrentViewportSize();
+
+	float BaseSizeX = 0.1f;
+	float BaseSizeY = 0.1f;
+
+	if (ViewportSize.X > 0 && ViewportSize.Y > 0)
+	{
+		BaseSizeX = ActiveImageSize.X / static_cast<float>(ViewportSize.X);
+		BaseSizeY = ActiveImageSize.Y / static_cast<float>(ViewportSize.Y);
+	}
+
+	FMaterialSpriteElement& Element = BillboardComponent->Elements[0];
+
+	if (BaseSizeX != Element.BaseSizeX || BaseSizeY != Element.BaseSizeY)	// TODO: use epsilon for comparison?
+	{
+		Element.BaseSizeX = BaseSizeX;
+		Element.BaseSizeY = BaseSizeY;
+
+		BillboardComponent->MarkRenderStateDirty();
+	}
+}
+
+void AGLTFHotspotActor::SetSpriteOpacity(const float Opacity) const
+{
+	GetSpriteMaterial()->SetScalarParameterValue(NAME_OpacityParameter, Opacity);
+}
+
+FIntPoint AGLTFHotspotActor::GetCurrentViewportSize()
+{
+	// TODO: verify that correct size is calculated in the various play-modes and in the editor
+
+	const FViewport* Viewport = nullptr;
+
+	if (UWorld* World = GetWorld())
+	{
+		if (World->IsGameWorld())
+		{
+			if (UGameViewportClient* GameViewportClient = World->GetGameViewport())
+			{
+				Viewport = GameViewportClient->Viewport;
+			}
+		}
+#if WITH_EDITOR
+		else
+		{
+			if (FModuleManager::Get().IsModuleLoaded(NAME_LevelEditorModule))
+			{
+				FLevelEditorModule& LevelEditorModule = FModuleManager::GetModuleChecked<FLevelEditorModule>(NAME_LevelEditorModule);
+
+				if (const TSharedPtr<SLevelViewport> ActiveLevelViewport = LevelEditorModule.GetFirstActiveLevelViewport())
+				{
+					Viewport = ActiveLevelViewport->GetActiveViewport();
+				}
+			}
+		}
+#endif
+	}
+
+	if (Viewport != nullptr)
+	{
+		if (!Viewport->ViewportResizedEvent.IsBoundToObject(this))
+		{
+			Viewport->ViewportResizedEvent.AddUObject(this, &AGLTFHotspotActor::ViewportResized);
+		}
+
+		return Viewport->GetSizeXY();
+	}
+	else
+	{
+		return FIntPoint(0, 0);
+	}
+}
+
+void AGLTFHotspotActor::ViewportResized(FViewport*, uint32)
+{
+	UpdateSpriteSize();
+}
+
+void AGLTFHotspotActor::ValidateAnimation()
+{
+	if (SkeletalMeshActor != nullptr && AnimationSequence != nullptr)
+	{
+		if (USkeletalMesh* SkeletalMesh = SkeletalMeshActor->GetSkeletalMeshComponent()->SkeletalMesh)
+		{
+			if (SkeletalMesh->Skeleton != AnimationSequence->GetSkeleton())
+			{
+				if (SkeletalMesh->Skeleton != nullptr)
+				{
+					UE_LOG(LogEditorGLTFInteractionHotspot, Warning, TEXT("Animation %s is incompatible with skeleton %s, removing animation from actor."), *AnimationSequence->GetName(), *SkeletalMesh->Skeleton->GetName());
+				}
+				else
+				{
+					UE_LOG(LogEditorGLTFInteractionHotspot, Warning, TEXT("Animation %s is incompatible because mesh %s has no skeleton, removing animation from actor."), *AnimationSequence->GetName(), *SkeletalMesh->GetName());
+				}
+
+				AnimationSequence = nullptr;
+			}
+		}
+	}
+}
+
+const UTexture2D* AGLTFHotspotActor::GetImageForState(EGLTFHotspotState State) const
+{
+	const UTexture2D* CurrentImage = DefaultImage;
+	const UTexture2D* CurrentHoveredImage = DefaultHoveredImage;
+	const UTexture2D* CurrentToggledImage = DefaultToggledImage;
+	const UTexture2D* CurrentToggledHoveredImage = DefaultToggledHoveredImage;
+
+	if (Image != nullptr)
+	{
+		CurrentImage = Image;
+		CurrentHoveredImage = HoveredImage != nullptr ? HoveredImage : CurrentImage;
+		CurrentToggledImage = ToggledImage != nullptr ? ToggledImage : CurrentImage;
+
+		if (ToggledHoveredImage != nullptr)
+		{
+			CurrentToggledHoveredImage = ToggledHoveredImage;
+		}
+		else if (ToggledImage != nullptr)
+		{
+			CurrentToggledHoveredImage = ToggledImage;
+		}
+		else if (HoveredImage != nullptr)
+		{
+			CurrentToggledHoveredImage = HoveredImage;
+		}
+		else
+		{
+			CurrentToggledHoveredImage = Image;
+		}
+	}
+
+	switch (State)
+	{
+		case EGLTFHotspotState::Default:        return CurrentImage;
+		case EGLTFHotspotState::Hovered:        return CurrentHoveredImage;
+		case EGLTFHotspotState::Toggled:        return CurrentToggledImage;
+		case EGLTFHotspotState::ToggledHovered: return CurrentToggledHoveredImage;
+		default:                                checkNoEntry();
+	}
+
+	return nullptr;
+}
diff --git a/Engine/Plugins/Enterprise/GLTFExporter/MigratedFromGit/Source/GLTFExporterRuntime/Private/Actors/GLTFInteractionHotspotActor.cpp b/Engine/Plugins/Enterprise/GLTFExporter/MigratedFromGit/Source/GLTFExporterRuntime/Private/Actors/GLTFInteractionHotspotActor.cpp
deleted file mode 100644
diff --git a/Engine/Plugins/Enterprise/GLTFExporter/MigratedFromGit/Source/GLTFExporterRuntime/Private/Actors/GLTFOrbitCameraActor.cpp b/Engine/Plugins/Enterprise/GLTFExporter/MigratedFromGit/Source/GLTFExporterRuntime/Private/Actors/GLTFOrbitCameraActor.cpp
deleted file mode 100644
diff --git a/Engine/Plugins/Enterprise/GLTFExporter/MigratedFromGit/Source/GLTFExporterRuntime/Public/Actors/GLTFCameraActor.h b/Engine/Plugins/Enterprise/GLTFExporter/MigratedFromGit/Source/GLTFExporterRuntime/Public/Actors/GLTFCameraActor.h
new file mode 100644
--- /dev/null
+++ b/Engine/Plugins/Enterprise/GLTFExporter/MigratedFromGit/Source/GLTFExporterRuntime/Public/Actors/GLTFCameraActor.h
@@ -0,0 +1,103 @@
+// Copyright Epic Games, Inc. All Rights Reserved.
+
+#pragma once
+
+#include "Camera/CameraActor.h"
+#include "GLTFCameraActor.generated.h"
+
+/**
+ * GLTF-compatible camera that will carry over settings and simulate the behavior in the resulting viewer.
+ * Focuses one actor in the scene and orbits it through mouse control.
+ */
+UCLASS(BlueprintType, Blueprintable, DisplayName = "GLTF Camera")
+class GLTFEXPORTERRUNTIME_API AGLTFCameraActor : public ACameraActor
+{
+	GENERATED_BODY()
+	//~ Begin UObject Interface
+public:
+	AGLTFCameraActor(const FObjectInitializer& ObjectInitializer);
+#if WITH_EDITOR
+	virtual void PostEditChangeProperty(FPropertyChangedEvent& PropertyChangedEvent) override;
+#endif // WITH_EDITOR
+	//~ End UObject Interface
+
+	//~ Begin AActor Interface
+protected:
+	virtual void BeginPlay() override;
+public:
+	virtual void Tick(float DeltaSeconds) override;
+	virtual void PreInitializeComponents() override;
+	virtual void PostActorCreated() override;
+	//~ End AActor Interface
+
+private:
+	UFUNCTION()
+	void OnMouseX(float AxisValue);
+
+	UFUNCTION()
+	void OnMouseY(float AxisValue);
+
+	UFUNCTION()
+	void OnMouseWheelAxis(float AxisValue);
+
+	float ClampDistance(float Value) const;
+
+	float ClampPitch(float Value) const;
+
+	float ClampYaw(float Value) const;
+
+	void RemoveInertia();
+
+	FRotator GetLookAtRotation(const FVector TargetPosition) const;
+
+	FVector GetFocusPosition() const;
+
+	bool SetAutoActivateForPlayer(const EAutoReceiveInput::Type Player);
+
+public:
+	/* Actor which the camera will focus on and subsequently orbit. */
+	UPROPERTY(EditAnywhere, BlueprintReadWrite, Category = "GLTF Orbit Camera Actor")
+	AActor* Focus;
+
+	/* Closest distance the camera can approach the focused actor. */
+	UPROPERTY(EditAnywhere, BlueprintReadWrite, Category = "GLTF Orbit Camera Actor")
+	float DistanceMin;
+
+	/* Farthest distance the camera can recede from the focused actor. */
+	UPROPERTY(EditAnywhere, BlueprintReadWrite, Category = "GLTF Orbit Camera Actor")
+	float DistanceMax;
+
+	/* Minimum angle (in degrees) that the camera can pitch relative to the focused actor. */
+	UPROPERTY(EditAnywhere, BlueprintReadWrite, Category = "GLTF Orbit Camera Actor")
+	float PitchAngleMin;
+
+	/* Maximum angle (in degrees) that the camera can pitch relative to the focused actor. */
+	UPROPERTY(EditAnywhere, BlueprintReadWrite, Category = "GLTF Orbit Camera Actor")
+	float PitchAngleMax;
+
+	/* Duration (in seconds) that it takes the camera to complete a change in distance. */
+	UPROPERTY(EditAnywhere, BlueprintReadOnly, Category = "GLTF Orbit Camera Actor")
+	float DollyDuration;
+
+	/* Deceleration that occurs after orbital movement. The higher the value, the longer it takes to settle. */
+	UPROPERTY(EditAnywhere, BlueprintReadOnly, Category = "GLTF Orbit Camera Actor")
+	float OrbitInertia;
+
+	/* Size of the orbital movement relative to user input. The higher the value, the faster it moves. */
+	UPROPERTY(EditAnywhere, BlueprintReadOnly, Category = "GLTF Orbit Camera Actor")
+	float OrbitSensitivity;
+
+	/* Size of the dolly movement relative to user input. The higher the value, the faster it moves. */
+	UPROPERTY(EditAnywhere, BlueprintReadOnly, Category = "GLTF Orbit Camera Actor")
+	float DistanceSensitivity;
+
+private:
+	float Distance;
+	float Pitch;
+	float Yaw;
+	float TargetDistance;
+	float TargetPitch;
+	float TargetYaw;
+	float DollyTime;
+	float DollyStartDistance;
+};
diff --git a/Engine/Plugins/Enterprise/GLTFExporter/MigratedFromGit/Source/GLTFExporterRuntime/Public/Actors/GLTFHotspotActor.h b/Engine/Plugins/Enterprise/GLTFExporter/MigratedFromGit/Source/GLTFExporterRuntime/Public/Actors/GLTFHotspotActor.h
new file mode 100644
--- /dev/null
+++ b/Engine/Plugins/Enterprise/GLTFExporter/MigratedFromGit/Source/GLTFExporterRuntime/Public/Actors/GLTFHotspotActor.h
@@ -0,0 +1,140 @@
+// Copyright Epic Games, Inc. All Rights Reserved.
+
+#pragma once
+
+#include "GameFramework/Actor.h"
+#include "GLTFHotspotActor.generated.h"
+
+class ASkeletalMeshActor;
+class UAnimSequence;
+class UTexture2D;
+class UTexture;
+class UMaterialInterface;
+class UMaterialBillboardComponent;
+class USphereComponent;
+class UMaterialInstanceDynamic;
+class FViewport;
+
+enum class EGLTFHotspotState : uint8
+{
+	Default,
+	Hovered,
+	Toggled,
+	ToggledHovered
+};
+
+/**
+ * Actor wrapper for the GLTF hotspot component. Appears as a billboard and allows playback of skeletal animations when cursor input is enabled.
+ */
+UCLASS(BlueprintType, Blueprintable, HideCategories = (Sprite, Physics, Collision, Navigation), DisplayName = "GLTF Hotspot")
+class GLTFEXPORTERRUNTIME_API AGLTFHotspotActor : public AActor
+{
+	GENERATED_BODY()
+	//~ Begin UObject Interface
+public:
+	AGLTFHotspotActor(const FObjectInitializer& ObjectInitializer);
+#if WITH_EDITOR
+	virtual void PostEditChangeProperty(FPropertyChangedEvent& PropertyChangedEvent) override;
+#endif // WITH_EDITOR
+	//~ End UObject Interface
+
+	//~ Begin AActor Interface
+	virtual void PostRegisterAllComponents() override;
+	virtual void Tick(float DeltaSeconds) override;
+	//~ End AActor Interface
+
+private:
+	UFUNCTION()
+	void BeginCursorOver(UPrimitiveComponent* TouchedComponent);
+
+	UFUNCTION()
+	void EndCursorOver(UPrimitiveComponent* TouchedComponent);
+
+	UFUNCTION()
+	void Clicked(UPrimitiveComponent* TouchedComponent, FKey ButtonPressed);
+
+	void UpdateActiveImageFromState(EGLTFHotspotState State);
+
+public:
+
+	/* The skeletal mesh actor that will be animated when the hotspot is clicked. */
+	UPROPERTY(EditAnywhere, BlueprintReadWrite, Category = "GLTF Interaction Hotspot")
+	ASkeletalMeshActor* SkeletalMeshActor;
+
+	/* The animation that will be played on the skeletal mesh actor. Must be compatible with its skeletal mesh asset. */
+	UPROPERTY(EditAnywhere, BlueprintReadWrite, Category = "GLTF Interaction Hotspot", meta=(EditCondition="SkeletalMeshActor != nullptr"))
+	UAnimSequence* AnimationSequence;
+
+	/* The billboard image that will be shown when the hotspot is in an inactive state or one without a specified image. */
+	UPROPERTY(EditAnywhere, BlueprintReadWrite, Category = "GLTF Interaction Hotspot")
+	UTexture2D* Image;
+
+	/** The optional billboard image that will be shown when a cursor enters the hotspot. */
+	UPROPERTY(EditAnywhere, BlueprintReadWrite, Category = "GLTF Interaction Hotspot", meta=(EditCondition="Image != nullptr"))
+	UTexture2D* HoveredImage;
+
+	/** The optional billboard image that will be shown when the hotspot is toggled by a click. */
+	UPROPERTY(EditAnywhere, BlueprintReadWrite, Category = "GLTF Interaction Hotspot", meta=(EditCondition="Image != nullptr"))
+	UTexture2D* ToggledImage;
+
+	/** The optional billboard image that will be shown when the hotspot is toggled by a click and a cursor enters it. */
+	UPROPERTY(EditAnywhere, BlueprintReadWrite, Category = "GLTF Interaction Hotspot", meta=(EditCondition="Image != nullptr"))
+	UTexture2D* ToggledHoveredImage;
+
+private:
+
+	UPROPERTY(Transient, DuplicateTransient)
+	UMaterialBillboardComponent* BillboardComponent;
+
+	UPROPERTY(Transient, DuplicateTransient)
+	USphereComponent* SphereComponent;
+
+	UPROPERTY(Transient, DuplicateTransient)
+	UMaterialInterface* DefaultMaterial;
+
+	UPROPERTY(Transient, DuplicateTransient)
+	UTexture2D* DefaultImage;
+
+	UPROPERTY(Transient, DuplicateTransient)
+	UTexture2D* DefaultHoveredImage;
+
+	UPROPERTY(Transient, DuplicateTransient)
+	UTexture2D* DefaultToggledImage;
+
+	UPROPERTY(Transient, DuplicateTransient)
+	UTexture2D* DefaultToggledHoveredImage;
+
+	UPROPERTY(Transient, DuplicateTransient)
+	UMaterialInterface* DefaultIconMaterial;
+
+	UPROPERTY(Transient, DuplicateTransient)
+	const UTexture* ActiveImage;
+
+	FVector2D ActiveImageSize;
+
+	bool bToggled;
+
+	bool bIsInteractable;
+
+	float RealtimeSecondsWhenLastInSight;
+
+	float RealtimeSecondsWhenLastHidden;
+
+	void SetupSpriteElement() const;
+
+	UMaterialInstanceDynamic* GetSpriteMaterial() const;
+
+	void UpdateSpriteSize();
+
+	void SetSpriteOpacity(const float Opacity) const;
+
+	FIntPoint GetCurrentViewportSize();
+
+	void ViewportResized(FViewport*, uint32);
+
+	void ValidateAnimation();
+
+public:
+
+	const UTexture2D* GetImageForState(EGLTFHotspotState State) const;
+};
diff --git a/Engine/Plugins/Enterprise/GLTFExporter/MigratedFromGit/Source/GLTFExporterRuntime/Public/Actors/GLTFInteractionHotspotActor.h b/Engine/Plugins/Enterprise/GLTFExporter/MigratedFromGit/Source/GLTFExporterRuntime/Public/Actors/GLTFInteractionHotspotActor.h
deleted file mode 100644
diff --git a/Engine/Plugins/Enterprise/GLTFExporter/MigratedFromGit/Source/GLTFExporterRuntime/Public/Actors/GLTFOrbitCameraActor.h b/Engine/Plugins/Enterprise/GLTFExporter/MigratedFromGit/Source/GLTFExporterRuntime/Public/Actors/GLTFOrbitCameraActor.h
deleted file mode 100644
