[GLTFExporter] refactor: switch animation and mesh task priority and remove actor priority

commit d134d7a60de264da991da0f698a8e14b966d8702
Date: Mon, 5 Oct 2020 13:35:13 +0200

    refactor: switch animation and mesh task priority and remove actor priority
    (cherry picked from commit 4a0845108a13765651f7d82c4e7fcedcc0d98f2e​)

#skipci

diff --git a/Engine/Plugins/Enterprise/GLTFExporter/MigratedFromGit/Source/GLTFExporter/Private/Builders/GLTFTaskBuilder.cpp b/Engine/Plugins/Enterprise/GLTFExporter/MigratedFromGit/Source/GLTFExporter/Private/Builders/GLTFTaskBuilder.cpp
--- a/Engine/Plugins/Enterprise/GLTFExporter/MigratedFromGit/Source/GLTFExporter/Private/Builders/GLTFTaskBuilder.cpp
+++ b/Engine/Plugins/Enterprise/GLTFExporter/MigratedFromGit/Source/GLTFExporter/Private/Builders/GLTFTaskBuilder.cpp
@@ -65,11 +65,10 @@ FText FGLTFTaskBuilder::GetPriorityMessageFormat(EGLTFTaskPriority Priority)
 {
 	switch (Priority)
 	{
-		case EGLTFTaskPriority::Actor:     return NSLOCTEXT("GLTFExporter", "ActorTaskMessage", "Converting actors: {0}");
-		case EGLTFTaskPriority::Mesh:      return NSLOCTEXT("GLTFExporter", "MeshTaskMessage", "Converting meshes: {0}");
-		case EGLTFTaskPriority::Animation: return NSLOCTEXT("GLTFExporter", "AnimationTaskMessage", "Converting animations: {0}");
-		case EGLTFTaskPriority::Material:  return NSLOCTEXT("GLTFExporter", "MaterialTaskMessage", "Converting materials: {0}");
-		case EGLTFTaskPriority::Texture:   return NSLOCTEXT("GLTFExporter", "TextureTaskMessage", "Converting textures: {0}");
+		case EGLTFTaskPriority::Animation: return NSLOCTEXT("GLTFExporter", "AnimationTaskMessage", "Animation(s): {0}");
+		case EGLTFTaskPriority::Mesh:      return NSLOCTEXT("GLTFExporter", "MeshTaskMessage", "Mesh(es): {0}");
+		case EGLTFTaskPriority::Material:  return NSLOCTEXT("GLTFExporter", "MaterialTaskMessage", "Material(s): {0}");
+		case EGLTFTaskPriority::Texture:   return NSLOCTEXT("GLTFExporter", "TextureTaskMessage", "Texture(s): {0}");
 		default:                           checkNoEntry();
 	}
 
diff --git a/Engine/Plugins/Enterprise/GLTFExporter/MigratedFromGit/Source/GLTFExporter/Private/Tasks/GLTFTask.h b/Engine/Plugins/Enterprise/GLTFExporter/MigratedFromGit/Source/GLTFExporter/Private/Tasks/GLTFTask.h
--- a/Engine/Plugins/Enterprise/GLTFExporter/MigratedFromGit/Source/GLTFExporter/Private/Tasks/GLTFTask.h
+++ b/Engine/Plugins/Enterprise/GLTFExporter/MigratedFromGit/Source/GLTFExporter/Private/Tasks/GLTFTask.h
@@ -6,12 +6,11 @@
 
 enum class EGLTFTaskPriority : uint8
 {
-	Actor, // TODO: add and implement tasks for actors
-    Mesh,
-    Animation,
-    Material,
-    Texture,
-    MAX
+	Animation,
+	Mesh,
+	Material,
+	Texture,
+	MAX
 };
 
 class FGLTFTask
