[GLTFExporter] chore: add module dependencies for exporting level sequences

commit 5e5d9eaa722ae0b9e3fac44b19cfb249f82c5319
Date: Mon, 5 Oct 2020 10:41:04 +0200

    chore: add module dependencies for exporting level sequences
    (cherry picked from commit f4949897ae614cd635409ca0d8db3330738b9086​)

#skipci

diff --git a/Engine/Plugins/Enterprise/GLTFExporter/MigratedFromGit/Source/GLTFExporter/GLTFExporter.Build.cs b/Engine/Plugins/Enterprise/GLTFExporter/MigratedFromGit/Source/GLTFExporter/GLTFExporter.Build.cs
--- a/Engine/Plugins/Enterprise/GLTFExporter/MigratedFromGit/Source/GLTFExporter/GLTFExporter.Build.cs
+++ b/Engine/Plugins/Enterprise/GLTFExporter/MigratedFromGit/Source/GLTFExporter/GLTFExporter.Build.cs
@@ -27,6 +27,9 @@ namespace UnrealBuildTool.Rules
 					"Projects",
 					"RenderCore",
 					"RHI",
+					"LevelSequence",
+					"MovieScene",
+					"MovieSceneTracks",
 					"MaterialBaking",
 					"VariantManagerContent",
 					"GLTFExporterRuntime"
