[GLTFExporter] feat: add GLTFExportOptions::bExportLevelSequences

commit 8467a48aad275cc20b7f929d570a87236a3b7e49
Date: Mon, 5 Oct 2020 11:22:19 +0200

    feat: add GLTFExportOptions::bExportLevelSequences
    (cherry picked from commit b968a4197982b8492c3476fb5f12f7d2a8fba5be​)

#skipci

diff --git a/Engine/Plugins/Enterprise/GLTFExporter/MigratedFromGit/Source/GLTFExporter/Private/GLTFExportOptions.cpp b/Engine/Plugins/Enterprise/GLTFExporter/MigratedFromGit/Source/GLTFExporter/Private/GLTFExportOptions.cpp
--- a/Engine/Plugins/Enterprise/GLTFExporter/MigratedFromGit/Source/GLTFExporter/Private/GLTFExportOptions.cpp
+++ b/Engine/Plugins/Enterprise/GLTFExporter/MigratedFromGit/Source/GLTFExporter/Private/GLTFExportOptions.cpp
@@ -23,6 +23,7 @@ UGLTFExportOptions::UGLTFExportOptions(const FObjectInitializer& ObjectInitializ
 	DefaultLevelOfDetail = 0;
 	bExportVertexSkinWeights = true;
 	bExportAnimationSequences = true;
+	bExportLevelSequences = true;
 	bRetargetBoneTransforms = true;
 	bExportSourceTextures = true;
 	bExportLightmaps = true;
diff --git a/Engine/Plugins/Enterprise/GLTFExporter/MigratedFromGit/Source/GLTFExporter/Public/GLTFExportOptions.h b/Engine/Plugins/Enterprise/GLTFExporter/MigratedFromGit/Source/GLTFExporter/Public/GLTFExportOptions.h
--- a/Engine/Plugins/Enterprise/GLTFExporter/MigratedFromGit/Source/GLTFExporter/Public/GLTFExportOptions.h
+++ b/Engine/Plugins/Enterprise/GLTFExporter/MigratedFromGit/Source/GLTFExporter/Public/GLTFExportOptions.h
@@ -79,6 +79,9 @@ public:
 	UPROPERTY(EditAnywhere, BlueprintReadWrite, config, Category = Animation)
 	uint32 bExportAnimationSequences : 1;
 
+	UPROPERTY(EditAnywhere, BlueprintReadWrite, config, Category = Animation)
+	uint32 bExportLevelSequences : 1;
+
 	UPROPERTY(EditAnywhere, BlueprintReadWrite, config, Category = Animation)
 	uint32 bRetargetBoneTransforms : 1;
 
