[GLTFExporter] feat: add converter and task for level sequence and call it from node converter

commit dbb202d68d61345dce3b20bd562b791b3ba94a61
Date: Mon, 5 Oct 2020 13:02:43 +0200

    feat: add converter and task for level sequence and call it from node converter
    (cherry picked from commit 31754aea1d3c6cb1f4d7cdacbf1213c7cc7abf86​)

#skipci

diff --git a/Engine/Plugins/Enterprise/GLTFExporter/MigratedFromGit/Source/GLTFExporter/Private/Builders/GLTFConvertBuilder.cpp b/Engine/Plugins/Enterprise/GLTFExporter/MigratedFromGit/Source/GLTFExporter/Private/Builders/GLTFConvertBuilder.cpp
--- a/Engine/Plugins/Enterprise/GLTFExporter/MigratedFromGit/Source/GLTFExporter/Private/Builders/GLTFConvertBuilder.cpp
+++ b/Engine/Plugins/Enterprise/GLTFExporter/MigratedFromGit/Source/GLTFExporter/Private/Builders/GLTFConvertBuilder.cpp
@@ -246,6 +246,16 @@ FGLTFJsonAnimationIndex FGLTFConvertBuilder::GetOrAddAnimation(FGLTFJsonNodeInde
 	return AnimationDataConverter.GetOrAdd(RootNode, SkeletalMeshComponent);
 }
 
+FGLTFJsonAnimationIndex FGLTFConvertBuilder::GetOrAddAnimation(const ALevelSequenceActor* LevelSequenceActor)
+{
+	if (LevelSequenceActor == nullptr)
+	{
+		return FGLTFJsonAnimationIndex(INDEX_NONE);
+	}
+
+	return LevelSequenceConverter.GetOrAdd(LevelSequenceActor);
+}
+
 FGLTFJsonNodeIndex FGLTFConvertBuilder::GetOrAddNode(const AActor* Actor)
 {
 	if (Actor == nullptr)
diff --git a/Engine/Plugins/Enterprise/GLTFExporter/MigratedFromGit/Source/GLTFExporter/Private/Builders/GLTFConvertBuilder.h b/Engine/Plugins/Enterprise/GLTFExporter/MigratedFromGit/Source/GLTFExporter/Private/Builders/GLTFConvertBuilder.h
--- a/Engine/Plugins/Enterprise/GLTFExporter/MigratedFromGit/Source/GLTFExporter/Private/Builders/GLTFConvertBuilder.h
+++ b/Engine/Plugins/Enterprise/GLTFExporter/MigratedFromGit/Source/GLTFExporter/Private/Builders/GLTFConvertBuilder.h
@@ -54,6 +54,7 @@ public:
 	FGLTFJsonSkinIndex GetOrAddSkin(FGLTFJsonNodeIndex RootNode, const USkeletalMeshComponent* SkeletalMeshComponent);
 	FGLTFJsonAnimationIndex GetOrAddAnimation(FGLTFJsonNodeIndex RootNode, const USkeletalMesh* SkeletalMesh, const UAnimSequence* AnimSequence);
 	FGLTFJsonAnimationIndex GetOrAddAnimation(FGLTFJsonNodeIndex RootNode, const USkeletalMeshComponent* SkeletalMeshComponent);
+	FGLTFJsonAnimationIndex GetOrAddAnimation(const ALevelSequenceActor* LevelSequenceActor);
 
 	FGLTFJsonNodeIndex GetOrAddNode(const AActor* Actor);
 	FGLTFJsonNodeIndex GetOrAddNode(const USceneComponent* SceneComponent);
@@ -95,6 +96,7 @@ private:
 	FGLTFSkinConverter SkinConverter{ *this };
 	FGLTFAnimationConverter AnimationConverter{ *this };
 	FGLTFAnimationDataConverter AnimationDataConverter{ *this };
+	FGLTFLevelSequenceConverter LevelSequenceConverter{ *this };
 
 	FGLTFActorConverter ActorConverter{ *this };
 	FGLTFComponentConverter ComponentConverter{ *this };
diff --git a/Engine/Plugins/Enterprise/GLTFExporter/MigratedFromGit/Source/GLTFExporter/Private/Converters/GLTFAnimationConverters.cpp b/Engine/Plugins/Enterprise/GLTFExporter/MigratedFromGit/Source/GLTFExporter/Private/Converters/GLTFAnimationConverters.cpp
--- a/Engine/Plugins/Enterprise/GLTFExporter/MigratedFromGit/Source/GLTFExporter/Private/Converters/GLTFAnimationConverters.cpp
+++ b/Engine/Plugins/Enterprise/GLTFExporter/MigratedFromGit/Source/GLTFExporter/Private/Converters/GLTFAnimationConverters.cpp
@@ -1,7 +1,6 @@
 // Copyright Epic Games, Inc. All Rights Reserved.
 
 #include "Converters/GLTFAnimationConverters.h"
-#include "Converters/GLTFConverterUtility.h"
 #include "Builders/GLTFConvertBuilder.h"
 #include "Tasks/GLTFAnimationTasks.h"
 #include "Animation/AnimSequence.h"
@@ -28,7 +27,7 @@ FGLTFJsonAnimationIndex FGLTFAnimationConverter::Convert(FGLTFJsonNodeIndex Root
 	}
 
 	const FGLTFJsonAnimationIndex AnimationIndex = Builder.AddAnimation();
-	Builder.SetupTask<FGLTFAnimationTask>(Builder, RootNode, SkeletalMesh, AnimSequence, AnimationIndex);
+	Builder.SetupTask<FGLTFAnimSequenceTask>(Builder, RootNode, SkeletalMesh, AnimSequence, AnimationIndex);
 	return AnimationIndex;
 }
 
@@ -56,3 +55,10 @@ FGLTFJsonAnimationIndex FGLTFAnimationDataConverter::Convert(FGLTFJsonNodeIndex
 
 	return AnimationIndex;
 }
+
+FGLTFJsonAnimationIndex FGLTFLevelSequenceConverter::Convert(const ALevelSequenceActor* LevelSequenceActor)
+{
+	const FGLTFJsonAnimationIndex AnimationIndex = Builder.AddAnimation();
+	Builder.SetupTask<FGLTFLevelSequenceTask>(Builder, LevelSequenceActor, AnimationIndex);
+	return AnimationIndex;
+}
diff --git a/Engine/Plugins/Enterprise/GLTFExporter/MigratedFromGit/Source/GLTFExporter/Private/Converters/GLTFAnimationConverters.h b/Engine/Plugins/Enterprise/GLTFExporter/MigratedFromGit/Source/GLTFExporter/Private/Converters/GLTFAnimationConverters.h
--- a/Engine/Plugins/Enterprise/GLTFExporter/MigratedFromGit/Source/GLTFExporter/Private/Converters/GLTFAnimationConverters.h
+++ b/Engine/Plugins/Enterprise/GLTFExporter/MigratedFromGit/Source/GLTFExporter/Private/Converters/GLTFAnimationConverters.h
@@ -7,6 +7,8 @@
 #include "Converters/GLTFBuilderContext.h"
 #include "Engine.h"
 
+class ALevelSequenceActor;
+
 template <typename... InputTypes>
 class TGLTFAnimationConverter : public FGLTFBuilderContext, public TGLTFConverter<FGLTFJsonAnimationIndex, InputTypes...>
 {
@@ -26,3 +28,10 @@ class FGLTFAnimationDataConverter final : public TGLTFAnimationConverter<FGLTFJs
 
 	virtual FGLTFJsonAnimationIndex Convert(FGLTFJsonNodeIndex RootNode, const USkeletalMeshComponent* SkeletalMeshComponent) override;
 };
+
+class FGLTFLevelSequenceConverter final : public TGLTFAnimationConverter<const ALevelSequenceActor*>
+{
+	using TGLTFAnimationConverter::TGLTFAnimationConverter;
+
+	virtual FGLTFJsonAnimationIndex Convert(const ALevelSequenceActor* LevelSequenceActor) override;
+};
diff --git a/Engine/Plugins/Enterprise/GLTFExporter/MigratedFromGit/Source/GLTFExporter/Private/Converters/GLTFNodeConverters.cpp b/Engine/Plugins/Enterprise/GLTFExporter/MigratedFromGit/Source/GLTFExporter/Private/Converters/GLTFNodeConverters.cpp
--- a/Engine/Plugins/Enterprise/GLTFExporter/MigratedFromGit/Source/GLTFExporter/Private/Converters/GLTFNodeConverters.cpp
+++ b/Engine/Plugins/Enterprise/GLTFExporter/MigratedFromGit/Source/GLTFExporter/Private/Converters/GLTFNodeConverters.cpp
@@ -6,6 +6,7 @@
 #include "Converters/GLTFActorUtility.h"
 #include "Converters/GLTFNameUtility.h"
 #include "Actors/GLTFHotspotActor.h"
+#include "LevelSequenceActor.h"
 
 FGLTFJsonNodeIndex FGLTFActorConverter::Convert(const AActor* Actor)
 {
@@ -31,6 +32,13 @@ FGLTFJsonNodeIndex FGLTFActorConverter::Convert(const AActor* Actor)
 			RootNode.Backdrop = Builder.GetOrAddBackdrop(Actor);
 		}
 	}
+	else if (const ALevelSequenceActor* LevelSequenceActor = Cast<ALevelSequenceActor>(Actor))
+	{
+		if (Builder.ExportOptions->bExportLevelSequences)
+		{
+			Builder.GetOrAddAnimation(LevelSequenceActor);
+		}
+	}
 	else if (const AGLTFHotspotActor* HotspotActor = Cast<AGLTFHotspotActor>(Actor))
 	{
 		if (Builder.ExportOptions->bExportInteractionHotspots)
diff --git a/Engine/Plugins/Enterprise/GLTFExporter/MigratedFromGit/Source/GLTFExporter/Private/Tasks/GLTFAnimationTasks.cpp b/Engine/Plugins/Enterprise/GLTFExporter/MigratedFromGit/Source/GLTFExporter/Private/Tasks/GLTFAnimationTasks.cpp
--- a/Engine/Plugins/Enterprise/GLTFExporter/MigratedFromGit/Source/GLTFExporter/Private/Tasks/GLTFAnimationTasks.cpp
+++ b/Engine/Plugins/Enterprise/GLTFExporter/MigratedFromGit/Source/GLTFExporter/Private/Tasks/GLTFAnimationTasks.cpp
@@ -5,7 +5,7 @@
 #include "Converters/GLTFConverterUtility.h"
 #include "Converters/GLTFBoneUtility.h"
 
-void FGLTFAnimationTask::Complete()
+void FGLTFAnimSequenceTask::Complete()
 {
 	const int32 FrameCount = AnimSequence->GetRawNumberOfFrames();
 	const USkeleton* Skeleton = AnimSequence->GetSkeleton();
@@ -171,3 +171,7 @@ void FGLTFAnimationTask::Complete()
 		}
 	}
 }
+
+void FGLTFLevelSequenceTask::Complete()
+{
+}
diff --git a/Engine/Plugins/Enterprise/GLTFExporter/MigratedFromGit/Source/GLTFExporter/Private/Tasks/GLTFAnimationTasks.h b/Engine/Plugins/Enterprise/GLTFExporter/MigratedFromGit/Source/GLTFExporter/Private/Tasks/GLTFAnimationTasks.h
--- a/Engine/Plugins/Enterprise/GLTFExporter/MigratedFromGit/Source/GLTFExporter/Private/Tasks/GLTFAnimationTasks.h
+++ b/Engine/Plugins/Enterprise/GLTFExporter/MigratedFromGit/Source/GLTFExporter/Private/Tasks/GLTFAnimationTasks.h
@@ -4,13 +4,15 @@
 
 #include "Tasks/GLTFTask.h"
 #include "Builders/GLTFConvertBuilder.h"
-#include "Engine.h"
+#include "LevelSequenceActor.h"
 
-class FGLTFAnimationTask : public FGLTFTask
+class ALevelSequenceActor;
+
+class FGLTFAnimSequenceTask : public FGLTFTask
 {
 public:
 
-	FGLTFAnimationTask(FGLTFConvertBuilder& Builder, FGLTFJsonNodeIndex RootNode, const USkeletalMesh* SkeletalMesh, const UAnimSequence* AnimSequence, FGLTFJsonAnimationIndex AnimationIndex)
+	FGLTFAnimSequenceTask(FGLTFConvertBuilder& Builder, FGLTFJsonNodeIndex RootNode, const USkeletalMesh* SkeletalMesh, const UAnimSequence* AnimSequence, FGLTFJsonAnimationIndex AnimationIndex)
         : FGLTFTask(EGLTFTaskPriority::Animation)
 		, Builder(Builder)
 		, RootNode(RootNode)
@@ -35,3 +37,29 @@ private:
 	const UAnimSequence* AnimSequence;
 	const FGLTFJsonAnimationIndex AnimationIndex;
 };
+
+class FGLTFLevelSequenceTask : public FGLTFTask
+{
+public:
+
+	FGLTFLevelSequenceTask(FGLTFConvertBuilder& Builder, const ALevelSequenceActor* LevelSequenceActor, FGLTFJsonAnimationIndex AnimationIndex)
+		: FGLTFTask(EGLTFTaskPriority::Animation)
+		, Builder(Builder)
+		, LevelSequenceActor(LevelSequenceActor)
+		, AnimationIndex(AnimationIndex)
+	{
+	}
+
+	virtual FString GetName() override
+	{
+		return LevelSequenceActor->GetName();
+	}
+
+	virtual void Complete() override;
+
+private:
+
+	FGLTFConvertBuilder& Builder;
+	const ALevelSequenceActor* LevelSequenceActor;
+	const FGLTFJsonAnimationIndex AnimationIndex;
+};
