[GLTFExporter] feat: implement export of level sequences via animation task

commit 878ffc7170d1414dde29ee0689006c88e9310160
Date: Mon, 5 Oct 2020 13:27:30 +0200

    feat: implement export of level sequences via animation task
    (cherry picked from commit 1d0afe048bfc18e6b27dfa4cfe5609ea1f029013​)

#skipci

diff --git a/Engine/Plugins/Enterprise/GLTFExporter/MigratedFromGit/Source/GLTFExporter/Private/Tasks/GLTFAnimationTasks.cpp b/Engine/Plugins/Enterprise/GLTFExporter/MigratedFromGit/Source/GLTFExporter/Private/Tasks/GLTFAnimationTasks.cpp
--- a/Engine/Plugins/Enterprise/GLTFExporter/MigratedFromGit/Source/GLTFExporter/Private/Tasks/GLTFAnimationTasks.cpp
+++ b/Engine/Plugins/Enterprise/GLTFExporter/MigratedFromGit/Source/GLTFExporter/Private/Tasks/GLTFAnimationTasks.cpp
@@ -4,6 +4,12 @@
 #include "Builders/GLTFContainerBuilder.h"
 #include "Converters/GLTFConverterUtility.h"
 #include "Converters/GLTFBoneUtility.h"
+#include "LevelSequence.h"
+#include "LevelSequencePlayer.h"
+#include "MovieSceneTimeHelpers.h"
+#include "Tracks/MovieScene3DTransformTrack.h"
+#include "Sections/MovieScene3DTransformSection.h"
+#include "Channels/MovieSceneChannelProxy.h"
 
 void FGLTFAnimSequenceTask::Complete()
 {
@@ -174,4 +180,209 @@ void FGLTFAnimSequenceTask::Complete()
 
 void FGLTFLevelSequenceTask::Complete()
 {
+	ULevelSequence* LevelSequence = LevelSequenceActor->LoadSequence();
+	UMovieScene* MovieScene = LevelSequence->GetMovieScene();
+
+	ULevelSequencePlayer* LevelSequencePlayer = LevelSequenceActor->SequencePlayer;
+	LevelSequencePlayer->Initialize(LevelSequence, LevelSequenceActor->GetLevel(), LevelSequenceActor->PlaybackSettings, LevelSequenceActor->CameraSettings);
+	LevelSequencePlayer->State.AssignSequence(MovieSceneSequenceID::Root, *LevelSequence, *LevelSequencePlayer);
+
+	FFrameRate TickResolution = MovieScene->GetTickResolution();
+	FFrameRate DisplayRate = MovieScene->GetDisplayRate();
+	TRange<FFrameNumber> PlaybackRange = MovieScene->GetPlaybackRange();
+
+	int32 FrameOffset = FFrameRate::TransformTime(FFrameTime(MovieScene::DiscreteInclusiveLower(PlaybackRange)), TickResolution, DisplayRate).RoundToFrame().Value;
+	int32 FrameCount = FFrameRate::TransformTime(FFrameTime(FFrameNumber(MovieScene::DiscreteSize(PlaybackRange))), TickResolution, DisplayRate).RoundToFrame().Value + 1;
+
+	FGLTFJsonAnimation& JsonAnimation = Builder.GetAnimation(AnimationIndex);
+	LevelSequence->GetName(JsonAnimation.Name);
+
+	TArray<float> Timestamps;
+	TArray<FFrameTime> FrameTimes;
+	Timestamps.AddUninitialized(FrameCount);
+	FrameTimes.AddUninitialized(FrameCount);
+
+	for (int32 Frame = 0; Frame < FrameCount; ++Frame)
+	{
+		Timestamps[Frame] = DisplayRate.AsSeconds(Frame);
+		FrameTimes[Frame] = FFrameRate::TransformTime(FFrameTime(FrameOffset + Frame), DisplayRate, TickResolution);
+	}
+
+	FGLTFJsonAccessor JsonInputAccessor;
+	JsonInputAccessor.BufferView = Builder.AddBufferView(Timestamps);
+	JsonInputAccessor.ComponentType = EGLTFJsonComponentType::F32;
+	JsonInputAccessor.Type = EGLTFJsonAccessorType::Scalar;
+	JsonInputAccessor.Count = FrameCount;
+	JsonInputAccessor.MinMaxLength = 1;
+	JsonInputAccessor.Min[0] = 0;
+	JsonInputAccessor.Max[0] = DisplayRate.AsSeconds(FrameCount - 1);
+	FGLTFJsonAccessorIndex InputAccessorIndex = Builder.AddAccessor(JsonInputAccessor);
+
+	for (const FMovieSceneBinding& Binding : MovieScene->GetBindings())
+	{
+		for (TWeakObjectPtr<UObject> Object : LevelSequencePlayer->FindBoundObjects(Binding.GetObjectGuid(), MovieSceneSequenceID::Root))
+		{
+			FGLTFJsonNodeIndex NodeIndex;
+
+			if (AActor* Actor = Cast<AActor>(Object.Get()))
+			{
+				NodeIndex = Builder.GetOrAddNode(Actor);
+			}
+			else if (USceneComponent* SceneComponent = Cast<USceneComponent>(Object.Get()))
+			{
+				NodeIndex = Builder.GetOrAddNode(SceneComponent);
+			}
+			else
+			{
+				// TODO: report warning
+				continue;
+			}
+
+			if (NodeIndex == INDEX_NONE)
+			{
+				// TODO: report warning
+				continue;
+			}
+
+			for (UMovieSceneTrack* Track : Binding.GetTracks())
+			{
+				UMovieScene3DTransformTrack* TransformTrack = Cast<UMovieScene3DTransformTrack>(Track);
+				if (TransformTrack == nullptr)
+				{
+					// TODO: report warning
+					continue;
+				}
+
+				for (UMovieSceneSection* Section : TransformTrack->GetAllSections())
+				{
+					UMovieScene3DTransformSection* TransformSection = Cast<UMovieScene3DTransformSection>(Section);
+					if (TransformSection == nullptr)
+					{
+						// TODO: report warning
+						continue;
+					}
+
+					FOptionalMovieSceneBlendType BlendType = TransformSection->GetBlendType();
+					if (BlendType.IsValid() && BlendType.Get() != EMovieSceneBlendType::Absolute)
+					{
+						// TODO: report warning
+						continue;
+					}
+
+					// TODO: do we need to account for TransformSection->GetOffsetTime() ?
+
+					TArrayView<FMovieSceneFloatChannel*> Channels = TransformSection->GetChannelProxy().GetChannels<FMovieSceneFloatChannel>();
+					EMovieSceneTransformChannel ChannelMask = TransformSection->GetMask().GetChannels();
+
+					if (EnumHasAnyFlags(ChannelMask, EMovieSceneTransformChannel::Translation))
+					{
+						TArray<FGLTFJsonVector3> Translations;
+						Translations.AddUninitialized(FrameCount);
+
+						for (int32 Frame = 0; Frame < FrameCount; ++Frame)
+						{
+							FFrameTime& FrameTime = FrameTimes[Frame];
+
+							FVector Translation = FVector::ZeroVector;
+							Channels[0]->Evaluate(FrameTime, Translation.X);
+							Channels[1]->Evaluate(FrameTime, Translation.Y);
+							Channels[2]->Evaluate(FrameTime, Translation.Z);
+
+							Translations[Frame] = FGLTFConverterUtility::ConvertPosition(Translation, Builder.ExportOptions->ExportScale);
+						}
+
+						FGLTFJsonAccessor JsonOutputAccessor;
+						JsonOutputAccessor.BufferView = Builder.AddBufferView(Translations);
+						JsonOutputAccessor.ComponentType = EGLTFJsonComponentType::F32;
+						JsonOutputAccessor.Count = Translations.Num();
+						JsonOutputAccessor.Type = EGLTFJsonAccessorType::Vec3;
+
+						FGLTFJsonAnimationSampler JsonSampler;
+						JsonSampler.Input = InputAccessorIndex;
+						JsonSampler.Output = Builder.AddAccessor(JsonOutputAccessor);
+						JsonSampler.Interpolation = EGLTFJsonInterpolation::Linear;
+
+						FGLTFJsonAnimationChannel JsonChannel;
+						JsonChannel.Sampler = FGLTFJsonAnimationSamplerIndex(JsonAnimation.Samplers.Add(JsonSampler));
+						JsonChannel.Target.Path = EGLTFJsonTargetPath::Translation;
+						JsonChannel.Target.Node = NodeIndex;
+						JsonAnimation.Channels.Add(JsonChannel);
+					}
+
+					if (EnumHasAnyFlags(ChannelMask, EMovieSceneTransformChannel::Rotation))
+					{
+						TArray<FGLTFJsonQuaternion> Rotations;
+						Rotations.AddUninitialized(FrameCount);
+
+						for (int32 Frame = 0; Frame < FrameCount; ++Frame)
+						{
+							FFrameTime& FrameTime = FrameTimes[Frame];
+
+							FRotator Rotator = FRotator::ZeroRotator;
+							Channels[3]->Evaluate(FrameTime, Rotator.Roll);
+							Channels[4]->Evaluate(FrameTime, Rotator.Pitch);
+							Channels[5]->Evaluate(FrameTime, Rotator.Yaw);
+
+							Rotations[Frame] = FGLTFConverterUtility::ConvertRotation(Rotator.Quaternion());
+						}
+
+						JsonInputAccessor.Count = Rotations.Num();
+						JsonInputAccessor.Max[0] = Timestamps[Rotations.Num() - 1];
+
+						FGLTFJsonAccessor JsonOutputAccessor;
+						JsonOutputAccessor.BufferView = Builder.AddBufferView(Rotations);
+						JsonOutputAccessor.ComponentType = EGLTFJsonComponentType::F32;
+						JsonOutputAccessor.Count = Rotations.Num();
+						JsonOutputAccessor.Type = EGLTFJsonAccessorType::Vec4;
+
+						FGLTFJsonAnimationSampler JsonSampler;
+						JsonSampler.Input = InputAccessorIndex;
+						JsonSampler.Output = Builder.AddAccessor(JsonOutputAccessor);
+						JsonSampler.Interpolation = EGLTFJsonInterpolation::Linear;
+
+						FGLTFJsonAnimationChannel JsonChannel;
+						JsonChannel.Sampler = FGLTFJsonAnimationSamplerIndex(JsonAnimation.Samplers.Add(JsonSampler));
+						JsonChannel.Target.Path = EGLTFJsonTargetPath::Rotation;
+						JsonChannel.Target.Node = NodeIndex;
+						JsonAnimation.Channels.Add(JsonChannel);
+					}
+
+					if (EnumHasAnyFlags(ChannelMask, EMovieSceneTransformChannel::Scale))
+					{
+						TArray<FGLTFJsonVector3> Scales;
+						Scales.AddUninitialized(FrameCount);
+
+						for (int32 Frame = 0; Frame < FrameCount; ++Frame)
+						{
+							FFrameTime& FrameTime = FrameTimes[Frame];
+
+							FVector Scale = FVector::OneVector;
+							Channels[6]->Evaluate(FrameTime, Scale.X);
+							Channels[7]->Evaluate(FrameTime, Scale.Y);
+							Channels[8]->Evaluate(FrameTime, Scale.Z);
+
+							Scales[Frame] = FGLTFConverterUtility::ConvertScale(Scale);
+						}
+
+						FGLTFJsonAccessor JsonOutputAccessor;
+						JsonOutputAccessor.BufferView = Builder.AddBufferView(Scales);
+						JsonOutputAccessor.ComponentType = EGLTFJsonComponentType::F32;
+						JsonOutputAccessor.Count = Scales.Num();
+						JsonOutputAccessor.Type = EGLTFJsonAccessorType::Vec3;
+
+						FGLTFJsonAnimationSampler JsonSampler;
+						JsonSampler.Input = InputAccessorIndex;
+						JsonSampler.Output = Builder.AddAccessor(JsonOutputAccessor);
+						JsonSampler.Interpolation = EGLTFJsonInterpolation::Linear;
+
+						FGLTFJsonAnimationChannel JsonChannel;
+						JsonChannel.Sampler = FGLTFJsonAnimationSamplerIndex(JsonAnimation.Samplers.Add(JsonSampler));
+						JsonChannel.Target.Path = EGLTFJsonTargetPath::Scale;
+						JsonChannel.Target.Node = NodeIndex;
+						JsonAnimation.Channels.Add(JsonChannel);
+					}
+				}
+			}
+		}
+	}
 }
