[GLTFExporter] feat: add support for exporting playback settings in level sequence

commit 4ee72eaa9a4662a96e3c6808e669829c65c3f7ab
Date: Mon, 5 Oct 2020 16:57:03 +0200

    feat: add support for exporting playback settings in level sequence
    (cherry picked from commit 8365ceaa99a22ca3b419acf6cbad296342454dd6​)

#skipci

diff --git a/Engine/Plugins/Enterprise/GLTFExporter/MigratedFromGit/Source/GLTFExporter/Private/Builders/GLTFConvertBuilder.cpp b/Engine/Plugins/Enterprise/GLTFExporter/MigratedFromGit/Source/GLTFExporter/Private/Builders/GLTFConvertBuilder.cpp
--- a/Engine/Plugins/Enterprise/GLTFExporter/MigratedFromGit/Source/GLTFExporter/Private/Builders/GLTFConvertBuilder.cpp
+++ b/Engine/Plugins/Enterprise/GLTFExporter/MigratedFromGit/Source/GLTFExporter/Private/Builders/GLTFConvertBuilder.cpp
@@ -246,14 +246,14 @@ FGLTFJsonAnimationIndex FGLTFConvertBuilder::GetOrAddAnimation(FGLTFJsonNodeInde
 	return AnimationDataConverter.GetOrAdd(RootNode, SkeletalMeshComponent);
 }
 
-FGLTFJsonAnimationIndex FGLTFConvertBuilder::GetOrAddAnimation(const ALevelSequenceActor* LevelSequenceActor)
+FGLTFJsonAnimationIndex FGLTFConvertBuilder::GetOrAddAnimation(const ALevelSequenceActor* LevelSequenceActor, const ULevelSequence* LevelSequence)
 {
 	if (LevelSequenceActor == nullptr)
 	{
 		return FGLTFJsonAnimationIndex(INDEX_NONE);
 	}
 
-	return LevelSequenceConverter.GetOrAdd(LevelSequenceActor);
+	return LevelSequenceConverter.GetOrAdd(LevelSequenceActor, LevelSequence);
 }
 
 FGLTFJsonNodeIndex FGLTFConvertBuilder::GetOrAddNode(const AActor* Actor)
diff --git a/Engine/Plugins/Enterprise/GLTFExporter/MigratedFromGit/Source/GLTFExporter/Private/Builders/GLTFConvertBuilder.h b/Engine/Plugins/Enterprise/GLTFExporter/MigratedFromGit/Source/GLTFExporter/Private/Builders/GLTFConvertBuilder.h
--- a/Engine/Plugins/Enterprise/GLTFExporter/MigratedFromGit/Source/GLTFExporter/Private/Builders/GLTFConvertBuilder.h
+++ b/Engine/Plugins/Enterprise/GLTFExporter/MigratedFromGit/Source/GLTFExporter/Private/Builders/GLTFConvertBuilder.h
@@ -54,7 +54,7 @@ public:
 	FGLTFJsonSkinIndex GetOrAddSkin(FGLTFJsonNodeIndex RootNode, const USkeletalMeshComponent* SkeletalMeshComponent);
 	FGLTFJsonAnimationIndex GetOrAddAnimation(FGLTFJsonNodeIndex RootNode, const USkeletalMesh* SkeletalMesh, const UAnimSequence* AnimSequence);
 	FGLTFJsonAnimationIndex GetOrAddAnimation(FGLTFJsonNodeIndex RootNode, const USkeletalMeshComponent* SkeletalMeshComponent);
-	FGLTFJsonAnimationIndex GetOrAddAnimation(const ALevelSequenceActor* LevelSequenceActor);
+	FGLTFJsonAnimationIndex GetOrAddAnimation(const ALevelSequenceActor* LevelSequenceActor, const ULevelSequence* LevelSequence);
 
 	FGLTFJsonNodeIndex GetOrAddNode(const AActor* Actor);
 	FGLTFJsonNodeIndex GetOrAddNode(const USceneComponent* SceneComponent);
diff --git a/Engine/Plugins/Enterprise/GLTFExporter/MigratedFromGit/Source/GLTFExporter/Private/Converters/GLTFAnimationConverters.cpp b/Engine/Plugins/Enterprise/GLTFExporter/MigratedFromGit/Source/GLTFExporter/Private/Converters/GLTFAnimationConverters.cpp
--- a/Engine/Plugins/Enterprise/GLTFExporter/MigratedFromGit/Source/GLTFExporter/Private/Converters/GLTFAnimationConverters.cpp
+++ b/Engine/Plugins/Enterprise/GLTFExporter/MigratedFromGit/Source/GLTFExporter/Private/Converters/GLTFAnimationConverters.cpp
@@ -56,9 +56,26 @@ FGLTFJsonAnimationIndex FGLTFAnimationDataConverter::Convert(FGLTFJsonNodeIndex
 	return AnimationIndex;
 }
 
-FGLTFJsonAnimationIndex FGLTFLevelSequenceConverter::Convert(const ALevelSequenceActor* LevelSequenceActor)
+FGLTFJsonAnimationIndex FGLTFLevelSequenceConverter::Convert(const ALevelSequenceActor* LevelSequenceActor, const ULevelSequence* LevelSequence)
 {
-	const FGLTFJsonAnimationIndex AnimationIndex = Builder.AddAnimation();
-	Builder.SetupTask<FGLTFLevelSequenceTask>(Builder, LevelSequenceActor, AnimationIndex);
+	FGLTFJsonAnimation JsonAnimation;
+
+	if (LevelSequence == nullptr)
+	{
+		LevelSequence = LevelSequenceActor->LoadSequence();
+		if (LevelSequence == nullptr)
+		{
+			return FGLTFJsonAnimationIndex(INDEX_NONE);
+		}
+
+		FGLTFJsonPlayData& JsonPlayData = JsonAnimation.PlayData;
+		JsonPlayData.Looping = LevelSequenceActor->PlaybackSettings.LoopCount.Value != 0;
+		JsonPlayData.Playing = LevelSequenceActor->PlaybackSettings.bAutoPlay;
+		JsonPlayData.PlayRate = LevelSequenceActor->PlaybackSettings.PlayRate;
+		JsonPlayData.Position = LevelSequenceActor->PlaybackSettings.StartTime;
+	}
+
+	const FGLTFJsonAnimationIndex AnimationIndex = Builder.AddAnimation(JsonAnimation);
+	Builder.SetupTask<FGLTFLevelSequenceTask>(Builder, LevelSequenceActor, LevelSequence, AnimationIndex);
 	return AnimationIndex;
 }
diff --git a/Engine/Plugins/Enterprise/GLTFExporter/MigratedFromGit/Source/GLTFExporter/Private/Converters/GLTFAnimationConverters.h b/Engine/Plugins/Enterprise/GLTFExporter/MigratedFromGit/Source/GLTFExporter/Private/Converters/GLTFAnimationConverters.h
--- a/Engine/Plugins/Enterprise/GLTFExporter/MigratedFromGit/Source/GLTFExporter/Private/Converters/GLTFAnimationConverters.h
+++ b/Engine/Plugins/Enterprise/GLTFExporter/MigratedFromGit/Source/GLTFExporter/Private/Converters/GLTFAnimationConverters.h
@@ -7,6 +7,7 @@
 #include "Converters/GLTFBuilderContext.h"
 #include "Engine.h"
 
+class ULevelSequence;
 class ALevelSequenceActor;
 
 template <typename... InputTypes>
@@ -29,9 +30,9 @@ class FGLTFAnimationDataConverter final : public TGLTFAnimationConverter<FGLTFJs
 	virtual FGLTFJsonAnimationIndex Convert(FGLTFJsonNodeIndex RootNode, const USkeletalMeshComponent* SkeletalMeshComponent) override;
 };
 
-class FGLTFLevelSequenceConverter final : public TGLTFAnimationConverter<const ALevelSequenceActor*>
+class FGLTFLevelSequenceConverter final : public TGLTFAnimationConverter<const ALevelSequenceActor*, const ULevelSequence*>
 {
 	using TGLTFAnimationConverter::TGLTFAnimationConverter;
 
-	virtual FGLTFJsonAnimationIndex Convert(const ALevelSequenceActor* LevelSequenceActor) override;
+	virtual FGLTFJsonAnimationIndex Convert(const ALevelSequenceActor* LevelSequenceActor, const ULevelSequence*) override;
 };
diff --git a/Engine/Plugins/Enterprise/GLTFExporter/MigratedFromGit/Source/GLTFExporter/Private/Converters/GLTFNodeConverters.cpp b/Engine/Plugins/Enterprise/GLTFExporter/MigratedFromGit/Source/GLTFExporter/Private/Converters/GLTFNodeConverters.cpp
--- a/Engine/Plugins/Enterprise/GLTFExporter/MigratedFromGit/Source/GLTFExporter/Private/Converters/GLTFNodeConverters.cpp
+++ b/Engine/Plugins/Enterprise/GLTFExporter/MigratedFromGit/Source/GLTFExporter/Private/Converters/GLTFNodeConverters.cpp
@@ -36,7 +36,7 @@ FGLTFJsonNodeIndex FGLTFActorConverter::Convert(const AActor* Actor)
 	{
 		if (Builder.ExportOptions->bExportLevelSequences)
 		{
-			Builder.GetOrAddAnimation(LevelSequenceActor);
+			Builder.GetOrAddAnimation(LevelSequenceActor, nullptr);
 		}
 	}
 	else if (const AGLTFHotspotActor* HotspotActor = Cast<AGLTFHotspotActor>(Actor))
diff --git a/Engine/Plugins/Enterprise/GLTFExporter/MigratedFromGit/Source/GLTFExporter/Private/Tasks/GLTFAnimationTasks.cpp b/Engine/Plugins/Enterprise/GLTFExporter/MigratedFromGit/Source/GLTFExporter/Private/Tasks/GLTFAnimationTasks.cpp
--- a/Engine/Plugins/Enterprise/GLTFExporter/MigratedFromGit/Source/GLTFExporter/Private/Tasks/GLTFAnimationTasks.cpp
+++ b/Engine/Plugins/Enterprise/GLTFExporter/MigratedFromGit/Source/GLTFExporter/Private/Tasks/GLTFAnimationTasks.cpp
@@ -180,12 +180,12 @@ void FGLTFAnimSequenceTask::Complete()
 
 void FGLTFLevelSequenceTask::Complete()
 {
-	ULevelSequence* LevelSequence = LevelSequenceActor->LoadSequence();
-	UMovieScene* MovieScene = LevelSequence->GetMovieScene();
+	ULevelSequence* Sequence = const_cast<ULevelSequence*>(LevelSequence);
+	UMovieScene* MovieScene = Sequence->GetMovieScene();
 
 	ULevelSequencePlayer* LevelSequencePlayer = LevelSequenceActor->SequencePlayer;
-	LevelSequencePlayer->Initialize(LevelSequence, LevelSequenceActor->GetLevel(), LevelSequenceActor->PlaybackSettings, LevelSequenceActor->CameraSettings);
-	LevelSequencePlayer->State.AssignSequence(MovieSceneSequenceID::Root, *LevelSequence, *LevelSequencePlayer);
+	LevelSequencePlayer->Initialize(Sequence, LevelSequenceActor->GetLevel(), LevelSequenceActor->PlaybackSettings, LevelSequenceActor->CameraSettings);
+	LevelSequencePlayer->State.AssignSequence(MovieSceneSequenceID::Root, *Sequence, *LevelSequencePlayer);
 
 	FFrameRate TickResolution = MovieScene->GetTickResolution();
 	FFrameRate DisplayRate = MovieScene->GetDisplayRate();
@@ -195,7 +195,7 @@ void FGLTFLevelSequenceTask::Complete()
 	int32 FrameCount = FFrameRate::TransformTime(FFrameTime(FFrameNumber(MovieScene::DiscreteSize(PlaybackRange))), TickResolution, DisplayRate).RoundToFrame().Value + 1;
 
 	FGLTFJsonAnimation& JsonAnimation = Builder.GetAnimation(AnimationIndex);
-	LevelSequence->GetName(JsonAnimation.Name);
+	Sequence->GetName(JsonAnimation.Name);
 
 	TArray<float> Timestamps;
 	TArray<FFrameTime> FrameTimes;
diff --git a/Engine/Plugins/Enterprise/GLTFExporter/MigratedFromGit/Source/GLTFExporter/Private/Tasks/GLTFAnimationTasks.h b/Engine/Plugins/Enterprise/GLTFExporter/MigratedFromGit/Source/GLTFExporter/Private/Tasks/GLTFAnimationTasks.h
--- a/Engine/Plugins/Enterprise/GLTFExporter/MigratedFromGit/Source/GLTFExporter/Private/Tasks/GLTFAnimationTasks.h
+++ b/Engine/Plugins/Enterprise/GLTFExporter/MigratedFromGit/Source/GLTFExporter/Private/Tasks/GLTFAnimationTasks.h
@@ -5,8 +5,7 @@
 #include "Tasks/GLTFTask.h"
 #include "Builders/GLTFConvertBuilder.h"
 #include "LevelSequenceActor.h"
-
-class ALevelSequenceActor;
+#include "LevelSequence.h"
 
 class FGLTFAnimSequenceTask : public FGLTFTask
 {
@@ -42,10 +41,11 @@ class FGLTFLevelSequenceTask : public FGLTFTask
 {
 public:
 
-	FGLTFLevelSequenceTask(FGLTFConvertBuilder& Builder, const ALevelSequenceActor* LevelSequenceActor, FGLTFJsonAnimationIndex AnimationIndex)
+	FGLTFLevelSequenceTask(FGLTFConvertBuilder& Builder, const ALevelSequenceActor* LevelSequenceActor, const ULevelSequence* LevelSequence, FGLTFJsonAnimationIndex AnimationIndex)
 		: FGLTFTask(EGLTFTaskPriority::Animation)
 		, Builder(Builder)
 		, LevelSequenceActor(LevelSequenceActor)
+		, LevelSequence(LevelSequence)
 		, AnimationIndex(AnimationIndex)
 	{
 	}
@@ -61,5 +61,6 @@ private:
 
 	FGLTFConvertBuilder& Builder;
 	const ALevelSequenceActor* LevelSequenceActor;
+	const ULevelSequence* LevelSequence;
 	const FGLTFJsonAnimationIndex AnimationIndex;
 };
