[GLTFExporter] fix: use default playback settings in level sequence conversion

commit e374b97b15d2831787272fc69abfa79be53db561
Date: Mon, 5 Oct 2020 18:55:02 +0200

    fix: use default playback settings in level sequence conversion
    (cherry picked from commit 3311904340936a14841c672844de593411935065​)

#skipci

diff --git a/Engine/Plugins/Enterprise/GLTFExporter/MigratedFromGit/Source/GLTFExporter/Private/Tasks/GLTFAnimationTasks.cpp b/Engine/Plugins/Enterprise/GLTFExporter/MigratedFromGit/Source/GLTFExporter/Private/Tasks/GLTFAnimationTasks.cpp
--- a/Engine/Plugins/Enterprise/GLTFExporter/MigratedFromGit/Source/GLTFExporter/Private/Tasks/GLTFAnimationTasks.cpp
+++ b/Engine/Plugins/Enterprise/GLTFExporter/MigratedFromGit/Source/GLTFExporter/Private/Tasks/GLTFAnimationTasks.cpp
@@ -184,7 +184,7 @@ void FGLTFLevelSequenceTask::Complete()
 	UMovieScene* MovieScene = Sequence->GetMovieScene();
 
 	ULevelSequencePlayer* LevelSequencePlayer = LevelSequenceActor->SequencePlayer;
-	LevelSequencePlayer->Initialize(Sequence, LevelSequenceActor->GetLevel(), LevelSequenceActor->PlaybackSettings, LevelSequenceActor->CameraSettings);
+	LevelSequencePlayer->Initialize(Sequence, LevelSequenceActor->GetLevel(), FMovieSceneSequencePlaybackSettings(), FLevelSequenceCameraSettings());
 	LevelSequencePlayer->State.AssignSequence(MovieSceneSequenceID::Root, *Sequence, *LevelSequencePlayer);
 
 	FFrameRate TickResolution = MovieScene->GetTickResolution();
