[GLTFExporter] refactor: rearrange GLTFExportOptions::bExportLevelSequences

commit 82c889fcdb0cf7ca9bc946d5e0f9b678a553a444
Date: Mon, 5 Oct 2020 18:55:30 +0200

    refactor: rearrange GLTFExportOptions::bExportLevelSequences
    (cherry picked from commit d9604850372c57c752668dfd5df47a6387dd14d6​)

#skipci

diff --git a/Engine/Plugins/Enterprise/GLTFExporter/MigratedFromGit/Source/GLTFExporter/Public/GLTFExportOptions.h b/Engine/Plugins/Enterprise/GLTFExporter/MigratedFromGit/Source/GLTFExporter/Public/GLTFExportOptions.h
--- a/Engine/Plugins/Enterprise/GLTFExporter/MigratedFromGit/Source/GLTFExporter/Public/GLTFExportOptions.h
+++ b/Engine/Plugins/Enterprise/GLTFExporter/MigratedFromGit/Source/GLTFExporter/Public/GLTFExportOptions.h
@@ -80,10 +80,10 @@ public:
 	uint32 bExportAnimationSequences : 1;
 
 	UPROPERTY(EditAnywhere, BlueprintReadWrite, config, Category = Animation)
-	uint32 bExportLevelSequences : 1;
+	uint32 bRetargetBoneTransforms : 1;
 
 	UPROPERTY(EditAnywhere, BlueprintReadWrite, config, Category = Animation)
-	uint32 bRetargetBoneTransforms : 1;
+	uint32 bExportLevelSequences : 1;
 
 	UPROPERTY(EditAnywhere, BlueprintReadWrite, config, Category = Texture)
 	uint32 bExportSourceTextures : 1;
