[GLTFExporter] feat: add support for level sequences in hotspot actor

commit 992cc7a86369fd81cb7a7a6c179173373eaf2ba1
Date: Mon, 5 Oct 2020 18:56:03 +0200

    feat: add support for level sequences in hotspot actor
    (cherry picked from commit 734b47b7eef2f04d9b089bc48c7bfa0470b246af​)

#skipci

diff --git a/Engine/Plugins/Enterprise/GLTFExporter/MigratedFromGit/Source/GLTFExporterRuntime/GLTFExporterRuntime.Build.cs b/Engine/Plugins/Enterprise/GLTFExporter/MigratedFromGit/Source/GLTFExporterRuntime/GLTFExporterRuntime.Build.cs
--- a/Engine/Plugins/Enterprise/GLTFExporter/MigratedFromGit/Source/GLTFExporterRuntime/GLTFExporterRuntime.Build.cs
+++ b/Engine/Plugins/Enterprise/GLTFExporter/MigratedFromGit/Source/GLTFExporterRuntime/GLTFExporterRuntime.Build.cs
@@ -31,7 +31,9 @@ namespace UnrealBuildTool.Rules
 					"Core",
 					"CoreUObject",
 					"Engine",
-					"InputCore"
+					"InputCore",
+					"LevelSequence",
+					"MovieScene"
 				}
 				);
 
diff --git a/Engine/Plugins/Enterprise/GLTFExporter/MigratedFromGit/Source/GLTFExporterRuntime/Private/Actors/GLTFHotspotActor.cpp b/Engine/Plugins/Enterprise/GLTFExporter/MigratedFromGit/Source/GLTFExporterRuntime/Private/Actors/GLTFHotspotActor.cpp
--- a/Engine/Plugins/Enterprise/GLTFExporter/MigratedFromGit/Source/GLTFExporterRuntime/Private/Actors/GLTFHotspotActor.cpp
+++ b/Engine/Plugins/Enterprise/GLTFExporter/MigratedFromGit/Source/GLTFExporterRuntime/Private/Actors/GLTFHotspotActor.cpp
@@ -3,6 +3,9 @@
 #include "Actors/GLTFHotspotActor.h"
 #include "Animation/SkeletalMeshActor.h"
 #include "Animation/AnimSequence.h"
+#include "LevelSequenceActor.h"
+#include "LevelSequence.h"
+#include "MovieSceneTimeHelpers.h"
 #include "Components/SphereComponent.h"
 #include "Components/MaterialBillboardComponent.h"
 #include "Kismet/GameplayStatics.h"
@@ -136,6 +139,14 @@ void AGLTFHotspotActor::PostEditChangeProperty(FPropertyChangedEvent& PropertyCh
 		{
 			ValidateAnimation();
 		}
+		else if (PropertyFName == GET_MEMBER_NAME_CHECKED(ThisClass, LevelSequenceActor))
+		{
+			if (LevelSequenceActor != nullptr && LevelSequence == nullptr)
+			{
+				LevelSequence = LevelSequenceActor->LoadSequence();
+			}
+			ValidateAnimation();
+		}
 	}
 }
 #endif // WITH_EDITOR
@@ -246,30 +257,7 @@ void AGLTFHotspotActor::Clicked(UPrimitiveComponent* TouchedComponent, FKey Butt
 		return;
 	}
 
-	const bool bReverseAnimation = bToggled;
-
-	if (SkeletalMeshActor != nullptr && AnimationSequence != nullptr)
-	{
-		USkeletalMeshComponent* SkeletalMeshComponent = SkeletalMeshActor->GetSkeletalMeshComponent();
-		SkeletalMeshComponent->SetAnimationMode(EAnimationMode::Type::AnimationSingleNode);
-		const float AbsolutePlayRate = FMath::Abs(SkeletalMeshComponent->GetPlayRate());
-		const UAnimSingleNodeInstance* SingleNodeInstance = SkeletalMeshComponent->GetSingleNodeInstance();
-
-		if (SkeletalMeshComponent->IsPlaying() && SingleNodeInstance != nullptr && SingleNodeInstance->GetAnimationAsset() == AnimationSequence)
-		{
-			// If the same animation is already playing, just reverse the play rate for a smooth transition
-			SkeletalMeshComponent->SetPlayRate(AbsolutePlayRate * -1.0f);
-		}
-		else
-		{
-			SkeletalMeshComponent->SetAnimation(AnimationSequence);
-			SkeletalMeshComponent->SetPlayRate(AbsolutePlayRate * (bReverseAnimation ? -1.0f : 1.0f));
-			SkeletalMeshComponent->SetPosition(bReverseAnimation ? AnimationSequence->GetPlayLength() : 0.0f);
-			SkeletalMeshComponent->Play(false);
-		}
-	}
-
-	bToggled = !bToggled;
+	ToggleAnimation();
 
 	UpdateActiveImageFromState(bToggled ? EGLTFHotspotState::ToggledHovered :  EGLTFHotspotState::Hovered);
 }
@@ -400,6 +388,57 @@ void AGLTFHotspotActor::ViewportResized(FViewport*, uint32)
 	UpdateSpriteSize();
 }
 
+void AGLTFHotspotActor::ToggleAnimation()
+{
+	const bool bReverseAnimation = bToggled;
+	const float TargetPlayRate = bToggled ? -1.0f : 1.0f;
+
+	if (SkeletalMeshActor != nullptr && AnimationSequence != nullptr)
+	{
+		USkeletalMeshComponent* SkeletalMeshComponent = SkeletalMeshActor->GetSkeletalMeshComponent();
+		SkeletalMeshComponent->SetAnimationMode(EAnimationMode::Type::AnimationSingleNode);
+		const UAnimSingleNodeInstance* SingleNodeInstance = SkeletalMeshComponent->GetSingleNodeInstance();
+
+		if (SkeletalMeshComponent->IsPlaying() && SingleNodeInstance != nullptr && SingleNodeInstance->GetAnimationAsset() == AnimationSequence)
+		{
+			// If the same animation is already playing, just reverse the play rate for a smooth transition
+			SkeletalMeshComponent->SetPlayRate(TargetPlayRate);
+		}
+		else
+		{
+			SkeletalMeshComponent->SetAnimation(AnimationSequence);
+			SkeletalMeshComponent->SetPlayRate(TargetPlayRate);
+			SkeletalMeshComponent->SetPosition(bReverseAnimation ? AnimationSequence->GetPlayLength() : 0.0f);
+			SkeletalMeshComponent->Play(false);
+		}
+	}
+	else if (LevelSequenceActor != nullptr && LevelSequence != nullptr)
+	{
+		ULevelSequencePlayer* LevelSequencePlayer = LevelSequenceActor->GetSequencePlayer();
+
+		if (LevelSequencePlayer != nullptr && LevelSequencePlayer->IsPlaying() && LevelSequenceActor->GetSequence() == LevelSequence)
+		{
+			// If the same animation is already playing, just reverse the play rate for a smooth transition
+			LevelSequencePlayer->SetPlayRate(TargetPlayRate);
+		}
+		else
+		{
+			// TODO: improve this calculation
+			const float SequenceLength = LevelSequence->GetMovieScene()->GetTickResolution().AsSeconds(MovieScene::DiscreteSize(LevelSequence->GetMovieScene()->GetPlaybackRange()));
+
+			LevelSequenceActor->PlaybackSettings.PlayRate = TargetPlayRate;
+			LevelSequenceActor->PlaybackSettings.StartTime = bReverseAnimation ? SequenceLength : 0.0f;
+			LevelSequenceActor->PlaybackSettings.bAutoPlay = true;
+			LevelSequenceActor->PlaybackSettings.bPauseAtEnd = true;
+			LevelSequenceActor->PlaybackSettings.LoopCount.Value = 0;
+			LevelSequenceActor->SetSequence(LevelSequence);
+			LevelSequenceActor->GetSequencePlayer()->Play();
+		}
+	}
+
+	bToggled = !bToggled;
+}
+
 void AGLTFHotspotActor::ValidateAnimation()
 {
 	if (SkeletalMeshActor != nullptr && AnimationSequence != nullptr)
diff --git a/Engine/Plugins/Enterprise/GLTFExporter/MigratedFromGit/Source/GLTFExporterRuntime/Public/Actors/GLTFHotspotActor.h b/Engine/Plugins/Enterprise/GLTFExporter/MigratedFromGit/Source/GLTFExporterRuntime/Public/Actors/GLTFHotspotActor.h
--- a/Engine/Plugins/Enterprise/GLTFExporter/MigratedFromGit/Source/GLTFExporterRuntime/Public/Actors/GLTFHotspotActor.h
+++ b/Engine/Plugins/Enterprise/GLTFExporter/MigratedFromGit/Source/GLTFExporterRuntime/Public/Actors/GLTFHotspotActor.h
@@ -7,6 +7,8 @@
 
 class ASkeletalMeshActor;
 class UAnimSequence;
+class ALevelSequenceActor;
+class ULevelSequence;
 class UTexture2D;
 class UTexture;
 class UMaterialInterface;
@@ -24,7 +26,7 @@ enum class EGLTFHotspotState : uint8
 };
 
 /**
- * Actor wrapper for the GLTF hotspot component. Appears as a billboard and allows playback of skeletal animations when cursor input is enabled.
+ * Actor wrapper for the GLTF hotspot component. Appears as a billboard and allows playback of animations when cursor input is enabled.
  */
 UCLASS(BlueprintType, Blueprintable, HideCategories = (Sprite, Physics, Collision, Navigation), DisplayName = "GLTF Hotspot")
 class GLTFEXPORTERRUNTIME_API AGLTFHotspotActor : public AActor
@@ -58,27 +60,35 @@ private:
 public:
 
 	/* The skeletal mesh actor that will be animated when the hotspot is clicked. */
-	UPROPERTY(EditAnywhere, BlueprintReadWrite, Category = "GLTF Interaction Hotspot")
+	UPROPERTY(EditAnywhere, BlueprintReadWrite, Category = "Hotspot Animation")
 	ASkeletalMeshActor* SkeletalMeshActor;
 
 	/* The animation that will be played on the skeletal mesh actor. Must be compatible with its skeletal mesh asset. */
-	UPROPERTY(EditAnywhere, BlueprintReadWrite, Category = "GLTF Interaction Hotspot", meta=(EditCondition="SkeletalMeshActor != nullptr"))
+	UPROPERTY(EditAnywhere, BlueprintReadWrite, Category = "Hotspot Animation", meta=(EditCondition="SkeletalMeshActor != nullptr"))
 	UAnimSequence* AnimationSequence;
 
+	/* The skeletal mesh actor that will be animated when the hotspot is clicked. */
+	UPROPERTY(EditAnywhere, BlueprintReadWrite, Category = "Hotspot Animation", meta=(EditCondition="SkeletalMeshActor == nullptr"))
+	ALevelSequenceActor* LevelSequenceActor;
+
+	/* The animation that will be played on the level sequence actor. */
+	UPROPERTY(EditAnywhere, BlueprintReadWrite, Category = "Hotspot Animation", meta=(EditCondition="SkeletalMeshActor == nullptr && LevelSequenceActor != nullptr"))
+	ULevelSequence* LevelSequence;
+
 	/* The billboard image that will be shown when the hotspot is in an inactive state or one without a specified image. */
-	UPROPERTY(EditAnywhere, BlueprintReadWrite, Category = "GLTF Interaction Hotspot")
+	UPROPERTY(EditAnywhere, BlueprintReadWrite, Category = "Hotspot Appearance")
 	UTexture2D* Image;
 
 	/** The optional billboard image that will be shown when a cursor enters the hotspot. */
-	UPROPERTY(EditAnywhere, BlueprintReadWrite, Category = "GLTF Interaction Hotspot", meta=(EditCondition="Image != nullptr"))
+	UPROPERTY(EditAnywhere, BlueprintReadWrite, Category = "Hotspot Appearance", meta=(EditCondition="Image != nullptr"))
 	UTexture2D* HoveredImage;
 
 	/** The optional billboard image that will be shown when the hotspot is toggled by a click. */
-	UPROPERTY(EditAnywhere, BlueprintReadWrite, Category = "GLTF Interaction Hotspot", meta=(EditCondition="Image != nullptr"))
+	UPROPERTY(EditAnywhere, BlueprintReadWrite, Category = "Hotspot Appearance", meta=(EditCondition="Image != nullptr"))
 	UTexture2D* ToggledImage;
 
 	/** The optional billboard image that will be shown when the hotspot is toggled by a click and a cursor enters it. */
-	UPROPERTY(EditAnywhere, BlueprintReadWrite, Category = "GLTF Interaction Hotspot", meta=(EditCondition="Image != nullptr"))
+	UPROPERTY(EditAnywhere, BlueprintReadWrite, Category = "Hotspot Appearance", meta=(EditCondition="Image != nullptr"))
 	UTexture2D* ToggledHoveredImage;
 
 private:
@@ -132,6 +142,8 @@ private:
 
 	void ViewportResized(FViewport*, uint32);
 
+	void ToggleAnimation();
+
 	void ValidateAnimation();
 
 public:
