[GLTFExporter] feat: add support for exporting hotspots with level sequences

commit b64eb28d018190dff85fff57deeb1c461ec1f34d
Date: Mon, 5 Oct 2020 18:56:25 +0200

    feat: add support for exporting hotspots with level sequences
    (cherry picked from commit b8d70af0968db6a79f12325763933addd9098f4c​)

#skipci

diff --git a/Engine/Plugins/Enterprise/GLTFExporter/MigratedFromGit/Source/GLTFExporter/Private/Converters/GLTFHotspotConverters.cpp b/Engine/Plugins/Enterprise/GLTFExporter/MigratedFromGit/Source/GLTFExporter/Private/Converters/GLTFHotspotConverters.cpp
--- a/Engine/Plugins/Enterprise/GLTFExporter/MigratedFromGit/Source/GLTFExporter/Private/Converters/GLTFHotspotConverters.cpp
+++ b/Engine/Plugins/Enterprise/GLTFExporter/MigratedFromGit/Source/GLTFExporter/Private/Converters/GLTFHotspotConverters.cpp
@@ -9,21 +9,21 @@ FGLTFJsonHotspotIndex FGLTFHotspotConverter::Convert(const AGLTFHotspotActor* Ho
 	FGLTFJsonHotspot JsonHotspot;
 	HotspotActor->GetName(JsonHotspot.Name);
 
-	if (!Builder.ExportOptions->bExportVertexSkinWeights)
+	if (const ASkeletalMeshActor* SkeletalMeshActor = HotspotActor->SkeletalMeshActor)
 	{
-		Builder.AddWarningMessage(
-			FString::Printf(TEXT("Can't export animation in hotspot %s because vertex skin weights are disabled by export options"),
-			*JsonHotspot.Name));
-	}
-	else if (!Builder.ExportOptions->bExportAnimationSequences)
-	{
-		Builder.AddWarningMessage(
-			FString::Printf(TEXT("Can't export animation in hotspot %s because animation sequences are disabled by export options"),
-			*JsonHotspot.Name));
-	}
-	else
-	{
-		if (const ASkeletalMeshActor* SkeletalMeshActor = HotspotActor->SkeletalMeshActor)
+		if (!Builder.ExportOptions->bExportVertexSkinWeights)
+		{
+			Builder.AddWarningMessage(
+				FString::Printf(TEXT("Can't export animation in hotspot %s because vertex skin weights are disabled by export options"),
+				*JsonHotspot.Name));
+		}
+		else if (!Builder.ExportOptions->bExportAnimationSequences)
+		{
+			Builder.AddWarningMessage(
+				FString::Printf(TEXT("Can't export animation in hotspot %s because animation sequences are disabled by export options"),
+				*JsonHotspot.Name));
+		}
+		else
 		{
 			const FGLTFJsonNodeIndex RootNode = Builder.GetOrAddNode(SkeletalMeshActor);
 
@@ -44,11 +44,31 @@ FGLTFJsonHotspotIndex FGLTFHotspotConverter::Convert(const AGLTFHotspotActor* Ho
 				// TODO: report warning
 			}
 		}
+	}
+	else if (const ALevelSequenceActor* LevelSequenceActor = HotspotActor->LevelSequenceActor)
+	{
+		if (!Builder.ExportOptions->bExportLevelSequences)
+		{
+			Builder.AddWarningMessage(
+                FString::Printf(TEXT("Can't export animation in hotspot %s because level sequences are disabled by export options"),
+                *JsonHotspot.Name));
+		}
 		else
 		{
-			// TODO: report warning
+			if (const ULevelSequence* LevelSequence = HotspotActor->LevelSequence)
+			{
+				JsonHotspot.Animation = Builder.GetOrAddAnimation(LevelSequenceActor, LevelSequence);
+			}
+			else
+			{
+				// TODO: report warning
+			}
 		}
 	}
+	else
+	{
+		// TODO: report warning
+	}
 
 	JsonHotspot.Image = Builder.GetOrAddTexture(HotspotActor->GetImageForState(EGLTFHotspotState::Default));
 	JsonHotspot.HoveredImage = Builder.GetOrAddTexture(HotspotActor->GetImageForState(EGLTFHotspotState::Hovered));
