[GLTFExporter] feat: remove the need for ALevelSequenceActor to play level sequences in hotspot actor

commit 46828d4bcb7d821f299534d501bf83b570d28e91
Date: Mon, 5 Oct 2020 23:22:47 +0200

    feat: remove the need for ALevelSequenceActor to play level sequences in hotspot actor
    (cherry picked from commit 751eaef0015605a7e260667b1f2d2b0264bb323a​)

#skipci

diff --git a/Engine/Plugins/Enterprise/GLTFExporter/MigratedFromGit/Source/GLTFExporter/Private/Builders/GLTFConvertBuilder.cpp b/Engine/Plugins/Enterprise/GLTFExporter/MigratedFromGit/Source/GLTFExporter/Private/Builders/GLTFConvertBuilder.cpp
--- a/Engine/Plugins/Enterprise/GLTFExporter/MigratedFromGit/Source/GLTFExporter/Private/Builders/GLTFConvertBuilder.cpp
+++ b/Engine/Plugins/Enterprise/GLTFExporter/MigratedFromGit/Source/GLTFExporter/Private/Builders/GLTFConvertBuilder.cpp
@@ -246,14 +246,24 @@ FGLTFJsonAnimationIndex FGLTFConvertBuilder::GetOrAddAnimation(FGLTFJsonNodeInde
 	return AnimationDataConverter.GetOrAdd(RootNode, SkeletalMeshComponent);
 }
 
-FGLTFJsonAnimationIndex FGLTFConvertBuilder::GetOrAddAnimation(const ALevelSequenceActor* LevelSequenceActor, const ULevelSequence* LevelSequence)
+FGLTFJsonAnimationIndex FGLTFConvertBuilder::GetOrAddAnimation(const ULevel* Level, const ULevelSequence* LevelSequence)
+{
+	if (Level == nullptr || LevelSequence == nullptr)
+	{
+		return FGLTFJsonAnimationIndex(INDEX_NONE);
+	}
+
+	return LevelSequenceConverter.GetOrAdd(Level, LevelSequence);
+}
+
+FGLTFJsonAnimationIndex FGLTFConvertBuilder::GetOrAddAnimation(const ALevelSequenceActor* LevelSequenceActor)
 {
 	if (LevelSequenceActor == nullptr)
 	{
 		return FGLTFJsonAnimationIndex(INDEX_NONE);
 	}
 
-	return LevelSequenceConverter.GetOrAdd(LevelSequenceActor, LevelSequence);
+	return LevelSequenceDataConverter.GetOrAdd(LevelSequenceActor);
 }
 
 FGLTFJsonNodeIndex FGLTFConvertBuilder::GetOrAddNode(const AActor* Actor)
diff --git a/Engine/Plugins/Enterprise/GLTFExporter/MigratedFromGit/Source/GLTFExporter/Private/Builders/GLTFConvertBuilder.h b/Engine/Plugins/Enterprise/GLTFExporter/MigratedFromGit/Source/GLTFExporter/Private/Builders/GLTFConvertBuilder.h
--- a/Engine/Plugins/Enterprise/GLTFExporter/MigratedFromGit/Source/GLTFExporter/Private/Builders/GLTFConvertBuilder.h
+++ b/Engine/Plugins/Enterprise/GLTFExporter/MigratedFromGit/Source/GLTFExporter/Private/Builders/GLTFConvertBuilder.h
@@ -54,7 +54,8 @@ public:
 	FGLTFJsonSkinIndex GetOrAddSkin(FGLTFJsonNodeIndex RootNode, const USkeletalMeshComponent* SkeletalMeshComponent);
 	FGLTFJsonAnimationIndex GetOrAddAnimation(FGLTFJsonNodeIndex RootNode, const USkeletalMesh* SkeletalMesh, const UAnimSequence* AnimSequence);
 	FGLTFJsonAnimationIndex GetOrAddAnimation(FGLTFJsonNodeIndex RootNode, const USkeletalMeshComponent* SkeletalMeshComponent);
-	FGLTFJsonAnimationIndex GetOrAddAnimation(const ALevelSequenceActor* LevelSequenceActor, const ULevelSequence* LevelSequence);
+	FGLTFJsonAnimationIndex GetOrAddAnimation(const ULevel* Level, const ULevelSequence* LevelSequence);
+	FGLTFJsonAnimationIndex GetOrAddAnimation(const ALevelSequenceActor* LevelSequenceActor);
 
 	FGLTFJsonNodeIndex GetOrAddNode(const AActor* Actor);
 	FGLTFJsonNodeIndex GetOrAddNode(const USceneComponent* SceneComponent);
@@ -97,6 +98,7 @@ private:
 	FGLTFAnimationConverter AnimationConverter{ *this };
 	FGLTFAnimationDataConverter AnimationDataConverter{ *this };
 	FGLTFLevelSequenceConverter LevelSequenceConverter{ *this };
+	FGLTFLevelSequenceDataConverter LevelSequenceDataConverter{ *this };
 
 	FGLTFActorConverter ActorConverter{ *this };
 	FGLTFComponentConverter ComponentConverter{ *this };
diff --git a/Engine/Plugins/Enterprise/GLTFExporter/MigratedFromGit/Source/GLTFExporter/Private/Converters/GLTFAnimationConverters.cpp b/Engine/Plugins/Enterprise/GLTFExporter/MigratedFromGit/Source/GLTFExporter/Private/Converters/GLTFAnimationConverters.cpp
--- a/Engine/Plugins/Enterprise/GLTFExporter/MigratedFromGit/Source/GLTFExporter/Private/Converters/GLTFAnimationConverters.cpp
+++ b/Engine/Plugins/Enterprise/GLTFExporter/MigratedFromGit/Source/GLTFExporter/Private/Converters/GLTFAnimationConverters.cpp
@@ -56,26 +56,34 @@ FGLTFJsonAnimationIndex FGLTFAnimationDataConverter::Convert(FGLTFJsonNodeIndex
 	return AnimationIndex;
 }
 
-FGLTFJsonAnimationIndex FGLTFLevelSequenceConverter::Convert(const ALevelSequenceActor* LevelSequenceActor, const ULevelSequence* LevelSequence)
+FGLTFJsonAnimationIndex FGLTFLevelSequenceConverter::Convert(const ULevel* Level, const ULevelSequence* LevelSequence)
 {
-	FGLTFJsonAnimation JsonAnimation;
+	const FGLTFJsonAnimationIndex AnimationIndex = Builder.AddAnimation();
+	Builder.SetupTask<FGLTFLevelSequenceTask>(Builder, Level, LevelSequence, AnimationIndex);
+	return AnimationIndex;
+}
 
-	if (LevelSequence == nullptr)
+FGLTFJsonAnimationIndex FGLTFLevelSequenceDataConverter::Convert(const ALevelSequenceActor* LevelSequenceActor)
+{
+	const ULevel* Level = LevelSequenceActor->GetLevel();
+	const ULevelSequence* LevelSequence = LevelSequenceActor->LoadSequence();
+
+	if (Level == nullptr || LevelSequence == nullptr)
 	{
-		LevelSequence = LevelSequenceActor->LoadSequence();
-		if (LevelSequence == nullptr)
-		{
-			return FGLTFJsonAnimationIndex(INDEX_NONE);
-		}
+		return FGLTFJsonAnimationIndex(INDEX_NONE);
+	}
 
+	const FGLTFJsonAnimationIndex AnimationIndex = Builder.GetOrAddAnimation(Level, LevelSequence);
+	if (AnimationIndex != INDEX_NONE)
+	{
+		FGLTFJsonAnimation& JsonAnimation = Builder.GetAnimation(AnimationIndex);
 		FGLTFJsonPlayData& JsonPlayData = JsonAnimation.PlayData;
+
 		JsonPlayData.Looping = LevelSequenceActor->PlaybackSettings.LoopCount.Value != 0;
 		JsonPlayData.Playing = LevelSequenceActor->PlaybackSettings.bAutoPlay;
 		JsonPlayData.PlayRate = LevelSequenceActor->PlaybackSettings.PlayRate;
 		JsonPlayData.Position = LevelSequenceActor->PlaybackSettings.StartTime;
 	}
 
-	const FGLTFJsonAnimationIndex AnimationIndex = Builder.AddAnimation(JsonAnimation);
-	Builder.SetupTask<FGLTFLevelSequenceTask>(Builder, LevelSequenceActor, LevelSequence, AnimationIndex);
 	return AnimationIndex;
 }
diff --git a/Engine/Plugins/Enterprise/GLTFExporter/MigratedFromGit/Source/GLTFExporter/Private/Converters/GLTFAnimationConverters.h b/Engine/Plugins/Enterprise/GLTFExporter/MigratedFromGit/Source/GLTFExporter/Private/Converters/GLTFAnimationConverters.h
--- a/Engine/Plugins/Enterprise/GLTFExporter/MigratedFromGit/Source/GLTFExporter/Private/Converters/GLTFAnimationConverters.h
+++ b/Engine/Plugins/Enterprise/GLTFExporter/MigratedFromGit/Source/GLTFExporter/Private/Converters/GLTFAnimationConverters.h
@@ -30,9 +30,16 @@ class FGLTFAnimationDataConverter final : public TGLTFAnimationConverter<FGLTFJs
 	virtual FGLTFJsonAnimationIndex Convert(FGLTFJsonNodeIndex RootNode, const USkeletalMeshComponent* SkeletalMeshComponent) override;
 };
 
-class FGLTFLevelSequenceConverter final : public TGLTFAnimationConverter<const ALevelSequenceActor*, const ULevelSequence*>
+class FGLTFLevelSequenceConverter final : public TGLTFAnimationConverter<const ULevel*, const ULevelSequence*>
 {
 	using TGLTFAnimationConverter::TGLTFAnimationConverter;
 
-	virtual FGLTFJsonAnimationIndex Convert(const ALevelSequenceActor* LevelSequenceActor, const ULevelSequence*) override;
+	virtual FGLTFJsonAnimationIndex Convert(const ULevel* Level, const ULevelSequence*) override;
+};
+
+class FGLTFLevelSequenceDataConverter final : public TGLTFAnimationConverter<const ALevelSequenceActor*>
+{
+	using TGLTFAnimationConverter::TGLTFAnimationConverter;
+
+	virtual FGLTFJsonAnimationIndex Convert(const ALevelSequenceActor* LevelSequenceActor) override;
 };
diff --git a/Engine/Plugins/Enterprise/GLTFExporter/MigratedFromGit/Source/GLTFExporter/Private/Converters/GLTFHotspotConverters.cpp b/Engine/Plugins/Enterprise/GLTFExporter/MigratedFromGit/Source/GLTFExporter/Private/Converters/GLTFHotspotConverters.cpp
--- a/Engine/Plugins/Enterprise/GLTFExporter/MigratedFromGit/Source/GLTFExporter/Private/Converters/GLTFHotspotConverters.cpp
+++ b/Engine/Plugins/Enterprise/GLTFExporter/MigratedFromGit/Source/GLTFExporter/Private/Converters/GLTFHotspotConverters.cpp
@@ -45,7 +45,7 @@ FGLTFJsonHotspotIndex FGLTFHotspotConverter::Convert(const AGLTFHotspotActor* Ho
 			}
 		}
 	}
-	else if (const ALevelSequenceActor* LevelSequenceActor = HotspotActor->LevelSequenceActor)
+	else if (const ULevelSequence* LevelSequence = HotspotActor->LevelSequence)
 	{
 		if (!Builder.ExportOptions->bExportLevelSequences)
 		{
@@ -55,14 +55,7 @@ FGLTFJsonHotspotIndex FGLTFHotspotConverter::Convert(const AGLTFHotspotActor* Ho
 		}
 		else
 		{
-			if (const ULevelSequence* LevelSequence = HotspotActor->LevelSequence)
-			{
-				JsonHotspot.Animation = Builder.GetOrAddAnimation(LevelSequenceActor, LevelSequence);
-			}
-			else
-			{
-				// TODO: report warning
-			}
+			JsonHotspot.Animation = Builder.GetOrAddAnimation(HotspotActor->GetLevel(), LevelSequence);
 		}
 	}
 	else
diff --git a/Engine/Plugins/Enterprise/GLTFExporter/MigratedFromGit/Source/GLTFExporter/Private/Converters/GLTFNodeConverters.cpp b/Engine/Plugins/Enterprise/GLTFExporter/MigratedFromGit/Source/GLTFExporter/Private/Converters/GLTFNodeConverters.cpp
--- a/Engine/Plugins/Enterprise/GLTFExporter/MigratedFromGit/Source/GLTFExporter/Private/Converters/GLTFNodeConverters.cpp
+++ b/Engine/Plugins/Enterprise/GLTFExporter/MigratedFromGit/Source/GLTFExporter/Private/Converters/GLTFNodeConverters.cpp
@@ -36,7 +36,7 @@ FGLTFJsonNodeIndex FGLTFActorConverter::Convert(const AActor* Actor)
 	{
 		if (Builder.ExportOptions->bExportLevelSequences)
 		{
-			Builder.GetOrAddAnimation(LevelSequenceActor, nullptr);
+			Builder.GetOrAddAnimation(LevelSequenceActor);
 		}
 	}
 	else if (const AGLTFHotspotActor* HotspotActor = Cast<AGLTFHotspotActor>(Actor))
diff --git a/Engine/Plugins/Enterprise/GLTFExporter/MigratedFromGit/Source/GLTFExporter/Private/Tasks/GLTFAnimationTasks.cpp b/Engine/Plugins/Enterprise/GLTFExporter/MigratedFromGit/Source/GLTFExporter/Private/Tasks/GLTFAnimationTasks.cpp
--- a/Engine/Plugins/Enterprise/GLTFExporter/MigratedFromGit/Source/GLTFExporter/Private/Tasks/GLTFAnimationTasks.cpp
+++ b/Engine/Plugins/Enterprise/GLTFExporter/MigratedFromGit/Source/GLTFExporter/Private/Tasks/GLTFAnimationTasks.cpp
@@ -183,9 +183,12 @@ void FGLTFLevelSequenceTask::Complete()
 	ULevelSequence* Sequence = const_cast<ULevelSequence*>(LevelSequence);
 	UMovieScene* MovieScene = Sequence->GetMovieScene();
 
-	ULevelSequencePlayer* LevelSequencePlayer = LevelSequenceActor->SequencePlayer;
-	LevelSequencePlayer->Initialize(Sequence, LevelSequenceActor->GetLevel(), FMovieSceneSequencePlaybackSettings(), FLevelSequenceCameraSettings());
-	LevelSequencePlayer->State.AssignSequence(MovieSceneSequenceID::Root, *Sequence, *LevelSequencePlayer);
+	ALevelSequenceActor* LevelSequenceActor;
+	UWorld* World = Level->GetWorld();
+	ULevelSequencePlayer* Player = ULevelSequencePlayer::CreateLevelSequencePlayer(World, Sequence, FMovieSceneSequencePlaybackSettings(), LevelSequenceActor);
+
+	Player->Initialize(Sequence, LevelSequenceActor->GetLevel(), FMovieSceneSequencePlaybackSettings(), FLevelSequenceCameraSettings());
+	Player->State.AssignSequence(MovieSceneSequenceID::Root, *Sequence, *Player);
 
 	FFrameRate TickResolution = MovieScene->GetTickResolution();
 	FFrameRate DisplayRate = MovieScene->GetDisplayRate();
@@ -220,7 +223,7 @@ void FGLTFLevelSequenceTask::Complete()
 
 	for (const FMovieSceneBinding& Binding : MovieScene->GetBindings())
 	{
-		for (TWeakObjectPtr<UObject> Object : LevelSequencePlayer->FindBoundObjects(Binding.GetObjectGuid(), MovieSceneSequenceID::Root))
+		for (TWeakObjectPtr<UObject> Object : Player->FindBoundObjects(Binding.GetObjectGuid(), MovieSceneSequenceID::Root))
 		{
 			FGLTFJsonNodeIndex NodeIndex;
 
@@ -385,4 +388,6 @@ void FGLTFLevelSequenceTask::Complete()
 			}
 		}
 	}
+
+	World->DestroyActor(LevelSequenceActor);
 }
diff --git a/Engine/Plugins/Enterprise/GLTFExporter/MigratedFromGit/Source/GLTFExporter/Private/Tasks/GLTFAnimationTasks.h b/Engine/Plugins/Enterprise/GLTFExporter/MigratedFromGit/Source/GLTFExporter/Private/Tasks/GLTFAnimationTasks.h
--- a/Engine/Plugins/Enterprise/GLTFExporter/MigratedFromGit/Source/GLTFExporter/Private/Tasks/GLTFAnimationTasks.h
+++ b/Engine/Plugins/Enterprise/GLTFExporter/MigratedFromGit/Source/GLTFExporter/Private/Tasks/GLTFAnimationTasks.h
@@ -41,10 +41,10 @@ class FGLTFLevelSequenceTask : public FGLTFTask
 {
 public:
 
-	FGLTFLevelSequenceTask(FGLTFConvertBuilder& Builder, const ALevelSequenceActor* LevelSequenceActor, const ULevelSequence* LevelSequence, FGLTFJsonAnimationIndex AnimationIndex)
+	FGLTFLevelSequenceTask(FGLTFConvertBuilder& Builder, const ULevel* Level, const ULevelSequence* LevelSequence, FGLTFJsonAnimationIndex AnimationIndex)
 		: FGLTFTask(EGLTFTaskPriority::Animation)
 		, Builder(Builder)
-		, LevelSequenceActor(LevelSequenceActor)
+		, Level(Level)
 		, LevelSequence(LevelSequence)
 		, AnimationIndex(AnimationIndex)
 	{
@@ -52,7 +52,7 @@ public:
 
 	virtual FString GetName() override
 	{
-		return LevelSequenceActor->GetName();
+		return LevelSequence->GetName();
 	}
 
 	virtual void Complete() override;
@@ -60,7 +60,7 @@ public:
 private:
 
 	FGLTFConvertBuilder& Builder;
-	const ALevelSequenceActor* LevelSequenceActor;
+	const ULevel* Level;
 	const ULevelSequence* LevelSequence;
 	const FGLTFJsonAnimationIndex AnimationIndex;
 };
diff --git a/Engine/Plugins/Enterprise/GLTFExporter/MigratedFromGit/Source/GLTFExporterRuntime/Private/Actors/GLTFHotspotActor.cpp b/Engine/Plugins/Enterprise/GLTFExporter/MigratedFromGit/Source/GLTFExporterRuntime/Private/Actors/GLTFHotspotActor.cpp
--- a/Engine/Plugins/Enterprise/GLTFExporter/MigratedFromGit/Source/GLTFExporterRuntime/Private/Actors/GLTFHotspotActor.cpp
+++ b/Engine/Plugins/Enterprise/GLTFExporter/MigratedFromGit/Source/GLTFExporterRuntime/Private/Actors/GLTFHotspotActor.cpp
@@ -109,6 +109,8 @@ AGLTFHotspotActor::AGLTFHotspotActor(const FObjectInitializer& ObjectInitializer
 
 	PrimaryActorTick.bCanEverTick = true;
 	PrimaryActorTick.bStartWithTickEnabled = true;
+
+	LevelSequencePlayer = ObjectInitializer.CreateDefaultSubobject<ULevelSequencePlayer>(this, "LevelSequencePlayer");
 }
 
 
@@ -139,14 +141,6 @@ void AGLTFHotspotActor::PostEditChangeProperty(FPropertyChangedEvent& PropertyCh
 		{
 			ValidateAnimation();
 		}
-		else if (PropertyFName == GET_MEMBER_NAME_CHECKED(ThisClass, LevelSequenceActor))
-		{
-			if (LevelSequenceActor != nullptr && LevelSequence == nullptr)
-			{
-				LevelSequence = LevelSequenceActor->LoadSequence();
-			}
-			ValidateAnimation();
-		}
 	}
 }
 #endif // WITH_EDITOR
@@ -235,6 +229,22 @@ void AGLTFHotspotActor::Tick(float DeltaTime)
 		SetSpriteOpacity(Opacity);
 		bIsInteractable = Opacity >= 0.5f;
 	}
+
+	if (LevelSequencePlayer != nullptr)
+	{
+		LevelSequencePlayer->Update(DeltaTime);
+	}
+}
+
+void AGLTFHotspotActor::EndPlay(const EEndPlayReason::Type EndPlayReason)
+{
+	if (LevelSequencePlayer != nullptr)
+	{
+		// Stop may modify a lot of actor state so it needs to be called
+		// during EndPlay (when Actors + World are still valid) instead
+		// of waiting for the UObject to be destroyed by GC.
+		LevelSequencePlayer->Stop();
+	}
 }
 
 void AGLTFHotspotActor::BeginCursorOver(UPrimitiveComponent* TouchedComponent)
@@ -412,11 +422,9 @@ void AGLTFHotspotActor::ToggleAnimation()
 			SkeletalMeshComponent->Play(false);
 		}
 	}
-	else if (LevelSequenceActor != nullptr && LevelSequence != nullptr)
+	else if (LevelSequence != nullptr)
 	{
-		ULevelSequencePlayer* LevelSequencePlayer = LevelSequenceActor->GetSequencePlayer();
-
-		if (LevelSequencePlayer != nullptr && LevelSequencePlayer->IsPlaying() && LevelSequenceActor->GetSequence() == LevelSequence)
+		if (LevelSequencePlayer->IsPlaying())
 		{
 			// If the same animation is already playing, just reverse the play rate for a smooth transition
 			LevelSequencePlayer->SetPlayRate(TargetPlayRate);
@@ -426,13 +434,14 @@ void AGLTFHotspotActor::ToggleAnimation()
 			// TODO: improve this calculation
 			const float SequenceLength = LevelSequence->GetMovieScene()->GetTickResolution().AsSeconds(MovieScene::DiscreteSize(LevelSequence->GetMovieScene()->GetPlaybackRange()));
 
-			LevelSequenceActor->PlaybackSettings.PlayRate = TargetPlayRate;
-			LevelSequenceActor->PlaybackSettings.StartTime = bReverseAnimation ? SequenceLength : 0.0f;
-			LevelSequenceActor->PlaybackSettings.bAutoPlay = true;
-			LevelSequenceActor->PlaybackSettings.bPauseAtEnd = true;
-			LevelSequenceActor->PlaybackSettings.LoopCount.Value = 0;
-			LevelSequenceActor->SetSequence(LevelSequence);
-			LevelSequenceActor->GetSequencePlayer()->Play();
+			FMovieSceneSequencePlaybackSettings PlaybackSettings;
+			PlaybackSettings.PlayRate = TargetPlayRate;
+			PlaybackSettings.StartTime = bReverseAnimation ? SequenceLength : 0.0f;
+			PlaybackSettings.bAutoPlay = true;
+			PlaybackSettings.bPauseAtEnd = true;
+			PlaybackSettings.LoopCount.Value = 0;
+			LevelSequencePlayer->Initialize(LevelSequence, GetLevel(), PlaybackSettings, FLevelSequenceCameraSettings());
+			LevelSequencePlayer->Play();
 		}
 	}
 
diff --git a/Engine/Plugins/Enterprise/GLTFExporter/MigratedFromGit/Source/GLTFExporterRuntime/Public/Actors/GLTFHotspotActor.h b/Engine/Plugins/Enterprise/GLTFExporter/MigratedFromGit/Source/GLTFExporterRuntime/Public/Actors/GLTFHotspotActor.h
--- a/Engine/Plugins/Enterprise/GLTFExporter/MigratedFromGit/Source/GLTFExporterRuntime/Public/Actors/GLTFHotspotActor.h
+++ b/Engine/Plugins/Enterprise/GLTFExporter/MigratedFromGit/Source/GLTFExporterRuntime/Public/Actors/GLTFHotspotActor.h
@@ -7,7 +7,7 @@
 
 class ASkeletalMeshActor;
 class UAnimSequence;
-class ALevelSequenceActor;
+class ULevelSequencePlayer;
 class ULevelSequence;
 class UTexture2D;
 class UTexture;
@@ -43,6 +43,7 @@ public:
 	//~ Begin AActor Interface
 	virtual void PostRegisterAllComponents() override;
 	virtual void Tick(float DeltaSeconds) override;
+	virtual void EndPlay(const EEndPlayReason::Type EndPlayReason) override;
 	//~ End AActor Interface
 
 private:
@@ -63,16 +64,12 @@ public:
 	UPROPERTY(EditAnywhere, BlueprintReadWrite, Category = "Hotspot Animation")
 	ASkeletalMeshActor* SkeletalMeshActor;
 
-	/* The animation that will be played on the skeletal mesh actor. Must be compatible with its skeletal mesh asset. */
+	/* The animation sequence that will be played on the skeletal mesh actor. Must be compatible with its skeletal mesh asset. */
 	UPROPERTY(EditAnywhere, BlueprintReadWrite, Category = "Hotspot Animation", meta=(EditCondition="SkeletalMeshActor != nullptr"))
 	UAnimSequence* AnimationSequence;
 
-	/* The skeletal mesh actor that will be animated when the hotspot is clicked. */
+	/* The level sequence that will be played in the level. */
 	UPROPERTY(EditAnywhere, BlueprintReadWrite, Category = "Hotspot Animation", meta=(EditCondition="SkeletalMeshActor == nullptr"))
-	ALevelSequenceActor* LevelSequenceActor;
-
-	/* The animation that will be played on the level sequence actor. */
-	UPROPERTY(EditAnywhere, BlueprintReadWrite, Category = "Hotspot Animation", meta=(EditCondition="SkeletalMeshActor == nullptr && LevelSequenceActor != nullptr"))
 	ULevelSequence* LevelSequence;
 
 	/* The billboard image that will be shown when the hotspot is in an inactive state or one without a specified image. */
@@ -120,6 +117,9 @@ private:
 	UPROPERTY(Transient, DuplicateTransient)
 	const UTexture* ActiveImage;
 
+	UPROPERTY(Instanced, Transient)
+	ULevelSequencePlayer* LevelSequencePlayer;
+
 	FVector2D ActiveImageSize;
 
 	bool bToggled;
