[GLTFExporter] fix: ensure each animation has a unique name

commit b78b8a945f45557b552761d2994364dad6079d34
Date: Tue, 6 Oct 2020 13:49:49 +0200

    fix: ensure each animation has a unique name
    (cherry picked from commit fcced6d99039e79d24cf72183c703e5bde01eb1c​)

#skipci

diff --git a/Engine/Plugins/Enterprise/GLTFExporter/MigratedFromGit/Source/GLTFExporter/Private/Tasks/GLTFAnimationTasks.cpp b/Engine/Plugins/Enterprise/GLTFExporter/MigratedFromGit/Source/GLTFExporter/Private/Tasks/GLTFAnimationTasks.cpp
--- a/Engine/Plugins/Enterprise/GLTFExporter/MigratedFromGit/Source/GLTFExporter/Private/Tasks/GLTFAnimationTasks.cpp
+++ b/Engine/Plugins/Enterprise/GLTFExporter/MigratedFromGit/Source/GLTFExporter/Private/Tasks/GLTFAnimationTasks.cpp
@@ -17,7 +17,7 @@ void FGLTFAnimSequenceTask::Complete()
 	const USkeleton* Skeleton = AnimSequence->GetSkeleton();
 
 	FGLTFJsonAnimation& JsonAnimation = Builder.GetAnimation(AnimationIndex);
-	AnimSequence->GetName(JsonAnimation.Name);
+	JsonAnimation.Name = AnimSequence->GetName() + TEXT("_") + FString::FromInt(AnimationIndex); // Ensure unique name
 
 	TArray<float> Timestamps;
 	Timestamps.AddUninitialized(FrameCount);
@@ -198,7 +198,7 @@ void FGLTFLevelSequenceTask::Complete()
 	int32 FrameCount = FFrameRate::TransformTime(FFrameTime(FFrameNumber(MovieScene::DiscreteSize(PlaybackRange))), TickResolution, DisplayRate).RoundToFrame().Value + 1;
 
 	FGLTFJsonAnimation& JsonAnimation = Builder.GetAnimation(AnimationIndex);
-	Sequence->GetName(JsonAnimation.Name);
+	JsonAnimation.Name = Sequence->GetName() + TEXT("_") + FString::FromInt(AnimationIndex); // Ensure unique name
 
 	TArray<float> Timestamps;
 	TArray<FFrameTime> FrameTimes;
