[GLTFExporter] refactor: rename EPIC_interaction_hotspots -> EPIC_animation_hotspots, EPIC_animation_play_data -> EPIC_animation_playback, Looping -> Loop, Playing -> AutoPlay, and Position -> StartTime

commit 45b05c11de7ea58474026436a4d85df908e2f7a2
Date: Tue, 6 Oct 2020 14:46:06 +0200

    refactor: rename EPIC_interaction_hotspots -> EPIC_animation_hotspots, EPIC_animation_play_data -> EPIC_animation_playback, Looping -> Loop, Playing -> AutoPlay, and Position -> StartTime
    (cherry picked from commit 39f4a5f6159789d6a7353dfe9f7346541647ddde​)

#skipci

diff --git a/Engine/Plugins/Enterprise/GLTFExporter/MigratedFromGit/Source/GLTFExporter/Private/Converters/GLTFAnimationConverters.cpp b/Engine/Plugins/Enterprise/GLTFExporter/MigratedFromGit/Source/GLTFExporter/Private/Converters/GLTFAnimationConverters.cpp
--- a/Engine/Plugins/Enterprise/GLTFExporter/MigratedFromGit/Source/GLTFExporter/Private/Converters/GLTFAnimationConverters.cpp
+++ b/Engine/Plugins/Enterprise/GLTFExporter/MigratedFromGit/Source/GLTFExporter/Private/Converters/GLTFAnimationConverters.cpp
@@ -45,12 +45,12 @@ FGLTFJsonAnimationIndex FGLTFAnimationDataConverter::Convert(FGLTFJsonNodeIndex
 	if (AnimationIndex != INDEX_NONE)
 	{
 		FGLTFJsonAnimation& JsonAnimation = Builder.GetAnimation(AnimationIndex);
-		FGLTFJsonPlayData& JsonPlayData = JsonAnimation.PlayData;
+		FGLTFJsonAnimationPlayback& JsonPlayback = JsonAnimation.Playback;
 
-		JsonPlayData.Looping = SkeletalMeshComponent->AnimationData.bSavedLooping;
-		JsonPlayData.Playing = SkeletalMeshComponent->AnimationData.bSavedPlaying;
-		JsonPlayData.PlayRate = SkeletalMeshComponent->AnimationData.SavedPlayRate;
-		JsonPlayData.Position = SkeletalMeshComponent->AnimationData.SavedPosition;
+		JsonPlayback.bLoop = SkeletalMeshComponent->AnimationData.bSavedLooping;
+		JsonPlayback.bAutoPlay = SkeletalMeshComponent->AnimationData.bSavedPlaying;
+		JsonPlayback.PlayRate = SkeletalMeshComponent->AnimationData.SavedPlayRate;
+		JsonPlayback.StartTime = SkeletalMeshComponent->AnimationData.SavedPosition;
 	}
 
 	return AnimationIndex;
@@ -77,12 +77,12 @@ FGLTFJsonAnimationIndex FGLTFLevelSequenceDataConverter::Convert(const ALevelSeq
 	if (AnimationIndex != INDEX_NONE)
 	{
 		FGLTFJsonAnimation& JsonAnimation = Builder.GetAnimation(AnimationIndex);
-		FGLTFJsonPlayData& JsonPlayData = JsonAnimation.PlayData;
+		FGLTFJsonAnimationPlayback& JsonPlayback = JsonAnimation.Playback;
 
-		JsonPlayData.Looping = LevelSequenceActor->PlaybackSettings.LoopCount.Value != 0;
-		JsonPlayData.Playing = LevelSequenceActor->PlaybackSettings.bAutoPlay;
-		JsonPlayData.PlayRate = LevelSequenceActor->PlaybackSettings.PlayRate;
-		JsonPlayData.Position = LevelSequenceActor->PlaybackSettings.StartTime;
+		JsonPlayback.bLoop = LevelSequenceActor->PlaybackSettings.LoopCount.Value != 0;
+		JsonPlayback.bAutoPlay = LevelSequenceActor->PlaybackSettings.bAutoPlay;
+		JsonPlayback.PlayRate = LevelSequenceActor->PlaybackSettings.PlayRate;
+		JsonPlayback.StartTime = LevelSequenceActor->PlaybackSettings.StartTime;
 	}
 
 	return AnimationIndex;
diff --git a/Engine/Plugins/Enterprise/GLTFExporter/MigratedFromGit/Source/GLTFExporter/Private/Converters/GLTFNodeConverters.cpp b/Engine/Plugins/Enterprise/GLTFExporter/MigratedFromGit/Source/GLTFExporter/Private/Converters/GLTFNodeConverters.cpp
--- a/Engine/Plugins/Enterprise/GLTFExporter/MigratedFromGit/Source/GLTFExporter/Private/Converters/GLTFNodeConverters.cpp
+++ b/Engine/Plugins/Enterprise/GLTFExporter/MigratedFromGit/Source/GLTFExporter/Private/Converters/GLTFNodeConverters.cpp
@@ -41,7 +41,7 @@ FGLTFJsonNodeIndex FGLTFActorConverter::Convert(const AActor* Actor)
 	}
 	else if (const AGLTFHotspotActor* HotspotActor = Cast<AGLTFHotspotActor>(Actor))
 	{
-		if (Builder.ExportOptions->bExportInteractionHotspots)
+		if (Builder.ExportOptions->bExportAnimationHotspots)
 		{
 			FGLTFJsonNode& RootNode = Builder.GetNode(RootNodeIndex);
 			RootNode.Hotspot = Builder.GetOrAddHotspot(HotspotActor);
diff --git a/Engine/Plugins/Enterprise/GLTFExporter/MigratedFromGit/Source/GLTFExporter/Private/Exporters/GLTFExporter.cpp b/Engine/Plugins/Enterprise/GLTFExporter/MigratedFromGit/Source/GLTFExporter/Private/Exporters/GLTFExporter.cpp
--- a/Engine/Plugins/Enterprise/GLTFExporter/MigratedFromGit/Source/GLTFExporter/Private/Exporters/GLTFExporter.cpp
+++ b/Engine/Plugins/Enterprise/GLTFExporter/MigratedFromGit/Source/GLTFExporter/Private/Exporters/GLTFExporter.cpp
@@ -55,6 +55,11 @@ bool UGLTFExporter::ExportBinary(UObject* Object, const TCHAR* Type, FArchive& A
 	else
 	{
 		Builder.ShowMessages();
+
+		if (bSuccess && Options->bShowFilesWhenDone)
+		{
+			FPlatformProcess::ExploreFolder(*Builder.DirPath);
+		}
 	}
 
 	return bSuccess;
diff --git a/Engine/Plugins/Enterprise/GLTFExporter/MigratedFromGit/Source/GLTFExporter/Private/GLTFExportOptions.cpp b/Engine/Plugins/Enterprise/GLTFExporter/MigratedFromGit/Source/GLTFExporter/Private/GLTFExportOptions.cpp
--- a/Engine/Plugins/Enterprise/GLTFExporter/MigratedFromGit/Source/GLTFExporter/Private/GLTFExportOptions.cpp
+++ b/Engine/Plugins/Enterprise/GLTFExporter/MigratedFromGit/Source/GLTFExporter/Private/GLTFExportOptions.cpp
@@ -35,10 +35,11 @@ UGLTFExportOptions::UGLTFExportOptions(const FObjectInitializer& ObjectInitializ
 	bExportOrbitalCameras = true;
 	bExportHDRIBackdrops = true;
 	bExportVariantSets = true;
-	bExportInteractionHotspots = true;
+	bExportAnimationHotspots = true;
 	bBundleWebViewer = true;
 	bExportPreviewMesh = true;
 	bAllExtensionsRequired = false;
+	bShowFilesWhenDone = true;
 }
 
 void UGLTFExportOptions::ResetToDefault()
diff --git a/Engine/Plugins/Enterprise/GLTFExporter/MigratedFromGit/Source/GLTFExporter/Private/Json/GLTFJsonAnimation.h b/Engine/Plugins/Enterprise/GLTFExporter/MigratedFromGit/Source/GLTFExporter/Private/Json/GLTFJsonAnimation.h
--- a/Engine/Plugins/Enterprise/GLTFExporter/MigratedFromGit/Source/GLTFExporter/Private/Json/GLTFJsonAnimation.h
+++ b/Engine/Plugins/Enterprise/GLTFExporter/MigratedFromGit/Source/GLTFExporter/Private/Json/GLTFJsonAnimation.h
@@ -4,7 +4,7 @@
 
 #include "Json/GLTFJsonEnums.h"
 #include "Json/GLTFJsonIndex.h"
-#include "Json/GLTFJsonPlayData.h"
+#include "Json/GLTFJsonAnimationPlayback.h"
 #include "Json/GLTFJsonUtility.h"
 #include "Serialization/JsonSerializer.h"
 
@@ -85,7 +85,7 @@ struct FGLTFJsonAnimation
 	TArray<FGLTFJsonAnimationChannel> Channels;
 	TArray<FGLTFJsonAnimationSampler> Samplers;
 
-	FGLTFJsonPlayData PlayData;
+	FGLTFJsonAnimationPlayback Playback;
 
 	template <class CharType = TCHAR, class PrintPolicy = TPrettyJsonPrintPolicy<CharType>>
 	void WriteObject(TJsonWriter<CharType, PrintPolicy>& JsonWriter, FGLTFJsonExtensions& Extensions) const
@@ -100,14 +100,14 @@ struct FGLTFJsonAnimation
 		FGLTFJsonUtility::WriteObjectArray(JsonWriter, TEXT("channels"), Channels, Extensions, true);
 		FGLTFJsonUtility::WriteObjectArray(JsonWriter, TEXT("samplers"), Samplers, Extensions, true);
 
-		if (PlayData != FGLTFJsonPlayData())
+		if (Playback != FGLTFJsonAnimationPlayback())
 		{
-			const EGLTFJsonExtension Extension = EGLTFJsonExtension::EPIC_AnimationPlayData;
+			const EGLTFJsonExtension Extension = EGLTFJsonExtension::EPIC_AnimationPlayback;
 			Extensions.Used.Add(Extension);
 
 			JsonWriter.WriteObjectStart(TEXT("extensions"));
 			JsonWriter.WriteIdentifierPrefix(FGLTFJsonUtility::ToString(Extension));
-			PlayData.WriteObject(JsonWriter, Extensions);
+			Playback.WriteObject(JsonWriter, Extensions);
 			JsonWriter.WriteObjectEnd();
 		}
 
diff --git a/Engine/Plugins/Enterprise/GLTFExporter/MigratedFromGit/Source/GLTFExporter/Private/Json/GLTFJsonAnimationPlayback.h b/Engine/Plugins/Enterprise/GLTFExporter/MigratedFromGit/Source/GLTFExporter/Private/Json/GLTFJsonAnimationPlayback.h
new file mode 100644
--- /dev/null
+++ b/Engine/Plugins/Enterprise/GLTFExporter/MigratedFromGit/Source/GLTFExporter/Private/Json/GLTFJsonAnimationPlayback.h
@@ -0,0 +1,71 @@
+// Copyright Epic Games, Inc. All Rights Reserved.
+
+#pragma once
+
+#include "Json/GLTFJsonUtility.h"
+#include "Serialization/JsonSerializer.h"
+
+struct FGLTFJsonAnimationPlayback
+{
+	FString Name;
+
+	bool bLoop;
+	bool bAutoPlay;
+
+	float PlayRate;
+	float StartTime;
+
+	FGLTFJsonAnimationPlayback()
+        : bLoop(true)
+        , bAutoPlay(true)
+        , PlayRate(1)
+        , StartTime(0)
+	{
+	}
+
+	template <class CharType = TCHAR, class PrintPolicy = TPrettyJsonPrintPolicy<CharType>>
+	void WriteObject(TJsonWriter<CharType, PrintPolicy>& JsonWriter, FGLTFJsonExtensions& Extensions) const
+	{
+		JsonWriter.WriteObjectStart();
+
+		if (!Name.IsEmpty())
+		{
+			JsonWriter.WriteValue(TEXT("name"), Name);
+		}
+
+		if (bLoop != true)
+		{
+			JsonWriter.WriteValue(TEXT("loop"), bLoop);
+		}
+
+		if (bAutoPlay != true)
+		{
+			JsonWriter.WriteValue(TEXT("autoPlay"), bAutoPlay);
+		}
+
+		if (PlayRate != 1)
+		{
+			FGLTFJsonUtility::WriteExactValue(JsonWriter, TEXT("playRate"), PlayRate);
+		}
+
+		if (StartTime != 0)
+		{
+			FGLTFJsonUtility::WriteExactValue(JsonWriter, TEXT("startTime"), StartTime);
+		}
+
+		JsonWriter.WriteObjectEnd();
+	}
+
+	bool operator==(const FGLTFJsonAnimationPlayback& Other) const
+	{
+		return bLoop == Other.bLoop
+            && bAutoPlay == Other.bAutoPlay
+            && PlayRate == Other.PlayRate
+            && StartTime == Other.StartTime;
+	}
+
+	bool operator!=(const FGLTFJsonAnimationPlayback& Other) const
+	{
+		return !(*this == Other);
+	}
+};
diff --git a/Engine/Plugins/Enterprise/GLTFExporter/MigratedFromGit/Source/GLTFExporter/Private/Json/GLTFJsonEnums.h b/Engine/Plugins/Enterprise/GLTFExporter/MigratedFromGit/Source/GLTFExporter/Private/Json/GLTFJsonEnums.h
--- a/Engine/Plugins/Enterprise/GLTFExporter/MigratedFromGit/Source/GLTFExporter/Private/Json/GLTFJsonEnums.h
+++ b/Engine/Plugins/Enterprise/GLTFExporter/MigratedFromGit/Source/GLTFExporter/Private/Json/GLTFJsonEnums.h
@@ -12,12 +12,12 @@ enum class EGLTFJsonExtension
 	KHR_MeshQuantization,
 	EPIC_OrbitalCameras,
 	EPIC_HDRIBackdrops,
-	EPIC_InteractionHotspots,
+	EPIC_AnimationHotspots,
 	EPIC_LightmapTextures,
 	EPIC_ReflectionCaptures,
 	EPIC_LevelVariantSets,
 	EPIC_TextureHDREncoding,
-    EPIC_AnimationPlayData
+    EPIC_AnimationPlayback
 };
 
 enum class EGLTFJsonShadingModel
diff --git a/Engine/Plugins/Enterprise/GLTFExporter/MigratedFromGit/Source/GLTFExporter/Private/Json/GLTFJsonNode.h b/Engine/Plugins/Enterprise/GLTFExporter/MigratedFromGit/Source/GLTFExporter/Private/Json/GLTFJsonNode.h
--- a/Engine/Plugins/Enterprise/GLTFExporter/MigratedFromGit/Source/GLTFExporter/Private/Json/GLTFJsonNode.h
+++ b/Engine/Plugins/Enterprise/GLTFExporter/MigratedFromGit/Source/GLTFExporter/Private/Json/GLTFJsonNode.h
@@ -136,11 +136,11 @@ struct FGLTFJsonNode
 
 			if (Hotspot != INDEX_NONE)
 			{
-				const EGLTFJsonExtension Extension = EGLTFJsonExtension::EPIC_InteractionHotspots;
+				const EGLTFJsonExtension Extension = EGLTFJsonExtension::EPIC_AnimationHotspots;
 				Extensions.Used.Add(Extension);
 
 				JsonWriter.WriteObjectStart(FGLTFJsonUtility::ToString(Extension));
-				JsonWriter.WriteValue(TEXT("interaction"), Hotspot);
+				JsonWriter.WriteValue(TEXT("hotspot"), Hotspot);
 				JsonWriter.WriteObjectEnd();
 			}
 
diff --git a/Engine/Plugins/Enterprise/GLTFExporter/MigratedFromGit/Source/GLTFExporter/Private/Json/GLTFJsonPlayData.h b/Engine/Plugins/Enterprise/GLTFExporter/MigratedFromGit/Source/GLTFExporter/Private/Json/GLTFJsonPlayData.h
deleted file mode 100644
diff --git a/Engine/Plugins/Enterprise/GLTFExporter/MigratedFromGit/Source/GLTFExporter/Private/Json/GLTFJsonRoot.h b/Engine/Plugins/Enterprise/GLTFExporter/MigratedFromGit/Source/GLTFExporter/Private/Json/GLTFJsonRoot.h
--- a/Engine/Plugins/Enterprise/GLTFExporter/MigratedFromGit/Source/GLTFExporter/Private/Json/GLTFJsonRoot.h
+++ b/Engine/Plugins/Enterprise/GLTFExporter/MigratedFromGit/Source/GLTFExporter/Private/Json/GLTFJsonRoot.h
@@ -156,11 +156,11 @@ struct FGLTFJsonRoot
 
 			if (Hotspots.Num() > 0)
 			{
-				const EGLTFJsonExtension Extension = EGLTFJsonExtension::EPIC_InteractionHotspots;
+				const EGLTFJsonExtension Extension = EGLTFJsonExtension::EPIC_AnimationHotspots;
 				AllExtensions.Used.Add(Extension);
 
 				JsonWriter.WriteObjectStart(FGLTFJsonUtility::ToString(Extension));
-				FGLTFJsonUtility::WriteObjectPtrArray(JsonWriter, TEXT("interactions"), Hotspots, AllExtensions);
+				FGLTFJsonUtility::WriteObjectPtrArray(JsonWriter, TEXT("hotspots"), Hotspots, AllExtensions);
 				JsonWriter.WriteObjectEnd();
 			}
 
diff --git a/Engine/Plugins/Enterprise/GLTFExporter/MigratedFromGit/Source/GLTFExporter/Private/Json/GLTFJsonUtility.h b/Engine/Plugins/Enterprise/GLTFExporter/MigratedFromGit/Source/GLTFExporter/Private/Json/GLTFJsonUtility.h
--- a/Engine/Plugins/Enterprise/GLTFExporter/MigratedFromGit/Source/GLTFExporter/Private/Json/GLTFJsonUtility.h
+++ b/Engine/Plugins/Enterprise/GLTFExporter/MigratedFromGit/Source/GLTFExporter/Private/Json/GLTFJsonUtility.h
@@ -19,19 +19,19 @@ struct FGLTFJsonUtility
 	{
 		switch (Value)
 		{
-			case EGLTFJsonExtension::KHR_LightsPunctual:       return TEXT("KHR_lights_punctual");
-			case EGLTFJsonExtension::KHR_MaterialsUnlit:       return TEXT("KHR_materials_unlit");
-			case EGLTFJsonExtension::KHR_MaterialsClearCoat:   return TEXT("KHR_materials_clearcoat");
-			case EGLTFJsonExtension::KHR_MeshQuantization:     return TEXT("KHR_mesh_quantization");
-			case EGLTFJsonExtension::EPIC_OrbitalCameras:      return TEXT("EPIC_orbital_cameras");
-			case EGLTFJsonExtension::EPIC_HDRIBackdrops:       return TEXT("EPIC_hdri_backdrops");
-			case EGLTFJsonExtension::EPIC_InteractionHotspots: return TEXT("EPIC_interaction_hotspots");
-			case EGLTFJsonExtension::EPIC_LightmapTextures:    return TEXT("EPIC_lightmap_textures");
-			case EGLTFJsonExtension::EPIC_ReflectionCaptures:  return TEXT("EPIC_reflection_captures");
-			case EGLTFJsonExtension::EPIC_LevelVariantSets:    return TEXT("EPIC_level_variant_sets");
-			case EGLTFJsonExtension::EPIC_TextureHDREncoding:  return TEXT("EPIC_texture_hdr_encoding");
-			case EGLTFJsonExtension::EPIC_AnimationPlayData:   return TEXT("EPIC_animation_play_data");
-			default:                                           return TEXT("unknown");
+			case EGLTFJsonExtension::KHR_LightsPunctual:      return TEXT("KHR_lights_punctual");
+			case EGLTFJsonExtension::KHR_MaterialsUnlit:      return TEXT("KHR_materials_unlit");
+			case EGLTFJsonExtension::KHR_MaterialsClearCoat:  return TEXT("KHR_materials_clearcoat");
+			case EGLTFJsonExtension::KHR_MeshQuantization:    return TEXT("KHR_mesh_quantization");
+			case EGLTFJsonExtension::EPIC_OrbitalCameras:     return TEXT("EPIC_orbital_cameras");
+			case EGLTFJsonExtension::EPIC_HDRIBackdrops:      return TEXT("EPIC_hdri_backdrops");
+			case EGLTFJsonExtension::EPIC_AnimationHotspots:  return TEXT("EPIC_animation_hotspots");
+			case EGLTFJsonExtension::EPIC_LightmapTextures:   return TEXT("EPIC_lightmap_textures");
+			case EGLTFJsonExtension::EPIC_ReflectionCaptures: return TEXT("EPIC_reflection_captures");
+			case EGLTFJsonExtension::EPIC_LevelVariantSets:   return TEXT("EPIC_level_variant_sets");
+			case EGLTFJsonExtension::EPIC_TextureHDREncoding: return TEXT("EPIC_texture_hdr_encoding");
+			case EGLTFJsonExtension::EPIC_AnimationPlayback:  return TEXT("EPIC_animation_playback");
+			default:                                          return TEXT("unknown");
 		}
 	}
 
diff --git a/Engine/Plugins/Enterprise/GLTFExporter/MigratedFromGit/Source/GLTFExporter/Public/GLTFExportOptions.h b/Engine/Plugins/Enterprise/GLTFExporter/MigratedFromGit/Source/GLTFExporter/Public/GLTFExportOptions.h
--- a/Engine/Plugins/Enterprise/GLTFExporter/MigratedFromGit/Source/GLTFExporter/Public/GLTFExportOptions.h
+++ b/Engine/Plugins/Enterprise/GLTFExporter/MigratedFromGit/Source/GLTFExporter/Public/GLTFExportOptions.h
@@ -118,7 +118,7 @@ public:
 	uint32 bExportVariantSets : 1;
 
 	UPROPERTY(EditAnywhere, BlueprintReadWrite, config, Category = Scene)
-	uint32 bExportInteractionHotspots : 1;
+	uint32 bExportAnimationHotspots : 1;
 
 	UPROPERTY(EditAnywhere, BlueprintReadWrite, config, Category = Exporter)
 	uint32 bBundleWebViewer : 1;
@@ -129,6 +129,9 @@ public:
 	UPROPERTY(EditAnywhere, BlueprintReadWrite, config, Category = Exporter)
 	uint32 bAllExtensionsRequired : 1;
 
+	UPROPERTY(EditAnywhere, BlueprintReadWrite, config, Category = Exporter)
+	uint32 bShowFilesWhenDone : 1;
+
 	/* Set all the each property to the CDO value */
 	void ResetToDefault();
 
diff --git a/Engine/Plugins/Enterprise/GLTFExporter/MigratedFromGit/Source/GLTFExporterRuntime/Private/Actors/GLTFCameraActor.cpp b/Engine/Plugins/Enterprise/GLTFExporter/MigratedFromGit/Source/GLTFExporterRuntime/Private/Actors/GLTFCameraActor.cpp
--- a/Engine/Plugins/Enterprise/GLTFExporter/MigratedFromGit/Source/GLTFExporterRuntime/Private/Actors/GLTFCameraActor.cpp
+++ b/Engine/Plugins/Enterprise/GLTFExporter/MigratedFromGit/Source/GLTFExporterRuntime/Private/Actors/GLTFCameraActor.cpp
@@ -58,6 +58,7 @@ void AGLTFCameraActor::PostEditChangeProperty(FPropertyChangedEvent& PropertyCha
 			if (Focus == this)
 			{
 				Focus = nullptr;
+				// TODO: don't use LogTemp
 				UE_LOG(LogTemp, Warning, TEXT("The camera focus must not be the camera's own actor"));
 			}
 		}
diff --git a/Engine/Plugins/Enterprise/GLTFExporter/MigratedFromGit/Source/GLTFExporterRuntime/Private/Actors/GLTFHotspotActor.cpp b/Engine/Plugins/Enterprise/GLTFExporter/MigratedFromGit/Source/GLTFExporterRuntime/Private/Actors/GLTFHotspotActor.cpp
--- a/Engine/Plugins/Enterprise/GLTFExporter/MigratedFromGit/Source/GLTFExporterRuntime/Private/Actors/GLTFHotspotActor.cpp
+++ b/Engine/Plugins/Enterprise/GLTFExporter/MigratedFromGit/Source/GLTFExporterRuntime/Private/Actors/GLTFHotspotActor.cpp
@@ -20,11 +20,11 @@
 #include "Slate/SceneViewport.h"
 #endif
 
-DEFINE_LOG_CATEGORY_STATIC(LogEditorGLTFInteractionHotspot, Log, All);
+DEFINE_LOG_CATEGORY_STATIC(LogGLTFHotspot, Log, All);
 
 namespace
 {
-	const FName NAME_InteractionHotspotTag = TEXT("InteractionHotspot");
+	const FName NAME_HotspotTag = TEXT("GLTFHotspot");
 	const FName NAME_LevelEditorModule = TEXT("LevelEditor");
 	const FName NAME_SpriteParameter = TEXT("Sprite");
 	const FName NAME_OpacityParameter = TEXT("OpacityMask");
@@ -64,7 +64,7 @@ AGLTFHotspotActor::AGLTFHotspotActor(const FObjectInitializer& ObjectInitializer
 	AddInstanceComponent(BillboardComponent);
 	SphereComponent->SetupAttachment(RootComponent);
 	SphereComponent->SetMobility(EComponentMobility::Movable);
-	SphereComponent->ComponentTags.Add(NAME_InteractionHotspotTag);
+	SphereComponent->ComponentTags.Add(NAME_HotspotTag);
 	SphereComponent->InitSphereRadius(100.0f);
 	SphereComponent->SetVisibility(false);
 
@@ -200,7 +200,7 @@ void AGLTFHotspotActor::Tick(float DeltaTime)
 		{
 			if (const UPrimitiveComponent* HitComponent = HitResult.GetComponent())
 			{
-				bIsHotspotOccluded = !HitComponent->ComponentTags.Contains(NAME_InteractionHotspotTag);
+				bIsHotspotOccluded = !HitComponent->ComponentTags.Contains(NAME_HotspotTag);
 			}
 		}
 
@@ -458,11 +458,11 @@ void AGLTFHotspotActor::ValidateAnimation()
 			{
 				if (SkeletalMesh->Skeleton != nullptr)
 				{
-					UE_LOG(LogEditorGLTFInteractionHotspot, Warning, TEXT("Animation %s is incompatible with skeleton %s, removing animation from actor."), *AnimationSequence->GetName(), *SkeletalMesh->Skeleton->GetName());
+					UE_LOG(LogGLTFHotspot, Warning, TEXT("Animation %s is incompatible with skeleton %s, removing animation from actor."), *AnimationSequence->GetName(), *SkeletalMesh->Skeleton->GetName());
 				}
 				else
 				{
-					UE_LOG(LogEditorGLTFInteractionHotspot, Warning, TEXT("Animation %s is incompatible because mesh %s has no skeleton, removing animation from actor."), *AnimationSequence->GetName(), *SkeletalMesh->GetName());
+					UE_LOG(LogGLTFHotspot, Warning, TEXT("Animation %s is incompatible because mesh %s has no skeleton, removing animation from actor."), *AnimationSequence->GetName(), *SkeletalMesh->GetName());
 				}
 
 				AnimationSequence = nullptr;
