[GLTFExporter] chore: update web viewer build

commit d6037571f620878ce178c199ad8856aee87f0e1a

Date: Tue, 6 Oct 2020 15:28:17 +0200

    chore: update web viewer build
    (cherry picked from commit 1dfd5982151f117350d29d1345e3baa35dd597f5​)

#skipci

diff --git a/Engine/Plugins/Enterprise/GLTFExporter/MigratedFromGit/Resources/WebViewer/index.html b/Engine/Plugins/Enterprise/GLTFExporter/MigratedFromGit/Resources/WebViewer/index.html
--- a/Engine/Plugins/Enterprise/GLTFExporter/MigratedFromGit/Resources/WebViewer/index.html
+++ b/Engine/Plugins/Enterprise/GLTFExporter/MigratedFromGit/Resources/WebViewer/index.html
@@ -30,8 +30,8 @@
   <script src="https:[REDACTED_EMAIL]/lib/mobx.umd.min.js"></script>
   <script src="https:[REDACTED_EMAIL]/prop-types.min.js"></script>
   <script src="./viewer/static/js/runtime-main.c850d8be.js"></script>
-  <script src="./viewer/static/js/2.fd937dd3.chunk.js"></script>
-  <script src="./viewer/static/js/main.76e9e68f.chunk.js"></script>
+  <script src="./viewer/static/js/2.86ba786d.chunk.js"></script>
+  <script src="./viewer/static/js/main.221b1708.chunk.js"></script>
 </body>
 
 </html>
\ No newline at end of file
diff --git a/Engine/Plugins/Enterprise/GLTFExporter/MigratedFromGit/Resources/WebViewer/viewer/static/js/2.86ba786d.chunk.js.LICENSE.txt b/Engine/Plugins/Enterprise/GLTFExporter/MigratedFromGit/Resources/WebViewer/viewer/static/js/2.86ba786d.chunk.js.LICENSE.txt
new file mode 100644
--- /dev/null
+++ b/Engine/Plugins/Enterprise/GLTFExporter/MigratedFromGit/Resources/WebViewer/viewer/static/js/2.86ba786d.chunk.js.LICENSE.txt
@@ -0,0 +1,50 @@
+/*!
+ * The buffer module from node.js, for the browser.
+ *
+ * @author   Feross Aboukhadijeh <http://feross.org>
+ * @license  MIT
+ */
+
+/*! *****************************************************************************
+Copyright (c) Microsoft Corporation. All rights reserved.
+Licensed under the Apache License, Version 2.0 (the "License"); you may not use
+this file except in compliance with the License. You may obtain a copy of the
+License at http://www.apache.org/licenses/LICENSE-2.0
+
+THIS CODE IS PROVIDED ON AN *AS IS* BASIS, WITHOUT WARRANTIES OR CONDITIONS OF ANY
+KIND, EITHER EXPRESS OR IMPLIED, INCLUDING WITHOUT LIMITATION ANY IMPLIED
+WARRANTIES OR CONDITIONS OF TITLE, FITNESS FOR A PARTICULAR PURPOSE,
+MERCHANTABLITY OR NON-INFRINGEMENT.
+
+See the Apache Version 2.0 License for specific language governing permissions
+and limitations under the License.
+***************************************************************************** */
+
+/*! Hammer.JS - v2.0.7 - 2016-04-22
+ * http://hammerjs.github.io/
+ *
+ * Copyright (c) 2016 Jorik Tangelder;
+ * Licensed under the MIT license */
+
+/**
+ * @license
+ * PlayCanvas Engine v1.36.0-rc.2 revision 798f16872
+ * Copyright 2011-2020 PlayCanvas Ltd. All rights reserved.
+ */
+
+/**
+ * A better abstraction over CSS.
+ *
+ * @copyright Oleg Isonen (Slobodskoi) / Isonen 2014-present
+ * @website [REDACTED_PROFILE_URL]/jss
+ * @license MIT
+ */
+
+/** @license React v16.13.0
+ * react-is.production.min.js
+ *
+ * Copyright (c) Facebook, Inc. and its affiliates.
+ *
+ * This source code is licensed under the MIT license found in the
+ * LICENSE file in the root directory of this source tree.
+ */
diff --git a/Engine/Plugins/Enterprise/GLTFExporter/MigratedFromGit/Resources/WebViewer/viewer/static/js/2.fd937dd3.chunk.js b/Engine/Plugins/Enterprise/GLTFExporter/MigratedFromGit/Resources/WebViewer/viewer/static/js/2.fd937dd3.chunk.js
deleted file mode 100644
diff --git a/Engine/Plugins/Enterprise/GLTFExporter/MigratedFromGit/Resources/WebViewer/viewer/static/js/2.fd937dd3.chunk.js.LICENSE.txt b/Engine/Plugins/Enterprise/GLTFExporter/MigratedFromGit/Resources/WebViewer/viewer/static/js/2.fd937dd3.chunk.js.LICENSE.txt
deleted file mode 100644
diff --git a/Engine/Plugins/Enterprise/GLTFExporter/MigratedFromGit/Resources/WebViewer/viewer/static/js/main.76e9e68f.chunk.js b/Engine/Plugins/Enterprise/GLTFExporter/MigratedFromGit/Resources/WebViewer/viewer/static/js/main.76e9e68f.chunk.js
deleted file mode 100644
