[GLTFExporter] feat: add and implement UGLTFExportOptions::bExportPlaybackSettings

commit aaa1b207f28081faff5eb346f2b0f05e51982f9e
Date: Tue, 6 Oct 2020 18:04:57 +0200

    feat: add and implement UGLTFExportOptions::bExportPlaybackSettings
    (cherry picked from commit bc845a7756f14d943e8bf688107e5d44e49e264b​)

#skipci

diff --git a/Engine/Plugins/Enterprise/GLTFExporter/MigratedFromGit/Source/GLTFExporter/Private/Converters/GLTFAnimationConverters.cpp b/Engine/Plugins/Enterprise/GLTFExporter/MigratedFromGit/Source/GLTFExporter/Private/Converters/GLTFAnimationConverters.cpp
--- a/Engine/Plugins/Enterprise/GLTFExporter/MigratedFromGit/Source/GLTFExporter/Private/Converters/GLTFAnimationConverters.cpp
+++ b/Engine/Plugins/Enterprise/GLTFExporter/MigratedFromGit/Source/GLTFExporter/Private/Converters/GLTFAnimationConverters.cpp
@@ -42,7 +42,7 @@ FGLTFJsonAnimationIndex FGLTFAnimationDataConverter::Convert(FGLTFJsonNodeIndex
 	}
 
 	const FGLTFJsonAnimationIndex AnimationIndex = Builder.GetOrAddAnimation(RootNode, SkeletalMesh, AnimSequence);
-	if (AnimationIndex != INDEX_NONE)
+	if (AnimationIndex != INDEX_NONE && Builder.ExportOptions->bExportPlaybackSettings)
 	{
 		FGLTFJsonAnimation& JsonAnimation = Builder.GetAnimation(AnimationIndex);
 		FGLTFJsonAnimationPlayback& JsonPlayback = JsonAnimation.Playback;
@@ -74,7 +74,7 @@ FGLTFJsonAnimationIndex FGLTFLevelSequenceDataConverter::Convert(const ALevelSeq
 	}
 
 	const FGLTFJsonAnimationIndex AnimationIndex = Builder.GetOrAddAnimation(Level, LevelSequence);
-	if (AnimationIndex != INDEX_NONE)
+	if (AnimationIndex != INDEX_NONE && Builder.ExportOptions->bExportPlaybackSettings)
 	{
 		FGLTFJsonAnimation& JsonAnimation = Builder.GetAnimation(AnimationIndex);
 		FGLTFJsonAnimationPlayback& JsonPlayback = JsonAnimation.Playback;
diff --git a/Engine/Plugins/Enterprise/GLTFExporter/MigratedFromGit/Source/GLTFExporter/Private/GLTFExportOptions.cpp b/Engine/Plugins/Enterprise/GLTFExporter/MigratedFromGit/Source/GLTFExporter/Private/GLTFExportOptions.cpp
--- a/Engine/Plugins/Enterprise/GLTFExporter/MigratedFromGit/Source/GLTFExporter/Private/GLTFExportOptions.cpp
+++ b/Engine/Plugins/Enterprise/GLTFExporter/MigratedFromGit/Source/GLTFExporter/Private/GLTFExportOptions.cpp
@@ -21,10 +21,11 @@ UGLTFExportOptions::UGLTFExportOptions(const FObjectInitializer& ObjectInitializ
 	bQuantizeVertexNormals = true;
 	bQuantizeVertexTangents = true;
 	DefaultLevelOfDetail = 0;
+	bExportLevelSequences = true;
 	bExportVertexSkinWeights = true;
 	bExportAnimationSequences = true;
 	bRetargetBoneTransforms = true;
-	bExportLevelSequences = true;
+	bExportPlaybackSettings = true;
 	bExportSourceTextures = true;
 	bExportLightmaps = true;
 	TextureHDREncoding = EGLTFExporterTextureHDREncoding::RGBM;
diff --git a/Engine/Plugins/Enterprise/GLTFExporter/MigratedFromGit/Source/GLTFExporter/Public/GLTFExportOptions.h b/Engine/Plugins/Enterprise/GLTFExporter/MigratedFromGit/Source/GLTFExporter/Public/GLTFExportOptions.h
--- a/Engine/Plugins/Enterprise/GLTFExporter/MigratedFromGit/Source/GLTFExporter/Public/GLTFExportOptions.h
+++ b/Engine/Plugins/Enterprise/GLTFExporter/MigratedFromGit/Source/GLTFExporter/Public/GLTFExportOptions.h
@@ -73,6 +73,9 @@ public:
 	UPROPERTY(EditAnywhere, BlueprintReadWrite, config, Category = Mesh)
 	int32 DefaultLevelOfDetail;
 
+	UPROPERTY(EditAnywhere, BlueprintReadWrite, config, Category = Animation)
+	uint32 bExportLevelSequences : 1;
+
 	UPROPERTY(EditAnywhere, BlueprintReadWrite, config, Category = Animation)
 	uint32 bExportVertexSkinWeights : 1;
 
@@ -83,7 +86,7 @@ public:
 	uint32 bRetargetBoneTransforms : 1;
 
 	UPROPERTY(EditAnywhere, BlueprintReadWrite, config, Category = Animation)
-	uint32 bExportLevelSequences : 1;
+	uint32 bExportPlaybackSettings : 1;
 
 	UPROPERTY(EditAnywhere, BlueprintReadWrite, config, Category = Texture)
 	uint32 bExportSourceTextures : 1;
