[GLTFExporter] chore: bump plugin version to 0.8

commit 67b46ba8778d96ca641150486aea0a4150cd5c79
Date: Tue, 6 Oct 2020 18:25:46 +0200

    chore: bump plugin version to 0.8
    (cherry picked from commit 1a9522b88e16c6d5681dd1e563fc57cf91d6aa70​)

#skipci

diff --git a/Engine/Plugins/Enterprise/GLTFExporter/MigratedFromGit/GLTFExporter.uplugin b/Engine/Plugins/Enterprise/GLTFExporter/MigratedFromGit/GLTFExporter.uplugin
--- a/Engine/Plugins/Enterprise/GLTFExporter/MigratedFromGit/GLTFExporter.uplugin
+++ b/Engine/Plugins/Enterprise/GLTFExporter/MigratedFromGit/GLTFExporter.uplugin
@@ -1,7 +1,7 @@
 {
 	"FileVersion": 3,
 	"Version": 5,
-	"VersionName": "0.5",
+	"VersionName": "0.8",
 	"FriendlyName": "glTF Exporter",
 	"Description": "An exporter for Khronos glTF 3D models.",
 	"Category": "Exporters",
