[GLTFExporter] fix: add special gltf default material to avoid unnecessary of export of WorldGridMaterial when converting backdrop mesh

commit 1598e640e7fb0513db3d9a9ddbe21b98c3d9648d
Date: Sat, 10 Oct 2020 15:07:56 +0200

    fix: add special gltf default material to avoid unnecessary of export of WorldGridMaterial when converting backdrop mesh
    (cherry picked from commit dd435652e52d9b57064ece51f9dee390a2c907f3​)

#skipci

diff --git a/Engine/Plugins/Enterprise/GLTFExporter/MigratedFromGit/Source/GLTFExporter/Private/Converters/GLTFBackdropConverters.cpp b/Engine/Plugins/Enterprise/GLTFExporter/MigratedFromGit/Source/GLTFExporter/Private/Converters/GLTFBackdropConverters.cpp
--- a/Engine/Plugins/Enterprise/GLTFExporter/MigratedFromGit/Source/GLTFExporter/Private/Converters/GLTFBackdropConverters.cpp
+++ b/Engine/Plugins/Enterprise/GLTFExporter/MigratedFromGit/Source/GLTFExporter/Private/Converters/GLTFBackdropConverters.cpp
@@ -2,6 +2,7 @@
 
 #include "Converters/GLTFBackdropConverters.h"
 #include "Converters/GLTFConverterUtility.h"
+#include "Converters/GLTFMaterialUtility.h"
 #include "Converters/GLTFActorUtility.h"
 #include "Builders/GLTFContainerBuilder.h"
 
@@ -20,7 +21,7 @@ FGLTFJsonBackdropIndex FGLTFBackdropConverter::Convert(const AActor* BackdropAct
 	if (FGLTFActorUtility::TryGetPropertyValue(BackdropActor, TEXT("Mesh"), Mesh))
 	{
 		// TODO: avoid exporting the mesh material (use gltf default material)
-		JsonBackdrop.Mesh = Builder.GetOrAddMesh(Mesh);
+		JsonBackdrop.Mesh = Builder.GetOrAddMesh(Mesh, -1, nullptr, { FGLTFMaterialUtility::GetDefault() });
 	}
 	else
 	{
diff --git a/Engine/Plugins/Enterprise/GLTFExporter/MigratedFromGit/Source/GLTFExporter/Private/Converters/GLTFMaterialConverters.cpp b/Engine/Plugins/Enterprise/GLTFExporter/MigratedFromGit/Source/GLTFExporter/Private/Converters/GLTFMaterialConverters.cpp
--- a/Engine/Plugins/Enterprise/GLTFExporter/MigratedFromGit/Source/GLTFExporter/Private/Converters/GLTFMaterialConverters.cpp
+++ b/Engine/Plugins/Enterprise/GLTFExporter/MigratedFromGit/Source/GLTFExporter/Private/Converters/GLTFMaterialConverters.cpp
@@ -1,11 +1,17 @@
 // Copyright Epic Games, Inc. All Rights Reserved.
 
 #include "Converters/GLTFMaterialConverters.h"
+#include "Converters/GLTFMaterialUtility.h"
 #include "Builders/GLTFConvertBuilder.h"
 #include "Tasks/GLTFMaterialTasks.h"
 
 FGLTFJsonMaterialIndex FGLTFMaterialConverter::Convert(const UMaterialInterface* Material)
 {
+	if (Material == FGLTFMaterialUtility::GetDefault())
+	{
+		return FGLTFJsonMaterialIndex(INDEX_NONE); // use default gltf definition
+	}
+
 	const FGLTFJsonMaterialIndex MaterialIndex = Builder.AddMaterial();
 	Builder.SetupTask<FGLTFMaterialTask>(Builder, Material, MaterialIndex);
 	return MaterialIndex;
diff --git a/Engine/Plugins/Enterprise/GLTFExporter/MigratedFromGit/Source/GLTFExporter/Private/Converters/GLTFMaterialUtility.cpp b/Engine/Plugins/Enterprise/GLTFExporter/MigratedFromGit/Source/GLTFExporter/Private/Converters/GLTFMaterialUtility.cpp
--- a/Engine/Plugins/Enterprise/GLTFExporter/MigratedFromGit/Source/GLTFExporter/Private/Converters/GLTFMaterialUtility.cpp
+++ b/Engine/Plugins/Enterprise/GLTFExporter/MigratedFromGit/Source/GLTFExporter/Private/Converters/GLTFMaterialUtility.cpp
@@ -12,6 +12,12 @@
 #include "Materials/MaterialExpressionClearCoatNormalCustomOutput.h"
 #include "Materials/MaterialExpressionTextureCoordinate.h"
 
+UMaterialInterface* FGLTFMaterialUtility::GetDefault()
+{
+	static UMaterialInterface* DefaultMaterial = LoadObject<UMaterialInterface>(nullptr, TEXT("/GLTFExporter/Materials/Default.Default"));
+	return DefaultMaterial;
+}
+
 const TCHAR* FGLTFMaterialUtility::GetPropertyName(EMaterialProperty Property)
 {
 	// TODO: replace this hardcoded list with some lookup function that should be in the engine.
diff --git a/Engine/Plugins/Enterprise/GLTFExporter/MigratedFromGit/Source/GLTFExporter/Private/Converters/GLTFMaterialUtility.h b/Engine/Plugins/Enterprise/GLTFExporter/MigratedFromGit/Source/GLTFExporter/Private/Converters/GLTFMaterialUtility.h
--- a/Engine/Plugins/Enterprise/GLTFExporter/MigratedFromGit/Source/GLTFExporter/Private/Converters/GLTFMaterialUtility.h
+++ b/Engine/Plugins/Enterprise/GLTFExporter/MigratedFromGit/Source/GLTFExporter/Private/Converters/GLTFMaterialUtility.h
@@ -34,6 +34,8 @@ struct FGLTFPropertyBakeOutput
 
 struct FGLTFMaterialUtility
 {
+	static UMaterialInterface* GetDefault();
+
 	static const TCHAR* GetPropertyName(EMaterialProperty Property);
 	static FVector4 GetPropertyDefaultValue(EMaterialProperty Property);
 	static FVector4 GetPropertyMask(EMaterialProperty Property);
