[GLTFExporter] refactor: when converting meshes clamp LODIndex if out of range

commit 4725fc3956b04c963d3be550cc21e22f2722baad
Date: Sat, 10 Oct 2020 15:20:13 +0200

    refactor: when converting meshes clamp LODIndex if out of range
    (cherry picked from commit 7d8d006fec6637f81748470bd1934ce602fbb594​)

#skipci

diff --git a/Engine/Plugins/Enterprise/GLTFExporter/MigratedFromGit/Source/GLTFExporter/Private/Builders/GLTFConvertBuilder.cpp b/Engine/Plugins/Enterprise/GLTFExporter/MigratedFromGit/Source/GLTFExporter/Private/Builders/GLTFConvertBuilder.cpp
--- a/Engine/Plugins/Enterprise/GLTFExporter/MigratedFromGit/Source/GLTFExporter/Private/Builders/GLTFConvertBuilder.cpp
+++ b/Engine/Plugins/Enterprise/GLTFExporter/MigratedFromGit/Source/GLTFExporter/Private/Builders/GLTFConvertBuilder.cpp
@@ -106,7 +106,7 @@ FGLTFJsonMeshIndex FGLTFConvertBuilder::GetOrAddMesh(const UStaticMeshComponent*
 	}
 
 	const UStaticMesh* StaticMesh = StaticMeshComponent->GetStaticMesh();
-	const int32 LODIndex = StaticMeshComponent->ForcedLodModel > 0 ? StaticMeshComponent->ForcedLodModel - 1 : ExportOptions->DefaultLevelOfDetail;
+	const int32 LODIndex = StaticMeshComponent->ForcedLodModel > 0 ? StaticMeshComponent->ForcedLodModel - 1 : -1;
 	const FColorVertexBuffer* OverrideVertexColors = StaticMeshComponent->LODData.IsValidIndex(LODIndex) ? StaticMeshComponent->LODData[LODIndex].OverrideVertexColors : nullptr;
 	const FGLTFMaterialArray OverrideMaterials = FGLTFMaterialArray(StaticMeshComponent->OverrideMaterials);
 
@@ -131,7 +131,7 @@ FGLTFJsonMeshIndex FGLTFConvertBuilder::GetOrAddMesh(const USkeletalMeshComponen
 	}
 
 	const USkeletalMesh* SkeletalMesh = SkeletalMeshComponent->SkeletalMesh;
-	const int32 LODIndex = SkeletalMeshComponent->GetForcedLOD() > 0 ? SkeletalMeshComponent->GetForcedLOD() - 1 : ExportOptions->DefaultLevelOfDetail;
+	const int32 LODIndex = SkeletalMeshComponent->GetForcedLOD() > 0 ? SkeletalMeshComponent->GetForcedLOD() - 1 : -1;
 	const FColorVertexBuffer* OverrideVertexColors = SkeletalMeshComponent->LODInfo.IsValidIndex(LODIndex) ? SkeletalMeshComponent->LODInfo[LODIndex].OverrideVertexColors : nullptr;
 	const FSkinWeightVertexBuffer* OverrideSkinWeights = SkeletalMeshComponent->LODInfo.IsValidIndex(LODIndex) ? SkeletalMeshComponent->LODInfo[LODIndex].OverrideSkinWeights : nullptr;
 	const FGLTFMaterialArray OverrideMaterials = FGLTFMaterialArray(SkeletalMeshComponent->OverrideMaterials);
diff --git a/Engine/Plugins/Enterprise/GLTFExporter/MigratedFromGit/Source/GLTFExporter/Private/Converters/GLTFMeshConverters.cpp b/Engine/Plugins/Enterprise/GLTFExporter/MigratedFromGit/Source/GLTFExporter/Private/Converters/GLTFMeshConverters.cpp
--- a/Engine/Plugins/Enterprise/GLTFExporter/MigratedFromGit/Source/GLTFExporter/Private/Converters/GLTFMeshConverters.cpp
+++ b/Engine/Plugins/Enterprise/GLTFExporter/MigratedFromGit/Source/GLTFExporter/Private/Converters/GLTFMeshConverters.cpp
@@ -22,23 +22,17 @@ void FGLTFStaticMeshConverter::Sanitize(const UStaticMesh*& StaticMesh, int32& L
 		OverrideMaterials.Empty();
 	}
 
-	if (StaticMesh->GetNumLODs() > LODIndex)
+	LODIndex = FMath::Clamp(LODIndex, 0, StaticMesh->GetNumLODs() - 1);
+	const FStaticMeshLODResources& MeshLOD = StaticMesh->GetLODForExport(LODIndex);
+
+	if (OverrideVertexColors == &MeshLOD.VertexBuffers.ColorVertexBuffer)
 	{
-		const FStaticMeshLODResources& MeshLOD = StaticMesh->GetLODForExport(LODIndex);
-		if (OverrideVertexColors == &MeshLOD.VertexBuffers.ColorVertexBuffer)
-		{
-			OverrideVertexColors = nullptr;
-		}
+		OverrideVertexColors = nullptr;
 	}
 }
 
 FGLTFJsonMeshIndex FGLTFStaticMeshConverter::Convert(const UStaticMesh* StaticMesh, int32 LODIndex, const FColorVertexBuffer* OverrideVertexColors, FGLTFMaterialArray OverrideMaterials)
 {
-	if (StaticMesh->GetNumLODs() <= LODIndex)
-	{
-		return FGLTFJsonMeshIndex(INDEX_NONE);
-	}
-
 	const FGLTFJsonMeshIndex MeshIndex = Builder.AddMesh();
 	Builder.SetupTask<FGLTFStaticMeshTask>(Builder, MeshSectionConverter, StaticMesh, LODIndex, OverrideVertexColors, OverrideMaterials, MeshIndex);
 	return MeshIndex;
@@ -67,31 +61,22 @@ void FGLTFSkeletalMeshConverter::Sanitize(const USkeletalMesh*& SkeletalMesh, in
 	}
 
 	const FSkeletalMeshRenderData* RenderData = SkeletalMesh->GetResourceForRendering();
-	if (RenderData->LODRenderData.Num() > LODIndex)
-	{
-		const FSkeletalMeshLODRenderData& MeshLOD = RenderData->LODRenderData[LODIndex];
+	LODIndex = FMath::Clamp(LODIndex, 0, RenderData->LODRenderData.Num() - 1);
+	const FSkeletalMeshLODRenderData& MeshLOD = RenderData->LODRenderData[LODIndex];
 
-		if (OverrideVertexColors == &MeshLOD.StaticVertexBuffers.ColorVertexBuffer)
-		{
-			OverrideVertexColors = nullptr;
-		}
+	if (OverrideVertexColors == &MeshLOD.StaticVertexBuffers.ColorVertexBuffer)
+	{
+		OverrideVertexColors = nullptr;
+	}
 
-		if (OverrideSkinWeights == MeshLOD.GetSkinWeightVertexBuffer())
-		{
-			OverrideSkinWeights = nullptr;
-		}
+	if (OverrideSkinWeights == MeshLOD.GetSkinWeightVertexBuffer())
+	{
+		OverrideSkinWeights = nullptr;
 	}
 }
 
 FGLTFJsonMeshIndex FGLTFSkeletalMeshConverter::Convert(const USkeletalMesh* SkeletalMesh, int32 LODIndex, const FColorVertexBuffer* OverrideVertexColors, const FSkinWeightVertexBuffer* OverrideSkinWeights, FGLTFMaterialArray OverrideMaterials)
 {
-	const FSkeletalMeshRenderData* RenderData = SkeletalMesh->GetResourceForRendering();
-
-	if (RenderData->LODRenderData.Num() <= LODIndex)
-	{
-		return FGLTFJsonMeshIndex(INDEX_NONE);
-	}
-
 	const FGLTFJsonMeshIndex MeshIndex = Builder.AddMesh();
 	Builder.SetupTask<FGLTFSkeletalMeshTask>(Builder, MeshSectionConverter, SkeletalMesh, LODIndex, OverrideVertexColors, OverrideSkinWeights, OverrideMaterials, MeshIndex);
 	return MeshIndex;
