[GLTFExporter] refactor: move test plugin from Plugins/Tests/GLTFExporterTests to Plugins/GLTFExporterTests and test definitions from GLTFExporterTests/Resources/Tests to GLTFExporterTests/Tests

commit 74f4ec3864be2afe09ccb7630d6cecf5a6c03826
Date: Sat, 10 Oct 2020 15:55:36 +0200

    refactor: move test plugin from Plugins/Tests/GLTFExporterTests to Plugins/GLTFExporterTests and test definitions from GLTFExporterTests/Resources/Tests to GLTFExporterTests/Tests
    (cherry picked from commit f2656bed96198bcfa7880a5b7f55f3c7d9a20256​)

#skipci

diff --git a/Engine/Plugins/Enterprise/GLTFExporter/MigratedFromGit/Config/FilterPlugin.ini b/Engine/Plugins/Enterprise/GLTFExporter/MigratedFromGit/Config/FilterPlugin.ini
--- a/Engine/Plugins/Enterprise/GLTFExporter/MigratedFromGit/Config/FilterPlugin.ini
+++ b/Engine/Plugins/Enterprise/GLTFExporter/MigratedFromGit/Config/FilterPlugin.ini
@@ -1,10 +1,2 @@
 [FilterPlugin]
-; This section lists additional files which will be packaged along with your plugin. Paths should be listed relative to the root plugin directory, and
-; may include "...", "*", and "?" wildcards to match directories, files, and individual characters respectively.
-;
-; Examples:
-;    /README.txt
-;    /Extras/...
-;    /Binaries/ThirdParty/*.dll
-
 /README.md
