[GLTFExporter] feat: use minimum LOD level (specified per mesh asset or mesh component) if higher than default level

commit 4a4c3aa320339b777640278a0a10420f20645196
Date: Sat, 10 Oct 2020 16:52:43 +0200

    feat: use minimum LOD level (specified per mesh asset or mesh component) if higher than default level
    (cherry picked from commit a9bdf6b16204a0100068caf1482ff62f0ccfd23e​)

#skipci

diff --git a/Engine/Plugins/Enterprise/GLTFExporter/MigratedFromGit/Source/GLTFExporter/GLTFExporter.Build.cs b/Engine/Plugins/Enterprise/GLTFExporter/MigratedFromGit/Source/GLTFExporter/GLTFExporter.Build.cs
--- a/Engine/Plugins/Enterprise/GLTFExporter/MigratedFromGit/Source/GLTFExporter/GLTFExporter.Build.cs
+++ b/Engine/Plugins/Enterprise/GLTFExporter/MigratedFromGit/Source/GLTFExporter/GLTFExporter.Build.cs
@@ -27,6 +27,7 @@ namespace UnrealBuildTool.Rules
 					"Projects",
 					"RenderCore",
 					"RHI",
+					"DesktopPlatform",
 					"LevelSequence",
 					"MovieScene",
 					"MovieSceneTracks",
diff --git a/Engine/Plugins/Enterprise/GLTFExporter/MigratedFromGit/Source/GLTFExporter/Private/Builders/GLTFConvertBuilder.cpp b/Engine/Plugins/Enterprise/GLTFExporter/MigratedFromGit/Source/GLTFExporter/Private/Builders/GLTFConvertBuilder.cpp
--- a/Engine/Plugins/Enterprise/GLTFExporter/MigratedFromGit/Source/GLTFExporter/Private/Builders/GLTFConvertBuilder.cpp
+++ b/Engine/Plugins/Enterprise/GLTFExporter/MigratedFromGit/Source/GLTFExporter/Private/Builders/GLTFConvertBuilder.cpp
@@ -106,7 +106,7 @@ FGLTFJsonMeshIndex FGLTFConvertBuilder::GetOrAddMesh(const UStaticMeshComponent*
 	}
 
 	const UStaticMesh* StaticMesh = StaticMeshComponent->GetStaticMesh();
-	const int32 LODIndex = StaticMeshComponent->ForcedLodModel > 0 ? StaticMeshComponent->ForcedLodModel - 1 : -1;
+	const int32 LODIndex = StaticMeshComponent->ForcedLodModel > 0 ? StaticMeshComponent->ForcedLodModel - 1 : FMath::Max(ExportOptions->DefaultLevelOfDetail, StaticMeshComponent->MinLOD);
 	const FColorVertexBuffer* OverrideVertexColors = StaticMeshComponent->LODData.IsValidIndex(LODIndex) ? StaticMeshComponent->LODData[LODIndex].OverrideVertexColors : nullptr;
 	const FGLTFMaterialArray OverrideMaterials = FGLTFMaterialArray(StaticMeshComponent->OverrideMaterials);
 
@@ -131,7 +131,7 @@ FGLTFJsonMeshIndex FGLTFConvertBuilder::GetOrAddMesh(const USkeletalMeshComponen
 	}
 
 	const USkeletalMesh* SkeletalMesh = SkeletalMeshComponent->SkeletalMesh;
-	const int32 LODIndex = SkeletalMeshComponent->GetForcedLOD() > 0 ? SkeletalMeshComponent->GetForcedLOD() - 1 : -1;
+	const int32 LODIndex = SkeletalMeshComponent->GetForcedLOD() > 0 ? SkeletalMeshComponent->GetForcedLOD() - 1 : FMath::Max(ExportOptions->DefaultLevelOfDetail, SkeletalMeshComponent->MinLodModel);
 	const FColorVertexBuffer* OverrideVertexColors = SkeletalMeshComponent->LODInfo.IsValidIndex(LODIndex) ? SkeletalMeshComponent->LODInfo[LODIndex].OverrideVertexColors : nullptr;
 	const FSkinWeightVertexBuffer* OverrideSkinWeights = SkeletalMeshComponent->LODInfo.IsValidIndex(LODIndex) ? SkeletalMeshComponent->LODInfo[LODIndex].OverrideSkinWeights : nullptr;
 	const FGLTFMaterialArray OverrideMaterials = FGLTFMaterialArray(SkeletalMeshComponent->OverrideMaterials);
diff --git a/Engine/Plugins/Enterprise/GLTFExporter/MigratedFromGit/Source/GLTFExporter/Private/Converters/GLTFMeshConverters.cpp b/Engine/Plugins/Enterprise/GLTFExporter/MigratedFromGit/Source/GLTFExporter/Private/Converters/GLTFMeshConverters.cpp
--- a/Engine/Plugins/Enterprise/GLTFExporter/MigratedFromGit/Source/GLTFExporter/Private/Converters/GLTFMeshConverters.cpp
+++ b/Engine/Plugins/Enterprise/GLTFExporter/MigratedFromGit/Source/GLTFExporter/Private/Converters/GLTFMeshConverters.cpp
@@ -1,6 +1,7 @@
 // Copyright Epic Games, Inc. All Rights Reserved.
 
 #include "Converters/GLTFMeshConverters.h"
+#include "Converters/GLTFMeshUtility.h"
 #include "Builders/GLTFConvertBuilder.h"
 #include "Tasks/GLTFMeshTasks.h"
 #include "Rendering/SkeletalMeshRenderData.h"
@@ -9,7 +10,7 @@ void FGLTFStaticMeshConverter::Sanitize(const UStaticMesh*& StaticMesh, int32& L
 {
 	if (LODIndex < 0)
 	{
-		LODIndex = FMath::Max(Builder.ExportOptions->DefaultLevelOfDetail, 0);
+		LODIndex = FMath::Max(Builder.ExportOptions->DefaultLevelOfDetail, FGLTFMeshUtility::GetMinimumLOD(StaticMesh));
 	}
 
 	if (!Builder.ExportOptions->bExportVertexColors)
@@ -42,7 +43,7 @@ void FGLTFSkeletalMeshConverter::Sanitize(const USkeletalMesh*& SkeletalMesh, in
 {
 	if (LODIndex < 0)
 	{
-		LODIndex = FMath::Max(Builder.ExportOptions->DefaultLevelOfDetail, 0);
+		LODIndex = FMath::Max(Builder.ExportOptions->DefaultLevelOfDetail, FGLTFMeshUtility::GetMinimumLOD(SkeletalMesh));
 	}
 
 	if (!Builder.ExportOptions->bExportVertexColors)
diff --git a/Engine/Plugins/Enterprise/GLTFExporter/MigratedFromGit/Source/GLTFExporter/Private/Converters/GLTFMeshUtility.cpp b/Engine/Plugins/Enterprise/GLTFExporter/MigratedFromGit/Source/GLTFExporter/Private/Converters/GLTFMeshUtility.cpp
new file mode 100644
--- /dev/null
+++ b/Engine/Plugins/Enterprise/GLTFExporter/MigratedFromGit/Source/GLTFExporter/Private/Converters/GLTFMeshUtility.cpp
@@ -0,0 +1,21 @@
+// Copyright Epic Games, Inc. All Rights Reserved.
+
+#include "Converters/GLTFMeshUtility.h"
+#include "PlatformInfo.h"
+
+int32 FGLTFMeshUtility::GetMinimumLOD(const UStaticMesh* StaticMesh)
+{
+	return GetValueForRunningPlatform<int32>(StaticMesh->MinLOD);
+}
+
+int32 FGLTFMeshUtility::GetMinimumLOD(const USkeletalMesh* SkeletalMesh)
+{
+	return GetValueForRunningPlatform<int32>(SkeletalMesh->MinLod);
+}
+
+template <typename ValueType, typename StructType>
+ValueType FGLTFMeshUtility::GetValueForRunningPlatform(const StructType& Properties)
+{
+	const PlatformInfo::FPlatformInfo& PlatformInfo = GetTargetPlatformManagerRef().GetRunningTargetPlatform()->GetPlatformInfo();
+	return Properties.GetValueForPlatformIdentifiers(PlatformInfo.PlatformGroupName, PlatformInfo.VanillaPlatformName);
+}
diff --git a/Engine/Plugins/Enterprise/GLTFExporter/MigratedFromGit/Source/GLTFExporter/Private/Converters/GLTFMeshUtility.h b/Engine/Plugins/Enterprise/GLTFExporter/MigratedFromGit/Source/GLTFExporter/Private/Converters/GLTFMeshUtility.h
new file mode 100644
--- /dev/null
+++ b/Engine/Plugins/Enterprise/GLTFExporter/MigratedFromGit/Source/GLTFExporter/Private/Converters/GLTFMeshUtility.h
@@ -0,0 +1,17 @@
+// Copyright Epic Games, Inc. All Rights Reserved.
+
+#pragma once
+
+#include "Engine.h"
+
+struct FGLTFMeshUtility
+{
+	static int32 GetMinimumLOD(const UStaticMesh* StaticMesh);
+
+	static int32 GetMinimumLOD(const USkeletalMesh* SkeletalMesh);
+
+private:
+
+	template <typename ValueType, typename StructType>
+	static ValueType GetValueForRunningPlatform(const StructType& Properties);
+};
