[GLTFExporter] fix: use source dimensions when exporting source data since render data may be smaller due to Maximum Texture Size

commit fe934f389d7d12b0506c7b7cfe761e956611eddc
Date: Sun, 11 Oct 2020 20:14:47 +0200

    fix: use source dimensions when exporting source data since render data may be smaller due to Maximum Texture Size
    (cherry picked from commit c195eb71ab903f0c818697a4266eee181a0a7bdb​)

#skipci

diff --git a/Engine/Plugins/Enterprise/GLTFExporter/MigratedFromGit/Source/GLTFExporter/Private/Tasks/GLTFTextureTasks.cpp b/Engine/Plugins/Enterprise/GLTFExporter/MigratedFromGit/Source/GLTFExporter/Private/Tasks/GLTFTextureTasks.cpp
--- a/Engine/Plugins/Enterprise/GLTFExporter/MigratedFromGit/Source/GLTFExporter/Private/Tasks/GLTFTextureTasks.cpp
+++ b/Engine/Plugins/Enterprise/GLTFExporter/MigratedFromGit/Source/GLTFExporter/Private/Tasks/GLTFTextureTasks.cpp
@@ -12,7 +12,6 @@ void FGLTFTexture2DTask::Complete()
 	FGLTFJsonTexture& JsonTexture = Builder.GetTexture(TextureIndex);
 	Texture2D->GetName(JsonTexture.Name);
 
-	const FIntPoint Size = { Texture2D->GetSizeX(), Texture2D->GetSizeY() };
 	const bool bPreferSourceExport = Builder.ExportOptions->bExportSourceTextures;
 
 	// NOTE: export of lightmaps and normalmaps via source data is done to workaround issues
@@ -28,6 +27,7 @@ void FGLTFTexture2DTask::Complete()
 		uint32 BitDepth;
 
 		FTextureSource& Source = const_cast<FTextureSource&>(Texture2D->Source);
+		const FIntPoint Size = { Source.GetSizeX(), Source.GetSizeY() };
 
 		if (Source.IsValid() && FGLTFTextureUtility::CanPNGCompressFormat(Source.GetFormat(), RGBFormat, BitDepth))
 		{
@@ -54,6 +54,8 @@ void FGLTFTexture2DTask::Complete()
 
 	if (ImageIndex == INDEX_NONE && !bRequireSourceExport)
 	{
+		const FIntPoint Size = { Texture2D->GetSizeX(), Texture2D->GetSizeY() };
+
 		// TODO: both bForceLinearGamma and TargetGamma=2.2f seem to be necessary for the exported images to match the results
 		// from exporting using the texture's source or its platform-data.
 		// It's not entirely clear why gamma must be set to 2.2 instead of 0.0 (which should use the correct gamma anyway),
