[GLTFExporter] feat(sky-sphere): add basic support for exporting sky spheres

commit a735e4b3c2d7414c50dded1122fd4e2b28847b7b
Date: Mon, 12 Oct 2020 17:14:13 +0200

    feat(sky-sphere): add basic support for exporting sky spheres
    (cherry picked from commit 5718b30abe1ea10035ca28839cfe44171224202e​)

#skipci

diff --git a/Engine/Plugins/Enterprise/GLTFExporter/MigratedFromGit/Source/GLTFExporter/Private/Builders/GLTFConvertBuilder.cpp b/Engine/Plugins/Enterprise/GLTFExporter/MigratedFromGit/Source/GLTFExporter/Private/Builders/GLTFConvertBuilder.cpp
--- a/Engine/Plugins/Enterprise/GLTFExporter/MigratedFromGit/Source/GLTFExporter/Private/Builders/GLTFConvertBuilder.cpp
+++ b/Engine/Plugins/Enterprise/GLTFExporter/MigratedFromGit/Source/GLTFExporter/Private/Builders/GLTFConvertBuilder.cpp
@@ -405,3 +405,13 @@ FGLTFJsonHotspotIndex FGLTFConvertBuilder::GetOrAddHotspot(const AGLTFHotspotAct
 
 	return HotspotConverter.GetOrAdd(HotspotActor);
 }
+
+FGLTFJsonSkySphereIndex FGLTFConvertBuilder::GetOrAddSkySphere(const AActor* SkySphereActor)
+{
+	if (SkySphereActor == nullptr)
+	{
+		return FGLTFJsonSkySphereIndex(INDEX_NONE);
+	}
+
+	return SkySphereConverter.GetOrAdd(SkySphereActor);
+}
diff --git a/Engine/Plugins/Enterprise/GLTFExporter/MigratedFromGit/Source/GLTFExporter/Private/Builders/GLTFConvertBuilder.h b/Engine/Plugins/Enterprise/GLTFExporter/MigratedFromGit/Source/GLTFExporter/Private/Builders/GLTFConvertBuilder.h
--- a/Engine/Plugins/Enterprise/GLTFExporter/MigratedFromGit/Source/GLTFExporter/Private/Builders/GLTFConvertBuilder.h
+++ b/Engine/Plugins/Enterprise/GLTFExporter/MigratedFromGit/Source/GLTFExporter/Private/Builders/GLTFConvertBuilder.h
@@ -18,6 +18,7 @@
 #include "Converters/GLTFVarationConverters.h"
 #include "Converters/GLTFLightMapConverters.h"
 #include "Converters/GLTFHotspotConverters.h"
+#include "Converters/GLTFSkySphereConverters.h"
 
 class FGLTFConvertBuilder : public FGLTFImageBuilder
 {
@@ -72,6 +73,7 @@ public:
 	FGLTFJsonVariationIndex GetOrAddVariation(const ALevelVariantSetsActor* LevelVariantSetsActor);
 	FGLTFJsonLightMapIndex GetOrAddLightMap(const UStaticMeshComponent* StaticMeshComponent);
 	FGLTFJsonHotspotIndex GetOrAddHotspot(const AGLTFHotspotActor* HotspotActor);
+	FGLTFJsonSkySphereIndex GetOrAddSkySphere(const AActor* SkySphereActor);
 
 private:
 
@@ -114,4 +116,5 @@ private:
 	FGLTFVariationConverter VariationConverter{ *this };
 	FGLTFLightMapConverter LightMapConverter{ *this };
 	FGLTFHotspotConverter HotspotConverter{ *this };
+	FGLTFSkySphereConverter SkySphereConverter{ *this };
 };
diff --git a/Engine/Plugins/Enterprise/GLTFExporter/MigratedFromGit/Source/GLTFExporter/Private/Builders/GLTFJsonBuilder.cpp b/Engine/Plugins/Enterprise/GLTFExporter/MigratedFromGit/Source/GLTFExporter/Private/Builders/GLTFJsonBuilder.cpp
--- a/Engine/Plugins/Enterprise/GLTFExporter/MigratedFromGit/Source/GLTFExporter/Private/Builders/GLTFJsonBuilder.cpp
+++ b/Engine/Plugins/Enterprise/GLTFExporter/MigratedFromGit/Source/GLTFExporter/Private/Builders/GLTFJsonBuilder.cpp
@@ -112,6 +112,11 @@ FGLTFJsonHotspotIndex FGLTFJsonBuilder::AddHotspot(const FGLTFJsonHotspot& JsonH
 	return FGLTFJsonHotspotIndex(JsonRoot.Hotspots.Add(MakeUnique<FGLTFJsonHotspot>(JsonHotspot)));
 }
 
+FGLTFJsonSkySphereIndex FGLTFJsonBuilder::AddSkySphere(const FGLTFJsonSkySphere& JsonSkySphere)
+{
+	return FGLTFJsonSkySphereIndex(JsonRoot.SkySpheres.Add(MakeUnique<FGLTFJsonSkySphere>(JsonSkySphere)));
+}
+
 FGLTFJsonNodeIndex FGLTFJsonBuilder::AddChildNode(FGLTFJsonNodeIndex ParentIndex, const FGLTFJsonNode& JsonNode)
 {
 	const FGLTFJsonNodeIndex ChildIndex = AddNode(JsonNode);
@@ -226,6 +231,11 @@ FGLTFJsonHotspot& FGLTFJsonBuilder::GetHotspot(FGLTFJsonHotspotIndex HotspotInde
 	return *JsonRoot.Hotspots[HotspotIndex];
 }
 
+FGLTFJsonSkySphere& FGLTFJsonBuilder::GetSkySphere(FGLTFJsonSkySphereIndex SkySphereIndex)
+{
+	return *JsonRoot.SkySpheres[SkySphereIndex];
+}
+
 FGLTFJsonNodeIndex FGLTFJsonBuilder::GetComponentNodeIndex(FGLTFJsonNodeIndex NodeIndex)
 {
 	if (NodeIndex == INDEX_NONE)
diff --git a/Engine/Plugins/Enterprise/GLTFExporter/MigratedFromGit/Source/GLTFExporter/Private/Builders/GLTFJsonBuilder.h b/Engine/Plugins/Enterprise/GLTFExporter/MigratedFromGit/Source/GLTFExporter/Private/Builders/GLTFJsonBuilder.h
--- a/Engine/Plugins/Enterprise/GLTFExporter/MigratedFromGit/Source/GLTFExporter/Private/Builders/GLTFJsonBuilder.h
+++ b/Engine/Plugins/Enterprise/GLTFExporter/MigratedFromGit/Source/GLTFExporter/Private/Builders/GLTFJsonBuilder.h
@@ -37,6 +37,7 @@ public:
 	FGLTFJsonLightMapIndex AddLightMap(const FGLTFJsonLightMap& JsonLightMap = {});
 	FGLTFJsonLightIndex AddLight(const FGLTFJsonLight& JsonLight = {});
 	FGLTFJsonHotspotIndex AddHotspot(const FGLTFJsonHotspot& JsonHotspot = {});
+	FGLTFJsonSkySphereIndex AddSkySphere(const FGLTFJsonSkySphere& JsonSkySphere = {});
 
 	FGLTFJsonNodeIndex AddChildNode(FGLTFJsonNodeIndex ParentNodeIndex, const FGLTFJsonNode& JsonNode = {});
 	FGLTFJsonNodeIndex AddChildComponentNode(FGLTFJsonNodeIndex ParentNodeIndex, const FGLTFJsonNode& JsonNode = {});
@@ -59,6 +60,7 @@ public:
 	FGLTFJsonLightMap& GetLightMap(FGLTFJsonLightMapIndex LightMapIndex);
 	FGLTFJsonLight& GetLight(FGLTFJsonLightIndex LightIndex);
 	FGLTFJsonHotspot& GetHotspot(FGLTFJsonHotspotIndex HotspotIndex);
+	FGLTFJsonSkySphere& GetSkySphere(FGLTFJsonSkySphereIndex SkySphereIndex);
 
 	FGLTFJsonNodeIndex GetComponentNodeIndex(FGLTFJsonNodeIndex NodeIndex);
 
diff --git a/Engine/Plugins/Enterprise/GLTFExporter/MigratedFromGit/Source/GLTFExporter/Private/Converters/GLTFNodeConverters.cpp b/Engine/Plugins/Enterprise/GLTFExporter/MigratedFromGit/Source/GLTFExporter/Private/Converters/GLTFNodeConverters.cpp
--- a/Engine/Plugins/Enterprise/GLTFExporter/MigratedFromGit/Source/GLTFExporter/Private/Converters/GLTFNodeConverters.cpp
+++ b/Engine/Plugins/Enterprise/GLTFExporter/MigratedFromGit/Source/GLTFExporter/Private/Converters/GLTFNodeConverters.cpp
@@ -21,8 +21,11 @@ FGLTFJsonNodeIndex FGLTFActorConverter::Convert(const AActor* Actor)
 	const UBlueprint* Blueprint = FGLTFActorUtility::GetBlueprintFromActor(Actor);
 	if (FGLTFActorUtility::IsSkySphereBlueprint(Blueprint))
 	{
-		// Ignore mesh and light components that are part of BP_SkySphere
-		// TODO: add support for BP_SkySphere later?
+		if (Builder.ExportOptions->bExportSkySpheres)
+		{
+			FGLTFJsonNode& RootNode = Builder.GetNode(RootNodeIndex);
+			RootNode.SkySphere = Builder.GetOrAddSkySphere(Actor);
+		}
 	}
 	else if (FGLTFActorUtility::IsHDRIBackdropBlueprint(Blueprint))
 	{
diff --git a/Engine/Plugins/Enterprise/GLTFExporter/MigratedFromGit/Source/GLTFExporter/Private/Converters/GLTFSkySphereConverters.cpp b/Engine/Plugins/Enterprise/GLTFExporter/MigratedFromGit/Source/GLTFExporter/Private/Converters/GLTFSkySphereConverters.cpp
new file mode 100644
--- /dev/null
+++ b/Engine/Plugins/Enterprise/GLTFExporter/MigratedFromGit/Source/GLTFExporter/Private/Converters/GLTFSkySphereConverters.cpp
@@ -0,0 +1,259 @@
+// Copyright Epic Games, Inc. All Rights Reserved.
+
+#include "Converters/GLTFSkySphereConverters.h"
+#include "Converters/GLTFConverterUtility.h"
+#include "Converters/GLTFMaterialUtility.h"
+#include "Converters/GLTFActorUtility.h"
+#include "Builders/GLTFContainerBuilder.h"
+
+FGLTFJsonSkySphereIndex FGLTFSkySphereConverter::Convert(const AActor* SkySphereActor)
+{
+	const UBlueprint* Blueprint = FGLTFActorUtility::GetBlueprintFromActor(SkySphereActor);
+	if (!FGLTFActorUtility::IsSkySphereBlueprint(Blueprint))
+	{
+		return FGLTFJsonSkySphereIndex(INDEX_NONE);
+	}
+
+	FGLTFJsonSkySphere JsonSkySphere;
+	SkySphereActor->GetName(JsonSkySphere.Name);
+
+	const UStaticMeshComponent* StaticMeshComponent = nullptr;
+	FGLTFActorUtility::TryGetPropertyValue(SkySphereActor, TEXT("SkySphereMesh"), StaticMeshComponent);
+
+	if (StaticMeshComponent != nullptr)
+	{
+		const USceneComponent* ParentComponent = StaticMeshComponent->GetAttachParent();
+		const FName SocketName = StaticMeshComponent->GetAttachSocketName();
+
+		const FTransform Transform = StaticMeshComponent->GetComponentTransform();
+		const FTransform ParentTransform = ParentComponent->GetSocketTransform(SocketName);
+		const FTransform RelativeTransform = Transform.GetRelativeTransform(ParentTransform);
+
+		JsonSkySphere.Scale = FGLTFConverterUtility::ConvertScale(RelativeTransform.GetScale3D());
+	}
+	else
+	{
+		Builder.AddWarningMessage(FString::Printf(TEXT("Failed to export Scale for Sky Sphere %s"), *SkySphereActor->GetName()));
+	}
+
+	const UStaticMesh* StaticMesh = StaticMeshComponent != nullptr ? StaticMeshComponent->GetStaticMesh() : nullptr;
+	if (StaticMesh)
+	{
+		JsonSkySphere.SkySphereMesh = Builder.GetOrAddMesh(StaticMesh, -1, nullptr, { FGLTFMaterialUtility::GetDefault() });
+	}
+	else
+	{
+		Builder.AddWarningMessage(FString::Printf(TEXT("Failed to export SkySphereMesh for Sky Sphere %s"), *SkySphereActor->GetName()));
+	}
+
+	const UMaterialInstance* SkyMaterial = nullptr;
+	FGLTFActorUtility::TryGetPropertyValue(SkySphereActor, TEXT("Sky material"), SkyMaterial);
+
+	if (const UTexture2D* SkyTexture = GetSkyTexture(SkyMaterial))
+	{
+		JsonSkySphere.SkyTexture = Builder.GetOrAddTexture(SkyTexture);
+	}
+	else
+	{
+		Builder.AddWarningMessage(FString::Printf(TEXT("Failed to export SkyTexture for Sky Sphere %s"), *SkySphereActor->GetName()));
+	}
+
+	if (const UTexture2D* CloudsTexture = GetCloudsTexture(SkyMaterial))
+	{
+		JsonSkySphere.CloudsTexture = Builder.GetOrAddTexture(CloudsTexture);
+	}
+	else
+	{
+		Builder.AddWarningMessage(FString::Printf(TEXT("Failed to export CloudsTexture for Sky Sphere %s"), *SkySphereActor->GetName()));
+	}
+
+	if (const UTexture2D* StarsTexture = GetStarsTexture(SkyMaterial))
+	{
+		JsonSkySphere.StarsTexture = Builder.GetOrAddTexture(StarsTexture);
+	}
+	else
+	{
+		Builder.AddWarningMessage(FString::Printf(TEXT("Failed to export StarsTexture for Sky Sphere %s"), *SkySphereActor->GetName()));
+	}
+
+	const ADirectionalLight* DirectionalLight = nullptr;
+	if (FGLTFActorUtility::TryGetPropertyValue(SkySphereActor, TEXT("Directional light actor"), DirectionalLight))
+	{
+		JsonSkySphere.DirectionalLight = Builder.GetOrAddNode(DirectionalLight);
+	}
+	else
+	{
+		Builder.AddWarningMessage(FString::Printf(TEXT("Failed to export DirectionalLight for Sky Sphere %s"), *SkySphereActor->GetName()));
+	}
+
+	float SunRadius;
+	if (SkyMaterial != nullptr && SkyMaterial->GetScalarParameterValue(TEXT("Sun Radius"), SunRadius))
+	{
+		JsonSkySphere.SunRadius = SunRadius;
+	}
+	else
+	{
+		Builder.AddWarningMessage(FString::Printf(TEXT("Failed to export SunRadius for Sky Sphere %s"), *SkySphereActor->GetName()));
+	}
+
+	float NoisePower1;
+	if (SkyMaterial != nullptr && SkyMaterial->GetScalarParameterValue(TEXT("NoisePower1"), NoisePower1))
+	{
+		JsonSkySphere.NoisePower1 = NoisePower1;
+	}
+	else
+	{
+		Builder.AddWarningMessage(FString::Printf(TEXT("Failed to export NoisePower1 for Sky Sphere %s"), *SkySphereActor->GetName()));
+	}
+
+	float NoisePower2;
+	if (SkyMaterial != nullptr && SkyMaterial->GetScalarParameterValue(TEXT("NoisePower2"), NoisePower2))
+	{
+		JsonSkySphere.NoisePower2 = NoisePower2;
+	}
+	else
+	{
+		Builder.AddWarningMessage(FString::Printf(TEXT("Failed to export NoisePower2 for Sky Sphere %s"), *SkySphereActor->GetName()));
+	}
+
+	float SunHeight;
+	if (FGLTFActorUtility::TryGetPropertyValue(SkySphereActor, TEXT("Sun height"), SunHeight))
+	{
+		JsonSkySphere.SunHeight = SunHeight;
+	}
+	else
+	{
+		Builder.AddWarningMessage(FString::Printf(TEXT("Failed to export SunHeight for Sky Sphere %s"), *SkySphereActor->GetName()));
+	}
+
+	float SunBrightness;
+	if (FGLTFActorUtility::TryGetPropertyValue(SkySphereActor, TEXT("Sun brightness"), SunBrightness))
+	{
+		JsonSkySphere.SunBrightness = SunBrightness;
+	}
+	else
+	{
+		Builder.AddWarningMessage(FString::Printf(TEXT("Failed to export SunBrightness for Sky Sphere %s"), *SkySphereActor->GetName()));
+	}
+
+	float StarsBrightness;
+	if (FGLTFActorUtility::TryGetPropertyValue(SkySphereActor, TEXT("Stars brightness"), StarsBrightness))
+	{
+		JsonSkySphere.StarsBrightness = StarsBrightness;
+	}
+	else
+	{
+		Builder.AddWarningMessage(FString::Printf(TEXT("Failed to export StarsBrightness for Sky Sphere %s"), *SkySphereActor->GetName()));
+	}
+
+	float CloudSpeed;
+	if (FGLTFActorUtility::TryGetPropertyValue(SkySphereActor, TEXT("Cloud speed"), CloudSpeed))
+	{
+		JsonSkySphere.CloudSpeed = CloudSpeed;
+	}
+	else
+	{
+		Builder.AddWarningMessage(FString::Printf(TEXT("Failed to export CloudSpeed for Sky Sphere %s"), *SkySphereActor->GetName()));
+	}
+
+	float CloudOpacity;
+	if (FGLTFActorUtility::TryGetPropertyValue(SkySphereActor, TEXT("Cloud opacity"), CloudOpacity))
+	{
+		JsonSkySphere.CloudOpacity = CloudOpacity;
+	}
+	else
+	{
+		Builder.AddWarningMessage(FString::Printf(TEXT("Failed to export CloudOpacity for Sky Sphere %s"), *SkySphereActor->GetName()));
+	}
+
+	float HorizonFalloff;
+	if (FGLTFActorUtility::TryGetPropertyValue(SkySphereActor, TEXT("Horizon Falloff"), HorizonFalloff))
+	{
+		JsonSkySphere.HorizonFalloff = HorizonFalloff;
+	}
+	else
+	{
+		Builder.AddWarningMessage(FString::Printf(TEXT("Failed to export HorizonFalloff for Sky Sphere %s"), *SkySphereActor->GetName()));
+	}
+
+	bool bColorsDeterminedBySunPosition;
+	if (FGLTFActorUtility::TryGetPropertyValue(SkySphereActor, TEXT("Colors determined by sun position"), bColorsDeterminedBySunPosition))
+	{
+		JsonSkySphere.bColorsDeterminedBySunPosition = bColorsDeterminedBySunPosition;
+	}
+	else
+	{
+		Builder.AddWarningMessage(FString::Printf(TEXT("Failed to export bColorsDeterminedBySunPosition for Sky Sphere %s"), *SkySphereActor->GetName()));
+	}
+
+	FLinearColor ZenithColor;
+	if (FGLTFActorUtility::TryGetPropertyValue(SkySphereActor, TEXT("Zenith Color"), ZenithColor))
+	{
+		JsonSkySphere.ZenithColor = FGLTFConverterUtility::ConvertColor(ZenithColor);
+	}
+	else
+	{
+		Builder.AddWarningMessage(FString::Printf(TEXT("Failed to export ZenithColor for Sky Sphere %s"), *SkySphereActor->GetName()));
+	}
+
+	FLinearColor HorizonColor;
+	if (FGLTFActorUtility::TryGetPropertyValue(SkySphereActor, TEXT("Horizon color"), HorizonColor))
+	{
+		JsonSkySphere.HorizonColor = FGLTFConverterUtility::ConvertColor(HorizonColor);
+	}
+	else
+	{
+		Builder.AddWarningMessage(FString::Printf(TEXT("Failed to export HorizonColor for Sky Sphere %s"), *SkySphereActor->GetName()));
+	}
+
+	FLinearColor CloudColor;
+	if (FGLTFActorUtility::TryGetPropertyValue(SkySphereActor, TEXT("Cloud color"), CloudColor))
+	{
+		JsonSkySphere.CloudColor = FGLTFConverterUtility::ConvertColor(CloudColor);
+	}
+	else
+	{
+		Builder.AddWarningMessage(FString::Printf(TEXT("Failed to export CloudColor for Sky Sphere %s"), *SkySphereActor->GetName()));
+	}
+
+	FLinearColor OverallColor;
+	if (FGLTFActorUtility::TryGetPropertyValue(SkySphereActor, TEXT("Overall Color"), OverallColor))
+	{
+		JsonSkySphere.OverallColor = FGLTFConverterUtility::ConvertColor(OverallColor);
+	}
+	else
+	{
+		Builder.AddWarningMessage(FString::Printf(TEXT("Failed to export OverallColor for Sky Sphere %s"), *SkySphereActor->GetName()));
+	}
+
+	// TODO: export color curves
+
+	return Builder.AddSkySphere(JsonSkySphere);
+}
+
+const UTexture2D* FGLTFSkySphereConverter::GetSkyTexture(const UMaterialInstance* SkyMaterial) const
+{
+	// TODO: the texture isn't exposed as a parameter on the default sky material. We should try to find
+	// a way of extracting and returning the correct texture from the material anyway.
+
+	static UTexture2D* Texture = LoadObject<UTexture2D>(nullptr, TEXT("/Engine/EngineSky/T_Sky_Blue.T_Sky_Blue"));
+	return Texture;
+}
+
+const UTexture2D* FGLTFSkySphereConverter::GetCloudsTexture(const UMaterialInstance* SkyMaterial) const
+{
+	// TODO: the texture isn't exposed as a parameter on the default sky material. We should try to find
+	// a way of extracting and returning the correct texture from the material anyway.
+
+	static UTexture2D* Texture = LoadObject<UTexture2D>(nullptr, TEXT("/Engine/EngineSky/T_Sky_Clouds_M.T_Sky_Clouds_M"));
+	return Texture;
+}
+
+const UTexture2D* FGLTFSkySphereConverter::GetStarsTexture(const UMaterialInstance* SkyMaterial) const
+{
+	// TODO: the texture isn't exposed as a parameter on the default sky material. We should try to find
+	// a way of extracting and returning the correct texture from the material anyway.
+
+	static UTexture2D* Texture = LoadObject<UTexture2D>(nullptr, TEXT("/Engine/EngineSky/T_Sky_Stars.T_Sky_Stars"));
+	return Texture;
+}
diff --git a/Engine/Plugins/Enterprise/GLTFExporter/MigratedFromGit/Source/GLTFExporter/Private/Converters/GLTFSkySphereConverters.h b/Engine/Plugins/Enterprise/GLTFExporter/MigratedFromGit/Source/GLTFExporter/Private/Converters/GLTFSkySphereConverters.h
new file mode 100644
--- /dev/null
+++ b/Engine/Plugins/Enterprise/GLTFExporter/MigratedFromGit/Source/GLTFExporter/Private/Converters/GLTFSkySphereConverters.h
@@ -0,0 +1,19 @@
+// Copyright Epic Games, Inc. All Rights Reserved.
+
+#pragma once
+
+#include "Json/GLTFJsonIndex.h"
+#include "Converters/GLTFConverter.h"
+#include "Converters/GLTFBuilderContext.h"
+#include "Engine.h"
+
+class FGLTFSkySphereConverter final : public FGLTFBuilderContext, public TGLTFConverter<FGLTFJsonSkySphereIndex, const AActor*>
+{
+	using FGLTFBuilderContext::FGLTFBuilderContext;
+
+	virtual FGLTFJsonSkySphereIndex Convert(const AActor* SkySphereActor) override;
+
+	const UTexture2D* GetSkyTexture(const UMaterialInstance* SkyMaterial) const;
+	const UTexture2D* GetCloudsTexture(const UMaterialInstance* SkyMaterial) const;
+	const UTexture2D* GetStarsTexture(const UMaterialInstance* SkyMaterial) const;
+};
diff --git a/Engine/Plugins/Enterprise/GLTFExporter/MigratedFromGit/Source/GLTFExporter/Private/GLTFExportOptions.cpp b/Engine/Plugins/Enterprise/GLTFExporter/MigratedFromGit/Source/GLTFExporter/Private/GLTFExportOptions.cpp
--- a/Engine/Plugins/Enterprise/GLTFExporter/MigratedFromGit/Source/GLTFExporter/Private/GLTFExportOptions.cpp
+++ b/Engine/Plugins/Enterprise/GLTFExporter/MigratedFromGit/Source/GLTFExporter/Private/GLTFExportOptions.cpp
@@ -35,6 +35,7 @@ UGLTFExportOptions::UGLTFExportOptions(const FObjectInitializer& ObjectInitializ
 	bExportCameras = true;
 	bExportOrbitalCameras = true;
 	bExportHDRIBackdrops = true;
+	bExportSkySpheres = true;
 	bExportVariantSets = true;
 	bExportAnimationHotspots = true;
 	bBundleWebViewer = true;
diff --git a/Engine/Plugins/Enterprise/GLTFExporter/MigratedFromGit/Source/GLTFExporter/Private/Json/GLTFJsonEnums.h b/Engine/Plugins/Enterprise/GLTFExporter/MigratedFromGit/Source/GLTFExporter/Private/Json/GLTFJsonEnums.h
--- a/Engine/Plugins/Enterprise/GLTFExporter/MigratedFromGit/Source/GLTFExporter/Private/Json/GLTFJsonEnums.h
+++ b/Engine/Plugins/Enterprise/GLTFExporter/MigratedFromGit/Source/GLTFExporter/Private/Json/GLTFJsonEnums.h
@@ -17,7 +17,8 @@ enum class EGLTFJsonExtension
 	EPIC_ReflectionCaptures,
 	EPIC_LevelVariantSets,
 	EPIC_TextureHDREncoding,
-    EPIC_AnimationPlayback
+    EPIC_AnimationPlayback,
+	EPIC_SkySpheres
 };
 
 enum class EGLTFJsonShadingModel
diff --git a/Engine/Plugins/Enterprise/GLTFExporter/MigratedFromGit/Source/GLTFExporter/Private/Json/GLTFJsonIndex.h b/Engine/Plugins/Enterprise/GLTFExporter/MigratedFromGit/Source/GLTFExporter/Private/Json/GLTFJsonIndex.h
--- a/Engine/Plugins/Enterprise/GLTFExporter/MigratedFromGit/Source/GLTFExporter/Private/Json/GLTFJsonIndex.h
+++ b/Engine/Plugins/Enterprise/GLTFExporter/MigratedFromGit/Source/GLTFExporter/Private/Json/GLTFJsonIndex.h
@@ -71,3 +71,4 @@ struct FGLTFJsonVariationIndex : FGLTFJsonIndex<FGLTFJsonVariationIndex> { using
 struct FGLTFJsonLightMapIndex : FGLTFJsonIndex<FGLTFJsonLightMapIndex> { using FGLTFJsonIndex::FGLTFJsonIndex; };
 struct FGLTFJsonLightIndex : FGLTFJsonIndex<FGLTFJsonLightIndex> { using FGLTFJsonIndex::FGLTFJsonIndex; };
 struct FGLTFJsonHotspotIndex : FGLTFJsonIndex<FGLTFJsonHotspotIndex> { using FGLTFJsonIndex::FGLTFJsonIndex; };
+struct FGLTFJsonSkySphereIndex : FGLTFJsonIndex<FGLTFJsonSkySphereIndex> { using FGLTFJsonIndex::FGLTFJsonIndex; };
diff --git a/Engine/Plugins/Enterprise/GLTFExporter/MigratedFromGit/Source/GLTFExporter/Private/Json/GLTFJsonNode.h b/Engine/Plugins/Enterprise/GLTFExporter/MigratedFromGit/Source/GLTFExporter/Private/Json/GLTFJsonNode.h
--- a/Engine/Plugins/Enterprise/GLTFExporter/MigratedFromGit/Source/GLTFExporter/Private/Json/GLTFJsonNode.h
+++ b/Engine/Plugins/Enterprise/GLTFExporter/MigratedFromGit/Source/GLTFExporter/Private/Json/GLTFJsonNode.h
@@ -26,13 +26,14 @@ struct FGLTFJsonNode
 		};
 	};
 
-	FGLTFJsonCameraIndex   Camera;
-	FGLTFJsonSkinIndex     Skin;
-	FGLTFJsonMeshIndex     Mesh;
-	FGLTFJsonBackdropIndex Backdrop;
-	FGLTFJsonLightMapIndex LightMap;
-	FGLTFJsonLightIndex    Light;
-	FGLTFJsonHotspotIndex  Hotspot;
+	FGLTFJsonCameraIndex    Camera;
+	FGLTFJsonSkinIndex      Skin;
+	FGLTFJsonMeshIndex      Mesh;
+	FGLTFJsonBackdropIndex  Backdrop;
+	FGLTFJsonLightMapIndex  LightMap;
+	FGLTFJsonLightIndex     Light;
+	FGLTFJsonHotspotIndex   Hotspot;
+	FGLTFJsonSkySphereIndex SkySphere;
 
 	FGLTFJsonNodeIndex ComponentNode;
 
@@ -100,7 +101,7 @@ struct FGLTFJsonNode
 			JsonWriter.WriteValue(TEXT("mesh"), Mesh);
 		}
 
-		if (Backdrop != INDEX_NONE || LightMap != INDEX_NONE || Light != INDEX_NONE || Hotspot != INDEX_NONE)
+		if (Backdrop != INDEX_NONE || LightMap != INDEX_NONE || Light != INDEX_NONE || Hotspot != INDEX_NONE || SkySphere != INDEX_NONE)
 		{
 			JsonWriter.WriteObjectStart(TEXT("extensions"));
 
@@ -144,6 +145,16 @@ struct FGLTFJsonNode
 				JsonWriter.WriteObjectEnd();
 			}
 
+			if (SkySphere != INDEX_NONE)
+			{
+				const EGLTFJsonExtension Extension = EGLTFJsonExtension::EPIC_SkySpheres;
+				Extensions.Used.Add(Extension);
+
+				JsonWriter.WriteObjectStart(FGLTFJsonUtility::ToString(Extension));
+				JsonWriter.WriteValue(TEXT("skySphere"), SkySphere);
+				JsonWriter.WriteObjectEnd();
+			}
+
 			JsonWriter.WriteObjectEnd();
 		}
 
diff --git a/Engine/Plugins/Enterprise/GLTFExporter/MigratedFromGit/Source/GLTFExporter/Private/Json/GLTFJsonRoot.h b/Engine/Plugins/Enterprise/GLTFExporter/MigratedFromGit/Source/GLTFExporter/Private/Json/GLTFJsonRoot.h
--- a/Engine/Plugins/Enterprise/GLTFExporter/MigratedFromGit/Source/GLTFExporter/Private/Json/GLTFJsonRoot.h
+++ b/Engine/Plugins/Enterprise/GLTFExporter/MigratedFromGit/Source/GLTFExporter/Private/Json/GLTFJsonRoot.h
@@ -22,6 +22,7 @@
 #include "Json/GLTFJsonLightMap.h"
 #include "Json/GLTFJsonLight.h"
 #include "Json/GLTFJsonHotspot.h"
+#include "Json/GLTFJsonSkySphere.h"
 #include "Policies/CondensedJsonPrintPolicy.h"
 
 
@@ -80,6 +81,7 @@ struct FGLTFJsonRoot
 	TArray<TUniquePtr<FGLTFJsonLight>>      Lights;
 	TArray<TUniquePtr<FGLTFJsonVariation>>  Variations;
 	TArray<TUniquePtr<FGLTFJsonHotspot>>    Hotspots;
+	TArray<TUniquePtr<FGLTFJsonSkySphere>>  SkySpheres;
 
 	template <class CharType = TCHAR, class PrintPolicy = TPrettyJsonPrintPolicy<CharType>>
 	void WriteObject(TJsonWriter<CharType, PrintPolicy>& JsonWriter, bool bAllExtensionsRequired) const
@@ -110,7 +112,7 @@ struct FGLTFJsonRoot
 		FGLTFJsonUtility::WriteObjectPtrArray(JsonWriter, TEXT("skins"), Skins, AllExtensions);
 		FGLTFJsonUtility::WriteObjectPtrArray(JsonWriter, TEXT("textures"), Textures, AllExtensions);
 
-		if (Backdrops.Num() > 0 || Variations.Num() > 0 || LightMaps.Num() > 0 || Lights.Num() > 0 || Hotspots.Num() > 0 /* TODO: add more extension support */)
+		if (Backdrops.Num() > 0 || Variations.Num() > 0 || LightMaps.Num() > 0 || Lights.Num() > 0 || Hotspots.Num() > 0 || SkySpheres.Num() > 0 /* TODO: add more extension support */)
 		{
 			JsonWriter.WriteObjectStart(TEXT("extensions"));
 
@@ -164,6 +166,16 @@ struct FGLTFJsonRoot
 				JsonWriter.WriteObjectEnd();
 			}
 
+			if (SkySpheres.Num() > 0)
+			{
+				const EGLTFJsonExtension Extension = EGLTFJsonExtension::EPIC_SkySpheres;
+				AllExtensions.Used.Add(Extension);
+
+				JsonWriter.WriteObjectStart(FGLTFJsonUtility::ToString(Extension));
+				FGLTFJsonUtility::WriteObjectPtrArray(JsonWriter, TEXT("skySpheres"), SkySpheres, AllExtensions);
+				JsonWriter.WriteObjectEnd();
+			}
+
 			JsonWriter.WriteObjectEnd();
 		}
 
diff --git a/Engine/Plugins/Enterprise/GLTFExporter/MigratedFromGit/Source/GLTFExporter/Private/Json/GLTFJsonSkySphere.h b/Engine/Plugins/Enterprise/GLTFExporter/MigratedFromGit/Source/GLTFExporter/Private/Json/GLTFJsonSkySphere.h
new file mode 100644
--- /dev/null
+++ b/Engine/Plugins/Enterprise/GLTFExporter/MigratedFromGit/Source/GLTFExporter/Private/Json/GLTFJsonSkySphere.h
@@ -0,0 +1,119 @@
+// Copyright Epic Games, Inc. All Rights Reserved.
+
+#pragma once
+
+#include "Json/GLTFJsonIndex.h"
+#include "Json/GLTFJsonColor4.h"
+#include "Json/GLTFJsonExtensions.h"
+#include "Serialization/JsonSerializer.h"
+
+struct FGLTFJsonSkySphere
+{
+	FString Name;
+
+	FGLTFJsonMeshIndex    SkySphereMesh;
+	FGLTFJsonTextureIndex SkyTexture;
+	FGLTFJsonTextureIndex CloudsTexture;
+	FGLTFJsonTextureIndex StarsTexture;
+	FGLTFJsonNodeIndex    DirectionalLight;
+
+	float SunHeight;
+	float SunBrightness;
+	float StarsBrightness;
+	float CloudSpeed;
+	float CloudOpacity;
+	float HorizonFalloff;
+
+	float SunRadius;
+	float NoisePower1;
+	float NoisePower2;
+
+	bool bColorsDeterminedBySunPosition;
+
+	FGLTFJsonColor4 ZenithColor;
+	FGLTFJsonColor4 HorizonColor;
+	FGLTFJsonColor4 CloudColor;
+	FGLTFJsonColor4 OverallColor;
+
+	FGLTFJsonVector3 Scale;
+
+	// TODO: add properties for color curves
+
+	FGLTFJsonSkySphere()
+		: SunHeight(0)
+		, SunBrightness(0)
+		, StarsBrightness(0)
+		, CloudSpeed(0)
+		, CloudOpacity(0)
+		, HorizonFalloff(0)
+		, SunRadius(0)
+		, NoisePower1(0)
+		, NoisePower2(0)
+		, bColorsDeterminedBySunPosition(false)
+		, ZenithColor(FGLTFJsonColor4::White)
+		, HorizonColor(FGLTFJsonColor4::White)
+		, CloudColor(FGLTFJsonColor4::White)
+		, OverallColor(FGLTFJsonColor4::White)
+		, Scale(FGLTFJsonVector3::One)
+	{
+	}
+
+	template <class CharType = TCHAR, class PrintPolicy = TPrettyJsonPrintPolicy<CharType>>
+	void WriteObject(TJsonWriter<CharType, PrintPolicy>& JsonWriter, FGLTFJsonExtensions& Extensions) const
+	{
+		JsonWriter.WriteObjectStart();
+
+		if (!Name.IsEmpty())
+		{
+			JsonWriter.WriteValue(TEXT("name"), Name);
+		}
+
+		JsonWriter.WriteValue(TEXT("skySphereMesh"), SkySphereMesh);
+		JsonWriter.WriteValue(TEXT("skyTexture"), SkyTexture);
+		JsonWriter.WriteValue(TEXT("cloudsTexture"), CloudsTexture);
+		JsonWriter.WriteValue(TEXT("starsTexture"), StarsTexture);
+
+		if (DirectionalLight != INDEX_NONE)
+		{
+			JsonWriter.WriteValue(TEXT("directionalLight"), DirectionalLight);
+		}
+
+		FGLTFJsonUtility::WriteExactValue(JsonWriter, TEXT("sunHeight"), SunHeight);
+		FGLTFJsonUtility::WriteExactValue(JsonWriter, TEXT("sunBrightness"), SunBrightness);
+		FGLTFJsonUtility::WriteExactValue(JsonWriter, TEXT("starsBrightness"), StarsBrightness);
+		FGLTFJsonUtility::WriteExactValue(JsonWriter, TEXT("cloudSpeed"), CloudSpeed);
+		FGLTFJsonUtility::WriteExactValue(JsonWriter, TEXT("cloudOpacity"), CloudOpacity);
+		FGLTFJsonUtility::WriteExactValue(JsonWriter, TEXT("horizonFalloff"), HorizonFalloff);
+
+		FGLTFJsonUtility::WriteExactValue(JsonWriter, TEXT("sunRadius"), SunRadius);
+		FGLTFJsonUtility::WriteExactValue(JsonWriter, TEXT("noisePower1"), NoisePower1);
+		FGLTFJsonUtility::WriteExactValue(JsonWriter, TEXT("noisePower2"), NoisePower2);
+
+		JsonWriter.WriteValue(TEXT("colorsDeterminedBySunPosition"), bColorsDeterminedBySunPosition);
+
+		JsonWriter.WriteIdentifierPrefix(TEXT("zenithColor"));
+		ZenithColor.WriteArray(JsonWriter);
+
+		JsonWriter.WriteIdentifierPrefix(TEXT("horizonColor"));
+		HorizonColor.WriteArray(JsonWriter);
+
+		JsonWriter.WriteIdentifierPrefix(TEXT("cloudColor"));
+		CloudColor.WriteArray(JsonWriter);
+
+		if (OverallColor != FGLTFJsonColor4::White)
+		{
+			JsonWriter.WriteIdentifierPrefix(TEXT("overallColor"));
+			OverallColor.WriteArray(JsonWriter);
+		}
+
+		if (Scale != FGLTFJsonVector3::One)
+		{
+			JsonWriter.WriteIdentifierPrefix(TEXT("scale"));
+			Scale.WriteArray(JsonWriter);
+		}
+
+		// TODO: export color curves
+
+		JsonWriter.WriteObjectEnd();
+	}
+};
diff --git a/Engine/Plugins/Enterprise/GLTFExporter/MigratedFromGit/Source/GLTFExporter/Private/Json/GLTFJsonUtility.h b/Engine/Plugins/Enterprise/GLTFExporter/MigratedFromGit/Source/GLTFExporter/Private/Json/GLTFJsonUtility.h
--- a/Engine/Plugins/Enterprise/GLTFExporter/MigratedFromGit/Source/GLTFExporter/Private/Json/GLTFJsonUtility.h
+++ b/Engine/Plugins/Enterprise/GLTFExporter/MigratedFromGit/Source/GLTFExporter/Private/Json/GLTFJsonUtility.h
@@ -31,6 +31,7 @@ struct FGLTFJsonUtility
 			case EGLTFJsonExtension::EPIC_LevelVariantSets:   return TEXT("EPIC_level_variant_sets");
 			case EGLTFJsonExtension::EPIC_TextureHDREncoding: return TEXT("EPIC_texture_hdr_encoding");
 			case EGLTFJsonExtension::EPIC_AnimationPlayback:  return TEXT("EPIC_animation_playback");
+			case EGLTFJsonExtension::EPIC_SkySpheres:         return TEXT("EPIC_sky_spheres");
 			default:                                          return TEXT("unknown");
 		}
 	}
diff --git a/Engine/Plugins/Enterprise/GLTFExporter/MigratedFromGit/Source/GLTFExporter/Public/GLTFExportOptions.h b/Engine/Plugins/Enterprise/GLTFExporter/MigratedFromGit/Source/GLTFExporter/Public/GLTFExportOptions.h
--- a/Engine/Plugins/Enterprise/GLTFExporter/MigratedFromGit/Source/GLTFExporter/Public/GLTFExportOptions.h
+++ b/Engine/Plugins/Enterprise/GLTFExporter/MigratedFromGit/Source/GLTFExporter/Public/GLTFExportOptions.h
@@ -117,6 +117,9 @@ public:
 	UPROPERTY(EditAnywhere, BlueprintReadWrite, config, Category = Scene)
 	uint32 bExportHDRIBackdrops : 1;
 
+	UPROPERTY(EditAnywhere, BlueprintReadWrite, config, Category = Scene)
+	uint32 bExportSkySpheres : 1;
+
 	UPROPERTY(EditAnywhere, BlueprintReadWrite, config, Category = Scene)
 	uint32 bExportVariantSets : 1;
 
