[GLTFExporter] refactor(sky-sphere): refactor property conversion to reduce duplicated code

commit 9255832ca7a7ed7ebac5d0a7a364a19b3c3249dc
Date: Wed, 14 Oct 2020 14:57:16 +0200

    refactor(sky-sphere): refactor property conversion to reduce duplicated code
    (cherry picked from commit 610349f6c8789e5dd21860e2b2d0036b80d0cde2​)

#skipci

diff --git a/Engine/Plugins/Enterprise/GLTFExporter/MigratedFromGit/Source/GLTFExporter/Private/Converters/GLTFSkySphereConverters.cpp b/Engine/Plugins/Enterprise/GLTFExporter/MigratedFromGit/Source/GLTFExporter/Private/Converters/GLTFSkySphereConverters.cpp
--- a/Engine/Plugins/Enterprise/GLTFExporter/MigratedFromGit/Source/GLTFExporter/Private/Converters/GLTFSkySphereConverters.cpp
+++ b/Engine/Plugins/Enterprise/GLTFExporter/MigratedFromGit/Source/GLTFExporter/Private/Converters/GLTFSkySphereConverters.cpp
@@ -86,145 +86,22 @@ FGLTFJsonSkySphereIndex FGLTFSkySphereConverter::Convert(const AActor* SkySphere
 		Builder.AddWarningMessage(FString::Printf(TEXT("Failed to export DirectionalLight for Sky Sphere %s"), *SkySphereActor->GetName()));
 	}
 
-	float SunRadius;
-	if (SkyMaterial != nullptr && SkyMaterial->GetScalarParameterValue(TEXT("Sun Radius"), SunRadius))
-	{
-		JsonSkySphere.SunRadius = SunRadius;
-	}
-	else
-	{
-		Builder.AddWarningMessage(FString::Printf(TEXT("Failed to export SunRadius for Sky Sphere %s"), *SkySphereActor->GetName()));
-	}
-
-	float NoisePower1;
-	if (SkyMaterial != nullptr && SkyMaterial->GetScalarParameterValue(TEXT("NoisePower1"), NoisePower1))
-	{
-		JsonSkySphere.NoisePower1 = NoisePower1;
-	}
-	else
-	{
-		Builder.AddWarningMessage(FString::Printf(TEXT("Failed to export NoisePower1 for Sky Sphere %s"), *SkySphereActor->GetName()));
-	}
-
-	float NoisePower2;
-	if (SkyMaterial != nullptr && SkyMaterial->GetScalarParameterValue(TEXT("NoisePower2"), NoisePower2))
-	{
-		JsonSkySphere.NoisePower2 = NoisePower2;
-	}
-	else
-	{
-		Builder.AddWarningMessage(FString::Printf(TEXT("Failed to export NoisePower2 for Sky Sphere %s"), *SkySphereActor->GetName()));
-	}
-
-	float SunHeight;
-	if (FGLTFActorUtility::TryGetPropertyValue(SkySphereActor, TEXT("Sun height"), SunHeight))
-	{
-		JsonSkySphere.SunHeight = SunHeight;
-	}
-	else
-	{
-		Builder.AddWarningMessage(FString::Printf(TEXT("Failed to export SunHeight for Sky Sphere %s"), *SkySphereActor->GetName()));
-	}
-
-	float SunBrightness;
-	if (FGLTFActorUtility::TryGetPropertyValue(SkySphereActor, TEXT("Sun brightness"), SunBrightness))
-	{
-		JsonSkySphere.SunBrightness = SunBrightness;
-	}
-	else
-	{
-		Builder.AddWarningMessage(FString::Printf(TEXT("Failed to export SunBrightness for Sky Sphere %s"), *SkySphereActor->GetName()));
-	}
-
-	float StarsBrightness;
-	if (FGLTFActorUtility::TryGetPropertyValue(SkySphereActor, TEXT("Stars brightness"), StarsBrightness))
-	{
-		JsonSkySphere.StarsBrightness = StarsBrightness;
-	}
-	else
-	{
-		Builder.AddWarningMessage(FString::Printf(TEXT("Failed to export StarsBrightness for Sky Sphere %s"), *SkySphereActor->GetName()));
-	}
-
-	float CloudSpeed;
-	if (FGLTFActorUtility::TryGetPropertyValue(SkySphereActor, TEXT("Cloud speed"), CloudSpeed))
-	{
-		JsonSkySphere.CloudSpeed = CloudSpeed;
-	}
-	else
-	{
-		Builder.AddWarningMessage(FString::Printf(TEXT("Failed to export CloudSpeed for Sky Sphere %s"), *SkySphereActor->GetName()));
-	}
-
-	float CloudOpacity;
-	if (FGLTFActorUtility::TryGetPropertyValue(SkySphereActor, TEXT("Cloud opacity"), CloudOpacity))
-	{
-		JsonSkySphere.CloudOpacity = CloudOpacity;
-	}
-	else
-	{
-		Builder.AddWarningMessage(FString::Printf(TEXT("Failed to export CloudOpacity for Sky Sphere %s"), *SkySphereActor->GetName()));
-	}
-
-	float HorizonFalloff;
-	if (FGLTFActorUtility::TryGetPropertyValue(SkySphereActor, TEXT("Horizon Falloff"), HorizonFalloff))
-	{
-		JsonSkySphere.HorizonFalloff = HorizonFalloff;
-	}
-	else
-	{
-		Builder.AddWarningMessage(FString::Printf(TEXT("Failed to export HorizonFalloff for Sky Sphere %s"), *SkySphereActor->GetName()));
-	}
-
-	bool bColorsDeterminedBySunPosition;
-	if (FGLTFActorUtility::TryGetPropertyValue(SkySphereActor, TEXT("Colors determined by sun position"), bColorsDeterminedBySunPosition))
-	{
-		JsonSkySphere.bColorsDeterminedBySunPosition = bColorsDeterminedBySunPosition;
-	}
-	else
-	{
-		Builder.AddWarningMessage(FString::Printf(TEXT("Failed to export bColorsDeterminedBySunPosition for Sky Sphere %s"), *SkySphereActor->GetName()));
-	}
-
-	FLinearColor ZenithColor;
-	if (FGLTFActorUtility::TryGetPropertyValue(SkySphereActor, TEXT("Zenith Color"), ZenithColor))
-	{
-		JsonSkySphere.ZenithColor = FGLTFConverterUtility::ConvertColor(ZenithColor);
-	}
-	else
-	{
-		Builder.AddWarningMessage(FString::Printf(TEXT("Failed to export ZenithColor for Sky Sphere %s"), *SkySphereActor->GetName()));
-	}
-
-	FLinearColor HorizonColor;
-	if (FGLTFActorUtility::TryGetPropertyValue(SkySphereActor, TEXT("Horizon color"), HorizonColor))
-	{
-		JsonSkySphere.HorizonColor = FGLTFConverterUtility::ConvertColor(HorizonColor);
-	}
-	else
-	{
-		Builder.AddWarningMessage(FString::Printf(TEXT("Failed to export HorizonColor for Sky Sphere %s"), *SkySphereActor->GetName()));
-	}
-
-	FLinearColor CloudColor;
-	if (FGLTFActorUtility::TryGetPropertyValue(SkySphereActor, TEXT("Cloud color"), CloudColor))
-	{
-		JsonSkySphere.CloudColor = FGLTFConverterUtility::ConvertColor(CloudColor);
-	}
-	else
-	{
-		Builder.AddWarningMessage(FString::Printf(TEXT("Failed to export CloudColor for Sky Sphere %s"), *SkySphereActor->GetName()));
-	}
-
-	FLinearColor OverallColor;
-	if (FGLTFActorUtility::TryGetPropertyValue(SkySphereActor, TEXT("Overall Color"), OverallColor))
-	{
-		JsonSkySphere.OverallColor = FGLTFConverterUtility::ConvertColor(OverallColor);
-	}
-	else
-	{
-		Builder.AddWarningMessage(FString::Printf(TEXT("Failed to export OverallColor for Sky Sphere %s"), *SkySphereActor->GetName()));
-	}
+	ConvertScalarParameter(*SkySphereActor, SkyMaterial, TEXT("Sun Radius"), TEXT("SunRadius"), JsonSkySphere.SunRadius);
+	ConvertScalarParameter(*SkySphereActor, SkyMaterial, TEXT("NoisePower1"), TEXT("NoisePower1"), JsonSkySphere.NoisePower1);
+	ConvertScalarParameter(*SkySphereActor, SkyMaterial, TEXT("NoisePower2"), TEXT("NoisePower2"), JsonSkySphere.NoisePower2);
+
+	ConvertProperty(*SkySphereActor, TEXT("Sun height"), TEXT("SunHeight"), JsonSkySphere.SunHeight);
+	ConvertProperty(*SkySphereActor, TEXT("Sun brightness"), TEXT("SunBrightness"), JsonSkySphere.SunBrightness);
+	ConvertProperty(*SkySphereActor, TEXT("Stars brightness"), TEXT("StarsBrightness"), JsonSkySphere.StarsBrightness);
+	ConvertProperty(*SkySphereActor, TEXT("Cloud speed"), TEXT("CloudSpeed"), JsonSkySphere.CloudSpeed);
+	ConvertProperty(*SkySphereActor, TEXT("Cloud opacity"), TEXT("CloudOpacity"), JsonSkySphere.CloudOpacity);
+	ConvertProperty(*SkySphereActor, TEXT("Horizon Falloff"), TEXT("HorizonFalloff"), JsonSkySphere.HorizonFalloff);
+	ConvertProperty(*SkySphereActor, TEXT("Colors determined by sun position"), TEXT("bColorsDeterminedBySunPosition"), JsonSkySphere.bColorsDeterminedBySunPosition);
+
+	ConvertColorProperty(*SkySphereActor, TEXT("Zenith Color"), TEXT("ZenithColor"), JsonSkySphere.ZenithColor);
+	ConvertColorProperty(*SkySphereActor, TEXT("Horizon color"), TEXT("HorizonColor"), JsonSkySphere.HorizonColor);
+	ConvertColorProperty(*SkySphereActor, TEXT("Cloud color"), TEXT("CloudColor"), JsonSkySphere.CloudColor);
+	ConvertColorProperty(*SkySphereActor, TEXT("Overall Color"), TEXT("OverallColor"), JsonSkySphere.OverallColor);
 
 	// TODO: export color curves
 
@@ -257,3 +134,42 @@ const UTexture2D* FGLTFSkySphereConverter::GetStarsTexture(const UMaterialInstan
 	static UTexture2D* Texture = LoadObject<UTexture2D>(nullptr, TEXT("/Engine/EngineSky/T_Sky_Stars.T_Sky_Stars"));
 	return Texture;
 }
+
+template <class ValueType>
+void FGLTFSkySphereConverter::ConvertProperty(const AActor& Actor, const TCHAR* PropertyName, const TCHAR* ExportedPropertyName, ValueType& OutValue) const
+{
+	if (!FGLTFActorUtility::TryGetPropertyValue(&Actor, PropertyName, OutValue))
+	{
+		Builder.AddWarningMessage(FString::Printf(
+			TEXT("Failed to export %s for Sky Sphere %s"),
+			ExportedPropertyName,
+			*Actor.GetName()));
+	}
+}
+
+void FGLTFSkySphereConverter::ConvertColorProperty(const AActor& Actor, const TCHAR* PropertyName, const TCHAR* ExportedPropertyName, FGLTFJsonColor4& OutValue) const
+{
+	FLinearColor LinearColor;
+	if (FGLTFActorUtility::TryGetPropertyValue(&Actor, PropertyName, LinearColor))
+	{
+		OutValue = FGLTFConverterUtility::ConvertColor(LinearColor);
+	}
+	else
+	{
+		Builder.AddWarningMessage(FString::Printf(
+			TEXT("Failed to export %s for Sky Sphere %s"),
+			ExportedPropertyName,
+			*Actor.GetName()));
+	}
+}
+
+void FGLTFSkySphereConverter::ConvertScalarParameter(const AActor& Actor, const UMaterialInstance* Material, const TCHAR* ParameterName, const TCHAR* ExportedPropertyName, float& OutValue) const
+{
+	if (Material == nullptr || !Material->GetScalarParameterValue(ParameterName, OutValue))
+	{
+		Builder.AddWarningMessage(FString::Printf(
+			TEXT("Failed to export %s for Sky Sphere %s"),
+			ExportedPropertyName,
+			*Actor.GetName()));
+	}
+}
diff --git a/Engine/Plugins/Enterprise/GLTFExporter/MigratedFromGit/Source/GLTFExporter/Private/Converters/GLTFSkySphereConverters.h b/Engine/Plugins/Enterprise/GLTFExporter/MigratedFromGit/Source/GLTFExporter/Private/Converters/GLTFSkySphereConverters.h
--- a/Engine/Plugins/Enterprise/GLTFExporter/MigratedFromGit/Source/GLTFExporter/Private/Converters/GLTFSkySphereConverters.h
+++ b/Engine/Plugins/Enterprise/GLTFExporter/MigratedFromGit/Source/GLTFExporter/Private/Converters/GLTFSkySphereConverters.h
@@ -7,6 +7,8 @@
 #include "Converters/GLTFBuilderContext.h"
 #include "Engine.h"
 
+struct FGLTFJsonColor4;
+
 class FGLTFSkySphereConverter final : public FGLTFBuilderContext, public TGLTFConverter<FGLTFJsonSkySphereIndex, const AActor*>
 {
 	using FGLTFBuilderContext::FGLTFBuilderContext;
@@ -16,4 +18,10 @@ class FGLTFSkySphereConverter final : public FGLTFBuilderContext, public TGLTFCo
 	const UTexture2D* GetSkyTexture(const UMaterialInstance* SkyMaterial) const;
 	const UTexture2D* GetCloudsTexture(const UMaterialInstance* SkyMaterial) const;
 	const UTexture2D* GetStarsTexture(const UMaterialInstance* SkyMaterial) const;
+
+	template <class ValueType>
+	void ConvertProperty(const AActor& Actor, const TCHAR* PropertyName, const TCHAR* ExportedPropertyName, ValueType& OutValue) const;
+
+	void ConvertColorProperty(const AActor& Actor, const TCHAR* PropertyName, const TCHAR* ExportedPropertyName, FGLTFJsonColor4& OutValue) const;
+	void ConvertScalarParameter(const AActor& Actor, const UMaterialInstance* Material, const TCHAR* ParameterName, const TCHAR* ExportedPropertyName, float& OutValue) const;
 };
