[GLTFExporter] feat(sky-sphere): add support for exporting color curves

commit 6d1a8fda69c174ba9b5f95814a9c376bfeb4aeb7
Date: Wed, 14 Oct 2020 14:59:02 +0200

    feat(sky-sphere): add support for exporting color curves
    (cherry picked from commit 65c7876d2196bb13078c8a80e187aa3bcb796c8b​)

#skipci

diff --git a/Engine/Plugins/Enterprise/GLTFExporter/MigratedFromGit/Source/GLTFExporter/Private/Converters/GLTFCurveUtility.cpp b/Engine/Plugins/Enterprise/GLTFExporter/MigratedFromGit/Source/GLTFExporter/Private/Converters/GLTFCurveUtility.cpp
new file mode 100644
--- /dev/null
+++ b/Engine/Plugins/Enterprise/GLTFExporter/MigratedFromGit/Source/GLTFExporter/Private/Converters/GLTFCurveUtility.cpp
@@ -0,0 +1,17 @@
+// Copyright Epic Games, Inc. All Rights Reserved.
+
+#include "Converters/GLTFCurveUtility.h"
+#include "Curves/CurveLinearColor.h"
+
+bool FGLTFCurveUtility::HasAnyAdjustment(const UCurveLinearColor& ColorCurve)
+{
+	const float ErrorTolerance = KINDA_SMALL_NUMBER;
+
+	return !FMath::IsNearlyEqual(ColorCurve.AdjustBrightness, 1.0f, ErrorTolerance) ||
+		!FMath::IsNearlyEqual(ColorCurve.AdjustHue, 0.0f, ErrorTolerance) ||
+		!FMath::IsNearlyEqual(ColorCurve.AdjustSaturation, 1.0f, ErrorTolerance) ||
+		!FMath::IsNearlyEqual(ColorCurve.AdjustVibrance, 0.0f, ErrorTolerance) ||
+		!FMath::IsNearlyEqual(ColorCurve.AdjustBrightnessCurve, 1.0f, ErrorTolerance) ||
+		!FMath::IsNearlyEqual(ColorCurve.AdjustMaxAlpha, 1.0f, ErrorTolerance) ||
+		!FMath::IsNearlyEqual(ColorCurve.AdjustMinAlpha, 0.0f, ErrorTolerance);
+}
diff --git a/Engine/Plugins/Enterprise/GLTFExporter/MigratedFromGit/Source/GLTFExporter/Private/Converters/GLTFCurveUtility.h b/Engine/Plugins/Enterprise/GLTFExporter/MigratedFromGit/Source/GLTFExporter/Private/Converters/GLTFCurveUtility.h
new file mode 100644
--- /dev/null
+++ b/Engine/Plugins/Enterprise/GLTFExporter/MigratedFromGit/Source/GLTFExporter/Private/Converters/GLTFCurveUtility.h
@@ -0,0 +1,12 @@
+// Copyright Epic Games, Inc. All Rights Reserved.
+
+#pragma once
+
+#include "Engine/EngineTypes.h"
+
+class UCurveLinearColor;
+
+struct FGLTFCurveUtility
+{
+	static bool HasAnyAdjustment(const UCurveLinearColor& ColorCurve);
+};
diff --git a/Engine/Plugins/Enterprise/GLTFExporter/MigratedFromGit/Source/GLTFExporter/Private/Converters/GLTFSkySphereConverters.cpp b/Engine/Plugins/Enterprise/GLTFExporter/MigratedFromGit/Source/GLTFExporter/Private/Converters/GLTFSkySphereConverters.cpp
--- a/Engine/Plugins/Enterprise/GLTFExporter/MigratedFromGit/Source/GLTFExporter/Private/Converters/GLTFSkySphereConverters.cpp
+++ b/Engine/Plugins/Enterprise/GLTFExporter/MigratedFromGit/Source/GLTFExporter/Private/Converters/GLTFSkySphereConverters.cpp
@@ -4,6 +4,7 @@
 #include "Converters/GLTFConverterUtility.h"
 #include "Converters/GLTFMaterialUtility.h"
 #include "Converters/GLTFActorUtility.h"
+#include "Converters/GLTFCurveUtility.h"
 #include "Builders/GLTFContainerBuilder.h"
 
 FGLTFJsonSkySphereIndex FGLTFSkySphereConverter::Convert(const AActor* SkySphereActor)
@@ -103,7 +104,9 @@ FGLTFJsonSkySphereIndex FGLTFSkySphereConverter::Convert(const AActor* SkySphere
 	ConvertColorProperty(*SkySphereActor, TEXT("Cloud color"), TEXT("CloudColor"), JsonSkySphere.CloudColor);
 	ConvertColorProperty(*SkySphereActor, TEXT("Overall Color"), TEXT("OverallColor"), JsonSkySphere.OverallColor);
 
-	// TODO: export color curves
+	ConvertColorCurveProperty(*SkySphereActor, TEXT("Zenith color curve"), TEXT("ZenithColorCurve"), JsonSkySphere.ZenithColorCurve);
+	ConvertColorCurveProperty(*SkySphereActor, TEXT("Horizon color curve"), TEXT("HorizonColorCurve"), JsonSkySphere.HorizonColorCurve);
+	ConvertColorCurveProperty(*SkySphereActor, TEXT("Cloud color curve"), TEXT("CloudColorCurve"), JsonSkySphere.CloudColorCurve);
 
 	return Builder.AddSkySphere(JsonSkySphere);
 }
@@ -163,6 +166,48 @@ void FGLTFSkySphereConverter::ConvertColorProperty(const AActor& Actor, const TC
 	}
 }
 
+void FGLTFSkySphereConverter::ConvertColorCurveProperty(const AActor& Actor, const TCHAR* PropertyName, const TCHAR* ExportedPropertyName, FGLTFJsonSkySphereColorCurve& OutValue) const
+{
+	const UCurveLinearColor* ColorCurve = nullptr;
+	FGLTFActorUtility::TryGetPropertyValue(&Actor, PropertyName, ColorCurve);
+
+	if (ColorCurve != nullptr)
+	{
+		if (FGLTFCurveUtility::HasAnyAdjustment(*ColorCurve))
+		{
+			Builder.AddWarningMessage(FString::Printf(
+				TEXT("Adjustments for %s in Sky Sphere %s are not supported and will be ignored"),
+				ExportedPropertyName,
+				*Actor.GetName()));
+		}
+
+		OutValue.ComponentCurves.AddDefaulted(3);
+
+		for (uint32 ComponentIndex = 0; ComponentIndex < 3; ++ComponentIndex)
+		{
+			const FRichCurve& FloatCurve = ColorCurve->FloatCurves[ComponentIndex];
+			const uint32 KeyCount = FloatCurve.Keys.Num();
+
+			TArray<FGLTFJsonSkySphereColorCurve::FKey>& ComponentKeys = OutValue.ComponentCurves[ComponentIndex].Keys;
+			ComponentKeys.AddUninitialized(KeyCount);
+
+			uint32 KeyIndex = 0;
+
+			for (const FRichCurveKey& Key: FloatCurve.Keys)
+			{
+				ComponentKeys[KeyIndex++] = { Key.Time, Key.Value };
+			}
+		}
+	}
+	else
+	{
+		Builder.AddWarningMessage(FString::Printf(
+			TEXT("Failed to export %s for Sky Sphere %s"),
+			ExportedPropertyName,
+			*Actor.GetName()));
+	}
+}
+
 void FGLTFSkySphereConverter::ConvertScalarParameter(const AActor& Actor, const UMaterialInstance* Material, const TCHAR* ParameterName, const TCHAR* ExportedPropertyName, float& OutValue) const
 {
 	if (Material == nullptr || !Material->GetScalarParameterValue(ParameterName, OutValue))
diff --git a/Engine/Plugins/Enterprise/GLTFExporter/MigratedFromGit/Source/GLTFExporter/Private/Converters/GLTFSkySphereConverters.h b/Engine/Plugins/Enterprise/GLTFExporter/MigratedFromGit/Source/GLTFExporter/Private/Converters/GLTFSkySphereConverters.h
--- a/Engine/Plugins/Enterprise/GLTFExporter/MigratedFromGit/Source/GLTFExporter/Private/Converters/GLTFSkySphereConverters.h
+++ b/Engine/Plugins/Enterprise/GLTFExporter/MigratedFromGit/Source/GLTFExporter/Private/Converters/GLTFSkySphereConverters.h
@@ -8,6 +8,7 @@
 #include "Engine.h"
 
 struct FGLTFJsonColor4;
+struct FGLTFJsonSkySphereColorCurve;
 
 class FGLTFSkySphereConverter final : public FGLTFBuilderContext, public TGLTFConverter<FGLTFJsonSkySphereIndex, const AActor*>
 {
@@ -23,5 +24,6 @@ class FGLTFSkySphereConverter final : public FGLTFBuilderContext, public TGLTFCo
 	void ConvertProperty(const AActor& Actor, const TCHAR* PropertyName, const TCHAR* ExportedPropertyName, ValueType& OutValue) const;
 
 	void ConvertColorProperty(const AActor& Actor, const TCHAR* PropertyName, const TCHAR* ExportedPropertyName, FGLTFJsonColor4& OutValue) const;
+	void ConvertColorCurveProperty(const AActor& Actor, const TCHAR* PropertyName, const TCHAR* ExportedPropertyName, FGLTFJsonSkySphereColorCurve& OutValue) const;
 	void ConvertScalarParameter(const AActor& Actor, const UMaterialInstance* Material, const TCHAR* ParameterName, const TCHAR* ExportedPropertyName, float& OutValue) const;
 };
diff --git a/Engine/Plugins/Enterprise/GLTFExporter/MigratedFromGit/Source/GLTFExporter/Private/Json/GLTFJsonSkySphere.h b/Engine/Plugins/Enterprise/GLTFExporter/MigratedFromGit/Source/GLTFExporter/Private/Json/GLTFJsonSkySphere.h
--- a/Engine/Plugins/Enterprise/GLTFExporter/MigratedFromGit/Source/GLTFExporter/Private/Json/GLTFJsonSkySphere.h
+++ b/Engine/Plugins/Enterprise/GLTFExporter/MigratedFromGit/Source/GLTFExporter/Private/Json/GLTFJsonSkySphere.h
@@ -7,6 +7,45 @@
 #include "Json/GLTFJsonExtensions.h"
 #include "Serialization/JsonSerializer.h"
 
+struct FGLTFJsonSkySphereColorCurve
+{
+	struct FKey
+	{
+		float Time;
+		float Value;
+	};
+
+	struct FComponentCurve
+	{
+		TArray<FKey> Keys;
+	};
+
+	TArray<FComponentCurve> ComponentCurves;
+
+	template <class CharType = TCHAR, class PrintPolicy = TPrettyJsonPrintPolicy<CharType>>
+	void WriteArray(TJsonWriter<CharType, PrintPolicy>& JsonWriter) const
+	{
+		JsonWriter.WriteArrayStart();
+
+		for (int32 ComponentIndex = 0; ComponentIndex < ComponentCurves.Num(); ++ComponentIndex)
+		{
+			const FComponentCurve& ComponentCurve = ComponentCurves[ComponentIndex];
+
+			JsonWriter.WriteArrayStart();
+
+			for (const FKey& Key: ComponentCurve.Keys)
+			{
+				FGLTFJsonUtility::WriteExactValue(JsonWriter, Key.Time);
+				FGLTFJsonUtility::WriteExactValue(JsonWriter, Key.Value);
+			}
+
+			JsonWriter.WriteArrayEnd();
+		}
+
+		JsonWriter.WriteArrayEnd();
+	}
+};
+
 struct FGLTFJsonSkySphere
 {
 	FString Name;
@@ -35,9 +74,11 @@ struct FGLTFJsonSkySphere
 	FGLTFJsonColor4 CloudColor;
 	FGLTFJsonColor4 OverallColor;
 
-	FGLTFJsonVector3 Scale;
+	FGLTFJsonSkySphereColorCurve ZenithColorCurve;
+	FGLTFJsonSkySphereColorCurve HorizonColorCurve;
+	FGLTFJsonSkySphereColorCurve CloudColorCurve;
 
-	// TODO: add properties for color curves
+	FGLTFJsonVector3 Scale;
 
 	FGLTFJsonSkySphere()
 		: SunHeight(0)
@@ -106,14 +147,30 @@ struct FGLTFJsonSkySphere
 			OverallColor.WriteArray(JsonWriter);
 		}
 
+		if (ZenithColorCurve.ComponentCurves.Num() >= 3)
+		{
+			JsonWriter.WriteIdentifierPrefix(TEXT("zenithColorCurve"));
+			ZenithColorCurve.WriteArray(JsonWriter);
+		}
+
+		if (HorizonColorCurve.ComponentCurves.Num() >= 3)
+		{
+			JsonWriter.WriteIdentifierPrefix(TEXT("horizonColorCurve"));
+			HorizonColorCurve.WriteArray(JsonWriter);
+		}
+
+		if (CloudColorCurve.ComponentCurves.Num() >= 3)
+		{
+			JsonWriter.WriteIdentifierPrefix(TEXT("cloudColorCurve"));
+			CloudColorCurve.WriteArray(JsonWriter);
+		}
+
 		if (Scale != FGLTFJsonVector3::One)
 		{
 			JsonWriter.WriteIdentifierPrefix(TEXT("scale"));
 			Scale.WriteArray(JsonWriter);
 		}
 
-		// TODO: export color curves
-
 		JsonWriter.WriteObjectEnd();
 	}
 };
