[GLTFExporter] refactor(sky-sphere): change wording of waning message for curve adjustments

commit 007c4f87d3037f7f0aa00aad1f74409f7d93ef19
Date: Thu, 15 Oct 2020 08:23:51 +0200

    refactor(sky-sphere): change wording of waning message for curve adjustments
    (cherry picked from commit 3e8c0be48c57934313d471ed546fc637f0e84350​)

#skipci

diff --git a/Engine/Plugins/Enterprise/GLTFExporter/MigratedFromGit/Source/GLTFExporter/Private/Converters/GLTFSkySphereConverters.cpp b/Engine/Plugins/Enterprise/GLTFExporter/MigratedFromGit/Source/GLTFExporter/Private/Converters/GLTFSkySphereConverters.cpp
--- a/Engine/Plugins/Enterprise/GLTFExporter/MigratedFromGit/Source/GLTFExporter/Private/Converters/GLTFSkySphereConverters.cpp
+++ b/Engine/Plugins/Enterprise/GLTFExporter/MigratedFromGit/Source/GLTFExporter/Private/Converters/GLTFSkySphereConverters.cpp
@@ -176,7 +176,7 @@ void FGLTFSkySphereConverter::ConvertColorCurveProperty(const AActor& Actor, con
 		if (FGLTFCurveUtility::HasAnyAdjustment(*ColorCurve))
 		{
 			Builder.AddWarningMessage(FString::Printf(
-				TEXT("Adjustments for %s in Sky Sphere %s are not supported and will be ignored"),
+				TEXT("Adjustments for %s in Sky Sphere %s are not supported"),
 				ExportedPropertyName,
 				*Actor.GetName()));
 		}
