[GLTFExporter] refactor: improve warning message when converting material with multiple shading models, by naming each.

commit 370952da61293cb9bfbd3b3c0fef410efe19f2e9
Date: Wed, 14 Oct 2020 21:07:35 +0200

    refactor: improve warning message when converting material with multiple shading models, by naming each.
    (cherry picked from commit 542ee1304f04493e0d9d41cc18c810336f519690​)

#skipci

diff --git a/Engine/Plugins/Enterprise/GLTFExporter/MigratedFromGit/Source/GLTFExporter/Private/Converters/GLTFMaterialUtility.cpp b/Engine/Plugins/Enterprise/GLTFExporter/MigratedFromGit/Source/GLTFExporter/Private/Converters/GLTFMaterialUtility.cpp
--- a/Engine/Plugins/Enterprise/GLTFExporter/MigratedFromGit/Source/GLTFExporter/Private/Converters/GLTFMaterialUtility.cpp
+++ b/Engine/Plugins/Enterprise/GLTFExporter/MigratedFromGit/Source/GLTFExporter/Private/Converters/GLTFMaterialUtility.cpp
@@ -2,6 +2,7 @@
 
 #include "Converters/GLTFMaterialUtility.h"
 #include "Converters/GLTFTextureUtility.h"
+#include "Converters/GLTFNameUtility.h"
 #include "Engine/TextureRenderTarget2D.h"
 #include "CanvasItem.h"
 #include "CanvasTypes.h"
@@ -461,3 +462,20 @@ void FGLTFMaterialUtility::ExpandAllFunctionExpressions(TArray<UMaterialExpressi
 		return Expression->IsA<UMaterialExpressionMaterialFunctionCall>() || Expression->IsA<UMaterialExpressionMaterialAttributeLayers>();
 	});
 }
+
+FString FGLTFMaterialUtility::ShadingModelsToString(const FMaterialShadingModelField& ShadingModels)
+{
+	FString Result;
+
+	for (uint32 Index = 0; Index < MSM_MAX; Index++)
+	{
+		const EMaterialShadingModel ShadingModel = static_cast<EMaterialShadingModel>(Index);
+		if (ShadingModels.HasShadingModel(ShadingModel))
+		{
+			FString Name = FGLTFNameUtility::GetName(ShadingModel);
+			Result += Result.IsEmpty() ? Name : TEXT(", ") + Name;
+		}
+	}
+
+	return Result;
+}
diff --git a/Engine/Plugins/Enterprise/GLTFExporter/MigratedFromGit/Source/GLTFExporter/Private/Converters/GLTFMaterialUtility.h b/Engine/Plugins/Enterprise/GLTFExporter/MigratedFromGit/Source/GLTFExporter/Private/Converters/GLTFMaterialUtility.h
--- a/Engine/Plugins/Enterprise/GLTFExporter/MigratedFromGit/Source/GLTFExporter/Private/Converters/GLTFMaterialUtility.h
+++ b/Engine/Plugins/Enterprise/GLTFExporter/MigratedFromGit/Source/GLTFExporter/Private/Converters/GLTFMaterialUtility.h
@@ -65,4 +65,6 @@ struct FGLTFMaterialUtility
 	static void GetAllTextureCoordinateIndices(const FExpressionInput& Input, TSet<int32>& OutTexCoords);
 
 	static void ExpandAllFunctionExpressions(TArray<UMaterialExpression*>& InOutExpressions);
+
+	static FString ShadingModelsToString(const FMaterialShadingModelField& ShadingModels);
 };
diff --git a/Engine/Plugins/Enterprise/GLTFExporter/MigratedFromGit/Source/GLTFExporter/Private/Tasks/GLTFMaterialTasks.cpp b/Engine/Plugins/Enterprise/GLTFExporter/MigratedFromGit/Source/GLTFExporter/Private/Tasks/GLTFMaterialTasks.cpp
--- a/Engine/Plugins/Enterprise/GLTFExporter/MigratedFromGit/Source/GLTFExporter/Private/Tasks/GLTFMaterialTasks.cpp
+++ b/Engine/Plugins/Enterprise/GLTFExporter/MigratedFromGit/Source/GLTFExporter/Private/Tasks/GLTFMaterialTasks.cpp
@@ -184,7 +184,7 @@ bool FGLTFMaterialTask::TryGetShadingModel(EGLTFJsonShadingModel& OutShadingMode
 
 	if (ShadingModelCount > 1)
 	{
-		Builder.AddWarningMessage(FString::Printf(TEXT("Support is limited to the first of multiple shading models found (%d) in material %s"), ShadingModelCount, *Material->GetName()));
+		Builder.AddWarningMessage(FString::Printf(TEXT("Support is limited to the first of multiple shading models found (%s) in material %s"), *FGLTFMaterialUtility::ShadingModelsToString(ShadingModels), *Material->GetName()));
 	}
 
 	const EMaterialShadingModel ShadingModel = ShadingModels.GetFirstShadingModel();
