[GLTFExporter] feat: improve support for materials with multiple shading models

commit af790e068daff322c3aabf9525870dd7021b4e5f
Date: Wed, 14 Oct 2020 21:24:59 +0200

    feat: improve support for materials with multiple shading models
    (cherry picked from commit a64bb5ea0392e987bf43d03986568ae3757c8795​)

#skipci

diff --git a/Engine/Plugins/Enterprise/GLTFExporter/MigratedFromGit/Source/GLTFExporter/Private/Tasks/GLTFMaterialTasks.cpp b/Engine/Plugins/Enterprise/GLTFExporter/MigratedFromGit/Source/GLTFExporter/Private/Tasks/GLTFMaterialTasks.cpp
--- a/Engine/Plugins/Enterprise/GLTFExporter/MigratedFromGit/Source/GLTFExporter/Private/Tasks/GLTFMaterialTasks.cpp
+++ b/Engine/Plugins/Enterprise/GLTFExporter/MigratedFromGit/Source/GLTFExporter/Private/Tasks/GLTFMaterialTasks.cpp
@@ -182,14 +182,28 @@ bool FGLTFMaterialTask::TryGetShadingModel(EGLTFJsonShadingModel& OutShadingMode
 		return false;
 	}
 
+	EMaterialShadingModel ShadingModel = ShadingModels.GetFirstShadingModel();
 	if (ShadingModelCount > 1)
 	{
-		Builder.AddWarningMessage(FString::Printf(TEXT("Support is limited to the first of multiple shading models found (%s) in material %s"), *FGLTFMaterialUtility::ShadingModelsToString(ShadingModels), *Material->GetName()));
+		Builder.AddWarningMessage(FString::Printf(TEXT("Support is limited to only one of the multiple shading models found (%s) in material %s"), *FGLTFMaterialUtility::ShadingModelsToString(ShadingModels), *Material->GetName()));
+
+		if (ShadingModels.HasShadingModel(MSM_ClearCoat))
+		{
+			ShadingModel = MSM_ClearCoat;
+		}
+		else if (ShadingModels.HasShadingModel(MSM_DefaultLit))
+		{
+			ShadingModel = MSM_DefaultLit;
+		}
+		else if (ShadingModels.HasShadingModel(MSM_Unlit))
+		{
+			ShadingModel = MSM_Unlit;
+		}
+
+		Builder.AddWarningMessage(FString::Printf(TEXT("Material %s will be exported as shading model %s"), *Material->GetName(), *FGLTFNameUtility::GetName(ShadingModel)));
 	}
 
-	const EMaterialShadingModel ShadingModel = ShadingModels.GetFirstShadingModel();
 	const EGLTFJsonShadingModel ConvertedShadingModel = FGLTFConverterUtility::ConvertShadingModel(ShadingModel);
-
 	if (ConvertedShadingModel == EGLTFJsonShadingModel::None)
 	{
 		const FString ShadingModelName = FGLTFNameUtility::GetName(ShadingModel);
