[GLTFExporter] feat: decode normalmaps when rendering via FCanvasTileItem

commit 4c21416db3475bbca04f84a9a409665dd93ff97e
Date: Fri, 16 Oct 2020 09:00:50 +0200

    feat: decode normalmaps when rendering via FCanvasTileItem
    renders normalmaps as decoded RGB instead of as packed data
    (cherry picked from commit 84900fb85133cea51331044404131e1779c1954e​)

#skipci

diff --git a/Engine/Plugins/Enterprise/GLTFExporter/MigratedFromGit/Source/GLTFExporter/Private/Converters/GLTFMaterialUtility.cpp b/Engine/Plugins/Enterprise/GLTFExporter/MigratedFromGit/Source/GLTFExporter/Private/Converters/GLTFMaterialUtility.cpp
--- a/Engine/Plugins/Enterprise/GLTFExporter/MigratedFromGit/Source/GLTFExporter/Private/Converters/GLTFMaterialUtility.cpp
+++ b/Engine/Plugins/Enterprise/GLTFExporter/MigratedFromGit/Source/GLTFExporter/Private/Converters/GLTFMaterialUtility.cpp
@@ -9,6 +9,7 @@
 #include "Modules/ModuleManager.h"
 #include "IMaterialBakingModule.h"
 #include "MaterialBakingStructures.h"
+#include "NormalMapPreview.h"
 #include "Materials/MaterialExpressionCustomOutput.h"
 #include "Materials/MaterialExpressionClearCoatNormalCustomOutput.h"
 #include "Materials/MaterialExpressionTextureCoordinate.h"
@@ -210,8 +211,16 @@ bool FGLTFMaterialUtility::CombineTextures(TArray<FColor>& OutPixels, const TArr
 	{
 		check(Source.Texture);
 
+		TRefCountPtr<FBatchedElementParameters> BatchedElementParameters;
+
+		if (Source.Texture->IsNormalMap())
+		{
+			BatchedElementParameters = new FNormalMapBatchedElementParameters();
+		}
+
 		FCanvasTileItem TileItem(TilePosition, Source.Texture->Resource, TileSize, Source.TintColor);
 
+		TileItem.BatchedElementParameters = BatchedElementParameters;
 		TileItem.BlendMode = Source.BlendMode;
 		TileItem.Draw(&Canvas);
 	}
diff --git a/Engine/Plugins/Enterprise/GLTFExporter/MigratedFromGit/Source/GLTFExporter/Private/Converters/GLTFTextureUtility.cpp b/Engine/Plugins/Enterprise/GLTFExporter/MigratedFromGit/Source/GLTFExporter/Private/Converters/GLTFTextureUtility.cpp
--- a/Engine/Plugins/Enterprise/GLTFExporter/MigratedFromGit/Source/GLTFExporter/Private/Converters/GLTFTextureUtility.cpp
+++ b/Engine/Plugins/Enterprise/GLTFExporter/MigratedFromGit/Source/GLTFExporter/Private/Converters/GLTFTextureUtility.cpp
@@ -2,6 +2,7 @@
 
 #include "Converters/GLTFTextureUtility.h"
 #include "IImageWrapper.h"
+#include "NormalMapPreview.h"
 
 bool FGLTFTextureUtility::IsHDRFormat(EPixelFormat Format)
 {
@@ -151,8 +152,16 @@ bool FGLTFTextureUtility::DrawTexture(UTextureRenderTarget2D* OutTarget, const U
 	const_cast<UTexture2D*>(InSource)->SetForceMipLevelsToBeResident(30.0f, true);
 	const_cast<UTexture2D*>(InSource)->WaitForStreaming();
 
+	TRefCountPtr<FBatchedElementParameters> BatchedElementParameters;
+
+	if (InSource->IsNormalMap())
+	{
+		BatchedElementParameters = new FNormalMapBatchedElementParameters();
+	}
+
 	FCanvas Canvas(RenderTarget, nullptr, 0.0f, 0.0f, 0.0f, GMaxRHIFeatureLevel);
 	FCanvasTileItem TileItem(FVector2D::ZeroVector, InSource->Resource, FLinearColor::White);
+	TileItem.BatchedElementParameters = BatchedElementParameters;
 
 	Canvas.PushAbsoluteTransform(InTransform);
 	TileItem.Draw(&Canvas);
