[GLTFExporter] feat: use render data instead of source data when exporting normalmaps

commit f34e4814d33ce41b54d801fe93b6239c3f2fd08e
Date: Fri, 16 Oct 2020 13:42:51 +0200

    feat: use render data instead of source data when exporting normalmaps
    (cherry picked from commit 2a5b8193f4e21069a862cba89a675baddfdb45ca​)

#skipci

diff --git a/Engine/Plugins/Enterprise/GLTFExporter/MigratedFromGit/Source/GLTFExporter/Private/Tasks/GLTFTextureTasks.cpp b/Engine/Plugins/Enterprise/GLTFExporter/MigratedFromGit/Source/GLTFExporter/Private/Tasks/GLTFTextureTasks.cpp
--- a/Engine/Plugins/Enterprise/GLTFExporter/MigratedFromGit/Source/GLTFExporter/Private/Tasks/GLTFTextureTasks.cpp
+++ b/Engine/Plugins/Enterprise/GLTFExporter/MigratedFromGit/Source/GLTFExporter/Private/Tasks/GLTFTextureTasks.cpp
@@ -12,14 +12,14 @@ void FGLTFTexture2DTask::Complete()
 	FGLTFJsonTexture& JsonTexture = Builder.GetTexture(TextureIndex);
 	Texture2D->GetName(JsonTexture.Name);
 
-	const bool bPreferSourceExport = Builder.ExportOptions->bExportSourceTextures;
+	// NOTE: export of normalmaps via source data is disabled because we need to manipulate the pixels during conversion.
+	const bool bPreferSourceExport = Texture2D->IsNormalMap() ? false : Builder.ExportOptions->bExportSourceTextures;
 
-	// NOTE: export of lightmaps and normalmaps via source data is done to workaround issues
-	// with using render data for normalmaps (which are compressed to red & green channels) and
-	// lightmaps (which can otherwise suffer quality-loss due to incorrect gamma transformation).
+	// NOTE: export of lightmaps via source data is used to work around issues
+	// with quality-loss due to incorrect gamma transformation when rendering to a canvas.
 
-	// TODO: make sure there are no other special cases than normalmaps and lightmaps that require source-export
-	const bool bRequireSourceExport = Texture2D->IsA<ULightMapTexture2D>() || Texture2D->IsNormalMap();
+	// TODO: make sure there are no other special cases than lightmaps that require source-export
+	const bool bRequireSourceExport = Texture2D->IsA<ULightMapTexture2D>();
 
 	if (bRequireSourceExport || bPreferSourceExport)
 	{
