[GLTFExporter] feat: invert the green channel of all normalmaps to match WebGL standard

commit 798f61b1a9e61a21e7846ac4c22d20d562f4ac7a
Date: Fri, 16 Oct 2020 14:55:35 +0200

    feat: invert the green channel of all normalmaps to match WebGL standard
    (cherry picked from commit d18240bcf5435f8716bf5a02c38f2e442293ef39​)

#skipci

diff --git a/Engine/Plugins/Enterprise/GLTFExporter/MigratedFromGit/Source/GLTFExporter/Private/Converters/GLTFMaterialUtility.cpp b/Engine/Plugins/Enterprise/GLTFExporter/MigratedFromGit/Source/GLTFExporter/Private/Converters/GLTFMaterialUtility.cpp
--- a/Engine/Plugins/Enterprise/GLTFExporter/MigratedFromGit/Source/GLTFExporter/Private/Converters/GLTFMaterialUtility.cpp
+++ b/Engine/Plugins/Enterprise/GLTFExporter/MigratedFromGit/Source/GLTFExporter/Private/Converters/GLTFMaterialUtility.cpp
@@ -315,6 +315,13 @@ FGLTFPropertyBakeOutput FGLTFMaterialUtility::BakeMaterialProperty(const FIntPoi
 		}
 	}
 
+	// TODO: add condition for ClearCoat normal property when it's been added to MaterialBaking module
+	if (Property == MP_Normal)
+	{
+		// Convert normalmaps to use +Y (OpenGL / WebGL standard)
+		FGLTFTextureUtility::FlipGreenChannel(BakedPixels);
+	}
+
 	FGLTFPropertyBakeOutput PropertyBakeOutput(Property, PF_B8G8R8A8, BakedPixels, BakedSize, EmissiveScale);
 
 	if (BakedPixels.Num() == 1)
diff --git a/Engine/Plugins/Enterprise/GLTFExporter/MigratedFromGit/Source/GLTFExporter/Private/Converters/GLTFTextureUtility.cpp b/Engine/Plugins/Enterprise/GLTFExporter/MigratedFromGit/Source/GLTFExporter/Private/Converters/GLTFTextureUtility.cpp
--- a/Engine/Plugins/Enterprise/GLTFExporter/MigratedFromGit/Source/GLTFExporter/Private/Converters/GLTFTextureUtility.cpp
+++ b/Engine/Plugins/Enterprise/GLTFExporter/MigratedFromGit/Source/GLTFExporter/Private/Converters/GLTFTextureUtility.cpp
@@ -337,3 +337,19 @@ bool FGLTFTextureUtility::LoadPlatformData(UTextureCube* TextureCube)
 
 	return true;
 }
+
+void FGLTFTextureUtility::FlipGreenChannel(TArray<FLinearColor>& Pixels)
+{
+	for (FLinearColor& Pixel: Pixels)
+	{
+		Pixel.G = 1.0f - FMath::Clamp(Pixel.G, 0.0f, 1.0f);
+	}
+}
+
+void FGLTFTextureUtility::FlipGreenChannel(TArray<FColor>& Pixels)
+{
+	for (FColor& Pixel: Pixels)
+	{
+		Pixel.G = 255 - Pixel.G;
+	}
+}
diff --git a/Engine/Plugins/Enterprise/GLTFExporter/MigratedFromGit/Source/GLTFExporter/Private/Converters/GLTFTextureUtility.h b/Engine/Plugins/Enterprise/GLTFExporter/MigratedFromGit/Source/GLTFExporter/Private/Converters/GLTFTextureUtility.h
--- a/Engine/Plugins/Enterprise/GLTFExporter/MigratedFromGit/Source/GLTFExporter/Private/Converters/GLTFTextureUtility.h
+++ b/Engine/Plugins/Enterprise/GLTFExporter/MigratedFromGit/Source/GLTFExporter/Private/Converters/GLTFTextureUtility.h
@@ -43,4 +43,7 @@ struct FGLTFTextureUtility
 	// TODO: maybe use template specialization to avoid the need for duplicated functions
 	static bool LoadPlatformData(UTexture2D* Texture);
 	static bool LoadPlatformData(UTextureCube* TextureCube);
+
+	static void FlipGreenChannel(TArray<FLinearColor>& Pixels);
+	static void FlipGreenChannel(TArray<FColor>& Pixels);
 };
diff --git a/Engine/Plugins/Enterprise/GLTFExporter/MigratedFromGit/Source/GLTFExporter/Private/Tasks/GLTFTextureTasks.cpp b/Engine/Plugins/Enterprise/GLTFExporter/MigratedFromGit/Source/GLTFExporter/Private/Tasks/GLTFTextureTasks.cpp
--- a/Engine/Plugins/Enterprise/GLTFExporter/MigratedFromGit/Source/GLTFExporter/Private/Tasks/GLTFTextureTasks.cpp
+++ b/Engine/Plugins/Enterprise/GLTFExporter/MigratedFromGit/Source/GLTFExporter/Private/Tasks/GLTFTextureTasks.cpp
@@ -72,6 +72,19 @@ void FGLTFTexture2DTask::Complete()
 		TArray<FColor> Pixels;
 		if (FGLTFTextureUtility::ReadEncodedPixels(RenderTarget, Pixels, JsonTexture.Encoding)) // TODO: use only encoding as specified by export options
 		{
+			if (Texture2D->IsNormalMap())
+			{
+				if (JsonTexture.Encoding == EGLTFJsonHDREncoding::None)
+				{
+					// Convert normalmaps to use +Y (OpenGL / WebGL standard)
+					FGLTFTextureUtility::FlipGreenChannel(Pixels);
+				}
+				else
+				{
+					Builder.AddWarningMessage(FString::Printf(TEXT("Failed to invert green channel for normalmap %s"), *JsonTexture.Name));
+				}
+			}
+
 			ImageIndex = Builder.AddImage(Pixels.GetData(), Size, JsonTexture.Name);
 		}
 	}
