[GLTFExporter] refactor(sky-sphere): make conversion more modular and improve warnings

commit 4daa3e88a65e0d80639cd0970cb9372646e53d04
Date: Fri, 16 Oct 2020 17:19:46 +0200

    refactor(sky-sphere): make conversion more modular and improve warnings
    (cherry picked from commit 63c3cb24b802e9b4b750cd77b63be612781b3f3b​)

#skipci

diff --git a/Engine/Plugins/Enterprise/GLTFExporter/MigratedFromGit/Source/GLTFExporter/Private/Converters/GLTFSkySphereConverters.cpp b/Engine/Plugins/Enterprise/GLTFExporter/MigratedFromGit/Source/GLTFExporter/Private/Converters/GLTFSkySphereConverters.cpp
--- a/Engine/Plugins/Enterprise/GLTFExporter/MigratedFromGit/Source/GLTFExporter/Private/Converters/GLTFSkySphereConverters.cpp
+++ b/Engine/Plugins/Enterprise/GLTFExporter/MigratedFromGit/Source/GLTFExporter/Private/Converters/GLTFSkySphereConverters.cpp
@@ -21,7 +21,8 @@ FGLTFJsonSkySphereIndex FGLTFSkySphereConverter::Convert(const AActor* SkySphere
 	const UStaticMeshComponent* StaticMeshComponent = nullptr;
 	FGLTFActorUtility::TryGetPropertyValue(SkySphereActor, TEXT("SkySphereMesh"), StaticMeshComponent);
 
-	if (StaticMeshComponent != nullptr)
+	const UStaticMesh* StaticMesh = StaticMeshComponent != nullptr ? StaticMeshComponent->GetStaticMesh() : nullptr;
+	if (StaticMesh != nullptr)
 	{
 		const USceneComponent* ParentComponent = StaticMeshComponent->GetAttachParent();
 		const FName SocketName = StaticMeshComponent->GetAttachSocketName();
@@ -31,126 +32,76 @@ FGLTFJsonSkySphereIndex FGLTFSkySphereConverter::Convert(const AActor* SkySphere
 		const FTransform RelativeTransform = Transform.GetRelativeTransform(ParentTransform);
 
 		JsonSkySphere.Scale = FGLTFConverterUtility::ConvertScale(RelativeTransform.GetScale3D());
+		JsonSkySphere.SkySphereMesh = Builder.GetOrAddMesh(StaticMesh, -1, nullptr, { FGLTFMaterialUtility::GetDefault() });
 	}
 	else
 	{
-		Builder.AddWarningMessage(FString::Printf(TEXT("Failed to export Scale for Sky Sphere %s"), *SkySphereActor->GetName()));
+		Builder.AddWarningMessage(FString::Printf(TEXT("Failed to export property SkySphereMesh for Sky Sphere %s"), *SkySphereActor->GetName()));
 	}
 
-	const UStaticMesh* StaticMesh = StaticMeshComponent != nullptr ? StaticMeshComponent->GetStaticMesh() : nullptr;
-	if (StaticMesh)
+	const ADirectionalLight* DirectionalLight = nullptr;
+	if (FGLTFActorUtility::TryGetPropertyValue(SkySphereActor, TEXT("Directional light actor"), DirectionalLight))
 	{
-		JsonSkySphere.SkySphereMesh = Builder.GetOrAddMesh(StaticMesh, -1, nullptr, { FGLTFMaterialUtility::GetDefault() });
+		JsonSkySphere.DirectionalLight = Builder.GetOrAddNode(DirectionalLight);
 	}
 	else
 	{
-		Builder.AddWarningMessage(FString::Printf(TEXT("Failed to export SkySphereMesh for Sky Sphere %s"), *SkySphereActor->GetName()));
+		Builder.AddWarningMessage(FString::Printf(TEXT("Failed to export property Directional light actor for Sky Sphere %s"), *SkySphereActor->GetName()));
 	}
 
 	const UMaterialInstance* SkyMaterial = nullptr;
 	FGLTFActorUtility::TryGetPropertyValue(SkySphereActor, TEXT("Sky material"), SkyMaterial);
 
-	if (const UTexture2D* SkyTexture = GetSkyTexture(SkyMaterial))
-	{
-		JsonSkySphere.SkyTexture = Builder.GetOrAddTexture(SkyTexture);
-	}
-	else
-	{
-		Builder.AddWarningMessage(FString::Printf(TEXT("Failed to export SkyTexture for Sky Sphere %s"), *SkySphereActor->GetName()));
-	}
-
-	if (const UTexture2D* CloudsTexture = GetCloudsTexture(SkyMaterial))
-	{
-		JsonSkySphere.CloudsTexture = Builder.GetOrAddTexture(CloudsTexture);
-	}
-	else
-	{
-		Builder.AddWarningMessage(FString::Printf(TEXT("Failed to export CloudsTexture for Sky Sphere %s"), *SkySphereActor->GetName()));
-	}
-
-	if (const UTexture2D* StarsTexture = GetStarsTexture(SkyMaterial))
-	{
-		JsonSkySphere.StarsTexture = Builder.GetOrAddTexture(StarsTexture);
-	}
-	else
+	if (SkyMaterial != nullptr)
 	{
-		Builder.AddWarningMessage(FString::Printf(TEXT("Failed to export StarsTexture for Sky Sphere %s"), *SkySphereActor->GetName()));
-	}
+		ConvertScalarParameter(*SkySphereActor, *SkyMaterial, TEXT("Sun Radius"), JsonSkySphere.SunRadius);
+		ConvertScalarParameter(*SkySphereActor, *SkyMaterial, TEXT("NoisePower1"), JsonSkySphere.NoisePower1);
+		ConvertScalarParameter(*SkySphereActor, *SkyMaterial, TEXT("NoisePower2"), JsonSkySphere.NoisePower2);
 
-	const ADirectionalLight* DirectionalLight = nullptr;
-	if (FGLTFActorUtility::TryGetPropertyValue(SkySphereActor, TEXT("Directional light actor"), DirectionalLight))
-	{
-		JsonSkySphere.DirectionalLight = Builder.GetOrAddNode(DirectionalLight);
+		ConvertTextureParameter(*SkySphereActor, *SkyMaterial, ESkySphereTextureParameter::SkyTexture, JsonSkySphere.SkyTexture);
+		ConvertTextureParameter(*SkySphereActor, *SkyMaterial, ESkySphereTextureParameter::CloudsTexture, JsonSkySphere.CloudsTexture);
+		ConvertTextureParameter(*SkySphereActor, *SkyMaterial, ESkySphereTextureParameter::StarsTexture, JsonSkySphere.StarsTexture);
 	}
 	else
 	{
-		Builder.AddWarningMessage(FString::Printf(TEXT("Failed to export DirectionalLight for Sky Sphere %s"), *SkySphereActor->GetName()));
+		Builder.AddWarningMessage(FString::Printf(
+			TEXT("Failed to export property Sky material for Sky Sphere %s"),
+			*SkySphereActor->GetName()));
 	}
 
-	ConvertScalarParameter(*SkySphereActor, SkyMaterial, TEXT("Sun Radius"), TEXT("SunRadius"), JsonSkySphere.SunRadius);
-	ConvertScalarParameter(*SkySphereActor, SkyMaterial, TEXT("NoisePower1"), TEXT("NoisePower1"), JsonSkySphere.NoisePower1);
-	ConvertScalarParameter(*SkySphereActor, SkyMaterial, TEXT("NoisePower2"), TEXT("NoisePower2"), JsonSkySphere.NoisePower2);
+	ConvertProperty(*SkySphereActor, TEXT("Sun height"), JsonSkySphere.SunHeight);
+	ConvertProperty(*SkySphereActor, TEXT("Sun brightness"), JsonSkySphere.SunBrightness);
+	ConvertProperty(*SkySphereActor, TEXT("Stars brightness"), JsonSkySphere.StarsBrightness);
+	ConvertProperty(*SkySphereActor, TEXT("Cloud speed"), JsonSkySphere.CloudSpeed);
+	ConvertProperty(*SkySphereActor, TEXT("Cloud opacity"), JsonSkySphere.CloudOpacity);
+	ConvertProperty(*SkySphereActor, TEXT("Horizon Falloff"), JsonSkySphere.HorizonFalloff);
+	ConvertProperty(*SkySphereActor, TEXT("Colors determined by sun position"), JsonSkySphere.bColorsDeterminedBySunPosition);
 
-	ConvertProperty(*SkySphereActor, TEXT("Sun height"), TEXT("SunHeight"), JsonSkySphere.SunHeight);
-	ConvertProperty(*SkySphereActor, TEXT("Sun brightness"), TEXT("SunBrightness"), JsonSkySphere.SunBrightness);
-	ConvertProperty(*SkySphereActor, TEXT("Stars brightness"), TEXT("StarsBrightness"), JsonSkySphere.StarsBrightness);
-	ConvertProperty(*SkySphereActor, TEXT("Cloud speed"), TEXT("CloudSpeed"), JsonSkySphere.CloudSpeed);
-	ConvertProperty(*SkySphereActor, TEXT("Cloud opacity"), TEXT("CloudOpacity"), JsonSkySphere.CloudOpacity);
-	ConvertProperty(*SkySphereActor, TEXT("Horizon Falloff"), TEXT("HorizonFalloff"), JsonSkySphere.HorizonFalloff);
-	ConvertProperty(*SkySphereActor, TEXT("Colors determined by sun position"), TEXT("bColorsDeterminedBySunPosition"), JsonSkySphere.bColorsDeterminedBySunPosition);
+	ConvertColorProperty(*SkySphereActor, TEXT("Zenith Color"), JsonSkySphere.ZenithColor);
+	ConvertColorProperty(*SkySphereActor, TEXT("Horizon color"), JsonSkySphere.HorizonColor);
+	ConvertColorProperty(*SkySphereActor, TEXT("Cloud color"), JsonSkySphere.CloudColor);
+	ConvertColorProperty(*SkySphereActor, TEXT("Overall Color"), JsonSkySphere.OverallColor);
 
-	ConvertColorProperty(*SkySphereActor, TEXT("Zenith Color"), TEXT("ZenithColor"), JsonSkySphere.ZenithColor);
-	ConvertColorProperty(*SkySphereActor, TEXT("Horizon color"), TEXT("HorizonColor"), JsonSkySphere.HorizonColor);
-	ConvertColorProperty(*SkySphereActor, TEXT("Cloud color"), TEXT("CloudColor"), JsonSkySphere.CloudColor);
-	ConvertColorProperty(*SkySphereActor, TEXT("Overall Color"), TEXT("OverallColor"), JsonSkySphere.OverallColor);
-
-	ConvertColorCurveProperty(*SkySphereActor, TEXT("Zenith color curve"), TEXT("ZenithColorCurve"), JsonSkySphere.ZenithColorCurve);
-	ConvertColorCurveProperty(*SkySphereActor, TEXT("Horizon color curve"), TEXT("HorizonColorCurve"), JsonSkySphere.HorizonColorCurve);
-	ConvertColorCurveProperty(*SkySphereActor, TEXT("Cloud color curve"), TEXT("CloudColorCurve"), JsonSkySphere.CloudColorCurve);
+	ConvertColorCurveProperty(*SkySphereActor, TEXT("Zenith color curve"), JsonSkySphere.ZenithColorCurve);
+	ConvertColorCurveProperty(*SkySphereActor, TEXT("Horizon color curve"), JsonSkySphere.HorizonColorCurve);
+	ConvertColorCurveProperty(*SkySphereActor, TEXT("Cloud color curve"), JsonSkySphere.CloudColorCurve);
 
 	return Builder.AddSkySphere(JsonSkySphere);
 }
 
-const UTexture2D* FGLTFSkySphereConverter::GetSkyTexture(const UMaterialInstance* SkyMaterial) const
-{
-	// TODO: the texture isn't exposed as a parameter on the default sky material. We should try to find
-	// a way of extracting and returning the correct texture from the material anyway.
-
-	static UTexture2D* Texture = LoadObject<UTexture2D>(nullptr, TEXT("/Engine/EngineSky/T_Sky_Blue.T_Sky_Blue"));
-	return Texture;
-}
-
-const UTexture2D* FGLTFSkySphereConverter::GetCloudsTexture(const UMaterialInstance* SkyMaterial) const
-{
-	// TODO: the texture isn't exposed as a parameter on the default sky material. We should try to find
-	// a way of extracting and returning the correct texture from the material anyway.
-
-	static UTexture2D* Texture = LoadObject<UTexture2D>(nullptr, TEXT("/Engine/EngineSky/T_Sky_Clouds_M.T_Sky_Clouds_M"));
-	return Texture;
-}
-
-const UTexture2D* FGLTFSkySphereConverter::GetStarsTexture(const UMaterialInstance* SkyMaterial) const
-{
-	// TODO: the texture isn't exposed as a parameter on the default sky material. We should try to find
-	// a way of extracting and returning the correct texture from the material anyway.
-
-	static UTexture2D* Texture = LoadObject<UTexture2D>(nullptr, TEXT("/Engine/EngineSky/T_Sky_Stars.T_Sky_Stars"));
-	return Texture;
-}
-
 template <class ValueType>
-void FGLTFSkySphereConverter::ConvertProperty(const AActor& Actor, const TCHAR* PropertyName, const TCHAR* ExportedPropertyName, ValueType& OutValue) const
+void FGLTFSkySphereConverter::ConvertProperty(const AActor& Actor, const TCHAR* PropertyName, ValueType& OutValue) const
 {
 	if (!FGLTFActorUtility::TryGetPropertyValue(&Actor, PropertyName, OutValue))
 	{
 		Builder.AddWarningMessage(FString::Printf(
-			TEXT("Failed to export %s for Sky Sphere %s"),
-			ExportedPropertyName,
+			TEXT("Failed to export property %s for Sky Sphere %s"),
+			PropertyName,
 			*Actor.GetName()));
 	}
 }
 
-void FGLTFSkySphereConverter::ConvertColorProperty(const AActor& Actor, const TCHAR* PropertyName, const TCHAR* ExportedPropertyName, FGLTFJsonColor4& OutValue) const
+void FGLTFSkySphereConverter::ConvertColorProperty(const AActor& Actor, const TCHAR* PropertyName, FGLTFJsonColor4& OutValue) const
 {
 	FLinearColor LinearColor;
 	if (FGLTFActorUtility::TryGetPropertyValue(&Actor, PropertyName, LinearColor))
@@ -160,13 +111,13 @@ void FGLTFSkySphereConverter::ConvertColorProperty(const AActor& Actor, const TC
 	else
 	{
 		Builder.AddWarningMessage(FString::Printf(
-			TEXT("Failed to export %s for Sky Sphere %s"),
-			ExportedPropertyName,
+			TEXT("Failed to export property %s for Sky Sphere %s"),
+			PropertyName,
 			*Actor.GetName()));
 	}
 }
 
-void FGLTFSkySphereConverter::ConvertColorCurveProperty(const AActor& Actor, const TCHAR* PropertyName, const TCHAR* ExportedPropertyName, FGLTFJsonSkySphereColorCurve& OutValue) const
+void FGLTFSkySphereConverter::ConvertColorCurveProperty(const AActor& Actor, const TCHAR* PropertyName, FGLTFJsonSkySphereColorCurve& OutValue) const
 {
 	const UCurveLinearColor* ColorCurve = nullptr;
 	FGLTFActorUtility::TryGetPropertyValue(&Actor, PropertyName, ColorCurve);
@@ -176,8 +127,8 @@ void FGLTFSkySphereConverter::ConvertColorCurveProperty(const AActor& Actor, con
 		if (FGLTFCurveUtility::HasAnyAdjustment(*ColorCurve))
 		{
 			Builder.AddWarningMessage(FString::Printf(
-				TEXT("Adjustments for %s in Sky Sphere %s are not supported"),
-				ExportedPropertyName,
+				TEXT("Adjustments for property %s in Sky Sphere %s are not supported"),
+				PropertyName,
 				*Actor.GetName()));
 		}
 
@@ -202,19 +153,49 @@ void FGLTFSkySphereConverter::ConvertColorCurveProperty(const AActor& Actor, con
 	else
 	{
 		Builder.AddWarningMessage(FString::Printf(
-			TEXT("Failed to export %s for Sky Sphere %s"),
-			ExportedPropertyName,
+			TEXT("Failed to export property %s for Sky Sphere %s"),
+			PropertyName,
 			*Actor.GetName()));
 	}
 }
 
-void FGLTFSkySphereConverter::ConvertScalarParameter(const AActor& Actor, const UMaterialInstance* Material, const TCHAR* ParameterName, const TCHAR* ExportedPropertyName, float& OutValue) const
+void FGLTFSkySphereConverter::ConvertScalarParameter(const AActor& Actor, const UMaterialInstance& Material, const TCHAR* ParameterName, float& OutValue) const
 {
-	if (Material == nullptr || !Material->GetScalarParameterValue(ParameterName, OutValue))
+	if (!Material.GetScalarParameterValue(ParameterName, OutValue))
+	{
+		Builder.AddWarningMessage(FString::Printf(
+			TEXT("Failed to export parameter %s (in material %s) for Sky Sphere %s"),
+			ParameterName,
+			*Material.GetName(),
+			*Actor.GetName()));
+	}
+}
+
+void FGLTFSkySphereConverter::ConvertTextureParameter(const AActor& Actor, const UMaterialInstance& Material, const ESkySphereTextureParameter Parameter, FGLTFJsonTextureIndex& OutValue) const
+{
+	const TCHAR* TexturePath = nullptr;
+
+	// TODO: the sky sphere material doesn't expose the relevant textures as parameters, so we use hard-coded asset paths.
+	// We should try to find a way of extracting and returning the correct texture from the material anyway.
+	switch (Parameter)
+	{
+		case ESkySphereTextureParameter::SkyTexture:    TexturePath = TEXT("/Engine/EngineSky/T_Sky_Blue.T_Sky_Blue"); break;
+		case ESkySphereTextureParameter::CloudsTexture: TexturePath = TEXT("/Engine/EngineSky/T_Sky_Clouds_M.T_Sky_Clouds_M"); break;
+		case ESkySphereTextureParameter::StarsTexture:  TexturePath = TEXT("/Engine/EngineSky/T_Sky_Stars.T_Sky_Stars"); break;
+		default:                                        checkNoEntry();
+	}
+
+	const UTexture2D* Texture2D = TexturePath != nullptr ? LoadObject<UTexture2D>(nullptr, TexturePath) : nullptr;
+	if (Texture2D != nullptr)
+	{
+		OutValue = Builder.GetOrAddTexture(Texture2D);
+	}
+	else
 	{
 		Builder.AddWarningMessage(FString::Printf(
-			TEXT("Failed to export %s for Sky Sphere %s"),
-			ExportedPropertyName,
+			TEXT("Failed to export texture %s (used in material %s) for Sky Sphere %s"),
+			TexturePath,
+			*Material.GetName(),
 			*Actor.GetName()));
 	}
 }
diff --git a/Engine/Plugins/Enterprise/GLTFExporter/MigratedFromGit/Source/GLTFExporter/Private/Converters/GLTFSkySphereConverters.h b/Engine/Plugins/Enterprise/GLTFExporter/MigratedFromGit/Source/GLTFExporter/Private/Converters/GLTFSkySphereConverters.h
--- a/Engine/Plugins/Enterprise/GLTFExporter/MigratedFromGit/Source/GLTFExporter/Private/Converters/GLTFSkySphereConverters.h
+++ b/Engine/Plugins/Enterprise/GLTFExporter/MigratedFromGit/Source/GLTFExporter/Private/Converters/GLTFSkySphereConverters.h
@@ -12,18 +12,23 @@ struct FGLTFJsonSkySphereColorCurve;
 
 class FGLTFSkySphereConverter final : public FGLTFBuilderContext, public TGLTFConverter<FGLTFJsonSkySphereIndex, const AActor*>
 {
+	enum class ESkySphereTextureParameter
+	{
+		SkyTexture,
+		CloudsTexture,
+		StarsTexture
+	};
+
 	using FGLTFBuilderContext::FGLTFBuilderContext;
 
 	virtual FGLTFJsonSkySphereIndex Convert(const AActor* SkySphereActor) override;
 
-	const UTexture2D* GetSkyTexture(const UMaterialInstance* SkyMaterial) const;
-	const UTexture2D* GetCloudsTexture(const UMaterialInstance* SkyMaterial) const;
-	const UTexture2D* GetStarsTexture(const UMaterialInstance* SkyMaterial) const;
-
 	template <class ValueType>
-	void ConvertProperty(const AActor& Actor, const TCHAR* PropertyName, const TCHAR* ExportedPropertyName, ValueType& OutValue) const;
+	void ConvertProperty(const AActor& Actor, const TCHAR* PropertyName, ValueType& OutValue) const;
+
+	void ConvertColorProperty(const AActor& Actor, const TCHAR* PropertyName, FGLTFJsonColor4& OutValue) const;
+	void ConvertColorCurveProperty(const AActor& Actor, const TCHAR* PropertyName, FGLTFJsonSkySphereColorCurve& OutValue) const;
+	void ConvertScalarParameter(const AActor& Actor, const UMaterialInstance& Material, const TCHAR* ParameterName, float& OutValue) const;
 
-	void ConvertColorProperty(const AActor& Actor, const TCHAR* PropertyName, const TCHAR* ExportedPropertyName, FGLTFJsonColor4& OutValue) const;
-	void ConvertColorCurveProperty(const AActor& Actor, const TCHAR* PropertyName, const TCHAR* ExportedPropertyName, FGLTFJsonSkySphereColorCurve& OutValue) const;
-	void ConvertScalarParameter(const AActor& Actor, const UMaterialInstance* Material, const TCHAR* ParameterName, const TCHAR* ExportedPropertyName, float& OutValue) const;
+	void ConvertTextureParameter(const AActor& Actor, const UMaterialInstance& Material, const ESkySphereTextureParameter Parameter, FGLTFJsonTextureIndex& OutValue) const;
 };
