[GLTFExporter] feat: add json data structures for KHR_texture_transform

commit 522e69c67a0857b80dc3490b46b532e6d13af2a4
Date: Sun, 18 Oct 2020 13:27:44 +0200

    feat: add json data structures for KHR_texture_transform
    (cherry picked from commit e8588a487bffb8a64a4e72d2e9754ecf30c65d41​)

#skipci

diff --git a/Engine/Plugins/Enterprise/GLTFExporter/MigratedFromGit/Source/GLTFExporter/Private/Json/GLTFJsonEnums.h b/Engine/Plugins/Enterprise/GLTFExporter/MigratedFromGit/Source/GLTFExporter/Private/Json/GLTFJsonEnums.h
--- a/Engine/Plugins/Enterprise/GLTFExporter/MigratedFromGit/Source/GLTFExporter/Private/Json/GLTFJsonEnums.h
+++ b/Engine/Plugins/Enterprise/GLTFExporter/MigratedFromGit/Source/GLTFExporter/Private/Json/GLTFJsonEnums.h
@@ -10,6 +10,7 @@ enum class EGLTFJsonExtension
 	KHR_MaterialsUnlit,
 	KHR_MaterialsClearCoat,
 	KHR_MeshQuantization,
+	KHR_TextureTransform,
 	EPIC_OrbitalCameras,
 	EPIC_HDRIBackdrops,
 	EPIC_AnimationHotspots,
diff --git a/Engine/Plugins/Enterprise/GLTFExporter/MigratedFromGit/Source/GLTFExporter/Private/Json/GLTFJsonMaterial.h b/Engine/Plugins/Enterprise/GLTFExporter/MigratedFromGit/Source/GLTFExporter/Private/Json/GLTFJsonMaterial.h
--- a/Engine/Plugins/Enterprise/GLTFExporter/MigratedFromGit/Source/GLTFExporter/Private/Json/GLTFJsonMaterial.h
+++ b/Engine/Plugins/Enterprise/GLTFExporter/MigratedFromGit/Source/GLTFExporter/Private/Json/GLTFJsonMaterial.h
@@ -5,6 +5,7 @@
 #include "Json/GLTFJsonEnums.h"
 #include "Json/GLTFJsonColor3.h"
 #include "Json/GLTFJsonColor4.h"
+#include "Json/GLTFJsonTextureTransform.h"
 #include "Json/GLTFJsonUtility.h"
 #include "Serialization/JsonSerializer.h"
 
@@ -13,6 +14,8 @@ struct FGLTFJsonTextureInfo
 	FGLTFJsonTextureIndex Index;
 	int32 TexCoord;
 
+	FGLTFJsonTextureTransform Transform;
+
 	FGLTFJsonTextureInfo()
 		: TexCoord(0)
 	{
@@ -30,6 +33,17 @@ struct FGLTFJsonTextureInfo
 			JsonWriter.WriteValue(TEXT("texCoord"), TexCoord);
 		}
 
+		if (Transform != FGLTFJsonTextureTransform())
+		{
+			const EGLTFJsonExtension Extension = EGLTFJsonExtension::KHR_TextureTransform;
+			Extensions.Used.Add(Extension);
+
+			JsonWriter.WriteObjectStart(TEXT("extensions"));
+			JsonWriter.WriteIdentifierPrefix(FGLTFJsonUtility::ToString(Extension));
+			Transform.WriteObject(JsonWriter, Extensions);
+			JsonWriter.WriteObjectEnd();
+		}
+
 		JsonWriter.WriteObjectEnd();
 	}
 };
diff --git a/Engine/Plugins/Enterprise/GLTFExporter/MigratedFromGit/Source/GLTFExporter/Private/Json/GLTFJsonTextureTransform.h b/Engine/Plugins/Enterprise/GLTFExporter/MigratedFromGit/Source/GLTFExporter/Private/Json/GLTFJsonTextureTransform.h
new file mode 100644
--- /dev/null
+++ b/Engine/Plugins/Enterprise/GLTFExporter/MigratedFromGit/Source/GLTFExporter/Private/Json/GLTFJsonTextureTransform.h
@@ -0,0 +1,59 @@
+// Copyright Epic Games, Inc. All Rights Reserved.
+
+#pragma once
+
+#include "Json/GLTFJsonUtility.h"
+#include "Json/GLTFJsonVector2.h"
+#include "Serialization/JsonSerializer.h"
+
+struct FGLTFJsonTextureTransform
+{
+	FGLTFJsonVector2 Offset;
+	FGLTFJsonVector2 Scale;
+	float Rotation;
+
+	FGLTFJsonTextureTransform()
+        : Offset(FGLTFJsonVector2::Zero)
+        , Scale(FGLTFJsonVector2::One)
+        , Rotation(0)
+	{
+	}
+
+	template <class CharType = TCHAR, class PrintPolicy = TPrettyJsonPrintPolicy<CharType>>
+	void WriteObject(TJsonWriter<CharType, PrintPolicy>& JsonWriter, FGLTFJsonExtensions& Extensions) const
+	{
+		JsonWriter.WriteObjectStart();
+
+		if (Offset != FGLTFJsonVector2::Zero)
+		{
+			JsonWriter.WriteIdentifierPrefix(TEXT("offset"));
+			Offset.WriteArray(JsonWriter);
+		}
+
+		if (Scale != FGLTFJsonVector2::One)
+		{
+			JsonWriter.WriteIdentifierPrefix(TEXT("scale"));
+			Scale.WriteArray(JsonWriter);
+		}
+
+		if (Rotation != 0)
+		{
+			JsonWriter.WriteIdentifierPrefix(TEXT("rotation"));
+			FGLTFJsonUtility::WriteExactValue(JsonWriter, Rotation);
+		}
+
+		JsonWriter.WriteObjectEnd();
+	}
+
+	bool operator==(const FGLTFJsonTextureTransform& Other) const
+	{
+		return Offset == Other.Offset
+            && Scale == Other.Scale
+            && Rotation == Other.Rotation;
+	}
+
+	bool operator!=(const FGLTFJsonTextureTransform& Other) const
+	{
+		return !(*this == Other);
+	}
+};
diff --git a/Engine/Plugins/Enterprise/GLTFExporter/MigratedFromGit/Source/GLTFExporter/Private/Json/GLTFJsonUtility.h b/Engine/Plugins/Enterprise/GLTFExporter/MigratedFromGit/Source/GLTFExporter/Private/Json/GLTFJsonUtility.h
--- a/Engine/Plugins/Enterprise/GLTFExporter/MigratedFromGit/Source/GLTFExporter/Private/Json/GLTFJsonUtility.h
+++ b/Engine/Plugins/Enterprise/GLTFExporter/MigratedFromGit/Source/GLTFExporter/Private/Json/GLTFJsonUtility.h
@@ -23,6 +23,7 @@ struct FGLTFJsonUtility
 			case EGLTFJsonExtension::KHR_MaterialsUnlit:      return TEXT("KHR_materials_unlit");
 			case EGLTFJsonExtension::KHR_MaterialsClearCoat:  return TEXT("KHR_materials_clearcoat");
 			case EGLTFJsonExtension::KHR_MeshQuantization:    return TEXT("KHR_mesh_quantization");
+			case EGLTFJsonExtension::KHR_TextureTransform:    return TEXT("KHR_texture_transform");
 			case EGLTFJsonExtension::EPIC_OrbitalCameras:     return TEXT("EPIC_orbital_cameras");
 			case EGLTFJsonExtension::EPIC_HDRIBackdrops:      return TEXT("EPIC_hdri_backdrops");
 			case EGLTFJsonExtension::EPIC_AnimationHotspots:  return TEXT("EPIC_animation_hotspots");
