[GLTFExporter] feat: convert UMaterialExpressionTextureCoordinate values to KHR_texture_transform

commit 4dcb39ecec58c332620e781b2b949f9235f3e8ce
Date: Sun, 18 Oct 2020 13:30:06 +0200

    feat: convert UMaterialExpressionTextureCoordinate values to KHR_texture_transform
    (cherry picked from commit fc4b49ccacb4f52670814d3512da2b09c14cffcc​)

#skipci

diff --git a/Engine/Plugins/Enterprise/GLTFExporter/MigratedFromGit/Source/GLTFExporter/Private/Converters/GLTFMaterialUtility.cpp b/Engine/Plugins/Enterprise/GLTFExporter/MigratedFromGit/Source/GLTFExporter/Private/Converters/GLTFMaterialUtility.cpp
--- a/Engine/Plugins/Enterprise/GLTFExporter/MigratedFromGit/Source/GLTFExporter/Private/Converters/GLTFMaterialUtility.cpp
+++ b/Engine/Plugins/Enterprise/GLTFExporter/MigratedFromGit/Source/GLTFExporter/Private/Converters/GLTFMaterialUtility.cpp
@@ -394,21 +394,29 @@ uint32 FGLTFMaterialUtility::GetMaskComponentCount(const FExpressionInput& Expre
 	return ExpressionInput.MaskR + ExpressionInput.MaskG + ExpressionInput.MaskB + ExpressionInput.MaskA;
 }
 
-bool FGLTFMaterialUtility::TryGetTextureCoordinateIndex(const UMaterialExpressionTextureSample* TextureSampler, int32& TexCoord)
+bool FGLTFMaterialUtility::TryGetTextureCoordinateIndex(const UMaterialExpressionTextureSample* TextureSampler, int32& TexCoord, FGLTFJsonTextureTransform& Transform)
 {
 	const UMaterialExpression* Expression = TextureSampler->Coordinates.Expression;
 	if (Expression == nullptr)
 	{
 		TexCoord = TextureSampler->ConstCoordinate;
+		Transform = {};
 		return true;
 	}
 
 	if (const UMaterialExpressionTextureCoordinate* TextureCoordinate = Cast<UMaterialExpressionTextureCoordinate>(Expression))
 	{
 		TexCoord = TextureCoordinate->CoordinateIndex;
+		Transform.Offset.X = TextureCoordinate->UnMirrorU ? TextureCoordinate->UTiling * 0.5 : 0.0;
+		Transform.Offset.Y = TextureCoordinate->UnMirrorV ? TextureCoordinate->VTiling * 0.5 : 0.0;
+		Transform.Scale.X = TextureCoordinate->UTiling * (TextureCoordinate->UnMirrorU ? 0.5 : 1.0);
+		Transform.Scale.Y = TextureCoordinate->VTiling * (TextureCoordinate->UnMirrorV ? 0.5 : 1.0);
+		Transform.Rotation = 0;
 		return true;
 	}
 
+	// TODO: add support for advanced expression tree (ex UMaterialExpressionTextureCoordinate -> UMaterialExpressionMultiply -> UMaterialExpressionAdd)
+
 	return false;
 }
 
diff --git a/Engine/Plugins/Enterprise/GLTFExporter/MigratedFromGit/Source/GLTFExporter/Private/Converters/GLTFMaterialUtility.h b/Engine/Plugins/Enterprise/GLTFExporter/MigratedFromGit/Source/GLTFExporter/Private/Converters/GLTFMaterialUtility.h
--- a/Engine/Plugins/Enterprise/GLTFExporter/MigratedFromGit/Source/GLTFExporter/Private/Converters/GLTFMaterialUtility.h
+++ b/Engine/Plugins/Enterprise/GLTFExporter/MigratedFromGit/Source/GLTFExporter/Private/Converters/GLTFMaterialUtility.h
@@ -61,7 +61,7 @@ struct FGLTFMaterialUtility
 	static FLinearColor GetMask(const FExpressionInput& ExpressionInput);
 	static uint32 GetMaskComponentCount(const FExpressionInput& ExpressionInput);
 
-	static bool TryGetTextureCoordinateIndex(const UMaterialExpressionTextureSample* TextureSampler, int32& TexCoord);
+	static bool TryGetTextureCoordinateIndex(const UMaterialExpressionTextureSample* TextureSampler, int32& TexCoord, FGLTFJsonTextureTransform& Transform);
 	static void GetAllTextureCoordinateIndices(const FExpressionInput& Input, TSet<int32>& OutTexCoords);
 
 	static void ExpandAllFunctionExpressions(TArray<UMaterialExpression*>& InOutExpressions);
diff --git a/Engine/Plugins/Enterprise/GLTFExporter/MigratedFromGit/Source/GLTFExporter/Private/Tasks/GLTFMaterialTasks.cpp b/Engine/Plugins/Enterprise/GLTFExporter/MigratedFromGit/Source/GLTFExporter/Private/Tasks/GLTFMaterialTasks.cpp
--- a/Engine/Plugins/Enterprise/GLTFExporter/MigratedFromGit/Source/GLTFExporter/Private/Tasks/GLTFMaterialTasks.cpp
+++ b/Engine/Plugins/Enterprise/GLTFExporter/MigratedFromGit/Source/GLTFExporter/Private/Tasks/GLTFMaterialTasks.cpp
@@ -248,18 +248,22 @@ bool FGLTFMaterialTask::TryGetBaseColorAndOpacity(FGLTFJsonPBRMetallicRoughness&
 	const UTexture2D* OpacityTexture;
 	int32 BaseColorTexCoord;
 	int32 OpacityTexCoord;
+	FGLTFJsonTextureTransform BaseColorTransform;
+	FGLTFJsonTextureTransform OpacityTransform;
 
-	const bool bHasBaseColorSourceTexture = TryGetSourceTexture(BaseColorTexture, BaseColorTexCoord, BaseColorProperty, BaseColorInputMasks);
-	const bool bHasOpacitySourceTexture = TryGetSourceTexture(OpacityTexture, OpacityTexCoord, OpacityProperty, OpacityInputMasks);
+	const bool bHasBaseColorSourceTexture = TryGetSourceTexture(BaseColorTexture, BaseColorTexCoord, BaseColorTransform, BaseColorProperty, BaseColorInputMasks);
+	const bool bHasOpacitySourceTexture = TryGetSourceTexture(OpacityTexture, OpacityTexCoord, OpacityTransform, OpacityProperty, OpacityInputMasks);
 
 	// Detect the "happy path" where both inputs share the same texture and are correctly masked.
 	if (bHasBaseColorSourceTexture &&
 		bHasOpacitySourceTexture &&
 		BaseColorTexture == OpacityTexture &&
-		BaseColorTexCoord == OpacityTexCoord)
+		BaseColorTexCoord == OpacityTexCoord &&
+        BaseColorTransform == OpacityTransform)
 	{
 		OutPBRParams.BaseColorTexture.Index = Builder.GetOrAddTexture(BaseColorTexture);
 		OutPBRParams.BaseColorTexture.TexCoord = BaseColorTexCoord;
+		OutPBRParams.BaseColorTexture.Transform = BaseColorTransform;
 		return true;
 	}
 
@@ -413,18 +417,22 @@ bool FGLTFMaterialTask::TryGetMetallicAndRoughness(FGLTFJsonPBRMetallicRoughness
 	const UTexture2D* RoughnessTexture;
 	int32 MetallicTexCoord;
 	int32 RoughnessTexCoord;
+	FGLTFJsonTextureTransform MetallicTransform;
+	FGLTFJsonTextureTransform RoughnessTransform;
 
-	const bool bHasMetallicSourceTexture = TryGetSourceTexture(MetallicTexture, MetallicTexCoord, MetallicProperty, MetallicInputMasks);
-	const bool bHasRoughnessSourceTexture = TryGetSourceTexture(RoughnessTexture, RoughnessTexCoord, RoughnessProperty, RoughnessInputMasks);
+	const bool bHasMetallicSourceTexture = TryGetSourceTexture(MetallicTexture, MetallicTexCoord, MetallicTransform, MetallicProperty, MetallicInputMasks);
+	const bool bHasRoughnessSourceTexture = TryGetSourceTexture(RoughnessTexture, RoughnessTexCoord, RoughnessTransform, RoughnessProperty, RoughnessInputMasks);
 
 	// Detect the "happy path" where both inputs share the same texture and are correctly masked.
 	if (bHasMetallicSourceTexture &&
 		bHasRoughnessSourceTexture &&
 		MetallicTexture == RoughnessTexture &&
-		MetallicTexCoord == RoughnessTexCoord)
+		MetallicTexCoord == RoughnessTexCoord &&
+        MetallicTransform == RoughnessTransform)
 	{
 		OutPBRParams.MetallicRoughnessTexture.Index = Builder.GetOrAddTexture(MetallicTexture);
 		OutPBRParams.MetallicRoughnessTexture.TexCoord = MetallicTexCoord;
+		OutPBRParams.MetallicRoughnessTexture.Transform = MetallicTransform;
 		return true;
 	}
 
@@ -574,21 +582,25 @@ bool FGLTFMaterialTask::TryGetClearCoatRoughness(FGLTFJsonClearCoatExtension& Ou
 	const UTexture2D* RoughnessTexture;
 	int32 IntensityTexCoord;
 	int32 RoughnessTexCoord;
+	FGLTFJsonTextureTransform IntensityTransform;
+	FGLTFJsonTextureTransform RoughnessTransform;
 
-	const bool bHasIntensitySourceTexture = TryGetSourceTexture(IntensityTexture, IntensityTexCoord, IntensityProperty, ClearCoatInputMasks);
-	const bool bHasRoughnessSourceTexture = TryGetSourceTexture(RoughnessTexture, RoughnessTexCoord, RoughnessProperty, ClearCoatRoughnessInputMasks);
+	const bool bHasIntensitySourceTexture = TryGetSourceTexture(IntensityTexture, IntensityTexCoord, IntensityTransform, IntensityProperty, ClearCoatInputMasks);
+	const bool bHasRoughnessSourceTexture = TryGetSourceTexture(RoughnessTexture, RoughnessTexCoord, RoughnessTransform, RoughnessProperty, ClearCoatRoughnessInputMasks);
 
 	// Detect the "happy path" where both inputs share the same texture and are correctly masked.
 	if (bHasIntensitySourceTexture &&
 		bHasRoughnessSourceTexture &&
 		IntensityTexture == RoughnessTexture &&
-		IntensityTexCoord == RoughnessTexCoord)
+		IntensityTexCoord == RoughnessTexCoord &&
+        IntensityTransform == RoughnessTransform)
 	{
 		const FGLTFJsonTextureIndex TextureIndex = Builder.GetOrAddTexture(IntensityTexture);
 		OutExtParams.ClearCoatTexture.Index = TextureIndex;
 		OutExtParams.ClearCoatTexture.TexCoord = IntensityTexCoord;
 		OutExtParams.ClearCoatRoughnessTexture.Index = TextureIndex;
-		OutExtParams.ClearCoatRoughnessTexture.TexCoord = RoughnessTexCoord;
+		OutExtParams.ClearCoatRoughnessTexture.TexCoord = IntensityTexCoord;
+		OutExtParams.ClearCoatRoughnessTexture.Transform = IntensityTransform;
 		return true;
 	}
 
@@ -1027,18 +1039,20 @@ bool FGLTFMaterialTask::TryGetSourceTexture(FGLTFJsonTextureInfo& OutTexInfo, EM
 {
 	const UTexture2D* Texture;
 	int32 TexCoord;
+	FGLTFJsonTextureTransform Transform;
 
-	if (TryGetSourceTexture(Texture, TexCoord, Property, AllowedMasks))
+	if (TryGetSourceTexture(Texture, TexCoord, Transform, Property, AllowedMasks))
 	{
 		OutTexInfo.Index = Builder.GetOrAddTexture(Texture);
 		OutTexInfo.TexCoord = TexCoord;
+		OutTexInfo.Transform = Transform;
 		return true;
 	}
 
 	return false;
 }
 
-bool FGLTFMaterialTask::TryGetSourceTexture(const UTexture2D*& OutTexture, int32& OutTexCoord, EMaterialProperty Property, const TArray<FLinearColor>& AllowedMasks) const
+bool FGLTFMaterialTask::TryGetSourceTexture(const UTexture2D*& OutTexture, int32& OutTexCoord, FGLTFJsonTextureTransform& OutTransform, EMaterialProperty Property, const TArray<FLinearColor>& AllowedMasks) const
 {
 	const FExpressionInput* MaterialInput = FGLTFMaterialUtility::GetInputForProperty(Material, Property);
 	if (MaterialInput == nullptr)
@@ -1090,7 +1104,7 @@ bool FGLTFMaterialTask::TryGetSourceTexture(const UTexture2D*& OutTexture, int32
 			return false;
 		}
 
-		if (!FGLTFMaterialUtility::TryGetTextureCoordinateIndex(TextureParameter, OutTexCoord))
+		if (!FGLTFMaterialUtility::TryGetTextureCoordinateIndex(TextureParameter, OutTexCoord, OutTransform))
 		{
 			// TODO: report error (failed to identify texture coordinate index)
 			return false;
@@ -1117,7 +1131,7 @@ bool FGLTFMaterialTask::TryGetSourceTexture(const UTexture2D*& OutTexture, int32
 			return false;
 		}
 
-		if (!FGLTFMaterialUtility::TryGetTextureCoordinateIndex(TextureSampler, OutTexCoord))
+		if (!FGLTFMaterialUtility::TryGetTextureCoordinateIndex(TextureSampler, OutTexCoord, OutTransform))
 		{
 			// TODO: report error (failed to identify texture coordinate index)
 			return false;
diff --git a/Engine/Plugins/Enterprise/GLTFExporter/MigratedFromGit/Source/GLTFExporter/Private/Tasks/GLTFMaterialTasks.h b/Engine/Plugins/Enterprise/GLTFExporter/MigratedFromGit/Source/GLTFExporter/Private/Tasks/GLTFMaterialTasks.h
--- a/Engine/Plugins/Enterprise/GLTFExporter/MigratedFromGit/Source/GLTFExporter/Private/Tasks/GLTFMaterialTasks.h
+++ b/Engine/Plugins/Enterprise/GLTFExporter/MigratedFromGit/Source/GLTFExporter/Private/Tasks/GLTFMaterialTasks.h
@@ -47,7 +47,7 @@ private:
 	bool TryGetConstantScalar(float& OutValue, EMaterialProperty Property) const;
 
 	bool TryGetSourceTexture(FGLTFJsonTextureInfo& OutTexInfo, EMaterialProperty Property, const TArray<FLinearColor>& AllowedMasks = {}) const;
-	bool TryGetSourceTexture(const UTexture2D*& OutTexture, int32& OutTexCoord, EMaterialProperty Property, const TArray<FLinearColor>& AllowedMasks = {}) const;
+	bool TryGetSourceTexture(const UTexture2D*& OutTexture, int32& OutTexCoord, FGLTFJsonTextureTransform& OutTransform, EMaterialProperty Property, const TArray<FLinearColor>& AllowedMasks = {}) const;
 
 	bool TryGetBakedMaterialProperty(FGLTFJsonTextureInfo& OutTexInfo, FGLTFJsonColor3& OutConstant, EMaterialProperty Property, const FString& PropertyName) const;
 	bool TryGetBakedMaterialProperty(FGLTFJsonTextureInfo& OutTexInfo, FGLTFJsonColor4& OutConstant, EMaterialProperty Property, const FString& PropertyName) const;
