[GLTFExporter] feat: add and implement export options bExportTextureTransforms

commit c2843c8e9c8aa4bc8c50ad61cc4f304125638c43
Date: Sun, 18 Oct 2020 13:45:52 +0200

    feat: add and implement export options bExportTextureTransforms
    (cherry picked from commit f6a225f765039a4d23693c9cffc487211075ea3d​)

#skipci

diff --git a/Engine/Plugins/Enterprise/GLTFExporter/MigratedFromGit/Source/GLTFExporter/Private/GLTFExportOptions.cpp b/Engine/Plugins/Enterprise/GLTFExporter/MigratedFromGit/Source/GLTFExporter/Private/GLTFExportOptions.cpp
--- a/Engine/Plugins/Enterprise/GLTFExporter/MigratedFromGit/Source/GLTFExporter/Private/GLTFExportOptions.cpp
+++ b/Engine/Plugins/Enterprise/GLTFExporter/MigratedFromGit/Source/GLTFExporter/Private/GLTFExportOptions.cpp
@@ -26,6 +26,7 @@ UGLTFExportOptions::UGLTFExportOptions(const FObjectInitializer& ObjectInitializ
 	bExportAnimationSequences = true;
 	bRetargetBoneTransforms = true;
 	bExportPlaybackSettings = true;
+	bExportTextureTransforms = true;
 	bExportSourceTextures = true;
 	bExportLightmaps = true;
 	TextureHDREncoding = EGLTFExporterTextureHDREncoding::RGBM;
diff --git a/Engine/Plugins/Enterprise/GLTFExporter/MigratedFromGit/Source/GLTFExporter/Private/Tasks/GLTFMaterialTasks.cpp b/Engine/Plugins/Enterprise/GLTFExporter/MigratedFromGit/Source/GLTFExporter/Private/Tasks/GLTFMaterialTasks.cpp
--- a/Engine/Plugins/Enterprise/GLTFExporter/MigratedFromGit/Source/GLTFExporter/Private/Tasks/GLTFMaterialTasks.cpp
+++ b/Engine/Plugins/Enterprise/GLTFExporter/MigratedFromGit/Source/GLTFExporter/Private/Tasks/GLTFMaterialTasks.cpp
@@ -1110,6 +1110,12 @@ bool FGLTFMaterialTask::TryGetSourceTexture(const UTexture2D*& OutTexture, int32
 			return false;
 		}
 
+		if (!Builder.ExportOptions->bExportTextureTransforms && OutTransform != FGLTFJsonTextureTransform())
+		{
+			Builder.AddWarningMessage(FString::Printf(TEXT("Texture transform for property %s in material %s disabled by export options"), *FGLTFNameUtility::GetName(Property), *Material->GetName()));
+			OutTransform = {};
+		}
+
 		return true;
 	}
 
@@ -1137,6 +1143,12 @@ bool FGLTFMaterialTask::TryGetSourceTexture(const UTexture2D*& OutTexture, int32
 			return false;
 		}
 
+		if (!Builder.ExportOptions->bExportTextureTransforms && OutTransform != FGLTFJsonTextureTransform())
+		{
+			Builder.AddWarningMessage(FString::Printf(TEXT("Texture transform for property %s in material %s disabled by export options"), *FGLTFNameUtility::GetName(Property), *Material->GetName()));
+			OutTransform = {};
+		}
+
 		return true;
 	}
 
diff --git a/Engine/Plugins/Enterprise/GLTFExporter/MigratedFromGit/Source/GLTFExporter/Public/GLTFExportOptions.h b/Engine/Plugins/Enterprise/GLTFExporter/MigratedFromGit/Source/GLTFExporter/Public/GLTFExportOptions.h
--- a/Engine/Plugins/Enterprise/GLTFExporter/MigratedFromGit/Source/GLTFExporter/Public/GLTFExportOptions.h
+++ b/Engine/Plugins/Enterprise/GLTFExporter/MigratedFromGit/Source/GLTFExporter/Public/GLTFExportOptions.h
@@ -88,6 +88,9 @@ public:
 	UPROPERTY(EditAnywhere, BlueprintReadWrite, config, Category = Animation)
 	uint32 bExportPlaybackSettings : 1;
 
+	UPROPERTY(EditAnywhere, BlueprintReadWrite, config, Category = Texture)
+	uint32 bExportTextureTransforms : 1;
+
 	UPROPERTY(EditAnywhere, BlueprintReadWrite, config, Category = Texture)
 	uint32 bExportSourceTextures : 1;
 
