[GLTFExporter] refactor: improve warning messages for texture transform

commit 2ff9758c24367a2d615418eeaca060181bb9e222
Date: Tue, 20 Oct 2020 22:09:59 +0200

    refactor: improve warning messages for texture transform
    (cherry picked from commit 71e1852d54b2fdc0bc4f9d33a6116664f797e3d3​)

#skipci

diff --git a/Engine/Plugins/Enterprise/GLTFExporter/MigratedFromGit/Source/GLTFExporter/Private/Tasks/GLTFMaterialTasks.cpp b/Engine/Plugins/Enterprise/GLTFExporter/MigratedFromGit/Source/GLTFExporter/Private/Tasks/GLTFMaterialTasks.cpp
--- a/Engine/Plugins/Enterprise/GLTFExporter/MigratedFromGit/Source/GLTFExporter/Private/Tasks/GLTFMaterialTasks.cpp
+++ b/Engine/Plugins/Enterprise/GLTFExporter/MigratedFromGit/Source/GLTFExporter/Private/Tasks/GLTFMaterialTasks.cpp
@@ -259,7 +259,7 @@ bool FGLTFMaterialTask::TryGetBaseColorAndOpacity(FGLTFJsonPBRMetallicRoughness&
 		bHasOpacitySourceTexture &&
 		BaseColorTexture == OpacityTexture &&
 		BaseColorTexCoord == OpacityTexCoord &&
-        BaseColorTransform == OpacityTransform)
+		BaseColorTransform == OpacityTransform)
 	{
 		OutPBRParams.BaseColorTexture.Index = Builder.GetOrAddTexture(BaseColorTexture);
 		OutPBRParams.BaseColorTexture.TexCoord = BaseColorTexCoord;
@@ -270,7 +270,7 @@ bool FGLTFMaterialTask::TryGetBaseColorAndOpacity(FGLTFJsonPBRMetallicRoughness&
 	if (!Builder.ExportOptions->bBakeMaterialInputs)
 	{
 		Builder.AddWarningMessage(FString::Printf(
-			TEXT("%s and %s for material %s needs to bake, but material baking disabled by export options"),
+			TEXT("%s and %s for material %s needs to bake, but material baking is disabled by export options"),
 			FGLTFMaterialUtility::GetPropertyName(BaseColorProperty),
 			FGLTFMaterialUtility::GetPropertyName(OpacityProperty),
 			*Material->GetName()));
@@ -428,7 +428,7 @@ bool FGLTFMaterialTask::TryGetMetallicAndRoughness(FGLTFJsonPBRMetallicRoughness
 		bHasRoughnessSourceTexture &&
 		MetallicTexture == RoughnessTexture &&
 		MetallicTexCoord == RoughnessTexCoord &&
-        MetallicTransform == RoughnessTransform)
+		MetallicTransform == RoughnessTransform)
 	{
 		OutPBRParams.MetallicRoughnessTexture.Index = Builder.GetOrAddTexture(MetallicTexture);
 		OutPBRParams.MetallicRoughnessTexture.TexCoord = MetallicTexCoord;
@@ -439,7 +439,7 @@ bool FGLTFMaterialTask::TryGetMetallicAndRoughness(FGLTFJsonPBRMetallicRoughness
 	if (!Builder.ExportOptions->bBakeMaterialInputs)
 	{
 		Builder.AddWarningMessage(FString::Printf(
-			TEXT("%s and %s for material %s needs to bake, but material baking disabled by export options"),
+			TEXT("%s and %s for material %s needs to bake, but material baking is disabled by export options"),
 			FGLTFMaterialUtility::GetPropertyName(MetallicProperty),
 			FGLTFMaterialUtility::GetPropertyName(RoughnessProperty),
 			*Material->GetName()));
@@ -593,7 +593,7 @@ bool FGLTFMaterialTask::TryGetClearCoatRoughness(FGLTFJsonClearCoatExtension& Ou
 		bHasRoughnessSourceTexture &&
 		IntensityTexture == RoughnessTexture &&
 		IntensityTexCoord == RoughnessTexCoord &&
-        IntensityTransform == RoughnessTransform)
+		IntensityTransform == RoughnessTransform)
 	{
 		const FGLTFJsonTextureIndex TextureIndex = Builder.GetOrAddTexture(IntensityTexture);
 		OutExtParams.ClearCoatTexture.Index = TextureIndex;
@@ -607,7 +607,7 @@ bool FGLTFMaterialTask::TryGetClearCoatRoughness(FGLTFJsonClearCoatExtension& Ou
 	if (!Builder.ExportOptions->bBakeMaterialInputs)
 	{
 		Builder.AddWarningMessage(FString::Printf(
-			TEXT("%s and %s for material %s needs to bake, but material baking disabled by export options"),
+			TEXT("%s and %s for material %s needs to bake, but material baking is disabled by export options"),
 			FGLTFMaterialUtility::GetPropertyName(IntensityProperty),
 			FGLTFMaterialUtility::GetPropertyName(RoughnessProperty),
 			*Material->GetName()));
@@ -751,7 +751,7 @@ bool FGLTFMaterialTask::TryGetEmissive(FGLTFJsonMaterial& JsonMaterial, EMateria
 	if (!Builder.ExportOptions->bBakeMaterialInputs)
 	{
 		Builder.AddWarningMessage(FString::Printf(
-			TEXT("%s for material %s needs to bake, but material baking disabled by export options"),
+			TEXT("%s for material %s needs to bake, but material baking is disabled by export options"),
 			FGLTFMaterialUtility::GetPropertyName(EmissiveProperty),
 			*Material->GetName()));
 		return false;
@@ -1112,7 +1112,11 @@ bool FGLTFMaterialTask::TryGetSourceTexture(const UTexture2D*& OutTexture, int32
 
 		if (!Builder.ExportOptions->bExportTextureTransforms && OutTransform != FGLTFJsonTextureTransform())
 		{
-			Builder.AddWarningMessage(FString::Printf(TEXT("Texture transform for property %s in material %s disabled by export options"), *FGLTFNameUtility::GetName(Property), *Material->GetName()));
+			Builder.AddWarningMessage(FString::Printf(
+				TEXT("Texture coordinates [%d] in %s for material %s are transformed, but texture transform is disabled by export options"),
+				OutTexCoord,
+				*FGLTFNameUtility::GetName(Property),
+				*Material->GetName()));
 			OutTransform = {};
 		}
 
@@ -1145,7 +1149,11 @@ bool FGLTFMaterialTask::TryGetSourceTexture(const UTexture2D*& OutTexture, int32
 
 		if (!Builder.ExportOptions->bExportTextureTransforms && OutTransform != FGLTFJsonTextureTransform())
 		{
-			Builder.AddWarningMessage(FString::Printf(TEXT("Texture transform for property %s in material %s disabled by export options"), *FGLTFNameUtility::GetName(Property), *Material->GetName()));
+			Builder.AddWarningMessage(FString::Printf(
+				TEXT("Texture coordinates [%d] in %s for material %s are transformed, but texture transform is disabled by export options"),
+				OutTexCoord,
+				*FGLTFNameUtility::GetName(Property),
+				*Material->GetName()));
 			OutTransform = {};
 		}
 
@@ -1160,7 +1168,7 @@ bool FGLTFMaterialTask::TryGetBakedMaterialProperty(FGLTFJsonTextureInfo& OutTex
 	if (!Builder.ExportOptions->bBakeMaterialInputs)
 	{
 		Builder.AddWarningMessage(FString::Printf(
-			TEXT("%s for material %s needs to bake, but material baking disabled by export options"),
+			TEXT("%s for material %s needs to bake, but material baking is disabled by export options"),
 			FGLTFMaterialUtility::GetPropertyName(Property),
 			*Material->GetName()));
 		return false;
@@ -1188,7 +1196,7 @@ bool FGLTFMaterialTask::TryGetBakedMaterialProperty(FGLTFJsonTextureInfo& OutTex
 	if (!Builder.ExportOptions->bBakeMaterialInputs)
 	{
 		Builder.AddWarningMessage(FString::Printf(
-			TEXT("%s for material %s needs to bake, but material baking disabled by export options"),
+			TEXT("%s for material %s needs to bake, but material baking is disabled by export options"),
 			FGLTFMaterialUtility::GetPropertyName(Property),
 			*Material->GetName()));
 		return false;
@@ -1216,7 +1224,7 @@ inline bool FGLTFMaterialTask::TryGetBakedMaterialProperty(FGLTFJsonTextureInfo&
 	if (!Builder.ExportOptions->bBakeMaterialInputs)
 	{
 		Builder.AddWarningMessage(FString::Printf(
-			TEXT("%s for material %s needs to bake, but material baking disabled by export options"),
+			TEXT("%s for material %s needs to bake, but material baking is disabled by export options"),
 			FGLTFMaterialUtility::GetPropertyName(Property),
 			*Material->GetName()));
 		return false;
@@ -1244,7 +1252,7 @@ bool FGLTFMaterialTask::TryGetBakedMaterialProperty(FGLTFJsonTextureInfo& OutTex
 	if (!Builder.ExportOptions->bBakeMaterialInputs)
 	{
 		Builder.AddWarningMessage(FString::Printf(
-			TEXT("%s for material %s needs to bake, but material baking disabled by export options"),
+			TEXT("%s for material %s needs to bake, but material baking is disabled by export options"),
 			FGLTFMaterialUtility::GetPropertyName(Property),
 			*Material->GetName()));
 		return false;
