[GLTFExporter] refactor: simplify LoadOptions and SaveOptions in GLTFExportOptions

commit e9126830dc2517a2f99cd6e61db44d96f6a30322
Date: Tue, 20 Oct 2020 22:35:12 +0200

    refactor: simplify LoadOptions and SaveOptions in GLTFExportOptions
    (cherry picked from commit 89219eea3fa3ab3c4c185f219b88051c5266183c​)

#skipci

diff --git a/Engine/Plugins/Enterprise/GLTFExporter/MigratedFromGit/Source/GLTFExporter/Private/GLTFExportOptions.cpp b/Engine/Plugins/Enterprise/GLTFExporter/MigratedFromGit/Source/GLTFExporter/Private/GLTFExportOptions.cpp
--- a/Engine/Plugins/Enterprise/GLTFExporter/MigratedFromGit/Source/GLTFExporter/Private/GLTFExportOptions.cpp
+++ b/Engine/Plugins/Enterprise/GLTFExporter/MigratedFromGit/Source/GLTFExporter/Private/GLTFExportOptions.cpp
@@ -52,145 +52,12 @@ void UGLTFExportOptions::ResetToDefault()
 
 void UGLTFExportOptions::LoadOptions()
 {
-	const int32 PortFlags = 0;
-
-	for (FProperty* Property = GetClass()->PropertyLink; Property; Property = Property->PropertyLinkNext)
-	{
-		if (!Property->HasAnyPropertyFlags(CPF_Config))
-		{
-			continue;
-		}
-		FString Section = GetClass()->GetName();
-		FString Key = Property->GetName();
-
-		const bool bIsPropertyInherited = Property->GetOwnerClass() != GetClass();
-		UObject* SuperClassDefaultObject = GetClass()->GetSuperClass()->GetDefaultObject();
-
-		const FString& PropFileName = GEditorPerProjectIni;
-
-		FArrayProperty* Array = CastField<FArrayProperty>(Property);
-		if (Array)
-		{
-			FConfigSection* Sec = GConfig->GetSectionPrivate(*Section, false, true, *GEditorPerProjectIni);
-			if (Sec != nullptr)
-			{
-				TArray<FConfigValue> List;
-				const FName KeyName(*Key, FNAME_Find);
-				Sec->MultiFind(KeyName, List);
-
-				FScriptArrayHelper_InContainer ArrayHelper(Array, this);
-				// Only override default properties if there is something to override them with.
-				if (List.Num() > 0)
-				{
-					ArrayHelper.EmptyAndAddValues(List.Num());
-					for (int32 i = List.Num() - 1, c = 0; i >= 0; i--, c++)
-					{
-						Array->Inner->ImportText(*List[i].GetValue(), ArrayHelper.GetRawPtr(c), PortFlags, this);
-					}
-				}
-				else
-				{
-					int32 Index = 0;
-					const FConfigValue* ElementValue = nullptr;
-					do
-					{
-						// Add array index number to end of key
-						FString IndexedKey = FString::Printf(TEXT("%s[%i]"), *Key, Index);
-
-						// Try to find value of key
-						const FName IndexedName(*IndexedKey, FNAME_Find);
-						if (IndexedName == NAME_None)
-						{
-							break;
-						}
-						ElementValue = Sec->Find(IndexedName);
-
-						// If found, import the element
-						if (ElementValue != nullptr)
-						{
-							// expand the array if necessary so that Index is a valid element
-							ArrayHelper.ExpandForIndex(Index);
-							Array->Inner->ImportText(*ElementValue->GetValue(), ArrayHelper.GetRawPtr(Index), PortFlags, this);
-						}
-
-						Index++;
-					} while (ElementValue || Index < ArrayHelper.Num());
-				}
-			}
-		}
-		else
-		{
-			for (int32 i = 0; i < Property->ArrayDim; i++)
-			{
-				if (Property->ArrayDim != 1)
-				{
-					Key = FString::Printf(TEXT("%s[%i]"), *Property->GetName(), i);
-				}
-
-				FString Value;
-				const bool bFoundValue = GConfig->GetString(*Section, *Key, Value, *GEditorPerProjectIni);
-
-				if (bFoundValue)
-				{
-					if (Property->ImportText(*Value, Property->ContainerPtrToValuePtr<uint8>(this, i), PortFlags, this) == nullptr)
-					{
-						// this should be an error as the properties from the .ini / .int file are not correctly being read in and probably are affecting things in subtle ways
-					}
-				}
-			}
-		}
-	}
+	LoadConfig(nullptr, *GEditorPerProjectIni);
 }
 
 void UGLTFExportOptions::SaveOptions()
 {
-	const int32 PortFlags = 0;
-
-	for (FProperty* Property = GetClass()->PropertyLink; Property; Property = Property->PropertyLinkNext)
-	{
-		if (!Property->HasAnyPropertyFlags(CPF_Config))
-		{
-			continue;
-		}
-		FString Section = GetClass()->GetName();
-		FString Key = Property->GetName();
-
-		const bool bIsPropertyInherited = Property->GetOwnerClass() != GetClass();
-		UObject* SuperClassDefaultObject = GetClass()->GetSuperClass()->GetDefaultObject();
-
-		FArrayProperty* Array = CastField<FArrayProperty>(Property);
-		if (Array)
-		{
-			FConfigSection* Sec = GConfig->GetSectionPrivate(*Section, true, false, *GEditorPerProjectIni);
-			check(Sec);
-			Sec->Remove(*Key);
-
-			FScriptArrayHelper_InContainer ArrayHelper(Array, this);
-			for (int32 i = 0; i < ArrayHelper.Num(); i++)
-			{
-				FString Buffer;
-				Array->Inner->ExportTextItem(Buffer, ArrayHelper.GetRawPtr(i), ArrayHelper.GetRawPtr(i), this, PortFlags);
-				Sec->Add(*Key, *Buffer);
-			}
-		}
-		else
-		{
-			TCHAR TempKey[MAX_SPRINTF] = TEXT("");
-			for (int32 Index = 0; Index < Property->ArrayDim; Index++)
-			{
-				if (Property->ArrayDim != 1)
-				{
-					FCString::Sprintf(TempKey, TEXT("%s[%i]"), *Property->GetName(), Index);
-					Key = TempKey;
-				}
-
-				FString Value;
-				Property->ExportText_InContainer(Index, Value, this, this, this, PortFlags);
-				GConfig->SetString(*Section, *Key, *Value, *GEditorPerProjectIni);
-			}
-		}
-	}
-	GConfig->Flush(false);
+	SaveConfig(CPF_Config, *GEditorPerProjectIni);
 }
 
 bool UGLTFExportOptions::CanEditChange(const FProperty* InProperty) const
