[GLTFExporter] refactor: disable bExportVertexColors by default since it often represented undesirably, due to gltf limitations. To properly export vertex-colors it should instead be baked into the material.

commit c5654cff3dcdc4afad1700e25532bfe0a5264c8f
Date: Tue, 20 Oct 2020 22:39:06 +0200

    refactor: disable bExportVertexColors by default since it often represented undesirably, due to gltf limitations. To properly export vertex-colors it should instead be baked into the material.
    (cherry picked from commit c45382492fc1b55ab9703073e84706ab0fe1283c​)

#skipci

diff --git a/Engine/Plugins/Enterprise/GLTFExporter/MigratedFromGit/Source/GLTFExporter/Private/GLTFExportOptions.cpp b/Engine/Plugins/Enterprise/GLTFExporter/MigratedFromGit/Source/GLTFExporter/Private/GLTFExportOptions.cpp
--- a/Engine/Plugins/Enterprise/GLTFExporter/MigratedFromGit/Source/GLTFExporter/Private/GLTFExportOptions.cpp
+++ b/Engine/Plugins/Enterprise/GLTFExporter/MigratedFromGit/Source/GLTFExporter/Private/GLTFExportOptions.cpp
@@ -17,7 +17,7 @@ UGLTFExportOptions::UGLTFExportOptions(const FObjectInitializer& ObjectInitializ
 	bExportClearCoatMaterials = true;
 	bBakeMaterialInputs = true;
 	DefaultMaterialBakeSize = EGLTFExporterMaterialBakeSize::POT_1024;
-	bExportVertexColors = true;
+	bExportVertexColors = false;
 	bQuantizeVertexNormals = true;
 	bQuantizeVertexTangents = true;
 	DefaultLevelOfDetail = 0;
