[GLTFExporter] fix: make sure we don't use saved editor settings during export tests

commit 5ead3f42f1196f6a279a143233516e6fea9a318b
Date: Wed, 21 Oct 2020 00:07:15 +0200

    fix: make sure we don't use saved editor settings during export tests
    (cherry picked from commit c0acee304f673256efd2d1f0e9f6a9210014519e​)

#skipci

diff --git a/Engine/Plugins/Enterprise/GLTFExporter/MigratedFromGit/Source/GLTFExporter/Private/Exporters/GLTFExporter.cpp b/Engine/Plugins/Enterprise/GLTFExporter/MigratedFromGit/Source/GLTFExporter/Private/Exporters/GLTFExporter.cpp
--- a/Engine/Plugins/Enterprise/GLTFExporter/MigratedFromGit/Source/GLTFExporter/Private/Exporters/GLTFExporter.cpp
+++ b/Engine/Plugins/Enterprise/GLTFExporter/MigratedFromGit/Source/GLTFExporter/Private/Exporters/GLTFExporter.cpp
@@ -88,7 +88,6 @@ UGLTFExportOptions* UGLTFExporter::GetExportOptions()
 	if (Options == nullptr)
 	{
 		Options = NewObject<UGLTFExportOptions>();
-		Options->LoadOptions();
 	}
 
 	if (GetShowExportOption() && !bAutomatedTask)
@@ -106,7 +105,7 @@ UGLTFExportOptions* UGLTFExporter::GetExportOptions()
 		}
 
 		SetShowExportOption(!bExportAll);
-		Options->SaveOptions();
+		Options->SaveConfig();
 	}
 
 	return Options;
diff --git a/Engine/Plugins/Enterprise/GLTFExporter/MigratedFromGit/Source/GLTFExporter/Private/GLTFExportOptions.cpp b/Engine/Plugins/Enterprise/GLTFExporter/MigratedFromGit/Source/GLTFExporter/Private/GLTFExportOptions.cpp
--- a/Engine/Plugins/Enterprise/GLTFExporter/MigratedFromGit/Source/GLTFExporter/Private/GLTFExportOptions.cpp
+++ b/Engine/Plugins/Enterprise/GLTFExporter/MigratedFromGit/Source/GLTFExporter/Private/GLTFExportOptions.cpp
@@ -1,17 +1,18 @@
 // Copyright Epic Games, Inc. All Rights Reserved.
 
 #include "GLTFExportOptions.h"
-#include "UObject/ObjectMacros.h"
 #include "UObject/UObjectGlobals.h"
 #include "UObject/Object.h"
-#include "UObject/UObjectBaseUtility.h"
 #include "UObject/Class.h"
 #include "UObject/UnrealType.h"
-#include "Misc/ConfigCacheIni.h"
-#include "Misc/App.h"
 
 UGLTFExportOptions::UGLTFExportOptions(const FObjectInitializer& ObjectInitializer)
 	: Super(ObjectInitializer)
+{
+	ResetToDefault();
+}
+
+void UGLTFExportOptions::ResetToDefault()
 {
 	bExportUnlitMaterials = true;
 	bExportClearCoatMaterials = true;
@@ -45,21 +46,6 @@ UGLTFExportOptions::UGLTFExportOptions(const FObjectInitializer& ObjectInitializ
 	bShowFilesWhenDone = true;
 }
 
-void UGLTFExportOptions::ResetToDefault()
-{
-	ReloadConfig();
-}
-
-void UGLTFExportOptions::LoadOptions()
-{
-	LoadConfig(nullptr, *GEditorPerProjectIni);
-}
-
-void UGLTFExportOptions::SaveOptions()
-{
-	SaveConfig(CPF_Config, *GEditorPerProjectIni);
-}
-
 bool UGLTFExportOptions::CanEditChange(const FProperty* InProperty) const
 {
 	const FName PropertyFName = InProperty->GetFName();
