[GLTFExporter] refactor: change integer bits to booleans in GLTFExportOptions

commit e7571f9d110163be6e22d7e1dd7ed682ef37f1f1
Date: Wed, 21 Oct 2020 00:08:09 +0200

    refactor: change integer bits to booleans in GLTFExportOptions
    (cherry picked from commit af4d27bb5c015af21a671cc5f23c1a76498fa959​)
    # Conflicts from cherry-pick:
    #	Plugins/GLTFExporter/Source/GLTFExporter/Public/GLTFExportOptions.h

#skipci

diff --git a/Engine/Plugins/Enterprise/GLTFExporter/MigratedFromGit/Source/GLTFExporter/Public/GLTFExportOptions.h b/Engine/Plugins/Enterprise/GLTFExporter/MigratedFromGit/Source/GLTFExporter/Public/GLTFExportOptions.h
--- a/Engine/Plugins/Enterprise/GLTFExporter/MigratedFromGit/Source/GLTFExporter/Public/GLTFExportOptions.h
+++ b/Engine/Plugins/Enterprise/GLTFExporter/MigratedFromGit/Source/GLTFExporter/Public/GLTFExportOptions.h
@@ -42,104 +42,104 @@ enum class EGLTFExporterLightMobility : uint8
 	All
 };
 
-UCLASS(config = EditorPerProjectUserSettings, HideCategories = (DebugProperty))
+UCLASS(Config=EditorPerProjectUserSettings, HideCategories=(DebugProperty))
 class GLTFEXPORTER_API UGLTFExportOptions : public UObject
 {
 	GENERATED_UCLASS_BODY()
 
 public:
 
-	UPROPERTY(EditAnywhere, BlueprintReadWrite, config, Category = Material)
-	uint32 bExportUnlitMaterials : 1;
+	UPROPERTY(EditAnywhere, BlueprintReadWrite, Config, Category = Material)
+	bool bExportUnlitMaterials;
 
-	UPROPERTY(EditAnywhere, BlueprintReadWrite, config, Category = Material)
-	uint32 bExportClearCoatMaterials : 1;
+	UPROPERTY(EditAnywhere, BlueprintReadWrite, Config, Category = Material)
+	bool bExportClearCoatMaterials;
 
-	UPROPERTY(EditAnywhere, BlueprintReadWrite, config, Category = Material)
-	uint32 bBakeMaterialInputs : 1;
+	UPROPERTY(EditAnywhere, BlueprintReadWrite, Config, Category = Material)
+	bool bBakeMaterialInputs;
 
-	UPROPERTY(EditAnywhere, BlueprintReadWrite, config, Category = Material)
+	UPROPERTY(EditAnywhere, BlueprintReadWrite, Config, Category = Material)
 	EGLTFExporterMaterialBakeSize DefaultMaterialBakeSize;
 
-	UPROPERTY(EditAnywhere, BlueprintReadWrite, config, Category = Mesh)
-	uint32 bExportVertexColors : 1;
+	UPROPERTY(EditAnywhere, BlueprintReadWrite, Config, Category = Mesh)
+	bool bExportVertexColors;
 
-	UPROPERTY(EditAnywhere, BlueprintReadWrite, config, Category = Mesh)
-	uint32 bQuantizeVertexNormals : 1;
+	UPROPERTY(EditAnywhere, BlueprintReadWrite, Config, Category = Mesh)
+	bool bQuantizeVertexNormals;
 
-	UPROPERTY(EditAnywhere, BlueprintReadWrite, config, Category = Mesh)
-	uint32 bQuantizeVertexTangents : 1;
+	UPROPERTY(EditAnywhere, BlueprintReadWrite, Config, Category = Mesh)
+	bool bQuantizeVertexTangents;
 
-	UPROPERTY(EditAnywhere, BlueprintReadWrite, config, Category = Mesh)
+	UPROPERTY(EditAnywhere, BlueprintReadWrite, Config, Category = Mesh)
 	int32 DefaultLevelOfDetail;
 
-	UPROPERTY(EditAnywhere, BlueprintReadWrite, config, Category = Animation)
-	uint32 bExportLevelSequences : 1;
+	UPROPERTY(EditAnywhere, BlueprintReadWrite, Config, Category = Animation)
+	bool bExportLevelSequences;
 
-	UPROPERTY(EditAnywhere, BlueprintReadWrite, config, Category = Animation)
-	uint32 bExportVertexSkinWeights : 1;
+	UPROPERTY(EditAnywhere, BlueprintReadWrite, Config, Category = Animation)
+	bool bExportVertexSkinWeights;
 
-	UPROPERTY(EditAnywhere, BlueprintReadWrite, config, Category = Animation)
-	uint32 bExportAnimationSequences : 1;
+	UPROPERTY(EditAnywhere, BlueprintReadWrite, Config, Category = Animation)
+	bool bExportAnimationSequences;
 
-	UPROPERTY(EditAnywhere, BlueprintReadWrite, config, Category = Animation)
-	uint32 bRetargetBoneTransforms : 1;
+	UPROPERTY(EditAnywhere, BlueprintReadWrite, Config, Category = Animation)
+	bool bRetargetBoneTransforms;
 
-	UPROPERTY(EditAnywhere, BlueprintReadWrite, config, Category = Animation)
-	uint32 bExportPlaybackSettings : 1;
+	UPROPERTY(EditAnywhere, BlueprintReadWrite, Config, Category = Animation)
+	bool bExportPlaybackSettings;
 
-	UPROPERTY(EditAnywhere, BlueprintReadWrite, config, Category = Texture)
-	uint32 bExportTextureTransforms : 1;
+	UPROPERTY(EditAnywhere, BlueprintReadWrite, Config, Category = Texture)
+	bool bExportTextureTransforms;
 
-	UPROPERTY(EditAnywhere, BlueprintReadWrite, config, Category = Texture)
-	uint32 bExportSourceTextures : 1;
+	UPROPERTY(EditAnywhere, BlueprintReadWrite, Config, Category = Texture)
+	bool bExportSourceTextures;
 
-	UPROPERTY(EditAnywhere, BlueprintReadWrite, config, Category = Texture)
-	uint32 bExportLightmaps : 1;
+	UPROPERTY(EditAnywhere, BlueprintReadWrite, Config, Category = Texture)
+	bool bExportLightmaps;
 
-	UPROPERTY(EditAnywhere, BlueprintReadWrite, config, Category = Texture)
+	UPROPERTY(EditAnywhere, BlueprintReadWrite, Config, Category = Texture)
 	EGLTFExporterTextureHDREncoding TextureHDREncoding;
 
-	UPROPERTY(EditAnywhere, BlueprintReadWrite, config, Category = Scene, meta =
+	UPROPERTY(EditAnywhere, BlueprintReadWrite, Config, Category = Scene, meta =
 		(DisplayName = "Export Uniform Scale",
 		 ToolTip = "Scale factor used for exporting assets, by default: 0.01, for conversion from centimeters(Unreal default) to meters(glTF)."))
 	float ExportScale;
 
-	UPROPERTY(EditAnywhere, BlueprintReadWrite, config, Category = Scene)
-	uint32 bExportHiddenInGame : 1;
+	UPROPERTY(EditAnywhere, BlueprintReadWrite, Config, Category = Scene)
+	bool bExportHiddenInGame;
 
-	UPROPERTY(EditAnywhere, BlueprintReadWrite, config, Category = Scene)
+	UPROPERTY(EditAnywhere, BlueprintReadWrite, Config, Category = Scene)
 	EGLTFExporterLightMobility ExportLights;
 
-	UPROPERTY(EditAnywhere, BlueprintReadWrite, config, Category = Scene)
-	uint32 bExportCameras : 1;
+	UPROPERTY(EditAnywhere, BlueprintReadWrite, Config, Category = Scene)
+	bool bExportCameras;
 
-	UPROPERTY(EditAnywhere, BlueprintReadWrite, config, Category = Scene)
-	uint32 bExportOrbitalCameras : 1;
+	UPROPERTY(EditAnywhere, BlueprintReadWrite, Config, Category = Scene)
+	bool bExportOrbitalCameras;
 
-	UPROPERTY(EditAnywhere, BlueprintReadWrite, config, Category = Scene)
-	uint32 bExportHDRIBackdrops : 1;
+	UPROPERTY(EditAnywhere, BlueprintReadWrite, Config, Category = Scene)
+	bool bExportHDRIBackdrops;
 
-	UPROPERTY(EditAnywhere, BlueprintReadWrite, config, Category = Scene)
-	uint32 bExportSkySpheres : 1;
+	UPROPERTY(EditAnywhere, BlueprintReadWrite, Config, Category = Scene)
+	bool bExportSkySpheres;
 
-	UPROPERTY(EditAnywhere, BlueprintReadWrite, config, Category = Scene)
-	uint32 bExportVariantSets : 1;
+	UPROPERTY(EditAnywhere, BlueprintReadWrite, Config, Category = Scene)
+	bool bExportVariantSets;
 
-	UPROPERTY(EditAnywhere, BlueprintReadWrite, config, Category = Scene)
-	uint32 bExportAnimationHotspots : 1;
+	UPROPERTY(EditAnywhere, BlueprintReadWrite, Config, Category = Scene)
+	bool bExportAnimationHotspots;
 
-	UPROPERTY(EditAnywhere, BlueprintReadWrite, config, Category = Exporter)
-	uint32 bBundleWebViewer : 1;
+	UPROPERTY(EditAnywhere, BlueprintReadWrite, Config, Category = Exporter)
+	bool bBundleWebViewer;
 
-	UPROPERTY(EditAnywhere, BlueprintReadWrite, config, Category = Exporter)
-	uint32 bExportPreviewMesh : 1;
+	UPROPERTY(EditAnywhere, BlueprintReadWrite, Config, Category = Exporter)
+	bool bExportPreviewMesh;
 
-	UPROPERTY(EditAnywhere, BlueprintReadWrite, config, Category = Exporter)
-	uint32 bAllExtensionsRequired : 1;
+	UPROPERTY(EditAnywhere, BlueprintReadWrite, Config, Category = Exporter)
+	bool bAllExtensionsRequired;
 
-	UPROPERTY(EditAnywhere, BlueprintReadWrite, config, Category = Exporter)
-	uint32 bShowFilesWhenDone : 1;
+	UPROPERTY(EditAnywhere, BlueprintReadWrite, Config, Category = Exporter)
+	bool bShowFilesWhenDone;
 
 	/* Set all the each property to the CDO value */
 	void ResetToDefault();
