[GLTFExporter] refactor: use checkNoEntry in enum conversions that need to be exhaustive

commit 291db6e828bbfe8df043a680d296343764f79983
Date: Mon, 19 Oct 2020 08:58:28 +0200

    refactor: use checkNoEntry in enum conversions that need to be exhaustive
    (cherry picked from commit 5f1539376bc7822fbf9ab743a29e242b79e02b6c​)
    # Conflicts from cherry-pick:
    #	Plugins/GLTFExporter/Source/GLTFExporter/Private/Json/GLTFJsonUtility.h

#skipci

diff --git a/Engine/Plugins/Enterprise/GLTFExporter/MigratedFromGit/Source/GLTFExporter/Private/Builders/GLTFBuilderUtility.cpp b/Engine/Plugins/Enterprise/GLTFExporter/MigratedFromGit/Source/GLTFExporter/Private/Builders/GLTFBuilderUtility.cpp
--- a/Engine/Plugins/Enterprise/GLTFExporter/MigratedFromGit/Source/GLTFExporter/Private/Builders/GLTFBuilderUtility.cpp
+++ b/Engine/Plugins/Enterprise/GLTFExporter/MigratedFromGit/Source/GLTFExporter/Private/Builders/GLTFBuilderUtility.cpp
@@ -8,7 +8,9 @@ const TCHAR* FGLTFBuilderUtility::GetFileExtension(EGLTFJsonMimeType MimeType)
 	{
 		case EGLTFJsonMimeType::PNG:  return TEXT(".png");
 		case EGLTFJsonMimeType::JPEG: return TEXT(".jpg");
-		default:                      return TEXT(".unknown");
+		default:
+			checkNoEntry();
+			return TEXT(".unknown");
 	}
 }
 
diff --git a/Engine/Plugins/Enterprise/GLTFExporter/MigratedFromGit/Source/GLTFExporter/Private/Converters/GLTFCameraConverters.cpp b/Engine/Plugins/Enterprise/GLTFExporter/MigratedFromGit/Source/GLTFExporter/Private/Converters/GLTFCameraConverters.cpp
--- a/Engine/Plugins/Enterprise/GLTFExporter/MigratedFromGit/Source/GLTFExporter/Private/Converters/GLTFCameraConverters.cpp
+++ b/Engine/Plugins/Enterprise/GLTFExporter/MigratedFromGit/Source/GLTFExporter/Private/Converters/GLTFCameraConverters.cpp
@@ -26,9 +26,13 @@ FGLTFJsonCameraIndex FGLTFCameraConverter::Convert(const UCameraComponent* Camer
 			Camera.Perspective = FGLTFCameraUtility::ConvertPerspective(DesiredView, Builder.ExportOptions->ExportScale);
 			break;
 
-		default:
+		case EGLTFJsonCameraType::None:
 			// TODO: report error (unsupported camera type)
 			return FGLTFJsonCameraIndex(INDEX_NONE);
+
+		default:
+			checkNoEntry();
+			break;
 	}
 
 	const AActor* Owner = CameraComponent->GetOwner();
diff --git a/Engine/Plugins/Enterprise/GLTFExporter/MigratedFromGit/Source/GLTFExporter/Private/Json/GLTFJsonUtility.h b/Engine/Plugins/Enterprise/GLTFExporter/MigratedFromGit/Source/GLTFExporter/Private/Json/GLTFJsonUtility.h
--- a/Engine/Plugins/Enterprise/GLTFExporter/MigratedFromGit/Source/GLTFExporter/Private/Json/GLTFJsonUtility.h
+++ b/Engine/Plugins/Enterprise/GLTFExporter/MigratedFromGit/Source/GLTFExporter/Private/Json/GLTFJsonUtility.h
@@ -13,8 +13,6 @@ struct FGLTFJsonUtility
 		return static_cast<int32>(Value);
 	}
 
-	// TODO: replace all default cases (i.e. "unknown") with checkNoEntry() in each ToString
-
 	static const TCHAR* ToString(EGLTFJsonExtension Value)
 	{
 		switch (Value)
@@ -33,7 +31,9 @@ struct FGLTFJsonUtility
 			case EGLTFJsonExtension::EPIC_TextureHDREncoding: return TEXT("EPIC_texture_hdr_encoding");
 			case EGLTFJsonExtension::EPIC_AnimationPlayback:  return TEXT("EPIC_animation_playback");
 			case EGLTFJsonExtension::EPIC_SkySpheres:         return TEXT("EPIC_sky_spheres");
-			default:                                          return TEXT("unknown");
+			default:
+				checkNoEntry();
+				return TEXT("unknown");
 		}
 	}
 
@@ -44,7 +44,9 @@ struct FGLTFJsonUtility
 			case EGLTFJsonAlphaMode::Opaque: return TEXT("OPAQUE");
 			case EGLTFJsonAlphaMode::Blend:  return TEXT("BLEND");
 			case EGLTFJsonAlphaMode::Mask:   return TEXT("MASK");
-			default:                         return TEXT("UNKNOWN");
+			default:
+				checkNoEntry();
+				return TEXT("UNKNOWN");
 		}
 	}
 
@@ -54,7 +56,9 @@ struct FGLTFJsonUtility
 		{
 			case EGLTFJsonMimeType::PNG:  return TEXT("image/png");
 			case EGLTFJsonMimeType::JPEG: return TEXT("image/jpeg");
-			default:                      return TEXT("unknown");
+			default:
+				checkNoEntry();
+				return TEXT("unknown");
 		}
 	}
 
@@ -69,7 +73,9 @@ struct FGLTFJsonUtility
 			case EGLTFJsonAccessorType::Mat2:   return TEXT("MAT2");
 			case EGLTFJsonAccessorType::Mat3:   return TEXT("MAT3");
 			case EGLTFJsonAccessorType::Mat4:   return TEXT("MAT4");
-			default:                            return TEXT("UNKNOWN");
+			default:
+				checkNoEntry();
+				return TEXT("UNKNOWN");
 		}
 	}
 
@@ -79,7 +85,9 @@ struct FGLTFJsonUtility
 		{
 			case EGLTFJsonHDREncoding::RGBE: return TEXT("RGBE");
 			case EGLTFJsonHDREncoding::RGBM: return TEXT("RGBM");
-			default:                         return TEXT("UNKNOWN");
+			default:
+				checkNoEntry();
+				return TEXT("UNKNOWN");
 		}
 	}
 
@@ -93,7 +101,9 @@ struct FGLTFJsonUtility
 			case EGLTFJsonCubeFace::NegY: return TEXT("NegY");
 			case EGLTFJsonCubeFace::PosZ: return TEXT("PosZ");
 			case EGLTFJsonCubeFace::NegZ: return TEXT("NegZ");
-			default:                      return TEXT("Unknown");
+			default:
+				checkNoEntry();
+				return TEXT("Unknown");
 		}
 	}
 
@@ -103,7 +113,9 @@ struct FGLTFJsonUtility
 		{
 			case EGLTFJsonCameraType::Perspective:  return TEXT("perspective");
 			case EGLTFJsonCameraType::Orthographic: return TEXT("orthographic");
-			default:                                return TEXT("unknown");
+			default:
+				checkNoEntry();
+				return TEXT("unknown");
 		}
 	}
 
@@ -114,7 +126,9 @@ struct FGLTFJsonUtility
 			case EGLTFJsonLightType::Directional: return TEXT("directional");
 			case EGLTFJsonLightType::Point:       return TEXT("point");
 			case EGLTFJsonLightType::Spot:        return TEXT("spot");
-			default:                              return TEXT("unknown");
+			default:
+				checkNoEntry();
+				return TEXT("unknown");
 		}
 	}
 
@@ -125,7 +139,9 @@ struct FGLTFJsonUtility
 			case EGLTFJsonInterpolation::Linear:      return TEXT("LINEAR");
 			case EGLTFJsonInterpolation::Step:        return TEXT("STEP");
 			case EGLTFJsonInterpolation::CubicSpline: return TEXT("CUBICSPLINE");
-			default:                                  return TEXT("UNKNOWN");
+			default:
+				checkNoEntry();
+				return TEXT("UNKNOWN");
 		}
 	}
 
@@ -137,7 +153,9 @@ struct FGLTFJsonUtility
 			case EGLTFJsonTargetPath::Rotation:    return TEXT("rotation");
 			case EGLTFJsonTargetPath::Scale:       return TEXT("scale");
 			case EGLTFJsonTargetPath::Weights:     return TEXT("weights");
-			default:                               return TEXT("unknown");
+			default:
+				checkNoEntry();
+				return TEXT("unknown");
 		}
 	}
 
