[GLTFExporter] refactor: add missing break statements to default cases

commit bdab5e099588c8f6379939ea78b44e51aa825951
Date: Mon, 19 Oct 2020 09:05:08 +0200

    refactor: add missing break statements to default cases
    (cherry picked from commit c40c333ce9ba3bb30eef47c55d4c4e69f0811260​)

#skipci

diff --git a/Engine/Plugins/Enterprise/GLTFExporter/MigratedFromGit/Source/GLTFExporter/Private/Converters/GLTFMaterialUtility.cpp b/Engine/Plugins/Enterprise/GLTFExporter/MigratedFromGit/Source/GLTFExporter/Private/Converters/GLTFMaterialUtility.cpp
--- a/Engine/Plugins/Enterprise/GLTFExporter/MigratedFromGit/Source/GLTFExporter/Private/Converters/GLTFMaterialUtility.cpp
+++ b/Engine/Plugins/Enterprise/GLTFExporter/MigratedFromGit/Source/GLTFExporter/Private/Converters/GLTFMaterialUtility.cpp
@@ -253,7 +253,7 @@ FGLTFPropertyBakeOutput FGLTFMaterialUtility::BakeMaterialProperty(const FIntPoi
 			case MP_CustomData0: HackProperty = MP_Metallic; break;
 			case MP_CustomData1: HackProperty = MP_Roughness; break;
 			case MP_CustomOutput: HackProperty = MP_Normal; break;
-			default: checkNoEntry();
+			default: checkNoEntry(); break;
 		}
 
 		UMaterial* ParentMaterial = const_cast<UMaterial*>(Material->GetMaterial());
diff --git a/Engine/Plugins/Enterprise/GLTFExporter/MigratedFromGit/Source/GLTFExporter/Private/Json/GLTFJsonMaterial.h b/Engine/Plugins/Enterprise/GLTFExporter/MigratedFromGit/Source/GLTFExporter/Private/Json/GLTFJsonMaterial.h
--- a/Engine/Plugins/Enterprise/GLTFExporter/MigratedFromGit/Source/GLTFExporter/Private/Json/GLTFJsonMaterial.h
+++ b/Engine/Plugins/Enterprise/GLTFExporter/MigratedFromGit/Source/GLTFExporter/Private/Json/GLTFJsonMaterial.h
@@ -322,6 +322,7 @@ struct FGLTFJsonMaterial
 
 				default:
 					checkNoEntry();
+					break;
 			}
 
 			JsonWriter.WriteObjectEnd();
diff --git a/Engine/Plugins/Enterprise/GLTFExporter/MigratedFromGit/Source/GLTFExporterRuntime/Private/Actors/GLTFHotspotActor.cpp b/Engine/Plugins/Enterprise/GLTFExporter/MigratedFromGit/Source/GLTFExporterRuntime/Private/Actors/GLTFHotspotActor.cpp
--- a/Engine/Plugins/Enterprise/GLTFExporter/MigratedFromGit/Source/GLTFExporterRuntime/Private/Actors/GLTFHotspotActor.cpp
+++ b/Engine/Plugins/Enterprise/GLTFExporter/MigratedFromGit/Source/GLTFExporterRuntime/Private/Actors/GLTFHotspotActor.cpp
@@ -508,7 +508,9 @@ const UTexture2D* AGLTFHotspotActor::GetImageForState(EGLTFHotspotState State) c
 		case EGLTFHotspotState::Hovered:        return CurrentHoveredImage;
 		case EGLTFHotspotState::Toggled:        return CurrentToggledImage;
 		case EGLTFHotspotState::ToggledHovered: return CurrentToggledHoveredImage;
-		default:                                checkNoEntry();
+		default:
+			checkNoEntry();
+			break;
 	}
 
 	return nullptr;
