[GLTFExporter] refactor: just return empty string after checkNoEntry

commit 55128e61de133f4cbb941f9e0a493d77490fbe24
Date: Tue, 20 Oct 2020 23:17:42 +0200

    refactor: just return empty string after checkNoEntry
    (cherry picked from commit 37803ee6a9774d05863f74787cd82585e2620760​)

#skipci

diff --git a/Engine/Plugins/Enterprise/GLTFExporter/MigratedFromGit/Source/GLTFExporter/Private/Builders/GLTFBuilderUtility.cpp b/Engine/Plugins/Enterprise/GLTFExporter/MigratedFromGit/Source/GLTFExporter/Private/Builders/GLTFBuilderUtility.cpp
--- a/Engine/Plugins/Enterprise/GLTFExporter/MigratedFromGit/Source/GLTFExporter/Private/Builders/GLTFBuilderUtility.cpp
+++ b/Engine/Plugins/Enterprise/GLTFExporter/MigratedFromGit/Source/GLTFExporter/Private/Builders/GLTFBuilderUtility.cpp
@@ -10,7 +10,7 @@ const TCHAR* FGLTFBuilderUtility::GetFileExtension(EGLTFJsonMimeType MimeType)
 		case EGLTFJsonMimeType::JPEG: return TEXT(".jpg");
 		default:
 			checkNoEntry();
-			return TEXT(".unknown");
+			return TEXT("");
 	}
 }
 
diff --git a/Engine/Plugins/Enterprise/GLTFExporter/MigratedFromGit/Source/GLTFExporter/Private/Json/GLTFJsonUtility.h b/Engine/Plugins/Enterprise/GLTFExporter/MigratedFromGit/Source/GLTFExporter/Private/Json/GLTFJsonUtility.h
--- a/Engine/Plugins/Enterprise/GLTFExporter/MigratedFromGit/Source/GLTFExporter/Private/Json/GLTFJsonUtility.h
+++ b/Engine/Plugins/Enterprise/GLTFExporter/MigratedFromGit/Source/GLTFExporter/Private/Json/GLTFJsonUtility.h
@@ -33,7 +33,7 @@ struct FGLTFJsonUtility
 			case EGLTFJsonExtension::EPIC_SkySpheres:         return TEXT("EPIC_sky_spheres");
 			default:
 				checkNoEntry();
-				return TEXT("unknown");
+				return TEXT("");
 		}
 	}
 
@@ -46,7 +46,7 @@ struct FGLTFJsonUtility
 			case EGLTFJsonAlphaMode::Mask:   return TEXT("MASK");
 			default:
 				checkNoEntry();
-				return TEXT("UNKNOWN");
+				return TEXT("");
 		}
 	}
 
@@ -58,7 +58,7 @@ struct FGLTFJsonUtility
 			case EGLTFJsonMimeType::JPEG: return TEXT("image/jpeg");
 			default:
 				checkNoEntry();
-				return TEXT("unknown");
+				return TEXT("");
 		}
 	}
 
@@ -75,7 +75,7 @@ struct FGLTFJsonUtility
 			case EGLTFJsonAccessorType::Mat4:   return TEXT("MAT4");
 			default:
 				checkNoEntry();
-				return TEXT("UNKNOWN");
+				return TEXT("");
 		}
 	}
 
@@ -87,7 +87,7 @@ struct FGLTFJsonUtility
 			case EGLTFJsonHDREncoding::RGBM: return TEXT("RGBM");
 			default:
 				checkNoEntry();
-				return TEXT("UNKNOWN");
+				return TEXT("");
 		}
 	}
 
@@ -103,7 +103,7 @@ struct FGLTFJsonUtility
 			case EGLTFJsonCubeFace::NegZ: return TEXT("NegZ");
 			default:
 				checkNoEntry();
-				return TEXT("Unknown");
+				return TEXT("");
 		}
 	}
 
@@ -115,7 +115,7 @@ struct FGLTFJsonUtility
 			case EGLTFJsonCameraType::Orthographic: return TEXT("orthographic");
 			default:
 				checkNoEntry();
-				return TEXT("unknown");
+				return TEXT("");
 		}
 	}
 
@@ -128,7 +128,7 @@ struct FGLTFJsonUtility
 			case EGLTFJsonLightType::Spot:        return TEXT("spot");
 			default:
 				checkNoEntry();
-				return TEXT("unknown");
+				return TEXT("");
 		}
 	}
 
@@ -141,7 +141,7 @@ struct FGLTFJsonUtility
 			case EGLTFJsonInterpolation::CubicSpline: return TEXT("CUBICSPLINE");
 			default:
 				checkNoEntry();
-				return TEXT("UNKNOWN");
+				return TEXT("");
 		}
 	}
 
@@ -155,7 +155,7 @@ struct FGLTFJsonUtility
 			case EGLTFJsonTargetPath::Weights:     return TEXT("weights");
 			default:
 				checkNoEntry();
-				return TEXT("unknown");
+				return TEXT("");
 		}
 	}
 
diff --git a/Engine/Plugins/Enterprise/GLTFExporter/MigratedFromGit/Source/GLTFExporterRuntime/Private/Actors/GLTFHotspotActor.cpp b/Engine/Plugins/Enterprise/GLTFExporter/MigratedFromGit/Source/GLTFExporterRuntime/Private/Actors/GLTFHotspotActor.cpp
--- a/Engine/Plugins/Enterprise/GLTFExporter/MigratedFromGit/Source/GLTFExporterRuntime/Private/Actors/GLTFHotspotActor.cpp
+++ b/Engine/Plugins/Enterprise/GLTFExporter/MigratedFromGit/Source/GLTFExporterRuntime/Private/Actors/GLTFHotspotActor.cpp
@@ -510,8 +510,6 @@ const UTexture2D* AGLTFHotspotActor::GetImageForState(EGLTFHotspotState State) c
 		case EGLTFHotspotState::ToggledHovered: return CurrentToggledHoveredImage;
 		default:
 			checkNoEntry();
-			break;
+			return nullptr;
 	}
-
-	return nullptr;
 }
