[GLTFExporter] feat: add warning message when alpha mode conversion failed (instead of silent fallback)

commit 5130be16f1a5bcb3fda735d192c0c552c4d9ba43
Date: Tue, 20 Oct 2020 23:19:50 +0200

    feat: add warning message when alpha mode conversion failed (instead of silent fallback)
    (cherry picked from commit 2faec5eb07c707f20c09face8fd474b3ca2b1eef​)

#skipci

diff --git a/Engine/Plugins/Enterprise/GLTFExporter/MigratedFromGit/Source/GLTFExporter/Private/Converters/GLTFConverterUtility.cpp b/Engine/Plugins/Enterprise/GLTFExporter/MigratedFromGit/Source/GLTFExporter/Private/Converters/GLTFConverterUtility.cpp
--- a/Engine/Plugins/Enterprise/GLTFExporter/MigratedFromGit/Source/GLTFExporter/Private/Converters/GLTFConverterUtility.cpp
+++ b/Engine/Plugins/Enterprise/GLTFExporter/MigratedFromGit/Source/GLTFExporter/Private/Converters/GLTFConverterUtility.cpp
@@ -52,7 +52,7 @@ EGLTFJsonAlphaMode FGLTFConverterUtility::ConvertBlendMode(EBlendMode Mode)
 		case BLEND_Opaque:      return EGLTFJsonAlphaMode::Opaque;
 		case BLEND_Translucent: return EGLTFJsonAlphaMode::Blend;
 		case BLEND_Masked:      return EGLTFJsonAlphaMode::Mask;
-		default:                return EGLTFJsonAlphaMode::Opaque; // fallback
+		default:                return EGLTFJsonAlphaMode::None;
 	}
 }
 
diff --git a/Engine/Plugins/Enterprise/GLTFExporter/MigratedFromGit/Source/GLTFExporter/Private/Json/GLTFJsonEnums.h b/Engine/Plugins/Enterprise/GLTFExporter/MigratedFromGit/Source/GLTFExporter/Private/Json/GLTFJsonEnums.h
--- a/Engine/Plugins/Enterprise/GLTFExporter/MigratedFromGit/Source/GLTFExporter/Private/Json/GLTFJsonEnums.h
+++ b/Engine/Plugins/Enterprise/GLTFExporter/MigratedFromGit/Source/GLTFExporter/Private/Json/GLTFJsonEnums.h
@@ -95,6 +95,7 @@ enum class EGLTFJsonPrimitiveMode
 
 enum class EGLTFJsonAlphaMode
 {
+	None = -1,
 	Opaque,
 	Mask,
 	Blend
diff --git a/Engine/Plugins/Enterprise/GLTFExporter/MigratedFromGit/Source/GLTFExporter/Private/Tasks/GLTFMaterialTasks.cpp b/Engine/Plugins/Enterprise/GLTFExporter/MigratedFromGit/Source/GLTFExporter/Private/Tasks/GLTFMaterialTasks.cpp
--- a/Engine/Plugins/Enterprise/GLTFExporter/MigratedFromGit/Source/GLTFExporter/Private/Tasks/GLTFMaterialTasks.cpp
+++ b/Engine/Plugins/Enterprise/GLTFExporter/MigratedFromGit/Source/GLTFExporter/Private/Tasks/GLTFMaterialTasks.cpp
@@ -63,8 +63,12 @@ void FGLTFMaterialTask::Complete()
 	FGLTFJsonMaterial& JsonMaterial = Builder.GetMaterial(MaterialIndex);
 	Material->GetName(JsonMaterial.Name);
 
-	// TODO: add support for additional blend modes (like Additive and Modulate)?
-	JsonMaterial.AlphaMode = FGLTFConverterUtility::ConvertBlendMode(Material->GetBlendMode());
+	if (!TryGetAlphaMode(JsonMaterial.AlphaMode))
+	{
+		JsonMaterial.AlphaMode = EGLTFJsonAlphaMode::Opaque;
+		Builder.AddWarningMessage(FString::Printf(TEXT("Material %s will be exported as blend mode %s"), *Material->GetName(), *FGLTFNameUtility::GetName(BLEND_Opaque)));
+	}
+
 	JsonMaterial.AlphaCutoff = Material->GetOpacityMaskClipValue();
 	JsonMaterial.DoubleSided = Material->IsTwoSided();
 
@@ -171,6 +175,24 @@ void FGLTFMaterialTask::Complete()
 	}
 }
 
+bool FGLTFMaterialTask::TryGetAlphaMode(EGLTFJsonAlphaMode& AlphaMode) const
+{
+	const EBlendMode BlendMode = Material->GetBlendMode();
+
+	// TODO: add support for additional blend modes (like Additive and Modulate)?
+
+	const EGLTFJsonAlphaMode ConvertedAlphaMode = FGLTFConverterUtility::ConvertBlendMode(BlendMode);
+	if (ConvertedAlphaMode == EGLTFJsonAlphaMode::None)
+	{
+		const FString BlendModeName = FGLTFNameUtility::GetName(BlendMode);
+		Builder.AddWarningMessage(FString::Printf(TEXT("Unsupported blend mode (%s) in material %s"), *BlendModeName, *Material->GetName()));
+		return false;
+	}
+
+	AlphaMode = ConvertedAlphaMode;
+	return true;
+}
+
 bool FGLTFMaterialTask::TryGetShadingModel(EGLTFJsonShadingModel& OutShadingModel) const
 {
 	const FMaterialShadingModelField ShadingModels = Material->GetShadingModels();
diff --git a/Engine/Plugins/Enterprise/GLTFExporter/MigratedFromGit/Source/GLTFExporter/Private/Tasks/GLTFMaterialTasks.h b/Engine/Plugins/Enterprise/GLTFExporter/MigratedFromGit/Source/GLTFExporter/Private/Tasks/GLTFMaterialTasks.h
--- a/Engine/Plugins/Enterprise/GLTFExporter/MigratedFromGit/Source/GLTFExporter/Private/Tasks/GLTFMaterialTasks.h
+++ b/Engine/Plugins/Enterprise/GLTFExporter/MigratedFromGit/Source/GLTFExporter/Private/Tasks/GLTFMaterialTasks.h
@@ -33,6 +33,7 @@ private:
 	const UMaterialInterface* Material;
 	const FGLTFJsonMaterialIndex MaterialIndex;
 
+	bool TryGetAlphaMode(EGLTFJsonAlphaMode& AlphaMode) const;
 	bool TryGetShadingModel(EGLTFJsonShadingModel& ShadingModel) const;
 
 	bool TryGetBaseColorAndOpacity(FGLTFJsonPBRMetallicRoughness& OutPBRParams, EMaterialProperty BaseColorProperty, EMaterialProperty OpacityProperty) const;
