[GLTFExporter] refactor: return EGLTFJsonInterpolation::None and EGLTFJsonTextureWrap::None if conversion fails

commit 2180d87ab6d59592d9d6903c6165bc77dd79ee4b
Date: Tue, 20 Oct 2020 23:22:06 +0200

    refactor: return EGLTFJsonInterpolation::None and EGLTFJsonTextureWrap::None if conversion fails
    (cherry picked from commit 6c451ace00f01802bdd15f47565228ca1d35f7ea​)

#skipci

diff --git a/Engine/Plugins/Enterprise/GLTFExporter/MigratedFromGit/Source/GLTFExporter/Private/Converters/GLTFConverterUtility.cpp b/Engine/Plugins/Enterprise/GLTFExporter/MigratedFromGit/Source/GLTFExporter/Private/Converters/GLTFConverterUtility.cpp
--- a/Engine/Plugins/Enterprise/GLTFExporter/MigratedFromGit/Source/GLTFExporter/Private/Converters/GLTFConverterUtility.cpp
+++ b/Engine/Plugins/Enterprise/GLTFExporter/MigratedFromGit/Source/GLTFExporter/Private/Converters/GLTFConverterUtility.cpp
@@ -30,7 +30,7 @@ EGLTFJsonInterpolation FGLTFConverterUtility::ConvertInterpolation(const EAnimIn
 	{
 		case EAnimInterpolationType::Linear: return EGLTFJsonInterpolation::Linear;
 		case EAnimInterpolationType::Step:   return EGLTFJsonInterpolation::Step;
-		default:                             return EGLTFJsonInterpolation::Linear; // fallback
+		default:                             return EGLTFJsonInterpolation::None;
 	}
 }
 
@@ -63,7 +63,7 @@ EGLTFJsonTextureWrap FGLTFConverterUtility::ConvertWrap(TextureAddress Address)
 		case TA_Wrap:   return EGLTFJsonTextureWrap::Repeat;
 		case TA_Mirror: return EGLTFJsonTextureWrap::MirroredRepeat;
 		case TA_Clamp:  return EGLTFJsonTextureWrap::ClampToEdge;
-		default:        return EGLTFJsonTextureWrap::Repeat; // fallback
+		default:        return EGLTFJsonTextureWrap::None; // TODO: add error handling in callers
 	}
 }
 
@@ -109,6 +109,6 @@ EGLTFJsonCubeFace FGLTFConverterUtility::ConvertCubeFace(ECubeFace CubeFace)
 		case CubeFace_NegY:	return EGLTFJsonCubeFace::NegZ;
 		case CubeFace_PosZ:	return EGLTFJsonCubeFace::PosY;
 		case CubeFace_NegZ:	return EGLTFJsonCubeFace::NegY;
-		default:            return EGLTFJsonCubeFace::None;
+		default:            return EGLTFJsonCubeFace::None; // TODO: add error handling in callers
 	}
 }
diff --git a/Engine/Plugins/Enterprise/GLTFExporter/MigratedFromGit/Source/GLTFExporter/Private/Json/GLTFJsonEnums.h b/Engine/Plugins/Enterprise/GLTFExporter/MigratedFromGit/Source/GLTFExporter/Private/Json/GLTFJsonEnums.h
--- a/Engine/Plugins/Enterprise/GLTFExporter/MigratedFromGit/Source/GLTFExporter/Private/Json/GLTFJsonEnums.h
+++ b/Engine/Plugins/Enterprise/GLTFExporter/MigratedFromGit/Source/GLTFExporter/Private/Json/GLTFJsonEnums.h
@@ -123,6 +123,7 @@ enum class EGLTFJsonTextureFilter
 
 enum class EGLTFJsonTextureWrap
 {
+	None = -1,
 	Repeat = 10497,
 	MirroredRepeat = 33648,
 	ClampToEdge = 33071
@@ -130,6 +131,7 @@ enum class EGLTFJsonTextureWrap
 
 enum class EGLTFJsonInterpolation
 {
+	None = -1,
 	Linear,
 	Step,
 	CubicSpline,
diff --git a/Engine/Plugins/Enterprise/GLTFExporter/MigratedFromGit/Source/GLTFExporter/Private/Tasks/GLTFAnimationTasks.cpp b/Engine/Plugins/Enterprise/GLTFExporter/MigratedFromGit/Source/GLTFExporter/Private/Tasks/GLTFAnimationTasks.cpp
--- a/Engine/Plugins/Enterprise/GLTFExporter/MigratedFromGit/Source/GLTFExporter/Private/Tasks/GLTFAnimationTasks.cpp
+++ b/Engine/Plugins/Enterprise/GLTFExporter/MigratedFromGit/Source/GLTFExporter/Private/Tasks/GLTFAnimationTasks.cpp
@@ -40,7 +40,13 @@ void FGLTFAnimSequenceTask::Complete()
 		FGLTFBoneUtility::InitializeToSkeleton(BoneContainer, Skeleton);
 	}
 
-	const EGLTFJsonInterpolation Interpolation = FGLTFConverterUtility::ConvertInterpolation(AnimSequence->Interpolation);
+	EGLTFJsonInterpolation Interpolation = FGLTFConverterUtility::ConvertInterpolation(AnimSequence->Interpolation);
+	if (Interpolation == EGLTFJsonInterpolation::None)
+	{
+		Interpolation = EGLTFJsonInterpolation::Linear;
+		// TODO: report warning (about unknown interpolation exported as linear)
+	}
+
 	const TArray<FName>& TrackNames = AnimSequence->GetAnimationTrackNames();
 	const int32 TrackCount = TrackNames.Num();
 
