[GLTFExporter] refactor: don't HDR-encode normalmap textures

commit f2ee885487a4a135a6527a3467a9d7a1bed85fe6
Date: Thu, 22 Oct 2020 16:09:38 +0200

    refactor: don't HDR-encode normalmap textures
    (cherry picked from commit 28b6123fbde50a81f0a50d5e8beedd2f427c7a55​)

#skipci

diff --git a/Engine/Plugins/Enterprise/GLTFExporter/MigratedFromGit/Source/GLTFExporter/Private/Converters/GLTFTextureUtility.cpp b/Engine/Plugins/Enterprise/GLTFExporter/MigratedFromGit/Source/GLTFExporter/Private/Converters/GLTFTextureUtility.cpp
--- a/Engine/Plugins/Enterprise/GLTFExporter/MigratedFromGit/Source/GLTFExporter/Private/Converters/GLTFTextureUtility.cpp
+++ b/Engine/Plugins/Enterprise/GLTFExporter/MigratedFromGit/Source/GLTFExporter/Private/Converters/GLTFTextureUtility.cpp
@@ -239,7 +239,7 @@ UTexture2D* FGLTFTextureUtility::CreateTextureFromCubeFace(const UTextureRenderT
 	return FaceTexture;
 }
 
-bool FGLTFTextureUtility::ReadEncodedPixels(const UTextureRenderTarget2D* InRenderTarget, TArray<FColor>& OutPixels, EGLTFJsonHDREncoding& OutEncoding)
+bool FGLTFTextureUtility::ReadPixels(const UTextureRenderTarget2D* InRenderTarget, TArray<FColor>& OutPixels)
 {
 	FTextureRenderTarget2DResource* Resource = static_cast<FTextureRenderTarget2DResource*>(InRenderTarget->Resource);
 	if (Resource == nullptr)
@@ -247,18 +247,34 @@ bool FGLTFTextureUtility::ReadEncodedPixels(const UTextureRenderTarget2D* InRend
 		return false;
 	}
 
-	if (IsHDRFormat(InRenderTarget->GetFormat()))
+	return Resource->ReadPixels(OutPixels);
+}
+
+bool FGLTFTextureUtility::ReadEncodedPixels(const UTextureRenderTarget2D* InRenderTarget, TArray<FColor>& OutPixels, EGLTFJsonHDREncoding Encoding)
+{
+	FTextureRenderTarget2DResource* Resource = static_cast<FTextureRenderTarget2DResource*>(InRenderTarget->Resource);
+	if (Resource == nullptr)
 	{
-		TArray<FLinearColor> HDRPixels;
-		Resource->ReadLinearColorPixels(HDRPixels);
+		return false;
+	}
 
-		EncodeRGBM(HDRPixels, OutPixels);
-		OutEncoding = EGLTFJsonHDREncoding::RGBM;
+	TArray<FLinearColor> HDRPixels;
+	if (!Resource->ReadLinearColorPixels(HDRPixels))
+	{
+		return false;
 	}
-	else
+
+	switch (Encoding)
 	{
-		Resource->ReadPixels(OutPixels);
-		OutEncoding = EGLTFJsonHDREncoding::None;
+		case EGLTFJsonHDREncoding::RGBM:
+			EncodeRGBM(HDRPixels, OutPixels);
+			break;
+
+		// TODO: add encoding for RGBE
+
+		default:
+			checkNoEntry();
+			break;
 	}
 
 	return true;
diff --git a/Engine/Plugins/Enterprise/GLTFExporter/MigratedFromGit/Source/GLTFExporter/Private/Converters/GLTFTextureUtility.h b/Engine/Plugins/Enterprise/GLTFExporter/MigratedFromGit/Source/GLTFExporter/Private/Converters/GLTFTextureUtility.h
--- a/Engine/Plugins/Enterprise/GLTFExporter/MigratedFromGit/Source/GLTFExporter/Private/Converters/GLTFTextureUtility.h
+++ b/Engine/Plugins/Enterprise/GLTFExporter/MigratedFromGit/Source/GLTFExporter/Private/Converters/GLTFTextureUtility.h
@@ -35,7 +35,8 @@ struct FGLTFTextureUtility
 	static UTexture2D* CreateTextureFromCubeFace(const UTextureCube* TextureCube, ECubeFace CubeFace);
 	static UTexture2D* CreateTextureFromCubeFace(const UTextureRenderTargetCube* RenderTargetCube, ECubeFace CubeFace);
 
-	static bool ReadEncodedPixels(const UTextureRenderTarget2D* InRenderTarget, TArray<FColor>& OutPixels, EGLTFJsonHDREncoding& OutEncoding);
+	static bool ReadPixels(const UTextureRenderTarget2D* InRenderTarget, TArray<FColor>& OutPixels);
+	static bool ReadEncodedPixels(const UTextureRenderTarget2D* InRenderTarget, TArray<FColor>& OutPixels, EGLTFJsonHDREncoding Encoding);
 
 	static FColor EncodeRGBM(const FLinearColor& Color, float MaxRange = 8);
 	static void EncodeRGBM(const TArray<FLinearColor>& InPixels, TArray<FColor>& OutPixels, float MaxRange = 8);
diff --git a/Engine/Plugins/Enterprise/GLTFExporter/MigratedFromGit/Source/GLTFExporter/Private/Tasks/GLTFTextureTasks.cpp b/Engine/Plugins/Enterprise/GLTFExporter/MigratedFromGit/Source/GLTFExporter/Private/Tasks/GLTFTextureTasks.cpp
--- a/Engine/Plugins/Enterprise/GLTFExporter/MigratedFromGit/Source/GLTFExporter/Private/Tasks/GLTFTextureTasks.cpp
+++ b/Engine/Plugins/Enterprise/GLTFExporter/MigratedFromGit/Source/GLTFExporter/Private/Tasks/GLTFTextureTasks.cpp
@@ -7,13 +7,14 @@
 
 void FGLTFTexture2DTask::Complete()
 {
+	bool bIsNormalMap = Texture2D->IsNormalMap();
 	FGLTFJsonImageIndex ImageIndex;
 
 	FGLTFJsonTexture& JsonTexture = Builder.GetTexture(TextureIndex);
 	Texture2D->GetName(JsonTexture.Name);
 
 	// NOTE: export of normalmaps via source data is disabled because we need to manipulate the pixels during conversion.
-	const bool bPreferSourceExport = Texture2D->IsNormalMap() ? false : Builder.ExportOptions->bExportSourceTextures;
+	const bool bPreferSourceExport = bIsNormalMap ? false : Builder.ExportOptions->bExportSourceTextures;
 
 	// NOTE: export of lightmaps via source data is used to work around issues
 	// with quality-loss due to incorrect gamma transformation when rendering to a canvas.
@@ -70,28 +71,33 @@ void FGLTFTexture2DTask::Complete()
 		FGLTFTextureUtility::DrawTexture(RenderTarget, Texture2D);
 
 		TArray<FColor> Pixels;
-		if (FGLTFTextureUtility::ReadEncodedPixels(RenderTarget, Pixels, JsonTexture.Encoding)) // TODO: use only encoding as specified by export options
+		if (!bIsNormalMap && FGLTFTextureUtility::IsHDRFormat(RenderTarget->GetFormat()))
 		{
-			if (Texture2D->IsNormalMap())
+			JsonTexture.Encoding = EGLTFJsonHDREncoding::RGBM;
+			FGLTFTextureUtility::ReadEncodedPixels(RenderTarget, Pixels, JsonTexture.Encoding); // TODO: use only encoding as specified by export options
+		}
+		else
+		{
+			FGLTFTextureUtility::ReadPixels(RenderTarget, Pixels);
+
+			if (bIsNormalMap)
 			{
-				if (JsonTexture.Encoding == EGLTFJsonHDREncoding::None)
-				{
-					// Convert normalmaps to use +Y (OpenGL / WebGL standard)
-					FGLTFTextureUtility::FlipGreenChannel(Pixels);
-				}
-				else
-				{
-					Builder.AddWarningMessage(FString::Printf(TEXT("Failed to invert green channel for normalmap %s"), *JsonTexture.Name));
-				}
+				FGLTFTextureUtility::FlipGreenChannel(Pixels);
 			}
+		}
 
-			ImageIndex = Builder.AddImage(Pixels.GetData(), Size, JsonTexture.Name);
+		if (Pixels.Num() == 0)
+		{
+			Builder.AddWarningMessage(FString::Printf(TEXT("Failed to read pixels for 2D texture %s"), *JsonTexture.Name));
+			return;
 		}
+
+		ImageIndex = Builder.AddImage(Pixels.GetData(), Size, JsonTexture.Name);
 	}
 
 	if (ImageIndex == INDEX_NONE)
 	{
-		Builder.AddWarningMessage(FString::Printf(TEXT("Failed to export texture %s"), *JsonTexture.Name));
+		Builder.AddWarningMessage(FString::Printf(TEXT("Failed to export 2D texture %s"), *JsonTexture.Name));
 		return;
 	}
 
@@ -121,9 +127,19 @@ void FGLTFTextureCubeTask::Complete()
 	FGLTFTextureUtility::RotateTexture(RenderTarget, FaceTexture, FaceRotation);
 
 	TArray<FColor> Pixels;
-	if (!FGLTFTextureUtility::ReadEncodedPixels(RenderTarget, Pixels, JsonTexture.Encoding)) // TODO: use only encoding as specified by export options
+	if (FGLTFTextureUtility::IsHDRFormat(RenderTarget->GetFormat()))
 	{
-		// TODO: report error
+		JsonTexture.Encoding = EGLTFJsonHDREncoding::RGBM;
+		FGLTFTextureUtility::ReadEncodedPixels(RenderTarget, Pixels, JsonTexture.Encoding); // TODO: use only encoding as specified by export options
+	}
+	else
+	{
+		FGLTFTextureUtility::ReadPixels(RenderTarget, Pixels);
+	}
+
+	if (Pixels.Num() == 0)
+	{
+		Builder.AddWarningMessage(FString::Printf(TEXT("Failed to read pixels for cubemap texture %s"), *JsonTexture.Name));
 		return;
 	}
 
@@ -137,9 +153,19 @@ void FGLTFTextureRenderTarget2DTask::Complete()
 	RenderTarget2D->GetName(JsonTexture.Name);
 
 	TArray<FColor> Pixels;
-	if (!FGLTFTextureUtility::ReadEncodedPixels(RenderTarget2D, Pixels, JsonTexture.Encoding)) // TODO: use only encoding as specified by export options
+	if (FGLTFTextureUtility::IsHDRFormat(RenderTarget2D->GetFormat()))
 	{
-		// TODO: report error
+		JsonTexture.Encoding = EGLTFJsonHDREncoding::RGBM;
+		FGLTFTextureUtility::ReadEncodedPixels(RenderTarget2D, Pixels, JsonTexture.Encoding); // TODO: use only encoding as specified by export options
+	}
+	else
+	{
+		FGLTFTextureUtility::ReadPixels(RenderTarget2D, Pixels);
+	}
+
+	if (Pixels.Num() == 0)
+	{
+		Builder.AddWarningMessage(FString::Printf(TEXT("Failed to read pixels for 2D render target %s"), *JsonTexture.Name));
 		return;
 	}
 
@@ -169,9 +195,19 @@ void FGLTFTextureRenderTargetCubeTask::Complete()
 	FGLTFTextureUtility::RotateTexture(RenderTarget, FaceTexture, FaceRotation);
 
 	TArray<FColor> Pixels;
-	if (!FGLTFTextureUtility::ReadEncodedPixels(RenderTarget, Pixels, JsonTexture.Encoding)) // TODO: use only encoding as specified by export options
+	if (FGLTFTextureUtility::IsHDRFormat(RenderTarget->GetFormat()))
 	{
-		// TODO: report error
+		JsonTexture.Encoding = EGLTFJsonHDREncoding::RGBM;
+		FGLTFTextureUtility::ReadEncodedPixels(RenderTarget, Pixels, JsonTexture.Encoding); // TODO: use only encoding as specified by export options
+	}
+	else
+	{
+		FGLTFTextureUtility::ReadPixels(RenderTarget, Pixels);
+	}
+
+	if (Pixels.Num() == 0)
+	{
+		Builder.AddWarningMessage(FString::Printf(TEXT("Failed to read pixels for cubemap render target %s"), *JsonTexture.Name));
 		return;
 	}
 
