[GLTFExporter] refactor(sky-sphere): use pointers instead of references for actors and materials

commit 392d9cfcf8375cd11d2ea33cb6a31165f997ea2e
Date: Thu, 22 Oct 2020 16:58:55 +0200

    refactor(sky-sphere): use pointers instead of references for actors and materials
    (cherry picked from commit 7bbd7112ebc484afcf41b5eae3729d6c3026cf27​)

#skipci

diff --git a/Engine/Plugins/Enterprise/GLTFExporter/MigratedFromGit/Source/GLTFExporter/Private/Converters/GLTFSkySphereConverters.cpp b/Engine/Plugins/Enterprise/GLTFExporter/MigratedFromGit/Source/GLTFExporter/Private/Converters/GLTFSkySphereConverters.cpp
--- a/Engine/Plugins/Enterprise/GLTFExporter/MigratedFromGit/Source/GLTFExporter/Private/Converters/GLTFSkySphereConverters.cpp
+++ b/Engine/Plugins/Enterprise/GLTFExporter/MigratedFromGit/Source/GLTFExporter/Private/Converters/GLTFSkySphereConverters.cpp
@@ -54,13 +54,13 @@ FGLTFJsonSkySphereIndex FGLTFSkySphereConverter::Convert(const AActor* SkySphere
 
 	if (SkyMaterial != nullptr)
 	{
-		ConvertScalarParameter(*SkySphereActor, *SkyMaterial, TEXT("Sun Radius"), JsonSkySphere.SunRadius);
-		ConvertScalarParameter(*SkySphereActor, *SkyMaterial, TEXT("NoisePower1"), JsonSkySphere.NoisePower1);
-		ConvertScalarParameter(*SkySphereActor, *SkyMaterial, TEXT("NoisePower2"), JsonSkySphere.NoisePower2);
+		ConvertScalarParameter(SkySphereActor, SkyMaterial, TEXT("Sun Radius"), JsonSkySphere.SunRadius);
+		ConvertScalarParameter(SkySphereActor, SkyMaterial, TEXT("NoisePower1"), JsonSkySphere.NoisePower1);
+		ConvertScalarParameter(SkySphereActor, SkyMaterial, TEXT("NoisePower2"), JsonSkySphere.NoisePower2);
 
-		ConvertTextureParameter(*SkySphereActor, *SkyMaterial, ESkySphereTextureParameter::SkyTexture, JsonSkySphere.SkyTexture);
-		ConvertTextureParameter(*SkySphereActor, *SkyMaterial, ESkySphereTextureParameter::CloudsTexture, JsonSkySphere.CloudsTexture);
-		ConvertTextureParameter(*SkySphereActor, *SkyMaterial, ESkySphereTextureParameter::StarsTexture, JsonSkySphere.StarsTexture);
+		ConvertTextureParameter(SkySphereActor, SkyMaterial, ESkySphereTextureParameter::SkyTexture, JsonSkySphere.SkyTexture);
+		ConvertTextureParameter(SkySphereActor, SkyMaterial, ESkySphereTextureParameter::CloudsTexture, JsonSkySphere.CloudsTexture);
+		ConvertTextureParameter(SkySphereActor, SkyMaterial, ESkySphereTextureParameter::StarsTexture, JsonSkySphere.StarsTexture);
 	}
 	else
 	{
@@ -69,42 +69,48 @@ FGLTFJsonSkySphereIndex FGLTFSkySphereConverter::Convert(const AActor* SkySphere
 			*SkySphereActor->GetName()));
 	}
 
-	ConvertProperty(*SkySphereActor, TEXT("Sun height"), JsonSkySphere.SunHeight);
-	ConvertProperty(*SkySphereActor, TEXT("Sun brightness"), JsonSkySphere.SunBrightness);
-	ConvertProperty(*SkySphereActor, TEXT("Stars brightness"), JsonSkySphere.StarsBrightness);
-	ConvertProperty(*SkySphereActor, TEXT("Cloud speed"), JsonSkySphere.CloudSpeed);
-	ConvertProperty(*SkySphereActor, TEXT("Cloud opacity"), JsonSkySphere.CloudOpacity);
-	ConvertProperty(*SkySphereActor, TEXT("Horizon Falloff"), JsonSkySphere.HorizonFalloff);
-	ConvertProperty(*SkySphereActor, TEXT("Colors determined by sun position"), JsonSkySphere.bColorsDeterminedBySunPosition);
+	ConvertProperty(SkySphereActor, TEXT("Sun height"), JsonSkySphere.SunHeight);
+	ConvertProperty(SkySphereActor, TEXT("Sun brightness"), JsonSkySphere.SunBrightness);
+	ConvertProperty(SkySphereActor, TEXT("Stars brightness"), JsonSkySphere.StarsBrightness);
+	ConvertProperty(SkySphereActor, TEXT("Cloud speed"), JsonSkySphere.CloudSpeed);
+	ConvertProperty(SkySphereActor, TEXT("Cloud opacity"), JsonSkySphere.CloudOpacity);
+	ConvertProperty(SkySphereActor, TEXT("Horizon Falloff"), JsonSkySphere.HorizonFalloff);
+	ConvertProperty(SkySphereActor, TEXT("Colors determined by sun position"), JsonSkySphere.bColorsDeterminedBySunPosition);
 
-	ConvertColorProperty(*SkySphereActor, TEXT("Zenith Color"), JsonSkySphere.ZenithColor);
-	ConvertColorProperty(*SkySphereActor, TEXT("Horizon color"), JsonSkySphere.HorizonColor);
-	ConvertColorProperty(*SkySphereActor, TEXT("Cloud color"), JsonSkySphere.CloudColor);
-	ConvertColorProperty(*SkySphereActor, TEXT("Overall Color"), JsonSkySphere.OverallColor);
+	ConvertColorProperty(SkySphereActor, TEXT("Zenith Color"), JsonSkySphere.ZenithColor);
+	ConvertColorProperty(SkySphereActor, TEXT("Horizon color"), JsonSkySphere.HorizonColor);
+	ConvertColorProperty(SkySphereActor, TEXT("Cloud color"), JsonSkySphere.CloudColor);
+	ConvertColorProperty(SkySphereActor, TEXT("Overall Color"), JsonSkySphere.OverallColor);
 
-	ConvertColorCurveProperty(*SkySphereActor, TEXT("Zenith color curve"), JsonSkySphere.ZenithColorCurve);
-	ConvertColorCurveProperty(*SkySphereActor, TEXT("Horizon color curve"), JsonSkySphere.HorizonColorCurve);
-	ConvertColorCurveProperty(*SkySphereActor, TEXT("Cloud color curve"), JsonSkySphere.CloudColorCurve);
+	ConvertColorCurveProperty(SkySphereActor, TEXT("Zenith color curve"), JsonSkySphere.ZenithColorCurve);
+	ConvertColorCurveProperty(SkySphereActor, TEXT("Horizon color curve"), JsonSkySphere.HorizonColorCurve);
+	ConvertColorCurveProperty(SkySphereActor, TEXT("Cloud color curve"), JsonSkySphere.CloudColorCurve);
 
 	return Builder.AddSkySphere(JsonSkySphere);
 }
 
 template <class ValueType>
-void FGLTFSkySphereConverter::ConvertProperty(const AActor& Actor, const TCHAR* PropertyName, ValueType& OutValue) const
+void FGLTFSkySphereConverter::ConvertProperty(const AActor* Actor, const TCHAR* PropertyName, ValueType& OutValue) const
 {
-	if (!FGLTFActorUtility::TryGetPropertyValue(&Actor, PropertyName, OutValue))
+	// Sanity check
+	check(Actor != nullptr);
+
+	if (!FGLTFActorUtility::TryGetPropertyValue(Actor, PropertyName, OutValue))
 	{
 		Builder.AddWarningMessage(FString::Printf(
 			TEXT("Failed to export property %s for Sky Sphere %s"),
 			PropertyName,
-			*Actor.GetName()));
+			*Actor->GetName()));
 	}
 }
 
-void FGLTFSkySphereConverter::ConvertColorProperty(const AActor& Actor, const TCHAR* PropertyName, FGLTFJsonColor4& OutValue) const
+void FGLTFSkySphereConverter::ConvertColorProperty(const AActor* Actor, const TCHAR* PropertyName, FGLTFJsonColor4& OutValue) const
 {
+	// Sanity check
+	check(Actor != nullptr);
+
 	FLinearColor LinearColor;
-	if (FGLTFActorUtility::TryGetPropertyValue(&Actor, PropertyName, LinearColor))
+	if (FGLTFActorUtility::TryGetPropertyValue(Actor, PropertyName, LinearColor))
 	{
 		OutValue = FGLTFConverterUtility::ConvertColor(LinearColor);
 	}
@@ -113,14 +119,17 @@ void FGLTFSkySphereConverter::ConvertColorProperty(const AActor& Actor, const TC
 		Builder.AddWarningMessage(FString::Printf(
 			TEXT("Failed to export property %s for Sky Sphere %s"),
 			PropertyName,
-			*Actor.GetName()));
+			*Actor->GetName()));
 	}
 }
 
-void FGLTFSkySphereConverter::ConvertColorCurveProperty(const AActor& Actor, const TCHAR* PropertyName, FGLTFJsonSkySphereColorCurve& OutValue) const
+void FGLTFSkySphereConverter::ConvertColorCurveProperty(const AActor* Actor, const TCHAR* PropertyName, FGLTFJsonSkySphereColorCurve& OutValue) const
 {
+	// Sanity check
+	check(Actor != nullptr);
+
 	const UCurveLinearColor* ColorCurve = nullptr;
-	FGLTFActorUtility::TryGetPropertyValue(&Actor, PropertyName, ColorCurve);
+	FGLTFActorUtility::TryGetPropertyValue(Actor, PropertyName, ColorCurve);
 
 	if (ColorCurve != nullptr)
 	{
@@ -129,7 +138,7 @@ void FGLTFSkySphereConverter::ConvertColorCurveProperty(const AActor& Actor, con
 			Builder.AddWarningMessage(FString::Printf(
 				TEXT("Adjustments for property %s in Sky Sphere %s are not supported"),
 				PropertyName,
-				*Actor.GetName()));
+				*Actor->GetName()));
 		}
 
 		OutValue.ComponentCurves.AddDefaulted(3);
@@ -155,24 +164,30 @@ void FGLTFSkySphereConverter::ConvertColorCurveProperty(const AActor& Actor, con
 		Builder.AddWarningMessage(FString::Printf(
 			TEXT("Failed to export property %s for Sky Sphere %s"),
 			PropertyName,
-			*Actor.GetName()));
+			*Actor->GetName()));
 	}
 }
 
-void FGLTFSkySphereConverter::ConvertScalarParameter(const AActor& Actor, const UMaterialInstance& Material, const TCHAR* ParameterName, float& OutValue) const
+void FGLTFSkySphereConverter::ConvertScalarParameter(const AActor* Actor, const UMaterialInstance* Material, const TCHAR* ParameterName, float& OutValue) const
 {
-	if (!Material.GetScalarParameterValue(ParameterName, OutValue))
+	// Sanity check
+	check(Actor != nullptr && Material != nullptr);
+
+	if (!Material->GetScalarParameterValue(ParameterName, OutValue))
 	{
 		Builder.AddWarningMessage(FString::Printf(
 			TEXT("Failed to export parameter %s (in material %s) for Sky Sphere %s"),
 			ParameterName,
-			*Material.GetName(),
-			*Actor.GetName()));
+			*Material->GetName(),
+			*Actor->GetName()));
 	}
 }
 
-void FGLTFSkySphereConverter::ConvertTextureParameter(const AActor& Actor, const UMaterialInstance& Material, const ESkySphereTextureParameter Parameter, FGLTFJsonTextureIndex& OutValue) const
+void FGLTFSkySphereConverter::ConvertTextureParameter(const AActor* Actor, const UMaterialInstance* Material, const ESkySphereTextureParameter Parameter, FGLTFJsonTextureIndex& OutValue) const
 {
+	// Sanity check
+	check(Actor != nullptr && Material != nullptr);
+
 	const TCHAR* TexturePath = nullptr;
 
 	// TODO: the sky sphere material doesn't expose the relevant textures as parameters, so we use hard-coded asset paths.
@@ -195,7 +210,7 @@ void FGLTFSkySphereConverter::ConvertTextureParameter(const AActor& Actor, const
 		Builder.AddWarningMessage(FString::Printf(
 			TEXT("Failed to export texture %s (used in material %s) for Sky Sphere %s"),
 			TexturePath,
-			*Material.GetName(),
-			*Actor.GetName()));
+			*Material->GetName(),
+			*Actor->GetName()));
 	}
 }
diff --git a/Engine/Plugins/Enterprise/GLTFExporter/MigratedFromGit/Source/GLTFExporter/Private/Converters/GLTFSkySphereConverters.h b/Engine/Plugins/Enterprise/GLTFExporter/MigratedFromGit/Source/GLTFExporter/Private/Converters/GLTFSkySphereConverters.h
--- a/Engine/Plugins/Enterprise/GLTFExporter/MigratedFromGit/Source/GLTFExporter/Private/Converters/GLTFSkySphereConverters.h
+++ b/Engine/Plugins/Enterprise/GLTFExporter/MigratedFromGit/Source/GLTFExporter/Private/Converters/GLTFSkySphereConverters.h
@@ -24,11 +24,11 @@ class FGLTFSkySphereConverter final : public FGLTFBuilderContext, public TGLTFCo
 	virtual FGLTFJsonSkySphereIndex Convert(const AActor* SkySphereActor) override;
 
 	template <class ValueType>
-	void ConvertProperty(const AActor& Actor, const TCHAR* PropertyName, ValueType& OutValue) const;
+	void ConvertProperty(const AActor* Actor, const TCHAR* PropertyName, ValueType& OutValue) const;
 
-	void ConvertColorProperty(const AActor& Actor, const TCHAR* PropertyName, FGLTFJsonColor4& OutValue) const;
-	void ConvertColorCurveProperty(const AActor& Actor, const TCHAR* PropertyName, FGLTFJsonSkySphereColorCurve& OutValue) const;
-	void ConvertScalarParameter(const AActor& Actor, const UMaterialInstance& Material, const TCHAR* ParameterName, float& OutValue) const;
+	void ConvertColorProperty(const AActor* Actor, const TCHAR* PropertyName, FGLTFJsonColor4& OutValue) const;
+	void ConvertColorCurveProperty(const AActor* Actor, const TCHAR* PropertyName, FGLTFJsonSkySphereColorCurve& OutValue) const;
+	void ConvertScalarParameter(const AActor* Actor, const UMaterialInstance* Material, const TCHAR* ParameterName, float& OutValue) const;
 
-	void ConvertTextureParameter(const AActor& Actor, const UMaterialInstance& Material, const ESkySphereTextureParameter Parameter, FGLTFJsonTextureIndex& OutValue) const;
+	void ConvertTextureParameter(const AActor* Actor, const UMaterialInstance* Material, const ESkySphereTextureParameter Parameter, FGLTFJsonTextureIndex& OutValue) const;
 };
