[GLTFExporter] refactor: remove self-explanatory comments

commit 39a2052c986fa60cd8d063e77ab2e7c1baaf8aea
Date: Thu, 22 Oct 2020 17:08:16 +0200

    refactor: remove self-explanatory comments
    (cherry picked from commit 69c7674a78ba87403e06dbbad83b8c5f2390c7c8​)

#skipci

diff --git a/Engine/Plugins/Enterprise/GLTFExporter/MigratedFromGit/Source/GLTFExporter/Private/Converters/GLTFSkySphereConverters.cpp b/Engine/Plugins/Enterprise/GLTFExporter/MigratedFromGit/Source/GLTFExporter/Private/Converters/GLTFSkySphereConverters.cpp
--- a/Engine/Plugins/Enterprise/GLTFExporter/MigratedFromGit/Source/GLTFExporter/Private/Converters/GLTFSkySphereConverters.cpp
+++ b/Engine/Plugins/Enterprise/GLTFExporter/MigratedFromGit/Source/GLTFExporter/Private/Converters/GLTFSkySphereConverters.cpp
@@ -92,7 +92,6 @@ FGLTFJsonSkySphereIndex FGLTFSkySphereConverter::Convert(const AActor* SkySphere
 template <class ValueType>
 void FGLTFSkySphereConverter::ConvertProperty(const AActor* Actor, const TCHAR* PropertyName, ValueType& OutValue) const
 {
-	// Sanity check
 	check(Actor != nullptr);
 
 	if (!FGLTFActorUtility::TryGetPropertyValue(Actor, PropertyName, OutValue))
@@ -106,7 +105,6 @@ void FGLTFSkySphereConverter::ConvertProperty(const AActor* Actor, const TCHAR*
 
 void FGLTFSkySphereConverter::ConvertColorProperty(const AActor* Actor, const TCHAR* PropertyName, FGLTFJsonColor4& OutValue) const
 {
-	// Sanity check
 	check(Actor != nullptr);
 
 	FLinearColor LinearColor;
@@ -125,7 +123,6 @@ void FGLTFSkySphereConverter::ConvertColorProperty(const AActor* Actor, const TC
 
 void FGLTFSkySphereConverter::ConvertColorCurveProperty(const AActor* Actor, const TCHAR* PropertyName, FGLTFJsonSkySphereColorCurve& OutValue) const
 {
-	// Sanity check
 	check(Actor != nullptr);
 
 	const UCurveLinearColor* ColorCurve = nullptr;
@@ -170,7 +167,6 @@ void FGLTFSkySphereConverter::ConvertColorCurveProperty(const AActor* Actor, con
 
 void FGLTFSkySphereConverter::ConvertScalarParameter(const AActor* Actor, const UMaterialInstance* Material, const TCHAR* ParameterName, float& OutValue) const
 {
-	// Sanity check
 	check(Actor != nullptr && Material != nullptr);
 
 	if (!Material->GetScalarParameterValue(ParameterName, OutValue))
@@ -185,7 +181,6 @@ void FGLTFSkySphereConverter::ConvertScalarParameter(const AActor* Actor, const
 
 void FGLTFSkySphereConverter::ConvertTextureParameter(const AActor* Actor, const UMaterialInstance* Material, const ESkySphereTextureParameter Parameter, FGLTFJsonTextureIndex& OutValue) const
 {
-	// Sanity check
 	check(Actor != nullptr && Material != nullptr);
 
 	const TCHAR* TexturePath = nullptr;
