[GLTFExporter] feat: add base for baking using mesh-data

commit bf7e2dfd764fc5c689355ae5abd8f6cd7ae57ce3
Date: Tue, 20 Oct 2020 10:49:43 +0200

    feat: add base for baking using mesh-data
    (cherry picked from commit 151ead2d708cd826dee582030f5faf8481e43100​)

#skipci

diff --git a/Engine/Plugins/Enterprise/GLTFExporter/MigratedFromGit/Source/GLTFExporter/Private/Builders/GLTFConvertBuilder.cpp b/Engine/Plugins/Enterprise/GLTFExporter/MigratedFromGit/Source/GLTFExporter/Private/Builders/GLTFConvertBuilder.cpp
--- a/Engine/Plugins/Enterprise/GLTFExporter/MigratedFromGit/Source/GLTFExporter/Private/Builders/GLTFConvertBuilder.cpp
+++ b/Engine/Plugins/Enterprise/GLTFExporter/MigratedFromGit/Source/GLTFExporter/Private/Builders/GLTFConvertBuilder.cpp
@@ -139,14 +139,14 @@ FGLTFJsonMeshIndex FGLTFConvertBuilder::GetOrAddMesh(const USkeletalMeshComponen
 	return GetOrAddMesh(SkeletalMesh, LODIndex, OverrideVertexColors, OverrideSkinWeights, OverrideMaterials);
 }
 
-FGLTFJsonMaterialIndex FGLTFConvertBuilder::GetOrAddMaterial(const UMaterialInterface* Material)
+FGLTFJsonMaterialIndex FGLTFConvertBuilder::GetOrAddMaterial(const UMaterialInterface* Material, const UStaticMesh* Mesh, int32 LODIndex)
 {
 	if (Material == nullptr)
 	{
 		return FGLTFJsonMaterialIndex(INDEX_NONE);
 	}
 
-	return MaterialConverter.GetOrAdd(Material);
+	return MaterialConverter.GetOrAdd(Material, Mesh, LODIndex);
 }
 
 FGLTFJsonSamplerIndex FGLTFConvertBuilder::GetOrAddSampler(const UTexture* Texture)
diff --git a/Engine/Plugins/Enterprise/GLTFExporter/MigratedFromGit/Source/GLTFExporter/Private/Builders/GLTFConvertBuilder.h b/Engine/Plugins/Enterprise/GLTFExporter/MigratedFromGit/Source/GLTFExporter/Private/Builders/GLTFConvertBuilder.h
--- a/Engine/Plugins/Enterprise/GLTFExporter/MigratedFromGit/Source/GLTFExporter/Private/Builders/GLTFConvertBuilder.h
+++ b/Engine/Plugins/Enterprise/GLTFExporter/MigratedFromGit/Source/GLTFExporter/Private/Builders/GLTFConvertBuilder.h
@@ -44,7 +44,7 @@ public:
 	FGLTFJsonMeshIndex GetOrAddMesh(const USkeletalMesh* SkeletalMesh, int32 LODIndex = -1, const FColorVertexBuffer* OverrideVertexColors = nullptr, const FSkinWeightVertexBuffer* OverrideSkinWeights = nullptr, const FGLTFMaterialArray& OverrideMaterials = {});
 	FGLTFJsonMeshIndex GetOrAddMesh(const USkeletalMeshComponent* SkeletalMeshComponent);
 
-	FGLTFJsonMaterialIndex GetOrAddMaterial(const UMaterialInterface* Material);
+	FGLTFJsonMaterialIndex GetOrAddMaterial(const UMaterialInterface* Material, const UStaticMesh* Mesh, int32 LODIndex = -1);
 	FGLTFJsonSamplerIndex GetOrAddSampler(const UTexture* Texture);
 	FGLTFJsonTextureIndex GetOrAddTexture(const UTexture2D* Texture);
 	FGLTFJsonTextureIndex GetOrAddTexture(const UTextureCube* Texture, ECubeFace CubeFace);
diff --git a/Engine/Plugins/Enterprise/GLTFExporter/MigratedFromGit/Source/GLTFExporter/Private/Converters/GLTFMaterialConverters.cpp b/Engine/Plugins/Enterprise/GLTFExporter/MigratedFromGit/Source/GLTFExporter/Private/Converters/GLTFMaterialConverters.cpp
--- a/Engine/Plugins/Enterprise/GLTFExporter/MigratedFromGit/Source/GLTFExporter/Private/Converters/GLTFMaterialConverters.cpp
+++ b/Engine/Plugins/Enterprise/GLTFExporter/MigratedFromGit/Source/GLTFExporter/Private/Converters/GLTFMaterialConverters.cpp
@@ -2,10 +2,53 @@
 
 #include "Converters/GLTFMaterialConverters.h"
 #include "Converters/GLTFMaterialUtility.h"
+#include "Converters/GLTFMeshUtility.h"
 #include "Builders/GLTFConvertBuilder.h"
 #include "Tasks/GLTFMaterialTasks.h"
 
-FGLTFJsonMaterialIndex FGLTFMaterialConverter::Convert(const UMaterialInterface* Material)
+void FGLTFMaterialConverter::Sanitize(const UMaterialInterface*& Material, const UStaticMesh*& Mesh, int32& LODIndex)
+{
+	if (Mesh != nullptr)
+	{
+		if (LODIndex < 0)
+		{
+			LODIndex = FMath::Max(Builder.ExportOptions->DefaultLevelOfDetail, FGLTFMeshUtility::GetMinimumLOD(Mesh));
+		}
+
+		const TArray<EMaterialProperty> Properties =
+		{
+			MP_BaseColor,
+			MP_EmissiveColor,
+			MP_Opacity,
+			MP_OpacityMask,
+			MP_Metallic,
+			MP_Roughness,
+			MP_Normal,
+			// MP_CustomOutput,	// NOTE: causes a crash when used with AnalyzeMaterialProperty
+			MP_AmbientOcclusion,
+			MP_CustomData0,
+			MP_CustomData1
+		};
+
+		bool bRequiresVertexData = false;
+
+		int32 NumTextureCoordinates;
+		bool bPropertyRequiresVertexData;
+
+		for (const EMaterialProperty Property: Properties)
+		{
+			const_cast<UMaterialInterface*>(Material)->AnalyzeMaterialProperty(Property, NumTextureCoordinates, bPropertyRequiresVertexData);
+			bRequiresVertexData |= bPropertyRequiresVertexData;
+		}
+
+		if (!bRequiresVertexData)
+		{
+			Mesh = nullptr;
+		}
+	}
+}
+
+FGLTFJsonMaterialIndex FGLTFMaterialConverter::Convert(const UMaterialInterface* Material, const UStaticMesh* Mesh, int32 LODIndex)
 {
 	if (Material == FGLTFMaterialUtility::GetDefault())
 	{
@@ -13,6 +56,6 @@ FGLTFJsonMaterialIndex FGLTFMaterialConverter::Convert(const UMaterialInterface*
 	}
 
 	const FGLTFJsonMaterialIndex MaterialIndex = Builder.AddMaterial();
-	Builder.SetupTask<FGLTFMaterialTask>(Builder, Material, MaterialIndex);
+	Builder.SetupTask<FGLTFMaterialTask>(Builder, Material, Mesh, LODIndex, MaterialIndex);
 	return MaterialIndex;
 }
diff --git a/Engine/Plugins/Enterprise/GLTFExporter/MigratedFromGit/Source/GLTFExporter/Private/Converters/GLTFMaterialConverters.h b/Engine/Plugins/Enterprise/GLTFExporter/MigratedFromGit/Source/GLTFExporter/Private/Converters/GLTFMaterialConverters.h
--- a/Engine/Plugins/Enterprise/GLTFExporter/MigratedFromGit/Source/GLTFExporter/Private/Converters/GLTFMaterialConverters.h
+++ b/Engine/Plugins/Enterprise/GLTFExporter/MigratedFromGit/Source/GLTFExporter/Private/Converters/GLTFMaterialConverters.h
@@ -7,9 +7,11 @@
 #include "Converters/GLTFBuilderContext.h"
 #include "Engine.h"
 
-class FGLTFMaterialConverter final : public FGLTFBuilderContext, public TGLTFConverter<FGLTFJsonMaterialIndex, const UMaterialInterface*>
+class FGLTFMaterialConverter final : public FGLTFBuilderContext, public TGLTFConverter<FGLTFJsonMaterialIndex, const UMaterialInterface*, const UStaticMesh*, int32>
 {
 	using FGLTFBuilderContext::FGLTFBuilderContext;
 
-	virtual FGLTFJsonMaterialIndex Convert(const UMaterialInterface* Material) override;
+	virtual void Sanitize(const UMaterialInterface*& Material, const UStaticMesh*& Mesh, int32& LODIndex) override;
+
+	virtual FGLTFJsonMaterialIndex Convert(const UMaterialInterface* Material, const UStaticMesh* Mesh, int32 LODIndex) override;
 };
diff --git a/Engine/Plugins/Enterprise/GLTFExporter/MigratedFromGit/Source/GLTFExporter/Private/Converters/GLTFMaterialUtility.cpp b/Engine/Plugins/Enterprise/GLTFExporter/MigratedFromGit/Source/GLTFExporter/Private/Converters/GLTFMaterialUtility.cpp
--- a/Engine/Plugins/Enterprise/GLTFExporter/MigratedFromGit/Source/GLTFExporter/Private/Converters/GLTFMaterialUtility.cpp
+++ b/Engine/Plugins/Enterprise/GLTFExporter/MigratedFromGit/Source/GLTFExporter/Private/Converters/GLTFMaterialUtility.cpp
@@ -240,7 +240,7 @@ bool FGLTFMaterialUtility::CombineTextures(TArray<FColor>& OutPixels, const TArr
 	return bReadSuccessful;
 }
 
-FGLTFPropertyBakeOutput FGLTFMaterialUtility::BakeMaterialProperty(const FIntPoint& OutputSize, EMaterialProperty Property, const UMaterialInterface* Material, int32 TexCoord, bool bCopyAlphaFromRedChannel)
+FGLTFPropertyBakeOutput FGLTFMaterialUtility::BakeMaterialProperty(const FIntPoint& OutputSize, EMaterialProperty Property, const UMaterialInterface* Material, int32 TexCoord, UStaticMesh* Mesh, int32 LODIndex, bool bCopyAlphaFromRedChannel)
 {
 	EMaterialProperty HackProperty = MP_MAX;
 	TTuple<int32, UMaterialExpression*> HackPrevState;
@@ -275,6 +275,19 @@ FGLTFPropertyBakeOutput FGLTFMaterialUtility::BakeMaterialProperty(const FIntPoi
 	MeshSet.TextureCoordinateBox = { { 0.0f, 0.0f }, { 1.0f, 1.0f } };
 	MeshSet.TextureCoordinateIndex = TexCoord;
 
+	if (Mesh != nullptr)
+	{
+		int32 NumTextureCoordinates;
+		bool bRequiresVertexData;
+		const_cast<UMaterialInterface*>(Material)->AnalyzeMaterialProperty(Property, NumTextureCoordinates, bRequiresVertexData);
+
+		if (bRequiresVertexData)
+		{
+			MeshSet.Mesh = Mesh;
+			MeshSet.RawMeshDescription =  Mesh->GetMeshDescription(LODIndex);
+		}
+	}
+
 	// TODO: Do we need to fill in any more info in MeshSet?
 
 	MeshSettings.Add(&MeshSet);
diff --git a/Engine/Plugins/Enterprise/GLTFExporter/MigratedFromGit/Source/GLTFExporter/Private/Converters/GLTFMaterialUtility.h b/Engine/Plugins/Enterprise/GLTFExporter/MigratedFromGit/Source/GLTFExporter/Private/Converters/GLTFMaterialUtility.h
--- a/Engine/Plugins/Enterprise/GLTFExporter/MigratedFromGit/Source/GLTFExporter/Private/Converters/GLTFMaterialUtility.h
+++ b/Engine/Plugins/Enterprise/GLTFExporter/MigratedFromGit/Source/GLTFExporter/Private/Converters/GLTFMaterialUtility.h
@@ -53,7 +53,7 @@ struct FGLTFMaterialUtility
 
 	static UTexture2D* CreateTransientTexture(const FGLTFPropertyBakeOutput& PropertyBakeOutput, bool bUseSRGB = false);
 	static bool CombineTextures(TArray<FColor>& OutPixels, const TArray<FGLTFTextureCombineSource>& Sources, const FIntPoint& OutputSize, EPixelFormat OutputPixelFormat);
-	static FGLTFPropertyBakeOutput BakeMaterialProperty(const FIntPoint& OutputSize, EMaterialProperty Property, const UMaterialInterface* Material, int32 TexCoord, bool bCopyAlphaFromRedChannel = false);
+	static FGLTFPropertyBakeOutput BakeMaterialProperty(const FIntPoint& OutputSize, EMaterialProperty Property, const UMaterialInterface* Material, int32 TexCoord, UStaticMesh* Mesh, int32 LODIndex, bool bCopyAlphaFromRedChannel);
 
 	static FGLTFJsonTextureIndex AddCombinedTexture(FGLTFConvertBuilder& Builder, const TArray<FGLTFTextureCombineSource>& CombineSources, const FIntPoint& TextureSize, const FString& TextureName, EGLTFJsonTextureFilter MinFilter, EGLTFJsonTextureFilter MagFilter, EGLTFJsonTextureWrap WrapS, EGLTFJsonTextureWrap WrapT);
 	static FGLTFJsonTextureIndex AddTexture(FGLTFConvertBuilder& Builder, const TArray<FColor>& Pixels, const FIntPoint& TextureSize, const FString& TextureName, EGLTFJsonTextureFilter MinFilter, EGLTFJsonTextureFilter MagFilter, EGLTFJsonTextureWrap WrapS, EGLTFJsonTextureWrap WrapT);
diff --git a/Engine/Plugins/Enterprise/GLTFExporter/MigratedFromGit/Source/GLTFExporter/Private/Converters/GLTFVarationConverters.cpp b/Engine/Plugins/Enterprise/GLTFExporter/MigratedFromGit/Source/GLTFExporter/Private/Converters/GLTFVarationConverters.cpp
--- a/Engine/Plugins/Enterprise/GLTFExporter/MigratedFromGit/Source/GLTFExporter/Private/Converters/GLTFVarationConverters.cpp
+++ b/Engine/Plugins/Enterprise/GLTFExporter/MigratedFromGit/Source/GLTFExporter/Private/Converters/GLTFVarationConverters.cpp
@@ -268,7 +268,7 @@ bool FGLTFVariationConverter::TryParseMaterialPropertyValue(FGLTFJsonVariant& Ou
 
 	// TODO: find way to determine whether the material is null because "None" was selected, or because it failed to resolve
 	const UMaterialInterface* Material = MaterialProperty->GetMaterial();
-	const FGLTFJsonMaterialIndex MaterialIndex = Builder.GetOrAddMaterial(Material);
+	const FGLTFJsonMaterialIndex MaterialIndex = Builder.GetOrAddMaterial(Material, nullptr);
 	const int32 ElementIndex = CapturedPropSegments[NumPropSegments - 1].PropertyIndex;
 
 	FGLTFJsonVariantMaterial VariantMaterial;
diff --git a/Engine/Plugins/Enterprise/GLTFExporter/MigratedFromGit/Source/GLTFExporter/Private/Exporters/GLTFMaterialExporter.cpp b/Engine/Plugins/Enterprise/GLTFExporter/MigratedFromGit/Source/GLTFExporter/Private/Exporters/GLTFMaterialExporter.cpp
--- a/Engine/Plugins/Enterprise/GLTFExporter/MigratedFromGit/Source/GLTFExporter/Private/Exporters/GLTFMaterialExporter.cpp
+++ b/Engine/Plugins/Enterprise/GLTFExporter/MigratedFromGit/Source/GLTFExporter/Private/Exporters/GLTFMaterialExporter.cpp
@@ -15,7 +15,8 @@ bool UGLTFMaterialExporter::AddObject(FGLTFContainerBuilder& Builder, const UObj
 {
 	const UMaterialInterface* Material = CastChecked<UMaterialInterface>(Object);
 
-	const FGLTFJsonMaterialIndex MaterialIndex = Builder.GetOrAddMaterial(Material);
+	// TODO: should we get the preview-mesh first, and pass it to GetOrAddMaterial for baking?
+	const FGLTFJsonMaterialIndex MaterialIndex = Builder.GetOrAddMaterial(Material, nullptr);
 	if (MaterialIndex == INDEX_NONE)
 	{
 		Builder.AddErrorMessage(
diff --git a/Engine/Plugins/Enterprise/GLTFExporter/MigratedFromGit/Source/GLTFExporter/Private/Tasks/GLTFMaterialTasks.cpp b/Engine/Plugins/Enterprise/GLTFExporter/MigratedFromGit/Source/GLTFExporter/Private/Tasks/GLTFMaterialTasks.cpp
--- a/Engine/Plugins/Enterprise/GLTFExporter/MigratedFromGit/Source/GLTFExporter/Private/Tasks/GLTFMaterialTasks.cpp
+++ b/Engine/Plugins/Enterprise/GLTFExporter/MigratedFromGit/Source/GLTFExporter/Private/Tasks/GLTFMaterialTasks.cpp
@@ -1365,6 +1365,8 @@ FGLTFPropertyBakeOutput FGLTFMaterialTask::BakeMaterialProperty(EMaterialPropert
 		Property,
 		Material,
 		OutTexCoord,
+		const_cast<UStaticMesh*>(Mesh),
+		LODIndex,
 		bCopyAlphaFromRedChannel);
 
 	if (!PropertyBakeOutput.bIsConstant && TexCoords.Num() == 0)
diff --git a/Engine/Plugins/Enterprise/GLTFExporter/MigratedFromGit/Source/GLTFExporter/Private/Tasks/GLTFMaterialTasks.h b/Engine/Plugins/Enterprise/GLTFExporter/MigratedFromGit/Source/GLTFExporter/Private/Tasks/GLTFMaterialTasks.h
--- a/Engine/Plugins/Enterprise/GLTFExporter/MigratedFromGit/Source/GLTFExporter/Private/Tasks/GLTFMaterialTasks.h
+++ b/Engine/Plugins/Enterprise/GLTFExporter/MigratedFromGit/Source/GLTFExporter/Private/Tasks/GLTFMaterialTasks.h
@@ -12,10 +12,12 @@ class FGLTFMaterialTask : public FGLTFTask
 {
 public:
 
-	FGLTFMaterialTask(FGLTFConvertBuilder& Builder, const UMaterialInterface* Material, FGLTFJsonMaterialIndex MaterialIndex)
+	FGLTFMaterialTask(FGLTFConvertBuilder& Builder, const UMaterialInterface* Material, const UStaticMesh* Mesh, int32 LODIndex, FGLTFJsonMaterialIndex MaterialIndex)
         : FGLTFTask(EGLTFTaskPriority::Material)
 		, Builder(Builder)
 		, Material(Material)
+		, Mesh(Mesh)
+		, LODIndex(LODIndex)
 		, MaterialIndex(MaterialIndex)
 	{
 	}
@@ -31,6 +33,8 @@ private:
 
 	FGLTFConvertBuilder& Builder;
 	const UMaterialInterface* Material;
+	const UStaticMesh* Mesh;
+	const int32 LODIndex;
 	const FGLTFJsonMaterialIndex MaterialIndex;
 
 	bool TryGetAlphaMode(EGLTFJsonAlphaMode& AlphaMode) const;
diff --git a/Engine/Plugins/Enterprise/GLTFExporter/MigratedFromGit/Source/GLTFExporter/Private/Tasks/GLTFMeshTasks.cpp b/Engine/Plugins/Enterprise/GLTFExporter/MigratedFromGit/Source/GLTFExporter/Private/Tasks/GLTFMeshTasks.cpp
--- a/Engine/Plugins/Enterprise/GLTFExporter/MigratedFromGit/Source/GLTFExporter/Private/Tasks/GLTFMeshTasks.cpp
+++ b/Engine/Plugins/Enterprise/GLTFExporter/MigratedFromGit/Source/GLTFExporter/Private/Tasks/GLTFMeshTasks.cpp
@@ -55,7 +55,7 @@ void FGLTFStaticMeshTask::Complete()
 			Material = UMaterial::GetDefaultMaterial(MD_Surface);
 		}
 
-		JsonPrimitive.Material = Builder.GetOrAddMaterial(Material);
+		JsonPrimitive.Material = Builder.GetOrAddMaterial(Material, StaticMesh, LODIndex);
 	}
 }
 
@@ -123,6 +123,7 @@ void FGLTFSkeletalMeshTask::Complete()
 			Material = UMaterial::GetDefaultMaterial(MD_Surface);
 		}
 
-		JsonPrimitive.Material = Builder.GetOrAddMaterial(Material);
+		// TODO: should we somehow pass a skeletal mesh for baking?
+		JsonPrimitive.Material = Builder.GetOrAddMaterial(Material, nullptr);
 	}
 }
