[GLTFExporter] fix: fix incorrect baking using mesh-data by specifying all necessary properties

commit 374f88ac7dc1702b9d5895175764844061f74681
Date: Tue, 20 Oct 2020 13:43:08 +0200

    fix: fix incorrect baking using mesh-data by specifying all necessary properties
    (cherry picked from commit 1c4150a1b5cae3676cce17b379dee29513a8694e​)

#skipci

diff --git a/Engine/Plugins/Enterprise/GLTFExporter/MigratedFromGit/Source/GLTFExporter/Private/Converters/GLTFMaterialUtility.cpp b/Engine/Plugins/Enterprise/GLTFExporter/MigratedFromGit/Source/GLTFExporter/Private/Converters/GLTFMaterialUtility.cpp
--- a/Engine/Plugins/Enterprise/GLTFExporter/MigratedFromGit/Source/GLTFExporter/Private/Converters/GLTFMaterialUtility.cpp
+++ b/Engine/Plugins/Enterprise/GLTFExporter/MigratedFromGit/Source/GLTFExporter/Private/Converters/GLTFMaterialUtility.cpp
@@ -13,6 +13,7 @@
 #include "Materials/MaterialExpressionCustomOutput.h"
 #include "Materials/MaterialExpressionClearCoatNormalCustomOutput.h"
 #include "Materials/MaterialExpressionTextureCoordinate.h"
+#include "MeshDescription.h"
 
 UMaterialInterface* FGLTFMaterialUtility::GetDefault()
 {
@@ -283,13 +284,24 @@ FGLTFPropertyBakeOutput FGLTFMaterialUtility::BakeMaterialProperty(const FIntPoi
 
 		if (bRequiresVertexData)
 		{
-			MeshSet.Mesh = Mesh;
-			MeshSet.RawMeshDescription =  Mesh->GetMeshDescription(LODIndex);
+			if (FMeshDescription* MeshDescription = Mesh->GetMeshDescription(LODIndex))
+			{
+				MeshSet.RawMeshDescription = MeshDescription;
+				MeshSet.LightMapIndex = Mesh->LightMapCoordinateIndex;
+
+				// NOTE: we use polygon group ID's instead of material indices since that's what
+				// FMeshMaterialRenderItem::PopulateWithMeshData compares MeshSet.MaterialIndices against.
+				for (const auto& PolygonGroupID: MeshDescription->PolygonGroups().GetElementIDs())
+				{
+					// TODO: should we only add polygon groups that are actually using the material
+					// we're trying to bake? Or would that prevent us from baking materials using unrelated meshes?
+
+					MeshSet.MaterialIndices.Add(PolygonGroupID.GetValue());
+				}
+			}
 		}
 	}
 
-	// TODO: Do we need to fill in any more info in MeshSet?
-
 	MeshSettings.Add(&MeshSet);
 
 	TArray<FMaterialData*> MatSettings;
