[GLTFExporter] fix: reset LODIndex when converting a material that doesn't need mesh data

commit 86573c85bedaaba6d81664fe3f41a6bb0d7b1b75
Date: Tue, 20 Oct 2020 14:17:39 +0200

    fix: reset LODIndex when converting a material that doesn't need mesh data
    prevents duplicate materials that only differ in LODIndex
    (cherry picked from commit 4c07277eb0c46489c23b2c1ccf55954e109aee6a​)

#skipci

diff --git a/Engine/Plugins/Enterprise/GLTFExporter/MigratedFromGit/Source/GLTFExporter/Private/Converters/GLTFMaterialConverters.cpp b/Engine/Plugins/Enterprise/GLTFExporter/MigratedFromGit/Source/GLTFExporter/Private/Converters/GLTFMaterialConverters.cpp
--- a/Engine/Plugins/Enterprise/GLTFExporter/MigratedFromGit/Source/GLTFExporter/Private/Converters/GLTFMaterialConverters.cpp
+++ b/Engine/Plugins/Enterprise/GLTFExporter/MigratedFromGit/Source/GLTFExporter/Private/Converters/GLTFMaterialConverters.cpp
@@ -44,6 +44,7 @@ void FGLTFMaterialConverter::Sanitize(const UMaterialInterface*& Material, const
 		if (!bRequiresVertexData)
 		{
 			Mesh = nullptr;
+			LODIndex = -1;
 		}
 	}
 }
