[GLTFExporter] feat: add bUsesVertexData flag to FGLTFPropertyBakeOutput

commit 97bc7b76ee392ebff775d058f18befcdf8e6ff8c
Date: Tue, 20 Oct 2020 15:21:10 +0200

    feat: add bUsesVertexData flag to FGLTFPropertyBakeOutput
    allows the caller to know if the baked output is mesh-specific or not
    (cherry picked from commit cfa4acc5d6e2e7887fe5c64ad9378bf3fcc3a9a4​)

#skipci

diff --git a/Engine/Plugins/Enterprise/GLTFExporter/MigratedFromGit/Source/GLTFExporter/Private/Converters/GLTFMaterialUtility.cpp b/Engine/Plugins/Enterprise/GLTFExporter/MigratedFromGit/Source/GLTFExporter/Private/Converters/GLTFMaterialUtility.cpp
--- a/Engine/Plugins/Enterprise/GLTFExporter/MigratedFromGit/Source/GLTFExporter/Private/Converters/GLTFMaterialUtility.cpp
+++ b/Engine/Plugins/Enterprise/GLTFExporter/MigratedFromGit/Source/GLTFExporter/Private/Converters/GLTFMaterialUtility.cpp
@@ -276,13 +276,14 @@ FGLTFPropertyBakeOutput FGLTFMaterialUtility::BakeMaterialProperty(const FIntPoi
 	MeshSet.TextureCoordinateBox = { { 0.0f, 0.0f }, { 1.0f, 1.0f } };
 	MeshSet.TextureCoordinateIndex = TexCoord;
 
+	bool bUsesVertexData = false;
+
 	if (Mesh != nullptr)
 	{
 		int32 NumTextureCoordinates;
-		bool bRequiresVertexData;
-		const_cast<UMaterialInterface*>(Material)->AnalyzeMaterialProperty(Property, NumTextureCoordinates, bRequiresVertexData);
+		const_cast<UMaterialInterface*>(Material)->AnalyzeMaterialProperty(Property, NumTextureCoordinates, bUsesVertexData);
 
-		if (bRequiresVertexData)
+		if (bUsesVertexData)
 		{
 			if (FMeshDescription* MeshDescription = Mesh->GetMeshDescription(LODIndex))
 			{
@@ -299,6 +300,11 @@ FGLTFPropertyBakeOutput FGLTFMaterialUtility::BakeMaterialProperty(const FIntPoi
 					MeshSet.MaterialIndices.Add(PolygonGroupID.GetValue());
 				}
 			}
+			else
+			{
+				// TODO: add warning?
+				bUsesVertexData = false;
+			}
 		}
 	}
 
@@ -348,6 +354,7 @@ FGLTFPropertyBakeOutput FGLTFMaterialUtility::BakeMaterialProperty(const FIntPoi
 	}
 
 	FGLTFPropertyBakeOutput PropertyBakeOutput(Property, PF_B8G8R8A8, BakedPixels, BakedSize, EmissiveScale);
+	PropertyBakeOutput.bUsesVertexData = bUsesVertexData;
 
 	if (BakedPixels.Num() == 1)
 	{
diff --git a/Engine/Plugins/Enterprise/GLTFExporter/MigratedFromGit/Source/GLTFExporter/Private/Converters/GLTFMaterialUtility.h b/Engine/Plugins/Enterprise/GLTFExporter/MigratedFromGit/Source/GLTFExporter/Private/Converters/GLTFMaterialUtility.h
--- a/Engine/Plugins/Enterprise/GLTFExporter/MigratedFromGit/Source/GLTFExporter/Private/Converters/GLTFMaterialUtility.h
+++ b/Engine/Plugins/Enterprise/GLTFExporter/MigratedFromGit/Source/GLTFExporter/Private/Converters/GLTFMaterialUtility.h
@@ -19,7 +19,7 @@ struct FGLTFTextureCombineSource
 struct FGLTFPropertyBakeOutput
 {
 	FORCEINLINE FGLTFPropertyBakeOutput(EMaterialProperty Property, EPixelFormat PixelFormat, TArray<FColor>& Pixels, FIntPoint Size, float EmissiveScale)
-		: Property(Property), PixelFormat(PixelFormat), Pixels(Pixels), Size(Size), EmissiveScale(EmissiveScale), bIsConstant(false)
+		: Property(Property), PixelFormat(PixelFormat), Pixels(Pixels), Size(Size), EmissiveScale(EmissiveScale), bIsConstant(false), bUsesVertexData(false)
 	{}
 
 	EMaterialProperty Property;
@@ -30,6 +30,8 @@ struct FGLTFPropertyBakeOutput
 
 	bool bIsConstant;
 	FLinearColor ConstantValue;
+
+	bool bUsesVertexData;
 };
 
 struct FGLTFMaterialUtility
