[GLTFExporter] refactor: move logic for deciding to use mesh for baking into FGLTFMaterialTasks

commit 14b3391e384475183640e41ef9e52147cb08b1e8
Date: Thu, 22 Oct 2020 09:10:46 +0200

    refactor: move logic for deciding to use mesh for baking into FGLTFMaterialTasks
    (cherry picked from commit a532ccc451509785a498eac33fcda00b68c8adb8​)

#skipci

diff --git a/Engine/Plugins/Enterprise/GLTFExporter/MigratedFromGit/Source/GLTFExporter/Private/Converters/GLTFMaterialUtility.cpp b/Engine/Plugins/Enterprise/GLTFExporter/MigratedFromGit/Source/GLTFExporter/Private/Converters/GLTFMaterialUtility.cpp
--- a/Engine/Plugins/Enterprise/GLTFExporter/MigratedFromGit/Source/GLTFExporter/Private/Converters/GLTFMaterialUtility.cpp
+++ b/Engine/Plugins/Enterprise/GLTFExporter/MigratedFromGit/Source/GLTFExporter/Private/Converters/GLTFMaterialUtility.cpp
@@ -241,7 +241,7 @@ bool FGLTFMaterialUtility::CombineTextures(TArray<FColor>& OutPixels, const TArr
 	return bReadSuccessful;
 }
 
-FGLTFPropertyBakeOutput FGLTFMaterialUtility::BakeMaterialProperty(const FIntPoint& OutputSize, EMaterialProperty Property, const UMaterialInterface* Material, int32 TexCoord, UStaticMesh* Mesh, int32 LODIndex, bool bCopyAlphaFromRedChannel)
+FGLTFPropertyBakeOutput FGLTFMaterialUtility::BakeMaterialProperty(const FIntPoint& OutputSize, EMaterialProperty Property, const UMaterialInterface* Material, int32 TexCoord, FMeshDescription* MeshDescription, bool bCopyAlphaFromRedChannel)
 {
 	EMaterialProperty HackProperty = MP_MAX;
 	TTuple<int32, UMaterialExpression*> HackPrevState;
@@ -275,36 +275,18 @@ FGLTFPropertyBakeOutput FGLTFMaterialUtility::BakeMaterialProperty(const FIntPoi
 	FMeshData MeshSet;
 	MeshSet.TextureCoordinateBox = { { 0.0f, 0.0f }, { 1.0f, 1.0f } };
 	MeshSet.TextureCoordinateIndex = TexCoord;
+	MeshSet.RawMeshDescription = MeshDescription;
 
-	bool bUsesVertexData = false;
-
-	if (Mesh != nullptr)
+	if (MeshDescription != nullptr)
 	{
-		int32 NumTextureCoordinates;
-		const_cast<UMaterialInterface*>(Material)->AnalyzeMaterialProperty(Property, NumTextureCoordinates, bUsesVertexData);
-
-		if (bUsesVertexData)
+		// NOTE: we use polygon group ID's instead of material indices since that's what
+		// FMeshMaterialRenderItem::PopulateWithMeshData compares MeshSet.MaterialIndices against.
+		for (const auto& PolygonGroupID: MeshDescription->PolygonGroups().GetElementIDs())
 		{
-			if (FMeshDescription* MeshDescription = Mesh->GetMeshDescription(LODIndex))
-			{
-				MeshSet.RawMeshDescription = MeshDescription;
-				MeshSet.LightMapIndex = Mesh->LightMapCoordinateIndex;
-
-				// NOTE: we use polygon group ID's instead of material indices since that's what
-				// FMeshMaterialRenderItem::PopulateWithMeshData compares MeshSet.MaterialIndices against.
-				for (const auto& PolygonGroupID: MeshDescription->PolygonGroups().GetElementIDs())
-				{
-					// TODO: should we only add polygon groups that are actually using the material
-					// we're trying to bake? Or would that prevent us from baking materials using unrelated meshes?
+			// TODO: should we only add polygon groups that are actually using the material
+			// we're trying to bake? Or would that prevent us from baking materials using unrelated meshes?
 
-					MeshSet.MaterialIndices.Add(PolygonGroupID.GetValue());
-				}
-			}
-			else
-			{
-				// TODO: add warning?
-				bUsesVertexData = false;
-			}
+			MeshSet.MaterialIndices.Add(PolygonGroupID.GetValue());
 		}
 	}
 
@@ -354,7 +336,6 @@ FGLTFPropertyBakeOutput FGLTFMaterialUtility::BakeMaterialProperty(const FIntPoi
 	}
 
 	FGLTFPropertyBakeOutput PropertyBakeOutput(Property, PF_B8G8R8A8, BakedPixels, BakedSize, EmissiveScale);
-	PropertyBakeOutput.bUsesVertexData = bUsesVertexData;
 
 	if (BakedPixels.Num() == 1)
 	{
diff --git a/Engine/Plugins/Enterprise/GLTFExporter/MigratedFromGit/Source/GLTFExporter/Private/Converters/GLTFMaterialUtility.h b/Engine/Plugins/Enterprise/GLTFExporter/MigratedFromGit/Source/GLTFExporter/Private/Converters/GLTFMaterialUtility.h
--- a/Engine/Plugins/Enterprise/GLTFExporter/MigratedFromGit/Source/GLTFExporter/Private/Converters/GLTFMaterialUtility.h
+++ b/Engine/Plugins/Enterprise/GLTFExporter/MigratedFromGit/Source/GLTFExporter/Private/Converters/GLTFMaterialUtility.h
@@ -55,7 +55,7 @@ struct FGLTFMaterialUtility
 
 	static UTexture2D* CreateTransientTexture(const FGLTFPropertyBakeOutput& PropertyBakeOutput, bool bUseSRGB = false);
 	static bool CombineTextures(TArray<FColor>& OutPixels, const TArray<FGLTFTextureCombineSource>& Sources, const FIntPoint& OutputSize, EPixelFormat OutputPixelFormat);
-	static FGLTFPropertyBakeOutput BakeMaterialProperty(const FIntPoint& OutputSize, EMaterialProperty Property, const UMaterialInterface* Material, int32 TexCoord, UStaticMesh* Mesh, int32 LODIndex, bool bCopyAlphaFromRedChannel);
+	static FGLTFPropertyBakeOutput BakeMaterialProperty(const FIntPoint& OutputSize, EMaterialProperty Property, const UMaterialInterface* Material, int32 TexCoord, FMeshDescription* MeshDescription, bool bCopyAlphaFromRedChannel);
 
 	static FGLTFJsonTextureIndex AddCombinedTexture(FGLTFConvertBuilder& Builder, const TArray<FGLTFTextureCombineSource>& CombineSources, const FIntPoint& TextureSize, const FString& TextureName, EGLTFJsonTextureFilter MinFilter, EGLTFJsonTextureFilter MagFilter, EGLTFJsonTextureWrap WrapS, EGLTFJsonTextureWrap WrapT);
 	static FGLTFJsonTextureIndex AddTexture(FGLTFConvertBuilder& Builder, const TArray<FColor>& Pixels, const FIntPoint& TextureSize, const FString& TextureName, EGLTFJsonTextureFilter MinFilter, EGLTFJsonTextureFilter MagFilter, EGLTFJsonTextureWrap WrapS, EGLTFJsonTextureWrap WrapT);
diff --git a/Engine/Plugins/Enterprise/GLTFExporter/MigratedFromGit/Source/GLTFExporter/Private/Tasks/GLTFMaterialTasks.cpp b/Engine/Plugins/Enterprise/GLTFExporter/MigratedFromGit/Source/GLTFExporter/Private/Tasks/GLTFMaterialTasks.cpp
--- a/Engine/Plugins/Enterprise/GLTFExporter/MigratedFromGit/Source/GLTFExporter/Private/Tasks/GLTFMaterialTasks.cpp
+++ b/Engine/Plugins/Enterprise/GLTFExporter/MigratedFromGit/Source/GLTFExporter/Private/Tasks/GLTFMaterialTasks.cpp
@@ -1349,6 +1349,8 @@ FGLTFPropertyBakeOutput FGLTFMaterialTask::BakeMaterialProperty(EMaterialPropert
 				OutTexCoord = 1; // assume ambient occlusion uses TexCoord1 when multiple
 			}
 		}
+
+		// TODO: should we perhaps always use the lightmap coordinate index for baking to guarantee unique uvs?
 	}
 	else
 	{
@@ -1358,6 +1360,25 @@ FGLTFPropertyBakeOutput FGLTFMaterialTask::BakeMaterialProperty(EMaterialPropert
 	const FIntPoint DefaultTextureSize = Builder.ExportOptions->GetDefaultMaterialBakeSize();
 	const FIntPoint TextureSize = PreferredTextureSize != nullptr ? *PreferredTextureSize : DefaultTextureSize;
 
+	FMeshDescription* MeshDescription = nullptr;
+
+	// Only use the mesh for baking if the material property requires it
+	if (Mesh != nullptr)
+	{
+		int32 NumTextureCoordinates;
+		bool bRequiresVertexData;
+		const_cast<UMaterialInterface*>(Material)->AnalyzeMaterialProperty(Property, NumTextureCoordinates, bRequiresVertexData);
+
+		if (bRequiresVertexData)
+		{
+			MeshDescription = Mesh->GetMeshDescription(LODIndex);
+			if (MeshDescription == nullptr)
+			{
+				// TODO: add warning
+			}
+		}
+	}
+
 	// TODO: add support for calculating the ideal resolution to use for baking based on connected (texture) nodes
 
 	const FGLTFPropertyBakeOutput PropertyBakeOutput = FGLTFMaterialUtility::BakeMaterialProperty(
@@ -1365,8 +1386,7 @@ FGLTFPropertyBakeOutput FGLTFMaterialTask::BakeMaterialProperty(EMaterialPropert
 		Property,
 		Material,
 		OutTexCoord,
-		const_cast<UStaticMesh*>(Mesh),
-		LODIndex,
+		MeshDescription,
 		bCopyAlphaFromRedChannel);
 
 	if (!PropertyBakeOutput.bIsConstant && TexCoords.Num() == 0)
