[GLTFExporter] feat: add and implement export option bBakeMaterialInputsUsingMeshData

commit 160aacdb36e2769a7073fc681e42b896ec06681f
Date: Thu, 22 Oct 2020 10:19:01 +0200

    feat: add and implement export option bBakeMaterialInputsUsingMeshData
    (cherry picked from commit 0a324377c45ff974e0d1743a6a1d76f6a7df8aec​)

#skipci

diff --git a/Engine/Plugins/Enterprise/GLTFExporter/MigratedFromGit/Source/GLTFExporter/Private/Converters/GLTFMaterialConverters.cpp b/Engine/Plugins/Enterprise/GLTFExporter/MigratedFromGit/Source/GLTFExporter/Private/Converters/GLTFMaterialConverters.cpp
--- a/Engine/Plugins/Enterprise/GLTFExporter/MigratedFromGit/Source/GLTFExporter/Private/Converters/GLTFMaterialConverters.cpp
+++ b/Engine/Plugins/Enterprise/GLTFExporter/MigratedFromGit/Source/GLTFExporter/Private/Converters/GLTFMaterialConverters.cpp
@@ -8,7 +8,11 @@
 
 void FGLTFMaterialConverter::Sanitize(const UMaterialInterface*& Material, const UStaticMesh*& Mesh, int32& LODIndex)
 {
-	if (Mesh != nullptr)
+	bool bRequiresVertexData = false;
+
+	if (Mesh != nullptr &&
+		Builder.ExportOptions->bBakeMaterialInputs &&
+		Builder.ExportOptions->bBakeMaterialInputsUsingMeshData)
 	{
 		if (LODIndex < 0)
 		{
@@ -30,8 +34,6 @@ void FGLTFMaterialConverter::Sanitize(const UMaterialInterface*& Material, const
 			MP_CustomData1
 		};
 
-		bool bRequiresVertexData = false;
-
 		int32 NumTextureCoordinates;
 		bool bPropertyRequiresVertexData;
 
@@ -40,12 +42,12 @@ void FGLTFMaterialConverter::Sanitize(const UMaterialInterface*& Material, const
 			const_cast<UMaterialInterface*>(Material)->AnalyzeMaterialProperty(Property, NumTextureCoordinates, bPropertyRequiresVertexData);
 			bRequiresVertexData |= bPropertyRequiresVertexData;
 		}
+	}
 
-		if (!bRequiresVertexData)
-		{
-			Mesh = nullptr;
-			LODIndex = -1;
-		}
+	if (!bRequiresVertexData)
+	{
+		Mesh = nullptr;
+		LODIndex = -1;
 	}
 }
 
diff --git a/Engine/Plugins/Enterprise/GLTFExporter/MigratedFromGit/Source/GLTFExporter/Private/GLTFExportOptions.cpp b/Engine/Plugins/Enterprise/GLTFExporter/MigratedFromGit/Source/GLTFExporter/Private/GLTFExportOptions.cpp
--- a/Engine/Plugins/Enterprise/GLTFExporter/MigratedFromGit/Source/GLTFExporter/Private/GLTFExportOptions.cpp
+++ b/Engine/Plugins/Enterprise/GLTFExporter/MigratedFromGit/Source/GLTFExporter/Private/GLTFExportOptions.cpp
@@ -17,6 +17,7 @@ void UGLTFExportOptions::ResetToDefault()
 	bExportUnlitMaterials = true;
 	bExportClearCoatMaterials = true;
 	bBakeMaterialInputs = true;
+	bBakeMaterialInputsUsingMeshData = true;
 	DefaultMaterialBakeSize = EGLTFExporterMaterialBakeSize::POT_1024;
 	bExportVertexColors = false;
 	bQuantizeVertexNormals = true;
@@ -77,6 +78,11 @@ bool UGLTFExportOptions::CanEditChange(const FProperty* InProperty) const
 		return bExportCameras;
 	}
 
+	if (PropertyFName == GET_MEMBER_NAME_CHECKED(ThisClass, bBakeMaterialInputsUsingMeshData))
+	{
+		return bBakeMaterialInputs;
+	}
+
 	return true;
 }
 
diff --git a/Engine/Plugins/Enterprise/GLTFExporter/MigratedFromGit/Source/GLTFExporter/Public/GLTFExportOptions.h b/Engine/Plugins/Enterprise/GLTFExporter/MigratedFromGit/Source/GLTFExporter/Public/GLTFExportOptions.h
--- a/Engine/Plugins/Enterprise/GLTFExporter/MigratedFromGit/Source/GLTFExporter/Public/GLTFExportOptions.h
+++ b/Engine/Plugins/Enterprise/GLTFExporter/MigratedFromGit/Source/GLTFExporter/Public/GLTFExportOptions.h
@@ -58,6 +58,9 @@ public:
 	UPROPERTY(EditAnywhere, BlueprintReadWrite, Config, Category = Material)
 	bool bBakeMaterialInputs;
 
+	UPROPERTY(EditAnywhere, BlueprintReadWrite, Config, Category = Material)
+	bool bBakeMaterialInputsUsingMeshData;
+
 	UPROPERTY(EditAnywhere, BlueprintReadWrite, Config, Category = Material)
 	EGLTFExporterMaterialBakeSize DefaultMaterialBakeSize;
 
