[GLTFExporter] refactor: use existing utility for retrieving mesh-description when baking

commit 58be8ca4e00daaecc6f0917b043f2f75497b0a61
Date: Thu, 22 Oct 2020 15:13:18 +0200

    refactor: use existing utility for retrieving mesh-description when baking
    (cherry picked from commit 1229424b970d9b45c5fbffaf47cb46482ad0db7f​)

#skipci

diff --git a/Engine/Plugins/Enterprise/GLTFExporter/MigratedFromGit/Source/GLTFExporter/GLTFExporter.Build.cs b/Engine/Plugins/Enterprise/GLTFExporter/MigratedFromGit/Source/GLTFExporter/GLTFExporter.Build.cs
--- a/Engine/Plugins/Enterprise/GLTFExporter/MigratedFromGit/Source/GLTFExporter/GLTFExporter.Build.cs
+++ b/Engine/Plugins/Enterprise/GLTFExporter/MigratedFromGit/Source/GLTFExporter/GLTFExporter.Build.cs
@@ -33,7 +33,8 @@ namespace UnrealBuildTool.Rules
 					"MovieSceneTracks",
 					"MaterialBaking",
 					"VariantManagerContent",
-					"GLTFExporterRuntime"
+					"GLTFExporterRuntime",
+					"MeshMergeUtilities"
 				}
 				);
 		}
diff --git a/Engine/Plugins/Enterprise/GLTFExporter/MigratedFromGit/Source/GLTFExporter/Private/Tasks/GLTFMaterialTasks.cpp b/Engine/Plugins/Enterprise/GLTFExporter/MigratedFromGit/Source/GLTFExporter/Private/Tasks/GLTFMaterialTasks.cpp
--- a/Engine/Plugins/Enterprise/GLTFExporter/MigratedFromGit/Source/GLTFExporter/Private/Tasks/GLTFMaterialTasks.cpp
+++ b/Engine/Plugins/Enterprise/GLTFExporter/MigratedFromGit/Source/GLTFExporter/Private/Tasks/GLTFMaterialTasks.cpp
@@ -13,6 +13,9 @@
 #include "Materials/MaterialExpressionVectorParameter.h"
 #include "Materials/MaterialExpressionTextureSample.h"
 #include "Materials/MaterialExpressionTextureSampleParameter2D.h"
+#include "MeshDescription.h"
+#include "StaticMeshAttributes.h"
+#include "Developer/MeshMergeUtilities/Private/MeshMergeHelpers.h"
 
 namespace
 {
@@ -1360,22 +1363,21 @@ FGLTFPropertyBakeOutput FGLTFMaterialTask::BakeMaterialProperty(EMaterialPropert
 	const FIntPoint DefaultTextureSize = Builder.ExportOptions->GetDefaultMaterialBakeSize();
 	const FIntPoint TextureSize = PreferredTextureSize != nullptr ? *PreferredTextureSize : DefaultTextureSize;
 
-	FMeshDescription* MeshDescription = nullptr;
+	bool bRequiresVertexData = false;
+	FMeshDescription MeshDescription;
 
 	// Only use the mesh for baking if the material property requires it
 	if (Mesh != nullptr)
 	{
 		int32 NumTextureCoordinates;
-		bool bRequiresVertexData;
 		const_cast<UMaterialInterface*>(Material)->AnalyzeMaterialProperty(Property, NumTextureCoordinates, bRequiresVertexData);
 
 		if (bRequiresVertexData)
 		{
-			MeshDescription = Mesh->GetMeshDescription(LODIndex);
-			if (MeshDescription == nullptr)
-			{
-				// TODO: add warning
-			}
+			// TODO: add support mesh-components (and perhaps skeletal meshes)
+
+			FStaticMeshAttributes(MeshDescription).Register();
+			FMeshMergeHelpers::RetrieveMesh(const_cast<UStaticMesh*>(Mesh), LODIndex, MeshDescription);
 		}
 	}
 
@@ -1386,7 +1388,7 @@ FGLTFPropertyBakeOutput FGLTFMaterialTask::BakeMaterialProperty(EMaterialPropert
 		Property,
 		Material,
 		OutTexCoord,
-		MeshDescription,
+		bRequiresVertexData ? &MeshDescription : nullptr,
 		bCopyAlphaFromRedChannel);
 
 	if (!PropertyBakeOutput.bIsConstant && TexCoords.Num() == 0)
