[GLTFExporter] refactor: remove bUsesVertexData from FGLTFPropertyBakeOutput since we dont need it

commit 9718373e19d6d6d655d7e0af82ba581b46778f1e
Date: Thu, 22 Oct 2020 14:42:09 +0200

    refactor: remove bUsesVertexData from FGLTFPropertyBakeOutput since we dont need it
    (cherry picked from commit 46e18bb79b08c5814f3f3e3e4f6b38207b2ec76c​)

#skipci

diff --git a/Engine/Plugins/Enterprise/GLTFExporter/MigratedFromGit/Source/GLTFExporter/Private/Converters/GLTFMaterialUtility.h b/Engine/Plugins/Enterprise/GLTFExporter/MigratedFromGit/Source/GLTFExporter/Private/Converters/GLTFMaterialUtility.h
--- a/Engine/Plugins/Enterprise/GLTFExporter/MigratedFromGit/Source/GLTFExporter/Private/Converters/GLTFMaterialUtility.h
+++ b/Engine/Plugins/Enterprise/GLTFExporter/MigratedFromGit/Source/GLTFExporter/Private/Converters/GLTFMaterialUtility.h
@@ -19,7 +19,7 @@ struct FGLTFTextureCombineSource
 struct FGLTFPropertyBakeOutput
 {
 	FORCEINLINE FGLTFPropertyBakeOutput(EMaterialProperty Property, EPixelFormat PixelFormat, TArray<FColor>& Pixels, FIntPoint Size, float EmissiveScale)
-		: Property(Property), PixelFormat(PixelFormat), Pixels(Pixels), Size(Size), EmissiveScale(EmissiveScale), bIsConstant(false), bUsesVertexData(false)
+		: Property(Property), PixelFormat(PixelFormat), Pixels(Pixels), Size(Size), EmissiveScale(EmissiveScale), bIsConstant(false)
 	{}
 
 	EMaterialProperty Property;
@@ -30,8 +30,6 @@ struct FGLTFPropertyBakeOutput
 
 	bool bIsConstant;
 	FLinearColor ConstantValue;
-
-	bool bUsesVertexData;
 };
 
 struct FGLTFMaterialUtility
