[GLTFExporter] feat: use mesh data for all properties when baking, if atleast one property needs it

commit ef86a5efe22e57cc3571e13dfb282ae6406131ed
Date: Thu, 22 Oct 2020 15:14:41 +0200

    feat: use mesh data for all properties when baking, if atleast one property needs it
    (cherry picked from commit a552e92468b58dd111ab4a11013ff06ca2eb6732​)

#skipci

diff --git a/Engine/Plugins/Enterprise/GLTFExporter/MigratedFromGit/Source/GLTFExporter/Private/Tasks/GLTFMaterialTasks.cpp b/Engine/Plugins/Enterprise/GLTFExporter/MigratedFromGit/Source/GLTFExporter/Private/Tasks/GLTFMaterialTasks.cpp
--- a/Engine/Plugins/Enterprise/GLTFExporter/MigratedFromGit/Source/GLTFExporter/Private/Tasks/GLTFMaterialTasks.cpp
+++ b/Engine/Plugins/Enterprise/GLTFExporter/MigratedFromGit/Source/GLTFExporter/Private/Tasks/GLTFMaterialTasks.cpp
@@ -1363,22 +1363,12 @@ FGLTFPropertyBakeOutput FGLTFMaterialTask::BakeMaterialProperty(EMaterialPropert
 	const FIntPoint DefaultTextureSize = Builder.ExportOptions->GetDefaultMaterialBakeSize();
 	const FIntPoint TextureSize = PreferredTextureSize != nullptr ? *PreferredTextureSize : DefaultTextureSize;
 
-	bool bRequiresVertexData = false;
 	FMeshDescription MeshDescription;
 
-	// Only use the mesh for baking if the material property requires it
 	if (Mesh != nullptr)
 	{
-		int32 NumTextureCoordinates;
-		const_cast<UMaterialInterface*>(Material)->AnalyzeMaterialProperty(Property, NumTextureCoordinates, bRequiresVertexData);
-
-		if (bRequiresVertexData)
-		{
-			// TODO: add support mesh-components (and perhaps skeletal meshes)
-
-			FStaticMeshAttributes(MeshDescription).Register();
-			FMeshMergeHelpers::RetrieveMesh(const_cast<UStaticMesh*>(Mesh), LODIndex, MeshDescription);
-		}
+		FStaticMeshAttributes(MeshDescription).Register();
+		FMeshMergeHelpers::RetrieveMesh(const_cast<UStaticMesh*>(Mesh), LODIndex, MeshDescription);
 	}
 
 	// TODO: add support for calculating the ideal resolution to use for baking based on connected (texture) nodes
@@ -1388,7 +1378,7 @@ FGLTFPropertyBakeOutput FGLTFMaterialTask::BakeMaterialProperty(EMaterialPropert
 		Property,
 		Material,
 		OutTexCoord,
-		bRequiresVertexData ? &MeshDescription : nullptr,
+		Mesh != nullptr ? &MeshDescription : nullptr,
 		bCopyAlphaFromRedChannel);
 
 	if (!PropertyBakeOutput.bIsConstant && TexCoords.Num() == 0)
