[GLTFExporter] refactor: minor clean up / refactoring of FGLTFMaterialConverter::Sanitize

commit bd02c6ed1e3477c914638130eede9f146f418ac3
Date: Thu, 22 Oct 2020 15:15:18 +0200

    refactor: minor clean up / refactoring of FGLTFMaterialConverter::Sanitize
    (cherry picked from commit 1a8ef886bc5321b12238736004477158441eedcd​)

#skipci

diff --git a/Engine/Plugins/Enterprise/GLTFExporter/MigratedFromGit/Source/GLTFExporter/Private/Converters/GLTFMaterialConverters.cpp b/Engine/Plugins/Enterprise/GLTFExporter/MigratedFromGit/Source/GLTFExporter/Private/Converters/GLTFMaterialConverters.cpp
--- a/Engine/Plugins/Enterprise/GLTFExporter/MigratedFromGit/Source/GLTFExporter/Private/Converters/GLTFMaterialConverters.cpp
+++ b/Engine/Plugins/Enterprise/GLTFExporter/MigratedFromGit/Source/GLTFExporter/Private/Converters/GLTFMaterialConverters.cpp
@@ -8,43 +8,19 @@
 
 void FGLTFMaterialConverter::Sanitize(const UMaterialInterface*& Material, const UStaticMesh*& Mesh, int32& LODIndex)
 {
-	bool bRequiresVertexData = false;
+	// TODO: add support for UStaticMeshComponent, USkeletalMesh and USkeletalMeshComponent too
 
 	if (Mesh != nullptr &&
 		Builder.ExportOptions->bBakeMaterialInputs &&
-		Builder.ExportOptions->bBakeMaterialInputsUsingMeshData)
+		Builder.ExportOptions->bBakeMaterialInputsUsingMeshData &&
+		MaterialNeedsVertexData(const_cast<UMaterialInterface*>(Material)))
 	{
 		if (LODIndex < 0)
 		{
 			LODIndex = FMath::Max(Builder.ExportOptions->DefaultLevelOfDetail, FGLTFMeshUtility::GetMinimumLOD(Mesh));
 		}
-
-		const TArray<EMaterialProperty> Properties =
-		{
-			MP_BaseColor,
-			MP_EmissiveColor,
-			MP_Opacity,
-			MP_OpacityMask,
-			MP_Metallic,
-			MP_Roughness,
-			MP_Normal,
-			// MP_CustomOutput,	// NOTE: causes a crash when used with AnalyzeMaterialProperty
-			MP_AmbientOcclusion,
-			MP_CustomData0,
-			MP_CustomData1
-		};
-
-		int32 NumTextureCoordinates;
-		bool bPropertyRequiresVertexData;
-
-		for (const EMaterialProperty Property: Properties)
-		{
-			const_cast<UMaterialInterface*>(Material)->AnalyzeMaterialProperty(Property, NumTextureCoordinates, bPropertyRequiresVertexData);
-			bRequiresVertexData |= bPropertyRequiresVertexData;
-		}
 	}
-
-	if (!bRequiresVertexData)
+	else
 	{
 		Mesh = nullptr;
 		LODIndex = -1;
@@ -62,3 +38,36 @@ FGLTFJsonMaterialIndex FGLTFMaterialConverter::Convert(const UMaterialInterface*
 	Builder.SetupTask<FGLTFMaterialTask>(Builder, Material, Mesh, LODIndex, MaterialIndex);
 	return MaterialIndex;
 }
+
+bool FGLTFMaterialConverter::MaterialNeedsVertexData(UMaterialInterface* Material)
+{
+	// TODO: only analyze properties that will be needed for this specific material
+
+	const TArray<EMaterialProperty> Properties =
+	{
+		MP_BaseColor,
+		MP_EmissiveColor,
+		MP_Opacity,
+		MP_OpacityMask,
+		MP_Metallic,
+		MP_Roughness,
+		MP_Normal,
+		// MP_CustomOutput,	// NOTE: causes a crash when used with AnalyzeMaterialProperty
+		MP_AmbientOcclusion,
+		MP_CustomData0,
+		MP_CustomData1
+	};
+
+	bool bRequiresVertexData = false;
+
+	int32 NumTextureCoordinates;
+	bool bPropertyRequiresVertexData;
+
+	for (const EMaterialProperty Property: Properties)
+	{
+		const_cast<UMaterialInterface*>(Material)->AnalyzeMaterialProperty(Property, NumTextureCoordinates, bPropertyRequiresVertexData);
+		bRequiresVertexData |= bPropertyRequiresVertexData;
+	}
+
+	return bRequiresVertexData;
+}
diff --git a/Engine/Plugins/Enterprise/GLTFExporter/MigratedFromGit/Source/GLTFExporter/Private/Converters/GLTFMaterialConverters.h b/Engine/Plugins/Enterprise/GLTFExporter/MigratedFromGit/Source/GLTFExporter/Private/Converters/GLTFMaterialConverters.h
--- a/Engine/Plugins/Enterprise/GLTFExporter/MigratedFromGit/Source/GLTFExporter/Private/Converters/GLTFMaterialConverters.h
+++ b/Engine/Plugins/Enterprise/GLTFExporter/MigratedFromGit/Source/GLTFExporter/Private/Converters/GLTFMaterialConverters.h
@@ -14,4 +14,6 @@ class FGLTFMaterialConverter final : public FGLTFBuilderContext, public TGLTFCon
 	virtual void Sanitize(const UMaterialInterface*& Material, const UStaticMesh*& Mesh, int32& LODIndex) override;
 
 	virtual FGLTFJsonMaterialIndex Convert(const UMaterialInterface* Material, const UStaticMesh* Mesh, int32 LODIndex) override;
+
+	static bool MaterialNeedsVertexData(UMaterialInterface* Material);
 };
