[GLTFExporter] fix: prevent exporting duplicate materials because of different arguments

commit f22be3603d27da906cabb18a7a9e2d20d31647c4
Date: Thu, 22 Oct 2020 15:43:22 +0200

    fix: prevent exporting duplicate materials because of different arguments
    (cherry picked from commit 85c03d57abf0dd78cdaaa90e584b21b3f77d5466​)

#skipci

diff --git a/Engine/Plugins/Enterprise/GLTFExporter/MigratedFromGit/Source/GLTFExporter/Private/Exporters/GLTFMaterialExporter.cpp b/Engine/Plugins/Enterprise/GLTFExporter/MigratedFromGit/Source/GLTFExporter/Private/Exporters/GLTFMaterialExporter.cpp
--- a/Engine/Plugins/Enterprise/GLTFExporter/MigratedFromGit/Source/GLTFExporter/Private/Exporters/GLTFMaterialExporter.cpp
+++ b/Engine/Plugins/Enterprise/GLTFExporter/MigratedFromGit/Source/GLTFExporter/Private/Exporters/GLTFMaterialExporter.cpp
@@ -14,9 +14,9 @@ UGLTFMaterialExporter::UGLTFMaterialExporter(const FObjectInitializer& ObjectIni
 bool UGLTFMaterialExporter::AddObject(FGLTFContainerBuilder& Builder, const UObject* Object)
 {
 	const UMaterialInterface* Material = CastChecked<UMaterialInterface>(Object);
+	const UStaticMesh* PreviewMesh = Builder.ExportOptions->bExportPreviewMesh ? FGLTFExporterUtility::GetPreviewMesh(Material) : nullptr;
 
-	// TODO: should we get the preview-mesh first, and pass it to GetOrAddMaterial for baking?
-	const FGLTFJsonMaterialIndex MaterialIndex = Builder.GetOrAddMaterial(Material, nullptr);
+	const FGLTFJsonMaterialIndex MaterialIndex = Builder.GetOrAddMaterial(Material, PreviewMesh);
 	if (MaterialIndex == INDEX_NONE)
 	{
 		Builder.AddErrorMessage(
@@ -27,7 +27,6 @@ bool UGLTFMaterialExporter::AddObject(FGLTFContainerBuilder& Builder, const UObj
 
 	if (Builder.ExportOptions->bExportPreviewMesh)
 	{
-		const UStaticMesh* PreviewMesh = FGLTFExporterUtility::GetPreviewMesh(Material);
 		if (PreviewMesh != nullptr)
 		{
 			const FGLTFJsonMeshIndex MeshIndex = Builder.GetOrAddMesh(PreviewMesh, -1, nullptr, { Material });
