[GLTFExporter] feat: use mesh name as part of material- and texture-names when baking with mesh data

commit 1620860f0441235c7edcbc34088efc567e5395d3
Date: Thu, 22 Oct 2020 16:16:32 +0200

    feat: use mesh name as part of material- and texture-names when baking with mesh data
    (cherry picked from commit 264d874dfa4f1187aed77c0a74606f6beb89cb78​)

#skipci

diff --git a/Engine/Plugins/Enterprise/GLTFExporter/MigratedFromGit/Source/GLTFExporter/Private/Tasks/GLTFMaterialTasks.cpp b/Engine/Plugins/Enterprise/GLTFExporter/MigratedFromGit/Source/GLTFExporter/Private/Tasks/GLTFMaterialTasks.cpp
--- a/Engine/Plugins/Enterprise/GLTFExporter/MigratedFromGit/Source/GLTFExporter/Private/Tasks/GLTFMaterialTasks.cpp
+++ b/Engine/Plugins/Enterprise/GLTFExporter/MigratedFromGit/Source/GLTFExporter/Private/Tasks/GLTFMaterialTasks.cpp
@@ -64,7 +64,7 @@ void FGLTFMaterialTask::Complete()
 	}
 
 	FGLTFJsonMaterial& JsonMaterial = Builder.GetMaterial(MaterialIndex);
-	Material->GetName(JsonMaterial.Name);
+	JsonMaterial.Name = GetMaterialName();
 
 	if (!TryGetAlphaMode(JsonMaterial.AlphaMode))
 	{
@@ -397,8 +397,6 @@ bool FGLTFMaterialTask::TryGetBaseColorAndOpacity(FGLTFJsonPBRMetallicRoughness&
 	BaseColorTexture = FGLTFMaterialUtility::CreateTransientTexture(BaseColorBakeOutput);
 	OpacityTexture = FGLTFMaterialUtility::CreateTransientTexture(OpacityBakeOutput);
 
-	const FString TextureName = Material->GetName() + TEXT("_BaseColor");
-
 	const TArray<FGLTFTextureCombineSource> CombineSources =
 	{
 		{ OpacityTexture, OpacityMask, SE_BLEND_Opaque },
@@ -409,7 +407,7 @@ bool FGLTFMaterialTask::TryGetBaseColorAndOpacity(FGLTFJsonPBRMetallicRoughness&
 		Builder,
 		CombineSources,
 		TextureSize,
-		TextureName,
+		GetBakedTextureName(TEXT("BaseColor")),
 		TextureMinFilter,
 		TextureMagFilter,
 		TextureWrapS,
@@ -563,8 +561,6 @@ bool FGLTFMaterialTask::TryGetMetallicAndRoughness(FGLTFJsonPBRMetallicRoughness
 	MetallicTexture = FGLTFMaterialUtility::CreateTransientTexture(MetallicBakeOutput);
 	RoughnessTexture = FGLTFMaterialUtility::CreateTransientTexture(RoughnessBakeOutput);
 
-	const FString TextureName = Material->GetName() + TEXT("_MetallicRoughness");
-
 	const TArray<FGLTFTextureCombineSource> CombineSources =
 	{
 		{ MetallicTexture, MetallicMask + AlphaMask, SE_BLEND_Opaque },
@@ -575,7 +571,7 @@ bool FGLTFMaterialTask::TryGetMetallicAndRoughness(FGLTFJsonPBRMetallicRoughness
 		Builder,
 		CombineSources,
 		TextureSize,
-		TextureName,
+		GetBakedTextureName(TEXT("MetallicRoughness")),
 		TextureMinFilter,
 		TextureMagFilter,
 		TextureWrapS,
@@ -731,8 +727,6 @@ bool FGLTFMaterialTask::TryGetClearCoatRoughness(FGLTFJsonClearCoatExtension& Ou
 	IntensityTexture = FGLTFMaterialUtility::CreateTransientTexture(IntensityBakeOutput);
 	RoughnessTexture = FGLTFMaterialUtility::CreateTransientTexture(RoughnessBakeOutput);
 
-	const FString TextureName = Material->GetName() + TEXT("_ClearCoatRoughness");
-
 	const TArray<FGLTFTextureCombineSource> CombineSources =
 	{
 		{ IntensityTexture, ClearCoatMask + AlphaMask, SE_BLEND_Opaque },
@@ -743,7 +737,7 @@ bool FGLTFMaterialTask::TryGetClearCoatRoughness(FGLTFJsonClearCoatExtension& Ou
 		Builder,
 		CombineSources,
 		TextureSize,
-		TextureName,
+		GetBakedTextureName(TEXT("ClearCoatRoughness")),
 		TextureMinFilter,
 		TextureMagFilter,
 		TextureWrapS,
@@ -1312,12 +1306,11 @@ bool FGLTFMaterialTask::TryGetBakedMaterialProperty(FGLTFJsonTextureInfo& OutTex
 	// even though the same material when set to opaque will properly bake AmbientOcclusion to a texture.
 	// For now, create a 1x1 texture with the constant value.
 
-	const FString TextureName = Material->GetName() + TEXT("_") + PropertyName;
 	const FGLTFJsonTextureIndex TextureIndex = FGLTFMaterialUtility::AddTexture(
 		Builder,
 		PropertyBakeOutput.Pixels,
 		PropertyBakeOutput.Size,
-		Material->GetName() + TEXT("_") + PropertyName,
+		GetBakedTextureName(PropertyName),
 		EGLTFJsonTextureFilter::Nearest,
 		EGLTFJsonTextureFilter::Nearest,
 		EGLTFJsonTextureWrap::ClampToEdge,
@@ -1391,8 +1384,6 @@ FGLTFPropertyBakeOutput FGLTFMaterialTask::BakeMaterialProperty(EMaterialPropert
 
 bool FGLTFMaterialTask::StoreBakedPropertyTexture(FGLTFJsonTextureInfo& OutTexInfo, const FGLTFPropertyBakeOutput& PropertyBakeOutput, const FString& PropertyName) const
 {
-	const FString TextureName = Material->GetName() + TEXT("_") + PropertyName;
-
 	// TODO: should this be the default wrap-mode?
 	const EGLTFJsonTextureWrap TextureWrapS = EGLTFJsonTextureWrap::Repeat;
 	const EGLTFJsonTextureWrap TextureWrapT = EGLTFJsonTextureWrap::Repeat;
@@ -1405,7 +1396,7 @@ bool FGLTFMaterialTask::StoreBakedPropertyTexture(FGLTFJsonTextureInfo& OutTexIn
 		Builder,
 		PropertyBakeOutput.Pixels,
 		PropertyBakeOutput.Size,
-		TextureName,
+		GetBakedTextureName(PropertyName),
 		TextureMinFilter,
 		TextureMagFilter,
 		TextureWrapS,
@@ -1414,3 +1405,20 @@ bool FGLTFMaterialTask::StoreBakedPropertyTexture(FGLTFJsonTextureInfo& OutTexIn
 	OutTexInfo.Index = TextureIndex;
 	return true;
 }
+
+FString FGLTFMaterialTask::GetMaterialName() const
+{
+	FString MaterialName = Material->GetName();
+
+	if (Mesh != nullptr)
+	{
+		MaterialName += TEXT("_") + FGLTFNameUtility::GetName(Mesh, LODIndex);
+	}
+
+	return MaterialName;
+}
+
+FString FGLTFMaterialTask::GetBakedTextureName(const FString& PropertyName) const
+{
+	return GetMaterialName() + TEXT("_") + PropertyName;
+}
diff --git a/Engine/Plugins/Enterprise/GLTFExporter/MigratedFromGit/Source/GLTFExporter/Private/Tasks/GLTFMaterialTasks.h b/Engine/Plugins/Enterprise/GLTFExporter/MigratedFromGit/Source/GLTFExporter/Private/Tasks/GLTFMaterialTasks.h
--- a/Engine/Plugins/Enterprise/GLTFExporter/MigratedFromGit/Source/GLTFExporter/Private/Tasks/GLTFMaterialTasks.h
+++ b/Engine/Plugins/Enterprise/GLTFExporter/MigratedFromGit/Source/GLTFExporter/Private/Tasks/GLTFMaterialTasks.h
@@ -62,4 +62,7 @@ private:
 	FGLTFPropertyBakeOutput BakeMaterialProperty(EMaterialProperty Property, int32& OutTexCoord, const FIntPoint* PreferredTextureSize = nullptr, bool bCopyAlphaFromRedChannel = false) const;
 
 	bool StoreBakedPropertyTexture(FGLTFJsonTextureInfo& OutTexInfo, const FGLTFPropertyBakeOutput& PropertyBakeOutput, const FString& PropertyName) const;
+
+	FString GetMaterialName() const;
+	FString GetBakedTextureName(const FString& PropertyName) const;
 };
