[GLTFExporter] feat: only use mesh sections that reference the current material when baking with mesh data

commit 7c14ac83f8346564c45c699a177aec7e6dd647bc
Date: Fri, 23 Oct 2020 09:54:28 +0200

    feat: only use mesh sections that reference the current material when baking with mesh data
    (cherry picked from commit c2461efc3cedeff76cb68b827601e78553551fc4​)

#skipci

diff --git a/Engine/Plugins/Enterprise/GLTFExporter/MigratedFromGit/Source/GLTFExporter/Private/Converters/GLTFMaterialUtility.cpp b/Engine/Plugins/Enterprise/GLTFExporter/MigratedFromGit/Source/GLTFExporter/Private/Converters/GLTFMaterialUtility.cpp
--- a/Engine/Plugins/Enterprise/GLTFExporter/MigratedFromGit/Source/GLTFExporter/Private/Converters/GLTFMaterialUtility.cpp
+++ b/Engine/Plugins/Enterprise/GLTFExporter/MigratedFromGit/Source/GLTFExporter/Private/Converters/GLTFMaterialUtility.cpp
@@ -241,7 +241,7 @@ bool FGLTFMaterialUtility::CombineTextures(TArray<FColor>& OutPixels, const TArr
 	return bReadSuccessful;
 }
 
-FGLTFPropertyBakeOutput FGLTFMaterialUtility::BakeMaterialProperty(const FIntPoint& OutputSize, EMaterialProperty Property, const UMaterialInterface* Material, int32 TexCoord, FMeshDescription* MeshDescription, bool bCopyAlphaFromRedChannel)
+FGLTFPropertyBakeOutput FGLTFMaterialUtility::BakeMaterialProperty(const FIntPoint& OutputSize, EMaterialProperty Property, const UMaterialInterface* Material, int32 TexCoord, FMeshDescription* MeshDescription, TArray<int32>* MeshMaterialIndices, bool bCopyAlphaFromRedChannel)
 {
 	EMaterialProperty HackProperty = MP_MAX;
 	TTuple<int32, UMaterialExpression*> HackPrevState;
@@ -277,17 +277,9 @@ FGLTFPropertyBakeOutput FGLTFMaterialUtility::BakeMaterialProperty(const FIntPoi
 	MeshSet.TextureCoordinateIndex = TexCoord;
 	MeshSet.RawMeshDescription = MeshDescription;
 
-	if (MeshDescription != nullptr)
+	if (MeshMaterialIndices != nullptr)
 	{
-		// NOTE: we use polygon group ID's instead of material indices since that's what
-		// FMeshMaterialRenderItem::PopulateWithMeshData compares MeshSet.MaterialIndices against.
-		for (const auto& PolygonGroupID: MeshDescription->PolygonGroups().GetElementIDs())
-		{
-			// TODO: should we only add polygon groups that are actually using the material
-			// we're trying to bake? Or would that prevent us from baking materials using unrelated meshes?
-
-			MeshSet.MaterialIndices.Add(PolygonGroupID.GetValue());
-		}
+		MeshSet.MaterialIndices = *MeshMaterialIndices;
 	}
 
 	MeshSettings.Add(&MeshSet);
diff --git a/Engine/Plugins/Enterprise/GLTFExporter/MigratedFromGit/Source/GLTFExporter/Private/Converters/GLTFMaterialUtility.h b/Engine/Plugins/Enterprise/GLTFExporter/MigratedFromGit/Source/GLTFExporter/Private/Converters/GLTFMaterialUtility.h
--- a/Engine/Plugins/Enterprise/GLTFExporter/MigratedFromGit/Source/GLTFExporter/Private/Converters/GLTFMaterialUtility.h
+++ b/Engine/Plugins/Enterprise/GLTFExporter/MigratedFromGit/Source/GLTFExporter/Private/Converters/GLTFMaterialUtility.h
@@ -53,7 +53,7 @@ struct FGLTFMaterialUtility
 
 	static UTexture2D* CreateTransientTexture(const FGLTFPropertyBakeOutput& PropertyBakeOutput, bool bUseSRGB = false);
 	static bool CombineTextures(TArray<FColor>& OutPixels, const TArray<FGLTFTextureCombineSource>& Sources, const FIntPoint& OutputSize, EPixelFormat OutputPixelFormat);
-	static FGLTFPropertyBakeOutput BakeMaterialProperty(const FIntPoint& OutputSize, EMaterialProperty Property, const UMaterialInterface* Material, int32 TexCoord, FMeshDescription* MeshDescription, bool bCopyAlphaFromRedChannel);
+	static FGLTFPropertyBakeOutput BakeMaterialProperty(const FIntPoint& OutputSize, EMaterialProperty Property, const UMaterialInterface* Material, int32 TexCoord, FMeshDescription* MeshDescription = nullptr, TArray<int32>* MeshMaterialIndices = nullptr, bool bCopyAlphaFromRedChannel = false);
 
 	static FGLTFJsonTextureIndex AddCombinedTexture(FGLTFConvertBuilder& Builder, const TArray<FGLTFTextureCombineSource>& CombineSources, const FIntPoint& TextureSize, const FString& TextureName, EGLTFJsonTextureFilter MinFilter, EGLTFJsonTextureFilter MagFilter, EGLTFJsonTextureWrap WrapS, EGLTFJsonTextureWrap WrapT);
 	static FGLTFJsonTextureIndex AddTexture(FGLTFConvertBuilder& Builder, const TArray<FColor>& Pixels, const FIntPoint& TextureSize, const FString& TextureName, EGLTFJsonTextureFilter MinFilter, EGLTFJsonTextureFilter MagFilter, EGLTFJsonTextureWrap WrapS, EGLTFJsonTextureWrap WrapT);
diff --git a/Engine/Plugins/Enterprise/GLTFExporter/MigratedFromGit/Source/GLTFExporter/Private/Tasks/GLTFMaterialTasks.cpp b/Engine/Plugins/Enterprise/GLTFExporter/MigratedFromGit/Source/GLTFExporter/Private/Tasks/GLTFMaterialTasks.cpp
--- a/Engine/Plugins/Enterprise/GLTFExporter/MigratedFromGit/Source/GLTFExporter/Private/Tasks/GLTFMaterialTasks.cpp
+++ b/Engine/Plugins/Enterprise/GLTFExporter/MigratedFromGit/Source/GLTFExporter/Private/Tasks/GLTFMaterialTasks.cpp
@@ -1357,11 +1357,32 @@ FGLTFPropertyBakeOutput FGLTFMaterialTask::BakeMaterialProperty(EMaterialPropert
 	const FIntPoint TextureSize = PreferredTextureSize != nullptr ? *PreferredTextureSize : DefaultTextureSize;
 
 	FMeshDescription MeshDescription;
+	TArray<int32> MaterialIndices;
 
 	if (Mesh != nullptr)
 	{
 		FStaticMeshAttributes(MeshDescription).Register();
-		FMeshMergeHelpers::RetrieveMesh(const_cast<UStaticMesh*>(Mesh), LODIndex, MeshDescription);
+		FMeshMergeHelpers::RetrieveMesh(Mesh, LODIndex, MeshDescription);
+
+		TArray<FSectionInfo> Sections;
+		FMeshMergeHelpers::ExtractSections(Mesh, LODIndex, Sections);
+
+		// Add material indices of all sections that reference the current material
+		for (int32 SectionIndex = 0; SectionIndex < Sections.Num(); ++SectionIndex)
+		{
+			FSectionInfo& Section = Sections[SectionIndex];
+			if (Section.Material == Material)
+			{
+				MaterialIndices.Add(Section.MaterialIndex);
+			}
+		}
+
+		if (MaterialIndices.Num() == 0)
+		{
+			// Fall back to using the first material index,
+			// just like when baking from within the editor
+			MaterialIndices.Add(0);
+		}
 	}
 
 	// TODO: add support for calculating the ideal resolution to use for baking based on connected (texture) nodes
@@ -1372,6 +1393,7 @@ FGLTFPropertyBakeOutput FGLTFMaterialTask::BakeMaterialProperty(EMaterialPropert
 		Material,
 		OutTexCoord,
 		Mesh != nullptr ? &MeshDescription : nullptr,
+		Mesh != nullptr ? &MaterialIndices : nullptr,
 		bCopyAlphaFromRedChannel);
 
 	if (!PropertyBakeOutput.bIsConstant && TexCoords.Num() == 0)
