[GLTFExporter] chore: add TODOs for missing functionality when baking with mesh data

commit 9e191747babf952ba7b3001bb6d3bb3902f66381
Date: Fri, 23 Oct 2020 10:03:48 +0200

    chore: add TODOs for missing functionality when baking with mesh data
    (cherry picked from commit a05843a3f3593dfa32c97ef4d2ea6e3411559ceb​)

#skipci

diff --git a/Engine/Plugins/Enterprise/GLTFExporter/MigratedFromGit/Source/GLTFExporter/Private/Tasks/GLTFMaterialTasks.cpp b/Engine/Plugins/Enterprise/GLTFExporter/MigratedFromGit/Source/GLTFExporter/Private/Tasks/GLTFMaterialTasks.cpp
--- a/Engine/Plugins/Enterprise/GLTFExporter/MigratedFromGit/Source/GLTFExporter/Private/Tasks/GLTFMaterialTasks.cpp
+++ b/Engine/Plugins/Enterprise/GLTFExporter/MigratedFromGit/Source/GLTFExporter/Private/Tasks/GLTFMaterialTasks.cpp
@@ -1361,6 +1361,12 @@ FGLTFPropertyBakeOutput FGLTFMaterialTask::BakeMaterialProperty(EMaterialPropert
 
 	if (Mesh != nullptr)
 	{
+		// TODO: add support for override materials
+
+		// TODO: add support for override vertex colors
+
+		// TODO: add support for baking world-position, -normals and other inputs that depend on component transform
+
 		FStaticMeshAttributes(MeshDescription).Register();
 		FMeshMergeHelpers::RetrieveMesh(Mesh, LODIndex, MeshDescription);
 
