[GLTFExporter] feat: add support for using mesh data from mesh-components when baking materials

commit 088fbeb03ec344c9613055e7be0e9585609d1094
Date: Mon, 26 Oct 2020 11:22:16 +0100

    feat: add support for using mesh data from mesh-components when baking materials
    (cherry picked from commit 1e556213399922abd7261641ab98e28ab4ebb1c7​)

#skipci

diff --git a/Engine/Plugins/Enterprise/GLTFExporter/MigratedFromGit/Source/GLTFExporter/Private/Builders/GLTFConvertBuilder.cpp b/Engine/Plugins/Enterprise/GLTFExporter/MigratedFromGit/Source/GLTFExporter/Private/Builders/GLTFConvertBuilder.cpp
--- a/Engine/Plugins/Enterprise/GLTFExporter/MigratedFromGit/Source/GLTFExporter/Private/Builders/GLTFConvertBuilder.cpp
+++ b/Engine/Plugins/Enterprise/GLTFExporter/MigratedFromGit/Source/GLTFExporter/Private/Builders/GLTFConvertBuilder.cpp
@@ -88,65 +88,88 @@ FGLTFJsonAccessorIndex FGLTFConvertBuilder::GetOrAddIndexAccessor(const FGLTFMes
 	return IndexBufferConverter.GetOrAdd(MeshSection);
 }
 
-FGLTFJsonMeshIndex FGLTFConvertBuilder::GetOrAddMesh(const UStaticMesh* StaticMesh, int32 LODIndex, const FColorVertexBuffer* OverrideVertexColors, const FGLTFMaterialArray& OverrideMaterials)
+FGLTFJsonMeshIndex FGLTFConvertBuilder::GetOrAddMesh(const UStaticMesh* StaticMesh, int32 LODIndex, const FGLTFMaterialArray& OverrideMaterials)
 {
 	if (StaticMesh == nullptr)
 	{
 		return FGLTFJsonMeshIndex(INDEX_NONE);
 	}
 
-	return StaticMeshConverter.GetOrAdd(StaticMesh, LODIndex, OverrideVertexColors, OverrideMaterials);
+	return StaticMeshConverter.GetOrAdd(StaticMesh, nullptr, LODIndex, OverrideMaterials);
 }
 
-FGLTFJsonMeshIndex FGLTFConvertBuilder::GetOrAddMesh(const UStaticMeshComponent* StaticMeshComponent)
+FGLTFJsonMeshIndex FGLTFConvertBuilder::GetOrAddMesh(const UStaticMeshComponent* StaticMeshComponent, const FGLTFMaterialArray& OverrideMaterials)
 {
 	if (StaticMeshComponent == nullptr)
 	{
 		return FGLTFJsonMeshIndex(INDEX_NONE);
 	}
 
-	const UStaticMesh* StaticMesh = StaticMeshComponent->GetStaticMesh();
 	const int32 LODIndex = StaticMeshComponent->ForcedLodModel > 0 ? StaticMeshComponent->ForcedLodModel - 1 : FMath::Max(ExportOptions->DefaultLevelOfDetail, StaticMeshComponent->MinLOD);
-	const FColorVertexBuffer* OverrideVertexColors = StaticMeshComponent->LODData.IsValidIndex(LODIndex) ? StaticMeshComponent->LODData[LODIndex].OverrideVertexColors : nullptr;
-	const FGLTFMaterialArray OverrideMaterials = FGLTFMaterialArray(StaticMeshComponent->OverrideMaterials);
 
-	return GetOrAddMesh(StaticMesh, LODIndex, OverrideVertexColors, OverrideMaterials);
+	return StaticMeshConverter.GetOrAdd(nullptr, StaticMeshComponent, LODIndex, OverrideMaterials);
 }
 
-FGLTFJsonMeshIndex FGLTFConvertBuilder::GetOrAddMesh(const USkeletalMesh* SkeletalMesh, int32 LODIndex, const FColorVertexBuffer* OverrideVertexColors, const FSkinWeightVertexBuffer* OverrideSkinWeights, const FGLTFMaterialArray& OverrideMaterials)
+FGLTFJsonMeshIndex FGLTFConvertBuilder::GetOrAddMesh(const USkeletalMesh* SkeletalMesh, int32 LODIndex, const FGLTFMaterialArray& OverrideMaterials)
 {
 	if (SkeletalMesh == nullptr)
 	{
 		return FGLTFJsonMeshIndex(INDEX_NONE);
 	}
 
-	return SkeletalMeshConverter.GetOrAdd(SkeletalMesh, LODIndex, OverrideVertexColors, OverrideSkinWeights, OverrideMaterials);
+	return SkeletalMeshConverter.GetOrAdd(SkeletalMesh, nullptr, LODIndex, OverrideMaterials);
 }
 
-FGLTFJsonMeshIndex FGLTFConvertBuilder::GetOrAddMesh(const USkeletalMeshComponent* SkeletalMeshComponent)
+FGLTFJsonMeshIndex FGLTFConvertBuilder::GetOrAddMesh(const USkeletalMeshComponent* SkeletalMeshComponent, const FGLTFMaterialArray& OverrideMaterials)
 {
 	if (SkeletalMeshComponent == nullptr)
 	{
 		return FGLTFJsonMeshIndex(INDEX_NONE);
 	}
 
-	const USkeletalMesh* SkeletalMesh = SkeletalMeshComponent->SkeletalMesh;
 	const int32 LODIndex = SkeletalMeshComponent->GetForcedLOD() > 0 ? SkeletalMeshComponent->GetForcedLOD() - 1 : FMath::Max(ExportOptions->DefaultLevelOfDetail, SkeletalMeshComponent->MinLodModel);
-	const FColorVertexBuffer* OverrideVertexColors = SkeletalMeshComponent->LODInfo.IsValidIndex(LODIndex) ? SkeletalMeshComponent->LODInfo[LODIndex].OverrideVertexColors : nullptr;
-	const FSkinWeightVertexBuffer* OverrideSkinWeights = SkeletalMeshComponent->LODInfo.IsValidIndex(LODIndex) ? SkeletalMeshComponent->LODInfo[LODIndex].OverrideSkinWeights : nullptr;
-	const FGLTFMaterialArray OverrideMaterials = FGLTFMaterialArray(SkeletalMeshComponent->OverrideMaterials);
 
-	return GetOrAddMesh(SkeletalMesh, LODIndex, OverrideVertexColors, OverrideSkinWeights, OverrideMaterials);
+	return SkeletalMeshConverter.GetOrAdd(nullptr, SkeletalMeshComponent, LODIndex, OverrideMaterials);
+}
+
+FGLTFJsonMaterialIndex FGLTFConvertBuilder::GetOrAddMaterial(const UMaterialInterface* Material)
+{
+	if (Material == nullptr)
+	{
+		return FGLTFJsonMaterialIndex(INDEX_NONE);
+	}
+
+	return MaterialConverter.GetOrAdd(Material, nullptr, -1, {});
+}
+
+FGLTFJsonMaterialIndex FGLTFConvertBuilder::GetOrAddMaterial(const UMaterialInterface* Material, const UStaticMesh* Mesh, int32 LODIndex, const FGLTFMaterialArray& OverrideMaterials)
+{
+	if (Material == nullptr)
+	{
+		return FGLTFJsonMaterialIndex(INDEX_NONE);
+	}
+
+	return MaterialConverter.GetOrAdd(Material, Mesh, LODIndex, OverrideMaterials);
+}
+
+FGLTFJsonMaterialIndex FGLTFConvertBuilder::GetOrAddMaterial(const UMaterialInterface* Material, const USkeletalMesh* Mesh, int32 LODIndex, const FGLTFMaterialArray& OverrideMaterials)
+{
+	if (Material == nullptr)
+	{
+		return FGLTFJsonMaterialIndex(INDEX_NONE);
+	}
+
+	return MaterialConverter.GetOrAdd(Material, Mesh, LODIndex, OverrideMaterials);
 }
 
-FGLTFJsonMaterialIndex FGLTFConvertBuilder::GetOrAddMaterial(const UMaterialInterface* Material, const UStaticMesh* Mesh, int32 LODIndex)
+FGLTFJsonMaterialIndex FGLTFConvertBuilder::GetOrAddMaterial(const UMaterialInterface* Material, const UMeshComponent* MeshComponent, int32 LODIndex, const FGLTFMaterialArray& OverrideMaterials)
 {
 	if (Material == nullptr)
 	{
 		return FGLTFJsonMaterialIndex(INDEX_NONE);
 	}
 
-	return MaterialConverter.GetOrAdd(Material, Mesh, LODIndex);
+	return MaterialConverter.GetOrAdd(Material, MeshComponent, LODIndex, OverrideMaterials);
 }
 
 FGLTFJsonSamplerIndex FGLTFConvertBuilder::GetOrAddSampler(const UTexture* Texture)
diff --git a/Engine/Plugins/Enterprise/GLTFExporter/MigratedFromGit/Source/GLTFExporter/Private/Builders/GLTFConvertBuilder.h b/Engine/Plugins/Enterprise/GLTFExporter/MigratedFromGit/Source/GLTFExporter/Private/Builders/GLTFConvertBuilder.h
--- a/Engine/Plugins/Enterprise/GLTFExporter/MigratedFromGit/Source/GLTFExporter/Private/Builders/GLTFConvertBuilder.h
+++ b/Engine/Plugins/Enterprise/GLTFExporter/MigratedFromGit/Source/GLTFExporter/Private/Builders/GLTFConvertBuilder.h
@@ -39,12 +39,16 @@ public:
 	FGLTFJsonAccessorIndex GetOrAddWeightAccessor(const FGLTFMeshSection* MeshSection, const FSkinWeightVertexBuffer* VertexBuffer, int32 InfluenceOffset);
 	FGLTFJsonAccessorIndex GetOrAddIndexAccessor(const FGLTFMeshSection* MeshSection);
 
-	FGLTFJsonMeshIndex GetOrAddMesh(const UStaticMesh* StaticMesh, int32 LODIndex = -1, const FColorVertexBuffer* OverrideVertexColors = nullptr, const FGLTFMaterialArray& OverrideMaterials = {});
-	FGLTFJsonMeshIndex GetOrAddMesh(const UStaticMeshComponent* StaticMeshComponent);
-	FGLTFJsonMeshIndex GetOrAddMesh(const USkeletalMesh* SkeletalMesh, int32 LODIndex = -1, const FColorVertexBuffer* OverrideVertexColors = nullptr, const FSkinWeightVertexBuffer* OverrideSkinWeights = nullptr, const FGLTFMaterialArray& OverrideMaterials = {});
-	FGLTFJsonMeshIndex GetOrAddMesh(const USkeletalMeshComponent* SkeletalMeshComponent);
+	FGLTFJsonMeshIndex GetOrAddMesh(const UStaticMesh* StaticMesh, int32 LODIndex = -1, const FGLTFMaterialArray& OverrideMaterials = {});
+	FGLTFJsonMeshIndex GetOrAddMesh(const UStaticMeshComponent* StaticMeshComponent, const FGLTFMaterialArray& OverrideMaterials = {});
+	FGLTFJsonMeshIndex GetOrAddMesh(const USkeletalMesh* SkeletalMesh, int32 LODIndex = -1, const FGLTFMaterialArray& OverrideMaterials = {});
+	FGLTFJsonMeshIndex GetOrAddMesh(const USkeletalMeshComponent* SkeletalMeshComponent, const FGLTFMaterialArray& OverrideMaterials = {});
+
+	FGLTFJsonMaterialIndex GetOrAddMaterial(const UMaterialInterface* Material);
+	FGLTFJsonMaterialIndex GetOrAddMaterial(const UMaterialInterface* Material, const UStaticMesh* Mesh, int32 LODIndex = -1, const FGLTFMaterialArray& OverrideMaterials = {});
+	FGLTFJsonMaterialIndex GetOrAddMaterial(const UMaterialInterface* Material, const USkeletalMesh* Mesh, int32 LODIndex = -1, const FGLTFMaterialArray& OverrideMaterials = {});
+	FGLTFJsonMaterialIndex GetOrAddMaterial(const UMaterialInterface* Material, const UMeshComponent* MeshComponent, int32 LODIndex = -1, const FGLTFMaterialArray& OverrideMaterials = {});
 
-	FGLTFJsonMaterialIndex GetOrAddMaterial(const UMaterialInterface* Material, const UStaticMesh* Mesh, int32 LODIndex = -1);
 	FGLTFJsonSamplerIndex GetOrAddSampler(const UTexture* Texture);
 	FGLTFJsonTextureIndex GetOrAddTexture(const UTexture2D* Texture);
 	FGLTFJsonTextureIndex GetOrAddTexture(const UTextureCube* Texture, ECubeFace CubeFace);
diff --git a/Engine/Plugins/Enterprise/GLTFExporter/MigratedFromGit/Source/GLTFExporter/Private/Converters/GLTFBackdropConverters.cpp b/Engine/Plugins/Enterprise/GLTFExporter/MigratedFromGit/Source/GLTFExporter/Private/Converters/GLTFBackdropConverters.cpp
--- a/Engine/Plugins/Enterprise/GLTFExporter/MigratedFromGit/Source/GLTFExporter/Private/Converters/GLTFBackdropConverters.cpp
+++ b/Engine/Plugins/Enterprise/GLTFExporter/MigratedFromGit/Source/GLTFExporter/Private/Converters/GLTFBackdropConverters.cpp
@@ -21,7 +21,7 @@ FGLTFJsonBackdropIndex FGLTFBackdropConverter::Convert(const AActor* BackdropAct
 	if (FGLTFActorUtility::TryGetPropertyValue(BackdropActor, TEXT("Mesh"), Mesh))
 	{
 		// TODO: avoid exporting the mesh material (use gltf default material)
-		JsonBackdrop.Mesh = Builder.GetOrAddMesh(Mesh, -1, nullptr, { FGLTFMaterialUtility::GetDefault() });
+		JsonBackdrop.Mesh = Builder.GetOrAddMesh(Mesh, -1, { FGLTFMaterialUtility::GetDefault() });
 	}
 	else
 	{
diff --git a/Engine/Plugins/Enterprise/GLTFExporter/MigratedFromGit/Source/GLTFExporter/Private/Converters/GLTFMaterialArray.h b/Engine/Plugins/Enterprise/GLTFExporter/MigratedFromGit/Source/GLTFExporter/Private/Converters/GLTFMaterialArray.h
--- a/Engine/Plugins/Enterprise/GLTFExporter/MigratedFromGit/Source/GLTFExporter/Private/Converters/GLTFMaterialArray.h
+++ b/Engine/Plugins/Enterprise/GLTFExporter/MigratedFromGit/Source/GLTFExporter/Private/Converters/GLTFMaterialArray.h
@@ -46,6 +46,41 @@ struct FGLTFMaterialArray : TArray<const UMaterialInterface*>
 		return GetOverride<FSkeletalMaterial>(Originals, Index);
 	}
 
+	const UMaterialInterface* GetOverride(const TArray<UMaterialInterface*>& Originals, int32 Index) const
+	{
+		return GetOverride<UMaterialInterface*>(Originals, Index);
+	}
+
+	TArray<UMaterialInterface*> GetOverrides(const TArray<FStaticMaterial>& Originals) const
+	{
+		return GetOverrides<FStaticMaterial>(Originals);
+	}
+
+	TArray<UMaterialInterface*> GetOverrides(const TArray<FSkeletalMaterial>& Originals) const
+	{
+		return GetOverrides<FSkeletalMaterial>(Originals);
+	}
+
+	TArray<UMaterialInterface*> GetOverrides(const TArray<UMaterialInterface*>& Originals) const
+	{
+		return GetOverrides<UMaterialInterface*>(Originals);
+	}
+
+	void ClearRedundantOverrides(const TArray<FStaticMaterial>& Originals)
+	{
+		ClearRedundantOverrides<FStaticMaterial>(Originals);
+	}
+
+	void ClearRedundantOverrides(const TArray<FSkeletalMaterial>& Originals)
+	{
+		ClearRedundantOverrides<FSkeletalMaterial>(Originals);
+	}
+
+	void ClearRedundantOverrides(const TArray<UMaterialInterface*>& Originals)
+	{
+		ClearRedundantOverrides<UMaterialInterface*>(Originals);
+	}
+
 	using TArray::operator=;
 	using TArray::operator==;
 	using TArray::operator!=;
@@ -70,6 +105,16 @@ struct FGLTFMaterialArray : TArray<const UMaterialInterface*>
 		return !Equals(Other);
 	}
 
+	bool operator==(const TArray<UMaterialInterface*>& Other) const
+	{
+		return Equals(Other);
+	}
+
+	bool operator!=(const TArray<UMaterialInterface*>& Other) const
+	{
+		return !Equals(Other);
+	}
+
 	friend uint32 GetTypeHash(const TArray& Array)
 	{
 		uint32 Hash = GetTypeHash(Array.Num());
@@ -92,12 +137,42 @@ private:
 
 		if (Originals.IsValidIndex(Index))
 		{
-			return Originals[Index].MaterialInterface;
+			return GetInterface(Originals[Index]);
 		}
 
 		return nullptr;
 	}
 
+	template <typename MeshMaterialType>
+	TArray<UMaterialInterface*> GetOverrides(const TArray<MeshMaterialType>& Originals) const
+	{
+		const int32 Count = Originals.Num();
+		TArray<UMaterialInterface*> Materials;
+		Materials.AddDefaulted(Count);
+
+		for (int32 Index = 0; Index < Count; ++Index)
+		{
+			Materials[Index] = const_cast<UMaterialInterface*>(GetOverride(Originals, Index));
+		}
+
+		return Materials;
+	}
+
+	template <typename MeshMaterialType>
+	void ClearRedundantOverrides(const TArray<MeshMaterialType>& Originals)
+	{
+		const int32 Count = Originals.Num();
+		SetNumZeroed(Count);
+
+		for (int32 Index = 0; Index < Count; ++Index)
+		{
+			if (GetInterface(Originals[Index]) == (*this)[Index])
+			{
+				(*this)[Index] = nullptr;
+			}
+		}
+	}
+
 	template <typename MeshMaterialType>
 	bool Equals(const TArray<MeshMaterialType>& Other) const
 	{
@@ -108,7 +183,7 @@ private:
 
 		for (int32 Index = 0; Index < Num(); ++Index)
 		{
-			if (Other[Index].MaterialInterface != (*this)[Index])
+			if (!Equals(Other[Index], Index))
 			{
 				return false;
 			}
@@ -116,4 +191,25 @@ private:
 
 		return true;
 	}
+
+	static const UMaterialInterface* GetInterface(const UMaterialInterface* Material)
+	{
+		return Material;
+	}
+
+	static const UMaterialInterface* GetInterface(const FStaticMaterial& Material)
+	{
+		return Material.MaterialInterface;
+	}
+
+	static const UMaterialInterface* GetInterface(const FSkeletalMaterial& Material)
+	{
+		return Material.MaterialInterface;
+	}
+
+	template <typename MeshMaterialType>
+	bool Equals(const MeshMaterialType& Other, int32 Index) const
+	{
+		return IsValidIndex(Index) && GetInterface(Other) == (*this)[Index];
+	}
 };
diff --git a/Engine/Plugins/Enterprise/GLTFExporter/MigratedFromGit/Source/GLTFExporter/Private/Converters/GLTFMaterialConverters.cpp b/Engine/Plugins/Enterprise/GLTFExporter/MigratedFromGit/Source/GLTFExporter/Private/Converters/GLTFMaterialConverters.cpp
--- a/Engine/Plugins/Enterprise/GLTFExporter/MigratedFromGit/Source/GLTFExporter/Private/Converters/GLTFMaterialConverters.cpp
+++ b/Engine/Plugins/Enterprise/GLTFExporter/MigratedFromGit/Source/GLTFExporter/Private/Converters/GLTFMaterialConverters.cpp
@@ -6,28 +6,49 @@
 #include "Builders/GLTFConvertBuilder.h"
 #include "Tasks/GLTFMaterialTasks.h"
 
-void FGLTFMaterialConverter::Sanitize(const UMaterialInterface*& Material, const UStaticMesh*& Mesh, int32& LODIndex)
+void FGLTFMaterialConverter::Sanitize(const UMaterialInterface*& Material, const UObject*& MeshOrComponent, int32& LODIndex, FGLTFMaterialArray& OverrideMaterials)
 {
-	// TODO: add support for UStaticMeshComponent, USkeletalMesh and USkeletalMeshComponent too
+	const UStaticMesh* StaticMesh = Cast<UStaticMesh>(MeshOrComponent);
+	const USkeletalMesh* SkeletalMesh = Cast<USkeletalMesh>(MeshOrComponent);
 
-	if (Mesh != nullptr &&
+	if (MeshOrComponent != nullptr && StaticMesh == nullptr && SkeletalMesh == nullptr)
+	{
+		if (const UStaticMeshComponent* StaticMeshComponent = Cast<UStaticMeshComponent>(MeshOrComponent))
+		{
+			StaticMesh = StaticMeshComponent->GetStaticMesh();
+		}
+		else if (const USkeletalMeshComponent* SkeletalMeshComponent = Cast<USkeletalMeshComponent>(MeshOrComponent))
+		{
+			SkeletalMesh = SkeletalMeshComponent->SkeletalMesh;
+		}
+	}
+
+	if ((StaticMesh != nullptr || SkeletalMesh != nullptr) &&
 		Builder.ExportOptions->bBakeMaterialInputs &&
 		Builder.ExportOptions->bBakeMaterialInputsUsingMeshData &&
-		MaterialNeedsVertexData(const_cast<UMaterialInterface*>(Material)))
+		FGLTFMaterialUtility::MaterialNeedsVertexData(Material))
 	{
 		if (LODIndex < 0)
 		{
-			LODIndex = FMath::Max(Builder.ExportOptions->DefaultLevelOfDetail, FGLTFMeshUtility::GetMinimumLOD(Mesh));
+			if (StaticMesh != nullptr)
+			{
+				LODIndex = FMath::Max(Builder.ExportOptions->DefaultLevelOfDetail, FGLTFMeshUtility::GetMinimumLOD(StaticMesh));
+			}
+			else if (SkeletalMesh != nullptr)
+			{
+				LODIndex = FMath::Max(Builder.ExportOptions->DefaultLevelOfDetail, FGLTFMeshUtility::GetMinimumLOD(SkeletalMesh));
+			}
 		}
 	}
 	else
 	{
-		Mesh = nullptr;
+		MeshOrComponent = nullptr;
 		LODIndex = -1;
+		OverrideMaterials = {};
 	}
 }
 
-FGLTFJsonMaterialIndex FGLTFMaterialConverter::Convert(const UMaterialInterface* Material, const UStaticMesh* Mesh, int32 LODIndex)
+FGLTFJsonMaterialIndex FGLTFMaterialConverter::Convert(const UMaterialInterface* Material, const UObject* MeshOrComponent, int32 LODIndex, const FGLTFMaterialArray OverrideMaterials)
 {
 	if (Material == FGLTFMaterialUtility::GetDefault())
 	{
@@ -35,39 +56,6 @@ FGLTFJsonMaterialIndex FGLTFMaterialConverter::Convert(const UMaterialInterface*
 	}
 
 	const FGLTFJsonMaterialIndex MaterialIndex = Builder.AddMaterial();
-	Builder.SetupTask<FGLTFMaterialTask>(Builder, Material, Mesh, LODIndex, MaterialIndex);
+	Builder.SetupTask<FGLTFMaterialTask>(Builder, Material, MeshOrComponent, LODIndex, OverrideMaterials, MaterialIndex);
 	return MaterialIndex;
 }
-
-bool FGLTFMaterialConverter::MaterialNeedsVertexData(UMaterialInterface* Material)
-{
-	// TODO: only analyze properties that will be needed for this specific material
-
-	const TArray<EMaterialProperty> Properties =
-	{
-		MP_BaseColor,
-		MP_EmissiveColor,
-		MP_Opacity,
-		MP_OpacityMask,
-		MP_Metallic,
-		MP_Roughness,
-		MP_Normal,
-		// MP_CustomOutput,	// NOTE: causes a crash when used with AnalyzeMaterialProperty
-		MP_AmbientOcclusion,
-		MP_CustomData0,
-		MP_CustomData1
-	};
-
-	bool bRequiresVertexData = false;
-
-	int32 NumTextureCoordinates;
-	bool bPropertyRequiresVertexData;
-
-	for (const EMaterialProperty Property: Properties)
-	{
-		const_cast<UMaterialInterface*>(Material)->AnalyzeMaterialProperty(Property, NumTextureCoordinates, bPropertyRequiresVertexData);
-		bRequiresVertexData |= bPropertyRequiresVertexData;
-	}
-
-	return bRequiresVertexData;
-}
diff --git a/Engine/Plugins/Enterprise/GLTFExporter/MigratedFromGit/Source/GLTFExporter/Private/Converters/GLTFMaterialConverters.h b/Engine/Plugins/Enterprise/GLTFExporter/MigratedFromGit/Source/GLTFExporter/Private/Converters/GLTFMaterialConverters.h
--- a/Engine/Plugins/Enterprise/GLTFExporter/MigratedFromGit/Source/GLTFExporter/Private/Converters/GLTFMaterialConverters.h
+++ b/Engine/Plugins/Enterprise/GLTFExporter/MigratedFromGit/Source/GLTFExporter/Private/Converters/GLTFMaterialConverters.h
@@ -5,15 +5,14 @@
 #include "Json/GLTFJsonIndex.h"
 #include "Converters/GLTFConverter.h"
 #include "Converters/GLTFBuilderContext.h"
+#include "Converters/GLTFMaterialArray.h"
 #include "Engine.h"
 
-class FGLTFMaterialConverter final : public FGLTFBuilderContext, public TGLTFConverter<FGLTFJsonMaterialIndex, const UMaterialInterface*, const UStaticMesh*, int32>
+class FGLTFMaterialConverter final : public FGLTFBuilderContext, public TGLTFConverter<FGLTFJsonMaterialIndex, const UMaterialInterface*, const UObject*, int32, FGLTFMaterialArray>
 {
 	using FGLTFBuilderContext::FGLTFBuilderContext;
 
-	virtual void Sanitize(const UMaterialInterface*& Material, const UStaticMesh*& Mesh, int32& LODIndex) override;
+	virtual void Sanitize(const UMaterialInterface*& Material, const UObject*& MeshOrComponent, int32& LODIndex, FGLTFMaterialArray& OverrideMaterials) override;
 
-	virtual FGLTFJsonMaterialIndex Convert(const UMaterialInterface* Material, const UStaticMesh* Mesh, int32 LODIndex) override;
-
-	static bool MaterialNeedsVertexData(UMaterialInterface* Material);
+	virtual FGLTFJsonMaterialIndex Convert(const UMaterialInterface* Material, const UObject* MeshOrComponent, int32 LODIndex, const FGLTFMaterialArray OverrideMaterials) override;
 };
diff --git a/Engine/Plugins/Enterprise/GLTFExporter/MigratedFromGit/Source/GLTFExporter/Private/Converters/GLTFMaterialUtility.cpp b/Engine/Plugins/Enterprise/GLTFExporter/MigratedFromGit/Source/GLTFExporter/Private/Converters/GLTFMaterialUtility.cpp
--- a/Engine/Plugins/Enterprise/GLTFExporter/MigratedFromGit/Source/GLTFExporter/Private/Converters/GLTFMaterialUtility.cpp
+++ b/Engine/Plugins/Enterprise/GLTFExporter/MigratedFromGit/Source/GLTFExporter/Private/Converters/GLTFMaterialUtility.cpp
@@ -508,3 +508,54 @@ FString FGLTFMaterialUtility::ShadingModelsToString(const FMaterialShadingModelF
 
 	return Result;
 }
+
+bool FGLTFMaterialUtility::MaterialNeedsVertexData(const UMaterialInterface* Material)
+{
+	if (Material == nullptr)
+	{
+		return false;
+	}
+
+	// TODO: only analyze properties that will be needed for this specific material
+
+	const TArray<EMaterialProperty> Properties =
+	{
+		MP_BaseColor,
+		MP_EmissiveColor,
+		MP_Opacity,
+		MP_OpacityMask,
+		MP_Metallic,
+		MP_Roughness,
+		MP_Normal,
+		// MP_CustomOutput,	// NOTE: causes a crash when used with AnalyzeMaterialProperty
+		MP_AmbientOcclusion,
+		MP_CustomData0,
+		MP_CustomData1
+	};
+
+	bool bRequiresVertexData = false;
+
+	int32 NumTextureCoordinates;
+	bool bPropertyRequiresVertexData;
+
+	for (const EMaterialProperty Property: Properties)
+	{
+		const_cast<UMaterialInterface*>(Material)->AnalyzeMaterialProperty(Property, NumTextureCoordinates, bPropertyRequiresVertexData);
+		bRequiresVertexData |= bPropertyRequiresVertexData;
+	}
+
+	return bRequiresVertexData;
+}
+
+bool FGLTFMaterialUtility::MaterialsNeedVertexData(const TArray<UMaterialInterface*> Materials)
+{
+	for (const UMaterialInterface* Material: Materials)
+	{
+		if (MaterialNeedsVertexData(Material))
+		{
+			return true;
+		}
+	}
+
+	return false;
+}
diff --git a/Engine/Plugins/Enterprise/GLTFExporter/MigratedFromGit/Source/GLTFExporter/Private/Converters/GLTFMaterialUtility.h b/Engine/Plugins/Enterprise/GLTFExporter/MigratedFromGit/Source/GLTFExporter/Private/Converters/GLTFMaterialUtility.h
--- a/Engine/Plugins/Enterprise/GLTFExporter/MigratedFromGit/Source/GLTFExporter/Private/Converters/GLTFMaterialUtility.h
+++ b/Engine/Plugins/Enterprise/GLTFExporter/MigratedFromGit/Source/GLTFExporter/Private/Converters/GLTFMaterialUtility.h
@@ -67,4 +67,7 @@ struct FGLTFMaterialUtility
 	static void ExpandAllFunctionExpressions(TArray<UMaterialExpression*>& InOutExpressions);
 
 	static FString ShadingModelsToString(const FMaterialShadingModelField& ShadingModels);
+
+	static bool MaterialNeedsVertexData(const UMaterialInterface* Material);
+	static bool MaterialsNeedVertexData(const TArray<UMaterialInterface*> Materials);
 };
diff --git a/Engine/Plugins/Enterprise/GLTFExporter/MigratedFromGit/Source/GLTFExporter/Private/Converters/GLTFMeshConverters.cpp b/Engine/Plugins/Enterprise/GLTFExporter/MigratedFromGit/Source/GLTFExporter/Private/Converters/GLTFMeshConverters.cpp
--- a/Engine/Plugins/Enterprise/GLTFExporter/MigratedFromGit/Source/GLTFExporter/Private/Converters/GLTFMeshConverters.cpp
+++ b/Engine/Plugins/Enterprise/GLTFExporter/MigratedFromGit/Source/GLTFExporter/Private/Converters/GLTFMeshConverters.cpp
@@ -1,84 +1,111 @@
 // Copyright Epic Games, Inc. All Rights Reserved.
 
 #include "Converters/GLTFMeshConverters.h"
+#include "Converters/GLTFMaterialUtility.h"
 #include "Converters/GLTFMeshUtility.h"
 #include "Builders/GLTFConvertBuilder.h"
 #include "Tasks/GLTFMeshTasks.h"
 #include "Rendering/SkeletalMeshRenderData.h"
 
-void FGLTFStaticMeshConverter::Sanitize(const UStaticMesh*& StaticMesh, int32& LODIndex, const FColorVertexBuffer*& OverrideVertexColors, FGLTFMaterialArray& OverrideMaterials)
+void FGLTFStaticMeshConverter::Sanitize(const UStaticMesh*& StaticMesh, const UStaticMeshComponent*& StaticMeshComponent, int32& LODIndex, FGLTFMaterialArray& OverrideMaterials)
 {
-	if (LODIndex < 0)
+	const uint32 MeshCount = (StaticMeshComponent != nullptr ? 1 : 0) + (StaticMesh != nullptr ? 1 : 0);
+	check(MeshCount == 1);
+
+	const TArray<UMaterialInterface*> Materials = StaticMeshComponent != nullptr
+		? OverrideMaterials.GetOverrides(StaticMeshComponent->GetMaterials())
+		: OverrideMaterials.GetOverrides(StaticMesh->StaticMaterials);
+
+	if (StaticMeshComponent != nullptr)
 	{
-		LODIndex = FMath::Max(Builder.ExportOptions->DefaultLevelOfDetail, FGLTFMeshUtility::GetMinimumLOD(StaticMesh));
+		// Only use the component if it's needed for baking, since we would
+		// otherwise export a copy of this mesh for each mesh-component.
+		if (!Builder.ExportOptions->bBakeMaterialInputs ||
+			!Builder.ExportOptions->bBakeMaterialInputsUsingMeshData ||
+			!FGLTFMaterialUtility::MaterialsNeedVertexData(Materials))
+		{
+			StaticMesh = StaticMeshComponent->GetStaticMesh();
+			StaticMeshComponent = nullptr;
+		}
 	}
 
-	if (!Builder.ExportOptions->bExportVertexColors)
+	// Clean up override materials to only contain items that differ from the original materials
+	OverrideMaterials = FGLTFMaterialArray(Materials);
+
+	if (StaticMeshComponent != nullptr)
 	{
-		OverrideVertexColors = nullptr;
+		OverrideMaterials.ClearRedundantOverrides(StaticMeshComponent->GetMaterials());
 	}
-
-	if (OverrideMaterials == StaticMesh->StaticMaterials)
+	else
 	{
-		OverrideMaterials.Empty();
+		OverrideMaterials.ClearRedundantOverrides(StaticMesh->StaticMaterials);
 	}
 
-	LODIndex = FMath::Clamp(LODIndex, 0, StaticMesh->GetNumLODs() - 1);
-	const FStaticMeshLODResources& MeshLOD = StaticMesh->GetLODForExport(LODIndex);
+	const UStaticMesh* Mesh = StaticMeshComponent != nullptr ? StaticMeshComponent->GetStaticMesh() : StaticMesh;
 
-	if (OverrideVertexColors == &MeshLOD.VertexBuffers.ColorVertexBuffer)
+	if (LODIndex < 0)
 	{
-		OverrideVertexColors = nullptr;
+		LODIndex = FMath::Max(Builder.ExportOptions->DefaultLevelOfDetail, FGLTFMeshUtility::GetMinimumLOD(Mesh));
 	}
+
+	LODIndex = FMath::Clamp(LODIndex, 0, Mesh->GetNumLODs() - 1);
 }
 
-FGLTFJsonMeshIndex FGLTFStaticMeshConverter::Convert(const UStaticMesh* StaticMesh, int32 LODIndex, const FColorVertexBuffer* OverrideVertexColors, FGLTFMaterialArray OverrideMaterials)
+FGLTFJsonMeshIndex FGLTFStaticMeshConverter::Convert(const UStaticMesh* StaticMesh, const UStaticMeshComponent* StaticMeshComponent, int32 LODIndex, FGLTFMaterialArray Materials)
 {
 	const FGLTFJsonMeshIndex MeshIndex = Builder.AddMesh();
-	Builder.SetupTask<FGLTFStaticMeshTask>(Builder, MeshSectionConverter, StaticMesh, LODIndex, OverrideVertexColors, OverrideMaterials, MeshIndex);
+	Builder.SetupTask<FGLTFStaticMeshTask>(Builder, MeshSectionConverter, StaticMesh, StaticMeshComponent, LODIndex, Materials, MeshIndex);
 	return MeshIndex;
 }
 
-void FGLTFSkeletalMeshConverter::Sanitize(const USkeletalMesh*& SkeletalMesh, int32& LODIndex, const FColorVertexBuffer*& OverrideVertexColors, const FSkinWeightVertexBuffer*& OverrideSkinWeights, FGLTFMaterialArray& OverrideMaterials)
+void FGLTFSkeletalMeshConverter::Sanitize(const USkeletalMesh*& SkeletalMesh, const USkeletalMeshComponent*& SkeletalMeshComponent, int32& LODIndex, FGLTFMaterialArray& OverrideMaterials)
 {
-	if (LODIndex < 0)
-	{
-		LODIndex = FMath::Max(Builder.ExportOptions->DefaultLevelOfDetail, FGLTFMeshUtility::GetMinimumLOD(SkeletalMesh));
-	}
+	const uint32 MeshCount = (SkeletalMeshComponent != nullptr ? 1 : 0) + (SkeletalMesh != nullptr ? 1 : 0);
+	check(MeshCount == 1);
+
+	const TArray<UMaterialInterface*> Materials = SkeletalMeshComponent != nullptr
+		? OverrideMaterials.GetOverrides(SkeletalMeshComponent->GetMaterials())
+		: OverrideMaterials.GetOverrides(SkeletalMesh->Materials);
 
-	if (!Builder.ExportOptions->bExportVertexColors)
+	if (SkeletalMeshComponent != nullptr)
 	{
-		OverrideVertexColors = nullptr;
+		// Only use the component if it's needed for baking, since we would
+		// otherwise export a copy of this mesh for each mesh-component.
+		if (!Builder.ExportOptions->bBakeMaterialInputs ||
+			!Builder.ExportOptions->bBakeMaterialInputsUsingMeshData ||
+			!FGLTFMaterialUtility::MaterialsNeedVertexData(Materials))
+		{
+			SkeletalMesh = SkeletalMeshComponent->SkeletalMesh;
+			SkeletalMeshComponent = nullptr;
+		}
 	}
 
-	if (!Builder.ExportOptions->bExportVertexSkinWeights)
+	// Clean up override materials to only contain items that differ from the original materials
+	OverrideMaterials = FGLTFMaterialArray(Materials);
+
+	if (SkeletalMeshComponent != nullptr)
 	{
-		OverrideSkinWeights = nullptr;
+		OverrideMaterials.ClearRedundantOverrides(SkeletalMeshComponent->GetMaterials());
 	}
-
-	if (OverrideMaterials == SkeletalMesh->Materials)
+	else
 	{
-		OverrideMaterials.Empty();
+		OverrideMaterials.ClearRedundantOverrides(SkeletalMesh->Materials);
 	}
 
-	const FSkeletalMeshRenderData* RenderData = SkeletalMesh->GetResourceForRendering();
-	LODIndex = FMath::Clamp(LODIndex, 0, RenderData->LODRenderData.Num() - 1);
-	const FSkeletalMeshLODRenderData& MeshLOD = RenderData->LODRenderData[LODIndex];
+	const USkeletalMesh* Mesh = SkeletalMeshComponent != nullptr ? SkeletalMeshComponent->SkeletalMesh : SkeletalMesh;
 
-	if (OverrideVertexColors == &MeshLOD.StaticVertexBuffers.ColorVertexBuffer)
+	if (LODIndex < 0)
 	{
-		OverrideVertexColors = nullptr;
+		LODIndex = FMath::Max(Builder.ExportOptions->DefaultLevelOfDetail, FGLTFMeshUtility::GetMinimumLOD(Mesh));
 	}
 
-	if (OverrideSkinWeights == MeshLOD.GetSkinWeightVertexBuffer())
-	{
-		OverrideSkinWeights = nullptr;
-	}
+	const FSkeletalMeshRenderData* RenderData = Mesh->GetResourceForRendering();
+	LODIndex = FMath::Clamp(LODIndex, 0, RenderData->LODRenderData.Num() - 1);
 }
 
-FGLTFJsonMeshIndex FGLTFSkeletalMeshConverter::Convert(const USkeletalMesh* SkeletalMesh, int32 LODIndex, const FColorVertexBuffer* OverrideVertexColors, const FSkinWeightVertexBuffer* OverrideSkinWeights, FGLTFMaterialArray OverrideMaterials)
+FGLTFJsonMeshIndex FGLTFSkeletalMeshConverter::Convert(const USkeletalMesh* SkeletalMesh, const USkeletalMeshComponent* SkeletalMeshComponent, int32 LODIndex, FGLTFMaterialArray Materials)
 {
 	const FGLTFJsonMeshIndex MeshIndex = Builder.AddMesh();
-	Builder.SetupTask<FGLTFSkeletalMeshTask>(Builder, MeshSectionConverter, SkeletalMesh, LODIndex, OverrideVertexColors, OverrideSkinWeights, OverrideMaterials, MeshIndex);
+	Builder.SetupTask<FGLTFSkeletalMeshTask>(Builder, MeshSectionConverter, SkeletalMesh, SkeletalMeshComponent, LODIndex, Materials, MeshIndex);
 	return MeshIndex;
 }
diff --git a/Engine/Plugins/Enterprise/GLTFExporter/MigratedFromGit/Source/GLTFExporter/Private/Converters/GLTFMeshConverters.h b/Engine/Plugins/Enterprise/GLTFExporter/MigratedFromGit/Source/GLTFExporter/Private/Converters/GLTFMeshConverters.h
--- a/Engine/Plugins/Enterprise/GLTFExporter/MigratedFromGit/Source/GLTFExporter/Private/Converters/GLTFMeshConverters.h
+++ b/Engine/Plugins/Enterprise/GLTFExporter/MigratedFromGit/Source/GLTFExporter/Private/Converters/GLTFMeshConverters.h
@@ -15,24 +15,24 @@ class TGLTFMeshConverter : public FGLTFBuilderContext, public TGLTFConverter<FGL
 	using FGLTFBuilderContext::FGLTFBuilderContext;
 };
 
-class FGLTFStaticMeshConverter final : public TGLTFMeshConverter<const UStaticMesh*, int32, const FColorVertexBuffer*, FGLTFMaterialArray>
+class FGLTFStaticMeshConverter final : public TGLTFMeshConverter<const UStaticMesh*, const UStaticMeshComponent*, int32, FGLTFMaterialArray>
 {
 	using TGLTFMeshConverter::TGLTFMeshConverter;
 
-	virtual void Sanitize(const UStaticMesh*& StaticMesh, int32& LODIndex, const FColorVertexBuffer*& OverrideVertexColors, FGLTFMaterialArray& OverrideMaterials) override;
+	virtual void Sanitize(const UStaticMesh*& StaticMesh, const UStaticMeshComponent*& StaticMeshComponent, int32& LODIndex, FGLTFMaterialArray& OverrideMaterials) override;
 
-	virtual FGLTFJsonMeshIndex Convert(const UStaticMesh* StaticMesh, int32 LODIndex, const FColorVertexBuffer* OverrideVertexColors, FGLTFMaterialArray OverrideMaterials) override;
+	virtual FGLTFJsonMeshIndex Convert(const UStaticMesh* StaticMesh, const UStaticMeshComponent* StaticMeshComponent, int32 LODIndex, FGLTFMaterialArray Materials) override;
 
 	FGLTFStaticMeshSectionConverter MeshSectionConverter;
 };
 
-class FGLTFSkeletalMeshConverter final : public TGLTFMeshConverter<const USkeletalMesh*, int32, const FColorVertexBuffer*, const FSkinWeightVertexBuffer*, FGLTFMaterialArray>
+class FGLTFSkeletalMeshConverter final : public TGLTFMeshConverter<const USkeletalMesh*, const USkeletalMeshComponent*, int32, FGLTFMaterialArray>
 {
 	using TGLTFMeshConverter::TGLTFMeshConverter;
 
-	virtual void Sanitize(const USkeletalMesh*& SkeletalMesh, int32& LODIndex, const FColorVertexBuffer*& OverrideVertexColors, const FSkinWeightVertexBuffer*& OverrideSkinWeights, FGLTFMaterialArray& OverrideMaterials) override;
+	virtual void Sanitize(const USkeletalMesh*& SkeletalMesh, const USkeletalMeshComponent*& SkeletalMeshComponent, int32& LODIndex, FGLTFMaterialArray& OverrideMaterials) override;
 
-	virtual FGLTFJsonMeshIndex Convert(const USkeletalMesh* SkeletalMesh, int32 LODIndex, const FColorVertexBuffer* OverrideVertexColors, const FSkinWeightVertexBuffer* OverrideSkinWeights, FGLTFMaterialArray OverrideMaterials) override;
+	virtual FGLTFJsonMeshIndex Convert(const USkeletalMesh* SkeletalMesh, const USkeletalMeshComponent* SkeletalMeshComponent, int32 LODIndex, FGLTFMaterialArray Materials) override;
 
 	FGLTFSkeletalMeshSectionConverter MeshSectionConverter;
 };
diff --git a/Engine/Plugins/Enterprise/GLTFExporter/MigratedFromGit/Source/GLTFExporter/Private/Converters/GLTFMeshUtility.h b/Engine/Plugins/Enterprise/GLTFExporter/MigratedFromGit/Source/GLTFExporter/Private/Converters/GLTFMeshUtility.h
--- a/Engine/Plugins/Enterprise/GLTFExporter/MigratedFromGit/Source/GLTFExporter/Private/Converters/GLTFMeshUtility.h
+++ b/Engine/Plugins/Enterprise/GLTFExporter/MigratedFromGit/Source/GLTFExporter/Private/Converters/GLTFMeshUtility.h
@@ -3,6 +3,7 @@
 #pragma once
 
 #include "Engine.h"
+#include "GLTFMaterialArray.h"
 
 struct FGLTFMeshUtility
 {
diff --git a/Engine/Plugins/Enterprise/GLTFExporter/MigratedFromGit/Source/GLTFExporter/Private/Converters/GLTFSkySphereConverters.cpp b/Engine/Plugins/Enterprise/GLTFExporter/MigratedFromGit/Source/GLTFExporter/Private/Converters/GLTFSkySphereConverters.cpp
--- a/Engine/Plugins/Enterprise/GLTFExporter/MigratedFromGit/Source/GLTFExporter/Private/Converters/GLTFSkySphereConverters.cpp
+++ b/Engine/Plugins/Enterprise/GLTFExporter/MigratedFromGit/Source/GLTFExporter/Private/Converters/GLTFSkySphereConverters.cpp
@@ -21,8 +21,7 @@ FGLTFJsonSkySphereIndex FGLTFSkySphereConverter::Convert(const AActor* SkySphere
 	const UStaticMeshComponent* StaticMeshComponent = nullptr;
 	FGLTFActorUtility::TryGetPropertyValue(SkySphereActor, TEXT("SkySphereMesh"), StaticMeshComponent);
 
-	const UStaticMesh* StaticMesh = StaticMeshComponent != nullptr ? StaticMeshComponent->GetStaticMesh() : nullptr;
-	if (StaticMesh != nullptr)
+	if (StaticMeshComponent != nullptr)
 	{
 		const USceneComponent* ParentComponent = StaticMeshComponent->GetAttachParent();
 		const FName SocketName = StaticMeshComponent->GetAttachSocketName();
@@ -32,7 +31,7 @@ FGLTFJsonSkySphereIndex FGLTFSkySphereConverter::Convert(const AActor* SkySphere
 		const FTransform RelativeTransform = Transform.GetRelativeTransform(ParentTransform);
 
 		JsonSkySphere.Scale = FGLTFConverterUtility::ConvertScale(RelativeTransform.GetScale3D());
-		JsonSkySphere.SkySphereMesh = Builder.GetOrAddMesh(StaticMesh, -1, nullptr, { FGLTFMaterialUtility::GetDefault() });
+		JsonSkySphere.SkySphereMesh = Builder.GetOrAddMesh(StaticMeshComponent, { FGLTFMaterialUtility::GetDefault() });
 	}
 	else
 	{
diff --git a/Engine/Plugins/Enterprise/GLTFExporter/MigratedFromGit/Source/GLTFExporter/Private/Converters/GLTFVarationConverters.cpp b/Engine/Plugins/Enterprise/GLTFExporter/MigratedFromGit/Source/GLTFExporter/Private/Converters/GLTFVarationConverters.cpp
--- a/Engine/Plugins/Enterprise/GLTFExporter/MigratedFromGit/Source/GLTFExporter/Private/Converters/GLTFVarationConverters.cpp
+++ b/Engine/Plugins/Enterprise/GLTFExporter/MigratedFromGit/Source/GLTFExporter/Private/Converters/GLTFVarationConverters.cpp
@@ -268,7 +268,7 @@ bool FGLTFVariationConverter::TryParseMaterialPropertyValue(FGLTFJsonVariant& Ou
 
 	// TODO: find way to determine whether the material is null because "None" was selected, or because it failed to resolve
 	const UMaterialInterface* Material = MaterialProperty->GetMaterial();
-	const FGLTFJsonMaterialIndex MaterialIndex = Builder.GetOrAddMaterial(Material, nullptr);
+	const FGLTFJsonMaterialIndex MaterialIndex = Builder.GetOrAddMaterial(Material);
 	const int32 ElementIndex = CapturedPropSegments[NumPropSegments - 1].PropertyIndex;
 
 	FGLTFJsonVariantMaterial VariantMaterial;
@@ -333,7 +333,7 @@ bool FGLTFVariationConverter::TryParseStaticMeshPropertyValue(FGLTFJsonVariant&
 	const FGLTFMaterialArray OverrideMaterials(MeshComponent->OverrideMaterials);
 	const FGLTFJsonNodeIndex NodeIndex = Builder.GetOrAddNode(Target);
 	const FGLTFJsonNodeIndex ComponentNodeIndex = Builder.GetComponentNodeIndex(NodeIndex);
-	const FGLTFJsonMeshIndex MeshIndex = Builder.GetOrAddMesh(StaticMesh, -1, nullptr, OverrideMaterials);
+	const FGLTFJsonMeshIndex MeshIndex = Builder.GetOrAddMesh(StaticMesh, -1, OverrideMaterials);
 	FGLTFJsonVariantNodeProperties& NodeProperties = OutVariant.Nodes.FindOrAdd(ComponentNodeIndex);
 
 	NodeProperties.Node = ComponentNodeIndex;
@@ -389,7 +389,7 @@ bool FGLTFVariationConverter::TryParseSkeletalMeshPropertyValue(FGLTFJsonVariant
 
 	const FGLTFMaterialArray OverrideMaterials(MeshComponent->OverrideMaterials);
 	const FGLTFJsonNodeIndex NodeIndex = Builder.GetOrAddNode(Target);
-	const FGLTFJsonMeshIndex MeshIndex = Builder.GetOrAddMesh(SkeletalMesh, -1, nullptr, nullptr, OverrideMaterials);
+	const FGLTFJsonMeshIndex MeshIndex = Builder.GetOrAddMesh(SkeletalMesh, -1, OverrideMaterials);
 	const FGLTFJsonNodeIndex ComponentNodeIndex = Builder.GetComponentNodeIndex(NodeIndex);
 	FGLTFJsonVariantNodeProperties& NodeProperties = OutVariant.Nodes.FindOrAdd(ComponentNodeIndex);
 
diff --git a/Engine/Plugins/Enterprise/GLTFExporter/MigratedFromGit/Source/GLTFExporter/Private/Exporters/GLTFMaterialExporter.cpp b/Engine/Plugins/Enterprise/GLTFExporter/MigratedFromGit/Source/GLTFExporter/Private/Exporters/GLTFMaterialExporter.cpp
--- a/Engine/Plugins/Enterprise/GLTFExporter/MigratedFromGit/Source/GLTFExporter/Private/Exporters/GLTFMaterialExporter.cpp
+++ b/Engine/Plugins/Enterprise/GLTFExporter/MigratedFromGit/Source/GLTFExporter/Private/Exporters/GLTFMaterialExporter.cpp
@@ -29,7 +29,7 @@ bool UGLTFMaterialExporter::AddObject(FGLTFContainerBuilder& Builder, const UObj
 	{
 		if (PreviewMesh != nullptr)
 		{
-			const FGLTFJsonMeshIndex MeshIndex = Builder.GetOrAddMesh(PreviewMesh, -1, nullptr, { Material });
+			const FGLTFJsonMeshIndex MeshIndex = Builder.GetOrAddMesh(PreviewMesh, -1, { Material });
 			if (MeshIndex == INDEX_NONE)
 			{
 				Builder.AddErrorMessage(
diff --git a/Engine/Plugins/Enterprise/GLTFExporter/MigratedFromGit/Source/GLTFExporter/Private/Tasks/GLTFMaterialTasks.cpp b/Engine/Plugins/Enterprise/GLTFExporter/MigratedFromGit/Source/GLTFExporter/Private/Tasks/GLTFMaterialTasks.cpp
--- a/Engine/Plugins/Enterprise/GLTFExporter/MigratedFromGit/Source/GLTFExporter/Private/Tasks/GLTFMaterialTasks.cpp
+++ b/Engine/Plugins/Enterprise/GLTFExporter/MigratedFromGit/Source/GLTFExporter/Private/Tasks/GLTFMaterialTasks.cpp
@@ -1359,25 +1359,44 @@ FGLTFPropertyBakeOutput FGLTFMaterialTask::BakeMaterialProperty(EMaterialPropert
 	FMeshDescription MeshDescription;
 	TArray<int32> MaterialIndices;
 
-	if (Mesh != nullptr)
-	{
-		// TODO: add support for override materials
-
-		// TODO: add support for override vertex colors
-
-		// TODO: add support for baking world-position, -normals and other inputs that depend on component transform
+	bool bUseMeshData = bHasValidMeshOrComponent;
 
+	if (bUseMeshData)
+	{
 		FStaticMeshAttributes(MeshDescription).Register();
-		FMeshMergeHelpers::RetrieveMesh(Mesh, LODIndex, MeshDescription);
-
 		TArray<FSectionInfo> Sections;
-		FMeshMergeHelpers::ExtractSections(Mesh, LODIndex, Sections);
+		const TArray<UMaterialInterface*> SectionMaterials;
+
+		if (StaticMesh != nullptr)
+		{
+			FMeshMergeHelpers::RetrieveMesh(StaticMesh, LODIndex, MeshDescription);
+			FMeshMergeHelpers::ExtractSections(StaticMesh, LODIndex, Sections);
+		}
+		else if (StaticMeshComponent != nullptr)
+		{
+			FMeshMergeHelpers::RetrieveMesh(StaticMeshComponent, LODIndex, MeshDescription, true);
+			FMeshMergeHelpers::ExtractSections(StaticMeshComponent, LODIndex, Sections);
+		}
+		else if (SkeletalMesh != nullptr)
+		{
+			// TODO: add support for skeletal meshes by implementing custom utilities for retrieving mesh & sections
+
+			bUseMeshData = false;
+		}
+		else if (SkeletalMeshComponent != nullptr)
+		{
+			FMeshMergeHelpers::RetrieveMesh(const_cast<USkeletalMeshComponent*>(SkeletalMeshComponent), LODIndex, MeshDescription, true);
+			FMeshMergeHelpers::ExtractSections(SkeletalMeshComponent, LODIndex, Sections);
+		}
 
 		// Add material indices of all sections that reference the current material
 		for (int32 SectionIndex = 0; SectionIndex < Sections.Num(); ++SectionIndex)
 		{
 			FSectionInfo& Section = Sections[SectionIndex];
-			if (Section.Material == Material)
+			const UMaterialInterface* OverrideMaterial = OverrideMaterials.IsValidIndex(SectionIndex) ? OverrideMaterials[SectionIndex] : nullptr;
+			const UMaterialInterface* SectionMaterial = OverrideMaterial != nullptr ? OverrideMaterial : Section.Material;
+
+			if (SectionMaterial == Material)
 			{
 				MaterialIndices.Add(Section.MaterialIndex);
 			}
@@ -1398,8 +1417,8 @@ FGLTFPropertyBakeOutput FGLTFMaterialTask::BakeMaterialProperty(EMaterialPropert
 		Property,
 		Material,
 		OutTexCoord,
-		Mesh != nullptr ? &MeshDescription : nullptr,
-		Mesh != nullptr ? &MaterialIndices : nullptr,
+		bUseMeshData ? &MeshDescription : nullptr,
+		bUseMeshData ? &MaterialIndices : nullptr,
 		bCopyAlphaFromRedChannel);
 
 	if (!PropertyBakeOutput.bIsConstant && TexCoords.Num() == 0)
@@ -1438,9 +1457,21 @@ FString FGLTFMaterialTask::GetMaterialName() const
 {
 	FString MaterialName = Material->GetName();
 
-	if (Mesh != nullptr)
+	if (StaticMesh != nullptr)
+	{
+		MaterialName += TEXT("_") + FGLTFNameUtility::GetName(StaticMesh, LODIndex);
+	}
+	else if (StaticMeshComponent != nullptr)
+	{
+		MaterialName += TEXT("_") + FGLTFNameUtility::GetName(StaticMeshComponent);
+	}
+	else if (SkeletalMesh != nullptr)
+	{
+		MaterialName += TEXT("_") + FGLTFNameUtility::GetName(SkeletalMesh, LODIndex);
+	}
+	else if (SkeletalMeshComponent != nullptr)
 	{
-		MaterialName += TEXT("_") + FGLTFNameUtility::GetName(Mesh, LODIndex);
+		MaterialName += TEXT("_") + FGLTFNameUtility::GetName(SkeletalMeshComponent);
 	}
 
 	return MaterialName;
diff --git a/Engine/Plugins/Enterprise/GLTFExporter/MigratedFromGit/Source/GLTFExporter/Private/Tasks/GLTFMaterialTasks.h b/Engine/Plugins/Enterprise/GLTFExporter/MigratedFromGit/Source/GLTFExporter/Private/Tasks/GLTFMaterialTasks.h
--- a/Engine/Plugins/Enterprise/GLTFExporter/MigratedFromGit/Source/GLTFExporter/Private/Tasks/GLTFMaterialTasks.h
+++ b/Engine/Plugins/Enterprise/GLTFExporter/MigratedFromGit/Source/GLTFExporter/Private/Tasks/GLTFMaterialTasks.h
@@ -12,13 +12,18 @@ class FGLTFMaterialTask : public FGLTFTask
 {
 public:
 
-	FGLTFMaterialTask(FGLTFConvertBuilder& Builder, const UMaterialInterface* Material, const UStaticMesh* Mesh, int32 LODIndex, FGLTFJsonMaterialIndex MaterialIndex)
+	FGLTFMaterialTask(FGLTFConvertBuilder& Builder, const UMaterialInterface* Material, const UObject* MeshOrComponent, int32 LODIndex, FGLTFMaterialArray OverrideMaterials, FGLTFJsonMaterialIndex MaterialIndex)
         : FGLTFTask(EGLTFTaskPriority::Material)
 		, Builder(Builder)
 		, Material(Material)
-		, Mesh(Mesh)
 		, LODIndex(LODIndex)
+		, OverrideMaterials(OverrideMaterials)
 		, MaterialIndex(MaterialIndex)
+		, StaticMesh(Cast<UStaticMesh>(MeshOrComponent))
+		, StaticMeshComponent(Cast<UStaticMeshComponent>(MeshOrComponent))
+		, SkeletalMesh(Cast<USkeletalMesh>(MeshOrComponent))
+		, SkeletalMeshComponent(Cast<USkeletalMeshComponent>(MeshOrComponent))
+		, bHasValidMeshOrComponent(StaticMesh != nullptr || StaticMeshComponent != nullptr || SkeletalMesh != nullptr || SkeletalMeshComponent != nullptr)
 	{
 	}
 
@@ -33,10 +38,16 @@ private:
 
 	FGLTFConvertBuilder& Builder;
 	const UMaterialInterface* Material;
-	const UStaticMesh* Mesh;
 	const int32 LODIndex;
+	const FGLTFMaterialArray OverrideMaterials;
 	const FGLTFJsonMaterialIndex MaterialIndex;
 
+	const UStaticMesh* StaticMesh;
+	const UStaticMeshComponent* StaticMeshComponent;
+	const USkeletalMesh* SkeletalMesh;
+	const USkeletalMeshComponent* SkeletalMeshComponent;
+	const bool bHasValidMeshOrComponent;
+
 	bool TryGetAlphaMode(EGLTFJsonAlphaMode& AlphaMode) const;
 	bool TryGetShadingModel(EGLTFJsonShadingModel& ShadingModel) const;
 
diff --git a/Engine/Plugins/Enterprise/GLTFExporter/MigratedFromGit/Source/GLTFExporter/Private/Tasks/GLTFMeshTasks.cpp b/Engine/Plugins/Enterprise/GLTFExporter/MigratedFromGit/Source/GLTFExporter/Private/Tasks/GLTFMeshTasks.cpp
--- a/Engine/Plugins/Enterprise/GLTFExporter/MigratedFromGit/Source/GLTFExporter/Private/Tasks/GLTFMeshTasks.cpp
+++ b/Engine/Plugins/Enterprise/GLTFExporter/MigratedFromGit/Source/GLTFExporter/Private/Tasks/GLTFMeshTasks.cpp
@@ -8,14 +8,28 @@
 
 void FGLTFStaticMeshTask::Complete()
 {
+	const uint32 MeshCount = (StaticMeshComponent != nullptr ? 1 : 0) + (StaticMesh != nullptr ? 1 : 0);
+	check(MeshCount == 1);
+
+	if (StaticMeshComponent != nullptr)
+	{
+		StaticMesh = StaticMeshComponent->GetStaticMesh();
+	}
+
 	const FStaticMeshLODResources& MeshLOD = StaticMesh->GetLODForExport(LODIndex);
 
 	FGLTFJsonMesh& JsonMesh = Builder.GetMesh(MeshIndex);
-	JsonMesh.Name = FGLTFNameUtility::GetName(StaticMesh, LODIndex);
+	JsonMesh.Name = StaticMeshComponent != nullptr ? FGLTFNameUtility::GetName(StaticMeshComponent) : FGLTFNameUtility::GetName(StaticMesh, LODIndex);
 
 	const FPositionVertexBuffer* PositionBuffer = &MeshLOD.VertexBuffers.PositionVertexBuffer;
 	const FStaticMeshVertexBuffer* VertexBuffer = &MeshLOD.VertexBuffers.StaticMeshVertexBuffer;
-	const FColorVertexBuffer* ColorBuffer = OverrideVertexColors != nullptr ? OverrideVertexColors : &MeshLOD.VertexBuffers.ColorVertexBuffer;
+	const FColorVertexBuffer* ColorBuffer = &MeshLOD.VertexBuffers.ColorVertexBuffer;
+
+	if (StaticMeshComponent != nullptr && StaticMeshComponent->LODData.IsValidIndex(LODIndex))
+	{
+		const FStaticMeshComponentLODInfo& LODInfo = StaticMeshComponent->LODData[LODIndex];
+		ColorBuffer = LODInfo.OverrideVertexColors != nullptr ? LODInfo.OverrideVertexColors : ColorBuffer;
+	}
 
 	const int32 MaterialCount = StaticMesh->StaticMaterials.Num();
 	JsonMesh.Primitives.AddDefaulted(MaterialCount);
@@ -49,28 +63,48 @@ void FGLTFStaticMeshTask::Complete()
 			JsonPrimitive.Attributes.TexCoords[UVIndex] = Builder.GetOrAddUVAccessor(ConvertedSection, VertexBuffer, UVIndex);
 		}
 
-		const UMaterialInterface* Material = OverrideMaterials.GetOverride(StaticMesh->StaticMaterials, MaterialIndex);
+		const UMaterialInterface* Material = StaticMeshComponent != nullptr
+			? OverrideMaterials.GetOverride(StaticMeshComponent->GetMaterials(), MaterialIndex)	// TODO: cache component materials once instead of retrieving multiple times
+			: OverrideMaterials.GetOverride(StaticMesh->StaticMaterials, MaterialIndex);
+
 		if (Material == nullptr)
 		{
 			Material = UMaterial::GetDefaultMaterial(MD_Surface);
 		}
 
-		JsonPrimitive.Material = Builder.GetOrAddMaterial(Material, StaticMesh, LODIndex);
+		JsonPrimitive.Material = StaticMeshComponent != nullptr
+			? Builder.GetOrAddMaterial(Material, StaticMeshComponent, LODIndex, OverrideMaterials)
+			: Builder.GetOrAddMaterial(Material, StaticMesh, LODIndex, OverrideMaterials);
 	}
 }
 
 void FGLTFSkeletalMeshTask::Complete()
 {
+	const uint32 MeshCount = (SkeletalMeshComponent != nullptr ? 1 : 0) + (SkeletalMesh != nullptr ? 1 : 0);
+	check(MeshCount == 1);
+
+	if (SkeletalMeshComponent != nullptr)
+	{
+		SkeletalMesh = SkeletalMeshComponent->SkeletalMesh;
+	}
+
 	const FSkeletalMeshRenderData* RenderData = SkeletalMesh->GetResourceForRendering();
 	const FSkeletalMeshLODRenderData& MeshLOD = RenderData->LODRenderData[LODIndex];
 
 	FGLTFJsonMesh& JsonMesh = Builder.GetMesh(MeshIndex);
-	JsonMesh.Name = FGLTFNameUtility::GetName(SkeletalMesh, LODIndex);
+	JsonMesh.Name = SkeletalMeshComponent != nullptr ? FGLTFNameUtility::GetName(SkeletalMeshComponent) : FGLTFNameUtility::GetName(SkeletalMesh, LODIndex);
 
 	const FPositionVertexBuffer* PositionBuffer = &MeshLOD.StaticVertexBuffers.PositionVertexBuffer;
 	const FStaticMeshVertexBuffer* VertexBuffer = &MeshLOD.StaticVertexBuffers.StaticMeshVertexBuffer;
-	const FColorVertexBuffer* ColorBuffer = OverrideVertexColors != nullptr ? OverrideVertexColors : &MeshLOD.StaticVertexBuffers.ColorVertexBuffer;
-	const FSkinWeightVertexBuffer* SkinWeightBuffer = OverrideSkinWeights != nullptr ? OverrideSkinWeights : MeshLOD.GetSkinWeightVertexBuffer();
+	const FColorVertexBuffer* ColorBuffer = &MeshLOD.StaticVertexBuffers.ColorVertexBuffer;
+	const FSkinWeightVertexBuffer* SkinWeightBuffer = MeshLOD.GetSkinWeightVertexBuffer();
+
+	if (SkeletalMeshComponent != nullptr && SkeletalMeshComponent->LODInfo.IsValidIndex(LODIndex))
+	{
+		const FSkelMeshComponentLODInfo& LODInfo = SkeletalMeshComponent->LODInfo[LODIndex];
+		ColorBuffer = LODInfo.OverrideVertexColors != nullptr ? LODInfo.OverrideVertexColors : ColorBuffer;
+		SkinWeightBuffer = LODInfo.OverrideSkinWeights != nullptr ? LODInfo.OverrideSkinWeights : SkinWeightBuffer;
+	}
 
 	const uint16 MaterialCount = SkeletalMesh->Materials.Num();
 	JsonMesh.Primitives.AddDefaulted(MaterialCount);
@@ -117,13 +151,17 @@ void FGLTFSkeletalMeshTask::Complete()
 			}
 		}
 
-		const UMaterialInterface* Material = OverrideMaterials.GetOverride(SkeletalMesh->Materials, MaterialIndex);
+		const UMaterialInterface* Material = SkeletalMeshComponent != nullptr
+			? OverrideMaterials.GetOverride(SkeletalMeshComponent->GetMaterials(), MaterialIndex)	// TODO: cache component materials once instead of retrieving multiple times
+			: OverrideMaterials.GetOverride(SkeletalMesh->Materials, MaterialIndex);
+
 		if (Material == nullptr)
 		{
 			Material = UMaterial::GetDefaultMaterial(MD_Surface);
 		}
 
-		// TODO: should we somehow pass a skeletal mesh for baking?
-		JsonPrimitive.Material = Builder.GetOrAddMaterial(Material, nullptr);
+		JsonPrimitive.Material = SkeletalMeshComponent != nullptr
+			? Builder.GetOrAddMaterial(Material, SkeletalMeshComponent, LODIndex, OverrideMaterials)
+			: Builder.GetOrAddMaterial(Material, SkeletalMesh, LODIndex, OverrideMaterials);
 	}
 }
diff --git a/Engine/Plugins/Enterprise/GLTFExporter/MigratedFromGit/Source/GLTFExporter/Private/Tasks/GLTFMeshTasks.h b/Engine/Plugins/Enterprise/GLTFExporter/MigratedFromGit/Source/GLTFExporter/Private/Tasks/GLTFMeshTasks.h
--- a/Engine/Plugins/Enterprise/GLTFExporter/MigratedFromGit/Source/GLTFExporter/Private/Tasks/GLTFMeshTasks.h
+++ b/Engine/Plugins/Enterprise/GLTFExporter/MigratedFromGit/Source/GLTFExporter/Private/Tasks/GLTFMeshTasks.h
@@ -6,19 +6,20 @@
 #include "Builders/GLTFConvertBuilder.h"
 #include "Converters/GLTFMeshSectionConverters.h"
 #include "Converters/GLTFMaterialArray.h"
+#include "Converters/GLTFNameUtility.h"
 #include "Engine.h"
 
 class FGLTFStaticMeshTask : public FGLTFTask
 {
 public:
 
-	FGLTFStaticMeshTask(FGLTFConvertBuilder& Builder, FGLTFStaticMeshSectionConverter& MeshSectionConverter, const UStaticMesh* StaticMesh, int32 LODIndex, const FColorVertexBuffer* OverrideVertexColors, FGLTFMaterialArray OverrideMaterials, FGLTFJsonMeshIndex MeshIndex)
+	FGLTFStaticMeshTask(FGLTFConvertBuilder& Builder, FGLTFStaticMeshSectionConverter& MeshSectionConverter, const UStaticMesh* StaticMesh, const UStaticMeshComponent* StaticMeshComponent, int32 LODIndex, FGLTFMaterialArray OverrideMaterials, FGLTFJsonMeshIndex MeshIndex)
 		: FGLTFTask(EGLTFTaskPriority::Mesh)
 		, Builder(Builder)
 		, MeshSectionConverter(MeshSectionConverter)
 		, StaticMesh(StaticMesh)
+		, StaticMeshComponent(StaticMeshComponent)
 		, LODIndex(LODIndex)
-		, OverrideVertexColors(OverrideVertexColors)
 		, OverrideMaterials(OverrideMaterials)
         , MeshIndex(MeshIndex)
 	{
@@ -26,7 +27,7 @@ public:
 
 	virtual FString GetName() override
 	{
-		return StaticMesh->GetName();
+		return StaticMeshComponent != nullptr ? FGLTFNameUtility::GetName(StaticMeshComponent) : FGLTFNameUtility::GetName(StaticMesh, LODIndex);
 	}
 
 	virtual void Complete() override;
@@ -36,8 +37,8 @@ private:
 	FGLTFConvertBuilder& Builder;
 	FGLTFStaticMeshSectionConverter& MeshSectionConverter;
 	const UStaticMesh* StaticMesh;
+	const UStaticMeshComponent* StaticMeshComponent;
 	int32 LODIndex;
-	const FColorVertexBuffer* OverrideVertexColors;
 	FGLTFMaterialArray OverrideMaterials;
 	const FGLTFJsonMeshIndex MeshIndex;
 };
@@ -46,14 +47,13 @@ class FGLTFSkeletalMeshTask : public FGLTFTask
 {
 public:
 
-	FGLTFSkeletalMeshTask(FGLTFConvertBuilder& Builder, FGLTFSkeletalMeshSectionConverter& MeshSectionConverter, const USkeletalMesh* SkeletalMesh, int32 LODIndex, const FColorVertexBuffer* OverrideVertexColors, const FSkinWeightVertexBuffer* OverrideSkinWeights, FGLTFMaterialArray OverrideMaterials, FGLTFJsonMeshIndex MeshIndex)
+	FGLTFSkeletalMeshTask(FGLTFConvertBuilder& Builder, FGLTFSkeletalMeshSectionConverter& MeshSectionConverter, const USkeletalMesh* SkeletalMesh, const USkeletalMeshComponent* SkeletalMeshComponent, int32 LODIndex, FGLTFMaterialArray OverrideMaterials, FGLTFJsonMeshIndex MeshIndex)
 		: FGLTFTask(EGLTFTaskPriority::Mesh)
 		, Builder(Builder)
 		, MeshSectionConverter(MeshSectionConverter)
 		, SkeletalMesh(SkeletalMesh)
+		, SkeletalMeshComponent(SkeletalMeshComponent)
 		, LODIndex(LODIndex)
-		, OverrideVertexColors(OverrideVertexColors)
-		, OverrideSkinWeights(OverrideSkinWeights)
 		, OverrideMaterials(OverrideMaterials)
 		, MeshIndex(MeshIndex)
 	{
@@ -61,7 +61,7 @@ public:
 
 	virtual FString GetName() override
 	{
-		return SkeletalMesh->GetName();
+		return SkeletalMeshComponent != nullptr ? FGLTFNameUtility::GetName(SkeletalMeshComponent) : FGLTFNameUtility::GetName(SkeletalMesh, LODIndex);
 	}
 
 	virtual void Complete() override;
@@ -71,9 +71,8 @@ private:
 	FGLTFConvertBuilder& Builder;
 	FGLTFSkeletalMeshSectionConverter& MeshSectionConverter;
 	const USkeletalMesh* SkeletalMesh;
+	const USkeletalMeshComponent* SkeletalMeshComponent;
 	int32 LODIndex;
-	const FColorVertexBuffer* OverrideVertexColors;
-	const FSkinWeightVertexBuffer* OverrideSkinWeights;
 	FGLTFMaterialArray OverrideMaterials;
 	const FGLTFJsonMeshIndex MeshIndex;
 };
