[GLTFExporter] refactor: fix incorrect argument names in mesh converters

commit f58b69c1eb0fece0ce252de94819fd701c6470a4
Date: Mon, 26 Oct 2020 16:56:16 +0100

    refactor: fix incorrect argument names in mesh converters
    (cherry picked from commit c04c03727263688341a085d87ae9050016b651d8​)

#skipci

diff --git a/Engine/Plugins/Enterprise/GLTFExporter/MigratedFromGit/Source/GLTFExporter/Private/Converters/GLTFMeshConverters.cpp b/Engine/Plugins/Enterprise/GLTFExporter/MigratedFromGit/Source/GLTFExporter/Private/Converters/GLTFMeshConverters.cpp
--- a/Engine/Plugins/Enterprise/GLTFExporter/MigratedFromGit/Source/GLTFExporter/Private/Converters/GLTFMeshConverters.cpp
+++ b/Engine/Plugins/Enterprise/GLTFExporter/MigratedFromGit/Source/GLTFExporter/Private/Converters/GLTFMeshConverters.cpp
@@ -51,10 +51,10 @@ void FGLTFStaticMeshConverter::Sanitize(const UStaticMesh*& StaticMesh, const US
 	LODIndex = FMath::Clamp(LODIndex, 0, Mesh->GetNumLODs() - 1);
 }
 
-FGLTFJsonMeshIndex FGLTFStaticMeshConverter::Convert(const UStaticMesh* StaticMesh, const UStaticMeshComponent* StaticMeshComponent, int32 LODIndex, FGLTFMaterialArray Materials)
+FGLTFJsonMeshIndex FGLTFStaticMeshConverter::Convert(const UStaticMesh* StaticMesh, const UStaticMeshComponent* StaticMeshComponent, int32 LODIndex, FGLTFMaterialArray OverrideMaterials)
 {
 	const FGLTFJsonMeshIndex MeshIndex = Builder.AddMesh();
-	Builder.SetupTask<FGLTFStaticMeshTask>(Builder, MeshSectionConverter, StaticMesh, StaticMeshComponent, LODIndex, Materials, MeshIndex);
+	Builder.SetupTask<FGLTFStaticMeshTask>(Builder, MeshSectionConverter, StaticMesh, StaticMeshComponent, LODIndex, OverrideMaterials, MeshIndex);
 	return MeshIndex;
 }
 
@@ -103,9 +103,9 @@ void FGLTFSkeletalMeshConverter::Sanitize(const USkeletalMesh*& SkeletalMesh, co
 	LODIndex = FMath::Clamp(LODIndex, 0, RenderData->LODRenderData.Num() - 1);
 }
 
-FGLTFJsonMeshIndex FGLTFSkeletalMeshConverter::Convert(const USkeletalMesh* SkeletalMesh, const USkeletalMeshComponent* SkeletalMeshComponent, int32 LODIndex, FGLTFMaterialArray Materials)
+FGLTFJsonMeshIndex FGLTFSkeletalMeshConverter::Convert(const USkeletalMesh* SkeletalMesh, const USkeletalMeshComponent* SkeletalMeshComponent, int32 LODIndex, FGLTFMaterialArray OverrideMaterials)
 {
 	const FGLTFJsonMeshIndex MeshIndex = Builder.AddMesh();
-	Builder.SetupTask<FGLTFSkeletalMeshTask>(Builder, MeshSectionConverter, SkeletalMesh, SkeletalMeshComponent, LODIndex, Materials, MeshIndex);
+	Builder.SetupTask<FGLTFSkeletalMeshTask>(Builder, MeshSectionConverter, SkeletalMesh, SkeletalMeshComponent, LODIndex, OverrideMaterials, MeshIndex);
 	return MeshIndex;
 }
diff --git a/Engine/Plugins/Enterprise/GLTFExporter/MigratedFromGit/Source/GLTFExporter/Private/Converters/GLTFMeshConverters.h b/Engine/Plugins/Enterprise/GLTFExporter/MigratedFromGit/Source/GLTFExporter/Private/Converters/GLTFMeshConverters.h
--- a/Engine/Plugins/Enterprise/GLTFExporter/MigratedFromGit/Source/GLTFExporter/Private/Converters/GLTFMeshConverters.h
+++ b/Engine/Plugins/Enterprise/GLTFExporter/MigratedFromGit/Source/GLTFExporter/Private/Converters/GLTFMeshConverters.h
@@ -21,7 +21,7 @@ class FGLTFStaticMeshConverter final : public TGLTFMeshConverter<const UStaticMe
 
 	virtual void Sanitize(const UStaticMesh*& StaticMesh, const UStaticMeshComponent*& StaticMeshComponent, int32& LODIndex, FGLTFMaterialArray& OverrideMaterials) override;
 
-	virtual FGLTFJsonMeshIndex Convert(const UStaticMesh* StaticMesh, const UStaticMeshComponent* StaticMeshComponent, int32 LODIndex, FGLTFMaterialArray Materials) override;
+	virtual FGLTFJsonMeshIndex Convert(const UStaticMesh* StaticMesh, const UStaticMeshComponent* StaticMeshComponent, int32 LODIndex, FGLTFMaterialArray OverrideMaterials) override;
 
 	FGLTFStaticMeshSectionConverter MeshSectionConverter;
 };
@@ -32,7 +32,7 @@ class FGLTFSkeletalMeshConverter final : public TGLTFMeshConverter<const USkelet
 
 	virtual void Sanitize(const USkeletalMesh*& SkeletalMesh, const USkeletalMeshComponent*& SkeletalMeshComponent, int32& LODIndex, FGLTFMaterialArray& OverrideMaterials) override;
 
-	virtual FGLTFJsonMeshIndex Convert(const USkeletalMesh* SkeletalMesh, const USkeletalMeshComponent* SkeletalMeshComponent, int32 LODIndex, FGLTFMaterialArray Materials) override;
+	virtual FGLTFJsonMeshIndex Convert(const USkeletalMesh* SkeletalMesh, const USkeletalMeshComponent* SkeletalMeshComponent, int32 LODIndex, FGLTFMaterialArray OverrideMaterials) override;
 
 	FGLTFSkeletalMeshSectionConverter MeshSectionConverter;
 };
