[GLTFExporter] feat: add GLTFZipUtility to extract zip files

commit ce4c730c6be5b96448f7bb81d592930044f0368c
Date: Sun, 25 Oct 2020 02:53:17 +0100

    feat: add GLTFZipUtility to extract zip files
    (cherry picked from commit 11daa1c9b421d4d47220deb1d3d735727e34f4c1​)
    # Conflicts from cherry-pick:
    #	Plugins/GLTFExporter/Source/GLTFExporter/GLTFExporter.Build.cs

#skipci

diff --git a/Engine/Plugins/Enterprise/GLTFExporter/MigratedFromGit/Source/GLTFExporter/GLTFExporter.Build.cs b/Engine/Plugins/Enterprise/GLTFExporter/MigratedFromGit/Source/GLTFExporter/GLTFExporter.Build.cs
--- a/Engine/Plugins/Enterprise/GLTFExporter/MigratedFromGit/Source/GLTFExporter/GLTFExporter.Build.cs
+++ b/Engine/Plugins/Enterprise/GLTFExporter/MigratedFromGit/Source/GLTFExporter/GLTFExporter.Build.cs
@@ -34,7 +34,8 @@ namespace UnrealBuildTool.Rules
 					"MaterialBaking",
 					"VariantManagerContent",
 					"GLTFExporterRuntime",
-					"MeshMergeUtilities"
+					"MeshMergeUtilities",
+					"zlib"
 				}
 				);
 		}
diff --git a/Engine/Plugins/Enterprise/GLTFExporter/MigratedFromGit/Source/GLTFExporter/Private/Builders/GLTFZipUtility.cpp b/Engine/Plugins/Enterprise/GLTFExporter/MigratedFromGit/Source/GLTFExporter/Private/Builders/GLTFZipUtility.cpp
new file mode 100644
--- /dev/null
+++ b/Engine/Plugins/Enterprise/GLTFExporter/MigratedFromGit/Source/GLTFExporter/Private/Builders/GLTFZipUtility.cpp
@@ -0,0 +1,94 @@
+// Copyright Epic Games, Inc. All Rights Reserved.
+
+#include "Builders/GLTFZipUtility.h"
+
+THIRD_PARTY_INCLUDES_START
+#include "ThirdParty/zlib/zlib-1.2.5/Inc/zlib.h"
+#include "ThirdParty/zlib/zlib-1.2.5/Src/contrib/minizip/unzip.h"
+THIRD_PARTY_INCLUDES_END
+
+bool FGLTFZipUtility::ExtractToDirectory(const FString& SourceFilePath, const FString& DestinationDirectoryPath, FGLTFMessageBuilder& Builder)
+{
+	const unzFile ZipFile = unzOpen64(TCHAR_TO_ANSI(*SourceFilePath));
+	if (ZipFile == nullptr)
+	{
+		Builder.AddErrorMessage(TEXT("Can't open zip archive"));
+		return false;
+	}
+
+	if (const int ErrorCode = unzGoToFirstFile(ZipFile))
+	{
+		Builder.AddErrorMessage(FString::Printf(TEXT("Can't locate first file in zip archive (error %d)"), ErrorCode));
+		unzClose(ZipFile);
+		return false;
+	}
+
+	do
+	{
+		if (!ExtractCurrentFile(ZipFile, DestinationDirectoryPath, Builder))
+		{
+			unzClose(ZipFile);
+			return false;
+		}
+	} while (unzGoToNextFile(ZipFile) == UNZ_OK);
+
+	unzClose(ZipFile);
+	return true;
+}
+
+bool FGLTFZipUtility::ExtractCurrentFile(void* ZipFile, const FString& DestinationDirectoryPath, FGLTFMessageBuilder& Builder)
+{
+	unz_file_info64 FileInfo;
+	if (const int ErrorCode = unzGetCurrentFileInfo64(ZipFile, &FileInfo, nullptr, 0, nullptr, 0, nullptr, 0))
+	{
+		Builder.AddErrorMessage(FString::Printf(TEXT("Can't get file info in zip archive (error %d)"), ErrorCode));
+		return false;
+	}
+
+	TArray<char> Filename;
+	Filename.Init('\0', FileInfo.size_filename + 1);
+	if (const int ErrorCode = unzGetCurrentFileInfo64(ZipFile, nullptr, Filename.GetData(), Filename.Num() - 1, nullptr, 0, nullptr, 0))
+	{
+		Builder.AddErrorMessage(FString::Printf(TEXT("Can't get file name in zip archive (error %d)"), ErrorCode));
+		return false;
+	}
+
+	const FString DestinationFilePath = DestinationDirectoryPath / Filename.GetData();
+	if (DestinationFilePath.EndsWith(TEXT("/")) || DestinationFilePath.EndsWith(TEXT("\\")))
+	{
+		IFileManager::Get().MakeDirectory(*DestinationFilePath, true);
+	}
+	else
+	{
+		FArchive* FileWriter = IFileManager::Get().CreateFileWriter(*DestinationFilePath);
+
+		if (const int ErrorCode = unzOpenCurrentFile(ZipFile))
+		{
+			Builder.AddErrorMessage(FString::Printf(TEXT("Can't open file %s in zip archive (error %d)"), ANSI_TO_TCHAR(Filename.GetData()), ErrorCode));
+			return false;
+		}
+
+		uint8 ReadBuffer[8192];
+		while (const int ReadSize = unzReadCurrentFile(ZipFile, ReadBuffer, sizeof(ReadBuffer)))
+		{
+			if (ReadSize < 0)
+			{
+				Builder.AddErrorMessage(FString::Printf(TEXT("Can't read file %s in zip archive (error %d)"), ANSI_TO_TCHAR(Filename.GetData()), ReadSize));
+				unzCloseCurrentFile(ZipFile);
+				return false;
+			}
+
+			FileWriter->Serialize(ReadBuffer, ReadSize);
+		}
+
+		if (const int ErrorCode = unzCloseCurrentFile(ZipFile))
+		{
+			Builder.AddErrorMessage(FString::Printf(TEXT("Can't close file %s in zip archive (error %d)"), ANSI_TO_TCHAR(Filename.GetData()), ErrorCode));
+			return false;
+		}
+
+		FileWriter->Close();
+	}
+
+	return true;
+}
diff --git a/Engine/Plugins/Enterprise/GLTFExporter/MigratedFromGit/Source/GLTFExporter/Private/Builders/GLTFZipUtility.h b/Engine/Plugins/Enterprise/GLTFExporter/MigratedFromGit/Source/GLTFExporter/Private/Builders/GLTFZipUtility.h
new file mode 100644
--- /dev/null
+++ b/Engine/Plugins/Enterprise/GLTFExporter/MigratedFromGit/Source/GLTFExporter/Private/Builders/GLTFZipUtility.h
@@ -0,0 +1,17 @@
+// Copyright Epic Games, Inc. All Rights Reserved.
+
+#pragma once
+
+#include "Builders/GLTFMessageBuilder.h"
+#include "CoreMinimal.h"
+
+class FGLTFZipUtility
+{
+public:
+
+	static bool ExtractToDirectory(const FString& SourceFilePath, const FString& DestinationDirectoryPath, FGLTFMessageBuilder& Builder);
+
+private:
+
+	static bool ExtractCurrentFile(void* ZipFile, const FString& DestinationDirectoryPath, FGLTFMessageBuilder& Builder);
+};
