[GLTFExporter] feat: use GLTFZipUtlity to extract viewer bundle

commit 149933fbb13daa91f6c6d9a1d105425f9f635a8f
Date: Sun, 25 Oct 2020 02:54:06 +0100

    feat: use GLTFZipUtlity to extract viewer bundle
    (cherry picked from commit 83be51b8136a2b90f9a25428a628591f92077fcf​)

#skipci

diff --git a/Engine/Plugins/Enterprise/GLTFExporter/MigratedFromGit/Source/GLTFExporter/Private/Builders/GLTFContainerBuilder.cpp b/Engine/Plugins/Enterprise/GLTFExporter/MigratedFromGit/Source/GLTFExporter/Private/Builders/GLTFContainerBuilder.cpp
--- a/Engine/Plugins/Enterprise/GLTFExporter/MigratedFromGit/Source/GLTFExporter/Private/Builders/GLTFContainerBuilder.cpp
+++ b/Engine/Plugins/Enterprise/GLTFExporter/MigratedFromGit/Source/GLTFExporter/Private/Builders/GLTFContainerBuilder.cpp
@@ -2,6 +2,7 @@
 
 #include "Builders/GLTFContainerBuilder.h"
 #include "Builders/GLTFContainerUtility.h"
+#include "Builders/GLTFZipUtility.h"
 #include "GLTFExporterModule.h"
 #include "Interfaces/IPluginManager.h"
 
@@ -22,27 +23,16 @@ void FGLTFContainerBuilder::WriteGlb(FArchive& Archive) const
 
 void FGLTFContainerBuilder::BundleWebViewer()
 {
-	// TODO: instead of raw files, store web viewer as a zip and unpack it when needed
-
-	const FString WebViewerPath = IPluginManager::Get().FindPlugin(GLTFEXPORTER_MODULE_NAME)->GetBaseDir() / TEXT("Resources") / TEXT("WebViewer");
-	TArray<FString> SourceFiles;
-
-	IFileManager& FileManager = IFileManager::Get();
-	FileManager.FindFilesRecursive(SourceFiles, *WebViewerPath, TEXT("*"), true, false);
-
-	if (SourceFiles.Num() == 0)
+	const FString WebViewerPath = IPluginManager::Get().FindPlugin(GLTFEXPORTER_MODULE_NAME)->GetBaseDir() / TEXT("Resources") / TEXT("WebViewer.zip");
+	if (!FPaths::FileExists(WebViewerPath))
 	{
-		AddErrorMessage(FString::Printf(TEXT("No source files for web viewer found at %s"), *WebViewerPath));
+		AddWarningMessage(FString::Printf(TEXT("No web viewer found at %s"), *WebViewerPath));
 		return;
 	}
 
-	for (const FString& SourceFile : SourceFiles)
+	if (!FGLTFZipUtility::ExtractToDirectory(WebViewerPath, DirPath, *this))
 	{
-		FString DestinationFile = DirPath / SourceFile.RightChop(WebViewerPath.Len());
-		if (FileManager.Copy(*DestinationFile, *SourceFile) != COPY_OK)
-		{
-			AddErrorMessage(FString::Printf(TEXT("Failed to copy web viewer file from %s to %s"), *SourceFile, *DestinationFile));
-		}
+		AddErrorMessage(FString::Printf(TEXT("Failed to extract web viewer at %s"), *WebViewerPath));
 	}
 }
 
