[GLTFExporter] refactor: warn about non-constrained aspect-ratio for orthographic camera

commit 8bc7160cba5fcf6bae1e51566a15fb4ec40ac939
Date: Sun, 25 Oct 2020 15:55:29 +0100

    refactor: warn about non-constrained aspect-ratio for orthographic camera
    (cherry picked from commit c2cbd03175f4a8464f312ed80c4961b4f87920c5​)

#skipci

diff --git a/Engine/Plugins/Enterprise/GLTFExporter/MigratedFromGit/Source/GLTFExporter/Private/Converters/GLTFCameraConverters.cpp b/Engine/Plugins/Enterprise/GLTFExporter/MigratedFromGit/Source/GLTFExporter/Private/Converters/GLTFCameraConverters.cpp
--- a/Engine/Plugins/Enterprise/GLTFExporter/MigratedFromGit/Source/GLTFExporter/Private/Converters/GLTFCameraConverters.cpp
+++ b/Engine/Plugins/Enterprise/GLTFExporter/MigratedFromGit/Source/GLTFExporter/Private/Converters/GLTFCameraConverters.cpp
@@ -4,7 +4,6 @@
 #include "Builders/GLTFContainerBuilder.h"
 #include "Converters/GLTFConverterUtility.h"
 #include "Converters/GLTFNameUtility.h"
-#include "Converters/GLTFCameraUtility.h"
 #include "Actors/GLTFCameraActor.h"
 
 FGLTFJsonCameraIndex FGLTFCameraConverter::Convert(const UCameraComponent* CameraComponent)
@@ -15,15 +14,30 @@ FGLTFJsonCameraIndex FGLTFCameraConverter::Convert(const UCameraComponent* Camer
 
 	FMinimalViewInfo DesiredView;
 	const_cast<UCameraComponent*>(CameraComponent)->GetCameraView(0, DesiredView);
+	const float ExportScale = Builder.ExportOptions->ExportScale;
 
 	switch (Camera.Type)
 	{
 		case EGLTFJsonCameraType::Orthographic:
-			Camera.Orthographic = FGLTFCameraUtility::ConvertOrthographic(DesiredView, Builder.ExportOptions->ExportScale);
+			if (!DesiredView.bConstrainAspectRatio)
+			{
+				Builder.AddWarningMessage(FString::Printf(TEXT("Aspect ratio for orthographic camera component %s (in actor %s) will be constrainted in glTF"), *CameraComponent->GetName(), *CameraComponent->GetOwner()->GetName()));
+			}
+			Camera.Orthographic.XMag = FGLTFConverterUtility::ConvertLength(DesiredView.OrthoWidth, ExportScale);
+			Camera.Orthographic.YMag = FGLTFConverterUtility::ConvertLength(DesiredView.OrthoWidth / DesiredView.AspectRatio, ExportScale); // TODO: is this correct?
+			Camera.Orthographic.ZFar = FGLTFConverterUtility::ConvertLength(DesiredView.OrthoFarClipPlane, ExportScale);
+			Camera.Orthographic.ZNear = FGLTFConverterUtility::ConvertLength(DesiredView.OrthoNearClipPlane, ExportScale);
 			break;
 
 		case EGLTFJsonCameraType::Perspective:
-			Camera.Perspective = FGLTFCameraUtility::ConvertPerspective(DesiredView, Builder.ExportOptions->ExportScale);
+			if (DesiredView.bConstrainAspectRatio)
+			{
+				Camera.Perspective.AspectRatio = DesiredView.AspectRatio;
+			}
+			Camera.Perspective.YFov = FGLTFConverterUtility::ConvertFieldOfView(DesiredView.FOV, DesiredView.AspectRatio);
+			// NOTE: even thought ZFar is optional, if we don't set it, then most gltf viewers won't handle it well.
+			Camera.Perspective.ZFar = FGLTFConverterUtility::ConvertLength(WORLD_MAX, ExportScale); // TODO: Unreal doesn't have max draw distance per view?
+			Camera.Perspective.ZNear = FGLTFConverterUtility::ConvertLength(GNearClippingPlane, ExportScale);
 			break;
 
 		case EGLTFJsonCameraType::None:
@@ -44,8 +58,8 @@ FGLTFJsonCameraIndex FGLTFCameraConverter::Convert(const UCameraComponent* Camer
 		{
 			FGLTFJsonOrbitCamera OrbitCamera;
 			OrbitCamera.Focus = Builder.GetOrAddNode(CameraActor->Focus);
-			OrbitCamera.MaxDistance = FGLTFConverterUtility::ConvertLength(CameraActor->DistanceMax, Builder.ExportOptions->ExportScale);
-			OrbitCamera.MinDistance = FGLTFConverterUtility::ConvertLength(CameraActor->DistanceMin, Builder.ExportOptions->ExportScale);
+			OrbitCamera.MaxDistance = FGLTFConverterUtility::ConvertLength(CameraActor->DistanceMax, ExportScale);
+			OrbitCamera.MinDistance = FGLTFConverterUtility::ConvertLength(CameraActor->DistanceMin, ExportScale);
 			OrbitCamera.MaxAngle = CameraActor->PitchAngleMax;
 			OrbitCamera.MinAngle = CameraActor->PitchAngleMin;
 			OrbitCamera.DistanceSensitivity = CameraActor->DistanceSensitivity;
diff --git a/Engine/Plugins/Enterprise/GLTFExporter/MigratedFromGit/Source/GLTFExporter/Private/Converters/GLTFCameraUtility.cpp b/Engine/Plugins/Enterprise/GLTFExporter/MigratedFromGit/Source/GLTFExporter/Private/Converters/GLTFCameraUtility.cpp
deleted file mode 100644
diff --git a/Engine/Plugins/Enterprise/GLTFExporter/MigratedFromGit/Source/GLTFExporter/Private/Converters/GLTFCameraUtility.h b/Engine/Plugins/Enterprise/GLTFExporter/MigratedFromGit/Source/GLTFExporter/Private/Converters/GLTFCameraUtility.h
deleted file mode 100644
diff --git a/Engine/Plugins/Enterprise/GLTFExporter/MigratedFromGit/Source/GLTFExporter/Private/Converters/GLTFConverterUtility.h b/Engine/Plugins/Enterprise/GLTFExporter/MigratedFromGit/Source/GLTFExporter/Private/Converters/GLTFConverterUtility.h
--- a/Engine/Plugins/Enterprise/GLTFExporter/MigratedFromGit/Source/GLTFExporter/Private/Converters/GLTFConverterUtility.h
+++ b/Engine/Plugins/Enterprise/GLTFExporter/MigratedFromGit/Source/GLTFExporter/Private/Converters/GLTFConverterUtility.h
@@ -144,6 +144,13 @@ struct FGLTFConverterUtility
 		return ConvertMatrix(Matrix);
 	}
 
+	static float ConvertFieldOfView(float FOVInDegress, float AspectRatio)
+	{
+		const float HorizontalFOV = FMath::DegreesToRadians(FOVInDegress);
+		const float VerticalFOV = 2 * FMath::Atan(FMath::Tan(HorizontalFOV / 2) / AspectRatio);
+		return VerticalFOV;
+	}
+
 	static FGLTFJsonQuaternion ConvertCameraDirection()
 	{
 		// Unreal uses +X axis as camera direction in Unreal coordinates.
diff --git a/Engine/Plugins/Enterprise/GLTFExporter/MigratedFromGit/Source/GLTFExporter/Private/Json/GLTFJsonCamera.h b/Engine/Plugins/Enterprise/GLTFExporter/MigratedFromGit/Source/GLTFExporter/Private/Json/GLTFJsonCamera.h
--- a/Engine/Plugins/Enterprise/GLTFExporter/MigratedFromGit/Source/GLTFExporter/Private/Json/GLTFJsonCamera.h
+++ b/Engine/Plugins/Enterprise/GLTFExporter/MigratedFromGit/Source/GLTFExporter/Private/Json/GLTFJsonCamera.h
@@ -27,11 +27,7 @@ struct FGLTFJsonOrthographic
 	{
 		JsonWriter.WriteObjectStart();
 
-		if (XMag != 0)
-		{
-			FGLTFJsonUtility::WriteExactValue(JsonWriter, TEXT("xmag"), XMag);
-		}
-
+		FGLTFJsonUtility::WriteExactValue(JsonWriter, TEXT("xmag"), XMag);
 		FGLTFJsonUtility::WriteExactValue(JsonWriter, TEXT("ymag"), YMag);
 		FGLTFJsonUtility::WriteExactValue(JsonWriter, TEXT("zfar"), ZFar);
 		FGLTFJsonUtility::WriteExactValue(JsonWriter, TEXT("znear"), ZNear);
