[GLTFExporter] refactor: add separate texture converter for lightmaps

commit 4702402d2f739852c8f7fd49c0664ff9d5ba0f72
Date: Mon, 26 Oct 2020 22:46:21 +0100

    refactor: add separate texture converter for lightmaps
    (cherry picked from commit e0e908ef8249eb3b389077409ead46ca34514da5​)

#skipci

diff --git a/Engine/Plugins/Enterprise/GLTFExporter/MigratedFromGit/Source/GLTFExporter/Private/Builders/GLTFConvertBuilder.cpp b/Engine/Plugins/Enterprise/GLTFExporter/MigratedFromGit/Source/GLTFExporter/Private/Builders/GLTFConvertBuilder.cpp
--- a/Engine/Plugins/Enterprise/GLTFExporter/MigratedFromGit/Source/GLTFExporter/Private/Builders/GLTFConvertBuilder.cpp
+++ b/Engine/Plugins/Enterprise/GLTFExporter/MigratedFromGit/Source/GLTFExporter/Private/Builders/GLTFConvertBuilder.cpp
@@ -222,6 +222,16 @@ FGLTFJsonTextureIndex FGLTFConvertBuilder::GetOrAddTexture(const UTextureRenderT
 	return TextureRenderTargetCubeConverter.GetOrAdd(Texture, CubeFace);
 }
 
+FGLTFJsonTextureIndex FGLTFConvertBuilder::GetOrAddTexture(const ULightMapTexture2D* Texture)
+{
+	if (Texture == nullptr)
+	{
+		return FGLTFJsonTextureIndex(INDEX_NONE);
+	}
+
+	return TextureLightMapConverter.GetOrAdd(Texture);
+}
+
 FGLTFJsonSkinIndex FGLTFConvertBuilder::GetOrAddSkin(FGLTFJsonNodeIndex RootNode, const USkeletalMesh* SkeletalMesh)
 {
 	if (RootNode == INDEX_NONE || SkeletalMesh == nullptr)
diff --git a/Engine/Plugins/Enterprise/GLTFExporter/MigratedFromGit/Source/GLTFExporter/Private/Builders/GLTFConvertBuilder.h b/Engine/Plugins/Enterprise/GLTFExporter/MigratedFromGit/Source/GLTFExporter/Private/Builders/GLTFConvertBuilder.h
--- a/Engine/Plugins/Enterprise/GLTFExporter/MigratedFromGit/Source/GLTFExporter/Private/Builders/GLTFConvertBuilder.h
+++ b/Engine/Plugins/Enterprise/GLTFExporter/MigratedFromGit/Source/GLTFExporter/Private/Builders/GLTFConvertBuilder.h
@@ -54,6 +54,7 @@ public:
 	FGLTFJsonTextureIndex GetOrAddTexture(const UTextureCube* Texture, ECubeFace CubeFace);
 	FGLTFJsonTextureIndex GetOrAddTexture(const UTextureRenderTarget2D* Texture);
 	FGLTFJsonTextureIndex GetOrAddTexture(const UTextureRenderTargetCube* Texture, ECubeFace CubeFace);
+	FGLTFJsonTextureIndex GetOrAddTexture(const ULightMapTexture2D* Texture);
 
 	FGLTFJsonSkinIndex GetOrAddSkin(FGLTFJsonNodeIndex RootNode, const USkeletalMesh* SkeletalMesh);
 	FGLTFJsonSkinIndex GetOrAddSkin(FGLTFJsonNodeIndex RootNode, const USkeletalMeshComponent* SkeletalMeshComponent);
@@ -99,6 +100,7 @@ private:
 	FGLTFTextureCubeConverter TextureCubeConverter{ *this };
 	FGLTFTextureRenderTarget2DConverter TextureRenderTarget2DConverter{ *this };
 	FGLTFTextureRenderTargetCubeConverter TextureRenderTargetCubeConverter{ *this };
+	FGLTFTextureLightMapConverter TextureLightMapConverter{ *this };
 
 	FGLTFSkinConverter SkinConverter{ *this };
 	FGLTFAnimationConverter AnimationConverter{ *this };
diff --git a/Engine/Plugins/Enterprise/GLTFExporter/MigratedFromGit/Source/GLTFExporter/Private/Converters/GLTFTextureConverters.cpp b/Engine/Plugins/Enterprise/GLTFExporter/MigratedFromGit/Source/GLTFExporter/Private/Converters/GLTFTextureConverters.cpp
--- a/Engine/Plugins/Enterprise/GLTFExporter/MigratedFromGit/Source/GLTFExporter/Private/Converters/GLTFTextureConverters.cpp
+++ b/Engine/Plugins/Enterprise/GLTFExporter/MigratedFromGit/Source/GLTFExporter/Private/Converters/GLTFTextureConverters.cpp
@@ -31,3 +31,10 @@ FGLTFJsonTextureIndex FGLTFTextureRenderTargetCubeConverter::Convert(const UText
 	Builder.SetupTask<FGLTFTextureRenderTargetCubeTask>(Builder, RenderTargetCube, CubeFace, TextureIndex);
 	return TextureIndex;
 }
+
+FGLTFJsonTextureIndex FGLTFTextureLightMapConverter::Convert(const ULightMapTexture2D* LightMap)
+{
+	const FGLTFJsonTextureIndex TextureIndex = Builder.AddTexture();
+	Builder.SetupTask<FGLTFTextureLightMapTask>(Builder, LightMap, TextureIndex);
+	return TextureIndex;
+}
diff --git a/Engine/Plugins/Enterprise/GLTFExporter/MigratedFromGit/Source/GLTFExporter/Private/Converters/GLTFTextureConverters.h b/Engine/Plugins/Enterprise/GLTFExporter/MigratedFromGit/Source/GLTFExporter/Private/Converters/GLTFTextureConverters.h
--- a/Engine/Plugins/Enterprise/GLTFExporter/MigratedFromGit/Source/GLTFExporter/Private/Converters/GLTFTextureConverters.h
+++ b/Engine/Plugins/Enterprise/GLTFExporter/MigratedFromGit/Source/GLTFExporter/Private/Converters/GLTFTextureConverters.h
@@ -40,3 +40,10 @@ class FGLTFTextureRenderTargetCubeConverter final : public TGLTFTextureConverter
 
 	virtual FGLTFJsonTextureIndex Convert(const UTextureRenderTargetCube* RenderTargetCube, ECubeFace CubeFace) override;
 };
+
+class FGLTFTextureLightMapConverter final : public TGLTFTextureConverter<const ULightMapTexture2D*>
+{
+	using TGLTFTextureConverter::TGLTFTextureConverter;
+
+	virtual FGLTFJsonTextureIndex Convert(const ULightMapTexture2D* LightMap) override;
+};
diff --git a/Engine/Plugins/Enterprise/GLTFExporter/MigratedFromGit/Source/GLTFExporter/Private/Tasks/GLTFTextureTasks.cpp b/Engine/Plugins/Enterprise/GLTFExporter/MigratedFromGit/Source/GLTFExporter/Private/Tasks/GLTFTextureTasks.cpp
--- a/Engine/Plugins/Enterprise/GLTFExporter/MigratedFromGit/Source/GLTFExporter/Private/Tasks/GLTFTextureTasks.cpp
+++ b/Engine/Plugins/Enterprise/GLTFExporter/MigratedFromGit/Source/GLTFExporter/Private/Tasks/GLTFTextureTasks.cpp
@@ -214,3 +214,34 @@ void FGLTFTextureRenderTargetCubeTask::Complete()
 	JsonTexture.Source = Builder.AddImage(Pixels.GetData(), Size, JsonTexture.Name);
 	JsonTexture.Sampler = Builder.GetOrAddSampler(RenderTargetCube);
 }
+
+void FGLTFTextureLightMapTask::Complete()
+{
+	FGLTFJsonTexture& JsonTexture = Builder.GetTexture(TextureIndex);
+	LightMap->GetName(JsonTexture.Name);
+
+	// NOTE: export of lightmaps via source data is used to work around issues with
+	// quality-loss due to incorrect gamma transformation when rendering to a canvas.
+
+	FTextureSource& Source = const_cast<FTextureSource&>(LightMap->Source);
+	if (!Source.IsValid())
+	{
+		Builder.AddWarningMessage(FString::Printf(TEXT("Failed to export lightmap texture %s because of missing source data"), *JsonTexture.Name));
+		return;
+	}
+
+	ERGBFormat RGBFormat;
+	uint32 BitDepth;
+	if (!FGLTFTextureUtility::CanPNGCompressFormat(Source.GetFormat(), RGBFormat, BitDepth))
+	{
+		Builder.AddWarningMessage(FString::Printf(TEXT("Failed to export lightmap texture %s because of unsupported format %s"), *JsonTexture.Name, *FGLTFNameUtility::GetName(Source.GetFormat())));
+		return;
+	}
+
+	const FIntPoint Size = { Source.GetSizeX(), Source.GetSizeY() };
+	const void* RawData = Source.LockMip(0);
+	JsonTexture.Source = Builder.AddImage(RawData, Source.CalcMipSize(0), Size, RGBFormat, BitDepth, JsonTexture.Name);
+	Source.UnlockMip(0);
+
+	JsonTexture.Sampler = Builder.GetOrAddSampler(LightMap);
+}
diff --git a/Engine/Plugins/Enterprise/GLTFExporter/MigratedFromGit/Source/GLTFExporter/Private/Tasks/GLTFTextureTasks.h b/Engine/Plugins/Enterprise/GLTFExporter/MigratedFromGit/Source/GLTFExporter/Private/Tasks/GLTFTextureTasks.h
--- a/Engine/Plugins/Enterprise/GLTFExporter/MigratedFromGit/Source/GLTFExporter/Private/Tasks/GLTFTextureTasks.h
+++ b/Engine/Plugins/Enterprise/GLTFExporter/MigratedFromGit/Source/GLTFExporter/Private/Tasks/GLTFTextureTasks.h
@@ -113,3 +113,29 @@ private:
 	ECubeFace CubeFace;
 	const FGLTFJsonTextureIndex TextureIndex;
 };
+
+class FGLTFTextureLightMapTask : public FGLTFTask
+{
+public:
+
+	FGLTFTextureLightMapTask(FGLTFConvertBuilder& Builder, const ULightMapTexture2D* LightMap, FGLTFJsonTextureIndex TextureIndex)
+        : FGLTFTask(EGLTFTaskPriority::Texture)
+        , Builder(Builder)
+        , LightMap(LightMap)
+        , TextureIndex(TextureIndex)
+	{
+	}
+
+	virtual FString GetName() override
+	{
+		return LightMap->GetName();
+	}
+
+	virtual void Complete() override;
+
+private:
+
+	FGLTFConvertBuilder& Builder;
+	const ULightMapTexture2D* LightMap;
+	const FGLTFJsonTextureIndex TextureIndex;
+};
