[GLTFExporter] refactor: remove export option bExportSourceTextures since it doesn't support adjustments or max texture size etc

commit ff8858c34508d39b2fa8450478382c5b481d3ae3
Date: Mon, 26 Oct 2020 22:38:36 +0100

    refactor: remove export option bExportSourceTextures since it doesn't support adjustments or max texture size etc
    (cherry picked from commit ced30d6b74737e1003410c62f5e214bc3c14d989​)

#skipci

diff --git a/Engine/Plugins/Enterprise/GLTFExporter/MigratedFromGit/Source/GLTFExporter/Private/Converters/GLTFTextureUtility.cpp b/Engine/Plugins/Enterprise/GLTFExporter/MigratedFromGit/Source/GLTFExporter/Private/Converters/GLTFTextureUtility.cpp
--- a/Engine/Plugins/Enterprise/GLTFExporter/MigratedFromGit/Source/GLTFExporter/Private/Converters/GLTFTextureUtility.cpp
+++ b/Engine/Plugins/Enterprise/GLTFExporter/MigratedFromGit/Source/GLTFExporter/Private/Converters/GLTFTextureUtility.cpp
@@ -43,21 +43,6 @@ bool FGLTFTextureUtility::IsCubemap(const UTexture* Texture)
 	return Texture->IsA<UTextureCube>() || Texture->IsA<UTextureRenderTargetCube>();
 }
 
-bool FGLTFTextureUtility::HasAnyAdjustment(const UTexture* Texture)
-{
-	const float ErrorTolerance = KINDA_SMALL_NUMBER;
-
-	return !FMath::IsNearlyEqual(Texture->AdjustBrightness, 1.0f, ErrorTolerance) ||
-		!FMath::IsNearlyEqual(Texture->AdjustBrightnessCurve, 1.0f, ErrorTolerance) ||
-		!FMath::IsNearlyEqual(Texture->AdjustSaturation, 1.0f, ErrorTolerance) ||
-		!FMath::IsNearlyEqual(Texture->AdjustVibrance, 0.0f, ErrorTolerance) ||
-		!FMath::IsNearlyEqual(Texture->AdjustRGBCurve, 1.0f, ErrorTolerance) ||
-		!FMath::IsNearlyEqual(Texture->AdjustHue, 0.0f, ErrorTolerance) ||
-		!FMath::IsNearlyEqual(Texture->AdjustMinAlpha, 0.0f, ErrorTolerance) ||
-		!FMath::IsNearlyEqual(Texture->AdjustMaxAlpha, 1.0f, ErrorTolerance) ||
-		Texture->bChromaKeyTexture;
-}
-
 float FGLTFTextureUtility::GetCubeFaceRotation(ECubeFace CubeFace)
 {
 	switch (CubeFace)
diff --git a/Engine/Plugins/Enterprise/GLTFExporter/MigratedFromGit/Source/GLTFExporter/Private/Converters/GLTFTextureUtility.h b/Engine/Plugins/Enterprise/GLTFExporter/MigratedFromGit/Source/GLTFExporter/Private/Converters/GLTFTextureUtility.h
--- a/Engine/Plugins/Enterprise/GLTFExporter/MigratedFromGit/Source/GLTFExporter/Private/Converters/GLTFTextureUtility.h
+++ b/Engine/Plugins/Enterprise/GLTFExporter/MigratedFromGit/Source/GLTFExporter/Private/Converters/GLTFTextureUtility.h
@@ -15,9 +15,6 @@ struct FGLTFTextureUtility
 	static bool CanPNGCompressFormat(EPixelFormat InFormat, ERGBFormat& OutFormat, uint32& OutBitDepth);
 
 	static bool IsCubemap(const UTexture* Texture);
-
-	static bool HasAnyAdjustment(const UTexture* Texture);
-
 	static float GetCubeFaceRotation(ECubeFace CubeFace);
 
 	static TextureFilter GetDefaultFilter(TextureGroup Group);
@@ -31,7 +28,6 @@ struct FGLTFTextureUtility
 	static bool DrawTexture(UTextureRenderTarget2D* OutTarget, const UTexture2D* InSource, const FMatrix& InTransform = FMatrix::Identity);
 	static bool RotateTexture(UTextureRenderTarget2D* OutTarget, const UTexture2D* InSource, float InDegrees);
 
-	// TODO: maybe use template specialization to avoid the need for duplicated functions
 	static UTexture2D* CreateTextureFromCubeFace(const UTextureCube* TextureCube, ECubeFace CubeFace);
 	static UTexture2D* CreateTextureFromCubeFace(const UTextureRenderTargetCube* RenderTargetCube, ECubeFace CubeFace);
 
diff --git a/Engine/Plugins/Enterprise/GLTFExporter/MigratedFromGit/Source/GLTFExporter/Private/GLTFExportOptions.cpp b/Engine/Plugins/Enterprise/GLTFExporter/MigratedFromGit/Source/GLTFExporter/Private/GLTFExportOptions.cpp
--- a/Engine/Plugins/Enterprise/GLTFExporter/MigratedFromGit/Source/GLTFExporter/Private/GLTFExportOptions.cpp
+++ b/Engine/Plugins/Enterprise/GLTFExporter/MigratedFromGit/Source/GLTFExporter/Private/GLTFExportOptions.cpp
@@ -29,7 +29,6 @@ void UGLTFExportOptions::ResetToDefault()
 	bRetargetBoneTransforms = true;
 	bExportPlaybackSettings = true;
 	bExportTextureTransforms = true;
-	bExportSourceTextures = true;
 	bExportLightmaps = true;
 	TextureHDREncoding = EGLTFExporterTextureHDREncoding::RGBM;
 	ExportScale = 0.01;
diff --git a/Engine/Plugins/Enterprise/GLTFExporter/MigratedFromGit/Source/GLTFExporter/Private/Tasks/GLTFTextureTasks.cpp b/Engine/Plugins/Enterprise/GLTFExporter/MigratedFromGit/Source/GLTFExporter/Private/Tasks/GLTFTextureTasks.cpp
--- a/Engine/Plugins/Enterprise/GLTFExporter/MigratedFromGit/Source/GLTFExporter/Private/Tasks/GLTFTextureTasks.cpp
+++ b/Engine/Plugins/Enterprise/GLTFExporter/MigratedFromGit/Source/GLTFExporter/Private/Tasks/GLTFTextureTasks.cpp
@@ -4,104 +4,51 @@
 #include "Builders/GLTFContainerBuilder.h"
 #include "Converters/GLTFConverterUtility.h"
 #include "Converters/GLTFTextureUtility.h"
+#include "Converters/GLTFNameUtility.h"
 
 void FGLTFTexture2DTask::Complete()
 {
-	bool bIsNormalMap = Texture2D->IsNormalMap();
-	FGLTFJsonImageIndex ImageIndex;
-
 	FGLTFJsonTexture& JsonTexture = Builder.GetTexture(TextureIndex);
 	Texture2D->GetName(JsonTexture.Name);
 
-	// NOTE: export of normalmaps via source data is disabled because we need to manipulate the pixels during conversion.
-	const bool bPreferSourceExport = bIsNormalMap ? false : Builder.ExportOptions->bExportSourceTextures;
+	// TODO: both bForceLinearGamma and TargetGamma=2.2f seem to be necessary for the exported images to match the results
+	// from exporting using the texture's source or its platform-data.
+	// It's not entirely clear why gamma must be set to 2.2 instead of 0.0 (which should use the correct gamma anyway),
+	// and why bInForceLinearGamma must also be true.
+	const bool bForceLinearGamma = true;
+	const float TargetGamma = 2.2f;
+	const FIntPoint Size = { Texture2D->GetSizeX(), Texture2D->GetSizeY() };
 
-	// NOTE: export of lightmaps via source data is used to work around issues
-	// with quality-loss due to incorrect gamma transformation when rendering to a canvas.
+	const EPixelFormat RenderTargetFormat = FGLTFTextureUtility::IsHDRFormat(Texture2D->GetPixelFormat()) ? PF_FloatRGBA : PF_B8G8R8A8;
+	UTextureRenderTarget2D* RenderTarget = FGLTFTextureUtility::CreateRenderTarget(Size, RenderTargetFormat, bForceLinearGamma);
+	RenderTarget->TargetGamma = TargetGamma;
 
-	// TODO: make sure there are no other special cases than lightmaps that require source-export
-	const bool bRequireSourceExport = Texture2D->IsA<ULightMapTexture2D>();
+	// TODO: preserve maximum image quality (avoid compression artifacts) by copying source data (and adjustments) to a temp texture
+	FGLTFTextureUtility::DrawTexture(RenderTarget, Texture2D);
 
-	if (bRequireSourceExport || bPreferSourceExport)
+	TArray<FColor> Pixels;
+	if (!Texture2D->IsNormalMap() && FGLTFTextureUtility::IsHDRFormat(RenderTarget->GetFormat()))
 	{
-		ERGBFormat RGBFormat;
-		uint32 BitDepth;
-
-		FTextureSource& Source = const_cast<FTextureSource&>(Texture2D->Source);
-		const FIntPoint Size = { Source.GetSizeX(), Source.GetSizeY() };
-
-		if (Source.IsValid() && FGLTFTextureUtility::CanPNGCompressFormat(Source.GetFormat(), RGBFormat, BitDepth))
-		{
-			const void* RawData = Source.LockMip(0);
-			ImageIndex = Builder.AddImage(RawData, Source.CalcMipSize(0), Size, RGBFormat, BitDepth, JsonTexture.Name);
-			Source.UnlockMip(0);
-		}
-
-		if (ImageIndex != INDEX_NONE)
-		{
-			if (FGLTFTextureUtility::HasAnyAdjustment(Texture2D))
-			{
-				Builder.AddWarningMessage(FString::Printf(TEXT("Adjustments for texture %s are not supported when exporting source data"), *JsonTexture.Name));
-			}
-		}
-		else
-		{
-			if (bPreferSourceExport)
-			{
-				Builder.AddWarningMessage(FString::Printf(TEXT("Unable to export source for texture %s, render data will be used as fallback"), *JsonTexture.Name));
-			}
-		}
+		JsonTexture.Encoding = EGLTFJsonHDREncoding::RGBM; // TODO: use only encoding as specified by export options
+		FGLTFTextureUtility::ReadEncodedPixels(RenderTarget, Pixels, JsonTexture.Encoding);
 	}
-
-	if (ImageIndex == INDEX_NONE && !bRequireSourceExport)
+	else
 	{
-		const FIntPoint Size = { Texture2D->GetSizeX(), Texture2D->GetSizeY() };
-
-		// TODO: both bForceLinearGamma and TargetGamma=2.2f seem to be necessary for the exported images to match the results
-		// from exporting using the texture's source or its platform-data.
-		// It's not entirely clear why gamma must be set to 2.2 instead of 0.0 (which should use the correct gamma anyway),
-		// and why bInForceLinearGamma must also be true.
-		const bool bForceLinearGamma = true;
-		const float TargetGamma = 2.2f;
-
-		const EPixelFormat RenderTargetFormat = FGLTFTextureUtility::IsHDRFormat(Texture2D->GetPixelFormat()) ? PF_FloatRGBA : PF_B8G8R8A8;
-		UTextureRenderTarget2D* RenderTarget = FGLTFTextureUtility::CreateRenderTarget(Size, RenderTargetFormat, bForceLinearGamma);
-		RenderTarget->TargetGamma = TargetGamma;
-
-		FGLTFTextureUtility::DrawTexture(RenderTarget, Texture2D);
-
-		TArray<FColor> Pixels;
-		if (!bIsNormalMap && FGLTFTextureUtility::IsHDRFormat(RenderTarget->GetFormat()))
-		{
-			JsonTexture.Encoding = EGLTFJsonHDREncoding::RGBM;
-			FGLTFTextureUtility::ReadEncodedPixels(RenderTarget, Pixels, JsonTexture.Encoding); // TODO: use only encoding as specified by export options
-		}
-		else
-		{
-			FGLTFTextureUtility::ReadPixels(RenderTarget, Pixels);
-
-			if (bIsNormalMap)
-			{
-				FGLTFTextureUtility::FlipGreenChannel(Pixels);
-			}
-		}
+		FGLTFTextureUtility::ReadPixels(RenderTarget, Pixels);
 
-		if (Pixels.Num() == 0)
+		if (Texture2D->IsNormalMap())
 		{
-			Builder.AddWarningMessage(FString::Printf(TEXT("Failed to read pixels for 2D texture %s"), *JsonTexture.Name));
-			return;
+			FGLTFTextureUtility::FlipGreenChannel(Pixels);
 		}
-
-		ImageIndex = Builder.AddImage(Pixels.GetData(), Size, JsonTexture.Name);
 	}
 
-	if (ImageIndex == INDEX_NONE)
+	if (Pixels.Num() == 0)
 	{
-		Builder.AddWarningMessage(FString::Printf(TEXT("Failed to export 2D texture %s"), *JsonTexture.Name));
+		Builder.AddWarningMessage(FString::Printf(TEXT("Failed to read pixels for 2D texture %s"), *JsonTexture.Name));
 		return;
 	}
 
-	JsonTexture.Source = ImageIndex;
+	JsonTexture.Source = Builder.AddImage(Pixels.GetData(), Size, JsonTexture.Name);
 	JsonTexture.Sampler = Builder.GetOrAddSampler(Texture2D);
 }
 
@@ -112,6 +59,7 @@ void FGLTFTextureCubeTask::Complete()
 
 	// TODO: add optimized "happy path" if cube face doesn't need rotation and has suitable pixel format
 
+	// TODO: preserve maximum image quality (avoid compression artifacts) by copying source data (and adjustments) to a temp texture
 	const UTexture2D* FaceTexture = FGLTFTextureUtility::CreateTextureFromCubeFace(TextureCube, CubeFace);
 	if (FaceTexture == nullptr)
 	{
@@ -129,8 +77,8 @@ void FGLTFTextureCubeTask::Complete()
 	TArray<FColor> Pixels;
 	if (FGLTFTextureUtility::IsHDRFormat(RenderTarget->GetFormat()))
 	{
-		JsonTexture.Encoding = EGLTFJsonHDREncoding::RGBM;
-		FGLTFTextureUtility::ReadEncodedPixels(RenderTarget, Pixels, JsonTexture.Encoding); // TODO: use only encoding as specified by export options
+		JsonTexture.Encoding = EGLTFJsonHDREncoding::RGBM; // TODO: use only encoding as specified by export options
+		FGLTFTextureUtility::ReadEncodedPixels(RenderTarget, Pixels, JsonTexture.Encoding);
 	}
 	else
 	{
@@ -155,8 +103,8 @@ void FGLTFTextureRenderTarget2DTask::Complete()
 	TArray<FColor> Pixels;
 	if (FGLTFTextureUtility::IsHDRFormat(RenderTarget2D->GetFormat()))
 	{
-		JsonTexture.Encoding = EGLTFJsonHDREncoding::RGBM;
-		FGLTFTextureUtility::ReadEncodedPixels(RenderTarget2D, Pixels, JsonTexture.Encoding); // TODO: use only encoding as specified by export options
+		JsonTexture.Encoding = EGLTFJsonHDREncoding::RGBM; // TODO: use only encoding as specified by export options
+		FGLTFTextureUtility::ReadEncodedPixels(RenderTarget2D, Pixels, JsonTexture.Encoding);
 	}
 	else
 	{
@@ -197,8 +145,8 @@ void FGLTFTextureRenderTargetCubeTask::Complete()
 	TArray<FColor> Pixels;
 	if (FGLTFTextureUtility::IsHDRFormat(RenderTarget->GetFormat()))
 	{
-		JsonTexture.Encoding = EGLTFJsonHDREncoding::RGBM;
-		FGLTFTextureUtility::ReadEncodedPixels(RenderTarget, Pixels, JsonTexture.Encoding); // TODO: use only encoding as specified by export options
+		JsonTexture.Encoding = EGLTFJsonHDREncoding::RGBM; // TODO: use only encoding as specified by export options
+		FGLTFTextureUtility::ReadEncodedPixels(RenderTarget, Pixels, JsonTexture.Encoding);
 	}
 	else
 	{
diff --git a/Engine/Plugins/Enterprise/GLTFExporter/MigratedFromGit/Source/GLTFExporter/Public/GLTFExportOptions.h b/Engine/Plugins/Enterprise/GLTFExporter/MigratedFromGit/Source/GLTFExporter/Public/GLTFExportOptions.h
--- a/Engine/Plugins/Enterprise/GLTFExporter/MigratedFromGit/Source/GLTFExporter/Public/GLTFExportOptions.h
+++ b/Engine/Plugins/Enterprise/GLTFExporter/MigratedFromGit/Source/GLTFExporter/Public/GLTFExportOptions.h
@@ -94,9 +94,6 @@ public:
 	UPROPERTY(EditAnywhere, BlueprintReadWrite, Config, Category = Texture)
 	bool bExportTextureTransforms;
 
-	UPROPERTY(EditAnywhere, BlueprintReadWrite, Config, Category = Texture)
-	bool bExportSourceTextures;
-
 	UPROPERTY(EditAnywhere, BlueprintReadWrite, Config, Category = Texture)
 	bool bExportLightmaps;
 
