[GLTFExporter] refactor: move GetDefaultMaterialBakeSize and ShouldExportLight from GLTFExportOptions to GLTFBuilder

commit 2a917dc7251a247fce334ef8102741ea44145303
Date: Tue, 27 Oct 2020 12:36:19 +0100

    refactor: move GetDefaultMaterialBakeSize and ShouldExportLight from GLTFExportOptions to GLTFBuilder
    (cherry picked from commit a4a570792de508ffe2fbcdc19ddf4348c3a0efcc​)

#skipci

diff --git a/Engine/Plugins/Enterprise/GLTFExporter/MigratedFromGit/Source/GLTFExporter/Private/Builders/GLTFBuilder.cpp b/Engine/Plugins/Enterprise/GLTFExporter/MigratedFromGit/Source/GLTFExporter/Private/Builders/GLTFBuilder.cpp
--- a/Engine/Plugins/Enterprise/GLTFExporter/MigratedFromGit/Source/GLTFExporter/Private/Builders/GLTFBuilder.cpp
+++ b/Engine/Plugins/Enterprise/GLTFExporter/MigratedFromGit/Source/GLTFExporter/Private/Builders/GLTFBuilder.cpp
@@ -10,3 +10,24 @@ FGLTFBuilder::FGLTFBuilder(const FString& FilePath, const UGLTFExportOptions* Ex
 	, ExportOptions(ExportOptions)
 {
 }
+
+FIntPoint FGLTFBuilder::GetDefaultMaterialBakeSize() const
+{
+	const int32 Size = 1 << static_cast<int>(ExportOptions->DefaultMaterialBakeSize);
+	return { Size, Size };
+}
+
+bool FGLTFBuilder::ShouldExportLight(EComponentMobility::Type LightMobility) const
+{
+	switch (ExportOptions->ExportLights)
+	{
+		case EGLTFExporterLightMobility::All:
+			return true;
+		case EGLTFExporterLightMobility::MovableAndStationary:
+			return LightMobility == EComponentMobility::Movable || LightMobility == EComponentMobility::Stationary;
+		case EGLTFExporterLightMobility::MovableOnly:
+			return LightMobility == EComponentMobility::Movable;
+		default:
+			return false;
+	}
+}
diff --git a/Engine/Plugins/Enterprise/GLTFExporter/MigratedFromGit/Source/GLTFExporter/Private/Builders/GLTFBuilder.h b/Engine/Plugins/Enterprise/GLTFExporter/MigratedFromGit/Source/GLTFExporter/Private/Builders/GLTFBuilder.h
--- a/Engine/Plugins/Enterprise/GLTFExporter/MigratedFromGit/Source/GLTFExporter/Private/Builders/GLTFBuilder.h
+++ b/Engine/Plugins/Enterprise/GLTFExporter/MigratedFromGit/Source/GLTFExporter/Private/Builders/GLTFBuilder.h
@@ -19,4 +19,8 @@ public:
 	const FString DirPath;
 
 	const UGLTFExportOptions* const ExportOptions;
+
+	FIntPoint GetDefaultMaterialBakeSize() const;
+
+	bool ShouldExportLight(EComponentMobility::Type LightMobility) const;
 };
diff --git a/Engine/Plugins/Enterprise/GLTFExporter/MigratedFromGit/Source/GLTFExporter/Private/Builders/GLTFJsonBuilder.cpp b/Engine/Plugins/Enterprise/GLTFExporter/MigratedFromGit/Source/GLTFExporter/Private/Builders/GLTFJsonBuilder.cpp
--- a/Engine/Plugins/Enterprise/GLTFExporter/MigratedFromGit/Source/GLTFExporter/Private/Builders/GLTFJsonBuilder.cpp
+++ b/Engine/Plugins/Enterprise/GLTFExporter/MigratedFromGit/Source/GLTFExporter/Private/Builders/GLTFJsonBuilder.cpp
@@ -3,7 +3,7 @@
 #include "Builders/GLTFJsonBuilder.h"
 
 FGLTFJsonBuilder::FGLTFJsonBuilder(const FString& FilePath, const UGLTFExportOptions* ExportOptions)
-	: FGLTFMessageBuilder(FilePath, ExportOptions)
+	: FGLTFTaskBuilder(FilePath, ExportOptions)
 	, DefaultScene(JsonRoot.DefaultScene)
 {
 }
diff --git a/Engine/Plugins/Enterprise/GLTFExporter/MigratedFromGit/Source/GLTFExporter/Private/Builders/GLTFJsonBuilder.h b/Engine/Plugins/Enterprise/GLTFExporter/MigratedFromGit/Source/GLTFExporter/Private/Builders/GLTFJsonBuilder.h
--- a/Engine/Plugins/Enterprise/GLTFExporter/MigratedFromGit/Source/GLTFExporter/Private/Builders/GLTFJsonBuilder.h
+++ b/Engine/Plugins/Enterprise/GLTFExporter/MigratedFromGit/Source/GLTFExporter/Private/Builders/GLTFJsonBuilder.h
@@ -3,9 +3,9 @@
 #pragma once
 
 #include "Json/GLTFJsonRoot.h"
-#include "Builders/GLTFMessageBuilder.h"
+#include "Builders/GLTFTaskBuilder.h"
 
-class FGLTFJsonBuilder : public FGLTFMessageBuilder
+class FGLTFJsonBuilder : public FGLTFTaskBuilder
 {
 protected:
 
diff --git a/Engine/Plugins/Enterprise/GLTFExporter/MigratedFromGit/Source/GLTFExporter/Private/Builders/GLTFMessageBuilder.cpp b/Engine/Plugins/Enterprise/GLTFExporter/MigratedFromGit/Source/GLTFExporter/Private/Builders/GLTFMessageBuilder.cpp
--- a/Engine/Plugins/Enterprise/GLTFExporter/MigratedFromGit/Source/GLTFExporter/Private/Builders/GLTFMessageBuilder.cpp
+++ b/Engine/Plugins/Enterprise/GLTFExporter/MigratedFromGit/Source/GLTFExporter/Private/Builders/GLTFMessageBuilder.cpp
@@ -7,7 +7,7 @@
 #include "IMessageLogListing.h"
 
 FGLTFMessageBuilder::FGLTFMessageBuilder(const FString& FilePath, const UGLTFExportOptions* ExportOptions)
-	: FGLTFTaskBuilder(FilePath, ExportOptions)
+	: FGLTFBuilder(FilePath, ExportOptions)
 {
 }
 
diff --git a/Engine/Plugins/Enterprise/GLTFExporter/MigratedFromGit/Source/GLTFExporter/Private/Builders/GLTFMessageBuilder.h b/Engine/Plugins/Enterprise/GLTFExporter/MigratedFromGit/Source/GLTFExporter/Private/Builders/GLTFMessageBuilder.h
--- a/Engine/Plugins/Enterprise/GLTFExporter/MigratedFromGit/Source/GLTFExporter/Private/Builders/GLTFMessageBuilder.h
+++ b/Engine/Plugins/Enterprise/GLTFExporter/MigratedFromGit/Source/GLTFExporter/Private/Builders/GLTFMessageBuilder.h
@@ -2,7 +2,7 @@
 
 #pragma once
 
-#include "Builders/GLTFTaskBuilder.h"
+#include "Builders/GLTFBuilder.h"
 
 enum class EGLTFMessageSeverity
 {
@@ -11,7 +11,7 @@ enum class EGLTFMessageSeverity
 	Error
 };
 
-class FGLTFMessageBuilder : public FGLTFTaskBuilder
+class FGLTFMessageBuilder : public FGLTFBuilder
 {
 protected:
 
diff --git a/Engine/Plugins/Enterprise/GLTFExporter/MigratedFromGit/Source/GLTFExporter/Private/Builders/GLTFTaskBuilder.cpp b/Engine/Plugins/Enterprise/GLTFExporter/MigratedFromGit/Source/GLTFExporter/Private/Builders/GLTFTaskBuilder.cpp
--- a/Engine/Plugins/Enterprise/GLTFExporter/MigratedFromGit/Source/GLTFExporter/Private/Builders/GLTFTaskBuilder.cpp
+++ b/Engine/Plugins/Enterprise/GLTFExporter/MigratedFromGit/Source/GLTFExporter/Private/Builders/GLTFTaskBuilder.cpp
@@ -5,7 +5,7 @@
 #include "Misc/ScopedSlowTask.h"
 
 FGLTFTaskBuilder::FGLTFTaskBuilder(const FString& FilePath, const UGLTFExportOptions* ExportOptions)
-	: FGLTFBuilder(FilePath, ExportOptions)
+	: FGLTFMessageBuilder(FilePath, ExportOptions)
 	, PriorityIndexLock(INDEX_NONE)
 {
 }
diff --git a/Engine/Plugins/Enterprise/GLTFExporter/MigratedFromGit/Source/GLTFExporter/Private/Builders/GLTFTaskBuilder.h b/Engine/Plugins/Enterprise/GLTFExporter/MigratedFromGit/Source/GLTFExporter/Private/Builders/GLTFTaskBuilder.h
--- a/Engine/Plugins/Enterprise/GLTFExporter/MigratedFromGit/Source/GLTFExporter/Private/Builders/GLTFTaskBuilder.h
+++ b/Engine/Plugins/Enterprise/GLTFExporter/MigratedFromGit/Source/GLTFExporter/Private/Builders/GLTFTaskBuilder.h
@@ -2,10 +2,10 @@
 
 #pragma once
 
-#include "Builders/GLTFBuilder.h"
+#include "Builders/GLTFMessageBuilder.h"
 #include "Tasks/GLTFTask.h"
 
-class FGLTFTaskBuilder : public FGLTFBuilder
+class FGLTFTaskBuilder : public FGLTFMessageBuilder
 {
 protected:
 
diff --git a/Engine/Plugins/Enterprise/GLTFExporter/MigratedFromGit/Source/GLTFExporter/Private/Converters/GLTFNodeConverters.cpp b/Engine/Plugins/Enterprise/GLTFExporter/MigratedFromGit/Source/GLTFExporter/Private/Converters/GLTFNodeConverters.cpp
--- a/Engine/Plugins/Enterprise/GLTFExporter/MigratedFromGit/Source/GLTFExporter/Private/Converters/GLTFNodeConverters.cpp
+++ b/Engine/Plugins/Enterprise/GLTFExporter/MigratedFromGit/Source/GLTFExporter/Private/Converters/GLTFNodeConverters.cpp
@@ -158,7 +158,7 @@ FGLTFJsonNodeIndex FGLTFComponentConverter::Convert(const USceneComponent* Scene
 	}
 	else if (const ULightComponent* LightComponent = Cast<ULightComponent>(SceneComponent))
 	{
-		if (Builder.ExportOptions->ShouldExportLight(LightComponent->Mobility))
+		if (Builder.ShouldExportLight(LightComponent->Mobility))
 		{
 			// TODO: conversion of light direction should be done in separate converter
 			FGLTFJsonNode LightNode;
diff --git a/Engine/Plugins/Enterprise/GLTFExporter/MigratedFromGit/Source/GLTFExporter/Private/GLTFExportOptions.cpp b/Engine/Plugins/Enterprise/GLTFExporter/MigratedFromGit/Source/GLTFExporter/Private/GLTFExportOptions.cpp
--- a/Engine/Plugins/Enterprise/GLTFExporter/MigratedFromGit/Source/GLTFExporter/Private/GLTFExportOptions.cpp
+++ b/Engine/Plugins/Enterprise/GLTFExporter/MigratedFromGit/Source/GLTFExporter/Private/GLTFExportOptions.cpp
@@ -97,24 +97,3 @@ void UGLTFExportOptions::PostEditChangeProperty(FPropertyChangedEvent& PropertyC
 		}
 	}
 }
-
-FIntPoint UGLTFExportOptions::GetDefaultMaterialBakeSize() const
-{
-	const int32 Size = 1 << static_cast<int>(DefaultMaterialBakeSize);
-	return { Size, Size };
-}
-
-bool UGLTFExportOptions::ShouldExportLight(EComponentMobility::Type LightMobility) const
-{
-	switch (ExportLights)
-	{
-		case EGLTFExporterLightMobility::All:
-			return true;
-		case EGLTFExporterLightMobility::MovableAndStationary:
-			return LightMobility == EComponentMobility::Movable || LightMobility == EComponentMobility::Stationary;
-		case EGLTFExporterLightMobility::MovableOnly:
-			return LightMobility == EComponentMobility::Movable;
-		default:
-			return false;
-	}
-}
diff --git a/Engine/Plugins/Enterprise/GLTFExporter/MigratedFromGit/Source/GLTFExporter/Private/Tasks/GLTFMaterialTasks.cpp b/Engine/Plugins/Enterprise/GLTFExporter/MigratedFromGit/Source/GLTFExporter/Private/Tasks/GLTFMaterialTasks.cpp
--- a/Engine/Plugins/Enterprise/GLTFExporter/MigratedFromGit/Source/GLTFExporter/Private/Tasks/GLTFMaterialTasks.cpp
+++ b/Engine/Plugins/Enterprise/GLTFExporter/MigratedFromGit/Source/GLTFExporter/Private/Tasks/GLTFMaterialTasks.cpp
@@ -303,7 +303,7 @@ bool FGLTFMaterialTask::TryGetBaseColorAndOpacity(FGLTFJsonPBRMetallicRoughness&
 	}
 
 	// TODO: add support for calculating the ideal resolution to use for baking based on connected (texture) nodes
-	FIntPoint TextureSize = Builder.ExportOptions->GetDefaultMaterialBakeSize();
+	FIntPoint TextureSize = Builder.GetDefaultMaterialBakeSize();
 
 	// TODO: should this be the default wrap-mode?
 	EGLTFJsonTextureWrap TextureWrapS = EGLTFJsonTextureWrap::Repeat;
@@ -470,7 +470,7 @@ bool FGLTFMaterialTask::TryGetMetallicAndRoughness(FGLTFJsonPBRMetallicRoughness
 	}
 
 	// TODO: add support for calculating the ideal resolution to use for baking based on connected (texture) nodes
-	FIntPoint TextureSize = Builder.ExportOptions->GetDefaultMaterialBakeSize();
+	FIntPoint TextureSize = Builder.GetDefaultMaterialBakeSize();
 
 	// TODO: should this be the default wrap-mode?
 	EGLTFJsonTextureWrap TextureWrapS = EGLTFJsonTextureWrap::Repeat;
@@ -636,7 +636,7 @@ bool FGLTFMaterialTask::TryGetClearCoatRoughness(FGLTFJsonClearCoatExtension& Ou
 	}
 
 	// TODO: add support for calculating the ideal resolution to use for baking based on connected (texture) nodes
-	FIntPoint TextureSize = Builder.ExportOptions->GetDefaultMaterialBakeSize();
+	FIntPoint TextureSize = Builder.GetDefaultMaterialBakeSize();
 
 	// TODO: should this be the default wrap-mode?
 	EGLTFJsonTextureWrap TextureWrapS = EGLTFJsonTextureWrap::Repeat;
@@ -1353,7 +1353,7 @@ FGLTFPropertyBakeOutput FGLTFMaterialTask::BakeMaterialProperty(EMaterialPropert
 		OutTexCoord = 0; // assume TexCoord0 even thought property seems to be texture coordinate independent
 	}
 
-	const FIntPoint DefaultTextureSize = Builder.ExportOptions->GetDefaultMaterialBakeSize();
+	const FIntPoint DefaultTextureSize = Builder.GetDefaultMaterialBakeSize();
 	const FIntPoint TextureSize = PreferredTextureSize != nullptr ? *PreferredTextureSize : DefaultTextureSize;
 
 	FMeshDescription MeshDescription;
diff --git a/Engine/Plugins/Enterprise/GLTFExporter/MigratedFromGit/Source/GLTFExporter/Public/GLTFExportOptions.h b/Engine/Plugins/Enterprise/GLTFExporter/MigratedFromGit/Source/GLTFExporter/Public/GLTFExportOptions.h
--- a/Engine/Plugins/Enterprise/GLTFExporter/MigratedFromGit/Source/GLTFExporter/Public/GLTFExportOptions.h
+++ b/Engine/Plugins/Enterprise/GLTFExporter/MigratedFromGit/Source/GLTFExporter/Public/GLTFExportOptions.h
@@ -153,8 +153,4 @@ public:
 	virtual bool CanEditChange(const FProperty* InProperty) const override;
 
 	virtual void PostEditChangeProperty(FPropertyChangedEvent& PropertyChangedEvent) override;
-
-	FIntPoint GetDefaultMaterialBakeSize() const;
-
-	bool ShouldExportLight(EComponentMobility::Type LightMobility) const;
 };
