[GLTFExporter] feat: add support for RGBE texture encoding

commit 61e536c93dbd8ffb5cea6acf784185cacba617e6
Date: Tue, 27 Oct 2020 12:42:14 +0100

    feat: add support for RGBE texture encoding
    (cherry picked from commit d48eb19dc95877b767e7dbbe99d27a66e7df0225​)

#skipci

diff --git a/Engine/Plugins/Enterprise/GLTFExporter/MigratedFromGit/Source/GLTFExporter/Private/Builders/GLTFBuilder.cpp b/Engine/Plugins/Enterprise/GLTFExporter/MigratedFromGit/Source/GLTFExporter/Private/Builders/GLTFBuilder.cpp
--- a/Engine/Plugins/Enterprise/GLTFExporter/MigratedFromGit/Source/GLTFExporter/Private/Builders/GLTFBuilder.cpp
+++ b/Engine/Plugins/Enterprise/GLTFExporter/MigratedFromGit/Source/GLTFExporter/Private/Builders/GLTFBuilder.cpp
@@ -17,6 +17,19 @@ FIntPoint FGLTFBuilder::GetDefaultMaterialBakeSize() const
 	return { Size, Size };
 }
 
+EGLTFJsonHDREncoding FGLTFBuilder::GetTextureHDREncoding() const
+{
+	switch (ExportOptions->TextureHDREncoding)
+	{
+		case EGLTFExporterTextureHDREncoding::None: return EGLTFJsonHDREncoding::None;
+		case EGLTFExporterTextureHDREncoding::RGBM: return EGLTFJsonHDREncoding::RGBM;
+		case EGLTFExporterTextureHDREncoding::RGBE: return EGLTFJsonHDREncoding::RGBE;
+		default:
+			checkNoEntry();
+			return EGLTFJsonHDREncoding::None;
+	}
+}
+
 bool FGLTFBuilder::ShouldExportLight(EComponentMobility::Type LightMobility) const
 {
 	switch (ExportOptions->ExportLights)
diff --git a/Engine/Plugins/Enterprise/GLTFExporter/MigratedFromGit/Source/GLTFExporter/Private/Builders/GLTFBuilder.h b/Engine/Plugins/Enterprise/GLTFExporter/MigratedFromGit/Source/GLTFExporter/Private/Builders/GLTFBuilder.h
--- a/Engine/Plugins/Enterprise/GLTFExporter/MigratedFromGit/Source/GLTFExporter/Private/Builders/GLTFBuilder.h
+++ b/Engine/Plugins/Enterprise/GLTFExporter/MigratedFromGit/Source/GLTFExporter/Private/Builders/GLTFBuilder.h
@@ -3,6 +3,7 @@
 #pragma once
 
 #include "GLTFExportOptions.h"
+#include "Json/GLTFJsonEnums.h"
 
 class FGLTFBuilder
 {
@@ -22,5 +23,7 @@ public:
 
 	FIntPoint GetDefaultMaterialBakeSize() const;
 
+	EGLTFJsonHDREncoding GetTextureHDREncoding() const;
+
 	bool ShouldExportLight(EComponentMobility::Type LightMobility) const;
 };
diff --git a/Engine/Plugins/Enterprise/GLTFExporter/MigratedFromGit/Source/GLTFExporter/Private/Converters/GLTFTextureUtility.cpp b/Engine/Plugins/Enterprise/GLTFExporter/MigratedFromGit/Source/GLTFExporter/Private/Converters/GLTFTextureUtility.cpp
--- a/Engine/Plugins/Enterprise/GLTFExporter/MigratedFromGit/Source/GLTFExporter/Private/Converters/GLTFTextureUtility.cpp
+++ b/Engine/Plugins/Enterprise/GLTFExporter/MigratedFromGit/Source/GLTFExporter/Private/Converters/GLTFTextureUtility.cpp
@@ -224,7 +224,7 @@ UTexture2D* FGLTFTextureUtility::CreateTextureFromCubeFace(const UTextureRenderT
 	return FaceTexture;
 }
 
-bool FGLTFTextureUtility::ReadPixels(const UTextureRenderTarget2D* InRenderTarget, TArray<FColor>& OutPixels)
+bool FGLTFTextureUtility::ReadPixels(const UTextureRenderTarget2D* InRenderTarget, TArray<FColor>& OutPixels, EGLTFJsonHDREncoding Encoding)
 {
 	FTextureRenderTarget2DResource* Resource = static_cast<FTextureRenderTarget2DResource*>(InRenderTarget->Resource);
 	if (Resource == nullptr)
@@ -232,15 +232,9 @@ bool FGLTFTextureUtility::ReadPixels(const UTextureRenderTarget2D* InRenderTarge
 		return false;
 	}
 
-	return Resource->ReadPixels(OutPixels);
-}
-
-bool FGLTFTextureUtility::ReadEncodedPixels(const UTextureRenderTarget2D* InRenderTarget, TArray<FColor>& OutPixels, EGLTFJsonHDREncoding Encoding)
-{
-	FTextureRenderTarget2DResource* Resource = static_cast<FTextureRenderTarget2DResource*>(InRenderTarget->Resource);
-	if (Resource == nullptr)
+	if (Encoding == EGLTFJsonHDREncoding::None)
 	{
-		return false;
+		return Resource->ReadPixels(OutPixels);
 	}
 
 	TArray<FLinearColor> HDRPixels;
@@ -254,9 +248,9 @@ bool FGLTFTextureUtility::ReadEncodedPixels(const UTextureRenderTarget2D* InRend
 		case EGLTFJsonHDREncoding::RGBM:
 			EncodeRGBM(HDRPixels, OutPixels);
 			break;
-
-		// TODO: add encoding for RGBE
-
+		case EGLTFJsonHDREncoding::RGBE:
+			EncodeRGBE(HDRPixels, OutPixels);
+			break;
 		default:
 			checkNoEntry();
 			break;
@@ -265,37 +259,41 @@ bool FGLTFTextureUtility::ReadEncodedPixels(const UTextureRenderTarget2D* InRend
 	return true;
 }
 
-FColor FGLTFTextureUtility::EncodeRGBM(const FLinearColor& Color, float MaxRange)
+void FGLTFTextureUtility::EncodeRGBM(const TArray<FLinearColor>& InPixels, TArray<FColor>& OutPixels, float MaxRange)
 {
-	// Based on PlayCanvas modified RGBM encoding.
+	OutPixels.AddUninitialized(InPixels.Num());
 
-	FLinearColor RGBM;
+	for (int32 Index = 0; Index < InPixels.Num(); ++Index)
+	{
+		const FLinearColor& Color = InPixels[Index];
+		FLinearColor RGBM;
 
-	RGBM.R = FMath::Sqrt(Color.R);
-	RGBM.G = FMath::Sqrt(Color.G);
-	RGBM.B = FMath::Sqrt(Color.B);
+		RGBM.R = FMath::Sqrt(Color.R);
+		RGBM.G = FMath::Sqrt(Color.G);
+		RGBM.B = FMath::Sqrt(Color.B);
 
-	RGBM.R /= MaxRange;
-	RGBM.G /= MaxRange;
-	RGBM.B /= MaxRange;
+		RGBM.R /= MaxRange;
+		RGBM.G /= MaxRange;
+		RGBM.B /= MaxRange;
 
-	RGBM.A = FMath::Max(FMath::Max(RGBM.R, RGBM.G), FMath::Max(RGBM.B, 1.0f / 255.0f));
-	RGBM.A = FMath::CeilToFloat(RGBM.A * 255.0f) / 255.0f;
+		RGBM.A = FMath::Max(FMath::Max(RGBM.R, RGBM.G), FMath::Max(RGBM.B, 1.0f / 255.0f));
+		RGBM.A = FMath::CeilToFloat(RGBM.A * 255.0f) / 255.0f;
 
-	RGBM.R /= RGBM.A;
-	RGBM.G /= RGBM.A;
-	RGBM.B /= RGBM.A;
+		RGBM.R /= RGBM.A;
+		RGBM.G /= RGBM.A;
+		RGBM.B /= RGBM.A;
 
-	return RGBM.ToFColor(false);
+		OutPixels[Index] = RGBM.ToFColor(false);
+	}
 }
 
-void FGLTFTextureUtility::EncodeRGBM(const TArray<FLinearColor>& InPixels, TArray<FColor>& OutPixels, float MaxRange)
+void FGLTFTextureUtility::EncodeRGBE(const TArray<FLinearColor>& InPixels, TArray<FColor>& OutPixels)
 {
 	OutPixels.AddUninitialized(InPixels.Num());
 
 	for (int32 Index = 0; Index < InPixels.Num(); ++Index)
 	{
-		OutPixels[Index] = EncodeRGBM(InPixels[Index], MaxRange);
+		OutPixels[Index] = InPixels[Index].ToRGBE();
 	}
 }
 
diff --git a/Engine/Plugins/Enterprise/GLTFExporter/MigratedFromGit/Source/GLTFExporter/Private/Converters/GLTFTextureUtility.h b/Engine/Plugins/Enterprise/GLTFExporter/MigratedFromGit/Source/GLTFExporter/Private/Converters/GLTFTextureUtility.h
--- a/Engine/Plugins/Enterprise/GLTFExporter/MigratedFromGit/Source/GLTFExporter/Private/Converters/GLTFTextureUtility.h
+++ b/Engine/Plugins/Enterprise/GLTFExporter/MigratedFromGit/Source/GLTFExporter/Private/Converters/GLTFTextureUtility.h
@@ -31,11 +31,10 @@ struct FGLTFTextureUtility
 	static UTexture2D* CreateTextureFromCubeFace(const UTextureCube* TextureCube, ECubeFace CubeFace);
 	static UTexture2D* CreateTextureFromCubeFace(const UTextureRenderTargetCube* RenderTargetCube, ECubeFace CubeFace);
 
-	static bool ReadPixels(const UTextureRenderTarget2D* InRenderTarget, TArray<FColor>& OutPixels);
-	static bool ReadEncodedPixels(const UTextureRenderTarget2D* InRenderTarget, TArray<FColor>& OutPixels, EGLTFJsonHDREncoding Encoding);
+	static bool ReadPixels(const UTextureRenderTarget2D* InRenderTarget, TArray<FColor>& OutPixels, EGLTFJsonHDREncoding Encoding);
 
-	static FColor EncodeRGBM(const FLinearColor& Color, float MaxRange = 8);
 	static void EncodeRGBM(const TArray<FLinearColor>& InPixels, TArray<FColor>& OutPixels, float MaxRange = 8);
+	static void EncodeRGBE(const TArray<FLinearColor>& InPixels, TArray<FColor>& OutPixels);
 
 	// TODO: maybe use template specialization to avoid the need for duplicated functions
 	static bool LoadPlatformData(UTexture2D* Texture);
diff --git a/Engine/Plugins/Enterprise/GLTFExporter/MigratedFromGit/Source/GLTFExporter/Private/Json/GLTFJsonEnums.h b/Engine/Plugins/Enterprise/GLTFExporter/MigratedFromGit/Source/GLTFExporter/Private/Json/GLTFJsonEnums.h
--- a/Engine/Plugins/Enterprise/GLTFExporter/MigratedFromGit/Source/GLTFExporter/Private/Json/GLTFJsonEnums.h
+++ b/Engine/Plugins/Enterprise/GLTFExporter/MigratedFromGit/Source/GLTFExporter/Private/Json/GLTFJsonEnums.h
@@ -33,8 +33,8 @@ enum class EGLTFJsonShadingModel
 enum class EGLTFJsonHDREncoding
 {
 	None = -1,
-	RGBE,
-	RGBM
+	RGBM,
+	RGBE
 };
 
 enum class EGLTFJsonCubeFace
diff --git a/Engine/Plugins/Enterprise/GLTFExporter/MigratedFromGit/Source/GLTFExporter/Private/Tasks/GLTFTextureTasks.cpp b/Engine/Plugins/Enterprise/GLTFExporter/MigratedFromGit/Source/GLTFExporter/Private/Tasks/GLTFTextureTasks.cpp
--- a/Engine/Plugins/Enterprise/GLTFExporter/MigratedFromGit/Source/GLTFExporter/Private/Tasks/GLTFTextureTasks.cpp
+++ b/Engine/Plugins/Enterprise/GLTFExporter/MigratedFromGit/Source/GLTFExporter/Private/Tasks/GLTFTextureTasks.cpp
@@ -26,28 +26,23 @@ void FGLTFTexture2DTask::Complete()
 	// TODO: preserve maximum image quality (avoid compression artifacts) by copying source data (and adjustments) to a temp texture
 	FGLTFTextureUtility::DrawTexture(RenderTarget, Texture2D);
 
-	TArray<FColor> Pixels;
 	if (!Texture2D->IsNormalMap() && FGLTFTextureUtility::IsHDRFormat(RenderTarget->GetFormat()))
 	{
-		JsonTexture.Encoding = EGLTFJsonHDREncoding::RGBM; // TODO: use only encoding as specified by export options
-		FGLTFTextureUtility::ReadEncodedPixels(RenderTarget, Pixels, JsonTexture.Encoding);
+		JsonTexture.Encoding = Builder.GetTextureHDREncoding();
 	}
-	else
-	{
-		FGLTFTextureUtility::ReadPixels(RenderTarget, Pixels);
 
-		if (Texture2D->IsNormalMap())
-		{
-			FGLTFTextureUtility::FlipGreenChannel(Pixels);
-		}
-	}
-
-	if (Pixels.Num() == 0)
+	TArray<FColor> Pixels;
+	if (!FGLTFTextureUtility::ReadPixels(RenderTarget, Pixels, JsonTexture.Encoding))
 	{
 		Builder.AddWarningMessage(FString::Printf(TEXT("Failed to read pixels for 2D texture %s"), *JsonTexture.Name));
 		return;
 	}
 
+	if (Texture2D->IsNormalMap())
+	{
+		FGLTFTextureUtility::FlipGreenChannel(Pixels);
+	}
+
 	JsonTexture.Source = Builder.AddImage(Pixels.GetData(), Size, JsonTexture.Name);
 	JsonTexture.Sampler = Builder.GetOrAddSampler(Texture2D);
 }
@@ -74,18 +69,13 @@ void FGLTFTextureCubeTask::Complete()
 	const float FaceRotation = FGLTFTextureUtility::GetCubeFaceRotation(CubeFace);
 	FGLTFTextureUtility::RotateTexture(RenderTarget, FaceTexture, FaceRotation);
 
-	TArray<FColor> Pixels;
 	if (FGLTFTextureUtility::IsHDRFormat(RenderTarget->GetFormat()))
 	{
-		JsonTexture.Encoding = EGLTFJsonHDREncoding::RGBM; // TODO: use only encoding as specified by export options
-		FGLTFTextureUtility::ReadEncodedPixels(RenderTarget, Pixels, JsonTexture.Encoding);
-	}
-	else
-	{
-		FGLTFTextureUtility::ReadPixels(RenderTarget, Pixels);
+		JsonTexture.Encoding = Builder.GetTextureHDREncoding();
 	}
 
-	if (Pixels.Num() == 0)
+	TArray<FColor> Pixels;
+	if (!FGLTFTextureUtility::ReadPixels(RenderTarget, Pixels, JsonTexture.Encoding))
 	{
 		Builder.AddWarningMessage(FString::Printf(TEXT("Failed to read pixels for cubemap texture %s"), *JsonTexture.Name));
 		return;
@@ -100,18 +90,13 @@ void FGLTFTextureRenderTarget2DTask::Complete()
 	FGLTFJsonTexture& JsonTexture = Builder.GetTexture(TextureIndex);
 	RenderTarget2D->GetName(JsonTexture.Name);
 
-	TArray<FColor> Pixels;
 	if (FGLTFTextureUtility::IsHDRFormat(RenderTarget2D->GetFormat()))
 	{
-		JsonTexture.Encoding = EGLTFJsonHDREncoding::RGBM; // TODO: use only encoding as specified by export options
-		FGLTFTextureUtility::ReadEncodedPixels(RenderTarget2D, Pixels, JsonTexture.Encoding);
-	}
-	else
-	{
-		FGLTFTextureUtility::ReadPixels(RenderTarget2D, Pixels);
+		JsonTexture.Encoding = Builder.GetTextureHDREncoding();
 	}
 
-	if (Pixels.Num() == 0)
+	TArray<FColor> Pixels;
+	if (!FGLTFTextureUtility::ReadPixels(RenderTarget2D, Pixels, JsonTexture.Encoding))
 	{
 		Builder.AddWarningMessage(FString::Printf(TEXT("Failed to read pixels for 2D render target %s"), *JsonTexture.Name));
 		return;
@@ -142,18 +127,13 @@ void FGLTFTextureRenderTargetCubeTask::Complete()
 	const float FaceRotation = FGLTFTextureUtility::GetCubeFaceRotation(CubeFace);
 	FGLTFTextureUtility::RotateTexture(RenderTarget, FaceTexture, FaceRotation);
 
-	TArray<FColor> Pixels;
 	if (FGLTFTextureUtility::IsHDRFormat(RenderTarget->GetFormat()))
 	{
-		JsonTexture.Encoding = EGLTFJsonHDREncoding::RGBM; // TODO: use only encoding as specified by export options
-		FGLTFTextureUtility::ReadEncodedPixels(RenderTarget, Pixels, JsonTexture.Encoding);
-	}
-	else
-	{
-		FGLTFTextureUtility::ReadPixels(RenderTarget, Pixels);
+		JsonTexture.Encoding = Builder.GetTextureHDREncoding();
 	}
 
-	if (Pixels.Num() == 0)
+	TArray<FColor> Pixels;
+	if (!FGLTFTextureUtility::ReadPixels(RenderTarget, Pixels, JsonTexture.Encoding))
 	{
 		Builder.AddWarningMessage(FString::Printf(TEXT("Failed to read pixels for cubemap render target %s"), *JsonTexture.Name));
 		return;
