[GLTFExporter] feat: add support for configuring HDR encoding in export options

commit 69fa87422188cd6eb9980acad938b0d291ad31a4
Date: Tue, 27 Oct 2020 12:43:01 +0100

    feat: add support for configuring HDR encoding in export options
    (cherry picked from commit 3212630b08500856388fe64f617e09a61ae887b9​)

#skipci

diff --git a/Engine/Plugins/Enterprise/GLTFExporter/MigratedFromGit/Source/GLTFExporter/Private/GLTFExportOptions.cpp b/Engine/Plugins/Enterprise/GLTFExporter/MigratedFromGit/Source/GLTFExporter/Private/GLTFExportOptions.cpp
--- a/Engine/Plugins/Enterprise/GLTFExporter/MigratedFromGit/Source/GLTFExporter/Private/GLTFExportOptions.cpp
+++ b/Engine/Plugins/Enterprise/GLTFExporter/MigratedFromGit/Source/GLTFExporter/Private/GLTFExportOptions.cpp
@@ -51,8 +51,7 @@ bool UGLTFExportOptions::CanEditChange(const FProperty* InProperty) const
 	const FName PropertyFName = InProperty->GetFName();
 
 	// TODO: remove options that have been implemented in exporters/converters.
-	if (PropertyFName == GET_MEMBER_NAME_CHECKED(ThisClass, TextureHDREncoding) ||
-		PropertyFName == GET_MEMBER_NAME_CHECKED(ThisClass, bExportHiddenInGame))
+	if (PropertyFName == GET_MEMBER_NAME_CHECKED(ThisClass, bExportHiddenInGame))
 	{
 		return false;
 	}
diff --git a/Engine/Plugins/Enterprise/GLTFExporter/MigratedFromGit/Source/GLTFExporter/Public/GLTFExportOptions.h b/Engine/Plugins/Enterprise/GLTFExporter/MigratedFromGit/Source/GLTFExporter/Public/GLTFExportOptions.h
--- a/Engine/Plugins/Enterprise/GLTFExporter/MigratedFromGit/Source/GLTFExporter/Public/GLTFExportOptions.h
+++ b/Engine/Plugins/Enterprise/GLTFExporter/MigratedFromGit/Source/GLTFExporter/Public/GLTFExportOptions.h
@@ -29,8 +29,7 @@ enum class EGLTFExporterTextureHDREncoding : uint8
 {
 	None,
 	RGBM,
-	RGBE,
-	HDR
+	RGBE
 };
 
 UENUM(BlueprintType)
