[GLTFExporter] feat(backdrop): implement and export angle property for hdri backdrops

commit ad5d4fbc27538824d6652c0f0a0af741ac7ff690
Date: Tue, 27 Oct 2020 10:24:11 +0100

    feat(backdrop): implement and export angle property for hdri backdrops
    (cherry picked from commit 5ffea2f51ee113d36f4d032f3b933d297d2d6a72​)

#skipci

diff --git a/Engine/Plugins/Enterprise/GLTFExporter/MigratedFromGit/Source/GLTFExporter/Private/Converters/GLTFBackdropConverters.cpp b/Engine/Plugins/Enterprise/GLTFExporter/MigratedFromGit/Source/GLTFExporter/Private/Converters/GLTFBackdropConverters.cpp
--- a/Engine/Plugins/Enterprise/GLTFExporter/MigratedFromGit/Source/GLTFExporter/Private/Converters/GLTFBackdropConverters.cpp
+++ b/Engine/Plugins/Enterprise/GLTFExporter/MigratedFromGit/Source/GLTFExporter/Private/Converters/GLTFBackdropConverters.cpp
@@ -17,6 +17,16 @@ FGLTFJsonBackdropIndex FGLTFBackdropConverter::Convert(const AActor* BackdropAct
 	FGLTFJsonBackdrop JsonBackdrop;
 	BackdropActor->GetName(JsonBackdrop.Name);
 
+	const USceneComponent* SceneComponent = BackdropActor->GetRootComponent();
+	const FRotator Rotation = SceneComponent->GetComponentRotation();
+
+	// NOTE: when calculating map rotation in the HDRI_Attributes material function using RotateAboutAxis,
+	// an angle in radians is used as input. But RotateAboutAxis expects a value of 1 per full revolution, not 2*PI.
+	// The end result is that the map rotation is always 2*PI more than the actual yaw of the scene component.
+
+	// TODO: remove the 2*PI multiplier if HDRI_Attributes is updated to match map rotation with yaw
+	JsonBackdrop.Angle = FRotator::ClampAxis(Rotation.Yaw * 2.0f * PI);
+
 	const UStaticMesh* Mesh;
 	if (FGLTFActorUtility::TryGetPropertyValue(BackdropActor, TEXT("Mesh"), Mesh))
 	{
diff --git a/Engine/Plugins/Enterprise/GLTFExporter/MigratedFromGit/Source/GLTFExporter/Private/Json/GLTFJsonBackdrop.h b/Engine/Plugins/Enterprise/GLTFExporter/MigratedFromGit/Source/GLTFExporter/Private/Json/GLTFJsonBackdrop.h
--- a/Engine/Plugins/Enterprise/GLTFExporter/MigratedFromGit/Source/GLTFExporter/Private/Json/GLTFJsonBackdrop.h
+++ b/Engine/Plugins/Enterprise/GLTFExporter/MigratedFromGit/Source/GLTFExporter/Private/Json/GLTFJsonBackdrop.h
@@ -15,6 +15,7 @@ struct FGLTFJsonBackdrop
 
 	float Intensity;
 	float Size;
+	float Angle;
 
 	FGLTFJsonVector3 ProjectionCenter;
 
@@ -24,6 +25,7 @@ struct FGLTFJsonBackdrop
 	FGLTFJsonBackdrop()
 		: Intensity(0)
 		, Size(0)
+		, Angle(0)
 		, ProjectionCenter(FGLTFJsonVector3::Zero)
 		, LightingDistanceFactor(0)
 		, UseCameraProjection(false)
@@ -50,6 +52,11 @@ struct FGLTFJsonBackdrop
 		FGLTFJsonUtility::WriteExactValue(JsonWriter, TEXT("intensity"), Intensity);
 		FGLTFJsonUtility::WriteExactValue(JsonWriter, TEXT("size"), Size);
 
+		if (!FMath::IsNearlyZero(Angle))
+		{
+			FGLTFJsonUtility::WriteExactValue(JsonWriter, TEXT("angle"), Angle);
+		}
+
 		JsonWriter.WriteIdentifierPrefix(TEXT("projectionCenter"));
 		ProjectionCenter.WriteArray(JsonWriter);
 
