[GLTFExporter] feat: implement support for export option bExportHiddenInGame

commit 69445a07309551fe97fac9dc370cb287dca80e4c
Date: Tue, 27 Oct 2020 21:03:24 +0100

    feat: implement support for export option bExportHiddenInGame
    (cherry picked from commit e1a8ce0b87d89cf75d1a58f1012468e3318b25da​)

#skipci

diff --git a/Engine/Plugins/Enterprise/GLTFExporter/MigratedFromGit/Source/GLTFExporter/Private/Converters/GLTFNodeConverters.cpp b/Engine/Plugins/Enterprise/GLTFExporter/MigratedFromGit/Source/GLTFExporter/Private/Converters/GLTFNodeConverters.cpp
--- a/Engine/Plugins/Enterprise/GLTFExporter/MigratedFromGit/Source/GLTFExporter/Private/Converters/GLTFNodeConverters.cpp
+++ b/Engine/Plugins/Enterprise/GLTFExporter/MigratedFromGit/Source/GLTFExporter/Private/Converters/GLTFNodeConverters.cpp
@@ -112,65 +112,62 @@ FGLTFJsonNodeIndex FGLTFComponentConverter::Convert(const USceneComponent* Scene
 	Node.Rotation = FGLTFConverterUtility::ConvertRotation(RelativeTransform.GetRotation());
 	Node.Scale = FGLTFConverterUtility::ConvertScale(RelativeTransform.GetScale3D());
 
-	if (SceneComponent->bHiddenInGame) // TODO: make this configurable
+	if (!SceneComponent->bHiddenInGame || Builder.ExportOptions->bExportHiddenInGame)
 	{
-		// ignore any visible properties
-	}
-	else if (const UStaticMeshComponent* StaticMeshComponent = Cast<UStaticMeshComponent>(SceneComponent))
-	{
-		Node.Mesh = Builder.GetOrAddMesh(StaticMeshComponent);
-
-		if (Builder.ExportOptions->bExportLightmaps)
+		if (const UStaticMeshComponent* StaticMeshComponent = Cast<UStaticMeshComponent>(SceneComponent))
 		{
-			Node.LightMap = Builder.GetOrAddLightMap(StaticMeshComponent);
-		}
-	}
-	else if (const USkeletalMeshComponent* SkeletalMeshComponent = Cast<USkeletalMeshComponent>(SceneComponent))
-	{
-		Node.Mesh = Builder.GetOrAddMesh(SkeletalMeshComponent);
+			Node.Mesh = Builder.GetOrAddMesh(StaticMeshComponent);
 
-		if (Builder.ExportOptions->bExportVertexSkinWeights)
-		{
-			// TODO: remove need for NodeIndex by adding support for cyclic calls in converter
-			const FGLTFJsonSkinIndex SkinIndex = Builder.GetOrAddSkin(NodeIndex, SkeletalMeshComponent);
-			if (SkinIndex != INDEX_NONE)
+			if (Builder.ExportOptions->bExportLightmaps)
 			{
-				Builder.GetNode(NodeIndex).Skin = SkinIndex;
+				Node.LightMap = Builder.GetOrAddLightMap(StaticMeshComponent);
+			}
+		}
+		else if (const USkeletalMeshComponent* SkeletalMeshComponent = Cast<USkeletalMeshComponent>(SceneComponent))
+		{
+			Node.Mesh = Builder.GetOrAddMesh(SkeletalMeshComponent);
 
-				if (Builder.ExportOptions->bExportAnimationSequences)
+			if (Builder.ExportOptions->bExportVertexSkinWeights)
+			{
+				// TODO: remove need for NodeIndex by adding support for cyclic calls in converter
+				const FGLTFJsonSkinIndex SkinIndex = Builder.GetOrAddSkin(NodeIndex, SkeletalMeshComponent);
+				if (SkinIndex != INDEX_NONE)
 				{
-					Builder.GetOrAddAnimation(NodeIndex, SkeletalMeshComponent);
+					Builder.GetNode(NodeIndex).Skin = SkinIndex;
+
+					if (Builder.ExportOptions->bExportAnimationSequences)
+					{
+						Builder.GetOrAddAnimation(NodeIndex, SkeletalMeshComponent);
+					}
 				}
 			}
 		}
-	}
-	else if (const UCameraComponent* CameraComponent = Cast<UCameraComponent>(SceneComponent))
-	{
-		if (Builder.ExportOptions->bExportCameras)
+		else if (const UCameraComponent* CameraComponent = Cast<UCameraComponent>(SceneComponent))
 		{
-			// TODO: conversion of camera direction should be done in separate converter
-			FGLTFJsonNode CameraNode;
-			CameraNode.Name = FGLTFNameUtility::GetName(CameraComponent);
-			CameraNode.Rotation = FGLTFConverterUtility::ConvertCameraDirection();
-			CameraNode.Camera = Builder.GetOrAddCamera(CameraComponent);
-			Builder.AddChildComponentNode(NodeIndex, CameraNode);
+			if (Builder.ExportOptions->bExportCameras)
+			{
+				// TODO: conversion of camera direction should be done in separate converter
+				FGLTFJsonNode CameraNode;
+				CameraNode.Name = FGLTFNameUtility::GetName(CameraComponent);
+				CameraNode.Rotation = FGLTFConverterUtility::ConvertCameraDirection();
+				CameraNode.Camera = Builder.GetOrAddCamera(CameraComponent);
+				Builder.AddChildComponentNode(NodeIndex, CameraNode);
+			}
 		}
-	}
-	else if (const ULightComponent* LightComponent = Cast<ULightComponent>(SceneComponent))
-	{
-		if (Builder.ShouldExportLight(LightComponent->Mobility))
+		else if (const ULightComponent* LightComponent = Cast<ULightComponent>(SceneComponent))
 		{
-			// TODO: conversion of light direction should be done in separate converter
-			FGLTFJsonNode LightNode;
-			LightNode.Name = FGLTFNameUtility::GetName(LightComponent);
-			LightNode.Rotation = FGLTFConverterUtility::ConvertLightDirection();
-			LightNode.Light = Builder.GetOrAddLight(LightComponent);
-			Builder.AddChildComponentNode(NodeIndex, LightNode);
+			if (Builder.ShouldExportLight(LightComponent->Mobility))
+			{
+				// TODO: conversion of light direction should be done in separate converter
+				FGLTFJsonNode LightNode;
+				LightNode.Name = FGLTFNameUtility::GetName(LightComponent);
+				LightNode.Rotation = FGLTFConverterUtility::ConvertLightDirection();
+				LightNode.Light = Builder.GetOrAddLight(LightComponent);
+				Builder.AddChildComponentNode(NodeIndex, LightNode);
+			}
 		}
 	}
 
-	// TODO: add support for SkyLight?
-
 	return NodeIndex;
 }
 
diff --git a/Engine/Plugins/Enterprise/GLTFExporter/MigratedFromGit/Source/GLTFExporter/Private/GLTFExportOptions.cpp b/Engine/Plugins/Enterprise/GLTFExporter/MigratedFromGit/Source/GLTFExporter/Private/GLTFExportOptions.cpp
--- a/Engine/Plugins/Enterprise/GLTFExporter/MigratedFromGit/Source/GLTFExporter/Private/GLTFExportOptions.cpp
+++ b/Engine/Plugins/Enterprise/GLTFExporter/MigratedFromGit/Source/GLTFExporter/Private/GLTFExportOptions.cpp
@@ -50,12 +50,6 @@ bool UGLTFExportOptions::CanEditChange(const FProperty* InProperty) const
 {
 	const FName PropertyFName = InProperty->GetFName();
 
-	// TODO: remove options that have been implemented in exporters/converters.
-	if (PropertyFName == GET_MEMBER_NAME_CHECKED(ThisClass, bExportHiddenInGame))
-	{
-		return false;
-	}
-
 	if (PropertyFName == GET_MEMBER_NAME_CHECKED(ThisClass, DefaultMaterialBakeSize))
 	{
 		return bBakeMaterialInputs;
