[GLTFExporter] fix: add error handling when can't write to file in zip utility

commit f46724fa09d533a845e16faf481b9b07df40d5d7
Date: Tue, 27 Oct 2020 22:20:36 +0100

    fix: add error handling when can't write to file in zip utility
    (cherry picked from commit cc595f2c46b803bd2a86b9834529c98fba7b1b62​)

#skipci

diff --git a/Engine/Plugins/Enterprise/GLTFExporter/MigratedFromGit/Source/GLTFExporter/Private/Builders/GLTFZipUtility.cpp b/Engine/Plugins/Enterprise/GLTFExporter/MigratedFromGit/Source/GLTFExporter/Private/Builders/GLTFZipUtility.cpp
--- a/Engine/Plugins/Enterprise/GLTFExporter/MigratedFromGit/Source/GLTFExporter/Private/Builders/GLTFZipUtility.cpp
+++ b/Engine/Plugins/Enterprise/GLTFExporter/MigratedFromGit/Source/GLTFExporter/Private/Builders/GLTFZipUtility.cpp
@@ -61,6 +61,11 @@ bool FGLTFZipUtility::ExtractCurrentFile(void* ZipFile, const FString& Destinati
 	else
 	{
 		FArchive* FileWriter = IFileManager::Get().CreateFileWriter(*DestinationFilePath);
+		if (FileWriter == nullptr)
+		{
+			Builder.AddErrorMessage(FString::Printf(TEXT("Can't write to file %s from zip archive"), *DestinationFilePath));
+			return false;
+		}
 
 		if (const int ErrorCode = unzOpenCurrentFile(ZipFile))
 		{
