[GLTFExporter] refactor: disable RGBE in export options because its missing support in the viewer

commit b4910df1e1e64437a721c74d4fe5f85651a48d98
Date: Thu, 29 Oct 2020 17:51:38 +0100

    refactor: disable RGBE in export options because its missing support in the viewer
    (cherry picked from commit b4ae3f4f31b99c82e9e6f411cfa93759e2cfa003​)

#skipci

diff --git a/Engine/Plugins/Enterprise/GLTFExporter/MigratedFromGit/Source/GLTFExporter/Private/Builders/GLTFBuilder.cpp b/Engine/Plugins/Enterprise/GLTFExporter/MigratedFromGit/Source/GLTFExporter/Private/Builders/GLTFBuilder.cpp
--- a/Engine/Plugins/Enterprise/GLTFExporter/MigratedFromGit/Source/GLTFExporter/Private/Builders/GLTFBuilder.cpp
+++ b/Engine/Plugins/Enterprise/GLTFExporter/MigratedFromGit/Source/GLTFExporter/Private/Builders/GLTFBuilder.cpp
@@ -23,7 +23,7 @@ EGLTFJsonHDREncoding FGLTFBuilder::GetTextureHDREncoding() const
 	{
 		case EGLTFExporterTextureHDREncoding::None: return EGLTFJsonHDREncoding::None;
 		case EGLTFExporterTextureHDREncoding::RGBM: return EGLTFJsonHDREncoding::RGBM;
-		case EGLTFExporterTextureHDREncoding::RGBE: return EGLTFJsonHDREncoding::RGBE;
+		// TODO: add more encodings (like RGBE) when viewer supports them
 		default:
 			checkNoEntry();
 			return EGLTFJsonHDREncoding::None;
diff --git a/Engine/Plugins/Enterprise/GLTFExporter/MigratedFromGit/Source/GLTFExporter/Public/GLTFExportOptions.h b/Engine/Plugins/Enterprise/GLTFExporter/MigratedFromGit/Source/GLTFExporter/Public/GLTFExportOptions.h
--- a/Engine/Plugins/Enterprise/GLTFExporter/MigratedFromGit/Source/GLTFExporter/Public/GLTFExportOptions.h
+++ b/Engine/Plugins/Enterprise/GLTFExporter/MigratedFromGit/Source/GLTFExporter/Public/GLTFExportOptions.h
@@ -28,8 +28,7 @@ UENUM(BlueprintType)
 enum class EGLTFExporterTextureHDREncoding : uint8
 {
 	None,
-	RGBM,
-	RGBE
+	RGBM
 };
 
 UENUM(BlueprintType)
