[GLTFExporter] fix: use compression settings to detect HDR textures instead of bytes per pixel (since the latter can have uneven bytes per channel)

commit eb8eaa9da19743bc9fad65c105d80146aed0e736
Date: Sat, 31 Oct 2020 23:26:29 +0100

    fix: use compression settings to detect HDR textures instead of bytes per pixel (since the latter can have uneven bytes per channel)
    (cherry picked from commit 6863a8fe34f3014162fcf44dc496d9ecfeeb963c​)

#skipci

diff --git a/Engine/Plugins/Enterprise/GLTFExporter/MigratedFromGit/Source/GLTFExporter/Private/Converters/GLTFTextureUtility.cpp b/Engine/Plugins/Enterprise/GLTFExporter/MigratedFromGit/Source/GLTFExporter/Private/Converters/GLTFTextureUtility.cpp
--- a/Engine/Plugins/Enterprise/GLTFExporter/MigratedFromGit/Source/GLTFExporter/Private/Converters/GLTFTextureUtility.cpp
+++ b/Engine/Plugins/Enterprise/GLTFExporter/MigratedFromGit/Source/GLTFExporter/Private/Converters/GLTFTextureUtility.cpp
@@ -4,9 +4,9 @@
 #include "IImageWrapper.h"
 #include "NormalMapPreview.h"
 
-bool FGLTFTextureUtility::IsHDRFormat(EPixelFormat Format)
+bool FGLTFTextureUtility::IsHDR(const UTexture* Texture)
 {
-	return CalculateImageBytes(1, 1, 0, Format) > 4;
+	return Texture->CompressionSettings == TC_HDR || Texture->CompressionSettings == TC_HDR_Compressed;
 }
 
 bool FGLTFTextureUtility::CanPNGCompressFormat(ETextureSourceFormat InFormat, ERGBFormat& OutFormat, uint32& OutBitDepth)
@@ -195,12 +195,16 @@ UTexture2D* FGLTFTextureUtility::CreateTextureFromCubeFace(const UTextureCube* T
 
 UTexture2D* FGLTFTextureUtility::CreateTextureFromCubeFace(const UTextureRenderTargetCube* RenderTargetCube, ECubeFace CubeFace)
 {
-	UTexture2D* FaceTexture;
-	const EPixelFormat Format = RenderTargetCube->GetFormat();
-	const FIntPoint Size(RenderTargetCube->SizeX, RenderTargetCube->SizeX);
 	FTextureRenderTargetCubeResource* Resource = static_cast<FTextureRenderTargetCubeResource*>(RenderTargetCube->Resource);
+	if (Resource == nullptr)
+	{
+		return nullptr;
+	}
+
+	const FIntPoint Size(RenderTargetCube->SizeX, RenderTargetCube->SizeX);
+	UTexture2D* FaceTexture;
 
-	if (IsHDRFormat(Format))
+	if (IsHDR(RenderTargetCube))
 	{
 		TArray<FFloat16Color> Pixels;
 		if (!Resource->ReadPixels(Pixels, FReadSurfaceDataFlags(RCM_UNorm, CubeFace)))
diff --git a/Engine/Plugins/Enterprise/GLTFExporter/MigratedFromGit/Source/GLTFExporter/Private/Converters/GLTFTextureUtility.h b/Engine/Plugins/Enterprise/GLTFExporter/MigratedFromGit/Source/GLTFExporter/Private/Converters/GLTFTextureUtility.h
--- a/Engine/Plugins/Enterprise/GLTFExporter/MigratedFromGit/Source/GLTFExporter/Private/Converters/GLTFTextureUtility.h
+++ b/Engine/Plugins/Enterprise/GLTFExporter/MigratedFromGit/Source/GLTFExporter/Private/Converters/GLTFTextureUtility.h
@@ -9,7 +9,7 @@ enum class ERGBFormat : int8;
 
 struct FGLTFTextureUtility
 {
-	static bool IsHDRFormat(EPixelFormat Format);
+	static bool IsHDR(const UTexture* Texture);
 
 	static bool CanPNGCompressFormat(ETextureSourceFormat InFormat, ERGBFormat& OutFormat, uint32& OutBitDepth);
 	static bool CanPNGCompressFormat(EPixelFormat InFormat, ERGBFormat& OutFormat, uint32& OutBitDepth);
diff --git a/Engine/Plugins/Enterprise/GLTFExporter/MigratedFromGit/Source/GLTFExporter/Private/Tasks/GLTFTextureTasks.cpp b/Engine/Plugins/Enterprise/GLTFExporter/MigratedFromGit/Source/GLTFExporter/Private/Tasks/GLTFTextureTasks.cpp
--- a/Engine/Plugins/Enterprise/GLTFExporter/MigratedFromGit/Source/GLTFExporter/Private/Tasks/GLTFTextureTasks.cpp
+++ b/Engine/Plugins/Enterprise/GLTFExporter/MigratedFromGit/Source/GLTFExporter/Private/Tasks/GLTFTextureTasks.cpp
@@ -17,16 +17,18 @@ void FGLTFTexture2DTask::Complete()
 	// and why bInForceLinearGamma must also be true.
 	const bool bForceLinearGamma = true;
 	const float TargetGamma = 2.2f;
+
+	const bool bIsHDR = FGLTFTextureUtility::IsHDR(Texture2D);
 	const FIntPoint Size = { Texture2D->GetSizeX(), Texture2D->GetSizeY() };
 
-	const EPixelFormat RenderTargetFormat = FGLTFTextureUtility::IsHDRFormat(Texture2D->GetPixelFormat()) ? PF_FloatRGBA : PF_B8G8R8A8;
+	const EPixelFormat RenderTargetFormat = bIsHDR ? PF_FloatRGBA : PF_B8G8R8A8;
 	UTextureRenderTarget2D* RenderTarget = FGLTFTextureUtility::CreateRenderTarget(Size, RenderTargetFormat, bForceLinearGamma);
 	RenderTarget->TargetGamma = TargetGamma;
 
 	// TODO: preserve maximum image quality (avoid compression artifacts) by copying source data (and adjustments) to a temp texture
 	FGLTFTextureUtility::DrawTexture(RenderTarget, Texture2D);
 
-	if (!Texture2D->IsNormalMap() && FGLTFTextureUtility::IsHDRFormat(RenderTarget->GetFormat()))
+	if (!Texture2D->IsNormalMap() && bIsHDR)
 	{
 		JsonTexture.Encoding = Builder.GetTextureHDREncoding();
 	}
@@ -62,14 +64,16 @@ void FGLTFTextureCubeTask::Complete()
 		return;
 	}
 
-	const FIntPoint Size = { FaceTexture->GetSizeX(), FaceTexture->GetSizeY() };
-	const EPixelFormat RenderTargetFormat = FGLTFTextureUtility::IsHDRFormat(FaceTexture->GetPixelFormat()) ? PF_FloatRGBA : PF_B8G8R8A8;
+	const bool bIsHDR = FGLTFTextureUtility::IsHDR(TextureCube);
+	const FIntPoint Size = { TextureCube->GetSizeX(), TextureCube->GetSizeY() };
+
+	const EPixelFormat RenderTargetFormat = bIsHDR ? PF_FloatRGBA : PF_B8G8R8A8;
 	UTextureRenderTarget2D* RenderTarget = FGLTFTextureUtility::CreateRenderTarget(Size, RenderTargetFormat, true);
 
 	const float FaceRotation = FGLTFTextureUtility::GetCubeFaceRotation(CubeFace);
 	FGLTFTextureUtility::RotateTexture(RenderTarget, FaceTexture, FaceRotation);
 
-	if (FGLTFTextureUtility::IsHDRFormat(RenderTarget->GetFormat()))
+	if (bIsHDR)
 	{
 		JsonTexture.Encoding = Builder.GetTextureHDREncoding();
 	}
@@ -90,7 +94,10 @@ void FGLTFTextureRenderTarget2DTask::Complete()
 	FGLTFJsonTexture& JsonTexture = Builder.GetTexture(TextureIndex);
 	RenderTarget2D->GetName(JsonTexture.Name);
 
-	if (FGLTFTextureUtility::IsHDRFormat(RenderTarget2D->GetFormat()))
+	const bool bIsHDR = FGLTFTextureUtility::IsHDR(RenderTarget2D);
+	const FIntPoint Size = { RenderTarget2D->SizeX, RenderTarget2D->SizeY };
+
+	if (bIsHDR)
 	{
 		JsonTexture.Encoding = Builder.GetTextureHDREncoding();
 	}
@@ -102,7 +109,7 @@ void FGLTFTextureRenderTarget2DTask::Complete()
 		return;
 	}
 
-	JsonTexture.Source = Builder.AddImage(Pixels.GetData(), { RenderTarget2D->SizeX, RenderTarget2D->SizeY }, JsonTexture.Name);
+	JsonTexture.Source = Builder.AddImage(Pixels.GetData(), Size, JsonTexture.Name);
 	JsonTexture.Sampler = Builder.GetOrAddSampler(RenderTarget2D);
 }
 
@@ -120,14 +127,16 @@ void FGLTFTextureRenderTargetCubeTask::Complete()
 		return;
 	}
 
-	const FIntPoint Size = { FaceTexture->GetSizeX(), FaceTexture->GetSizeY() };
-	const EPixelFormat RenderTargetFormat = FGLTFTextureUtility::IsHDRFormat(FaceTexture->GetPixelFormat()) ? PF_FloatRGBA : PF_B8G8R8A8;
+	const bool bIsHDR = FGLTFTextureUtility::IsHDR(RenderTargetCube);
+	const FIntPoint Size = { RenderTargetCube->SizeX, RenderTargetCube->SizeX };
+
+	const EPixelFormat RenderTargetFormat = bIsHDR ? PF_FloatRGBA : PF_B8G8R8A8;
 	UTextureRenderTarget2D* RenderTarget = FGLTFTextureUtility::CreateRenderTarget(Size, RenderTargetFormat, true);
 
 	const float FaceRotation = FGLTFTextureUtility::GetCubeFaceRotation(CubeFace);
 	FGLTFTextureUtility::RotateTexture(RenderTarget, FaceTexture, FaceRotation);
 
-	if (FGLTFTextureUtility::IsHDRFormat(RenderTarget->GetFormat()))
+	if (bIsHDR)
 	{
 		JsonTexture.Encoding = Builder.GetTextureHDREncoding();
 	}
