[GLTFExporter] refactor: improve export option UPROPERTY specifiers

commit 0249297b714c068de03f9fe3f6f02cf36b4fd1de
Date: Sun, 1 Nov 2020 00:09:08 +0100

    refactor: improve export option UPROPERTY specifiers
    (cherry picked from commit 04fb8737c07a8470efea56785dcde99eeef009d1​)
    # Conflicts from cherry-pick:
    #	Plugins/GLTFExporter/Source/GLTFExporter/Private/GLTFExportOptions.cpp
    #	Plugins/GLTFExporter/Source/GLTFExporter/Public/GLTFExportOptions.h

#skipci

diff --git a/Engine/Plugins/Enterprise/GLTFExporter/MigratedFromGit/Source/GLTFExporter/Private/GLTFExportOptions.cpp b/Engine/Plugins/Enterprise/GLTFExporter/MigratedFromGit/Source/GLTFExporter/Private/GLTFExportOptions.cpp
--- a/Engine/Plugins/Enterprise/GLTFExporter/MigratedFromGit/Source/GLTFExporter/Private/GLTFExportOptions.cpp
+++ b/Engine/Plugins/Enterprise/GLTFExporter/MigratedFromGit/Source/GLTFExporter/Private/GLTFExportOptions.cpp
@@ -45,48 +45,3 @@ void UGLTFExportOptions::ResetToDefault()
 	bAllExtensionsRequired = false;
 	bShowFilesWhenDone = true;
 }
-
-bool UGLTFExportOptions::CanEditChange(const FProperty* InProperty) const
-{
-	const FName PropertyFName = InProperty->GetFName();
-
-	if (PropertyFName == GET_MEMBER_NAME_CHECKED(ThisClass, DefaultMaterialBakeSize))
-	{
-		return bBakeMaterialInputs;
-	}
-
-	if (PropertyFName == GET_MEMBER_NAME_CHECKED(ThisClass, bExportAnimationSequences))
-	{
-		return bExportVertexSkinWeights;
-	}
-
-	if (PropertyFName == GET_MEMBER_NAME_CHECKED(ThisClass, bRetargetBoneTransforms))
-	{
-		return bExportVertexSkinWeights && bExportAnimationSequences;
-	}
-
-	if (PropertyFName == GET_MEMBER_NAME_CHECKED(ThisClass, bExportOrbitalCameras))
-	{
-		return bExportCameras;
-	}
-
-	if (PropertyFName == GET_MEMBER_NAME_CHECKED(ThisClass, bBakeMaterialInputsUsingMeshData))
-	{
-		return bBakeMaterialInputs;
-	}
-
-	return true;
-}
-
-void UGLTFExportOptions::PostEditChangeProperty(FPropertyChangedEvent& PropertyChangedEvent)
-{
-	const FName PropertyFName = PropertyChangedEvent.GetPropertyName();
-
-	if (PropertyFName == GET_MEMBER_NAME_CHECKED(ThisClass, DefaultLevelOfDetail))
-	{
-		if (DefaultLevelOfDetail < 0)
-		{
-			DefaultLevelOfDetail = 0;
-		}
-	}
-}
diff --git a/Engine/Plugins/Enterprise/GLTFExporter/MigratedFromGit/Source/GLTFExporter/Public/GLTFExportOptions.h b/Engine/Plugins/Enterprise/GLTFExporter/MigratedFromGit/Source/GLTFExporter/Public/GLTFExportOptions.h
--- a/Engine/Plugins/Enterprise/GLTFExporter/MigratedFromGit/Source/GLTFExporter/Public/GLTFExportOptions.h
+++ b/Engine/Plugins/Enterprise/GLTFExporter/MigratedFromGit/Source/GLTFExporter/Public/GLTFExportOptions.h
@@ -56,10 +56,10 @@ public:
 	UPROPERTY(EditAnywhere, BlueprintReadWrite, Config, Category = Material)
 	bool bBakeMaterialInputs;
 
-	UPROPERTY(EditAnywhere, BlueprintReadWrite, Config, Category = Material)
+	UPROPERTY(EditAnywhere, BlueprintReadWrite, Config, Category = Material, Meta = (EditCondition = "bBakeMaterialInputs"))
 	bool bBakeMaterialInputsUsingMeshData;
 
-	UPROPERTY(EditAnywhere, BlueprintReadWrite, Config, Category = Material)
+	UPROPERTY(EditAnywhere, BlueprintReadWrite, Config, Category = Material, Meta = (EditCondition = "bBakeMaterialInputs"))
 	EGLTFExporterMaterialBakeSize DefaultMaterialBakeSize;
 
 	UPROPERTY(EditAnywhere, BlueprintReadWrite, Config, Category = Mesh)
@@ -71,7 +71,7 @@ public:
 	UPROPERTY(EditAnywhere, BlueprintReadWrite, Config, Category = Mesh)
 	bool bQuantizeVertexTangents;
 
-	UPROPERTY(EditAnywhere, BlueprintReadWrite, Config, Category = Mesh)
+	UPROPERTY(EditAnywhere, BlueprintReadWrite, Config, Category = Mesh, Meta = (ClampMin = "0"))
 	int32 DefaultLevelOfDetail;
 
 	UPROPERTY(EditAnywhere, BlueprintReadWrite, Config, Category = Animation)
@@ -80,10 +80,10 @@ public:
 	UPROPERTY(EditAnywhere, BlueprintReadWrite, Config, Category = Animation)
 	bool bExportVertexSkinWeights;
 
-	UPROPERTY(EditAnywhere, BlueprintReadWrite, Config, Category = Animation)
+	UPROPERTY(EditAnywhere, BlueprintReadWrite, Config, Category = Animation, Meta = (EditCondition = "bExportVertexSkinWeights"))
 	bool bExportAnimationSequences;
 
-	UPROPERTY(EditAnywhere, BlueprintReadWrite, Config, Category = Animation)
+	UPROPERTY(EditAnywhere, BlueprintReadWrite, Config, Category = Animation, Meta = (EditCondition = "bExportVertexSkinWeights && bExportAnimationSequences"))
 	bool bRetargetBoneTransforms;
 
 	UPROPERTY(EditAnywhere, BlueprintReadWrite, Config, Category = Animation)
@@ -98,7 +98,7 @@ public:
 	UPROPERTY(EditAnywhere, BlueprintReadWrite, Config, Category = Texture)
 	EGLTFExporterTextureHDREncoding TextureHDREncoding;
 
-	UPROPERTY(EditAnywhere, BlueprintReadWrite, Config, Category = Scene, meta =
+	UPROPERTY(EditAnywhere, BlueprintReadWrite, Config, Category = Scene, Meta =
 		(DisplayName = "Export Uniform Scale",
 		 ToolTip = "Scale factor used for exporting assets, by default: 0.01, for conversion from centimeters(Unreal default) to meters(glTF)."))
 	float ExportScale;
@@ -112,7 +112,7 @@ public:
 	UPROPERTY(EditAnywhere, BlueprintReadWrite, Config, Category = Scene)
 	bool bExportCameras;
 
-	UPROPERTY(EditAnywhere, BlueprintReadWrite, Config, Category = Scene)
+	UPROPERTY(EditAnywhere, BlueprintReadWrite, Config, Category = Scene, Meta = (EditCondition = "bExportCameras"))
 	bool bExportOrbitalCameras;
 
 	UPROPERTY(EditAnywhere, BlueprintReadWrite, Config, Category = Scene)
@@ -139,16 +139,5 @@ public:
 	UPROPERTY(EditAnywhere, BlueprintReadWrite, Config, Category = Exporter)
 	bool bShowFilesWhenDone;
 
-	/* Set all the each property to the CDO value */
 	void ResetToDefault();
-
-	/* Save the each property to a local ini to retrieve the value the next time we call function LoadOptions() */
-	void SaveOptions();
-
-	/* Load the each property from a local ini which the value was store by the function SaveOptions() */
-	void LoadOptions();
-
-	virtual bool CanEditChange(const FProperty* InProperty) const override;
-
-	virtual void PostEditChangeProperty(FPropertyChangedEvent& PropertyChangedEvent) override;
 };
