[GLTFExporter] refactor: improve texture utility for creating render target

commit 75b1d10c572cebdc747687c96431cdf7345afa52
Date: Sun, 1 Nov 2020 10:01:54 +0100

    refactor: improve texture utility for creating render target
    (cherry picked from commit 1f552cb6ff3a0f71837cdade116da838b4cfd196​)

#skipci

diff --git a/Engine/Plugins/Enterprise/GLTFExporter/MigratedFromGit/Source/GLTFExporter/Private/Converters/GLTFTextureUtility.cpp b/Engine/Plugins/Enterprise/GLTFExporter/MigratedFromGit/Source/GLTFExporter/Private/Converters/GLTFTextureUtility.cpp
--- a/Engine/Plugins/Enterprise/GLTFExporter/MigratedFromGit/Source/GLTFExporter/Private/Converters/GLTFTextureUtility.cpp
+++ b/Engine/Plugins/Enterprise/GLTFExporter/MigratedFromGit/Source/GLTFExporter/Private/Converters/GLTFTextureUtility.cpp
@@ -99,7 +99,7 @@ TTuple<TextureAddress, TextureAddress> FGLTFTextureUtility::GetAddressXY(const U
 	return MakeTuple(AddressX, AddressY);
 }
 
-UTexture2D* FGLTFTextureUtility::CreateTransientTexture(const void* RawData, int64 ByteLength, const FIntPoint& Size, EPixelFormat Format, bool bUseSRGB)
+UTexture2D* FGLTFTextureUtility::CreateTransientTexture(const void* RawData, int64 ByteLength, const FIntPoint& Size, EPixelFormat Format, bool bSRGB)
 {
 	check(CalculateImageBytes(Size.X, Size.Y, 0, Format) == ByteLength);
 
@@ -110,7 +110,7 @@ UTexture2D* FGLTFTextureUtility::CreateTransientTexture(const void* RawData, int
 	FMemory::Memcpy(MipData, RawData, ByteLength);
 	Texture->PlatformData->Mips[0].BulkData.Unlock();
 
-	Texture->SRGB = bUseSRGB ? 1 : 0;
+	Texture->SRGB = bSRGB ? 1 : 0;
 	Texture->CompressionNone = true;
 	Texture->MipGenSettings = TMGS_NoMipmaps;
 
@@ -118,10 +118,20 @@ UTexture2D* FGLTFTextureUtility::CreateTransientTexture(const void* RawData, int
 	return Texture;
 }
 
-UTextureRenderTarget2D* FGLTFTextureUtility::CreateRenderTarget(const FIntPoint& Size, EPixelFormat Format, bool bInForceLinearGamma)
+UTextureRenderTarget2D* FGLTFTextureUtility::CreateRenderTarget(const FIntPoint& Size, bool bIsHDR)
 {
+	// TODO: instead of PF_FloatRGBA (i.e. RTF_RGBA16f) use PF_A32B32G32R32F (i.e. RTF_RGBA32f) to avoid accuracy loss
+	const EPixelFormat PixelFormat = bIsHDR ? PF_FloatRGBA : PF_B8G8R8A8;
+
+	// NOTE: both bForceLinearGamma and TargetGamma=2.2 seem necessary for exported images to match their source data.
+	// It's not entirely clear why gamma must be 2.2 (instead of 0.0) and why bInForceLinearGamma must also be true.
+	const bool bForceLinearGamma = true;
+	const float TargetGamma = 2.2f;
+
 	UTextureRenderTarget2D* RenderTarget = NewObject<UTextureRenderTarget2D>();
-	RenderTarget->InitCustomFormat(Size.X, Size.Y, Format, bInForceLinearGamma);
+	RenderTarget->InitCustomFormat(Size.X, Size.Y, PixelFormat, bForceLinearGamma);
+
+	RenderTarget->TargetGamma = TargetGamma;
 	return RenderTarget;
 }
 
diff --git a/Engine/Plugins/Enterprise/GLTFExporter/MigratedFromGit/Source/GLTFExporter/Private/Converters/GLTFTextureUtility.h b/Engine/Plugins/Enterprise/GLTFExporter/MigratedFromGit/Source/GLTFExporter/Private/Converters/GLTFTextureUtility.h
--- a/Engine/Plugins/Enterprise/GLTFExporter/MigratedFromGit/Source/GLTFExporter/Private/Converters/GLTFTextureUtility.h
+++ b/Engine/Plugins/Enterprise/GLTFExporter/MigratedFromGit/Source/GLTFExporter/Private/Converters/GLTFTextureUtility.h
@@ -21,9 +21,9 @@ struct FGLTFTextureUtility
 
 	static TTuple<TextureAddress, TextureAddress> GetAddressXY(const UTexture* Texture);
 
-	static UTexture2D* CreateTransientTexture(const void* RawData, int64 ByteLength, const FIntPoint& Size, EPixelFormat Format, bool bUseSRGB = false);
+	static UTexture2D* CreateTransientTexture(const void* RawData, int64 ByteLength, const FIntPoint& Size, EPixelFormat Format, bool bSRGB = false);
 
-	static UTextureRenderTarget2D* CreateRenderTarget(const FIntPoint& Size, EPixelFormat Format, bool bInForceLinearGamma = false);
+	static UTextureRenderTarget2D* CreateRenderTarget(const FIntPoint& Size, bool bIsHDR);
 
 	static bool DrawTexture(UTextureRenderTarget2D* OutTarget, const UTexture2D* InSource, const FMatrix& InTransform = FMatrix::Identity);
 	static bool RotateTexture(UTextureRenderTarget2D* OutTarget, const UTexture2D* InSource, float InDegrees);
diff --git a/Engine/Plugins/Enterprise/GLTFExporter/MigratedFromGit/Source/GLTFExporter/Private/Tasks/GLTFTextureTasks.cpp b/Engine/Plugins/Enterprise/GLTFExporter/MigratedFromGit/Source/GLTFExporter/Private/Tasks/GLTFTextureTasks.cpp
--- a/Engine/Plugins/Enterprise/GLTFExporter/MigratedFromGit/Source/GLTFExporter/Private/Tasks/GLTFTextureTasks.cpp
+++ b/Engine/Plugins/Enterprise/GLTFExporter/MigratedFromGit/Source/GLTFExporter/Private/Tasks/GLTFTextureTasks.cpp
@@ -11,19 +11,9 @@ void FGLTFTexture2DTask::Complete()
 	FGLTFJsonTexture& JsonTexture = Builder.GetTexture(TextureIndex);
 	Texture2D->GetName(JsonTexture.Name);
 
-	// TODO: both bForceLinearGamma and TargetGamma=2.2f seem to be necessary for the exported images to match the results
-	// from exporting using the texture's source or its platform-data.
-	// It's not entirely clear why gamma must be set to 2.2 instead of 0.0 (which should use the correct gamma anyway),
-	// and why bInForceLinearGamma must also be true.
-	const bool bForceLinearGamma = true;
-	const float TargetGamma = 2.2f;
-
 	const bool bIsHDR = FGLTFTextureUtility::IsHDR(Texture2D);
 	const FIntPoint Size = { Texture2D->GetSizeX(), Texture2D->GetSizeY() };
-
-	const EPixelFormat RenderTargetFormat = bIsHDR ? PF_FloatRGBA : PF_B8G8R8A8;
-	UTextureRenderTarget2D* RenderTarget = FGLTFTextureUtility::CreateRenderTarget(Size, RenderTargetFormat, bForceLinearGamma);
-	RenderTarget->TargetGamma = TargetGamma;
+	UTextureRenderTarget2D* RenderTarget = FGLTFTextureUtility::CreateRenderTarget(Size, bIsHDR);
 
 	// TODO: preserve maximum image quality (avoid compression artifacts) by copying source data (and adjustments) to a temp texture
 	FGLTFTextureUtility::DrawTexture(RenderTarget, Texture2D);
@@ -66,9 +56,7 @@ void FGLTFTextureCubeTask::Complete()
 
 	const bool bIsHDR = FGLTFTextureUtility::IsHDR(TextureCube);
 	const FIntPoint Size = { TextureCube->GetSizeX(), TextureCube->GetSizeY() };
-
-	const EPixelFormat RenderTargetFormat = bIsHDR ? PF_FloatRGBA : PF_B8G8R8A8;
-	UTextureRenderTarget2D* RenderTarget = FGLTFTextureUtility::CreateRenderTarget(Size, RenderTargetFormat, true);
+	UTextureRenderTarget2D* RenderTarget = FGLTFTextureUtility::CreateRenderTarget(Size, bIsHDR);
 
 	const float FaceRotation = FGLTFTextureUtility::GetCubeFaceRotation(CubeFace);
 	FGLTFTextureUtility::RotateTexture(RenderTarget, FaceTexture, FaceRotation);
@@ -129,9 +117,7 @@ void FGLTFTextureRenderTargetCubeTask::Complete()
 
 	const bool bIsHDR = FGLTFTextureUtility::IsHDR(RenderTargetCube);
 	const FIntPoint Size = { RenderTargetCube->SizeX, RenderTargetCube->SizeX };
-
-	const EPixelFormat RenderTargetFormat = bIsHDR ? PF_FloatRGBA : PF_B8G8R8A8;
-	UTextureRenderTarget2D* RenderTarget = FGLTFTextureUtility::CreateRenderTarget(Size, RenderTargetFormat, true);
+	UTextureRenderTarget2D* RenderTarget = FGLTFTextureUtility::CreateRenderTarget(Size, bIsHDR);
 
 	const float FaceRotation = FGLTFTextureUtility::GetCubeFaceRotation(CubeFace);
 	FGLTFTextureUtility::RotateTexture(RenderTarget, FaceTexture, FaceRotation);
