[GLTFExporter] refactor: sanity check number of pixels against expected image dimensions in GLTFImageBuilder

commit 521ef15df5218203b25eeb94813364c1e3565d1b
Date: Sun, 1 Nov 2020 10:33:38 +0100

    refactor: sanity check number of pixels against expected image dimensions in GLTFImageBuilder
    (cherry picked from commit a195fc7e6e2656450e1d36640fe4c42746c985f3​)

#skipci

diff --git a/Engine/Plugins/Enterprise/GLTFExporter/MigratedFromGit/Source/GLTFExporter/Private/Builders/GLTFImageBuilder.cpp b/Engine/Plugins/Enterprise/GLTFExporter/MigratedFromGit/Source/GLTFExporter/Private/Builders/GLTFImageBuilder.cpp
--- a/Engine/Plugins/Enterprise/GLTFExporter/MigratedFromGit/Source/GLTFExporter/Private/Builders/GLTFImageBuilder.cpp
+++ b/Engine/Plugins/Enterprise/GLTFExporter/MigratedFromGit/Source/GLTFExporter/Private/Builders/GLTFImageBuilder.cpp
@@ -62,10 +62,11 @@ FGLTFJsonImageIndex FGLTFImageBuilder::AddImage(const void* RawData, int64 ByteL
 	return AddImage(ImageData.GetData(), ImageData.Num(), EGLTFJsonMimeType::PNG, Name);
 }
 
-FGLTFJsonImageIndex FGLTFImageBuilder::AddImage(const FColor* Pixels, FIntPoint Size, const FString& Name)
+FGLTFJsonImageIndex FGLTFImageBuilder::AddImage(const TArray<FColor>& Pixels, FIntPoint Size, const FString& Name)
 {
+	check(Pixels.Num() == Size.X * Size.Y);
 	const int64 ByteLength = Size.X * Size.Y * sizeof(FColor);
-	return AddImage(Pixels, ByteLength, Size, ERGBFormat::BGRA, 8, Name);
+	return AddImage(Pixels.GetData(), ByteLength, Size, ERGBFormat::BGRA, 8, Name);
 }
 
 FString FGLTFImageBuilder::SaveImageToFile(const void* CompressedData, int64 CompressedByteLength, EGLTFJsonMimeType MimeType, const FString& Name)
diff --git a/Engine/Plugins/Enterprise/GLTFExporter/MigratedFromGit/Source/GLTFExporter/Private/Builders/GLTFImageBuilder.h b/Engine/Plugins/Enterprise/GLTFExporter/MigratedFromGit/Source/GLTFExporter/Private/Builders/GLTFImageBuilder.h
--- a/Engine/Plugins/Enterprise/GLTFExporter/MigratedFromGit/Source/GLTFExporter/Private/Builders/GLTFImageBuilder.h
+++ b/Engine/Plugins/Enterprise/GLTFExporter/MigratedFromGit/Source/GLTFExporter/Private/Builders/GLTFImageBuilder.h
@@ -18,7 +18,7 @@ public:
 	FGLTFJsonImageIndex AddImage(const void* CompressedData, int64 CompressedByteLength, EGLTFJsonMimeType MimeType, const FString& Name);
 	FGLTFJsonImageIndex AddImage(const void* RawData, int64 ByteLength, FIntPoint Size, ERGBFormat Format, int32 BitDepth, const FString& Name);
 
-	FGLTFJsonImageIndex AddImage(const FColor* Pixels, FIntPoint Size, const FString& Name);
+	FGLTFJsonImageIndex AddImage(const TArray<FColor>& Pixels, FIntPoint Size, const FString& Name);
 
 private:
 
diff --git a/Engine/Plugins/Enterprise/GLTFExporter/MigratedFromGit/Source/GLTFExporter/Private/Converters/GLTFMaterialUtility.cpp b/Engine/Plugins/Enterprise/GLTFExporter/MigratedFromGit/Source/GLTFExporter/Private/Converters/GLTFMaterialUtility.cpp
--- a/Engine/Plugins/Enterprise/GLTFExporter/MigratedFromGit/Source/GLTFExporter/Private/Converters/GLTFMaterialUtility.cpp
+++ b/Engine/Plugins/Enterprise/GLTFExporter/MigratedFromGit/Source/GLTFExporter/Private/Converters/GLTFMaterialUtility.cpp
@@ -384,7 +384,7 @@ FGLTFJsonTextureIndex FGLTFMaterialUtility::AddTexture(FGLTFConvertBuilder& Buil
 	FGLTFJsonTexture JsonTexture;
 	JsonTexture.Name = TextureName;
 	JsonTexture.Sampler = Builder.AddSampler(JsonSampler);
-	JsonTexture.Source = Builder.AddImage(Pixels.GetData(), TextureSize, TextureName);
+	JsonTexture.Source = Builder.AddImage(Pixels, TextureSize, TextureName);
 
 	return Builder.AddTexture(JsonTexture);
 }
diff --git a/Engine/Plugins/Enterprise/GLTFExporter/MigratedFromGit/Source/GLTFExporter/Private/Tasks/GLTFTextureTasks.cpp b/Engine/Plugins/Enterprise/GLTFExporter/MigratedFromGit/Source/GLTFExporter/Private/Tasks/GLTFTextureTasks.cpp
--- a/Engine/Plugins/Enterprise/GLTFExporter/MigratedFromGit/Source/GLTFExporter/Private/Tasks/GLTFTextureTasks.cpp
+++ b/Engine/Plugins/Enterprise/GLTFExporter/MigratedFromGit/Source/GLTFExporter/Private/Tasks/GLTFTextureTasks.cpp
@@ -35,7 +35,7 @@ void FGLTFTexture2DTask::Complete()
 		FGLTFTextureUtility::FlipGreenChannel(Pixels);
 	}
 
-	JsonTexture.Source = Builder.AddImage(Pixels.GetData(), Size, JsonTexture.Name);
+	JsonTexture.Source = Builder.AddImage(Pixels, Size, JsonTexture.Name);
 	JsonTexture.Sampler = Builder.GetOrAddSampler(Texture2D);
 }
 
@@ -73,7 +73,7 @@ void FGLTFTextureCubeTask::Complete()
 		return;
 	}
 
-	JsonTexture.Source = Builder.AddImage(Pixels.GetData(), Size, JsonTexture.Name);
+	JsonTexture.Source = Builder.AddImage(Pixels, Size, JsonTexture.Name);
 	JsonTexture.Sampler = Builder.GetOrAddSampler(TextureCube);
 }
 
@@ -97,7 +97,7 @@ void FGLTFTextureRenderTarget2DTask::Complete()
 		return;
 	}
 
-	JsonTexture.Source = Builder.AddImage(Pixels.GetData(), Size, JsonTexture.Name);
+	JsonTexture.Source = Builder.AddImage(Pixels, Size, JsonTexture.Name);
 	JsonTexture.Sampler = Builder.GetOrAddSampler(RenderTarget2D);
 }
 
@@ -134,7 +134,7 @@ void FGLTFTextureRenderTargetCubeTask::Complete()
 		return;
 	}
 
-	JsonTexture.Source = Builder.AddImage(Pixels.GetData(), Size, JsonTexture.Name);
+	JsonTexture.Source = Builder.AddImage(Pixels, Size, JsonTexture.Name);
 	JsonTexture.Sampler = Builder.GetOrAddSampler(RenderTargetCube);
 }
 
@@ -162,6 +162,7 @@ void FGLTFTextureLightMapTask::Complete()
 	}
 
 	const FIntPoint Size = { Source.GetSizeX(), Source.GetSizeY() };
+
 	const void* RawData = Source.LockMip(0);
 	JsonTexture.Source = Builder.AddImage(RawData, Source.CalcMipSize(0), Size, RGBFormat, BitDepth, JsonTexture.Name);
 	Source.UnlockMip(0);
