[GLTFExporter] fix: use pixel format instead of compression settings to detect HDR textures (since render targets don't use compression settings)

commit 53d7071579974aae77023544470ab62c21d1dcc9
Date: Sun, 1 Nov 2020 20:09:11 +0100

    fix: use pixel format instead of compression settings to detect HDR textures (since render targets don't use compression settings)
    (cherry picked from commit b581502af9aba241183a6deb3a2c307e0e673b34​)

#skipci

diff --git a/Engine/Plugins/Enterprise/GLTFExporter/MigratedFromGit/Source/GLTFExporter/Private/Converters/GLTFTextureUtility.cpp b/Engine/Plugins/Enterprise/GLTFExporter/MigratedFromGit/Source/GLTFExporter/Private/Converters/GLTFTextureUtility.cpp
--- a/Engine/Plugins/Enterprise/GLTFExporter/MigratedFromGit/Source/GLTFExporter/Private/Converters/GLTFTextureUtility.cpp
+++ b/Engine/Plugins/Enterprise/GLTFExporter/MigratedFromGit/Source/GLTFExporter/Private/Converters/GLTFTextureUtility.cpp
@@ -4,9 +4,28 @@
 #include "IImageWrapper.h"
 #include "NormalMapPreview.h"
 
-bool FGLTFTextureUtility::IsHDR(const UTexture* Texture)
+bool FGLTFTextureUtility::IsHDR(EPixelFormat Format)
 {
-	return Texture->CompressionSettings == TC_HDR || Texture->CompressionSettings == TC_HDR_Compressed;
+	switch (Format)
+	{
+		case PF_A32B32G32R32F:
+		case PF_FloatRGB:
+		case PF_FloatRGBA:
+		case PF_R32_FLOAT:
+		case PF_G16R16F:
+		case PF_G16R16F_FILTER:
+		case PF_G32R32F:
+		case PF_R16F:
+		case PF_R16F_FILTER:
+		case PF_FloatR11G11B10:
+		case PF_BC6H:
+		case PF_PLATFORM_HDR_0:
+		case PF_PLATFORM_HDR_1:
+		case PF_PLATFORM_HDR_2:
+			return true;
+		default:
+			return false;
+	}
 }
 
 bool FGLTFTextureUtility::CanPNGCompressFormat(ETextureSourceFormat InFormat, ERGBFormat& OutFormat, uint32& OutBitDepth)
@@ -118,10 +137,10 @@ UTexture2D* FGLTFTextureUtility::CreateTransientTexture(const void* RawData, int
 	return Texture;
 }
 
-UTextureRenderTarget2D* FGLTFTextureUtility::CreateRenderTarget(const FIntPoint& Size, bool bIsHDR)
+UTextureRenderTarget2D* FGLTFTextureUtility::CreateRenderTarget(const FIntPoint& Size, bool bHDR)
 {
 	// TODO: instead of PF_FloatRGBA (i.e. RTF_RGBA16f) use PF_A32B32G32R32F (i.e. RTF_RGBA32f) to avoid accuracy loss
-	const EPixelFormat PixelFormat = bIsHDR ? PF_FloatRGBA : PF_B8G8R8A8;
+	const EPixelFormat PixelFormat = bHDR ? PF_FloatRGBA : PF_B8G8R8A8;
 
 	// NOTE: both bForceLinearGamma and TargetGamma=2.2 seem necessary for exported images to match their source data.
 	// It's not entirely clear why gamma must be 2.2 (instead of 0.0) and why bInForceLinearGamma must also be true.
@@ -214,7 +233,7 @@ UTexture2D* FGLTFTextureUtility::CreateTextureFromCubeFace(const UTextureRenderT
 	const FIntPoint Size(RenderTargetCube->SizeX, RenderTargetCube->SizeX);
 	UTexture2D* FaceTexture;
 
-	if (IsHDR(RenderTargetCube))
+	if (IsHDR(RenderTargetCube->GetFormat()))
 	{
 		TArray<FFloat16Color> Pixels;
 		if (!Resource->ReadPixels(Pixels, FReadSurfaceDataFlags(RCM_UNorm, CubeFace)))
diff --git a/Engine/Plugins/Enterprise/GLTFExporter/MigratedFromGit/Source/GLTFExporter/Private/Converters/GLTFTextureUtility.h b/Engine/Plugins/Enterprise/GLTFExporter/MigratedFromGit/Source/GLTFExporter/Private/Converters/GLTFTextureUtility.h
--- a/Engine/Plugins/Enterprise/GLTFExporter/MigratedFromGit/Source/GLTFExporter/Private/Converters/GLTFTextureUtility.h
+++ b/Engine/Plugins/Enterprise/GLTFExporter/MigratedFromGit/Source/GLTFExporter/Private/Converters/GLTFTextureUtility.h
@@ -9,12 +9,13 @@ enum class ERGBFormat : int8;
 
 struct FGLTFTextureUtility
 {
-	static bool IsHDR(const UTexture* Texture);
+	static bool IsHDR(EPixelFormat Format);
 
 	static bool CanPNGCompressFormat(ETextureSourceFormat InFormat, ERGBFormat& OutFormat, uint32& OutBitDepth);
 	static bool CanPNGCompressFormat(EPixelFormat InFormat, ERGBFormat& OutFormat, uint32& OutBitDepth);
 
 	static bool IsCubemap(const UTexture* Texture);
+
 	static float GetCubeFaceRotation(ECubeFace CubeFace);
 
 	static TextureFilter GetDefaultFilter(TextureGroup Group);
diff --git a/Engine/Plugins/Enterprise/GLTFExporter/MigratedFromGit/Source/GLTFExporter/Private/Tasks/GLTFTextureTasks.cpp b/Engine/Plugins/Enterprise/GLTFExporter/MigratedFromGit/Source/GLTFExporter/Private/Tasks/GLTFTextureTasks.cpp
--- a/Engine/Plugins/Enterprise/GLTFExporter/MigratedFromGit/Source/GLTFExporter/Private/Tasks/GLTFTextureTasks.cpp
+++ b/Engine/Plugins/Enterprise/GLTFExporter/MigratedFromGit/Source/GLTFExporter/Private/Tasks/GLTFTextureTasks.cpp
@@ -11,7 +11,7 @@ void FGLTFTexture2DTask::Complete()
 	FGLTFJsonTexture& JsonTexture = Builder.GetTexture(TextureIndex);
 	Texture2D->GetName(JsonTexture.Name);
 
-	const bool bIsHDR = FGLTFTextureUtility::IsHDR(Texture2D);
+	const bool bIsHDR = FGLTFTextureUtility::IsHDR(Texture2D->GetPixelFormat());
 	const FIntPoint Size = { Texture2D->GetSizeX(), Texture2D->GetSizeY() };
 	UTextureRenderTarget2D* RenderTarget = FGLTFTextureUtility::CreateRenderTarget(Size, bIsHDR);
 
@@ -54,7 +54,7 @@ void FGLTFTextureCubeTask::Complete()
 		return;
 	}
 
-	const bool bIsHDR = FGLTFTextureUtility::IsHDR(TextureCube);
+	const bool bIsHDR = FGLTFTextureUtility::IsHDR(TextureCube->GetPixelFormat());
 	const FIntPoint Size = { TextureCube->GetSizeX(), TextureCube->GetSizeY() };
 	UTextureRenderTarget2D* RenderTarget = FGLTFTextureUtility::CreateRenderTarget(Size, bIsHDR);
 
@@ -82,7 +82,7 @@ void FGLTFTextureRenderTarget2DTask::Complete()
 	FGLTFJsonTexture& JsonTexture = Builder.GetTexture(TextureIndex);
 	RenderTarget2D->GetName(JsonTexture.Name);
 
-	const bool bIsHDR = FGLTFTextureUtility::IsHDR(RenderTarget2D);
+	const bool bIsHDR = FGLTFTextureUtility::IsHDR(RenderTarget2D->GetFormat());
 	const FIntPoint Size = { RenderTarget2D->SizeX, RenderTarget2D->SizeY };
 
 	if (bIsHDR)
@@ -115,7 +115,7 @@ void FGLTFTextureRenderTargetCubeTask::Complete()
 		return;
 	}
 
-	const bool bIsHDR = FGLTFTextureUtility::IsHDR(RenderTargetCube);
+	const bool bIsHDR = FGLTFTextureUtility::IsHDR(RenderTargetCube->GetFormat());
 	const FIntPoint Size = { RenderTargetCube->SizeX, RenderTargetCube->SizeX };
 	UTextureRenderTarget2D* RenderTarget = FGLTFTextureUtility::CreateRenderTarget(Size, bIsHDR);
 
