[GLTFExporter] refactor(camera-actor): rename orbit-specific properties in preparation of new mode

commit fe999daec2f502621637ff0a3c7e0bc301665255
Date: Wed, 28 Oct 2020 16:01:27 +0100

    refactor(camera-actor): rename orbit-specific properties in preparation of new mode
    (cherry picked from commit f18fc1eb3ac90aec40c118639bd78d82d5a69657​)

#skipci

diff --git a/Engine/Plugins/Enterprise/GLTFExporter/MigratedFromGit/Source/GLTFExporter/Private/Converters/GLTFCameraConverters.cpp b/Engine/Plugins/Enterprise/GLTFExporter/MigratedFromGit/Source/GLTFExporter/Private/Converters/GLTFCameraConverters.cpp
--- a/Engine/Plugins/Enterprise/GLTFExporter/MigratedFromGit/Source/GLTFExporter/Private/Converters/GLTFCameraConverters.cpp
+++ b/Engine/Plugins/Enterprise/GLTFExporter/MigratedFromGit/Source/GLTFExporter/Private/Converters/GLTFCameraConverters.cpp
@@ -62,9 +62,9 @@ FGLTFJsonCameraIndex FGLTFCameraConverter::Convert(const UCameraComponent* Camer
 			OrbitCamera.MinDistance = FGLTFConverterUtility::ConvertLength(CameraActor->DistanceMin, ExportScale);
 			OrbitCamera.MaxAngle = CameraActor->PitchAngleMax;
 			OrbitCamera.MinAngle = CameraActor->PitchAngleMin;
-			OrbitCamera.DistanceSensitivity = CameraActor->DistanceSensitivity;
-			OrbitCamera.OrbitSensitivity = CameraActor->OrbitSensitivity;
-			OrbitCamera.OrbitInertia = CameraActor->OrbitInertia;
+			OrbitCamera.DistanceSensitivity = CameraActor->DollySensitivity;
+			OrbitCamera.OrbitSensitivity = CameraActor->RotationSensitivity;
+			OrbitCamera.OrbitInertia = CameraActor->RotationInertia;
 			OrbitCamera.DollyDuration = CameraActor->DollyDuration;
 
 			Camera.OrbitCamera = OrbitCamera;
diff --git a/Engine/Plugins/Enterprise/GLTFExporter/MigratedFromGit/Source/GLTFExporterRuntime/Private/Actors/GLTFCameraActor.cpp b/Engine/Plugins/Enterprise/GLTFExporter/MigratedFromGit/Source/GLTFExporterRuntime/Private/Actors/GLTFCameraActor.cpp
--- a/Engine/Plugins/Enterprise/GLTFExporter/MigratedFromGit/Source/GLTFExporterRuntime/Private/Actors/GLTFCameraActor.cpp
+++ b/Engine/Plugins/Enterprise/GLTFExporter/MigratedFromGit/Source/GLTFExporterRuntime/Private/Actors/GLTFCameraActor.cpp
@@ -15,21 +15,21 @@ namespace
 
 	// Scales to convert from the export-friendly sensitivity-values stored in our properties
 	// to values that we can use when processing axis-values (to get similar results as in the viewer).
-	const float OrbitSensitivityScale = 16.667f;
-	const float DistanceSensitivityScale = 0.1f;
+	const float RotationSensitivityScale = 16.667f;
+	const float DollySensitivityScale = 0.1f;
 }
 
 AGLTFCameraActor::AGLTFCameraActor(const FObjectInitializer& ObjectInitializer)
 	: Super(ObjectInitializer)
 	, Focus(nullptr)
-	, DistanceMin(100.0f)
-	, DistanceMax(1000.0f)
 	, PitchAngleMin(-90.0f)
 	, PitchAngleMax(90.0f)
+	, DistanceMin(100.0f)
+	, DistanceMax(1000.0f)
 	, DollyDuration(0.2f)
-	, OrbitInertia(0.1f)
-	, OrbitSensitivity(0.3f)
-	, DistanceSensitivity(0.5f)
+	, DollySensitivity(0.5f)
+	, RotationInertia(0.1f)
+	, RotationSensitivity(0.3f)
 	, Distance(0.0f)
 	, Pitch(0.0f)
 	, Yaw(0.0f)
@@ -104,7 +104,7 @@ void AGLTFCameraActor::Tick(float DeltaSeconds)
 		Distance = FMath::InterpEaseInOut(DollyStartDistance, TargetDistance, (DollyDuration - DollyTime) / DollyDuration, 1.2f);
 	}
 
-	const float Alpha = (OrbitInertia == 0.0f) ? 1.0f : FMath::Min(DeltaSeconds / OrbitInertia, 1.0f);
+	const float Alpha = (RotationInertia == 0.0f) ? 1.0f : FMath::Min(DeltaSeconds / RotationInertia, 1.0f);
 	Yaw = FMath::Lerp(Yaw, TargetYaw, Alpha);
 	Pitch = FMath::Lerp(Pitch, TargetPitch, Alpha);
 
@@ -141,19 +141,19 @@ void AGLTFCameraActor::PostActorCreated()
 
 void AGLTFCameraActor::OnMouseX(float AxisValue)
 {
-	TargetYaw += AxisValue * OrbitSensitivity * OrbitSensitivityScale;
+	TargetYaw += AxisValue * RotationSensitivity * RotationSensitivityScale;
 }
 
 void AGLTFCameraActor::OnMouseY(float AxisValue)
 {
-	TargetPitch = ClampPitch(TargetPitch - AxisValue * OrbitSensitivity * OrbitSensitivityScale);
+	TargetPitch = ClampPitch(TargetPitch - AxisValue * RotationSensitivity * RotationSensitivityScale);
 }
 
 void AGLTFCameraActor::OnMouseWheelAxis(float AxisValue)
 {
 	if (!FMath::IsNearlyZero(AxisValue))
 	{
-		const float DeltaDistance = -AxisValue * (TargetDistance * DistanceSensitivity * DistanceSensitivityScale);
+		const float DeltaDistance = -AxisValue * (TargetDistance * DollySensitivity * DollySensitivityScale);
 
 		DollyTime = DollyDuration;
 		TargetDistance = ClampDistance(TargetDistance + DeltaDistance);
diff --git a/Engine/Plugins/Enterprise/GLTFExporter/MigratedFromGit/Source/GLTFExporterRuntime/Public/Actors/GLTFCameraActor.h b/Engine/Plugins/Enterprise/GLTFExporter/MigratedFromGit/Source/GLTFExporterRuntime/Public/Actors/GLTFCameraActor.h
--- a/Engine/Plugins/Enterprise/GLTFExporter/MigratedFromGit/Source/GLTFExporterRuntime/Public/Actors/GLTFCameraActor.h
+++ b/Engine/Plugins/Enterprise/GLTFExporter/MigratedFromGit/Source/GLTFExporterRuntime/Public/Actors/GLTFCameraActor.h
@@ -59,37 +59,37 @@ public:
 	UPROPERTY(EditAnywhere, BlueprintReadWrite, Category = "GLTF Orbit Camera Actor")
 	AActor* Focus;
 
+	/* Minimum pitch angle (in degrees) for the camera. */
+	UPROPERTY(EditAnywhere, BlueprintReadWrite, Category = "GLTF Camera Actor")
+	float PitchAngleMin;
+
+	/* Maximum pitch angle (in degrees) for the camera. */
+	UPROPERTY(EditAnywhere, BlueprintReadWrite, Category = "GLTF Camera Actor")
+	float PitchAngleMax;
+
 	/* Closest distance the camera can approach the focused actor. */
-	UPROPERTY(EditAnywhere, BlueprintReadWrite, Category = "GLTF Orbit Camera Actor")
+	UPROPERTY(EditAnywhere, BlueprintReadWrite, Category = "GLTF Camera Actor")
 	float DistanceMin;
 
 	/* Farthest distance the camera can recede from the focused actor. */
-	UPROPERTY(EditAnywhere, BlueprintReadWrite, Category = "GLTF Orbit Camera Actor")
+	UPROPERTY(EditAnywhere, BlueprintReadWrite, Category = "GLTF Camera Actor")
 	float DistanceMax;
 
-	/* Minimum angle (in degrees) that the camera can pitch relative to the focused actor. */
-	UPROPERTY(EditAnywhere, BlueprintReadWrite, Category = "GLTF Orbit Camera Actor")
-	float PitchAngleMin;
-
-	/* Maximum angle (in degrees) that the camera can pitch relative to the focused actor. */
-	UPROPERTY(EditAnywhere, BlueprintReadWrite, Category = "GLTF Orbit Camera Actor")
-	float PitchAngleMax;
-
 	/* Duration (in seconds) that it takes the camera to complete a change in distance. */
-	UPROPERTY(EditAnywhere, BlueprintReadOnly, Category = "GLTF Orbit Camera Actor")
+	UPROPERTY(EditAnywhere, BlueprintReadOnly, Category = "GLTF Camera Actor")
 	float DollyDuration;
 
-	/* Deceleration that occurs after orbital movement. The higher the value, the longer it takes to settle. */
-	UPROPERTY(EditAnywhere, BlueprintReadOnly, Category = "GLTF Orbit Camera Actor")
-	float OrbitInertia;
+	/* Size of the dolly movement relative to user input. The higher the value, the faster it moves. */
+	UPROPERTY(EditAnywhere, BlueprintReadOnly, Category = "GLTF Camera Actor")
+	float DollySensitivity;
 
-	/* Size of the orbital movement relative to user input. The higher the value, the faster it moves. */
-	UPROPERTY(EditAnywhere, BlueprintReadOnly, Category = "GLTF Orbit Camera Actor")
-	float OrbitSensitivity;
+	/* Deceleration that occurs after rotational movement. The higher the value, the longer it takes to settle. */
+	UPROPERTY(EditAnywhere, BlueprintReadOnly, Category = "GLTF Camera Actor")
+	float RotationInertia;
 
-	/* Size of the dolly movement relative to user input. The higher the value, the faster it moves. */
-	UPROPERTY(EditAnywhere, BlueprintReadOnly, Category = "GLTF Orbit Camera Actor")
-	float DistanceSensitivity;
+	/* Size of the rotational movement relative to user input. The higher the value, the faster it moves. */
+	UPROPERTY(EditAnywhere, BlueprintReadOnly, Category = "GLTF Camera Actor")
+	float RotationSensitivity;
 
 private:
 	float Distance;
