[GLTFExporter] feat(camera-actor): add structure for two separate camera modes

commit e45e2abc58fdcb9e3f1595903cc0e678265d4a31
Date: Thu, 29 Oct 2020 08:32:34 +0100

    feat(camera-actor): add structure for two separate camera modes
    (cherry picked from commit 9b4590aa664905f1b85343f5e2cea8a183fa5e9b​)

#skipci

diff --git a/Engine/Plugins/Enterprise/GLTFExporter/MigratedFromGit/Source/GLTFExporterRuntime/Private/Actors/GLTFCameraActor.cpp b/Engine/Plugins/Enterprise/GLTFExporter/MigratedFromGit/Source/GLTFExporterRuntime/Private/Actors/GLTFCameraActor.cpp
--- a/Engine/Plugins/Enterprise/GLTFExporter/MigratedFromGit/Source/GLTFExporterRuntime/Private/Actors/GLTFCameraActor.cpp
+++ b/Engine/Plugins/Enterprise/GLTFExporter/MigratedFromGit/Source/GLTFExporterRuntime/Private/Actors/GLTFCameraActor.cpp
@@ -21,6 +21,7 @@ namespace
 
 AGLTFCameraActor::AGLTFCameraActor(const FObjectInitializer& ObjectInitializer)
 	: Super(ObjectInitializer)
+	, Mode(EGLTFCameraMode::FirstPerson)
 	, Focus(nullptr)
 	, PitchAngleMin(-90.0f)
 	, PitchAngleMax(90.0f)
@@ -70,21 +71,32 @@ void AGLTFCameraActor::BeginPlay()
 {
 	Super::BeginPlay();
 
-	const FVector FocusPosition = GetFocusPosition();
+	if (Mode == EGLTFCameraMode::FirstPerson)
+	{
+		// TODO: implement
+	}
+	else if (Mode == EGLTFCameraMode::ThirdPerson)
+	{
+		const FVector FocusPosition = GetFocusPosition();
 
-	// Ensure that the camera is initially aimed at the focus-position
-	SetActorRotation(GetLookAtRotation(FocusPosition));
+		// Ensure that the camera is initially aimed at the focus-position
+		SetActorRotation(GetLookAtRotation(FocusPosition));
 
-	const FVector Position = GetActorLocation();
-	const FRotator Rotation = GetActorRotation();
+		const FVector Position = GetActorLocation();
+		const FRotator Rotation = GetActorRotation();
 
-	// Calculate values based on the current location and orientation
-	Distance = ClampDistance((FocusPosition - Position).Size());
-	Pitch = ClampPitch(-Rotation.Pitch);
-	Yaw = ClampYaw(Rotation.Yaw);
-	TargetDistance = Distance;
-	TargetPitch = Pitch;
-	TargetYaw = Yaw;
+		// Calculate values based on the current location and orientation
+		Distance = ClampDistance((FocusPosition - Position).Size());
+		Pitch = ClampPitch(-Rotation.Pitch);
+		Yaw = ClampYaw(Rotation.Yaw);
+		TargetDistance = Distance;
+		TargetPitch = Pitch;
+		TargetYaw = Yaw;
+	}
+	else
+	{
+		checkNoEntry();
+	}
 
 	if (InputComponent)
 	{
@@ -98,33 +110,40 @@ void AGLTFCameraActor::Tick(float DeltaSeconds)
 {
 	Super::Tick(DeltaSeconds);
 
-	if (DollyTime != 0.0f)
+	if (Mode == EGLTFCameraMode::FirstPerson)
 	{
-		DollyTime = FMath::Max(DollyTime - DeltaSeconds, 0.0f);
-		Distance = FMath::InterpEaseInOut(DollyStartDistance, TargetDistance, (DollyDuration - DollyTime) / DollyDuration, 1.2f);
+		// TODO: implement
 	}
+	else if (Mode == EGLTFCameraMode::ThirdPerson)
+	{
+		if (DollyTime != 0.0f)
+		{
+			DollyTime = FMath::Max(DollyTime - DeltaSeconds, 0.0f);
+			Distance = FMath::InterpEaseInOut(DollyStartDistance, TargetDistance, (DollyDuration - DollyTime) / DollyDuration, 1.2f);
+		}
 
-	const float Alpha = (RotationInertia == 0.0f) ? 1.0f : FMath::Min(DeltaSeconds / RotationInertia, 1.0f);
-	Yaw = FMath::Lerp(Yaw, TargetYaw, Alpha);
-	Pitch = FMath::Lerp(Pitch, TargetPitch, Alpha);
+		const float Alpha = (RotationInertia == 0.0f) ? 1.0f : FMath::Min(DeltaSeconds / RotationInertia, 1.0f);
+		Yaw = FMath::Lerp(Yaw, TargetYaw, Alpha);
+		Pitch = FMath::Lerp(Pitch, TargetPitch, Alpha);
 
-	const int32 YawCycleIndex = AngleCycleIndex(Yaw);
-	const int32 TargetYawCycleIndex = AngleCycleIndex(TargetYaw);
+		const int32 YawCycleIndex = AngleCycleIndex(Yaw);
+		const int32 TargetYawCycleIndex = AngleCycleIndex(TargetYaw);
 
-	// Clamp the angles to a positive 360 degrees while they are within the same cycle.
-	// It is important that this doesn't happen during a transition as it will otherwise be skewed.
-	if (YawCycleIndex == TargetYawCycleIndex && YawCycleIndex != 0)
-	{
-		Yaw = ClampYaw(Yaw);
-		TargetYaw = ClampYaw(TargetYaw);
-	}
+		// Clamp the angles to a positive 360 degrees while they are within the same cycle.
+		// It is important that this doesn't happen during a transition as it will otherwise be skewed.
+		if (YawCycleIndex == TargetYawCycleIndex && YawCycleIndex != 0)
+		{
+			Yaw = ClampYaw(Yaw);
+			TargetYaw = ClampYaw(TargetYaw);
+		}
 
-	const FTransform FocusTransform = FTransform(GetFocusPosition());
-	const FTransform DollyTransform = FTransform(-FVector::ForwardVector * Distance);
-	const FTransform RotationTransform = FTransform(FQuat::MakeFromEuler(FVector(0.0f, -Pitch, Yaw)));
-	const FTransform ResultTransform = DollyTransform * RotationTransform * FocusTransform;
+		const FTransform FocusTransform = FTransform(GetFocusPosition());
+		const FTransform DollyTransform = FTransform(-FVector::ForwardVector * Distance);
+		const FTransform RotationTransform = FTransform(FQuat::MakeFromEuler(FVector(0.0f, -Pitch, Yaw)));
+		const FTransform ResultTransform = DollyTransform * RotationTransform * FocusTransform;
 
-	SetActorTransform(ResultTransform);
+		SetActorTransform(ResultTransform);
+	}
 }
 
 void AGLTFCameraActor::PreInitializeComponents()
@@ -151,13 +170,16 @@ void AGLTFCameraActor::OnMouseY(float AxisValue)
 
 void AGLTFCameraActor::OnMouseWheelAxis(float AxisValue)
 {
-	if (!FMath::IsNearlyZero(AxisValue))
+	if (Mode == EGLTFCameraMode::ThirdPerson)
 	{
-		const float DeltaDistance = -AxisValue * (TargetDistance * DollySensitivity * DollySensitivityScale);
+		if (!FMath::IsNearlyZero(AxisValue))
+		{
+			const float DeltaDistance = -AxisValue * (TargetDistance * DollySensitivity * DollySensitivityScale);
 
-		DollyTime = DollyDuration;
-		TargetDistance = ClampDistance(TargetDistance + DeltaDistance);
-		DollyStartDistance = Distance;
+			DollyTime = DollyDuration;
+			TargetDistance = ClampDistance(TargetDistance + DeltaDistance);
+			DollyStartDistance = Distance;
+		}
 	}
 }
 
diff --git a/Engine/Plugins/Enterprise/GLTFExporter/MigratedFromGit/Source/GLTFExporterRuntime/Public/Actors/GLTFCameraActor.h b/Engine/Plugins/Enterprise/GLTFExporter/MigratedFromGit/Source/GLTFExporterRuntime/Public/Actors/GLTFCameraActor.h
--- a/Engine/Plugins/Enterprise/GLTFExporter/MigratedFromGit/Source/GLTFExporterRuntime/Public/Actors/GLTFCameraActor.h
+++ b/Engine/Plugins/Enterprise/GLTFExporter/MigratedFromGit/Source/GLTFExporterRuntime/Public/Actors/GLTFCameraActor.h
@@ -5,9 +5,18 @@
 #include "Camera/CameraActor.h"
 #include "GLTFCameraActor.generated.h"
 
+UENUM(BlueprintType)
+enum class EGLTFCameraMode : uint8
+{
+	FirstPerson,
+	ThirdPerson
+};
+
 /**
  * GLTF-compatible camera that will carry over settings and simulate the behavior in the resulting viewer.
- * Focuses one actor in the scene and orbits it through mouse control.
+ * Supports two modes:
+ * - FirstPerson: Allows the user to rotate the camera without modifying its position.
+ * - ThirdPerson: Focuses one actor in the scene and orbits it through mouse control.
  */
 UCLASS(BlueprintType, Blueprintable, DisplayName = "GLTF Camera")
 class GLTFEXPORTERRUNTIME_API AGLTFCameraActor : public ACameraActor
@@ -55,8 +64,13 @@ private:
 	bool SetAutoActivateForPlayer(const EAutoReceiveInput::Type Player);
 
 public:
-	/* Actor which the camera will focus on and subsequently orbit. */
-	UPROPERTY(EditAnywhere, BlueprintReadWrite, Category = "GLTF Orbit Camera Actor")
+
+	/* Camera mode */
+	UPROPERTY(EditAnywhere, BlueprintReadWrite, Category = "GLTF Camera Actor")
+	EGLTFCameraMode Mode;
+
+	/* Actor which the camera will focus on and subsequently orbit when using Third Person mode. */
+	UPROPERTY(EditAnywhere, BlueprintReadWrite, Category = "GLTF Camera Actor", meta=(EditCondition="Mode == EGLTFCameraMode::ThirdPerson"))
 	AActor* Focus;
 
 	/* Minimum pitch angle (in degrees) for the camera. */
