[GLTFExporter] feat(camera-actor): invert Pitch to match pitch of transforms in the editor

commit aa2e3c308e641043885bc553f12e8934829bc4c1
Date: Thu, 29 Oct 2020 08:36:59 +0100

    feat(camera-actor): invert Pitch to match pitch of transforms in the editor
    (cherry picked from commit 7982b8eef952b254518762f5b678e6ee1ea72beb​)

#skipci

diff --git a/Engine/Plugins/Enterprise/GLTFExporter/MigratedFromGit/Source/GLTFExporterRuntime/Private/Actors/GLTFCameraActor.cpp b/Engine/Plugins/Enterprise/GLTFExporter/MigratedFromGit/Source/GLTFExporterRuntime/Private/Actors/GLTFCameraActor.cpp
--- a/Engine/Plugins/Enterprise/GLTFExporter/MigratedFromGit/Source/GLTFExporterRuntime/Private/Actors/GLTFCameraActor.cpp
+++ b/Engine/Plugins/Enterprise/GLTFExporter/MigratedFromGit/Source/GLTFExporterRuntime/Private/Actors/GLTFCameraActor.cpp
@@ -87,7 +87,7 @@ void AGLTFCameraActor::BeginPlay()
 
 		// Calculate values based on the current location and orientation
 		Distance = ClampDistance((FocusPosition - Position).Size());
-		Pitch = ClampPitch(-Rotation.Pitch);
+		Pitch = ClampPitch(Rotation.Pitch);
 		Yaw = ClampYaw(Rotation.Yaw);
 		TargetDistance = Distance;
 		TargetPitch = Pitch;
@@ -139,7 +139,7 @@ void AGLTFCameraActor::Tick(float DeltaSeconds)
 
 		const FTransform FocusTransform = FTransform(GetFocusPosition());
 		const FTransform DollyTransform = FTransform(-FVector::ForwardVector * Distance);
-		const FTransform RotationTransform = FTransform(FQuat::MakeFromEuler(FVector(0.0f, -Pitch, Yaw)));
+		const FTransform RotationTransform = FTransform(FQuat::MakeFromEuler(FVector(0.0f, Pitch, Yaw)));
 		const FTransform ResultTransform = DollyTransform * RotationTransform * FocusTransform;
 
 		SetActorTransform(ResultTransform);
@@ -165,7 +165,7 @@ void AGLTFCameraActor::OnMouseX(float AxisValue)
 
 void AGLTFCameraActor::OnMouseY(float AxisValue)
 {
-	TargetPitch = ClampPitch(TargetPitch - AxisValue * RotationSensitivity * RotationSensitivityScale);
+	TargetPitch = ClampPitch(TargetPitch + AxisValue * RotationSensitivity * RotationSensitivityScale);
 }
 
 void AGLTFCameraActor::OnMouseWheelAxis(float AxisValue)
