[GLTFExporter] refactor(camera-actor): remove checks for cyclic index of Yaw

commit a495af4073d4188858f07f76c99c5fc5bfe3befb
Date: Thu, 29 Oct 2020 08:41:00 +0100

    refactor(camera-actor): remove checks for cyclic index of Yaw
    we will use another way of keeping Yaw in a suitable range later
    (cherry picked from commit d40c548da8f0668c1a49df51bdc602964c5743f1​)

#skipci

diff --git a/Engine/Plugins/Enterprise/GLTFExporter/MigratedFromGit/Source/GLTFExporterRuntime/Private/Actors/GLTFCameraActor.cpp b/Engine/Plugins/Enterprise/GLTFExporter/MigratedFromGit/Source/GLTFExporterRuntime/Private/Actors/GLTFCameraActor.cpp
--- a/Engine/Plugins/Enterprise/GLTFExporter/MigratedFromGit/Source/GLTFExporterRuntime/Private/Actors/GLTFCameraActor.cpp
+++ b/Engine/Plugins/Enterprise/GLTFExporter/MigratedFromGit/Source/GLTFExporterRuntime/Private/Actors/GLTFCameraActor.cpp
@@ -5,14 +5,6 @@
 
 namespace
 {
-	// Uniquely identifies an angle by how many times it has crossed the 0-360 degree range.
-	// Positive angles are indexed from 0 and negative angles from -1.
-	int32 AngleCycleIndex(float Angle)
-	{
-		// TODO: is this similar to viewer camera?
-		return (static_cast<int32>(Angle) / 360) + ((Angle < 0.0f) ? -1 : 0);
-	}
-
 	// Scales to convert from the export-friendly sensitivity-values stored in our properties
 	// to values that we can use when processing axis-values (to get similar results as in the viewer).
 	const float RotationSensitivityScale = 16.667f;
@@ -126,17 +118,6 @@ void AGLTFCameraActor::Tick(float DeltaSeconds)
 		Yaw = FMath::Lerp(Yaw, TargetYaw, Alpha);
 		Pitch = FMath::Lerp(Pitch, TargetPitch, Alpha);
 
-		const int32 YawCycleIndex = AngleCycleIndex(Yaw);
-		const int32 TargetYawCycleIndex = AngleCycleIndex(TargetYaw);
-
-		// Clamp the angles to a positive 360 degrees while they are within the same cycle.
-		// It is important that this doesn't happen during a transition as it will otherwise be skewed.
-		if (YawCycleIndex == TargetYawCycleIndex && YawCycleIndex != 0)
-		{
-			Yaw = ClampYaw(Yaw);
-			TargetYaw = ClampYaw(TargetYaw);
-		}
-
 		const FTransform FocusTransform = FTransform(GetFocusPosition());
 		const FTransform DollyTransform = FTransform(-FVector::ForwardVector * Distance);
 		const FTransform RotationTransform = FTransform(FQuat::MakeFromEuler(FVector(0.0f, Pitch, Yaw)));
@@ -195,10 +176,9 @@ float AGLTFCameraActor::ClampPitch(float Value) const
 
 float AGLTFCameraActor::ClampYaw(float Value) const
 {
-	const int32 CycleIndex = AngleCycleIndex(Value);
-	const float Offset = 360.0f * static_cast<float>(FMath::Abs(FMath::Min(CycleIndex, 0)));
+	// TODO: implement
 
-	return FMath::Fmod(Value + Offset, 360.0f);
+	return Value;
 }
 
 void AGLTFCameraActor::RemoveInertia()
