[GLTFExporter] feat(camera-actor): add basic implementation of FirstPerson mode

commit 8640121c9f1d14bd268ddc8b7615222009560068
Date: Thu, 29 Oct 2020 08:43:14 +0100

    feat(camera-actor): add basic implementation of FirstPerson mode
    (cherry picked from commit b3be2f44b632754bfcfd798eab06d54ab3964076​)

#skipci

diff --git a/Engine/Plugins/Enterprise/GLTFExporter/MigratedFromGit/Source/GLTFExporterRuntime/Private/Actors/GLTFCameraActor.cpp b/Engine/Plugins/Enterprise/GLTFExporter/MigratedFromGit/Source/GLTFExporterRuntime/Private/Actors/GLTFCameraActor.cpp
--- a/Engine/Plugins/Enterprise/GLTFExporter/MigratedFromGit/Source/GLTFExporterRuntime/Private/Actors/GLTFCameraActor.cpp
+++ b/Engine/Plugins/Enterprise/GLTFExporter/MigratedFromGit/Source/GLTFExporterRuntime/Private/Actors/GLTFCameraActor.cpp
@@ -65,7 +65,12 @@ void AGLTFCameraActor::BeginPlay()
 
 	if (Mode == EGLTFCameraMode::FirstPerson)
 	{
-		// TODO: implement
+		const FRotator Rotation = GetActorRotation();
+
+		Pitch = ClampPitch(Rotation.Pitch);
+		Yaw = ClampYaw(Rotation.Yaw);
+		TargetPitch = Pitch;
+		TargetYaw = Yaw;
 	}
 	else if (Mode == EGLTFCameraMode::ThirdPerson)
 	{
@@ -104,7 +109,11 @@ void AGLTFCameraActor::Tick(float DeltaSeconds)
 
 	if (Mode == EGLTFCameraMode::FirstPerson)
 	{
-		// TODO: implement
+		const float Alpha = (RotationInertia == 0.0f) ? 1.0f : FMath::Min(DeltaSeconds / RotationInertia, 1.0f);
+		Yaw = FMath::Lerp(Yaw, TargetYaw, Alpha);
+		Pitch = FMath::Lerp(Pitch, TargetPitch, Alpha);
+
+		SetActorRotation(FQuat::MakeFromEuler(FVector(0.0f, Pitch, Yaw)));
 	}
 	else if (Mode == EGLTFCameraMode::ThirdPerson)
 	{
