[GLTFExporter] feat(camera-actor): add yaw limits

commit dd31032a2bbb768e2957c31d2022944158029c97
Date: Thu, 29 Oct 2020 08:46:15 +0100

    feat(camera-actor): add yaw limits
    (cherry picked from commit fa242f371ac4302ee1d7c795e643a3b9eb229e7d​)

#skipci

diff --git a/Engine/Plugins/Enterprise/GLTFExporter/MigratedFromGit/Source/GLTFExporterRuntime/Private/Actors/GLTFCameraActor.cpp b/Engine/Plugins/Enterprise/GLTFExporter/MigratedFromGit/Source/GLTFExporterRuntime/Private/Actors/GLTFCameraActor.cpp
--- a/Engine/Plugins/Enterprise/GLTFExporter/MigratedFromGit/Source/GLTFExporterRuntime/Private/Actors/GLTFCameraActor.cpp
+++ b/Engine/Plugins/Enterprise/GLTFExporter/MigratedFromGit/Source/GLTFExporterRuntime/Private/Actors/GLTFCameraActor.cpp
@@ -17,6 +17,8 @@ AGLTFCameraActor::AGLTFCameraActor(const FObjectInitializer& ObjectInitializer)
 	, Focus(nullptr)
 	, PitchAngleMin(-90.0f)
 	, PitchAngleMax(90.0f)
+	, YawAngleMin(0.0f)
+	, YawAngleMax(360.0f)
 	, DistanceMin(100.0f)
 	, DistanceMax(1000.0f)
 	, DollyDuration(0.2f)
@@ -150,7 +152,7 @@ void AGLTFCameraActor::PostActorCreated()
 
 void AGLTFCameraActor::OnMouseX(float AxisValue)
 {
-	TargetYaw += AxisValue * RotationSensitivity * RotationSensitivityScale;
+	TargetYaw = ClampYaw(TargetYaw + AxisValue * RotationSensitivity * RotationSensitivityScale);
 }
 
 void AGLTFCameraActor::OnMouseY(float AxisValue)
@@ -185,9 +187,7 @@ float AGLTFCameraActor::ClampPitch(float Value) const
 
 float AGLTFCameraActor::ClampYaw(float Value) const
 {
-	// TODO: implement
-
-	return Value;
+	return UsesYawLimits() ? FMath::Clamp(Value, YawAngleMin, YawAngleMax) : Value;
 }
 
 void AGLTFCameraActor::RemoveInertia()
@@ -227,3 +227,8 @@ bool AGLTFCameraActor::SetAutoActivateForPlayer(const EAutoReceiveInput::Type Pl
 	*ValuePtr = Player;
 	return true;
 }
+
+bool AGLTFCameraActor::UsesYawLimits() const
+{
+	return !FMath::IsNearlyEqual(YawAngleMax - YawAngleMin, 360.0f);
+}
diff --git a/Engine/Plugins/Enterprise/GLTFExporter/MigratedFromGit/Source/GLTFExporterRuntime/Public/Actors/GLTFCameraActor.h b/Engine/Plugins/Enterprise/GLTFExporter/MigratedFromGit/Source/GLTFExporterRuntime/Public/Actors/GLTFCameraActor.h
--- a/Engine/Plugins/Enterprise/GLTFExporter/MigratedFromGit/Source/GLTFExporterRuntime/Public/Actors/GLTFCameraActor.h
+++ b/Engine/Plugins/Enterprise/GLTFExporter/MigratedFromGit/Source/GLTFExporterRuntime/Public/Actors/GLTFCameraActor.h
@@ -63,6 +63,8 @@ private:
 
 	bool SetAutoActivateForPlayer(const EAutoReceiveInput::Type Player);
 
+	bool UsesYawLimits() const;
+
 public:
 
 	/* Camera mode */
@@ -81,6 +83,14 @@ public:
 	UPROPERTY(EditAnywhere, BlueprintReadWrite, Category = "GLTF Camera Actor")
 	float PitchAngleMax;
 
+	/* Minimum yaw angle (in degrees) for the camera. */
+	UPROPERTY(EditAnywhere, BlueprintReadWrite, Category = "GLTF Camera Actor")
+	float YawAngleMin;
+
+	/* Maximum yaw angle (in degrees) for the camera. */
+	UPROPERTY(EditAnywhere, BlueprintReadWrite, Category = "GLTF Camera Actor")
+	float YawAngleMax;
+
 	/* Closest distance the camera can approach the focused actor. */
 	UPROPERTY(EditAnywhere, BlueprintReadWrite, Category = "GLTF Camera Actor")
 	float DistanceMin;
