[GLTFExporter] feat(camera-actor): keep Yaw and TargetYaw in a reasonable range when not using limits

commit 12715610c9061be07cc8b701080383beb8618684
Date: Thu, 29 Oct 2020 08:47:03 +0100

    feat(camera-actor): keep Yaw and TargetYaw in a reasonable range when not using limits
    (cherry picked from commit d6dc47616bd284a53261e60c1003cd8a1912df34​)

#skipci

diff --git a/Engine/Plugins/Enterprise/GLTFExporter/MigratedFromGit/Source/GLTFExporterRuntime/Private/Actors/GLTFCameraActor.cpp b/Engine/Plugins/Enterprise/GLTFExporter/MigratedFromGit/Source/GLTFExporterRuntime/Private/Actors/GLTFCameraActor.cpp
--- a/Engine/Plugins/Enterprise/GLTFExporter/MigratedFromGit/Source/GLTFExporterRuntime/Private/Actors/GLTFCameraActor.cpp
+++ b/Engine/Plugins/Enterprise/GLTFExporter/MigratedFromGit/Source/GLTFExporterRuntime/Private/Actors/GLTFCameraActor.cpp
@@ -153,6 +153,15 @@ void AGLTFCameraActor::PostActorCreated()
 void AGLTFCameraActor::OnMouseX(float AxisValue)
 {
 	TargetYaw = ClampYaw(TargetYaw + AxisValue * RotationSensitivity * RotationSensitivityScale);
+
+	if (!UsesYawLimits())
+	{
+		// Prevent TargetYaw from reaching excessive values by keeping it in the range (-360, 360)
+		const float ClampedTargetYaw = FMath::Fmod(TargetYaw, 360.f);
+
+		Yaw += ClampedTargetYaw - TargetYaw;
+		TargetYaw = ClampedTargetYaw;
+	}
 }
 
 void AGLTFCameraActor::OnMouseY(float AxisValue)
