[GLTFExporter] feat(camera-actor): add range validation for min- and max-angle properties

commit 7a2d13d049d9fd97e40ee494287b6d9f87d805ee
Date: Thu, 29 Oct 2020 08:47:44 +0100

    feat(camera-actor): add range validation for min- and max-angle properties
    (cherry picked from commit 64ece3e345bd0e22140611033b9f94c8f29812a1​)

#skipci

diff --git a/Engine/Plugins/Enterprise/GLTFExporter/MigratedFromGit/Source/GLTFExporterRuntime/Private/Actors/GLTFCameraActor.cpp b/Engine/Plugins/Enterprise/GLTFExporter/MigratedFromGit/Source/GLTFExporterRuntime/Private/Actors/GLTFCameraActor.cpp
--- a/Engine/Plugins/Enterprise/GLTFExporter/MigratedFromGit/Source/GLTFExporterRuntime/Private/Actors/GLTFCameraActor.cpp
+++ b/Engine/Plugins/Enterprise/GLTFExporter/MigratedFromGit/Source/GLTFExporterRuntime/Private/Actors/GLTFCameraActor.cpp
@@ -57,6 +57,28 @@ void AGLTFCameraActor::PostEditChangeProperty(FPropertyChangedEvent& PropertyCha
 				UE_LOG(LogTemp, Warning, TEXT("The camera focus must not be the camera's own actor"));
 			}
 		}
+		else if (PropertyName == GET_MEMBER_NAME_CHECKED(ThisClass, PitchAngleMin))
+		{
+			PitchAngleMin = FMath::Clamp(PitchAngleMin, -90.0f, 90.0f);
+		}
+		else if (PropertyName == GET_MEMBER_NAME_CHECKED(ThisClass, PitchAngleMax))
+		{
+			PitchAngleMax = FMath::Clamp(PitchAngleMax, -90.0f, 90.0f);
+		}
+		else if (PropertyName == GET_MEMBER_NAME_CHECKED(ThisClass, YawAngleMin))
+		{
+			YawAngleMin = FMath::Clamp(
+				YawAngleMin,
+				FMath::Max(-360.0f, YawAngleMax - 360.0f),
+				FMath::Min(360.0f, YawAngleMax));
+		}
+		else if (PropertyName == GET_MEMBER_NAME_CHECKED(ThisClass, YawAngleMax))
+		{
+			YawAngleMax = FMath::Clamp(
+				YawAngleMax,
+				FMath::Max(-360.0f, YawAngleMin),
+				FMath::Min(360.0f, YawAngleMin + 360.0f));
+		}
 	}
 }
 #endif // WITH_EDITOR
