[GLTFExporter] refactor: rename EPIC_orbital_cameras extension to EPIC_player_cameras

commit 6e896ffc3c7780a13d52569fbd7d364e8da9c25b
Date: Thu, 29 Oct 2020 08:56:40 +0100

    refactor: rename EPIC_orbital_cameras extension to EPIC_player_cameras
    (cherry picked from commit b10b696c72775495473b166837d5fb0145f61f32​)
    # Conflicts from cherry-pick:
    #	Plugins/GLTFExporter/Source/GLTFExporter/Private/GLTFExportOptions.cpp
    #	Plugins/GLTFExporter/Source/GLTFExporter/Public/GLTFExportOptions.h

#skipci

diff --git a/Engine/Plugins/Enterprise/GLTFExporter/MigratedFromGit/Source/GLTFExporter/Private/Converters/GLTFCameraConverters.cpp b/Engine/Plugins/Enterprise/GLTFExporter/MigratedFromGit/Source/GLTFExporter/Private/Converters/GLTFCameraConverters.cpp
--- a/Engine/Plugins/Enterprise/GLTFExporter/MigratedFromGit/Source/GLTFExporter/Private/Converters/GLTFCameraConverters.cpp
+++ b/Engine/Plugins/Enterprise/GLTFExporter/MigratedFromGit/Source/GLTFExporter/Private/Converters/GLTFCameraConverters.cpp
@@ -54,20 +54,20 @@ FGLTFJsonCameraIndex FGLTFCameraConverter::Convert(const UCameraComponent* Camer
 
 	if (CameraActor != nullptr)
 	{
-		if (Builder.ExportOptions->bExportOrbitalCameras)
+		if (Builder.ExportOptions->bExportPlayerCameras)
 		{
-			FGLTFJsonOrbitCamera OrbitCamera;
-			OrbitCamera.Focus = Builder.GetOrAddNode(CameraActor->Focus);
-			OrbitCamera.MaxDistance = FGLTFConverterUtility::ConvertLength(CameraActor->DistanceMax, ExportScale);
-			OrbitCamera.MinDistance = FGLTFConverterUtility::ConvertLength(CameraActor->DistanceMin, ExportScale);
-			OrbitCamera.MaxAngle = CameraActor->PitchAngleMax;
-			OrbitCamera.MinAngle = CameraActor->PitchAngleMin;
-			OrbitCamera.DistanceSensitivity = CameraActor->DollySensitivity;
-			OrbitCamera.OrbitSensitivity = CameraActor->RotationSensitivity;
-			OrbitCamera.OrbitInertia = CameraActor->RotationInertia;
-			OrbitCamera.DollyDuration = CameraActor->DollyDuration;
+			FGLTFJsonPlayerCamera PlayerCamera;
+			PlayerCamera.Focus = Builder.GetOrAddNode(CameraActor->Focus);
+			PlayerCamera.MaxDistance = FGLTFConverterUtility::ConvertLength(CameraActor->DistanceMax, ExportScale);
+			PlayerCamera.MinDistance = FGLTFConverterUtility::ConvertLength(CameraActor->DistanceMin, ExportScale);
+			PlayerCamera.MaxAngle = CameraActor->PitchAngleMax;
+			PlayerCamera.MinAngle = CameraActor->PitchAngleMin;
+			PlayerCamera.DistanceSensitivity = CameraActor->DollySensitivity;
+			PlayerCamera.OrbitSensitivity = CameraActor->RotationSensitivity;
+			PlayerCamera.OrbitInertia = CameraActor->RotationInertia;
+			PlayerCamera.DollyDuration = CameraActor->DollyDuration;
 
-			Camera.OrbitCamera = OrbitCamera;
+			Camera.PlayerCamera = PlayerCamera;
 		}
 	}
 
diff --git a/Engine/Plugins/Enterprise/GLTFExporter/MigratedFromGit/Source/GLTFExporter/Private/GLTFExportOptions.cpp b/Engine/Plugins/Enterprise/GLTFExporter/MigratedFromGit/Source/GLTFExporter/Private/GLTFExportOptions.cpp
--- a/Engine/Plugins/Enterprise/GLTFExporter/MigratedFromGit/Source/GLTFExporter/Private/GLTFExportOptions.cpp
+++ b/Engine/Plugins/Enterprise/GLTFExporter/MigratedFromGit/Source/GLTFExporter/Private/GLTFExportOptions.cpp
@@ -35,7 +35,7 @@ void UGLTFExportOptions::ResetToDefault()
 	bExportHiddenInGame = false;
 	ExportLights = EGLTFExporterLightMobility::MovableAndStationary;
 	bExportCameras = true;
-	bExportOrbitalCameras = true;
+	bExportPlayerCameras = true;
 	bExportHDRIBackdrops = true;
 	bExportSkySpheres = true;
 	bExportVariantSets = true;
diff --git a/Engine/Plugins/Enterprise/GLTFExporter/MigratedFromGit/Source/GLTFExporter/Private/Json/GLTFJsonCamera.h b/Engine/Plugins/Enterprise/GLTFExporter/MigratedFromGit/Source/GLTFExporter/Private/Json/GLTFJsonCamera.h
--- a/Engine/Plugins/Enterprise/GLTFExporter/MigratedFromGit/Source/GLTFExporter/Private/Json/GLTFJsonCamera.h
+++ b/Engine/Plugins/Enterprise/GLTFExporter/MigratedFromGit/Source/GLTFExporter/Private/Json/GLTFJsonCamera.h
@@ -3,7 +3,7 @@
 #pragma once
 
 #include "Json/GLTFJsonEnums.h"
-#include "Json/GLTFJsonOrbitCamera.h"
+#include "Json/GLTFJsonPlayerCamera.h"
 #include "Json/GLTFJsonUtility.h"
 #include "Serialization/JsonSerializer.h"
 
@@ -79,7 +79,7 @@ struct FGLTFJsonCamera
 	FString Name;
 
 	EGLTFJsonCameraType             Type;
-	TOptional<FGLTFJsonOrbitCamera> OrbitCamera;
+	TOptional<FGLTFJsonPlayerCamera> PlayerCamera;
 
 	union {
 		FGLTFJsonOrthographic Orthographic;
@@ -121,15 +121,15 @@ struct FGLTFJsonCamera
 				break;
 		}
 
-		if (OrbitCamera.IsSet())
+		if (PlayerCamera.IsSet())
 		{
 			JsonWriter.WriteObjectStart(TEXT("extensions"));
 
-			const EGLTFJsonExtension Extension = EGLTFJsonExtension::EPIC_OrbitalCameras;
+			const EGLTFJsonExtension Extension = EGLTFJsonExtension::EPIC_PlayerCameras;
 			Extensions.Used.Add(Extension);
 
 			JsonWriter.WriteIdentifierPrefix(FGLTFJsonUtility::ToString(Extension));
-			OrbitCamera.GetValue().WriteObject(JsonWriter, Extensions);
+			PlayerCamera.GetValue().WriteObject(JsonWriter, Extensions);
 
 			JsonWriter.WriteObjectEnd();
 		}
diff --git a/Engine/Plugins/Enterprise/GLTFExporter/MigratedFromGit/Source/GLTFExporter/Private/Json/GLTFJsonEnums.h b/Engine/Plugins/Enterprise/GLTFExporter/MigratedFromGit/Source/GLTFExporter/Private/Json/GLTFJsonEnums.h
--- a/Engine/Plugins/Enterprise/GLTFExporter/MigratedFromGit/Source/GLTFExporter/Private/Json/GLTFJsonEnums.h
+++ b/Engine/Plugins/Enterprise/GLTFExporter/MigratedFromGit/Source/GLTFExporter/Private/Json/GLTFJsonEnums.h
@@ -11,7 +11,7 @@ enum class EGLTFJsonExtension
 	KHR_MaterialsClearCoat,
 	KHR_MeshQuantization,
 	KHR_TextureTransform,
-	EPIC_OrbitalCameras,
+	EPIC_PlayerCameras,
 	EPIC_HDRIBackdrops,
 	EPIC_AnimationHotspots,
 	EPIC_LightmapTextures,
diff --git a/Engine/Plugins/Enterprise/GLTFExporter/MigratedFromGit/Source/GLTFExporter/Private/Json/GLTFJsonOrbitCamera.h b/Engine/Plugins/Enterprise/GLTFExporter/MigratedFromGit/Source/GLTFExporter/Private/Json/GLTFJsonOrbitCamera.h
deleted file mode 100644
diff --git a/Engine/Plugins/Enterprise/GLTFExporter/MigratedFromGit/Source/GLTFExporter/Private/Json/GLTFJsonPlayerCamera.h b/Engine/Plugins/Enterprise/GLTFExporter/MigratedFromGit/Source/GLTFExporter/Private/Json/GLTFJsonPlayerCamera.h
new file mode 100644
--- /dev/null
+++ b/Engine/Plugins/Enterprise/GLTFExporter/MigratedFromGit/Source/GLTFExporter/Private/Json/GLTFJsonPlayerCamera.h
@@ -0,0 +1,37 @@
+// Copyright Epic Games, Inc. All Rights Reserved.
+
+#pragma once
+
+#include "Json/GLTFJsonIndex.h"
+#include "Json/GLTFJsonVector4.h"
+
+struct FGLTFJsonPlayerCamera
+{
+	FGLTFJsonNodeIndex Focus;
+	float              MaxDistance;
+	float              MinDistance;
+	float              MaxAngle;
+	float              MinAngle;
+	float              DistanceSensitivity;
+	float              OrbitSensitivity;
+	float              OrbitInertia;
+	float              DollyDuration;
+
+	template <class CharType = TCHAR, class PrintPolicy = TPrettyJsonPrintPolicy<CharType>>
+	void WriteObject(TJsonWriter<CharType, PrintPolicy>& JsonWriter, FGLTFJsonExtensions& Extensions) const
+	{
+		JsonWriter.WriteObjectStart();
+
+		JsonWriter.WriteValue(TEXT("focus"), Focus);
+		FGLTFJsonUtility::WriteExactValue(JsonWriter, TEXT("maxDistance"), MaxDistance);
+		FGLTFJsonUtility::WriteExactValue(JsonWriter, TEXT("minDistance"), MinDistance);
+		FGLTFJsonUtility::WriteExactValue(JsonWriter, TEXT("maxAngle"), MaxAngle);
+		FGLTFJsonUtility::WriteExactValue(JsonWriter, TEXT("minAngle"), MinAngle);
+		FGLTFJsonUtility::WriteExactValue(JsonWriter, TEXT("distanceSensitivity"), DistanceSensitivity);
+		FGLTFJsonUtility::WriteExactValue(JsonWriter, TEXT("orbitSensitivity"), OrbitSensitivity);
+		FGLTFJsonUtility::WriteExactValue(JsonWriter, TEXT("orbitInertia"), OrbitInertia);
+		FGLTFJsonUtility::WriteExactValue(JsonWriter, TEXT("dollyDuration"), DollyDuration);
+
+		JsonWriter.WriteObjectEnd();
+	}
+};
diff --git a/Engine/Plugins/Enterprise/GLTFExporter/MigratedFromGit/Source/GLTFExporter/Private/Json/GLTFJsonUtility.h b/Engine/Plugins/Enterprise/GLTFExporter/MigratedFromGit/Source/GLTFExporter/Private/Json/GLTFJsonUtility.h
--- a/Engine/Plugins/Enterprise/GLTFExporter/MigratedFromGit/Source/GLTFExporter/Private/Json/GLTFJsonUtility.h
+++ b/Engine/Plugins/Enterprise/GLTFExporter/MigratedFromGit/Source/GLTFExporter/Private/Json/GLTFJsonUtility.h
@@ -22,7 +22,7 @@ struct FGLTFJsonUtility
 			case EGLTFJsonExtension::KHR_MaterialsClearCoat:  return TEXT("KHR_materials_clearcoat");
 			case EGLTFJsonExtension::KHR_MeshQuantization:    return TEXT("KHR_mesh_quantization");
 			case EGLTFJsonExtension::KHR_TextureTransform:    return TEXT("KHR_texture_transform");
-			case EGLTFJsonExtension::EPIC_OrbitalCameras:     return TEXT("EPIC_orbital_cameras");
+			case EGLTFJsonExtension::EPIC_PlayerCameras:      return TEXT("EPIC_player_cameras");
 			case EGLTFJsonExtension::EPIC_HDRIBackdrops:      return TEXT("EPIC_hdri_backdrops");
 			case EGLTFJsonExtension::EPIC_AnimationHotspots:  return TEXT("EPIC_animation_hotspots");
 			case EGLTFJsonExtension::EPIC_LightmapTextures:   return TEXT("EPIC_lightmap_textures");
diff --git a/Engine/Plugins/Enterprise/GLTFExporter/MigratedFromGit/Source/GLTFExporter/Public/GLTFExportOptions.h b/Engine/Plugins/Enterprise/GLTFExporter/MigratedFromGit/Source/GLTFExporter/Public/GLTFExportOptions.h
--- a/Engine/Plugins/Enterprise/GLTFExporter/MigratedFromGit/Source/GLTFExporter/Public/GLTFExportOptions.h
+++ b/Engine/Plugins/Enterprise/GLTFExporter/MigratedFromGit/Source/GLTFExporter/Public/GLTFExportOptions.h
@@ -113,7 +113,7 @@ public:
 	bool bExportCameras;
 
 	UPROPERTY(EditAnywhere, BlueprintReadWrite, Config, Category = Scene, Meta = (EditCondition = "bExportCameras"))
-	bool bExportOrbitalCameras;
+	bool bExportPlayerCameras;
 
 	UPROPERTY(EditAnywhere, BlueprintReadWrite, Config, Category = Scene)
 	bool bExportHDRIBackdrops;
