[GLTFExporter] refactor(player-cameras): rename properties in extension to match camera-actor

commit 8e8146926ba8c2b60127f152fd9f4c80f398ddf1
Date: Thu, 29 Oct 2020 09:31:22 +0100

    refactor(player-cameras): rename properties in extension to match camera-actor
    (cherry picked from commit 834e583e16617ce09c6e01bdf2b1a88d63b8dafe​)

#skipci

diff --git a/Engine/Plugins/Enterprise/GLTFExporter/MigratedFromGit/Source/GLTFExporter/Private/Converters/GLTFCameraConverters.cpp b/Engine/Plugins/Enterprise/GLTFExporter/MigratedFromGit/Source/GLTFExporter/Private/Converters/GLTFCameraConverters.cpp
--- a/Engine/Plugins/Enterprise/GLTFExporter/MigratedFromGit/Source/GLTFExporter/Private/Converters/GLTFCameraConverters.cpp
+++ b/Engine/Plugins/Enterprise/GLTFExporter/MigratedFromGit/Source/GLTFExporter/Private/Converters/GLTFCameraConverters.cpp
@@ -60,11 +60,11 @@ FGLTFJsonCameraIndex FGLTFCameraConverter::Convert(const UCameraComponent* Camer
 			PlayerCamera.Focus = Builder.GetOrAddNode(CameraActor->Focus);
 			PlayerCamera.MaxDistance = FGLTFConverterUtility::ConvertLength(CameraActor->DistanceMax, ExportScale);
 			PlayerCamera.MinDistance = FGLTFConverterUtility::ConvertLength(CameraActor->DistanceMin, ExportScale);
-			PlayerCamera.MaxAngle = CameraActor->PitchAngleMax;
-			PlayerCamera.MinAngle = CameraActor->PitchAngleMin;
-			PlayerCamera.DistanceSensitivity = CameraActor->DollySensitivity;
-			PlayerCamera.OrbitSensitivity = CameraActor->RotationSensitivity;
-			PlayerCamera.OrbitInertia = CameraActor->RotationInertia;
+			PlayerCamera.MaxPitch = CameraActor->PitchAngleMax;
+			PlayerCamera.MinPitch = CameraActor->PitchAngleMin;
+			PlayerCamera.DollySensitivity = CameraActor->DollySensitivity;
+			PlayerCamera.RotationSensitivity = CameraActor->RotationSensitivity;
+			PlayerCamera.RotationInertia = CameraActor->RotationInertia;
 			PlayerCamera.DollyDuration = CameraActor->DollyDuration;
 
 			Camera.PlayerCamera = PlayerCamera;
diff --git a/Engine/Plugins/Enterprise/GLTFExporter/MigratedFromGit/Source/GLTFExporter/Private/Json/GLTFJsonPlayerCamera.h b/Engine/Plugins/Enterprise/GLTFExporter/MigratedFromGit/Source/GLTFExporter/Private/Json/GLTFJsonPlayerCamera.h
--- a/Engine/Plugins/Enterprise/GLTFExporter/MigratedFromGit/Source/GLTFExporter/Private/Json/GLTFJsonPlayerCamera.h
+++ b/Engine/Plugins/Enterprise/GLTFExporter/MigratedFromGit/Source/GLTFExporter/Private/Json/GLTFJsonPlayerCamera.h
@@ -10,11 +10,11 @@ struct FGLTFJsonPlayerCamera
 	FGLTFJsonNodeIndex Focus;
 	float              MaxDistance;
 	float              MinDistance;
-	float              MaxAngle;
-	float              MinAngle;
-	float              DistanceSensitivity;
-	float              OrbitSensitivity;
-	float              OrbitInertia;
+	float              MaxPitch;
+	float              MinPitch;
+	float              DollySensitivity;
+	float              RotationSensitivity;
+	float              RotationInertia;
 	float              DollyDuration;
 
 	template <class CharType = TCHAR, class PrintPolicy = TPrettyJsonPrintPolicy<CharType>>
@@ -25,11 +25,11 @@ struct FGLTFJsonPlayerCamera
 		JsonWriter.WriteValue(TEXT("focus"), Focus);
 		FGLTFJsonUtility::WriteExactValue(JsonWriter, TEXT("maxDistance"), MaxDistance);
 		FGLTFJsonUtility::WriteExactValue(JsonWriter, TEXT("minDistance"), MinDistance);
-		FGLTFJsonUtility::WriteExactValue(JsonWriter, TEXT("maxAngle"), MaxAngle);
-		FGLTFJsonUtility::WriteExactValue(JsonWriter, TEXT("minAngle"), MinAngle);
-		FGLTFJsonUtility::WriteExactValue(JsonWriter, TEXT("distanceSensitivity"), DistanceSensitivity);
-		FGLTFJsonUtility::WriteExactValue(JsonWriter, TEXT("orbitSensitivity"), OrbitSensitivity);
-		FGLTFJsonUtility::WriteExactValue(JsonWriter, TEXT("orbitInertia"), OrbitInertia);
+		FGLTFJsonUtility::WriteExactValue(JsonWriter, TEXT("maxPitch"), MaxPitch);
+		FGLTFJsonUtility::WriteExactValue(JsonWriter, TEXT("minPitch"), MinPitch);
+		FGLTFJsonUtility::WriteExactValue(JsonWriter, TEXT("dollySensitivity"), DollySensitivity);
+		FGLTFJsonUtility::WriteExactValue(JsonWriter, TEXT("rotationSensitivity"), RotationSensitivity);
+		FGLTFJsonUtility::WriteExactValue(JsonWriter, TEXT("rotationInertia"), RotationInertia);
 		FGLTFJsonUtility::WriteExactValue(JsonWriter, TEXT("dollyDuration"), DollyDuration);
 
 		JsonWriter.WriteObjectEnd();
