[GLTFExporter] feat(player-cameras): add mode and yaw limits to extension

commit cd2759cccd2808b2c73649a8d060e0d1c2a995b7
Date: Thu, 29 Oct 2020 09:34:46 +0100

    feat(player-cameras): add mode and yaw limits to extension
    (cherry picked from commit ef01995702532c5e5cd71238ad411a6aeb8b8e68​)

#skipci

diff --git a/Engine/Plugins/Enterprise/GLTFExporter/MigratedFromGit/Source/GLTFExporter/Private/Converters/GLTFCameraConverters.cpp b/Engine/Plugins/Enterprise/GLTFExporter/MigratedFromGit/Source/GLTFExporter/Private/Converters/GLTFCameraConverters.cpp
--- a/Engine/Plugins/Enterprise/GLTFExporter/MigratedFromGit/Source/GLTFExporter/Private/Converters/GLTFCameraConverters.cpp
+++ b/Engine/Plugins/Enterprise/GLTFExporter/MigratedFromGit/Source/GLTFExporter/Private/Converters/GLTFCameraConverters.cpp
@@ -57,11 +57,14 @@ FGLTFJsonCameraIndex FGLTFCameraConverter::Convert(const UCameraComponent* Camer
 		if (Builder.ExportOptions->bExportPlayerCameras)
 		{
 			FGLTFJsonPlayerCamera PlayerCamera;
+			PlayerCamera.Mode = FGLTFConverterUtility::ConvertPlayerCameraMode(CameraActor->Mode);
 			PlayerCamera.Focus = Builder.GetOrAddNode(CameraActor->Focus);
 			PlayerCamera.MaxDistance = FGLTFConverterUtility::ConvertLength(CameraActor->DistanceMax, ExportScale);
 			PlayerCamera.MinDistance = FGLTFConverterUtility::ConvertLength(CameraActor->DistanceMin, ExportScale);
 			PlayerCamera.MaxPitch = CameraActor->PitchAngleMax;
 			PlayerCamera.MinPitch = CameraActor->PitchAngleMin;
+			PlayerCamera.MaxYaw = CameraActor->YawAngleMax;
+			PlayerCamera.MinYaw = CameraActor->YawAngleMin;
 			PlayerCamera.RotationSensitivity = CameraActor->RotationSensitivity;
 			PlayerCamera.RotationInertia = CameraActor->RotationInertia;
 			PlayerCamera.DollySensitivity = CameraActor->DollySensitivity;
diff --git a/Engine/Plugins/Enterprise/GLTFExporter/MigratedFromGit/Source/GLTFExporter/Private/Converters/GLTFConverterUtility.cpp b/Engine/Plugins/Enterprise/GLTFExporter/MigratedFromGit/Source/GLTFExporter/Private/Converters/GLTFConverterUtility.cpp
--- a/Engine/Plugins/Enterprise/GLTFExporter/MigratedFromGit/Source/GLTFExporter/Private/Converters/GLTFConverterUtility.cpp
+++ b/Engine/Plugins/Enterprise/GLTFExporter/MigratedFromGit/Source/GLTFExporter/Private/Converters/GLTFConverterUtility.cpp
@@ -2,6 +2,7 @@
 
 #include "Converters/GLTFConverterUtility.h"
 #include "Converters/GLTFTextureUtility.h"
+#include "Actors/GLTFCameraActor.h"
 
 EGLTFJsonCameraType FGLTFConverterUtility::ConvertCameraType(ECameraProjectionMode::Type ProjectionMode)
 {
@@ -112,3 +113,13 @@ EGLTFJsonCubeFace FGLTFConverterUtility::ConvertCubeFace(ECubeFace CubeFace)
 		default:            return EGLTFJsonCubeFace::None; // TODO: add error handling in callers
 	}
 }
+
+EGLTFJsonPlayerCameraMode FGLTFConverterUtility::ConvertPlayerCameraMode(EGLTFCameraMode CameraMode)
+{
+	switch (CameraMode)
+	{
+		case EGLTFCameraMode::FirstPerson: return EGLTFJsonPlayerCameraMode::FirstPerson;
+		case EGLTFCameraMode::ThirdPerson: return EGLTFJsonPlayerCameraMode::ThirdPerson;
+		default:                           return EGLTFJsonPlayerCameraMode::None; // TODO: add error handling in callers
+	}
+}
diff --git a/Engine/Plugins/Enterprise/GLTFExporter/MigratedFromGit/Source/GLTFExporter/Private/Converters/GLTFConverterUtility.h b/Engine/Plugins/Enterprise/GLTFExporter/MigratedFromGit/Source/GLTFExporter/Private/Converters/GLTFConverterUtility.h
--- a/Engine/Plugins/Enterprise/GLTFExporter/MigratedFromGit/Source/GLTFExporter/Private/Converters/GLTFConverterUtility.h
+++ b/Engine/Plugins/Enterprise/GLTFExporter/MigratedFromGit/Source/GLTFExporter/Private/Converters/GLTFConverterUtility.h
@@ -12,6 +12,8 @@
 #include "Json/GLTFJsonQuaternion.h"
 #include "Engine/EngineTypes.h"
 
+enum class EGLTFCameraMode : unsigned char;
+
 struct FGLTFConverterUtility
 {
 	static float ConvertLength(const float Length, const float ConversionScale)
@@ -193,6 +195,8 @@ struct FGLTFConverterUtility
 
 	static EGLTFJsonCubeFace ConvertCubeFace(ECubeFace CubeFace);
 
+	static EGLTFJsonPlayerCameraMode ConvertPlayerCameraMode(EGLTFCameraMode CameraMode);
+
 	template <typename ComponentType>
 	static EGLTFJsonComponentType GetComponentType()
 	{
diff --git a/Engine/Plugins/Enterprise/GLTFExporter/MigratedFromGit/Source/GLTFExporter/Private/Json/GLTFJsonEnums.h b/Engine/Plugins/Enterprise/GLTFExporter/MigratedFromGit/Source/GLTFExporter/Private/Json/GLTFJsonEnums.h
--- a/Engine/Plugins/Enterprise/GLTFExporter/MigratedFromGit/Source/GLTFExporter/Private/Json/GLTFJsonEnums.h
+++ b/Engine/Plugins/Enterprise/GLTFExporter/MigratedFromGit/Source/GLTFExporter/Private/Json/GLTFJsonEnums.h
@@ -160,3 +160,10 @@ enum class EGLTFJsonLightType
 	Point,
 	Spot
 };
+
+enum class EGLTFJsonPlayerCameraMode
+{
+	None = -1,
+	FirstPerson,
+	ThirdPerson
+};
diff --git a/Engine/Plugins/Enterprise/GLTFExporter/MigratedFromGit/Source/GLTFExporter/Private/Json/GLTFJsonPlayerCamera.h b/Engine/Plugins/Enterprise/GLTFExporter/MigratedFromGit/Source/GLTFExporter/Private/Json/GLTFJsonPlayerCamera.h
--- a/Engine/Plugins/Enterprise/GLTFExporter/MigratedFromGit/Source/GLTFExporter/Private/Json/GLTFJsonPlayerCamera.h
+++ b/Engine/Plugins/Enterprise/GLTFExporter/MigratedFromGit/Source/GLTFExporter/Private/Json/GLTFJsonPlayerCamera.h
@@ -7,26 +7,32 @@
 
 struct FGLTFJsonPlayerCamera
 {
-	FGLTFJsonNodeIndex Focus;
-	float              MaxDistance;
-	float              MinDistance;
-	float              MaxPitch;
-	float              MinPitch;
-	float              RotationSensitivity;
-	float              RotationInertia;
-	float              DollySensitivity;
-	float              DollyDuration;
+	EGLTFJsonPlayerCameraMode Mode;
+	FGLTFJsonNodeIndex        Focus;
+	float                     MaxDistance;
+	float                     MinDistance;
+	float                     MaxPitch;
+	float                     MinPitch;
+	float                     MaxYaw;
+	float                     MinYaw;
+	float                     RotationSensitivity;
+	float                     RotationInertia;
+	float                     DollySensitivity;
+	float                     DollyDuration;
 
 	template <class CharType = TCHAR, class PrintPolicy = TPrettyJsonPrintPolicy<CharType>>
 	void WriteObject(TJsonWriter<CharType, PrintPolicy>& JsonWriter, FGLTFJsonExtensions& Extensions) const
 	{
 		JsonWriter.WriteObjectStart();
 
+		JsonWriter.WriteValue(TEXT("mode"), FGLTFJsonUtility::ToString(Mode));
 		JsonWriter.WriteValue(TEXT("focus"), Focus);
 		FGLTFJsonUtility::WriteExactValue(JsonWriter, TEXT("maxDistance"), MaxDistance);
 		FGLTFJsonUtility::WriteExactValue(JsonWriter, TEXT("minDistance"), MinDistance);
 		FGLTFJsonUtility::WriteExactValue(JsonWriter, TEXT("maxPitch"), MaxPitch);
 		FGLTFJsonUtility::WriteExactValue(JsonWriter, TEXT("minPitch"), MinPitch);
+		FGLTFJsonUtility::WriteExactValue(JsonWriter, TEXT("maxYaw"), MaxYaw);
+		FGLTFJsonUtility::WriteExactValue(JsonWriter, TEXT("minYaw"), MinYaw);
 		FGLTFJsonUtility::WriteExactValue(JsonWriter, TEXT("rotationSensitivity"), RotationSensitivity);
 		FGLTFJsonUtility::WriteExactValue(JsonWriter, TEXT("rotationInertia"), RotationInertia);
 		FGLTFJsonUtility::WriteExactValue(JsonWriter, TEXT("dollySensitivity"), DollySensitivity);
diff --git a/Engine/Plugins/Enterprise/GLTFExporter/MigratedFromGit/Source/GLTFExporter/Private/Json/GLTFJsonUtility.h b/Engine/Plugins/Enterprise/GLTFExporter/MigratedFromGit/Source/GLTFExporter/Private/Json/GLTFJsonUtility.h
--- a/Engine/Plugins/Enterprise/GLTFExporter/MigratedFromGit/Source/GLTFExporter/Private/Json/GLTFJsonUtility.h
+++ b/Engine/Plugins/Enterprise/GLTFExporter/MigratedFromGit/Source/GLTFExporter/Private/Json/GLTFJsonUtility.h
@@ -159,6 +159,18 @@ struct FGLTFJsonUtility
 		}
 	}
 
+	static const TCHAR* ToString(EGLTFJsonPlayerCameraMode Value)
+	{
+		switch (Value)
+		{
+			case EGLTFJsonPlayerCameraMode::FirstPerson: return TEXT("firstPerson");
+			case EGLTFJsonPlayerCameraMode::ThirdPerson: return TEXT("thirdPerson");
+			default:
+				checkNoEntry();
+				return TEXT("");
+		}
+	}
+
 	template <class CharType = TCHAR, class PrintPolicy = TPrettyJsonPrintPolicy<CharType>>
 	static void WriteExactValue(TJsonWriter<CharType, PrintPolicy>& JsonWriter, float Value)
 	{
