[GLTFExporter] refactor(player-cameras): add TODO for conversion of angle limits

commit 106f4153770c1c34ae311f9df8acc0fe7caee548
Date: Thu, 29 Oct 2020 09:46:06 +0100

    refactor(player-cameras): add TODO for conversion of angle limits
    (cherry picked from commit a812f6a6f5ef1a49ee4b36ae132e9d0eaf672450​)

#skipci

diff --git a/Engine/Plugins/Enterprise/GLTFExporter/MigratedFromGit/Source/GLTFExporter/Private/Converters/GLTFCameraConverters.cpp b/Engine/Plugins/Enterprise/GLTFExporter/MigratedFromGit/Source/GLTFExporter/Private/Converters/GLTFCameraConverters.cpp
--- a/Engine/Plugins/Enterprise/GLTFExporter/MigratedFromGit/Source/GLTFExporter/Private/Converters/GLTFCameraConverters.cpp
+++ b/Engine/Plugins/Enterprise/GLTFExporter/MigratedFromGit/Source/GLTFExporter/Private/Converters/GLTFCameraConverters.cpp
@@ -61,10 +61,14 @@ FGLTFJsonCameraIndex FGLTFCameraConverter::Convert(const UCameraComponent* Camer
 			PlayerCamera.Focus = Builder.GetOrAddNode(CameraActor->Focus);
 			PlayerCamera.MaxDistance = FGLTFConverterUtility::ConvertLength(CameraActor->DistanceMax, ExportScale);
 			PlayerCamera.MinDistance = FGLTFConverterUtility::ConvertLength(CameraActor->DistanceMin, ExportScale);
+
+			// TODO: convert limits for pitch and yaw to match differences in camera direction between Unreal and glTF
+
 			PlayerCamera.MaxPitch = CameraActor->PitchAngleMax;
 			PlayerCamera.MinPitch = CameraActor->PitchAngleMin;
 			PlayerCamera.MaxYaw = CameraActor->YawAngleMax;
 			PlayerCamera.MinYaw = CameraActor->YawAngleMin;
+
 			PlayerCamera.RotationSensitivity = CameraActor->RotationSensitivity;
 			PlayerCamera.RotationInertia = CameraActor->RotationInertia;
 			PlayerCamera.DollySensitivity = CameraActor->DollySensitivity;
