[GLTFExporter] feat(player-cameras): transform yaw limits to match glTF specification

commit bc479850db025af10143cc8935882b8fc22f74b5
Date: Thu, 29 Oct 2020 14:57:19 +0100

    feat(player-cameras): transform yaw limits to match glTF specification
    (cherry picked from commit 5f28215ba852822d9d9d66101f26c52b078a7874​)

#skipci

diff --git a/Engine/Plugins/Enterprise/GLTFExporter/MigratedFromGit/Source/GLTFExporter/Private/Converters/GLTFCameraConverters.cpp b/Engine/Plugins/Enterprise/GLTFExporter/MigratedFromGit/Source/GLTFExporter/Private/Converters/GLTFCameraConverters.cpp
--- a/Engine/Plugins/Enterprise/GLTFExporter/MigratedFromGit/Source/GLTFExporter/Private/Converters/GLTFCameraConverters.cpp
+++ b/Engine/Plugins/Enterprise/GLTFExporter/MigratedFromGit/Source/GLTFExporter/Private/Converters/GLTFCameraConverters.cpp
@@ -61,13 +61,20 @@ FGLTFJsonCameraIndex FGLTFCameraConverter::Convert(const UCameraComponent* Camer
 			PlayerCamera.Focus = Builder.GetOrAddNode(CameraActor->Focus);
 			PlayerCamera.MaxDistance = FGLTFConverterUtility::ConvertLength(CameraActor->DistanceMax, ExportScale);
 			PlayerCamera.MinDistance = FGLTFConverterUtility::ConvertLength(CameraActor->DistanceMin, ExportScale);
-
-			// TODO: convert limits for pitch and yaw to match differences in camera direction between Unreal and glTF
-
 			PlayerCamera.MaxPitch = CameraActor->PitchAngleMax;
 			PlayerCamera.MinPitch = CameraActor->PitchAngleMin;
-			PlayerCamera.MaxYaw = CameraActor->YawAngleMax;
-			PlayerCamera.MinYaw = CameraActor->YawAngleMin;
+
+			// Transform yaw limits to match right-handed system and glTF specification for cameras, i.e
+			// positive rotation is CCW, and camera looks down Z- (instead of X+).
+			const float MaxYaw = FMath::Max(-CameraActor->YawAngleMin, -CameraActor->YawAngleMax) - 90.0f;
+			const float MinYaw = FMath::Min(-CameraActor->YawAngleMin, -CameraActor->YawAngleMax) - 90.0f;
+
+			// We prefer the limits to be in the 0..360 range, but we only use MaxYaw to calculate
+			// the needed offset since we need to keep both limits a fixed distance apart from each other.
+			const float PositiveRangeOffset = FRotator::ClampAxis(MaxYaw) - MaxYaw;
+
+			PlayerCamera.MaxYaw = MaxYaw + PositiveRangeOffset;
+			PlayerCamera.MinYaw = MinYaw + PositiveRangeOffset;
 
 			PlayerCamera.RotationSensitivity = CameraActor->RotationSensitivity;
 			PlayerCamera.RotationInertia = CameraActor->RotationInertia;
