[GLTFExporter] refactor(player-cameras): use more intuitive names for camera modes

commit 2352c48e0fcffc72e2a81d376327fc16b699e83d
Date: Mon, 2 Nov 2020 10:06:53 +0100

    refactor(player-cameras): use more intuitive names for camera modes
    (cherry picked from commit 93c2e66cdd2a5ffe9830fcc0930d10e9b9af3d56​)

#skipci

diff --git a/Engine/Plugins/Enterprise/GLTFExporter/MigratedFromGit/Source/GLTFExporter/Private/Converters/GLTFConverterUtility.cpp b/Engine/Plugins/Enterprise/GLTFExporter/MigratedFromGit/Source/GLTFExporter/Private/Converters/GLTFConverterUtility.cpp
--- a/Engine/Plugins/Enterprise/GLTFExporter/MigratedFromGit/Source/GLTFExporter/Private/Converters/GLTFConverterUtility.cpp
+++ b/Engine/Plugins/Enterprise/GLTFExporter/MigratedFromGit/Source/GLTFExporter/Private/Converters/GLTFConverterUtility.cpp
@@ -118,8 +118,8 @@ EGLTFJsonPlayerCameraMode FGLTFConverterUtility::ConvertPlayerCameraMode(EGLTFCa
 {
 	switch (CameraMode)
 	{
-		case EGLTFCameraMode::FirstPerson: return EGLTFJsonPlayerCameraMode::FirstPerson;
-		case EGLTFCameraMode::ThirdPerson: return EGLTFJsonPlayerCameraMode::ThirdPerson;
-		default:                           return EGLTFJsonPlayerCameraMode::None; // TODO: add error handling in callers
+		case EGLTFCameraMode::FreeLook: return EGLTFJsonPlayerCameraMode::FreeLook;
+		case EGLTFCameraMode::Orbital:  return EGLTFJsonPlayerCameraMode::Orbital;
+		default:                        return EGLTFJsonPlayerCameraMode::None; // TODO: add error handling in callers
 	}
 }
diff --git a/Engine/Plugins/Enterprise/GLTFExporter/MigratedFromGit/Source/GLTFExporter/Private/Json/GLTFJsonEnums.h b/Engine/Plugins/Enterprise/GLTFExporter/MigratedFromGit/Source/GLTFExporter/Private/Json/GLTFJsonEnums.h
--- a/Engine/Plugins/Enterprise/GLTFExporter/MigratedFromGit/Source/GLTFExporter/Private/Json/GLTFJsonEnums.h
+++ b/Engine/Plugins/Enterprise/GLTFExporter/MigratedFromGit/Source/GLTFExporter/Private/Json/GLTFJsonEnums.h
@@ -164,6 +164,6 @@ enum class EGLTFJsonLightType
 enum class EGLTFJsonPlayerCameraMode
 {
 	None = -1,
-	FirstPerson,
-	ThirdPerson
+	FreeLook,
+	Orbital
 };
diff --git a/Engine/Plugins/Enterprise/GLTFExporter/MigratedFromGit/Source/GLTFExporter/Private/Json/GLTFJsonUtility.h b/Engine/Plugins/Enterprise/GLTFExporter/MigratedFromGit/Source/GLTFExporter/Private/Json/GLTFJsonUtility.h
--- a/Engine/Plugins/Enterprise/GLTFExporter/MigratedFromGit/Source/GLTFExporter/Private/Json/GLTFJsonUtility.h
+++ b/Engine/Plugins/Enterprise/GLTFExporter/MigratedFromGit/Source/GLTFExporter/Private/Json/GLTFJsonUtility.h
@@ -163,8 +163,8 @@ struct FGLTFJsonUtility
 	{
 		switch (Value)
 		{
-			case EGLTFJsonPlayerCameraMode::FirstPerson: return TEXT("firstPerson");
-			case EGLTFJsonPlayerCameraMode::ThirdPerson: return TEXT("thirdPerson");
+			case EGLTFJsonPlayerCameraMode::FreeLook: return TEXT("freeLook");
+			case EGLTFJsonPlayerCameraMode::Orbital:  return TEXT("orbital");
 			default:
 				checkNoEntry();
 				return TEXT("");
diff --git a/Engine/Plugins/Enterprise/GLTFExporter/MigratedFromGit/Source/GLTFExporterRuntime/Private/Actors/GLTFCameraActor.cpp b/Engine/Plugins/Enterprise/GLTFExporter/MigratedFromGit/Source/GLTFExporterRuntime/Private/Actors/GLTFCameraActor.cpp
--- a/Engine/Plugins/Enterprise/GLTFExporter/MigratedFromGit/Source/GLTFExporterRuntime/Private/Actors/GLTFCameraActor.cpp
+++ b/Engine/Plugins/Enterprise/GLTFExporter/MigratedFromGit/Source/GLTFExporterRuntime/Private/Actors/GLTFCameraActor.cpp
@@ -13,7 +13,7 @@ namespace
 
 AGLTFCameraActor::AGLTFCameraActor(const FObjectInitializer& ObjectInitializer)
 	: Super(ObjectInitializer)
-	, Mode(EGLTFCameraMode::FirstPerson)
+	, Mode(EGLTFCameraMode::FreeLook)
 	, Focus(nullptr)
 	, PitchAngleMin(-90.0f)
 	, PitchAngleMax(90.0f)
@@ -87,7 +87,7 @@ void AGLTFCameraActor::BeginPlay()
 {
 	Super::BeginPlay();
 
-	if (Mode == EGLTFCameraMode::FirstPerson)
+	if (Mode == EGLTFCameraMode::FreeLook)
 	{
 		const FRotator Rotation = GetActorRotation();
 
@@ -96,7 +96,7 @@ void AGLTFCameraActor::BeginPlay()
 		TargetPitch = Pitch;
 		TargetYaw = Yaw;
 	}
-	else if (Mode == EGLTFCameraMode::ThirdPerson)
+	else if (Mode == EGLTFCameraMode::Orbital)
 	{
 		const FVector FocusPosition = GetFocusPosition();
 
@@ -131,7 +131,7 @@ void AGLTFCameraActor::Tick(float DeltaSeconds)
 {
 	Super::Tick(DeltaSeconds);
 
-	if (Mode == EGLTFCameraMode::FirstPerson)
+	if (Mode == EGLTFCameraMode::FreeLook)
 	{
 		const float Alpha = (RotationInertia == 0.0f) ? 1.0f : FMath::Min(DeltaSeconds / RotationInertia, 1.0f);
 		Yaw = FMath::Lerp(Yaw, TargetYaw, Alpha);
@@ -139,7 +139,7 @@ void AGLTFCameraActor::Tick(float DeltaSeconds)
 
 		SetActorRotation(FQuat::MakeFromEuler(FVector(0.0f, Pitch, Yaw)));
 	}
-	else if (Mode == EGLTFCameraMode::ThirdPerson)
+	else if (Mode == EGLTFCameraMode::Orbital)
 	{
 		if (DollyTime != 0.0f)
 		{
@@ -193,7 +193,7 @@ void AGLTFCameraActor::OnMouseY(float AxisValue)
 
 void AGLTFCameraActor::OnMouseWheelAxis(float AxisValue)
 {
-	if (Mode == EGLTFCameraMode::ThirdPerson)
+	if (Mode == EGLTFCameraMode::Orbital)
 	{
 		if (!FMath::IsNearlyZero(AxisValue))
 		{
diff --git a/Engine/Plugins/Enterprise/GLTFExporter/MigratedFromGit/Source/GLTFExporterRuntime/Public/Actors/GLTFCameraActor.h b/Engine/Plugins/Enterprise/GLTFExporter/MigratedFromGit/Source/GLTFExporterRuntime/Public/Actors/GLTFCameraActor.h
--- a/Engine/Plugins/Enterprise/GLTFExporter/MigratedFromGit/Source/GLTFExporterRuntime/Public/Actors/GLTFCameraActor.h
+++ b/Engine/Plugins/Enterprise/GLTFExporter/MigratedFromGit/Source/GLTFExporterRuntime/Public/Actors/GLTFCameraActor.h
@@ -8,15 +8,15 @@
 UENUM(BlueprintType)
 enum class EGLTFCameraMode : uint8
 {
-	FirstPerson,
-	ThirdPerson
+	FreeLook,
+	Orbital
 };
 
 /**
  * GLTF-compatible camera that will carry over settings and simulate the behavior in the resulting viewer.
  * Supports two modes:
- * - FirstPerson: Allows the user to rotate the camera without modifying its position.
- * - ThirdPerson: Focuses one actor in the scene and orbits it through mouse control.
+ * - FreeLook: Allows the user to rotate the camera without modifying its position.
+ * - Orbital: Focuses one actor in the scene and orbits it through mouse control.
  */
 UCLASS(BlueprintType, Blueprintable, DisplayName = "GLTF Camera")
 class GLTFEXPORTERRUNTIME_API AGLTFCameraActor : public ACameraActor
@@ -72,7 +72,7 @@ public:
 	EGLTFCameraMode Mode;
 
 	/* Actor which the camera will focus on and subsequently orbit when using Third Person mode. */
-	UPROPERTY(EditAnywhere, BlueprintReadWrite, Category = "GLTF Camera Actor", meta=(EditCondition="Mode == EGLTFCameraMode::ThirdPerson"))
+	UPROPERTY(EditAnywhere, BlueprintReadWrite, Category = "GLTF Camera Actor", meta=(EditCondition="Mode == EGLTFCameraMode::Orbital"))
 	AActor* Focus;
 
 	/* Minimum pitch angle (in degrees) for the camera. */
