[GLTFExporter] refactor: rename extension EPIC_player_cameras to EPIC_camera_controls

commit 7bd41fe7421bd98d9d2b4c26010d5babc9cec3bd
Date: Mon, 2 Nov 2020 10:19:47 +0100

    refactor: rename extension EPIC_player_cameras to EPIC_camera_controls
    (cherry picked from commit 99cfb39c3b36a78c21c2c687578bd98bf293316d​)
    # Conflicts from cherry-pick:
    #	Plugins/GLTFExporter/Source/GLTFExporter/Private/GLTFExportOptions.cpp
    #	Plugins/GLTFExporter/Source/GLTFExporter/Public/GLTFExportOptions.h

#skipci

diff --git a/Engine/Plugins/Enterprise/GLTFExporter/MigratedFromGit/Source/GLTFExporter/Private/Converters/GLTFCameraConverters.cpp b/Engine/Plugins/Enterprise/GLTFExporter/MigratedFromGit/Source/GLTFExporter/Private/Converters/GLTFCameraConverters.cpp
--- a/Engine/Plugins/Enterprise/GLTFExporter/MigratedFromGit/Source/GLTFExporter/Private/Converters/GLTFCameraConverters.cpp
+++ b/Engine/Plugins/Enterprise/GLTFExporter/MigratedFromGit/Source/GLTFExporter/Private/Converters/GLTFCameraConverters.cpp
@@ -54,15 +54,15 @@ FGLTFJsonCameraIndex FGLTFCameraConverter::Convert(const UCameraComponent* Camer
 
 	if (CameraActor != nullptr)
 	{
-		if (Builder.ExportOptions->bExportPlayerCameras)
+		if (Builder.ExportOptions->bExportCameraControls)
 		{
-			FGLTFJsonPlayerCamera PlayerCamera;
-			PlayerCamera.Mode = FGLTFConverterUtility::ConvertPlayerCameraMode(CameraActor->Mode);
-			PlayerCamera.Focus = Builder.GetOrAddNode(CameraActor->Focus);
-			PlayerCamera.MaxDistance = FGLTFConverterUtility::ConvertLength(CameraActor->DistanceMax, ExportScale);
-			PlayerCamera.MinDistance = FGLTFConverterUtility::ConvertLength(CameraActor->DistanceMin, ExportScale);
-			PlayerCamera.MaxPitch = CameraActor->PitchAngleMax;
-			PlayerCamera.MinPitch = CameraActor->PitchAngleMin;
+			FGLTFJsonCameraControl CameraControl;
+			CameraControl.Mode = FGLTFConverterUtility::ConvertCameraControlMode(CameraActor->Mode);
+			CameraControl.Focus = Builder.GetOrAddNode(CameraActor->Focus);
+			CameraControl.MaxDistance = FGLTFConverterUtility::ConvertLength(CameraActor->DistanceMax, ExportScale);
+			CameraControl.MinDistance = FGLTFConverterUtility::ConvertLength(CameraActor->DistanceMin, ExportScale);
+			CameraControl.MaxPitch = CameraActor->PitchAngleMax;
+			CameraControl.MinPitch = CameraActor->PitchAngleMin;
 
 			// Transform yaw limits to match right-handed system and glTF specification for cameras, i.e
 			// positive rotation is CCW, and camera looks down Z- (instead of X+).
@@ -73,15 +73,15 @@ FGLTFJsonCameraIndex FGLTFCameraConverter::Convert(const UCameraComponent* Camer
 			// the needed offset since we need to keep both limits a fixed distance apart from each other.
 			const float PositiveRangeOffset = FRotator::ClampAxis(MaxYaw) - MaxYaw;
 
-			PlayerCamera.MaxYaw = MaxYaw + PositiveRangeOffset;
-			PlayerCamera.MinYaw = MinYaw + PositiveRangeOffset;
+			CameraControl.MaxYaw = MaxYaw + PositiveRangeOffset;
+			CameraControl.MinYaw = MinYaw + PositiveRangeOffset;
 
-			PlayerCamera.RotationSensitivity = CameraActor->RotationSensitivity;
-			PlayerCamera.RotationInertia = CameraActor->RotationInertia;
-			PlayerCamera.DollySensitivity = CameraActor->DollySensitivity;
-			PlayerCamera.DollyDuration = CameraActor->DollyDuration;
+			CameraControl.RotationSensitivity = CameraActor->RotationSensitivity;
+			CameraControl.RotationInertia = CameraActor->RotationInertia;
+			CameraControl.DollySensitivity = CameraActor->DollySensitivity;
+			CameraControl.DollyDuration = CameraActor->DollyDuration;
 
-			Camera.PlayerCamera = PlayerCamera;
+			Camera.CameraControl = CameraControl;
 		}
 	}
 
diff --git a/Engine/Plugins/Enterprise/GLTFExporter/MigratedFromGit/Source/GLTFExporter/Private/Converters/GLTFConverterUtility.cpp b/Engine/Plugins/Enterprise/GLTFExporter/MigratedFromGit/Source/GLTFExporter/Private/Converters/GLTFConverterUtility.cpp
--- a/Engine/Plugins/Enterprise/GLTFExporter/MigratedFromGit/Source/GLTFExporter/Private/Converters/GLTFConverterUtility.cpp
+++ b/Engine/Plugins/Enterprise/GLTFExporter/MigratedFromGit/Source/GLTFExporter/Private/Converters/GLTFConverterUtility.cpp
@@ -114,12 +114,12 @@ EGLTFJsonCubeFace FGLTFConverterUtility::ConvertCubeFace(ECubeFace CubeFace)
 	}
 }
 
-EGLTFJsonPlayerCameraMode FGLTFConverterUtility::ConvertPlayerCameraMode(EGLTFCameraMode CameraMode)
+EGLTFJsonCameraControlMode FGLTFConverterUtility::ConvertCameraControlMode(EGLTFCameraControlMode CameraMode)
 {
 	switch (CameraMode)
 	{
-		case EGLTFCameraMode::FreeLook: return EGLTFJsonPlayerCameraMode::FreeLook;
-		case EGLTFCameraMode::Orbital:  return EGLTFJsonPlayerCameraMode::Orbital;
-		default:                        return EGLTFJsonPlayerCameraMode::None; // TODO: add error handling in callers
+		case EGLTFCameraControlMode::FreeLook: return EGLTFJsonCameraControlMode::FreeLook;
+		case EGLTFCameraControlMode::Orbital:  return EGLTFJsonCameraControlMode::Orbital;
+		default:                               return EGLTFJsonCameraControlMode::None; // TODO: add error handling in callers
 	}
 }
diff --git a/Engine/Plugins/Enterprise/GLTFExporter/MigratedFromGit/Source/GLTFExporter/Private/Converters/GLTFConverterUtility.h b/Engine/Plugins/Enterprise/GLTFExporter/MigratedFromGit/Source/GLTFExporter/Private/Converters/GLTFConverterUtility.h
--- a/Engine/Plugins/Enterprise/GLTFExporter/MigratedFromGit/Source/GLTFExporter/Private/Converters/GLTFConverterUtility.h
+++ b/Engine/Plugins/Enterprise/GLTFExporter/MigratedFromGit/Source/GLTFExporter/Private/Converters/GLTFConverterUtility.h
@@ -12,7 +12,7 @@
 #include "Json/GLTFJsonQuaternion.h"
 #include "Engine/EngineTypes.h"
 
-enum class EGLTFCameraMode : unsigned char;
+enum class EGLTFCameraControlMode : unsigned char;
 
 struct FGLTFConverterUtility
 {
@@ -195,7 +195,7 @@ struct FGLTFConverterUtility
 
 	static EGLTFJsonCubeFace ConvertCubeFace(ECubeFace CubeFace);
 
-	static EGLTFJsonPlayerCameraMode ConvertPlayerCameraMode(EGLTFCameraMode CameraMode);
+	static EGLTFJsonCameraControlMode ConvertCameraControlMode(EGLTFCameraControlMode CameraMode);
 
 	template <typename ComponentType>
 	static EGLTFJsonComponentType GetComponentType()
diff --git a/Engine/Plugins/Enterprise/GLTFExporter/MigratedFromGit/Source/GLTFExporter/Private/GLTFExportOptions.cpp b/Engine/Plugins/Enterprise/GLTFExporter/MigratedFromGit/Source/GLTFExporter/Private/GLTFExportOptions.cpp
--- a/Engine/Plugins/Enterprise/GLTFExporter/MigratedFromGit/Source/GLTFExporter/Private/GLTFExportOptions.cpp
+++ b/Engine/Plugins/Enterprise/GLTFExporter/MigratedFromGit/Source/GLTFExporter/Private/GLTFExportOptions.cpp
@@ -35,7 +35,7 @@ void UGLTFExportOptions::ResetToDefault()
 	bExportHiddenInGame = false;
 	ExportLights = EGLTFExporterLightMobility::MovableAndStationary;
 	bExportCameras = true;
-	bExportPlayerCameras = true;
+	bExportCameraControls = true;
 	bExportHDRIBackdrops = true;
 	bExportSkySpheres = true;
 	bExportVariantSets = true;
diff --git a/Engine/Plugins/Enterprise/GLTFExporter/MigratedFromGit/Source/GLTFExporter/Private/Json/GLTFJsonCamera.h b/Engine/Plugins/Enterprise/GLTFExporter/MigratedFromGit/Source/GLTFExporter/Private/Json/GLTFJsonCamera.h
--- a/Engine/Plugins/Enterprise/GLTFExporter/MigratedFromGit/Source/GLTFExporter/Private/Json/GLTFJsonCamera.h
+++ b/Engine/Plugins/Enterprise/GLTFExporter/MigratedFromGit/Source/GLTFExporter/Private/Json/GLTFJsonCamera.h
@@ -3,7 +3,7 @@
 #pragma once
 
 #include "Json/GLTFJsonEnums.h"
-#include "Json/GLTFJsonPlayerCamera.h"
+#include "Json/GLTFJsonCameraControl.h"
 #include "Json/GLTFJsonUtility.h"
 #include "Serialization/JsonSerializer.h"
 
@@ -78,8 +78,8 @@ struct FGLTFJsonCamera
 {
 	FString Name;
 
-	EGLTFJsonCameraType             Type;
-	TOptional<FGLTFJsonPlayerCamera> PlayerCamera;
+	EGLTFJsonCameraType               Type;
+	TOptional<FGLTFJsonCameraControl> CameraControl;
 
 	union {
 		FGLTFJsonOrthographic Orthographic;
@@ -121,15 +121,15 @@ struct FGLTFJsonCamera
 				break;
 		}
 
-		if (PlayerCamera.IsSet())
+		if (CameraControl.IsSet())
 		{
 			JsonWriter.WriteObjectStart(TEXT("extensions"));
 
-			const EGLTFJsonExtension Extension = EGLTFJsonExtension::EPIC_PlayerCameras;
+			const EGLTFJsonExtension Extension = EGLTFJsonExtension::EPIC_CameraControls;
 			Extensions.Used.Add(Extension);
 
 			JsonWriter.WriteIdentifierPrefix(FGLTFJsonUtility::ToString(Extension));
-			PlayerCamera.GetValue().WriteObject(JsonWriter, Extensions);
+			CameraControl.GetValue().WriteObject(JsonWriter, Extensions);
 
 			JsonWriter.WriteObjectEnd();
 		}
diff --git a/Engine/Plugins/Enterprise/GLTFExporter/MigratedFromGit/Source/GLTFExporter/Private/Json/GLTFJsonCameraControl.h b/Engine/Plugins/Enterprise/GLTFExporter/MigratedFromGit/Source/GLTFExporter/Private/Json/GLTFJsonCameraControl.h
new file mode 100644
--- /dev/null
+++ b/Engine/Plugins/Enterprise/GLTFExporter/MigratedFromGit/Source/GLTFExporter/Private/Json/GLTFJsonCameraControl.h
@@ -0,0 +1,43 @@
+// Copyright Epic Games, Inc. All Rights Reserved.
+
+#pragma once
+
+#include "Json/GLTFJsonIndex.h"
+#include "Json/GLTFJsonVector4.h"
+
+struct FGLTFJsonCameraControl
+{
+	EGLTFJsonCameraControlMode Mode;
+	FGLTFJsonNodeIndex         Focus;
+	float                      MaxDistance;
+	float                      MinDistance;
+	float                      MaxPitch;
+	float                      MinPitch;
+	float                      MaxYaw;
+	float                      MinYaw;
+	float                      RotationSensitivity;
+	float                      RotationInertia;
+	float                      DollySensitivity;
+	float                      DollyDuration;
+
+	template <class CharType = TCHAR, class PrintPolicy = TPrettyJsonPrintPolicy<CharType>>
+	void WriteObject(TJsonWriter<CharType, PrintPolicy>& JsonWriter, FGLTFJsonExtensions& Extensions) const
+	{
+		JsonWriter.WriteObjectStart();
+
+		JsonWriter.WriteValue(TEXT("mode"), FGLTFJsonUtility::ToString(Mode));
+		JsonWriter.WriteValue(TEXT("focus"), Focus);
+		FGLTFJsonUtility::WriteExactValue(JsonWriter, TEXT("maxDistance"), MaxDistance);
+		FGLTFJsonUtility::WriteExactValue(JsonWriter, TEXT("minDistance"), MinDistance);
+		FGLTFJsonUtility::WriteExactValue(JsonWriter, TEXT("maxPitch"), MaxPitch);
+		FGLTFJsonUtility::WriteExactValue(JsonWriter, TEXT("minPitch"), MinPitch);
+		FGLTFJsonUtility::WriteExactValue(JsonWriter, TEXT("maxYaw"), MaxYaw);
+		FGLTFJsonUtility::WriteExactValue(JsonWriter, TEXT("minYaw"), MinYaw);
+		FGLTFJsonUtility::WriteExactValue(JsonWriter, TEXT("rotationSensitivity"), RotationSensitivity);
+		FGLTFJsonUtility::WriteExactValue(JsonWriter, TEXT("rotationInertia"), RotationInertia);
+		FGLTFJsonUtility::WriteExactValue(JsonWriter, TEXT("dollySensitivity"), DollySensitivity);
+		FGLTFJsonUtility::WriteExactValue(JsonWriter, TEXT("dollyDuration"), DollyDuration);
+
+		JsonWriter.WriteObjectEnd();
+	}
+};
diff --git a/Engine/Plugins/Enterprise/GLTFExporter/MigratedFromGit/Source/GLTFExporter/Private/Json/GLTFJsonEnums.h b/Engine/Plugins/Enterprise/GLTFExporter/MigratedFromGit/Source/GLTFExporter/Private/Json/GLTFJsonEnums.h
--- a/Engine/Plugins/Enterprise/GLTFExporter/MigratedFromGit/Source/GLTFExporter/Private/Json/GLTFJsonEnums.h
+++ b/Engine/Plugins/Enterprise/GLTFExporter/MigratedFromGit/Source/GLTFExporter/Private/Json/GLTFJsonEnums.h
@@ -11,7 +11,7 @@ enum class EGLTFJsonExtension
 	KHR_MaterialsClearCoat,
 	KHR_MeshQuantization,
 	KHR_TextureTransform,
-	EPIC_PlayerCameras,
+	EPIC_CameraControls,
 	EPIC_HDRIBackdrops,
 	EPIC_AnimationHotspots,
 	EPIC_LightmapTextures,
@@ -161,7 +161,7 @@ enum class EGLTFJsonLightType
 	Spot
 };
 
-enum class EGLTFJsonPlayerCameraMode
+enum class EGLTFJsonCameraControlMode
 {
 	None = -1,
 	FreeLook,
diff --git a/Engine/Plugins/Enterprise/GLTFExporter/MigratedFromGit/Source/GLTFExporter/Private/Json/GLTFJsonPlayerCamera.h b/Engine/Plugins/Enterprise/GLTFExporter/MigratedFromGit/Source/GLTFExporter/Private/Json/GLTFJsonPlayerCamera.h
deleted file mode 100644
diff --git a/Engine/Plugins/Enterprise/GLTFExporter/MigratedFromGit/Source/GLTFExporter/Private/Json/GLTFJsonUtility.h b/Engine/Plugins/Enterprise/GLTFExporter/MigratedFromGit/Source/GLTFExporter/Private/Json/GLTFJsonUtility.h
--- a/Engine/Plugins/Enterprise/GLTFExporter/MigratedFromGit/Source/GLTFExporter/Private/Json/GLTFJsonUtility.h
+++ b/Engine/Plugins/Enterprise/GLTFExporter/MigratedFromGit/Source/GLTFExporter/Private/Json/GLTFJsonUtility.h
@@ -22,7 +22,7 @@ struct FGLTFJsonUtility
 			case EGLTFJsonExtension::KHR_MaterialsClearCoat:  return TEXT("KHR_materials_clearcoat");
 			case EGLTFJsonExtension::KHR_MeshQuantization:    return TEXT("KHR_mesh_quantization");
 			case EGLTFJsonExtension::KHR_TextureTransform:    return TEXT("KHR_texture_transform");
-			case EGLTFJsonExtension::EPIC_PlayerCameras:      return TEXT("EPIC_player_cameras");
+			case EGLTFJsonExtension::EPIC_CameraControls:     return TEXT("EPIC_camera_controls");
 			case EGLTFJsonExtension::EPIC_HDRIBackdrops:      return TEXT("EPIC_hdri_backdrops");
 			case EGLTFJsonExtension::EPIC_AnimationHotspots:  return TEXT("EPIC_animation_hotspots");
 			case EGLTFJsonExtension::EPIC_LightmapTextures:   return TEXT("EPIC_lightmap_textures");
@@ -159,12 +159,12 @@ struct FGLTFJsonUtility
 		}
 	}
 
-	static const TCHAR* ToString(EGLTFJsonPlayerCameraMode Value)
+	static const TCHAR* ToString(EGLTFJsonCameraControlMode Value)
 	{
 		switch (Value)
 		{
-			case EGLTFJsonPlayerCameraMode::FreeLook: return TEXT("freeLook");
-			case EGLTFJsonPlayerCameraMode::Orbital:  return TEXT("orbital");
+			case EGLTFJsonCameraControlMode::FreeLook: return TEXT("freeLook");
+			case EGLTFJsonCameraControlMode::Orbital:  return TEXT("orbital");
 			default:
 				checkNoEntry();
 				return TEXT("");
diff --git a/Engine/Plugins/Enterprise/GLTFExporter/MigratedFromGit/Source/GLTFExporter/Public/GLTFExportOptions.h b/Engine/Plugins/Enterprise/GLTFExporter/MigratedFromGit/Source/GLTFExporter/Public/GLTFExportOptions.h
--- a/Engine/Plugins/Enterprise/GLTFExporter/MigratedFromGit/Source/GLTFExporter/Public/GLTFExportOptions.h
+++ b/Engine/Plugins/Enterprise/GLTFExporter/MigratedFromGit/Source/GLTFExporter/Public/GLTFExportOptions.h
@@ -113,7 +113,7 @@ public:
 	bool bExportCameras;
 
 	UPROPERTY(EditAnywhere, BlueprintReadWrite, Config, Category = Scene, Meta = (EditCondition = "bExportCameras"))
-	bool bExportPlayerCameras;
+	bool bExportCameraControls;
 
 	UPROPERTY(EditAnywhere, BlueprintReadWrite, Config, Category = Scene)
 	bool bExportHDRIBackdrops;
diff --git a/Engine/Plugins/Enterprise/GLTFExporter/MigratedFromGit/Source/GLTFExporterRuntime/Private/Actors/GLTFCameraActor.cpp b/Engine/Plugins/Enterprise/GLTFExporter/MigratedFromGit/Source/GLTFExporterRuntime/Private/Actors/GLTFCameraActor.cpp
--- a/Engine/Plugins/Enterprise/GLTFExporter/MigratedFromGit/Source/GLTFExporterRuntime/Private/Actors/GLTFCameraActor.cpp
+++ b/Engine/Plugins/Enterprise/GLTFExporter/MigratedFromGit/Source/GLTFExporterRuntime/Private/Actors/GLTFCameraActor.cpp
@@ -13,7 +13,7 @@ namespace
 
 AGLTFCameraActor::AGLTFCameraActor(const FObjectInitializer& ObjectInitializer)
 	: Super(ObjectInitializer)
-	, Mode(EGLTFCameraMode::FreeLook)
+	, Mode(EGLTFCameraControlMode::FreeLook)
 	, Focus(nullptr)
 	, PitchAngleMin(-90.0f)
 	, PitchAngleMax(90.0f)
@@ -87,7 +87,7 @@ void AGLTFCameraActor::BeginPlay()
 {
 	Super::BeginPlay();
 
-	if (Mode == EGLTFCameraMode::FreeLook)
+	if (Mode == EGLTFCameraControlMode::FreeLook)
 	{
 		const FRotator Rotation = GetActorRotation();
 
@@ -96,7 +96,7 @@ void AGLTFCameraActor::BeginPlay()
 		TargetPitch = Pitch;
 		TargetYaw = Yaw;
 	}
-	else if (Mode == EGLTFCameraMode::Orbital)
+	else if (Mode == EGLTFCameraControlMode::Orbital)
 	{
 		const FVector FocusPosition = GetFocusPosition();
 
@@ -131,7 +131,7 @@ void AGLTFCameraActor::Tick(float DeltaSeconds)
 {
 	Super::Tick(DeltaSeconds);
 
-	if (Mode == EGLTFCameraMode::FreeLook)
+	if (Mode == EGLTFCameraControlMode::FreeLook)
 	{
 		const float Alpha = (RotationInertia == 0.0f) ? 1.0f : FMath::Min(DeltaSeconds / RotationInertia, 1.0f);
 		Yaw = FMath::Lerp(Yaw, TargetYaw, Alpha);
@@ -139,7 +139,7 @@ void AGLTFCameraActor::Tick(float DeltaSeconds)
 
 		SetActorRotation(FQuat::MakeFromEuler(FVector(0.0f, Pitch, Yaw)));
 	}
-	else if (Mode == EGLTFCameraMode::Orbital)
+	else if (Mode == EGLTFCameraControlMode::Orbital)
 	{
 		if (DollyTime != 0.0f)
 		{
@@ -193,7 +193,7 @@ void AGLTFCameraActor::OnMouseY(float AxisValue)
 
 void AGLTFCameraActor::OnMouseWheelAxis(float AxisValue)
 {
-	if (Mode == EGLTFCameraMode::Orbital)
+	if (Mode == EGLTFCameraControlMode::Orbital)
 	{
 		if (!FMath::IsNearlyZero(AxisValue))
 		{
diff --git a/Engine/Plugins/Enterprise/GLTFExporter/MigratedFromGit/Source/GLTFExporterRuntime/Public/Actors/GLTFCameraActor.h b/Engine/Plugins/Enterprise/GLTFExporter/MigratedFromGit/Source/GLTFExporterRuntime/Public/Actors/GLTFCameraActor.h
--- a/Engine/Plugins/Enterprise/GLTFExporter/MigratedFromGit/Source/GLTFExporterRuntime/Public/Actors/GLTFCameraActor.h
+++ b/Engine/Plugins/Enterprise/GLTFExporter/MigratedFromGit/Source/GLTFExporterRuntime/Public/Actors/GLTFCameraActor.h
@@ -6,7 +6,7 @@
 #include "GLTFCameraActor.generated.h"
 
 UENUM(BlueprintType)
-enum class EGLTFCameraMode : uint8
+enum class EGLTFCameraControlMode : uint8
 {
 	FreeLook,
 	Orbital
@@ -69,10 +69,10 @@ public:
 
 	/* Camera mode */
 	UPROPERTY(EditAnywhere, BlueprintReadWrite, Category = "GLTF Camera Actor")
-	EGLTFCameraMode Mode;
+	EGLTFCameraControlMode Mode;
 
 	/* Actor which the camera will focus on and subsequently orbit when using Third Person mode. */
-	UPROPERTY(EditAnywhere, BlueprintReadWrite, Category = "GLTF Camera Actor", meta=(EditCondition="Mode == EGLTFCameraMode::Orbital"))
+	UPROPERTY(EditAnywhere, BlueprintReadWrite, Category = "GLTF Camera Actor", meta=(EditCondition="Mode == EGLTFCameraControlMode::Orbital"))
 	AActor* Focus;
 
 	/* Minimum pitch angle (in degrees) for the camera. */
