[GLTFExporter] refactor(camera-controls): rename focus property to target

commit e12f902626b30319498f53521315810de9768a78
Date: Mon, 2 Nov 2020 10:24:15 +0100

    refactor(camera-controls): rename focus property to target
    (cherry picked from commit 6d9ec107c733e45c49185a316e2422043160de43​)

#skipci

diff --git a/Engine/Plugins/Enterprise/GLTFExporter/MigratedFromGit/Source/GLTFExporter/Private/Converters/GLTFCameraConverters.cpp b/Engine/Plugins/Enterprise/GLTFExporter/MigratedFromGit/Source/GLTFExporter/Private/Converters/GLTFCameraConverters.cpp
--- a/Engine/Plugins/Enterprise/GLTFExporter/MigratedFromGit/Source/GLTFExporter/Private/Converters/GLTFCameraConverters.cpp
+++ b/Engine/Plugins/Enterprise/GLTFExporter/MigratedFromGit/Source/GLTFExporter/Private/Converters/GLTFCameraConverters.cpp
@@ -58,7 +58,7 @@ FGLTFJsonCameraIndex FGLTFCameraConverter::Convert(const UCameraComponent* Camer
 		{
 			FGLTFJsonCameraControl CameraControl;
 			CameraControl.Mode = FGLTFConverterUtility::ConvertCameraControlMode(CameraActor->Mode);
-			CameraControl.Focus = Builder.GetOrAddNode(CameraActor->Focus);
+			CameraControl.Target = Builder.GetOrAddNode(CameraActor->Target);
 			CameraControl.MaxDistance = FGLTFConverterUtility::ConvertLength(CameraActor->DistanceMax, ExportScale);
 			CameraControl.MinDistance = FGLTFConverterUtility::ConvertLength(CameraActor->DistanceMin, ExportScale);
 			CameraControl.MaxPitch = CameraActor->PitchAngleMax;
diff --git a/Engine/Plugins/Enterprise/GLTFExporter/MigratedFromGit/Source/GLTFExporter/Private/Json/GLTFJsonCameraControl.h b/Engine/Plugins/Enterprise/GLTFExporter/MigratedFromGit/Source/GLTFExporter/Private/Json/GLTFJsonCameraControl.h
--- a/Engine/Plugins/Enterprise/GLTFExporter/MigratedFromGit/Source/GLTFExporter/Private/Json/GLTFJsonCameraControl.h
+++ b/Engine/Plugins/Enterprise/GLTFExporter/MigratedFromGit/Source/GLTFExporter/Private/Json/GLTFJsonCameraControl.h
@@ -8,7 +8,7 @@
 struct FGLTFJsonCameraControl
 {
 	EGLTFJsonCameraControlMode Mode;
-	FGLTFJsonNodeIndex         Focus;
+	FGLTFJsonNodeIndex         Target;
 	float                      MaxDistance;
 	float                      MinDistance;
 	float                      MaxPitch;
@@ -26,7 +26,7 @@ struct FGLTFJsonCameraControl
 		JsonWriter.WriteObjectStart();
 
 		JsonWriter.WriteValue(TEXT("mode"), FGLTFJsonUtility::ToString(Mode));
-		JsonWriter.WriteValue(TEXT("focus"), Focus);
+		JsonWriter.WriteValue(TEXT("target"), Target);
 		FGLTFJsonUtility::WriteExactValue(JsonWriter, TEXT("maxDistance"), MaxDistance);
 		FGLTFJsonUtility::WriteExactValue(JsonWriter, TEXT("minDistance"), MinDistance);
 		FGLTFJsonUtility::WriteExactValue(JsonWriter, TEXT("maxPitch"), MaxPitch);
diff --git a/Engine/Plugins/Enterprise/GLTFExporter/MigratedFromGit/Source/GLTFExporterRuntime/Private/Actors/GLTFCameraActor.cpp b/Engine/Plugins/Enterprise/GLTFExporter/MigratedFromGit/Source/GLTFExporterRuntime/Private/Actors/GLTFCameraActor.cpp
--- a/Engine/Plugins/Enterprise/GLTFExporter/MigratedFromGit/Source/GLTFExporterRuntime/Private/Actors/GLTFCameraActor.cpp
+++ b/Engine/Plugins/Enterprise/GLTFExporter/MigratedFromGit/Source/GLTFExporterRuntime/Private/Actors/GLTFCameraActor.cpp
@@ -14,7 +14,7 @@ namespace
 AGLTFCameraActor::AGLTFCameraActor(const FObjectInitializer& ObjectInitializer)
 	: Super(ObjectInitializer)
 	, Mode(EGLTFCameraControlMode::FreeLook)
-	, Focus(nullptr)
+	, Target(nullptr)
 	, PitchAngleMin(-90.0f)
 	, PitchAngleMax(90.0f)
 	, YawAngleMin(0.0f)
@@ -48,13 +48,13 @@ void AGLTFCameraActor::PostEditChangeProperty(FPropertyChangedEvent& PropertyCha
 	{
 		const FName PropertyName = PropertyThatChanged->GetFName();
 
-		if (PropertyName == GET_MEMBER_NAME_CHECKED(ThisClass, Focus))
+		if (PropertyName == GET_MEMBER_NAME_CHECKED(ThisClass, Target))
 		{
-			if (Focus == this)
+			if (Target == this)
 			{
-				Focus = nullptr;
+				Target = nullptr;
 				// TODO: don't use LogTemp
-				UE_LOG(LogTemp, Warning, TEXT("The camera focus must not be the camera's own actor"));
+				UE_LOG(LogTemp, Warning, TEXT("The camera target must not be the camera's own actor"));
 			}
 		}
 		else if (PropertyName == GET_MEMBER_NAME_CHECKED(ThisClass, PitchAngleMin))
@@ -237,7 +237,7 @@ FRotator AGLTFCameraActor::GetLookAtRotation(const FVector TargetPosition) const
 
 FVector AGLTFCameraActor::GetFocusPosition() const
 {
-	return this->Focus != nullptr ? this->Focus->GetActorLocation() : FVector::ZeroVector;
+	return this->Target != nullptr ? this->Target->GetActorLocation() : FVector::ZeroVector;
 }
 
 bool AGLTFCameraActor::SetAutoActivateForPlayer(const EAutoReceiveInput::Type Player)
diff --git a/Engine/Plugins/Enterprise/GLTFExporter/MigratedFromGit/Source/GLTFExporterRuntime/Public/Actors/GLTFCameraActor.h b/Engine/Plugins/Enterprise/GLTFExporter/MigratedFromGit/Source/GLTFExporterRuntime/Public/Actors/GLTFCameraActor.h
--- a/Engine/Plugins/Enterprise/GLTFExporter/MigratedFromGit/Source/GLTFExporterRuntime/Public/Actors/GLTFCameraActor.h
+++ b/Engine/Plugins/Enterprise/GLTFExporter/MigratedFromGit/Source/GLTFExporterRuntime/Public/Actors/GLTFCameraActor.h
@@ -73,7 +73,7 @@ public:
 
 	/* Actor which the camera will focus on and subsequently orbit when using Third Person mode. */
 	UPROPERTY(EditAnywhere, BlueprintReadWrite, Category = "GLTF Camera Actor", meta=(EditCondition="Mode == EGLTFCameraControlMode::Orbital"))
-	AActor* Focus;
+	AActor* Target;
 
 	/* Minimum pitch angle (in degrees) for the camera. */
 	UPROPERTY(EditAnywhere, BlueprintReadWrite, Category = "GLTF Camera Actor")
