[GLTFExporter] refactor: don't write target in EPIC_camera_controls if target is -1 or mode is not orbital

commit e1d50e5adb8fb42cb00442435efdb92da15cb585
Date: Tue, 3 Nov 2020 00:08:46 +0100

    refactor: don't write target in EPIC_camera_controls if target is -1 or mode is not orbital
    (cherry picked from commit 589457830777641f566da00ba2f5e62826e76f67​)

#skipci

diff --git a/Engine/Plugins/Enterprise/GLTFExporter/MigratedFromGit/Source/GLTFExporter/Private/Json/GLTFJsonCameraControl.h b/Engine/Plugins/Enterprise/GLTFExporter/MigratedFromGit/Source/GLTFExporter/Private/Json/GLTFJsonCameraControl.h
--- a/Engine/Plugins/Enterprise/GLTFExporter/MigratedFromGit/Source/GLTFExporter/Private/Json/GLTFJsonCameraControl.h
+++ b/Engine/Plugins/Enterprise/GLTFExporter/MigratedFromGit/Source/GLTFExporter/Private/Json/GLTFJsonCameraControl.h
@@ -8,17 +8,17 @@
 struct FGLTFJsonCameraControl
 {
 	EGLTFJsonCameraControlMode Mode;
-	FGLTFJsonNodeIndex         Target;
-	float                      MaxDistance;
-	float                      MinDistance;
-	float                      MaxPitch;
-	float                      MinPitch;
-	float                      MaxYaw;
-	float                      MinYaw;
-	float                      RotationSensitivity;
-	float                      RotationInertia;
-	float                      DollySensitivity;
-	float                      DollyDuration;
+	FGLTFJsonNodeIndex Target;
+	float MaxDistance;
+	float MinDistance;
+	float MaxPitch;
+	float MinPitch;
+	float MaxYaw;
+	float MinYaw;
+	float RotationSensitivity;
+	float RotationInertia;
+	float DollySensitivity;
+	float DollyDuration;
 
 	template <class CharType = TCHAR, class PrintPolicy = TPrettyJsonPrintPolicy<CharType>>
 	void WriteObject(TJsonWriter<CharType, PrintPolicy>& JsonWriter, FGLTFJsonExtensions& Extensions) const
@@ -26,7 +26,12 @@ struct FGLTFJsonCameraControl
 		JsonWriter.WriteObjectStart();
 
 		JsonWriter.WriteValue(TEXT("mode"), FGLTFJsonUtility::ToString(Mode));
-		JsonWriter.WriteValue(TEXT("target"), Target);
+
+		if (Target != INDEX_NONE && Mode == EGLTFJsonCameraControlMode::Orbital)
+		{
+			JsonWriter.WriteValue(TEXT("target"), Target);
+		}
+
 		FGLTFJsonUtility::WriteExactValue(JsonWriter, TEXT("maxDistance"), MaxDistance);
 		FGLTFJsonUtility::WriteExactValue(JsonWriter, TEXT("minDistance"), MinDistance);
 		FGLTFJsonUtility::WriteExactValue(JsonWriter, TEXT("maxPitch"), MaxPitch);
