[GLTFExporter] refactor(camera-controls): don't write angle limits in EPIC_camera_controls if they use default values

commit e2b26ed239f533c4e84a3447cb40b7942023a10e
Date: Tue, 3 Nov 2020 10:23:14 +0100

    refactor(camera-controls): don't write angle limits in EPIC_camera_controls if they use default values
    (cherry picked from commit 7856ec2e201e4d467e97a9a2b27544966dce9d0e​)

#skipci

diff --git a/Engine/Plugins/Enterprise/GLTFExporter/MigratedFromGit/Source/GLTFExporter/Private/Json/GLTFJsonCameraControl.h b/Engine/Plugins/Enterprise/GLTFExporter/MigratedFromGit/Source/GLTFExporter/Private/Json/GLTFJsonCameraControl.h
--- a/Engine/Plugins/Enterprise/GLTFExporter/MigratedFromGit/Source/GLTFExporter/Private/Json/GLTFJsonCameraControl.h
+++ b/Engine/Plugins/Enterprise/GLTFExporter/MigratedFromGit/Source/GLTFExporter/Private/Json/GLTFJsonCameraControl.h
@@ -20,6 +20,21 @@ struct FGLTFJsonCameraControl
 	float DollySensitivity;
 	float DollyDuration;
 
+	FGLTFJsonCameraControl()
+		: Mode(EGLTFJsonCameraControlMode::FreeLook)
+		, MaxDistance(0)
+		, MinDistance(0)
+		, MaxPitch(90.0f)
+		, MinPitch(-90.0f)
+		, MaxYaw(360.0f)
+		, MinYaw(0)
+		, RotationSensitivity(0)
+		, RotationInertia(0)
+		, DollySensitivity(0)
+		, DollyDuration(0)
+	{
+	}
+
 	template <class CharType = TCHAR, class PrintPolicy = TPrettyJsonPrintPolicy<CharType>>
 	void WriteObject(TJsonWriter<CharType, PrintPolicy>& JsonWriter, FGLTFJsonExtensions& Extensions) const
 	{
@@ -34,10 +49,27 @@ struct FGLTFJsonCameraControl
 
 		FGLTFJsonUtility::WriteExactValue(JsonWriter, TEXT("maxDistance"), MaxDistance);
 		FGLTFJsonUtility::WriteExactValue(JsonWriter, TEXT("minDistance"), MinDistance);
-		FGLTFJsonUtility::WriteExactValue(JsonWriter, TEXT("maxPitch"), MaxPitch);
-		FGLTFJsonUtility::WriteExactValue(JsonWriter, TEXT("minPitch"), MinPitch);
-		FGLTFJsonUtility::WriteExactValue(JsonWriter, TEXT("maxYaw"), MaxYaw);
-		FGLTFJsonUtility::WriteExactValue(JsonWriter, TEXT("minYaw"), MinYaw);
+
+		if (!FMath::IsNearlyEqual(MaxPitch, 90.0f))
+		{
+			FGLTFJsonUtility::WriteExactValue(JsonWriter, TEXT("maxPitch"), MaxPitch);
+		}
+
+		if (!FMath::IsNearlyEqual(MinPitch, -90.0f))
+		{
+			FGLTFJsonUtility::WriteExactValue(JsonWriter, TEXT("minPitch"), MinPitch);
+		}
+
+		if (!FMath::IsNearlyEqual(MaxYaw, 360.0f))
+		{
+			FGLTFJsonUtility::WriteExactValue(JsonWriter, TEXT("maxYaw"), MaxYaw);
+		}
+
+		if (!FMath::IsNearlyEqual(MinYaw, 0.0f))
+		{
+			FGLTFJsonUtility::WriteExactValue(JsonWriter, TEXT("minYaw"), MinYaw);
+		}
+
 		FGLTFJsonUtility::WriteExactValue(JsonWriter, TEXT("rotationSensitivity"), RotationSensitivity);
 		FGLTFJsonUtility::WriteExactValue(JsonWriter, TEXT("rotationInertia"), RotationInertia);
 		FGLTFJsonUtility::WriteExactValue(JsonWriter, TEXT("dollySensitivity"), DollySensitivity);
