[GLTFExporter] refactor(camera-controls): only specify non-default yaw limits if they are actually used

commit 7f02e6fd768373a9e2f675d73312ef948daf99d2
Date: Tue, 3 Nov 2020 10:25:41 +0100

    refactor(camera-controls): only specify non-default yaw limits if they are actually used
    (cherry picked from commit 60bd62aeedcc55ac940bc331c68bc439b856eeaf​)

#skipci

diff --git a/Engine/Plugins/Enterprise/GLTFExporter/MigratedFromGit/Source/GLTFExporter/Private/Converters/GLTFCameraConverters.cpp b/Engine/Plugins/Enterprise/GLTFExporter/MigratedFromGit/Source/GLTFExporter/Private/Converters/GLTFCameraConverters.cpp
--- a/Engine/Plugins/Enterprise/GLTFExporter/MigratedFromGit/Source/GLTFExporter/Private/Converters/GLTFCameraConverters.cpp
+++ b/Engine/Plugins/Enterprise/GLTFExporter/MigratedFromGit/Source/GLTFExporter/Private/Converters/GLTFCameraConverters.cpp
@@ -64,17 +64,20 @@ FGLTFJsonCameraIndex FGLTFCameraConverter::Convert(const UCameraComponent* Camer
 			CameraControl.MaxPitch = CameraActor->PitchAngleMax;
 			CameraControl.MinPitch = CameraActor->PitchAngleMin;
 
-			// Transform yaw limits to match right-handed system and glTF specification for cameras, i.e
-			// positive rotation is CCW, and camera looks down Z- (instead of X+).
-			const float MaxYaw = FMath::Max(-CameraActor->YawAngleMin, -CameraActor->YawAngleMax) - 90.0f;
-			const float MinYaw = FMath::Min(-CameraActor->YawAngleMin, -CameraActor->YawAngleMax) - 90.0f;
+			if (CameraActor->UsesYawLimits())
+			{
+				// Transform yaw limits to match right-handed system and glTF specification for cameras, i.e
+				// positive rotation is CCW, and camera looks down Z- (instead of X+).
+				const float MaxYaw = FMath::Max(-CameraActor->YawAngleMin, -CameraActor->YawAngleMax) - 90.0f;
+				const float MinYaw = FMath::Min(-CameraActor->YawAngleMin, -CameraActor->YawAngleMax) - 90.0f;
 
-			// We prefer the limits to be in the 0..360 range, but we only use MaxYaw to calculate
-			// the needed offset since we need to keep both limits a fixed distance apart from each other.
-			const float PositiveRangeOffset = FRotator::ClampAxis(MaxYaw) - MaxYaw;
+				// We prefer the limits to be in the 0..360 range, but we only use MaxYaw to calculate
+				// the needed offset since we need to keep both limits a fixed distance apart from each other.
+				const float PositiveRangeOffset = FRotator::ClampAxis(MaxYaw) - MaxYaw;
 
-			CameraControl.MaxYaw = MaxYaw + PositiveRangeOffset;
-			CameraControl.MinYaw = MinYaw + PositiveRangeOffset;
+				CameraControl.MaxYaw = MaxYaw + PositiveRangeOffset;
+				CameraControl.MinYaw = MinYaw + PositiveRangeOffset;
+			}
 
 			CameraControl.RotationSensitivity = CameraActor->RotationSensitivity;
 			CameraControl.RotationInertia = CameraActor->RotationInertia;
diff --git a/Engine/Plugins/Enterprise/GLTFExporter/MigratedFromGit/Source/GLTFExporterRuntime/Public/Actors/GLTFCameraActor.h b/Engine/Plugins/Enterprise/GLTFExporter/MigratedFromGit/Source/GLTFExporterRuntime/Public/Actors/GLTFCameraActor.h
--- a/Engine/Plugins/Enterprise/GLTFExporter/MigratedFromGit/Source/GLTFExporterRuntime/Public/Actors/GLTFCameraActor.h
+++ b/Engine/Plugins/Enterprise/GLTFExporter/MigratedFromGit/Source/GLTFExporterRuntime/Public/Actors/GLTFCameraActor.h
@@ -63,8 +63,6 @@ private:
 
 	bool SetAutoActivateForPlayer(const EAutoReceiveInput::Type Player);
 
-	bool UsesYawLimits() const;
-
 public:
 
 	/* Camera mode */
@@ -124,4 +122,8 @@ private:
 	float TargetYaw;
 	float DollyTime;
 	float DollyStartDistance;
+
+public:
+
+	bool UsesYawLimits() const;
 };
