[GLTFExporter] refactor: change parameter MeshMaterialIndices in FGLTFMaterialUtility::BakeMaterialProperty from pointer to reference

commit 9f4f5399fa518788f12326e3868f7b9f34fcddef
Date: Mon, 2 Nov 2020 16:21:14 +0100

    refactor: change parameter MeshMaterialIndices in FGLTFMaterialUtility::BakeMaterialProperty from pointer to reference
    (cherry picked from commit fc8c159f591720e4866f9ea2405ae62aecd78cfe​)

#skipci

diff --git a/Engine/Plugins/Enterprise/GLTFExporter/MigratedFromGit/Source/GLTFExporter/Private/Converters/GLTFMaterialUtility.cpp b/Engine/Plugins/Enterprise/GLTFExporter/MigratedFromGit/Source/GLTFExporter/Private/Converters/GLTFMaterialUtility.cpp
--- a/Engine/Plugins/Enterprise/GLTFExporter/MigratedFromGit/Source/GLTFExporter/Private/Converters/GLTFMaterialUtility.cpp
+++ b/Engine/Plugins/Enterprise/GLTFExporter/MigratedFromGit/Source/GLTFExporter/Private/Converters/GLTFMaterialUtility.cpp
@@ -241,7 +241,7 @@ bool FGLTFMaterialUtility::CombineTextures(TArray<FColor>& OutPixels, const TArr
 	return bReadSuccessful;
 }
 
-FGLTFPropertyBakeOutput FGLTFMaterialUtility::BakeMaterialProperty(const FIntPoint& OutputSize, EMaterialProperty Property, const UMaterialInterface* Material, int32 TexCoord, FMeshDescription* MeshDescription, TArray<int32>* MeshMaterialIndices, bool bCopyAlphaFromRedChannel)
+FGLTFPropertyBakeOutput FGLTFMaterialUtility::BakeMaterialProperty(const FIntPoint& OutputSize, EMaterialProperty Property, const UMaterialInterface* Material, int32 TexCoord, const FMeshDescription* MeshDescription, const TArray<int32>& MeshMaterialIndices, bool bCopyAlphaFromRedChannel)
 {
 	EMaterialProperty HackProperty = MP_MAX;
 	TTuple<int32, UMaterialExpression*> HackPrevState;
@@ -270,26 +270,19 @@ FGLTFPropertyBakeOutput FGLTFMaterialUtility::BakeMaterialProperty(const FIntPoi
 		Property = HackProperty;
 	}
 
-	TArray<FMeshData*> MeshSettings;
-
 	FMeshData MeshSet;
 	MeshSet.TextureCoordinateBox = { { 0.0f, 0.0f }, { 1.0f, 1.0f } };
 	MeshSet.TextureCoordinateIndex = TexCoord;
-	MeshSet.RawMeshDescription = MeshDescription;
-
-	if (MeshMaterialIndices != nullptr)
-	{
-		MeshSet.MaterialIndices = *MeshMaterialIndices;
-	}
-
-	MeshSettings.Add(&MeshSet);
-
-	TArray<FMaterialData*> MatSettings;
+	MeshSet.RawMeshDescription = const_cast<FMeshDescription*>(MeshDescription);
+	MeshSet.MaterialIndices = MeshMaterialIndices;
 
 	FMaterialData MatSet;
 	MatSet.Material = const_cast<UMaterialInterface*>(Material);
 	MatSet.PropertySizes.Add(Property, OutputSize);
 
+	TArray<FMeshData*> MeshSettings;
+	TArray<FMaterialData*> MatSettings;
+	MeshSettings.Add(&MeshSet);
 	MatSettings.Add(&MatSet);
 
 	TArray<FBakeOutput> BakeOutputs;
diff --git a/Engine/Plugins/Enterprise/GLTFExporter/MigratedFromGit/Source/GLTFExporter/Private/Converters/GLTFMaterialUtility.h b/Engine/Plugins/Enterprise/GLTFExporter/MigratedFromGit/Source/GLTFExporter/Private/Converters/GLTFMaterialUtility.h
--- a/Engine/Plugins/Enterprise/GLTFExporter/MigratedFromGit/Source/GLTFExporter/Private/Converters/GLTFMaterialUtility.h
+++ b/Engine/Plugins/Enterprise/GLTFExporter/MigratedFromGit/Source/GLTFExporter/Private/Converters/GLTFMaterialUtility.h
@@ -53,7 +53,7 @@ struct FGLTFMaterialUtility
 
 	static UTexture2D* CreateTransientTexture(const FGLTFPropertyBakeOutput& PropertyBakeOutput, bool bUseSRGB = false);
 	static bool CombineTextures(TArray<FColor>& OutPixels, const TArray<FGLTFTextureCombineSource>& Sources, const FIntPoint& OutputSize, EPixelFormat OutputPixelFormat);
-	static FGLTFPropertyBakeOutput BakeMaterialProperty(const FIntPoint& OutputSize, EMaterialProperty Property, const UMaterialInterface* Material, int32 TexCoord, FMeshDescription* MeshDescription = nullptr, TArray<int32>* MeshMaterialIndices = nullptr, bool bCopyAlphaFromRedChannel = false);
+	static FGLTFPropertyBakeOutput BakeMaterialProperty(const FIntPoint& OutputSize, EMaterialProperty Property, const UMaterialInterface* Material, int32 TexCoord, const FMeshDescription* MeshDescription = nullptr, const TArray<int32>& MeshMaterialIndices = {}, bool bCopyAlphaFromRedChannel = false);
 
 	static FGLTFJsonTextureIndex AddCombinedTexture(FGLTFConvertBuilder& Builder, const TArray<FGLTFTextureCombineSource>& CombineSources, const FIntPoint& TextureSize, const FString& TextureName, EGLTFJsonTextureFilter MinFilter, EGLTFJsonTextureFilter MagFilter, EGLTFJsonTextureWrap WrapS, EGLTFJsonTextureWrap WrapT);
 	static FGLTFJsonTextureIndex AddTexture(FGLTFConvertBuilder& Builder, const TArray<FColor>& Pixels, const FIntPoint& TextureSize, const FString& TextureName, EGLTFJsonTextureFilter MinFilter, EGLTFJsonTextureFilter MagFilter, EGLTFJsonTextureWrap WrapS, EGLTFJsonTextureWrap WrapT);
diff --git a/Engine/Plugins/Enterprise/GLTFExporter/MigratedFromGit/Source/GLTFExporter/Private/Tasks/GLTFMaterialTasks.cpp b/Engine/Plugins/Enterprise/GLTFExporter/MigratedFromGit/Source/GLTFExporter/Private/Tasks/GLTFMaterialTasks.cpp
--- a/Engine/Plugins/Enterprise/GLTFExporter/MigratedFromGit/Source/GLTFExporter/Private/Tasks/GLTFMaterialTasks.cpp
+++ b/Engine/Plugins/Enterprise/GLTFExporter/MigratedFromGit/Source/GLTFExporter/Private/Tasks/GLTFMaterialTasks.cpp
@@ -1357,7 +1357,7 @@ FGLTFPropertyBakeOutput FGLTFMaterialTask::BakeMaterialProperty(EMaterialPropert
 	const FIntPoint TextureSize = PreferredTextureSize != nullptr ? *PreferredTextureSize : DefaultTextureSize;
 
 	FMeshDescription MeshDescription;
-	TArray<int32> MaterialIndices;
+	TArray<int32> MeshMaterialIndices;
 
 	bool bUseMeshData = bHasValidMeshOrComponent;
 
@@ -1398,15 +1398,15 @@ FGLTFPropertyBakeOutput FGLTFMaterialTask::BakeMaterialProperty(EMaterialPropert
 
 			if (SectionMaterial == Material)
 			{
-				MaterialIndices.Add(Section.MaterialIndex);
+				MeshMaterialIndices.Add(Section.MaterialIndex);
 			}
 		}
 
-		if (MaterialIndices.Num() == 0)
+		if (MeshMaterialIndices.Num() == 0)
 		{
 			// Fall back to using the first material index,
 			// just like when baking from within the editor
-			MaterialIndices.Add(0);
+			MeshMaterialIndices.Add(0);
 		}
 	}
 
@@ -1418,7 +1418,7 @@ FGLTFPropertyBakeOutput FGLTFMaterialTask::BakeMaterialProperty(EMaterialPropert
 		Material,
 		OutTexCoord,
 		bUseMeshData ? &MeshDescription : nullptr,
-		bUseMeshData ? &MaterialIndices : nullptr,
+		MeshMaterialIndices,
 		bCopyAlphaFromRedChannel);
 
 	if (!PropertyBakeOutput.bIsConstant && TexCoords.Num() == 0)
