[GLTFExporter] refactor: don't create gltf mesh object if StaticMesh or SkeletalMesh is null in mesh component

commit ae7cab2795ed60ca3f94c70a647a9db993ca6028
Date: Mon, 2 Nov 2020 16:42:04 +0100

    refactor: don't create gltf mesh object if StaticMesh or SkeletalMesh is null in mesh component
    (cherry picked from commit 26bf28128e6419716ccb280de7cc63d7f6c8a1b7​)

#skipci

diff --git a/Engine/Plugins/Enterprise/GLTFExporter/MigratedFromGit/Source/GLTFExporter/Private/Builders/GLTFConvertBuilder.cpp b/Engine/Plugins/Enterprise/GLTFExporter/MigratedFromGit/Source/GLTFExporter/Private/Builders/GLTFConvertBuilder.cpp
--- a/Engine/Plugins/Enterprise/GLTFExporter/MigratedFromGit/Source/GLTFExporter/Private/Builders/GLTFConvertBuilder.cpp
+++ b/Engine/Plugins/Enterprise/GLTFExporter/MigratedFromGit/Source/GLTFExporter/Private/Builders/GLTFConvertBuilder.cpp
@@ -88,14 +88,14 @@ FGLTFJsonAccessorIndex FGLTFConvertBuilder::GetOrAddIndexAccessor(const FGLTFMes
 	return IndexBufferConverter.GetOrAdd(MeshSection);
 }
 
-FGLTFJsonMeshIndex FGLTFConvertBuilder::GetOrAddMesh(const UStaticMesh* StaticMesh, int32 LODIndex, const FGLTFMaterialArray& OverrideMaterials)
+FGLTFJsonMeshIndex FGLTFConvertBuilder::GetOrAddMesh(const UStaticMesh* StaticMesh, const FGLTFMaterialArray& OverrideMaterials)
 {
 	if (StaticMesh == nullptr)
 	{
 		return FGLTFJsonMeshIndex(INDEX_NONE);
 	}
 
-	return StaticMeshConverter.GetOrAdd(StaticMesh, nullptr, LODIndex, OverrideMaterials);
+	return StaticMeshConverter.GetOrAdd(StaticMesh, nullptr, OverrideMaterials);
 }
 
 FGLTFJsonMeshIndex FGLTFConvertBuilder::GetOrAddMesh(const UStaticMeshComponent* StaticMeshComponent, const FGLTFMaterialArray& OverrideMaterials)
@@ -105,19 +105,23 @@ FGLTFJsonMeshIndex FGLTFConvertBuilder::GetOrAddMesh(const UStaticMeshComponent*
 		return FGLTFJsonMeshIndex(INDEX_NONE);
 	}
 
-	const int32 LODIndex = StaticMeshComponent->ForcedLodModel > 0 ? StaticMeshComponent->ForcedLodModel - 1 : FMath::Max(ExportOptions->DefaultLevelOfDetail, StaticMeshComponent->MinLOD);
+	const UStaticMesh* StaticMesh = StaticMeshComponent->GetStaticMesh();
+	if (StaticMesh == nullptr)
+	{
+		return FGLTFJsonMeshIndex(INDEX_NONE);
+	}
 
-	return StaticMeshConverter.GetOrAdd(nullptr, StaticMeshComponent, LODIndex, OverrideMaterials);
+	return StaticMeshConverter.GetOrAdd(StaticMesh, StaticMeshComponent, OverrideMaterials);
 }
 
-FGLTFJsonMeshIndex FGLTFConvertBuilder::GetOrAddMesh(const USkeletalMesh* SkeletalMesh, int32 LODIndex, const FGLTFMaterialArray& OverrideMaterials)
+FGLTFJsonMeshIndex FGLTFConvertBuilder::GetOrAddMesh(const USkeletalMesh* SkeletalMesh, const FGLTFMaterialArray& OverrideMaterials)
 {
 	if (SkeletalMesh == nullptr)
 	{
 		return FGLTFJsonMeshIndex(INDEX_NONE);
 	}
 
-	return SkeletalMeshConverter.GetOrAdd(SkeletalMesh, nullptr, LODIndex, OverrideMaterials);
+	return SkeletalMeshConverter.GetOrAdd(SkeletalMesh, nullptr, OverrideMaterials);
 }
 
 FGLTFJsonMeshIndex FGLTFConvertBuilder::GetOrAddMesh(const USkeletalMeshComponent* SkeletalMeshComponent, const FGLTFMaterialArray& OverrideMaterials)
@@ -127,9 +131,13 @@ FGLTFJsonMeshIndex FGLTFConvertBuilder::GetOrAddMesh(const USkeletalMeshComponen
 		return FGLTFJsonMeshIndex(INDEX_NONE);
 	}
 
-	const int32 LODIndex = SkeletalMeshComponent->GetForcedLOD() > 0 ? SkeletalMeshComponent->GetForcedLOD() - 1 : FMath::Max(ExportOptions->DefaultLevelOfDetail, SkeletalMeshComponent->MinLodModel);
+	const USkeletalMesh* SkeletalMesh = SkeletalMeshComponent->SkeletalMesh;
+	if (SkeletalMesh == nullptr)
+	{
+		return FGLTFJsonMeshIndex(INDEX_NONE);
+	}
 
-	return SkeletalMeshConverter.GetOrAdd(nullptr, SkeletalMeshComponent, LODIndex, OverrideMaterials);
+	return SkeletalMeshConverter.GetOrAdd(SkeletalMesh, SkeletalMeshComponent, OverrideMaterials);
 }
 
 FGLTFJsonMaterialIndex FGLTFConvertBuilder::GetOrAddMaterial(const UMaterialInterface* Material)
@@ -139,37 +147,37 @@ FGLTFJsonMaterialIndex FGLTFConvertBuilder::GetOrAddMaterial(const UMaterialInte
 		return FGLTFJsonMaterialIndex(INDEX_NONE);
 	}
 
-	return MaterialConverter.GetOrAdd(Material, nullptr, -1, {});
+	return MaterialConverter.GetOrAdd(Material, nullptr, {});
 }
 
-FGLTFJsonMaterialIndex FGLTFConvertBuilder::GetOrAddMaterial(const UMaterialInterface* Material, const UStaticMesh* Mesh, int32 LODIndex, const FGLTFMaterialArray& OverrideMaterials)
+FGLTFJsonMaterialIndex FGLTFConvertBuilder::GetOrAddMaterial(const UMaterialInterface* Material, const UStaticMesh* Mesh, const FGLTFMaterialArray& OverrideMaterials)
 {
 	if (Material == nullptr)
 	{
 		return FGLTFJsonMaterialIndex(INDEX_NONE);
 	}
 
-	return MaterialConverter.GetOrAdd(Material, Mesh, LODIndex, OverrideMaterials);
+	return MaterialConverter.GetOrAdd(Material, Mesh, OverrideMaterials);
 }
 
-FGLTFJsonMaterialIndex FGLTFConvertBuilder::GetOrAddMaterial(const UMaterialInterface* Material, const USkeletalMesh* Mesh, int32 LODIndex, const FGLTFMaterialArray& OverrideMaterials)
+FGLTFJsonMaterialIndex FGLTFConvertBuilder::GetOrAddMaterial(const UMaterialInterface* Material, const USkeletalMesh* Mesh, const FGLTFMaterialArray& OverrideMaterials)
 {
 	if (Material == nullptr)
 	{
 		return FGLTFJsonMaterialIndex(INDEX_NONE);
 	}
 
-	return MaterialConverter.GetOrAdd(Material, Mesh, LODIndex, OverrideMaterials);
+	return MaterialConverter.GetOrAdd(Material, Mesh, OverrideMaterials);
 }
 
-FGLTFJsonMaterialIndex FGLTFConvertBuilder::GetOrAddMaterial(const UMaterialInterface* Material, const UMeshComponent* MeshComponent, int32 LODIndex, const FGLTFMaterialArray& OverrideMaterials)
+FGLTFJsonMaterialIndex FGLTFConvertBuilder::GetOrAddMaterial(const UMaterialInterface* Material, const UMeshComponent* MeshComponent, const FGLTFMaterialArray& OverrideMaterials)
 {
 	if (Material == nullptr)
 	{
 		return FGLTFJsonMaterialIndex(INDEX_NONE);
 	}
 
-	return MaterialConverter.GetOrAdd(Material, MeshComponent, LODIndex, OverrideMaterials);
+	return MaterialConverter.GetOrAdd(Material, MeshComponent, OverrideMaterials);
 }
 
 FGLTFJsonSamplerIndex FGLTFConvertBuilder::GetOrAddSampler(const UTexture* Texture)
diff --git a/Engine/Plugins/Enterprise/GLTFExporter/MigratedFromGit/Source/GLTFExporter/Private/Builders/GLTFConvertBuilder.h b/Engine/Plugins/Enterprise/GLTFExporter/MigratedFromGit/Source/GLTFExporter/Private/Builders/GLTFConvertBuilder.h
--- a/Engine/Plugins/Enterprise/GLTFExporter/MigratedFromGit/Source/GLTFExporter/Private/Builders/GLTFConvertBuilder.h
+++ b/Engine/Plugins/Enterprise/GLTFExporter/MigratedFromGit/Source/GLTFExporter/Private/Builders/GLTFConvertBuilder.h
@@ -39,15 +39,15 @@ public:
 	FGLTFJsonAccessorIndex GetOrAddWeightAccessor(const FGLTFMeshSection* MeshSection, const FSkinWeightVertexBuffer* VertexBuffer, int32 InfluenceOffset);
 	FGLTFJsonAccessorIndex GetOrAddIndexAccessor(const FGLTFMeshSection* MeshSection);
 
-	FGLTFJsonMeshIndex GetOrAddMesh(const UStaticMesh* StaticMesh, int32 LODIndex = -1, const FGLTFMaterialArray& OverrideMaterials = {});
+	FGLTFJsonMeshIndex GetOrAddMesh(const UStaticMesh* StaticMesh, const FGLTFMaterialArray& OverrideMaterials = {});
 	FGLTFJsonMeshIndex GetOrAddMesh(const UStaticMeshComponent* StaticMeshComponent, const FGLTFMaterialArray& OverrideMaterials = {});
-	FGLTFJsonMeshIndex GetOrAddMesh(const USkeletalMesh* SkeletalMesh, int32 LODIndex = -1, const FGLTFMaterialArray& OverrideMaterials = {});
+	FGLTFJsonMeshIndex GetOrAddMesh(const USkeletalMesh* SkeletalMesh, const FGLTFMaterialArray& OverrideMaterials = {});
 	FGLTFJsonMeshIndex GetOrAddMesh(const USkeletalMeshComponent* SkeletalMeshComponent, const FGLTFMaterialArray& OverrideMaterials = {});
 
 	FGLTFJsonMaterialIndex GetOrAddMaterial(const UMaterialInterface* Material);
-	FGLTFJsonMaterialIndex GetOrAddMaterial(const UMaterialInterface* Material, const UStaticMesh* Mesh, int32 LODIndex = -1, const FGLTFMaterialArray& OverrideMaterials = {});
-	FGLTFJsonMaterialIndex GetOrAddMaterial(const UMaterialInterface* Material, const USkeletalMesh* Mesh, int32 LODIndex = -1, const FGLTFMaterialArray& OverrideMaterials = {});
-	FGLTFJsonMaterialIndex GetOrAddMaterial(const UMaterialInterface* Material, const UMeshComponent* MeshComponent, int32 LODIndex = -1, const FGLTFMaterialArray& OverrideMaterials = {});
+	FGLTFJsonMaterialIndex GetOrAddMaterial(const UMaterialInterface* Material, const UStaticMesh* Mesh, const FGLTFMaterialArray& OverrideMaterials = {});
+	FGLTFJsonMaterialIndex GetOrAddMaterial(const UMaterialInterface* Material, const USkeletalMesh* Mesh, const FGLTFMaterialArray& OverrideMaterials = {});
+	FGLTFJsonMaterialIndex GetOrAddMaterial(const UMaterialInterface* Material, const UMeshComponent* MeshComponent, const FGLTFMaterialArray& OverrideMaterials = {});
 
 	FGLTFJsonSamplerIndex GetOrAddSampler(const UTexture* Texture);
 	FGLTFJsonTextureIndex GetOrAddTexture(const UTexture2D* Texture);
diff --git a/Engine/Plugins/Enterprise/GLTFExporter/MigratedFromGit/Source/GLTFExporter/Private/Converters/GLTFBackdropConverters.cpp b/Engine/Plugins/Enterprise/GLTFExporter/MigratedFromGit/Source/GLTFExporter/Private/Converters/GLTFBackdropConverters.cpp
--- a/Engine/Plugins/Enterprise/GLTFExporter/MigratedFromGit/Source/GLTFExporter/Private/Converters/GLTFBackdropConverters.cpp
+++ b/Engine/Plugins/Enterprise/GLTFExporter/MigratedFromGit/Source/GLTFExporter/Private/Converters/GLTFBackdropConverters.cpp
@@ -31,7 +31,7 @@ FGLTFJsonBackdropIndex FGLTFBackdropConverter::Convert(const AActor* BackdropAct
 	if (FGLTFActorUtility::TryGetPropertyValue(BackdropActor, TEXT("Mesh"), Mesh))
 	{
 		// TODO: avoid exporting the mesh material (use gltf default material)
-		JsonBackdrop.Mesh = Builder.GetOrAddMesh(Mesh, -1, { FGLTFMaterialUtility::GetDefault() });
+		JsonBackdrop.Mesh = Builder.GetOrAddMesh(Mesh, { FGLTFMaterialUtility::GetDefault() });
 	}
 	else
 	{
diff --git a/Engine/Plugins/Enterprise/GLTFExporter/MigratedFromGit/Source/GLTFExporter/Private/Converters/GLTFMaterialConverters.cpp b/Engine/Plugins/Enterprise/GLTFExporter/MigratedFromGit/Source/GLTFExporter/Private/Converters/GLTFMaterialConverters.cpp
--- a/Engine/Plugins/Enterprise/GLTFExporter/MigratedFromGit/Source/GLTFExporter/Private/Converters/GLTFMaterialConverters.cpp
+++ b/Engine/Plugins/Enterprise/GLTFExporter/MigratedFromGit/Source/GLTFExporter/Private/Converters/GLTFMaterialConverters.cpp
@@ -6,7 +6,7 @@
 #include "Builders/GLTFConvertBuilder.h"
 #include "Tasks/GLTFMaterialTasks.h"
 
-void FGLTFMaterialConverter::Sanitize(const UMaterialInterface*& Material, const UObject*& MeshOrComponent, int32& LODIndex, FGLTFMaterialArray& OverrideMaterials)
+void FGLTFMaterialConverter::Sanitize(const UMaterialInterface*& Material, const UObject*& MeshOrComponent, FGLTFMaterialArray& OverrideMaterials)
 {
 	const UStaticMesh* StaticMesh = Cast<UStaticMesh>(MeshOrComponent);
 	const USkeletalMesh* SkeletalMesh = Cast<USkeletalMesh>(MeshOrComponent);
@@ -23,38 +23,25 @@ void FGLTFMaterialConverter::Sanitize(const UMaterialInterface*& Material, const
 		}
 	}
 
-	if ((StaticMesh != nullptr || SkeletalMesh != nullptr) &&
-		Builder.ExportOptions->bBakeMaterialInputs &&
-		Builder.ExportOptions->bBakeMaterialInputsUsingMeshData &&
-		FGLTFMaterialUtility::MaterialNeedsVertexData(Material))
-	{
-		if (LODIndex < 0)
-		{
-			if (StaticMesh != nullptr)
-			{
-				LODIndex = FMath::Max(Builder.ExportOptions->DefaultLevelOfDetail, FGLTFMeshUtility::GetMinimumLOD(StaticMesh));
-			}
-			else if (SkeletalMesh != nullptr)
-			{
-				LODIndex = FMath::Max(Builder.ExportOptions->DefaultLevelOfDetail, FGLTFMeshUtility::GetMinimumLOD(SkeletalMesh));
-			}
-		}
-	}
-	else
+	if (StaticMesh == nullptr && SkeletalMesh == nullptr ||
+		!Builder.ExportOptions->bBakeMaterialInputs ||
+		!Builder.ExportOptions->bBakeMaterialInputsUsingMeshData ||
+		!FGLTFMaterialUtility::MaterialNeedsVertexData(Material))
 	{
 		MeshOrComponent = nullptr;
-		LODIndex = -1;
 		OverrideMaterials = {};
 	}
 }
 
-FGLTFJsonMaterialIndex FGLTFMaterialConverter::Convert(const UMaterialInterface* Material, const UObject* MeshOrComponent, int32 LODIndex, const FGLTFMaterialArray OverrideMaterials)
+FGLTFJsonMaterialIndex FGLTFMaterialConverter::Convert(const UMaterialInterface* Material, const UObject* MeshOrComponent, const FGLTFMaterialArray OverrideMaterials)
 {
 	if (Material == FGLTFMaterialUtility::GetDefault())
 	{
 		return FGLTFJsonMaterialIndex(INDEX_NONE); // use default gltf definition
 	}
 
+	const int32 LODIndex = FGLTFMeshUtility::GetLOD(MeshOrComponent, Builder.ExportOptions->DefaultLevelOfDetail);
+
 	const FGLTFJsonMaterialIndex MaterialIndex = Builder.AddMaterial();
 	Builder.SetupTask<FGLTFMaterialTask>(Builder, Material, MeshOrComponent, LODIndex, OverrideMaterials, MaterialIndex);
 	return MaterialIndex;
diff --git a/Engine/Plugins/Enterprise/GLTFExporter/MigratedFromGit/Source/GLTFExporter/Private/Converters/GLTFMaterialConverters.h b/Engine/Plugins/Enterprise/GLTFExporter/MigratedFromGit/Source/GLTFExporter/Private/Converters/GLTFMaterialConverters.h
--- a/Engine/Plugins/Enterprise/GLTFExporter/MigratedFromGit/Source/GLTFExporter/Private/Converters/GLTFMaterialConverters.h
+++ b/Engine/Plugins/Enterprise/GLTFExporter/MigratedFromGit/Source/GLTFExporter/Private/Converters/GLTFMaterialConverters.h
@@ -8,11 +8,11 @@
 #include "Converters/GLTFMaterialArray.h"
 #include "Engine.h"
 
-class FGLTFMaterialConverter final : public FGLTFBuilderContext, public TGLTFConverter<FGLTFJsonMaterialIndex, const UMaterialInterface*, const UObject*, int32, FGLTFMaterialArray>
+class FGLTFMaterialConverter final : public FGLTFBuilderContext, public TGLTFConverter<FGLTFJsonMaterialIndex, const UMaterialInterface*, const UObject*, FGLTFMaterialArray>
 {
 	using FGLTFBuilderContext::FGLTFBuilderContext;
 
-	virtual void Sanitize(const UMaterialInterface*& Material, const UObject*& MeshOrComponent, int32& LODIndex, FGLTFMaterialArray& OverrideMaterials) override;
+	virtual void Sanitize(const UMaterialInterface*& Material, const UObject*& MeshOrComponent, FGLTFMaterialArray& OverrideMaterials) override;
 
-	virtual FGLTFJsonMaterialIndex Convert(const UMaterialInterface* Material, const UObject* MeshOrComponent, int32 LODIndex, const FGLTFMaterialArray OverrideMaterials) override;
+	virtual FGLTFJsonMaterialIndex Convert(const UMaterialInterface* Material, const UObject* MeshOrComponent, const FGLTFMaterialArray OverrideMaterials) override;
 };
diff --git a/Engine/Plugins/Enterprise/GLTFExporter/MigratedFromGit/Source/GLTFExporter/Private/Converters/GLTFMeshConverters.cpp b/Engine/Plugins/Enterprise/GLTFExporter/MigratedFromGit/Source/GLTFExporter/Private/Converters/GLTFMeshConverters.cpp
--- a/Engine/Plugins/Enterprise/GLTFExporter/MigratedFromGit/Source/GLTFExporter/Private/Converters/GLTFMeshConverters.cpp
+++ b/Engine/Plugins/Enterprise/GLTFExporter/MigratedFromGit/Source/GLTFExporter/Private/Converters/GLTFMeshConverters.cpp
@@ -2,16 +2,11 @@
 
 #include "Converters/GLTFMeshConverters.h"
 #include "Converters/GLTFMaterialUtility.h"
-#include "Converters/GLTFMeshUtility.h"
 #include "Builders/GLTFConvertBuilder.h"
 #include "Tasks/GLTFMeshTasks.h"
-#include "Rendering/SkeletalMeshRenderData.h"
 
-void FGLTFStaticMeshConverter::Sanitize(const UStaticMesh*& StaticMesh, const UStaticMeshComponent*& StaticMeshComponent, int32& LODIndex, FGLTFMaterialArray& OverrideMaterials)
+void FGLTFStaticMeshConverter::Sanitize(const UStaticMesh*& StaticMesh, const UStaticMeshComponent*& StaticMeshComponent, FGLTFMaterialArray& OverrideMaterials)
 {
-	const uint32 MeshCount = (StaticMeshComponent != nullptr ? 1 : 0) + (StaticMesh != nullptr ? 1 : 0);
-	check(MeshCount == 1);
-
 	const TArray<UMaterialInterface*> Materials = StaticMeshComponent != nullptr
 		? OverrideMaterials.GetOverrides(StaticMeshComponent->GetMaterials())
 		: OverrideMaterials.GetOverrides(StaticMesh->StaticMaterials);
@@ -24,7 +19,6 @@ void FGLTFStaticMeshConverter::Sanitize(const UStaticMesh*& StaticMesh, const US
 			!Builder.ExportOptions->bBakeMaterialInputsUsingMeshData ||
 			!FGLTFMaterialUtility::MaterialsNeedVertexData(Materials))
 		{
-			StaticMesh = StaticMeshComponent->GetStaticMesh();
 			StaticMeshComponent = nullptr;
 		}
 	}
@@ -40,29 +34,17 @@ void FGLTFStaticMeshConverter::Sanitize(const UStaticMesh*& StaticMesh, const US
 	{
 		OverrideMaterials.ClearRedundantOverrides(StaticMesh->StaticMaterials);
 	}
-
-	const UStaticMesh* Mesh = StaticMeshComponent != nullptr ? StaticMeshComponent->GetStaticMesh() : StaticMesh;
-
-	if (LODIndex < 0)
-	{
-		LODIndex = FMath::Max(Builder.ExportOptions->DefaultLevelOfDetail, FGLTFMeshUtility::GetMinimumLOD(Mesh));
-	}
-
-	LODIndex = FMath::Clamp(LODIndex, 0, Mesh->GetNumLODs() - 1);
 }
 
-FGLTFJsonMeshIndex FGLTFStaticMeshConverter::Convert(const UStaticMesh* StaticMesh, const UStaticMeshComponent* StaticMeshComponent, int32 LODIndex, FGLTFMaterialArray OverrideMaterials)
+FGLTFJsonMeshIndex FGLTFStaticMeshConverter::Convert(const UStaticMesh* StaticMesh, const UStaticMeshComponent* StaticMeshComponent, FGLTFMaterialArray OverrideMaterials)
 {
 	const FGLTFJsonMeshIndex MeshIndex = Builder.AddMesh();
-	Builder.SetupTask<FGLTFStaticMeshTask>(Builder, MeshSectionConverter, StaticMesh, StaticMeshComponent, LODIndex, OverrideMaterials, MeshIndex);
+	Builder.SetupTask<FGLTFStaticMeshTask>(Builder, MeshSectionConverter, StaticMesh, StaticMeshComponent, OverrideMaterials, MeshIndex);
 	return MeshIndex;
 }
 
-void FGLTFSkeletalMeshConverter::Sanitize(const USkeletalMesh*& SkeletalMesh, const USkeletalMeshComponent*& SkeletalMeshComponent, int32& LODIndex, FGLTFMaterialArray& OverrideMaterials)
+void FGLTFSkeletalMeshConverter::Sanitize(const USkeletalMesh*& SkeletalMesh, const USkeletalMeshComponent*& SkeletalMeshComponent, FGLTFMaterialArray& OverrideMaterials)
 {
-	const uint32 MeshCount = (SkeletalMeshComponent != nullptr ? 1 : 0) + (SkeletalMesh != nullptr ? 1 : 0);
-	check(MeshCount == 1);
-
 	const TArray<UMaterialInterface*> Materials = SkeletalMeshComponent != nullptr
 		? OverrideMaterials.GetOverrides(SkeletalMeshComponent->GetMaterials())
 		: OverrideMaterials.GetOverrides(SkeletalMesh->Materials);
@@ -75,7 +57,6 @@ void FGLTFSkeletalMeshConverter::Sanitize(const USkeletalMesh*& SkeletalMesh, co
 			!Builder.ExportOptions->bBakeMaterialInputsUsingMeshData ||
 			!FGLTFMaterialUtility::MaterialsNeedVertexData(Materials))
 		{
-			SkeletalMesh = SkeletalMeshComponent->SkeletalMesh;
 			SkeletalMeshComponent = nullptr;
 		}
 	}
@@ -91,21 +72,11 @@ void FGLTFSkeletalMeshConverter::Sanitize(const USkeletalMesh*& SkeletalMesh, co
 	{
 		OverrideMaterials.ClearRedundantOverrides(SkeletalMesh->Materials);
 	}
-
-	const USkeletalMesh* Mesh = SkeletalMeshComponent != nullptr ? SkeletalMeshComponent->SkeletalMesh : SkeletalMesh;
-
-	if (LODIndex < 0)
-	{
-		LODIndex = FMath::Max(Builder.ExportOptions->DefaultLevelOfDetail, FGLTFMeshUtility::GetMinimumLOD(Mesh));
-	}
-
-	const FSkeletalMeshRenderData* RenderData = Mesh->GetResourceForRendering();
-	LODIndex = FMath::Clamp(LODIndex, 0, RenderData->LODRenderData.Num() - 1);
 }
 
-FGLTFJsonMeshIndex FGLTFSkeletalMeshConverter::Convert(const USkeletalMesh* SkeletalMesh, const USkeletalMeshComponent* SkeletalMeshComponent, int32 LODIndex, FGLTFMaterialArray OverrideMaterials)
+FGLTFJsonMeshIndex FGLTFSkeletalMeshConverter::Convert(const USkeletalMesh* SkeletalMesh, const USkeletalMeshComponent* SkeletalMeshComponent, FGLTFMaterialArray OverrideMaterials)
 {
 	const FGLTFJsonMeshIndex MeshIndex = Builder.AddMesh();
-	Builder.SetupTask<FGLTFSkeletalMeshTask>(Builder, MeshSectionConverter, SkeletalMesh, SkeletalMeshComponent, LODIndex, OverrideMaterials, MeshIndex);
+	Builder.SetupTask<FGLTFSkeletalMeshTask>(Builder, MeshSectionConverter, SkeletalMesh, SkeletalMeshComponent, OverrideMaterials, MeshIndex);
 	return MeshIndex;
 }
diff --git a/Engine/Plugins/Enterprise/GLTFExporter/MigratedFromGit/Source/GLTFExporter/Private/Converters/GLTFMeshConverters.h b/Engine/Plugins/Enterprise/GLTFExporter/MigratedFromGit/Source/GLTFExporter/Private/Converters/GLTFMeshConverters.h
--- a/Engine/Plugins/Enterprise/GLTFExporter/MigratedFromGit/Source/GLTFExporter/Private/Converters/GLTFMeshConverters.h
+++ b/Engine/Plugins/Enterprise/GLTFExporter/MigratedFromGit/Source/GLTFExporter/Private/Converters/GLTFMeshConverters.h
@@ -15,24 +15,24 @@ class TGLTFMeshConverter : public FGLTFBuilderContext, public TGLTFConverter<FGL
 	using FGLTFBuilderContext::FGLTFBuilderContext;
 };
 
-class FGLTFStaticMeshConverter final : public TGLTFMeshConverter<const UStaticMesh*, const UStaticMeshComponent*, int32, FGLTFMaterialArray>
+class FGLTFStaticMeshConverter final : public TGLTFMeshConverter<const UStaticMesh*, const UStaticMeshComponent*, FGLTFMaterialArray>
 {
 	using TGLTFMeshConverter::TGLTFMeshConverter;
 
-	virtual void Sanitize(const UStaticMesh*& StaticMesh, const UStaticMeshComponent*& StaticMeshComponent, int32& LODIndex, FGLTFMaterialArray& OverrideMaterials) override;
+	virtual void Sanitize(const UStaticMesh*& StaticMesh, const UStaticMeshComponent*& StaticMeshComponent, FGLTFMaterialArray& OverrideMaterials) override;
 
-	virtual FGLTFJsonMeshIndex Convert(const UStaticMesh* StaticMesh, const UStaticMeshComponent* StaticMeshComponent, int32 LODIndex, FGLTFMaterialArray OverrideMaterials) override;
+	virtual FGLTFJsonMeshIndex Convert(const UStaticMesh* StaticMesh, const UStaticMeshComponent* StaticMeshComponent, FGLTFMaterialArray OverrideMaterials) override;
 
 	FGLTFStaticMeshSectionConverter MeshSectionConverter;
 };
 
-class FGLTFSkeletalMeshConverter final : public TGLTFMeshConverter<const USkeletalMesh*, const USkeletalMeshComponent*, int32, FGLTFMaterialArray>
+class FGLTFSkeletalMeshConverter final : public TGLTFMeshConverter<const USkeletalMesh*, const USkeletalMeshComponent*, FGLTFMaterialArray>
 {
 	using TGLTFMeshConverter::TGLTFMeshConverter;
 
-	virtual void Sanitize(const USkeletalMesh*& SkeletalMesh, const USkeletalMeshComponent*& SkeletalMeshComponent, int32& LODIndex, FGLTFMaterialArray& OverrideMaterials) override;
+	virtual void Sanitize(const USkeletalMesh*& SkeletalMesh, const USkeletalMeshComponent*& SkeletalMeshComponent, FGLTFMaterialArray& OverrideMaterials) override;
 
-	virtual FGLTFJsonMeshIndex Convert(const USkeletalMesh* SkeletalMesh, const USkeletalMeshComponent* SkeletalMeshComponent, int32 LODIndex, FGLTFMaterialArray OverrideMaterials) override;
+	virtual FGLTFJsonMeshIndex Convert(const USkeletalMesh* SkeletalMesh, const USkeletalMeshComponent* SkeletalMeshComponent, FGLTFMaterialArray OverrideMaterials) override;
 
 	FGLTFSkeletalMeshSectionConverter MeshSectionConverter;
 };
diff --git a/Engine/Plugins/Enterprise/GLTFExporter/MigratedFromGit/Source/GLTFExporter/Private/Converters/GLTFMeshUtility.cpp b/Engine/Plugins/Enterprise/GLTFExporter/MigratedFromGit/Source/GLTFExporter/Private/Converters/GLTFMeshUtility.cpp
--- a/Engine/Plugins/Enterprise/GLTFExporter/MigratedFromGit/Source/GLTFExporter/Private/Converters/GLTFMeshUtility.cpp
+++ b/Engine/Plugins/Enterprise/GLTFExporter/MigratedFromGit/Source/GLTFExporter/Private/Converters/GLTFMeshUtility.cpp
@@ -1,16 +1,95 @@
 // Copyright Epic Games, Inc. All Rights Reserved.
 
 #include "Converters/GLTFMeshUtility.h"
+#include "Rendering/SkeletalMeshRenderData.h"
 #include "PlatformInfo.h"
 
-int32 FGLTFMeshUtility::GetMinimumLOD(const UStaticMesh* StaticMesh)
+int32 FGLTFMeshUtility::GetLOD(const UObject* MeshOrComponent, int32 DefaultLOD)
 {
-	return GetValueForRunningPlatform<int32>(StaticMesh->MinLOD);
+	if (const UStaticMesh* StaticMesh = Cast<UStaticMesh>(MeshOrComponent))
+	{
+		return GetLOD(StaticMesh, nullptr, DefaultLOD);
+	}
+
+	if (const USkeletalMesh* SkeletalMesh = Cast<USkeletalMesh>(MeshOrComponent))
+	{
+		return GetLOD(SkeletalMesh, nullptr, DefaultLOD);
+	}
+
+	if (const UStaticMeshComponent* StaticMeshComponent = Cast<UStaticMeshComponent>(MeshOrComponent))
+	{
+		return GetLOD(StaticMeshComponent->GetStaticMesh(), StaticMeshComponent, DefaultLOD);
+	}
+
+	if (const USkeletalMeshComponent* SkeletalMeshComponent = Cast<USkeletalMeshComponent>(MeshOrComponent))
+	{
+		return GetLOD(SkeletalMeshComponent->SkeletalMesh, SkeletalMeshComponent, DefaultLOD);
+	}
+
+	return -1;
+}
+
+int32 FGLTFMeshUtility::GetLOD(const UStaticMesh* StaticMesh, const UStaticMeshComponent* StaticMeshComponent, int32 DefaultLOD)
+{
+	const int32 ForcedLOD = StaticMeshComponent != nullptr ? StaticMeshComponent->ForcedLodModel - 1 : -1;
+	const int32 LOD = ForcedLOD > 0 ? ForcedLOD : FMath::Max(DefaultLOD, GetMinimumLOD(StaticMesh, StaticMeshComponent));
+	return FMath::Min(LOD, GetMaximumLOD(StaticMesh));
+}
+
+int32 FGLTFMeshUtility::GetLOD(const USkeletalMesh* SkeletalMesh, const USkeletalMeshComponent* SkeletalMeshComponent, int32 DefaultLOD)
+{
+	const int32 ForcedLOD = SkeletalMeshComponent != nullptr ? SkeletalMeshComponent->GetForcedLOD() - 1 : -1;
+	const int32 LOD = ForcedLOD > 0 ? ForcedLOD : FMath::Max(DefaultLOD, GetMinimumLOD(SkeletalMesh, SkeletalMeshComponent));
+	return FMath::Min(LOD, GetMaximumLOD(SkeletalMesh));
+}
+
+int32 FGLTFMeshUtility::GetMaximumLOD(const UStaticMesh* StaticMesh)
+{
+	return StaticMesh != nullptr ? StaticMesh->GetNumLODs() - 1 : -1;
+}
+
+int32 FGLTFMeshUtility::GetMaximumLOD(const USkeletalMesh* SkeletalMesh)
+{
+	if (SkeletalMesh != nullptr)
+	{
+		const FSkeletalMeshRenderData* RenderData = SkeletalMesh->GetResourceForRendering();
+		if (RenderData != nullptr)
+		{
+			return RenderData->LODRenderData.Num() - 1;
+		}
+	}
+
+	return -1;
+}
+
+int32 FGLTFMeshUtility::GetMinimumLOD(const UStaticMesh* StaticMesh, const UStaticMeshComponent* StaticMeshComponent)
+{
+	if (StaticMeshComponent != nullptr && StaticMeshComponent->bOverrideMinLOD)
+	{
+		return StaticMeshComponent->MinLOD;
+	}
+
+	if (StaticMesh != nullptr)
+	{
+		return GetValueForRunningPlatform<int32>(StaticMesh->MinLOD);
+	}
+
+	return -1;
 }
 
-int32 FGLTFMeshUtility::GetMinimumLOD(const USkeletalMesh* SkeletalMesh)
+int32 FGLTFMeshUtility::GetMinimumLOD(const USkeletalMesh* SkeletalMesh, const USkeletalMeshComponent* SkeletalMeshComponent)
 {
-	return GetValueForRunningPlatform<int32>(SkeletalMesh->MinLod);
+	if (SkeletalMeshComponent != nullptr && SkeletalMeshComponent->bOverrideMinLod)
+	{
+		return SkeletalMeshComponent->MinLodModel;
+	}
+
+	if (SkeletalMesh != nullptr)
+	{
+		return GetValueForRunningPlatform<int32>(SkeletalMesh->MinLod);
+	}
+
+	return -1;
 }
 
 template <typename ValueType, typename StructType>
diff --git a/Engine/Plugins/Enterprise/GLTFExporter/MigratedFromGit/Source/GLTFExporter/Private/Converters/GLTFMeshUtility.h b/Engine/Plugins/Enterprise/GLTFExporter/MigratedFromGit/Source/GLTFExporter/Private/Converters/GLTFMeshUtility.h
--- a/Engine/Plugins/Enterprise/GLTFExporter/MigratedFromGit/Source/GLTFExporter/Private/Converters/GLTFMeshUtility.h
+++ b/Engine/Plugins/Enterprise/GLTFExporter/MigratedFromGit/Source/GLTFExporter/Private/Converters/GLTFMeshUtility.h
@@ -3,13 +3,19 @@
 #pragma once
 
 #include "Engine.h"
-#include "GLTFMaterialArray.h"
 
 struct FGLTFMeshUtility
 {
-	static int32 GetMinimumLOD(const UStaticMesh* StaticMesh);
+	static int32 GetLOD(const UObject* MeshOrComponent, int32 DefaultLOD);
 
-	static int32 GetMinimumLOD(const USkeletalMesh* SkeletalMesh);
+	static int32 GetLOD(const UStaticMesh* StaticMesh, const UStaticMeshComponent* StaticMeshComponent, int32 DefaultLOD);
+	static int32 GetLOD(const USkeletalMesh* SkeletalMesh, const USkeletalMeshComponent* SkeletalMeshComponent, int32 DefaultLOD);
+
+	static int32 GetMaximumLOD(const UStaticMesh* StaticMesh);
+	static int32 GetMaximumLOD(const USkeletalMesh* SkeletalMesh);
+
+	static int32 GetMinimumLOD(const UStaticMesh* StaticMesh, const UStaticMeshComponent* StaticMeshComponent);
+	static int32 GetMinimumLOD(const USkeletalMesh* SkeletalMesh, const USkeletalMeshComponent* SkeletalMeshComponent);
 
 private:
 
diff --git a/Engine/Plugins/Enterprise/GLTFExporter/MigratedFromGit/Source/GLTFExporter/Private/Converters/GLTFVarationConverters.cpp b/Engine/Plugins/Enterprise/GLTFExporter/MigratedFromGit/Source/GLTFExporter/Private/Converters/GLTFVarationConverters.cpp
--- a/Engine/Plugins/Enterprise/GLTFExporter/MigratedFromGit/Source/GLTFExporter/Private/Converters/GLTFVarationConverters.cpp
+++ b/Engine/Plugins/Enterprise/GLTFExporter/MigratedFromGit/Source/GLTFExporter/Private/Converters/GLTFVarationConverters.cpp
@@ -333,7 +333,7 @@ bool FGLTFVariationConverter::TryParseStaticMeshPropertyValue(FGLTFJsonVariant&
 	const FGLTFMaterialArray OverrideMaterials(MeshComponent->OverrideMaterials);
 	const FGLTFJsonNodeIndex NodeIndex = Builder.GetOrAddNode(Target);
 	const FGLTFJsonNodeIndex ComponentNodeIndex = Builder.GetComponentNodeIndex(NodeIndex);
-	const FGLTFJsonMeshIndex MeshIndex = Builder.GetOrAddMesh(StaticMesh, -1, OverrideMaterials);
+	const FGLTFJsonMeshIndex MeshIndex = Builder.GetOrAddMesh(StaticMesh, OverrideMaterials);
 	FGLTFJsonVariantNodeProperties& NodeProperties = OutVariant.Nodes.FindOrAdd(ComponentNodeIndex);
 
 	NodeProperties.Node = ComponentNodeIndex;
@@ -389,7 +389,7 @@ bool FGLTFVariationConverter::TryParseSkeletalMeshPropertyValue(FGLTFJsonVariant
 
 	const FGLTFMaterialArray OverrideMaterials(MeshComponent->OverrideMaterials);
 	const FGLTFJsonNodeIndex NodeIndex = Builder.GetOrAddNode(Target);
-	const FGLTFJsonMeshIndex MeshIndex = Builder.GetOrAddMesh(SkeletalMesh, -1, OverrideMaterials);
+	const FGLTFJsonMeshIndex MeshIndex = Builder.GetOrAddMesh(SkeletalMesh, OverrideMaterials);
 	const FGLTFJsonNodeIndex ComponentNodeIndex = Builder.GetComponentNodeIndex(NodeIndex);
 	FGLTFJsonVariantNodeProperties& NodeProperties = OutVariant.Nodes.FindOrAdd(ComponentNodeIndex);
 
diff --git a/Engine/Plugins/Enterprise/GLTFExporter/MigratedFromGit/Source/GLTFExporter/Private/Exporters/GLTFMaterialExporter.cpp b/Engine/Plugins/Enterprise/GLTFExporter/MigratedFromGit/Source/GLTFExporter/Private/Exporters/GLTFMaterialExporter.cpp
--- a/Engine/Plugins/Enterprise/GLTFExporter/MigratedFromGit/Source/GLTFExporter/Private/Exporters/GLTFMaterialExporter.cpp
+++ b/Engine/Plugins/Enterprise/GLTFExporter/MigratedFromGit/Source/GLTFExporter/Private/Exporters/GLTFMaterialExporter.cpp
@@ -29,7 +29,7 @@ bool UGLTFMaterialExporter::AddObject(FGLTFContainerBuilder& Builder, const UObj
 	{
 		if (PreviewMesh != nullptr)
 		{
-			const FGLTFJsonMeshIndex MeshIndex = Builder.GetOrAddMesh(PreviewMesh, -1, { Material });
+			const FGLTFJsonMeshIndex MeshIndex = Builder.GetOrAddMesh(PreviewMesh, { Material });
 			if (MeshIndex == INDEX_NONE)
 			{
 				Builder.AddErrorMessage(
diff --git a/Engine/Plugins/Enterprise/GLTFExporter/MigratedFromGit/Source/GLTFExporter/Private/Tasks/GLTFMeshTasks.cpp b/Engine/Plugins/Enterprise/GLTFExporter/MigratedFromGit/Source/GLTFExporter/Private/Tasks/GLTFMeshTasks.cpp
--- a/Engine/Plugins/Enterprise/GLTFExporter/MigratedFromGit/Source/GLTFExporter/Private/Tasks/GLTFMeshTasks.cpp
+++ b/Engine/Plugins/Enterprise/GLTFExporter/MigratedFromGit/Source/GLTFExporter/Private/Tasks/GLTFMeshTasks.cpp
@@ -2,25 +2,18 @@
 
 #include "Tasks/GLTFMeshTasks.h"
 #include "Converters/GLTFConverterUtility.h"
-#include "Converters/GLTFNameUtility.h"
+#include "Converters/GLTFMeshUtility.h"
 #include "Builders/GLTFConvertBuilder.h"
 #include "Rendering/SkeletalMeshRenderData.h"
 
 void FGLTFStaticMeshTask::Complete()
 {
-	const uint32 MeshCount = (StaticMeshComponent != nullptr ? 1 : 0) + (StaticMesh != nullptr ? 1 : 0);
-	check(MeshCount == 1);
-
-	if (StaticMeshComponent != nullptr)
-	{
-		StaticMesh = StaticMeshComponent->GetStaticMesh();
-	}
+	FGLTFJsonMesh& JsonMesh = Builder.GetMesh(MeshIndex);
+	JsonMesh.Name = StaticMeshComponent != nullptr ? FGLTFNameUtility::GetName(StaticMeshComponent) : StaticMesh->GetName();
 
+	const int32 LODIndex = FGLTFMeshUtility::GetLOD(StaticMesh, StaticMeshComponent, Builder.ExportOptions->DefaultLevelOfDetail);
 	const FStaticMeshLODResources& MeshLOD = StaticMesh->GetLODForExport(LODIndex);
 
-	FGLTFJsonMesh& JsonMesh = Builder.GetMesh(MeshIndex);
-	JsonMesh.Name = StaticMeshComponent != nullptr ? FGLTFNameUtility::GetName(StaticMeshComponent) : FGLTFNameUtility::GetName(StaticMesh, LODIndex);
-
 	const FPositionVertexBuffer* PositionBuffer = &MeshLOD.VertexBuffers.PositionVertexBuffer;
 	const FStaticMeshVertexBuffer* VertexBuffer = &MeshLOD.VertexBuffers.StaticMeshVertexBuffer;
 	const FColorVertexBuffer* ColorBuffer = &MeshLOD.VertexBuffers.ColorVertexBuffer;
@@ -73,27 +66,20 @@ void FGLTFStaticMeshTask::Complete()
 		}
 
 		JsonPrimitive.Material = StaticMeshComponent != nullptr
-			? Builder.GetOrAddMaterial(Material, StaticMeshComponent, LODIndex, OverrideMaterials)
-			: Builder.GetOrAddMaterial(Material, StaticMesh, LODIndex, OverrideMaterials);
+			? Builder.GetOrAddMaterial(Material, StaticMeshComponent, OverrideMaterials)
+			: Builder.GetOrAddMaterial(Material, StaticMesh, OverrideMaterials);
 	}
 }
 
 void FGLTFSkeletalMeshTask::Complete()
 {
-	const uint32 MeshCount = (SkeletalMeshComponent != nullptr ? 1 : 0) + (SkeletalMesh != nullptr ? 1 : 0);
-	check(MeshCount == 1);
-
-	if (SkeletalMeshComponent != nullptr)
-	{
-		SkeletalMesh = SkeletalMeshComponent->SkeletalMesh;
-	}
+	FGLTFJsonMesh& JsonMesh = Builder.GetMesh(MeshIndex);
+	JsonMesh.Name = SkeletalMeshComponent != nullptr ? FGLTFNameUtility::GetName(SkeletalMeshComponent) : SkeletalMesh->GetName();
 
+	const int32 LODIndex = FGLTFMeshUtility::GetLOD(SkeletalMesh, SkeletalMeshComponent, Builder.ExportOptions->DefaultLevelOfDetail);
 	const FSkeletalMeshRenderData* RenderData = SkeletalMesh->GetResourceForRendering();
 	const FSkeletalMeshLODRenderData& MeshLOD = RenderData->LODRenderData[LODIndex];
 
-	FGLTFJsonMesh& JsonMesh = Builder.GetMesh(MeshIndex);
-	JsonMesh.Name = SkeletalMeshComponent != nullptr ? FGLTFNameUtility::GetName(SkeletalMeshComponent) : FGLTFNameUtility::GetName(SkeletalMesh, LODIndex);
-
 	const FPositionVertexBuffer* PositionBuffer = &MeshLOD.StaticVertexBuffers.PositionVertexBuffer;
 	const FStaticMeshVertexBuffer* VertexBuffer = &MeshLOD.StaticVertexBuffers.StaticMeshVertexBuffer;
 	const FColorVertexBuffer* ColorBuffer = &MeshLOD.StaticVertexBuffers.ColorVertexBuffer;
@@ -161,7 +147,7 @@ void FGLTFSkeletalMeshTask::Complete()
 		}
 
 		JsonPrimitive.Material = SkeletalMeshComponent != nullptr
-			? Builder.GetOrAddMaterial(Material, SkeletalMeshComponent, LODIndex, OverrideMaterials)
-			: Builder.GetOrAddMaterial(Material, SkeletalMesh, LODIndex, OverrideMaterials);
+			? Builder.GetOrAddMaterial(Material, SkeletalMeshComponent, OverrideMaterials)
+			: Builder.GetOrAddMaterial(Material, SkeletalMesh, OverrideMaterials);
 	}
 }
diff --git a/Engine/Plugins/Enterprise/GLTFExporter/MigratedFromGit/Source/GLTFExporter/Private/Tasks/GLTFMeshTasks.h b/Engine/Plugins/Enterprise/GLTFExporter/MigratedFromGit/Source/GLTFExporter/Private/Tasks/GLTFMeshTasks.h
--- a/Engine/Plugins/Enterprise/GLTFExporter/MigratedFromGit/Source/GLTFExporter/Private/Tasks/GLTFMeshTasks.h
+++ b/Engine/Plugins/Enterprise/GLTFExporter/MigratedFromGit/Source/GLTFExporter/Private/Tasks/GLTFMeshTasks.h
@@ -13,13 +13,12 @@ class FGLTFStaticMeshTask : public FGLTFTask
 {
 public:
 
-	FGLTFStaticMeshTask(FGLTFConvertBuilder& Builder, FGLTFStaticMeshSectionConverter& MeshSectionConverter, const UStaticMesh* StaticMesh, const UStaticMeshComponent* StaticMeshComponent, int32 LODIndex, FGLTFMaterialArray OverrideMaterials, FGLTFJsonMeshIndex MeshIndex)
+	FGLTFStaticMeshTask(FGLTFConvertBuilder& Builder, FGLTFStaticMeshSectionConverter& MeshSectionConverter, const UStaticMesh* StaticMesh, const UStaticMeshComponent* StaticMeshComponent, FGLTFMaterialArray OverrideMaterials, FGLTFJsonMeshIndex MeshIndex)
 		: FGLTFTask(EGLTFTaskPriority::Mesh)
 		, Builder(Builder)
 		, MeshSectionConverter(MeshSectionConverter)
 		, StaticMesh(StaticMesh)
 		, StaticMeshComponent(StaticMeshComponent)
-		, LODIndex(LODIndex)
 		, OverrideMaterials(OverrideMaterials)
         , MeshIndex(MeshIndex)
 	{
@@ -27,7 +26,7 @@ public:
 
 	virtual FString GetName() override
 	{
-		return StaticMeshComponent != nullptr ? FGLTFNameUtility::GetName(StaticMeshComponent) : FGLTFNameUtility::GetName(StaticMesh, LODIndex);
+		return StaticMeshComponent != nullptr ? FGLTFNameUtility::GetName(StaticMeshComponent) : StaticMesh->GetName();
 	}
 
 	virtual void Complete() override;
@@ -38,7 +37,6 @@ private:
 	FGLTFStaticMeshSectionConverter& MeshSectionConverter;
 	const UStaticMesh* StaticMesh;
 	const UStaticMeshComponent* StaticMeshComponent;
-	int32 LODIndex;
 	FGLTFMaterialArray OverrideMaterials;
 	const FGLTFJsonMeshIndex MeshIndex;
 };
@@ -47,13 +45,12 @@ class FGLTFSkeletalMeshTask : public FGLTFTask
 {
 public:
 
-	FGLTFSkeletalMeshTask(FGLTFConvertBuilder& Builder, FGLTFSkeletalMeshSectionConverter& MeshSectionConverter, const USkeletalMesh* SkeletalMesh, const USkeletalMeshComponent* SkeletalMeshComponent, int32 LODIndex, FGLTFMaterialArray OverrideMaterials, FGLTFJsonMeshIndex MeshIndex)
+	FGLTFSkeletalMeshTask(FGLTFConvertBuilder& Builder, FGLTFSkeletalMeshSectionConverter& MeshSectionConverter, const USkeletalMesh* SkeletalMesh, const USkeletalMeshComponent* SkeletalMeshComponent, FGLTFMaterialArray OverrideMaterials, FGLTFJsonMeshIndex MeshIndex)
 		: FGLTFTask(EGLTFTaskPriority::Mesh)
 		, Builder(Builder)
 		, MeshSectionConverter(MeshSectionConverter)
 		, SkeletalMesh(SkeletalMesh)
 		, SkeletalMeshComponent(SkeletalMeshComponent)
-		, LODIndex(LODIndex)
 		, OverrideMaterials(OverrideMaterials)
 		, MeshIndex(MeshIndex)
 	{
@@ -61,7 +58,7 @@ public:
 
 	virtual FString GetName() override
 	{
-		return SkeletalMeshComponent != nullptr ? FGLTFNameUtility::GetName(SkeletalMeshComponent) : FGLTFNameUtility::GetName(SkeletalMesh, LODIndex);
+		return SkeletalMeshComponent != nullptr ? FGLTFNameUtility::GetName(SkeletalMeshComponent) : SkeletalMesh->GetName();
 	}
 
 	virtual void Complete() override;
@@ -72,7 +69,6 @@ private:
 	FGLTFSkeletalMeshSectionConverter& MeshSectionConverter;
 	const USkeletalMesh* SkeletalMesh;
 	const USkeletalMeshComponent* SkeletalMeshComponent;
-	int32 LODIndex;
 	FGLTFMaterialArray OverrideMaterials;
 	const FGLTFJsonMeshIndex MeshIndex;
 };
