[GLTFExporter] refactor: add GLTFMeshData to simplify material converter parameters and reuse retrieved mesh (+ extracted sections)

commit 5715850cd51f7cab32dc2598668a400c35854b1b
Date: Mon, 2 Nov 2020 23:40:07 +0100

    refactor: add GLTFMeshData to simplify material converter parameters and reuse retrieved mesh (+ extracted sections)
    (cherry picked from commit 71c66d4e41cc3a49c3c04bb4b32ad2d7e9375903​)

#skipci

diff --git a/Engine/Plugins/Enterprise/GLTFExporter/MigratedFromGit/Source/GLTFExporter/Private/Builders/GLTFConvertBuilder.cpp b/Engine/Plugins/Enterprise/GLTFExporter/MigratedFromGit/Source/GLTFExporter/Private/Builders/GLTFConvertBuilder.cpp
--- a/Engine/Plugins/Enterprise/GLTFExporter/MigratedFromGit/Source/GLTFExporter/Private/Builders/GLTFConvertBuilder.cpp
+++ b/Engine/Plugins/Enterprise/GLTFExporter/MigratedFromGit/Source/GLTFExporter/Private/Builders/GLTFConvertBuilder.cpp
@@ -140,44 +140,14 @@ FGLTFJsonMeshIndex FGLTFConvertBuilder::GetOrAddMesh(const USkeletalMeshComponen
 	return SkeletalMeshConverter.GetOrAdd(SkeletalMesh, SkeletalMeshComponent, OverrideMaterials);
 }
 
-FGLTFJsonMaterialIndex FGLTFConvertBuilder::GetOrAddMaterial(const UMaterialInterface* Material)
+FGLTFJsonMaterialIndex FGLTFConvertBuilder::GetOrAddMaterial(const UMaterialInterface* Material, const FGLTFMeshData* MeshData, const FGLTFMaterialArray& OverrideMaterials)
 {
 	if (Material == nullptr)
 	{
 		return FGLTFJsonMaterialIndex(INDEX_NONE);
 	}
 
-	return MaterialConverter.GetOrAdd(Material, nullptr, {});
-}
-
-FGLTFJsonMaterialIndex FGLTFConvertBuilder::GetOrAddMaterial(const UMaterialInterface* Material, const UStaticMesh* Mesh, const FGLTFMaterialArray& OverrideMaterials)
-{
-	if (Material == nullptr)
-	{
-		return FGLTFJsonMaterialIndex(INDEX_NONE);
-	}
-
-	return MaterialConverter.GetOrAdd(Material, Mesh, OverrideMaterials);
-}
-
-FGLTFJsonMaterialIndex FGLTFConvertBuilder::GetOrAddMaterial(const UMaterialInterface* Material, const USkeletalMesh* Mesh, const FGLTFMaterialArray& OverrideMaterials)
-{
-	if (Material == nullptr)
-	{
-		return FGLTFJsonMaterialIndex(INDEX_NONE);
-	}
-
-	return MaterialConverter.GetOrAdd(Material, Mesh, OverrideMaterials);
-}
-
-FGLTFJsonMaterialIndex FGLTFConvertBuilder::GetOrAddMaterial(const UMaterialInterface* Material, const UMeshComponent* MeshComponent, const FGLTFMaterialArray& OverrideMaterials)
-{
-	if (Material == nullptr)
-	{
-		return FGLTFJsonMaterialIndex(INDEX_NONE);
-	}
-
-	return MaterialConverter.GetOrAdd(Material, MeshComponent, OverrideMaterials);
+	return MaterialConverter.GetOrAdd(Material, MeshData, OverrideMaterials);
 }
 
 FGLTFJsonSamplerIndex FGLTFConvertBuilder::GetOrAddSampler(const UTexture* Texture)
diff --git a/Engine/Plugins/Enterprise/GLTFExporter/MigratedFromGit/Source/GLTFExporter/Private/Builders/GLTFConvertBuilder.h b/Engine/Plugins/Enterprise/GLTFExporter/MigratedFromGit/Source/GLTFExporter/Private/Builders/GLTFConvertBuilder.h
--- a/Engine/Plugins/Enterprise/GLTFExporter/MigratedFromGit/Source/GLTFExporter/Private/Builders/GLTFConvertBuilder.h
+++ b/Engine/Plugins/Enterprise/GLTFExporter/MigratedFromGit/Source/GLTFExporter/Private/Builders/GLTFConvertBuilder.h
@@ -44,10 +44,7 @@ public:
 	FGLTFJsonMeshIndex GetOrAddMesh(const USkeletalMesh* SkeletalMesh, const FGLTFMaterialArray& OverrideMaterials = {});
 	FGLTFJsonMeshIndex GetOrAddMesh(const USkeletalMeshComponent* SkeletalMeshComponent, const FGLTFMaterialArray& OverrideMaterials = {});
 
-	FGLTFJsonMaterialIndex GetOrAddMaterial(const UMaterialInterface* Material);
-	FGLTFJsonMaterialIndex GetOrAddMaterial(const UMaterialInterface* Material, const UStaticMesh* Mesh, const FGLTFMaterialArray& OverrideMaterials = {});
-	FGLTFJsonMaterialIndex GetOrAddMaterial(const UMaterialInterface* Material, const USkeletalMesh* Mesh, const FGLTFMaterialArray& OverrideMaterials = {});
-	FGLTFJsonMaterialIndex GetOrAddMaterial(const UMaterialInterface* Material, const UMeshComponent* MeshComponent, const FGLTFMaterialArray& OverrideMaterials = {});
+	FGLTFJsonMaterialIndex GetOrAddMaterial(const UMaterialInterface* Material, const FGLTFMeshData* MeshData = nullptr, const FGLTFMaterialArray& OverrideMaterials = {});
 
 	FGLTFJsonSamplerIndex GetOrAddSampler(const UTexture* Texture);
 	FGLTFJsonTextureIndex GetOrAddTexture(const UTexture2D* Texture);
diff --git a/Engine/Plugins/Enterprise/GLTFExporter/MigratedFromGit/Source/GLTFExporter/Private/Converters/GLTFMaterialConverters.cpp b/Engine/Plugins/Enterprise/GLTFExporter/MigratedFromGit/Source/GLTFExporter/Private/Converters/GLTFMaterialConverters.cpp
--- a/Engine/Plugins/Enterprise/GLTFExporter/MigratedFromGit/Source/GLTFExporter/Private/Converters/GLTFMaterialConverters.cpp
+++ b/Engine/Plugins/Enterprise/GLTFExporter/MigratedFromGit/Source/GLTFExporter/Private/Converters/GLTFMaterialConverters.cpp
@@ -6,43 +6,26 @@
 #include "Builders/GLTFConvertBuilder.h"
 #include "Tasks/GLTFMaterialTasks.h"
 
-void FGLTFMaterialConverter::Sanitize(const UMaterialInterface*& Material, const UObject*& MeshOrComponent, FGLTFMaterialArray& OverrideMaterials)
+void FGLTFMaterialConverter::Sanitize(const UMaterialInterface*& Material, const FGLTFMeshData*& MeshData, FGLTFMaterialArray& OverrideMaterials)
 {
-	const UStaticMesh* StaticMesh = Cast<UStaticMesh>(MeshOrComponent);
-	const USkeletalMesh* SkeletalMesh = Cast<USkeletalMesh>(MeshOrComponent);
-
-	if (MeshOrComponent != nullptr && StaticMesh == nullptr && SkeletalMesh == nullptr)
-	{
-		if (const UStaticMeshComponent* StaticMeshComponent = Cast<UStaticMeshComponent>(MeshOrComponent))
-		{
-			StaticMesh = StaticMeshComponent->GetStaticMesh();
-		}
-		else if (const USkeletalMeshComponent* SkeletalMeshComponent = Cast<USkeletalMeshComponent>(MeshOrComponent))
-		{
-			SkeletalMesh = SkeletalMeshComponent->SkeletalMesh;
-		}
-	}
-
-	if (StaticMesh == nullptr && SkeletalMesh == nullptr ||
+	if (MeshData == nullptr ||
 		!Builder.ExportOptions->bBakeMaterialInputs ||
 		!Builder.ExportOptions->bBakeMaterialInputsUsingMeshData ||
 		!FGLTFMaterialUtility::MaterialNeedsVertexData(Material))
 	{
-		MeshOrComponent = nullptr;
+		MeshData = nullptr;
 		OverrideMaterials = {};
 	}
 }
 
-FGLTFJsonMaterialIndex FGLTFMaterialConverter::Convert(const UMaterialInterface* Material, const UObject* MeshOrComponent, const FGLTFMaterialArray OverrideMaterials)
+FGLTFJsonMaterialIndex FGLTFMaterialConverter::Convert(const UMaterialInterface* Material, const FGLTFMeshData* MeshData, const FGLTFMaterialArray OverrideMaterials)
 {
 	if (Material == FGLTFMaterialUtility::GetDefault())
 	{
 		return FGLTFJsonMaterialIndex(INDEX_NONE); // use default gltf definition
 	}
 
-	const int32 LODIndex = FGLTFMeshUtility::GetLOD(MeshOrComponent, Builder.ExportOptions->DefaultLevelOfDetail);
-
 	const FGLTFJsonMaterialIndex MaterialIndex = Builder.AddMaterial();
-	Builder.SetupTask<FGLTFMaterialTask>(Builder, Material, MeshOrComponent, LODIndex, OverrideMaterials, MaterialIndex);
+	Builder.SetupTask<FGLTFMaterialTask>(Builder, Material, MeshData, OverrideMaterials, MaterialIndex);
 	return MaterialIndex;
 }
diff --git a/Engine/Plugins/Enterprise/GLTFExporter/MigratedFromGit/Source/GLTFExporter/Private/Converters/GLTFMaterialConverters.h b/Engine/Plugins/Enterprise/GLTFExporter/MigratedFromGit/Source/GLTFExporter/Private/Converters/GLTFMaterialConverters.h
--- a/Engine/Plugins/Enterprise/GLTFExporter/MigratedFromGit/Source/GLTFExporter/Private/Converters/GLTFMaterialConverters.h
+++ b/Engine/Plugins/Enterprise/GLTFExporter/MigratedFromGit/Source/GLTFExporter/Private/Converters/GLTFMaterialConverters.h
@@ -5,14 +5,15 @@
 #include "Json/GLTFJsonIndex.h"
 #include "Converters/GLTFConverter.h"
 #include "Converters/GLTFBuilderContext.h"
+#include "Converters/GLTFMeshData.h"
 #include "Converters/GLTFMaterialArray.h"
 #include "Engine.h"
 
-class FGLTFMaterialConverter final : public FGLTFBuilderContext, public TGLTFConverter<FGLTFJsonMaterialIndex, const UMaterialInterface*, const UObject*, FGLTFMaterialArray>
+class FGLTFMaterialConverter final : public FGLTFBuilderContext, public TGLTFConverter<FGLTFJsonMaterialIndex, const UMaterialInterface*, const FGLTFMeshData*, FGLTFMaterialArray>
 {
 	using FGLTFBuilderContext::FGLTFBuilderContext;
 
-	virtual void Sanitize(const UMaterialInterface*& Material, const UObject*& MeshOrComponent, FGLTFMaterialArray& OverrideMaterials) override;
+	virtual void Sanitize(const UMaterialInterface*& Material, const FGLTFMeshData*& MeshData, FGLTFMaterialArray& OverrideMaterials) override;
 
-	virtual FGLTFJsonMaterialIndex Convert(const UMaterialInterface* Material, const UObject* MeshOrComponent, const FGLTFMaterialArray OverrideMaterials) override;
+	virtual FGLTFJsonMaterialIndex Convert(const UMaterialInterface* Material, const FGLTFMeshData* MeshData, const FGLTFMaterialArray OverrideMaterials) override;
 };
diff --git a/Engine/Plugins/Enterprise/GLTFExporter/MigratedFromGit/Source/GLTFExporter/Private/Converters/GLTFMeshConverters.cpp b/Engine/Plugins/Enterprise/GLTFExporter/MigratedFromGit/Source/GLTFExporter/Private/Converters/GLTFMeshConverters.cpp
--- a/Engine/Plugins/Enterprise/GLTFExporter/MigratedFromGit/Source/GLTFExporter/Private/Converters/GLTFMeshConverters.cpp
+++ b/Engine/Plugins/Enterprise/GLTFExporter/MigratedFromGit/Source/GLTFExporter/Private/Converters/GLTFMeshConverters.cpp
@@ -39,7 +39,7 @@ void FGLTFStaticMeshConverter::Sanitize(const UStaticMesh*& StaticMesh, const US
 FGLTFJsonMeshIndex FGLTFStaticMeshConverter::Convert(const UStaticMesh* StaticMesh, const UStaticMeshComponent* StaticMeshComponent, FGLTFMaterialArray OverrideMaterials)
 {
 	const FGLTFJsonMeshIndex MeshIndex = Builder.AddMesh();
-	Builder.SetupTask<FGLTFStaticMeshTask>(Builder, MeshSectionConverter, StaticMesh, StaticMeshComponent, OverrideMaterials, MeshIndex);
+	Builder.SetupTask<FGLTFStaticMeshTask>(Builder, MeshSectionConverter, MeshDataConverter, StaticMesh, StaticMeshComponent, OverrideMaterials, MeshIndex);
 	return MeshIndex;
 }
 
@@ -77,6 +77,6 @@ void FGLTFSkeletalMeshConverter::Sanitize(const USkeletalMesh*& SkeletalMesh, co
 FGLTFJsonMeshIndex FGLTFSkeletalMeshConverter::Convert(const USkeletalMesh* SkeletalMesh, const USkeletalMeshComponent* SkeletalMeshComponent, FGLTFMaterialArray OverrideMaterials)
 {
 	const FGLTFJsonMeshIndex MeshIndex = Builder.AddMesh();
-	Builder.SetupTask<FGLTFSkeletalMeshTask>(Builder, MeshSectionConverter, SkeletalMesh, SkeletalMeshComponent, OverrideMaterials, MeshIndex);
+	Builder.SetupTask<FGLTFSkeletalMeshTask>(Builder, MeshSectionConverter, MeshDataConverter, SkeletalMesh, SkeletalMeshComponent, OverrideMaterials, MeshIndex);
 	return MeshIndex;
 }
diff --git a/Engine/Plugins/Enterprise/GLTFExporter/MigratedFromGit/Source/GLTFExporter/Private/Converters/GLTFMeshConverters.h b/Engine/Plugins/Enterprise/GLTFExporter/MigratedFromGit/Source/GLTFExporter/Private/Converters/GLTFMeshConverters.h
--- a/Engine/Plugins/Enterprise/GLTFExporter/MigratedFromGit/Source/GLTFExporter/Private/Converters/GLTFMeshConverters.h
+++ b/Engine/Plugins/Enterprise/GLTFExporter/MigratedFromGit/Source/GLTFExporter/Private/Converters/GLTFMeshConverters.h
@@ -6,6 +6,7 @@
 #include "Converters/GLTFConverter.h"
 #include "Converters/GLTFBuilderContext.h"
 #include "Converters/GLTFMeshSectionConverters.h"
+#include "Converters/GLTFMeshDataConverters.h"
 #include "Converters/GLTFMaterialArray.h"
 #include "Engine.h"
 
@@ -24,6 +25,7 @@ class FGLTFStaticMeshConverter final : public TGLTFMeshConverter<const UStaticMe
 	virtual FGLTFJsonMeshIndex Convert(const UStaticMesh* StaticMesh, const UStaticMeshComponent* StaticMeshComponent, FGLTFMaterialArray OverrideMaterials) override;
 
 	FGLTFStaticMeshSectionConverter MeshSectionConverter;
+	FGLTFStaticMeshDataConverter MeshDataConverter;
 };
 
 class FGLTFSkeletalMeshConverter final : public TGLTFMeshConverter<const USkeletalMesh*, const USkeletalMeshComponent*, FGLTFMaterialArray>
@@ -35,4 +37,5 @@ class FGLTFSkeletalMeshConverter final : public TGLTFMeshConverter<const USkelet
 	virtual FGLTFJsonMeshIndex Convert(const USkeletalMesh* SkeletalMesh, const USkeletalMeshComponent* SkeletalMeshComponent, FGLTFMaterialArray OverrideMaterials) override;
 
 	FGLTFSkeletalMeshSectionConverter MeshSectionConverter;
+	FGLTFSkeletalMeshDataConverter MeshDataConverter;
 };
diff --git a/Engine/Plugins/Enterprise/GLTFExporter/MigratedFromGit/Source/GLTFExporter/Private/Converters/GLTFMeshData.cpp b/Engine/Plugins/Enterprise/GLTFExporter/MigratedFromGit/Source/GLTFExporter/Private/Converters/GLTFMeshData.cpp
new file mode 100644
--- /dev/null
+++ b/Engine/Plugins/Enterprise/GLTFExporter/MigratedFromGit/Source/GLTFExporter/Private/Converters/GLTFMeshData.cpp
@@ -0,0 +1,41 @@
+// Copyright Epic Games, Inc. All Rights Reserved.
+
+#include "Converters/GLTFMeshData.h"
+#include "Converters/GLTFNameUtility.h"
+#include "StaticMeshAttributes.h"
+#include "Developer/MeshMergeUtilities/Private/MeshMergeHelpers.h"
+
+FGLTFMeshData::FGLTFMeshData(const UStaticMesh* StaticMesh, const UStaticMeshComponent* StaticMeshComponent, int32 LODIndex)
+{
+	FStaticMeshAttributes(Description).Register();
+
+	if (StaticMeshComponent != nullptr)
+	{
+		FMeshMergeHelpers::RetrieveMesh(StaticMeshComponent, LODIndex, Description, true);
+		FMeshMergeHelpers::ExtractSections(StaticMeshComponent, LODIndex, Sections);
+		Name = FGLTFNameUtility::GetName(StaticMeshComponent);
+	}
+	else
+	{
+		FMeshMergeHelpers::RetrieveMesh(StaticMesh, LODIndex, Description);
+		FMeshMergeHelpers::ExtractSections(StaticMesh, LODIndex, Sections);
+		StaticMesh->GetName(Name);
+	}
+}
+
+FGLTFMeshData::FGLTFMeshData(const USkeletalMesh* SkeletalMesh, const USkeletalMeshComponent* SkeletalMeshComponent, int32 LODIndex)
+{
+	FStaticMeshAttributes(Description).Register();
+
+	if (SkeletalMeshComponent != nullptr)
+	{
+		FMeshMergeHelpers::RetrieveMesh(const_cast<USkeletalMeshComponent*>(SkeletalMeshComponent), LODIndex, Description, true);
+		FMeshMergeHelpers::ExtractSections(SkeletalMeshComponent, LODIndex, Sections);
+		Name = FGLTFNameUtility::GetName(SkeletalMeshComponent);
+	}
+	else
+	{
+		// TODO: add support for skeletal meshes by implementing custom utilities for retrieving mesh & sections
+		SkeletalMesh->GetName(Name);
+	}
+}
diff --git a/Engine/Plugins/Enterprise/GLTFExporter/MigratedFromGit/Source/GLTFExporter/Private/Converters/GLTFMeshData.h b/Engine/Plugins/Enterprise/GLTFExporter/MigratedFromGit/Source/GLTFExporter/Private/Converters/GLTFMeshData.h
new file mode 100644
--- /dev/null
+++ b/Engine/Plugins/Enterprise/GLTFExporter/MigratedFromGit/Source/GLTFExporter/Private/Converters/GLTFMeshData.h
@@ -0,0 +1,16 @@
+// Copyright Epic Games, Inc. All Rights Reserved.
+
+#pragma once
+
+#include "CoreMinimal.h"
+#include "MeshDescription.h"
+
+struct FGLTFMeshData
+{
+	FGLTFMeshData(const UStaticMesh* StaticMesh, const UStaticMeshComponent* StaticMeshComponent, int32 LODIndex);
+	FGLTFMeshData(const USkeletalMesh* SkeletalMesh, const USkeletalMeshComponent* SkeletalMeshComponent, int32 LODIndex);
+
+	FMeshDescription Description;
+	TArray<FSectionInfo> Sections;
+	FString Name;
+};
diff --git a/Engine/Plugins/Enterprise/GLTFExporter/MigratedFromGit/Source/GLTFExporter/Private/Converters/GLTFMeshDataConverters.h b/Engine/Plugins/Enterprise/GLTFExporter/MigratedFromGit/Source/GLTFExporter/Private/Converters/GLTFMeshDataConverters.h
new file mode 100644
--- /dev/null
+++ b/Engine/Plugins/Enterprise/GLTFExporter/MigratedFromGit/Source/GLTFExporter/Private/Converters/GLTFMeshDataConverters.h
@@ -0,0 +1,21 @@
+// Copyright Epic Games, Inc. All Rights Reserved.
+
+#pragma once
+
+#include "Converters/GLTFConverter.h"
+#include "Converters/GLTFMeshData.h"
+#include "Engine.h"
+
+template <typename MeshType, typename MeshComponentType>
+class TGLTFMeshDataConverter final : public TGLTFConverter<const FGLTFMeshData*, const MeshType*, const MeshComponentType*, int32>
+{
+	TArray<TUniquePtr<FGLTFMeshData>> Outputs;
+
+	const FGLTFMeshData* Convert(const MeshType* Mesh, const MeshComponentType* MeshComponent, int32 LODIndex)
+	{
+		return Outputs.Add_GetRef(MakeUnique<FGLTFMeshData>(Mesh, MeshComponent, LODIndex)).Get();
+	}
+};
+
+typedef TGLTFMeshDataConverter<UStaticMesh, UStaticMeshComponent> FGLTFStaticMeshDataConverter;
+typedef TGLTFMeshDataConverter<USkeletalMesh, USkeletalMeshComponent> FGLTFSkeletalMeshDataConverter;
diff --git a/Engine/Plugins/Enterprise/GLTFExporter/MigratedFromGit/Source/GLTFExporter/Private/Converters/GLTFMeshUtility.cpp b/Engine/Plugins/Enterprise/GLTFExporter/MigratedFromGit/Source/GLTFExporter/Private/Converters/GLTFMeshUtility.cpp
--- a/Engine/Plugins/Enterprise/GLTFExporter/MigratedFromGit/Source/GLTFExporter/Private/Converters/GLTFMeshUtility.cpp
+++ b/Engine/Plugins/Enterprise/GLTFExporter/MigratedFromGit/Source/GLTFExporter/Private/Converters/GLTFMeshUtility.cpp
@@ -4,31 +4,6 @@
 #include "Rendering/SkeletalMeshRenderData.h"
 #include "PlatformInfo.h"
 
-int32 FGLTFMeshUtility::GetLOD(const UObject* MeshOrComponent, int32 DefaultLOD)
-{
-	if (const UStaticMesh* StaticMesh = Cast<UStaticMesh>(MeshOrComponent))
-	{
-		return GetLOD(StaticMesh, nullptr, DefaultLOD);
-	}
-
-	if (const USkeletalMesh* SkeletalMesh = Cast<USkeletalMesh>(MeshOrComponent))
-	{
-		return GetLOD(SkeletalMesh, nullptr, DefaultLOD);
-	}
-
-	if (const UStaticMeshComponent* StaticMeshComponent = Cast<UStaticMeshComponent>(MeshOrComponent))
-	{
-		return GetLOD(StaticMeshComponent->GetStaticMesh(), StaticMeshComponent, DefaultLOD);
-	}
-
-	if (const USkeletalMeshComponent* SkeletalMeshComponent = Cast<USkeletalMeshComponent>(MeshOrComponent))
-	{
-		return GetLOD(SkeletalMeshComponent->SkeletalMesh, SkeletalMeshComponent, DefaultLOD);
-	}
-
-	return -1;
-}
-
 int32 FGLTFMeshUtility::GetLOD(const UStaticMesh* StaticMesh, const UStaticMeshComponent* StaticMeshComponent, int32 DefaultLOD)
 {
 	const int32 ForcedLOD = StaticMeshComponent != nullptr ? StaticMeshComponent->ForcedLodModel - 1 : -1;
diff --git a/Engine/Plugins/Enterprise/GLTFExporter/MigratedFromGit/Source/GLTFExporter/Private/Converters/GLTFMeshUtility.h b/Engine/Plugins/Enterprise/GLTFExporter/MigratedFromGit/Source/GLTFExporter/Private/Converters/GLTFMeshUtility.h
--- a/Engine/Plugins/Enterprise/GLTFExporter/MigratedFromGit/Source/GLTFExporter/Private/Converters/GLTFMeshUtility.h
+++ b/Engine/Plugins/Enterprise/GLTFExporter/MigratedFromGit/Source/GLTFExporter/Private/Converters/GLTFMeshUtility.h
@@ -6,8 +6,6 @@
 
 struct FGLTFMeshUtility
 {
-	static int32 GetLOD(const UObject* MeshOrComponent, int32 DefaultLOD);
-
 	static int32 GetLOD(const UStaticMesh* StaticMesh, const UStaticMeshComponent* StaticMeshComponent, int32 DefaultLOD);
 	static int32 GetLOD(const USkeletalMesh* SkeletalMesh, const USkeletalMeshComponent* SkeletalMeshComponent, int32 DefaultLOD);
 
diff --git a/Engine/Plugins/Enterprise/GLTFExporter/MigratedFromGit/Source/GLTFExporter/Private/Converters/GLTFNameUtility.cpp b/Engine/Plugins/Enterprise/GLTFExporter/MigratedFromGit/Source/GLTFExporter/Private/Converters/GLTFNameUtility.cpp
--- a/Engine/Plugins/Enterprise/GLTFExporter/MigratedFromGit/Source/GLTFExporter/Private/Converters/GLTFNameUtility.cpp
+++ b/Engine/Plugins/Enterprise/GLTFExporter/MigratedFromGit/Source/GLTFExporter/Private/Converters/GLTFNameUtility.cpp
@@ -6,19 +6,6 @@
 
 namespace
 {
-	FString GetNameWithLODSuffix(const UObject* Object, int32 LODIndex)
-	{
-		FString Name;
-		Object->GetName(Name);
-
-		if (LODIndex != 0)
-		{
-			Name += TEXT("_LOD") + FString::FromInt(LODIndex);
-		}
-
-		return Name;
-	}
-
 	template <typename ActorType>
 	FString GetActorNameIfOwnerOfType(const UActorComponent* Component)
 	{
@@ -31,16 +18,6 @@ namespace
 	}
 }
 
-FString FGLTFNameUtility::GetName(const UStaticMesh* StaticMesh, int32 LODIndex)
-{
-	return GetNameWithLODSuffix(StaticMesh, LODIndex);
-}
-
-FString FGLTFNameUtility::GetName(const USkeletalMesh* SkeletalMesh, int32 LODIndex)
-{
-	return GetNameWithLODSuffix(SkeletalMesh, LODIndex);
-}
-
 FString FGLTFNameUtility::GetName(const USceneComponent* Component)
 {
 	if (const AActor* Owner = Component->GetOwner())
diff --git a/Engine/Plugins/Enterprise/GLTFExporter/MigratedFromGit/Source/GLTFExporter/Private/Converters/GLTFNameUtility.h b/Engine/Plugins/Enterprise/GLTFExporter/MigratedFromGit/Source/GLTFExporter/Private/Converters/GLTFNameUtility.h
--- a/Engine/Plugins/Enterprise/GLTFExporter/MigratedFromGit/Source/GLTFExporter/Private/Converters/GLTFNameUtility.h
+++ b/Engine/Plugins/Enterprise/GLTFExporter/MigratedFromGit/Source/GLTFExporter/Private/Converters/GLTFNameUtility.h
@@ -18,9 +18,6 @@ struct FGLTFNameUtility
 		return Enum->GetDisplayNameTextByValue(Value).ToString();
 	}
 
-	static FString GetName(const UStaticMesh* StaticMesh, int32 LODIndex);
-	static FString GetName(const USkeletalMesh* SkeletalMesh, int32 LODIndex);
-
 	static FString GetName(const USceneComponent* Component);
 
 	static FString GetName(const UStaticMeshComponent* Component);
diff --git a/Engine/Plugins/Enterprise/GLTFExporter/MigratedFromGit/Source/GLTFExporter/Private/Exporters/GLTFMaterialExporter.cpp b/Engine/Plugins/Enterprise/GLTFExporter/MigratedFromGit/Source/GLTFExporter/Private/Exporters/GLTFMaterialExporter.cpp
--- a/Engine/Plugins/Enterprise/GLTFExporter/MigratedFromGit/Source/GLTFExporter/Private/Exporters/GLTFMaterialExporter.cpp
+++ b/Engine/Plugins/Enterprise/GLTFExporter/MigratedFromGit/Source/GLTFExporter/Private/Exporters/GLTFMaterialExporter.cpp
@@ -16,15 +16,6 @@ bool UGLTFMaterialExporter::AddObject(FGLTFContainerBuilder& Builder, const UObj
 	const UMaterialInterface* Material = CastChecked<UMaterialInterface>(Object);
 	const UStaticMesh* PreviewMesh = Builder.ExportOptions->bExportPreviewMesh ? FGLTFExporterUtility::GetPreviewMesh(Material) : nullptr;
 
-	const FGLTFJsonMaterialIndex MaterialIndex = Builder.GetOrAddMaterial(Material, PreviewMesh);
-	if (MaterialIndex == INDEX_NONE)
-	{
-		Builder.AddErrorMessage(
-			FString::Printf(TEXT("Failed to export material %s"),
-			*Material->GetName()));
-		return false;
-	}
-
 	if (Builder.ExportOptions->bExportPreviewMesh)
 	{
 		if (PreviewMesh != nullptr)
@@ -57,6 +48,17 @@ bool UGLTFMaterialExporter::AddObject(FGLTFContainerBuilder& Builder, const UObj
 			return false;
 		}
 	}
+	else
+	{
+		const FGLTFJsonMaterialIndex MaterialIndex = Builder.GetOrAddMaterial(Material);
+		if (MaterialIndex == INDEX_NONE)
+		{
+			Builder.AddErrorMessage(
+                FString::Printf(TEXT("Failed to export material %s"),
+                *Material->GetName()));
+			return false;
+		}
+	}
 
 	return true;
 }
diff --git a/Engine/Plugins/Enterprise/GLTFExporter/MigratedFromGit/Source/GLTFExporter/Private/Tasks/GLTFMaterialTasks.cpp b/Engine/Plugins/Enterprise/GLTFExporter/MigratedFromGit/Source/GLTFExporter/Private/Tasks/GLTFMaterialTasks.cpp
--- a/Engine/Plugins/Enterprise/GLTFExporter/MigratedFromGit/Source/GLTFExporter/Private/Tasks/GLTFMaterialTasks.cpp
+++ b/Engine/Plugins/Enterprise/GLTFExporter/MigratedFromGit/Source/GLTFExporter/Private/Tasks/GLTFMaterialTasks.cpp
@@ -13,9 +13,6 @@
 #include "Materials/MaterialExpressionVectorParameter.h"
 #include "Materials/MaterialExpressionTextureSample.h"
 #include "Materials/MaterialExpressionTextureSampleParameter2D.h"
-#include "MeshDescription.h"
-#include "StaticMeshAttributes.h"
-#include "Developer/MeshMergeUtilities/Private/MeshMergeHelpers.h"
 
 namespace
 {
@@ -1356,43 +1353,17 @@ FGLTFPropertyBakeOutput FGLTFMaterialTask::BakeMaterialProperty(EMaterialPropert
 	const FIntPoint DefaultTextureSize = Builder.GetDefaultMaterialBakeSize();
 	const FIntPoint TextureSize = PreferredTextureSize != nullptr ? *PreferredTextureSize : DefaultTextureSize;
 
-	FMeshDescription MeshDescription;
 	TArray<int32> MeshMaterialIndices;
 
-	bool bUseMeshData = bHasValidMeshOrComponent;
-
-	if (bUseMeshData)
+	if (MeshData != nullptr)
 	{
-		FStaticMeshAttributes(MeshDescription).Register();
-		TArray<FSectionInfo> Sections;
+		const TArray<FSectionInfo>& Sections = MeshData->Sections;
 		const TArray<UMaterialInterface*> SectionMaterials;
 
-		if (StaticMesh != nullptr)
-		{
-			FMeshMergeHelpers::RetrieveMesh(StaticMesh, LODIndex, MeshDescription);
-			FMeshMergeHelpers::ExtractSections(StaticMesh, LODIndex, Sections);
-		}
-		else if (StaticMeshComponent != nullptr)
-		{
-			FMeshMergeHelpers::RetrieveMesh(StaticMeshComponent, LODIndex, MeshDescription, true);
-			FMeshMergeHelpers::ExtractSections(StaticMeshComponent, LODIndex, Sections);
-		}
-		else if (SkeletalMesh != nullptr)
-		{
-			// TODO: add support for skeletal meshes by implementing custom utilities for retrieving mesh & sections
-
-			bUseMeshData = false;
-		}
-		else if (SkeletalMeshComponent != nullptr)
-		{
-			FMeshMergeHelpers::RetrieveMesh(const_cast<USkeletalMeshComponent*>(SkeletalMeshComponent), LODIndex, MeshDescription, true);
-			FMeshMergeHelpers::ExtractSections(SkeletalMeshComponent, LODIndex, Sections);
-		}
-
 		// Add material indices of all sections that reference the current material
 		for (int32 SectionIndex = 0; SectionIndex < Sections.Num(); ++SectionIndex)
 		{
-			FSectionInfo& Section = Sections[SectionIndex];
+			const FSectionInfo& Section = Sections[SectionIndex];
 			const UMaterialInterface* OverrideMaterial = OverrideMaterials.IsValidIndex(SectionIndex) ? OverrideMaterials[SectionIndex] : nullptr;
 			const UMaterialInterface* SectionMaterial = OverrideMaterial != nullptr ? OverrideMaterial : Section.Material;
 
@@ -1417,7 +1388,7 @@ FGLTFPropertyBakeOutput FGLTFMaterialTask::BakeMaterialProperty(EMaterialPropert
 		Property,
 		Material,
 		OutTexCoord,
-		bUseMeshData ? &MeshDescription : nullptr,
+		MeshData != nullptr ? &MeshData->Description : nullptr,
 		MeshMaterialIndices,
 		bCopyAlphaFromRedChannel);
 
@@ -1457,21 +1428,9 @@ FString FGLTFMaterialTask::GetMaterialName() const
 {
 	FString MaterialName = Material->GetName();
 
-	if (StaticMesh != nullptr)
-	{
-		MaterialName += TEXT("_") + FGLTFNameUtility::GetName(StaticMesh, LODIndex);
-	}
-	else if (StaticMeshComponent != nullptr)
-	{
-		MaterialName += TEXT("_") + FGLTFNameUtility::GetName(StaticMeshComponent);
-	}
-	else if (SkeletalMesh != nullptr)
-	{
-		MaterialName += TEXT("_") + FGLTFNameUtility::GetName(SkeletalMesh, LODIndex);
-	}
-	else if (SkeletalMeshComponent != nullptr)
+	if (MeshData != nullptr)
 	{
-		MaterialName += TEXT("_") + FGLTFNameUtility::GetName(SkeletalMeshComponent);
+		MaterialName += TEXT("_") + MeshData->Name;
 	}
 
 	return MaterialName;
diff --git a/Engine/Plugins/Enterprise/GLTFExporter/MigratedFromGit/Source/GLTFExporter/Private/Tasks/GLTFMaterialTasks.h b/Engine/Plugins/Enterprise/GLTFExporter/MigratedFromGit/Source/GLTFExporter/Private/Tasks/GLTFMaterialTasks.h
--- a/Engine/Plugins/Enterprise/GLTFExporter/MigratedFromGit/Source/GLTFExporter/Private/Tasks/GLTFMaterialTasks.h
+++ b/Engine/Plugins/Enterprise/GLTFExporter/MigratedFromGit/Source/GLTFExporter/Private/Tasks/GLTFMaterialTasks.h
@@ -3,6 +3,7 @@
 #pragma once
 
 #include "Tasks/GLTFTask.h"
+#include "Converters/GLTFMeshData.h"
 #include "Builders/GLTFConvertBuilder.h"
 #include "Engine.h"
 
@@ -12,18 +13,13 @@ class FGLTFMaterialTask : public FGLTFTask
 {
 public:
 
-	FGLTFMaterialTask(FGLTFConvertBuilder& Builder, const UMaterialInterface* Material, const UObject* MeshOrComponent, int32 LODIndex, FGLTFMaterialArray OverrideMaterials, FGLTFJsonMaterialIndex MaterialIndex)
+	FGLTFMaterialTask(FGLTFConvertBuilder& Builder, const UMaterialInterface* Material, const FGLTFMeshData* MeshData, FGLTFMaterialArray OverrideMaterials, FGLTFJsonMaterialIndex MaterialIndex)
         : FGLTFTask(EGLTFTaskPriority::Material)
 		, Builder(Builder)
 		, Material(Material)
-		, LODIndex(LODIndex)
+		, MeshData(MeshData)
 		, OverrideMaterials(OverrideMaterials)
 		, MaterialIndex(MaterialIndex)
-		, StaticMesh(Cast<UStaticMesh>(MeshOrComponent))
-		, StaticMeshComponent(Cast<UStaticMeshComponent>(MeshOrComponent))
-		, SkeletalMesh(Cast<USkeletalMesh>(MeshOrComponent))
-		, SkeletalMeshComponent(Cast<USkeletalMeshComponent>(MeshOrComponent))
-		, bHasValidMeshOrComponent(StaticMesh != nullptr || StaticMeshComponent != nullptr || SkeletalMesh != nullptr || SkeletalMeshComponent != nullptr)
 	{
 	}
 
@@ -38,16 +34,10 @@ private:
 
 	FGLTFConvertBuilder& Builder;
 	const UMaterialInterface* Material;
-	const int32 LODIndex;
+	const FGLTFMeshData* MeshData;
 	const FGLTFMaterialArray OverrideMaterials;
 	const FGLTFJsonMaterialIndex MaterialIndex;
 
-	const UStaticMesh* StaticMesh;
-	const UStaticMeshComponent* StaticMeshComponent;
-	const USkeletalMesh* SkeletalMesh;
-	const USkeletalMeshComponent* SkeletalMeshComponent;
-	const bool bHasValidMeshOrComponent;
-
 	bool TryGetAlphaMode(EGLTFJsonAlphaMode& AlphaMode) const;
 	bool TryGetShadingModel(EGLTFJsonShadingModel& ShadingModel) const;
 
diff --git a/Engine/Plugins/Enterprise/GLTFExporter/MigratedFromGit/Source/GLTFExporter/Private/Tasks/GLTFMeshTasks.cpp b/Engine/Plugins/Enterprise/GLTFExporter/MigratedFromGit/Source/GLTFExporter/Private/Tasks/GLTFMeshTasks.cpp
--- a/Engine/Plugins/Enterprise/GLTFExporter/MigratedFromGit/Source/GLTFExporter/Private/Tasks/GLTFMeshTasks.cpp
+++ b/Engine/Plugins/Enterprise/GLTFExporter/MigratedFromGit/Source/GLTFExporter/Private/Tasks/GLTFMeshTasks.cpp
@@ -24,6 +24,10 @@ void FGLTFStaticMeshTask::Complete()
 		ColorBuffer = LODInfo.OverrideVertexColors != nullptr ? LODInfo.OverrideVertexColors : ColorBuffer;
 	}
 
+	const FGLTFMeshData* MeshData =
+        Builder.ExportOptions->bBakeMaterialInputs && Builder.ExportOptions->bBakeMaterialInputsUsingMeshData ?
+            MeshDataConverter.GetOrAdd(StaticMesh, StaticMeshComponent, LODIndex) : nullptr;
+
 	const int32 MaterialCount = StaticMesh->StaticMaterials.Num();
 	JsonMesh.Primitives.AddDefaulted(MaterialCount);
 
@@ -65,9 +69,7 @@ void FGLTFStaticMeshTask::Complete()
 			Material = UMaterial::GetDefaultMaterial(MD_Surface);
 		}
 
-		JsonPrimitive.Material = StaticMeshComponent != nullptr
-			? Builder.GetOrAddMaterial(Material, StaticMeshComponent, OverrideMaterials)
-			: Builder.GetOrAddMaterial(Material, StaticMesh, OverrideMaterials);
+		JsonPrimitive.Material =  Builder.GetOrAddMaterial(Material, MeshData, OverrideMaterials);
 	}
 }
 
@@ -92,6 +94,15 @@ void FGLTFSkeletalMeshTask::Complete()
 		SkinWeightBuffer = LODInfo.OverrideSkinWeights != nullptr ? LODInfo.OverrideSkinWeights : SkinWeightBuffer;
 	}
 
+	const FGLTFMeshData* MeshData =
+		Builder.ExportOptions->bBakeMaterialInputs && Builder.ExportOptions->bBakeMaterialInputsUsingMeshData ?
+			MeshDataConverter.GetOrAdd(SkeletalMesh, SkeletalMeshComponent, LODIndex) : nullptr;
+
+	if (SkeletalMeshComponent == nullptr)
+	{
+		MeshData = nullptr; // TODO: remove and add support for skeletal meshes in FGLTFMeshData
+	}
+
 	const uint16 MaterialCount = SkeletalMesh->Materials.Num();
 	JsonMesh.Primitives.AddDefaulted(MaterialCount);
 
@@ -146,8 +157,6 @@ void FGLTFSkeletalMeshTask::Complete()
 			Material = UMaterial::GetDefaultMaterial(MD_Surface);
 		}
 
-		JsonPrimitive.Material = SkeletalMeshComponent != nullptr
-			? Builder.GetOrAddMaterial(Material, SkeletalMeshComponent, OverrideMaterials)
-			: Builder.GetOrAddMaterial(Material, SkeletalMesh, OverrideMaterials);
+		JsonPrimitive.Material =  Builder.GetOrAddMaterial(Material, MeshData, OverrideMaterials);
 	}
 }
diff --git a/Engine/Plugins/Enterprise/GLTFExporter/MigratedFromGit/Source/GLTFExporter/Private/Tasks/GLTFMeshTasks.h b/Engine/Plugins/Enterprise/GLTFExporter/MigratedFromGit/Source/GLTFExporter/Private/Tasks/GLTFMeshTasks.h
--- a/Engine/Plugins/Enterprise/GLTFExporter/MigratedFromGit/Source/GLTFExporter/Private/Tasks/GLTFMeshTasks.h
+++ b/Engine/Plugins/Enterprise/GLTFExporter/MigratedFromGit/Source/GLTFExporter/Private/Tasks/GLTFMeshTasks.h
@@ -5,6 +5,7 @@
 #include "Tasks/GLTFTask.h"
 #include "Builders/GLTFConvertBuilder.h"
 #include "Converters/GLTFMeshSectionConverters.h"
+#include "Converters/GLTFMeshDataConverters.h"
 #include "Converters/GLTFMaterialArray.h"
 #include "Converters/GLTFNameUtility.h"
 #include "Engine.h"
@@ -13,10 +14,11 @@ class FGLTFStaticMeshTask : public FGLTFTask
 {
 public:
 
-	FGLTFStaticMeshTask(FGLTFConvertBuilder& Builder, FGLTFStaticMeshSectionConverter& MeshSectionConverter, const UStaticMesh* StaticMesh, const UStaticMeshComponent* StaticMeshComponent, FGLTFMaterialArray OverrideMaterials, FGLTFJsonMeshIndex MeshIndex)
+	FGLTFStaticMeshTask(FGLTFConvertBuilder& Builder, FGLTFStaticMeshSectionConverter& MeshSectionConverter, FGLTFStaticMeshDataConverter& MeshDataConverter, const UStaticMesh* StaticMesh, const UStaticMeshComponent* StaticMeshComponent, FGLTFMaterialArray OverrideMaterials, FGLTFJsonMeshIndex MeshIndex)
 		: FGLTFTask(EGLTFTaskPriority::Mesh)
 		, Builder(Builder)
 		, MeshSectionConverter(MeshSectionConverter)
+		, MeshDataConverter(MeshDataConverter)
 		, StaticMesh(StaticMesh)
 		, StaticMeshComponent(StaticMeshComponent)
 		, OverrideMaterials(OverrideMaterials)
@@ -35,6 +37,7 @@ private:
 
 	FGLTFConvertBuilder& Builder;
 	FGLTFStaticMeshSectionConverter& MeshSectionConverter;
+	FGLTFStaticMeshDataConverter& MeshDataConverter;
 	const UStaticMesh* StaticMesh;
 	const UStaticMeshComponent* StaticMeshComponent;
 	FGLTFMaterialArray OverrideMaterials;
@@ -45,10 +48,11 @@ class FGLTFSkeletalMeshTask : public FGLTFTask
 {
 public:
 
-	FGLTFSkeletalMeshTask(FGLTFConvertBuilder& Builder, FGLTFSkeletalMeshSectionConverter& MeshSectionConverter, const USkeletalMesh* SkeletalMesh, const USkeletalMeshComponent* SkeletalMeshComponent, FGLTFMaterialArray OverrideMaterials, FGLTFJsonMeshIndex MeshIndex)
+	FGLTFSkeletalMeshTask(FGLTFConvertBuilder& Builder, FGLTFSkeletalMeshSectionConverter& MeshSectionConverter, FGLTFSkeletalMeshDataConverter& MeshDataConverter, const USkeletalMesh* SkeletalMesh, const USkeletalMeshComponent* SkeletalMeshComponent, FGLTFMaterialArray OverrideMaterials, FGLTFJsonMeshIndex MeshIndex)
 		: FGLTFTask(EGLTFTaskPriority::Mesh)
 		, Builder(Builder)
 		, MeshSectionConverter(MeshSectionConverter)
+		, MeshDataConverter(MeshDataConverter)
 		, SkeletalMesh(SkeletalMesh)
 		, SkeletalMeshComponent(SkeletalMeshComponent)
 		, OverrideMaterials(OverrideMaterials)
@@ -67,6 +71,7 @@ private:
 
 	FGLTFConvertBuilder& Builder;
 	FGLTFSkeletalMeshSectionConverter& MeshSectionConverter;
+	FGLTFSkeletalMeshDataConverter& MeshDataConverter;
 	const USkeletalMesh* SkeletalMesh;
 	const USkeletalMeshComponent* SkeletalMeshComponent;
 	FGLTFMaterialArray OverrideMaterials;
