[GLTFExporter] refactor: re-add LODIndex parameter to mesh converters to make sure mesh component LOD settings are not lost even if mesh data is not used in material baking

commit 2b8b5ac92f8a5177760e4f887c1691859018b5dc
Date: Tue, 3 Nov 2020 13:48:24 +0100

    refactor: re-add LODIndex parameter to mesh converters to make sure mesh component LOD settings are not lost even if mesh data is not used in material baking
    (cherry picked from commit a1f46803754da3c74a1ad727bd9634f3abe2b904​)

#skipci

diff --git a/Engine/Plugins/Enterprise/GLTFExporter/MigratedFromGit/Source/GLTFExporter/Private/Builders/GLTFConvertBuilder.cpp b/Engine/Plugins/Enterprise/GLTFExporter/MigratedFromGit/Source/GLTFExporter/Private/Builders/GLTFConvertBuilder.cpp
--- a/Engine/Plugins/Enterprise/GLTFExporter/MigratedFromGit/Source/GLTFExporter/Private/Builders/GLTFConvertBuilder.cpp
+++ b/Engine/Plugins/Enterprise/GLTFExporter/MigratedFromGit/Source/GLTFExporter/Private/Builders/GLTFConvertBuilder.cpp
@@ -88,17 +88,17 @@ FGLTFJsonAccessorIndex FGLTFConvertBuilder::GetOrAddIndexAccessor(const FGLTFMes
 	return IndexBufferConverter.GetOrAdd(MeshSection);
 }
 
-FGLTFJsonMeshIndex FGLTFConvertBuilder::GetOrAddMesh(const UStaticMesh* StaticMesh, const FGLTFMaterialArray& OverrideMaterials)
+FGLTFJsonMeshIndex FGLTFConvertBuilder::GetOrAddMesh(const UStaticMesh* StaticMesh, const FGLTFMaterialArray& OverrideMaterials, int32 LODIndex)
 {
 	if (StaticMesh == nullptr)
 	{
 		return FGLTFJsonMeshIndex(INDEX_NONE);
 	}
 
-	return StaticMeshConverter.GetOrAdd(StaticMesh, nullptr, OverrideMaterials);
+	return StaticMeshConverter.GetOrAdd(StaticMesh, nullptr, OverrideMaterials, LODIndex);
 }
 
-FGLTFJsonMeshIndex FGLTFConvertBuilder::GetOrAddMesh(const UStaticMeshComponent* StaticMeshComponent, const FGLTFMaterialArray& OverrideMaterials)
+FGLTFJsonMeshIndex FGLTFConvertBuilder::GetOrAddMesh(const UStaticMeshComponent* StaticMeshComponent, const FGLTFMaterialArray& OverrideMaterials, int32 LODIndex)
 {
 	if (StaticMeshComponent == nullptr)
 	{
@@ -111,20 +111,20 @@ FGLTFJsonMeshIndex FGLTFConvertBuilder::GetOrAddMesh(const UStaticMeshComponent*
 		return FGLTFJsonMeshIndex(INDEX_NONE);
 	}
 
-	return StaticMeshConverter.GetOrAdd(StaticMesh, StaticMeshComponent, OverrideMaterials);
+	return StaticMeshConverter.GetOrAdd(StaticMesh, StaticMeshComponent, OverrideMaterials, LODIndex);
 }
 
-FGLTFJsonMeshIndex FGLTFConvertBuilder::GetOrAddMesh(const USkeletalMesh* SkeletalMesh, const FGLTFMaterialArray& OverrideMaterials)
+FGLTFJsonMeshIndex FGLTFConvertBuilder::GetOrAddMesh(const USkeletalMesh* SkeletalMesh, const FGLTFMaterialArray& OverrideMaterials, int32 LODIndex)
 {
 	if (SkeletalMesh == nullptr)
 	{
 		return FGLTFJsonMeshIndex(INDEX_NONE);
 	}
 
-	return SkeletalMeshConverter.GetOrAdd(SkeletalMesh, nullptr, OverrideMaterials);
+	return SkeletalMeshConverter.GetOrAdd(SkeletalMesh, nullptr, OverrideMaterials, LODIndex);
 }
 
-FGLTFJsonMeshIndex FGLTFConvertBuilder::GetOrAddMesh(const USkeletalMeshComponent* SkeletalMeshComponent, const FGLTFMaterialArray& OverrideMaterials)
+FGLTFJsonMeshIndex FGLTFConvertBuilder::GetOrAddMesh(const USkeletalMeshComponent* SkeletalMeshComponent, const FGLTFMaterialArray& OverrideMaterials, int32 LODIndex)
 {
 	if (SkeletalMeshComponent == nullptr)
 	{
@@ -137,7 +137,7 @@ FGLTFJsonMeshIndex FGLTFConvertBuilder::GetOrAddMesh(const USkeletalMeshComponen
 		return FGLTFJsonMeshIndex(INDEX_NONE);
 	}
 
-	return SkeletalMeshConverter.GetOrAdd(SkeletalMesh, SkeletalMeshComponent, OverrideMaterials);
+	return SkeletalMeshConverter.GetOrAdd(SkeletalMesh, SkeletalMeshComponent, OverrideMaterials, LODIndex);
 }
 
 FGLTFJsonMaterialIndex FGLTFConvertBuilder::GetOrAddMaterial(const UMaterialInterface* Material, const FGLTFMeshData* MeshData, const FGLTFMaterialArray& OverrideMaterials)
diff --git a/Engine/Plugins/Enterprise/GLTFExporter/MigratedFromGit/Source/GLTFExporter/Private/Builders/GLTFConvertBuilder.h b/Engine/Plugins/Enterprise/GLTFExporter/MigratedFromGit/Source/GLTFExporter/Private/Builders/GLTFConvertBuilder.h
--- a/Engine/Plugins/Enterprise/GLTFExporter/MigratedFromGit/Source/GLTFExporter/Private/Builders/GLTFConvertBuilder.h
+++ b/Engine/Plugins/Enterprise/GLTFExporter/MigratedFromGit/Source/GLTFExporter/Private/Builders/GLTFConvertBuilder.h
@@ -39,10 +39,10 @@ public:
 	FGLTFJsonAccessorIndex GetOrAddWeightAccessor(const FGLTFMeshSection* MeshSection, const FSkinWeightVertexBuffer* VertexBuffer, int32 InfluenceOffset);
 	FGLTFJsonAccessorIndex GetOrAddIndexAccessor(const FGLTFMeshSection* MeshSection);
 
-	FGLTFJsonMeshIndex GetOrAddMesh(const UStaticMesh* StaticMesh, const FGLTFMaterialArray& OverrideMaterials = {});
-	FGLTFJsonMeshIndex GetOrAddMesh(const UStaticMeshComponent* StaticMeshComponent, const FGLTFMaterialArray& OverrideMaterials = {});
-	FGLTFJsonMeshIndex GetOrAddMesh(const USkeletalMesh* SkeletalMesh, const FGLTFMaterialArray& OverrideMaterials = {});
-	FGLTFJsonMeshIndex GetOrAddMesh(const USkeletalMeshComponent* SkeletalMeshComponent, const FGLTFMaterialArray& OverrideMaterials = {});
+	FGLTFJsonMeshIndex GetOrAddMesh(const UStaticMesh* StaticMesh, const FGLTFMaterialArray& OverrideMaterials = {}, int32 LODIndex = -1);
+	FGLTFJsonMeshIndex GetOrAddMesh(const UStaticMeshComponent* StaticMeshComponent, const FGLTFMaterialArray& OverrideMaterials = {}, int32 LODIndex = -1);
+	FGLTFJsonMeshIndex GetOrAddMesh(const USkeletalMesh* SkeletalMesh, const FGLTFMaterialArray& OverrideMaterials = {}, int32 LODIndex = -1);
+	FGLTFJsonMeshIndex GetOrAddMesh(const USkeletalMeshComponent* SkeletalMeshComponent, const FGLTFMaterialArray& OverrideMaterials = {}, int32 LODIndex = -1);
 
 	FGLTFJsonMaterialIndex GetOrAddMaterial(const UMaterialInterface* Material, const FGLTFMeshData* MeshData = nullptr, const FGLTFMaterialArray& OverrideMaterials = {});
 
diff --git a/Engine/Plugins/Enterprise/GLTFExporter/MigratedFromGit/Source/GLTFExporter/Private/Converters/GLTFMeshConverters.cpp b/Engine/Plugins/Enterprise/GLTFExporter/MigratedFromGit/Source/GLTFExporter/Private/Converters/GLTFMeshConverters.cpp
--- a/Engine/Plugins/Enterprise/GLTFExporter/MigratedFromGit/Source/GLTFExporter/Private/Converters/GLTFMeshConverters.cpp
+++ b/Engine/Plugins/Enterprise/GLTFExporter/MigratedFromGit/Source/GLTFExporter/Private/Converters/GLTFMeshConverters.cpp
@@ -1,16 +1,26 @@
 // Copyright Epic Games, Inc. All Rights Reserved.
 
 #include "Converters/GLTFMeshConverters.h"
+#include "Converters/GLTFMeshUtility.h"
 #include "Converters/GLTFMaterialUtility.h"
 #include "Builders/GLTFConvertBuilder.h"
 #include "Tasks/GLTFMeshTasks.h"
 
-void FGLTFStaticMeshConverter::Sanitize(const UStaticMesh*& StaticMesh, const UStaticMeshComponent*& StaticMeshComponent, FGLTFMaterialArray& OverrideMaterials)
+void FGLTFStaticMeshConverter::Sanitize(const UStaticMesh*& StaticMesh, const UStaticMeshComponent*& StaticMeshComponent, FGLTFMaterialArray& OverrideMaterials, int32& LODIndex)
 {
 	const TArray<UMaterialInterface*> Materials = StaticMeshComponent != nullptr
 		? OverrideMaterials.GetOverrides(StaticMeshComponent->GetMaterials())
 		: OverrideMaterials.GetOverrides(StaticMesh->StaticMaterials);
 
+	if (LODIndex < 0)
+	{
+		LODIndex = FGLTFMeshUtility::GetLOD(StaticMesh, StaticMeshComponent, Builder.ExportOptions->DefaultLevelOfDetail);
+	}
+	else
+	{
+		LODIndex = FMath::Min(LODIndex, FGLTFMeshUtility::GetMaximumLOD(StaticMesh));
+	}
+
 	if (StaticMeshComponent != nullptr)
 	{
 		// Only use the component if it's needed for baking, since we would
@@ -36,19 +46,28 @@ void FGLTFStaticMeshConverter::Sanitize(const UStaticMesh*& StaticMesh, const US
 	}
 }
 
-FGLTFJsonMeshIndex FGLTFStaticMeshConverter::Convert(const UStaticMesh* StaticMesh, const UStaticMeshComponent* StaticMeshComponent, FGLTFMaterialArray OverrideMaterials)
+FGLTFJsonMeshIndex FGLTFStaticMeshConverter::Convert(const UStaticMesh* StaticMesh,  const UStaticMeshComponent* StaticMeshComponent, FGLTFMaterialArray OverrideMaterials, int32 LODIndex)
 {
 	const FGLTFJsonMeshIndex MeshIndex = Builder.AddMesh();
-	Builder.SetupTask<FGLTFStaticMeshTask>(Builder, MeshSectionConverter, MeshDataConverter, StaticMesh, StaticMeshComponent, OverrideMaterials, MeshIndex);
+	Builder.SetupTask<FGLTFStaticMeshTask>(Builder, MeshSectionConverter, MeshDataConverter, StaticMesh, StaticMeshComponent, OverrideMaterials, LODIndex, MeshIndex);
 	return MeshIndex;
 }
 
-void FGLTFSkeletalMeshConverter::Sanitize(const USkeletalMesh*& SkeletalMesh, const USkeletalMeshComponent*& SkeletalMeshComponent, FGLTFMaterialArray& OverrideMaterials)
+void FGLTFSkeletalMeshConverter::Sanitize(const USkeletalMesh*& SkeletalMesh, const USkeletalMeshComponent*& SkeletalMeshComponent, FGLTFMaterialArray& OverrideMaterials, int32& LODIndex)
 {
 	const TArray<UMaterialInterface*> Materials = SkeletalMeshComponent != nullptr
 		? OverrideMaterials.GetOverrides(SkeletalMeshComponent->GetMaterials())
 		: OverrideMaterials.GetOverrides(SkeletalMesh->Materials);
 
+	if (LODIndex < 0)
+	{
+		LODIndex = FGLTFMeshUtility::GetLOD(SkeletalMesh, SkeletalMeshComponent, Builder.ExportOptions->DefaultLevelOfDetail);
+	}
+	else
+	{
+		LODIndex = FMath::Min(LODIndex, FGLTFMeshUtility::GetMaximumLOD(SkeletalMesh));
+	}
+
 	if (SkeletalMeshComponent != nullptr)
 	{
 		// Only use the component if it's needed for baking, since we would
@@ -74,9 +93,9 @@ void FGLTFSkeletalMeshConverter::Sanitize(const USkeletalMesh*& SkeletalMesh, co
 	}
 }
 
-FGLTFJsonMeshIndex FGLTFSkeletalMeshConverter::Convert(const USkeletalMesh* SkeletalMesh, const USkeletalMeshComponent* SkeletalMeshComponent, FGLTFMaterialArray OverrideMaterials)
+FGLTFJsonMeshIndex FGLTFSkeletalMeshConverter::Convert(const USkeletalMesh* SkeletalMesh, const USkeletalMeshComponent* SkeletalMeshComponent, FGLTFMaterialArray OverrideMaterials, int32 LODIndex)
 {
 	const FGLTFJsonMeshIndex MeshIndex = Builder.AddMesh();
-	Builder.SetupTask<FGLTFSkeletalMeshTask>(Builder, MeshSectionConverter, MeshDataConverter, SkeletalMesh, SkeletalMeshComponent, OverrideMaterials, MeshIndex);
+	Builder.SetupTask<FGLTFSkeletalMeshTask>(Builder, MeshSectionConverter, MeshDataConverter, SkeletalMesh, SkeletalMeshComponent, OverrideMaterials, LODIndex, MeshIndex);
 	return MeshIndex;
 }
diff --git a/Engine/Plugins/Enterprise/GLTFExporter/MigratedFromGit/Source/GLTFExporter/Private/Converters/GLTFMeshConverters.h b/Engine/Plugins/Enterprise/GLTFExporter/MigratedFromGit/Source/GLTFExporter/Private/Converters/GLTFMeshConverters.h
--- a/Engine/Plugins/Enterprise/GLTFExporter/MigratedFromGit/Source/GLTFExporter/Private/Converters/GLTFMeshConverters.h
+++ b/Engine/Plugins/Enterprise/GLTFExporter/MigratedFromGit/Source/GLTFExporter/Private/Converters/GLTFMeshConverters.h
@@ -16,25 +16,25 @@ class TGLTFMeshConverter : public FGLTFBuilderContext, public TGLTFConverter<FGL
 	using FGLTFBuilderContext::FGLTFBuilderContext;
 };
 
-class FGLTFStaticMeshConverter final : public TGLTFMeshConverter<const UStaticMesh*, const UStaticMeshComponent*, FGLTFMaterialArray>
+class FGLTFStaticMeshConverter final : public TGLTFMeshConverter<const UStaticMesh*, const UStaticMeshComponent*, FGLTFMaterialArray, int32>
 {
 	using TGLTFMeshConverter::TGLTFMeshConverter;
 
-	virtual void Sanitize(const UStaticMesh*& StaticMesh, const UStaticMeshComponent*& StaticMeshComponent, FGLTFMaterialArray& OverrideMaterials) override;
+	virtual void Sanitize(const UStaticMesh*& StaticMesh, const UStaticMeshComponent*& StaticMeshComponent, FGLTFMaterialArray& OverrideMaterials, int32& LODIndex) override;
 
-	virtual FGLTFJsonMeshIndex Convert(const UStaticMesh* StaticMesh, const UStaticMeshComponent* StaticMeshComponent, FGLTFMaterialArray OverrideMaterials) override;
+	virtual FGLTFJsonMeshIndex Convert(const UStaticMesh* StaticMesh, const UStaticMeshComponent* StaticMeshComponent, FGLTFMaterialArray OverrideMaterials, int32 LODIndex) override;
 
 	FGLTFStaticMeshSectionConverter MeshSectionConverter;
 	FGLTFStaticMeshDataConverter MeshDataConverter;
 };
 
-class FGLTFSkeletalMeshConverter final : public TGLTFMeshConverter<const USkeletalMesh*, const USkeletalMeshComponent*, FGLTFMaterialArray>
+class FGLTFSkeletalMeshConverter final : public TGLTFMeshConverter<const USkeletalMesh*, const USkeletalMeshComponent*, FGLTFMaterialArray, int32>
 {
 	using TGLTFMeshConverter::TGLTFMeshConverter;
 
-	virtual void Sanitize(const USkeletalMesh*& SkeletalMesh, const USkeletalMeshComponent*& SkeletalMeshComponent, FGLTFMaterialArray& OverrideMaterials) override;
+	virtual void Sanitize(const USkeletalMesh*& SkeletalMesh, const USkeletalMeshComponent*& SkeletalMeshComponent, FGLTFMaterialArray& OverrideMaterials, int32& LODIndex) override;
 
-	virtual FGLTFJsonMeshIndex Convert(const USkeletalMesh* SkeletalMesh, const USkeletalMeshComponent* SkeletalMeshComponent, FGLTFMaterialArray OverrideMaterials) override;
+	virtual FGLTFJsonMeshIndex Convert(const USkeletalMesh* SkeletalMesh, const USkeletalMeshComponent* SkeletalMeshComponent, FGLTFMaterialArray OverrideMaterials, int32 LODIndex) override;
 
 	FGLTFSkeletalMeshSectionConverter MeshSectionConverter;
 	FGLTFSkeletalMeshDataConverter MeshDataConverter;
diff --git a/Engine/Plugins/Enterprise/GLTFExporter/MigratedFromGit/Source/GLTFExporter/Private/Tasks/GLTFMeshTasks.cpp b/Engine/Plugins/Enterprise/GLTFExporter/MigratedFromGit/Source/GLTFExporter/Private/Tasks/GLTFMeshTasks.cpp
--- a/Engine/Plugins/Enterprise/GLTFExporter/MigratedFromGit/Source/GLTFExporter/Private/Tasks/GLTFMeshTasks.cpp
+++ b/Engine/Plugins/Enterprise/GLTFExporter/MigratedFromGit/Source/GLTFExporter/Private/Tasks/GLTFMeshTasks.cpp
@@ -2,7 +2,6 @@
 
 #include "Tasks/GLTFMeshTasks.h"
 #include "Converters/GLTFConverterUtility.h"
-#include "Converters/GLTFMeshUtility.h"
 #include "Builders/GLTFConvertBuilder.h"
 #include "Rendering/SkeletalMeshRenderData.h"
 
@@ -11,9 +10,7 @@ void FGLTFStaticMeshTask::Complete()
 	FGLTFJsonMesh& JsonMesh = Builder.GetMesh(MeshIndex);
 	JsonMesh.Name = StaticMeshComponent != nullptr ? FGLTFNameUtility::GetName(StaticMeshComponent) : StaticMesh->GetName();
 
-	const int32 LODIndex = FGLTFMeshUtility::GetLOD(StaticMesh, StaticMeshComponent, Builder.ExportOptions->DefaultLevelOfDetail);
 	const FStaticMeshLODResources& MeshLOD = StaticMesh->GetLODForExport(LODIndex);
-
 	const FPositionVertexBuffer* PositionBuffer = &MeshLOD.VertexBuffers.PositionVertexBuffer;
 	const FStaticMeshVertexBuffer* VertexBuffer = &MeshLOD.VertexBuffers.StaticMeshVertexBuffer;
 	const FColorVertexBuffer* ColorBuffer = &MeshLOD.VertexBuffers.ColorVertexBuffer;
@@ -78,7 +75,6 @@ void FGLTFSkeletalMeshTask::Complete()
 	FGLTFJsonMesh& JsonMesh = Builder.GetMesh(MeshIndex);
 	JsonMesh.Name = SkeletalMeshComponent != nullptr ? FGLTFNameUtility::GetName(SkeletalMeshComponent) : SkeletalMesh->GetName();
 
-	const int32 LODIndex = FGLTFMeshUtility::GetLOD(SkeletalMesh, SkeletalMeshComponent, Builder.ExportOptions->DefaultLevelOfDetail);
 	const FSkeletalMeshRenderData* RenderData = SkeletalMesh->GetResourceForRendering();
 	const FSkeletalMeshLODRenderData& MeshLOD = RenderData->LODRenderData[LODIndex];
 
diff --git a/Engine/Plugins/Enterprise/GLTFExporter/MigratedFromGit/Source/GLTFExporter/Private/Tasks/GLTFMeshTasks.h b/Engine/Plugins/Enterprise/GLTFExporter/MigratedFromGit/Source/GLTFExporter/Private/Tasks/GLTFMeshTasks.h
--- a/Engine/Plugins/Enterprise/GLTFExporter/MigratedFromGit/Source/GLTFExporter/Private/Tasks/GLTFMeshTasks.h
+++ b/Engine/Plugins/Enterprise/GLTFExporter/MigratedFromGit/Source/GLTFExporter/Private/Tasks/GLTFMeshTasks.h
@@ -14,7 +14,7 @@ class FGLTFStaticMeshTask : public FGLTFTask
 {
 public:
 
-	FGLTFStaticMeshTask(FGLTFConvertBuilder& Builder, FGLTFStaticMeshSectionConverter& MeshSectionConverter, FGLTFStaticMeshDataConverter& MeshDataConverter, const UStaticMesh* StaticMesh, const UStaticMeshComponent* StaticMeshComponent, FGLTFMaterialArray OverrideMaterials, FGLTFJsonMeshIndex MeshIndex)
+	FGLTFStaticMeshTask(FGLTFConvertBuilder& Builder, FGLTFStaticMeshSectionConverter& MeshSectionConverter, FGLTFStaticMeshDataConverter& MeshDataConverter, const UStaticMesh* StaticMesh, const UStaticMeshComponent* StaticMeshComponent, FGLTFMaterialArray OverrideMaterials, int32 LODIndex, FGLTFJsonMeshIndex MeshIndex)
 		: FGLTFTask(EGLTFTaskPriority::Mesh)
 		, Builder(Builder)
 		, MeshSectionConverter(MeshSectionConverter)
@@ -22,6 +22,7 @@ public:
 		, StaticMesh(StaticMesh)
 		, StaticMeshComponent(StaticMeshComponent)
 		, OverrideMaterials(OverrideMaterials)
+		, LODIndex(LODIndex)
         , MeshIndex(MeshIndex)
 	{
 	}
@@ -40,7 +41,8 @@ private:
 	FGLTFStaticMeshDataConverter& MeshDataConverter;
 	const UStaticMesh* StaticMesh;
 	const UStaticMeshComponent* StaticMeshComponent;
-	FGLTFMaterialArray OverrideMaterials;
+	const FGLTFMaterialArray OverrideMaterials;
+	const int32 LODIndex;
 	const FGLTFJsonMeshIndex MeshIndex;
 };
 
@@ -48,7 +50,7 @@ class FGLTFSkeletalMeshTask : public FGLTFTask
 {
 public:
 
-	FGLTFSkeletalMeshTask(FGLTFConvertBuilder& Builder, FGLTFSkeletalMeshSectionConverter& MeshSectionConverter, FGLTFSkeletalMeshDataConverter& MeshDataConverter, const USkeletalMesh* SkeletalMesh, const USkeletalMeshComponent* SkeletalMeshComponent, FGLTFMaterialArray OverrideMaterials, FGLTFJsonMeshIndex MeshIndex)
+	FGLTFSkeletalMeshTask(FGLTFConvertBuilder& Builder, FGLTFSkeletalMeshSectionConverter& MeshSectionConverter, FGLTFSkeletalMeshDataConverter& MeshDataConverter, const USkeletalMesh* SkeletalMesh, const USkeletalMeshComponent* SkeletalMeshComponent, FGLTFMaterialArray OverrideMaterials, int32 LODIndex, FGLTFJsonMeshIndex MeshIndex)
 		: FGLTFTask(EGLTFTaskPriority::Mesh)
 		, Builder(Builder)
 		, MeshSectionConverter(MeshSectionConverter)
@@ -56,6 +58,7 @@ public:
 		, SkeletalMesh(SkeletalMesh)
 		, SkeletalMeshComponent(SkeletalMeshComponent)
 		, OverrideMaterials(OverrideMaterials)
+		, LODIndex(LODIndex)
 		, MeshIndex(MeshIndex)
 	{
 	}
@@ -74,6 +77,7 @@ private:
 	FGLTFSkeletalMeshDataConverter& MeshDataConverter;
 	const USkeletalMesh* SkeletalMesh;
 	const USkeletalMeshComponent* SkeletalMeshComponent;
-	FGLTFMaterialArray OverrideMaterials;
+	const FGLTFMaterialArray OverrideMaterials;
+	const int32 LODIndex;
 	const FGLTFJsonMeshIndex MeshIndex;
 };
