[GLTFExporter] fix: material indices is actually section indices when material baking using mesh data

commit a824714b9fe53eb2953f8c3c4564f3b6076c6960
Date: Tue, 3 Nov 2020 19:46:57 +0100

    fix: material indices is actually section indices when material baking using mesh data
    (cherry picked from commit 96fa477ff3979edf257de2f83de73543b423bf65​)

#skipci

diff --git a/Engine/Plugins/Enterprise/GLTFExporter/MigratedFromGit/Source/GLTFExporter/Private/Converters/GLTFMaterialUtility.cpp b/Engine/Plugins/Enterprise/GLTFExporter/MigratedFromGit/Source/GLTFExporter/Private/Converters/GLTFMaterialUtility.cpp
--- a/Engine/Plugins/Enterprise/GLTFExporter/MigratedFromGit/Source/GLTFExporter/Private/Converters/GLTFMaterialUtility.cpp
+++ b/Engine/Plugins/Enterprise/GLTFExporter/MigratedFromGit/Source/GLTFExporter/Private/Converters/GLTFMaterialUtility.cpp
@@ -241,7 +241,7 @@ bool FGLTFMaterialUtility::CombineTextures(TArray<FColor>& OutPixels, const TArr
 	return bReadSuccessful;
 }
 
-FGLTFPropertyBakeOutput FGLTFMaterialUtility::BakeMaterialProperty(const FIntPoint& OutputSize, EMaterialProperty Property, const UMaterialInterface* Material, int32 TexCoord, const FMeshDescription* MeshDescription, const TArray<int32>& MeshMaterialIndices, bool bCopyAlphaFromRedChannel)
+FGLTFPropertyBakeOutput FGLTFMaterialUtility::BakeMaterialProperty(const FIntPoint& OutputSize, EMaterialProperty Property, const UMaterialInterface* Material, int32 TexCoord, const FMeshDescription* MeshDescription, const TArray<int32>& MeshSectionIndices, bool bCopyAlphaFromRedChannel)
 {
 	EMaterialProperty HackProperty = MP_MAX;
 	TTuple<int32, UMaterialExpression*> HackPrevState;
@@ -274,7 +274,7 @@ FGLTFPropertyBakeOutput FGLTFMaterialUtility::BakeMaterialProperty(const FIntPoi
 	MeshSet.TextureCoordinateBox = { { 0.0f, 0.0f }, { 1.0f, 1.0f } };
 	MeshSet.TextureCoordinateIndex = TexCoord;
 	MeshSet.RawMeshDescription = const_cast<FMeshDescription*>(MeshDescription);
-	MeshSet.MaterialIndices = MeshMaterialIndices;
+	MeshSet.MaterialIndices = MeshSectionIndices; // NOTE: MaterialIndices is actually section indices
 
 	FMaterialData MatSet;
 	MatSet.Material = const_cast<UMaterialInterface*>(Material);
diff --git a/Engine/Plugins/Enterprise/GLTFExporter/MigratedFromGit/Source/GLTFExporter/Private/Converters/GLTFMaterialUtility.h b/Engine/Plugins/Enterprise/GLTFExporter/MigratedFromGit/Source/GLTFExporter/Private/Converters/GLTFMaterialUtility.h
--- a/Engine/Plugins/Enterprise/GLTFExporter/MigratedFromGit/Source/GLTFExporter/Private/Converters/GLTFMaterialUtility.h
+++ b/Engine/Plugins/Enterprise/GLTFExporter/MigratedFromGit/Source/GLTFExporter/Private/Converters/GLTFMaterialUtility.h
@@ -53,7 +53,7 @@ struct FGLTFMaterialUtility
 
 	static UTexture2D* CreateTransientTexture(const FGLTFPropertyBakeOutput& PropertyBakeOutput, bool bUseSRGB = false);
 	static bool CombineTextures(TArray<FColor>& OutPixels, const TArray<FGLTFTextureCombineSource>& Sources, const FIntPoint& OutputSize, EPixelFormat OutputPixelFormat);
-	static FGLTFPropertyBakeOutput BakeMaterialProperty(const FIntPoint& OutputSize, EMaterialProperty Property, const UMaterialInterface* Material, int32 TexCoord, const FMeshDescription* MeshDescription = nullptr, const TArray<int32>& MeshMaterialIndices = {}, bool bCopyAlphaFromRedChannel = false);
+	static FGLTFPropertyBakeOutput BakeMaterialProperty(const FIntPoint& OutputSize, EMaterialProperty Property, const UMaterialInterface* Material, int32 TexCoord, const FMeshDescription* MeshDescription = nullptr, const TArray<int32>& MeshSectionIndices = {}, bool bCopyAlphaFromRedChannel = false);
 
 	static FGLTFJsonTextureIndex AddCombinedTexture(FGLTFConvertBuilder& Builder, const TArray<FGLTFTextureCombineSource>& CombineSources, const FIntPoint& TextureSize, const FString& TextureName, EGLTFJsonTextureFilter MinFilter, EGLTFJsonTextureFilter MagFilter, EGLTFJsonTextureWrap WrapS, EGLTFJsonTextureWrap WrapT);
 	static FGLTFJsonTextureIndex AddTexture(FGLTFConvertBuilder& Builder, const TArray<FColor>& Pixels, const FIntPoint& TextureSize, const FString& TextureName, EGLTFJsonTextureFilter MinFilter, EGLTFJsonTextureFilter MagFilter, EGLTFJsonTextureWrap WrapS, EGLTFJsonTextureWrap WrapT);
diff --git a/Engine/Plugins/Enterprise/GLTFExporter/MigratedFromGit/Source/GLTFExporter/Private/Tasks/GLTFMaterialTasks.cpp b/Engine/Plugins/Enterprise/GLTFExporter/MigratedFromGit/Source/GLTFExporter/Private/Tasks/GLTFMaterialTasks.cpp
--- a/Engine/Plugins/Enterprise/GLTFExporter/MigratedFromGit/Source/GLTFExporter/Private/Tasks/GLTFMaterialTasks.cpp
+++ b/Engine/Plugins/Enterprise/GLTFExporter/MigratedFromGit/Source/GLTFExporter/Private/Tasks/GLTFMaterialTasks.cpp
@@ -1353,7 +1353,7 @@ FGLTFPropertyBakeOutput FGLTFMaterialTask::BakeMaterialProperty(EMaterialPropert
 	const FIntPoint DefaultTextureSize = Builder.GetDefaultMaterialBakeSize();
 	const FIntPoint TextureSize = PreferredTextureSize != nullptr ? *PreferredTextureSize : DefaultTextureSize;
 
-	TArray<int32> MeshMaterialIndices;
+	TArray<int32> MeshSectionIndices;
 
 	if (MeshData != nullptr)
 	{
@@ -1364,20 +1364,20 @@ FGLTFPropertyBakeOutput FGLTFMaterialTask::BakeMaterialProperty(EMaterialPropert
 		for (int32 SectionIndex = 0; SectionIndex < Sections.Num(); ++SectionIndex)
 		{
 			const FSectionInfo& Section = Sections[SectionIndex];
-			const UMaterialInterface* OverrideMaterial = OverrideMaterials.IsValidIndex(SectionIndex) ? OverrideMaterials[SectionIndex] : nullptr;
+			const UMaterialInterface* OverrideMaterial = OverrideMaterials.IsValidIndex(Section.MaterialIndex) ? OverrideMaterials[Section.MaterialIndex] : nullptr;
 			const UMaterialInterface* SectionMaterial = OverrideMaterial != nullptr ? OverrideMaterial : Section.Material;
 
 			if (SectionMaterial == Material)
 			{
-				MeshMaterialIndices.Add(Section.MaterialIndex);
+				MeshSectionIndices.Add(SectionIndex);
 			}
 		}
 
-		if (MeshMaterialIndices.Num() == 0)
+		if (MeshSectionIndices.Num() == 0)
 		{
 			// Fall back to using the first material index,
 			// just like when baking from within the editor
-			MeshMaterialIndices.Add(0);
+			MeshSectionIndices.Add(0);
 		}
 	}
 
@@ -1389,7 +1389,7 @@ FGLTFPropertyBakeOutput FGLTFMaterialTask::BakeMaterialProperty(EMaterialPropert
 		Material,
 		OutTexCoord,
 		MeshData != nullptr ? &MeshData->Description : nullptr,
-		MeshMaterialIndices,
+		MeshSectionIndices,
 		bCopyAlphaFromRedChannel);
 
 	if (!PropertyBakeOutput.bIsConstant && TexCoords.Num() == 0)
