[GLTFExporter] refactor: simplify material array parameter in converters to contain fully resolved materials

commit 676c7feae14952b51c1cfc04e736dc81358587ed
Date: Tue, 3 Nov 2020 19:51:57 +0100

    refactor: simplify material array parameter in converters to contain fully resolved materials
    (cherry picked from commit 2ea2e5f97e67667e02c7168c0f3a7dc7d9ff2a4b​)

#skipci

diff --git a/Engine/Plugins/Enterprise/GLTFExporter/MigratedFromGit/Source/GLTFExporter/Private/Builders/GLTFConvertBuilder.cpp b/Engine/Plugins/Enterprise/GLTFExporter/MigratedFromGit/Source/GLTFExporter/Private/Builders/GLTFConvertBuilder.cpp
--- a/Engine/Plugins/Enterprise/GLTFExporter/MigratedFromGit/Source/GLTFExporter/Private/Builders/GLTFConvertBuilder.cpp
+++ b/Engine/Plugins/Enterprise/GLTFExporter/MigratedFromGit/Source/GLTFExporter/Private/Builders/GLTFConvertBuilder.cpp
@@ -88,17 +88,17 @@ FGLTFJsonAccessorIndex FGLTFConvertBuilder::GetOrAddIndexAccessor(const FGLTFMes
 	return IndexBufferConverter.GetOrAdd(MeshSection);
 }
 
-FGLTFJsonMeshIndex FGLTFConvertBuilder::GetOrAddMesh(const UStaticMesh* StaticMesh, const FGLTFMaterialArray& OverrideMaterials, int32 LODIndex)
+FGLTFJsonMeshIndex FGLTFConvertBuilder::GetOrAddMesh(const UStaticMesh* StaticMesh, const FGLTFMaterialArray& Materials, int32 LODIndex)
 {
 	if (StaticMesh == nullptr)
 	{
 		return FGLTFJsonMeshIndex(INDEX_NONE);
 	}
 
-	return StaticMeshConverter.GetOrAdd(StaticMesh, nullptr, OverrideMaterials, LODIndex);
+	return StaticMeshConverter.GetOrAdd(StaticMesh, nullptr, Materials, LODIndex);
 }
 
-FGLTFJsonMeshIndex FGLTFConvertBuilder::GetOrAddMesh(const UStaticMeshComponent* StaticMeshComponent, const FGLTFMaterialArray& OverrideMaterials, int32 LODIndex)
+FGLTFJsonMeshIndex FGLTFConvertBuilder::GetOrAddMesh(const UStaticMeshComponent* StaticMeshComponent, const FGLTFMaterialArray& Materials, int32 LODIndex)
 {
 	if (StaticMeshComponent == nullptr)
 	{
@@ -111,20 +111,20 @@ FGLTFJsonMeshIndex FGLTFConvertBuilder::GetOrAddMesh(const UStaticMeshComponent*
 		return FGLTFJsonMeshIndex(INDEX_NONE);
 	}
 
-	return StaticMeshConverter.GetOrAdd(StaticMesh, StaticMeshComponent, OverrideMaterials, LODIndex);
+	return StaticMeshConverter.GetOrAdd(StaticMesh, StaticMeshComponent, Materials, LODIndex);
 }
 
-FGLTFJsonMeshIndex FGLTFConvertBuilder::GetOrAddMesh(const USkeletalMesh* SkeletalMesh, const FGLTFMaterialArray& OverrideMaterials, int32 LODIndex)
+FGLTFJsonMeshIndex FGLTFConvertBuilder::GetOrAddMesh(const USkeletalMesh* SkeletalMesh, const FGLTFMaterialArray& Materials, int32 LODIndex)
 {
 	if (SkeletalMesh == nullptr)
 	{
 		return FGLTFJsonMeshIndex(INDEX_NONE);
 	}
 
-	return SkeletalMeshConverter.GetOrAdd(SkeletalMesh, nullptr, OverrideMaterials, LODIndex);
+	return SkeletalMeshConverter.GetOrAdd(SkeletalMesh, nullptr, Materials, LODIndex);
 }
 
-FGLTFJsonMeshIndex FGLTFConvertBuilder::GetOrAddMesh(const USkeletalMeshComponent* SkeletalMeshComponent, const FGLTFMaterialArray& OverrideMaterials, int32 LODIndex)
+FGLTFJsonMeshIndex FGLTFConvertBuilder::GetOrAddMesh(const USkeletalMeshComponent* SkeletalMeshComponent, const FGLTFMaterialArray& Materials, int32 LODIndex)
 {
 	if (SkeletalMeshComponent == nullptr)
 	{
@@ -137,17 +137,17 @@ FGLTFJsonMeshIndex FGLTFConvertBuilder::GetOrAddMesh(const USkeletalMeshComponen
 		return FGLTFJsonMeshIndex(INDEX_NONE);
 	}
 
-	return SkeletalMeshConverter.GetOrAdd(SkeletalMesh, SkeletalMeshComponent, OverrideMaterials, LODIndex);
+	return SkeletalMeshConverter.GetOrAdd(SkeletalMesh, SkeletalMeshComponent, Materials, LODIndex);
 }
 
-FGLTFJsonMaterialIndex FGLTFConvertBuilder::GetOrAddMaterial(const UMaterialInterface* Material, const FGLTFMeshData* MeshData, const FGLTFMaterialArray& OverrideMaterials)
+FGLTFJsonMaterialIndex FGLTFConvertBuilder::GetOrAddMaterial(const UMaterialInterface* Material, const FGLTFMeshData* MeshData, const FGLTFMaterialArray& Materials)
 {
 	if (Material == nullptr)
 	{
 		return FGLTFJsonMaterialIndex(INDEX_NONE);
 	}
 
-	return MaterialConverter.GetOrAdd(Material, MeshData, OverrideMaterials);
+	return MaterialConverter.GetOrAdd(Material, MeshData, Materials);
 }
 
 FGLTFJsonSamplerIndex FGLTFConvertBuilder::GetOrAddSampler(const UTexture* Texture)
diff --git a/Engine/Plugins/Enterprise/GLTFExporter/MigratedFromGit/Source/GLTFExporter/Private/Builders/GLTFConvertBuilder.h b/Engine/Plugins/Enterprise/GLTFExporter/MigratedFromGit/Source/GLTFExporter/Private/Builders/GLTFConvertBuilder.h
--- a/Engine/Plugins/Enterprise/GLTFExporter/MigratedFromGit/Source/GLTFExporter/Private/Builders/GLTFConvertBuilder.h
+++ b/Engine/Plugins/Enterprise/GLTFExporter/MigratedFromGit/Source/GLTFExporter/Private/Builders/GLTFConvertBuilder.h
@@ -39,12 +39,12 @@ public:
 	FGLTFJsonAccessorIndex GetOrAddWeightAccessor(const FGLTFMeshSection* MeshSection, const FSkinWeightVertexBuffer* VertexBuffer, int32 InfluenceOffset);
 	FGLTFJsonAccessorIndex GetOrAddIndexAccessor(const FGLTFMeshSection* MeshSection);
 
-	FGLTFJsonMeshIndex GetOrAddMesh(const UStaticMesh* StaticMesh, const FGLTFMaterialArray& OverrideMaterials = {}, int32 LODIndex = -1);
-	FGLTFJsonMeshIndex GetOrAddMesh(const UStaticMeshComponent* StaticMeshComponent, const FGLTFMaterialArray& OverrideMaterials = {}, int32 LODIndex = -1);
-	FGLTFJsonMeshIndex GetOrAddMesh(const USkeletalMesh* SkeletalMesh, const FGLTFMaterialArray& OverrideMaterials = {}, int32 LODIndex = -1);
-	FGLTFJsonMeshIndex GetOrAddMesh(const USkeletalMeshComponent* SkeletalMeshComponent, const FGLTFMaterialArray& OverrideMaterials = {}, int32 LODIndex = -1);
+	FGLTFJsonMeshIndex GetOrAddMesh(const UStaticMesh* StaticMesh, const FGLTFMaterialArray& Materials = {}, int32 LODIndex = -1);
+	FGLTFJsonMeshIndex GetOrAddMesh(const UStaticMeshComponent* StaticMeshComponent, const FGLTFMaterialArray& Materials = {}, int32 LODIndex = -1);
+	FGLTFJsonMeshIndex GetOrAddMesh(const USkeletalMesh* SkeletalMesh, const FGLTFMaterialArray& Materials = {}, int32 LODIndex = -1);
+	FGLTFJsonMeshIndex GetOrAddMesh(const USkeletalMeshComponent* SkeletalMeshComponent, const FGLTFMaterialArray& Materials = {}, int32 LODIndex = -1);
 
-	FGLTFJsonMaterialIndex GetOrAddMaterial(const UMaterialInterface* Material, const FGLTFMeshData* MeshData = nullptr, const FGLTFMaterialArray& OverrideMaterials = {});
+	FGLTFJsonMaterialIndex GetOrAddMaterial(const UMaterialInterface* Material, const FGLTFMeshData* MeshData = nullptr, const FGLTFMaterialArray& Materials = {});
 
 	FGLTFJsonSamplerIndex GetOrAddSampler(const UTexture* Texture);
 	FGLTFJsonTextureIndex GetOrAddTexture(const UTexture2D* Texture);
diff --git a/Engine/Plugins/Enterprise/GLTFExporter/MigratedFromGit/Source/GLTFExporter/Private/Converters/GLTFMaterialArray.h b/Engine/Plugins/Enterprise/GLTFExporter/MigratedFromGit/Source/GLTFExporter/Private/Converters/GLTFMaterialArray.h
--- a/Engine/Plugins/Enterprise/GLTFExporter/MigratedFromGit/Source/GLTFExporter/Private/Converters/GLTFMaterialArray.h
+++ b/Engine/Plugins/Enterprise/GLTFExporter/MigratedFromGit/Source/GLTFExporter/Private/Converters/GLTFMaterialArray.h
@@ -36,85 +36,25 @@ struct FGLTFMaterialArray : TArray<const UMaterialInterface*>
 	{
 	}
 
-	const UMaterialInterface* GetOverride(const TArray<FStaticMaterial>& Originals, int32 Index) const
+	void FillIn(const TArray<UMaterialInterface*>& Defaults)
 	{
-		return GetOverride<FStaticMaterial>(Originals, Index);
+		return FillIn<UMaterialInterface*>(Defaults);
 	}
 
-	const UMaterialInterface* GetOverride(const TArray<FSkeletalMaterial>& Originals, int32 Index) const
+	void FillIn(const TArray<FStaticMaterial>& Defaults)
 	{
-		return GetOverride<FSkeletalMaterial>(Originals, Index);
+		return FillIn<FStaticMaterial>(Defaults);
 	}
 
-	const UMaterialInterface* GetOverride(const TArray<UMaterialInterface*>& Originals, int32 Index) const
+	void FillIn(const TArray<FSkeletalMaterial>& Defaults)
 	{
-		return GetOverride<UMaterialInterface*>(Originals, Index);
-	}
-
-	TArray<UMaterialInterface*> GetOverrides(const TArray<FStaticMaterial>& Originals) const
-	{
-		return GetOverrides<FStaticMaterial>(Originals);
-	}
-
-	TArray<UMaterialInterface*> GetOverrides(const TArray<FSkeletalMaterial>& Originals) const
-	{
-		return GetOverrides<FSkeletalMaterial>(Originals);
-	}
-
-	TArray<UMaterialInterface*> GetOverrides(const TArray<UMaterialInterface*>& Originals) const
-	{
-		return GetOverrides<UMaterialInterface*>(Originals);
-	}
-
-	void ClearRedundantOverrides(const TArray<FStaticMaterial>& Originals)
-	{
-		ClearRedundantOverrides<FStaticMaterial>(Originals);
-	}
-
-	void ClearRedundantOverrides(const TArray<FSkeletalMaterial>& Originals)
-	{
-		ClearRedundantOverrides<FSkeletalMaterial>(Originals);
-	}
-
-	void ClearRedundantOverrides(const TArray<UMaterialInterface*>& Originals)
-	{
-		ClearRedundantOverrides<UMaterialInterface*>(Originals);
+		return FillIn<FSkeletalMaterial>(Defaults);
 	}
 
 	using TArray::operator=;
 	using TArray::operator==;
 	using TArray::operator!=;
 
-	bool operator==(const TArray<FStaticMaterial>& Other) const
-	{
-		return Equals(Other);
-	}
-
-	bool operator!=(const TArray<FStaticMaterial>& Other) const
-	{
-		return !Equals(Other);
-	}
-
-	bool operator==(const TArray<FSkeletalMaterial>& Other) const
-	{
-		return Equals(Other);
-	}
-
-	bool operator!=(const TArray<FSkeletalMaterial>& Other) const
-	{
-		return !Equals(Other);
-	}
-
-	bool operator==(const TArray<UMaterialInterface*>& Other) const
-	{
-		return Equals(Other);
-	}
-
-	bool operator!=(const TArray<UMaterialInterface*>& Other) const
-	{
-		return !Equals(Other);
-	}
-
 	friend uint32 GetTypeHash(const TArray& Array)
 	{
 		uint32 Hash = GetTypeHash(Array.Num());
@@ -128,70 +68,25 @@ struct FGLTFMaterialArray : TArray<const UMaterialInterface*>
 private:
 
 	template <typename MeshMaterialType>
-	const UMaterialInterface* GetOverride(const TArray<MeshMaterialType>& Originals, int32 Index) const
+	void FillIn(const TArray<MeshMaterialType>& Defaults)
 	{
-		if (IsValidIndex(Index) && (*this)[Index] != nullptr)
-		{
-			return (*this)[Index];
-		}
-
-		if (Originals.IsValidIndex(Index))
-		{
-			return GetInterface(Originals[Index]);
-		}
-
-		return nullptr;
-	}
-
-	template <typename MeshMaterialType>
-	TArray<UMaterialInterface*> GetOverrides(const TArray<MeshMaterialType>& Originals) const
-	{
-		const int32 Count = Originals.Num();
-		TArray<UMaterialInterface*> Materials;
-		Materials.AddDefaulted(Count);
-
-		for (int32 Index = 0; Index < Count; ++Index)
-		{
-			Materials[Index] = const_cast<UMaterialInterface*>(GetOverride(Originals, Index));
-		}
-
-		return Materials;
-	}
-
-	template <typename MeshMaterialType>
-	void ClearRedundantOverrides(const TArray<MeshMaterialType>& Originals)
-	{
-		const int32 Count = Originals.Num();
+		const int32 Count = Defaults.Num();
 		SetNumZeroed(Count);
 
 		for (int32 Index = 0; Index < Count; ++Index)
 		{
-			if (GetInterface(Originals[Index]) == (*this)[Index])
+			const UMaterialInterface*& Element = (*this)[Index];
+			if (Element == nullptr)
 			{
-				(*this)[Index] = nullptr;
+				Element = GetInterface(Defaults[Index]);
+				if (Element == nullptr)
+				{
+					Element = UMaterial::GetDefaultMaterial(MD_Surface);
+				}
 			}
 		}
 	}
 
-	template <typename MeshMaterialType>
-	bool Equals(const TArray<MeshMaterialType>& Other) const
-	{
-		if (Other.Num() != Num())
-		{
-			return false;
-		}
-
-		for (int32 Index = 0; Index < Num(); ++Index)
-		{
-			if (!Equals(Other[Index], Index))
-			{
-				return false;
-			}
-		}
-
-		return true;
-	}
-
 	static const UMaterialInterface* GetInterface(const UMaterialInterface* Material)
 	{
 		return Material;
@@ -206,10 +101,4 @@ private:
 	{
 		return Material.MaterialInterface;
 	}
-
-	template <typename MeshMaterialType>
-	bool Equals(const MeshMaterialType& Other, int32 Index) const
-	{
-		return IsValidIndex(Index) && GetInterface(Other) == (*this)[Index];
-	}
 };
diff --git a/Engine/Plugins/Enterprise/GLTFExporter/MigratedFromGit/Source/GLTFExporter/Private/Converters/GLTFMaterialConverters.cpp b/Engine/Plugins/Enterprise/GLTFExporter/MigratedFromGit/Source/GLTFExporter/Private/Converters/GLTFMaterialConverters.cpp
--- a/Engine/Plugins/Enterprise/GLTFExporter/MigratedFromGit/Source/GLTFExporter/Private/Converters/GLTFMaterialConverters.cpp
+++ b/Engine/Plugins/Enterprise/GLTFExporter/MigratedFromGit/Source/GLTFExporter/Private/Converters/GLTFMaterialConverters.cpp
@@ -6,7 +6,7 @@
 #include "Builders/GLTFConvertBuilder.h"
 #include "Tasks/GLTFMaterialTasks.h"
 
-void FGLTFMaterialConverter::Sanitize(const UMaterialInterface*& Material, const FGLTFMeshData*& MeshData, FGLTFMaterialArray& OverrideMaterials)
+void FGLTFMaterialConverter::Sanitize(const UMaterialInterface*& Material, const FGLTFMeshData*& MeshData, FGLTFMaterialArray& Materials)
 {
 	if (MeshData == nullptr ||
 		!Builder.ExportOptions->bBakeMaterialInputs ||
@@ -14,11 +14,11 @@ void FGLTFMaterialConverter::Sanitize(const UMaterialInterface*& Material, const
 		!FGLTFMaterialUtility::MaterialNeedsVertexData(Material))
 	{
 		MeshData = nullptr;
-		OverrideMaterials = {};
+		Materials = {};
 	}
 }
 
-FGLTFJsonMaterialIndex FGLTFMaterialConverter::Convert(const UMaterialInterface* Material, const FGLTFMeshData* MeshData, const FGLTFMaterialArray OverrideMaterials)
+FGLTFJsonMaterialIndex FGLTFMaterialConverter::Convert(const UMaterialInterface* Material, const FGLTFMeshData* MeshData, FGLTFMaterialArray Materials)
 {
 	if (Material == FGLTFMaterialUtility::GetDefault())
 	{
@@ -26,6 +26,6 @@ FGLTFJsonMaterialIndex FGLTFMaterialConverter::Convert(const UMaterialInterface*
 	}
 
 	const FGLTFJsonMaterialIndex MaterialIndex = Builder.AddMaterial();
-	Builder.SetupTask<FGLTFMaterialTask>(Builder, Material, MeshData, OverrideMaterials, MaterialIndex);
+	Builder.SetupTask<FGLTFMaterialTask>(Builder, Material, MeshData, Materials, MaterialIndex);
 	return MaterialIndex;
 }
diff --git a/Engine/Plugins/Enterprise/GLTFExporter/MigratedFromGit/Source/GLTFExporter/Private/Converters/GLTFMaterialConverters.h b/Engine/Plugins/Enterprise/GLTFExporter/MigratedFromGit/Source/GLTFExporter/Private/Converters/GLTFMaterialConverters.h
--- a/Engine/Plugins/Enterprise/GLTFExporter/MigratedFromGit/Source/GLTFExporter/Private/Converters/GLTFMaterialConverters.h
+++ b/Engine/Plugins/Enterprise/GLTFExporter/MigratedFromGit/Source/GLTFExporter/Private/Converters/GLTFMaterialConverters.h
@@ -13,7 +13,7 @@ class FGLTFMaterialConverter final : public FGLTFBuilderContext, public TGLTFCon
 {
 	using FGLTFBuilderContext::FGLTFBuilderContext;
 
-	virtual void Sanitize(const UMaterialInterface*& Material, const FGLTFMeshData*& MeshData, FGLTFMaterialArray& OverrideMaterials) override;
+	virtual void Sanitize(const UMaterialInterface*& Material, const FGLTFMeshData*& MeshData, FGLTFMaterialArray& Materials) override;
 
-	virtual FGLTFJsonMaterialIndex Convert(const UMaterialInterface* Material, const FGLTFMeshData* MeshData, const FGLTFMaterialArray OverrideMaterials) override;
+	virtual FGLTFJsonMaterialIndex Convert(const UMaterialInterface* Material, const FGLTFMeshData* MeshData, FGLTFMaterialArray Materials) override;
 };
diff --git a/Engine/Plugins/Enterprise/GLTFExporter/MigratedFromGit/Source/GLTFExporter/Private/Converters/GLTFMaterialUtility.cpp b/Engine/Plugins/Enterprise/GLTFExporter/MigratedFromGit/Source/GLTFExporter/Private/Converters/GLTFMaterialUtility.cpp
--- a/Engine/Plugins/Enterprise/GLTFExporter/MigratedFromGit/Source/GLTFExporter/Private/Converters/GLTFMaterialUtility.cpp
+++ b/Engine/Plugins/Enterprise/GLTFExporter/MigratedFromGit/Source/GLTFExporter/Private/Converters/GLTFMaterialUtility.cpp
@@ -540,7 +540,7 @@ bool FGLTFMaterialUtility::MaterialNeedsVertexData(const UMaterialInterface* Mat
 	return bRequiresVertexData;
 }
 
-bool FGLTFMaterialUtility::MaterialsNeedVertexData(const TArray<UMaterialInterface*> Materials)
+bool FGLTFMaterialUtility::MaterialsNeedVertexData(const TArray<const UMaterialInterface*>& Materials)
 {
 	for (const UMaterialInterface* Material: Materials)
 	{
diff --git a/Engine/Plugins/Enterprise/GLTFExporter/MigratedFromGit/Source/GLTFExporter/Private/Converters/GLTFMaterialUtility.h b/Engine/Plugins/Enterprise/GLTFExporter/MigratedFromGit/Source/GLTFExporter/Private/Converters/GLTFMaterialUtility.h
--- a/Engine/Plugins/Enterprise/GLTFExporter/MigratedFromGit/Source/GLTFExporter/Private/Converters/GLTFMaterialUtility.h
+++ b/Engine/Plugins/Enterprise/GLTFExporter/MigratedFromGit/Source/GLTFExporter/Private/Converters/GLTFMaterialUtility.h
@@ -69,5 +69,5 @@ struct FGLTFMaterialUtility
 	static FString ShadingModelsToString(const FMaterialShadingModelField& ShadingModels);
 
 	static bool MaterialNeedsVertexData(const UMaterialInterface* Material);
-	static bool MaterialsNeedVertexData(const TArray<UMaterialInterface*> Materials);
+	static bool MaterialsNeedVertexData(const TArray<const UMaterialInterface*>& Materials);
 };
diff --git a/Engine/Plugins/Enterprise/GLTFExporter/MigratedFromGit/Source/GLTFExporter/Private/Converters/GLTFMeshConverters.cpp b/Engine/Plugins/Enterprise/GLTFExporter/MigratedFromGit/Source/GLTFExporter/Private/Converters/GLTFMeshConverters.cpp
--- a/Engine/Plugins/Enterprise/GLTFExporter/MigratedFromGit/Source/GLTFExporter/Private/Converters/GLTFMeshConverters.cpp
+++ b/Engine/Plugins/Enterprise/GLTFExporter/MigratedFromGit/Source/GLTFExporter/Private/Converters/GLTFMeshConverters.cpp
@@ -6,11 +6,16 @@
 #include "Builders/GLTFConvertBuilder.h"
 #include "Tasks/GLTFMeshTasks.h"
 
-void FGLTFStaticMeshConverter::Sanitize(const UStaticMesh*& StaticMesh, const UStaticMeshComponent*& StaticMeshComponent, FGLTFMaterialArray& OverrideMaterials, int32& LODIndex)
+void FGLTFStaticMeshConverter::Sanitize(const UStaticMesh*& StaticMesh, const UStaticMeshComponent*& StaticMeshComponent, FGLTFMaterialArray& Materials, int32& LODIndex)
 {
-	const TArray<UMaterialInterface*> Materials = StaticMeshComponent != nullptr
-		? OverrideMaterials.GetOverrides(StaticMeshComponent->GetMaterials())
-		: OverrideMaterials.GetOverrides(StaticMesh->StaticMaterials);
+	if (StaticMeshComponent != nullptr)
+	{
+		Materials.FillIn(StaticMeshComponent->GetMaterials());
+	}
+	else
+	{
+		Materials.FillIn(StaticMesh->StaticMaterials);
+	}
 
 	if (LODIndex < 0)
 	{
@@ -32,32 +37,25 @@ void FGLTFStaticMeshConverter::Sanitize(const UStaticMesh*& StaticMesh, const US
 			StaticMeshComponent = nullptr;
 		}
 	}
-
-	// Clean up override materials to only contain items that differ from the original materials
-	OverrideMaterials = FGLTFMaterialArray(Materials);
-
-	if (StaticMeshComponent != nullptr)
-	{
-		OverrideMaterials.ClearRedundantOverrides(StaticMeshComponent->GetMaterials());
-	}
-	else
-	{
-		OverrideMaterials.ClearRedundantOverrides(StaticMesh->StaticMaterials);
-	}
 }
 
-FGLTFJsonMeshIndex FGLTFStaticMeshConverter::Convert(const UStaticMesh* StaticMesh,  const UStaticMeshComponent* StaticMeshComponent, FGLTFMaterialArray OverrideMaterials, int32 LODIndex)
+FGLTFJsonMeshIndex FGLTFStaticMeshConverter::Convert(const UStaticMesh* StaticMesh,  const UStaticMeshComponent* StaticMeshComponent, FGLTFMaterialArray Materials, int32 LODIndex)
 {
 	const FGLTFJsonMeshIndex MeshIndex = Builder.AddMesh();
-	Builder.SetupTask<FGLTFStaticMeshTask>(Builder, MeshSectionConverter, MeshDataConverter, StaticMesh, StaticMeshComponent, OverrideMaterials, LODIndex, MeshIndex);
+	Builder.SetupTask<FGLTFStaticMeshTask>(Builder, MeshSectionConverter, MeshDataConverter, StaticMesh, StaticMeshComponent, Materials, LODIndex, MeshIndex);
 	return MeshIndex;
 }
 
-void FGLTFSkeletalMeshConverter::Sanitize(const USkeletalMesh*& SkeletalMesh, const USkeletalMeshComponent*& SkeletalMeshComponent, FGLTFMaterialArray& OverrideMaterials, int32& LODIndex)
+void FGLTFSkeletalMeshConverter::Sanitize(const USkeletalMesh*& SkeletalMesh, const USkeletalMeshComponent*& SkeletalMeshComponent, FGLTFMaterialArray& Materials, int32& LODIndex)
 {
-	const TArray<UMaterialInterface*> Materials = SkeletalMeshComponent != nullptr
-		? OverrideMaterials.GetOverrides(SkeletalMeshComponent->GetMaterials())
-		: OverrideMaterials.GetOverrides(SkeletalMesh->Materials);
+	if (SkeletalMeshComponent != nullptr)
+	{
+		Materials.FillIn(SkeletalMeshComponent->GetMaterials());
+	}
+	else
+	{
+		Materials.FillIn(SkeletalMesh->Materials);
+	}
 
 	if (LODIndex < 0)
 	{
@@ -79,23 +77,11 @@ void FGLTFSkeletalMeshConverter::Sanitize(const USkeletalMesh*& SkeletalMesh, co
 			SkeletalMeshComponent = nullptr;
 		}
 	}
-
-	// Clean up override materials to only contain items that differ from the original materials
-	OverrideMaterials = FGLTFMaterialArray(Materials);
-
-	if (SkeletalMeshComponent != nullptr)
-	{
-		OverrideMaterials.ClearRedundantOverrides(SkeletalMeshComponent->GetMaterials());
-	}
-	else
-	{
-		OverrideMaterials.ClearRedundantOverrides(SkeletalMesh->Materials);
-	}
 }
 
-FGLTFJsonMeshIndex FGLTFSkeletalMeshConverter::Convert(const USkeletalMesh* SkeletalMesh, const USkeletalMeshComponent* SkeletalMeshComponent, FGLTFMaterialArray OverrideMaterials, int32 LODIndex)
+FGLTFJsonMeshIndex FGLTFSkeletalMeshConverter::Convert(const USkeletalMesh* SkeletalMesh, const USkeletalMeshComponent* SkeletalMeshComponent, FGLTFMaterialArray Materials, int32 LODIndex)
 {
 	const FGLTFJsonMeshIndex MeshIndex = Builder.AddMesh();
-	Builder.SetupTask<FGLTFSkeletalMeshTask>(Builder, MeshSectionConverter, MeshDataConverter, SkeletalMesh, SkeletalMeshComponent, OverrideMaterials, LODIndex, MeshIndex);
+	Builder.SetupTask<FGLTFSkeletalMeshTask>(Builder, MeshSectionConverter, MeshDataConverter, SkeletalMesh, SkeletalMeshComponent, Materials, LODIndex, MeshIndex);
 	return MeshIndex;
 }
diff --git a/Engine/Plugins/Enterprise/GLTFExporter/MigratedFromGit/Source/GLTFExporter/Private/Converters/GLTFMeshConverters.h b/Engine/Plugins/Enterprise/GLTFExporter/MigratedFromGit/Source/GLTFExporter/Private/Converters/GLTFMeshConverters.h
--- a/Engine/Plugins/Enterprise/GLTFExporter/MigratedFromGit/Source/GLTFExporter/Private/Converters/GLTFMeshConverters.h
+++ b/Engine/Plugins/Enterprise/GLTFExporter/MigratedFromGit/Source/GLTFExporter/Private/Converters/GLTFMeshConverters.h
@@ -20,9 +20,9 @@ class FGLTFStaticMeshConverter final : public TGLTFMeshConverter<const UStaticMe
 {
 	using TGLTFMeshConverter::TGLTFMeshConverter;
 
-	virtual void Sanitize(const UStaticMesh*& StaticMesh, const UStaticMeshComponent*& StaticMeshComponent, FGLTFMaterialArray& OverrideMaterials, int32& LODIndex) override;
+	virtual void Sanitize(const UStaticMesh*& StaticMesh, const UStaticMeshComponent*& StaticMeshComponent, FGLTFMaterialArray& Materials, int32& LODIndex) override;
 
-	virtual FGLTFJsonMeshIndex Convert(const UStaticMesh* StaticMesh, const UStaticMeshComponent* StaticMeshComponent, FGLTFMaterialArray OverrideMaterials, int32 LODIndex) override;
+	virtual FGLTFJsonMeshIndex Convert(const UStaticMesh* StaticMesh, const UStaticMeshComponent* StaticMeshComponent, FGLTFMaterialArray Materials, int32 LODIndex) override;
 
 	FGLTFStaticMeshSectionConverter MeshSectionConverter;
 	FGLTFStaticMeshDataConverter MeshDataConverter;
@@ -32,9 +32,9 @@ class FGLTFSkeletalMeshConverter final : public TGLTFMeshConverter<const USkelet
 {
 	using TGLTFMeshConverter::TGLTFMeshConverter;
 
-	virtual void Sanitize(const USkeletalMesh*& SkeletalMesh, const USkeletalMeshComponent*& SkeletalMeshComponent, FGLTFMaterialArray& OverrideMaterials, int32& LODIndex) override;
+	virtual void Sanitize(const USkeletalMesh*& SkeletalMesh, const USkeletalMeshComponent*& SkeletalMeshComponent, FGLTFMaterialArray& Materials, int32& LODIndex) override;
 
-	virtual FGLTFJsonMeshIndex Convert(const USkeletalMesh* SkeletalMesh, const USkeletalMeshComponent* SkeletalMeshComponent, FGLTFMaterialArray OverrideMaterials, int32 LODIndex) override;
+	virtual FGLTFJsonMeshIndex Convert(const USkeletalMesh* SkeletalMesh, const USkeletalMeshComponent* SkeletalMeshComponent, FGLTFMaterialArray Materials, int32 LODIndex) override;
 
 	FGLTFSkeletalMeshSectionConverter MeshSectionConverter;
 	FGLTFSkeletalMeshDataConverter MeshDataConverter;
diff --git a/Engine/Plugins/Enterprise/GLTFExporter/MigratedFromGit/Source/GLTFExporter/Private/Tasks/GLTFMaterialTasks.cpp b/Engine/Plugins/Enterprise/GLTFExporter/MigratedFromGit/Source/GLTFExporter/Private/Tasks/GLTFMaterialTasks.cpp
--- a/Engine/Plugins/Enterprise/GLTFExporter/MigratedFromGit/Source/GLTFExporter/Private/Tasks/GLTFMaterialTasks.cpp
+++ b/Engine/Plugins/Enterprise/GLTFExporter/MigratedFromGit/Source/GLTFExporter/Private/Tasks/GLTFMaterialTasks.cpp
@@ -1364,10 +1364,7 @@ FGLTFPropertyBakeOutput FGLTFMaterialTask::BakeMaterialProperty(EMaterialPropert
 		for (int32 SectionIndex = 0; SectionIndex < Sections.Num(); ++SectionIndex)
 		{
 			const FSectionInfo& Section = Sections[SectionIndex];
-			const UMaterialInterface* OverrideMaterial = OverrideMaterials.IsValidIndex(Section.MaterialIndex) ? OverrideMaterials[Section.MaterialIndex] : nullptr;
-			const UMaterialInterface* SectionMaterial = OverrideMaterial != nullptr ? OverrideMaterial : Section.Material;
-
-			if (SectionMaterial == Material)
+			if (Materials[Section.MaterialIndex] == Material)
 			{
 				MeshSectionIndices.Add(SectionIndex);
 			}
diff --git a/Engine/Plugins/Enterprise/GLTFExporter/MigratedFromGit/Source/GLTFExporter/Private/Tasks/GLTFMaterialTasks.h b/Engine/Plugins/Enterprise/GLTFExporter/MigratedFromGit/Source/GLTFExporter/Private/Tasks/GLTFMaterialTasks.h
--- a/Engine/Plugins/Enterprise/GLTFExporter/MigratedFromGit/Source/GLTFExporter/Private/Tasks/GLTFMaterialTasks.h
+++ b/Engine/Plugins/Enterprise/GLTFExporter/MigratedFromGit/Source/GLTFExporter/Private/Tasks/GLTFMaterialTasks.h
@@ -13,12 +13,12 @@ class FGLTFMaterialTask : public FGLTFTask
 {
 public:
 
-	FGLTFMaterialTask(FGLTFConvertBuilder& Builder, const UMaterialInterface* Material, const FGLTFMeshData* MeshData, FGLTFMaterialArray OverrideMaterials, FGLTFJsonMaterialIndex MaterialIndex)
+	FGLTFMaterialTask(FGLTFConvertBuilder& Builder, const UMaterialInterface* Material, const FGLTFMeshData* MeshData, FGLTFMaterialArray Materials, FGLTFJsonMaterialIndex MaterialIndex)
         : FGLTFTask(EGLTFTaskPriority::Material)
 		, Builder(Builder)
 		, Material(Material)
 		, MeshData(MeshData)
-		, OverrideMaterials(OverrideMaterials)
+		, Materials(Materials)
 		, MaterialIndex(MaterialIndex)
 	{
 	}
@@ -35,7 +35,7 @@ private:
 	FGLTFConvertBuilder& Builder;
 	const UMaterialInterface* Material;
 	const FGLTFMeshData* MeshData;
-	const FGLTFMaterialArray OverrideMaterials;
+	const FGLTFMaterialArray Materials;
 	const FGLTFJsonMaterialIndex MaterialIndex;
 
 	bool TryGetAlphaMode(EGLTFJsonAlphaMode& AlphaMode) const;
diff --git a/Engine/Plugins/Enterprise/GLTFExporter/MigratedFromGit/Source/GLTFExporter/Private/Tasks/GLTFMeshTasks.cpp b/Engine/Plugins/Enterprise/GLTFExporter/MigratedFromGit/Source/GLTFExporter/Private/Tasks/GLTFMeshTasks.cpp
--- a/Engine/Plugins/Enterprise/GLTFExporter/MigratedFromGit/Source/GLTFExporter/Private/Tasks/GLTFMeshTasks.cpp
+++ b/Engine/Plugins/Enterprise/GLTFExporter/MigratedFromGit/Source/GLTFExporter/Private/Tasks/GLTFMeshTasks.cpp
@@ -57,16 +57,8 @@ void FGLTFStaticMeshTask::Complete()
 			JsonPrimitive.Attributes.TexCoords[UVIndex] = Builder.GetOrAddUVAccessor(ConvertedSection, VertexBuffer, UVIndex);
 		}
 
-		const UMaterialInterface* Material = StaticMeshComponent != nullptr
-			? OverrideMaterials.GetOverride(StaticMeshComponent->GetMaterials(), MaterialIndex)	// TODO: cache component materials once instead of retrieving multiple times
-			: OverrideMaterials.GetOverride(StaticMesh->StaticMaterials, MaterialIndex);
-
-		if (Material == nullptr)
-		{
-			Material = UMaterial::GetDefaultMaterial(MD_Surface);
-		}
-
-		JsonPrimitive.Material =  Builder.GetOrAddMaterial(Material, MeshData, OverrideMaterials);
+		const UMaterialInterface* Material = Materials[MaterialIndex];
+		JsonPrimitive.Material =  Builder.GetOrAddMaterial(Material, MeshData, Materials);
 	}
 }
 
@@ -144,15 +136,7 @@ void FGLTFSkeletalMeshTask::Complete()
 			}
 		}
 
-		const UMaterialInterface* Material = SkeletalMeshComponent != nullptr
-			? OverrideMaterials.GetOverride(SkeletalMeshComponent->GetMaterials(), MaterialIndex)	// TODO: cache component materials once instead of retrieving multiple times
-			: OverrideMaterials.GetOverride(SkeletalMesh->Materials, MaterialIndex);
-
-		if (Material == nullptr)
-		{
-			Material = UMaterial::GetDefaultMaterial(MD_Surface);
-		}
-
-		JsonPrimitive.Material =  Builder.GetOrAddMaterial(Material, MeshData, OverrideMaterials);
+		const UMaterialInterface* Material = Materials[MaterialIndex];
+		JsonPrimitive.Material =  Builder.GetOrAddMaterial(Material, MeshData, Materials);
 	}
 }
diff --git a/Engine/Plugins/Enterprise/GLTFExporter/MigratedFromGit/Source/GLTFExporter/Private/Tasks/GLTFMeshTasks.h b/Engine/Plugins/Enterprise/GLTFExporter/MigratedFromGit/Source/GLTFExporter/Private/Tasks/GLTFMeshTasks.h
--- a/Engine/Plugins/Enterprise/GLTFExporter/MigratedFromGit/Source/GLTFExporter/Private/Tasks/GLTFMeshTasks.h
+++ b/Engine/Plugins/Enterprise/GLTFExporter/MigratedFromGit/Source/GLTFExporter/Private/Tasks/GLTFMeshTasks.h
@@ -14,14 +14,14 @@ class FGLTFStaticMeshTask : public FGLTFTask
 {
 public:
 
-	FGLTFStaticMeshTask(FGLTFConvertBuilder& Builder, FGLTFStaticMeshSectionConverter& MeshSectionConverter, FGLTFStaticMeshDataConverter& MeshDataConverter, const UStaticMesh* StaticMesh, const UStaticMeshComponent* StaticMeshComponent, FGLTFMaterialArray OverrideMaterials, int32 LODIndex, FGLTFJsonMeshIndex MeshIndex)
+	FGLTFStaticMeshTask(FGLTFConvertBuilder& Builder, FGLTFStaticMeshSectionConverter& MeshSectionConverter, FGLTFStaticMeshDataConverter& MeshDataConverter, const UStaticMesh* StaticMesh, const UStaticMeshComponent* StaticMeshComponent, FGLTFMaterialArray Materials, int32 LODIndex, FGLTFJsonMeshIndex MeshIndex)
 		: FGLTFTask(EGLTFTaskPriority::Mesh)
 		, Builder(Builder)
 		, MeshSectionConverter(MeshSectionConverter)
 		, MeshDataConverter(MeshDataConverter)
 		, StaticMesh(StaticMesh)
 		, StaticMeshComponent(StaticMeshComponent)
-		, OverrideMaterials(OverrideMaterials)
+		, Materials(Materials)
 		, LODIndex(LODIndex)
         , MeshIndex(MeshIndex)
 	{
@@ -41,7 +41,7 @@ private:
 	FGLTFStaticMeshDataConverter& MeshDataConverter;
 	const UStaticMesh* StaticMesh;
 	const UStaticMeshComponent* StaticMeshComponent;
-	const FGLTFMaterialArray OverrideMaterials;
+	const FGLTFMaterialArray Materials;
 	const int32 LODIndex;
 	const FGLTFJsonMeshIndex MeshIndex;
 };
@@ -50,14 +50,14 @@ class FGLTFSkeletalMeshTask : public FGLTFTask
 {
 public:
 
-	FGLTFSkeletalMeshTask(FGLTFConvertBuilder& Builder, FGLTFSkeletalMeshSectionConverter& MeshSectionConverter, FGLTFSkeletalMeshDataConverter& MeshDataConverter, const USkeletalMesh* SkeletalMesh, const USkeletalMeshComponent* SkeletalMeshComponent, FGLTFMaterialArray OverrideMaterials, int32 LODIndex, FGLTFJsonMeshIndex MeshIndex)
+	FGLTFSkeletalMeshTask(FGLTFConvertBuilder& Builder, FGLTFSkeletalMeshSectionConverter& MeshSectionConverter, FGLTFSkeletalMeshDataConverter& MeshDataConverter, const USkeletalMesh* SkeletalMesh, const USkeletalMeshComponent* SkeletalMeshComponent, FGLTFMaterialArray Materials, int32 LODIndex, FGLTFJsonMeshIndex MeshIndex)
 		: FGLTFTask(EGLTFTaskPriority::Mesh)
 		, Builder(Builder)
 		, MeshSectionConverter(MeshSectionConverter)
 		, MeshDataConverter(MeshDataConverter)
 		, SkeletalMesh(SkeletalMesh)
 		, SkeletalMeshComponent(SkeletalMeshComponent)
-		, OverrideMaterials(OverrideMaterials)
+		, Materials(Materials)
 		, LODIndex(LODIndex)
 		, MeshIndex(MeshIndex)
 	{
@@ -77,7 +77,7 @@ private:
 	FGLTFSkeletalMeshDataConverter& MeshDataConverter;
 	const USkeletalMesh* SkeletalMesh;
 	const USkeletalMeshComponent* SkeletalMeshComponent;
-	const FGLTFMaterialArray OverrideMaterials;
+	const FGLTFMaterialArray Materials;
 	const int32 LODIndex;
 	const FGLTFJsonMeshIndex MeshIndex;
 };
