[GLTFExporter] refactor: add GLTFHashableArray as generalized base class for GLTFMaterialArray

commit 204b41a99db43255510da9091f49e02135ec97c1
Date: Tue, 3 Nov 2020 22:22:44 +0100

    refactor: add GLTFHashableArray as generalized base class for GLTFMaterialArray
    (cherry picked from commit def488bf0592d40565e2d6a090e9e520f5a722a3​)

#skipci

diff --git a/Engine/Plugins/Enterprise/GLTFExporter/MigratedFromGit/Source/GLTFExporter/Private/Converters/GLTFHashableArray.h b/Engine/Plugins/Enterprise/GLTFExporter/MigratedFromGit/Source/GLTFExporter/Private/Converters/GLTFHashableArray.h
new file mode 100644
--- /dev/null
+++ b/Engine/Plugins/Enterprise/GLTFExporter/MigratedFromGit/Source/GLTFExporter/Private/Converters/GLTFHashableArray.h
@@ -0,0 +1,48 @@
+// Copyright Epic Games, Inc. All Rights Reserved.
+
+#pragma once
+
+#include "CoreMinimal.h"
+
+template <typename ElementType>
+struct FGLTFHashableArray : TArray<ElementType>
+{
+	FGLTFHashableArray()
+	{
+	}
+
+	FGLTFHashableArray(const TArray<ElementType>& Other)
+		: TArray(Other)
+	{
+	}
+
+	FGLTFHashableArray(const FGLTFHashableArray& Other)
+        : TArray(Other)
+	{
+	}
+
+	template <typename OtherElementType, typename OtherAllocatorType>
+	explicit FGLTFHashableArray(const TArray<OtherElementType, OtherAllocatorType>& Other)
+		: TArray(Other)
+	{
+	}
+
+	FGLTFHashableArray(std::initializer_list<ElementType> InitList)
+		: TArray(InitList)
+	{
+	}
+
+	using TArray<ElementType>::operator=;
+	using TArray<ElementType>::operator==;
+	using TArray<ElementType>::operator!=;
+
+	friend uint32 GetTypeHash(const TArray<ElementType>& Array)
+	{
+		uint32 Hash = GetTypeHash(Array.Num());
+		for (const ElementType& Element : Array)
+		{
+			Hash = HashCombine(Hash, GetTypeHash(Element));
+		}
+		return Hash;
+	}
+};
diff --git a/Engine/Plugins/Enterprise/GLTFExporter/MigratedFromGit/Source/GLTFExporter/Private/Converters/GLTFMaterialArray.h b/Engine/Plugins/Enterprise/GLTFExporter/MigratedFromGit/Source/GLTFExporter/Private/Converters/GLTFMaterialArray.h
--- a/Engine/Plugins/Enterprise/GLTFExporter/MigratedFromGit/Source/GLTFExporter/Private/Converters/GLTFMaterialArray.h
+++ b/Engine/Plugins/Enterprise/GLTFExporter/MigratedFromGit/Source/GLTFExporter/Private/Converters/GLTFMaterialArray.h
@@ -2,39 +2,16 @@
 
 #pragma once
 
-#include "CoreMinimal.h"
+#include "GLTFHashableArray.h"
 
 class UMaterialInterface;
 
-// The reason behind a custom material array type are two:
-// 1. Forwarding a TArray by copy in FGLTFConverter::GetOrAdd twice will result in a empty array
-// 2. There is no GetTypeHash for generic arrays
-struct FGLTFMaterialArray : TArray<const UMaterialInterface*>
+struct FGLTFMaterialArray : FGLTFHashableArray<const UMaterialInterface*>
 {
-	FGLTFMaterialArray()
-	{
-	}
-
-	FGLTFMaterialArray(const TArray& Other)
-		: TArray(Other)
-	{
-	}
-
-	FGLTFMaterialArray(const FGLTFMaterialArray& Other)
-		: TArray(Other)
-	{
-	}
-
-	template <typename ElementType, typename AllocatorType>
-	explicit FGLTFMaterialArray(const TArray<ElementType, AllocatorType>& Other)
-		: TArray(Other)
-	{
-	}
-
-	FGLTFMaterialArray(std::initializer_list<const UMaterialInterface*> InitList)
-		: TArray(InitList)
-	{
-	}
+	using FGLTFHashableArray::FGLTFHashableArray;
+	using FGLTFHashableArray::operator=;
+	using FGLTFHashableArray::operator==;
+	using FGLTFHashableArray::operator!=;
 
 	void FillIn(const TArray<UMaterialInterface*>& Defaults)
 	{
@@ -51,20 +28,6 @@ struct FGLTFMaterialArray : TArray<const UMaterialInterface*>
 		return FillIn<FSkeletalMaterial>(Defaults);
 	}
 
-	using TArray::operator=;
-	using TArray::operator==;
-	using TArray::operator!=;
-
-	friend uint32 GetTypeHash(const TArray& Array)
-	{
-		uint32 Hash = GetTypeHash(Array.Num());
-		for (const UMaterialInterface* Material : Array)
-		{
-			Hash = HashCombine(Hash, GetTypeHash(Material));
-		}
-		return Hash;
-	}
-
 private:
 
 	template <typename MeshMaterialType>
