[GLTFExporter] refactor: simplify material converter to use section indices instead of material array as parameter

commit 3ab2ca1f82a25b643fa458e929c006fe4f9c70ed
Date: Tue, 3 Nov 2020 22:56:55 +0100

    refactor: simplify material converter to use section indices instead of material array as parameter
    (cherry picked from commit b9682b6efbe041011ecdbc0a88e9f0f7c0f8d707​)

#skipci

diff --git a/Engine/Plugins/Enterprise/GLTFExporter/MigratedFromGit/Source/GLTFExporter/Private/Builders/GLTFConvertBuilder.cpp b/Engine/Plugins/Enterprise/GLTFExporter/MigratedFromGit/Source/GLTFExporter/Private/Builders/GLTFConvertBuilder.cpp
--- a/Engine/Plugins/Enterprise/GLTFExporter/MigratedFromGit/Source/GLTFExporter/Private/Builders/GLTFConvertBuilder.cpp
+++ b/Engine/Plugins/Enterprise/GLTFExporter/MigratedFromGit/Source/GLTFExporter/Private/Builders/GLTFConvertBuilder.cpp
@@ -140,14 +140,14 @@ FGLTFJsonMeshIndex FGLTFConvertBuilder::GetOrAddMesh(const USkeletalMeshComponen
 	return SkeletalMeshConverter.GetOrAdd(SkeletalMesh, SkeletalMeshComponent, Materials, LODIndex);
 }
 
-FGLTFJsonMaterialIndex FGLTFConvertBuilder::GetOrAddMaterial(const UMaterialInterface* Material, const FGLTFMeshData* MeshData, const FGLTFMaterialArray& Materials)
+FGLTFJsonMaterialIndex FGLTFConvertBuilder::GetOrAddMaterial(const UMaterialInterface* Material, const FGLTFMeshData* MeshData, const FGLTFHashableArray<int32>& SectionIndices)
 {
 	if (Material == nullptr)
 	{
 		return FGLTFJsonMaterialIndex(INDEX_NONE);
 	}
 
-	return MaterialConverter.GetOrAdd(Material, MeshData, Materials);
+	return MaterialConverter.GetOrAdd(Material, MeshData, SectionIndices);
 }
 
 FGLTFJsonSamplerIndex FGLTFConvertBuilder::GetOrAddSampler(const UTexture* Texture)
diff --git a/Engine/Plugins/Enterprise/GLTFExporter/MigratedFromGit/Source/GLTFExporter/Private/Builders/GLTFConvertBuilder.h b/Engine/Plugins/Enterprise/GLTFExporter/MigratedFromGit/Source/GLTFExporter/Private/Builders/GLTFConvertBuilder.h
--- a/Engine/Plugins/Enterprise/GLTFExporter/MigratedFromGit/Source/GLTFExporter/Private/Builders/GLTFConvertBuilder.h
+++ b/Engine/Plugins/Enterprise/GLTFExporter/MigratedFromGit/Source/GLTFExporter/Private/Builders/GLTFConvertBuilder.h
@@ -44,7 +44,7 @@ public:
 	FGLTFJsonMeshIndex GetOrAddMesh(const USkeletalMesh* SkeletalMesh, const FGLTFMaterialArray& Materials = {}, int32 LODIndex = -1);
 	FGLTFJsonMeshIndex GetOrAddMesh(const USkeletalMeshComponent* SkeletalMeshComponent, const FGLTFMaterialArray& Materials = {}, int32 LODIndex = -1);
 
-	FGLTFJsonMaterialIndex GetOrAddMaterial(const UMaterialInterface* Material, const FGLTFMeshData* MeshData = nullptr, const FGLTFMaterialArray& Materials = {});
+	FGLTFJsonMaterialIndex GetOrAddMaterial(const UMaterialInterface* Material, const FGLTFMeshData* MeshData = nullptr, const FGLTFHashableArray<int32>& SectionIndices = {});
 
 	FGLTFJsonSamplerIndex GetOrAddSampler(const UTexture* Texture);
 	FGLTFJsonTextureIndex GetOrAddTexture(const UTexture2D* Texture);
diff --git a/Engine/Plugins/Enterprise/GLTFExporter/MigratedFromGit/Source/GLTFExporter/Private/Converters/GLTFHashableArray.h b/Engine/Plugins/Enterprise/GLTFExporter/MigratedFromGit/Source/GLTFExporter/Private/Converters/GLTFHashableArray.h
--- a/Engine/Plugins/Enterprise/GLTFExporter/MigratedFromGit/Source/GLTFExporter/Private/Converters/GLTFHashableArray.h
+++ b/Engine/Plugins/Enterprise/GLTFExporter/MigratedFromGit/Source/GLTFExporter/Private/Converters/GLTFHashableArray.h
@@ -36,7 +36,7 @@ struct FGLTFHashableArray : TArray<ElementType>
 	using TArray<ElementType>::operator==;
 	using TArray<ElementType>::operator!=;
 
-	friend uint32 GetTypeHash(const TArray<ElementType>& Array)
+	friend uint32 GetTypeHash(const FGLTFHashableArray<ElementType>& Array)
 	{
 		uint32 Hash = GetTypeHash(Array.Num());
 		for (const ElementType& Element : Array)
diff --git a/Engine/Plugins/Enterprise/GLTFExporter/MigratedFromGit/Source/GLTFExporter/Private/Converters/GLTFMaterialConverters.cpp b/Engine/Plugins/Enterprise/GLTFExporter/MigratedFromGit/Source/GLTFExporter/Private/Converters/GLTFMaterialConverters.cpp
--- a/Engine/Plugins/Enterprise/GLTFExporter/MigratedFromGit/Source/GLTFExporter/Private/Converters/GLTFMaterialConverters.cpp
+++ b/Engine/Plugins/Enterprise/GLTFExporter/MigratedFromGit/Source/GLTFExporter/Private/Converters/GLTFMaterialConverters.cpp
@@ -6,7 +6,7 @@
 #include "Builders/GLTFConvertBuilder.h"
 #include "Tasks/GLTFMaterialTasks.h"
 
-void FGLTFMaterialConverter::Sanitize(const UMaterialInterface*& Material, const FGLTFMeshData*& MeshData, FGLTFMaterialArray& Materials)
+void FGLTFMaterialConverter::Sanitize(const UMaterialInterface*& Material, const FGLTFMeshData*& MeshData, FGLTFHashableArray<int32>& SectionIndices)
 {
 	if (MeshData == nullptr ||
 		!Builder.ExportOptions->bBakeMaterialInputs ||
@@ -14,11 +14,11 @@ void FGLTFMaterialConverter::Sanitize(const UMaterialInterface*& Material, const
 		!FGLTFMaterialUtility::MaterialNeedsVertexData(Material))
 	{
 		MeshData = nullptr;
-		Materials = {};
+		SectionIndices = {};
 	}
 }
 
-FGLTFJsonMaterialIndex FGLTFMaterialConverter::Convert(const UMaterialInterface* Material, const FGLTFMeshData* MeshData, FGLTFMaterialArray Materials)
+FGLTFJsonMaterialIndex FGLTFMaterialConverter::Convert(const UMaterialInterface* Material, const FGLTFMeshData* MeshData, FGLTFHashableArray<int32> SectionIndices)
 {
 	if (Material == FGLTFMaterialUtility::GetDefault())
 	{
@@ -26,6 +26,6 @@ FGLTFJsonMaterialIndex FGLTFMaterialConverter::Convert(const UMaterialInterface*
 	}
 
 	const FGLTFJsonMaterialIndex MaterialIndex = Builder.AddMaterial();
-	Builder.SetupTask<FGLTFMaterialTask>(Builder, Material, MeshData, Materials, MaterialIndex);
+	Builder.SetupTask<FGLTFMaterialTask>(Builder, Material, MeshData, SectionIndices, MaterialIndex);
 	return MaterialIndex;
 }
diff --git a/Engine/Plugins/Enterprise/GLTFExporter/MigratedFromGit/Source/GLTFExporter/Private/Converters/GLTFMaterialConverters.h b/Engine/Plugins/Enterprise/GLTFExporter/MigratedFromGit/Source/GLTFExporter/Private/Converters/GLTFMaterialConverters.h
--- a/Engine/Plugins/Enterprise/GLTFExporter/MigratedFromGit/Source/GLTFExporter/Private/Converters/GLTFMaterialConverters.h
+++ b/Engine/Plugins/Enterprise/GLTFExporter/MigratedFromGit/Source/GLTFExporter/Private/Converters/GLTFMaterialConverters.h
@@ -6,14 +6,14 @@
 #include "Converters/GLTFConverter.h"
 #include "Converters/GLTFBuilderContext.h"
 #include "Converters/GLTFMeshData.h"
-#include "Converters/GLTFMaterialArray.h"
+#include "Converters/GLTFHashableArray.h"
 #include "Engine.h"
 
-class FGLTFMaterialConverter final : public FGLTFBuilderContext, public TGLTFConverter<FGLTFJsonMaterialIndex, const UMaterialInterface*, const FGLTFMeshData*, FGLTFMaterialArray>
+class FGLTFMaterialConverter final : public FGLTFBuilderContext, public TGLTFConverter<FGLTFJsonMaterialIndex, const UMaterialInterface*, const FGLTFMeshData*, FGLTFHashableArray<int32>>
 {
 	using FGLTFBuilderContext::FGLTFBuilderContext;
 
-	virtual void Sanitize(const UMaterialInterface*& Material, const FGLTFMeshData*& MeshData, FGLTFMaterialArray& Materials) override;
+	virtual void Sanitize(const UMaterialInterface*& Material, const FGLTFMeshData*& MeshData, FGLTFHashableArray<int32>& SectionIndices) override;
 
-	virtual FGLTFJsonMaterialIndex Convert(const UMaterialInterface* Material, const FGLTFMeshData* MeshData, FGLTFMaterialArray Materials) override;
+	virtual FGLTFJsonMaterialIndex Convert(const UMaterialInterface* Material, const FGLTFMeshData* MeshData, FGLTFHashableArray<int32> SectionIndices) override;
 };
diff --git a/Engine/Plugins/Enterprise/GLTFExporter/MigratedFromGit/Source/GLTFExporter/Private/Converters/GLTFMeshData.cpp b/Engine/Plugins/Enterprise/GLTFExporter/MigratedFromGit/Source/GLTFExporter/Private/Converters/GLTFMeshData.cpp
--- a/Engine/Plugins/Enterprise/GLTFExporter/MigratedFromGit/Source/GLTFExporter/Private/Converters/GLTFMeshData.cpp
+++ b/Engine/Plugins/Enterprise/GLTFExporter/MigratedFromGit/Source/GLTFExporter/Private/Converters/GLTFMeshData.cpp
@@ -12,13 +12,11 @@ FGLTFMeshData::FGLTFMeshData(const UStaticMesh* StaticMesh, const UStaticMeshCom
 	if (StaticMeshComponent != nullptr)
 	{
 		FMeshMergeHelpers::RetrieveMesh(StaticMeshComponent, LODIndex, Description, true);
-		FMeshMergeHelpers::ExtractSections(StaticMeshComponent, LODIndex, Sections);
 		Name = FGLTFNameUtility::GetName(StaticMeshComponent);
 	}
 	else
 	{
 		FMeshMergeHelpers::RetrieveMesh(StaticMesh, LODIndex, Description);
-		FMeshMergeHelpers::ExtractSections(StaticMesh, LODIndex, Sections);
 		StaticMesh->GetName(Name);
 	}
 }
@@ -30,7 +28,6 @@ FGLTFMeshData::FGLTFMeshData(const USkeletalMesh* SkeletalMesh, const USkeletalM
 	if (SkeletalMeshComponent != nullptr)
 	{
 		FMeshMergeHelpers::RetrieveMesh(const_cast<USkeletalMeshComponent*>(SkeletalMeshComponent), LODIndex, Description, true);
-		FMeshMergeHelpers::ExtractSections(SkeletalMeshComponent, LODIndex, Sections);
 		Name = FGLTFNameUtility::GetName(SkeletalMeshComponent);
 	}
 	else
diff --git a/Engine/Plugins/Enterprise/GLTFExporter/MigratedFromGit/Source/GLTFExporter/Private/Converters/GLTFMeshData.h b/Engine/Plugins/Enterprise/GLTFExporter/MigratedFromGit/Source/GLTFExporter/Private/Converters/GLTFMeshData.h
--- a/Engine/Plugins/Enterprise/GLTFExporter/MigratedFromGit/Source/GLTFExporter/Private/Converters/GLTFMeshData.h
+++ b/Engine/Plugins/Enterprise/GLTFExporter/MigratedFromGit/Source/GLTFExporter/Private/Converters/GLTFMeshData.h
@@ -11,6 +11,5 @@ struct FGLTFMeshData
 	FGLTFMeshData(const USkeletalMesh* SkeletalMesh, const USkeletalMeshComponent* SkeletalMeshComponent, int32 LODIndex);
 
 	FMeshDescription Description;
-	TArray<FSectionInfo> Sections;
 	FString Name;
 };
diff --git a/Engine/Plugins/Enterprise/GLTFExporter/MigratedFromGit/Source/GLTFExporter/Private/Converters/GLTFMeshSection.cpp b/Engine/Plugins/Enterprise/GLTFExporter/MigratedFromGit/Source/GLTFExporter/Private/Converters/GLTFMeshSection.cpp
--- a/Engine/Plugins/Enterprise/GLTFExporter/MigratedFromGit/Source/GLTFExporter/Private/Converters/GLTFMeshSection.cpp
+++ b/Engine/Plugins/Enterprise/GLTFExporter/MigratedFromGit/Source/GLTFExporter/Private/Converters/GLTFMeshSection.cpp
@@ -5,18 +5,14 @@
 #include "Rendering/SkeletalMeshRenderData.h"
 #include "Algo/MaxElement.h"
 
-FGLTFMeshSection::FGLTFMeshSection(const FStaticMeshLODResources* MeshLOD, const int32 MaterialIndex)
+FGLTFMeshSection::FGLTFMeshSection(const FStaticMeshLODResources* MeshLOD, const TArray<int32>& SectionIndices)
 {
-	uint32 TotalIndexCount = 0;
+	const FStaticMeshLODResources::FStaticMeshSectionArray& Sections = MeshLOD->Sections;
 
-	TArray<const FStaticMeshSection*> Sections;
-	for (const FStaticMeshSection& Section : MeshLOD->Sections)
+	uint32 TotalIndexCount = 0;
+	for (int32 SectionIndex : SectionIndices)
 	{
-		if (Section.MaterialIndex == MaterialIndex)
-		{
-			Sections.Add(&Section);
-			TotalIndexCount += Section.NumTriangles * 3;
-		}
+		TotalIndexCount += Sections[SectionIndex].NumTriangles * 3;
 	}
 
 	IndexMap.Reserve(TotalIndexCount);
@@ -25,10 +21,11 @@ FGLTFMeshSection::FGLTFMeshSection(const FStaticMeshLODResources* MeshLOD, const
 
 	TMap<uint32, uint32> IndexLookup;
 
-	for (const FStaticMeshSection* MeshSection : Sections)
+	for (int32 SectionIndex : SectionIndices)
 	{
-		const uint32 IndexOffset = MeshSection->FirstIndex;
-		const uint32 IndexCount = MeshSection->NumTriangles * 3;
+		const FStaticMeshSection& MeshSection = Sections[SectionIndex];
+		const uint32 IndexOffset = MeshSection.FirstIndex;
+		const uint32 IndexCount = MeshSection.NumTriangles * 3;
 
 		for (uint32 Index = 0; Index < IndexCount; Index++)
 		{
@@ -55,19 +52,15 @@ FGLTFMeshSection::FGLTFMeshSection(const FStaticMeshLODResources* MeshLOD, const
 	MaxBoneIndex = 0;
 }
 
-FGLTFMeshSection::FGLTFMeshSection(const FSkeletalMeshLODRenderData* MeshLOD, const uint16 MaterialIndex)
+FGLTFMeshSection::FGLTFMeshSection(const FSkeletalMeshLODRenderData* MeshLOD, const TArray<int32>& SectionIndices)
 	: MaxBoneIndex(0)
 {
-	uint32 TotalIndexCount = 0;
+	const TArray<FSkelMeshRenderSection>& Sections = MeshLOD->RenderSections;
 
-	TArray<const FSkelMeshRenderSection*> Sections;
-	for (const FSkelMeshRenderSection& Section : MeshLOD->RenderSections)
+	uint32 TotalIndexCount = 0;
+	for (int32 SectionIndex : SectionIndices)
 	{
-		if (Section.MaterialIndex == MaterialIndex)
-		{
-			Sections.Add(&Section);
-			TotalIndexCount += Section.NumTriangles * 3;
-		}
+		TotalIndexCount += Sections[SectionIndex].NumTriangles * 3;
 	}
 
 	IndexMap.Reserve(TotalIndexCount);
@@ -78,10 +71,11 @@ FGLTFMeshSection::FGLTFMeshSection(const FSkeletalMeshLODRenderData* MeshLOD, co
 
 	const FRawStaticIndexBuffer16or32Interface* OldIndexBuffer = MeshLOD->MultiSizeIndexContainer.GetIndexBuffer();
 
-	for (const FSkelMeshRenderSection* MeshSection : Sections)
+	for (int32 SectionIndex : SectionIndices)
 	{
-		const uint32 IndexOffset = MeshSection->BaseIndex;
-		const uint32 IndexCount = MeshSection->NumTriangles * 3;
+		const FSkelMeshRenderSection& MeshSection = Sections[SectionIndex];
+		const uint32 IndexOffset = MeshSection.BaseIndex;
+		const uint32 IndexCount = MeshSection.NumTriangles * 3;
 		const uint32 BoneMapIndex = BoneMaps.Num();
 
 		for (uint32 Index = 0; Index < IndexCount; Index++)
@@ -104,9 +98,9 @@ FGLTFMeshSection::FGLTFMeshSection(const FSkeletalMeshLODRenderData* MeshLOD, co
 			IndexBuffer[Index] = NewIndex;
 		}
 
-		BoneMaps.Add(MeshSection->BoneMap);
+		BoneMaps.Add(MeshSection.BoneMap);
 
-		if (const FBoneIndexType* MaxSectionBoneIndex = Algo::MaxElement(MeshSection->BoneMap))
+		if (const FBoneIndexType* MaxSectionBoneIndex = Algo::MaxElement(MeshSection.BoneMap))
 		{
 		    MaxBoneIndex = FMath::Max(*MaxSectionBoneIndex, MaxBoneIndex);
 		}
diff --git a/Engine/Plugins/Enterprise/GLTFExporter/MigratedFromGit/Source/GLTFExporter/Private/Converters/GLTFMeshSection.h b/Engine/Plugins/Enterprise/GLTFExporter/MigratedFromGit/Source/GLTFExporter/Private/Converters/GLTFMeshSection.h
--- a/Engine/Plugins/Enterprise/GLTFExporter/MigratedFromGit/Source/GLTFExporter/Private/Converters/GLTFMeshSection.h
+++ b/Engine/Plugins/Enterprise/GLTFExporter/MigratedFromGit/Source/GLTFExporter/Private/Converters/GLTFMeshSection.h
@@ -6,8 +6,8 @@
 
 struct FGLTFMeshSection
 {
-	FGLTFMeshSection(const FStaticMeshLODResources* MeshLOD, const int32 MaterialIndex);
-	FGLTFMeshSection(const FSkeletalMeshLODRenderData* MeshLOD, const uint16 MaterialIndex);
+	FGLTFMeshSection(const FStaticMeshLODResources* MeshLOD, const TArray<int32>& SectionIndices);
+	FGLTFMeshSection(const FSkeletalMeshLODRenderData* MeshLOD, const TArray<int32>& SectionIndices);
 
 	TArray<uint32> IndexMap;
 	TArray<uint32> IndexBuffer;
diff --git a/Engine/Plugins/Enterprise/GLTFExporter/MigratedFromGit/Source/GLTFExporter/Private/Converters/GLTFMeshSectionConverters.h b/Engine/Plugins/Enterprise/GLTFExporter/MigratedFromGit/Source/GLTFExporter/Private/Converters/GLTFMeshSectionConverters.h
--- a/Engine/Plugins/Enterprise/GLTFExporter/MigratedFromGit/Source/GLTFExporter/Private/Converters/GLTFMeshSectionConverters.h
+++ b/Engine/Plugins/Enterprise/GLTFExporter/MigratedFromGit/Source/GLTFExporter/Private/Converters/GLTFMeshSectionConverters.h
@@ -4,18 +4,18 @@
 
 #include "Converters/GLTFConverter.h"
 #include "Converters/GLTFMeshSection.h"
-#include "Engine.h"
+#include "Converters/GLTFHashableArray.h"
 
-template <typename MeshLODType, typename MaterialIndexType>
-class TGLTFMeshSectionConverter final : public TGLTFConverter<const FGLTFMeshSection*, const MeshLODType*, const MaterialIndexType>
+template <typename MeshLODType>
+class TGLTFMeshSectionConverter final : public TGLTFConverter<const FGLTFMeshSection*, const MeshLODType*, FGLTFHashableArray<int32>>
 {
 	TArray<TUniquePtr<FGLTFMeshSection>> Outputs;
 
-	const FGLTFMeshSection* Convert(const MeshLODType* MeshLOD, const MaterialIndexType MaterialIndex)
+	const FGLTFMeshSection* Convert(const MeshLODType* MeshLOD, FGLTFHashableArray<int32> SectionIndices)
 	{
-		return Outputs.Add_GetRef(MakeUnique<FGLTFMeshSection>(MeshLOD, MaterialIndex)).Get();
+		return Outputs.Add_GetRef(MakeUnique<FGLTFMeshSection>(MeshLOD, SectionIndices)).Get();
 	}
 };
 
-typedef TGLTFMeshSectionConverter<FStaticMeshLODResources, int32> FGLTFStaticMeshSectionConverter;
-typedef TGLTFMeshSectionConverter<FSkeletalMeshLODRenderData, uint16> FGLTFSkeletalMeshSectionConverter;
+typedef TGLTFMeshSectionConverter<FStaticMeshLODResources> FGLTFStaticMeshSectionConverter;
+typedef TGLTFMeshSectionConverter<FSkeletalMeshLODRenderData> FGLTFSkeletalMeshSectionConverter;
diff --git a/Engine/Plugins/Enterprise/GLTFExporter/MigratedFromGit/Source/GLTFExporter/Private/Converters/GLTFMeshUtility.cpp b/Engine/Plugins/Enterprise/GLTFExporter/MigratedFromGit/Source/GLTFExporter/Private/Converters/GLTFMeshUtility.cpp
--- a/Engine/Plugins/Enterprise/GLTFExporter/MigratedFromGit/Source/GLTFExporter/Private/Converters/GLTFMeshUtility.cpp
+++ b/Engine/Plugins/Enterprise/GLTFExporter/MigratedFromGit/Source/GLTFExporter/Private/Converters/GLTFMeshUtility.cpp
@@ -4,6 +4,42 @@
 #include "Rendering/SkeletalMeshRenderData.h"
 #include "PlatformInfo.h"
 
+TArray<int32> FGLTFMeshUtility::GetSectionIndices(const FStaticMeshLODResources& MeshLOD, int32 MaterialIndex)
+{
+	const FStaticMeshLODResources::FStaticMeshSectionArray& Sections = MeshLOD.Sections;
+
+	TArray<int32> SectionIndices;
+	SectionIndices.Reserve(Sections.Num());
+
+	for (int32 SectionIndex = 0; SectionIndex < Sections.Num(); ++SectionIndex)
+	{
+		if (Sections[SectionIndex].MaterialIndex == MaterialIndex)
+		{
+			SectionIndices.Add(SectionIndex);
+		}
+	}
+
+	return SectionIndices;
+}
+
+TArray<int32> FGLTFMeshUtility::GetSectionIndices(const FSkeletalMeshLODRenderData& MeshLOD, int32 MaterialIndex)
+{
+	const TArray<FSkelMeshRenderSection>& Sections = MeshLOD.RenderSections;
+
+	TArray<int32> SectionIndices;
+	SectionIndices.Reserve(Sections.Num());
+
+	for (int32 SectionIndex = 0; SectionIndex < Sections.Num(); ++SectionIndex)
+	{
+		if (Sections[SectionIndex].MaterialIndex == MaterialIndex)
+		{
+			SectionIndices.Add(SectionIndex);
+		}
+	}
+
+	return SectionIndices;
+}
+
 int32 FGLTFMeshUtility::GetLOD(const UStaticMesh* StaticMesh, const UStaticMeshComponent* StaticMeshComponent, int32 DefaultLOD)
 {
 	const int32 ForcedLOD = StaticMeshComponent != nullptr ? StaticMeshComponent->ForcedLodModel - 1 : -1;
diff --git a/Engine/Plugins/Enterprise/GLTFExporter/MigratedFromGit/Source/GLTFExporter/Private/Converters/GLTFMeshUtility.h b/Engine/Plugins/Enterprise/GLTFExporter/MigratedFromGit/Source/GLTFExporter/Private/Converters/GLTFMeshUtility.h
--- a/Engine/Plugins/Enterprise/GLTFExporter/MigratedFromGit/Source/GLTFExporter/Private/Converters/GLTFMeshUtility.h
+++ b/Engine/Plugins/Enterprise/GLTFExporter/MigratedFromGit/Source/GLTFExporter/Private/Converters/GLTFMeshUtility.h
@@ -4,8 +4,13 @@
 
 #include "Engine.h"
 
+class FSkeletalMeshLODRenderData;
+
 struct FGLTFMeshUtility
 {
+	static TArray<int32> GetSectionIndices(const FStaticMeshLODResources& MeshLOD, int32 MaterialIndex);
+	static TArray<int32> GetSectionIndices(const FSkeletalMeshLODRenderData& MeshLOD, int32 MaterialIndex);
+
 	static int32 GetLOD(const UStaticMesh* StaticMesh, const UStaticMeshComponent* StaticMeshComponent, int32 DefaultLOD);
 	static int32 GetLOD(const USkeletalMesh* SkeletalMesh, const USkeletalMeshComponent* SkeletalMeshComponent, int32 DefaultLOD);
 
diff --git a/Engine/Plugins/Enterprise/GLTFExporter/MigratedFromGit/Source/GLTFExporter/Private/Tasks/GLTFMaterialTasks.cpp b/Engine/Plugins/Enterprise/GLTFExporter/MigratedFromGit/Source/GLTFExporter/Private/Tasks/GLTFMaterialTasks.cpp
--- a/Engine/Plugins/Enterprise/GLTFExporter/MigratedFromGit/Source/GLTFExporter/Private/Tasks/GLTFMaterialTasks.cpp
+++ b/Engine/Plugins/Enterprise/GLTFExporter/MigratedFromGit/Source/GLTFExporter/Private/Tasks/GLTFMaterialTasks.cpp
@@ -1353,31 +1353,6 @@ FGLTFPropertyBakeOutput FGLTFMaterialTask::BakeMaterialProperty(EMaterialPropert
 	const FIntPoint DefaultTextureSize = Builder.GetDefaultMaterialBakeSize();
 	const FIntPoint TextureSize = PreferredTextureSize != nullptr ? *PreferredTextureSize : DefaultTextureSize;
 
-	TArray<int32> MeshSectionIndices;
-
-	if (MeshData != nullptr)
-	{
-		const TArray<FSectionInfo>& Sections = MeshData->Sections;
-		const TArray<UMaterialInterface*> SectionMaterials;
-
-		// Add material indices of all sections that reference the current material
-		for (int32 SectionIndex = 0; SectionIndex < Sections.Num(); ++SectionIndex)
-		{
-			const FSectionInfo& Section = Sections[SectionIndex];
-			if (Materials[Section.MaterialIndex] == Material)
-			{
-				MeshSectionIndices.Add(SectionIndex);
-			}
-		}
-
-		if (MeshSectionIndices.Num() == 0)
-		{
-			// Fall back to using the first material index,
-			// just like when baking from within the editor
-			MeshSectionIndices.Add(0);
-		}
-	}
-
 	// TODO: add support for calculating the ideal resolution to use for baking based on connected (texture) nodes
 
 	const FGLTFPropertyBakeOutput PropertyBakeOutput = FGLTFMaterialUtility::BakeMaterialProperty(
@@ -1386,7 +1361,7 @@ FGLTFPropertyBakeOutput FGLTFMaterialTask::BakeMaterialProperty(EMaterialPropert
 		Material,
 		OutTexCoord,
 		MeshData != nullptr ? &MeshData->Description : nullptr,
-		MeshSectionIndices,
+		SectionIndices,
 		bCopyAlphaFromRedChannel);
 
 	if (!PropertyBakeOutput.bIsConstant && TexCoords.Num() == 0)
diff --git a/Engine/Plugins/Enterprise/GLTFExporter/MigratedFromGit/Source/GLTFExporter/Private/Tasks/GLTFMaterialTasks.h b/Engine/Plugins/Enterprise/GLTFExporter/MigratedFromGit/Source/GLTFExporter/Private/Tasks/GLTFMaterialTasks.h
--- a/Engine/Plugins/Enterprise/GLTFExporter/MigratedFromGit/Source/GLTFExporter/Private/Tasks/GLTFMaterialTasks.h
+++ b/Engine/Plugins/Enterprise/GLTFExporter/MigratedFromGit/Source/GLTFExporter/Private/Tasks/GLTFMaterialTasks.h
@@ -13,12 +13,12 @@ class FGLTFMaterialTask : public FGLTFTask
 {
 public:
 
-	FGLTFMaterialTask(FGLTFConvertBuilder& Builder, const UMaterialInterface* Material, const FGLTFMeshData* MeshData, FGLTFMaterialArray Materials, FGLTFJsonMaterialIndex MaterialIndex)
+	FGLTFMaterialTask(FGLTFConvertBuilder& Builder, const UMaterialInterface* Material, const FGLTFMeshData* MeshData, FGLTFHashableArray<int32> SectionIndices, FGLTFJsonMaterialIndex MaterialIndex)
         : FGLTFTask(EGLTFTaskPriority::Material)
 		, Builder(Builder)
 		, Material(Material)
 		, MeshData(MeshData)
-		, Materials(Materials)
+		, SectionIndices(SectionIndices)
 		, MaterialIndex(MaterialIndex)
 	{
 	}
@@ -35,7 +35,7 @@ private:
 	FGLTFConvertBuilder& Builder;
 	const UMaterialInterface* Material;
 	const FGLTFMeshData* MeshData;
-	const FGLTFMaterialArray Materials;
+	const FGLTFHashableArray<int32> SectionIndices;
 	const FGLTFJsonMaterialIndex MaterialIndex;
 
 	bool TryGetAlphaMode(EGLTFJsonAlphaMode& AlphaMode) const;
diff --git a/Engine/Plugins/Enterprise/GLTFExporter/MigratedFromGit/Source/GLTFExporter/Private/Tasks/GLTFMeshTasks.cpp b/Engine/Plugins/Enterprise/GLTFExporter/MigratedFromGit/Source/GLTFExporter/Private/Tasks/GLTFMeshTasks.cpp
--- a/Engine/Plugins/Enterprise/GLTFExporter/MigratedFromGit/Source/GLTFExporter/Private/Tasks/GLTFMeshTasks.cpp
+++ b/Engine/Plugins/Enterprise/GLTFExporter/MigratedFromGit/Source/GLTFExporter/Private/Tasks/GLTFMeshTasks.cpp
@@ -2,6 +2,7 @@
 
 #include "Tasks/GLTFMeshTasks.h"
 #include "Converters/GLTFConverterUtility.h"
+#include "Converters/GLTFMeshUtility.h"
 #include "Builders/GLTFConvertBuilder.h"
 #include "Rendering/SkeletalMeshRenderData.h"
 
@@ -30,7 +31,8 @@ void FGLTFStaticMeshTask::Complete()
 
 	for (int32 MaterialIndex = 0; MaterialIndex < MaterialCount; ++MaterialIndex)
 	{
-		const FGLTFMeshSection* ConvertedSection = MeshSectionConverter.GetOrAdd(&MeshLOD, MaterialIndex);
+		const FGLTFHashableArray<int32> SectionIndices = FGLTFMeshUtility::GetSectionIndices(MeshLOD, MaterialIndex);
+		const FGLTFMeshSection* ConvertedSection = MeshSectionConverter.GetOrAdd(&MeshLOD, SectionIndices);
 
 		FGLTFJsonPrimitive& JsonPrimitive = JsonMesh.Primitives[MaterialIndex];
 		JsonPrimitive.Indices = Builder.GetOrAddIndexAccessor(ConvertedSection);
@@ -58,7 +60,7 @@ void FGLTFStaticMeshTask::Complete()
 		}
 
 		const UMaterialInterface* Material = Materials[MaterialIndex];
-		JsonPrimitive.Material =  Builder.GetOrAddMaterial(Material, MeshData, Materials);
+		JsonPrimitive.Material =  Builder.GetOrAddMaterial(Material, MeshData, SectionIndices);
 	}
 }
 
@@ -96,7 +98,8 @@ void FGLTFSkeletalMeshTask::Complete()
 
 	for (uint16 MaterialIndex = 0; MaterialIndex < MaterialCount; ++MaterialIndex)
 	{
-		const FGLTFMeshSection* ConvertedSection = MeshSectionConverter.GetOrAdd(&MeshLOD, MaterialIndex);
+		const FGLTFHashableArray<int32> SectionIndices = FGLTFMeshUtility::GetSectionIndices(MeshLOD, MaterialIndex);
+		const FGLTFMeshSection* ConvertedSection = MeshSectionConverter.GetOrAdd(&MeshLOD, SectionIndices);
 
 		FGLTFJsonPrimitive& JsonPrimitive = JsonMesh.Primitives[MaterialIndex];
 		JsonPrimitive.Indices = Builder.GetOrAddIndexAccessor(ConvertedSection);
@@ -137,6 +140,6 @@ void FGLTFSkeletalMeshTask::Complete()
 		}
 
 		const UMaterialInterface* Material = Materials[MaterialIndex];
-		JsonPrimitive.Material =  Builder.GetOrAddMaterial(Material, MeshData, Materials);
+		JsonPrimitive.Material =  Builder.GetOrAddMaterial(Material, MeshData, SectionIndices);
 	}
 }
