[GLTFExporter] refactor: remove GLTFHashableArray and simplify GLTFMaterialArray to a simple typedef

commit e1f6ab45f7efbcb70ba54691e2cab276cc57c979
Date: Tue, 3 Nov 2020 23:21:16 +0100

    refactor: remove GLTFHashableArray and simplify GLTFMaterialArray to a simple typedef
    (cherry picked from commit 2e97483384e45bf4884d4cf2b72d4d01e6ed2ac1​)

#skipci

diff --git a/Engine/Plugins/Enterprise/GLTFExporter/MigratedFromGit/Source/GLTFExporter/Private/Builders/GLTFConvertBuilder.cpp b/Engine/Plugins/Enterprise/GLTFExporter/MigratedFromGit/Source/GLTFExporter/Private/Builders/GLTFConvertBuilder.cpp
--- a/Engine/Plugins/Enterprise/GLTFExporter/MigratedFromGit/Source/GLTFExporter/Private/Builders/GLTFConvertBuilder.cpp
+++ b/Engine/Plugins/Enterprise/GLTFExporter/MigratedFromGit/Source/GLTFExporter/Private/Builders/GLTFConvertBuilder.cpp
@@ -140,7 +140,7 @@ FGLTFJsonMeshIndex FGLTFConvertBuilder::GetOrAddMesh(const USkeletalMeshComponen
 	return SkeletalMeshConverter.GetOrAdd(SkeletalMesh, SkeletalMeshComponent, Materials, LODIndex);
 }
 
-FGLTFJsonMaterialIndex FGLTFConvertBuilder::GetOrAddMaterial(const UMaterialInterface* Material, const FGLTFMeshData* MeshData, const FGLTFHashableArray<int32>& SectionIndices)
+FGLTFJsonMaterialIndex FGLTFConvertBuilder::GetOrAddMaterial(const UMaterialInterface* Material, const FGLTFMeshData* MeshData, const TArray<int32>& SectionIndices)
 {
 	if (Material == nullptr)
 	{
diff --git a/Engine/Plugins/Enterprise/GLTFExporter/MigratedFromGit/Source/GLTFExporter/Private/Builders/GLTFConvertBuilder.h b/Engine/Plugins/Enterprise/GLTFExporter/MigratedFromGit/Source/GLTFExporter/Private/Builders/GLTFConvertBuilder.h
--- a/Engine/Plugins/Enterprise/GLTFExporter/MigratedFromGit/Source/GLTFExporter/Private/Builders/GLTFConvertBuilder.h
+++ b/Engine/Plugins/Enterprise/GLTFExporter/MigratedFromGit/Source/GLTFExporter/Private/Builders/GLTFConvertBuilder.h
@@ -44,7 +44,7 @@ public:
 	FGLTFJsonMeshIndex GetOrAddMesh(const USkeletalMesh* SkeletalMesh, const FGLTFMaterialArray& Materials = {}, int32 LODIndex = -1);
 	FGLTFJsonMeshIndex GetOrAddMesh(const USkeletalMeshComponent* SkeletalMeshComponent, const FGLTFMaterialArray& Materials = {}, int32 LODIndex = -1);
 
-	FGLTFJsonMaterialIndex GetOrAddMaterial(const UMaterialInterface* Material, const FGLTFMeshData* MeshData = nullptr, const FGLTFHashableArray<int32>& SectionIndices = {});
+	FGLTFJsonMaterialIndex GetOrAddMaterial(const UMaterialInterface* Material, const FGLTFMeshData* MeshData = nullptr, const TArray<int32>& SectionIndices = {});
 
 	FGLTFJsonSamplerIndex GetOrAddSampler(const UTexture* Texture);
 	FGLTFJsonTextureIndex GetOrAddTexture(const UTexture2D* Texture);
diff --git a/Engine/Plugins/Enterprise/GLTFExporter/MigratedFromGit/Source/GLTFExporter/Private/Converters/GLTFConverter.h b/Engine/Plugins/Enterprise/GLTFExporter/MigratedFromGit/Source/GLTFExporter/Private/Converters/GLTFConverter.h
--- a/Engine/Plugins/Enterprise/GLTFExporter/MigratedFromGit/Source/GLTFExporter/Private/Converters/GLTFConverter.h
+++ b/Engine/Plugins/Enterprise/GLTFExporter/MigratedFromGit/Source/GLTFExporter/Private/Converters/GLTFConverter.h
@@ -28,7 +28,7 @@ public:
 	OutputType Add(InputTypes... Inputs)
 	{
 		Sanitize(Inputs...);
-		const InputKeyType InputKey(Forward<InputTypes>(Inputs)...);
+		const InputKeyType InputKey(Inputs...);
 		OutputType NewOutput = Convert(Forward<InputTypes>(Inputs)...);
 
 		SavedOutputs.Add(InputKey, NewOutput);
@@ -38,7 +38,7 @@ public:
 	OutputType GetOrAdd(InputTypes... Inputs)
 	{
 		Sanitize(Inputs...);
-		const InputKeyType InputKey(Forward<InputTypes>(Inputs)...);
+		const InputKeyType InputKey(Inputs...);
 		if (OutputType* SavedOutput = SavedOutputs.Find(InputKey))
 		{
 			return *SavedOutput;
diff --git a/Engine/Plugins/Enterprise/GLTFExporter/MigratedFromGit/Source/GLTFExporter/Private/Converters/GLTFHashableArray.h b/Engine/Plugins/Enterprise/GLTFExporter/MigratedFromGit/Source/GLTFExporter/Private/Converters/GLTFHashableArray.h
deleted file mode 100644
diff --git a/Engine/Plugins/Enterprise/GLTFExporter/MigratedFromGit/Source/GLTFExporter/Private/Converters/GLTFMaterialArray.h b/Engine/Plugins/Enterprise/GLTFExporter/MigratedFromGit/Source/GLTFExporter/Private/Converters/GLTFMaterialArray.h
--- a/Engine/Plugins/Enterprise/GLTFExporter/MigratedFromGit/Source/GLTFExporter/Private/Converters/GLTFMaterialArray.h
+++ b/Engine/Plugins/Enterprise/GLTFExporter/MigratedFromGit/Source/GLTFExporter/Private/Converters/GLTFMaterialArray.h
@@ -2,66 +2,18 @@
 
 #pragma once
 
-#include "GLTFHashableArray.h"
+#include "CoreMinimal.h"
 
 class UMaterialInterface;
 
-struct FGLTFMaterialArray : FGLTFHashableArray<const UMaterialInterface*>
-{
-	using FGLTFHashableArray::FGLTFHashableArray;
-	using FGLTFHashableArray::operator=;
-	using FGLTFHashableArray::operator==;
-	using FGLTFHashableArray::operator!=;
-
-	void FillIn(const TArray<UMaterialInterface*>& Defaults)
-	{
-		return FillIn<UMaterialInterface*>(Defaults);
-	}
-
-	void FillIn(const TArray<FStaticMaterial>& Defaults)
-	{
-		return FillIn<FStaticMaterial>(Defaults);
-	}
-
-	void FillIn(const TArray<FSkeletalMaterial>& Defaults)
-	{
-		return FillIn<FSkeletalMaterial>(Defaults);
-	}
-
-private:
+typedef TArray<const UMaterialInterface*> FGLTFMaterialArray;
 
-	template <typename MeshMaterialType>
-	void FillIn(const TArray<MeshMaterialType>& Defaults)
-	{
-		const int32 Count = Defaults.Num();
-		SetNumZeroed(Count);
-
-		for (int32 Index = 0; Index < Count; ++Index)
-		{
-			const UMaterialInterface*& Element = (*this)[Index];
-			if (Element == nullptr)
-			{
-				Element = GetInterface(Defaults[Index]);
-				if (Element == nullptr)
-				{
-					Element = UMaterial::GetDefaultMaterial(MD_Surface);
-				}
-			}
-		}
-	}
-
-	static const UMaterialInterface* GetInterface(const UMaterialInterface* Material)
-	{
-		return Material;
-	}
-
-	static const UMaterialInterface* GetInterface(const FStaticMaterial& Material)
-	{
-		return Material.MaterialInterface;
-	}
-
-	static const UMaterialInterface* GetInterface(const FSkeletalMaterial& Material)
+inline uint32 GetTypeHash(const FGLTFMaterialArray& MaterialArray)
+{
+	uint32 Hash = GetTypeHash(MaterialArray.Num());
+	for (const UMaterialInterface* Material : MaterialArray)
 	{
-		return Material.MaterialInterface;
+		Hash = HashCombine(Hash, GetTypeHash(Material));
 	}
-};
+	return Hash;
+}
diff --git a/Engine/Plugins/Enterprise/GLTFExporter/MigratedFromGit/Source/GLTFExporter/Private/Converters/GLTFMaterialConverters.cpp b/Engine/Plugins/Enterprise/GLTFExporter/MigratedFromGit/Source/GLTFExporter/Private/Converters/GLTFMaterialConverters.cpp
--- a/Engine/Plugins/Enterprise/GLTFExporter/MigratedFromGit/Source/GLTFExporter/Private/Converters/GLTFMaterialConverters.cpp
+++ b/Engine/Plugins/Enterprise/GLTFExporter/MigratedFromGit/Source/GLTFExporter/Private/Converters/GLTFMaterialConverters.cpp
@@ -6,7 +6,7 @@
 #include "Builders/GLTFConvertBuilder.h"
 #include "Tasks/GLTFMaterialTasks.h"
 
-void FGLTFMaterialConverter::Sanitize(const UMaterialInterface*& Material, const FGLTFMeshData*& MeshData, FGLTFHashableArray<int32>& SectionIndices)
+void FGLTFMaterialConverter::Sanitize(const UMaterialInterface*& Material, const FGLTFMeshData*& MeshData, TArray<int32>& SectionIndices)
 {
 	if (MeshData == nullptr ||
 		!Builder.ExportOptions->bBakeMaterialInputs ||
@@ -18,7 +18,7 @@ void FGLTFMaterialConverter::Sanitize(const UMaterialInterface*& Material, const
 	}
 }
 
-FGLTFJsonMaterialIndex FGLTFMaterialConverter::Convert(const UMaterialInterface* Material, const FGLTFMeshData* MeshData, FGLTFHashableArray<int32> SectionIndices)
+FGLTFJsonMaterialIndex FGLTFMaterialConverter::Convert(const UMaterialInterface* Material, const FGLTFMeshData* MeshData, TArray<int32> SectionIndices)
 {
 	if (Material == FGLTFMaterialUtility::GetDefault())
 	{
diff --git a/Engine/Plugins/Enterprise/GLTFExporter/MigratedFromGit/Source/GLTFExporter/Private/Converters/GLTFMaterialConverters.h b/Engine/Plugins/Enterprise/GLTFExporter/MigratedFromGit/Source/GLTFExporter/Private/Converters/GLTFMaterialConverters.h
--- a/Engine/Plugins/Enterprise/GLTFExporter/MigratedFromGit/Source/GLTFExporter/Private/Converters/GLTFMaterialConverters.h
+++ b/Engine/Plugins/Enterprise/GLTFExporter/MigratedFromGit/Source/GLTFExporter/Private/Converters/GLTFMaterialConverters.h
@@ -6,14 +6,13 @@
 #include "Converters/GLTFConverter.h"
 #include "Converters/GLTFBuilderContext.h"
 #include "Converters/GLTFMeshData.h"
-#include "Converters/GLTFHashableArray.h"
 #include "Engine.h"
 
-class FGLTFMaterialConverter final : public FGLTFBuilderContext, public TGLTFConverter<FGLTFJsonMaterialIndex, const UMaterialInterface*, const FGLTFMeshData*, FGLTFHashableArray<int32>>
+class FGLTFMaterialConverter final : public FGLTFBuilderContext, public TGLTFConverter<FGLTFJsonMaterialIndex, const UMaterialInterface*, const FGLTFMeshData*, TArray<int32>>
 {
 	using FGLTFBuilderContext::FGLTFBuilderContext;
 
-	virtual void Sanitize(const UMaterialInterface*& Material, const FGLTFMeshData*& MeshData, FGLTFHashableArray<int32>& SectionIndices) override;
+	virtual void Sanitize(const UMaterialInterface*& Material, const FGLTFMeshData*& MeshData, TArray<int32>& SectionIndices) override;
 
-	virtual FGLTFJsonMaterialIndex Convert(const UMaterialInterface* Material, const FGLTFMeshData* MeshData, FGLTFHashableArray<int32> SectionIndices) override;
+	virtual FGLTFJsonMaterialIndex Convert(const UMaterialInterface* Material, const FGLTFMeshData* MeshData, TArray<int32> SectionIndices) override;
 };
diff --git a/Engine/Plugins/Enterprise/GLTFExporter/MigratedFromGit/Source/GLTFExporter/Private/Converters/GLTFMaterialUtility.cpp b/Engine/Plugins/Enterprise/GLTFExporter/MigratedFromGit/Source/GLTFExporter/Private/Converters/GLTFMaterialUtility.cpp
--- a/Engine/Plugins/Enterprise/GLTFExporter/MigratedFromGit/Source/GLTFExporter/Private/Converters/GLTFMaterialUtility.cpp
+++ b/Engine/Plugins/Enterprise/GLTFExporter/MigratedFromGit/Source/GLTFExporter/Private/Converters/GLTFMaterialUtility.cpp
@@ -552,3 +552,53 @@ bool FGLTFMaterialUtility::MaterialsNeedVertexData(const TArray<const UMaterialI
 
 	return false;
 }
+
+const UMaterialInterface* FGLTFMaterialUtility::GetInterface(const UMaterialInterface* Material)
+{
+	return Material;
+}
+
+const UMaterialInterface* FGLTFMaterialUtility::GetInterface(const FStaticMaterial& Material)
+{
+	return Material.MaterialInterface;
+}
+
+const UMaterialInterface* FGLTFMaterialUtility::GetInterface(const FSkeletalMaterial& Material)
+{
+	return Material.MaterialInterface;
+}
+
+void FGLTFMaterialUtility::ResolveOverrides(TArray<const UMaterialInterface*>& Overrides, const TArray<UMaterialInterface*>& Defaults)
+{
+	ResolveOverrides<UMaterialInterface*>(Overrides, Defaults);
+}
+
+void FGLTFMaterialUtility::ResolveOverrides(TArray<const UMaterialInterface*>& Overrides, const TArray<FStaticMaterial>& Defaults)
+{
+	ResolveOverrides<FStaticMaterial>(Overrides, Defaults);
+}
+
+void FGLTFMaterialUtility::ResolveOverrides(TArray<const UMaterialInterface*>& Overrides, const TArray<FSkeletalMaterial>& Defaults)
+{
+	ResolveOverrides<FSkeletalMaterial>(Overrides, Defaults);
+}
+
+template <typename MaterialType>
+void FGLTFMaterialUtility::ResolveOverrides(TArray<const UMaterialInterface*>& Overrides, const TArray<MaterialType>& Defaults)
+{
+	const int32 Count = Defaults.Num();
+	Overrides.SetNumZeroed(Count);
+
+	for (int32 Index = 0; Index < Count; ++Index)
+	{
+		const UMaterialInterface*& Element = Overrides[Index];
+		if (Element == nullptr)
+		{
+			Element = GetInterface(Defaults[Index]);
+			if (Element == nullptr)
+			{
+				Element = UMaterial::GetDefaultMaterial(MD_Surface);
+			}
+		}
+	}
+}
diff --git a/Engine/Plugins/Enterprise/GLTFExporter/MigratedFromGit/Source/GLTFExporter/Private/Converters/GLTFMaterialUtility.h b/Engine/Plugins/Enterprise/GLTFExporter/MigratedFromGit/Source/GLTFExporter/Private/Converters/GLTFMaterialUtility.h
--- a/Engine/Plugins/Enterprise/GLTFExporter/MigratedFromGit/Source/GLTFExporter/Private/Converters/GLTFMaterialUtility.h
+++ b/Engine/Plugins/Enterprise/GLTFExporter/MigratedFromGit/Source/GLTFExporter/Private/Converters/GLTFMaterialUtility.h
@@ -70,4 +70,17 @@ struct FGLTFMaterialUtility
 
 	static bool MaterialNeedsVertexData(const UMaterialInterface* Material);
 	static bool MaterialsNeedVertexData(const TArray<const UMaterialInterface*>& Materials);
+
+	static const UMaterialInterface* GetInterface(const UMaterialInterface* Material);
+	static const UMaterialInterface* GetInterface(const FStaticMaterial& Material);
+	static const UMaterialInterface* GetInterface(const FSkeletalMaterial& Material);
+
+	static void ResolveOverrides(TArray<const UMaterialInterface*>& Overrides, const TArray<UMaterialInterface*>& Defaults);
+	static void ResolveOverrides(TArray<const UMaterialInterface*>& Overrides, const TArray<FStaticMaterial>& Defaults);
+	static void ResolveOverrides(TArray<const UMaterialInterface*>& Overrides, const TArray<FSkeletalMaterial>& Defaults);
+
+private:
+
+	template <typename MaterialType>
+	static void ResolveOverrides(TArray<const UMaterialInterface*>& Overrides, const TArray<MaterialType>& Defaults);
 };
diff --git a/Engine/Plugins/Enterprise/GLTFExporter/MigratedFromGit/Source/GLTFExporter/Private/Converters/GLTFMeshConverters.cpp b/Engine/Plugins/Enterprise/GLTFExporter/MigratedFromGit/Source/GLTFExporter/Private/Converters/GLTFMeshConverters.cpp
--- a/Engine/Plugins/Enterprise/GLTFExporter/MigratedFromGit/Source/GLTFExporter/Private/Converters/GLTFMeshConverters.cpp
+++ b/Engine/Plugins/Enterprise/GLTFExporter/MigratedFromGit/Source/GLTFExporter/Private/Converters/GLTFMeshConverters.cpp
@@ -10,11 +10,11 @@ void FGLTFStaticMeshConverter::Sanitize(const UStaticMesh*& StaticMesh, const US
 {
 	if (StaticMeshComponent != nullptr)
 	{
-		Materials.FillIn(StaticMeshComponent->GetMaterials());
+		FGLTFMaterialUtility::ResolveOverrides(Materials, StaticMeshComponent->GetMaterials());
 	}
 	else
 	{
-		Materials.FillIn(StaticMesh->StaticMaterials);
+		FGLTFMaterialUtility::ResolveOverrides(Materials, StaticMesh->StaticMaterials);
 	}
 
 	if (LODIndex < 0)
@@ -50,11 +50,11 @@ void FGLTFSkeletalMeshConverter::Sanitize(const USkeletalMesh*& SkeletalMesh, co
 {
 	if (SkeletalMeshComponent != nullptr)
 	{
-		Materials.FillIn(SkeletalMeshComponent->GetMaterials());
+		FGLTFMaterialUtility::ResolveOverrides(Materials, SkeletalMeshComponent->GetMaterials());
 	}
 	else
 	{
-		Materials.FillIn(SkeletalMesh->Materials);
+		FGLTFMaterialUtility::ResolveOverrides(Materials, SkeletalMesh->Materials);
 	}
 
 	if (LODIndex < 0)
diff --git a/Engine/Plugins/Enterprise/GLTFExporter/MigratedFromGit/Source/GLTFExporter/Private/Converters/GLTFMeshSectionConverters.h b/Engine/Plugins/Enterprise/GLTFExporter/MigratedFromGit/Source/GLTFExporter/Private/Converters/GLTFMeshSectionConverters.h
--- a/Engine/Plugins/Enterprise/GLTFExporter/MigratedFromGit/Source/GLTFExporter/Private/Converters/GLTFMeshSectionConverters.h
+++ b/Engine/Plugins/Enterprise/GLTFExporter/MigratedFromGit/Source/GLTFExporter/Private/Converters/GLTFMeshSectionConverters.h
@@ -4,14 +4,13 @@
 
 #include "Converters/GLTFConverter.h"
 #include "Converters/GLTFMeshSection.h"
-#include "Converters/GLTFHashableArray.h"
 
 template <typename MeshLODType>
-class TGLTFMeshSectionConverter final : public TGLTFConverter<const FGLTFMeshSection*, const MeshLODType*, FGLTFHashableArray<int32>>
+class TGLTFMeshSectionConverter final : public TGLTFConverter<const FGLTFMeshSection*, const MeshLODType*, TArray<int32>>
 {
 	TArray<TUniquePtr<FGLTFMeshSection>> Outputs;
 
-	const FGLTFMeshSection* Convert(const MeshLODType* MeshLOD, FGLTFHashableArray<int32> SectionIndices)
+	const FGLTFMeshSection* Convert(const MeshLODType* MeshLOD, TArray<int32> SectionIndices)
 	{
 		return Outputs.Add_GetRef(MakeUnique<FGLTFMeshSection>(MeshLOD, SectionIndices)).Get();
 	}
diff --git a/Engine/Plugins/Enterprise/GLTFExporter/MigratedFromGit/Source/GLTFExporter/Private/Tasks/GLTFMaterialTasks.h b/Engine/Plugins/Enterprise/GLTFExporter/MigratedFromGit/Source/GLTFExporter/Private/Tasks/GLTFMaterialTasks.h
--- a/Engine/Plugins/Enterprise/GLTFExporter/MigratedFromGit/Source/GLTFExporter/Private/Tasks/GLTFMaterialTasks.h
+++ b/Engine/Plugins/Enterprise/GLTFExporter/MigratedFromGit/Source/GLTFExporter/Private/Tasks/GLTFMaterialTasks.h
@@ -13,7 +13,7 @@ class FGLTFMaterialTask : public FGLTFTask
 {
 public:
 
-	FGLTFMaterialTask(FGLTFConvertBuilder& Builder, const UMaterialInterface* Material, const FGLTFMeshData* MeshData, FGLTFHashableArray<int32> SectionIndices, FGLTFJsonMaterialIndex MaterialIndex)
+	FGLTFMaterialTask(FGLTFConvertBuilder& Builder, const UMaterialInterface* Material, const FGLTFMeshData* MeshData, TArray<int32> SectionIndices, FGLTFJsonMaterialIndex MaterialIndex)
         : FGLTFTask(EGLTFTaskPriority::Material)
 		, Builder(Builder)
 		, Material(Material)
@@ -35,7 +35,7 @@ private:
 	FGLTFConvertBuilder& Builder;
 	const UMaterialInterface* Material;
 	const FGLTFMeshData* MeshData;
-	const FGLTFHashableArray<int32> SectionIndices;
+	const TArray<int32> SectionIndices;
 	const FGLTFJsonMaterialIndex MaterialIndex;
 
 	bool TryGetAlphaMode(EGLTFJsonAlphaMode& AlphaMode) const;
diff --git a/Engine/Plugins/Enterprise/GLTFExporter/MigratedFromGit/Source/GLTFExporter/Private/Tasks/GLTFMeshTasks.cpp b/Engine/Plugins/Enterprise/GLTFExporter/MigratedFromGit/Source/GLTFExporter/Private/Tasks/GLTFMeshTasks.cpp
--- a/Engine/Plugins/Enterprise/GLTFExporter/MigratedFromGit/Source/GLTFExporter/Private/Tasks/GLTFMeshTasks.cpp
+++ b/Engine/Plugins/Enterprise/GLTFExporter/MigratedFromGit/Source/GLTFExporter/Private/Tasks/GLTFMeshTasks.cpp
@@ -31,7 +31,7 @@ void FGLTFStaticMeshTask::Complete()
 
 	for (int32 MaterialIndex = 0; MaterialIndex < MaterialCount; ++MaterialIndex)
 	{
-		const FGLTFHashableArray<int32> SectionIndices = FGLTFMeshUtility::GetSectionIndices(MeshLOD, MaterialIndex);
+		const TArray<int32> SectionIndices = FGLTFMeshUtility::GetSectionIndices(MeshLOD, MaterialIndex);
 		const FGLTFMeshSection* ConvertedSection = MeshSectionConverter.GetOrAdd(&MeshLOD, SectionIndices);
 
 		FGLTFJsonPrimitive& JsonPrimitive = JsonMesh.Primitives[MaterialIndex];
@@ -98,7 +98,7 @@ void FGLTFSkeletalMeshTask::Complete()
 
 	for (uint16 MaterialIndex = 0; MaterialIndex < MaterialCount; ++MaterialIndex)
 	{
-		const FGLTFHashableArray<int32> SectionIndices = FGLTFMeshUtility::GetSectionIndices(MeshLOD, MaterialIndex);
+		const TArray<int32> SectionIndices = FGLTFMeshUtility::GetSectionIndices(MeshLOD, MaterialIndex);
 		const FGLTFMeshSection* ConvertedSection = MeshSectionConverter.GetOrAdd(&MeshLOD, SectionIndices);
 
 		FGLTFJsonPrimitive& JsonPrimitive = JsonMesh.Primitives[MaterialIndex];
