[GLTFExporter] refactor: rename Material(s)NeedVertexData to NeedsVertexData and add some TODO comments

commit a38bf8e50df0d1326919e71a2e1e641e1c905381
Date: Wed, 4 Nov 2020 19:16:25 +0100

    refactor: rename Material(s)NeedVertexData to NeedsVertexData and add some TODO comments
    (cherry picked from commit 7269950ab39db703e0be3b7dfc74f93cb750637a​)

#skipci

diff --git a/Engine/Plugins/Enterprise/GLTFExporter/MigratedFromGit/Source/GLTFExporter/Private/Converters/GLTFMaterialConverters.cpp b/Engine/Plugins/Enterprise/GLTFExporter/MigratedFromGit/Source/GLTFExporter/Private/Converters/GLTFMaterialConverters.cpp
--- a/Engine/Plugins/Enterprise/GLTFExporter/MigratedFromGit/Source/GLTFExporter/Private/Converters/GLTFMaterialConverters.cpp
+++ b/Engine/Plugins/Enterprise/GLTFExporter/MigratedFromGit/Source/GLTFExporter/Private/Converters/GLTFMaterialConverters.cpp
@@ -11,7 +11,7 @@ void FGLTFMaterialConverter::Sanitize(const UMaterialInterface*& Material, const
 	if (MeshData == nullptr ||
 		!Builder.ExportOptions->bBakeMaterialInputs ||
 		!Builder.ExportOptions->bBakeMaterialInputsUsingMeshData ||
-		!FGLTFMaterialUtility::MaterialNeedsVertexData(Material))
+		!FGLTFMaterialUtility::NeedsMeshData(Material))
 	{
 		MeshData = nullptr;
 		SectionIndices = {};
diff --git a/Engine/Plugins/Enterprise/GLTFExporter/MigratedFromGit/Source/GLTFExporter/Private/Converters/GLTFMaterialUtility.cpp b/Engine/Plugins/Enterprise/GLTFExporter/MigratedFromGit/Source/GLTFExporter/Private/Converters/GLTFMaterialUtility.cpp
--- a/Engine/Plugins/Enterprise/GLTFExporter/MigratedFromGit/Source/GLTFExporter/Private/Converters/GLTFMaterialUtility.cpp
+++ b/Engine/Plugins/Enterprise/GLTFExporter/MigratedFromGit/Source/GLTFExporter/Private/Converters/GLTFMaterialUtility.cpp
@@ -502,7 +502,7 @@ FString FGLTFMaterialUtility::ShadingModelsToString(const FMaterialShadingModelF
 	return Result;
 }
 
-bool FGLTFMaterialUtility::MaterialNeedsVertexData(const UMaterialInterface* Material)
+bool FGLTFMaterialUtility::NeedsMeshData(const UMaterialInterface* Material)
 {
 	if (Material == nullptr)
 	{
@@ -510,7 +510,6 @@ bool FGLTFMaterialUtility::MaterialNeedsVertexData(const UMaterialInterface* Mat
 	}
 
 	// TODO: only analyze properties that will be needed for this specific material
-
 	const TArray<EMaterialProperty> Properties =
 	{
 		MP_BaseColor,
@@ -520,10 +519,10 @@ bool FGLTFMaterialUtility::MaterialNeedsVertexData(const UMaterialInterface* Mat
 		MP_Metallic,
 		MP_Roughness,
 		MP_Normal,
-		// MP_CustomOutput,	// NOTE: causes a crash when used with AnalyzeMaterialProperty
 		MP_AmbientOcclusion,
 		MP_CustomData0,
-		MP_CustomData1
+		MP_CustomData1,
+		// TODO: add support for custom output ClearCoatBottomNormal
 	};
 
 	bool bRequiresVertexData = false;
@@ -540,11 +539,11 @@ bool FGLTFMaterialUtility::MaterialNeedsVertexData(const UMaterialInterface* Mat
 	return bRequiresVertexData;
 }
 
-bool FGLTFMaterialUtility::MaterialsNeedVertexData(const TArray<const UMaterialInterface*>& Materials)
+bool FGLTFMaterialUtility::NeedsMeshData(const TArray<const UMaterialInterface*>& Materials)
 {
 	for (const UMaterialInterface* Material: Materials)
 	{
-		if (MaterialNeedsVertexData(Material))
+		if (NeedsMeshData(Material))
 		{
 			return true;
 		}
diff --git a/Engine/Plugins/Enterprise/GLTFExporter/MigratedFromGit/Source/GLTFExporter/Private/Converters/GLTFMaterialUtility.h b/Engine/Plugins/Enterprise/GLTFExporter/MigratedFromGit/Source/GLTFExporter/Private/Converters/GLTFMaterialUtility.h
--- a/Engine/Plugins/Enterprise/GLTFExporter/MigratedFromGit/Source/GLTFExporter/Private/Converters/GLTFMaterialUtility.h
+++ b/Engine/Plugins/Enterprise/GLTFExporter/MigratedFromGit/Source/GLTFExporter/Private/Converters/GLTFMaterialUtility.h
@@ -68,8 +68,8 @@ struct FGLTFMaterialUtility
 
 	static FString ShadingModelsToString(const FMaterialShadingModelField& ShadingModels);
 
-	static bool MaterialNeedsVertexData(const UMaterialInterface* Material);
-	static bool MaterialsNeedVertexData(const TArray<const UMaterialInterface*>& Materials);
+	static bool NeedsMeshData(const UMaterialInterface* Material);
+	static bool NeedsMeshData(const TArray<const UMaterialInterface*>& Materials);
 
 	static const UMaterialInterface* GetInterface(const UMaterialInterface* Material);
 	static const UMaterialInterface* GetInterface(const FStaticMaterial& Material);
diff --git a/Engine/Plugins/Enterprise/GLTFExporter/MigratedFromGit/Source/GLTFExporter/Private/Converters/GLTFMeshConverters.cpp b/Engine/Plugins/Enterprise/GLTFExporter/MigratedFromGit/Source/GLTFExporter/Private/Converters/GLTFMeshConverters.cpp
--- a/Engine/Plugins/Enterprise/GLTFExporter/MigratedFromGit/Source/GLTFExporter/Private/Converters/GLTFMeshConverters.cpp
+++ b/Engine/Plugins/Enterprise/GLTFExporter/MigratedFromGit/Source/GLTFExporter/Private/Converters/GLTFMeshConverters.cpp
@@ -32,7 +32,7 @@ void FGLTFStaticMeshConverter::Sanitize(const UStaticMesh*& StaticMesh, const US
 		// otherwise export a copy of this mesh for each mesh-component.
 		if (!Builder.ExportOptions->bBakeMaterialInputs ||
 			!Builder.ExportOptions->bBakeMaterialInputsUsingMeshData ||
-			!FGLTFMaterialUtility::MaterialsNeedVertexData(Materials))
+			!FGLTFMaterialUtility::NeedsMeshData(Materials)) // TODO: if this expensive, cache the results for each material
 		{
 			StaticMeshComponent = nullptr;
 		}
@@ -72,7 +72,7 @@ void FGLTFSkeletalMeshConverter::Sanitize(const USkeletalMesh*& SkeletalMesh, co
 		// otherwise export a copy of this mesh for each mesh-component.
 		if (!Builder.ExportOptions->bBakeMaterialInputs ||
 			!Builder.ExportOptions->bBakeMaterialInputsUsingMeshData ||
-			!FGLTFMaterialUtility::MaterialsNeedVertexData(Materials))
+			!FGLTFMaterialUtility::NeedsMeshData(Materials)) // TODO: if this expensive, cache the results for each material
 		{
 			SkeletalMeshComponent = nullptr;
 		}
