[GLTFExporter] refactor: rename export option bBakeMaterialInputsUsingMeshData to bMaterialBakeUsingMeshData

commit fcd1fa1d4ad30293137087c7355d49fcb0bf0a68
Date: Wed, 4 Nov 2020 19:33:44 +0100

    refactor: rename export option bBakeMaterialInputsUsingMeshData to bMaterialBakeUsingMeshData
    (cherry picked from commit 81f7f8a11ad78c44bf8d0199c259528b8251da6f​)

#skipci

diff --git a/Engine/Plugins/Enterprise/GLTFExporter/MigratedFromGit/Source/GLTFExporter/Private/Converters/GLTFMaterialConverters.cpp b/Engine/Plugins/Enterprise/GLTFExporter/MigratedFromGit/Source/GLTFExporter/Private/Converters/GLTFMaterialConverters.cpp
--- a/Engine/Plugins/Enterprise/GLTFExporter/MigratedFromGit/Source/GLTFExporter/Private/Converters/GLTFMaterialConverters.cpp
+++ b/Engine/Plugins/Enterprise/GLTFExporter/MigratedFromGit/Source/GLTFExporter/Private/Converters/GLTFMaterialConverters.cpp
@@ -10,7 +10,7 @@ void FGLTFMaterialConverter::Sanitize(const UMaterialInterface*& Material, const
 {
 	if (MeshData == nullptr ||
 		!Builder.ExportOptions->bBakeMaterialInputs ||
-		!Builder.ExportOptions->bBakeMaterialInputsUsingMeshData ||
+		!Builder.ExportOptions->bMaterialBakeUsingMeshData ||
 		!FGLTFMaterialUtility::NeedsMeshData(Material))
 	{
 		MeshData = nullptr;
diff --git a/Engine/Plugins/Enterprise/GLTFExporter/MigratedFromGit/Source/GLTFExporter/Private/Converters/GLTFMeshConverters.cpp b/Engine/Plugins/Enterprise/GLTFExporter/MigratedFromGit/Source/GLTFExporter/Private/Converters/GLTFMeshConverters.cpp
--- a/Engine/Plugins/Enterprise/GLTFExporter/MigratedFromGit/Source/GLTFExporter/Private/Converters/GLTFMeshConverters.cpp
+++ b/Engine/Plugins/Enterprise/GLTFExporter/MigratedFromGit/Source/GLTFExporter/Private/Converters/GLTFMeshConverters.cpp
@@ -31,7 +31,7 @@ void FGLTFStaticMeshConverter::Sanitize(const UStaticMesh*& StaticMesh, const US
 		// Only use the component if it's needed for baking, since we would
 		// otherwise export a copy of this mesh for each mesh-component.
 		if (!Builder.ExportOptions->bBakeMaterialInputs ||
-			!Builder.ExportOptions->bBakeMaterialInputsUsingMeshData ||
+			!Builder.ExportOptions->bMaterialBakeUsingMeshData ||
 			!FGLTFMaterialUtility::NeedsMeshData(Materials)) // TODO: if this expensive, cache the results for each material
 		{
 			StaticMeshComponent = nullptr;
@@ -71,7 +71,7 @@ void FGLTFSkeletalMeshConverter::Sanitize(const USkeletalMesh*& SkeletalMesh, co
 		// Only use the component if it's needed for baking, since we would
 		// otherwise export a copy of this mesh for each mesh-component.
 		if (!Builder.ExportOptions->bBakeMaterialInputs ||
-			!Builder.ExportOptions->bBakeMaterialInputsUsingMeshData ||
+			!Builder.ExportOptions->bMaterialBakeUsingMeshData ||
 			!FGLTFMaterialUtility::NeedsMeshData(Materials)) // TODO: if this expensive, cache the results for each material
 		{
 			SkeletalMeshComponent = nullptr;
diff --git a/Engine/Plugins/Enterprise/GLTFExporter/MigratedFromGit/Source/GLTFExporter/Private/GLTFExportOptions.cpp b/Engine/Plugins/Enterprise/GLTFExporter/MigratedFromGit/Source/GLTFExporter/Private/GLTFExportOptions.cpp
--- a/Engine/Plugins/Enterprise/GLTFExporter/MigratedFromGit/Source/GLTFExporter/Private/GLTFExportOptions.cpp
+++ b/Engine/Plugins/Enterprise/GLTFExporter/MigratedFromGit/Source/GLTFExporter/Private/GLTFExportOptions.cpp
@@ -17,7 +17,7 @@ void UGLTFExportOptions::ResetToDefault()
 	bExportUnlitMaterials = true;
 	bExportClearCoatMaterials = true;
 	bBakeMaterialInputs = true;
-	bBakeMaterialInputsUsingMeshData = true;
+	bMaterialBakeUsingMeshData = true;
 	DefaultMaterialBakeSize = EGLTFExporterMaterialBakeSize::POT_1024;
 	bExportVertexColors = false;
 	bQuantizeVertexNormals = true;
diff --git a/Engine/Plugins/Enterprise/GLTFExporter/MigratedFromGit/Source/GLTFExporter/Private/Tasks/GLTFMeshTasks.cpp b/Engine/Plugins/Enterprise/GLTFExporter/MigratedFromGit/Source/GLTFExporter/Private/Tasks/GLTFMeshTasks.cpp
--- a/Engine/Plugins/Enterprise/GLTFExporter/MigratedFromGit/Source/GLTFExporter/Private/Tasks/GLTFMeshTasks.cpp
+++ b/Engine/Plugins/Enterprise/GLTFExporter/MigratedFromGit/Source/GLTFExporter/Private/Tasks/GLTFMeshTasks.cpp
@@ -23,7 +23,7 @@ void FGLTFStaticMeshTask::Complete()
 	}
 
 	const FGLTFMeshData* MeshData =
-        Builder.ExportOptions->bBakeMaterialInputs && Builder.ExportOptions->bBakeMaterialInputsUsingMeshData ?
+        Builder.ExportOptions->bBakeMaterialInputs && Builder.ExportOptions->bMaterialBakeUsingMeshData ?
             MeshDataConverter.GetOrAdd(StaticMesh, StaticMeshComponent, LODIndex) : nullptr;
 
 	const int32 MaterialCount = StaticMesh->StaticMaterials.Num();
@@ -85,7 +85,7 @@ void FGLTFSkeletalMeshTask::Complete()
 	}
 
 	const FGLTFMeshData* MeshData =
-		Builder.ExportOptions->bBakeMaterialInputs && Builder.ExportOptions->bBakeMaterialInputsUsingMeshData ?
+		Builder.ExportOptions->bBakeMaterialInputs && Builder.ExportOptions->bMaterialBakeUsingMeshData ?
 			MeshDataConverter.GetOrAdd(SkeletalMesh, SkeletalMeshComponent, LODIndex) : nullptr;
 
 	if (SkeletalMeshComponent == nullptr)
diff --git a/Engine/Plugins/Enterprise/GLTFExporter/MigratedFromGit/Source/GLTFExporter/Public/GLTFExportOptions.h b/Engine/Plugins/Enterprise/GLTFExporter/MigratedFromGit/Source/GLTFExporter/Public/GLTFExportOptions.h
--- a/Engine/Plugins/Enterprise/GLTFExporter/MigratedFromGit/Source/GLTFExporter/Public/GLTFExportOptions.h
+++ b/Engine/Plugins/Enterprise/GLTFExporter/MigratedFromGit/Source/GLTFExporter/Public/GLTFExportOptions.h
@@ -57,7 +57,7 @@ public:
 	bool bBakeMaterialInputs;
 
 	UPROPERTY(EditAnywhere, BlueprintReadWrite, Config, Category = Material, Meta = (EditCondition = "bBakeMaterialInputs"))
-	bool bBakeMaterialInputsUsingMeshData;
+	bool bMaterialBakeUsingMeshData;
 
 	UPROPERTY(EditAnywhere, BlueprintReadWrite, Config, Category = Material, Meta = (EditCondition = "bBakeMaterialInputs"))
 	EGLTFExporterMaterialBakeSize DefaultMaterialBakeSize;
