[GLTFExporter] chore: change default for bMaterialBakeUsingMeshData from false to true and update test assets

commit 82e3a66606fe0b39a0b406a8bb1d074c1e28a552
Date: Wed, 4 Nov 2020 19:40:49 +0100

    chore: change default for bMaterialBakeUsingMeshData from false to true and update test assets
    (cherry picked from commit f7dfb299b4788a850537496f1463c59fe8a41d8f​)

#skipci

diff --git a/Engine/Plugins/Enterprise/GLTFExporter/MigratedFromGit/Source/GLTFExporter/Private/GLTFExportOptions.cpp b/Engine/Plugins/Enterprise/GLTFExporter/MigratedFromGit/Source/GLTFExporter/Private/GLTFExportOptions.cpp
--- a/Engine/Plugins/Enterprise/GLTFExporter/MigratedFromGit/Source/GLTFExporter/Private/GLTFExportOptions.cpp
+++ b/Engine/Plugins/Enterprise/GLTFExporter/MigratedFromGit/Source/GLTFExporter/Private/GLTFExportOptions.cpp
@@ -17,7 +17,7 @@ void UGLTFExportOptions::ResetToDefault()
 	bExportUnlitMaterials = true;
 	bExportClearCoatMaterials = true;
 	bBakeMaterialInputs = true;
-	bMaterialBakeUsingMeshData = true;
+	bMaterialBakeUsingMeshData = false;
 	DefaultMaterialBakeSize = EGLTFExporterMaterialBakeSize::POT_1024;
 	bExportVertexColors = false;
 	bQuantizeVertexNormals = true;
