[GLTFExporter] chore: copy build and source files from https://github.com/EpicGames/UnrealEngine/tree/4.25/Engine/Source/Developer/MaterialBaking to GLTFMaterialBaking module

commit 0c4fd45183637d53548b1b94583e6b459949d160
Date: Wed, 4 Nov 2020 21:07:34 +0100

    chore: copy build and source files from https://github.com/EpicGames/UnrealEngine/tree/4.25/Engine/Source/Developer/MaterialBaking to GLTFMaterialBaking module
    (cherry picked from commit ecf6db69272cc6857568a72e7a19862c6fbc9bb6​)

#skipci

diff --git a/Engine/Plugins/Enterprise/GLTFExporter/MigratedFromGit/Source/GLTFMaterialBaking/MaterialBaking.Build.cs b/Engine/Plugins/Enterprise/GLTFExporter/MigratedFromGit/Source/GLTFMaterialBaking/MaterialBaking.Build.cs
new file mode 100644
--- /dev/null
+++ b/Engine/Plugins/Enterprise/GLTFExporter/MigratedFromGit/Source/GLTFMaterialBaking/MaterialBaking.Build.cs
@@ -0,0 +1,29 @@
+// Copyright Epic Games, Inc. All Rights Reserved.
+
+using UnrealBuildTool;
+using System.IO;
+public class MaterialBaking : ModuleRules
+{
+	public MaterialBaking(ReadOnlyTargetRules Target) : base(Target)
+	{
+        PrivateDependencyModuleNames.AddRange(
+			new string [] {
+				"Core",
+				"CoreUObject",
+				"Engine",
+				"RenderCore",
+                "RHI",                
+                "UnrealEd",
+                "MainFrame",
+                "SlateCore",
+                "Slate",
+                "InputCore",
+                "PropertyEditor",
+                "EditorStyle",
+                "Renderer",
+                "MeshDescription",
+				"StaticMeshDescription"
+            }
+		);
+    }
+}
diff --git a/Engine/Plugins/Enterprise/GLTFExporter/MigratedFromGit/Source/GLTFMaterialBaking/Private/ExportMaterialProxy.h b/Engine/Plugins/Enterprise/GLTFExporter/MigratedFromGit/Source/GLTFMaterialBaking/Private/ExportMaterialProxy.h
new file mode 100644
--- /dev/null
+++ b/Engine/Plugins/Enterprise/GLTFExporter/MigratedFromGit/Source/GLTFMaterialBaking/Private/ExportMaterialProxy.h
@@ -0,0 +1,620 @@
+// Copyright Epic Games, Inc. All Rights Reserved.
+
+#pragma once
+
+#include "MaterialShared.h"
+#include "MaterialCompiler.h"
+#include "Materials/MaterialParameterCollection.h"
+
+#include "Engine/TextureLODSettings.h"
+#include "Engine/Texture2D.h"
+#include "Engine/Texture.h"
+#include "Engine/TextureCube.h"
+#include "Engine/Texture2DArray.h"
+
+#include "DeviceProfiles/DeviceProfileManager.h"
+#include "DeviceProfiles/DeviceProfile.h"
+#include "Materials/MaterialInterface.h"
+#include "SceneTypes.h"
+#include "Materials/Material.h"
+
+struct FExportMaterialCompiler : public FProxyMaterialCompiler
+{
+	FExportMaterialCompiler(FMaterialCompiler* InCompiler) :
+		FProxyMaterialCompiler(InCompiler)
+	{}
+
+	// gets value stored by SetMaterialProperty()
+	virtual EShaderFrequency GetCurrentShaderFrequency() const override
+	{
+		return SF_Pixel;
+	}
+
+	virtual FMaterialShadingModelField GetMaterialShadingModels() const override
+	{
+		return MSM_MAX;
+	}
+
+	virtual int32 WorldPosition(EWorldPositionIncludedOffsets WorldPositionIncludedOffsets) override
+	{
+#if WITH_EDITOR
+		return Compiler->MaterialBakingWorldPosition();
+#else
+		return Compiler->WorldPosition(WorldPositionIncludedOffsets);
+#endif
+	}
+
+	virtual int32 ObjectWorldPosition() override
+	{
+		return Compiler->ObjectWorldPosition();
+	}
+
+	virtual int32 DistanceCullFade() override
+	{
+		return Compiler->Constant(1.0f);
+	}
+
+	virtual int32 ActorWorldPosition() override
+	{
+		return Compiler->ActorWorldPosition();
+	}
+
+	virtual int32 ParticleRelativeTime() override
+	{
+		return Compiler->Constant(0.0f);
+	}
+
+	virtual int32 ParticleMotionBlurFade() override
+	{
+		return Compiler->Constant(1.0f);
+	}
+
+	virtual int32 PixelNormalWS() override
+	{
+		// Current returning vertex normal since pixel normal will contain incorrect data (normal calculated from uv data used as vertex positions to render out the material)
+		return Compiler->VertexNormal();
+	}
+
+	virtual int32 ParticleRandom() override
+	{
+		return Compiler->Constant(0.0f);
+	}
+
+	virtual int32 ParticleDirection() override
+	{
+		return Compiler->Constant3(0.0f, 0.0f, 0.0f);
+	}
+
+	virtual int32 ParticleSpeed() override
+	{
+		return Compiler->Constant(0.0f);
+	}
+
+	virtual int32 ParticleSize() override
+	{
+		return Compiler->Constant2(0.0f, 0.0f);
+	}
+
+	virtual int32 ObjectRadius() override
+	{
+		return Compiler->Constant(500);
+	}
+
+	virtual int32 ObjectBounds() override
+	{
+		return Compiler->ObjectBounds();
+	}
+
+	virtual int32 PreSkinnedLocalBounds(int32 OutputIndex) override
+	{
+		return Compiler->PreSkinnedLocalBounds(OutputIndex);
+	}
+
+	virtual int32 CameraVector() override
+	{
+		return Compiler->Constant3(0.0f, 0.0f, 1.0f);
+	}
+
+	virtual int32 ReflectionAboutCustomWorldNormal(int32 CustomWorldNormal, int32 bNormalizeCustomWorldNormal) override
+	{
+		return Compiler->ReflectionAboutCustomWorldNormal(CustomWorldNormal, bNormalizeCustomWorldNormal);
+	}
+
+	virtual int32 VertexColor() override
+	{
+		return Compiler->VertexColor();
+	}
+
+	virtual int32 PreSkinnedPosition() override
+	{
+		return Compiler->PreSkinnedPosition();
+	}
+
+	virtual int32 PreSkinnedNormal() override
+	{
+		return Compiler->PreSkinnedNormal();
+	}
+
+	virtual int32 VertexInterpolator(uint32 InterpolatorIndex) override
+	{
+		return Compiler->VertexInterpolator(InterpolatorIndex);
+	}
+
+	virtual int32 LightVector() override
+	{
+		return Compiler->LightVector();
+	}
+
+	virtual int32 ReflectionVector() override
+	{
+		return Compiler->ReflectionVector();
+	}
+
+	virtual int32 AtmosphericFogColor(int32 WorldPosition) override
+	{
+		return INDEX_NONE;
+	}
+
+	virtual int32 PrecomputedAOMask() override
+	{
+		return Compiler->PrecomputedAOMask();
+	}
+
+#if WITH_EDITOR
+	virtual int32 MaterialBakingWorldPosition() override
+	{
+		return Compiler->MaterialBakingWorldPosition();
+	}
+#endif
+
+	virtual int32 AccessCollectionParameter(UMaterialParameterCollection* ParameterCollection, int32 ParameterIndex, int32 ComponentIndex) override
+	{
+		if (!ParameterCollection || ParameterIndex == -1)
+		{
+			return INDEX_NONE;
+		}
+
+		// Collect names of all parameters
+		TArray<FName> ParameterNames;
+		ParameterCollection->GetParameterNames(ParameterNames, /*bVectorParameters=*/ false);
+		int32 NumScalarParameters = ParameterNames.Num();
+		ParameterCollection->GetParameterNames(ParameterNames, /*bVectorParameters=*/ true);
+
+		// Find a parameter corresponding to ParameterIndex/ComponentIndex pair
+		int32 Index;
+		for (Index = 0; Index < ParameterNames.Num(); Index++)
+		{
+			FGuid ParameterId = ParameterCollection->GetParameterId(ParameterNames[Index]);
+			int32 CheckParameterIndex, CheckComponentIndex;
+			ParameterCollection->GetParameterIndex(ParameterId, CheckParameterIndex, CheckComponentIndex);
+			if (CheckParameterIndex == ParameterIndex && CheckComponentIndex == ComponentIndex)
+			{
+				// Found
+				break;
+			}
+		}
+		if (Index >= ParameterNames.Num())
+		{
+			// Not found, should not happen
+			return INDEX_NONE;
+		}
+
+		// Create code for parameter
+		if (Index < NumScalarParameters)
+		{
+			const FCollectionScalarParameter* ScalarParameter = ParameterCollection->GetScalarParameterByName(ParameterNames[Index]);
+			check(ScalarParameter);
+			return Constant(ScalarParameter->DefaultValue);
+		}
+		else
+		{
+			const FCollectionVectorParameter* VectorParameter = ParameterCollection->GetVectorParameterByName(ParameterNames[Index]);
+			check(VectorParameter);
+			const FLinearColor& Color = VectorParameter->DefaultValue;
+			return Constant4(Color.R, Color.G, Color.B, Color.A);
+		}
+	}
+
+	virtual EMaterialCompilerType GetCompilerType() const override { return EMaterialCompilerType::MaterialProxy; }
+};
+
+class FExportMaterialProxy : public FMaterial, public FMaterialRenderProxy
+{
+public:
+	FExportMaterialProxy(UMaterialInterface* InMaterialInterface, EMaterialProperty InPropertyToCompile, bool bInSynchronousCompilation = true)
+		: FMaterial()
+		, MaterialInterface(InMaterialInterface)
+		, PropertyToCompile(InPropertyToCompile)
+		, bSynchronousCompilation(bInSynchronousCompilation)
+	{
+		SetQualityLevelProperties(EMaterialQualityLevel::High, false, GMaxRHIFeatureLevel);
+		Material = InMaterialInterface->GetMaterial();
+		ReferencedTextures = InMaterialInterface->GetReferencedTextures();
+
+		const FMaterialResource* Resource = InMaterialInterface->GetMaterialResource(GMaxRHIFeatureLevel);
+
+		FMaterialShaderMapId ResourceId;
+		Resource->GetShaderMapId(GMaxRHIShaderPlatform, nullptr, ResourceId);
+
+		// Our Id must be the same as BaseMaterialId for the shader compiler
+		// to be able to set back GameThreadShaderMap after async compilation.
+		Id = ResourceId.BaseMaterialId;
+
+
+		{
+			TArray<FShaderType*> ShaderTypes;
+			TArray<FVertexFactoryType*> VFTypes;
+			TArray<const FShaderPipelineType*> ShaderPipelineTypes;
+			GetDependentShaderAndVFTypes(GMaxRHIShaderPlatform, ShaderTypes, ShaderPipelineTypes, VFTypes);
+
+			// Overwrite the shader map Id's dependencies with ones that came from the FMaterial actually being compiled (this)
+			// This is necessary as we change FMaterial attributes like GetShadingModels(), which factor into the ShouldCache functions that determine dependent shader types
+			ResourceId.SetShaderDependencies(ShaderTypes, ShaderPipelineTypes, VFTypes, GMaxRHIShaderPlatform);
+		}
+
+		// Override with a special usage so we won't re-use the shader map used by the material for rendering
+		switch (InPropertyToCompile)
+		{
+		case MP_BaseColor: ResourceId.Usage = EMaterialShaderMapUsage::MaterialExportBaseColor; break;
+		case MP_Specular: ResourceId.Usage = EMaterialShaderMapUsage::MaterialExportSpecular; break;
+		case MP_Normal: ResourceId.Usage = EMaterialShaderMapUsage::MaterialExportNormal; break;
+		case MP_Tangent: ResourceId.Usage = EMaterialShaderMapUsage::MaterialExportTangent; break;
+		case MP_Metallic: ResourceId.Usage = EMaterialShaderMapUsage::MaterialExportMetallic; break;
+		case MP_Roughness: ResourceId.Usage = EMaterialShaderMapUsage::MaterialExportRoughness; break;
+		case MP_Anisotropy: ResourceId.Usage = EMaterialShaderMapUsage::MaterialExportAnisotropy; break;
+		case MP_AmbientOcclusion: ResourceId.Usage = EMaterialShaderMapUsage::MaterialExportAO; break;
+		case MP_EmissiveColor: ResourceId.Usage = EMaterialShaderMapUsage::MaterialExportEmissive; break;
+		case MP_Opacity: ResourceId.Usage = EMaterialShaderMapUsage::MaterialExportOpacity; break;
+		case MP_OpacityMask: ResourceId.Usage = EMaterialShaderMapUsage::MaterialExportOpacityMask; break;
+		case MP_SubsurfaceColor: ResourceId.Usage = EMaterialShaderMapUsage::MaterialExportSubSurfaceColor; break;
+		default:
+			ensureMsgf(false, TEXT("ExportMaterial has no usage for property %i.  Will likely reuse the normal rendering shader and crash later with a parameter mismatch"), (int32)InPropertyToCompile);
+			break;
+		};
+
+		CacheShaders(ResourceId, GMaxRHIShaderPlatform);
+	}
+
+	/** This override is required otherwise the shaders aren't ready for use when the surface is rendered resulting in a blank image */
+	virtual bool RequiresSynchronousCompilation() const override { return bSynchronousCompilation; };
+
+	/**
+	* Should the shader for this material with the given platform, shader type and vertex
+	* factory type combination be compiled
+	*
+	* @param Platform		The platform currently being compiled for
+	* @param ShaderType	Which shader is being compiled
+	* @param VertexFactory	Which vertex factory is being compiled (can be NULL)
+	*
+	* @return true if the shader should be compiled
+	*/
+	virtual bool ShouldCache(EShaderPlatform Platform, const FShaderType* ShaderType, const FVertexFactoryType* VertexFactoryType) const override
+	{
+		const bool bCorrectVertexFactory = VertexFactoryType == FindVertexFactoryType(FName(TEXT("FLocalVertexFactory"), FNAME_Find));
+		const bool bPCPlatform = !IsConsolePlatform(Platform);
+		const bool bCorrectFrequency = ShaderType->GetFrequency() == SF_Vertex || ShaderType->GetFrequency() == SF_Pixel;
+		return bCorrectVertexFactory && bPCPlatform && bCorrectFrequency;
+	}
+
+	virtual TArrayView<UObject* const> GetReferencedTextures() const override
+	{
+		return ReferencedTextures;
+	}
+
+	virtual void GetStaticParameterSet(EShaderPlatform Platform, FStaticParameterSet& OutSet) const override
+	{
+		if (const FMaterialResource* Resource = MaterialInterface->GetMaterialResource(GMaxRHIFeatureLevel))
+		{
+			Resource->GetStaticParameterSet(Platform, OutSet);
+		}
+	}
+
+	////////////////
+	// FMaterialRenderProxy interface.
+	virtual const FMaterial& GetMaterialWithFallback(ERHIFeatureLevel::Type FeatureLevel, const FMaterialRenderProxy*& OutFallbackMaterialRenderProxy) const override
+	{
+		if(GetRenderingThreadShaderMap())
+		{
+			return *this;
+		}
+		else
+		{
+			OutFallbackMaterialRenderProxy = UMaterial::GetDefaultMaterial(MD_Surface)->GetRenderProxy();
+			return OutFallbackMaterialRenderProxy->GetMaterialWithFallback(FeatureLevel, OutFallbackMaterialRenderProxy);
+		}
+	}
+
+	virtual bool GetVectorValue(const FHashedMaterialParameterInfo& ParameterInfo, FLinearColor* OutValue, const FMaterialRenderContext& Context) const override
+	{
+		return MaterialInterface->GetRenderProxy()->GetVectorValue(ParameterInfo, OutValue, Context);
+	}
+
+	virtual bool GetScalarValue(const FHashedMaterialParameterInfo& ParameterInfo, float* OutValue, const FMaterialRenderContext& Context) const override
+	{
+		return MaterialInterface->GetRenderProxy()->GetScalarValue(ParameterInfo, OutValue, Context);
+	}
+
+	virtual bool GetTextureValue(const FHashedMaterialParameterInfo& ParameterInfo, const UTexture** OutValue, const FMaterialRenderContext& Context) const override
+	{
+		return MaterialInterface->GetRenderProxy()->GetTextureValue(ParameterInfo, OutValue, Context);
+	}
+
+	virtual bool GetTextureValue(const FHashedMaterialParameterInfo& ParameterInfo, const URuntimeVirtualTexture** OutValue, const FMaterialRenderContext& Context) const override
+	{
+		return MaterialInterface->GetRenderProxy()->GetTextureValue(ParameterInfo, OutValue, Context);
+	}
+
+	// Material properties.
+	/** Entry point for compiling a specific material property.  This must call SetMaterialProperty. */
+	virtual int32 CompilePropertyAndSetMaterialProperty(EMaterialProperty Property, FMaterialCompiler* Compiler, EShaderFrequency OverrideShaderFrequency, bool bUsePreviousFrameTime) const override
+	{
+		// needs to be called in this function!!
+		Compiler->SetMaterialProperty(Property, OverrideShaderFrequency, bUsePreviousFrameTime);
+		const int32 Ret = CompilePropertyAndSetMaterialPropertyWithoutCast(Property, Compiler);
+		return Compiler->ForceCast(Ret, FMaterialAttributeDefinitionMap::GetValueType(Property), MFCF_ExactMatch | MFCF_ReplicateValue);
+	}
+
+	/** helper for CompilePropertyAndSetMaterialProperty() */
+	int32 CompilePropertyAndSetMaterialPropertyWithoutCast(EMaterialProperty Property, FMaterialCompiler* Compiler) const
+	{
+		if (Property == MP_EmissiveColor)
+		{
+			const EBlendMode BlendMode = MaterialInterface->GetBlendMode();
+			FExportMaterialCompiler ProxyCompiler(Compiler);
+			const uint32 ForceCast_Exact_Replicate = MFCF_ForceCast | MFCF_ExactMatch | MFCF_ReplicateValue;
+
+			switch (PropertyToCompile)
+			{
+			case MP_EmissiveColor:
+				// Emissive is ALWAYS returned...
+				return MaterialInterface->CompileProperty(&ProxyCompiler, MP_EmissiveColor, ForceCast_Exact_Replicate);
+			case MP_BaseColor:
+				return MaterialInterface->CompileProperty(&ProxyCompiler, MP_BaseColor, ForceCast_Exact_Replicate);
+				break;
+			case MP_Specular:
+			case MP_Roughness:
+			case MP_Anisotropy:
+			case MP_Metallic:
+			case MP_AmbientOcclusion:
+				// Only return for Opaque and Masked...
+				if (BlendMode == BLEND_Opaque || BlendMode == BLEND_Masked)
+				{
+					return MaterialInterface->CompileProperty(&ProxyCompiler, PropertyToCompile, ForceCast_Exact_Replicate);
+				}
+				break;
+
+			case MP_Opacity:
+			case MP_OpacityMask:
+			{
+				return MaterialInterface->CompileProperty(&ProxyCompiler, PropertyToCompile, ForceCast_Exact_Replicate);
+			}
+			case MP_Normal:
+			case MP_Tangent:
+				// Only return for Opaque and Masked...
+				if (BlendMode == BLEND_Opaque || BlendMode == BLEND_Masked)
+				{
+					return Compiler->Add(
+						Compiler->Mul(MaterialInterface->CompileProperty(&ProxyCompiler, PropertyToCompile, ForceCast_Exact_Replicate), Compiler->Constant(0.5f)), // [-1,1] * 0.5
+						Compiler->Constant(0.5f)); // [-0.5,0.5] + 0.5
+				}
+				break;
+			case MP_ShadingModel:
+				return MaterialInterface->CompileProperty(&ProxyCompiler, MP_ShadingModel);
+			default:
+				return Compiler->Constant(1.0f);
+			}
+
+			return Compiler->Constant(0.0f);
+		}
+		else if (Property == MP_WorldPositionOffset)
+		{
+			//This property MUST return 0 as a default or during the process of rendering textures out for lightmass to use, pixels will be off by 1.
+			return Compiler->Constant(0.0f);
+		}
+		else if (Property >= MP_CustomizedUVs0 && Property <= MP_CustomizedUVs7)
+		{
+			// Pass through customized UVs
+			return MaterialInterface->CompileProperty(Compiler, Property);
+		}
+		else if (Property == MP_ShadingModel)
+		{
+			return MaterialInterface->CompileProperty(Compiler, MP_ShadingModel);
+
+		}
+		else
+		{
+			return Compiler->Constant(1.0f);
+		}
+	}
+
+	virtual FString GetMaterialUsageDescription() const override
+	{
+		return FString::Printf(TEXT("MaterialBaking_%s"), MaterialInterface ? *MaterialInterface->GetName() : TEXT("NULL"));
+	}
+	virtual EMaterialDomain GetMaterialDomain() const override
+	{
+		if (Material)
+		{
+			return Material->MaterialDomain;
+		}
+		return MD_Surface;
+	}
+	virtual bool IsTwoSided() const  override
+	{
+		if (MaterialInterface)
+		{
+			return MaterialInterface->IsTwoSided();
+		}
+		return false;
+	}
+	virtual bool IsDitheredLODTransition() const  override
+	{
+		if (MaterialInterface)
+		{
+			return MaterialInterface->IsDitheredLODTransition();
+		}
+		return false;
+	}
+	virtual bool IsLightFunction() const override
+	{
+		if (Material)
+		{
+			return (Material->MaterialDomain == MD_LightFunction);
+		}
+		return false;
+	}
+	virtual bool IsDeferredDecal() const override
+	{
+		return Material && Material->MaterialDomain == MD_DeferredDecal;
+	}
+	virtual bool IsVolumetricPrimitive() const override
+	{
+		return Material && Material->MaterialDomain == MD_Volume;
+	}
+	virtual bool IsSpecialEngineMaterial() const override
+	{
+		if (Material)
+		{
+			return (Material->bUsedAsSpecialEngineMaterial == 1);
+		}
+		return false;
+	}
+	virtual bool IsWireframe() const override
+	{
+		if (Material)
+		{
+			return (Material->Wireframe == 1);
+		}
+		return false;
+	}
+	virtual bool IsMasked() const override { return false; }
+	virtual enum EBlendMode GetBlendMode() const override { return BLEND_Opaque; }
+	virtual FMaterialShadingModelField GetShadingModels() const override { return MSM_DefaultLit; }
+	virtual bool IsShadingModelFromMaterialExpression() const override { return false; }
+	virtual float GetOpacityMaskClipValue() const override { return 0.5f; }
+	virtual bool GetCastDynamicShadowAsMasked() const override { return false; }
+	virtual FString GetFriendlyName() const override { return FString::Printf(TEXT("FExportMaterialRenderer %s"), MaterialInterface ? *MaterialInterface->GetName() : TEXT("NULL")); }
+	/**
+	* Should shaders compiled for this material be saved to disk?
+	*/
+	virtual bool IsPersistent() const override { return true; }
+	virtual FGuid GetMaterialId() const override { return Id; }
+
+	virtual UMaterialInterface* GetMaterialInterface() const override
+	{
+		return MaterialInterface;
+	}
+
+	friend FArchive& operator<< (FArchive& Ar, FExportMaterialProxy& V)
+	{
+		return Ar << V.MaterialInterface;
+	}
+
+	/**
+	* Iterate through all textures used by the material and return the maximum texture resolution used
+	* (ideally this could be made dependent of the material property)
+	*
+	* @param MaterialInterface The material to scan for texture size
+	*
+	* @return Size (width and height)
+	*/
+	FIntPoint FindMaxTextureSize(UMaterialInterface* InMaterialInterface, FIntPoint MinimumSize = FIntPoint(1, 1)) const
+	{
+		// static lod settings so that we only initialize them once
+		const UTextureLODSettings* GameTextureLODSettings = UDeviceProfileManager::Get().GetActiveProfile()->GetTextureLODSettings();
+
+		TArray<UTexture*> MaterialTextures;
+		InMaterialInterface->GetUsedTextures(MaterialTextures, EMaterialQualityLevel::Num, false, GMaxRHIFeatureLevel, false);
+
+		// find the largest texture in the list (applying it's LOD bias)
+		FIntPoint MaxSize = MinimumSize;
+		for (int32 TexIndex = 0; TexIndex < MaterialTextures.Num(); TexIndex++)
+		{
+			const UTexture* Texture = MaterialTextures[TexIndex];
+
+			if (Texture == NULL)
+			{
+				continue;
+			}
+
+			// get the max size of the texture
+			const FIntPoint LocalSize = [&]()
+			{
+				if (Texture->IsA(UTexture2D::StaticClass()))
+				{
+					UTexture2D* Tex2D = (UTexture2D*)Texture;
+					return FIntPoint(Tex2D->GetSizeX(), Tex2D->GetSizeY());
+				}
+				else if (Texture->IsA(UTextureCube::StaticClass()))
+				{
+					UTextureCube* TexCube = (UTextureCube*)Texture;
+					return FIntPoint(TexCube->GetSizeX(), TexCube->GetSizeY());
+				}
+				else if (Texture->IsA(UTexture2DArray::StaticClass())) 
+				{
+					UTexture2DArray* TexArray = (UTexture2DArray*)Texture;
+					return FIntPoint(TexArray->GetSizeX(), TexArray->GetSizeY());
+				}
+				return FIntPoint(0, 0);
+			}();
+			
+			// bias the texture size based on LOD group
+			const int32 LocalBias = GameTextureLODSettings->CalculateLODBias(Texture);
+			MaxSize.X = FMath::Max(LocalSize.X >> LocalBias, MaxSize.X);
+			MaxSize.Y = FMath::Max(LocalSize.Y >> LocalBias, MaxSize.Y);
+		}
+
+		return MaxSize;
+	}
+
+	static bool WillFillData(EBlendMode InBlendMode, EMaterialProperty InMaterialProperty)
+	{
+		if (InMaterialProperty == MP_EmissiveColor)
+		{
+			return true;
+		}
+
+		switch (InBlendMode)
+		{
+		case BLEND_Opaque:
+		{
+			switch (InMaterialProperty)
+			{
+			case MP_BaseColor:				return true;
+			case MP_Specular:				return true;
+			case MP_Normal:					return true;
+			case MP_Tangent:				return true;
+			case MP_Metallic:				return true;
+			case MP_Roughness:				return true;
+			case MP_Anisotropy:				return true;
+			case MP_AmbientOcclusion:		return true;
+			}
+		}
+		break;
+		}
+		return false;
+	}
+
+	virtual bool IsUsedWithStaticLighting() const override
+	{
+		return true; 
+	}
+
+	virtual void GatherExpressionsForCustomInterpolators(TArray<UMaterialExpression*>& OutExpressions) const override
+	{
+		if(Material)
+		{
+			Material->GetAllExpressionsForCustomInterpolators(OutExpressions);
+		}
+	}
+
+private:
+	/** The material interface for this proxy */
+	UMaterialInterface* MaterialInterface;
+	UMaterial* Material;
+	TArray<UObject*> ReferencedTextures;
+	/** The property to compile for rendering the sample */
+	EMaterialProperty PropertyToCompile;
+	FGuid Id;
+	bool bSynchronousCompilation;
+};
\ No newline at end of file
diff --git a/Engine/Plugins/Enterprise/GLTFExporter/MigratedFromGit/Source/GLTFMaterialBaking/Private/MaterialBakingHelpers.cpp b/Engine/Plugins/Enterprise/GLTFExporter/MigratedFromGit/Source/GLTFMaterialBaking/Private/MaterialBakingHelpers.cpp
new file mode 100644
--- /dev/null
+++ b/Engine/Plugins/Enterprise/GLTFExporter/MigratedFromGit/Source/GLTFMaterialBaking/Private/MaterialBakingHelpers.cpp
@@ -0,0 +1,331 @@
+// Copyright Epic Games, Inc. All Rights Reserved.
+
+#include "MaterialBakingHelpers.h"
+#include "Math/Color.h"
+#include "Async/ParallelFor.h"
+
+namespace FMaterialBakingHelpersImpl
+{
+	static FColor BoxBlurSample(FColor* InBMP, int32 X, int32 Y, int32 InImageWidth, int32 InImageHeight)
+	{
+		int32 PixelsSampled = 0;
+		int32 CombinedColorR = 0;
+		int32 CombinedColorG = 0;
+		int32 CombinedColorB = 0;
+		int32 CombinedColorA = 0;
+	
+		const uint32 MagentaMask = FColor(255, 0, 255).DWColor();
+
+		const int32 BaseIndex = (Y * InImageWidth) + X;
+	
+		int32 SampleIndex = BaseIndex - (InImageWidth + 1);
+		for (int32 YI = 0; YI < 3; ++YI)
+		{
+			for (int32 XI = 0; XI < 3; ++XI)
+			{
+				const FColor SampledColor = InBMP[SampleIndex++];
+				// Check if the pixel is a rendered one (not clear color)
+				if (SampledColor.DWColor() != MagentaMask)
+				{
+					CombinedColorR += SampledColor.R;
+					CombinedColorG += SampledColor.G;
+					CombinedColorB += SampledColor.B;
+					CombinedColorA += SampledColor.A;
+					++PixelsSampled;
+				}
+			}
+			SampleIndex += InImageWidth - 3;
+		}
+
+		if (PixelsSampled == 0)
+		{
+			const FColor Magenta = FColor(255, 0, 255);
+			return Magenta;
+		}
+
+		return FColor(CombinedColorR / PixelsSampled, CombinedColorG / PixelsSampled, CombinedColorB / PixelsSampled, CombinedColorA / PixelsSampled);
+	}
+
+	static bool HasBorderingPixel(FColor* InBMP, int32 X, int32 Y, int32 InImageWidth, int32 InImageHeight)
+	{
+		const uint32 MagentaMask = FColor(255, 0, 255).DWColor();
+
+		const int32 BaseIndex = Y*InImageWidth + X;
+
+		int32 SampleIndex = BaseIndex - (InImageWidth + 1);
+		for (int32 YI = 0; YI < 3; ++YI)
+		{
+			for (int32 XI = 0; XI < 3; ++XI)
+			{
+				if (InBMP[SampleIndex++].DWColor() != MagentaMask)
+				{
+					return true;
+				}
+			}
+			SampleIndex += InImageWidth - 3;
+		}
+		return false;
+	}
+
+	void PerformUVBorderSmear(TArray<FColor>& InOutPixels, int32& ImageWidth, int32& ImageHeight, int32 MaxIterations, bool bCanShrink)
+	{
+		TRACE_CPUPROFILER_EVENT_SCOPE(FMaterialBakingHelpers::PerformUVBorderSmear)
+
+		check(InOutPixels.Num() == ImageWidth*ImageHeight);
+
+		const uint32 MagentaMask = FColor(255, 0, 255).DWColor();
+		const int32 PaddedImageWidth = ImageWidth + 2;
+		const int32 PaddedImageHeight = ImageHeight + 2;
+
+		TArray<uint32> RowsSwap1;
+		TArray<uint32> RowsSwap2;
+		TArray<uint32>* CurrentRowsLeft = &RowsSwap1;
+		TArray<uint32>* NextRowsLeft = &RowsSwap2;
+
+		//
+		// Create a copy of the image with a 1 pixel magenta border around the edge.
+		// This avoids needing to bounds check in inner loops.
+		//
+
+		FColor* Current = InOutPixels.GetData();
+
+		// We can hopefully skip the entire smearing process if there is a single
+		// non-magenta color in the entire array since smearing that completely
+		// would lead to a monochome output.
+		if (MaxIterations == -1)
+		{
+			uint32 SingleColor     = MagentaMask;
+			bool   bHasSingleColor = true;
+			for (int32 Index = 0, Num = InOutPixels.Num(); Index < Num; ++Index)
+			{
+				const uint32 Color = Current[Index].DWColor();
+				if (Color != MagentaMask)
+				{
+					if (SingleColor == MagentaMask)
+					{
+						// This is the first time we stumble on a color other than magenta, keep it.
+						SingleColor = Color;
+					}
+					// Compare with the known non-magenta color
+					else if (SingleColor != Color)
+					{
+						bHasSingleColor = false;
+						break;
+					}
+				}
+			}
+
+			if (bHasSingleColor)
+			{
+				// Shrink to single texel when allowed
+				if (bCanShrink)
+				{
+					InOutPixels.SetNum(1);
+					ImageWidth = ImageHeight = 1;
+					Current = InOutPixels.GetData();
+				}
+
+				FColor Color(SingleColor);
+				for (int32 Index = 0, Num = InOutPixels.Num(); Index < Num; ++Index)
+				{
+					Current[Index] = Color;
+				}
+
+				return;
+			}
+		}
+
+		TArray<FColor> ScratchBuffer;
+		ScratchBuffer.SetNumUninitialized(PaddedImageWidth*PaddedImageHeight);
+		FColor* Scratch = ScratchBuffer.GetData();
+
+		// Set top and bottom bordering pixels to magenta
+		for (int32 X = 0; X < PaddedImageWidth; ++X)
+		{
+			Scratch[X] = FColor(MagentaMask);
+			Scratch[(PaddedImageHeight - 1)*PaddedImageWidth + X] = FColor(MagentaMask);
+		}
+
+		// Set leftg and right border pixels to magenta and copy image data into rows
+		for (int32 Y = 1; Y <= ImageHeight; ++Y)
+		{
+			const int32 YOffset = Y * PaddedImageWidth;
+			Scratch[YOffset] = FColor(MagentaMask);
+			FMemory::Memcpy(&Scratch[YOffset + 1], &Current[(Y - 1)*ImageWidth], ImageWidth * sizeof(FColor));
+			Scratch[YOffset + ImageWidth + 1] = FColor(MagentaMask);
+		}
+
+		//
+		// Find our initial workset of all rows that have magenta pixels bordering non-magenta pixels.
+		// Also find all rows that have zero magenta pixels - these will never be added to the workset.
+		//
+
+		TArray<bool> RowCompleted;
+		RowCompleted.SetNumZeroed(PaddedImageHeight);
+
+		// Set border rows to start completed
+		RowCompleted[0] = true;
+		RowCompleted[PaddedImageHeight - 1] = true;
+
+		bool bHasAnyData = false;
+		for (int32 Y = 1; Y <= ImageHeight; ++Y)
+		{
+			bool bHasMagenta = false;
+			bool bBordersNonMagenta = false;
+			const int32 YOffset = Y * PaddedImageWidth;
+			for (int32 X = 1; X <= ImageWidth; ++X)
+			{
+				if (Scratch[YOffset + X].DWColor() == MagentaMask)
+				{
+					bHasMagenta = true;
+					if (HasBorderingPixel(Scratch, X, Y, PaddedImageWidth, PaddedImageHeight))
+					{
+						bBordersNonMagenta = true;
+						break;
+					}
+				}
+			}
+
+			if (!bHasMagenta)
+			{
+				bHasAnyData = true;
+				RowCompleted[Y] = true;
+			}
+			else if (bBordersNonMagenta)
+			{
+				bHasAnyData = true;
+				CurrentRowsLeft->Add(Y);
+			}
+		}
+
+		// early-out if no valid pixels were found.
+		if (!bHasAnyData)
+		{
+			FMemory::Memzero(InOutPixels.GetData(), ImageWidth * ImageHeight * sizeof(FColor));
+			return;
+		}
+
+		// Early out if all pixels were valid.
+		if (CurrentRowsLeft->Num() == 0)
+		{
+			return;
+		}
+
+		TArray<uint32> RowRemainingPixels;
+		RowRemainingPixels.SetNumZeroed(PaddedImageHeight);
+
+		const int32 MaxThreads = FPlatformProcess::SupportsMultithreading() ? FPlatformMisc::NumberOfCoresIncludingHyperthreads() : 1;
+
+		//
+		// Iteratively smear until all rows are filled.
+		//
+
+		int32 LoopCount = 0;
+		while (CurrentRowsLeft->Num() && (MaxIterations <= 0 || LoopCount <= MaxIterations))
+		{
+			const int32 NumThreads = FMath::Min(CurrentRowsLeft->Num(), MaxThreads);
+			const int32 LinesPerThread = FMath::CeilToInt((float)CurrentRowsLeft->Num() / (float)NumThreads);
+
+			// split up rows that still have magenta pixels amongst threads
+			ParallelFor(NumThreads, [ImageWidth, PaddedImageWidth, PaddedImageHeight, Current, Scratch, CurrentRowsLeft, LinesPerThread, &RowRemainingPixels, MagentaMask](int32 ThreadIndex)
+			{
+				const int32 StartY = ThreadIndex*LinesPerThread;
+				const int32 EndY = FMath::Min((ThreadIndex + 1) * LinesPerThread, CurrentRowsLeft->Num());
+
+				for (int32 Index = StartY; Index < EndY; Index++)
+				{
+					const int32 Y = (*CurrentRowsLeft)[Index];
+					RowRemainingPixels[Index] = 0;
+
+					int32 PixelIndex = (Y - 1) * ImageWidth;
+					for (int32 X = 1; X <= ImageWidth; X++)
+					{
+						FColor& Color = Current[PixelIndex++];
+						if (Color.DWColor() == MagentaMask)
+						{
+							const FColor SampledColor = BoxBlurSample(Scratch, X, Y, PaddedImageWidth, PaddedImageHeight);
+							// If it's a valid pixel
+							if (SampledColor.DWColor() != MagentaMask)
+							{
+								Color = SampledColor;
+							}
+							else
+							{
+								RowRemainingPixels[Index]++;
+							}
+						}
+					}
+				}
+			});
+
+			// Mark all completed rows and copy data between buffers
+			const int32 RowMemorySize = ImageWidth * sizeof(FColor);
+			for (int32 RowIndex = 0; RowIndex < CurrentRowsLeft->Num(); ++RowIndex)
+			{
+				const int32 Y = (*CurrentRowsLeft)[RowIndex];
+				if (!RowRemainingPixels[RowIndex])
+				{
+					RowCompleted[Y] = true;
+				}
+				FMemory::Memcpy(&Scratch[Y * PaddedImageWidth + 1], &Current[(Y - 1) * ImageWidth], RowMemorySize);
+			}
+
+			// Find the next set of rows to process.
+			// This will be our current rows that haven't completed, plus any bordering rows that haven't been marked completed.
+			NextRowsLeft->SetNum(0, false);
+			{
+				int32 PreviousY = -1;
+				for (int32 RowIndex = 0; RowIndex < CurrentRowsLeft->Num(); ++RowIndex)
+				{
+					const int32 Y = (*CurrentRowsLeft)[RowIndex];
+					if (!RowCompleted[Y - 1] && PreviousY < Y - 1)
+					{
+						PreviousY = Y - 1;
+						NextRowsLeft->Add(Y - 1);
+					}
+					if (!RowCompleted[Y] && RowRemainingPixels[RowIndex] && PreviousY < Y)
+					{
+						PreviousY = Y;
+						NextRowsLeft->Add(Y);
+					}
+					if (!RowCompleted[Y + 1])
+					{
+						PreviousY = Y + 1;
+						NextRowsLeft->Add(Y + 1);
+					}
+				}
+			}
+
+			// Swap rows left buffers
+			{			
+				TArray<uint32>* Temp = NextRowsLeft;
+				NextRowsLeft = CurrentRowsLeft;
+				CurrentRowsLeft = Temp;
+			}
+
+			LoopCount++;
+		}
+
+		// If we finished before replacing all pixels, replace the remaining pixels with black.
+		if (CurrentRowsLeft->Num() > 0)
+		{
+			for (int32 i = 0; i < InOutPixels.Num(); ++i)
+			{
+				if (InOutPixels[i].DWColor() == MagentaMask)
+				{
+					InOutPixels[i] = FColor::Black;
+				}
+			}
+		}
+	}
+}
+
+void FMaterialBakingHelpers::PerformUVBorderSmear(TArray<FColor>& InOutPixels, int32 ImageWidth, int32 ImageHeight, int32 MaxIterations)
+{
+	FMaterialBakingHelpersImpl::PerformUVBorderSmear(InOutPixels, ImageWidth, ImageHeight, MaxIterations, false);
+}
+
+void FMaterialBakingHelpers::PerformUVBorderSmearAndShrink(TArray<FColor>& InOutPixels, int32& InOutImageWidth, int32& InOutImageHeight)
+{
+	FMaterialBakingHelpersImpl::PerformUVBorderSmear(InOutPixels, InOutImageWidth, InOutImageHeight, -1, true);
+}
\ No newline at end of file
diff --git a/Engine/Plugins/Enterprise/GLTFExporter/MigratedFromGit/Source/GLTFMaterialBaking/Private/MaterialBakingModule.cpp b/Engine/Plugins/Enterprise/GLTFExporter/MigratedFromGit/Source/GLTFMaterialBaking/Private/MaterialBakingModule.cpp
new file mode 100644
--- /dev/null
+++ b/Engine/Plugins/Enterprise/GLTFExporter/MigratedFromGit/Source/GLTFMaterialBaking/Private/MaterialBakingModule.cpp
@@ -0,0 +1,902 @@
+// Copyright Epic Games, Inc. All Rights Reserved.
+
+#include "MaterialBakingModule.h"
+#include "MaterialRenderItem.h"
+#include "Engine/TextureRenderTarget2D.h"
+#include "ExportMaterialProxy.h"
+#include "Interfaces/IMainFrameModule.h"
+#include "MaterialOptionsWindow.h"
+#include "MaterialOptions.h"
+#include "PropertyEditorModule.h"
+#include "MaterialOptionsCustomization.h"
+#include "UObject/UObjectGlobals.h"
+#include "MaterialBakingStructures.h"
+#include "Framework/Application/SlateApplication.h"
+#include "MaterialBakingHelpers.h"
+#include "Async/Async.h"
+#include "Async/ParallelFor.h"
+#include "Materials/MaterialInstance.h"
+#include "RenderingThread.h"
+#include "RHISurfaceDataConversion.h"
+#include "Misc/ScopedSlowTask.h"
+#include "MeshDescription.h"
+#if WITH_EDITOR
+#include "Misc/FileHelper.h"
+#endif
+
+IMPLEMENT_MODULE(FMaterialBakingModule, MaterialBaking);
+
+#define LOCTEXT_NAMESPACE "MaterialBakingModule"
+
+/** Cvars for advanced features */
+static TAutoConsoleVariable<int32> CVarUseMaterialProxyCaching(
+	TEXT("MaterialBaking.UseMaterialProxyCaching"),
+	1,
+	TEXT("Determines whether or not Material Proxies should be cached to speed up material baking.\n")
+	TEXT("0: Turned Off\n")
+	TEXT("1: Turned On"),	
+	ECVF_Default);
+
+static TAutoConsoleVariable<int32> CVarSaveIntermediateTextures(
+	TEXT("MaterialBaking.SaveIntermediateTextures"),
+	0,
+	TEXT("Determines whether or not to save out intermediate BMP images for each flattened material property.\n")
+	TEXT("0: Turned Off\n")
+	TEXT("1: Turned On"),
+	ECVF_Default);
+
+namespace FMaterialBakingModuleImpl
+{
+	// Custom dynamic mesh allocator specifically tailored for Material Baking.
+	// This will always reuse the same couple buffers, so searching linearly is not a problem.
+	class FMaterialBakingDynamicMeshBufferAllocator : public FDynamicMeshBufferAllocator
+	{
+		// This must be smaller than the large allocation blocks on Windows 10 which is currently ~508K.
+		// Large allocations uses VirtualAlloc directly without any kind of buffering before
+		// releasing pages to the kernel, so it causes lots of soft page fault when
+		// memory is first initialized.
+		const uint32 SmallestPooledBufferSize = 256*1024;
+
+		TArray<FIndexBufferRHIRef>  IndexBuffers;
+		TArray<FVertexBufferRHIRef> VertexBuffers;
+
+		template <typename RefType>
+		RefType GetSmallestFit(uint32 SizeInBytes, TArray<RefType>& Array)
+		{
+			uint32 SmallestFitIndex = UINT32_MAX;
+			uint32 SmallestFitSize  = UINT32_MAX;
+			for (int32 Index = 0; Index < Array.Num(); ++Index)
+			{
+				uint32 Size = Array[Index]->GetSize();
+				if (Size >= SizeInBytes && (SmallestFitIndex == UINT32_MAX || Size < SmallestFitSize))
+				{
+					SmallestFitIndex = Index;
+					SmallestFitSize  = Size;
+				}
+			}
+
+			RefType Ref;
+			// Do not reuse the smallest fit if it's a lot bigger than what we requested
+			if (SmallestFitIndex != UINT32_MAX && SmallestFitSize < SizeInBytes*2)
+			{
+				Ref = Array[SmallestFitIndex];
+				Array.RemoveAtSwap(SmallestFitIndex);
+			}
+
+			return Ref;
+		}
+
+		virtual FIndexBufferRHIRef AllocIndexBuffer(uint32 NumElements) override
+		{
+			uint32 BufferSize = GetIndexBufferSize(NumElements);
+			if (BufferSize > SmallestPooledBufferSize)
+			{
+				FIndexBufferRHIRef Ref = GetSmallestFit(GetIndexBufferSize(NumElements), IndexBuffers);
+				if (Ref.IsValid())
+				{
+					return Ref;
+				}
+			}
+
+			return FDynamicMeshBufferAllocator::AllocIndexBuffer(NumElements);
+		}
+
+		virtual void ReleaseIndexBuffer(FIndexBufferRHIRef& IndexBufferRHI) override
+		{
+			if (IndexBufferRHI->GetSize() > SmallestPooledBufferSize)
+			{
+				IndexBuffers.Add(MoveTemp(IndexBufferRHI));
+			}
+
+			IndexBufferRHI = nullptr;
+		}
+
+		virtual FVertexBufferRHIRef AllocVertexBuffer(uint32 Stride, uint32 NumElements) override
+		{
+			uint32 BufferSize = GetVertexBufferSize(Stride, NumElements);
+			if (BufferSize > SmallestPooledBufferSize)
+			{
+				FVertexBufferRHIRef Ref = GetSmallestFit(BufferSize, VertexBuffers);
+				if (Ref.IsValid())
+				{
+					return Ref;
+				}
+			}
+
+			return FDynamicMeshBufferAllocator::AllocVertexBuffer(Stride, NumElements);
+		}
+
+		virtual void ReleaseVertexBuffer(FVertexBufferRHIRef& VertexBufferRHI) override
+		{
+			if (VertexBufferRHI->GetSize() > SmallestPooledBufferSize)
+			{
+				VertexBuffers.Add(MoveTemp(VertexBufferRHI));
+			}
+
+			VertexBufferRHI = nullptr;
+		}
+	};
+
+	class FStagingBufferPool
+	{
+	public:
+		FTexture2DRHIRef CreateStagingBuffer_RenderThread(FRHICommandListImmediate& RHICmdList, int32 Width, int32 Height, EPixelFormat Format)
+		{
+			TRACE_CPUPROFILER_EVENT_SCOPE(CreateStagingBuffer_RenderThread)
+
+			auto StagingBufferPredicate = 
+				[Width, Height, Format](const FTexture2DRHIRef& Texture2DRHIRef)
+				{
+					return Texture2DRHIRef->GetSizeX() == Width && Texture2DRHIRef->GetSizeY() == Height && Texture2DRHIRef->GetFormat() == Format;
+				};
+
+			// Process any staging buffers available for unmapping
+			{
+				TArray<FTexture2DRHIRef> ToUnmapLocal;
+				{
+					FScopeLock Lock(&ToUnmapLock);
+					ToUnmapLocal = MoveTemp(ToUnmap);
+				}
+
+				for (int32 Index = 0, Num = ToUnmapLocal.Num(); Index < Num; ++Index)
+				{
+					RHICmdList.UnmapStagingSurface(ToUnmapLocal[Index]);
+					Pool.Add(MoveTemp(ToUnmapLocal[Index]));
+				}
+			}
+
+			// Find any pooled staging buffer with suitable properties.
+			int32 Index = Pool.IndexOfByPredicate(StagingBufferPredicate);
+
+			if (Index != -1)
+			{
+				FTexture2DRHIRef StagingBuffer = MoveTemp(Pool[Index]);
+				Pool.RemoveAtSwap(Index);
+				return StagingBuffer;
+			}
+
+			TRACE_CPUPROFILER_EVENT_SCOPE(RHICreateTexture2D)
+			FRHIResourceCreateInfo CreateInfo;
+			return RHICreateTexture2D(Width, Height, Format, 1, 1, TexCreate_CPUReadback, CreateInfo);
+		}
+
+		void ReleaseStagingBufferForUnmap_AnyThread(FTexture2DRHIRef& Texture2DRHIRef)
+		{
+			TRACE_CPUPROFILER_EVENT_SCOPE(ReleaseStagingBufferForUnmap_AnyThread)
+			FScopeLock Lock(&ToUnmapLock);
+			ToUnmap.Emplace(MoveTemp(Texture2DRHIRef));
+		}
+
+		void Clear_RenderThread(FRHICommandListImmediate& RHICmdList)
+		{
+			TRACE_CPUPROFILER_EVENT_SCOPE(Clear_RenderThread)
+			for (FTexture2DRHIRef& StagingSurface : ToUnmap)
+			{
+				RHICmdList.UnmapStagingSurface(StagingSurface);
+			}
+
+			ToUnmap.Empty();
+			Pool.Empty();
+		}
+
+		~FStagingBufferPool()
+		{
+			check(Pool.Num() == 0);
+		}
+
+	private:
+		TArray<FTexture2DRHIRef> Pool;
+
+		// Not contented enough to warrant the use of lockless structures.
+		FCriticalSection         ToUnmapLock;
+		TArray<FTexture2DRHIRef> ToUnmap;
+	};
+
+	struct FRenderItemKey
+	{
+		const FMeshData* RenderData;
+		const FIntPoint  RenderSize;
+
+		FRenderItemKey(const FMeshData* InRenderData, const FIntPoint& InRenderSize)
+			: RenderData(InRenderData)
+			, RenderSize(InRenderSize)
+		{
+		}
+
+		bool operator == (const FRenderItemKey& Other) const
+		{
+			return RenderData == Other.RenderData &&
+				RenderSize == Other.RenderSize;
+		}
+	};
+
+	uint32 GetTypeHash(const FRenderItemKey& Key)
+	{
+		return HashCombine(GetTypeHash(Key.RenderData), GetTypeHash(Key.RenderSize));
+	}
+}
+
+void FMaterialBakingModule::StartupModule()
+{
+	// Set which properties should enforce gamma correction
+	FMemory::Memset(PerPropertyGamma, (uint8)false);
+	PerPropertyGamma[MP_Normal] = true;
+	PerPropertyGamma[MP_Opacity] = true;
+	PerPropertyGamma[MP_OpacityMask] = true;
+
+	// Set which pixel format should be used for the possible baked out material properties
+	FMemory::Memset(PerPropertyFormat, (uint8)PF_Unknown);
+	PerPropertyFormat[MP_EmissiveColor] = PF_FloatRGBA;
+	PerPropertyFormat[MP_Opacity] = PF_B8G8R8A8;
+	PerPropertyFormat[MP_OpacityMask] = PF_B8G8R8A8;
+	PerPropertyFormat[MP_BaseColor] = PF_B8G8R8A8;
+	PerPropertyFormat[MP_Metallic] = PF_B8G8R8A8;
+	PerPropertyFormat[MP_Specular] = PF_B8G8R8A8;
+	PerPropertyFormat[MP_Roughness] = PF_B8G8R8A8;
+	PerPropertyFormat[MP_Anisotropy] = PF_B8G8R8A8;
+	PerPropertyFormat[MP_Normal] = PF_B8G8R8A8;
+	PerPropertyFormat[MP_Tangent] = PF_B8G8R8A8;
+	PerPropertyFormat[MP_AmbientOcclusion] = PF_B8G8R8A8;
+	PerPropertyFormat[MP_SubsurfaceColor] = PF_B8G8R8A8;
+
+	// Register property customization
+	FPropertyEditorModule& Module = FModuleManager::Get().LoadModuleChecked<FPropertyEditorModule>("PropertyEditor");
+	Module.RegisterCustomPropertyTypeLayout(TEXT("PropertyEntry"), FOnGetPropertyTypeCustomizationInstance::CreateStatic(&FPropertyEntryCustomization::MakeInstance));
+	
+	// Register callback for modified objects
+	FCoreUObjectDelegates::OnObjectModified.AddRaw(this, &FMaterialBakingModule::OnObjectModified);
+}
+
+void FMaterialBakingModule::ShutdownModule()
+{
+	// Unregister customization and callback
+	FPropertyEditorModule* PropertyEditorModule = FModuleManager::GetModulePtr<FPropertyEditorModule>("PropertyEditor");
+
+	if (PropertyEditorModule)
+	{
+		PropertyEditorModule->UnregisterCustomPropertyTypeLayout(TEXT("PropertyEntry"));
+	}
+	FCoreUObjectDelegates::OnObjectModified.RemoveAll(this);
+}
+
+void FMaterialBakingModule::BakeMaterials(const TArray<FMaterialData*>& MaterialSettings, const TArray<FMeshData*>& MeshSettings, TArray<FBakeOutput>& Output)
+{
+	TRACE_CPUPROFILER_EVENT_SCOPE(FMaterialBakingModule::BakeMaterials)
+
+	checkf(MaterialSettings.Num() == MeshSettings.Num(), TEXT("Number of material settings does not match that of MeshSettings"));
+	const int32 NumMaterials = MaterialSettings.Num();
+	const bool bSaveIntermediateTextures = CVarSaveIntermediateTextures.GetValueOnAnyThread() == 1;
+
+	using namespace FMaterialBakingModuleImpl;
+	FMaterialBakingDynamicMeshBufferAllocator MaterialBakingDynamicMeshBufferAllocator;
+
+	FScopedSlowTask Progress(NumMaterials, LOCTEXT("BakeMaterials", "Baking Materials..."), true );
+	Progress.MakeDialog(true);
+
+	TArray<uint32> ProcessingOrder;
+	ProcessingOrder.Reserve(MeshSettings.Num());
+	for (int32 Index = 0; Index < MeshSettings.Num(); ++Index)
+	{
+		ProcessingOrder.Add(Index);
+	}
+
+	// Start with the biggest mesh first so we can always reuse the same vertex/index buffers.
+	// This will decrease the number of allocations backed by newly allocated memory from the OS,
+	// which will reduce soft page faults while copying into that memory.
+	// Soft page faults are now incredibly expensive on Windows 10.
+	Algo::SortBy(
+		ProcessingOrder,
+		[&MeshSettings](const uint32 Index){ return MeshSettings[Index]->RawMeshDescription ? MeshSettings[Index]->RawMeshDescription->Vertices().Num() : 0; },
+		TGreater<>()
+	);
+
+	Output.SetNum(NumMaterials);
+
+	struct FPipelineContext
+	{
+		typedef TFunction<void (FRHICommandListImmediate& RHICmdList)> FReadCommand;
+		FReadCommand ReadCommand;
+	};
+
+	// Distance between the command sent to rendering and the GPU read-back of the result
+	// to minimize sync time waiting on GPU.
+	const int32 PipelineDepth = 16;
+	int32 PipelineIndex = 0;
+	FPipelineContext PipelineContext[PipelineDepth];
+
+	// This will create and prepare FMeshMaterialRenderItem for each property sizes we're going to need
+	auto PrepareRenderItems_AnyThread =
+		[&](int32 MaterialIndex)
+	{
+		TRACE_CPUPROFILER_EVENT_SCOPE(PrepareRenderItems);
+		
+		TMap<FMaterialBakingModuleImpl::FRenderItemKey, FMeshMaterialRenderItem*>* RenderItems = new TMap<FRenderItemKey, FMeshMaterialRenderItem *>();
+		const FMaterialData* CurrentMaterialSettings = MaterialSettings[MaterialIndex];
+		const FMeshData* CurrentMeshSettings = MeshSettings[MaterialIndex];
+
+		for (TMap<EMaterialProperty, FIntPoint>::TConstIterator PropertySizeIterator = CurrentMaterialSettings->PropertySizes.CreateConstIterator(); PropertySizeIterator; ++PropertySizeIterator)
+		{
+			FRenderItemKey RenderItemKey(CurrentMeshSettings, PropertySizeIterator.Value());
+			if (RenderItems->Find(RenderItemKey) == nullptr)
+			{
+				RenderItems->Add(RenderItemKey, new FMeshMaterialRenderItem(CurrentMaterialSettings, CurrentMeshSettings, PropertySizeIterator.Key(), &MaterialBakingDynamicMeshBufferAllocator));
+			}
+		}
+
+		return RenderItems;
+	};
+
+	// We reuse the pipeline depth to prepare render items in advance to avoid stalling the game thread
+	int NextRenderItem = 0;
+	TFuture<TMap<FRenderItemKey, FMeshMaterialRenderItem*>*> PreparedRenderItems[PipelineDepth];
+	for (; NextRenderItem < NumMaterials && NextRenderItem < PipelineDepth; ++NextRenderItem)
+	{
+		PreparedRenderItems[NextRenderItem] = 
+			Async(
+				EAsyncExecution::ThreadPool,
+				[&PrepareRenderItems_AnyThread, &ProcessingOrder, NextRenderItem]()
+				{
+					return PrepareRenderItems_AnyThread(ProcessingOrder[NextRenderItem]);
+				}
+			);
+	}
+
+	// Create all material proxies right away to start compiling shaders asynchronously and avoid stalling the baking process as much as possible
+	{
+		TRACE_CPUPROFILER_EVENT_SCOPE(CreateMaterialProxies)
+
+		for (int32 Index = 0; Index < NumMaterials; ++Index)
+		{
+			int32 MaterialIndex = ProcessingOrder[Index];
+			const FMaterialData* CurrentMaterialSettings = MaterialSettings[MaterialIndex];
+
+			TArray<UTexture*> MaterialTextures;
+			CurrentMaterialSettings->Material->GetUsedTextures(MaterialTextures, EMaterialQualityLevel::Num, true, GMaxRHIFeatureLevel, true);
+
+			// Force load materials used by the current material
+			{
+				TRACE_CPUPROFILER_EVENT_SCOPE(LoadTexturesForMaterial)
+
+				for (UTexture* Texture : MaterialTextures)
+				{
+					if (Texture != NULL)
+					{
+						UTexture2D* Texture2D = Cast<UTexture2D>(Texture);
+						if (Texture2D)
+						{
+							Texture2D->SetForceMipLevelsToBeResident(30.0f);
+							Texture2D->WaitForStreaming();
+						}
+					}
+				}
+			}
+
+			for (TMap<EMaterialProperty, FIntPoint>::TConstIterator PropertySizeIterator = CurrentMaterialSettings->PropertySizes.CreateConstIterator(); PropertySizeIterator; ++PropertySizeIterator)
+			{
+				// They will be stored in the pool and compiled asynchronously
+				CreateMaterialProxy(CurrentMaterialSettings->Material, PropertySizeIterator.Key());
+			}
+		}
+	}
+
+	TAtomic<uint32>    NumTasks(0);
+	FStagingBufferPool StagingBufferPool;
+
+	for (int32 Index = 0; Index < NumMaterials; ++Index)
+	{
+		TRACE_CPUPROFILER_EVENT_SCOPE(BakeOneMaterial)
+
+		Progress.EnterProgressFrame(1.0f, FText::Format(LOCTEXT("BakingMaterial", "Baking Material {0}/{1}"), Index, NumMaterials));
+
+		int32 MaterialIndex = ProcessingOrder[Index];
+		TMap<FRenderItemKey, FMeshMaterialRenderItem*>* RenderItems;
+
+		{
+			TRACE_CPUPROFILER_EVENT_SCOPE(WaitOnPreparedRenderItems)
+			RenderItems = PreparedRenderItems[Index % PipelineDepth].Get();
+		}
+
+		// Prepare the next render item in advance
+		if (NextRenderItem < NumMaterials)
+		{
+			check((NextRenderItem % PipelineDepth) == (Index % PipelineDepth));
+			PreparedRenderItems[NextRenderItem % PipelineDepth] = 
+				Async(
+					EAsyncExecution::ThreadPool,
+					[&PrepareRenderItems_AnyThread, NextMaterialIndex = ProcessingOrder[NextRenderItem]]()
+					{
+						return PrepareRenderItems_AnyThread(NextMaterialIndex);
+					}
+				);
+			NextRenderItem++;
+		}
+
+		const FMaterialData* CurrentMaterialSettings = MaterialSettings[MaterialIndex];
+		const FMeshData* CurrentMeshSettings = MeshSettings[MaterialIndex];
+		FBakeOutput& CurrentOutput = Output[MaterialIndex];
+
+		TArray<EMaterialProperty> MaterialPropertiesToBakeOut;
+		CurrentMaterialSettings->PropertySizes.GenerateKeyArray(MaterialPropertiesToBakeOut);
+
+		const int32 NumPropertiesToRender = MaterialPropertiesToBakeOut.Num();
+		if (NumPropertiesToRender > 0)
+		{
+			TRACE_CPUPROFILER_EVENT_SCOPE(RenderOneMaterial)
+
+			// Ensure data in memory will not change place passed this point to avoid race conditions
+			CurrentOutput.PropertySizes = CurrentMaterialSettings->PropertySizes;
+			for (int32 PropertyIndex = 0; PropertyIndex < NumPropertiesToRender; ++PropertyIndex)
+			{
+				const EMaterialProperty Property = MaterialPropertiesToBakeOut[PropertyIndex];
+				CurrentOutput.PropertyData.Add(Property);
+			}
+
+			for (int32 PropertyIndex = 0; PropertyIndex < NumPropertiesToRender; ++PropertyIndex)
+			{
+				const EMaterialProperty Property = MaterialPropertiesToBakeOut[PropertyIndex];
+				FExportMaterialProxy* ExportMaterialProxy = CreateMaterialProxy(CurrentMaterialSettings->Material, Property);
+
+				if (!ExportMaterialProxy->IsCompilationFinished())
+				{
+					TRACE_CPUPROFILER_EVENT_SCOPE(WaitForMaterialProxyCompilation)
+					ExportMaterialProxy->FinishCompilation();
+				}
+
+				// It is safe to reuse the same render target for each draw pass since they all execute sequentially on the GPU and are copied to staging buffers before
+				// being reused.
+				UTextureRenderTarget2D* RenderTarget = CreateRenderTarget(PerPropertyGamma[Property], PerPropertyFormat[Property], CurrentOutput.PropertySizes[Property]);
+				if (RenderTarget != nullptr)
+				{
+					// Perform everything left of the operation directly on the render thread since we need to modify some RenderItem's properties
+					// for each render pass and we can't do that without costly synchronization (flush) between the game thread and render thread.
+					// Everything slow to execute has already been prepared on the game thread anyway.
+					ENQUEUE_RENDER_COMMAND(RenderOneMaterial)(
+						[this, RenderItems, RenderTarget, Property, ExportMaterialProxy, &PipelineContext, PipelineIndex, &StagingBufferPool, &NumTasks, bSaveIntermediateTextures, &MaterialSettings, &MeshSettings, MaterialIndex, &Output](FRHICommandListImmediate& RHICmdList)
+						{
+							const FMaterialData* CurrentMaterialSettings = MaterialSettings[MaterialIndex];
+							const FMeshData* CurrentMeshSettings = MeshSettings[MaterialIndex];
+
+							FMeshMaterialRenderItem& RenderItem = *RenderItems->FindChecked(FRenderItemKey(CurrentMeshSettings, FIntPoint(RenderTarget->GetSurfaceWidth(), RenderTarget->GetSurfaceHeight())));
+
+							FSceneViewFamily ViewFamily(FSceneViewFamily::ConstructionValues(RenderTarget->GetRenderTargetResource(), nullptr,
+								FEngineShowFlags(ESFIM_Game))
+								.SetWorldTimes(0.0f, 0.0f, 0.0f)
+								.SetGammaCorrection(RenderTarget->GetRenderTargetResource()->GetDisplayGamma()));
+
+							RenderItem.MaterialProperty = Property;
+							RenderItem.MaterialRenderProxy = ExportMaterialProxy;
+							RenderItem.ViewFamily = &ViewFamily;
+
+							FTextureRenderTargetResource* RenderTargetResource = RenderTarget->GetRenderTargetResource();
+							FCanvas Canvas(RenderTargetResource, nullptr, FApp::GetCurrentTime() - GStartTime, FApp::GetDeltaTime(), FApp::GetCurrentTime() - GStartTime, GMaxRHIFeatureLevel);
+							Canvas.SetAllowedModes(FCanvas::Allow_Flush);
+							Canvas.SetRenderTargetRect(FIntRect(0, 0, RenderTarget->GetSurfaceWidth(), RenderTarget->GetSurfaceHeight()));
+							Canvas.SetBaseTransform(Canvas.CalcBaseTransform2D(RenderTarget->GetSurfaceWidth(), RenderTarget->GetSurfaceHeight()));
+
+							// Do rendering
+							Canvas.Clear(RenderTarget->ClearColor);
+							FCanvas::FCanvasSortElement& SortElement = Canvas.GetSortElement(Canvas.TopDepthSortKey());
+							SortElement.RenderBatchArray.Add(&RenderItem);
+							Canvas.Flush_RenderThread(RHICmdList);
+							SortElement.RenderBatchArray.Empty();
+
+							FTexture2DRHIRef StagingBufferRef = StagingBufferPool.CreateStagingBuffer_RenderThread(RHICmdList, RenderTargetResource->GetSizeX(), RenderTargetResource->GetSizeY(), PerPropertyFormat[Property]);
+							FGPUFenceRHIRef GPUFence = RHICreateGPUFence(TEXT("MaterialBackingFence"));
+
+							FResolveRect Rect(0, 0, RenderTargetResource->GetSizeX(), RenderTargetResource->GetSizeY());
+							RHICmdList.CopyToResolveTarget(RenderTargetResource->GetRenderTargetTexture(), StagingBufferRef, FResolveParams(Rect));	
+							RHICmdList.WriteGPUFence(GPUFence);
+
+							// Prepare a lambda for final processing that will be executed asynchronously
+							NumTasks++;
+							auto FinalProcessing_AnyThread =
+								[&NumTasks, bSaveIntermediateTextures, CurrentMaterialSettings, &StagingBufferPool, &Output, Property, MaterialIndex](FTexture2DRHIRef& StagingBuffer, void * Data, int32 DataWidth, int32 DataHeight)
+								{
+									TRACE_CPUPROFILER_EVENT_SCOPE(FinalProcessing)
+
+									FBakeOutput& CurrentOutput  = Output[MaterialIndex];
+									TArray<FColor>& OutputColor = CurrentOutput.PropertyData[Property];
+									FIntPoint& OutputSize       = CurrentOutput.PropertySizes[Property];
+
+									OutputColor.SetNum(OutputSize.X * OutputSize.Y);
+
+									if (Property == MP_EmissiveColor)
+									{
+										// Only one thread will write to CurrentOutput.EmissiveScale since there can be only one emissive channel property per FBakeOutput
+										FMaterialBakingModule::ProcessEmissiveOutput((const FFloat16Color*)Data, DataWidth, OutputSize, OutputColor, CurrentOutput.EmissiveScale);
+									}
+									else
+									{
+										TRACE_CPUPROFILER_EVENT_SCOPE(ConvertRawB8G8R8A8DataToFColor)
+										
+										check(StagingBuffer->GetFormat() == PF_B8G8R8A8);
+										ConvertRawB8G8R8A8DataToFColor(OutputSize.X, OutputSize.Y, (uint8*)Data, DataWidth * sizeof(FColor), OutputColor.GetData());
+									}
+
+									// We can't unmap ourself since we're not on the render thread
+									StagingBufferPool.ReleaseStagingBufferForUnmap_AnyThread(StagingBuffer);
+
+									if (CurrentMaterialSettings->bPerformBorderSmear)
+									{
+										// This will resize the output to a single pixel if the result is monochrome.
+										FMaterialBakingHelpers::PerformUVBorderSmearAndShrink(OutputColor, OutputSize.X, OutputSize.Y);
+									}
+#if WITH_EDITOR
+									// If saving intermediates is turned on
+									if (bSaveIntermediateTextures)
+									{
+										TRACE_CPUPROFILER_EVENT_SCOPE(SaveIntermediateTextures)
+										const UEnum* PropertyEnum = StaticEnum<EMaterialProperty>();
+										FName PropertyName = PropertyEnum->GetNameByValue(Property);
+										FString TrimmedPropertyName = PropertyName.ToString();
+										TrimmedPropertyName.RemoveFromStart(TEXT("MP_"));
+
+										const FString DirectoryPath = FPaths::ConvertRelativePathToFull(FPaths::ProjectIntermediateDir() + TEXT("MaterialBaking/"));
+										FString FilenameString = FString::Printf(TEXT("%s%s-%d-%s.bmp"), *DirectoryPath, *CurrentMaterialSettings->Material->GetName(), MaterialIndex, *TrimmedPropertyName);
+										FFileHelper::CreateBitmap(*FilenameString, CurrentOutput.PropertySizes[Property].X, CurrentOutput.PropertySizes[Property].Y, CurrentOutput.PropertyData[Property].GetData());
+									}
+#endif // WITH_EDITOR
+									NumTasks--;
+								};
+
+							// Run previous command if we're going to overwrite it meaning pipeline depth has been reached
+							if (PipelineContext[PipelineIndex].ReadCommand)
+							{
+								PipelineContext[PipelineIndex].ReadCommand(RHICmdList);
+							}
+
+							// Generate a texture reading command that will be executed once it reaches the end of the pipeline
+							PipelineContext[PipelineIndex].ReadCommand =
+								[FinalProcessing_AnyThread, StagingBufferRef = MoveTemp(StagingBufferRef), GPUFence = MoveTemp(GPUFence)](FRHICommandListImmediate& RHICmdList) mutable
+								{
+									TRACE_CPUPROFILER_EVENT_SCOPE(MapAndEnqueue)
+
+									void * Data = nullptr;
+									int32 Width; int32 Height;
+									RHICmdList.MapStagingSurface(StagingBufferRef, GPUFence.GetReference(), Data, Width, Height);
+
+									// Schedule the copy and processing on another thread to free up the render thread as much as possible
+									Async(
+										EAsyncExecution::ThreadPool,
+										[FinalProcessing_AnyThread, Data, Width, Height, StagingBufferRef = MoveTemp(StagingBufferRef)]() mutable
+										{
+											FinalProcessing_AnyThread(StagingBufferRef, Data, Width, Height);
+										}
+									);
+								};
+						}
+					);
+
+					PipelineIndex = (PipelineIndex + 1) % PipelineDepth;
+				}
+			}
+
+		}
+
+		// Destroying Render Items must happen on the render thread to ensure
+		// they are not used anymore.
+		ENQUEUE_RENDER_COMMAND(DestroyRenderItems)(
+			[RenderItems](FRHICommandListImmediate& RHICmdList)
+			{
+				for (auto RenderItem : (*RenderItems))
+				{
+					delete RenderItem.Value;
+				}
+
+				delete RenderItems;
+			}
+		);
+	}
+
+	ENQUEUE_RENDER_COMMAND(ProcessRemainingReads)(
+		[&PipelineContext, PipelineDepth, PipelineIndex](FRHICommandListImmediate& RHICmdList)
+		{
+			// Enqueue remaining reads
+			for (int32 Index = 0; Index < PipelineDepth; Index++)
+			{
+				int32 LocalPipelineIndex = (PipelineIndex + Index) % PipelineDepth;
+
+				if (PipelineContext[LocalPipelineIndex].ReadCommand)
+				{
+					PipelineContext[LocalPipelineIndex].ReadCommand(RHICmdList);
+				}
+			}
+		}
+	);
+
+	// Wait until every tasks have been queued so that NumTasks is only decreasing
+	FlushRenderingCommands();
+
+	// Wait for any remaining final processing tasks
+	while (NumTasks.Load(EMemoryOrder::Relaxed) > 0)
+	{
+		FPlatformProcess::Sleep(0.1f);
+	}
+
+	// Wait for all tasks to have been processed before clearing the staging buffers
+	FlushRenderingCommands();
+
+	ENQUEUE_RENDER_COMMAND(ClearStagingBufferPool)(
+		[&StagingBufferPool](FRHICommandListImmediate& RHICmdList)
+		{
+			StagingBufferPool.Clear_RenderThread(RHICmdList);
+		}
+	);
+
+	// Wait for StagingBufferPool clear to have executed before exiting the function
+	FlushRenderingCommands();
+
+	if (!CVarUseMaterialProxyCaching.GetValueOnAnyThread())
+	{
+		CleanupMaterialProxies();
+	}
+}
+
+bool FMaterialBakingModule::SetupMaterialBakeSettings(TArray<TWeakObjectPtr<UObject>>& OptionObjects, int32 NumLODs)
+{
+	TSharedRef<SWindow> Window = SNew(SWindow)
+		.Title(LOCTEXT("WindowTitle", "Material Baking Options"))
+		.SizingRule(ESizingRule::Autosized);
+
+	TSharedPtr<SMaterialOptions> Options;
+
+	Window->SetContent
+	(
+		SAssignNew(Options, SMaterialOptions)
+		.WidgetWindow(Window)
+		.NumLODs(NumLODs)
+		.SettingsObjects(OptionObjects)
+	);
+
+	TSharedPtr<SWindow> ParentWindow;
+	if (FModuleManager::Get().IsModuleLoaded("MainFrame"))
+	{
+		IMainFrameModule& MainFrame = FModuleManager::LoadModuleChecked<IMainFrameModule>("MainFrame");
+		ParentWindow = MainFrame.GetParentWindow();
+		FSlateApplication::Get().AddModalWindow(Window, ParentWindow, false);
+		return !Options->WasUserCancelled();
+	}
+
+	return false;
+}
+
+void FMaterialBakingModule::CleanupMaterialProxies()
+{
+	for (auto Iterator : MaterialProxyPool)
+	{
+		delete Iterator.Value.Value;
+	}
+	MaterialProxyPool.Reset();
+}
+
+UTextureRenderTarget2D* FMaterialBakingModule::CreateRenderTarget(bool bInForceLinearGamma, EPixelFormat InPixelFormat, const FIntPoint& InTargetSize)
+{
+	TRACE_CPUPROFILER_EVENT_SCOPE(FMaterialBakingModule::CreateRenderTarget)
+
+	UTextureRenderTarget2D* RenderTarget = nullptr;
+	const int32 MaxTextureSize = 1 << (MAX_TEXTURE_MIP_COUNT - 1); // Don't use GetMax2DTextureDimension() as this is for the RHI only.
+	const FIntPoint ClampedTargetSize(FMath::Clamp(InTargetSize.X, 1, MaxTextureSize), FMath::Clamp(InTargetSize.Y, 1, MaxTextureSize));
+	auto RenderTargetComparison = [bInForceLinearGamma, InPixelFormat, ClampedTargetSize](const UTextureRenderTarget2D* CompareRenderTarget) -> bool
+	{
+		return (CompareRenderTarget->SizeX == ClampedTargetSize.X && CompareRenderTarget->SizeY == ClampedTargetSize.Y && CompareRenderTarget->OverrideFormat == InPixelFormat && CompareRenderTarget->bForceLinearGamma == bInForceLinearGamma);
+	};
+
+	// Find any pooled render target with suitable properties.
+	UTextureRenderTarget2D** FindResult = RenderTargetPool.FindByPredicate(RenderTargetComparison);
+	
+	if (FindResult)
+	{
+		RenderTarget = *FindResult;
+	}
+	else
+	{
+		TRACE_CPUPROFILER_EVENT_SCOPE(CreateNewRenderTarget)
+
+		// Not found - create a new one.
+		RenderTarget = NewObject<UTextureRenderTarget2D>();
+		check(RenderTarget);
+		RenderTarget->AddToRoot();
+		RenderTarget->ClearColor = FLinearColor(1.0f, 0.0f, 1.0f);
+		RenderTarget->ClearColor.A = 1.0f;
+		RenderTarget->TargetGamma = 0.0f;
+		RenderTarget->InitCustomFormat(ClampedTargetSize.X, ClampedTargetSize.Y, InPixelFormat, bInForceLinearGamma);
+
+		RenderTargetPool.Add(RenderTarget);
+	}
+
+	checkf(RenderTarget != nullptr, TEXT("Unable to create or find valid render target"));
+	return RenderTarget;
+}
+
+FExportMaterialProxy* FMaterialBakingModule::CreateMaterialProxy(UMaterialInterface* Material, const EMaterialProperty Property)
+{
+	TRACE_CPUPROFILER_EVENT_SCOPE(FMaterialBakingModule::CreateMaterialProxy)
+
+	FExportMaterialProxy* Proxy = nullptr;
+
+	// Find all pooled material proxy matching this material
+	TArray<FMaterialPoolValue> Entries;
+	MaterialProxyPool.MultiFind(Material, Entries);
+
+	// Look for the matching property
+	for (FMaterialPoolValue& Entry : Entries)
+	{
+		if (Entry.Key == Property)
+		{
+			Proxy = Entry.Value;
+			break;
+		}
+	}
+
+	// Not found, create a new entry
+	if (Proxy == nullptr)
+	{
+		Proxy = new FExportMaterialProxy(Material, Property, false /* bInSynchronousCompilation */);
+		MaterialProxyPool.Add(Material, FMaterialPoolValue(Property, Proxy));
+	}
+
+	return Proxy;
+}
+
+void FMaterialBakingModule::ProcessEmissiveOutput(const FFloat16Color* Color16, int32 Color16Pitch, const FIntPoint& OutputSize, TArray<FColor>& OutputColor, float& EmissiveScale)
+{
+	TRACE_CPUPROFILER_EVENT_SCOPE(FMaterialBakingModule::ProcessEmissiveOutput)
+
+	const int32 NumThreads = [&]()
+	{
+		return FPlatformProcess::SupportsMultithreading() ? FPlatformMisc::NumberOfCores() : 1;
+	}();
+
+	float* MaxValue = new float[NumThreads];
+	FMemory::Memset(MaxValue, 0, NumThreads * sizeof(MaxValue[0]));
+	const int32 LinesPerThread = FMath::CeilToInt((float)OutputSize.Y / (float)NumThreads);
+
+	// Find maximum float value across texture
+	ParallelFor(NumThreads, [&Color16, LinesPerThread, MaxValue, OutputSize, Color16Pitch](int32 Index)
+	{
+		const int32 EndY = FMath::Min((Index + 1) * LinesPerThread, OutputSize.Y);			
+		float& CurrentMaxValue = MaxValue[Index];
+		const FFloat16Color MagentaFloat16 = FFloat16Color(FLinearColor(1.0f, 0.0f, 1.0f));
+		for (int32 PixelY = Index * LinesPerThread; PixelY < EndY; ++PixelY)
+		{
+			const int32 SrcYOffset = PixelY * Color16Pitch;
+			for (int32 PixelX = 0; PixelX < OutputSize.X; PixelX++)
+			{
+				const FFloat16Color& Pixel16 = Color16[PixelX + SrcYOffset];
+				// Find maximum channel value across texture
+				if (!(Pixel16 == MagentaFloat16))
+				{
+					CurrentMaxValue = FMath::Max(CurrentMaxValue, FMath::Max3(Pixel16.R.GetFloat(), Pixel16.G.GetFloat(), Pixel16.B.GetFloat()));
+				}
+			}
+		}
+	});
+
+	const float GlobalMaxValue = [&MaxValue, NumThreads]
+	{
+		float TempValue = 0.0f;
+		for (int32 ThreadIndex = 0; ThreadIndex < NumThreads; ++ThreadIndex)
+		{
+			TempValue = FMath::Max(TempValue, MaxValue[ThreadIndex]);
+		}
+
+		return TempValue;
+	}();
+		
+	if (GlobalMaxValue <= 0.01f)
+	{
+		// Black emissive, drop it			
+	}
+
+	// Now convert Float16 to Color using the scale
+	OutputColor.SetNumUninitialized(OutputSize.X * OutputSize.Y);
+	const float Scale = 255.0f / GlobalMaxValue;
+	ParallelFor(NumThreads, [&Color16, LinesPerThread, &OutputColor, OutputSize, Color16Pitch, Scale](int32 Index)
+	{
+		const FFloat16Color MagentaFloat16 = FFloat16Color(FLinearColor(1.0f, 0.0f, 1.0f));
+
+		const int32 EndY = FMath::Min((Index + 1) * LinesPerThread, OutputSize.Y);
+		for (int32 PixelY = Index * LinesPerThread; PixelY < EndY; ++PixelY)
+		{
+			const int32 SrcYOffset = PixelY * Color16Pitch;
+			const int32 DstYOffset = PixelY * OutputSize.X;
+
+			for (int32 PixelX = 0; PixelX < OutputSize.X; PixelX++)
+			{
+				const FFloat16Color& Pixel16 = Color16[PixelX + SrcYOffset];
+				FColor& Pixel8 = OutputColor[PixelX + DstYOffset];
+
+				if (Pixel16 == MagentaFloat16)
+				{
+					Pixel8.R = 255;
+					Pixel8.G = 0;
+					Pixel8.B = 255;
+				}
+				else
+				{
+					Pixel8.R = (uint8)FMath::RoundToInt(Pixel16.R.GetFloat() * Scale);
+					Pixel8.G = (uint8)FMath::RoundToInt(Pixel16.G.GetFloat() * Scale);
+					Pixel8.B = (uint8)FMath::RoundToInt(Pixel16.B.GetFloat() * Scale);
+				}
+					
+				Pixel8.A = 255;
+			}
+		}
+	});
+
+	// This scale will be used in the proxy material to get the original range of emissive values outside of 0-1
+	EmissiveScale = GlobalMaxValue;
+}
+
+void FMaterialBakingModule::OnObjectModified(UObject* Object)
+{
+	TRACE_CPUPROFILER_EVENT_SCOPE(FMaterialBakingModule::OnObjectModified)
+
+	if (CVarUseMaterialProxyCaching.GetValueOnAnyThread())
+	{
+		if (Object && Object->IsA<UMaterialInterface>())
+		{
+			UMaterialInterface* MaterialToInvalidate = Cast<UMaterialInterface>(Object);
+
+			// Search our proxy pool for materials or material instances that refer to MaterialToInvalidate
+			for (auto It = MaterialProxyPool.CreateIterator(); It; ++It)
+			{
+				TWeakObjectPtr<UMaterialInterface> PoolMaterialPtr = It.Key();
+
+				// Remove stale entries from the pool
+				bool bMustDelete = PoolMaterialPtr.IsValid();
+				if (!bMustDelete)
+				{
+					bMustDelete = PoolMaterialPtr == MaterialToInvalidate;
+				}
+
+				// No match - Test the MaterialInstance hierarchy
+				if (!bMustDelete)
+				{
+					UMaterialInstance* MaterialInstance = Cast<UMaterialInstance>(PoolMaterialPtr);
+					while (!bMustDelete && MaterialInstance && MaterialInstance->Parent != nullptr)
+					{
+						bMustDelete = MaterialInstance->Parent == MaterialToInvalidate;
+						MaterialInstance = Cast<UMaterialInstance>(MaterialInstance->Parent);
+					}
+				}
+
+				// We have a match, remove the entry from our pool
+				if (bMustDelete)
+				{
+					FExportMaterialProxy* Proxy = It.Value().Value;
+
+					ENQUEUE_RENDER_COMMAND(DeleteCachedMaterialProxy)(
+						[Proxy](FRHICommandListImmediate& RHICmdList)
+						{
+							delete Proxy;
+						});
+
+					It.RemoveCurrent();
+				}
+			}
+		}
+	}
+}
+
+#undef LOCTEXT_NAMESPACE //"MaterialBakingModule"
diff --git a/Engine/Plugins/Enterprise/GLTFExporter/MigratedFromGit/Source/GLTFMaterialBaking/Private/MaterialBakingModule.h b/Engine/Plugins/Enterprise/GLTFExporter/MigratedFromGit/Source/GLTFMaterialBaking/Private/MaterialBakingModule.h
new file mode 100644
--- /dev/null
+++ b/Engine/Plugins/Enterprise/GLTFExporter/MigratedFromGit/Source/GLTFMaterialBaking/Private/MaterialBakingModule.h
@@ -0,0 +1,66 @@
+// Copyright Epic Games, Inc. All Rights Reserved.
+
+#pragma once
+
+#include "CoreMinimal.h"
+#include "Modules/ModuleInterface.h"
+#include "PixelFormat.h"
+#include "SceneTypes.h"
+#include "IMaterialBakingModule.h"
+
+class UTextureRenderTarget2D;
+class UMaterialOptions;
+class UMaterialInterface;
+
+class FExportMaterialProxy;
+class FTextureRenderTargetResource;
+
+struct FMaterialData; 
+struct FMeshData;
+struct FBakeOutput;
+
+class MATERIALBAKING_API FMaterialBakingModule : public IMaterialBakingModule
+{
+public:
+	/** IModuleInterface overrides begin */
+	virtual void StartupModule();
+	virtual void ShutdownModule();
+	/** IModuleInterface overrides end */
+
+	/** Bakes out material properties according to MaterialSettings using MeshSettings and stores the output in Output */
+	virtual void BakeMaterials(const TArray<FMaterialData*>& MaterialSettings, const TArray<FMeshData*>& MeshSettings, TArray<FBakeOutput>& Output) override;
+
+	/** Promps a slate window to allow the user to populate specific material baking settings used while baking out materials */
+	virtual bool SetupMaterialBakeSettings(TArray<TWeakObjectPtr<UObject>>& OptionObjects, int32 NumLODs) override;
+
+protected:
+	/* Creates and adds or reuses a RenderTarget from the pool */
+	UTextureRenderTarget2D* CreateRenderTarget(bool bInForceLinearGamma, EPixelFormat InPixelFormat, const FIntPoint& InTargetSize);
+
+	/* Creates and adds (or reuses a ExportMaterialProxy from the pool if MaterialBaking.UseMaterialProxyCaching is set to 1) */
+	FExportMaterialProxy* CreateMaterialProxy(UMaterialInterface* Material, const EMaterialProperty Property );
+
+	/** Helper for emissive color conversion to Output */
+	static void ProcessEmissiveOutput(const FFloat16Color* Color16, int32 Color16Pitch, const FIntPoint& OutputSize, TArray<FColor>& Output, float& EmissiveScale);
+
+	/** Cleans up all cached material proxies in MaterialProxyPool */
+	void CleanupMaterialProxies();
+
+	/** Callback for modified objects which should be removed from MaterialProxyPool in that case */
+	void OnObjectModified(UObject* Object);
+private:
+	/** Pool of available render targets, cached for re-using on consecutive property rendering */
+	TArray<UTextureRenderTarget2D*> RenderTargetPool;
+
+	/** Pool of cached material proxies to optimize material baking workflow, stays resident when MaterialBaking.UseMaterialProxyCaching is set to 1 */
+	typedef TWeakObjectPtr<UMaterialInterface>				FMaterialPoolKey;
+	typedef TPair<EMaterialProperty, FExportMaterialProxy*> FMaterialPoolValue;
+	typedef TMultiMap<FMaterialPoolKey, FMaterialPoolValue, FDefaultSetAllocator, TWeakObjectPtrMapKeyFuncs<FMaterialPoolKey, FMaterialPoolValue, true /*bInAllowDuplicateKeys*/>> FMaterialPoolMap;
+	FMaterialPoolMap MaterialProxyPool;
+
+	/** Pixel formats to use for baking out specific material properties */
+	EPixelFormat PerPropertyFormat[MP_MAX];
+
+	/** Whether or not to enforce gamma correction while baking out specific material properties */
+	bool PerPropertyGamma[MP_MAX];
+};
diff --git a/Engine/Plugins/Enterprise/GLTFExporter/MigratedFromGit/Source/GLTFMaterialBaking/Private/MaterialOptionsCustomization.cpp b/Engine/Plugins/Enterprise/GLTFExporter/MigratedFromGit/Source/GLTFMaterialBaking/Private/MaterialOptionsCustomization.cpp
new file mode 100644
--- /dev/null
+++ b/Engine/Plugins/Enterprise/GLTFExporter/MigratedFromGit/Source/GLTFMaterialBaking/Private/MaterialOptionsCustomization.cpp
@@ -0,0 +1,215 @@
+// Copyright Epic Games, Inc. All Rights Reserved.
+
+#include "MaterialOptionsCustomization.h"
+#include "MaterialOptions.h"
+#include "PropertyHandle.h"
+#include "DetailWidgetRow.h"
+#include "IDetailChildrenBuilder.h"
+#include "IDetailPropertyRow.h"
+#include "DetailCategoryBuilder.h"
+#include "DetailLayoutBuilder.h"
+#include "Widgets/Input/SCheckBox.h"
+#include "Widgets/Text/STextBlock.h"
+#include "RHI.h"
+
+#include "IPropertyTypeCustomization.h"
+#include "IPropertyUtilities.h"
+
+static const TArray<EMaterialProperty> DisabledProperties = { MP_Refraction };
+
+TSharedRef<IPropertyTypeCustomization> FPropertyEntryCustomization::MakeInstance()
+{
+	return MakeShareable(new FPropertyEntryCustomization);
+}
+
+FPropertyEntryCustomization::FPropertyEntryCustomization()
+{
+	PropertyRestriction = MakeShareable(new FPropertyRestriction(FText::FromString("Property already set on for a different entry")));
+	CurrentOptions = nullptr;
+}
+
+void FPropertyEntryCustomization::CustomizeHeader(TSharedRef<IPropertyHandle> PropertyHandle, FDetailWidgetRow& HeaderRow, IPropertyTypeCustomizationUtils& CustomizationUtils)
+{
+	TSharedPtr<IPropertyHandle> MaterialPropertyHandle = PropertyHandle->GetChildHandle(GET_MEMBER_NAME_CHECKED(FPropertyEntry, Property));
+	HeaderRow.NameContent()
+	[
+		MaterialPropertyHandle->CreatePropertyValueWidget()
+	];
+
+	const TArray<TWeakObjectPtr<UObject>>& SelectedObjects = CustomizationUtils.GetPropertyUtilities()->GetSelectedObjects();
+
+	// Try and find material options instance in currently edited objects
+	const TWeakObjectPtr<UObject>* MaterialOptionsPtr = SelectedObjects.FindByPredicate([](TWeakObjectPtr<UObject> Object) { return Cast<UMaterialOptions>(Object.Get()); });
+	CurrentOptions = MaterialOptionsPtr ? Cast<UMaterialOptions>((MaterialOptionsPtr)->Get()) : nullptr;
+	
+	const int32 Index = PropertyHandle->GetIndexInArray();
+
+	// Add restriction to ensure the user cannot set up two entries with the same EMaterialProperty value	
+	MaterialPropertyHandle->AddRestriction(PropertyRestriction.ToSharedRef());
+
+	// Set Parent handle on property change to update restrictions across all the entry properties
+	TSharedPtr<IPropertyHandle> ParentHandle = PropertyHandle->GetParentHandle();
+	if (ParentHandle.IsValid())
+	{
+		TSharedPtr<IPropertyHandle> TopParentHandle = ParentHandle->GetParentHandle();
+		if (TopParentHandle.IsValid())
+		{
+			TopParentHandle->SetOnChildPropertyValueChanged(FSimpleDelegate::CreateSP(this, &FPropertyEntryCustomization::UpdateRestrictions, Index));
+		}
+	}
+
+	UpdateRestrictions(Index);
+}
+
+void FPropertyEntryCustomization::CustomizeChildren(TSharedRef<IPropertyHandle> PropertyHandle, IDetailChildrenBuilder& ChildBuilder, IPropertyTypeCustomizationUtils& CustomizationUtils)
+{
+	TSharedPtr<IPropertyHandle> CustomSizePropertyHandle = PropertyHandle->GetChildHandle(GET_MEMBER_NAME_CHECKED(FPropertyEntry, CustomSize));
+	IDetailPropertyRow& SizeRow = ChildBuilder.AddProperty(CustomSizePropertyHandle.ToSharedRef());
+
+	AddTextureSizeClamping(CustomSizePropertyHandle);
+
+	TSharedPtr<IPropertyHandle> ConstantValuePropertyHandle = PropertyHandle->GetChildHandle(GET_MEMBER_NAME_CHECKED(FPropertyEntry, ConstantValue));
+	IDetailPropertyRow& ConstantValueRow = ChildBuilder.AddProperty(ConstantValuePropertyHandle.ToSharedRef());
+}
+
+void FPropertyEntryCustomization::UpdateRestrictions(const int32 EntryIndex)
+{	
+	PropertyRestriction->RemoveAll();
+	const UEnum* PropertyEnum = StaticEnum<EMaterialProperty>();
+	if (CurrentOptions)
+	{	
+		// Add all previously set material properties to be disabled
+		for (int32 Index = 0; Index < CurrentOptions->Properties.Num(); ++Index)
+		{
+			if (Index != EntryIndex)
+			{
+				const FPropertyEntry& Entry = CurrentOptions->Properties[Index];
+				PropertyRestriction->AddDisabledValue(PropertyEnum->GetNameStringByValue(Entry.Property));
+			}
+		}
+	}
+
+	for (const EMaterialProperty& Property : DisabledProperties)
+	{
+		PropertyRestriction->AddDisabledValue(PropertyEnum->GetNameStringByValue(Property));
+	}
+}
+
+TSharedRef<IDetailCustomization> FMaterialOptionsCustomization::MakeInstance(int32 InNumLODs)
+{
+	return MakeShareable(new FMaterialOptionsCustomization(InNumLODs));
+}
+
+FMaterialOptionsCustomization::FMaterialOptionsCustomization(int32 InNumLODs)
+	: NumLODs(InNumLODs)
+{}
+
+void FMaterialOptionsCustomization::CustomizeDetails(IDetailLayoutBuilder& DetailBuilder)
+{
+	const FName CategoryName = FName(TEXT("MeshSettings"));
+	IDetailCategoryBuilder& CategoryBuilder = DetailBuilder.EditCategory(CategoryName);
+
+	// Add custom LOD index selection row
+	FDetailWidgetRow& LODsRow = CategoryBuilder.AddCustomRow(FText::FromString(TEXT("LODs")));
+	LODsRow.NameContent()
+	[
+		SNew(STextBlock)
+		.Text(FText::FromString(TEXT("LODs")))
+		.Font(DetailBuilder.GetDetailFont())
+	];
+
+	TSharedPtr<SHorizontalBox> ContentBox;
+	LODsRow.ValueContent()
+	[
+		SAssignNew(ContentBox, SHorizontalBox)
+	];
+	
+	TArray<TWeakObjectPtr<UObject>> WeakObjects;
+	DetailBuilder.GetObjectsBeingCustomized(WeakObjects);
+	// Try and find material options instance in currently edited objects
+	UMaterialOptions* CurrentOptions = Cast<UMaterialOptions>((WeakObjects.FindByPredicate([](TWeakObjectPtr<UObject> Object) { return Cast<UMaterialOptions>(Object.Get()); }))->Get());
+
+	TSharedRef<IPropertyHandle> TextureSizePropertyHandle = DetailBuilder.GetProperty(GET_MEMBER_NAME_CHECKED(UMaterialOptions, TextureSize));
+	if (TextureSizePropertyHandle->IsValidHandle())
+	{
+		AddTextureSizeClamping(TextureSizePropertyHandle);
+	}
+
+	TSharedRef<IPropertyHandle> PropertiesHandle = DetailBuilder.GetProperty(GET_MEMBER_NAME_CHECKED(UMaterialOptions, Properties));
+	if (PropertiesHandle->IsValidHandle())
+	{
+		// Setup delegate for when the number of Material Property items changes
+		FSimpleDelegate RefreshDelegate = FSimpleDelegate::CreateLambda([this, &DetailBuilder, CurrentOptions]()
+		{
+			TArray<EMaterialProperty> Properties;
+			// Ensure that we set duplicate entries to MP_MAX
+			for (FPropertyEntry& Entry : CurrentOptions->Properties)
+			{
+				if (Properties.Contains(Entry.Property))
+				{
+					Entry.Property = MP_MAX;
+				}
+				else if (Entry.Property != MP_MAX && !DisabledProperties.Contains(Entry.Property))
+				{
+					Properties.Add(Entry.Property);
+				}
+			}
+			
+			DetailBuilder.ForceRefreshDetails();
+		});
+		PropertiesHandle->SetOnPropertyValueChanged(RefreshDelegate);
+	}
+
+	// Only allow changes to LOD indices if we have a valid options instance and if there is actually more than one index
+	ContentBox->SetEnabled(TAttribute<bool>::Create([=]() -> bool { return NumLODs > 1 && CurrentOptions != nullptr; }));
+	for (int32 LODIndex = 0; LODIndex < NumLODs; ++LODIndex)
+	{
+		ContentBox->AddSlot()
+		.Padding(FMargin(0.0f, 0.0f, 0.0f, 0.0f))
+		.AutoWidth()
+		[
+			SNew(SCheckBox)
+			.IsChecked(LODIndex == 0 ? ECheckBoxState::Checked : ECheckBoxState::Unchecked)
+			.OnCheckStateChanged_Lambda([=](ECheckBoxState NewState)
+			{
+				if (NewState == ECheckBoxState::Unchecked)
+				{
+					CurrentOptions->LODIndices.Remove(LODIndex);
+				}
+				else
+				{
+					CurrentOptions->LODIndices.Add(LODIndex);
+				}
+			})
+		];
+
+		ContentBox->AddSlot()
+		.Padding(FMargin(3.0f, 2.0f, 4.0f, 0.0f))
+		.AutoWidth()
+		[
+			SNew(STextBlock)
+			.Text(FText::FromString(FString::FromInt(LODIndex)))
+			.Font(DetailBuilder.GetDetailFont())
+		];
+	}
+}
+
+void AddTextureSizeClamping(TSharedPtr<IPropertyHandle> TextureSizeProperty)
+{
+	TSharedPtr<IPropertyHandle> PropertyX = TextureSizeProperty->GetChildHandle(GET_MEMBER_NAME_CHECKED(FIntPoint, X));
+	TSharedPtr<IPropertyHandle> PropertyY = TextureSizeProperty->GetChildHandle(GET_MEMBER_NAME_CHECKED(FIntPoint, Y));
+
+	const FString MaxTextureResolutionString = FString::FromInt(GetMax2DTextureDimension());
+	TextureSizeProperty->GetProperty()->SetMetaData(TEXT("ClampMax"), *MaxTextureResolutionString);
+	TextureSizeProperty->GetProperty()->SetMetaData(TEXT("UIMax"), *MaxTextureResolutionString);
+	PropertyX->SetInstanceMetaData(TEXT("ClampMax"), *MaxTextureResolutionString);
+	PropertyX->SetInstanceMetaData(TEXT("UIMax"), *MaxTextureResolutionString);
+	PropertyY->SetInstanceMetaData(TEXT("ClampMax"), *MaxTextureResolutionString);
+	PropertyY->SetInstanceMetaData(TEXT("UIMax"), *MaxTextureResolutionString);
+
+	const FString MinTextureResolutionString("1");
+	PropertyX->SetInstanceMetaData(TEXT("ClampMin"), *MinTextureResolutionString);
+	PropertyX->SetInstanceMetaData(TEXT("UIMin"), *MinTextureResolutionString);
+	PropertyY->SetInstanceMetaData(TEXT("ClampMin"), *MinTextureResolutionString);
+	PropertyY->SetInstanceMetaData(TEXT("UIMin"), *MinTextureResolutionString);
+}
diff --git a/Engine/Plugins/Enterprise/GLTFExporter/MigratedFromGit/Source/GLTFMaterialBaking/Private/MaterialOptionsCustomization.h b/Engine/Plugins/Enterprise/GLTFExporter/MigratedFromGit/Source/GLTFMaterialBaking/Private/MaterialOptionsCustomization.h
new file mode 100644
--- /dev/null
+++ b/Engine/Plugins/Enterprise/GLTFExporter/MigratedFromGit/Source/GLTFMaterialBaking/Private/MaterialOptionsCustomization.h
@@ -0,0 +1,38 @@
+// Copyright Epic Games, Inc. All Rights Reserved.
+
+#pragma once
+
+#include "IPropertyTypeCustomization.h"
+#include "IDetailCustomization.h"
+#include "PropertyRestriction.h"
+
+/** Property customization for Material Property entries */
+class FPropertyEntryCustomization : public IPropertyTypeCustomization
+{
+public:
+	static TSharedRef<IPropertyTypeCustomization> MakeInstance();
+	FPropertyEntryCustomization();
+
+	virtual void CustomizeHeader(TSharedRef<IPropertyHandle> PropertyHandle, FDetailWidgetRow& HeaderRow, IPropertyTypeCustomizationUtils& CustomizationUtils) override;
+	virtual void CustomizeChildren(TSharedRef<IPropertyHandle> PropertyHandle, IDetailChildrenBuilder& ChildBuilder, IPropertyTypeCustomizationUtils& CustomizationUtils) override;
+protected:
+	void UpdateRestrictions(const int32 EntryIndex);
+protected:
+	/** Property restriction instance used for limiting EMaterialProperty selection */
+	TSharedPtr<FPropertyRestriction> PropertyRestriction;
+	class UMaterialOptions* CurrentOptions;
+};
+
+/** Detail customization for UMaterialOptions */
+class FMaterialOptionsCustomization : public IDetailCustomization
+{
+public:
+	static TSharedRef<IDetailCustomization> MakeInstance(int32 InNumLODs);
+	FMaterialOptionsCustomization(int32 InNumLODs);
+
+	virtual void CustomizeDetails(IDetailLayoutBuilder& DetailBuilder) override;
+protected:
+	int32 NumLODs;
+};
+
+static void AddTextureSizeClamping(TSharedPtr<IPropertyHandle> TextureSizeProperty);
diff --git a/Engine/Plugins/Enterprise/GLTFExporter/MigratedFromGit/Source/GLTFMaterialBaking/Private/MaterialOptionsWindow.cpp b/Engine/Plugins/Enterprise/GLTFExporter/MigratedFromGit/Source/GLTFMaterialBaking/Private/MaterialOptionsWindow.cpp
new file mode 100644
--- /dev/null
+++ b/Engine/Plugins/Enterprise/GLTFExporter/MigratedFromGit/Source/GLTFMaterialBaking/Private/MaterialOptionsWindow.cpp
@@ -0,0 +1,123 @@
+// Copyright Epic Games, Inc. All Rights Reserved.
+
+#include "MaterialOptionsWindow.h"
+#include "Modules/ModuleManager.h"
+#include "PropertyEditorModule.h"
+
+#include "IDetailsView.h"
+#include "Widgets/Layout/SUniformGridPanel.h"
+#include "Widgets/Input/SButton.h"
+
+#include "MaterialBakingStructures.h"
+#include "MaterialOptions.h"
+#include "MaterialOptionsCustomization.h"
+#include "Misc/MessageDialog.h"
+
+#define LOCTEXT_NAMESPACE "SMaterialOptions"
+
+SMaterialOptions::SMaterialOptions() : bUserCancelled(true)
+{
+
+}
+
+void SMaterialOptions::Construct(const FArguments& InArgs)
+{
+	WidgetWindow = InArgs._WidgetWindow;
+
+	// Retrieve property editor module and create a SDetailsView
+	FPropertyEditorModule& PropertyEditorModule = FModuleManager::GetModuleChecked<FPropertyEditorModule>("PropertyEditor");
+	FDetailsViewArgs DetailsViewArgs;
+	DetailsViewArgs.bAllowSearch = false;
+	DetailsViewArgs.NameAreaSettings = FDetailsViewArgs::HideNameArea;
+	DetailsViewArgs.bAllowMultipleTopLevelObjects = true;
+
+	DetailsView = PropertyEditorModule.CreateDetailView(DetailsViewArgs);
+	// Register instance property customization
+	DetailsView->RegisterInstancedCustomPropertyLayout(UMaterialOptions::StaticClass(), FOnGetDetailCustomizationInstance::CreateLambda([=]() { return FMaterialOptionsCustomization::MakeInstance(InArgs._NumLODs); }));
+
+	// Set provided objects on SDetailsView
+	DetailsView->SetObjects(InArgs._SettingsObjects, true);
+	
+	this->ChildSlot
+	[
+		SNew(SVerticalBox)
+
+		+ SVerticalBox::Slot()
+		.Padding(2)
+		.MaxHeight(500.0f)
+		[
+			DetailsView->AsShared()
+		]
+
+		+ SVerticalBox::Slot()
+		.AutoHeight()
+		.HAlign(HAlign_Right)
+		.Padding(2)
+		[
+			SNew(SUniformGridPanel)
+			.SlotPadding(2)
+			+ SUniformGridPanel::Slot(0, 0)
+			[
+				SAssignNew(ConfirmButton, SButton)
+				.HAlign(HAlign_Center)
+				.Text(LOCTEXT("MaterialBakeOptionWindow_Import", "Confirm"))
+				.OnClicked(this, &SMaterialOptions::OnConfirm)
+			]
+			+ SUniformGridPanel::Slot(1, 0)
+			[
+				SNew(SButton)
+				.HAlign(HAlign_Center)
+				.Text(LOCTEXT("MaterialBakeOptionWindow_Cancel", "Cancel"))
+				.ToolTipText(LOCTEXT("MaterialBakeOptionWindow_Cancel_ToolTip", "Cancels baking out Material"))
+				.OnClicked(this, &SMaterialOptions::OnCancel)
+			]
+		]
+	];
+}
+
+FReply SMaterialOptions::OnConfirm()
+{
+	// Ensure the user has selected at least one LOD index
+	if (GetMutableDefault<UMaterialOptions>()->LODIndices.Num() == 0)
+	{
+		FText Title = LOCTEXT("MaterialBake_SelectLODErrorTitle", "Invalid options");
+		FMessageDialog::Open(EAppMsgType::Ok, LOCTEXT("MaterialBake_SelectLODError", "Ensure that atleast one LOD index is selected."), &Title);
+	}
+	else
+	{
+		bUserCancelled = false;
+		if (WidgetWindow.IsValid())
+		{
+			WidgetWindow.Pin()->RequestDestroyWindow();
+		}
+	}	
+	
+	return FReply::Handled();
+}
+
+FReply SMaterialOptions::OnCancel()
+{
+	if (WidgetWindow.IsValid())
+	{
+		WidgetWindow.Pin()->RequestDestroyWindow();
+	}
+	return FReply::Handled();
+}
+
+FReply SMaterialOptions::OnKeyDown(const FGeometry& MyGeometry, const FKeyEvent& InKeyEvent)
+{
+	if (InKeyEvent.GetKey() == EKeys::Escape)
+	{
+		return OnCancel();
+	}
+
+	return FReply::Unhandled();
+}
+
+bool SMaterialOptions::WasUserCancelled()
+{
+	return bUserCancelled;
+}
+
+#undef LOCTEXT_NAMESPACE //"SMaterialOptions"
+
diff --git a/Engine/Plugins/Enterprise/GLTFExporter/MigratedFromGit/Source/GLTFMaterialBaking/Private/MaterialOptionsWindow.h b/Engine/Plugins/Enterprise/GLTFExporter/MigratedFromGit/Source/GLTFMaterialBaking/Private/MaterialOptionsWindow.h
new file mode 100644
--- /dev/null
+++ b/Engine/Plugins/Enterprise/GLTFExporter/MigratedFromGit/Source/GLTFMaterialBaking/Private/MaterialOptionsWindow.h
@@ -0,0 +1,49 @@
+// Copyright Epic Games, Inc. All Rights Reserved.
+
+#pragma once
+
+#include "Widgets/SCompoundWidget.h"
+#include "Widgets/DeclarativeSyntaxSupport.h"
+#include "Widgets/SWindow.h"
+#include "IDetailRootObjectCustomization.h"
+
+struct FMaterialData;
+class SButton;
+
+/** Options window used to populate provided settings objects */
+class SMaterialOptions : public SCompoundWidget
+{
+public:
+	SLATE_BEGIN_ARGS(SMaterialOptions)
+		: _WidgetWindow(), _NumLODs(1)
+	{}
+		SLATE_ARGUMENT(TSharedPtr<SWindow>, WidgetWindow)
+		SLATE_ARGUMENT(int32, NumLODs)
+		SLATE_ARGUMENT(TArray<TWeakObjectPtr<UObject>>, SettingsObjects)
+	SLATE_END_ARGS()
+
+public:
+	SMaterialOptions();
+	void Construct(const FArguments& InArgs);
+	
+	/** Begin SCompoundWidget overrides */
+	virtual bool SupportsKeyboardFocus() const override { return true; }
+	virtual FReply OnKeyDown(const FGeometry& MyGeometry, const FKeyEvent& InKeyEvent) override;
+	/** End SCompoundWidget overrides */
+
+	/** Callbacks used for Confirm and Cancel buttons */
+	FReply OnConfirm();
+	FReply OnCancel();
+
+	/** Returns whether or not the user cancelled the operation */
+	bool WasUserCancelled();
+private:
+	/** Owning window this widget is part of */
+	TWeakPtr< SWindow > WidgetWindow;
+	/** Whether or not the cancel button was clicked by the user */
+	bool bUserCancelled;
+	/** Detailsview used to display SettingsObjects, and allowing user to change options */
+	TSharedPtr<class IDetailsView> DetailsView;
+	/** Shared ptr to Confirm button */
+	TSharedPtr<SButton> ConfirmButton;
+};
\ No newline at end of file
diff --git a/Engine/Plugins/Enterprise/GLTFExporter/MigratedFromGit/Source/GLTFMaterialBaking/Private/MaterialRenderItem.cpp b/Engine/Plugins/Enterprise/GLTFExporter/MigratedFromGit/Source/GLTFMaterialBaking/Private/MaterialRenderItem.cpp
new file mode 100644
--- /dev/null
+++ b/Engine/Plugins/Enterprise/GLTFExporter/MigratedFromGit/Source/GLTFMaterialBaking/Private/MaterialRenderItem.cpp
@@ -0,0 +1,287 @@
+// Copyright Epic Games, Inc. All Rights Reserved.
+
+#include "MaterialRenderItem.h"
+#include "MaterialBakingStructures.h"
+
+#include "EngineModule.h"
+#include "StaticMeshAttributes.h"
+#include "DynamicMeshBuilder.h"
+#include "MeshPassProcessor.h"
+
+#define SHOW_WIREFRAME_MESH 0
+
+FMeshMaterialRenderItem::FMeshMaterialRenderItem(
+	const FMaterialData* InMaterialSettings,
+	const FMeshData* InMeshSettings,
+	EMaterialProperty InMaterialProperty,
+	FDynamicMeshBufferAllocator* InDynamicMeshBufferAllocator)
+	: MeshSettings(InMeshSettings)
+	, MaterialSettings(InMaterialSettings)
+	, MaterialProperty(InMaterialProperty)
+	, MaterialRenderProxy(nullptr)
+	, ViewFamily(nullptr)
+	, bMeshElementDirty(true)
+	, DynamicMeshBufferAllocator(InDynamicMeshBufferAllocator)
+{
+	GenerateRenderData();
+	LCI = new FMeshRenderInfo(InMeshSettings->LightMap, nullptr, nullptr, InMeshSettings->LightmapResourceCluster);
+}
+
+bool FMeshMaterialRenderItem::Render_RenderThread(FRHICommandListImmediate& RHICmdList, FMeshPassProcessorRenderState& DrawRenderState, const FCanvas* Canvas)
+{
+	checkSlow(ViewFamily && MaterialSettings && MeshSettings && MaterialRenderProxy);
+	// current render target set for the canvas
+	const FRenderTarget* CanvasRenderTarget = Canvas->GetRenderTarget();
+	const FIntRect ViewRect(FIntPoint(0, 0), CanvasRenderTarget->GetSizeXY());
+
+	// make a temporary view
+	FSceneViewInitOptions ViewInitOptions;
+	ViewInitOptions.ViewFamily = ViewFamily;
+	ViewInitOptions.SetViewRectangle(ViewRect);
+	ViewInitOptions.ViewOrigin = FVector::ZeroVector;
+	ViewInitOptions.ViewRotationMatrix = FMatrix::Identity;
+	ViewInitOptions.ProjectionMatrix = Canvas->GetTransformStack().Top().GetMatrix();
+	ViewInitOptions.BackgroundColor = FLinearColor::Black;
+	ViewInitOptions.OverlayColor = FLinearColor::White;
+
+	FSceneView View(ViewInitOptions);
+	View.FinalPostProcessSettings.bOverride_IndirectLightingIntensity = 1;
+	View.FinalPostProcessSettings.IndirectLightingIntensity = 0.0f;
+
+	const bool bNeedsToSwitchVerticalAxis = RHINeedsToSwitchVerticalAxis(Canvas->GetShaderPlatform()) && !Canvas->GetAllowSwitchVerticalAxis();
+	check(bNeedsToSwitchVerticalAxis == false);
+
+	if (Vertices.Num() && Indices.Num())
+	{
+		FMeshPassProcessorRenderState LocalDrawRenderState(View);
+
+		// disable depth test & writes
+		LocalDrawRenderState.SetBlendState(TStaticBlendState<CW_RGBA>::GetRHI());
+		LocalDrawRenderState.SetDepthStencilState(TStaticDepthStencilState<false, CF_Always>::GetRHI());
+
+		QueueMaterial(RHICmdList, LocalDrawRenderState, &View);
+	}
+
+	return true;
+}
+
+bool FMeshMaterialRenderItem::Render_GameThread(const FCanvas* Canvas, FRenderThreadScope& RenderScope)
+{
+	RenderScope.EnqueueRenderCommand(
+		[this, Canvas](FRHICommandListImmediate& RHICmdList)
+		{
+			// Render_RenderThread uses its own render state
+			FMeshPassProcessorRenderState DummyRenderState;
+			Render_RenderThread(RHICmdList, DummyRenderState, Canvas);
+		}
+	);
+
+	return true;
+}
+
+void FMeshMaterialRenderItem::GenerateRenderData()
+{
+	TRACE_CPUPROFILER_EVENT_SCOPE(FMeshMaterialRenderItem::GenerateRenderData)
+
+	// Reset array without resizing
+	Vertices.SetNum(0, false);
+	Indices.SetNum(0, false);
+	if (MeshSettings->RawMeshDescription)
+	{
+		// Use supplied FMeshDescription data to populate render data
+		PopulateWithMeshData();
+	}
+	else
+	{
+		// Use simple rectangle
+		PopulateWithQuadData();
+	}
+
+	bMeshElementDirty = true;
+}
+
+FMeshMaterialRenderItem::~FMeshMaterialRenderItem()
+{
+	// Send the release of the buffers to the render thread
+	ENQUEUE_RENDER_COMMAND(ReleaseResources)(
+		[ToRelease = MoveTemp(MeshBuilderResources)](FRHICommandListImmediate& RHICmdList) {}
+	);
+}
+
+void FMeshMaterialRenderItem::QueueMaterial(FRHICommandListImmediate& RHICmdList, FMeshPassProcessorRenderState& DrawRenderState, const FSceneView* View)
+{
+	TRACE_CPUPROFILER_EVENT_SCOPE(FMeshMaterialRenderItem::QueueMaterial)
+
+	if (bMeshElementDirty)
+	{
+		MeshBuilderResources.Clear();
+		FDynamicMeshBuilder DynamicMeshBuilder(View->GetFeatureLevel(), MAX_STATIC_TEXCOORDS, MeshSettings->LightMapIndex, false, DynamicMeshBufferAllocator);
+		{
+			TRACE_CPUPROFILER_EVENT_SCOPE(CopyData);
+			DynamicMeshBuilder.AddVertices(Vertices);
+			DynamicMeshBuilder.AddTriangles(Indices);
+		}
+
+		DynamicMeshBuilder.GetMeshElement(FMatrix::Identity, MaterialRenderProxy, SDPG_Foreground, true, false, 0, MeshBuilderResources, MeshElement);
+
+		check(MeshBuilderResources.IsValidForRendering());
+		bMeshElementDirty = false;
+	}
+
+	MeshElement.MaterialRenderProxy = MaterialRenderProxy;
+
+	LCI->CreatePrecomputedLightingUniformBuffer_RenderingThread(View->GetFeatureLevel());
+	MeshElement.LCI = LCI;
+
+#if SHOW_WIREFRAME_MESH
+	MeshElement.bWireframe = true;
+#endif
+
+	const int32 NumTris = FMath::TruncToInt(Indices.Num() / 3);
+	if (NumTris == 0)
+	{
+		// there's nothing to do here
+		return;
+	}
+
+	// Bake the material out to a tile
+	GetRendererModule().DrawTileMesh(RHICmdList, DrawRenderState, *View, MeshElement, false /*bIsHitTesting*/, FHitProxyId());
+}
+
+void FMeshMaterialRenderItem::PopulateWithQuadData()
+{
+	Vertices.Empty(4);
+	Indices.Empty(6);
+
+	const float U = MeshSettings->TextureCoordinateBox.Min.X;
+	const float V = MeshSettings->TextureCoordinateBox.Min.Y;
+	const float SizeU = MeshSettings->TextureCoordinateBox.Max.X - MeshSettings->TextureCoordinateBox.Min.X;
+	const float SizeV = MeshSettings->TextureCoordinateBox.Max.Y - MeshSettings->TextureCoordinateBox.Min.Y;
+	const FIntPoint& PropertySize = MaterialSettings->PropertySizes[MaterialProperty];
+	const float ScaleX = PropertySize.X;
+	const float ScaleY = PropertySize.Y;
+
+	// add vertices
+	for (int32 VertIndex = 0; VertIndex < 4; VertIndex++)
+	{
+		FDynamicMeshVertex* Vert = new(Vertices)FDynamicMeshVertex();
+		const int32 X = VertIndex & 1;
+		const int32 Y = (VertIndex >> 1) & 1;
+		Vert->Position.Set(ScaleX * X, ScaleY * Y, 0);
+		Vert->SetTangents(FVector(1, 0, 0), FVector(0, 1, 0), FVector(0, 0, 1));
+		FMemory::Memzero(&Vert->TextureCoordinate, sizeof(Vert->TextureCoordinate));
+		for (int32 TexcoordIndex = 0; TexcoordIndex < MAX_STATIC_TEXCOORDS; TexcoordIndex++)
+		{
+			Vert->TextureCoordinate[TexcoordIndex].Set(U + SizeU * X, V + SizeV * Y);
+		}		
+		Vert->Color = FColor::White;
+	}
+
+	// add indices
+	static const uint32 TriangleIndices[6] = { 0, 2, 1, 2, 3, 1 };
+	Indices.Append(TriangleIndices, 6);
+}
+
+void FMeshMaterialRenderItem::PopulateWithMeshData()
+{
+	const FMeshDescription* RawMesh = MeshSettings->RawMeshDescription;
+
+	TVertexAttributesConstRef<FVector> VertexPositions = RawMesh->VertexAttributes().GetAttributesRef<FVector>(MeshAttribute::Vertex::Position);
+	TVertexInstanceAttributesConstRef<FVector> VertexInstanceNormals = RawMesh->VertexInstanceAttributes().GetAttributesRef<FVector>(MeshAttribute::VertexInstance::Normal);
+	TVertexInstanceAttributesConstRef<FVector> VertexInstanceTangents = RawMesh->VertexInstanceAttributes().GetAttributesRef<FVector>(MeshAttribute::VertexInstance::Tangent);
+	TVertexInstanceAttributesConstRef<float> VertexInstanceBinormalSigns = RawMesh->VertexInstanceAttributes().GetAttributesRef<float>(MeshAttribute::VertexInstance::BinormalSign);
+	TVertexInstanceAttributesConstRef<FVector2D> VertexInstanceUVs = RawMesh->VertexInstanceAttributes().GetAttributesRef<FVector2D>(MeshAttribute::VertexInstance::TextureCoordinate);
+	TVertexInstanceAttributesConstRef<FVector4> VertexInstanceColors = RawMesh->VertexInstanceAttributes().GetAttributesRef<FVector4>(MeshAttribute::VertexInstance::Color);
+	const int32 NumVerts = RawMesh->Vertices().Num();
+
+	// reserve renderer data
+	Vertices.Empty(NumVerts);
+	Indices.Empty(NumVerts >> 1);
+
+	const FIntPoint& PropertySize = MaterialSettings->PropertySizes[MaterialProperty];
+	const float ScaleX = PropertySize.X;
+	const float ScaleY = PropertySize.Y;
+
+	const static int32 VertexPositionStoredUVChannel = 6;
+	// count number of texture coordinates for this mesh
+	const int32 NumTexcoords = [&]()
+	{
+		return FMath::Min(VertexInstanceUVs.GetNumIndices(), VertexPositionStoredUVChannel);
+	}();		
+
+	// check if we should use NewUVs or original UV set
+	const bool bUseNewUVs = MeshSettings->CustomTextureCoordinates.Num() > 0;
+	if (bUseNewUVs)
+	{
+		check(MeshSettings->CustomTextureCoordinates.Num() == VertexInstanceUVs.GetNumElements() && VertexInstanceUVs.GetNumIndices() > MeshSettings->TextureCoordinateIndex);
+	}
+
+	// add vertices
+	int32 VertIndex = 0;
+	int32 FaceIndex = 0;
+	for(const FPolygonID PolygonID : RawMesh->Polygons().GetElementIDs())
+	{
+		const FPolygonGroupID PolygonGroupID = RawMesh->GetPolygonPolygonGroup(PolygonID);
+		const TArray<FTriangleID>& TriangleIDs = RawMesh->GetPolygonTriangleIDs(PolygonID);
+		for (const FTriangleID TriangleID : TriangleIDs)
+		{
+			if (MeshSettings->MaterialIndices.Contains(PolygonGroupID.GetValue()))
+			{
+				const int32 NUM_VERTICES = 3;
+				for (int32 Corner = 0; Corner < NUM_VERTICES; Corner++)
+				{
+					// Swap vertices order if mesh is mirrored
+					const int32 CornerIdx = !MeshSettings->bMirrored ? Corner : NUM_VERTICES - Corner - 1;
+
+					const int32 SrcVertIndex = FaceIndex * NUM_VERTICES + CornerIdx;
+					const FVertexInstanceID SrcVertexInstanceID = RawMesh->GetTriangleVertexInstance(TriangleID, Corner);
+					const FVertexID SrcVertexID = RawMesh->GetVertexInstanceVertex(SrcVertexInstanceID);
+
+					// add vertex
+					FDynamicMeshVertex* Vert = new(Vertices)FDynamicMeshVertex();
+					if (!bUseNewUVs)
+					{
+						// compute vertex position from original UV
+						const FVector2D& UV = VertexInstanceUVs.Get(SrcVertexInstanceID, MeshSettings->TextureCoordinateIndex);
+						Vert->Position.Set(UV.X * ScaleX, UV.Y * ScaleY, 0);
+					}
+					else
+					{
+						const FVector2D& UV = MeshSettings->CustomTextureCoordinates[SrcVertIndex];
+						Vert->Position.Set(UV.X * ScaleX, UV.Y * ScaleY, 0);
+					}
+					FVector TangentX = VertexInstanceTangents[SrcVertexInstanceID];
+					FVector TangentZ = VertexInstanceNormals[SrcVertexInstanceID];
+					FVector TangentY = FVector::CrossProduct(TangentZ, TangentX).GetSafeNormal() * VertexInstanceBinormalSigns[SrcVertexInstanceID];
+					Vert->SetTangents(TangentX, TangentY, TangentZ);
+					for (int32 TexcoordIndex = 0; TexcoordIndex < NumTexcoords; TexcoordIndex++)
+					{
+						Vert->TextureCoordinate[TexcoordIndex] = VertexInstanceUVs.Get(SrcVertexInstanceID, TexcoordIndex);
+					}
+				
+					if (NumTexcoords < VertexPositionStoredUVChannel)
+					{
+						for (int32 TexcoordIndex = NumTexcoords; TexcoordIndex < VertexPositionStoredUVChannel; TexcoordIndex++)
+						{
+							Vert->TextureCoordinate[TexcoordIndex] = Vert->TextureCoordinate[FMath::Max(NumTexcoords - 1, 0)];
+						}
+					}
+					// Store original vertex positions in texture coordinate data
+					Vert->TextureCoordinate[6].X = VertexPositions[SrcVertexID].X;
+					Vert->TextureCoordinate[6].Y = VertexPositions[SrcVertexID].Y;
+					Vert->TextureCoordinate[7].X = VertexPositions[SrcVertexID].Z;
+
+					Vert->Color = FLinearColor(VertexInstanceColors[SrcVertexInstanceID]).ToFColor(true);
+					// add index
+					Indices.Add(VertIndex);
+					VertIndex++;
+				}
+				// add the same triangle with opposite vertex order
+				Indices.Add(VertIndex - 3);
+				Indices.Add(VertIndex - 1);
+				Indices.Add(VertIndex - 2);
+			}
+			FaceIndex++;
+		}
+	}
+}
\ No newline at end of file
diff --git a/Engine/Plugins/Enterprise/GLTFExporter/MigratedFromGit/Source/GLTFMaterialBaking/Private/MaterialRenderItem.h b/Engine/Plugins/Enterprise/GLTFExporter/MigratedFromGit/Source/GLTFMaterialBaking/Private/MaterialRenderItem.h
new file mode 100644
--- /dev/null
+++ b/Engine/Plugins/Enterprise/GLTFExporter/MigratedFromGit/Source/GLTFMaterialBaking/Private/MaterialRenderItem.h
@@ -0,0 +1,89 @@
+// Copyright Epic Games, Inc. All Rights Reserved.
+
+#pragma once
+
+#include "CanvasTypes.h"
+#include "MaterialRenderItemData.h"
+#include "DynamicMeshBuilder.h"
+
+class FSceneViewFamily;
+class FMaterialRenderProxy;
+class FSceneView;
+class FRHICommandListImmediate;
+struct FMaterialData;
+struct FMeshData;
+struct FMeshPassProcessorRenderState;
+
+// This will hold onto the resources until not needed anymore.
+// Move constructor makes it easier to send the destruction to another thread (render thread).
+class FMeshBuilderResources : public FMeshBuilderOneFrameResources
+{
+public:
+	FMeshBuilderResources() = default;
+	FMeshBuilderResources(FMeshBuilderResources&& Other)
+	{
+		if (this != &Other)
+		{
+			VertexBuffer = Other.VertexBuffer;
+			IndexBuffer = Other.IndexBuffer;
+			VertexFactory = Other.VertexFactory;
+			PrimitiveUniformBuffer = Other.PrimitiveUniformBuffer;
+
+			Other.VertexBuffer = nullptr;
+			Other.IndexBuffer = nullptr;
+			Other.VertexFactory = nullptr;
+			Other.PrimitiveUniformBuffer = nullptr;
+		}
+	}
+
+	void Clear()
+	{
+		FMeshBuilderOneFrameResources::~FMeshBuilderOneFrameResources();
+		VertexBuffer = nullptr;
+		IndexBuffer = nullptr;
+		VertexFactory = nullptr;
+		PrimitiveUniformBuffer = nullptr;
+	}
+};
+
+class FMeshMaterialRenderItem : public FCanvasBaseRenderItem
+{
+public:
+	FMeshMaterialRenderItem(const FMaterialData* InMaterialSettings, const FMeshData* InMeshSettings, EMaterialProperty InMaterialProperty, FDynamicMeshBufferAllocator* InDynamicMeshBufferAllocator = nullptr);
+	virtual ~FMeshMaterialRenderItem();
+
+	/** Begin FCanvasBaseRenderItem overrides */
+	virtual bool Render_RenderThread(FRHICommandListImmediate& RHICmdList, FMeshPassProcessorRenderState& DrawRenderState, const FCanvas* Canvas) final;
+	virtual bool Render_GameThread(const FCanvas* Canvas, FRenderThreadScope& RenderScope) final;
+	/** End FCanvasBaseRenderItem overrides */
+
+	/** Populate vertices and indices according to available mesh data and otherwise uses simple quad */
+	void GenerateRenderData();
+protected:
+	/** Enqueues the current material to be rendered */
+	void QueueMaterial(FRHICommandListImmediate& RHICmdList, FMeshPassProcessorRenderState& DrawRenderState, const FSceneView* View);
+	
+	/** Helper functions to populate render data using either mesh data or a simple quad */
+	void PopulateWithQuadData();
+	void PopulateWithMeshData();
+public:
+	/** Mesh and material settings to use while baking out the material */
+	const FMeshData* MeshSettings;
+	const FMaterialData* MaterialSettings;
+	/** Material property to bake out */
+	EMaterialProperty MaterialProperty;
+	/** Material render proxy (material/shader) to use while baking */
+	FMaterialRenderProxy* MaterialRenderProxy;
+	/** Vertex and index data representing the mesh or a quad */
+	TArray<FDynamicMeshVertex> Vertices;
+	TArray<uint32> Indices;
+	/** Light cache interface object to simulate lightmap behavior in case the material used prebaked ambient occlusion */
+	FLightCacheInterface* LCI;
+	/** View family to use while baking */
+	FSceneViewFamily* ViewFamily;
+private:
+	FMeshBatch MeshElement;
+	bool bMeshElementDirty;
+	FMeshBuilderResources MeshBuilderResources;
+	FDynamicMeshBufferAllocator* DynamicMeshBufferAllocator;
+};
diff --git a/Engine/Plugins/Enterprise/GLTFExporter/MigratedFromGit/Source/GLTFMaterialBaking/Private/MaterialRenderItemData.cpp b/Engine/Plugins/Enterprise/GLTFExporter/MigratedFromGit/Source/GLTFMaterialBaking/Private/MaterialRenderItemData.cpp
new file mode 100644
--- /dev/null
+++ b/Engine/Plugins/Enterprise/GLTFExporter/MigratedFromGit/Source/GLTFMaterialBaking/Private/MaterialRenderItemData.cpp
@@ -0,0 +1,3 @@
+// Copyright Epic Games, Inc. All Rights Reserved.
+
+#include "MaterialRenderItemData.h"
diff --git a/Engine/Plugins/Enterprise/GLTFExporter/MigratedFromGit/Source/GLTFMaterialBaking/Private/MaterialRenderItemData.h b/Engine/Plugins/Enterprise/GLTFExporter/MigratedFromGit/Source/GLTFMaterialBaking/Private/MaterialRenderItemData.h
new file mode 100644
--- /dev/null
+++ b/Engine/Plugins/Enterprise/GLTFExporter/MigratedFromGit/Source/GLTFMaterialBaking/Private/MaterialRenderItemData.h
@@ -0,0 +1,30 @@
+// Copyright Epic Games, Inc. All Rights Reserved.
+
+#pragma once
+
+#include "LightMapHelpers.h"
+#include "Components.h"
+#include "LocalVertexFactory.h"
+#include "SceneManagement.h"
+
+/** Simple implementation for light cache to simulated lightmap behaviour (used for accessing prebaked ambient occlusion values) */
+class FMeshRenderInfo : public FLightCacheInterface
+{
+public:
+	FMeshRenderInfo(const FLightMap* InLightMap, const FShadowMap* InShadowMap, FUniformBufferRHIRef Buffer, const FLightmapResourceCluster* InLightmapResourceCluster)
+		: FLightCacheInterface()
+	{
+		SetLightMap(InLightMap);
+		SetShadowMap(InShadowMap);
+		SetPrecomputedLightingBuffer(Buffer);
+		if (InLightmapResourceCluster)
+		{
+			SetResourceCluster(InLightmapResourceCluster);
+		}
+	}
+
+	virtual FLightInteraction GetInteraction(const class FLightSceneProxy* LightSceneProxy) const override
+	{
+		return LIT_CachedLightMap;
+	}
+};
\ No newline at end of file
diff --git a/Engine/Plugins/Enterprise/GLTFExporter/MigratedFromGit/Source/GLTFMaterialBaking/Public/IMaterialBakingModule.h b/Engine/Plugins/Enterprise/GLTFExporter/MigratedFromGit/Source/GLTFMaterialBaking/Public/IMaterialBakingModule.h
new file mode 100644
--- /dev/null
+++ b/Engine/Plugins/Enterprise/GLTFExporter/MigratedFromGit/Source/GLTFMaterialBaking/Public/IMaterialBakingModule.h
@@ -0,0 +1,29 @@
+// Copyright Epic Games, Inc. All Rights Reserved.
+
+#pragma once
+
+#include "CoreMinimal.h"
+#include "Modules/ModuleInterface.h"
+#include "PixelFormat.h"
+#include "SceneTypes.h"
+
+class UTextureRenderTarget2D;
+class UMaterialOptions;
+class UMaterialInterface;
+
+class FExportMaterialProxy;
+class FTextureRenderTargetResource;
+
+struct FMaterialData; 
+struct FMeshData;
+struct FBakeOutput;
+
+class MATERIALBAKING_API IMaterialBakingModule : public IModuleInterface
+{
+public:
+	/** Bakes out material properties according to MaterialSettings using MeshSettings and stores the output in Output */
+	virtual void BakeMaterials(const TArray<FMaterialData*>& MaterialSettings, const TArray<FMeshData*>& MeshSettings, TArray<FBakeOutput>& Output) = 0;
+
+	/** Promps a slate window to allow the user to populate specific material baking settings used while baking out materials */
+	virtual  bool SetupMaterialBakeSettings(TArray<TWeakObjectPtr<UObject>>& OptionObjects, int32 NumLODs) = 0;
+};
diff --git a/Engine/Plugins/Enterprise/GLTFExporter/MigratedFromGit/Source/GLTFMaterialBaking/Public/MaterialBakingHelpers.h b/Engine/Plugins/Enterprise/GLTFExporter/MigratedFromGit/Source/GLTFMaterialBaking/Public/MaterialBakingHelpers.h
new file mode 100644
--- /dev/null
+++ b/Engine/Plugins/Enterprise/GLTFExporter/MigratedFromGit/Source/GLTFMaterialBaking/Public/MaterialBakingHelpers.h
@@ -0,0 +1,30 @@
+// Copyright Epic Games, Inc. All Rights Reserved.
+
+#pragma once
+
+#include "Containers/Array.h"
+
+struct FColor;
+
+class MATERIALBAKING_API FMaterialBakingHelpers
+{
+public:
+	/** 
+	 * Applies a box blur to magenta pixels found in given texture represented by InOutPixels using non-magenta pixels, this creates a smear across the magenta/filled pixels.
+	 * @param  InOutPixels     The image to apply the smear to
+	 * @param  ImageWidth      The width of the image in pixels
+	 * @param  ImageHeight     The height of the image in pixels
+	 * @param  MaxIterations   The max distance in pixels to smear the edges of the texture. When set to -1 this value is the max of the width and height of the buffer
+	 */
+	static void PerformUVBorderSmear(TArray<FColor>& InOutPixels, int32 ImageWidth, int32 ImageHeight, int32 MaxIterations = -1);
+
+	/** 
+	 * Applies a box blur to magenta pixels found in given texture represented by InOutPixels using non-magenta pixels, this creates a smear across the magenta/filled pixels.
+	 * @param  InOutPixels      The image to apply the smear to
+	 * @param  InOutImageWidth  The width of the image in pixels
+	 * @param  InOutImageHeight The height of the image in pixels
+	 *
+	 * @note   This version is allowed to shrink the array to a single pixel if the entire output is monochrome. Params will be adjusted accordingly when this occurs.
+	 */
+	static void PerformUVBorderSmearAndShrink(TArray<FColor>& InOutPixels, int32& InOutImageWidth, int32& InOutImageHeight);
+};
\ No newline at end of file
diff --git a/Engine/Plugins/Enterprise/GLTFExporter/MigratedFromGit/Source/GLTFMaterialBaking/Public/MaterialBakingStructures.h b/Engine/Plugins/Enterprise/GLTFExporter/MigratedFromGit/Source/GLTFMaterialBaking/Public/MaterialBakingStructures.h
new file mode 100644
--- /dev/null
+++ b/Engine/Plugins/Enterprise/GLTFExporter/MigratedFromGit/Source/GLTFMaterialBaking/Public/MaterialBakingStructures.h
@@ -0,0 +1,83 @@
+// Copyright Epic Games, Inc. All Rights Reserved.
+
+#pragma once
+
+#include "Math/IntPoint.h"
+#include "SceneTypes.h"
+#include "LightMap.h"
+
+class UMaterialInterface;
+struct FMeshDescription;
+
+/** Structure containing information about the material which is being baked out */
+struct FMaterialData
+{
+	FMaterialData()
+		: Material(nullptr)
+		, bPerformBorderSmear(true)
+	{}
+
+	/** Material to bake out */
+	UMaterialInterface* Material;
+	/** Properties and the texture size at which they should be baked out */
+	TMap<EMaterialProperty, FIntPoint> PropertySizes;
+	/** Whether to smear borders after baking */
+	bool bPerformBorderSmear;
+};
+
+struct FMeshData
+{
+	FMeshData()
+		: RawMeshDescription(nullptr), Mesh(nullptr), bMirrored(false), VertexColorHash(0), TextureCoordinateIndex(0), LightMapIndex(0), LightMap(nullptr), LightmapResourceCluster(nullptr)
+	{}
+
+	/** Ptr to raw mesh data to use for baking out the material data, if nullptr a standard quad is used */
+	FMeshDescription* RawMeshDescription;
+
+	/** Ptr to original static mesh this mesh data came from */
+	UStaticMesh* Mesh;
+
+	/** Transform determinant used to detect mirroring */
+	bool bMirrored;
+
+	/** A hash of the vertex color buffer for the rawmesh */
+	uint32 VertexColorHash;
+
+	/** Material indices to test the Raw Mesh data against, ensuring we only bake out triangles which use the currently baked out material */
+	TArray<int32> MaterialIndices;
+
+	/** Set of custom texture coordinates which ensure that the material is baked out with unique/non-overlapping positions */
+	TArray<FVector2D> CustomTextureCoordinates;
+
+	/** Box which's space contains the UV coordinates used to bake out the material */
+	FBox2D TextureCoordinateBox;
+	
+	/** Specific texture coordinate index to use as texture coordinates to bake out the material (is overruled if CustomTextureCoordinates contains any data) */
+	int32 TextureCoordinateIndex;
+
+	/** Light map index used to retrieve the light-map UVs from RawMesh */
+	int32 LightMapIndex;
+
+	/** Reference to the lightmap texture part of the level in the currently being baked out mesh instance data is resident */
+	FLightMapRef LightMap;
+
+	/** Pointer to the LightmapResourceCluster to be passed on the the LightCacheInterface when baking */
+	const FLightmapResourceCluster* LightmapResourceCluster;
+};
+
+/** Structure containing data being processed while baking out materials*/
+struct FBakeOutput
+{
+	FBakeOutput()
+		: EmissiveScale(1.0f)
+	{}
+
+	/** Contains the resulting texture data for baking out a material's property */
+	TMap<EMaterialProperty, TArray<FColor>> PropertyData;
+
+	/** Contains the resulting texture size for baking out a material's property */
+	TMap<EMaterialProperty, FIntPoint> PropertySizes;
+
+	/** Scale used to allow having wide ranges of emissive values in the source materials, the final proxy material will use this value to scale the emissive texture's pixel values */
+	float EmissiveScale;
+};
diff --git a/Engine/Plugins/Enterprise/GLTFExporter/MigratedFromGit/Source/GLTFMaterialBaking/Public/MaterialOptions.h b/Engine/Plugins/Enterprise/GLTFExporter/MigratedFromGit/Source/GLTFMaterialBaking/Public/MaterialOptions.h
new file mode 100644
--- /dev/null
+++ b/Engine/Plugins/Enterprise/GLTFExporter/MigratedFromGit/Source/GLTFMaterialBaking/Public/MaterialOptions.h
@@ -0,0 +1,132 @@
+// Copyright Epic Games, Inc. All Rights Reserved.
+
+#pragma once
+
+#include "SceneTypes.h"
+#include "Engine/EngineTypes.h"
+
+#include "MaterialOptions.generated.h"
+
+/** Enum to define different types of baking materials */
+UENUM()
+enum class EMaterialBakeMethod : uint8
+{
+	IndividualMaterial UMETA(DisplayName = "Bake out Materials Individually"),
+	AtlasMaterial UMETA(DisplayName = "Combine Materials into Atlassed Material"),
+	BinnedMaterial UMETA(DisplayName = "Combine Materials into Binned Material")
+};
+
+/** Structure to represent a single property the user wants to bake out for a given set of materials */
+USTRUCT(Blueprintable)
+struct FPropertyEntry
+{
+	GENERATED_BODY()
+
+	FPropertyEntry()
+		: Property(MP_EmissiveColor)
+		, bUseCustomSize(false)
+		, CustomSize(0, 0)
+		, bUseConstantValue(false)
+		, ConstantValue(0.0f)
+	{}
+
+	FPropertyEntry(EMaterialProperty InProperty)
+		: Property(InProperty)
+		, bUseCustomSize(false)
+		, CustomSize(0, 0)
+		, bUseConstantValue(false)
+		, ConstantValue(0.0f)
+	{}
+
+	/** Property which should be baked out */
+	UPROPERTY(EditAnywhere, BlueprintReadWrite, Category = Property, meta = (ExposeOnSpawn))
+	TEnumAsByte<EMaterialProperty> Property;
+
+	/** Whether or not to use the value of custom size for the output texture */
+	UPROPERTY(EditAnywhere, Category = Property, BlueprintReadWrite, meta = (InlineEditConditionToggle, ExposeOnSpawn))
+	bool bUseCustomSize;
+
+	/** Defines the size of the output textures for the baked out material properties */
+	UPROPERTY(EditAnywhere, Category = Property, BlueprintReadWrite, meta = (EditCondition = bUseCustomSize, ExposeOnSpawn, ClampMin = "1", UIMin = "1"))
+	FIntPoint CustomSize;
+
+	/** Wheter or not to use Constant Value as the final 'baked out' value for the this property */
+	UPROPERTY(EditAnywhere, Category = Property, BlueprintReadWrite, meta = (InlineEditConditionToggle, ExposeOnSpawn))
+	bool bUseConstantValue;
+
+	/** Defines the value representing this property in the final proxy material */
+	UPROPERTY(EditAnywhere, Category = Property, BlueprintReadWrite, meta = (EditCondition = bUseConstantValue, ExposeOnSpawn))
+	float ConstantValue;
+};
+
+/** Options object to define what and how a material should be baked out */
+UCLASS(config = Editor, Blueprintable)
+class MATERIALBAKING_API UMaterialOptions : public UObject 
+{
+	GENERATED_BODY()
+
+public:
+
+	UMaterialOptions()
+		: TextureSize(128, 128), bUseMeshData(false), bUseSpecificUVIndex(false), TextureCoordinateIndex(0)
+	{
+		Properties.Add(MP_BaseColor);
+		LODIndices.Add(0);
+	}
+
+	/** Properties which are supposed to be baked out for the material(s) */
+	UPROPERTY(EditAnywhere, config, BlueprintReadWrite, Category= MaterialBakeSettings, meta = (ExposeOnSpawn))
+	TArray<FPropertyEntry> Properties;
+
+	/** Size of the final texture(s) containing the baked out property data */
+	UPROPERTY(EditAnywhere, config, BlueprintReadWrite, Category = MaterialBakeSettings, meta = (ExposeOnSpawn, ClampMin = "1", UIMin = "1"))
+	FIntPoint TextureSize;
+
+	/** LOD indices for which the materials should be baked out */
+	UPROPERTY(BlueprintReadWrite, Category = MeshSettings, meta = (ExposeOnSpawn))
+	TArray<int32> LODIndices;
+
+	/** Determines whether to not allow usage of the source mesh data while baking out material properties */
+	UPROPERTY(EditAnywhere, config, BlueprintReadWrite, Category = MeshSettings, meta = (ExposeOnSpawn))
+	bool bUseMeshData;
+
+	/** Flag whether or not the value of TextureCoordinateIndex should be used while baking out material properties */
+	UPROPERTY(EditAnywhere, Category = MeshSettings, BlueprintReadWrite, meta = (InlineEditConditionToggle, EditCondition = bUseMeshData, ExposeOnSpawn))
+	bool bUseSpecificUVIndex;
+
+	/** Specific texture coordinate which should be used to while baking out material properties as the positions stream */
+	UPROPERTY(EditAnywhere, Category = MeshSettings, BlueprintReadWrite, meta = (EditCondition = bUseSpecificUVIndex, ExposeOnSpawn))
+	int32 TextureCoordinateIndex;
+};
+
+/** Asset bake options object */
+UCLASS(Config = Editor, Blueprintable)
+class MATERIALBAKING_API UAssetBakeOptions : public UObject
+{
+	GENERATED_BODY()
+public:
+
+	UAssetBakeOptions()
+	{
+	}
+};
+
+/** Material merge options object */
+UCLASS(Config = Editor, Blueprintable)
+class MATERIALBAKING_API UMaterialMergeOptions: public UObject
+{
+	GENERATED_BODY()
+public:
+
+	UMaterialMergeOptions() : Method(EMaterialBakeMethod::IndividualMaterial), BlendMode(EBlendMode::BLEND_Opaque)
+	{
+	}
+
+	/** Method used to bake out the materials, hidden for now */
+	UPROPERTY(/*EditAnywhere, BlueprintReadWrite, config, Category = MeshSettings, meta = (ExposeOnSpawn)*/)
+	EMaterialBakeMethod Method;
+
+	/** Blend mode for the final proxy material(s) */
+	UPROPERTY(EditAnywhere, BlueprintReadWrite, config, Category = OutputMaterialSettings, meta = (ExposeOnSpawn))
+	TEnumAsByte<EBlendMode> BlendMode;
+};
\ No newline at end of file
