[GLTFExporter] chore: rename build file MaterialBaking.Build.cs to GLTFMaterialBaking.Build.cs

commit ac5cdeac5f04395a1bf5282a85be4e17139ed058
Date: Wed, 4 Nov 2020 21:09:18 +0100

    chore: rename build file MaterialBaking.Build.cs to GLTFMaterialBaking.Build.cs
    (cherry picked from commit 5aa70423f849bef25bc69030bd8343897f1075d5​)

#skipci

diff --git a/Engine/Plugins/Enterprise/GLTFExporter/MigratedFromGit/Source/GLTFMaterialBaking/GLTFMaterialBaking.Build.cs b/Engine/Plugins/Enterprise/GLTFExporter/MigratedFromGit/Source/GLTFMaterialBaking/GLTFMaterialBaking.Build.cs
new file mode 100644
--- /dev/null
+++ b/Engine/Plugins/Enterprise/GLTFExporter/MigratedFromGit/Source/GLTFMaterialBaking/GLTFMaterialBaking.Build.cs
@@ -0,0 +1,29 @@
+// Copyright Epic Games, Inc. All Rights Reserved.
+
+using UnrealBuildTool;
+using System.IO;
+public class MaterialBaking : ModuleRules
+{
+	public MaterialBaking(ReadOnlyTargetRules Target) : base(Target)
+	{
+        PrivateDependencyModuleNames.AddRange(
+			new string [] {
+				"Core",
+				"CoreUObject",
+				"Engine",
+				"RenderCore",
+                "RHI",                
+                "UnrealEd",
+                "MainFrame",
+                "SlateCore",
+                "Slate",
+                "InputCore",
+                "PropertyEditor",
+                "EditorStyle",
+                "Renderer",
+                "MeshDescription",
+				"StaticMeshDescription"
+            }
+		);
+    }
+}
diff --git a/Engine/Plugins/Enterprise/GLTFExporter/MigratedFromGit/Source/GLTFMaterialBaking/MaterialBaking.Build.cs b/Engine/Plugins/Enterprise/GLTFExporter/MigratedFromGit/Source/GLTFMaterialBaking/MaterialBaking.Build.cs
deleted file mode 100644
