[GLTFExporter] refactor: rename MaterialBaking to GLTFMaterialBaking and add define MATERIALBAKING_API=GLTFMATERIALBAKING_API

commit 8baf05aa61b7755e1f5132dce39bef4859d1f8b4
Date: Wed, 4 Nov 2020 21:24:48 +0100

    refactor: rename MaterialBaking to GLTFMaterialBaking and add define MATERIALBAKING_API=GLTFMATERIALBAKING_API
    (cherry picked from commit 1f7cb77eb2e73ec6c7ccf618c2cdda79c2c5a85b​)

#skipci

diff --git a/Engine/Plugins/Enterprise/GLTFExporter/MigratedFromGit/Source/GLTFMaterialBaking/GLTFMaterialBaking.Build.cs b/Engine/Plugins/Enterprise/GLTFExporter/MigratedFromGit/Source/GLTFMaterialBaking/GLTFMaterialBaking.Build.cs
--- a/Engine/Plugins/Enterprise/GLTFExporter/MigratedFromGit/Source/GLTFMaterialBaking/GLTFMaterialBaking.Build.cs
+++ b/Engine/Plugins/Enterprise/GLTFExporter/MigratedFromGit/Source/GLTFMaterialBaking/GLTFMaterialBaking.Build.cs
@@ -2,9 +2,9 @@
 
 using UnrealBuildTool;
 using System.IO;
-public class MaterialBaking : ModuleRules
+public class GLTFMaterialBaking : ModuleRules
 {
-	public MaterialBaking(ReadOnlyTargetRules Target) : base(Target)
+	public GLTFMaterialBaking(ReadOnlyTargetRules Target) : base(Target)
 	{
         PrivateDependencyModuleNames.AddRange(
 			new string [] {
@@ -25,5 +25,8 @@ public class MaterialBaking : ModuleRules
 				"StaticMeshDescription"
             }
 		);
+
+		// NOTE: avoid renaming all instaces of MATERIALBAKING_API by redirecting it to GLTFMATERIALBAKING_API
+		PublicDefinitions.Add("MATERIALBAKING_API=GLTFMATERIALBAKING_API");
     }
 }
