[GLTFExporter] refactor: rename UMaterialOptions (to UGLTFMaterialOptions), UMaterialMergeOptions (to UGLTFMaterialMergeOptions), and UGAssetBakeOptions (to UGLTFAssetBakeOptions)

commit fa0cda526ea3f14e215a12d53d0e23411e5a821b
Date: Wed, 4 Nov 2020 21:24:58 +0100

    refactor: rename UMaterialOptions (to UGLTFMaterialOptions), UMaterialMergeOptions (to UGLTFMaterialMergeOptions), and UGAssetBakeOptions (to UGLTFAssetBakeOptions)
    (cherry picked from commit fa24e649d69972b4d2612578a265d6982ee92b1f​)

#skipci

diff --git a/Engine/Plugins/Enterprise/GLTFExporter/MigratedFromGit/Source/GLTFMaterialBaking/Private/MaterialBakingModule.cpp b/Engine/Plugins/Enterprise/GLTFExporter/MigratedFromGit/Source/GLTFMaterialBaking/Private/MaterialBakingModule.cpp
--- a/Engine/Plugins/Enterprise/GLTFExporter/MigratedFromGit/Source/GLTFMaterialBaking/Private/MaterialBakingModule.cpp
+++ b/Engine/Plugins/Enterprise/GLTFExporter/MigratedFromGit/Source/GLTFMaterialBaking/Private/MaterialBakingModule.cpp
@@ -24,7 +24,7 @@
 #include "Misc/FileHelper.h"
 #endif
 
-IMPLEMENT_MODULE(FMaterialBakingModule, MaterialBaking);
+IMPLEMENT_MODULE(FMaterialBakingModule, GLTFMaterialBaking);
 
 #define LOCTEXT_NAMESPACE "MaterialBakingModule"
 
diff --git a/Engine/Plugins/Enterprise/GLTFExporter/MigratedFromGit/Source/GLTFMaterialBaking/Private/MaterialBakingModule.h b/Engine/Plugins/Enterprise/GLTFExporter/MigratedFromGit/Source/GLTFMaterialBaking/Private/MaterialBakingModule.h
--- a/Engine/Plugins/Enterprise/GLTFExporter/MigratedFromGit/Source/GLTFMaterialBaking/Private/MaterialBakingModule.h
+++ b/Engine/Plugins/Enterprise/GLTFExporter/MigratedFromGit/Source/GLTFMaterialBaking/Private/MaterialBakingModule.h
@@ -9,7 +9,7 @@
 #include "IMaterialBakingModule.h"
 
 class UTextureRenderTarget2D;
-class UMaterialOptions;
+class UGLTFMaterialOptions;
 class UMaterialInterface;
 
 class FExportMaterialProxy;
diff --git a/Engine/Plugins/Enterprise/GLTFExporter/MigratedFromGit/Source/GLTFMaterialBaking/Private/MaterialOptionsCustomization.cpp b/Engine/Plugins/Enterprise/GLTFExporter/MigratedFromGit/Source/GLTFMaterialBaking/Private/MaterialOptionsCustomization.cpp
--- a/Engine/Plugins/Enterprise/GLTFExporter/MigratedFromGit/Source/GLTFMaterialBaking/Private/MaterialOptionsCustomization.cpp
+++ b/Engine/Plugins/Enterprise/GLTFExporter/MigratedFromGit/Source/GLTFMaterialBaking/Private/MaterialOptionsCustomization.cpp
@@ -39,8 +39,8 @@ void FPropertyEntryCustomization::CustomizeHeader(TSharedRef<IPropertyHandle> Pr
 	const TArray<TWeakObjectPtr<UObject>>& SelectedObjects = CustomizationUtils.GetPropertyUtilities()->GetSelectedObjects();
 
 	// Try and find material options instance in currently edited objects
-	const TWeakObjectPtr<UObject>* MaterialOptionsPtr = SelectedObjects.FindByPredicate([](TWeakObjectPtr<UObject> Object) { return Cast<UMaterialOptions>(Object.Get()); });
-	CurrentOptions = MaterialOptionsPtr ? Cast<UMaterialOptions>((MaterialOptionsPtr)->Get()) : nullptr;
+	const TWeakObjectPtr<UObject>* MaterialOptionsPtr = SelectedObjects.FindByPredicate([](TWeakObjectPtr<UObject> Object) { return Cast<UGLTFMaterialOptions>(Object.Get()); });
+	CurrentOptions = MaterialOptionsPtr ? Cast<UGLTFMaterialOptions>((MaterialOptionsPtr)->Get()) : nullptr;
 	
 	const int32 Index = PropertyHandle->GetIndexInArray();
 
@@ -127,15 +127,15 @@ void FMaterialOptionsCustomization::CustomizeDetails(IDetailLayoutBuilder& Detai
 	TArray<TWeakObjectPtr<UObject>> WeakObjects;
 	DetailBuilder.GetObjectsBeingCustomized(WeakObjects);
 	// Try and find material options instance in currently edited objects
-	UMaterialOptions* CurrentOptions = Cast<UMaterialOptions>((WeakObjects.FindByPredicate([](TWeakObjectPtr<UObject> Object) { return Cast<UMaterialOptions>(Object.Get()); }))->Get());
+	UGLTFMaterialOptions* CurrentOptions = Cast<UGLTFMaterialOptions>((WeakObjects.FindByPredicate([](TWeakObjectPtr<UObject> Object) { return Cast<UGLTFMaterialOptions>(Object.Get()); }))->Get());
 
-	TSharedRef<IPropertyHandle> TextureSizePropertyHandle = DetailBuilder.GetProperty(GET_MEMBER_NAME_CHECKED(UMaterialOptions, TextureSize));
+	TSharedRef<IPropertyHandle> TextureSizePropertyHandle = DetailBuilder.GetProperty(GET_MEMBER_NAME_CHECKED(UGLTFMaterialOptions, TextureSize));
 	if (TextureSizePropertyHandle->IsValidHandle())
 	{
 		AddTextureSizeClamping(TextureSizePropertyHandle);
 	}
 
-	TSharedRef<IPropertyHandle> PropertiesHandle = DetailBuilder.GetProperty(GET_MEMBER_NAME_CHECKED(UMaterialOptions, Properties));
+	TSharedRef<IPropertyHandle> PropertiesHandle = DetailBuilder.GetProperty(GET_MEMBER_NAME_CHECKED(UGLTFMaterialOptions, Properties));
 	if (PropertiesHandle->IsValidHandle())
 	{
 		// Setup delegate for when the number of Material Property items changes
diff --git a/Engine/Plugins/Enterprise/GLTFExporter/MigratedFromGit/Source/GLTFMaterialBaking/Private/MaterialOptionsCustomization.h b/Engine/Plugins/Enterprise/GLTFExporter/MigratedFromGit/Source/GLTFMaterialBaking/Private/MaterialOptionsCustomization.h
--- a/Engine/Plugins/Enterprise/GLTFExporter/MigratedFromGit/Source/GLTFMaterialBaking/Private/MaterialOptionsCustomization.h
+++ b/Engine/Plugins/Enterprise/GLTFExporter/MigratedFromGit/Source/GLTFMaterialBaking/Private/MaterialOptionsCustomization.h
@@ -20,10 +20,10 @@ protected:
 protected:
 	/** Property restriction instance used for limiting EMaterialProperty selection */
 	TSharedPtr<FPropertyRestriction> PropertyRestriction;
-	class UMaterialOptions* CurrentOptions;
+	class UGLTFMaterialOptions* CurrentOptions;
 };
 
-/** Detail customization for UMaterialOptions */
+/** Detail customization for UGLTFMaterialOptions */
 class FMaterialOptionsCustomization : public IDetailCustomization
 {
 public:
diff --git a/Engine/Plugins/Enterprise/GLTFExporter/MigratedFromGit/Source/GLTFMaterialBaking/Private/MaterialOptionsWindow.cpp b/Engine/Plugins/Enterprise/GLTFExporter/MigratedFromGit/Source/GLTFMaterialBaking/Private/MaterialOptionsWindow.cpp
--- a/Engine/Plugins/Enterprise/GLTFExporter/MigratedFromGit/Source/GLTFMaterialBaking/Private/MaterialOptionsWindow.cpp
+++ b/Engine/Plugins/Enterprise/GLTFExporter/MigratedFromGit/Source/GLTFMaterialBaking/Private/MaterialOptionsWindow.cpp
@@ -33,7 +33,7 @@ void SMaterialOptions::Construct(const FArguments& InArgs)
 
 	DetailsView = PropertyEditorModule.CreateDetailView(DetailsViewArgs);
 	// Register instance property customization
-	DetailsView->RegisterInstancedCustomPropertyLayout(UMaterialOptions::StaticClass(), FOnGetDetailCustomizationInstance::CreateLambda([=]() { return FMaterialOptionsCustomization::MakeInstance(InArgs._NumLODs); }));
+	DetailsView->RegisterInstancedCustomPropertyLayout(UGLTFMaterialOptions::StaticClass(), FOnGetDetailCustomizationInstance::CreateLambda([=]() { return FMaterialOptionsCustomization::MakeInstance(InArgs._NumLODs); }));
 
 	// Set provided objects on SDetailsView
 	DetailsView->SetObjects(InArgs._SettingsObjects, true);
@@ -78,7 +78,7 @@ void SMaterialOptions::Construct(const FArguments& InArgs)
 FReply SMaterialOptions::OnConfirm()
 {
 	// Ensure the user has selected at least one LOD index
-	if (GetMutableDefault<UMaterialOptions>()->LODIndices.Num() == 0)
+	if (GetMutableDefault<UGLTFMaterialOptions>()->LODIndices.Num() == 0)
 	{
 		FText Title = LOCTEXT("MaterialBake_SelectLODErrorTitle", "Invalid options");
 		FMessageDialog::Open(EAppMsgType::Ok, LOCTEXT("MaterialBake_SelectLODError", "Ensure that atleast one LOD index is selected."), &Title);
diff --git a/Engine/Plugins/Enterprise/GLTFExporter/MigratedFromGit/Source/GLTFMaterialBaking/Public/IMaterialBakingModule.h b/Engine/Plugins/Enterprise/GLTFExporter/MigratedFromGit/Source/GLTFMaterialBaking/Public/IMaterialBakingModule.h
--- a/Engine/Plugins/Enterprise/GLTFExporter/MigratedFromGit/Source/GLTFMaterialBaking/Public/IMaterialBakingModule.h
+++ b/Engine/Plugins/Enterprise/GLTFExporter/MigratedFromGit/Source/GLTFMaterialBaking/Public/IMaterialBakingModule.h
@@ -8,7 +8,7 @@
 #include "SceneTypes.h"
 
 class UTextureRenderTarget2D;
-class UMaterialOptions;
+class UGLTFMaterialOptions;
 class UMaterialInterface;
 
 class FExportMaterialProxy;
diff --git a/Engine/Plugins/Enterprise/GLTFExporter/MigratedFromGit/Source/GLTFMaterialBaking/Public/MaterialOptions.h b/Engine/Plugins/Enterprise/GLTFExporter/MigratedFromGit/Source/GLTFMaterialBaking/Public/MaterialOptions.h
--- a/Engine/Plugins/Enterprise/GLTFExporter/MigratedFromGit/Source/GLTFMaterialBaking/Public/MaterialOptions.h
+++ b/Engine/Plugins/Enterprise/GLTFExporter/MigratedFromGit/Source/GLTFMaterialBaking/Public/MaterialOptions.h
@@ -61,13 +61,13 @@ struct FPropertyEntry
 
 /** Options object to define what and how a material should be baked out */
 UCLASS(config = Editor, Blueprintable)
-class MATERIALBAKING_API UMaterialOptions : public UObject 
+class MATERIALBAKING_API UGLTFMaterialOptions : public UObject 
 {
 	GENERATED_BODY()
 
 public:
 
-	UMaterialOptions()
+	UGLTFMaterialOptions()
 		: TextureSize(128, 128), bUseMeshData(false), bUseSpecificUVIndex(false), TextureCoordinateIndex(0)
 	{
 		Properties.Add(MP_BaseColor);
@@ -101,24 +101,24 @@ public:
 
 /** Asset bake options object */
 UCLASS(Config = Editor, Blueprintable)
-class MATERIALBAKING_API UAssetBakeOptions : public UObject
+class MATERIALBAKING_API UGLTFAssetBakeOptions : public UObject
 {
 	GENERATED_BODY()
 public:
 
-	UAssetBakeOptions()
+	UGLTFAssetBakeOptions()
 	{
 	}
 };
 
 /** Material merge options object */
 UCLASS(Config = Editor, Blueprintable)
-class MATERIALBAKING_API UMaterialMergeOptions: public UObject
+class MATERIALBAKING_API UGLTFMaterialMergeOptions: public UObject
 {
 	GENERATED_BODY()
 public:
 
-	UMaterialMergeOptions() : Method(EMaterialBakeMethod::IndividualMaterial), BlendMode(EBlendMode::BLEND_Opaque)
+	UGLTFMaterialMergeOptions() : Method(EMaterialBakeMethod::IndividualMaterial), BlendMode(EBlendMode::BLEND_Opaque)
 	{
 	}
 
