[GLTFExporter] refactor: switch GLTFExporter to use plugin-embedded GLTFMaterialBaking module (instead of engine-built-in MaterialBaking module)

commit 64392191fcda955e314e5d85df06a7a99bc9805c
Date: Wed, 4 Nov 2020 21:28:38 +0100

    refactor: switch GLTFExporter to use plugin-embedded GLTFMaterialBaking module (instead of engine-built-in MaterialBaking module)
    (cherry picked from commit d9e61ac2667ac101abdbf86926a2c4f7b1b170fe​)

#skipci

diff --git a/Engine/Plugins/Enterprise/GLTFExporter/MigratedFromGit/Source/GLTFExporter/GLTFExporter.Build.cs b/Engine/Plugins/Enterprise/GLTFExporter/MigratedFromGit/Source/GLTFExporter/GLTFExporter.Build.cs
--- a/Engine/Plugins/Enterprise/GLTFExporter/MigratedFromGit/Source/GLTFExporter/GLTFExporter.Build.cs
+++ b/Engine/Plugins/Enterprise/GLTFExporter/MigratedFromGit/Source/GLTFExporter/GLTFExporter.Build.cs
@@ -31,7 +31,7 @@ namespace UnrealBuildTool.Rules
 					"LevelSequence",
 					"MovieScene",
 					"MovieSceneTracks",
-					"MaterialBaking",
+					"GLTFMaterialBaking",
 					"VariantManagerContent",
 					"GLTFExporterRuntime",
 					"MeshMergeUtilities",
diff --git a/Engine/Plugins/Enterprise/GLTFExporter/MigratedFromGit/Source/GLTFExporter/Private/Converters/GLTFMaterialUtility.cpp b/Engine/Plugins/Enterprise/GLTFExporter/MigratedFromGit/Source/GLTFExporter/Private/Converters/GLTFMaterialUtility.cpp
--- a/Engine/Plugins/Enterprise/GLTFExporter/MigratedFromGit/Source/GLTFExporter/Private/Converters/GLTFMaterialUtility.cpp
+++ b/Engine/Plugins/Enterprise/GLTFExporter/MigratedFromGit/Source/GLTFExporter/Private/Converters/GLTFMaterialUtility.cpp
@@ -7,8 +7,8 @@
 #include "CanvasItem.h"
 #include "CanvasTypes.h"
 #include "Modules/ModuleManager.h"
-#include "IMaterialBakingModule.h"
-#include "MaterialBakingStructures.h"
+#include "GLTFMaterialBaking/Public/IMaterialBakingModule.h"
+#include "GLTFMaterialBaking/Public/MaterialBakingStructures.h"
 #include "NormalMapPreview.h"
 #include "Materials/MaterialExpressionCustomOutput.h"
 #include "Materials/MaterialExpressionClearCoatNormalCustomOutput.h"
@@ -286,7 +286,7 @@ FGLTFPropertyBakeOutput FGLTFMaterialUtility::BakeMaterialProperty(const FIntPoi
 	MatSettings.Add(&MatSet);
 
 	TArray<FBakeOutput> BakeOutputs;
-	IMaterialBakingModule& Module = FModuleManager::Get().LoadModuleChecked<IMaterialBakingModule>("MaterialBaking");
+	IMaterialBakingModule& Module = FModuleManager::Get().LoadModuleChecked<IMaterialBakingModule>("GLTFMaterialBaking");
 
 	Module.BakeMaterials(MatSettings, MeshSettings, BakeOutputs);
 
