[GLTFExporter] feat: add support for all clear-coat properties (incl custom output) in custom baking module (see https://github.com/EpicGames/UnrealEngine/pull/7326)

commit 5e01f764a8acfd9ce381beafa1f905a5388022c0
Date: Wed, 4 Nov 2020 22:29:45 +0100

    feat: add support for all clear-coat properties (incl custom output) in custom baking module (see https://github.com/EpicGames/UnrealEngine/pull/7326)
    (cherry picked from commit c7fc3d28e1a066daf45a386ec0ff52ee52ea1044​)

#skipci

diff --git a/Engine/Plugins/Enterprise/GLTFExporter/MigratedFromGit/Source/GLTFMaterialBaking/Private/ExportMaterialProxy.h b/Engine/Plugins/Enterprise/GLTFExporter/MigratedFromGit/Source/GLTFMaterialBaking/Private/ExportMaterialProxy.h
--- a/Engine/Plugins/Enterprise/GLTFExporter/MigratedFromGit/Source/GLTFMaterialBaking/Private/ExportMaterialProxy.h
+++ b/Engine/Plugins/Enterprise/GLTFExporter/MigratedFromGit/Source/GLTFMaterialBaking/Private/ExportMaterialProxy.h
@@ -17,6 +17,7 @@
 #include "Materials/MaterialInterface.h"
 #include "SceneTypes.h"
 #include "Materials/Material.h"
+#include "Materials/MaterialExpressionCustomOutput.h"
 
 struct FExportMaterialCompiler : public FProxyMaterialCompiler
 {
@@ -221,10 +222,11 @@ struct FExportMaterialCompiler : public FProxyMaterialCompiler
 class FExportMaterialProxy : public FMaterial, public FMaterialRenderProxy
 {
 public:
-	FExportMaterialProxy(UMaterialInterface* InMaterialInterface, EMaterialProperty InPropertyToCompile, bool bInSynchronousCompilation = true)
+	FExportMaterialProxy(UMaterialInterface* InMaterialInterface, EMaterialProperty InPropertyToCompile, const FString& InCustomOutputToCompile = TEXT(""), bool bInSynchronousCompilation = true)
 		: FMaterial()
 		, MaterialInterface(InMaterialInterface)
 		, PropertyToCompile(InPropertyToCompile)
+		, CustomOutputToCompile(InCustomOutputToCompile)
 		, bSynchronousCompilation(bInSynchronousCompilation)
 	{
 		SetQualityLevelProperties(EMaterialQualityLevel::High, false, GMaxRHIFeatureLevel);
@@ -267,6 +269,12 @@ public:
 		case MP_Opacity: ResourceId.Usage = EMaterialShaderMapUsage::MaterialExportOpacity; break;
 		case MP_OpacityMask: ResourceId.Usage = EMaterialShaderMapUsage::MaterialExportOpacityMask; break;
 		case MP_SubsurfaceColor: ResourceId.Usage = EMaterialShaderMapUsage::MaterialExportSubSurfaceColor; break;
+		case MP_CustomData0: ResourceId.Usage = EMaterialShaderMapUsage::MaterialExportClearCoat; break;
+		case MP_CustomData1: ResourceId.Usage = EMaterialShaderMapUsage::MaterialExportClearCoatRoughness; break;
+		case MP_CustomOutput:
+			ResourceId.Usage = EMaterialShaderMapUsage::MaterialExportCustomOutput;
+			ResourceId.UsageCustomOutput = InCustomOutputToCompile;
+			break;
 		default:
 			ensureMsgf(false, TEXT("ExportMaterial has no usage for property %i.  Will likely reuse the normal rendering shader and crash later with a parameter mismatch"), (int32)InPropertyToCompile);
 			break;
@@ -376,6 +384,8 @@ public:
 			case MP_Anisotropy:
 			case MP_Metallic:
 			case MP_AmbientOcclusion:
+			case MP_CustomData0:
+			case MP_CustomData1:
 				// Only return for Opaque and Masked...
 				if (BlendMode == BLEND_Opaque || BlendMode == BLEND_Masked)
 				{
@@ -400,6 +410,9 @@ public:
 				break;
 			case MP_ShadingModel:
 				return MaterialInterface->CompileProperty(&ProxyCompiler, MP_ShadingModel);
+			case MP_CustomOutput:
+				 // NOTE: Currently we can assume input index is always 0, which it is for all custom outputs that are registered as material attributes
+				return CompileInputForCustomOutput(Compiler, 0, ForceCast_Exact_Replicate);
 			default:
 				return Compiler->Constant(1.0f);
 			}
@@ -608,6 +621,48 @@ public:
 		}
 	}
 
+private:
+	int32 CompileInputForCustomOutput(FMaterialCompiler* Compiler, int32 InputIndex, uint32 ForceCastFlags) const
+	{
+		FGuid AttributeID = FMaterialAttributeDefinitionMap::GetCustomAttributeID(CustomOutputToCompile);
+		check(AttributeID.IsValid());
+
+		UMaterialExpressionCustomOutput* Expression = GetCustomOutputExpressionToCompile();
+		FExpressionInput* ExpressionInput = Expression ? Expression->GetInput(InputIndex) : nullptr;
+		int32 Result = INDEX_NONE;
+
+		if (ExpressionInput)
+		{
+			Result = ExpressionInput->Compile(Compiler);
+		}
+		else
+		{
+			Result = FMaterialAttributeDefinitionMap::CompileDefaultExpression(Compiler, AttributeID);
+		}
+
+		if (ForceCastFlags & MFCF_ForceCast)
+		{
+			Result = Compiler->ForceCast(Result, FMaterialAttributeDefinitionMap::GetValueType(AttributeID), ForceCastFlags);
+		}
+
+		return Result;
+	}
+
+	UMaterialExpressionCustomOutput* GetCustomOutputExpressionToCompile() const
+	{
+		UMaterial* BaseMaterial = MaterialInterface->GetMaterial();
+		for (UMaterialExpression* Expression : BaseMaterial->Expressions)
+		{
+			UMaterialExpressionCustomOutput* CustomOutputExpression = Cast<UMaterialExpressionCustomOutput>(Expression);
+			if (CustomOutputExpression && CustomOutputExpression->GetDisplayName() == CustomOutputToCompile)
+			{
+				return CustomOutputExpression;
+			}
+		}
+
+		return nullptr;
+	}
+
 private:
 	/** The material interface for this proxy */
 	UMaterialInterface* MaterialInterface;
@@ -615,6 +670,8 @@ private:
 	TArray<UObject*> ReferencedTextures;
 	/** The property to compile for rendering the sample */
 	EMaterialProperty PropertyToCompile;
+	/** The name of the specific custom output to compile for rendering the sample. Only used if PropertyToCompile is MP_CustomOutput */
+	FString CustomOutputToCompile;
 	FGuid Id;
 	bool bSynchronousCompilation;
 };
\ No newline at end of file
diff --git a/Engine/Plugins/Enterprise/GLTFExporter/MigratedFromGit/Source/GLTFMaterialBaking/Private/MaterialBakingModule.cpp b/Engine/Plugins/Enterprise/GLTFExporter/MigratedFromGit/Source/GLTFMaterialBaking/Private/MaterialBakingModule.cpp
--- a/Engine/Plugins/Enterprise/GLTFExporter/MigratedFromGit/Source/GLTFMaterialBaking/Private/MaterialBakingModule.cpp
+++ b/Engine/Plugins/Enterprise/GLTFExporter/MigratedFromGit/Source/GLTFMaterialBaking/Private/MaterialBakingModule.cpp
@@ -239,25 +239,27 @@ namespace FMaterialBakingModuleImpl
 void FMaterialBakingModule::StartupModule()
 {
 	// Set which properties should enforce gamma correction
-	FMemory::Memset(PerPropertyGamma, (uint8)false);
-	PerPropertyGamma[MP_Normal] = true;
-	PerPropertyGamma[MP_Opacity] = true;
-	PerPropertyGamma[MP_OpacityMask] = true;
+	PerPropertyGamma.Add(MP_Normal);
+	PerPropertyGamma.Add(MP_Opacity);
+	PerPropertyGamma.Add(MP_OpacityMask);
+	PerPropertyGamma.Add(TEXT("ClearCoatBottomNormal"));
 
 	// Set which pixel format should be used for the possible baked out material properties
-	FMemory::Memset(PerPropertyFormat, (uint8)PF_Unknown);
-	PerPropertyFormat[MP_EmissiveColor] = PF_FloatRGBA;
-	PerPropertyFormat[MP_Opacity] = PF_B8G8R8A8;
-	PerPropertyFormat[MP_OpacityMask] = PF_B8G8R8A8;
-	PerPropertyFormat[MP_BaseColor] = PF_B8G8R8A8;
-	PerPropertyFormat[MP_Metallic] = PF_B8G8R8A8;
-	PerPropertyFormat[MP_Specular] = PF_B8G8R8A8;
-	PerPropertyFormat[MP_Roughness] = PF_B8G8R8A8;
-	PerPropertyFormat[MP_Anisotropy] = PF_B8G8R8A8;
-	PerPropertyFormat[MP_Normal] = PF_B8G8R8A8;
-	PerPropertyFormat[MP_Tangent] = PF_B8G8R8A8;
-	PerPropertyFormat[MP_AmbientOcclusion] = PF_B8G8R8A8;
-	PerPropertyFormat[MP_SubsurfaceColor] = PF_B8G8R8A8;
+	PerPropertyFormat.Add(MP_EmissiveColor, PF_FloatRGBA);
+	PerPropertyFormat.Add(MP_Opacity, PF_B8G8R8A8);
+	PerPropertyFormat.Add(MP_OpacityMask, PF_B8G8R8A8);
+	PerPropertyFormat.Add(MP_BaseColor, PF_B8G8R8A8);
+	PerPropertyFormat.Add(MP_Metallic, PF_B8G8R8A8);
+	PerPropertyFormat.Add(MP_Specular, PF_B8G8R8A8);
+	PerPropertyFormat.Add(MP_Roughness, PF_B8G8R8A8);
+	PerPropertyFormat.Add(MP_Anisotropy, PF_B8G8R8A8);
+	PerPropertyFormat.Add(MP_Normal, PF_B8G8R8A8);
+	PerPropertyFormat.Add(MP_Tangent, PF_B8G8R8A8);
+	PerPropertyFormat.Add(MP_AmbientOcclusion, PF_B8G8R8A8);
+	PerPropertyFormat.Add(MP_SubsurfaceColor, PF_B8G8R8A8);
+	PerPropertyFormat.Add(MP_CustomData0, PF_B8G8R8A8);
+	PerPropertyFormat.Add(MP_CustomData1, PF_B8G8R8A8);
+	PerPropertyFormat.Add(TEXT("ClearCoatBottomNormal"), PF_B8G8R8A8);
 
 	// Register property customization
 	FPropertyEditorModule& Module = FModuleManager::Get().LoadModuleChecked<FPropertyEditorModule>("PropertyEditor");
@@ -280,6 +282,47 @@ void FMaterialBakingModule::ShutdownModule()
 }
 
 void FMaterialBakingModule::BakeMaterials(const TArray<FMaterialData*>& MaterialSettings, const TArray<FMeshData*>& MeshSettings, TArray<FBakeOutput>& Output)
+{
+	TArray<FMaterialDataEx> MaterialDataExs;
+	TArray<FMaterialDataEx*> MaterialSettingsEx;
+
+	// Translate old material data to extended types
+	for (const FMaterialData* MaterialData : MaterialSettings)
+	{
+		FMaterialDataEx& MaterialDataEx = MaterialDataExs.AddDefaulted_GetRef();
+		MaterialDataEx.Material = MaterialData->Material;
+		MaterialDataEx.bPerformBorderSmear = MaterialData->bPerformBorderSmear;
+
+		for (const TPair<EMaterialProperty, FIntPoint>& PropertySizePair : MaterialData->PropertySizes)
+		{
+			MaterialDataEx.PropertySizes.Add(PropertySizePair.Key, PropertySizePair.Value);
+		}
+
+		MaterialSettingsEx.Add(&MaterialDataEx);
+	}
+
+	TArray<FBakeOutputEx> BakeOutputExs;
+	BakeMaterials(MaterialSettingsEx, MeshSettings, BakeOutputExs);
+
+	// Translate extended bake output to old types
+	for (FBakeOutputEx& BakeOutputEx : BakeOutputExs)
+	{
+		FBakeOutput& BakeOutput = Output.AddDefaulted_GetRef();
+		BakeOutput.EmissiveScale = BakeOutputEx.EmissiveScale;
+
+		for (TPair<FMaterialPropertyEx, FIntPoint>& PropertySizePair : BakeOutputEx.PropertySizes)
+		{
+			BakeOutput.PropertySizes.Add(PropertySizePair.Key.Type, PropertySizePair.Value);
+		}
+
+		for (TPair<FMaterialPropertyEx, TArray<FColor>>& PropertyDataPair : BakeOutputEx.PropertyData)
+		{
+			BakeOutput.PropertyData.Add(PropertyDataPair.Key.Type, MoveTemp(PropertyDataPair.Value));
+		}
+	}
+}
+
+void FMaterialBakingModule::BakeMaterials(const TArray<FMaterialDataEx*>& MaterialSettings, const TArray<FMeshData*>& MeshSettings, TArray<FBakeOutputEx>& Output)
 {
 	TRACE_CPUPROFILER_EVENT_SCOPE(FMaterialBakingModule::BakeMaterials)
 
@@ -331,15 +374,15 @@ void FMaterialBakingModule::BakeMaterials(const TArray<FMaterialData*>& Material
 		TRACE_CPUPROFILER_EVENT_SCOPE(PrepareRenderItems);
 		
 		TMap<FMaterialBakingModuleImpl::FRenderItemKey, FMeshMaterialRenderItem*>* RenderItems = new TMap<FRenderItemKey, FMeshMaterialRenderItem *>();
-		const FMaterialData* CurrentMaterialSettings = MaterialSettings[MaterialIndex];
+		const FMaterialDataEx* CurrentMaterialSettings = MaterialSettings[MaterialIndex];
 		const FMeshData* CurrentMeshSettings = MeshSettings[MaterialIndex];
 
-		for (TMap<EMaterialProperty, FIntPoint>::TConstIterator PropertySizeIterator = CurrentMaterialSettings->PropertySizes.CreateConstIterator(); PropertySizeIterator; ++PropertySizeIterator)
+		for (TMap<FMaterialPropertyEx, FIntPoint>::TConstIterator PropertySizeIterator = CurrentMaterialSettings->PropertySizes.CreateConstIterator(); PropertySizeIterator; ++PropertySizeIterator)
 		{
 			FRenderItemKey RenderItemKey(CurrentMeshSettings, PropertySizeIterator.Value());
 			if (RenderItems->Find(RenderItemKey) == nullptr)
 			{
-				RenderItems->Add(RenderItemKey, new FMeshMaterialRenderItem(CurrentMaterialSettings, CurrentMeshSettings, PropertySizeIterator.Key(), &MaterialBakingDynamicMeshBufferAllocator));
+				RenderItems->Add(RenderItemKey, new FMeshMaterialRenderItem(PropertySizeIterator.Value(), CurrentMeshSettings, &MaterialBakingDynamicMeshBufferAllocator));
 			}
 		}
 
@@ -368,7 +411,7 @@ void FMaterialBakingModule::BakeMaterials(const TArray<FMaterialData*>& Material
 		for (int32 Index = 0; Index < NumMaterials; ++Index)
 		{
 			int32 MaterialIndex = ProcessingOrder[Index];
-			const FMaterialData* CurrentMaterialSettings = MaterialSettings[MaterialIndex];
+			const FMaterialDataEx* CurrentMaterialSettings = MaterialSettings[MaterialIndex];
 
 			TArray<UTexture*> MaterialTextures;
 			CurrentMaterialSettings->Material->GetUsedTextures(MaterialTextures, EMaterialQualityLevel::Num, true, GMaxRHIFeatureLevel, true);
@@ -391,7 +434,7 @@ void FMaterialBakingModule::BakeMaterials(const TArray<FMaterialData*>& Material
 				}
 			}
 
-			for (TMap<EMaterialProperty, FIntPoint>::TConstIterator PropertySizeIterator = CurrentMaterialSettings->PropertySizes.CreateConstIterator(); PropertySizeIterator; ++PropertySizeIterator)
+			for (TMap<FMaterialPropertyEx, FIntPoint>::TConstIterator PropertySizeIterator = CurrentMaterialSettings->PropertySizes.CreateConstIterator(); PropertySizeIterator; ++PropertySizeIterator)
 			{
 				// They will be stored in the pool and compiled asynchronously
 				CreateMaterialProxy(CurrentMaterialSettings->Material, PropertySizeIterator.Key());
@@ -431,11 +474,11 @@ void FMaterialBakingModule::BakeMaterials(const TArray<FMaterialData*>& Material
 			NextRenderItem++;
 		}
 
-		const FMaterialData* CurrentMaterialSettings = MaterialSettings[MaterialIndex];
+		const FMaterialDataEx* CurrentMaterialSettings = MaterialSettings[MaterialIndex];
 		const FMeshData* CurrentMeshSettings = MeshSettings[MaterialIndex];
-		FBakeOutput& CurrentOutput = Output[MaterialIndex];
+		FBakeOutputEx& CurrentOutput = Output[MaterialIndex];
 
-		TArray<EMaterialProperty> MaterialPropertiesToBakeOut;
+		TArray<FMaterialPropertyEx> MaterialPropertiesToBakeOut;
 		CurrentMaterialSettings->PropertySizes.GenerateKeyArray(MaterialPropertiesToBakeOut);
 
 		const int32 NumPropertiesToRender = MaterialPropertiesToBakeOut.Num();
@@ -447,13 +490,13 @@ void FMaterialBakingModule::BakeMaterials(const TArray<FMaterialData*>& Material
 			CurrentOutput.PropertySizes = CurrentMaterialSettings->PropertySizes;
 			for (int32 PropertyIndex = 0; PropertyIndex < NumPropertiesToRender; ++PropertyIndex)
 			{
-				const EMaterialProperty Property = MaterialPropertiesToBakeOut[PropertyIndex];
+				const FMaterialPropertyEx& Property = MaterialPropertiesToBakeOut[PropertyIndex];
 				CurrentOutput.PropertyData.Add(Property);
 			}
 
 			for (int32 PropertyIndex = 0; PropertyIndex < NumPropertiesToRender; ++PropertyIndex)
 			{
-				const EMaterialProperty Property = MaterialPropertiesToBakeOut[PropertyIndex];
+				const FMaterialPropertyEx& Property = MaterialPropertiesToBakeOut[PropertyIndex];
 				FExportMaterialProxy* ExportMaterialProxy = CreateMaterialProxy(CurrentMaterialSettings->Material, Property);
 
 				if (!ExportMaterialProxy->IsCompilationFinished())
@@ -462,9 +505,13 @@ void FMaterialBakingModule::BakeMaterials(const TArray<FMaterialData*>& Material
 					ExportMaterialProxy->FinishCompilation();
 				}
 
+				// Lookup gamma and format settings for property, if not found use default values
+				bool InForceLinearGamma = PerPropertyGamma.Contains(Property);
+				EPixelFormat PixelFormat = PerPropertyFormat.Contains(Property) ? PerPropertyFormat[Property] : PF_B8G8R8A8;
+
 				// It is safe to reuse the same render target for each draw pass since they all execute sequentially on the GPU and are copied to staging buffers before
 				// being reused.
-				UTextureRenderTarget2D* RenderTarget = CreateRenderTarget(PerPropertyGamma[Property], PerPropertyFormat[Property], CurrentOutput.PropertySizes[Property]);
+				UTextureRenderTarget2D* RenderTarget = CreateRenderTarget(InForceLinearGamma, PixelFormat, CurrentOutput.PropertySizes[Property]);
 				if (RenderTarget != nullptr)
 				{
 					// Perform everything left of the operation directly on the render thread since we need to modify some RenderItem's properties
@@ -473,7 +520,7 @@ void FMaterialBakingModule::BakeMaterials(const TArray<FMaterialData*>& Material
 					ENQUEUE_RENDER_COMMAND(RenderOneMaterial)(
 						[this, RenderItems, RenderTarget, Property, ExportMaterialProxy, &PipelineContext, PipelineIndex, &StagingBufferPool, &NumTasks, bSaveIntermediateTextures, &MaterialSettings, &MeshSettings, MaterialIndex, &Output](FRHICommandListImmediate& RHICmdList)
 						{
-							const FMaterialData* CurrentMaterialSettings = MaterialSettings[MaterialIndex];
+							const FMaterialDataEx* CurrentMaterialSettings = MaterialSettings[MaterialIndex];
 							const FMeshData* CurrentMeshSettings = MeshSettings[MaterialIndex];
 
 							FMeshMaterialRenderItem& RenderItem = *RenderItems->FindChecked(FRenderItemKey(CurrentMeshSettings, FIntPoint(RenderTarget->GetSurfaceWidth(), RenderTarget->GetSurfaceHeight())));
@@ -483,7 +530,6 @@ void FMaterialBakingModule::BakeMaterials(const TArray<FMaterialData*>& Material
 								.SetWorldTimes(0.0f, 0.0f, 0.0f)
 								.SetGammaCorrection(RenderTarget->GetRenderTargetResource()->GetDisplayGamma()));
 
-							RenderItem.MaterialProperty = Property;
 							RenderItem.MaterialRenderProxy = ExportMaterialProxy;
 							RenderItem.ViewFamily = &ViewFamily;
 
@@ -500,7 +546,7 @@ void FMaterialBakingModule::BakeMaterials(const TArray<FMaterialData*>& Material
 							Canvas.Flush_RenderThread(RHICmdList);
 							SortElement.RenderBatchArray.Empty();
 
-							FTexture2DRHIRef StagingBufferRef = StagingBufferPool.CreateStagingBuffer_RenderThread(RHICmdList, RenderTargetResource->GetSizeX(), RenderTargetResource->GetSizeY(), PerPropertyFormat[Property]);
+							FTexture2DRHIRef StagingBufferRef = StagingBufferPool.CreateStagingBuffer_RenderThread(RHICmdList, RenderTargetResource->GetSizeX(), RenderTargetResource->GetSizeY(), RenderTarget->GetFormat());
 							FGPUFenceRHIRef GPUFence = RHICreateGPUFence(TEXT("MaterialBackingFence"));
 
 							FResolveRect Rect(0, 0, RenderTargetResource->GetSizeX(), RenderTargetResource->GetSizeY());
@@ -514,15 +560,15 @@ void FMaterialBakingModule::BakeMaterials(const TArray<FMaterialData*>& Material
 								{
 									TRACE_CPUPROFILER_EVENT_SCOPE(FinalProcessing)
 
-									FBakeOutput& CurrentOutput  = Output[MaterialIndex];
+									FBakeOutputEx& CurrentOutput  = Output[MaterialIndex];
 									TArray<FColor>& OutputColor = CurrentOutput.PropertyData[Property];
 									FIntPoint& OutputSize       = CurrentOutput.PropertySizes[Property];
 
 									OutputColor.SetNum(OutputSize.X * OutputSize.Y);
 
-									if (Property == MP_EmissiveColor)
+									if (Property.Type == MP_EmissiveColor)
 									{
-										// Only one thread will write to CurrentOutput.EmissiveScale since there can be only one emissive channel property per FBakeOutput
+										// Only one thread will write to CurrentOutput.EmissiveScale since there can be only one emissive channel property per FBakeOutputEx
 										FMaterialBakingModule::ProcessEmissiveOutput((const FFloat16Color*)Data, DataWidth, OutputSize, OutputColor, CurrentOutput.EmissiveScale);
 									}
 									else
@@ -546,10 +592,7 @@ void FMaterialBakingModule::BakeMaterials(const TArray<FMaterialData*>& Material
 									if (bSaveIntermediateTextures)
 									{
 										TRACE_CPUPROFILER_EVENT_SCOPE(SaveIntermediateTextures)
-										const UEnum* PropertyEnum = StaticEnum<EMaterialProperty>();
-										FName PropertyName = PropertyEnum->GetNameByValue(Property);
-										FString TrimmedPropertyName = PropertyName.ToString();
-										TrimmedPropertyName.RemoveFromStart(TEXT("MP_"));
+										FString TrimmedPropertyName = Property.ToString();
 
 										const FString DirectoryPath = FPaths::ConvertRelativePathToFull(FPaths::ProjectIntermediateDir() + TEXT("MaterialBaking/"));
 										FString FilenameString = FString::Printf(TEXT("%s%s-%d-%s.bmp"), *DirectoryPath, *CurrentMaterialSettings->Material->GetName(), MaterialIndex, *TrimmedPropertyName);
@@ -728,7 +771,7 @@ UTextureRenderTarget2D* FMaterialBakingModule::CreateRenderTarget(bool bInForceL
 	return RenderTarget;
 }
 
-FExportMaterialProxy* FMaterialBakingModule::CreateMaterialProxy(UMaterialInterface* Material, const EMaterialProperty Property)
+FExportMaterialProxy* FMaterialBakingModule::CreateMaterialProxy(UMaterialInterface* Material, const FMaterialPropertyEx& Property)
 {
 	TRACE_CPUPROFILER_EVENT_SCOPE(FMaterialBakingModule::CreateMaterialProxy)
 
@@ -751,7 +794,7 @@ FExportMaterialProxy* FMaterialBakingModule::CreateMaterialProxy(UMaterialInterf
 	// Not found, create a new entry
 	if (Proxy == nullptr)
 	{
-		Proxy = new FExportMaterialProxy(Material, Property, false /* bInSynchronousCompilation */);
+		Proxy = new FExportMaterialProxy(Material, Property.Type, Property.CustomOutput.ToString(), false /* bInSynchronousCompilation */);
 		MaterialProxyPool.Add(Material, FMaterialPoolValue(Property, Proxy));
 	}
 
diff --git a/Engine/Plugins/Enterprise/GLTFExporter/MigratedFromGit/Source/GLTFMaterialBaking/Private/MaterialBakingModule.h b/Engine/Plugins/Enterprise/GLTFExporter/MigratedFromGit/Source/GLTFMaterialBaking/Private/MaterialBakingModule.h
--- a/Engine/Plugins/Enterprise/GLTFExporter/MigratedFromGit/Source/GLTFMaterialBaking/Private/MaterialBakingModule.h
+++ b/Engine/Plugins/Enterprise/GLTFExporter/MigratedFromGit/Source/GLTFMaterialBaking/Private/MaterialBakingModule.h
@@ -7,6 +7,7 @@
 #include "PixelFormat.h"
 #include "SceneTypes.h"
 #include "IMaterialBakingModule.h"
+#include "MaterialPropertyEx.h"
 
 class UTextureRenderTarget2D;
 class UGLTFMaterialOptions;
@@ -30,6 +31,9 @@ public:
 	/** Bakes out material properties according to MaterialSettings using MeshSettings and stores the output in Output */
 	virtual void BakeMaterials(const TArray<FMaterialData*>& MaterialSettings, const TArray<FMeshData*>& MeshSettings, TArray<FBakeOutput>& Output) override;
 
+	/** Bakes out material properties according to extended MaterialSettings using MeshSettings and stores the output in Output */
+	virtual void BakeMaterials(const TArray<FMaterialDataEx*>& MaterialSettings, const TArray<FMeshData*>& MeshSettings, TArray<FBakeOutputEx>& Output) override;
+
 	/** Promps a slate window to allow the user to populate specific material baking settings used while baking out materials */
 	virtual bool SetupMaterialBakeSettings(TArray<TWeakObjectPtr<UObject>>& OptionObjects, int32 NumLODs) override;
 
@@ -38,7 +42,7 @@ protected:
 	UTextureRenderTarget2D* CreateRenderTarget(bool bInForceLinearGamma, EPixelFormat InPixelFormat, const FIntPoint& InTargetSize);
 
 	/* Creates and adds (or reuses a ExportMaterialProxy from the pool if MaterialBaking.UseMaterialProxyCaching is set to 1) */
-	FExportMaterialProxy* CreateMaterialProxy(UMaterialInterface* Material, const EMaterialProperty Property );
+	FExportMaterialProxy* CreateMaterialProxy(UMaterialInterface* Material, const FMaterialPropertyEx& Property);
 
 	/** Helper for emissive color conversion to Output */
 	static void ProcessEmissiveOutput(const FFloat16Color* Color16, int32 Color16Pitch, const FIntPoint& OutputSize, TArray<FColor>& Output, float& EmissiveScale);
@@ -54,13 +58,13 @@ private:
 
 	/** Pool of cached material proxies to optimize material baking workflow, stays resident when MaterialBaking.UseMaterialProxyCaching is set to 1 */
 	typedef TWeakObjectPtr<UMaterialInterface>				FMaterialPoolKey;
-	typedef TPair<EMaterialProperty, FExportMaterialProxy*> FMaterialPoolValue;
+	typedef TPair<FMaterialPropertyEx, FExportMaterialProxy*> FMaterialPoolValue;
 	typedef TMultiMap<FMaterialPoolKey, FMaterialPoolValue, FDefaultSetAllocator, TWeakObjectPtrMapKeyFuncs<FMaterialPoolKey, FMaterialPoolValue, true /*bInAllowDuplicateKeys*/>> FMaterialPoolMap;
 	FMaterialPoolMap MaterialProxyPool;
 
 	/** Pixel formats to use for baking out specific material properties */
-	EPixelFormat PerPropertyFormat[MP_MAX];
+	TMap<FMaterialPropertyEx, EPixelFormat> PerPropertyFormat;
 
 	/** Whether or not to enforce gamma correction while baking out specific material properties */
-	bool PerPropertyGamma[MP_MAX];
+	TSet<FMaterialPropertyEx> PerPropertyGamma;
 };
diff --git a/Engine/Plugins/Enterprise/GLTFExporter/MigratedFromGit/Source/GLTFMaterialBaking/Private/MaterialRenderItem.cpp b/Engine/Plugins/Enterprise/GLTFExporter/MigratedFromGit/Source/GLTFMaterialBaking/Private/MaterialRenderItem.cpp
--- a/Engine/Plugins/Enterprise/GLTFExporter/MigratedFromGit/Source/GLTFMaterialBaking/Private/MaterialRenderItem.cpp
+++ b/Engine/Plugins/Enterprise/GLTFExporter/MigratedFromGit/Source/GLTFMaterialBaking/Private/MaterialRenderItem.cpp
@@ -11,13 +11,11 @@
 #define SHOW_WIREFRAME_MESH 0
 
 FMeshMaterialRenderItem::FMeshMaterialRenderItem(
-	const FMaterialData* InMaterialSettings,
+	const FIntPoint& InTextureSize,
 	const FMeshData* InMeshSettings,
-	EMaterialProperty InMaterialProperty,
 	FDynamicMeshBufferAllocator* InDynamicMeshBufferAllocator)
 	: MeshSettings(InMeshSettings)
-	, MaterialSettings(InMaterialSettings)
-	, MaterialProperty(InMaterialProperty)
+	, TextureSize(InTextureSize)
 	, MaterialRenderProxy(nullptr)
 	, ViewFamily(nullptr)
 	, bMeshElementDirty(true)
@@ -157,9 +155,8 @@ void FMeshMaterialRenderItem::PopulateWithQuadData()
 	const float V = MeshSettings->TextureCoordinateBox.Min.Y;
 	const float SizeU = MeshSettings->TextureCoordinateBox.Max.X - MeshSettings->TextureCoordinateBox.Min.X;
 	const float SizeV = MeshSettings->TextureCoordinateBox.Max.Y - MeshSettings->TextureCoordinateBox.Min.Y;
-	const FIntPoint& PropertySize = MaterialSettings->PropertySizes[MaterialProperty];
-	const float ScaleX = PropertySize.X;
-	const float ScaleY = PropertySize.Y;
+	const float ScaleX = TextureSize.X;
+	const float ScaleY = TextureSize.Y;
 
 	// add vertices
 	for (int32 VertIndex = 0; VertIndex < 4; VertIndex++)
@@ -198,9 +195,8 @@ void FMeshMaterialRenderItem::PopulateWithMeshData()
 	Vertices.Empty(NumVerts);
 	Indices.Empty(NumVerts >> 1);
 
-	const FIntPoint& PropertySize = MaterialSettings->PropertySizes[MaterialProperty];
-	const float ScaleX = PropertySize.X;
-	const float ScaleY = PropertySize.Y;
+	const float ScaleX = TextureSize.X;
+	const float ScaleY = TextureSize.Y;
 
 	const static int32 VertexPositionStoredUVChannel = 6;
 	// count number of texture coordinates for this mesh
@@ -223,7 +219,7 @@ void FMeshMaterialRenderItem::PopulateWithMeshData()
 	{
 		const FPolygonGroupID PolygonGroupID = RawMesh->GetPolygonPolygonGroup(PolygonID);
 		const TArray<FTriangleID>& TriangleIDs = RawMesh->GetPolygonTriangleIDs(PolygonID);
-		for (const FTriangleID TriangleID : TriangleIDs)
+		for (const FTriangleID& TriangleID : TriangleIDs)
 		{
 			if (MeshSettings->MaterialIndices.Contains(PolygonGroupID.GetValue()))
 			{
@@ -284,4 +280,4 @@ void FMeshMaterialRenderItem::PopulateWithMeshData()
 			FaceIndex++;
 		}
 	}
-}
\ No newline at end of file
+}
diff --git a/Engine/Plugins/Enterprise/GLTFExporter/MigratedFromGit/Source/GLTFMaterialBaking/Private/MaterialRenderItem.h b/Engine/Plugins/Enterprise/GLTFExporter/MigratedFromGit/Source/GLTFMaterialBaking/Private/MaterialRenderItem.h
--- a/Engine/Plugins/Enterprise/GLTFExporter/MigratedFromGit/Source/GLTFMaterialBaking/Private/MaterialRenderItem.h
+++ b/Engine/Plugins/Enterprise/GLTFExporter/MigratedFromGit/Source/GLTFMaterialBaking/Private/MaterialRenderItem.h
@@ -49,7 +49,7 @@ public:
 class FMeshMaterialRenderItem : public FCanvasBaseRenderItem
 {
 public:
-	FMeshMaterialRenderItem(const FMaterialData* InMaterialSettings, const FMeshData* InMeshSettings, EMaterialProperty InMaterialProperty, FDynamicMeshBufferAllocator* InDynamicMeshBufferAllocator = nullptr);
+	FMeshMaterialRenderItem(const FIntPoint& InTextureSize, const FMeshData* InMeshSettings, FDynamicMeshBufferAllocator* InDynamicMeshBufferAllocator = nullptr);
 	virtual ~FMeshMaterialRenderItem();
 
 	/** Begin FCanvasBaseRenderItem overrides */
@@ -67,11 +67,10 @@ protected:
 	void PopulateWithQuadData();
 	void PopulateWithMeshData();
 public:
-	/** Mesh and material settings to use while baking out the material */
+	/** Mesh settings to use while baking out the material */
 	const FMeshData* MeshSettings;
-	const FMaterialData* MaterialSettings;
-	/** Material property to bake out */
-	EMaterialProperty MaterialProperty;
+	/** The texture size to use while baking */
+	FIntPoint TextureSize;
 	/** Material render proxy (material/shader) to use while baking */
 	FMaterialRenderProxy* MaterialRenderProxy;
 	/** Vertex and index data representing the mesh or a quad */
diff --git a/Engine/Plugins/Enterprise/GLTFExporter/MigratedFromGit/Source/GLTFMaterialBaking/Public/IMaterialBakingModule.h b/Engine/Plugins/Enterprise/GLTFExporter/MigratedFromGit/Source/GLTFMaterialBaking/Public/IMaterialBakingModule.h
--- a/Engine/Plugins/Enterprise/GLTFExporter/MigratedFromGit/Source/GLTFMaterialBaking/Public/IMaterialBakingModule.h
+++ b/Engine/Plugins/Enterprise/GLTFExporter/MigratedFromGit/Source/GLTFMaterialBaking/Public/IMaterialBakingModule.h
@@ -15,8 +15,10 @@ class FExportMaterialProxy;
 class FTextureRenderTargetResource;
 
 struct FMaterialData; 
+struct FMaterialDataEx;
 struct FMeshData;
 struct FBakeOutput;
+struct FBakeOutputEx;
 
 class MATERIALBAKING_API IMaterialBakingModule : public IModuleInterface
 {
@@ -24,6 +26,9 @@ public:
 	/** Bakes out material properties according to MaterialSettings using MeshSettings and stores the output in Output */
 	virtual void BakeMaterials(const TArray<FMaterialData*>& MaterialSettings, const TArray<FMeshData*>& MeshSettings, TArray<FBakeOutput>& Output) = 0;
 
+	/** Bakes out material properties according to extended MaterialSettings using MeshSettings and stores the output in Output */
+	virtual void BakeMaterials(const TArray<FMaterialDataEx*>& MaterialSettings, const TArray<FMeshData*>& MeshSettings, TArray<FBakeOutputEx>& Output) = 0;
+
 	/** Promps a slate window to allow the user to populate specific material baking settings used while baking out materials */
 	virtual  bool SetupMaterialBakeSettings(TArray<TWeakObjectPtr<UObject>>& OptionObjects, int32 NumLODs) = 0;
 };
diff --git a/Engine/Plugins/Enterprise/GLTFExporter/MigratedFromGit/Source/GLTFMaterialBaking/Public/MaterialBakingStructures.h b/Engine/Plugins/Enterprise/GLTFExporter/MigratedFromGit/Source/GLTFMaterialBaking/Public/MaterialBakingStructures.h
--- a/Engine/Plugins/Enterprise/GLTFExporter/MigratedFromGit/Source/GLTFMaterialBaking/Public/MaterialBakingStructures.h
+++ b/Engine/Plugins/Enterprise/GLTFExporter/MigratedFromGit/Source/GLTFMaterialBaking/Public/MaterialBakingStructures.h
@@ -5,6 +5,7 @@
 #include "Math/IntPoint.h"
 #include "SceneTypes.h"
 #include "LightMap.h"
+#include "MaterialPropertyEx.h"
 
 class UMaterialInterface;
 struct FMeshDescription;
@@ -25,6 +26,22 @@ struct FMaterialData
 	bool bPerformBorderSmear;
 };
 
+/** Structure containing extended information about the material and properties which is being baked out */
+struct FMaterialDataEx
+{
+	FMaterialDataEx()
+		: Material(nullptr)
+		, bPerformBorderSmear(true)
+	{}
+
+	/** Material to bake out */
+	UMaterialInterface* Material;
+	/** Extended properties and the texture size at which they should be baked out */
+	TMap<FMaterialPropertyEx, FIntPoint> PropertySizes;
+	/** Whether to smear borders after baking */
+	bool bPerformBorderSmear;
+};
+
 struct FMeshData
 {
 	FMeshData()
@@ -81,3 +98,20 @@ struct FBakeOutput
 	/** Scale used to allow having wide ranges of emissive values in the source materials, the final proxy material will use this value to scale the emissive texture's pixel values */
 	float EmissiveScale;
 };
+
+/** Structure containing extended data being processed while baking out materials*/
+struct FBakeOutputEx
+{
+	FBakeOutputEx()
+		: EmissiveScale(1.0f)
+	{}
+
+	/** Contains the resulting texture data for baking out a extened material's property */
+	TMap<FMaterialPropertyEx, TArray<FColor>> PropertyData;
+
+	/** Contains the resulting texture size for baking out a extened material's property */
+	TMap<FMaterialPropertyEx, FIntPoint> PropertySizes;
+
+	/** Scale used to allow having wide ranges of emissive values in the source materials, the final proxy material will use this value to scale the emissive texture's pixel values */
+	float EmissiveScale;
+};
diff --git a/Engine/Plugins/Enterprise/GLTFExporter/MigratedFromGit/Source/GLTFMaterialBaking/Public/MaterialPropertyEx.h b/Engine/Plugins/Enterprise/GLTFExporter/MigratedFromGit/Source/GLTFMaterialBaking/Public/MaterialPropertyEx.h
new file mode 100644
--- /dev/null
+++ b/Engine/Plugins/Enterprise/GLTFExporter/MigratedFromGit/Source/GLTFMaterialBaking/Public/MaterialPropertyEx.h
@@ -0,0 +1,62 @@
+// Copyright Epic Games, Inc. All Rights Reserved.
+
+#pragma once
+
+#include "SceneTypes.h"
+#include "UObject/NameTypes.h"
+
+/** Structure extending EMaterialProperty to allow detailed information about custom output */
+struct FMaterialPropertyEx
+{
+	FMaterialPropertyEx(EMaterialProperty Type = MP_MAX, const FName& CustomOutput = NAME_None)
+		: Type(Type)
+		, CustomOutput(CustomOutput)
+	{}
+
+	FMaterialPropertyEx(const FName& CustomOutput)
+		: Type(MP_CustomOutput)
+		, CustomOutput(CustomOutput)
+	{}
+
+	FMaterialPropertyEx(const TCHAR* CustomOutput)
+		: Type(MP_CustomOutput)
+		, CustomOutput(CustomOutput)
+	{}
+
+	FORCEINLINE bool operator ==(const FMaterialPropertyEx& Other) const
+	{
+		return Type == Other.Type && (Type != MP_CustomOutput || CustomOutput == Other.CustomOutput);
+	}
+
+	FORCEINLINE bool operator !=(const FMaterialPropertyEx& Other) const
+	{
+		return !(*this == Other);
+	}
+
+	friend FORCEINLINE uint32 GetTypeHash(const FMaterialPropertyEx& Other)
+	{
+		return Other.Type != MP_CustomOutput ? GetTypeHash(Other.Type) : GetTypeHash(Other.CustomOutput);
+	}
+
+	FString ToString() const
+	{
+		if (Type != MP_CustomOutput)
+		{
+			const UEnum* Enum = StaticEnum<EMaterialProperty>();
+			FName Name = Enum->GetNameByValue(Type);
+			FString TrimmedName = Name.ToString();
+			TrimmedName.RemoveFromStart(TEXT("MP_"), ESearchCase::CaseSensitive);
+			return TrimmedName;
+		}
+		else
+		{
+			return CustomOutput.ToString();
+		}
+	}
+	
+	/** The material property */
+	EMaterialProperty Type;
+	
+	/** The name of a specific custom output. Only used if property is MP_CustomOutput */
+	FName CustomOutput;
+};
