[GLTFExporter] fix: build issues in modified custom baking module due to missing engine changes (from https://github.com/EpicGames/UnrealEngine/pull/7326)

commit 990a95017469196924765c5626838489fd4d67a6
Date: Wed, 4 Nov 2020 22:30:57 +0100

    fix: build issues in modified custom baking module due to missing engine changes (from https://github.com/EpicGames/UnrealEngine/pull/7326)
    (cherry picked from commit 08e728ea54773c0db8ed1919b1430bbd7d543cda​)

#skipci

diff --git a/Engine/Plugins/Enterprise/GLTFExporter/MigratedFromGit/Source/GLTFMaterialBaking/Private/ExportMaterialProxy.h b/Engine/Plugins/Enterprise/GLTFExporter/MigratedFromGit/Source/GLTFMaterialBaking/Private/ExportMaterialProxy.h
--- a/Engine/Plugins/Enterprise/GLTFExporter/MigratedFromGit/Source/GLTFMaterialBaking/Private/ExportMaterialProxy.h
+++ b/Engine/Plugins/Enterprise/GLTFExporter/MigratedFromGit/Source/GLTFMaterialBaking/Private/ExportMaterialProxy.h
@@ -269,12 +269,12 @@ public:
 		case MP_Opacity: ResourceId.Usage = EMaterialShaderMapUsage::MaterialExportOpacity; break;
 		case MP_OpacityMask: ResourceId.Usage = EMaterialShaderMapUsage::MaterialExportOpacityMask; break;
 		case MP_SubsurfaceColor: ResourceId.Usage = EMaterialShaderMapUsage::MaterialExportSubSurfaceColor; break;
-		case MP_CustomData0: ResourceId.Usage = EMaterialShaderMapUsage::MaterialExportClearCoat; break;
-		case MP_CustomData1: ResourceId.Usage = EMaterialShaderMapUsage::MaterialExportClearCoatRoughness; break;
-		case MP_CustomOutput:
-			ResourceId.Usage = EMaterialShaderMapUsage::MaterialExportCustomOutput;
-			ResourceId.UsageCustomOutput = InCustomOutputToCompile;
-			break;
+
+		// NOTE: the following cases for custom data 0/1 and custom output are hacks since a proper solution requires engine changes:
+		case MP_CustomData0: ResourceId.Usage = static_cast<EMaterialShaderMapUsage::Type>(EMaterialShaderMapUsage::DebugViewMode + 1); break;
+		case MP_CustomData1: ResourceId.Usage = static_cast<EMaterialShaderMapUsage::Type>(EMaterialShaderMapUsage::DebugViewMode + 2); break;
+		case MP_CustomOutput: ResourceId.Usage = static_cast<EMaterialShaderMapUsage::Type>(EMaterialShaderMapUsage::DebugViewMode + 3); break;
+
 		default:
 			ensureMsgf(false, TEXT("ExportMaterial has no usage for property %i.  Will likely reuse the normal rendering shader and crash later with a parameter mismatch"), (int32)InPropertyToCompile);
 			break;
@@ -622,9 +622,25 @@ public:
 	}
 
 private:
+	static FGuid GetCustomAttributeID(const FString& AttributeName)
+	{
+		TArray<FMaterialCustomOutputAttributeDefintion> CustomAttributes;
+		FMaterialAttributeDefinitionMap::GetCustomAttributeList(CustomAttributes);
+
+		for (auto& Attribute : CustomAttributes)
+		{
+			if (Attribute.AttributeName == AttributeName)
+			{
+				return Attribute.AttributeID;
+			}
+		}
+
+		return {};
+	}
+
 	int32 CompileInputForCustomOutput(FMaterialCompiler* Compiler, int32 InputIndex, uint32 ForceCastFlags) const
 	{
-		FGuid AttributeID = FMaterialAttributeDefinitionMap::GetCustomAttributeID(CustomOutputToCompile);
+		FGuid AttributeID = GetCustomAttributeID(CustomOutputToCompile);
 		check(AttributeID.IsValid());
 
 		UMaterialExpressionCustomOutput* Expression = GetCustomOutputExpressionToCompile();
