[GLTFExporter] refactor: remove clear coat hack for material baking

commit 77455a834d5d351c75c6de93230fc2161e2d12c6
Date: Wed, 4 Nov 2020 22:46:33 +0100

    refactor: remove clear coat hack for material baking
    (cherry picked from commit 85296a643046f2ab7c7c6bc3ec913c47c1912044​)

#skipci

diff --git a/Engine/Plugins/Enterprise/GLTFExporter/MigratedFromGit/Source/GLTFExporter/Private/Converters/GLTFMaterialUtility.cpp b/Engine/Plugins/Enterprise/GLTFExporter/MigratedFromGit/Source/GLTFExporter/Private/Converters/GLTFMaterialUtility.cpp
--- a/Engine/Plugins/Enterprise/GLTFExporter/MigratedFromGit/Source/GLTFExporter/Private/Converters/GLTFMaterialUtility.cpp
+++ b/Engine/Plugins/Enterprise/GLTFExporter/MigratedFromGit/Source/GLTFExporter/Private/Converters/GLTFMaterialUtility.cpp
@@ -243,32 +243,7 @@ bool FGLTFMaterialUtility::CombineTextures(TArray<FColor>& OutPixels, const TArr
 
 FGLTFPropertyBakeOutput FGLTFMaterialUtility::BakeMaterialProperty(const FIntPoint& OutputSize, EMaterialProperty Property, const UMaterialInterface* Material, int32 TexCoord, const FMeshDescription* MeshDescription, const TArray<int32>& MeshSectionIndices, bool bCopyAlphaFromRedChannel)
 {
-	EMaterialProperty HackProperty = MP_MAX;
-	TTuple<int32, UMaterialExpression*> HackPrevState;
-
-	// TODO: replace this hack by adding proper support for ClearCoat properties in MaterialBaking module
-	if (Property == MP_CustomData0 || Property == MP_CustomData1 || Property == MP_CustomOutput)
-	{
-		switch (Property)
-		{
-			case MP_CustomData0: HackProperty = MP_Metallic; break;
-			case MP_CustomData1: HackProperty = MP_Roughness; break;
-			case MP_CustomOutput: HackProperty = MP_Normal; break;
-			default: checkNoEntry(); break;
-		}
-
-		UMaterial* ParentMaterial = const_cast<UMaterial*>(Material->GetMaterial());
-		ParentMaterial->Modify();
-		{
-			const FExpressionInput& Input = *GetInputForProperty(Material, Property);
-			FExpressionInput& HackInput = *const_cast<FExpressionInput*>(GetInputForProperty(Material, HackProperty));
-			HackPrevState = MakeTuple(HackInput.OutputIndex, HackInput.Expression);
-			HackInput.Connect(Input.OutputIndex, Input.Expression);
-		}
-		ParentMaterial->PostEditChange();
-
-		Property = HackProperty;
-	}
+	const FMaterialPropertyEx PropertyEx(Property, Property == MP_CustomOutput ? FName(TEXT("ClearCoatBottomNormal")) : NAME_None);
 
 	FMeshData MeshSet;
 	MeshSet.TextureCoordinateBox = { { 0.0f, 0.0f }, { 1.0f, 1.0f } };
@@ -276,24 +251,24 @@ FGLTFPropertyBakeOutput FGLTFMaterialUtility::BakeMaterialProperty(const FIntPoi
 	MeshSet.RawMeshDescription = const_cast<FMeshDescription*>(MeshDescription);
 	MeshSet.MaterialIndices = MeshSectionIndices; // NOTE: MaterialIndices is actually section indices
 
-	FMaterialData MatSet;
+	FMaterialDataEx MatSet;
 	MatSet.Material = const_cast<UMaterialInterface*>(Material);
-	MatSet.PropertySizes.Add(Property, OutputSize);
+	MatSet.PropertySizes.Add(PropertyEx, OutputSize);
 
 	TArray<FMeshData*> MeshSettings;
-	TArray<FMaterialData*> MatSettings;
+	TArray<FMaterialDataEx*> MatSettings;
 	MeshSettings.Add(&MeshSet);
 	MatSettings.Add(&MatSet);
 
-	TArray<FBakeOutput> BakeOutputs;
+	TArray<FBakeOutputEx> BakeOutputs;
 	IMaterialBakingModule& Module = FModuleManager::Get().LoadModuleChecked<IMaterialBakingModule>("GLTFMaterialBaking");
 
 	Module.BakeMaterials(MatSettings, MeshSettings, BakeOutputs);
 
-	FBakeOutput& BakeOutput = BakeOutputs[0];
+	FBakeOutputEx& BakeOutput = BakeOutputs[0];
 
-	TArray<FColor> BakedPixels = MoveTemp(BakeOutput.PropertyData.FindChecked(Property));
-	const FIntPoint BakedSize = BakeOutput.PropertySizes.FindChecked(Property);
+	TArray<FColor> BakedPixels = MoveTemp(BakeOutput.PropertyData.FindChecked(PropertyEx));
+	const FIntPoint BakedSize = BakeOutput.PropertySizes.FindChecked(PropertyEx);
 	const float EmissiveScale = BakeOutput.EmissiveScale;
 
 	if (bCopyAlphaFromRedChannel)
@@ -313,8 +288,7 @@ FGLTFPropertyBakeOutput FGLTFMaterialUtility::BakeMaterialProperty(const FIntPoi
 		}
 	}
 
-	// TODO: add condition for ClearCoat normal property when it's been added to MaterialBaking module
-	if (Property == MP_Normal)
+	if (Property == MP_Normal || Property == MP_CustomOutput)
 	{
 		// Convert normalmaps to use +Y (OpenGL / WebGL standard)
 		FGLTFTextureUtility::FlipGreenChannel(BakedPixels);
@@ -333,18 +307,6 @@ FGLTFPropertyBakeOutput FGLTFMaterialUtility::BakeMaterialProperty(const FIntPoi
 		PropertyBakeOutput.ConstantValue = Pixel;
 	}
 
-	// TODO: replace this hack by adding proper support for ClearCoat properties in MaterialBaking module
-	if (HackProperty != MP_MAX)
-	{
-		UMaterial* ParentMaterial = const_cast<UMaterial*>(Material->GetMaterial());
-		ParentMaterial->Modify();
-		{
-			FExpressionInput& HackInput = *const_cast<FExpressionInput*>(GetInputForProperty(Material, HackProperty));
-			HackInput.Connect(HackPrevState.Get<0>(), HackPrevState.Get<1>());
-		}
-		ParentMaterial->PostEditChange();
-	}
-
 	return PropertyBakeOutput;
 }
 
