[GLTFExporter] feat: add GLTFExportOptions::TextureCompression

commit 4c9d626ec466e674a0be56af94177829c34f7e30
Date: Sun, 8 Nov 2020 17:34:20 +0100

    feat: add GLTFExportOptions::TextureCompression
    (cherry picked from commit b8f688494d738ac1a3e9523cc174657e54f66912​)

#skipci

diff --git a/Engine/Plugins/Enterprise/GLTFExporter/MigratedFromGit/Source/GLTFExporter/Private/GLTFExportOptions.cpp b/Engine/Plugins/Enterprise/GLTFExporter/MigratedFromGit/Source/GLTFExporter/Private/GLTFExportOptions.cpp
--- a/Engine/Plugins/Enterprise/GLTFExporter/MigratedFromGit/Source/GLTFExporter/Private/GLTFExportOptions.cpp
+++ b/Engine/Plugins/Enterprise/GLTFExporter/MigratedFromGit/Source/GLTFExporter/Private/GLTFExportOptions.cpp
@@ -28,6 +28,7 @@ void UGLTFExportOptions::ResetToDefault()
 	bExportAnimationSequences = true;
 	bRetargetBoneTransforms = true;
 	bExportPlaybackSettings = true;
+	TextureCompression = EGLTFExporterTextureCompression::PNG;
 	bExportTextureTransforms = true;
 	bExportLightmaps = true;
 	TextureHDREncoding = EGLTFExporterTextureHDREncoding::RGBM;
diff --git a/Engine/Plugins/Enterprise/GLTFExporter/MigratedFromGit/Source/GLTFExporter/Public/GLTFExportOptions.h b/Engine/Plugins/Enterprise/GLTFExporter/MigratedFromGit/Source/GLTFExporter/Public/GLTFExportOptions.h
--- a/Engine/Plugins/Enterprise/GLTFExporter/MigratedFromGit/Source/GLTFExporter/Public/GLTFExportOptions.h
+++ b/Engine/Plugins/Enterprise/GLTFExporter/MigratedFromGit/Source/GLTFExporter/Public/GLTFExportOptions.h
@@ -24,6 +24,13 @@ enum class EGLTFExporterMaterialBakeSize : uint8
 	POT_8192 UMETA(DisplayName = "8192 x 8192")
 };
 
+UENUM(BlueprintType)
+enum class EGLTFExporterTextureCompression : uint8
+{
+	PNG UMETA(DisplayName = "PNG"),
+    JPEG UMETA(DisplayName = "JPEG (if no alpha)")
+};
+
 UENUM(BlueprintType)
 enum class EGLTFExporterTextureHDREncoding : uint8
 {
@@ -89,6 +96,9 @@ public:
 	UPROPERTY(EditAnywhere, BlueprintReadWrite, Config, Category = Animation)
 	bool bExportPlaybackSettings;
 
+	UPROPERTY(EditAnywhere, BlueprintReadWrite, Config, Category = Texture)
+	EGLTFExporterTextureCompression TextureCompression;
+
 	UPROPERTY(EditAnywhere, BlueprintReadWrite, Config, Category = Texture)
 	bool bExportTextureTransforms;
 
