[GLTFExporter] refactor: simplify GLTFImageBuilder by limiting lightmap texture formats to BGRA8

commit 94381ed0ecf03f56bd788cdf9bd44cd9c11ef808
Date: Sun, 8 Nov 2020 17:45:20 +0100

    refactor: simplify GLTFImageBuilder by limiting lightmap texture formats to BGRA8
    (cherry picked from commit 0bd1bc5852f4be8378173558c8e6a14b6bbe84b4​)

#skipci

diff --git a/Engine/Plugins/Enterprise/GLTFExporter/MigratedFromGit/Source/GLTFExporter/Private/Builders/GLTFImageBuilder.cpp b/Engine/Plugins/Enterprise/GLTFExporter/MigratedFromGit/Source/GLTFExporter/Private/Builders/GLTFImageBuilder.cpp
--- a/Engine/Plugins/Enterprise/GLTFExporter/MigratedFromGit/Source/GLTFExporter/Private/Builders/GLTFImageBuilder.cpp
+++ b/Engine/Plugins/Enterprise/GLTFExporter/MigratedFromGit/Source/GLTFExporter/Private/Builders/GLTFImageBuilder.cpp
@@ -11,6 +11,40 @@ FGLTFImageBuilder::FGLTFImageBuilder(const FString& FilePath, const UGLTFExportO
 {
 }
 
+FGLTFJsonImageIndex FGLTFImageBuilder::AddImage(const TArray<FColor>& Pixels, FIntPoint Size, const FString& Name)
+{
+	check(Pixels.Num() == Size.X * Size.Y);
+	return AddImage(Pixels.GetData(), Size, Name);
+}
+
+FGLTFJsonImageIndex FGLTFImageBuilder::AddImage(const FColor* Pixels, int64 ByteLength, FIntPoint Size, const FString& Name)
+{
+	check(ByteLength == Size.X * Size.Y * sizeof(FColor));
+	return AddImage(Pixels, Size, Name);
+}
+
+FGLTFJsonImageIndex FGLTFImageBuilder::AddImage(const FColor* Pixels, FIntPoint Size, const FString& Name)
+{
+	IImageWrapperModule& ImageWrapperModule = FModuleManager::Get().LoadModuleChecked<IImageWrapperModule>(TEXT("ImageWrapper"));
+	TSharedPtr<IImageWrapper> ImageWrapper = ImageWrapperModule.CreateImageWrapper(EImageFormat::PNG);
+
+	if (!ImageWrapper.IsValid())
+	{
+		// TODO: report error
+		return FGLTFJsonImageIndex(INDEX_NONE);
+	}
+
+	const int64 ByteLength = Size.X * Size.Y * sizeof(FColor);
+	if (!ImageWrapper->SetRaw(Pixels, ByteLength, Size.X, Size.Y, ERGBFormat::BGRA, 8))
+	{
+		// TODO: report error
+		return FGLTFJsonImageIndex(INDEX_NONE);
+	}
+
+	const TArray64<uint8>& ImageData = ImageWrapper->GetCompressed();
+	return AddImage(ImageData.GetData(), ImageData.Num(), EGLTFJsonMimeType::PNG, Name);
+}
+
 FGLTFJsonImageIndex FGLTFImageBuilder::AddImage(const void* CompressedData, int64 CompressedByteLength, EGLTFJsonMimeType MimeType, const FString& Name)
 {
 	// TODO: should this function be renamed to GetOrAddImage?
@@ -39,36 +73,6 @@ FGLTFJsonImageIndex FGLTFImageBuilder::AddImage(const void* CompressedData, int6
 	return ImageIndex;
 }
 
-FGLTFJsonImageIndex FGLTFImageBuilder::AddImage(const void* RawData, int64 ByteLength, FIntPoint Size, ERGBFormat Format, int32 BitDepth, const FString& Name)
-{
-	IImageWrapperModule& ImageWrapperModule = FModuleManager::Get().LoadModuleChecked<IImageWrapperModule>(TEXT("ImageWrapper"));
-	TSharedPtr<IImageWrapper> ImageWrapper = ImageWrapperModule.CreateImageWrapper(EImageFormat::PNG);
-
-	if (!ImageWrapper.IsValid())
-	{
-		// TODO: report error
-		return FGLTFJsonImageIndex(INDEX_NONE);
-	}
-
-	const bool bFormatSupported = ImageWrapper->SetRaw(RawData, ByteLength, Size.X, Size.Y, Format, BitDepth);
-
-	if (!bFormatSupported)
-	{
-		// TODO: report error
-		return FGLTFJsonImageIndex(INDEX_NONE);
-	}
-
-	const TArray64<uint8>& ImageData = ImageWrapper->GetCompressed();
-	return AddImage(ImageData.GetData(), ImageData.Num(), EGLTFJsonMimeType::PNG, Name);
-}
-
-FGLTFJsonImageIndex FGLTFImageBuilder::AddImage(const TArray<FColor>& Pixels, FIntPoint Size, const FString& Name)
-{
-	check(Pixels.Num() == Size.X * Size.Y);
-	const int64 ByteLength = Size.X * Size.Y * sizeof(FColor);
-	return AddImage(Pixels.GetData(), ByteLength, Size, ERGBFormat::BGRA, 8, Name);
-}
-
 FString FGLTFImageBuilder::SaveImageToFile(const void* CompressedData, int64 CompressedByteLength, EGLTFJsonMimeType MimeType, const FString& Name)
 {
 	const TCHAR* Extension = FGLTFBuilderUtility::GetFileExtension(MimeType);
diff --git a/Engine/Plugins/Enterprise/GLTFExporter/MigratedFromGit/Source/GLTFExporter/Private/Builders/GLTFImageBuilder.h b/Engine/Plugins/Enterprise/GLTFExporter/MigratedFromGit/Source/GLTFExporter/Private/Builders/GLTFImageBuilder.h
--- a/Engine/Plugins/Enterprise/GLTFExporter/MigratedFromGit/Source/GLTFExporter/Private/Builders/GLTFImageBuilder.h
+++ b/Engine/Plugins/Enterprise/GLTFExporter/MigratedFromGit/Source/GLTFExporter/Private/Builders/GLTFImageBuilder.h
@@ -5,8 +5,6 @@
 #include "Builders/GLTFBufferBuilder.h"
 #include "Builders/GLTFBinaryHashKey.h"
 
-enum class ERGBFormat : int8;
-
 class FGLTFImageBuilder : public FGLTFBufferBuilder
 {
 protected:
@@ -15,13 +13,14 @@ protected:
 
 public:
 
-	FGLTFJsonImageIndex AddImage(const void* CompressedData, int64 CompressedByteLength, EGLTFJsonMimeType MimeType, const FString& Name);
-	FGLTFJsonImageIndex AddImage(const void* RawData, int64 ByteLength, FIntPoint Size, ERGBFormat Format, int32 BitDepth, const FString& Name);
-
 	FGLTFJsonImageIndex AddImage(const TArray<FColor>& Pixels, FIntPoint Size, const FString& Name);
+	FGLTFJsonImageIndex AddImage(const FColor* Pixels, int64 ByteLength, FIntPoint Size, const FString& Name);
 
 private:
 
+	FGLTFJsonImageIndex AddImage(const FColor* Pixels, FIntPoint Size, const FString& Name);
+	FGLTFJsonImageIndex AddImage(const void* CompressedData, int64 CompressedByteLength, EGLTFJsonMimeType MimeType, const FString& Name);
+
 	FString SaveImageToFile(const void* CompressedData, int64 CompressedByteLength, EGLTFJsonMimeType MimeType, const FString& Name);
 
 	TSet<FString> UniqueImageUris;
diff --git a/Engine/Plugins/Enterprise/GLTFExporter/MigratedFromGit/Source/GLTFExporter/Private/Converters/GLTFTextureUtility.cpp b/Engine/Plugins/Enterprise/GLTFExporter/MigratedFromGit/Source/GLTFExporter/Private/Converters/GLTFTextureUtility.cpp
--- a/Engine/Plugins/Enterprise/GLTFExporter/MigratedFromGit/Source/GLTFExporter/Private/Converters/GLTFTextureUtility.cpp
+++ b/Engine/Plugins/Enterprise/GLTFExporter/MigratedFromGit/Source/GLTFExporter/Private/Converters/GLTFTextureUtility.cpp
@@ -1,7 +1,6 @@
 // Copyright Epic Games, Inc. All Rights Reserved.
 
 #include "Converters/GLTFTextureUtility.h"
-#include "IImageWrapper.h"
 #include "NormalMapPreview.h"
 
 bool FGLTFTextureUtility::IsHDR(EPixelFormat Format)
@@ -28,35 +27,6 @@ bool FGLTFTextureUtility::IsHDR(EPixelFormat Format)
 	}
 }
 
-bool FGLTFTextureUtility::CanPNGCompressFormat(ETextureSourceFormat InFormat, ERGBFormat& OutFormat, uint32& OutBitDepth)
-{
-	switch (InFormat)
-	{
-		case TSF_BGRA8:   OutFormat = ERGBFormat::BGRA; OutBitDepth = 8;  return true;
-		case TSF_RGBA8:   OutFormat = ERGBFormat::RGBA; OutBitDepth = 8;  return true;
-		case TSF_RGBA16:  OutFormat = ERGBFormat::RGBA; OutBitDepth = 16; return true;
-		case TSF_G8:      OutFormat = ERGBFormat::Gray; OutBitDepth = 8;  return true;
-		case TSF_G16:     OutFormat = ERGBFormat::Gray; OutBitDepth = 16; return true;
-		default:
-			return false;
-	}
-}
-
-bool FGLTFTextureUtility::CanPNGCompressFormat(EPixelFormat InFormat, ERGBFormat& OutFormat, uint32& OutBitDepth)
-{
-	switch (InFormat)
-	{
-		case PF_B8G8R8A8:           OutFormat = ERGBFormat::BGRA; OutBitDepth = 8;  return true;
-		case PF_R8G8B8A8:           OutFormat = ERGBFormat::RGBA; OutBitDepth = 8;  return true;
-		case PF_R16G16B16A16_UNORM: OutFormat = ERGBFormat::RGBA; OutBitDepth = 16; return true;
-		case PF_L8:                 OutFormat = ERGBFormat::Gray; OutBitDepth = 8;  return true;
-		case PF_G8:                 OutFormat = ERGBFormat::Gray; OutBitDepth = 8;  return true;
-		case PF_G16:                OutFormat = ERGBFormat::Gray; OutBitDepth = 16; return true;
-		default:
-			return false;
-	}
-}
-
 bool FGLTFTextureUtility::IsCubemap(const UTexture* Texture)
 {
 	return Texture->IsA<UTextureCube>() || Texture->IsA<UTextureRenderTargetCube>();
diff --git a/Engine/Plugins/Enterprise/GLTFExporter/MigratedFromGit/Source/GLTFExporter/Private/Converters/GLTFTextureUtility.h b/Engine/Plugins/Enterprise/GLTFExporter/MigratedFromGit/Source/GLTFExporter/Private/Converters/GLTFTextureUtility.h
--- a/Engine/Plugins/Enterprise/GLTFExporter/MigratedFromGit/Source/GLTFExporter/Private/Converters/GLTFTextureUtility.h
+++ b/Engine/Plugins/Enterprise/GLTFExporter/MigratedFromGit/Source/GLTFExporter/Private/Converters/GLTFTextureUtility.h
@@ -5,15 +5,9 @@
 #include "Json/GLTFJsonEnums.h"
 #include "Engine.h"
 
-enum class ERGBFormat : int8;
-
 struct FGLTFTextureUtility
 {
 	static bool IsHDR(EPixelFormat Format);
-
-	static bool CanPNGCompressFormat(ETextureSourceFormat InFormat, ERGBFormat& OutFormat, uint32& OutBitDepth);
-	static bool CanPNGCompressFormat(EPixelFormat InFormat, ERGBFormat& OutFormat, uint32& OutBitDepth);
-
 	static bool IsCubemap(const UTexture* Texture);
 
 	static float GetCubeFaceRotation(ECubeFace CubeFace);
diff --git a/Engine/Plugins/Enterprise/GLTFExporter/MigratedFromGit/Source/GLTFExporter/Private/Tasks/GLTFTextureTasks.cpp b/Engine/Plugins/Enterprise/GLTFExporter/MigratedFromGit/Source/GLTFExporter/Private/Tasks/GLTFTextureTasks.cpp
--- a/Engine/Plugins/Enterprise/GLTFExporter/MigratedFromGit/Source/GLTFExporter/Private/Tasks/GLTFTextureTasks.cpp
+++ b/Engine/Plugins/Enterprise/GLTFExporter/MigratedFromGit/Source/GLTFExporter/Private/Tasks/GLTFTextureTasks.cpp
@@ -153,18 +153,18 @@ void FGLTFTextureLightMapTask::Complete()
 		return;
 	}
 
-	ERGBFormat RGBFormat;
-	uint32 BitDepth;
-	if (!FGLTFTextureUtility::CanPNGCompressFormat(Source.GetFormat(), RGBFormat, BitDepth))
+	const ETextureSourceFormat SourceFormat = Source.GetFormat();
+	if (SourceFormat != TSF_BGRA8)
 	{
-		Builder.AddWarningMessage(FString::Printf(TEXT("Failed to export lightmap texture %s because of unsupported format %s"), *JsonTexture.Name, *FGLTFNameUtility::GetName(Source.GetFormat())));
+		Builder.AddWarningMessage(FString::Printf(TEXT("Failed to export lightmap texture %s because of unsupported source format %s"), *JsonTexture.Name, *FGLTFNameUtility::GetName(SourceFormat)));
 		return;
 	}
 
 	const FIntPoint Size = { Source.GetSizeX(), Source.GetSizeY() };
+	const int64 ByteLength = Source.CalcMipSize(0);
 
 	const void* RawData = Source.LockMip(0);
-	JsonTexture.Source = Builder.AddImage(RawData, Source.CalcMipSize(0), Size, RGBFormat, BitDepth, JsonTexture.Name);
+	JsonTexture.Source = Builder.AddImage(static_cast<const FColor*>(RawData), ByteLength, Size, JsonTexture.Name);
 	Source.UnlockMip(0);
 
 	JsonTexture.Sampler = Builder.GetOrAddSampler(LightMap);
