[GLTFExporter] refactor: detect if texture format has no alpha channel and pass information to GLTFImageBuilder

commit 6dd425bddcd2253128a0a7409f111902c6e29d67
Date: Sun, 8 Nov 2020 17:58:07 +0100

    refactor: detect if texture format has no alpha channel and pass information to GLTFImageBuilder
    (cherry picked from commit 237286a8f915f88644628bba0f8c1e35dbd853f1​)

#skipci

diff --git a/Engine/Plugins/Enterprise/GLTFExporter/MigratedFromGit/Source/GLTFExporter/Private/Builders/GLTFImageBuilder.cpp b/Engine/Plugins/Enterprise/GLTFExporter/MigratedFromGit/Source/GLTFExporter/Private/Builders/GLTFImageBuilder.cpp
--- a/Engine/Plugins/Enterprise/GLTFExporter/MigratedFromGit/Source/GLTFExporter/Private/Builders/GLTFImageBuilder.cpp
+++ b/Engine/Plugins/Enterprise/GLTFExporter/MigratedFromGit/Source/GLTFExporter/Private/Builders/GLTFImageBuilder.cpp
@@ -11,19 +11,19 @@ FGLTFImageBuilder::FGLTFImageBuilder(const FString& FilePath, const UGLTFExportO
 {
 }
 
-FGLTFJsonImageIndex FGLTFImageBuilder::AddImage(const TArray<FColor>& Pixels, FIntPoint Size, const FString& Name)
+FGLTFJsonImageIndex FGLTFImageBuilder::AddImage(const TArray<FColor>& Pixels, FIntPoint Size, bool bIgnoreAlpha, const FString& Name)
 {
 	check(Pixels.Num() == Size.X * Size.Y);
-	return AddImage(Pixels.GetData(), Size, Name);
+	return AddImage(Pixels.GetData(), Size, bIgnoreAlpha, Name);
 }
 
-FGLTFJsonImageIndex FGLTFImageBuilder::AddImage(const FColor* Pixels, int64 ByteLength, FIntPoint Size, const FString& Name)
+FGLTFJsonImageIndex FGLTFImageBuilder::AddImage(const FColor* Pixels, int64 ByteLength, FIntPoint Size, bool bIgnoreAlpha, const FString& Name)
 {
 	check(ByteLength == Size.X * Size.Y * sizeof(FColor));
-	return AddImage(Pixels, Size, Name);
+	return AddImage(Pixels, Size, bIgnoreAlpha, Name);
 }
 
-FGLTFJsonImageIndex FGLTFImageBuilder::AddImage(const FColor* Pixels, FIntPoint Size, const FString& Name)
+FGLTFJsonImageIndex FGLTFImageBuilder::AddImage(const FColor* Pixels, FIntPoint Size, bool bIgnoreAlpha, const FString& Name)
 {
 	IImageWrapperModule& ImageWrapperModule = FModuleManager::Get().LoadModuleChecked<IImageWrapperModule>(TEXT("ImageWrapper"));
 	TSharedPtr<IImageWrapper> ImageWrapper = ImageWrapperModule.CreateImageWrapper(EImageFormat::PNG);
diff --git a/Engine/Plugins/Enterprise/GLTFExporter/MigratedFromGit/Source/GLTFExporter/Private/Builders/GLTFImageBuilder.h b/Engine/Plugins/Enterprise/GLTFExporter/MigratedFromGit/Source/GLTFExporter/Private/Builders/GLTFImageBuilder.h
--- a/Engine/Plugins/Enterprise/GLTFExporter/MigratedFromGit/Source/GLTFExporter/Private/Builders/GLTFImageBuilder.h
+++ b/Engine/Plugins/Enterprise/GLTFExporter/MigratedFromGit/Source/GLTFExporter/Private/Builders/GLTFImageBuilder.h
@@ -13,12 +13,12 @@ protected:
 
 public:
 
-	FGLTFJsonImageIndex AddImage(const TArray<FColor>& Pixels, FIntPoint Size, const FString& Name);
-	FGLTFJsonImageIndex AddImage(const FColor* Pixels, int64 ByteLength, FIntPoint Size, const FString& Name);
+	FGLTFJsonImageIndex AddImage(const TArray<FColor>& Pixels, FIntPoint Size, bool bIgnoreAlpha, const FString& Name);
+	FGLTFJsonImageIndex AddImage(const FColor* Pixels, int64 ByteLength, FIntPoint Size, bool bIgnoreAlpha, const FString& Name);
 
 private:
 
-	FGLTFJsonImageIndex AddImage(const FColor* Pixels, FIntPoint Size, const FString& Name);
+	FGLTFJsonImageIndex AddImage(const FColor* Pixels, FIntPoint Size, bool bIgnoreAlpha, const FString& Name);
 	FGLTFJsonImageIndex AddImage(const void* CompressedData, int64 CompressedByteLength, EGLTFJsonMimeType MimeType, const FString& Name);
 
 	FString SaveImageToFile(const void* CompressedData, int64 CompressedByteLength, EGLTFJsonMimeType MimeType, const FString& Name);
diff --git a/Engine/Plugins/Enterprise/GLTFExporter/MigratedFromGit/Source/GLTFExporter/Private/Converters/GLTFMaterialUtility.cpp b/Engine/Plugins/Enterprise/GLTFExporter/MigratedFromGit/Source/GLTFExporter/Private/Converters/GLTFMaterialUtility.cpp
--- a/Engine/Plugins/Enterprise/GLTFExporter/MigratedFromGit/Source/GLTFExporter/Private/Converters/GLTFMaterialUtility.cpp
+++ b/Engine/Plugins/Enterprise/GLTFExporter/MigratedFromGit/Source/GLTFExporter/Private/Converters/GLTFMaterialUtility.cpp
@@ -339,7 +339,7 @@ FGLTFJsonTextureIndex FGLTFMaterialUtility::AddTexture(FGLTFConvertBuilder& Buil
 	FGLTFJsonTexture JsonTexture;
 	JsonTexture.Name = TextureName;
 	JsonTexture.Sampler = Builder.AddSampler(JsonSampler);
-	JsonTexture.Source = Builder.AddImage(Pixels, TextureSize, TextureName);
+	JsonTexture.Source = Builder.AddImage(Pixels, TextureSize, false, TextureName);
 
 	return Builder.AddTexture(JsonTexture);
 }
diff --git a/Engine/Plugins/Enterprise/GLTFExporter/MigratedFromGit/Source/GLTFExporter/Private/Converters/GLTFTextureUtility.cpp b/Engine/Plugins/Enterprise/GLTFExporter/MigratedFromGit/Source/GLTFExporter/Private/Converters/GLTFTextureUtility.cpp
--- a/Engine/Plugins/Enterprise/GLTFExporter/MigratedFromGit/Source/GLTFExporter/Private/Converters/GLTFTextureUtility.cpp
+++ b/Engine/Plugins/Enterprise/GLTFExporter/MigratedFromGit/Source/GLTFExporter/Private/Converters/GLTFTextureUtility.cpp
@@ -27,6 +27,26 @@ bool FGLTFTextureUtility::IsHDR(EPixelFormat Format)
 	}
 }
 
+bool FGLTFTextureUtility::IsAlphaless(EPixelFormat PixelFormat)
+{
+	switch (PixelFormat)
+	{
+		case PF_ATC_RGB:
+		case PF_BC4:
+		case PF_BC5:
+		case PF_DXT1:
+		case PF_ETC1:
+		case PF_ETC2_RGB:
+		case PF_FloatR11G11B10:
+		case PF_FloatRGB:
+		case PF_R5G6B5_UNORM:
+			// TODO: add more pixel formats that don't support alpha, but beware of formats like PF_G8 (that still seem to return alpha in some cases)
+			return true;
+		default:
+			return false;
+	}
+}
+
 bool FGLTFTextureUtility::IsCubemap(const UTexture* Texture)
 {
 	return Texture->IsA<UTextureCube>() || Texture->IsA<UTextureRenderTargetCube>();
@@ -340,14 +360,6 @@ bool FGLTFTextureUtility::LoadPlatformData(UTextureCube* TextureCube)
 	return true;
 }
 
-void FGLTFTextureUtility::FlipGreenChannel(TArray<FLinearColor>& Pixels)
-{
-	for (FLinearColor& Pixel: Pixels)
-	{
-		Pixel.G = 1.0f - FMath::Clamp(Pixel.G, 0.0f, 1.0f);
-	}
-}
-
 void FGLTFTextureUtility::FlipGreenChannel(TArray<FColor>& Pixels)
 {
 	for (FColor& Pixel: Pixels)
diff --git a/Engine/Plugins/Enterprise/GLTFExporter/MigratedFromGit/Source/GLTFExporter/Private/Converters/GLTFTextureUtility.h b/Engine/Plugins/Enterprise/GLTFExporter/MigratedFromGit/Source/GLTFExporter/Private/Converters/GLTFTextureUtility.h
--- a/Engine/Plugins/Enterprise/GLTFExporter/MigratedFromGit/Source/GLTFExporter/Private/Converters/GLTFTextureUtility.h
+++ b/Engine/Plugins/Enterprise/GLTFExporter/MigratedFromGit/Source/GLTFExporter/Private/Converters/GLTFTextureUtility.h
@@ -8,6 +8,8 @@
 struct FGLTFTextureUtility
 {
 	static bool IsHDR(EPixelFormat Format);
+	static bool IsAlphaless(EPixelFormat PixelFormat);
+
 	static bool IsCubemap(const UTexture* Texture);
 
 	static float GetCubeFaceRotation(ECubeFace CubeFace);
@@ -35,6 +37,8 @@ struct FGLTFTextureUtility
 	static bool LoadPlatformData(UTexture2D* Texture);
 	static bool LoadPlatformData(UTextureCube* TextureCube);
 
-	static void FlipGreenChannel(TArray<FLinearColor>& Pixels);
 	static void FlipGreenChannel(TArray<FColor>& Pixels);
+
+	static bool HasAlphaChannel(const UTexture* Texture);
+	static bool HasAlphaChannel(EPixelFormat PixelFormat);
 };
diff --git a/Engine/Plugins/Enterprise/GLTFExporter/MigratedFromGit/Source/GLTFExporter/Private/Tasks/GLTFTextureTasks.cpp b/Engine/Plugins/Enterprise/GLTFExporter/MigratedFromGit/Source/GLTFExporter/Private/Tasks/GLTFTextureTasks.cpp
--- a/Engine/Plugins/Enterprise/GLTFExporter/MigratedFromGit/Source/GLTFExporter/Private/Tasks/GLTFTextureTasks.cpp
+++ b/Engine/Plugins/Enterprise/GLTFExporter/MigratedFromGit/Source/GLTFExporter/Private/Tasks/GLTFTextureTasks.cpp
@@ -35,7 +35,8 @@ void FGLTFTexture2DTask::Complete()
 		FGLTFTextureUtility::FlipGreenChannel(Pixels);
 	}
 
-	JsonTexture.Source = Builder.AddImage(Pixels, Size, JsonTexture.Name);
+	const bool bIgnoreAlpha = FGLTFTextureUtility::IsAlphaless(Texture2D->GetPixelFormat());
+	JsonTexture.Source = Builder.AddImage(Pixels, Size, bIgnoreAlpha, JsonTexture.Name);
 	JsonTexture.Sampler = Builder.GetOrAddSampler(Texture2D);
 }
 
@@ -73,7 +74,8 @@ void FGLTFTextureCubeTask::Complete()
 		return;
 	}
 
-	JsonTexture.Source = Builder.AddImage(Pixels, Size, JsonTexture.Name);
+	const bool bIgnoreAlpha = FGLTFTextureUtility::IsAlphaless(TextureCube->GetPixelFormat());
+	JsonTexture.Source = Builder.AddImage(Pixels, Size, bIgnoreAlpha, JsonTexture.Name);
 	JsonTexture.Sampler = Builder.GetOrAddSampler(TextureCube);
 }
 
@@ -97,7 +99,8 @@ void FGLTFTextureRenderTarget2DTask::Complete()
 		return;
 	}
 
-	JsonTexture.Source = Builder.AddImage(Pixels, Size, JsonTexture.Name);
+	const bool bIgnoreAlpha = FGLTFTextureUtility::IsAlphaless(RenderTarget2D->GetFormat());
+	JsonTexture.Source = Builder.AddImage(Pixels, Size, bIgnoreAlpha, JsonTexture.Name);
 	JsonTexture.Sampler = Builder.GetOrAddSampler(RenderTarget2D);
 }
 
@@ -134,7 +137,8 @@ void FGLTFTextureRenderTargetCubeTask::Complete()
 		return;
 	}
 
-	JsonTexture.Source = Builder.AddImage(Pixels, Size, JsonTexture.Name);
+	const bool bIgnoreAlpha = FGLTFTextureUtility::IsAlphaless(RenderTargetCube->GetFormat());
+	JsonTexture.Source = Builder.AddImage(Pixels, Size, bIgnoreAlpha, JsonTexture.Name);
 	JsonTexture.Sampler = Builder.GetOrAddSampler(RenderTargetCube);
 }
 
@@ -164,7 +168,7 @@ void FGLTFTextureLightMapTask::Complete()
 	const int64 ByteLength = Source.CalcMipSize(0);
 
 	const void* RawData = Source.LockMip(0);
-	JsonTexture.Source = Builder.AddImage(static_cast<const FColor*>(RawData), ByteLength, Size, JsonTexture.Name);
+	JsonTexture.Source = Builder.AddImage(static_cast<const FColor*>(RawData), ByteLength, Size, false, JsonTexture.Name);
 	Source.UnlockMip(0);
 
 	JsonTexture.Sampler = Builder.GetOrAddSampler(LightMap);
