[GLTFExporter] feat: implement support for JPEG compression and check if alpha not needed

commit 7635f45ebbbad374be051d852e9b54e13e6365cb
Date: Sun, 8 Nov 2020 18:01:47 +0100

    feat: implement support for JPEG compression and check if alpha not needed
    (cherry picked from commit 9678c6591022787c063e88e54132b439be503f4e​)

#skipci

diff --git a/Engine/Plugins/Enterprise/GLTFExporter/MigratedFromGit/Source/GLTFExporter/Private/Builders/GLTFImageBuilder.cpp b/Engine/Plugins/Enterprise/GLTFExporter/MigratedFromGit/Source/GLTFExporter/Private/Builders/GLTFImageBuilder.cpp
--- a/Engine/Plugins/Enterprise/GLTFExporter/MigratedFromGit/Source/GLTFExporter/Private/Builders/GLTFImageBuilder.cpp
+++ b/Engine/Plugins/Enterprise/GLTFExporter/MigratedFromGit/Source/GLTFExporter/Private/Builders/GLTFImageBuilder.cpp
@@ -2,8 +2,7 @@
 
 #include "Builders/GLTFImageBuilder.h"
 #include "Builders/GLTFBuilderUtility.h"
-#include "IImageWrapperModule.h"
-#include "IImageWrapper.h"
+#include "Builders/GLTFImageUtility.h"
 #include "Misc/FileHelper.h"
 
 FGLTFImageBuilder::FGLTFImageBuilder(const FString& FilePath, const UGLTFExportOptions* ExportOptions)
@@ -25,24 +24,25 @@ FGLTFJsonImageIndex FGLTFImageBuilder::AddImage(const FColor* Pixels, int64 Byte
 
 FGLTFJsonImageIndex FGLTFImageBuilder::AddImage(const FColor* Pixels, FIntPoint Size, bool bIgnoreAlpha, const FString& Name)
 {
-	IImageWrapperModule& ImageWrapperModule = FModuleManager::Get().LoadModuleChecked<IImageWrapperModule>(TEXT("ImageWrapper"));
-	TSharedPtr<IImageWrapper> ImageWrapper = ImageWrapperModule.CreateImageWrapper(EImageFormat::PNG);
+	TArray64<uint8> CompressedData;
 
-	if (!ImageWrapper.IsValid())
+	const EGLTFJsonMimeType ImageFormat = GetImageFormat(Pixels, Size, bIgnoreAlpha);
+	switch (ImageFormat)
 	{
-		// TODO: report error
-		return FGLTFJsonImageIndex(INDEX_NONE);
-	}
+		case EGLTFJsonMimeType::PNG:
+			FGLTFImageUtility::CompressToPNG(Pixels, Size, CompressedData);
+			break;
 
-	const int64 ByteLength = Size.X * Size.Y * sizeof(FColor);
-	if (!ImageWrapper->SetRaw(Pixels, ByteLength, Size.X, Size.Y, ERGBFormat::BGRA, 8))
-	{
-		// TODO: report error
-		return FGLTFJsonImageIndex(INDEX_NONE);
+		case EGLTFJsonMimeType::JPEG:
+			FGLTFImageUtility::CompressToJPEG(Pixels, Size, 0, CompressedData);
+			break;
+
+		default:
+			checkNoEntry();
+			break;
 	}
 
-	const TArray64<uint8>& ImageData = ImageWrapper->GetCompressed();
-	return AddImage(ImageData.GetData(), ImageData.Num(), EGLTFJsonMimeType::PNG, Name);
+	return AddImage(CompressedData.GetData(), CompressedData.Num(), ImageFormat, Name);
 }
 
 FGLTFJsonImageIndex FGLTFImageBuilder::AddImage(const void* CompressedData, int64 CompressedByteLength, EGLTFJsonMimeType MimeType, const FString& Name)
@@ -73,6 +73,22 @@ FGLTFJsonImageIndex FGLTFImageBuilder::AddImage(const void* CompressedData, int6
 	return ImageIndex;
 }
 
+EGLTFJsonMimeType FGLTFImageBuilder::GetImageFormat(const FColor* Pixels, FIntPoint Size, bool bIgnoreAlpha) const
+{
+	switch (ExportOptions->TextureCompression)
+	{
+		case EGLTFExporterTextureCompression::PNG:
+			return EGLTFJsonMimeType::PNG;
+
+		case EGLTFExporterTextureCompression::JPEG:
+			return bIgnoreAlpha || FGLTFImageUtility::NoAlphaNeeded(Pixels, Size) ? EGLTFJsonMimeType::JPEG : EGLTFJsonMimeType::PNG;
+
+		default:
+			checkNoEntry();
+			return EGLTFJsonMimeType::None;
+	}
+}
+
 FString FGLTFImageBuilder::SaveImageToFile(const void* CompressedData, int64 CompressedByteLength, EGLTFJsonMimeType MimeType, const FString& Name)
 {
 	const TCHAR* Extension = FGLTFBuilderUtility::GetFileExtension(MimeType);
diff --git a/Engine/Plugins/Enterprise/GLTFExporter/MigratedFromGit/Source/GLTFExporter/Private/Builders/GLTFImageBuilder.h b/Engine/Plugins/Enterprise/GLTFExporter/MigratedFromGit/Source/GLTFExporter/Private/Builders/GLTFImageBuilder.h
--- a/Engine/Plugins/Enterprise/GLTFExporter/MigratedFromGit/Source/GLTFExporter/Private/Builders/GLTFImageBuilder.h
+++ b/Engine/Plugins/Enterprise/GLTFExporter/MigratedFromGit/Source/GLTFExporter/Private/Builders/GLTFImageBuilder.h
@@ -21,6 +21,8 @@ private:
 	FGLTFJsonImageIndex AddImage(const FColor* Pixels, FIntPoint Size, bool bIgnoreAlpha, const FString& Name);
 	FGLTFJsonImageIndex AddImage(const void* CompressedData, int64 CompressedByteLength, EGLTFJsonMimeType MimeType, const FString& Name);
 
+	EGLTFJsonMimeType GetImageFormat(const FColor* Pixels, FIntPoint Size, bool bIgnoreAlpha) const;
+
 	FString SaveImageToFile(const void* CompressedData, int64 CompressedByteLength, EGLTFJsonMimeType MimeType, const FString& Name);
 
 	TSet<FString> UniqueImageUris;
diff --git a/Engine/Plugins/Enterprise/GLTFExporter/MigratedFromGit/Source/GLTFExporter/Private/Builders/GLTFImageUtility.cpp b/Engine/Plugins/Enterprise/GLTFExporter/MigratedFromGit/Source/GLTFExporter/Private/Builders/GLTFImageUtility.cpp
new file mode 100644
--- /dev/null
+++ b/Engine/Plugins/Enterprise/GLTFExporter/MigratedFromGit/Source/GLTFExporter/Private/Builders/GLTFImageUtility.cpp
@@ -0,0 +1,56 @@
+// Copyright Epic Games, Inc. All Rights Reserved.
+
+#include "Builders/GLTFImageUtility.h"
+#include "IImageWrapperModule.h"
+#include "IImageWrapper.h"
+
+bool FGLTFImageUtility::NoAlphaNeeded(const FColor* Pixels, FIntPoint Size)
+{
+	const int64 Count = Size.X * Size.Y;
+	for (int64 Index = 0; Index < Count; ++Index)
+	{
+		const FColor& Pixel = Pixels[Index];
+		if (Pixel.A < 255)
+		{
+			return false;
+		}
+	}
+
+	return true;
+}
+
+bool FGLTFImageUtility::CompressToPNG(const FColor* InPixels, FIntPoint InSize, TArray64<uint8>& OutCompressedData)
+{
+	return CompressToFormat(InPixels, InSize, EImageFormat::PNG, 0, OutCompressedData);
+}
+
+bool FGLTFImageUtility::CompressToJPEG(const FColor* InPixels, FIntPoint InSize, int32 InCompressionQuality, TArray64<uint8>& OutCompressedData)
+{
+	return CompressToFormat(InPixels, InSize, EImageFormat::JPEG, InCompressionQuality, OutCompressedData);
+}
+
+bool FGLTFImageUtility::CompressToFormat(const FColor* InPixels, FIntPoint InSize, EImageFormat InCompressionFormat, int32 InCompressionQuality, TArray64<uint8>& OutCompressedData)
+{
+	const int64 ByteLength = InSize.X * InSize.Y * sizeof(FColor);
+	return CompressToFormat(InPixels, ByteLength, InSize.X, InSize.Y, ERGBFormat::BGRA, 8, InCompressionFormat, InCompressionQuality, OutCompressedData);
+}
+
+bool FGLTFImageUtility::CompressToFormat(const void* InRawData, int64 InRawSize, int32 InWidth, int32 InHeight, ERGBFormat InRGBFormat, int32 InBitDepth, EImageFormat InCompressionFormat, int32 InCompressionQuality, TArray64<uint8>& OutCompressedData)
+{
+	if (InCompressionFormat == EImageFormat::JPEG && InRGBFormat == ERGBFormat::Gray)
+	{
+		InCompressionFormat = EImageFormat::GrayscaleJPEG;
+	}
+
+	IImageWrapperModule& ImageWrapperModule = FModuleManager::Get().LoadModuleChecked<IImageWrapperModule>(TEXT("ImageWrapper"));
+	TSharedPtr<IImageWrapper> ImageWrapper = ImageWrapperModule.CreateImageWrapper(InCompressionFormat);
+
+	if (ImageWrapper.IsValid() && ImageWrapper->SetRaw(InRawData, InRawSize, InWidth, InHeight, InRGBFormat, InBitDepth))
+	{
+		const TArray64<uint8>& CompressedData = ImageWrapper->GetCompressed(InCompressionQuality);
+		OutCompressedData = MoveTemp(const_cast<TArray64<uint8>&>(CompressedData)); // safe to do this since we don't keep image wrapper
+		return true;
+	}
+
+	return false;
+}
diff --git a/Engine/Plugins/Enterprise/GLTFExporter/MigratedFromGit/Source/GLTFExporter/Private/Builders/GLTFImageUtility.h b/Engine/Plugins/Enterprise/GLTFExporter/MigratedFromGit/Source/GLTFExporter/Private/Builders/GLTFImageUtility.h
new file mode 100644
--- /dev/null
+++ b/Engine/Plugins/Enterprise/GLTFExporter/MigratedFromGit/Source/GLTFExporter/Private/Builders/GLTFImageUtility.h
@@ -0,0 +1,19 @@
+// Copyright Epic Games, Inc. All Rights Reserved.
+
+#pragma once
+
+#include "CoreMinimal.h"
+
+enum class ERGBFormat : int8;
+enum class EImageFormat : int8;
+
+struct FGLTFImageUtility
+{
+	static bool NoAlphaNeeded(const FColor* Pixels, FIntPoint Size);
+
+	static bool CompressToPNG(const FColor* InPixels, FIntPoint InSize, TArray64<uint8>& OutCompressedData);
+	static bool CompressToJPEG(const FColor* InPixels, FIntPoint InSize, int32 InCompressionQuality, TArray64<uint8>& OutCompressedData);
+
+	static bool CompressToFormat(const FColor* InPixels, FIntPoint InSize, EImageFormat InCompressionFormat, int32 InCompressionQuality, TArray64<uint8>& OutCompressedData);
+	static bool CompressToFormat(const void* InRawData, int64 InRawSize, int32 InWidth, int32 InHeight, ERGBFormat InRGBFormat, int32 InBitDepth, EImageFormat InCompressionFormat, int32 InCompressionQuality, TArray64<uint8>& OutCompressedData);
+};
