[GLTFExporter] feat: add GLTFExportOptions::TextureCompressionQuality

commit 4ec2db7323af0508b0b87110b945266d57772976
Date: Sun, 8 Nov 2020 18:03:26 +0100

    feat: add GLTFExportOptions::TextureCompressionQuality
    (cherry picked from commit 5c44268e112a4d4056484c1759a233236fd308c6​)

#skipci

diff --git a/Engine/Plugins/Enterprise/GLTFExporter/MigratedFromGit/Source/GLTFExporter/Private/Builders/GLTFImageBuilder.cpp b/Engine/Plugins/Enterprise/GLTFExporter/MigratedFromGit/Source/GLTFExporter/Private/Builders/GLTFImageBuilder.cpp
--- a/Engine/Plugins/Enterprise/GLTFExporter/MigratedFromGit/Source/GLTFExporter/Private/Builders/GLTFImageBuilder.cpp
+++ b/Engine/Plugins/Enterprise/GLTFExporter/MigratedFromGit/Source/GLTFExporter/Private/Builders/GLTFImageBuilder.cpp
@@ -34,7 +34,7 @@ FGLTFJsonImageIndex FGLTFImageBuilder::AddImage(const FColor* Pixels, FIntPoint
 			break;
 
 		case EGLTFJsonMimeType::JPEG:
-			FGLTFImageUtility::CompressToJPEG(Pixels, Size, 0, CompressedData);
+			FGLTFImageUtility::CompressToJPEG(Pixels, Size, ExportOptions->TextureCompressionQuality, CompressedData);
 			break;
 
 		default:
diff --git a/Engine/Plugins/Enterprise/GLTFExporter/MigratedFromGit/Source/GLTFExporter/Private/Builders/GLTFImageUtility.cpp b/Engine/Plugins/Enterprise/GLTFExporter/MigratedFromGit/Source/GLTFExporter/Private/Builders/GLTFImageUtility.cpp
--- a/Engine/Plugins/Enterprise/GLTFExporter/MigratedFromGit/Source/GLTFExporter/Private/Builders/GLTFImageUtility.cpp
+++ b/Engine/Plugins/Enterprise/GLTFExporter/MigratedFromGit/Source/GLTFExporter/Private/Builders/GLTFImageUtility.cpp
@@ -37,11 +37,6 @@ bool FGLTFImageUtility::CompressToFormat(const FColor* InPixels, FIntPoint InSiz
 
 bool FGLTFImageUtility::CompressToFormat(const void* InRawData, int64 InRawSize, int32 InWidth, int32 InHeight, ERGBFormat InRGBFormat, int32 InBitDepth, EImageFormat InCompressionFormat, int32 InCompressionQuality, TArray64<uint8>& OutCompressedData)
 {
-	if (InCompressionFormat == EImageFormat::JPEG && InRGBFormat == ERGBFormat::Gray)
-	{
-		InCompressionFormat = EImageFormat::GrayscaleJPEG;
-	}
-
 	IImageWrapperModule& ImageWrapperModule = FModuleManager::Get().LoadModuleChecked<IImageWrapperModule>(TEXT("ImageWrapper"));
 	TSharedPtr<IImageWrapper> ImageWrapper = ImageWrapperModule.CreateImageWrapper(InCompressionFormat);
 
diff --git a/Engine/Plugins/Enterprise/GLTFExporter/MigratedFromGit/Source/GLTFExporter/Private/GLTFExportOptions.cpp b/Engine/Plugins/Enterprise/GLTFExporter/MigratedFromGit/Source/GLTFExporter/Private/GLTFExportOptions.cpp
--- a/Engine/Plugins/Enterprise/GLTFExporter/MigratedFromGit/Source/GLTFExporter/Private/GLTFExportOptions.cpp
+++ b/Engine/Plugins/Enterprise/GLTFExporter/MigratedFromGit/Source/GLTFExporter/Private/GLTFExportOptions.cpp
@@ -29,6 +29,7 @@ void UGLTFExportOptions::ResetToDefault()
 	bRetargetBoneTransforms = true;
 	bExportPlaybackSettings = true;
 	TextureCompression = EGLTFExporterTextureCompression::PNG;
+	TextureCompressionQuality = 0;
 	bExportTextureTransforms = true;
 	bExportLightmaps = true;
 	TextureHDREncoding = EGLTFExporterTextureHDREncoding::RGBM;
diff --git a/Engine/Plugins/Enterprise/GLTFExporter/MigratedFromGit/Source/GLTFExporter/Public/GLTFExportOptions.h b/Engine/Plugins/Enterprise/GLTFExporter/MigratedFromGit/Source/GLTFExporter/Public/GLTFExportOptions.h
--- a/Engine/Plugins/Enterprise/GLTFExporter/MigratedFromGit/Source/GLTFExporter/Public/GLTFExportOptions.h
+++ b/Engine/Plugins/Enterprise/GLTFExporter/MigratedFromGit/Source/GLTFExporter/Public/GLTFExportOptions.h
@@ -99,6 +99,9 @@ public:
 	UPROPERTY(EditAnywhere, BlueprintReadWrite, Config, Category = Texture)
 	EGLTFExporterTextureCompression TextureCompression;
 
+	UPROPERTY(EditAnywhere, BlueprintReadWrite, Config, Category = Texture, Meta = (ClampMin = "0", ClampMax = "100", EditCondition = "TextureCompression == EGLTFExporterTextureCompression::JPEG"))
+	int32 TextureCompressionQuality;
+
 	UPROPERTY(EditAnywhere, BlueprintReadWrite, Config, Category = Texture)
 	bool bExportTextureTransforms;
 
