[GLTFExporter] refactor: add some TODO and NOTE comments regarding JPEG

commit e3359730b48568ab864afdc7fba55a101936fb94
Date: Sun, 8 Nov 2020 19:33:50 +0100

    refactor: add some TODO and NOTE comments regarding JPEG
    (cherry picked from commit 9c07aa64724f95bbf77da64a475daa1b6fa73833​)

#skipci

diff --git a/Engine/Plugins/Enterprise/GLTFExporter/MigratedFromGit/Source/GLTFExporter/Private/Builders/GLTFImageUtility.cpp b/Engine/Plugins/Enterprise/GLTFExporter/MigratedFromGit/Source/GLTFExporter/Private/Builders/GLTFImageUtility.cpp
--- a/Engine/Plugins/Enterprise/GLTFExporter/MigratedFromGit/Source/GLTFExporter/Private/Builders/GLTFImageUtility.cpp
+++ b/Engine/Plugins/Enterprise/GLTFExporter/MigratedFromGit/Source/GLTFExporter/Private/Builders/GLTFImageUtility.cpp
@@ -26,6 +26,7 @@ bool FGLTFImageUtility::CompressToPNG(const FColor* InPixels, FIntPoint InSize,
 
 bool FGLTFImageUtility::CompressToJPEG(const FColor* InPixels, FIntPoint InSize, int32 InCompressionQuality, TArray64<uint8>& OutCompressedData)
 {
+	// NOTE: Tests have shown that explicitly checking for greyscale images (and using EImageFormat::GrayscaleJPEG) results in the same as EImageFormat::JPEG
 	return CompressToFormat(InPixels, InSize, EImageFormat::JPEG, InCompressionQuality, OutCompressedData);
 }
 
