[GLTFExporter] refactor: ignore alpha for all baked properties but Base Color + Opacity (Mask)

commit b7f4bab9fb9298f57d6f87167fb741fddf893d2a
Date: Sun, 8 Nov 2020 21:10:53 +0100

    refactor: ignore alpha for all baked properties but Base Color + Opacity (Mask)
    (cherry picked from commit 8e43fbe251755ae44acced4444ef7c36dfd913b0​)

#skipci

diff --git a/Engine/Plugins/Enterprise/GLTFExporter/MigratedFromGit/Source/GLTFExporter/Private/Converters/GLTFMaterialUtility.cpp b/Engine/Plugins/Enterprise/GLTFExporter/MigratedFromGit/Source/GLTFExporter/Private/Converters/GLTFMaterialUtility.cpp
--- a/Engine/Plugins/Enterprise/GLTFExporter/MigratedFromGit/Source/GLTFExporter/Private/Converters/GLTFMaterialUtility.cpp
+++ b/Engine/Plugins/Enterprise/GLTFExporter/MigratedFromGit/Source/GLTFExporter/Private/Converters/GLTFMaterialUtility.cpp
@@ -310,7 +310,7 @@ FGLTFPropertyBakeOutput FGLTFMaterialUtility::BakeMaterialProperty(const FIntPoi
 	return PropertyBakeOutput;
 }
 
-FGLTFJsonTextureIndex FGLTFMaterialUtility::AddCombinedTexture(FGLTFConvertBuilder& Builder, const TArray<FGLTFTextureCombineSource>& CombineSources, const FIntPoint& TextureSize, const FString& TextureName, EGLTFJsonTextureFilter MinFilter, EGLTFJsonTextureFilter MagFilter, EGLTFJsonTextureWrap WrapS, EGLTFJsonTextureWrap WrapT)
+FGLTFJsonTextureIndex FGLTFMaterialUtility::AddCombinedTexture(FGLTFConvertBuilder& Builder, const TArray<FGLTFTextureCombineSource>& CombineSources, const FIntPoint& TextureSize, bool bIgnoreAlpha, const FString& TextureName, EGLTFJsonTextureFilter MinFilter, EGLTFJsonTextureFilter MagFilter, EGLTFJsonTextureWrap WrapS, EGLTFJsonTextureWrap WrapT)
 {
 	check(CombineSources.Num() > 0);
 
@@ -322,10 +322,10 @@ FGLTFJsonTextureIndex FGLTFMaterialUtility::AddCombinedTexture(FGLTFConvertBuild
 		return FGLTFJsonTextureIndex(INDEX_NONE);
 	}
 
-	return AddTexture(Builder, Pixels, TextureSize, TextureName, MinFilter, MagFilter, WrapS, WrapT);
+	return AddTexture(Builder, Pixels, TextureSize, bIgnoreAlpha, TextureName, MinFilter, MagFilter, WrapS, WrapT);
 }
 
-FGLTFJsonTextureIndex FGLTFMaterialUtility::AddTexture(FGLTFConvertBuilder& Builder, const TArray<FColor>& Pixels, const FIntPoint& TextureSize, const FString& TextureName, EGLTFJsonTextureFilter MinFilter, EGLTFJsonTextureFilter MagFilter, EGLTFJsonTextureWrap WrapS, EGLTFJsonTextureWrap WrapT)
+FGLTFJsonTextureIndex FGLTFMaterialUtility::AddTexture(FGLTFConvertBuilder& Builder, const TArray<FColor>& Pixels, const FIntPoint& TextureSize, bool bIgnoreAlpha, const FString& TextureName, EGLTFJsonTextureFilter MinFilter, EGLTFJsonTextureFilter MagFilter, EGLTFJsonTextureWrap WrapS, EGLTFJsonTextureWrap WrapT)
 {
 	// TODO: maybe we should reuse existing samplers?
 	FGLTFJsonSampler JsonSampler;
@@ -339,7 +339,7 @@ FGLTFJsonTextureIndex FGLTFMaterialUtility::AddTexture(FGLTFConvertBuilder& Buil
 	FGLTFJsonTexture JsonTexture;
 	JsonTexture.Name = TextureName;
 	JsonTexture.Sampler = Builder.AddSampler(JsonSampler);
-	JsonTexture.Source = Builder.AddImage(Pixels, TextureSize, false, TextureName);
+	JsonTexture.Source = Builder.AddImage(Pixels, TextureSize, bIgnoreAlpha, TextureName);
 
 	return Builder.AddTexture(JsonTexture);
 }
diff --git a/Engine/Plugins/Enterprise/GLTFExporter/MigratedFromGit/Source/GLTFExporter/Private/Converters/GLTFMaterialUtility.h b/Engine/Plugins/Enterprise/GLTFExporter/MigratedFromGit/Source/GLTFExporter/Private/Converters/GLTFMaterialUtility.h
--- a/Engine/Plugins/Enterprise/GLTFExporter/MigratedFromGit/Source/GLTFExporter/Private/Converters/GLTFMaterialUtility.h
+++ b/Engine/Plugins/Enterprise/GLTFExporter/MigratedFromGit/Source/GLTFExporter/Private/Converters/GLTFMaterialUtility.h
@@ -55,8 +55,8 @@ struct FGLTFMaterialUtility
 	static bool CombineTextures(TArray<FColor>& OutPixels, const TArray<FGLTFTextureCombineSource>& Sources, const FIntPoint& OutputSize, EPixelFormat OutputPixelFormat);
 	static FGLTFPropertyBakeOutput BakeMaterialProperty(const FIntPoint& OutputSize, EMaterialProperty Property, const UMaterialInterface* Material, int32 TexCoord, const FMeshDescription* MeshDescription = nullptr, const TArray<int32>& MeshSectionIndices = {}, bool bCopyAlphaFromRedChannel = false);
 
-	static FGLTFJsonTextureIndex AddCombinedTexture(FGLTFConvertBuilder& Builder, const TArray<FGLTFTextureCombineSource>& CombineSources, const FIntPoint& TextureSize, const FString& TextureName, EGLTFJsonTextureFilter MinFilter, EGLTFJsonTextureFilter MagFilter, EGLTFJsonTextureWrap WrapS, EGLTFJsonTextureWrap WrapT);
-	static FGLTFJsonTextureIndex AddTexture(FGLTFConvertBuilder& Builder, const TArray<FColor>& Pixels, const FIntPoint& TextureSize, const FString& TextureName, EGLTFJsonTextureFilter MinFilter, EGLTFJsonTextureFilter MagFilter, EGLTFJsonTextureWrap WrapS, EGLTFJsonTextureWrap WrapT);
+	static FGLTFJsonTextureIndex AddCombinedTexture(FGLTFConvertBuilder& Builder, const TArray<FGLTFTextureCombineSource>& CombineSources, const FIntPoint& TextureSize, bool bIgnoreAlpha, const FString& TextureName, EGLTFJsonTextureFilter MinFilter, EGLTFJsonTextureFilter MagFilter, EGLTFJsonTextureWrap WrapS, EGLTFJsonTextureWrap WrapT);
+	static FGLTFJsonTextureIndex AddTexture(FGLTFConvertBuilder& Builder, const TArray<FColor>& Pixels, const FIntPoint& TextureSize, bool bIgnoreAlpha, const FString& TextureName, EGLTFJsonTextureFilter MinFilter, EGLTFJsonTextureFilter MagFilter, EGLTFJsonTextureWrap WrapS, EGLTFJsonTextureWrap WrapT);
 
 	static FLinearColor GetMask(const FExpressionInput& ExpressionInput);
 	static uint32 GetMaskComponentCount(const FExpressionInput& ExpressionInput);
diff --git a/Engine/Plugins/Enterprise/GLTFExporter/MigratedFromGit/Source/GLTFExporter/Private/Tasks/GLTFMaterialTasks.cpp b/Engine/Plugins/Enterprise/GLTFExporter/MigratedFromGit/Source/GLTFExporter/Private/Tasks/GLTFMaterialTasks.cpp
--- a/Engine/Plugins/Enterprise/GLTFExporter/MigratedFromGit/Source/GLTFExporter/Private/Tasks/GLTFMaterialTasks.cpp
+++ b/Engine/Plugins/Enterprise/GLTFExporter/MigratedFromGit/Source/GLTFExporter/Private/Tasks/GLTFMaterialTasks.cpp
@@ -404,6 +404,7 @@ bool FGLTFMaterialTask::TryGetBaseColorAndOpacity(FGLTFJsonPBRMetallicRoughness&
 		Builder,
 		CombineSources,
 		TextureSize,
+		false,
 		GetBakedTextureName(TEXT("BaseColor")),
 		TextureMinFilter,
 		TextureMagFilter,
@@ -568,6 +569,7 @@ bool FGLTFMaterialTask::TryGetMetallicAndRoughness(FGLTFJsonPBRMetallicRoughness
 		Builder,
 		CombineSources,
 		TextureSize,
+		true, // NOTE: we can ignore alpha in everything but TryGetBaseColorAndOpacity
 		GetBakedTextureName(TEXT("MetallicRoughness")),
 		TextureMinFilter,
 		TextureMagFilter,
@@ -734,6 +736,7 @@ bool FGLTFMaterialTask::TryGetClearCoatRoughness(FGLTFJsonClearCoatExtension& Ou
 		Builder,
 		CombineSources,
 		TextureSize,
+		true, // NOTE: we can ignore alpha in everything but TryGetBaseColorAndOpacity
 		GetBakedTextureName(TEXT("ClearCoatRoughness")),
 		TextureMinFilter,
 		TextureMagFilter,
@@ -1307,6 +1310,7 @@ bool FGLTFMaterialTask::TryGetBakedMaterialProperty(FGLTFJsonTextureInfo& OutTex
 		Builder,
 		PropertyBakeOutput.Pixels,
 		PropertyBakeOutput.Size,
+		true, // NOTE: we can ignore alpha in everything but TryGetBaseColorAndOpacity
 		GetBakedTextureName(PropertyName),
 		EGLTFJsonTextureFilter::Nearest,
 		EGLTFJsonTextureFilter::Nearest,
@@ -1386,6 +1390,7 @@ bool FGLTFMaterialTask::StoreBakedPropertyTexture(FGLTFJsonTextureInfo& OutTexIn
 		Builder,
 		PropertyBakeOutput.Pixels,
 		PropertyBakeOutput.Size,
+		true, // NOTE: we can ignore alpha in everything but TryGetBaseColorAndOpacity
 		GetBakedTextureName(PropertyName),
 		TextureMinFilter,
 		TextureMagFilter,
