[GLTFExporter] feat: add and implement GLTFExportOptions::LosslessCompressTextures to exempt certain textures from lossy compression

commit 0bbf8bc51a9e6ffd7fa71e57e20e0d57122eaa16
Date: Sun, 8 Nov 2020 23:11:46 +0100

    feat: add and implement GLTFExportOptions::LosslessCompressTextures to exempt certain textures from lossy compression
    (cherry picked from commit a139ef19e6cb8d5fe2120f8697ce5b7eac9093b3​)

#skipci

diff --git a/Engine/Plugins/Enterprise/GLTFExporter/MigratedFromGit/Source/GLTFExporter/Private/Builders/GLTFImageBuilder.cpp b/Engine/Plugins/Enterprise/GLTFExporter/MigratedFromGit/Source/GLTFExporter/Private/Builders/GLTFImageBuilder.cpp
--- a/Engine/Plugins/Enterprise/GLTFExporter/MigratedFromGit/Source/GLTFExporter/Private/Builders/GLTFImageBuilder.cpp
+++ b/Engine/Plugins/Enterprise/GLTFExporter/MigratedFromGit/Source/GLTFExporter/Private/Builders/GLTFImageBuilder.cpp
@@ -10,23 +10,23 @@ FGLTFImageBuilder::FGLTFImageBuilder(const FString& FilePath, const UGLTFExportO
 {
 }
 
-FGLTFJsonImageIndex FGLTFImageBuilder::AddImage(const TArray<FColor>& Pixels, FIntPoint Size, bool bIgnoreAlpha, const FString& Name)
+FGLTFJsonImageIndex FGLTFImageBuilder::AddImage(const TArray<FColor>& Pixels, FIntPoint Size, bool bIgnoreAlpha, EGLTFExporterTextureFlags Flags, const FString& Name)
 {
 	check(Pixels.Num() == Size.X * Size.Y);
-	return AddImage(Pixels.GetData(), Size, bIgnoreAlpha, Name);
+	return AddImage(Pixels.GetData(), Size, bIgnoreAlpha, Flags, Name);
 }
 
-FGLTFJsonImageIndex FGLTFImageBuilder::AddImage(const FColor* Pixels, int64 ByteLength, FIntPoint Size, bool bIgnoreAlpha, const FString& Name)
+FGLTFJsonImageIndex FGLTFImageBuilder::AddImage(const FColor* Pixels, int64 ByteLength, FIntPoint Size, bool bIgnoreAlpha, EGLTFExporterTextureFlags Flags, const FString& Name)
 {
 	check(ByteLength == Size.X * Size.Y * sizeof(FColor));
-	return AddImage(Pixels, Size, bIgnoreAlpha, Name);
+	return AddImage(Pixels, Size, bIgnoreAlpha, Flags, Name);
 }
 
-FGLTFJsonImageIndex FGLTFImageBuilder::AddImage(const FColor* Pixels, FIntPoint Size, bool bIgnoreAlpha, const FString& Name)
+FGLTFJsonImageIndex FGLTFImageBuilder::AddImage(const FColor* Pixels, FIntPoint Size, bool bIgnoreAlpha, EGLTFExporterTextureFlags Flags, const FString& Name)
 {
 	TArray64<uint8> CompressedData;
 
-	const EGLTFJsonMimeType ImageFormat = GetImageFormat(Pixels, Size, bIgnoreAlpha);
+	const EGLTFJsonMimeType ImageFormat = GetImageFormat(Pixels, Size, bIgnoreAlpha, Flags);
 	switch (ImageFormat)
 	{
 		case EGLTFJsonMimeType::PNG:
@@ -73,7 +73,7 @@ FGLTFJsonImageIndex FGLTFImageBuilder::AddImage(const void* CompressedData, int6
 	return ImageIndex;
 }
 
-EGLTFJsonMimeType FGLTFImageBuilder::GetImageFormat(const FColor* Pixels, FIntPoint Size, bool bIgnoreAlpha) const
+EGLTFJsonMimeType FGLTFImageBuilder::GetImageFormat(const FColor* Pixels, FIntPoint Size, bool bIgnoreAlpha, EGLTFExporterTextureFlags Flags) const
 {
 	switch (ExportOptions->TextureCompression)
 	{
@@ -81,7 +81,10 @@ EGLTFJsonMimeType FGLTFImageBuilder::GetImageFormat(const FColor* Pixels, FIntPo
 			return EGLTFJsonMimeType::PNG;
 
 		case EGLTFExporterTextureCompression::JPEG:
-			return bIgnoreAlpha || FGLTFImageUtility::NoAlphaNeeded(Pixels, Size) ? EGLTFJsonMimeType::JPEG : EGLTFJsonMimeType::PNG;
+			return
+				!EnumHasAllFlags(static_cast<EGLTFExporterTextureFlags>(ExportOptions->LosslessCompressTextures), Flags) &&
+				(bIgnoreAlpha || FGLTFImageUtility::NoAlphaNeeded(Pixels, Size)) ?
+				EGLTFJsonMimeType::JPEG : EGLTFJsonMimeType::PNG;
 
 		default:
 			checkNoEntry();
diff --git a/Engine/Plugins/Enterprise/GLTFExporter/MigratedFromGit/Source/GLTFExporter/Private/Builders/GLTFImageBuilder.h b/Engine/Plugins/Enterprise/GLTFExporter/MigratedFromGit/Source/GLTFExporter/Private/Builders/GLTFImageBuilder.h
--- a/Engine/Plugins/Enterprise/GLTFExporter/MigratedFromGit/Source/GLTFExporter/Private/Builders/GLTFImageBuilder.h
+++ b/Engine/Plugins/Enterprise/GLTFExporter/MigratedFromGit/Source/GLTFExporter/Private/Builders/GLTFImageBuilder.h
@@ -13,15 +13,15 @@ protected:
 
 public:
 
-	FGLTFJsonImageIndex AddImage(const TArray<FColor>& Pixels, FIntPoint Size, bool bIgnoreAlpha, const FString& Name);
-	FGLTFJsonImageIndex AddImage(const FColor* Pixels, int64 ByteLength, FIntPoint Size, bool bIgnoreAlpha, const FString& Name);
+	FGLTFJsonImageIndex AddImage(const TArray<FColor>& Pixels, FIntPoint Size, bool bIgnoreAlpha, EGLTFExporterTextureFlags Flags, const FString& Name);
+	FGLTFJsonImageIndex AddImage(const FColor* Pixels, int64 ByteLength, FIntPoint Size, bool bIgnoreAlpha, EGLTFExporterTextureFlags Flags, const FString& Name);
 
 private:
 
-	FGLTFJsonImageIndex AddImage(const FColor* Pixels, FIntPoint Size, bool bIgnoreAlpha, const FString& Name);
+	FGLTFJsonImageIndex AddImage(const FColor* Pixels, FIntPoint Size, bool bIgnoreAlpha, EGLTFExporterTextureFlags Flags, const FString& Name);
 	FGLTFJsonImageIndex AddImage(const void* CompressedData, int64 CompressedByteLength, EGLTFJsonMimeType MimeType, const FString& Name);
 
-	EGLTFJsonMimeType GetImageFormat(const FColor* Pixels, FIntPoint Size, bool bIgnoreAlpha) const;
+	EGLTFJsonMimeType GetImageFormat(const FColor* Pixels, FIntPoint Size, bool bIgnoreAlpha, EGLTFExporterTextureFlags Flags) const;
 
 	FString SaveImageToFile(const void* CompressedData, int64 CompressedByteLength, EGLTFJsonMimeType MimeType, const FString& Name);
 
diff --git a/Engine/Plugins/Enterprise/GLTFExporter/MigratedFromGit/Source/GLTFExporter/Private/Converters/GLTFMaterialUtility.cpp b/Engine/Plugins/Enterprise/GLTFExporter/MigratedFromGit/Source/GLTFExporter/Private/Converters/GLTFMaterialUtility.cpp
--- a/Engine/Plugins/Enterprise/GLTFExporter/MigratedFromGit/Source/GLTFExporter/Private/Converters/GLTFMaterialUtility.cpp
+++ b/Engine/Plugins/Enterprise/GLTFExporter/MigratedFromGit/Source/GLTFExporter/Private/Converters/GLTFMaterialUtility.cpp
@@ -322,10 +322,10 @@ FGLTFJsonTextureIndex FGLTFMaterialUtility::AddCombinedTexture(FGLTFConvertBuild
 		return FGLTFJsonTextureIndex(INDEX_NONE);
 	}
 
-	return AddTexture(Builder, Pixels, TextureSize, bIgnoreAlpha, TextureName, MinFilter, MagFilter, WrapS, WrapT);
+	return AddTexture(Builder, Pixels, TextureSize, bIgnoreAlpha, false, TextureName, MinFilter, MagFilter, WrapS, WrapT);
 }
 
-FGLTFJsonTextureIndex FGLTFMaterialUtility::AddTexture(FGLTFConvertBuilder& Builder, const TArray<FColor>& Pixels, const FIntPoint& TextureSize, bool bIgnoreAlpha, const FString& TextureName, EGLTFJsonTextureFilter MinFilter, EGLTFJsonTextureFilter MagFilter, EGLTFJsonTextureWrap WrapS, EGLTFJsonTextureWrap WrapT)
+FGLTFJsonTextureIndex FGLTFMaterialUtility::AddTexture(FGLTFConvertBuilder& Builder, const TArray<FColor>& Pixels, const FIntPoint& TextureSize, bool bIgnoreAlpha, bool bIsNormalmap, const FString& TextureName, EGLTFJsonTextureFilter MinFilter, EGLTFJsonTextureFilter MagFilter, EGLTFJsonTextureWrap WrapS, EGLTFJsonTextureWrap WrapT)
 {
 	// TODO: maybe we should reuse existing samplers?
 	FGLTFJsonSampler JsonSampler;
@@ -339,7 +339,7 @@ FGLTFJsonTextureIndex FGLTFMaterialUtility::AddTexture(FGLTFConvertBuilder& Buil
 	FGLTFJsonTexture JsonTexture;
 	JsonTexture.Name = TextureName;
 	JsonTexture.Sampler = Builder.AddSampler(JsonSampler);
-	JsonTexture.Source = Builder.AddImage(Pixels, TextureSize, bIgnoreAlpha, TextureName);
+	JsonTexture.Source = Builder.AddImage(Pixels, TextureSize, bIgnoreAlpha, bIsNormalmap ? EGLTFExporterTextureFlags::Normalmaps : EGLTFExporterTextureFlags::None, TextureName);
 
 	return Builder.AddTexture(JsonTexture);
 }
diff --git a/Engine/Plugins/Enterprise/GLTFExporter/MigratedFromGit/Source/GLTFExporter/Private/Converters/GLTFMaterialUtility.h b/Engine/Plugins/Enterprise/GLTFExporter/MigratedFromGit/Source/GLTFExporter/Private/Converters/GLTFMaterialUtility.h
--- a/Engine/Plugins/Enterprise/GLTFExporter/MigratedFromGit/Source/GLTFExporter/Private/Converters/GLTFMaterialUtility.h
+++ b/Engine/Plugins/Enterprise/GLTFExporter/MigratedFromGit/Source/GLTFExporter/Private/Converters/GLTFMaterialUtility.h
@@ -56,7 +56,7 @@ struct FGLTFMaterialUtility
 	static FGLTFPropertyBakeOutput BakeMaterialProperty(const FIntPoint& OutputSize, EMaterialProperty Property, const UMaterialInterface* Material, int32 TexCoord, const FMeshDescription* MeshDescription = nullptr, const TArray<int32>& MeshSectionIndices = {}, bool bCopyAlphaFromRedChannel = false);
 
 	static FGLTFJsonTextureIndex AddCombinedTexture(FGLTFConvertBuilder& Builder, const TArray<FGLTFTextureCombineSource>& CombineSources, const FIntPoint& TextureSize, bool bIgnoreAlpha, const FString& TextureName, EGLTFJsonTextureFilter MinFilter, EGLTFJsonTextureFilter MagFilter, EGLTFJsonTextureWrap WrapS, EGLTFJsonTextureWrap WrapT);
-	static FGLTFJsonTextureIndex AddTexture(FGLTFConvertBuilder& Builder, const TArray<FColor>& Pixels, const FIntPoint& TextureSize, bool bIgnoreAlpha, const FString& TextureName, EGLTFJsonTextureFilter MinFilter, EGLTFJsonTextureFilter MagFilter, EGLTFJsonTextureWrap WrapS, EGLTFJsonTextureWrap WrapT);
+	static FGLTFJsonTextureIndex AddTexture(FGLTFConvertBuilder& Builder, const TArray<FColor>& Pixels, const FIntPoint& TextureSize, bool bIgnoreAlpha, bool bIsNormalmap, const FString& TextureName, EGLTFJsonTextureFilter MinFilter, EGLTFJsonTextureFilter MagFilter, EGLTFJsonTextureWrap WrapS, EGLTFJsonTextureWrap WrapT);
 
 	static FLinearColor GetMask(const FExpressionInput& ExpressionInput);
 	static uint32 GetMaskComponentCount(const FExpressionInput& ExpressionInput);
diff --git a/Engine/Plugins/Enterprise/GLTFExporter/MigratedFromGit/Source/GLTFExporter/Private/GLTFExportOptions.cpp b/Engine/Plugins/Enterprise/GLTFExporter/MigratedFromGit/Source/GLTFExporter/Private/GLTFExportOptions.cpp
--- a/Engine/Plugins/Enterprise/GLTFExporter/MigratedFromGit/Source/GLTFExporter/Private/GLTFExportOptions.cpp
+++ b/Engine/Plugins/Enterprise/GLTFExporter/MigratedFromGit/Source/GLTFExporter/Private/GLTFExportOptions.cpp
@@ -30,6 +30,7 @@ void UGLTFExportOptions::ResetToDefault()
 	bExportPlaybackSettings = true;
 	TextureCompression = EGLTFExporterTextureCompression::PNG;
 	TextureCompressionQuality = 0;
+	LosslessCompressTextures = static_cast<int32>(EGLTFExporterTextureFlags::All);
 	bExportTextureTransforms = true;
 	bExportLightmaps = true;
 	TextureHDREncoding = EGLTFExporterTextureHDREncoding::RGBM;
diff --git a/Engine/Plugins/Enterprise/GLTFExporter/MigratedFromGit/Source/GLTFExporter/Private/Tasks/GLTFMaterialTasks.cpp b/Engine/Plugins/Enterprise/GLTFExporter/MigratedFromGit/Source/GLTFExporter/Private/Tasks/GLTFMaterialTasks.cpp
--- a/Engine/Plugins/Enterprise/GLTFExporter/MigratedFromGit/Source/GLTFExporter/Private/Tasks/GLTFMaterialTasks.cpp
+++ b/Engine/Plugins/Enterprise/GLTFExporter/MigratedFromGit/Source/GLTFExporter/Private/Tasks/GLTFMaterialTasks.cpp
@@ -1311,6 +1311,7 @@ bool FGLTFMaterialTask::TryGetBakedMaterialProperty(FGLTFJsonTextureInfo& OutTex
 		PropertyBakeOutput.Pixels,
 		PropertyBakeOutput.Size,
 		true, // NOTE: we can ignore alpha in everything but TryGetBaseColorAndOpacity
+		false, // Normal and ClearCoatBottomNormal are handled above
 		GetBakedTextureName(PropertyName),
 		EGLTFJsonTextureFilter::Nearest,
 		EGLTFJsonTextureFilter::Nearest,
@@ -1386,11 +1387,14 @@ bool FGLTFMaterialTask::StoreBakedPropertyTexture(FGLTFJsonTextureInfo& OutTexIn
 	const EGLTFJsonTextureFilter TextureMinFilter = EGLTFJsonTextureFilter::LinearMipmapLinear;
 	const EGLTFJsonTextureFilter TextureMagFilter = EGLTFJsonTextureFilter::Linear;
 
+	const bool bIsNormalmap = PropertyBakeOutput.Property == MP_Normal || PropertyBakeOutput.Property == MP_CustomOutput /* ClearCoatBottomNormal */;
+
 	const FGLTFJsonTextureIndex TextureIndex = FGLTFMaterialUtility::AddTexture(
 		Builder,
 		PropertyBakeOutput.Pixels,
 		PropertyBakeOutput.Size,
 		true, // NOTE: we can ignore alpha in everything but TryGetBaseColorAndOpacity
+		bIsNormalmap,
 		GetBakedTextureName(PropertyName),
 		TextureMinFilter,
 		TextureMagFilter,
diff --git a/Engine/Plugins/Enterprise/GLTFExporter/MigratedFromGit/Source/GLTFExporter/Private/Tasks/GLTFTextureTasks.cpp b/Engine/Plugins/Enterprise/GLTFExporter/MigratedFromGit/Source/GLTFExporter/Private/Tasks/GLTFTextureTasks.cpp
--- a/Engine/Plugins/Enterprise/GLTFExporter/MigratedFromGit/Source/GLTFExporter/Private/Tasks/GLTFTextureTasks.cpp
+++ b/Engine/Plugins/Enterprise/GLTFExporter/MigratedFromGit/Source/GLTFExporter/Private/Tasks/GLTFTextureTasks.cpp
@@ -36,7 +36,11 @@ void FGLTFTexture2DTask::Complete()
 	}
 
 	const bool bIgnoreAlpha = FGLTFTextureUtility::IsAlphaless(Texture2D->GetPixelFormat());
-	JsonTexture.Source = Builder.AddImage(Pixels, Size, bIgnoreAlpha, JsonTexture.Name);
+	const EGLTFExporterTextureFlags Flags =
+		(Texture2D->IsNormalMap() ? EGLTFExporterTextureFlags::Normalmaps : EGLTFExporterTextureFlags::None) |
+		(bIsHDR ? EGLTFExporterTextureFlags::HDR : EGLTFExporterTextureFlags::None);
+
+	JsonTexture.Source = Builder.AddImage(Pixels, Size, bIgnoreAlpha, Flags, JsonTexture.Name);
 	JsonTexture.Sampler = Builder.GetOrAddSampler(Texture2D);
 }
 
@@ -75,7 +79,9 @@ void FGLTFTextureCubeTask::Complete()
 	}
 
 	const bool bIgnoreAlpha = FGLTFTextureUtility::IsAlphaless(TextureCube->GetPixelFormat());
-	JsonTexture.Source = Builder.AddImage(Pixels, Size, bIgnoreAlpha, JsonTexture.Name);
+	const EGLTFExporterTextureFlags Flags = bIsHDR ? EGLTFExporterTextureFlags::HDR : EGLTFExporterTextureFlags::None;
+
+	JsonTexture.Source = Builder.AddImage(Pixels, Size, bIgnoreAlpha, Flags, JsonTexture.Name);
 	JsonTexture.Sampler = Builder.GetOrAddSampler(TextureCube);
 }
 
@@ -100,7 +106,9 @@ void FGLTFTextureRenderTarget2DTask::Complete()
 	}
 
 	const bool bIgnoreAlpha = FGLTFTextureUtility::IsAlphaless(RenderTarget2D->GetFormat());
-	JsonTexture.Source = Builder.AddImage(Pixels, Size, bIgnoreAlpha, JsonTexture.Name);
+	const EGLTFExporterTextureFlags Flags = bIsHDR ? EGLTFExporterTextureFlags::HDR : EGLTFExporterTextureFlags::None;
+
+	JsonTexture.Source = Builder.AddImage(Pixels, Size, bIgnoreAlpha, Flags, JsonTexture.Name);
 	JsonTexture.Sampler = Builder.GetOrAddSampler(RenderTarget2D);
 }
 
@@ -138,7 +146,9 @@ void FGLTFTextureRenderTargetCubeTask::Complete()
 	}
 
 	const bool bIgnoreAlpha = FGLTFTextureUtility::IsAlphaless(RenderTargetCube->GetFormat());
-	JsonTexture.Source = Builder.AddImage(Pixels, Size, bIgnoreAlpha, JsonTexture.Name);
+	const EGLTFExporterTextureFlags Flags = bIsHDR ? EGLTFExporterTextureFlags::HDR : EGLTFExporterTextureFlags::None;
+
+	JsonTexture.Source = Builder.AddImage(Pixels, Size, bIgnoreAlpha, Flags, JsonTexture.Name);
 	JsonTexture.Sampler = Builder.GetOrAddSampler(RenderTargetCube);
 }
 
@@ -167,8 +177,11 @@ void FGLTFTextureLightMapTask::Complete()
 	const FIntPoint Size = { Source.GetSizeX(), Source.GetSizeY() };
 	const int64 ByteLength = Source.CalcMipSize(0);
 
+	const bool bIgnoreAlpha = false;
+	const EGLTFExporterTextureFlags Flags = EGLTFExporterTextureFlags::Lightmaps;
+
 	const void* RawData = Source.LockMip(0);
-	JsonTexture.Source = Builder.AddImage(static_cast<const FColor*>(RawData), ByteLength, Size, false, JsonTexture.Name);
+	JsonTexture.Source = Builder.AddImage(static_cast<const FColor*>(RawData), ByteLength, Size, bIgnoreAlpha, Flags, JsonTexture.Name);
 	Source.UnlockMip(0);
 
 	JsonTexture.Sampler = Builder.GetOrAddSampler(LightMap);
diff --git a/Engine/Plugins/Enterprise/GLTFExporter/MigratedFromGit/Source/GLTFExporter/Public/GLTFExportOptions.h b/Engine/Plugins/Enterprise/GLTFExporter/MigratedFromGit/Source/GLTFExporter/Public/GLTFExportOptions.h
--- a/Engine/Plugins/Enterprise/GLTFExporter/MigratedFromGit/Source/GLTFExporter/Public/GLTFExportOptions.h
+++ b/Engine/Plugins/Enterprise/GLTFExporter/MigratedFromGit/Source/GLTFExporter/Public/GLTFExportOptions.h
@@ -27,10 +27,23 @@ enum class EGLTFExporterMaterialBakeSize : uint8
 UENUM(BlueprintType)
 enum class EGLTFExporterTextureCompression : uint8
 {
-	PNG UMETA(DisplayName = "PNG"),
+	PNG,
     JPEG UMETA(DisplayName = "JPEG (if no alpha)")
 };
 
+UENUM(BlueprintType, Meta = (Bitflags, UseEnumValuesAsMaskValuesInEditor = "true"))
+enum class EGLTFExporterTextureFlags : uint8
+{
+	None = 0 UMETA(Hidden),
+
+	HDR = 1 << 0,
+	Normalmaps = 1 << 1,
+    Lightmaps = 1 << 2,
+
+	All = HDR | Normalmaps | Lightmaps UMETA(Hidden)
+};
+ENUM_CLASS_FLAGS(EGLTFExporterTextureFlags);
+
 UENUM(BlueprintType)
 enum class EGLTFExporterTextureHDREncoding : uint8
 {
@@ -99,9 +112,12 @@ public:
 	UPROPERTY(EditAnywhere, BlueprintReadWrite, Config, Category = Texture)
 	EGLTFExporterTextureCompression TextureCompression;
 
-	UPROPERTY(EditAnywhere, BlueprintReadWrite, Config, Category = Texture, Meta = (ClampMin = "0", ClampMax = "100", EditCondition = "TextureCompression == EGLTFExporterTextureCompression::JPEG"))
+	UPROPERTY(EditAnywhere, BlueprintReadWrite, Config, Category = Texture, Meta = (ClampMin = "0", ClampMax = "100", EditCondition = "TextureCompression != EGLTFExporterTextureCompression::PNG"))
 	int32 TextureCompressionQuality;
 
+	UPROPERTY(EditAnywhere, BlueprintReadWrite, Config, Category = Texture, Meta = (Bitmask, BitmaskEnum = EGLTFExporterTextureFlags, EditCondition = "TextureCompression != EGLTFExporterTextureCompression::PNG"))
+	int32 LosslessCompressTextures;
+
 	UPROPERTY(EditAnywhere, BlueprintReadWrite, Config, Category = Texture)
 	bool bExportTextureTransforms;
 
