[GLTFExporter] refactor: replace export options bQuantizeVertexNormals and bQuantizeVertexTangents with bExportMeshQuantization

commit 647b59ee6cbe87de138feca4617b2c125009ecb6
Date: Sun, 8 Nov 2020 23:17:52 +0100

    refactor: replace export options bQuantizeVertexNormals and bQuantizeVertexTangents with bExportMeshQuantization
    (cherry picked from commit e3bec3fca96fa16fc404210b8afb5b44d037cfc9​)

#skipci

diff --git a/Engine/Plugins/Enterprise/GLTFExporter/MigratedFromGit/Source/GLTFExporter/Private/Converters/GLTFAccessorConverters.cpp b/Engine/Plugins/Enterprise/GLTFExporter/MigratedFromGit/Source/GLTFExporter/Private/Converters/GLTFAccessorConverters.cpp
--- a/Engine/Plugins/Enterprise/GLTFExporter/MigratedFromGit/Source/GLTFExporter/Private/Converters/GLTFAccessorConverters.cpp
+++ b/Engine/Plugins/Enterprise/GLTFExporter/MigratedFromGit/Source/GLTFExporter/Private/Converters/GLTFAccessorConverters.cpp
@@ -134,7 +134,7 @@ FGLTFJsonAccessorIndex FGLTFNormalBufferConverter::Convert(const FGLTFMeshSectio
 	FGLTFJsonBufferViewIndex BufferViewIndex;
 	EGLTFJsonComponentType ComponentType;
 
-	const bool bMeshQuantization = Builder.ExportOptions->bQuantizeVertexNormals;
+	const bool bMeshQuantization = Builder.ExportOptions->bExportMeshQuantization;
 	const bool bHighPrecision = VertexBuffer->GetUseHighPrecisionTangentBasis();
 
 	if (bMeshQuantization)
@@ -210,7 +210,7 @@ FGLTFJsonAccessorIndex FGLTFTangentBufferConverter::Convert(const FGLTFMeshSecti
 	FGLTFJsonBufferViewIndex BufferViewIndex;
 	EGLTFJsonComponentType ComponentType;
 
-	const bool bMeshQuantization = Builder.ExportOptions->bQuantizeVertexTangents;
+	const bool bMeshQuantization = Builder.ExportOptions->bExportMeshQuantization;
 	const bool bHighPrecision = VertexBuffer->GetUseHighPrecisionTangentBasis();
 
 	if (bMeshQuantization)
diff --git a/Engine/Plugins/Enterprise/GLTFExporter/MigratedFromGit/Source/GLTFExporter/Private/GLTFExportOptions.cpp b/Engine/Plugins/Enterprise/GLTFExporter/MigratedFromGit/Source/GLTFExporter/Private/GLTFExportOptions.cpp
--- a/Engine/Plugins/Enterprise/GLTFExporter/MigratedFromGit/Source/GLTFExporter/Private/GLTFExportOptions.cpp
+++ b/Engine/Plugins/Enterprise/GLTFExporter/MigratedFromGit/Source/GLTFExporter/Private/GLTFExportOptions.cpp
@@ -20,8 +20,7 @@ void UGLTFExportOptions::ResetToDefault()
 	bMaterialBakeUsingMeshData = false;
 	DefaultMaterialBakeSize = EGLTFExporterMaterialBakeSize::POT_1024;
 	bExportVertexColors = false;
-	bQuantizeVertexNormals = true;
-	bQuantizeVertexTangents = true;
+	bExportMeshQuantization = true;
 	DefaultLevelOfDetail = 0;
 	bExportLevelSequences = true;
 	bExportVertexSkinWeights = true;
diff --git a/Engine/Plugins/Enterprise/GLTFExporter/MigratedFromGit/Source/GLTFExporter/Public/GLTFExportOptions.h b/Engine/Plugins/Enterprise/GLTFExporter/MigratedFromGit/Source/GLTFExporter/Public/GLTFExportOptions.h
--- a/Engine/Plugins/Enterprise/GLTFExporter/MigratedFromGit/Source/GLTFExporter/Public/GLTFExportOptions.h
+++ b/Engine/Plugins/Enterprise/GLTFExporter/MigratedFromGit/Source/GLTFExporter/Public/GLTFExportOptions.h
@@ -86,10 +86,7 @@ public:
 	bool bExportVertexColors;
 
 	UPROPERTY(EditAnywhere, BlueprintReadWrite, Config, Category = Mesh)
-	bool bQuantizeVertexNormals;
-
-	UPROPERTY(EditAnywhere, BlueprintReadWrite, Config, Category = Mesh)
-	bool bQuantizeVertexTangents;
+	bool bExportMeshQuantization;
 
 	UPROPERTY(EditAnywhere, BlueprintReadWrite, Config, Category = Mesh, Meta = (ClampMin = "0"))
 	int32 DefaultLevelOfDetail;
