[GLTFExporter] refactor: add FGLTFMaterialUtility::IsNormalMap

commit ce06dd1988890f02009eadfe2f451808f52929db
Date: Sun, 8 Nov 2020 23:33:53 +0100

    refactor: add FGLTFMaterialUtility::IsNormalMap
    (cherry picked from commit ef590ab81e4c047c0825ef48460712a6359ab6be​)

#skipci

diff --git a/Engine/Plugins/Enterprise/GLTFExporter/MigratedFromGit/Source/GLTFExporter/Private/Converters/GLTFMaterialUtility.cpp b/Engine/Plugins/Enterprise/GLTFExporter/MigratedFromGit/Source/GLTFExporter/Private/Converters/GLTFMaterialUtility.cpp
--- a/Engine/Plugins/Enterprise/GLTFExporter/MigratedFromGit/Source/GLTFExporter/Private/Converters/GLTFMaterialUtility.cpp
+++ b/Engine/Plugins/Enterprise/GLTFExporter/MigratedFromGit/Source/GLTFExporter/Private/Converters/GLTFMaterialUtility.cpp
@@ -21,6 +21,11 @@ UMaterialInterface* FGLTFMaterialUtility::GetDefault()
 	return DefaultMaterial;
 }
 
+bool FGLTFMaterialUtility::IsNormalMap(EMaterialProperty Property)
+{
+	return Property == MP_Normal || Property == MP_CustomOutput /* ClearCoatBottomNormal */;
+}
+
 const TCHAR* FGLTFMaterialUtility::GetPropertyName(EMaterialProperty Property)
 {
 	// TODO: replace this hardcoded list with some lookup function that should be in the engine.
@@ -288,7 +293,7 @@ FGLTFPropertyBakeOutput FGLTFMaterialUtility::BakeMaterialProperty(const FIntPoi
 		}
 	}
 
-	if (Property == MP_Normal || Property == MP_CustomOutput)
+	if (IsNormalMap(Property))
 	{
 		// Convert normalmaps to use +Y (OpenGL / WebGL standard)
 		FGLTFTextureUtility::FlipGreenChannel(BakedPixels);
diff --git a/Engine/Plugins/Enterprise/GLTFExporter/MigratedFromGit/Source/GLTFExporter/Private/Converters/GLTFMaterialUtility.h b/Engine/Plugins/Enterprise/GLTFExporter/MigratedFromGit/Source/GLTFExporter/Private/Converters/GLTFMaterialUtility.h
--- a/Engine/Plugins/Enterprise/GLTFExporter/MigratedFromGit/Source/GLTFExporter/Private/Converters/GLTFMaterialUtility.h
+++ b/Engine/Plugins/Enterprise/GLTFExporter/MigratedFromGit/Source/GLTFExporter/Private/Converters/GLTFMaterialUtility.h
@@ -36,6 +36,8 @@ struct FGLTFMaterialUtility
 {
 	static UMaterialInterface* GetDefault();
 
+	static bool IsNormalMap(EMaterialProperty Property);
+
 	static const TCHAR* GetPropertyName(EMaterialProperty Property);
 	static FVector4 GetPropertyDefaultValue(EMaterialProperty Property);
 	static FVector4 GetPropertyMask(EMaterialProperty Property);
diff --git a/Engine/Plugins/Enterprise/GLTFExporter/MigratedFromGit/Source/GLTFExporter/Private/Tasks/GLTFMaterialTasks.cpp b/Engine/Plugins/Enterprise/GLTFExporter/MigratedFromGit/Source/GLTFExporter/Private/Tasks/GLTFMaterialTasks.cpp
--- a/Engine/Plugins/Enterprise/GLTFExporter/MigratedFromGit/Source/GLTFExporter/Private/Tasks/GLTFMaterialTasks.cpp
+++ b/Engine/Plugins/Enterprise/GLTFExporter/MigratedFromGit/Source/GLTFExporter/Private/Tasks/GLTFMaterialTasks.cpp
@@ -1292,7 +1292,7 @@ bool FGLTFMaterialTask::TryGetBakedMaterialProperty(FGLTFJsonTextureInfo& OutTex
 		return true;
 	}
 
-	if (Property == MP_Normal || Property == MP_CustomOutput /* ClearCoatBottomNormal */)
+	if (FGLTFMaterialUtility::IsNormalMap(Property))
 	{
 		// TODO: In some cases baking normal can result in constant vector that differs slight from default (i.e 0,0,1).
 		// Yet often, when looking at such a material, it should be exactly default. Needs further investigation.
@@ -1387,14 +1387,12 @@ bool FGLTFMaterialTask::StoreBakedPropertyTexture(FGLTFJsonTextureInfo& OutTexIn
 	const EGLTFJsonTextureFilter TextureMinFilter = EGLTFJsonTextureFilter::LinearMipmapLinear;
 	const EGLTFJsonTextureFilter TextureMagFilter = EGLTFJsonTextureFilter::Linear;
 
-	const bool bIsNormalMap = PropertyBakeOutput.Property == MP_Normal || PropertyBakeOutput.Property == MP_CustomOutput /* ClearCoatBottomNormal */;
-
 	const FGLTFJsonTextureIndex TextureIndex = FGLTFMaterialUtility::AddTexture(
 		Builder,
 		PropertyBakeOutput.Pixels,
 		PropertyBakeOutput.Size,
 		true, // NOTE: we can ignore alpha in everything but TryGetBaseColorAndOpacity
-		bIsNormalMap,
+		FGLTFMaterialUtility::IsNormalMap(PropertyBakeOutput.Property),
 		GetBakedTextureName(PropertyName),
 		TextureMinFilter,
 		TextureMagFilter,
