[GLTFExporter] fix: avoid flagging Texture2D as both normalmap and HDR

commit 39d07ff4778542b1892245e85ca9e44248f01fe4
Date: Sun, 8 Nov 2020 23:36:19 +0100

    fix: avoid flagging Texture2D as both normalmap and HDR
    (cherry picked from commit 29e04dc639d6c0f5c8547c4fcd4ed6fa99819c6a​)

#skipci

diff --git a/Engine/Plugins/Enterprise/GLTFExporter/MigratedFromGit/Source/GLTFExporter/Private/Tasks/GLTFTextureTasks.cpp b/Engine/Plugins/Enterprise/GLTFExporter/MigratedFromGit/Source/GLTFExporter/Private/Tasks/GLTFTextureTasks.cpp
--- a/Engine/Plugins/Enterprise/GLTFExporter/MigratedFromGit/Source/GLTFExporter/Private/Tasks/GLTFTextureTasks.cpp
+++ b/Engine/Plugins/Enterprise/GLTFExporter/MigratedFromGit/Source/GLTFExporter/Private/Tasks/GLTFTextureTasks.cpp
@@ -37,8 +37,8 @@ void FGLTFTexture2DTask::Complete()
 
 	const bool bIgnoreAlpha = FGLTFTextureUtility::IsAlphaless(Texture2D->GetPixelFormat());
 	const EGLTFExporterTextureFlags Flags =
-		(Texture2D->IsNormalMap() ? EGLTFExporterTextureFlags::Normalmaps : EGLTFExporterTextureFlags::None) |
-		(bIsHDR ? EGLTFExporterTextureFlags::HDR : EGLTFExporterTextureFlags::None);
+		Texture2D->IsNormalMap() ? EGLTFExporterTextureFlags::Normalmaps :
+		bIsHDR ? EGLTFExporterTextureFlags::HDR : EGLTFExporterTextureFlags::None;
 
 	JsonTexture.Source = Builder.AddImage(Pixels, Size, bIgnoreAlpha, Flags, JsonTexture.Name);
 	JsonTexture.Sampler = Builder.GetOrAddSampler(Texture2D);
