[GLTFExporter] refactor: rename all export option enum types

commit de2df876517b9a381eef0a71ed306e9e8a75c6ef
Date: Mon, 9 Nov 2020 00:16:53 +0100

    refactor: rename all export option enum types
    EGLTFExporterMaterialBakeSize -> EGLTFMaterialBakeSizePOT
    EGLTFExporterTextureCompression -> EGLTFTextureCompression
    EGLTFExporterTextureFlags -> EGLTFTextureGroupFlags
    EGLTFExporterTextureHDREncoding -> EGLTFTextureHDREncoding
    EGLTFExporterLightMobility -> EGLTFExportLightMobility
    (cherry picked from commit 7d10840c7503a33224fd7d052651c95adb98fbe4​)

#skipci

diff --git a/Engine/Plugins/Enterprise/GLTFExporter/MigratedFromGit/Source/GLTFExporter/Private/Builders/GLTFBuilder.cpp b/Engine/Plugins/Enterprise/GLTFExporter/MigratedFromGit/Source/GLTFExporter/Private/Builders/GLTFBuilder.cpp
--- a/Engine/Plugins/Enterprise/GLTFExporter/MigratedFromGit/Source/GLTFExporter/Private/Builders/GLTFBuilder.cpp
+++ b/Engine/Plugins/Enterprise/GLTFExporter/MigratedFromGit/Source/GLTFExporter/Private/Builders/GLTFBuilder.cpp
@@ -21,8 +21,8 @@ EGLTFJsonHDREncoding FGLTFBuilder::GetTextureHDREncoding() const
 {
 	switch (ExportOptions->TextureHDREncoding)
 	{
-		case EGLTFExporterTextureHDREncoding::None: return EGLTFJsonHDREncoding::None;
-		case EGLTFExporterTextureHDREncoding::RGBM: return EGLTFJsonHDREncoding::RGBM;
+		case EGLTFTextureHDREncoding::None: return EGLTFJsonHDREncoding::None;
+		case EGLTFTextureHDREncoding::RGBM: return EGLTFJsonHDREncoding::RGBM;
 		// TODO: add more encodings (like RGBE) when viewer supports them
 		default:
 			checkNoEntry();
@@ -34,11 +34,11 @@ bool FGLTFBuilder::ShouldExportLight(EComponentMobility::Type LightMobility) con
 {
 	switch (ExportOptions->ExportLights)
 	{
-		case EGLTFExporterLightMobility::All:
+		case EGLTFExportLightMobility::All:
 			return true;
-		case EGLTFExporterLightMobility::MovableAndStationary:
+		case EGLTFExportLightMobility::MovableAndStationary:
 			return LightMobility == EComponentMobility::Movable || LightMobility == EComponentMobility::Stationary;
-		case EGLTFExporterLightMobility::MovableOnly:
+		case EGLTFExportLightMobility::MovableOnly:
 			return LightMobility == EComponentMobility::Movable;
 		default:
 			return false;
diff --git a/Engine/Plugins/Enterprise/GLTFExporter/MigratedFromGit/Source/GLTFExporter/Private/Builders/GLTFImageBuilder.cpp b/Engine/Plugins/Enterprise/GLTFExporter/MigratedFromGit/Source/GLTFExporter/Private/Builders/GLTFImageBuilder.cpp
--- a/Engine/Plugins/Enterprise/GLTFExporter/MigratedFromGit/Source/GLTFExporter/Private/Builders/GLTFImageBuilder.cpp
+++ b/Engine/Plugins/Enterprise/GLTFExporter/MigratedFromGit/Source/GLTFExporter/Private/Builders/GLTFImageBuilder.cpp
@@ -10,19 +10,19 @@ FGLTFImageBuilder::FGLTFImageBuilder(const FString& FilePath, const UGLTFExportO
 {
 }
 
-FGLTFJsonImageIndex FGLTFImageBuilder::AddImage(const TArray<FColor>& Pixels, FIntPoint Size, bool bIgnoreAlpha, EGLTFExporterTextureFlags Flags, const FString& Name)
+FGLTFJsonImageIndex FGLTFImageBuilder::AddImage(const TArray<FColor>& Pixels, FIntPoint Size, bool bIgnoreAlpha, EGLTFTextureGroupFlags Flags, const FString& Name)
 {
 	check(Pixels.Num() == Size.X * Size.Y);
 	return AddImage(Pixels.GetData(), Size, bIgnoreAlpha, Flags, Name);
 }
 
-FGLTFJsonImageIndex FGLTFImageBuilder::AddImage(const FColor* Pixels, int64 ByteLength, FIntPoint Size, bool bIgnoreAlpha, EGLTFExporterTextureFlags Flags, const FString& Name)
+FGLTFJsonImageIndex FGLTFImageBuilder::AddImage(const FColor* Pixels, int64 ByteLength, FIntPoint Size, bool bIgnoreAlpha, EGLTFTextureGroupFlags Flags, const FString& Name)
 {
 	check(ByteLength == Size.X * Size.Y * sizeof(FColor));
 	return AddImage(Pixels, Size, bIgnoreAlpha, Flags, Name);
 }
 
-FGLTFJsonImageIndex FGLTFImageBuilder::AddImage(const FColor* Pixels, FIntPoint Size, bool bIgnoreAlpha, EGLTFExporterTextureFlags Flags, const FString& Name)
+FGLTFJsonImageIndex FGLTFImageBuilder::AddImage(const FColor* Pixels, FIntPoint Size, bool bIgnoreAlpha, EGLTFTextureGroupFlags Flags, const FString& Name)
 {
 	TArray64<uint8> CompressedData;
 
@@ -73,16 +73,16 @@ FGLTFJsonImageIndex FGLTFImageBuilder::AddImage(const void* CompressedData, int6
 	return ImageIndex;
 }
 
-EGLTFJsonMimeType FGLTFImageBuilder::GetImageFormat(const FColor* Pixels, FIntPoint Size, bool bIgnoreAlpha, EGLTFExporterTextureFlags Flags) const
+EGLTFJsonMimeType FGLTFImageBuilder::GetImageFormat(const FColor* Pixels, FIntPoint Size, bool bIgnoreAlpha, EGLTFTextureGroupFlags Flags) const
 {
 	switch (ExportOptions->TextureCompression)
 	{
-		case EGLTFExporterTextureCompression::PNG:
+		case EGLTFTextureCompression::PNG:
 			return EGLTFJsonMimeType::PNG;
 
-		case EGLTFExporterTextureCompression::JPEG:
+		case EGLTFTextureCompression::JPEG:
 			return
-				!EnumHasAllFlags(static_cast<EGLTFExporterTextureFlags>(ExportOptions->LosslessCompressTextures), Flags) &&
+				!EnumHasAllFlags(static_cast<EGLTFTextureGroupFlags>(ExportOptions->LosslessCompressTextures), Flags) &&
 				(bIgnoreAlpha || FGLTFImageUtility::NoAlphaNeeded(Pixels, Size)) ?
 				EGLTFJsonMimeType::JPEG : EGLTFJsonMimeType::PNG;
 
diff --git a/Engine/Plugins/Enterprise/GLTFExporter/MigratedFromGit/Source/GLTFExporter/Private/Builders/GLTFImageBuilder.h b/Engine/Plugins/Enterprise/GLTFExporter/MigratedFromGit/Source/GLTFExporter/Private/Builders/GLTFImageBuilder.h
--- a/Engine/Plugins/Enterprise/GLTFExporter/MigratedFromGit/Source/GLTFExporter/Private/Builders/GLTFImageBuilder.h
+++ b/Engine/Plugins/Enterprise/GLTFExporter/MigratedFromGit/Source/GLTFExporter/Private/Builders/GLTFImageBuilder.h
@@ -13,15 +13,15 @@ protected:
 
 public:
 
-	FGLTFJsonImageIndex AddImage(const TArray<FColor>& Pixels, FIntPoint Size, bool bIgnoreAlpha, EGLTFExporterTextureFlags Flags, const FString& Name);
-	FGLTFJsonImageIndex AddImage(const FColor* Pixels, int64 ByteLength, FIntPoint Size, bool bIgnoreAlpha, EGLTFExporterTextureFlags Flags, const FString& Name);
+	FGLTFJsonImageIndex AddImage(const TArray<FColor>& Pixels, FIntPoint Size, bool bIgnoreAlpha, EGLTFTextureGroupFlags Flags, const FString& Name);
+	FGLTFJsonImageIndex AddImage(const FColor* Pixels, int64 ByteLength, FIntPoint Size, bool bIgnoreAlpha, EGLTFTextureGroupFlags Flags, const FString& Name);
 
 private:
 
-	FGLTFJsonImageIndex AddImage(const FColor* Pixels, FIntPoint Size, bool bIgnoreAlpha, EGLTFExporterTextureFlags Flags, const FString& Name);
+	FGLTFJsonImageIndex AddImage(const FColor* Pixels, FIntPoint Size, bool bIgnoreAlpha, EGLTFTextureGroupFlags Flags, const FString& Name);
 	FGLTFJsonImageIndex AddImage(const void* CompressedData, int64 CompressedByteLength, EGLTFJsonMimeType MimeType, const FString& Name);
 
-	EGLTFJsonMimeType GetImageFormat(const FColor* Pixels, FIntPoint Size, bool bIgnoreAlpha, EGLTFExporterTextureFlags Flags) const;
+	EGLTFJsonMimeType GetImageFormat(const FColor* Pixels, FIntPoint Size, bool bIgnoreAlpha, EGLTFTextureGroupFlags Flags) const;
 
 	FString SaveImageToFile(const void* CompressedData, int64 CompressedByteLength, EGLTFJsonMimeType MimeType, const FString& Name);
 
diff --git a/Engine/Plugins/Enterprise/GLTFExporter/MigratedFromGit/Source/GLTFExporter/Private/Converters/GLTFMaterialUtility.cpp b/Engine/Plugins/Enterprise/GLTFExporter/MigratedFromGit/Source/GLTFExporter/Private/Converters/GLTFMaterialUtility.cpp
--- a/Engine/Plugins/Enterprise/GLTFExporter/MigratedFromGit/Source/GLTFExporter/Private/Converters/GLTFMaterialUtility.cpp
+++ b/Engine/Plugins/Enterprise/GLTFExporter/MigratedFromGit/Source/GLTFExporter/Private/Converters/GLTFMaterialUtility.cpp
@@ -344,7 +344,7 @@ FGLTFJsonTextureIndex FGLTFMaterialUtility::AddTexture(FGLTFConvertBuilder& Buil
 	FGLTFJsonTexture JsonTexture;
 	JsonTexture.Name = TextureName;
 	JsonTexture.Sampler = Builder.AddSampler(JsonSampler);
-	JsonTexture.Source = Builder.AddImage(Pixels, TextureSize, bIgnoreAlpha, bIsNormalMap ? EGLTFExporterTextureFlags::Normalmaps : EGLTFExporterTextureFlags::None, TextureName);
+	JsonTexture.Source = Builder.AddImage(Pixels, TextureSize, bIgnoreAlpha, bIsNormalMap ? EGLTFTextureGroupFlags::Normalmaps : EGLTFTextureGroupFlags::None, TextureName);
 
 	return Builder.AddTexture(JsonTexture);
 }
diff --git a/Engine/Plugins/Enterprise/GLTFExporter/MigratedFromGit/Source/GLTFExporter/Private/GLTFExportOptions.cpp b/Engine/Plugins/Enterprise/GLTFExporter/MigratedFromGit/Source/GLTFExporter/Private/GLTFExportOptions.cpp
--- a/Engine/Plugins/Enterprise/GLTFExporter/MigratedFromGit/Source/GLTFExporter/Private/GLTFExportOptions.cpp
+++ b/Engine/Plugins/Enterprise/GLTFExporter/MigratedFromGit/Source/GLTFExporter/Private/GLTFExportOptions.cpp
@@ -18,7 +18,7 @@ void UGLTFExportOptions::ResetToDefault()
 	bExportClearCoatMaterials = true;
 	bBakeMaterialInputs = true;
 	bMaterialBakeUsingMeshData = false;
-	DefaultMaterialBakeSize = EGLTFExporterMaterialBakeSize::POT_1024;
+	DefaultMaterialBakeSize = EGLTFMaterialBakeSizePOT::POT_1024;
 	bExportVertexColors = false;
 	bExportMeshQuantization = true;
 	DefaultLevelOfDetail = 0;
@@ -27,15 +27,15 @@ void UGLTFExportOptions::ResetToDefault()
 	bExportAnimationSequences = true;
 	bRetargetBoneTransforms = true;
 	bExportPlaybackSettings = true;
-	TextureCompression = EGLTFExporterTextureCompression::PNG;
+	TextureCompression = EGLTFTextureCompression::PNG;
 	TextureCompressionQuality = 0;
-	LosslessCompressTextures = static_cast<int32>(EGLTFExporterTextureFlags::All);
+	LosslessCompressTextures = static_cast<int32>(EGLTFTextureGroupFlags::All);
 	bExportTextureTransforms = true;
 	bExportLightmaps = true;
-	TextureHDREncoding = EGLTFExporterTextureHDREncoding::RGBM;
+	TextureHDREncoding = EGLTFTextureHDREncoding::RGBM;
 	ExportScale = 0.01;
 	bExportHiddenInGame = false;
-	ExportLights = EGLTFExporterLightMobility::MovableAndStationary;
+	ExportLights = EGLTFExportLightMobility::MovableAndStationary;
 	bExportCameras = true;
 	bExportCameraControls = true;
 	bExportHDRIBackdrops = true;
diff --git a/Engine/Plugins/Enterprise/GLTFExporter/MigratedFromGit/Source/GLTFExporter/Private/Tasks/GLTFTextureTasks.cpp b/Engine/Plugins/Enterprise/GLTFExporter/MigratedFromGit/Source/GLTFExporter/Private/Tasks/GLTFTextureTasks.cpp
--- a/Engine/Plugins/Enterprise/GLTFExporter/MigratedFromGit/Source/GLTFExporter/Private/Tasks/GLTFTextureTasks.cpp
+++ b/Engine/Plugins/Enterprise/GLTFExporter/MigratedFromGit/Source/GLTFExporter/Private/Tasks/GLTFTextureTasks.cpp
@@ -36,9 +36,9 @@ void FGLTFTexture2DTask::Complete()
 	}
 
 	const bool bIgnoreAlpha = FGLTFTextureUtility::IsAlphaless(Texture2D->GetPixelFormat());
-	const EGLTFExporterTextureFlags Flags =
-		Texture2D->IsNormalMap() ? EGLTFExporterTextureFlags::Normalmaps :
-		bIsHDR ? EGLTFExporterTextureFlags::HDR : EGLTFExporterTextureFlags::None;
+	const EGLTFTextureGroupFlags Flags =
+		Texture2D->IsNormalMap() ? EGLTFTextureGroupFlags::Normalmaps :
+		bIsHDR ? EGLTFTextureGroupFlags::HDR : EGLTFTextureGroupFlags::None;
 
 	JsonTexture.Source = Builder.AddImage(Pixels, Size, bIgnoreAlpha, Flags, JsonTexture.Name);
 	JsonTexture.Sampler = Builder.GetOrAddSampler(Texture2D);
@@ -79,7 +79,7 @@ void FGLTFTextureCubeTask::Complete()
 	}
 
 	const bool bIgnoreAlpha = FGLTFTextureUtility::IsAlphaless(TextureCube->GetPixelFormat());
-	const EGLTFExporterTextureFlags Flags = bIsHDR ? EGLTFExporterTextureFlags::HDR : EGLTFExporterTextureFlags::None;
+	const EGLTFTextureGroupFlags Flags = bIsHDR ? EGLTFTextureGroupFlags::HDR : EGLTFTextureGroupFlags::None;
 
 	JsonTexture.Source = Builder.AddImage(Pixels, Size, bIgnoreAlpha, Flags, JsonTexture.Name);
 	JsonTexture.Sampler = Builder.GetOrAddSampler(TextureCube);
@@ -106,7 +106,7 @@ void FGLTFTextureRenderTarget2DTask::Complete()
 	}
 
 	const bool bIgnoreAlpha = FGLTFTextureUtility::IsAlphaless(RenderTarget2D->GetFormat());
-	const EGLTFExporterTextureFlags Flags = bIsHDR ? EGLTFExporterTextureFlags::HDR : EGLTFExporterTextureFlags::None;
+	const EGLTFTextureGroupFlags Flags = bIsHDR ? EGLTFTextureGroupFlags::HDR : EGLTFTextureGroupFlags::None;
 
 	JsonTexture.Source = Builder.AddImage(Pixels, Size, bIgnoreAlpha, Flags, JsonTexture.Name);
 	JsonTexture.Sampler = Builder.GetOrAddSampler(RenderTarget2D);
@@ -146,7 +146,7 @@ void FGLTFTextureRenderTargetCubeTask::Complete()
 	}
 
 	const bool bIgnoreAlpha = FGLTFTextureUtility::IsAlphaless(RenderTargetCube->GetFormat());
-	const EGLTFExporterTextureFlags Flags = bIsHDR ? EGLTFExporterTextureFlags::HDR : EGLTFExporterTextureFlags::None;
+	const EGLTFTextureGroupFlags Flags = bIsHDR ? EGLTFTextureGroupFlags::HDR : EGLTFTextureGroupFlags::None;
 
 	JsonTexture.Source = Builder.AddImage(Pixels, Size, bIgnoreAlpha, Flags, JsonTexture.Name);
 	JsonTexture.Sampler = Builder.GetOrAddSampler(RenderTargetCube);
@@ -178,7 +178,7 @@ void FGLTFTextureLightMapTask::Complete()
 	const int64 ByteLength = Source.CalcMipSize(0);
 
 	const bool bIgnoreAlpha = false;
-	const EGLTFExporterTextureFlags Flags = EGLTFExporterTextureFlags::Lightmaps;
+	const EGLTFTextureGroupFlags Flags = EGLTFTextureGroupFlags::Lightmaps;
 
 	const void* RawData = Source.LockMip(0);
 	JsonTexture.Source = Builder.AddImage(static_cast<const FColor*>(RawData), ByteLength, Size, bIgnoreAlpha, Flags, JsonTexture.Name);
diff --git a/Engine/Plugins/Enterprise/GLTFExporter/MigratedFromGit/Source/GLTFExporter/Public/GLTFExportOptions.h b/Engine/Plugins/Enterprise/GLTFExporter/MigratedFromGit/Source/GLTFExporter/Public/GLTFExportOptions.h
--- a/Engine/Plugins/Enterprise/GLTFExporter/MigratedFromGit/Source/GLTFExporter/Public/GLTFExportOptions.h
+++ b/Engine/Plugins/Enterprise/GLTFExporter/MigratedFromGit/Source/GLTFExporter/Public/GLTFExportOptions.h
@@ -6,7 +6,7 @@
 #include "GLTFExportOptions.generated.h"
 
 UENUM(BlueprintType)
-enum class EGLTFExporterMaterialBakeSize : uint8
+enum class EGLTFMaterialBakeSizePOT : uint8
 {
 	POT_1 UMETA(DisplayName = "1 x 1"),
 	POT_2 UMETA(DisplayName = "2 x 2"),
@@ -25,14 +25,14 @@ enum class EGLTFExporterMaterialBakeSize : uint8
 };
 
 UENUM(BlueprintType)
-enum class EGLTFExporterTextureCompression : uint8
+enum class EGLTFTextureCompression : uint8
 {
 	PNG,
     JPEG UMETA(DisplayName = "JPEG (if no alpha)")
 };
 
 UENUM(BlueprintType, Meta = (Bitflags, UseEnumValuesAsMaskValuesInEditor = "true"))
-enum class EGLTFExporterTextureFlags : uint8
+enum class EGLTFTextureGroupFlags : uint8
 {
 	None = 0 UMETA(Hidden),
 
@@ -42,17 +42,17 @@ enum class EGLTFExporterTextureFlags : uint8
 
 	All = HDR | Normalmaps | Lightmaps UMETA(Hidden)
 };
-ENUM_CLASS_FLAGS(EGLTFExporterTextureFlags);
+ENUM_CLASS_FLAGS(EGLTFTextureGroupFlags);
 
 UENUM(BlueprintType)
-enum class EGLTFExporterTextureHDREncoding : uint8
+enum class EGLTFTextureHDREncoding : uint8
 {
 	None,
 	RGBM
 };
 
 UENUM(BlueprintType)
-enum class EGLTFExporterLightMobility : uint8
+enum class EGLTFExportLightMobility : uint8
 {
 	None,
 	MovableOnly,
@@ -80,7 +80,7 @@ public:
 	bool bMaterialBakeUsingMeshData;
 
 	UPROPERTY(EditAnywhere, BlueprintReadWrite, Config, Category = Material, Meta = (EditCondition = "bBakeMaterialInputs"))
-	EGLTFExporterMaterialBakeSize DefaultMaterialBakeSize;
+	EGLTFMaterialBakeSizePOT DefaultMaterialBakeSize;
 
 	UPROPERTY(EditAnywhere, BlueprintReadWrite, Config, Category = Mesh)
 	bool bExportVertexColors;
@@ -107,12 +107,12 @@ public:
 	bool bExportPlaybackSettings;
 
 	UPROPERTY(EditAnywhere, BlueprintReadWrite, Config, Category = Texture)
-	EGLTFExporterTextureCompression TextureCompression;
+	EGLTFTextureCompression TextureCompression;
 
 	UPROPERTY(EditAnywhere, BlueprintReadWrite, Config, Category = Texture, Meta = (ClampMin = "0", ClampMax = "100", EditCondition = "TextureCompression != EGLTFExporterTextureCompression::PNG"))
 	int32 TextureCompressionQuality;
 
-	UPROPERTY(EditAnywhere, BlueprintReadWrite, Config, Category = Texture, Meta = (Bitmask, BitmaskEnum = EGLTFExporterTextureFlags, EditCondition = "TextureCompression != EGLTFExporterTextureCompression::PNG"))
+	UPROPERTY(EditAnywhere, BlueprintReadWrite, Config, Category = Texture, Meta = (Bitmask, BitmaskEnum = EGLTFTextureGroupFlags, EditCondition = "TextureCompression != EGLTFExporterTextureCompression::PNG"))
 	int32 LosslessCompressTextures;
 
 	UPROPERTY(EditAnywhere, BlueprintReadWrite, Config, Category = Texture)
@@ -122,7 +122,7 @@ public:
 	bool bExportLightmaps;
 
 	UPROPERTY(EditAnywhere, BlueprintReadWrite, Config, Category = Texture)
-	EGLTFExporterTextureHDREncoding TextureHDREncoding;
+	EGLTFTextureHDREncoding TextureHDREncoding;
 
 	UPROPERTY(EditAnywhere, BlueprintReadWrite, Config, Category = Scene, Meta =
 		(DisplayName = "Export Uniform Scale",
@@ -133,7 +133,7 @@ public:
 	bool bExportHiddenInGame;
 
 	UPROPERTY(EditAnywhere, BlueprintReadWrite, Config, Category = Scene)
-	EGLTFExporterLightMobility ExportLights;
+	EGLTFExportLightMobility ExportLights;
 
 	UPROPERTY(EditAnywhere, BlueprintReadWrite, Config, Category = Scene)
 	bool bExportCameras;
