[GLTFExporter] refactor: fix display name for export options TextureHDREncoding and bExportHDRIBackdrops

commit 5d4262b61e003f32998b2eff66fd86ea7b9175cc
Date: Mon, 9 Nov 2020 00:17:50 +0100

    refactor: fix display name for export options TextureHDREncoding and bExportHDRIBackdrops
    (cherry picked from commit 918402dec7494d7c1dd40075a7a3cef1cafb3029​)

#skipci

diff --git a/Engine/Plugins/Enterprise/GLTFExporter/MigratedFromGit/Source/GLTFExporter/Public/GLTFExportOptions.h b/Engine/Plugins/Enterprise/GLTFExporter/MigratedFromGit/Source/GLTFExporter/Public/GLTFExportOptions.h
--- a/Engine/Plugins/Enterprise/GLTFExporter/MigratedFromGit/Source/GLTFExporter/Public/GLTFExportOptions.h
+++ b/Engine/Plugins/Enterprise/GLTFExporter/MigratedFromGit/Source/GLTFExporter/Public/GLTFExportOptions.h
@@ -121,7 +121,7 @@ public:
 	UPROPERTY(EditAnywhere, BlueprintReadWrite, Config, Category = Texture)
 	bool bExportLightmaps;
 
-	UPROPERTY(EditAnywhere, BlueprintReadWrite, Config, Category = Texture)
+	UPROPERTY(EditAnywhere, BlueprintReadWrite, Config, Category = Texture, Meta = (DisplayName = "Texture HDR Encoding"))
 	EGLTFTextureHDREncoding TextureHDREncoding;
 
 	UPROPERTY(EditAnywhere, BlueprintReadWrite, Config, Category = Scene, Meta =
@@ -141,7 +141,7 @@ public:
 	UPROPERTY(EditAnywhere, BlueprintReadWrite, Config, Category = Scene, Meta = (EditCondition = "bExportCameras"))
 	bool bExportCameraControls;
 
-	UPROPERTY(EditAnywhere, BlueprintReadWrite, Config, Category = Scene)
+	UPROPERTY(EditAnywhere, BlueprintReadWrite, Config, Category = Scene, Meta = (DisplayName = "Export HDRI Backdrops"))
 	bool bExportHDRIBackdrops;
 
 	UPROPERTY(EditAnywhere, BlueprintReadWrite, Config, Category = Scene)
