[GLTFExporter] refactor: update some TODO comments

commit dd57f7a365e4c354de48e90109b5e784729ae81c
Date: Mon, 9 Nov 2020 10:52:25 +0100

    refactor: update some TODO comments
    (cherry picked from commit 1dbe0536172a283304c01e1367b59f66ff2218a5​)

#skipci

diff --git a/Engine/Plugins/Enterprise/GLTFExporter/MigratedFromGit/Source/GLTFExporter/Private/Converters/GLTFBackdropConverters.cpp b/Engine/Plugins/Enterprise/GLTFExporter/MigratedFromGit/Source/GLTFExporter/Private/Converters/GLTFBackdropConverters.cpp
--- a/Engine/Plugins/Enterprise/GLTFExporter/MigratedFromGit/Source/GLTFExporter/Private/Converters/GLTFBackdropConverters.cpp
+++ b/Engine/Plugins/Enterprise/GLTFExporter/MigratedFromGit/Source/GLTFExporter/Private/Converters/GLTFBackdropConverters.cpp
@@ -30,7 +30,6 @@ FGLTFJsonBackdropIndex FGLTFBackdropConverter::Convert(const AActor* BackdropAct
 	const UStaticMesh* Mesh;
 	if (FGLTFActorUtility::TryGetPropertyValue(BackdropActor, TEXT("Mesh"), Mesh))
 	{
-		// TODO: avoid exporting the mesh material (use gltf default material)
 		JsonBackdrop.Mesh = Builder.GetOrAddMesh(Mesh, { FGLTFMaterialUtility::GetDefault() });
 	}
 	else
diff --git a/Engine/Plugins/Enterprise/GLTFExporter/MigratedFromGit/Source/GLTFExporter/Private/Converters/GLTFVarationConverters.cpp b/Engine/Plugins/Enterprise/GLTFExporter/MigratedFromGit/Source/GLTFExporter/Private/Converters/GLTFVarationConverters.cpp
--- a/Engine/Plugins/Enterprise/GLTFExporter/MigratedFromGit/Source/GLTFExporter/Private/Converters/GLTFVarationConverters.cpp
+++ b/Engine/Plugins/Enterprise/GLTFExporter/MigratedFromGit/Source/GLTFExporter/Private/Converters/GLTFVarationConverters.cpp
@@ -90,7 +90,7 @@ bool FGLTFVariationConverter::TryParseVariant(FGLTFJsonVariant& OutVariant, cons
 
 	if (const UTexture2D* Thumbnail = const_cast<UVariant*>(Variant)->GetThumbnail())
 	{
-		// TODO: if thumbnail is generated from viewport, give it a variant-relevant name (instead of "Texture2D_#")
+		// TODO: if thumbnail has generic name "Texture2D", give it a variant-relevant name
 		JsonVariant.Thumbnail = Builder.GetOrAddTexture(Thumbnail);
 	}
 
