[GLTFExporter] refactor: change parameter PreferredTextureSize (in FGLTFMaterialTask::BakeMaterialProperty) to from pass-by-reference to pass-by-value

commit 5272c7ae9cba77063eae8ecbff323c560de323af
Date: Mon, 9 Nov 2020 20:29:24 +0100

    refactor: change parameter PreferredTextureSize (in FGLTFMaterialTask::BakeMaterialProperty) to from pass-by-reference to pass-by-value
    (cherry picked from commit 0387013a0774bd5658b4eaecd82267307fc19f85​)

#skipci

diff --git a/Engine/Plugins/Enterprise/GLTFExporter/MigratedFromGit/Source/GLTFExporter/Private/Tasks/GLTFMaterialTasks.cpp b/Engine/Plugins/Enterprise/GLTFExporter/MigratedFromGit/Source/GLTFExporter/Private/Tasks/GLTFMaterialTasks.cpp
--- a/Engine/Plugins/Enterprise/GLTFExporter/MigratedFromGit/Source/GLTFExporter/Private/Tasks/GLTFMaterialTasks.cpp
+++ b/Engine/Plugins/Enterprise/GLTFExporter/MigratedFromGit/Source/GLTFExporter/Private/Tasks/GLTFMaterialTasks.cpp
@@ -357,8 +357,8 @@ bool FGLTFMaterialTask::TryGetBaseColorAndOpacity(FGLTFJsonPBRMetallicRoughness&
 		TextureMagFilter = FGLTFConverterUtility::ConvertMagFilter(OpacityTexture->Filter, OpacityTexture->LODGroup);
 	}
 
-	const FGLTFPropertyBakeOutput BaseColorBakeOutput = BakeMaterialProperty(BaseColorProperty, BaseColorTexCoord, &TextureSize);
-	const FGLTFPropertyBakeOutput OpacityBakeOutput = BakeMaterialProperty(OpacityProperty, OpacityTexCoord, &TextureSize, true);
+	const FGLTFPropertyBakeOutput BaseColorBakeOutput = BakeMaterialProperty(BaseColorProperty, BaseColorTexCoord, TextureSize);
+	const FGLTFPropertyBakeOutput OpacityBakeOutput = BakeMaterialProperty(OpacityProperty, OpacityTexCoord, TextureSize, true);
 	const float BaseColorScale = BaseColorProperty == MP_EmissiveColor ? BaseColorBakeOutput.EmissiveScale : 1;
 
 	// Detect when both baked properties are constants, which means we can avoid exporting a texture
@@ -525,8 +525,8 @@ bool FGLTFMaterialTask::TryGetMetallicAndRoughness(FGLTFJsonPBRMetallicRoughness
 		TextureMagFilter = FGLTFConverterUtility::ConvertMagFilter(RoughnessTexture->Filter, RoughnessTexture->LODGroup);
 	}
 
-	const FGLTFPropertyBakeOutput MetallicBakeOutput = BakeMaterialProperty(MetallicProperty, MetallicTexCoord, &TextureSize);
-	const FGLTFPropertyBakeOutput RoughnessBakeOutput = BakeMaterialProperty(RoughnessProperty, RoughnessTexCoord, &TextureSize);
+	const FGLTFPropertyBakeOutput MetallicBakeOutput = BakeMaterialProperty(MetallicProperty, MetallicTexCoord, TextureSize);
+	const FGLTFPropertyBakeOutput RoughnessBakeOutput = BakeMaterialProperty(RoughnessProperty, RoughnessTexCoord, TextureSize);
 
 	// Detect when both baked properties are constants, which means we can use factors and avoid exporting a texture
 	if (MetallicBakeOutput.bIsConstant && RoughnessBakeOutput.bIsConstant)
@@ -692,8 +692,8 @@ bool FGLTFMaterialTask::TryGetClearCoatRoughness(FGLTFJsonClearCoatExtension& Ou
 		TextureMagFilter = FGLTFConverterUtility::ConvertMagFilter(RoughnessTexture->Filter, RoughnessTexture->LODGroup);
 	}
 
-	const FGLTFPropertyBakeOutput IntensityBakeOutput = BakeMaterialProperty(IntensityProperty, IntensityTexCoord, &TextureSize);
-	const FGLTFPropertyBakeOutput RoughnessBakeOutput = BakeMaterialProperty(RoughnessProperty, RoughnessTexCoord, &TextureSize);
+	const FGLTFPropertyBakeOutput IntensityBakeOutput = BakeMaterialProperty(IntensityProperty, IntensityTexCoord, TextureSize);
+	const FGLTFPropertyBakeOutput RoughnessBakeOutput = BakeMaterialProperty(RoughnessProperty, RoughnessTexCoord, TextureSize);
 
 	// Detect when both baked properties are constants, which means we can use factors and avoid exporting a texture
 	if (IntensityBakeOutput.bIsConstant && RoughnessBakeOutput.bIsConstant)
@@ -1322,7 +1322,7 @@ bool FGLTFMaterialTask::TryGetBakedMaterialProperty(FGLTFJsonTextureInfo& OutTex
 	return true;
 }
 
-FGLTFPropertyBakeOutput FGLTFMaterialTask::BakeMaterialProperty(EMaterialProperty Property, int32& OutTexCoord, const FIntPoint* PreferredTextureSize, bool bCopyAlphaFromRedChannel) const
+FGLTFPropertyBakeOutput FGLTFMaterialTask::BakeMaterialProperty(EMaterialProperty Property, int32& OutTexCoord, FIntPoint PreferredTextureSize, bool bCopyAlphaFromRedChannel) const
 {
 	const FExpressionInput* PropertyInput = FGLTFMaterialUtility::GetInputForProperty(Material, Property);
 	TSet<int32> TexCoords;
@@ -1356,7 +1356,7 @@ FGLTFPropertyBakeOutput FGLTFMaterialTask::BakeMaterialProperty(EMaterialPropert
 	}
 
 	const FIntPoint DefaultTextureSize = Builder.GetDefaultMaterialBakeSize();
-	const FIntPoint TextureSize = PreferredTextureSize != nullptr ? *PreferredTextureSize : DefaultTextureSize;
+	const FIntPoint TextureSize = PreferredTextureSize != FIntPoint::ZeroValue ? PreferredTextureSize : DefaultTextureSize;
 
 	// TODO: add support for calculating the ideal resolution to use for baking based on connected (texture) nodes
 
diff --git a/Engine/Plugins/Enterprise/GLTFExporter/MigratedFromGit/Source/GLTFExporter/Private/Tasks/GLTFMaterialTasks.h b/Engine/Plugins/Enterprise/GLTFExporter/MigratedFromGit/Source/GLTFExporter/Private/Tasks/GLTFMaterialTasks.h
--- a/Engine/Plugins/Enterprise/GLTFExporter/MigratedFromGit/Source/GLTFExporter/Private/Tasks/GLTFMaterialTasks.h
+++ b/Engine/Plugins/Enterprise/GLTFExporter/MigratedFromGit/Source/GLTFExporter/Private/Tasks/GLTFMaterialTasks.h
@@ -60,7 +60,7 @@ private:
 	bool TryGetBakedMaterialProperty(FGLTFJsonTextureInfo& OutTexInfo, float& OutConstant, EMaterialProperty Property, const FString& PropertyName) const;
 	bool TryGetBakedMaterialProperty(FGLTFJsonTextureInfo& OutTexInfo, EMaterialProperty Property, const FString& PropertyName) const;
 
-	FGLTFPropertyBakeOutput BakeMaterialProperty(EMaterialProperty Property, int32& OutTexCoord, const FIntPoint* PreferredTextureSize = nullptr, bool bCopyAlphaFromRedChannel = false) const;
+	FGLTFPropertyBakeOutput BakeMaterialProperty(EMaterialProperty Property, int32& OutTexCoord, FIntPoint PreferredTextureSize = FIntPoint::ZeroValue, bool bCopyAlphaFromRedChannel = false) const;
 
 	bool StoreBakedPropertyTexture(FGLTFJsonTextureInfo& OutTexInfo, const FGLTFPropertyBakeOutput& PropertyBakeOutput, const FString& PropertyName) const;
 
