[GLTFExporter] feat: detect shading model (if multiple found in material) more accurately by baking the input expression

commit c46d1e45f81bf0db867b939762afad45376a4214
Date: Mon, 9 Nov 2020 22:02:49 +0100

    feat: detect shading model (if multiple found in material) more accurately by baking the input expression
    (cherry picked from commit bbfdf81c3c0251d9eb412eb0d3fd2ac388f06a68​)

#skipci

diff --git a/Engine/Plugins/Enterprise/GLTFExporter/MigratedFromGit/Source/GLTFExporter/Private/Tasks/GLTFMaterialTasks.cpp b/Engine/Plugins/Enterprise/GLTFExporter/MigratedFromGit/Source/GLTFExporter/Private/Tasks/GLTFMaterialTasks.cpp
--- a/Engine/Plugins/Enterprise/GLTFExporter/MigratedFromGit/Source/GLTFExporter/Private/Tasks/GLTFMaterialTasks.cpp
+++ b/Engine/Plugins/Enterprise/GLTFExporter/MigratedFromGit/Source/GLTFExporter/Private/Tasks/GLTFMaterialTasks.cpp
@@ -207,22 +207,15 @@ bool FGLTFMaterialTask::TryGetShadingModel(EGLTFJsonShadingModel& OutShadingMode
 	EMaterialShadingModel ShadingModel = ShadingModels.GetFirstShadingModel();
 	if (ShadingModelCount > 1)
 	{
-		Builder.AddWarningMessage(FString::Printf(TEXT("Support is limited to only one of the multiple shading models found (%s) in material %s"), *FGLTFMaterialUtility::ShadingModelsToString(ShadingModels), *Material->GetName()));
-
-		if (ShadingModels.HasShadingModel(MSM_ClearCoat))
-		{
-			ShadingModel = MSM_ClearCoat;
-		}
-		else if (ShadingModels.HasShadingModel(MSM_DefaultLit))
-		{
-			ShadingModel = MSM_DefaultLit;
-		}
-		else if (ShadingModels.HasShadingModel(MSM_Unlit))
-		{
-			ShadingModel = MSM_Unlit;
-		}
-
-		Builder.AddWarningMessage(FString::Printf(TEXT("Material %s will be exported as shading model %s"), *Material->GetName(), *FGLTFNameUtility::GetName(ShadingModel)));
+		int32 TexCoord;
+		const FGLTFPropertyBakeOutput BakeOutput = BakeMaterialProperty(MP_ShadingModel, TexCoord, { 1, 1 });
+		ShadingModel = static_cast<EMaterialShadingModel>(BakeOutput.Pixels[0].R);
+
+		Builder.AddWarningMessage(
+			FString::Printf(TEXT("Multiple shading models (%s) in material %s will be reduced to one (%s)"),
+			*FGLTFMaterialUtility::ShadingModelsToString(ShadingModels),
+			*Material->GetName(),
+			*FGLTFNameUtility::GetName(ShadingModel)));
 	}
 
 	const EGLTFJsonShadingModel ConvertedShadingModel = FGLTFConverterUtility::ConvertShadingModel(ShadingModel);
