[GLTFExporter] feat: add support for baking shading model in embedded baking module

commit e8749571bfb67be6ae0a7d6f040d65263b45f272
Date: Mon, 9 Nov 2020 20:22:24 +0100

    feat: add support for baking shading model in embedded baking module
    (cherry picked from commit 960f1cbaab6ec80817084125f04555ee2fda5b6e​)

#skipci

diff --git a/Engine/Plugins/Enterprise/GLTFExporter/MigratedFromGit/Source/GLTFMaterialBaking/GLTFMaterialBaking.Build.cs b/Engine/Plugins/Enterprise/GLTFExporter/MigratedFromGit/Source/GLTFMaterialBaking/GLTFMaterialBaking.Build.cs
--- a/Engine/Plugins/Enterprise/GLTFExporter/MigratedFromGit/Source/GLTFMaterialBaking/GLTFMaterialBaking.Build.cs
+++ b/Engine/Plugins/Enterprise/GLTFExporter/MigratedFromGit/Source/GLTFMaterialBaking/GLTFMaterialBaking.Build.cs
@@ -26,6 +26,9 @@ public class GLTFMaterialBaking : ModuleRules
             }
 		);
 
+		// NOTE: ugly hack to access HLSLMaterialTranslator to bake shading model
+        PrivateIncludePaths.Add(EngineDirectory + "/Source/Runtime/Engine/Private");
+
 		// NOTE: avoid renaming all instaces of MATERIALBAKING_API by redirecting it to GLTFMATERIALBAKING_API
 		PublicDefinitions.Add("MATERIALBAKING_API=GLTFMATERIALBAKING_API");
     }
diff --git a/Engine/Plugins/Enterprise/GLTFExporter/MigratedFromGit/Source/GLTFMaterialBaking/Private/ExportMaterialProxy.h b/Engine/Plugins/Enterprise/GLTFExporter/MigratedFromGit/Source/GLTFMaterialBaking/Private/ExportMaterialProxy.h
--- a/Engine/Plugins/Enterprise/GLTFExporter/MigratedFromGit/Source/GLTFMaterialBaking/Private/ExportMaterialProxy.h
+++ b/Engine/Plugins/Enterprise/GLTFExporter/MigratedFromGit/Source/GLTFMaterialBaking/Private/ExportMaterialProxy.h
@@ -19,6 +19,8 @@
 #include "Materials/Material.h"
 #include "Materials/MaterialExpressionCustomOutput.h"
 
+#include "Materials/HLSLMaterialTranslator.h"
+
 struct FExportMaterialCompiler : public FProxyMaterialCompiler
 {
 	FExportMaterialCompiler(FMaterialCompiler* InCompiler) :
@@ -273,7 +275,8 @@ public:
 		// NOTE: the following cases for custom data 0/1 and custom output are hacks since a proper solution requires engine changes:
 		case MP_CustomData0: ResourceId.Usage = static_cast<EMaterialShaderMapUsage::Type>(EMaterialShaderMapUsage::DebugViewMode + 1); break;
 		case MP_CustomData1: ResourceId.Usage = static_cast<EMaterialShaderMapUsage::Type>(EMaterialShaderMapUsage::DebugViewMode + 2); break;
-		case MP_CustomOutput: ResourceId.Usage = static_cast<EMaterialShaderMapUsage::Type>(EMaterialShaderMapUsage::DebugViewMode + 3); break;
+		case MP_ShadingModel: ResourceId.Usage = static_cast<EMaterialShaderMapUsage::Type>(EMaterialShaderMapUsage::DebugViewMode + 3); break;
+		case MP_CustomOutput: ResourceId.Usage = static_cast<EMaterialShaderMapUsage::Type>(EMaterialShaderMapUsage::DebugViewMode + 4); break;
 
 		default:
 			ensureMsgf(false, TEXT("ExportMaterial has no usage for property %i.  Will likely reuse the normal rendering shader and crash later with a parameter mismatch"), (int32)InPropertyToCompile);
@@ -359,6 +362,10 @@ public:
 		// needs to be called in this function!!
 		Compiler->SetMaterialProperty(Property, OverrideShaderFrequency, bUsePreviousFrameTime);
 		const int32 Ret = CompilePropertyAndSetMaterialPropertyWithoutCast(Property, Compiler);
+		if (Property == MP_ShadingModel)
+		{
+			return Ret; // NOTE: don't try to cast shading model since it will fail
+		}
 		return Compiler->ForceCast(Ret, FMaterialAttributeDefinitionMap::GetValueType(Property), MFCF_ExactMatch | MFCF_ReplicateValue);
 	}
 
@@ -409,7 +416,7 @@ public:
 				}
 				break;
 			case MP_ShadingModel:
-				return MaterialInterface->CompileProperty(&ProxyCompiler, MP_ShadingModel);
+				return CompileShadingModel(Compiler, MaterialInterface->CompileProperty(&ProxyCompiler, MP_ShadingModel));
 			case MP_CustomOutput:
 				 // NOTE: Currently we can assume input index is always 0, which it is for all custom outputs that are registered as material attributes
 				return CompileInputForCustomOutput(Compiler, 0, ForceCast_Exact_Replicate);
@@ -679,6 +686,26 @@ private:
 		return nullptr;
 	}
 
+	static int32 CompileShadingModel(FMaterialCompiler* Compiler, int32 Code)
+	{
+		class FMaterialCompilerHack : public FHLSLMaterialTranslator
+		{
+			using FHLSLMaterialTranslator::FHLSLMaterialTranslator;
+
+		public:
+
+			void SetParameterType(int32 Index, EMaterialValueType Type)
+			{
+				check(Index >= 0 && Index < CurrentScopeChunks->Num());
+				(*CurrentScopeChunks)[Index].Type = Type;
+			}
+		};
+
+		// NOTE: ugly hack to workaround casting shading model (uint) to float
+		static_cast<FMaterialCompilerHack*>(Compiler)->SetParameterType(Code, MCT_Float1);
+		return Compiler->Div(Code, Compiler->Constant(255)); // [0,255] / 255
+	}
+
 private:
 	/** The material interface for this proxy */
 	UMaterialInterface* MaterialInterface;
diff --git a/Engine/Plugins/Enterprise/GLTFExporter/MigratedFromGit/Source/GLTFMaterialBaking/Private/MaterialBakingModule.cpp b/Engine/Plugins/Enterprise/GLTFExporter/MigratedFromGit/Source/GLTFMaterialBaking/Private/MaterialBakingModule.cpp
--- a/Engine/Plugins/Enterprise/GLTFExporter/MigratedFromGit/Source/GLTFMaterialBaking/Private/MaterialBakingModule.cpp
+++ b/Engine/Plugins/Enterprise/GLTFExporter/MigratedFromGit/Source/GLTFMaterialBaking/Private/MaterialBakingModule.cpp
@@ -242,6 +242,7 @@ void FMaterialBakingModule::StartupModule()
 	PerPropertyGamma.Add(MP_Normal);
 	PerPropertyGamma.Add(MP_Opacity);
 	PerPropertyGamma.Add(MP_OpacityMask);
+	PerPropertyGamma.Add(MP_ShadingModel);
 	PerPropertyGamma.Add(TEXT("ClearCoatBottomNormal"));
 
 	// Set which pixel format should be used for the possible baked out material properties
@@ -259,6 +260,7 @@ void FMaterialBakingModule::StartupModule()
 	PerPropertyFormat.Add(MP_SubsurfaceColor, PF_B8G8R8A8);
 	PerPropertyFormat.Add(MP_CustomData0, PF_B8G8R8A8);
 	PerPropertyFormat.Add(MP_CustomData1, PF_B8G8R8A8);
+	PerPropertyFormat.Add(MP_ShadingModel, PF_B8G8R8A8);
 	PerPropertyFormat.Add(TEXT("ClearCoatBottomNormal"), PF_B8G8R8A8);
 
 	// Register property customization
