[GLTFExporter] refactor: add sanity check for evaluating shading model and remove warning about reduction

commit b0cbad55b75084089da482a59ddb6322bee2bbe4
Date: Mon, 9 Nov 2020 22:29:26 +0100

    refactor: add sanity check for evaluating shading model and remove warning about reduction
    (cherry picked from commit c47ff22421bb0e0ffb3c2ab85f50367ef9963323​)

#skipci

diff --git a/Engine/Plugins/Enterprise/GLTFExporter/MigratedFromGit/Source/GLTFExporter/Private/Converters/GLTFMaterialUtility.cpp b/Engine/Plugins/Enterprise/GLTFExporter/MigratedFromGit/Source/GLTFExporter/Private/Converters/GLTFMaterialUtility.cpp
--- a/Engine/Plugins/Enterprise/GLTFExporter/MigratedFromGit/Source/GLTFExporter/Private/Converters/GLTFMaterialUtility.cpp
+++ b/Engine/Plugins/Enterprise/GLTFExporter/MigratedFromGit/Source/GLTFExporter/Private/Converters/GLTFMaterialUtility.cpp
@@ -452,23 +452,6 @@ void FGLTFMaterialUtility::ExpandAllFunctionExpressions(TArray<UMaterialExpressi
 	});
 }
 
-FString FGLTFMaterialUtility::ShadingModelsToString(const FMaterialShadingModelField& ShadingModels)
-{
-	FString Result;
-
-	for (uint32 Index = 0; Index < MSM_MAX; Index++)
-	{
-		const EMaterialShadingModel ShadingModel = static_cast<EMaterialShadingModel>(Index);
-		if (ShadingModels.HasShadingModel(ShadingModel))
-		{
-			FString Name = FGLTFNameUtility::GetName(ShadingModel);
-			Result += Result.IsEmpty() ? Name : TEXT(", ") + Name;
-		}
-	}
-
-	return Result;
-}
-
 bool FGLTFMaterialUtility::NeedsMeshData(const UMaterialInterface* Material)
 {
 	if (Material == nullptr)
diff --git a/Engine/Plugins/Enterprise/GLTFExporter/MigratedFromGit/Source/GLTFExporter/Private/Converters/GLTFMaterialUtility.h b/Engine/Plugins/Enterprise/GLTFExporter/MigratedFromGit/Source/GLTFExporter/Private/Converters/GLTFMaterialUtility.h
--- a/Engine/Plugins/Enterprise/GLTFExporter/MigratedFromGit/Source/GLTFExporter/Private/Converters/GLTFMaterialUtility.h
+++ b/Engine/Plugins/Enterprise/GLTFExporter/MigratedFromGit/Source/GLTFExporter/Private/Converters/GLTFMaterialUtility.h
@@ -68,8 +68,6 @@ struct FGLTFMaterialUtility
 
 	static void ExpandAllFunctionExpressions(TArray<UMaterialExpression*>& InOutExpressions);
 
-	static FString ShadingModelsToString(const FMaterialShadingModelField& ShadingModels);
-
 	static bool NeedsMeshData(const UMaterialInterface* Material);
 	static bool NeedsMeshData(const TArray<const UMaterialInterface*>& Materials);
 
diff --git a/Engine/Plugins/Enterprise/GLTFExporter/MigratedFromGit/Source/GLTFExporter/Private/Tasks/GLTFMaterialTasks.cpp b/Engine/Plugins/Enterprise/GLTFExporter/MigratedFromGit/Source/GLTFExporter/Private/Tasks/GLTFMaterialTasks.cpp
--- a/Engine/Plugins/Enterprise/GLTFExporter/MigratedFromGit/Source/GLTFExporter/Private/Tasks/GLTFMaterialTasks.cpp
+++ b/Engine/Plugins/Enterprise/GLTFExporter/MigratedFromGit/Source/GLTFExporter/Private/Tasks/GLTFMaterialTasks.cpp
@@ -211,11 +211,14 @@ bool FGLTFMaterialTask::TryGetShadingModel(EGLTFJsonShadingModel& OutShadingMode
 		const FGLTFPropertyBakeOutput BakeOutput = BakeMaterialProperty(MP_ShadingModel, TexCoord, { 1, 1 });
 		ShadingModel = static_cast<EMaterialShadingModel>(BakeOutput.Pixels[0].R);
 
-		Builder.AddWarningMessage(
-			FString::Printf(TEXT("Multiple shading models (%s) in material %s will be reduced to one (%s)"),
-			*FGLTFMaterialUtility::ShadingModelsToString(ShadingModels),
-			*Material->GetName(),
-			*FGLTFNameUtility::GetName(ShadingModel)));
+		if (!ShadingModels.HasShadingModel(ShadingModel))
+		{
+			ShadingModel = ShadingModels.GetFirstShadingModel();
+			Builder.AddWarningMessage(
+				FString::Printf(TEXT("Failed to properly evaluate shading model expression in material %s, will export as %s"),
+				*Material->GetName(),
+				*FGLTFNameUtility::GetName(ShadingModel)));
+		}
 	}
 
 	const EGLTFJsonShadingModel ConvertedShadingModel = FGLTFConverterUtility::ConvertShadingModel(ShadingModel);
