[GLTFExporter] feat: improve evaluation (aka baking) of shading model expression in material converter

commit 1145a53dc24572e9207da0618e224549060d3107
Date: Tue, 10 Nov 2020 13:53:14 +0100

    feat: improve evaluation (aka baking) of shading model expression in material converter
    (cherry picked from commit 19a7764f3bae6fe00c411432cbe7a7a424882287​)

#skipci

diff --git a/Engine/Plugins/Enterprise/GLTFExporter/MigratedFromGit/Source/GLTFExporter/Private/Converters/GLTFMaterialUtility.cpp b/Engine/Plugins/Enterprise/GLTFExporter/MigratedFromGit/Source/GLTFExporter/Private/Converters/GLTFMaterialUtility.cpp
--- a/Engine/Plugins/Enterprise/GLTFExporter/MigratedFromGit/Source/GLTFExporter/Private/Converters/GLTFMaterialUtility.cpp
+++ b/Engine/Plugins/Enterprise/GLTFExporter/MigratedFromGit/Source/GLTFExporter/Private/Converters/GLTFMaterialUtility.cpp
@@ -452,6 +452,45 @@ void FGLTFMaterialUtility::ExpandAllFunctionExpressions(TArray<UMaterialExpressi
 	});
 }
 
+EMaterialShadingModel FGLTFMaterialUtility::GetRichestShadingModel(const FMaterialShadingModelField& ShadingModels)
+{
+	if (ShadingModels.HasShadingModel(MSM_ClearCoat))
+	{
+		return MSM_ClearCoat;
+	}
+
+	if (ShadingModels.HasShadingModel(MSM_DefaultLit))
+	{
+		return MSM_DefaultLit;
+	}
+
+	if (ShadingModels.HasShadingModel(MSM_Unlit))
+	{
+		return MSM_Unlit;
+	}
+
+	// TODO: add more shading models when conversion supported
+
+	return ShadingModels.GetFirstShadingModel();
+}
+
+FString FGLTFMaterialUtility::ShadingModelsToString(const FMaterialShadingModelField& ShadingModels)
+{
+	FString Result;
+
+	for (uint32 Index = 0; Index < MSM_NUM; Index++)
+	{
+		const EMaterialShadingModel ShadingModel = static_cast<EMaterialShadingModel>(Index);
+		if (ShadingModels.HasShadingModel(ShadingModel))
+		{
+			FString Name = FGLTFNameUtility::GetName(ShadingModel);
+			Result += Result.IsEmpty() ? Name : TEXT(", ") + Name;
+		}
+	}
+
+	return Result;
+}
+
 bool FGLTFMaterialUtility::NeedsMeshData(const UMaterialInterface* Material)
 {
 	if (Material == nullptr)
diff --git a/Engine/Plugins/Enterprise/GLTFExporter/MigratedFromGit/Source/GLTFExporter/Private/Converters/GLTFMaterialUtility.h b/Engine/Plugins/Enterprise/GLTFExporter/MigratedFromGit/Source/GLTFExporter/Private/Converters/GLTFMaterialUtility.h
--- a/Engine/Plugins/Enterprise/GLTFExporter/MigratedFromGit/Source/GLTFExporter/Private/Converters/GLTFMaterialUtility.h
+++ b/Engine/Plugins/Enterprise/GLTFExporter/MigratedFromGit/Source/GLTFExporter/Private/Converters/GLTFMaterialUtility.h
@@ -68,6 +68,9 @@ struct FGLTFMaterialUtility
 
 	static void ExpandAllFunctionExpressions(TArray<UMaterialExpression*>& InOutExpressions);
 
+	static EMaterialShadingModel GetRichestShadingModel(const FMaterialShadingModelField& ShadingModels);
+	static FString ShadingModelsToString(const FMaterialShadingModelField& ShadingModels);
+
 	static bool NeedsMeshData(const UMaterialInterface* Material);
 	static bool NeedsMeshData(const TArray<const UMaterialInterface*>& Materials);
 
diff --git a/Engine/Plugins/Enterprise/GLTFExporter/MigratedFromGit/Source/GLTFExporter/Private/Tasks/GLTFMaterialTasks.cpp b/Engine/Plugins/Enterprise/GLTFExporter/MigratedFromGit/Source/GLTFExporter/Private/Tasks/GLTFMaterialTasks.cpp
--- a/Engine/Plugins/Enterprise/GLTFExporter/MigratedFromGit/Source/GLTFExporter/Private/Tasks/GLTFMaterialTasks.cpp
+++ b/Engine/Plugins/Enterprise/GLTFExporter/MigratedFromGit/Source/GLTFExporter/Private/Tasks/GLTFMaterialTasks.cpp
@@ -196,30 +196,30 @@ bool FGLTFMaterialTask::TryGetAlphaMode(EGLTFJsonAlphaMode& AlphaMode) const
 bool FGLTFMaterialTask::TryGetShadingModel(EGLTFJsonShadingModel& OutShadingModel) const
 {
 	const FMaterialShadingModelField ShadingModels = Material->GetShadingModels();
-	const int32 ShadingModelCount = ShadingModels.CountShadingModels();
+	EMaterialShadingModel ShadingModel;
 
-	if (ShadingModelCount <= 0)
+	if (Material->IsShadingModelFromMaterialExpression())
 	{
-		Builder.AddWarningMessage(FString::Printf(TEXT("No shading model found for material %s"), *Material->GetName()));
-		return false;
-	}
-
-	EMaterialShadingModel ShadingModel = ShadingModels.GetFirstShadingModel();
-	if (ShadingModelCount > 1)
-	{
-		int32 TexCoord;
-		const FGLTFPropertyBakeOutput BakeOutput = BakeMaterialProperty(MP_ShadingModel, TexCoord, { 1, 1 });
-		ShadingModel = static_cast<EMaterialShadingModel>(BakeOutput.Pixels[0].R);
-
-		if (!ShadingModels.HasShadingModel(ShadingModel))
+		const FMaterialShadingModelField Evaluation = EvaluateShadingModelExpression();
+		if (Evaluation.CountShadingModels() == 1)
 		{
-			ShadingModel = ShadingModels.GetFirstShadingModel();
+			ShadingModel = Evaluation.GetFirstShadingModel();
+		}
+		else
+		{
+			ShadingModel = FGLTFMaterialUtility::GetRichestShadingModel(Evaluation.CountShadingModels() > 0 ? Evaluation : ShadingModels);
+
 			Builder.AddWarningMessage(
-				FString::Printf(TEXT("Failed to properly evaluate shading model expression in material %s, will export as %s"),
-				*Material->GetName(),
-				*FGLTFNameUtility::GetName(ShadingModel)));
+                FString::Printf(TEXT("Evaluation of shading model expression in material %s is inconclusive (%s), will export as %s"),
+                *Material->GetName(),
+                *FGLTFMaterialUtility::ShadingModelsToString(Evaluation),
+                *FGLTFNameUtility::GetName(ShadingModel)));
 		}
 	}
+	else
+	{
+		ShadingModel = ShadingModels.GetFirstShadingModel();
+	}
 
 	const EGLTFJsonShadingModel ConvertedShadingModel = FGLTFConverterUtility::ConvertShadingModel(ShadingModel);
 	if (ConvertedShadingModel == EGLTFJsonShadingModel::None)
@@ -247,6 +247,23 @@ bool FGLTFMaterialTask::TryGetShadingModel(EGLTFJsonShadingModel& OutShadingMode
 	return true;
 }
 
+FMaterialShadingModelField FGLTFMaterialTask::EvaluateShadingModelExpression() const
+{
+	int32 TexCoord;
+	const FGLTFPropertyBakeOutput BakeOutput = BakeMaterialProperty(MP_ShadingModel, TexCoord); // TODO: should we use default bake size?
+
+	const FMaterialShadingModelField PossibleValues = Material->GetShadingModels();
+	FMaterialShadingModelField Evaluation;
+
+	for (const FColor& Pixel : BakeOutput.Pixels)
+	{
+		const EMaterialShadingModel ShadingModel = static_cast<EMaterialShadingModel>(Pixel.R);
+		if (PossibleValues.HasShadingModel(ShadingModel)) Evaluation.AddShadingModel(ShadingModel);
+	}
+
+	return Evaluation;
+}
+
 bool FGLTFMaterialTask::TryGetBaseColorAndOpacity(FGLTFJsonPBRMetallicRoughness& OutPBRParams, EMaterialProperty BaseColorProperty, EMaterialProperty OpacityProperty) const
 {
 	const bool bIsBaseColorConstant = TryGetConstantColor(OutPBRParams.BaseColorFactor, BaseColorProperty);
diff --git a/Engine/Plugins/Enterprise/GLTFExporter/MigratedFromGit/Source/GLTFExporter/Private/Tasks/GLTFMaterialTasks.h b/Engine/Plugins/Enterprise/GLTFExporter/MigratedFromGit/Source/GLTFExporter/Private/Tasks/GLTFMaterialTasks.h
--- a/Engine/Plugins/Enterprise/GLTFExporter/MigratedFromGit/Source/GLTFExporter/Private/Tasks/GLTFMaterialTasks.h
+++ b/Engine/Plugins/Enterprise/GLTFExporter/MigratedFromGit/Source/GLTFExporter/Private/Tasks/GLTFMaterialTasks.h
@@ -41,6 +41,8 @@ private:
 	bool TryGetAlphaMode(EGLTFJsonAlphaMode& AlphaMode) const;
 	bool TryGetShadingModel(EGLTFJsonShadingModel& ShadingModel) const;
 
+	FMaterialShadingModelField EvaluateShadingModelExpression() const;
+
 	bool TryGetBaseColorAndOpacity(FGLTFJsonPBRMetallicRoughness& OutPBRParams, EMaterialProperty BaseColorProperty, EMaterialProperty OpacityProperty) const;
 	bool TryGetMetallicAndRoughness(FGLTFJsonPBRMetallicRoughness& OutPBRParams, EMaterialProperty MetallicProperty, EMaterialProperty RoughnessProperty) const;
 	bool TryGetClearCoatRoughness(FGLTFJsonClearCoatExtension& OutExtParams, EMaterialProperty IntensityProperty, EMaterialProperty RoughnessProperty) const;
