[GLTFExporter] fix: if enum value has no name, return value as string in GLTFNameUtility

commit 245c6c04277880bcbc8f80c91fb3b2477f7e822f
Date: Tue, 10 Nov 2020 13:54:06 +0100

    fix: if enum value has no name, return value as string in GLTFNameUtility
    (cherry picked from commit 009d40b63821dfa18a12fe6acc157734a6cc2a05​)

#skipci

diff --git a/Engine/Plugins/Enterprise/GLTFExporter/MigratedFromGit/Source/GLTFExporter/Private/Converters/GLTFNameUtility.h b/Engine/Plugins/Enterprise/GLTFExporter/MigratedFromGit/Source/GLTFExporter/Private/Converters/GLTFNameUtility.h
--- a/Engine/Plugins/Enterprise/GLTFExporter/MigratedFromGit/Source/GLTFExporter/Private/Converters/GLTFNameUtility.h
+++ b/Engine/Plugins/Enterprise/GLTFExporter/MigratedFromGit/Source/GLTFExporter/Private/Converters/GLTFNameUtility.h
@@ -15,7 +15,8 @@ struct FGLTFNameUtility
 	{
 		const UEnum* Enum = StaticEnum<EnumType>();
 		check(Enum != nullptr);
-		return Enum->GetDisplayNameTextByValue(Value).ToString();
+		const FString DisplayName = Enum->GetDisplayNameTextByValue(Value).ToString();
+		return DisplayName.IsEmpty() ? FString::FromInt(static_cast<int32>(Value)) : DisplayName;
 	}
 
 	static FString GetName(const USceneComponent* Component);
