[GLTFExporter] fix: if material baking is disabled in export options don't bake shading model expression

commit 2e732ee3fa0f1213baa9d58d73a0492c6bbf7f28
Date: Tue, 10 Nov 2020 14:17:40 +0100

    fix: if material baking is disabled in export options don't bake shading model expression
    (cherry picked from commit 9055d2e0da0e681629ca11d49c7df6fdb1e12c54​)

#skipci

diff --git a/Engine/Plugins/Enterprise/GLTFExporter/MigratedFromGit/Source/GLTFExporter/Private/Tasks/GLTFMaterialTasks.cpp b/Engine/Plugins/Enterprise/GLTFExporter/MigratedFromGit/Source/GLTFExporter/Private/Tasks/GLTFMaterialTasks.cpp
--- a/Engine/Plugins/Enterprise/GLTFExporter/MigratedFromGit/Source/GLTFExporter/Private/Tasks/GLTFMaterialTasks.cpp
+++ b/Engine/Plugins/Enterprise/GLTFExporter/MigratedFromGit/Source/GLTFExporter/Private/Tasks/GLTFMaterialTasks.cpp
@@ -195,31 +195,9 @@ bool FGLTFMaterialTask::TryGetAlphaMode(EGLTFJsonAlphaMode& AlphaMode) const
 
 bool FGLTFMaterialTask::TryGetShadingModel(EGLTFJsonShadingModel& OutShadingModel) const
 {
-	const FMaterialShadingModelField ShadingModels = Material->GetShadingModels();
-	EMaterialShadingModel ShadingModel;
-
-	if (Material->IsShadingModelFromMaterialExpression())
-	{
-		const FMaterialShadingModelField Evaluation = EvaluateShadingModelExpression();
-		if (Evaluation.CountShadingModels() == 1)
-		{
-			ShadingModel = Evaluation.GetFirstShadingModel();
-		}
-		else
-		{
-			ShadingModel = FGLTFMaterialUtility::GetRichestShadingModel(Evaluation.CountShadingModels() > 0 ? Evaluation : ShadingModels);
-
-			Builder.AddWarningMessage(
-                FString::Printf(TEXT("Evaluation of shading model expression in material %s is inconclusive (%s), will export as %s"),
-                *Material->GetName(),
-                *FGLTFMaterialUtility::ShadingModelsToString(Evaluation),
-                *FGLTFNameUtility::GetName(ShadingModel)));
-		}
-	}
-	else
-	{
-		ShadingModel = ShadingModels.GetFirstShadingModel();
-	}
+	const EMaterialShadingModel ShadingModel =
+		Material->IsShadingModelFromMaterialExpression() ?
+			EvaluateShadingModelExpression() : Material->GetShadingModels().GetFirstShadingModel();
 
 	const EGLTFJsonShadingModel ConvertedShadingModel = FGLTFConverterUtility::ConvertShadingModel(ShadingModel);
 	if (ConvertedShadingModel == EGLTFJsonShadingModel::None)
@@ -247,12 +225,22 @@ bool FGLTFMaterialTask::TryGetShadingModel(EGLTFJsonShadingModel& OutShadingMode
 	return true;
 }
 
-FMaterialShadingModelField FGLTFMaterialTask::EvaluateShadingModelExpression() const
+EMaterialShadingModel FGLTFMaterialTask::EvaluateShadingModelExpression() const
 {
+	const FMaterialShadingModelField PossibleValues = Material->GetShadingModels();
+
+	if (!Builder.ExportOptions->bBakeMaterialInputs)
+	{
+		const EMaterialShadingModel ShadingModel = FGLTFMaterialUtility::GetRichestShadingModel(PossibleValues);
+		Builder.AddWarningMessage(FString::Printf(
+			TEXT("Can't evaluate shading model expression for material %s, because material baking is disabled by export options, will export as %s"),
+			*Material->GetName(),
+			*FGLTFNameUtility::GetName(ShadingModel)));
+		return ShadingModel;
+	}
+
 	int32 TexCoord;
 	const FGLTFPropertyBakeOutput BakeOutput = BakeMaterialProperty(MP_ShadingModel, TexCoord); // TODO: should we use default bake size?
-
-	const FMaterialShadingModelField PossibleValues = Material->GetShadingModels();
 	FMaterialShadingModelField Evaluation;
 
 	for (const FColor& Pixel : BakeOutput.Pixels)
@@ -261,7 +249,18 @@ FMaterialShadingModelField FGLTFMaterialTask::EvaluateShadingModelExpression() c
 		if (PossibleValues.HasShadingModel(ShadingModel)) Evaluation.AddShadingModel(ShadingModel);
 	}
 
-	return Evaluation;
+	if (Evaluation.CountShadingModels() != 1)
+	{
+		const EMaterialShadingModel ShadingModel = FGLTFMaterialUtility::GetRichestShadingModel(Evaluation.CountShadingModels() > 0 ? Evaluation : PossibleValues);
+		Builder.AddWarningMessage(FString::Printf(
+			TEXT("Evaluation of shading model expression in material %s is inconclusive (%s), will export as %s"),
+			*Material->GetName(),
+			*FGLTFMaterialUtility::ShadingModelsToString(Evaluation),
+			*FGLTFNameUtility::GetName(ShadingModel)));
+		return ShadingModel;
+	}
+
+	return Evaluation.GetFirstShadingModel();
 }
 
 bool FGLTFMaterialTask::TryGetBaseColorAndOpacity(FGLTFJsonPBRMetallicRoughness& OutPBRParams, EMaterialProperty BaseColorProperty, EMaterialProperty OpacityProperty) const
diff --git a/Engine/Plugins/Enterprise/GLTFExporter/MigratedFromGit/Source/GLTFExporter/Private/Tasks/GLTFMaterialTasks.h b/Engine/Plugins/Enterprise/GLTFExporter/MigratedFromGit/Source/GLTFExporter/Private/Tasks/GLTFMaterialTasks.h
--- a/Engine/Plugins/Enterprise/GLTFExporter/MigratedFromGit/Source/GLTFExporter/Private/Tasks/GLTFMaterialTasks.h
+++ b/Engine/Plugins/Enterprise/GLTFExporter/MigratedFromGit/Source/GLTFExporter/Private/Tasks/GLTFMaterialTasks.h
@@ -41,7 +41,7 @@ private:
 	bool TryGetAlphaMode(EGLTFJsonAlphaMode& AlphaMode) const;
 	bool TryGetShadingModel(EGLTFJsonShadingModel& ShadingModel) const;
 
-	FMaterialShadingModelField EvaluateShadingModelExpression() const;
+	EMaterialShadingModel EvaluateShadingModelExpression() const;
 
 	bool TryGetBaseColorAndOpacity(FGLTFJsonPBRMetallicRoughness& OutPBRParams, EMaterialProperty BaseColorProperty, EMaterialProperty OpacityProperty) const;
 	bool TryGetMetallicAndRoughness(FGLTFJsonPBRMetallicRoughness& OutPBRParams, EMaterialProperty MetallicProperty, EMaterialProperty RoughnessProperty) const;
