[GLTFExporter] feat: transform custom output ClearCoatBottomNormal to correct range when baking

commit d656a6f78c88955a6683c34e2800cd19d692b2db
Date: Fri, 6 Nov 2020 11:09:36 +0100

    feat: transform custom output ClearCoatBottomNormal to correct range when baking
    (cherry picked from commit 0ebfd4ba994b99fa2febf9086f53a03b88852729​)

#skipci

diff --git a/Engine/Plugins/Enterprise/GLTFExporter/MigratedFromGit/Source/GLTFMaterialBaking/Private/ExportMaterialProxy.h b/Engine/Plugins/Enterprise/GLTFExporter/MigratedFromGit/Source/GLTFMaterialBaking/Private/ExportMaterialProxy.h
--- a/Engine/Plugins/Enterprise/GLTFExporter/MigratedFromGit/Source/GLTFMaterialBaking/Private/ExportMaterialProxy.h
+++ b/Engine/Plugins/Enterprise/GLTFExporter/MigratedFromGit/Source/GLTFMaterialBaking/Private/ExportMaterialProxy.h
@@ -410,9 +410,9 @@ public:
 				// Only return for Opaque and Masked...
 				if (BlendMode == BLEND_Opaque || BlendMode == BLEND_Masked)
 				{
-					return Compiler->Add(
-						Compiler->Mul(MaterialInterface->CompileProperty(&ProxyCompiler, PropertyToCompile, ForceCast_Exact_Replicate), Compiler->Constant(0.5f)), // [-1,1] * 0.5
-						Compiler->Constant(0.5f)); // [-0.5,0.5] + 0.5
+					return CompileNormalTransform(
+						Compiler,
+						MaterialInterface->CompileProperty(&ProxyCompiler, PropertyToCompile, ForceCast_Exact_Replicate));
 				}
 				break;
 			case MP_ShadingModel:
@@ -657,6 +657,11 @@ private:
 		if (ExpressionInput)
 		{
 			Result = ExpressionInput->Compile(Compiler);
+
+			if (CustomOutputToCompile ==  TEXT("ClearCoatBottomNormal"))
+			{
+				Result = CompileNormalTransform(Compiler, Result);
+			}
 		}
 		else
 		{
@@ -706,6 +711,13 @@ private:
 		return Compiler->Div(Code, Compiler->Constant(255)); // [0,255] / 255
 	}
 
+	int32 CompileNormalTransform(FMaterialCompiler* Compiler, int32 ExpressionIndex) const
+	{
+		return Compiler->Add(
+			Compiler->Mul(ExpressionIndex, Compiler->Constant(0.5f)), // [-1,1] * 0.5
+			Compiler->Constant(0.5f)); // [-0.5,0.5] + 0.5
+	}
+
 private:
 	/** The material interface for this proxy */
 	UMaterialInterface* MaterialInterface;
@@ -717,4 +729,4 @@ private:
 	FString CustomOutputToCompile;
 	FGuid Id;
 	bool bSynchronousCompilation;
-};
\ No newline at end of file
+};
