[GLTFExporter] feat: transform baked normals to tangent-space

commit f6c1ffa5089ec76fed7ca84bbccf772b01b6983e
Date: Fri, 6 Nov 2020 13:46:28 +0100

    feat: transform baked normals to tangent-space
    (cherry picked from commit c172379a827cd4c8e4b35337f36e7a376aff91b8​)

#skipci

diff --git a/Engine/Plugins/Enterprise/GLTFExporter/MigratedFromGit/Source/GLTFMaterialBaking/Private/ExportMaterialProxy.h b/Engine/Plugins/Enterprise/GLTFExporter/MigratedFromGit/Source/GLTFMaterialBaking/Private/ExportMaterialProxy.h
--- a/Engine/Plugins/Enterprise/GLTFExporter/MigratedFromGit/Source/GLTFMaterialBaking/Private/ExportMaterialProxy.h
+++ b/Engine/Plugins/Enterprise/GLTFExporter/MigratedFromGit/Source/GLTFMaterialBaking/Private/ExportMaterialProxy.h
@@ -713,6 +713,12 @@ private:
 
 	int32 CompileNormalTransform(FMaterialCompiler* Compiler, int32 ExpressionIndex) const
 	{
+		// TODO: make this configurable instead of assuming that we always want tangent-space normals
+		if (!Material->bTangentSpaceNormal)
+		{
+			ExpressionIndex = Compiler->TransformVector(MCB_World, MCB_Tangent, ExpressionIndex);
+		}
+
 		return Compiler->Add(
 			Compiler->Mul(ExpressionIndex, Compiler->Constant(0.5f)), // [-1,1] * 0.5
 			Compiler->Constant(0.5f)); // [-0.5,0.5] + 0.5
