[GLTFExporter] refactor: when showing export files, select gltf/glb file

commit 6521c7bccc869c3946d80418f916a4c71899e1a7
Date: Tue, 10 Nov 2020 21:01:11 +0100

    refactor: when showing export files, select gltf/glb file
    (cherry picked from commit 0b540651107414fd47395e0a3891419afd0389c9​)

#skipci

diff --git a/Engine/Plugins/Enterprise/GLTFExporter/MigratedFromGit/Source/GLTFExporter/Private/Exporters/GLTFExporter.cpp b/Engine/Plugins/Enterprise/GLTFExporter/MigratedFromGit/Source/GLTFExporter/Private/Exporters/GLTFExporter.cpp
--- a/Engine/Plugins/Enterprise/GLTFExporter/MigratedFromGit/Source/GLTFExporter/Private/Exporters/GLTFExporter.cpp
+++ b/Engine/Plugins/Enterprise/GLTFExporter/MigratedFromGit/Source/GLTFExporter/Private/Exporters/GLTFExporter.cpp
@@ -58,7 +58,7 @@ bool UGLTFExporter::ExportBinary(UObject* Object, const TCHAR* Type, FArchive& A
 
 		if (bSuccess && Options->bShowFilesWhenDone)
 		{
-			FPlatformProcess::ExploreFolder(*Builder.DirPath);
+			FPlatformProcess::ExploreFolder(*Builder.FilePath);
 		}
 	}
 
