[GLTFExporter] feat: update web viewer index json

commit 3d45acc68a8566e43fd648c195f68d7984564436
Date: Tue, 27 Oct 2020 22:18:50 +0100

    feat: update web viewer index json
    (cherry picked from commit ea319a2262434bce796f941afaeb834d7b416e77​)

#skipci

diff --git a/Engine/Plugins/Enterprise/GLTFExporter/MigratedFromGit/Source/GLTFExporter/Private/Builders/GLTFContainerBuilder.cpp b/Engine/Plugins/Enterprise/GLTFExporter/MigratedFromGit/Source/GLTFExporter/Private/Builders/GLTFContainerBuilder.cpp
--- a/Engine/Plugins/Enterprise/GLTFExporter/MigratedFromGit/Source/GLTFExporter/Private/Builders/GLTFContainerBuilder.cpp
+++ b/Engine/Plugins/Enterprise/GLTFExporter/MigratedFromGit/Source/GLTFExporter/Private/Builders/GLTFContainerBuilder.cpp
@@ -11,6 +11,25 @@ FGLTFContainerBuilder::FGLTFContainerBuilder(const FString& FilePath, const UGLT
 {
 }
 
+void FGLTFContainerBuilder::Write(FArchive& Archive, FFeedbackContext* Context)
+{
+	CompleteAllTasks(Context);
+
+	if (bIsGlbFile)
+	{
+		WriteGlb(Archive);
+	}
+	else
+	{
+		WriteJson(Archive);
+	}
+
+	if (ExportOptions->bBundleWebViewer)
+	{
+		BundleWebViewer();
+	}
+}
+
 void FGLTFContainerBuilder::WriteGlb(FArchive& Archive) const
 {
 	FBufferArchive JsonData;
@@ -23,34 +42,70 @@ void FGLTFContainerBuilder::WriteGlb(FArchive& Archive) const
 
 void FGLTFContainerBuilder::BundleWebViewer()
 {
-	const FString WebViewerPath = IPluginManager::Get().FindPlugin(GLTFEXPORTER_MODULE_NAME)->GetBaseDir() / TEXT("Resources") / TEXT("WebViewer.zip");
-	if (!FPaths::FileExists(WebViewerPath))
+	const FString PluginDir = IPluginManager::Get().FindPlugin(GLTFEXPORTER_MODULE_NAME)->GetBaseDir();
+	const FString ArchiveFile = PluginDir / TEXT("Resources") / TEXT("WebViewer.zip");
+
+	if (!FPaths::FileExists(ArchiveFile))
 	{
-		AddWarningMessage(FString::Printf(TEXT("No web viewer found at %s"), *WebViewerPath));
+		AddWarningMessage(FString::Printf(TEXT("No web viewer found at %s"), *ArchiveFile));
 		return;
 	}
 
-	if (!FGLTFZipUtility::ExtractToDirectory(WebViewerPath, DirPath, *this))
+	if (!FGLTFZipUtility::ExtractToDirectory(ArchiveFile, DirPath, *this))
 	{
-		AddErrorMessage(FString::Printf(TEXT("Failed to extract web viewer at %s"), *WebViewerPath));
+		AddErrorMessage(FString::Printf(TEXT("Failed to extract web viewer at %s"), *ArchiveFile));
+		return;
 	}
+
+	UpdateWebViewerIndex();
 }
 
-void FGLTFContainerBuilder::Write(FArchive& Archive, FFeedbackContext* Context)
+void FGLTFContainerBuilder::UpdateWebViewerIndex()
 {
-	CompleteAllTasks(Context);
+	const FString IndexFile = DirPath / TEXT("index.json");
 
-	if (bIsGlbFile)
+	if (!FPaths::FileExists(IndexFile))
 	{
-		WriteGlb(Archive);
+		AddWarningMessage(FString::Printf(TEXT("No index file found at %s"), *IndexFile));
+		return;
 	}
-	else
+
+	TSharedPtr<FJsonObject> JsonObject;
+	if (!ReadJsonFile(IndexFile, JsonObject))
 	{
-		WriteJson(Archive);
+		AddWarningMessage(FString::Printf(TEXT("Failed to read index file at %s"), *IndexFile));
+		return;
 	}
 
-	if (ExportOptions->bBundleWebViewer)
+	JsonObject->SetArrayField(TEXT("assets"), { MakeShared<FJsonValueString>(FPaths::GetCleanFilename(FilePath)) });
+
+	if (!WriteJsonFile(IndexFile, JsonObject.ToSharedRef()))
 	{
-		BundleWebViewer();
+		AddWarningMessage(FString::Printf(TEXT("Failed to write index file at %s"), *IndexFile));
 	}
 }
+
+bool FGLTFContainerBuilder::ReadJsonFile(const FString& FilePath, TSharedPtr<FJsonObject>& JsonObject)
+{
+	FString JsonContent;
+	if (!FFileHelper::LoadFileToString(JsonContent, *FilePath))
+	{
+		return false;
+	}
+
+	const TSharedRef<TJsonReader<>> JsonReader = TJsonReaderFactory<>::Create(JsonContent);
+	return FJsonSerializer::Deserialize(JsonReader, JsonObject) && JsonObject.IsValid();
+}
+
+bool FGLTFContainerBuilder::WriteJsonFile(const FString& FilePath, const TSharedRef<FJsonObject>& JsonObject)
+{
+	FString JsonContent;
+	const TSharedRef<TJsonWriter<>> JsonWriter = TJsonWriterFactory<>::Create(&JsonContent);
+
+	if (!FJsonSerializer::Serialize(JsonObject, JsonWriter))
+	{
+		return false;
+	}
+
+	return FFileHelper::SaveStringToFile(JsonContent, *FilePath, FFileHelper::EEncodingOptions::ForceUTF8);
+}
diff --git a/Engine/Plugins/Enterprise/GLTFExporter/MigratedFromGit/Source/GLTFExporter/Private/Builders/GLTFContainerBuilder.h b/Engine/Plugins/Enterprise/GLTFExporter/MigratedFromGit/Source/GLTFExporter/Private/Builders/GLTFContainerBuilder.h
--- a/Engine/Plugins/Enterprise/GLTFExporter/MigratedFromGit/Source/GLTFExporter/Private/Builders/GLTFContainerBuilder.h
+++ b/Engine/Plugins/Enterprise/GLTFExporter/MigratedFromGit/Source/GLTFExporter/Private/Builders/GLTFContainerBuilder.h
@@ -17,4 +17,9 @@ private:
 	void WriteGlb(FArchive& Archive) const;
 
 	void BundleWebViewer();
+
+	void UpdateWebViewerIndex();
+
+	static bool ReadJsonFile(const FString& FilePath, TSharedPtr<FJsonObject>& JsonObject);
+	static bool WriteJsonFile(const FString& FilePath, const TSharedRef<FJsonObject>& JsonObject);
 };
