[GLTFExporter] chore: rename WebViewer.zip to GLTFWebViewer.zip

commit 820658c84e94f5145f4bac5ec73ad90d44c18399
Date: Tue, 10 Nov 2020 16:58:55 +0100

    chore: rename WebViewer.zip to GLTFWebViewer.zip
    (cherry picked from commit 1acc67dd39fb6f47dd8ed944e12bfc46d4817c8e​)

#skipci

diff --git a/Engine/Plugins/Enterprise/GLTFExporter/MigratedFromGit/Source/GLTFExporter/Private/Builders/GLTFContainerBuilder.cpp b/Engine/Plugins/Enterprise/GLTFExporter/MigratedFromGit/Source/GLTFExporter/Private/Builders/GLTFContainerBuilder.cpp
--- a/Engine/Plugins/Enterprise/GLTFExporter/MigratedFromGit/Source/GLTFExporter/Private/Builders/GLTFContainerBuilder.cpp
+++ b/Engine/Plugins/Enterprise/GLTFExporter/MigratedFromGit/Source/GLTFExporter/Private/Builders/GLTFContainerBuilder.cpp
@@ -43,7 +43,7 @@ void FGLTFContainerBuilder::WriteGlb(FArchive& Archive) const
 void FGLTFContainerBuilder::BundleWebViewer()
 {
 	const FString PluginDir = IPluginManager::Get().FindPlugin(GLTFEXPORTER_MODULE_NAME)->GetBaseDir();
-	const FString ArchiveFile = PluginDir / TEXT("Resources") / TEXT("WebViewer.zip");
+	const FString ArchiveFile = PluginDir / TEXT("Resources") / TEXT("GLTFWebViewer.zip");
 
 	if (!FPaths::FileExists(ArchiveFile))
 	{
