[GLTFExporter] refactor: simplify GLTFZipUtility implementation

commit b1ab2aeba170121f57491fe35ac810f8822769ae
Date: Tue, 10 Nov 2020 17:11:58 +0100

    refactor: simplify GLTFZipUtility implementation
    (cherry picked from commit aa5c9f615c9e2681818ce948f01a9973bdf5aec5​)

#skipci

diff --git a/Engine/Plugins/Enterprise/GLTFExporter/MigratedFromGit/Source/GLTFExporter/Private/Builders/GLTFContainerBuilder.cpp b/Engine/Plugins/Enterprise/GLTFExporter/MigratedFromGit/Source/GLTFExporter/Private/Builders/GLTFContainerBuilder.cpp
--- a/Engine/Plugins/Enterprise/GLTFExporter/MigratedFromGit/Source/GLTFExporter/Private/Builders/GLTFContainerBuilder.cpp
+++ b/Engine/Plugins/Enterprise/GLTFExporter/MigratedFromGit/Source/GLTFExporter/Private/Builders/GLTFContainerBuilder.cpp
@@ -51,7 +51,7 @@ void FGLTFContainerBuilder::BundleWebViewer()
 		return;
 	}
 
-	if (!FGLTFZipUtility::ExtractToDirectory(ArchiveFile, DirPath, *this))
+	if (!FGLTFZipUtility::ExtractToDirectory(ArchiveFile, DirPath))
 	{
 		AddErrorMessage(FString::Printf(TEXT("Failed to extract web viewer at %s"), *ArchiveFile));
 		return;
diff --git a/Engine/Plugins/Enterprise/GLTFExporter/MigratedFromGit/Source/GLTFExporter/Private/Builders/GLTFZipUtility.cpp b/Engine/Plugins/Enterprise/GLTFExporter/MigratedFromGit/Source/GLTFExporter/Private/Builders/GLTFZipUtility.cpp
--- a/Engine/Plugins/Enterprise/GLTFExporter/MigratedFromGit/Source/GLTFExporter/Private/Builders/GLTFZipUtility.cpp
+++ b/Engine/Plugins/Enterprise/GLTFExporter/MigratedFromGit/Source/GLTFExporter/Private/Builders/GLTFZipUtility.cpp
@@ -7,53 +7,49 @@ THIRD_PARTY_INCLUDES_START
 #include "ThirdParty/zlib/zlib-1.2.5/Src/contrib/minizip/unzip.h"
 THIRD_PARTY_INCLUDES_END
 
-bool FGLTFZipUtility::ExtractToDirectory(const FString& SourceFilePath, const FString& DestinationDirectoryPath, FGLTFMessageBuilder& Builder)
+bool FGLTFZipUtility::ExtractToDirectory(const FString& ArchiveFile, const FString& TargetDirectory)
 {
-	const unzFile ZipFile = unzOpen64(TCHAR_TO_ANSI(*SourceFilePath));
-	if (ZipFile == nullptr)
+	const unzFile ZipHandle = unzOpen64(TCHAR_TO_ANSI(*ArchiveFile));
+	if (ZipHandle == nullptr)
 	{
-		Builder.AddErrorMessage(TEXT("Can't open zip archive"));
 		return false;
 	}
 
-	if (const int ErrorCode = unzGoToFirstFile(ZipFile))
+	if (unzGoToFirstFile(ZipHandle) != UNZ_OK)
 	{
-		Builder.AddErrorMessage(FString::Printf(TEXT("Can't locate first file in zip archive (error %d)"), ErrorCode));
-		unzClose(ZipFile);
+		unzClose(ZipHandle);
 		return false;
 	}
 
 	do
 	{
-		if (!ExtractCurrentFile(ZipFile, DestinationDirectoryPath, Builder))
+		if (!ExtractCurrentFile(ZipHandle, TargetDirectory))
 		{
-			unzClose(ZipFile);
+			unzClose(ZipHandle);
 			return false;
 		}
-	} while (unzGoToNextFile(ZipFile) == UNZ_OK);
+	} while (unzGoToNextFile(ZipHandle) == UNZ_OK);
 
-	unzClose(ZipFile);
+	unzClose(ZipHandle);
 	return true;
 }
 
-bool FGLTFZipUtility::ExtractCurrentFile(void* ZipFile, const FString& DestinationDirectoryPath, FGLTFMessageBuilder& Builder)
+bool FGLTFZipUtility::ExtractCurrentFile(void* ZipHandle, const FString& TargetDirectory)
 {
 	unz_file_info64 FileInfo;
-	if (const int ErrorCode = unzGetCurrentFileInfo64(ZipFile, &FileInfo, nullptr, 0, nullptr, 0, nullptr, 0))
+	if (unzGetCurrentFileInfo64(ZipHandle, &FileInfo, nullptr, 0, nullptr, 0, nullptr, 0) != UNZ_OK)
 	{
-		Builder.AddErrorMessage(FString::Printf(TEXT("Can't get file info in zip archive (error %d)"), ErrorCode));
 		return false;
 	}
 
 	TArray<char> Filename;
 	Filename.Init('\0', FileInfo.size_filename + 1);
-	if (const int ErrorCode = unzGetCurrentFileInfo64(ZipFile, nullptr, Filename.GetData(), Filename.Num() - 1, nullptr, 0, nullptr, 0))
+	if (unzGetCurrentFileInfo64(ZipHandle, nullptr, Filename.GetData(), Filename.Num() - 1, nullptr, 0, nullptr, 0) != UNZ_OK)
 	{
-		Builder.AddErrorMessage(FString::Printf(TEXT("Can't get file name in zip archive (error %d)"), ErrorCode));
 		return false;
 	}
 
-	const FString DestinationFilePath = DestinationDirectoryPath / Filename.GetData();
+	const FString DestinationFilePath = TargetDirectory / Filename.GetData();
 	if (DestinationFilePath.EndsWith(TEXT("/")) || DestinationFilePath.EndsWith(TEXT("\\")))
 	{
 		IFileManager::Get().MakeDirectory(*DestinationFilePath, true);
@@ -63,32 +59,28 @@ bool FGLTFZipUtility::ExtractCurrentFile(void* ZipFile, const FString& Destinati
 		FArchive* FileWriter = IFileManager::Get().CreateFileWriter(*DestinationFilePath);
 		if (FileWriter == nullptr)
 		{
-			Builder.AddErrorMessage(FString::Printf(TEXT("Can't write to file %s from zip archive"), *DestinationFilePath));
 			return false;
 		}
 
-		if (const int ErrorCode = unzOpenCurrentFile(ZipFile))
+		if (unzOpenCurrentFile(ZipHandle) != UNZ_OK)
 		{
-			Builder.AddErrorMessage(FString::Printf(TEXT("Can't open file %s in zip archive (error %d)"), ANSI_TO_TCHAR(Filename.GetData()), ErrorCode));
 			return false;
 		}
 
 		uint8 ReadBuffer[8192];
-		while (const int ReadSize = unzReadCurrentFile(ZipFile, ReadBuffer, sizeof(ReadBuffer)))
+		while (const int ReadSize = unzReadCurrentFile(ZipHandle, ReadBuffer, sizeof(ReadBuffer)))
 		{
 			if (ReadSize < 0)
 			{
-				Builder.AddErrorMessage(FString::Printf(TEXT("Can't read file %s in zip archive (error %d)"), ANSI_TO_TCHAR(Filename.GetData()), ReadSize));
-				unzCloseCurrentFile(ZipFile);
+				unzCloseCurrentFile(ZipHandle);
 				return false;
 			}
 
 			FileWriter->Serialize(ReadBuffer, ReadSize);
 		}
 
-		if (const int ErrorCode = unzCloseCurrentFile(ZipFile))
+		if (unzCloseCurrentFile(ZipHandle) != UNZ_OK)
 		{
-			Builder.AddErrorMessage(FString::Printf(TEXT("Can't close file %s in zip archive (error %d)"), ANSI_TO_TCHAR(Filename.GetData()), ErrorCode));
 			return false;
 		}
 
diff --git a/Engine/Plugins/Enterprise/GLTFExporter/MigratedFromGit/Source/GLTFExporter/Private/Builders/GLTFZipUtility.h b/Engine/Plugins/Enterprise/GLTFExporter/MigratedFromGit/Source/GLTFExporter/Private/Builders/GLTFZipUtility.h
--- a/Engine/Plugins/Enterprise/GLTFExporter/MigratedFromGit/Source/GLTFExporter/Private/Builders/GLTFZipUtility.h
+++ b/Engine/Plugins/Enterprise/GLTFExporter/MigratedFromGit/Source/GLTFExporter/Private/Builders/GLTFZipUtility.h
@@ -2,16 +2,15 @@
 
 #pragma once
 
-#include "Builders/GLTFMessageBuilder.h"
 #include "CoreMinimal.h"
 
 class FGLTFZipUtility
 {
 public:
 
-	static bool ExtractToDirectory(const FString& SourceFilePath, const FString& DestinationDirectoryPath, FGLTFMessageBuilder& Builder);
+	static bool ExtractToDirectory(const FString& ArchiveFile, const FString& TargetDirectory);
 
 private:
 
-	static bool ExtractCurrentFile(void* ZipFile, const FString& DestinationDirectoryPath, FGLTFMessageBuilder& Builder);
+	static bool ExtractCurrentFile(void* ZipHandle, const FString& TargetDirectory);
 };
