[GLTFExporter] refactor: add additional helper functions to GLTFZipUtility

commit f6fc76dcb262fa7e65aac6664217483297f41afb
Date: Tue, 10 Nov 2020 19:22:44 +0100

    refactor: add additional helper functions to GLTFZipUtility
    (cherry picked from commit 2c6d8101d29e5968f6557dc0436a0d470fc4f218​)

#skipci

diff --git a/Engine/Plugins/Enterprise/GLTFExporter/MigratedFromGit/Source/GLTFExporter/Private/Builders/GLTFContainerBuilder.cpp b/Engine/Plugins/Enterprise/GLTFExporter/MigratedFromGit/Source/GLTFExporter/Private/Builders/GLTFContainerBuilder.cpp
--- a/Engine/Plugins/Enterprise/GLTFExporter/MigratedFromGit/Source/GLTFExporter/Private/Builders/GLTFContainerBuilder.cpp
+++ b/Engine/Plugins/Enterprise/GLTFExporter/MigratedFromGit/Source/GLTFExporter/Private/Builders/GLTFContainerBuilder.cpp
@@ -51,7 +51,7 @@ void FGLTFContainerBuilder::BundleWebViewer()
 		return;
 	}
 
-	if (!FGLTFZipUtility::ExtractToDirectory(ArchiveFile, DirPath))
+	if (!FGLTFZipUtility::ExtractAllFiles(ArchiveFile, DirPath))
 	{
 		AddErrorMessage(FString::Printf(TEXT("Failed to extract web viewer at %s"), *ArchiveFile));
 		return;
diff --git a/Engine/Plugins/Enterprise/GLTFExporter/MigratedFromGit/Source/GLTFExporter/Private/Builders/GLTFZipUtility.cpp b/Engine/Plugins/Enterprise/GLTFExporter/MigratedFromGit/Source/GLTFExporter/Private/Builders/GLTFZipUtility.cpp
--- a/Engine/Plugins/Enterprise/GLTFExporter/MigratedFromGit/Source/GLTFExporter/Private/Builders/GLTFZipUtility.cpp
+++ b/Engine/Plugins/Enterprise/GLTFExporter/MigratedFromGit/Source/GLTFExporter/Private/Builders/GLTFZipUtility.cpp
@@ -7,7 +7,39 @@ THIRD_PARTY_INCLUDES_START
 #include "ThirdParty/zlib/zlib-1.2.5/Src/contrib/minizip/unzip.h"
 THIRD_PARTY_INCLUDES_END
 
-bool FGLTFZipUtility::ExtractToDirectory(const FString& ArchiveFile, const FString& TargetDirectory)
+TArray<FString> FGLTFZipUtility::GetAllFiles(const FString& ArchiveFile)
+{
+	const unzFile ZipHandle = unzOpen64(TCHAR_TO_ANSI(*ArchiveFile));
+	if (ZipHandle == nullptr)
+	{
+		return {};
+	}
+
+	if (unzGoToFirstFile(ZipHandle) != UNZ_OK)
+	{
+		unzClose(ZipHandle);
+		return {};
+	}
+
+	TArray<FString> Filenames;
+
+	do
+	{
+		const FString Filename = GetCurrentFilename(ZipHandle);
+		if (Filename.IsEmpty())
+		{
+			unzClose(ZipHandle);
+			return {};
+		}
+
+		Filenames.Add(Filename);
+	} while (unzGoToNextFile(ZipHandle) == UNZ_OK);
+
+	unzClose(ZipHandle);
+	return Filenames;
+}
+
+bool FGLTFZipUtility::ExtractAllFiles(const FString& ArchiveFile, const FString& TargetDirectory)
 {
 	const unzFile ZipHandle = unzOpen64(TCHAR_TO_ANSI(*ArchiveFile));
 	if (ZipHandle == nullptr)
@@ -23,7 +55,14 @@ bool FGLTFZipUtility::ExtractToDirectory(const FString& ArchiveFile, const FStri
 
 	do
 	{
-		if (!ExtractCurrentFile(ZipHandle, TargetDirectory))
+		const FString Filename = GetCurrentFilename(ZipHandle);
+		if (Filename.IsEmpty())
+		{
+			unzClose(ZipHandle);
+			return false;
+		}
+
+		if (!ExtractCurrentFile(ZipHandle, TargetDirectory / Filename))
 		{
 			unzClose(ZipHandle);
 			return false;
@@ -34,29 +73,70 @@ bool FGLTFZipUtility::ExtractToDirectory(const FString& ArchiveFile, const FStri
 	return true;
 }
 
-bool FGLTFZipUtility::ExtractCurrentFile(void* ZipHandle, const FString& TargetDirectory)
+bool FGLTFZipUtility::ExtractOneFile(const FString& ArchiveFile, const FString& FileToExtract, const FString& TargetDirectory)
+{
+	const unzFile ZipHandle = unzOpen64(TCHAR_TO_ANSI(*ArchiveFile));
+	if (ZipHandle == nullptr)
+	{
+		return false;
+	}
+
+	if (unzGoToFirstFile(ZipHandle) != UNZ_OK)
+	{
+		unzClose(ZipHandle);
+		return false;
+	}
+
+	if (unzLocateFile(ZipHandle, TCHAR_TO_ANSI(*FileToExtract), 0) != UNZ_OK)
+	{
+		unzClose(ZipHandle);
+		return false;
+	}
+
+	if (!ExtractCurrentFile(ZipHandle, TargetDirectory / FileToExtract))
+	{
+		unzClose(ZipHandle);
+		return false;
+	}
+
+	unzClose(ZipHandle);
+	return true;
+}
+
+FString FGLTFZipUtility::GetCurrentFilename(void* ZipHandle)
 {
 	unz_file_info64 FileInfo;
 	if (unzGetCurrentFileInfo64(ZipHandle, &FileInfo, nullptr, 0, nullptr, 0, nullptr, 0) != UNZ_OK)
 	{
-		return false;
+		return {};
 	}
 
 	TArray<char> Filename;
+
 	Filename.Init('\0', FileInfo.size_filename + 1);
 	if (unzGetCurrentFileInfo64(ZipHandle, nullptr, Filename.GetData(), Filename.Num() - 1, nullptr, 0, nullptr, 0) != UNZ_OK)
+	{
+		return {};
+	}
+
+	return Filename.GetData();
+}
+
+bool FGLTFZipUtility::ExtractCurrentFile(void* ZipHandle, const FString& TargetFile)
+{
+	const FString Filename = GetCurrentFilename(ZipHandle);
+	if (Filename.IsEmpty())
 	{
 		return false;
 	}
 
-	const FString DestinationFilePath = TargetDirectory / Filename.GetData();
-	if (DestinationFilePath.EndsWith(TEXT("/")) || DestinationFilePath.EndsWith(TEXT("\\")))
+	if (TargetFile.EndsWith(TEXT("/")) || TargetFile.EndsWith(TEXT("\\")))
 	{
-		IFileManager::Get().MakeDirectory(*DestinationFilePath, true);
+		IFileManager::Get().MakeDirectory(*TargetFile, true);
 	}
 	else
 	{
-		FArchive* FileWriter = IFileManager::Get().CreateFileWriter(*DestinationFilePath);
+		FArchive* FileWriter = IFileManager::Get().CreateFileWriter(*TargetFile);
 		if (FileWriter == nullptr)
 		{
 			return false;
diff --git a/Engine/Plugins/Enterprise/GLTFExporter/MigratedFromGit/Source/GLTFExporter/Private/Builders/GLTFZipUtility.h b/Engine/Plugins/Enterprise/GLTFExporter/MigratedFromGit/Source/GLTFExporter/Private/Builders/GLTFZipUtility.h
--- a/Engine/Plugins/Enterprise/GLTFExporter/MigratedFromGit/Source/GLTFExporter/Private/Builders/GLTFZipUtility.h
+++ b/Engine/Plugins/Enterprise/GLTFExporter/MigratedFromGit/Source/GLTFExporter/Private/Builders/GLTFZipUtility.h
@@ -8,9 +8,14 @@ class FGLTFZipUtility
 {
 public:
 
-	static bool ExtractToDirectory(const FString& ArchiveFile, const FString& TargetDirectory);
+	static TArray<FString> GetAllFiles(const FString& ArchiveFile);
+
+	static bool ExtractAllFiles(const FString& ArchiveFile, const FString& TargetDirectory);
+	static bool ExtractOneFile(const FString& ArchiveFile, const FString& FileToExtract, const FString& TargetDirectory);
 
 private:
 
-	static bool ExtractCurrentFile(void* ZipHandle, const FString& TargetDirectory);
+	static FString GetCurrentFilename(void* ZipHandle);
+
+	static bool ExtractCurrentFile(void* ZipHandle, const FString& TargetFile);
 };
