[GLTFExporter] feat: bundle GLTFLaunchHelper when extracting web viewer

commit 71e87f93f603786a5c69d7dbb48c91e20527b5cd
Date: Tue, 10 Nov 2020 20:56:36 +0100

    feat: bundle GLTFLaunchHelper when extracting web viewer
    (cherry picked from commit 4999b6152b27655cf08cff258b6cbf602a0cbb1e​)

#skipci

diff --git a/Engine/Plugins/Enterprise/GLTFExporter/MigratedFromGit/Source/GLTFExporter/Private/Builders/GLTFBuilderUtility.cpp b/Engine/Plugins/Enterprise/GLTFExporter/MigratedFromGit/Source/GLTFExporter/Private/Builders/GLTFBuilderUtility.cpp
--- a/Engine/Plugins/Enterprise/GLTFExporter/MigratedFromGit/Source/GLTFExporter/Private/Builders/GLTFBuilderUtility.cpp
+++ b/Engine/Plugins/Enterprise/GLTFExporter/MigratedFromGit/Source/GLTFExporter/Private/Builders/GLTFBuilderUtility.cpp
@@ -1,6 +1,21 @@
 // Copyright Epic Games, Inc. All Rights Reserved.
 
 #include "Builders/GLTFBuilderUtility.h"
+#include "GLTFExporterModule.h"
+#include "Interfaces/IPluginManager.h"
+#include "Serialization/JsonReader.h"
+#include "Serialization/JsonSerializer.h"
+#include "Misc/FileHelper.h"
+
+#if PLATFORM_LINUX || PLATFORM_MAC
+// TODO: uncomment when GLTFLaunchHelper added for linux and macos
+// #include <sys/stat.h>
+#endif
+
+FString FGLTFBuilderUtility::GetPluginDir()
+{
+	return IPluginManager::Get().FindPlugin(GLTFEXPORTER_MODULE_NAME)->GetBaseDir();
+}
 
 const TCHAR* FGLTFBuilderUtility::GetFileExtension(EGLTFJsonMimeType MimeType)
 {
@@ -37,3 +52,50 @@ bool FGLTFBuilderUtility::IsGlbFile(const FString& Filename)
 {
 	return FPaths::GetExtension(Filename).Equals(TEXT("glb"), ESearchCase::IgnoreCase);
 }
+
+bool FGLTFBuilderUtility::ReadJsonFile(const FString& FilePath, TSharedPtr<FJsonObject>& JsonObject)
+{
+	FString JsonContent;
+	if (!FFileHelper::LoadFileToString(JsonContent, *FilePath))
+	{
+		return false;
+	}
+
+	const TSharedRef<TJsonReader<>> JsonReader = TJsonReaderFactory<>::Create(JsonContent);
+	return FJsonSerializer::Deserialize(JsonReader, JsonObject) && JsonObject.IsValid();
+}
+
+bool FGLTFBuilderUtility::WriteJsonFile(const FString& FilePath, const TSharedRef<FJsonObject>& JsonObject)
+{
+	FString JsonContent;
+	const TSharedRef<TJsonWriter<>> JsonWriter = TJsonWriterFactory<>::Create(&JsonContent);
+
+	if (!FJsonSerializer::Serialize(JsonObject, JsonWriter))
+	{
+		return false;
+	}
+
+	return FFileHelper::SaveStringToFile(JsonContent, *FilePath, FFileHelper::EEncodingOptions::ForceUTF8);
+}
+
+bool FGLTFBuilderUtility::SetExecutable(const TCHAR* Filename, bool bIsExecutable)
+{
+/* TODO: uncomment when GLTFLaunchHelper added for linux and macos
+#if PLATFORM_LINUX || PLATFORM_MAC
+	struct stat FileInfo;
+	FTCHARToUTF8 FilenameUtf8(Filename);
+	bool bSuccess = false;
+	// Set executable permission bit
+	if (stat(FilenameUtf8.Get(), &FileInfo) == 0)
+	{
+		mode_t ExeFlags = S_IXUSR | S_IXGRP | S_IXOTH;
+		FileInfo.st_mode = bIsExecutable ? (FileInfo.st_mode | ExeFlags) : (FileInfo.st_mode & ~ExeFlags);
+		bSuccess = chmod(FilenameUtf8.Get(), FileInfo.st_mode) == 0;
+	}
+	return bSuccess;
+#else
+	return true;
+#endif
+*/
+	return true;
+}
diff --git a/Engine/Plugins/Enterprise/GLTFExporter/MigratedFromGit/Source/GLTFExporter/Private/Builders/GLTFBuilderUtility.h b/Engine/Plugins/Enterprise/GLTFExporter/MigratedFromGit/Source/GLTFExporter/Private/Builders/GLTFBuilderUtility.h
--- a/Engine/Plugins/Enterprise/GLTFExporter/MigratedFromGit/Source/GLTFExporter/Private/Builders/GLTFBuilderUtility.h
+++ b/Engine/Plugins/Enterprise/GLTFExporter/MigratedFromGit/Source/GLTFExporter/Private/Builders/GLTFBuilderUtility.h
@@ -7,9 +7,16 @@
 
 struct FGLTFBuilderUtility
 {
+	static FString GetPluginDir();
+
 	static const TCHAR* GetFileExtension(EGLTFJsonMimeType MimeType);
 
 	static FString GetUniqueFilename(const FString& BaseFilename, const FString& FileExtension, const TSet<FString>& UniqueFilenames);
 
 	static bool IsGlbFile(const FString& Filename);
+
+	static bool ReadJsonFile(const FString& FilePath, TSharedPtr<FJsonObject>& JsonObject);
+	static bool WriteJsonFile(const FString& FilePath, const TSharedRef<FJsonObject>& JsonObject);
+
+	static bool SetExecutable(const TCHAR* Filename, bool bIsExecutable);
 };
diff --git a/Engine/Plugins/Enterprise/GLTFExporter/MigratedFromGit/Source/GLTFExporter/Private/Builders/GLTFContainerBuilder.cpp b/Engine/Plugins/Enterprise/GLTFExporter/MigratedFromGit/Source/GLTFExporter/Private/Builders/GLTFContainerBuilder.cpp
--- a/Engine/Plugins/Enterprise/GLTFExporter/MigratedFromGit/Source/GLTFExporter/Private/Builders/GLTFContainerBuilder.cpp
+++ b/Engine/Plugins/Enterprise/GLTFExporter/MigratedFromGit/Source/GLTFExporter/Private/Builders/GLTFContainerBuilder.cpp
@@ -1,35 +1,12 @@
 // Copyright Epic Games, Inc. All Rights Reserved.
 
 #include "Builders/GLTFContainerBuilder.h"
-#include "Builders/GLTFContainerUtility.h"
-#include "Builders/GLTFZipUtility.h"
-#include "GLTFExporterModule.h"
-#include "Interfaces/IPluginManager.h"
 
 FGLTFContainerBuilder::FGLTFContainerBuilder(const FString& FilePath, const UGLTFExportOptions* ExportOptions, bool bSelectedActorsOnly)
 	: FGLTFConvertBuilder(FilePath, ExportOptions, bSelectedActorsOnly)
 {
 }
 
-void FGLTFContainerBuilder::Write(FArchive& Archive, FFeedbackContext* Context)
-{
-	CompleteAllTasks(Context);
-
-	if (bIsGlbFile)
-	{
-		WriteGlb(Archive);
-	}
-	else
-	{
-		WriteJson(Archive);
-	}
-
-	if (ExportOptions->bBundleWebViewer)
-	{
-		BundleWebViewer();
-	}
-}
-
 void FGLTFContainerBuilder::WriteGlb(FArchive& Archive) const
 {
 	FBufferArchive JsonData;
@@ -37,75 +14,68 @@ void FGLTFContainerBuilder::WriteGlb(FArchive& Archive) const
 
 	const TArray<uint8>& BufferData = GetBufferData();
 
-	FGLTFContainerUtility::WriteGlb(Archive, JsonData, BufferData);
+	WriteGlb(Archive, JsonData, BufferData);
 }
 
-void FGLTFContainerBuilder::BundleWebViewer()
+void FGLTFContainerBuilder::WriteGlb(FArchive& Archive, const TArray<uint8>& JsonData, const TArray<uint8>& BinaryData)
 {
-	const FString PluginDir = IPluginManager::Get().FindPlugin(GLTFEXPORTER_MODULE_NAME)->GetBaseDir();
-	const FString ArchiveFile = PluginDir / TEXT("Resources") / TEXT("GLTFWebViewer.zip");
-
-	if (!FPaths::FileExists(ArchiveFile))
-	{
-		AddWarningMessage(FString::Printf(TEXT("No web viewer found at %s"), *ArchiveFile));
-		return;
-	}
+	const uint32 JsonChunkType = 0x4E4F534A; // "JSON" in ASCII
+	const uint32 BinaryChunkType = 0x004E4942; // "BIN" in ASCII
+	const uint32 FileSize =
+        3 * sizeof(uint32) +
+        2 * sizeof(uint32) + GetPaddedChunkSize(JsonData.Num()) +
+        2 * sizeof(uint32) + GetPaddedChunkSize(BinaryData.Num());
+
+	WriteHeader(Archive, FileSize);
+	WriteChunk(Archive, JsonChunkType, JsonData, 0x20);
+	WriteChunk(Archive, BinaryChunkType, BinaryData, 0x0);
+}
 
-	if (!FGLTFZipUtility::ExtractAllFiles(ArchiveFile, DirPath))
-	{
-		AddErrorMessage(FString::Printf(TEXT("Failed to extract web viewer at %s"), *ArchiveFile));
-		return;
-	}
+void FGLTFContainerBuilder::WriteHeader(FArchive& Archive, uint32 FileSize)
+{
+	const uint32 FileSignature = 0x46546C67; // "glTF" in ASCII
+	const uint32 FileVersion = 2;
 
-	UpdateWebViewerIndex();
+	WriteInt(Archive, FileSignature);
+	WriteInt(Archive, FileVersion);
+	WriteInt(Archive, FileSize);
 }
 
-void FGLTFContainerBuilder::UpdateWebViewerIndex()
+void FGLTFContainerBuilder::WriteChunk(FArchive& Archive, uint32 ChunkType, const TArray<uint8>& ChunkData, uint8 ChunkTrailingByte)
 {
-	const FString IndexFile = DirPath / TEXT("index.json");
+	const uint32 ChunkLength = GetPaddedChunkSize(ChunkData.Num());
+	const uint32 ChunkTrailing = GetTrailingChunkSize(ChunkData.Num());
 
-	if (!FPaths::FileExists(IndexFile))
-	{
-		AddWarningMessage(FString::Printf(TEXT("No index file found at %s"), *IndexFile));
-		return;
-	}
-
-	TSharedPtr<FJsonObject> JsonObject;
-	if (!ReadJsonFile(IndexFile, JsonObject))
-	{
-		AddWarningMessage(FString::Printf(TEXT("Failed to read index file at %s"), *IndexFile));
-		return;
-	}
+	WriteInt(Archive, ChunkLength);
+	WriteInt(Archive, ChunkType);
+	WriteData(Archive, ChunkData);
+	WriteFill(Archive, ChunkTrailing, ChunkTrailingByte);
+}
 
-	JsonObject->SetArrayField(TEXT("assets"), { MakeShared<FJsonValueString>(FPaths::GetCleanFilename(FilePath)) });
+void FGLTFContainerBuilder::WriteInt(FArchive& Archive, uint32 Value)
+{
+	Archive.SerializeInt(Value, MAX_uint32);
+}
 
-	if (!WriteJsonFile(IndexFile, JsonObject.ToSharedRef()))
-	{
-		AddWarningMessage(FString::Printf(TEXT("Failed to write index file at %s"), *IndexFile));
-	}
+void FGLTFContainerBuilder::WriteData(FArchive& Archive, const TArray<uint8>& Data)
+{
+	Archive.Serialize(const_cast<uint8*>(Data.GetData()), Data.Num());
 }
 
-bool FGLTFContainerBuilder::ReadJsonFile(const FString& FilePath, TSharedPtr<FJsonObject>& JsonObject)
+void FGLTFContainerBuilder::WriteFill(FArchive& Archive, int32 Size, uint8 Value)
 {
-	FString JsonContent;
-	if (!FFileHelper::LoadFileToString(JsonContent, *FilePath))
+	while (--Size >= 0)
 	{
-		return false;
+		Archive.Serialize(&Value, sizeof(Value));
 	}
-
-	const TSharedRef<TJsonReader<>> JsonReader = TJsonReaderFactory<>::Create(JsonContent);
-	return FJsonSerializer::Deserialize(JsonReader, JsonObject) && JsonObject.IsValid();
 }
 
-bool FGLTFContainerBuilder::WriteJsonFile(const FString& FilePath, const TSharedRef<FJsonObject>& JsonObject)
+int32 FGLTFContainerBuilder::GetPaddedChunkSize(int32 Size)
 {
-	FString JsonContent;
-	const TSharedRef<TJsonWriter<>> JsonWriter = TJsonWriterFactory<>::Create(&JsonContent);
-
-	if (!FJsonSerializer::Serialize(JsonObject, JsonWriter))
-	{
-		return false;
-	}
+	return (Size + 3) & ~3;
+}
 
-	return FFileHelper::SaveStringToFile(JsonContent, *FilePath, FFileHelper::EEncodingOptions::ForceUTF8);
+int32 FGLTFContainerBuilder::GetTrailingChunkSize(int32 Size)
+{
+	return (4 - (Size & 3)) & 3;
 }
diff --git a/Engine/Plugins/Enterprise/GLTFExporter/MigratedFromGit/Source/GLTFExporter/Private/Builders/GLTFContainerBuilder.h b/Engine/Plugins/Enterprise/GLTFExporter/MigratedFromGit/Source/GLTFExporter/Private/Builders/GLTFContainerBuilder.h
--- a/Engine/Plugins/Enterprise/GLTFExporter/MigratedFromGit/Source/GLTFExporter/Private/Builders/GLTFContainerBuilder.h
+++ b/Engine/Plugins/Enterprise/GLTFExporter/MigratedFromGit/Source/GLTFExporter/Private/Builders/GLTFContainerBuilder.h
@@ -6,20 +6,23 @@
 
 class FGLTFContainerBuilder : public FGLTFConvertBuilder
 {
-public:
+protected:
 
 	FGLTFContainerBuilder(const FString& FilePath, const UGLTFExportOptions* ExportOptions, bool bSelectedActorsOnly);
 
-	void Write(FArchive& Archive, FFeedbackContext* Context = GWarn);
+	void WriteGlb(FArchive& Archive) const;
 
 private:
 
-	void WriteGlb(FArchive& Archive) const;
+	static void WriteGlb(FArchive& Archive, const TArray<uint8>& JsonData, const TArray<uint8>& BinaryData);
 
-	void BundleWebViewer();
+	static void WriteHeader(FArchive& Archive, uint32 FileSize);
+	static void WriteChunk(FArchive& Archive, uint32 ChunkType, const TArray<uint8>& ChunkData, uint8 ChunkTrailingByte);
 
-	void UpdateWebViewerIndex();
+	static void WriteInt(FArchive& Archive, uint32 Value);
+	static void WriteData(FArchive& Archive, const TArray<uint8>& Data);
+	static void WriteFill(FArchive& Archive, int32 Size, uint8 Value);
 
-	static bool ReadJsonFile(const FString& FilePath, TSharedPtr<FJsonObject>& JsonObject);
-	static bool WriteJsonFile(const FString& FilePath, const TSharedRef<FJsonObject>& JsonObject);
+	static int32 GetPaddedChunkSize(int32 Size);
+	static int32 GetTrailingChunkSize(int32 Size);
 };
diff --git a/Engine/Plugins/Enterprise/GLTFExporter/MigratedFromGit/Source/GLTFExporter/Private/Builders/GLTFContainerUtility.cpp b/Engine/Plugins/Enterprise/GLTFExporter/MigratedFromGit/Source/GLTFExporter/Private/Builders/GLTFContainerUtility.cpp
deleted file mode 100644
diff --git a/Engine/Plugins/Enterprise/GLTFExporter/MigratedFromGit/Source/GLTFExporter/Private/Builders/GLTFContainerUtility.h b/Engine/Plugins/Enterprise/GLTFExporter/MigratedFromGit/Source/GLTFExporter/Private/Builders/GLTFContainerUtility.h
deleted file mode 100644
diff --git a/Engine/Plugins/Enterprise/GLTFExporter/MigratedFromGit/Source/GLTFExporter/Private/Builders/GLTFWebBuilder.cpp b/Engine/Plugins/Enterprise/GLTFExporter/MigratedFromGit/Source/GLTFExporter/Private/Builders/GLTFWebBuilder.cpp
new file mode 100644
--- /dev/null
+++ b/Engine/Plugins/Enterprise/GLTFExporter/MigratedFromGit/Source/GLTFExporter/Private/Builders/GLTFWebBuilder.cpp
@@ -0,0 +1,105 @@
+// Copyright Epic Games, Inc. All Rights Reserved.
+
+#include "Builders/GLTFWebBuilder.h"
+#include "Builders/GLTFBuilderUtility.h"
+#include "Builders/GLTFZipUtility.h"
+
+FGLTFWebBuilder::FGLTFWebBuilder(const FString& FilePath, const UGLTFExportOptions* ExportOptions, bool bSelectedActorsOnly)
+	: FGLTFContainerBuilder(FilePath, ExportOptions, bSelectedActorsOnly)
+{
+}
+
+void FGLTFWebBuilder::Write(FArchive& Archive, FFeedbackContext* Context)
+{
+	CompleteAllTasks(Context);
+
+	if (bIsGlbFile)
+	{
+		WriteGlb(Archive);
+	}
+	else
+	{
+		WriteJson(Archive);
+	}
+
+	if (ExportOptions->bBundleWebViewer)
+	{
+		const FString ResourcesDir = FGLTFBuilderUtility::GetPluginDir() / TEXT("Resources");
+		BundleWebViewer(ResourcesDir);
+		BundleLaunchHelper(ResourcesDir);
+	}
+}
+
+void FGLTFWebBuilder::BundleWebViewer(const FString& ResourcesDir)
+{
+	const FString ArchiveFile = ResourcesDir / TEXT("GLTFWebViewer.zip");
+
+	if (!FPaths::FileExists(ArchiveFile))
+	{
+		AddWarningMessage(FString::Printf(TEXT("No web viewer archive found at %s"), *ArchiveFile));
+		return;
+	}
+
+	if (!FGLTFZipUtility::ExtractAllFiles(ArchiveFile, DirPath))
+	{
+		AddErrorMessage(FString::Printf(TEXT("Failed to extract web viewer files from %s"), *ArchiveFile));
+		return;
+	}
+
+	const FString IndexFile = DirPath / TEXT("index.json");
+	TSharedPtr<FJsonObject> JsonObject;
+
+	if (!FGLTFBuilderUtility::ReadJsonFile(IndexFile, JsonObject))
+	{
+		AddWarningMessage(FString::Printf(TEXT("Failed to read web viewer index at %s"), *IndexFile));
+		return;
+	}
+
+	JsonObject->SetArrayField(TEXT("assets"), { MakeShared<FJsonValueString>(FPaths::GetCleanFilename(FilePath)) });
+
+	if (!FGLTFBuilderUtility::WriteJsonFile(IndexFile, JsonObject.ToSharedRef()))
+	{
+		AddWarningMessage(FString::Printf(TEXT("Failed to write web viewer index at %s"), *IndexFile));
+	}
+}
+
+void FGLTFWebBuilder::BundleLaunchHelper(const FString& ResourcesDir)
+{
+	const FString ExecutableName = GetLaunchHelperExecutable();
+	if (ExecutableName.IsEmpty())
+	{
+		return;
+	}
+
+	const FString ArchiveFile = ResourcesDir / TEXT("GLTFLaunchHelper.zip");
+
+	if (!FPaths::FileExists(ArchiveFile))
+	{
+		AddWarningMessage(FString::Printf(TEXT("No launch helper archive found at %s"), *ArchiveFile));
+		return;
+	}
+
+	if (!FGLTFZipUtility::ExtractOneFile(ArchiveFile, ExecutableName, DirPath))
+	{
+		AddErrorMessage(FString::Printf(TEXT("Failed to extract launch helper file (%s) from %s"), *ExecutableName, *ArchiveFile));
+		return;
+	}
+
+	const FString ExecutableFile = DirPath / ExecutableName;
+
+	if (!FGLTFBuilderUtility::SetExecutable(*ExecutableFile, true))
+	{
+		AddWarningMessage(FString::Printf(TEXT("Failed to make launch helper file executable at %s"), *ExecutableFile));
+	}
+}
+
+const TCHAR* FGLTFWebBuilder::GetLaunchHelperExecutable()
+{
+#if PLATFORM_WINDOWS
+	return TEXT("GLTFLaunchHelper.exe");
+#elif PLATFORM_MAC
+	return TEXT(""); // TODO: add GLTFLaunchHelper for macos
+#else
+	return TEXT(""); // TODO: add GLTFLaunchHelper for linux
+#endif
+}
diff --git a/Engine/Plugins/Enterprise/GLTFExporter/MigratedFromGit/Source/GLTFExporter/Private/Builders/GLTFWebBuilder.h b/Engine/Plugins/Enterprise/GLTFExporter/MigratedFromGit/Source/GLTFExporter/Private/Builders/GLTFWebBuilder.h
new file mode 100644
--- /dev/null
+++ b/Engine/Plugins/Enterprise/GLTFExporter/MigratedFromGit/Source/GLTFExporter/Private/Builders/GLTFWebBuilder.h
@@ -0,0 +1,21 @@
+// Copyright Epic Games, Inc. All Rights Reserved.
+
+#pragma once
+
+#include "Builders/GLTFContainerBuilder.h"
+
+class FGLTFWebBuilder : public FGLTFContainerBuilder
+{
+public:
+
+	FGLTFWebBuilder(const FString& FilePath, const UGLTFExportOptions* ExportOptions, bool bSelectedActorsOnly);
+
+	void Write(FArchive& Archive, FFeedbackContext* Context = GWarn);
+
+private:
+
+	void BundleWebViewer(const FString& ResourcesDir);
+	void BundleLaunchHelper(const FString& ResourcesDir);
+
+	static const TCHAR* GetLaunchHelperExecutable();
+};
diff --git a/Engine/Plugins/Enterprise/GLTFExporter/MigratedFromGit/Source/GLTFExporter/Private/Builders/GLTFZipUtility.cpp b/Engine/Plugins/Enterprise/GLTFExporter/MigratedFromGit/Source/GLTFExporter/Private/Builders/GLTFZipUtility.cpp
--- a/Engine/Plugins/Enterprise/GLTFExporter/MigratedFromGit/Source/GLTFExporter/Private/Builders/GLTFZipUtility.cpp
+++ b/Engine/Plugins/Enterprise/GLTFExporter/MigratedFromGit/Source/GLTFExporter/Private/Builders/GLTFZipUtility.cpp
@@ -9,7 +9,7 @@ THIRD_PARTY_INCLUDES_END
 
 TArray<FString> FGLTFZipUtility::GetAllFiles(const FString& ArchiveFile)
 {
-	const unzFile ZipHandle = unzOpen64(TCHAR_TO_ANSI(*ArchiveFile));
+	const unzFile ZipHandle = unzOpen64(TCHAR_TO_UTF8(*ArchiveFile));
 	if (ZipHandle == nullptr)
 	{
 		return {};
@@ -41,7 +41,7 @@ TArray<FString> FGLTFZipUtility::GetAllFiles(const FString& ArchiveFile)
 
 bool FGLTFZipUtility::ExtractAllFiles(const FString& ArchiveFile, const FString& TargetDirectory)
 {
-	const unzFile ZipHandle = unzOpen64(TCHAR_TO_ANSI(*ArchiveFile));
+	const unzFile ZipHandle = unzOpen64(TCHAR_TO_UTF8(*ArchiveFile));
 	if (ZipHandle == nullptr)
 	{
 		return false;
@@ -75,7 +75,7 @@ bool FGLTFZipUtility::ExtractAllFiles(const FString& ArchiveFile, const FString&
 
 bool FGLTFZipUtility::ExtractOneFile(const FString& ArchiveFile, const FString& FileToExtract, const FString& TargetDirectory)
 {
-	const unzFile ZipHandle = unzOpen64(TCHAR_TO_ANSI(*ArchiveFile));
+	const unzFile ZipHandle = unzOpen64(TCHAR_TO_UTF8(*ArchiveFile));
 	if (ZipHandle == nullptr)
 	{
 		return false;
@@ -87,7 +87,7 @@ bool FGLTFZipUtility::ExtractOneFile(const FString& ArchiveFile, const FString&
 		return false;
 	}
 
-	if (unzLocateFile(ZipHandle, TCHAR_TO_ANSI(*FileToExtract), 0) != UNZ_OK)
+	if (unzLocateFile(ZipHandle, TCHAR_TO_UTF8(*FileToExtract), 0) != UNZ_OK)
 	{
 		unzClose(ZipHandle);
 		return false;
@@ -112,14 +112,14 @@ FString FGLTFZipUtility::GetCurrentFilename(void* ZipHandle)
 	}
 
 	TArray<char> Filename;
-
 	Filename.Init('\0', FileInfo.size_filename + 1);
-	if (unzGetCurrentFileInfo64(ZipHandle, nullptr, Filename.GetData(), Filename.Num() - 1, nullptr, 0, nullptr, 0) != UNZ_OK)
+
+	if (unzGetCurrentFileInfo64(ZipHandle, nullptr, Filename.GetData(), FileInfo.size_filename, nullptr, 0, nullptr, 0) != UNZ_OK)
 	{
 		return {};
 	}
 
-	return Filename.GetData();
+	return UTF8_TO_TCHAR(Filename.GetData());
 }
 
 bool FGLTFZipUtility::ExtractCurrentFile(void* ZipHandle, const FString& TargetFile)
diff --git a/Engine/Plugins/Enterprise/GLTFExporter/MigratedFromGit/Source/GLTFExporter/Private/Exporters/GLTFExporter.cpp b/Engine/Plugins/Enterprise/GLTFExporter/MigratedFromGit/Source/GLTFExporter/Private/Exporters/GLTFExporter.cpp
--- a/Engine/Plugins/Enterprise/GLTFExporter/MigratedFromGit/Source/GLTFExporter/Private/Exporters/GLTFExporter.cpp
+++ b/Engine/Plugins/Enterprise/GLTFExporter/MigratedFromGit/Source/GLTFExporter/Private/Exporters/GLTFExporter.cpp
@@ -3,7 +3,7 @@
 #include "Exporters/GLTFExporter.h"
 #include "GLTFExportOptions.h"
 #include "UI/GLTFExportOptionsWindow.h"
-#include "Builders/GLTFContainerBuilder.h"
+#include "Builders/GLTFWebBuilder.h"
 #include "UObject/GCObjectScopeGuard.h"
 #include "AssetExportTask.h"
 
@@ -33,7 +33,7 @@ bool UGLTFExporter::ExportBinary(UObject* Object, const TCHAR* Type, FArchive& A
 	// TODO: add support for UAssetExportTask::IgnoreObjectList?
 
 	FGCObjectScopeGuard OptionsGuard(Options);
-	FGLTFContainerBuilder Builder(CurrentFilename, Options, bSelectedOnly);
+	FGLTFWebBuilder Builder(CurrentFilename, Options, bSelectedOnly);
 
 	const bool bSuccess = AddObject(Builder, Object);
 	if (bSuccess)
