[GLTFExporter] refactor: rename GLTFBuilderUtility to GLTFFileUtility

commit 0cc11958653c882ccd82a5429e5fc9da86d10f4c
Date: Tue, 10 Nov 2020 21:04:17 +0100

    refactor: rename GLTFBuilderUtility to GLTFFileUtility
    (cherry picked from commit 488934d2c621ee30ca2589b41afcdc52db51177f​)

#skipci

diff --git a/Engine/Plugins/Enterprise/GLTFExporter/MigratedFromGit/Source/GLTFExporter/Private/Builders/GLTFBuilder.cpp b/Engine/Plugins/Enterprise/GLTFExporter/MigratedFromGit/Source/GLTFExporter/Private/Builders/GLTFBuilder.cpp
--- a/Engine/Plugins/Enterprise/GLTFExporter/MigratedFromGit/Source/GLTFExporter/Private/Builders/GLTFBuilder.cpp
+++ b/Engine/Plugins/Enterprise/GLTFExporter/MigratedFromGit/Source/GLTFExporter/Private/Builders/GLTFBuilder.cpp
@@ -1,10 +1,10 @@
 // Copyright Epic Games, Inc. All Rights Reserved.
 
 #include "Builders/GLTFBuilder.h"
-#include "Builders/GLTFBuilderUtility.h"
+#include "Builders/GLTFFileUtility.h"
 
 FGLTFBuilder::FGLTFBuilder(const FString& FilePath, const UGLTFExportOptions* ExportOptions)
-	: bIsGlbFile(FGLTFBuilderUtility::IsGlbFile(FilePath))
+	: bIsGlbFile(FGLTFFileUtility::IsGlbFile(FilePath))
 	, FilePath(FilePath)
 	, DirPath(FPaths::GetPath(FilePath))
 	, ExportOptions(ExportOptions)
diff --git a/Engine/Plugins/Enterprise/GLTFExporter/MigratedFromGit/Source/GLTFExporter/Private/Builders/GLTFBuilderUtility.cpp b/Engine/Plugins/Enterprise/GLTFExporter/MigratedFromGit/Source/GLTFExporter/Private/Builders/GLTFBuilderUtility.cpp
deleted file mode 100644
diff --git a/Engine/Plugins/Enterprise/GLTFExporter/MigratedFromGit/Source/GLTFExporter/Private/Builders/GLTFBuilderUtility.h b/Engine/Plugins/Enterprise/GLTFExporter/MigratedFromGit/Source/GLTFExporter/Private/Builders/GLTFBuilderUtility.h
deleted file mode 100644
diff --git a/Engine/Plugins/Enterprise/GLTFExporter/MigratedFromGit/Source/GLTFExporter/Private/Builders/GLTFFileUtility.cpp b/Engine/Plugins/Enterprise/GLTFExporter/MigratedFromGit/Source/GLTFExporter/Private/Builders/GLTFFileUtility.cpp
new file mode 100644
--- /dev/null
+++ b/Engine/Plugins/Enterprise/GLTFExporter/MigratedFromGit/Source/GLTFExporter/Private/Builders/GLTFFileUtility.cpp
@@ -0,0 +1,101 @@
+// Copyright Epic Games, Inc. All Rights Reserved.
+
+#include "Builders/GLTFFileUtility.h"
+#include "GLTFExporterModule.h"
+#include "Interfaces/IPluginManager.h"
+#include "Serialization/JsonReader.h"
+#include "Serialization/JsonSerializer.h"
+#include "Misc/FileHelper.h"
+
+#if PLATFORM_LINUX || PLATFORM_MAC
+// TODO: uncomment when GLTFLaunchHelper added for linux and macos
+// #include <sys/stat.h>
+#endif
+
+FString FGLTFFileUtility::GetPluginDir()
+{
+	return IPluginManager::Get().FindPlugin(GLTFEXPORTER_MODULE_NAME)->GetBaseDir();
+}
+
+const TCHAR* FGLTFFileUtility::GetFileExtension(EGLTFJsonMimeType MimeType)
+{
+	switch (MimeType)
+	{
+		case EGLTFJsonMimeType::PNG:  return TEXT(".png");
+		case EGLTFJsonMimeType::JPEG: return TEXT(".jpg");
+		default:
+			checkNoEntry();
+			return TEXT("");
+	}
+}
+
+FString FGLTFFileUtility::GetUniqueFilename(const FString& BaseFilename, const FString& FileExtension, const TSet<FString>& UniqueFilenames)
+{
+	FString Filename = BaseFilename + FileExtension;
+	if (!UniqueFilenames.Contains(Filename))
+	{
+		return Filename;
+	}
+
+	int32 Suffix = 1;
+	do
+	{
+		Filename = BaseFilename + TEXT('_') + FString::FromInt(Suffix) + FileExtension;
+		Suffix++;
+	}
+	while (UniqueFilenames.Contains(Filename));
+
+	return Filename;
+}
+
+bool FGLTFFileUtility::IsGlbFile(const FString& Filename)
+{
+	return FPaths::GetExtension(Filename).Equals(TEXT("glb"), ESearchCase::IgnoreCase);
+}
+
+bool FGLTFFileUtility::ReadJsonFile(const FString& FilePath, TSharedPtr<FJsonObject>& JsonObject)
+{
+	FString JsonContent;
+	if (!FFileHelper::LoadFileToString(JsonContent, *FilePath))
+	{
+		return false;
+	}
+
+	const TSharedRef<TJsonReader<>> JsonReader = TJsonReaderFactory<>::Create(JsonContent);
+	return FJsonSerializer::Deserialize(JsonReader, JsonObject) && JsonObject.IsValid();
+}
+
+bool FGLTFFileUtility::WriteJsonFile(const FString& FilePath, const TSharedRef<FJsonObject>& JsonObject)
+{
+	FString JsonContent;
+	const TSharedRef<TJsonWriter<>> JsonWriter = TJsonWriterFactory<>::Create(&JsonContent);
+
+	if (!FJsonSerializer::Serialize(JsonObject, JsonWriter))
+	{
+		return false;
+	}
+
+	return FFileHelper::SaveStringToFile(JsonContent, *FilePath, FFileHelper::EEncodingOptions::ForceUTF8);
+}
+
+bool FGLTFFileUtility::SetExecutable(const TCHAR* Filename, bool bIsExecutable)
+{
+/* TODO: uncomment when GLTFLaunchHelper added for linux and macos
+#if PLATFORM_LINUX || PLATFORM_MAC
+	struct stat FileInfo;
+	FTCHARToUTF8 FilenameUtf8(Filename);
+	bool bSuccess = false;
+	// Set executable permission bit
+	if (stat(FilenameUtf8.Get(), &FileInfo) == 0)
+	{
+		mode_t ExeFlags = S_IXUSR | S_IXGRP | S_IXOTH;
+		FileInfo.st_mode = bIsExecutable ? (FileInfo.st_mode | ExeFlags) : (FileInfo.st_mode & ~ExeFlags);
+		bSuccess = chmod(FilenameUtf8.Get(), FileInfo.st_mode) == 0;
+	}
+	return bSuccess;
+#else
+	return true;
+#endif
+*/
+	return true;
+}
diff --git a/Engine/Plugins/Enterprise/GLTFExporter/MigratedFromGit/Source/GLTFExporter/Private/Builders/GLTFFileUtility.h b/Engine/Plugins/Enterprise/GLTFExporter/MigratedFromGit/Source/GLTFExporter/Private/Builders/GLTFFileUtility.h
new file mode 100644
--- /dev/null
+++ b/Engine/Plugins/Enterprise/GLTFExporter/MigratedFromGit/Source/GLTFExporter/Private/Builders/GLTFFileUtility.h
@@ -0,0 +1,22 @@
+// Copyright Epic Games, Inc. All Rights Reserved.
+
+#pragma once
+
+#include "CoreMinimal.h"
+#include "Json/GLTFJsonEnums.h"
+
+struct FGLTFFileUtility
+{
+	static FString GetPluginDir();
+
+	static const TCHAR* GetFileExtension(EGLTFJsonMimeType MimeType);
+
+	static FString GetUniqueFilename(const FString& BaseFilename, const FString& FileExtension, const TSet<FString>& UniqueFilenames);
+
+	static bool IsGlbFile(const FString& Filename);
+
+	static bool ReadJsonFile(const FString& FilePath, TSharedPtr<FJsonObject>& JsonObject);
+	static bool WriteJsonFile(const FString& FilePath, const TSharedRef<FJsonObject>& JsonObject);
+
+	static bool SetExecutable(const TCHAR* Filename, bool bIsExecutable);
+};
diff --git a/Engine/Plugins/Enterprise/GLTFExporter/MigratedFromGit/Source/GLTFExporter/Private/Builders/GLTFImageBuilder.cpp b/Engine/Plugins/Enterprise/GLTFExporter/MigratedFromGit/Source/GLTFExporter/Private/Builders/GLTFImageBuilder.cpp
--- a/Engine/Plugins/Enterprise/GLTFExporter/MigratedFromGit/Source/GLTFExporter/Private/Builders/GLTFImageBuilder.cpp
+++ b/Engine/Plugins/Enterprise/GLTFExporter/MigratedFromGit/Source/GLTFExporter/Private/Builders/GLTFImageBuilder.cpp
@@ -1,7 +1,7 @@
 // Copyright Epic Games, Inc. All Rights Reserved.
 
 #include "Builders/GLTFImageBuilder.h"
-#include "Builders/GLTFBuilderUtility.h"
+#include "Builders/GLTFFileUtility.h"
 #include "Builders/GLTFImageUtility.h"
 #include "Misc/FileHelper.h"
 
@@ -94,8 +94,8 @@ EGLTFJsonMimeType FGLTFImageBuilder::GetImageFormat(const FColor* Pixels, FIntPo
 
 FString FGLTFImageBuilder::SaveImageToFile(const void* CompressedData, int64 CompressedByteLength, EGLTFJsonMimeType MimeType, const FString& Name)
 {
-	const TCHAR* Extension = FGLTFBuilderUtility::GetFileExtension(MimeType);
-	const FString ImageUri = FGLTFBuilderUtility::GetUniqueFilename(Name, Extension, UniqueImageUris);
+	const TCHAR* Extension = FGLTFFileUtility::GetFileExtension(MimeType);
+	const FString ImageUri = FGLTFFileUtility::GetUniqueFilename(Name, Extension, UniqueImageUris);
 
 	const TArrayView<const uint8> ImageData(static_cast<const uint8*>(CompressedData), CompressedByteLength);
 	const FString ImagePath = FPaths::Combine(DirPath, ImageUri);
diff --git a/Engine/Plugins/Enterprise/GLTFExporter/MigratedFromGit/Source/GLTFExporter/Private/Builders/GLTFWebBuilder.cpp b/Engine/Plugins/Enterprise/GLTFExporter/MigratedFromGit/Source/GLTFExporter/Private/Builders/GLTFWebBuilder.cpp
--- a/Engine/Plugins/Enterprise/GLTFExporter/MigratedFromGit/Source/GLTFExporter/Private/Builders/GLTFWebBuilder.cpp
+++ b/Engine/Plugins/Enterprise/GLTFExporter/MigratedFromGit/Source/GLTFExporter/Private/Builders/GLTFWebBuilder.cpp
@@ -1,7 +1,7 @@
 // Copyright Epic Games, Inc. All Rights Reserved.
 
 #include "Builders/GLTFWebBuilder.h"
-#include "Builders/GLTFBuilderUtility.h"
+#include "Builders/GLTFFileUtility.h"
 #include "Builders/GLTFZipUtility.h"
 
 FGLTFWebBuilder::FGLTFWebBuilder(const FString& FilePath, const UGLTFExportOptions* ExportOptions, bool bSelectedActorsOnly)
@@ -24,7 +24,7 @@ void FGLTFWebBuilder::Write(FArchive& Archive, FFeedbackContext* Context)
 
 	if (ExportOptions->bBundleWebViewer)
 	{
-		const FString ResourcesDir = FGLTFBuilderUtility::GetPluginDir() / TEXT("Resources");
+		const FString ResourcesDir = FGLTFFileUtility::GetPluginDir() / TEXT("Resources");
 		BundleWebViewer(ResourcesDir);
 		BundleLaunchHelper(ResourcesDir);
 	}
@@ -49,7 +49,7 @@ void FGLTFWebBuilder::BundleWebViewer(const FString& ResourcesDir)
 	const FString IndexFile = DirPath / TEXT("index.json");
 	TSharedPtr<FJsonObject> JsonObject;
 
-	if (!FGLTFBuilderUtility::ReadJsonFile(IndexFile, JsonObject))
+	if (!FGLTFFileUtility::ReadJsonFile(IndexFile, JsonObject))
 	{
 		AddWarningMessage(FString::Printf(TEXT("Failed to read web viewer index at %s"), *IndexFile));
 		return;
@@ -57,7 +57,7 @@ void FGLTFWebBuilder::BundleWebViewer(const FString& ResourcesDir)
 
 	JsonObject->SetArrayField(TEXT("assets"), { MakeShared<FJsonValueString>(FPaths::GetCleanFilename(FilePath)) });
 
-	if (!FGLTFBuilderUtility::WriteJsonFile(IndexFile, JsonObject.ToSharedRef()))
+	if (!FGLTFFileUtility::WriteJsonFile(IndexFile, JsonObject.ToSharedRef()))
 	{
 		AddWarningMessage(FString::Printf(TEXT("Failed to write web viewer index at %s"), *IndexFile));
 	}
@@ -87,7 +87,7 @@ void FGLTFWebBuilder::BundleLaunchHelper(const FString& ResourcesDir)
 
 	const FString ExecutableFile = DirPath / ExecutableName;
 
-	if (!FGLTFBuilderUtility::SetExecutable(*ExecutableFile, true))
+	if (!FGLTFFileUtility::SetExecutable(*ExecutableFile, true))
 	{
 		AddWarningMessage(FString::Printf(TEXT("Failed to make launch helper file executable at %s"), *ExecutableFile));
 	}
