[GLTFExporter] refactor: improve checking for non-uniform parent scale by using FVecto::IsUniform

commit d9f6533fa1d8f646158926060bd3316e56cedd60
Date: Wed, 11 Nov 2020 07:16:24 +0100

    refactor: improve checking for non-uniform parent scale by using FVecto::IsUniform
    (cherry picked from commit 9efb78ed2b791bfa665e49a3556bd1d779430a38​)

#skipci

diff --git a/Engine/Plugins/Enterprise/GLTFExporter/MigratedFromGit/Source/GLTFExporter/Private/Converters/GLTFLightMapConverters.cpp b/Engine/Plugins/Enterprise/GLTFExporter/MigratedFromGit/Source/GLTFExporter/Private/Converters/GLTFLightMapConverters.cpp
--- a/Engine/Plugins/Enterprise/GLTFExporter/MigratedFromGit/Source/GLTFExporter/Private/Converters/GLTFLightMapConverters.cpp
+++ b/Engine/Plugins/Enterprise/GLTFExporter/MigratedFromGit/Source/GLTFExporter/Private/Converters/GLTFLightMapConverters.cpp
@@ -1,7 +1,6 @@
 // Copyright Epic Games, Inc. All Rights Reserved.
 
 #include "Converters/GLTFLightMapConverters.h"
-#include "Json/GLTFJsonLightMap.h"
 #include "Builders/GLTFContainerBuilder.h"
 #include "Engine/MapBuildDataRegistry.h"
 
diff --git a/Engine/Plugins/Enterprise/GLTFExporter/MigratedFromGit/Source/GLTFExporter/Private/Converters/GLTFNodeConverters.cpp b/Engine/Plugins/Enterprise/GLTFExporter/MigratedFromGit/Source/GLTFExporter/Private/Converters/GLTFNodeConverters.cpp
--- a/Engine/Plugins/Enterprise/GLTFExporter/MigratedFromGit/Source/GLTFExporter/Private/Converters/GLTFNodeConverters.cpp
+++ b/Engine/Plugins/Enterprise/GLTFExporter/MigratedFromGit/Source/GLTFExporter/Private/Converters/GLTFNodeConverters.cpp
@@ -91,7 +91,7 @@ FGLTFJsonNodeIndex FGLTFComponentConverter::Convert(const USceneComponent* Scene
 	if (ParentComponent != nullptr && !SceneComponent->IsUsingAbsoluteScale())
 	{
 		const FVector ParentScale = ParentComponent->GetComponentScale();
-		if (ParentScale.X != ParentScale.Y || ParentScale.Y != ParentScale.Z)
+		if (!ParentScale.IsUniform())
 		{
 			Builder.AddWarningMessage(
 				FString::Printf(TEXT("Non-uniform parent scale (%s) for component %s (in actor %s) may be represented differently in glTF"),
