[GLTFExporter] fix: update expressions for some EditCondition in GLTFExportOptions to fix issue with texture compression settings not being disabled correctly

commit 859203fb6c159097ef3518aa8f5bd28795cdf5b3
Date: Thu, 12 Nov 2020 00:12:45 +0100

    fix: update expressions for some EditCondition in GLTFExportOptions to fix issue with texture compression settings not being disabled correctly
    (cherry picked from commit 057cc47032f0323cb9bfa0f6b86b598936221bd0​)

#skipci

diff --git a/Engine/Plugins/Enterprise/GLTFExporter/MigratedFromGit/Source/GLTFExporter/Public/GLTFExportOptions.h b/Engine/Plugins/Enterprise/GLTFExporter/MigratedFromGit/Source/GLTFExporter/Public/GLTFExportOptions.h
--- a/Engine/Plugins/Enterprise/GLTFExporter/MigratedFromGit/Source/GLTFExporter/Public/GLTFExportOptions.h
+++ b/Engine/Plugins/Enterprise/GLTFExporter/MigratedFromGit/Source/GLTFExporter/Public/GLTFExportOptions.h
@@ -109,10 +109,10 @@ public:
 	UPROPERTY(EditAnywhere, BlueprintReadWrite, Config, Category = Texture)
 	EGLTFTextureCompression TextureCompression;
 
-	UPROPERTY(EditAnywhere, BlueprintReadWrite, Config, Category = Texture, Meta = (ClampMin = "0", ClampMax = "100", EditCondition = "TextureCompression != EGLTFExporterTextureCompression::PNG"))
+	UPROPERTY(EditAnywhere, BlueprintReadWrite, Config, Category = Texture, Meta = (ClampMin = "0", ClampMax = "100", EditCondition = "TextureCompression != EGLTFTextureCompression::PNG"))
 	int32 TextureCompressionQuality;
 
-	UPROPERTY(EditAnywhere, BlueprintReadWrite, Config, Category = Texture, Meta = (Bitmask, BitmaskEnum = EGLTFTextureGroupFlags, EditCondition = "TextureCompression != EGLTFExporterTextureCompression::PNG"))
+	UPROPERTY(EditAnywhere, BlueprintReadWrite, Config, Category = Texture, Meta = (Bitmask, BitmaskEnum = EGLTFTextureGroupFlags, EditCondition = "TextureCompression != EGLTFTextureCompression::PNG"))
 	int32 LosslessCompressTextures;
 
 	UPROPERTY(EditAnywhere, BlueprintReadWrite, Config, Category = Texture)
