[GLTFExporter] feat: add support for using mesh data when baking materials for USkeletalMesh

commit 0e84d78782bb5597074ea4f6d6c96097529d887f
Date: Wed, 11 Nov 2020 09:52:09 +0100

    feat: add support for using mesh data when baking materials for USkeletalMesh
    (cherry picked from commit bd63704bf68780c4dd26f488577c06003742a707​)

#skipci

diff --git a/Engine/Plugins/Enterprise/GLTFExporter/MigratedFromGit/Source/GLTFExporter/Private/Converters/GLTFMeshData.cpp b/Engine/Plugins/Enterprise/GLTFExporter/MigratedFromGit/Source/GLTFExporter/Private/Converters/GLTFMeshData.cpp
--- a/Engine/Plugins/Enterprise/GLTFExporter/MigratedFromGit/Source/GLTFExporter/Private/Converters/GLTFMeshData.cpp
+++ b/Engine/Plugins/Enterprise/GLTFExporter/MigratedFromGit/Source/GLTFExporter/Private/Converters/GLTFMeshData.cpp
@@ -2,6 +2,7 @@
 
 #include "Converters/GLTFMeshData.h"
 #include "Converters/GLTFNameUtility.h"
+#include "Converters/GLTFMeshUtility.h"
 #include "StaticMeshAttributes.h"
 #include "Developer/MeshMergeUtilities/Private/MeshMergeHelpers.h"
 
@@ -32,7 +33,7 @@ FGLTFMeshData::FGLTFMeshData(const USkeletalMesh* SkeletalMesh, const USkeletalM
 	}
 	else
 	{
-		// TODO: add support for skeletal meshes by implementing custom utilities for retrieving mesh & sections
+		FGLTFMeshUtility::RetrieveMesh(const_cast<USkeletalMesh*>(SkeletalMesh), LODIndex, Description);
 		SkeletalMesh->GetName(Name);
 	}
 }
diff --git a/Engine/Plugins/Enterprise/GLTFExporter/MigratedFromGit/Source/GLTFExporter/Private/Converters/GLTFMeshUtility.cpp b/Engine/Plugins/Enterprise/GLTFExporter/MigratedFromGit/Source/GLTFExporter/Private/Converters/GLTFMeshUtility.cpp
--- a/Engine/Plugins/Enterprise/GLTFExporter/MigratedFromGit/Source/GLTFExporter/Private/Converters/GLTFMeshUtility.cpp
+++ b/Engine/Plugins/Enterprise/GLTFExporter/MigratedFromGit/Source/GLTFExporter/Private/Converters/GLTFMeshUtility.cpp
@@ -3,6 +3,8 @@
 #include "Converters/GLTFMeshUtility.h"
 #include "Rendering/SkeletalMeshRenderData.h"
 #include "PlatformInfo.h"
+#include "StaticMeshAttributes.h"
+#include "Developer/MeshMergeUtilities/Private/MeshMergeHelpers.h"
 
 TArray<int32> FGLTFMeshUtility::GetSectionIndices(const FStaticMeshLODResources& MeshLOD, int32 MaterialIndex)
 {
@@ -103,6 +105,33 @@ int32 FGLTFMeshUtility::GetMinimumLOD(const USkeletalMesh* SkeletalMesh, const U
 	return -1;
 }
 
+void FGLTFMeshUtility::RetrieveMesh(USkeletalMesh* SkeletalMesh, int32 LODIndex, FMeshDescription& OutDescription)
+{
+	// NOTE: this is a workaround for the fact that there's no overload for FGLTFMeshUtility::RetrieveMesh
+	// that accepts a USkeletalMesh, only a USkeletalMeshComponent.
+	// Writing a custom function that would work on "standalone" skeletal meshes is problematic since
+	// we would need to implement an equivalent of USkinnedMeshComponent::GetCPUSkinnedVertices too.
+
+	if (UWorld* World = GEditor->GetEditorWorldContext().World())
+	{
+		FActorSpawnParameters SpawnParams;
+		SpawnParams.SpawnCollisionHandlingOverride = ESpawnActorCollisionHandlingMethod::AlwaysSpawn;
+		SpawnParams.ObjectFlags |= RF_Transient;
+		SpawnParams.bAllowDuringConstructionScript = true;
+
+		if (AActor* Actor = World->SpawnActor<AActor>(SpawnParams))
+		{
+			USkeletalMeshComponent* SkeletalMeshComponent = NewObject<USkeletalMeshComponent>(Actor, TEXT(""), RF_Transient);
+			SkeletalMeshComponent->RegisterComponent();
+			SkeletalMeshComponent->SetSkeletalMesh(SkeletalMesh);
+
+			FMeshMergeHelpers::RetrieveMesh(SkeletalMeshComponent, LODIndex, OutDescription, true);
+
+			World->DestroyActor(Actor, false, false);
+		}
+	}
+}
+
 template <typename ValueType, typename StructType>
 ValueType FGLTFMeshUtility::GetValueForRunningPlatform(const StructType& Properties)
 {
diff --git a/Engine/Plugins/Enterprise/GLTFExporter/MigratedFromGit/Source/GLTFExporter/Private/Converters/GLTFMeshUtility.h b/Engine/Plugins/Enterprise/GLTFExporter/MigratedFromGit/Source/GLTFExporter/Private/Converters/GLTFMeshUtility.h
--- a/Engine/Plugins/Enterprise/GLTFExporter/MigratedFromGit/Source/GLTFExporter/Private/Converters/GLTFMeshUtility.h
+++ b/Engine/Plugins/Enterprise/GLTFExporter/MigratedFromGit/Source/GLTFExporter/Private/Converters/GLTFMeshUtility.h
@@ -20,6 +20,8 @@ struct FGLTFMeshUtility
 	static int32 GetMinimumLOD(const UStaticMesh* StaticMesh, const UStaticMeshComponent* StaticMeshComponent);
 	static int32 GetMinimumLOD(const USkeletalMesh* SkeletalMesh, const USkeletalMeshComponent* SkeletalMeshComponent);
 
+	static void RetrieveMesh(USkeletalMesh* SkeletalMesh, int32 LODIndex, FMeshDescription& OutDescription);
+
 private:
 
 	template <typename ValueType, typename StructType>
diff --git a/Engine/Plugins/Enterprise/GLTFExporter/MigratedFromGit/Source/GLTFExporter/Private/Tasks/GLTFMeshTasks.cpp b/Engine/Plugins/Enterprise/GLTFExporter/MigratedFromGit/Source/GLTFExporter/Private/Tasks/GLTFMeshTasks.cpp
--- a/Engine/Plugins/Enterprise/GLTFExporter/MigratedFromGit/Source/GLTFExporter/Private/Tasks/GLTFMeshTasks.cpp
+++ b/Engine/Plugins/Enterprise/GLTFExporter/MigratedFromGit/Source/GLTFExporter/Private/Tasks/GLTFMeshTasks.cpp
@@ -26,6 +26,12 @@ void FGLTFStaticMeshTask::Complete()
 		Builder.ExportOptions->bBakeMaterialInputs && Builder.ExportOptions->bMaterialBakeUsingMeshData ?
 			MeshDataConverter.GetOrAdd(StaticMesh, StaticMeshComponent, LODIndex) : nullptr;
 
+	if (MeshData != nullptr && MeshData->Description.IsEmpty())
+	{
+		// TODO: add warning in case the mesh actually has data, which means we failed to extract a mesh description.
+		MeshData = nullptr;
+	}
+
 	const int32 MaterialCount = StaticMesh->StaticMaterials.Num();
 	JsonMesh.Primitives.AddDefaulted(MaterialCount);
 
@@ -88,9 +94,10 @@ void FGLTFSkeletalMeshTask::Complete()
 		Builder.ExportOptions->bBakeMaterialInputs && Builder.ExportOptions->bMaterialBakeUsingMeshData ?
 			MeshDataConverter.GetOrAdd(SkeletalMesh, SkeletalMeshComponent, LODIndex) : nullptr;
 
-	if (SkeletalMeshComponent == nullptr)
+	if (MeshData != nullptr && MeshData->Description.IsEmpty())
 	{
-		MeshData = nullptr; // TODO: remove and add support for skeletal meshes in FGLTFMeshData
+		// TODO: add warning in case the mesh actually has data, which means we failed to extract a mesh description.
+		MeshData = nullptr;
 	}
 
 	const uint16 MaterialCount = SkeletalMesh->Materials.Num();
