[GLTFExporter] refactor: fix minor typo in NOTE comment

commit 01846ce316efaaece71b1f941f090df231173a0d
Date: Thu, 12 Nov 2020 00:06:53 +0100

    refactor: fix minor typo in NOTE comment
    (cherry picked from commit c13a49b3462e7c2a10de7b08b2fa27e348ef4c8f​)

#skipci

diff --git a/Engine/Plugins/Enterprise/GLTFExporter/MigratedFromGit/Source/GLTFExporter/Private/Converters/GLTFMeshUtility.cpp b/Engine/Plugins/Enterprise/GLTFExporter/MigratedFromGit/Source/GLTFExporter/Private/Converters/GLTFMeshUtility.cpp
--- a/Engine/Plugins/Enterprise/GLTFExporter/MigratedFromGit/Source/GLTFExporter/Private/Converters/GLTFMeshUtility.cpp
+++ b/Engine/Plugins/Enterprise/GLTFExporter/MigratedFromGit/Source/GLTFExporter/Private/Converters/GLTFMeshUtility.cpp
@@ -107,7 +107,7 @@ int32 FGLTFMeshUtility::GetMinimumLOD(const USkeletalMesh* SkeletalMesh, const U
 
 void FGLTFMeshUtility::RetrieveMesh(USkeletalMesh* SkeletalMesh, int32 LODIndex, FMeshDescription& OutDescription)
 {
-	// NOTE: this is a workaround for the fact that there's no overload for FGLTFMeshUtility::RetrieveMesh
+	// NOTE: this is a workaround for the fact that there's no overload for FMeshMergeHelpers::RetrieveMesh
 	// that accepts a USkeletalMesh, only a USkeletalMeshComponent.
 	// Writing a custom function that would work on "standalone" skeletal meshes is problematic since
 	// we would need to implement an equivalent of USkinnedMeshComponent::GetCPUSkinnedVertices too.
