[GLTFExporter] refactor: move workaround for retrieving mesh data from SkeletalMesh to GLTFMeshData since its only needed there

commit d24a3d6a93c26a33f338a1b9cd858e329e683780
Date: Thu, 12 Nov 2020 00:09:07 +0100

    refactor: move workaround for retrieving mesh data from SkeletalMesh to GLTFMeshData since its only needed there
    (cherry picked from commit 233f57765e5e9566172f437c23b145e46c7ef033​)

#skipci

diff --git a/Engine/Plugins/Enterprise/GLTFExporter/MigratedFromGit/Source/GLTFExporter/Private/Converters/GLTFMeshData.cpp b/Engine/Plugins/Enterprise/GLTFExporter/MigratedFromGit/Source/GLTFExporter/Private/Converters/GLTFMeshData.cpp
--- a/Engine/Plugins/Enterprise/GLTFExporter/MigratedFromGit/Source/GLTFExporter/Private/Converters/GLTFMeshData.cpp
+++ b/Engine/Plugins/Enterprise/GLTFExporter/MigratedFromGit/Source/GLTFExporter/Private/Converters/GLTFMeshData.cpp
@@ -33,7 +33,30 @@ FGLTFMeshData::FGLTFMeshData(const USkeletalMesh* SkeletalMesh, const USkeletalM
 	}
 	else
 	{
-		FGLTFMeshUtility::RetrieveMesh(const_cast<USkeletalMesh*>(SkeletalMesh), LODIndex, Description);
+		// NOTE: this is a workaround for the fact that there's no overload for FMeshMergeHelpers::RetrieveMesh
+		// that accepts a USkeletalMesh, only a USkeletalMeshComponent.
+		// Writing a custom utility function that would work on a "standalone" skeletal mesh is problematic
+		// since we would need to implement an equivalent of USkinnedMeshComponent::GetCPUSkinnedVertices too.
+
+		if (UWorld* World = GEditor->GetEditorWorldContext().World())
+		{
+			FActorSpawnParameters SpawnParams;
+			SpawnParams.SpawnCollisionHandlingOverride = ESpawnActorCollisionHandlingMethod::AlwaysSpawn;
+			SpawnParams.ObjectFlags |= RF_Transient;
+			SpawnParams.bAllowDuringConstructionScript = true;
+
+			if (AActor* Actor = World->SpawnActor<AActor>(SpawnParams))
+			{
+				USkeletalMeshComponent* Component = NewObject<USkeletalMeshComponent>(Actor, TEXT(""), RF_Transient);
+				Component->RegisterComponent();
+				Component->SetSkeletalMesh(const_cast<USkeletalMesh*>(SkeletalMesh));
+
+				FMeshMergeHelpers::RetrieveMesh(Component, LODIndex, Description, true);
+
+				World->DestroyActor(Actor, false, false);
+			}
+		}
+
 		SkeletalMesh->GetName(Name);
 	}
 }
diff --git a/Engine/Plugins/Enterprise/GLTFExporter/MigratedFromGit/Source/GLTFExporter/Private/Converters/GLTFMeshUtility.cpp b/Engine/Plugins/Enterprise/GLTFExporter/MigratedFromGit/Source/GLTFExporter/Private/Converters/GLTFMeshUtility.cpp
--- a/Engine/Plugins/Enterprise/GLTFExporter/MigratedFromGit/Source/GLTFExporter/Private/Converters/GLTFMeshUtility.cpp
+++ b/Engine/Plugins/Enterprise/GLTFExporter/MigratedFromGit/Source/GLTFExporter/Private/Converters/GLTFMeshUtility.cpp
@@ -3,8 +3,6 @@
 #include "Converters/GLTFMeshUtility.h"
 #include "Rendering/SkeletalMeshRenderData.h"
 #include "PlatformInfo.h"
-#include "StaticMeshAttributes.h"
-#include "Developer/MeshMergeUtilities/Private/MeshMergeHelpers.h"
 
 TArray<int32> FGLTFMeshUtility::GetSectionIndices(const FStaticMeshLODResources& MeshLOD, int32 MaterialIndex)
 {
@@ -105,33 +103,6 @@ int32 FGLTFMeshUtility::GetMinimumLOD(const USkeletalMesh* SkeletalMesh, const U
 	return -1;
 }
 
-void FGLTFMeshUtility::RetrieveMesh(USkeletalMesh* SkeletalMesh, int32 LODIndex, FMeshDescription& OutDescription)
-{
-	// NOTE: this is a workaround for the fact that there's no overload for FMeshMergeHelpers::RetrieveMesh
-	// that accepts a USkeletalMesh, only a USkeletalMeshComponent.
-	// Writing a custom function that would work on "standalone" skeletal meshes is problematic since
-	// we would need to implement an equivalent of USkinnedMeshComponent::GetCPUSkinnedVertices too.
-
-	if (UWorld* World = GEditor->GetEditorWorldContext().World())
-	{
-		FActorSpawnParameters SpawnParams;
-		SpawnParams.SpawnCollisionHandlingOverride = ESpawnActorCollisionHandlingMethod::AlwaysSpawn;
-		SpawnParams.ObjectFlags |= RF_Transient;
-		SpawnParams.bAllowDuringConstructionScript = true;
-
-		if (AActor* Actor = World->SpawnActor<AActor>(SpawnParams))
-		{
-			USkeletalMeshComponent* SkeletalMeshComponent = NewObject<USkeletalMeshComponent>(Actor, TEXT(""), RF_Transient);
-			SkeletalMeshComponent->RegisterComponent();
-			SkeletalMeshComponent->SetSkeletalMesh(SkeletalMesh);
-
-			FMeshMergeHelpers::RetrieveMesh(SkeletalMeshComponent, LODIndex, OutDescription, true);
-
-			World->DestroyActor(Actor, false, false);
-		}
-	}
-}
-
 template <typename ValueType, typename StructType>
 ValueType FGLTFMeshUtility::GetValueForRunningPlatform(const StructType& Properties)
 {
diff --git a/Engine/Plugins/Enterprise/GLTFExporter/MigratedFromGit/Source/GLTFExporter/Private/Converters/GLTFMeshUtility.h b/Engine/Plugins/Enterprise/GLTFExporter/MigratedFromGit/Source/GLTFExporter/Private/Converters/GLTFMeshUtility.h
--- a/Engine/Plugins/Enterprise/GLTFExporter/MigratedFromGit/Source/GLTFExporter/Private/Converters/GLTFMeshUtility.h
+++ b/Engine/Plugins/Enterprise/GLTFExporter/MigratedFromGit/Source/GLTFExporter/Private/Converters/GLTFMeshUtility.h
@@ -20,8 +20,6 @@ struct FGLTFMeshUtility
 	static int32 GetMinimumLOD(const UStaticMesh* StaticMesh, const UStaticMeshComponent* StaticMeshComponent);
 	static int32 GetMinimumLOD(const USkeletalMesh* SkeletalMesh, const USkeletalMeshComponent* SkeletalMeshComponent);
 
-	static void RetrieveMesh(USkeletalMesh* SkeletalMesh, int32 LODIndex, FMeshDescription& OutDescription);
-
 private:
 
 	template <typename ValueType, typename StructType>
