[GLTFExporter] feat: add new export option bVariantMaterialBakeUsingMeshData

commit 57d4a9ae43eafe52a22039d698ec25df5684844d
Date: Thu, 12 Nov 2020 09:05:33 +0100

    feat: add new export option bVariantMaterialBakeUsingMeshData
    (cherry picked from commit 7a17f6e1961cc64c0e808445ff6e1d350cc5be4e​)

#skipci

diff --git a/Engine/Plugins/Enterprise/GLTFExporter/MigratedFromGit/Source/GLTFExporter/Private/GLTFExportOptions.cpp b/Engine/Plugins/Enterprise/GLTFExporter/MigratedFromGit/Source/GLTFExporter/Private/GLTFExportOptions.cpp
--- a/Engine/Plugins/Enterprise/GLTFExporter/MigratedFromGit/Source/GLTFExporter/Private/GLTFExportOptions.cpp
+++ b/Engine/Plugins/Enterprise/GLTFExporter/MigratedFromGit/Source/GLTFExporter/Private/GLTFExportOptions.cpp
@@ -41,6 +41,7 @@ void UGLTFExportOptions::ResetToDefault()
 	bExportHDRIBackdrops = true;
 	bExportSkySpheres = true;
 	bExportVariantSets = true;
+	bVariantMaterialBakeUsingMeshData = false;
 	bExportAnimationHotspots = true;
 	bBundleWebViewer = true;
 	bExportPreviewMesh = true;
diff --git a/Engine/Plugins/Enterprise/GLTFExporter/MigratedFromGit/Source/GLTFExporter/Public/GLTFExportOptions.h b/Engine/Plugins/Enterprise/GLTFExporter/MigratedFromGit/Source/GLTFExporter/Public/GLTFExportOptions.h
--- a/Engine/Plugins/Enterprise/GLTFExporter/MigratedFromGit/Source/GLTFExporter/Public/GLTFExportOptions.h
+++ b/Engine/Plugins/Enterprise/GLTFExporter/MigratedFromGit/Source/GLTFExporter/Public/GLTFExportOptions.h
@@ -150,6 +150,9 @@ public:
 	UPROPERTY(EditAnywhere, BlueprintReadWrite, Config, Category = Scene)
 	bool bExportVariantSets;
 
+	UPROPERTY(EditAnywhere, BlueprintReadWrite, Config, Category = Material, Meta = (EditCondition = "bExportVariantSets && bBakeMaterialInputs && bMaterialBakeUsingMeshData"))
+	bool bVariantMaterialBakeUsingMeshData;
+
 	UPROPERTY(EditAnywhere, BlueprintReadWrite, Config, Category = Scene)
 	bool bExportAnimationHotspots;
 
