[GLTFExporter] feat: move mesh-data converters into FGLTFConvertBuilder to allow access from other converters

commit 2443d15cab6bd0bde741c994d51fa896b4b0d997
Date: Thu, 12 Nov 2020 09:09:19 +0100

    feat: move mesh-data converters into FGLTFConvertBuilder to allow access from other converters
    (cherry picked from commit f639e6b83f38f8e3212d68348ce4c73402f863f4​)

#skipci

diff --git a/Engine/Plugins/Enterprise/GLTFExporter/MigratedFromGit/Source/GLTFExporter/Private/Builders/GLTFConvertBuilder.h b/Engine/Plugins/Enterprise/GLTFExporter/MigratedFromGit/Source/GLTFExporter/Private/Builders/GLTFConvertBuilder.h
--- a/Engine/Plugins/Enterprise/GLTFExporter/MigratedFromGit/Source/GLTFExporter/Private/Builders/GLTFConvertBuilder.h
+++ b/Engine/Plugins/Enterprise/GLTFExporter/MigratedFromGit/Source/GLTFExporter/Private/Builders/GLTFConvertBuilder.h
@@ -5,6 +5,7 @@
 #include "Builders/GLTFImageBuilder.h"
 #include "Converters/GLTFAccessorConverters.h"
 #include "Converters/GLTFMeshConverters.h"
+#include "Converters/GLTFMeshDataConverters.h"
 #include "Converters/GLTFMaterialConverters.h"
 #include "Converters/GLTFSamplerConverters.h"
 #include "Converters/GLTFTextureConverters.h"
@@ -30,6 +31,9 @@ public:
 
 	const bool bSelectedActorsOnly;
 
+	FGLTFStaticMeshDataConverter StaticMeshDataConverter;
+	FGLTFSkeletalMeshDataConverter SkeletalMeshDataConverter;
+
 	FGLTFJsonAccessorIndex GetOrAddPositionAccessor(const FGLTFMeshSection* MeshSection, const FPositionVertexBuffer* VertexBuffer);
 	FGLTFJsonAccessorIndex GetOrAddColorAccessor(const FGLTFMeshSection* MeshSection, const FColorVertexBuffer* VertexBuffer);
 	FGLTFJsonAccessorIndex GetOrAddNormalAccessor(const FGLTFMeshSection* MeshSection, const FStaticMeshVertexBuffer* VertexBuffer);
diff --git a/Engine/Plugins/Enterprise/GLTFExporter/MigratedFromGit/Source/GLTFExporter/Private/Converters/GLTFMeshConverters.cpp b/Engine/Plugins/Enterprise/GLTFExporter/MigratedFromGit/Source/GLTFExporter/Private/Converters/GLTFMeshConverters.cpp
--- a/Engine/Plugins/Enterprise/GLTFExporter/MigratedFromGit/Source/GLTFExporter/Private/Converters/GLTFMeshConverters.cpp
+++ b/Engine/Plugins/Enterprise/GLTFExporter/MigratedFromGit/Source/GLTFExporter/Private/Converters/GLTFMeshConverters.cpp
@@ -42,7 +42,7 @@ void FGLTFStaticMeshConverter::Sanitize(const UStaticMesh*& StaticMesh, const US
 FGLTFJsonMeshIndex FGLTFStaticMeshConverter::Convert(const UStaticMesh* StaticMesh,  const UStaticMeshComponent* StaticMeshComponent, FGLTFMaterialArray Materials, int32 LODIndex)
 {
 	const FGLTFJsonMeshIndex MeshIndex = Builder.AddMesh();
-	Builder.SetupTask<FGLTFStaticMeshTask>(Builder, MeshSectionConverter, MeshDataConverter, StaticMesh, StaticMeshComponent, Materials, LODIndex, MeshIndex);
+	Builder.SetupTask<FGLTFStaticMeshTask>(Builder, MeshSectionConverter, StaticMesh, StaticMeshComponent, Materials, LODIndex, MeshIndex);
 	return MeshIndex;
 }
 
@@ -82,6 +82,6 @@ void FGLTFSkeletalMeshConverter::Sanitize(const USkeletalMesh*& SkeletalMesh, co
 FGLTFJsonMeshIndex FGLTFSkeletalMeshConverter::Convert(const USkeletalMesh* SkeletalMesh, const USkeletalMeshComponent* SkeletalMeshComponent, FGLTFMaterialArray Materials, int32 LODIndex)
 {
 	const FGLTFJsonMeshIndex MeshIndex = Builder.AddMesh();
-	Builder.SetupTask<FGLTFSkeletalMeshTask>(Builder, MeshSectionConverter, MeshDataConverter, SkeletalMesh, SkeletalMeshComponent, Materials, LODIndex, MeshIndex);
+	Builder.SetupTask<FGLTFSkeletalMeshTask>(Builder, MeshSectionConverter, SkeletalMesh, SkeletalMeshComponent, Materials, LODIndex, MeshIndex);
 	return MeshIndex;
 }
diff --git a/Engine/Plugins/Enterprise/GLTFExporter/MigratedFromGit/Source/GLTFExporter/Private/Converters/GLTFMeshConverters.h b/Engine/Plugins/Enterprise/GLTFExporter/MigratedFromGit/Source/GLTFExporter/Private/Converters/GLTFMeshConverters.h
--- a/Engine/Plugins/Enterprise/GLTFExporter/MigratedFromGit/Source/GLTFExporter/Private/Converters/GLTFMeshConverters.h
+++ b/Engine/Plugins/Enterprise/GLTFExporter/MigratedFromGit/Source/GLTFExporter/Private/Converters/GLTFMeshConverters.h
@@ -25,7 +25,6 @@ class FGLTFStaticMeshConverter final : public TGLTFMeshConverter<const UStaticMe
 	virtual FGLTFJsonMeshIndex Convert(const UStaticMesh* StaticMesh, const UStaticMeshComponent* StaticMeshComponent, FGLTFMaterialArray Materials, int32 LODIndex) override;
 
 	FGLTFStaticMeshSectionConverter MeshSectionConverter;
-	FGLTFStaticMeshDataConverter MeshDataConverter;
 };
 
 class FGLTFSkeletalMeshConverter final : public TGLTFMeshConverter<const USkeletalMesh*, const USkeletalMeshComponent*, FGLTFMaterialArray, int32>
@@ -37,5 +36,4 @@ class FGLTFSkeletalMeshConverter final : public TGLTFMeshConverter<const USkelet
 	virtual FGLTFJsonMeshIndex Convert(const USkeletalMesh* SkeletalMesh, const USkeletalMeshComponent* SkeletalMeshComponent, FGLTFMaterialArray Materials, int32 LODIndex) override;
 
 	FGLTFSkeletalMeshSectionConverter MeshSectionConverter;
-	FGLTFSkeletalMeshDataConverter MeshDataConverter;
 };
diff --git a/Engine/Plugins/Enterprise/GLTFExporter/MigratedFromGit/Source/GLTFExporter/Private/Tasks/GLTFMeshTasks.cpp b/Engine/Plugins/Enterprise/GLTFExporter/MigratedFromGit/Source/GLTFExporter/Private/Tasks/GLTFMeshTasks.cpp
--- a/Engine/Plugins/Enterprise/GLTFExporter/MigratedFromGit/Source/GLTFExporter/Private/Tasks/GLTFMeshTasks.cpp
+++ b/Engine/Plugins/Enterprise/GLTFExporter/MigratedFromGit/Source/GLTFExporter/Private/Tasks/GLTFMeshTasks.cpp
@@ -24,7 +24,7 @@ void FGLTFStaticMeshTask::Complete()
 
 	const FGLTFMeshData* MeshData =
 		Builder.ExportOptions->bBakeMaterialInputs && Builder.ExportOptions->bMaterialBakeUsingMeshData ?
-			MeshDataConverter.GetOrAdd(StaticMesh, StaticMeshComponent, LODIndex) : nullptr;
+			Builder.StaticMeshDataConverter.GetOrAdd(StaticMesh, StaticMeshComponent, LODIndex) : nullptr;
 
 	if (MeshData != nullptr && MeshData->Description.IsEmpty())
 	{
@@ -92,7 +92,7 @@ void FGLTFSkeletalMeshTask::Complete()
 
 	const FGLTFMeshData* MeshData =
 		Builder.ExportOptions->bBakeMaterialInputs && Builder.ExportOptions->bMaterialBakeUsingMeshData ?
-			MeshDataConverter.GetOrAdd(SkeletalMesh, SkeletalMeshComponent, LODIndex) : nullptr;
+			Builder.SkeletalMeshDataConverter.GetOrAdd(SkeletalMesh, SkeletalMeshComponent, LODIndex) : nullptr;
 
 	if (MeshData != nullptr && MeshData->Description.IsEmpty())
 	{
diff --git a/Engine/Plugins/Enterprise/GLTFExporter/MigratedFromGit/Source/GLTFExporter/Private/Tasks/GLTFMeshTasks.h b/Engine/Plugins/Enterprise/GLTFExporter/MigratedFromGit/Source/GLTFExporter/Private/Tasks/GLTFMeshTasks.h
--- a/Engine/Plugins/Enterprise/GLTFExporter/MigratedFromGit/Source/GLTFExporter/Private/Tasks/GLTFMeshTasks.h
+++ b/Engine/Plugins/Enterprise/GLTFExporter/MigratedFromGit/Source/GLTFExporter/Private/Tasks/GLTFMeshTasks.h
@@ -5,7 +5,6 @@
 #include "Tasks/GLTFTask.h"
 #include "Builders/GLTFConvertBuilder.h"
 #include "Converters/GLTFMeshSectionConverters.h"
-#include "Converters/GLTFMeshDataConverters.h"
 #include "Converters/GLTFMaterialArray.h"
 #include "Converters/GLTFNameUtility.h"
 #include "Engine.h"
@@ -14,11 +13,10 @@ class FGLTFStaticMeshTask : public FGLTFTask
 {
 public:
 
-	FGLTFStaticMeshTask(FGLTFConvertBuilder& Builder, FGLTFStaticMeshSectionConverter& MeshSectionConverter, FGLTFStaticMeshDataConverter& MeshDataConverter, const UStaticMesh* StaticMesh, const UStaticMeshComponent* StaticMeshComponent, FGLTFMaterialArray Materials, int32 LODIndex, FGLTFJsonMeshIndex MeshIndex)
+	FGLTFStaticMeshTask(FGLTFConvertBuilder& Builder, FGLTFStaticMeshSectionConverter& MeshSectionConverter, const UStaticMesh* StaticMesh, const UStaticMeshComponent* StaticMeshComponent, FGLTFMaterialArray Materials, int32 LODIndex, FGLTFJsonMeshIndex MeshIndex)
 		: FGLTFTask(EGLTFTaskPriority::Mesh)
 		, Builder(Builder)
 		, MeshSectionConverter(MeshSectionConverter)
-		, MeshDataConverter(MeshDataConverter)
 		, StaticMesh(StaticMesh)
 		, StaticMeshComponent(StaticMeshComponent)
 		, Materials(Materials)
@@ -38,7 +36,6 @@ private:
 
 	FGLTFConvertBuilder& Builder;
 	FGLTFStaticMeshSectionConverter& MeshSectionConverter;
-	FGLTFStaticMeshDataConverter& MeshDataConverter;
 	const UStaticMesh* StaticMesh;
 	const UStaticMeshComponent* StaticMeshComponent;
 	const FGLTFMaterialArray Materials;
@@ -50,11 +47,10 @@ class FGLTFSkeletalMeshTask : public FGLTFTask
 {
 public:
 
-	FGLTFSkeletalMeshTask(FGLTFConvertBuilder& Builder, FGLTFSkeletalMeshSectionConverter& MeshSectionConverter, FGLTFSkeletalMeshDataConverter& MeshDataConverter, const USkeletalMesh* SkeletalMesh, const USkeletalMeshComponent* SkeletalMeshComponent, FGLTFMaterialArray Materials, int32 LODIndex, FGLTFJsonMeshIndex MeshIndex)
+	FGLTFSkeletalMeshTask(FGLTFConvertBuilder& Builder, FGLTFSkeletalMeshSectionConverter& MeshSectionConverter, const USkeletalMesh* SkeletalMesh, const USkeletalMeshComponent* SkeletalMeshComponent, FGLTFMaterialArray Materials, int32 LODIndex, FGLTFJsonMeshIndex MeshIndex)
 		: FGLTFTask(EGLTFTaskPriority::Mesh)
 		, Builder(Builder)
 		, MeshSectionConverter(MeshSectionConverter)
-		, MeshDataConverter(MeshDataConverter)
 		, SkeletalMesh(SkeletalMesh)
 		, SkeletalMeshComponent(SkeletalMeshComponent)
 		, Materials(Materials)
@@ -74,7 +70,6 @@ private:
 
 	FGLTFConvertBuilder& Builder;
 	FGLTFSkeletalMeshSectionConverter& MeshSectionConverter;
-	FGLTFSkeletalMeshDataConverter& MeshDataConverter;
 	const USkeletalMesh* SkeletalMesh;
 	const USkeletalMeshComponent* SkeletalMeshComponent;
 	const FGLTFMaterialArray Materials;
