[GLTFExporter] feat: bake materials for variants using mesh-data (when allowed by export options)

commit 0ae2a85d94a19b8000e35f77df319ea53fa6491e
Date: Thu, 12 Nov 2020 09:10:11 +0100

    feat: bake materials for variants using mesh-data (when allowed by export options)
    (cherry picked from commit 68fa7f10e913c7e3edef3edc87029301bb3d0c81​)

#skipci

diff --git a/Engine/Plugins/Enterprise/GLTFExporter/MigratedFromGit/Source/GLTFExporter/Private/Converters/GLTFVarationConverters.cpp b/Engine/Plugins/Enterprise/GLTFExporter/MigratedFromGit/Source/GLTFExporter/Private/Converters/GLTFVarationConverters.cpp
--- a/Engine/Plugins/Enterprise/GLTFExporter/MigratedFromGit/Source/GLTFExporter/Private/Converters/GLTFVarationConverters.cpp
+++ b/Engine/Plugins/Enterprise/GLTFExporter/MigratedFromGit/Source/GLTFExporter/Private/Converters/GLTFVarationConverters.cpp
@@ -1,7 +1,9 @@
 // Copyright Epic Games, Inc. All Rights Reserved.
 
 #include "Converters/GLTFVarationConverters.h"
+#include "Converters/GLTFMeshUtility.h"
 #include "Builders/GLTFContainerBuilder.h"
+#include "Rendering/SkeletalMeshRenderData.h"
 #include "LevelVariantSetsActor.h"
 #include "VariantObjectBinding.h"
 #include "PropertyValueMaterial.h"
@@ -268,12 +270,41 @@ bool FGLTFVariationConverter::TryParseMaterialPropertyValue(FGLTFJsonVariant& Ou
 
 	// TODO: find way to determine whether the material is null because "None" was selected, or because it failed to resolve
 	const UMaterialInterface* Material = MaterialProperty->GetMaterial();
-	const FGLTFJsonMaterialIndex MaterialIndex = Builder.GetOrAddMaterial(Material);
-	const int32 ElementIndex = CapturedPropSegments[NumPropSegments - 1].PropertyIndex;
+	const int32 MaterialIndex = CapturedPropSegments[NumPropSegments - 1].PropertyIndex;
+
+	const FGLTFMeshData* MeshData = nullptr;
+	TArray<int32> SectionIndices;
+
+	if (Builder.ExportOptions->bVariantMaterialBakeUsingMeshData &&
+		Builder.ExportOptions->bMaterialBakeUsingMeshData &&
+		Builder.ExportOptions->bBakeMaterialInputs)
+	{
+		const int32 DefaultLODIndex = Builder.ExportOptions->DefaultLevelOfDetail;
+
+		if (const UStaticMeshComponent* StaticMeshComponent = Cast<UStaticMeshComponent>(Target))
+		{
+			const UStaticMesh* StaticMesh = StaticMeshComponent->GetStaticMesh();
+			const int32 LODIndex = FGLTFMeshUtility::GetLOD(StaticMesh, StaticMeshComponent, DefaultLODIndex);
+			const FStaticMeshLODResources& MeshLOD = StaticMesh->GetLODForExport(LODIndex);
+
+			SectionIndices = FGLTFMeshUtility::GetSectionIndices(MeshLOD, MaterialIndex);
+			MeshData = Builder.StaticMeshDataConverter.GetOrAdd(StaticMesh, StaticMeshComponent, LODIndex);
+		}
+		else if (const USkeletalMeshComponent* SkeletalMeshComponent = Cast<USkeletalMeshComponent>(Target))
+		{
+			const USkeletalMesh* SkeletalMesh = SkeletalMeshComponent->SkeletalMesh;
+			const int32 LODIndex = FGLTFMeshUtility::GetLOD(SkeletalMesh, SkeletalMeshComponent, DefaultLODIndex);
+			const FSkeletalMeshRenderData* RenderData = SkeletalMesh->GetResourceForRendering();
+			const FSkeletalMeshLODRenderData& MeshLOD = RenderData->LODRenderData[LODIndex];
+
+			SectionIndices = FGLTFMeshUtility::GetSectionIndices(MeshLOD, MaterialIndex);
+			MeshData = Builder.SkeletalMeshDataConverter.GetOrAdd(SkeletalMesh, SkeletalMeshComponent, LODIndex);
+		}
+	}
 
 	FGLTFJsonVariantMaterial VariantMaterial;
-	VariantMaterial.Material = MaterialIndex;
-	VariantMaterial.Index = ElementIndex;
+	VariantMaterial.Material = Builder.GetOrAddMaterial(Material, MeshData, SectionIndices);
+	VariantMaterial.Index = MaterialIndex;
 
 	const FGLTFJsonNodeIndex NodeIndex = Builder.GetOrAddNode(Target);
 	const FGLTFJsonNodeIndex ComponentNodeIndex = Builder.GetComponentNodeIndex(NodeIndex);
