[GLTFExporter] refactor: add TODO comment regarding placement of mesh data converters

commit 810b4a11e8c85010b6f3c74e4247f5d99797de2d
Date: Thu, 12 Nov 2020 10:52:04 +0100

    refactor: add TODO comment regarding placement of mesh data converters
    (cherry picked from commit c51f9e7564d8a9be03682891ccb0d389578d8f4c​)

#skipci

diff --git a/Engine/Plugins/Enterprise/GLTFExporter/MigratedFromGit/Source/GLTFExporter/Private/Builders/GLTFConvertBuilder.h b/Engine/Plugins/Enterprise/GLTFExporter/MigratedFromGit/Source/GLTFExporter/Private/Builders/GLTFConvertBuilder.h
--- a/Engine/Plugins/Enterprise/GLTFExporter/MigratedFromGit/Source/GLTFExporter/Private/Builders/GLTFConvertBuilder.h
+++ b/Engine/Plugins/Enterprise/GLTFExporter/MigratedFromGit/Source/GLTFExporter/Private/Builders/GLTFConvertBuilder.h
@@ -31,6 +31,7 @@ public:
 
 	const bool bSelectedActorsOnly;
 
+	// TODO: find a better place for this types of indirect converters
 	FGLTFStaticMeshDataConverter StaticMeshDataConverter;
 	FGLTFSkeletalMeshDataConverter SkeletalMeshDataConverter;
 
