[GLTFExporter] refactor: order extensions object alphabetically in gltf output

commit 7ba1a1bd9b31b2a5c1e623bfb989c115a78d82f2
Date: Thu, 12 Nov 2020 19:39:53 +0100

    refactor: order extensions object alphabetically in gltf output
    (cherry picked from commit 8c02d36040dda88f76d820c4bdd8d110da704877​)

#skipci

diff --git a/Engine/Plugins/Enterprise/GLTFExporter/MigratedFromGit/Source/GLTFExporter/Private/Builders/GLTFJsonBuilder.cpp b/Engine/Plugins/Enterprise/GLTFExporter/MigratedFromGit/Source/GLTFExporter/Private/Builders/GLTFJsonBuilder.cpp
--- a/Engine/Plugins/Enterprise/GLTFExporter/MigratedFromGit/Source/GLTFExporter/Private/Builders/GLTFJsonBuilder.cpp
+++ b/Engine/Plugins/Enterprise/GLTFExporter/MigratedFromGit/Source/GLTFExporter/Private/Builders/GLTFJsonBuilder.cpp
@@ -92,14 +92,9 @@ FGLTFJsonBackdropIndex FGLTFJsonBuilder::AddBackdrop(const FGLTFJsonBackdrop& Js
 	return FGLTFJsonBackdropIndex(JsonRoot.Backdrops.Add(MakeUnique<FGLTFJsonBackdrop>(JsonBackdrop)));
 }
 
-FGLTFJsonVariationIndex FGLTFJsonBuilder::AddVariation(const FGLTFJsonVariation& JsonVariation)
-{
-	return FGLTFJsonVariationIndex(JsonRoot.Variations.Add(MakeUnique<FGLTFJsonVariation>(JsonVariation)));
-}
-
-FGLTFJsonLightMapIndex FGLTFJsonBuilder::AddLightMap(const FGLTFJsonLightMap& JsonLightMap)
+FGLTFJsonHotspotIndex FGLTFJsonBuilder::AddHotspot(const FGLTFJsonHotspot& JsonHotspot)
 {
-	return FGLTFJsonLightMapIndex(JsonRoot.LightMaps.Add(MakeUnique<FGLTFJsonLightMap>(JsonLightMap)));
+	return FGLTFJsonHotspotIndex(JsonRoot.Hotspots.Add(MakeUnique<FGLTFJsonHotspot>(JsonHotspot)));
 }
 
 FGLTFJsonLightIndex FGLTFJsonBuilder::AddLight(const FGLTFJsonLight& JsonLight)
@@ -107,9 +102,9 @@ FGLTFJsonLightIndex FGLTFJsonBuilder::AddLight(const FGLTFJsonLight& JsonLight)
 	return FGLTFJsonLightIndex(JsonRoot.Lights.Add(MakeUnique<FGLTFJsonLight>(JsonLight)));
 }
 
-FGLTFJsonHotspotIndex FGLTFJsonBuilder::AddHotspot(const FGLTFJsonHotspot& JsonHotspot)
+FGLTFJsonLightMapIndex FGLTFJsonBuilder::AddLightMap(const FGLTFJsonLightMap& JsonLightMap)
 {
-	return FGLTFJsonHotspotIndex(JsonRoot.Hotspots.Add(MakeUnique<FGLTFJsonHotspot>(JsonHotspot)));
+	return FGLTFJsonLightMapIndex(JsonRoot.LightMaps.Add(MakeUnique<FGLTFJsonLightMap>(JsonLightMap)));
 }
 
 FGLTFJsonSkySphereIndex FGLTFJsonBuilder::AddSkySphere(const FGLTFJsonSkySphere& JsonSkySphere)
@@ -117,6 +112,11 @@ FGLTFJsonSkySphereIndex FGLTFJsonBuilder::AddSkySphere(const FGLTFJsonSkySphere&
 	return FGLTFJsonSkySphereIndex(JsonRoot.SkySpheres.Add(MakeUnique<FGLTFJsonSkySphere>(JsonSkySphere)));
 }
 
+FGLTFJsonVariationIndex FGLTFJsonBuilder::AddVariation(const FGLTFJsonVariation& JsonVariation)
+{
+	return FGLTFJsonVariationIndex(JsonRoot.Variations.Add(MakeUnique<FGLTFJsonVariation>(JsonVariation)));
+}
+
 FGLTFJsonNodeIndex FGLTFJsonBuilder::AddChildNode(FGLTFJsonNodeIndex ParentIndex, const FGLTFJsonNode& JsonNode)
 {
 	const FGLTFJsonNodeIndex ChildIndex = AddNode(JsonNode);
@@ -211,14 +211,9 @@ FGLTFJsonBackdrop& FGLTFJsonBuilder::GetBackdrop(FGLTFJsonBackdropIndex Backdrop
 	return *JsonRoot.Backdrops[BackdropIndex];
 }
 
-FGLTFJsonVariation& FGLTFJsonBuilder::GetVariation(FGLTFJsonVariationIndex VariationIndex)
-{
-	return *JsonRoot.Variations[VariationIndex];
-}
-
-FGLTFJsonLightMap& FGLTFJsonBuilder::GetLightMap(FGLTFJsonLightMapIndex LightMapIndex)
+FGLTFJsonHotspot& FGLTFJsonBuilder::GetHotspot(FGLTFJsonHotspotIndex HotspotIndex)
 {
-	return *JsonRoot.LightMaps[LightMapIndex];
+	return *JsonRoot.Hotspots[HotspotIndex];
 }
 
 FGLTFJsonLight& FGLTFJsonBuilder::GetLight(FGLTFJsonLightIndex LightIndex)
@@ -226,9 +221,9 @@ FGLTFJsonLight& FGLTFJsonBuilder::GetLight(FGLTFJsonLightIndex LightIndex)
 	return *JsonRoot.Lights[LightIndex];
 }
 
-FGLTFJsonHotspot& FGLTFJsonBuilder::GetHotspot(FGLTFJsonHotspotIndex HotspotIndex)
+FGLTFJsonLightMap& FGLTFJsonBuilder::GetLightMap(FGLTFJsonLightMapIndex LightMapIndex)
 {
-	return *JsonRoot.Hotspots[HotspotIndex];
+	return *JsonRoot.LightMaps[LightMapIndex];
 }
 
 FGLTFJsonSkySphere& FGLTFJsonBuilder::GetSkySphere(FGLTFJsonSkySphereIndex SkySphereIndex)
@@ -236,6 +231,11 @@ FGLTFJsonSkySphere& FGLTFJsonBuilder::GetSkySphere(FGLTFJsonSkySphereIndex SkySp
 	return *JsonRoot.SkySpheres[SkySphereIndex];
 }
 
+FGLTFJsonVariation& FGLTFJsonBuilder::GetVariation(FGLTFJsonVariationIndex VariationIndex)
+{
+	return *JsonRoot.Variations[VariationIndex];
+}
+
 FGLTFJsonNodeIndex FGLTFJsonBuilder::GetComponentNodeIndex(FGLTFJsonNodeIndex NodeIndex)
 {
 	if (NodeIndex == INDEX_NONE)
diff --git a/Engine/Plugins/Enterprise/GLTFExporter/MigratedFromGit/Source/GLTFExporter/Private/Builders/GLTFJsonBuilder.h b/Engine/Plugins/Enterprise/GLTFExporter/MigratedFromGit/Source/GLTFExporter/Private/Builders/GLTFJsonBuilder.h
--- a/Engine/Plugins/Enterprise/GLTFExporter/MigratedFromGit/Source/GLTFExporter/Private/Builders/GLTFJsonBuilder.h
+++ b/Engine/Plugins/Enterprise/GLTFExporter/MigratedFromGit/Source/GLTFExporter/Private/Builders/GLTFJsonBuilder.h
@@ -33,11 +33,11 @@ public:
 	FGLTFJsonSkinIndex AddSkin(const FGLTFJsonSkin& JsonSkin = {});
 	FGLTFJsonTextureIndex AddTexture(const FGLTFJsonTexture& JsonTexture = {});
 	FGLTFJsonBackdropIndex AddBackdrop(const FGLTFJsonBackdrop& JsonBackdrop = {});
-	FGLTFJsonVariationIndex AddVariation(const FGLTFJsonVariation& JsonVariation = {});
-	FGLTFJsonLightMapIndex AddLightMap(const FGLTFJsonLightMap& JsonLightMap = {});
-	FGLTFJsonLightIndex AddLight(const FGLTFJsonLight& JsonLight = {});
 	FGLTFJsonHotspotIndex AddHotspot(const FGLTFJsonHotspot& JsonHotspot = {});
+	FGLTFJsonLightIndex AddLight(const FGLTFJsonLight& JsonLight = {});
+	FGLTFJsonLightMapIndex AddLightMap(const FGLTFJsonLightMap& JsonLightMap = {});
 	FGLTFJsonSkySphereIndex AddSkySphere(const FGLTFJsonSkySphere& JsonSkySphere = {});
+	FGLTFJsonVariationIndex AddVariation(const FGLTFJsonVariation& JsonVariation = {});
 
 	FGLTFJsonNodeIndex AddChildNode(FGLTFJsonNodeIndex ParentNodeIndex, const FGLTFJsonNode& JsonNode = {});
 	FGLTFJsonNodeIndex AddChildComponentNode(FGLTFJsonNodeIndex ParentNodeIndex, const FGLTFJsonNode& JsonNode = {});
@@ -56,11 +56,11 @@ public:
 	FGLTFJsonSkin& GetSkin(FGLTFJsonSkinIndex SkinIndex);
 	FGLTFJsonTexture& GetTexture(FGLTFJsonTextureIndex TextureIndex);
 	FGLTFJsonBackdrop& GetBackdrop(FGLTFJsonBackdropIndex BackdropIndex);
-	FGLTFJsonVariation& GetVariation(FGLTFJsonVariationIndex VariationIndex);
-	FGLTFJsonLightMap& GetLightMap(FGLTFJsonLightMapIndex LightMapIndex);
-	FGLTFJsonLight& GetLight(FGLTFJsonLightIndex LightIndex);
 	FGLTFJsonHotspot& GetHotspot(FGLTFJsonHotspotIndex HotspotIndex);
+	FGLTFJsonLight& GetLight(FGLTFJsonLightIndex LightIndex);
+	FGLTFJsonLightMap& GetLightMap(FGLTFJsonLightMapIndex LightMapIndex);
 	FGLTFJsonSkySphere& GetSkySphere(FGLTFJsonSkySphereIndex SkySphereIndex);
+	FGLTFJsonVariation& GetVariation(FGLTFJsonVariationIndex VariationIndex);
 
 	FGLTFJsonNodeIndex GetComponentNodeIndex(FGLTFJsonNodeIndex NodeIndex);
 
diff --git a/Engine/Plugins/Enterprise/GLTFExporter/MigratedFromGit/Source/GLTFExporter/Private/Json/GLTFJsonIndex.h b/Engine/Plugins/Enterprise/GLTFExporter/MigratedFromGit/Source/GLTFExporter/Private/Json/GLTFJsonIndex.h
--- a/Engine/Plugins/Enterprise/GLTFExporter/MigratedFromGit/Source/GLTFExporter/Private/Json/GLTFJsonIndex.h
+++ b/Engine/Plugins/Enterprise/GLTFExporter/MigratedFromGit/Source/GLTFExporter/Private/Json/GLTFJsonIndex.h
@@ -67,8 +67,8 @@ struct FGLTFJsonSceneIndex : FGLTFJsonIndex<FGLTFJsonSceneIndex> { using FGLTFJs
 struct FGLTFJsonSkinIndex : FGLTFJsonIndex<FGLTFJsonSkinIndex> { using FGLTFJsonIndex::FGLTFJsonIndex; };
 struct FGLTFJsonTextureIndex : FGLTFJsonIndex<FGLTFJsonTextureIndex> { using FGLTFJsonIndex::FGLTFJsonIndex; };
 struct FGLTFJsonBackdropIndex : FGLTFJsonIndex<FGLTFJsonBackdropIndex> { using FGLTFJsonIndex::FGLTFJsonIndex; };
-struct FGLTFJsonVariationIndex : FGLTFJsonIndex<FGLTFJsonVariationIndex> { using FGLTFJsonIndex::FGLTFJsonIndex; };
-struct FGLTFJsonLightMapIndex : FGLTFJsonIndex<FGLTFJsonLightMapIndex> { using FGLTFJsonIndex::FGLTFJsonIndex; };
-struct FGLTFJsonLightIndex : FGLTFJsonIndex<FGLTFJsonLightIndex> { using FGLTFJsonIndex::FGLTFJsonIndex; };
 struct FGLTFJsonHotspotIndex : FGLTFJsonIndex<FGLTFJsonHotspotIndex> { using FGLTFJsonIndex::FGLTFJsonIndex; };
+struct FGLTFJsonLightIndex : FGLTFJsonIndex<FGLTFJsonLightIndex> { using FGLTFJsonIndex::FGLTFJsonIndex; };
+struct FGLTFJsonLightMapIndex : FGLTFJsonIndex<FGLTFJsonLightMapIndex> { using FGLTFJsonIndex::FGLTFJsonIndex; };
 struct FGLTFJsonSkySphereIndex : FGLTFJsonIndex<FGLTFJsonSkySphereIndex> { using FGLTFJsonIndex::FGLTFJsonIndex; };
+struct FGLTFJsonVariationIndex : FGLTFJsonIndex<FGLTFJsonVariationIndex> { using FGLTFJsonIndex::FGLTFJsonIndex; };
diff --git a/Engine/Plugins/Enterprise/GLTFExporter/MigratedFromGit/Source/GLTFExporter/Private/Json/GLTFJsonNode.h b/Engine/Plugins/Enterprise/GLTFExporter/MigratedFromGit/Source/GLTFExporter/Private/Json/GLTFJsonNode.h
--- a/Engine/Plugins/Enterprise/GLTFExporter/MigratedFromGit/Source/GLTFExporter/Private/Json/GLTFJsonNode.h
+++ b/Engine/Plugins/Enterprise/GLTFExporter/MigratedFromGit/Source/GLTFExporter/Private/Json/GLTFJsonNode.h
@@ -30,9 +30,9 @@ struct FGLTFJsonNode
 	FGLTFJsonSkinIndex      Skin;
 	FGLTFJsonMeshIndex      Mesh;
 	FGLTFJsonBackdropIndex  Backdrop;
-	FGLTFJsonLightMapIndex  LightMap;
-	FGLTFJsonLightIndex     Light;
 	FGLTFJsonHotspotIndex   Hotspot;
+	FGLTFJsonLightIndex     Light;
+	FGLTFJsonLightMapIndex  LightMap;
 	FGLTFJsonSkySphereIndex SkySphere;
 
 	FGLTFJsonNodeIndex ComponentNode;
@@ -101,7 +101,7 @@ struct FGLTFJsonNode
 			JsonWriter.WriteValue(TEXT("mesh"), Mesh);
 		}
 
-		if (Backdrop != INDEX_NONE || LightMap != INDEX_NONE || Light != INDEX_NONE || Hotspot != INDEX_NONE || SkySphere != INDEX_NONE)
+		if (Backdrop != INDEX_NONE || Hotspot != INDEX_NONE || Light != INDEX_NONE || LightMap != INDEX_NONE || SkySphere != INDEX_NONE)
 		{
 			JsonWriter.WriteObjectStart(TEXT("extensions"));
 
@@ -115,13 +115,13 @@ struct FGLTFJsonNode
 				JsonWriter.WriteObjectEnd();
 			}
 
-			if (LightMap != INDEX_NONE)
+			if (Hotspot != INDEX_NONE)
 			{
-				const EGLTFJsonExtension Extension = EGLTFJsonExtension::EPIC_LightmapTextures;
+				const EGLTFJsonExtension Extension = EGLTFJsonExtension::EPIC_AnimationHotspots;
 				Extensions.Used.Add(Extension);
 
 				JsonWriter.WriteObjectStart(FGLTFJsonUtility::ToString(Extension));
-				JsonWriter.WriteValue(TEXT("lightmap"), LightMap);
+				JsonWriter.WriteValue(TEXT("hotspot"), Hotspot);
 				JsonWriter.WriteObjectEnd();
 			}
 
@@ -135,13 +135,13 @@ struct FGLTFJsonNode
 				JsonWriter.WriteObjectEnd();
 			}
 
-			if (Hotspot != INDEX_NONE)
+			if (LightMap != INDEX_NONE)
 			{
-				const EGLTFJsonExtension Extension = EGLTFJsonExtension::EPIC_AnimationHotspots;
+				const EGLTFJsonExtension Extension = EGLTFJsonExtension::EPIC_LightmapTextures;
 				Extensions.Used.Add(Extension);
 
 				JsonWriter.WriteObjectStart(FGLTFJsonUtility::ToString(Extension));
-				JsonWriter.WriteValue(TEXT("hotspot"), Hotspot);
+				JsonWriter.WriteValue(TEXT("lightmap"), LightMap);
 				JsonWriter.WriteObjectEnd();
 			}
 
diff --git a/Engine/Plugins/Enterprise/GLTFExporter/MigratedFromGit/Source/GLTFExporter/Private/Json/GLTFJsonRoot.h b/Engine/Plugins/Enterprise/GLTFExporter/MigratedFromGit/Source/GLTFExporter/Private/Json/GLTFJsonRoot.h
--- a/Engine/Plugins/Enterprise/GLTFExporter/MigratedFromGit/Source/GLTFExporter/Private/Json/GLTFJsonRoot.h
+++ b/Engine/Plugins/Enterprise/GLTFExporter/MigratedFromGit/Source/GLTFExporter/Private/Json/GLTFJsonRoot.h
@@ -18,11 +18,11 @@
 #include "Json/GLTFJsonSkin.h"
 #include "Json/GLTFJsonTexture.h"
 #include "Json/GLTFJsonBackdrop.h"
-#include "Json/GLTFJsonVariation.h"
-#include "Json/GLTFJsonLightMap.h"
-#include "Json/GLTFJsonLight.h"
 #include "Json/GLTFJsonHotspot.h"
+#include "Json/GLTFJsonLight.h"
+#include "Json/GLTFJsonLightMap.h"
 #include "Json/GLTFJsonSkySphere.h"
+#include "Json/GLTFJsonVariation.h"
 #include "Policies/CondensedJsonPrintPolicy.h"
 
 
@@ -77,11 +77,11 @@ struct FGLTFJsonRoot
 	TArray<TUniquePtr<FGLTFJsonSkin>>       Skins;
 	TArray<TUniquePtr<FGLTFJsonTexture>>    Textures;
 	TArray<TUniquePtr<FGLTFJsonBackdrop>>   Backdrops;
-	TArray<TUniquePtr<FGLTFJsonLightMap>>   LightMaps;
-	TArray<TUniquePtr<FGLTFJsonLight>>      Lights;
-	TArray<TUniquePtr<FGLTFJsonVariation>>  Variations;
 	TArray<TUniquePtr<FGLTFJsonHotspot>>    Hotspots;
+	TArray<TUniquePtr<FGLTFJsonLight>>      Lights;
+	TArray<TUniquePtr<FGLTFJsonLightMap>>   LightMaps;
 	TArray<TUniquePtr<FGLTFJsonSkySphere>>  SkySpheres;
+	TArray<TUniquePtr<FGLTFJsonVariation>>  Variations;
 
 	template <class CharType = TCHAR, class PrintPolicy = TPrettyJsonPrintPolicy<CharType>>
 	void WriteObject(TJsonWriter<CharType, PrintPolicy>& JsonWriter, bool bAllExtensionsRequired) const
@@ -112,7 +112,7 @@ struct FGLTFJsonRoot
 		FGLTFJsonUtility::WriteObjectPtrArray(JsonWriter, TEXT("skins"), Skins, AllExtensions);
 		FGLTFJsonUtility::WriteObjectPtrArray(JsonWriter, TEXT("textures"), Textures, AllExtensions);
 
-		if (Backdrops.Num() > 0 || Variations.Num() > 0 || LightMaps.Num() > 0 || Lights.Num() > 0 || Hotspots.Num() > 0 || SkySpheres.Num() > 0 /* TODO: add more extension support */)
+		if (Backdrops.Num() > 0 || Hotspots.Num() > 0 || Lights.Num() > 0 || LightMaps.Num() > 0 || SkySpheres.Num() > 0 || Variations.Num() > 0)
 		{
 			JsonWriter.WriteObjectStart(TEXT("extensions"));
 
@@ -126,23 +126,23 @@ struct FGLTFJsonRoot
 				JsonWriter.WriteObjectEnd();
 			}
 
-			if (Variations.Num() > 0)
+			if (Hotspots.Num() > 0)
 			{
-				const EGLTFJsonExtension Extension = EGLTFJsonExtension::EPIC_LevelVariantSets;
+				const EGLTFJsonExtension Extension = EGLTFJsonExtension::EPIC_AnimationHotspots;
 				AllExtensions.Used.Add(Extension);
 
 				JsonWriter.WriteObjectStart(FGLTFJsonUtility::ToString(Extension));
-				FGLTFJsonUtility::WriteObjectPtrArray(JsonWriter, TEXT("levelVariantSets"), Variations, AllExtensions);
+				FGLTFJsonUtility::WriteObjectPtrArray(JsonWriter, TEXT("hotspots"), Hotspots, AllExtensions);
 				JsonWriter.WriteObjectEnd();
 			}
 
-			if (LightMaps.Num() > 0)
+			if (Variations.Num() > 0)
 			{
-				const EGLTFJsonExtension Extension = EGLTFJsonExtension::EPIC_LightmapTextures;
+				const EGLTFJsonExtension Extension = EGLTFJsonExtension::EPIC_LevelVariantSets;
 				AllExtensions.Used.Add(Extension);
 
 				JsonWriter.WriteObjectStart(FGLTFJsonUtility::ToString(Extension));
-				FGLTFJsonUtility::WriteObjectPtrArray(JsonWriter, TEXT("lightmaps"), LightMaps, AllExtensions);
+				FGLTFJsonUtility::WriteObjectPtrArray(JsonWriter, TEXT("levelVariantSets"), Variations, AllExtensions);
 				JsonWriter.WriteObjectEnd();
 			}
 
@@ -156,13 +156,13 @@ struct FGLTFJsonRoot
 				JsonWriter.WriteObjectEnd();
 			}
 
-			if (Hotspots.Num() > 0)
+			if (LightMaps.Num() > 0)
 			{
-				const EGLTFJsonExtension Extension = EGLTFJsonExtension::EPIC_AnimationHotspots;
+				const EGLTFJsonExtension Extension = EGLTFJsonExtension::EPIC_LightmapTextures;
 				AllExtensions.Used.Add(Extension);
 
 				JsonWriter.WriteObjectStart(FGLTFJsonUtility::ToString(Extension));
-				FGLTFJsonUtility::WriteObjectPtrArray(JsonWriter, TEXT("hotspots"), Hotspots, AllExtensions);
+				FGLTFJsonUtility::WriteObjectPtrArray(JsonWriter, TEXT("lightmaps"), LightMaps, AllExtensions);
 				JsonWriter.WriteObjectEnd();
 			}
 
