[GLTFExporter] refactor: combine export options bBakeMaterialInputs and bMaterialBakeUsingMeshData into single enum BakeMaterialInputs

commit caab2c5706c003506b592e97644870a8b5bf595a
Date: Fri, 13 Nov 2020 17:32:18 +0100

    refactor: combine export options bBakeMaterialInputs and bMaterialBakeUsingMeshData into single enum BakeMaterialInputs
    (cherry picked from commit f0a4d0b36fad60470f53c152110b9ae3cb39f30b​)

#skipci

diff --git a/Engine/Plugins/Enterprise/GLTFExporter/MigratedFromGit/Source/GLTFExporter/Private/Converters/GLTFMaterialConverters.cpp b/Engine/Plugins/Enterprise/GLTFExporter/MigratedFromGit/Source/GLTFExporter/Private/Converters/GLTFMaterialConverters.cpp
--- a/Engine/Plugins/Enterprise/GLTFExporter/MigratedFromGit/Source/GLTFExporter/Private/Converters/GLTFMaterialConverters.cpp
+++ b/Engine/Plugins/Enterprise/GLTFExporter/MigratedFromGit/Source/GLTFExporter/Private/Converters/GLTFMaterialConverters.cpp
@@ -9,8 +9,7 @@
 void FGLTFMaterialConverter::Sanitize(const UMaterialInterface*& Material, const FGLTFMeshData*& MeshData, TArray<int32>& SectionIndices)
 {
 	if (MeshData == nullptr ||
-		!Builder.ExportOptions->bBakeMaterialInputs ||
-		!Builder.ExportOptions->bMaterialBakeUsingMeshData ||
+		Builder.ExportOptions->BakeMaterialInputs != EGLTFMaterialBakeMode::UseMeshData ||
 		!FGLTFMaterialUtility::NeedsMeshData(Material))
 	{
 		MeshData = nullptr;
diff --git a/Engine/Plugins/Enterprise/GLTFExporter/MigratedFromGit/Source/GLTFExporter/Private/Converters/GLTFMeshConverters.cpp b/Engine/Plugins/Enterprise/GLTFExporter/MigratedFromGit/Source/GLTFExporter/Private/Converters/GLTFMeshConverters.cpp
--- a/Engine/Plugins/Enterprise/GLTFExporter/MigratedFromGit/Source/GLTFExporter/Private/Converters/GLTFMeshConverters.cpp
+++ b/Engine/Plugins/Enterprise/GLTFExporter/MigratedFromGit/Source/GLTFExporter/Private/Converters/GLTFMeshConverters.cpp
@@ -30,8 +30,7 @@ void FGLTFStaticMeshConverter::Sanitize(const UStaticMesh*& StaticMesh, const US
 	{
 		// Only use the component if it's needed for baking, since we would
 		// otherwise export a copy of this mesh for each mesh-component.
-		if (!Builder.ExportOptions->bBakeMaterialInputs ||
-			!Builder.ExportOptions->bMaterialBakeUsingMeshData ||
+		if (Builder.ExportOptions->BakeMaterialInputs != EGLTFMaterialBakeMode::UseMeshData ||
 			!FGLTFMaterialUtility::NeedsMeshData(Materials)) // TODO: if this expensive, cache the results for each material
 		{
 			StaticMeshComponent = nullptr;
@@ -70,8 +69,7 @@ void FGLTFSkeletalMeshConverter::Sanitize(const USkeletalMesh*& SkeletalMesh, co
 	{
 		// Only use the component if it's needed for baking, since we would
 		// otherwise export a copy of this mesh for each mesh-component.
-		if (!Builder.ExportOptions->bBakeMaterialInputs ||
-			!Builder.ExportOptions->bMaterialBakeUsingMeshData ||
+		if (Builder.ExportOptions->BakeMaterialInputs != EGLTFMaterialBakeMode::UseMeshData ||
 			!FGLTFMaterialUtility::NeedsMeshData(Materials)) // TODO: if this expensive, cache the results for each material
 		{
 			SkeletalMeshComponent = nullptr;
diff --git a/Engine/Plugins/Enterprise/GLTFExporter/MigratedFromGit/Source/GLTFExporter/Private/Converters/GLTFVarationConverters.cpp b/Engine/Plugins/Enterprise/GLTFExporter/MigratedFromGit/Source/GLTFExporter/Private/Converters/GLTFVarationConverters.cpp
--- a/Engine/Plugins/Enterprise/GLTFExporter/MigratedFromGit/Source/GLTFExporter/Private/Converters/GLTFVarationConverters.cpp
+++ b/Engine/Plugins/Enterprise/GLTFExporter/MigratedFromGit/Source/GLTFExporter/Private/Converters/GLTFVarationConverters.cpp
@@ -276,8 +276,7 @@ bool FGLTFVariationConverter::TryParseMaterialPropertyValue(FGLTFJsonVariant& Ou
 	TArray<int32> SectionIndices;
 
 	if (Builder.ExportOptions->bVariantMaterialBakeUsingMeshData &&
-		Builder.ExportOptions->bMaterialBakeUsingMeshData &&
-		Builder.ExportOptions->bBakeMaterialInputs)
+		Builder.ExportOptions->BakeMaterialInputs == EGLTFMaterialBakeMode::UseMeshData)
 	{
 		const int32 DefaultLODIndex = Builder.ExportOptions->DefaultLevelOfDetail;
 
diff --git a/Engine/Plugins/Enterprise/GLTFExporter/MigratedFromGit/Source/GLTFExporter/Private/GLTFExportOptions.cpp b/Engine/Plugins/Enterprise/GLTFExporter/MigratedFromGit/Source/GLTFExporter/Private/GLTFExportOptions.cpp
--- a/Engine/Plugins/Enterprise/GLTFExporter/MigratedFromGit/Source/GLTFExporter/Private/GLTFExportOptions.cpp
+++ b/Engine/Plugins/Enterprise/GLTFExporter/MigratedFromGit/Source/GLTFExporter/Private/GLTFExportOptions.cpp
@@ -16,8 +16,7 @@ void UGLTFExportOptions::ResetToDefault()
 {
 	bExportUnlitMaterials = true;
 	bExportClearCoatMaterials = true;
-	bBakeMaterialInputs = true;
-	bMaterialBakeUsingMeshData = false;
+	BakeMaterialInputs = EGLTFMaterialBakeMode::Simple;
 	DefaultMaterialBakeSize = EGLTFMaterialBakeSizePOT::POT_1024;
 	bExportVertexColors = false;
 	bExportMeshQuantization = true;
diff --git a/Engine/Plugins/Enterprise/GLTFExporter/MigratedFromGit/Source/GLTFExporter/Private/Tasks/GLTFMaterialTasks.cpp b/Engine/Plugins/Enterprise/GLTFExporter/MigratedFromGit/Source/GLTFExporter/Private/Tasks/GLTFMaterialTasks.cpp
--- a/Engine/Plugins/Enterprise/GLTFExporter/MigratedFromGit/Source/GLTFExporter/Private/Tasks/GLTFMaterialTasks.cpp
+++ b/Engine/Plugins/Enterprise/GLTFExporter/MigratedFromGit/Source/GLTFExporter/Private/Tasks/GLTFMaterialTasks.cpp
@@ -229,7 +229,7 @@ EMaterialShadingModel FGLTFMaterialTask::EvaluateShadingModelExpression() const
 {
 	const FMaterialShadingModelField PossibleValues = Material->GetShadingModels();
 
-	if (!Builder.ExportOptions->bBakeMaterialInputs)
+	if (Builder.ExportOptions->BakeMaterialInputs == EGLTFMaterialBakeMode::Disabled)
 	{
 		const EMaterialShadingModel ShadingModel = FGLTFMaterialUtility::GetRichestShadingModel(PossibleValues);
 		Builder.AddWarningMessage(FString::Printf(
@@ -301,7 +301,7 @@ bool FGLTFMaterialTask::TryGetBaseColorAndOpacity(FGLTFJsonPBRMetallicRoughness&
 		return true;
 	}
 
-	if (!Builder.ExportOptions->bBakeMaterialInputs)
+	if (Builder.ExportOptions->BakeMaterialInputs == EGLTFMaterialBakeMode::Disabled)
 	{
 		Builder.AddWarningMessage(FString::Printf(
 			TEXT("%s and %s for material %s needs to bake, but material baking is disabled by export options"),
@@ -469,7 +469,7 @@ bool FGLTFMaterialTask::TryGetMetallicAndRoughness(FGLTFJsonPBRMetallicRoughness
 		return true;
 	}
 
-	if (!Builder.ExportOptions->bBakeMaterialInputs)
+	if (Builder.ExportOptions->BakeMaterialInputs == EGLTFMaterialBakeMode::Disabled)
 	{
 		Builder.AddWarningMessage(FString::Printf(
 			TEXT("%s and %s for material %s needs to bake, but material baking is disabled by export options"),
@@ -636,7 +636,7 @@ bool FGLTFMaterialTask::TryGetClearCoatRoughness(FGLTFJsonClearCoatExtension& Ou
 		return true;
 	}
 
-	if (!Builder.ExportOptions->bBakeMaterialInputs)
+	if (Builder.ExportOptions->BakeMaterialInputs == EGLTFMaterialBakeMode::Disabled)
 	{
 		Builder.AddWarningMessage(FString::Printf(
 			TEXT("%s and %s for material %s needs to bake, but material baking is disabled by export options"),
@@ -779,7 +779,7 @@ bool FGLTFMaterialTask::TryGetEmissive(FGLTFJsonMaterial& JsonMaterial, EMateria
 		return true;
 	}
 
-	if (!Builder.ExportOptions->bBakeMaterialInputs)
+	if (Builder.ExportOptions->BakeMaterialInputs == EGLTFMaterialBakeMode::Disabled)
 	{
 		Builder.AddWarningMessage(FString::Printf(
 			TEXT("%s for material %s needs to bake, but material baking is disabled by export options"),
@@ -1196,7 +1196,7 @@ bool FGLTFMaterialTask::TryGetSourceTexture(const UTexture2D*& OutTexture, int32
 
 bool FGLTFMaterialTask::TryGetBakedMaterialProperty(FGLTFJsonTextureInfo& OutTexInfo, FGLTFJsonColor3& OutConstant, EMaterialProperty Property, const FString& PropertyName) const
 {
-	if (!Builder.ExportOptions->bBakeMaterialInputs)
+	if (Builder.ExportOptions->BakeMaterialInputs == EGLTFMaterialBakeMode::Disabled)
 	{
 		Builder.AddWarningMessage(FString::Printf(
 			TEXT("%s for material %s needs to bake, but material baking is disabled by export options"),
@@ -1224,7 +1224,7 @@ bool FGLTFMaterialTask::TryGetBakedMaterialProperty(FGLTFJsonTextureInfo& OutTex
 
 bool FGLTFMaterialTask::TryGetBakedMaterialProperty(FGLTFJsonTextureInfo& OutTexInfo, FGLTFJsonColor4& OutConstant, EMaterialProperty Property, const FString& PropertyName) const
 {
-	if (!Builder.ExportOptions->bBakeMaterialInputs)
+	if (Builder.ExportOptions->BakeMaterialInputs == EGLTFMaterialBakeMode::Disabled)
 	{
 		Builder.AddWarningMessage(FString::Printf(
 			TEXT("%s for material %s needs to bake, but material baking is disabled by export options"),
@@ -1252,7 +1252,7 @@ bool FGLTFMaterialTask::TryGetBakedMaterialProperty(FGLTFJsonTextureInfo& OutTex
 
 inline bool FGLTFMaterialTask::TryGetBakedMaterialProperty(FGLTFJsonTextureInfo& OutTexInfo, float& OutConstant, EMaterialProperty Property, const FString& PropertyName) const
 {
-	if (!Builder.ExportOptions->bBakeMaterialInputs)
+	if (Builder.ExportOptions->BakeMaterialInputs == EGLTFMaterialBakeMode::Disabled)
 	{
 		Builder.AddWarningMessage(FString::Printf(
 			TEXT("%s for material %s needs to bake, but material baking is disabled by export options"),
@@ -1280,7 +1280,7 @@ inline bool FGLTFMaterialTask::TryGetBakedMaterialProperty(FGLTFJsonTextureInfo&
 
 bool FGLTFMaterialTask::TryGetBakedMaterialProperty(FGLTFJsonTextureInfo& OutTexInfo, EMaterialProperty Property, const FString& PropertyName) const
 {
-	if (!Builder.ExportOptions->bBakeMaterialInputs)
+	if (Builder.ExportOptions->BakeMaterialInputs == EGLTFMaterialBakeMode::Disabled)
 	{
 		Builder.AddWarningMessage(FString::Printf(
 			TEXT("%s for material %s needs to bake, but material baking is disabled by export options"),
diff --git a/Engine/Plugins/Enterprise/GLTFExporter/MigratedFromGit/Source/GLTFExporter/Private/Tasks/GLTFMeshTasks.cpp b/Engine/Plugins/Enterprise/GLTFExporter/MigratedFromGit/Source/GLTFExporter/Private/Tasks/GLTFMeshTasks.cpp
--- a/Engine/Plugins/Enterprise/GLTFExporter/MigratedFromGit/Source/GLTFExporter/Private/Tasks/GLTFMeshTasks.cpp
+++ b/Engine/Plugins/Enterprise/GLTFExporter/MigratedFromGit/Source/GLTFExporter/Private/Tasks/GLTFMeshTasks.cpp
@@ -22,9 +22,8 @@ void FGLTFStaticMeshTask::Complete()
 		ColorBuffer = LODInfo.OverrideVertexColors != nullptr ? LODInfo.OverrideVertexColors : ColorBuffer;
 	}
 
-	const FGLTFMeshData* MeshData =
-		Builder.ExportOptions->bBakeMaterialInputs && Builder.ExportOptions->bMaterialBakeUsingMeshData ?
-			Builder.StaticMeshDataConverter.GetOrAdd(StaticMesh, StaticMeshComponent, LODIndex) : nullptr;
+	const FGLTFMeshData* MeshData = Builder.ExportOptions->BakeMaterialInputs == EGLTFMaterialBakeMode::UseMeshData ?
+		Builder.StaticMeshDataConverter.GetOrAdd(StaticMesh, StaticMeshComponent, LODIndex) : nullptr;
 
 	if (MeshData != nullptr && MeshData->Description.IsEmpty())
 	{
@@ -90,9 +89,8 @@ void FGLTFSkeletalMeshTask::Complete()
 		SkinWeightBuffer = LODInfo.OverrideSkinWeights != nullptr ? LODInfo.OverrideSkinWeights : SkinWeightBuffer;
 	}
 
-	const FGLTFMeshData* MeshData =
-		Builder.ExportOptions->bBakeMaterialInputs && Builder.ExportOptions->bMaterialBakeUsingMeshData ?
-			Builder.SkeletalMeshDataConverter.GetOrAdd(SkeletalMesh, SkeletalMeshComponent, LODIndex) : nullptr;
+	const FGLTFMeshData* MeshData = Builder.ExportOptions->BakeMaterialInputs == EGLTFMaterialBakeMode::UseMeshData ?
+		Builder.SkeletalMeshDataConverter.GetOrAdd(SkeletalMesh, SkeletalMeshComponent, LODIndex) : nullptr;
 
 	if (MeshData != nullptr && MeshData->Description.IsEmpty())
 	{
diff --git a/Engine/Plugins/Enterprise/GLTFExporter/MigratedFromGit/Source/GLTFExporter/Public/GLTFExportOptions.h b/Engine/Plugins/Enterprise/GLTFExporter/MigratedFromGit/Source/GLTFExporter/Public/GLTFExportOptions.h
--- a/Engine/Plugins/Enterprise/GLTFExporter/MigratedFromGit/Source/GLTFExporter/Public/GLTFExportOptions.h
+++ b/Engine/Plugins/Enterprise/GLTFExporter/MigratedFromGit/Source/GLTFExporter/Public/GLTFExportOptions.h
@@ -5,6 +5,14 @@
 #include "CoreMinimal.h"
 #include "GLTFExportOptions.generated.h"
 
+UENUM(BlueprintType)
+enum class EGLTFMaterialBakeMode : uint8
+{
+	Disabled,
+    Simple,
+    UseMeshData,
+};
+
 UENUM(BlueprintType)
 enum class EGLTFMaterialBakeSizePOT : uint8
 {
@@ -74,12 +82,9 @@ public:
 	bool bExportClearCoatMaterials;
 
 	UPROPERTY(EditAnywhere, BlueprintReadWrite, Config, Category = Material)
-	bool bBakeMaterialInputs;
-
-	UPROPERTY(EditAnywhere, BlueprintReadWrite, Config, Category = Material, Meta = (EditCondition = "bBakeMaterialInputs"))
-	bool bMaterialBakeUsingMeshData;
+	EGLTFMaterialBakeMode BakeMaterialInputs;
 
-	UPROPERTY(EditAnywhere, BlueprintReadWrite, Config, Category = Material, Meta = (EditCondition = "bBakeMaterialInputs"))
+	UPROPERTY(EditAnywhere, BlueprintReadWrite, Config, Category = Material, Meta = (EditCondition = "BakeMaterialInputs != EGLTFMaterialBakeMode::Disabled"))
 	EGLTFMaterialBakeSizePOT DefaultMaterialBakeSize;
 
 	UPROPERTY(EditAnywhere, BlueprintReadWrite, Config, Category = Mesh)
@@ -150,7 +155,7 @@ public:
 	UPROPERTY(EditAnywhere, BlueprintReadWrite, Config, Category = Scene)
 	bool bExportVariantSets;
 
-	UPROPERTY(EditAnywhere, BlueprintReadWrite, Config, Category = Material, Meta = (EditCondition = "bExportVariantSets && bBakeMaterialInputs && bMaterialBakeUsingMeshData"))
+	UPROPERTY(EditAnywhere, BlueprintReadWrite, Config, Category = Material, Meta = (EditCondition = "bExportVariantSets && BakeMaterialInputs == EGLTFMaterialBakeMode::UseMeshData"))
 	bool bVariantMaterialBakeUsingMeshData;
 
 	UPROPERTY(EditAnywhere, BlueprintReadWrite, Config, Category = Scene)
