[GLTFExporter] refactor: change export option bVariantMaterialBakeUsingMeshData to ExportMaterialVariants and rename EGLTFMaterialBakeMode::Disabled to EGLTFMaterialBakeMode::None

commit 937c292b6e62920c3b847636ddadb82a3ffc70a2
Date: Fri, 13 Nov 2020 18:05:26 +0100

    refactor: change export option bVariantMaterialBakeUsingMeshData to ExportMaterialVariants and rename EGLTFMaterialBakeMode::Disabled to EGLTFMaterialBakeMode::None
    (cherry picked from commit db540a800a10784b5d6bf44521a0118fc60e788d​)

#skipci

diff --git a/Engine/Plugins/Enterprise/GLTFExporter/MigratedFromGit/Source/GLTFExporter/Private/Converters/GLTFVarationConverters.cpp b/Engine/Plugins/Enterprise/GLTFExporter/MigratedFromGit/Source/GLTFExporter/Private/Converters/GLTFVarationConverters.cpp
--- a/Engine/Plugins/Enterprise/GLTFExporter/MigratedFromGit/Source/GLTFExporter/Private/Converters/GLTFVarationConverters.cpp
+++ b/Engine/Plugins/Enterprise/GLTFExporter/MigratedFromGit/Source/GLTFExporter/Private/Converters/GLTFVarationConverters.cpp
@@ -114,38 +114,33 @@ bool FGLTFVariationConverter::TryParseVariantBinding(FGLTFJsonVariant& OutVarian
 			continue;
 		}
 
-		const bool bIsVisibilityProperty =
-			Property->GetPropertyName() == TEXT("bVisible") &&
-			Property->GetPropertyClass()->IsChildOf(FBoolProperty::StaticClass());
+		const FName PropertyName = Property->GetPropertyName();
+		const FFieldClass* PropertyClass = Property->GetPropertyClass();
 
-		const bool bIsMaterialProperty = Property->IsA<UPropertyValueMaterial>();
-		const bool bIsStaticMeshProperty = Property->GetPropertyName() == TEXT("StaticMesh");
-		const bool bIsSkeletalMeshProperty = Property->GetPropertyName() == TEXT("SkeletalMesh");
-
-		if (bIsVisibilityProperty)
+		if (Property->IsA<UPropertyValueMaterial>())
 		{
-			if (TryParseVisibilityPropertyValue(OutVariant, Property))
+			if (Builder.ExportOptions->ExportMaterialVariants != EGLTFMaterialBakeMode::None && TryParseMaterialPropertyValue(OutVariant, Property))
 			{
 				bHasParsedAnyProperty = true;
 			}
 		}
-		else if (bIsMaterialProperty)
+		else if (PropertyName == TEXT("StaticMesh")) // TODO: should we not also check PropertyClass?
 		{
-			if (TryParseMaterialPropertyValue(OutVariant, Property))
+			if (TryParseStaticMeshPropertyValue(OutVariant, Property))
 			{
 				bHasParsedAnyProperty = true;
 			}
 		}
-		else if (bIsStaticMeshProperty)
+		else if (PropertyName == TEXT("SkeletalMesh")) // TODO: should we not also check PropertyClass?
 		{
-			if (TryParseStaticMeshPropertyValue(OutVariant, Property))
+			if (TryParseSkeletalMeshPropertyValue(OutVariant, Property))
 			{
 				bHasParsedAnyProperty = true;
 			}
 		}
-		else if (bIsSkeletalMeshProperty)
+		else if (PropertyName == TEXT("bVisible") && PropertyClass->IsChildOf(FBoolProperty::StaticClass()))
 		{
-			if (TryParseSkeletalMeshPropertyValue(OutVariant, Property))
+			if (TryParseVisibilityPropertyValue(OutVariant, Property))
 			{
 				bHasParsedAnyProperty = true;
 			}
@@ -275,29 +270,35 @@ bool FGLTFVariationConverter::TryParseMaterialPropertyValue(FGLTFJsonVariant& Ou
 	const FGLTFMeshData* MeshData = nullptr;
 	TArray<int32> SectionIndices;
 
-	if (Builder.ExportOptions->bVariantMaterialBakeUsingMeshData &&
-		Builder.ExportOptions->BakeMaterialInputs == EGLTFMaterialBakeMode::UseMeshData)
+	if (Builder.ExportOptions->ExportMaterialVariants == EGLTFMaterialBakeMode::UseMeshData)
 	{
-		const int32 DefaultLODIndex = Builder.ExportOptions->DefaultLevelOfDetail;
-
-		if (const UStaticMeshComponent* StaticMeshComponent = Cast<UStaticMeshComponent>(Target))
+		if (Builder.ExportOptions->BakeMaterialInputs == EGLTFMaterialBakeMode::UseMeshData)
 		{
-			const UStaticMesh* StaticMesh = StaticMeshComponent->GetStaticMesh();
-			const int32 LODIndex = FGLTFMeshUtility::GetLOD(StaticMesh, StaticMeshComponent, DefaultLODIndex);
-			const FStaticMeshLODResources& MeshLOD = StaticMesh->GetLODForExport(LODIndex);
+			const int32 DefaultLODIndex = Builder.ExportOptions->DefaultLevelOfDetail;
 
-			SectionIndices = FGLTFMeshUtility::GetSectionIndices(MeshLOD, MaterialIndex);
-			MeshData = Builder.StaticMeshDataConverter.GetOrAdd(StaticMesh, StaticMeshComponent, LODIndex);
+			if (const UStaticMeshComponent* StaticMeshComponent = Cast<UStaticMeshComponent>(Target))
+			{
+				const UStaticMesh* StaticMesh = StaticMeshComponent->GetStaticMesh();
+				const int32 LODIndex = FGLTFMeshUtility::GetLOD(StaticMesh, StaticMeshComponent, DefaultLODIndex);
+				const FStaticMeshLODResources& MeshLOD = StaticMesh->GetLODForExport(LODIndex);
+
+				SectionIndices = FGLTFMeshUtility::GetSectionIndices(MeshLOD, MaterialIndex);
+				MeshData = Builder.StaticMeshDataConverter.GetOrAdd(StaticMesh, StaticMeshComponent, LODIndex);
+			}
+			else if (const USkeletalMeshComponent* SkeletalMeshComponent = Cast<USkeletalMeshComponent>(Target))
+			{
+				const USkeletalMesh* SkeletalMesh = SkeletalMeshComponent->SkeletalMesh;
+				const int32 LODIndex = FGLTFMeshUtility::GetLOD(SkeletalMesh, SkeletalMeshComponent, DefaultLODIndex);
+				const FSkeletalMeshRenderData* RenderData = SkeletalMesh->GetResourceForRendering();
+				const FSkeletalMeshLODRenderData& MeshLOD = RenderData->LODRenderData[LODIndex];
+
+				SectionIndices = FGLTFMeshUtility::GetSectionIndices(MeshLOD, MaterialIndex);
+				MeshData = Builder.SkeletalMeshDataConverter.GetOrAdd(SkeletalMesh, SkeletalMeshComponent, LODIndex);
+			}
 		}
-		else if (const USkeletalMeshComponent* SkeletalMeshComponent = Cast<USkeletalMeshComponent>(Target))
+		else
 		{
-			const USkeletalMesh* SkeletalMesh = SkeletalMeshComponent->SkeletalMesh;
-			const int32 LODIndex = FGLTFMeshUtility::GetLOD(SkeletalMesh, SkeletalMeshComponent, DefaultLODIndex);
-			const FSkeletalMeshRenderData* RenderData = SkeletalMesh->GetResourceForRendering();
-			const FSkeletalMeshLODRenderData& MeshLOD = RenderData->LODRenderData[LODIndex];
-
-			SectionIndices = FGLTFMeshUtility::GetSectionIndices(MeshLOD, MaterialIndex);
-			MeshData = Builder.SkeletalMeshDataConverter.GetOrAdd(SkeletalMesh, SkeletalMeshComponent, LODIndex);
+			// TODO: report warning (about materials won't be export using mesh data because BakeMaterialInputs is not set to UseMeshData)
 		}
 	}
 
diff --git a/Engine/Plugins/Enterprise/GLTFExporter/MigratedFromGit/Source/GLTFExporter/Private/GLTFExportOptions.cpp b/Engine/Plugins/Enterprise/GLTFExporter/MigratedFromGit/Source/GLTFExporter/Private/GLTFExportOptions.cpp
--- a/Engine/Plugins/Enterprise/GLTFExporter/MigratedFromGit/Source/GLTFExporter/Private/GLTFExportOptions.cpp
+++ b/Engine/Plugins/Enterprise/GLTFExporter/MigratedFromGit/Source/GLTFExporter/Private/GLTFExportOptions.cpp
@@ -40,8 +40,8 @@ void UGLTFExportOptions::ResetToDefault()
 	bExportHDRIBackdrops = true;
 	bExportSkySpheres = true;
 	bExportVariantSets = true;
-	bVariantMaterialBakeUsingMeshData = false;
 	bExportAnimationHotspots = true;
+	ExportMaterialVariants = EGLTFMaterialBakeMode::Simple;
 	bBundleWebViewer = true;
 	bExportPreviewMesh = true;
 	bAllExtensionsRequired = false;
diff --git a/Engine/Plugins/Enterprise/GLTFExporter/MigratedFromGit/Source/GLTFExporter/Private/Tasks/GLTFMaterialTasks.cpp b/Engine/Plugins/Enterprise/GLTFExporter/MigratedFromGit/Source/GLTFExporter/Private/Tasks/GLTFMaterialTasks.cpp
--- a/Engine/Plugins/Enterprise/GLTFExporter/MigratedFromGit/Source/GLTFExporter/Private/Tasks/GLTFMaterialTasks.cpp
+++ b/Engine/Plugins/Enterprise/GLTFExporter/MigratedFromGit/Source/GLTFExporter/Private/Tasks/GLTFMaterialTasks.cpp
@@ -229,7 +229,7 @@ EMaterialShadingModel FGLTFMaterialTask::EvaluateShadingModelExpression() const
 {
 	const FMaterialShadingModelField PossibleValues = Material->GetShadingModels();
 
-	if (Builder.ExportOptions->BakeMaterialInputs == EGLTFMaterialBakeMode::Disabled)
+	if (Builder.ExportOptions->BakeMaterialInputs == EGLTFMaterialBakeMode::None)
 	{
 		const EMaterialShadingModel ShadingModel = FGLTFMaterialUtility::GetRichestShadingModel(PossibleValues);
 		Builder.AddWarningMessage(FString::Printf(
@@ -301,7 +301,7 @@ bool FGLTFMaterialTask::TryGetBaseColorAndOpacity(FGLTFJsonPBRMetallicRoughness&
 		return true;
 	}
 
-	if (Builder.ExportOptions->BakeMaterialInputs == EGLTFMaterialBakeMode::Disabled)
+	if (Builder.ExportOptions->BakeMaterialInputs == EGLTFMaterialBakeMode::None)
 	{
 		Builder.AddWarningMessage(FString::Printf(
 			TEXT("%s and %s for material %s needs to bake, but material baking is disabled by export options"),
@@ -469,7 +469,7 @@ bool FGLTFMaterialTask::TryGetMetallicAndRoughness(FGLTFJsonPBRMetallicRoughness
 		return true;
 	}
 
-	if (Builder.ExportOptions->BakeMaterialInputs == EGLTFMaterialBakeMode::Disabled)
+	if (Builder.ExportOptions->BakeMaterialInputs == EGLTFMaterialBakeMode::None)
 	{
 		Builder.AddWarningMessage(FString::Printf(
 			TEXT("%s and %s for material %s needs to bake, but material baking is disabled by export options"),
@@ -636,7 +636,7 @@ bool FGLTFMaterialTask::TryGetClearCoatRoughness(FGLTFJsonClearCoatExtension& Ou
 		return true;
 	}
 
-	if (Builder.ExportOptions->BakeMaterialInputs == EGLTFMaterialBakeMode::Disabled)
+	if (Builder.ExportOptions->BakeMaterialInputs == EGLTFMaterialBakeMode::None)
 	{
 		Builder.AddWarningMessage(FString::Printf(
 			TEXT("%s and %s for material %s needs to bake, but material baking is disabled by export options"),
@@ -779,7 +779,7 @@ bool FGLTFMaterialTask::TryGetEmissive(FGLTFJsonMaterial& JsonMaterial, EMateria
 		return true;
 	}
 
-	if (Builder.ExportOptions->BakeMaterialInputs == EGLTFMaterialBakeMode::Disabled)
+	if (Builder.ExportOptions->BakeMaterialInputs == EGLTFMaterialBakeMode::None)
 	{
 		Builder.AddWarningMessage(FString::Printf(
 			TEXT("%s for material %s needs to bake, but material baking is disabled by export options"),
@@ -1196,7 +1196,7 @@ bool FGLTFMaterialTask::TryGetSourceTexture(const UTexture2D*& OutTexture, int32
 
 bool FGLTFMaterialTask::TryGetBakedMaterialProperty(FGLTFJsonTextureInfo& OutTexInfo, FGLTFJsonColor3& OutConstant, EMaterialProperty Property, const FString& PropertyName) const
 {
-	if (Builder.ExportOptions->BakeMaterialInputs == EGLTFMaterialBakeMode::Disabled)
+	if (Builder.ExportOptions->BakeMaterialInputs == EGLTFMaterialBakeMode::None)
 	{
 		Builder.AddWarningMessage(FString::Printf(
 			TEXT("%s for material %s needs to bake, but material baking is disabled by export options"),
@@ -1224,7 +1224,7 @@ bool FGLTFMaterialTask::TryGetBakedMaterialProperty(FGLTFJsonTextureInfo& OutTex
 
 bool FGLTFMaterialTask::TryGetBakedMaterialProperty(FGLTFJsonTextureInfo& OutTexInfo, FGLTFJsonColor4& OutConstant, EMaterialProperty Property, const FString& PropertyName) const
 {
-	if (Builder.ExportOptions->BakeMaterialInputs == EGLTFMaterialBakeMode::Disabled)
+	if (Builder.ExportOptions->BakeMaterialInputs == EGLTFMaterialBakeMode::None)
 	{
 		Builder.AddWarningMessage(FString::Printf(
 			TEXT("%s for material %s needs to bake, but material baking is disabled by export options"),
@@ -1252,7 +1252,7 @@ bool FGLTFMaterialTask::TryGetBakedMaterialProperty(FGLTFJsonTextureInfo& OutTex
 
 inline bool FGLTFMaterialTask::TryGetBakedMaterialProperty(FGLTFJsonTextureInfo& OutTexInfo, float& OutConstant, EMaterialProperty Property, const FString& PropertyName) const
 {
-	if (Builder.ExportOptions->BakeMaterialInputs == EGLTFMaterialBakeMode::Disabled)
+	if (Builder.ExportOptions->BakeMaterialInputs == EGLTFMaterialBakeMode::None)
 	{
 		Builder.AddWarningMessage(FString::Printf(
 			TEXT("%s for material %s needs to bake, but material baking is disabled by export options"),
@@ -1280,7 +1280,7 @@ inline bool FGLTFMaterialTask::TryGetBakedMaterialProperty(FGLTFJsonTextureInfo&
 
 bool FGLTFMaterialTask::TryGetBakedMaterialProperty(FGLTFJsonTextureInfo& OutTexInfo, EMaterialProperty Property, const FString& PropertyName) const
 {
-	if (Builder.ExportOptions->BakeMaterialInputs == EGLTFMaterialBakeMode::Disabled)
+	if (Builder.ExportOptions->BakeMaterialInputs == EGLTFMaterialBakeMode::None)
 	{
 		Builder.AddWarningMessage(FString::Printf(
 			TEXT("%s for material %s needs to bake, but material baking is disabled by export options"),
diff --git a/Engine/Plugins/Enterprise/GLTFExporter/MigratedFromGit/Source/GLTFExporter/Public/GLTFExportOptions.h b/Engine/Plugins/Enterprise/GLTFExporter/MigratedFromGit/Source/GLTFExporter/Public/GLTFExportOptions.h
--- a/Engine/Plugins/Enterprise/GLTFExporter/MigratedFromGit/Source/GLTFExporter/Public/GLTFExportOptions.h
+++ b/Engine/Plugins/Enterprise/GLTFExporter/MigratedFromGit/Source/GLTFExporter/Public/GLTFExportOptions.h
@@ -8,7 +8,7 @@
 UENUM(BlueprintType)
 enum class EGLTFMaterialBakeMode : uint8
 {
-	Disabled,
+	None,
     Simple,
     UseMeshData,
 };
@@ -84,7 +84,7 @@ public:
 	UPROPERTY(EditAnywhere, BlueprintReadWrite, Config, Category = Material)
 	EGLTFMaterialBakeMode BakeMaterialInputs;
 
-	UPROPERTY(EditAnywhere, BlueprintReadWrite, Config, Category = Material, Meta = (EditCondition = "BakeMaterialInputs != EGLTFMaterialBakeMode::Disabled"))
+	UPROPERTY(EditAnywhere, BlueprintReadWrite, Config, Category = Material, Meta = (EditCondition = "BakeMaterialInputs != EGLTFMaterialBakeMode::None"))
 	EGLTFMaterialBakeSizePOT DefaultMaterialBakeSize;
 
 	UPROPERTY(EditAnywhere, BlueprintReadWrite, Config, Category = Mesh)
@@ -155,12 +155,12 @@ public:
 	UPROPERTY(EditAnywhere, BlueprintReadWrite, Config, Category = Scene)
 	bool bExportVariantSets;
 
-	UPROPERTY(EditAnywhere, BlueprintReadWrite, Config, Category = Material, Meta = (EditCondition = "bExportVariantSets && BakeMaterialInputs == EGLTFMaterialBakeMode::UseMeshData"))
-	bool bVariantMaterialBakeUsingMeshData;
-
 	UPROPERTY(EditAnywhere, BlueprintReadWrite, Config, Category = Scene)
 	bool bExportAnimationHotspots;
 
+	UPROPERTY(EditAnywhere, BlueprintReadWrite, Config, Category = VariantSets, Meta = (EditCondition = "bExportVariantSets"))
+	EGLTFMaterialBakeMode ExportMaterialVariants;
+
 	UPROPERTY(EditAnywhere, BlueprintReadWrite, Config, Category = Exporter)
 	bool bBundleWebViewer;
 
