[GLTFExporter] feat: add export options bExportMeshVariants and bExportVisibilityVariants

commit a6d503270e1110f15eb097679bbc64d309245010
Date: Fri, 13 Nov 2020 18:07:43 +0100

    feat: add export options bExportMeshVariants and bExportVisibilityVariants
    (cherry picked from commit 20f742ca9ba9dfb5c555c62047712f331bae5885​)

#skipci

diff --git a/Engine/Plugins/Enterprise/GLTFExporter/MigratedFromGit/Source/GLTFExporter/Private/Converters/GLTFVarationConverters.cpp b/Engine/Plugins/Enterprise/GLTFExporter/MigratedFromGit/Source/GLTFExporter/Private/Converters/GLTFVarationConverters.cpp
--- a/Engine/Plugins/Enterprise/GLTFExporter/MigratedFromGit/Source/GLTFExporter/Private/Converters/GLTFVarationConverters.cpp
+++ b/Engine/Plugins/Enterprise/GLTFExporter/MigratedFromGit/Source/GLTFExporter/Private/Converters/GLTFVarationConverters.cpp
@@ -126,21 +126,21 @@ bool FGLTFVariationConverter::TryParseVariantBinding(FGLTFJsonVariant& OutVarian
 		}
 		else if (PropertyName == TEXT("StaticMesh")) // TODO: should we not also check PropertyClass?
 		{
-			if (TryParseStaticMeshPropertyValue(OutVariant, Property))
+			if (Builder.ExportOptions->bExportMeshVariants && TryParseStaticMeshPropertyValue(OutVariant, Property))
 			{
 				bHasParsedAnyProperty = true;
 			}
 		}
 		else if (PropertyName == TEXT("SkeletalMesh")) // TODO: should we not also check PropertyClass?
 		{
-			if (TryParseSkeletalMeshPropertyValue(OutVariant, Property))
+			if (Builder.ExportOptions->bExportMeshVariants && TryParseSkeletalMeshPropertyValue(OutVariant, Property))
 			{
 				bHasParsedAnyProperty = true;
 			}
 		}
 		else if (PropertyName == TEXT("bVisible") && PropertyClass->IsChildOf(FBoolProperty::StaticClass()))
 		{
-			if (TryParseVisibilityPropertyValue(OutVariant, Property))
+			if (Builder.ExportOptions->bExportVisibilityVariants && TryParseVisibilityPropertyValue(OutVariant, Property))
 			{
 				bHasParsedAnyProperty = true;
 			}
diff --git a/Engine/Plugins/Enterprise/GLTFExporter/MigratedFromGit/Source/GLTFExporter/Private/GLTFExportOptions.cpp b/Engine/Plugins/Enterprise/GLTFExporter/MigratedFromGit/Source/GLTFExporter/Private/GLTFExportOptions.cpp
--- a/Engine/Plugins/Enterprise/GLTFExporter/MigratedFromGit/Source/GLTFExporter/Private/GLTFExportOptions.cpp
+++ b/Engine/Plugins/Enterprise/GLTFExporter/MigratedFromGit/Source/GLTFExporter/Private/GLTFExportOptions.cpp
@@ -42,6 +42,8 @@ void UGLTFExportOptions::ResetToDefault()
 	bExportVariantSets = true;
 	bExportAnimationHotspots = true;
 	ExportMaterialVariants = EGLTFMaterialBakeMode::Simple;
+	bExportMeshVariants = true;
+	bExportVisibilityVariants = true;
 	bBundleWebViewer = true;
 	bExportPreviewMesh = true;
 	bAllExtensionsRequired = false;
diff --git a/Engine/Plugins/Enterprise/GLTFExporter/MigratedFromGit/Source/GLTFExporter/Public/GLTFExportOptions.h b/Engine/Plugins/Enterprise/GLTFExporter/MigratedFromGit/Source/GLTFExporter/Public/GLTFExportOptions.h
--- a/Engine/Plugins/Enterprise/GLTFExporter/MigratedFromGit/Source/GLTFExporter/Public/GLTFExportOptions.h
+++ b/Engine/Plugins/Enterprise/GLTFExporter/MigratedFromGit/Source/GLTFExporter/Public/GLTFExportOptions.h
@@ -161,6 +161,12 @@ public:
 	UPROPERTY(EditAnywhere, BlueprintReadWrite, Config, Category = VariantSets, Meta = (EditCondition = "bExportVariantSets"))
 	EGLTFMaterialBakeMode ExportMaterialVariants;
 
+	UPROPERTY(EditAnywhere, BlueprintReadWrite, Config, Category = VariantSets, Meta = (EditCondition = "bExportVariantSets"))
+	bool bExportMeshVariants;
+
+	UPROPERTY(EditAnywhere, BlueprintReadWrite, Config, Category = VariantSets, Meta = (EditCondition = "bExportVariantSets"))
+	bool bExportVisibilityVariants;
+
 	UPROPERTY(EditAnywhere, BlueprintReadWrite, Config, Category = Exporter)
 	bool bBundleWebViewer;
 
