[GLTFExporter] refactor: simplify variation converter to avoid warning about unsupported properties

commit 9af7fb56b2d5e9c492bd62fa7089396c2219b4f6
Date: Fri, 13 Nov 2020 18:25:33 +0100

    refactor: simplify variation converter to avoid warning about unsupported properties
    (cherry picked from commit 0f7516bfb607805ad6209834aa36b2cc42357e87​)

#skipci

diff --git a/Engine/Plugins/Enterprise/GLTFExporter/MigratedFromGit/Source/GLTFExporter/Private/Converters/GLTFVarationConverters.cpp b/Engine/Plugins/Enterprise/GLTFExporter/MigratedFromGit/Source/GLTFExporter/Private/Converters/GLTFVarationConverters.cpp
--- a/Engine/Plugins/Enterprise/GLTFExporter/MigratedFromGit/Source/GLTFExporter/Private/Converters/GLTFVarationConverters.cpp
+++ b/Engine/Plugins/Enterprise/GLTFExporter/MigratedFromGit/Source/GLTFExporter/Private/Converters/GLTFVarationConverters.cpp
@@ -54,12 +54,6 @@ FGLTFJsonVariationIndex FGLTFVariationConverter::Convert(const ALevelVariantSets
 		{
 			JsonVariation.VariantSets.Add(JsonVariantSet);
 		}
-		else
-		{
-			Builder.AddWarningMessage(FString::Printf(
-				TEXT("Variant-set has no supported variants and will be skipped. Context: %s"),
-				*GetLogContext(VariantSet)));
-		}
 	}
 
 	if (JsonVariation.VariantSets.Num() == 0)
@@ -81,9 +75,6 @@ bool FGLTFVariationConverter::TryParseVariant(FGLTFJsonVariant& OutVariant, cons
 
 	if (JsonVariant.Nodes.Num() == 0)
 	{
-		Builder.AddWarningMessage(FString::Printf(
-			TEXT("Variant has no supported bindings and will be skipped. Context: %s"),
-			*GetLogContext(Variant)));
 		return false;
 	}
 
@@ -108,9 +99,6 @@ bool FGLTFVariationConverter::TryParseVariantBinding(FGLTFJsonVariant& OutVarian
 	{
 		if (!const_cast<UPropertyValue*>(Property)->Resolve() || !Property->HasRecordedData())
 		{
-			Builder.AddWarningMessage(FString::Printf(
-				TEXT("Property is missing recorded data, it will be skipped. Context: %s"),
-				*GetLogContext(Property)));
 			continue;
 		}
 
@@ -145,21 +133,8 @@ bool FGLTFVariationConverter::TryParseVariantBinding(FGLTFJsonVariant& OutVarian
 				bHasParsedAnyProperty = true;
 			}
 		}
-		else
-		{
-			// TODO: add support for more properties
 
-			Builder.AddWarningMessage(FString::Printf(
-				TEXT("Property is not supported and will be skipped. Context: %s"),
-				*GetLogContext(Property)));
-		}
-	}
-
-	if (!bHasParsedAnyProperty)
-	{
-		Builder.AddWarningMessage(FString::Printf(
-			TEXT("Binding has no supported properties and will be skipped. Context: %s"),
-			*GetLogContext(Binding)));
+		// TODO: add support for more properties
 	}
 
 	return bHasParsedAnyProperty;
