[GLTFExporter] refactor: rearrange some export options and add TODO comments

commit f1d2a0be7d07920247df8499c7868318366c0585
Date: Fri, 13 Nov 2020 18:37:54 +0100

    refactor: rearrange some export options and add TODO comments
    (cherry picked from commit 9126520c9c8b660bd367eddfea67fcf0b469c824​)

#skipci

diff --git a/Engine/Plugins/Enterprise/GLTFExporter/MigratedFromGit/Source/GLTFExporter/Private/Tasks/GLTFMeshTasks.cpp b/Engine/Plugins/Enterprise/GLTFExporter/MigratedFromGit/Source/GLTFExporter/Private/Tasks/GLTFMeshTasks.cpp
--- a/Engine/Plugins/Enterprise/GLTFExporter/MigratedFromGit/Source/GLTFExporter/Private/Tasks/GLTFMeshTasks.cpp
+++ b/Engine/Plugins/Enterprise/GLTFExporter/MigratedFromGit/Source/GLTFExporter/Private/Tasks/GLTFMeshTasks.cpp
@@ -45,11 +45,13 @@ void FGLTFStaticMeshTask::Complete()
 		JsonPrimitive.Attributes.Position = Builder.GetOrAddPositionAccessor(ConvertedSection, PositionBuffer);
 		if (JsonPrimitive.Attributes.Position == INDEX_NONE)
 		{
-			// TODO: report warning
+			// TODO: report warning?
 		}
 
 		if (Builder.ExportOptions->bExportVertexColors)
 		{
+			// TODO: report warning (about vertex colors may be represented differently in glTF)
+
 			JsonPrimitive.Attributes.Color0 = Builder.GetOrAddColorAccessor(ConvertedSection, ColorBuffer);
 		}
 
@@ -112,11 +114,13 @@ void FGLTFSkeletalMeshTask::Complete()
 		JsonPrimitive.Attributes.Position = Builder.GetOrAddPositionAccessor(ConvertedSection, PositionBuffer);
 		if (JsonPrimitive.Attributes.Position == INDEX_NONE)
 		{
-			// TODO: report warning
+			// TODO: report warning?
 		}
 
 		if (Builder.ExportOptions->bExportVertexColors)
 		{
+			// TODO: report warning (about vertex colors may be represented differently in glTF)
+
 			JsonPrimitive.Attributes.Color0 = Builder.GetOrAddColorAccessor(ConvertedSection, ColorBuffer);
 		}
 
diff --git a/Engine/Plugins/Enterprise/GLTFExporter/MigratedFromGit/Source/GLTFExporter/Public/GLTFExportOptions.h b/Engine/Plugins/Enterprise/GLTFExporter/MigratedFromGit/Source/GLTFExporter/Public/GLTFExportOptions.h
--- a/Engine/Plugins/Enterprise/GLTFExporter/MigratedFromGit/Source/GLTFExporter/Public/GLTFExportOptions.h
+++ b/Engine/Plugins/Enterprise/GLTFExporter/MigratedFromGit/Source/GLTFExporter/Public/GLTFExportOptions.h
@@ -87,21 +87,21 @@ public:
 	UPROPERTY(EditAnywhere, BlueprintReadWrite, Config, Category = Material, Meta = (EditCondition = "BakeMaterialInputs != EGLTFMaterialBakeMode::None"))
 	EGLTFMaterialBakeSizePOT DefaultMaterialBakeSize;
 
+	UPROPERTY(EditAnywhere, BlueprintReadWrite, Config, Category = Mesh, Meta = (ClampMin = "0"))
+	int32 DefaultLevelOfDetail;
+
 	UPROPERTY(EditAnywhere, BlueprintReadWrite, Config, Category = Mesh)
 	bool bExportVertexColors;
 
 	UPROPERTY(EditAnywhere, BlueprintReadWrite, Config, Category = Mesh)
-	bool bExportMeshQuantization;
+	bool bExportVertexSkinWeights;
 
-	UPROPERTY(EditAnywhere, BlueprintReadWrite, Config, Category = Mesh, Meta = (ClampMin = "0"))
-	int32 DefaultLevelOfDetail;
+	UPROPERTY(EditAnywhere, BlueprintReadWrite, Config, Category = Mesh)
+	bool bExportMeshQuantization;
 
 	UPROPERTY(EditAnywhere, BlueprintReadWrite, Config, Category = Animation)
 	bool bExportLevelSequences;
 
-	UPROPERTY(EditAnywhere, BlueprintReadWrite, Config, Category = Animation)
-	bool bExportVertexSkinWeights;
-
 	UPROPERTY(EditAnywhere, BlueprintReadWrite, Config, Category = Animation, Meta = (EditCondition = "bExportVertexSkinWeights"))
 	bool bExportAnimationSequences;
 
@@ -146,6 +146,9 @@ public:
 	UPROPERTY(EditAnywhere, BlueprintReadWrite, Config, Category = Scene, Meta = (EditCondition = "bExportCameras"))
 	bool bExportCameraControls;
 
+	UPROPERTY(EditAnywhere, BlueprintReadWrite, Config, Category = Scene)
+	bool bExportAnimationHotspots;
+
 	UPROPERTY(EditAnywhere, BlueprintReadWrite, Config, Category = Scene, Meta = (DisplayName = "Export HDRI Backdrops"))
 	bool bExportHDRIBackdrops;
 
@@ -155,9 +158,6 @@ public:
 	UPROPERTY(EditAnywhere, BlueprintReadWrite, Config, Category = Scene)
 	bool bExportVariantSets;
 
-	UPROPERTY(EditAnywhere, BlueprintReadWrite, Config, Category = Scene)
-	bool bExportAnimationHotspots;
-
 	UPROPERTY(EditAnywhere, BlueprintReadWrite, Config, Category = VariantSets, Meta = (EditCondition = "bExportVariantSets"))
 	EGLTFMaterialBakeMode ExportMaterialVariants;
 
