[GLTFExporter] fix: use correct compiler when baking custom outputs for materials

commit 8c77b9a2b9fd164fb3cfff5671de8696f112ca6a
Date: Fri, 13 Nov 2020 19:53:38 +0100

    fix: use correct compiler when baking custom outputs for materials
    (cherry picked from commit afb4ebe9d0b8554ae215bb6a7a224a3cd5fdbba0​)

#skipci

diff --git a/Engine/Plugins/Enterprise/GLTFExporter/MigratedFromGit/Source/GLTFMaterialBaking/Private/ExportMaterialProxy.h b/Engine/Plugins/Enterprise/GLTFExporter/MigratedFromGit/Source/GLTFMaterialBaking/Private/ExportMaterialProxy.h
--- a/Engine/Plugins/Enterprise/GLTFExporter/MigratedFromGit/Source/GLTFMaterialBaking/Private/ExportMaterialProxy.h
+++ b/Engine/Plugins/Enterprise/GLTFExporter/MigratedFromGit/Source/GLTFMaterialBaking/Private/ExportMaterialProxy.h
@@ -411,7 +411,7 @@ public:
 				if (BlendMode == BLEND_Opaque || BlendMode == BLEND_Masked)
 				{
 					return CompileNormalTransform(
-						Compiler,
+						&ProxyCompiler,
 						MaterialInterface->CompileProperty(&ProxyCompiler, PropertyToCompile, ForceCast_Exact_Replicate));
 				}
 				break;
@@ -419,7 +419,7 @@ public:
 				return CompileShadingModel(Compiler, MaterialInterface->CompileProperty(&ProxyCompiler, MP_ShadingModel));
 			case MP_CustomOutput:
 				 // NOTE: Currently we can assume input index is always 0, which it is for all custom outputs that are registered as material attributes
-				return CompileInputForCustomOutput(Compiler, 0, ForceCast_Exact_Replicate);
+				return CompileInputForCustomOutput(&ProxyCompiler, 0, ForceCast_Exact_Replicate);
 			default:
 				return Compiler->Constant(1.0f);
 			}
