[GLTFExporter] chore: add empty module GLTFMaterialAnalyzer to exporter plugin

commit e37cf20244ec940e4ce01ac6d448733b6aa726b3
Date: Wed, 11 Nov 2020 22:22:44 +0100

    chore: add empty module GLTFMaterialAnalyzer to exporter plugin
    (cherry picked from commit 2afa2d935c0de3acaf8202e04790395f7ced1556​)

#skipci

diff --git a/Engine/Plugins/Enterprise/GLTFExporter/MigratedFromGit/GLTFExporter.uplugin b/Engine/Plugins/Enterprise/GLTFExporter/MigratedFromGit/GLTFExporter.uplugin
--- a/Engine/Plugins/Enterprise/GLTFExporter/MigratedFromGit/GLTFExporter.uplugin
+++ b/Engine/Plugins/Enterprise/GLTFExporter/MigratedFromGit/GLTFExporter.uplugin
@@ -15,18 +15,23 @@
 	"IsBetaVersion": true,
 	"Modules": [
 		{
-			"Name": "GLTFMaterialBaking",
+			"Name": "GLTFExporter",
 			"Type": "Editor",
 			"LoadingPhase": "Default"
 		},
 		{
-			"Name": "GLTFExporter",
+			"Name": "GLTFExporterRuntime",
+			"Type": "Runtime",
+			"LoadingPhase": "Default"
+		},
+		{
+			"Name": "GLTFMaterialAnalyzer",
 			"Type": "Editor",
 			"LoadingPhase": "Default"
 		},
 		{
-			"Name": "GLTFExporterRuntime",
-			"Type": "Runtime",
+			"Name": "GLTFMaterialBaking",
+			"Type": "Editor",
 			"LoadingPhase": "Default"
 		}
 	],
diff --git a/Engine/Plugins/Enterprise/GLTFExporter/MigratedFromGit/Source/GLTFMaterialAnalyzer/GLTFMaterialAnalyzer.Build.cs b/Engine/Plugins/Enterprise/GLTFExporter/MigratedFromGit/Source/GLTFMaterialAnalyzer/GLTFMaterialAnalyzer.Build.cs
new file mode 100644
--- /dev/null
+++ b/Engine/Plugins/Enterprise/GLTFExporter/MigratedFromGit/Source/GLTFMaterialAnalyzer/GLTFMaterialAnalyzer.Build.cs
@@ -0,0 +1,53 @@
+// Copyright Epic Games, Inc. All Rights Reserved.
+
+using UnrealBuildTool;
+
+public class GLTFMaterialAnalyzer : ModuleRules
+{
+	public GLTFMaterialAnalyzer(ReadOnlyTargetRules Target) : base(Target)
+	{
+		PCHUsage = ModuleRules.PCHUsageMode.UseExplicitOrSharedPCHs;
+
+		PublicIncludePaths.AddRange(
+			new string[] {
+				// ... add public include paths required here ...
+			}
+			);
+
+
+		PrivateIncludePaths.AddRange(
+			new string[] {
+				// ... add other private include paths required here ...
+			}
+			);
+
+
+		PublicDependencyModuleNames.AddRange(
+			new string[]
+			{
+				"Core",
+				// ... add other public dependencies that you statically link with here ...
+			}
+			);
+
+
+		PrivateDependencyModuleNames.AddRange(
+			new string[]
+			{
+				"CoreUObject",
+				"Engine",
+				"Slate",
+				"SlateCore",
+				// ... add private dependencies that you statically link with here ...
+			}
+			);
+
+
+		DynamicallyLoadedModuleNames.AddRange(
+			new string[]
+			{
+				// ... add any modules that your module loads dynamically here ...
+			}
+			);
+	}
+}
diff --git a/Engine/Plugins/Enterprise/GLTFExporter/MigratedFromGit/Source/GLTFMaterialAnalyzer/Private/GLTFMaterialAnalyzerModule.cpp b/Engine/Plugins/Enterprise/GLTFExporter/MigratedFromGit/Source/GLTFMaterialAnalyzer/Private/GLTFMaterialAnalyzerModule.cpp
new file mode 100644
--- /dev/null
+++ b/Engine/Plugins/Enterprise/GLTFExporter/MigratedFromGit/Source/GLTFMaterialAnalyzer/Private/GLTFMaterialAnalyzerModule.cpp
@@ -0,0 +1,20 @@
+// Copyright Epic Games, Inc. All Rights Reserved.
+
+#include "GLTFMaterialAnalyzerModule.h"
+
+#define LOCTEXT_NAMESPACE "FGLTFMaterialAnalyzerModule"
+
+void FGLTFMaterialAnalyzerModule::StartupModule()
+{
+	// This code will execute after your module is loaded into memory; the exact timing is specified in the .uplugin file per-module
+}
+
+void FGLTFMaterialAnalyzerModule::ShutdownModule()
+{
+	// This function may be called during shutdown to clean up your module.  For modules that support dynamic reloading,
+	// we call this function before unloading the module.
+}
+
+#undef LOCTEXT_NAMESPACE
+
+IMPLEMENT_MODULE(FGLTFMaterialAnalyzerModule, GLTFMaterialAnalyzer)
diff --git a/Engine/Plugins/Enterprise/GLTFExporter/MigratedFromGit/Source/GLTFMaterialAnalyzer/Public/GLTFMaterialAnalyzerModule.h b/Engine/Plugins/Enterprise/GLTFExporter/MigratedFromGit/Source/GLTFMaterialAnalyzer/Public/GLTFMaterialAnalyzerModule.h
new file mode 100644
--- /dev/null
+++ b/Engine/Plugins/Enterprise/GLTFExporter/MigratedFromGit/Source/GLTFMaterialAnalyzer/Public/GLTFMaterialAnalyzerModule.h
@@ -0,0 +1,15 @@
+// Copyright Epic Games, Inc. All Rights Reserved.
+
+#pragma once
+
+#include "CoreMinimal.h"
+#include "Modules/ModuleManager.h"
+
+class FGLTFMaterialAnalyzerModule : public IModuleInterface
+{
+public:
+
+	/** IModuleInterface implementation */
+	virtual void StartupModule() override;
+	virtual void ShutdownModule() override;
+};
