[GLTFExporter] chore: simplify build script for GLTFMaterialAnalyzer and add private include path to access HLSLMaterialTranslator

commit 4f04c3e309a99dcfbb7c06b39733d8abaafe9b96
Date: Wed, 11 Nov 2020 22:27:34 +0100

    chore: simplify build script for GLTFMaterialAnalyzer and add private include path to access HLSLMaterialTranslator
    (cherry picked from commit ebccf8cbc972833141d82ad1712a844424d3dca6​)

#skipci

diff --git a/Engine/Plugins/Enterprise/GLTFExporter/MigratedFromGit/Source/GLTFMaterialAnalyzer/GLTFMaterialAnalyzer.Build.cs b/Engine/Plugins/Enterprise/GLTFExporter/MigratedFromGit/Source/GLTFMaterialAnalyzer/GLTFMaterialAnalyzer.Build.cs
--- a/Engine/Plugins/Enterprise/GLTFExporter/MigratedFromGit/Source/GLTFMaterialAnalyzer/GLTFMaterialAnalyzer.Build.cs
+++ b/Engine/Plugins/Enterprise/GLTFExporter/MigratedFromGit/Source/GLTFMaterialAnalyzer/GLTFMaterialAnalyzer.Build.cs
@@ -6,48 +6,17 @@ public class GLTFMaterialAnalyzer : ModuleRules
 {
 	public GLTFMaterialAnalyzer(ReadOnlyTargetRules Target) : base(Target)
 	{
-		PCHUsage = ModuleRules.PCHUsageMode.UseExplicitOrSharedPCHs;
-
-		PublicIncludePaths.AddRange(
-			new string[] {
-				// ... add public include paths required here ...
-			}
-			);
-
-
-		PrivateIncludePaths.AddRange(
-			new string[] {
-				// ... add other private include paths required here ...
-			}
-			);
-
-
-		PublicDependencyModuleNames.AddRange(
-			new string[]
-			{
-				"Core",
-				// ... add other public dependencies that you statically link with here ...
-			}
-			);
-
-
 		PrivateDependencyModuleNames.AddRange(
 			new string[]
 			{
+				"Core",
 				"CoreUObject",
 				"Engine",
-				"Slate",
-				"SlateCore",
-				// ... add private dependencies that you statically link with here ...
+				"RenderCore"
 			}
-			);
+		);
 
-
-		DynamicallyLoadedModuleNames.AddRange(
-			new string[]
-			{
-				// ... add any modules that your module loads dynamically here ...
-			}
-			);
+		// NOTE: ugly hack to access HLSLMaterialTranslator to analyze materials
+		PrivateIncludePaths.Add(EngineDirectory + "/Source/Runtime/Engine/Private");
 	}
 }
