[GLTFExporter] feat: implement GLTFMaterialAnalyzer with hack to workaround accessing FHLSLMaterialTranslator data

commit 1ced7381e5e72a477a1f4751b6bdfedea8caa1b7
Date: Wed, 11 Nov 2020 23:41:03 +0100

    feat: implement GLTFMaterialAnalyzer with hack to workaround accessing FHLSLMaterialTranslator data
    (cherry picked from commit 34f1a1decdf9a194a740d5b9521af669443a7a94​)

#skipci

diff --git a/Engine/Plugins/Enterprise/GLTFExporter/MigratedFromGit/Source/GLTFMaterialAnalyzer/Private/GLTFMaterialAnalyzer.cpp b/Engine/Plugins/Enterprise/GLTFExporter/MigratedFromGit/Source/GLTFMaterialAnalyzer/Private/GLTFMaterialAnalyzer.cpp
new file mode 100644
--- /dev/null
+++ b/Engine/Plugins/Enterprise/GLTFExporter/MigratedFromGit/Source/GLTFMaterialAnalyzer/Private/GLTFMaterialAnalyzer.cpp
@@ -0,0 +1,74 @@
+// Copyright Epic Games, Inc. All Rights Reserved.
+
+#include "GLTFMaterialAnalyzer.h"
+#include "Materials/HLSLMaterialTranslator.h"
+
+void UGLTFMaterialAnalyzer::AnalyzeMaterialProperty(const UMaterialInterface* InMaterial, EMaterialProperty InProperty)
+{
+	Material = const_cast<UMaterialInterface*>(InMaterial);
+
+	int32 NumTextureCoordinates;
+	bool RequiresVertexData;
+	UMaterialInterface::AnalyzeMaterialProperty(InProperty, NumTextureCoordinates, RequiresVertexData);
+}
+
+FMaterialResource* UGLTFMaterialAnalyzer::GetMaterialResource(ERHIFeatureLevel::Type InFeatureLevel, EMaterialQualityLevel::Type QualityLevel)
+{
+	return Material->GetMaterialResource(InFeatureLevel, QualityLevel);
+}
+
+int32 UGLTFMaterialAnalyzer::CompilePropertyEx(FMaterialCompiler* Compiler, const FGuid& AttributeID)
+{
+	class FTranslatorHack : public FHLSLMaterialTranslator
+	{
+		using FHLSLMaterialTranslator::FHLSLMaterialTranslator;
+
+		friend UGLTFMaterialAnalyzer;
+	};
+
+	const int32 Result = Material->CompilePropertyEx(Compiler, AttributeID);
+	FTranslatorHack* Translator = static_cast<FTranslatorHack*>(Compiler);
+
+	bUsesSceneDepth = Translator->bUsesSceneDepth;
+	bNeedsParticlePosition = Translator->bNeedsParticlePosition;
+	bNeedsParticleVelocity = Translator->bNeedsParticleVelocity;
+	bNeedsParticleTime = Translator->bNeedsParticleTime;
+	bUsesParticleMotionBlur = Translator->bUsesParticleMotionBlur;
+	bNeedsParticleRandom = Translator->bNeedsParticleRandom;
+	bUsesSphericalParticleOpacity = Translator->bUsesSphericalParticleOpacity;
+	bUsesParticleSubUVs = Translator->bUsesParticleSubUVs;
+	bUsesLightmapUVs = Translator->bUsesLightmapUVs;
+	bUsesAOMaterialMask = Translator->bUsesAOMaterialMask;
+	bUsesSpeedTree = Translator->bUsesSpeedTree;
+	bNeedsWorldPositionExcludingShaderOffsets = Translator->bNeedsWorldPositionExcludingShaderOffsets;
+	bNeedsParticleSize = Translator->bNeedsParticleSize;
+	bNeedsSceneTexturePostProcessInputs = Translator->bNeedsSceneTexturePostProcessInputs;
+	bUsesAtmosphericFog = Translator->bUsesAtmosphericFog;
+	bUsesSkyAtmosphere = Translator->bUsesSkyAtmosphere;
+	bUsesVertexColor = Translator->bUsesVertexColor;
+	bUsesParticleColor = Translator->bUsesParticleColor;
+	bUsesParticleLocalToWorld = Translator->bUsesParticleLocalToWorld;
+	bUsesParticleWorldToLocal = Translator->bUsesParticleWorldToLocal;
+	bUsesVertexPosition = Translator->bUsesVertexPosition;
+	bUsesTransformVector = Translator->bUsesTransformVector;
+	bCompilingPreviousFrame = Translator->bCompilingPreviousFrame;
+	bOutputsBasePassVelocities = Translator->bOutputsBasePassVelocities;
+	bUsesPixelDepthOffset = Translator->bUsesPixelDepthOffset;
+	bUsesWorldPositionOffset = Translator->bUsesWorldPositionOffset;
+	bUsesEmissiveColor = Translator->bUsesEmissiveColor;
+	bUsesDistanceCullFade = Translator->bUsesDistanceCullFade;
+	bIsFullyRough = Translator->bIsFullyRough;
+	bAllowCodeChunkGeneration = Translator->bAllowCodeChunkGeneration;
+	bUsesPerInstanceCustomData = Translator->bUsesPerInstanceCustomData;
+	AllocatedUserTexCoords = Translator->AllocatedUserTexCoords;
+	AllocatedUserVertexTexCoords = Translator->AllocatedUserVertexTexCoords;
+	DynamicParticleParameterMask = Translator->DynamicParticleParameterMask;
+	ShadingModelsFromCompilation = Translator->ShadingModelsFromCompilation;
+
+	return Result;
+}
+
+bool UGLTFMaterialAnalyzer::IsPropertyActive(EMaterialProperty InProperty) const
+{
+	return Material->IsPropertyActive(InProperty);
+}
diff --git a/Engine/Plugins/Enterprise/GLTFExporter/MigratedFromGit/Source/GLTFMaterialAnalyzer/Public/GLTFMaterialAnalyzer.h b/Engine/Plugins/Enterprise/GLTFExporter/MigratedFromGit/Source/GLTFMaterialAnalyzer/Public/GLTFMaterialAnalyzer.h
new file mode 100644
--- /dev/null
+++ b/Engine/Plugins/Enterprise/GLTFExporter/MigratedFromGit/Source/GLTFMaterialAnalyzer/Public/GLTFMaterialAnalyzer.h
@@ -0,0 +1,98 @@
+// Copyright Epic Games, Inc. All Rights Reserved.
+
+#pragma once
+
+#include "Materials/MaterialInterface.h"
+#include "GLTFMaterialAnalyzer.generated.h"
+
+UCLASS()
+class GLTFMATERIALANALYZER_API UGLTFMaterialAnalyzer : public UMaterialInterface
+{
+	GENERATED_BODY()
+
+public:
+
+	void AnalyzeMaterialProperty(const UMaterialInterface* InMaterial, EMaterialProperty InProperty);
+
+	/** Whether the compiled material uses scene depth. */
+	uint32 bUsesSceneDepth : 1;
+	/** true if the material needs particle position. */
+	uint32 bNeedsParticlePosition : 1;
+	/** true if the material needs particle velocity. */
+	uint32 bNeedsParticleVelocity : 1;
+	/** true if the material needs particle relative time. */
+	uint32 bNeedsParticleTime : 1;
+	/** true if the material uses particle motion blur. */
+	uint32 bUsesParticleMotionBlur : 1;
+	/** true if the material needs particle random value. */
+	uint32 bNeedsParticleRandom : 1;
+	/** true if the material uses spherical particle opacity. */
+	uint32 bUsesSphericalParticleOpacity : 1;
+	/** true if the material uses particle sub uvs. */
+	uint32 bUsesParticleSubUVs : 1;
+	/** Boolean indicating using LightmapUvs */
+	uint32 bUsesLightmapUVs : 1;
+	/** Whether the material uses AO Material Mask */
+	uint32 bUsesAOMaterialMask : 1;
+	/** true if needs SpeedTree code */
+	uint32 bUsesSpeedTree : 1;
+	/** Boolean indicating the material uses worldspace position without shader offsets applied */
+	uint32 bNeedsWorldPositionExcludingShaderOffsets : 1;
+	/** true if the material needs particle size. */
+	uint32 bNeedsParticleSize : 1;
+	/** true if any scene texture expressions are reading from post process inputs */
+	uint32 bNeedsSceneTexturePostProcessInputs : 1;
+	/** true if any atmospheric fog expressions are used */
+	uint32 bUsesAtmosphericFog : 1;
+	/** true if any SkyAtmosphere expressions are used */
+	uint32 bUsesSkyAtmosphere : 1;
+	/** true if the material reads vertex color in the pixel shader. */
+	uint32 bUsesVertexColor : 1;
+	/** true if the material reads particle color in the pixel shader. */
+	uint32 bUsesParticleColor : 1;
+	/** true if the material reads mesh particle local to world in the pixel shader. */
+	uint32 bUsesParticleLocalToWorld : 1;
+	/** true if the material reads mesh particle world to local in the pixel shader. */
+	uint32 bUsesParticleWorldToLocal : 1;
+
+	/** true if the material uses any type of vertex position */
+	uint32 bUsesVertexPosition : 1;
+
+	uint32 bUsesTransformVector : 1;
+	// True if the current property requires last frame's information
+	uint32 bCompilingPreviousFrame : 1;
+	/** True if material will output accurate velocities during base pass rendering. */
+	uint32 bOutputsBasePassVelocities : 1;
+	uint32 bUsesPixelDepthOffset : 1;
+	uint32 bUsesWorldPositionOffset : 1;
+	uint32 bUsesEmissiveColor : 1;
+	uint32 bUsesDistanceCullFade : 1;
+	/** true if the Roughness input evaluates to a constant 1.0 */
+	uint32 bIsFullyRough : 1;
+	/** true if allowed to generate code chunks. Translator operates in two phases; generate all code chunks & query meta data based on generated code chunks. */
+	uint32 bAllowCodeChunkGeneration : 1;
+
+	/** True if this material reads any per-instance custom data */
+	uint32 bUsesPerInstanceCustomData : 1;
+
+	/** Tracks the texture coordinates used by this material. */
+	TBitArray<> AllocatedUserTexCoords;
+	/** Tracks the texture coordinates used by the vertex shader in this material. */
+	TBitArray<> AllocatedUserVertexTexCoords;
+
+	uint32 DynamicParticleParameterMask;
+
+	/** Will contain all the shading models picked up from the material expression graph */
+	FMaterialShadingModelField ShadingModelsFromCompilation;
+
+private:
+
+	virtual FMaterialResource* GetMaterialResource(ERHIFeatureLevel::Type InFeatureLevel, EMaterialQualityLevel::Type QualityLevel) override;
+
+	virtual int32 CompilePropertyEx(FMaterialCompiler* Compiler, const FGuid& AttributeID) override;
+
+	virtual bool IsPropertyActive(EMaterialProperty InProperty) const override;
+
+	// ReSharper disable once CppUE4ProbableMemoryIssuesWithUObject
+	UMaterialInterface* Material;
+};
