[GLTFExporter] refactor: use GLTFMaterialAnalyzer in GLTFMaterialUtility

commit ed25d17451bab0a691ca30edcbcc0885bd1e0ec1
Date: Wed, 11 Nov 2020 23:41:32 +0100

    refactor: use GLTFMaterialAnalyzer in GLTFMaterialUtility
    (cherry picked from commit c1caf25042a7c825a67b4aa5cd67c7b9817311e1​)

#skipci

diff --git a/Engine/Plugins/Enterprise/GLTFExporter/MigratedFromGit/Source/GLTFExporter/GLTFExporter.Build.cs b/Engine/Plugins/Enterprise/GLTFExporter/MigratedFromGit/Source/GLTFExporter/GLTFExporter.Build.cs
--- a/Engine/Plugins/Enterprise/GLTFExporter/MigratedFromGit/Source/GLTFExporter/GLTFExporter.Build.cs
+++ b/Engine/Plugins/Enterprise/GLTFExporter/MigratedFromGit/Source/GLTFExporter/GLTFExporter.Build.cs
@@ -31,11 +31,12 @@ namespace UnrealBuildTool.Rules
 					"LevelSequence",
 					"MovieScene",
 					"MovieSceneTracks",
-					"GLTFMaterialBaking",
 					"VariantManagerContent",
-					"GLTFExporterRuntime",
 					"MeshMergeUtilities",
-					"zlib"
+					"zlib",
+					"GLTFExporterRuntime",
+					"GLTFMaterialBaking",
+					"GLTFMaterialAnalyzer"
 				}
 				);
 		}
diff --git a/Engine/Plugins/Enterprise/GLTFExporter/MigratedFromGit/Source/GLTFExporter/Private/Converters/GLTFMaterialUtility.cpp b/Engine/Plugins/Enterprise/GLTFExporter/MigratedFromGit/Source/GLTFExporter/Private/Converters/GLTFMaterialUtility.cpp
--- a/Engine/Plugins/Enterprise/GLTFExporter/MigratedFromGit/Source/GLTFExporter/Private/Converters/GLTFMaterialUtility.cpp
+++ b/Engine/Plugins/Enterprise/GLTFExporter/MigratedFromGit/Source/GLTFExporter/Private/Converters/GLTFMaterialUtility.cpp
@@ -3,6 +3,7 @@
 #include "Converters/GLTFMaterialUtility.h"
 #include "Converters/GLTFTextureUtility.h"
 #include "Converters/GLTFNameUtility.h"
+#include "GLTFMaterialAnalyzer.h"
 #include "Engine/TextureRenderTarget2D.h"
 #include "CanvasItem.h"
 #include "CanvasTypes.h"
@@ -521,7 +522,7 @@ bool FGLTFMaterialUtility::NeedsMeshData(const UMaterialInterface* Material)
 
 	for (const EMaterialProperty Property: Properties)
 	{
-		const_cast<UMaterialInterface*>(Material)->AnalyzeMaterialProperty(Property, NumTextureCoordinates, bPropertyRequiresVertexData);
+		AnalyzeMaterialProperty(Material, Property, NumTextureCoordinates, bPropertyRequiresVertexData);
 		bRequiresVertexData |= bPropertyRequiresVertexData;
 	}
 
@@ -541,6 +542,15 @@ bool FGLTFMaterialUtility::NeedsMeshData(const TArray<const UMaterialInterface*>
 	return false;
 }
 
+void FGLTFMaterialUtility::AnalyzeMaterialProperty(const UMaterialInterface* InMaterial, EMaterialProperty InProperty, int32& OutNumTextureCoordinates, bool& bOutRequiresVertexData)
+{
+	UGLTFMaterialAnalyzer* Analyzer = NewObject<UGLTFMaterialAnalyzer>();
+	Analyzer->AnalyzeMaterialProperty(InMaterial, InProperty);
+
+	OutNumTextureCoordinates = Analyzer->AllocatedUserTexCoords.FindLast(true) + 1;
+	bOutRequiresVertexData = Analyzer->bUsesVertexColor || Analyzer->bUsesTransformVector || Analyzer->bNeedsWorldPositionExcludingShaderOffsets || Analyzer->bUsesAOMaterialMask || Analyzer->bUsesVertexPosition;
+}
+
 const UMaterialInterface* FGLTFMaterialUtility::GetInterface(const UMaterialInterface* Material)
 {
 	return Material;
diff --git a/Engine/Plugins/Enterprise/GLTFExporter/MigratedFromGit/Source/GLTFExporter/Private/Converters/GLTFMaterialUtility.h b/Engine/Plugins/Enterprise/GLTFExporter/MigratedFromGit/Source/GLTFExporter/Private/Converters/GLTFMaterialUtility.h
--- a/Engine/Plugins/Enterprise/GLTFExporter/MigratedFromGit/Source/GLTFExporter/Private/Converters/GLTFMaterialUtility.h
+++ b/Engine/Plugins/Enterprise/GLTFExporter/MigratedFromGit/Source/GLTFExporter/Private/Converters/GLTFMaterialUtility.h
@@ -74,6 +74,8 @@ struct FGLTFMaterialUtility
 	static bool NeedsMeshData(const UMaterialInterface* Material);
 	static bool NeedsMeshData(const TArray<const UMaterialInterface*>& Materials);
 
+	static void AnalyzeMaterialProperty(const UMaterialInterface* Material, EMaterialProperty InProperty, int32& OutNumTextureCoordinates, bool& bOutRequiresVertexData);
+
 	static const UMaterialInterface* GetInterface(const UMaterialInterface* Material);
 	static const UMaterialInterface* GetInterface(const FStaticMaterial& Material);
 	static const UMaterialInterface* GetInterface(const FSkeletalMaterial& Material);
