[GLTFExporter] refactor: use separate struct for the results of AnalyzeMaterialProperty

commit 25cae31e839ca422506e133d309ea594d46f4465
Date: Thu, 12 Nov 2020 14:22:02 +0100

    refactor: use separate struct for the results of AnalyzeMaterialProperty
    (cherry picked from commit 9b58475ee7239f8c2d9ae41c772cef87b371e8b0​)

#skipci

diff --git a/Engine/Plugins/Enterprise/GLTFExporter/MigratedFromGit/Source/GLTFExporter/Private/Converters/GLTFMaterialUtility.cpp b/Engine/Plugins/Enterprise/GLTFExporter/MigratedFromGit/Source/GLTFExporter/Private/Converters/GLTFMaterialUtility.cpp
--- a/Engine/Plugins/Enterprise/GLTFExporter/MigratedFromGit/Source/GLTFExporter/Private/Converters/GLTFMaterialUtility.cpp
+++ b/Engine/Plugins/Enterprise/GLTFExporter/MigratedFromGit/Source/GLTFExporter/Private/Converters/GLTFMaterialUtility.cpp
@@ -3,6 +3,7 @@
 #include "Converters/GLTFMaterialUtility.h"
 #include "Converters/GLTFTextureUtility.h"
 #include "Converters/GLTFNameUtility.h"
+#include "GLTFMaterialStatistics.h"
 #include "GLTFMaterialAnalyzer.h"
 #include "Engine/TextureRenderTarget2D.h"
 #include "CanvasItem.h"
@@ -516,14 +517,12 @@ bool FGLTFMaterialUtility::NeedsMeshData(const UMaterialInterface* Material)
 	};
 
 	bool bRequiresVertexData = false;
-
-	int32 NumTextureCoordinates;
-	bool bPropertyRequiresVertexData;
+	FGLTFMaterialStatistics MaterialStatistics;
 
 	for (const EMaterialProperty Property: Properties)
 	{
-		AnalyzeMaterialProperty(Material, Property, NumTextureCoordinates, bPropertyRequiresVertexData);
-		bRequiresVertexData |= bPropertyRequiresVertexData;
+		AnalyzeMaterialProperty(Material, Property, MaterialStatistics);
+		bRequiresVertexData |= MaterialStatistics.bRequiresVertexData;
 	}
 
 	return bRequiresVertexData;
@@ -542,13 +541,10 @@ bool FGLTFMaterialUtility::NeedsMeshData(const TArray<const UMaterialInterface*>
 	return false;
 }
 
-void FGLTFMaterialUtility::AnalyzeMaterialProperty(const UMaterialInterface* InMaterial, EMaterialProperty InProperty, int32& OutNumTextureCoordinates, bool& bOutRequiresVertexData)
+void FGLTFMaterialUtility::AnalyzeMaterialProperty(const UMaterialInterface* InMaterial, EMaterialProperty InProperty, FGLTFMaterialStatistics& OutMaterialStatistics)
 {
 	UGLTFMaterialAnalyzer* Analyzer = NewObject<UGLTFMaterialAnalyzer>();
-	Analyzer->AnalyzeMaterialProperty(InMaterial, InProperty);
-
-	OutNumTextureCoordinates = Analyzer->AllocatedUserTexCoords.FindLast(true) + 1;
-	bOutRequiresVertexData = Analyzer->bUsesVertexColor || Analyzer->bUsesTransformVector || Analyzer->bNeedsWorldPositionExcludingShaderOffsets || Analyzer->bUsesAOMaterialMask || Analyzer->bUsesVertexPosition;
+	Analyzer->AnalyzeMaterialProperty(InMaterial, InProperty, OutMaterialStatistics);
 }
 
 const UMaterialInterface* FGLTFMaterialUtility::GetInterface(const UMaterialInterface* Material)
diff --git a/Engine/Plugins/Enterprise/GLTFExporter/MigratedFromGit/Source/GLTFExporter/Private/Converters/GLTFMaterialUtility.h b/Engine/Plugins/Enterprise/GLTFExporter/MigratedFromGit/Source/GLTFExporter/Private/Converters/GLTFMaterialUtility.h
--- a/Engine/Plugins/Enterprise/GLTFExporter/MigratedFromGit/Source/GLTFExporter/Private/Converters/GLTFMaterialUtility.h
+++ b/Engine/Plugins/Enterprise/GLTFExporter/MigratedFromGit/Source/GLTFExporter/Private/Converters/GLTFMaterialUtility.h
@@ -5,6 +5,8 @@
 #include "Builders/GLTFContainerBuilder.h"
 #include "Engine/Texture2D.h"
 
+struct FGLTFMaterialStatistics;
+
 struct FGLTFTextureCombineSource
 {
 	constexpr FORCEINLINE FGLTFTextureCombineSource(const UTexture2D* Texture, FLinearColor TintColor = { 1.0f, 1.0f, 1.0f, 1.0f }, ESimpleElementBlendMode BlendMode = SE_BLEND_Additive)
@@ -74,7 +76,7 @@ struct FGLTFMaterialUtility
 	static bool NeedsMeshData(const UMaterialInterface* Material);
 	static bool NeedsMeshData(const TArray<const UMaterialInterface*>& Materials);
 
-	static void AnalyzeMaterialProperty(const UMaterialInterface* Material, EMaterialProperty InProperty, int32& OutNumTextureCoordinates, bool& bOutRequiresVertexData);
+	static void AnalyzeMaterialProperty(const UMaterialInterface* Material, EMaterialProperty InProperty, FGLTFMaterialStatistics& OutMaterialStatistics);
 
 	static const UMaterialInterface* GetInterface(const UMaterialInterface* Material);
 	static const UMaterialInterface* GetInterface(const FStaticMaterial& Material);
diff --git a/Engine/Plugins/Enterprise/GLTFExporter/MigratedFromGit/Source/GLTFMaterialAnalyzer/Private/GLTFMaterialAnalyzer.cpp b/Engine/Plugins/Enterprise/GLTFExporter/MigratedFromGit/Source/GLTFMaterialAnalyzer/Private/GLTFMaterialAnalyzer.cpp
--- a/Engine/Plugins/Enterprise/GLTFExporter/MigratedFromGit/Source/GLTFMaterialAnalyzer/Private/GLTFMaterialAnalyzer.cpp
+++ b/Engine/Plugins/Enterprise/GLTFExporter/MigratedFromGit/Source/GLTFMaterialAnalyzer/Private/GLTFMaterialAnalyzer.cpp
@@ -1,15 +1,20 @@
 // Copyright Epic Games, Inc. All Rights Reserved.
 
 #include "GLTFMaterialAnalyzer.h"
+#include "GLTFMaterialStatistics.h"
 #include "Materials/HLSLMaterialTranslator.h"
 
-void UGLTFMaterialAnalyzer::AnalyzeMaterialProperty(const UMaterialInterface* InMaterial, EMaterialProperty InProperty)
+void UGLTFMaterialAnalyzer::AnalyzeMaterialProperty(const UMaterialInterface* InMaterial, EMaterialProperty InProperty, FGLTFMaterialStatistics& OutMaterialStatistics)
 {
 	Material = const_cast<UMaterialInterface*>(InMaterial);
+	MaterialStatistics = &OutMaterialStatistics;
 
 	int32 NumTextureCoordinates;
 	bool RequiresVertexData;
 	UMaterialInterface::AnalyzeMaterialProperty(InProperty, NumTextureCoordinates, RequiresVertexData);
+
+	Material = nullptr;
+	MaterialStatistics = nullptr;
 }
 
 FMaterialResource* UGLTFMaterialAnalyzer::GetMaterialResource(ERHIFeatureLevel::Type InFeatureLevel, EMaterialQualityLevel::Type QualityLevel)
@@ -29,41 +34,49 @@ int32 UGLTFMaterialAnalyzer::CompilePropertyEx(FMaterialCompiler* Compiler, cons
 	const int32 Result = Material->CompilePropertyEx(Compiler, AttributeID);
 	FTranslatorHack* Translator = static_cast<FTranslatorHack*>(Compiler);
 
-	bUsesSceneDepth = Translator->bUsesSceneDepth;
-	bNeedsParticlePosition = Translator->bNeedsParticlePosition;
-	bNeedsParticleVelocity = Translator->bNeedsParticleVelocity;
-	bNeedsParticleTime = Translator->bNeedsParticleTime;
-	bUsesParticleMotionBlur = Translator->bUsesParticleMotionBlur;
-	bNeedsParticleRandom = Translator->bNeedsParticleRandom;
-	bUsesSphericalParticleOpacity = Translator->bUsesSphericalParticleOpacity;
-	bUsesParticleSubUVs = Translator->bUsesParticleSubUVs;
-	bUsesLightmapUVs = Translator->bUsesLightmapUVs;
-	bUsesAOMaterialMask = Translator->bUsesAOMaterialMask;
-	bUsesSpeedTree = Translator->bUsesSpeedTree;
-	bNeedsWorldPositionExcludingShaderOffsets = Translator->bNeedsWorldPositionExcludingShaderOffsets;
-	bNeedsParticleSize = Translator->bNeedsParticleSize;
-	bNeedsSceneTexturePostProcessInputs = Translator->bNeedsSceneTexturePostProcessInputs;
-	bUsesAtmosphericFog = Translator->bUsesAtmosphericFog;
-	bUsesSkyAtmosphere = Translator->bUsesSkyAtmosphere;
-	bUsesVertexColor = Translator->bUsesVertexColor;
-	bUsesParticleColor = Translator->bUsesParticleColor;
-	bUsesParticleLocalToWorld = Translator->bUsesParticleLocalToWorld;
-	bUsesParticleWorldToLocal = Translator->bUsesParticleWorldToLocal;
-	bUsesVertexPosition = Translator->bUsesVertexPosition;
-	bUsesTransformVector = Translator->bUsesTransformVector;
-	bCompilingPreviousFrame = Translator->bCompilingPreviousFrame;
-	bOutputsBasePassVelocities = Translator->bOutputsBasePassVelocities;
-	bUsesPixelDepthOffset = Translator->bUsesPixelDepthOffset;
-	bUsesWorldPositionOffset = Translator->bUsesWorldPositionOffset;
-	bUsesEmissiveColor = Translator->bUsesEmissiveColor;
-	bUsesDistanceCullFade = Translator->bUsesDistanceCullFade;
-	bIsFullyRough = Translator->bIsFullyRough;
-	bAllowCodeChunkGeneration = Translator->bAllowCodeChunkGeneration;
-	bUsesPerInstanceCustomData = Translator->bUsesPerInstanceCustomData;
-	AllocatedUserTexCoords = Translator->AllocatedUserTexCoords;
-	AllocatedUserVertexTexCoords = Translator->AllocatedUserVertexTexCoords;
-	DynamicParticleParameterMask = Translator->DynamicParticleParameterMask;
-	ShadingModelsFromCompilation = Translator->ShadingModelsFromCompilation;
+	MaterialStatistics->bUsesSceneDepth = Translator->bUsesSceneDepth;
+	MaterialStatistics->bNeedsParticlePosition = Translator->bNeedsParticlePosition;
+	MaterialStatistics->bNeedsParticleVelocity = Translator->bNeedsParticleVelocity;
+	MaterialStatistics->bNeedsParticleTime = Translator->bNeedsParticleTime;
+	MaterialStatistics->bUsesParticleMotionBlur = Translator->bUsesParticleMotionBlur;
+	MaterialStatistics->bNeedsParticleRandom = Translator->bNeedsParticleRandom;
+	MaterialStatistics->bUsesSphericalParticleOpacity = Translator->bUsesSphericalParticleOpacity;
+	MaterialStatistics->bUsesParticleSubUVs = Translator->bUsesParticleSubUVs;
+	MaterialStatistics->bUsesLightmapUVs = Translator->bUsesLightmapUVs;
+	MaterialStatistics->bUsesAOMaterialMask = Translator->bUsesAOMaterialMask;
+	MaterialStatistics->bUsesSpeedTree = Translator->bUsesSpeedTree;
+	MaterialStatistics->bNeedsWorldPositionExcludingShaderOffsets = Translator->bNeedsWorldPositionExcludingShaderOffsets;
+	MaterialStatistics->bNeedsParticleSize = Translator->bNeedsParticleSize;
+	MaterialStatistics->bNeedsSceneTexturePostProcessInputs = Translator->bNeedsSceneTexturePostProcessInputs;
+	MaterialStatistics->bUsesAtmosphericFog = Translator->bUsesAtmosphericFog;
+	MaterialStatistics->bUsesSkyAtmosphere = Translator->bUsesSkyAtmosphere;
+	MaterialStatistics->bUsesVertexColor = Translator->bUsesVertexColor;
+	MaterialStatistics->bUsesParticleColor = Translator->bUsesParticleColor;
+	MaterialStatistics->bUsesParticleLocalToWorld = Translator->bUsesParticleLocalToWorld;
+	MaterialStatistics->bUsesParticleWorldToLocal = Translator->bUsesParticleWorldToLocal;
+	MaterialStatistics->bUsesVertexPosition = Translator->bUsesVertexPosition;
+	MaterialStatistics->bUsesTransformVector = Translator->bUsesTransformVector;
+	MaterialStatistics->bCompilingPreviousFrame = Translator->bCompilingPreviousFrame;
+	MaterialStatistics->bOutputsBasePassVelocities = Translator->bOutputsBasePassVelocities;
+	MaterialStatistics->bUsesPixelDepthOffset = Translator->bUsesPixelDepthOffset;
+	MaterialStatistics->bUsesWorldPositionOffset = Translator->bUsesWorldPositionOffset;
+	MaterialStatistics->bUsesEmissiveColor = Translator->bUsesEmissiveColor;
+	MaterialStatistics->bUsesDistanceCullFade = Translator->bUsesDistanceCullFade;
+	MaterialStatistics->bIsFullyRough = Translator->bIsFullyRough;
+	MaterialStatistics->bAllowCodeChunkGeneration = Translator->bAllowCodeChunkGeneration;
+	MaterialStatistics->bUsesPerInstanceCustomData = Translator->bUsesPerInstanceCustomData;
+	MaterialStatistics->AllocatedUserTexCoords = Translator->AllocatedUserTexCoords;
+	MaterialStatistics->AllocatedUserVertexTexCoords = Translator->AllocatedUserVertexTexCoords;
+	MaterialStatistics->DynamicParticleParameterMask = Translator->DynamicParticleParameterMask;
+	MaterialStatistics->ShadingModelsFromCompilation = Translator->ShadingModelsFromCompilation;
+
+	// TODO: investigate if we need to check more conditions to determine that vertex data is required
+	MaterialStatistics->bRequiresVertexData =
+		Translator->bUsesVertexColor ||
+		Translator->bUsesTransformVector ||
+		Translator->bNeedsWorldPositionExcludingShaderOffsets ||
+		Translator->bUsesAOMaterialMask ||
+		Translator->bUsesVertexPosition;
 
 	return Result;
 }
diff --git a/Engine/Plugins/Enterprise/GLTFExporter/MigratedFromGit/Source/GLTFMaterialAnalyzer/Public/GLTFMaterialAnalyzer.h b/Engine/Plugins/Enterprise/GLTFExporter/MigratedFromGit/Source/GLTFMaterialAnalyzer/Public/GLTFMaterialAnalyzer.h
--- a/Engine/Plugins/Enterprise/GLTFExporter/MigratedFromGit/Source/GLTFMaterialAnalyzer/Public/GLTFMaterialAnalyzer.h
+++ b/Engine/Plugins/Enterprise/GLTFExporter/MigratedFromGit/Source/GLTFMaterialAnalyzer/Public/GLTFMaterialAnalyzer.h
@@ -5,6 +5,8 @@
 #include "Materials/MaterialInterface.h"
 #include "GLTFMaterialAnalyzer.generated.h"
 
+struct FGLTFMaterialStatistics;
+
 UCLASS()
 class GLTFMATERIALANALYZER_API UGLTFMaterialAnalyzer : public UMaterialInterface
 {
@@ -12,78 +14,7 @@ class GLTFMATERIALANALYZER_API UGLTFMaterialAnalyzer : public UMaterialInterface
 
 public:
 
-	void AnalyzeMaterialProperty(const UMaterialInterface* InMaterial, EMaterialProperty InProperty);
-
-	/** Whether the compiled material uses scene depth. */
-	uint32 bUsesSceneDepth : 1;
-	/** true if the material needs particle position. */
-	uint32 bNeedsParticlePosition : 1;
-	/** true if the material needs particle velocity. */
-	uint32 bNeedsParticleVelocity : 1;
-	/** true if the material needs particle relative time. */
-	uint32 bNeedsParticleTime : 1;
-	/** true if the material uses particle motion blur. */
-	uint32 bUsesParticleMotionBlur : 1;
-	/** true if the material needs particle random value. */
-	uint32 bNeedsParticleRandom : 1;
-	/** true if the material uses spherical particle opacity. */
-	uint32 bUsesSphericalParticleOpacity : 1;
-	/** true if the material uses particle sub uvs. */
-	uint32 bUsesParticleSubUVs : 1;
-	/** Boolean indicating using LightmapUvs */
-	uint32 bUsesLightmapUVs : 1;
-	/** Whether the material uses AO Material Mask */
-	uint32 bUsesAOMaterialMask : 1;
-	/** true if needs SpeedTree code */
-	uint32 bUsesSpeedTree : 1;
-	/** Boolean indicating the material uses worldspace position without shader offsets applied */
-	uint32 bNeedsWorldPositionExcludingShaderOffsets : 1;
-	/** true if the material needs particle size. */
-	uint32 bNeedsParticleSize : 1;
-	/** true if any scene texture expressions are reading from post process inputs */
-	uint32 bNeedsSceneTexturePostProcessInputs : 1;
-	/** true if any atmospheric fog expressions are used */
-	uint32 bUsesAtmosphericFog : 1;
-	/** true if any SkyAtmosphere expressions are used */
-	uint32 bUsesSkyAtmosphere : 1;
-	/** true if the material reads vertex color in the pixel shader. */
-	uint32 bUsesVertexColor : 1;
-	/** true if the material reads particle color in the pixel shader. */
-	uint32 bUsesParticleColor : 1;
-	/** true if the material reads mesh particle local to world in the pixel shader. */
-	uint32 bUsesParticleLocalToWorld : 1;
-	/** true if the material reads mesh particle world to local in the pixel shader. */
-	uint32 bUsesParticleWorldToLocal : 1;
-
-	/** true if the material uses any type of vertex position */
-	uint32 bUsesVertexPosition : 1;
-
-	uint32 bUsesTransformVector : 1;
-	// True if the current property requires last frame's information
-	uint32 bCompilingPreviousFrame : 1;
-	/** True if material will output accurate velocities during base pass rendering. */
-	uint32 bOutputsBasePassVelocities : 1;
-	uint32 bUsesPixelDepthOffset : 1;
-	uint32 bUsesWorldPositionOffset : 1;
-	uint32 bUsesEmissiveColor : 1;
-	uint32 bUsesDistanceCullFade : 1;
-	/** true if the Roughness input evaluates to a constant 1.0 */
-	uint32 bIsFullyRough : 1;
-	/** true if allowed to generate code chunks. Translator operates in two phases; generate all code chunks & query meta data based on generated code chunks. */
-	uint32 bAllowCodeChunkGeneration : 1;
-
-	/** True if this material reads any per-instance custom data */
-	uint32 bUsesPerInstanceCustomData : 1;
-
-	/** Tracks the texture coordinates used by this material. */
-	TBitArray<> AllocatedUserTexCoords;
-	/** Tracks the texture coordinates used by the vertex shader in this material. */
-	TBitArray<> AllocatedUserVertexTexCoords;
-
-	uint32 DynamicParticleParameterMask;
-
-	/** Will contain all the shading models picked up from the material expression graph */
-	FMaterialShadingModelField ShadingModelsFromCompilation;
+	void AnalyzeMaterialProperty(const UMaterialInterface* InMaterial, EMaterialProperty InProperty, FGLTFMaterialStatistics& OutMaterialStatistics);
 
 private:
 
@@ -95,4 +26,6 @@ private:
 
 	// ReSharper disable once CppUE4ProbableMemoryIssuesWithUObject
 	UMaterialInterface* Material;
+
+	FGLTFMaterialStatistics* MaterialStatistics;
 };
diff --git a/Engine/Plugins/Enterprise/GLTFExporter/MigratedFromGit/Source/GLTFMaterialAnalyzer/Public/GLTFMaterialStatistics.h b/Engine/Plugins/Enterprise/GLTFExporter/MigratedFromGit/Source/GLTFMaterialAnalyzer/Public/GLTFMaterialStatistics.h
new file mode 100644
--- /dev/null
+++ b/Engine/Plugins/Enterprise/GLTFExporter/MigratedFromGit/Source/GLTFMaterialAnalyzer/Public/GLTFMaterialStatistics.h
@@ -0,0 +1,80 @@
+// Copyright Epic Games, Inc. All Rights Reserved.
+
+#pragma once
+
+struct FGLTFMaterialStatistics
+{
+	/** Whether the compiled material uses scene depth. */
+	uint32 bUsesSceneDepth : 1;
+	/** true if the material needs particle position. */
+	uint32 bNeedsParticlePosition : 1;
+	/** true if the material needs particle velocity. */
+	uint32 bNeedsParticleVelocity : 1;
+	/** true if the material needs particle relative time. */
+	uint32 bNeedsParticleTime : 1;
+	/** true if the material uses particle motion blur. */
+	uint32 bUsesParticleMotionBlur : 1;
+	/** true if the material needs particle random value. */
+	uint32 bNeedsParticleRandom : 1;
+	/** true if the material uses spherical particle opacity. */
+	uint32 bUsesSphericalParticleOpacity : 1;
+	/** true if the material uses particle sub uvs. */
+	uint32 bUsesParticleSubUVs : 1;
+	/** Boolean indicating using LightmapUvs */
+	uint32 bUsesLightmapUVs : 1;
+	/** Whether the material uses AO Material Mask */
+	uint32 bUsesAOMaterialMask : 1;
+	/** true if needs SpeedTree code */
+	uint32 bUsesSpeedTree : 1;
+	/** Boolean indicating the material uses worldspace position without shader offsets applied */
+	uint32 bNeedsWorldPositionExcludingShaderOffsets : 1;
+	/** true if the material needs particle size. */
+	uint32 bNeedsParticleSize : 1;
+	/** true if any scene texture expressions are reading from post process inputs */
+	uint32 bNeedsSceneTexturePostProcessInputs : 1;
+	/** true if any atmospheric fog expressions are used */
+	uint32 bUsesAtmosphericFog : 1;
+	/** true if any SkyAtmosphere expressions are used */
+	uint32 bUsesSkyAtmosphere : 1;
+	/** true if the material reads vertex color in the pixel shader. */
+	uint32 bUsesVertexColor : 1;
+	/** true if the material reads particle color in the pixel shader. */
+	uint32 bUsesParticleColor : 1;
+	/** true if the material reads mesh particle local to world in the pixel shader. */
+	uint32 bUsesParticleLocalToWorld : 1;
+	/** true if the material reads mesh particle world to local in the pixel shader. */
+	uint32 bUsesParticleWorldToLocal : 1;
+
+	/** true if the material uses any type of vertex position */
+	uint32 bUsesVertexPosition : 1;
+
+	uint32 bUsesTransformVector : 1;
+	// True if the current property requires last frame's information
+	uint32 bCompilingPreviousFrame : 1;
+	/** True if material will output accurate velocities during base pass rendering. */
+	uint32 bOutputsBasePassVelocities : 1;
+	uint32 bUsesPixelDepthOffset : 1;
+	uint32 bUsesWorldPositionOffset : 1;
+	uint32 bUsesEmissiveColor : 1;
+	uint32 bUsesDistanceCullFade : 1;
+	/** true if the Roughness input evaluates to a constant 1.0 */
+	uint32 bIsFullyRough : 1;
+	/** true if allowed to generate code chunks. Translator operates in two phases; generate all code chunks & query meta data based on generated code chunks. */
+	uint32 bAllowCodeChunkGeneration : 1;
+
+	/** True if this material reads any per-instance custom data */
+	uint32 bUsesPerInstanceCustomData : 1;
+
+	/** Tracks the texture coordinates used by this material. */
+	TBitArray<> AllocatedUserTexCoords;
+	/** Tracks the texture coordinates used by the vertex shader in this material. */
+	TBitArray<> AllocatedUserVertexTexCoords;
+
+	uint32 DynamicParticleParameterMask;
+
+	/** Will contain all the shading models picked up from the material expression graph */
+	FMaterialShadingModelField ShadingModelsFromCompilation;
+
+	/** True if this material reads any vertex data */
+	uint32 bRequiresVertexData : 1;
+};
