[GLTFExporter] feat: add support for analyzing custom outputs in AnalyzeMaterialProperty

commit b66f7fd5c1c6e1c9ae99b19faf028423bf396119
Date: Thu, 12 Nov 2020 15:28:45 +0100

    feat: add support for analyzing custom outputs in AnalyzeMaterialProperty
    (cherry picked from commit 3df3c574be5f4e42d41fdcdceeac563fd781ff4d​)

#skipci

diff --git a/Engine/Plugins/Enterprise/GLTFExporter/MigratedFromGit/Source/GLTFExporter/Private/Converters/GLTFMaterialUtility.cpp b/Engine/Plugins/Enterprise/GLTFExporter/MigratedFromGit/Source/GLTFExporter/Private/Converters/GLTFMaterialUtility.cpp
--- a/Engine/Plugins/Enterprise/GLTFExporter/MigratedFromGit/Source/GLTFExporter/Private/Converters/GLTFMaterialUtility.cpp
+++ b/Engine/Plugins/Enterprise/GLTFExporter/MigratedFromGit/Source/GLTFExporter/Private/Converters/GLTFMaterialUtility.cpp
@@ -501,7 +501,7 @@ bool FGLTFMaterialUtility::NeedsMeshData(const UMaterialInterface* Material)
 	}
 
 	// TODO: only analyze properties that will be needed for this specific material
-	const TArray<EMaterialProperty> Properties =
+	const TArray<FMaterialPropertyEx> Properties =
 	{
 		MP_BaseColor,
 		MP_EmissiveColor,
@@ -513,13 +513,13 @@ bool FGLTFMaterialUtility::NeedsMeshData(const UMaterialInterface* Material)
 		MP_AmbientOcclusion,
 		MP_CustomData0,
 		MP_CustomData1,
-		// TODO: add support for custom output ClearCoatBottomNormal
+		{ MP_CustomOutput, FName(TEXT("ClearCoatBottomNormal")) },
 	};
 
 	bool bRequiresVertexData = false;
 	FGLTFMaterialStatistics MaterialStatistics;
 
-	for (const EMaterialProperty Property: Properties)
+	for (const FMaterialPropertyEx& Property: Properties)
 	{
 		AnalyzeMaterialProperty(Material, Property, MaterialStatistics);
 		bRequiresVertexData |= MaterialStatistics.bRequiresVertexData;
@@ -541,8 +541,10 @@ bool FGLTFMaterialUtility::NeedsMeshData(const TArray<const UMaterialInterface*>
 	return false;
 }
 
-void FGLTFMaterialUtility::AnalyzeMaterialProperty(const UMaterialInterface* InMaterial, EMaterialProperty InProperty, FGLTFMaterialStatistics& OutMaterialStatistics)
+void FGLTFMaterialUtility::AnalyzeMaterialProperty(const UMaterialInterface* InMaterial, const FMaterialPropertyEx& InProperty, FGLTFMaterialStatistics& OutMaterialStatistics)
 {
+	// TODO: use a shared UGLTFMaterialAnalyzer instance instead of creating a new one for each invocation
+
 	UGLTFMaterialAnalyzer* Analyzer = NewObject<UGLTFMaterialAnalyzer>();
 	Analyzer->AnalyzeMaterialProperty(InMaterial, InProperty, OutMaterialStatistics);
 }
diff --git a/Engine/Plugins/Enterprise/GLTFExporter/MigratedFromGit/Source/GLTFExporter/Private/Converters/GLTFMaterialUtility.h b/Engine/Plugins/Enterprise/GLTFExporter/MigratedFromGit/Source/GLTFExporter/Private/Converters/GLTFMaterialUtility.h
--- a/Engine/Plugins/Enterprise/GLTFExporter/MigratedFromGit/Source/GLTFExporter/Private/Converters/GLTFMaterialUtility.h
+++ b/Engine/Plugins/Enterprise/GLTFExporter/MigratedFromGit/Source/GLTFExporter/Private/Converters/GLTFMaterialUtility.h
@@ -6,6 +6,7 @@
 #include "Engine/Texture2D.h"
 
 struct FGLTFMaterialStatistics;
+struct FMaterialPropertyEx;
 
 struct FGLTFTextureCombineSource
 {
@@ -76,7 +77,7 @@ struct FGLTFMaterialUtility
 	static bool NeedsMeshData(const UMaterialInterface* Material);
 	static bool NeedsMeshData(const TArray<const UMaterialInterface*>& Materials);
 
-	static void AnalyzeMaterialProperty(const UMaterialInterface* Material, EMaterialProperty InProperty, FGLTFMaterialStatistics& OutMaterialStatistics);
+	static void AnalyzeMaterialProperty(const UMaterialInterface* Material, const FMaterialPropertyEx& InProperty, FGLTFMaterialStatistics& OutMaterialStatistics);
 
 	static const UMaterialInterface* GetInterface(const UMaterialInterface* Material);
 	static const UMaterialInterface* GetInterface(const FStaticMaterial& Material);
diff --git a/Engine/Plugins/Enterprise/GLTFExporter/MigratedFromGit/Source/GLTFMaterialAnalyzer/GLTFMaterialAnalyzer.Build.cs b/Engine/Plugins/Enterprise/GLTFExporter/MigratedFromGit/Source/GLTFMaterialAnalyzer/GLTFMaterialAnalyzer.Build.cs
--- a/Engine/Plugins/Enterprise/GLTFExporter/MigratedFromGit/Source/GLTFMaterialAnalyzer/GLTFMaterialAnalyzer.Build.cs
+++ b/Engine/Plugins/Enterprise/GLTFExporter/MigratedFromGit/Source/GLTFMaterialAnalyzer/GLTFMaterialAnalyzer.Build.cs
@@ -12,7 +12,8 @@ public class GLTFMaterialAnalyzer : ModuleRules
 				"Core",
 				"CoreUObject",
 				"Engine",
-				"RenderCore"
+				"RenderCore",
+				"GLTFMaterialBaking"	// NOTE: only used to get access to baking structures
 			}
 		);
 
diff --git a/Engine/Plugins/Enterprise/GLTFExporter/MigratedFromGit/Source/GLTFMaterialAnalyzer/Private/GLTFMaterialAnalyzer.cpp b/Engine/Plugins/Enterprise/GLTFExporter/MigratedFromGit/Source/GLTFMaterialAnalyzer/Private/GLTFMaterialAnalyzer.cpp
--- a/Engine/Plugins/Enterprise/GLTFExporter/MigratedFromGit/Source/GLTFMaterialAnalyzer/Private/GLTFMaterialAnalyzer.cpp
+++ b/Engine/Plugins/Enterprise/GLTFExporter/MigratedFromGit/Source/GLTFMaterialAnalyzer/Private/GLTFMaterialAnalyzer.cpp
@@ -2,17 +2,24 @@
 
 #include "GLTFMaterialAnalyzer.h"
 #include "GLTFMaterialStatistics.h"
+#include "GLTFMaterialBaking/Public/MaterialBakingStructures.h"
 #include "Materials/HLSLMaterialTranslator.h"
 
-void UGLTFMaterialAnalyzer::AnalyzeMaterialProperty(const UMaterialInterface* InMaterial, EMaterialProperty InProperty, FGLTFMaterialStatistics& OutMaterialStatistics)
+void UGLTFMaterialAnalyzer::AnalyzeMaterialProperty(const UMaterialInterface* InMaterial, const FMaterialPropertyEx& InProperty, FGLTFMaterialStatistics& OutMaterialStatistics)
 {
+	Property = &InProperty;
 	Material = const_cast<UMaterialInterface*>(InMaterial);
 	MaterialStatistics = &OutMaterialStatistics;
 
 	int32 NumTextureCoordinates;
 	bool RequiresVertexData;
-	UMaterialInterface::AnalyzeMaterialProperty(InProperty, NumTextureCoordinates, RequiresVertexData);
 
+	// NOTE: When analyzing custom outputs, the property *must* be set to MP_MAX or the compiler will refuse to compile the output
+	const EMaterialProperty TempProperty = InProperty.Type == MP_CustomOutput ? MP_MAX : InProperty.Type;
+
+	UMaterialInterface::AnalyzeMaterialProperty(TempProperty, NumTextureCoordinates, RequiresVertexData);
+
+	Property = nullptr;
 	Material = nullptr;
 	MaterialStatistics = nullptr;
 }
@@ -31,7 +38,28 @@ int32 UGLTFMaterialAnalyzer::CompilePropertyEx(FMaterialCompiler* Compiler, cons
 		friend UGLTFMaterialAnalyzer;
 	};
 
-	const int32 Result = Material->CompilePropertyEx(Compiler, AttributeID);
+	int32 Result = INDEX_NONE;
+
+	if (Property->Type == MP_CustomOutput)
+	{
+		const FString CustomOutput = Property->CustomOutput.ToString();
+
+		for (UMaterialExpression* Expression : Material->GetMaterial()->Expressions)
+		{
+			UMaterialExpressionCustomOutput* CustomOutputExpression = Cast<UMaterialExpressionCustomOutput>(Expression);
+			if (CustomOutputExpression && CustomOutputExpression->GetDisplayName() == CustomOutput)
+			{
+				// TODO: can we rely on OutputIndex always being 0?
+				Result = CustomOutputExpression->Compile(Compiler, 0);
+				break;
+			}
+		}
+	}
+	else
+	{
+		Result = Material->CompilePropertyEx(Compiler, AttributeID);
+	}
+
 	FTranslatorHack* Translator = static_cast<FTranslatorHack*>(Compiler);
 
 	MaterialStatistics->bUsesSceneDepth = Translator->bUsesSceneDepth;
diff --git a/Engine/Plugins/Enterprise/GLTFExporter/MigratedFromGit/Source/GLTFMaterialAnalyzer/Public/GLTFMaterialAnalyzer.h b/Engine/Plugins/Enterprise/GLTFExporter/MigratedFromGit/Source/GLTFMaterialAnalyzer/Public/GLTFMaterialAnalyzer.h
--- a/Engine/Plugins/Enterprise/GLTFExporter/MigratedFromGit/Source/GLTFMaterialAnalyzer/Public/GLTFMaterialAnalyzer.h
+++ b/Engine/Plugins/Enterprise/GLTFExporter/MigratedFromGit/Source/GLTFMaterialAnalyzer/Public/GLTFMaterialAnalyzer.h
@@ -6,6 +6,7 @@
 #include "GLTFMaterialAnalyzer.generated.h"
 
 struct FGLTFMaterialStatistics;
+struct FMaterialPropertyEx;
 
 UCLASS()
 class GLTFMATERIALANALYZER_API UGLTFMaterialAnalyzer : public UMaterialInterface
@@ -14,7 +15,7 @@ class GLTFMATERIALANALYZER_API UGLTFMaterialAnalyzer : public UMaterialInterface
 
 public:
 
-	void AnalyzeMaterialProperty(const UMaterialInterface* InMaterial, EMaterialProperty InProperty, FGLTFMaterialStatistics& OutMaterialStatistics);
+	void AnalyzeMaterialProperty(const UMaterialInterface* InMaterial, const FMaterialPropertyEx& InProperty, FGLTFMaterialStatistics& OutMaterialStatistics);
 
 private:
 
@@ -24,6 +25,8 @@ private:
 
 	virtual bool IsPropertyActive(EMaterialProperty InProperty) const override;
 
+	const FMaterialPropertyEx* Property;
+
 	// ReSharper disable once CppUE4ProbableMemoryIssuesWithUObject
 	UMaterialInterface* Material;
 
