[GLTFExporter] refactor: remove dependency to FMaterialPropertyEx in material analyzer

commit e37af3f8d1535a19dbda181ddd488af77af017de
Date: Thu, 12 Nov 2020 22:55:59 +0100

    refactor: remove dependency to FMaterialPropertyEx in material analyzer
    (cherry picked from commit 283925d6c96bd35edc9906fe6576c23f0a044413​)

#skipci

diff --git a/Engine/Plugins/Enterprise/GLTFExporter/MigratedFromGit/Source/GLTFExporter/Private/Converters/GLTFMaterialUtility.cpp b/Engine/Plugins/Enterprise/GLTFExporter/MigratedFromGit/Source/GLTFExporter/Private/Converters/GLTFMaterialUtility.cpp
--- a/Engine/Plugins/Enterprise/GLTFExporter/MigratedFromGit/Source/GLTFExporter/Private/Converters/GLTFMaterialUtility.cpp
+++ b/Engine/Plugins/Enterprise/GLTFExporter/MigratedFromGit/Source/GLTFExporter/Private/Converters/GLTFMaterialUtility.cpp
@@ -513,7 +513,7 @@ bool FGLTFMaterialUtility::NeedsMeshData(const UMaterialInterface* Material)
 		MP_AmbientOcclusion,
 		MP_CustomData0,
 		MP_CustomData1,
-		{ MP_CustomOutput, FName(TEXT("ClearCoatBottomNormal")) },
+		TEXT("ClearCoatBottomNormal"),
 	};
 
 	bool bRequiresVertexData = false;
@@ -546,7 +546,7 @@ void FGLTFMaterialUtility::AnalyzeMaterialProperty(const UMaterialInterface* InM
 	// TODO: use a shared UGLTFMaterialAnalyzer instance instead of creating a new one for each invocation
 
 	UGLTFMaterialAnalyzer* Analyzer = NewObject<UGLTFMaterialAnalyzer>();
-	Analyzer->AnalyzeMaterialProperty(InMaterial, InProperty, OutMaterialStatistics);
+	Analyzer->AnalyzeMaterialProperty(InMaterial, InProperty.Type, InProperty.CustomOutput.ToString(), OutMaterialStatistics);
 }
 
 const UMaterialInterface* FGLTFMaterialUtility::GetInterface(const UMaterialInterface* Material)
diff --git a/Engine/Plugins/Enterprise/GLTFExporter/MigratedFromGit/Source/GLTFMaterialAnalyzer/GLTFMaterialAnalyzer.Build.cs b/Engine/Plugins/Enterprise/GLTFExporter/MigratedFromGit/Source/GLTFMaterialAnalyzer/GLTFMaterialAnalyzer.Build.cs
--- a/Engine/Plugins/Enterprise/GLTFExporter/MigratedFromGit/Source/GLTFMaterialAnalyzer/GLTFMaterialAnalyzer.Build.cs
+++ b/Engine/Plugins/Enterprise/GLTFExporter/MigratedFromGit/Source/GLTFMaterialAnalyzer/GLTFMaterialAnalyzer.Build.cs
@@ -12,8 +12,7 @@ public class GLTFMaterialAnalyzer : ModuleRules
 				"Core",
 				"CoreUObject",
 				"Engine",
-				"RenderCore",
-				"GLTFMaterialBaking"	// NOTE: only used to get access to baking structures
+				"RenderCore"
 			}
 		);
 
diff --git a/Engine/Plugins/Enterprise/GLTFExporter/MigratedFromGit/Source/GLTFMaterialAnalyzer/Private/GLTFMaterialAnalyzer.cpp b/Engine/Plugins/Enterprise/GLTFExporter/MigratedFromGit/Source/GLTFMaterialAnalyzer/Private/GLTFMaterialAnalyzer.cpp
--- a/Engine/Plugins/Enterprise/GLTFExporter/MigratedFromGit/Source/GLTFMaterialAnalyzer/Private/GLTFMaterialAnalyzer.cpp
+++ b/Engine/Plugins/Enterprise/GLTFExporter/MigratedFromGit/Source/GLTFMaterialAnalyzer/Private/GLTFMaterialAnalyzer.cpp
@@ -5,25 +5,40 @@
 #include "GLTFMaterialBaking/Public/MaterialBakingStructures.h"
 #include "Materials/HLSLMaterialTranslator.h"
 
-void UGLTFMaterialAnalyzer::AnalyzeMaterialProperty(const UMaterialInterface* InMaterial, const FMaterialPropertyEx& InProperty, FGLTFMaterialStatistics& OutMaterialStatistics)
+void UGLTFMaterialAnalyzer::AnalyzeMaterialProperty(const UMaterialInterface* InMaterial, const EMaterialProperty& InProperty, const FString& InCustomOutput, FGLTFMaterialStatistics& OutMaterialStatistics)
 {
-	Property = &InProperty;
+	Property = InProperty;
+	CustomOutput = InCustomOutput;
 	Material = const_cast<UMaterialInterface*>(InMaterial);
 	MaterialStatistics = &OutMaterialStatistics;
 
-	int32 NumTextureCoordinates;
-	bool RequiresVertexData;
-
 	// NOTE: When analyzing custom outputs, the property *must* be set to MP_MAX or the compiler will refuse to compile the output
-	const EMaterialProperty TempProperty = InProperty.Type == MP_CustomOutput ? MP_MAX : InProperty.Type;
+	const EMaterialProperty SafeProperty = InProperty == MP_CustomOutput ? MP_MAX : InProperty;
 
-	UMaterialInterface::AnalyzeMaterialProperty(TempProperty, NumTextureCoordinates, RequiresVertexData);
+	int32 DummyNumTextureCoordinates; // Dummy value from built-in analysis not used since it's insufficient
+	bool DummyRequiresVertexData; // Dummy value from built-in analysis not used since it's insufficient
+	Super::AnalyzeMaterialProperty(SafeProperty, DummyNumTextureCoordinates, DummyRequiresVertexData);
 
-	Property = nullptr;
+	Property = MP_MAX;
+	CustomOutput = {};
 	Material = nullptr;
 	MaterialStatistics = nullptr;
 }
 
+UMaterialExpressionCustomOutput* UGLTFMaterialAnalyzer::GetCustomOutputExpression() const
+{
+	for (UMaterialExpression* Expression : Material->GetMaterial()->Expressions)
+	{
+		UMaterialExpressionCustomOutput* CustomOutputExpression = Cast<UMaterialExpressionCustomOutput>(Expression);
+		if (CustomOutputExpression != nullptr && CustomOutputExpression->GetDisplayName() == CustomOutput)
+		{
+			return CustomOutputExpression;
+		}
+	}
+
+	return nullptr;
+}
+
 FMaterialResource* UGLTFMaterialAnalyzer::GetMaterialResource(ERHIFeatureLevel::Type InFeatureLevel, EMaterialQualityLevel::Type QualityLevel)
 {
 	return Material->GetMaterialResource(InFeatureLevel, QualityLevel);
@@ -38,22 +53,12 @@ int32 UGLTFMaterialAnalyzer::CompilePropertyEx(FMaterialCompiler* Compiler, cons
 		friend UGLTFMaterialAnalyzer;
 	};
 
-	int32 Result = INDEX_NONE;
+	int32 Result;
 
-	if (Property->Type == MP_CustomOutput)
+	if (Property == MP_CustomOutput)
 	{
-		const FString CustomOutput = Property->CustomOutput.ToString();
-
-		for (UMaterialExpression* Expression : Material->GetMaterial()->Expressions)
-		{
-			UMaterialExpressionCustomOutput* CustomOutputExpression = Cast<UMaterialExpressionCustomOutput>(Expression);
-			if (CustomOutputExpression && CustomOutputExpression->GetDisplayName() == CustomOutput)
-			{
-				// TODO: can we rely on OutputIndex always being 0?
-				Result = CustomOutputExpression->Compile(Compiler, 0);
-				break;
-			}
-		}
+		UMaterialExpressionCustomOutput* CustomOutputExpression = GetCustomOutputExpression();
+		Result = CustomOutputExpression != nullptr ? CustomOutputExpression->Compile(Compiler, 0) : INDEX_NONE;
 	}
 	else
 	{
diff --git a/Engine/Plugins/Enterprise/GLTFExporter/MigratedFromGit/Source/GLTFMaterialAnalyzer/Public/GLTFMaterialAnalyzer.h b/Engine/Plugins/Enterprise/GLTFExporter/MigratedFromGit/Source/GLTFMaterialAnalyzer/Public/GLTFMaterialAnalyzer.h
--- a/Engine/Plugins/Enterprise/GLTFExporter/MigratedFromGit/Source/GLTFMaterialAnalyzer/Public/GLTFMaterialAnalyzer.h
+++ b/Engine/Plugins/Enterprise/GLTFExporter/MigratedFromGit/Source/GLTFMaterialAnalyzer/Public/GLTFMaterialAnalyzer.h
@@ -6,7 +6,6 @@
 #include "GLTFMaterialAnalyzer.generated.h"
 
 struct FGLTFMaterialStatistics;
-struct FMaterialPropertyEx;
 
 UCLASS()
 class GLTFMATERIALANALYZER_API UGLTFMaterialAnalyzer : public UMaterialInterface
@@ -15,17 +14,20 @@ class GLTFMATERIALANALYZER_API UGLTFMaterialAnalyzer : public UMaterialInterface
 
 public:
 
-	void AnalyzeMaterialProperty(const UMaterialInterface* InMaterial, const FMaterialPropertyEx& InProperty, FGLTFMaterialStatistics& OutMaterialStatistics);
+	void AnalyzeMaterialProperty(const UMaterialInterface* InMaterial, const EMaterialProperty& InProperty, const FString& InCustomOutput, FGLTFMaterialStatistics& OutMaterialStatistics);
 
 private:
 
+	UMaterialExpressionCustomOutput* GetCustomOutputExpression() const;
+
 	virtual FMaterialResource* GetMaterialResource(ERHIFeatureLevel::Type InFeatureLevel, EMaterialQualityLevel::Type QualityLevel) override;
 
 	virtual int32 CompilePropertyEx(FMaterialCompiler* Compiler, const FGuid& AttributeID) override;
 
 	virtual bool IsPropertyActive(EMaterialProperty InProperty) const override;
 
-	const FMaterialPropertyEx* Property;
+	EMaterialProperty Property;
+	FString CustomOutput;
 
 	// ReSharper disable once CppUE4ProbableMemoryIssuesWithUObject
 	UMaterialInterface* Material;
