[GLTFExporter] refactor: simplify GLTFMaterialStatistics to only necessary output and rename to GLTFMaterialAnalysis

commit 18b93109d8c218ef822c0a815e47ad56d56b8d54
Date: Fri, 13 Nov 2020 00:14:49 +0100

    refactor: simplify GLTFMaterialStatistics to only necessary output and rename to GLTFMaterialAnalysis
    (cherry picked from commit 3a09bc71ba8db11e499b6da81e8d8f336b25fcac​)

#skipci

diff --git a/Engine/Plugins/Enterprise/GLTFExporter/MigratedFromGit/Source/GLTFExporter/Private/Converters/GLTFMaterialUtility.cpp b/Engine/Plugins/Enterprise/GLTFExporter/MigratedFromGit/Source/GLTFExporter/Private/Converters/GLTFMaterialUtility.cpp
--- a/Engine/Plugins/Enterprise/GLTFExporter/MigratedFromGit/Source/GLTFExporter/Private/Converters/GLTFMaterialUtility.cpp
+++ b/Engine/Plugins/Enterprise/GLTFExporter/MigratedFromGit/Source/GLTFExporter/Private/Converters/GLTFMaterialUtility.cpp
@@ -3,7 +3,7 @@
 #include "Converters/GLTFMaterialUtility.h"
 #include "Converters/GLTFTextureUtility.h"
 #include "Converters/GLTFNameUtility.h"
-#include "GLTFMaterialStatistics.h"
+#include "GLTFMaterialAnalysis.h"
 #include "GLTFMaterialAnalyzer.h"
 #include "Engine/TextureRenderTarget2D.h"
 #include "CanvasItem.h"
@@ -517,12 +517,12 @@ bool FGLTFMaterialUtility::NeedsMeshData(const UMaterialInterface* Material)
 	};
 
 	bool bRequiresVertexData = false;
-	FGLTFMaterialStatistics MaterialStatistics;
+	FGLTFMaterialAnalysis Analysis;
 
 	for (const FMaterialPropertyEx& Property: Properties)
 	{
-		AnalyzeMaterialProperty(Material, Property, MaterialStatistics);
-		bRequiresVertexData |= MaterialStatistics.bRequiresVertexData;
+		AnalyzeMaterialProperty(Material, Property, Analysis);
+		bRequiresVertexData |= Analysis.bRequiresVertexData;
 	}
 
 	return bRequiresVertexData;
@@ -541,12 +541,12 @@ bool FGLTFMaterialUtility::NeedsMeshData(const TArray<const UMaterialInterface*>
 	return false;
 }
 
-void FGLTFMaterialUtility::AnalyzeMaterialProperty(const UMaterialInterface* InMaterial, const FMaterialPropertyEx& InProperty, FGLTFMaterialStatistics& OutMaterialStatistics)
+void FGLTFMaterialUtility::AnalyzeMaterialProperty(const UMaterialInterface* InMaterial, const FMaterialPropertyEx& InProperty, FGLTFMaterialAnalysis& OutMaterialAnalysis)
 {
 	// TODO: use a shared UGLTFMaterialAnalyzer instance instead of creating a new one for each invocation
 
 	UGLTFMaterialAnalyzer* Analyzer = NewObject<UGLTFMaterialAnalyzer>();
-	Analyzer->AnalyzeMaterialProperty(InMaterial, InProperty.Type, InProperty.CustomOutput.ToString(), OutMaterialStatistics);
+	Analyzer->AnalyzeMaterialProperty(InMaterial, InProperty.Type, InProperty.CustomOutput.ToString(), OutMaterialAnalysis);
 }
 
 const UMaterialInterface* FGLTFMaterialUtility::GetInterface(const UMaterialInterface* Material)
diff --git a/Engine/Plugins/Enterprise/GLTFExporter/MigratedFromGit/Source/GLTFExporter/Private/Converters/GLTFMaterialUtility.h b/Engine/Plugins/Enterprise/GLTFExporter/MigratedFromGit/Source/GLTFExporter/Private/Converters/GLTFMaterialUtility.h
--- a/Engine/Plugins/Enterprise/GLTFExporter/MigratedFromGit/Source/GLTFExporter/Private/Converters/GLTFMaterialUtility.h
+++ b/Engine/Plugins/Enterprise/GLTFExporter/MigratedFromGit/Source/GLTFExporter/Private/Converters/GLTFMaterialUtility.h
@@ -5,7 +5,7 @@
 #include "Builders/GLTFContainerBuilder.h"
 #include "Engine/Texture2D.h"
 
-struct FGLTFMaterialStatistics;
+struct FGLTFMaterialAnalysis;
 struct FMaterialPropertyEx;
 
 struct FGLTFTextureCombineSource
@@ -77,7 +77,7 @@ struct FGLTFMaterialUtility
 	static bool NeedsMeshData(const UMaterialInterface* Material);
 	static bool NeedsMeshData(const TArray<const UMaterialInterface*>& Materials);
 
-	static void AnalyzeMaterialProperty(const UMaterialInterface* Material, const FMaterialPropertyEx& InProperty, FGLTFMaterialStatistics& OutMaterialStatistics);
+	static void AnalyzeMaterialProperty(const UMaterialInterface* Material, const FMaterialPropertyEx& InProperty, FGLTFMaterialAnalysis& OutMaterialAnalysis);
 
 	static const UMaterialInterface* GetInterface(const UMaterialInterface* Material);
 	static const UMaterialInterface* GetInterface(const FStaticMaterial& Material);
diff --git a/Engine/Plugins/Enterprise/GLTFExporter/MigratedFromGit/Source/GLTFMaterialAnalyzer/Private/GLTFMaterialAnalyzer.cpp b/Engine/Plugins/Enterprise/GLTFExporter/MigratedFromGit/Source/GLTFMaterialAnalyzer/Private/GLTFMaterialAnalyzer.cpp
--- a/Engine/Plugins/Enterprise/GLTFExporter/MigratedFromGit/Source/GLTFMaterialAnalyzer/Private/GLTFMaterialAnalyzer.cpp
+++ b/Engine/Plugins/Enterprise/GLTFExporter/MigratedFromGit/Source/GLTFMaterialAnalyzer/Private/GLTFMaterialAnalyzer.cpp
@@ -1,16 +1,16 @@
 // Copyright Epic Games, Inc. All Rights Reserved.
 
 #include "GLTFMaterialAnalyzer.h"
-#include "GLTFMaterialStatistics.h"
+#include "GLTFMaterialAnalysis.h"
 #include "GLTFMaterialBaking/Public/MaterialBakingStructures.h"
 #include "Materials/HLSLMaterialTranslator.h"
 
-void UGLTFMaterialAnalyzer::AnalyzeMaterialProperty(const UMaterialInterface* InMaterial, const EMaterialProperty& InProperty, const FString& InCustomOutput, FGLTFMaterialStatistics& OutMaterialStatistics)
+void UGLTFMaterialAnalyzer::AnalyzeMaterialProperty(const UMaterialInterface* InMaterial, const EMaterialProperty& InProperty, const FString& InCustomOutput, FGLTFMaterialAnalysis& OutAnalysis)
 {
 	Property = InProperty;
 	CustomOutput = InCustomOutput;
 	Material = const_cast<UMaterialInterface*>(InMaterial);
-	MaterialStatistics = &OutMaterialStatistics;
+	Analysis = &OutAnalysis;
 
 	// NOTE: When analyzing custom outputs, the property *must* be set to MP_MAX or the compiler will refuse to compile the output
 	const EMaterialProperty SafeProperty = InProperty == MP_CustomOutput ? MP_MAX : InProperty;
@@ -22,7 +22,7 @@ void UGLTFMaterialAnalyzer::AnalyzeMaterialProperty(const UMaterialInterface* In
 	Property = MP_MAX;
 	CustomOutput = {};
 	Material = nullptr;
-	MaterialStatistics = nullptr;
+	Analysis = nullptr;
 }
 
 UMaterialExpressionCustomOutput* UGLTFMaterialAnalyzer::GetCustomOutputExpression() const
@@ -53,6 +53,7 @@ int32 UGLTFMaterialAnalyzer::CompilePropertyEx(FMaterialCompiler* Compiler, cons
 		friend UGLTFMaterialAnalyzer;
 	};
 
+	FTranslatorHack* Translator = static_cast<FTranslatorHack*>(Compiler);
 	int32 Result;
 
 	if (Property == MP_CustomOutput)
@@ -62,49 +63,20 @@ int32 UGLTFMaterialAnalyzer::CompilePropertyEx(FMaterialCompiler* Compiler, cons
 	}
 	else
 	{
-		Result = Material->CompilePropertyEx(Compiler, AttributeID);
+		Result = Material->CompilePropertyEx(Translator, AttributeID);
 	}
 
-	FTranslatorHack* Translator = static_cast<FTranslatorHack*>(Compiler);
+	Analysis->TextureCoordinates = Translator->AllocatedUserTexCoords;
+	Analysis->NumTextureCoordinates = 0;
 
-	MaterialStatistics->bUsesSceneDepth = Translator->bUsesSceneDepth;
-	MaterialStatistics->bNeedsParticlePosition = Translator->bNeedsParticlePosition;
-	MaterialStatistics->bNeedsParticleVelocity = Translator->bNeedsParticleVelocity;
-	MaterialStatistics->bNeedsParticleTime = Translator->bNeedsParticleTime;
-	MaterialStatistics->bUsesParticleMotionBlur = Translator->bUsesParticleMotionBlur;
-	MaterialStatistics->bNeedsParticleRandom = Translator->bNeedsParticleRandom;
-	MaterialStatistics->bUsesSphericalParticleOpacity = Translator->bUsesSphericalParticleOpacity;
-	MaterialStatistics->bUsesParticleSubUVs = Translator->bUsesParticleSubUVs;
-	MaterialStatistics->bUsesLightmapUVs = Translator->bUsesLightmapUVs;
-	MaterialStatistics->bUsesAOMaterialMask = Translator->bUsesAOMaterialMask;
-	MaterialStatistics->bUsesSpeedTree = Translator->bUsesSpeedTree;
-	MaterialStatistics->bNeedsWorldPositionExcludingShaderOffsets = Translator->bNeedsWorldPositionExcludingShaderOffsets;
-	MaterialStatistics->bNeedsParticleSize = Translator->bNeedsParticleSize;
-	MaterialStatistics->bNeedsSceneTexturePostProcessInputs = Translator->bNeedsSceneTexturePostProcessInputs;
-	MaterialStatistics->bUsesAtmosphericFog = Translator->bUsesAtmosphericFog;
-	MaterialStatistics->bUsesSkyAtmosphere = Translator->bUsesSkyAtmosphere;
-	MaterialStatistics->bUsesVertexColor = Translator->bUsesVertexColor;
-	MaterialStatistics->bUsesParticleColor = Translator->bUsesParticleColor;
-	MaterialStatistics->bUsesParticleLocalToWorld = Translator->bUsesParticleLocalToWorld;
-	MaterialStatistics->bUsesParticleWorldToLocal = Translator->bUsesParticleWorldToLocal;
-	MaterialStatistics->bUsesVertexPosition = Translator->bUsesVertexPosition;
-	MaterialStatistics->bUsesTransformVector = Translator->bUsesTransformVector;
-	MaterialStatistics->bCompilingPreviousFrame = Translator->bCompilingPreviousFrame;
-	MaterialStatistics->bOutputsBasePassVelocities = Translator->bOutputsBasePassVelocities;
-	MaterialStatistics->bUsesPixelDepthOffset = Translator->bUsesPixelDepthOffset;
-	MaterialStatistics->bUsesWorldPositionOffset = Translator->bUsesWorldPositionOffset;
-	MaterialStatistics->bUsesEmissiveColor = Translator->bUsesEmissiveColor;
-	MaterialStatistics->bUsesDistanceCullFade = Translator->bUsesDistanceCullFade;
-	MaterialStatistics->bIsFullyRough = Translator->bIsFullyRough;
-	MaterialStatistics->bAllowCodeChunkGeneration = Translator->bAllowCodeChunkGeneration;
-	MaterialStatistics->bUsesPerInstanceCustomData = Translator->bUsesPerInstanceCustomData;
-	MaterialStatistics->AllocatedUserTexCoords = Translator->AllocatedUserTexCoords;
-	MaterialStatistics->AllocatedUserVertexTexCoords = Translator->AllocatedUserVertexTexCoords;
-	MaterialStatistics->DynamicParticleParameterMask = Translator->DynamicParticleParameterMask;
-	MaterialStatistics->ShadingModelsFromCompilation = Translator->ShadingModelsFromCompilation;
+	for (int32 TexCoordIndex = 0; TexCoordIndex < Translator->AllocatedUserTexCoords.Num(); TexCoordIndex++)
+	{
+		const bool IsTexCoordUsed = Translator->AllocatedUserTexCoords[TexCoordIndex];
+		Analysis->NumTextureCoordinates += IsTexCoordUsed ? 1 : 0;
+	}
 
 	// TODO: investigate if we need to check more conditions to determine that vertex data is required
-	MaterialStatistics->bRequiresVertexData =
+	Analysis->bRequiresVertexData =
 		Translator->bUsesVertexColor ||
 		Translator->bUsesTransformVector ||
 		Translator->bNeedsWorldPositionExcludingShaderOffsets ||
diff --git a/Engine/Plugins/Enterprise/GLTFExporter/MigratedFromGit/Source/GLTFMaterialAnalyzer/Public/GLTFMaterialAnalysis.h b/Engine/Plugins/Enterprise/GLTFExporter/MigratedFromGit/Source/GLTFMaterialAnalyzer/Public/GLTFMaterialAnalysis.h
new file mode 100644
--- /dev/null
+++ b/Engine/Plugins/Enterprise/GLTFExporter/MigratedFromGit/Source/GLTFMaterialAnalyzer/Public/GLTFMaterialAnalysis.h
@@ -0,0 +1,18 @@
+// Copyright Epic Games, Inc. All Rights Reserved.
+
+#pragma once
+
+struct FGLTFMaterialAnalysis
+{
+	/** Tracks the texture coordinates used by this material */
+	TBitArray<> TextureCoordinates;
+
+	/** Number of texture coordinates used by this material */
+	uint32 NumTextureCoordinates;
+
+	/** Will contain all the shading models picked up from the material expression graph */
+	FMaterialShadingModelField ShadingModels;
+
+	/** True if this material reads any vertex data */
+	bool bRequiresVertexData;
+};
diff --git a/Engine/Plugins/Enterprise/GLTFExporter/MigratedFromGit/Source/GLTFMaterialAnalyzer/Public/GLTFMaterialAnalyzer.h b/Engine/Plugins/Enterprise/GLTFExporter/MigratedFromGit/Source/GLTFMaterialAnalyzer/Public/GLTFMaterialAnalyzer.h
--- a/Engine/Plugins/Enterprise/GLTFExporter/MigratedFromGit/Source/GLTFMaterialAnalyzer/Public/GLTFMaterialAnalyzer.h
+++ b/Engine/Plugins/Enterprise/GLTFExporter/MigratedFromGit/Source/GLTFMaterialAnalyzer/Public/GLTFMaterialAnalyzer.h
@@ -5,7 +5,7 @@
 #include "Materials/MaterialInterface.h"
 #include "GLTFMaterialAnalyzer.generated.h"
 
-struct FGLTFMaterialStatistics;
+struct FGLTFMaterialAnalysis;
 
 UCLASS()
 class GLTFMATERIALANALYZER_API UGLTFMaterialAnalyzer : public UMaterialInterface
@@ -14,7 +14,7 @@ class GLTFMATERIALANALYZER_API UGLTFMaterialAnalyzer : public UMaterialInterface
 
 public:
 
-	void AnalyzeMaterialProperty(const UMaterialInterface* InMaterial, const EMaterialProperty& InProperty, const FString& InCustomOutput, FGLTFMaterialStatistics& OutMaterialStatistics);
+	void AnalyzeMaterialProperty(const UMaterialInterface* InMaterial, const EMaterialProperty& InProperty, const FString& InCustomOutput, FGLTFMaterialAnalysis& OutAnalysis);
 
 private:
 
@@ -32,5 +32,5 @@ private:
 	// ReSharper disable once CppUE4ProbableMemoryIssuesWithUObject
 	UMaterialInterface* Material;
 
-	FGLTFMaterialStatistics* MaterialStatistics;
+	FGLTFMaterialAnalysis* Analysis;
 };
diff --git a/Engine/Plugins/Enterprise/GLTFExporter/MigratedFromGit/Source/GLTFMaterialAnalyzer/Public/GLTFMaterialStatistics.h b/Engine/Plugins/Enterprise/GLTFExporter/MigratedFromGit/Source/GLTFMaterialAnalyzer/Public/GLTFMaterialStatistics.h
deleted file mode 100644
