[GLTFExporter] refactor: make material analyzer more self-contained and safer from incorrect usage

commit 8522c781bdfa090bd347443ea8126be4e5f463b6
Date: Fri, 13 Nov 2020 00:25:43 +0100

    refactor: make material analyzer more self-contained and safer from incorrect usage
    (cherry picked from commit 8c072ea7cbeec5062f805140ff15a82dadedf06d​)

#skipci

diff --git a/Engine/Plugins/Enterprise/GLTFExporter/MigratedFromGit/Source/GLTFExporter/Private/Converters/GLTFMaterialUtility.cpp b/Engine/Plugins/Enterprise/GLTFExporter/MigratedFromGit/Source/GLTFExporter/Private/Converters/GLTFMaterialUtility.cpp
--- a/Engine/Plugins/Enterprise/GLTFExporter/MigratedFromGit/Source/GLTFExporter/Private/Converters/GLTFMaterialUtility.cpp
+++ b/Engine/Plugins/Enterprise/GLTFExporter/MigratedFromGit/Source/GLTFExporter/Private/Converters/GLTFMaterialUtility.cpp
@@ -541,12 +541,9 @@ bool FGLTFMaterialUtility::NeedsMeshData(const TArray<const UMaterialInterface*>
 	return false;
 }
 
-void FGLTFMaterialUtility::AnalyzeMaterialProperty(const UMaterialInterface* InMaterial, const FMaterialPropertyEx& InProperty, FGLTFMaterialAnalysis& OutMaterialAnalysis)
+void FGLTFMaterialUtility::AnalyzeMaterialProperty(const UMaterialInterface* InMaterial, const FMaterialPropertyEx& InProperty, FGLTFMaterialAnalysis& OutAnalysis)
 {
-	// TODO: use a shared UGLTFMaterialAnalyzer instance instead of creating a new one for each invocation
-
-	UGLTFMaterialAnalyzer* Analyzer = NewObject<UGLTFMaterialAnalyzer>();
-	Analyzer->AnalyzeMaterialProperty(InMaterial, InProperty.Type, InProperty.CustomOutput.ToString(), OutMaterialAnalysis);
+	UGLTFMaterialAnalyzer::AnalyzeMaterialPropertyEx(InMaterial, InProperty.Type, InProperty.CustomOutput.ToString(), OutAnalysis);
 }
 
 const UMaterialInterface* FGLTFMaterialUtility::GetInterface(const UMaterialInterface* Material)
diff --git a/Engine/Plugins/Enterprise/GLTFExporter/MigratedFromGit/Source/GLTFExporter/Private/Converters/GLTFMaterialUtility.h b/Engine/Plugins/Enterprise/GLTFExporter/MigratedFromGit/Source/GLTFExporter/Private/Converters/GLTFMaterialUtility.h
--- a/Engine/Plugins/Enterprise/GLTFExporter/MigratedFromGit/Source/GLTFExporter/Private/Converters/GLTFMaterialUtility.h
+++ b/Engine/Plugins/Enterprise/GLTFExporter/MigratedFromGit/Source/GLTFExporter/Private/Converters/GLTFMaterialUtility.h
@@ -77,7 +77,7 @@ struct FGLTFMaterialUtility
 	static bool NeedsMeshData(const UMaterialInterface* Material);
 	static bool NeedsMeshData(const TArray<const UMaterialInterface*>& Materials);
 
-	static void AnalyzeMaterialProperty(const UMaterialInterface* Material, const FMaterialPropertyEx& InProperty, FGLTFMaterialAnalysis& OutMaterialAnalysis);
+	static void AnalyzeMaterialProperty(const UMaterialInterface* Material, const FMaterialPropertyEx& InProperty, FGLTFMaterialAnalysis& OutAnalysis);
 
 	static const UMaterialInterface* GetInterface(const UMaterialInterface* Material);
 	static const UMaterialInterface* GetInterface(const FStaticMaterial& Material);
diff --git a/Engine/Plugins/Enterprise/GLTFExporter/MigratedFromGit/Source/GLTFMaterialAnalyzer/Private/GLTFMaterialAnalyzer.cpp b/Engine/Plugins/Enterprise/GLTFExporter/MigratedFromGit/Source/GLTFMaterialAnalyzer/Private/GLTFMaterialAnalyzer.cpp
--- a/Engine/Plugins/Enterprise/GLTFExporter/MigratedFromGit/Source/GLTFMaterialAnalyzer/Private/GLTFMaterialAnalyzer.cpp
+++ b/Engine/Plugins/Enterprise/GLTFExporter/MigratedFromGit/Source/GLTFMaterialAnalyzer/Private/GLTFMaterialAnalyzer.cpp
@@ -5,24 +5,37 @@
 #include "GLTFMaterialBaking/Public/MaterialBakingStructures.h"
 #include "Materials/HLSLMaterialTranslator.h"
 
-void UGLTFMaterialAnalyzer::AnalyzeMaterialProperty(const UMaterialInterface* InMaterial, const EMaterialProperty& InProperty, const FString& InCustomOutput, FGLTFMaterialAnalysis& OutAnalysis)
+UGLTFMaterialAnalyzer::UGLTFMaterialAnalyzer()
 {
-	Property = InProperty;
-	CustomOutput = InCustomOutput;
-	Material = const_cast<UMaterialInterface*>(InMaterial);
-	Analysis = &OutAnalysis;
+	ResetToDefaults();
+}
+
+void UGLTFMaterialAnalyzer::ResetToDefaults()
+{
+	Property = MP_MAX;
+	CustomOutput = {};
+	Material = nullptr;
+	Analysis = nullptr;
+}
+
+void UGLTFMaterialAnalyzer::AnalyzeMaterialPropertyEx(const UMaterialInterface* InMaterial, const EMaterialProperty& InProperty, const FString& InCustomOutput, FGLTFMaterialAnalysis& OutAnalysis)
+{
+	// TODO: use a shared UGLTFMaterialAnalyzer instance instead of creating a new one for each invocation
+	UGLTFMaterialAnalyzer* Analyzer = NewObject<UGLTFMaterialAnalyzer>();
+
+	Analyzer->Property = InProperty;
+	Analyzer->CustomOutput = InCustomOutput;
+	Analyzer->Material = const_cast<UMaterialInterface*>(InMaterial);
+	Analyzer->Analysis = &OutAnalysis;
 
 	// NOTE: When analyzing custom outputs, the property *must* be set to MP_MAX or the compiler will refuse to compile the output
 	const EMaterialProperty SafeProperty = InProperty == MP_CustomOutput ? MP_MAX : InProperty;
 
 	int32 DummyNumTextureCoordinates; // Dummy value from built-in analysis not used since it's insufficient
 	bool DummyRequiresVertexData; // Dummy value from built-in analysis not used since it's insufficient
-	Super::AnalyzeMaterialProperty(SafeProperty, DummyNumTextureCoordinates, DummyRequiresVertexData);
+	Analyzer->AnalyzeMaterialProperty(SafeProperty, DummyNumTextureCoordinates, DummyRequiresVertexData);
 
-	Property = MP_MAX;
-	CustomOutput = {};
-	Material = nullptr;
-	Analysis = nullptr;
+	Analyzer->ResetToDefaults();
 }
 
 UMaterialExpressionCustomOutput* UGLTFMaterialAnalyzer::GetCustomOutputExpression() const
@@ -41,11 +54,16 @@ UMaterialExpressionCustomOutput* UGLTFMaterialAnalyzer::GetCustomOutputExpressio
 
 FMaterialResource* UGLTFMaterialAnalyzer::GetMaterialResource(ERHIFeatureLevel::Type InFeatureLevel, EMaterialQualityLevel::Type QualityLevel)
 {
-	return Material->GetMaterialResource(InFeatureLevel, QualityLevel);
+	return Material != nullptr ? Material->GetMaterialResource(InFeatureLevel, QualityLevel) : nullptr;
 }
 
 int32 UGLTFMaterialAnalyzer::CompilePropertyEx(FMaterialCompiler* Compiler, const FGuid& AttributeID)
 {
+	if (Material == nullptr)
+	{
+		return INDEX_NONE;
+	}
+
 	class FTranslatorHack : public FHLSLMaterialTranslator
 	{
 		using FHLSLMaterialTranslator::FHLSLMaterialTranslator;
@@ -59,7 +77,12 @@ int32 UGLTFMaterialAnalyzer::CompilePropertyEx(FMaterialCompiler* Compiler, cons
 	if (Property == MP_CustomOutput)
 	{
 		UMaterialExpressionCustomOutput* CustomOutputExpression = GetCustomOutputExpression();
-		Result = CustomOutputExpression != nullptr ? CustomOutputExpression->Compile(Compiler, 0) : INDEX_NONE;
+		if (CustomOutputExpression == nullptr)
+		{
+			return INDEX_NONE;
+		}
+
+		Result = CustomOutputExpression->Compile(Compiler, 0);
 	}
 	else
 	{
@@ -88,5 +111,5 @@ int32 UGLTFMaterialAnalyzer::CompilePropertyEx(FMaterialCompiler* Compiler, cons
 
 bool UGLTFMaterialAnalyzer::IsPropertyActive(EMaterialProperty InProperty) const
 {
-	return Material->IsPropertyActive(InProperty);
+	return Material != nullptr && Material->IsPropertyActive(InProperty);
 }
diff --git a/Engine/Plugins/Enterprise/GLTFExporter/MigratedFromGit/Source/GLTFMaterialAnalyzer/Public/GLTFMaterialAnalyzer.h b/Engine/Plugins/Enterprise/GLTFExporter/MigratedFromGit/Source/GLTFMaterialAnalyzer/Public/GLTFMaterialAnalyzer.h
--- a/Engine/Plugins/Enterprise/GLTFExporter/MigratedFromGit/Source/GLTFMaterialAnalyzer/Public/GLTFMaterialAnalyzer.h
+++ b/Engine/Plugins/Enterprise/GLTFExporter/MigratedFromGit/Source/GLTFMaterialAnalyzer/Public/GLTFMaterialAnalyzer.h
@@ -7,17 +7,21 @@
 
 struct FGLTFMaterialAnalysis;
 
-UCLASS()
+UCLASS(NotBlueprintType, Transient)
 class GLTFMATERIALANALYZER_API UGLTFMaterialAnalyzer : public UMaterialInterface
 {
 	GENERATED_BODY()
 
 public:
 
-	void AnalyzeMaterialProperty(const UMaterialInterface* InMaterial, const EMaterialProperty& InProperty, const FString& InCustomOutput, FGLTFMaterialAnalysis& OutAnalysis);
+	static void AnalyzeMaterialPropertyEx(const UMaterialInterface* InMaterial, const EMaterialProperty& InProperty, const FString& InCustomOutput, FGLTFMaterialAnalysis& OutAnalysis);
 
 private:
 
+	UGLTFMaterialAnalyzer();
+
+	void ResetToDefaults();
+
 	UMaterialExpressionCustomOutput* GetCustomOutputExpression() const;
 
 	virtual FMaterialResource* GetMaterialResource(ERHIFeatureLevel::Type InFeatureLevel, EMaterialQualityLevel::Type QualityLevel) override;
