[GLTFExporter] feat: improve texture coordinate detection when baking materials

commit 1a18e3b6f544ea1719f10a5a87bc89db0fdf6734
Date: Fri, 13 Nov 2020 23:23:44 +0100

    feat: improve texture coordinate detection when baking materials
    (cherry picked from commit c91bbdefbe2a5d53c31f99731986d081e37f2685​)

#skipci

diff --git a/Engine/Plugins/Enterprise/GLTFExporter/MigratedFromGit/Source/GLTFExporter/Private/Converters/GLTFMaterialUtility.cpp b/Engine/Plugins/Enterprise/GLTFExporter/MigratedFromGit/Source/GLTFExporter/Private/Converters/GLTFMaterialUtility.cpp
--- a/Engine/Plugins/Enterprise/GLTFExporter/MigratedFromGit/Source/GLTFExporter/Private/Converters/GLTFMaterialUtility.cpp
+++ b/Engine/Plugins/Enterprise/GLTFExporter/MigratedFromGit/Source/GLTFExporter/Private/Converters/GLTFMaterialUtility.cpp
@@ -387,71 +387,18 @@ bool FGLTFMaterialUtility::TryGetTextureCoordinateIndex(const UMaterialExpressio
 	return false;
 }
 
-void FGLTFMaterialUtility::GetAllTextureCoordinateIndices(const FExpressionInput& Input, TSet<int32>& OutTexCoords)
+void FGLTFMaterialUtility::GetAllTextureCoordinateIndices(const UMaterialInterface* InMaterial, const FMaterialPropertyEx& InProperty, int32& OutNumTexCoords, TBitArray<>& OutTexCoords)
 {
-	UMaterialExpression* InputExpression = Input.Expression;
-	if (InputExpression == nullptr)
-	{
-		return;
-	}
-
-	TArray<UMaterialExpression*> Expressions;
-	InputExpression->GetAllInputExpressions(Expressions);
-
-	ExpandAllFunctionExpressions(Expressions);
-
-	for (const UMaterialExpression* Expression : Expressions)
-	{
-		if (const UMaterialExpressionTextureSample* TextureSample = Cast<UMaterialExpressionTextureSample>(Expression))
-		{
-			if (TextureSample->Coordinates.Expression == nullptr)
-			{
-				OutTexCoords.Add(TextureSample->ConstCoordinate);
-			}
-		}
-		else if (const UMaterialExpressionTextureCoordinate* TextureCoordinate = Cast<UMaterialExpressionTextureCoordinate>(Expression))
-		{
-			OutTexCoords.Add(TextureCoordinate->CoordinateIndex);
-		}
-	}
-}
+	FGLTFMaterialAnalysis Analysis;
+	AnalyzeMaterialProperty(InMaterial, InProperty, Analysis);
 
-void FGLTFMaterialUtility::ExpandAllFunctionExpressions(TArray<UMaterialExpression*>& InOutExpressions)
-{
-	TArray<UMaterialExpression*> UnexpandedExpressions(InOutExpressions);
+	OutNumTexCoords = 0;
+	OutTexCoords = Analysis.TextureCoordinates;
 
-	for (const UMaterialExpression* UnexpandedExpression : UnexpandedExpressions)
+	for (int32 Index = 0; Index < OutTexCoords.Num(); Index++)
 	{
-		if (const UMaterialExpressionMaterialFunctionCall* FunctionCall = Cast<UMaterialExpressionMaterialFunctionCall>(UnexpandedExpression))
-		{
-			if (UMaterialFunctionInterface* Function = FunctionCall->MaterialFunction)
-			{
-				Function->GetAllExpressionsOfType<UMaterialExpression>(InOutExpressions);
-			}
-		}
-		else if (const UMaterialExpressionMaterialAttributeLayers* AttributeLayers = Cast<UMaterialExpressionMaterialAttributeLayers>(UnexpandedExpression))
-		{
-			for (UMaterialFunctionInterface* Layer : AttributeLayers->GetLayers())
-			{
-				if (Layer != nullptr)
-				{
-					Layer->GetAllExpressionsOfType<UMaterialExpression>(InOutExpressions);
-				}
-			}
-
-			for (UMaterialFunctionInterface* Blend : AttributeLayers->GetBlends())
-			{
-				if (Blend != nullptr)
-				{
-					Blend->GetAllExpressionsOfType<UMaterialExpression>(InOutExpressions);
-				}
-			}
-		}
+		OutNumTexCoords += OutTexCoords[Index] ? 1 : 0;
 	}
-
-	InOutExpressions.RemoveAll([](const UMaterialExpression* Expression) {
-		return Expression->IsA<UMaterialExpressionMaterialFunctionCall>() || Expression->IsA<UMaterialExpressionMaterialAttributeLayers>();
-	});
 }
 
 EMaterialShadingModel FGLTFMaterialUtility::GetRichestShadingModel(const FMaterialShadingModelField& ShadingModels)
diff --git a/Engine/Plugins/Enterprise/GLTFExporter/MigratedFromGit/Source/GLTFExporter/Private/Converters/GLTFMaterialUtility.h b/Engine/Plugins/Enterprise/GLTFExporter/MigratedFromGit/Source/GLTFExporter/Private/Converters/GLTFMaterialUtility.h
--- a/Engine/Plugins/Enterprise/GLTFExporter/MigratedFromGit/Source/GLTFExporter/Private/Converters/GLTFMaterialUtility.h
+++ b/Engine/Plugins/Enterprise/GLTFExporter/MigratedFromGit/Source/GLTFExporter/Private/Converters/GLTFMaterialUtility.h
@@ -67,9 +67,7 @@ struct FGLTFMaterialUtility
 	static uint32 GetMaskComponentCount(const FExpressionInput& ExpressionInput);
 
 	static bool TryGetTextureCoordinateIndex(const UMaterialExpressionTextureSample* TextureSampler, int32& TexCoord, FGLTFJsonTextureTransform& Transform);
-	static void GetAllTextureCoordinateIndices(const FExpressionInput& Input, TSet<int32>& OutTexCoords);
-
-	static void ExpandAllFunctionExpressions(TArray<UMaterialExpression*>& InOutExpressions);
+	static void GetAllTextureCoordinateIndices(const UMaterialInterface* InMaterial, const FMaterialPropertyEx& InProperty, int32& OutNumTexCoords, TBitArray<>& OutTexCoords);
 
 	static EMaterialShadingModel GetRichestShadingModel(const FMaterialShadingModelField& ShadingModels);
 	static FString ShadingModelsToString(const FMaterialShadingModelField& ShadingModels);
diff --git a/Engine/Plugins/Enterprise/GLTFExporter/MigratedFromGit/Source/GLTFExporter/Private/Converters/GLTFMeshData.cpp b/Engine/Plugins/Enterprise/GLTFExporter/MigratedFromGit/Source/GLTFExporter/Private/Converters/GLTFMeshData.cpp
--- a/Engine/Plugins/Enterprise/GLTFExporter/MigratedFromGit/Source/GLTFExporter/Private/Converters/GLTFMeshData.cpp
+++ b/Engine/Plugins/Enterprise/GLTFExporter/MigratedFromGit/Source/GLTFExporter/Private/Converters/GLTFMeshData.cpp
@@ -5,6 +5,7 @@
 #include "Converters/GLTFMeshUtility.h"
 #include "StaticMeshAttributes.h"
 #include "Developer/MeshMergeUtilities/Private/MeshMergeHelpers.h"
+#include "Rendering/SkeletalMeshRenderData.h"
 
 FGLTFMeshData::FGLTFMeshData(const UStaticMesh* StaticMesh, const UStaticMeshComponent* StaticMeshComponent, int32 LODIndex)
 {
@@ -12,14 +13,18 @@ FGLTFMeshData::FGLTFMeshData(const UStaticMesh* StaticMesh, const UStaticMeshCom
 
 	if (StaticMeshComponent != nullptr)
 	{
-		FMeshMergeHelpers::RetrieveMesh(StaticMeshComponent, LODIndex, Description, true);
 		Name = FGLTFNameUtility::GetName(StaticMeshComponent);
+		FMeshMergeHelpers::RetrieveMesh(StaticMeshComponent, LODIndex, Description, true);
 	}
 	else
 	{
-		FMeshMergeHelpers::RetrieveMesh(StaticMesh, LODIndex, Description);
 		StaticMesh->GetName(Name);
+		FMeshMergeHelpers::RetrieveMesh(StaticMesh, LODIndex, Description);
 	}
+
+	TexCoord = StaticMesh->LightMapCoordinateIndex;
+
+	// TODO: add warning if texture coordinate has overlap
 }
 
 FGLTFMeshData::FGLTFMeshData(const USkeletalMesh* SkeletalMesh, const USkeletalMeshComponent* SkeletalMeshComponent, int32 LODIndex)
@@ -28,11 +33,13 @@ FGLTFMeshData::FGLTFMeshData(const USkeletalMesh* SkeletalMesh, const USkeletalM
 
 	if (SkeletalMeshComponent != nullptr)
 	{
-		FMeshMergeHelpers::RetrieveMesh(const_cast<USkeletalMeshComponent*>(SkeletalMeshComponent), LODIndex, Description, true);
 		Name = FGLTFNameUtility::GetName(SkeletalMeshComponent);
+		FMeshMergeHelpers::RetrieveMesh(const_cast<USkeletalMeshComponent*>(SkeletalMeshComponent), LODIndex, Description, true);
 	}
 	else
 	{
+		SkeletalMesh->GetName(Name);
+
 		// NOTE: this is a workaround for the fact that there's no overload for FMeshMergeHelpers::RetrieveMesh
 		// that accepts a USkeletalMesh, only a USkeletalMeshComponent.
 		// Writing a custom utility function that would work on a "standalone" skeletal mesh is problematic
@@ -56,7 +63,11 @@ FGLTFMeshData::FGLTFMeshData(const USkeletalMesh* SkeletalMesh, const USkeletalM
 				World->DestroyActor(Actor, false, false);
 			}
 		}
-
-		SkeletalMesh->GetName(Name);
 	}
+
+	// TODO: don't assume last UV channel is non-overlapping
+	const uint32 NumTexCoords = SkeletalMesh->GetResourceForRendering()->LODRenderData[LODIndex].StaticVertexBuffers.StaticMeshVertexBuffer.GetNumTexCoords();
+	TexCoord = NumTexCoords - 1;
+
+	// TODO: add warning if texture coordinate has overlap
 }
diff --git a/Engine/Plugins/Enterprise/GLTFExporter/MigratedFromGit/Source/GLTFExporter/Private/Converters/GLTFMeshData.h b/Engine/Plugins/Enterprise/GLTFExporter/MigratedFromGit/Source/GLTFExporter/Private/Converters/GLTFMeshData.h
--- a/Engine/Plugins/Enterprise/GLTFExporter/MigratedFromGit/Source/GLTFExporter/Private/Converters/GLTFMeshData.h
+++ b/Engine/Plugins/Enterprise/GLTFExporter/MigratedFromGit/Source/GLTFExporter/Private/Converters/GLTFMeshData.h
@@ -10,6 +10,7 @@ struct FGLTFMeshData
 	FGLTFMeshData(const UStaticMesh* StaticMesh, const UStaticMeshComponent* StaticMeshComponent, int32 LODIndex);
 	FGLTFMeshData(const USkeletalMesh* SkeletalMesh, const USkeletalMeshComponent* SkeletalMeshComponent, int32 LODIndex);
 
-	FMeshDescription Description;
 	FString Name;
+	FMeshDescription Description;
+	int32 TexCoord;
 };
diff --git a/Engine/Plugins/Enterprise/GLTFExporter/MigratedFromGit/Source/GLTFExporter/Private/Tasks/GLTFMaterialTasks.cpp b/Engine/Plugins/Enterprise/GLTFExporter/MigratedFromGit/Source/GLTFExporter/Private/Tasks/GLTFMaterialTasks.cpp
--- a/Engine/Plugins/Enterprise/GLTFExporter/MigratedFromGit/Source/GLTFExporter/Private/Tasks/GLTFMaterialTasks.cpp
+++ b/Engine/Plugins/Enterprise/GLTFExporter/MigratedFromGit/Source/GLTFExporter/Private/Tasks/GLTFMaterialTasks.cpp
@@ -5,6 +5,7 @@
 #include "Converters/GLTFNameUtility.h"
 #include "Converters/GLTFMaterialUtility.h"
 #include "Builders/GLTFContainerBuilder.h"
+#include "MaterialPropertyEx.h"
 #include "Materials/MaterialExpressionConstant.h"
 #include "Materials/MaterialExpressionConstant2Vector.h"
 #include "Materials/MaterialExpressionConstant3Vector.h"
@@ -1336,35 +1337,32 @@ bool FGLTFMaterialTask::TryGetBakedMaterialProperty(FGLTFJsonTextureInfo& OutTex
 
 FGLTFPropertyBakeOutput FGLTFMaterialTask::BakeMaterialProperty(EMaterialProperty Property, int32& OutTexCoord, FIntPoint PreferredTextureSize, bool bCopyAlphaFromRedChannel) const
 {
-	const FExpressionInput* PropertyInput = FGLTFMaterialUtility::GetInputForProperty(Material, Property);
-	TSet<int32> TexCoords;
-
-	FGLTFMaterialUtility::GetAllTextureCoordinateIndices(*PropertyInput, TexCoords);
-	if (TexCoords.Num() > 0)
+	if (MeshData == nullptr)
 	{
-		// TODO: is there a nicer way to get the first element in a TSet?
-		for (int32 TexCoord : TexCoords)
-		{
-			OutTexCoord = TexCoord;
-			break;
-		}
+		// TODO: replace dummy enum MP_CustomOutput workaround for ClearCoatBottomNormal
+		const FMaterialPropertyEx PropertyEx = Property != MP_CustomOutput ? FMaterialPropertyEx(Property) : FMaterialPropertyEx(TEXT("ClearCoatBottomNormal"));
 
-		if (TexCoords.Num() > 1)
+		int32 NumTexCoords;
+		TBitArray<> TexCoords;
+		FGLTFMaterialUtility::GetAllTextureCoordinateIndices(Material, PropertyEx, NumTexCoords, TexCoords);
+
+		if (NumTexCoords > 0)
 		{
-			// TODO: report warning (multiple texture coordinates found, will use first)
+			OutTexCoord = TexCoords.Find(true);
 
-			// TODO: replace this hardcoded hack with something more configurable and proper
-			if (Property == MP_AmbientOcclusion && TexCoords.Contains(1))
+			if (TexCoords.Num() > 1)
 			{
-				OutTexCoord = 1; // assume ambient occlusion uses TexCoord1 when multiple
+				// TODO: report warning (multiple texture coordinates found, will use first)
 			}
 		}
-
-		// TODO: should we perhaps always use the lightmap coordinate index for baking to guarantee unique uvs?
+		else
+		{
+			OutTexCoord = 0;
+		}
 	}
 	else
 	{
-		OutTexCoord = 0; // assume TexCoord0 even thought property seems to be texture coordinate independent
+		OutTexCoord = MeshData->TexCoord;
 	}
 
 	const FIntPoint DefaultTextureSize = Builder.GetDefaultMaterialBakeSize();
@@ -1381,11 +1379,6 @@ FGLTFPropertyBakeOutput FGLTFMaterialTask::BakeMaterialProperty(EMaterialPropert
 		SectionIndices,
 		bCopyAlphaFromRedChannel);
 
-	if (!PropertyBakeOutput.bIsConstant && TexCoords.Num() == 0)
-	{
-		// TODO: report warning about property not being constant yet texture coordinate independent
-	}
-
 	return PropertyBakeOutput;
 }
 
diff --git a/Engine/Plugins/Enterprise/GLTFExporter/MigratedFromGit/Source/GLTFMaterialAnalyzer/Private/GLTFMaterialAnalyzer.cpp b/Engine/Plugins/Enterprise/GLTFExporter/MigratedFromGit/Source/GLTFMaterialAnalyzer/Private/GLTFMaterialAnalyzer.cpp
--- a/Engine/Plugins/Enterprise/GLTFExporter/MigratedFromGit/Source/GLTFMaterialAnalyzer/Private/GLTFMaterialAnalyzer.cpp
+++ b/Engine/Plugins/Enterprise/GLTFExporter/MigratedFromGit/Source/GLTFMaterialAnalyzer/Private/GLTFMaterialAnalyzer.cpp
@@ -90,13 +90,6 @@ int32 UGLTFMaterialAnalyzer::CompilePropertyEx(FMaterialCompiler* Compiler, cons
 	}
 
 	Analysis->TextureCoordinates = Translator->AllocatedUserTexCoords;
-	Analysis->NumTextureCoordinates = 0;
-
-	for (int32 TexCoordIndex = 0; TexCoordIndex < Translator->AllocatedUserTexCoords.Num(); TexCoordIndex++)
-	{
-		const bool IsTexCoordUsed = Translator->AllocatedUserTexCoords[TexCoordIndex];
-		Analysis->NumTextureCoordinates += IsTexCoordUsed ? 1 : 0;
-	}
 
 	// TODO: investigate if we need to check more conditions to determine that vertex data is required
 	Analysis->bRequiresVertexData =
diff --git a/Engine/Plugins/Enterprise/GLTFExporter/MigratedFromGit/Source/GLTFMaterialAnalyzer/Public/GLTFMaterialAnalysis.h b/Engine/Plugins/Enterprise/GLTFExporter/MigratedFromGit/Source/GLTFMaterialAnalyzer/Public/GLTFMaterialAnalysis.h
--- a/Engine/Plugins/Enterprise/GLTFExporter/MigratedFromGit/Source/GLTFMaterialAnalyzer/Public/GLTFMaterialAnalysis.h
+++ b/Engine/Plugins/Enterprise/GLTFExporter/MigratedFromGit/Source/GLTFMaterialAnalyzer/Public/GLTFMaterialAnalysis.h
@@ -7,12 +7,6 @@ struct FGLTFMaterialAnalysis
 	/** Tracks the texture coordinates used by this material */
 	TBitArray<> TextureCoordinates;
 
-	/** Number of texture coordinates used by this material */
-	uint32 NumTextureCoordinates;
-
-	/** Will contain all the shading models picked up from the material expression graph */
-	FMaterialShadingModelField ShadingModels;
-
 	/** True if this material reads any vertex data */
 	bool bRequiresVertexData;
 };
