[GLTFExporter] refactor: improve MaterialPropertyEx in embedded baking module

commit e935ed079e8dd0d2ad745f61629001a99aa5d850
Date: Fri, 13 Nov 2020 23:54:18 +0100

    refactor: improve MaterialPropertyEx in embedded baking module
    (cherry picked from commit 3ffbe9dc329a6a5ad2140d13d9c6b0e33e02db52​)

#skipci

diff --git a/Engine/Plugins/Enterprise/GLTFExporter/MigratedFromGit/Source/GLTFMaterialBaking/Public/MaterialPropertyEx.h b/Engine/Plugins/Enterprise/GLTFExporter/MigratedFromGit/Source/GLTFMaterialBaking/Public/MaterialPropertyEx.h
--- a/Engine/Plugins/Enterprise/GLTFExporter/MigratedFromGit/Source/GLTFMaterialBaking/Public/MaterialPropertyEx.h
+++ b/Engine/Plugins/Enterprise/GLTFExporter/MigratedFromGit/Source/GLTFMaterialBaking/Public/MaterialPropertyEx.h
@@ -23,24 +23,24 @@ struct FMaterialPropertyEx
 		, CustomOutput(CustomOutput)
 	{}
 
-	FORCEINLINE bool operator ==(const FMaterialPropertyEx& Other) const
+	bool IsCustomOutput() const
 	{
-		return Type == Other.Type && (Type != MP_CustomOutput || CustomOutput == Other.CustomOutput);
+		return Type == MP_CustomOutput;
 	}
 
-	FORCEINLINE bool operator !=(const FMaterialPropertyEx& Other) const
+	bool Equals(const FMaterialPropertyEx& Other) const
 	{
-		return !(*this == Other);
+		return Type == Other.Type && (!IsCustomOutput() || CustomOutput == Other.CustomOutput);
 	}
 
-	friend FORCEINLINE uint32 GetTypeHash(const FMaterialPropertyEx& Other)
+	int32 Compare(const FMaterialPropertyEx& Other) const
 	{
-		return Other.Type != MP_CustomOutput ? GetTypeHash(Other.Type) : GetTypeHash(Other.CustomOutput);
+		return !IsCustomOutput() || !Other.IsCustomOutput() ? Type - Other.Type : CustomOutput.Compare(Other.CustomOutput);
 	}
 
 	FString ToString() const
 	{
-		if (Type != MP_CustomOutput)
+		if (!IsCustomOutput())
 		{
 			const UEnum* Enum = StaticEnum<EMaterialProperty>();
 			FName Name = Enum->GetNameByValue(Type);
@@ -48,15 +48,48 @@ struct FMaterialPropertyEx
 			TrimmedName.RemoveFromStart(TEXT("MP_"), ESearchCase::CaseSensitive);
 			return TrimmedName;
 		}
-		else
-		{
-			return CustomOutput.ToString();
-		}
+
+		return CustomOutput.ToString();
+	}
+
+	bool operator ==(const FMaterialPropertyEx& Other) const
+	{
+		return Equals(Other);
+	}
+
+	bool operator !=(const FMaterialPropertyEx& Other) const
+	{
+		return !Equals(Other);
+	}
+
+	bool operator >=(const FMaterialPropertyEx& Other) const
+	{
+		return Compare(Other) >= 0;
+	}
+
+	bool operator >(const FMaterialPropertyEx& Other) const
+	{
+		return Compare(Other) > 0;
 	}
-	
+
+	bool operator <=(const FMaterialPropertyEx& Other) const
+	{
+		return Compare(Other) <= 0;
+	}
+
+	bool operator <(const FMaterialPropertyEx& Other) const
+	{
+		return Compare(Other) < 0;
+	}
+
+	friend uint32 GetTypeHash(const FMaterialPropertyEx& Other)
+	{
+		return !Other.IsCustomOutput() ? GetTypeHash(Other.Type) : GetTypeHash(Other.CustomOutput);
+	}
+
 	/** The material property */
 	EMaterialProperty Type;
-	
+
 	/** The name of a specific custom output. Only used if property is MP_CustomOutput */
 	FName CustomOutput;
 };
