[GLTFExporter] refactor: replace usage of EMaterialProperty with FMaterialPropertyEx

commit ad71d5afb36303953dd166da41e529873be3eb3c
Date: Fri, 13 Nov 2020 23:59:55 +0100

    refactor: replace usage of EMaterialProperty with FMaterialPropertyEx
    (cherry picked from commit 485a1000d2d80ce70bb398367063c00218db505a​)

#skipci

diff --git a/Engine/Plugins/Enterprise/GLTFExporter/MigratedFromGit/Source/GLTFExporter/Private/Converters/GLTFMaterialUtility.cpp b/Engine/Plugins/Enterprise/GLTFExporter/MigratedFromGit/Source/GLTFExporter/Private/Converters/GLTFMaterialUtility.cpp
--- a/Engine/Plugins/Enterprise/GLTFExporter/MigratedFromGit/Source/GLTFExporter/Private/Converters/GLTFMaterialUtility.cpp
+++ b/Engine/Plugins/Enterprise/GLTFExporter/MigratedFromGit/Source/GLTFExporter/Private/Converters/GLTFMaterialUtility.cpp
@@ -23,58 +23,16 @@ UMaterialInterface* FGLTFMaterialUtility::GetDefault()
 	return DefaultMaterial;
 }
 
-bool FGLTFMaterialUtility::IsNormalMap(EMaterialProperty Property)
+bool FGLTFMaterialUtility::IsNormalMap(const FMaterialPropertyEx& Property)
 {
-	return Property == MP_Normal || Property == MP_CustomOutput /* ClearCoatBottomNormal */;
+	return Property == MP_Normal || Property == TEXT("ClearCoatBottomNormal");
 }
 
-const TCHAR* FGLTFMaterialUtility::GetPropertyName(EMaterialProperty Property)
-{
-	// TODO: replace this hardcoded list with some lookup function that should be in the engine.
-
-	switch (Property)
-	{
-		case MP_EmissiveColor:          return TEXT("EmissiveColor");
-		case MP_Opacity:                return TEXT("Opacity");
-		case MP_OpacityMask:            return TEXT("OpacityMask");
-		case MP_BaseColor:              return TEXT("BaseColor");
-		case MP_Metallic:               return TEXT("Metallic");
-		case MP_Specular:               return TEXT("Specular");
-		case MP_Roughness:              return TEXT("Roughness");
-		case MP_Anisotropy:             return TEXT("Anisotropy");
-		case MP_Normal:                 return TEXT("Normal");
-		case MP_Tangent:                return TEXT("Tangent");
-		case MP_WorldPositionOffset:    return TEXT("WorldPositionOffset");
-		case MP_WorldDisplacement:      return TEXT("WorldDisplacement");
-		case MP_TessellationMultiplier: return TEXT("TessellationMultiplier");
-		case MP_SubsurfaceColor:        return TEXT("SubsurfaceColor");
-		case MP_CustomData0:            return TEXT("ClearCoat");
-		case MP_CustomData1:            return TEXT("ClearCoatRoughness");
-		case MP_AmbientOcclusion:       return TEXT("AmbientOcclusion");
-		case MP_Refraction:             return TEXT("Refraction");
-		case MP_CustomizedUVs0:         return TEXT("CustomizedUV0");
-		case MP_CustomizedUVs1:         return TEXT("CustomizedUV1");
-		case MP_CustomizedUVs2:         return TEXT("CustomizedUV2");
-		case MP_CustomizedUVs3:         return TEXT("CustomizedUV3");
-		case MP_CustomizedUVs4:         return TEXT("CustomizedUV4");
-		case MP_CustomizedUVs5:         return TEXT("CustomizedUV5");
-		case MP_CustomizedUVs6:         return TEXT("CustomizedUV6");
-		case MP_CustomizedUVs7:         return TEXT("CustomizedUV7");
-		case MP_PixelDepthOffset:       return TEXT("PixelDepthOffset");
-		case MP_ShadingModel:           return TEXT("ShadingModel");
-		case MP_MaterialAttributes:     return TEXT("MaterialAttributes");
-		case MP_CustomOutput:           return TEXT("ClearCoatBottomNormal");
-		default:
-			checkNoEntry();
-			return TEXT("");
-	}
-}
-
-FVector4 FGLTFMaterialUtility::GetPropertyDefaultValue(EMaterialProperty Property)
+FVector4 FGLTFMaterialUtility::GetPropertyDefaultValue(const FMaterialPropertyEx& Property)
 {
 	// TODO: replace with GMaterialPropertyAttributesMap lookup (when public API available)
 
-	switch (Property)
+	switch (Property.Type)
 	{
 		case MP_EmissiveColor:          return FVector4(0,0,0,0);
 		case MP_Opacity:                return FVector4(1,0,0,0);
@@ -104,8 +62,7 @@ FVector4 FGLTFMaterialUtility::GetPropertyDefaultValue(EMaterialProperty Propert
 		return FVector4(0,0,0,0);
 	}
 
-	// TODO: replace workaround for ClearCoatBottomNormal
-	if (Property == MP_CustomOutput)
+	if (Property == TEXT("ClearCoatBottomNormal"))
 	{
 		return FVector4(0,0,1,0);
 	}
@@ -114,11 +71,11 @@ FVector4 FGLTFMaterialUtility::GetPropertyDefaultValue(EMaterialProperty Propert
 	return FVector4();
 }
 
-FVector4 FGLTFMaterialUtility::GetPropertyMask(EMaterialProperty Property)
+FVector4 FGLTFMaterialUtility::GetPropertyMask(const FMaterialPropertyEx& Property)
 {
 	// TODO: replace with GMaterialPropertyAttributesMap lookup (when public API available)
 
-	switch (Property)
+	switch (Property.Type)
 	{
 		case MP_EmissiveColor:          return FVector4(1,1,1,0);
 		case MP_Opacity:                return FVector4(1,0,0,0);
@@ -148,8 +105,7 @@ FVector4 FGLTFMaterialUtility::GetPropertyMask(EMaterialProperty Property)
 		return FVector4(1,1,0,0);
 	}
 
-	// TODO: replace workaround for ClearCoatBottomNormal
-	if (Property == MP_CustomOutput)
+	if (Property == TEXT("ClearCoatBottomNormal"))
 	{
 		return FVector4(1,1,1,0);
 	}
@@ -158,17 +114,16 @@ FVector4 FGLTFMaterialUtility::GetPropertyMask(EMaterialProperty Property)
 	return FVector4();
 }
 
-const FExpressionInput* FGLTFMaterialUtility::GetInputForProperty(const UMaterialInterface* Material, EMaterialProperty Property)
+const FExpressionInput* FGLTFMaterialUtility::GetInputForProperty(const UMaterialInterface* Material, const FMaterialPropertyEx& Property)
 {
-	// TODO: replace workaround for ClearCoatBottomNormal
-	if (Property == MP_CustomOutput)
+	if (Property.IsCustomOutput())
 	{
-		const UMaterialExpressionCustomOutput* CustomOutput = GetCustomOutputByName(Material, TEXT("ClearCoatBottomNormal"));
+		const UMaterialExpressionCustomOutput* CustomOutput = GetCustomOutputByName(Material, Property.CustomOutput.ToString());
 		return CustomOutput != nullptr ? &CastChecked<UMaterialExpressionClearCoatNormalCustomOutput>(CustomOutput)->Input : nullptr;
 	}
 
 	UMaterial* UnderlyingMaterial = const_cast<UMaterial*>(Material->GetMaterial());
-	return UnderlyingMaterial->GetExpressionInputForProperty(Property);
+	return UnderlyingMaterial->GetExpressionInputForProperty(Property.Type);
 }
 
 const UMaterialExpressionCustomOutput* FGLTFMaterialUtility::GetCustomOutputByName(const UMaterialInterface* Material, const FString& Name)
@@ -178,7 +133,7 @@ const UMaterialExpressionCustomOutput* FGLTFMaterialUtility::GetCustomOutputByNa
 	for (const UMaterialExpression* Expression : Material->GetMaterial()->Expressions)
 	{
 		const UMaterialExpressionCustomOutput* CustomOutput = Cast<UMaterialExpressionCustomOutput>(Expression);
-		if (CustomOutput != nullptr && CustomOutput->GetFunctionName() == Name)
+		if (CustomOutput != nullptr && CustomOutput->GetDisplayName() == Name)
 		{
 			return CustomOutput;
 		}
@@ -248,10 +203,8 @@ bool FGLTFMaterialUtility::CombineTextures(TArray<FColor>& OutPixels, const TArr
 	return bReadSuccessful;
 }
 
-FGLTFPropertyBakeOutput FGLTFMaterialUtility::BakeMaterialProperty(const FIntPoint& OutputSize, EMaterialProperty Property, const UMaterialInterface* Material, int32 TexCoord, const FMeshDescription* MeshDescription, const TArray<int32>& MeshSectionIndices, bool bCopyAlphaFromRedChannel)
+FGLTFPropertyBakeOutput FGLTFMaterialUtility::BakeMaterialProperty(const FIntPoint& OutputSize, const FMaterialPropertyEx& Property, const UMaterialInterface* Material, int32 TexCoord, const FMeshDescription* MeshDescription, const TArray<int32>& MeshSectionIndices, bool bCopyAlphaFromRedChannel)
 {
-	const FMaterialPropertyEx PropertyEx(Property, Property == MP_CustomOutput ? FName(TEXT("ClearCoatBottomNormal")) : NAME_None);
-
 	FMeshData MeshSet;
 	MeshSet.TextureCoordinateBox = { { 0.0f, 0.0f }, { 1.0f, 1.0f } };
 	MeshSet.TextureCoordinateIndex = TexCoord;
@@ -260,7 +213,7 @@ FGLTFPropertyBakeOutput FGLTFMaterialUtility::BakeMaterialProperty(const FIntPoi
 
 	FMaterialDataEx MatSet;
 	MatSet.Material = const_cast<UMaterialInterface*>(Material);
-	MatSet.PropertySizes.Add(PropertyEx, OutputSize);
+	MatSet.PropertySizes.Add(Property, OutputSize);
 
 	TArray<FMeshData*> MeshSettings;
 	TArray<FMaterialDataEx*> MatSettings;
@@ -274,8 +227,8 @@ FGLTFPropertyBakeOutput FGLTFMaterialUtility::BakeMaterialProperty(const FIntPoi
 
 	FBakeOutputEx& BakeOutput = BakeOutputs[0];
 
-	TArray<FColor> BakedPixels = MoveTemp(BakeOutput.PropertyData.FindChecked(PropertyEx));
-	const FIntPoint BakedSize = BakeOutput.PropertySizes.FindChecked(PropertyEx);
+	TArray<FColor> BakedPixels = MoveTemp(BakeOutput.PropertyData.FindChecked(Property));
+	const FIntPoint BakedSize = BakeOutput.PropertySizes.FindChecked(Property);
 	const float EmissiveScale = BakeOutput.EmissiveScale;
 
 	if (bCopyAlphaFromRedChannel)
diff --git a/Engine/Plugins/Enterprise/GLTFExporter/MigratedFromGit/Source/GLTFExporter/Private/Converters/GLTFMaterialUtility.h b/Engine/Plugins/Enterprise/GLTFExporter/MigratedFromGit/Source/GLTFExporter/Private/Converters/GLTFMaterialUtility.h
--- a/Engine/Plugins/Enterprise/GLTFExporter/MigratedFromGit/Source/GLTFExporter/Private/Converters/GLTFMaterialUtility.h
+++ b/Engine/Plugins/Enterprise/GLTFExporter/MigratedFromGit/Source/GLTFExporter/Private/Converters/GLTFMaterialUtility.h
@@ -21,11 +21,11 @@ struct FGLTFTextureCombineSource
 
 struct FGLTFPropertyBakeOutput
 {
-	FORCEINLINE FGLTFPropertyBakeOutput(EMaterialProperty Property, EPixelFormat PixelFormat, TArray<FColor>& Pixels, FIntPoint Size, float EmissiveScale)
+	FORCEINLINE FGLTFPropertyBakeOutput(const FMaterialPropertyEx& Property, EPixelFormat PixelFormat, TArray<FColor>& Pixels, FIntPoint Size, float EmissiveScale)
 		: Property(Property), PixelFormat(PixelFormat), Pixels(Pixels), Size(Size), EmissiveScale(EmissiveScale), bIsConstant(false)
 	{}
 
-	EMaterialProperty Property;
+	const FMaterialPropertyEx& Property;
 	EPixelFormat PixelFormat;
 	TArray<FColor> Pixels;
 	FIntPoint Size;
@@ -39,16 +39,15 @@ struct FGLTFMaterialUtility
 {
 	static UMaterialInterface* GetDefault();
 
-	static bool IsNormalMap(EMaterialProperty Property);
+	static bool IsNormalMap(const FMaterialPropertyEx& Property);
 
-	static const TCHAR* GetPropertyName(EMaterialProperty Property);
-	static FVector4 GetPropertyDefaultValue(EMaterialProperty Property);
-	static FVector4 GetPropertyMask(EMaterialProperty Property);
+	static FVector4 GetPropertyDefaultValue(const FMaterialPropertyEx& Property);
+	static FVector4 GetPropertyMask(const FMaterialPropertyEx& Property);
 
-	static const FExpressionInput* GetInputForProperty(const UMaterialInterface* Material, EMaterialProperty Property);
+	static const FExpressionInput* GetInputForProperty(const UMaterialInterface* Material, const FMaterialPropertyEx& Property);
 
 	template<class InputType>
-	static const FMaterialInput<InputType>* GetInputForProperty(const UMaterialInterface* Material, EMaterialProperty Property)
+	static const FMaterialInput<InputType>* GetInputForProperty(const UMaterialInterface* Material, const FMaterialPropertyEx& Property)
 	{
 		const FExpressionInput* ExpressionInput = GetInputForProperty(Material, Property);
 		return static_cast<const FMaterialInput<InputType>*>(ExpressionInput);
@@ -58,7 +57,7 @@ struct FGLTFMaterialUtility
 
 	static UTexture2D* CreateTransientTexture(const FGLTFPropertyBakeOutput& PropertyBakeOutput, bool bUseSRGB = false);
 	static bool CombineTextures(TArray<FColor>& OutPixels, const TArray<FGLTFTextureCombineSource>& Sources, const FIntPoint& OutputSize, EPixelFormat OutputPixelFormat);
-	static FGLTFPropertyBakeOutput BakeMaterialProperty(const FIntPoint& OutputSize, EMaterialProperty Property, const UMaterialInterface* Material, int32 TexCoord, const FMeshDescription* MeshDescription = nullptr, const TArray<int32>& MeshSectionIndices = {}, bool bCopyAlphaFromRedChannel = false);
+	static FGLTFPropertyBakeOutput BakeMaterialProperty(const FIntPoint& OutputSize, const FMaterialPropertyEx& Property, const UMaterialInterface* Material, int32 TexCoord, const FMeshDescription* MeshDescription = nullptr, const TArray<int32>& MeshSectionIndices = {}, bool bCopyAlphaFromRedChannel = false);
 
 	static FGLTFJsonTextureIndex AddCombinedTexture(FGLTFConvertBuilder& Builder, const TArray<FGLTFTextureCombineSource>& CombineSources, const FIntPoint& TextureSize, bool bIgnoreAlpha, const FString& TextureName, EGLTFJsonTextureFilter MinFilter, EGLTFJsonTextureFilter MagFilter, EGLTFJsonTextureWrap WrapS, EGLTFJsonTextureWrap WrapT);
 	static FGLTFJsonTextureIndex AddTexture(FGLTFConvertBuilder& Builder, const TArray<FColor>& Pixels, const FIntPoint& TextureSize, bool bIgnoreAlpha, bool bIsNormalMap, const FString& TextureName, EGLTFJsonTextureFilter MinFilter, EGLTFJsonTextureFilter MagFilter, EGLTFJsonTextureWrap WrapS, EGLTFJsonTextureWrap WrapT);
diff --git a/Engine/Plugins/Enterprise/GLTFExporter/MigratedFromGit/Source/GLTFExporter/Private/Tasks/GLTFMaterialTasks.cpp b/Engine/Plugins/Enterprise/GLTFExporter/MigratedFromGit/Source/GLTFExporter/Private/Tasks/GLTFMaterialTasks.cpp
--- a/Engine/Plugins/Enterprise/GLTFExporter/MigratedFromGit/Source/GLTFExporter/Private/Tasks/GLTFMaterialTasks.cpp
+++ b/Engine/Plugins/Enterprise/GLTFExporter/MigratedFromGit/Source/GLTFExporter/Private/Tasks/GLTFMaterialTasks.cpp
@@ -81,8 +81,8 @@ void FGLTFMaterialTask::Complete()
 
 	if (JsonMaterial.ShadingModel != EGLTFJsonShadingModel::None)
 	{
-		const EMaterialProperty BaseColorProperty = JsonMaterial.ShadingModel == EGLTFJsonShadingModel::Unlit ? MP_EmissiveColor : MP_BaseColor;
-		const EMaterialProperty OpacityProperty = JsonMaterial.AlphaMode == EGLTFJsonAlphaMode::Mask ? MP_OpacityMask : MP_Opacity;
+		const FMaterialPropertyEx BaseColorProperty = JsonMaterial.ShadingModel == EGLTFJsonShadingModel::Unlit ? MP_EmissiveColor : MP_BaseColor;
+		const FMaterialPropertyEx OpacityProperty = JsonMaterial.AlphaMode == EGLTFJsonAlphaMode::Mask ? MP_OpacityMask : MP_Opacity;
 
 		// TODO: check if a property is active before trying to get it (i.e. Material->IsPropertyActive)
 
@@ -94,7 +94,7 @@ void FGLTFMaterialTask::Complete()
 				{
 					if (!TryGetBakedMaterialProperty(JsonMaterial.PBRMetallicRoughness.BaseColorTexture, JsonMaterial.PBRMetallicRoughness.BaseColorFactor, BaseColorProperty, TEXT("BaseColor")))
 					{
-						Builder.AddWarningMessage(FString::Printf(TEXT("Failed to export %s for material %s"), FGLTFMaterialUtility::GetPropertyName(BaseColorProperty), *Material->GetName()));
+						Builder.AddWarningMessage(FString::Printf(TEXT("Failed to export %s for material %s"), *BaseColorProperty.ToString(), *Material->GetName()));
 					}
 				}
 			}
@@ -105,69 +105,68 @@ void FGLTFMaterialTask::Complete()
 		{
 			if (!TryGetBaseColorAndOpacity(JsonMaterial.PBRMetallicRoughness, BaseColorProperty, OpacityProperty))
 			{
-				Builder.AddWarningMessage(FString::Printf(TEXT("Failed to export %s and %s for material %s"), FGLTFMaterialUtility::GetPropertyName(BaseColorProperty), FGLTFMaterialUtility::GetPropertyName(OpacityProperty), *Material->GetName()));
+				Builder.AddWarningMessage(FString::Printf(TEXT("Failed to export %s and %s for material %s"), *BaseColorProperty.ToString(), *OpacityProperty.ToString(), *Material->GetName()));
 			}
 		}
 
 		if (JsonMaterial.ShadingModel == EGLTFJsonShadingModel::Default || JsonMaterial.ShadingModel == EGLTFJsonShadingModel::ClearCoat)
 		{
-			const EMaterialProperty MetallicProperty = MP_Metallic;
-			const EMaterialProperty RoughnessProperty = MP_Roughness;
+			const FMaterialPropertyEx MetallicProperty = MP_Metallic;
+			const FMaterialPropertyEx RoughnessProperty = MP_Roughness;
 
 			if (!TryGetMetallicAndRoughness(JsonMaterial.PBRMetallicRoughness, MetallicProperty, RoughnessProperty))
 			{
-				Builder.AddWarningMessage(FString::Printf(TEXT("Failed to export %s and %s for material %s"), FGLTFMaterialUtility::GetPropertyName(MetallicProperty), FGLTFMaterialUtility::GetPropertyName(RoughnessProperty), *Material->GetName()));
+				Builder.AddWarningMessage(FString::Printf(TEXT("Failed to export %s and %s for material %s"), *MetallicProperty.ToString(), *RoughnessProperty.ToString(), *Material->GetName()));
 			}
 
-			const EMaterialProperty EmissiveProperty = MP_EmissiveColor;
+			const FMaterialPropertyEx EmissiveProperty = MP_EmissiveColor;
 			if (!TryGetEmissive(JsonMaterial, EmissiveProperty))
 			{
-				Builder.AddWarningMessage(FString::Printf(TEXT("Failed to export %s for material %s"), FGLTFMaterialUtility::GetPropertyName(EmissiveProperty), *Material->GetName()));
+				Builder.AddWarningMessage(FString::Printf(TEXT("Failed to export %s for material %s"), *EmissiveProperty.ToString(), *Material->GetName()));
 			}
 
-			// TODO: replace dummy enum MP_CustomOutput workaround for ClearCoatBottomNormal with proper support for custom outputs
-			const EMaterialProperty NormalProperty = JsonMaterial.ShadingModel == EGLTFJsonShadingModel::ClearCoat ? MP_CustomOutput : MP_Normal;
+			const FMaterialPropertyEx NormalProperty = JsonMaterial.ShadingModel == EGLTFJsonShadingModel::ClearCoat ? FMaterialPropertyEx(TEXT("ClearCoatBottomNormal")) : FMaterialPropertyEx(MP_Normal);
 			if (IsPropertyNonDefault(NormalProperty))
 			{
 				if (!TryGetSourceTexture(JsonMaterial.NormalTexture, NormalProperty, DefaultColorInputMasks))
 				{
 					if (!TryGetBakedMaterialProperty(JsonMaterial.NormalTexture, NormalProperty, TEXT("Normal")))
 					{
-						Builder.AddWarningMessage(FString::Printf(TEXT("Failed to export %s for material %s"), FGLTFMaterialUtility::GetPropertyName(NormalProperty), *Material->GetName()));
+						Builder.AddWarningMessage(FString::Printf(TEXT("Failed to export %s for material %s"), *NormalProperty.ToString(), *Material->GetName()));
 					}
 				}
 			}
 
-			const EMaterialProperty AmbientOcclusionProperty = MP_AmbientOcclusion;
+			const FMaterialPropertyEx AmbientOcclusionProperty = MP_AmbientOcclusion;
 			if (IsPropertyNonDefault(AmbientOcclusionProperty))
 			{
 				if (!TryGetSourceTexture(JsonMaterial.OcclusionTexture, AmbientOcclusionProperty, OcclusionInputMasks))
 				{
 					if (!TryGetBakedMaterialProperty(JsonMaterial.OcclusionTexture, AmbientOcclusionProperty, TEXT("Occlusion")))
 					{
-						Builder.AddWarningMessage(FString::Printf(TEXT("Failed to export %s for material %s"), FGLTFMaterialUtility::GetPropertyName(AmbientOcclusionProperty), *Material->GetName()));
+						Builder.AddWarningMessage(FString::Printf(TEXT("Failed to export %s for material %s"), *AmbientOcclusionProperty.ToString(), *Material->GetName()));
 					}
 				}
 			}
 
 			if (JsonMaterial.ShadingModel == EGLTFJsonShadingModel::ClearCoat)
 			{
-				const EMaterialProperty ClearCoatProperty = MP_CustomData0;
-				const EMaterialProperty ClearCoatRoughnessProperty = MP_CustomData1;
+				const FMaterialPropertyEx ClearCoatProperty = MP_CustomData0;
+				const FMaterialPropertyEx ClearCoatRoughnessProperty = MP_CustomData1;
 
 				if (!TryGetClearCoatRoughness(JsonMaterial.ClearCoat, ClearCoatProperty, ClearCoatRoughnessProperty))
 				{
-					Builder.AddWarningMessage(FString::Printf(TEXT("Failed to export %s and %s for material %s"), FGLTFMaterialUtility::GetPropertyName(ClearCoatProperty), FGLTFMaterialUtility::GetPropertyName(ClearCoatRoughnessProperty), *Material->GetName()));
+					Builder.AddWarningMessage(FString::Printf(TEXT("Failed to export %s and %s for material %s"), *ClearCoatProperty.ToString(), *ClearCoatRoughnessProperty.ToString(), *Material->GetName()));
 				}
 
-				const EMaterialProperty ClearCoatNormalProperty = MP_Normal;
+				const FMaterialPropertyEx ClearCoatNormalProperty = MP_Normal;
 				if (IsPropertyNonDefault(ClearCoatNormalProperty))
 				{
 					if (!TryGetSourceTexture(JsonMaterial.ClearCoat.ClearCoatNormalTexture, ClearCoatNormalProperty, DefaultColorInputMasks))
 					{
 						if (!TryGetBakedMaterialProperty(JsonMaterial.ClearCoat.ClearCoatNormalTexture, ClearCoatNormalProperty, TEXT("ClearCoatNormal")))
 						{
-							Builder.AddWarningMessage(FString::Printf(TEXT("Failed to export %s for material %s"), FGLTFMaterialUtility::GetPropertyName(ClearCoatNormalProperty), *Material->GetName()));
+							Builder.AddWarningMessage(FString::Printf(TEXT("Failed to export %s for material %s"), *ClearCoatNormalProperty.ToString(), *Material->GetName()));
 						}
 					}
 				}
@@ -264,7 +263,7 @@ EMaterialShadingModel FGLTFMaterialTask::EvaluateShadingModelExpression() const
 	return Evaluation.GetFirstShadingModel();
 }
 
-bool FGLTFMaterialTask::TryGetBaseColorAndOpacity(FGLTFJsonPBRMetallicRoughness& OutPBRParams, EMaterialProperty BaseColorProperty, EMaterialProperty OpacityProperty) const
+bool FGLTFMaterialTask::TryGetBaseColorAndOpacity(FGLTFJsonPBRMetallicRoughness& OutPBRParams, const FMaterialPropertyEx& BaseColorProperty, const FMaterialPropertyEx& OpacityProperty) const
 {
 	const bool bIsBaseColorConstant = TryGetConstantColor(OutPBRParams.BaseColorFactor, BaseColorProperty);
 	const bool bIsOpacityConstant = TryGetConstantScalar(OutPBRParams.BaseColorFactor.A, OpacityProperty);
@@ -306,8 +305,8 @@ bool FGLTFMaterialTask::TryGetBaseColorAndOpacity(FGLTFJsonPBRMetallicRoughness&
 	{
 		Builder.AddWarningMessage(FString::Printf(
 			TEXT("%s and %s for material %s needs to bake, but material baking is disabled by export options"),
-			FGLTFMaterialUtility::GetPropertyName(BaseColorProperty),
-			FGLTFMaterialUtility::GetPropertyName(OpacityProperty),
+			*BaseColorProperty.ToString(),
+			*OpacityProperty.ToString(),
 			*Material->GetName()));
 		return false;
 	}
@@ -431,7 +430,7 @@ bool FGLTFMaterialTask::TryGetBaseColorAndOpacity(FGLTFJsonPBRMetallicRoughness&
 	return true;
 }
 
-bool FGLTFMaterialTask::TryGetMetallicAndRoughness(FGLTFJsonPBRMetallicRoughness& OutPBRParams, EMaterialProperty MetallicProperty, EMaterialProperty RoughnessProperty) const
+bool FGLTFMaterialTask::TryGetMetallicAndRoughness(FGLTFJsonPBRMetallicRoughness& OutPBRParams, const FMaterialPropertyEx& MetallicProperty, const FMaterialPropertyEx& RoughnessProperty) const
 {
 	const bool bIsMetallicConstant = TryGetConstantScalar(OutPBRParams.MetallicFactor, MetallicProperty);
 	const bool bIsRoughnessConstant = TryGetConstantScalar(OutPBRParams.RoughnessFactor, RoughnessProperty);
@@ -474,8 +473,8 @@ bool FGLTFMaterialTask::TryGetMetallicAndRoughness(FGLTFJsonPBRMetallicRoughness
 	{
 		Builder.AddWarningMessage(FString::Printf(
 			TEXT("%s and %s for material %s needs to bake, but material baking is disabled by export options"),
-			FGLTFMaterialUtility::GetPropertyName(MetallicProperty),
-			FGLTFMaterialUtility::GetPropertyName(RoughnessProperty),
+			*MetallicProperty.ToString(),
+			*RoughnessProperty.ToString(),
 			*Material->GetName()));
 		return false;
 	}
@@ -595,7 +594,7 @@ bool FGLTFMaterialTask::TryGetMetallicAndRoughness(FGLTFJsonPBRMetallicRoughness
 	return true;
 }
 
-bool FGLTFMaterialTask::TryGetClearCoatRoughness(FGLTFJsonClearCoatExtension& OutExtParams, EMaterialProperty IntensityProperty, EMaterialProperty RoughnessProperty) const
+bool FGLTFMaterialTask::TryGetClearCoatRoughness(FGLTFJsonClearCoatExtension& OutExtParams, const FMaterialPropertyEx& IntensityProperty, const FMaterialPropertyEx& RoughnessProperty) const
 {
 	const bool bIsIntensityConstant = TryGetConstantScalar(OutExtParams.ClearCoatFactor, IntensityProperty);
 	const bool bIsRoughnessConstant = TryGetConstantScalar(OutExtParams.ClearCoatRoughnessFactor, RoughnessProperty);
@@ -641,8 +640,8 @@ bool FGLTFMaterialTask::TryGetClearCoatRoughness(FGLTFJsonClearCoatExtension& Ou
 	{
 		Builder.AddWarningMessage(FString::Printf(
 			TEXT("%s and %s for material %s needs to bake, but material baking is disabled by export options"),
-			FGLTFMaterialUtility::GetPropertyName(IntensityProperty),
-			FGLTFMaterialUtility::GetPropertyName(RoughnessProperty),
+			*IntensityProperty.ToString(),
+			*RoughnessProperty.ToString(),
 			*Material->GetName()));
 		return false;
 	}
@@ -764,7 +763,7 @@ bool FGLTFMaterialTask::TryGetClearCoatRoughness(FGLTFJsonClearCoatExtension& Ou
 	return true;
 }
 
-bool FGLTFMaterialTask::TryGetEmissive(FGLTFJsonMaterial& JsonMaterial, EMaterialProperty EmissiveProperty) const
+bool FGLTFMaterialTask::TryGetEmissive(FGLTFJsonMaterial& JsonMaterial, const FMaterialPropertyEx& EmissiveProperty) const
 {
 	// TODO: right now we allow EmissiveFactor to be > 1.0 to support very bright emission, although it's not valid according to the glTF standard.
 	// We may want to change this behaviour and store factors above 1.0 using a custom extension instead.
@@ -784,7 +783,7 @@ bool FGLTFMaterialTask::TryGetEmissive(FGLTFJsonMaterial& JsonMaterial, EMateria
 	{
 		Builder.AddWarningMessage(FString::Printf(
 			TEXT("%s for material %s needs to bake, but material baking is disabled by export options"),
-			FGLTFMaterialUtility::GetPropertyName(EmissiveProperty),
+			*EmissiveProperty.ToString(),
 			*Material->GetName()));
 		return false;
 	}
@@ -810,7 +809,7 @@ bool FGLTFMaterialTask::TryGetEmissive(FGLTFJsonMaterial& JsonMaterial, EMateria
 	return true;
 }
 
-bool FGLTFMaterialTask::IsPropertyNonDefault(EMaterialProperty Property) const
+bool FGLTFMaterialTask::IsPropertyNonDefault(const FMaterialPropertyEx& Property) const
 {
 	const bool bUseMaterialAttributes = Material->GetMaterial()->bUseMaterialAttributes;
 	if (bUseMaterialAttributes)
@@ -835,7 +834,7 @@ bool FGLTFMaterialTask::IsPropertyNonDefault(EMaterialProperty Property) const
 	return true;
 }
 
-bool FGLTFMaterialTask::TryGetConstantColor(FGLTFJsonColor3& OutValue, EMaterialProperty Property) const
+bool FGLTFMaterialTask::TryGetConstantColor(FGLTFJsonColor3& OutValue, const FMaterialPropertyEx& Property) const
 {
 	FLinearColor Value;
 	if (TryGetConstantColor(Value, Property))
@@ -847,7 +846,7 @@ bool FGLTFMaterialTask::TryGetConstantColor(FGLTFJsonColor3& OutValue, EMaterial
 	return false;
 }
 
-bool FGLTFMaterialTask::TryGetConstantColor(FGLTFJsonColor4& OutValue, EMaterialProperty Property) const
+bool FGLTFMaterialTask::TryGetConstantColor(FGLTFJsonColor4& OutValue, const FMaterialPropertyEx& Property) const
 {
 	FLinearColor Value;
 	if (TryGetConstantColor(Value, Property))
@@ -859,7 +858,7 @@ bool FGLTFMaterialTask::TryGetConstantColor(FGLTFJsonColor4& OutValue, EMaterial
 	return false;
 }
 
-bool FGLTFMaterialTask::TryGetConstantColor(FLinearColor& OutValue, EMaterialProperty Property) const
+bool FGLTFMaterialTask::TryGetConstantColor(FLinearColor& OutValue, const FMaterialPropertyEx& Property) const
 {
 	const bool bUseMaterialAttributes = Material->GetMaterial()->bUseMaterialAttributes;
 	if (bUseMaterialAttributes)
@@ -966,7 +965,7 @@ bool FGLTFMaterialTask::TryGetConstantColor(FLinearColor& OutValue, EMaterialPro
 	return false;
 }
 
-bool FGLTFMaterialTask::TryGetConstantScalar(float& OutValue, EMaterialProperty Property) const
+bool FGLTFMaterialTask::TryGetConstantScalar(float& OutValue, const FMaterialPropertyEx& Property) const
 {
 	const bool bUseMaterialAttributes = Material->GetMaterial()->bUseMaterialAttributes;
 	if (bUseMaterialAttributes)
@@ -1067,7 +1066,7 @@ bool FGLTFMaterialTask::TryGetConstantScalar(float& OutValue, EMaterialProperty
 	return false;
 }
 
-bool FGLTFMaterialTask::TryGetSourceTexture(FGLTFJsonTextureInfo& OutTexInfo, EMaterialProperty Property, const TArray<FLinearColor>& AllowedMasks) const
+bool FGLTFMaterialTask::TryGetSourceTexture(FGLTFJsonTextureInfo& OutTexInfo, const FMaterialPropertyEx& Property, const TArray<FLinearColor>& AllowedMasks) const
 {
 	const UTexture2D* Texture;
 	int32 TexCoord;
@@ -1084,7 +1083,7 @@ bool FGLTFMaterialTask::TryGetSourceTexture(FGLTFJsonTextureInfo& OutTexInfo, EM
 	return false;
 }
 
-bool FGLTFMaterialTask::TryGetSourceTexture(const UTexture2D*& OutTexture, int32& OutTexCoord, FGLTFJsonTextureTransform& OutTransform, EMaterialProperty Property, const TArray<FLinearColor>& AllowedMasks) const
+bool FGLTFMaterialTask::TryGetSourceTexture(const UTexture2D*& OutTexture, int32& OutTexCoord, FGLTFJsonTextureTransform& OutTransform, const FMaterialPropertyEx& Property, const TArray<FLinearColor>& AllowedMasks) const
 {
 	const FExpressionInput* MaterialInput = FGLTFMaterialUtility::GetInputForProperty(Material, Property);
 	if (MaterialInput == nullptr)
@@ -1147,7 +1146,7 @@ bool FGLTFMaterialTask::TryGetSourceTexture(const UTexture2D*& OutTexture, int32
 			Builder.AddWarningMessage(FString::Printf(
 				TEXT("Texture coordinates [%d] in %s for material %s are transformed, but texture transform is disabled by export options"),
 				OutTexCoord,
-				*FGLTFNameUtility::GetName(Property),
+				*Property.ToString(),
 				*Material->GetName()));
 			OutTransform = {};
 		}
@@ -1184,7 +1183,7 @@ bool FGLTFMaterialTask::TryGetSourceTexture(const UTexture2D*& OutTexture, int32
 			Builder.AddWarningMessage(FString::Printf(
 				TEXT("Texture coordinates [%d] in %s for material %s are transformed, but texture transform is disabled by export options"),
 				OutTexCoord,
-				*FGLTFNameUtility::GetName(Property),
+				*Property.ToString(),
 				*Material->GetName()));
 			OutTransform = {};
 		}
@@ -1195,13 +1194,13 @@ bool FGLTFMaterialTask::TryGetSourceTexture(const UTexture2D*& OutTexture, int32
 	return false;
 }
 
-bool FGLTFMaterialTask::TryGetBakedMaterialProperty(FGLTFJsonTextureInfo& OutTexInfo, FGLTFJsonColor3& OutConstant, EMaterialProperty Property, const FString& PropertyName) const
+bool FGLTFMaterialTask::TryGetBakedMaterialProperty(FGLTFJsonTextureInfo& OutTexInfo, FGLTFJsonColor3& OutConstant, const FMaterialPropertyEx& Property, const FString& PropertyName) const
 {
 	if (Builder.ExportOptions->BakeMaterialInputs == EGLTFMaterialBakeMode::None)
 	{
 		Builder.AddWarningMessage(FString::Printf(
 			TEXT("%s for material %s needs to bake, but material baking is disabled by export options"),
-			FGLTFMaterialUtility::GetPropertyName(Property),
+			*Property.ToString(),
 			*Material->GetName()));
 		return false;
 	}
@@ -1223,13 +1222,13 @@ bool FGLTFMaterialTask::TryGetBakedMaterialProperty(FGLTFJsonTextureInfo& OutTex
 	return false;
 }
 
-bool FGLTFMaterialTask::TryGetBakedMaterialProperty(FGLTFJsonTextureInfo& OutTexInfo, FGLTFJsonColor4& OutConstant, EMaterialProperty Property, const FString& PropertyName) const
+bool FGLTFMaterialTask::TryGetBakedMaterialProperty(FGLTFJsonTextureInfo& OutTexInfo, FGLTFJsonColor4& OutConstant, const FMaterialPropertyEx& Property, const FString& PropertyName) const
 {
 	if (Builder.ExportOptions->BakeMaterialInputs == EGLTFMaterialBakeMode::None)
 	{
 		Builder.AddWarningMessage(FString::Printf(
 			TEXT("%s for material %s needs to bake, but material baking is disabled by export options"),
-			FGLTFMaterialUtility::GetPropertyName(Property),
+			*Property.ToString(),
 			*Material->GetName()));
 		return false;
 	}
@@ -1251,13 +1250,13 @@ bool FGLTFMaterialTask::TryGetBakedMaterialProperty(FGLTFJsonTextureInfo& OutTex
 	return false;
 }
 
-inline bool FGLTFMaterialTask::TryGetBakedMaterialProperty(FGLTFJsonTextureInfo& OutTexInfo, float& OutConstant, EMaterialProperty Property, const FString& PropertyName) const
+inline bool FGLTFMaterialTask::TryGetBakedMaterialProperty(FGLTFJsonTextureInfo& OutTexInfo, float& OutConstant, const FMaterialPropertyEx& Property, const FString& PropertyName) const
 {
 	if (Builder.ExportOptions->BakeMaterialInputs == EGLTFMaterialBakeMode::None)
 	{
 		Builder.AddWarningMessage(FString::Printf(
 			TEXT("%s for material %s needs to bake, but material baking is disabled by export options"),
-			FGLTFMaterialUtility::GetPropertyName(Property),
+			*Property.ToString(),
 			*Material->GetName()));
 		return false;
 	}
@@ -1279,13 +1278,13 @@ inline bool FGLTFMaterialTask::TryGetBakedMaterialProperty(FGLTFJsonTextureInfo&
 	return false;
 }
 
-bool FGLTFMaterialTask::TryGetBakedMaterialProperty(FGLTFJsonTextureInfo& OutTexInfo, EMaterialProperty Property, const FString& PropertyName) const
+bool FGLTFMaterialTask::TryGetBakedMaterialProperty(FGLTFJsonTextureInfo& OutTexInfo, const FMaterialPropertyEx& Property, const FString& PropertyName) const
 {
 	if (Builder.ExportOptions->BakeMaterialInputs == EGLTFMaterialBakeMode::None)
 	{
 		Builder.AddWarningMessage(FString::Printf(
 			TEXT("%s for material %s needs to bake, but material baking is disabled by export options"),
-			FGLTFMaterialUtility::GetPropertyName(Property),
+			*Property.ToString(),
 			*Material->GetName()));
 		return false;
 	}
@@ -1335,16 +1334,13 @@ bool FGLTFMaterialTask::TryGetBakedMaterialProperty(FGLTFJsonTextureInfo& OutTex
 	return true;
 }
 
-FGLTFPropertyBakeOutput FGLTFMaterialTask::BakeMaterialProperty(EMaterialProperty Property, int32& OutTexCoord, FIntPoint PreferredTextureSize, bool bCopyAlphaFromRedChannel) const
+FGLTFPropertyBakeOutput FGLTFMaterialTask::BakeMaterialProperty(const FMaterialPropertyEx& Property, int32& OutTexCoord, FIntPoint PreferredTextureSize, bool bCopyAlphaFromRedChannel) const
 {
 	if (MeshData == nullptr)
 	{
-		// TODO: replace dummy enum MP_CustomOutput workaround for ClearCoatBottomNormal
-		const FMaterialPropertyEx PropertyEx = Property != MP_CustomOutput ? FMaterialPropertyEx(Property) : FMaterialPropertyEx(TEXT("ClearCoatBottomNormal"));
-
 		int32 NumTexCoords;
 		TBitArray<> TexCoords;
-		FGLTFMaterialUtility::GetAllTextureCoordinateIndices(Material, PropertyEx, NumTexCoords, TexCoords);
+		FGLTFMaterialUtility::GetAllTextureCoordinateIndices(Material, Property, NumTexCoords, TexCoords);
 
 		if (NumTexCoords > 0)
 		{
diff --git a/Engine/Plugins/Enterprise/GLTFExporter/MigratedFromGit/Source/GLTFExporter/Private/Tasks/GLTFMaterialTasks.h b/Engine/Plugins/Enterprise/GLTFExporter/MigratedFromGit/Source/GLTFExporter/Private/Tasks/GLTFMaterialTasks.h
--- a/Engine/Plugins/Enterprise/GLTFExporter/MigratedFromGit/Source/GLTFExporter/Private/Tasks/GLTFMaterialTasks.h
+++ b/Engine/Plugins/Enterprise/GLTFExporter/MigratedFromGit/Source/GLTFExporter/Private/Tasks/GLTFMaterialTasks.h
@@ -5,6 +5,7 @@
 #include "Tasks/GLTFTask.h"
 #include "Converters/GLTFMeshData.h"
 #include "Builders/GLTFConvertBuilder.h"
+#include "MaterialPropertyEx.h"
 #include "Engine.h"
 
 struct FGLTFPropertyBakeOutput;
@@ -43,26 +44,26 @@ private:
 
 	EMaterialShadingModel EvaluateShadingModelExpression() const;
 
-	bool TryGetBaseColorAndOpacity(FGLTFJsonPBRMetallicRoughness& OutPBRParams, EMaterialProperty BaseColorProperty, EMaterialProperty OpacityProperty) const;
-	bool TryGetMetallicAndRoughness(FGLTFJsonPBRMetallicRoughness& OutPBRParams, EMaterialProperty MetallicProperty, EMaterialProperty RoughnessProperty) const;
-	bool TryGetClearCoatRoughness(FGLTFJsonClearCoatExtension& OutExtParams, EMaterialProperty IntensityProperty, EMaterialProperty RoughnessProperty) const;
-	bool TryGetEmissive(FGLTFJsonMaterial& JsonMaterial, EMaterialProperty EmissiveProperty) const;
+	bool TryGetBaseColorAndOpacity(FGLTFJsonPBRMetallicRoughness& OutPBRParams, const FMaterialPropertyEx& BaseColorProperty, const FMaterialPropertyEx& OpacityProperty) const;
+	bool TryGetMetallicAndRoughness(FGLTFJsonPBRMetallicRoughness& OutPBRParams, const FMaterialPropertyEx& MetallicProperty, const FMaterialPropertyEx& RoughnessProperty) const;
+	bool TryGetClearCoatRoughness(FGLTFJsonClearCoatExtension& OutExtParams, const FMaterialPropertyEx& IntensityProperty, const FMaterialPropertyEx& RoughnessProperty) const;
+	bool TryGetEmissive(FGLTFJsonMaterial& JsonMaterial, const FMaterialPropertyEx& EmissiveProperty) const;
 
-	bool IsPropertyNonDefault(EMaterialProperty Property) const;
-	bool TryGetConstantColor(FGLTFJsonColor3& OutValue, EMaterialProperty Property) const;
-	bool TryGetConstantColor(FGLTFJsonColor4& OutValue, EMaterialProperty Property) const;
-	bool TryGetConstantColor(FLinearColor& OutValue, EMaterialProperty Property) const;
-	bool TryGetConstantScalar(float& OutValue, EMaterialProperty Property) const;
+	bool IsPropertyNonDefault(const FMaterialPropertyEx& Property) const;
+	bool TryGetConstantColor(FGLTFJsonColor3& OutValue, const FMaterialPropertyEx& Property) const;
+	bool TryGetConstantColor(FGLTFJsonColor4& OutValue, const FMaterialPropertyEx& Property) const;
+	bool TryGetConstantColor(FLinearColor& OutValue, const FMaterialPropertyEx& Property) const;
+	bool TryGetConstantScalar(float& OutValue, const FMaterialPropertyEx& Property) const;
 
-	bool TryGetSourceTexture(FGLTFJsonTextureInfo& OutTexInfo, EMaterialProperty Property, const TArray<FLinearColor>& AllowedMasks = {}) const;
-	bool TryGetSourceTexture(const UTexture2D*& OutTexture, int32& OutTexCoord, FGLTFJsonTextureTransform& OutTransform, EMaterialProperty Property, const TArray<FLinearColor>& AllowedMasks = {}) const;
+	bool TryGetSourceTexture(FGLTFJsonTextureInfo& OutTexInfo, const FMaterialPropertyEx& Property, const TArray<FLinearColor>& AllowedMasks = {}) const;
+	bool TryGetSourceTexture(const UTexture2D*& OutTexture, int32& OutTexCoord, FGLTFJsonTextureTransform& OutTransform, const FMaterialPropertyEx& Property, const TArray<FLinearColor>& AllowedMasks = {}) const;
 
-	bool TryGetBakedMaterialProperty(FGLTFJsonTextureInfo& OutTexInfo, FGLTFJsonColor3& OutConstant, EMaterialProperty Property, const FString& PropertyName) const;
-	bool TryGetBakedMaterialProperty(FGLTFJsonTextureInfo& OutTexInfo, FGLTFJsonColor4& OutConstant, EMaterialProperty Property, const FString& PropertyName) const;
-	bool TryGetBakedMaterialProperty(FGLTFJsonTextureInfo& OutTexInfo, float& OutConstant, EMaterialProperty Property, const FString& PropertyName) const;
-	bool TryGetBakedMaterialProperty(FGLTFJsonTextureInfo& OutTexInfo, EMaterialProperty Property, const FString& PropertyName) const;
+	bool TryGetBakedMaterialProperty(FGLTFJsonTextureInfo& OutTexInfo, FGLTFJsonColor3& OutConstant, const FMaterialPropertyEx& Property, const FString& PropertyName) const;
+	bool TryGetBakedMaterialProperty(FGLTFJsonTextureInfo& OutTexInfo, FGLTFJsonColor4& OutConstant, const FMaterialPropertyEx& Property, const FString& PropertyName) const;
+	bool TryGetBakedMaterialProperty(FGLTFJsonTextureInfo& OutTexInfo, float& OutConstant, const FMaterialPropertyEx& Property, const FString& PropertyName) const;
+	bool TryGetBakedMaterialProperty(FGLTFJsonTextureInfo& OutTexInfo, const FMaterialPropertyEx& Property, const FString& PropertyName) const;
 
-	FGLTFPropertyBakeOutput BakeMaterialProperty(EMaterialProperty Property, int32& OutTexCoord, FIntPoint PreferredTextureSize = FIntPoint::ZeroValue, bool bCopyAlphaFromRedChannel = false) const;
+	FGLTFPropertyBakeOutput BakeMaterialProperty(const FMaterialPropertyEx& Property, int32& OutTexCoord, FIntPoint PreferredTextureSize = FIntPoint::ZeroValue, bool bCopyAlphaFromRedChannel = false) const;
 
 	bool StoreBakedPropertyTexture(FGLTFJsonTextureInfo& OutTexInfo, const FGLTFPropertyBakeOutput& PropertyBakeOutput, const FString& PropertyName) const;
 
