[GLTFExporter] refactor: don't use mesh data for material baking if no suitable texture coordinate found

commit 053ee9221de643c997eb2d9dfa1172716f392903
Date: Sat, 14 Nov 2020 01:07:48 +0100

    refactor: don't use mesh data for material baking if no suitable texture coordinate found
    (cherry picked from commit 98975a1311ca76c9e388253aecbc6a14f6a01e71​)

#skipci

diff --git a/Engine/Plugins/Enterprise/GLTFExporter/MigratedFromGit/Source/GLTFExporter/Private/Converters/GLTFMeshData.cpp b/Engine/Plugins/Enterprise/GLTFExporter/MigratedFromGit/Source/GLTFExporter/Private/Converters/GLTFMeshData.cpp
--- a/Engine/Plugins/Enterprise/GLTFExporter/MigratedFromGit/Source/GLTFExporter/Private/Converters/GLTFMeshData.cpp
+++ b/Engine/Plugins/Enterprise/GLTFExporter/MigratedFromGit/Source/GLTFExporter/Private/Converters/GLTFMeshData.cpp
@@ -22,7 +22,8 @@ FGLTFMeshData::FGLTFMeshData(const UStaticMesh* StaticMesh, const UStaticMeshCom
 		FMeshMergeHelpers::RetrieveMesh(StaticMesh, LODIndex, Description);
 	}
 
-	TexCoord = StaticMesh->LightMapCoordinateIndex;
+	const int32 NumTexCoords = StaticMesh->GetLODForExport(LODIndex).VertexBuffers.StaticMeshVertexBuffer.GetNumTexCoords();
+	TexCoord = FMath::Min(StaticMesh->LightMapCoordinateIndex, NumTexCoords - 1);
 
 	// TODO: add warning if texture coordinate has overlap
 }
@@ -66,7 +67,7 @@ FGLTFMeshData::FGLTFMeshData(const USkeletalMesh* SkeletalMesh, const USkeletalM
 	}
 
 	// TODO: don't assume last UV channel is non-overlapping
-	const uint32 NumTexCoords = SkeletalMesh->GetResourceForRendering()->LODRenderData[LODIndex].StaticVertexBuffers.StaticMeshVertexBuffer.GetNumTexCoords();
+	const int32 NumTexCoords = SkeletalMesh->GetResourceForRendering()->LODRenderData[LODIndex].StaticVertexBuffers.StaticMeshVertexBuffer.GetNumTexCoords();
 	TexCoord = NumTexCoords - 1;
 
 	// TODO: add warning if texture coordinate has overlap
diff --git a/Engine/Plugins/Enterprise/GLTFExporter/MigratedFromGit/Source/GLTFExporter/Private/Tasks/GLTFMeshTasks.cpp b/Engine/Plugins/Enterprise/GLTFExporter/MigratedFromGit/Source/GLTFExporter/Private/Tasks/GLTFMeshTasks.cpp
--- a/Engine/Plugins/Enterprise/GLTFExporter/MigratedFromGit/Source/GLTFExporter/Private/Tasks/GLTFMeshTasks.cpp
+++ b/Engine/Plugins/Enterprise/GLTFExporter/MigratedFromGit/Source/GLTFExporter/Private/Tasks/GLTFMeshTasks.cpp
@@ -25,10 +25,18 @@ void FGLTFStaticMeshTask::Complete()
 	const FGLTFMeshData* MeshData = Builder.ExportOptions->BakeMaterialInputs == EGLTFMaterialBakeMode::UseMeshData ?
 		Builder.StaticMeshDataConverter.GetOrAdd(StaticMesh, StaticMeshComponent, LODIndex) : nullptr;
 
-	if (MeshData != nullptr && MeshData->Description.IsEmpty())
+	if (MeshData != nullptr)
 	{
-		// TODO: add warning in case the mesh actually has data, which means we failed to extract a mesh description.
-		MeshData = nullptr;
+		if (MeshData->Description.IsEmpty())
+		{
+			// TODO: report warning in case the mesh actually has data, which means we failed to extract a mesh description.
+			MeshData = nullptr;
+		}
+		else if (MeshData->TexCoord < 0)
+		{
+			// TODO: report warning (about missing non-overlapping texture coordinate for baking with mesh data).
+			MeshData = nullptr;
+		}
 	}
 
 	const int32 MaterialCount = StaticMesh->StaticMaterials.Num();
@@ -94,10 +102,18 @@ void FGLTFSkeletalMeshTask::Complete()
 	const FGLTFMeshData* MeshData = Builder.ExportOptions->BakeMaterialInputs == EGLTFMaterialBakeMode::UseMeshData ?
 		Builder.SkeletalMeshDataConverter.GetOrAdd(SkeletalMesh, SkeletalMeshComponent, LODIndex) : nullptr;
 
-	if (MeshData != nullptr && MeshData->Description.IsEmpty())
+	if (MeshData != nullptr)
 	{
-		// TODO: add warning in case the mesh actually has data, which means we failed to extract a mesh description.
-		MeshData = nullptr;
+		if (MeshData->Description.IsEmpty())
+		{
+			// TODO: report warning in case the mesh actually has data, which means we failed to extract a mesh description.
+			MeshData = nullptr;
+		}
+		else if (MeshData->TexCoord < 0)
+		{
+			// TODO: report warning (about missing non-overlapping texture coordinate for baking with mesh data).
+			MeshData = nullptr;
+		}
 	}
 
 	const uint16 MaterialCount = SkeletalMesh->Materials.Num();
