[GLTFExporter] fix: instead of baking shading model (which sometimes produces bad results) use gltf material analyzer

commit 86471c1896f0130a72d0fa1e20ab0aa2c0490dbf
Date: Sun, 15 Nov 2020 20:25:34 +0100

    fix: instead of baking shading model (which sometimes produces bad results) use gltf material analyzer
    (cherry picked from commit f67a27c644c2c739bf027f69c9b6fc806c83660f​)

#skipci

diff --git a/Engine/Plugins/Enterprise/GLTFExporter/MigratedFromGit/Source/GLTFExporter/Private/Converters/GLTFMaterialUtility.cpp b/Engine/Plugins/Enterprise/GLTFExporter/MigratedFromGit/Source/GLTFExporter/Private/Converters/GLTFMaterialUtility.cpp
--- a/Engine/Plugins/Enterprise/GLTFExporter/MigratedFromGit/Source/GLTFExporter/Private/Converters/GLTFMaterialUtility.cpp
+++ b/Engine/Plugins/Enterprise/GLTFExporter/MigratedFromGit/Source/GLTFExporter/Private/Converters/GLTFMaterialUtility.cpp
@@ -354,6 +354,13 @@ void FGLTFMaterialUtility::GetAllTextureCoordinateIndices(const UMaterialInterfa
 	}
 }
 
+FMaterialShadingModelField FGLTFMaterialUtility::EvaluateShadingModelExpression(const UMaterialInterface* Material)
+{
+	FGLTFMaterialAnalysis Analysis;
+	AnalyzeMaterialProperty(Material, MP_ShadingModel, Analysis);
+	return Analysis.ShadingModels;
+}
+
 EMaterialShadingModel FGLTFMaterialUtility::GetRichestShadingModel(const FMaterialShadingModelField& ShadingModels)
 {
 	if (ShadingModels.HasShadingModel(MSM_ClearCoat))
diff --git a/Engine/Plugins/Enterprise/GLTFExporter/MigratedFromGit/Source/GLTFExporter/Private/Converters/GLTFMaterialUtility.h b/Engine/Plugins/Enterprise/GLTFExporter/MigratedFromGit/Source/GLTFExporter/Private/Converters/GLTFMaterialUtility.h
--- a/Engine/Plugins/Enterprise/GLTFExporter/MigratedFromGit/Source/GLTFExporter/Private/Converters/GLTFMaterialUtility.h
+++ b/Engine/Plugins/Enterprise/GLTFExporter/MigratedFromGit/Source/GLTFExporter/Private/Converters/GLTFMaterialUtility.h
@@ -68,6 +68,7 @@ struct FGLTFMaterialUtility
 	static bool TryGetTextureCoordinateIndex(const UMaterialExpressionTextureSample* TextureSampler, int32& TexCoord, FGLTFJsonTextureTransform& Transform);
 	static void GetAllTextureCoordinateIndices(const UMaterialInterface* InMaterial, const FMaterialPropertyEx& InProperty, int32& OutNumTexCoords, TBitArray<>& OutTexCoords);
 
+	static FMaterialShadingModelField EvaluateShadingModelExpression(const UMaterialInterface* Material);
 	static EMaterialShadingModel GetRichestShadingModel(const FMaterialShadingModelField& ShadingModels);
 	static FString ShadingModelsToString(const FMaterialShadingModelField& ShadingModels);
 
diff --git a/Engine/Plugins/Enterprise/GLTFExporter/MigratedFromGit/Source/GLTFExporter/Private/Tasks/GLTFMaterialTasks.cpp b/Engine/Plugins/Enterprise/GLTFExporter/MigratedFromGit/Source/GLTFExporter/Private/Tasks/GLTFMaterialTasks.cpp
--- a/Engine/Plugins/Enterprise/GLTFExporter/MigratedFromGit/Source/GLTFExporter/Private/Tasks/GLTFMaterialTasks.cpp
+++ b/Engine/Plugins/Enterprise/GLTFExporter/MigratedFromGit/Source/GLTFExporter/Private/Tasks/GLTFMaterialTasks.cpp
@@ -195,9 +195,39 @@ bool FGLTFMaterialTask::TryGetAlphaMode(EGLTFJsonAlphaMode& AlphaMode) const
 
 bool FGLTFMaterialTask::TryGetShadingModel(EGLTFJsonShadingModel& OutShadingModel) const
 {
-	const EMaterialShadingModel ShadingModel =
-		Material->IsShadingModelFromMaterialExpression() ?
-			EvaluateShadingModelExpression() : Material->GetShadingModels().GetFirstShadingModel();
+	EMaterialShadingModel ShadingModel;
+	{
+		const FMaterialShadingModelField ShadingModels = Material->GetShadingModels();
+		if (Material->IsShadingModelFromMaterialExpression())
+		{
+			const FMaterialShadingModelField Evaluation = FGLTFMaterialUtility::EvaluateShadingModelExpression(Material);
+			if (Evaluation.CountShadingModels() == 1)
+			{
+				ShadingModel = Evaluation.GetFirstShadingModel();
+			}
+			else if (Evaluation.CountShadingModels() > 1)
+			{
+				ShadingModel = FGLTFMaterialUtility::GetRichestShadingModel(Evaluation);
+				Builder.AddWarningMessage(FString::Printf(
+					TEXT("Evaluation of shading model expression in material %s is inconclusive (%s), will export as %s"),
+					*Material->GetName(),
+					*FGLTFMaterialUtility::ShadingModelsToString(Evaluation),
+					*FGLTFNameUtility::GetName(ShadingModel)));
+			}
+			else
+			{
+				ShadingModel = FGLTFMaterialUtility::GetRichestShadingModel(ShadingModels);
+				Builder.AddWarningMessage(FString::Printf(
+					TEXT("Evaluation of shading model expression in material %s returned none, will export as %s"),
+					*Material->GetName(),
+					*FGLTFNameUtility::GetName(ShadingModel)));
+			}
+		}
+		else
+		{
+			ShadingModel = ShadingModels.GetFirstShadingModel();
+		}
+	}
 
 	const EGLTFJsonShadingModel ConvertedShadingModel = FGLTFConverterUtility::ConvertShadingModel(ShadingModel);
 	if (ConvertedShadingModel == EGLTFJsonShadingModel::None)
@@ -225,44 +255,6 @@ bool FGLTFMaterialTask::TryGetShadingModel(EGLTFJsonShadingModel& OutShadingMode
 	return true;
 }
 
-EMaterialShadingModel FGLTFMaterialTask::EvaluateShadingModelExpression() const
-{
-	const FMaterialShadingModelField PossibleValues = Material->GetShadingModels();
-
-	if (Builder.ExportOptions->BakeMaterialInputs == EGLTFMaterialBakeMode::None)
-	{
-		const EMaterialShadingModel ShadingModel = FGLTFMaterialUtility::GetRichestShadingModel(PossibleValues);
-		Builder.AddWarningMessage(FString::Printf(
-			TEXT("Can't evaluate shading model expression for material %s, because material baking is disabled by export options, will export as %s"),
-			*Material->GetName(),
-			*FGLTFNameUtility::GetName(ShadingModel)));
-		return ShadingModel;
-	}
-
-	int32 TexCoord;
-	const FGLTFPropertyBakeOutput BakeOutput = BakeMaterialProperty(MP_ShadingModel, TexCoord); // TODO: should we use default bake size?
-	FMaterialShadingModelField Evaluation;
-
-	for (const FColor& Pixel : BakeOutput.Pixels)
-	{
-		const EMaterialShadingModel ShadingModel = static_cast<EMaterialShadingModel>(Pixel.R);
-		if (PossibleValues.HasShadingModel(ShadingModel)) Evaluation.AddShadingModel(ShadingModel);
-	}
-
-	if (Evaluation.CountShadingModels() != 1)
-	{
-		const EMaterialShadingModel ShadingModel = FGLTFMaterialUtility::GetRichestShadingModel(Evaluation.CountShadingModels() > 0 ? Evaluation : PossibleValues);
-		Builder.AddWarningMessage(FString::Printf(
-			TEXT("Evaluation of shading model expression in material %s is inconclusive (%s), will export as %s"),
-			*Material->GetName(),
-			*FGLTFMaterialUtility::ShadingModelsToString(Evaluation),
-			*FGLTFNameUtility::GetName(ShadingModel)));
-		return ShadingModel;
-	}
-
-	return Evaluation.GetFirstShadingModel();
-}
-
 bool FGLTFMaterialTask::TryGetBaseColorAndOpacity(FGLTFJsonPBRMetallicRoughness& OutPBRParams, const FMaterialPropertyEx& BaseColorProperty, const FMaterialPropertyEx& OpacityProperty) const
 {
 	const bool bIsBaseColorConstant = TryGetConstantColor(OutPBRParams.BaseColorFactor, BaseColorProperty);
diff --git a/Engine/Plugins/Enterprise/GLTFExporter/MigratedFromGit/Source/GLTFExporter/Private/Tasks/GLTFMaterialTasks.h b/Engine/Plugins/Enterprise/GLTFExporter/MigratedFromGit/Source/GLTFExporter/Private/Tasks/GLTFMaterialTasks.h
--- a/Engine/Plugins/Enterprise/GLTFExporter/MigratedFromGit/Source/GLTFExporter/Private/Tasks/GLTFMaterialTasks.h
+++ b/Engine/Plugins/Enterprise/GLTFExporter/MigratedFromGit/Source/GLTFExporter/Private/Tasks/GLTFMaterialTasks.h
@@ -42,8 +42,6 @@ private:
 	bool TryGetAlphaMode(EGLTFJsonAlphaMode& AlphaMode) const;
 	bool TryGetShadingModel(EGLTFJsonShadingModel& ShadingModel) const;
 
-	EMaterialShadingModel EvaluateShadingModelExpression() const;
-
 	bool TryGetBaseColorAndOpacity(FGLTFJsonPBRMetallicRoughness& OutPBRParams, const FMaterialPropertyEx& BaseColorProperty, const FMaterialPropertyEx& OpacityProperty) const;
 	bool TryGetMetallicAndRoughness(FGLTFJsonPBRMetallicRoughness& OutPBRParams, const FMaterialPropertyEx& MetallicProperty, const FMaterialPropertyEx& RoughnessProperty) const;
 	bool TryGetClearCoatRoughness(FGLTFJsonClearCoatExtension& OutExtParams, const FMaterialPropertyEx& IntensityProperty, const FMaterialPropertyEx& RoughnessProperty) const;
diff --git a/Engine/Plugins/Enterprise/GLTFExporter/MigratedFromGit/Source/GLTFMaterialAnalyzer/Private/GLTFMaterialAnalyzer.cpp b/Engine/Plugins/Enterprise/GLTFExporter/MigratedFromGit/Source/GLTFMaterialAnalyzer/Private/GLTFMaterialAnalyzer.cpp
--- a/Engine/Plugins/Enterprise/GLTFExporter/MigratedFromGit/Source/GLTFMaterialAnalyzer/Private/GLTFMaterialAnalyzer.cpp
+++ b/Engine/Plugins/Enterprise/GLTFExporter/MigratedFromGit/Source/GLTFMaterialAnalyzer/Private/GLTFMaterialAnalyzer.cpp
@@ -90,6 +90,7 @@ int32 UGLTFMaterialAnalyzer::CompilePropertyEx(FMaterialCompiler* Compiler, cons
 	}
 
 	Analysis->TextureCoordinates = Translator->AllocatedUserTexCoords;
+	Analysis->ShadingModels = Translator->ShadingModelsFromCompilation;
 
 	// TODO: investigate if we need to check more conditions to determine that vertex data is required
 	Analysis->bRequiresVertexData =
diff --git a/Engine/Plugins/Enterprise/GLTFExporter/MigratedFromGit/Source/GLTFMaterialAnalyzer/Public/GLTFMaterialAnalysis.h b/Engine/Plugins/Enterprise/GLTFExporter/MigratedFromGit/Source/GLTFMaterialAnalyzer/Public/GLTFMaterialAnalysis.h
--- a/Engine/Plugins/Enterprise/GLTFExporter/MigratedFromGit/Source/GLTFMaterialAnalyzer/Public/GLTFMaterialAnalysis.h
+++ b/Engine/Plugins/Enterprise/GLTFExporter/MigratedFromGit/Source/GLTFMaterialAnalyzer/Public/GLTFMaterialAnalysis.h
@@ -7,6 +7,9 @@ struct FGLTFMaterialAnalysis
 	/** Tracks the texture coordinates used by this material */
 	TBitArray<> TextureCoordinates;
 
+	/** Will contain all the shading models picked up from the material expression graph */
+	FMaterialShadingModelField ShadingModels;
+
 	/** True if this material reads any vertex data */
 	bool bRequiresVertexData;
 };
