[GLTFExporter] fix: improve shading model evaluation by ignoring raytracing quality switch in material analyzer

commit ad7f0b33c5015179de9bc68d7dd9899712d0ee06
Date: Sun, 15 Nov 2020 20:29:12 +0100

    fix: improve shading model evaluation by ignoring raytracing quality switch in material analyzer
    (cherry picked from commit 682d2175185ba74ff80097a4bc7bf56988cf372c​)

#skipci

diff --git a/Engine/Plugins/Enterprise/GLTFExporter/MigratedFromGit/Source/GLTFExporter/Private/Converters/GLTFMaterialUtility.cpp b/Engine/Plugins/Enterprise/GLTFExporter/MigratedFromGit/Source/GLTFExporter/Private/Converters/GLTFMaterialUtility.cpp
--- a/Engine/Plugins/Enterprise/GLTFExporter/MigratedFromGit/Source/GLTFExporter/Private/Converters/GLTFMaterialUtility.cpp
+++ b/Engine/Plugins/Enterprise/GLTFExporter/MigratedFromGit/Source/GLTFExporter/Private/Converters/GLTFMaterialUtility.cpp
@@ -3,7 +3,6 @@
 #include "Converters/GLTFMaterialUtility.h"
 #include "Converters/GLTFTextureUtility.h"
 #include "Converters/GLTFNameUtility.h"
-#include "GLTFMaterialAnalysis.h"
 #include "GLTFMaterialAnalyzer.h"
 #include "Engine/TextureRenderTarget2D.h"
 #include "CanvasItem.h"
@@ -16,6 +15,7 @@
 #include "Materials/MaterialExpressionClearCoatNormalCustomOutput.h"
 #include "Materials/MaterialExpressionTextureCoordinate.h"
 #include "MeshDescription.h"
+#include "Misc/DefaultValueHelper.h"
 
 UMaterialInterface* FGLTFMaterialUtility::GetDefault()
 {
@@ -358,6 +358,13 @@ FMaterialShadingModelField FGLTFMaterialUtility::EvaluateShadingModelExpression(
 {
 	FGLTFMaterialAnalysis Analysis;
 	AnalyzeMaterialProperty(Material, MP_ShadingModel, Analysis);
+
+	int32 Value;
+	if (FDefaultValueHelper::ParseInt(Analysis.ParameterCode, Value))
+	{
+		return static_cast<EMaterialShadingModel>(Value);
+    }
+
 	return Analysis.ShadingModels;
 }
 
diff --git a/Engine/Plugins/Enterprise/GLTFExporter/MigratedFromGit/Source/GLTFMaterialAnalyzer/Private/GLTFMaterialAnalyzer.cpp b/Engine/Plugins/Enterprise/GLTFExporter/MigratedFromGit/Source/GLTFMaterialAnalyzer/Private/GLTFMaterialAnalyzer.cpp
--- a/Engine/Plugins/Enterprise/GLTFExporter/MigratedFromGit/Source/GLTFMaterialAnalyzer/Private/GLTFMaterialAnalyzer.cpp
+++ b/Engine/Plugins/Enterprise/GLTFExporter/MigratedFromGit/Source/GLTFMaterialAnalyzer/Private/GLTFMaterialAnalyzer.cpp
@@ -1,8 +1,7 @@
 // Copyright Epic Games, Inc. All Rights Reserved.
 
 #include "GLTFMaterialAnalyzer.h"
-#include "GLTFMaterialAnalysis.h"
-#include "GLTFMaterialBaking/Public/MaterialBakingStructures.h"
+#include "GLTFProxyMaterialCompiler.h"
 #include "Materials/HLSLMaterialTranslator.h"
 
 UGLTFMaterialAnalyzer::UGLTFMaterialAnalyzer()
@@ -64,14 +63,7 @@ int32 UGLTFMaterialAnalyzer::CompilePropertyEx(FMaterialCompiler* Compiler, cons
 		return INDEX_NONE;
 	}
 
-	class FTranslatorHack : public FHLSLMaterialTranslator
-	{
-		using FHLSLMaterialTranslator::FHLSLMaterialTranslator;
-
-		friend UGLTFMaterialAnalyzer;
-	};
-
-	FTranslatorHack* Translator = static_cast<FTranslatorHack*>(Compiler);
+	FGLTFProxyMaterialCompiler ProxyCompiler(Compiler);
 	int32 Result;
 
 	if (Property == MP_CustomOutput)
@@ -82,23 +74,37 @@ int32 UGLTFMaterialAnalyzer::CompilePropertyEx(FMaterialCompiler* Compiler, cons
 			return INDEX_NONE;
 		}
 
-		Result = CustomOutputExpression->Compile(Compiler, 0);
+		Result = CustomOutputExpression->Compile(&ProxyCompiler, 0);
 	}
 	else
 	{
-		Result = Material->CompilePropertyEx(Translator, AttributeID);
+		Result = Material->CompilePropertyEx(&ProxyCompiler, AttributeID);
 	}
 
-	Analysis->TextureCoordinates = Translator->AllocatedUserTexCoords;
-	Analysis->ShadingModels = Translator->ShadingModelsFromCompilation;
+	class FHLSLMaterialTranslatorHack : public FHLSLMaterialTranslator
+	{
+		using FHLSLMaterialTranslator::FHLSLMaterialTranslator;
+
+		friend UGLTFMaterialAnalyzer;
+	};
+
+	FHLSLMaterialTranslatorHack* TranslatorHack = static_cast<FHLSLMaterialTranslatorHack*>(Compiler);
+
+	Analysis->TextureCoordinates = TranslatorHack->AllocatedUserTexCoords;
+	Analysis->ShadingModels = TranslatorHack->ShadingModelsFromCompilation;
+
+	if (Result != INDEX_NONE)
+	{
+		Analysis->ParameterCode = TranslatorHack->GetParameterCode(Result);
+	}
 
 	// TODO: investigate if we need to check more conditions to determine that vertex data is required
 	Analysis->bRequiresVertexData =
-		Translator->bUsesVertexColor ||
-		Translator->bUsesTransformVector ||
-		Translator->bNeedsWorldPositionExcludingShaderOffsets ||
-		Translator->bUsesAOMaterialMask ||
-		Translator->bUsesVertexPosition;
+		TranslatorHack->bUsesVertexColor ||
+		TranslatorHack->bUsesTransformVector ||
+		TranslatorHack->bNeedsWorldPositionExcludingShaderOffsets ||
+		TranslatorHack->bUsesAOMaterialMask ||
+		TranslatorHack->bUsesVertexPosition;
 
 	return Result;
 }
diff --git a/Engine/Plugins/Enterprise/GLTFExporter/MigratedFromGit/Source/GLTFMaterialAnalyzer/Private/GLTFProxyMaterialCompiler.h b/Engine/Plugins/Enterprise/GLTFExporter/MigratedFromGit/Source/GLTFMaterialAnalyzer/Private/GLTFProxyMaterialCompiler.h
new file mode 100644
--- /dev/null
+++ b/Engine/Plugins/Enterprise/GLTFExporter/MigratedFromGit/Source/GLTFMaterialAnalyzer/Private/GLTFProxyMaterialCompiler.h
@@ -0,0 +1,60 @@
+// Copyright Epic Games, Inc. All Rights Reserved.
+
+#pragma once
+
+#include "MaterialCompiler.h"
+
+class FGLTFProxyMaterialCompiler : public FProxyMaterialCompiler
+{
+    using FProxyMaterialCompiler::FProxyMaterialCompiler;
+
+    virtual int32 RayTracingQualitySwitchReplace(int32 Normal, int32 RayTraced) override
+    {
+    	return Normal; // ignore RayTraced branch
+    }
+
+    virtual int32 ReflectionAboutCustomWorldNormal(int32 CustomWorldNormal, int32 bNormalizeCustomWorldNormal) override
+    {
+    	return Compiler->ReflectionAboutCustomWorldNormal(CustomWorldNormal, bNormalizeCustomWorldNormal);
+    }
+
+    virtual int32 ParticleRelativeTime() override
+    {
+    	return Compiler->ParticleRelativeTime();
+    }
+
+    virtual int32 ParticleMotionBlurFade() override
+    {
+    	return Compiler->ParticleMotionBlurFade();
+    }
+
+    virtual int32 ParticleRandom() override
+    {
+    	return Compiler->ParticleRandom();
+    }
+
+    virtual int32 ParticleDirection() override
+    {
+    	return Compiler->ParticleDirection();
+    }
+
+    virtual int32 ParticleSpeed() override
+    {
+    	return Compiler->ParticleSpeed();
+    }
+
+    virtual int32 ParticleSize() override
+    {
+    	return Compiler->ParticleSize();
+    }
+
+    virtual int32 VertexInterpolator(uint32 InterpolatorIndex) override
+    {
+    	return Compiler->VertexInterpolator(InterpolatorIndex);
+    }
+
+    virtual int32 MaterialBakingWorldPosition() override
+    {
+    	return Compiler->MaterialBakingWorldPosition();
+    }
+};
diff --git a/Engine/Plugins/Enterprise/GLTFExporter/MigratedFromGit/Source/GLTFMaterialAnalyzer/Public/GLTFMaterialAnalysis.h b/Engine/Plugins/Enterprise/GLTFExporter/MigratedFromGit/Source/GLTFMaterialAnalyzer/Public/GLTFMaterialAnalysis.h
--- a/Engine/Plugins/Enterprise/GLTFExporter/MigratedFromGit/Source/GLTFMaterialAnalyzer/Public/GLTFMaterialAnalysis.h
+++ b/Engine/Plugins/Enterprise/GLTFExporter/MigratedFromGit/Source/GLTFMaterialAnalyzer/Public/GLTFMaterialAnalysis.h
@@ -2,6 +2,8 @@
 
 #pragma once
 
+#include "CoreMinimal.h"
+
 struct FGLTFMaterialAnalysis
 {
 	/** Tracks the texture coordinates used by this material */
@@ -10,6 +12,9 @@ struct FGLTFMaterialAnalysis
 	/** Will contain all the shading models picked up from the material expression graph */
 	FMaterialShadingModelField ShadingModels;
 
+	/** The resulting code corresponding to the currently compiled property or custom output. */
+	FString ParameterCode;
+
 	/** True if this material reads any vertex data */
 	bool bRequiresVertexData;
 };
diff --git a/Engine/Plugins/Enterprise/GLTFExporter/MigratedFromGit/Source/GLTFMaterialAnalyzer/Public/GLTFMaterialAnalyzer.h b/Engine/Plugins/Enterprise/GLTFExporter/MigratedFromGit/Source/GLTFMaterialAnalyzer/Public/GLTFMaterialAnalyzer.h
--- a/Engine/Plugins/Enterprise/GLTFExporter/MigratedFromGit/Source/GLTFMaterialAnalyzer/Public/GLTFMaterialAnalyzer.h
+++ b/Engine/Plugins/Enterprise/GLTFExporter/MigratedFromGit/Source/GLTFMaterialAnalyzer/Public/GLTFMaterialAnalyzer.h
@@ -3,10 +3,9 @@
 #pragma once
 
 #include "Materials/MaterialInterface.h"
+#include "GLTFMaterialAnalysis.h"
 #include "GLTFMaterialAnalyzer.generated.h"
 
-struct FGLTFMaterialAnalysis;
-
 UCLASS(NotBlueprintType, Transient)
 class GLTFMATERIALANALYZER_API UGLTFMaterialAnalyzer : public UMaterialInterface
 {
