[GLTFExporter] refactor: remove unecessary error message "AddObject failed"

commit 1c0375aaa6d97c9bbefbf96250cd9a39bbabb928
Date: Sun, 15 Nov 2020 20:44:35 +0100

    refactor: remove unecessary error message "AddObject failed"
    (cherry picked from commit 76202934f8fc4614c13eed5a755c87ac7f6bc088​)

#skipci

diff --git a/Engine/Plugins/Enterprise/GLTFExporter/MigratedFromGit/Source/GLTFExporter/Private/Exporters/GLTFExporter.cpp b/Engine/Plugins/Enterprise/GLTFExporter/MigratedFromGit/Source/GLTFExporter/Private/Exporters/GLTFExporter.cpp
--- a/Engine/Plugins/Enterprise/GLTFExporter/MigratedFromGit/Source/GLTFExporter/Private/Exporters/GLTFExporter.cpp
+++ b/Engine/Plugins/Enterprise/GLTFExporter/MigratedFromGit/Source/GLTFExporter/Private/Exporters/GLTFExporter.cpp
@@ -40,11 +40,6 @@ bool UGLTFExporter::ExportBinary(UObject* Object, const TCHAR* Type, FArchive& A
 	{
 		Builder.Write(Archive, Warn);
 	}
-	else
-	{
-		// TODO: more descriptive error
-		Builder.AddErrorMessage(TEXT("AddObject failed"));
-	}
 
 	// TODO: should we copy messages to UAssetExportTask::Errors?
 
