[GLTFExporter] Revert "refactor: when showing export files, select gltf/glb file"

commit 86a95dc267b0258ad7b638295c2569c4821cd25f
Date: Sun, 15 Nov 2020 20:45:28 +0100

    Revert "refactor: when showing export files, select gltf/glb file"
    This reverts commit 0b540651107414fd47395e0a3891419afd0389c9.
    (cherry picked from commit d595c631f88dcac1656244365d66872a3c1eba2d​)

#skipci

diff --git a/Engine/Plugins/Enterprise/GLTFExporter/MigratedFromGit/Source/GLTFExporter/Private/Exporters/GLTFExporter.cpp b/Engine/Plugins/Enterprise/GLTFExporter/MigratedFromGit/Source/GLTFExporter/Private/Exporters/GLTFExporter.cpp
--- a/Engine/Plugins/Enterprise/GLTFExporter/MigratedFromGit/Source/GLTFExporter/Private/Exporters/GLTFExporter.cpp
+++ b/Engine/Plugins/Enterprise/GLTFExporter/MigratedFromGit/Source/GLTFExporter/Private/Exporters/GLTFExporter.cpp
@@ -53,7 +53,7 @@ bool UGLTFExporter::ExportBinary(UObject* Object, const TCHAR* Type, FArchive& A
 
 		if (bSuccess && Options->bShowFilesWhenDone)
 		{
-			FPlatformProcess::ExploreFolder(*Builder.FilePath);
+			FPlatformProcess::ExploreFolder(*Builder.DirPath);
 		}
 	}
 
