[GLTFExporter] feat: add and implement custom extension EPIC_blend_modes

commit 1c712711b52cd56757c1fb0e5d3504429bd3b4d7
Date: Sun, 15 Nov 2020 21:21:15 +0100

    feat: add and implement custom extension EPIC_blend_modes
    (cherry picked from commit b59665c76c9ff00e88db93faade0440ae0b9eee1​)

#skipci

diff --git a/Engine/Plugins/Enterprise/GLTFExporter/MigratedFromGit/Source/GLTFExporter/Private/Converters/GLTFConverterUtility.cpp b/Engine/Plugins/Enterprise/GLTFExporter/MigratedFromGit/Source/GLTFExporter/Private/Converters/GLTFConverterUtility.cpp
--- a/Engine/Plugins/Enterprise/GLTFExporter/MigratedFromGit/Source/GLTFExporter/Private/Converters/GLTFConverterUtility.cpp
+++ b/Engine/Plugins/Enterprise/GLTFExporter/MigratedFromGit/Source/GLTFExporter/Private/Converters/GLTFConverterUtility.cpp
@@ -46,14 +46,28 @@ EGLTFJsonShadingModel FGLTFConverterUtility::ConvertShadingModel(EMaterialShadin
 	}
 }
 
-EGLTFJsonAlphaMode FGLTFConverterUtility::ConvertBlendMode(EBlendMode Mode)
+EGLTFJsonAlphaMode FGLTFConverterUtility::ConvertAlphaMode(EBlendMode Mode)
 {
 	switch (Mode)
 	{
-		case BLEND_Opaque:      return EGLTFJsonAlphaMode::Opaque;
-		case BLEND_Translucent: return EGLTFJsonAlphaMode::Blend;
-		case BLEND_Masked:      return EGLTFJsonAlphaMode::Mask;
-		default:                return EGLTFJsonAlphaMode::None;
+		case BLEND_Opaque:         return EGLTFJsonAlphaMode::Opaque;
+		case BLEND_Masked:         return EGLTFJsonAlphaMode::Mask;
+		case BLEND_Translucent:    return EGLTFJsonAlphaMode::Blend;
+		case BLEND_Additive:       return EGLTFJsonAlphaMode::Blend;
+		case BLEND_Modulate:       return EGLTFJsonAlphaMode::Blend;
+		case BLEND_AlphaComposite: return EGLTFJsonAlphaMode::Blend;
+		default:                   return EGLTFJsonAlphaMode::None;
+	}
+}
+
+EGLTFJsonBlendMode FGLTFConverterUtility::ConvertBlendMode(EBlendMode Mode)
+{
+	switch (Mode)
+	{
+		case BLEND_Additive:       return EGLTFJsonBlendMode::Additive;
+		case BLEND_Modulate:       return EGLTFJsonBlendMode::Modulate;
+		case BLEND_AlphaComposite: return EGLTFJsonBlendMode::AlphaComposite;
+		default:                   return EGLTFJsonBlendMode::None;
 	}
 }
 
diff --git a/Engine/Plugins/Enterprise/GLTFExporter/MigratedFromGit/Source/GLTFExporter/Private/Converters/GLTFConverterUtility.h b/Engine/Plugins/Enterprise/GLTFExporter/MigratedFromGit/Source/GLTFExporter/Private/Converters/GLTFConverterUtility.h
--- a/Engine/Plugins/Enterprise/GLTFExporter/MigratedFromGit/Source/GLTFExporter/Private/Converters/GLTFConverterUtility.h
+++ b/Engine/Plugins/Enterprise/GLTFExporter/MigratedFromGit/Source/GLTFExporter/Private/Converters/GLTFConverterUtility.h
@@ -183,7 +183,8 @@ struct FGLTFConverterUtility
 
 	static EGLTFJsonShadingModel ConvertShadingModel(EMaterialShadingModel ShadingModel);
 
-	static EGLTFJsonAlphaMode ConvertBlendMode(EBlendMode Mode);
+	static EGLTFJsonAlphaMode ConvertAlphaMode(EBlendMode Mode);
+	static EGLTFJsonBlendMode ConvertBlendMode(EBlendMode Mode);
 
 	static EGLTFJsonTextureWrap ConvertWrap(TextureAddress Address);
 
diff --git a/Engine/Plugins/Enterprise/GLTFExporter/MigratedFromGit/Source/GLTFExporter/Private/Json/GLTFJsonEnums.h b/Engine/Plugins/Enterprise/GLTFExporter/MigratedFromGit/Source/GLTFExporter/Private/Json/GLTFJsonEnums.h
--- a/Engine/Plugins/Enterprise/GLTFExporter/MigratedFromGit/Source/GLTFExporter/Private/Json/GLTFJsonEnums.h
+++ b/Engine/Plugins/Enterprise/GLTFExporter/MigratedFromGit/Source/GLTFExporter/Private/Json/GLTFJsonEnums.h
@@ -7,19 +7,19 @@
 enum class EGLTFJsonExtension
 {
 	KHR_LightsPunctual,
-	KHR_MaterialsUnlit,
 	KHR_MaterialsClearCoat,
+	KHR_MaterialsUnlit,
 	KHR_MeshQuantization,
 	KHR_TextureTransform,
+	EPIC_AnimationHotspots,
+	EPIC_AnimationPlayback,
+	EPIC_BlendModes,
 	EPIC_CameraControls,
 	EPIC_HDRIBackdrops,
-	EPIC_AnimationHotspots,
-	EPIC_LightmapTextures,
-	EPIC_ReflectionCaptures,
 	EPIC_LevelVariantSets,
-	EPIC_TextureHDREncoding,
-	EPIC_AnimationPlayback,
-	EPIC_SkySpheres
+	EPIC_LightmapTextures,
+	EPIC_SkySpheres,
+	EPIC_TextureHDREncoding
 };
 
 enum class EGLTFJsonShadingModel
@@ -101,6 +101,14 @@ enum class EGLTFJsonAlphaMode
 	Blend
 };
 
+enum class EGLTFJsonBlendMode
+{
+	None = -1,
+	Additive,
+	Modulate,
+	AlphaComposite
+};
+
 enum class EGLTFJsonMimeType
 {
 	None = -1,
diff --git a/Engine/Plugins/Enterprise/GLTFExporter/MigratedFromGit/Source/GLTFExporter/Private/Json/GLTFJsonMaterial.h b/Engine/Plugins/Enterprise/GLTFExporter/MigratedFromGit/Source/GLTFExporter/Private/Json/GLTFJsonMaterial.h
--- a/Engine/Plugins/Enterprise/GLTFExporter/MigratedFromGit/Source/GLTFExporter/Private/Json/GLTFJsonMaterial.h
+++ b/Engine/Plugins/Enterprise/GLTFExporter/MigratedFromGit/Source/GLTFExporter/Private/Json/GLTFJsonMaterial.h
@@ -233,6 +233,8 @@ struct FGLTFJsonMaterial
 
 	bool DoubleSided;
 
+	EGLTFJsonBlendMode BlendMode;
+
 	FGLTFJsonClearCoatExtension ClearCoat;
 
 	FGLTFJsonMaterial()
@@ -241,6 +243,7 @@ struct FGLTFJsonMaterial
 		, AlphaMode(EGLTFJsonAlphaMode::Opaque)
 		, AlphaCutoff(0.5f)
 		, DoubleSided(false)
+		, BlendMode(EGLTFJsonBlendMode::None)
 	{
 	}
 
@@ -299,30 +302,37 @@ struct FGLTFJsonMaterial
 			JsonWriter.WriteValue(TEXT("doubleSided"), DoubleSided);
 		}
 
-		if (ShadingModel == EGLTFJsonShadingModel::Unlit || ShadingModel == EGLTFJsonShadingModel::ClearCoat)
+		if (BlendMode != EGLTFJsonBlendMode::None || ShadingModel == EGLTFJsonShadingModel::Unlit || ShadingModel == EGLTFJsonShadingModel::ClearCoat)
 		{
-			const EGLTFJsonExtension Extension = ShadingModel == EGLTFJsonShadingModel::Unlit ?
-				EGLTFJsonExtension::KHR_MaterialsUnlit : EGLTFJsonExtension::KHR_MaterialsClearCoat;
+			JsonWriter.WriteObjectStart(TEXT("extensions"));
 
-			Extensions.Used.Add(Extension);
+			if (BlendMode != EGLTFJsonBlendMode::None)
+			{
+				const EGLTFJsonExtension Extension = EGLTFJsonExtension::EPIC_BlendModes;
+				Extensions.Used.Add(Extension);
 
-			JsonWriter.WriteObjectStart(TEXT("extensions"));
-			JsonWriter.WriteIdentifierPrefix(FGLTFJsonUtility::ToString(Extension));
+				JsonWriter.WriteIdentifierPrefix(FGLTFJsonUtility::ToString(Extension));
+				JsonWriter.WriteObjectStart();
+				JsonWriter.WriteValue(TEXT("blendMode"), FGLTFJsonUtility::ToString(BlendMode));
+				JsonWriter.WriteObjectEnd();
+			}
 
-			switch (ShadingModel)
+			if (ShadingModel == EGLTFJsonShadingModel::Unlit)
 			{
-				case EGLTFJsonShadingModel::Unlit:
-					JsonWriter.WriteObjectStart(); // Write empty object
-					JsonWriter.WriteObjectEnd();
-					break;
-
-				case EGLTFJsonShadingModel::ClearCoat:
-					ClearCoat.WriteObject(JsonWriter, Extensions);
-					break;
-
-				default:
-					checkNoEntry();
-					break;
+				const EGLTFJsonExtension Extension = EGLTFJsonExtension::KHR_MaterialsUnlit;
+				Extensions.Used.Add(Extension);
+
+				JsonWriter.WriteIdentifierPrefix(FGLTFJsonUtility::ToString(Extension));
+				JsonWriter.WriteObjectStart(); // Write empty object
+				JsonWriter.WriteObjectEnd();
+			}
+			else if (ShadingModel == EGLTFJsonShadingModel::ClearCoat)
+			{
+				const EGLTFJsonExtension Extension = EGLTFJsonExtension::KHR_MaterialsClearCoat;
+				Extensions.Used.Add(Extension);
+
+				JsonWriter.WriteIdentifierPrefix(FGLTFJsonUtility::ToString(Extension));
+				ClearCoat.WriteObject(JsonWriter, Extensions);
 			}
 
 			JsonWriter.WriteObjectEnd();
diff --git a/Engine/Plugins/Enterprise/GLTFExporter/MigratedFromGit/Source/GLTFExporter/Private/Json/GLTFJsonUtility.h b/Engine/Plugins/Enterprise/GLTFExporter/MigratedFromGit/Source/GLTFExporter/Private/Json/GLTFJsonUtility.h
--- a/Engine/Plugins/Enterprise/GLTFExporter/MigratedFromGit/Source/GLTFExporter/Private/Json/GLTFJsonUtility.h
+++ b/Engine/Plugins/Enterprise/GLTFExporter/MigratedFromGit/Source/GLTFExporter/Private/Json/GLTFJsonUtility.h
@@ -18,19 +18,19 @@ struct FGLTFJsonUtility
 		switch (Value)
 		{
 			case EGLTFJsonExtension::KHR_LightsPunctual:      return TEXT("KHR_lights_punctual");
-			case EGLTFJsonExtension::KHR_MaterialsUnlit:      return TEXT("KHR_materials_unlit");
 			case EGLTFJsonExtension::KHR_MaterialsClearCoat:  return TEXT("KHR_materials_clearcoat");
+			case EGLTFJsonExtension::KHR_MaterialsUnlit:      return TEXT("KHR_materials_unlit");
 			case EGLTFJsonExtension::KHR_MeshQuantization:    return TEXT("KHR_mesh_quantization");
 			case EGLTFJsonExtension::KHR_TextureTransform:    return TEXT("KHR_texture_transform");
+			case EGLTFJsonExtension::EPIC_AnimationHotspots:  return TEXT("EPIC_animation_hotspots");
+			case EGLTFJsonExtension::EPIC_AnimationPlayback:  return TEXT("EPIC_animation_playback");
+			case EGLTFJsonExtension::EPIC_BlendModes:         return TEXT("EPIC_blend_modes");
 			case EGLTFJsonExtension::EPIC_CameraControls:     return TEXT("EPIC_camera_controls");
 			case EGLTFJsonExtension::EPIC_HDRIBackdrops:      return TEXT("EPIC_hdri_backdrops");
-			case EGLTFJsonExtension::EPIC_AnimationHotspots:  return TEXT("EPIC_animation_hotspots");
-			case EGLTFJsonExtension::EPIC_LightmapTextures:   return TEXT("EPIC_lightmap_textures");
-			case EGLTFJsonExtension::EPIC_ReflectionCaptures: return TEXT("EPIC_reflection_captures");
 			case EGLTFJsonExtension::EPIC_LevelVariantSets:   return TEXT("EPIC_level_variant_sets");
-			case EGLTFJsonExtension::EPIC_TextureHDREncoding: return TEXT("EPIC_texture_hdr_encoding");
-			case EGLTFJsonExtension::EPIC_AnimationPlayback:  return TEXT("EPIC_animation_playback");
+			case EGLTFJsonExtension::EPIC_LightmapTextures:   return TEXT("EPIC_lightmap_textures");
 			case EGLTFJsonExtension::EPIC_SkySpheres:         return TEXT("EPIC_sky_spheres");
+			case EGLTFJsonExtension::EPIC_TextureHDREncoding: return TEXT("EPIC_texture_hdr_encoding");
 			default:
 				checkNoEntry();
 				return TEXT("");
@@ -50,6 +50,19 @@ struct FGLTFJsonUtility
 		}
 	}
 
+	static const TCHAR* ToString(EGLTFJsonBlendMode Value)
+	{
+		switch (Value)
+		{
+			case EGLTFJsonBlendMode::Additive:       return TEXT("ADDITIVE");
+			case EGLTFJsonBlendMode::Modulate:       return TEXT("MODULATE");
+			case EGLTFJsonBlendMode::AlphaComposite: return TEXT("ALPHACOMPOSITE");
+			default:
+				checkNoEntry();
+				return TEXT("");
+		}
+	}
+
 	static const TCHAR* ToString(EGLTFJsonMimeType Value)
 	{
 		switch (Value)
diff --git a/Engine/Plugins/Enterprise/GLTFExporter/MigratedFromGit/Source/GLTFExporter/Private/Tasks/GLTFMaterialTasks.cpp b/Engine/Plugins/Enterprise/GLTFExporter/MigratedFromGit/Source/GLTFExporter/Private/Tasks/GLTFMaterialTasks.cpp
--- a/Engine/Plugins/Enterprise/GLTFExporter/MigratedFromGit/Source/GLTFExporter/Private/Tasks/GLTFMaterialTasks.cpp
+++ b/Engine/Plugins/Enterprise/GLTFExporter/MigratedFromGit/Source/GLTFExporter/Private/Tasks/GLTFMaterialTasks.cpp
@@ -64,10 +64,10 @@ void FGLTFMaterialTask::Complete()
 	FGLTFJsonMaterial& JsonMaterial = Builder.GetMaterial(MaterialIndex);
 	JsonMaterial.Name = GetMaterialName();
 
-	if (!TryGetAlphaMode(JsonMaterial.AlphaMode))
+	if (!TryGetAlphaMode(JsonMaterial.AlphaMode, JsonMaterial.BlendMode))
 	{
-		JsonMaterial.AlphaMode = EGLTFJsonAlphaMode::Opaque;
-		Builder.AddWarningMessage(FString::Printf(TEXT("Material %s will be exported as blend mode %s"), *Material->GetName(), *FGLTFNameUtility::GetName(BLEND_Opaque)));
+		JsonMaterial.AlphaMode = EGLTFJsonAlphaMode::Blend;
+		Builder.AddWarningMessage(FString::Printf(TEXT("Material %s will be exported as blend mode %s"), *Material->GetName(), *FGLTFNameUtility::GetName(BLEND_Translucent)));
 	}
 
 	JsonMaterial.AlphaCutoff = Material->GetOpacityMaskClipValue();
@@ -175,21 +175,24 @@ void FGLTFMaterialTask::Complete()
 	}
 }
 
-bool FGLTFMaterialTask::TryGetAlphaMode(EGLTFJsonAlphaMode& AlphaMode) const
+bool FGLTFMaterialTask::TryGetAlphaMode(EGLTFJsonAlphaMode& OutAlphaMode, EGLTFJsonBlendMode& OutBlendMode) const
 {
 	const EBlendMode BlendMode = Material->GetBlendMode();
 
 	// TODO: add support for additional blend modes (like Additive and Modulate)?
 
-	const EGLTFJsonAlphaMode ConvertedAlphaMode = FGLTFConverterUtility::ConvertBlendMode(BlendMode);
-	if (ConvertedAlphaMode == EGLTFJsonAlphaMode::None)
+	const EGLTFJsonAlphaMode AlphaMode = FGLTFConverterUtility::ConvertAlphaMode(BlendMode);
+	if (AlphaMode == EGLTFJsonAlphaMode::None)
 	{
-		const FString BlendModeName = FGLTFNameUtility::GetName(BlendMode);
-		Builder.AddWarningMessage(FString::Printf(TEXT("Unsupported blend mode (%s) in material %s"), *BlendModeName, *Material->GetName()));
+		Builder.AddWarningMessage(FString::Printf(
+			TEXT("Unsupported blend mode (%s) in material %s"),
+			*FGLTFNameUtility::GetName(BlendMode),
+			*Material->GetName()));
 		return false;
 	}
 
-	AlphaMode = ConvertedAlphaMode;
+	OutAlphaMode = AlphaMode;
+	OutBlendMode = AlphaMode == EGLTFJsonAlphaMode::Blend ? FGLTFConverterUtility::ConvertBlendMode(BlendMode) : EGLTFJsonBlendMode::None;
 	return true;
 }
 
diff --git a/Engine/Plugins/Enterprise/GLTFExporter/MigratedFromGit/Source/GLTFExporter/Private/Tasks/GLTFMaterialTasks.h b/Engine/Plugins/Enterprise/GLTFExporter/MigratedFromGit/Source/GLTFExporter/Private/Tasks/GLTFMaterialTasks.h
--- a/Engine/Plugins/Enterprise/GLTFExporter/MigratedFromGit/Source/GLTFExporter/Private/Tasks/GLTFMaterialTasks.h
+++ b/Engine/Plugins/Enterprise/GLTFExporter/MigratedFromGit/Source/GLTFExporter/Private/Tasks/GLTFMaterialTasks.h
@@ -39,8 +39,8 @@ private:
 	const TArray<int32> SectionIndices;
 	const FGLTFJsonMaterialIndex MaterialIndex;
 
-	bool TryGetAlphaMode(EGLTFJsonAlphaMode& AlphaMode) const;
-	bool TryGetShadingModel(EGLTFJsonShadingModel& ShadingModel) const;
+	bool TryGetAlphaMode(EGLTFJsonAlphaMode& OutAlphaMode, EGLTFJsonBlendMode& OutBlendMode) const;
+	bool TryGetShadingModel(EGLTFJsonShadingModel& OutShadingModel) const;
 
 	bool TryGetBaseColorAndOpacity(FGLTFJsonPBRMetallicRoughness& OutPBRParams, const FMaterialPropertyEx& BaseColorProperty, const FMaterialPropertyEx& OpacityProperty) const;
 	bool TryGetMetallicAndRoughness(FGLTFJsonPBRMetallicRoughness& OutPBRParams, const FMaterialPropertyEx& MetallicProperty, const FMaterialPropertyEx& RoughnessProperty) const;
