[GLTFExporter] feat: add and implement export option bExportExtraBlendModes

commit 6c5571d502cecb64b836318ab21e3ee58be73213
Date: Sun, 15 Nov 2020 22:14:02 +0100

    feat: add and implement export option bExportExtraBlendModes
    (cherry picked from commit 7491fe5bb96ef48e70f2c54a69eb38b3a704bba1​)

#skipci

diff --git a/Engine/Plugins/Enterprise/GLTFExporter/MigratedFromGit/Source/GLTFExporter/Private/GLTFExportOptions.cpp b/Engine/Plugins/Enterprise/GLTFExporter/MigratedFromGit/Source/GLTFExporter/Private/GLTFExportOptions.cpp
--- a/Engine/Plugins/Enterprise/GLTFExporter/MigratedFromGit/Source/GLTFExporter/Private/GLTFExportOptions.cpp
+++ b/Engine/Plugins/Enterprise/GLTFExporter/MigratedFromGit/Source/GLTFExporter/Private/GLTFExportOptions.cpp
@@ -16,6 +16,7 @@ void UGLTFExportOptions::ResetToDefault()
 {
 	bExportUnlitMaterials = true;
 	bExportClearCoatMaterials = true;
+	bExportExtraBlendModes = true;
 	BakeMaterialInputs = EGLTFMaterialBakeMode::Simple;
 	DefaultMaterialBakeSize = EGLTFMaterialBakeSizePOT::POT_1024;
 	DefaultLevelOfDetail = 0;
diff --git a/Engine/Plugins/Enterprise/GLTFExporter/MigratedFromGit/Source/GLTFExporter/Private/Tasks/GLTFMaterialTasks.cpp b/Engine/Plugins/Enterprise/GLTFExporter/MigratedFromGit/Source/GLTFExporter/Private/Tasks/GLTFMaterialTasks.cpp
--- a/Engine/Plugins/Enterprise/GLTFExporter/MigratedFromGit/Source/GLTFExporter/Private/Tasks/GLTFMaterialTasks.cpp
+++ b/Engine/Plugins/Enterprise/GLTFExporter/MigratedFromGit/Source/GLTFExporter/Private/Tasks/GLTFMaterialTasks.cpp
@@ -179,10 +179,8 @@ bool FGLTFMaterialTask::TryGetAlphaMode(EGLTFJsonAlphaMode& OutAlphaMode, EGLTFJ
 {
 	const EBlendMode BlendMode = Material->GetBlendMode();
 
-	// TODO: add support for additional blend modes (like Additive and Modulate)?
-
-	const EGLTFJsonAlphaMode AlphaMode = FGLTFConverterUtility::ConvertAlphaMode(BlendMode);
-	if (AlphaMode == EGLTFJsonAlphaMode::None)
+	const EGLTFJsonAlphaMode ConvertedAlphaMode = FGLTFConverterUtility::ConvertAlphaMode(BlendMode);
+	if (ConvertedAlphaMode == EGLTFJsonAlphaMode::None)
 	{
 		Builder.AddWarningMessage(FString::Printf(
 			TEXT("Unsupported blend mode (%s) in material %s"),
@@ -191,8 +189,18 @@ bool FGLTFMaterialTask::TryGetAlphaMode(EGLTFJsonAlphaMode& OutAlphaMode, EGLTFJ
 		return false;
 	}
 
-	OutAlphaMode = AlphaMode;
-	OutBlendMode = AlphaMode == EGLTFJsonAlphaMode::Blend ? FGLTFConverterUtility::ConvertBlendMode(BlendMode) : EGLTFJsonBlendMode::None;
+	const EGLTFJsonBlendMode ConvertedBlendMode = ConvertedAlphaMode == EGLTFJsonAlphaMode::Blend ? FGLTFConverterUtility::ConvertBlendMode(BlendMode) : EGLTFJsonBlendMode::None;
+	if (ConvertedBlendMode != EGLTFJsonBlendMode::None && !Builder.ExportOptions->bExportExtraBlendModes)
+	{
+		Builder.AddWarningMessage(FString::Printf(
+			TEXT("Extra blend mode (%s) in material %s disabled by export options"),
+			*FGLTFNameUtility::GetName(BlendMode),
+			*Material->GetName()));
+		return false;
+	}
+
+	OutAlphaMode = ConvertedAlphaMode;
+	OutBlendMode = ConvertedBlendMode;
 	return true;
 }
 
diff --git a/Engine/Plugins/Enterprise/GLTFExporter/MigratedFromGit/Source/GLTFExporter/Public/GLTFExportOptions.h b/Engine/Plugins/Enterprise/GLTFExporter/MigratedFromGit/Source/GLTFExporter/Public/GLTFExportOptions.h
--- a/Engine/Plugins/Enterprise/GLTFExporter/MigratedFromGit/Source/GLTFExporter/Public/GLTFExportOptions.h
+++ b/Engine/Plugins/Enterprise/GLTFExporter/MigratedFromGit/Source/GLTFExporter/Public/GLTFExportOptions.h
@@ -81,6 +81,9 @@ public:
 	UPROPERTY(EditAnywhere, BlueprintReadWrite, Config, Category = Material)
 	bool bExportClearCoatMaterials;
 
+	UPROPERTY(EditAnywhere, BlueprintReadWrite, Config, Category = Material)
+	bool bExportExtraBlendModes;
+
 	UPROPERTY(EditAnywhere, BlueprintReadWrite, Config, Category = Material)
 	EGLTFMaterialBakeMode BakeMaterialInputs;
 
