[GLTFExporter] refactor: minor improvement to lightmap converter

commit fe96af4220f7985945aab90e44297d8512cf2635
Date: Mon, 16 Nov 2020 10:17:53 +0100

    refactor: minor improvement to lightmap converter
    (cherry picked from commit fb02230c388cca64ee1b983e62486a538c8737e8​)

#skipci

diff --git a/Engine/Plugins/Enterprise/GLTFExporter/MigratedFromGit/Source/GLTFExporter/Private/Converters/GLTFLightMapConverters.cpp b/Engine/Plugins/Enterprise/GLTFExporter/MigratedFromGit/Source/GLTFExporter/Private/Converters/GLTFLightMapConverters.cpp
--- a/Engine/Plugins/Enterprise/GLTFExporter/MigratedFromGit/Source/GLTFExporter/Private/Converters/GLTFLightMapConverters.cpp
+++ b/Engine/Plugins/Enterprise/GLTFExporter/MigratedFromGit/Source/GLTFExporter/Private/Converters/GLTFLightMapConverters.cpp
@@ -37,13 +37,13 @@ FGLTFJsonLightMapIndex FGLTFLightMapConverter::Convert(const UStaticMeshComponen
 	}
 
 	const FLightMap2D* LightMap2D = MeshMapBuildData->LightMap->GetLightMap2D();
-	check(LightMap2D);
-
-	// TODO: is it correct to use SM5?
-	const FLightMapInteraction LightMapInteraction = LightMap2D->GetInteraction(ERHIFeatureLevel::Type::SM5);
+	if (LightMap2D == nullptr)
+	{
+		return FGLTFJsonLightMapIndex(INDEX_NONE);
+	}
 
-	const bool bUseHighQualityLightMap = true;
-	const ULightMapTexture2D* Texture = LightMapInteraction.GetTexture(bUseHighQualityLightMap);
+	const FLightMapInteraction LightMapInteraction = LightMap2D->GetInteraction(GMaxRHIFeatureLevel);
+	const ULightMapTexture2D* Texture = LightMapInteraction.GetTexture(true);
 	const FGLTFJsonTextureIndex TextureIndex = Builder.GetOrAddTexture(Texture);
 
 	if (TextureIndex == INDEX_NONE)
