[GLTFExporter] refactor: update some TODO comments

commit c8f69cc06f1bb8adc5898d8c70a7f09e1779a778
Date: Mon, 16 Nov 2020 10:18:20 +0100

    refactor: update some TODO comments
    (cherry picked from commit c6d2bcd7b1d26a4f431ed6576c753a8641abb906​)

#skipci

diff --git a/Engine/Plugins/Enterprise/GLTFExporter/MigratedFromGit/Source/GLTFExporter/Private/Converters/GLTFSamplerConverters.cpp b/Engine/Plugins/Enterprise/GLTFExporter/MigratedFromGit/Source/GLTFExporter/Private/Converters/GLTFSamplerConverters.cpp
--- a/Engine/Plugins/Enterprise/GLTFExporter/MigratedFromGit/Source/GLTFExporter/Private/Converters/GLTFSamplerConverters.cpp
+++ b/Engine/Plugins/Enterprise/GLTFExporter/MigratedFromGit/Source/GLTFExporter/Private/Converters/GLTFSamplerConverters.cpp
@@ -7,7 +7,7 @@
 
 FGLTFJsonSamplerIndex FGLTFSamplerConverter::Convert(const UTexture* Texture)
 {
-	// TODO: maybe we should reuse existing samplers?
+	// TODO: reuse existing samplers, don't create a unique sampler per texture
 
 	FGLTFJsonSampler JsonSampler;
 	Texture->GetName(JsonSampler.Name);
diff --git a/Engine/Plugins/Enterprise/GLTFExporter/MigratedFromGit/Source/GLTFExporter/Private/Tasks/GLTFMaterialTasks.cpp b/Engine/Plugins/Enterprise/GLTFExporter/MigratedFromGit/Source/GLTFExporter/Private/Tasks/GLTFMaterialTasks.cpp
--- a/Engine/Plugins/Enterprise/GLTFExporter/MigratedFromGit/Source/GLTFExporter/Private/Tasks/GLTFMaterialTasks.cpp
+++ b/Engine/Plugins/Enterprise/GLTFExporter/MigratedFromGit/Source/GLTFExporter/Private/Tasks/GLTFMaterialTasks.cpp
@@ -54,11 +54,6 @@ void FGLTFMaterialTask::Complete()
 		{
 			// TODO: report warning (non-surface materials not supported, will be treated as surface)
 		}
-
-		if (!ParentMaterial->bTangentSpaceNormal)
-		{
-			// TODO: report warning (world-space normals not supported, will be treated as tangent-space)
-		}
 	}
 
 	FGLTFJsonMaterial& JsonMaterial = Builder.GetMaterial(MaterialIndex);
diff --git a/Engine/Plugins/Enterprise/GLTFExporter/MigratedFromGit/Source/GLTFExporterRuntime/Private/Actors/GLTFHotspotActor.cpp b/Engine/Plugins/Enterprise/GLTFExporter/MigratedFromGit/Source/GLTFExporterRuntime/Private/Actors/GLTFHotspotActor.cpp
--- a/Engine/Plugins/Enterprise/GLTFExporter/MigratedFromGit/Source/GLTFExporterRuntime/Private/Actors/GLTFHotspotActor.cpp
+++ b/Engine/Plugins/Enterprise/GLTFExporter/MigratedFromGit/Source/GLTFExporterRuntime/Private/Actors/GLTFHotspotActor.cpp
@@ -431,7 +431,7 @@ void AGLTFHotspotActor::ToggleAnimation()
 		}
 		else
 		{
-			// TODO: improve this calculation
+			// TODO: is there a simpler way to get the length (in seconds) of a level sequence?
 			const float SequenceLength = LevelSequence->GetMovieScene()->GetTickResolution().AsSeconds(MovieScene::DiscreteSize(LevelSequence->GetMovieScene()->GetPlaybackRange()));
 
 			FMovieSceneSequencePlaybackSettings PlaybackSettings;
