[GLTFExporter] fix: use MaterialProxy as compiler-type for GLTFProxyMaterialCompiler

commit c458de8c09422da049ac2b7c72183e7de5738a8a
Date: Mon, 16 Nov 2020 16:07:17 +0100

    fix: use MaterialProxy as compiler-type for GLTFProxyMaterialCompiler
    inlines code that would otherwise use interpolators, and allows the analyzer to better detect expressions that need vertex data
    (cherry picked from commit 6c40565105982b00446fc1113c12de7be0b540fb​)

#skipci

diff --git a/Engine/Plugins/Enterprise/GLTFExporter/MigratedFromGit/Source/GLTFMaterialAnalyzer/Private/GLTFProxyMaterialCompiler.h b/Engine/Plugins/Enterprise/GLTFExporter/MigratedFromGit/Source/GLTFMaterialAnalyzer/Private/GLTFProxyMaterialCompiler.h
--- a/Engine/Plugins/Enterprise/GLTFExporter/MigratedFromGit/Source/GLTFMaterialAnalyzer/Private/GLTFProxyMaterialCompiler.h
+++ b/Engine/Plugins/Enterprise/GLTFExporter/MigratedFromGit/Source/GLTFMaterialAnalyzer/Private/GLTFProxyMaterialCompiler.h
@@ -57,4 +57,9 @@ class FGLTFProxyMaterialCompiler : public FProxyMaterialCompiler
     {
     	return Compiler->MaterialBakingWorldPosition();
     }
+
+	virtual EMaterialCompilerType GetCompilerType() const override
+	{
+		return EMaterialCompilerType::MaterialProxy;
+	}
 };
