[GLTFExporter] fix: fix bug in utility that retrieves texture address/wrap pair for a texture

commit f365b67c52a1a776e09386cc9ca54a1a76364cd0
Date: Mon, 16 Nov 2020 16:41:32 +0100

    fix: fix bug in utility that retrieves texture address/wrap pair for a texture
    (cherry picked from commit 0c576424afa2e769839f6e5393d8b9ecbf9a5311​)

#skipci

diff --git a/Engine/Plugins/Enterprise/GLTFExporter/MigratedFromGit/Source/GLTFExporter/Private/Converters/GLTFTextureUtility.cpp b/Engine/Plugins/Enterprise/GLTFExporter/MigratedFromGit/Source/GLTFExporter/Private/Converters/GLTFTextureUtility.cpp
--- a/Engine/Plugins/Enterprise/GLTFExporter/MigratedFromGit/Source/GLTFExporter/Private/Converters/GLTFTextureUtility.cpp
+++ b/Engine/Plugins/Enterprise/GLTFExporter/MigratedFromGit/Source/GLTFExporter/Private/Converters/GLTFTextureUtility.cpp
@@ -92,12 +92,12 @@ TTuple<TextureAddress, TextureAddress> FGLTFTextureUtility::GetAddressXY(const U
 	if (const UTexture2D* Texture2D = Cast<UTexture2D>(Texture))
 	{
 		AddressX = Texture2D->AddressX;
-		AddressY = Texture2D->AddressX;
+		AddressY = Texture2D->AddressY;
 	}
 	else if (const UTextureRenderTarget2D* RenderTarget2D = Cast<UTextureRenderTarget2D>(Texture))
 	{
 		AddressX = RenderTarget2D->AddressX;
-		AddressY = RenderTarget2D->AddressX;
+		AddressY = RenderTarget2D->AddressY;
 	}
 	else
 	{
