[GLTFExporter] refactor: remove unnecessary bypass for shading model in baking module

commit dfd4aac3513a8ff4b90844d5cf087257df507e03
Date: Sat, 14 Nov 2020 00:23:00 +0100

    refactor: remove unnecessary bypass for shading model in baking module
    (cherry picked from commit 4a3214e3db0887e8f5bd5d58019693a8717b5831​)

#skipci

diff --git a/Engine/Plugins/Enterprise/GLTFExporter/MigratedFromGit/Source/GLTFMaterialBaking/Private/ExportMaterialProxy.h b/Engine/Plugins/Enterprise/GLTFExporter/MigratedFromGit/Source/GLTFMaterialBaking/Private/ExportMaterialProxy.h
--- a/Engine/Plugins/Enterprise/GLTFExporter/MigratedFromGit/Source/GLTFMaterialBaking/Private/ExportMaterialProxy.h
+++ b/Engine/Plugins/Enterprise/GLTFExporter/MigratedFromGit/Source/GLTFMaterialBaking/Private/ExportMaterialProxy.h
@@ -362,10 +362,6 @@ public:
 		// needs to be called in this function!!
 		Compiler->SetMaterialProperty(Property, OverrideShaderFrequency, bUsePreviousFrameTime);
 		const int32 Ret = CompilePropertyAndSetMaterialPropertyWithoutCast(Property, Compiler);
-		if (Property == MP_ShadingModel)
-		{
-			return Ret; // NOTE: don't try to cast shading model since it will fail
-		}
 		return Compiler->ForceCast(Ret, FMaterialAttributeDefinitionMap::GetValueType(Property), MFCF_ExactMatch | MFCF_ReplicateValue);
 	}
 
