[GLTFExporter] feat: add GLTFMaterialUserData that can be added to material assets as user data to override export options

commit 0f5cc21a2ac200f5c84a25c9ef625c12c2c627de
Date: Sat, 17 Oct 2020 00:25:54 +0200

    feat: add GLTFMaterialUserData that can be added to material assets as user data to override export options
    (cherry picked from commit 9b0f93feb012c3965e15493deb4d06bdd3fff105​)

#skipci

diff --git a/Engine/Plugins/Enterprise/GLTFExporter/MigratedFromGit/Source/GLTFExporterRuntime/Private/UserData/GLTFMaterialUserData.cpp b/Engine/Plugins/Enterprise/GLTFExporter/MigratedFromGit/Source/GLTFExporterRuntime/Private/UserData/GLTFMaterialUserData.cpp
new file mode 100644
--- /dev/null
+++ b/Engine/Plugins/Enterprise/GLTFExporter/MigratedFromGit/Source/GLTFExporterRuntime/Private/UserData/GLTFMaterialUserData.cpp
@@ -0,0 +1,65 @@
+// Copyright Epic Games, Inc. All Rights Reserved.
+
+#include "UserData/GLTFMaterialUserData.h"
+#include "Materials/MaterialInterface.h"
+
+UGLTFMaterialUserData::UGLTFMaterialUserData()
+	: DefaultBakeSize(EGLTFOverrideMaterialBakeSizePOT::POT_1024)
+{
+}
+
+EGLTFOverrideMaterialBakeSizePOT UGLTFMaterialUserData::GetBakeSizeForProperty(EMaterialProperty Property) const
+{
+	const EGLTFOverrideMaterialPropertyGroup PropertyGroup = GetPropertyGroup(Property);
+	if (const EGLTFOverrideMaterialBakeSizePOT* PropertyBakeSize = BakeSizePerProperty.Find(PropertyGroup))
+	{
+		return *PropertyBakeSize;
+	}
+
+	return DefaultBakeSize;
+}
+
+enum class EGLTFOverrideMater : uint8
+{
+	None UMETA(DisplayName = "None"),
+
+    BaseColorOpacity UMETA(DisplayName = "Base Color + Opacity (Mask)"),
+    MetallicRoughness UMETA(DisplayName = "Metallic + Roughness"),
+    EmissiveColor UMETA(DisplayName = "Emissive Color"),
+    Normal UMETA(DisplayName = "Normal"),
+    AmbientOcclusion UMETA(DisplayName = "Ambient Occlusion"),
+    ClearCoatRoughness UMETA(DisplayName = "Clear Coat + Clear Coat Roughness"),
+    ClearCoatBottomNormal UMETA(DisplayName = "Clear Coat Bottom Normal"),
+};
+
+EGLTFOverrideMaterialPropertyGroup UGLTFMaterialUserData::GetPropertyGroup(EMaterialProperty Property)
+{
+	switch (Property)
+	{
+		case MP_BaseColor:
+		case MP_Opacity:
+		case MP_OpacityMask:
+			return EGLTFOverrideMaterialPropertyGroup::BaseColorOpacity;
+		case MP_Metallic:
+		case MP_Roughness:
+			return EGLTFOverrideMaterialPropertyGroup::MetallicRoughness;
+		case MP_EmissiveColor:
+			return EGLTFOverrideMaterialPropertyGroup::EmissiveColor;
+		case MP_Normal:
+			return EGLTFOverrideMaterialPropertyGroup::Normal;
+		case MP_AmbientOcclusion:
+			return EGLTFOverrideMaterialPropertyGroup::AmbientOcclusion;
+		case MP_CustomData0:
+		case MP_CustomData1:
+			return EGLTFOverrideMaterialPropertyGroup::ClearCoatRoughness;
+		case MP_CustomOutput: // TODO: fix assumption
+			return EGLTFOverrideMaterialPropertyGroup::ClearCoatBottomNormal;
+		default:
+			return EGLTFOverrideMaterialPropertyGroup::None;
+	}
+}
+
+const UGLTFMaterialUserData* UGLTFMaterialUserData::GetUserData(const UMaterialInterface* Material)
+{
+	return const_cast<UMaterialInterface*>(Material)->GetAssetUserData<UGLTFMaterialUserData>();
+}
diff --git a/Engine/Plugins/Enterprise/GLTFExporter/MigratedFromGit/Source/GLTFExporterRuntime/Public/UserData/GLTFMaterialUserData.h b/Engine/Plugins/Enterprise/GLTFExporter/MigratedFromGit/Source/GLTFExporterRuntime/Public/UserData/GLTFMaterialUserData.h
new file mode 100644
--- /dev/null
+++ b/Engine/Plugins/Enterprise/GLTFExporter/MigratedFromGit/Source/GLTFExporterRuntime/Public/UserData/GLTFMaterialUserData.h
@@ -0,0 +1,67 @@
+// Copyright Epic Games, Inc. All Rights Reserved.
+
+#pragma once
+
+#include "Engine/AssetUserData.h"
+#include "GLTFMaterialUserData.generated.h"
+
+class UMaterialInterface;
+enum EMaterialProperty;
+
+UENUM(BlueprintType)
+enum class EGLTFOverrideMaterialBakeSizePOT : uint8
+{
+	NoOverride UMETA(DisplayName = "No Override"),
+
+    POT_1 UMETA(DisplayName = "1 x 1"),
+    POT_2 UMETA(DisplayName = "2 x 2"),
+    POT_4 UMETA(DisplayName = "4 x 4"),
+    POT_8 UMETA(DisplayName = "8 x 8"),
+    POT_16 UMETA(DisplayName = "16 x 16"),
+    POT_32 UMETA(DisplayName = "32 x 32"),
+    POT_64 UMETA(DisplayName = "64 x 64"),
+    POT_128 UMETA(DisplayName = "128 x 128"),
+    POT_256 UMETA(DisplayName = "256 x 256"),
+    POT_512 UMETA(DisplayName = "512 x 512"),
+    POT_1024 UMETA(DisplayName = "1024 x 1024"),
+    POT_2048 UMETA(DisplayName = "2048 x 2048"),
+    POT_4096 UMETA(DisplayName = "4096 x 4096"),
+    POT_8192 UMETA(DisplayName = "8192 x 8192")
+};
+
+UENUM(BlueprintType)
+enum class EGLTFOverrideMaterialPropertyGroup : uint8
+{
+	None UMETA(DisplayName = "None"),
+
+	BaseColorOpacity UMETA(DisplayName = "Base Color + Opacity (Mask)"),
+    MetallicRoughness UMETA(DisplayName = "Metallic + Roughness"),
+    EmissiveColor UMETA(DisplayName = "Emissive Color"),
+    Normal UMETA(DisplayName = "Normal"),
+    AmbientOcclusion UMETA(DisplayName = "Ambient Occlusion"),
+    ClearCoatRoughness UMETA(DisplayName = "Clear Coat + Clear Coat Roughness"),
+    ClearCoatBottomNormal UMETA(DisplayName = "Clear Coat Bottom Normal"),
+};
+
+/** glTF-specific user data that can be added to material assets to override export options */
+UCLASS(BlueprintType, meta = (DisplayName = "GLTF Material User Data"))
+class GLTFEXPORTERRUNTIME_API UGLTFMaterialUserData : public UAssetUserData
+{
+	GENERATED_BODY()
+
+public:
+
+	UGLTFMaterialUserData();
+
+	UPROPERTY(EditAnywhere, BlueprintReadWrite, Category = "Override Export Options")
+	EGLTFOverrideMaterialBakeSizePOT DefaultBakeSize;
+
+	UPROPERTY(EditAnywhere, BlueprintReadWrite, Category = "Override Export Options")
+	TMap<EGLTFOverrideMaterialPropertyGroup, EGLTFOverrideMaterialBakeSizePOT> BakeSizePerProperty;
+
+	EGLTFOverrideMaterialBakeSizePOT GetBakeSizeForProperty(EMaterialProperty Property) const;
+
+	static EGLTFOverrideMaterialPropertyGroup GetPropertyGroup(EMaterialProperty Property);
+
+	static const UGLTFMaterialUserData* GetUserData(const UMaterialInterface* Material);
+};
