[GLTFExporter] feat: implement support for overriding default material bake size using GLTFMaterialUserData

commit b16c071392e635e74ee84c4edb0d947df03547ed
Date: Sat, 17 Oct 2020 00:27:43 +0200

    feat: implement support for overriding default material bake size using GLTFMaterialUserData
    (cherry picked from commit c52e23d868bfce0467d9c29a66d841650f55c9d8​)

#skipci

diff --git a/Engine/Plugins/Enterprise/GLTFExporter/MigratedFromGit/Source/GLTFExporter/Private/Builders/GLTFBuilder.cpp b/Engine/Plugins/Enterprise/GLTFExporter/MigratedFromGit/Source/GLTFExporter/Private/Builders/GLTFBuilder.cpp
--- a/Engine/Plugins/Enterprise/GLTFExporter/MigratedFromGit/Source/GLTFExporter/Private/Builders/GLTFBuilder.cpp
+++ b/Engine/Plugins/Enterprise/GLTFExporter/MigratedFromGit/Source/GLTFExporter/Private/Builders/GLTFBuilder.cpp
@@ -1,6 +1,7 @@
 // Copyright Epic Games, Inc. All Rights Reserved.
 
 #include "Builders/GLTFBuilder.h"
+#include "UserData/GLTFMaterialUserData.h"
 #include "Builders/GLTFFileUtility.h"
 
 FGLTFBuilder::FGLTFBuilder(const FString& FilePath, const UGLTFExportOptions* ExportOptions)
@@ -11,9 +12,19 @@ FGLTFBuilder::FGLTFBuilder(const FString& FilePath, const UGLTFExportOptions* Ex
 {
 }
 
-FIntPoint FGLTFBuilder::GetDefaultMaterialBakeSize() const
+FIntPoint FGLTFBuilder::GetBakeSizeForMaterialProperty(const UMaterialInterface* Material, const FMaterialPropertyEx& Property) const
 {
-	const int32 Size = 1 << static_cast<int>(ExportOptions->DefaultMaterialBakeSize);
+	int32 Size = 1 << (static_cast<uint8>(ExportOptions->DefaultMaterialBakeSize) - static_cast<uint8>(EGLTFMaterialBakeSizePOT::POT_1));
+
+	if (const UGLTFMaterialUserData* UserData = UGLTFMaterialUserData::GetUserData(Material))
+	{
+		EGLTFOverrideMaterialBakeSizePOT OverridePOT = UserData->GetBakeSizeForProperty(Property.Type);
+		if (OverridePOT != EGLTFOverrideMaterialBakeSizePOT::NoOverride)
+		{
+			Size = 1 << (static_cast<uint8>(OverridePOT) - static_cast<uint8>(EGLTFOverrideMaterialBakeSizePOT::POT_1));
+		}
+	}
+
 	return { Size, Size };
 }
 
diff --git a/Engine/Plugins/Enterprise/GLTFExporter/MigratedFromGit/Source/GLTFExporter/Private/Builders/GLTFBuilder.h b/Engine/Plugins/Enterprise/GLTFExporter/MigratedFromGit/Source/GLTFExporter/Private/Builders/GLTFBuilder.h
--- a/Engine/Plugins/Enterprise/GLTFExporter/MigratedFromGit/Source/GLTFExporter/Private/Builders/GLTFBuilder.h
+++ b/Engine/Plugins/Enterprise/GLTFExporter/MigratedFromGit/Source/GLTFExporter/Private/Builders/GLTFBuilder.h
@@ -4,6 +4,7 @@
 
 #include "GLTFExportOptions.h"
 #include "Json/GLTFJsonEnums.h"
+#include "MaterialPropertyEx.h"
 
 class FGLTFBuilder
 {
@@ -19,7 +20,7 @@ public:
 
 	const UGLTFExportOptions* const ExportOptions;
 
-	FIntPoint GetDefaultMaterialBakeSize() const;
+	FIntPoint GetBakeSizeForMaterialProperty(const UMaterialInterface* Material, const FMaterialPropertyEx& Property) const;
 
 	EGLTFJsonHDREncoding GetTextureHDREncoding() const;
 
diff --git a/Engine/Plugins/Enterprise/GLTFExporter/MigratedFromGit/Source/GLTFExporter/Private/Tasks/GLTFMaterialTasks.cpp b/Engine/Plugins/Enterprise/GLTFExporter/MigratedFromGit/Source/GLTFExporter/Private/Tasks/GLTFMaterialTasks.cpp
--- a/Engine/Plugins/Enterprise/GLTFExporter/MigratedFromGit/Source/GLTFExporter/Private/Tasks/GLTFMaterialTasks.cpp
+++ b/Engine/Plugins/Enterprise/GLTFExporter/MigratedFromGit/Source/GLTFExporter/Private/Tasks/GLTFMaterialTasks.cpp
@@ -310,7 +310,7 @@ bool FGLTFMaterialTask::TryGetBaseColorAndOpacity(FGLTFJsonPBRMetallicRoughness&
 	}
 
 	// TODO: add support for calculating the ideal resolution to use for baking based on connected (texture) nodes
-	FIntPoint TextureSize = Builder.GetDefaultMaterialBakeSize();
+	FIntPoint TextureSize = Builder.GetBakeSizeForMaterialProperty(Material, BaseColorProperty);
 
 	// TODO: should this be the default wrap-mode?
 	EGLTFJsonTextureWrap TextureWrapS = EGLTFJsonTextureWrap::Repeat;
@@ -478,7 +478,7 @@ bool FGLTFMaterialTask::TryGetMetallicAndRoughness(FGLTFJsonPBRMetallicRoughness
 	}
 
 	// TODO: add support for calculating the ideal resolution to use for baking based on connected (texture) nodes
-	FIntPoint TextureSize = Builder.GetDefaultMaterialBakeSize();
+	FIntPoint TextureSize = Builder.GetBakeSizeForMaterialProperty(Material, MetallicProperty);
 
 	// TODO: should this be the default wrap-mode?
 	EGLTFJsonTextureWrap TextureWrapS = EGLTFJsonTextureWrap::Repeat;
@@ -645,7 +645,7 @@ bool FGLTFMaterialTask::TryGetClearCoatRoughness(FGLTFJsonClearCoatExtension& Ou
 	}
 
 	// TODO: add support for calculating the ideal resolution to use for baking based on connected (texture) nodes
-	FIntPoint TextureSize = Builder.GetDefaultMaterialBakeSize();
+	FIntPoint TextureSize = Builder.GetBakeSizeForMaterialProperty(Material, IntensityProperty);
 
 	// TODO: should this be the default wrap-mode?
 	EGLTFJsonTextureWrap TextureWrapS = EGLTFJsonTextureWrap::Repeat;
@@ -1359,7 +1359,7 @@ FGLTFPropertyBakeOutput FGLTFMaterialTask::BakeMaterialProperty(const FMaterialP
 		OutTexCoord = MeshData->TexCoord;
 	}
 
-	const FIntPoint DefaultTextureSize = Builder.GetDefaultMaterialBakeSize();
+	const FIntPoint DefaultTextureSize = Builder.GetBakeSizeForMaterialProperty(Material, Property);
 	const FIntPoint TextureSize = PreferredTextureSize != FIntPoint::ZeroValue ? PreferredTextureSize : DefaultTextureSize;
 
 	// TODO: add support for calculating the ideal resolution to use for baking based on connected (texture) nodes
