[GLTFExporter] feat: implement support for overriding default texture filter and wrapping in bake settings

commit 1778fbdc96d5dd72acd8b55730c6d43703e08fd7
Date: Mon, 16 Nov 2020 17:01:18 +0100

    feat: implement support for overriding default texture filter and wrapping in bake settings
    (cherry picked from commit a447fdf53cccfe1f57e0a01cb81370b02342a346​)

#skipci

diff --git a/Engine/Plugins/Enterprise/GLTFExporter/MigratedFromGit/Source/GLTFExporter/GLTFExporter.Build.cs b/Engine/Plugins/Enterprise/GLTFExporter/MigratedFromGit/Source/GLTFExporter/GLTFExporter.Build.cs
--- a/Engine/Plugins/Enterprise/GLTFExporter/MigratedFromGit/Source/GLTFExporter/GLTFExporter.Build.cs
+++ b/Engine/Plugins/Enterprise/GLTFExporter/MigratedFromGit/Source/GLTFExporter/GLTFExporter.Build.cs
@@ -34,11 +34,17 @@ namespace UnrealBuildTool.Rules
 					"VariantManagerContent",
 					"MeshMergeUtilities",
 					"zlib",
-					"GLTFExporterRuntime",
 					"GLTFMaterialBaking",
 					"GLTFMaterialAnalyzer"
 				}
 				);
+
+			PublicDependencyModuleNames.AddRange(
+				new string[]
+				{
+					"GLTFExporterRuntime"
+				}
+				);
 		}
 	}
 }
diff --git a/Engine/Plugins/Enterprise/GLTFExporter/MigratedFromGit/Source/GLTFExporter/Private/Builders/GLTFBuilder.cpp b/Engine/Plugins/Enterprise/GLTFExporter/MigratedFromGit/Source/GLTFExporter/Private/Builders/GLTFBuilder.cpp
--- a/Engine/Plugins/Enterprise/GLTFExporter/MigratedFromGit/Source/GLTFExporter/Private/Builders/GLTFBuilder.cpp
+++ b/Engine/Plugins/Enterprise/GLTFExporter/MigratedFromGit/Source/GLTFExporter/Private/Builders/GLTFBuilder.cpp
@@ -14,18 +14,52 @@ FGLTFBuilder::FGLTFBuilder(const FString& FilePath, const UGLTFExportOptions* Ex
 
 FIntPoint FGLTFBuilder::GetBakeSizeForMaterialProperty(const UMaterialInterface* Material, const FMaterialPropertyEx& Property) const
 {
-	int32 Size = 1 << (static_cast<uint8>(ExportOptions->DefaultMaterialBakeSize) - static_cast<uint8>(EGLTFMaterialBakeSizePOT::POT_1));
+	const EGLTFMaterialPropertyGroup PropertyGroup = GetPropertyGroup(Property);
+	EGLTFMaterialBakeSizePOT DefaultValue = ExportOptions->DefaultMaterialBakeSize;
 
-	if (const UGLTFMaterialUserData* UserData = UGLTFMaterialUserData::GetUserData(Material))
+	if (const FGLTFOverrideMaterialBakeSettings* BakeSettings = ExportOptions->OverrideBakeSettings.Find(PropertyGroup))
 	{
-		EGLTFOverrideMaterialBakeSizePOT OverridePOT = UserData->GetBakeSizeForProperty(Property.Type);
-		if (OverridePOT != EGLTFOverrideMaterialBakeSizePOT::NoOverride)
+		if (BakeSettings->bOverrideSize)
 		{
-			Size = 1 << (static_cast<uint8>(OverridePOT) - static_cast<uint8>(EGLTFOverrideMaterialBakeSizePOT::POT_1));
+			DefaultValue = BakeSettings->Size;
 		}
 	}
 
-	return { Size, Size };
+	const EGLTFMaterialBakeSizePOT Size = UGLTFMaterialExportOptions::GetBakeSizeForPropertyGroup(Material, PropertyGroup, DefaultValue);
+	const int32 PixelSize = 1 << static_cast<uint8>(Size);
+	return { PixelSize, PixelSize };
+}
+
+TextureFilter FGLTFBuilder::GetBakeFilterForMaterialProperty(const UMaterialInterface* Material, const FMaterialPropertyEx& Property) const
+{
+	const EGLTFMaterialPropertyGroup PropertyGroup = GetPropertyGroup(Property);
+	TextureFilter DefaultValue = ExportOptions->DefaultMaterialBakeFilter;
+
+	if (const FGLTFOverrideMaterialBakeSettings* BakeSettings = ExportOptions->OverrideBakeSettings.Find(PropertyGroup))
+	{
+		if (BakeSettings->bOverrideFilter)
+		{
+			DefaultValue = BakeSettings->Filter;
+		}
+	}
+
+	return UGLTFMaterialExportOptions::GetBakeFilterForPropertyGroup(Material, PropertyGroup, DefaultValue);
+}
+
+TextureAddress FGLTFBuilder::GetBakeTilingForMaterialProperty(const UMaterialInterface* Material, const FMaterialPropertyEx& Property) const
+{
+	const EGLTFMaterialPropertyGroup PropertyGroup = GetPropertyGroup(Property);
+	TextureAddress DefaultValue = ExportOptions->DefaultMaterialBakeTiling;
+
+	if (const FGLTFOverrideMaterialBakeSettings* BakeSettings = ExportOptions->OverrideBakeSettings.Find(PropertyGroup))
+	{
+		if (BakeSettings->bOverrideTiling)
+		{
+			DefaultValue = BakeSettings->Tiling;
+		}
+	}
+
+	return UGLTFMaterialExportOptions::GetBakeTilingForPropertyGroup(Material, PropertyGroup, DefaultValue);
 }
 
 EGLTFJsonHDREncoding FGLTFBuilder::GetTextureHDREncoding() const
@@ -55,3 +89,33 @@ bool FGLTFBuilder::ShouldExportLight(EComponentMobility::Type LightMobility) con
 			return false;
 	}
 }
+
+EGLTFMaterialPropertyGroup FGLTFBuilder::GetPropertyGroup(const FMaterialPropertyEx& Property)
+{
+	switch (Property.Type)
+	{
+		case MP_BaseColor:
+        case MP_Opacity:
+        case MP_OpacityMask:
+            return EGLTFMaterialPropertyGroup::BaseColorOpacity;
+		case MP_Metallic:
+        case MP_Roughness:
+            return EGLTFMaterialPropertyGroup::MetallicRoughness;
+		case MP_EmissiveColor:
+			return EGLTFMaterialPropertyGroup::EmissiveColor;
+		case MP_Normal:
+			return EGLTFMaterialPropertyGroup::Normal;
+		case MP_AmbientOcclusion:
+			return EGLTFMaterialPropertyGroup::AmbientOcclusion;
+		case MP_CustomData0:
+        case MP_CustomData1:
+            return EGLTFMaterialPropertyGroup::ClearCoatRoughness;
+		case MP_CustomOutput:
+			if (Property.CustomOutput == TEXT("ClearCoatBottomNormal"))
+			{
+				return EGLTFMaterialPropertyGroup::ClearCoatBottomNormal;
+			}
+		default:
+			return EGLTFMaterialPropertyGroup::None;
+	}
+}
diff --git a/Engine/Plugins/Enterprise/GLTFExporter/MigratedFromGit/Source/GLTFExporter/Private/Builders/GLTFBuilder.h b/Engine/Plugins/Enterprise/GLTFExporter/MigratedFromGit/Source/GLTFExporter/Private/Builders/GLTFBuilder.h
--- a/Engine/Plugins/Enterprise/GLTFExporter/MigratedFromGit/Source/GLTFExporter/Private/Builders/GLTFBuilder.h
+++ b/Engine/Plugins/Enterprise/GLTFExporter/MigratedFromGit/Source/GLTFExporter/Private/Builders/GLTFBuilder.h
@@ -21,8 +21,14 @@ public:
 	const UGLTFExportOptions* const ExportOptions;
 
 	FIntPoint GetBakeSizeForMaterialProperty(const UMaterialInterface* Material, const FMaterialPropertyEx& Property) const;
+	TextureFilter GetBakeFilterForMaterialProperty(const UMaterialInterface* Material, const FMaterialPropertyEx& Property) const;
+	TextureAddress GetBakeTilingForMaterialProperty(const UMaterialInterface* Material, const FMaterialPropertyEx& Property) const;
 
 	EGLTFJsonHDREncoding GetTextureHDREncoding() const;
 
 	bool ShouldExportLight(EComponentMobility::Type LightMobility) const;
+
+private:
+
+	static EGLTFMaterialPropertyGroup GetPropertyGroup(const FMaterialPropertyEx& Property);
 };
diff --git a/Engine/Plugins/Enterprise/GLTFExporter/MigratedFromGit/Source/GLTFExporter/Private/GLTFExportOptions.cpp b/Engine/Plugins/Enterprise/GLTFExporter/MigratedFromGit/Source/GLTFExporter/Private/GLTFExportOptions.cpp
--- a/Engine/Plugins/Enterprise/GLTFExporter/MigratedFromGit/Source/GLTFExporter/Private/GLTFExportOptions.cpp
+++ b/Engine/Plugins/Enterprise/GLTFExporter/MigratedFromGit/Source/GLTFExporter/Private/GLTFExportOptions.cpp
@@ -3,8 +3,6 @@
 #include "GLTFExportOptions.h"
 #include "UObject/UObjectGlobals.h"
 #include "UObject/Object.h"
-#include "UObject/Class.h"
-#include "UObject/UnrealType.h"
 
 UGLTFExportOptions::UGLTFExportOptions(const FObjectInitializer& ObjectInitializer)
 	: Super(ObjectInitializer)
@@ -19,7 +17,8 @@ void UGLTFExportOptions::ResetToDefault()
 	bExportExtraBlendModes = true;
 	BakeMaterialInputs = EGLTFMaterialBakeMode::Simple;
 	DefaultMaterialBakeSize = EGLTFMaterialBakeSizePOT::POT_1024;
-	DefaultLevelOfDetail = 0;
+	DefaultMaterialBakeFilter = TF_Trilinear;
+	DefaultMaterialBakeTiling = TA_Wrap;
 	bExportVertexColors = false;
 	bExportVertexSkinWeights = true;
 	bExportMeshQuantization = true;
diff --git a/Engine/Plugins/Enterprise/GLTFExporter/MigratedFromGit/Source/GLTFExporter/Private/Tasks/GLTFMaterialTasks.cpp b/Engine/Plugins/Enterprise/GLTFExporter/MigratedFromGit/Source/GLTFExporter/Private/Tasks/GLTFMaterialTasks.cpp
--- a/Engine/Plugins/Enterprise/GLTFExporter/MigratedFromGit/Source/GLTFExporter/Private/Tasks/GLTFMaterialTasks.cpp
+++ b/Engine/Plugins/Enterprise/GLTFExporter/MigratedFromGit/Source/GLTFExporter/Private/Tasks/GLTFMaterialTasks.cpp
@@ -309,63 +309,15 @@ bool FGLTFMaterialTask::TryGetBaseColorAndOpacity(FGLTFJsonPBRMetallicRoughness&
 		return false;
 	}
 
-	// TODO: add support for calculating the ideal resolution to use for baking based on connected (texture) nodes
 	FIntPoint TextureSize = Builder.GetBakeSizeForMaterialProperty(Material, BaseColorProperty);
 
-	// TODO: should this be the default wrap-mode?
-	EGLTFJsonTextureWrap TextureWrapS = EGLTFJsonTextureWrap::Repeat;
-	EGLTFJsonTextureWrap TextureWrapT = EGLTFJsonTextureWrap::Repeat;
+	const TextureAddress TextureAddress = Builder.GetBakeTilingForMaterialProperty(Material, BaseColorProperty);
+	const EGLTFJsonTextureWrap TextureWrapS = FGLTFConverterUtility::ConvertWrap(TextureAddress);
+	const EGLTFJsonTextureWrap TextureWrapT = FGLTFConverterUtility::ConvertWrap(TextureAddress);
 
-	// TODO: should this be the default filter?
-	EGLTFJsonTextureFilter TextureMinFilter = EGLTFJsonTextureFilter::LinearMipmapLinear;
-	EGLTFJsonTextureFilter TextureMagFilter = EGLTFJsonTextureFilter::Linear;
-
-	if (bHasBaseColorSourceTexture && bHasOpacitySourceTexture)
-	{
-		const bool bAreTexturesCompatible = BaseColorTexCoord == OpacityTexCoord &&
-			BaseColorTexture->AddressX == OpacityTexture->AddressX &&
-			BaseColorTexture->AddressY == OpacityTexture->AddressY;
-
-		if (!bAreTexturesCompatible)
-		{
-			// TODO: handle differences in wrapping or uv-coords
-			Builder.AddWarningMessage(FString::Printf(
-				TEXT("BaseColor- and Opacity-textures for material %s were not able to be combined and will be skipped"),
-				*Material->GetName()));
-
-			return false;
-		}
-
-		TextureSize =
-		{
-			FMath::Max(BaseColorTexture->GetSizeX(), OpacityTexture->GetSizeX()),
-			FMath::Max(BaseColorTexture->GetSizeY(), OpacityTexture->GetSizeY())
-		};
-
-		TextureWrapS = FGLTFConverterUtility::ConvertWrap(BaseColorTexture->AddressX);
-		TextureWrapT = FGLTFConverterUtility::ConvertWrap(BaseColorTexture->AddressY);
-
-		// TODO: compare min- and mag-filter for BaseColorTexture and OpacityTexture. If they differ,
-		// we should choose one or the other and inform the user about the choice made by logging to the console.
-		TextureMinFilter = FGLTFConverterUtility::ConvertMinFilter(BaseColorTexture->Filter, BaseColorTexture->LODGroup);
-		TextureMagFilter = FGLTFConverterUtility::ConvertMagFilter(BaseColorTexture->Filter, BaseColorTexture->LODGroup);
-	}
-	else if (bHasBaseColorSourceTexture)
-	{
-		TextureSize = { BaseColorTexture->GetSizeX(), BaseColorTexture->GetSizeY() };
-		TextureWrapS = FGLTFConverterUtility::ConvertWrap(BaseColorTexture->AddressX);
-		TextureWrapT = FGLTFConverterUtility::ConvertWrap(BaseColorTexture->AddressY);
-		TextureMinFilter = FGLTFConverterUtility::ConvertMinFilter(BaseColorTexture->Filter, BaseColorTexture->LODGroup);
-		TextureMagFilter = FGLTFConverterUtility::ConvertMagFilter(BaseColorTexture->Filter, BaseColorTexture->LODGroup);
-	}
-	else if (bHasOpacitySourceTexture)
-	{
-		TextureSize = { OpacityTexture->GetSizeX(), OpacityTexture->GetSizeY() };
-		TextureWrapS = FGLTFConverterUtility::ConvertWrap(OpacityTexture->AddressX);
-		TextureWrapT = FGLTFConverterUtility::ConvertWrap(OpacityTexture->AddressY);
-		TextureMinFilter = FGLTFConverterUtility::ConvertMinFilter(OpacityTexture->Filter, OpacityTexture->LODGroup);
-		TextureMagFilter = FGLTFConverterUtility::ConvertMagFilter(OpacityTexture->Filter, OpacityTexture->LODGroup);
-	}
+	const TextureFilter TextureFilter = Builder.GetBakeFilterForMaterialProperty(Material, BaseColorProperty);
+	const EGLTFJsonTextureFilter TextureMinFilter = FGLTFConverterUtility::ConvertMinFilter(TextureFilter);
+	const EGLTFJsonTextureFilter TextureMagFilter = FGLTFConverterUtility::ConvertMagFilter(TextureFilter);
 
 	const FGLTFPropertyBakeOutput BaseColorBakeOutput = BakeMaterialProperty(BaseColorProperty, BaseColorTexCoord, TextureSize);
 	const FGLTFPropertyBakeOutput OpacityBakeOutput = BakeMaterialProperty(OpacityProperty, OpacityTexCoord, TextureSize, true);
@@ -480,60 +432,13 @@ bool FGLTFMaterialTask::TryGetMetallicAndRoughness(FGLTFJsonPBRMetallicRoughness
 	// TODO: add support for calculating the ideal resolution to use for baking based on connected (texture) nodes
 	FIntPoint TextureSize = Builder.GetBakeSizeForMaterialProperty(Material, MetallicProperty);
 
-	// TODO: should this be the default wrap-mode?
-	EGLTFJsonTextureWrap TextureWrapS = EGLTFJsonTextureWrap::Repeat;
-	EGLTFJsonTextureWrap TextureWrapT = EGLTFJsonTextureWrap::Repeat;
-
-	// TODO: should this be the default filter?
-	EGLTFJsonTextureFilter TextureMinFilter = EGLTFJsonTextureFilter::LinearMipmapLinear;
-	EGLTFJsonTextureFilter TextureMagFilter = EGLTFJsonTextureFilter::Linear;
-
-	if (bHasMetallicSourceTexture && bHasRoughnessSourceTexture)
-	{
-		const bool bAreTexturesCompatible = MetallicTexCoord == RoughnessTexCoord &&
-			MetallicTexture->AddressX == RoughnessTexture->AddressX &&
-			MetallicTexture->AddressY == RoughnessTexture->AddressY;
-
-		if (!bAreTexturesCompatible)
-		{
-			// TODO: handle differences in wrapping or uv-coords
-			Builder.AddWarningMessage(FString::Printf(
-				TEXT("Metallic- and Roughness-textures for material %s were not able to be combined and will be skipped"),
-				*Material->GetName()));
-
-			return false;
-		}
-
-		TextureSize =
-		{
-			FMath::Max(MetallicTexture->GetSizeX(), RoughnessTexture->GetSizeX()),
-			FMath::Max(MetallicTexture->GetSizeY(), RoughnessTexture->GetSizeY())
-		};
+	const TextureAddress TextureAddress = Builder.GetBakeTilingForMaterialProperty(Material, MetallicProperty);
+	const EGLTFJsonTextureWrap TextureWrapS = FGLTFConverterUtility::ConvertWrap(TextureAddress);
+	const EGLTFJsonTextureWrap TextureWrapT = FGLTFConverterUtility::ConvertWrap(TextureAddress);
 
-		TextureWrapS = FGLTFConverterUtility::ConvertWrap(MetallicTexture->AddressX);
-		TextureWrapT = FGLTFConverterUtility::ConvertWrap(MetallicTexture->AddressY);
-
-		// TODO: compare min- and mag-filter for BaseColorTexture and OpacityTexture. If they differ,
-		// we should choose one or the other and inform the user about the choice made by logging to the console.
-		TextureMinFilter = FGLTFConverterUtility::ConvertMinFilter(MetallicTexture->Filter, MetallicTexture->LODGroup);
-		TextureMagFilter = FGLTFConverterUtility::ConvertMagFilter(MetallicTexture->Filter, MetallicTexture->LODGroup);
-	}
-	else if (bHasMetallicSourceTexture)
-	{
-		TextureSize = { MetallicTexture->GetSizeX(), MetallicTexture->GetSizeY() };
-		TextureWrapS = FGLTFConverterUtility::ConvertWrap(MetallicTexture->AddressX);
-		TextureWrapT = FGLTFConverterUtility::ConvertWrap(MetallicTexture->AddressY);
-		TextureMinFilter = FGLTFConverterUtility::ConvertMinFilter(MetallicTexture->Filter, MetallicTexture->LODGroup);
-		TextureMagFilter = FGLTFConverterUtility::ConvertMagFilter(MetallicTexture->Filter, MetallicTexture->LODGroup);
-	}
-	else if (bHasRoughnessSourceTexture)
-	{
-		TextureSize = { RoughnessTexture->GetSizeX(), RoughnessTexture->GetSizeY() };
-		TextureWrapS = FGLTFConverterUtility::ConvertWrap(RoughnessTexture->AddressX);
-		TextureWrapT = FGLTFConverterUtility::ConvertWrap(RoughnessTexture->AddressY);
-		TextureMinFilter = FGLTFConverterUtility::ConvertMinFilter(RoughnessTexture->Filter, RoughnessTexture->LODGroup);
-		TextureMagFilter = FGLTFConverterUtility::ConvertMagFilter(RoughnessTexture->Filter, RoughnessTexture->LODGroup);
-	}
+	const TextureFilter TextureFilter = Builder.GetBakeFilterForMaterialProperty(Material, MetallicProperty);
+	const EGLTFJsonTextureFilter TextureMinFilter = FGLTFConverterUtility::ConvertMinFilter(TextureFilter);
+	const EGLTFJsonTextureFilter TextureMagFilter = FGLTFConverterUtility::ConvertMagFilter(TextureFilter);
 
 	const FGLTFPropertyBakeOutput MetallicBakeOutput = BakeMaterialProperty(MetallicProperty, MetallicTexCoord, TextureSize);
 	const FGLTFPropertyBakeOutput RoughnessBakeOutput = BakeMaterialProperty(RoughnessProperty, RoughnessTexCoord, TextureSize);
@@ -644,63 +549,15 @@ bool FGLTFMaterialTask::TryGetClearCoatRoughness(FGLTFJsonClearCoatExtension& Ou
 		return false;
 	}
 
-	// TODO: add support for calculating the ideal resolution to use for baking based on connected (texture) nodes
 	FIntPoint TextureSize = Builder.GetBakeSizeForMaterialProperty(Material, IntensityProperty);
 
-	// TODO: should this be the default wrap-mode?
-	EGLTFJsonTextureWrap TextureWrapS = EGLTFJsonTextureWrap::Repeat;
-	EGLTFJsonTextureWrap TextureWrapT = EGLTFJsonTextureWrap::Repeat;
-
-	// TODO: should this be the default filter?
-	EGLTFJsonTextureFilter TextureMinFilter = EGLTFJsonTextureFilter::LinearMipmapLinear;
-	EGLTFJsonTextureFilter TextureMagFilter = EGLTFJsonTextureFilter::Linear;
-
-	if (bHasIntensitySourceTexture && bHasRoughnessSourceTexture)
-	{
-		const bool bAreTexturesCompatible = IntensityTexCoord == RoughnessTexCoord &&
-			IntensityTexture->AddressX == RoughnessTexture->AddressX &&
-			IntensityTexture->AddressY == RoughnessTexture->AddressY;
+	const TextureAddress TextureAddress = Builder.GetBakeTilingForMaterialProperty(Material, IntensityProperty);
+	const EGLTFJsonTextureWrap TextureWrapS = FGLTFConverterUtility::ConvertWrap(TextureAddress);
+	const EGLTFJsonTextureWrap TextureWrapT = FGLTFConverterUtility::ConvertWrap(TextureAddress);
 
-		if (!bAreTexturesCompatible)
-		{
-			// TODO: handle differences in wrapping or uv-coords
-			Builder.AddWarningMessage(FString::Printf(
-				TEXT("Intensity- and Roughness-textures for material %s were not able to be combined and will be skipped"),
-				*Material->GetName()));
-
-			return false;
-		}
-
-		TextureSize =
-		{
-			FMath::Max(IntensityTexture->GetSizeX(), RoughnessTexture->GetSizeX()),
-			FMath::Max(IntensityTexture->GetSizeY(), RoughnessTexture->GetSizeY())
-		};
-
-		TextureWrapS = FGLTFConverterUtility::ConvertWrap(IntensityTexture->AddressX);
-		TextureWrapT = FGLTFConverterUtility::ConvertWrap(IntensityTexture->AddressY);
-
-		// TODO: compare min- and mag-filter for BaseColorTexture and OpacityTexture. If they differ,
-		// we should choose one or the other and inform the user about the choice made by logging to the console.
-		TextureMinFilter = FGLTFConverterUtility::ConvertMinFilter(IntensityTexture->Filter, IntensityTexture->LODGroup);
-		TextureMagFilter = FGLTFConverterUtility::ConvertMagFilter(IntensityTexture->Filter, IntensityTexture->LODGroup);
-	}
-	else if (bHasIntensitySourceTexture)
-	{
-		TextureSize = { IntensityTexture->GetSizeX(), IntensityTexture->GetSizeY() };
-		TextureWrapS = FGLTFConverterUtility::ConvertWrap(IntensityTexture->AddressX);
-		TextureWrapT = FGLTFConverterUtility::ConvertWrap(IntensityTexture->AddressY);
-		TextureMinFilter = FGLTFConverterUtility::ConvertMinFilter(IntensityTexture->Filter, IntensityTexture->LODGroup);
-		TextureMagFilter = FGLTFConverterUtility::ConvertMagFilter(IntensityTexture->Filter, IntensityTexture->LODGroup);
-	}
-	else if (bHasRoughnessSourceTexture)
-	{
-		TextureSize = { RoughnessTexture->GetSizeX(), RoughnessTexture->GetSizeY() };
-		TextureWrapS = FGLTFConverterUtility::ConvertWrap(RoughnessTexture->AddressX);
-		TextureWrapT = FGLTFConverterUtility::ConvertWrap(RoughnessTexture->AddressY);
-		TextureMinFilter = FGLTFConverterUtility::ConvertMinFilter(RoughnessTexture->Filter, RoughnessTexture->LODGroup);
-		TextureMagFilter = FGLTFConverterUtility::ConvertMagFilter(RoughnessTexture->Filter, RoughnessTexture->LODGroup);
-	}
+	const TextureFilter TextureFilter = Builder.GetBakeFilterForMaterialProperty(Material, IntensityProperty);
+	const EGLTFJsonTextureFilter TextureMinFilter = FGLTFConverterUtility::ConvertMinFilter(TextureFilter);
+	const EGLTFJsonTextureFilter TextureMagFilter = FGLTFConverterUtility::ConvertMagFilter(TextureFilter);
 
 	const FGLTFPropertyBakeOutput IntensityBakeOutput = BakeMaterialProperty(IntensityProperty, IntensityTexCoord, TextureSize);
 	const FGLTFPropertyBakeOutput RoughnessBakeOutput = BakeMaterialProperty(RoughnessProperty, RoughnessTexCoord, TextureSize);
@@ -1332,7 +1189,13 @@ bool FGLTFMaterialTask::TryGetBakedMaterialProperty(FGLTFJsonTextureInfo& OutTex
 	return true;
 }
 
-FGLTFPropertyBakeOutput FGLTFMaterialTask::BakeMaterialProperty(const FMaterialPropertyEx& Property, int32& OutTexCoord, FIntPoint PreferredTextureSize, bool bCopyAlphaFromRedChannel) const
+FGLTFPropertyBakeOutput FGLTFMaterialTask::BakeMaterialProperty(const FMaterialPropertyEx& Property, int32& OutTexCoord, bool bCopyAlphaFromRedChannel) const
+{
+	const FIntPoint TextureSize = Builder.GetBakeSizeForMaterialProperty(Material, Property);
+	return BakeMaterialProperty(Property, OutTexCoord, TextureSize, bCopyAlphaFromRedChannel);
+}
+
+FGLTFPropertyBakeOutput FGLTFMaterialTask::BakeMaterialProperty(const FMaterialPropertyEx& Property, int32& OutTexCoord, const FIntPoint& TextureSize, bool bCopyAlphaFromRedChannel) const
 {
 	if (MeshData == nullptr)
 	{
@@ -1359,9 +1222,6 @@ FGLTFPropertyBakeOutput FGLTFMaterialTask::BakeMaterialProperty(const FMaterialP
 		OutTexCoord = MeshData->TexCoord;
 	}
 
-	const FIntPoint DefaultTextureSize = Builder.GetBakeSizeForMaterialProperty(Material, Property);
-	const FIntPoint TextureSize = PreferredTextureSize != FIntPoint::ZeroValue ? PreferredTextureSize : DefaultTextureSize;
-
 	// TODO: add support for calculating the ideal resolution to use for baking based on connected (texture) nodes
 
 	const FGLTFPropertyBakeOutput PropertyBakeOutput = FGLTFMaterialUtility::BakeMaterialProperty(
diff --git a/Engine/Plugins/Enterprise/GLTFExporter/MigratedFromGit/Source/GLTFExporter/Private/Tasks/GLTFMaterialTasks.h b/Engine/Plugins/Enterprise/GLTFExporter/MigratedFromGit/Source/GLTFExporter/Private/Tasks/GLTFMaterialTasks.h
--- a/Engine/Plugins/Enterprise/GLTFExporter/MigratedFromGit/Source/GLTFExporter/Private/Tasks/GLTFMaterialTasks.h
+++ b/Engine/Plugins/Enterprise/GLTFExporter/MigratedFromGit/Source/GLTFExporter/Private/Tasks/GLTFMaterialTasks.h
@@ -61,7 +61,8 @@ private:
 	bool TryGetBakedMaterialProperty(FGLTFJsonTextureInfo& OutTexInfo, float& OutConstant, const FMaterialPropertyEx& Property, const FString& PropertyName) const;
 	bool TryGetBakedMaterialProperty(FGLTFJsonTextureInfo& OutTexInfo, const FMaterialPropertyEx& Property, const FString& PropertyName) const;
 
-	FGLTFPropertyBakeOutput BakeMaterialProperty(const FMaterialPropertyEx& Property, int32& OutTexCoord, FIntPoint PreferredTextureSize = FIntPoint::ZeroValue, bool bCopyAlphaFromRedChannel = false) const;
+	FGLTFPropertyBakeOutput BakeMaterialProperty(const FMaterialPropertyEx& Property, int32& OutTexCoord, bool bCopyAlphaFromRedChannel = false) const;
+	FGLTFPropertyBakeOutput BakeMaterialProperty(const FMaterialPropertyEx& Property, int32& OutTexCoord, const FIntPoint& TextureSize, bool bCopyAlphaFromRedChannel = false) const;
 
 	bool StoreBakedPropertyTexture(FGLTFJsonTextureInfo& OutTexInfo, const FGLTFPropertyBakeOutput& PropertyBakeOutput, const FString& PropertyName) const;
 
diff --git a/Engine/Plugins/Enterprise/GLTFExporter/MigratedFromGit/Source/GLTFExporter/Public/GLTFExportOptions.h b/Engine/Plugins/Enterprise/GLTFExporter/MigratedFromGit/Source/GLTFExporter/Public/GLTFExportOptions.h
--- a/Engine/Plugins/Enterprise/GLTFExporter/MigratedFromGit/Source/GLTFExporter/Public/GLTFExportOptions.h
+++ b/Engine/Plugins/Enterprise/GLTFExporter/MigratedFromGit/Source/GLTFExporter/Public/GLTFExportOptions.h
@@ -3,35 +3,9 @@
 #pragma once
 
 #include "CoreMinimal.h"
+#include "UserData/GLTFMaterialUserData.h"
 #include "GLTFExportOptions.generated.h"
 
-UENUM(BlueprintType)
-enum class EGLTFMaterialBakeMode : uint8
-{
-	None,
-    Simple,
-    UseMeshData,
-};
-
-UENUM(BlueprintType)
-enum class EGLTFMaterialBakeSizePOT : uint8
-{
-	POT_1 UMETA(DisplayName = "1 x 1"),
-	POT_2 UMETA(DisplayName = "2 x 2"),
-	POT_4 UMETA(DisplayName = "4 x 4"),
-	POT_8 UMETA(DisplayName = "8 x 8"),
-	POT_16 UMETA(DisplayName = "16 x 16"),
-	POT_32 UMETA(DisplayName = "32 x 32"),
-	POT_64 UMETA(DisplayName = "64 x 64"),
-	POT_128 UMETA(DisplayName = "128 x 128"),
-	POT_256 UMETA(DisplayName = "256 x 256"),
-	POT_512 UMETA(DisplayName = "512 x 512"),
-	POT_1024 UMETA(DisplayName = "1024 x 1024"),
-	POT_2048 UMETA(DisplayName = "2048 x 2048"),
-	POT_4096 UMETA(DisplayName = "4096 x 4096"),
-	POT_8192 UMETA(DisplayName = "8192 x 8192")
-};
-
 UENUM(BlueprintType)
 enum class EGLTFTextureCompression : uint8
 {
@@ -90,6 +64,15 @@ public:
 	UPROPERTY(EditAnywhere, BlueprintReadWrite, Config, Category = Material, Meta = (EditCondition = "BakeMaterialInputs != EGLTFMaterialBakeMode::None"))
 	EGLTFMaterialBakeSizePOT DefaultMaterialBakeSize;
 
+	UPROPERTY(EditAnywhere, BlueprintReadWrite, Config, Category = Material, Meta = (EditCondition = "BakeMaterialInputs != EGLTFMaterialBakeMode::None", ValidEnumValues="TF_Nearest, TF_Bilinear, TF_Trilinear"))
+	TEnumAsByte<TextureFilter> DefaultMaterialBakeFilter;
+
+	UPROPERTY(EditAnywhere, BlueprintReadWrite, Config, Category = Material, Meta = (EditCondition = "BakeMaterialInputs != EGLTFMaterialBakeMode::None"))
+	TEnumAsByte<TextureAddress> DefaultMaterialBakeTiling;
+
+	UPROPERTY(EditAnywhere, BlueprintReadWrite, Config, Category = Material)
+	TMap<EGLTFMaterialPropertyGroup, FGLTFOverrideMaterialBakeSettings> OverrideBakeSettings;
+
 	UPROPERTY(EditAnywhere, BlueprintReadWrite, Config, Category = Mesh, Meta = (ClampMin = "0"))
 	int32 DefaultLevelOfDetail;
 
diff --git a/Engine/Plugins/Enterprise/GLTFExporter/MigratedFromGit/Source/GLTFExporterRuntime/Private/UserData/GLTFMaterialUserData.cpp b/Engine/Plugins/Enterprise/GLTFExporter/MigratedFromGit/Source/GLTFExporterRuntime/Private/UserData/GLTFMaterialUserData.cpp
--- a/Engine/Plugins/Enterprise/GLTFExporter/MigratedFromGit/Source/GLTFExporterRuntime/Private/UserData/GLTFMaterialUserData.cpp
+++ b/Engine/Plugins/Enterprise/GLTFExporter/MigratedFromGit/Source/GLTFExporterRuntime/Private/UserData/GLTFMaterialUserData.cpp
@@ -1,65 +1,73 @@
 // Copyright Epic Games, Inc. All Rights Reserved.
 
 #include "UserData/GLTFMaterialUserData.h"
-#include "Materials/MaterialInterface.h"
+#include "Materials/MaterialInstance.h"
+#include "Engine/Texture.h"
 
-UGLTFMaterialUserData::UGLTFMaterialUserData()
-	: DefaultBakeSize(EGLTFOverrideMaterialBakeSizePOT::POT_1024)
+FGLTFOverrideMaterialBakeSettings::FGLTFOverrideMaterialBakeSettings()
+	: bOverrideSize(false)
+	, Size(EGLTFMaterialBakeSizePOT::POT_1024)
+	, bOverrideFilter(false)
+	, Filter(TF_Trilinear)
+	, bOverrideTiling(false)
+	, Tiling(TA_Wrap)
 {
 }
 
-EGLTFOverrideMaterialBakeSizePOT UGLTFMaterialUserData::GetBakeSizeForProperty(EMaterialProperty Property) const
+EGLTFMaterialBakeSizePOT UGLTFMaterialExportOptions::GetBakeSizeForPropertyGroup(const UMaterialInterface* Material, EGLTFMaterialPropertyGroup PropertyGroup, EGLTFMaterialBakeSizePOT DefaultValue)
 {
-	const EGLTFOverrideMaterialPropertyGroup PropertyGroup = GetPropertyGroup(Property);
-	if (const EGLTFOverrideMaterialBakeSizePOT* PropertyBakeSize = BakeSizePerProperty.Find(PropertyGroup))
+	if (const FGLTFOverrideMaterialBakeSettings* BakeSettings = GetBakeSettingsByPredicate(Material, PropertyGroup, [](const FGLTFOverrideMaterialBakeSettings& BakeSettings) { return BakeSettings.bOverrideSize; }))
 	{
-		return *PropertyBakeSize;
+		return BakeSettings->Size;
 	}
 
-	return DefaultBakeSize;
+	return DefaultValue;
 }
 
-enum class EGLTFOverrideMater : uint8
+TextureFilter UGLTFMaterialExportOptions::GetBakeFilterForPropertyGroup(const UMaterialInterface* Material, EGLTFMaterialPropertyGroup PropertyGroup, TextureFilter DefaultValue)
 {
-	None UMETA(DisplayName = "None"),
+	if (const FGLTFOverrideMaterialBakeSettings* BakeSettings = GetBakeSettingsByPredicate(Material, PropertyGroup, [](const FGLTFOverrideMaterialBakeSettings& BakeSettings) { return BakeSettings.bOverrideFilter; }))
+	{
+		return BakeSettings->Filter;
+	}
 
-    BaseColorOpacity UMETA(DisplayName = "Base Color + Opacity (Mask)"),
-    MetallicRoughness UMETA(DisplayName = "Metallic + Roughness"),
-    EmissiveColor UMETA(DisplayName = "Emissive Color"),
-    Normal UMETA(DisplayName = "Normal"),
-    AmbientOcclusion UMETA(DisplayName = "Ambient Occlusion"),
-    ClearCoatRoughness UMETA(DisplayName = "Clear Coat + Clear Coat Roughness"),
-    ClearCoatBottomNormal UMETA(DisplayName = "Clear Coat Bottom Normal"),
-};
+	return DefaultValue;
+}
 
-EGLTFOverrideMaterialPropertyGroup UGLTFMaterialUserData::GetPropertyGroup(EMaterialProperty Property)
+TextureAddress UGLTFMaterialExportOptions::GetBakeTilingForPropertyGroup(const UMaterialInterface* Material, EGLTFMaterialPropertyGroup PropertyGroup, TextureAddress DefaultValue)
 {
-	switch (Property)
+	if (const FGLTFOverrideMaterialBakeSettings* BakeSettings = GetBakeSettingsByPredicate(Material, PropertyGroup, [](const FGLTFOverrideMaterialBakeSettings& BakeSettings) { return BakeSettings.bOverrideTiling; }))
 	{
-		case MP_BaseColor:
-		case MP_Opacity:
-		case MP_OpacityMask:
-			return EGLTFOverrideMaterialPropertyGroup::BaseColorOpacity;
-		case MP_Metallic:
-		case MP_Roughness:
-			return EGLTFOverrideMaterialPropertyGroup::MetallicRoughness;
-		case MP_EmissiveColor:
-			return EGLTFOverrideMaterialPropertyGroup::EmissiveColor;
-		case MP_Normal:
-			return EGLTFOverrideMaterialPropertyGroup::Normal;
-		case MP_AmbientOcclusion:
-			return EGLTFOverrideMaterialPropertyGroup::AmbientOcclusion;
-		case MP_CustomData0:
-		case MP_CustomData1:
-			return EGLTFOverrideMaterialPropertyGroup::ClearCoatRoughness;
-		case MP_CustomOutput: // TODO: fix assumption
-			return EGLTFOverrideMaterialPropertyGroup::ClearCoatBottomNormal;
-		default:
-			return EGLTFOverrideMaterialPropertyGroup::None;
+		return BakeSettings->Tiling;
 	}
+
+	return DefaultValue;
 }
 
-const UGLTFMaterialUserData* UGLTFMaterialUserData::GetUserData(const UMaterialInterface* Material)
+template <typename Predicate>
+const FGLTFOverrideMaterialBakeSettings* UGLTFMaterialExportOptions::GetBakeSettingsByPredicate(const UMaterialInterface* Material, EGLTFMaterialPropertyGroup PropertyGroup, Predicate Pred)
 {
-	return const_cast<UMaterialInterface*>(Material)->GetAssetUserData<UGLTFMaterialUserData>();
+	do
+	{
+		if (const UGLTFMaterialExportOptions* UserData = const_cast<UMaterialInterface*>(Material)->GetAssetUserData<UGLTFMaterialExportOptions>())
+		{
+			if (const FGLTFOverrideMaterialBakeSettings* BakeSettings = UserData->PerInput.Find(PropertyGroup))
+			{
+				if (Pred(*BakeSettings))
+				{
+					return BakeSettings;
+				}
+			}
+
+			if (Pred(UserData->Default))
+			{
+				return &UserData->Default;
+			}
+		}
+
+		const UMaterialInstance* MaterialInstance = Cast<UMaterialInstance>(Material);
+		Material = MaterialInstance != nullptr ? MaterialInstance->Parent : nullptr;
+	} while (Material != nullptr);
+
+	return nullptr;
 }
diff --git a/Engine/Plugins/Enterprise/GLTFExporter/MigratedFromGit/Source/GLTFExporterRuntime/Public/UserData/GLTFMaterialUserData.h b/Engine/Plugins/Enterprise/GLTFExporter/MigratedFromGit/Source/GLTFExporterRuntime/Public/UserData/GLTFMaterialUserData.h
--- a/Engine/Plugins/Enterprise/GLTFExporter/MigratedFromGit/Source/GLTFExporterRuntime/Public/UserData/GLTFMaterialUserData.h
+++ b/Engine/Plugins/Enterprise/GLTFExporter/MigratedFromGit/Source/GLTFExporterRuntime/Public/UserData/GLTFMaterialUserData.h
@@ -6,13 +6,20 @@
 #include "GLTFMaterialUserData.generated.h"
 
 class UMaterialInterface;
-enum EMaterialProperty;
+enum TextureFilter;
+enum TextureAddress;
 
 UENUM(BlueprintType)
-enum class EGLTFOverrideMaterialBakeSizePOT : uint8
+enum class EGLTFMaterialBakeMode : uint8
 {
-	NoOverride UMETA(DisplayName = "No Override"),
+    None,
+	Simple,
+	UseMeshData,
+};
 
+UENUM(BlueprintType)
+enum class EGLTFMaterialBakeSizePOT : uint8
+{
     POT_1 UMETA(DisplayName = "1 x 1"),
     POT_2 UMETA(DisplayName = "2 x 2"),
     POT_4 UMETA(DisplayName = "4 x 4"),
@@ -30,7 +37,7 @@ enum class EGLTFOverrideMaterialBakeSizePOT : uint8
 };
 
 UENUM(BlueprintType)
-enum class EGLTFOverrideMaterialPropertyGroup : uint8
+enum class EGLTFMaterialPropertyGroup : uint8
 {
 	None UMETA(DisplayName = "None"),
 
@@ -43,25 +50,52 @@ enum class EGLTFOverrideMaterialPropertyGroup : uint8
     ClearCoatBottomNormal UMETA(DisplayName = "Clear Coat Bottom Normal"),
 };
 
+USTRUCT(Blueprintable)
+struct FGLTFOverrideMaterialBakeSettings
+{
+	GENERATED_BODY()
+
+	UPROPERTY(EditAnywhere, BlueprintReadWrite, meta = (InlineEditConditionToggle))
+	bool bOverrideSize;
+
+	UPROPERTY(EditAnywhere, BlueprintReadWrite, meta = (EditCondition = "bOverrideSize"))
+	EGLTFMaterialBakeSizePOT Size;
+
+	UPROPERTY(EditAnywhere, BlueprintReadWrite, meta = (InlineEditConditionToggle))
+	bool bOverrideFilter;
+
+	UPROPERTY(EditAnywhere, BlueprintReadWrite, meta = (EditCondition = "bOverrideFilter", ValidEnumValues="TF_Nearest, TF_Bilinear, TF_Trilinear"))
+	TEnumAsByte<TextureFilter> Filter;
+
+	UPROPERTY(EditAnywhere, BlueprintReadWrite, meta = (InlineEditConditionToggle))
+	bool bOverrideTiling;
+
+	UPROPERTY(EditAnywhere, BlueprintReadWrite, meta = (EditCondition = "bOverrideTiling"))
+	TEnumAsByte<TextureAddress> Tiling;
+
+	FGLTFOverrideMaterialBakeSettings();
+};
+
 /** glTF-specific user data that can be added to material assets to override export options */
-UCLASS(BlueprintType, meta = (DisplayName = "GLTF Material User Data"))
-class GLTFEXPORTERRUNTIME_API UGLTFMaterialUserData : public UAssetUserData
+UCLASS(BlueprintType, meta = (DisplayName = "GLTF Material Export Options"))
+class GLTFEXPORTERRUNTIME_API UGLTFMaterialExportOptions : public UAssetUserData
 {
 	GENERATED_BODY()
 
 public:
 
-	UGLTFMaterialUserData();
-
-	UPROPERTY(EditAnywhere, BlueprintReadWrite, Category = "Override Export Options")
-	EGLTFOverrideMaterialBakeSizePOT DefaultBakeSize;
+	UPROPERTY(EditAnywhere, BlueprintReadWrite, Category = "Override Bake Settings", meta = (ShowOnlyInnerProperties))
+	FGLTFOverrideMaterialBakeSettings Default;
 
-	UPROPERTY(EditAnywhere, BlueprintReadWrite, Category = "Override Export Options")
-	TMap<EGLTFOverrideMaterialPropertyGroup, EGLTFOverrideMaterialBakeSizePOT> BakeSizePerProperty;
+	UPROPERTY(EditAnywhere, BlueprintReadWrite, Category = "Override Bake Settings")
+	TMap<EGLTFMaterialPropertyGroup, FGLTFOverrideMaterialBakeSettings> PerInput;
 
-	EGLTFOverrideMaterialBakeSizePOT GetBakeSizeForProperty(EMaterialProperty Property) const;
+	static EGLTFMaterialBakeSizePOT GetBakeSizeForPropertyGroup(const UMaterialInterface* Material, EGLTFMaterialPropertyGroup PropertyGroup, EGLTFMaterialBakeSizePOT DefaultValue);
+	static TextureFilter GetBakeFilterForPropertyGroup(const UMaterialInterface* Material, EGLTFMaterialPropertyGroup PropertyGroup, TextureFilter DefaultValue);
+	static TextureAddress GetBakeTilingForPropertyGroup(const UMaterialInterface* Material, EGLTFMaterialPropertyGroup PropertyGroup, TextureAddress DefaultValue);
 
-	static EGLTFOverrideMaterialPropertyGroup GetPropertyGroup(EMaterialProperty Property);
+private:
 
-	static const UGLTFMaterialUserData* GetUserData(const UMaterialInterface* Material);
+	template <typename Predicate>
+	static const FGLTFOverrideMaterialBakeSettings* GetBakeSettingsByPredicate(const UMaterialInterface* Material, EGLTFMaterialPropertyGroup PropertyGroup, Predicate Pred);
 };
