[GLTFExporter] refactor: add some TODO comments

commit b760f9c6a758fd7f627c992126adcc25e79a722b
Date: Mon, 16 Nov 2020 17:05:57 +0100

    refactor: add some TODO comments
    (cherry picked from commit c62a50f0c106a1677fccd49b600a913b83dfba7b​)

#skipci

diff --git a/Engine/Plugins/Enterprise/GLTFExporter/MigratedFromGit/Source/GLTFExporter/Private/Builders/GLTFBuilder.h b/Engine/Plugins/Enterprise/GLTFExporter/MigratedFromGit/Source/GLTFExporter/Private/Builders/GLTFBuilder.h
--- a/Engine/Plugins/Enterprise/GLTFExporter/MigratedFromGit/Source/GLTFExporter/Private/Builders/GLTFBuilder.h
+++ b/Engine/Plugins/Enterprise/GLTFExporter/MigratedFromGit/Source/GLTFExporter/Private/Builders/GLTFBuilder.h
@@ -18,6 +18,7 @@ public:
 	const FString FilePath;
 	const FString DirPath;
 
+	// TODO: make ExportOptions private and expose each option via getters to ease overriding settings in future
 	const UGLTFExportOptions* const ExportOptions;
 
 	FIntPoint GetBakeSizeForMaterialProperty(const UMaterialInterface* Material, const FMaterialPropertyEx& Property) const;
diff --git a/Engine/Plugins/Enterprise/GLTFExporter/MigratedFromGit/Source/GLTFExporter/Private/Tasks/GLTFAnimationTasks.cpp b/Engine/Plugins/Enterprise/GLTFExporter/MigratedFromGit/Source/GLTFExporter/Private/Tasks/GLTFAnimationTasks.cpp
--- a/Engine/Plugins/Enterprise/GLTFExporter/MigratedFromGit/Source/GLTFExporter/Private/Tasks/GLTFAnimationTasks.cpp
+++ b/Engine/Plugins/Enterprise/GLTFExporter/MigratedFromGit/Source/GLTFExporter/Private/Tasks/GLTFAnimationTasks.cpp
@@ -13,6 +13,8 @@
 
 void FGLTFAnimSequenceTask::Complete()
 {
+	// TODO: bone transforms should be absolute (not relative) according to gltf spec
+
 	const int32 FrameCount = AnimSequence->GetRawNumberOfFrames();
 	const USkeleton* Skeleton = AnimSequence->GetSkeleton();
 
diff --git a/Engine/Plugins/Enterprise/GLTFExporter/MigratedFromGit/Source/GLTFExporterRuntime/Public/UserData/GLTFMaterialUserData.h b/Engine/Plugins/Enterprise/GLTFExporter/MigratedFromGit/Source/GLTFExporterRuntime/Public/UserData/GLTFMaterialUserData.h
--- a/Engine/Plugins/Enterprise/GLTFExporter/MigratedFromGit/Source/GLTFExporterRuntime/Public/UserData/GLTFMaterialUserData.h
+++ b/Engine/Plugins/Enterprise/GLTFExporter/MigratedFromGit/Source/GLTFExporterRuntime/Public/UserData/GLTFMaterialUserData.h
@@ -84,6 +84,8 @@ class GLTFEXPORTERRUNTIME_API UGLTFMaterialExportOptions : public UAssetUserData
 
 public:
 
+	// TODO: add support for overriding more export options
+
 	UPROPERTY(EditAnywhere, BlueprintReadWrite, Category = "Override Bake Settings", meta = (ShowOnlyInnerProperties))
 	FGLTFOverrideMaterialBakeSettings Default;
 
