[GLTFExporter] refactor: simplify material conversion of alpha mode and shading model

commit 7faec7bf889d2913a573e90fbd419df6af715235
Date: Mon, 16 Nov 2020 17:44:52 +0100

    refactor: simplify material conversion of alpha mode and shading model
    (cherry picked from commit 1959e8bef91d9ced97449107b1521fedfff75b38​)

#skipci

diff --git a/Engine/Plugins/Enterprise/GLTFExporter/MigratedFromGit/Source/GLTFExporter/Private/Tasks/GLTFMaterialTasks.cpp b/Engine/Plugins/Enterprise/GLTFExporter/MigratedFromGit/Source/GLTFExporter/Private/Tasks/GLTFMaterialTasks.cpp
--- a/Engine/Plugins/Enterprise/GLTFExporter/MigratedFromGit/Source/GLTFExporter/Private/Tasks/GLTFMaterialTasks.cpp
+++ b/Engine/Plugins/Enterprise/GLTFExporter/MigratedFromGit/Source/GLTFExporter/Private/Tasks/GLTFMaterialTasks.cpp
@@ -58,21 +58,11 @@ void FGLTFMaterialTask::Complete()
 
 	FGLTFJsonMaterial& JsonMaterial = Builder.GetMaterial(MaterialIndex);
 	JsonMaterial.Name = GetMaterialName();
-
-	if (!TryGetAlphaMode(JsonMaterial.AlphaMode, JsonMaterial.BlendMode))
-	{
-		JsonMaterial.AlphaMode = EGLTFJsonAlphaMode::Blend;
-		Builder.AddWarningMessage(FString::Printf(TEXT("Material %s will be exported as blend mode %s"), *Material->GetName(), *FGLTFNameUtility::GetName(BLEND_Translucent)));
-	}
-
 	JsonMaterial.AlphaCutoff = Material->GetOpacityMaskClipValue();
 	JsonMaterial.DoubleSided = Material->IsTwoSided();
 
-	if (!TryGetShadingModel(JsonMaterial.ShadingModel))
-	{
-		JsonMaterial.ShadingModel = EGLTFJsonShadingModel::Default;
-		Builder.AddWarningMessage(FString::Printf(TEXT("Material %s will be exported as shading model %s"), *Material->GetName(), *FGLTFNameUtility::GetName(MSM_DefaultLit)));
-	}
+	ConvertAlphaMode(JsonMaterial.AlphaMode, JsonMaterial.BlendMode);
+	ConvertShadingModel(JsonMaterial.ShadingModel);
 
 	if (JsonMaterial.ShadingModel != EGLTFJsonShadingModel::None)
 	{
@@ -170,36 +160,41 @@ void FGLTFMaterialTask::Complete()
 	}
 }
 
-bool FGLTFMaterialTask::TryGetAlphaMode(EGLTFJsonAlphaMode& OutAlphaMode, EGLTFJsonBlendMode& OutBlendMode) const
+void FGLTFMaterialTask::ConvertAlphaMode(EGLTFJsonAlphaMode& OutAlphaMode, EGLTFJsonBlendMode& OutBlendMode) const
 {
 	const EBlendMode BlendMode = Material->GetBlendMode();
 
-	const EGLTFJsonAlphaMode ConvertedAlphaMode = FGLTFConverterUtility::ConvertAlphaMode(BlendMode);
-	if (ConvertedAlphaMode == EGLTFJsonAlphaMode::None)
+	OutAlphaMode = FGLTFConverterUtility::ConvertAlphaMode(BlendMode);
+	if (OutAlphaMode == EGLTFJsonAlphaMode::None)
 	{
+		OutAlphaMode = EGLTFJsonAlphaMode::Blend;
+		OutBlendMode = EGLTFJsonBlendMode::None;
+
 		Builder.AddWarningMessage(FString::Printf(
-			TEXT("Unsupported blend mode (%s) in material %s"),
+			TEXT("Unsupported blend mode (%s) in material %s, will export as %s"),
 			*FGLTFNameUtility::GetName(BlendMode),
-			*Material->GetName()));
-		return false;
+			*Material->GetName(),
+			*FGLTFNameUtility::GetName(BLEND_Translucent)));
+		return;
 	}
 
-	const EGLTFJsonBlendMode ConvertedBlendMode = ConvertedAlphaMode == EGLTFJsonAlphaMode::Blend ? FGLTFConverterUtility::ConvertBlendMode(BlendMode) : EGLTFJsonBlendMode::None;
-	if (ConvertedBlendMode != EGLTFJsonBlendMode::None && !Builder.ExportOptions->bExportExtraBlendModes)
+	if (OutAlphaMode == EGLTFJsonAlphaMode::Blend)
 	{
-		Builder.AddWarningMessage(FString::Printf(
-			TEXT("Extra blend mode (%s) in material %s disabled by export options"),
-			*FGLTFNameUtility::GetName(BlendMode),
-			*Material->GetName()));
-		return false;
-	}
+		OutBlendMode = FGLTFConverterUtility::ConvertBlendMode(BlendMode);
+		if (OutBlendMode != EGLTFJsonBlendMode::None && !Builder.ExportOptions->bExportExtraBlendModes)
+		{
+			OutBlendMode = EGLTFJsonBlendMode::None;
 
-	OutAlphaMode = ConvertedAlphaMode;
-	OutBlendMode = ConvertedBlendMode;
-	return true;
+			Builder.AddWarningMessage(FString::Printf(
+				TEXT("Extra blend mode (%s) in material %s disabled by export options, will export as %s"),
+				*FGLTFNameUtility::GetName(BlendMode),
+				*Material->GetName(),
+				*FGLTFNameUtility::GetName(BLEND_Translucent)));
+		}
+	}
 }
 
-bool FGLTFMaterialTask::TryGetShadingModel(EGLTFJsonShadingModel& OutShadingModel) const
+void FGLTFMaterialTask::ConvertShadingModel(EGLTFJsonShadingModel& OutShadingModel) const
 {
 	EMaterialShadingModel ShadingModel;
 	{
@@ -235,30 +230,41 @@ bool FGLTFMaterialTask::TryGetShadingModel(EGLTFJsonShadingModel& OutShadingMode
 		}
 	}
 
-	const EGLTFJsonShadingModel ConvertedShadingModel = FGLTFConverterUtility::ConvertShadingModel(ShadingModel);
-	if (ConvertedShadingModel == EGLTFJsonShadingModel::None)
+	OutShadingModel = FGLTFConverterUtility::ConvertShadingModel(ShadingModel);
+	if (OutShadingModel == EGLTFJsonShadingModel::None)
 	{
-		const FString ShadingModelName = FGLTFNameUtility::GetName(ShadingModel);
-		Builder.AddWarningMessage(FString::Printf(TEXT("Unsupported shading model (%s) in material %s"), *ShadingModelName, *Material->GetName()));
-		return false;
+		OutShadingModel = EGLTFJsonShadingModel::Default;
+
+		Builder.AddWarningMessage(FString::Printf(
+			TEXT("Unsupported shading model (%s) in material %s, will export as %s"),
+			*FGLTFNameUtility::GetName(ShadingModel),
+			*Material->GetName(),
+			*FGLTFNameUtility::GetName(MSM_DefaultLit)));
+		return;
 	}
 
-	if (ConvertedShadingModel == EGLTFJsonShadingModel::Unlit && !Builder.ExportOptions->bExportUnlitMaterials)
+	if (OutShadingModel == EGLTFJsonShadingModel::Unlit && !Builder.ExportOptions->bExportUnlitMaterials)
 	{
-		const FString ShadingModelName = FGLTFNameUtility::GetName(ShadingModel);
-		Builder.AddWarningMessage(FString::Printf(TEXT("Shading model (%s) in material %s disabled by export options"), *ShadingModelName, *Material->GetName()));
-		return false;
+		OutShadingModel = EGLTFJsonShadingModel::Default;
+
+		Builder.AddWarningMessage(FString::Printf(
+			TEXT("Shading model (%s) in material %s disabled by export options, will export as %s"),
+			*FGLTFNameUtility::GetName(ShadingModel),
+			*Material->GetName(),
+			*FGLTFNameUtility::GetName(MSM_DefaultLit)));
+		return;
 	}
 
-	if (ConvertedShadingModel == EGLTFJsonShadingModel::ClearCoat && !Builder.ExportOptions->bExportClearCoatMaterials)
+	if (OutShadingModel == EGLTFJsonShadingModel::ClearCoat && !Builder.ExportOptions->bExportClearCoatMaterials)
 	{
-		const FString ShadingModelName = FGLTFNameUtility::GetName(ShadingModel);
-		Builder.AddWarningMessage(FString::Printf(TEXT("Shading model (%s) in material %s disabled by export options"), *ShadingModelName, *Material->GetName()));
-		return false;
-	}
+		OutShadingModel = EGLTFJsonShadingModel::Default;
 
-	OutShadingModel = ConvertedShadingModel;
-	return true;
+		Builder.AddWarningMessage(FString::Printf(
+			TEXT("Shading model (%s) in material %s disabled by export options, will export as %s"),
+			*FGLTFNameUtility::GetName(ShadingModel),
+			*Material->GetName(),
+			*FGLTFNameUtility::GetName(MSM_DefaultLit)));
+	}
 }
 
 bool FGLTFMaterialTask::TryGetBaseColorAndOpacity(FGLTFJsonPBRMetallicRoughness& OutPBRParams, const FMaterialPropertyEx& BaseColorProperty, const FMaterialPropertyEx& OpacityProperty) const
diff --git a/Engine/Plugins/Enterprise/GLTFExporter/MigratedFromGit/Source/GLTFExporter/Private/Tasks/GLTFMaterialTasks.h b/Engine/Plugins/Enterprise/GLTFExporter/MigratedFromGit/Source/GLTFExporter/Private/Tasks/GLTFMaterialTasks.h
--- a/Engine/Plugins/Enterprise/GLTFExporter/MigratedFromGit/Source/GLTFExporter/Private/Tasks/GLTFMaterialTasks.h
+++ b/Engine/Plugins/Enterprise/GLTFExporter/MigratedFromGit/Source/GLTFExporter/Private/Tasks/GLTFMaterialTasks.h
@@ -39,8 +39,8 @@ private:
 	const TArray<int32> SectionIndices;
 	const FGLTFJsonMaterialIndex MaterialIndex;
 
-	bool TryGetAlphaMode(EGLTFJsonAlphaMode& OutAlphaMode, EGLTFJsonBlendMode& OutBlendMode) const;
-	bool TryGetShadingModel(EGLTFJsonShadingModel& OutShadingModel) const;
+	void ConvertAlphaMode(EGLTFJsonAlphaMode& OutAlphaMode, EGLTFJsonBlendMode& OutBlendMode) const;
+	void ConvertShadingModel(EGLTFJsonShadingModel& OutShadingModel) const;
 
 	bool TryGetBaseColorAndOpacity(FGLTFJsonPBRMetallicRoughness& OutPBRParams, const FMaterialPropertyEx& BaseColorProperty, const FMaterialPropertyEx& OpacityProperty) const;
 	bool TryGetMetallicAndRoughness(FGLTFJsonPBRMetallicRoughness& OutPBRParams, const FMaterialPropertyEx& MetallicProperty, const FMaterialPropertyEx& RoughnessProperty) const;
