[GLTFExporter] refactor: change default value for export options BakeMaterialInputs and ExportMaterialVariants from Simple to UseMeshData

commit c31abc4ab7f21090fe46c5005b43de181a70b289
Date: Mon, 16 Nov 2020 19:38:58 +0100

    refactor: change default value for export options BakeMaterialInputs and ExportMaterialVariants from Simple to UseMeshData
    (cherry picked from commit ec6bcd0748e3fb6c560227df89ddddf034fef584​)

#skipci

diff --git a/Engine/Plugins/Enterprise/GLTFExporter/MigratedFromGit/Source/GLTFExporter/Private/GLTFExportOptions.cpp b/Engine/Plugins/Enterprise/GLTFExporter/MigratedFromGit/Source/GLTFExporter/Private/GLTFExportOptions.cpp
--- a/Engine/Plugins/Enterprise/GLTFExporter/MigratedFromGit/Source/GLTFExporter/Private/GLTFExportOptions.cpp
+++ b/Engine/Plugins/Enterprise/GLTFExporter/MigratedFromGit/Source/GLTFExporter/Private/GLTFExportOptions.cpp
@@ -15,7 +15,7 @@ void UGLTFExportOptions::ResetToDefault()
 	bExportUnlitMaterials = true;
 	bExportClearCoatMaterials = true;
 	bExportExtraBlendModes = true;
-	BakeMaterialInputs = EGLTFMaterialBakeMode::Simple;
+	BakeMaterialInputs = EGLTFMaterialBakeMode::UseMeshData;
 	DefaultMaterialBakeSize = EGLTFMaterialBakeSizePOT::POT_1024;
 	DefaultMaterialBakeFilter = TF_Trilinear;
 	DefaultMaterialBakeTiling = TA_Wrap;
@@ -41,7 +41,7 @@ void UGLTFExportOptions::ResetToDefault()
 	bExportHDRIBackdrops = true;
 	bExportSkySpheres = true;
 	bExportVariantSets = true;
-	ExportMaterialVariants = EGLTFMaterialBakeMode::Simple;
+	ExportMaterialVariants = EGLTFMaterialBakeMode::UseMeshData;
 	bExportMeshVariants = true;
 	bExportVisibilityVariants = true;
 	bBundleWebViewer = true;
