[GLTFExporter] fix: incorrect assumption about the JSON schema used by index.json

commit ecff36c80b241200d599e452dabd77cffa8e8cdf
Date: Mon, 16 Nov 2020 21:42:57 +0100

    fix: incorrect assumption about the JSON schema used by index.json
    (cherry picked from commit 11cc1dbf73818faaaeffaa680bff0bd9e1899b15​)

#skipci

diff --git a/Engine/Plugins/Enterprise/GLTFExporter/MigratedFromGit/Source/GLTFExporter/Private/Builders/GLTFWebBuilder.cpp b/Engine/Plugins/Enterprise/GLTFExporter/MigratedFromGit/Source/GLTFExporter/Private/Builders/GLTFWebBuilder.cpp
--- a/Engine/Plugins/Enterprise/GLTFExporter/MigratedFromGit/Source/GLTFExporter/Private/Builders/GLTFWebBuilder.cpp
+++ b/Engine/Plugins/Enterprise/GLTFExporter/MigratedFromGit/Source/GLTFExporter/Private/Builders/GLTFWebBuilder.cpp
@@ -47,17 +47,25 @@ void FGLTFWebBuilder::BundleWebViewer(const FString& ResourcesDir)
 	}
 
 	const FString IndexFile = DirPath / TEXT("index.json");
-	TSharedPtr<FJsonObject> JsonObject;
+	TSharedPtr<FJsonObject> JsonIndexRoot;
 
-	if (!FGLTFFileUtility::ReadJsonFile(IndexFile, JsonObject))
+	if (!FGLTFFileUtility::ReadJsonFile(IndexFile, JsonIndexRoot))
 	{
 		AddWarningMessage(FString::Printf(TEXT("Failed to read web viewer index at %s"), *IndexFile));
 		return;
 	}
 
-	JsonObject->SetArrayField(TEXT("assets"), { MakeShared<FJsonValueString>(FPaths::GetCleanFilename(FilePath)) });
+	const FString FileName = FPaths::GetCleanFilename(FilePath);
+	const FString Name = FPaths::GetBaseFilename(FileName);
 
-	if (!FGLTFFileUtility::WriteJsonFile(IndexFile, JsonObject.ToSharedRef()))
+	// TODO: simplify the index.json to avoid unnecessary json fields
+	TSharedPtr<FJsonObject> JsonAssetObject = MakeShared<FJsonObject>();
+	JsonAssetObject->SetStringField(TEXT("name"), Name);
+	JsonAssetObject->SetStringField(TEXT("filePath"), TEXT("../../") + FileName); // TODO: change root to mean the index.html folder (not viewer/playcanvas/)
+	JsonAssetObject->SetStringField(TEXT("blobFileName"), FileName);
+	JsonIndexRoot->SetArrayField(TEXT("assets"), { MakeShared<FJsonValueObject>(JsonAssetObject) });
+
+	if (!FGLTFFileUtility::WriteJsonFile(IndexFile, JsonIndexRoot.ToSharedRef()))
 	{
 		AddWarningMessage(FString::Printf(TEXT("Failed to write web viewer index at %s"), *IndexFile));
 	}
