[GLTFExporter] chore: update plugin version to 0.9

commit 898feae8d95981f579d06cd551950c9bba649e6d
Date: Mon, 16 Nov 2020 21:58:55 +0100

    chore: update plugin version to 0.9
    (cherry picked from commit 08bf0c1e96c21c26cee55a2e421d93844ea5d218​)

#skipci

diff --git a/Engine/Plugins/Enterprise/GLTFExporter/MigratedFromGit/GLTFExporter.uplugin b/Engine/Plugins/Enterprise/GLTFExporter/MigratedFromGit/GLTFExporter.uplugin
--- a/Engine/Plugins/Enterprise/GLTFExporter/MigratedFromGit/GLTFExporter.uplugin
+++ b/Engine/Plugins/Enterprise/GLTFExporter/MigratedFromGit/GLTFExporter.uplugin
@@ -1,7 +1,7 @@
 {
 	"FileVersion": 3,
-	"Version": 5,
-	"VersionName": "0.8",
+	"Version": 9,
+	"VersionName": "0.9",
 	"FriendlyName": "glTF Exporter",
 	"Description": "An exporter for Khronos glTF 3D models.",
 	"Category": "Exporters",
