[GLTFExporter] feat: support specifying a blend-mode when baking a material

commit cbeaf665f21f890670edcb9cde574cc3435787eb
Date: Tue, 17 Nov 2020 10:52:32 +0100

    feat: support specifying a blend-mode when baking a material
    (cherry picked from commit 56a9974846ca087116f7719d520e89a18f4de820​)

#skipci

diff --git a/Engine/Plugins/Enterprise/GLTFExporter/MigratedFromGit/Source/GLTFMaterialBaking/Private/ExportMaterialProxy.h b/Engine/Plugins/Enterprise/GLTFExporter/MigratedFromGit/Source/GLTFMaterialBaking/Private/ExportMaterialProxy.h
--- a/Engine/Plugins/Enterprise/GLTFExporter/MigratedFromGit/Source/GLTFMaterialBaking/Private/ExportMaterialProxy.h
+++ b/Engine/Plugins/Enterprise/GLTFExporter/MigratedFromGit/Source/GLTFMaterialBaking/Private/ExportMaterialProxy.h
@@ -224,12 +224,13 @@ struct FExportMaterialCompiler : public FProxyMaterialCompiler
 class FExportMaterialProxy : public FMaterial, public FMaterialRenderProxy
 {
 public:
-	FExportMaterialProxy(UMaterialInterface* InMaterialInterface, EMaterialProperty InPropertyToCompile, const FString& InCustomOutputToCompile = TEXT(""), bool bInSynchronousCompilation = true)
+	FExportMaterialProxy(UMaterialInterface* InMaterialInterface, EMaterialProperty InPropertyToCompile, const FString& InCustomOutputToCompile = TEXT(""), bool bInSynchronousCompilation = true, EBlendMode ProxyBlendMode = BLEND_Opaque)
 		: FMaterial()
 		, MaterialInterface(InMaterialInterface)
 		, PropertyToCompile(InPropertyToCompile)
 		, CustomOutputToCompile(InCustomOutputToCompile)
 		, bSynchronousCompilation(bInSynchronousCompilation)
+		, ProxyBlendMode(ProxyBlendMode)
 	{
 		SetQualityLevelProperties(EMaterialQualityLevel::High, false, GMaxRHIFeatureLevel);
 		Material = InMaterialInterface->GetMaterial();
@@ -504,7 +505,7 @@ public:
 		return false;
 	}
 	virtual bool IsMasked() const override { return false; }
-	virtual enum EBlendMode GetBlendMode() const override { return BLEND_Opaque; }
+	virtual enum EBlendMode GetBlendMode() const override { return ProxyBlendMode; }
 	virtual FMaterialShadingModelField GetShadingModels() const override { return MSM_DefaultLit; }
 	virtual bool IsShadingModelFromMaterialExpression() const override { return false; }
 	virtual float GetOpacityMaskClipValue() const override { return 0.5f; }
@@ -731,4 +732,5 @@ private:
 	FString CustomOutputToCompile;
 	FGuid Id;
 	bool bSynchronousCompilation;
+	EBlendMode ProxyBlendMode;
 };
diff --git a/Engine/Plugins/Enterprise/GLTFExporter/MigratedFromGit/Source/GLTFMaterialBaking/Private/MaterialBakingModule.cpp b/Engine/Plugins/Enterprise/GLTFExporter/MigratedFromGit/Source/GLTFMaterialBaking/Private/MaterialBakingModule.cpp
--- a/Engine/Plugins/Enterprise/GLTFExporter/MigratedFromGit/Source/GLTFMaterialBaking/Private/MaterialBakingModule.cpp
+++ b/Engine/Plugins/Enterprise/GLTFExporter/MigratedFromGit/Source/GLTFMaterialBaking/Private/MaterialBakingModule.cpp
@@ -294,6 +294,7 @@ void FMaterialBakingModule::BakeMaterials(const TArray<FMaterialData*>& Material
 		FMaterialDataEx& MaterialDataEx = MaterialDataExs.AddDefaulted_GetRef();
 		MaterialDataEx.Material = MaterialData->Material;
 		MaterialDataEx.bPerformBorderSmear = MaterialData->bPerformBorderSmear;
+		MaterialDataEx.BlendMode = MaterialData->BlendMode;
 
 		for (const TPair<EMaterialProperty, FIntPoint>& PropertySizePair : MaterialData->PropertySizes)
 		{
@@ -439,7 +440,7 @@ void FMaterialBakingModule::BakeMaterials(const TArray<FMaterialDataEx*>& Materi
 			for (TMap<FMaterialPropertyEx, FIntPoint>::TConstIterator PropertySizeIterator = CurrentMaterialSettings->PropertySizes.CreateConstIterator(); PropertySizeIterator; ++PropertySizeIterator)
 			{
 				// They will be stored in the pool and compiled asynchronously
-				CreateMaterialProxy(CurrentMaterialSettings->Material, PropertySizeIterator.Key());
+				CreateMaterialProxy(*CurrentMaterialSettings, PropertySizeIterator.Key());
 			}
 		}
 	}
@@ -499,7 +500,7 @@ void FMaterialBakingModule::BakeMaterials(const TArray<FMaterialDataEx*>& Materi
 			for (int32 PropertyIndex = 0; PropertyIndex < NumPropertiesToRender; ++PropertyIndex)
 			{
 				const FMaterialPropertyEx& Property = MaterialPropertiesToBakeOut[PropertyIndex];
-				FExportMaterialProxy* ExportMaterialProxy = CreateMaterialProxy(CurrentMaterialSettings->Material, Property);
+				FExportMaterialProxy* ExportMaterialProxy = CreateMaterialProxy(*CurrentMaterialSettings, Property);
 
 				if (!ExportMaterialProxy->IsCompilationFinished())
 				{
@@ -773,7 +774,7 @@ UTextureRenderTarget2D* FMaterialBakingModule::CreateRenderTarget(bool bInForceL
 	return RenderTarget;
 }
 
-FExportMaterialProxy* FMaterialBakingModule::CreateMaterialProxy(UMaterialInterface* Material, const FMaterialPropertyEx& Property)
+FExportMaterialProxy* FMaterialBakingModule::CreateMaterialProxy(const FMaterialDataEx& MatSettings, const FMaterialPropertyEx& Property)
 {
 	TRACE_CPUPROFILER_EVENT_SCOPE(FMaterialBakingModule::CreateMaterialProxy)
 
@@ -781,12 +782,12 @@ FExportMaterialProxy* FMaterialBakingModule::CreateMaterialProxy(UMaterialInterf
 
 	// Find all pooled material proxy matching this material
 	TArray<FMaterialPoolValue> Entries;
-	MaterialProxyPool.MultiFind(Material, Entries);
+	MaterialProxyPool.MultiFind(MatSettings.Material, Entries);
 
 	// Look for the matching property
 	for (FMaterialPoolValue& Entry : Entries)
 	{
-		if (Entry.Key == Property)
+		if (Entry.Key == Property && Entry.Value->GetBlendMode() == MatSettings.BlendMode)
 		{
 			Proxy = Entry.Value;
 			break;
@@ -796,8 +797,8 @@ FExportMaterialProxy* FMaterialBakingModule::CreateMaterialProxy(UMaterialInterf
 	// Not found, create a new entry
 	if (Proxy == nullptr)
 	{
-		Proxy = new FExportMaterialProxy(Material, Property.Type, Property.CustomOutput.ToString(), false /* bInSynchronousCompilation */);
-		MaterialProxyPool.Add(Material, FMaterialPoolValue(Property, Proxy));
+		Proxy = new FExportMaterialProxy(MatSettings.Material, Property.Type, Property.CustomOutput.ToString(), false /* bInSynchronousCompilation */, MatSettings.BlendMode);
+		MaterialProxyPool.Add(MatSettings.Material, FMaterialPoolValue(Property, Proxy));
 	}
 
 	return Proxy;
diff --git a/Engine/Plugins/Enterprise/GLTFExporter/MigratedFromGit/Source/GLTFMaterialBaking/Private/MaterialBakingModule.h b/Engine/Plugins/Enterprise/GLTFExporter/MigratedFromGit/Source/GLTFMaterialBaking/Private/MaterialBakingModule.h
--- a/Engine/Plugins/Enterprise/GLTFExporter/MigratedFromGit/Source/GLTFMaterialBaking/Private/MaterialBakingModule.h
+++ b/Engine/Plugins/Enterprise/GLTFExporter/MigratedFromGit/Source/GLTFMaterialBaking/Private/MaterialBakingModule.h
@@ -42,7 +42,7 @@ protected:
 	UTextureRenderTarget2D* CreateRenderTarget(bool bInForceLinearGamma, EPixelFormat InPixelFormat, const FIntPoint& InTargetSize);
 
 	/* Creates and adds (or reuses a ExportMaterialProxy from the pool if MaterialBaking.UseMaterialProxyCaching is set to 1) */
-	FExportMaterialProxy* CreateMaterialProxy(UMaterialInterface* Material, const FMaterialPropertyEx& Property);
+	FExportMaterialProxy* CreateMaterialProxy(const FMaterialDataEx& MatSettings, const FMaterialPropertyEx& Property);
 
 	/** Helper for emissive color conversion to Output */
 	static void ProcessEmissiveOutput(const FFloat16Color* Color16, int32 Color16Pitch, const FIntPoint& OutputSize, TArray<FColor>& Output, float& EmissiveScale);
diff --git a/Engine/Plugins/Enterprise/GLTFExporter/MigratedFromGit/Source/GLTFMaterialBaking/Public/MaterialBakingStructures.h b/Engine/Plugins/Enterprise/GLTFExporter/MigratedFromGit/Source/GLTFMaterialBaking/Public/MaterialBakingStructures.h
--- a/Engine/Plugins/Enterprise/GLTFExporter/MigratedFromGit/Source/GLTFMaterialBaking/Public/MaterialBakingStructures.h
+++ b/Engine/Plugins/Enterprise/GLTFExporter/MigratedFromGit/Source/GLTFMaterialBaking/Public/MaterialBakingStructures.h
@@ -16,6 +16,7 @@ struct FMaterialData
 	FMaterialData()
 		: Material(nullptr)
 		, bPerformBorderSmear(true)
+		, BlendMode(BLEND_Opaque)
 	{}
 
 	/** Material to bake out */
@@ -24,6 +25,8 @@ struct FMaterialData
 	TMap<EMaterialProperty, FIntPoint> PropertySizes;
 	/** Whether to smear borders after baking */
 	bool bPerformBorderSmear;
+	/** Blend mode to use when baking */
+	EBlendMode BlendMode;
 };
 
 /** Structure containing extended information about the material and properties which is being baked out */
@@ -32,6 +35,7 @@ struct FMaterialDataEx
 	FMaterialDataEx()
 		: Material(nullptr)
 		, bPerformBorderSmear(true)
+		, BlendMode(BLEND_Opaque)
 	{}
 
 	/** Material to bake out */
@@ -40,6 +44,8 @@ struct FMaterialDataEx
 	TMap<FMaterialPropertyEx, FIntPoint> PropertySizes;
 	/** Whether to smear borders after baking */
 	bool bPerformBorderSmear;
+	/** Blend mode to use when baking */
+	EBlendMode BlendMode;
 };
 
 struct FMeshData
