[GLTFExporter] feat: detect and warn if overlapping uv:s are used when baking with mesh data

commit a4ad30994deb64975550dd9faf0a7053d1bd6792
Date: Tue, 17 Nov 2020 14:09:33 +0100

    feat: detect and warn if overlapping uv:s are used when baking with mesh data
    (cherry picked from commit 2eab605322d4a87cdfbdb040f1d51d4fccdbe87b​)

#skipci

diff --git a/Engine/Plugins/Enterprise/GLTFExporter/MigratedFromGit/Source/GLTFExporter/Private/Converters/GLTFMaterialUtility.cpp b/Engine/Plugins/Enterprise/GLTFExporter/MigratedFromGit/Source/GLTFExporter/Private/Converters/GLTFMaterialUtility.cpp
--- a/Engine/Plugins/Enterprise/GLTFExporter/MigratedFromGit/Source/GLTFExporter/Private/Converters/GLTFMaterialUtility.cpp
+++ b/Engine/Plugins/Enterprise/GLTFExporter/MigratedFromGit/Source/GLTFExporter/Private/Converters/GLTFMaterialUtility.cpp
@@ -460,6 +460,71 @@ void FGLTFMaterialUtility::AnalyzeMaterialProperty(const UMaterialInterface* InM
 	UGLTFMaterialAnalyzer::AnalyzeMaterialPropertyEx(InMaterial, InProperty.Type, InProperty.CustomOutput.ToString(), OutAnalysis);
 }
 
+float FGLTFMaterialUtility::CalcOverlappingUVPercentage(int32 TexCoord, const FIntPoint& OutputSize, const FMeshDescription& MeshDescription, const TArray<int32>& MeshSectionIndices)
+{
+	static UMaterialInterface* Material = LoadObject<UMaterialInterface>(nullptr, TEXT("/GLTFExporter/Materials/OverlappingUVs.OverlappingUVs"));
+	check(Material);
+
+	const FMaterialPropertyEx Property = MP_EmissiveColor;
+
+	FMeshData MeshSet;
+	MeshSet.TextureCoordinateBox = { { 0.0f, 0.0f }, { 1.0f, 1.0f } };
+	MeshSet.TextureCoordinateIndex = TexCoord;
+	MeshSet.RawMeshDescription = const_cast<FMeshDescription*>(&MeshDescription);
+	MeshSet.MaterialIndices = MeshSectionIndices; // NOTE: MaterialIndices is actually section indices
+
+	FMaterialDataEx MatSet;
+	MatSet.Material = Material;
+	MatSet.PropertySizes.Add(Property, OutputSize);
+	MatSet.BlendMode = BLEND_Additive;
+	MatSet.bPerformBorderSmear = false;
+
+	TArray<FMeshData*> MeshSettings;
+	TArray<FMaterialDataEx*> MatSettings;
+	MeshSettings.Add(&MeshSet);
+	MatSettings.Add(&MatSet);
+
+	TArray<FBakeOutputEx> BakeOutputs;
+	IMaterialBakingModule& Module = FModuleManager::Get().LoadModuleChecked<IMaterialBakingModule>("GLTFMaterialBaking");
+
+	Module.BakeMaterials(MatSettings, MeshSettings, BakeOutputs);
+
+	FBakeOutputEx& BakeOutput = BakeOutputs[0];
+	TArray<FColor> BakedPixels = MoveTemp(BakeOutput.PropertyData.FindChecked(Property));
+
+	// NOTE: the emissive value of each pixel will be incremented by 10 each time a triangle is drawn on it.
+	// Therefore a value of 0.0 indicates an unreferenced pixel, a value of 10.0 indicates a uniquely referenced pixel,
+	// and a value of 20+ indicates an overlapping pixel.
+	// The value may differ somewhat from what was specified in the material due to conversions (color-space, gamma, float => uint etc),
+	// so to be safe we set the limit at 15 instead of 10 to prevent incorrectly flagging some pixels as overlapping.
+	const float EmissiveThreshold = 15.0f;
+
+	uint32 ProcessedPixels = 0;
+	uint32 OverlappingPixels = 0;
+
+	if (BakeOutput.EmissiveScale > EmissiveThreshold)
+	{
+		const FColor Magenta(255, 0, 255, 255);
+		const uint8 ColorThreshold = FMath::RoundToInt((EmissiveThreshold / BakeOutput.EmissiveScale) * 255.0f); // This should never overflow since we know that the quotient will be <= 1.0
+
+		bool bIsMagenta;
+		bool bIsOverlapping;
+
+		for (const FColor& Pixel: BakedPixels)
+		{
+			bIsMagenta = Pixel == Magenta;
+			bIsOverlapping = Pixel.G > ColorThreshold;
+
+			ProcessedPixels += !bIsMagenta ? 1 : 0;
+			OverlappingPixels += !bIsMagenta && bIsOverlapping ? 1 : 0;
+		}
+	}
+
+	return ProcessedPixels > 0
+		? (static_cast<float>(OverlappingPixels) / ProcessedPixels) * 100.0f
+		: 0.0f;
+}
+
 const UMaterialInterface* FGLTFMaterialUtility::GetInterface(const UMaterialInterface* Material)
 {
 	return Material;
diff --git a/Engine/Plugins/Enterprise/GLTFExporter/MigratedFromGit/Source/GLTFExporter/Private/Converters/GLTFMaterialUtility.h b/Engine/Plugins/Enterprise/GLTFExporter/MigratedFromGit/Source/GLTFExporter/Private/Converters/GLTFMaterialUtility.h
--- a/Engine/Plugins/Enterprise/GLTFExporter/MigratedFromGit/Source/GLTFExporter/Private/Converters/GLTFMaterialUtility.h
+++ b/Engine/Plugins/Enterprise/GLTFExporter/MigratedFromGit/Source/GLTFExporter/Private/Converters/GLTFMaterialUtility.h
@@ -77,6 +77,8 @@ struct FGLTFMaterialUtility
 
 	static void AnalyzeMaterialProperty(const UMaterialInterface* Material, const FMaterialPropertyEx& InProperty, FGLTFMaterialAnalysis& OutAnalysis);
 
+	static float CalcOverlappingUVPercentage(int32 TexCoord, const FIntPoint& OutputSize, const FMeshDescription& MeshDescription, const TArray<int32>& MeshSectionIndices);
+
 	static const UMaterialInterface* GetInterface(const UMaterialInterface* Material);
 	static const UMaterialInterface* GetInterface(const FStaticMaterial& Material);
 	static const UMaterialInterface* GetInterface(const FSkeletalMaterial& Material);
diff --git a/Engine/Plugins/Enterprise/GLTFExporter/MigratedFromGit/Source/GLTFExporter/Private/Tasks/GLTFMaterialTasks.cpp b/Engine/Plugins/Enterprise/GLTFExporter/MigratedFromGit/Source/GLTFExporter/Private/Tasks/GLTFMaterialTasks.cpp
--- a/Engine/Plugins/Enterprise/GLTFExporter/MigratedFromGit/Source/GLTFExporter/Private/Tasks/GLTFMaterialTasks.cpp
+++ b/Engine/Plugins/Enterprise/GLTFExporter/MigratedFromGit/Source/GLTFExporter/Private/Tasks/GLTFMaterialTasks.cpp
@@ -1226,6 +1226,19 @@ FGLTFPropertyBakeOutput FGLTFMaterialTask::BakeMaterialProperty(const FMaterialP
 	else
 	{
 		OutTexCoord = MeshData->TexCoord;
+
+		const float MaxOverlapPercentage = 1.0f;
+		const float OverlapPercentage = FGLTFMaterialUtility::CalcOverlappingUVPercentage(MeshData->TexCoord, TextureSize, MeshData->Description, SectionIndices);
+
+		if (OverlapPercentage > MaxOverlapPercentage)
+		{
+			Builder.AddWarningMessage(FString::Printf(
+				TEXT("%s for material %s is baked using mesh data, but TexCoord %d is overlapping (%.2f%%) and may produce incorrect results"),
+				*Property.ToString(),
+				*Material->GetName(),
+				OutTexCoord,
+				OverlapPercentage));
+		}
 	}
 
 	// TODO: add support for calculating the ideal resolution to use for baking based on connected (texture) nodes
