[GLTFExporter] refactor: add TODO for possible optimization when calculating overlapping uv's

commit b6eac4bd39452b07d937da14ee4cab16e67f362f
Date: Tue, 17 Nov 2020 14:35:13 +0100

    refactor: add TODO for possible optimization when calculating overlapping uv's
    (cherry picked from commit 37059068f74e95ed8c84b712be7f8ba9ae32b413​)

#skipci

diff --git a/Engine/Plugins/Enterprise/GLTFExporter/MigratedFromGit/Source/GLTFExporter/Private/Tasks/GLTFMaterialTasks.cpp b/Engine/Plugins/Enterprise/GLTFExporter/MigratedFromGit/Source/GLTFExporter/Private/Tasks/GLTFMaterialTasks.cpp
--- a/Engine/Plugins/Enterprise/GLTFExporter/MigratedFromGit/Source/GLTFExporter/Private/Tasks/GLTFMaterialTasks.cpp
+++ b/Engine/Plugins/Enterprise/GLTFExporter/MigratedFromGit/Source/GLTFExporter/Private/Tasks/GLTFMaterialTasks.cpp
@@ -1227,6 +1227,9 @@ FGLTFPropertyBakeOutput FGLTFMaterialTask::BakeMaterialProperty(const FMaterialP
 	{
 		OutTexCoord = MeshData->TexCoord;
 
+		// TODO: use a smaller texture size when calculating overlapping uv's? The results are more correct when
+		// using the exact same size as the baked texture, but it also incurs an extra processing cost.
+
 		const float MaxOverlapPercentage = 1.0f;
 		const float OverlapPercentage = FGLTFMaterialUtility::CalcOverlappingUVPercentage(MeshData->TexCoord, TextureSize, MeshData->Description, SectionIndices);
 
