[GLTFExporter] refactor: improve warning for overlapping UVs and add more TODOs for potential improvements

commit 29b4a22346b5d7dfc36890450ffffac7f721d1a0
Date: Tue, 17 Nov 2020 15:01:50 +0100

    refactor: improve warning for overlapping UVs and add more TODOs for potential improvements
    (cherry picked from commit 109706620373f574d6ccf85a943ff02540dbfbb3​)

#skipci

diff --git a/Engine/Plugins/Enterprise/GLTFExporter/MigratedFromGit/Source/GLTFExporter/Private/Tasks/GLTFMaterialTasks.cpp b/Engine/Plugins/Enterprise/GLTFExporter/MigratedFromGit/Source/GLTFExporter/Private/Tasks/GLTFMaterialTasks.cpp
--- a/Engine/Plugins/Enterprise/GLTFExporter/MigratedFromGit/Source/GLTFExporter/Private/Tasks/GLTFMaterialTasks.cpp
+++ b/Engine/Plugins/Enterprise/GLTFExporter/MigratedFromGit/Source/GLTFExporter/Private/Tasks/GLTFMaterialTasks.cpp
@@ -1230,15 +1230,20 @@ FGLTFPropertyBakeOutput FGLTFMaterialTask::BakeMaterialProperty(const FMaterialP
 		// TODO: use a smaller texture size when calculating overlapping uv's? The results are more correct when
 		// using the exact same size as the baked texture, but it also incurs an extra processing cost.
 
+		// TODO: only calculate overlap once for the material (and perhaps once per baked texture size) instead of once for each baked property
+
+		// TODO: report overlapping uv's once for the material instead of once per property?
+
 		const float MaxOverlapPercentage = 1.0f;
 		const float OverlapPercentage = FGLTFMaterialUtility::CalcOverlappingUVPercentage(MeshData->TexCoord, TextureSize, MeshData->Description, SectionIndices);
 
 		if (OverlapPercentage > MaxOverlapPercentage)
 		{
 			Builder.AddWarningMessage(FString::Printf(
-				TEXT("%s for material %s is baked using mesh data, but TexCoord %d is overlapping (%.2f%%) and may produce incorrect results"),
+				TEXT("%s for material %s is baked using mesh data from %s, but TexCoord %d is overlapping (%.2f%%) and may produce incorrect results"),
 				*Property.ToString(),
 				*Material->GetName(),
+				*MeshData->Name,
 				OutTexCoord,
 				OverlapPercentage));
 		}
