[GLTFExporter] refactor: re-use UV overlap calculations as much as possible

commit 869249eec41c27117853d56a46c4a01fc1cbc928
Date: Wed, 18 Nov 2020 09:38:53 +0100

    refactor: re-use UV overlap calculations as much as possible
    (cherry picked from commit 7d7d3aa974a7fd71b5b8738d624b0b41699d5423​)

#skipci

diff --git a/Engine/Plugins/Enterprise/GLTFExporter/MigratedFromGit/Source/GLTFExporter/Private/Builders/GLTFConvertBuilder.h b/Engine/Plugins/Enterprise/GLTFExporter/MigratedFromGit/Source/GLTFExporter/Private/Builders/GLTFConvertBuilder.h
--- a/Engine/Plugins/Enterprise/GLTFExporter/MigratedFromGit/Source/GLTFExporter/Private/Builders/GLTFConvertBuilder.h
+++ b/Engine/Plugins/Enterprise/GLTFExporter/MigratedFromGit/Source/GLTFExporter/Private/Builders/GLTFConvertBuilder.h
@@ -20,6 +20,7 @@
 #include "Converters/GLTFLightMapConverters.h"
 #include "Converters/GLTFHotspotConverters.h"
 #include "Converters/GLTFSkySphereConverters.h"
+#include "Converters/GLTFUVAnalysisConverters.h"
 
 class FGLTFConvertBuilder : public FGLTFImageBuilder
 {
@@ -34,6 +35,7 @@ public:
 	// TODO: find a better place for this types of indirect converters
 	FGLTFStaticMeshDataConverter StaticMeshDataConverter;
 	FGLTFSkeletalMeshDataConverter SkeletalMeshDataConverter;
+	FGLTFUVAnalysisConverter UVAnalysisConverter;
 
 	FGLTFJsonAccessorIndex GetOrAddPositionAccessor(const FGLTFMeshSection* MeshSection, const FPositionVertexBuffer* VertexBuffer);
 	FGLTFJsonAccessorIndex GetOrAddColorAccessor(const FGLTFMeshSection* MeshSection, const FColorVertexBuffer* VertexBuffer);
diff --git a/Engine/Plugins/Enterprise/GLTFExporter/MigratedFromGit/Source/GLTFExporter/Private/Converters/GLTFUVAnalysis.h b/Engine/Plugins/Enterprise/GLTFExporter/MigratedFromGit/Source/GLTFExporter/Private/Converters/GLTFUVAnalysis.h
new file mode 100644
--- /dev/null
+++ b/Engine/Plugins/Enterprise/GLTFExporter/MigratedFromGit/Source/GLTFExporter/Private/Converters/GLTFUVAnalysis.h
@@ -0,0 +1,15 @@
+// Copyright Epic Games, Inc. All Rights Reserved.
+
+#pragma once
+
+#include "CoreMinimal.h"
+
+struct FGLTFUVAnalysis
+{
+	FGLTFUVAnalysis(float OverlapPercentage)
+		: OverlapPercentage(OverlapPercentage)
+	{
+	}
+
+	const float OverlapPercentage;
+};
diff --git a/Engine/Plugins/Enterprise/GLTFExporter/MigratedFromGit/Source/GLTFExporter/Private/Converters/GLTFUVAnalysisConverters.cpp b/Engine/Plugins/Enterprise/GLTFExporter/MigratedFromGit/Source/GLTFExporter/Private/Converters/GLTFUVAnalysisConverters.cpp
new file mode 100644
--- /dev/null
+++ b/Engine/Plugins/Enterprise/GLTFExporter/MigratedFromGit/Source/GLTFExporter/Private/Converters/GLTFUVAnalysisConverters.cpp
@@ -0,0 +1,45 @@
+// Copyright Epic Games, Inc. All Rights Reserved.
+
+#include "Converters/GLTFUVAnalysisConverters.h"
+#include "Converters/GLTFMaterialUtility.h"
+#include "StaticMeshAttributes.h"
+
+void FGLTFUVAnalysisConverter::Sanitize(const FMeshDescription*& Description, const TArray<int32>& SectionIndices, int32& TexCoord)
+{
+	if (Description != nullptr)
+	{
+		const TVertexInstanceAttributesConstRef<FVector2D> VertexInstanceUVs =
+			Description->VertexInstanceAttributes().GetAttributesRef<FVector2D>(MeshAttribute::VertexInstance::TextureCoordinate);
+		const int32 TexCoordCount = VertexInstanceUVs.GetNumIndices();
+
+		if (TexCoord < 0 || TexCoord >= TexCoordCount)
+		{
+			Description = nullptr;
+		}
+
+		const int32 MinSectionIndex = FMath::Min(SectionIndices);
+		const int32 MaxSectionIndex = FMath::Max(SectionIndices);
+		const int32 SectionCount = Description->PolygonGroups().GetArraySize();
+
+		if (MinSectionIndex < 0 || MaxSectionIndex >= SectionCount)
+		{
+			Description = nullptr;
+		}
+	}
+}
+
+const FGLTFUVAnalysis* FGLTFUVAnalysisConverter::Convert(const FMeshDescription* Description, const TArray<int32> SectionIndices, int32 TexCoord)
+{
+	if (Description == nullptr)
+	{
+		// TODO: add warning?
+
+		return nullptr;
+	}
+
+	// TODO: investigate if the fixed size is high enough to properly calculate overlap
+	const FIntPoint TextureSize(512, 512);
+	const float OverlapPercentage = FGLTFMaterialUtility::CalcOverlappingUVPercentage(TexCoord, TextureSize, *Description, SectionIndices);
+
+	return Outputs.Add_GetRef(MakeUnique<FGLTFUVAnalysis>(OverlapPercentage)).Get();
+}
diff --git a/Engine/Plugins/Enterprise/GLTFExporter/MigratedFromGit/Source/GLTFExporter/Private/Converters/GLTFUVAnalysisConverters.h b/Engine/Plugins/Enterprise/GLTFExporter/MigratedFromGit/Source/GLTFExporter/Private/Converters/GLTFUVAnalysisConverters.h
new file mode 100644
--- /dev/null
+++ b/Engine/Plugins/Enterprise/GLTFExporter/MigratedFromGit/Source/GLTFExporter/Private/Converters/GLTFUVAnalysisConverters.h
@@ -0,0 +1,16 @@
+// Copyright Epic Games, Inc. All Rights Reserved.
+
+#pragma once
+
+#include "Converters/GLTFConverter.h"
+#include "Converters/GLTFUVAnalysis.h"
+#include "Engine.h"
+
+class FGLTFUVAnalysisConverter final : public TGLTFConverter<const FGLTFUVAnalysis*, const FMeshDescription*, const TArray<int32>, int32>
+{
+	TArray<TUniquePtr<FGLTFUVAnalysis>> Outputs;
+
+	virtual void Sanitize(const FMeshDescription*& Description, const TArray<int32>& SectionIndices, int32& TexCoord) override;
+
+	virtual const FGLTFUVAnalysis* Convert(const FMeshDescription* Description, const TArray<int32> SectionIndices, int32 TexCoord) override;
+};
diff --git a/Engine/Plugins/Enterprise/GLTFExporter/MigratedFromGit/Source/GLTFExporter/Private/Tasks/GLTFMaterialTasks.cpp b/Engine/Plugins/Enterprise/GLTFExporter/MigratedFromGit/Source/GLTFExporter/Private/Tasks/GLTFMaterialTasks.cpp
--- a/Engine/Plugins/Enterprise/GLTFExporter/MigratedFromGit/Source/GLTFExporter/Private/Tasks/GLTFMaterialTasks.cpp
+++ b/Engine/Plugins/Enterprise/GLTFExporter/MigratedFromGit/Source/GLTFExporter/Private/Tasks/GLTFMaterialTasks.cpp
@@ -1227,17 +1227,12 @@ FGLTFPropertyBakeOutput FGLTFMaterialTask::BakeMaterialProperty(const FMaterialP
 	{
 		OutTexCoord = MeshData->TexCoord;
 
-		// TODO: use a smaller texture size when calculating overlapping uv's? The results are more correct when
-		// using the exact same size as the baked texture, but it also incurs an extra processing cost.
-
-		// TODO: only calculate overlap once for the material (and perhaps once per baked texture size) instead of once for each baked property
-
 		// TODO: report overlapping uv's once for the material instead of once per property?
 
 		const float MaxOverlapPercentage = 1.0f;
-		const float OverlapPercentage = FGLTFMaterialUtility::CalcOverlappingUVPercentage(MeshData->TexCoord, TextureSize, MeshData->Description, SectionIndices);
+		const FGLTFUVAnalysis* UVAnalysis = Builder.UVAnalysisConverter.GetOrAdd(&MeshData->Description, SectionIndices, MeshData->TexCoord);
 
-		if (OverlapPercentage > MaxOverlapPercentage)
+		if (UVAnalysis != nullptr && UVAnalysis->OverlapPercentage > MaxOverlapPercentage)
 		{
 			Builder.AddWarningMessage(FString::Printf(
 				TEXT("%s for material %s is baked using mesh data from %s, but TexCoord %d is overlapping (%.2f%%) and may produce incorrect results"),
@@ -1245,7 +1240,7 @@ FGLTFPropertyBakeOutput FGLTFMaterialTask::BakeMaterialProperty(const FMaterialP
 				*Material->GetName(),
 				*MeshData->Name,
 				OutTexCoord,
-				OverlapPercentage));
+				UVAnalysis->OverlapPercentage));
 		}
 	}
 
