[GLTFExporter] refactor: make function for calculating UV overlap a part of the designated converter

commit fb6be99b98ad945f1a686f0c88db864cf511c9d0
Date: Wed, 18 Nov 2020 09:50:49 +0100

    refactor: make function for calculating UV overlap a part of the designated converter
    (cherry picked from commit 2dbfb949910c78e242e8ab14df05a0d9bb443327​)

#skipci

diff --git a/Engine/Plugins/Enterprise/GLTFExporter/MigratedFromGit/Source/GLTFExporter/Private/Converters/GLTFMaterialUtility.cpp b/Engine/Plugins/Enterprise/GLTFExporter/MigratedFromGit/Source/GLTFExporter/Private/Converters/GLTFMaterialUtility.cpp
--- a/Engine/Plugins/Enterprise/GLTFExporter/MigratedFromGit/Source/GLTFExporter/Private/Converters/GLTFMaterialUtility.cpp
+++ b/Engine/Plugins/Enterprise/GLTFExporter/MigratedFromGit/Source/GLTFExporter/Private/Converters/GLTFMaterialUtility.cpp
@@ -460,71 +460,6 @@ void FGLTFMaterialUtility::AnalyzeMaterialProperty(const UMaterialInterface* InM
 	UGLTFMaterialAnalyzer::AnalyzeMaterialPropertyEx(InMaterial, InProperty.Type, InProperty.CustomOutput.ToString(), OutAnalysis);
 }
 
-float FGLTFMaterialUtility::CalcOverlappingUVPercentage(int32 TexCoord, const FIntPoint& OutputSize, const FMeshDescription& MeshDescription, const TArray<int32>& MeshSectionIndices)
-{
-	static UMaterialInterface* Material = LoadObject<UMaterialInterface>(nullptr, TEXT("/GLTFExporter/Materials/OverlappingUVs.OverlappingUVs"));
-	check(Material);
-
-	const FMaterialPropertyEx Property = MP_EmissiveColor;
-
-	FMeshData MeshSet;
-	MeshSet.TextureCoordinateBox = { { 0.0f, 0.0f }, { 1.0f, 1.0f } };
-	MeshSet.TextureCoordinateIndex = TexCoord;
-	MeshSet.RawMeshDescription = const_cast<FMeshDescription*>(&MeshDescription);
-	MeshSet.MaterialIndices = MeshSectionIndices; // NOTE: MaterialIndices is actually section indices
-
-	FMaterialDataEx MatSet;
-	MatSet.Material = Material;
-	MatSet.PropertySizes.Add(Property, OutputSize);
-	MatSet.BlendMode = BLEND_Additive;
-	MatSet.bPerformBorderSmear = false;
-
-	TArray<FMeshData*> MeshSettings;
-	TArray<FMaterialDataEx*> MatSettings;
-	MeshSettings.Add(&MeshSet);
-	MatSettings.Add(&MatSet);
-
-	TArray<FBakeOutputEx> BakeOutputs;
-	IMaterialBakingModule& Module = FModuleManager::Get().LoadModuleChecked<IMaterialBakingModule>("GLTFMaterialBaking");
-
-	Module.BakeMaterials(MatSettings, MeshSettings, BakeOutputs);
-
-	FBakeOutputEx& BakeOutput = BakeOutputs[0];
-	TArray<FColor> BakedPixels = MoveTemp(BakeOutput.PropertyData.FindChecked(Property));
-
-	// NOTE: the emissive value of each pixel will be incremented by 10 each time a triangle is drawn on it.
-	// Therefore a value of 0.0 indicates an unreferenced pixel, a value of 10.0 indicates a uniquely referenced pixel,
-	// and a value of 20+ indicates an overlapping pixel.
-	// The value may differ somewhat from what was specified in the material due to conversions (color-space, gamma, float => uint etc),
-	// so to be safe we set the limit at 15 instead of 10 to prevent incorrectly flagging some pixels as overlapping.
-	const float EmissiveThreshold = 15.0f;
-
-	uint32 ProcessedPixels = 0;
-	uint32 OverlappingPixels = 0;
-
-	if (BakeOutput.EmissiveScale > EmissiveThreshold)
-	{
-		const FColor Magenta(255, 0, 255, 255);
-		const uint8 ColorThreshold = FMath::RoundToInt((EmissiveThreshold / BakeOutput.EmissiveScale) * 255.0f); // This should never overflow since we know that the quotient will be <= 1.0
-
-		bool bIsMagenta;
-		bool bIsOverlapping;
-
-		for (const FColor& Pixel: BakedPixels)
-		{
-			bIsMagenta = Pixel == Magenta;
-			bIsOverlapping = Pixel.G > ColorThreshold;
-
-			ProcessedPixels += !bIsMagenta ? 1 : 0;
-			OverlappingPixels += !bIsMagenta && bIsOverlapping ? 1 : 0;
-		}
-	}
-
-	return ProcessedPixels > 0
-		? (static_cast<float>(OverlappingPixels) / ProcessedPixels) * 100.0f
-		: 0.0f;
-}
-
 const UMaterialInterface* FGLTFMaterialUtility::GetInterface(const UMaterialInterface* Material)
 {
 	return Material;
diff --git a/Engine/Plugins/Enterprise/GLTFExporter/MigratedFromGit/Source/GLTFExporter/Private/Converters/GLTFMaterialUtility.h b/Engine/Plugins/Enterprise/GLTFExporter/MigratedFromGit/Source/GLTFExporter/Private/Converters/GLTFMaterialUtility.h
--- a/Engine/Plugins/Enterprise/GLTFExporter/MigratedFromGit/Source/GLTFExporter/Private/Converters/GLTFMaterialUtility.h
+++ b/Engine/Plugins/Enterprise/GLTFExporter/MigratedFromGit/Source/GLTFExporter/Private/Converters/GLTFMaterialUtility.h
@@ -77,8 +77,6 @@ struct FGLTFMaterialUtility
 
 	static void AnalyzeMaterialProperty(const UMaterialInterface* Material, const FMaterialPropertyEx& InProperty, FGLTFMaterialAnalysis& OutAnalysis);
 
-	static float CalcOverlappingUVPercentage(int32 TexCoord, const FIntPoint& OutputSize, const FMeshDescription& MeshDescription, const TArray<int32>& MeshSectionIndices);
-
 	static const UMaterialInterface* GetInterface(const UMaterialInterface* Material);
 	static const UMaterialInterface* GetInterface(const FStaticMaterial& Material);
 	static const UMaterialInterface* GetInterface(const FSkeletalMaterial& Material);
diff --git a/Engine/Plugins/Enterprise/GLTFExporter/MigratedFromGit/Source/GLTFExporter/Private/Converters/GLTFUVAnalysisConverters.cpp b/Engine/Plugins/Enterprise/GLTFExporter/MigratedFromGit/Source/GLTFExporter/Private/Converters/GLTFUVAnalysisConverters.cpp
--- a/Engine/Plugins/Enterprise/GLTFExporter/MigratedFromGit/Source/GLTFExporter/Private/Converters/GLTFUVAnalysisConverters.cpp
+++ b/Engine/Plugins/Enterprise/GLTFExporter/MigratedFromGit/Source/GLTFExporter/Private/Converters/GLTFUVAnalysisConverters.cpp
@@ -3,6 +3,10 @@
 #include "Converters/GLTFUVAnalysisConverters.h"
 #include "Converters/GLTFMaterialUtility.h"
 #include "StaticMeshAttributes.h"
+#include "MeshDescription.h"
+#include "Modules/ModuleManager.h"
+#include "GLTFMaterialBaking/Public/IMaterialBakingModule.h"
+#include "GLTFMaterialBaking/Public/MaterialBakingStructures.h"
 
 void FGLTFUVAnalysisConverter::Sanitize(const FMeshDescription*& Description, const TArray<int32>& SectionIndices, int32& TexCoord)
 {
@@ -39,7 +43,76 @@ const FGLTFUVAnalysis* FGLTFUVAnalysisConverter::Convert(const FMeshDescription*
 
 	// TODO: investigate if the fixed size is high enough to properly calculate overlap
 	const FIntPoint TextureSize(512, 512);
-	const float OverlapPercentage = FGLTFMaterialUtility::CalcOverlappingUVPercentage(TexCoord, TextureSize, *Description, SectionIndices);
 
-	return Outputs.Add_GetRef(MakeUnique<FGLTFUVAnalysis>(OverlapPercentage)).Get();
+	return Outputs.Add_GetRef(MakeUnique<FGLTFUVAnalysis>(CalcOverlapPercentage(TexCoord, TextureSize, *Description, SectionIndices))).Get();
+}
+
+const UMaterialInterface* FGLTFUVAnalysisConverter::GetOverlapMaterial()
+{
+	static const UMaterialInterface* Material = LoadObject<UMaterialInterface>(nullptr, TEXT("/GLTFExporter/Materials/OverlappingUVs.OverlappingUVs"));
+	check(Material);
+	return Material;
+}
+
+float FGLTFUVAnalysisConverter::CalcOverlapPercentage(int32 TexCoord, const FIntPoint& OutputSize, const FMeshDescription& MeshDescription, const TArray<int32>& MeshSectionIndices)
+{
+	const UMaterialInterface* Material = GetOverlapMaterial();
+	const FMaterialPropertyEx Property = MP_EmissiveColor;
+
+	FMeshData MeshSet;
+	MeshSet.TextureCoordinateBox = { { 0.0f, 0.0f }, { 1.0f, 1.0f } };
+	MeshSet.TextureCoordinateIndex = TexCoord;
+	MeshSet.RawMeshDescription = const_cast<FMeshDescription*>(&MeshDescription);
+	MeshSet.MaterialIndices = MeshSectionIndices; // NOTE: MaterialIndices is actually section indices
+
+	FMaterialDataEx MatSet;
+	MatSet.Material = const_cast<UMaterialInterface*>(Material);
+	MatSet.PropertySizes.Add(Property, OutputSize);
+	MatSet.BlendMode = BLEND_Additive;
+	MatSet.bPerformBorderSmear = false;
+
+	TArray<FMeshData*> MeshSettings;
+	TArray<FMaterialDataEx*> MatSettings;
+	MeshSettings.Add(&MeshSet);
+	MatSettings.Add(&MatSet);
+
+	TArray<FBakeOutputEx> BakeOutputs;
+	IMaterialBakingModule& Module = FModuleManager::Get().LoadModuleChecked<IMaterialBakingModule>("GLTFMaterialBaking");
+
+	Module.BakeMaterials(MatSettings, MeshSettings, BakeOutputs);
+
+	FBakeOutputEx& BakeOutput = BakeOutputs[0];
+	TArray<FColor> BakedPixels = MoveTemp(BakeOutput.PropertyData.FindChecked(Property));
+
+	// NOTE: the emissive value of each pixel will be incremented by 10 each time a triangle is drawn on it.
+	// Therefore a value of 0.0 indicates an unreferenced pixel, a value of 10.0 indicates a uniquely referenced pixel,
+	// and a value of 20+ indicates an overlapping pixel.
+	// The value may differ somewhat from what was specified in the material due to conversions (color-space, gamma, float => uint etc),
+	// so to be safe we set the limit at 15 instead of 10 to prevent incorrectly flagging some pixels as overlapping.
+	const float EmissiveThreshold = 15.0f;
+
+	uint32 ProcessedPixels = 0;
+	uint32 OverlappingPixels = 0;
+
+	if (BakeOutput.EmissiveScale > EmissiveThreshold)
+	{
+		const FColor Magenta(255, 0, 255, 255);
+		const uint8 ColorThreshold = FMath::RoundToInt((EmissiveThreshold / BakeOutput.EmissiveScale) * 255.0f); // This should never overflow since we know that the quotient will be <= 1.0
+
+		bool bIsMagenta;
+		bool bIsOverlapping;
+
+		for (const FColor& Pixel: BakedPixels)
+		{
+			bIsMagenta = Pixel == Magenta;
+			bIsOverlapping = Pixel.G > ColorThreshold;
+
+			ProcessedPixels += !bIsMagenta ? 1 : 0;
+			OverlappingPixels += !bIsMagenta && bIsOverlapping ? 1 : 0;
+		}
+	}
+
+	return ProcessedPixels > 0
+		? (static_cast<float>(OverlappingPixels) / ProcessedPixels) * 100.0f
+		: 0.0f;
 }
diff --git a/Engine/Plugins/Enterprise/GLTFExporter/MigratedFromGit/Source/GLTFExporter/Private/Converters/GLTFUVAnalysisConverters.h b/Engine/Plugins/Enterprise/GLTFExporter/MigratedFromGit/Source/GLTFExporter/Private/Converters/GLTFUVAnalysisConverters.h
--- a/Engine/Plugins/Enterprise/GLTFExporter/MigratedFromGit/Source/GLTFExporter/Private/Converters/GLTFUVAnalysisConverters.h
+++ b/Engine/Plugins/Enterprise/GLTFExporter/MigratedFromGit/Source/GLTFExporter/Private/Converters/GLTFUVAnalysisConverters.h
@@ -13,4 +13,8 @@ class FGLTFUVAnalysisConverter final : public TGLTFConverter<const FGLTFUVAnalys
 	virtual void Sanitize(const FMeshDescription*& Description, const TArray<int32>& SectionIndices, int32& TexCoord) override;
 
 	virtual const FGLTFUVAnalysis* Convert(const FMeshDescription* Description, const TArray<int32> SectionIndices, int32 TexCoord) override;
+
+	static const UMaterialInterface* GetOverlapMaterial();
+
+	static float CalcOverlapPercentage(int32 TexCoord, const FIntPoint& OutputSize, const FMeshDescription& MeshDescription, const TArray<int32>& MeshSectionIndices);
 };
