[GLTFExporter] refactor: only warn about overlapping UVs once per converted material

commit 958b58ff2a0670b2f84cdff13673094d4a135dc2
Date: Wed, 18 Nov 2020 10:27:59 +0100

    refactor: only warn about overlapping UVs once per converted material
    (cherry picked from commit 9a94a8e654ecfd7b2d5ecdcb26106235d2416054​)

#skipci

diff --git a/Engine/Plugins/Enterprise/GLTFExporter/MigratedFromGit/Source/GLTFExporter/Private/Tasks/GLTFMaterialTasks.cpp b/Engine/Plugins/Enterprise/GLTFExporter/MigratedFromGit/Source/GLTFExporter/Private/Tasks/GLTFMaterialTasks.cpp
--- a/Engine/Plugins/Enterprise/GLTFExporter/MigratedFromGit/Source/GLTFExporter/Private/Tasks/GLTFMaterialTasks.cpp
+++ b/Engine/Plugins/Enterprise/GLTFExporter/MigratedFromGit/Source/GLTFExporter/Private/Tasks/GLTFMaterialTasks.cpp
@@ -158,6 +158,28 @@ void FGLTFMaterialTask::Complete()
 			}
 		}
 	}
+
+	// Warn if properties baked using mesh data are using overlapping UVs
+	if (MeshData != nullptr && MeshDataBakedProperties.Num() > 0)
+	{
+		// TODO: when a mesh-component is used, MeshData will be unique for each converted component even if
+		// they use the same mesh internally. This leads to an excessive number of unique UV analyses.
+		// We should look into creating a separate FGLTFMeshData for the mesh itself in these cases, and pass
+		// that as parameter to the UV analysis converter instead of MeshData.
+
+		const float MaxOverlapPercentage = 1.0f;
+		const FGLTFUVAnalysis* UVAnalysis = Builder.UVAnalysisConverter.GetOrAdd(&MeshData->Description, SectionIndices, MeshData->TexCoord);
+
+		if (UVAnalysis != nullptr && UVAnalysis->OverlapPercentage > MaxOverlapPercentage)
+		{
+			Builder.AddWarningMessage(FString::Printf(
+				TEXT("Material %s is baked using TexCoord %d of %s which has %.2f%% overlapping UVs and may produce incorrect results"),
+				*Material->GetName(),
+				MeshData->TexCoord,
+				*MeshData->Name,
+				UVAnalysis->OverlapPercentage));
+		}
+	}
 }
 
 void FGLTFMaterialTask::ConvertAlphaMode(EGLTFJsonAlphaMode& OutAlphaMode, EGLTFJsonBlendMode& OutBlendMode) const
@@ -267,7 +289,7 @@ void FGLTFMaterialTask::ConvertShadingModel(EGLTFJsonShadingModel& OutShadingMod
 	}
 }
 
-bool FGLTFMaterialTask::TryGetBaseColorAndOpacity(FGLTFJsonPBRMetallicRoughness& OutPBRParams, const FMaterialPropertyEx& BaseColorProperty, const FMaterialPropertyEx& OpacityProperty) const
+bool FGLTFMaterialTask::TryGetBaseColorAndOpacity(FGLTFJsonPBRMetallicRoughness& OutPBRParams, const FMaterialPropertyEx& BaseColorProperty, const FMaterialPropertyEx& OpacityProperty)
 {
 	const bool bIsBaseColorConstant = TryGetConstantColor(OutPBRParams.BaseColorFactor, BaseColorProperty);
 	const bool bIsOpacityConstant = TryGetConstantScalar(OutPBRParams.BaseColorFactor.A, OpacityProperty);
@@ -386,7 +408,7 @@ bool FGLTFMaterialTask::TryGetBaseColorAndOpacity(FGLTFJsonPBRMetallicRoughness&
 	return true;
 }
 
-bool FGLTFMaterialTask::TryGetMetallicAndRoughness(FGLTFJsonPBRMetallicRoughness& OutPBRParams, const FMaterialPropertyEx& MetallicProperty, const FMaterialPropertyEx& RoughnessProperty) const
+bool FGLTFMaterialTask::TryGetMetallicAndRoughness(FGLTFJsonPBRMetallicRoughness& OutPBRParams, const FMaterialPropertyEx& MetallicProperty, const FMaterialPropertyEx& RoughnessProperty)
 {
 	const bool bIsMetallicConstant = TryGetConstantScalar(OutPBRParams.MetallicFactor, MetallicProperty);
 	const bool bIsRoughnessConstant = TryGetConstantScalar(OutPBRParams.RoughnessFactor, RoughnessProperty);
@@ -503,7 +525,7 @@ bool FGLTFMaterialTask::TryGetMetallicAndRoughness(FGLTFJsonPBRMetallicRoughness
 	return true;
 }
 
-bool FGLTFMaterialTask::TryGetClearCoatRoughness(FGLTFJsonClearCoatExtension& OutExtParams, const FMaterialPropertyEx& IntensityProperty, const FMaterialPropertyEx& RoughnessProperty) const
+bool FGLTFMaterialTask::TryGetClearCoatRoughness(FGLTFJsonClearCoatExtension& OutExtParams, const FMaterialPropertyEx& IntensityProperty, const FMaterialPropertyEx& RoughnessProperty)
 {
 	const bool bIsIntensityConstant = TryGetConstantScalar(OutExtParams.ClearCoatFactor, IntensityProperty);
 	const bool bIsRoughnessConstant = TryGetConstantScalar(OutExtParams.ClearCoatRoughnessFactor, RoughnessProperty);
@@ -624,7 +646,7 @@ bool FGLTFMaterialTask::TryGetClearCoatRoughness(FGLTFJsonClearCoatExtension& Ou
 	return true;
 }
 
-bool FGLTFMaterialTask::TryGetEmissive(FGLTFJsonMaterial& JsonMaterial, const FMaterialPropertyEx& EmissiveProperty) const
+bool FGLTFMaterialTask::TryGetEmissive(FGLTFJsonMaterial& JsonMaterial, const FMaterialPropertyEx& EmissiveProperty)
 {
 	// TODO: right now we allow EmissiveFactor to be > 1.0 to support very bright emission, although it's not valid according to the glTF standard.
 	// We may want to change this behaviour and store factors above 1.0 using a custom extension instead.
@@ -1055,7 +1077,7 @@ bool FGLTFMaterialTask::TryGetSourceTexture(const UTexture2D*& OutTexture, int32
 	return false;
 }
 
-bool FGLTFMaterialTask::TryGetBakedMaterialProperty(FGLTFJsonTextureInfo& OutTexInfo, FGLTFJsonColor3& OutConstant, const FMaterialPropertyEx& Property, const FString& PropertyName) const
+bool FGLTFMaterialTask::TryGetBakedMaterialProperty(FGLTFJsonTextureInfo& OutTexInfo, FGLTFJsonColor3& OutConstant, const FMaterialPropertyEx& Property, const FString& PropertyName)
 {
 	if (Builder.ExportOptions->BakeMaterialInputs == EGLTFMaterialBakeMode::None)
 	{
@@ -1083,7 +1105,7 @@ bool FGLTFMaterialTask::TryGetBakedMaterialProperty(FGLTFJsonTextureInfo& OutTex
 	return false;
 }
 
-bool FGLTFMaterialTask::TryGetBakedMaterialProperty(FGLTFJsonTextureInfo& OutTexInfo, FGLTFJsonColor4& OutConstant, const FMaterialPropertyEx& Property, const FString& PropertyName) const
+bool FGLTFMaterialTask::TryGetBakedMaterialProperty(FGLTFJsonTextureInfo& OutTexInfo, FGLTFJsonColor4& OutConstant, const FMaterialPropertyEx& Property, const FString& PropertyName)
 {
 	if (Builder.ExportOptions->BakeMaterialInputs == EGLTFMaterialBakeMode::None)
 	{
@@ -1111,7 +1133,7 @@ bool FGLTFMaterialTask::TryGetBakedMaterialProperty(FGLTFJsonTextureInfo& OutTex
 	return false;
 }
 
-inline bool FGLTFMaterialTask::TryGetBakedMaterialProperty(FGLTFJsonTextureInfo& OutTexInfo, float& OutConstant, const FMaterialPropertyEx& Property, const FString& PropertyName) const
+inline bool FGLTFMaterialTask::TryGetBakedMaterialProperty(FGLTFJsonTextureInfo& OutTexInfo, float& OutConstant, const FMaterialPropertyEx& Property, const FString& PropertyName)
 {
 	if (Builder.ExportOptions->BakeMaterialInputs == EGLTFMaterialBakeMode::None)
 	{
@@ -1139,7 +1161,7 @@ inline bool FGLTFMaterialTask::TryGetBakedMaterialProperty(FGLTFJsonTextureInfo&
 	return false;
 }
 
-bool FGLTFMaterialTask::TryGetBakedMaterialProperty(FGLTFJsonTextureInfo& OutTexInfo, const FMaterialPropertyEx& Property, const FString& PropertyName) const
+bool FGLTFMaterialTask::TryGetBakedMaterialProperty(FGLTFJsonTextureInfo& OutTexInfo, const FMaterialPropertyEx& Property, const FString& PropertyName)
 {
 	if (Builder.ExportOptions->BakeMaterialInputs == EGLTFMaterialBakeMode::None)
 	{
@@ -1195,13 +1217,13 @@ bool FGLTFMaterialTask::TryGetBakedMaterialProperty(FGLTFJsonTextureInfo& OutTex
 	return true;
 }
 
-FGLTFPropertyBakeOutput FGLTFMaterialTask::BakeMaterialProperty(const FMaterialPropertyEx& Property, int32& OutTexCoord, bool bCopyAlphaFromRedChannel) const
+FGLTFPropertyBakeOutput FGLTFMaterialTask::BakeMaterialProperty(const FMaterialPropertyEx& Property, int32& OutTexCoord, bool bCopyAlphaFromRedChannel)
 {
 	const FIntPoint TextureSize = Builder.GetBakeSizeForMaterialProperty(Material, Property);
 	return BakeMaterialProperty(Property, OutTexCoord, TextureSize, bCopyAlphaFromRedChannel);
 }
 
-FGLTFPropertyBakeOutput FGLTFMaterialTask::BakeMaterialProperty(const FMaterialPropertyEx& Property, int32& OutTexCoord, const FIntPoint& TextureSize, bool bCopyAlphaFromRedChannel) const
+FGLTFPropertyBakeOutput FGLTFMaterialTask::BakeMaterialProperty(const FMaterialPropertyEx& Property, int32& OutTexCoord, const FIntPoint& TextureSize, bool bCopyAlphaFromRedChannel)
 {
 	if (MeshData == nullptr)
 	{
@@ -1226,22 +1248,7 @@ FGLTFPropertyBakeOutput FGLTFMaterialTask::BakeMaterialProperty(const FMaterialP
 	else
 	{
 		OutTexCoord = MeshData->TexCoord;
-
-		// TODO: report overlapping uv's once for the material instead of once per property?
-
-		const float MaxOverlapPercentage = 1.0f;
-		const FGLTFUVAnalysis* UVAnalysis = Builder.UVAnalysisConverter.GetOrAdd(&MeshData->Description, SectionIndices, MeshData->TexCoord);
-
-		if (UVAnalysis != nullptr && UVAnalysis->OverlapPercentage > MaxOverlapPercentage)
-		{
-			Builder.AddWarningMessage(FString::Printf(
-				TEXT("%s for material %s is baked using mesh data from %s, but TexCoord %d is overlapping (%.2f%%) and may produce incorrect results"),
-				*Property.ToString(),
-				*Material->GetName(),
-				*MeshData->Name,
-				OutTexCoord,
-				UVAnalysis->OverlapPercentage));
-		}
+		MeshDataBakedProperties.AddUnique(Property);
 	}
 
 	// TODO: add support for calculating the ideal resolution to use for baking based on connected (texture) nodes
diff --git a/Engine/Plugins/Enterprise/GLTFExporter/MigratedFromGit/Source/GLTFExporter/Private/Tasks/GLTFMaterialTasks.h b/Engine/Plugins/Enterprise/GLTFExporter/MigratedFromGit/Source/GLTFExporter/Private/Tasks/GLTFMaterialTasks.h
--- a/Engine/Plugins/Enterprise/GLTFExporter/MigratedFromGit/Source/GLTFExporter/Private/Tasks/GLTFMaterialTasks.h
+++ b/Engine/Plugins/Enterprise/GLTFExporter/MigratedFromGit/Source/GLTFExporter/Private/Tasks/GLTFMaterialTasks.h
@@ -39,13 +39,15 @@ private:
 	const TArray<int32> SectionIndices;
 	const FGLTFJsonMaterialIndex MaterialIndex;
 
+	TArray<FMaterialPropertyEx> MeshDataBakedProperties;
+
 	void ConvertAlphaMode(EGLTFJsonAlphaMode& OutAlphaMode, EGLTFJsonBlendMode& OutBlendMode) const;
 	void ConvertShadingModel(EGLTFJsonShadingModel& OutShadingModel) const;
 
-	bool TryGetBaseColorAndOpacity(FGLTFJsonPBRMetallicRoughness& OutPBRParams, const FMaterialPropertyEx& BaseColorProperty, const FMaterialPropertyEx& OpacityProperty) const;
-	bool TryGetMetallicAndRoughness(FGLTFJsonPBRMetallicRoughness& OutPBRParams, const FMaterialPropertyEx& MetallicProperty, const FMaterialPropertyEx& RoughnessProperty) const;
-	bool TryGetClearCoatRoughness(FGLTFJsonClearCoatExtension& OutExtParams, const FMaterialPropertyEx& IntensityProperty, const FMaterialPropertyEx& RoughnessProperty) const;
-	bool TryGetEmissive(FGLTFJsonMaterial& JsonMaterial, const FMaterialPropertyEx& EmissiveProperty) const;
+	bool TryGetBaseColorAndOpacity(FGLTFJsonPBRMetallicRoughness& OutPBRParams, const FMaterialPropertyEx& BaseColorProperty, const FMaterialPropertyEx& OpacityProperty);
+	bool TryGetMetallicAndRoughness(FGLTFJsonPBRMetallicRoughness& OutPBRParams, const FMaterialPropertyEx& MetallicProperty, const FMaterialPropertyEx& RoughnessProperty);
+	bool TryGetClearCoatRoughness(FGLTFJsonClearCoatExtension& OutExtParams, const FMaterialPropertyEx& IntensityProperty, const FMaterialPropertyEx& RoughnessProperty);
+	bool TryGetEmissive(FGLTFJsonMaterial& JsonMaterial, const FMaterialPropertyEx& EmissiveProperty);
 
 	bool IsPropertyNonDefault(const FMaterialPropertyEx& Property) const;
 	bool TryGetConstantColor(FGLTFJsonColor3& OutValue, const FMaterialPropertyEx& Property) const;
@@ -56,13 +58,13 @@ private:
 	bool TryGetSourceTexture(FGLTFJsonTextureInfo& OutTexInfo, const FMaterialPropertyEx& Property, const TArray<FLinearColor>& AllowedMasks = {}) const;
 	bool TryGetSourceTexture(const UTexture2D*& OutTexture, int32& OutTexCoord, FGLTFJsonTextureTransform& OutTransform, const FMaterialPropertyEx& Property, const TArray<FLinearColor>& AllowedMasks = {}) const;
 
-	bool TryGetBakedMaterialProperty(FGLTFJsonTextureInfo& OutTexInfo, FGLTFJsonColor3& OutConstant, const FMaterialPropertyEx& Property, const FString& PropertyName) const;
-	bool TryGetBakedMaterialProperty(FGLTFJsonTextureInfo& OutTexInfo, FGLTFJsonColor4& OutConstant, const FMaterialPropertyEx& Property, const FString& PropertyName) const;
-	bool TryGetBakedMaterialProperty(FGLTFJsonTextureInfo& OutTexInfo, float& OutConstant, const FMaterialPropertyEx& Property, const FString& PropertyName) const;
-	bool TryGetBakedMaterialProperty(FGLTFJsonTextureInfo& OutTexInfo, const FMaterialPropertyEx& Property, const FString& PropertyName) const;
+	bool TryGetBakedMaterialProperty(FGLTFJsonTextureInfo& OutTexInfo, FGLTFJsonColor3& OutConstant, const FMaterialPropertyEx& Property, const FString& PropertyName);
+	bool TryGetBakedMaterialProperty(FGLTFJsonTextureInfo& OutTexInfo, FGLTFJsonColor4& OutConstant, const FMaterialPropertyEx& Property, const FString& PropertyName);
+	bool TryGetBakedMaterialProperty(FGLTFJsonTextureInfo& OutTexInfo, float& OutConstant, const FMaterialPropertyEx& Property, const FString& PropertyName);
+	bool TryGetBakedMaterialProperty(FGLTFJsonTextureInfo& OutTexInfo, const FMaterialPropertyEx& Property, const FString& PropertyName);
 
-	FGLTFPropertyBakeOutput BakeMaterialProperty(const FMaterialPropertyEx& Property, int32& OutTexCoord, bool bCopyAlphaFromRedChannel = false) const;
-	FGLTFPropertyBakeOutput BakeMaterialProperty(const FMaterialPropertyEx& Property, int32& OutTexCoord, const FIntPoint& TextureSize, bool bCopyAlphaFromRedChannel = false) const;
+	FGLTFPropertyBakeOutput BakeMaterialProperty(const FMaterialPropertyEx& Property, int32& OutTexCoord, bool bCopyAlphaFromRedChannel = false);
+	FGLTFPropertyBakeOutput BakeMaterialProperty(const FMaterialPropertyEx& Property, int32& OutTexCoord, const FIntPoint& TextureSize, bool bCopyAlphaFromRedChannel = false);
 
 	bool StoreBakedPropertyTexture(FGLTFJsonTextureInfo& OutTexInfo, const FGLTFPropertyBakeOutput& PropertyBakeOutput, const FString& PropertyName) const;
 
