[GLTFExporter] refactor: use BlendMode from the material when calculating overlapping UVs

commit dba97403d8e576ff5722a57cd9b34038ac811e3d
Date: Wed, 18 Nov 2020 10:33:11 +0100

    refactor: use BlendMode from the material when calculating overlapping UVs
    (cherry picked from commit 059d0224ff2a53c34dfd0b829371c715b5d6e405​)

#skipci

diff --git a/Engine/Plugins/Enterprise/GLTFExporter/MigratedFromGit/Source/GLTFExporter/Private/Converters/GLTFUVAnalysisConverters.cpp b/Engine/Plugins/Enterprise/GLTFExporter/MigratedFromGit/Source/GLTFExporter/Private/Converters/GLTFUVAnalysisConverters.cpp
--- a/Engine/Plugins/Enterprise/GLTFExporter/MigratedFromGit/Source/GLTFExporter/Private/Converters/GLTFUVAnalysisConverters.cpp
+++ b/Engine/Plugins/Enterprise/GLTFExporter/MigratedFromGit/Source/GLTFExporter/Private/Converters/GLTFUVAnalysisConverters.cpp
@@ -68,7 +68,7 @@ float FGLTFUVAnalysisConverter::CalcOverlapPercentage(int32 TexCoord, const FInt
 	FMaterialDataEx MatSet;
 	MatSet.Material = const_cast<UMaterialInterface*>(Material);
 	MatSet.PropertySizes.Add(Property, OutputSize);
-	MatSet.BlendMode = BLEND_Additive;
+	MatSet.BlendMode = Material->GetBlendMode();
 	MatSet.bPerformBorderSmear = false;
 
 	TArray<FMeshData*> MeshSettings;
