[GLTFExporter] refactor: simplify GLTFUVAnalysisConverter to avoid unnecessary usage of memory heap

commit e536d1cf77013f43e145a1c0a4f5dcbc2bf53654
Date: Tue, 24 Nov 2020 09:33:14 +0100

    refactor: simplify GLTFUVAnalysisConverter to avoid unnecessary usage of memory heap
    (cherry picked from commit 94b47b616c8a17cfc66575d89ca03f98def72ad8​)

#skipci

diff --git a/Engine/Plugins/Enterprise/GLTFExporter/MigratedFromGit/Source/GLTFExporter/Private/Converters/GLTFUVAnalysisConverters.cpp b/Engine/Plugins/Enterprise/GLTFExporter/MigratedFromGit/Source/GLTFExporter/Private/Converters/GLTFUVAnalysisConverters.cpp
--- a/Engine/Plugins/Enterprise/GLTFExporter/MigratedFromGit/Source/GLTFExporter/Private/Converters/GLTFUVAnalysisConverters.cpp
+++ b/Engine/Plugins/Enterprise/GLTFExporter/MigratedFromGit/Source/GLTFExporter/Private/Converters/GLTFUVAnalysisConverters.cpp
@@ -32,19 +32,19 @@ void FGLTFUVAnalysisConverter::Sanitize(const FMeshDescription*& Description, co
 	}
 }
 
-const FGLTFUVAnalysis* FGLTFUVAnalysisConverter::Convert(const FMeshDescription* Description, const TArray<int32> SectionIndices, int32 TexCoord)
+FGLTFUVAnalysis FGLTFUVAnalysisConverter::Convert(const FMeshDescription* Description, const TArray<int32> SectionIndices, int32 TexCoord)
 {
 	if (Description == nullptr)
 	{
 		// TODO: add warning?
 
-		return nullptr;
+		return -1;
 	}
 
 	// TODO: investigate if the fixed size is high enough to properly calculate overlap
 	const FIntPoint TextureSize(512, 512);
 
-	return Outputs.Add_GetRef(MakeUnique<FGLTFUVAnalysis>(CalcOverlapPercentage(TexCoord, TextureSize, *Description, SectionIndices))).Get();
+	return CalcOverlapPercentage(TexCoord, TextureSize, *Description, SectionIndices);
 }
 
 const UMaterialInterface* FGLTFUVAnalysisConverter::GetOverlapMaterial()
diff --git a/Engine/Plugins/Enterprise/GLTFExporter/MigratedFromGit/Source/GLTFExporter/Private/Converters/GLTFUVAnalysisConverters.h b/Engine/Plugins/Enterprise/GLTFExporter/MigratedFromGit/Source/GLTFExporter/Private/Converters/GLTFUVAnalysisConverters.h
--- a/Engine/Plugins/Enterprise/GLTFExporter/MigratedFromGit/Source/GLTFExporter/Private/Converters/GLTFUVAnalysisConverters.h
+++ b/Engine/Plugins/Enterprise/GLTFExporter/MigratedFromGit/Source/GLTFExporter/Private/Converters/GLTFUVAnalysisConverters.h
@@ -6,13 +6,11 @@
 #include "Converters/GLTFUVAnalysis.h"
 #include "Engine.h"
 
-class FGLTFUVAnalysisConverter final : public TGLTFConverter<const FGLTFUVAnalysis*, const FMeshDescription*, const TArray<int32>, int32>
+class FGLTFUVAnalysisConverter final : public TGLTFConverter<FGLTFUVAnalysis, const FMeshDescription*, const TArray<int32>, int32>
 {
-	TArray<TUniquePtr<FGLTFUVAnalysis>> Outputs;
-
 	virtual void Sanitize(const FMeshDescription*& Description, const TArray<int32>& SectionIndices, int32& TexCoord) override;
 
-	virtual const FGLTFUVAnalysis* Convert(const FMeshDescription* Description, const TArray<int32> SectionIndices, int32 TexCoord) override;
+	virtual FGLTFUVAnalysis Convert(const FMeshDescription* Description, const TArray<int32> SectionIndices, int32 TexCoord) override;
 
 	static const UMaterialInterface* GetOverlapMaterial();
 
diff --git a/Engine/Plugins/Enterprise/GLTFExporter/MigratedFromGit/Source/GLTFExporter/Private/Tasks/GLTFMaterialTasks.cpp b/Engine/Plugins/Enterprise/GLTFExporter/MigratedFromGit/Source/GLTFExporter/Private/Tasks/GLTFMaterialTasks.cpp
--- a/Engine/Plugins/Enterprise/GLTFExporter/MigratedFromGit/Source/GLTFExporter/Private/Tasks/GLTFMaterialTasks.cpp
+++ b/Engine/Plugins/Enterprise/GLTFExporter/MigratedFromGit/Source/GLTFExporter/Private/Tasks/GLTFMaterialTasks.cpp
@@ -168,16 +168,16 @@ void FGLTFMaterialTask::Complete()
 		// that as parameter to the UV analysis converter instead of MeshData.
 
 		const float MaxOverlapPercentage = 1.0f;
-		const FGLTFUVAnalysis* UVAnalysis = Builder.UVAnalysisConverter.GetOrAdd(&MeshData->Description, SectionIndices, MeshData->TexCoord);
+		const FGLTFUVAnalysis UVAnalysis = Builder.UVAnalysisConverter.GetOrAdd(&MeshData->Description, SectionIndices, MeshData->TexCoord);
 
-		if (UVAnalysis != nullptr && UVAnalysis->OverlapPercentage > MaxOverlapPercentage)
+		if (UVAnalysis.OverlapPercentage > MaxOverlapPercentage)
 		{
 			Builder.AddWarningMessage(FString::Printf(
 				TEXT("Material %s is baked using TexCoord %d of %s which has %.2f%% overlapping UVs and may produce incorrect results"),
 				*Material->GetName(),
 				MeshData->TexCoord,
 				*MeshData->Name,
-				UVAnalysis->OverlapPercentage));
+				UVAnalysis.OverlapPercentage));
 		}
 	}
 }
