[GLTFExporter] refactor: rename GLTFUVAnalysisConverter to GLTFUVOverlapChecker and isolate it to material converter (instead of builder)

commit 3e2cabb7dd3f9b7a97eea7d234a3f983629fc5b5
Date: Tue, 24 Nov 2020 12:15:44 +0100

    refactor: rename GLTFUVAnalysisConverter to GLTFUVOverlapChecker and isolate it to material converter (instead of builder)
    (cherry picked from commit 58af15f36e6caeee8395a8afa0f869ec76321cd0​)

#skipci

diff --git a/Engine/Plugins/Enterprise/GLTFExporter/MigratedFromGit/Source/GLTFExporter/Private/Builders/GLTFConvertBuilder.h b/Engine/Plugins/Enterprise/GLTFExporter/MigratedFromGit/Source/GLTFExporter/Private/Builders/GLTFConvertBuilder.h
--- a/Engine/Plugins/Enterprise/GLTFExporter/MigratedFromGit/Source/GLTFExporter/Private/Builders/GLTFConvertBuilder.h
+++ b/Engine/Plugins/Enterprise/GLTFExporter/MigratedFromGit/Source/GLTFExporter/Private/Builders/GLTFConvertBuilder.h
@@ -20,7 +20,6 @@
 #include "Converters/GLTFLightMapConverters.h"
 #include "Converters/GLTFHotspotConverters.h"
 #include "Converters/GLTFSkySphereConverters.h"
-#include "Converters/GLTFUVAnalysisConverters.h"
 
 class FGLTFConvertBuilder : public FGLTFImageBuilder
 {
@@ -35,7 +34,6 @@ public:
 	// TODO: find a better place for this types of indirect converters
 	FGLTFStaticMeshDataConverter StaticMeshDataConverter;
 	FGLTFSkeletalMeshDataConverter SkeletalMeshDataConverter;
-	FGLTFUVAnalysisConverter UVAnalysisConverter;
 
 	FGLTFJsonAccessorIndex GetOrAddPositionAccessor(const FGLTFMeshSection* MeshSection, const FPositionVertexBuffer* VertexBuffer);
 	FGLTFJsonAccessorIndex GetOrAddColorAccessor(const FGLTFMeshSection* MeshSection, const FColorVertexBuffer* VertexBuffer);
diff --git a/Engine/Plugins/Enterprise/GLTFExporter/MigratedFromGit/Source/GLTFExporter/Private/Converters/GLTFMaterialConverters.cpp b/Engine/Plugins/Enterprise/GLTFExporter/MigratedFromGit/Source/GLTFExporter/Private/Converters/GLTFMaterialConverters.cpp
--- a/Engine/Plugins/Enterprise/GLTFExporter/MigratedFromGit/Source/GLTFExporter/Private/Converters/GLTFMaterialConverters.cpp
+++ b/Engine/Plugins/Enterprise/GLTFExporter/MigratedFromGit/Source/GLTFExporter/Private/Converters/GLTFMaterialConverters.cpp
@@ -2,7 +2,6 @@
 
 #include "Converters/GLTFMaterialConverters.h"
 #include "Converters/GLTFMaterialUtility.h"
-#include "Converters/GLTFMeshUtility.h"
 #include "Builders/GLTFConvertBuilder.h"
 #include "Tasks/GLTFMaterialTasks.h"
 
@@ -25,6 +24,6 @@ FGLTFJsonMaterialIndex FGLTFMaterialConverter::Convert(const UMaterialInterface*
 	}
 
 	const FGLTFJsonMaterialIndex MaterialIndex = Builder.AddMaterial();
-	Builder.SetupTask<FGLTFMaterialTask>(Builder, Material, MeshData, SectionIndices, MaterialIndex);
+	Builder.SetupTask<FGLTFMaterialTask>(Builder, UVOverlapChecker, Material, MeshData, SectionIndices, MaterialIndex);
 	return MaterialIndex;
 }
diff --git a/Engine/Plugins/Enterprise/GLTFExporter/MigratedFromGit/Source/GLTFExporter/Private/Converters/GLTFMaterialConverters.h b/Engine/Plugins/Enterprise/GLTFExporter/MigratedFromGit/Source/GLTFExporter/Private/Converters/GLTFMaterialConverters.h
--- a/Engine/Plugins/Enterprise/GLTFExporter/MigratedFromGit/Source/GLTFExporter/Private/Converters/GLTFMaterialConverters.h
+++ b/Engine/Plugins/Enterprise/GLTFExporter/MigratedFromGit/Source/GLTFExporter/Private/Converters/GLTFMaterialConverters.h
@@ -6,6 +6,7 @@
 #include "Converters/GLTFConverter.h"
 #include "Converters/GLTFBuilderContext.h"
 #include "Converters/GLTFMeshData.h"
+#include "Converters/GLTFUVOverlapChecker.h"
 #include "Engine.h"
 
 class FGLTFMaterialConverter final : public FGLTFBuilderContext, public TGLTFConverter<FGLTFJsonMaterialIndex, const UMaterialInterface*, const FGLTFMeshData*, TArray<int32>>
@@ -15,4 +16,6 @@ class FGLTFMaterialConverter final : public FGLTFBuilderContext, public TGLTFCon
 	virtual void Sanitize(const UMaterialInterface*& Material, const FGLTFMeshData*& MeshData, TArray<int32>& SectionIndices) override;
 
 	virtual FGLTFJsonMaterialIndex Convert(const UMaterialInterface* Material, const FGLTFMeshData* MeshData, TArray<int32> SectionIndices) override;
+
+	FGLTFUVOverlapChecker UVOverlapChecker;
 };
diff --git a/Engine/Plugins/Enterprise/GLTFExporter/MigratedFromGit/Source/GLTFExporter/Private/Converters/GLTFUVAnalysis.h b/Engine/Plugins/Enterprise/GLTFExporter/MigratedFromGit/Source/GLTFExporter/Private/Converters/GLTFUVAnalysis.h
deleted file mode 100644
diff --git a/Engine/Plugins/Enterprise/GLTFExporter/MigratedFromGit/Source/GLTFExporter/Private/Converters/GLTFUVAnalysisConverters.cpp b/Engine/Plugins/Enterprise/GLTFExporter/MigratedFromGit/Source/GLTFExporter/Private/Converters/GLTFUVAnalysisConverters.cpp
deleted file mode 100644
diff --git a/Engine/Plugins/Enterprise/GLTFExporter/MigratedFromGit/Source/GLTFExporter/Private/Converters/GLTFUVAnalysisConverters.h b/Engine/Plugins/Enterprise/GLTFExporter/MigratedFromGit/Source/GLTFExporter/Private/Converters/GLTFUVAnalysisConverters.h
deleted file mode 100644
diff --git a/Engine/Plugins/Enterprise/GLTFExporter/MigratedFromGit/Source/GLTFExporter/Private/Converters/GLTFUVOverlapChecker.cpp b/Engine/Plugins/Enterprise/GLTFExporter/MigratedFromGit/Source/GLTFExporter/Private/Converters/GLTFUVOverlapChecker.cpp
new file mode 100644
--- /dev/null
+++ b/Engine/Plugins/Enterprise/GLTFExporter/MigratedFromGit/Source/GLTFExporter/Private/Converters/GLTFUVOverlapChecker.cpp
@@ -0,0 +1,113 @@
+// Copyright Epic Games, Inc. All Rights Reserved.
+
+#include "Converters/GLTFUVOverlapChecker.h"
+#include "StaticMeshAttributes.h"
+#include "MeshDescription.h"
+#include "Modules/ModuleManager.h"
+#include "IMaterialBakingModule.h"
+#include "MaterialBakingStructures.h"
+
+void FGLTFUVOverlapChecker::Sanitize(const FMeshDescription*& Description, const TArray<int32>& SectionIndices, int32& TexCoord)
+{
+	if (Description != nullptr)
+	{
+		const TVertexInstanceAttributesConstRef<FVector2D> VertexInstanceUVs =
+			Description->VertexInstanceAttributes().GetAttributesRef<FVector2D>(MeshAttribute::VertexInstance::TextureCoordinate);
+		const int32 TexCoordCount = VertexInstanceUVs.GetNumIndices();
+
+		if (TexCoord < 0 || TexCoord >= TexCoordCount)
+		{
+			Description = nullptr;
+		}
+
+		const int32 MinSectionIndex = FMath::Min(SectionIndices);
+		const int32 MaxSectionIndex = FMath::Max(SectionIndices);
+		const int32 SectionCount = Description->PolygonGroups().GetArraySize();
+
+		if (MinSectionIndex < 0 || MaxSectionIndex >= SectionCount)
+		{
+			Description = nullptr;
+		}
+	}
+}
+
+float FGLTFUVOverlapChecker::Convert(const FMeshDescription* Description, const TArray<int32> SectionIndices, int32 TexCoord)
+{
+	if (Description == nullptr)
+	{
+		// TODO: report warning?
+
+		return -1;
+	}
+
+	// TODO: investigate if the fixed size is high enough to properly calculate overlap
+	const FIntPoint TextureSize(512, 512);
+	const FMaterialPropertyEx Property = MP_EmissiveColor;
+
+	FMeshData MeshSet;
+	MeshSet.TextureCoordinateBox = { { 0.0f, 0.0f }, { 1.0f, 1.0f } };
+	MeshSet.TextureCoordinateIndex = TexCoord;
+	MeshSet.RawMeshDescription = const_cast<FMeshDescription*>(Description);
+	MeshSet.MaterialIndices = SectionIndices; // NOTE: MaterialIndices is actually section indices
+
+	FMaterialDataEx MatSet;
+	MatSet.Material = GetMaterial();
+	MatSet.PropertySizes.Add(Property, TextureSize);
+	MatSet.BlendMode = MatSet.Material->GetBlendMode();
+	MatSet.bPerformBorderSmear = false;
+
+	TArray<FMeshData*> MeshSettings;
+	TArray<FMaterialDataEx*> MatSettings;
+	MeshSettings.Add(&MeshSet);
+	MatSettings.Add(&MatSet);
+
+	TArray<FBakeOutputEx> BakeOutputs;
+	IMaterialBakingModule& Module = FModuleManager::Get().LoadModuleChecked<IMaterialBakingModule>("GLTFMaterialBaking");
+
+	Module.BakeMaterials(MatSettings, MeshSettings, BakeOutputs);
+
+	const FBakeOutputEx& BakeOutput = BakeOutputs[0];
+	const TArray<FColor>& BakedPixels = BakeOutput.PropertyData.FindChecked(Property);
+
+	if (BakedPixels.Num() <= 0)
+	{
+		return -1;
+	}
+
+	// NOTE: the emissive value of each pixel will be incremented by 10 each time a triangle is drawn on it.
+	// Therefore a value of 0.0 indicates an unreferenced pixel, a value of 10.0 indicates a uniquely referenced pixel,
+	// and a value of 20+ indicates an overlapping pixel.
+	// The value may differ somewhat from what was specified in the material due to conversions (color-space, gamma, float => uint etc),
+	// so to be safe we set the limit at 15 instead of 10 to prevent incorrectly flagging some pixels as overlapping.
+	const float EmissiveThreshold = 15.0f;
+
+	float ProcessedPixels = 0;
+	float OverlappingPixels = 0;
+
+	if (BakeOutput.EmissiveScale > EmissiveThreshold)
+	{
+		const FColor Magenta(255, 0, 255, 255);
+		const uint8 ColorThreshold = FMath::RoundToInt((EmissiveThreshold / BakeOutput.EmissiveScale) * 255.0f); // This should never overflow since we know that the quotient will be <= 1.0
+
+		bool bIsMagenta;
+		bool bIsOverlapping;
+
+		for (const FColor& Pixel: BakedPixels)
+		{
+			bIsMagenta = Pixel == Magenta;
+			bIsOverlapping = Pixel.G > ColorThreshold;
+
+			ProcessedPixels += !bIsMagenta ? 1 : 0;
+			OverlappingPixels += !bIsMagenta && bIsOverlapping ? 1 : 0;
+		}
+	}
+
+	return OverlappingPixels / ProcessedPixels;
+}
+
+UMaterialInterface* FGLTFUVOverlapChecker::GetMaterial()
+{
+	static UMaterialInterface* Material = LoadObject<UMaterialInterface>(nullptr, TEXT("/GLTFExporter/Materials/OverlappingUVs.OverlappingUVs"));
+	check(Material);
+	return Material;
+}
diff --git a/Engine/Plugins/Enterprise/GLTFExporter/MigratedFromGit/Source/GLTFExporter/Private/Converters/GLTFUVOverlapChecker.h b/Engine/Plugins/Enterprise/GLTFExporter/MigratedFromGit/Source/GLTFExporter/Private/Converters/GLTFUVOverlapChecker.h
new file mode 100644
--- /dev/null
+++ b/Engine/Plugins/Enterprise/GLTFExporter/MigratedFromGit/Source/GLTFExporter/Private/Converters/GLTFUVOverlapChecker.h
@@ -0,0 +1,15 @@
+// Copyright Epic Games, Inc. All Rights Reserved.
+
+#pragma once
+
+#include "Converters/GLTFConverter.h"
+#include "Engine.h"
+
+class FGLTFUVOverlapChecker final : public TGLTFConverter<float, const FMeshDescription*, const TArray<int32>, int32>
+{
+	virtual void Sanitize(const FMeshDescription*& Description, const TArray<int32>& SectionIndices, int32& TexCoord) override;
+
+	virtual float Convert(const FMeshDescription* Description, const TArray<int32> SectionIndices, int32 TexCoord) override;
+
+	static UMaterialInterface* GetMaterial();
+};
diff --git a/Engine/Plugins/Enterprise/GLTFExporter/MigratedFromGit/Source/GLTFExporter/Private/Tasks/GLTFMaterialTasks.cpp b/Engine/Plugins/Enterprise/GLTFExporter/MigratedFromGit/Source/GLTFExporter/Private/Tasks/GLTFMaterialTasks.cpp
--- a/Engine/Plugins/Enterprise/GLTFExporter/MigratedFromGit/Source/GLTFExporter/Private/Tasks/GLTFMaterialTasks.cpp
+++ b/Engine/Plugins/Enterprise/GLTFExporter/MigratedFromGit/Source/GLTFExporter/Private/Tasks/GLTFMaterialTasks.cpp
@@ -163,21 +163,21 @@ void FGLTFMaterialTask::Complete()
 	if (MeshData != nullptr && MeshDataBakedProperties.Num() > 0)
 	{
 		// TODO: when a mesh-component is used, MeshData will be unique for each converted component even if
-		// they use the same mesh internally. This leads to an excessive number of unique UV analyses.
+		// they use the same mesh internally. This leads to an excessive number of unique UV overlapp checks.
 		// We should look into creating a separate FGLTFMeshData for the mesh itself in these cases, and pass
 		// that as parameter to the UV analysis converter instead of MeshData.
 
-		const float MaxOverlapPercentage = 1.0f;
-		const FGLTFUVAnalysis UVAnalysis = Builder.UVAnalysisConverter.GetOrAdd(&MeshData->Description, SectionIndices, MeshData->TexCoord);
+		const float UVOverlapThreshold = 1.0f / 100.0f;
+		const float UVOverlap = UVOverlapChecker.GetOrAdd(&MeshData->Description, SectionIndices, MeshData->TexCoord);
 
-		if (UVAnalysis.OverlapPercentage > MaxOverlapPercentage)
+		if (UVOverlap > UVOverlapThreshold)
 		{
 			Builder.AddWarningMessage(FString::Printf(
-				TEXT("Material %s is baked using TexCoord %d of %s which has %.2f%% overlapping UVs and may produce incorrect results"),
+				TEXT("Material %s is baked using mesh data from %s but the lightmap UV (channel %d) are overlapping by %.2f%% and may produce incorrect results"),
 				*Material->GetName(),
-				MeshData->TexCoord,
 				*MeshData->Name,
-				UVAnalysis.OverlapPercentage));
+				MeshData->TexCoord,
+				UVOverlap * 100));
 		}
 	}
 }
diff --git a/Engine/Plugins/Enterprise/GLTFExporter/MigratedFromGit/Source/GLTFExporter/Private/Tasks/GLTFMaterialTasks.h b/Engine/Plugins/Enterprise/GLTFExporter/MigratedFromGit/Source/GLTFExporter/Private/Tasks/GLTFMaterialTasks.h
--- a/Engine/Plugins/Enterprise/GLTFExporter/MigratedFromGit/Source/GLTFExporter/Private/Tasks/GLTFMaterialTasks.h
+++ b/Engine/Plugins/Enterprise/GLTFExporter/MigratedFromGit/Source/GLTFExporter/Private/Tasks/GLTFMaterialTasks.h
@@ -3,8 +3,9 @@
 #pragma once
 
 #include "Tasks/GLTFTask.h"
-#include "Converters/GLTFMeshData.h"
 #include "Builders/GLTFConvertBuilder.h"
+#include "Converters/GLTFUVOverlapChecker.h"
+#include "Converters/GLTFMeshData.h"
 #include "MaterialPropertyEx.h"
 #include "Engine.h"
 
@@ -14,9 +15,10 @@ class FGLTFMaterialTask : public FGLTFTask
 {
 public:
 
-	FGLTFMaterialTask(FGLTFConvertBuilder& Builder, const UMaterialInterface* Material, const FGLTFMeshData* MeshData, TArray<int32> SectionIndices, FGLTFJsonMaterialIndex MaterialIndex)
+	FGLTFMaterialTask(FGLTFConvertBuilder& Builder, FGLTFUVOverlapChecker& UVOverlapChecker, const UMaterialInterface* Material, const FGLTFMeshData* MeshData, TArray<int32> SectionIndices, FGLTFJsonMaterialIndex MaterialIndex)
 		: FGLTFTask(EGLTFTaskPriority::Material)
 		, Builder(Builder)
+		, UVOverlapChecker(UVOverlapChecker)
 		, Material(Material)
 		, MeshData(MeshData)
 		, SectionIndices(SectionIndices)
@@ -34,6 +36,7 @@ public:
 private:
 
 	FGLTFConvertBuilder& Builder;
+	FGLTFUVOverlapChecker& UVOverlapChecker;
 	const UMaterialInterface* Material;
 	const FGLTFMeshData* MeshData;
 	const TArray<int32> SectionIndices;
