[GLTFExporter] refactor: improve performance by reusing UV overlap checks for same mesh

commit ad8672b4e2ad69a880cc3ea20d7774c04fd8a922
Date: Tue, 24 Nov 2020 12:38:31 +0100

    refactor: improve performance by reusing UV overlap checks for same mesh
    (cherry picked from commit 9994013d7c208b2e9003633d262db4f4dc2aecc7​)

#skipci

diff --git a/Engine/Plugins/Enterprise/GLTFExporter/MigratedFromGit/Source/GLTFExporter/Private/Converters/GLTFMeshData.cpp b/Engine/Plugins/Enterprise/GLTFExporter/MigratedFromGit/Source/GLTFExporter/Private/Converters/GLTFMeshData.cpp
--- a/Engine/Plugins/Enterprise/GLTFExporter/MigratedFromGit/Source/GLTFExporter/Private/Converters/GLTFMeshData.cpp
+++ b/Engine/Plugins/Enterprise/GLTFExporter/MigratedFromGit/Source/GLTFExporter/Private/Converters/GLTFMeshData.cpp
@@ -8,6 +8,7 @@
 #include "Rendering/SkeletalMeshRenderData.h"
 
 FGLTFMeshData::FGLTFMeshData(const UStaticMesh* StaticMesh, const UStaticMeshComponent* StaticMeshComponent, int32 LODIndex)
+	: Parent(nullptr)
 {
 	FStaticMeshAttributes(Description).Register();
 
@@ -29,6 +30,7 @@ FGLTFMeshData::FGLTFMeshData(const UStaticMesh* StaticMesh, const UStaticMeshCom
 }
 
 FGLTFMeshData::FGLTFMeshData(const USkeletalMesh* SkeletalMesh, const USkeletalMeshComponent* SkeletalMeshComponent, int32 LODIndex)
+	: Parent(nullptr)
 {
 	FStaticMeshAttributes(Description).Register();
 
@@ -72,3 +74,8 @@ FGLTFMeshData::FGLTFMeshData(const USkeletalMesh* SkeletalMesh, const USkeletalM
 
 	// TODO: add warning if texture coordinate has overlap
 }
+
+const FGLTFMeshData* FGLTFMeshData::GetParent() const
+{
+	return Parent != nullptr ? Parent : this;
+}
diff --git a/Engine/Plugins/Enterprise/GLTFExporter/MigratedFromGit/Source/GLTFExporter/Private/Converters/GLTFMeshData.h b/Engine/Plugins/Enterprise/GLTFExporter/MigratedFromGit/Source/GLTFExporter/Private/Converters/GLTFMeshData.h
--- a/Engine/Plugins/Enterprise/GLTFExporter/MigratedFromGit/Source/GLTFExporter/Private/Converters/GLTFMeshData.h
+++ b/Engine/Plugins/Enterprise/GLTFExporter/MigratedFromGit/Source/GLTFExporter/Private/Converters/GLTFMeshData.h
@@ -10,7 +10,12 @@ struct FGLTFMeshData
 	FGLTFMeshData(const UStaticMesh* StaticMesh, const UStaticMeshComponent* StaticMeshComponent, int32 LODIndex);
 	FGLTFMeshData(const USkeletalMesh* SkeletalMesh, const USkeletalMeshComponent* SkeletalMeshComponent, int32 LODIndex);
 
+	const FGLTFMeshData* GetParent() const;
+
 	FString Name;
 	FMeshDescription Description;
 	int32 TexCoord;
+
+	// TODO: find a better name for referencing the mesh-only data (no component)
+	const FGLTFMeshData* Parent;
 };
diff --git a/Engine/Plugins/Enterprise/GLTFExporter/MigratedFromGit/Source/GLTFExporter/Private/Converters/GLTFMeshDataConverters.h b/Engine/Plugins/Enterprise/GLTFExporter/MigratedFromGit/Source/GLTFExporter/Private/Converters/GLTFMeshDataConverters.h
--- a/Engine/Plugins/Enterprise/GLTFExporter/MigratedFromGit/Source/GLTFExporter/Private/Converters/GLTFMeshDataConverters.h
+++ b/Engine/Plugins/Enterprise/GLTFExporter/MigratedFromGit/Source/GLTFExporter/Private/Converters/GLTFMeshDataConverters.h
@@ -13,7 +13,14 @@ class TGLTFMeshDataConverter final : public TGLTFConverter<const FGLTFMeshData*,
 
 	const FGLTFMeshData* Convert(const MeshType* Mesh, const MeshComponentType* MeshComponent, int32 LODIndex)
 	{
-		return Outputs.Add_GetRef(MakeUnique<FGLTFMeshData>(Mesh, MeshComponent, LODIndex)).Get();
+		TUniquePtr<FGLTFMeshData> Output = MakeUnique<FGLTFMeshData>(Mesh, MeshComponent, LODIndex);
+
+		if (MeshComponent != nullptr)
+		{
+			Output->Parent = GetOrAdd(Mesh, nullptr, LODIndex);
+		}
+
+		return Outputs.Add_GetRef(MoveTemp(Output)).Get();
 	}
 };
 
diff --git a/Engine/Plugins/Enterprise/GLTFExporter/MigratedFromGit/Source/GLTFExporter/Private/Tasks/GLTFMaterialTasks.cpp b/Engine/Plugins/Enterprise/GLTFExporter/MigratedFromGit/Source/GLTFExporter/Private/Tasks/GLTFMaterialTasks.cpp
--- a/Engine/Plugins/Enterprise/GLTFExporter/MigratedFromGit/Source/GLTFExporter/Private/Tasks/GLTFMaterialTasks.cpp
+++ b/Engine/Plugins/Enterprise/GLTFExporter/MigratedFromGit/Source/GLTFExporter/Private/Tasks/GLTFMaterialTasks.cpp
@@ -162,20 +162,17 @@ void FGLTFMaterialTask::Complete()
 	// Warn if properties baked using mesh data are using overlapping UVs
 	if (MeshData != nullptr && MeshDataBakedProperties.Num() > 0)
 	{
-		// TODO: when a mesh-component is used, MeshData will be unique for each converted component even if
-		// they use the same mesh internally. This leads to an excessive number of unique UV overlapp checks.
-		// We should look into creating a separate FGLTFMeshData for the mesh itself in these cases, and pass
-		// that as parameter to the UV analysis converter instead of MeshData.
+		const FGLTFMeshData* ParentMeshData = MeshData->GetParent();
 
 		const float UVOverlapThreshold = 1.0f / 100.0f;
-		const float UVOverlap = UVOverlapChecker.GetOrAdd(&MeshData->Description, SectionIndices, MeshData->TexCoord);
+		const float UVOverlap = UVOverlapChecker.GetOrAdd(&ParentMeshData->Description, SectionIndices, MeshData->TexCoord);
 
 		if (UVOverlap > UVOverlapThreshold)
 		{
 			Builder.AddWarningMessage(FString::Printf(
 				TEXT("Material %s is baked using mesh data from %s but the lightmap UV (channel %d) are overlapping by %.2f%% and may produce incorrect results"),
 				*Material->GetName(),
-				*MeshData->Name,
+				*ParentMeshData->Name,
 				MeshData->TexCoord,
 				UVOverlap * 100));
 		}
