[GLTFExporter] fix: fix crash in editor due to UGLTFMaterialAnalyzer being treated as a real material

commit e03f3a43758df20bdf4f967d3d2f7a3d814c84f5
Date: Thu, 19 Nov 2020 10:00:45 +0100

    fix: fix crash in editor due to UGLTFMaterialAnalyzer being treated as a real material
    (cherry picked from commit c4e8a0161c1b3178acfbc3a7f2af1cdaaa03459d​)

#skipci

diff --git a/Engine/Plugins/Enterprise/GLTFExporter/MigratedFromGit/Source/GLTFMaterialAnalyzer/Private/GLTFMaterialAnalyzer.cpp b/Engine/Plugins/Enterprise/GLTFExporter/MigratedFromGit/Source/GLTFMaterialAnalyzer/Private/GLTFMaterialAnalyzer.cpp
--- a/Engine/Plugins/Enterprise/GLTFExporter/MigratedFromGit/Source/GLTFMaterialAnalyzer/Private/GLTFMaterialAnalyzer.cpp
+++ b/Engine/Plugins/Enterprise/GLTFExporter/MigratedFromGit/Source/GLTFMaterialAnalyzer/Private/GLTFMaterialAnalyzer.cpp
@@ -19,8 +19,17 @@ void UGLTFMaterialAnalyzer::ResetToDefaults()
 
 void UGLTFMaterialAnalyzer::AnalyzeMaterialPropertyEx(const UMaterialInterface* InMaterial, const EMaterialProperty& InProperty, const FString& InCustomOutput, FGLTFMaterialAnalysis& OutAnalysis)
 {
-	// TODO: use a shared UGLTFMaterialAnalyzer instance instead of creating a new one for each invocation
-	UGLTFMaterialAnalyzer* Analyzer = NewObject<UGLTFMaterialAnalyzer>();
+	// NOTE: besides avoiding creating a new object each time a property is analyzed,
+	// the use of GetDefaultObject() also server another purpose:
+	//
+	// When a property in an object that may be referenced by materials is modified,
+	// the editor will evaluate which materials need updating.
+	// UGLTFMaterialAnalyzer has no implementations for the pure virtual functions
+	// defined in UMaterialInterface, and the editor would crash if it was processed.
+	// Fortunately objects with the RF_ClassDefaultObject flag
+	// are excluded from being processed.
+
+	UGLTFMaterialAnalyzer* Analyzer = Cast<UGLTFMaterialAnalyzer>(StaticClass()->GetDefaultObject());
 
 	Analyzer->Property = InProperty;
 	Analyzer->CustomOutput = InCustomOutput;
