[GLTFExporter] fix: ensure that bRequiresVertexData is initialized to prevent random results when analyzing materials

commit f96f295717477cad46a2163b85664f20764e5512
Date: Sat, 21 Nov 2020 15:50:06 +0100

    fix: ensure that bRequiresVertexData is initialized to prevent random results when analyzing materials
    the flag isn't assigned by the analyzer if the analyzed material is already compiled
    (cherry picked from commit b973497b37ae6a822b49c5aabd136ed8e2ee1e36​)

#skipci

diff --git a/Engine/Plugins/Enterprise/GLTFExporter/MigratedFromGit/Source/GLTFMaterialAnalyzer/Public/GLTFMaterialAnalysis.h b/Engine/Plugins/Enterprise/GLTFExporter/MigratedFromGit/Source/GLTFMaterialAnalyzer/Public/GLTFMaterialAnalysis.h
--- a/Engine/Plugins/Enterprise/GLTFExporter/MigratedFromGit/Source/GLTFMaterialAnalyzer/Public/GLTFMaterialAnalysis.h
+++ b/Engine/Plugins/Enterprise/GLTFExporter/MigratedFromGit/Source/GLTFMaterialAnalyzer/Public/GLTFMaterialAnalysis.h
@@ -6,6 +6,11 @@
 
 struct FGLTFMaterialAnalysis
 {
+	FGLTFMaterialAnalysis()
+		: bRequiresVertexData(false)
+	{
+	}
+
 	/** Tracks the texture coordinates used by this material */
 	TBitArray<> TextureCoordinates;
 
