[GLTFExporter] fix: use a camera vector identical to the vertex normal when baking materials

commit 7e2a01138d0ba0a5403c9aa64e9c796ff39ac822
Date: Wed, 18 Nov 2020 14:06:46 +0100

    fix: use a camera vector identical to the vertex normal when baking materials
    fixes issues when baking automotive glass materials (which calculate fresnel)
    (cherry picked from commit 65134a752c3be757c943d1b29d9e03ee90e462ea​)

#skipci

diff --git a/Engine/Plugins/Enterprise/GLTFExporter/MigratedFromGit/Source/GLTFMaterialBaking/Private/ExportMaterialProxy.h b/Engine/Plugins/Enterprise/GLTFExporter/MigratedFromGit/Source/GLTFMaterialBaking/Private/ExportMaterialProxy.h
--- a/Engine/Plugins/Enterprise/GLTFExporter/MigratedFromGit/Source/GLTFMaterialBaking/Private/ExportMaterialProxy.h
+++ b/Engine/Plugins/Enterprise/GLTFExporter/MigratedFromGit/Source/GLTFMaterialBaking/Private/ExportMaterialProxy.h
@@ -115,7 +115,10 @@ struct FExportMaterialCompiler : public FProxyMaterialCompiler
 
 	virtual int32 CameraVector() override
 	{
-		return Compiler->Constant3(0.0f, 0.0f, 1.0f);
+		// NOTE: by using VertexNormal instead of a fixed (world-space) vector, we allow
+		// more correct baking of materials that use the dot-product between the vertex normal
+		// and the camera vector for effects such as fresnel.
+		return Compiler->VertexNormal();
 	}
 
 	virtual int32 ReflectionAboutCustomWorldNormal(int32 CustomWorldNormal, int32 bNormalizeCustomWorldNormal) override
