[GLTFExporter] fix: avoid gamma-conversion in FGLTFMaterialUtility::CombineTextures

commit cdece2babf3d8ea8249a535ea0edccf1a6b24f5f
Date: Wed, 18 Nov 2020 16:12:08 +0100

    fix: avoid gamma-conversion in FGLTFMaterialUtility::CombineTextures
    preserves the original gamma when combining multiple sources
    (cherry picked from commit 89c6490a081a4a07e5a413112af021019f066ce0​)

#skipci

diff --git a/Engine/Plugins/Enterprise/GLTFExporter/MigratedFromGit/Source/GLTFExporter/Private/Converters/GLTFMaterialUtility.cpp b/Engine/Plugins/Enterprise/GLTFExporter/MigratedFromGit/Source/GLTFExporter/Private/Converters/GLTFMaterialUtility.cpp
--- a/Engine/Plugins/Enterprise/GLTFExporter/MigratedFromGit/Source/GLTFExporter/Private/Converters/GLTFMaterialUtility.cpp
+++ b/Engine/Plugins/Enterprise/GLTFExporter/MigratedFromGit/Source/GLTFExporter/Private/Converters/GLTFMaterialUtility.cpp
@@ -154,12 +154,17 @@ UTexture2D* FGLTFMaterialUtility::CreateTransientTexture(const FGLTFPropertyBake
 
 bool FGLTFMaterialUtility::CombineTextures(TArray<FColor>& OutPixels, const TArray<FGLTFTextureCombineSource>& Sources, const FIntPoint& OutputSize, const EPixelFormat OutputPixelFormat)
 {
+	// NOTE: both bForceLinearGamma and TargetGamma=2.2 seem necessary for exported images to match their source data.
+	// It's not entirely clear why gamma must be 2.2 (instead of 0.0) and why bInForceLinearGamma must also be true.
+	const bool bForceLinearGamma = true;
+	const float TargetGamma = 2.2f;
+
 	UTextureRenderTarget2D* RenderTarget2D = NewObject<UTextureRenderTarget2D>();
 
 	RenderTarget2D->AddToRoot();
 	RenderTarget2D->ClearColor = { 0.0f, 0.0f, 0.0f, 0.0f };
-	RenderTarget2D->InitCustomFormat(OutputSize.X, OutputSize.Y, OutputPixelFormat, false);
-	RenderTarget2D->TargetGamma = 0.0f;
+	RenderTarget2D->InitCustomFormat(OutputSize.X, OutputSize.Y, OutputPixelFormat, bForceLinearGamma);
+	RenderTarget2D->TargetGamma = TargetGamma;
 
 	FRenderTarget* RenderTarget = RenderTarget2D->GameThread_GetRenderTargetResource();
 	FCanvas Canvas(RenderTarget, nullptr, 0.0f, 0.0f, 0.0f, GMaxRHIFeatureLevel);
