[GLTFExporter] fix: use linear colors when baking metallic-roughness and occlusion textures

commit 1792c0c71a049273bf08769131d05092328fe52d
Date: Wed, 18 Nov 2020 17:07:52 +0100

    fix: use linear colors when baking metallic-roughness and occlusion textures
    (cherry picked from commit 2f0407710b7a1e0f5881e7d568aef93ad8fe2a7e​)
    # Conflicts from cherry-pick:
    #	Plugins/GLTFExporter/Source/GLTFExporter/Private/Tasks/GLTFMaterialTasks.cpp
    #	Plugins/GLTFExporter/Source/GLTFExporter/Private/Tasks/GLTFMaterialTasks.h

#skipci

diff --git a/Engine/Plugins/Enterprise/GLTFExporter/MigratedFromGit/Source/GLTFExporter/Private/Converters/GLTFMaterialUtility.cpp b/Engine/Plugins/Enterprise/GLTFExporter/MigratedFromGit/Source/GLTFExporter/Private/Converters/GLTFMaterialUtility.cpp
--- a/Engine/Plugins/Enterprise/GLTFExporter/MigratedFromGit/Source/GLTFExporter/Private/Converters/GLTFMaterialUtility.cpp
+++ b/Engine/Plugins/Enterprise/GLTFExporter/MigratedFromGit/Source/GLTFExporter/Private/Converters/GLTFMaterialUtility.cpp
@@ -309,6 +309,14 @@ FGLTFJsonTextureIndex FGLTFMaterialUtility::AddTexture(FGLTFConvertBuilder& Buil
 	return Builder.AddTexture(JsonTexture);
 }
 
+void FGLTFMaterialUtility::TransformToLinear(TArray<FColor>& InOutPixels)
+{
+	for (FColor& Pixel: InOutPixels)
+	{
+		Pixel = FLinearColor(Pixel).ToFColor(false);
+	}
+}
+
 FLinearColor FGLTFMaterialUtility::GetMask(const FExpressionInput& ExpressionInput)
 {
 	return FLinearColor(ExpressionInput.MaskR, ExpressionInput.MaskG, ExpressionInput.MaskB, ExpressionInput.MaskA);
diff --git a/Engine/Plugins/Enterprise/GLTFExporter/MigratedFromGit/Source/GLTFExporter/Private/Converters/GLTFMaterialUtility.h b/Engine/Plugins/Enterprise/GLTFExporter/MigratedFromGit/Source/GLTFExporter/Private/Converters/GLTFMaterialUtility.h
--- a/Engine/Plugins/Enterprise/GLTFExporter/MigratedFromGit/Source/GLTFExporter/Private/Converters/GLTFMaterialUtility.h
+++ b/Engine/Plugins/Enterprise/GLTFExporter/MigratedFromGit/Source/GLTFExporter/Private/Converters/GLTFMaterialUtility.h
@@ -62,6 +62,8 @@ struct FGLTFMaterialUtility
 	static FGLTFJsonTextureIndex AddCombinedTexture(FGLTFConvertBuilder& Builder, const TArray<FGLTFTextureCombineSource>& CombineSources, const FIntPoint& TextureSize, bool bIgnoreAlpha, const FString& TextureName, EGLTFJsonTextureFilter MinFilter, EGLTFJsonTextureFilter MagFilter, EGLTFJsonTextureWrap WrapS, EGLTFJsonTextureWrap WrapT);
 	static FGLTFJsonTextureIndex AddTexture(FGLTFConvertBuilder& Builder, const TArray<FColor>& Pixels, const FIntPoint& TextureSize, bool bIgnoreAlpha, bool bIsNormalMap, const FString& TextureName, EGLTFJsonTextureFilter MinFilter, EGLTFJsonTextureFilter MagFilter, EGLTFJsonTextureWrap WrapS, EGLTFJsonTextureWrap WrapT);
 
+	static void TransformToLinear(TArray<FColor>& InOutPixels);
+
 	static FLinearColor GetMask(const FExpressionInput& ExpressionInput);
 	static uint32 GetMaskComponentCount(const FExpressionInput& ExpressionInput);
 
diff --git a/Engine/Plugins/Enterprise/GLTFExporter/MigratedFromGit/Source/GLTFExporter/Private/Tasks/GLTFMaterialTasks.cpp b/Engine/Plugins/Enterprise/GLTFExporter/MigratedFromGit/Source/GLTFExporter/Private/Tasks/GLTFMaterialTasks.cpp
--- a/Engine/Plugins/Enterprise/GLTFExporter/MigratedFromGit/Source/GLTFExporter/Private/Tasks/GLTFMaterialTasks.cpp
+++ b/Engine/Plugins/Enterprise/GLTFExporter/MigratedFromGit/Source/GLTFExporter/Private/Tasks/GLTFMaterialTasks.cpp
@@ -127,7 +127,7 @@ void FGLTFMaterialTask::Complete()
 			{
 				if (!TryGetSourceTexture(JsonMaterial.OcclusionTexture, AmbientOcclusionProperty, OcclusionInputMasks))
 				{
-					if (!TryGetBakedMaterialProperty(JsonMaterial.OcclusionTexture, AmbientOcclusionProperty, TEXT("Occlusion")))
+					if (!TryGetBakedMaterialProperty(JsonMaterial.OcclusionTexture, AmbientOcclusionProperty, TEXT("Occlusion"), true))
 					{
 						Builder.AddWarningMessage(FString::Printf(TEXT("Failed to export %s for material %s"), *AmbientOcclusionProperty.ToString(), *Material->GetName()));
 					}
@@ -465,8 +465,8 @@ bool FGLTFMaterialTask::TryGetMetallicAndRoughness(FGLTFJsonPBRMetallicRoughness
 	const EGLTFJsonTextureFilter TextureMinFilter = FGLTFConverterUtility::ConvertMinFilter(TextureFilter);
 	const EGLTFJsonTextureFilter TextureMagFilter = FGLTFConverterUtility::ConvertMagFilter(TextureFilter);
 
-	const FGLTFPropertyBakeOutput MetallicBakeOutput = BakeMaterialProperty(MetallicProperty, MetallicTexCoord, TextureSize);
-	const FGLTFPropertyBakeOutput RoughnessBakeOutput = BakeMaterialProperty(RoughnessProperty, RoughnessTexCoord, TextureSize);
+	FGLTFPropertyBakeOutput MetallicBakeOutput = BakeMaterialProperty(MetallicProperty, MetallicTexCoord, TextureSize);
+	FGLTFPropertyBakeOutput RoughnessBakeOutput = BakeMaterialProperty(RoughnessProperty, RoughnessTexCoord, TextureSize);
 
 	// Detect when both baked properties are constants, which means we can use factors and avoid exporting a texture
 	if (MetallicBakeOutput.bIsConstant && RoughnessBakeOutput.bIsConstant)
@@ -495,6 +495,9 @@ bool FGLTFMaterialTask::TryGetMetallicAndRoughness(FGLTFJsonPBRMetallicRoughness
 		return false;
 	}
 
+	FGLTFMaterialUtility::TransformToLinear(MetallicBakeOutput.Pixels);
+	FGLTFMaterialUtility::TransformToLinear(RoughnessBakeOutput.Pixels);
+
 	TextureSize = RoughnessBakeOutput.Size.ComponentMax(MetallicBakeOutput.Size);
 	MetallicTexture = FGLTFMaterialUtility::CreateTransientTexture(MetallicBakeOutput);
 	RoughnessTexture = FGLTFMaterialUtility::CreateTransientTexture(RoughnessBakeOutput);
@@ -1074,7 +1077,7 @@ bool FGLTFMaterialTask::TryGetSourceTexture(const UTexture2D*& OutTexture, int32
 	return false;
 }
 
-bool FGLTFMaterialTask::TryGetBakedMaterialProperty(FGLTFJsonTextureInfo& OutTexInfo, FGLTFJsonColor3& OutConstant, const FMaterialPropertyEx& Property, const FString& PropertyName)
+bool FGLTFMaterialTask::TryGetBakedMaterialProperty(FGLTFJsonTextureInfo& OutTexInfo, FGLTFJsonColor3& OutConstant, const FMaterialPropertyEx& Property, const FString& PropertyName, bool bTransformToLinear)
 {
 	if (Builder.ExportOptions->BakeMaterialInputs == EGLTFMaterialBakeMode::None)
 	{
@@ -1085,7 +1088,7 @@ bool FGLTFMaterialTask::TryGetBakedMaterialProperty(FGLTFJsonTextureInfo& OutTex
 		return false;
 	}
 
-	const FGLTFPropertyBakeOutput PropertyBakeOutput = BakeMaterialProperty(Property, OutTexInfo.TexCoord);
+	FGLTFPropertyBakeOutput PropertyBakeOutput = BakeMaterialProperty(Property, OutTexInfo.TexCoord);
 
 	if (PropertyBakeOutput.bIsConstant)
 	{
@@ -1093,6 +1096,11 @@ bool FGLTFMaterialTask::TryGetBakedMaterialProperty(FGLTFJsonTextureInfo& OutTex
 		return true;
 	}
 
+	if (bTransformToLinear)
+	{
+		FGLTFMaterialUtility::TransformToLinear(PropertyBakeOutput.Pixels);
+	}
+
 	if (StoreBakedPropertyTexture(OutTexInfo, PropertyBakeOutput, PropertyName))
 	{
 		OutConstant = FGLTFJsonColor3::White; // make sure property is not zero
@@ -1102,7 +1110,7 @@ bool FGLTFMaterialTask::TryGetBakedMaterialProperty(FGLTFJsonTextureInfo& OutTex
 	return false;
 }
 
-bool FGLTFMaterialTask::TryGetBakedMaterialProperty(FGLTFJsonTextureInfo& OutTexInfo, FGLTFJsonColor4& OutConstant, const FMaterialPropertyEx& Property, const FString& PropertyName)
+bool FGLTFMaterialTask::TryGetBakedMaterialProperty(FGLTFJsonTextureInfo& OutTexInfo, FGLTFJsonColor4& OutConstant, const FMaterialPropertyEx& Property, const FString& PropertyName, bool bTransformToLinear)
 {
 	if (Builder.ExportOptions->BakeMaterialInputs == EGLTFMaterialBakeMode::None)
 	{
@@ -1113,7 +1121,7 @@ bool FGLTFMaterialTask::TryGetBakedMaterialProperty(FGLTFJsonTextureInfo& OutTex
 		return false;
 	}
 
-	const FGLTFPropertyBakeOutput PropertyBakeOutput = BakeMaterialProperty(Property, OutTexInfo.TexCoord);
+	FGLTFPropertyBakeOutput PropertyBakeOutput = BakeMaterialProperty(Property, OutTexInfo.TexCoord);
 
 	if (PropertyBakeOutput.bIsConstant)
 	{
@@ -1121,6 +1129,11 @@ bool FGLTFMaterialTask::TryGetBakedMaterialProperty(FGLTFJsonTextureInfo& OutTex
 		return true;
 	}
 
+	if (bTransformToLinear)
+	{
+		FGLTFMaterialUtility::TransformToLinear(PropertyBakeOutput.Pixels);
+	}
+
 	if (StoreBakedPropertyTexture(OutTexInfo, PropertyBakeOutput, PropertyName))
 	{
 		OutConstant = FGLTFJsonColor4::White; // make sure property is not zero
@@ -1130,7 +1143,7 @@ bool FGLTFMaterialTask::TryGetBakedMaterialProperty(FGLTFJsonTextureInfo& OutTex
 	return false;
 }
 
-inline bool FGLTFMaterialTask::TryGetBakedMaterialProperty(FGLTFJsonTextureInfo& OutTexInfo, float& OutConstant, const FMaterialPropertyEx& Property, const FString& PropertyName)
+inline bool FGLTFMaterialTask::TryGetBakedMaterialProperty(FGLTFJsonTextureInfo& OutTexInfo, float& OutConstant, const FMaterialPropertyEx& Property, const FString& PropertyName, bool bTransformToLinear)
 {
 	if (Builder.ExportOptions->BakeMaterialInputs == EGLTFMaterialBakeMode::None)
 	{
@@ -1141,7 +1154,7 @@ inline bool FGLTFMaterialTask::TryGetBakedMaterialProperty(FGLTFJsonTextureInfo&
 		return false;
 	}
 
-	const FGLTFPropertyBakeOutput PropertyBakeOutput = BakeMaterialProperty(Property, OutTexInfo.TexCoord);
+	FGLTFPropertyBakeOutput PropertyBakeOutput = BakeMaterialProperty(Property, OutTexInfo.TexCoord);
 
 	if (PropertyBakeOutput.bIsConstant)
 	{
@@ -1149,6 +1162,11 @@ inline bool FGLTFMaterialTask::TryGetBakedMaterialProperty(FGLTFJsonTextureInfo&
 		return true;
 	}
 
+	if (bTransformToLinear)
+	{
+		FGLTFMaterialUtility::TransformToLinear(PropertyBakeOutput.Pixels);
+	}
+
 	if (StoreBakedPropertyTexture(OutTexInfo, PropertyBakeOutput, PropertyName))
 	{
 		OutConstant = 1; // make sure property is not zero
@@ -1158,7 +1176,7 @@ inline bool FGLTFMaterialTask::TryGetBakedMaterialProperty(FGLTFJsonTextureInfo&
 	return false;
 }
 
-bool FGLTFMaterialTask::TryGetBakedMaterialProperty(FGLTFJsonTextureInfo& OutTexInfo, const FMaterialPropertyEx& Property, const FString& PropertyName)
+bool FGLTFMaterialTask::TryGetBakedMaterialProperty(FGLTFJsonTextureInfo& OutTexInfo, const FMaterialPropertyEx& Property, const FString& PropertyName, bool bTransformToLinear)
 {
 	if (Builder.ExportOptions->BakeMaterialInputs == EGLTFMaterialBakeMode::None)
 	{
@@ -1169,10 +1187,15 @@ bool FGLTFMaterialTask::TryGetBakedMaterialProperty(FGLTFJsonTextureInfo& OutTex
 		return false;
 	}
 
-	const FGLTFPropertyBakeOutput PropertyBakeOutput = BakeMaterialProperty(Property, OutTexInfo.TexCoord);
+	FGLTFPropertyBakeOutput PropertyBakeOutput = BakeMaterialProperty(Property, OutTexInfo.TexCoord);
 
 	if (!PropertyBakeOutput.bIsConstant)
 	{
+		if (bTransformToLinear)
+		{
+			FGLTFMaterialUtility::TransformToLinear(PropertyBakeOutput.Pixels);
+		}
+
 		return StoreBakedPropertyTexture(OutTexInfo, PropertyBakeOutput, PropertyName);
 	}
 
diff --git a/Engine/Plugins/Enterprise/GLTFExporter/MigratedFromGit/Source/GLTFExporter/Private/Tasks/GLTFMaterialTasks.h b/Engine/Plugins/Enterprise/GLTFExporter/MigratedFromGit/Source/GLTFExporter/Private/Tasks/GLTFMaterialTasks.h
--- a/Engine/Plugins/Enterprise/GLTFExporter/MigratedFromGit/Source/GLTFExporter/Private/Tasks/GLTFMaterialTasks.h
+++ b/Engine/Plugins/Enterprise/GLTFExporter/MigratedFromGit/Source/GLTFExporter/Private/Tasks/GLTFMaterialTasks.h
@@ -61,10 +61,10 @@ private:
 	bool TryGetSourceTexture(FGLTFJsonTextureInfo& OutTexInfo, const FMaterialPropertyEx& Property, const TArray<FLinearColor>& AllowedMasks = {}) const;
 	bool TryGetSourceTexture(const UTexture2D*& OutTexture, int32& OutTexCoord, FGLTFJsonTextureTransform& OutTransform, const FMaterialPropertyEx& Property, const TArray<FLinearColor>& AllowedMasks = {}) const;
 
-	bool TryGetBakedMaterialProperty(FGLTFJsonTextureInfo& OutTexInfo, FGLTFJsonColor3& OutConstant, const FMaterialPropertyEx& Property, const FString& PropertyName);
-	bool TryGetBakedMaterialProperty(FGLTFJsonTextureInfo& OutTexInfo, FGLTFJsonColor4& OutConstant, const FMaterialPropertyEx& Property, const FString& PropertyName);
-	bool TryGetBakedMaterialProperty(FGLTFJsonTextureInfo& OutTexInfo, float& OutConstant, const FMaterialPropertyEx& Property, const FString& PropertyName);
-	bool TryGetBakedMaterialProperty(FGLTFJsonTextureInfo& OutTexInfo, const FMaterialPropertyEx& Property, const FString& PropertyName);
+	bool TryGetBakedMaterialProperty(FGLTFJsonTextureInfo& OutTexInfo, FGLTFJsonColor3& OutConstant, const FMaterialPropertyEx& Property, const FString& PropertyName, bool bTransformToLinear = false);
+	bool TryGetBakedMaterialProperty(FGLTFJsonTextureInfo& OutTexInfo, FGLTFJsonColor4& OutConstant, const FMaterialPropertyEx& Property, const FString& PropertyName, bool bTransformToLinear = false);
+	bool TryGetBakedMaterialProperty(FGLTFJsonTextureInfo& OutTexInfo, float& OutConstant, const FMaterialPropertyEx& Property, const FString& PropertyName, bool bTransformToLinear = false);
+	bool TryGetBakedMaterialProperty(FGLTFJsonTextureInfo& OutTexInfo, const FMaterialPropertyEx& Property, const FString& PropertyName, bool bTransformToLinear = false);
 
 	FGLTFPropertyBakeOutput BakeMaterialProperty(const FMaterialPropertyEx& Property, int32& OutTexCoord, bool bCopyAlphaFromRedChannel = false);
 	FGLTFPropertyBakeOutput BakeMaterialProperty(const FMaterialPropertyEx& Property, int32& OutTexCoord, const FIntPoint& TextureSize, bool bCopyAlphaFromRedChannel = false);
