[GLTFExporter] refactor: improve warning message for overlapping UVs by including mesh section indices

commit 767c25aad07d76466acaebb677aa15d165c0f47d
Date: Tue, 24 Nov 2020 14:15:01 +0100

    refactor: improve warning message for overlapping UVs by including mesh section indices
    (cherry picked from commit 6b5fc2d43daa6a8c8c6061aeb3ffa9a9b8ef9f42​)

#skipci

diff --git a/Engine/Plugins/Enterprise/GLTFExporter/MigratedFromGit/Source/GLTFExporter/Private/Tasks/GLTFMaterialTasks.cpp b/Engine/Plugins/Enterprise/GLTFExporter/MigratedFromGit/Source/GLTFExporter/Private/Tasks/GLTFMaterialTasks.cpp
--- a/Engine/Plugins/Enterprise/GLTFExporter/MigratedFromGit/Source/GLTFExporter/Private/Tasks/GLTFMaterialTasks.cpp
+++ b/Engine/Plugins/Enterprise/GLTFExporter/MigratedFromGit/Source/GLTFExporter/Private/Tasks/GLTFMaterialTasks.cpp
@@ -169,12 +169,17 @@ void FGLTFMaterialTask::Complete()
 
 		if (UVOverlap > UVOverlapThreshold)
 		{
+			FString SectionString = TEXT("mesh section");
+			SectionString += SectionIndices.Num() > 1 ? TEXT("s ") : TEXT(" ");
+			SectionString += FString::JoinBy(SectionIndices, TEXT(", "), [](int32 Index) { return FString::FromInt(Index); });
+
 			Builder.AddWarningMessage(FString::Printf(
-				TEXT("Material %s is baked using mesh data from %s but the lightmap UV (channel %d) are overlapping by %.2f%% and may produce incorrect results"),
+				TEXT("Material %s is baked using mesh data from %s but the lightmap UV (channel %d) are overlapping by %.2f%% (in %s) and may produce incorrect results"),
 				*Material->GetName(),
 				*ParentMeshData->Name,
 				MeshData->TexCoord,
-				UVOverlap * 100));
+				UVOverlap * 100,
+				*SectionString));
 		}
 	}
 }
