[GLTFExporter] refactor: minor code simplification in material task

commit c7774f41c6c498a077b824d776d2b71888478499
Date: Tue, 24 Nov 2020 14:42:01 +0100

    refactor: minor code simplification in material task
    (cherry picked from commit edb4ea1c6e9967e95ef6f36ea1446b40ffe78bcd​)

#skipci

diff --git a/Engine/Plugins/Enterprise/GLTFExporter/MigratedFromGit/Source/GLTFExporter/Private/Tasks/GLTFMaterialTasks.cpp b/Engine/Plugins/Enterprise/GLTFExporter/MigratedFromGit/Source/GLTFExporter/Private/Tasks/GLTFMaterialTasks.cpp
--- a/Engine/Plugins/Enterprise/GLTFExporter/MigratedFromGit/Source/GLTFExporter/Private/Tasks/GLTFMaterialTasks.cpp
+++ b/Engine/Plugins/Enterprise/GLTFExporter/MigratedFromGit/Source/GLTFExporter/Private/Tasks/GLTFMaterialTasks.cpp
@@ -171,7 +171,7 @@ void FGLTFMaterialTask::Complete()
 		{
 			FString SectionString = TEXT("mesh section");
 			SectionString += SectionIndices.Num() > 1 ? TEXT("s ") : TEXT(" ");
-			SectionString += FString::JoinBy(SectionIndices, TEXT(", "), [](int32 Index) { return FString::FromInt(Index); });
+			SectionString += FString::JoinBy(SectionIndices, TEXT(", "), FString::FromInt);
 
 			Builder.AddWarningMessage(FString::Printf(
 				TEXT("Material %s is baked using mesh data from %s but the lightmap UV (channel %d) are overlapping by %.2f%% (in %s) and may produce incorrect results"),
@@ -1273,7 +1273,7 @@ FGLTFPropertyBakeOutput FGLTFMaterialTask::BakeMaterialProperty(const FMaterialP
 	else
 	{
 		OutTexCoord = MeshData->TexCoord;
-		MeshDataBakedProperties.AddUnique(Property);
+		MeshDataBakedProperties.Add(Property);
 	}
 
 	// TODO: add support for calculating the ideal resolution to use for baking based on connected (texture) nodes
diff --git a/Engine/Plugins/Enterprise/GLTFExporter/MigratedFromGit/Source/GLTFExporter/Private/Tasks/GLTFMaterialTasks.h b/Engine/Plugins/Enterprise/GLTFExporter/MigratedFromGit/Source/GLTFExporter/Private/Tasks/GLTFMaterialTasks.h
--- a/Engine/Plugins/Enterprise/GLTFExporter/MigratedFromGit/Source/GLTFExporter/Private/Tasks/GLTFMaterialTasks.h
+++ b/Engine/Plugins/Enterprise/GLTFExporter/MigratedFromGit/Source/GLTFExporter/Private/Tasks/GLTFMaterialTasks.h
@@ -42,7 +42,7 @@ private:
 	const TArray<int32> SectionIndices;
 	const FGLTFJsonMaterialIndex MaterialIndex;
 
-	TArray<FMaterialPropertyEx> MeshDataBakedProperties;
+	TSet<FMaterialPropertyEx> MeshDataBakedProperties;
 
 	void ConvertAlphaMode(EGLTFJsonAlphaMode& OutAlphaMode, EGLTFJsonBlendMode& OutBlendMode) const;
 	void ConvertShadingModel(EGLTFJsonShadingModel& OutShadingModel) const;
