[GLTFExporter] feat: add warning if custom extensions used but web viewer is not bundled

commit 6935a9395050ccac1046ba47f5894ec7c1bb52f5
Date: Tue, 24 Nov 2020 14:38:56 +0100

    feat: add warning if custom extensions used but web viewer is not bundled
    (cherry picked from commit 052beac6deebd03c3bfc5c0d8efebd08177b90ef​)

#skipci

diff --git a/Engine/Plugins/Enterprise/GLTFExporter/MigratedFromGit/Source/GLTFExporter/Private/Builders/GLTFJsonBuilder.cpp b/Engine/Plugins/Enterprise/GLTFExporter/MigratedFromGit/Source/GLTFExporter/Private/Builders/GLTFJsonBuilder.cpp
--- a/Engine/Plugins/Enterprise/GLTFExporter/MigratedFromGit/Source/GLTFExporter/Private/Builders/GLTFJsonBuilder.cpp
+++ b/Engine/Plugins/Enterprise/GLTFExporter/MigratedFromGit/Source/GLTFExporter/Private/Builders/GLTFJsonBuilder.cpp
@@ -13,6 +13,24 @@ void FGLTFJsonBuilder::WriteJson(FArchive& Archive) const
 	JsonRoot.ToJson(&Archive, bIsGlbFile, ExportOptions->bAllExtensionsRequired);
 }
 
+TSet<EGLTFJsonExtension> FGLTFJsonBuilder::GetCustomExtensionsUsed() const
+{
+	static const TSet<EGLTFJsonExtension> AllCustomExtensions =
+	{
+		EGLTFJsonExtension::EPIC_AnimationHotspots,
+		EGLTFJsonExtension::EPIC_AnimationPlayback,
+		EGLTFJsonExtension::EPIC_BlendModes,
+		EGLTFJsonExtension::EPIC_CameraControls,
+		EGLTFJsonExtension::EPIC_HDRIBackdrops,
+		EGLTFJsonExtension::EPIC_LevelVariantSets,
+		EGLTFJsonExtension::EPIC_LightmapTextures,
+		EGLTFJsonExtension::EPIC_SkySpheres,
+		EGLTFJsonExtension::EPIC_TextureHDREncoding
+	};
+
+	return JsonRoot.Extensions.Used.Intersect(AllCustomExtensions);
+}
+
 void FGLTFJsonBuilder::AddExtension(EGLTFJsonExtension Extension, bool bIsRequired)
 {
 	JsonRoot.Extensions.Used.Add(Extension);
diff --git a/Engine/Plugins/Enterprise/GLTFExporter/MigratedFromGit/Source/GLTFExporter/Private/Builders/GLTFJsonBuilder.h b/Engine/Plugins/Enterprise/GLTFExporter/MigratedFromGit/Source/GLTFExporter/Private/Builders/GLTFJsonBuilder.h
--- a/Engine/Plugins/Enterprise/GLTFExporter/MigratedFromGit/Source/GLTFExporter/Private/Builders/GLTFJsonBuilder.h
+++ b/Engine/Plugins/Enterprise/GLTFExporter/MigratedFromGit/Source/GLTFExporter/Private/Builders/GLTFJsonBuilder.h
@@ -13,6 +13,8 @@ protected:
 
 	void WriteJson(FArchive& Archive) const;
 
+	TSet<EGLTFJsonExtension> GetCustomExtensionsUsed() const;
+
 public:
 
 	FGLTFJsonSceneIndex& DefaultScene;
diff --git a/Engine/Plugins/Enterprise/GLTFExporter/MigratedFromGit/Source/GLTFExporter/Private/Builders/GLTFWebBuilder.cpp b/Engine/Plugins/Enterprise/GLTFExporter/MigratedFromGit/Source/GLTFExporter/Private/Builders/GLTFWebBuilder.cpp
--- a/Engine/Plugins/Enterprise/GLTFExporter/MigratedFromGit/Source/GLTFExporter/Private/Builders/GLTFWebBuilder.cpp
+++ b/Engine/Plugins/Enterprise/GLTFExporter/MigratedFromGit/Source/GLTFExporter/Private/Builders/GLTFWebBuilder.cpp
@@ -28,6 +28,20 @@ void FGLTFWebBuilder::Write(FArchive& Archive, FFeedbackContext* Context)
 		BundleWebViewer(ResourcesDir);
 		BundleLaunchHelper(ResourcesDir);
 	}
+	else
+	{
+		const TSet<EGLTFJsonExtension> CustomExtensions = GetCustomExtensionsUsed();
+		if (CustomExtensions.Num() > 0)
+		{
+			const FString ExtensionsString = FString::JoinBy(CustomExtensions, TEXT(", "),
+				[](EGLTFJsonExtension Extension)
+			{
+				return FGLTFJsonUtility::ToString(Extension);
+			});
+
+			AddWarningMessage(FString::Printf(TEXT("Please note that export uses the following extensions that may not be supported in all glTF viewers: %s"), *ExtensionsString));
+		}
+	}
 }
 
 void FGLTFWebBuilder::BundleWebViewer(const FString& ResourcesDir)
diff --git a/Engine/Plugins/Enterprise/GLTFExporter/MigratedFromGit/Source/GLTFExporter/Private/Builders/GLTFWebBuilder.h b/Engine/Plugins/Enterprise/GLTFExporter/MigratedFromGit/Source/GLTFExporter/Private/Builders/GLTFWebBuilder.h
--- a/Engine/Plugins/Enterprise/GLTFExporter/MigratedFromGit/Source/GLTFExporter/Private/Builders/GLTFWebBuilder.h
+++ b/Engine/Plugins/Enterprise/GLTFExporter/MigratedFromGit/Source/GLTFExporter/Private/Builders/GLTFWebBuilder.h
@@ -18,4 +18,6 @@ private:
 	void BundleLaunchHelper(const FString& ResourcesDir);
 
 	static const TCHAR* GetLaunchHelperExecutable();
+
+	static bool IsCustomExtension(EGLTFJsonExtension Value);
 };
