[GLTFExporter] fix: update FGLTFJsonRoot::Extensions to include all detected extensions during write

commit 82859fe69d4ffecf88955db2bb47e3db71817d8f
Date: Tue, 24 Nov 2020 15:01:48 +0100

    fix: update FGLTFJsonRoot::Extensions to include all detected extensions during write
    (cherry picked from commit 41efd90e29852c1ca8aa52ec934ecf294844daec​)

#skipci

diff --git a/Engine/Plugins/Enterprise/GLTFExporter/MigratedFromGit/Source/GLTFExporter/Private/Builders/GLTFContainerBuilder.cpp b/Engine/Plugins/Enterprise/GLTFExporter/MigratedFromGit/Source/GLTFExporter/Private/Builders/GLTFContainerBuilder.cpp
--- a/Engine/Plugins/Enterprise/GLTFExporter/MigratedFromGit/Source/GLTFExporter/Private/Builders/GLTFContainerBuilder.cpp
+++ b/Engine/Plugins/Enterprise/GLTFExporter/MigratedFromGit/Source/GLTFExporter/Private/Builders/GLTFContainerBuilder.cpp
@@ -7,7 +7,7 @@ FGLTFContainerBuilder::FGLTFContainerBuilder(const FString& FilePath, const UGLT
 {
 }
 
-void FGLTFContainerBuilder::WriteGlb(FArchive& Archive) const
+void FGLTFContainerBuilder::WriteGlb(FArchive& Archive)
 {
 	FBufferArchive JsonData;
 	WriteJson(JsonData);
diff --git a/Engine/Plugins/Enterprise/GLTFExporter/MigratedFromGit/Source/GLTFExporter/Private/Builders/GLTFContainerBuilder.h b/Engine/Plugins/Enterprise/GLTFExporter/MigratedFromGit/Source/GLTFExporter/Private/Builders/GLTFContainerBuilder.h
--- a/Engine/Plugins/Enterprise/GLTFExporter/MigratedFromGit/Source/GLTFExporter/Private/Builders/GLTFContainerBuilder.h
+++ b/Engine/Plugins/Enterprise/GLTFExporter/MigratedFromGit/Source/GLTFExporter/Private/Builders/GLTFContainerBuilder.h
@@ -10,7 +10,7 @@ protected:
 
 	FGLTFContainerBuilder(const FString& FilePath, const UGLTFExportOptions* ExportOptions, bool bSelectedActorsOnly);
 
-	void WriteGlb(FArchive& Archive) const;
+	void WriteGlb(FArchive& Archive);
 
 private:
 
diff --git a/Engine/Plugins/Enterprise/GLTFExporter/MigratedFromGit/Source/GLTFExporter/Private/Builders/GLTFJsonBuilder.cpp b/Engine/Plugins/Enterprise/GLTFExporter/MigratedFromGit/Source/GLTFExporter/Private/Builders/GLTFJsonBuilder.cpp
--- a/Engine/Plugins/Enterprise/GLTFExporter/MigratedFromGit/Source/GLTFExporter/Private/Builders/GLTFJsonBuilder.cpp
+++ b/Engine/Plugins/Enterprise/GLTFExporter/MigratedFromGit/Source/GLTFExporter/Private/Builders/GLTFJsonBuilder.cpp
@@ -8,7 +8,7 @@ FGLTFJsonBuilder::FGLTFJsonBuilder(const FString& FilePath, const UGLTFExportOpt
 {
 }
 
-void FGLTFJsonBuilder::WriteJson(FArchive& Archive) const
+void FGLTFJsonBuilder::WriteJson(FArchive& Archive)
 {
 	JsonRoot.ToJson(&Archive, bIsGlbFile, ExportOptions->bAllExtensionsRequired);
 }
diff --git a/Engine/Plugins/Enterprise/GLTFExporter/MigratedFromGit/Source/GLTFExporter/Private/Builders/GLTFJsonBuilder.h b/Engine/Plugins/Enterprise/GLTFExporter/MigratedFromGit/Source/GLTFExporter/Private/Builders/GLTFJsonBuilder.h
--- a/Engine/Plugins/Enterprise/GLTFExporter/MigratedFromGit/Source/GLTFExporter/Private/Builders/GLTFJsonBuilder.h
+++ b/Engine/Plugins/Enterprise/GLTFExporter/MigratedFromGit/Source/GLTFExporter/Private/Builders/GLTFJsonBuilder.h
@@ -11,7 +11,7 @@ protected:
 
 	FGLTFJsonBuilder(const FString& FilePath, const UGLTFExportOptions* ExportOptions);
 
-	void WriteJson(FArchive& Archive) const;
+	void WriteJson(FArchive& Archive);
 
 	TSet<EGLTFJsonExtension> GetCustomExtensionsUsed() const;
 
diff --git a/Engine/Plugins/Enterprise/GLTFExporter/MigratedFromGit/Source/GLTFExporter/Private/Json/GLTFJsonRoot.h b/Engine/Plugins/Enterprise/GLTFExporter/MigratedFromGit/Source/GLTFExporter/Private/Json/GLTFJsonRoot.h
--- a/Engine/Plugins/Enterprise/GLTFExporter/MigratedFromGit/Source/GLTFExporter/Private/Json/GLTFJsonRoot.h
+++ b/Engine/Plugins/Enterprise/GLTFExporter/MigratedFromGit/Source/GLTFExporter/Private/Json/GLTFJsonRoot.h
@@ -84,33 +84,31 @@ struct FGLTFJsonRoot
 	TArray<TUniquePtr<FGLTFJsonVariation>>  Variations;
 
 	template <class CharType = TCHAR, class PrintPolicy = TPrettyJsonPrintPolicy<CharType>>
-	void WriteObject(TJsonWriter<CharType, PrintPolicy>& JsonWriter, bool bAllExtensionsRequired) const
+	void WriteObject(TJsonWriter<CharType, PrintPolicy>& JsonWriter, bool bAllExtensionsRequired)
 	{
-		FGLTFJsonExtensions AllExtensions(Extensions);
-
 		JsonWriter.WriteObjectStart();
 
 		JsonWriter.WriteIdentifierPrefix(TEXT("asset"));
-		Asset.WriteObject(JsonWriter, AllExtensions);
+		Asset.WriteObject(JsonWriter, Extensions);
 
 		if (DefaultScene != INDEX_NONE)
 		{
 			JsonWriter.WriteValue(TEXT("scene"), DefaultScene);
 		}
 
-		FGLTFJsonUtility::WriteObjectPtrArray(JsonWriter, TEXT("accessors"), Accessors, AllExtensions);
-		FGLTFJsonUtility::WriteObjectPtrArray(JsonWriter, TEXT("animations"), Animations, AllExtensions);
-		FGLTFJsonUtility::WriteObjectPtrArray(JsonWriter, TEXT("buffers"), Buffers, AllExtensions);
-		FGLTFJsonUtility::WriteObjectPtrArray(JsonWriter, TEXT("bufferViews"), BufferViews, AllExtensions);
-		FGLTFJsonUtility::WriteObjectPtrArray(JsonWriter, TEXT("cameras"), Cameras, AllExtensions);
-		FGLTFJsonUtility::WriteObjectPtrArray(JsonWriter, TEXT("images"), Images, AllExtensions);
-		FGLTFJsonUtility::WriteObjectPtrArray(JsonWriter, TEXT("materials"), Materials, AllExtensions);
-		FGLTFJsonUtility::WriteObjectPtrArray(JsonWriter, TEXT("meshes"), Meshes, AllExtensions);
-		FGLTFJsonUtility::WriteObjectPtrArray(JsonWriter, TEXT("nodes"), Nodes, AllExtensions);
-		FGLTFJsonUtility::WriteObjectPtrArray(JsonWriter, TEXT("samplers"), Samplers, AllExtensions);
-		FGLTFJsonUtility::WriteObjectPtrArray(JsonWriter, TEXT("scenes"), Scenes, AllExtensions);
-		FGLTFJsonUtility::WriteObjectPtrArray(JsonWriter, TEXT("skins"), Skins, AllExtensions);
-		FGLTFJsonUtility::WriteObjectPtrArray(JsonWriter, TEXT("textures"), Textures, AllExtensions);
+		FGLTFJsonUtility::WriteObjectPtrArray(JsonWriter, TEXT("accessors"), Accessors, Extensions);
+		FGLTFJsonUtility::WriteObjectPtrArray(JsonWriter, TEXT("animations"), Animations, Extensions);
+		FGLTFJsonUtility::WriteObjectPtrArray(JsonWriter, TEXT("buffers"), Buffers, Extensions);
+		FGLTFJsonUtility::WriteObjectPtrArray(JsonWriter, TEXT("bufferViews"), BufferViews, Extensions);
+		FGLTFJsonUtility::WriteObjectPtrArray(JsonWriter, TEXT("cameras"), Cameras, Extensions);
+		FGLTFJsonUtility::WriteObjectPtrArray(JsonWriter, TEXT("images"), Images, Extensions);
+		FGLTFJsonUtility::WriteObjectPtrArray(JsonWriter, TEXT("materials"), Materials, Extensions);
+		FGLTFJsonUtility::WriteObjectPtrArray(JsonWriter, TEXT("meshes"), Meshes, Extensions);
+		FGLTFJsonUtility::WriteObjectPtrArray(JsonWriter, TEXT("nodes"), Nodes, Extensions);
+		FGLTFJsonUtility::WriteObjectPtrArray(JsonWriter, TEXT("samplers"), Samplers, Extensions);
+		FGLTFJsonUtility::WriteObjectPtrArray(JsonWriter, TEXT("scenes"), Scenes, Extensions);
+		FGLTFJsonUtility::WriteObjectPtrArray(JsonWriter, TEXT("skins"), Skins, Extensions);
+		FGLTFJsonUtility::WriteObjectPtrArray(JsonWriter, TEXT("textures"), Textures, Extensions);
 
 		if (Backdrops.Num() > 0 || Hotspots.Num() > 0 || Lights.Num() > 0 || LightMaps.Num() > 0 || SkySpheres.Num() > 0 || Variations.Num() > 0)
 		{
@@ -119,81 +117,86 @@ struct FGLTFJsonRoot
 			if (Backdrops.Num() > 0)
 			{
 				const EGLTFJsonExtension Extension = EGLTFJsonExtension::EPIC_HDRIBackdrops;
-				AllExtensions.Used.Add(Extension);
+				Extensions.Used.Add(Extension);
 
 				JsonWriter.WriteObjectStart(FGLTFJsonUtility::ToString(Extension));
-				FGLTFJsonUtility::WriteObjectPtrArray(JsonWriter, TEXT("backdrops"), Backdrops, AllExtensions);
+				FGLTFJsonUtility::WriteObjectPtrArray(JsonWriter, TEXT("backdrops"), Backdrops, Extensions);
 				JsonWriter.WriteObjectEnd();
 			}
 
 			if (Hotspots.Num() > 0)
 			{
 				const EGLTFJsonExtension Extension = EGLTFJsonExtension::EPIC_AnimationHotspots;
-				AllExtensions.Used.Add(Extension);
+				Extensions.Used.Add(Extension);
 
 				JsonWriter.WriteObjectStart(FGLTFJsonUtility::ToString(Extension));
-				FGLTFJsonUtility::WriteObjectPtrArray(JsonWriter, TEXT("hotspots"), Hotspots, AllExtensions);
+				FGLTFJsonUtility::WriteObjectPtrArray(JsonWriter, TEXT("hotspots"), Hotspots, Extensions);
 				JsonWriter.WriteObjectEnd();
 			}
 
 			if (Variations.Num() > 0)
 			{
 				const EGLTFJsonExtension Extension = EGLTFJsonExtension::EPIC_LevelVariantSets;
-				AllExtensions.Used.Add(Extension);
+				Extensions.Used.Add(Extension);
 
 				JsonWriter.WriteObjectStart(FGLTFJsonUtility::ToString(Extension));
-				FGLTFJsonUtility::WriteObjectPtrArray(JsonWriter, TEXT("levelVariantSets"), Variations, AllExtensions);
+				FGLTFJsonUtility::WriteObjectPtrArray(JsonWriter, TEXT("levelVariantSets"), Variations, Extensions);
 				JsonWriter.WriteObjectEnd();
 			}
 
 			if (Lights.Num() > 0)
 			{
 				const EGLTFJsonExtension Extension = EGLTFJsonExtension::KHR_LightsPunctual;
-				AllExtensions.Used.Add(Extension);
+				Extensions.Used.Add(Extension);
 
 				JsonWriter.WriteObjectStart(FGLTFJsonUtility::ToString(Extension));
-				FGLTFJsonUtility::WriteObjectPtrArray(JsonWriter, TEXT("lights"), Lights, AllExtensions);
+				FGLTFJsonUtility::WriteObjectPtrArray(JsonWriter, TEXT("lights"), Lights, Extensions);
 				JsonWriter.WriteObjectEnd();
 			}
 
 			if (LightMaps.Num() > 0)
 			{
 				const EGLTFJsonExtension Extension = EGLTFJsonExtension::EPIC_LightmapTextures;
-				AllExtensions.Used.Add(Extension);
+				Extensions.Used.Add(Extension);
 
 				JsonWriter.WriteObjectStart(FGLTFJsonUtility::ToString(Extension));
-				FGLTFJsonUtility::WriteObjectPtrArray(JsonWriter, TEXT("lightmaps"), LightMaps, AllExtensions);
+				FGLTFJsonUtility::WriteObjectPtrArray(JsonWriter, TEXT("lightmaps"), LightMaps, Extensions);
 				JsonWriter.WriteObjectEnd();
 			}
 
 			if (SkySpheres.Num() > 0)
 			{
 				const EGLTFJsonExtension Extension = EGLTFJsonExtension::EPIC_SkySpheres;
-				AllExtensions.Used.Add(Extension);
+				Extensions.Used.Add(Extension);
 
 				JsonWriter.WriteObjectStart(FGLTFJsonUtility::ToString(Extension));
-				FGLTFJsonUtility::WriteObjectPtrArray(JsonWriter, TEXT("skySpheres"), SkySpheres, AllExtensions);
+				FGLTFJsonUtility::WriteObjectPtrArray(JsonWriter, TEXT("skySpheres"), SkySpheres, Extensions);
 				JsonWriter.WriteObjectEnd();
 			}
 
 			JsonWriter.WriteObjectEnd();
 		}
 
-		FGLTFJsonUtility::WriteStringArray(JsonWriter, TEXT("extensionsUsed"), AllExtensions.Used);
-		FGLTFJsonUtility::WriteStringArray(JsonWriter, TEXT("extensionsRequired"), bAllExtensionsRequired ? AllExtensions.Used : AllExtensions.Required);
+		if (bAllExtensionsRequired)
+		{
+			Extensions.Required.Append(Extensions.Used);
+		}
+
+		FGLTFJsonUtility::WriteStringArray(JsonWriter, TEXT("extensionsUsed"), Extensions.Used);
+		FGLTFJsonUtility::WriteStringArray(JsonWriter, TEXT("extensionsRequired"), Extensions.Required);
 
 		JsonWriter.WriteObjectEnd();
 	}
 
 	template <class CharType = TCHAR, class PrintPolicy = TPrettyJsonPrintPolicy<CharType>>
-	void ToJson(FArchive* const Archive, bool bAllExtensionsRequired) const
+	void ToJson(FArchive* const Archive, bool bAllExtensionsRequired)
 	{
 		TSharedRef<TJsonWriter<CharType, PrintPolicy>> JsonWriter = TJsonWriterFactory<CharType, PrintPolicy>::Create(Archive);
 		WriteObject(*JsonWriter, bAllExtensionsRequired);
 		JsonWriter->Close();
 	}
 
-	void ToJson(FArchive* const Archive, bool bCondensedPrint, bool bAllExtensionsRequired) const
+	void ToJson(FArchive* const Archive, bool bCondensedPrint, bool bAllExtensionsRequired)
 	{
 		if (bCondensedPrint)
 		{
