[GLTFExporter] refactor: remove need to manually maintain identify of each custom extension in GLTFJsonBuilder

commit de5cda0d8c04682bb5220ecaa5f79099aa7e5161
Date: Tue, 24 Nov 2020 15:51:43 +0100

    refactor: remove need to manually maintain identify of each custom extension in GLTFJsonBuilder
    (cherry picked from commit 65641f99e8e0f300238b0c8386e9c9f431a477d9​)

#skipci

diff --git a/Engine/Plugins/Enterprise/GLTFExporter/MigratedFromGit/Source/GLTFExporter/Private/Builders/GLTFJsonBuilder.cpp b/Engine/Plugins/Enterprise/GLTFExporter/MigratedFromGit/Source/GLTFExporter/Private/Builders/GLTFJsonBuilder.cpp
--- a/Engine/Plugins/Enterprise/GLTFExporter/MigratedFromGit/Source/GLTFExporter/Private/Builders/GLTFJsonBuilder.cpp
+++ b/Engine/Plugins/Enterprise/GLTFExporter/MigratedFromGit/Source/GLTFExporter/Private/Builders/GLTFJsonBuilder.cpp
@@ -15,20 +15,20 @@ void FGLTFJsonBuilder::WriteJson(FArchive& Archive)
 
 TSet<EGLTFJsonExtension> FGLTFJsonBuilder::GetCustomExtensionsUsed() const
 {
-	static const TSet<EGLTFJsonExtension> AllCustomExtensions =
+	const TCHAR CustomPrefix[] = TEXT("EPIC_");
+
+	TSet<EGLTFJsonExtension> CustomExtensions;
+
+	for (EGLTFJsonExtension Extension : JsonRoot.Extensions.Used)
 	{
-		EGLTFJsonExtension::EPIC_AnimationHotspots,
-		EGLTFJsonExtension::EPIC_AnimationPlayback,
-		EGLTFJsonExtension::EPIC_BlendModes,
-		EGLTFJsonExtension::EPIC_CameraControls,
-		EGLTFJsonExtension::EPIC_HDRIBackdrops,
-		EGLTFJsonExtension::EPIC_LevelVariantSets,
-		EGLTFJsonExtension::EPIC_LightmapTextures,
-		EGLTFJsonExtension::EPIC_SkySpheres,
-		EGLTFJsonExtension::EPIC_TextureHDREncoding
-	};
-
-	return JsonRoot.Extensions.Used.Intersect(AllCustomExtensions);
+		const TCHAR* ExtensionString = FGLTFJsonUtility::ToString(Extension);
+		if (FCString::Strncmp(ExtensionString, CustomPrefix, (sizeof(CustomPrefix) / sizeof(TCHAR)) - 1) == 0)
+		{
+			CustomExtensions.Add(Extension);
+		}
+	}
+
+	return CustomExtensions;
 }
 
 void FGLTFJsonBuilder::AddExtension(EGLTFJsonExtension Extension, bool bIsRequired)
