[GLTFExporter] refactor: rename export option ExportScale to ExportUniformScale, to match FBX options

commit f30dd12388a5a83b71ab920d3fc8b5e7d460fd46
Date: Tue, 24 Nov 2020 21:28:45 +0100

    refactor: rename export option ExportScale to ExportUniformScale, to match FBX options
    (cherry picked from commit bf3b2c98c693c7da1a5dcb605b3a743ebe2b22f5​)

#skipci

diff --git a/Engine/Plugins/Enterprise/GLTFExporter/MigratedFromGit/Source/GLTFExporter/Private/Converters/GLTFAccessorConverters.cpp b/Engine/Plugins/Enterprise/GLTFExporter/MigratedFromGit/Source/GLTFExporter/Private/Converters/GLTFAccessorConverters.cpp
--- a/Engine/Plugins/Enterprise/GLTFExporter/MigratedFromGit/Source/GLTFExporter/Private/Converters/GLTFAccessorConverters.cpp
+++ b/Engine/Plugins/Enterprise/GLTFExporter/MigratedFromGit/Source/GLTFExporter/Private/Converters/GLTFAccessorConverters.cpp
@@ -60,7 +60,7 @@ FGLTFJsonAccessorIndex FGLTFPositionBufferConverter::Convert(const FGLTFMeshSect
 	for (uint32 VertexIndex = 0; VertexIndex < VertexCount; ++VertexIndex)
 	{
 		const uint32 MappedVertexIndex = IndexMap[VertexIndex];
-		Positions[VertexIndex] = FGLTFConverterUtility::ConvertPosition(VertexBuffer->VertexPosition(MappedVertexIndex), Builder.ExportOptions->ExportScale);
+		Positions[VertexIndex] = FGLTFConverterUtility::ConvertPosition(VertexBuffer->VertexPosition(MappedVertexIndex), Builder.ExportOptions->ExportUniformScale);
 	}
 
 	// More accurate bounding box if based on raw vertex values
diff --git a/Engine/Plugins/Enterprise/GLTFExporter/MigratedFromGit/Source/GLTFExporter/Private/Converters/GLTFBackdropConverters.cpp b/Engine/Plugins/Enterprise/GLTFExporter/MigratedFromGit/Source/GLTFExporter/Private/Converters/GLTFBackdropConverters.cpp
--- a/Engine/Plugins/Enterprise/GLTFExporter/MigratedFromGit/Source/GLTFExporter/Private/Converters/GLTFBackdropConverters.cpp
+++ b/Engine/Plugins/Enterprise/GLTFExporter/MigratedFromGit/Source/GLTFExporter/Private/Converters/GLTFBackdropConverters.cpp
@@ -77,7 +77,7 @@ FGLTFJsonBackdropIndex FGLTFBackdropConverter::Convert(const AActor* BackdropAct
 	FVector ProjectionCenter;
 	if (FGLTFActorUtility::TryGetPropertyValue(BackdropActor, TEXT("ProjectionCenter"), ProjectionCenter))
 	{
-		JsonBackdrop.ProjectionCenter = FGLTFConverterUtility::ConvertPosition(ProjectionCenter, Builder.ExportOptions->ExportScale);
+		JsonBackdrop.ProjectionCenter = FGLTFConverterUtility::ConvertPosition(ProjectionCenter, Builder.ExportOptions->ExportUniformScale);
 	}
 	else
 	{
diff --git a/Engine/Plugins/Enterprise/GLTFExporter/MigratedFromGit/Source/GLTFExporter/Private/Converters/GLTFCameraConverters.cpp b/Engine/Plugins/Enterprise/GLTFExporter/MigratedFromGit/Source/GLTFExporter/Private/Converters/GLTFCameraConverters.cpp
--- a/Engine/Plugins/Enterprise/GLTFExporter/MigratedFromGit/Source/GLTFExporter/Private/Converters/GLTFCameraConverters.cpp
+++ b/Engine/Plugins/Enterprise/GLTFExporter/MigratedFromGit/Source/GLTFExporter/Private/Converters/GLTFCameraConverters.cpp
@@ -14,7 +14,7 @@ FGLTFJsonCameraIndex FGLTFCameraConverter::Convert(const UCameraComponent* Camer
 
 	FMinimalViewInfo DesiredView;
 	const_cast<UCameraComponent*>(CameraComponent)->GetCameraView(0, DesiredView);
-	const float ExportScale = Builder.ExportOptions->ExportScale;
+	const float ExportScale = Builder.ExportOptions->ExportUniformScale;
 
 	switch (Camera.Type)
 	{
diff --git a/Engine/Plugins/Enterprise/GLTFExporter/MigratedFromGit/Source/GLTFExporter/Private/Converters/GLTFLightConverters.cpp b/Engine/Plugins/Enterprise/GLTFExporter/MigratedFromGit/Source/GLTFExporter/Private/Converters/GLTFLightConverters.cpp
--- a/Engine/Plugins/Enterprise/GLTFExporter/MigratedFromGit/Source/GLTFExporter/Private/Converters/GLTFLightConverters.cpp
+++ b/Engine/Plugins/Enterprise/GLTFExporter/MigratedFromGit/Source/GLTFExporter/Private/Converters/GLTFLightConverters.cpp
@@ -25,7 +25,7 @@ FGLTFJsonLightIndex FGLTFLightConverter::Convert(const ULightComponent* LightCom
 
 	if (const UPointLightComponent* PointLightComponent = Cast<UPointLightComponent>(LightComponent))
 	{
-		Light.Range = FGLTFConverterUtility::ConvertLength(PointLightComponent->AttenuationRadius, Builder.ExportOptions->ExportScale);
+		Light.Range = FGLTFConverterUtility::ConvertLength(PointLightComponent->AttenuationRadius, Builder.ExportOptions->ExportUniformScale);
 	}
 
 	if (const USpotLightComponent* SpotLightComponent = Cast<USpotLightComponent>(LightComponent))
diff --git a/Engine/Plugins/Enterprise/GLTFExporter/MigratedFromGit/Source/GLTFExporter/Private/Converters/GLTFNodeConverters.cpp b/Engine/Plugins/Enterprise/GLTFExporter/MigratedFromGit/Source/GLTFExporter/Private/Converters/GLTFNodeConverters.cpp
--- a/Engine/Plugins/Enterprise/GLTFExporter/MigratedFromGit/Source/GLTFExporter/Private/Converters/GLTFNodeConverters.cpp
+++ b/Engine/Plugins/Enterprise/GLTFExporter/MigratedFromGit/Source/GLTFExporter/Private/Converters/GLTFNodeConverters.cpp
@@ -108,7 +108,7 @@ FGLTFJsonNodeIndex FGLTFComponentConverter::Convert(const USceneComponent* Scene
 	const FGLTFJsonNodeIndex NodeIndex = Builder.AddChildNode(ParentNodeIndex);
 	FGLTFJsonNode& Node = Builder.GetNode(NodeIndex);
 	Node.Name = FGLTFNameUtility::GetName(SceneComponent);
-	Node.Translation = FGLTFConverterUtility::ConvertPosition(RelativeTransform.GetTranslation(), Builder.ExportOptions->ExportScale);
+	Node.Translation = FGLTFConverterUtility::ConvertPosition(RelativeTransform.GetTranslation(), Builder.ExportOptions->ExportUniformScale);
 	Node.Rotation = FGLTFConverterUtility::ConvertRotation(RelativeTransform.GetRotation());
 	Node.Scale = FGLTFConverterUtility::ConvertScale(RelativeTransform.GetScale3D());
 
@@ -214,7 +214,7 @@ FGLTFJsonNodeIndex FGLTFStaticSocketConverter::Convert(FGLTFJsonNodeIndex RootNo
 	Node.Name = SocketName.ToString();
 
 	// TODO: add warning check for non-uniform scaling
-	Node.Translation = FGLTFConverterUtility::ConvertPosition(Socket->RelativeLocation, Builder.ExportOptions->ExportScale);
+	Node.Translation = FGLTFConverterUtility::ConvertPosition(Socket->RelativeLocation, Builder.ExportOptions->ExportUniformScale);
 	Node.Rotation = FGLTFConverterUtility::ConvertRotation(Socket->RelativeRotation.Quaternion());
 	Node.Scale = FGLTFConverterUtility::ConvertScale(Socket->RelativeScale);
 
@@ -230,7 +230,7 @@ FGLTFJsonNodeIndex FGLTFSkeletalSocketConverter::Convert(FGLTFJsonNodeIndex Root
 		Node.Name = SocketName.ToString();
 
 		// TODO: add warning check for non-uniform scaling
-		Node.Translation = FGLTFConverterUtility::ConvertPosition(Socket->RelativeLocation, Builder.ExportOptions->ExportScale);
+		Node.Translation = FGLTFConverterUtility::ConvertPosition(Socket->RelativeLocation, Builder.ExportOptions->ExportUniformScale);
 		Node.Rotation = FGLTFConverterUtility::ConvertRotation(Socket->RelativeRotation.Quaternion());
 		Node.Scale = FGLTFConverterUtility::ConvertScale(Socket->RelativeScale);
 
@@ -270,7 +270,7 @@ FGLTFJsonNodeIndex FGLTFSkeletalBoneConverter::Convert(FGLTFJsonNodeIndex RootNo
 	{
 		// TODO: add warning check for non-uniform scaling
 		const FTransform& BonePose = BonePoses[BoneIndex];
-		Node.Translation = FGLTFConverterUtility::ConvertPosition(BonePose.GetTranslation(), Builder.ExportOptions->ExportScale);
+		Node.Translation = FGLTFConverterUtility::ConvertPosition(BonePose.GetTranslation(), Builder.ExportOptions->ExportUniformScale);
 		Node.Rotation = FGLTFConverterUtility::ConvertRotation(BonePose.GetRotation());
 		Node.Scale = FGLTFConverterUtility::ConvertScale(BonePose.GetScale3D());
 	}
diff --git a/Engine/Plugins/Enterprise/GLTFExporter/MigratedFromGit/Source/GLTFExporter/Private/Converters/GLTFSkinConverters.cpp b/Engine/Plugins/Enterprise/GLTFExporter/MigratedFromGit/Source/GLTFExporter/Private/Converters/GLTFSkinConverters.cpp
--- a/Engine/Plugins/Enterprise/GLTFExporter/MigratedFromGit/Source/GLTFExporter/Private/Converters/GLTFSkinConverters.cpp
+++ b/Engine/Plugins/Enterprise/GLTFExporter/MigratedFromGit/Source/GLTFExporter/Private/Converters/GLTFSkinConverters.cpp
@@ -31,7 +31,7 @@ FGLTFJsonSkinIndex FGLTFSkinConverter::Convert(FGLTFJsonNodeIndex RootNode, cons
 	for (int32 BoneIndex = 0; BoneIndex < BoneCount; ++BoneIndex)
 	{
 		const FTransform InverseBindTransform = FGLTFBoneUtility::GetBindTransform(SkeletalMesh->RefSkeleton, BoneIndex).Inverse();
-		InverseBindMatrices[BoneIndex] = FGLTFConverterUtility::ConvertTransform(InverseBindTransform, Builder.ExportOptions->ExportScale);
+		InverseBindMatrices[BoneIndex] = FGLTFConverterUtility::ConvertTransform(InverseBindTransform, Builder.ExportOptions->ExportUniformScale);
 	}
 
 	FGLTFJsonAccessor JsonAccessor;
diff --git a/Engine/Plugins/Enterprise/GLTFExporter/MigratedFromGit/Source/GLTFExporter/Private/GLTFExportOptions.cpp b/Engine/Plugins/Enterprise/GLTFExporter/MigratedFromGit/Source/GLTFExporter/Private/GLTFExportOptions.cpp
--- a/Engine/Plugins/Enterprise/GLTFExporter/MigratedFromGit/Source/GLTFExporter/Private/GLTFExportOptions.cpp
+++ b/Engine/Plugins/Enterprise/GLTFExporter/MigratedFromGit/Source/GLTFExporter/Private/GLTFExportOptions.cpp
@@ -32,7 +32,7 @@ void UGLTFExportOptions::ResetToDefault()
 	bExportTextureTransforms = true;
 	bExportLightmaps = true;
 	TextureHDREncoding = EGLTFTextureHDREncoding::RGBM;
-	ExportScale = 0.01;
+	ExportUniformScale = 0.01;
 	bExportHiddenInGame = false;
 	ExportLights = EGLTFExportLightMobility::MovableAndStationary;
 	bExportCameras = true;
diff --git a/Engine/Plugins/Enterprise/GLTFExporter/MigratedFromGit/Source/GLTFExporter/Private/Tasks/GLTFAnimationTasks.cpp b/Engine/Plugins/Enterprise/GLTFExporter/MigratedFromGit/Source/GLTFExporter/Private/Tasks/GLTFAnimationTasks.cpp
--- a/Engine/Plugins/Enterprise/GLTFExporter/MigratedFromGit/Source/GLTFExporter/Private/Tasks/GLTFAnimationTasks.cpp
+++ b/Engine/Plugins/Enterprise/GLTFExporter/MigratedFromGit/Source/GLTFExporter/Private/Tasks/GLTFAnimationTasks.cpp
@@ -96,7 +96,7 @@ void FGLTFAnimSequenceTask::Complete()
 			for (int32 Key = 0; Key < KeyPositions.Num(); ++Key)
 			{
 				const FVector KeyPosition = KeyTransforms[Key].GetTranslation();
-				Translations[Key] = FGLTFConverterUtility::ConvertPosition(KeyPosition, Builder.ExportOptions->ExportScale);
+				Translations[Key] = FGLTFConverterUtility::ConvertPosition(KeyPosition, Builder.ExportOptions->ExportUniformScale);
 			}
 
 			JsonInputAccessor.Count = Translations.Num();
@@ -299,7 +299,7 @@ void FGLTFLevelSequenceTask::Complete()
 							Channels[1]->Evaluate(FrameTime, Translation.Y);
 							Channels[2]->Evaluate(FrameTime, Translation.Z);
 
-							Translations[Frame] = FGLTFConverterUtility::ConvertPosition(Translation, Builder.ExportOptions->ExportScale);
+							Translations[Frame] = FGLTFConverterUtility::ConvertPosition(Translation, Builder.ExportOptions->ExportUniformScale);
 						}
 
 						FGLTFJsonAccessor JsonOutputAccessor;
diff --git a/Engine/Plugins/Enterprise/GLTFExporter/MigratedFromGit/Source/GLTFExporter/Public/GLTFExportOptions.h b/Engine/Plugins/Enterprise/GLTFExporter/MigratedFromGit/Source/GLTFExporter/Public/GLTFExportOptions.h
--- a/Engine/Plugins/Enterprise/GLTFExporter/MigratedFromGit/Source/GLTFExporter/Public/GLTFExportOptions.h
+++ b/Engine/Plugins/Enterprise/GLTFExporter/MigratedFromGit/Source/GLTFExporter/Public/GLTFExportOptions.h
@@ -115,10 +115,9 @@ public:
 	UPROPERTY(EditAnywhere, BlueprintReadWrite, Config, Category = Texture, Meta = (DisplayName = "Texture HDR Encoding"))
 	EGLTFTextureHDREncoding TextureHDREncoding;
 
-	UPROPERTY(EditAnywhere, BlueprintReadWrite, Config, Category = Scene, Meta =
-		(DisplayName = "Export Uniform Scale",
-		 ToolTip = "Scale factor used for exporting assets, by default: 0.01, for conversion from centimeters(Unreal default) to meters(glTF)."))
-	float ExportScale;
+	/** Scale factor used for exporting assets (0.01 by default) for conversion from centimeters (Unreal default) to meters (glTF). */
+	UPROPERTY(EditAnywhere, BlueprintReadWrite, Config, Category = Scene)
+	float ExportUniformScale;
 
 	UPROPERTY(EditAnywhere, BlueprintReadWrite, Config, Category = Scene)
 	bool bExportHiddenInGame;
