[GLTFExporter] refactor: simplify warning message if custom extensions used but web viewer not bundled

commit f15f54da801371f8f7a866727c05e494318c7418
Date: Tue, 24 Nov 2020 22:23:06 +0100

    refactor: simplify warning message if custom extensions used but web viewer not bundled
    (cherry picked from commit c903f34126d2e4b9144167fca7df486692dcf0c9​)

#skipci

diff --git a/Engine/Plugins/Enterprise/GLTFExporter/MigratedFromGit/Source/GLTFExporter/Private/Builders/GLTFWebBuilder.cpp b/Engine/Plugins/Enterprise/GLTFExporter/MigratedFromGit/Source/GLTFExporter/Private/Builders/GLTFWebBuilder.cpp
--- a/Engine/Plugins/Enterprise/GLTFExporter/MigratedFromGit/Source/GLTFExporter/Private/Builders/GLTFWebBuilder.cpp
+++ b/Engine/Plugins/Enterprise/GLTFExporter/MigratedFromGit/Source/GLTFExporter/Private/Builders/GLTFWebBuilder.cpp
@@ -39,7 +39,7 @@ void FGLTFWebBuilder::Write(FArchive& Archive, FFeedbackContext* Context)
 				return FGLTFJsonUtility::ToString(Extension);
 			});
 
-			AddWarningMessage(FString::Printf(TEXT("Please note that export uses the following extensions that may not be supported in all glTF viewers: %s"), *ExtensionsString));
+			AddWarningMessage(FString::Printf(TEXT("Export uses some extensions that may only be supported in Unreal's bundled glTF viewer: %s"), *ExtensionsString));
 		}
 	}
 }
