[GLTFExporter] refactor: rearrange general export options and remove bAllExtensionsRequired

commit 52ba5f09b88980f09bef248d0dd98fa217f20090
Date: Tue, 24 Nov 2020 22:33:13 +0100

    refactor: rearrange general export options and remove bAllExtensionsRequired
    (cherry picked from commit eadd86a70a777c5b692f91ffab3a18ff5d3c7715​)

#skipci

diff --git a/Engine/Plugins/Enterprise/GLTFExporter/MigratedFromGit/Source/GLTFExporter/Private/Builders/GLTFJsonBuilder.cpp b/Engine/Plugins/Enterprise/GLTFExporter/MigratedFromGit/Source/GLTFExporter/Private/Builders/GLTFJsonBuilder.cpp
--- a/Engine/Plugins/Enterprise/GLTFExporter/MigratedFromGit/Source/GLTFExporter/Private/Builders/GLTFJsonBuilder.cpp
+++ b/Engine/Plugins/Enterprise/GLTFExporter/MigratedFromGit/Source/GLTFExporter/Private/Builders/GLTFJsonBuilder.cpp
@@ -10,7 +10,7 @@ FGLTFJsonBuilder::FGLTFJsonBuilder(const FString& FilePath, const UGLTFExportOpt
 
 void FGLTFJsonBuilder::WriteJson(FArchive& Archive)
 {
-	JsonRoot.ToJson(&Archive, bIsGlbFile, ExportOptions->bAllExtensionsRequired);
+	JsonRoot.ToJson(&Archive, bIsGlbFile);
 }
 
 TSet<EGLTFJsonExtension> FGLTFJsonBuilder::GetCustomExtensionsUsed() const
diff --git a/Engine/Plugins/Enterprise/GLTFExporter/MigratedFromGit/Source/GLTFExporter/Private/GLTFExportOptions.cpp b/Engine/Plugins/Enterprise/GLTFExporter/MigratedFromGit/Source/GLTFExporter/Private/GLTFExportOptions.cpp
--- a/Engine/Plugins/Enterprise/GLTFExporter/MigratedFromGit/Source/GLTFExporter/Private/GLTFExportOptions.cpp
+++ b/Engine/Plugins/Enterprise/GLTFExporter/MigratedFromGit/Source/GLTFExporter/Private/GLTFExportOptions.cpp
@@ -12,6 +12,10 @@ UGLTFExportOptions::UGLTFExportOptions(const FObjectInitializer& ObjectInitializ
 
 void UGLTFExportOptions::ResetToDefault()
 {
+	ExportUniformScale = 0.01;
+	bExportPreviewMesh = true;
+	bBundleWebViewer = true;
+	bShowFilesWhenDone = true;
 	bExportUnlitMaterials = true;
 	bExportClearCoatMaterials = true;
 	bExportExtraBlendModes = true;
@@ -32,7 +36,6 @@ void UGLTFExportOptions::ResetToDefault()
 	bExportTextureTransforms = true;
 	bExportLightmaps = true;
 	TextureHDREncoding = EGLTFTextureHDREncoding::RGBM;
-	ExportUniformScale = 0.01;
 	bExportHiddenInGame = false;
 	ExportLights = EGLTFExportLightMobility::MovableAndStationary;
 	bExportCameras = true;
@@ -44,8 +47,4 @@ void UGLTFExportOptions::ResetToDefault()
 	ExportMaterialVariants = EGLTFMaterialBakeMode::UseMeshData;
 	bExportMeshVariants = true;
 	bExportVisibilityVariants = true;
-	bBundleWebViewer = true;
-	bExportPreviewMesh = true;
-	bAllExtensionsRequired = false;
-	bShowFilesWhenDone = true;
 }
diff --git a/Engine/Plugins/Enterprise/GLTFExporter/MigratedFromGit/Source/GLTFExporter/Private/Json/GLTFJsonRoot.h b/Engine/Plugins/Enterprise/GLTFExporter/MigratedFromGit/Source/GLTFExporter/Private/Json/GLTFJsonRoot.h
--- a/Engine/Plugins/Enterprise/GLTFExporter/MigratedFromGit/Source/GLTFExporter/Private/Json/GLTFJsonRoot.h
+++ b/Engine/Plugins/Enterprise/GLTFExporter/MigratedFromGit/Source/GLTFExporter/Private/Json/GLTFJsonRoot.h
@@ -84,7 +84,7 @@ struct FGLTFJsonRoot
 	TArray<TUniquePtr<FGLTFJsonVariation>>  Variations;
 
 	template <class CharType = TCHAR, class PrintPolicy = TPrettyJsonPrintPolicy<CharType>>
-	void WriteObject(TJsonWriter<CharType, PrintPolicy>& JsonWriter, bool bAllExtensionsRequired)
+	void WriteObject(TJsonWriter<CharType, PrintPolicy>& JsonWriter)
 	{
 		JsonWriter.WriteObjectStart();
 
@@ -177,11 +177,6 @@ struct FGLTFJsonRoot
 			JsonWriter.WriteObjectEnd();
 		}
 
-		if (bAllExtensionsRequired)
-		{
-			Extensions.Required.Append(Extensions.Used);
-		}
-
 		FGLTFJsonUtility::WriteStringArray(JsonWriter, TEXT("extensionsUsed"), Extensions.Used);
 		FGLTFJsonUtility::WriteStringArray(JsonWriter, TEXT("extensionsRequired"), Extensions.Required);
 
@@ -189,22 +184,22 @@ struct FGLTFJsonRoot
 	}
 
 	template <class CharType = TCHAR, class PrintPolicy = TPrettyJsonPrintPolicy<CharType>>
-	void ToJson(FArchive* const Archive, bool bAllExtensionsRequired)
+	void ToJson(FArchive* const Archive)
 	{
 		TSharedRef<TJsonWriter<CharType, PrintPolicy>> JsonWriter = TJsonWriterFactory<CharType, PrintPolicy>::Create(Archive);
-		WriteObject(*JsonWriter, bAllExtensionsRequired);
+		WriteObject(*JsonWriter);
 		JsonWriter->Close();
 	}
 
-	void ToJson(FArchive* const Archive, bool bCondensedPrint, bool bAllExtensionsRequired)
+	void ToJson(FArchive* const Archive, bool bCondensedPrint)
 	{
 		if (bCondensedPrint)
 		{
-			ToJson<UTF8CHAR, TCondensedJsonPrintPolicy<UTF8CHAR>>(Archive, bAllExtensionsRequired);
+			ToJson<UTF8CHAR, TCondensedJsonPrintPolicy<UTF8CHAR>>(Archive);
 		}
 		else
 		{
-			ToJson<UTF8CHAR, TPrettyJsonPrintPolicy<UTF8CHAR>>(Archive, bAllExtensionsRequired);
+			ToJson<UTF8CHAR, TPrettyJsonPrintPolicy<UTF8CHAR>>(Archive);
 		}
 	}
 };
diff --git a/Engine/Plugins/Enterprise/GLTFExporter/MigratedFromGit/Source/GLTFExporter/Public/GLTFExportOptions.h b/Engine/Plugins/Enterprise/GLTFExporter/MigratedFromGit/Source/GLTFExporter/Public/GLTFExportOptions.h
--- a/Engine/Plugins/Enterprise/GLTFExporter/MigratedFromGit/Source/GLTFExporter/Public/GLTFExportOptions.h
+++ b/Engine/Plugins/Enterprise/GLTFExporter/MigratedFromGit/Source/GLTFExporter/Public/GLTFExportOptions.h
@@ -47,7 +47,18 @@ class GLTFEXPORTER_API UGLTFExportOptions : public UObject
 {
 	GENERATED_UCLASS_BODY()
 
-public:
+	/** Scale factor used for exporting assets (0.01 by default) for conversion from centimeters (Unreal default) to meters (glTF). */
+	UPROPERTY(EditAnywhere, BlueprintReadWrite, Config, Category = General)
+	float ExportUniformScale;
+
+	UPROPERTY(EditAnywhere, BlueprintReadWrite, Config, Category = General)
+	bool bExportPreviewMesh;
+
+	UPROPERTY(EditAnywhere, BlueprintReadWrite, Config, Category = General)
+	bool bBundleWebViewer;
+
+	UPROPERTY(EditAnywhere, BlueprintReadWrite, Config, Category = General)
+	bool bShowFilesWhenDone;
 
 	UPROPERTY(EditAnywhere, BlueprintReadWrite, Config, Category = Material)
 	bool bExportUnlitMaterials;
@@ -115,10 +126,6 @@ public:
 	UPROPERTY(EditAnywhere, BlueprintReadWrite, Config, Category = Texture, Meta = (DisplayName = "Texture HDR Encoding"))
 	EGLTFTextureHDREncoding TextureHDREncoding;
 
-	/** Scale factor used for exporting assets (0.01 by default) for conversion from centimeters (Unreal default) to meters (glTF). */
-	UPROPERTY(EditAnywhere, BlueprintReadWrite, Config, Category = Scene)
-	float ExportUniformScale;
-
 	UPROPERTY(EditAnywhere, BlueprintReadWrite, Config, Category = Scene)
 	bool bExportHiddenInGame;
 
@@ -152,17 +159,5 @@ public:
 	UPROPERTY(EditAnywhere, BlueprintReadWrite, Config, Category = VariantSets, Meta = (EditCondition = "bExportVariantSets"))
 	bool bExportVisibilityVariants;
 
-	UPROPERTY(EditAnywhere, BlueprintReadWrite, Config, Category = Exporter)
-	bool bBundleWebViewer;
-
-	UPROPERTY(EditAnywhere, BlueprintReadWrite, Config, Category = Exporter)
-	bool bExportPreviewMesh;
-
-	UPROPERTY(EditAnywhere, BlueprintReadWrite, Config, Category = Exporter)
-	bool bAllExtensionsRequired;
-
-	UPROPERTY(EditAnywhere, BlueprintReadWrite, Config, Category = Exporter)
-	bool bShowFilesWhenDone;
-
 	void ResetToDefault();
 };
