[GLTFExporter] refactor: rename export option LosslessCompressTextures to NoLossyCompressionFor and improve enum types for ExportLights and ExportMaterialVariants

commit 847cfff402d69a37522c78a508725cf902aea816
Date: Wed, 25 Nov 2020 01:08:40 +0100

    refactor: rename export option LosslessCompressTextures to NoLossyCompressionFor and improve enum types for ExportLights and ExportMaterialVariants
    (cherry picked from commit 019f81038b9cd22a6717db55c911a2ac0772f17c​)

#skipci

diff --git a/Engine/Plugins/Enterprise/GLTFExporter/MigratedFromGit/Source/GLTFExporter/Private/Builders/GLTFBuilder.cpp b/Engine/Plugins/Enterprise/GLTFExporter/MigratedFromGit/Source/GLTFExporter/Private/Builders/GLTFBuilder.cpp
--- a/Engine/Plugins/Enterprise/GLTFExporter/MigratedFromGit/Source/GLTFExporter/Private/Builders/GLTFBuilder.cpp
+++ b/Engine/Plugins/Enterprise/GLTFExporter/MigratedFromGit/Source/GLTFExporter/Private/Builders/GLTFBuilder.cpp
@@ -77,16 +77,21 @@ EGLTFJsonHDREncoding FGLTFBuilder::GetTextureHDREncoding() const
 
 bool FGLTFBuilder::ShouldExportLight(EComponentMobility::Type LightMobility) const
 {
-	switch (ExportOptions->ExportLights)
+	const EGLTFSceneMobility AllowedMobility = static_cast<EGLTFSceneMobility>(ExportOptions->ExportLights);
+	const EGLTFSceneMobility QueriedMobility = GetSceneMobility(LightMobility);
+	return EnumHasAllFlags(AllowedMobility, QueriedMobility);
+}
+
+EGLTFSceneMobility FGLTFBuilder::GetSceneMobility(EComponentMobility::Type Mobility)
+{
+	switch (Mobility)
 	{
-		case EGLTFExportLightMobility::All:
-			return true;
-		case EGLTFExportLightMobility::MovableAndStationary:
-			return LightMobility == EComponentMobility::Movable || LightMobility == EComponentMobility::Stationary;
-		case EGLTFExportLightMobility::MovableOnly:
-			return LightMobility == EComponentMobility::Movable;
+		case EComponentMobility::Static:     return EGLTFSceneMobility::Static;
+		case EComponentMobility::Stationary: return EGLTFSceneMobility::Stationary;
+		case EComponentMobility::Movable:    return EGLTFSceneMobility::Movable;
 		default:
-			return false;
+			checkNoEntry();
+			return EGLTFSceneMobility::None;
 	}
 }
 
diff --git a/Engine/Plugins/Enterprise/GLTFExporter/MigratedFromGit/Source/GLTFExporter/Private/Builders/GLTFBuilder.h b/Engine/Plugins/Enterprise/GLTFExporter/MigratedFromGit/Source/GLTFExporter/Private/Builders/GLTFBuilder.h
--- a/Engine/Plugins/Enterprise/GLTFExporter/MigratedFromGit/Source/GLTFExporter/Private/Builders/GLTFBuilder.h
+++ b/Engine/Plugins/Enterprise/GLTFExporter/MigratedFromGit/Source/GLTFExporter/Private/Builders/GLTFBuilder.h
@@ -31,5 +31,6 @@ public:
 
 private:
 
+	static EGLTFSceneMobility GetSceneMobility(EComponentMobility::Type Mobility);
 	static EGLTFMaterialPropertyGroup GetPropertyGroup(const FMaterialPropertyEx& Property);
 };
diff --git a/Engine/Plugins/Enterprise/GLTFExporter/MigratedFromGit/Source/GLTFExporter/Private/Builders/GLTFImageBuilder.cpp b/Engine/Plugins/Enterprise/GLTFExporter/MigratedFromGit/Source/GLTFExporter/Private/Builders/GLTFImageBuilder.cpp
--- a/Engine/Plugins/Enterprise/GLTFExporter/MigratedFromGit/Source/GLTFExporter/Private/Builders/GLTFImageBuilder.cpp
+++ b/Engine/Plugins/Enterprise/GLTFExporter/MigratedFromGit/Source/GLTFExporter/Private/Builders/GLTFImageBuilder.cpp
@@ -10,23 +10,23 @@ FGLTFImageBuilder::FGLTFImageBuilder(const FString& FilePath, const UGLTFExportO
 {
 }
 
-FGLTFJsonImageIndex FGLTFImageBuilder::AddImage(const TArray<FColor>& Pixels, FIntPoint Size, bool bIgnoreAlpha, EGLTFTextureGroupFlags Flags, const FString& Name)
+FGLTFJsonImageIndex FGLTFImageBuilder::AddImage(const TArray<FColor>& Pixels, FIntPoint Size, bool bIgnoreAlpha, EGLTFTextureType Type, const FString& Name)
 {
 	check(Pixels.Num() == Size.X * Size.Y);
-	return AddImage(Pixels.GetData(), Size, bIgnoreAlpha, Flags, Name);
+	return AddImage(Pixels.GetData(), Size, bIgnoreAlpha, Type, Name);
 }
 
-FGLTFJsonImageIndex FGLTFImageBuilder::AddImage(const FColor* Pixels, int64 ByteLength, FIntPoint Size, bool bIgnoreAlpha, EGLTFTextureGroupFlags Flags, const FString& Name)
+FGLTFJsonImageIndex FGLTFImageBuilder::AddImage(const FColor* Pixels, int64 ByteLength, FIntPoint Size, bool bIgnoreAlpha, EGLTFTextureType Type, const FString& Name)
 {
 	check(ByteLength == Size.X * Size.Y * sizeof(FColor));
-	return AddImage(Pixels, Size, bIgnoreAlpha, Flags, Name);
+	return AddImage(Pixels, Size, bIgnoreAlpha, Type, Name);
 }
 
-FGLTFJsonImageIndex FGLTFImageBuilder::AddImage(const FColor* Pixels, FIntPoint Size, bool bIgnoreAlpha, EGLTFTextureGroupFlags Flags, const FString& Name)
+FGLTFJsonImageIndex FGLTFImageBuilder::AddImage(const FColor* Pixels, FIntPoint Size, bool bIgnoreAlpha, EGLTFTextureType Type, const FString& Name)
 {
 	TArray64<uint8> CompressedData;
 
-	const EGLTFJsonMimeType ImageFormat = GetImageFormat(Pixels, Size, bIgnoreAlpha, Flags);
+	const EGLTFJsonMimeType ImageFormat = GetImageFormat(Pixels, Size, bIgnoreAlpha, Type);
 	switch (ImageFormat)
 	{
 		case EGLTFJsonMimeType::PNG:
@@ -73,7 +73,7 @@ FGLTFJsonImageIndex FGLTFImageBuilder::AddImage(const void* CompressedData, int6
 	return ImageIndex;
 }
 
-EGLTFJsonMimeType FGLTFImageBuilder::GetImageFormat(const FColor* Pixels, FIntPoint Size, bool bIgnoreAlpha, EGLTFTextureGroupFlags Flags) const
+EGLTFJsonMimeType FGLTFImageBuilder::GetImageFormat(const FColor* Pixels, FIntPoint Size, bool bIgnoreAlpha, EGLTFTextureType Type) const
 {
 	switch (ExportOptions->TextureCompression)
 	{
@@ -82,7 +82,7 @@ EGLTFJsonMimeType FGLTFImageBuilder::GetImageFormat(const FColor* Pixels, FIntPo
 
 		case EGLTFTextureCompression::JPEG:
 			return
-				!EnumHasAllFlags(static_cast<EGLTFTextureGroupFlags>(ExportOptions->LosslessCompressTextures), Flags) &&
+				!EnumHasAllFlags(static_cast<EGLTFTextureType>(ExportOptions->NoLossyCompressionFor), Type) &&
 				(bIgnoreAlpha || FGLTFImageUtility::NoAlphaNeeded(Pixels, Size)) ?
 				EGLTFJsonMimeType::JPEG : EGLTFJsonMimeType::PNG;
 
diff --git a/Engine/Plugins/Enterprise/GLTFExporter/MigratedFromGit/Source/GLTFExporter/Private/Builders/GLTFImageBuilder.h b/Engine/Plugins/Enterprise/GLTFExporter/MigratedFromGit/Source/GLTFExporter/Private/Builders/GLTFImageBuilder.h
--- a/Engine/Plugins/Enterprise/GLTFExporter/MigratedFromGit/Source/GLTFExporter/Private/Builders/GLTFImageBuilder.h
+++ b/Engine/Plugins/Enterprise/GLTFExporter/MigratedFromGit/Source/GLTFExporter/Private/Builders/GLTFImageBuilder.h
@@ -13,15 +13,15 @@ protected:
 
 public:
 
-	FGLTFJsonImageIndex AddImage(const TArray<FColor>& Pixels, FIntPoint Size, bool bIgnoreAlpha, EGLTFTextureGroupFlags Flags, const FString& Name);
-	FGLTFJsonImageIndex AddImage(const FColor* Pixels, int64 ByteLength, FIntPoint Size, bool bIgnoreAlpha, EGLTFTextureGroupFlags Flags, const FString& Name);
+	FGLTFJsonImageIndex AddImage(const TArray<FColor>& Pixels, FIntPoint Size, bool bIgnoreAlpha, EGLTFTextureType Type, const FString& Name);
+	FGLTFJsonImageIndex AddImage(const FColor* Pixels, int64 ByteLength, FIntPoint Size, bool bIgnoreAlpha, EGLTFTextureType Type, const FString& Name);
 
 private:
 
-	FGLTFJsonImageIndex AddImage(const FColor* Pixels, FIntPoint Size, bool bIgnoreAlpha, EGLTFTextureGroupFlags Flags, const FString& Name);
+	FGLTFJsonImageIndex AddImage(const FColor* Pixels, FIntPoint Size, bool bIgnoreAlpha, EGLTFTextureType Type, const FString& Name);
 	FGLTFJsonImageIndex AddImage(const void* CompressedData, int64 CompressedByteLength, EGLTFJsonMimeType MimeType, const FString& Name);
 
-	EGLTFJsonMimeType GetImageFormat(const FColor* Pixels, FIntPoint Size, bool bIgnoreAlpha, EGLTFTextureGroupFlags Flags) const;
+	EGLTFJsonMimeType GetImageFormat(const FColor* Pixels, FIntPoint Size, bool bIgnoreAlpha, EGLTFTextureType Type) const;
 
 	FString SaveImageToFile(const void* CompressedData, int64 CompressedByteLength, EGLTFJsonMimeType MimeType, const FString& Name);
 
diff --git a/Engine/Plugins/Enterprise/GLTFExporter/MigratedFromGit/Source/GLTFExporter/Private/Converters/GLTFMaterialUtility.cpp b/Engine/Plugins/Enterprise/GLTFExporter/MigratedFromGit/Source/GLTFExporter/Private/Converters/GLTFMaterialUtility.cpp
--- a/Engine/Plugins/Enterprise/GLTFExporter/MigratedFromGit/Source/GLTFExporter/Private/Converters/GLTFMaterialUtility.cpp
+++ b/Engine/Plugins/Enterprise/GLTFExporter/MigratedFromGit/Source/GLTFExporter/Private/Converters/GLTFMaterialUtility.cpp
@@ -304,7 +304,7 @@ FGLTFJsonTextureIndex FGLTFMaterialUtility::AddTexture(FGLTFConvertBuilder& Buil
 	FGLTFJsonTexture JsonTexture;
 	JsonTexture.Name = TextureName;
 	JsonTexture.Sampler = Builder.AddSampler(JsonSampler);
-	JsonTexture.Source = Builder.AddImage(Pixels, TextureSize, bIgnoreAlpha, bIsNormalMap ? EGLTFTextureGroupFlags::Normalmaps : EGLTFTextureGroupFlags::None, TextureName);
+	JsonTexture.Source = Builder.AddImage(Pixels, TextureSize, bIgnoreAlpha, bIsNormalMap ? EGLTFTextureType::Normalmaps : EGLTFTextureType::None, TextureName);
 
 	return Builder.AddTexture(JsonTexture);
 }
diff --git a/Engine/Plugins/Enterprise/GLTFExporter/MigratedFromGit/Source/GLTFExporter/Private/Converters/GLTFVarationConverters.cpp b/Engine/Plugins/Enterprise/GLTFExporter/MigratedFromGit/Source/GLTFExporter/Private/Converters/GLTFVarationConverters.cpp
--- a/Engine/Plugins/Enterprise/GLTFExporter/MigratedFromGit/Source/GLTFExporter/Private/Converters/GLTFVarationConverters.cpp
+++ b/Engine/Plugins/Enterprise/GLTFExporter/MigratedFromGit/Source/GLTFExporter/Private/Converters/GLTFVarationConverters.cpp
@@ -107,7 +107,7 @@ bool FGLTFVariationConverter::TryParseVariantBinding(FGLTFJsonVariant& OutVarian
 
 		if (Property->IsA<UPropertyValueMaterial>())
 		{
-			if (Builder.ExportOptions->ExportMaterialVariants != EGLTFMaterialBakeMode::None && TryParseMaterialPropertyValue(OutVariant, Property))
+			if (Builder.ExportOptions->ExportMaterialVariants != EGLTFMaterialVariantMode::None && TryParseMaterialPropertyValue(OutVariant, Property))
 			{
 				bHasParsedAnyProperty = true;
 			}
@@ -245,7 +245,7 @@ bool FGLTFVariationConverter::TryParseMaterialPropertyValue(FGLTFJsonVariant& Ou
 	const FGLTFMeshData* MeshData = nullptr;
 	TArray<int32> SectionIndices;
 
-	if (Builder.ExportOptions->ExportMaterialVariants == EGLTFMaterialBakeMode::UseMeshData)
+	if (Builder.ExportOptions->ExportMaterialVariants == EGLTFMaterialVariantMode::UseMeshData)
 	{
 		if (Builder.ExportOptions->BakeMaterialInputs == EGLTFMaterialBakeMode::UseMeshData)
 		{
diff --git a/Engine/Plugins/Enterprise/GLTFExporter/MigratedFromGit/Source/GLTFExporter/Private/GLTFExportOptions.cpp b/Engine/Plugins/Enterprise/GLTFExporter/MigratedFromGit/Source/GLTFExporter/Private/GLTFExportOptions.cpp
--- a/Engine/Plugins/Enterprise/GLTFExporter/MigratedFromGit/Source/GLTFExporter/Private/GLTFExportOptions.cpp
+++ b/Engine/Plugins/Enterprise/GLTFExporter/MigratedFromGit/Source/GLTFExporter/Private/GLTFExportOptions.cpp
@@ -32,19 +32,19 @@ void UGLTFExportOptions::ResetToDefault()
 	bExportPlaybackSettings = true;
 	TextureCompression = EGLTFTextureCompression::PNG;
 	TextureCompressionQuality = 0;
-	LosslessCompressTextures = static_cast<int32>(EGLTFTextureGroupFlags::All);
+	NoLossyCompressionFor = static_cast<int32>(EGLTFTextureType::All);
 	bExportTextureTransforms = true;
 	bExportLightmaps = true;
 	TextureHDREncoding = EGLTFTextureHDREncoding::RGBM;
 	bExportHiddenInGame = false;
-	ExportLights = EGLTFExportLightMobility::MovableAndStationary;
+	ExportLights = static_cast<int32>(EGLTFSceneMobility::Stationary | EGLTFSceneMobility::Movable);
 	bExportCameras = true;
 	bExportCameraControls = true;
 	bExportAnimationHotspots = true;
 	bExportHDRIBackdrops = true;
 	bExportSkySpheres = true;
 	bExportVariantSets = true;
-	ExportMaterialVariants = EGLTFMaterialBakeMode::UseMeshData;
+	ExportMaterialVariants = EGLTFMaterialVariantMode::UseMeshData;
 	bExportMeshVariants = true;
 	bExportVisibilityVariants = true;
 }
diff --git a/Engine/Plugins/Enterprise/GLTFExporter/MigratedFromGit/Source/GLTFExporter/Private/Tasks/GLTFMaterialTasks.cpp b/Engine/Plugins/Enterprise/GLTFExporter/MigratedFromGit/Source/GLTFExporter/Private/Tasks/GLTFMaterialTasks.cpp
--- a/Engine/Plugins/Enterprise/GLTFExporter/MigratedFromGit/Source/GLTFExporter/Private/Tasks/GLTFMaterialTasks.cpp
+++ b/Engine/Plugins/Enterprise/GLTFExporter/MigratedFromGit/Source/GLTFExporter/Private/Tasks/GLTFMaterialTasks.cpp
@@ -329,7 +329,7 @@ bool FGLTFMaterialTask::TryGetBaseColorAndOpacity(FGLTFJsonPBRMetallicRoughness&
 		return true;
 	}
 
-	if (Builder.ExportOptions->BakeMaterialInputs == EGLTFMaterialBakeMode::None)
+	if (Builder.ExportOptions->BakeMaterialInputs == EGLTFMaterialBakeMode::Disabled)
 	{
 		Builder.AddWarningMessage(FString::Printf(
 			TEXT("%s and %s for material %s needs to bake, but material baking is disabled by export options"),
@@ -449,7 +449,7 @@ bool FGLTFMaterialTask::TryGetMetallicAndRoughness(FGLTFJsonPBRMetallicRoughness
 		return true;
 	}
 
-	if (Builder.ExportOptions->BakeMaterialInputs == EGLTFMaterialBakeMode::None)
+	if (Builder.ExportOptions->BakeMaterialInputs == EGLTFMaterialBakeMode::Disabled)
 	{
 		Builder.AddWarningMessage(FString::Printf(
 			TEXT("%s and %s for material %s needs to bake, but material baking is disabled by export options"),
@@ -572,7 +572,7 @@ bool FGLTFMaterialTask::TryGetClearCoatRoughness(FGLTFJsonClearCoatExtension& Ou
 		return true;
 	}
 
-	if (Builder.ExportOptions->BakeMaterialInputs == EGLTFMaterialBakeMode::None)
+	if (Builder.ExportOptions->BakeMaterialInputs == EGLTFMaterialBakeMode::Disabled)
 	{
 		Builder.AddWarningMessage(FString::Printf(
 			TEXT("%s and %s for material %s needs to bake, but material baking is disabled by export options"),
@@ -667,7 +667,7 @@ bool FGLTFMaterialTask::TryGetEmissive(FGLTFJsonMaterial& JsonMaterial, const FM
 		return true;
 	}
 
-	if (Builder.ExportOptions->BakeMaterialInputs == EGLTFMaterialBakeMode::None)
+	if (Builder.ExportOptions->BakeMaterialInputs == EGLTFMaterialBakeMode::Disabled)
 	{
 		Builder.AddWarningMessage(FString::Printf(
 			TEXT("%s for material %s needs to bake, but material baking is disabled by export options"),
@@ -1084,7 +1084,7 @@ bool FGLTFMaterialTask::TryGetSourceTexture(const UTexture2D*& OutTexture, int32
 
 bool FGLTFMaterialTask::TryGetBakedMaterialProperty(FGLTFJsonTextureInfo& OutTexInfo, FGLTFJsonColor3& OutConstant, const FMaterialPropertyEx& Property, const FString& PropertyName, bool bTransformToLinear)
 {
-	if (Builder.ExportOptions->BakeMaterialInputs == EGLTFMaterialBakeMode::None)
+	if (Builder.ExportOptions->BakeMaterialInputs == EGLTFMaterialBakeMode::Disabled)
 	{
 		Builder.AddWarningMessage(FString::Printf(
 			TEXT("%s for material %s needs to bake, but material baking is disabled by export options"),
@@ -1117,7 +1117,7 @@ bool FGLTFMaterialTask::TryGetBakedMaterialProperty(FGLTFJsonTextureInfo& OutTex
 
 bool FGLTFMaterialTask::TryGetBakedMaterialProperty(FGLTFJsonTextureInfo& OutTexInfo, FGLTFJsonColor4& OutConstant, const FMaterialPropertyEx& Property, const FString& PropertyName, bool bTransformToLinear)
 {
-	if (Builder.ExportOptions->BakeMaterialInputs == EGLTFMaterialBakeMode::None)
+	if (Builder.ExportOptions->BakeMaterialInputs == EGLTFMaterialBakeMode::Disabled)
 	{
 		Builder.AddWarningMessage(FString::Printf(
 			TEXT("%s for material %s needs to bake, but material baking is disabled by export options"),
@@ -1150,7 +1150,7 @@ bool FGLTFMaterialTask::TryGetBakedMaterialProperty(FGLTFJsonTextureInfo& OutTex
 
 inline bool FGLTFMaterialTask::TryGetBakedMaterialProperty(FGLTFJsonTextureInfo& OutTexInfo, float& OutConstant, const FMaterialPropertyEx& Property, const FString& PropertyName, bool bTransformToLinear)
 {
-	if (Builder.ExportOptions->BakeMaterialInputs == EGLTFMaterialBakeMode::None)
+	if (Builder.ExportOptions->BakeMaterialInputs == EGLTFMaterialBakeMode::Disabled)
 	{
 		Builder.AddWarningMessage(FString::Printf(
 			TEXT("%s for material %s needs to bake, but material baking is disabled by export options"),
@@ -1183,7 +1183,7 @@ inline bool FGLTFMaterialTask::TryGetBakedMaterialProperty(FGLTFJsonTextureInfo&
 
 bool FGLTFMaterialTask::TryGetBakedMaterialProperty(FGLTFJsonTextureInfo& OutTexInfo, const FMaterialPropertyEx& Property, const FString& PropertyName, bool bTransformToLinear)
 {
-	if (Builder.ExportOptions->BakeMaterialInputs == EGLTFMaterialBakeMode::None)
+	if (Builder.ExportOptions->BakeMaterialInputs == EGLTFMaterialBakeMode::Disabled)
 	{
 		Builder.AddWarningMessage(FString::Printf(
 			TEXT("%s for material %s needs to bake, but material baking is disabled by export options"),
diff --git a/Engine/Plugins/Enterprise/GLTFExporter/MigratedFromGit/Source/GLTFExporter/Private/Tasks/GLTFTextureTasks.cpp b/Engine/Plugins/Enterprise/GLTFExporter/MigratedFromGit/Source/GLTFExporter/Private/Tasks/GLTFTextureTasks.cpp
--- a/Engine/Plugins/Enterprise/GLTFExporter/MigratedFromGit/Source/GLTFExporter/Private/Tasks/GLTFTextureTasks.cpp
+++ b/Engine/Plugins/Enterprise/GLTFExporter/MigratedFromGit/Source/GLTFExporter/Private/Tasks/GLTFTextureTasks.cpp
@@ -36,11 +36,11 @@ void FGLTFTexture2DTask::Complete()
 	}
 
 	const bool bIgnoreAlpha = FGLTFTextureUtility::IsAlphaless(Texture2D->GetPixelFormat());
-	const EGLTFTextureGroupFlags Flags =
-		Texture2D->IsNormalMap() ? EGLTFTextureGroupFlags::Normalmaps :
-		bIsHDR ? EGLTFTextureGroupFlags::HDR : EGLTFTextureGroupFlags::None;
+	const EGLTFTextureType Type =
+		Texture2D->IsNormalMap() ? EGLTFTextureType::Normalmaps :
+		bIsHDR ? EGLTFTextureType::HDR : EGLTFTextureType::None;
 
-	JsonTexture.Source = Builder.AddImage(Pixels, Size, bIgnoreAlpha, Flags, JsonTexture.Name);
+	JsonTexture.Source = Builder.AddImage(Pixels, Size, bIgnoreAlpha, Type, JsonTexture.Name);
 	JsonTexture.Sampler = Builder.GetOrAddSampler(Texture2D);
 }
 
@@ -79,9 +79,9 @@ void FGLTFTextureCubeTask::Complete()
 	}
 
 	const bool bIgnoreAlpha = FGLTFTextureUtility::IsAlphaless(TextureCube->GetPixelFormat());
-	const EGLTFTextureGroupFlags Flags = bIsHDR ? EGLTFTextureGroupFlags::HDR : EGLTFTextureGroupFlags::None;
+	const EGLTFTextureType Type = bIsHDR ? EGLTFTextureType::HDR : EGLTFTextureType::None;
 
-	JsonTexture.Source = Builder.AddImage(Pixels, Size, bIgnoreAlpha, Flags, JsonTexture.Name);
+	JsonTexture.Source = Builder.AddImage(Pixels, Size, bIgnoreAlpha, Type, JsonTexture.Name);
 	JsonTexture.Sampler = Builder.GetOrAddSampler(TextureCube);
 }
 
@@ -106,9 +106,9 @@ void FGLTFTextureRenderTarget2DTask::Complete()
 	}
 
 	const bool bIgnoreAlpha = FGLTFTextureUtility::IsAlphaless(RenderTarget2D->GetFormat());
-	const EGLTFTextureGroupFlags Flags = bIsHDR ? EGLTFTextureGroupFlags::HDR : EGLTFTextureGroupFlags::None;
+	const EGLTFTextureType Type = bIsHDR ? EGLTFTextureType::HDR : EGLTFTextureType::None;
 
-	JsonTexture.Source = Builder.AddImage(Pixels, Size, bIgnoreAlpha, Flags, JsonTexture.Name);
+	JsonTexture.Source = Builder.AddImage(Pixels, Size, bIgnoreAlpha, Type, JsonTexture.Name);
 	JsonTexture.Sampler = Builder.GetOrAddSampler(RenderTarget2D);
 }
 
@@ -146,9 +146,9 @@ void FGLTFTextureRenderTargetCubeTask::Complete()
 	}
 
 	const bool bIgnoreAlpha = FGLTFTextureUtility::IsAlphaless(RenderTargetCube->GetFormat());
-	const EGLTFTextureGroupFlags Flags = bIsHDR ? EGLTFTextureGroupFlags::HDR : EGLTFTextureGroupFlags::None;
+	const EGLTFTextureType Type = bIsHDR ? EGLTFTextureType::HDR : EGLTFTextureType::None;
 
-	JsonTexture.Source = Builder.AddImage(Pixels, Size, bIgnoreAlpha, Flags, JsonTexture.Name);
+	JsonTexture.Source = Builder.AddImage(Pixels, Size, bIgnoreAlpha, Type, JsonTexture.Name);
 	JsonTexture.Sampler = Builder.GetOrAddSampler(RenderTargetCube);
 }
 
@@ -178,10 +178,10 @@ void FGLTFTextureLightMapTask::Complete()
 	const int64 ByteLength = Source.CalcMipSize(0);
 
 	const bool bIgnoreAlpha = false;
-	const EGLTFTextureGroupFlags Flags = EGLTFTextureGroupFlags::Lightmaps;
+	const EGLTFTextureType Type = EGLTFTextureType::Lightmaps;
 
 	const void* RawData = Source.LockMip(0);
-	JsonTexture.Source = Builder.AddImage(static_cast<const FColor*>(RawData), ByteLength, Size, bIgnoreAlpha, Flags, JsonTexture.Name);
+	JsonTexture.Source = Builder.AddImage(static_cast<const FColor*>(RawData), ByteLength, Size, bIgnoreAlpha, Type, JsonTexture.Name);
 	Source.UnlockMip(0);
 
 	JsonTexture.Sampler = Builder.GetOrAddSampler(LightMap);
diff --git a/Engine/Plugins/Enterprise/GLTFExporter/MigratedFromGit/Source/GLTFExporter/Public/GLTFExportOptions.h b/Engine/Plugins/Enterprise/GLTFExporter/MigratedFromGit/Source/GLTFExporter/Public/GLTFExportOptions.h
--- a/Engine/Plugins/Enterprise/GLTFExporter/MigratedFromGit/Source/GLTFExporter/Public/GLTFExportOptions.h
+++ b/Engine/Plugins/Enterprise/GLTFExporter/MigratedFromGit/Source/GLTFExporter/Public/GLTFExportOptions.h
@@ -14,7 +14,7 @@ enum class EGLTFTextureCompression : uint8
 };
 
 UENUM(BlueprintType, Meta = (Bitflags, UseEnumValuesAsMaskValuesInEditor = "true"))
-enum class EGLTFTextureGroupFlags : uint8
+enum class EGLTFTextureType : uint8
 {
 	None = 0 UMETA(Hidden),
 
@@ -24,7 +24,7 @@ enum class EGLTFTextureGroupFlags : uint8
 
 	All = HDR | Normalmaps | Lightmaps UMETA(Hidden)
 };
-ENUM_CLASS_FLAGS(EGLTFTextureGroupFlags);
+ENUM_CLASS_FLAGS(EGLTFTextureType);
 
 UENUM(BlueprintType)
 enum class EGLTFTextureHDREncoding : uint8
@@ -33,13 +33,25 @@ enum class EGLTFTextureHDREncoding : uint8
 	RGBM
 };
 
+UENUM(BlueprintType, Meta = (Bitflags, UseEnumValuesAsMaskValuesInEditor = "true"))
+enum class EGLTFSceneMobility : uint8
+{
+	None = 0 UMETA(Hidden),
+
+	Static = 1 << 0,
+	Stationary = 1 << 1,
+	Movable = 1 << 2,
+
+	All = Static | Stationary | Movable UMETA(Hidden)
+};
+ENUM_CLASS_FLAGS(EGLTFSceneMobility);
+
 UENUM(BlueprintType)
-enum class EGLTFExportLightMobility : uint8
+enum class EGLTFMaterialVariantMode : uint8
 {
 	None,
-	MovableOnly,
-	MovableAndStationary,
-	All
+    Simple,
+    UseMeshData,
 };
 
 UCLASS(Config=EditorPerProjectUserSettings, HideCategories=(DebugProperty))
@@ -114,8 +126,8 @@ class GLTFEXPORTER_API UGLTFExportOptions : public UObject
 	UPROPERTY(EditAnywhere, BlueprintReadWrite, Config, Category = Texture, Meta = (ClampMin = "0", ClampMax = "100", EditCondition = "TextureCompression != EGLTFTextureCompression::PNG"))
 	int32 TextureCompressionQuality;
 
-	UPROPERTY(EditAnywhere, BlueprintReadWrite, Config, Category = Texture, Meta = (Bitmask, BitmaskEnum = EGLTFTextureGroupFlags, EditCondition = "TextureCompression != EGLTFTextureCompression::PNG"))
-	int32 LosslessCompressTextures;
+	UPROPERTY(EditAnywhere, BlueprintReadWrite, Config, Category = Texture, Meta = (Bitmask, BitmaskEnum = EGLTFTextureType, EditCondition = "TextureCompression != EGLTFTextureCompression::PNG"))
+	int32 NoLossyCompressionFor;
 
 	UPROPERTY(EditAnywhere, BlueprintReadWrite, Config, Category = Texture)
 	bool bExportTextureTransforms;
@@ -129,8 +141,8 @@ class GLTFEXPORTER_API UGLTFExportOptions : public UObject
 	UPROPERTY(EditAnywhere, BlueprintReadWrite, Config, Category = Scene)
 	bool bExportHiddenInGame;
 
-	UPROPERTY(EditAnywhere, BlueprintReadWrite, Config, Category = Scene)
-	EGLTFExportLightMobility ExportLights;
+	UPROPERTY(EditAnywhere, BlueprintReadWrite, Config, Category = Scene, Meta = (Bitmask, BitmaskEnum = EGLTFSceneMobility))
+	int32 ExportLights;
 
 	UPROPERTY(EditAnywhere, BlueprintReadWrite, Config, Category = Scene)
 	bool bExportCameras;
@@ -151,7 +163,7 @@ class GLTFEXPORTER_API UGLTFExportOptions : public UObject
 	bool bExportVariantSets;
 
 	UPROPERTY(EditAnywhere, BlueprintReadWrite, Config, Category = VariantSets, Meta = (EditCondition = "bExportVariantSets"))
-	EGLTFMaterialBakeMode ExportMaterialVariants;
+	EGLTFMaterialVariantMode ExportMaterialVariants;
 
 	UPROPERTY(EditAnywhere, BlueprintReadWrite, Config, Category = VariantSets, Meta = (EditCondition = "bExportVariantSets"))
 	bool bExportMeshVariants;
diff --git a/Engine/Plugins/Enterprise/GLTFExporter/MigratedFromGit/Source/GLTFExporterRuntime/Public/UserData/GLTFMaterialUserData.h b/Engine/Plugins/Enterprise/GLTFExporter/MigratedFromGit/Source/GLTFExporterRuntime/Public/UserData/GLTFMaterialUserData.h
--- a/Engine/Plugins/Enterprise/GLTFExporter/MigratedFromGit/Source/GLTFExporterRuntime/Public/UserData/GLTFMaterialUserData.h
+++ b/Engine/Plugins/Enterprise/GLTFExporter/MigratedFromGit/Source/GLTFExporterRuntime/Public/UserData/GLTFMaterialUserData.h
@@ -12,7 +12,7 @@ enum TextureAddress;
 UENUM(BlueprintType)
 enum class EGLTFMaterialBakeMode : uint8
 {
-    None,
+    Disabled,
 	Simple,
 	UseMeshData,
 };
