[GLTFExporter] refactor: rename export option OverrideBakeSettings to DefaultInputBakeSettings

commit c3c4592e9e07e8c06e586dd89f22e0117f9499f8
Date: Wed, 25 Nov 2020 02:02:20 +0100

    refactor: rename export option OverrideBakeSettings to DefaultInputBakeSettings
    (cherry picked from commit f1ae7305b9c53a1e6081c34eb062ae18c1911b9b​)

#skipci

diff --git a/Engine/Plugins/Enterprise/GLTFExporter/MigratedFromGit/Source/GLTFExporter/Private/Builders/GLTFBuilder.cpp b/Engine/Plugins/Enterprise/GLTFExporter/MigratedFromGit/Source/GLTFExporter/Private/Builders/GLTFBuilder.cpp
--- a/Engine/Plugins/Enterprise/GLTFExporter/MigratedFromGit/Source/GLTFExporter/Private/Builders/GLTFBuilder.cpp
+++ b/Engine/Plugins/Enterprise/GLTFExporter/MigratedFromGit/Source/GLTFExporter/Private/Builders/GLTFBuilder.cpp
@@ -17,7 +17,7 @@ FIntPoint FGLTFBuilder::GetBakeSizeForMaterialProperty(const UMaterialInterface*
 	const EGLTFMaterialPropertyGroup PropertyGroup = GetPropertyGroup(Property);
 	EGLTFMaterialBakeSizePOT DefaultValue = ExportOptions->DefaultMaterialBakeSize;
 
-	if (const FGLTFOverrideMaterialBakeSettings* BakeSettings = ExportOptions->OverrideBakeSettings.Find(PropertyGroup))
+	if (const FGLTFOverrideMaterialBakeSettings* BakeSettings = ExportOptions->DefaultInputBakeSettings.Find(PropertyGroup))
 	{
 		if (BakeSettings->bOverrideSize)
 		{
@@ -35,7 +35,7 @@ TextureFilter FGLTFBuilder::GetBakeFilterForMaterialProperty(const UMaterialInte
 	const EGLTFMaterialPropertyGroup PropertyGroup = GetPropertyGroup(Property);
 	TextureFilter DefaultValue = ExportOptions->DefaultMaterialBakeFilter;
 
-	if (const FGLTFOverrideMaterialBakeSettings* BakeSettings = ExportOptions->OverrideBakeSettings.Find(PropertyGroup))
+	if (const FGLTFOverrideMaterialBakeSettings* BakeSettings = ExportOptions->DefaultInputBakeSettings.Find(PropertyGroup))
 	{
 		if (BakeSettings->bOverrideFilter)
 		{
@@ -51,7 +51,7 @@ TextureAddress FGLTFBuilder::GetBakeTilingForMaterialProperty(const UMaterialInt
 	const EGLTFMaterialPropertyGroup PropertyGroup = GetPropertyGroup(Property);
 	TextureAddress DefaultValue = ExportOptions->DefaultMaterialBakeTiling;
 
-	if (const FGLTFOverrideMaterialBakeSettings* BakeSettings = ExportOptions->OverrideBakeSettings.Find(PropertyGroup))
+	if (const FGLTFOverrideMaterialBakeSettings* BakeSettings = ExportOptions->DefaultInputBakeSettings.Find(PropertyGroup))
 	{
 		if (BakeSettings->bOverrideTiling)
 		{
diff --git a/Engine/Plugins/Enterprise/GLTFExporter/MigratedFromGit/Source/GLTFExporter/Public/GLTFExportOptions.h b/Engine/Plugins/Enterprise/GLTFExporter/MigratedFromGit/Source/GLTFExporter/Public/GLTFExportOptions.h
--- a/Engine/Plugins/Enterprise/GLTFExporter/MigratedFromGit/Source/GLTFExporter/Public/GLTFExportOptions.h
+++ b/Engine/Plugins/Enterprise/GLTFExporter/MigratedFromGit/Source/GLTFExporter/Public/GLTFExportOptions.h
@@ -94,7 +94,7 @@ class GLTFEXPORTER_API UGLTFExportOptions : public UObject
 	TEnumAsByte<TextureAddress> DefaultMaterialBakeTiling;
 
 	UPROPERTY(EditAnywhere, BlueprintReadWrite, Config, Category = Material)
-	TMap<EGLTFMaterialPropertyGroup, FGLTFOverrideMaterialBakeSettings> OverrideBakeSettings;
+	TMap<EGLTFMaterialPropertyGroup, FGLTFOverrideMaterialBakeSettings> DefaultInputBakeSettings;
 
 	UPROPERTY(EditAnywhere, BlueprintReadWrite, Config, Category = Mesh, Meta = (ClampMin = "0"))
 	int32 DefaultLevelOfDetail;
