[GLTFExporter] refactor: rename property PerInput to Inputs in GLTFMaterialExportOptions

commit 1365ce14e7a29ed47f124d005d71a0a12c36423a
Date: Wed, 25 Nov 2020 02:02:33 +0100

    refactor: rename property PerInput to Inputs in GLTFMaterialExportOptions
    (cherry picked from commit 188c29c97a96e29f1cf58067afd3470c9321f099​)

#skipci

diff --git a/Engine/Plugins/Enterprise/GLTFExporter/MigratedFromGit/Source/GLTFExporterRuntime/Private/UserData/GLTFMaterialUserData.cpp b/Engine/Plugins/Enterprise/GLTFExporter/MigratedFromGit/Source/GLTFExporterRuntime/Private/UserData/GLTFMaterialUserData.cpp
--- a/Engine/Plugins/Enterprise/GLTFExporter/MigratedFromGit/Source/GLTFExporterRuntime/Private/UserData/GLTFMaterialUserData.cpp
+++ b/Engine/Plugins/Enterprise/GLTFExporter/MigratedFromGit/Source/GLTFExporterRuntime/Private/UserData/GLTFMaterialUserData.cpp
@@ -51,7 +51,7 @@ const FGLTFOverrideMaterialBakeSettings* UGLTFMaterialExportOptions::GetBakeSett
 	{
 		if (const UGLTFMaterialExportOptions* UserData = const_cast<UMaterialInterface*>(Material)->GetAssetUserData<UGLTFMaterialExportOptions>())
 		{
-			if (const FGLTFOverrideMaterialBakeSettings* BakeSettings = UserData->PerInput.Find(PropertyGroup))
+			if (const FGLTFOverrideMaterialBakeSettings* BakeSettings = UserData->Inputs.Find(PropertyGroup))
 			{
 				if (Pred(*BakeSettings))
 				{
diff --git a/Engine/Plugins/Enterprise/GLTFExporter/MigratedFromGit/Source/GLTFExporterRuntime/Public/UserData/GLTFMaterialUserData.h b/Engine/Plugins/Enterprise/GLTFExporter/MigratedFromGit/Source/GLTFExporterRuntime/Public/UserData/GLTFMaterialUserData.h
--- a/Engine/Plugins/Enterprise/GLTFExporter/MigratedFromGit/Source/GLTFExporterRuntime/Public/UserData/GLTFMaterialUserData.h
+++ b/Engine/Plugins/Enterprise/GLTFExporter/MigratedFromGit/Source/GLTFExporterRuntime/Public/UserData/GLTFMaterialUserData.h
@@ -90,7 +90,7 @@ public:
 	FGLTFOverrideMaterialBakeSettings Default;
 
 	UPROPERTY(EditAnywhere, BlueprintReadWrite, Category = "Override Bake Settings")
-	TMap<EGLTFMaterialPropertyGroup, FGLTFOverrideMaterialBakeSettings> PerInput;
+	TMap<EGLTFMaterialPropertyGroup, FGLTFOverrideMaterialBakeSettings> Inputs;
 
 	static EGLTFMaterialBakeSizePOT GetBakeSizeForPropertyGroup(const UMaterialInterface* Material, EGLTFMaterialPropertyGroup PropertyGroup, EGLTFMaterialBakeSizePOT DefaultValue);
 	static TextureFilter GetBakeFilterForPropertyGroup(const UMaterialInterface* Material, EGLTFMaterialPropertyGroup PropertyGroup, TextureFilter DefaultValue);
