[GLTFExporter] feat: add warnings for zero-length and non unit-length vectors

commit 0d6ac840130367564b0eadc908145f37533cfcf8
Date: Mon, 23 Nov 2020 14:40:49 +0100

    feat: add warnings for zero-length and non unit-length vectors
    (cherry picked from commit 9339b851d9ed703e2bde14c00d4f82329785af8b​)

#skipci

diff --git a/Engine/Plugins/Enterprise/GLTFExporter/MigratedFromGit/Source/GLTFExporter/Private/Converters/GLTFAccessorConverters.cpp b/Engine/Plugins/Enterprise/GLTFExporter/MigratedFromGit/Source/GLTFExporter/Private/Converters/GLTFAccessorConverters.cpp
--- a/Engine/Plugins/Enterprise/GLTFExporter/MigratedFromGit/Source/GLTFExporter/Private/Converters/GLTFAccessorConverters.cpp
+++ b/Engine/Plugins/Enterprise/GLTFExporter/MigratedFromGit/Source/GLTFExporter/Private/Converters/GLTFAccessorConverters.cpp
@@ -10,38 +10,40 @@ namespace
 	template <typename DestinationType, typename SourceType>
 	struct TGLTFVertexNormalUtility
 	{
-		static DestinationType Convert(TStaticMeshVertexTangentDatum<SourceType> Vertex)
+		static DestinationType Convert(const SourceType& Normal)
 		{
-			return FGLTFConverterUtility::ConvertNormal(Vertex.TangentZ);
+			return FGLTFConverterUtility::ConvertNormal(Normal);
 		}
 	};
 
 	template <typename SourceType>
 	struct TGLTFVertexNormalUtility<FGLTFJsonVector3, SourceType>
 	{
-		static FGLTFJsonVector3 Convert(TStaticMeshVertexTangentDatum<SourceType> Vertex)
+		static FGLTFJsonVector3 Convert(const SourceType& Normal)
 		{
-			return FGLTFConverterUtility::ConvertNormal(Vertex.TangentZ.ToFVector());
+			return FGLTFConverterUtility::ConvertNormal(Normal.ToFVector());
 		}
 	};
 
 	template <typename DestinationType, typename SourceType>
 	struct TGLTFVertexTangentUtility
 	{
-		static DestinationType Convert(TStaticMeshVertexTangentDatum<SourceType> Vertex)
+		static DestinationType Convert(const SourceType& Tangent)
 		{
-			return FGLTFConverterUtility::ConvertTangent(Vertex.TangentX);
+			return FGLTFConverterUtility::ConvertTangent(Tangent);
 		}
 	};
 
 	template <typename SourceType>
 	struct TGLTFVertexTangentUtility<FGLTFJsonVector4, SourceType>
 	{
-		static FGLTFJsonVector4 Convert(TStaticMeshVertexTangentDatum<SourceType> Vertex)
+		static FGLTFJsonVector4 Convert(const SourceType& Tangent)
 		{
-			return FGLTFConverterUtility::ConvertTangent(Vertex.TangentX.ToFVector());
+			return FGLTFConverterUtility::ConvertTangent(Tangent.ToFVector());
 		}
 	};
+
+	const float UnitLengthToleranceSq = FMath::Pow(1.006f, 2.0f) - 1.0f;	// Slightly less than the tolerance used by the glTF Validator for Vec3
 }
 
 FGLTFJsonAccessorIndex FGLTFPositionBufferConverter::Convert(const FGLTFMeshSection* MeshSection, const FPositionVertexBuffer* VertexBuffer)
@@ -191,13 +193,41 @@ FGLTFJsonBufferViewIndex FGLTFNormalBufferConverter::ConvertBufferView(const FGL
 	TArray<DestinationType> Normals;
 	Normals.AddUninitialized(VertexCount);
 
+	bool bHasZeroVectors = false;
+	bool bHasNonUnitVectors = false;
+
 	for (uint32 VertexIndex = 0; VertexIndex < VertexCount; ++VertexIndex)
 	{
 		const uint32 MappedVertexIndex = IndexMap[VertexIndex];
-		Normals[VertexIndex] = TGLTFVertexNormalUtility<DestinationType, SourceType>::Convert(VertexTangents[MappedVertexIndex]);
+		SourceType Normal = VertexTangents[MappedVertexIndex].TangentZ;
+		FVector NormalVector = Normal.ToFVector();
+
+		if (NormalVector.IsNearlyZero())
+		{
+			bHasZeroVectors = true;
+		}
+		else if (!NormalVector.IsUnit(UnitLengthToleranceSq))
+		{
+			bHasNonUnitVectors = true;
+		}
+
+		Normals[VertexIndex] = TGLTFVertexNormalUtility<DestinationType, SourceType>::Convert(Normal);
+	}
+
+	const FGLTFJsonBufferViewIndex BufferViewIndex = Builder.AddBufferView(Normals, EGLTFJsonBufferTarget::ArrayBuffer);
+	FGLTFJsonBufferView& BufferView = Builder.GetBufferView(BufferViewIndex);
+
+	if (bHasZeroVectors)
+	{
+		BufferView.ErrorFlags |= EGLTFJsonBufferViewErrorFlags::ContainsZeroLengthVectors;
+	}
+
+	if (bHasNonUnitVectors)
+	{
+		BufferView.ErrorFlags |= EGLTFJsonBufferViewErrorFlags::ContainsNonUnitLengthVectors;
 	}
 
-	return Builder.AddBufferView(Normals, EGLTFJsonBufferTarget::ArrayBuffer);
+	return BufferViewIndex;
 }
 
 FGLTFJsonAccessorIndex FGLTFTangentBufferConverter::Convert(const FGLTFMeshSection* MeshSection, const FStaticMeshVertexBuffer* VertexBuffer)
@@ -265,13 +295,41 @@ FGLTFJsonBufferViewIndex FGLTFTangentBufferConverter::ConvertBufferView(const FG
 	TArray<DestinationType> Tangents;
 	Tangents.AddUninitialized(VertexCount);
 
+	bool bHasZeroVectors = false;
+	bool bHasNonUnitVectors = false;
+
 	for (uint32 VertexIndex = 0; VertexIndex < VertexCount; ++VertexIndex)
 	{
 		const uint32 MappedVertexIndex = IndexMap[VertexIndex];
-		Tangents[VertexIndex] = TGLTFVertexTangentUtility<DestinationType, SourceType>::Convert(VertexTangents[MappedVertexIndex]);
+		SourceType Tangent = VertexTangents[MappedVertexIndex].TangentX;
+		FVector TangentVector = Tangent.ToFVector();
+
+		if (TangentVector.IsNearlyZero())
+		{
+			bHasZeroVectors = true;
+		}
+		else if (!TangentVector.IsUnit(UnitLengthToleranceSq))
+		{
+			bHasNonUnitVectors = true;
+		}
+
+		Tangents[VertexIndex] = TGLTFVertexTangentUtility<DestinationType, SourceType>::Convert(Tangent);
+	}
+
+	const FGLTFJsonBufferViewIndex BufferViewIndex = Builder.AddBufferView(Tangents, EGLTFJsonBufferTarget::ArrayBuffer);
+	FGLTFJsonBufferView& BufferView = Builder.GetBufferView(BufferViewIndex);
+
+	if (bHasZeroVectors)
+	{
+		BufferView.ErrorFlags |= EGLTFJsonBufferViewErrorFlags::ContainsZeroLengthVectors;
+	}
+
+	if (bHasNonUnitVectors)
+	{
+		BufferView.ErrorFlags |= EGLTFJsonBufferViewErrorFlags::ContainsNonUnitLengthVectors;
 	}
 
-	return Builder.AddBufferView(Tangents, EGLTFJsonBufferTarget::ArrayBuffer);
+	return BufferViewIndex;
 }
 
 FGLTFJsonAccessorIndex FGLTFUVBufferConverter::Convert(const FGLTFMeshSection* MeshSection, const FStaticMeshVertexBuffer* VertexBuffer, uint32 UVIndex)
diff --git a/Engine/Plugins/Enterprise/GLTFExporter/MigratedFromGit/Source/GLTFExporter/Private/Json/GLTFJsonBufferView.h b/Engine/Plugins/Enterprise/GLTFExporter/MigratedFromGit/Source/GLTFExporter/Private/Json/GLTFJsonBufferView.h
--- a/Engine/Plugins/Enterprise/GLTFExporter/MigratedFromGit/Source/GLTFExporter/Private/Json/GLTFJsonBufferView.h
+++ b/Engine/Plugins/Enterprise/GLTFExporter/MigratedFromGit/Source/GLTFExporter/Private/Json/GLTFJsonBufferView.h
@@ -19,11 +19,14 @@ struct FGLTFJsonBufferView
 
 	EGLTFJsonBufferTarget Target;
 
+	EGLTFJsonBufferViewErrorFlags ErrorFlags;
+
 	FGLTFJsonBufferView()
 		: ByteLength(0)
 		, ByteOffset(0)
 		, ByteStride(0)
 		, Target(EGLTFJsonBufferTarget::None)
+		, ErrorFlags(EGLTFJsonBufferViewErrorFlags::None)
 	{
 		// check that view fits completely inside the buffer
 	}
diff --git a/Engine/Plugins/Enterprise/GLTFExporter/MigratedFromGit/Source/GLTFExporter/Private/Json/GLTFJsonEnums.h b/Engine/Plugins/Enterprise/GLTFExporter/MigratedFromGit/Source/GLTFExporter/Private/Json/GLTFJsonEnums.h
--- a/Engine/Plugins/Enterprise/GLTFExporter/MigratedFromGit/Source/GLTFExporter/Private/Json/GLTFJsonEnums.h
+++ b/Engine/Plugins/Enterprise/GLTFExporter/MigratedFromGit/Source/GLTFExporter/Private/Json/GLTFJsonEnums.h
@@ -175,3 +175,12 @@ enum class EGLTFJsonCameraControlMode
 	FreeLook,
 	Orbital
 };
+
+enum class EGLTFJsonBufferViewErrorFlags : uint8
+{
+	None = 0,
+	ContainsZeroLengthVectors = 1 << 0,
+	/** Note that zero-length vectors are not treated as non unit-length vectors */
+	ContainsNonUnitLengthVectors = 1 << 1,
+};
+ENUM_CLASS_FLAGS(EGLTFJsonBufferViewErrorFlags);
diff --git a/Engine/Plugins/Enterprise/GLTFExporter/MigratedFromGit/Source/GLTFExporter/Private/Tasks/GLTFMeshTasks.cpp b/Engine/Plugins/Enterprise/GLTFExporter/MigratedFromGit/Source/GLTFExporter/Private/Tasks/GLTFMeshTasks.cpp
--- a/Engine/Plugins/Enterprise/GLTFExporter/MigratedFromGit/Source/GLTFExporter/Private/Tasks/GLTFMeshTasks.cpp
+++ b/Engine/Plugins/Enterprise/GLTFExporter/MigratedFromGit/Source/GLTFExporter/Private/Tasks/GLTFMeshTasks.cpp
@@ -6,6 +6,47 @@
 #include "Builders/GLTFConvertBuilder.h"
 #include "Rendering/SkeletalMeshRenderData.h"
 
+namespace
+{
+	EGLTFJsonBufferViewErrorFlags GetBufferViewErrorFlags(FGLTFConvertBuilder& Builder, FGLTFJsonAccessorIndex AccessorIndex)
+	{
+		if (AccessorIndex != INDEX_NONE)
+		{
+			const FGLTFJsonBufferViewIndex BufferViewIndex = Builder.GetAccessor(AccessorIndex).BufferView;
+			if (BufferViewIndex != INDEX_NONE)
+			{
+				return Builder.GetBufferView(BufferViewIndex).ErrorFlags;
+			}
+		}
+
+		return EGLTFJsonBufferViewErrorFlags::None;
+	}
+
+	void AddWarningsForBufferViewErrors(FGLTFConvertBuilder& Builder, EGLTFJsonBufferViewErrorFlags ErrorFlags, const TCHAR* MeshName, const TCHAR* TypeName)
+	{
+		const FString TypeNamePlural = FString(TypeName) + TEXT("s");
+		const FString TypeNamePluralLC = FString(TypeNamePlural).ToLower();
+
+		if (EnumHasAllFlags(ErrorFlags, EGLTFJsonBufferViewErrorFlags::ContainsZeroLengthVectors))
+		{
+			Builder.AddWarningMessage(FString::Printf(
+				TEXT("Mesh %s contains zero-length %s. Consider checking 'Recompute %s' in the build-settings for the mesh"),
+				MeshName,
+				*TypeNamePluralLC,
+				*TypeNamePlural));
+		}
+
+		if (EnumHasAllFlags(ErrorFlags, EGLTFJsonBufferViewErrorFlags::ContainsNonUnitLengthVectors))
+		{
+			Builder.AddWarningMessage(FString::Printf(
+				TEXT("Mesh %s contains non unit-length %s. Consider checking 'Recompute %s' in the build-settings for the mesh, or enable normalization in export options"),
+				MeshName,
+				*TypeNamePluralLC,
+				*TypeNamePlural));
+		}
+	}
+}
+
 void FGLTFStaticMeshTask::Complete()
 {
 	FGLTFJsonMesh& JsonMesh = Builder.GetMesh(MeshIndex);
@@ -42,6 +83,9 @@ void FGLTFStaticMeshTask::Complete()
 	const int32 MaterialCount = StaticMesh->StaticMaterials.Num();
 	JsonMesh.Primitives.AddDefaulted(MaterialCount);
 
+	EGLTFJsonBufferViewErrorFlags NormalBufferErrorFlags(EGLTFJsonBufferViewErrorFlags::None);
+	EGLTFJsonBufferViewErrorFlags TangentBufferErrorFlags(EGLTFJsonBufferViewErrorFlags::None);
+
 	for (int32 MaterialIndex = 0; MaterialIndex < MaterialCount; ++MaterialIndex)
 	{
 		const TArray<int32> SectionIndices = FGLTFMeshUtility::GetSectionIndices(MeshLOD, MaterialIndex);
@@ -66,6 +110,9 @@ void FGLTFStaticMeshTask::Complete()
 		JsonPrimitive.Attributes.Normal = Builder.GetOrAddNormalAccessor(ConvertedSection, VertexBuffer);
 		JsonPrimitive.Attributes.Tangent = Builder.GetOrAddTangentAccessor(ConvertedSection, VertexBuffer);
 
+		NormalBufferErrorFlags |= GetBufferViewErrorFlags(Builder, JsonPrimitive.Attributes.Normal);
+		TangentBufferErrorFlags |= GetBufferViewErrorFlags(Builder, JsonPrimitive.Attributes.Tangent);
+
 		const uint32 UVCount = VertexBuffer->GetNumTexCoords();
 		JsonPrimitive.Attributes.TexCoords.AddUninitialized(UVCount);
 
@@ -77,6 +124,9 @@ void FGLTFStaticMeshTask::Complete()
 		const UMaterialInterface* Material = Materials[MaterialIndex];
 		JsonPrimitive.Material =  Builder.GetOrAddMaterial(Material, MeshData, SectionIndices);
 	}
+
+	AddWarningsForBufferViewErrors(Builder, NormalBufferErrorFlags, *JsonMesh.Name, TEXT("Normal"));
+	AddWarningsForBufferViewErrors(Builder, TangentBufferErrorFlags, *JsonMesh.Name, TEXT("Tangent"));
 }
 
 void FGLTFSkeletalMeshTask::Complete()
@@ -119,6 +169,9 @@ void FGLTFSkeletalMeshTask::Complete()
 	const uint16 MaterialCount = SkeletalMesh->Materials.Num();
 	JsonMesh.Primitives.AddDefaulted(MaterialCount);
 
+	EGLTFJsonBufferViewErrorFlags NormalBufferErrorFlags(EGLTFJsonBufferViewErrorFlags::None);
+	EGLTFJsonBufferViewErrorFlags TangentBufferErrorFlags(EGLTFJsonBufferViewErrorFlags::None);
+
 	for (uint16 MaterialIndex = 0; MaterialIndex < MaterialCount; ++MaterialIndex)
 	{
 		const TArray<int32> SectionIndices = FGLTFMeshUtility::GetSectionIndices(MeshLOD, MaterialIndex);
@@ -143,6 +196,9 @@ void FGLTFSkeletalMeshTask::Complete()
 		JsonPrimitive.Attributes.Normal = Builder.GetOrAddNormalAccessor(ConvertedSection, VertexBuffer);
 		JsonPrimitive.Attributes.Tangent = Builder.GetOrAddTangentAccessor(ConvertedSection, VertexBuffer);
 
+		NormalBufferErrorFlags |= GetBufferViewErrorFlags(Builder, JsonPrimitive.Attributes.Normal);
+		TangentBufferErrorFlags |= GetBufferViewErrorFlags(Builder, JsonPrimitive.Attributes.Tangent);
+
 		const uint32 UVCount = VertexBuffer->GetNumTexCoords();
 		JsonPrimitive.Attributes.TexCoords.AddUninitialized(UVCount);
 
@@ -167,4 +223,7 @@ void FGLTFSkeletalMeshTask::Complete()
 		const UMaterialInterface* Material = Materials[MaterialIndex];
 		JsonPrimitive.Material =  Builder.GetOrAddMaterial(Material, MeshData, SectionIndices);
 	}
+
+	AddWarningsForBufferViewErrors(Builder, NormalBufferErrorFlags, *JsonMesh.Name, TEXT("Normal"));
+	AddWarningsForBufferViewErrors(Builder, TangentBufferErrorFlags, *JsonMesh.Name, TEXT("Tangent"));
 }
