[GLTFExporter] feat: add and implement export option for normalizing non-zero vectors

commit c86ff4db6ae2ba306f54bd7d96c3c725b258df19
Date: Mon, 23 Nov 2020 14:40:38 +0100

    feat: add and implement export option for normalizing non-zero vectors
    (cherry picked from commit 273d4cb06b2b57d34b3c465d610069433e80b14b​)

#skipci

diff --git a/Engine/Plugins/Enterprise/GLTFExporter/MigratedFromGit/Source/GLTFExporter/Private/Converters/GLTFAccessorConverters.cpp b/Engine/Plugins/Enterprise/GLTFExporter/MigratedFromGit/Source/GLTFExporter/Private/Converters/GLTFAccessorConverters.cpp
--- a/Engine/Plugins/Enterprise/GLTFExporter/MigratedFromGit/Source/GLTFExporter/Private/Converters/GLTFAccessorConverters.cpp
+++ b/Engine/Plugins/Enterprise/GLTFExporter/MigratedFromGit/Source/GLTFExporter/Private/Converters/GLTFAccessorConverters.cpp
@@ -195,6 +195,7 @@ FGLTFJsonBufferViewIndex FGLTFNormalBufferConverter::ConvertBufferView(const FGL
 
 	bool bHasZeroVectors = false;
 	bool bHasNonUnitVectors = false;
+	const bool bNormalizeVectors = Builder.ExportOptions->bNormalizeNonZeroVectors;
 
 	for (uint32 VertexIndex = 0; VertexIndex < VertexCount; ++VertexIndex)
 	{
@@ -202,16 +203,20 @@ FGLTFJsonBufferViewIndex FGLTFNormalBufferConverter::ConvertBufferView(const FGL
 		SourceType Normal = VertexTangents[MappedVertexIndex].TangentZ;
 		FVector NormalVector = Normal.ToFVector();
 
-		if (NormalVector.IsNearlyZero())
-		{
-			bHasZeroVectors = true;
-		}
-		else if (!NormalVector.IsUnit(UnitLengthToleranceSq))
+		const bool bIsZero = NormalVector.IsNearlyZero();
+		const bool bIsUnit = NormalVector.IsUnit(UnitLengthToleranceSq);
+		bool bNeedsNormalizing = !bIsZero && !bIsUnit;
+
+		if (bNormalizeVectors && bNeedsNormalizing)
 		{
-			bHasNonUnitVectors = true;
+			NormalVector = NormalVector.GetSafeNormal();
+			Normal = NormalVector;
+			bNeedsNormalizing = false;
 		}
 
 		Normals[VertexIndex] = TGLTFVertexNormalUtility<DestinationType, SourceType>::Convert(Normal);
+		bHasZeroVectors |= bIsZero;
+		bHasNonUnitVectors |= bNeedsNormalizing;
 	}
 
 	const FGLTFJsonBufferViewIndex BufferViewIndex = Builder.AddBufferView(Normals, EGLTFJsonBufferTarget::ArrayBuffer);
@@ -297,6 +302,7 @@ FGLTFJsonBufferViewIndex FGLTFTangentBufferConverter::ConvertBufferView(const FG
 
 	bool bHasZeroVectors = false;
 	bool bHasNonUnitVectors = false;
+	const bool bNormalizeVectors = Builder.ExportOptions->bNormalizeNonZeroVectors;
 
 	for (uint32 VertexIndex = 0; VertexIndex < VertexCount; ++VertexIndex)
 	{
@@ -304,16 +310,20 @@ FGLTFJsonBufferViewIndex FGLTFTangentBufferConverter::ConvertBufferView(const FG
 		SourceType Tangent = VertexTangents[MappedVertexIndex].TangentX;
 		FVector TangentVector = Tangent.ToFVector();
 
-		if (TangentVector.IsNearlyZero())
-		{
-			bHasZeroVectors = true;
-		}
-		else if (!TangentVector.IsUnit(UnitLengthToleranceSq))
+		const bool bIsZero = TangentVector.IsNearlyZero();
+		const bool bIsUnit = TangentVector.IsUnit(UnitLengthToleranceSq);
+		bool bNeedsNormalizing = !bIsZero && !bIsUnit;
+
+		if (bNormalizeVectors && bNeedsNormalizing)
 		{
-			bHasNonUnitVectors = true;
+			TangentVector = TangentVector.GetSafeNormal();
+			Tangent = TangentVector;
+			bNeedsNormalizing = false;
 		}
 
 		Tangents[VertexIndex] = TGLTFVertexTangentUtility<DestinationType, SourceType>::Convert(Tangent);
+		bHasZeroVectors |= bIsZero;
+		bHasNonUnitVectors |= bNeedsNormalizing;
 	}
 
 	const FGLTFJsonBufferViewIndex BufferViewIndex = Builder.AddBufferView(Tangents, EGLTFJsonBufferTarget::ArrayBuffer);
diff --git a/Engine/Plugins/Enterprise/GLTFExporter/MigratedFromGit/Source/GLTFExporter/Private/GLTFExportOptions.cpp b/Engine/Plugins/Enterprise/GLTFExporter/MigratedFromGit/Source/GLTFExporter/Private/GLTFExportOptions.cpp
--- a/Engine/Plugins/Enterprise/GLTFExporter/MigratedFromGit/Source/GLTFExporter/Private/GLTFExportOptions.cpp
+++ b/Engine/Plugins/Enterprise/GLTFExporter/MigratedFromGit/Source/GLTFExporter/Private/GLTFExportOptions.cpp
@@ -26,6 +26,7 @@ void UGLTFExportOptions::ResetToDefault()
 	bExportVertexColors = false;
 	bExportVertexSkinWeights = true;
 	bExportMeshQuantization = true;
+	bNormalizeNonZeroVectors = true;
 	bExportLevelSequences = true;
 	bExportAnimationSequences = true;
 	bRetargetBoneTransforms = true;
diff --git a/Engine/Plugins/Enterprise/GLTFExporter/MigratedFromGit/Source/GLTFExporter/Public/GLTFExportOptions.h b/Engine/Plugins/Enterprise/GLTFExporter/MigratedFromGit/Source/GLTFExporter/Public/GLTFExportOptions.h
--- a/Engine/Plugins/Enterprise/GLTFExporter/MigratedFromGit/Source/GLTFExporter/Public/GLTFExportOptions.h
+++ b/Engine/Plugins/Enterprise/GLTFExporter/MigratedFromGit/Source/GLTFExporter/Public/GLTFExportOptions.h
@@ -108,6 +108,9 @@ class GLTFEXPORTER_API UGLTFExportOptions : public UObject
 	UPROPERTY(EditAnywhere, BlueprintReadWrite, Config, Category = Mesh)
 	bool bExportMeshQuantization;
 
+	UPROPERTY(EditAnywhere, BlueprintReadWrite, Config, Category = Mesh)
+	bool bNormalizeNonZeroVectors;
+
 	UPROPERTY(EditAnywhere, BlueprintReadWrite, Config, Category = Animation)
 	bool bExportLevelSequences;
 
