[GLTFExporter] fix: make tolerance for unit-length more restrictive than in the gltf-validator

commit 0a8193153a577856953c1140a9b0ba192addb29f
Date: Mon, 23 Nov 2020 15:26:10 +0100

    fix: make tolerance for unit-length more restrictive than in the gltf-validator
    because some precision is lost during conversion of vector formats, we need a slightly tighter tolerance to keep inside the tolerance of the validator
    (cherry picked from commit 766c98734e1945e9aa2268216a2d986ada520d31​)

#skipci

diff --git a/Engine/Plugins/Enterprise/GLTFExporter/MigratedFromGit/Source/GLTFExporter/Private/Converters/GLTFAccessorConverters.cpp b/Engine/Plugins/Enterprise/GLTFExporter/MigratedFromGit/Source/GLTFExporter/Private/Converters/GLTFAccessorConverters.cpp
--- a/Engine/Plugins/Enterprise/GLTFExporter/MigratedFromGit/Source/GLTFExporter/Private/Converters/GLTFAccessorConverters.cpp
+++ b/Engine/Plugins/Enterprise/GLTFExporter/MigratedFromGit/Source/GLTFExporter/Private/Converters/GLTFAccessorConverters.cpp
@@ -43,7 +43,7 @@ namespace
 		}
 	};
 
-	const float UnitLengthToleranceSq = FMath::Pow(1.006f, 2.0f) - 1.0f;	// Slightly less than the tolerance used by the glTF Validator for Vec3
+	const float UnitLengthToleranceSq = FMath::Pow(1.0067f, 2.0f) - 1.0f;	// Slightly less than the tolerance of 0.00674 used by the glTF Validator for Vec3
 }
 
 FGLTFJsonAccessorIndex FGLTFPositionBufferConverter::Convert(const FGLTFMeshSection* MeshSection, const FPositionVertexBuffer* VertexBuffer)
