[GLTFExporter] Revert "fix: make tolerance for unit-length more restrictive than in the gltf-validator"

commit fdff3786fe3af80d1e251486d8bcde88d93b455f
Date: Tue, 24 Nov 2020 14:50:46 +0100

    Revert "fix: make tolerance for unit-length more restrictive than in the gltf-validator"
    This reverts commit 337045d97e1a27bb7617c6978cc18cda0f221b2f.
    (cherry picked from commit 51f7a1363c233275c2e2e375bc45aa4043260f69​)

#skipci

diff --git a/Engine/Plugins/Enterprise/GLTFExporter/MigratedFromGit/Source/GLTFExporter/Private/Converters/GLTFAccessorConverters.cpp b/Engine/Plugins/Enterprise/GLTFExporter/MigratedFromGit/Source/GLTFExporter/Private/Converters/GLTFAccessorConverters.cpp
--- a/Engine/Plugins/Enterprise/GLTFExporter/MigratedFromGit/Source/GLTFExporter/Private/Converters/GLTFAccessorConverters.cpp
+++ b/Engine/Plugins/Enterprise/GLTFExporter/MigratedFromGit/Source/GLTFExporter/Private/Converters/GLTFAccessorConverters.cpp
@@ -43,7 +43,7 @@ namespace
 		}
 	};
 
-	const float UnitLengthToleranceSq = FMath::Pow(1.0067f, 2.0f) - 1.0f;	// Slightly less than the tolerance of 0.00674 used by the glTF Validator for Vec3
+	const float UnitLengthToleranceSq = FMath::Pow(1.006f, 2.0f) - 1.0f;	// Slightly less than the tolerance used by the glTF Validator for Vec3
 }
 
 FGLTFJsonAccessorIndex FGLTFPositionBufferConverter::Convert(const FGLTFMeshSection* MeshSection, const FPositionVertexBuffer* VertexBuffer)
