[GLTFExporter] Revert "feat: add and implement export option for normalizing non-zero vectors"

commit 21098590152d8ce97f9e3173d888cfbbe4a0439b
Date: Tue, 24 Nov 2020 14:50:46 +0100

    Revert "feat: add and implement export option for normalizing non-zero vectors"
    This reverts commit 39ccddb0ea4ccb56887577b74e5885b61e262adf.
    (cherry picked from commit be7a4e725c2eaee345a2fd939a0baa5cc9ec9837​)

#skipci

diff --git a/Engine/Plugins/Enterprise/GLTFExporter/MigratedFromGit/Source/GLTFExporter/Private/Converters/GLTFAccessorConverters.cpp b/Engine/Plugins/Enterprise/GLTFExporter/MigratedFromGit/Source/GLTFExporter/Private/Converters/GLTFAccessorConverters.cpp
--- a/Engine/Plugins/Enterprise/GLTFExporter/MigratedFromGit/Source/GLTFExporter/Private/Converters/GLTFAccessorConverters.cpp
+++ b/Engine/Plugins/Enterprise/GLTFExporter/MigratedFromGit/Source/GLTFExporter/Private/Converters/GLTFAccessorConverters.cpp
@@ -195,7 +195,6 @@ FGLTFJsonBufferViewIndex FGLTFNormalBufferConverter::ConvertBufferView(const FGL
 
 	bool bHasZeroVectors = false;
 	bool bHasNonUnitVectors = false;
-	const bool bNormalizeVectors = Builder.ExportOptions->bNormalizeNonZeroVectors;
 
 	for (uint32 VertexIndex = 0; VertexIndex < VertexCount; ++VertexIndex)
 	{
@@ -203,20 +202,16 @@ FGLTFJsonBufferViewIndex FGLTFNormalBufferConverter::ConvertBufferView(const FGL
 		SourceType Normal = VertexTangents[MappedVertexIndex].TangentZ;
 		FVector NormalVector = Normal.ToFVector();
 
-		const bool bIsZero = NormalVector.IsNearlyZero();
-		const bool bIsUnit = NormalVector.IsUnit(UnitLengthToleranceSq);
-		bool bNeedsNormalizing = !bIsZero && !bIsUnit;
-
-		if (bNormalizeVectors && bNeedsNormalizing)
+		if (NormalVector.IsNearlyZero())
+		{
+			bHasZeroVectors = true;
+		}
+		else if (!NormalVector.IsUnit(UnitLengthToleranceSq))
 		{
-			NormalVector = NormalVector.GetSafeNormal();
-			Normal = NormalVector;
-			bNeedsNormalizing = false;
+			bHasNonUnitVectors = true;
 		}
 
 		Normals[VertexIndex] = TGLTFVertexNormalUtility<DestinationType, SourceType>::Convert(Normal);
-		bHasZeroVectors |= bIsZero;
-		bHasNonUnitVectors |= bNeedsNormalizing;
 	}
 
 	const FGLTFJsonBufferViewIndex BufferViewIndex = Builder.AddBufferView(Normals, EGLTFJsonBufferTarget::ArrayBuffer);
@@ -302,7 +297,6 @@ FGLTFJsonBufferViewIndex FGLTFTangentBufferConverter::ConvertBufferView(const FG
 
 	bool bHasZeroVectors = false;
 	bool bHasNonUnitVectors = false;
-	const bool bNormalizeVectors = Builder.ExportOptions->bNormalizeNonZeroVectors;
 
 	for (uint32 VertexIndex = 0; VertexIndex < VertexCount; ++VertexIndex)
 	{
@@ -310,20 +304,16 @@ FGLTFJsonBufferViewIndex FGLTFTangentBufferConverter::ConvertBufferView(const FG
 		SourceType Tangent = VertexTangents[MappedVertexIndex].TangentX;
 		FVector TangentVector = Tangent.ToFVector();
 
-		const bool bIsZero = TangentVector.IsNearlyZero();
-		const bool bIsUnit = TangentVector.IsUnit(UnitLengthToleranceSq);
-		bool bNeedsNormalizing = !bIsZero && !bIsUnit;
-
-		if (bNormalizeVectors && bNeedsNormalizing)
+		if (TangentVector.IsNearlyZero())
+		{
+			bHasZeroVectors = true;
+		}
+		else if (!TangentVector.IsUnit(UnitLengthToleranceSq))
 		{
-			TangentVector = TangentVector.GetSafeNormal();
-			Tangent = TangentVector;
-			bNeedsNormalizing = false;
+			bHasNonUnitVectors = true;
 		}
 
 		Tangents[VertexIndex] = TGLTFVertexTangentUtility<DestinationType, SourceType>::Convert(Tangent);
-		bHasZeroVectors |= bIsZero;
-		bHasNonUnitVectors |= bNeedsNormalizing;
 	}
 
 	const FGLTFJsonBufferViewIndex BufferViewIndex = Builder.AddBufferView(Tangents, EGLTFJsonBufferTarget::ArrayBuffer);
diff --git a/Engine/Plugins/Enterprise/GLTFExporter/MigratedFromGit/Source/GLTFExporter/Private/GLTFExportOptions.cpp b/Engine/Plugins/Enterprise/GLTFExporter/MigratedFromGit/Source/GLTFExporter/Private/GLTFExportOptions.cpp
--- a/Engine/Plugins/Enterprise/GLTFExporter/MigratedFromGit/Source/GLTFExporter/Private/GLTFExportOptions.cpp
+++ b/Engine/Plugins/Enterprise/GLTFExporter/MigratedFromGit/Source/GLTFExporter/Private/GLTFExportOptions.cpp
@@ -26,7 +26,6 @@ void UGLTFExportOptions::ResetToDefault()
 	bExportVertexColors = false;
 	bExportVertexSkinWeights = true;
 	bExportMeshQuantization = true;
-	bNormalizeNonZeroVectors = true;
 	bExportLevelSequences = true;
 	bExportAnimationSequences = true;
 	bRetargetBoneTransforms = true;
diff --git a/Engine/Plugins/Enterprise/GLTFExporter/MigratedFromGit/Source/GLTFExporter/Public/GLTFExportOptions.h b/Engine/Plugins/Enterprise/GLTFExporter/MigratedFromGit/Source/GLTFExporter/Public/GLTFExportOptions.h
--- a/Engine/Plugins/Enterprise/GLTFExporter/MigratedFromGit/Source/GLTFExporter/Public/GLTFExportOptions.h
+++ b/Engine/Plugins/Enterprise/GLTFExporter/MigratedFromGit/Source/GLTFExporter/Public/GLTFExportOptions.h
@@ -108,9 +108,6 @@ class GLTFEXPORTER_API UGLTFExportOptions : public UObject
 	UPROPERTY(EditAnywhere, BlueprintReadWrite, Config, Category = Mesh)
 	bool bExportMeshQuantization;
 
-	UPROPERTY(EditAnywhere, BlueprintReadWrite, Config, Category = Mesh)
-	bool bNormalizeNonZeroVectors;
-
 	UPROPERTY(EditAnywhere, BlueprintReadWrite, Config, Category = Animation)
 	bool bExportLevelSequences;
 
