[GLTFExporter] Revert "feat: add warnings for zero-length and non unit-length vectors"

commit ad68b4ae0e362ffcd9c242eca6fac6a907a35ebd
Date: Tue, 24 Nov 2020 14:50:46 +0100

    Revert "feat: add warnings for zero-length and non unit-length vectors"
    This reverts commit 2e6a930fbe6eec34af5546bcac62cc545522eb14.
    (cherry picked from commit bdeb6779a4d56cc8b56d90319015675c833fab0f​)

#skipci

diff --git a/Engine/Plugins/Enterprise/GLTFExporter/MigratedFromGit/Source/GLTFExporter/Private/Converters/GLTFAccessorConverters.cpp b/Engine/Plugins/Enterprise/GLTFExporter/MigratedFromGit/Source/GLTFExporter/Private/Converters/GLTFAccessorConverters.cpp
--- a/Engine/Plugins/Enterprise/GLTFExporter/MigratedFromGit/Source/GLTFExporter/Private/Converters/GLTFAccessorConverters.cpp
+++ b/Engine/Plugins/Enterprise/GLTFExporter/MigratedFromGit/Source/GLTFExporter/Private/Converters/GLTFAccessorConverters.cpp
@@ -10,40 +10,38 @@ namespace
 	template <typename DestinationType, typename SourceType>
 	struct TGLTFVertexNormalUtility
 	{
-		static DestinationType Convert(const SourceType& Normal)
+		static DestinationType Convert(TStaticMeshVertexTangentDatum<SourceType> Vertex)
 		{
-			return FGLTFConverterUtility::ConvertNormal(Normal);
+			return FGLTFConverterUtility::ConvertNormal(Vertex.TangentZ);
 		}
 	};
 
 	template <typename SourceType>
 	struct TGLTFVertexNormalUtility<FGLTFJsonVector3, SourceType>
 	{
-		static FGLTFJsonVector3 Convert(const SourceType& Normal)
+		static FGLTFJsonVector3 Convert(TStaticMeshVertexTangentDatum<SourceType> Vertex)
 		{
-			return FGLTFConverterUtility::ConvertNormal(Normal.ToFVector());
+			return FGLTFConverterUtility::ConvertNormal(Vertex.TangentZ.ToFVector());
 		}
 	};
 
 	template <typename DestinationType, typename SourceType>
 	struct TGLTFVertexTangentUtility
 	{
-		static DestinationType Convert(const SourceType& Tangent)
+		static DestinationType Convert(TStaticMeshVertexTangentDatum<SourceType> Vertex)
 		{
-			return FGLTFConverterUtility::ConvertTangent(Tangent);
+			return FGLTFConverterUtility::ConvertTangent(Vertex.TangentX);
 		}
 	};
 
 	template <typename SourceType>
 	struct TGLTFVertexTangentUtility<FGLTFJsonVector4, SourceType>
 	{
-		static FGLTFJsonVector4 Convert(const SourceType& Tangent)
+		static FGLTFJsonVector4 Convert(TStaticMeshVertexTangentDatum<SourceType> Vertex)
 		{
-			return FGLTFConverterUtility::ConvertTangent(Tangent.ToFVector());
+			return FGLTFConverterUtility::ConvertTangent(Vertex.TangentX.ToFVector());
 		}
 	};
-
-	const float UnitLengthToleranceSq = FMath::Pow(1.006f, 2.0f) - 1.0f;	// Slightly less than the tolerance used by the glTF Validator for Vec3
 }
 
 FGLTFJsonAccessorIndex FGLTFPositionBufferConverter::Convert(const FGLTFMeshSection* MeshSection, const FPositionVertexBuffer* VertexBuffer)
@@ -193,41 +191,13 @@ FGLTFJsonBufferViewIndex FGLTFNormalBufferConverter::ConvertBufferView(const FGL
 	TArray<DestinationType> Normals;
 	Normals.AddUninitialized(VertexCount);
 
-	bool bHasZeroVectors = false;
-	bool bHasNonUnitVectors = false;
-
 	for (uint32 VertexIndex = 0; VertexIndex < VertexCount; ++VertexIndex)
 	{
 		const uint32 MappedVertexIndex = IndexMap[VertexIndex];
-		SourceType Normal = VertexTangents[MappedVertexIndex].TangentZ;
-		FVector NormalVector = Normal.ToFVector();
-
-		if (NormalVector.IsNearlyZero())
-		{
-			bHasZeroVectors = true;
-		}
-		else if (!NormalVector.IsUnit(UnitLengthToleranceSq))
-		{
-			bHasNonUnitVectors = true;
-		}
-
-		Normals[VertexIndex] = TGLTFVertexNormalUtility<DestinationType, SourceType>::Convert(Normal);
-	}
-
-	const FGLTFJsonBufferViewIndex BufferViewIndex = Builder.AddBufferView(Normals, EGLTFJsonBufferTarget::ArrayBuffer);
-	FGLTFJsonBufferView& BufferView = Builder.GetBufferView(BufferViewIndex);
-
-	if (bHasZeroVectors)
-	{
-		BufferView.ErrorFlags |= EGLTFJsonBufferViewErrorFlags::ContainsZeroLengthVectors;
-	}
-
-	if (bHasNonUnitVectors)
-	{
-		BufferView.ErrorFlags |= EGLTFJsonBufferViewErrorFlags::ContainsNonUnitLengthVectors;
+		Normals[VertexIndex] = TGLTFVertexNormalUtility<DestinationType, SourceType>::Convert(VertexTangents[MappedVertexIndex]);
 	}
 
-	return BufferViewIndex;
+	return Builder.AddBufferView(Normals, EGLTFJsonBufferTarget::ArrayBuffer);
 }
 
 FGLTFJsonAccessorIndex FGLTFTangentBufferConverter::Convert(const FGLTFMeshSection* MeshSection, const FStaticMeshVertexBuffer* VertexBuffer)
@@ -295,41 +265,13 @@ FGLTFJsonBufferViewIndex FGLTFTangentBufferConverter::ConvertBufferView(const FG
 	TArray<DestinationType> Tangents;
 	Tangents.AddUninitialized(VertexCount);
 
-	bool bHasZeroVectors = false;
-	bool bHasNonUnitVectors = false;
-
 	for (uint32 VertexIndex = 0; VertexIndex < VertexCount; ++VertexIndex)
 	{
 		const uint32 MappedVertexIndex = IndexMap[VertexIndex];
-		SourceType Tangent = VertexTangents[MappedVertexIndex].TangentX;
-		FVector TangentVector = Tangent.ToFVector();
-
-		if (TangentVector.IsNearlyZero())
-		{
-			bHasZeroVectors = true;
-		}
-		else if (!TangentVector.IsUnit(UnitLengthToleranceSq))
-		{
-			bHasNonUnitVectors = true;
-		}
-
-		Tangents[VertexIndex] = TGLTFVertexTangentUtility<DestinationType, SourceType>::Convert(Tangent);
-	}
-
-	const FGLTFJsonBufferViewIndex BufferViewIndex = Builder.AddBufferView(Tangents, EGLTFJsonBufferTarget::ArrayBuffer);
-	FGLTFJsonBufferView& BufferView = Builder.GetBufferView(BufferViewIndex);
-
-	if (bHasZeroVectors)
-	{
-		BufferView.ErrorFlags |= EGLTFJsonBufferViewErrorFlags::ContainsZeroLengthVectors;
-	}
-
-	if (bHasNonUnitVectors)
-	{
-		BufferView.ErrorFlags |= EGLTFJsonBufferViewErrorFlags::ContainsNonUnitLengthVectors;
+		Tangents[VertexIndex] = TGLTFVertexTangentUtility<DestinationType, SourceType>::Convert(VertexTangents[MappedVertexIndex]);
 	}
 
-	return BufferViewIndex;
+	return Builder.AddBufferView(Tangents, EGLTFJsonBufferTarget::ArrayBuffer);
 }
 
 FGLTFJsonAccessorIndex FGLTFUVBufferConverter::Convert(const FGLTFMeshSection* MeshSection, const FStaticMeshVertexBuffer* VertexBuffer, uint32 UVIndex)
diff --git a/Engine/Plugins/Enterprise/GLTFExporter/MigratedFromGit/Source/GLTFExporter/Private/Json/GLTFJsonBufferView.h b/Engine/Plugins/Enterprise/GLTFExporter/MigratedFromGit/Source/GLTFExporter/Private/Json/GLTFJsonBufferView.h
--- a/Engine/Plugins/Enterprise/GLTFExporter/MigratedFromGit/Source/GLTFExporter/Private/Json/GLTFJsonBufferView.h
+++ b/Engine/Plugins/Enterprise/GLTFExporter/MigratedFromGit/Source/GLTFExporter/Private/Json/GLTFJsonBufferView.h
@@ -19,14 +19,11 @@ struct FGLTFJsonBufferView
 
 	EGLTFJsonBufferTarget Target;
 
-	EGLTFJsonBufferViewErrorFlags ErrorFlags;
-
 	FGLTFJsonBufferView()
 		: ByteLength(0)
 		, ByteOffset(0)
 		, ByteStride(0)
 		, Target(EGLTFJsonBufferTarget::None)
-		, ErrorFlags(EGLTFJsonBufferViewErrorFlags::None)
 	{
 		// check that view fits completely inside the buffer
 	}
diff --git a/Engine/Plugins/Enterprise/GLTFExporter/MigratedFromGit/Source/GLTFExporter/Private/Json/GLTFJsonEnums.h b/Engine/Plugins/Enterprise/GLTFExporter/MigratedFromGit/Source/GLTFExporter/Private/Json/GLTFJsonEnums.h
--- a/Engine/Plugins/Enterprise/GLTFExporter/MigratedFromGit/Source/GLTFExporter/Private/Json/GLTFJsonEnums.h
+++ b/Engine/Plugins/Enterprise/GLTFExporter/MigratedFromGit/Source/GLTFExporter/Private/Json/GLTFJsonEnums.h
@@ -175,12 +175,3 @@ enum class EGLTFJsonCameraControlMode
 	FreeLook,
 	Orbital
 };
-
-enum class EGLTFJsonBufferViewErrorFlags : uint8
-{
-	None = 0,
-	ContainsZeroLengthVectors = 1 << 0,
-	/** Note that zero-length vectors are not treated as non unit-length vectors */
-	ContainsNonUnitLengthVectors = 1 << 1,
-};
-ENUM_CLASS_FLAGS(EGLTFJsonBufferViewErrorFlags);
diff --git a/Engine/Plugins/Enterprise/GLTFExporter/MigratedFromGit/Source/GLTFExporter/Private/Tasks/GLTFMeshTasks.cpp b/Engine/Plugins/Enterprise/GLTFExporter/MigratedFromGit/Source/GLTFExporter/Private/Tasks/GLTFMeshTasks.cpp
--- a/Engine/Plugins/Enterprise/GLTFExporter/MigratedFromGit/Source/GLTFExporter/Private/Tasks/GLTFMeshTasks.cpp
+++ b/Engine/Plugins/Enterprise/GLTFExporter/MigratedFromGit/Source/GLTFExporter/Private/Tasks/GLTFMeshTasks.cpp
@@ -6,47 +6,6 @@
 #include "Builders/GLTFConvertBuilder.h"
 #include "Rendering/SkeletalMeshRenderData.h"
 
-namespace
-{
-	EGLTFJsonBufferViewErrorFlags GetBufferViewErrorFlags(FGLTFConvertBuilder& Builder, FGLTFJsonAccessorIndex AccessorIndex)
-	{
-		if (AccessorIndex != INDEX_NONE)
-		{
-			const FGLTFJsonBufferViewIndex BufferViewIndex = Builder.GetAccessor(AccessorIndex).BufferView;
-			if (BufferViewIndex != INDEX_NONE)
-			{
-				return Builder.GetBufferView(BufferViewIndex).ErrorFlags;
-			}
-		}
-
-		return EGLTFJsonBufferViewErrorFlags::None;
-	}
-
-	void AddWarningsForBufferViewErrors(FGLTFConvertBuilder& Builder, EGLTFJsonBufferViewErrorFlags ErrorFlags, const TCHAR* MeshName, const TCHAR* TypeName)
-	{
-		const FString TypeNamePlural = FString(TypeName) + TEXT("s");
-		const FString TypeNamePluralLC = FString(TypeNamePlural).ToLower();
-
-		if (EnumHasAllFlags(ErrorFlags, EGLTFJsonBufferViewErrorFlags::ContainsZeroLengthVectors))
-		{
-			Builder.AddWarningMessage(FString::Printf(
-				TEXT("Mesh %s contains zero-length %s. Consider checking 'Recompute %s' in the build-settings for the mesh"),
-				MeshName,
-				*TypeNamePluralLC,
-				*TypeNamePlural));
-		}
-
-		if (EnumHasAllFlags(ErrorFlags, EGLTFJsonBufferViewErrorFlags::ContainsNonUnitLengthVectors))
-		{
-			Builder.AddWarningMessage(FString::Printf(
-				TEXT("Mesh %s contains non unit-length %s. Consider checking 'Recompute %s' in the build-settings for the mesh, or enable normalization in export options"),
-				MeshName,
-				*TypeNamePluralLC,
-				*TypeNamePlural));
-		}
-	}
-}
-
 void FGLTFStaticMeshTask::Complete()
 {
 	FGLTFJsonMesh& JsonMesh = Builder.GetMesh(MeshIndex);
@@ -83,9 +42,6 @@ void FGLTFStaticMeshTask::Complete()
 	const int32 MaterialCount = StaticMesh->StaticMaterials.Num();
 	JsonMesh.Primitives.AddDefaulted(MaterialCount);
 
-	EGLTFJsonBufferViewErrorFlags NormalBufferErrorFlags(EGLTFJsonBufferViewErrorFlags::None);
-	EGLTFJsonBufferViewErrorFlags TangentBufferErrorFlags(EGLTFJsonBufferViewErrorFlags::None);
-
 	for (int32 MaterialIndex = 0; MaterialIndex < MaterialCount; ++MaterialIndex)
 	{
 		const TArray<int32> SectionIndices = FGLTFMeshUtility::GetSectionIndices(MeshLOD, MaterialIndex);
@@ -110,9 +66,6 @@ void FGLTFStaticMeshTask::Complete()
 		JsonPrimitive.Attributes.Normal = Builder.GetOrAddNormalAccessor(ConvertedSection, VertexBuffer);
 		JsonPrimitive.Attributes.Tangent = Builder.GetOrAddTangentAccessor(ConvertedSection, VertexBuffer);
 
-		NormalBufferErrorFlags |= GetBufferViewErrorFlags(Builder, JsonPrimitive.Attributes.Normal);
-		TangentBufferErrorFlags |= GetBufferViewErrorFlags(Builder, JsonPrimitive.Attributes.Tangent);
-
 		const uint32 UVCount = VertexBuffer->GetNumTexCoords();
 		JsonPrimitive.Attributes.TexCoords.AddUninitialized(UVCount);
 
@@ -124,9 +77,6 @@ void FGLTFStaticMeshTask::Complete()
 		const UMaterialInterface* Material = Materials[MaterialIndex];
 		JsonPrimitive.Material =  Builder.GetOrAddMaterial(Material, MeshData, SectionIndices);
 	}
-
-	AddWarningsForBufferViewErrors(Builder, NormalBufferErrorFlags, *JsonMesh.Name, TEXT("Normal"));
-	AddWarningsForBufferViewErrors(Builder, TangentBufferErrorFlags, *JsonMesh.Name, TEXT("Tangent"));
 }
 
 void FGLTFSkeletalMeshTask::Complete()
@@ -169,9 +119,6 @@ void FGLTFSkeletalMeshTask::Complete()
 	const uint16 MaterialCount = SkeletalMesh->Materials.Num();
 	JsonMesh.Primitives.AddDefaulted(MaterialCount);
 
-	EGLTFJsonBufferViewErrorFlags NormalBufferErrorFlags(EGLTFJsonBufferViewErrorFlags::None);
-	EGLTFJsonBufferViewErrorFlags TangentBufferErrorFlags(EGLTFJsonBufferViewErrorFlags::None);
-
 	for (uint16 MaterialIndex = 0; MaterialIndex < MaterialCount; ++MaterialIndex)
 	{
 		const TArray<int32> SectionIndices = FGLTFMeshUtility::GetSectionIndices(MeshLOD, MaterialIndex);
@@ -196,9 +143,6 @@ void FGLTFSkeletalMeshTask::Complete()
 		JsonPrimitive.Attributes.Normal = Builder.GetOrAddNormalAccessor(ConvertedSection, VertexBuffer);
 		JsonPrimitive.Attributes.Tangent = Builder.GetOrAddTangentAccessor(ConvertedSection, VertexBuffer);
 
-		NormalBufferErrorFlags |= GetBufferViewErrorFlags(Builder, JsonPrimitive.Attributes.Normal);
-		TangentBufferErrorFlags |= GetBufferViewErrorFlags(Builder, JsonPrimitive.Attributes.Tangent);
-
 		const uint32 UVCount = VertexBuffer->GetNumTexCoords();
 		JsonPrimitive.Attributes.TexCoords.AddUninitialized(UVCount);
 
@@ -223,7 +167,4 @@ void FGLTFSkeletalMeshTask::Complete()
 		const UMaterialInterface* Material = Materials[MaterialIndex];
 		JsonPrimitive.Material =  Builder.GetOrAddMaterial(Material, MeshData, SectionIndices);
 	}
-
-	AddWarningsForBufferViewErrors(Builder, NormalBufferErrorFlags, *JsonMesh.Name, TEXT("Normal"));
-	AddWarningsForBufferViewErrors(Builder, TangentBufferErrorFlags, *JsonMesh.Name, TEXT("Tangent"));
 }
