[GLTFExporter] feat: always normalize normals and tangents during conversion

commit 0355d8061ad07aa4688da3e62e20cd767bf215ff
Date: Tue, 24 Nov 2020 15:52:07 +0100

    feat: always normalize normals and tangents during conversion
    (cherry picked from commit 4f440ee3e0be77e3936c5e3aeebc91191b5b784c​)

#skipci

diff --git a/Engine/Plugins/Enterprise/GLTFExporter/MigratedFromGit/Source/GLTFExporter/Private/Converters/GLTFAccessorConverters.cpp b/Engine/Plugins/Enterprise/GLTFExporter/MigratedFromGit/Source/GLTFExporter/Private/Converters/GLTFAccessorConverters.cpp
--- a/Engine/Plugins/Enterprise/GLTFExporter/MigratedFromGit/Source/GLTFExporter/Private/Converters/GLTFAccessorConverters.cpp
+++ b/Engine/Plugins/Enterprise/GLTFExporter/MigratedFromGit/Source/GLTFExporter/Private/Converters/GLTFAccessorConverters.cpp
@@ -10,36 +10,36 @@ namespace
 	template <typename DestinationType, typename SourceType>
 	struct TGLTFVertexNormalUtility
 	{
-		static DestinationType Convert(TStaticMeshVertexTangentDatum<SourceType> Vertex)
+		static DestinationType Convert(const FVector& Normal)
 		{
-			return FGLTFConverterUtility::ConvertNormal(Vertex.TangentZ);
+			return FGLTFConverterUtility::ConvertNormal(SourceType(Normal));
 		}
 	};
 
 	template <typename SourceType>
 	struct TGLTFVertexNormalUtility<FGLTFJsonVector3, SourceType>
 	{
-		static FGLTFJsonVector3 Convert(TStaticMeshVertexTangentDatum<SourceType> Vertex)
+		static FGLTFJsonVector3 Convert(const FVector& Normal)
 		{
-			return FGLTFConverterUtility::ConvertNormal(Vertex.TangentZ.ToFVector());
+			return FGLTFConverterUtility::ConvertNormal(Normal);
 		}
 	};
 
 	template <typename DestinationType, typename SourceType>
 	struct TGLTFVertexTangentUtility
 	{
-		static DestinationType Convert(TStaticMeshVertexTangentDatum<SourceType> Vertex)
+		static DestinationType Convert(const FVector& Tangent)
 		{
-			return FGLTFConverterUtility::ConvertTangent(Vertex.TangentX);
+			return FGLTFConverterUtility::ConvertTangent(SourceType(Tangent));
 		}
 	};
 
 	template <typename SourceType>
 	struct TGLTFVertexTangentUtility<FGLTFJsonVector4, SourceType>
 	{
-		static FGLTFJsonVector4 Convert(TStaticMeshVertexTangentDatum<SourceType> Vertex)
+		static FGLTFJsonVector4 Convert(const FVector& Tangent)
 		{
-			return FGLTFConverterUtility::ConvertTangent(Vertex.TangentX.ToFVector());
+			return FGLTFConverterUtility::ConvertTangent(Tangent);
 		}
 	};
 }
@@ -194,7 +194,8 @@ FGLTFJsonBufferViewIndex FGLTFNormalBufferConverter::ConvertBufferView(const FGL
 	for (uint32 VertexIndex = 0; VertexIndex < VertexCount; ++VertexIndex)
 	{
 		const uint32 MappedVertexIndex = IndexMap[VertexIndex];
-		Normals[VertexIndex] = TGLTFVertexNormalUtility<DestinationType, SourceType>::Convert(VertexTangents[MappedVertexIndex]);
+		const FVector Normal = VertexTangents[MappedVertexIndex].TangentZ.ToFVector().GetSafeNormal();
+		Normals[VertexIndex] = TGLTFVertexNormalUtility<DestinationType, SourceType>::Convert(Normal);
 	}
 
 	return Builder.AddBufferView(Normals, EGLTFJsonBufferTarget::ArrayBuffer);
@@ -268,7 +269,8 @@ FGLTFJsonBufferViewIndex FGLTFTangentBufferConverter::ConvertBufferView(const FG
 	for (uint32 VertexIndex = 0; VertexIndex < VertexCount; ++VertexIndex)
 	{
 		const uint32 MappedVertexIndex = IndexMap[VertexIndex];
-		Tangents[VertexIndex] = TGLTFVertexTangentUtility<DestinationType, SourceType>::Convert(VertexTangents[MappedVertexIndex]);
+		const FVector Tangent = VertexTangents[MappedVertexIndex].TangentX.ToFVector().GetSafeNormal();
+		Tangents[VertexIndex] = TGLTFVertexTangentUtility<DestinationType, SourceType>::Convert(Tangent);
 	}
 
 	return Builder.AddBufferView(Tangents, EGLTFJsonBufferTarget::ArrayBuffer);
