[GLTFExporter] feat: warn if mesh contains zero-length normals or tangents

commit 2624ed42816d0887afa622a3db2b19d6fb90fa18
Date: Tue, 24 Nov 2020 17:02:34 +0100

    feat: warn if mesh contains zero-length normals or tangents
    (cherry picked from commit 4b24ff5a0c239119cd9f6770a09094e216522aee​)

#skipci

diff --git a/Engine/Plugins/Enterprise/GLTFExporter/MigratedFromGit/Source/GLTFExporter/Private/Tasks/GLTFMeshTasks.cpp b/Engine/Plugins/Enterprise/GLTFExporter/MigratedFromGit/Source/GLTFExporter/Private/Tasks/GLTFMeshTasks.cpp
--- a/Engine/Plugins/Enterprise/GLTFExporter/MigratedFromGit/Source/GLTFExporter/Private/Tasks/GLTFMeshTasks.cpp
+++ b/Engine/Plugins/Enterprise/GLTFExporter/MigratedFromGit/Source/GLTFExporter/Private/Tasks/GLTFMeshTasks.cpp
@@ -6,6 +6,70 @@
 #include "Builders/GLTFConvertBuilder.h"
 #include "Rendering/SkeletalMeshRenderData.h"
 
+namespace
+{
+	template <typename VertexBufferType>
+	void ValidateVertexBuffer(FGLTFConvertBuilder& Builder, const VertexBufferType* VertexBuffer, const TCHAR* MeshName)
+	{
+		if (VertexBuffer == nullptr)
+		{
+			return;
+		}
+
+		bool bHasZeroLengthNormals;
+		bool bHasZeroLengthTangents;
+
+		if (VertexBuffer->GetUseHighPrecisionTangentBasis())
+		{
+			AnalyzeTangents<VertexBufferType, FPackedRGBA16N>(VertexBuffer, bHasZeroLengthNormals, bHasZeroLengthTangents);
+		}
+		else
+		{
+			AnalyzeTangents<VertexBufferType, FPackedNormal>(VertexBuffer, bHasZeroLengthNormals, bHasZeroLengthTangents);
+		}
+
+		if (bHasZeroLengthNormals)
+		{
+			Builder.AddWarningMessage(FString::Printf(
+				TEXT("Mesh %s contains zero-length normals. Consider checking 'Recompute Normals' in the build-settings for the mesh"),
+				MeshName));
+		}
+
+		if (bHasZeroLengthTangents)
+		{
+			Builder.AddWarningMessage(FString::Printf(
+				TEXT("Mesh %s contains zero-length tangents. Consider checking 'Recompute Tangents' in the build-settings for the mesh"),
+				MeshName));
+		}
+	}
+
+	template <typename VertexBufferType, typename TangentVectorType>
+	void AnalyzeTangents(const VertexBufferType* VertexBuffer, bool& bOutHasZeroLengthNormals, bool& bOutHasZeroLengthTangents)
+	{
+		bool bHasZeroLengthNormals = false;
+		bool bHasZeroLengthTangents = false;
+
+		typedef TStaticMeshVertexTangentDatum<TangentVectorType> VertexTangentType;
+
+		const void* TangentData = const_cast<FStaticMeshVertexBuffer*>(VertexBuffer)->GetTangentData();
+		const VertexTangentType* VertexTangents = TangentData != nullptr ? static_cast<const VertexTangentType*>(TangentData) : nullptr;
+
+		if (VertexTangents != nullptr)
+		{
+			const uint32 VertexCount = VertexBuffer->GetNumVertices();
+			for (uint32 VertexIndex = 0; VertexIndex < VertexCount; ++VertexIndex)
+			{
+				const VertexTangentType& VertexTangent = VertexTangents[VertexIndex];
+				bHasZeroLengthNormals |= VertexTangent.TangentZ.ToFVector().IsNearlyZero();
+				bHasZeroLengthTangents |= VertexTangent.TangentX.ToFVector().IsNearlyZero();
+			}
+		}
+
+		bOutHasZeroLengthNormals = bHasZeroLengthNormals;
+		bOutHasZeroLengthTangents = bHasZeroLengthTangents;
+	}
+}
+
 void FGLTFStaticMeshTask::Complete()
 {
 	FGLTFJsonMesh& JsonMesh = Builder.GetMesh(MeshIndex);
@@ -39,6 +103,8 @@ void FGLTFStaticMeshTask::Complete()
 		}
 	}
 
+	ValidateVertexBuffer(Builder, VertexBuffer, *JsonMesh.Name);
+
 	const int32 MaterialCount = StaticMesh->StaticMaterials.Num();
 	JsonMesh.Primitives.AddDefaulted(MaterialCount);
 
@@ -116,6 +182,8 @@ void FGLTFSkeletalMeshTask::Complete()
 		}
 	}
 
+	ValidateVertexBuffer(Builder, VertexBuffer, *JsonMesh.Name);
+
 	const uint16 MaterialCount = SkeletalMesh->Materials.Num();
 	JsonMesh.Primitives.AddDefaulted(MaterialCount);
 
