[GLTFExporter] refactor: tweak warning message based on similar checks and phrasing in engine

commit daa483767aa9819d96346962e5603783df399de2
Date: Wed, 25 Nov 2020 02:18:58 +0100

    refactor: tweak warning message based on similar checks and phrasing in engine
    (cherry picked from commit 315fcf6cad00cdcce601070ff132ae82590a8084​)

#skipci

diff --git a/Engine/Plugins/Enterprise/GLTFExporter/MigratedFromGit/Source/GLTFExporter/Private/Tasks/GLTFMeshTasks.cpp b/Engine/Plugins/Enterprise/GLTFExporter/MigratedFromGit/Source/GLTFExporter/Private/Tasks/GLTFMeshTasks.cpp
--- a/Engine/Plugins/Enterprise/GLTFExporter/MigratedFromGit/Source/GLTFExporter/Private/Tasks/GLTFMeshTasks.cpp
+++ b/Engine/Plugins/Enterprise/GLTFExporter/MigratedFromGit/Source/GLTFExporter/Private/Tasks/GLTFMeshTasks.cpp
@@ -31,14 +31,14 @@ namespace
 		if (bHasZeroLengthNormals)
 		{
 			Builder.AddWarningMessage(FString::Printf(
-				TEXT("Mesh %s contains zero-length normals. Consider checking 'Recompute Normals' in the build-settings for the mesh"),
+				TEXT("Mesh %s has some nearly zero-length normals which can create some issues. Consider checking 'Recompute Normals' in the asset settings"),
 				MeshName));
 		}
 
 		if (bHasZeroLengthTangents)
 		{
 			Builder.AddWarningMessage(FString::Printf(
-				TEXT("Mesh %s contains zero-length tangents. Consider checking 'Recompute Tangents' in the build-settings for the mesh"),
+				TEXT("Mesh %s chas some nearly zero-length tangents which can create some issues. Consider checking 'Recompute Tangents' in the asset settings"),
 				MeshName));
 		}
 	}
@@ -103,7 +103,7 @@ void FGLTFStaticMeshTask::Complete()
 		}
 	}
 
-	ValidateVertexBuffer(Builder, VertexBuffer, *JsonMesh.Name);
+	ValidateVertexBuffer(Builder, VertexBuffer, *StaticMesh->GetName());
 
 	const int32 MaterialCount = StaticMesh->StaticMaterials.Num();
 	JsonMesh.Primitives.AddDefaulted(MaterialCount);
@@ -182,7 +182,7 @@ void FGLTFSkeletalMeshTask::Complete()
 		}
 	}
 
-	ValidateVertexBuffer(Builder, VertexBuffer, *JsonMesh.Name);
+	ValidateVertexBuffer(Builder, VertexBuffer, *SkeletalMesh->GetName());
 
 	const uint16 MaterialCount = SkeletalMesh->Materials.Num();
 	JsonMesh.Primitives.AddDefaulted(MaterialCount);
