[GLTFExporter] refactor: simplify some variable names based on terminology used in engine (UStaticMesh::BuildInternal)

commit 7f447f315f690ea564858eaba8d52a690970a70c
Date: Wed, 25 Nov 2020 02:25:51 +0100

    refactor: simplify some variable names based on terminology used in engine (UStaticMesh::BuildInternal)
    (cherry picked from commit 94c855a5116114c34e26375a578f0f6ecfe2a178​)

#skipci

diff --git a/Engine/Plugins/Enterprise/GLTFExporter/MigratedFromGit/Source/GLTFExporter/Private/Converters/GLTFAccessorConverters.cpp b/Engine/Plugins/Enterprise/GLTFExporter/MigratedFromGit/Source/GLTFExporter/Private/Converters/GLTFAccessorConverters.cpp
--- a/Engine/Plugins/Enterprise/GLTFExporter/MigratedFromGit/Source/GLTFExporter/Private/Converters/GLTFAccessorConverters.cpp
+++ b/Engine/Plugins/Enterprise/GLTFExporter/MigratedFromGit/Source/GLTFExporter/Private/Converters/GLTFAccessorConverters.cpp
@@ -194,8 +194,8 @@ FGLTFJsonBufferViewIndex FGLTFNormalBufferConverter::ConvertBufferView(const FGL
 	for (uint32 VertexIndex = 0; VertexIndex < VertexCount; ++VertexIndex)
 	{
 		const uint32 MappedVertexIndex = IndexMap[VertexIndex];
-		const FVector Normal = VertexTangents[MappedVertexIndex].TangentZ.ToFVector().GetSafeNormal();
-		Normals[VertexIndex] = TGLTFVertexNormalUtility<DestinationType, SourceType>::Convert(Normal);
+		const FVector SafeNormal = VertexTangents[MappedVertexIndex].TangentZ.ToFVector().GetSafeNormal();
+		Normals[VertexIndex] = TGLTFVertexNormalUtility<DestinationType, SourceType>::Convert(SafeNormal);
 	}
 
 	return Builder.AddBufferView(Normals, EGLTFJsonBufferTarget::ArrayBuffer);
@@ -269,8 +269,8 @@ FGLTFJsonBufferViewIndex FGLTFTangentBufferConverter::ConvertBufferView(const FG
 	for (uint32 VertexIndex = 0; VertexIndex < VertexCount; ++VertexIndex)
 	{
 		const uint32 MappedVertexIndex = IndexMap[VertexIndex];
-		const FVector Tangent = VertexTangents[MappedVertexIndex].TangentX.ToFVector().GetSafeNormal();
-		Tangents[VertexIndex] = TGLTFVertexTangentUtility<DestinationType, SourceType>::Convert(Tangent);
+		const FVector SafeTangent = VertexTangents[MappedVertexIndex].TangentX.ToFVector().GetSafeNormal();
+		Tangents[VertexIndex] = TGLTFVertexTangentUtility<DestinationType, SourceType>::Convert(SafeTangent);
 	}
 
 	return Builder.AddBufferView(Tangents, EGLTFJsonBufferTarget::ArrayBuffer);
diff --git a/Engine/Plugins/Enterprise/GLTFExporter/MigratedFromGit/Source/GLTFExporter/Private/Tasks/GLTFMeshTasks.cpp b/Engine/Plugins/Enterprise/GLTFExporter/MigratedFromGit/Source/GLTFExporter/Private/Tasks/GLTFMeshTasks.cpp
--- a/Engine/Plugins/Enterprise/GLTFExporter/MigratedFromGit/Source/GLTFExporter/Private/Tasks/GLTFMeshTasks.cpp
+++ b/Engine/Plugins/Enterprise/GLTFExporter/MigratedFromGit/Source/GLTFExporter/Private/Tasks/GLTFMeshTasks.cpp
@@ -16,26 +16,26 @@ namespace
 			return;
 		}
 
-		bool bHasZeroLengthNormals;
-		bool bHasZeroLengthTangents;
+		bool bZeroNormals;
+		bool bZeroTangents;
 
 		if (VertexBuffer->GetUseHighPrecisionTangentBasis())
 		{
-			AnalyzeTangents<VertexBufferType, FPackedRGBA16N>(VertexBuffer, bHasZeroLengthNormals, bHasZeroLengthTangents);
+			CheckTangentVectors<VertexBufferType, FPackedRGBA16N>(VertexBuffer, bZeroNormals, bZeroTangents);
 		}
 		else
 		{
-			AnalyzeTangents<VertexBufferType, FPackedNormal>(VertexBuffer, bHasZeroLengthNormals, bHasZeroLengthTangents);
+			CheckTangentVectors<VertexBufferType, FPackedNormal>(VertexBuffer, bZeroNormals, bZeroTangents);
 		}
 
-		if (bHasZeroLengthNormals)
+		if (bZeroNormals)
 		{
 			Builder.AddWarningMessage(FString::Printf(
 				TEXT("Mesh %s has some nearly zero-length normals which can create some issues. Consider checking 'Recompute Normals' in the asset settings"),
 				MeshName));
 		}
 
-		if (bHasZeroLengthTangents)
+		if (bZeroTangents)
 		{
 			Builder.AddWarningMessage(FString::Printf(
 				TEXT("Mesh %s chas some nearly zero-length tangents which can create some issues. Consider checking 'Recompute Tangents' in the asset settings"),
@@ -44,10 +44,10 @@ namespace
 	}
 
 	template <typename VertexBufferType, typename TangentVectorType>
-	void AnalyzeTangents(const VertexBufferType* VertexBuffer, bool& bOutHasZeroLengthNormals, bool& bOutHasZeroLengthTangents)
+	void CheckTangentVectors(const VertexBufferType* VertexBuffer, bool& bOutZeroNormals, bool& bOutZeroTangents)
 	{
-		bool bHasZeroLengthNormals = false;
-		bool bHasZeroLengthTangents = false;
+		bool bZeroNormals = false;
+		bool bZeroTangents = false;
 
 		typedef TStaticMeshVertexTangentDatum<TangentVectorType> VertexTangentType;
 
@@ -60,13 +60,13 @@ namespace
 			for (uint32 VertexIndex = 0; VertexIndex < VertexCount; ++VertexIndex)
 			{
 				const VertexTangentType& VertexTangent = VertexTangents[VertexIndex];
-				bHasZeroLengthNormals |= VertexTangent.TangentZ.ToFVector().IsNearlyZero();
-				bHasZeroLengthTangents |= VertexTangent.TangentX.ToFVector().IsNearlyZero();
+				bZeroNormals |= VertexTangent.TangentZ.ToFVector().IsNearlyZero();
+				bZeroTangents |= VertexTangent.TangentX.ToFVector().IsNearlyZero();
 			}
 		}
 
-		bOutHasZeroLengthNormals = bHasZeroLengthNormals;
-		bOutHasZeroLengthTangents = bHasZeroLengthTangents;
+		bOutZeroNormals = bZeroNormals;
+		bOutZeroTangents = bZeroTangents;
 	}
 }
 
