[GLTFExporter] refactor: fix minor typo in warning message

commit d5c68aa76344ee694668e77ff6e6a1582a11d04c
Date: Wed, 25 Nov 2020 13:12:54 +0100

    refactor: fix minor typo in warning message
    (cherry picked from commit 7898088040ea875195edaf657893a3c28c9836a8​)

#skipci

diff --git a/Engine/Plugins/Enterprise/GLTFExporter/MigratedFromGit/Source/GLTFExporter/Private/Tasks/GLTFMeshTasks.cpp b/Engine/Plugins/Enterprise/GLTFExporter/MigratedFromGit/Source/GLTFExporter/Private/Tasks/GLTFMeshTasks.cpp
--- a/Engine/Plugins/Enterprise/GLTFExporter/MigratedFromGit/Source/GLTFExporter/Private/Tasks/GLTFMeshTasks.cpp
+++ b/Engine/Plugins/Enterprise/GLTFExporter/MigratedFromGit/Source/GLTFExporter/Private/Tasks/GLTFMeshTasks.cpp
@@ -38,7 +38,7 @@ namespace
 		if (bZeroTangents)
 		{
 			Builder.AddWarningMessage(FString::Printf(
-				TEXT("Mesh %s chas some nearly zero-length tangents which can create some issues. Consider checking 'Recompute Tangents' in the asset settings"),
+				TEXT("Mesh %s has some nearly zero-length tangents which can create some issues. Consider checking 'Recompute Tangents' in the asset settings"),
 				MeshName));
 		}
 	}
