[GLTFExporter] fix: issue error "An explicit Category specifier is required for any property exposed in an Engine module"

commit 11039bb99d3b32c5e63f7de10d45c18230d6821b
Date: Wed, 25 Nov 2020 15:42:32 +0100

    fix: issue error "An explicit Category specifier is required for any property exposed in an Engine module"
    (cherry picked from commit b59b3a582fa2841f94a6477c5c8689d9760cec95​)

#skipci

diff --git a/Engine/Plugins/Enterprise/GLTFExporter/MigratedFromGit/Source/GLTFExporterRuntime/Public/UserData/GLTFMaterialUserData.h b/Engine/Plugins/Enterprise/GLTFExporter/MigratedFromGit/Source/GLTFExporterRuntime/Public/UserData/GLTFMaterialUserData.h
--- a/Engine/Plugins/Enterprise/GLTFExporter/MigratedFromGit/Source/GLTFExporterRuntime/Public/UserData/GLTFMaterialUserData.h
+++ b/Engine/Plugins/Enterprise/GLTFExporter/MigratedFromGit/Source/GLTFExporterRuntime/Public/UserData/GLTFMaterialUserData.h
@@ -55,22 +55,22 @@ struct FGLTFOverrideMaterialBakeSettings
 {
 	GENERATED_BODY()
 
-	UPROPERTY(EditAnywhere, BlueprintReadWrite, meta = (InlineEditConditionToggle))
+	UPROPERTY(EditAnywhere, BlueprintReadWrite, Category = "", meta = (InlineEditConditionToggle))
 	bool bOverrideSize;
 
-	UPROPERTY(EditAnywhere, BlueprintReadWrite, meta = (EditCondition = "bOverrideSize"))
+	UPROPERTY(EditAnywhere, BlueprintReadWrite, Category = "", meta = (EditCondition = "bOverrideSize"))
 	EGLTFMaterialBakeSizePOT Size;
 
-	UPROPERTY(EditAnywhere, BlueprintReadWrite, meta = (InlineEditConditionToggle))
+	UPROPERTY(EditAnywhere, BlueprintReadWrite, Category = "", meta = (InlineEditConditionToggle))
 	bool bOverrideFilter;
 
-	UPROPERTY(EditAnywhere, BlueprintReadWrite, meta = (EditCondition = "bOverrideFilter", ValidEnumValues="TF_Nearest, TF_Bilinear, TF_Trilinear"))
+	UPROPERTY(EditAnywhere, BlueprintReadWrite, Category = "", meta = (EditCondition = "bOverrideFilter", ValidEnumValues="TF_Nearest, TF_Bilinear, TF_Trilinear"))
 	TEnumAsByte<TextureFilter> Filter;
 
-	UPROPERTY(EditAnywhere, BlueprintReadWrite, meta = (InlineEditConditionToggle))
+	UPROPERTY(EditAnywhere, BlueprintReadWrite, Category = "", meta = (InlineEditConditionToggle))
 	bool bOverrideTiling;
 
-	UPROPERTY(EditAnywhere, BlueprintReadWrite, meta = (EditCondition = "bOverrideTiling"))
+	UPROPERTY(EditAnywhere, BlueprintReadWrite, Category = "", meta = (EditCondition = "bOverrideTiling"))
 	TEnumAsByte<TextureAddress> Tiling;
 
 	FGLTFOverrideMaterialBakeSettings();
