[GLTFExporter] feat: add tooltips and descriptions for each export option, relevant enum value, and user data property

commit 6c5ca870303651aba4c06989b53529a249dd8395
Date: Wed, 25 Nov 2020 02:06:51 +0100

    feat: add tooltips and descriptions for each export option, relevant enum value, and user data property
    (cherry picked from commit 79ece650ae2dac0753a8de7b6f70ba6b4d4cbfd0​)
    # Conflicts from cherry-pick:
    #	Plugins/GLTFExporter/Source/GLTFExporterRuntime/Public/UserData/GLTFMaterialUserData.h

#skipci

diff --git a/Engine/Plugins/Enterprise/GLTFExporter/MigratedFromGit/Source/GLTFExporter/Public/GLTFExportOptions.h b/Engine/Plugins/Enterprise/GLTFExporter/MigratedFromGit/Source/GLTFExporter/Public/GLTFExportOptions.h
--- a/Engine/Plugins/Enterprise/GLTFExporter/MigratedFromGit/Source/GLTFExporter/Public/GLTFExportOptions.h
+++ b/Engine/Plugins/Enterprise/GLTFExporter/MigratedFromGit/Source/GLTFExporter/Public/GLTFExportOptions.h
@@ -9,7 +9,9 @@
 UENUM(BlueprintType)
 enum class EGLTFTextureCompression : uint8
 {
+	/** Always use PNG (lossless compression). */
 	PNG,
+	/** If texture does not have an alpha channel, use JPEG (lossy compression); otherwise fallback to PNG. */
     JPEG UMETA(DisplayName = "JPEG (if no alpha)")
 };
 
@@ -29,7 +31,9 @@ ENUM_CLASS_FLAGS(EGLTFTextureType);
 UENUM(BlueprintType)
 enum class EGLTFTextureHDREncoding : uint8
 {
+	/** Clamp HDR colors to standard 8-bit per channel. */
 	None,
+	/** Encode HDR colors to RGBM (will discard alpha). */
 	RGBM
 };
 
@@ -49,8 +53,11 @@ ENUM_CLASS_FLAGS(EGLTFSceneMobility);
 UENUM(BlueprintType)
 enum class EGLTFMaterialVariantMode : uint8
 {
+	/** Never export material variants. */
 	None,
+    /** Export material variants but only use a simple quad if a material input needs to be baked out. */
     Simple,
+    /** Export material variants and allow usage of the mesh data if a material input needs to be baked out with vertex data. */
     UseMeshData,
 };
 
@@ -59,115 +66,151 @@ class GLTFEXPORTER_API UGLTFExportOptions : public UObject
 {
 	GENERATED_UCLASS_BODY()
 
-	/** Scale factor used for exporting assets (0.01 by default) for conversion from centimeters (Unreal default) to meters (glTF). */
+	/** Scale factor used for exporting all assets (0.01 by default) for conversion from centimeters (Unreal default) to meters (glTF). */
 	UPROPERTY(EditAnywhere, BlueprintReadWrite, Config, Category = General)
 	float ExportUniformScale;
 
+	/** If enabled, the preview mesh for a standalone animation or material asset will also be exported. */
 	UPROPERTY(EditAnywhere, BlueprintReadWrite, Config, Category = General)
 	bool bExportPreviewMesh;
 
+	/** If enabled, Unreal's glTF viewer (including a executable launch helper) will be bundled with the exported files. It supports all extensions used. */
 	UPROPERTY(EditAnywhere, BlueprintReadWrite, Config, Category = General)
 	bool bBundleWebViewer;
 
+	/** If enabled, the target folder (with all exported files) will be shown in explorer when export successfully completed. */
 	UPROPERTY(EditAnywhere, BlueprintReadWrite, Config, Category = General)
 	bool bShowFilesWhenDone;
 
+	/** If enabled, materials with shading model unlit will be properly exported. Uses extension KHR_materials_unlit. */
 	UPROPERTY(EditAnywhere, BlueprintReadWrite, Config, Category = Material)
 	bool bExportUnlitMaterials;
 
+	/** If enabled, materials with shading model clear coat will be properly exported. Requires extension KHR_materials_clearcoat, which is not supported by all glTF viewers. */
 	UPROPERTY(EditAnywhere, BlueprintReadWrite, Config, Category = Material)
 	bool bExportClearCoatMaterials;
 
+	/** If enabled, materials with blend modes additive, modulate, and alpha composite will be properly exported. Uses extension EPIC_blend_modes, which is supported by Unreal's glTF viewer. */
 	UPROPERTY(EditAnywhere, BlueprintReadWrite, Config, Category = Material)
 	bool bExportExtraBlendModes;
 
+	/** Bake mode determining if and how a material input is baked out to a texture. Can be overriden by material- and input-specific bake settings, see GLTFMaterialExportOptions. */
 	UPROPERTY(EditAnywhere, BlueprintReadWrite, Config, Category = Material)
 	EGLTFMaterialBakeMode BakeMaterialInputs;
 
+	/** Default size of the baked out texture (containing the material input). Can be overriden by material- and input-specific bake settings, see GLTFMaterialExportOptions. */
 	UPROPERTY(EditAnywhere, BlueprintReadWrite, Config, Category = Material, Meta = (EditCondition = "BakeMaterialInputs != EGLTFMaterialBakeMode::None"))
 	EGLTFMaterialBakeSizePOT DefaultMaterialBakeSize;
 
+	/** Default filtering mode used when sampling the baked out texture. Can be overriden by material- and input-specific bake settings, see GLTFMaterialExportOptions. */
 	UPROPERTY(EditAnywhere, BlueprintReadWrite, Config, Category = Material, Meta = (EditCondition = "BakeMaterialInputs != EGLTFMaterialBakeMode::None", ValidEnumValues="TF_Nearest, TF_Bilinear, TF_Trilinear"))
 	TEnumAsByte<TextureFilter> DefaultMaterialBakeFilter;
 
+	/** Default addressing mode used when sampling the baked out texture. Can be overriden by material- and input-specific bake settings, see GLTFMaterialExportOptions. */
 	UPROPERTY(EditAnywhere, BlueprintReadWrite, Config, Category = Material, Meta = (EditCondition = "BakeMaterialInputs != EGLTFMaterialBakeMode::None"))
 	TEnumAsByte<TextureAddress> DefaultMaterialBakeTiling;
 
+	/** Input-specific default bake settings that override the general defaults above. */
 	UPROPERTY(EditAnywhere, BlueprintReadWrite, Config, Category = Material)
 	TMap<EGLTFMaterialPropertyGroup, FGLTFOverrideMaterialBakeSettings> DefaultInputBakeSettings;
 
+	/** Default LOD level used for exporting a mesh. Can be overriden by component or asset settings (e.g. minimum or forced LOD level). */
 	UPROPERTY(EditAnywhere, BlueprintReadWrite, Config, Category = Mesh, Meta = (ClampMin = "0"))
 	int32 DefaultLevelOfDetail;
 
+	/** If enabled, export vertex color. Not recommended due to vertex colors always being used as a base color multiplier in glTF, regardless of material. Often producing undesirable results. */
 	UPROPERTY(EditAnywhere, BlueprintReadWrite, Config, Category = Mesh)
 	bool bExportVertexColors;
 
+	/** If enabled, export vertex bone weights and indices in skeletal meshes. Necessary for animation sequences. */
 	UPROPERTY(EditAnywhere, BlueprintReadWrite, Config, Category = Mesh)
 	bool bExportVertexSkinWeights;
 
+	/** If enabled, export Unreal-configured quantization for vertex tangents and normals, reducing size. Requires extension KHR_mesh_quantization, which is not supported by all glTF viewers. */
 	UPROPERTY(EditAnywhere, BlueprintReadWrite, Config, Category = Mesh)
 	bool bExportMeshQuantization;
 
+	/** If enabled, export level sequences. Only transform tracks are currently supported. The level sequence will be played at the assigned display rate. */
 	UPROPERTY(EditAnywhere, BlueprintReadWrite, Config, Category = Animation)
 	bool bExportLevelSequences;
 
+	/** If enabled, export single animation asset used by a skeletal mesh component or hotspot actor. Export of vertex skin weights must be enabled. */
 	UPROPERTY(EditAnywhere, BlueprintReadWrite, Config, Category = Animation, Meta = (EditCondition = "bExportVertexSkinWeights"))
 	bool bExportAnimationSequences;
 
+	/** If enabled, apply animation retargeting to skeleton bones when exporting a animation sequence. */
 	UPROPERTY(EditAnywhere, BlueprintReadWrite, Config, Category = Animation, Meta = (EditCondition = "bExportVertexSkinWeights && bExportAnimationSequences"))
 	bool bRetargetBoneTransforms;
 
+	/** If enabled, export play rate, start time, looping, and auto play for a animation or level sequence. Uses extension EPIC_animation_playback, which is supported by Unreal's glTF viewer. */
 	UPROPERTY(EditAnywhere, BlueprintReadWrite, Config, Category = Animation)
 	bool bExportPlaybackSettings;
 
+	/** Desired image format used for exported textures. */
 	UPROPERTY(EditAnywhere, BlueprintReadWrite, Config, Category = Texture)
 	EGLTFTextureCompression TextureCompression;
 
+	/** Level of compression used for exported textures, between 1 (worst quality, best compression) and 100 (best quality, worst compression). Does not apply to lossless formats (e.g. PNG). */
 	UPROPERTY(EditAnywhere, BlueprintReadWrite, Config, Category = Texture, Meta = (ClampMin = "0", ClampMax = "100", EditCondition = "TextureCompression != EGLTFTextureCompression::PNG"))
 	int32 TextureCompressionQuality;
 
+	/** Texture types that will always use lossless formats (e.g. PNG) because of sensitivity to compression artifacts. */
 	UPROPERTY(EditAnywhere, BlueprintReadWrite, Config, Category = Texture, Meta = (Bitmask, BitmaskEnum = EGLTFTextureType, EditCondition = "TextureCompression != EGLTFTextureCompression::PNG"))
 	int32 NoLossyCompressionFor;
 
+	/** If enabled, export UV tiling and un-mirroring settings in a texture coordinate expression node for simple material input expressions. Uses extension KHR_texture_transform. */
 	UPROPERTY(EditAnywhere, BlueprintReadWrite, Config, Category = Texture)
 	bool bExportTextureTransforms;
 
+	/** If enabled, export lightmaps (created by Lightmass) when exporting a level. Uses extension EPIC_lightmap_textures, which is supported by Unreal's glTF viewer. */
 	UPROPERTY(EditAnywhere, BlueprintReadWrite, Config, Category = Texture)
 	bool bExportLightmaps;
 
+	/** Encoding used to store textures that have HDR pixel data. Uses extension EPIC_texture_hdr_encoding, which is supported by Unreal's glTF viewer. */
 	UPROPERTY(EditAnywhere, BlueprintReadWrite, Config, Category = Texture, Meta = (DisplayName = "Texture HDR Encoding"))
 	EGLTFTextureHDREncoding TextureHDREncoding;
 
+	/** If enabled, export components that are flagged as hidden in-game. */
 	UPROPERTY(EditAnywhere, BlueprintReadWrite, Config, Category = Scene)
 	bool bExportHiddenInGame;
 
+	/** Mobility of directional, point, and spot light components that will be exported. Uses extension KHR_lights_punctual. */
 	UPROPERTY(EditAnywhere, BlueprintReadWrite, Config, Category = Scene, Meta = (Bitmask, BitmaskEnum = EGLTFSceneMobility))
 	int32 ExportLights;
 
+	/** If enabled, export camera components. */
 	UPROPERTY(EditAnywhere, BlueprintReadWrite, Config, Category = Scene)
 	bool bExportCameras;
 
+	/** If enabled, export GLTFCameraActors. Uses extension EPIC_camera_controls, which is supported by Unreal's glTF viewer. */
 	UPROPERTY(EditAnywhere, BlueprintReadWrite, Config, Category = Scene, Meta = (EditCondition = "bExportCameras"))
 	bool bExportCameraControls;
 
+	/** If enabled, export GLTFHotspotActors. Uses extension EPIC_animation_hotspots, which is supported by Unreal's glTF viewer. */
 	UPROPERTY(EditAnywhere, BlueprintReadWrite, Config, Category = Scene)
 	bool bExportAnimationHotspots;
 
+	/** If enabled, export HDRIBackdrop blueprints. Uses extension EPIC_hdri_backdrops, which is supported by Unreal's glTF viewer. */
 	UPROPERTY(EditAnywhere, BlueprintReadWrite, Config, Category = Scene, Meta = (DisplayName = "Export HDRI Backdrops"))
 	bool bExportHDRIBackdrops;
 
+	/** If enabled, export SkySphere blueprints. Uses extension EPIC_sky_spheres, which is supported by Unreal's glTF viewer. */
 	UPROPERTY(EditAnywhere, BlueprintReadWrite, Config, Category = Scene)
 	bool bExportSkySpheres;
 
+	/** If enabled, export LevelVariantSetsActors. Uses extension EPIC_level_variant_sets, which is supported by Unreal's glTF viewer. */
 	UPROPERTY(EditAnywhere, BlueprintReadWrite, Config, Category = Scene)
 	bool bExportVariantSets;
 
+	/** Mode determining if and how to export material variants that change the materials property on a static or skeletal mesh component. */
 	UPROPERTY(EditAnywhere, BlueprintReadWrite, Config, Category = VariantSets, Meta = (EditCondition = "bExportVariantSets"))
 	EGLTFMaterialVariantMode ExportMaterialVariants;
 
+	/** If enabled, export variants that change the mesh property on a static or skeletal mesh component. */
 	UPROPERTY(EditAnywhere, BlueprintReadWrite, Config, Category = VariantSets, Meta = (EditCondition = "bExportVariantSets"))
 	bool bExportMeshVariants;
 
+	/** If enabled, export variants that change the visible property on a scene component. */
 	UPROPERTY(EditAnywhere, BlueprintReadWrite, Config, Category = VariantSets, Meta = (EditCondition = "bExportVariantSets"))
 	bool bExportVisibilityVariants;
 
diff --git a/Engine/Plugins/Enterprise/GLTFExporter/MigratedFromGit/Source/GLTFExporterRuntime/Public/UserData/GLTFMaterialUserData.h b/Engine/Plugins/Enterprise/GLTFExporter/MigratedFromGit/Source/GLTFExporterRuntime/Public/UserData/GLTFMaterialUserData.h
--- a/Engine/Plugins/Enterprise/GLTFExporter/MigratedFromGit/Source/GLTFExporterRuntime/Public/UserData/GLTFMaterialUserData.h
+++ b/Engine/Plugins/Enterprise/GLTFExporter/MigratedFromGit/Source/GLTFExporterRuntime/Public/UserData/GLTFMaterialUserData.h
@@ -12,8 +12,11 @@ enum TextureAddress;
 UENUM(BlueprintType)
 enum class EGLTFMaterialBakeMode : uint8
 {
-    Disabled,
+	/** Never bake material inputs. */
+	Disabled,
+	/** Only use a simple quad if a material input needs to be baked out. */
 	Simple,
+	/** Allow usage of the mesh data if a material input needs to be baked out with vertex data. */
 	UseMeshData,
 };
 
@@ -55,25 +58,31 @@ struct FGLTFOverrideMaterialBakeSettings
 {
 	GENERATED_BODY()
 
+	FGLTFOverrideMaterialBakeSettings();
+
+	/** If enabled, override size of the baked out texture. */
 	UPROPERTY(EditAnywhere, BlueprintReadWrite, Category = "", meta = (InlineEditConditionToggle))
 	bool bOverrideSize;
 
+	/** Overriding size of the baked out texture. */
 	UPROPERTY(EditAnywhere, BlueprintReadWrite, Category = "", meta = (EditCondition = "bOverrideSize"))
 	EGLTFMaterialBakeSizePOT Size;
 
+	/** If enabled, override filtering mode used when sampling the baked out texture. */
 	UPROPERTY(EditAnywhere, BlueprintReadWrite, Category = "", meta = (InlineEditConditionToggle))
 	bool bOverrideFilter;
 
+	/** Overriding filtering mode used when sampling the baked out texture. */
 	UPROPERTY(EditAnywhere, BlueprintReadWrite, Category = "", meta = (EditCondition = "bOverrideFilter", ValidEnumValues="TF_Nearest, TF_Bilinear, TF_Trilinear"))
 	TEnumAsByte<TextureFilter> Filter;
 
+	/** If enabled, override addressing mode used when sampling the baked out texture. */
 	UPROPERTY(EditAnywhere, BlueprintReadWrite, Category = "", meta = (InlineEditConditionToggle))
 	bool bOverrideTiling;
 
+	/** Overriding addressing mode used when sampling the baked out texture. */
 	UPROPERTY(EditAnywhere, BlueprintReadWrite, Category = "", meta = (EditCondition = "bOverrideTiling"))
 	TEnumAsByte<TextureAddress> Tiling;
-
-	FGLTFOverrideMaterialBakeSettings();
 };
 
 /** glTF-specific user data that can be added to material assets to override export options */
@@ -86,9 +95,11 @@ public:
 
 	// TODO: add support for overriding more export options
 
+	/** Default bake settings for this material in general. */
 	UPROPERTY(EditAnywhere, BlueprintReadWrite, Category = "Override Bake Settings", meta = (ShowOnlyInnerProperties))
 	FGLTFOverrideMaterialBakeSettings Default;
 
+	/** Input-specific bake settings that override the defaults above. */
 	UPROPERTY(EditAnywhere, BlueprintReadWrite, Category = "Override Bake Settings")
 	TMap<EGLTFMaterialPropertyGroup, FGLTFOverrideMaterialBakeSettings> Inputs;
 
