[GLTFExporter] refactor: fix minor typos in export option descriptions

commit 0dc189d778eed75516509b6b54bd844753db2083
Date: Wed, 25 Nov 2020 12:56:25 +0100

    refactor: fix minor typos in export option descriptions
    (cherry picked from commit 6f73becf4131ef6ec2e3ac9db380c25afb472bc1​)

#skipci

diff --git a/Engine/Plugins/Enterprise/GLTFExporter/MigratedFromGit/Source/GLTFExporter/Public/GLTFExportOptions.h b/Engine/Plugins/Enterprise/GLTFExporter/MigratedFromGit/Source/GLTFExporter/Public/GLTFExportOptions.h
--- a/Engine/Plugins/Enterprise/GLTFExporter/MigratedFromGit/Source/GLTFExporter/Public/GLTFExportOptions.h
+++ b/Engine/Plugins/Enterprise/GLTFExporter/MigratedFromGit/Source/GLTFExporter/Public/GLTFExportOptions.h
@@ -74,11 +74,11 @@ class GLTFEXPORTER_API UGLTFExportOptions : public UObject
 	UPROPERTY(EditAnywhere, BlueprintReadWrite, Config, Category = General)
 	bool bExportPreviewMesh;
 
-	/** If enabled, Unreal's glTF viewer (including a executable launch helper) will be bundled with the exported files. It supports all extensions used. */
+	/** If enabled, Unreal's glTF viewer (including an executable launch helper) will be bundled with the exported files. It supports all extensions used. */
 	UPROPERTY(EditAnywhere, BlueprintReadWrite, Config, Category = General)
 	bool bBundleWebViewer;
 
-	/** If enabled, the target folder (with all exported files) will be shown in explorer when export successfully completed. */
+	/** If enabled, the target folder (with all exported files) will be shown in explorer once the export has been completed successfully. */
 	UPROPERTY(EditAnywhere, BlueprintReadWrite, Config, Category = General)
 	bool bShowFilesWhenDone;
 
@@ -138,11 +138,11 @@ class GLTFEXPORTER_API UGLTFExportOptions : public UObject
 	UPROPERTY(EditAnywhere, BlueprintReadWrite, Config, Category = Animation, Meta = (EditCondition = "bExportVertexSkinWeights"))
 	bool bExportAnimationSequences;
 
-	/** If enabled, apply animation retargeting to skeleton bones when exporting a animation sequence. */
+	/** If enabled, apply animation retargeting to skeleton bones when exporting an animation sequence. */
 	UPROPERTY(EditAnywhere, BlueprintReadWrite, Config, Category = Animation, Meta = (EditCondition = "bExportVertexSkinWeights && bExportAnimationSequences"))
 	bool bRetargetBoneTransforms;
 
-	/** If enabled, export play rate, start time, looping, and auto play for a animation or level sequence. Uses extension EPIC_animation_playback, which is supported by Unreal's glTF viewer. */
+	/** If enabled, export play rate, start time, looping, and auto play for an animation or level sequence. Uses extension EPIC_animation_playback, which is supported by Unreal's glTF viewer. */
 	UPROPERTY(EditAnywhere, BlueprintReadWrite, Config, Category = Animation)
 	bool bExportPlaybackSettings;
 
