[GLTFExporter] refactor: minor tweaks to some export option descriptions

commit 1a56df6b26bac6be4c83b067ad350aaa2d55fc9b
Date: Wed, 25 Nov 2020 14:54:05 +0100

    refactor: minor tweaks to some export option descriptions
    (cherry picked from commit 1594a6b7a99f8cf02ce957744abd265c19268d03​)

#skipci

diff --git a/Engine/Plugins/Enterprise/GLTFExporter/MigratedFromGit/Source/GLTFExporter/Public/GLTFExportOptions.h b/Engine/Plugins/Enterprise/GLTFExporter/MigratedFromGit/Source/GLTFExporter/Public/GLTFExportOptions.h
--- a/Engine/Plugins/Enterprise/GLTFExporter/MigratedFromGit/Source/GLTFExporter/Public/GLTFExportOptions.h
+++ b/Engine/Plugins/Enterprise/GLTFExporter/MigratedFromGit/Source/GLTFExporter/Public/GLTFExportOptions.h
@@ -86,7 +86,7 @@ class GLTFEXPORTER_API UGLTFExportOptions : public UObject
 	UPROPERTY(EditAnywhere, BlueprintReadWrite, Config, Category = Material)
 	bool bExportUnlitMaterials;
 
-	/** If enabled, materials with shading model clear coat will be properly exported. Requires extension KHR_materials_clearcoat, which is not supported by all glTF viewers. */
+	/** If enabled, materials with shading model clear coat will be properly exported. Uses extension KHR_materials_clearcoat, which is not supported by all glTF viewers. */
 	UPROPERTY(EditAnywhere, BlueprintReadWrite, Config, Category = Material)
 	bool bExportClearCoatMaterials;
 
@@ -166,7 +166,7 @@ class GLTFEXPORTER_API UGLTFExportOptions : public UObject
 	UPROPERTY(EditAnywhere, BlueprintReadWrite, Config, Category = Texture)
 	bool bExportLightmaps;
 
-	/** Encoding used to store textures that have HDR pixel data. Uses extension EPIC_texture_hdr_encoding, which is supported by Unreal's glTF viewer. */
+	/** Encoding used to store textures that have pixel colors with more than 8-bit per channel. Uses extension EPIC_texture_hdr_encoding, which is supported by Unreal's glTF viewer. */
 	UPROPERTY(EditAnywhere, BlueprintReadWrite, Config, Category = Texture, Meta = (DisplayName = "Texture HDR Encoding"))
 	EGLTFTextureHDREncoding TextureHDREncoding;
 
