[GLTFExporter] refactor: rephrase descriptions for properties that override bake settings

commit 4438526ca0056a93a2f64221fd0d9ad8a0df8f95
Date: Wed, 25 Nov 2020 14:56:17 +0100

    refactor: rephrase descriptions for properties that override bake settings
    (cherry picked from commit 190a909f65ecb7095fbe903032b17ad27186f98d​)
    # Conflicts from cherry-pick:
    #	Plugins/GLTFExporter/Source/GLTFExporterRuntime/Public/UserData/GLTFMaterialUserData.h

#skipci

diff --git a/Engine/Plugins/Enterprise/GLTFExporter/MigratedFromGit/Source/GLTFExporterRuntime/Public/UserData/GLTFMaterialUserData.h b/Engine/Plugins/Enterprise/GLTFExporter/MigratedFromGit/Source/GLTFExporterRuntime/Public/UserData/GLTFMaterialUserData.h
--- a/Engine/Plugins/Enterprise/GLTFExporter/MigratedFromGit/Source/GLTFExporterRuntime/Public/UserData/GLTFMaterialUserData.h
+++ b/Engine/Plugins/Enterprise/GLTFExporter/MigratedFromGit/Source/GLTFExporterRuntime/Public/UserData/GLTFMaterialUserData.h
@@ -60,27 +60,27 @@ struct FGLTFOverrideMaterialBakeSettings
 
 	FGLTFOverrideMaterialBakeSettings();
 
-	/** If enabled, override size of the baked out texture. */
+	/** If enabled, default size will be overridden by the corresponding property. */
 	UPROPERTY(EditAnywhere, BlueprintReadWrite, Category = "", meta = (InlineEditConditionToggle))
 	bool bOverrideSize;
 
-	/** Overriding size of the baked out texture. */
+	/** Overrides default size of the baked out texture. */
 	UPROPERTY(EditAnywhere, BlueprintReadWrite, Category = "", meta = (EditCondition = "bOverrideSize"))
 	EGLTFMaterialBakeSizePOT Size;
 
-	/** If enabled, override filtering mode used when sampling the baked out texture. */
+	/** If enabled, default filtering mode will be overridden by the corresponding property. */
 	UPROPERTY(EditAnywhere, BlueprintReadWrite, Category = "", meta = (InlineEditConditionToggle))
 	bool bOverrideFilter;
 
-	/** Overriding filtering mode used when sampling the baked out texture. */
+	/** Overrides the default filtering mode used when sampling the baked out texture. */
 	UPROPERTY(EditAnywhere, BlueprintReadWrite, Category = "", meta = (EditCondition = "bOverrideFilter", ValidEnumValues="TF_Nearest, TF_Bilinear, TF_Trilinear"))
 	TEnumAsByte<TextureFilter> Filter;
 
-	/** If enabled, override addressing mode used when sampling the baked out texture. */
+	/** If enabled, default addressing mode will be overridden by the corresponding property. */
 	UPROPERTY(EditAnywhere, BlueprintReadWrite, Category = "", meta = (InlineEditConditionToggle))
 	bool bOverrideTiling;
 
-	/** Overriding addressing mode used when sampling the baked out texture. */
+	/** Overrides the default addressing mode used when sampling the baked out texture. */
 	UPROPERTY(EditAnywhere, BlueprintReadWrite, Category = "", meta = (EditCondition = "bOverrideTiling"))
 	TEnumAsByte<TextureAddress> Tiling;
 };
