[GLTFExporter] fix: fixed error ('BakeSettings': redefinition; previous definition was 'data variable') when running packaging plugin

commit c17bb293f5c9e9c62395cc10bfe6923e6c4b26a0
Date: Thu, 26 Nov 2020 11:27:16 +0100

    fix: fixed error ('BakeSettings': redefinition; previous definition was 'data variable') when running packaging plugin
    (cherry picked from commit b8897f52391435513c87bd72c62de4a0ae1c02f7​)

#skipci

diff --git a/Engine/Plugins/Enterprise/GLTFExporter/MigratedFromGit/Source/GLTFExporterRuntime/Private/UserData/GLTFMaterialUserData.cpp b/Engine/Plugins/Enterprise/GLTFExporter/MigratedFromGit/Source/GLTFExporterRuntime/Private/UserData/GLTFMaterialUserData.cpp
--- a/Engine/Plugins/Enterprise/GLTFExporter/MigratedFromGit/Source/GLTFExporterRuntime/Private/UserData/GLTFMaterialUserData.cpp
+++ b/Engine/Plugins/Enterprise/GLTFExporter/MigratedFromGit/Source/GLTFExporterRuntime/Private/UserData/GLTFMaterialUserData.cpp
@@ -16,7 +16,7 @@ FGLTFOverrideMaterialBakeSettings::FGLTFOverrideMaterialBakeSettings()
 
 EGLTFMaterialBakeSizePOT UGLTFMaterialExportOptions::GetBakeSizeForPropertyGroup(const UMaterialInterface* Material, EGLTFMaterialPropertyGroup PropertyGroup, EGLTFMaterialBakeSizePOT DefaultValue)
 {
-	if (const FGLTFOverrideMaterialBakeSettings* BakeSettings = GetBakeSettingsByPredicate(Material, PropertyGroup, [](const FGLTFOverrideMaterialBakeSettings& BakeSettings) { return BakeSettings.bOverrideSize; }))
+	if (const FGLTFOverrideMaterialBakeSettings* BakeSettings = GetBakeSettingsByPredicate(Material, PropertyGroup, [](const FGLTFOverrideMaterialBakeSettings& Settings) { return Settings.bOverrideSize; }))
 	{
 		return BakeSettings->Size;
 	}
@@ -26,7 +26,7 @@ EGLTFMaterialBakeSizePOT UGLTFMaterialExportOptions::GetBakeSizeForPropertyGroup
 
 TextureFilter UGLTFMaterialExportOptions::GetBakeFilterForPropertyGroup(const UMaterialInterface* Material, EGLTFMaterialPropertyGroup PropertyGroup, TextureFilter DefaultValue)
 {
-	if (const FGLTFOverrideMaterialBakeSettings* BakeSettings = GetBakeSettingsByPredicate(Material, PropertyGroup, [](const FGLTFOverrideMaterialBakeSettings& BakeSettings) { return BakeSettings.bOverrideFilter; }))
+	if (const FGLTFOverrideMaterialBakeSettings* BakeSettings = GetBakeSettingsByPredicate(Material, PropertyGroup, [](const FGLTFOverrideMaterialBakeSettings& Settings) { return Settings.bOverrideFilter; }))
 	{
 		return BakeSettings->Filter;
 	}
@@ -36,7 +36,7 @@ TextureFilter UGLTFMaterialExportOptions::GetBakeFilterForPropertyGroup(const UM
 
 TextureAddress UGLTFMaterialExportOptions::GetBakeTilingForPropertyGroup(const UMaterialInterface* Material, EGLTFMaterialPropertyGroup PropertyGroup, TextureAddress DefaultValue)
 {
-	if (const FGLTFOverrideMaterialBakeSettings* BakeSettings = GetBakeSettingsByPredicate(Material, PropertyGroup, [](const FGLTFOverrideMaterialBakeSettings& BakeSettings) { return BakeSettings.bOverrideTiling; }))
+	if (const FGLTFOverrideMaterialBakeSettings* BakeSettings = GetBakeSettingsByPredicate(Material, PropertyGroup, [](const FGLTFOverrideMaterialBakeSettings& Settings) { return Settings.bOverrideTiling; }))
 	{
 		return BakeSettings->Tiling;
 	}
