[GLTFExporter] feat: add option to not save any textures for clean geometry exports

commit 952d7db3aeec5787fbcdc9e615fc11bf288ade2b
Date: Fri, 27 Nov 2020 02:53:57 +0100

    feat: add option to not save any textures for clean geometry exports
    (cherry picked from commit 447663409e96282a5e56fca54dda903e577dd09a​)

#skipci

diff --git a/Engine/Plugins/Enterprise/GLTFExporter/MigratedFromGit/Source/GLTFExporter/Private/Builders/GLTFImageBuilder.cpp b/Engine/Plugins/Enterprise/GLTFExporter/MigratedFromGit/Source/GLTFExporter/Private/Builders/GLTFImageBuilder.cpp
--- a/Engine/Plugins/Enterprise/GLTFExporter/MigratedFromGit/Source/GLTFExporter/Private/Builders/GLTFImageBuilder.cpp
+++ b/Engine/Plugins/Enterprise/GLTFExporter/MigratedFromGit/Source/GLTFExporter/Private/Builders/GLTFImageBuilder.cpp
@@ -29,6 +29,9 @@ FGLTFJsonImageIndex FGLTFImageBuilder::AddImage(const FColor* Pixels, FIntPoint
 	const EGLTFJsonMimeType ImageFormat = GetImageFormat(Pixels, Size, bIgnoreAlpha, Type);
 	switch (ImageFormat)
 	{
+		case EGLTFJsonMimeType::None:
+			return FGLTFJsonImageIndex(INDEX_NONE);
+
 		case EGLTFJsonMimeType::PNG:
 			FGLTFImageUtility::CompressToPNG(Pixels, Size, CompressedData);
 			break;
@@ -77,6 +80,9 @@ EGLTFJsonMimeType FGLTFImageBuilder::GetImageFormat(const FColor* Pixels, FIntPo
 {
 	switch (ExportOptions->TextureCompression)
 	{
+		case EGLTFTextureCompression::None:
+			return EGLTFJsonMimeType::None;
+
 		case EGLTFTextureCompression::PNG:
 			return EGLTFJsonMimeType::PNG;
 
diff --git a/Engine/Plugins/Enterprise/GLTFExporter/MigratedFromGit/Source/GLTFExporter/Private/Converters/GLTFTextureConverters.cpp b/Engine/Plugins/Enterprise/GLTFExporter/MigratedFromGit/Source/GLTFExporter/Private/Converters/GLTFTextureConverters.cpp
--- a/Engine/Plugins/Enterprise/GLTFExporter/MigratedFromGit/Source/GLTFExporter/Private/Converters/GLTFTextureConverters.cpp
+++ b/Engine/Plugins/Enterprise/GLTFExporter/MigratedFromGit/Source/GLTFExporter/Private/Converters/GLTFTextureConverters.cpp
@@ -6,6 +6,11 @@
 
 FGLTFJsonTextureIndex FGLTFTexture2DConverter::Convert(const UTexture2D* Texture2D)
 {
+	if (Builder.ExportOptions->TextureCompression == EGLTFTextureCompression::None)
+	{
+		return FGLTFJsonTextureIndex(INDEX_NONE);
+	}
+
 	const FGLTFJsonTextureIndex TextureIndex = Builder.AddTexture();
 	Builder.SetupTask<FGLTFTexture2DTask>(Builder, Texture2D, TextureIndex);
 	return TextureIndex;
@@ -13,6 +18,11 @@ FGLTFJsonTextureIndex FGLTFTexture2DConverter::Convert(const UTexture2D* Texture
 
 FGLTFJsonTextureIndex FGLTFTextureCubeConverter::Convert(const UTextureCube* TextureCube, ECubeFace CubeFace)
 {
+	if (Builder.ExportOptions->TextureCompression == EGLTFTextureCompression::None)
+	{
+		return FGLTFJsonTextureIndex(INDEX_NONE);
+	}
+
 	const FGLTFJsonTextureIndex TextureIndex = Builder.AddTexture();
 	Builder.SetupTask<FGLTFTextureCubeTask>(Builder, TextureCube, CubeFace, TextureIndex);
 	return TextureIndex;
@@ -20,6 +30,11 @@ FGLTFJsonTextureIndex FGLTFTextureCubeConverter::Convert(const UTextureCube* Tex
 
 FGLTFJsonTextureIndex FGLTFTextureRenderTarget2DConverter::Convert(const UTextureRenderTarget2D* RenderTarget2D)
 {
+	if (Builder.ExportOptions->TextureCompression == EGLTFTextureCompression::None)
+	{
+		return FGLTFJsonTextureIndex(INDEX_NONE);
+	}
+
 	const FGLTFJsonTextureIndex TextureIndex = Builder.AddTexture();
 	Builder.SetupTask<FGLTFTextureRenderTarget2DTask>(Builder, RenderTarget2D, TextureIndex);
 	return TextureIndex;
@@ -27,6 +42,11 @@ FGLTFJsonTextureIndex FGLTFTextureRenderTarget2DConverter::Convert(const UTextur
 
 FGLTFJsonTextureIndex FGLTFTextureRenderTargetCubeConverter::Convert(const UTextureRenderTargetCube* RenderTargetCube, ECubeFace CubeFace)
 {
+	if (Builder.ExportOptions->TextureCompression == EGLTFTextureCompression::None)
+	{
+		return FGLTFJsonTextureIndex(INDEX_NONE);
+	}
+
 	const FGLTFJsonTextureIndex TextureIndex = Builder.AddTexture();
 	Builder.SetupTask<FGLTFTextureRenderTargetCubeTask>(Builder, RenderTargetCube, CubeFace, TextureIndex);
 	return TextureIndex;
@@ -34,6 +54,11 @@ FGLTFJsonTextureIndex FGLTFTextureRenderTargetCubeConverter::Convert(const UText
 
 FGLTFJsonTextureIndex FGLTFTextureLightMapConverter::Convert(const ULightMapTexture2D* LightMap)
 {
+	if (Builder.ExportOptions->TextureCompression == EGLTFTextureCompression::None)
+	{
+		return FGLTFJsonTextureIndex(INDEX_NONE);
+	}
+
 	const FGLTFJsonTextureIndex TextureIndex = Builder.AddTexture();
 	Builder.SetupTask<FGLTFTextureLightMapTask>(Builder, LightMap, TextureIndex);
 	return TextureIndex;
diff --git a/Engine/Plugins/Enterprise/GLTFExporter/MigratedFromGit/Source/GLTFExporter/Private/Tasks/GLTFMaterialTasks.cpp b/Engine/Plugins/Enterprise/GLTFExporter/MigratedFromGit/Source/GLTFExporter/Private/Tasks/GLTFMaterialTasks.cpp
--- a/Engine/Plugins/Enterprise/GLTFExporter/MigratedFromGit/Source/GLTFExporter/Private/Tasks/GLTFMaterialTasks.cpp
+++ b/Engine/Plugins/Enterprise/GLTFExporter/MigratedFromGit/Source/GLTFExporter/Private/Tasks/GLTFMaterialTasks.cpp
@@ -363,6 +363,11 @@ bool FGLTFMaterialTask::TryGetBaseColorAndOpacity(FGLTFJsonPBRMetallicRoughness&
 		return true;
 	}
 
+	if (Builder.ExportOptions->TextureCompression == EGLTFTextureCompression::None)
+	{
+		return true;
+	}
+
 	int32 TexCoord;
 	if (BaseColorBakeOutput.bIsConstant)
 	{
@@ -481,6 +486,11 @@ bool FGLTFMaterialTask::TryGetMetallicAndRoughness(FGLTFJsonPBRMetallicRoughness
 		return true;
 	}
 
+	if (Builder.ExportOptions->TextureCompression == EGLTFTextureCompression::None)
+	{
+		return true;
+	}
+
 	int32 TexCoord;
 	if (MetallicBakeOutput.bIsConstant)
 	{
@@ -603,6 +613,11 @@ bool FGLTFMaterialTask::TryGetClearCoatRoughness(FGLTFJsonClearCoatExtension& Ou
 		return true;
 	}
 
+	if (Builder.ExportOptions->TextureCompression == EGLTFTextureCompression::None)
+	{
+		return true;
+	}
+
 	int32 TexCoord;
 	if (IntensityBakeOutput.bIsConstant)
 	{
@@ -686,6 +701,12 @@ bool FGLTFMaterialTask::TryGetEmissive(FGLTFJsonMaterial& JsonMaterial, const FM
 	}
 	else
 	{
+		if (Builder.ExportOptions->TextureCompression == EGLTFTextureCompression::None)
+		{
+			JsonMaterial.EmissiveTexture.Index = FGLTFJsonTextureIndex(INDEX_NONE);
+			return true;
+		}
+
 		if (!StoreBakedPropertyTexture(JsonMaterial.EmissiveTexture, PropertyBakeOutput, TEXT("Emissive")))
 		{
 			return false;
@@ -973,6 +994,11 @@ bool FGLTFMaterialTask::TryGetSourceTexture(FGLTFJsonTextureInfo& OutTexInfo, co
 
 bool FGLTFMaterialTask::TryGetSourceTexture(const UTexture2D*& OutTexture, int32& OutTexCoord, FGLTFJsonTextureTransform& OutTransform, const FMaterialPropertyEx& Property, const TArray<FLinearColor>& AllowedMasks) const
 {
+	if (Builder.ExportOptions->TextureCompression == EGLTFTextureCompression::None)
+	{
+		return false;
+	}
+
 	const FExpressionInput* MaterialInput = FGLTFMaterialUtility::GetInputForProperty(Material, Property);
 	if (MaterialInput == nullptr)
 	{
@@ -1101,6 +1127,12 @@ bool FGLTFMaterialTask::TryGetBakedMaterialProperty(FGLTFJsonTextureInfo& OutTex
 		return true;
 	}
 
+	if (Builder.ExportOptions->TextureCompression == EGLTFTextureCompression::None)
+	{
+		OutTexInfo.Index = FGLTFJsonTextureIndex(INDEX_NONE);
+		return true;
+	}
+
 	if (bTransformToLinear)
 	{
 		FGLTFMaterialUtility::TransformToLinear(PropertyBakeOutput.Pixels);
@@ -1134,6 +1166,12 @@ bool FGLTFMaterialTask::TryGetBakedMaterialProperty(FGLTFJsonTextureInfo& OutTex
 		return true;
 	}
 
+	if (Builder.ExportOptions->TextureCompression == EGLTFTextureCompression::None)
+	{
+		OutTexInfo.Index = FGLTFJsonTextureIndex(INDEX_NONE);
+		return true;
+	}
+
 	if (bTransformToLinear)
 	{
 		FGLTFMaterialUtility::TransformToLinear(PropertyBakeOutput.Pixels);
@@ -1167,6 +1205,12 @@ inline bool FGLTFMaterialTask::TryGetBakedMaterialProperty(FGLTFJsonTextureInfo&
 		return true;
 	}
 
+	if (Builder.ExportOptions->TextureCompression == EGLTFTextureCompression::None)
+	{
+		OutTexInfo.Index = FGLTFJsonTextureIndex(INDEX_NONE);
+		return true;
+	}
+
 	if (bTransformToLinear)
 	{
 		FGLTFMaterialUtility::TransformToLinear(PropertyBakeOutput.Pixels);
@@ -1196,6 +1240,12 @@ bool FGLTFMaterialTask::TryGetBakedMaterialProperty(FGLTFJsonTextureInfo& OutTex
 
 	if (!PropertyBakeOutput.bIsConstant)
 	{
+		if (Builder.ExportOptions->TextureCompression == EGLTFTextureCompression::None)
+		{
+			OutTexInfo.Index = FGLTFJsonTextureIndex(INDEX_NONE);
+			return true;
+		}
+
 		if (bTransformToLinear)
 		{
 			FGLTFMaterialUtility::TransformToLinear(PropertyBakeOutput.Pixels);
diff --git a/Engine/Plugins/Enterprise/GLTFExporter/MigratedFromGit/Source/GLTFExporter/Public/GLTFExportOptions.h b/Engine/Plugins/Enterprise/GLTFExporter/MigratedFromGit/Source/GLTFExporter/Public/GLTFExportOptions.h
--- a/Engine/Plugins/Enterprise/GLTFExporter/MigratedFromGit/Source/GLTFExporter/Public/GLTFExportOptions.h
+++ b/Engine/Plugins/Enterprise/GLTFExporter/MigratedFromGit/Source/GLTFExporter/Public/GLTFExportOptions.h
@@ -9,6 +9,8 @@
 UENUM(BlueprintType)
 enum class EGLTFTextureCompression : uint8
 {
+	/** Don't export any textures. */
+	None,
 	/** Always use PNG (lossless compression). */
 	PNG,
 	/** If texture does not have an alpha channel, use JPEG (lossy compression); otherwise fallback to PNG. */
@@ -151,23 +153,23 @@ class GLTFEXPORTER_API UGLTFExportOptions : public UObject
 	EGLTFTextureCompression TextureCompression;
 
 	/** Level of compression used for exported textures, between 1 (worst quality, best compression) and 100 (best quality, worst compression). Does not apply to lossless formats (e.g. PNG). */
-	UPROPERTY(EditAnywhere, BlueprintReadWrite, Config, Category = Texture, Meta = (ClampMin = "0", ClampMax = "100", EditCondition = "TextureCompression != EGLTFTextureCompression::PNG"))
+	UPROPERTY(EditAnywhere, BlueprintReadWrite, Config, Category = Texture, Meta = (ClampMin = "0", ClampMax = "100", EditCondition = "TextureCompression == EGLTFTextureCompression::JPEG"))
 	int32 TextureCompressionQuality;
 
 	/** Texture types that will always use lossless formats (e.g. PNG) because of sensitivity to compression artifacts. */
-	UPROPERTY(EditAnywhere, BlueprintReadWrite, Config, Category = Texture, Meta = (Bitmask, BitmaskEnum = EGLTFTextureType, EditCondition = "TextureCompression != EGLTFTextureCompression::PNG"))
+	UPROPERTY(EditAnywhere, BlueprintReadWrite, Config, Category = Texture, Meta = (Bitmask, BitmaskEnum = EGLTFTextureType, EditCondition = "TextureCompression == EGLTFTextureCompression::JPEG"))
 	int32 NoLossyCompressionFor;
 
 	/** If enabled, export UV tiling and un-mirroring settings in a texture coordinate expression node for simple material input expressions. Uses extension KHR_texture_transform. */
-	UPROPERTY(EditAnywhere, BlueprintReadWrite, Config, Category = Texture)
+	UPROPERTY(EditAnywhere, BlueprintReadWrite, Config, Category = Texture, Meta = (EditCondition = "TextureCompression != EGLTFTextureCompression::None"))
 	bool bExportTextureTransforms;
 
 	/** If enabled, export lightmaps (created by Lightmass) when exporting a level. Uses extension EPIC_lightmap_textures, which is supported by Unreal's glTF viewer. */
-	UPROPERTY(EditAnywhere, BlueprintReadWrite, Config, Category = Texture)
+	UPROPERTY(EditAnywhere, BlueprintReadWrite, Config, Category = Texture, Meta = (EditCondition = "TextureCompression != EGLTFTextureCompression::None"))
 	bool bExportLightmaps;
 
 	/** Encoding used to store textures that have pixel colors with more than 8-bit per channel. Uses extension EPIC_texture_hdr_encoding, which is supported by Unreal's glTF viewer. */
-	UPROPERTY(EditAnywhere, BlueprintReadWrite, Config, Category = Texture, Meta = (DisplayName = "Texture HDR Encoding"))
+	UPROPERTY(EditAnywhere, BlueprintReadWrite, Config, Category = Texture, Meta = (DisplayName = "Texture HDR Encoding", EditCondition = "TextureCompression != EGLTFTextureCompression::None"))
 	EGLTFTextureHDREncoding TextureHDREncoding;
 
 	/** If enabled, export components that are flagged as hidden in-game. */
