[GLTFExporter] refactor: rename some export options for textures

commit f0f673e44bd2fd6d98216d0d6c725dd29c3370e2
Date: Fri, 27 Nov 2020 02:57:30 +0100

    refactor: rename some export options for textures
    TextureCompression -> TextureImageFormat
    TextureCompressionQuality -> TextureImageQuality
    NoLossyCompressionFor -> NoLossyImageFormatFor
    (cherry picked from commit 927f9716608725bf6d6d301672ec2649b26c738f​)

#skipci

diff --git a/Engine/Plugins/Enterprise/GLTFExporter/MigratedFromGit/Source/GLTFExporter/Private/Builders/GLTFImageBuilder.cpp b/Engine/Plugins/Enterprise/GLTFExporter/MigratedFromGit/Source/GLTFExporter/Private/Builders/GLTFImageBuilder.cpp
--- a/Engine/Plugins/Enterprise/GLTFExporter/MigratedFromGit/Source/GLTFExporter/Private/Builders/GLTFImageBuilder.cpp
+++ b/Engine/Plugins/Enterprise/GLTFExporter/MigratedFromGit/Source/GLTFExporter/Private/Builders/GLTFImageBuilder.cpp
@@ -37,7 +37,7 @@ FGLTFJsonImageIndex FGLTFImageBuilder::AddImage(const FColor* Pixels, FIntPoint
 			break;
 
 		case EGLTFJsonMimeType::JPEG:
-			FGLTFImageUtility::CompressToJPEG(Pixels, Size, ExportOptions->TextureCompressionQuality, CompressedData);
+			FGLTFImageUtility::CompressToJPEG(Pixels, Size, ExportOptions->TextureImageQuality, CompressedData);
 			break;
 
 		default:
@@ -78,17 +78,17 @@ FGLTFJsonImageIndex FGLTFImageBuilder::AddImage(const void* CompressedData, int6
 
 EGLTFJsonMimeType FGLTFImageBuilder::GetImageFormat(const FColor* Pixels, FIntPoint Size, bool bIgnoreAlpha, EGLTFTextureType Type) const
 {
-	switch (ExportOptions->TextureCompression)
+	switch (ExportOptions->TextureImageFormat)
 	{
-		case EGLTFTextureCompression::None:
+		case EGLTFTextureImageFormat::None:
 			return EGLTFJsonMimeType::None;
 
-		case EGLTFTextureCompression::PNG:
+		case EGLTFTextureImageFormat::PNG:
 			return EGLTFJsonMimeType::PNG;
 
-		case EGLTFTextureCompression::JPEG:
+		case EGLTFTextureImageFormat::JPEG:
 			return
-				!EnumHasAllFlags(static_cast<EGLTFTextureType>(ExportOptions->NoLossyCompressionFor), Type) &&
+				!EnumHasAllFlags(static_cast<EGLTFTextureType>(ExportOptions->NoLossyImageFormatFor), Type) &&
 				(bIgnoreAlpha || FGLTFImageUtility::NoAlphaNeeded(Pixels, Size)) ?
 				EGLTFJsonMimeType::JPEG : EGLTFJsonMimeType::PNG;
 
diff --git a/Engine/Plugins/Enterprise/GLTFExporter/MigratedFromGit/Source/GLTFExporter/Private/Converters/GLTFTextureConverters.cpp b/Engine/Plugins/Enterprise/GLTFExporter/MigratedFromGit/Source/GLTFExporter/Private/Converters/GLTFTextureConverters.cpp
--- a/Engine/Plugins/Enterprise/GLTFExporter/MigratedFromGit/Source/GLTFExporter/Private/Converters/GLTFTextureConverters.cpp
+++ b/Engine/Plugins/Enterprise/GLTFExporter/MigratedFromGit/Source/GLTFExporter/Private/Converters/GLTFTextureConverters.cpp
@@ -6,7 +6,7 @@
 
 FGLTFJsonTextureIndex FGLTFTexture2DConverter::Convert(const UTexture2D* Texture2D)
 {
-	if (Builder.ExportOptions->TextureCompression == EGLTFTextureCompression::None)
+	if (Builder.ExportOptions->TextureImageFormat == EGLTFTextureImageFormat::None)
 	{
 		return FGLTFJsonTextureIndex(INDEX_NONE);
 	}
@@ -18,7 +18,7 @@ FGLTFJsonTextureIndex FGLTFTexture2DConverter::Convert(const UTexture2D* Texture
 
 FGLTFJsonTextureIndex FGLTFTextureCubeConverter::Convert(const UTextureCube* TextureCube, ECubeFace CubeFace)
 {
-	if (Builder.ExportOptions->TextureCompression == EGLTFTextureCompression::None)
+	if (Builder.ExportOptions->TextureImageFormat == EGLTFTextureImageFormat::None)
 	{
 		return FGLTFJsonTextureIndex(INDEX_NONE);
 	}
@@ -30,7 +30,7 @@ FGLTFJsonTextureIndex FGLTFTextureCubeConverter::Convert(const UTextureCube* Tex
 
 FGLTFJsonTextureIndex FGLTFTextureRenderTarget2DConverter::Convert(const UTextureRenderTarget2D* RenderTarget2D)
 {
-	if (Builder.ExportOptions->TextureCompression == EGLTFTextureCompression::None)
+	if (Builder.ExportOptions->TextureImageFormat == EGLTFTextureImageFormat::None)
 	{
 		return FGLTFJsonTextureIndex(INDEX_NONE);
 	}
@@ -42,7 +42,7 @@ FGLTFJsonTextureIndex FGLTFTextureRenderTarget2DConverter::Convert(const UTextur
 
 FGLTFJsonTextureIndex FGLTFTextureRenderTargetCubeConverter::Convert(const UTextureRenderTargetCube* RenderTargetCube, ECubeFace CubeFace)
 {
-	if (Builder.ExportOptions->TextureCompression == EGLTFTextureCompression::None)
+	if (Builder.ExportOptions->TextureImageFormat == EGLTFTextureImageFormat::None)
 	{
 		return FGLTFJsonTextureIndex(INDEX_NONE);
 	}
@@ -54,7 +54,7 @@ FGLTFJsonTextureIndex FGLTFTextureRenderTargetCubeConverter::Convert(const UText
 
 FGLTFJsonTextureIndex FGLTFTextureLightMapConverter::Convert(const ULightMapTexture2D* LightMap)
 {
-	if (Builder.ExportOptions->TextureCompression == EGLTFTextureCompression::None)
+	if (Builder.ExportOptions->TextureImageFormat == EGLTFTextureImageFormat::None)
 	{
 		return FGLTFJsonTextureIndex(INDEX_NONE);
 	}
diff --git a/Engine/Plugins/Enterprise/GLTFExporter/MigratedFromGit/Source/GLTFExporter/Private/GLTFExportOptions.cpp b/Engine/Plugins/Enterprise/GLTFExporter/MigratedFromGit/Source/GLTFExporter/Private/GLTFExportOptions.cpp
--- a/Engine/Plugins/Enterprise/GLTFExporter/MigratedFromGit/Source/GLTFExporter/Private/GLTFExportOptions.cpp
+++ b/Engine/Plugins/Enterprise/GLTFExporter/MigratedFromGit/Source/GLTFExporter/Private/GLTFExportOptions.cpp
@@ -30,9 +30,9 @@ void UGLTFExportOptions::ResetToDefault()
 	bExportAnimationSequences = true;
 	bRetargetBoneTransforms = true;
 	bExportPlaybackSettings = true;
-	TextureCompression = EGLTFTextureCompression::PNG;
-	TextureCompressionQuality = 0;
-	NoLossyCompressionFor = static_cast<int32>(EGLTFTextureType::All);
+	TextureImageFormat = EGLTFTextureImageFormat::PNG;
+	TextureImageQuality = 0;
+	NoLossyImageFormatFor = static_cast<int32>(EGLTFTextureType::All);
 	bExportTextureTransforms = true;
 	bExportLightmaps = true;
 	TextureHDREncoding = EGLTFTextureHDREncoding::RGBM;
diff --git a/Engine/Plugins/Enterprise/GLTFExporter/MigratedFromGit/Source/GLTFExporter/Private/Tasks/GLTFMaterialTasks.cpp b/Engine/Plugins/Enterprise/GLTFExporter/MigratedFromGit/Source/GLTFExporter/Private/Tasks/GLTFMaterialTasks.cpp
--- a/Engine/Plugins/Enterprise/GLTFExporter/MigratedFromGit/Source/GLTFExporter/Private/Tasks/GLTFMaterialTasks.cpp
+++ b/Engine/Plugins/Enterprise/GLTFExporter/MigratedFromGit/Source/GLTFExporter/Private/Tasks/GLTFMaterialTasks.cpp
@@ -363,7 +363,7 @@ bool FGLTFMaterialTask::TryGetBaseColorAndOpacity(FGLTFJsonPBRMetallicRoughness&
 		return true;
 	}
 
-	if (Builder.ExportOptions->TextureCompression == EGLTFTextureCompression::None)
+	if (Builder.ExportOptions->TextureImageFormat == EGLTFTextureImageFormat::None)
 	{
 		return true;
 	}
@@ -486,7 +486,7 @@ bool FGLTFMaterialTask::TryGetMetallicAndRoughness(FGLTFJsonPBRMetallicRoughness
 		return true;
 	}
 
-	if (Builder.ExportOptions->TextureCompression == EGLTFTextureCompression::None)
+	if (Builder.ExportOptions->TextureImageFormat == EGLTFTextureImageFormat::None)
 	{
 		return true;
 	}
@@ -613,7 +613,7 @@ bool FGLTFMaterialTask::TryGetClearCoatRoughness(FGLTFJsonClearCoatExtension& Ou
 		return true;
 	}
 
-	if (Builder.ExportOptions->TextureCompression == EGLTFTextureCompression::None)
+	if (Builder.ExportOptions->TextureImageFormat == EGLTFTextureImageFormat::None)
 	{
 		return true;
 	}
@@ -701,7 +701,7 @@ bool FGLTFMaterialTask::TryGetEmissive(FGLTFJsonMaterial& JsonMaterial, const FM
 	}
 	else
 	{
-		if (Builder.ExportOptions->TextureCompression == EGLTFTextureCompression::None)
+		if (Builder.ExportOptions->TextureImageFormat == EGLTFTextureImageFormat::None)
 		{
 			JsonMaterial.EmissiveTexture.Index = FGLTFJsonTextureIndex(INDEX_NONE);
 			return true;
@@ -994,7 +994,7 @@ bool FGLTFMaterialTask::TryGetSourceTexture(FGLTFJsonTextureInfo& OutTexInfo, co
 
 bool FGLTFMaterialTask::TryGetSourceTexture(const UTexture2D*& OutTexture, int32& OutTexCoord, FGLTFJsonTextureTransform& OutTransform, const FMaterialPropertyEx& Property, const TArray<FLinearColor>& AllowedMasks) const
 {
-	if (Builder.ExportOptions->TextureCompression == EGLTFTextureCompression::None)
+	if (Builder.ExportOptions->TextureImageFormat == EGLTFTextureImageFormat::None)
 	{
 		return false;
 	}
@@ -1127,7 +1127,7 @@ bool FGLTFMaterialTask::TryGetBakedMaterialProperty(FGLTFJsonTextureInfo& OutTex
 		return true;
 	}
 
-	if (Builder.ExportOptions->TextureCompression == EGLTFTextureCompression::None)
+	if (Builder.ExportOptions->TextureImageFormat == EGLTFTextureImageFormat::None)
 	{
 		OutTexInfo.Index = FGLTFJsonTextureIndex(INDEX_NONE);
 		return true;
@@ -1166,7 +1166,7 @@ bool FGLTFMaterialTask::TryGetBakedMaterialProperty(FGLTFJsonTextureInfo& OutTex
 		return true;
 	}
 
-	if (Builder.ExportOptions->TextureCompression == EGLTFTextureCompression::None)
+	if (Builder.ExportOptions->TextureImageFormat == EGLTFTextureImageFormat::None)
 	{
 		OutTexInfo.Index = FGLTFJsonTextureIndex(INDEX_NONE);
 		return true;
@@ -1205,7 +1205,7 @@ inline bool FGLTFMaterialTask::TryGetBakedMaterialProperty(FGLTFJsonTextureInfo&
 		return true;
 	}
 
-	if (Builder.ExportOptions->TextureCompression == EGLTFTextureCompression::None)
+	if (Builder.ExportOptions->TextureImageFormat == EGLTFTextureImageFormat::None)
 	{
 		OutTexInfo.Index = FGLTFJsonTextureIndex(INDEX_NONE);
 		return true;
@@ -1240,7 +1240,7 @@ bool FGLTFMaterialTask::TryGetBakedMaterialProperty(FGLTFJsonTextureInfo& OutTex
 
 	if (!PropertyBakeOutput.bIsConstant)
 	{
-		if (Builder.ExportOptions->TextureCompression == EGLTFTextureCompression::None)
+		if (Builder.ExportOptions->TextureImageFormat == EGLTFTextureImageFormat::None)
 		{
 			OutTexInfo.Index = FGLTFJsonTextureIndex(INDEX_NONE);
 			return true;
diff --git a/Engine/Plugins/Enterprise/GLTFExporter/MigratedFromGit/Source/GLTFExporter/Public/GLTFExportOptions.h b/Engine/Plugins/Enterprise/GLTFExporter/MigratedFromGit/Source/GLTFExporter/Public/GLTFExportOptions.h
--- a/Engine/Plugins/Enterprise/GLTFExporter/MigratedFromGit/Source/GLTFExporter/Public/GLTFExportOptions.h
+++ b/Engine/Plugins/Enterprise/GLTFExporter/MigratedFromGit/Source/GLTFExporter/Public/GLTFExportOptions.h
@@ -7,7 +7,7 @@
 #include "GLTFExportOptions.generated.h"
 
 UENUM(BlueprintType)
-enum class EGLTFTextureCompression : uint8
+enum class EGLTFTextureImageFormat : uint8
 {
 	/** Don't export any textures. */
 	None,
@@ -150,26 +150,26 @@ class GLTFEXPORTER_API UGLTFExportOptions : public UObject
 
 	/** Desired image format used for exported textures. */
 	UPROPERTY(EditAnywhere, BlueprintReadWrite, Config, Category = Texture)
-	EGLTFTextureCompression TextureCompression;
+	EGLTFTextureImageFormat TextureImageFormat;
 
 	/** Level of compression used for exported textures, between 1 (worst quality, best compression) and 100 (best quality, worst compression). Does not apply to lossless formats (e.g. PNG). */
-	UPROPERTY(EditAnywhere, BlueprintReadWrite, Config, Category = Texture, Meta = (ClampMin = "0", ClampMax = "100", EditCondition = "TextureCompression == EGLTFTextureCompression::JPEG"))
-	int32 TextureCompressionQuality;
+	UPROPERTY(EditAnywhere, BlueprintReadWrite, Config, Category = Texture, Meta = (ClampMin = "0", ClampMax = "100", EditCondition = "TextureImageFormat == EGLTFTextureImageFormat::JPEG"))
+	int32 TextureImageQuality;
 
 	/** Texture types that will always use lossless formats (e.g. PNG) because of sensitivity to compression artifacts. */
-	UPROPERTY(EditAnywhere, BlueprintReadWrite, Config, Category = Texture, Meta = (Bitmask, BitmaskEnum = EGLTFTextureType, EditCondition = "TextureCompression == EGLTFTextureCompression::JPEG"))
-	int32 NoLossyCompressionFor;
+	UPROPERTY(EditAnywhere, BlueprintReadWrite, Config, Category = Texture, Meta = (Bitmask, BitmaskEnum = EGLTFTextureType, EditCondition = "TextureImageFormat == EGLTFTextureImageFormat::JPEG"))
+	int32 NoLossyImageFormatFor;
 
 	/** If enabled, export UV tiling and un-mirroring settings in a texture coordinate expression node for simple material input expressions. Uses extension KHR_texture_transform. */
-	UPROPERTY(EditAnywhere, BlueprintReadWrite, Config, Category = Texture, Meta = (EditCondition = "TextureCompression != EGLTFTextureCompression::None"))
+	UPROPERTY(EditAnywhere, BlueprintReadWrite, Config, Category = Texture, Meta = (EditCondition = "TextureImageFormat != EGLTFTextureImageFormat::None"))
 	bool bExportTextureTransforms;
 
 	/** If enabled, export lightmaps (created by Lightmass) when exporting a level. Uses extension EPIC_lightmap_textures, which is supported by Unreal's glTF viewer. */
-	UPROPERTY(EditAnywhere, BlueprintReadWrite, Config, Category = Texture, Meta = (EditCondition = "TextureCompression != EGLTFTextureCompression::None"))
+	UPROPERTY(EditAnywhere, BlueprintReadWrite, Config, Category = Texture, Meta = (EditCondition = "TextureImageFormat != EGLTFTextureImageFormat::None"))
 	bool bExportLightmaps;
 
 	/** Encoding used to store textures that have pixel colors with more than 8-bit per channel. Uses extension EPIC_texture_hdr_encoding, which is supported by Unreal's glTF viewer. */
-	UPROPERTY(EditAnywhere, BlueprintReadWrite, Config, Category = Texture, Meta = (DisplayName = "Texture HDR Encoding", EditCondition = "TextureCompression != EGLTFTextureCompression::None"))
+	UPROPERTY(EditAnywhere, BlueprintReadWrite, Config, Category = Texture, Meta = (DisplayName = "Texture HDR Encoding", EditCondition = "TextureImageFormat != EGLTFTextureImageFormat::None"))
 	EGLTFTextureHDREncoding TextureHDREncoding;
 
 	/** If enabled, export components that are flagged as hidden in-game. */
