[GLTFExporter] feat: add warning and fall back to simple baking if UVs are degenerate

commit e18599259885c243e1fb7073cdb048e391ba76f9
Date: Wed, 25 Nov 2020 16:39:11 +0100

    feat: add warning and fall back to simple baking if UVs are degenerate
    (cherry picked from commit df4e94fae12be5038c22f82c423a67d085503f8f​)

#skipci

diff --git a/Engine/Plugins/Enterprise/GLTFExporter/MigratedFromGit/Source/GLTFExporter/Private/Converters/GLTFDegenerateUVSectionsChecker.cpp b/Engine/Plugins/Enterprise/GLTFExporter/MigratedFromGit/Source/GLTFExporter/Private/Converters/GLTFDegenerateUVSectionsChecker.cpp
new file mode 100644
--- /dev/null
+++ b/Engine/Plugins/Enterprise/GLTFExporter/MigratedFromGit/Source/GLTFExporter/Private/Converters/GLTFDegenerateUVSectionsChecker.cpp
@@ -0,0 +1,75 @@
+// Copyright Epic Games, Inc. All Rights Reserved.
+
+#include "Converters/GLTFDegenerateUVSectionsChecker.h"
+#include "StaticMeshAttributes.h"
+#include "MeshDescription.h"
+
+void FGLTFDegenerateUVSectionsChecker::Sanitize(const FMeshDescription*& Description, int32& TexCoord)
+{
+	if (Description != nullptr)
+	{
+		const TVertexInstanceAttributesConstRef<FVector2D> VertexInstanceUVs =
+			Description->VertexInstanceAttributes().GetAttributesRef<FVector2D>(MeshAttribute::VertexInstance::TextureCoordinate);
+		const int32 TexCoordCount = VertexInstanceUVs.GetNumIndices();
+
+		if (TexCoord < 0 || TexCoord >= TexCoordCount)
+		{
+			Description = nullptr;
+		}
+	}
+}
+
+const TArray<int32>* FGLTFDegenerateUVSectionsChecker::Convert(const FMeshDescription* Description, int32 TexCoord)
+{
+	if (Description == nullptr)
+	{
+		// TODO: report warning?
+
+		return nullptr;
+	}
+
+	TArray<uint32> DegenerateSections;
+
+	const TVertexInstanceAttributesConstRef<FVector2D> VertexInstanceUVs =
+		Description->VertexInstanceAttributes().GetAttributesRef<FVector2D>(MeshAttribute::VertexInstance::TextureCoordinate);
+
+	for (const FPolygonGroupID PolygonGroupID : Description->PolygonGroups().GetElementIDs())
+	{
+		bool bGroupIsValid = false;
+		uint32 Index = 0;
+		FVector2D ReferenceUV;
+
+		for (const FPolygonID PolygonID : Description->GetPolygonGroupPolygons(PolygonGroupID))
+		{
+			for (const FVertexInstanceID VertexInstanceID : Description->GetPolygonVertexInstances(PolygonID))
+			{
+				const FVector2D UV = VertexInstanceUVs.Get(VertexInstanceID, TexCoord);
+
+				if (Index == 0)
+				{
+					ReferenceUV = UV;
+				}
+				else if (!UV.Equals(ReferenceUV))
+				{
+					bGroupIsValid = true;
+					break;
+				}
+
+				++Index;
+			}
+
+			if (bGroupIsValid)
+			{
+				break;
+			}
+		}
+
+		if (!bGroupIsValid)
+		{
+			DegenerateSections.Add(PolygonGroupID.GetValue());
+		}
+	}
+
+	TUniquePtr<TArray<int32>> Output = MakeUnique<TArray<int32>>(DegenerateSections);
+	return Outputs.Add_GetRef(MoveTemp(Output)).Get();
+}
diff --git a/Engine/Plugins/Enterprise/GLTFExporter/MigratedFromGit/Source/GLTFExporter/Private/Converters/GLTFDegenerateUVSectionsChecker.h b/Engine/Plugins/Enterprise/GLTFExporter/MigratedFromGit/Source/GLTFExporter/Private/Converters/GLTFDegenerateUVSectionsChecker.h
new file mode 100644
--- /dev/null
+++ b/Engine/Plugins/Enterprise/GLTFExporter/MigratedFromGit/Source/GLTFExporter/Private/Converters/GLTFDegenerateUVSectionsChecker.h
@@ -0,0 +1,15 @@
+// Copyright Epic Games, Inc. All Rights Reserved.
+
+#pragma once
+
+#include "Converters/GLTFConverter.h"
+#include "Engine.h"
+
+class FGLTFDegenerateUVSectionsChecker final : public TGLTFConverter<const TArray<int32>*, const FMeshDescription*, int32>
+{
+	TArray<TUniquePtr<TArray<int32>>> Outputs;
+
+	virtual void Sanitize(const FMeshDescription*& Description, int32& TexCoord) override;
+
+	virtual const TArray<int32>* Convert(const FMeshDescription* Description, int32 TexCoord) override;
+};
diff --git a/Engine/Plugins/Enterprise/GLTFExporter/MigratedFromGit/Source/GLTFExporter/Private/Converters/GLTFMaterialConverters.cpp b/Engine/Plugins/Enterprise/GLTFExporter/MigratedFromGit/Source/GLTFExporter/Private/Converters/GLTFMaterialConverters.cpp
--- a/Engine/Plugins/Enterprise/GLTFExporter/MigratedFromGit/Source/GLTFExporter/Private/Converters/GLTFMaterialConverters.cpp
+++ b/Engine/Plugins/Enterprise/GLTFExporter/MigratedFromGit/Source/GLTFExporter/Private/Converters/GLTFMaterialConverters.cpp
@@ -14,6 +14,32 @@ void FGLTFMaterialConverter::Sanitize(const UMaterialInterface*& Material, const
 		MeshData = nullptr;
 		SectionIndices = {};
 	}
+
+	if (MeshData != nullptr)
+	{
+		const FGLTFMeshData* ParentMeshData = MeshData->GetParent();
+		const TArray<int32>* DegenerateSections = DegenerateUVSectionsChecker.GetOrAdd(&ParentMeshData->Description, MeshData->TexCoord);
+		if (DegenerateSections != nullptr)
+		{
+			bool bHasDegenerateUVs = false;
+			for (const int32 SectionIndex : SectionIndices)
+			{
+				bHasDegenerateUVs |= DegenerateSections->Contains(SectionIndex);
+			}
+
+			if (bHasDegenerateUVs)
+			{
+				Builder.AddWarningMessage(FString::Printf(
+					TEXT("Material %s is using mesh data from %s but the lightmap UVs (channel %d) are degenerate. Simple baking will be used as fallback"),
+					*Material->GetName(),
+					*ParentMeshData->Name,
+					MeshData->TexCoord));
+
+				MeshData = nullptr;
+				SectionIndices = {};
+			}
+		}
+	}
 }
 
 FGLTFJsonMaterialIndex FGLTFMaterialConverter::Convert(const UMaterialInterface* Material, const FGLTFMeshData* MeshData, TArray<int32> SectionIndices)
diff --git a/Engine/Plugins/Enterprise/GLTFExporter/MigratedFromGit/Source/GLTFExporter/Private/Converters/GLTFMaterialConverters.h b/Engine/Plugins/Enterprise/GLTFExporter/MigratedFromGit/Source/GLTFExporter/Private/Converters/GLTFMaterialConverters.h
--- a/Engine/Plugins/Enterprise/GLTFExporter/MigratedFromGit/Source/GLTFExporter/Private/Converters/GLTFMaterialConverters.h
+++ b/Engine/Plugins/Enterprise/GLTFExporter/MigratedFromGit/Source/GLTFExporter/Private/Converters/GLTFMaterialConverters.h
@@ -7,6 +7,7 @@
 #include "Converters/GLTFBuilderContext.h"
 #include "Converters/GLTFMeshData.h"
 #include "Converters/GLTFUVOverlapChecker.h"
+#include "Converters/GLTFDegenerateUVSectionsChecker.h"
 #include "Engine.h"
 
 class FGLTFMaterialConverter final : public FGLTFBuilderContext, public TGLTFConverter<FGLTFJsonMaterialIndex, const UMaterialInterface*, const FGLTFMeshData*, TArray<int32>>
@@ -18,4 +19,5 @@ class FGLTFMaterialConverter final : public FGLTFBuilderContext, public TGLTFCon
 	virtual FGLTFJsonMaterialIndex Convert(const UMaterialInterface* Material, const FGLTFMeshData* MeshData, TArray<int32> SectionIndices) override;
 
 	FGLTFUVOverlapChecker UVOverlapChecker;
+	FGLTFDegenerateUVSectionsChecker DegenerateUVSectionsChecker;
 };
