[GLTFExporter] refactor: simplify code flow in FGLTFDegenerateUVSectionsChecker by reducing it to one section at a time

commit 7803f09f5071e60ecf8dc43f90783df568e00b28
Date: Thu, 26 Nov 2020 22:17:59 +0100

    refactor: simplify code flow in FGLTFDegenerateUVSectionsChecker by reducing it to one section at a time
    (cherry picked from commit 24edf4ff437ef89c846f9244bccaa7d50d998d88​)

#skipci

diff --git a/Engine/Plugins/Enterprise/GLTFExporter/MigratedFromGit/Source/GLTFExporter/Private/Converters/GLTFDegenerateUVSectionsChecker.cpp b/Engine/Plugins/Enterprise/GLTFExporter/MigratedFromGit/Source/GLTFExporter/Private/Converters/GLTFDegenerateUVSectionsChecker.cpp
--- a/Engine/Plugins/Enterprise/GLTFExporter/MigratedFromGit/Source/GLTFExporter/Private/Converters/GLTFDegenerateUVSectionsChecker.cpp
+++ b/Engine/Plugins/Enterprise/GLTFExporter/MigratedFromGit/Source/GLTFExporter/Private/Converters/GLTFDegenerateUVSectionsChecker.cpp
@@ -4,14 +4,20 @@
 #include "StaticMeshAttributes.h"
 #include "MeshDescription.h"
 
-void FGLTFDegenerateUVSectionsChecker::Sanitize(const FMeshDescription*& Description, int32& TexCoord)
+void FGLTFDegenerateUVSectionsChecker::Sanitize(const FMeshDescription*& Description, int32& SectionIndex, int32& TexCoord)
 {
 	if (Description != nullptr)
 	{
-		const TVertexInstanceAttributesConstRef<FVector2D> VertexInstanceUVs =
-			Description->VertexInstanceAttributes().GetAttributesRef<FVector2D>(MeshAttribute::VertexInstance::TextureCoordinate);
-		const int32 TexCoordCount = VertexInstanceUVs.GetNumIndices();
+		const int32 SectionCount = Description->PolygonGroups().Num();
+		if (SectionIndex < 0 || SectionIndex >= SectionCount)
+		{
+			Description = nullptr;
+		}
+	}
 
+	if (Description != nullptr)
+	{
+		const int32 TexCoordCount = Description->VertexInstanceAttributes().GetAttributesRef<FVector2D>(MeshAttribute::VertexInstance::TextureCoordinate).GetNumIndices();
 		if (TexCoord < 0 || TexCoord >= TexCoordCount)
 		{
 			Description = nullptr;
@@ -19,57 +25,40 @@ void FGLTFDegenerateUVSectionsChecker::Sanitize(const FMeshDescription*& Descrip
 	}
 }
 
-const TArray<int32>* FGLTFDegenerateUVSectionsChecker::Convert(const FMeshDescription* Description, int32 TexCoord)
+bool FGLTFDegenerateUVSectionsChecker::Convert(const FMeshDescription* Description, int32 SectionIndex, int32 TexCoord)
 {
 	if (Description == nullptr)
 	{
 		// TODO: report warning?
 
-		return nullptr;
+		return false;
 	}
 
-	TArray<uint32> DegenerateSections;
-
+	const FPolygonGroupID PolygonGroupID = FPolygonGroupID(SectionIndex);
 	const TVertexInstanceAttributesConstRef<FVector2D> VertexInstanceUVs =
 		Description->VertexInstanceAttributes().GetAttributesRef<FVector2D>(MeshAttribute::VertexInstance::TextureCoordinate);
 
-	for (const FPolygonGroupID PolygonGroupID : Description->PolygonGroups().GetElementIDs())
-	{
-		bool bGroupIsValid = false;
-		uint32 Index = 0;
-		FVector2D ReferenceUV;
+	uint32 Index = 0;
+	FVector2D ReferenceUV;
 
-		for (const FPolygonID PolygonID : Description->GetPolygonGroupPolygons(PolygonGroupID))
+	for (const FPolygonID PolygonID : Description->GetPolygonGroupPolygons(PolygonGroupID))
+	{
+		for (const FVertexInstanceID VertexInstanceID : Description->GetPolygonVertexInstances(PolygonID))
 		{
-			for (const FVertexInstanceID VertexInstanceID : Description->GetPolygonVertexInstances(PolygonID))
-			{
-				const FVector2D UV = VertexInstanceUVs.Get(VertexInstanceID, TexCoord);
+			const FVector2D UV = VertexInstanceUVs.Get(VertexInstanceID, TexCoord);
 
-				if (Index == 0)
-				{
-					ReferenceUV = UV;
-				}
-				else if (!UV.Equals(ReferenceUV))
-				{
-					bGroupIsValid = true;
-					break;
-				}
-
-				++Index;
+			if (Index == 0)
+			{
+				ReferenceUV = UV;
 			}
-
-			if (bGroupIsValid)
+			else if (!UV.Equals(ReferenceUV))
 			{
-				break;
+				return false;
 			}
-		}
 
-		if (!bGroupIsValid)
-		{
-			DegenerateSections.Add(PolygonGroupID.GetValue());
+			++Index;
 		}
 	}
 
-	TUniquePtr<TArray<int32>> Output = MakeUnique<TArray<int32>>(DegenerateSections);
-	return Outputs.Add_GetRef(MoveTemp(Output)).Get();
+	return true;
 }
diff --git a/Engine/Plugins/Enterprise/GLTFExporter/MigratedFromGit/Source/GLTFExporter/Private/Converters/GLTFDegenerateUVSectionsChecker.h b/Engine/Plugins/Enterprise/GLTFExporter/MigratedFromGit/Source/GLTFExporter/Private/Converters/GLTFDegenerateUVSectionsChecker.h
--- a/Engine/Plugins/Enterprise/GLTFExporter/MigratedFromGit/Source/GLTFExporter/Private/Converters/GLTFDegenerateUVSectionsChecker.h
+++ b/Engine/Plugins/Enterprise/GLTFExporter/MigratedFromGit/Source/GLTFExporter/Private/Converters/GLTFDegenerateUVSectionsChecker.h
@@ -5,11 +5,9 @@
 #include "Converters/GLTFConverter.h"
 #include "Engine.h"
 
-class FGLTFDegenerateUVSectionsChecker final : public TGLTFConverter<const TArray<int32>*, const FMeshDescription*, int32>
+class FGLTFDegenerateUVSectionsChecker final : public TGLTFConverter<bool, const FMeshDescription*, int32, int32>
 {
-	TArray<TUniquePtr<TArray<int32>>> Outputs;
+	virtual void Sanitize(const FMeshDescription*& Description, int32& SectionIndex, int32& TexCoord) override;
 
-	virtual void Sanitize(const FMeshDescription*& Description, int32& TexCoord) override;
-
-	virtual const TArray<int32>* Convert(const FMeshDescription* Description, int32 TexCoord) override;
+	virtual bool Convert(const FMeshDescription* Description, int32 SectionIndex, int32 TexCoord) override;
 };
diff --git a/Engine/Plugins/Enterprise/GLTFExporter/MigratedFromGit/Source/GLTFExporter/Private/Converters/GLTFMaterialConverters.cpp b/Engine/Plugins/Enterprise/GLTFExporter/MigratedFromGit/Source/GLTFExporter/Private/Converters/GLTFMaterialConverters.cpp
--- a/Engine/Plugins/Enterprise/GLTFExporter/MigratedFromGit/Source/GLTFExporter/Private/Converters/GLTFMaterialConverters.cpp
+++ b/Engine/Plugins/Enterprise/GLTFExporter/MigratedFromGit/Source/GLTFExporter/Private/Converters/GLTFMaterialConverters.cpp
@@ -17,26 +17,21 @@ void FGLTFMaterialConverter::Sanitize(const UMaterialInterface*& Material, const
 
 	if (MeshData != nullptr)
 	{
-		const FGLTFMeshData* ParentMeshData = MeshData->GetParent();
-		const TArray<int32>* DegenerateSections = DegenerateUVSectionsChecker.GetOrAdd(&ParentMeshData->Description, MeshData->TexCoord);
-		if (DegenerateSections != nullptr)
+		const FMeshDescription& MeshDescription = MeshData->GetParent()->Description;
+		for (const int32 SectionIndex : SectionIndices)
 		{
-			bool bHasDegenerateUVs = false;
-			for (const int32 SectionIndex : SectionIndices)
-			{
-				bHasDegenerateUVs |= DegenerateSections->Contains(SectionIndex);
-			}
-
+			const bool bHasDegenerateUVs = DegenerateUVSectionsChecker.GetOrAdd(&MeshDescription, SectionIndex, MeshData->TexCoord);
 			if (bHasDegenerateUVs)
 			{
 				Builder.AddWarningMessage(FString::Printf(
 					TEXT("Material %s is using mesh data from %s but the lightmap UVs (channel %d) are degenerate. Simple baking will be used as fallback"),
 					*Material->GetName(),
-					*ParentMeshData->Name,
+					*MeshData->GetParent()->Name,
 					MeshData->TexCoord));
 
 				MeshData = nullptr;
 				SectionIndices = {};
+				break;
 			}
 		}
 	}
