[GLTFExporter] refactor: GLTFDegenerateUVSectionsChecker to GLTFUVDegenerateChecker and add mesh section info to warning message

commit 60ee3fcd9e8805ca52aea22f385e9de710f8f919
Date: Fri, 27 Nov 2020 00:00:50 +0100

    refactor: GLTFDegenerateUVSectionsChecker to GLTFUVDegenerateChecker and add mesh section info to warning message
    (cherry picked from commit 48f1c17e3f24f9c6ecb1832a8bd9025136e1417c​)

#skipci

diff --git a/Engine/Plugins/Enterprise/GLTFExporter/MigratedFromGit/Source/GLTFExporter/Private/Converters/GLTFDegenerateUVSectionsChecker.cpp b/Engine/Plugins/Enterprise/GLTFExporter/MigratedFromGit/Source/GLTFExporter/Private/Converters/GLTFDegenerateUVSectionsChecker.cpp
deleted file mode 100644
diff --git a/Engine/Plugins/Enterprise/GLTFExporter/MigratedFromGit/Source/GLTFExporter/Private/Converters/GLTFDegenerateUVSectionsChecker.h b/Engine/Plugins/Enterprise/GLTFExporter/MigratedFromGit/Source/GLTFExporter/Private/Converters/GLTFDegenerateUVSectionsChecker.h
deleted file mode 100644
diff --git a/Engine/Plugins/Enterprise/GLTFExporter/MigratedFromGit/Source/GLTFExporter/Private/Converters/GLTFMaterialConverters.cpp b/Engine/Plugins/Enterprise/GLTFExporter/MigratedFromGit/Source/GLTFExporter/Private/Converters/GLTFMaterialConverters.cpp
--- a/Engine/Plugins/Enterprise/GLTFExporter/MigratedFromGit/Source/GLTFExporter/Private/Converters/GLTFMaterialConverters.cpp
+++ b/Engine/Plugins/Enterprise/GLTFExporter/MigratedFromGit/Source/GLTFExporter/Private/Converters/GLTFMaterialConverters.cpp
@@ -20,14 +20,15 @@ void FGLTFMaterialConverter::Sanitize(const UMaterialInterface*& Material, const
 		const FMeshDescription& MeshDescription = MeshData->GetParent()->Description;
 		for (const int32 SectionIndex : SectionIndices)
 		{
-			const bool bHasDegenerateUVs = DegenerateUVSectionsChecker.GetOrAdd(&MeshDescription, SectionIndex, MeshData->TexCoord);
+			const bool bHasDegenerateUVs = UVDegenerateChecker.GetOrAdd(&MeshDescription, SectionIndex, MeshData->TexCoord);
 			if (bHasDegenerateUVs)
 			{
 				Builder.AddWarningMessage(FString::Printf(
-					TEXT("Material %s is using mesh data from %s but the lightmap UVs (channel %d) are degenerate. Simple baking will be used as fallback"),
+					TEXT("Material %s is using mesh data from %s but the lightmap UVs (channel %d) are degenerate (in mesh section %d). Simple baking will be used as fallback"),
 					*Material->GetName(),
 					*MeshData->GetParent()->Name,
-					MeshData->TexCoord));
+					MeshData->TexCoord,
+					SectionIndex));
 
 				MeshData = nullptr;
 				SectionIndices = {};
diff --git a/Engine/Plugins/Enterprise/GLTFExporter/MigratedFromGit/Source/GLTFExporter/Private/Converters/GLTFMaterialConverters.h b/Engine/Plugins/Enterprise/GLTFExporter/MigratedFromGit/Source/GLTFExporter/Private/Converters/GLTFMaterialConverters.h
--- a/Engine/Plugins/Enterprise/GLTFExporter/MigratedFromGit/Source/GLTFExporter/Private/Converters/GLTFMaterialConverters.h
+++ b/Engine/Plugins/Enterprise/GLTFExporter/MigratedFromGit/Source/GLTFExporter/Private/Converters/GLTFMaterialConverters.h
@@ -7,7 +7,7 @@
 #include "Converters/GLTFBuilderContext.h"
 #include "Converters/GLTFMeshData.h"
 #include "Converters/GLTFUVOverlapChecker.h"
-#include "Converters/GLTFDegenerateUVSectionsChecker.h"
+#include "Converters/GLTFUVDegenerateChecker.h"
 #include "Engine.h"
 
 class FGLTFMaterialConverter final : public FGLTFBuilderContext, public TGLTFConverter<FGLTFJsonMaterialIndex, const UMaterialInterface*, const FGLTFMeshData*, TArray<int32>>
@@ -19,5 +19,5 @@ class FGLTFMaterialConverter final : public FGLTFBuilderContext, public TGLTFCon
 	virtual FGLTFJsonMaterialIndex Convert(const UMaterialInterface* Material, const FGLTFMeshData* MeshData, TArray<int32> SectionIndices) override;
 
 	FGLTFUVOverlapChecker UVOverlapChecker;
-	FGLTFDegenerateUVSectionsChecker DegenerateUVSectionsChecker;
+	FGLTFUVDegenerateChecker UVDegenerateChecker;
 };
diff --git a/Engine/Plugins/Enterprise/GLTFExporter/MigratedFromGit/Source/GLTFExporter/Private/Converters/GLTFUVDegenerateChecker.cpp b/Engine/Plugins/Enterprise/GLTFExporter/MigratedFromGit/Source/GLTFExporter/Private/Converters/GLTFUVDegenerateChecker.cpp
new file mode 100644
--- /dev/null
+++ b/Engine/Plugins/Enterprise/GLTFExporter/MigratedFromGit/Source/GLTFExporter/Private/Converters/GLTFUVDegenerateChecker.cpp
@@ -0,0 +1,66 @@
+// Copyright Epic Games, Inc. All Rights Reserved.
+
+#include "Converters/GLTFUVDegenerateChecker.h"
+#include "StaticMeshAttributes.h"
+#include "MeshDescription.h"
+
+void FGLTFUVDegenerateChecker::Sanitize(const FMeshDescription*& Description, int32& SectionIndex, int32& TexCoord)
+{
+	if (Description != nullptr)
+	{
+		const int32 SectionCount = Description->PolygonGroups().Num();
+		if (SectionIndex < 0 || SectionIndex >= SectionCount)
+		{
+			Description = nullptr;
+		}
+	}
+
+	if (Description != nullptr)
+	{
+		const int32 TexCoordCount = Description->VertexInstanceAttributes().GetAttributesRef<FVector2D>(MeshAttribute::VertexInstance::TextureCoordinate).GetNumIndices();
+		if (TexCoord < 0 || TexCoord >= TexCoordCount)
+		{
+			Description = nullptr;
+		}
+	}
+}
+
+bool FGLTFUVDegenerateChecker::Convert(const FMeshDescription* Description, int32 SectionIndex, int32 TexCoord)
+{
+	if (Description == nullptr)
+	{
+		// TODO: report warning?
+
+		return false;
+	}
+
+	// TODO: since we are not really checking for degenerate UVs, but simply if all UVs have the same value or not, we should call this class something else
+
+	const FPolygonGroupID PolygonGroupID = FPolygonGroupID(SectionIndex);
+	const TVertexInstanceAttributesConstRef<FVector2D> VertexInstanceUVs =
+		Description->VertexInstanceAttributes().GetAttributesRef<FVector2D>(MeshAttribute::VertexInstance::TextureCoordinate);
+
+	uint32 Index = 0;
+	FVector2D ReferenceUV;
+
+	for (const FPolygonID PolygonID : Description->GetPolygonGroupPolygons(PolygonGroupID))
+	{
+		for (const FVertexInstanceID VertexInstanceID : Description->GetPolygonVertexInstances(PolygonID))
+		{
+			const FVector2D UV = VertexInstanceUVs.Get(VertexInstanceID, TexCoord);
+
+			if (Index == 0)
+			{
+				ReferenceUV = UV;
+			}
+			else if (!UV.Equals(ReferenceUV))
+			{
+				return false;
+			}
+
+			++Index;
+		}
+	}
+
+	return true;
+}
diff --git a/Engine/Plugins/Enterprise/GLTFExporter/MigratedFromGit/Source/GLTFExporter/Private/Converters/GLTFUVDegenerateChecker.h b/Engine/Plugins/Enterprise/GLTFExporter/MigratedFromGit/Source/GLTFExporter/Private/Converters/GLTFUVDegenerateChecker.h
new file mode 100644
--- /dev/null
+++ b/Engine/Plugins/Enterprise/GLTFExporter/MigratedFromGit/Source/GLTFExporter/Private/Converters/GLTFUVDegenerateChecker.h
@@ -0,0 +1,13 @@
+// Copyright Epic Games, Inc. All Rights Reserved.
+
+#pragma once
+
+#include "Converters/GLTFConverter.h"
+#include "Engine.h"
+
+class FGLTFUVDegenerateChecker final : public TGLTFConverter<bool, const FMeshDescription*, int32, int32>
+{
+	virtual void Sanitize(const FMeshDescription*& Description, int32& SectionIndex, int32& TexCoord) override;
+
+	virtual bool Convert(const FMeshDescription* Description, int32 SectionIndex, int32 TexCoord) override;
+};
