[GLTFExporter] feat: add explicit exporters for LevelSequence and LevelVariantSets

commit e49cac4d3014c0e52ff0066e4a5618cc133c4ffd
Date: Sun, 29 Nov 2020 21:04:41 +0100

    feat: add explicit exporters for LevelSequence and LevelVariantSets
    (cherry picked from commit bce7408bcefc5791106108cbad39cb9e49c669a0​)

#skipci

diff --git a/Engine/Plugins/Enterprise/GLTFExporter/MigratedFromGit/Source/GLTFExporter/Private/Builders/GLTFConvertBuilder.cpp b/Engine/Plugins/Enterprise/GLTFExporter/MigratedFromGit/Source/GLTFExporter/Private/Builders/GLTFConvertBuilder.cpp
--- a/Engine/Plugins/Enterprise/GLTFExporter/MigratedFromGit/Source/GLTFExporter/Private/Builders/GLTFConvertBuilder.cpp
+++ b/Engine/Plugins/Enterprise/GLTFExporter/MigratedFromGit/Source/GLTFExporter/Private/Builders/GLTFConvertBuilder.cpp
@@ -387,14 +387,14 @@ FGLTFJsonBackdropIndex FGLTFConvertBuilder::GetOrAddBackdrop(const AActor* Backd
 	return BackdropConverter.GetOrAdd(BackdropActor);
 }
 
-FGLTFJsonVariationIndex FGLTFConvertBuilder::GetOrAddVariation(const ALevelVariantSetsActor* LevelVariantSetsActor)
+FGLTFJsonVariationIndex FGLTFConvertBuilder::GetOrAddVariation(const ULevelVariantSets* LevelVariantSets)
 {
-	if (LevelVariantSetsActor == nullptr)
+	if (LevelVariantSets == nullptr)
 	{
 		return FGLTFJsonVariationIndex(INDEX_NONE);
 	}
 
-	return VariationConverter.GetOrAdd(LevelVariantSetsActor);
+	return VariationConverter.GetOrAdd(LevelVariantSets);
 }
 
 FGLTFJsonLightMapIndex FGLTFConvertBuilder::GetOrAddLightMap(const UStaticMeshComponent* StaticMeshComponent)
diff --git a/Engine/Plugins/Enterprise/GLTFExporter/MigratedFromGit/Source/GLTFExporter/Private/Builders/GLTFConvertBuilder.h b/Engine/Plugins/Enterprise/GLTFExporter/MigratedFromGit/Source/GLTFExporter/Private/Builders/GLTFConvertBuilder.h
--- a/Engine/Plugins/Enterprise/GLTFExporter/MigratedFromGit/Source/GLTFExporter/Private/Builders/GLTFConvertBuilder.h
+++ b/Engine/Plugins/Enterprise/GLTFExporter/MigratedFromGit/Source/GLTFExporter/Private/Builders/GLTFConvertBuilder.h
@@ -77,7 +77,7 @@ public:
 	FGLTFJsonCameraIndex GetOrAddCamera(const UCameraComponent* CameraComponent);
 	FGLTFJsonLightIndex GetOrAddLight(const ULightComponent* LightComponent);
 	FGLTFJsonBackdropIndex GetOrAddBackdrop(const AActor* BackdropActor);
-	FGLTFJsonVariationIndex GetOrAddVariation(const ALevelVariantSetsActor* LevelVariantSetsActor);
+	FGLTFJsonVariationIndex GetOrAddVariation(const ULevelVariantSets* LevelVariantSets);
 	FGLTFJsonLightMapIndex GetOrAddLightMap(const UStaticMeshComponent* StaticMeshComponent);
 	FGLTFJsonHotspotIndex GetOrAddHotspot(const AGLTFHotspotActor* HotspotActor);
 	FGLTFJsonSkySphereIndex GetOrAddSkySphere(const AActor* SkySphereActor);
diff --git a/Engine/Plugins/Enterprise/GLTFExporter/MigratedFromGit/Source/GLTFExporter/Private/Converters/GLTFSceneConverters.cpp b/Engine/Plugins/Enterprise/GLTFExporter/MigratedFromGit/Source/GLTFExporter/Private/Converters/GLTFSceneConverters.cpp
--- a/Engine/Plugins/Enterprise/GLTFExporter/MigratedFromGit/Source/GLTFExporter/Private/Converters/GLTFSceneConverters.cpp
+++ b/Engine/Plugins/Enterprise/GLTFExporter/MigratedFromGit/Source/GLTFExporter/Private/Converters/GLTFSceneConverters.cpp
@@ -4,6 +4,7 @@
 #include "Builders/GLTFContainerBuilder.h"
 #include "Converters/GLTFConverterUtility.h"
 #include "Converters/GLTFActorUtility.h"
+#include "LevelVariantSetsActor.h"
 
 FGLTFJsonSceneIndex FGLTFSceneConverter::Convert(const ULevel* Level)
 {
@@ -27,10 +28,13 @@ FGLTFJsonSceneIndex FGLTFSceneConverter::Convert(const ULevel* Level)
 		{
 			if (Builder.ExportOptions->bExportVariantSets)
 			{
-				const FGLTFJsonVariationIndex VariationIndex = Builder.GetOrAddVariation(LevelVariantSetsActor);
-				if (VariationIndex != INDEX_NONE)
+				if (const ULevelVariantSets* LevelVariantSets = const_cast<ALevelVariantSetsActor*>(LevelVariantSetsActor)->GetLevelVariantSets(true))
 				{
-					Scene.Variations.Add(VariationIndex);
+					const FGLTFJsonVariationIndex VariationIndex = Builder.GetOrAddVariation(LevelVariantSets);
+					if (VariationIndex != INDEX_NONE)
+					{
+						Scene.Variations.Add(VariationIndex);
+					}
 				}
 			}
 		}
diff --git a/Engine/Plugins/Enterprise/GLTFExporter/MigratedFromGit/Source/GLTFExporter/Private/Converters/GLTFVarationConverters.cpp b/Engine/Plugins/Enterprise/GLTFExporter/MigratedFromGit/Source/GLTFExporter/Private/Converters/GLTFVarationConverters.cpp
--- a/Engine/Plugins/Enterprise/GLTFExporter/MigratedFromGit/Source/GLTFExporter/Private/Converters/GLTFVarationConverters.cpp
+++ b/Engine/Plugins/Enterprise/GLTFExporter/MigratedFromGit/Source/GLTFExporter/Private/Converters/GLTFVarationConverters.cpp
@@ -4,7 +4,6 @@
 #include "Converters/GLTFMeshUtility.h"
 #include "Builders/GLTFContainerBuilder.h"
 #include "Rendering/SkeletalMeshRenderData.h"
-#include "LevelVariantSetsActor.h"
 #include "VariantObjectBinding.h"
 #include "PropertyValueMaterial.h"
 #include "LevelVariantSets.h"
@@ -25,14 +24,8 @@ namespace
 	};
 }
 
-FGLTFJsonVariationIndex FGLTFVariationConverter::Convert(const ALevelVariantSetsActor* LevelVariantSetsActor)
+FGLTFJsonVariationIndex FGLTFVariationConverter::Convert(const ULevelVariantSets* LevelVariantSets)
 {
-	const ULevelVariantSets* LevelVariantSets = const_cast<ALevelVariantSetsActor*>(LevelVariantSetsActor)->GetLevelVariantSets(true);
-	if (LevelVariantSets == nullptr)
-	{
-		return FGLTFJsonVariationIndex(INDEX_NONE);
-	}
-
 	FGLTFJsonVariation JsonVariation;
 	LevelVariantSets->GetName(JsonVariation.Name);
 
diff --git a/Engine/Plugins/Enterprise/GLTFExporter/MigratedFromGit/Source/GLTFExporter/Private/Converters/GLTFVarationConverters.h b/Engine/Plugins/Enterprise/GLTFExporter/MigratedFromGit/Source/GLTFExporter/Private/Converters/GLTFVarationConverters.h
--- a/Engine/Plugins/Enterprise/GLTFExporter/MigratedFromGit/Source/GLTFExporter/Private/Converters/GLTFVarationConverters.h
+++ b/Engine/Plugins/Enterprise/GLTFExporter/MigratedFromGit/Source/GLTFExporter/Private/Converters/GLTFVarationConverters.h
@@ -7,16 +7,15 @@
 #include "Converters/GLTFConverter.h"
 #include "Converters/GLTFBuilderContext.h"
 #include "Engine.h"
-#include "LevelVariantSetsActor.h"
 #include "Variant.h"
 #include "PropertyValue.h"
 #include "LevelVariantSets.h"
 
-class FGLTFVariationConverter final : public FGLTFBuilderContext, public TGLTFConverter<FGLTFJsonVariationIndex, const ALevelVariantSetsActor*>
+class FGLTFVariationConverter final : public FGLTFBuilderContext, public TGLTFConverter<FGLTFJsonVariationIndex, const ULevelVariantSets*>
 {
 	using FGLTFBuilderContext::FGLTFBuilderContext;
 
-	virtual FGLTFJsonVariationIndex Convert(const ALevelVariantSetsActor* LevelVariantSetsActor) override;
+	virtual FGLTFJsonVariationIndex Convert(const ULevelVariantSets* LevelVariantSets) override;
 
 	bool TryParseVariant(FGLTFJsonVariant& OutVariant, const UVariant* Variant) const;
 	bool TryParseVariantBinding(FGLTFJsonVariant& OutVariant, const UVariantObjectBinding* Binding) const;
diff --git a/Engine/Plugins/Enterprise/GLTFExporter/MigratedFromGit/Source/GLTFExporter/Private/Exporters/GLTFAnimSequenceExporter.cpp b/Engine/Plugins/Enterprise/GLTFExporter/MigratedFromGit/Source/GLTFExporter/Private/Exporters/GLTFAnimSequenceExporter.cpp
--- a/Engine/Plugins/Enterprise/GLTFExporter/MigratedFromGit/Source/GLTFExporter/Private/Exporters/GLTFAnimSequenceExporter.cpp
+++ b/Engine/Plugins/Enterprise/GLTFExporter/MigratedFromGit/Source/GLTFExporter/Private/Exporters/GLTFAnimSequenceExporter.cpp
@@ -15,6 +15,22 @@ bool UGLTFAnimSequenceExporter::AddObject(FGLTFContainerBuilder& Builder, const
 {
 	const UAnimSequence* AnimSequence = CastChecked<UAnimSequence>(Object);
 
+	if (!Builder.ExportOptions->bExportAnimationSequences)
+	{
+		Builder.AddErrorMessage(
+			FString::Printf(TEXT("Failed to export animation sequence %s because animation sequences are disabled by export options"),
+			*AnimSequence->GetName()));
+		return false;
+	}
+
+	if (!Builder.ExportOptions->bExportVertexSkinWeights)
+	{
+		Builder.AddErrorMessage(
+			FString::Printf(TEXT("Failed to export animation sequence %s because vertex skin weights are disabled by export options"),
+			*AnimSequence->GetName()));
+		return false;
+	}
+
 	const USkeletalMesh* SkeletalMesh = FGLTFExporterUtility::GetPreviewMesh(AnimSequence);
 	if (SkeletalMesh == nullptr)
 	{
@@ -43,43 +59,25 @@ bool UGLTFAnimSequenceExporter::AddObject(FGLTFContainerBuilder& Builder, const
 
 	const FGLTFJsonNodeIndex NodeIndex = Builder.AddNode(Node);
 
-	if (Builder.ExportOptions->bExportVertexSkinWeights)
+	const FGLTFJsonSkinIndex SkinIndex = Builder.GetOrAddSkin(NodeIndex, SkeletalMesh);
+	if (SkinIndex == INDEX_NONE)
 	{
-		const FGLTFJsonSkinIndex SkinIndex = Builder.GetOrAddSkin(NodeIndex, SkeletalMesh);
-		if (SkinIndex == INDEX_NONE)
-		{
-			Builder.AddErrorMessage(
-				FString::Printf(TEXT("Failed to export bones in skeletal mesh %s for animation sequence %s"),
-				*SkeletalMesh->GetName(),
-				*AnimSequence->GetName()));
-			return false;
-		}
+		Builder.AddErrorMessage(
+			FString::Printf(TEXT("Failed to export bones in skeletal mesh %s for animation sequence %s"),
+			*SkeletalMesh->GetName(),
+			*AnimSequence->GetName()));
+		return false;
+	}
 
-		Builder.GetNode(NodeIndex).Skin = SkinIndex;
+	Builder.GetNode(NodeIndex).Skin = SkinIndex;
 
-		if (Builder.ExportOptions->bExportAnimationSequences)
-		{
-			const FGLTFJsonAnimationIndex AnimationIndex =  Builder.GetOrAddAnimation(NodeIndex, SkeletalMesh, AnimSequence);
-			if (AnimationIndex == INDEX_NONE)
-			{
-				Builder.AddErrorMessage(
-					FString::Printf(TEXT("Failed to export animation sequence %s"),
-					*AnimSequence->GetName()));
-				return false;
-			}
-		}
-		else
-		{
-			Builder.AddWarningMessage(
-				FString::Printf(TEXT("Failed to export animation sequence %s because animation sequences are disabled by export options"),
-				*AnimSequence->GetName()));
-		}
-	}
-	else
+	const FGLTFJsonAnimationIndex AnimationIndex =  Builder.GetOrAddAnimation(NodeIndex, SkeletalMesh, AnimSequence);
+	if (AnimationIndex == INDEX_NONE)
 	{
-		Builder.AddWarningMessage(
-			FString::Printf(TEXT("Failed to export animation sequence %s because vertex skin weights are disabled by export options"),
+		Builder.AddErrorMessage(
+			FString::Printf(TEXT("Failed to export animation sequence %s"),
 			*AnimSequence->GetName()));
+		return false;
 	}
 
 	FGLTFJsonScene Scene;
diff --git a/Engine/Plugins/Enterprise/GLTFExporter/MigratedFromGit/Source/GLTFExporter/Private/Exporters/GLTFExporterUtility.cpp b/Engine/Plugins/Enterprise/GLTFExporter/MigratedFromGit/Source/GLTFExporter/Private/Exporters/GLTFExporterUtility.cpp
--- a/Engine/Plugins/Enterprise/GLTFExporter/MigratedFromGit/Source/GLTFExporter/Private/Exporters/GLTFExporterUtility.cpp
+++ b/Engine/Plugins/Enterprise/GLTFExporter/MigratedFromGit/Source/GLTFExporter/Private/Exporters/GLTFExporterUtility.cpp
@@ -1,6 +1,8 @@
 // Copyright Epic Games, Inc. All Rights Reserved.
 
 #include "Exporters/GLTFExporterUtility.h"
+#include "AssetRegistryModule.h"
+#include "LevelSequence.h"
 #include "UnrealEd.h"
 
 const UStaticMesh* FGLTFExporterUtility::GetPreviewMesh(const UMaterialInterface* Material)
@@ -45,3 +47,28 @@ const USkeletalMesh* FGLTFExporterUtility::GetPreviewMesh(const UAnimSequence* A
 
 	return PreviewMesh;
 }
+
+TArray<ULevel*> FGLTFExporterUtility::GetReferencedLevels(const UObject* Object)
+{
+	FAssetRegistryModule& AssetRegistryModule = FModuleManager::LoadModuleChecked<FAssetRegistryModule>("AssetRegistry");
+	const FName OuterPathName = *Object->GetOutermost()->GetPathName();
+
+	TArray<ULevel*> ReferencedLevels;
+	TArray<FAssetIdentifier> AssetDependencies;
+	AssetRegistryModule.Get().GetDependencies(OuterPathName, AssetDependencies);
+
+	for (FAssetIdentifier& AssetDependency : AssetDependencies)
+	{
+		FString PackageName = AssetDependency.PackageName.ToString();
+		if (FEditorFileUtils::IsMapPackageAsset(PackageName))
+		{
+			UWorld* World = LoadObject<UWorld>(nullptr, *PackageName);
+			if (World != nullptr)
+			{
+				ReferencedLevels.AddUnique(World->PersistentLevel);
+			}
+		}
+	}
+
+	return ReferencedLevels;
+}
diff --git a/Engine/Plugins/Enterprise/GLTFExporter/MigratedFromGit/Source/GLTFExporter/Private/Exporters/GLTFExporterUtility.h b/Engine/Plugins/Enterprise/GLTFExporter/MigratedFromGit/Source/GLTFExporter/Private/Exporters/GLTFExporterUtility.h
--- a/Engine/Plugins/Enterprise/GLTFExporter/MigratedFromGit/Source/GLTFExporter/Private/Exporters/GLTFExporterUtility.h
+++ b/Engine/Plugins/Enterprise/GLTFExporter/MigratedFromGit/Source/GLTFExporter/Private/Exporters/GLTFExporterUtility.h
@@ -4,9 +4,13 @@
 
 class UStaticMesh;
 class UMaterialInterface;
+class UAnimSequence;
+class ULevelSequence;
 
 struct FGLTFExporterUtility
 {
 	static const UStaticMesh* GetPreviewMesh(const UMaterialInterface* Material);
 	static const USkeletalMesh* GetPreviewMesh(const UAnimSequence* AnimSequence);
+
+	static TArray<ULevel*> GetReferencedLevels(const UObject* LevelSequence);
 };
diff --git a/Engine/Plugins/Enterprise/GLTFExporter/MigratedFromGit/Source/GLTFExporter/Private/Exporters/GLTFLevelSequenceExporter.cpp b/Engine/Plugins/Enterprise/GLTFExporter/MigratedFromGit/Source/GLTFExporter/Private/Exporters/GLTFLevelSequenceExporter.cpp
new file mode 100644
--- /dev/null
+++ b/Engine/Plugins/Enterprise/GLTFExporter/MigratedFromGit/Source/GLTFExporter/Private/Exporters/GLTFLevelSequenceExporter.cpp
@@ -0,0 +1,64 @@
+// Copyright Epic Games, Inc. All Rights Reserved.
+
+#include "Exporters/GLTFLevelSequenceExporter.h"
+#include "Exporters/GLTFExporterUtility.h"
+#include "Builders/GLTFContainerBuilder.h"
+#include "LevelSequence.h"
+
+UGLTFLevelSequenceExporter::UGLTFLevelSequenceExporter(const FObjectInitializer& ObjectInitializer)
+	: Super(ObjectInitializer)
+{
+	SupportedClass = ULevelSequence::StaticClass();
+}
+
+bool UGLTFLevelSequenceExporter::AddObject(FGLTFContainerBuilder& Builder, const UObject* Object)
+{
+	const ULevelSequence* LevelSequence = CastChecked<ULevelSequence>(Object);
+
+	if (!Builder.ExportOptions->bExportLevelSequences)
+	{
+		Builder.AddErrorMessage(
+			FString::Printf(TEXT("Failed to export level sequence %s because level sequences are disabled by export options"),
+			*LevelSequence->GetName()));
+		return false;
+	}
+
+	TArray<ULevel*> Levels = FGLTFExporterUtility::GetReferencedLevels(LevelSequence);
+	if (Levels.Num() == 0)
+	{
+		Builder.AddErrorMessage(
+			FString::Printf(TEXT("Failed to export level sequence %s because no level referenced"),
+			*LevelSequence->GetName()));
+		return false;
+	}
+
+	if (Levels.Num() > 1)
+	{
+		Builder.AddErrorMessage(
+            FString::Printf(TEXT("Failed to export level sequence %s because more than one level referenced"),
+            *LevelSequence->GetName()));
+		return false;
+	}
+
+	ULevel* Level = Levels[0];
+	const FGLTFJsonSceneIndex SceneIndex = Builder.GetOrAddScene(Level);
+	if (SceneIndex == INDEX_NONE)
+	{
+		Builder.AddErrorMessage(
+			FString::Printf(TEXT("Failed to export level %s for level sequence %s"),
+			*Level->GetName(),
+			*LevelSequence->GetName()));
+		return false;
+	}
+
+	const FGLTFJsonAnimationIndex AnimationIndex = Builder.GetOrAddAnimation(Level, LevelSequence);
+	if (AnimationIndex == INDEX_NONE)
+	{
+		Builder.AddErrorMessage(
+			FString::Printf(TEXT("Failed to export level sequence %s"),
+			*LevelSequence->GetName()));
+		return false;
+	}
+
+	return true;
+}
diff --git a/Engine/Plugins/Enterprise/GLTFExporter/MigratedFromGit/Source/GLTFExporter/Private/Exporters/GLTFLevelVariantSetsExporter.cpp b/Engine/Plugins/Enterprise/GLTFExporter/MigratedFromGit/Source/GLTFExporter/Private/Exporters/GLTFLevelVariantSetsExporter.cpp
new file mode 100644
--- /dev/null
+++ b/Engine/Plugins/Enterprise/GLTFExporter/MigratedFromGit/Source/GLTFExporter/Private/Exporters/GLTFLevelVariantSetsExporter.cpp
@@ -0,0 +1,67 @@
+// Copyright Epic Games, Inc. All Rights Reserved.
+
+#include "Exporters/GLTFLevelVariantSetsExporter.h"
+#include "Exporters/GLTFExporterUtility.h"
+#include "Builders/GLTFContainerBuilder.h"
+#include "LevelVariantSets.h"
+
+UGLTFLevelVariantSetsExporter::UGLTFLevelVariantSetsExporter(const FObjectInitializer& ObjectInitializer)
+	: Super(ObjectInitializer)
+{
+	SupportedClass = ULevelVariantSets::StaticClass();
+}
+
+bool UGLTFLevelVariantSetsExporter::AddObject(FGLTFContainerBuilder& Builder, const UObject* Object)
+{
+	const ULevelVariantSets* LevelVariantSets = CastChecked<ULevelVariantSets>(Object);
+
+	if (!Builder.ExportOptions->bExportVariantSets)
+	{
+		Builder.AddErrorMessage(
+			FString::Printf(TEXT("Failed to export level variant sets %s because variant sets are disabled by export options"),
+			*LevelVariantSets->GetName()));
+		return false;
+	}
+
+	TArray<ULevel*> Levels = FGLTFExporterUtility::GetReferencedLevels(LevelVariantSets);
+	if (Levels.Num() == 0)
+	{
+		Builder.AddErrorMessage(
+			FString::Printf(TEXT("Failed to export level variant sets %s because no level referenced"),
+			*LevelVariantSets->GetName()));
+		return false;
+	}
+
+	if (Levels.Num() > 1)
+	{
+		Builder.AddErrorMessage(
+            FString::Printf(TEXT("Failed to export level variant sets %s because more than one level referenced"),
+            *LevelVariantSets->GetName()));
+		return false;
+	}
+
+	ULevel* Level = Levels[0];
+	const FGLTFJsonSceneIndex SceneIndex = Builder.GetOrAddScene(Level);
+	if (SceneIndex == INDEX_NONE)
+	{
+		Builder.AddErrorMessage(
+			FString::Printf(TEXT("Failed to export level %s for level variant sets %s"),
+			*Level->GetName(),
+			*LevelVariantSets->GetName()));
+		return false;
+	}
+
+	const FGLTFJsonVariationIndex VariationIndex = Builder.GetOrAddVariation(LevelVariantSets);
+	if (VariationIndex == INDEX_NONE)
+	{
+		Builder.AddErrorMessage(
+			FString::Printf(TEXT("Failed to export level variant sets %s"),
+			*LevelVariantSets->GetName()));
+		return false;
+	}
+
+	Builder.GetScene(SceneIndex).Variations.AddUnique(VariationIndex);
+
+	Builder.DefaultScene = SceneIndex;
+	return true;
+}
diff --git a/Engine/Plugins/Enterprise/GLTFExporter/MigratedFromGit/Source/GLTFExporter/Public/Exporters/GLTFLevelSequenceExporter.h b/Engine/Plugins/Enterprise/GLTFExporter/MigratedFromGit/Source/GLTFExporter/Public/Exporters/GLTFLevelSequenceExporter.h
new file mode 100644
--- /dev/null
+++ b/Engine/Plugins/Enterprise/GLTFExporter/MigratedFromGit/Source/GLTFExporter/Public/Exporters/GLTFLevelSequenceExporter.h
@@ -0,0 +1,18 @@
+// Copyright Epic Games, Inc. All Rights Reserved.
+
+#pragma once
+
+#include "Exporters/GLTFExporter.h"
+#include "GLTFLevelSequenceExporter.generated.h"
+
+UCLASS()
+class GLTFEXPORTER_API UGLTFLevelSequenceExporter final : public UGLTFExporter
+{
+public:
+
+	GENERATED_BODY()
+
+	explicit UGLTFLevelSequenceExporter(const FObjectInitializer& ObjectInitializer = FObjectInitializer());
+
+	bool AddObject(FGLTFContainerBuilder& Builder, const UObject* Object) override;
+};
diff --git a/Engine/Plugins/Enterprise/GLTFExporter/MigratedFromGit/Source/GLTFExporter/Public/Exporters/GLTFLevelVariantSetsExporter.h b/Engine/Plugins/Enterprise/GLTFExporter/MigratedFromGit/Source/GLTFExporter/Public/Exporters/GLTFLevelVariantSetsExporter.h
new file mode 100644
--- /dev/null
+++ b/Engine/Plugins/Enterprise/GLTFExporter/MigratedFromGit/Source/GLTFExporter/Public/Exporters/GLTFLevelVariantSetsExporter.h
@@ -0,0 +1,18 @@
+// Copyright Epic Games, Inc. All Rights Reserved.
+
+#pragma once
+
+#include "Exporters/GLTFExporter.h"
+#include "GLTFLevelVariantSetsExporter.generated.h"
+
+UCLASS()
+class GLTFEXPORTER_API UGLTFLevelVariantSetsExporter final : public UGLTFExporter
+{
+public:
+
+	GENERATED_BODY()
+
+	explicit UGLTFLevelVariantSetsExporter(const FObjectInitializer& ObjectInitializer = FObjectInitializer());
+
+	bool AddObject(FGLTFContainerBuilder& Builder, const UObject* Object) override;
+};
