[GLTFExporter] fix: improve detection of degenerate UVs

commit f4a5786dd1780b5d08812bf36479141af13f0b2e
Date: Sun, 29 Nov 2020 21:08:43 +0100

    fix: improve detection of degenerate UVs
    (cherry picked from commit eb54ac47f31bbf9461fe3f3912f416eaa2e6bf5a​)
    # Conflicts from cherry-pick:
    #	Plugins/GLTFExporter/README.md

#skipci

diff --git a/Engine/Plugins/Enterprise/GLTFExporter/MigratedFromGit/Source/GLTFExporter/Private/Converters/GLTFMaterialConverters.cpp b/Engine/Plugins/Enterprise/GLTFExporter/MigratedFromGit/Source/GLTFExporter/Private/Converters/GLTFMaterialConverters.cpp
--- a/Engine/Plugins/Enterprise/GLTFExporter/MigratedFromGit/Source/GLTFExporter/Private/Converters/GLTFMaterialConverters.cpp
+++ b/Engine/Plugins/Enterprise/GLTFExporter/MigratedFromGit/Source/GLTFExporter/Private/Converters/GLTFMaterialConverters.cpp
@@ -18,22 +18,23 @@ void FGLTFMaterialConverter::Sanitize(const UMaterialInterface*& Material, const
 	if (MeshData != nullptr)
 	{
 		const FMeshDescription& MeshDescription = MeshData->GetParent()->Description;
-		for (const int32 SectionIndex : SectionIndices)
+
+		const float DegenerateUVPercentage = UVDegenerateChecker.GetOrAdd(&MeshDescription, SectionIndices, MeshData->TexCoord);
+		if (FMath::IsNearlyEqual(DegenerateUVPercentage, 1))
 		{
-			const bool bHasDegenerateUVs = UVDegenerateChecker.GetOrAdd(&MeshDescription, SectionIndex, MeshData->TexCoord);
-			if (bHasDegenerateUVs)
-			{
-				Builder.AddWarningMessage(FString::Printf(
-					TEXT("Material %s is using mesh data from %s but the lightmap UVs (channel %d) are degenerate (in mesh section %d). Simple baking will be used as fallback"),
-					*Material->GetName(),
-					*MeshData->GetParent()->Name,
-					MeshData->TexCoord,
-					SectionIndex));
-
-				MeshData = nullptr;
-				SectionIndices = {};
-				break;
-			}
+			FString SectionString = TEXT("mesh section");
+			SectionString += SectionIndices.Num() > 1 ? TEXT("s ") : TEXT(" ");
+			SectionString += FString::JoinBy(SectionIndices, TEXT(", "), FString::FromInt);
+
+			Builder.AddWarningMessage(FString::Printf(
+				TEXT("Material %s uses mesh data from %s but the lightmap UVs (channel %d) are 100%% degenerate (in %s). Simple baking will be used as fallback"),
+				*Material->GetName(),
+				*MeshData->GetParent()->Name,
+				MeshData->TexCoord,
+				*SectionString));
+
+			MeshData = nullptr;
+			SectionIndices = {};
 		}
 	}
 }
diff --git a/Engine/Plugins/Enterprise/GLTFExporter/MigratedFromGit/Source/GLTFExporter/Private/Converters/GLTFUVDegenerateChecker.cpp b/Engine/Plugins/Enterprise/GLTFExporter/MigratedFromGit/Source/GLTFExporter/Private/Converters/GLTFUVDegenerateChecker.cpp
--- a/Engine/Plugins/Enterprise/GLTFExporter/MigratedFromGit/Source/GLTFExporter/Private/Converters/GLTFUVDegenerateChecker.cpp
+++ b/Engine/Plugins/Enterprise/GLTFExporter/MigratedFromGit/Source/GLTFExporter/Private/Converters/GLTFUVDegenerateChecker.cpp
@@ -4,14 +4,17 @@
 #include "StaticMeshAttributes.h"
 #include "MeshDescription.h"
 
-void FGLTFUVDegenerateChecker::Sanitize(const FMeshDescription*& Description, int32& SectionIndex, int32& TexCoord)
+void FGLTFUVDegenerateChecker::Sanitize(const FMeshDescription*& Description, TArray<int32>& SectionIndices, int32& TexCoord)
 {
 	if (Description != nullptr)
 	{
 		const int32 SectionCount = Description->PolygonGroups().Num();
-		if (SectionIndex < 0 || SectionIndex >= SectionCount)
+		for (const int32 SectionIndex : SectionIndices)
 		{
-			Description = nullptr;
+			if (SectionIndex < 0 || SectionIndex >= SectionCount)
+			{
+				Description = nullptr;
+			}
 		}
 	}
 
@@ -25,42 +28,79 @@ void FGLTFUVDegenerateChecker::Sanitize(const FMeshDescription*& Description, in
 	}
 }
 
-bool FGLTFUVDegenerateChecker::Convert(const FMeshDescription* Description, int32 SectionIndex, int32 TexCoord)
+float FGLTFUVDegenerateChecker::Convert(const FMeshDescription* Description, TArray<int32> SectionIndices, int32 TexCoord)
 {
 	if (Description == nullptr)
 	{
 		// TODO: report warning?
 
-		return false;
+		return -1;
 	}
 
-	// TODO: since we are not really checking for degenerate UVs, but simply if all UVs have the same value or not, we should call this class something else
-
-	const FPolygonGroupID PolygonGroupID = FPolygonGroupID(SectionIndex);
-	const TVertexInstanceAttributesConstRef<FVector2D> VertexInstanceUVs =
+	const TVertexAttributesConstRef<FVector> Positions =
+		Description->VertexAttributes().GetAttributesRef<FVector>(MeshAttribute::Vertex::Position);
+	const TVertexInstanceAttributesConstRef<FVector2D> UVs =
 		Description->VertexInstanceAttributes().GetAttributesRef<FVector2D>(MeshAttribute::VertexInstance::TextureCoordinate);
 
-	uint32 Index = 0;
-	FVector2D ReferenceUV;
+	int32 TriangleCount = 0;
+	int32 DegenerateCount = 0;
 
-	for (const FPolygonID PolygonID : Description->GetPolygonGroupPolygons(PolygonGroupID))
+	for (const int32 SectionIndex : SectionIndices)
 	{
-		for (const FVertexInstanceID VertexInstanceID : Description->GetPolygonVertexInstances(PolygonID))
+		for (const FPolygonID PolygonID : Description->GetPolygonGroupPolygons(FPolygonGroupID(SectionIndex)))
 		{
-			const FVector2D UV = VertexInstanceUVs.Get(VertexInstanceID, TexCoord);
-
-			if (Index == 0)
-			{
-				ReferenceUV = UV;
-			}
-			else if (!UV.Equals(ReferenceUV))
+			for (const FTriangleID TriangleID : Description->GetPolygonTriangleIDs(PolygonID))
 			{
-				return false;
-			}
+				const TStaticArray<FVertexID, 3> TriangleVertexIDs = Description->GetTriangleVertices(TriangleID);
+				TArrayView<const FVertexInstanceID> TriangleVertexInstanceIDs = Description->GetTriangleVertexInstances(TriangleID);
+
+				TStaticArray<FVector, 3> TrianglePositions;
+				TStaticArray<FVector2D, 3> TriangleUVs;
 
-			++Index;
+				for (int32 Index = 0; Index < 3; Index++)
+				{
+					TrianglePositions[Index] = Positions.Get(TriangleVertexIDs[Index]);
+					TriangleUVs[Index] = UVs.Get(TriangleVertexInstanceIDs[Index], TexCoord);
+				}
+
+				if (!IsDegenerateTriangle(TrianglePositions))
+				{
+					TriangleCount++;
+
+					if (IsDegenerateTriangle(TriangleUVs))
+					{
+						DegenerateCount++;
+					}
+				}
+			}
 		}
 	}
 
-	return true;
+	if (TriangleCount == 0)
+	{
+		return -1;
+	}
+
+	if (TriangleCount == DegenerateCount)
+	{
+		return 1;
+	}
+
+	return static_cast<float>(DegenerateCount) / static_cast<float>(TriangleCount);
+}
+
+bool FGLTFUVDegenerateChecker::IsDegenerateTriangle(const TStaticArray<FVector2D, 3>& Points)
+{
+	const FVector2D AB = Points[1] - Points[0];
+	const FVector2D AC = Points[2] - Points[0];
+	const float DoubleArea = FMath::Abs(AB ^ AC);
+	return DoubleArea < 2 * SMALL_NUMBER;
+}
+
+bool FGLTFUVDegenerateChecker::IsDegenerateTriangle(const TStaticArray<FVector, 3>& Points)
+{
+	const FVector AB = Points[1] - Points[0];
+	const FVector AC = Points[2] - Points[0];
+	const float DoubleAreaSquared = (AB ^ AC).SizeSquared();
+	return DoubleAreaSquared < 2 * SMALL_NUMBER * SMALL_NUMBER;
 }
diff --git a/Engine/Plugins/Enterprise/GLTFExporter/MigratedFromGit/Source/GLTFExporter/Private/Converters/GLTFUVDegenerateChecker.h b/Engine/Plugins/Enterprise/GLTFExporter/MigratedFromGit/Source/GLTFExporter/Private/Converters/GLTFUVDegenerateChecker.h
--- a/Engine/Plugins/Enterprise/GLTFExporter/MigratedFromGit/Source/GLTFExporter/Private/Converters/GLTFUVDegenerateChecker.h
+++ b/Engine/Plugins/Enterprise/GLTFExporter/MigratedFromGit/Source/GLTFExporter/Private/Converters/GLTFUVDegenerateChecker.h
@@ -5,9 +5,12 @@
 #include "Converters/GLTFConverter.h"
 #include "Engine.h"
 
-class FGLTFUVDegenerateChecker final : public TGLTFConverter<bool, const FMeshDescription*, int32, int32>
+class FGLTFUVDegenerateChecker final : public TGLTFConverter<float, const FMeshDescription*, TArray<int32>, int32>
 {
-	virtual void Sanitize(const FMeshDescription*& Description, int32& SectionIndex, int32& TexCoord) override;
+	virtual void Sanitize(const FMeshDescription*& Description, TArray<int32>& SectionIndices, int32& TexCoord) override;
 
-	virtual bool Convert(const FMeshDescription* Description, int32 SectionIndex, int32 TexCoord) override;
+	virtual float Convert(const FMeshDescription* Description, TArray<int32> SectionIndices, int32 TexCoord) override;
+
+	static bool IsDegenerateTriangle(const TStaticArray<FVector2D, 3>& Points);
+	static bool IsDegenerateTriangle(const TStaticArray<FVector, 3>& Points);
 };
diff --git a/Engine/Plugins/Enterprise/GLTFExporter/MigratedFromGit/Source/GLTFExporter/Private/Converters/GLTFUVOverlapChecker.cpp b/Engine/Plugins/Enterprise/GLTFExporter/MigratedFromGit/Source/GLTFExporter/Private/Converters/GLTFUVOverlapChecker.cpp
--- a/Engine/Plugins/Enterprise/GLTFExporter/MigratedFromGit/Source/GLTFExporter/Private/Converters/GLTFUVOverlapChecker.cpp
+++ b/Engine/Plugins/Enterprise/GLTFExporter/MigratedFromGit/Source/GLTFExporter/Private/Converters/GLTFUVOverlapChecker.cpp
@@ -81,28 +81,38 @@ float FGLTFUVOverlapChecker::Convert(const FMeshDescription* Description, const
 	// so to be safe we set the limit at 15 instead of 10 to prevent incorrectly flagging some pixels as overlapping.
 	const float EmissiveThreshold = 15.0f;
 
-	float ProcessedPixels = 0;
-	float OverlappingPixels = 0;
+	int32 ValidCount = 0;
+	int32 OverlapCount = 0;
 
 	if (BakeOutput.EmissiveScale > EmissiveThreshold)
 	{
-		const FColor Magenta(255, 0, 255, 255);
-		const uint8 ColorThreshold = FMath::RoundToInt((EmissiveThreshold / BakeOutput.EmissiveScale) * 255.0f); // This should never overflow since we know that the quotient will be <= 1.0
+		// This should never overflow since we know that the quotient will be <= 1.0
+		const uint8 ColorThreshold = FMath::RoundToInt((EmissiveThreshold / BakeOutput.EmissiveScale) * 255.0f);
 
 		bool bIsMagenta;
 		bool bIsOverlapping;
 
 		for (const FColor& Pixel: BakedPixels)
 		{
-			bIsMagenta = Pixel == Magenta;
+			bIsMagenta = Pixel == FColor::Magenta;
 			bIsOverlapping = Pixel.G > ColorThreshold;
 
-			ProcessedPixels += !bIsMagenta ? 1 : 0;
-			OverlappingPixels += !bIsMagenta && bIsOverlapping ? 1 : 0;
+			ValidCount += !bIsMagenta ? 1 : 0;
+			OverlapCount += !bIsMagenta && bIsOverlapping ? 1 : 0;
 		}
 	}
 
-	return OverlappingPixels / ProcessedPixels;
+	if (ValidCount == 0)
+	{
+		return -1;
+	}
+
+	if (ValidCount == OverlapCount)
+	{
+		return 1;
+	}
+
+	return static_cast<float>(OverlapCount) / static_cast<float>(ValidCount);
 }
 
 UMaterialInterface* FGLTFUVOverlapChecker::GetMaterial()
