[GLTFExporter] fix: add important missing error and warning messages

commit 6fd7459d35edb984645c00152b34e8e8545b2fe5
Date: Mon, 30 Nov 2020 23:13:39 +0100

    fix: add important missing error and warning messages
    (cherry picked from commit d233f4b6620745b5b8377bc12a07f751076e92b5​)

#skipci

diff --git a/Engine/Plugins/Enterprise/GLTFExporter/MigratedFromGit/Source/GLTFExporter/Private/Builders/GLTFBufferBuilder.cpp b/Engine/Plugins/Enterprise/GLTFExporter/MigratedFromGit/Source/GLTFExporter/Private/Builders/GLTFBufferBuilder.cpp
--- a/Engine/Plugins/Enterprise/GLTFExporter/MigratedFromGit/Source/GLTFExporter/Private/Builders/GLTFBufferBuilder.cpp
+++ b/Engine/Plugins/Enterprise/GLTFExporter/MigratedFromGit/Source/GLTFExporter/Private/Builders/GLTFBufferBuilder.cpp
@@ -22,7 +22,7 @@ FGLTFBufferBuilder::FGLTFBufferBuilder(const FString& FilePath, const UGLTFExpor
 		BufferArchive.Reset(IFileManager::Get().CreateFileWriter(*ExternalBinaryPath));
 		if (BufferArchive == nullptr)
 		{
-			// TODO: report error
+			AddErrorMessage(FString::Printf(TEXT("Failed to write external binary buffer to file: %s"), *ExternalBinaryPath));
 		}
 	}
 
diff --git a/Engine/Plugins/Enterprise/GLTFExporter/MigratedFromGit/Source/GLTFExporter/Private/Builders/GLTFImageBuilder.cpp b/Engine/Plugins/Enterprise/GLTFExporter/MigratedFromGit/Source/GLTFExporter/Private/Builders/GLTFImageBuilder.cpp
--- a/Engine/Plugins/Enterprise/GLTFExporter/MigratedFromGit/Source/GLTFExporter/Private/Builders/GLTFImageBuilder.cpp
+++ b/Engine/Plugins/Enterprise/GLTFExporter/MigratedFromGit/Source/GLTFExporter/Private/Builders/GLTFImageBuilder.cpp
@@ -108,7 +108,7 @@ FString FGLTFImageBuilder::SaveImageToFile(const void* CompressedData, int64 Com
 
 	if (!FFileHelper::SaveArrayToFile(ImageData, *ImagePath))
 	{
-		// TODO: report error
+		AddErrorMessage(FString::Printf(TEXT("Failed to save image %s to file: %s"), *Name, *ImagePath));
 		return TEXT("");
 	}
 
diff --git a/Engine/Plugins/Enterprise/GLTFExporter/MigratedFromGit/Source/GLTFExporter/Private/Builders/GLTFTaskBuilder.cpp b/Engine/Plugins/Enterprise/GLTFExporter/MigratedFromGit/Source/GLTFExporter/Private/Builders/GLTFTaskBuilder.cpp
--- a/Engine/Plugins/Enterprise/GLTFExporter/MigratedFromGit/Source/GLTFExporter/Private/Builders/GLTFTaskBuilder.cpp
+++ b/Engine/Plugins/Enterprise/GLTFExporter/MigratedFromGit/Source/GLTFExporter/Private/Builders/GLTFTaskBuilder.cpp
@@ -47,13 +47,13 @@ bool FGLTFTaskBuilder::SetupTask(TUniquePtr<FGLTFTask> Task)
 	const EGLTFTaskPriority Priority = Task->Priority;
 	if (static_cast<int32>(Priority) >= static_cast<int32>(EGLTFTaskPriority::MAX))
 	{
-		// TODO: report error
+		checkNoEntry();
 		return false;
 	}
 
 	if (static_cast<int32>(Priority) <= PriorityIndexLock)
 	{
-		// TODO: report error
+		checkNoEntry();
 		return false;
 	}
 
diff --git a/Engine/Plugins/Enterprise/GLTFExporter/MigratedFromGit/Source/GLTFExporter/Private/Tasks/GLTFAnimationTasks.cpp b/Engine/Plugins/Enterprise/GLTFExporter/MigratedFromGit/Source/GLTFExporter/Private/Tasks/GLTFAnimationTasks.cpp
--- a/Engine/Plugins/Enterprise/GLTFExporter/MigratedFromGit/Source/GLTFExporter/Private/Tasks/GLTFAnimationTasks.cpp
+++ b/Engine/Plugins/Enterprise/GLTFExporter/MigratedFromGit/Source/GLTFExporter/Private/Tasks/GLTFAnimationTasks.cpp
@@ -266,6 +266,8 @@ void FGLTFLevelSequenceTask::Complete()
 
 				for (UMovieSceneSection* Section : TransformTrack->GetAllSections())
 				{
+					// TODO: add support for Section->GetCompletionMode() (i.e. when finished)
+
 					UMovieScene3DTransformSection* TransformSection = Cast<UMovieScene3DTransformSection>(Section);
 					if (TransformSection == nullptr)
 					{
diff --git a/Engine/Plugins/Enterprise/GLTFExporter/MigratedFromGit/Source/GLTFExporter/Private/Tasks/GLTFMeshTasks.cpp b/Engine/Plugins/Enterprise/GLTFExporter/MigratedFromGit/Source/GLTFExporter/Private/Tasks/GLTFMeshTasks.cpp
--- a/Engine/Plugins/Enterprise/GLTFExporter/MigratedFromGit/Source/GLTFExporter/Private/Tasks/GLTFMeshTasks.cpp
+++ b/Engine/Plugins/Enterprise/GLTFExporter/MigratedFromGit/Source/GLTFExporter/Private/Tasks/GLTFMeshTasks.cpp
@@ -98,7 +98,7 @@ void FGLTFStaticMeshTask::Complete()
 		}
 		else if (MeshData->TexCoord < 0)
 		{
-			// TODO: report warning (about missing non-overlapping texture coordinate for baking with mesh data).
+			// TODO: report warning (about missing lightmap texture coordinate for baking with mesh data).
 			MeshData = nullptr;
 		}
 	}
diff --git a/Engine/Plugins/Enterprise/GLTFExporter/MigratedFromGit/Source/GLTFExporter/Private/Tasks/GLTFTextureTasks.cpp b/Engine/Plugins/Enterprise/GLTFExporter/MigratedFromGit/Source/GLTFExporter/Private/Tasks/GLTFTextureTasks.cpp
--- a/Engine/Plugins/Enterprise/GLTFExporter/MigratedFromGit/Source/GLTFExporter/Private/Tasks/GLTFTextureTasks.cpp
+++ b/Engine/Plugins/Enterprise/GLTFExporter/MigratedFromGit/Source/GLTFExporter/Private/Tasks/GLTFTextureTasks.cpp
@@ -55,7 +55,7 @@ void FGLTFTextureCubeTask::Complete()
 	const UTexture2D* FaceTexture = FGLTFTextureUtility::CreateTextureFromCubeFace(TextureCube, CubeFace);
 	if (FaceTexture == nullptr)
 	{
-		// TODO: report error
+		Builder.AddWarningMessage(FString::Printf(TEXT("Failed to extract cube face %d for cubemap texture %s"), CubeFace, *TextureCube->GetName()));
 		return;
 	}
 
@@ -74,7 +74,7 @@ void FGLTFTextureCubeTask::Complete()
 	TArray<FColor> Pixels;
 	if (!FGLTFTextureUtility::ReadPixels(RenderTarget, Pixels, JsonTexture.Encoding))
 	{
-		Builder.AddWarningMessage(FString::Printf(TEXT("Failed to read pixels for cubemap texture %s"), *JsonTexture.Name));
+		Builder.AddWarningMessage(FString::Printf(TEXT("Failed to read pixels (cube face %d) for cubemap texture %s"), CubeFace, *TextureCube->GetName()));
 		return;
 	}
 
@@ -122,7 +122,7 @@ void FGLTFTextureRenderTargetCubeTask::Complete()
 	const UTexture2D* FaceTexture = FGLTFTextureUtility::CreateTextureFromCubeFace(RenderTargetCube, CubeFace);
 	if (FaceTexture == nullptr)
 	{
-		// TODO: report error
+		Builder.AddWarningMessage(FString::Printf(TEXT("Failed to extract cube face %d for cubemap render target %s"), CubeFace, *RenderTargetCube->GetName()));
 		return;
 	}
 
@@ -141,7 +141,7 @@ void FGLTFTextureRenderTargetCubeTask::Complete()
 	TArray<FColor> Pixels;
 	if (!FGLTFTextureUtility::ReadPixels(RenderTarget, Pixels, JsonTexture.Encoding))
 	{
-		Builder.AddWarningMessage(FString::Printf(TEXT("Failed to read pixels for cubemap render target %s"), *JsonTexture.Name));
+		Builder.AddWarningMessage(FString::Printf(TEXT("Failed to read pixels (cube face %d) for cubemap render target %s"), CubeFace, *RenderTargetCube->GetName()));
 		return;
 	}
 
