[GLTFExporter] refactor: always include GLTFLaunchHelper.exe in web viewer bundle even on non-Windows platforms

commit dd2d4b039d2003113ba0dd5159a3bba16eead319
Date: Tue, 1 Dec 2020 04:29:11 +0100

    refactor: always include GLTFLaunchHelper.exe in web viewer bundle even on non-Windows platforms
    (cherry picked from commit 00327ac0c90aa773ade95970077c24dde05b8c72​)

#skipci

diff --git a/Engine/Plugins/Enterprise/GLTFExporter/MigratedFromGit/Source/GLTFExporter/Private/Builders/GLTFWebBuilder.cpp b/Engine/Plugins/Enterprise/GLTFExporter/MigratedFromGit/Source/GLTFExporter/Private/Builders/GLTFWebBuilder.cpp
--- a/Engine/Plugins/Enterprise/GLTFExporter/MigratedFromGit/Source/GLTFExporter/Private/Builders/GLTFWebBuilder.cpp
+++ b/Engine/Plugins/Enterprise/GLTFExporter/MigratedFromGit/Source/GLTFExporter/Private/Builders/GLTFWebBuilder.cpp
@@ -87,12 +87,6 @@ void FGLTFWebBuilder::BundleWebViewer(const FString& ResourcesDir)
 
 void FGLTFWebBuilder::BundleLaunchHelper(const FString& ResourcesDir)
 {
-	const FString ExecutableName = GetLaunchHelperExecutable();
-	if (ExecutableName.IsEmpty())
-	{
-		return;
-	}
-
 	const FString ArchiveFile = ResourcesDir / TEXT("GLTFLaunchHelper.zip");
 
 	if (!FPaths::FileExists(ArchiveFile))
@@ -101,6 +95,8 @@ void FGLTFWebBuilder::BundleLaunchHelper(const FString& ResourcesDir)
 		return;
 	}
 
+	const FString ExecutableName = TEXT("GLTFLaunchHelper.exe");
+
 	if (!FGLTFZipUtility::ExtractOneFile(ArchiveFile, ExecutableName, DirPath))
 	{
 		AddErrorMessage(FString::Printf(TEXT("Failed to extract launch helper file (%s) from %s"), *ExecutableName, *ArchiveFile));
@@ -114,14 +110,3 @@ void FGLTFWebBuilder::BundleLaunchHelper(const FString& ResourcesDir)
 		AddWarningMessage(FString::Printf(TEXT("Failed to make launch helper file executable at %s"), *ExecutableFile));
 	}
 }
-
-const TCHAR* FGLTFWebBuilder::GetLaunchHelperExecutable()
-{
-#if PLATFORM_WINDOWS
-	return TEXT("GLTFLaunchHelper.exe");
-#elif PLATFORM_MAC
-	return TEXT(""); // TODO: add GLTFLaunchHelper for macos
-#else
-	return TEXT(""); // TODO: add GLTFLaunchHelper for linux
-#endif
-}
