[GLTFExporter] fix: update edit condition in bake settings export options

commit a8cfce5e0ea2789e8a5a8fcaac3837ddc7436daf
Date: Wed, 2 Dec 2020 02:44:03 +0100

    fix: update edit condition in bake settings export options
    (cherry picked from commit 4ab6dac5303af9ce395af19473216c85cd43f974​)

#skipci

diff --git a/Engine/Plugins/Enterprise/GLTFExporter/MigratedFromGit/Source/GLTFExporter/Public/GLTFExportOptions.h b/Engine/Plugins/Enterprise/GLTFExporter/MigratedFromGit/Source/GLTFExporter/Public/GLTFExportOptions.h
--- a/Engine/Plugins/Enterprise/GLTFExporter/MigratedFromGit/Source/GLTFExporter/Public/GLTFExportOptions.h
+++ b/Engine/Plugins/Enterprise/GLTFExporter/MigratedFromGit/Source/GLTFExporter/Public/GLTFExportOptions.h
@@ -101,15 +101,15 @@ class GLTFEXPORTER_API UGLTFExportOptions : public UObject
 	EGLTFMaterialBakeMode BakeMaterialInputs;
 
 	/** Default size of the baked out texture (containing the material input). Can be overriden by material- and input-specific bake settings, see GLTFMaterialExportOptions. */
-	UPROPERTY(EditAnywhere, BlueprintReadWrite, Config, Category = Material, Meta = (EditCondition = "BakeMaterialInputs != EGLTFMaterialBakeMode::None"))
+	UPROPERTY(EditAnywhere, BlueprintReadWrite, Config, Category = Material, Meta = (EditCondition = "BakeMaterialInputs != EGLTFMaterialBakeMode::Disabled"))
 	EGLTFMaterialBakeSizePOT DefaultMaterialBakeSize;
 
 	/** Default filtering mode used when sampling the baked out texture. Can be overriden by material- and input-specific bake settings, see GLTFMaterialExportOptions. */
-	UPROPERTY(EditAnywhere, BlueprintReadWrite, Config, Category = Material, Meta = (EditCondition = "BakeMaterialInputs != EGLTFMaterialBakeMode::None", ValidEnumValues="TF_Nearest, TF_Bilinear, TF_Trilinear"))
+	UPROPERTY(EditAnywhere, BlueprintReadWrite, Config, Category = Material, Meta = (EditCondition = "BakeMaterialInputs != EGLTFMaterialBakeMode::Disabled", ValidEnumValues="TF_Nearest, TF_Bilinear, TF_Trilinear"))
 	TEnumAsByte<TextureFilter> DefaultMaterialBakeFilter;
 
 	/** Default addressing mode used when sampling the baked out texture. Can be overriden by material- and input-specific bake settings, see GLTFMaterialExportOptions. */
-	UPROPERTY(EditAnywhere, BlueprintReadWrite, Config, Category = Material, Meta = (EditCondition = "BakeMaterialInputs != EGLTFMaterialBakeMode::None"))
+	UPROPERTY(EditAnywhere, BlueprintReadWrite, Config, Category = Material, Meta = (EditCondition = "BakeMaterialInputs != EGLTFMaterialBakeMode::Disabled"))
 	TEnumAsByte<TextureAddress> DefaultMaterialBakeTiling;
 
 	/** Input-specific default bake settings that override the general defaults above. */
