[GLTFExporter] feat: warn if vertex colors exported due to often unexpected results

commit 2b62b288e4ec583886dad3b6f4992f4596dcb7c9
Date: Wed, 2 Dec 2020 02:46:15 +0100

    feat: warn if vertex colors exported due to often unexpected results
    (cherry picked from commit e2ec1366c8281c29d6f54d7a449aba7bd5c7deb6​)

#skipci

diff --git a/Engine/Plugins/Enterprise/GLTFExporter/MigratedFromGit/Source/GLTFExporter/Private/Tasks/GLTFMeshTasks.cpp b/Engine/Plugins/Enterprise/GLTFExporter/MigratedFromGit/Source/GLTFExporter/Private/Tasks/GLTFMeshTasks.cpp
--- a/Engine/Plugins/Enterprise/GLTFExporter/MigratedFromGit/Source/GLTFExporter/Private/Tasks/GLTFMeshTasks.cpp
+++ b/Engine/Plugins/Enterprise/GLTFExporter/MigratedFromGit/Source/GLTFExporter/Private/Tasks/GLTFMeshTasks.cpp
@@ -8,8 +8,33 @@
 
 namespace
 {
-	template <typename VertexBufferType>
-	void ValidateVertexBuffer(FGLTFConvertBuilder& Builder, const VertexBufferType* VertexBuffer, const TCHAR* MeshName)
+	template <typename TangentVectorType>
+	void CheckTangentVectors(const FStaticMeshVertexBuffer* VertexBuffer, bool& bOutZeroNormals, bool& bOutZeroTangents)
+	{
+		bool bZeroNormals = false;
+		bool bZeroTangents = false;
+
+		typedef TStaticMeshVertexTangentDatum<TangentVectorType> VertexTangentType;
+
+		const void* TangentData = const_cast<FStaticMeshVertexBuffer*>(VertexBuffer)->GetTangentData();
+		const VertexTangentType* VertexTangents = TangentData != nullptr ? static_cast<const VertexTangentType*>(TangentData) : nullptr;
+
+		if (VertexTangents != nullptr)
+		{
+			const uint32 VertexCount = VertexBuffer->GetNumVertices();
+			for (uint32 VertexIndex = 0; VertexIndex < VertexCount; ++VertexIndex)
+			{
+				const VertexTangentType& VertexTangent = VertexTangents[VertexIndex];
+				bZeroNormals |= VertexTangent.TangentZ.ToFVector().IsNearlyZero();
+				bZeroTangents |= VertexTangent.TangentX.ToFVector().IsNearlyZero();
+			}
+		}
+
+		bOutZeroNormals = bZeroNormals;
+		bOutZeroTangents = bZeroTangents;
+	}
+
+	void ValidateVertexBuffer(FGLTFConvertBuilder& Builder, const FStaticMeshVertexBuffer* VertexBuffer, const TCHAR* MeshName)
 	{
 		if (VertexBuffer == nullptr)
 		{
@@ -21,11 +46,11 @@ namespace
 
 		if (VertexBuffer->GetUseHighPrecisionTangentBasis())
 		{
-			CheckTangentVectors<VertexBufferType, FPackedRGBA16N>(VertexBuffer, bZeroNormals, bZeroTangents);
+			CheckTangentVectors<FPackedRGBA16N>(VertexBuffer, bZeroNormals, bZeroTangents);
 		}
 		else
 		{
-			CheckTangentVectors<VertexBufferType, FPackedNormal>(VertexBuffer, bZeroNormals, bZeroTangents);
+			CheckTangentVectors<FPackedNormal>(VertexBuffer, bZeroNormals, bZeroTangents);
 		}
 
 		if (bZeroNormals)
@@ -43,30 +68,20 @@ namespace
 		}
 	}
 
-	template <typename VertexBufferType, typename TangentVectorType>
-	void CheckTangentVectors(const VertexBufferType* VertexBuffer, bool& bOutZeroNormals, bool& bOutZeroTangents)
+	bool HasVertexColors(const FColorVertexBuffer* VertexBuffer)
 	{
-		bool bZeroNormals = false;
-		bool bZeroTangents = false;
+		const uint32 VertexCount = VertexBuffer->GetNumVertices();
 
-		typedef TStaticMeshVertexTangentDatum<TangentVectorType> VertexTangentType;
-
-		const void* TangentData = const_cast<FStaticMeshVertexBuffer*>(VertexBuffer)->GetTangentData();
-		const VertexTangentType* VertexTangents = TangentData != nullptr ? static_cast<const VertexTangentType*>(TangentData) : nullptr;
-
-		if (VertexTangents != nullptr)
+		for (uint32 VertexIndex = 0; VertexIndex < VertexCount; VertexIndex++)
 		{
-			const uint32 VertexCount = VertexBuffer->GetNumVertices();
-			for (uint32 VertexIndex = 0; VertexIndex < VertexCount; ++VertexIndex)
+			const FColor& Color = VertexBuffer->VertexColor(VertexIndex);
+			if (Color != FColor::White)
 			{
-				const VertexTangentType& VertexTangent = VertexTangents[VertexIndex];
-				bZeroNormals |= VertexTangent.TangentZ.ToFVector().IsNearlyZero();
-				bZeroTangents |= VertexTangent.TangentX.ToFVector().IsNearlyZero();
+				return true;
 			}
 		}
 
-		bOutZeroNormals = bZeroNormals;
-		bOutZeroTangents = bZeroTangents;
+		return false;
 	}
 }
 
@@ -80,6 +95,17 @@ void FGLTFStaticMeshTask::Complete()
 	const FStaticMeshVertexBuffer* VertexBuffer = &MeshLOD.VertexBuffers.StaticMeshVertexBuffer;
 	const FColorVertexBuffer* ColorBuffer = &MeshLOD.VertexBuffers.ColorVertexBuffer;
 
+	if (Builder.ExportOptions->bExportVertexColors && HasVertexColors(ColorBuffer))
+	{
+		Builder.AddWarningMessage(FString::Printf(
+			TEXT("Vertex colors in mesh %s will act as a multiplier for base color in glTF, regardless of material, which may produce undesirable results."),
+			*StaticMesh->GetName()));
+	}
+	else
+	{
+		ColorBuffer = nullptr;
+	}
+
 	if (StaticMeshComponent != nullptr && StaticMeshComponent->LODData.IsValidIndex(LODIndex))
 	{
 		const FStaticMeshComponentLODInfo& LODInfo = StaticMeshComponent->LODData[LODIndex];
@@ -122,10 +148,8 @@ void FGLTFStaticMeshTask::Complete()
 			// TODO: report warning?
 		}
 
-		if (Builder.ExportOptions->bExportVertexColors)
+		if (ColorBuffer != nullptr)
 		{
-			// TODO: report warning (about vertex colors may be represented differently in glTF)
-
 			JsonPrimitive.Attributes.Color0 = Builder.GetOrAddColorAccessor(ConvertedSection, ColorBuffer);
 		}
 
@@ -158,6 +182,17 @@ void FGLTFSkeletalMeshTask::Complete()
 	const FColorVertexBuffer* ColorBuffer = &MeshLOD.StaticVertexBuffers.ColorVertexBuffer;
 	const FSkinWeightVertexBuffer* SkinWeightBuffer = MeshLOD.GetSkinWeightVertexBuffer();
 
+	if (Builder.ExportOptions->bExportVertexColors && HasVertexColors(ColorBuffer))
+	{
+		Builder.AddWarningMessage(FString::Printf(
+			TEXT("Vertex colors in mesh %s will act as a multiplier for base color in glTF, regardless of material, which may produce undesirable results."),
+			*SkeletalMesh->GetName()));
+	}
+	else
+	{
+		ColorBuffer = nullptr;
+	}
+
 	if (SkeletalMeshComponent != nullptr && SkeletalMeshComponent->LODInfo.IsValidIndex(LODIndex))
 	{
 		const FSkelMeshComponentLODInfo& LODInfo = SkeletalMeshComponent->LODInfo[LODIndex];
@@ -201,10 +236,8 @@ void FGLTFSkeletalMeshTask::Complete()
 			// TODO: report warning?
 		}
 
-		if (Builder.ExportOptions->bExportVertexColors)
+		if (ColorBuffer != nullptr)
 		{
-			// TODO: report warning (about vertex colors may be represented differently in glTF)
-
 			JsonPrimitive.Attributes.Color0 = Builder.GetOrAddColorAccessor(ConvertedSection, ColorBuffer);
 		}
 
