[GLTFExporter] chore: update plugin version to 1.0

commit f4ec07c7ed3818e1626da5c04357f671cb0486a5
Date: Wed, 2 Dec 2020 03:04:26 +0100

    chore: update plugin version to 1.0
    (cherry picked from commit c56e8afd677a94ab83d9430c4dabe087f92dfd83​)

#skipci

diff --git a/Engine/Plugins/Enterprise/GLTFExporter/MigratedFromGit/GLTFExporter.uplugin b/Engine/Plugins/Enterprise/GLTFExporter/MigratedFromGit/GLTFExporter.uplugin
--- a/Engine/Plugins/Enterprise/GLTFExporter/MigratedFromGit/GLTFExporter.uplugin
+++ b/Engine/Plugins/Enterprise/GLTFExporter/MigratedFromGit/GLTFExporter.uplugin
@@ -1,9 +1,9 @@
 {
 	"FileVersion": 3,
-	"Version": 9,
-	"VersionName": "0.9",
+	"Version": 10,
+	"VersionName": "1.0",
 	"FriendlyName": "glTF Exporter",
-	"Description": "An exporter for Khronos glTF 3D models.",
+	"Description": "An exporter for Khronos glTF 2.0.",
 	"Category": "Exporters",
 	"CreatedBy": "Epic Games, Inc.",
 	"CreatedByURL": "http://epicgames.com",
@@ -41,4 +41,4 @@
 			"Enabled": true
 		}
 	]
-}
+}
\ No newline at end of file
