[GLTFExporter] feat: add warning if multiple texture coordinates found when baking material input

commit d6a14b2a06201ba2011c901c500c176a9c59c8d4
Date: Sun, 6 Dec 2020 07:43:43 +0100

    feat: add warning if multiple texture coordinates found when baking material input
    (cherry picked from commit 74509dbbe917b3fd87be11cfa307f1d4b7aa99c0​)

#skipci

diff --git a/Engine/Plugins/Enterprise/GLTFExporter/MigratedFromGit/Source/GLTFExporter/Private/Converters/GLTFMaterialConverters.cpp b/Engine/Plugins/Enterprise/GLTFExporter/MigratedFromGit/Source/GLTFExporter/Private/Converters/GLTFMaterialConverters.cpp
--- a/Engine/Plugins/Enterprise/GLTFExporter/MigratedFromGit/Source/GLTFExporter/Private/Converters/GLTFMaterialConverters.cpp
+++ b/Engine/Plugins/Enterprise/GLTFExporter/MigratedFromGit/Source/GLTFExporter/Private/Converters/GLTFMaterialConverters.cpp
@@ -27,7 +27,7 @@ void FGLTFMaterialConverter::Sanitize(const UMaterialInterface*& Material, const
 			SectionString += FString::JoinBy(SectionIndices, TEXT(", "), FString::FromInt);
 
 			Builder.AddWarningMessage(FString::Printf(
-				TEXT("Material %s uses mesh data from %s but the lightmap UVs (channel %d) are 100%% degenerate (in %s). Simple baking will be used as fallback"),
+				TEXT("Material %s uses mesh data from %s but the lightmap UVs (channel %d) are nearly 100%% degenerate (in %s). Simple baking will be used as fallback"),
 				*Material->GetName(),
 				*MeshData->GetParent()->Name,
 				MeshData->TexCoord,
diff --git a/Engine/Plugins/Enterprise/GLTFExporter/MigratedFromGit/Source/GLTFExporter/Private/Converters/GLTFMaterialUtility.cpp b/Engine/Plugins/Enterprise/GLTFExporter/MigratedFromGit/Source/GLTFExporter/Private/Converters/GLTFMaterialUtility.cpp
--- a/Engine/Plugins/Enterprise/GLTFExporter/MigratedFromGit/Source/GLTFExporter/Private/Converters/GLTFMaterialUtility.cpp
+++ b/Engine/Plugins/Enterprise/GLTFExporter/MigratedFromGit/Source/GLTFExporter/Private/Converters/GLTFMaterialUtility.cpp
@@ -353,17 +353,18 @@ bool FGLTFMaterialUtility::TryGetTextureCoordinateIndex(const UMaterialExpressio
 	return false;
 }
 
-void FGLTFMaterialUtility::GetAllTextureCoordinateIndices(const UMaterialInterface* InMaterial, const FMaterialPropertyEx& InProperty, int32& OutNumTexCoords, TBitArray<>& OutTexCoords)
+void FGLTFMaterialUtility::GetAllTextureCoordinateIndices(const UMaterialInterface* InMaterial, const FMaterialPropertyEx& InProperty, TArray<int32>& OutTexCoords)
 {
 	FGLTFMaterialAnalysis Analysis;
 	AnalyzeMaterialProperty(InMaterial, InProperty, Analysis);
 
-	OutNumTexCoords = 0;
-	OutTexCoords = Analysis.TextureCoordinates;
-
-	for (int32 Index = 0; Index < OutTexCoords.Num(); Index++)
+	const TBitArray<>& TexCoords = Analysis.TextureCoordinates;
+	for (int32 Index = 0; Index < TexCoords.Num(); Index++)
 	{
-		OutNumTexCoords += OutTexCoords[Index] ? 1 : 0;
+		if (TexCoords[Index])
+		{
+			OutTexCoords.Add(Index);
+		}
 	}
 }
 
diff --git a/Engine/Plugins/Enterprise/GLTFExporter/MigratedFromGit/Source/GLTFExporter/Private/Converters/GLTFMaterialUtility.h b/Engine/Plugins/Enterprise/GLTFExporter/MigratedFromGit/Source/GLTFExporter/Private/Converters/GLTFMaterialUtility.h
--- a/Engine/Plugins/Enterprise/GLTFExporter/MigratedFromGit/Source/GLTFExporter/Private/Converters/GLTFMaterialUtility.h
+++ b/Engine/Plugins/Enterprise/GLTFExporter/MigratedFromGit/Source/GLTFExporter/Private/Converters/GLTFMaterialUtility.h
@@ -68,7 +68,7 @@ struct FGLTFMaterialUtility
 	static uint32 GetMaskComponentCount(const FExpressionInput& ExpressionInput);
 
 	static bool TryGetTextureCoordinateIndex(const UMaterialExpressionTextureSample* TextureSampler, int32& TexCoord, FGLTFJsonTextureTransform& Transform);
-	static void GetAllTextureCoordinateIndices(const UMaterialInterface* InMaterial, const FMaterialPropertyEx& InProperty, int32& OutNumTexCoords, TBitArray<>& OutTexCoords);
+	static void GetAllTextureCoordinateIndices(const UMaterialInterface* InMaterial, const FMaterialPropertyEx& InProperty, TArray<int32>& OutTexCoords);
 
 	static FMaterialShadingModelField EvaluateShadingModelExpression(const UMaterialInterface* Material);
 	static EMaterialShadingModel GetRichestShadingModel(const FMaterialShadingModelField& ShadingModels);
diff --git a/Engine/Plugins/Enterprise/GLTFExporter/MigratedFromGit/Source/GLTFExporter/Private/Tasks/GLTFMaterialTasks.cpp b/Engine/Plugins/Enterprise/GLTFExporter/MigratedFromGit/Source/GLTFExporter/Private/Tasks/GLTFMaterialTasks.cpp
--- a/Engine/Plugins/Enterprise/GLTFExporter/MigratedFromGit/Source/GLTFExporter/Private/Tasks/GLTFMaterialTasks.cpp
+++ b/Engine/Plugins/Enterprise/GLTFExporter/MigratedFromGit/Source/GLTFExporter/Private/Tasks/GLTFMaterialTasks.cpp
@@ -1302,17 +1302,21 @@ FGLTFPropertyBakeOutput FGLTFMaterialTask::BakeMaterialProperty(const FMaterialP
 {
 	if (MeshData == nullptr)
 	{
-		int32 NumTexCoords;
-		TBitArray<> TexCoords;
-		FGLTFMaterialUtility::GetAllTextureCoordinateIndices(Material, Property, NumTexCoords, TexCoords);
+		TArray<int32> TexCoords;
+		FGLTFMaterialUtility::GetAllTextureCoordinateIndices(Material, Property, TexCoords);
 
-		if (NumTexCoords > 0)
+		if (TexCoords.Num() > 0)
 		{
-			OutTexCoord = TexCoords.Find(true);
+			OutTexCoord = TexCoords[0];
 
 			if (TexCoords.Num() > 1)
 			{
-				// TODO: report warning (multiple texture coordinates found, will use first)
+				Builder.AddWarningMessage(FString::Printf(
+					TEXT("%s for material %s uses multiple texture coordinates (%s), baked texture will be sampled using only the first (%d)"),
+					*Property.ToString(),
+					*Material->GetName(),
+					*FString::JoinBy(TexCoords, TEXT(", "), FString::FromInt),
+					OutTexCoord));
 			}
 		}
 		else
