[GLTFExporter] fix: fixed non-grouped material inputs not using bake settings for filter and tiling

commit 392dddb19fea2026b1bf7761b3624195607276e2
Date: Sun, 6 Dec 2020 20:26:42 +0100

    fix: fixed non-grouped material inputs not using bake settings for filter and tiling
    (cherry picked from commit 0b17c98e818d21acd59e3957b061d731b95f7cf4​)

#skipci

diff --git a/Engine/Plugins/Enterprise/GLTFExporter/MigratedFromGit/Source/GLTFExporter/Private/Tasks/GLTFMaterialTasks.cpp b/Engine/Plugins/Enterprise/GLTFExporter/MigratedFromGit/Source/GLTFExporter/Private/Tasks/GLTFMaterialTasks.cpp
--- a/Engine/Plugins/Enterprise/GLTFExporter/MigratedFromGit/Source/GLTFExporter/Private/Tasks/GLTFMaterialTasks.cpp
+++ b/Engine/Plugins/Enterprise/GLTFExporter/MigratedFromGit/Source/GLTFExporter/Private/Tasks/GLTFMaterialTasks.cpp
@@ -1346,13 +1346,13 @@ FGLTFPropertyBakeOutput FGLTFMaterialTask::BakeMaterialProperty(const FMaterialP
 
 bool FGLTFMaterialTask::StoreBakedPropertyTexture(FGLTFJsonTextureInfo& OutTexInfo, const FGLTFPropertyBakeOutput& PropertyBakeOutput, const FString& PropertyName) const
 {
-	// TODO: should this be the default wrap-mode?
-	const EGLTFJsonTextureWrap TextureWrapS = EGLTFJsonTextureWrap::Repeat;
-	const EGLTFJsonTextureWrap TextureWrapT = EGLTFJsonTextureWrap::Repeat;
+	const TextureAddress TextureAddress = Builder.GetBakeTilingForMaterialProperty(Material, PropertyBakeOutput.Property);
+	const EGLTFJsonTextureWrap TextureWrapS = FGLTFConverterUtility::ConvertWrap(TextureAddress);
+	const EGLTFJsonTextureWrap TextureWrapT = FGLTFConverterUtility::ConvertWrap(TextureAddress);
 
-	// TODO: should this be the default filter?
-	const EGLTFJsonTextureFilter TextureMinFilter = EGLTFJsonTextureFilter::LinearMipmapLinear;
-	const EGLTFJsonTextureFilter TextureMagFilter = EGLTFJsonTextureFilter::Linear;
+	const TextureFilter TextureFilter = Builder.GetBakeFilterForMaterialProperty(Material, PropertyBakeOutput.Property);
+	const EGLTFJsonTextureFilter TextureMinFilter = FGLTFConverterUtility::ConvertMinFilter(TextureFilter);
+	const EGLTFJsonTextureFilter TextureMagFilter = FGLTFConverterUtility::ConvertMagFilter(TextureFilter);
 
 	const FGLTFJsonTextureIndex TextureIndex = FGLTFMaterialUtility::AddTexture(
 		Builder,
