[GLTFExporter] fix: include streaming levels when exporting world

commit f5749d17a6a882449a7494f2b8a630dc969e64fd
Date: Tue, 15 Dec 2020 20:56:31 +0100

    fix: include streaming levels when exporting world
    (cherry picked from commit ef3b08b6807836ede277f65c33c4740ac40d1cd5​)

#skipci

diff --git a/Engine/Plugins/Enterprise/GLTFExporter/MigratedFromGit/Source/GLTFExporter/Private/Builders/GLTFConvertBuilder.cpp b/Engine/Plugins/Enterprise/GLTFExporter/MigratedFromGit/Source/GLTFExporter/Private/Builders/GLTFConvertBuilder.cpp
--- a/Engine/Plugins/Enterprise/GLTFExporter/MigratedFromGit/Source/GLTFExporter/Private/Builders/GLTFConvertBuilder.cpp
+++ b/Engine/Plugins/Enterprise/GLTFExporter/MigratedFromGit/Source/GLTFExporter/Private/Builders/GLTFConvertBuilder.cpp
@@ -337,16 +337,6 @@ FGLTFJsonNodeIndex FGLTFConvertBuilder::GetOrAddNode(FGLTFJsonNodeIndex RootNode
 	return SkeletalBoneConverter.GetOrAdd(RootNode, SkeletalMesh, BoneIndex);
 }
 
-FGLTFJsonSceneIndex FGLTFConvertBuilder::GetOrAddScene(const ULevel* Level)
-{
-	if (Level == nullptr)
-	{
-		return FGLTFJsonSceneIndex(INDEX_NONE);
-	}
-
-	return SceneConverter.GetOrAdd(Level);
-}
-
 FGLTFJsonSceneIndex FGLTFConvertBuilder::GetOrAddScene(const UWorld* World)
 {
 	if (World == nullptr)
@@ -354,7 +344,7 @@ FGLTFJsonSceneIndex FGLTFConvertBuilder::GetOrAddScene(const UWorld* World)
 		return FGLTFJsonSceneIndex(INDEX_NONE);
 	}
 
-	return GetOrAddScene(World->PersistentLevel);
+	return SceneConverter.GetOrAdd(World);
 }
 
 FGLTFJsonCameraIndex FGLTFConvertBuilder::GetOrAddCamera(const UCameraComponent* CameraComponent)
diff --git a/Engine/Plugins/Enterprise/GLTFExporter/MigratedFromGit/Source/GLTFExporter/Private/Builders/GLTFConvertBuilder.h b/Engine/Plugins/Enterprise/GLTFExporter/MigratedFromGit/Source/GLTFExporter/Private/Builders/GLTFConvertBuilder.h
--- a/Engine/Plugins/Enterprise/GLTFExporter/MigratedFromGit/Source/GLTFExporter/Private/Builders/GLTFConvertBuilder.h
+++ b/Engine/Plugins/Enterprise/GLTFExporter/MigratedFromGit/Source/GLTFExporter/Private/Builders/GLTFConvertBuilder.h
@@ -71,7 +71,6 @@ public:
 	FGLTFJsonNodeIndex GetOrAddNode(FGLTFJsonNodeIndex RootNode, const UStaticMesh* StaticMesh, FName SocketName);
 	FGLTFJsonNodeIndex GetOrAddNode(FGLTFJsonNodeIndex RootNode, const USkeletalMesh* SkeletalMesh, FName SocketName);
 	FGLTFJsonNodeIndex GetOrAddNode(FGLTFJsonNodeIndex RootNode, const USkeletalMesh* SkeletalMesh, int32 BoneIndex);
-	FGLTFJsonSceneIndex GetOrAddScene(const ULevel* Level);
 	FGLTFJsonSceneIndex GetOrAddScene(const UWorld* World);
 
 	FGLTFJsonCameraIndex GetOrAddCamera(const UCameraComponent* CameraComponent);
diff --git a/Engine/Plugins/Enterprise/GLTFExporter/MigratedFromGit/Source/GLTFExporter/Private/Converters/GLTFSceneConverters.cpp b/Engine/Plugins/Enterprise/GLTFExporter/MigratedFromGit/Source/GLTFExporter/Private/Converters/GLTFSceneConverters.cpp
--- a/Engine/Plugins/Enterprise/GLTFExporter/MigratedFromGit/Source/GLTFExporter/Private/Converters/GLTFSceneConverters.cpp
+++ b/Engine/Plugins/Enterprise/GLTFExporter/MigratedFromGit/Source/GLTFExporter/Private/Converters/GLTFSceneConverters.cpp
@@ -6,44 +6,45 @@
 #include "Converters/GLTFActorUtility.h"
 #include "LevelVariantSetsActor.h"
 
-FGLTFJsonSceneIndex FGLTFSceneConverter::Convert(const ULevel* Level)
+FGLTFJsonSceneIndex FGLTFSceneConverter::Convert(const UWorld* World)
 {
 	FGLTFJsonScene Scene;
+	World->GetName(Scene.Name);
 
-	if (const UWorld* World = Level->GetWorld())
+	for (int32 Index = 0; Index < World->GetNumLevels(); ++Index)
 	{
-		World->GetName(Scene.Name);
-	}
-	else
-	{
-		Level->GetName(Scene.Name);
-	}
+		ULevel* Level = World->GetLevel(Index);
+		if (Level == nullptr)
+		{
+			continue;
+		}
 
-	// TODO: export Level->Model ?
+		// TODO: export Level->Model ?
 
-	for (const AActor* Actor : Level->Actors)
-	{
-		// TODO: should a LevelVariantSet be exported even if not selected for export?
-		if (const ALevelVariantSetsActor *LevelVariantSetsActor = Cast<ALevelVariantSetsActor>(Actor))
+		for (const AActor* Actor : Level->Actors)
 		{
-			if (Builder.ExportOptions->bExportVariantSets)
+			// TODO: should a LevelVariantSet be exported even if not selected for export?
+			if (const ALevelVariantSetsActor *LevelVariantSetsActor = Cast<ALevelVariantSetsActor>(Actor))
 			{
-				if (const ULevelVariantSets* LevelVariantSets = const_cast<ALevelVariantSetsActor*>(LevelVariantSetsActor)->GetLevelVariantSets(true))
+				if (Builder.ExportOptions->bExportVariantSets)
 				{
-					const FGLTFJsonVariationIndex VariationIndex = Builder.GetOrAddVariation(LevelVariantSets);
-					if (VariationIndex != INDEX_NONE)
+					if (const ULevelVariantSets* LevelVariantSets = const_cast<ALevelVariantSetsActor*>(LevelVariantSetsActor)->GetLevelVariantSets(true))
 					{
-						Scene.Variations.Add(VariationIndex);
+						const FGLTFJsonVariationIndex VariationIndex = Builder.GetOrAddVariation(LevelVariantSets);
+						if (VariationIndex != INDEX_NONE)
+						{
+							Scene.Variations.Add(VariationIndex);
+						}
 					}
 				}
 			}
-		}
 
-		const FGLTFJsonNodeIndex NodeIndex = Builder.GetOrAddNode(Actor);
-		if (NodeIndex != INDEX_NONE && FGLTFActorUtility::IsRootActor(Actor, Builder.bSelectedActorsOnly))
-		{
-			// TODO: to avoid having to add irrelevant actors/components let GLTFComponentConverter decide and add root nodes to scene. This change may require node converters to support cyclic calls.
-			Scene.Nodes.Add(NodeIndex);
+			const FGLTFJsonNodeIndex NodeIndex = Builder.GetOrAddNode(Actor);
+			if (NodeIndex != INDEX_NONE && FGLTFActorUtility::IsRootActor(Actor, Builder.bSelectedActorsOnly))
+			{
+				// TODO: to avoid having to add irrelevant actors/components let GLTFComponentConverter decide and add root nodes to scene. This change may require node converters to support cyclic calls.
+				Scene.Nodes.Add(NodeIndex);
+			}
 		}
 	}
 
diff --git a/Engine/Plugins/Enterprise/GLTFExporter/MigratedFromGit/Source/GLTFExporter/Private/Converters/GLTFSceneConverters.h b/Engine/Plugins/Enterprise/GLTFExporter/MigratedFromGit/Source/GLTFExporter/Private/Converters/GLTFSceneConverters.h
--- a/Engine/Plugins/Enterprise/GLTFExporter/MigratedFromGit/Source/GLTFExporter/Private/Converters/GLTFSceneConverters.h
+++ b/Engine/Plugins/Enterprise/GLTFExporter/MigratedFromGit/Source/GLTFExporter/Private/Converters/GLTFSceneConverters.h
@@ -7,9 +7,9 @@
 #include "Converters/GLTFBuilderContext.h"
 #include "Engine.h"
 
-class FGLTFSceneConverter final : public FGLTFBuilderContext, public TGLTFConverter<FGLTFJsonSceneIndex, const ULevel*>
+class FGLTFSceneConverter final : public FGLTFBuilderContext, public TGLTFConverter<FGLTFJsonSceneIndex, const UWorld*>
 {
 	using FGLTFBuilderContext::FGLTFBuilderContext;
 
-	virtual FGLTFJsonSceneIndex Convert(const ULevel* Level) override;
+	virtual FGLTFJsonSceneIndex Convert(const UWorld* Level) override;
 };
diff --git a/Engine/Plugins/Enterprise/GLTFExporter/MigratedFromGit/Source/GLTFExporter/Private/Exporters/GLTFExporterUtility.cpp b/Engine/Plugins/Enterprise/GLTFExporter/MigratedFromGit/Source/GLTFExporter/Private/Exporters/GLTFExporterUtility.cpp
--- a/Engine/Plugins/Enterprise/GLTFExporter/MigratedFromGit/Source/GLTFExporter/Private/Exporters/GLTFExporterUtility.cpp
+++ b/Engine/Plugins/Enterprise/GLTFExporter/MigratedFromGit/Source/GLTFExporter/Private/Exporters/GLTFExporterUtility.cpp
@@ -48,12 +48,12 @@ const USkeletalMesh* FGLTFExporterUtility::GetPreviewMesh(const UAnimSequence* A
 	return PreviewMesh;
 }
 
-TArray<ULevel*> FGLTFExporterUtility::GetReferencedLevels(const UObject* Object)
+TArray<UWorld*> FGLTFExporterUtility::GetReferencedWorlds(const UObject* Object)
 {
 	FAssetRegistryModule& AssetRegistryModule = FModuleManager::LoadModuleChecked<FAssetRegistryModule>("AssetRegistry");
 	const FName OuterPathName = *Object->GetOutermost()->GetPathName();
 
-	TArray<ULevel*> ReferencedLevels;
+	TArray<UWorld*> ReferencedWorlds;
 	TArray<FAssetIdentifier> AssetDependencies;
 	AssetRegistryModule.Get().GetDependencies(OuterPathName, AssetDependencies);
 
@@ -65,10 +65,10 @@ TArray<ULevel*> FGLTFExporterUtility::GetReferencedLevels(const UObject* Object)
 			UWorld* World = LoadObject<UWorld>(nullptr, *PackageName);
 			if (World != nullptr)
 			{
-				ReferencedLevels.AddUnique(World->PersistentLevel);
+				ReferencedWorlds.AddUnique(World);
 			}
 		}
 	}
 
-	return ReferencedLevels;
+	return ReferencedWorlds;
 }
diff --git a/Engine/Plugins/Enterprise/GLTFExporter/MigratedFromGit/Source/GLTFExporter/Private/Exporters/GLTFExporterUtility.h b/Engine/Plugins/Enterprise/GLTFExporter/MigratedFromGit/Source/GLTFExporter/Private/Exporters/GLTFExporterUtility.h
--- a/Engine/Plugins/Enterprise/GLTFExporter/MigratedFromGit/Source/GLTFExporter/Private/Exporters/GLTFExporterUtility.h
+++ b/Engine/Plugins/Enterprise/GLTFExporter/MigratedFromGit/Source/GLTFExporter/Private/Exporters/GLTFExporterUtility.h
@@ -12,5 +12,5 @@ struct FGLTFExporterUtility
 	static const UStaticMesh* GetPreviewMesh(const UMaterialInterface* Material);
 	static const USkeletalMesh* GetPreviewMesh(const UAnimSequence* AnimSequence);
 
-	static TArray<ULevel*> GetReferencedLevels(const UObject* LevelSequence);
+	static TArray<UWorld*> GetReferencedWorlds(const UObject* LevelSequence);
 };
diff --git a/Engine/Plugins/Enterprise/GLTFExporter/MigratedFromGit/Source/GLTFExporter/Private/Exporters/GLTFLevelSequenceExporter.cpp b/Engine/Plugins/Enterprise/GLTFExporter/MigratedFromGit/Source/GLTFExporter/Private/Exporters/GLTFLevelSequenceExporter.cpp
--- a/Engine/Plugins/Enterprise/GLTFExporter/MigratedFromGit/Source/GLTFExporter/Private/Exporters/GLTFLevelSequenceExporter.cpp
+++ b/Engine/Plugins/Enterprise/GLTFExporter/MigratedFromGit/Source/GLTFExporter/Private/Exporters/GLTFLevelSequenceExporter.cpp
@@ -23,8 +23,8 @@ bool UGLTFLevelSequenceExporter::AddObject(FGLTFContainerBuilder& Builder, const
 		return false;
 	}
 
-	TArray<ULevel*> Levels = FGLTFExporterUtility::GetReferencedLevels(LevelSequence);
-	if (Levels.Num() == 0)
+	TArray<UWorld*> Worlds = FGLTFExporterUtility::GetReferencedWorlds(LevelSequence);
+	if (Worlds.Num() == 0)
 	{
 		Builder.AddErrorMessage(
 			FString::Printf(TEXT("Failed to export level sequence %s because no level referenced"),
@@ -32,7 +32,7 @@ bool UGLTFLevelSequenceExporter::AddObject(FGLTFContainerBuilder& Builder, const
 		return false;
 	}
 
-	if (Levels.Num() > 1)
+	if (Worlds.Num() > 1)
 	{
 		Builder.AddErrorMessage(
             FString::Printf(TEXT("Failed to export level sequence %s because more than one level referenced"),
@@ -40,18 +40,18 @@ bool UGLTFLevelSequenceExporter::AddObject(FGLTFContainerBuilder& Builder, const
 		return false;
 	}
 
-	ULevel* Level = Levels[0];
-	const FGLTFJsonSceneIndex SceneIndex = Builder.GetOrAddScene(Level);
+	UWorld* World = Worlds[0];
+	const FGLTFJsonSceneIndex SceneIndex = Builder.GetOrAddScene(World);
 	if (SceneIndex == INDEX_NONE)
 	{
 		Builder.AddErrorMessage(
 			FString::Printf(TEXT("Failed to export level %s for level sequence %s"),
-			*Level->GetName(),
+			*World->GetName(),
 			*LevelSequence->GetName()));
 		return false;
 	}
 
-	const FGLTFJsonAnimationIndex AnimationIndex = Builder.GetOrAddAnimation(Level, LevelSequence);
+	const FGLTFJsonAnimationIndex AnimationIndex = Builder.GetOrAddAnimation(World->PersistentLevel, LevelSequence);
 	if (AnimationIndex == INDEX_NONE)
 	{
 		Builder.AddErrorMessage(
diff --git a/Engine/Plugins/Enterprise/GLTFExporter/MigratedFromGit/Source/GLTFExporter/Private/Exporters/GLTFLevelVariantSetsExporter.cpp b/Engine/Plugins/Enterprise/GLTFExporter/MigratedFromGit/Source/GLTFExporter/Private/Exporters/GLTFLevelVariantSetsExporter.cpp
--- a/Engine/Plugins/Enterprise/GLTFExporter/MigratedFromGit/Source/GLTFExporter/Private/Exporters/GLTFLevelVariantSetsExporter.cpp
+++ b/Engine/Plugins/Enterprise/GLTFExporter/MigratedFromGit/Source/GLTFExporter/Private/Exporters/GLTFLevelVariantSetsExporter.cpp
@@ -23,8 +23,8 @@ bool UGLTFLevelVariantSetsExporter::AddObject(FGLTFContainerBuilder& Builder, co
 		return false;
 	}
 
-	TArray<ULevel*> Levels = FGLTFExporterUtility::GetReferencedLevels(LevelVariantSets);
-	if (Levels.Num() == 0)
+	TArray<UWorld*> Worlds = FGLTFExporterUtility::GetReferencedWorlds(LevelVariantSets);
+	if (Worlds.Num() == 0)
 	{
 		Builder.AddErrorMessage(
 			FString::Printf(TEXT("Failed to export level variant sets %s because no level referenced"),
@@ -32,7 +32,7 @@ bool UGLTFLevelVariantSetsExporter::AddObject(FGLTFContainerBuilder& Builder, co
 		return false;
 	}
 
-	if (Levels.Num() > 1)
+	if (Worlds.Num() > 1)
 	{
 		Builder.AddErrorMessage(
             FString::Printf(TEXT("Failed to export level variant sets %s because more than one level referenced"),
@@ -40,13 +40,13 @@ bool UGLTFLevelVariantSetsExporter::AddObject(FGLTFContainerBuilder& Builder, co
 		return false;
 	}
 
-	ULevel* Level = Levels[0];
-	const FGLTFJsonSceneIndex SceneIndex = Builder.GetOrAddScene(Level);
+	UWorld* World = Worlds[0];
+	const FGLTFJsonSceneIndex SceneIndex = Builder.GetOrAddScene(World);
 	if (SceneIndex == INDEX_NONE)
 	{
 		Builder.AddErrorMessage(
 			FString::Printf(TEXT("Failed to export level %s for level variant sets %s"),
-			*Level->GetName(),
+			*World->GetName(),
 			*LevelVariantSets->GetName()));
 		return false;
 	}
