[GLTFExporter] fix: use active device profile to get texture lod group settings instead of running platform

commit 1661d3c11b1c6234e4792cd3f5749be94116ec88
Date: Sat, 9 Jan 2021 20:47:12 +0100

    fix: use active device profile to get texture lod group settings instead of running platform
    (cherry picked from commit b58bfc4d584ae1f5fab2f89e5cc6893091b79fa1​)

#skipci

diff --git a/Engine/Plugins/Enterprise/GLTFExporter/MigratedFromGit/Source/GLTFExporter/Private/Converters/GLTFTextureUtility.cpp b/Engine/Plugins/Enterprise/GLTFExporter/MigratedFromGit/Source/GLTFExporter/Private/Converters/GLTFTextureUtility.cpp
--- a/Engine/Plugins/Enterprise/GLTFExporter/MigratedFromGit/Source/GLTFExporter/Private/Converters/GLTFTextureUtility.cpp
+++ b/Engine/Plugins/Enterprise/GLTFExporter/MigratedFromGit/Source/GLTFExporter/Private/Converters/GLTFTextureUtility.cpp
@@ -68,12 +68,10 @@ float FGLTFTextureUtility::GetCubeFaceRotation(ECubeFace CubeFace)
 
 TextureFilter FGLTFTextureUtility::GetDefaultFilter(TextureGroup LODGroup)
 {
-	// TODO: should this be the running platform?
-	const ITargetPlatform* RunningPlatform = GetTargetPlatformManagerRef().GetRunningTargetPlatform();
-	const UTextureLODSettings& TextureLODSettings = RunningPlatform->GetTextureLODSettings();
-	const FTextureLODGroup& TextureLODGroup = TextureLODSettings.GetTextureLODGroup(LODGroup);
+	const UTextureLODSettings* TextureLODSettings = UDeviceProfileManager::Get().GetActiveProfile()->GetTextureLODSettings();
+	const ETextureSamplerFilter Filter = TextureLODSettings->GetTextureLODGroup(LODGroup).Filter;
 
-	switch (TextureLODGroup.Filter)
+	switch (Filter)
 	{
 		case ETextureSamplerFilter::Point:             return TF_Nearest;
 		case ETextureSamplerFilter::Bilinear:          return TF_Bilinear;
