[GLTFExporter] fix: remove UI options and window in modified baking module to fix some build issues in UE 4.26

commit 174f1389d475d8c0ffba0d954d5ddea0df9dc386
Date: Tue, 12 Jan 2021 05:27:58 +0100

    fix: remove UI options and window in modified baking module to fix some build issues in UE 4.26
    (cherry picked from commit b6b8b08760ca7bfe2852ea2e7e856290205a5e70​)

#skipci

diff --git a/Engine/Plugins/Enterprise/GLTFExporter/MigratedFromGit/Source/GLTFMaterialBaking/Private/MaterialBakingModule.cpp b/Engine/Plugins/Enterprise/GLTFExporter/MigratedFromGit/Source/GLTFMaterialBaking/Private/MaterialBakingModule.cpp
--- a/Engine/Plugins/Enterprise/GLTFExporter/MigratedFromGit/Source/GLTFMaterialBaking/Private/MaterialBakingModule.cpp
+++ b/Engine/Plugins/Enterprise/GLTFExporter/MigratedFromGit/Source/GLTFMaterialBaking/Private/MaterialBakingModule.cpp
@@ -5,10 +5,6 @@
 #include "Engine/TextureRenderTarget2D.h"
 #include "ExportMaterialProxy.h"
 #include "Interfaces/IMainFrameModule.h"
-#include "MaterialOptionsWindow.h"
-#include "MaterialOptions.h"
-#include "PropertyEditorModule.h"
-#include "MaterialOptionsCustomization.h"
 #include "UObject/UObjectGlobals.h"
 #include "MaterialBakingStructures.h"
 #include "Framework/Application/SlateApplication.h"
@@ -263,10 +259,6 @@ void FMaterialBakingModule::StartupModule()
 	PerPropertyFormat.Add(MP_ShadingModel, PF_B8G8R8A8);
 	PerPropertyFormat.Add(TEXT("ClearCoatBottomNormal"), PF_B8G8R8A8);
 
-	// Register property customization
-	FPropertyEditorModule& Module = FModuleManager::Get().LoadModuleChecked<FPropertyEditorModule>("PropertyEditor");
-	Module.RegisterCustomPropertyTypeLayout(TEXT("PropertyEntry"), FOnGetPropertyTypeCustomizationInstance::CreateStatic(&FPropertyEntryCustomization::MakeInstance));
-	
 	// Register callback for modified objects
 	FCoreUObjectDelegates::OnObjectModified.AddRaw(this, &FMaterialBakingModule::OnObjectModified);
 }
@@ -698,34 +690,6 @@ void FMaterialBakingModule::BakeMaterials(const TArray<FMaterialDataEx*>& Materi
 	}
 }
 
-bool FMaterialBakingModule::SetupMaterialBakeSettings(TArray<TWeakObjectPtr<UObject>>& OptionObjects, int32 NumLODs)
-{
-	TSharedRef<SWindow> Window = SNew(SWindow)
-		.Title(LOCTEXT("WindowTitle", "Material Baking Options"))
-		.SizingRule(ESizingRule::Autosized);
-
-	TSharedPtr<SMaterialOptions> Options;
-
-	Window->SetContent
-	(
-		SAssignNew(Options, SMaterialOptions)
-		.WidgetWindow(Window)
-		.NumLODs(NumLODs)
-		.SettingsObjects(OptionObjects)
-	);
-
-	TSharedPtr<SWindow> ParentWindow;
-	if (FModuleManager::Get().IsModuleLoaded("MainFrame"))
-	{
-		IMainFrameModule& MainFrame = FModuleManager::LoadModuleChecked<IMainFrameModule>("MainFrame");
-		ParentWindow = MainFrame.GetParentWindow();
-		FSlateApplication::Get().AddModalWindow(Window, ParentWindow, false);
-		return !Options->WasUserCancelled();
-	}
-
-	return false;
-}
-
 void FMaterialBakingModule::CleanupMaterialProxies()
 {
 	for (auto Iterator : MaterialProxyPool)
diff --git a/Engine/Plugins/Enterprise/GLTFExporter/MigratedFromGit/Source/GLTFMaterialBaking/Private/MaterialBakingModule.h b/Engine/Plugins/Enterprise/GLTFExporter/MigratedFromGit/Source/GLTFMaterialBaking/Private/MaterialBakingModule.h
--- a/Engine/Plugins/Enterprise/GLTFExporter/MigratedFromGit/Source/GLTFMaterialBaking/Private/MaterialBakingModule.h
+++ b/Engine/Plugins/Enterprise/GLTFExporter/MigratedFromGit/Source/GLTFMaterialBaking/Private/MaterialBakingModule.h
@@ -10,7 +10,6 @@
 #include "MaterialPropertyEx.h"
 
 class UTextureRenderTarget2D;
-class UGLTFMaterialOptions;
 class UMaterialInterface;
 
 class FExportMaterialProxy;
@@ -34,9 +33,6 @@ public:
 	/** Bakes out material properties according to extended MaterialSettings using MeshSettings and stores the output in Output */
 	virtual void BakeMaterials(const TArray<FMaterialDataEx*>& MaterialSettings, const TArray<FMeshData*>& MeshSettings, TArray<FBakeOutputEx>& Output) override;
 
-	/** Promps a slate window to allow the user to populate specific material baking settings used while baking out materials */
-	virtual bool SetupMaterialBakeSettings(TArray<TWeakObjectPtr<UObject>>& OptionObjects, int32 NumLODs) override;
-
 protected:
 	/* Creates and adds or reuses a RenderTarget from the pool */
 	UTextureRenderTarget2D* CreateRenderTarget(bool bInForceLinearGamma, EPixelFormat InPixelFormat, const FIntPoint& InTargetSize);
diff --git a/Engine/Plugins/Enterprise/GLTFExporter/MigratedFromGit/Source/GLTFMaterialBaking/Private/MaterialOptionsCustomization.cpp b/Engine/Plugins/Enterprise/GLTFExporter/MigratedFromGit/Source/GLTFMaterialBaking/Private/MaterialOptionsCustomization.cpp
deleted file mode 100644
diff --git a/Engine/Plugins/Enterprise/GLTFExporter/MigratedFromGit/Source/GLTFMaterialBaking/Private/MaterialOptionsCustomization.h b/Engine/Plugins/Enterprise/GLTFExporter/MigratedFromGit/Source/GLTFMaterialBaking/Private/MaterialOptionsCustomization.h
deleted file mode 100644
diff --git a/Engine/Plugins/Enterprise/GLTFExporter/MigratedFromGit/Source/GLTFMaterialBaking/Private/MaterialOptionsWindow.cpp b/Engine/Plugins/Enterprise/GLTFExporter/MigratedFromGit/Source/GLTFMaterialBaking/Private/MaterialOptionsWindow.cpp
deleted file mode 100644
diff --git a/Engine/Plugins/Enterprise/GLTFExporter/MigratedFromGit/Source/GLTFMaterialBaking/Private/MaterialOptionsWindow.h b/Engine/Plugins/Enterprise/GLTFExporter/MigratedFromGit/Source/GLTFMaterialBaking/Private/MaterialOptionsWindow.h
deleted file mode 100644
diff --git a/Engine/Plugins/Enterprise/GLTFExporter/MigratedFromGit/Source/GLTFMaterialBaking/Public/IMaterialBakingModule.h b/Engine/Plugins/Enterprise/GLTFExporter/MigratedFromGit/Source/GLTFMaterialBaking/Public/IMaterialBakingModule.h
--- a/Engine/Plugins/Enterprise/GLTFExporter/MigratedFromGit/Source/GLTFMaterialBaking/Public/IMaterialBakingModule.h
+++ b/Engine/Plugins/Enterprise/GLTFExporter/MigratedFromGit/Source/GLTFMaterialBaking/Public/IMaterialBakingModule.h
@@ -8,7 +8,6 @@
 #include "SceneTypes.h"
 
 class UTextureRenderTarget2D;
-class UGLTFMaterialOptions;
 class UMaterialInterface;
 
 class FExportMaterialProxy;
@@ -28,7 +27,4 @@ public:
 
 	/** Bakes out material properties according to extended MaterialSettings using MeshSettings and stores the output in Output */
 	virtual void BakeMaterials(const TArray<FMaterialDataEx*>& MaterialSettings, const TArray<FMeshData*>& MeshSettings, TArray<FBakeOutputEx>& Output) = 0;
-
-	/** Promps a slate window to allow the user to populate specific material baking settings used while baking out materials */
-	virtual  bool SetupMaterialBakeSettings(TArray<TWeakObjectPtr<UObject>>& OptionObjects, int32 NumLODs) = 0;
 };
diff --git a/Engine/Plugins/Enterprise/GLTFExporter/MigratedFromGit/Source/GLTFMaterialBaking/Public/MaterialOptions.h b/Engine/Plugins/Enterprise/GLTFExporter/MigratedFromGit/Source/GLTFMaterialBaking/Public/MaterialOptions.h
deleted file mode 100644
