[GLTFExporter] fix: add GLTFIndexArray to fix some build issues in UE 4.26

commit e2d4a3e2354eac8c22ec229a83b37f9595f900fc
Date: Tue, 12 Jan 2021 05:33:49 +0100

    fix: add GLTFIndexArray to fix some build issues in UE 4.26
    (cherry picked from commit ff47f87ba045ce3ac99ccbecf187b63dfd16c02e​)

#skipci

diff --git a/Engine/Plugins/Enterprise/GLTFExporter/MigratedFromGit/Source/GLTFExporter/Private/Builders/GLTFConvertBuilder.cpp b/Engine/Plugins/Enterprise/GLTFExporter/MigratedFromGit/Source/GLTFExporter/Private/Builders/GLTFConvertBuilder.cpp
--- a/Engine/Plugins/Enterprise/GLTFExporter/MigratedFromGit/Source/GLTFExporter/Private/Builders/GLTFConvertBuilder.cpp
+++ b/Engine/Plugins/Enterprise/GLTFExporter/MigratedFromGit/Source/GLTFExporter/Private/Builders/GLTFConvertBuilder.cpp
@@ -140,7 +140,7 @@ FGLTFJsonMeshIndex FGLTFConvertBuilder::GetOrAddMesh(const USkeletalMeshComponen
 	return SkeletalMeshConverter.GetOrAdd(SkeletalMesh, SkeletalMeshComponent, Materials, LODIndex);
 }
 
-FGLTFJsonMaterialIndex FGLTFConvertBuilder::GetOrAddMaterial(const UMaterialInterface* Material, const FGLTFMeshData* MeshData, const TArray<int32>& SectionIndices)
+FGLTFJsonMaterialIndex FGLTFConvertBuilder::GetOrAddMaterial(const UMaterialInterface* Material, const FGLTFMeshData* MeshData, const FGLTFIndexArray& SectionIndices)
 {
 	if (Material == nullptr)
 	{
diff --git a/Engine/Plugins/Enterprise/GLTFExporter/MigratedFromGit/Source/GLTFExporter/Private/Builders/GLTFConvertBuilder.h b/Engine/Plugins/Enterprise/GLTFExporter/MigratedFromGit/Source/GLTFExporter/Private/Builders/GLTFConvertBuilder.h
--- a/Engine/Plugins/Enterprise/GLTFExporter/MigratedFromGit/Source/GLTFExporter/Private/Builders/GLTFConvertBuilder.h
+++ b/Engine/Plugins/Enterprise/GLTFExporter/MigratedFromGit/Source/GLTFExporter/Private/Builders/GLTFConvertBuilder.h
@@ -49,7 +49,7 @@ public:
 	FGLTFJsonMeshIndex GetOrAddMesh(const USkeletalMesh* SkeletalMesh, const FGLTFMaterialArray& Materials = {}, int32 LODIndex = -1);
 	FGLTFJsonMeshIndex GetOrAddMesh(const USkeletalMeshComponent* SkeletalMeshComponent, const FGLTFMaterialArray& Materials = {}, int32 LODIndex = -1);
 
-	FGLTFJsonMaterialIndex GetOrAddMaterial(const UMaterialInterface* Material, const FGLTFMeshData* MeshData = nullptr, const TArray<int32>& SectionIndices = {});
+	FGLTFJsonMaterialIndex GetOrAddMaterial(const UMaterialInterface* Material, const FGLTFMeshData* MeshData = nullptr, const FGLTFIndexArray& SectionIndices = {});
 
 	FGLTFJsonSamplerIndex GetOrAddSampler(const UTexture* Texture);
 	FGLTFJsonTextureIndex GetOrAddTexture(const UTexture2D* Texture);
diff --git a/Engine/Plugins/Enterprise/GLTFExporter/MigratedFromGit/Source/GLTFExporter/Private/Converters/GLTFIndexArray.h b/Engine/Plugins/Enterprise/GLTFExporter/MigratedFromGit/Source/GLTFExporter/Private/Converters/GLTFIndexArray.h
new file mode 100644
--- /dev/null
+++ b/Engine/Plugins/Enterprise/GLTFExporter/MigratedFromGit/Source/GLTFExporter/Private/Converters/GLTFIndexArray.h
@@ -0,0 +1,19 @@
+// Copyright Epic Games, Inc. All Rights Reserved.
+
+#pragma once
+
+#include "CoreMinimal.h"
+
+typedef TArray<int32> FGLTFIndexArray;
+
+#if !(ENGINE_MAJOR_VERSION == 4 && ENGINE_MINOR_VERSION <= 25)
+inline uint32 GetTypeHash(const FGLTFIndexArray& IndexArray)
+{
+	uint32 Hash = GetTypeHash(IndexArray.Num());
+	for (int32 Index : IndexArray)
+	{
+		Hash = HashCombine(Hash, GetTypeHash(Index));
+	}
+	return Hash;
+}
+#endif
diff --git a/Engine/Plugins/Enterprise/GLTFExporter/MigratedFromGit/Source/GLTFExporter/Private/Converters/GLTFMaterialConverters.cpp b/Engine/Plugins/Enterprise/GLTFExporter/MigratedFromGit/Source/GLTFExporter/Private/Converters/GLTFMaterialConverters.cpp
--- a/Engine/Plugins/Enterprise/GLTFExporter/MigratedFromGit/Source/GLTFExporter/Private/Converters/GLTFMaterialConverters.cpp
+++ b/Engine/Plugins/Enterprise/GLTFExporter/MigratedFromGit/Source/GLTFExporter/Private/Converters/GLTFMaterialConverters.cpp
@@ -5,7 +5,7 @@
 #include "Builders/GLTFConvertBuilder.h"
 #include "Tasks/GLTFMaterialTasks.h"
 
-void FGLTFMaterialConverter::Sanitize(const UMaterialInterface*& Material, const FGLTFMeshData*& MeshData, TArray<int32>& SectionIndices)
+void FGLTFMaterialConverter::Sanitize(const UMaterialInterface*& Material, const FGLTFMeshData*& MeshData, FGLTFIndexArray& SectionIndices)
 {
 	if (MeshData == nullptr ||
 		Builder.ExportOptions->BakeMaterialInputs != EGLTFMaterialBakeMode::UseMeshData ||
@@ -39,7 +39,7 @@ void FGLTFMaterialConverter::Sanitize(const UMaterialInterface*& Material, const
 	}
 }
 
-FGLTFJsonMaterialIndex FGLTFMaterialConverter::Convert(const UMaterialInterface* Material, const FGLTFMeshData* MeshData, TArray<int32> SectionIndices)
+FGLTFJsonMaterialIndex FGLTFMaterialConverter::Convert(const UMaterialInterface* Material, const FGLTFMeshData* MeshData, FGLTFIndexArray SectionIndices)
 {
 	if (Material == FGLTFMaterialUtility::GetDefault())
 	{
diff --git a/Engine/Plugins/Enterprise/GLTFExporter/MigratedFromGit/Source/GLTFExporter/Private/Converters/GLTFMaterialConverters.h b/Engine/Plugins/Enterprise/GLTFExporter/MigratedFromGit/Source/GLTFExporter/Private/Converters/GLTFMaterialConverters.h
--- a/Engine/Plugins/Enterprise/GLTFExporter/MigratedFromGit/Source/GLTFExporter/Private/Converters/GLTFMaterialConverters.h
+++ b/Engine/Plugins/Enterprise/GLTFExporter/MigratedFromGit/Source/GLTFExporter/Private/Converters/GLTFMaterialConverters.h
@@ -6,17 +6,18 @@
 #include "Converters/GLTFConverter.h"
 #include "Converters/GLTFBuilderContext.h"
 #include "Converters/GLTFMeshData.h"
+#include "Converters/GLTFIndexArray.h"
 #include "Converters/GLTFUVOverlapChecker.h"
 #include "Converters/GLTFUVDegenerateChecker.h"
 #include "Engine.h"
 
-class FGLTFMaterialConverter final : public FGLTFBuilderContext, public TGLTFConverter<FGLTFJsonMaterialIndex, const UMaterialInterface*, const FGLTFMeshData*, TArray<int32>>
+class FGLTFMaterialConverter final : public FGLTFBuilderContext, public TGLTFConverter<FGLTFJsonMaterialIndex, const UMaterialInterface*, const FGLTFMeshData*, FGLTFIndexArray>
 {
 	using FGLTFBuilderContext::FGLTFBuilderContext;
 
-	virtual void Sanitize(const UMaterialInterface*& Material, const FGLTFMeshData*& MeshData, TArray<int32>& SectionIndices) override;
+	virtual void Sanitize(const UMaterialInterface*& Material, const FGLTFMeshData*& MeshData, FGLTFIndexArray& SectionIndices) override;
 
-	virtual FGLTFJsonMaterialIndex Convert(const UMaterialInterface* Material, const FGLTFMeshData* MeshData, TArray<int32> SectionIndices) override;
+	virtual FGLTFJsonMaterialIndex Convert(const UMaterialInterface* Material, const FGLTFMeshData* MeshData, FGLTFIndexArray SectionIndices) override;
 
 	FGLTFUVOverlapChecker UVOverlapChecker;
 	FGLTFUVDegenerateChecker UVDegenerateChecker;
diff --git a/Engine/Plugins/Enterprise/GLTFExporter/MigratedFromGit/Source/GLTFExporter/Private/Converters/GLTFMaterialUtility.cpp b/Engine/Plugins/Enterprise/GLTFExporter/MigratedFromGit/Source/GLTFExporter/Private/Converters/GLTFMaterialUtility.cpp
--- a/Engine/Plugins/Enterprise/GLTFExporter/MigratedFromGit/Source/GLTFExporter/Private/Converters/GLTFMaterialUtility.cpp
+++ b/Engine/Plugins/Enterprise/GLTFExporter/MigratedFromGit/Source/GLTFExporter/Private/Converters/GLTFMaterialUtility.cpp
@@ -208,7 +208,7 @@ bool FGLTFMaterialUtility::CombineTextures(TArray<FColor>& OutPixels, const TArr
 	return bReadSuccessful;
 }
 
-FGLTFPropertyBakeOutput FGLTFMaterialUtility::BakeMaterialProperty(const FIntPoint& OutputSize, const FMaterialPropertyEx& Property, const UMaterialInterface* Material, int32 TexCoord, const FMeshDescription* MeshDescription, const TArray<int32>& MeshSectionIndices, bool bCopyAlphaFromRedChannel)
+FGLTFPropertyBakeOutput FGLTFMaterialUtility::BakeMaterialProperty(const FIntPoint& OutputSize, const FMaterialPropertyEx& Property, const UMaterialInterface* Material, int32 TexCoord, const FMeshDescription* MeshDescription, const FGLTFIndexArray& MeshSectionIndices, bool bCopyAlphaFromRedChannel)
 {
 	FMeshData MeshSet;
 	MeshSet.TextureCoordinateBox = { { 0.0f, 0.0f }, { 1.0f, 1.0f } };
@@ -353,7 +353,7 @@ bool FGLTFMaterialUtility::TryGetTextureCoordinateIndex(const UMaterialExpressio
 	return false;
 }
 
-void FGLTFMaterialUtility::GetAllTextureCoordinateIndices(const UMaterialInterface* InMaterial, const FMaterialPropertyEx& InProperty, TArray<int32>& OutTexCoords)
+void FGLTFMaterialUtility::GetAllTextureCoordinateIndices(const UMaterialInterface* InMaterial, const FMaterialPropertyEx& InProperty, FGLTFIndexArray& OutTexCoords)
 {
 	FGLTFMaterialAnalysis Analysis;
 	AnalyzeMaterialProperty(InMaterial, InProperty, Analysis);
diff --git a/Engine/Plugins/Enterprise/GLTFExporter/MigratedFromGit/Source/GLTFExporter/Private/Converters/GLTFMaterialUtility.h b/Engine/Plugins/Enterprise/GLTFExporter/MigratedFromGit/Source/GLTFExporter/Private/Converters/GLTFMaterialUtility.h
--- a/Engine/Plugins/Enterprise/GLTFExporter/MigratedFromGit/Source/GLTFExporter/Private/Converters/GLTFMaterialUtility.h
+++ b/Engine/Plugins/Enterprise/GLTFExporter/MigratedFromGit/Source/GLTFExporter/Private/Converters/GLTFMaterialUtility.h
@@ -57,7 +57,7 @@ struct FGLTFMaterialUtility
 
 	static UTexture2D* CreateTransientTexture(const FGLTFPropertyBakeOutput& PropertyBakeOutput, bool bUseSRGB = false);
 	static bool CombineTextures(TArray<FColor>& OutPixels, const TArray<FGLTFTextureCombineSource>& Sources, const FIntPoint& OutputSize, EPixelFormat OutputPixelFormat);
-	static FGLTFPropertyBakeOutput BakeMaterialProperty(const FIntPoint& OutputSize, const FMaterialPropertyEx& Property, const UMaterialInterface* Material, int32 TexCoord, const FMeshDescription* MeshDescription = nullptr, const TArray<int32>& MeshSectionIndices = {}, bool bCopyAlphaFromRedChannel = false);
+	static FGLTFPropertyBakeOutput BakeMaterialProperty(const FIntPoint& OutputSize, const FMaterialPropertyEx& Property, const UMaterialInterface* Material, int32 TexCoord, const FMeshDescription* MeshDescription = nullptr, const FGLTFIndexArray& MeshSectionIndices = {}, bool bCopyAlphaFromRedChannel = false);
 
 	static FGLTFJsonTextureIndex AddCombinedTexture(FGLTFConvertBuilder& Builder, const TArray<FGLTFTextureCombineSource>& CombineSources, const FIntPoint& TextureSize, bool bIgnoreAlpha, const FString& TextureName, EGLTFJsonTextureFilter MinFilter, EGLTFJsonTextureFilter MagFilter, EGLTFJsonTextureWrap WrapS, EGLTFJsonTextureWrap WrapT);
 	static FGLTFJsonTextureIndex AddTexture(FGLTFConvertBuilder& Builder, const TArray<FColor>& Pixels, const FIntPoint& TextureSize, bool bIgnoreAlpha, bool bIsNormalMap, const FString& TextureName, EGLTFJsonTextureFilter MinFilter, EGLTFJsonTextureFilter MagFilter, EGLTFJsonTextureWrap WrapS, EGLTFJsonTextureWrap WrapT);
@@ -68,7 +68,7 @@ struct FGLTFMaterialUtility
 	static uint32 GetMaskComponentCount(const FExpressionInput& ExpressionInput);
 
 	static bool TryGetTextureCoordinateIndex(const UMaterialExpressionTextureSample* TextureSampler, int32& TexCoord, FGLTFJsonTextureTransform& Transform);
-	static void GetAllTextureCoordinateIndices(const UMaterialInterface* InMaterial, const FMaterialPropertyEx& InProperty, TArray<int32>& OutTexCoords);
+	static void GetAllTextureCoordinateIndices(const UMaterialInterface* InMaterial, const FMaterialPropertyEx& InProperty, FGLTFIndexArray& OutTexCoords);
 
 	static FMaterialShadingModelField EvaluateShadingModelExpression(const UMaterialInterface* Material);
 	static EMaterialShadingModel GetRichestShadingModel(const FMaterialShadingModelField& ShadingModels);
diff --git a/Engine/Plugins/Enterprise/GLTFExporter/MigratedFromGit/Source/GLTFExporter/Private/Converters/GLTFMeshSection.cpp b/Engine/Plugins/Enterprise/GLTFExporter/MigratedFromGit/Source/GLTFExporter/Private/Converters/GLTFMeshSection.cpp
--- a/Engine/Plugins/Enterprise/GLTFExporter/MigratedFromGit/Source/GLTFExporter/Private/Converters/GLTFMeshSection.cpp
+++ b/Engine/Plugins/Enterprise/GLTFExporter/MigratedFromGit/Source/GLTFExporter/Private/Converters/GLTFMeshSection.cpp
@@ -5,7 +5,7 @@
 #include "Rendering/SkeletalMeshRenderData.h"
 #include "Algo/MaxElement.h"
 
-FGLTFMeshSection::FGLTFMeshSection(const FStaticMeshLODResources* MeshLOD, const TArray<int32>& SectionIndices)
+FGLTFMeshSection::FGLTFMeshSection(const FStaticMeshLODResources* MeshLOD, const FGLTFIndexArray& SectionIndices)
 {
 	const FStaticMeshLODResources::FStaticMeshSectionArray& Sections = MeshLOD->Sections;
 
@@ -52,7 +52,7 @@ FGLTFMeshSection::FGLTFMeshSection(const FStaticMeshLODResources* MeshLOD, const
 	MaxBoneIndex = 0;
 }
 
-FGLTFMeshSection::FGLTFMeshSection(const FSkeletalMeshLODRenderData* MeshLOD, const TArray<int32>& SectionIndices)
+FGLTFMeshSection::FGLTFMeshSection(const FSkeletalMeshLODRenderData* MeshLOD, const FGLTFIndexArray& SectionIndices)
 	: MaxBoneIndex(0)
 {
 	const TArray<FSkelMeshRenderSection>& Sections = MeshLOD->RenderSections;
diff --git a/Engine/Plugins/Enterprise/GLTFExporter/MigratedFromGit/Source/GLTFExporter/Private/Converters/GLTFMeshSection.h b/Engine/Plugins/Enterprise/GLTFExporter/MigratedFromGit/Source/GLTFExporter/Private/Converters/GLTFMeshSection.h
--- a/Engine/Plugins/Enterprise/GLTFExporter/MigratedFromGit/Source/GLTFExporter/Private/Converters/GLTFMeshSection.h
+++ b/Engine/Plugins/Enterprise/GLTFExporter/MigratedFromGit/Source/GLTFExporter/Private/Converters/GLTFMeshSection.h
@@ -3,11 +3,12 @@
 #pragma once
 
 #include "CoreMinimal.h"
+#include "Converters/GLTFIndexArray.h"
 
 struct FGLTFMeshSection
 {
-	FGLTFMeshSection(const FStaticMeshLODResources* MeshLOD, const TArray<int32>& SectionIndices);
-	FGLTFMeshSection(const FSkeletalMeshLODRenderData* MeshLOD, const TArray<int32>& SectionIndices);
+	FGLTFMeshSection(const FStaticMeshLODResources* MeshLOD, const FGLTFIndexArray& SectionIndices);
+	FGLTFMeshSection(const FSkeletalMeshLODRenderData* MeshLOD, const FGLTFIndexArray& SectionIndices);
 
 	TArray<uint32> IndexMap;
 	TArray<uint32> IndexBuffer;
diff --git a/Engine/Plugins/Enterprise/GLTFExporter/MigratedFromGit/Source/GLTFExporter/Private/Converters/GLTFMeshSectionConverters.h b/Engine/Plugins/Enterprise/GLTFExporter/MigratedFromGit/Source/GLTFExporter/Private/Converters/GLTFMeshSectionConverters.h
--- a/Engine/Plugins/Enterprise/GLTFExporter/MigratedFromGit/Source/GLTFExporter/Private/Converters/GLTFMeshSectionConverters.h
+++ b/Engine/Plugins/Enterprise/GLTFExporter/MigratedFromGit/Source/GLTFExporter/Private/Converters/GLTFMeshSectionConverters.h
@@ -4,13 +4,14 @@
 
 #include "Converters/GLTFConverter.h"
 #include "Converters/GLTFMeshSection.h"
+#include "Converters/GLTFIndexArray.h"
 
 template <typename MeshLODType>
-class TGLTFMeshSectionConverter final : public TGLTFConverter<const FGLTFMeshSection*, const MeshLODType*, TArray<int32>>
+class TGLTFMeshSectionConverter final : public TGLTFConverter<const FGLTFMeshSection*, const MeshLODType*, FGLTFIndexArray>
 {
 	TArray<TUniquePtr<FGLTFMeshSection>> Outputs;
 
-	const FGLTFMeshSection* Convert(const MeshLODType* MeshLOD, TArray<int32> SectionIndices)
+	const FGLTFMeshSection* Convert(const MeshLODType* MeshLOD, FGLTFIndexArray SectionIndices)
 	{
 		return Outputs.Add_GetRef(MakeUnique<FGLTFMeshSection>(MeshLOD, SectionIndices)).Get();
 	}
diff --git a/Engine/Plugins/Enterprise/GLTFExporter/MigratedFromGit/Source/GLTFExporter/Private/Converters/GLTFMeshUtility.cpp b/Engine/Plugins/Enterprise/GLTFExporter/MigratedFromGit/Source/GLTFExporter/Private/Converters/GLTFMeshUtility.cpp
--- a/Engine/Plugins/Enterprise/GLTFExporter/MigratedFromGit/Source/GLTFExporter/Private/Converters/GLTFMeshUtility.cpp
+++ b/Engine/Plugins/Enterprise/GLTFExporter/MigratedFromGit/Source/GLTFExporter/Private/Converters/GLTFMeshUtility.cpp
@@ -4,11 +4,11 @@
 #include "Rendering/SkeletalMeshRenderData.h"
 #include "PlatformInfo.h"
 
-TArray<int32> FGLTFMeshUtility::GetSectionIndices(const FStaticMeshLODResources& MeshLOD, int32 MaterialIndex)
+FGLTFIndexArray FGLTFMeshUtility::GetSectionIndices(const FStaticMeshLODResources& MeshLOD, int32 MaterialIndex)
 {
 	const FStaticMeshLODResources::FStaticMeshSectionArray& Sections = MeshLOD.Sections;
 
-	TArray<int32> SectionIndices;
+	FGLTFIndexArray SectionIndices;
 	SectionIndices.Reserve(Sections.Num());
 
 	for (int32 SectionIndex = 0; SectionIndex < Sections.Num(); ++SectionIndex)
@@ -22,11 +22,11 @@ TArray<int32> FGLTFMeshUtility::GetSectionIndices(const FStaticMeshLODResources&
 	return SectionIndices;
 }
 
-TArray<int32> FGLTFMeshUtility::GetSectionIndices(const FSkeletalMeshLODRenderData& MeshLOD, int32 MaterialIndex)
+FGLTFIndexArray FGLTFMeshUtility::GetSectionIndices(const FSkeletalMeshLODRenderData& MeshLOD, int32 MaterialIndex)
 {
 	const TArray<FSkelMeshRenderSection>& Sections = MeshLOD.RenderSections;
 
-	TArray<int32> SectionIndices;
+	FGLTFIndexArray SectionIndices;
 	SectionIndices.Reserve(Sections.Num());
 
 	for (int32 SectionIndex = 0; SectionIndex < Sections.Num(); ++SectionIndex)
diff --git a/Engine/Plugins/Enterprise/GLTFExporter/MigratedFromGit/Source/GLTFExporter/Private/Converters/GLTFMeshUtility.h b/Engine/Plugins/Enterprise/GLTFExporter/MigratedFromGit/Source/GLTFExporter/Private/Converters/GLTFMeshUtility.h
--- a/Engine/Plugins/Enterprise/GLTFExporter/MigratedFromGit/Source/GLTFExporter/Private/Converters/GLTFMeshUtility.h
+++ b/Engine/Plugins/Enterprise/GLTFExporter/MigratedFromGit/Source/GLTFExporter/Private/Converters/GLTFMeshUtility.h
@@ -2,14 +2,15 @@
 
 #pragma once
 
+#include "Converters/GLTFIndexArray.h"
 #include "Engine.h"
 
 class FSkeletalMeshLODRenderData;
 
 struct FGLTFMeshUtility
 {
-	static TArray<int32> GetSectionIndices(const FStaticMeshLODResources& MeshLOD, int32 MaterialIndex);
-	static TArray<int32> GetSectionIndices(const FSkeletalMeshLODRenderData& MeshLOD, int32 MaterialIndex);
+	static FGLTFIndexArray GetSectionIndices(const FStaticMeshLODResources& MeshLOD, int32 MaterialIndex);
+	static FGLTFIndexArray GetSectionIndices(const FSkeletalMeshLODRenderData& MeshLOD, int32 MaterialIndex);
 
 	static int32 GetLOD(const UStaticMesh* StaticMesh, const UStaticMeshComponent* StaticMeshComponent, int32 DefaultLOD);
 	static int32 GetLOD(const USkeletalMesh* SkeletalMesh, const USkeletalMeshComponent* SkeletalMeshComponent, int32 DefaultLOD);
diff --git a/Engine/Plugins/Enterprise/GLTFExporter/MigratedFromGit/Source/GLTFExporter/Private/Converters/GLTFUVDegenerateChecker.cpp b/Engine/Plugins/Enterprise/GLTFExporter/MigratedFromGit/Source/GLTFExporter/Private/Converters/GLTFUVDegenerateChecker.cpp
--- a/Engine/Plugins/Enterprise/GLTFExporter/MigratedFromGit/Source/GLTFExporter/Private/Converters/GLTFUVDegenerateChecker.cpp
+++ b/Engine/Plugins/Enterprise/GLTFExporter/MigratedFromGit/Source/GLTFExporter/Private/Converters/GLTFUVDegenerateChecker.cpp
@@ -4,7 +4,7 @@
 #include "StaticMeshAttributes.h"
 #include "MeshDescription.h"
 
-void FGLTFUVDegenerateChecker::Sanitize(const FMeshDescription*& Description, TArray<int32>& SectionIndices, int32& TexCoord)
+void FGLTFUVDegenerateChecker::Sanitize(const FMeshDescription*& Description, FGLTFIndexArray& SectionIndices, int32& TexCoord)
 {
 	if (Description != nullptr)
 	{
@@ -28,7 +28,7 @@ void FGLTFUVDegenerateChecker::Sanitize(const FMeshDescription*& Description, TA
 	}
 }
 
-float FGLTFUVDegenerateChecker::Convert(const FMeshDescription* Description, TArray<int32> SectionIndices, int32 TexCoord)
+float FGLTFUVDegenerateChecker::Convert(const FMeshDescription* Description, FGLTFIndexArray SectionIndices, int32 TexCoord)
 {
 	if (Description == nullptr)
 	{
diff --git a/Engine/Plugins/Enterprise/GLTFExporter/MigratedFromGit/Source/GLTFExporter/Private/Converters/GLTFUVDegenerateChecker.h b/Engine/Plugins/Enterprise/GLTFExporter/MigratedFromGit/Source/GLTFExporter/Private/Converters/GLTFUVDegenerateChecker.h
--- a/Engine/Plugins/Enterprise/GLTFExporter/MigratedFromGit/Source/GLTFExporter/Private/Converters/GLTFUVDegenerateChecker.h
+++ b/Engine/Plugins/Enterprise/GLTFExporter/MigratedFromGit/Source/GLTFExporter/Private/Converters/GLTFUVDegenerateChecker.h
@@ -3,13 +3,14 @@
 #pragma once
 
 #include "Converters/GLTFConverter.h"
+#include "Converters/GLTFIndexArray.h"
 #include "Engine.h"
 
-class FGLTFUVDegenerateChecker final : public TGLTFConverter<float, const FMeshDescription*, TArray<int32>, int32>
+class FGLTFUVDegenerateChecker final : public TGLTFConverter<float, const FMeshDescription*, FGLTFIndexArray, int32>
 {
-	virtual void Sanitize(const FMeshDescription*& Description, TArray<int32>& SectionIndices, int32& TexCoord) override;
+	virtual void Sanitize(const FMeshDescription*& Description, FGLTFIndexArray& SectionIndices, int32& TexCoord) override;
 
-	virtual float Convert(const FMeshDescription* Description, TArray<int32> SectionIndices, int32 TexCoord) override;
+	virtual float Convert(const FMeshDescription* Description, FGLTFIndexArray SectionIndices, int32 TexCoord) override;
 
 	static bool IsDegenerateTriangle(const TStaticArray<FVector2D, 3>& Points);
 	static bool IsDegenerateTriangle(const TStaticArray<FVector, 3>& Points);
diff --git a/Engine/Plugins/Enterprise/GLTFExporter/MigratedFromGit/Source/GLTFExporter/Private/Converters/GLTFUVOverlapChecker.cpp b/Engine/Plugins/Enterprise/GLTFExporter/MigratedFromGit/Source/GLTFExporter/Private/Converters/GLTFUVOverlapChecker.cpp
--- a/Engine/Plugins/Enterprise/GLTFExporter/MigratedFromGit/Source/GLTFExporter/Private/Converters/GLTFUVOverlapChecker.cpp
+++ b/Engine/Plugins/Enterprise/GLTFExporter/MigratedFromGit/Source/GLTFExporter/Private/Converters/GLTFUVOverlapChecker.cpp
@@ -7,7 +7,7 @@
 #include "IMaterialBakingModule.h"
 #include "MaterialBakingStructures.h"
 
-void FGLTFUVOverlapChecker::Sanitize(const FMeshDescription*& Description, const TArray<int32>& SectionIndices, int32& TexCoord)
+void FGLTFUVOverlapChecker::Sanitize(const FMeshDescription*& Description, FGLTFIndexArray& SectionIndices, int32& TexCoord)
 {
 	if (Description != nullptr)
 	{
@@ -31,7 +31,7 @@ void FGLTFUVOverlapChecker::Sanitize(const FMeshDescription*& Description, const
 	}
 }
 
-float FGLTFUVOverlapChecker::Convert(const FMeshDescription* Description, const TArray<int32> SectionIndices, int32 TexCoord)
+float FGLTFUVOverlapChecker::Convert(const FMeshDescription* Description, FGLTFIndexArray SectionIndices, int32 TexCoord)
 {
 	if (Description == nullptr)
 	{
diff --git a/Engine/Plugins/Enterprise/GLTFExporter/MigratedFromGit/Source/GLTFExporter/Private/Converters/GLTFUVOverlapChecker.h b/Engine/Plugins/Enterprise/GLTFExporter/MigratedFromGit/Source/GLTFExporter/Private/Converters/GLTFUVOverlapChecker.h
--- a/Engine/Plugins/Enterprise/GLTFExporter/MigratedFromGit/Source/GLTFExporter/Private/Converters/GLTFUVOverlapChecker.h
+++ b/Engine/Plugins/Enterprise/GLTFExporter/MigratedFromGit/Source/GLTFExporter/Private/Converters/GLTFUVOverlapChecker.h
@@ -3,13 +3,14 @@
 #pragma once
 
 #include "Converters/GLTFConverter.h"
+#include "Converters/GLTFIndexArray.h"
 #include "Engine.h"
 
-class FGLTFUVOverlapChecker final : public TGLTFConverter<float, const FMeshDescription*, const TArray<int32>, int32>
+class FGLTFUVOverlapChecker final : public TGLTFConverter<float, const FMeshDescription*, FGLTFIndexArray, int32>
 {
-	virtual void Sanitize(const FMeshDescription*& Description, const TArray<int32>& SectionIndices, int32& TexCoord) override;
+	virtual void Sanitize(const FMeshDescription*& Description, FGLTFIndexArray& SectionIndices, int32& TexCoord) override;
 
-	virtual float Convert(const FMeshDescription* Description, const TArray<int32> SectionIndices, int32 TexCoord) override;
+	virtual float Convert(const FMeshDescription* Description, FGLTFIndexArray SectionIndices, int32 TexCoord) override;
 
 	static UMaterialInterface* GetMaterial();
 };
diff --git a/Engine/Plugins/Enterprise/GLTFExporter/MigratedFromGit/Source/GLTFExporter/Private/Converters/GLTFVarationConverters.cpp b/Engine/Plugins/Enterprise/GLTFExporter/MigratedFromGit/Source/GLTFExporter/Private/Converters/GLTFVarationConverters.cpp
--- a/Engine/Plugins/Enterprise/GLTFExporter/MigratedFromGit/Source/GLTFExporter/Private/Converters/GLTFVarationConverters.cpp
+++ b/Engine/Plugins/Enterprise/GLTFExporter/MigratedFromGit/Source/GLTFExporter/Private/Converters/GLTFVarationConverters.cpp
@@ -236,7 +236,7 @@ bool FGLTFVariationConverter::TryParseMaterialPropertyValue(FGLTFJsonVariant& Ou
 	const int32 MaterialIndex = CapturedPropSegments[NumPropSegments - 1].PropertyIndex;
 
 	const FGLTFMeshData* MeshData = nullptr;
-	TArray<int32> SectionIndices;
+	FGLTFIndexArray SectionIndices;
 
 	if (Builder.ExportOptions->ExportMaterialVariants == EGLTFMaterialVariantMode::UseMeshData)
 	{
diff --git a/Engine/Plugins/Enterprise/GLTFExporter/MigratedFromGit/Source/GLTFExporter/Private/Tasks/GLTFMaterialTasks.cpp b/Engine/Plugins/Enterprise/GLTFExporter/MigratedFromGit/Source/GLTFExporter/Private/Tasks/GLTFMaterialTasks.cpp
--- a/Engine/Plugins/Enterprise/GLTFExporter/MigratedFromGit/Source/GLTFExporter/Private/Tasks/GLTFMaterialTasks.cpp
+++ b/Engine/Plugins/Enterprise/GLTFExporter/MigratedFromGit/Source/GLTFExporter/Private/Tasks/GLTFMaterialTasks.cpp
@@ -1302,7 +1302,7 @@ FGLTFPropertyBakeOutput FGLTFMaterialTask::BakeMaterialProperty(const FMaterialP
 {
 	if (MeshData == nullptr)
 	{
-		TArray<int32> TexCoords;
+		FGLTFIndexArray TexCoords;
 		FGLTFMaterialUtility::GetAllTextureCoordinateIndices(Material, Property, TexCoords);
 
 		if (TexCoords.Num() > 0)
diff --git a/Engine/Plugins/Enterprise/GLTFExporter/MigratedFromGit/Source/GLTFExporter/Private/Tasks/GLTFMaterialTasks.h b/Engine/Plugins/Enterprise/GLTFExporter/MigratedFromGit/Source/GLTFExporter/Private/Tasks/GLTFMaterialTasks.h
--- a/Engine/Plugins/Enterprise/GLTFExporter/MigratedFromGit/Source/GLTFExporter/Private/Tasks/GLTFMaterialTasks.h
+++ b/Engine/Plugins/Enterprise/GLTFExporter/MigratedFromGit/Source/GLTFExporter/Private/Tasks/GLTFMaterialTasks.h
@@ -15,7 +15,7 @@ class FGLTFMaterialTask : public FGLTFTask
 {
 public:
 
-	FGLTFMaterialTask(FGLTFConvertBuilder& Builder, FGLTFUVOverlapChecker& UVOverlapChecker, const UMaterialInterface* Material, const FGLTFMeshData* MeshData, TArray<int32> SectionIndices, FGLTFJsonMaterialIndex MaterialIndex)
+	FGLTFMaterialTask(FGLTFConvertBuilder& Builder, FGLTFUVOverlapChecker& UVOverlapChecker, const UMaterialInterface* Material, const FGLTFMeshData* MeshData, FGLTFIndexArray SectionIndices, FGLTFJsonMaterialIndex MaterialIndex)
 		: FGLTFTask(EGLTFTaskPriority::Material)
 		, Builder(Builder)
 		, UVOverlapChecker(UVOverlapChecker)
@@ -39,7 +39,7 @@ private:
 	FGLTFUVOverlapChecker& UVOverlapChecker;
 	const UMaterialInterface* Material;
 	const FGLTFMeshData* MeshData;
-	const TArray<int32> SectionIndices;
+	const FGLTFIndexArray SectionIndices;
 	const FGLTFJsonMaterialIndex MaterialIndex;
 
 	TSet<FMaterialPropertyEx> MeshDataBakedProperties;
diff --git a/Engine/Plugins/Enterprise/GLTFExporter/MigratedFromGit/Source/GLTFExporter/Private/Tasks/GLTFMeshTasks.cpp b/Engine/Plugins/Enterprise/GLTFExporter/MigratedFromGit/Source/GLTFExporter/Private/Tasks/GLTFMeshTasks.cpp
--- a/Engine/Plugins/Enterprise/GLTFExporter/MigratedFromGit/Source/GLTFExporter/Private/Tasks/GLTFMeshTasks.cpp
+++ b/Engine/Plugins/Enterprise/GLTFExporter/MigratedFromGit/Source/GLTFExporter/Private/Tasks/GLTFMeshTasks.cpp
@@ -136,7 +136,7 @@ void FGLTFStaticMeshTask::Complete()
 
 	for (int32 MaterialIndex = 0; MaterialIndex < MaterialCount; ++MaterialIndex)
 	{
-		const TArray<int32> SectionIndices = FGLTFMeshUtility::GetSectionIndices(MeshLOD, MaterialIndex);
+		const FGLTFIndexArray SectionIndices = FGLTFMeshUtility::GetSectionIndices(MeshLOD, MaterialIndex);
 		const FGLTFMeshSection* ConvertedSection = MeshSectionConverter.GetOrAdd(&MeshLOD, SectionIndices);
 
 		FGLTFJsonPrimitive& JsonPrimitive = JsonMesh.Primitives[MaterialIndex];
@@ -224,7 +224,7 @@ void FGLTFSkeletalMeshTask::Complete()
 
 	for (uint16 MaterialIndex = 0; MaterialIndex < MaterialCount; ++MaterialIndex)
 	{
-		const TArray<int32> SectionIndices = FGLTFMeshUtility::GetSectionIndices(MeshLOD, MaterialIndex);
+		const FGLTFIndexArray SectionIndices = FGLTFMeshUtility::GetSectionIndices(MeshLOD, MaterialIndex);
 		const FGLTFMeshSection* ConvertedSection = MeshSectionConverter.GetOrAdd(&MeshLOD, SectionIndices);
 
 		FGLTFJsonPrimitive& JsonPrimitive = JsonMesh.Primitives[MaterialIndex];
