[GLTFExporter] fix: add "using namespace UE" to fix some build issues in UE 4.26

commit 3b121f516949166608fdccfa354090ff831044c4
Date: Tue, 12 Jan 2021 05:41:02 +0100

    fix: add "using namespace UE" to fix some build issues in UE 4.26
    (cherry picked from commit 7d406e5b05e735f46514414fb649e9e4a242fbce​)

#skipci

diff --git a/Engine/Plugins/Enterprise/GLTFExporter/MigratedFromGit/Source/GLTFExporter/Private/Tasks/GLTFAnimationTasks.cpp b/Engine/Plugins/Enterprise/GLTFExporter/MigratedFromGit/Source/GLTFExporter/Private/Tasks/GLTFAnimationTasks.cpp
--- a/Engine/Plugins/Enterprise/GLTFExporter/MigratedFromGit/Source/GLTFExporter/Private/Tasks/GLTFAnimationTasks.cpp
+++ b/Engine/Plugins/Enterprise/GLTFExporter/MigratedFromGit/Source/GLTFExporter/Private/Tasks/GLTFAnimationTasks.cpp
@@ -11,6 +11,10 @@
 #include "Sections/MovieScene3DTransformSection.h"
 #include "Channels/MovieSceneChannelProxy.h"
 
+#if !(ENGINE_MAJOR_VERSION == 4 && ENGINE_MINOR_VERSION <= 25)
+using namespace UE;
+#endif
+
 void FGLTFAnimSequenceTask::Complete()
 {
 	// TODO: bone transforms should be absolute (not relative) according to gltf spec
diff --git a/Engine/Plugins/Enterprise/GLTFExporter/MigratedFromGit/Source/GLTFExporterRuntime/Private/Actors/GLTFHotspotActor.cpp b/Engine/Plugins/Enterprise/GLTFExporter/MigratedFromGit/Source/GLTFExporterRuntime/Private/Actors/GLTFHotspotActor.cpp
--- a/Engine/Plugins/Enterprise/GLTFExporter/MigratedFromGit/Source/GLTFExporterRuntime/Private/Actors/GLTFHotspotActor.cpp
+++ b/Engine/Plugins/Enterprise/GLTFExporter/MigratedFromGit/Source/GLTFExporterRuntime/Private/Actors/GLTFHotspotActor.cpp
@@ -20,6 +20,10 @@
 #include "Slate/SceneViewport.h"
 #endif
 
+#if !(ENGINE_MAJOR_VERSION == 4 && ENGINE_MINOR_VERSION <= 25)
+using namespace UE;
+#endif
+
 DEFINE_LOG_CATEGORY_STATIC(LogGLTFHotspot, Log, All);
 
 namespace
