[GLTFExporter] fix: fix remaining UE 4.26 build issues

commit c182fedbfedeecbd87bfe5fc6aeea9b9ab6cfaf7
Date: Tue, 12 Jan 2021 05:51:12 +0100

    fix: fix remaining UE 4.26 build issues
    (cherry picked from commit 3b9cfacc567e14ca9cdcd8369f8671d2ee14c956​)

#skipci

diff --git a/Engine/Plugins/Enterprise/GLTFExporter/MigratedFromGit/Source/GLTFMaterialAnalyzer/Private/GLTFProxyMaterialCompiler.h b/Engine/Plugins/Enterprise/GLTFExporter/MigratedFromGit/Source/GLTFMaterialAnalyzer/Private/GLTFProxyMaterialCompiler.h
--- a/Engine/Plugins/Enterprise/GLTFExporter/MigratedFromGit/Source/GLTFMaterialAnalyzer/Private/GLTFProxyMaterialCompiler.h
+++ b/Engine/Plugins/Enterprise/GLTFExporter/MigratedFromGit/Source/GLTFMaterialAnalyzer/Private/GLTFProxyMaterialCompiler.h
@@ -6,60 +6,72 @@
 
 class FGLTFProxyMaterialCompiler : public FProxyMaterialCompiler
 {
-    using FProxyMaterialCompiler::FProxyMaterialCompiler;
-
-    virtual int32 RayTracingQualitySwitchReplace(int32 Normal, int32 RayTraced) override
-    {
-    	return Normal; // ignore RayTraced branch
-    }
-
-    virtual int32 ReflectionAboutCustomWorldNormal(int32 CustomWorldNormal, int32 bNormalizeCustomWorldNormal) override
-    {
-    	return Compiler->ReflectionAboutCustomWorldNormal(CustomWorldNormal, bNormalizeCustomWorldNormal);
-    }
-
-    virtual int32 ParticleRelativeTime() override
-    {
-    	return Compiler->ParticleRelativeTime();
-    }
-
-    virtual int32 ParticleMotionBlurFade() override
-    {
-    	return Compiler->ParticleMotionBlurFade();
-    }
-
-    virtual int32 ParticleRandom() override
-    {
-    	return Compiler->ParticleRandom();
-    }
-
-    virtual int32 ParticleDirection() override
-    {
-    	return Compiler->ParticleDirection();
-    }
-
-    virtual int32 ParticleSpeed() override
-    {
-    	return Compiler->ParticleSpeed();
-    }
-
-    virtual int32 ParticleSize() override
-    {
-    	return Compiler->ParticleSize();
-    }
-
-    virtual int32 VertexInterpolator(uint32 InterpolatorIndex) override
-    {
-    	return Compiler->VertexInterpolator(InterpolatorIndex);
-    }
-
-    virtual int32 MaterialBakingWorldPosition() override
-    {
-    	return Compiler->MaterialBakingWorldPosition();
-    }
+	using FProxyMaterialCompiler::FProxyMaterialCompiler;
+
+	virtual int32 RayTracingQualitySwitchReplace(int32 Normal, int32 RayTraced) override
+	{
+		return Normal; // ignore RayTraced branch
+	}
+
+	virtual int32 ReflectionAboutCustomWorldNormal(int32 CustomWorldNormal, int32 bNormalizeCustomWorldNormal) override
+	{
+		return Compiler->ReflectionAboutCustomWorldNormal(CustomWorldNormal, bNormalizeCustomWorldNormal);
+	}
+
+	virtual int32 ParticleRelativeTime() override
+	{
+		return Compiler->ParticleRelativeTime();
+	}
+
+	virtual int32 ParticleMotionBlurFade() override
+	{
+		return Compiler->ParticleMotionBlurFade();
+	}
+
+	virtual int32 ParticleRandom() override
+	{
+		return Compiler->ParticleRandom();
+	}
+
+	virtual int32 ParticleDirection() override
+	{
+		return Compiler->ParticleDirection();
+	}
+
+	virtual int32 ParticleSpeed() override
+	{
+		return Compiler->ParticleSpeed();
+	}
+
+	virtual int32 ParticleSize() override
+	{
+		return Compiler->ParticleSize();
+	}
+
+	virtual int32 VertexInterpolator(uint32 InterpolatorIndex) override
+	{
+		return Compiler->VertexInterpolator(InterpolatorIndex);
+	}
+
+	virtual int32 MaterialBakingWorldPosition() override
+	{
+		return Compiler->MaterialBakingWorldPosition();
+	}
 
 	virtual EMaterialCompilerType GetCompilerType() const override
 	{
 		return EMaterialCompilerType::MaterialProxy;
 	}
+
+#if !(ENGINE_MAJOR_VERSION == 4 && ENGINE_MINOR_VERSION <= 25)
+	virtual int32 PreSkinVertexOffset() override
+	{
+		return Compiler->PreSkinVertexOffset();
+	}
+
+	virtual int32 PostSkinVertexOffset() override
+	{
+		return Compiler->PostSkinVertexOffset();
+	}
+#endif
 };
diff --git a/Engine/Plugins/Enterprise/GLTFExporter/MigratedFromGit/Source/GLTFMaterialBaking/Private/ExportMaterialProxy.h b/Engine/Plugins/Enterprise/GLTFExporter/MigratedFromGit/Source/GLTFMaterialBaking/Private/ExportMaterialProxy.h
--- a/Engine/Plugins/Enterprise/GLTFExporter/MigratedFromGit/Source/GLTFMaterialBaking/Private/ExportMaterialProxy.h
+++ b/Engine/Plugins/Enterprise/GLTFExporter/MigratedFromGit/Source/GLTFMaterialBaking/Private/ExportMaterialProxy.h
@@ -221,7 +221,22 @@ struct FExportMaterialCompiler : public FProxyMaterialCompiler
 		}
 	}
 
-	virtual EMaterialCompilerType GetCompilerType() const override { return EMaterialCompilerType::MaterialProxy; }
+	virtual EMaterialCompilerType GetCompilerType() const override
+	{
+		return EMaterialCompilerType::MaterialProxy;
+	}
+
+#if !(ENGINE_MAJOR_VERSION == 4 && ENGINE_MINOR_VERSION <= 25)
+	virtual int32 PreSkinVertexOffset() override
+	{
+		return Compiler->PreSkinVertexOffset();
+	}
+
+	virtual int32 PostSkinVertexOffset() override
+	{
+		return Compiler->PostSkinVertexOffset();
+	}
+#endif
 };
 
 class FExportMaterialProxy : public FMaterial, public FMaterialRenderProxy
@@ -235,7 +250,12 @@ public:
 		, bSynchronousCompilation(bInSynchronousCompilation)
 		, ProxyBlendMode(ProxyBlendMode)
 	{
+
+#if (ENGINE_MAJOR_VERSION == 4 && ENGINE_MINOR_VERSION <= 25)
 		SetQualityLevelProperties(EMaterialQualityLevel::High, false, GMaxRHIFeatureLevel);
+#else
+		SetQualityLevelProperties(GMaxRHIFeatureLevel, EMaterialQualityLevel::High);
+#endif
 		Material = InMaterialInterface->GetMaterial();
 		ReferencedTextures = InMaterialInterface->GetReferencedTextures();
 
