[GLTFExporter] fix: export 2d texture using calculated ingame size (to account for lod/mip bias)

commit 13ee1ff090cf533bf6c5da2b3e5c8ed64e2e278f
Date: Sat, 16 Jan 2021 16:44:59 +0100

    fix: export 2d texture using calculated ingame size (to account for lod/mip bias)
    (cherry picked from commit 66656783e481f33575423a0862ca876ccb459536​)

#skipci

diff --git a/Engine/Plugins/Enterprise/GLTFExporter/MigratedFromGit/Source/GLTFExporter/Private/Converters/GLTFTextureUtility.cpp b/Engine/Plugins/Enterprise/GLTFExporter/MigratedFromGit/Source/GLTFExporter/Private/Converters/GLTFTextureUtility.cpp
--- a/Engine/Plugins/Enterprise/GLTFExporter/MigratedFromGit/Source/GLTFExporter/Private/Converters/GLTFTextureUtility.cpp
+++ b/Engine/Plugins/Enterprise/GLTFExporter/MigratedFromGit/Source/GLTFExporter/Private/Converters/GLTFTextureUtility.cpp
@@ -82,6 +82,29 @@ TextureFilter FGLTFTextureUtility::GetDefaultFilter(TextureGroup LODGroup)
 	}
 }
 
+int32 FGLTFTextureUtility::GetMipBias(const UTexture* Texture)
+{
+	if (const UTexture2D* Texture2D = Cast<UTexture2D>(Texture))
+	{
+		return Texture2D->GetNumMips() - Texture2D->GetNumMipsAllowed(true);
+	}
+
+	return Texture->GetCachedLODBias();
+}
+
+FIntPoint FGLTFTextureUtility::GetInGameSize(const UTexture* Texture)
+{
+	const int32 Width = Texture->GetSurfaceWidth();
+	const int32 Height = Texture->GetSurfaceHeight();
+
+	const int32 MipBias = GetMipBias(Texture);
+
+	const int32 InGameWidth = FMath::Max(Width >> MipBias, 1);
+	const int32 InGameHeight = FMath::Max(Height >> MipBias, 1);
+
+	return { InGameWidth, InGameHeight };
+}
+
 TTuple<TextureAddress, TextureAddress> FGLTFTextureUtility::GetAddressXY(const UTexture* Texture)
 {
 	TextureAddress AddressX;
@@ -142,7 +165,7 @@ UTextureRenderTarget2D* FGLTFTextureUtility::CreateRenderTarget(const FIntPoint&
 	return RenderTarget;
 }
 
-bool FGLTFTextureUtility::DrawTexture(UTextureRenderTarget2D* OutTarget, const UTexture2D* InSource, const FMatrix& InTransform)
+bool FGLTFTextureUtility::DrawTexture(UTextureRenderTarget2D* OutTarget, const UTexture2D* InSource, const FVector2D& InPosition, const FVector2D& InSize, const FMatrix& InTransform)
 {
 	FRenderTarget* RenderTarget = OutTarget->GameThread_GetRenderTargetResource();
 	if (RenderTarget == nullptr)
@@ -162,7 +185,7 @@ bool FGLTFTextureUtility::DrawTexture(UTextureRenderTarget2D* OutTarget, const U
 	}
 
 	FCanvas Canvas(RenderTarget, nullptr, 0.0f, 0.0f, 0.0f, GMaxRHIFeatureLevel);
-	FCanvasTileItem TileItem(FVector2D::ZeroVector, InSource->Resource, FLinearColor::White);
+	FCanvasTileItem TileItem(InPosition, InSource->Resource, InSize, FLinearColor::White);
 	TileItem.BatchedElementParameters = BatchedElementParameters;
 
 	Canvas.PushAbsoluteTransform(InTransform);
@@ -177,16 +200,16 @@ bool FGLTFTextureUtility::DrawTexture(UTextureRenderTarget2D* OutTarget, const U
 	return true;
 }
 
-bool FGLTFTextureUtility::RotateTexture(UTextureRenderTarget2D* OutTarget, const UTexture2D* InSource, float InDegrees)
+bool FGLTFTextureUtility::RotateTexture(UTextureRenderTarget2D* OutTarget, const UTexture2D* InSource, const FVector2D& InPosition, const FVector2D& InSize, float InDegrees)
 {
 	FMatrix Transform = FMatrix::Identity;
 	if (InDegrees != 0)
 	{
-		const FVector Center = FVector(InSource->GetSizeX() / 2.0f, InSource->GetSizeY() / 2.0f, 0);
+		const FVector Center = FVector(InSize.X / 2.0f, InSize.Y / 2.0f, 0);
 		Transform = FTranslationMatrix(-Center) * FRotationMatrix({ 0, InDegrees, 0 }) * FTranslationMatrix(Center);
 	}
 
-	return DrawTexture(OutTarget, InSource, Transform);
+	return DrawTexture(OutTarget, InSource, InPosition, InSize, Transform);
 }
 
 UTexture2D* FGLTFTextureUtility::CreateTextureFromCubeFace(const UTextureCube* TextureCube, ECubeFace CubeFace)
diff --git a/Engine/Plugins/Enterprise/GLTFExporter/MigratedFromGit/Source/GLTFExporter/Private/Converters/GLTFTextureUtility.h b/Engine/Plugins/Enterprise/GLTFExporter/MigratedFromGit/Source/GLTFExporter/Private/Converters/GLTFTextureUtility.h
--- a/Engine/Plugins/Enterprise/GLTFExporter/MigratedFromGit/Source/GLTFExporter/Private/Converters/GLTFTextureUtility.h
+++ b/Engine/Plugins/Enterprise/GLTFExporter/MigratedFromGit/Source/GLTFExporter/Private/Converters/GLTFTextureUtility.h
@@ -16,14 +16,18 @@ struct FGLTFTextureUtility
 
 	static TextureFilter GetDefaultFilter(TextureGroup Group);
 
+	static int32 GetMipBias(const UTexture* Texture);
+
+	static FIntPoint GetInGameSize(const UTexture* Texture);
+
 	static TTuple<TextureAddress, TextureAddress> GetAddressXY(const UTexture* Texture);
 
 	static UTexture2D* CreateTransientTexture(const void* RawData, int64 ByteLength, const FIntPoint& Size, EPixelFormat Format, bool bSRGB = false);
 
 	static UTextureRenderTarget2D* CreateRenderTarget(const FIntPoint& Size, bool bIsHDR);
 
-	static bool DrawTexture(UTextureRenderTarget2D* OutTarget, const UTexture2D* InSource, const FMatrix& InTransform = FMatrix::Identity);
-	static bool RotateTexture(UTextureRenderTarget2D* OutTarget, const UTexture2D* InSource, float InDegrees);
+	static bool DrawTexture(UTextureRenderTarget2D* OutTarget, const UTexture2D* InSource, const FVector2D& InPosition, const FVector2D& InSize, const FMatrix& InTransform = FMatrix::Identity);
+	static bool RotateTexture(UTextureRenderTarget2D* OutTarget, const UTexture2D* InSource, const FVector2D& InPosition, const FVector2D& InSize, float InDegrees);
 
 	static UTexture2D* CreateTextureFromCubeFace(const UTextureCube* TextureCube, ECubeFace CubeFace);
 	static UTexture2D* CreateTextureFromCubeFace(const UTextureRenderTargetCube* RenderTargetCube, ECubeFace CubeFace);
@@ -38,7 +42,4 @@ struct FGLTFTextureUtility
 	static bool LoadPlatformData(UTextureCube* TextureCube);
 
 	static void FlipGreenChannel(TArray<FColor>& Pixels);
-
-	static bool HasAlphaChannel(const UTexture* Texture);
-	static bool HasAlphaChannel(EPixelFormat PixelFormat);
 };
diff --git a/Engine/Plugins/Enterprise/GLTFExporter/MigratedFromGit/Source/GLTFExporter/Private/Tasks/GLTFTextureTasks.cpp b/Engine/Plugins/Enterprise/GLTFExporter/MigratedFromGit/Source/GLTFExporter/Private/Tasks/GLTFTextureTasks.cpp
--- a/Engine/Plugins/Enterprise/GLTFExporter/MigratedFromGit/Source/GLTFExporter/Private/Tasks/GLTFTextureTasks.cpp
+++ b/Engine/Plugins/Enterprise/GLTFExporter/MigratedFromGit/Source/GLTFExporter/Private/Tasks/GLTFTextureTasks.cpp
@@ -12,11 +12,11 @@ void FGLTFTexture2DTask::Complete()
 	Texture2D->GetName(JsonTexture.Name);
 
 	const bool bIsHDR = FGLTFTextureUtility::IsHDR(Texture2D->GetPixelFormat());
-	const FIntPoint Size = { Texture2D->GetSizeX(), Texture2D->GetSizeY() };
+	const FIntPoint Size = FGLTFTextureUtility::GetInGameSize(Texture2D);
 	UTextureRenderTarget2D* RenderTarget = FGLTFTextureUtility::CreateRenderTarget(Size, bIsHDR);
 
 	// TODO: preserve maximum image quality (avoid compression artifacts) by copying source data (and adjustments) to a temp texture
-	FGLTFTextureUtility::DrawTexture(RenderTarget, Texture2D);
+	FGLTFTextureUtility::DrawTexture(RenderTarget, Texture2D, FVector2D::ZeroVector, Size);
 
 	if (!Texture2D->IsNormalMap() && bIsHDR)
 	{
@@ -64,7 +64,7 @@ void FGLTFTextureCubeTask::Complete()
 	UTextureRenderTarget2D* RenderTarget = FGLTFTextureUtility::CreateRenderTarget(Size, bIsHDR);
 
 	const float FaceRotation = FGLTFTextureUtility::GetCubeFaceRotation(CubeFace);
-	FGLTFTextureUtility::RotateTexture(RenderTarget, FaceTexture, FaceRotation);
+	FGLTFTextureUtility::RotateTexture(RenderTarget, FaceTexture, FVector2D::ZeroVector, Size, FaceRotation);
 
 	if (bIsHDR)
 	{
@@ -131,7 +131,7 @@ void FGLTFTextureRenderTargetCubeTask::Complete()
 	UTextureRenderTarget2D* RenderTarget = FGLTFTextureUtility::CreateRenderTarget(Size, bIsHDR);
 
 	const float FaceRotation = FGLTFTextureUtility::GetCubeFaceRotation(CubeFace);
-	FGLTFTextureUtility::RotateTexture(RenderTarget, FaceTexture, FaceRotation);
+	FGLTFTextureUtility::RotateTexture(RenderTarget, FaceTexture, FVector2D::ZeroVector, Size, FaceRotation);
 
 	if (bIsHDR)
 	{
