[GLTFExporter] chore: update project and plugin descriptors to UE 4.26

commit 257c2564475ccc0f464b98bdef9669909adc7cb3
Date: Sun, 17 Jan 2021 17:39:17 +0100

    chore: update project and plugin descriptors to UE 4.26
    (cherry picked from commit 81e8ce46184c4054b83c33d4db74096f7f357b3c​)

#skipci

diff --git a/Engine/Plugins/Enterprise/GLTFExporter/MigratedFromGit/GLTFExporter.uplugin b/Engine/Plugins/Enterprise/GLTFExporter/MigratedFromGit/GLTFExporter.uplugin
--- a/Engine/Plugins/Enterprise/GLTFExporter/MigratedFromGit/GLTFExporter.uplugin
+++ b/Engine/Plugins/Enterprise/GLTFExporter/MigratedFromGit/GLTFExporter.uplugin
@@ -10,7 +10,7 @@
 	"DocsURL": "",
 	"MarketplaceURL": "",
 	"SupportURL": "",
-	"EngineVersion": "4.25.0",
+	"EngineVersion": "4.26.0",
 	"CanContainContent": true,
 	"IsBetaVersion": true,
 	"Modules": [
