[GLTFExporter] chore: add initial documentation for gltf exporter

commit 718f89718b03cb4d2b619d85bfdcad2604317fb5
Date: Sun, 17 Jan 2021 17:44:46 +0100

    chore: add initial documentation for gltf exporter
    (cherry picked from commit de2e5fc442d5a4d204df5b6a9ab976d236af1d8b​)

#skipci

diff --git a/Engine/Plugins/Enterprise/GLTFExporter/MigratedFromGit/Docs/export-options-reference.md b/Engine/Plugins/Enterprise/GLTFExporter/MigratedFromGit/Docs/export-options-reference.md
new file mode 100644
--- /dev/null
+++ b/Engine/Plugins/Enterprise/GLTFExporter/MigratedFromGit/Docs/export-options-reference.md
@@ -0,0 +1,41 @@
+# Export options reference
+
+Option                         | Description
+-------------------------------| ----------------------------------------------------------------------------------------------------------------------------
+`Export Uniform Scale`         | Scale factor used for exporting all assets (0.01 by default) for conversion from centimeters (Unreal default) to meters (glTF).
+`Export Preview Mesh`          | If enabled, the preview mesh for a standalone animation or material asset will also be exported.
+`Bundle Web Viewer`            | If enabled, Unreal's glTF viewer (including an executable launch helper) will be bundled with the exported files. It supports all extensions used.
+`Show Files When Done`         | If enabled, the target folder (with all exported files) will be shown in explorer once the export has been completed successfully.
+`Export Unlit Materials`       | If enabled, materials with shading model unlit will be properly exported. Uses extension KHR_materials_unlit.
+`Export Clear Coat Materials`  | If enabled, materials with shading model clear coat will be properly exported. Uses extension KHR_materials_clearcoat, which is not supported by all glTF viewers.
+`Export Extra Blend Modes`     | If enabled, materials with blend modes additive, modulate, and alpha composite will be properly exported. Uses extension EPIC_blend_modes, which is supported by Unreal's glTF viewer.
+`Bake Material Inputs`         | Bake mode determining if and how a material input is baked out to a texture. Can be overriden by material- and input-specific bake settings, see GLTFMaterialExportOptions.
+`Default Material Bake Size`   | Default size of the baked out texture (containing the material input). Can be overriden by material- and input-specific bake settings, see GLTFMaterialExportOptions.
+`Default Material Bake Filter` | Default filtering mode used when sampling the baked out texture. Can be overriden by material- and input-specific bake settings, see GLTFMaterialExportOptions.
+`Default Material Bake Tiling` | Default addressing mode used when sampling the baked out texture. Can be overriden by material- and input-specific bake settings, see GLTFMaterialExportOptions.
+`Default Input Bake Settings`  | Input-specific default bake settings that override the general defaults above.
+`Default Level Of Detail`      | Default LOD level used for exporting a mesh. Can be overriden by component or asset settings (e.g. minimum or forced LOD level).
+`Export Vertex Colors`         | If enabled, export vertex color. Not recommended due to vertex colors always being used as a base color multiplier in glTF, regardless of material. Often producing undesirable results.
+`Export Vertex Skin Weights`   | If enabled, export vertex bone weights and indices in skeletal meshes. Necessary for animation sequences.
+`Export Mesh Quantization`     | If enabled, export Unreal-configured quantization for vertex tangents and normals, reducing size. Requires extension KHR_mesh_quantization, which is not supported by all glTF viewers.
+`Export Level Sequences`       | If enabled, export level sequences. Only transform tracks are currently supported. The level sequence will be played at the assigned display rate.
+`Export Animation Sequences`   | If enabled, export single animation asset used by a skeletal mesh component or hotspot actor. Export of vertex skin weights must be enabled.
+`Retarget Bone Transforms`     | If enabled, apply animation retargeting to skeleton bones when exporting an animation sequence.
+`Export Playback Settings`     | If enabled, export play rate, start time, looping, and auto play for an animation or level sequence. Uses extension EPIC_animation_playback, which is supported by Unreal's glTF viewer.
+`Texture Image Format`         | Desired image format used for exported textures.
+`Texture Image Quality`        | Level of compression used for exported textures, between 1 (worst quality, best compression) and 100 (best quality, worst compression). Does not apply to lossless formats (e.g. PNG).
+`No Lossy Image Format For`    | Texture types that will always use lossless formats (e.g. PNG) because of sensitivity to compression artifacts.
+`Export Texture Transforms`    | If enabled, export UV tiling and un-mirroring settings in a texture coordinate expression node for simple material input expressions. Uses extension KHR_texture_transform.
+`Export Lightmaps`             | If enabled, export lightmaps (created by Lightmass) when exporting a level. Uses extension EPIC_lightmap_textures, which is supported by Unreal's glTF viewer.
+`Texture HDR Encoding`         | Encoding used to store textures that have pixel colors with more than 8-bit per channel. Uses extension EPIC_texture_hdr_encoding, which is supported by Unreal's glTF viewer.
+`Export Hidden In Game`        | If enabled, export components that are flagged as hidden in-game.
+`Export Lights`                | Mobility of directional, point, and spot light components that will be exported. Uses extension KHR_lights_punctual.
+`Export Cameras`               | If enabled, export camera components.
+`Export Camera Controls`       | If enabled, export GLTFCameraActors. Uses extension EPIC_camera_controls, which is supported by Unreal's glTF viewer.
+`Export Animation Hotspots`    | If enabled, export GLTFHotspotActors. Uses extension EPIC_animation_hotspots, which is supported by Unreal's glTF viewer.
+`Export HDRI Backdrops`        | If enabled, export HDRIBackdrop blueprints. Uses extension EPIC_hdri_backdrops, which is supported by Unreal's glTF viewer.
+`Export Sky Spheres`           | If enabled, export SkySphere blueprints. Uses extension EPIC_sky_spheres, which is supported by Unreal's glTF viewer.
+`Export Variant Sets`          | If enabled, export LevelVariantSetsActors. Uses extension EPIC_level_variant_sets, which is supported by Unreal's glTF viewer.
+`Export Material Variants`      | Mode determining if and how to export material variants that change the materials property on a static or skeletal mesh component.
+`Export Mesh Variants`         | If enabled, export variants that change the mesh property on a static or skeletal mesh component.
+`Export Visibility Variants`   | If enabled, export variants that change the visible property on a scene component.
diff --git a/Engine/Plugins/Enterprise/GLTFExporter/MigratedFromGit/Docs/what-can-be-exported.md b/Engine/Plugins/Enterprise/GLTFExporter/MigratedFromGit/Docs/what-can-be-exported.md
new file mode 100644
--- /dev/null
+++ b/Engine/Plugins/Enterprise/GLTFExporter/MigratedFromGit/Docs/what-can-be-exported.md
@@ -0,0 +1,452 @@
+# What can be exported?
+
+Because glTF is a data-driven format (even with extensions) not everything in UE4 can be exported to glTF. The range of supported content can be separated into primary and secondary assets. Primary assets are directly exportable by right-clicking on an asset of the following type in the Content Browser and selecting `Asset Actions -> Export...`:
+
+- [Materials](#materials)
+- [Static Meshes](#static-meshes)
+- [Skeletal Meshes](#skeletal-meshes)
+- [Animation Sequences](#animation-sequences)
+- [Level Sequences](#level-sequences)
+- [Level Variant Sets](#level-variant-sets)
+- [Levels](#levels)
+
+Secondary assets are exported indirectly by being used or referenced in a primary asset, and include:
+
+- [Actors](#actors)
+- [Components](#components)
+- [Textures](#textures)
+
+## Materials
+
+Like UE4 glTF primarily uses a metallic/roughness PBR (physically based rendering) workflow enabling accurate photorealistic materials. However, unlike UE4, glTF materials do not support arbitrary material expressions. Instead, the format only allows a single texture and/or constant for each material input. Some inputs, like `Metallic`/`Roughness` and `Base Color`/`Opacity (Mask)`, must even share the same texture. So, to convert a UE4 material to glTF, the exporter will automatically use one of the following two methods (ordered by preference):
+
+- Material expression matching
+- Material baking
+
+### Material expression matching
+
+Material expression matching is the quicker and more accurate approach but limited to very simple and strict expression patterns. The exporter will always prefer it whenever a material expression satisfies the necessary conditions. Simplified, each material input expression is examined to see if it matches a limited range of supported expression node patterns, like a single texture or constant expression. If there is a match, then the values from the material expression are extracted and converted without the need for material baking. Below are some examples of material expressions that meet the conditions:
+
+![Examples of Constant and Vector Parameter material nodes](figures/constant-nodes.png)
+
+- `Constant` or `Vector Parameter` node (does not apply to `Normal` or `Ambient Occlusion` because of limitations in glTF)
+
+![Examples of Texture Sample and Texture Parameter material nodes](figures/texture-nodes.png)
+
+- `Texture Sample` or `Texture Parameter` node with or without a `Texture Coordinate` node
+
+> Currently, expression matching does not support multiplying a texture sample node with a constant.
+
+However, since most UE4 material inputs use more advanced expressions, the exporter will typically fall back to material baking if expression matching fails.
+
+### Material baking
+
+Material baking is a technique that renders the full expression of a specific material input into a 2D texture (often using a simple plane as mesh data). Material baking can handle most material expressions but is slower and requires some key input from the user including:
+
+- `Size` of the resulting texture containing the baked-out material input.
+- `Filter`ing mode to use when sampling the resulting texture in the glTF material.
+- `Tiling`/addressing mode to use for the X and Y axis when sampling the resulting texture in the glTF material.
+
+These bake settings can be overridden for each specific material input (any non-overridden setting will fall back to the corresponding default setting).
+
+Since its often desireable to customize these bake settings on a per-asset-basis, each setting can be further overridden by adding `GLTF Material Export Options` as user data to a material.
+
+![Overriding material bake settings with user data](figures/material-userdata.png)
+
+If a bake setting is overridden on a specific material asset (with user data), then any related material instance (that has the specific asset as direct or indirect parent material) will automatically inherit that bake setting - unless the material instance has itself overridden the same setting. Basically, the exporter will first look for overridden bake settings in the material's own user data, and if none was found, it will go up the hierarchy chain of parents until the closest (relation-wise) is found. If none was ultimately found, then it will fall back to the global bake settings defined by the general export options.
+
+Aside from the discussed bake settings, material baking can be set to one of the following modes:
+
+- `Disabled` - Never use material baking and only rely on material expression matching.
+- `Simple` - If a material input needs baking, then only use a simple plane as mesh data.
+- `Use Mesh Data` - If a material input uses some mesh-related data (such as vertex color, world position, vector transform nodes etc), then this mode will allow that data to be included in the material baking (and the baked out texture); Otherwise, fall back to using a simple plane.
+
+Simple baking mode is preferential to baking with mesh-specific data since the output always covers the entire texture size and the material will not be unique for each mesh. Additionally, baking a material using specific mesh data assumes the mesh has proper lightmap UVs (i.e., covers the entire mesh and does not overlap), which is not always the case. If the resulting exported materials look different than expected its recommended to try changing `Bake Material Inputs` to `Simple`. However, as with the latest official automotive material pack, its also not always possible to avoid baking without the mesh data mode, since some complex material rely heavily on mesh data.
+
+Finally, its important to note that because each input expression is evaluated pixel by pixel when baking a material (regardless of which mesh data is used), all dynamic nodes (like time, camera position and reflection vector etc) will naturally become static. Hence, it is not recommended to use these kinds of view-dependent expressions when exporting a material. The more independent the material is of the view state (e.g., camera, time, screen etc) the better and more accurately it will be converted to glTF.
+
+### Shading Models
+
+Not all UE4 shading models are supported by glTF. Currently, only the following can be exported:
+- `Default Lit`
+- `Clear Coat`
+- `Unlit`
+
+Unsupported shading models will be treated as Default Lit. If the shading model is determined by the material expression graph, the exporter will attempt to evaluate it using static analysis. The static analyzer is limited to certain types of expressions however:
+
+![Example of Shading Model Expression that can be evaluated](figures/good-shadingmodel-expression.png)
+
+More complicated expressions (that use non-static variables) may not be fully resolved. If the static analyzer cannot resolve the evaluation to a single static shading model, then the exporter will simply choose one of the remaining possible outcomes.
+
+![Example of Shading Model Expression that can not be evaluated](figures/bad-shadingmodel-expression.png)
+
+It is also worth remembering that while a UE4 material can use different shading models simultaneously (i.e., for different regions), glTF is limited to a single shading model per material. Like before, the exporter will in such a case simply choose one of the final shading models.
+
+#### Default Lit
+
+The most common and default shading model in UE4 is `Default Lit`, which shares almost the exact same material inputs with the default shading model in glTF:
+
+- `Base Color`
+- `Metallic`
+- `Roughness`
+- `Emissive Color`
+- `Opacity`/`Opacity Mask` (depending on blend mode)
+- `Normal`
+- `Ambient Occlusion`
+
+Nonetheless, there are a few noteworthy differences:
+
+1. In glTF some input pairs share the same texture slot, specifically:
+
+   - `Base Color` and `Opacity (Mask)`
+   - `Metallic` and `Roughness`
+
+   This means that paired inputs must always use the same texture, just different color channels. So, in the case of `Metallic` and `Roughness`, the former is represented by the blue channel and the latter uses the green channel. But since both share the same texture, they will always have the same texture resolution and coordinates.
+
+2. In glTF some inputs are limited to only textures, specifically:
+
+   - `Normal`
+   - `Ambient Occlusion`
+
+   This means that unlike other inputs, `Normal` and `Ambient Occlusion` cannot have a constant non-default value, unless as a 1x1 texture. The reasoning is that constant values for such inputs should be avoided.
+
+3. glTF's metallic/roughness PBR workflow does not support `Specular` (like in UE4).
+
+#### Clear Coat
+
+> Please note that exporting shading model `Clear Coat` uses the glTF extension `KHR_materials_clearcoat`, which can be turned off in the export options.
+
+One of the most important shading models among automotive materials is `Clear Coat`. This shading model extends upon the previous covered `Default Lit`, with three additional material inputs:
+
+- `Clear Coat` (intensity)
+- `Clear Coat Roughness`
+- `Clear Coat Bottom Normal`
+
+The last one `Clear Coat Bottom Normal` is not a regular material input pin but a separate custom output node, and requires the feature [Clear Coat Enable Second Normal](https://docs.unrealengine.com/en-US/RenderingAndGraphics/Materials/HowTo/ClearCoatDualNormal/index.html) to be enabled in project settings. But like the regular `Normal` input, `Clear Coat Bottom Normal` is also limited to only textures, meaning it cannot have a constant non-default value, unless as a 1x1 texture.
+
+Also, worth noting is that `Clear Coat` and `Clear Coat Roughness` are paired inputs (like `Metallic` and `Roughness`) and thus share the same texture slot. `Clear Coat` is represented by the red channel and `Clear Coat Roughness` uses the green channel.
+
+#### Unlit
+
+> Please note that exporting shading model `Unlit` uses the glTF extension `KHR_materials_unlit`, which can be turned off in the export options.
+
+Unlike the previously covered shading model, `Unlit` does not extend upon `Default Lit` and is instead standalone with only two material inputs:
+
+- `Emissive Color`
+- `Opacity`/`Opacity Mask` (depending on blend mode)
+
+Just like `Base Color`, `Emissive Color` will in this case be paired with `Opacity (Mask)`.
+
+### Blend Modes
+
+Aside from shading models, the following blend modes are also supported:
+
+- `Opaque`
+- `Masked`
+- `Translucent`
+- `Additive`
+- `Modulate`
+- `Alpha Composite`
+
+> The last three blend modes (`Additive`, `Modulate`, and `Alpha Composite`) uses the glTF extension `EPIC_blend_modes`, which can be turned off in the export options.
+
+## Static Meshes
+
+A static mesh is almost fully supported by glTF, with only a few caveats:
+
+- Vertex colors will always act as a multiplier for base color in glTF, regardless of material, which may produce undesirable results. Unless you know what, you are doing its recommended to disable `Export Vertex Colors`, and instead select `Use Mesh Data` to bake materials.
+- glTF does not support half precision UVs (16-bit), thus its recommended to enable `Use Full Precision UVs` (32-bit) in a static mesh's asset settings.
+- Even thought all UV channels are exported, most glTF applications (including Unreal's glTF viewer) only supports handling two UV texture coordinate channels, thus materials that use texture coordinates beyond the first two channels should be avoided.
+- Collision geometry is not supported in glTF and will be ignored during export.
+
+### Level of Detail
+
+Because glTF does not support multiple levels of detail. The exporter will select a single LOD for export based on the following settings:
+
+- `Default Level Of Detail` (assigned by the export options), which is the global fallback.
+- `Minimum LOD` (assigned by each mesh asset), which is used instead of the default LOD, if the later is lower.
+- `Override Min LOD` (assigned by each mesh component/actor), if assigned can override the minimum LOD.
+- `Forced Lod Model` (assigned by each mesh component/actor), if assigned can override the final LOD and nullify all other settings.
+
+### Mesh Quantization
+
+To save disk and memory space, the following vertex attributes are quantized in UE4 and glTF:
+
+- Vertex colors (8-bit per component)
+- Vertex normals (8-bit or 16-bit per component)
+- Vertex tangents (8-bit or 16-bit per component)
+
+To achieve high quality reflections and lighting its recommended to enable `Use High Precision Tangent Basis` (which is assigned in each mesh asset's settings). This setting decides how much vertex normals and tangents are quantized (whether 8-bit or 16-bit).
+
+> Quantization of the vertex normals and tangents require the glTF extension `KHR_mesh_quantization`, which can be turned off in the export options.
+
+## Skeletal Meshes
+
+Apart from the caveats regarding static meshes, skeletal meshes in glTF also have the two additional:
+
+- No support for mesh clothing assets in glTF.
+- No support for morph target animations, currently.
+
+## Animation Sequences
+
+UE4 animation sequences are fully supported in glTF, as long a vertex skin weights are also exported. Additionally, UE animation retargeting may also be accounted for in the export.
+
+## Level Sequences
+
+Support for level sequences is restricted to transform tracks in absolute space (i.e., no blending of multiple tracks). Each level sequence is also exported at their selected display rate. For a level sequence asset to be included in a scene export, the asset needs to be assigned to a `Level Sequence Actor` in the scene.
+
+## Level Variant Sets
+
+> Please note that export of Level Variant Sets uses the glTF extension `EPIC_level_variant_sets`, which can be turned off in the export options.
+
+In UE4, variant sets can be used to configure almost any property in a scene. The exporter currently only supports the following properties:
+
+- `Material` - changing any material asset on a static or skeletal mesh component
+- `Static Mesh` - changing the mesh asset on a static mesh component
+- `Skeletal Mesh` - changing the mesh asset on a skeletal mesh component
+- `Visible` - changing the visibility on a scene component
+
+Support for any of the listed properties can be turned off in the export settings. For a level variant sets asset to be included in a scene export, the asset needs to be assigned to a `Level Variant Sets Actor` in the scene.
+
+## Levels
+
+A Level can be exported either by right-clicking on the asset in Content Browser or via the menu bar, `File -> Export All` or `File -> Export Selected` (the latter will only export selected actors).
+
+## Actors
+
+Support for actors in a level, can be split into two categories:
+
+- An actor that is one of the following types:
+
+  - [`HDRI Backdrop`](#hdri-backdrops)
+  - [`Sky Sphere`](#sky-spheres)
+  - [`Level Sequence Actor`](#level-sequence-actors)
+  - [`Level Variant Sets Actor`](#level-variant-sets-actors)
+  - [`GLTF Camera`](#gltf-cameras)
+  - [`GLTF Hotspot`](#gltf-hotspots)
+
+- An actor that has one or more supported components.
+
+If an actor matches any type in the first category, then the exact properties exported is detailed in the respective subsection below. If an actor is of the second category, only the properties of the supported component(s) will be exported, as described in the [Components](#components) section.
+
+### HDRI Backdrops
+
+> Please note that export of `HDRI Backdrop` uses the glTF extension `EPIC_hdri_backdrops`, which can be turned off in the export options.
+
+[`HDRI Backdrop`](https://docs.unrealengine.com/en-US/BuildingWorlds/LightingAndShadows/HDRIBackdrop/index.html) is a blueprint actor that is provided by an engine-distributed plugin, to quickly set up product visualization using HDR image projection.
+
+The following properties are supported by the exporter:
+
+Property                   | Description
+---------------------------| ------------------------------------------------------------------------------------------------------------------------------------------------------------------
+`Cubemap`                  | An HDR image that will be projected onto the ground and backdrop and is used by the built-in Sky Light source.
+`Intensity`                | The intensity of the embedded Sky Light and how emissive the backdrop HDR image is. Higher values result in brighter ambient lighting being samples from the HDR image (cd/m2).
+`Size`                     | The size (in meters) of the mesh used to project the HDR image. It controls the diameter of the backdrop mesh and should be adjusted based on the HDR image used, the horizon height, and the content in the scene. For most exterior scenes, a typical size should be approximately 100 meters.
+`Projection Center`        | Defines the projection point of the HDR image.
+`Lighting Distance Factor` | Specifies the ground area that will be affected by lighting and shadows. Lit areas will have slightly different shading depending on the intensity and other lighting parameters in the scene. This enables the lit area range to blend smoothly around the camera to reduce shading differences with the background HDR projection.
+`Use Camera Projection`    | Disables ground tracking and enables the HDR image to follow the camera.
+`Mesh`                     | The static mesh to use as a backdrop from which the HDR image is projected.
+
+### Sky Spheres
+
+> Please note that export of `Sky Sphere` uses the glTF extension `EPIC_sky_spheres`, which can be turned off in the export options.
+
+`Sky Sphere` is a blueprint actor that is part of the core engine content, to provide a customizable dynamic sky sphere that mimics the effects of atmospheric fog. The following properties are supported by the exporter:
+
+Property                            | Description
+------------------------------------| -----------------------------------------------------------------------------------
+`Directional Light Actor`           | An optional directional light actor to match the sky's sun position and color.
+`Sun Brightness`                    | Brightness multiplier for the sun disk.
+`Cloud Speed`                       | Panning speed for the clouds.
+`Cloud Opacity`                     | Opacity of the panning clouds.
+`Stars Brightness`                  | Multiplier for the brightness of the stars when the sun is below the horizon.
+`Colors Determined By Sun Position` | If enabled, sky colors will change according to the sun's position.
+`Sun Height`                        | If no directional light is assigned, this value determines the height of the sun.
+`Horizon Falloff`                   | Affects the size of the gradient from zenith color to horizon color.
+`Zenith Color`                      |
+`Horizon color`                     |
+`Cloud color`                       |
+`Overall Color`                     |
+
+### Level Sequence Actors
+
+`Level Sequence Actor` is an engine actor responsible for playing level sequences in the scene. The following properties are supported by the exporter:
+
+Property         | Description
+-----------------| --------------------------------------------------------------------
+`Level Sequence` | Auto-play the sequence when created
+`Auto Play`      | Auto-play the sequence when created
+`Play Rate`      | The rate at which to playback the animation
+`Start Time`     | Start playback at the specified offset from the start of the sequence's playback range
+`Loop Count`     | Number of times to loop playback. -1 for infinite, else the number of times to loop before stopping
+
+> Please note that export of playback settings (i.e., the last four properties) uses the glTF extension `EPIC_animation_playback`, which can be turned off in the export options.
+
+### Level Variant Sets Actors
+
+> Please note that export of `Level Variant Sets Actor` uses the glTF extension `EPIC_level_variant_sets`, which can be turned off in the export options.
+
+`Level Variant Sets Actor` is provided by an engine-distributed plugin (`Variant Manager`) and ìs responsible for querying and activating variants in a `Level Variant Sets` at runtime. Only the `Level Variant Sets` property is exported.
+
+### GLTF Cameras
+
+> Please note that export of `GLTF Camera` uses the glTF extension `EPIC_camera_controls`, which can be turned off in the export options.
+
+`GLTF Camera` is an actor provided by the exporter plugin, to add user-controlled cameras to a glTF scene. It has the following properties that are supported by the exporter:
+
+Property               | Description
+-----------------------| ---------------------------------------------------------------------------------------------------------
+`Mode`                 | Camera mode which can be `Orbital` (rotate the camera around a target and zoom in/out) or `Free Look` (rotate the camera in first person but no movement).
+`Target`               | Actor which the camera will focus on and subsequently orbit when using Orbital mode.
+`Pitch Angle Min`      | Minimum pitch angle (in degrees) for the camera.
+`Pitch Angle Max`      | Maximum pitch angle (in degrees) for the camera.
+`Yaw Angle Min`        | Minimum yaw angle (in degrees) for the camera.
+`Yaw Angle Max`        | Maximum yaw angle (in degrees) for the camera.
+`Distance Min`         | Closest distance the camera can approach the focused actor.
+`Distance Max`         | Farthest distance the camera can recede from the focused actor.
+`Dolly Duration`       | Duration (in seconds) that it takes the camera to complete a change in distance.
+`Dolly Sensitivity`    | Size of the dolly movement relative to user input. The higher the value, the faster it moves.
+`Rotation Inertia`     | Deceleration that occurs after rotational movement. The higher the value, the longer it takes to settle.
+`Rotation Sensitivity` | Size of the rotational movement relative to user input. The higher the value, the faster it moves.
+
+### GLTF Hotspots
+
+> Please note that export of `GLTF Hotspot` uses the glTF extension `EPIC_animation_hotspots`, which can be turned off in the export options.
+
+`GLTF Hotspot` is an actor provided by the exporter plugin, to add a clickable UI element (that triggers animations) to a glTF scene. It has the following properties that are supported by the exporter:
+
+Property                | Description
+------------------------| ------------------------------------------------------------------------------------------------------------------------
+`Skeletal Mesh Actor`   | The skeletal mesh actor that will be animated when the hotspot is clicked.
+`Animation Sequence`    | The animation sequence that will be played on the skeletal mesh actor. Must be compatible with its skeletal mesh asset.
+`Level Sequence`        | The level sequence that will be played in the level.
+`Image`                 | The billboard image that will be shown when the hotspot is in an inactive state or one without a specified image.
+`Hovered Image`         | The optional billboard image that will be shown when a cursor enters the hotspot
+`Toggled Image`         | The optional billboard image that will be shown when the hotspot is toggled by a click.
+`Toggled Hovered Image` | The optional billboard image that will be shown when the hotspot is toggled by a click and a cursor enters it.
+
+## Components
+
+Supported components can be separated into three groups:
+
+- Cameras:
+  - [`Camera Component`](#camera-components)
+- Lights:
+  - [`Directional Light`](#directional-lights)
+  - [`Point Light`](#point-lights)
+  - [`Spot Light`](#spot-lights)
+- Primitives:
+  - [`Scene Component`](#scene-components)
+  - [`Static Mesh Component`](#static-mesh-components)
+  - [`Skeletal Mesh Component`](#skeletal-mesh-components)
+
+For more details on which properties in each component type is supported, see the respective subsection below.
+
+### Scene Components
+
+`Scene Component` has the following properties that are supported by the exporter:
+
+Property   | Description
+-----------| ---------------------------
+`Location` | Location of the component.
+`Rotation` | Rotation of the component.
+`Scale`    | Scale of the component.
+
+Due to differences in how UE4 applies scale, non-uniform scale may be represented differently in glTF.
+
+### Camera Components
+
+In addition to the properties covered by the [Scene Components](#scene-components) section, `Camera Component` also has the following properties that are supported by the exporter:
+
+Property                 | Description
+-------------------------| ----------------------------------------------------------------------------------------------
+`Projection Mode`        | The type of camera
+`Field Of View`          | The horizontal field of view (in degrees) in perspective mode (ignored in Orthographic mode).
+`Ortho Width`            | The desired width (in world units) of the orthographic view (ignored in Perspective mode).
+`Ortho Near Clip Plane`  | The near plane distance of the orthographic view (in world units).
+`Ortho Far Clip Plane`   | The far plane distance of the orthographic view (in world units).
+`Aspect Ratio`           | The ratio of width to height.
+`Constrain Aspect Ratio` | If enabled, black bars will be added if the destination view has a different aspect ratio than this camera requested.
+
+### Directional Lights
+
+> Please note that export of `Directional Light` uses the glTF extension `KHR_lights_punctual`, which can be turned off in the export options.
+
+In addition to the properties covered by the [Scene Components](#scene-components) section, `Directional Light` also has the following properties that are supported by the exporter:
+
+Property          | Description
+------------------| ---------------------------------------------------------
+`Intensity`       | Total energy that the light emits.
+`Light Color`     | Filter color of the light.
+`Temperature`     | Color temperature in Kelvin of the blackbody illuminant.
+`Use Temperature` | If disabled, use white (D65) as illuminant.
+
+### Point Lights
+
+> Please note that export of `Point Light` uses the glTF extension `KHR_lights_punctual`, which can be turned off in the export options.
+
+In addition to the properties covered by the [Directional Lights](#directional-lights) section, `Point Light` also has the following properties that are supported by the exporter:
+
+Property             | Description
+---------------------| --------------------------------------
+`Attenuation Radius` | Bounds the light's visible influence.
+
+### Spot Lights
+
+> Please note that export of `Spot Light` uses the glTF extension `KHR_lights_punctual`, which can be turned off in the export options.
+
+In addition to the properties covered by the [Point Lights](#point-lights) section, `Spot Light` also has the following properties that are supported by the exporter:
+
+Property             | Description
+---------------------| ------------------------------------------------------------------
+`Attenuation Radius` | Bounds the light's visible influence.
+`Inner Cone Angle`   | Angle, in degrees, from centre of spotlight where falloff begins.
+`Outer Cone Angle`   | Angle, in degrees, from centre of spotlight where falloff ends.
+
+### Static Mesh Components
+
+In addition to the properties covered by the [Scene Components](#scene-components) section, `Static Mesh Component` also has the following properties that are supported by the exporter:
+
+Property      | Description
+--------------| ----------------------------------------
+`Static Mesh` | The static mesh used by this component.
+`Materials`   | The materials used by this component.
+
+The level of detail used when exporting the static mesh is determined by the component's LOD settings (i.e., `Forced Lod Model`, `Min LOD`, and `Override Min LOD`) as well as the asset's `Minimum LOD` and the export option's `Default Level Of Detail`. For more details see section [Level of Detail](#level-of-detail).
+
+### Skeletal Mesh Components
+
+In addition to the properties covered by the [Scene Components](#scene-components) section, `Skeletal Mesh Component` also has the following properties that are supported by the exporter:
+
+Property           | Description
+-------------------| --------------------------------------------
+`Skeletal Mesh`    | The skeletal mesh used by this component.
+`Materials`        | The materials used by this component.
+`Anim to Play`     | The sequence to play on this skeletal mesh.
+`Looping`          | If enabled, the sequence will be looped.
+`Playing`          | Auto-play the sequence when created.
+`PlayRate`         | The rate at which to play the animation.
+`Initial Position` | Start play at the specified offset from the start of the sequence's playback range.
+
+> Please note that export of playback settings (i.e., the last four properties) uses the glTF extension `EPIC_animation_playback`, which can be turned off in the export options.
+
+The level of detail used when exporting the static mesh is determined by the component's LOD settings (i.e., `Forced Lod Model`, `Min Lod Model`, and `Override Min Lod`) as well as the asset's `Minimum LOD` and the export option's `Default Level Of Detail`. For more details see section [Level of Detail](#level-of-detail).
+
+
+## Textures
+
+The following texture types are supported by the exporter:
+
+- `Texture 2D`
+- `Texture Cube`
+- `Light Map Texture 2D`
+
+> Please note that export of `Light Map Texture 2D` uses the glTF extension `EPIC_lightmap_textures`, which can be turned off in the export options.
+
+To support all texture settings (like color adjustments), the exporter uses the render data (i.e., platform data) stored internally by UE4 (rather than the source data). This ensure the exported texture is identical to the texture rendered in the UE4 editor or in-game. The downside of this approach is that any artifacts introduced by UE4's compression settings will also be exported. Thus, its recommended to use compression settings `UserInterface2D` and `HDR` wherever possible.
+
+### HDR encoding
+
+> Please note that export of HDR textures uses the glTF extension `EPIC_texture_hdr_encoding`, which can be turned off in the export options.
+
+Because glTF does not support HDR (high dynamic range) texture formats, the exporter works around this limitation by using modified RGBM encoding. It stores pixels as one byte each for RGB (red, green, and blue) values with a one byte (alpha) as a shared multiplier.
diff --git a/Engine/Plugins/Enterprise/GLTFExporter/MigratedFromGit/Docs/what-is-gltf.md b/Engine/Plugins/Enterprise/GLTFExporter/MigratedFromGit/Docs/what-is-gltf.md
new file mode 100644
--- /dev/null
+++ b/Engine/Plugins/Enterprise/GLTFExporter/MigratedFromGit/Docs/what-is-gltf.md
@@ -0,0 +1,46 @@
+# What is glTF?
+
+glTF™ (GL Transmission Format) is an open standard file format (developed and maintained by Khronos Group) to efficiently share rich 3D content between a wide variety of applications. glTF has been specifically designed for compact file size, fast loading, complete scene representation, and extensibility – making it ideal for (but not limited to) web development.
+
+![glTF logo](figures/gltf.png)
+
+## glTF 2.0
+
+Although the initial version of glTF (1.0) was heavily intended for WebGL and OpenGL applications (due to dependence on GLSL shaders), the current version (2.0) is much more run-time independent and instead only relies on well established PBR (physically based rendering) workflows. For these and other technical reasons the exporter (and much of the software ecosystem) exclusively support glTF 2.0, hereafter referred to as just glTF.
+
+## File format
+
+A glTF file can have one of two possible file extensions, `.gltf` (JSON) or `.glb` (binary). When choosing the first the exporter will save the full scene description as a JSON-formatted human-readable text file (`.gltf`), while textures (e.g. `.png` or `.jpeg`) and binary data (`.bin`) will be saved as separate files in the same directory. When choosing the later, the exporter will embed all textures and binary data together with the full scene description into a single self-contained binary file (`.glb`).
+
+![glTF file extensions and related formats](figures/files.png)
+
+## Extensions
+
+Due to the complexities of modern 3D graphics, an engine-independent format such as glTF cannot cover all features, especially in a powerful game engine like UE4. The good news is that the glTF base format is specifically designed, as previously mentioned, with extensibility in mind. This extensibility is facilitated using so called glTF extensions.
+
+It is important to note that the extensibility does come at a cost - not all applications implement every glTF extension. If an application does not support an extension, it may still be able to load and show other parts of the glTF that do not use the specific extension - unless that extension is stated as explicitly required by the glTF, in which case the application will typically fail to open the entire glTF.
+
+Each glTF extension has a unique name, allowing any application to easily identify all extensions in a glTF file, regardless of whether all are supported by the application in question or not. The first part of the extension name (i.e., prefix) can also be used to identify how well supported the extension was generally when first conceived:
+
+- Khronos ratified extensions use the reserved prefix `KHR` and are typically the most supported.
+- Multi-vendor extensions use the reserved prefix `EXT` and are typically supported by more than one vendor (i.e., company or application).
+- Vendor-specific extensions use registered prefixes (like `EPIC`) and are typically supported by at least that vendor. Applications by other vendors may also still support such an extension.
+
+To support many of UE4's features the exporter has implemented the following extensions (all of which can be turned off in the export options):
+
+Extension                   | Description
+----------------------------|--------------------------------------------------------
+`KHR_lights_punctual`       | Point, spot, and directional lights
+`KHR_materials_unlit`       | Materials with unlit shading model
+`KHR_materials_clearcoat`   | Materials with clear coat shading model
+`KHR_mesh_quantization`     | Decrease vertex data size and precision
+`KHR_texture_transform`     | Tiling and mirroring texture coordinates
+`EPIC_lightmap_textures`    | Lightmass baked UE4-encoded lightmaps
+`EPIC_level_variant_sets`   | Scene variants by UE4's variant manager
+`EPIC_hdri_backdrops`       | UE4 backdrop actors for HDR image projection
+`EPIC_sky_spheres`          | UE4's default sky sphere actors
+`EPIC_animation_hotspots`   | Exporter-provided actors for interactive animations
+`EPIC_animation_playback`   | Animation looping, auto play, start time, play rate
+`EPIC_camera_controls`      | Exporter-provided actors for user-controlled cameras
+`EPIC_texture_hdr_encoding` | HDR encoding for image formats like PNG and JPEG
+`EPIC_blend_modes`          | Additive, modulate, and alpha composite blend modes
diff --git a/Engine/Plugins/Enterprise/GLTFExporter/MigratedFromGit/README.md b/Engine/Plugins/Enterprise/GLTFExporter/MigratedFromGit/README.md
--- a/Engine/Plugins/Enterprise/GLTFExporter/MigratedFromGit/README.md
+++ b/Engine/Plugins/Enterprise/GLTFExporter/MigratedFromGit/README.md
@@ -1,37 +1,46 @@
-glTF exporter plugin
-====================
+# glTF Exporter plugin
 
-Unreal Editor plugin for exporting UE4 assets to glTF 2.0 (with custom extensions).
+The glTF Exporter lets UE4 content creators showcase any scene or model on the web in full interactive real-time 3D with minimal effort.
 
 
-Installation
-------------
+## Compatibility
 
-* Alt 1: Install to project:
+Currently, the plugin works with both Unreal Engine **4.25** and **4.26**. It has primarily been tested on Windows but is designed to be platform-independent.
+
+
+## Installation
+
+- Alt 1: Install to project:
   1. Go to project folder which contains `[ProjectName].uproject`
   1. Create a folder called `Plugins` (if it doesn't already exist)
   1. Copy this `GLTFExporter` folder into the `Plugins` folder
 
-* Alt 2: Install to Unreal Engine:
+- Alt 2: Install to Unreal Engine:
   1. Go to the plugin folder of Unreal Engine which is `Engine/Plugins`
   1. Copy this `GLTFExporter` folder into the `Plugins` folder
 
 In Unreal Editor, open `Menu -> Edit -> Plugins` and make sure `glTF Exporter` is installed and enabled.
 
 
-Usage
------
+## Usage
 
-* Alt 1: Export asset via Content Browser
+- Alt 1: Export asset via Content Browser
   1. Right-click on a `StaticMesh`, `SkeletalMesh`, `AnimSequence`, `Level`, or `Material` asset in the Content Browser.
   1. Select `Asset Actions -> Export...`
   1. Change `Save as type` to `.gltf` (or `.glb`) and click `Save`
   1. When `glTF Export Options` window is shown, click `Export`
 
-* Alt 2: Export current level via File Menu
+- Alt 2: Export current level via File Menu
   1. Select any number of actors in the current level
   1. In the top menu, select `File -> Export Selected...`
   1. Change `Save as type` to `.gltf` (or `.glb`) and click `Save`
   1. When `glTF Export Options` window is shown, click `Export`
 
 If the Web Viewer is bundled, run `GLTFLaunchHelper.exe` in the target folder to view the export asset in the Unreal glTF Viewer.
+
+
+## Documentation
+
+- [What is glTF?](Docs/what-is-gltf.md)
+- [What can be exported?](Docs/what-can-be-exported.md)
+- [Export options reference](Docs/export-options-reference.md)
