[GLTFExporter] refactor: add some TODO comments

commit 1d3dae1aba73f95c14ac082f0cc75aac277dd2f0
Date: Sun, 17 Jan 2021 17:46:21 +0100

    refactor: add some TODO comments
    (cherry picked from commit 55b7c7f18d0674b9d1214e8ef7c5d1d5ec047266​)

#skipci

diff --git a/Engine/Plugins/Enterprise/GLTFExporter/MigratedFromGit/Source/GLTFExporter/Private/Converters/GLTFAccessorConverters.cpp b/Engine/Plugins/Enterprise/GLTFExporter/MigratedFromGit/Source/GLTFExporter/Private/Converters/GLTFAccessorConverters.cpp
--- a/Engine/Plugins/Enterprise/GLTFExporter/MigratedFromGit/Source/GLTFExporter/Private/Converters/GLTFAccessorConverters.cpp
+++ b/Engine/Plugins/Enterprise/GLTFExporter/MigratedFromGit/Source/GLTFExporter/Private/Converters/GLTFAccessorConverters.cpp
@@ -293,6 +293,8 @@ FGLTFJsonAccessorIndex FGLTFUVBufferConverter::Convert(const FGLTFMeshSection* M
 		return FGLTFJsonAccessorIndex(INDEX_NONE);
 	}
 
+	// TODO: report warning or add support for half float precision UVs, i.e. !VertexBuffer->GetUseFullPrecisionUVs()?
+
 	TArray<FGLTFJsonVector2> UVs;
 	UVs.AddUninitialized(VertexCount);
 
diff --git a/Engine/Plugins/Enterprise/GLTFExporter/MigratedFromGit/Source/GLTFExporter/Private/Converters/GLTFAnimationConverters.cpp b/Engine/Plugins/Enterprise/GLTFExporter/MigratedFromGit/Source/GLTFExporter/Private/Converters/GLTFAnimationConverters.cpp
--- a/Engine/Plugins/Enterprise/GLTFExporter/MigratedFromGit/Source/GLTFExporter/Private/Converters/GLTFAnimationConverters.cpp
+++ b/Engine/Plugins/Enterprise/GLTFExporter/MigratedFromGit/Source/GLTFExporter/Private/Converters/GLTFAnimationConverters.cpp
@@ -79,6 +79,7 @@ FGLTFJsonAnimationIndex FGLTFLevelSequenceDataConverter::Convert(const ALevelSeq
 		FGLTFJsonAnimation& JsonAnimation = Builder.GetAnimation(AnimationIndex);
 		FGLTFJsonAnimationPlayback& JsonPlayback = JsonAnimation.Playback;
 
+		// TODO: report warning if loop count is not 0 or -1 (infinite)
 		JsonPlayback.bLoop = LevelSequenceActor->PlaybackSettings.LoopCount.Value != 0;
 		JsonPlayback.bAutoPlay = LevelSequenceActor->PlaybackSettings.bAutoPlay;
 		JsonPlayback.PlayRate = LevelSequenceActor->PlaybackSettings.PlayRate;
diff --git a/Engine/Plugins/Enterprise/GLTFExporter/MigratedFromGit/Source/GLTFExporter/Private/Converters/GLTFNodeConverters.cpp b/Engine/Plugins/Enterprise/GLTFExporter/MigratedFromGit/Source/GLTFExporter/Private/Converters/GLTFNodeConverters.cpp
--- a/Engine/Plugins/Enterprise/GLTFExporter/MigratedFromGit/Source/GLTFExporter/Private/Converters/GLTFNodeConverters.cpp
+++ b/Engine/Plugins/Enterprise/GLTFExporter/MigratedFromGit/Source/GLTFExporter/Private/Converters/GLTFNodeConverters.cpp
@@ -112,6 +112,9 @@ FGLTFJsonNodeIndex FGLTFComponentConverter::Convert(const USceneComponent* Scene
 	Node.Rotation = FGLTFConverterUtility::ConvertRotation(RelativeTransform.GetRotation());
 	Node.Scale = FGLTFConverterUtility::ConvertScale(RelativeTransform.GetScale3D());
 
+	// TODO: don't export invisible components unless visibility is variable due to variant sets
+
+	// TODO: should hidden in game be configurable like this?
 	if (!SceneComponent->bHiddenInGame || Builder.ExportOptions->bExportHiddenInGame)
 	{
 		if (const UStaticMeshComponent* StaticMeshComponent = Cast<UStaticMeshComponent>(SceneComponent))
diff --git a/Engine/Plugins/Enterprise/GLTFExporter/MigratedFromGit/Source/GLTFExporter/Private/Tasks/GLTFAnimationTasks.cpp b/Engine/Plugins/Enterprise/GLTFExporter/MigratedFromGit/Source/GLTFExporter/Private/Tasks/GLTFAnimationTasks.cpp
--- a/Engine/Plugins/Enterprise/GLTFExporter/MigratedFromGit/Source/GLTFExporter/Private/Tasks/GLTFAnimationTasks.cpp
+++ b/Engine/Plugins/Enterprise/GLTFExporter/MigratedFromGit/Source/GLTFExporter/Private/Tasks/GLTFAnimationTasks.cpp
@@ -23,7 +23,7 @@ void FGLTFAnimSequenceTask::Complete()
 	const USkeleton* Skeleton = AnimSequence->GetSkeleton();
 
 	FGLTFJsonAnimation& JsonAnimation = Builder.GetAnimation(AnimationIndex);
-	JsonAnimation.Name = AnimSequence->GetName() + TEXT("_") + FString::FromInt(AnimationIndex); // Ensure unique name
+	JsonAnimation.Name = AnimSequence->GetName() + TEXT("_") + FString::FromInt(AnimationIndex); // Ensure unique name due to limitation in certain gltf viewers
 
 	TArray<float> Timestamps;
 	Timestamps.AddUninitialized(FrameCount);
@@ -33,6 +33,8 @@ void FGLTFAnimSequenceTask::Complete()
 		Timestamps[Frame] = AnimSequence->GetTimeAtFrame(Frame);
 	}
 
+	// TODO: add animation data accessor converters to reuse track information
+
 	FGLTFJsonAccessor JsonInputAccessor;
 	JsonInputAccessor.BufferView = Builder.AddBufferView(Timestamps);
 	JsonInputAccessor.ComponentType = EGLTFJsonComponentType::F32;
@@ -203,14 +205,14 @@ void FGLTFLevelSequenceTask::Complete()
 	Player->State.AssignSequence(MovieSceneSequenceID::Root, *Sequence, *Player);
 
 	FFrameRate TickResolution = MovieScene->GetTickResolution();
-	FFrameRate DisplayRate = MovieScene->GetDisplayRate();
+	FFrameRate DisplayRate = MovieScene->GetDisplayRate(); // TODO: add option to switch between DisplayRate and TickResolution?
 	TRange<FFrameNumber> PlaybackRange = MovieScene->GetPlaybackRange();
 
 	int32 FrameOffset = FFrameRate::TransformTime(FFrameTime(MovieScene::DiscreteInclusiveLower(PlaybackRange)), TickResolution, DisplayRate).RoundToFrame().Value;
 	int32 FrameCount = FFrameRate::TransformTime(FFrameTime(FFrameNumber(MovieScene::DiscreteSize(PlaybackRange))), TickResolution, DisplayRate).RoundToFrame().Value + 1;
 
 	FGLTFJsonAnimation& JsonAnimation = Builder.GetAnimation(AnimationIndex);
-	JsonAnimation.Name = Sequence->GetName() + TEXT("_") + FString::FromInt(AnimationIndex); // Ensure unique name
+	JsonAnimation.Name = Sequence->GetName() + TEXT("_") + FString::FromInt(AnimationIndex); // Ensure unique name due to limitation in certain gltf viewers
 
 	TArray<float> Timestamps;
 	TArray<FFrameTime> FrameTimes;
@@ -223,6 +225,8 @@ void FGLTFLevelSequenceTask::Complete()
 		FrameTimes[Frame] = FFrameRate::TransformTime(FFrameTime(FrameOffset + Frame), DisplayRate, TickResolution);
 	}
 
+	// TODO: add animation data accessor converters to reuse track information
+
 	FGLTFJsonAccessor JsonInputAccessor;
 	JsonInputAccessor.BufferView = Builder.AddBufferView(Timestamps);
 	JsonInputAccessor.ComponentType = EGLTFJsonComponentType::F32;
diff --git a/Engine/Plugins/Enterprise/GLTFExporter/MigratedFromGit/Source/GLTFExporter/Private/Tasks/GLTFMeshTasks.cpp b/Engine/Plugins/Enterprise/GLTFExporter/MigratedFromGit/Source/GLTFExporter/Private/Tasks/GLTFMeshTasks.cpp
--- a/Engine/Plugins/Enterprise/GLTFExporter/MigratedFromGit/Source/GLTFExporter/Private/Tasks/GLTFMeshTasks.cpp
+++ b/Engine/Plugins/Enterprise/GLTFExporter/MigratedFromGit/Source/GLTFExporter/Private/Tasks/GLTFMeshTasks.cpp
@@ -93,7 +93,7 @@ void FGLTFStaticMeshTask::Complete()
 	const FStaticMeshLODResources& MeshLOD = StaticMesh->GetLODForExport(LODIndex);
 	const FPositionVertexBuffer* PositionBuffer = &MeshLOD.VertexBuffers.PositionVertexBuffer;
 	const FStaticMeshVertexBuffer* VertexBuffer = &MeshLOD.VertexBuffers.StaticMeshVertexBuffer;
-	const FColorVertexBuffer* ColorBuffer = &MeshLOD.VertexBuffers.ColorVertexBuffer;
+	const FColorVertexBuffer* ColorBuffer = &MeshLOD.VertexBuffers.ColorVertexBuffer; // TODO: add support for overriding color buffer by component
 
 	if (Builder.ExportOptions->bExportVertexColors && HasVertexColors(ColorBuffer))
 	{
@@ -153,10 +153,12 @@ void FGLTFStaticMeshTask::Complete()
 			JsonPrimitive.Attributes.Color0 = Builder.GetOrAddColorAccessor(ConvertedSection, ColorBuffer);
 		}
 
+		// TODO: report warning if both Mesh Quantization (export options) and Use High Precision Tangent Basis (vertex buffer) are disabled
 		JsonPrimitive.Attributes.Normal = Builder.GetOrAddNormalAccessor(ConvertedSection, VertexBuffer);
 		JsonPrimitive.Attributes.Tangent = Builder.GetOrAddTangentAccessor(ConvertedSection, VertexBuffer);
 
 		const uint32 UVCount = VertexBuffer->GetNumTexCoords();
+		// TODO: report warning or option to limit UV channels since most viewers don't support more than 2?
 		JsonPrimitive.Attributes.TexCoords.AddUninitialized(UVCount);
 
 		for (uint32 UVIndex = 0; UVIndex < UVCount; ++UVIndex)
@@ -179,8 +181,10 @@ void FGLTFSkeletalMeshTask::Complete()
 
 	const FPositionVertexBuffer* PositionBuffer = &MeshLOD.StaticVertexBuffers.PositionVertexBuffer;
 	const FStaticMeshVertexBuffer* VertexBuffer = &MeshLOD.StaticVertexBuffers.StaticMeshVertexBuffer;
-	const FColorVertexBuffer* ColorBuffer = &MeshLOD.StaticVertexBuffers.ColorVertexBuffer;
-	const FSkinWeightVertexBuffer* SkinWeightBuffer = MeshLOD.GetSkinWeightVertexBuffer();
+	const FColorVertexBuffer* ColorBuffer = &MeshLOD.StaticVertexBuffers.ColorVertexBuffer; // TODO: add support for overriding color buffer by component
+	const FSkinWeightVertexBuffer* SkinWeightBuffer = MeshLOD.GetSkinWeightVertexBuffer(); // TODO: add support for overriding skin weight buffer by component
+	// TODO: add support for skin weight profiles?
+	// TODO: add support for morph targets
 
 	if (Builder.ExportOptions->bExportVertexColors && HasVertexColors(ColorBuffer))
 	{
@@ -241,10 +245,12 @@ void FGLTFSkeletalMeshTask::Complete()
 			JsonPrimitive.Attributes.Color0 = Builder.GetOrAddColorAccessor(ConvertedSection, ColorBuffer);
 		}
 
+		// TODO: report warning if both Mesh Quantization (export options) and Use High Precision Tangent Basis (vertex buffer) are disabled
 		JsonPrimitive.Attributes.Normal = Builder.GetOrAddNormalAccessor(ConvertedSection, VertexBuffer);
 		JsonPrimitive.Attributes.Tangent = Builder.GetOrAddTangentAccessor(ConvertedSection, VertexBuffer);
 
 		const uint32 UVCount = VertexBuffer->GetNumTexCoords();
+		// TODO: report warning or option to limit UV channels since most viewers don't support more than 2?
 		JsonPrimitive.Attributes.TexCoords.AddUninitialized(UVCount);
 
 		for (uint32 UVIndex = 0; UVIndex < UVCount; ++UVIndex)
@@ -255,6 +261,7 @@ void FGLTFSkeletalMeshTask::Complete()
 		if (Builder.ExportOptions->bExportVertexSkinWeights)
 		{
 			const uint32 GroupCount = (SkinWeightBuffer->GetMaxBoneInfluences() + 3) / 4;
+			// TODO: report warning or option to limit groups (of joints and weights) since most viewers don't support more than one?
 			JsonPrimitive.Attributes.Joints.AddUninitialized(GroupCount);
 			JsonPrimitive.Attributes.Weights.AddUninitialized(GroupCount);
 
diff --git a/Engine/Plugins/Enterprise/GLTFExporter/MigratedFromGit/Source/GLTFExporterRuntime/Public/Actors/GLTFCameraActor.h b/Engine/Plugins/Enterprise/GLTFExporter/MigratedFromGit/Source/GLTFExporterRuntime/Public/Actors/GLTFCameraActor.h
--- a/Engine/Plugins/Enterprise/GLTFExporter/MigratedFromGit/Source/GLTFExporterRuntime/Public/Actors/GLTFCameraActor.h
+++ b/Engine/Plugins/Enterprise/GLTFExporter/MigratedFromGit/Source/GLTFExporterRuntime/Public/Actors/GLTFCameraActor.h
@@ -69,7 +69,7 @@ public:
 	UPROPERTY(EditAnywhere, BlueprintReadWrite, Category = "GLTF Camera Actor")
 	EGLTFCameraControlMode Mode;
 
-	/* Actor which the camera will focus on and subsequently orbit when using Third Person mode. */
+	/* Actor which the camera will focus on and subsequently orbit when using Orbital mode. */
 	UPROPERTY(EditAnywhere, BlueprintReadWrite, Category = "GLTF Camera Actor", meta=(EditCondition="Mode == EGLTFCameraControlMode::Orbital"))
 	AActor* Target;
 
