[GLTFExporter] fix: use light color as multiplier also when using temperature.

commit 97c69eac046dc254713e22e3cb342256f4895884
Date: Sun, 17 Jan 2021 17:47:33 +0100

    fix: use light color as multiplier also when using temperature.
    (cherry picked from commit 2c5ae4092ce4eaa03ebd88a8995f83c2a0da7003​)

#skipci

diff --git a/Engine/Plugins/Enterprise/GLTFExporter/MigratedFromGit/Source/GLTFExporter/Private/Converters/GLTFLightConverters.cpp b/Engine/Plugins/Enterprise/GLTFExporter/MigratedFromGit/Source/GLTFExporter/Private/Converters/GLTFLightConverters.cpp
--- a/Engine/Plugins/Enterprise/GLTFExporter/MigratedFromGit/Source/GLTFExporter/Private/Converters/GLTFLightConverters.cpp
+++ b/Engine/Plugins/Enterprise/GLTFExporter/MigratedFromGit/Source/GLTFExporter/Private/Converters/GLTFLightConverters.cpp
@@ -20,8 +20,8 @@ FGLTFJsonLightIndex FGLTFLightConverter::Convert(const ULightComponent* LightCom
 
 	Light.Intensity = LightComponent->Intensity;
 
-	const FLinearColor LightColor = LightComponent->bUseTemperature ? FLinearColor::MakeFromColorTemperature(LightComponent->Temperature) : LightComponent->GetLightColor();
-	Light.Color = FGLTFConverterUtility::ConvertColor(LightColor);
+	const FLinearColor TemperatureColor = LightComponent->bUseTemperature ? FLinearColor::MakeFromColorTemperature(LightComponent->Temperature) : FLinearColor::White;
+	Light.Color = FGLTFConverterUtility::ConvertColor(TemperatureColor * LightComponent->GetLightColor());
 
 	if (const UPointLightComponent* PointLightComponent = Cast<UPointLightComponent>(LightComponent))
 	{
