[GLTFExporter] refactor: remove detailed version info (engine + plugin) in gltf output to avoid bumping all test assets every release

commit 0ac214700e0519133d0cab6362bba6d870133f12
Date: Sun, 17 Jan 2021 18:01:07 +0100

    refactor: remove detailed version info (engine + plugin) in gltf output to avoid bumping all test assets every release
    (cherry picked from commit d988e86265064daecd814ac4e17b22e9b0181e7a​)

#skipci

diff --git a/Engine/Plugins/Enterprise/GLTFExporter/MigratedFromGit/Source/GLTFExporter/Private/Json/GLTFJsonRoot.cpp b/Engine/Plugins/Enterprise/GLTFExporter/MigratedFromGit/Source/GLTFExporter/Private/Json/GLTFJsonRoot.cpp
--- a/Engine/Plugins/Enterprise/GLTFExporter/MigratedFromGit/Source/GLTFExporter/Private/Json/GLTFJsonRoot.cpp
+++ b/Engine/Plugins/Enterprise/GLTFExporter/MigratedFromGit/Source/GLTFExporter/Private/Json/GLTFJsonRoot.cpp
@@ -11,6 +11,7 @@ FGLTFJsonAsset::FGLTFJsonAsset()
 	const TSharedPtr<IPlugin> Plugin = IPluginManager::Get().FindPlugin(GLTFEXPORTER_MODULE_NAME);
 	const FPluginDescriptor& PluginDescriptor = Plugin->GetDescriptor();
 
-	Generator = TEXT(EPIC_PRODUCT_NAME) TEXT(" ") VERSION_STRINGIFY(ENGINE_MAJOR_VERSION) TEXT(".") VERSION_STRINGIFY(ENGINE_MINOR_VERSION)
-		TEXT(" ") + PluginDescriptor.FriendlyName + TEXT(" ") + PluginDescriptor.VersionName;
+	// TODO: add configuration option to include detailed engine and plugin version info?
+
+	Generator = TEXT(EPIC_PRODUCT_NAME) TEXT(" ") VERSION_STRINGIFY(ENGINE_MAJOR_VERSION) TEXT(" ") + PluginDescriptor.FriendlyName;
 }
