[GLTFExporter] chore: bump version to 1.0.1

commit ae1302352eb1bce1e82a6f1534ff2d2782d41f31
Date: Sun, 17 Jan 2021 18:01:38 +0100

    chore: bump version to 1.0.1
    (cherry picked from commit 0a84a880941fc5462be79182b32b2bac02492b9b​)

#skipci

diff --git a/Engine/Plugins/Enterprise/GLTFExporter/MigratedFromGit/GLTFExporter.uplugin b/Engine/Plugins/Enterprise/GLTFExporter/MigratedFromGit/GLTFExporter.uplugin
--- a/Engine/Plugins/Enterprise/GLTFExporter/MigratedFromGit/GLTFExporter.uplugin
+++ b/Engine/Plugins/Enterprise/GLTFExporter/MigratedFromGit/GLTFExporter.uplugin
@@ -1,7 +1,7 @@
 {
 	"FileVersion": 3,
-	"Version": 10,
-	"VersionName": "1.0",
+	"Version": 101,
+	"VersionName": "1.0.1",
 	"FriendlyName": "glTF Exporter",
 	"Description": "An exporter for Khronos glTF 2.0.",
 	"Category": "Exporters",
@@ -41,4 +41,4 @@
 			"Enabled": true
 		}
 	]
-}
\ No newline at end of file
+}
