[GLTFExporter] refactor: fix minor typos and code style issues

commit 9884298da627e8016857c653af41dcb745255b1b
Date: Sun, 21 Feb 2021 23:14:01 +0100

    refactor: fix minor typos and code style issues
    (cherry picked from commit 7ecabfea3f767d353c1b36b5d4d3473d9f9678bf​)

#skipci

diff --git a/Engine/Plugins/Enterprise/GLTFExporter/MigratedFromGit/Docs/export-options-reference.md b/Engine/Plugins/Enterprise/GLTFExporter/MigratedFromGit/Docs/export-options-reference.md
--- a/Engine/Plugins/Enterprise/GLTFExporter/MigratedFromGit/Docs/export-options-reference.md
+++ b/Engine/Plugins/Enterprise/GLTFExporter/MigratedFromGit/Docs/export-options-reference.md
@@ -9,12 +9,12 @@ Option                         | Description
 `Export Unlit Materials`       | If enabled, materials with shading model unlit will be properly exported. Uses extension KHR_materials_unlit.
 `Export Clear Coat Materials`  | If enabled, materials with shading model clear coat will be properly exported. Uses extension KHR_materials_clearcoat, which is not supported by all glTF viewers.
 `Export Extra Blend Modes`     | If enabled, materials with blend modes additive, modulate, and alpha composite will be properly exported. Uses extension EPIC_blend_modes, which is supported by Unreal's glTF viewer.
-`Bake Material Inputs`         | Bake mode determining if and how a material input is baked out to a texture. Can be overriden by material- and input-specific bake settings, see GLTFMaterialExportOptions.
-`Default Material Bake Size`   | Default size of the baked out texture (containing the material input). Can be overriden by material- and input-specific bake settings, see GLTFMaterialExportOptions.
-`Default Material Bake Filter` | Default filtering mode used when sampling the baked out texture. Can be overriden by material- and input-specific bake settings, see GLTFMaterialExportOptions.
-`Default Material Bake Tiling` | Default addressing mode used when sampling the baked out texture. Can be overriden by material- and input-specific bake settings, see GLTFMaterialExportOptions.
+`Bake Material Inputs`         | Bake mode determining if and how a material input is baked out to a texture. Can be overridden by material- and input-specific bake settings, see GLTFMaterialExportOptions.
+`Default Material Bake Size`   | Default size of the baked out texture (containing the material input). Can be overridden by material- and input-specific bake settings, see GLTFMaterialExportOptions.
+`Default Material Bake Filter` | Default filtering mode used when sampling the baked out texture. Can be overridden by material- and input-specific bake settings, see GLTFMaterialExportOptions.
+`Default Material Bake Tiling` | Default addressing mode used when sampling the baked out texture. Can be overridden by material- and input-specific bake settings, see GLTFMaterialExportOptions.
 `Default Input Bake Settings`  | Input-specific default bake settings that override the general defaults above.
-`Default Level Of Detail`      | Default LOD level used for exporting a mesh. Can be overriden by component or asset settings (e.g. minimum or forced LOD level).
+`Default Level Of Detail`      | Default LOD level used for exporting a mesh. Can be overridden by component or asset settings (e.g. minimum or forced LOD level).
 `Export Vertex Colors`         | If enabled, export vertex color. Not recommended due to vertex colors always being used as a base color multiplier in glTF, regardless of material. Often producing undesirable results.
 `Export Vertex Skin Weights`   | If enabled, export vertex bone weights and indices in skeletal meshes. Necessary for animation sequences.
 `Export Mesh Quantization`     | If enabled, export Unreal-configured quantization for vertex tangents and normals, reducing size. Requires extension KHR_mesh_quantization, which is not supported by all glTF viewers.
diff --git a/Engine/Plugins/Enterprise/GLTFExporter/MigratedFromGit/Source/GLTFExporter/Private/Builders/GLTFConvertBuilder.h b/Engine/Plugins/Enterprise/GLTFExporter/MigratedFromGit/Source/GLTFExporter/Private/Builders/GLTFConvertBuilder.h
--- a/Engine/Plugins/Enterprise/GLTFExporter/MigratedFromGit/Source/GLTFExporter/Private/Builders/GLTFConvertBuilder.h
+++ b/Engine/Plugins/Enterprise/GLTFExporter/MigratedFromGit/Source/GLTFExporter/Private/Builders/GLTFConvertBuilder.h
@@ -16,7 +16,7 @@
 #include "Converters/GLTFCameraConverters.h"
 #include "Converters/GLTFLightConverters.h"
 #include "Converters/GLTFBackdropConverters.h"
-#include "Converters/GLTFVarationConverters.h"
+#include "Converters/GLTFVariationConverters.h"
 #include "Converters/GLTFLightMapConverters.h"
 #include "Converters/GLTFHotspotConverters.h"
 #include "Converters/GLTFSkySphereConverters.h"
diff --git a/Engine/Plugins/Enterprise/GLTFExporter/MigratedFromGit/Source/GLTFExporter/Private/Converters/GLTFConverterUtility.h b/Engine/Plugins/Enterprise/GLTFExporter/MigratedFromGit/Source/GLTFExporter/Private/Converters/GLTFConverterUtility.h
--- a/Engine/Plugins/Enterprise/GLTFExporter/MigratedFromGit/Source/GLTFExporter/Private/Converters/GLTFConverterUtility.h
+++ b/Engine/Plugins/Enterprise/GLTFExporter/MigratedFromGit/Source/GLTFExporter/Private/Converters/GLTFConverterUtility.h
@@ -146,9 +146,9 @@ struct FGLTFConverterUtility
 		return ConvertMatrix(Matrix);
 	}
 
-	static float ConvertFieldOfView(float FOVInDegress, float AspectRatio)
+	static float ConvertFieldOfView(float FOVInDegrees, float AspectRatio)
 	{
-		const float HorizontalFOV = FMath::DegreesToRadians(FOVInDegress);
+		const float HorizontalFOV = FMath::DegreesToRadians(FOVInDegrees);
 		const float VerticalFOV = 2 * FMath::Atan(FMath::Tan(HorizontalFOV / 2) / AspectRatio);
 		return VerticalFOV;
 	}
diff --git a/Engine/Plugins/Enterprise/GLTFExporter/MigratedFromGit/Source/GLTFExporter/Private/Converters/GLTFCurveUtility.h b/Engine/Plugins/Enterprise/GLTFExporter/MigratedFromGit/Source/GLTFExporter/Private/Converters/GLTFCurveUtility.h
--- a/Engine/Plugins/Enterprise/GLTFExporter/MigratedFromGit/Source/GLTFExporter/Private/Converters/GLTFCurveUtility.h
+++ b/Engine/Plugins/Enterprise/GLTFExporter/MigratedFromGit/Source/GLTFExporter/Private/Converters/GLTFCurveUtility.h
@@ -2,8 +2,6 @@
 
 #pragma once
 
-#include "Engine/EngineTypes.h"
-
 class UCurveLinearColor;
 
 struct FGLTFCurveUtility
diff --git a/Engine/Plugins/Enterprise/GLTFExporter/MigratedFromGit/Source/GLTFExporter/Private/Converters/GLTFSkinWeightVertexBufferHack.h b/Engine/Plugins/Enterprise/GLTFExporter/MigratedFromGit/Source/GLTFExporter/Private/Converters/GLTFSkinWeightVertexBufferHack.h
--- a/Engine/Plugins/Enterprise/GLTFExporter/MigratedFromGit/Source/GLTFExporter/Private/Converters/GLTFSkinWeightVertexBufferHack.h
+++ b/Engine/Plugins/Enterprise/GLTFExporter/MigratedFromGit/Source/GLTFExporter/Private/Converters/GLTFSkinWeightVertexBufferHack.h
@@ -24,29 +24,23 @@ public:
 				FBoneIndex16* BoneIndex16Ptr = (FBoneIndex16*)BoneData;
 				return BoneIndex16Ptr[InfluenceIndex];
 			}
-			else
-			{
-				return BoneData[InfluenceIndex];
-			}
-		}
-		else
-		{
-			return 0;
+
+			return BoneData[InfluenceIndex];
 		}
+
+		return 0;
 	}
 
 	uint8 GetBoneWeight(uint32 VertexWeightOffset, uint32 VertexInfluenceCount, uint32 InfluenceIndex) const
 	{
 		if (InfluenceIndex < VertexInfluenceCount)
 		{
-			uint8* BoneData = Data + VertexWeightOffset;
-			uint32 BoneWeightOffset = GetBoneIndexByteSize() * VertexInfluenceCount;
+			const uint8* BoneData = Data + VertexWeightOffset;
+			const uint32 BoneWeightOffset = GetBoneIndexByteSize() * VertexInfluenceCount;
 			return BoneData[BoneWeightOffset + InfluenceIndex];
 		}
-		else
-		{
-			return 0;
-		}
+
+		return 0;
 	}
 
 	// guaranteed only to be valid if the vertex buffer is valid
@@ -83,8 +77,8 @@ public:
 
 	void GetWeightOffsetAndInfluenceCount(uint32 VertexIndex, uint32& OutWeightOffset, uint32& OutInfluenceCount) const
 	{
-		uint32 Offset = VertexIndex * 4;
-		uint32 DataUInt32 = *((uint32*)(&Data[Offset]));
+		const uint32 Offset = VertexIndex * 4;
+		const uint32 DataUInt32 = *((uint32*)(&Data[Offset]));
 		OutWeightOffset = DataUInt32 >> 8;
 		OutInfluenceCount = DataUInt32 & 0xff;
 	}
diff --git a/Engine/Plugins/Enterprise/GLTFExporter/MigratedFromGit/Source/GLTFExporter/Private/Converters/GLTFTextureConverters.cpp b/Engine/Plugins/Enterprise/GLTFExporter/MigratedFromGit/Source/GLTFExporter/Private/Converters/GLTFTextureConverters.cpp
--- a/Engine/Plugins/Enterprise/GLTFExporter/MigratedFromGit/Source/GLTFExporter/Private/Converters/GLTFTextureConverters.cpp
+++ b/Engine/Plugins/Enterprise/GLTFExporter/MigratedFromGit/Source/GLTFExporter/Private/Converters/GLTFTextureConverters.cpp
@@ -1,7 +1,6 @@
 // Copyright Epic Games, Inc. All Rights Reserved.
 
 #include "Converters/GLTFTextureConverters.h"
-#include "Builders/GLTFContainerBuilder.h"
 #include "Tasks/GLTFTextureTasks.h"
 
 FGLTFJsonTextureIndex FGLTFTexture2DConverter::Convert(const UTexture2D* Texture2D)
diff --git a/Engine/Plugins/Enterprise/GLTFExporter/MigratedFromGit/Source/GLTFExporter/Private/Converters/GLTFVarationConverters.cpp b/Engine/Plugins/Enterprise/GLTFExporter/MigratedFromGit/Source/GLTFExporter/Private/Converters/GLTFVarationConverters.cpp
deleted file mode 100644
diff --git a/Engine/Plugins/Enterprise/GLTFExporter/MigratedFromGit/Source/GLTFExporter/Private/Converters/GLTFVarationConverters.h b/Engine/Plugins/Enterprise/GLTFExporter/MigratedFromGit/Source/GLTFExporter/Private/Converters/GLTFVarationConverters.h
deleted file mode 100644
diff --git a/Engine/Plugins/Enterprise/GLTFExporter/MigratedFromGit/Source/GLTFExporter/Private/Converters/GLTFVariationConverters.cpp b/Engine/Plugins/Enterprise/GLTFExporter/MigratedFromGit/Source/GLTFExporter/Private/Converters/GLTFVariationConverters.cpp
new file mode 100644
--- /dev/null
+++ b/Engine/Plugins/Enterprise/GLTFExporter/MigratedFromGit/Source/GLTFExporter/Private/Converters/GLTFVariationConverters.cpp
@@ -0,0 +1,439 @@
+// Copyright Epic Games, Inc. All Rights Reserved.
+
+#include "Converters/GLTFVariationConverters.h"
+#include "Converters/GLTFMeshUtility.h"
+#include "Builders/GLTFContainerBuilder.h"
+#include "Rendering/SkeletalMeshRenderData.h"
+#include "VariantObjectBinding.h"
+#include "PropertyValueMaterial.h"
+#include "LevelVariantSets.h"
+#include "PropertyValue.h"
+#include "VariantSet.h"
+#include "Variant.h"
+
+namespace
+{
+	// TODO: replace hack with safe solution by adding proper API access to UPropertyValue
+	class UPropertyValueHack : public UPropertyValue
+	{
+	public:
+		const TArray<FCapturedPropSegment>& GetCapturedPropSegments() const
+		{
+			return CapturedPropSegments;
+		}
+	};
+}
+
+FGLTFJsonVariationIndex FGLTFVariationConverter::Convert(const ULevelVariantSets* LevelVariantSets)
+{
+	FGLTFJsonVariation JsonVariation;
+	LevelVariantSets->GetName(JsonVariation.Name);
+
+	for (const UVariantSet* VariantSet: LevelVariantSets->GetVariantSets())
+	{
+		FGLTFJsonVariantSet JsonVariantSet;
+		JsonVariantSet.Name = VariantSet->GetDisplayText().ToString();
+
+		for (const UVariant* Variant: VariantSet->GetVariants())
+		{
+			FGLTFJsonVariant JsonVariant;
+			if (TryParseVariant(JsonVariant, Variant))
+			{
+				JsonVariantSet.Variants.Add(JsonVariant);
+			}
+		}
+
+		if (JsonVariantSet.Variants.Num() > 0)
+		{
+			JsonVariation.VariantSets.Add(JsonVariantSet);
+		}
+	}
+
+	if (JsonVariation.VariantSets.Num() == 0)
+	{
+		return FGLTFJsonVariationIndex(INDEX_NONE);
+	}
+
+	return Builder.AddVariation(JsonVariation);
+}
+
+bool FGLTFVariationConverter::TryParseVariant(FGLTFJsonVariant& OutVariant, const UVariant* Variant) const
+{
+	FGLTFJsonVariant JsonVariant;
+
+	for (const UVariantObjectBinding* Binding: Variant->GetBindings())
+	{
+		TryParseVariantBinding(JsonVariant, Binding);
+	}
+
+	if (JsonVariant.Nodes.Num() == 0)
+	{
+		return false;
+	}
+
+	JsonVariant.Name = Variant->GetDisplayText().ToString();
+	JsonVariant.bIsActive = const_cast<UVariant*>(Variant)->IsActive();
+
+	if (const UTexture2D* Thumbnail = const_cast<UVariant*>(Variant)->GetThumbnail())
+	{
+		// TODO: if thumbnail has generic name "Texture2D", give it a variant-relevant name
+		JsonVariant.Thumbnail = Builder.GetOrAddTexture(Thumbnail);
+	}
+
+	OutVariant = JsonVariant;
+	return true;
+}
+
+bool FGLTFVariationConverter::TryParseVariantBinding(FGLTFJsonVariant& OutVariant, const UVariantObjectBinding* Binding) const
+{
+	bool bHasParsedAnyProperty = false;
+
+	for (UPropertyValue* Property: Binding->GetCapturedProperties())
+	{
+		if (!const_cast<UPropertyValue*>(Property)->Resolve() || !Property->HasRecordedData())
+		{
+			continue;
+		}
+
+		const FName PropertyName = Property->GetPropertyName();
+		const FFieldClass* PropertyClass = Property->GetPropertyClass();
+
+		if (Property->IsA<UPropertyValueMaterial>())
+		{
+			if (Builder.ExportOptions->ExportMaterialVariants != EGLTFMaterialVariantMode::None && TryParseMaterialPropertyValue(OutVariant, Property))
+			{
+				bHasParsedAnyProperty = true;
+			}
+		}
+		else if (PropertyName == TEXT("StaticMesh")) // TODO: should we not also check PropertyClass?
+		{
+			if (Builder.ExportOptions->bExportMeshVariants && TryParseStaticMeshPropertyValue(OutVariant, Property))
+			{
+				bHasParsedAnyProperty = true;
+			}
+		}
+		else if (PropertyName == TEXT("SkeletalMesh")) // TODO: should we not also check PropertyClass?
+		{
+			if (Builder.ExportOptions->bExportMeshVariants && TryParseSkeletalMeshPropertyValue(OutVariant, Property))
+			{
+				bHasParsedAnyProperty = true;
+			}
+		}
+		else if (PropertyName == TEXT("bVisible") && PropertyClass->IsChildOf(FBoolProperty::StaticClass()))
+		{
+			if (Builder.ExportOptions->bExportVisibilityVariants && TryParseVisibilityPropertyValue(OutVariant, Property))
+			{
+				bHasParsedAnyProperty = true;
+			}
+		}
+
+		// TODO: add support for more properties
+	}
+
+	return bHasParsedAnyProperty;
+}
+
+bool FGLTFVariationConverter::TryParseVisibilityPropertyValue(FGLTFJsonVariant& OutVariant, const UPropertyValue* Property) const
+{
+	const USceneComponent* Target = static_cast<USceneComponent*>(Property->GetPropertyParentContainerAddress());
+	if (Target == nullptr)
+	{
+		Builder.AddWarningMessage(FString::Printf(
+			TEXT("Target object for property is invalid, the property will be skipped. Context: %s"),
+			*GetLogContext(Property)));
+		return false;
+	}
+
+	const AActor* Owner = Target->GetOwner();
+	if (Owner == nullptr)
+	{
+		Builder.AddWarningMessage(FString::Printf(TEXT("Invalid scene component, the property will be skipped. Context: %s"),
+			*GetLogContext(Property)));
+		return false;
+	}
+
+	if (Builder.bSelectedActorsOnly && !Owner->IsSelected())
+	{
+		Builder.AddWarningMessage(FString::Printf(
+			TEXT("Target actor for property is not selected for export, the property will be skipped. Context: %s"),
+			*GetLogContext(Property)));
+		return false;
+	}
+
+	bool bIsVisible;
+
+	if (!TryGetPropertyValue(const_cast<UPropertyValue*>(Property), bIsVisible))
+	{
+		Builder.AddWarningMessage(FString::Printf(
+			TEXT("Failed to parse recorded data for property, it will be skipped. Context: %s"),
+			*GetLogContext(Property)));
+		return false;
+	}
+
+	const FGLTFJsonNodeIndex NodeIndex = Builder.GetOrAddNode(Target);
+	const FGLTFJsonNodeIndex ComponentNodeIndex = Builder.GetComponentNodeIndex(NodeIndex);
+
+	FGLTFJsonVariantNodeProperties& NodeProperties = OutVariant.Nodes.FindOrAdd(ComponentNodeIndex);
+
+	NodeProperties.Node = ComponentNodeIndex;
+	NodeProperties.bIsVisible = bIsVisible;
+	return true;
+}
+
+bool FGLTFVariationConverter::TryParseMaterialPropertyValue(FGLTFJsonVariant& OutVariant, const UPropertyValue* Property) const
+{
+	UPropertyValueMaterial* MaterialProperty = Cast<UPropertyValueMaterial>(const_cast<UPropertyValue*>(Property));
+	if (MaterialProperty == nullptr)
+	{
+		Builder.AddWarningMessage(FString::Printf(
+			TEXT("Material property is invalid, the property will be skipped. Context: %s"),
+			*GetLogContext(Property)));
+		return false;
+	}
+
+	const USceneComponent* Target = static_cast<USceneComponent*>(Property->GetPropertyParentContainerAddress());
+	if (Target == nullptr)
+	{
+		Builder.AddWarningMessage(FString::Printf(
+			TEXT("Target object for property is invalid, the property will be skipped. Context: %s"),
+			*GetLogContext(Property)));
+		return false;
+	}
+
+	const AActor* Owner = Target->GetOwner();
+	if (Owner == nullptr)
+	{
+		Builder.AddWarningMessage(FString::Printf(TEXT("Invalid scene component, the property will be skipped. Context: %s"),
+			*GetLogContext(Property)));
+		return false;
+	}
+
+	if (Builder.bSelectedActorsOnly && !Owner->IsSelected())
+	{
+		Builder.AddWarningMessage(FString::Printf(
+			TEXT("Target actor for property is not selected for export, the property will be skipped. Context: %s"),
+			*GetLogContext(Property)));
+		return false;
+	}
+
+	const TArray<FCapturedPropSegment>& CapturedPropSegments = reinterpret_cast<UPropertyValueHack*>(MaterialProperty)->GetCapturedPropSegments();
+	const int32 NumPropSegments = CapturedPropSegments.Num();
+
+	if (NumPropSegments < 1)
+	{
+		Builder.AddWarningMessage(FString::Printf(
+			TEXT("Failed to parse element index to apply the material to, the property will be skipped. Context: %s"),
+			*GetLogContext(Property)));
+		return false;
+	}
+
+	// NOTE: UPropertyValueMaterial::GetMaterial does *not* ensure that the recorded data has been loaded,
+	// so we need to call UProperty::GetRecordedData first to make that happen.
+	MaterialProperty->GetRecordedData();
+
+	// TODO: find way to determine whether the material is null because "None" was selected, or because it failed to resolve
+	const UMaterialInterface* Material = MaterialProperty->GetMaterial();
+	const int32 MaterialIndex = CapturedPropSegments[NumPropSegments - 1].PropertyIndex;
+
+	const FGLTFMeshData* MeshData = nullptr;
+	FGLTFIndexArray SectionIndices;
+
+	if (Builder.ExportOptions->ExportMaterialVariants == EGLTFMaterialVariantMode::UseMeshData)
+	{
+		if (Builder.ExportOptions->BakeMaterialInputs == EGLTFMaterialBakeMode::UseMeshData)
+		{
+			const int32 DefaultLODIndex = Builder.ExportOptions->DefaultLevelOfDetail;
+
+			if (const UStaticMeshComponent* StaticMeshComponent = Cast<UStaticMeshComponent>(Target))
+			{
+				const UStaticMesh* StaticMesh = StaticMeshComponent->GetStaticMesh();
+				const int32 LODIndex = FGLTFMeshUtility::GetLOD(StaticMesh, StaticMeshComponent, DefaultLODIndex);
+				const FStaticMeshLODResources& MeshLOD = StaticMesh->GetLODForExport(LODIndex);
+
+				SectionIndices = FGLTFMeshUtility::GetSectionIndices(MeshLOD, MaterialIndex);
+				MeshData = Builder.StaticMeshDataConverter.GetOrAdd(StaticMesh, StaticMeshComponent, LODIndex);
+			}
+			else if (const USkeletalMeshComponent* SkeletalMeshComponent = Cast<USkeletalMeshComponent>(Target))
+			{
+				const USkeletalMesh* SkeletalMesh = SkeletalMeshComponent->SkeletalMesh;
+				const int32 LODIndex = FGLTFMeshUtility::GetLOD(SkeletalMesh, SkeletalMeshComponent, DefaultLODIndex);
+				const FSkeletalMeshRenderData* RenderData = SkeletalMesh->GetResourceForRendering();
+				const FSkeletalMeshLODRenderData& MeshLOD = RenderData->LODRenderData[LODIndex];
+
+				SectionIndices = FGLTFMeshUtility::GetSectionIndices(MeshLOD, MaterialIndex);
+				MeshData = Builder.SkeletalMeshDataConverter.GetOrAdd(SkeletalMesh, SkeletalMeshComponent, LODIndex);
+			}
+		}
+		else
+		{
+			// TODO: report warning (about materials won't be export using mesh data because BakeMaterialInputs is not set to UseMeshData)
+		}
+	}
+
+	FGLTFJsonVariantMaterial VariantMaterial;
+	VariantMaterial.Material = Builder.GetOrAddMaterial(Material, MeshData, SectionIndices);
+	VariantMaterial.Index = MaterialIndex;
+
+	const FGLTFJsonNodeIndex NodeIndex = Builder.GetOrAddNode(Target);
+	const FGLTFJsonNodeIndex ComponentNodeIndex = Builder.GetComponentNodeIndex(NodeIndex);
+	FGLTFJsonVariantNodeProperties& NodeProperties = OutVariant.Nodes.FindOrAdd(ComponentNodeIndex);
+
+	NodeProperties.Node = ComponentNodeIndex;
+	NodeProperties.Materials.Add(VariantMaterial);
+	return true;
+}
+
+bool FGLTFVariationConverter::TryParseStaticMeshPropertyValue(FGLTFJsonVariant& OutVariant, const UPropertyValue* Property) const
+{
+	const USceneComponent* Target = static_cast<USceneComponent*>(Property->GetPropertyParentContainerAddress());
+	if (Target == nullptr)
+	{
+		Builder.AddWarningMessage(FString::Printf(
+			TEXT("Target object for property is invalid, the property will be skipped. Context: %s"),
+			*GetLogContext(Property)));
+		return false;
+	}
+
+	const AActor* Owner = Target->GetOwner();
+	if (Owner == nullptr)
+	{
+		Builder.AddWarningMessage(FString::Printf(TEXT("Invalid scene component, the property will be skipped. Context: %s"),
+			*GetLogContext(Property)));
+		return false;
+	}
+
+	if (Builder.bSelectedActorsOnly && !Owner->IsSelected())
+	{
+		Builder.AddWarningMessage(FString::Printf(
+			TEXT("Target actor for property is not selected for export, the property will be skipped. Context: %s"),
+			*GetLogContext(Property)));
+		return false;
+	}
+
+	const UMeshComponent* MeshComponent = Cast<UMeshComponent>(Target);
+	if (MeshComponent == nullptr)
+	{
+		Builder.AddWarningMessage(FString::Printf(
+			TEXT("Target object for property has no mesh-component, the property will be skipped. Context: %s"),
+			*GetLogContext(Property)));
+		return false;
+	}
+
+	const UStaticMesh* StaticMesh;
+
+	if (!TryGetPropertyValue(const_cast<UPropertyValue*>(Property), StaticMesh))
+	{
+		Builder.AddWarningMessage(FString::Printf(
+			TEXT("Failed to parse recorded data for property, it will be skipped. Context: %s"),
+			*GetLogContext(Property)));
+		return false;
+	}
+
+	const FGLTFMaterialArray OverrideMaterials(MeshComponent->OverrideMaterials);
+	const FGLTFJsonNodeIndex NodeIndex = Builder.GetOrAddNode(Target);
+	const FGLTFJsonNodeIndex ComponentNodeIndex = Builder.GetComponentNodeIndex(NodeIndex);
+	const FGLTFJsonMeshIndex MeshIndex = Builder.GetOrAddMesh(StaticMesh, OverrideMaterials);
+	FGLTFJsonVariantNodeProperties& NodeProperties = OutVariant.Nodes.FindOrAdd(ComponentNodeIndex);
+
+	NodeProperties.Node = ComponentNodeIndex;
+	NodeProperties.Mesh = MeshIndex;
+	return true;
+}
+
+bool FGLTFVariationConverter::TryParseSkeletalMeshPropertyValue(FGLTFJsonVariant& OutVariant, const UPropertyValue* Property) const
+{
+	const USceneComponent* Target = static_cast<USceneComponent*>(Property->GetPropertyParentContainerAddress());
+	if (Target == nullptr)
+	{
+		Builder.AddWarningMessage(FString::Printf(
+			TEXT("Target object for property is invalid, the property will be skipped. Context: %s"),
+			*GetLogContext(Property)));
+		return false;
+	}
+
+	const AActor* Owner = Target->GetOwner();
+	if (Owner == nullptr)
+	{
+		Builder.AddWarningMessage(FString::Printf(TEXT("Invalid scene component, the property will be skipped. Context: %s"),
+			*GetLogContext(Property)));
+		return false;
+	}
+
+	if (Builder.bSelectedActorsOnly && !Owner->IsSelected())
+	{
+		Builder.AddWarningMessage(FString::Printf(
+			TEXT("Target actor for property is not selected for export, the property will be skipped. Context: %s"),
+			*GetLogContext(Property)));
+		return false;
+	}
+
+	const UMeshComponent* MeshComponent = Cast<UMeshComponent>(Target);
+	if (MeshComponent == nullptr)
+	{
+		Builder.AddWarningMessage(FString::Printf(
+			TEXT("Target object for property has no mesh-component, the property will be skipped. Context: %s"),
+			*GetLogContext(Property)));
+		return false;
+	}
+
+	const USkeletalMesh* SkeletalMesh;
+
+	if (!TryGetPropertyValue(const_cast<UPropertyValue*>(Property), SkeletalMesh) || SkeletalMesh == nullptr)
+	{
+		Builder.AddWarningMessage(FString::Printf(
+			TEXT("Failed to parse recorded data for property, it will be skipped. Context: %s"),
+			*GetLogContext(Property)));
+		return false;
+	}
+
+	const FGLTFMaterialArray OverrideMaterials(MeshComponent->OverrideMaterials);
+	const FGLTFJsonNodeIndex NodeIndex = Builder.GetOrAddNode(Target);
+	const FGLTFJsonMeshIndex MeshIndex = Builder.GetOrAddMesh(SkeletalMesh, OverrideMaterials);
+	const FGLTFJsonNodeIndex ComponentNodeIndex = Builder.GetComponentNodeIndex(NodeIndex);
+	FGLTFJsonVariantNodeProperties& NodeProperties = OutVariant.Nodes.FindOrAdd(ComponentNodeIndex);
+
+	NodeProperties.Node = ComponentNodeIndex;
+	NodeProperties.Mesh = MeshIndex;
+	return true;
+}
+
+template<typename T>
+bool FGLTFVariationConverter::TryGetPropertyValue(UPropertyValue* Property, T& OutValue) const
+{
+	if (Property == nullptr || !Property->HasRecordedData())
+	{
+		return false;
+	}
+
+	FMemory::Memcpy(&OutValue, Property->GetRecordedData().GetData(), sizeof(T));
+	return true;
+}
+
+FString FGLTFVariationConverter::GetLogContext(const UPropertyValue* Property) const
+{
+	const UVariantObjectBinding* Parent = Property->GetParent();
+	return GetLogContext(Parent) + TEXT("/") + Property->GetFullDisplayString();
+}
+
+FString FGLTFVariationConverter::GetLogContext(const UVariantObjectBinding* Binding) const
+{
+	const UVariant* Parent = const_cast<UVariantObjectBinding*>(Binding)->GetParent();
+	return GetLogContext(Parent) + TEXT("/") + Binding->GetDisplayText().ToString();
+}
+
+FString FGLTFVariationConverter::GetLogContext(const UVariant* Variant) const
+{
+	const UVariantSet* Parent = const_cast<UVariant*>(Variant)->GetParent();
+	return GetLogContext(Parent) + TEXT("/") + Variant->GetDisplayText().ToString();
+}
+
+FString FGLTFVariationConverter::GetLogContext(const UVariantSet* VariantSet) const
+{
+	const ULevelVariantSets* Parent = const_cast<UVariantSet*>(VariantSet)->GetParent();
+	return GetLogContext(Parent) + TEXT("/") + VariantSet->GetDisplayText().ToString();
+}
+
+FString FGLTFVariationConverter::GetLogContext(const ULevelVariantSets* LevelVariantSets) const
+{
+	return LevelVariantSets->GetName();
+}
diff --git a/Engine/Plugins/Enterprise/GLTFExporter/MigratedFromGit/Source/GLTFExporter/Private/Converters/GLTFVariationConverters.h b/Engine/Plugins/Enterprise/GLTFExporter/MigratedFromGit/Source/GLTFExporter/Private/Converters/GLTFVariationConverters.h
new file mode 100644
--- /dev/null
+++ b/Engine/Plugins/Enterprise/GLTFExporter/MigratedFromGit/Source/GLTFExporter/Private/Converters/GLTFVariationConverters.h
@@ -0,0 +1,35 @@
+// Copyright Epic Games, Inc. All Rights Reserved.
+
+#pragma once
+
+#include "Json/GLTFJsonIndex.h"
+#include "Json/GLTFJsonVariation.h"
+#include "Converters/GLTFConverter.h"
+#include "Converters/GLTFBuilderContext.h"
+#include "Engine.h"
+#include "Variant.h"
+#include "PropertyValue.h"
+#include "LevelVariantSets.h"
+
+class FGLTFVariationConverter final : public FGLTFBuilderContext, public TGLTFConverter<FGLTFJsonVariationIndex, const ULevelVariantSets*>
+{
+	using FGLTFBuilderContext::FGLTFBuilderContext;
+
+	virtual FGLTFJsonVariationIndex Convert(const ULevelVariantSets* LevelVariantSets) override;
+
+	bool TryParseVariant(FGLTFJsonVariant& OutVariant, const UVariant* Variant) const;
+	bool TryParseVariantBinding(FGLTFJsonVariant& OutVariant, const UVariantObjectBinding* Binding) const;
+	bool TryParseVisibilityPropertyValue(FGLTFJsonVariant& OutVariant, const UPropertyValue* Property) const;
+	bool TryParseMaterialPropertyValue(FGLTFJsonVariant& OutVariant, const UPropertyValue* Property) const;
+	bool TryParseStaticMeshPropertyValue(FGLTFJsonVariant& OutVariant, const UPropertyValue* Property) const;
+	bool TryParseSkeletalMeshPropertyValue(FGLTFJsonVariant& OutVariant, const UPropertyValue* Property) const;
+
+	template<typename T>
+	bool TryGetPropertyValue(UPropertyValue* Property, T& OutValue) const;
+
+	FString GetLogContext(const UPropertyValue* Property) const;
+	FString GetLogContext(const UVariantObjectBinding* Binding) const;
+	FString GetLogContext(const UVariant* Variant) const;
+	FString GetLogContext(const UVariantSet* VariantSet) const;
+	FString GetLogContext(const ULevelVariantSets* LevelVariantSets) const;
+};
diff --git a/Engine/Plugins/Enterprise/GLTFExporter/MigratedFromGit/Source/GLTFExporter/Private/Json/GLTFJsonBuffer.h b/Engine/Plugins/Enterprise/GLTFExporter/MigratedFromGit/Source/GLTFExporter/Private/Json/GLTFJsonBuffer.h
--- a/Engine/Plugins/Enterprise/GLTFExporter/MigratedFromGit/Source/GLTFExporter/Private/Json/GLTFJsonBuffer.h
+++ b/Engine/Plugins/Enterprise/GLTFExporter/MigratedFromGit/Source/GLTFExporter/Private/Json/GLTFJsonBuffer.h
@@ -2,7 +2,6 @@
 
 #pragma once
 
-#include "Json/GLTFJsonIndex.h"
 #include "Serialization/JsonSerializer.h"
 
 struct FGLTFJsonBuffer
diff --git a/Engine/Plugins/Enterprise/GLTFExporter/MigratedFromGit/Source/GLTFExporter/Private/Json/GLTFJsonColor4.h b/Engine/Plugins/Enterprise/GLTFExporter/MigratedFromGit/Source/GLTFExporter/Private/Json/GLTFJsonColor4.h
--- a/Engine/Plugins/Enterprise/GLTFExporter/MigratedFromGit/Source/GLTFExporter/Private/Json/GLTFJsonColor4.h
+++ b/Engine/Plugins/Enterprise/GLTFExporter/MigratedFromGit/Source/GLTFExporter/Private/Json/GLTFJsonColor4.h
@@ -1,4 +1,4 @@
-// CopGright Epic Games, Inc. All Rights Reserved.
+// Copyright Epic Games, Inc. All Rights Reserved.
 
 #pragma once
 
diff --git a/Engine/Plugins/Enterprise/GLTFExporter/MigratedFromGit/Source/GLTFExporter/Private/Tasks/GLTFAnimationTasks.cpp b/Engine/Plugins/Enterprise/GLTFExporter/MigratedFromGit/Source/GLTFExporter/Private/Tasks/GLTFAnimationTasks.cpp
--- a/Engine/Plugins/Enterprise/GLTFExporter/MigratedFromGit/Source/GLTFExporter/Private/Tasks/GLTFAnimationTasks.cpp
+++ b/Engine/Plugins/Enterprise/GLTFExporter/MigratedFromGit/Source/GLTFExporter/Private/Tasks/GLTFAnimationTasks.cpp
@@ -147,7 +147,7 @@ void FGLTFAnimSequenceTask::Complete()
 			JsonOutputAccessor.Type = EGLTFJsonAccessorType::Vec4;
 
 			FGLTFJsonAnimationSampler JsonSampler;
-			JsonSampler.Input = Builder.AddAccessor(JsonInputAccessor);;
+			JsonSampler.Input = Builder.AddAccessor(JsonInputAccessor);
 			JsonSampler.Output = Builder.AddAccessor(JsonOutputAccessor);
 			JsonSampler.Interpolation = Interpolation;
 
diff --git a/Engine/Plugins/Enterprise/GLTFExporter/MigratedFromGit/Source/GLTFExporter/Private/UI/GLTFExportOptionsWindow.h b/Engine/Plugins/Enterprise/GLTFExporter/MigratedFromGit/Source/GLTFExporter/Private/UI/GLTFExportOptionsWindow.h
--- a/Engine/Plugins/Enterprise/GLTFExporter/MigratedFromGit/Source/GLTFExporter/Private/UI/GLTFExportOptionsWindow.h
+++ b/Engine/Plugins/Enterprise/GLTFExporter/MigratedFromGit/Source/GLTFExporter/Private/UI/GLTFExportOptionsWindow.h
@@ -3,7 +3,6 @@
 #pragma once
 
 #include "GLTFExportOptions.h"
-#include "InputCoreTypes.h"
 #include "Widgets/DeclarativeSyntaxSupport.h"
 #include "Input/Reply.h"
 #include "Widgets/SCompoundWidget.h"
diff --git a/Engine/Plugins/Enterprise/GLTFExporter/MigratedFromGit/Source/GLTFExporter/Public/Exporters/GLTFAnimSequenceExporter.h b/Engine/Plugins/Enterprise/GLTFExporter/MigratedFromGit/Source/GLTFExporter/Public/Exporters/GLTFAnimSequenceExporter.h
--- a/Engine/Plugins/Enterprise/GLTFExporter/MigratedFromGit/Source/GLTFExporter/Public/Exporters/GLTFAnimSequenceExporter.h
+++ b/Engine/Plugins/Enterprise/GLTFExporter/MigratedFromGit/Source/GLTFExporter/Public/Exporters/GLTFAnimSequenceExporter.h
@@ -14,5 +14,5 @@ public:
 
 	explicit UGLTFAnimSequenceExporter(const FObjectInitializer& ObjectInitializer = FObjectInitializer());
 
-	bool AddObject(FGLTFContainerBuilder& Builder, const UObject* Object) override;
+	virtual bool AddObject(FGLTFContainerBuilder& Builder, const UObject* Object) override;
 };
diff --git a/Engine/Plugins/Enterprise/GLTFExporter/MigratedFromGit/Source/GLTFExporter/Public/Exporters/GLTFLevelExporter.h b/Engine/Plugins/Enterprise/GLTFExporter/MigratedFromGit/Source/GLTFExporter/Public/Exporters/GLTFLevelExporter.h
--- a/Engine/Plugins/Enterprise/GLTFExporter/MigratedFromGit/Source/GLTFExporter/Public/Exporters/GLTFLevelExporter.h
+++ b/Engine/Plugins/Enterprise/GLTFExporter/MigratedFromGit/Source/GLTFExporter/Public/Exporters/GLTFLevelExporter.h
@@ -14,5 +14,5 @@ public:
 
 	explicit UGLTFLevelExporter(const FObjectInitializer& ObjectInitializer = FObjectInitializer());
 
-	bool AddObject(FGLTFContainerBuilder& Builder, const UObject* Object) override;
+	virtual bool AddObject(FGLTFContainerBuilder& Builder, const UObject* Object) override;
 };
diff --git a/Engine/Plugins/Enterprise/GLTFExporter/MigratedFromGit/Source/GLTFExporter/Public/Exporters/GLTFLevelSequenceExporter.h b/Engine/Plugins/Enterprise/GLTFExporter/MigratedFromGit/Source/GLTFExporter/Public/Exporters/GLTFLevelSequenceExporter.h
--- a/Engine/Plugins/Enterprise/GLTFExporter/MigratedFromGit/Source/GLTFExporter/Public/Exporters/GLTFLevelSequenceExporter.h
+++ b/Engine/Plugins/Enterprise/GLTFExporter/MigratedFromGit/Source/GLTFExporter/Public/Exporters/GLTFLevelSequenceExporter.h
@@ -14,5 +14,5 @@ public:
 
 	explicit UGLTFLevelSequenceExporter(const FObjectInitializer& ObjectInitializer = FObjectInitializer());
 
-	bool AddObject(FGLTFContainerBuilder& Builder, const UObject* Object) override;
+	virtual bool AddObject(FGLTFContainerBuilder& Builder, const UObject* Object) override;
 };
diff --git a/Engine/Plugins/Enterprise/GLTFExporter/MigratedFromGit/Source/GLTFExporter/Public/Exporters/GLTFLevelVariantSetsExporter.h b/Engine/Plugins/Enterprise/GLTFExporter/MigratedFromGit/Source/GLTFExporter/Public/Exporters/GLTFLevelVariantSetsExporter.h
--- a/Engine/Plugins/Enterprise/GLTFExporter/MigratedFromGit/Source/GLTFExporter/Public/Exporters/GLTFLevelVariantSetsExporter.h
+++ b/Engine/Plugins/Enterprise/GLTFExporter/MigratedFromGit/Source/GLTFExporter/Public/Exporters/GLTFLevelVariantSetsExporter.h
@@ -14,5 +14,5 @@ public:
 
 	explicit UGLTFLevelVariantSetsExporter(const FObjectInitializer& ObjectInitializer = FObjectInitializer());
 
-	bool AddObject(FGLTFContainerBuilder& Builder, const UObject* Object) override;
+	virtual bool AddObject(FGLTFContainerBuilder& Builder, const UObject* Object) override;
 };
diff --git a/Engine/Plugins/Enterprise/GLTFExporter/MigratedFromGit/Source/GLTFExporter/Public/Exporters/GLTFMaterialExporter.h b/Engine/Plugins/Enterprise/GLTFExporter/MigratedFromGit/Source/GLTFExporter/Public/Exporters/GLTFMaterialExporter.h
--- a/Engine/Plugins/Enterprise/GLTFExporter/MigratedFromGit/Source/GLTFExporter/Public/Exporters/GLTFMaterialExporter.h
+++ b/Engine/Plugins/Enterprise/GLTFExporter/MigratedFromGit/Source/GLTFExporter/Public/Exporters/GLTFMaterialExporter.h
@@ -14,5 +14,5 @@ public:
 
 	explicit UGLTFMaterialExporter(const FObjectInitializer& ObjectInitializer = FObjectInitializer());
 
-	bool AddObject(FGLTFContainerBuilder& Builder, const UObject* Object) override;
+	virtual bool AddObject(FGLTFContainerBuilder& Builder, const UObject* Object) override;
 };
diff --git a/Engine/Plugins/Enterprise/GLTFExporter/MigratedFromGit/Source/GLTFExporter/Public/Exporters/GLTFSkeletalMeshExporter.h b/Engine/Plugins/Enterprise/GLTFExporter/MigratedFromGit/Source/GLTFExporter/Public/Exporters/GLTFSkeletalMeshExporter.h
--- a/Engine/Plugins/Enterprise/GLTFExporter/MigratedFromGit/Source/GLTFExporter/Public/Exporters/GLTFSkeletalMeshExporter.h
+++ b/Engine/Plugins/Enterprise/GLTFExporter/MigratedFromGit/Source/GLTFExporter/Public/Exporters/GLTFSkeletalMeshExporter.h
@@ -14,5 +14,5 @@ public:
 
 	explicit UGLTFSkeletalMeshExporter(const FObjectInitializer& ObjectInitializer = FObjectInitializer());
 
-	bool AddObject(FGLTFContainerBuilder& Builder, const UObject* Object) override;
+	virtual bool AddObject(FGLTFContainerBuilder& Builder, const UObject* Object) override;
 };
diff --git a/Engine/Plugins/Enterprise/GLTFExporter/MigratedFromGit/Source/GLTFExporter/Public/Exporters/GLTFStaticMeshExporter.h b/Engine/Plugins/Enterprise/GLTFExporter/MigratedFromGit/Source/GLTFExporter/Public/Exporters/GLTFStaticMeshExporter.h
--- a/Engine/Plugins/Enterprise/GLTFExporter/MigratedFromGit/Source/GLTFExporter/Public/Exporters/GLTFStaticMeshExporter.h
+++ b/Engine/Plugins/Enterprise/GLTFExporter/MigratedFromGit/Source/GLTFExporter/Public/Exporters/GLTFStaticMeshExporter.h
@@ -14,5 +14,5 @@ public:
 
 	explicit UGLTFStaticMeshExporter(const FObjectInitializer& ObjectInitializer = FObjectInitializer());
 
-	bool AddObject(FGLTFContainerBuilder& Builder, const UObject* Object) override;
+	virtual bool AddObject(FGLTFContainerBuilder& Builder, const UObject* Object) override;
 };
diff --git a/Engine/Plugins/Enterprise/GLTFExporter/MigratedFromGit/Source/GLTFExporter/Public/GLTFExportOptions.h b/Engine/Plugins/Enterprise/GLTFExporter/MigratedFromGit/Source/GLTFExporter/Public/GLTFExportOptions.h
--- a/Engine/Plugins/Enterprise/GLTFExporter/MigratedFromGit/Source/GLTFExporter/Public/GLTFExportOptions.h
+++ b/Engine/Plugins/Enterprise/GLTFExporter/MigratedFromGit/Source/GLTFExporter/Public/GLTFExportOptions.h
@@ -96,19 +96,19 @@ class GLTFEXPORTER_API UGLTFExportOptions : public UObject
 	UPROPERTY(EditAnywhere, BlueprintReadWrite, Config, Category = Material)
 	bool bExportExtraBlendModes;
 
-	/** Bake mode determining if and how a material input is baked out to a texture. Can be overriden by material- and input-specific bake settings, see GLTFMaterialExportOptions. */
+	/** Bake mode determining if and how a material input is baked out to a texture. Can be overridden by material- and input-specific bake settings, see GLTFMaterialExportOptions. */
 	UPROPERTY(EditAnywhere, BlueprintReadWrite, Config, Category = Material)
 	EGLTFMaterialBakeMode BakeMaterialInputs;
 
-	/** Default size of the baked out texture (containing the material input). Can be overriden by material- and input-specific bake settings, see GLTFMaterialExportOptions. */
+	/** Default size of the baked out texture (containing the material input). Can be overridden by material- and input-specific bake settings, see GLTFMaterialExportOptions. */
 	UPROPERTY(EditAnywhere, BlueprintReadWrite, Config, Category = Material, Meta = (EditCondition = "BakeMaterialInputs != EGLTFMaterialBakeMode::Disabled"))
 	EGLTFMaterialBakeSizePOT DefaultMaterialBakeSize;
 
-	/** Default filtering mode used when sampling the baked out texture. Can be overriden by material- and input-specific bake settings, see GLTFMaterialExportOptions. */
+	/** Default filtering mode used when sampling the baked out texture. Can be overridden by material- and input-specific bake settings, see GLTFMaterialExportOptions. */
 	UPROPERTY(EditAnywhere, BlueprintReadWrite, Config, Category = Material, Meta = (EditCondition = "BakeMaterialInputs != EGLTFMaterialBakeMode::Disabled", ValidEnumValues="TF_Nearest, TF_Bilinear, TF_Trilinear"))
 	TEnumAsByte<TextureFilter> DefaultMaterialBakeFilter;
 
-	/** Default addressing mode used when sampling the baked out texture. Can be overriden by material- and input-specific bake settings, see GLTFMaterialExportOptions. */
+	/** Default addressing mode used when sampling the baked out texture. Can be overridden by material- and input-specific bake settings, see GLTFMaterialExportOptions. */
 	UPROPERTY(EditAnywhere, BlueprintReadWrite, Config, Category = Material, Meta = (EditCondition = "BakeMaterialInputs != EGLTFMaterialBakeMode::Disabled"))
 	TEnumAsByte<TextureAddress> DefaultMaterialBakeTiling;
 
@@ -116,7 +116,7 @@ class GLTFEXPORTER_API UGLTFExportOptions : public UObject
 	UPROPERTY(EditAnywhere, BlueprintReadWrite, Config, Category = Material)
 	TMap<EGLTFMaterialPropertyGroup, FGLTFOverrideMaterialBakeSettings> DefaultInputBakeSettings;
 
-	/** Default LOD level used for exporting a mesh. Can be overriden by component or asset settings (e.g. minimum or forced LOD level). */
+	/** Default LOD level used for exporting a mesh. Can be overridden by component or asset settings (e.g. minimum or forced LOD level). */
 	UPROPERTY(EditAnywhere, BlueprintReadWrite, Config, Category = Mesh, Meta = (ClampMin = "0"))
 	int32 DefaultLevelOfDetail;
 
diff --git a/Engine/Plugins/Enterprise/GLTFExporter/MigratedFromGit/Source/GLTFExporterRuntime/Private/Actors/GLTFHotspotActor.cpp b/Engine/Plugins/Enterprise/GLTFExporter/MigratedFromGit/Source/GLTFExporterRuntime/Private/Actors/GLTFHotspotActor.cpp
--- a/Engine/Plugins/Enterprise/GLTFExporter/MigratedFromGit/Source/GLTFExporterRuntime/Private/Actors/GLTFHotspotActor.cpp
+++ b/Engine/Plugins/Enterprise/GLTFExporter/MigratedFromGit/Source/GLTFExporterRuntime/Private/Actors/GLTFHotspotActor.cpp
@@ -391,10 +391,8 @@ FIntPoint AGLTFHotspotActor::GetCurrentViewportSize()
 
 		return Viewport->GetSizeXY();
 	}
-	else
-	{
-		return FIntPoint(0, 0);
-	}
+
+	return FIntPoint::ZeroValue;
 }
 
 void AGLTFHotspotActor::ViewportResized(FViewport*, uint32)
diff --git a/Engine/Plugins/Enterprise/GLTFExporter/MigratedFromGit/Source/GLTFMaterialAnalyzer/Public/GLTFMaterialAnalyzerModule.h b/Engine/Plugins/Enterprise/GLTFExporter/MigratedFromGit/Source/GLTFMaterialAnalyzer/Public/GLTFMaterialAnalyzerModule.h
--- a/Engine/Plugins/Enterprise/GLTFExporter/MigratedFromGit/Source/GLTFMaterialAnalyzer/Public/GLTFMaterialAnalyzerModule.h
+++ b/Engine/Plugins/Enterprise/GLTFExporter/MigratedFromGit/Source/GLTFMaterialAnalyzer/Public/GLTFMaterialAnalyzerModule.h
@@ -3,7 +3,6 @@
 #pragma once
 
 #include "CoreMinimal.h"
-#include "Modules/ModuleManager.h"
 
 class FGLTFMaterialAnalyzerModule : public IModuleInterface
 {
