[GLTFExporter] refactor: rename internal term Variation to LevelVariantSets

commit 7e8b5d444206ef25c1e763a872fad093c164d7cd
Date: Sun, 21 Feb 2021 23:22:02 +0100

    refactor: rename internal term Variation to LevelVariantSets
    (cherry picked from commit 101b0b90e2a9b26efaa6a8083bf86579b0942149​)

#skipci

diff --git a/Engine/Plugins/Enterprise/GLTFExporter/MigratedFromGit/Source/GLTFExporter/Private/Builders/GLTFConvertBuilder.cpp b/Engine/Plugins/Enterprise/GLTFExporter/MigratedFromGit/Source/GLTFExporter/Private/Builders/GLTFConvertBuilder.cpp
--- a/Engine/Plugins/Enterprise/GLTFExporter/MigratedFromGit/Source/GLTFExporter/Private/Builders/GLTFConvertBuilder.cpp
+++ b/Engine/Plugins/Enterprise/GLTFExporter/MigratedFromGit/Source/GLTFExporter/Private/Builders/GLTFConvertBuilder.cpp
@@ -377,14 +377,14 @@ FGLTFJsonBackdropIndex FGLTFConvertBuilder::GetOrAddBackdrop(const AActor* Backd
 	return BackdropConverter.GetOrAdd(BackdropActor);
 }
 
-FGLTFJsonVariationIndex FGLTFConvertBuilder::GetOrAddVariation(const ULevelVariantSets* LevelVariantSets)
+FGLTFJsonLevelVariantSetsIndex FGLTFConvertBuilder::GetOrAddLevelVariantSets(const ULevelVariantSets* LevelVariantSets)
 {
 	if (LevelVariantSets == nullptr)
 	{
-		return FGLTFJsonVariationIndex(INDEX_NONE);
+		return FGLTFJsonLevelVariantSetsIndex(INDEX_NONE);
 	}
 
-	return VariationConverter.GetOrAdd(LevelVariantSets);
+	return LevelVariantSetsConverter.GetOrAdd(LevelVariantSets);
 }
 
 FGLTFJsonLightMapIndex FGLTFConvertBuilder::GetOrAddLightMap(const UStaticMeshComponent* StaticMeshComponent)
diff --git a/Engine/Plugins/Enterprise/GLTFExporter/MigratedFromGit/Source/GLTFExporter/Private/Builders/GLTFConvertBuilder.h b/Engine/Plugins/Enterprise/GLTFExporter/MigratedFromGit/Source/GLTFExporter/Private/Builders/GLTFConvertBuilder.h
--- a/Engine/Plugins/Enterprise/GLTFExporter/MigratedFromGit/Source/GLTFExporter/Private/Builders/GLTFConvertBuilder.h
+++ b/Engine/Plugins/Enterprise/GLTFExporter/MigratedFromGit/Source/GLTFExporter/Private/Builders/GLTFConvertBuilder.h
@@ -16,7 +16,7 @@
 #include "Converters/GLTFCameraConverters.h"
 #include "Converters/GLTFLightConverters.h"
 #include "Converters/GLTFBackdropConverters.h"
-#include "Converters/GLTFVariationConverters.h"
+#include "Converters/GLTFLevelVariantSetsConverters.h"
 #include "Converters/GLTFLightMapConverters.h"
 #include "Converters/GLTFHotspotConverters.h"
 #include "Converters/GLTFSkySphereConverters.h"
@@ -76,7 +76,7 @@ public:
 	FGLTFJsonCameraIndex GetOrAddCamera(const UCameraComponent* CameraComponent);
 	FGLTFJsonLightIndex GetOrAddLight(const ULightComponent* LightComponent);
 	FGLTFJsonBackdropIndex GetOrAddBackdrop(const AActor* BackdropActor);
-	FGLTFJsonVariationIndex GetOrAddVariation(const ULevelVariantSets* LevelVariantSets);
+	FGLTFJsonLevelVariantSetsIndex GetOrAddLevelVariantSets(const ULevelVariantSets* LevelVariantSets);
 	FGLTFJsonLightMapIndex GetOrAddLightMap(const UStaticMeshComponent* StaticMeshComponent);
 	FGLTFJsonHotspotIndex GetOrAddHotspot(const AGLTFHotspotActor* HotspotActor);
 	FGLTFJsonSkySphereIndex GetOrAddSkySphere(const AActor* SkySphereActor);
@@ -120,7 +120,7 @@ private:
 	FGLTFCameraConverter CameraConverter{ *this };
 	FGLTFLightConverter LightConverter{ *this };
 	FGLTFBackdropConverter BackdropConverter{ *this };
-	FGLTFVariationConverter VariationConverter{ *this };
+	FGLTFLevelVariantSetsConverter LevelVariantSetsConverter{ *this };
 	FGLTFLightMapConverter LightMapConverter{ *this };
 	FGLTFHotspotConverter HotspotConverter{ *this };
 	FGLTFSkySphereConverter SkySphereConverter{ *this };
diff --git a/Engine/Plugins/Enterprise/GLTFExporter/MigratedFromGit/Source/GLTFExporter/Private/Builders/GLTFJsonBuilder.cpp b/Engine/Plugins/Enterprise/GLTFExporter/MigratedFromGit/Source/GLTFExporter/Private/Builders/GLTFJsonBuilder.cpp
--- a/Engine/Plugins/Enterprise/GLTFExporter/MigratedFromGit/Source/GLTFExporter/Private/Builders/GLTFJsonBuilder.cpp
+++ b/Engine/Plugins/Enterprise/GLTFExporter/MigratedFromGit/Source/GLTFExporter/Private/Builders/GLTFJsonBuilder.cpp
@@ -130,9 +130,9 @@ FGLTFJsonSkySphereIndex FGLTFJsonBuilder::AddSkySphere(const FGLTFJsonSkySphere&
 	return FGLTFJsonSkySphereIndex(JsonRoot.SkySpheres.Add(MakeUnique<FGLTFJsonSkySphere>(JsonSkySphere)));
 }
 
-FGLTFJsonVariationIndex FGLTFJsonBuilder::AddVariation(const FGLTFJsonVariation& JsonVariation)
+FGLTFJsonLevelVariantSetsIndex FGLTFJsonBuilder::AddLevelVariantSets(const FGLTFJsonLevelVariantSets& JsonLevelVariantSets)
 {
-	return FGLTFJsonVariationIndex(JsonRoot.Variations.Add(MakeUnique<FGLTFJsonVariation>(JsonVariation)));
+	return FGLTFJsonLevelVariantSetsIndex(JsonRoot.LevelVariantSets.Add(MakeUnique<FGLTFJsonLevelVariantSets>(JsonLevelVariantSets)));
 }
 
 FGLTFJsonNodeIndex FGLTFJsonBuilder::AddChildNode(FGLTFJsonNodeIndex ParentIndex, const FGLTFJsonNode& JsonNode)
@@ -249,9 +249,9 @@ FGLTFJsonSkySphere& FGLTFJsonBuilder::GetSkySphere(FGLTFJsonSkySphereIndex SkySp
 	return *JsonRoot.SkySpheres[SkySphereIndex];
 }
 
-FGLTFJsonVariation& FGLTFJsonBuilder::GetVariation(FGLTFJsonVariationIndex VariationIndex)
+FGLTFJsonLevelVariantSets& FGLTFJsonBuilder::GetLevelVariantSets(FGLTFJsonLevelVariantSetsIndex LevelVariantSetsIndex)
 {
-	return *JsonRoot.Variations[VariationIndex];
+	return *JsonRoot.LevelVariantSets[LevelVariantSetsIndex];
 }
 
 FGLTFJsonNodeIndex FGLTFJsonBuilder::GetComponentNodeIndex(FGLTFJsonNodeIndex NodeIndex)
diff --git a/Engine/Plugins/Enterprise/GLTFExporter/MigratedFromGit/Source/GLTFExporter/Private/Builders/GLTFJsonBuilder.h b/Engine/Plugins/Enterprise/GLTFExporter/MigratedFromGit/Source/GLTFExporter/Private/Builders/GLTFJsonBuilder.h
--- a/Engine/Plugins/Enterprise/GLTFExporter/MigratedFromGit/Source/GLTFExporter/Private/Builders/GLTFJsonBuilder.h
+++ b/Engine/Plugins/Enterprise/GLTFExporter/MigratedFromGit/Source/GLTFExporter/Private/Builders/GLTFJsonBuilder.h
@@ -39,7 +39,7 @@ public:
 	FGLTFJsonLightIndex AddLight(const FGLTFJsonLight& JsonLight = {});
 	FGLTFJsonLightMapIndex AddLightMap(const FGLTFJsonLightMap& JsonLightMap = {});
 	FGLTFJsonSkySphereIndex AddSkySphere(const FGLTFJsonSkySphere& JsonSkySphere = {});
-	FGLTFJsonVariationIndex AddVariation(const FGLTFJsonVariation& JsonVariation = {});
+	FGLTFJsonLevelVariantSetsIndex AddLevelVariantSets(const FGLTFJsonLevelVariantSets& JsonLevelVariantSets = {});
 
 	FGLTFJsonNodeIndex AddChildNode(FGLTFJsonNodeIndex ParentNodeIndex, const FGLTFJsonNode& JsonNode = {});
 	FGLTFJsonNodeIndex AddChildComponentNode(FGLTFJsonNodeIndex ParentNodeIndex, const FGLTFJsonNode& JsonNode = {});
@@ -62,7 +62,7 @@ public:
 	FGLTFJsonLight& GetLight(FGLTFJsonLightIndex LightIndex);
 	FGLTFJsonLightMap& GetLightMap(FGLTFJsonLightMapIndex LightMapIndex);
 	FGLTFJsonSkySphere& GetSkySphere(FGLTFJsonSkySphereIndex SkySphereIndex);
-	FGLTFJsonVariation& GetVariation(FGLTFJsonVariationIndex VariationIndex);
+	FGLTFJsonLevelVariantSets& GetLevelVariantSets(FGLTFJsonLevelVariantSetsIndex LevelVariantSetsIndex);
 
 	FGLTFJsonNodeIndex GetComponentNodeIndex(FGLTFJsonNodeIndex NodeIndex);
 
diff --git a/Engine/Plugins/Enterprise/GLTFExporter/MigratedFromGit/Source/GLTFExporter/Private/Converters/GLTFLevelVariantSetsConverters.cpp b/Engine/Plugins/Enterprise/GLTFExporter/MigratedFromGit/Source/GLTFExporter/Private/Converters/GLTFLevelVariantSetsConverters.cpp
new file mode 100644
--- /dev/null
+++ b/Engine/Plugins/Enterprise/GLTFExporter/MigratedFromGit/Source/GLTFExporter/Private/Converters/GLTFLevelVariantSetsConverters.cpp
@@ -0,0 +1,439 @@
+// Copyright Epic Games, Inc. All Rights Reserved.
+
+#include "Converters/GLTFLevelVariantSetsConverters.h"
+#include "Converters/GLTFMeshUtility.h"
+#include "Builders/GLTFContainerBuilder.h"
+#include "Rendering/SkeletalMeshRenderData.h"
+#include "VariantObjectBinding.h"
+#include "PropertyValueMaterial.h"
+#include "LevelVariantSets.h"
+#include "PropertyValue.h"
+#include "VariantSet.h"
+#include "Variant.h"
+
+namespace
+{
+	// TODO: replace hack with safe solution by adding proper API access to UPropertyValue
+	class UPropertyValueHack : public UPropertyValue
+	{
+	public:
+		const TArray<FCapturedPropSegment>& GetCapturedPropSegments() const
+		{
+			return CapturedPropSegments;
+		}
+	};
+}
+
+FGLTFJsonLevelVariantSetsIndex FGLTFLevelVariantSetsConverter::Convert(const ULevelVariantSets* LevelVariantSets)
+{
+	FGLTFJsonLevelVariantSets JsonLevelVariantSets;
+	LevelVariantSets->GetName(JsonLevelVariantSets.Name);
+
+	for (const UVariantSet* VariantSet: LevelVariantSets->GetVariantSets())
+	{
+		FGLTFJsonVariantSet JsonVariantSet;
+		JsonVariantSet.Name = VariantSet->GetDisplayText().ToString();
+
+		for (const UVariant* Variant: VariantSet->GetVariants())
+		{
+			FGLTFJsonVariant JsonVariant;
+			if (TryParseVariant(JsonVariant, Variant))
+			{
+				JsonVariantSet.Variants.Add(JsonVariant);
+			}
+		}
+
+		if (JsonVariantSet.Variants.Num() > 0)
+		{
+			JsonLevelVariantSets.VariantSets.Add(JsonVariantSet);
+		}
+	}
+
+	if (JsonLevelVariantSets.VariantSets.Num() == 0)
+	{
+		return FGLTFJsonLevelVariantSetsIndex(INDEX_NONE);
+	}
+
+	return Builder.AddLevelVariantSets(JsonLevelVariantSets);
+}
+
+bool FGLTFLevelVariantSetsConverter::TryParseVariant(FGLTFJsonVariant& OutVariant, const UVariant* Variant) const
+{
+	FGLTFJsonVariant JsonVariant;
+
+	for (const UVariantObjectBinding* Binding: Variant->GetBindings())
+	{
+		TryParseVariantBinding(JsonVariant, Binding);
+	}
+
+	if (JsonVariant.Nodes.Num() == 0)
+	{
+		return false;
+	}
+
+	JsonVariant.Name = Variant->GetDisplayText().ToString();
+	JsonVariant.bIsActive = const_cast<UVariant*>(Variant)->IsActive();
+
+	if (const UTexture2D* Thumbnail = const_cast<UVariant*>(Variant)->GetThumbnail())
+	{
+		// TODO: if thumbnail has generic name "Texture2D", give it a variant-relevant name
+		JsonVariant.Thumbnail = Builder.GetOrAddTexture(Thumbnail);
+	}
+
+	OutVariant = JsonVariant;
+	return true;
+}
+
+bool FGLTFLevelVariantSetsConverter::TryParseVariantBinding(FGLTFJsonVariant& OutVariant, const UVariantObjectBinding* Binding) const
+{
+	bool bHasParsedAnyProperty = false;
+
+	for (UPropertyValue* Property: Binding->GetCapturedProperties())
+	{
+		if (!const_cast<UPropertyValue*>(Property)->Resolve() || !Property->HasRecordedData())
+		{
+			continue;
+		}
+
+		const FName PropertyName = Property->GetPropertyName();
+		const FFieldClass* PropertyClass = Property->GetPropertyClass();
+
+		if (Property->IsA<UPropertyValueMaterial>())
+		{
+			if (Builder.ExportOptions->ExportMaterialVariants != EGLTFMaterialVariantMode::None && TryParseMaterialPropertyValue(OutVariant, Property))
+			{
+				bHasParsedAnyProperty = true;
+			}
+		}
+		else if (PropertyName == TEXT("StaticMesh")) // TODO: should we not also check PropertyClass?
+		{
+			if (Builder.ExportOptions->bExportMeshVariants && TryParseStaticMeshPropertyValue(OutVariant, Property))
+			{
+				bHasParsedAnyProperty = true;
+			}
+		}
+		else if (PropertyName == TEXT("SkeletalMesh")) // TODO: should we not also check PropertyClass?
+		{
+			if (Builder.ExportOptions->bExportMeshVariants && TryParseSkeletalMeshPropertyValue(OutVariant, Property))
+			{
+				bHasParsedAnyProperty = true;
+			}
+		}
+		else if (PropertyName == TEXT("bVisible") && PropertyClass->IsChildOf(FBoolProperty::StaticClass()))
+		{
+			if (Builder.ExportOptions->bExportVisibilityVariants && TryParseVisibilityPropertyValue(OutVariant, Property))
+			{
+				bHasParsedAnyProperty = true;
+			}
+		}
+
+		// TODO: add support for more properties
+	}
+
+	return bHasParsedAnyProperty;
+}
+
+bool FGLTFLevelVariantSetsConverter::TryParseVisibilityPropertyValue(FGLTFJsonVariant& OutVariant, const UPropertyValue* Property) const
+{
+	const USceneComponent* Target = static_cast<USceneComponent*>(Property->GetPropertyParentContainerAddress());
+	if (Target == nullptr)
+	{
+		Builder.AddWarningMessage(FString::Printf(
+			TEXT("Target object for property is invalid, the property will be skipped. Context: %s"),
+			*GetLogContext(Property)));
+		return false;
+	}
+
+	const AActor* Owner = Target->GetOwner();
+	if (Owner == nullptr)
+	{
+		Builder.AddWarningMessage(FString::Printf(TEXT("Invalid scene component, the property will be skipped. Context: %s"),
+			*GetLogContext(Property)));
+		return false;
+	}
+
+	if (Builder.bSelectedActorsOnly && !Owner->IsSelected())
+	{
+		Builder.AddWarningMessage(FString::Printf(
+			TEXT("Target actor for property is not selected for export, the property will be skipped. Context: %s"),
+			*GetLogContext(Property)));
+		return false;
+	}
+
+	bool bIsVisible;
+
+	if (!TryGetPropertyValue(const_cast<UPropertyValue*>(Property), bIsVisible))
+	{
+		Builder.AddWarningMessage(FString::Printf(
+			TEXT("Failed to parse recorded data for property, it will be skipped. Context: %s"),
+			*GetLogContext(Property)));
+		return false;
+	}
+
+	const FGLTFJsonNodeIndex NodeIndex = Builder.GetOrAddNode(Target);
+	const FGLTFJsonNodeIndex ComponentNodeIndex = Builder.GetComponentNodeIndex(NodeIndex);
+
+	FGLTFJsonVariantNodeProperties& NodeProperties = OutVariant.Nodes.FindOrAdd(ComponentNodeIndex);
+
+	NodeProperties.Node = ComponentNodeIndex;
+	NodeProperties.bIsVisible = bIsVisible;
+	return true;
+}
+
+bool FGLTFLevelVariantSetsConverter::TryParseMaterialPropertyValue(FGLTFJsonVariant& OutVariant, const UPropertyValue* Property) const
+{
+	UPropertyValueMaterial* MaterialProperty = Cast<UPropertyValueMaterial>(const_cast<UPropertyValue*>(Property));
+	if (MaterialProperty == nullptr)
+	{
+		Builder.AddWarningMessage(FString::Printf(
+			TEXT("Material property is invalid, the property will be skipped. Context: %s"),
+			*GetLogContext(Property)));
+		return false;
+	}
+
+	const USceneComponent* Target = static_cast<USceneComponent*>(Property->GetPropertyParentContainerAddress());
+	if (Target == nullptr)
+	{
+		Builder.AddWarningMessage(FString::Printf(
+			TEXT("Target object for property is invalid, the property will be skipped. Context: %s"),
+			*GetLogContext(Property)));
+		return false;
+	}
+
+	const AActor* Owner = Target->GetOwner();
+	if (Owner == nullptr)
+	{
+		Builder.AddWarningMessage(FString::Printf(TEXT("Invalid scene component, the property will be skipped. Context: %s"),
+			*GetLogContext(Property)));
+		return false;
+	}
+
+	if (Builder.bSelectedActorsOnly && !Owner->IsSelected())
+	{
+		Builder.AddWarningMessage(FString::Printf(
+			TEXT("Target actor for property is not selected for export, the property will be skipped. Context: %s"),
+			*GetLogContext(Property)));
+		return false;
+	}
+
+	const TArray<FCapturedPropSegment>& CapturedPropSegments = reinterpret_cast<UPropertyValueHack*>(MaterialProperty)->GetCapturedPropSegments();
+	const int32 NumPropSegments = CapturedPropSegments.Num();
+
+	if (NumPropSegments < 1)
+	{
+		Builder.AddWarningMessage(FString::Printf(
+			TEXT("Failed to parse element index to apply the material to, the property will be skipped. Context: %s"),
+			*GetLogContext(Property)));
+		return false;
+	}
+
+	// NOTE: UPropertyValueMaterial::GetMaterial does *not* ensure that the recorded data has been loaded,
+	// so we need to call UProperty::GetRecordedData first to make that happen.
+	MaterialProperty->GetRecordedData();
+
+	// TODO: find way to determine whether the material is null because "None" was selected, or because it failed to resolve
+	const UMaterialInterface* Material = MaterialProperty->GetMaterial();
+	const int32 MaterialIndex = CapturedPropSegments[NumPropSegments - 1].PropertyIndex;
+
+	const FGLTFMeshData* MeshData = nullptr;
+	FGLTFIndexArray SectionIndices;
+
+	if (Builder.ExportOptions->ExportMaterialVariants == EGLTFMaterialVariantMode::UseMeshData)
+	{
+		if (Builder.ExportOptions->BakeMaterialInputs == EGLTFMaterialBakeMode::UseMeshData)
+		{
+			const int32 DefaultLODIndex = Builder.ExportOptions->DefaultLevelOfDetail;
+
+			if (const UStaticMeshComponent* StaticMeshComponent = Cast<UStaticMeshComponent>(Target))
+			{
+				const UStaticMesh* StaticMesh = StaticMeshComponent->GetStaticMesh();
+				const int32 LODIndex = FGLTFMeshUtility::GetLOD(StaticMesh, StaticMeshComponent, DefaultLODIndex);
+				const FStaticMeshLODResources& MeshLOD = StaticMesh->GetLODForExport(LODIndex);
+
+				SectionIndices = FGLTFMeshUtility::GetSectionIndices(MeshLOD, MaterialIndex);
+				MeshData = Builder.StaticMeshDataConverter.GetOrAdd(StaticMesh, StaticMeshComponent, LODIndex);
+			}
+			else if (const USkeletalMeshComponent* SkeletalMeshComponent = Cast<USkeletalMeshComponent>(Target))
+			{
+				const USkeletalMesh* SkeletalMesh = SkeletalMeshComponent->SkeletalMesh;
+				const int32 LODIndex = FGLTFMeshUtility::GetLOD(SkeletalMesh, SkeletalMeshComponent, DefaultLODIndex);
+				const FSkeletalMeshRenderData* RenderData = SkeletalMesh->GetResourceForRendering();
+				const FSkeletalMeshLODRenderData& MeshLOD = RenderData->LODRenderData[LODIndex];
+
+				SectionIndices = FGLTFMeshUtility::GetSectionIndices(MeshLOD, MaterialIndex);
+				MeshData = Builder.SkeletalMeshDataConverter.GetOrAdd(SkeletalMesh, SkeletalMeshComponent, LODIndex);
+			}
+		}
+		else
+		{
+			// TODO: report warning (about materials won't be export using mesh data because BakeMaterialInputs is not set to UseMeshData)
+		}
+	}
+
+	FGLTFJsonVariantMaterial VariantMaterial;
+	VariantMaterial.Material = Builder.GetOrAddMaterial(Material, MeshData, SectionIndices);
+	VariantMaterial.Index = MaterialIndex;
+
+	const FGLTFJsonNodeIndex NodeIndex = Builder.GetOrAddNode(Target);
+	const FGLTFJsonNodeIndex ComponentNodeIndex = Builder.GetComponentNodeIndex(NodeIndex);
+	FGLTFJsonVariantNodeProperties& NodeProperties = OutVariant.Nodes.FindOrAdd(ComponentNodeIndex);
+
+	NodeProperties.Node = ComponentNodeIndex;
+	NodeProperties.Materials.Add(VariantMaterial);
+	return true;
+}
+
+bool FGLTFLevelVariantSetsConverter::TryParseStaticMeshPropertyValue(FGLTFJsonVariant& OutVariant, const UPropertyValue* Property) const
+{
+	const USceneComponent* Target = static_cast<USceneComponent*>(Property->GetPropertyParentContainerAddress());
+	if (Target == nullptr)
+	{
+		Builder.AddWarningMessage(FString::Printf(
+			TEXT("Target object for property is invalid, the property will be skipped. Context: %s"),
+			*GetLogContext(Property)));
+		return false;
+	}
+
+	const AActor* Owner = Target->GetOwner();
+	if (Owner == nullptr)
+	{
+		Builder.AddWarningMessage(FString::Printf(TEXT("Invalid scene component, the property will be skipped. Context: %s"),
+			*GetLogContext(Property)));
+		return false;
+	}
+
+	if (Builder.bSelectedActorsOnly && !Owner->IsSelected())
+	{
+		Builder.AddWarningMessage(FString::Printf(
+			TEXT("Target actor for property is not selected for export, the property will be skipped. Context: %s"),
+			*GetLogContext(Property)));
+		return false;
+	}
+
+	const UMeshComponent* MeshComponent = Cast<UMeshComponent>(Target);
+	if (MeshComponent == nullptr)
+	{
+		Builder.AddWarningMessage(FString::Printf(
+			TEXT("Target object for property has no mesh-component, the property will be skipped. Context: %s"),
+			*GetLogContext(Property)));
+		return false;
+	}
+
+	const UStaticMesh* StaticMesh;
+
+	if (!TryGetPropertyValue(const_cast<UPropertyValue*>(Property), StaticMesh))
+	{
+		Builder.AddWarningMessage(FString::Printf(
+			TEXT("Failed to parse recorded data for property, it will be skipped. Context: %s"),
+			*GetLogContext(Property)));
+		return false;
+	}
+
+	const FGLTFMaterialArray OverrideMaterials(MeshComponent->OverrideMaterials);
+	const FGLTFJsonNodeIndex NodeIndex = Builder.GetOrAddNode(Target);
+	const FGLTFJsonNodeIndex ComponentNodeIndex = Builder.GetComponentNodeIndex(NodeIndex);
+	const FGLTFJsonMeshIndex MeshIndex = Builder.GetOrAddMesh(StaticMesh, OverrideMaterials);
+	FGLTFJsonVariantNodeProperties& NodeProperties = OutVariant.Nodes.FindOrAdd(ComponentNodeIndex);
+
+	NodeProperties.Node = ComponentNodeIndex;
+	NodeProperties.Mesh = MeshIndex;
+	return true;
+}
+
+bool FGLTFLevelVariantSetsConverter::TryParseSkeletalMeshPropertyValue(FGLTFJsonVariant& OutVariant, const UPropertyValue* Property) const
+{
+	const USceneComponent* Target = static_cast<USceneComponent*>(Property->GetPropertyParentContainerAddress());
+	if (Target == nullptr)
+	{
+		Builder.AddWarningMessage(FString::Printf(
+			TEXT("Target object for property is invalid, the property will be skipped. Context: %s"),
+			*GetLogContext(Property)));
+		return false;
+	}
+
+	const AActor* Owner = Target->GetOwner();
+	if (Owner == nullptr)
+	{
+		Builder.AddWarningMessage(FString::Printf(TEXT("Invalid scene component, the property will be skipped. Context: %s"),
+			*GetLogContext(Property)));
+		return false;
+	}
+
+	if (Builder.bSelectedActorsOnly && !Owner->IsSelected())
+	{
+		Builder.AddWarningMessage(FString::Printf(
+			TEXT("Target actor for property is not selected for export, the property will be skipped. Context: %s"),
+			*GetLogContext(Property)));
+		return false;
+	}
+
+	const UMeshComponent* MeshComponent = Cast<UMeshComponent>(Target);
+	if (MeshComponent == nullptr)
+	{
+		Builder.AddWarningMessage(FString::Printf(
+			TEXT("Target object for property has no mesh-component, the property will be skipped. Context: %s"),
+			*GetLogContext(Property)));
+		return false;
+	}
+
+	const USkeletalMesh* SkeletalMesh;
+
+	if (!TryGetPropertyValue(const_cast<UPropertyValue*>(Property), SkeletalMesh) || SkeletalMesh == nullptr)
+	{
+		Builder.AddWarningMessage(FString::Printf(
+			TEXT("Failed to parse recorded data for property, it will be skipped. Context: %s"),
+			*GetLogContext(Property)));
+		return false;
+	}
+
+	const FGLTFMaterialArray OverrideMaterials(MeshComponent->OverrideMaterials);
+	const FGLTFJsonNodeIndex NodeIndex = Builder.GetOrAddNode(Target);
+	const FGLTFJsonMeshIndex MeshIndex = Builder.GetOrAddMesh(SkeletalMesh, OverrideMaterials);
+	const FGLTFJsonNodeIndex ComponentNodeIndex = Builder.GetComponentNodeIndex(NodeIndex);
+	FGLTFJsonVariantNodeProperties& NodeProperties = OutVariant.Nodes.FindOrAdd(ComponentNodeIndex);
+
+	NodeProperties.Node = ComponentNodeIndex;
+	NodeProperties.Mesh = MeshIndex;
+	return true;
+}
+
+template<typename T>
+bool FGLTFLevelVariantSetsConverter::TryGetPropertyValue(UPropertyValue* Property, T& OutValue) const
+{
+	if (Property == nullptr || !Property->HasRecordedData())
+	{
+		return false;
+	}
+
+	FMemory::Memcpy(&OutValue, Property->GetRecordedData().GetData(), sizeof(T));
+	return true;
+}
+
+FString FGLTFLevelVariantSetsConverter::GetLogContext(const UPropertyValue* Property) const
+{
+	const UVariantObjectBinding* Parent = Property->GetParent();
+	return GetLogContext(Parent) + TEXT("/") + Property->GetFullDisplayString();
+}
+
+FString FGLTFLevelVariantSetsConverter::GetLogContext(const UVariantObjectBinding* Binding) const
+{
+	const UVariant* Parent = const_cast<UVariantObjectBinding*>(Binding)->GetParent();
+	return GetLogContext(Parent) + TEXT("/") + Binding->GetDisplayText().ToString();
+}
+
+FString FGLTFLevelVariantSetsConverter::GetLogContext(const UVariant* Variant) const
+{
+	const UVariantSet* Parent = const_cast<UVariant*>(Variant)->GetParent();
+	return GetLogContext(Parent) + TEXT("/") + Variant->GetDisplayText().ToString();
+}
+
+FString FGLTFLevelVariantSetsConverter::GetLogContext(const UVariantSet* VariantSet) const
+{
+	const ULevelVariantSets* Parent = const_cast<UVariantSet*>(VariantSet)->GetParent();
+	return GetLogContext(Parent) + TEXT("/") + VariantSet->GetDisplayText().ToString();
+}
+
+FString FGLTFLevelVariantSetsConverter::GetLogContext(const ULevelVariantSets* LevelVariantSets) const
+{
+	return LevelVariantSets->GetName();
+}
diff --git a/Engine/Plugins/Enterprise/GLTFExporter/MigratedFromGit/Source/GLTFExporter/Private/Converters/GLTFLevelVariantSetsConverters.h b/Engine/Plugins/Enterprise/GLTFExporter/MigratedFromGit/Source/GLTFExporter/Private/Converters/GLTFLevelVariantSetsConverters.h
new file mode 100644
--- /dev/null
+++ b/Engine/Plugins/Enterprise/GLTFExporter/MigratedFromGit/Source/GLTFExporter/Private/Converters/GLTFLevelVariantSetsConverters.h
@@ -0,0 +1,35 @@
+// Copyright Epic Games, Inc. All Rights Reserved.
+
+#pragma once
+
+#include "Json/GLTFJsonIndex.h"
+#include "Json/GLTFJsonLevelVariantSets.h"
+#include "Converters/GLTFConverter.h"
+#include "Converters/GLTFBuilderContext.h"
+#include "Engine.h"
+#include "Variant.h"
+#include "PropertyValue.h"
+#include "LevelVariantSets.h"
+
+class FGLTFLevelVariantSetsConverter final : public FGLTFBuilderContext, public TGLTFConverter<FGLTFJsonLevelVariantSetsIndex, const ULevelVariantSets*>
+{
+	using FGLTFBuilderContext::FGLTFBuilderContext;
+
+	virtual FGLTFJsonLevelVariantSetsIndex Convert(const ULevelVariantSets* LevelVariantSets) override;
+
+	bool TryParseVariant(FGLTFJsonVariant& OutVariant, const UVariant* Variant) const;
+	bool TryParseVariantBinding(FGLTFJsonVariant& OutVariant, const UVariantObjectBinding* Binding) const;
+	bool TryParseVisibilityPropertyValue(FGLTFJsonVariant& OutVariant, const UPropertyValue* Property) const;
+	bool TryParseMaterialPropertyValue(FGLTFJsonVariant& OutVariant, const UPropertyValue* Property) const;
+	bool TryParseStaticMeshPropertyValue(FGLTFJsonVariant& OutVariant, const UPropertyValue* Property) const;
+	bool TryParseSkeletalMeshPropertyValue(FGLTFJsonVariant& OutVariant, const UPropertyValue* Property) const;
+
+	template<typename T>
+	bool TryGetPropertyValue(UPropertyValue* Property, T& OutValue) const;
+
+	FString GetLogContext(const UPropertyValue* Property) const;
+	FString GetLogContext(const UVariantObjectBinding* Binding) const;
+	FString GetLogContext(const UVariant* Variant) const;
+	FString GetLogContext(const UVariantSet* VariantSet) const;
+	FString GetLogContext(const ULevelVariantSets* LevelVariantSets) const;
+};
diff --git a/Engine/Plugins/Enterprise/GLTFExporter/MigratedFromGit/Source/GLTFExporter/Private/Converters/GLTFSceneConverters.cpp b/Engine/Plugins/Enterprise/GLTFExporter/MigratedFromGit/Source/GLTFExporter/Private/Converters/GLTFSceneConverters.cpp
--- a/Engine/Plugins/Enterprise/GLTFExporter/MigratedFromGit/Source/GLTFExporter/Private/Converters/GLTFSceneConverters.cpp
+++ b/Engine/Plugins/Enterprise/GLTFExporter/MigratedFromGit/Source/GLTFExporter/Private/Converters/GLTFSceneConverters.cpp
@@ -30,10 +30,10 @@ FGLTFJsonSceneIndex FGLTFSceneConverter::Convert(const UWorld* World)
 				{
 					if (const ULevelVariantSets* LevelVariantSets = const_cast<ALevelVariantSetsActor*>(LevelVariantSetsActor)->GetLevelVariantSets(true))
 					{
-						const FGLTFJsonVariationIndex VariationIndex = Builder.GetOrAddVariation(LevelVariantSets);
-						if (VariationIndex != INDEX_NONE)
+						const FGLTFJsonLevelVariantSetsIndex LevelVariantSetsIndex = Builder.GetOrAddLevelVariantSets(LevelVariantSets);
+						if (LevelVariantSetsIndex != INDEX_NONE)
 						{
-							Scene.Variations.Add(VariationIndex);
+							Scene.LevelVariantSets.Add(LevelVariantSetsIndex);
 						}
 					}
 				}
diff --git a/Engine/Plugins/Enterprise/GLTFExporter/MigratedFromGit/Source/GLTFExporter/Private/Converters/GLTFVariationConverters.cpp b/Engine/Plugins/Enterprise/GLTFExporter/MigratedFromGit/Source/GLTFExporter/Private/Converters/GLTFVariationConverters.cpp
deleted file mode 100644
diff --git a/Engine/Plugins/Enterprise/GLTFExporter/MigratedFromGit/Source/GLTFExporter/Private/Converters/GLTFVariationConverters.h b/Engine/Plugins/Enterprise/GLTFExporter/MigratedFromGit/Source/GLTFExporter/Private/Converters/GLTFVariationConverters.h
deleted file mode 100644
diff --git a/Engine/Plugins/Enterprise/GLTFExporter/MigratedFromGit/Source/GLTFExporter/Private/Exporters/GLTFLevelVariantSetsExporter.cpp b/Engine/Plugins/Enterprise/GLTFExporter/MigratedFromGit/Source/GLTFExporter/Private/Exporters/GLTFLevelVariantSetsExporter.cpp
--- a/Engine/Plugins/Enterprise/GLTFExporter/MigratedFromGit/Source/GLTFExporter/Private/Exporters/GLTFLevelVariantSetsExporter.cpp
+++ b/Engine/Plugins/Enterprise/GLTFExporter/MigratedFromGit/Source/GLTFExporter/Private/Exporters/GLTFLevelVariantSetsExporter.cpp
@@ -51,8 +51,8 @@ bool UGLTFLevelVariantSetsExporter::AddObject(FGLTFContainerBuilder& Builder, co
 		return false;
 	}
 
-	const FGLTFJsonVariationIndex VariationIndex = Builder.GetOrAddVariation(LevelVariantSets);
-	if (VariationIndex == INDEX_NONE)
+	const FGLTFJsonLevelVariantSetsIndex LevelVariantSetsIndex = Builder.GetOrAddLevelVariantSets(LevelVariantSets);
+	if (LevelVariantSetsIndex == INDEX_NONE)
 	{
 		Builder.AddErrorMessage(
 			FString::Printf(TEXT("Failed to export level variant sets %s"),
@@ -60,7 +60,7 @@ bool UGLTFLevelVariantSetsExporter::AddObject(FGLTFContainerBuilder& Builder, co
 		return false;
 	}
 
-	Builder.GetScene(SceneIndex).Variations.AddUnique(VariationIndex);
+	Builder.GetScene(SceneIndex).LevelVariantSets.AddUnique(LevelVariantSetsIndex);
 
 	Builder.DefaultScene = SceneIndex;
 	return true;
diff --git a/Engine/Plugins/Enterprise/GLTFExporter/MigratedFromGit/Source/GLTFExporter/Private/Json/GLTFJsonIndex.h b/Engine/Plugins/Enterprise/GLTFExporter/MigratedFromGit/Source/GLTFExporter/Private/Json/GLTFJsonIndex.h
--- a/Engine/Plugins/Enterprise/GLTFExporter/MigratedFromGit/Source/GLTFExporter/Private/Json/GLTFJsonIndex.h
+++ b/Engine/Plugins/Enterprise/GLTFExporter/MigratedFromGit/Source/GLTFExporter/Private/Json/GLTFJsonIndex.h
@@ -71,4 +71,4 @@ struct FGLTFJsonHotspotIndex : FGLTFJsonIndex<FGLTFJsonHotspotIndex> { using FGL
 struct FGLTFJsonLightIndex : FGLTFJsonIndex<FGLTFJsonLightIndex> { using FGLTFJsonIndex::FGLTFJsonIndex; };
 struct FGLTFJsonLightMapIndex : FGLTFJsonIndex<FGLTFJsonLightMapIndex> { using FGLTFJsonIndex::FGLTFJsonIndex; };
 struct FGLTFJsonSkySphereIndex : FGLTFJsonIndex<FGLTFJsonSkySphereIndex> { using FGLTFJsonIndex::FGLTFJsonIndex; };
-struct FGLTFJsonVariationIndex : FGLTFJsonIndex<FGLTFJsonVariationIndex> { using FGLTFJsonIndex::FGLTFJsonIndex; };
+struct FGLTFJsonLevelVariantSetsIndex : FGLTFJsonIndex<FGLTFJsonLevelVariantSetsIndex> { using FGLTFJsonIndex::FGLTFJsonIndex; };
diff --git a/Engine/Plugins/Enterprise/GLTFExporter/MigratedFromGit/Source/GLTFExporter/Private/Json/GLTFJsonLevelVariantSets.h b/Engine/Plugins/Enterprise/GLTFExporter/MigratedFromGit/Source/GLTFExporter/Private/Json/GLTFJsonLevelVariantSets.h
new file mode 100644
--- /dev/null
+++ b/Engine/Plugins/Enterprise/GLTFExporter/MigratedFromGit/Source/GLTFExporter/Private/Json/GLTFJsonLevelVariantSets.h
@@ -0,0 +1,154 @@
+// Copyright Epic Games, Inc. All Rights Reserved.
+
+#pragma once
+
+#include "GLTFJsonIndex.h"
+#include "Json/GLTFJsonExtensions.h"
+#include "Json/GLTFJsonUtility.h"
+#include "Serialization/JsonSerializer.h"
+#include "Misc/Optional.h"
+
+struct FGLTFJsonVariantMaterial
+{
+	FGLTFJsonMaterialIndex Material;
+	int32                  Index;
+
+	FGLTFJsonVariantMaterial()
+		: Material(INDEX_NONE)
+		, Index(INDEX_NONE)
+	{
+	}
+
+	template <class CharType = TCHAR, class PrintPolicy = TPrettyJsonPrintPolicy<CharType>>
+	void WriteObject(TJsonWriter<CharType, PrintPolicy>& JsonWriter, FGLTFJsonExtensions& Extensions) const
+	{
+		JsonWriter.WriteObjectStart();
+
+		JsonWriter.WriteValue(TEXT("material"), Material);
+
+		if (Index != INDEX_NONE)
+		{
+			JsonWriter.WriteValue(TEXT("index"), Index);
+		}
+
+		JsonWriter.WriteObjectEnd();
+	}
+};
+
+struct FGLTFJsonVariantNodeProperties
+{
+	FGLTFJsonNodeIndex Node;
+	TOptional<bool>    bIsVisible;
+
+	TOptional<FGLTFJsonMeshIndex>    Mesh;
+	TArray<FGLTFJsonVariantMaterial> Materials;
+
+	template <class CharType = TCHAR, class PrintPolicy = TPrettyJsonPrintPolicy<CharType>>
+	void WriteObject(TJsonWriter<CharType, PrintPolicy>& JsonWriter, FGLTFJsonExtensions& Extensions) const
+	{
+		JsonWriter.WriteObjectStart();
+
+		if (Node != INDEX_NONE)
+		{
+			JsonWriter.WriteValue(TEXT("node"), Node);
+		}
+
+		JsonWriter.WriteObjectStart(TEXT("properties"));
+
+		if (bIsVisible.IsSet())
+		{
+			JsonWriter.WriteValue(TEXT("visible"), bIsVisible.GetValue());
+		}
+
+		if (Mesh.IsSet())
+		{
+			JsonWriter.WriteValue(TEXT("mesh"), Mesh.GetValue());
+		}
+
+		FGLTFJsonUtility::WriteObjectArray(JsonWriter, TEXT("materials"), Materials, Extensions);
+
+		JsonWriter.WriteObjectEnd();
+		JsonWriter.WriteObjectEnd();
+	}
+};
+
+struct FGLTFJsonVariant
+{
+	typedef TMap<FGLTFJsonNodeIndex, FGLTFJsonVariantNodeProperties> FNodeProperties;
+
+	FString Name;
+	bool    bIsActive;
+
+	FGLTFJsonTextureIndex Thumbnail;
+	FNodeProperties       Nodes;
+
+	template <class CharType = TCHAR, class PrintPolicy = TPrettyJsonPrintPolicy<CharType>>
+	void WriteObject(TJsonWriter<CharType, PrintPolicy>& JsonWriter, FGLTFJsonExtensions& Extensions) const
+	{
+		JsonWriter.WriteObjectStart();
+		JsonWriter.WriteValue(TEXT("name"), Name);
+		JsonWriter.WriteValue(TEXT("active"), bIsActive);
+
+		if (Thumbnail != INDEX_NONE)
+		{
+			JsonWriter.WriteValue(TEXT("thumbnail"), Thumbnail);
+		}
+
+		JsonWriter.WriteArrayStart(TEXT("nodes"));
+
+		for (const auto& DataPair: Nodes)
+		{
+			if (DataPair.Key != INDEX_NONE)
+			{
+				DataPair.Value.WriteObject(JsonWriter, Extensions);
+			}
+		}
+
+		JsonWriter.WriteArrayEnd();
+		JsonWriter.WriteObjectEnd();
+	}
+};
+
+struct FGLTFJsonVariantSet
+{
+	FString Name;
+
+	TArray<FGLTFJsonVariant> Variants;
+
+	template <class CharType = TCHAR, class PrintPolicy = TPrettyJsonPrintPolicy<CharType>>
+	void WriteObject(TJsonWriter<CharType, PrintPolicy>& JsonWriter, FGLTFJsonExtensions& Extensions) const
+	{
+		JsonWriter.WriteObjectStart();
+
+		if (!Name.IsEmpty())
+		{
+			JsonWriter.WriteValue(TEXT("name"), Name);
+		}
+
+		FGLTFJsonUtility::WriteObjectArray(JsonWriter, TEXT("variants"), Variants, Extensions);
+
+		JsonWriter.WriteObjectEnd();
+	}
+};
+
+struct FGLTFJsonLevelVariantSets
+{
+	FString Name;
+
+	TArray<FGLTFJsonVariantSet> VariantSets;
+
+	template <class CharType = TCHAR, class PrintPolicy = TPrettyJsonPrintPolicy<CharType>>
+	void WriteObject(TJsonWriter<CharType, PrintPolicy>& JsonWriter, FGLTFJsonExtensions& Extensions) const
+	{
+		JsonWriter.WriteObjectStart();
+
+		if (!Name.IsEmpty())
+		{
+			JsonWriter.WriteValue(TEXT("name"), Name);
+		}
+
+		FGLTFJsonUtility::WriteObjectArray(JsonWriter, TEXT("variantSets"), VariantSets, Extensions);
+
+		JsonWriter.WriteObjectEnd();
+	}
+};
diff --git a/Engine/Plugins/Enterprise/GLTFExporter/MigratedFromGit/Source/GLTFExporter/Private/Json/GLTFJsonRoot.h b/Engine/Plugins/Enterprise/GLTFExporter/MigratedFromGit/Source/GLTFExporter/Private/Json/GLTFJsonRoot.h
--- a/Engine/Plugins/Enterprise/GLTFExporter/MigratedFromGit/Source/GLTFExporter/Private/Json/GLTFJsonRoot.h
+++ b/Engine/Plugins/Enterprise/GLTFExporter/MigratedFromGit/Source/GLTFExporter/Private/Json/GLTFJsonRoot.h
@@ -22,7 +22,7 @@
 #include "Json/GLTFJsonLight.h"
 #include "Json/GLTFJsonLightMap.h"
 #include "Json/GLTFJsonSkySphere.h"
-#include "Json/GLTFJsonVariation.h"
+#include "Json/GLTFJsonLevelVariantSets.h"
 #include "Policies/CondensedJsonPrintPolicy.h"
 
 
@@ -81,7 +81,7 @@ struct FGLTFJsonRoot
 	TArray<TUniquePtr<FGLTFJsonLight>>      Lights;
 	TArray<TUniquePtr<FGLTFJsonLightMap>>   LightMaps;
 	TArray<TUniquePtr<FGLTFJsonSkySphere>>  SkySpheres;
-	TArray<TUniquePtr<FGLTFJsonVariation>>  Variations;
+	TArray<TUniquePtr<FGLTFJsonLevelVariantSets>>  LevelVariantSets;
 
 	template <class CharType = TCHAR, class PrintPolicy = TPrettyJsonPrintPolicy<CharType>>
 	void WriteObject(TJsonWriter<CharType, PrintPolicy>& JsonWriter)
@@ -110,7 +110,7 @@ struct FGLTFJsonRoot
 		FGLTFJsonUtility::WriteObjectPtrArray(JsonWriter, TEXT("skins"), Skins, Extensions);
 		FGLTFJsonUtility::WriteObjectPtrArray(JsonWriter, TEXT("textures"), Textures, Extensions);
 
-		if (Backdrops.Num() > 0 || Hotspots.Num() > 0 || Lights.Num() > 0 || LightMaps.Num() > 0 || SkySpheres.Num() > 0 || Variations.Num() > 0)
+		if (Backdrops.Num() > 0 || Hotspots.Num() > 0 || Lights.Num() > 0 || LightMaps.Num() > 0 || SkySpheres.Num() > 0 || LevelVariantSets.Num() > 0)
 		{
 			JsonWriter.WriteObjectStart(TEXT("extensions"));
 
@@ -134,13 +134,13 @@ struct FGLTFJsonRoot
 				JsonWriter.WriteObjectEnd();
 			}
 
-			if (Variations.Num() > 0)
+			if (LevelVariantSets.Num() > 0)
 			{
 				const EGLTFJsonExtension Extension = EGLTFJsonExtension::EPIC_LevelVariantSets;
 				Extensions.Used.Add(Extension);
 
 				JsonWriter.WriteObjectStart(FGLTFJsonUtility::ToString(Extension));
-				FGLTFJsonUtility::WriteObjectPtrArray(JsonWriter, TEXT("levelVariantSets"), Variations, Extensions);
+				FGLTFJsonUtility::WriteObjectPtrArray(JsonWriter, TEXT("levelVariantSets"), LevelVariantSets, Extensions);
 				JsonWriter.WriteObjectEnd();
 			}
 
diff --git a/Engine/Plugins/Enterprise/GLTFExporter/MigratedFromGit/Source/GLTFExporter/Private/Json/GLTFJsonScene.h b/Engine/Plugins/Enterprise/GLTFExporter/MigratedFromGit/Source/GLTFExporter/Private/Json/GLTFJsonScene.h
--- a/Engine/Plugins/Enterprise/GLTFExporter/MigratedFromGit/Source/GLTFExporter/Private/Json/GLTFJsonScene.h
+++ b/Engine/Plugins/Enterprise/GLTFExporter/MigratedFromGit/Source/GLTFExporter/Private/Json/GLTFJsonScene.h
@@ -10,7 +10,7 @@ struct FGLTFJsonScene
 	FString Name;
 
 	TArray<FGLTFJsonNodeIndex> Nodes;
-	TArray<FGLTFJsonVariationIndex>  Variations;
+	TArray<FGLTFJsonLevelVariantSetsIndex>  LevelVariantSets;
 
 	template <class CharType = TCHAR, class PrintPolicy = TPrettyJsonPrintPolicy<CharType>>
 	void WriteObject(TJsonWriter<CharType, PrintPolicy>& JsonWriter, FGLTFJsonExtensions& Extensions) const
@@ -27,7 +27,7 @@ struct FGLTFJsonScene
 			JsonWriter.WriteValue(TEXT("nodes"), Nodes);
 		}
 
-		if (Variations.Num() > 0)
+		if (LevelVariantSets.Num() > 0)
 		{
 			const EGLTFJsonExtension Extension = EGLTFJsonExtension::EPIC_LevelVariantSets;
 			Extensions.Used.Add(Extension);
@@ -35,7 +35,7 @@ struct FGLTFJsonScene
 			JsonWriter.WriteObjectStart(TEXT("extensions"));
 			JsonWriter.WriteObjectStart(FGLTFJsonUtility::ToString(Extension));
 
-			JsonWriter.WriteValue(TEXT("levelVariantSets"), Variations);
+			JsonWriter.WriteValue(TEXT("levelVariantSets"), LevelVariantSets);
 
 			JsonWriter.WriteObjectEnd();
 			JsonWriter.WriteObjectEnd();
diff --git a/Engine/Plugins/Enterprise/GLTFExporter/MigratedFromGit/Source/GLTFExporter/Private/Json/GLTFJsonVariation.h b/Engine/Plugins/Enterprise/GLTFExporter/MigratedFromGit/Source/GLTFExporter/Private/Json/GLTFJsonVariation.h
deleted file mode 100644
