[GLTFExporter] feat: add export option bStrictCompliance for truncating values into formal valid range according to gltf specs

commit 22cbc7ed933b819d91deefe7b8d3524319a7fc30
Date: Mon, 22 Feb 2021 17:44:44 +0100

    feat: add export option bStrictCompliance for truncating values into formal valid range according to gltf specs
    (cherry picked from commit 3b43cf23e2822299cdc0366a3f791f49c1ce8831​)

#skipci

diff --git a/Engine/Plugins/Enterprise/GLTFExporter/MigratedFromGit/Docs/export-options-reference.md b/Engine/Plugins/Enterprise/GLTFExporter/MigratedFromGit/Docs/export-options-reference.md
--- a/Engine/Plugins/Enterprise/GLTFExporter/MigratedFromGit/Docs/export-options-reference.md
+++ b/Engine/Plugins/Enterprise/GLTFExporter/MigratedFromGit/Docs/export-options-reference.md
@@ -4,6 +4,7 @@ Option                         | Description
 -------------------------------| ----------------------------------------------------------------------------------------------------------------------------
 `Export Uniform Scale`         | Scale factor used for exporting all assets (0.01 by default) for conversion from centimeters (Unreal default) to meters (glTF).
 `Export Preview Mesh`          | If enabled, the preview mesh for a standalone animation or material asset will also be exported.
+`Strict Compliance`            | If enabled, certain values (like HDR colors and light angles) will be truncated during export to strictly conform to the formal glTF specification.
 `Bundle Web Viewer`            | If enabled, Unreal's glTF viewer (including an executable launch helper) will be bundled with the exported files. It supports all extensions used.
 `Show Files When Done`         | If enabled, the target folder (with all exported files) will be shown in explorer once the export has been completed successfully.
 `Export Unlit Materials`       | If enabled, materials with shading model unlit will be properly exported. Uses extension KHR_materials_unlit.
diff --git a/Engine/Plugins/Enterprise/GLTFExporter/MigratedFromGit/Source/GLTFExporter/Private/Converters/GLTFConverterUtility.h b/Engine/Plugins/Enterprise/GLTFExporter/MigratedFromGit/Source/GLTFExporter/Private/Converters/GLTFConverterUtility.h
--- a/Engine/Plugins/Enterprise/GLTFExporter/MigratedFromGit/Source/GLTFExporter/Private/Converters/GLTFConverterUtility.h
+++ b/Engine/Plugins/Enterprise/GLTFExporter/MigratedFromGit/Source/GLTFExporter/Private/Converters/GLTFConverterUtility.h
@@ -75,8 +75,18 @@ struct FGLTFConverterUtility
 		return { UV.X, UV.Y };
 	}
 
-	static FGLTFJsonColor4 ConvertColor(const FLinearColor& Color)
+	static FGLTFJsonColor4 ConvertColor(const FLinearColor& Color, bool bForceLDR)
 	{
+		if (bForceLDR)
+		{
+			return {
+				FMath::Clamp(Color.R, 0.0f, 1.0f),
+                FMath::Clamp(Color.G, 0.0f, 1.0f),
+                FMath::Clamp(Color.B, 0.0f, 1.0f),
+                FMath::Clamp(Color.A, 0.0f, 1.0f)
+            };
+		}
+
 		// Just make sure its non-negative (which can happen when using MakeFromColorTemperature).
 		return {
 			FMath::Max(Color.R, 0.0f),
@@ -112,7 +122,8 @@ struct FGLTFConverterUtility
 			return FGLTFJsonQuaternion::Identity;
 		}
 
-		return { -Rotation.X, -Rotation.Z, -Rotation.Y, Rotation.W };
+		const FQuat Normalized = Rotation.GetNormalized();
+		return { -Normalized.X, -Normalized.Z, -Normalized.Y, Normalized.W };
 	}
 
 	static FGLTFJsonMatrix4 ConvertMatrix(const FMatrix& Matrix)
diff --git a/Engine/Plugins/Enterprise/GLTFExporter/MigratedFromGit/Source/GLTFExporter/Private/Converters/GLTFLightConverters.cpp b/Engine/Plugins/Enterprise/GLTFExporter/MigratedFromGit/Source/GLTFExporter/Private/Converters/GLTFLightConverters.cpp
--- a/Engine/Plugins/Enterprise/GLTFExporter/MigratedFromGit/Source/GLTFExporter/Private/Converters/GLTFLightConverters.cpp
+++ b/Engine/Plugins/Enterprise/GLTFExporter/MigratedFromGit/Source/GLTFExporter/Private/Converters/GLTFLightConverters.cpp
@@ -21,7 +21,7 @@ FGLTFJsonLightIndex FGLTFLightConverter::Convert(const ULightComponent* LightCom
 	Light.Intensity = LightComponent->Intensity;
 
 	const FLinearColor TemperatureColor = LightComponent->bUseTemperature ? FLinearColor::MakeFromColorTemperature(LightComponent->Temperature) : FLinearColor::White;
-	Light.Color = FGLTFConverterUtility::ConvertColor(TemperatureColor * LightComponent->GetLightColor());
+	Light.Color = FGLTFConverterUtility::ConvertColor(TemperatureColor * LightComponent->GetLightColor(), Builder.ExportOptions->bStrictCompliance);
 
 	if (const UPointLightComponent* PointLightComponent = Cast<UPointLightComponent>(LightComponent))
 	{
@@ -32,6 +32,12 @@ FGLTFJsonLightIndex FGLTFLightConverter::Convert(const ULightComponent* LightCom
 	{
 		Light.Spot.InnerConeAngle = FGLTFConverterUtility::ConvertLightAngle(SpotLightComponent->InnerConeAngle);
 		Light.Spot.OuterConeAngle = FGLTFConverterUtility::ConvertLightAngle(SpotLightComponent->OuterConeAngle);
+
+		if (Builder.ExportOptions->bStrictCompliance)
+		{
+			Light.Spot.InnerConeAngle = FMath::Clamp(Light.Spot.InnerConeAngle, 0.0f, nextafterf(Light.Spot.OuterConeAngle, 0.0f));
+			Light.Spot.OuterConeAngle = FMath::Clamp(Light.Spot.OuterConeAngle, nextafterf(Light.Spot.InnerConeAngle, HALF_PI), HALF_PI);
+		}
 	}
 
 	return Builder.AddLight(Light);
diff --git a/Engine/Plugins/Enterprise/GLTFExporter/MigratedFromGit/Source/GLTFExporter/Private/Converters/GLTFSkySphereConverters.cpp b/Engine/Plugins/Enterprise/GLTFExporter/MigratedFromGit/Source/GLTFExporter/Private/Converters/GLTFSkySphereConverters.cpp
--- a/Engine/Plugins/Enterprise/GLTFExporter/MigratedFromGit/Source/GLTFExporter/Private/Converters/GLTFSkySphereConverters.cpp
+++ b/Engine/Plugins/Enterprise/GLTFExporter/MigratedFromGit/Source/GLTFExporter/Private/Converters/GLTFSkySphereConverters.cpp
@@ -109,7 +109,7 @@ void FGLTFSkySphereConverter::ConvertColorProperty(const AActor* Actor, const TC
 	FLinearColor LinearColor;
 	if (FGLTFActorUtility::TryGetPropertyValue(Actor, PropertyName, LinearColor))
 	{
-		OutValue = FGLTFConverterUtility::ConvertColor(LinearColor);
+		OutValue = FGLTFConverterUtility::ConvertColor(LinearColor, Builder.ExportOptions->bStrictCompliance);
 	}
 	else
 	{
diff --git a/Engine/Plugins/Enterprise/GLTFExporter/MigratedFromGit/Source/GLTFExporter/Private/GLTFExportOptions.cpp b/Engine/Plugins/Enterprise/GLTFExporter/MigratedFromGit/Source/GLTFExporter/Private/GLTFExportOptions.cpp
--- a/Engine/Plugins/Enterprise/GLTFExporter/MigratedFromGit/Source/GLTFExporter/Private/GLTFExportOptions.cpp
+++ b/Engine/Plugins/Enterprise/GLTFExporter/MigratedFromGit/Source/GLTFExporter/Private/GLTFExportOptions.cpp
@@ -14,6 +14,7 @@ void UGLTFExportOptions::ResetToDefault()
 {
 	ExportUniformScale = 0.01;
 	bExportPreviewMesh = true;
+	bStrictCompliance = true;
 	bBundleWebViewer = true;
 	bShowFilesWhenDone = true;
 	bExportUnlitMaterials = true;
diff --git a/Engine/Plugins/Enterprise/GLTFExporter/MigratedFromGit/Source/GLTFExporter/Private/Tasks/GLTFMaterialTasks.cpp b/Engine/Plugins/Enterprise/GLTFExporter/MigratedFromGit/Source/GLTFExporter/Private/Tasks/GLTFMaterialTasks.cpp
--- a/Engine/Plugins/Enterprise/GLTFExporter/MigratedFromGit/Source/GLTFExporter/Private/Tasks/GLTFMaterialTasks.cpp
+++ b/Engine/Plugins/Enterprise/GLTFExporter/MigratedFromGit/Source/GLTFExporter/Private/Tasks/GLTFMaterialTasks.cpp
@@ -359,7 +359,7 @@ bool FGLTFMaterialTask::TryGetBaseColorAndOpacity(FGLTFJsonPBRMetallicRoughness&
 		FLinearColor BaseColorFactor(BaseColorBakeOutput.ConstantValue * BaseColorScale);
 		BaseColorFactor.A = OpacityBakeOutput.ConstantValue.A;
 
-		OutPBRParams.BaseColorFactor = FGLTFConverterUtility::ConvertColor(BaseColorFactor);
+		OutPBRParams.BaseColorFactor = FGLTFConverterUtility::ConvertColor(BaseColorFactor, Builder.ExportOptions->bStrictCompliance);
 		return true;
 	}
 
@@ -410,7 +410,7 @@ bool FGLTFMaterialTask::TryGetBaseColorAndOpacity(FGLTFJsonPBRMetallicRoughness&
 
 	OutPBRParams.BaseColorTexture.TexCoord = TexCoord;
 	OutPBRParams.BaseColorTexture.Index = TextureIndex;
-	OutPBRParams.BaseColorFactor = { BaseColorScale, BaseColorScale, BaseColorScale };
+	OutPBRParams.BaseColorFactor = FGLTFConverterUtility::ConvertColor({ BaseColorScale, BaseColorScale, BaseColorScale }, Builder.ExportOptions->bStrictCompliance);
 
 	return true;
 }
@@ -609,7 +609,7 @@ bool FGLTFMaterialTask::TryGetClearCoatRoughness(FGLTFJsonClearCoatExtension& Ou
 	if (IntensityBakeOutput.bIsConstant && RoughnessBakeOutput.bIsConstant)
 	{
 		OutExtParams.ClearCoatFactor = IntensityBakeOutput.ConstantValue.R;
-		OutExtParams.ClearCoatRoughnessFactor =  RoughnessBakeOutput.ConstantValue.R;
+		OutExtParams.ClearCoatRoughnessFactor = RoughnessBakeOutput.ConstantValue.R;
 		return true;
 	}
 
@@ -697,7 +697,7 @@ bool FGLTFMaterialTask::TryGetEmissive(FGLTFJsonMaterial& JsonMaterial, const FM
 	if (PropertyBakeOutput.bIsConstant)
 	{
 		const FLinearColor EmissiveColor = PropertyBakeOutput.ConstantValue;
-		JsonMaterial.EmissiveFactor = FGLTFConverterUtility::ConvertColor(EmissiveColor * EmissiveScale);
+		JsonMaterial.EmissiveFactor = FGLTFConverterUtility::ConvertColor(EmissiveColor * EmissiveScale, Builder.ExportOptions->bStrictCompliance);
 	}
 	else
 	{
@@ -748,7 +748,7 @@ bool FGLTFMaterialTask::TryGetConstantColor(FGLTFJsonColor3& OutValue, const FMa
 	FLinearColor Value;
 	if (TryGetConstantColor(Value, Property))
 	{
-		OutValue = FGLTFConverterUtility::ConvertColor(Value);
+		OutValue = FGLTFConverterUtility::ConvertColor(Value, Builder.ExportOptions->bStrictCompliance);
 		return true;
 	}
 
@@ -760,7 +760,7 @@ bool FGLTFMaterialTask::TryGetConstantColor(FGLTFJsonColor4& OutValue, const FMa
 	FLinearColor Value;
 	if (TryGetConstantColor(Value, Property))
 	{
-		OutValue = FGLTFConverterUtility::ConvertColor(Value);
+		OutValue = FGLTFConverterUtility::ConvertColor(Value, Builder.ExportOptions->bStrictCompliance);
 		return true;
 	}
 
@@ -1123,7 +1123,7 @@ bool FGLTFMaterialTask::TryGetBakedMaterialProperty(FGLTFJsonTextureInfo& OutTex
 
 	if (PropertyBakeOutput.bIsConstant)
 	{
-		OutConstant = FGLTFConverterUtility::ConvertColor(PropertyBakeOutput.ConstantValue);
+		OutConstant = FGLTFConverterUtility::ConvertColor(PropertyBakeOutput.ConstantValue, Builder.ExportOptions->bStrictCompliance);
 		return true;
 	}
 
@@ -1162,7 +1162,7 @@ bool FGLTFMaterialTask::TryGetBakedMaterialProperty(FGLTFJsonTextureInfo& OutTex
 
 	if (PropertyBakeOutput.bIsConstant)
 	{
-		OutConstant = FGLTFConverterUtility::ConvertColor(PropertyBakeOutput.ConstantValue);
+		OutConstant = FGLTFConverterUtility::ConvertColor(PropertyBakeOutput.ConstantValue, Builder.ExportOptions->bStrictCompliance);
 		return true;
 	}
 
diff --git a/Engine/Plugins/Enterprise/GLTFExporter/MigratedFromGit/Source/GLTFExporter/Public/GLTFExportOptions.h b/Engine/Plugins/Enterprise/GLTFExporter/MigratedFromGit/Source/GLTFExporter/Public/GLTFExportOptions.h
--- a/Engine/Plugins/Enterprise/GLTFExporter/MigratedFromGit/Source/GLTFExporter/Public/GLTFExportOptions.h
+++ b/Engine/Plugins/Enterprise/GLTFExporter/MigratedFromGit/Source/GLTFExporter/Public/GLTFExportOptions.h
@@ -76,6 +76,10 @@ class GLTFEXPORTER_API UGLTFExportOptions : public UObject
 	UPROPERTY(EditAnywhere, BlueprintReadWrite, Config, Category = General)
 	bool bExportPreviewMesh;
 
+	/** If enabled, certain values (like HDR colors and light angles) will be truncated during export to strictly conform to the formal glTF specification. */
+	UPROPERTY(EditAnywhere, BlueprintReadWrite, Config, Category = General)
+	bool bStrictCompliance;
+
 	/** If enabled, Unreal's glTF viewer (including an executable launch helper) will be bundled with the exported files. It supports all extensions used. */
 	UPROPERTY(EditAnywhere, BlueprintReadWrite, Config, Category = General)
 	bool bBundleWebViewer;
