[GLTFExporter] refactor: remove bundled Unreal glTF Viewer

commit 32262c3b6a8aaba68e53273c037be24e489bf153
Date: Fri, 19 Feb 2021 21:35:47 +0100

    refactor: remove bundled Unreal glTF Viewer
    (cherry picked from commit 2315edd2f38bffae3ec788bad05f82f50f4451c8​)

#skipci

diff --git a/Engine/Plugins/Enterprise/GLTFExporter/MigratedFromGit/Docs/export-options-reference.md b/Engine/Plugins/Enterprise/GLTFExporter/MigratedFromGit/Docs/export-options-reference.md
--- a/Engine/Plugins/Enterprise/GLTFExporter/MigratedFromGit/Docs/export-options-reference.md
+++ b/Engine/Plugins/Enterprise/GLTFExporter/MigratedFromGit/Docs/export-options-reference.md
@@ -5,7 +5,6 @@ Option                         | Description
 `Export Uniform Scale`         | Scale factor used for exporting all assets (0.01 by default) for conversion from centimeters (Unreal default) to meters (glTF).
 `Export Preview Mesh`          | If enabled, the preview mesh for a standalone animation or material asset will also be exported.
 `Strict Compliance`            | If enabled, certain values (like HDR colors and light angles) will be truncated during export to strictly conform to the formal glTF specification.
-`Bundle Web Viewer`            | If enabled, Unreal's glTF viewer (including an executable launch helper) will be bundled with the exported files. It supports all extensions used.
 `Show Files When Done`         | If enabled, the target folder (with all exported files) will be shown in explorer once the export has been completed successfully.
 `Export Unlit Materials`       | If enabled, materials with shading model unlit will be properly exported. Uses extension KHR_materials_unlit.
 `Export Clear Coat Materials`  | If enabled, materials with shading model clear coat will be properly exported. Uses extension KHR_materials_clearcoat, which is not supported by all glTF viewers.
diff --git a/Engine/Plugins/Enterprise/GLTFExporter/MigratedFromGit/README.md b/Engine/Plugins/Enterprise/GLTFExporter/MigratedFromGit/README.md
--- a/Engine/Plugins/Enterprise/GLTFExporter/MigratedFromGit/README.md
+++ b/Engine/Plugins/Enterprise/GLTFExporter/MigratedFromGit/README.md
@@ -36,8 +36,6 @@ In Unreal Editor, open `Menu -> Edit -> Plugins` and make sure `glTF Exporter` i
   1. Change `Save as type` to `.gltf` (or `.glb`) and click `Save`
   1. When `glTF Export Options` window is shown, click `Export`
 
-If the Web Viewer is bundled, run `GLTFLaunchHelper.exe` in the target folder to view the export asset in the Unreal glTF Viewer.
-
 
 ## Documentation
 
diff --git a/Engine/Plugins/Enterprise/GLTFExporter/MigratedFromGit/Source/GLTFExporter/GLTFExporter.Build.cs b/Engine/Plugins/Enterprise/GLTFExporter/MigratedFromGit/Source/GLTFExporter/GLTFExporter.Build.cs
--- a/Engine/Plugins/Enterprise/GLTFExporter/MigratedFromGit/Source/GLTFExporter/GLTFExporter.Build.cs
+++ b/Engine/Plugins/Enterprise/GLTFExporter/MigratedFromGit/Source/GLTFExporter/GLTFExporter.Build.cs
@@ -33,7 +33,6 @@ namespace UnrealBuildTool.Rules
 					"MovieSceneTracks",
 					"VariantManagerContent",
 					"MeshMergeUtilities",
-					"zlib",
 					"GLTFMaterialBaking",
 					"GLTFMaterialAnalyzer"
 				}
diff --git a/Engine/Plugins/Enterprise/GLTFExporter/MigratedFromGit/Source/GLTFExporter/Private/Builders/GLTFContainerBuilder.cpp b/Engine/Plugins/Enterprise/GLTFExporter/MigratedFromGit/Source/GLTFExporter/Private/Builders/GLTFContainerBuilder.cpp
--- a/Engine/Plugins/Enterprise/GLTFExporter/MigratedFromGit/Source/GLTFExporter/Private/Builders/GLTFContainerBuilder.cpp
+++ b/Engine/Plugins/Enterprise/GLTFExporter/MigratedFromGit/Source/GLTFExporter/Private/Builders/GLTFContainerBuilder.cpp
@@ -2,6 +2,32 @@
 
 #include "Builders/GLTFContainerBuilder.h"
 
+void FGLTFContainerBuilder::Write(FArchive& Archive, FFeedbackContext* Context)
+{
+	CompleteAllTasks(Context);
+
+	if (bIsGlbFile)
+	{
+		WriteGlb(Archive);
+	}
+	else
+	{
+		WriteJson(Archive);
+	}
+
+	const TSet<EGLTFJsonExtension> CustomExtensions = GetCustomExtensionsUsed();
+	if (CustomExtensions.Num() > 0)
+	{
+		const FString ExtensionsString = FString::JoinBy(CustomExtensions, TEXT(", "),
+            [](EGLTFJsonExtension Extension)
+        {
+            return FGLTFJsonUtility::ToString(Extension);
+        });
+
+		AddWarningMessage(FString::Printf(TEXT("Export uses some extensions that may only be supported in Unreal's glTF viewer: %s"), *ExtensionsString));
+	}
+}
+
 FGLTFContainerBuilder::FGLTFContainerBuilder(const FString& FilePath, const UGLTFExportOptions* ExportOptions, bool bSelectedActorsOnly)
 	: FGLTFConvertBuilder(FilePath, ExportOptions, bSelectedActorsOnly)
 {
diff --git a/Engine/Plugins/Enterprise/GLTFExporter/MigratedFromGit/Source/GLTFExporter/Private/Builders/GLTFContainerBuilder.h b/Engine/Plugins/Enterprise/GLTFExporter/MigratedFromGit/Source/GLTFExporter/Private/Builders/GLTFContainerBuilder.h
--- a/Engine/Plugins/Enterprise/GLTFExporter/MigratedFromGit/Source/GLTFExporter/Private/Builders/GLTFContainerBuilder.h
+++ b/Engine/Plugins/Enterprise/GLTFExporter/MigratedFromGit/Source/GLTFExporter/Private/Builders/GLTFContainerBuilder.h
@@ -6,10 +6,14 @@
 
 class FGLTFContainerBuilder : public FGLTFConvertBuilder
 {
-protected:
+public:
 
 	FGLTFContainerBuilder(const FString& FilePath, const UGLTFExportOptions* ExportOptions, bool bSelectedActorsOnly);
 
+	void Write(FArchive& Archive, FFeedbackContext* Context);
+
+protected:
+
 	void WriteGlb(FArchive& Archive);
 
 private:
diff --git a/Engine/Plugins/Enterprise/GLTFExporter/MigratedFromGit/Source/GLTFExporter/Private/Builders/GLTFWebBuilder.cpp b/Engine/Plugins/Enterprise/GLTFExporter/MigratedFromGit/Source/GLTFExporter/Private/Builders/GLTFWebBuilder.cpp
deleted file mode 100644
diff --git a/Engine/Plugins/Enterprise/GLTFExporter/MigratedFromGit/Source/GLTFExporter/Private/Builders/GLTFWebBuilder.h b/Engine/Plugins/Enterprise/GLTFExporter/MigratedFromGit/Source/GLTFExporter/Private/Builders/GLTFWebBuilder.h
deleted file mode 100644
diff --git a/Engine/Plugins/Enterprise/GLTFExporter/MigratedFromGit/Source/GLTFExporter/Private/Builders/GLTFZipUtility.cpp b/Engine/Plugins/Enterprise/GLTFExporter/MigratedFromGit/Source/GLTFExporter/Private/Builders/GLTFZipUtility.cpp
deleted file mode 100644
diff --git a/Engine/Plugins/Enterprise/GLTFExporter/MigratedFromGit/Source/GLTFExporter/Private/Builders/GLTFZipUtility.h b/Engine/Plugins/Enterprise/GLTFExporter/MigratedFromGit/Source/GLTFExporter/Private/Builders/GLTFZipUtility.h
deleted file mode 100644
diff --git a/Engine/Plugins/Enterprise/GLTFExporter/MigratedFromGit/Source/GLTFExporter/Private/Exporters/GLTFExporter.cpp b/Engine/Plugins/Enterprise/GLTFExporter/MigratedFromGit/Source/GLTFExporter/Private/Exporters/GLTFExporter.cpp
--- a/Engine/Plugins/Enterprise/GLTFExporter/MigratedFromGit/Source/GLTFExporter/Private/Exporters/GLTFExporter.cpp
+++ b/Engine/Plugins/Enterprise/GLTFExporter/MigratedFromGit/Source/GLTFExporter/Private/Exporters/GLTFExporter.cpp
@@ -3,7 +3,7 @@
 #include "Exporters/GLTFExporter.h"
 #include "GLTFExportOptions.h"
 #include "UI/GLTFExportOptionsWindow.h"
-#include "Builders/GLTFWebBuilder.h"
+#include "Builders/GLTFContainerBuilder.h"
 #include "UObject/GCObjectScopeGuard.h"
 #include "AssetExportTask.h"
 
@@ -33,7 +33,7 @@ bool UGLTFExporter::ExportBinary(UObject* Object, const TCHAR* Type, FArchive& A
 	// TODO: add support for UAssetExportTask::IgnoreObjectList?
 
 	FGCObjectScopeGuard OptionsGuard(Options);
-	FGLTFWebBuilder Builder(CurrentFilename, Options, bSelectedOnly);
+	FGLTFContainerBuilder Builder(CurrentFilename, Options, bSelectedOnly);
 
 	const bool bSuccess = AddObject(Builder, Object);
 	if (bSuccess)
diff --git a/Engine/Plugins/Enterprise/GLTFExporter/MigratedFromGit/Source/GLTFExporter/Private/GLTFExportOptions.cpp b/Engine/Plugins/Enterprise/GLTFExporter/MigratedFromGit/Source/GLTFExporter/Private/GLTFExportOptions.cpp
--- a/Engine/Plugins/Enterprise/GLTFExporter/MigratedFromGit/Source/GLTFExporter/Private/GLTFExportOptions.cpp
+++ b/Engine/Plugins/Enterprise/GLTFExporter/MigratedFromGit/Source/GLTFExporter/Private/GLTFExportOptions.cpp
@@ -15,11 +15,10 @@ void UGLTFExportOptions::ResetToDefault()
 	ExportUniformScale = 0.01;
 	bExportPreviewMesh = true;
 	bStrictCompliance = true;
-	bBundleWebViewer = true;
 	bShowFilesWhenDone = true;
 	bExportUnlitMaterials = true;
 	bExportClearCoatMaterials = true;
-	bExportExtraBlendModes = true;
+	bExportExtraBlendModes = false;
 	BakeMaterialInputs = EGLTFMaterialBakeMode::UseMeshData;
 	DefaultMaterialBakeSize = EGLTFMaterialBakeSizePOT::POT_1024;
 	DefaultMaterialBakeFilter = TF_Trilinear;
@@ -30,21 +29,21 @@ void UGLTFExportOptions::ResetToDefault()
 	bExportLevelSequences = true;
 	bExportAnimationSequences = true;
 	bRetargetBoneTransforms = true;
-	bExportPlaybackSettings = true;
+	bExportPlaybackSettings = false;
 	TextureImageFormat = EGLTFTextureImageFormat::PNG;
 	TextureImageQuality = 0;
 	NoLossyImageFormatFor = static_cast<int32>(EGLTFTextureType::All);
 	bExportTextureTransforms = true;
-	bExportLightmaps = true;
-	TextureHDREncoding = EGLTFTextureHDREncoding::RGBM;
+	bExportLightmaps = false;
+	TextureHDREncoding = EGLTFTextureHDREncoding::None;
 	bExportHiddenInGame = false;
 	ExportLights = static_cast<int32>(EGLTFSceneMobility::Stationary | EGLTFSceneMobility::Movable);
 	bExportCameras = true;
-	bExportCameraControls = true;
-	bExportAnimationHotspots = true;
-	bExportHDRIBackdrops = true;
-	bExportSkySpheres = true;
-	bExportVariantSets = true;
+	bExportCameraControls = false;
+	bExportAnimationHotspots = false;
+	bExportHDRIBackdrops = false;
+	bExportSkySpheres = false;
+	bExportVariantSets = false;
 	ExportMaterialVariants = EGLTFMaterialVariantMode::UseMeshData;
 	bExportMeshVariants = true;
 	bExportVisibilityVariants = true;
diff --git a/Engine/Plugins/Enterprise/GLTFExporter/MigratedFromGit/Source/GLTFExporter/Public/GLTFExportOptions.h b/Engine/Plugins/Enterprise/GLTFExporter/MigratedFromGit/Source/GLTFExporter/Public/GLTFExportOptions.h
--- a/Engine/Plugins/Enterprise/GLTFExporter/MigratedFromGit/Source/GLTFExporter/Public/GLTFExportOptions.h
+++ b/Engine/Plugins/Enterprise/GLTFExporter/MigratedFromGit/Source/GLTFExporter/Public/GLTFExportOptions.h
@@ -80,10 +80,6 @@ class GLTFEXPORTER_API UGLTFExportOptions : public UObject
 	UPROPERTY(EditAnywhere, BlueprintReadWrite, Config, Category = General)
 	bool bStrictCompliance;
 
-	/** If enabled, Unreal's glTF viewer (including an executable launch helper) will be bundled with the exported files. It supports all extensions used. */
-	UPROPERTY(EditAnywhere, BlueprintReadWrite, Config, Category = General)
-	bool bBundleWebViewer;
-
 	/** If enabled, the target folder (with all exported files) will be shown in explorer once the export has been completed successfully. */
 	UPROPERTY(EditAnywhere, BlueprintReadWrite, Config, Category = General)
 	bool bShowFilesWhenDone;
