[GLTFExporter] chore: bump version to 1.0.2

commit b9fd49b92c7d604a5fce0c54a22617803d460d3e
Date: Mon, 22 Feb 2021 18:52:12 +0100

    chore: bump version to 1.0.2
    (cherry picked from commit a40b208957f60a5e40c5b6047a2b8df5c3d7be3f​)

#skipci

diff --git a/Engine/Plugins/Enterprise/GLTFExporter/MigratedFromGit/GLTFExporter.uplugin b/Engine/Plugins/Enterprise/GLTFExporter/MigratedFromGit/GLTFExporter.uplugin
--- a/Engine/Plugins/Enterprise/GLTFExporter/MigratedFromGit/GLTFExporter.uplugin
+++ b/Engine/Plugins/Enterprise/GLTFExporter/MigratedFromGit/GLTFExporter.uplugin
@@ -1,7 +1,7 @@
 {
 	"FileVersion": 3,
-	"Version": 101,
-	"VersionName": "1.0.1",
+	"Version": 102,
+	"VersionName": "1.0.2",
 	"FriendlyName": "glTF Exporter",
 	"Description": "An exporter for Khronos glTF 2.0.",
 	"Category": "Exporters",
