[GLTFExporter] fix: add missing lightmaps for non-zero-LOD static meshes

commit ccf7dfac7bc3607e3870a91b9d2a5837945f11b8
Date: Wed, 24 Mar 2021 08:02:28 +0100

    fix: add missing lightmaps for non-zero-LOD static meshes
    (cherry picked from commit 533d70be79934d0a76d01e23da27d5ba38fc935a​)

#skipci

diff --git a/Engine/Plugins/Enterprise/GLTFExporter/MigratedFromGit/Source/GLTFExporter/Private/Converters/GLTFLightMapConverters.cpp b/Engine/Plugins/Enterprise/GLTFExporter/MigratedFromGit/Source/GLTFExporter/Private/Converters/GLTFLightMapConverters.cpp
--- a/Engine/Plugins/Enterprise/GLTFExporter/MigratedFromGit/Source/GLTFExporter/Private/Converters/GLTFLightMapConverters.cpp
+++ b/Engine/Plugins/Enterprise/GLTFExporter/MigratedFromGit/Source/GLTFExporter/Private/Converters/GLTFLightMapConverters.cpp
@@ -1,6 +1,7 @@
 // Copyright Epic Games, Inc. All Rights Reserved.
 
 #include "Converters/GLTFLightMapConverters.h"
+#include "Converters/GLTFMeshUtility.h"
 #include "Builders/GLTFContainerBuilder.h"
 #include "Engine/MapBuildDataRegistry.h"
 
@@ -13,22 +14,22 @@ FGLTFJsonLightMapIndex FGLTFLightMapConverter::Convert(const UStaticMeshComponen
 		return FGLTFJsonLightMapIndex(INDEX_NONE);
 	}
 
-	const int32 LODIndex = StaticMeshComponent->ForcedLodModel > 0 ? StaticMeshComponent->ForcedLodModel - 1 : Builder.ExportOptions->DefaultLevelOfDetail;
+	const int32 LODIndex = FGLTFMeshUtility::GetLOD(StaticMesh, StaticMeshComponent, Builder.ExportOptions->DefaultLevelOfDetail);
+	const FStaticMeshLODResources& LODResources = StaticMesh->GetLODForExport(LODIndex);
 
-	if (LODIndex < 0 || StaticMesh->GetNumLODs() <= LODIndex || StaticMeshComponent->LODData.Num() <= LODIndex)
+	const int32 CoordinateIndex = StaticMesh->LightMapCoordinateIndex;
+	if (CoordinateIndex < 0 || CoordinateIndex >= LODResources.GetNumTexCoords())
 	{
 		return FGLTFJsonLightMapIndex(INDEX_NONE);
 	}
 
-	const int32 CoordinateIndex = StaticMesh->LightMapCoordinateIndex;
-	const FStaticMeshLODResources& LODResources = StaticMesh->GetLODForExport(LODIndex);
-
-	if (CoordinateIndex < 0 || CoordinateIndex >= LODResources.GetNumTexCoords())
+	const int32 LightMapLODIndex = 0;
+	if (!StaticMeshComponent->LODData.IsValidIndex(LightMapLODIndex))
 	{
 		return FGLTFJsonLightMapIndex(INDEX_NONE);
 	}
 
-	const FStaticMeshComponentLODInfo& ComponentLODInfo = StaticMeshComponent->LODData[LODIndex];
+	const FStaticMeshComponentLODInfo& ComponentLODInfo = StaticMeshComponent->LODData[LightMapLODIndex];
 	const FMeshMapBuildData* MeshMapBuildData = StaticMeshComponent->GetMeshMapBuildData(ComponentLODInfo);
 
 	if (MeshMapBuildData == nullptr || MeshMapBuildData->LightMap == nullptr)
