[GLTFExporter] refactor: rename export option "Export Mesh Quantization" to "Use Mesh Quantization" and uncheck by default

commit 5d39fb8c65dfba3e9c240c6414e88e4fe232b7e0
Date: Wed, 24 Mar 2021 13:34:05 +0100

    refactor: rename export option "Export Mesh Quantization" to "Use Mesh Quantization" and uncheck by default
    (cherry picked from commit 0e091bedbc13882b04d79373dae979a45bf69f2d​)

#skipci

diff --git a/Engine/Plugins/Enterprise/GLTFExporter/MigratedFromGit/Docs/export-options-reference.md b/Engine/Plugins/Enterprise/GLTFExporter/MigratedFromGit/Docs/export-options-reference.md
--- a/Engine/Plugins/Enterprise/GLTFExporter/MigratedFromGit/Docs/export-options-reference.md
+++ b/Engine/Plugins/Enterprise/GLTFExporter/MigratedFromGit/Docs/export-options-reference.md
@@ -17,7 +17,7 @@ Option                         | Description
 `Default Level Of Detail`      | Default LOD level used for exporting a mesh. Can be overridden by component or asset settings (e.g. minimum or forced LOD level).
 `Export Vertex Colors`         | If enabled, export vertex color. Not recommended due to vertex colors always being used as a base color multiplier in glTF, regardless of material. Often producing undesirable results.
 `Export Vertex Skin Weights`   | If enabled, export vertex bone weights and indices in skeletal meshes. Necessary for animation sequences.
-`Export Mesh Quantization`     | If enabled, export Unreal-configured quantization for vertex tangents and normals, reducing size. Requires extension KHR_mesh_quantization, which is not supported by all glTF viewers.
+`Use Mesh Quantization`        | If enabled, export Unreal-configured quantization for vertex tangents and normals, reducing size. Requires extension KHR_mesh_quantization, which is not supported by all glTF viewers.
 `Export Level Sequences`       | If enabled, export level sequences. Only transform tracks are currently supported. The level sequence will be played at the assigned display rate.
 `Export Animation Sequences`   | If enabled, export single animation asset used by a skeletal mesh component or hotspot actor. Export of vertex skin weights must be enabled.
 `Retarget Bone Transforms`     | If enabled, apply animation retargeting to skeleton bones when exporting an animation sequence.
diff --git a/Engine/Plugins/Enterprise/GLTFExporter/MigratedFromGit/Source/GLTFExporter/Private/Converters/GLTFAccessorConverters.cpp b/Engine/Plugins/Enterprise/GLTFExporter/MigratedFromGit/Source/GLTFExporter/Private/Converters/GLTFAccessorConverters.cpp
--- a/Engine/Plugins/Enterprise/GLTFExporter/MigratedFromGit/Source/GLTFExporter/Private/Converters/GLTFAccessorConverters.cpp
+++ b/Engine/Plugins/Enterprise/GLTFExporter/MigratedFromGit/Source/GLTFExporter/Private/Converters/GLTFAccessorConverters.cpp
@@ -134,7 +134,7 @@ FGLTFJsonAccessorIndex FGLTFNormalBufferConverter::Convert(const FGLTFMeshSectio
 	FGLTFJsonBufferViewIndex BufferViewIndex;
 	EGLTFJsonComponentType ComponentType;
 
-	const bool bMeshQuantization = Builder.ExportOptions->bExportMeshQuantization;
+	const bool bMeshQuantization = Builder.ExportOptions->bUseMeshQuantization;
 	const bool bHighPrecision = VertexBuffer->GetUseHighPrecisionTangentBasis();
 
 	if (bMeshQuantization)
@@ -211,7 +211,7 @@ FGLTFJsonAccessorIndex FGLTFTangentBufferConverter::Convert(const FGLTFMeshSecti
 	FGLTFJsonBufferViewIndex BufferViewIndex;
 	EGLTFJsonComponentType ComponentType;
 
-	const bool bMeshQuantization = Builder.ExportOptions->bExportMeshQuantization;
+	const bool bMeshQuantization = Builder.ExportOptions->bUseMeshQuantization;
 	const bool bHighPrecision = VertexBuffer->GetUseHighPrecisionTangentBasis();
 
 	if (bMeshQuantization)
diff --git a/Engine/Plugins/Enterprise/GLTFExporter/MigratedFromGit/Source/GLTFExporter/Private/GLTFExportOptions.cpp b/Engine/Plugins/Enterprise/GLTFExporter/MigratedFromGit/Source/GLTFExporter/Private/GLTFExportOptions.cpp
--- a/Engine/Plugins/Enterprise/GLTFExporter/MigratedFromGit/Source/GLTFExporter/Private/GLTFExportOptions.cpp
+++ b/Engine/Plugins/Enterprise/GLTFExporter/MigratedFromGit/Source/GLTFExporter/Private/GLTFExportOptions.cpp
@@ -25,7 +25,7 @@ void UGLTFExportOptions::ResetToDefault()
 	DefaultMaterialBakeTiling = TA_Wrap;
 	bExportVertexColors = false;
 	bExportVertexSkinWeights = true;
-	bExportMeshQuantization = true;
+	bUseMeshQuantization = false;
 	bExportLevelSequences = true;
 	bExportAnimationSequences = true;
 	bRetargetBoneTransforms = true;
diff --git a/Engine/Plugins/Enterprise/GLTFExporter/MigratedFromGit/Source/GLTFExporter/Public/GLTFExportOptions.h b/Engine/Plugins/Enterprise/GLTFExporter/MigratedFromGit/Source/GLTFExporter/Public/GLTFExportOptions.h
--- a/Engine/Plugins/Enterprise/GLTFExporter/MigratedFromGit/Source/GLTFExporter/Public/GLTFExportOptions.h
+++ b/Engine/Plugins/Enterprise/GLTFExporter/MigratedFromGit/Source/GLTFExporter/Public/GLTFExportOptions.h
@@ -130,7 +130,7 @@ class GLTFEXPORTER_API UGLTFExportOptions : public UObject
 
 	/** If enabled, export Unreal-configured quantization for vertex tangents and normals, reducing size. Requires extension KHR_mesh_quantization, which is not supported by all glTF viewers. */
 	UPROPERTY(EditAnywhere, BlueprintReadWrite, Config, Category = Mesh)
-	bool bExportMeshQuantization;
+	bool bUseMeshQuantization;
 
 	/** If enabled, export level sequences. Only transform tracks are currently supported. The level sequence will be played at the assigned display rate. */
 	UPROPERTY(EditAnywhere, BlueprintReadWrite, Config, Category = Animation)
