[GLTFExporter] refactor: remove export option "Show Files When Done"

commit fce3736c5e04d55b645bc5e5932d7b5f981b2a15
Date: Wed, 24 Mar 2021 08:03:42 +0100

    refactor: remove export option "Show Files When Done"
    (cherry picked from commit c3cb3dbcca28bcf785aa07088e312bdc6f9159f9​)

#skipci

diff --git a/Engine/Plugins/Enterprise/GLTFExporter/MigratedFromGit/Docs/export-options-reference.md b/Engine/Plugins/Enterprise/GLTFExporter/MigratedFromGit/Docs/export-options-reference.md
--- a/Engine/Plugins/Enterprise/GLTFExporter/MigratedFromGit/Docs/export-options-reference.md
+++ b/Engine/Plugins/Enterprise/GLTFExporter/MigratedFromGit/Docs/export-options-reference.md
@@ -5,7 +5,6 @@ Option                         | Description
 `Export Uniform Scale`         | Scale factor used for exporting all assets (0.01 by default) for conversion from centimeters (Unreal default) to meters (glTF).
 `Export Preview Mesh`          | If enabled, the preview mesh for a standalone animation or material asset will also be exported.
 `Strict Compliance`            | If enabled, certain values (like HDR colors and light angles) will be truncated during export to strictly conform to the formal glTF specification.
-`Show Files When Done`         | If enabled, the target folder (with all exported files) will be shown in explorer once the export has been completed successfully.
 `Export Unlit Materials`       | If enabled, materials with shading model unlit will be properly exported. Uses extension KHR_materials_unlit.
 `Export Clear Coat Materials`  | If enabled, materials with shading model clear coat will be properly exported. Uses extension KHR_materials_clearcoat, which is not supported by all glTF viewers.
 `Export Extra Blend Modes`     | If enabled, materials with blend modes additive, modulate, and alpha composite will be properly exported. Uses extension EPIC_blend_modes, which is supported by Unreal's glTF viewer.
diff --git a/Engine/Plugins/Enterprise/GLTFExporter/MigratedFromGit/Source/GLTFExporter/Private/Exporters/GLTFExporter.cpp b/Engine/Plugins/Enterprise/GLTFExporter/MigratedFromGit/Source/GLTFExporter/Private/Exporters/GLTFExporter.cpp
--- a/Engine/Plugins/Enterprise/GLTFExporter/MigratedFromGit/Source/GLTFExporter/Private/Exporters/GLTFExporter.cpp
+++ b/Engine/Plugins/Enterprise/GLTFExporter/MigratedFromGit/Source/GLTFExporter/Private/Exporters/GLTFExporter.cpp
@@ -50,11 +50,6 @@ bool UGLTFExporter::ExportBinary(UObject* Object, const TCHAR* Type, FArchive& A
 	else
 	{
 		Builder.ShowMessages();
-
-		if (bSuccess && Options->bShowFilesWhenDone)
-		{
-			FPlatformProcess::ExploreFolder(*Builder.DirPath);
-		}
 	}
 
 	return bSuccess;
diff --git a/Engine/Plugins/Enterprise/GLTFExporter/MigratedFromGit/Source/GLTFExporter/Private/GLTFExportOptions.cpp b/Engine/Plugins/Enterprise/GLTFExporter/MigratedFromGit/Source/GLTFExporter/Private/GLTFExportOptions.cpp
--- a/Engine/Plugins/Enterprise/GLTFExporter/MigratedFromGit/Source/GLTFExporter/Private/GLTFExportOptions.cpp
+++ b/Engine/Plugins/Enterprise/GLTFExporter/MigratedFromGit/Source/GLTFExporter/Private/GLTFExportOptions.cpp
@@ -15,7 +15,6 @@ void UGLTFExportOptions::ResetToDefault()
 	ExportUniformScale = 0.01;
 	bExportPreviewMesh = true;
 	bStrictCompliance = true;
-	bShowFilesWhenDone = true;
 	bExportUnlitMaterials = true;
 	bExportClearCoatMaterials = true;
 	bExportExtraBlendModes = false;
diff --git a/Engine/Plugins/Enterprise/GLTFExporter/MigratedFromGit/Source/GLTFExporter/Public/GLTFExportOptions.h b/Engine/Plugins/Enterprise/GLTFExporter/MigratedFromGit/Source/GLTFExporter/Public/GLTFExportOptions.h
--- a/Engine/Plugins/Enterprise/GLTFExporter/MigratedFromGit/Source/GLTFExporter/Public/GLTFExportOptions.h
+++ b/Engine/Plugins/Enterprise/GLTFExporter/MigratedFromGit/Source/GLTFExporter/Public/GLTFExportOptions.h
@@ -80,10 +80,6 @@ class GLTFEXPORTER_API UGLTFExportOptions : public UObject
 	UPROPERTY(EditAnywhere, BlueprintReadWrite, Config, Category = General)
 	bool bStrictCompliance;
 
-	/** If enabled, the target folder (with all exported files) will be shown in explorer once the export has been completed successfully. */
-	UPROPERTY(EditAnywhere, BlueprintReadWrite, Config, Category = General)
-	bool bShowFilesWhenDone;
-
 	/** If enabled, materials with shading model unlit will be properly exported. Uses extension KHR_materials_unlit. */
 	UPROPERTY(EditAnywhere, BlueprintReadWrite, Config, Category = Material)
 	bool bExportUnlitMaterials;
