[GLTFExporter] refactor: skip analyzing material property if no input node is connected

commit 8d0cab68e56c2e8f2a7ebdab19b99089aa4554f4
Date: Wed, 24 Mar 2021 11:50:41 +0100

    refactor: skip analyzing material property if no input node is connected
    (cherry picked from commit 0d3d53868bf46d01f29c438e37b9aed9492b612d​)

#skipci

diff --git a/Engine/Plugins/Enterprise/GLTFExporter/MigratedFromGit/Source/GLTFExporter/Private/Converters/GLTFMaterialUtility.cpp b/Engine/Plugins/Enterprise/GLTFExporter/MigratedFromGit/Source/GLTFExporter/Private/Converters/GLTFMaterialUtility.cpp
--- a/Engine/Plugins/Enterprise/GLTFExporter/MigratedFromGit/Source/GLTFExporter/Private/Converters/GLTFMaterialUtility.cpp
+++ b/Engine/Plugins/Enterprise/GLTFExporter/MigratedFromGit/Source/GLTFExporter/Private/Converters/GLTFMaterialUtility.cpp
@@ -471,6 +471,12 @@ bool FGLTFMaterialUtility::NeedsMeshData(const TArray<const UMaterialInterface*>
 
 void FGLTFMaterialUtility::AnalyzeMaterialProperty(const UMaterialInterface* InMaterial, const FMaterialPropertyEx& InProperty, FGLTFMaterialAnalysis& OutAnalysis)
 {
+	if (GetInputForProperty(InMaterial, InProperty) == nullptr)
+	{
+		OutAnalysis = FGLTFMaterialAnalysis();
+		return;
+	}
+
 	UGLTFMaterialAnalyzer::AnalyzeMaterialPropertyEx(InMaterial, InProperty.Type, InProperty.CustomOutput.ToString(), OutAnalysis);
 }
 
