[GLTFExporter] chore: bump version to 1.0.3

commit 08e6911ebdec2e7ae2119afbb956ab084a113105
Date: Wed, 24 Mar 2021 14:35:58 +0100

    chore: bump version to 1.0.3
    (cherry picked from commit bc2bb3ca587944e76f0c4d4f9bccb35b96f471a4​)

#skipci

diff --git a/Engine/Plugins/Enterprise/GLTFExporter/MigratedFromGit/GLTFExporter.uplugin b/Engine/Plugins/Enterprise/GLTFExporter/MigratedFromGit/GLTFExporter.uplugin
--- a/Engine/Plugins/Enterprise/GLTFExporter/MigratedFromGit/GLTFExporter.uplugin
+++ b/Engine/Plugins/Enterprise/GLTFExporter/MigratedFromGit/GLTFExporter.uplugin
@@ -1,7 +1,7 @@
 {
 	"FileVersion": 3,
-	"Version": 102,
-	"VersionName": "1.0.2",
+	"Version": 103,
+	"VersionName": "1.0.3",
 	"FriendlyName": "glTF Exporter",
 	"Description": "An exporter for Khronos glTF 2.0.",
 	"Category": "Exporters",
