[GLTFExporter] fix: avoid creating empty binary buffer if not needed

commit 092aec11360e4a7c8d80354c89f15263ee20a6d3
Date: Fri, 26 Mar 2021 10:57:50 +0100

    fix: avoid creating empty binary buffer if not needed
    (cherry picked from commit a48069fffb84dff6aec49015cbf6213c0c6bb360​)

#skipci

diff --git a/Engine/Plugins/Enterprise/GLTFExporter/MigratedFromGit/Source/GLTFExporter/Private/Builders/GLTFBufferBuilder.cpp b/Engine/Plugins/Enterprise/GLTFExporter/MigratedFromGit/Source/GLTFExporter/Private/Builders/GLTFBufferBuilder.cpp
--- a/Engine/Plugins/Enterprise/GLTFExporter/MigratedFromGit/Source/GLTFExporter/Private/Builders/GLTFBufferBuilder.cpp
+++ b/Engine/Plugins/Enterprise/GLTFExporter/MigratedFromGit/Source/GLTFExporter/Private/Builders/GLTFBufferBuilder.cpp
@@ -1,12 +1,23 @@
 // Copyright Epic Games, Inc. All Rights Reserved.
 
 #include "Builders/GLTFBufferBuilder.h"
-#include "Serialization/BufferArchive.h"
 #include "Misc/FileHelper.h"
 #include "Misc/Paths.h"
 
 FGLTFBufferBuilder::FGLTFBufferBuilder(const FString& FilePath, const UGLTFExportOptions* ExportOptions)
 	: FGLTFJsonBuilder(FilePath, ExportOptions)
+{
+}
+
+FGLTFBufferBuilder::~FGLTFBufferBuilder()
+{
+	if (BufferArchive != nullptr)
+	{
+		BufferArchive->Close();
+	}
+}
+
+bool FGLTFBufferBuilder::InitializeBuffer()
 {
 	FGLTFJsonBuffer JsonBuffer;
 
@@ -23,29 +34,22 @@ FGLTFBufferBuilder::FGLTFBufferBuilder(const FString& FilePath, const UGLTFExpor
 		if (BufferArchive == nullptr)
 		{
 			AddErrorMessage(FString::Printf(TEXT("Failed to write external binary buffer to file: %s"), *ExternalBinaryPath));
+			return false;
 		}
 	}
 
 	BufferIndex = AddBuffer(JsonBuffer);
+	return true;
 }
 
-FGLTFBufferBuilder::~FGLTFBufferBuilder()
-{
-	if (BufferArchive != nullptr)
-	{
-		BufferArchive->Close();
-	}
-}
-
-const TArray<uint8>& FGLTFBufferBuilder::GetBufferData() const
+const FBufferArchive* FGLTFBufferBuilder::GetBufferData() const
 {
-	check(bIsGlbFile);
-	return *static_cast<FBufferArchive*>(BufferArchive.Get());
+	return bIsGlbFile ? static_cast<FBufferArchive*>(BufferArchive.Get()) : nullptr;
 }
 
 FGLTFJsonBufferViewIndex FGLTFBufferBuilder::AddBufferView(const void* RawData, uint64 ByteLength, EGLTFJsonBufferTarget BufferTarget, uint8 DataAlignment)
 {
-	if (BufferArchive == nullptr)
+	if (BufferArchive == nullptr && !InitializeBuffer())
 	{
 		// TODO: report error
 		return FGLTFJsonBufferViewIndex(INDEX_NONE);
diff --git a/Engine/Plugins/Enterprise/GLTFExporter/MigratedFromGit/Source/GLTFExporter/Private/Builders/GLTFBufferBuilder.h b/Engine/Plugins/Enterprise/GLTFExporter/MigratedFromGit/Source/GLTFExporter/Private/Builders/GLTFBufferBuilder.h
--- a/Engine/Plugins/Enterprise/GLTFExporter/MigratedFromGit/Source/GLTFExporter/Private/Builders/GLTFBufferBuilder.h
+++ b/Engine/Plugins/Enterprise/GLTFExporter/MigratedFromGit/Source/GLTFExporter/Private/Builders/GLTFBufferBuilder.h
@@ -3,6 +3,7 @@
 #pragma once
 
 #include "Builders/GLTFJsonBuilder.h"
+#include "Serialization/BufferArchive.h"
 
 class FGLTFBufferBuilder : public FGLTFJsonBuilder
 {
@@ -11,7 +12,7 @@ protected:
 	FGLTFBufferBuilder(const FString& FilePath, const UGLTFExportOptions* ExportOptions);
 	~FGLTFBufferBuilder();
 
-	const TArray<uint8>& GetBufferData() const;
+	const FBufferArchive* GetBufferData() const;
 
 public:
 
@@ -25,6 +26,8 @@ public:
 
 private:
 
+	bool InitializeBuffer();
+
 	FGLTFJsonBufferIndex BufferIndex;
 	TUniquePtr<FArchive> BufferArchive;
 };
diff --git a/Engine/Plugins/Enterprise/GLTFExporter/MigratedFromGit/Source/GLTFExporter/Private/Builders/GLTFContainerBuilder.cpp b/Engine/Plugins/Enterprise/GLTFExporter/MigratedFromGit/Source/GLTFExporter/Private/Builders/GLTFContainerBuilder.cpp
--- a/Engine/Plugins/Enterprise/GLTFExporter/MigratedFromGit/Source/GLTFExporter/Private/Builders/GLTFContainerBuilder.cpp
+++ b/Engine/Plugins/Enterprise/GLTFExporter/MigratedFromGit/Source/GLTFExporter/Private/Builders/GLTFContainerBuilder.cpp
@@ -38,23 +38,30 @@ void FGLTFContainerBuilder::WriteGlb(FArchive& Archive)
 	FBufferArchive JsonData;
 	WriteJson(JsonData);
 
-	const TArray<uint8>& BufferData = GetBufferData();
-
-	WriteGlb(Archive, JsonData, BufferData);
+	const FBufferArchive* BufferData = GetBufferData();
+	WriteGlb(Archive, &JsonData, BufferData);
 }
 
-void FGLTFContainerBuilder::WriteGlb(FArchive& Archive, const TArray<uint8>& JsonData, const TArray<uint8>& BinaryData)
+void FGLTFContainerBuilder::WriteGlb(FArchive& Archive, const FBufferArchive* JsonData, const FBufferArchive* BinaryData)
 {
 	const uint32 JsonChunkType = 0x4E4F534A; // "JSON" in ASCII
 	const uint32 BinaryChunkType = 0x004E4942; // "BIN" in ASCII
-	const uint32 FileSize =
+	uint32 FileSize =
 		3 * sizeof(uint32) +
-		2 * sizeof(uint32) + GetPaddedChunkSize(JsonData.Num()) +
-		2 * sizeof(uint32) + GetPaddedChunkSize(BinaryData.Num());
+		2 * sizeof(uint32) + GetPaddedChunkSize(JsonData->Num());
+
+	if (BinaryData != nullptr)
+	{
+		FileSize += 2 * sizeof(uint32) + GetPaddedChunkSize(BinaryData->Num());
+	}
 
 	WriteHeader(Archive, FileSize);
-	WriteChunk(Archive, JsonChunkType, JsonData, 0x20);
-	WriteChunk(Archive, BinaryChunkType, BinaryData, 0x0);
+	WriteChunk(Archive, JsonChunkType, *JsonData, 0x20);
+
+	if (BinaryData != nullptr)
+	{
+		WriteChunk(Archive, BinaryChunkType, *BinaryData, 0x0);
+	}
 }
 
 void FGLTFContainerBuilder::WriteHeader(FArchive& Archive, uint32 FileSize)
diff --git a/Engine/Plugins/Enterprise/GLTFExporter/MigratedFromGit/Source/GLTFExporter/Private/Builders/GLTFContainerBuilder.h b/Engine/Plugins/Enterprise/GLTFExporter/MigratedFromGit/Source/GLTFExporter/Private/Builders/GLTFContainerBuilder.h
--- a/Engine/Plugins/Enterprise/GLTFExporter/MigratedFromGit/Source/GLTFExporter/Private/Builders/GLTFContainerBuilder.h
+++ b/Engine/Plugins/Enterprise/GLTFExporter/MigratedFromGit/Source/GLTFExporter/Private/Builders/GLTFContainerBuilder.h
@@ -18,7 +18,7 @@ protected:
 
 private:
 
-	static void WriteGlb(FArchive& Archive, const TArray<uint8>& JsonData, const TArray<uint8>& BinaryData);
+	static void WriteGlb(FArchive& Archive, const FBufferArchive* JsonData, const FBufferArchive* BinaryData);
 
 	static void WriteHeader(FArchive& Archive, uint32 FileSize);
 	static void WriteChunk(FArchive& Archive, uint32 ChunkType, const TArray<uint8>& ChunkData, uint8 ChunkTrailingByte);
