[GLTFExporter] fix: avoid exporting unnecessary brush and physics volume in level

commit ec38031c675ca4663a387739077f581d446fec1c
Date: Fri, 26 Mar 2021 10:58:15 +0100

    fix: avoid exporting unnecessary brush and physics volume in level
    (cherry picked from commit 58b50a4e553003f0a24946ce3ca6f3d5e471e6cc​)

#skipci

diff --git a/Engine/Plugins/Enterprise/GLTFExporter/MigratedFromGit/Source/GLTFExporter/Private/Converters/GLTFSceneConverters.cpp b/Engine/Plugins/Enterprise/GLTFExporter/MigratedFromGit/Source/GLTFExporter/Private/Converters/GLTFSceneConverters.cpp
--- a/Engine/Plugins/Enterprise/GLTFExporter/MigratedFromGit/Source/GLTFExporter/Private/Converters/GLTFSceneConverters.cpp
+++ b/Engine/Plugins/Enterprise/GLTFExporter/MigratedFromGit/Source/GLTFExporter/Private/Converters/GLTFSceneConverters.cpp
@@ -23,6 +23,16 @@ FGLTFJsonSceneIndex FGLTFSceneConverter::Convert(const UWorld* World)
 
 		for (const AActor* Actor : Level->Actors)
 		{
+			if (Actor == Level->GetDefaultBrush())
+			{
+				continue; // TODO: can we safely assume no other actor is ever attached to the default brush?
+			}
+
+			if (World->HasDefaultPhysicsVolume() && Actor == World->GetDefaultPhysicsVolume())
+			{
+				continue; // TODO: can we safely assume no other actor is ever attached to the default physics volume?
+			}
+
 			// TODO: should a LevelVariantSet be exported even if not selected for export?
 			if (const ALevelVariantSetsActor *LevelVariantSetsActor = Cast<ALevelVariantSetsActor>(Actor))
 			{
@@ -42,7 +52,8 @@ FGLTFJsonSceneIndex FGLTFSceneConverter::Convert(const UWorld* World)
 			const FGLTFJsonNodeIndex NodeIndex = Builder.GetOrAddNode(Actor);
 			if (NodeIndex != INDEX_NONE && FGLTFActorUtility::IsRootActor(Actor, Builder.bSelectedActorsOnly))
 			{
-				// TODO: to avoid having to add irrelevant actors/components let GLTFComponentConverter decide and add root nodes to scene. This change may require node converters to support cyclic calls.
+				// TODO: to avoid having to add irrelevant actors/components let GLTFComponentConverter decide and add root nodes to scene.
+				// This change may require node converters to support cyclic calls.
 				Scene.Nodes.Add(NodeIndex);
 			}
 		}
