[GLTFExporter] docs: improve tooltip for export option "Use Mesh Quantization"

commit 92a5fdcd2da91415a9fcbe33878650febcee30f2
Date: Fri, 26 Mar 2021 10:59:27 +0100

    docs: improve tooltip for export option "Use Mesh Quantization"
    (cherry picked from commit d1af05e5d861f28f754d974fb0d1e5c53007fdcc​)

#skipci

diff --git a/Engine/Plugins/Enterprise/GLTFExporter/MigratedFromGit/Docs/export-options-reference.md b/Engine/Plugins/Enterprise/GLTFExporter/MigratedFromGit/Docs/export-options-reference.md
--- a/Engine/Plugins/Enterprise/GLTFExporter/MigratedFromGit/Docs/export-options-reference.md
+++ b/Engine/Plugins/Enterprise/GLTFExporter/MigratedFromGit/Docs/export-options-reference.md
@@ -16,7 +16,7 @@ Option                         | Description
 `Default Level Of Detail`      | Default LOD level used for exporting a mesh. Can be overridden by component or asset settings (e.g. minimum or forced LOD level).
 `Export Vertex Colors`         | If enabled, export vertex color. Not recommended due to vertex colors always being used as a base color multiplier in glTF, regardless of material. Often producing undesirable results.
 `Export Vertex Skin Weights`   | If enabled, export vertex bone weights and indices in skeletal meshes. Necessary for animation sequences.
-`Use Mesh Quantization`        | If enabled, export Unreal-configured quantization for vertex tangents and normals, reducing size. Requires extension KHR_mesh_quantization, which is not supported by all glTF viewers.
+`Use Mesh Quantization`        | If enabled, use Unreal-configured quantization for vertex tangents and normals, reducing size. Requires extension KHR_mesh_quantization, the mesh might not load in some glTF viewers.
 `Export Level Sequences`       | If enabled, export level sequences. Only transform tracks are currently supported. The level sequence will be played at the assigned display rate.
 `Export Animation Sequences`   | If enabled, export single animation asset used by a skeletal mesh component or hotspot actor. Export of vertex skin weights must be enabled.
 `Retarget Bone Transforms`     | If enabled, apply animation retargeting to skeleton bones when exporting an animation sequence.
diff --git a/Engine/Plugins/Enterprise/GLTFExporter/MigratedFromGit/Source/GLTFExporter/Public/GLTFExportOptions.h b/Engine/Plugins/Enterprise/GLTFExporter/MigratedFromGit/Source/GLTFExporter/Public/GLTFExportOptions.h
--- a/Engine/Plugins/Enterprise/GLTFExporter/MigratedFromGit/Source/GLTFExporter/Public/GLTFExportOptions.h
+++ b/Engine/Plugins/Enterprise/GLTFExporter/MigratedFromGit/Source/GLTFExporter/Public/GLTFExportOptions.h
@@ -124,7 +124,7 @@ class GLTFEXPORTER_API UGLTFExportOptions : public UObject
 	UPROPERTY(EditAnywhere, BlueprintReadWrite, Config, Category = Mesh)
 	bool bExportVertexSkinWeights;
 
-	/** If enabled, export Unreal-configured quantization for vertex tangents and normals, reducing size. Requires extension KHR_mesh_quantization, which is not supported by all glTF viewers. */
+	/** If enabled, use Unreal-configured quantization for vertex tangents and normals, reducing size. Requires extension KHR_mesh_quantization, the mesh might not load in some glTF viewers. */
 	UPROPERTY(EditAnywhere, BlueprintReadWrite, Config, Category = Mesh)
 	bool bUseMeshQuantization;
 
