[GLTFExporter] fix: improve check for skipping analyzing material property if no input node is connected

commit 837e065650df941d6375d11081130f4271529bb2
Date: Fri, 26 Mar 2021 10:59:48 +0100

    fix: improve check for skipping analyzing material property if no input node is connected
    (cherry picked from commit bc1f2d8e0ed6f66f13c4869d3a68136c3fa18f36​)

#skipci

diff --git a/Engine/Plugins/Enterprise/GLTFExporter/MigratedFromGit/Source/GLTFExporter/Private/Converters/GLTFMaterialUtility.cpp b/Engine/Plugins/Enterprise/GLTFExporter/MigratedFromGit/Source/GLTFExporter/Private/Converters/GLTFMaterialUtility.cpp
--- a/Engine/Plugins/Enterprise/GLTFExporter/MigratedFromGit/Source/GLTFExporter/Private/Converters/GLTFMaterialUtility.cpp
+++ b/Engine/Plugins/Enterprise/GLTFExporter/MigratedFromGit/Source/GLTFExporter/Private/Converters/GLTFMaterialUtility.cpp
@@ -471,7 +471,8 @@ bool FGLTFMaterialUtility::NeedsMeshData(const TArray<const UMaterialInterface*>
 
 void FGLTFMaterialUtility::AnalyzeMaterialProperty(const UMaterialInterface* InMaterial, const FMaterialPropertyEx& InProperty, FGLTFMaterialAnalysis& OutAnalysis)
 {
-	if (GetInputForProperty(InMaterial, InProperty) == nullptr)
+	const FExpressionInput* Input = GetInputForProperty(InMaterial, InProperty);
+	if (Input == nullptr || Input->Expression == nullptr)
 	{
 		OutAnalysis = FGLTFMaterialAnalysis();
 		return;
