[GLTFExporter] fix: force recompile material if internal resource missing to avoid potential crash during material analysis

commit 981834513c3901ebf61263b942e678bd1919de99
Date: Fri, 26 Mar 2021 11:00:13 +0100

    fix: force recompile material if internal resource missing to avoid potential crash during material analysis
    (cherry picked from commit d559122374b9cad32601f597aaabd91f39805ead​)

#skipci

diff --git a/Engine/Plugins/Enterprise/GLTFExporter/MigratedFromGit/Source/GLTFMaterialAnalyzer/GLTFMaterialAnalyzer.Build.cs b/Engine/Plugins/Enterprise/GLTFExporter/MigratedFromGit/Source/GLTFMaterialAnalyzer/GLTFMaterialAnalyzer.Build.cs
--- a/Engine/Plugins/Enterprise/GLTFExporter/MigratedFromGit/Source/GLTFMaterialAnalyzer/GLTFMaterialAnalyzer.Build.cs
+++ b/Engine/Plugins/Enterprise/GLTFExporter/MigratedFromGit/Source/GLTFMaterialAnalyzer/GLTFMaterialAnalyzer.Build.cs
@@ -12,7 +12,8 @@ public class GLTFMaterialAnalyzer : ModuleRules
 				"Core",
 				"CoreUObject",
 				"Engine",
-				"RenderCore"
+				"RenderCore",
+				"RHI"
 			}
 		);
 
diff --git a/Engine/Plugins/Enterprise/GLTFExporter/MigratedFromGit/Source/GLTFMaterialAnalyzer/Private/GLTFMaterialAnalyzer.cpp b/Engine/Plugins/Enterprise/GLTFExporter/MigratedFromGit/Source/GLTFMaterialAnalyzer/Private/GLTFMaterialAnalyzer.cpp
--- a/Engine/Plugins/Enterprise/GLTFExporter/MigratedFromGit/Source/GLTFMaterialAnalyzer/Private/GLTFMaterialAnalyzer.cpp
+++ b/Engine/Plugins/Enterprise/GLTFExporter/MigratedFromGit/Source/GLTFMaterialAnalyzer/Private/GLTFMaterialAnalyzer.cpp
@@ -19,6 +19,12 @@ void UGLTFMaterialAnalyzer::ResetToDefaults()
 
 void UGLTFMaterialAnalyzer::AnalyzeMaterialPropertyEx(const UMaterialInterface* InMaterial, const EMaterialProperty& InProperty, const FString& InCustomOutput, FGLTFMaterialAnalysis& OutAnalysis)
 {
+	if (const_cast<UMaterialInterface*>(InMaterial)->GetMaterialResource(GMaxRHIFeatureLevel) == nullptr)
+	{
+		// NOTE: make sure material has a resource before calling UMaterialInterface::AnalyzeMaterialProperty
+		const_cast<UMaterialInterface*>(InMaterial)->ForceRecompileForRendering();
+	}
+
 	// NOTE: besides avoiding creating a new object each time a property is analyzed,
 	// the use of GetDefaultObject() also server another purpose:
 	//
@@ -28,7 +34,6 @@ void UGLTFMaterialAnalyzer::AnalyzeMaterialPropertyEx(const UMaterialInterface*
 	// defined in UMaterialInterface, and the editor would crash if it was processed.
 	// Fortunately objects with the RF_ClassDefaultObject flag
 	// are excluded from being processed.
-
 	UGLTFMaterialAnalyzer* Analyzer = Cast<UGLTFMaterialAnalyzer>(StaticClass()->GetDefaultObject());
 
 	Analyzer->Property = InProperty;
