[GLTFExporter] chore: bump version to 1.0.4

commit 2fefff08c4b2e5444e2ca14b4b5a23f4ffc7c5c9
Date: Fri, 26 Mar 2021 11:00:37 +0100

    chore: bump version to 1.0.4
    (cherry picked from commit 5627a4511512771d8d49eee9afdcf4add591dff3​)

#skipci

diff --git a/Engine/Plugins/Enterprise/GLTFExporter/MigratedFromGit/GLTFExporter.uplugin b/Engine/Plugins/Enterprise/GLTFExporter/MigratedFromGit/GLTFExporter.uplugin
--- a/Engine/Plugins/Enterprise/GLTFExporter/MigratedFromGit/GLTFExporter.uplugin
+++ b/Engine/Plugins/Enterprise/GLTFExporter/MigratedFromGit/GLTFExporter.uplugin
@@ -1,7 +1,7 @@
 {
 	"FileVersion": 3,
-	"Version": 103,
-	"VersionName": "1.0.3",
+	"Version": 104,
+	"VersionName": "1.0.4",
 	"FriendlyName": "glTF Exporter",
 	"Description": "An exporter for Khronos glTF 2.0.",
 	"Category": "Exporters",
