[GLTFExporter] refactor: move GLTFJsonAsset to own files

commit 9a89024903d99abe58b29f8e993c2355f303e7a8
Date: Mon, 5 Apr 2021 02:17:59 +0200

    refactor: move GLTFJsonAsset to own files
    (cherry picked from commit 0e5c2ca418c32b45ebdbcc50a4f5b5fc16e67eb4​)

#skipci

diff --git a/Engine/Plugins/Enterprise/GLTFExporter/MigratedFromGit/Source/GLTFExporter/Private/Json/GLTFJsonAsset.cpp b/Engine/Plugins/Enterprise/GLTFExporter/MigratedFromGit/Source/GLTFExporter/Private/Json/GLTFJsonAsset.cpp
new file mode 100644
--- /dev/null
+++ b/Engine/Plugins/Enterprise/GLTFExporter/MigratedFromGit/Source/GLTFExporter/Private/Json/GLTFJsonAsset.cpp
@@ -0,0 +1,17 @@
+// Copyright Epic Games, Inc. All Rights Reserved.
+
+#include "Json/GLTFJsonAsset.h"
+#include "GLTFExporterModule.h"
+#include "Interfaces/IPluginManager.h"
+#include "Runtime/Launch/Resources/Version.h"
+
+FGLTFJsonAsset::FGLTFJsonAsset()
+	: Version(TEXT("2.0"))
+{
+	const TSharedPtr<IPlugin> Plugin = IPluginManager::Get().FindPlugin(GLTFEXPORTER_MODULE_NAME);
+	const FPluginDescriptor& PluginDescriptor = Plugin->GetDescriptor();
+
+	// TODO: add configuration option to include detailed engine and plugin version info?
+
+	Generator = TEXT(EPIC_PRODUCT_NAME) TEXT(" ") VERSION_STRINGIFY(ENGINE_MAJOR_VERSION) TEXT(" ") + PluginDescriptor.FriendlyName;
+}
diff --git a/Engine/Plugins/Enterprise/GLTFExporter/MigratedFromGit/Source/GLTFExporter/Private/Json/GLTFJsonAsset.h b/Engine/Plugins/Enterprise/GLTFExporter/MigratedFromGit/Source/GLTFExporter/Private/Json/GLTFJsonAsset.h
new file mode 100644
--- /dev/null
+++ b/Engine/Plugins/Enterprise/GLTFExporter/MigratedFromGit/Source/GLTFExporter/Private/Json/GLTFJsonAsset.h
@@ -0,0 +1,35 @@
+// Copyright Epic Games, Inc. All Rights Reserved.
+
+#pragma once
+
+#include "Json/GLTFJsonExtensions.h"
+#include "Serialization/JsonSerializer.h"
+
+struct FGLTFJsonAsset
+{
+	FString Version;
+	FString Generator;
+	FString Copyright;
+
+	FGLTFJsonAsset();
+
+	template <class CharType = TCHAR, class PrintPolicy = TPrettyJsonPrintPolicy<CharType>>
+	void WriteObject(TJsonWriter<CharType, PrintPolicy>& JsonWriter, FGLTFJsonExtensions& Extensions) const
+	{
+		JsonWriter.WriteObjectStart();
+
+		JsonWriter.WriteValue(TEXT("version"), Version);
+
+		if (!Generator.IsEmpty())
+		{
+			JsonWriter.WriteValue(TEXT("generator"), Generator);
+		}
+
+		if (!Copyright.IsEmpty())
+		{
+			JsonWriter.WriteValue(TEXT("copyright"), Copyright);
+		}
+
+		JsonWriter.WriteObjectEnd();
+	}
+};
diff --git a/Engine/Plugins/Enterprise/GLTFExporter/MigratedFromGit/Source/GLTFExporter/Private/Json/GLTFJsonRoot.cpp b/Engine/Plugins/Enterprise/GLTFExporter/MigratedFromGit/Source/GLTFExporter/Private/Json/GLTFJsonRoot.cpp
deleted file mode 100644
diff --git a/Engine/Plugins/Enterprise/GLTFExporter/MigratedFromGit/Source/GLTFExporter/Private/Json/GLTFJsonRoot.h b/Engine/Plugins/Enterprise/GLTFExporter/MigratedFromGit/Source/GLTFExporter/Private/Json/GLTFJsonRoot.h
--- a/Engine/Plugins/Enterprise/GLTFExporter/MigratedFromGit/Source/GLTFExporter/Private/Json/GLTFJsonRoot.h
+++ b/Engine/Plugins/Enterprise/GLTFExporter/MigratedFromGit/Source/GLTFExporter/Private/Json/GLTFJsonRoot.h
@@ -3,6 +3,7 @@
 #pragma once
 
 #include "Json/GLTFJsonIndex.h"
+#include "Json/GLTFJsonAsset.h"
 #include "Json/GLTFJsonExtensions.h"
 #include "Json/GLTFJsonAccessor.h"
 #include "Json/GLTFJsonAnimation.h"
@@ -25,36 +26,6 @@
 #include "Json/GLTFJsonLevelVariantSets.h"
 #include "Policies/CondensedJsonPrintPolicy.h"
 
-
-struct FGLTFJsonAsset
-{
-	FString Version;
-	FString Generator;
-	FString Copyright;
-
-	FGLTFJsonAsset();
-
-	template <class CharType = TCHAR, class PrintPolicy = TPrettyJsonPrintPolicy<CharType>>
-	void WriteObject(TJsonWriter<CharType, PrintPolicy>& JsonWriter, FGLTFJsonExtensions& Extensions) const
-	{
-		JsonWriter.WriteObjectStart();
-
-		JsonWriter.WriteValue(TEXT("version"), Version);
-
-		if (!Generator.IsEmpty())
-		{
-			JsonWriter.WriteValue(TEXT("generator"), Generator);
-		}
-
-		if (!Copyright.IsEmpty())
-		{
-			JsonWriter.WriteValue(TEXT("copyright"), Copyright);
-		}
-
-		JsonWriter.WriteObjectEnd();
-	}
-};
-
 struct FGLTFJsonRoot
 {
 	FGLTFJsonAsset Asset;
