[GLTFExporter] feat: add export option "Include Generator Version" to add more version info as metadata

commit 0e55528e03e384755064e29c328de9c8aee27357
Date: Mon, 5 Apr 2021 04:12:01 +0200

    feat: add export option "Include Generator Version" to add more version info as metadata
    (cherry picked from commit 6ff58b4bf9177d87cadad494004b9f91bdd43e51​)

#skipci

diff --git a/Engine/Plugins/Enterprise/GLTFExporter/MigratedFromGit/Source/GLTFExporter/Private/Builders/GLTFJsonBuilder.cpp b/Engine/Plugins/Enterprise/GLTFExporter/MigratedFromGit/Source/GLTFExporter/Private/Builders/GLTFJsonBuilder.cpp
--- a/Engine/Plugins/Enterprise/GLTFExporter/MigratedFromGit/Source/GLTFExporter/Private/Builders/GLTFJsonBuilder.cpp
+++ b/Engine/Plugins/Enterprise/GLTFExporter/MigratedFromGit/Source/GLTFExporter/Private/Builders/GLTFJsonBuilder.cpp
@@ -1,11 +1,15 @@
 // Copyright Epic Games, Inc. All Rights Reserved.
 
 #include "Builders/GLTFJsonBuilder.h"
+#include "GLTFExporterModule.h"
+#include "Interfaces/IPluginManager.h"
+#include "Runtime/Launch/Resources/Version.h"
 
 FGLTFJsonBuilder::FGLTFJsonBuilder(const FString& FilePath, const UGLTFExportOptions* ExportOptions)
 	: FGLTFTaskBuilder(FilePath, ExportOptions)
 	, DefaultScene(JsonRoot.DefaultScene)
 {
+	JsonRoot.Asset.Generator = GetGeneratorString();
 }
 
 void FGLTFJsonBuilder::WriteJson(FArchive& Archive)
@@ -264,3 +268,13 @@ FGLTFJsonNodeIndex FGLTFJsonBuilder::GetComponentNodeIndex(FGLTFJsonNodeIndex No
 	const FGLTFJsonNode& Node = GetNode(NodeIndex);
 	return Node.ComponentNode != INDEX_NONE ? Node.ComponentNode : NodeIndex;
 }
+
+FString FGLTFJsonBuilder::GetGeneratorString() const
+{
+	const TSharedPtr<IPlugin> Plugin = IPluginManager::Get().FindPlugin(GLTFEXPORTER_MODULE_NAME);
+	const FPluginDescriptor& PluginDescriptor = Plugin->GetDescriptor();
+
+	return ExportOptions->bIncludeGeneratorVersion
+		? TEXT(EPIC_PRODUCT_NAME) TEXT(" ") ENGINE_VERSION_STRING TEXT(" ") + PluginDescriptor.FriendlyName + TEXT(" ") + PluginDescriptor.VersionName
+		: TEXT(EPIC_PRODUCT_NAME) TEXT(" ") + PluginDescriptor.FriendlyName;
+}
diff --git a/Engine/Plugins/Enterprise/GLTFExporter/MigratedFromGit/Source/GLTFExporter/Private/Builders/GLTFJsonBuilder.h b/Engine/Plugins/Enterprise/GLTFExporter/MigratedFromGit/Source/GLTFExporter/Private/Builders/GLTFJsonBuilder.h
--- a/Engine/Plugins/Enterprise/GLTFExporter/MigratedFromGit/Source/GLTFExporter/Private/Builders/GLTFJsonBuilder.h
+++ b/Engine/Plugins/Enterprise/GLTFExporter/MigratedFromGit/Source/GLTFExporter/Private/Builders/GLTFJsonBuilder.h
@@ -68,5 +68,7 @@ public:
 
 private:
 
+	FString GetGeneratorString() const;
+
 	FGLTFJsonRoot JsonRoot;
 };
diff --git a/Engine/Plugins/Enterprise/GLTFExporter/MigratedFromGit/Source/GLTFExporter/Private/GLTFExportOptions.cpp b/Engine/Plugins/Enterprise/GLTFExporter/MigratedFromGit/Source/GLTFExporter/Private/GLTFExportOptions.cpp
--- a/Engine/Plugins/Enterprise/GLTFExporter/MigratedFromGit/Source/GLTFExporter/Private/GLTFExportOptions.cpp
+++ b/Engine/Plugins/Enterprise/GLTFExporter/MigratedFromGit/Source/GLTFExporter/Private/GLTFExportOptions.cpp
@@ -15,6 +15,7 @@ void UGLTFExportOptions::ResetToDefault()
 	ExportUniformScale = 0.01;
 	bExportPreviewMesh = true;
 	bStrictCompliance = true;
+	bIncludeGeneratorVersion = true;
 	bExportUnlitMaterials = true;
 	bExportClearCoatMaterials = true;
 	bExportExtraBlendModes = false;
diff --git a/Engine/Plugins/Enterprise/GLTFExporter/MigratedFromGit/Source/GLTFExporter/Private/Json/GLTFJsonAsset.cpp b/Engine/Plugins/Enterprise/GLTFExporter/MigratedFromGit/Source/GLTFExporter/Private/Json/GLTFJsonAsset.cpp
deleted file mode 100644
diff --git a/Engine/Plugins/Enterprise/GLTFExporter/MigratedFromGit/Source/GLTFExporter/Private/Json/GLTFJsonAsset.h b/Engine/Plugins/Enterprise/GLTFExporter/MigratedFromGit/Source/GLTFExporter/Private/Json/GLTFJsonAsset.h
--- a/Engine/Plugins/Enterprise/GLTFExporter/MigratedFromGit/Source/GLTFExporter/Private/Json/GLTFJsonAsset.h
+++ b/Engine/Plugins/Enterprise/GLTFExporter/MigratedFromGit/Source/GLTFExporter/Private/Json/GLTFJsonAsset.h
@@ -11,7 +11,10 @@ struct FGLTFJsonAsset
 	FString Generator;
 	FString Copyright;
 
-	FGLTFJsonAsset();
+	FGLTFJsonAsset()
+		: Version(TEXT("2.0"))
+	{
+	}
 
 	template <class CharType = TCHAR, class PrintPolicy = TPrettyJsonPrintPolicy<CharType>>
 	void WriteObject(TJsonWriter<CharType, PrintPolicy>& JsonWriter, FGLTFJsonExtensions& Extensions) const
diff --git a/Engine/Plugins/Enterprise/GLTFExporter/MigratedFromGit/Source/GLTFExporter/Public/GLTFExportOptions.h b/Engine/Plugins/Enterprise/GLTFExporter/MigratedFromGit/Source/GLTFExporter/Public/GLTFExportOptions.h
--- a/Engine/Plugins/Enterprise/GLTFExporter/MigratedFromGit/Source/GLTFExporter/Public/GLTFExportOptions.h
+++ b/Engine/Plugins/Enterprise/GLTFExporter/MigratedFromGit/Source/GLTFExporter/Public/GLTFExportOptions.h
@@ -80,6 +80,10 @@ class GLTFEXPORTER_API UGLTFExportOptions : public UObject
 	UPROPERTY(EditAnywhere, BlueprintReadWrite, Config, Category = General)
 	bool bStrictCompliance;
 
+	/** If enabled, version info for Unreal Engine and exporter plugin will be included as metadata in the glTF asset, which is useful when reporting issues. */
+	UPROPERTY(EditAnywhere, BlueprintReadWrite, Config, Category = General)
+	bool bIncludeGeneratorVersion;
+
 	/** If enabled, materials with shading model unlit will be properly exported. Uses extension KHR_materials_unlit. */
 	UPROPERTY(EditAnywhere, BlueprintReadWrite, Config, Category = Material)
 	bool bExportUnlitMaterials;
