[GLTFExporter] docs: improve description for export options bUseMeshQuantization and TextureImageQuality

commit 9edaca7124ee22c3e08b46c7edffc17d76ff07e6
Date: Sun, 18 Apr 2021 22:38:26 +0200

    docs: improve description for export options bUseMeshQuantization and TextureImageQuality
    (cherry picked from commit 8df2d09009a5babfdc1a02dfc3bae72ea4751292​)

#skipci

diff --git a/Engine/Plugins/Enterprise/GLTFExporter/MigratedFromGit/Source/GLTFExporter/Public/GLTFExportOptions.h b/Engine/Plugins/Enterprise/GLTFExporter/MigratedFromGit/Source/GLTFExporter/Public/GLTFExportOptions.h
--- a/Engine/Plugins/Enterprise/GLTFExporter/MigratedFromGit/Source/GLTFExporter/Public/GLTFExportOptions.h
+++ b/Engine/Plugins/Enterprise/GLTFExporter/MigratedFromGit/Source/GLTFExporter/Public/GLTFExportOptions.h
@@ -128,7 +128,7 @@ class GLTFEXPORTER_API UGLTFExportOptions : public UObject
 	UPROPERTY(EditAnywhere, BlueprintReadWrite, Config, Category = Mesh)
 	bool bExportVertexSkinWeights;
 
-	/** If enabled, use Unreal-configured quantization for vertex tangents and normals, reducing size. Requires extension KHR_mesh_quantization, the mesh might not load in some glTF viewers. */
+	/** If enabled, use quantization for vertex tangents and normals, reducing size. Requires extension KHR_mesh_quantization, which may result in the mesh not loading in some glTF viewers. */
 	UPROPERTY(EditAnywhere, BlueprintReadWrite, Config, Category = Mesh)
 	bool bUseMeshQuantization;
 
@@ -152,7 +152,7 @@ class GLTFEXPORTER_API UGLTFExportOptions : public UObject
 	UPROPERTY(EditAnywhere, BlueprintReadWrite, Config, Category = Texture)
 	EGLTFTextureImageFormat TextureImageFormat;
 
-	/** Level of compression used for exported textures, between 1 (worst quality, best compression) and 100 (best quality, worst compression). Does not apply to lossless formats (e.g. PNG). */
+	/** Level of compression used for textures exported with lossy image formats, 0 (default) or value between 1 (worst quality, best compression) and 100 (best quality, worst compression). */
 	UPROPERTY(EditAnywhere, BlueprintReadWrite, Config, Category = Texture, Meta = (ClampMin = "0", ClampMax = "100", EditCondition = "TextureImageFormat == EGLTFTextureImageFormat::JPEG"))
 	int32 TextureImageQuality;
 
