[GLTFExporter] docs: update external export options reference document

commit 7226251d0e12598fec01a5cfee6e4a7043779d0f
Date: Sun, 18 Apr 2021 22:38:58 +0200

    docs: update external export options reference document
    (cherry picked from commit ad9ee613a77c3a4533580498f292da7186c938cc​)

#skipci

diff --git a/Engine/Plugins/Enterprise/GLTFExporter/MigratedFromGit/Docs/export-options-reference.md b/Engine/Plugins/Enterprise/GLTFExporter/MigratedFromGit/Docs/export-options-reference.md
--- a/Engine/Plugins/Enterprise/GLTFExporter/MigratedFromGit/Docs/export-options-reference.md
+++ b/Engine/Plugins/Enterprise/GLTFExporter/MigratedFromGit/Docs/export-options-reference.md
@@ -5,6 +5,7 @@ Option                         | Description
 `Export Uniform Scale`         | Scale factor used for exporting all assets (0.01 by default) for conversion from centimeters (Unreal default) to meters (glTF).
 `Export Preview Mesh`          | If enabled, the preview mesh for a standalone animation or material asset will also be exported.
 `Strict Compliance`            | If enabled, certain values (like HDR colors and light angles) will be truncated during export to strictly conform to the formal glTF specification.
+`Include Generator Version`    | If enabled, version info for Unreal Engine and exporter plugin will be included as metadata in the glTF asset, which is useful when reporting issues.
 `Export Unlit Materials`       | If enabled, materials with shading model unlit will be properly exported. Uses extension KHR_materials_unlit.
 `Export Clear Coat Materials`  | If enabled, materials with shading model clear coat will be properly exported. Uses extension KHR_materials_clearcoat, which is not supported by all glTF viewers.
 `Export Extra Blend Modes`     | If enabled, materials with blend modes additive, modulate, and alpha composite will be properly exported. Uses extension EPIC_blend_modes, which is supported by Unreal's glTF viewer.
@@ -16,13 +17,13 @@ Option                         | Description
 `Default Level Of Detail`      | Default LOD level used for exporting a mesh. Can be overridden by component or asset settings (e.g. minimum or forced LOD level).
 `Export Vertex Colors`         | If enabled, export vertex color. Not recommended due to vertex colors always being used as a base color multiplier in glTF, regardless of material. Often producing undesirable results.
 `Export Vertex Skin Weights`   | If enabled, export vertex bone weights and indices in skeletal meshes. Necessary for animation sequences.
-`Use Mesh Quantization`        | If enabled, use Unreal-configured quantization for vertex tangents and normals, reducing size. Requires extension KHR_mesh_quantization, the mesh might not load in some glTF viewers.
+`Use Mesh Quantization`        | If enabled, use quantization for vertex tangents and normals, reducing size. Requires extension KHR_mesh_quantization, which may result in the mesh not loading in some glTF viewers.
 `Export Level Sequences`       | If enabled, export level sequences. Only transform tracks are currently supported. The level sequence will be played at the assigned display rate.
 `Export Animation Sequences`   | If enabled, export single animation asset used by a skeletal mesh component or hotspot actor. Export of vertex skin weights must be enabled.
 `Retarget Bone Transforms`     | If enabled, apply animation retargeting to skeleton bones when exporting an animation sequence.
 `Export Playback Settings`     | If enabled, export play rate, start time, looping, and auto play for an animation or level sequence. Uses extension EPIC_animation_playback, which is supported by Unreal's glTF viewer.
 `Texture Image Format`         | Desired image format used for exported textures.
-`Texture Image Quality`        | Level of compression used for exported textures, between 1 (worst quality, best compression) and 100 (best quality, worst compression). Does not apply to lossless formats (e.g. PNG).
+`Texture Image Quality`        | Level of compression used for textures exported with lossy image formats, 0 (default) or value between 1 (worst quality, best compression) and 100 (best quality, worst compression).
 `No Lossy Image Format For`    | Texture types that will always use lossless formats (e.g. PNG) because of sensitivity to compression artifacts.
 `Export Texture Transforms`    | If enabled, export UV tiling and un-mirroring settings in a texture coordinate expression node for simple material input expressions. Uses extension KHR_texture_transform.
 `Export Lightmaps`             | If enabled, export lightmaps (created by Lightmass) when exporting a level. Uses extension EPIC_lightmap_textures, which is supported by Unreal's glTF viewer.
