[GLTFExporter] chore: add docs folder in filter so the folder is packaged

commit 4764da8d5bc243eceac53f1396a227dacd52a01f
Date: Mon, 19 Apr 2021 19:32:14 +0100

    chore: add docs folder in filter so the folder is packaged
    (cherry picked from commit 26bd98630fb36d9c99d4702f16dd3f3ba641310b​)

#skipci

diff --git a/Engine/Plugins/Enterprise/GLTFExporter/MigratedFromGit/Config/FilterPlugin.ini b/Engine/Plugins/Enterprise/GLTFExporter/MigratedFromGit/Config/FilterPlugin.ini
--- a/Engine/Plugins/Enterprise/GLTFExporter/MigratedFromGit/Config/FilterPlugin.ini
+++ b/Engine/Plugins/Enterprise/GLTFExporter/MigratedFromGit/Config/FilterPlugin.ini
@@ -1,2 +1,3 @@
 [FilterPlugin]
+/Docs/
 /README.md
