[GLTFExporter] chore: add whitelist platforms as requested in marketplace prerequisites

commit 7d55467dd37089b6546856b62dfc3741c87d2be8
Date: Mon, 19 Apr 2021 19:33:02 +0100

    chore: add whitelist platforms as requested in marketplace prerequisites
    (cherry picked from commit 56fde99fd20417ea4711510e3bbeec3eeb98c0b0​)

#skipci

diff --git a/Engine/Plugins/Enterprise/GLTFExporter/MigratedFromGit/GLTFExporter.uplugin b/Engine/Plugins/Enterprise/GLTFExporter/MigratedFromGit/GLTFExporter.uplugin
--- a/Engine/Plugins/Enterprise/GLTFExporter/MigratedFromGit/GLTFExporter.uplugin
+++ b/Engine/Plugins/Enterprise/GLTFExporter/MigratedFromGit/GLTFExporter.uplugin
@@ -17,22 +17,26 @@
 		{
 			"Name": "GLTFExporter",
 			"Type": "Editor",
-			"LoadingPhase": "Default"
+			"LoadingPhase": "Default",
+			"WhitelistPlatforms": [ "Win64", "Win32" ]
 		},
 		{
 			"Name": "GLTFExporterRuntime",
 			"Type": "Runtime",
-			"LoadingPhase": "Default"
+			"LoadingPhase": "Default",
+			"WhitelistPlatforms": [ "Win64", "Win32" ]
 		},
 		{
 			"Name": "GLTFMaterialAnalyzer",
 			"Type": "Editor",
-			"LoadingPhase": "Default"
+			"LoadingPhase": "Default",
+			"WhitelistPlatforms": [ "Win64", "Win32" ]
 		},
 		{
 			"Name": "GLTFMaterialBaking",
 			"Type": "Editor",
-			"LoadingPhase": "Default"
+			"LoadingPhase": "Default",
+			"WhitelistPlatforms": [ "Win64", "Win32" ]
 		}
 	],
 	"Plugins": [
