[GLTFExporter] chore: bump version to 1.0.5

commit 66950aa5aac65b356e3318010e918b2d1c49371a
Date: Mon, 19 Apr 2021 22:09:18 +0200

    chore: bump version to 1.0.5
    (cherry picked from commit 860d0f1db62395597503eb2d79031a204e8e4b5f​)

#skipci

diff --git a/Engine/Plugins/Enterprise/GLTFExporter/MigratedFromGit/GLTFExporter.uplugin b/Engine/Plugins/Enterprise/GLTFExporter/MigratedFromGit/GLTFExporter.uplugin
--- a/Engine/Plugins/Enterprise/GLTFExporter/MigratedFromGit/GLTFExporter.uplugin
+++ b/Engine/Plugins/Enterprise/GLTFExporter/MigratedFromGit/GLTFExporter.uplugin
@@ -1,7 +1,7 @@
 {
 	"FileVersion": 3,
-	"Version": 104,
-	"VersionName": "1.0.4",
+	"Version": 105,
+	"VersionName": "1.0.5",
 	"FriendlyName": "glTF Exporter",
 	"Description": "An exporter for Khronos glTF 2.0.",
 	"Category": "Exporters",
