[GLTFExporter] refactor: reduce code in GLTFAccessorConverters by using std::conditional

commit c6cbc619f76868ce51154f5dc1c7163668e1f150
Date: Mon, 7 Jun 2021 20:12:19 +0200

    refactor: reduce code in GLTFAccessorConverters by using std::conditional
    (cherry picked from commit 523a4955c45887ec7d080b3fbc0d169b94c15572​)

#skipci

diff --git a/Engine/Plugins/Enterprise/GLTFExporter/MigratedFromGit/Source/GLTFExporter/Private/Converters/GLTFAccessorConverters.cpp b/Engine/Plugins/Enterprise/GLTFExporter/MigratedFromGit/Source/GLTFExporter/Private/Converters/GLTFAccessorConverters.cpp
--- a/Engine/Plugins/Enterprise/GLTFExporter/MigratedFromGit/Source/GLTFExporter/Private/Converters/GLTFAccessorConverters.cpp
+++ b/Engine/Plugins/Enterprise/GLTFExporter/MigratedFromGit/Source/GLTFExporter/Private/Converters/GLTFAccessorConverters.cpp
@@ -5,44 +5,20 @@
 #include "Converters/GLTFSkinWeightVertexBufferHack.h"
 #include "Builders/GLTFConvertBuilder.h"
 
-namespace
+// TODO: Unreal-style implementation of std::conditional to avoid mixing in STL. Should be added to the engine.
+template <bool Condition, class TypeIfTrue, class TypeIfFalse>
+class TConditional
 {
-	template <typename DestinationType, typename SourceType>
-	struct TGLTFVertexNormalUtility
-	{
-		static DestinationType Convert(const FVector& Normal)
-		{
-			return FGLTFConverterUtility::ConvertNormal(SourceType(Normal));
-		}
-	};
-
-	template <typename SourceType>
-	struct TGLTFVertexNormalUtility<FGLTFJsonVector3, SourceType>
-	{
-		static FGLTFJsonVector3 Convert(const FVector& Normal)
-		{
-			return FGLTFConverterUtility::ConvertNormal(Normal);
-		}
-	};
-
-	template <typename DestinationType, typename SourceType>
-	struct TGLTFVertexTangentUtility
-	{
-		static DestinationType Convert(const FVector& Tangent)
-		{
-			return FGLTFConverterUtility::ConvertTangent(SourceType(Tangent));
-		}
-	};
+public:
+	typedef TypeIfFalse Type;
+};
 
-	template <typename SourceType>
-	struct TGLTFVertexTangentUtility<FGLTFJsonVector4, SourceType>
-	{
-		static FGLTFJsonVector4 Convert(const FVector& Tangent)
-		{
-			return FGLTFConverterUtility::ConvertTangent(Tangent);
-		}
-	};
-}
+template <class TypeIfTrue, class TypeIfFalse>
+class TConditional<true, TypeIfTrue, TypeIfFalse>
+{
+public:
+	typedef TypeIfTrue Type;
+};
 
 FGLTFJsonAccessorIndex FGLTFPositionBufferConverter::Convert(const FGLTFMeshSection* MeshSection, const FPositionVertexBuffer* VertexBuffer)
 {
@@ -195,7 +171,9 @@ FGLTFJsonBufferViewIndex FGLTFNormalBufferConverter::ConvertBufferView(const FGL
 	{
 		const uint32 MappedVertexIndex = IndexMap[VertexIndex];
 		const FVector SafeNormal = VertexTangents[MappedVertexIndex].TangentZ.ToFVector().GetSafeNormal();
-		Normals[VertexIndex] = TGLTFVertexNormalUtility<DestinationType, SourceType>::Convert(SafeNormal);
+
+		typedef typename TConditional<TIsSame<DestinationType, FGLTFJsonVector3>::Value, FVector, SourceType>::Type IntermediateType;
+		Normals[VertexIndex] = FGLTFConverterUtility::ConvertNormal(IntermediateType(SafeNormal));
 	}
 
 	return Builder.AddBufferView(Normals, EGLTFJsonBufferTarget::ArrayBuffer);
@@ -270,7 +248,9 @@ FGLTFJsonBufferViewIndex FGLTFTangentBufferConverter::ConvertBufferView(const FG
 	{
 		const uint32 MappedVertexIndex = IndexMap[VertexIndex];
 		const FVector SafeTangent = VertexTangents[MappedVertexIndex].TangentX.ToFVector().GetSafeNormal();
-		Tangents[VertexIndex] = TGLTFVertexTangentUtility<DestinationType, SourceType>::Convert(SafeTangent);
+
+		typedef typename TConditional<TIsSame<DestinationType, FGLTFJsonVector4>::Value, FVector, SourceType>::Type IntermediateType;
+		Tangents[VertexIndex] = FGLTFConverterUtility::ConvertTangent(IntermediateType(SafeTangent));
 	}
 
 	return Builder.AddBufferView(Tangents, EGLTFJsonBufferTarget::ArrayBuffer);
