[GLTFExporter] refactor: wrap TJsonWriter template into new IGLTFJsonWriter interface and reduce need for FGLTFJsonUtility

commit 3fed58b24f34b5e09cf05d04a915198d77399177
Date: Sun, 6 Jun 2021 23:59:30 +0200

    refactor: wrap TJsonWriter template into new IGLTFJsonWriter interface and reduce need for FGLTFJsonUtility
    (cherry picked from commit f3720d755c92212415fa9144cf205bf7e58d9cb2​)

#skipci

diff --git a/Engine/Plugins/Enterprise/GLTFExporter/MigratedFromGit/Source/GLTFExporter/Private/Builders/GLTFContainerBuilder.cpp b/Engine/Plugins/Enterprise/GLTFExporter/MigratedFromGit/Source/GLTFExporter/Private/Builders/GLTFContainerBuilder.cpp
--- a/Engine/Plugins/Enterprise/GLTFExporter/MigratedFromGit/Source/GLTFExporter/Private/Builders/GLTFContainerBuilder.cpp
+++ b/Engine/Plugins/Enterprise/GLTFExporter/MigratedFromGit/Source/GLTFExporter/Private/Builders/GLTFContainerBuilder.cpp
@@ -21,7 +21,7 @@ void FGLTFContainerBuilder::Write(FArchive& Archive, FFeedbackContext* Context)
 		const FString ExtensionsString = FString::JoinBy(CustomExtensions, TEXT(", "),
 			[](EGLTFJsonExtension Extension)
 		{
-			return FGLTFJsonUtility::ToString(Extension);
+			return FGLTFJsonUtility::GetValue(Extension);
 		});
 
 		AddWarningMessage(FString::Printf(TEXT("Export uses some extensions that may only be supported in Unreal's glTF viewer: %s"), *ExtensionsString));
diff --git a/Engine/Plugins/Enterprise/GLTFExporter/MigratedFromGit/Source/GLTFExporter/Private/Builders/GLTFJsonBuilder.cpp b/Engine/Plugins/Enterprise/GLTFExporter/MigratedFromGit/Source/GLTFExporter/Private/Builders/GLTFJsonBuilder.cpp
--- a/Engine/Plugins/Enterprise/GLTFExporter/MigratedFromGit/Source/GLTFExporter/Private/Builders/GLTFJsonBuilder.cpp
+++ b/Engine/Plugins/Enterprise/GLTFExporter/MigratedFromGit/Source/GLTFExporter/Private/Builders/GLTFJsonBuilder.cpp
@@ -14,7 +14,7 @@ FGLTFJsonBuilder::FGLTFJsonBuilder(const FString& FilePath, const UGLTFExportOpt
 
 void FGLTFJsonBuilder::WriteJson(FArchive& Archive)
 {
-	JsonRoot.ToJson(&Archive, bIsGlbFile);
+	JsonRoot.ToJson(&Archive, !bIsGlbFile);
 }
 
 TSet<EGLTFJsonExtension> FGLTFJsonBuilder::GetCustomExtensionsUsed() const
@@ -25,7 +25,7 @@ TSet<EGLTFJsonExtension> FGLTFJsonBuilder::GetCustomExtensionsUsed() const
 
 	for (EGLTFJsonExtension Extension : JsonRoot.Extensions.Used)
 	{
-		const TCHAR* ExtensionString = FGLTFJsonUtility::ToString(Extension);
+		const TCHAR* ExtensionString = FGLTFJsonUtility::GetValue(Extension);
 		if (FCString::Strncmp(ExtensionString, CustomPrefix, (sizeof(CustomPrefix) / sizeof(TCHAR)) - 1) == 0)
 		{
 			CustomExtensions.Add(Extension);
diff --git a/Engine/Plugins/Enterprise/GLTFExporter/MigratedFromGit/Source/GLTFExporter/Private/Converters/GLTFAccessorConverters.cpp b/Engine/Plugins/Enterprise/GLTFExporter/MigratedFromGit/Source/GLTFExporter/Private/Converters/GLTFAccessorConverters.cpp
--- a/Engine/Plugins/Enterprise/GLTFExporter/MigratedFromGit/Source/GLTFExporter/Private/Converters/GLTFAccessorConverters.cpp
+++ b/Engine/Plugins/Enterprise/GLTFExporter/MigratedFromGit/Source/GLTFExporter/Private/Converters/GLTFAccessorConverters.cpp
@@ -30,7 +30,7 @@ FGLTFJsonAccessorIndex FGLTFPositionBufferConverter::Convert(const FGLTFMeshSect
 	const TArray<uint32>& IndexMap = MeshSection->IndexMap;
 	const uint32 VertexCount = IndexMap.Num();
 
-	TArray<FGLTFJsonVector3> Positions;
+	TArray<FGLTFRawVector3> Positions;
 	Positions.AddUninitialized(VertexCount);
 
 	for (uint32 VertexIndex = 0; VertexIndex < VertexCount; ++VertexIndex)
@@ -40,12 +40,12 @@ FGLTFJsonAccessorIndex FGLTFPositionBufferConverter::Convert(const FGLTFMeshSect
 	}
 
 	// More accurate bounding box if based on raw vertex values
-	FGLTFJsonVector3 MinPosition = VertexCount > 0 ? Positions[0] : FGLTFJsonVector3::Zero;
-	FGLTFJsonVector3 MaxPosition = VertexCount > 0 ? Positions[0] : FGLTFJsonVector3::Zero;
+	FGLTFRawVector3 MinPosition = VertexCount > 0 ? Positions[0] : FGLTFRawVector3(0, 0, 0); // TODO: make static const
+	FGLTFRawVector3 MaxPosition = VertexCount > 0 ? Positions[0] : FGLTFRawVector3(0, 0, 0); // TODO: make static const
 
 	for (uint32 VertexIndex = 1; VertexIndex < VertexCount; ++VertexIndex)
 	{
-		const FGLTFJsonVector3& Position = Positions[VertexIndex];
+		const FGLTFRawVector3& Position = Positions[VertexIndex];
 		MinPosition.X = FMath::Min(MinPosition.X, Position.X);
 		MinPosition.Y = FMath::Min(MinPosition.Y, Position.Y);
 		MinPosition.Z = FMath::Min(MinPosition.Z, Position.Z);
@@ -134,8 +134,8 @@ FGLTFJsonAccessorIndex FGLTFNormalBufferConverter::Convert(const FGLTFMeshSectio
 	{
 		ComponentType = EGLTFJsonComponentType::F32;
 		BufferViewIndex = bHighPrecision
-			? ConvertBufferView<FGLTFJsonVector3, FPackedRGBA16N>(MeshSection, VertexBuffer)
-			: ConvertBufferView<FGLTFJsonVector3, FPackedNormal>(MeshSection, VertexBuffer);
+			? ConvertBufferView<FGLTFRawVector3, FPackedRGBA16N>(MeshSection, VertexBuffer)
+			: ConvertBufferView<FGLTFRawVector3, FPackedNormal>(MeshSection, VertexBuffer);
 	}
 
 	FGLTFJsonAccessor JsonAccessor;
@@ -172,7 +172,7 @@ FGLTFJsonBufferViewIndex FGLTFNormalBufferConverter::ConvertBufferView(const FGL
 		const uint32 MappedVertexIndex = IndexMap[VertexIndex];
 		const FVector SafeNormal = VertexTangents[MappedVertexIndex].TangentZ.ToFVector().GetSafeNormal();
 
-		typedef typename TConditional<TIsSame<DestinationType, FGLTFJsonVector3>::Value, FVector, SourceType>::Type IntermediateType;
+		typedef typename TConditional<TIsSame<DestinationType, FGLTFRawVector3>::Value, FVector, SourceType>::Type IntermediateType;
 		Normals[VertexIndex] = FGLTFConverterUtility::ConvertNormal(IntermediateType(SafeNormal));
 	}
 
@@ -211,8 +211,8 @@ FGLTFJsonAccessorIndex FGLTFTangentBufferConverter::Convert(const FGLTFMeshSecti
 	{
 		ComponentType = EGLTFJsonComponentType::F32;
 		BufferViewIndex = bHighPrecision
-			? ConvertBufferView<FGLTFJsonVector4, FPackedRGBA16N>(MeshSection, VertexBuffer)
-			: ConvertBufferView<FGLTFJsonVector4, FPackedNormal>(MeshSection, VertexBuffer);
+			? ConvertBufferView<FGLTFRawVector4, FPackedRGBA16N>(MeshSection, VertexBuffer)
+			: ConvertBufferView<FGLTFRawVector4, FPackedNormal>(MeshSection, VertexBuffer);
 	}
 
 	FGLTFJsonAccessor JsonAccessor;
@@ -249,7 +249,7 @@ FGLTFJsonBufferViewIndex FGLTFTangentBufferConverter::ConvertBufferView(const FG
 		const uint32 MappedVertexIndex = IndexMap[VertexIndex];
 		const FVector SafeTangent = VertexTangents[MappedVertexIndex].TangentX.ToFVector().GetSafeNormal();
 
-		typedef typename TConditional<TIsSame<DestinationType, FGLTFJsonVector4>::Value, FVector, SourceType>::Type IntermediateType;
+		typedef typename TConditional<TIsSame<DestinationType, FGLTFRawVector4>::Value, FVector, SourceType>::Type IntermediateType;
 		Tangents[VertexIndex] = FGLTFConverterUtility::ConvertTangent(IntermediateType(SafeTangent));
 	}
 
@@ -275,7 +275,7 @@ FGLTFJsonAccessorIndex FGLTFUVBufferConverter::Convert(const FGLTFMeshSection* M
 
 	// TODO: report warning or add support for half float precision UVs, i.e. !VertexBuffer->GetUseFullPrecisionUVs()?
 
-	TArray<FGLTFJsonVector2> UVs;
+	TArray<FGLTFRawVector2> UVs;
 	UVs.AddUninitialized(VertexCount);
 
 	for (uint32 VertexIndex = 0; VertexIndex < VertexCount; ++VertexIndex)
diff --git a/Engine/Plugins/Enterprise/GLTFExporter/MigratedFromGit/Source/GLTFExporter/Private/Converters/GLTFConverterUtility.h b/Engine/Plugins/Enterprise/GLTFExporter/MigratedFromGit/Source/GLTFExporter/Private/Converters/GLTFConverterUtility.h
--- a/Engine/Plugins/Enterprise/GLTFExporter/MigratedFromGit/Source/GLTFExporter/Private/Converters/GLTFConverterUtility.h
+++ b/Engine/Plugins/Enterprise/GLTFExporter/MigratedFromGit/Source/GLTFExporter/Private/Converters/GLTFConverterUtility.h
@@ -2,6 +2,7 @@
 
 #pragma once
 
+#include "Converters/GLTFRawTypes.h"
 #include "Converters/GLTFPackedTypes.h"
 #include "Json/GLTFJsonEnums.h"
 #include "Json/GLTFJsonVector2.h"
@@ -21,31 +22,31 @@ struct FGLTFConverterUtility
 		return Length * ConversionScale;
 	}
 
-	static FGLTFJsonVector3 ConvertVector(const FVector& Vector)
+	static FGLTFRawVector3 ConvertVector(const FVector& Vector)
 	{
 		// UE4 uses a left-handed coordinate system, with Z up.
 		// glTF uses a right-handed coordinate system, with Y up.
 		return { Vector.X, Vector.Z, Vector.Y };
 	}
 
-	static FGLTFJsonVector3 ConvertPosition(const FVector& Position, const float ConversionScale)
+	static FGLTFRawVector3 ConvertPosition(const FVector& Position, const float ConversionScale)
 	{
 		return ConvertVector(Position * ConversionScale);
 	}
 
-	static FGLTFJsonVector3 ConvertScale(const FVector& Scale)
+	static FGLTFRawVector3 ConvertScale(const FVector& Scale)
 	{
 		return ConvertVector(Scale);
 	}
 
-	static FGLTFJsonVector3 ConvertNormal(const FVector& Normal)
+	static FGLTFRawVector3 ConvertNormal(const FVector& Normal)
 	{
 		return ConvertVector(Normal);
 	}
 
 	static FGLTFPackedVector16 ConvertNormal(const FPackedRGBA16N& Normal)
 	{
-		return { Normal.X, Normal.Z, Normal.Y };
+		return { Normal.X, Normal.Z, Normal.Y, 0 };
 	}
 
 	static FGLTFPackedVector8 ConvertNormal(const FPackedNormal& Normal)
@@ -53,7 +54,7 @@ struct FGLTFConverterUtility
 		return { Normal.Vector.X, Normal.Vector.Z, Normal.Vector.Y };
 	}
 
-	static FGLTFJsonVector4 ConvertTangent(const FVector& Tangent)
+	static FGLTFRawVector4 ConvertTangent(const FVector& Tangent)
 	{
 		// glTF stores tangent as Vec4, with W component indicating handedness of tangent basis.
 		return { ConvertVector(Tangent), 1.0f };
@@ -69,7 +70,7 @@ struct FGLTFConverterUtility
 		return { Tangent.Vector.X, Tangent.Vector.Z, Tangent.Vector.Y, MAX_int8 /* = 1.0 */ };
 	}
 
-	static FGLTFJsonVector2 ConvertUV(const FVector2D& UV)
+	static FGLTFRawVector2 ConvertUV(const FVector2D& UV)
 	{
 		// No conversion actually needed, this is primarily for type-safety.
 		return { UV.X, UV.Y };
@@ -103,7 +104,7 @@ struct FGLTFConverterUtility
 		return { Color.R, Color.G, Color.B, Color.A };
 	}
 
-	static FGLTFJsonQuaternion ConvertRotation(const FQuat& Rotation)
+	static FGLTFRawQuaternion ConvertRotation(const FQuat& Rotation)
 	{
 		// UE4 uses a left-handed coordinate system, with Z up.
 		// glTF uses a right-handed coordinate system, with Y up.
@@ -117,21 +118,21 @@ struct FGLTFConverterUtility
 		// e.g. some sources use non-unit Quats for rotation tangents
 
 		// Return the identity quaternion when possible (depending on tolerance)
-		if (Rotation.Equals(FQuat::Identity))
+		if (Rotation.Equals(FQuat::Identity)) // TODO: doesn't this potentially remove intentional deviance from identity?
 		{
-			return FGLTFJsonQuaternion::Identity;
+			return FGLTFRawQuaternion(0, 0, 0, 1); // TODO: make static const
 		}
 
 		const FQuat Normalized = Rotation.GetNormalized();
 		return { -Normalized.X, -Normalized.Z, -Normalized.Y, Normalized.W };
 	}
 
-	static FGLTFJsonMatrix4 ConvertMatrix(const FMatrix& Matrix)
+	static FGLTFRawMatrix4 ConvertMatrix(const FMatrix& Matrix)
 	{
 		// Unreal stores matrix elements in row major order.
 		// glTF stores matrix elements in column major order.
 
-		FGLTFJsonMatrix4 Result;
+		FGLTFRawMatrix4 Result;
 		for (int32 Row = 0; Row < 4; ++Row)
 		{
 			for (int32 Col = 0; Col < 4; ++Col)
@@ -142,7 +143,7 @@ struct FGLTFConverterUtility
 		return Result;
 	}
 
-	static FGLTFJsonMatrix4 ConvertTransform(const FTransform& Transform, const float ConversionScale)
+	static FGLTFRawMatrix4 ConvertTransform(const FTransform& Transform, const float ConversionScale)
 	{
 		const FQuat Rotation = Transform.GetRotation();
 		const FVector Translation = Transform.GetTranslation();
diff --git a/Engine/Plugins/Enterprise/GLTFExporter/MigratedFromGit/Source/GLTFExporter/Private/Converters/GLTFRawTypes.h b/Engine/Plugins/Enterprise/GLTFExporter/MigratedFromGit/Source/GLTFExporter/Private/Converters/GLTFRawTypes.h
new file mode 100644
--- /dev/null
+++ b/Engine/Plugins/Enterprise/GLTFExporter/MigratedFromGit/Source/GLTFExporter/Private/Converters/GLTFRawTypes.h
@@ -0,0 +1,85 @@
+// Copyright Epic Games, Inc. All Rights Reserved.
+
+#pragma once
+
+#include "CoreTypes.h"
+
+// Primarily for type-safety
+struct FGLTFRawVector2
+{
+	float X, Y;
+
+	FGLTFRawVector2(float X, float Y)
+		: X(X), Y(Y)
+	{
+	}
+};
+
+// Primarily for type-safety
+struct FGLTFRawVector3
+{
+	float X, Y, Z;
+
+	FGLTFRawVector3(float X, float Y, float Z)
+		: X(X), Y(Y), Z(Z)
+	{
+	}
+
+	FGLTFRawVector3(const FGLTFRawVector2& Vector2, float Z)
+		: X(Vector2.X), Y(Vector2.Y), Z(Z)
+	{
+	}
+
+	operator FGLTFRawVector3() const
+	{
+		return { X, Y, Z };
+	}
+};
+
+// Primarily for type-safety
+struct FGLTFRawVector4
+{
+	float X, Y, Z, W;
+
+	FGLTFRawVector4(float X, float Y, float Z, float W)
+		: X(X), Y(Y), Z(Z), W(W)
+	{
+	}
+
+	FGLTFRawVector4(const FGLTFRawVector3& Vector3, float W)
+		: X(Vector3.X), Y(Vector3.Y), Z(Vector3.Z), W(W)
+	{
+	}
+
+	operator FGLTFRawVector3() const
+	{
+		return { X, Y, Z };
+	}
+};
+
+// Primarily for type-safety
+struct FGLTFRawQuaternion
+{
+	float X, Y, Z, W;
+
+	FGLTFRawQuaternion(float X, float Y, float Z, float W)
+		: X(X), Y(Y), Z(Z), W(W)
+	{
+	}
+};
+
+// Primarily for type-safety
+struct FGLTFRawMatrix4
+{
+	float Cells[16];
+
+	float& operator()(int32 Row, int32 Col)
+	{
+		return Cells[Row * 4 + Col];
+	}
+
+	const float& operator()(int32 Row, int32 Col) const
+	{
+		return Cells[Row * 4 + Col];
+	}
+};
diff --git a/Engine/Plugins/Enterprise/GLTFExporter/MigratedFromGit/Source/GLTFExporter/Private/Converters/GLTFSkinConverters.cpp b/Engine/Plugins/Enterprise/GLTFExporter/MigratedFromGit/Source/GLTFExporter/Private/Converters/GLTFSkinConverters.cpp
--- a/Engine/Plugins/Enterprise/GLTFExporter/MigratedFromGit/Source/GLTFExporter/Private/Converters/GLTFSkinConverters.cpp
+++ b/Engine/Plugins/Enterprise/GLTFExporter/MigratedFromGit/Source/GLTFExporter/Private/Converters/GLTFSkinConverters.cpp
@@ -25,7 +25,7 @@ FGLTFJsonSkinIndex FGLTFSkinConverter::Convert(FGLTFJsonNodeIndex RootNode, cons
 		Skin.Joints[BoneIndex] = Builder.GetOrAddNode(RootNode, SkeletalMesh, BoneIndex);
 	}
 
-	TArray<FGLTFJsonMatrix4> InverseBindMatrices;
+	TArray<FGLTFRawMatrix4> InverseBindMatrices;
 	InverseBindMatrices.AddUninitialized(BoneCount);
 
 	for (int32 BoneIndex = 0; BoneIndex < BoneCount; ++BoneIndex)
diff --git a/Engine/Plugins/Enterprise/GLTFExporter/MigratedFromGit/Source/GLTFExporter/Private/Json/GLTFJsonAccessor.h b/Engine/Plugins/Enterprise/GLTFExporter/MigratedFromGit/Source/GLTFExporter/Private/Json/GLTFJsonAccessor.h
--- a/Engine/Plugins/Enterprise/GLTFExporter/MigratedFromGit/Source/GLTFExporter/Private/Json/GLTFJsonAccessor.h
+++ b/Engine/Plugins/Enterprise/GLTFExporter/MigratedFromGit/Source/GLTFExporter/Private/Json/GLTFJsonAccessor.h
@@ -2,12 +2,11 @@
 
 #pragma once
 
+#include "Json/GLTFJsonObject.h"
 #include "Json/GLTFJsonEnums.h"
 #include "Json/GLTFJsonIndex.h"
-#include "Json/GLTFJsonUtility.h"
-#include "Serialization/JsonSerializer.h"
 
-struct FGLTFJsonAccessor
+struct FGLTFJsonAccessor : IGLTFJsonObject
 {
 	FString Name;
 
@@ -34,49 +33,33 @@ struct FGLTFJsonAccessor
 	{
 	}
 
-	template <class CharType = TCHAR, class PrintPolicy = TPrettyJsonPrintPolicy<CharType>>
-	void WriteObject(TJsonWriter<CharType, PrintPolicy>& JsonWriter, FGLTFJsonExtensions& Extensions) const
+	virtual void WriteObject(IGLTFJsonWriter& Writer) const override
 	{
-		JsonWriter.WriteObjectStart();
-
 		if (!Name.IsEmpty())
 		{
-			JsonWriter.WriteValue(TEXT("name"), Name);
+			Writer.Write(TEXT("name"), Name);
 		}
 
-		JsonWriter.WriteValue(TEXT("bufferView"), BufferView);
+		Writer.Write(TEXT("bufferView"), BufferView);
 
 		if (ByteOffset != 0)
 		{
-			JsonWriter.WriteValue(TEXT("byteOffset"), ByteOffset);
+			Writer.Write(TEXT("byteOffset"), ByteOffset);
 		}
 
-		JsonWriter.WriteValue(TEXT("count"), Count);
-		JsonWriter.WriteValue(TEXT("type"), FGLTFJsonUtility::ToString(Type));
-		JsonWriter.WriteValue(TEXT("componentType"), FGLTFJsonUtility::ToInteger(ComponentType));
+		Writer.Write(TEXT("count"), Count);
+		Writer.Write(TEXT("type"), Type);
+		Writer.Write(TEXT("componentType"), ComponentType);
 
 		if (bNormalized)
 		{
-			JsonWriter.WriteValue(TEXT("normalized"), bNormalized);
+			Writer.Write(TEXT("normalized"), bNormalized);
 		}
 
 		if (MinMaxLength > 0)
 		{
-			JsonWriter.WriteArrayStart(TEXT("min"));
-			for (int32 ComponentIndex = 0; ComponentIndex < MinMaxLength; ++ComponentIndex)
-			{
-				FGLTFJsonUtility::WriteExactValue(JsonWriter, Min[ComponentIndex]);
-			}
-			JsonWriter.WriteArrayEnd();
-
-			JsonWriter.WriteArrayStart(TEXT("max"));
-			for (int32 ComponentIndex = 0; ComponentIndex < MinMaxLength; ++ComponentIndex)
-			{
-				FGLTFJsonUtility::WriteExactValue(JsonWriter, Max[ComponentIndex]);
-			}
-			JsonWriter.WriteArrayEnd();
+			Writer.Write(TEXT("min"), Min, MinMaxLength);
+			Writer.Write(TEXT("max"), Max, MinMaxLength);
 		}
-
-		JsonWriter.WriteObjectEnd();
 	}
 };
diff --git a/Engine/Plugins/Enterprise/GLTFExporter/MigratedFromGit/Source/GLTFExporter/Private/Json/GLTFJsonAnimation.h b/Engine/Plugins/Enterprise/GLTFExporter/MigratedFromGit/Source/GLTFExporter/Private/Json/GLTFJsonAnimation.h
--- a/Engine/Plugins/Enterprise/GLTFExporter/MigratedFromGit/Source/GLTFExporter/Private/Json/GLTFJsonAnimation.h
+++ b/Engine/Plugins/Enterprise/GLTFExporter/MigratedFromGit/Source/GLTFExporter/Private/Json/GLTFJsonAnimation.h
@@ -2,13 +2,12 @@
 
 #pragma once
 
+#include "Json/GLTFJsonObject.h"
 #include "Json/GLTFJsonEnums.h"
 #include "Json/GLTFJsonIndex.h"
 #include "Json/GLTFJsonAnimationPlayback.h"
-#include "Json/GLTFJsonUtility.h"
-#include "Serialization/JsonSerializer.h"
 
-struct FGLTFJsonAnimationChannelTarget
+struct FGLTFJsonAnimationChannelTarget : IGLTFJsonObject
 {
 	FGLTFJsonNodeIndex Node;
 	EGLTFJsonTargetPath Path;
@@ -18,38 +17,27 @@ struct FGLTFJsonAnimationChannelTarget
 	{
 	}
 
-	template <class CharType = TCHAR, class PrintPolicy = TPrettyJsonPrintPolicy<CharType>>
-	void WriteObject(TJsonWriter<CharType, PrintPolicy>& JsonWriter, FGLTFJsonExtensions& Extensions) const
+	virtual void WriteObject(IGLTFJsonWriter& Writer) const override
 	{
-		JsonWriter.WriteObjectStart();
-
-		JsonWriter.WriteValue(TEXT("node"), Node);
-		JsonWriter.WriteValue(TEXT("path"), FGLTFJsonUtility::ToString(Path));
-
-		JsonWriter.WriteObjectEnd();
+		Writer.Write(TEXT("node"), Node);
+		Writer.Write(TEXT("path"), Path);
 	}
 };
 
-struct FGLTFJsonAnimationChannel
+struct FGLTFJsonAnimationChannel : IGLTFJsonObject
 {
 	FGLTFJsonAnimationSamplerIndex Sampler;
 	FGLTFJsonAnimationChannelTarget Target;
 
-	template <class CharType = TCHAR, class PrintPolicy = TPrettyJsonPrintPolicy<CharType>>
-	void WriteObject(TJsonWriter<CharType, PrintPolicy>& JsonWriter, FGLTFJsonExtensions& Extensions) const
+	virtual void WriteObject(IGLTFJsonWriter& Writer) const override
 	{
-		JsonWriter.WriteObjectStart();
-
-		JsonWriter.WriteValue(TEXT("sampler"), Sampler);
+		Writer.Write(TEXT("sampler"), Sampler);
 
-		JsonWriter.WriteIdentifierPrefix(TEXT("target"));
-		Target.WriteObject(JsonWriter, Extensions);
-
-		JsonWriter.WriteObjectEnd();
+		Writer.Write(TEXT("target"), Target);
 	}
 };
 
-struct FGLTFJsonAnimationSampler
+struct FGLTFJsonAnimationSampler : IGLTFJsonObject
 {
 	FGLTFJsonAccessorIndex Input;
 	FGLTFJsonAccessorIndex Output;
@@ -61,24 +49,19 @@ struct FGLTFJsonAnimationSampler
 	{
 	}
 
-	template <class CharType = TCHAR, class PrintPolicy = TPrettyJsonPrintPolicy<CharType>>
-	void WriteObject(TJsonWriter<CharType, PrintPolicy>& JsonWriter, FGLTFJsonExtensions& Extensions) const
+	virtual void WriteObject(IGLTFJsonWriter& Writer) const override
 	{
-		JsonWriter.WriteObjectStart();
-
-		JsonWriter.WriteValue(TEXT("input"), Input);
-		JsonWriter.WriteValue(TEXT("output"), Output);
+		Writer.Write(TEXT("input"), Input);
+		Writer.Write(TEXT("output"), Output);
 
 		if (Interpolation != EGLTFJsonInterpolation::Linear)
 		{
-			JsonWriter.WriteValue(TEXT("interpolation"), FGLTFJsonUtility::ToString(Interpolation));
+			Writer.Write(TEXT("interpolation"), Interpolation);
 		}
-
-		JsonWriter.WriteObjectEnd();
 	}
 };
 
-struct FGLTFJsonAnimation
+struct FGLTFJsonAnimation : IGLTFJsonObject
 {
 	FString Name;
 
@@ -87,30 +70,21 @@ struct FGLTFJsonAnimation
 
 	FGLTFJsonAnimationPlayback Playback;
 
-	template <class CharType = TCHAR, class PrintPolicy = TPrettyJsonPrintPolicy<CharType>>
-	void WriteObject(TJsonWriter<CharType, PrintPolicy>& JsonWriter, FGLTFJsonExtensions& Extensions) const
+	virtual void WriteObject(IGLTFJsonWriter& Writer) const override
 	{
-		JsonWriter.WriteObjectStart();
-
 		if (!Name.IsEmpty())
 		{
-			JsonWriter.WriteValue(TEXT("name"), Name);
+			Writer.Write(TEXT("name"), Name);
 		}
 
-		FGLTFJsonUtility::WriteObjectArray(JsonWriter, TEXT("channels"), Channels, Extensions, true);
-		FGLTFJsonUtility::WriteObjectArray(JsonWriter, TEXT("samplers"), Samplers, Extensions, true);
+		Writer.Write(TEXT("channels"), Channels);
+		Writer.Write(TEXT("samplers"), Samplers);
 
 		if (Playback != FGLTFJsonAnimationPlayback())
 		{
-			const EGLTFJsonExtension Extension = EGLTFJsonExtension::EPIC_AnimationPlayback;
-			Extensions.Used.Add(Extension);
-
-			JsonWriter.WriteObjectStart(TEXT("extensions"));
-			JsonWriter.WriteIdentifierPrefix(FGLTFJsonUtility::ToString(Extension));
-			Playback.WriteObject(JsonWriter, Extensions);
-			JsonWriter.WriteObjectEnd();
+			Writer.StartExtensions();
+			Writer.Write(EGLTFJsonExtension::EPIC_AnimationPlayback, Playback);
+			Writer.EndExtensions();
 		}
-
-		JsonWriter.WriteObjectEnd();
 	}
 };
diff --git a/Engine/Plugins/Enterprise/GLTFExporter/MigratedFromGit/Source/GLTFExporter/Private/Json/GLTFJsonAnimationPlayback.h b/Engine/Plugins/Enterprise/GLTFExporter/MigratedFromGit/Source/GLTFExporter/Private/Json/GLTFJsonAnimationPlayback.h
--- a/Engine/Plugins/Enterprise/GLTFExporter/MigratedFromGit/Source/GLTFExporter/Private/Json/GLTFJsonAnimationPlayback.h
+++ b/Engine/Plugins/Enterprise/GLTFExporter/MigratedFromGit/Source/GLTFExporter/Private/Json/GLTFJsonAnimationPlayback.h
@@ -2,10 +2,9 @@
 
 #pragma once
 
-#include "Json/GLTFJsonUtility.h"
-#include "Serialization/JsonSerializer.h"
+#include "Json/GLTFJsonObject.h"
 
-struct FGLTFJsonAnimationPlayback
+struct FGLTFJsonAnimationPlayback : IGLTFJsonObject
 {
 	FString Name;
 
@@ -23,37 +22,32 @@ struct FGLTFJsonAnimationPlayback
 	{
 	}
 
-	template <class CharType = TCHAR, class PrintPolicy = TPrettyJsonPrintPolicy<CharType>>
-	void WriteObject(TJsonWriter<CharType, PrintPolicy>& JsonWriter, FGLTFJsonExtensions& Extensions) const
+	virtual void WriteObject(IGLTFJsonWriter& Writer) const override
 	{
-		JsonWriter.WriteObjectStart();
-
 		if (!Name.IsEmpty())
 		{
-			JsonWriter.WriteValue(TEXT("name"), Name);
+			Writer.Write(TEXT("name"), Name);
 		}
 
 		if (bLoop != true)
 		{
-			JsonWriter.WriteValue(TEXT("loop"), bLoop);
+			Writer.Write(TEXT("loop"), bLoop);
 		}
 
 		if (bAutoPlay != true)
 		{
-			JsonWriter.WriteValue(TEXT("autoPlay"), bAutoPlay);
+			Writer.Write(TEXT("autoPlay"), bAutoPlay);
 		}
 
 		if (PlayRate != 1)
 		{
-			FGLTFJsonUtility::WriteExactValue(JsonWriter, TEXT("playRate"), PlayRate);
+			Writer.Write(TEXT("playRate"), PlayRate);
 		}
 
 		if (StartTime != 0)
 		{
-			FGLTFJsonUtility::WriteExactValue(JsonWriter, TEXT("startTime"), StartTime);
+			Writer.Write(TEXT("startTime"), StartTime);
 		}
-
-		JsonWriter.WriteObjectEnd();
 	}
 
 	bool operator==(const FGLTFJsonAnimationPlayback& Other) const
diff --git a/Engine/Plugins/Enterprise/GLTFExporter/MigratedFromGit/Source/GLTFExporter/Private/Json/GLTFJsonArray.h b/Engine/Plugins/Enterprise/GLTFExporter/MigratedFromGit/Source/GLTFExporter/Private/Json/GLTFJsonArray.h
new file mode 100644
--- /dev/null
+++ b/Engine/Plugins/Enterprise/GLTFExporter/MigratedFromGit/Source/GLTFExporter/Private/Json/GLTFJsonArray.h
@@ -0,0 +1,18 @@
+// Copyright Epic Games, Inc. All Rights Reserved.
+
+#pragma once
+
+#include "Json/GLTFJsonValue.h"
+#include "Json/GLTFJsonWriter.h"
+
+struct IGLTFJsonArray : IGLTFJsonValue
+{
+	virtual void WriteValue(IGLTFJsonWriter& Writer) const override final
+	{
+		Writer.StartArray();
+		WriteArray(Writer);
+		Writer.EndArray();
+	}
+
+	virtual void WriteArray(IGLTFJsonWriter& Writer) const = 0;
+};
diff --git a/Engine/Plugins/Enterprise/GLTFExporter/MigratedFromGit/Source/GLTFExporter/Private/Json/GLTFJsonAsset.h b/Engine/Plugins/Enterprise/GLTFExporter/MigratedFromGit/Source/GLTFExporter/Private/Json/GLTFJsonAsset.h
--- a/Engine/Plugins/Enterprise/GLTFExporter/MigratedFromGit/Source/GLTFExporter/Private/Json/GLTFJsonAsset.h
+++ b/Engine/Plugins/Enterprise/GLTFExporter/MigratedFromGit/Source/GLTFExporter/Private/Json/GLTFJsonAsset.h
@@ -2,10 +2,9 @@
 
 #pragma once
 
-#include "Json/GLTFJsonExtensions.h"
-#include "Serialization/JsonSerializer.h"
+#include "Json/GLTFJsonObject.h"
 
-struct FGLTFJsonAsset
+struct FGLTFJsonAsset : IGLTFJsonObject
 {
 	FString Version;
 	FString Generator;
@@ -16,23 +15,18 @@ struct FGLTFJsonAsset
 	{
 	}
 
-	template <class CharType = TCHAR, class PrintPolicy = TPrettyJsonPrintPolicy<CharType>>
-	void WriteObject(TJsonWriter<CharType, PrintPolicy>& JsonWriter, FGLTFJsonExtensions& Extensions) const
+	virtual void WriteObject(IGLTFJsonWriter& Writer) const override
 	{
-		JsonWriter.WriteObjectStart();
-
-		JsonWriter.WriteValue(TEXT("version"), Version);
+		Writer.Write(TEXT("version"), Version);
 
 		if (!Generator.IsEmpty())
 		{
-			JsonWriter.WriteValue(TEXT("generator"), Generator);
+			Writer.Write(TEXT("generator"), Generator);
 		}
 
 		if (!Copyright.IsEmpty())
 		{
-			JsonWriter.WriteValue(TEXT("copyright"), Copyright);
+			Writer.Write(TEXT("copyright"), Copyright);
 		}
-
-		JsonWriter.WriteObjectEnd();
 	}
 };
diff --git a/Engine/Plugins/Enterprise/GLTFExporter/MigratedFromGit/Source/GLTFExporter/Private/Json/GLTFJsonBackdrop.h b/Engine/Plugins/Enterprise/GLTFExporter/MigratedFromGit/Source/GLTFExporter/Private/Json/GLTFJsonBackdrop.h
--- a/Engine/Plugins/Enterprise/GLTFExporter/MigratedFromGit/Source/GLTFExporter/Private/Json/GLTFJsonBackdrop.h
+++ b/Engine/Plugins/Enterprise/GLTFExporter/MigratedFromGit/Source/GLTFExporter/Private/Json/GLTFJsonBackdrop.h
@@ -2,11 +2,11 @@
 
 #pragma once
 
+#include "Json/GLTFJsonObject.h"
 #include "Json/GLTFJsonIndex.h"
 #include "Json/GLTFJsonVector3.h"
-#include "Json/GLTFJsonUtility.h"
 
-struct FGLTFJsonBackdrop
+struct FGLTFJsonBackdrop : IGLTFJsonObject
 {
 	FString Name;
 
@@ -32,37 +32,31 @@ struct FGLTFJsonBackdrop
 	{
 	}
 
-	template <class CharType = TCHAR, class PrintPolicy = TPrettyJsonPrintPolicy<CharType>>
-	void WriteObject(TJsonWriter<CharType, PrintPolicy>& JsonWriter, FGLTFJsonExtensions& Extensions) const
+	virtual void WriteObject(IGLTFJsonWriter& Writer) const override
 	{
-		JsonWriter.WriteObjectStart();
-
 		if (!Name.IsEmpty())
 		{
-			JsonWriter.WriteValue(TEXT("name"), Name);
+			Writer.Write(TEXT("name"), Name);
 		}
 
 		if (Mesh != INDEX_NONE)
 		{
-			JsonWriter.WriteValue(TEXT("mesh"), Mesh);
+			Writer.Write(TEXT("mesh"), Mesh);
 		}
 
-		FGLTFJsonUtility::WriteFixedArray(JsonWriter, TEXT("cubemap"), Cubemap);
+		Writer.Write(TEXT("cubemap"), Cubemap);
 
-		FGLTFJsonUtility::WriteExactValue(JsonWriter, TEXT("intensity"), Intensity);
-		FGLTFJsonUtility::WriteExactValue(JsonWriter, TEXT("size"), Size);
+		Writer.Write(TEXT("intensity"), Intensity);
+		Writer.Write(TEXT("size"), Size);
 
 		if (!FMath::IsNearlyZero(Angle))
 		{
-			FGLTFJsonUtility::WriteExactValue(JsonWriter, TEXT("angle"), Angle);
+			Writer.Write(TEXT("angle"), Angle);
 		}
 
-		JsonWriter.WriteIdentifierPrefix(TEXT("projectionCenter"));
-		ProjectionCenter.WriteArray(JsonWriter);
-
-		FGLTFJsonUtility::WriteExactValue(JsonWriter, TEXT("lightingDistanceFactor"), LightingDistanceFactor);
-		JsonWriter.WriteValue(TEXT("useCameraProjection"), UseCameraProjection);
+		Writer.Write(TEXT("projectionCenter"), ProjectionCenter);
 
-		JsonWriter.WriteObjectEnd();
+		Writer.Write(TEXT("lightingDistanceFactor"), LightingDistanceFactor);
+		Writer.Write(TEXT("useCameraProjection"), UseCameraProjection);
 	}
 };
diff --git a/Engine/Plugins/Enterprise/GLTFExporter/MigratedFromGit/Source/GLTFExporter/Private/Json/GLTFJsonBuffer.h b/Engine/Plugins/Enterprise/GLTFExporter/MigratedFromGit/Source/GLTFExporter/Private/Json/GLTFJsonBuffer.h
--- a/Engine/Plugins/Enterprise/GLTFExporter/MigratedFromGit/Source/GLTFExporter/Private/Json/GLTFJsonBuffer.h
+++ b/Engine/Plugins/Enterprise/GLTFExporter/MigratedFromGit/Source/GLTFExporter/Private/Json/GLTFJsonBuffer.h
@@ -2,9 +2,9 @@
 
 #pragma once
 
-#include "Serialization/JsonSerializer.h"
+#include "Json/GLTFJsonObject.h"
 
-struct FGLTFJsonBuffer
+struct FGLTFJsonBuffer : IGLTFJsonObject
 {
 	FString Name;
 
@@ -16,23 +16,18 @@ struct FGLTFJsonBuffer
 	{
 	}
 
-	template <class CharType = TCHAR, class PrintPolicy = TPrettyJsonPrintPolicy<CharType>>
-	void WriteObject(TJsonWriter<CharType, PrintPolicy>& JsonWriter, FGLTFJsonExtensions& Extensions) const
+	virtual void WriteObject(IGLTFJsonWriter& Writer) const override
 	{
-		JsonWriter.WriteObjectStart();
-
 		if (!Name.IsEmpty())
 		{
-			JsonWriter.WriteValue(TEXT("name"), Name);
+			Writer.Write(TEXT("name"), Name);
 		}
 
 		if (!URI.IsEmpty())
 		{
-			JsonWriter.WriteValue(TEXT("uri"), URI);
+			Writer.Write(TEXT("uri"), URI);
 		}
 
-		JsonWriter.WriteValue(TEXT("byteLength"), ByteLength);
-
-		JsonWriter.WriteObjectEnd();
+		Writer.Write(TEXT("byteLength"), ByteLength);
 	}
 };
diff --git a/Engine/Plugins/Enterprise/GLTFExporter/MigratedFromGit/Source/GLTFExporter/Private/Json/GLTFJsonBufferView.h b/Engine/Plugins/Enterprise/GLTFExporter/MigratedFromGit/Source/GLTFExporter/Private/Json/GLTFJsonBufferView.h
--- a/Engine/Plugins/Enterprise/GLTFExporter/MigratedFromGit/Source/GLTFExporter/Private/Json/GLTFJsonBufferView.h
+++ b/Engine/Plugins/Enterprise/GLTFExporter/MigratedFromGit/Source/GLTFExporter/Private/Json/GLTFJsonBufferView.h
@@ -2,12 +2,11 @@
 
 #pragma once
 
+#include "Json/GLTFJsonObject.h"
 #include "Json/GLTFJsonEnums.h"
 #include "Json/GLTFJsonIndex.h"
-#include "Json/GLTFJsonUtility.h"
-#include "Serialization/JsonSerializer.h"
 
-struct FGLTFJsonBufferView
+struct FGLTFJsonBufferView : IGLTFJsonObject
 {
 	FString Name;
 
@@ -28,34 +27,29 @@ struct FGLTFJsonBufferView
 		// check that view fits completely inside the buffer
 	}
 
-	template <class CharType = TCHAR, class PrintPolicy = TPrettyJsonPrintPolicy<CharType>>
-	void WriteObject(TJsonWriter<CharType, PrintPolicy>& JsonWriter, FGLTFJsonExtensions& Extensions) const
+	virtual void WriteObject(IGLTFJsonWriter& Writer) const override
 	{
-		JsonWriter.WriteObjectStart();
-
 		if (!Name.IsEmpty())
 		{
-			JsonWriter.WriteValue(TEXT("name"), Name);
+			Writer.Write(TEXT("name"), Name);
 		}
 
-		JsonWriter.WriteValue(TEXT("buffer"), Buffer);
-		JsonWriter.WriteValue(TEXT("byteLength"), ByteLength);
+		Writer.Write(TEXT("buffer"), Buffer);
+		Writer.Write(TEXT("byteLength"), ByteLength);
 
 		if (ByteOffset != 0)
 		{
-			JsonWriter.WriteValue(TEXT("byteOffset"), ByteOffset);
+			Writer.Write(TEXT("byteOffset"), ByteOffset);
 		}
 
 		if (ByteStride != 0)
 		{
-			JsonWriter.WriteValue(TEXT("byteStride"), ByteStride);
+			Writer.Write(TEXT("byteStride"), ByteStride);
 		}
 
 		if (Target != EGLTFJsonBufferTarget::None)
 		{
-			JsonWriter.WriteValue(TEXT("target"), FGLTFJsonUtility::ToInteger(Target));
+			Writer.Write(TEXT("target"), Target);
 		}
-
-		JsonWriter.WriteObjectEnd();
 	}
 };
diff --git a/Engine/Plugins/Enterprise/GLTFExporter/MigratedFromGit/Source/GLTFExporter/Private/Json/GLTFJsonCamera.h b/Engine/Plugins/Enterprise/GLTFExporter/MigratedFromGit/Source/GLTFExporter/Private/Json/GLTFJsonCamera.h
--- a/Engine/Plugins/Enterprise/GLTFExporter/MigratedFromGit/Source/GLTFExporter/Private/Json/GLTFJsonCamera.h
+++ b/Engine/Plugins/Enterprise/GLTFExporter/MigratedFromGit/Source/GLTFExporter/Private/Json/GLTFJsonCamera.h
@@ -2,12 +2,11 @@
 
 #pragma once
 
+#include "Json/GLTFJsonObject.h"
 #include "Json/GLTFJsonEnums.h"
 #include "Json/GLTFJsonCameraControl.h"
-#include "Json/GLTFJsonUtility.h"
-#include "Serialization/JsonSerializer.h"
 
-struct FGLTFJsonOrthographic
+struct FGLTFJsonOrthographic : IGLTFJsonObject
 {
 	float XMag; // horizontal magnification of the view
 	float YMag; // vertical magnification of the view
@@ -22,21 +21,16 @@ struct FGLTFJsonOrthographic
 	{
 	}
 
-	template <class CharType = TCHAR, class PrintPolicy = TPrettyJsonPrintPolicy<CharType>>
-	void WriteObject(TJsonWriter<CharType, PrintPolicy>& JsonWriter, FGLTFJsonExtensions& Extensions) const
+	virtual void WriteObject(IGLTFJsonWriter& Writer) const override
 	{
-		JsonWriter.WriteObjectStart();
-
-		FGLTFJsonUtility::WriteExactValue(JsonWriter, TEXT("xmag"), XMag);
-		FGLTFJsonUtility::WriteExactValue(JsonWriter, TEXT("ymag"), YMag);
-		FGLTFJsonUtility::WriteExactValue(JsonWriter, TEXT("zfar"), ZFar);
-		FGLTFJsonUtility::WriteExactValue(JsonWriter, TEXT("znear"), ZNear);
-
-		JsonWriter.WriteObjectEnd();
+		Writer.Write(TEXT("xmag"), XMag);
+		Writer.Write(TEXT("ymag"), YMag);
+		Writer.Write(TEXT("zfar"), ZFar);
+		Writer.Write(TEXT("znear"), ZNear);
 	}
 };
 
-struct FGLTFJsonPerspective
+struct FGLTFJsonPerspective : IGLTFJsonObject
 {
 	float AspectRatio; // aspect ratio of the field of view
 	float YFov; // vertical field of view in radians
@@ -51,70 +45,56 @@ struct FGLTFJsonPerspective
 	{
 	}
 
-	template <class CharType = TCHAR, class PrintPolicy = TPrettyJsonPrintPolicy<CharType>>
-	void WriteObject(TJsonWriter<CharType, PrintPolicy>& JsonWriter, FGLTFJsonExtensions& Extensions) const
+	virtual void WriteObject(IGLTFJsonWriter& Writer) const override
 	{
-		JsonWriter.WriteObjectStart();
-
 		if (AspectRatio != 0)
 		{
-			FGLTFJsonUtility::WriteExactValue(JsonWriter, TEXT("aspectRatio"), AspectRatio);
+			Writer.Write(TEXT("aspectRatio"), AspectRatio);
 		}
 
-		FGLTFJsonUtility::WriteExactValue(JsonWriter, TEXT("yfov"), YFov);
+		Writer.Write(TEXT("yfov"), YFov);
 
 		if (ZFar != 0)
 		{
-			FGLTFJsonUtility::WriteExactValue(JsonWriter, TEXT("zfar"), ZFar);
+			Writer.Write(TEXT("zfar"), ZFar);
 		}
 
-		FGLTFJsonUtility::WriteExactValue(JsonWriter, TEXT("znear"), ZNear);
-
-		JsonWriter.WriteObjectEnd();
+		Writer.Write(TEXT("znear"), ZNear);
 	}
 };
 
-struct FGLTFJsonCamera
+struct FGLTFJsonCamera : IGLTFJsonObject
 {
 	FString Name;
 
 	EGLTFJsonCameraType               Type;
 	TOptional<FGLTFJsonCameraControl> CameraControl;
 
-	union {
-		FGLTFJsonOrthographic Orthographic;
-		FGLTFJsonPerspective  Perspective;
-	};
+	FGLTFJsonOrthographic Orthographic;
+	FGLTFJsonPerspective  Perspective;
 
 	FGLTFJsonCamera()
 		: Type(EGLTFJsonCameraType::None)
-		, Orthographic()
-		, Perspective()
 	{
 	}
 
-	template <class CharType = TCHAR, class PrintPolicy = TPrettyJsonPrintPolicy<CharType>>
-	void WriteObject(TJsonWriter<CharType, PrintPolicy>& JsonWriter, FGLTFJsonExtensions& Extensions) const
+	virtual void WriteObject(IGLTFJsonWriter& Writer) const override
 	{
-		JsonWriter.WriteObjectStart();
-
 		if (!Name.IsEmpty())
 		{
-			JsonWriter.WriteValue(TEXT("name"), Name);
+			Writer.Write(TEXT("name"), Name);
 		}
 
-		JsonWriter.WriteValue(TEXT("type"), FGLTFJsonUtility::ToString(Type));
+		Writer.Write(TEXT("type"), Type);
 
 		switch (Type)
 		{
 			case EGLTFJsonCameraType::Orthographic:
-				JsonWriter.WriteIdentifierPrefix(TEXT("orthographic"));
-				Orthographic.WriteObject(JsonWriter, Extensions);
+				Writer.Write(TEXT("orthographic"), Orthographic);
 				break;
 
 			case EGLTFJsonCameraType::Perspective:
-				JsonWriter.WriteIdentifierPrefix(TEXT("perspective"));
-				Perspective.WriteObject(JsonWriter, Extensions);
+				Writer.Write(TEXT("perspective"), Perspective);
 				break;
 
 			default:
@@ -123,17 +103,9 @@ struct FGLTFJsonCamera
 
 		if (CameraControl.IsSet())
 		{
-			JsonWriter.WriteObjectStart(TEXT("extensions"));
-
-			const EGLTFJsonExtension Extension = EGLTFJsonExtension::EPIC_CameraControls;
-			Extensions.Used.Add(Extension);
-
-			JsonWriter.WriteIdentifierPrefix(FGLTFJsonUtility::ToString(Extension));
-			CameraControl.GetValue().WriteObject(JsonWriter, Extensions);
-
-			JsonWriter.WriteObjectEnd();
+			Writer.StartExtensions();
+			Writer.Write(EGLTFJsonExtension::EPIC_CameraControls, CameraControl.GetValue());
+			Writer.EndExtensions();
 		}
-
-		JsonWriter.WriteObjectEnd();
 	}
 };
diff --git a/Engine/Plugins/Enterprise/GLTFExporter/MigratedFromGit/Source/GLTFExporter/Private/Json/GLTFJsonCameraControl.h b/Engine/Plugins/Enterprise/GLTFExporter/MigratedFromGit/Source/GLTFExporter/Private/Json/GLTFJsonCameraControl.h
--- a/Engine/Plugins/Enterprise/GLTFExporter/MigratedFromGit/Source/GLTFExporter/Private/Json/GLTFJsonCameraControl.h
+++ b/Engine/Plugins/Enterprise/GLTFExporter/MigratedFromGit/Source/GLTFExporter/Private/Json/GLTFJsonCameraControl.h
@@ -2,10 +2,11 @@
 
 #pragma once
 
+#include "Json/GLTFJsonObject.h"
 #include "Json/GLTFJsonIndex.h"
 #include "Json/GLTFJsonVector4.h"
 
-struct FGLTFJsonCameraControl
+struct FGLTFJsonCameraControl : IGLTFJsonObject
 {
 	EGLTFJsonCameraControlMode Mode;
 	FGLTFJsonNodeIndex Target;
@@ -35,46 +36,41 @@ struct FGLTFJsonCameraControl
 	{
 	}
 
-	template <class CharType = TCHAR, class PrintPolicy = TPrettyJsonPrintPolicy<CharType>>
-	void WriteObject(TJsonWriter<CharType, PrintPolicy>& JsonWriter, FGLTFJsonExtensions& Extensions) const
+	virtual void WriteObject(IGLTFJsonWriter& Writer) const override
 	{
-		JsonWriter.WriteObjectStart();
-
-		JsonWriter.WriteValue(TEXT("mode"), FGLTFJsonUtility::ToString(Mode));
+		Writer.Write(TEXT("mode"), Mode);
 
 		if (Target != INDEX_NONE && Mode == EGLTFJsonCameraControlMode::Orbital)
 		{
-			JsonWriter.WriteValue(TEXT("target"), Target);
+			Writer.Write(TEXT("target"), Target);
 		}
 
-		FGLTFJsonUtility::WriteExactValue(JsonWriter, TEXT("maxDistance"), MaxDistance);
-		FGLTFJsonUtility::WriteExactValue(JsonWriter, TEXT("minDistance"), MinDistance);
+		Writer.Write(TEXT("maxDistance"), MaxDistance);
+		Writer.Write(TEXT("minDistance"), MinDistance);
 
 		if (!FMath::IsNearlyEqual(MaxPitch, 90.0f))
 		{
-			FGLTFJsonUtility::WriteExactValue(JsonWriter, TEXT("maxPitch"), MaxPitch);
+			Writer.Write(TEXT("maxPitch"), MaxPitch);
 		}
 
 		if (!FMath::IsNearlyEqual(MinPitch, -90.0f))
 		{
-			FGLTFJsonUtility::WriteExactValue(JsonWriter, TEXT("minPitch"), MinPitch);
+			Writer.Write(TEXT("minPitch"), MinPitch);
 		}
 
 		if (!FMath::IsNearlyEqual(MaxYaw, 360.0f))
 		{
-			FGLTFJsonUtility::WriteExactValue(JsonWriter, TEXT("maxYaw"), MaxYaw);
+			Writer.Write(TEXT("maxYaw"), MaxYaw);
 		}
 
 		if (!FMath::IsNearlyEqual(MinYaw, 0.0f))
 		{
-			FGLTFJsonUtility::WriteExactValue(JsonWriter, TEXT("minYaw"), MinYaw);
+			Writer.Write(TEXT("minYaw"), MinYaw);
 		}
 
-		FGLTFJsonUtility::WriteExactValue(JsonWriter, TEXT("rotationSensitivity"), RotationSensitivity);
-		FGLTFJsonUtility::WriteExactValue(JsonWriter, TEXT("rotationInertia"), RotationInertia);
-		FGLTFJsonUtility::WriteExactValue(JsonWriter, TEXT("dollySensitivity"), DollySensitivity);
-		FGLTFJsonUtility::WriteExactValue(JsonWriter, TEXT("dollyDuration"), DollyDuration);
-
-		JsonWriter.WriteObjectEnd();
+		Writer.Write(TEXT("rotationSensitivity"), RotationSensitivity);
+		Writer.Write(TEXT("rotationInertia"), RotationInertia);
+		Writer.Write(TEXT("dollySensitivity"), DollySensitivity);
+		Writer.Write(TEXT("dollyDuration"), DollyDuration);
 	}
 };
diff --git a/Engine/Plugins/Enterprise/GLTFExporter/MigratedFromGit/Source/GLTFExporter/Private/Json/GLTFJsonColor3.h b/Engine/Plugins/Enterprise/GLTFExporter/MigratedFromGit/Source/GLTFExporter/Private/Json/GLTFJsonColor3.h
--- a/Engine/Plugins/Enterprise/GLTFExporter/MigratedFromGit/Source/GLTFExporter/Private/Json/GLTFJsonColor3.h
+++ b/Engine/Plugins/Enterprise/GLTFExporter/MigratedFromGit/Source/GLTFExporter/Private/Json/GLTFJsonColor3.h
@@ -2,10 +2,9 @@
 
 #pragma once
 
-#include "Json/GLTFJsonUtility.h"
-#include "Serialization/JsonSerializer.h"
+#include "Json/GLTFJsonArray.h"
 
-struct FGLTFJsonColor3
+struct FGLTFJsonColor3 : IGLTFJsonArray
 {
 	float R, G, B;
 
@@ -17,14 +16,11 @@ struct FGLTFJsonColor3
 	{
 	}
 
-	template <class CharType = TCHAR, class PrintPolicy = TPrettyJsonPrintPolicy<CharType>>
-	void WriteArray(TJsonWriter<CharType, PrintPolicy>& JsonWriter) const
+	virtual void WriteArray(IGLTFJsonWriter& Writer) const override
 	{
-		JsonWriter.WriteArrayStart();
-		FGLTFJsonUtility::WriteExactValue(JsonWriter, R);
-		FGLTFJsonUtility::WriteExactValue(JsonWriter, G);
-		FGLTFJsonUtility::WriteExactValue(JsonWriter, B);
-		JsonWriter.WriteArrayEnd();
+		Writer.Write(R);
+		Writer.Write(G);
+		Writer.Write(B);
 	}
 
 	bool operator==(const FGLTFJsonColor3& Other) const
diff --git a/Engine/Plugins/Enterprise/GLTFExporter/MigratedFromGit/Source/GLTFExporter/Private/Json/GLTFJsonColor4.h b/Engine/Plugins/Enterprise/GLTFExporter/MigratedFromGit/Source/GLTFExporter/Private/Json/GLTFJsonColor4.h
--- a/Engine/Plugins/Enterprise/GLTFExporter/MigratedFromGit/Source/GLTFExporter/Private/Json/GLTFJsonColor4.h
+++ b/Engine/Plugins/Enterprise/GLTFExporter/MigratedFromGit/Source/GLTFExporter/Private/Json/GLTFJsonColor4.h
@@ -2,11 +2,10 @@
 
 #pragma once
 
+#include "Json/GLTFJsonArray.h"
 #include "Json/GLTFJsonColor3.h"
-#include "Json/GLTFJsonUtility.h"
-#include "Serialization/JsonSerializer.h"
 
-struct FGLTFJsonColor4
+struct FGLTFJsonColor4 : IGLTFJsonArray
 {
 	float R, G, B, A;
 
@@ -23,15 +22,12 @@ struct FGLTFJsonColor4
 	{
 	}
 
-	template <class CharType = TCHAR, class PrintPolicy = TPrettyJsonPrintPolicy<CharType>>
-	void WriteArray(TJsonWriter<CharType, PrintPolicy>& JsonWriter) const
+	virtual void WriteArray(IGLTFJsonWriter& Writer) const override
 	{
-		JsonWriter.WriteArrayStart();
-		FGLTFJsonUtility::WriteExactValue(JsonWriter, R);
-		FGLTFJsonUtility::WriteExactValue(JsonWriter, G);
-		FGLTFJsonUtility::WriteExactValue(JsonWriter, B);
-		FGLTFJsonUtility::WriteExactValue(JsonWriter, A);
-		JsonWriter.WriteArrayEnd();
+		Writer.Write(R);
+		Writer.Write(G);
+		Writer.Write(B);
+		Writer.Write(A);
 	}
 
 	operator FGLTFJsonColor3() const
diff --git a/Engine/Plugins/Enterprise/GLTFExporter/MigratedFromGit/Source/GLTFExporter/Private/Json/GLTFJsonHotspot.h b/Engine/Plugins/Enterprise/GLTFExporter/MigratedFromGit/Source/GLTFExporter/Private/Json/GLTFJsonHotspot.h
--- a/Engine/Plugins/Enterprise/GLTFExporter/MigratedFromGit/Source/GLTFExporter/Private/Json/GLTFJsonHotspot.h
+++ b/Engine/Plugins/Enterprise/GLTFExporter/MigratedFromGit/Source/GLTFExporter/Private/Json/GLTFJsonHotspot.h
@@ -2,11 +2,10 @@
 
 #pragma once
 
+#include "Json/GLTFJsonObject.h"
 #include "Json/GLTFJsonIndex.h"
-#include "Json/GLTFJsonExtensions.h"
-#include "Serialization/JsonSerializer.h"
 
-struct FGLTFJsonHotspot
+struct FGLTFJsonHotspot : IGLTFJsonObject
 {
 	FString Name;
 
@@ -16,34 +15,29 @@ struct FGLTFJsonHotspot
 	FGLTFJsonTextureIndex   ToggledImage;
 	FGLTFJsonTextureIndex   ToggledHoveredImage;
 
-	template <class CharType = TCHAR, class PrintPolicy = TPrettyJsonPrintPolicy<CharType>>
-	void WriteObject(TJsonWriter<CharType, PrintPolicy>& JsonWriter, FGLTFJsonExtensions& Extensions) const
+	virtual void WriteObject(IGLTFJsonWriter& Writer) const override
 	{
-		JsonWriter.WriteObjectStart();
-
 		if (!Name.IsEmpty())
 		{
-			JsonWriter.WriteValue(TEXT("name"), Name);
+			Writer.Write(TEXT("name"), Name);
 		}
 
-		JsonWriter.WriteValue(TEXT("animation"), Animation);
-		JsonWriter.WriteValue(TEXT("image"), Image);
+		Writer.Write(TEXT("animation"), Animation);
+		Writer.Write(TEXT("image"), Image);
 
 		if (HoveredImage != INDEX_NONE)
 		{
-			JsonWriter.WriteValue(TEXT("hoveredImage"), HoveredImage);
+			Writer.Write(TEXT("hoveredImage"), HoveredImage);
 		}
 
 		if (ToggledImage != INDEX_NONE)
 		{
-			JsonWriter.WriteValue(TEXT("toggledImage"), ToggledImage);
+			Writer.Write(TEXT("toggledImage"), ToggledImage);
 		}
 
 		if (ToggledHoveredImage != INDEX_NONE)
 		{
-			JsonWriter.WriteValue(TEXT("toggledHoveredImage"), ToggledHoveredImage);
+			Writer.Write(TEXT("toggledHoveredImage"), ToggledHoveredImage);
 		}
-
-		JsonWriter.WriteObjectEnd();
 	}
 };
diff --git a/Engine/Plugins/Enterprise/GLTFExporter/MigratedFromGit/Source/GLTFExporter/Private/Json/GLTFJsonImage.h b/Engine/Plugins/Enterprise/GLTFExporter/MigratedFromGit/Source/GLTFExporter/Private/Json/GLTFJsonImage.h
--- a/Engine/Plugins/Enterprise/GLTFExporter/MigratedFromGit/Source/GLTFExporter/Private/Json/GLTFJsonImage.h
+++ b/Engine/Plugins/Enterprise/GLTFExporter/MigratedFromGit/Source/GLTFExporter/Private/Json/GLTFJsonImage.h
@@ -2,12 +2,11 @@
 
 #pragma once
 
+#include "Json/GLTFJsonObject.h"
 #include "Json/GLTFJsonEnums.h"
 #include "Json/GLTFJsonIndex.h"
-#include "Json/GLTFJsonUtility.h"
-#include "Serialization/JsonSerializer.h"
 
-struct FGLTFJsonImage
+struct FGLTFJsonImage : IGLTFJsonObject
 {
 	FString Name;
 	FString Uri;
@@ -21,31 +20,26 @@ struct FGLTFJsonImage
 	{
 	}
 
-	template <class CharType = TCHAR, class PrintPolicy = TPrettyJsonPrintPolicy<CharType>>
-	void WriteObject(TJsonWriter<CharType, PrintPolicy>& JsonWriter, FGLTFJsonExtensions& Extensions) const
+	virtual void WriteObject(IGLTFJsonWriter& Writer) const override
 	{
-		JsonWriter.WriteObjectStart();
-
 		if (!Name.IsEmpty())
 		{
-			JsonWriter.WriteValue(TEXT("name"), Name);
+			Writer.Write(TEXT("name"), Name);
 		}
 
 		if (!Uri.IsEmpty())
 		{
-			JsonWriter.WriteValue(TEXT("uri"), Uri);
+			Writer.Write(TEXT("uri"), Uri);
 		}
 
 		if (MimeType != EGLTFJsonMimeType::None)
 		{
-			JsonWriter.WriteValue(TEXT("mimeType"), FGLTFJsonUtility::ToString(MimeType));
+			Writer.Write(TEXT("mimeType"), MimeType);
 		}
 
 		if (BufferView != INDEX_NONE)
 		{
-			JsonWriter.WriteValue(TEXT("bufferView"), BufferView);
+			Writer.Write(TEXT("bufferView"), BufferView);
 		}
-
-		JsonWriter.WriteObjectEnd();
 	}
 };
diff --git a/Engine/Plugins/Enterprise/GLTFExporter/MigratedFromGit/Source/GLTFExporter/Private/Json/GLTFJsonLevelVariantSets.h b/Engine/Plugins/Enterprise/GLTFExporter/MigratedFromGit/Source/GLTFExporter/Private/Json/GLTFJsonLevelVariantSets.h
--- a/Engine/Plugins/Enterprise/GLTFExporter/MigratedFromGit/Source/GLTFExporter/Private/Json/GLTFJsonLevelVariantSets.h
+++ b/Engine/Plugins/Enterprise/GLTFExporter/MigratedFromGit/Source/GLTFExporter/Private/Json/GLTFJsonLevelVariantSets.h
@@ -2,13 +2,11 @@
 
 #pragma once
 
-#include "GLTFJsonIndex.h"
-#include "Json/GLTFJsonExtensions.h"
-#include "Json/GLTFJsonUtility.h"
-#include "Serialization/JsonSerializer.h"
+#include "Json/GLTFJsonObject.h"
+#include "Json/GLTFJsonIndex.h"
 #include "Misc/Optional.h"
 
-struct FGLTFJsonVariantMaterial
+struct FGLTFJsonVariantMaterial : IGLTFJsonObject
 {
 	FGLTFJsonMaterialIndex Material;
 	int32                  Index;
@@ -19,23 +17,18 @@ struct FGLTFJsonVariantMaterial
 	{
 	}
 
-	template <class CharType = TCHAR, class PrintPolicy = TPrettyJsonPrintPolicy<CharType>>
-	void WriteObject(TJsonWriter<CharType, PrintPolicy>& JsonWriter, FGLTFJsonExtensions& Extensions) const
+	virtual void WriteObject(IGLTFJsonWriter& Writer) const override
 	{
-		JsonWriter.WriteObjectStart();
-
-		JsonWriter.WriteValue(TEXT("material"), Material);
+		Writer.Write(TEXT("material"), Material);
 
 		if (Index != INDEX_NONE)
 		{
-			JsonWriter.WriteValue(TEXT("index"), Index);
+			Writer.Write(TEXT("index"), Index);
 		}
-
-		JsonWriter.WriteObjectEnd();
 	}
 };
 
-struct FGLTFJsonVariantNodeProperties
+struct FGLTFJsonVariantNodeProperties : IGLTFJsonObject
 {
 	FGLTFJsonNodeIndex Node;
 	TOptional<bool>    bIsVisible;
@@ -43,112 +36,88 @@ struct FGLTFJsonVariantNodeProperties
 	TOptional<FGLTFJsonMeshIndex>    Mesh;
 	TArray<FGLTFJsonVariantMaterial> Materials;
 
-	template <class CharType = TCHAR, class PrintPolicy = TPrettyJsonPrintPolicy<CharType>>
-	void WriteObject(TJsonWriter<CharType, PrintPolicy>& JsonWriter, FGLTFJsonExtensions& Extensions) const
+	virtual void WriteObject(IGLTFJsonWriter& Writer) const override
 	{
-		JsonWriter.WriteObjectStart();
-
 		if (Node != INDEX_NONE)
 		{
-			JsonWriter.WriteValue(TEXT("node"), Node);
+			Writer.Write(TEXT("node"), Node);
 		}
 
-		JsonWriter.WriteObjectStart(TEXT("properties"));
+		Writer.StartObject(TEXT("properties"));
 
 		if (bIsVisible.IsSet())
 		{
-			JsonWriter.WriteValue(TEXT("visible"), bIsVisible.GetValue());
+			Writer.Write(TEXT("visible"), bIsVisible.GetValue());
 		}
 
 		if (Mesh.IsSet())
 		{
-			JsonWriter.WriteValue(TEXT("mesh"), Mesh.GetValue());
+			Writer.Write(TEXT("mesh"), Mesh.GetValue());
 		}
 
-		FGLTFJsonUtility::WriteObjectArray(JsonWriter, TEXT("materials"), Materials, Extensions);
+		if (Materials.Num() > 0)
+		{
+			Writer.Write(TEXT("materials"), Materials);
+		}
 
-		JsonWriter.WriteObjectEnd();
-		JsonWriter.WriteObjectEnd();
+		Writer.EndObject();
 	}
 };
 
-struct FGLTFJsonVariant
+struct FGLTFJsonVariant : IGLTFJsonObject
 {
-	typedef TMap<FGLTFJsonNodeIndex, FGLTFJsonVariantNodeProperties> FNodeProperties;
-
 	FString Name;
 	bool    bIsActive;
 
 	FGLTFJsonTextureIndex Thumbnail;
-	FNodeProperties       Nodes;
+	TMap<FGLTFJsonNodeIndex, FGLTFJsonVariantNodeProperties> Nodes;
 
-	template <class CharType = TCHAR, class PrintPolicy = TPrettyJsonPrintPolicy<CharType>>
-	void WriteObject(TJsonWriter<CharType, PrintPolicy>& JsonWriter, FGLTFJsonExtensions& Extensions) const
+	virtual void WriteObject(IGLTFJsonWriter& Writer) const override
 	{
-		JsonWriter.WriteObjectStart();
-		JsonWriter.WriteValue(TEXT("name"), Name);
-		JsonWriter.WriteValue(TEXT("active"), bIsActive);
+		Writer.Write(TEXT("name"), Name);
+		Writer.Write(TEXT("active"), bIsActive);
 
 		if (Thumbnail != INDEX_NONE)
 		{
-			JsonWriter.WriteValue(TEXT("thumbnail"), Thumbnail);
-		}
-
-		JsonWriter.WriteArrayStart(TEXT("nodes"));
-
-		for (const auto& DataPair: Nodes)
-		{
-			if (DataPair.Key != INDEX_NONE)
-			{
-				DataPair.Value.WriteObject(JsonWriter, Extensions);
-			}
+			Writer.Write(TEXT("thumbnail"), Thumbnail);
 		}
 
-		JsonWriter.WriteArrayEnd();
-		JsonWriter.WriteObjectEnd();
+		TArray<FGLTFJsonVariantNodeProperties> NodeValues;
+		Nodes.GenerateValueArray(NodeValues);
+		Writer.Write(TEXT("nodes"), NodeValues);
 	}
 };
 
-struct FGLTFJsonVariantSet
+struct FGLTFJsonVariantSet : IGLTFJsonObject
 {
 	FString Name;
 
 	TArray<FGLTFJsonVariant> Variants;
 
-	template <class CharType = TCHAR, class PrintPolicy = TPrettyJsonPrintPolicy<CharType>>
-	void WriteObject(TJsonWriter<CharType, PrintPolicy>& JsonWriter, FGLTFJsonExtensions& Extensions) const
+	virtual void WriteObject(IGLTFJsonWriter& Writer) const override
 	{
-		JsonWriter.WriteObjectStart();
-
 		if (!Name.IsEmpty())
 		{
-			JsonWriter.WriteValue(TEXT("name"), Name);
+			Writer.Write(TEXT("name"), Name);
 		}
 
-		FGLTFJsonUtility::WriteObjectArray(JsonWriter, TEXT("variants"), Variants, Extensions);
-
-		JsonWriter.WriteObjectEnd();
+		Writer.Write(TEXT("variants"), Variants);
 	}
 };
 
-struct FGLTFJsonLevelVariantSets
+struct FGLTFJsonLevelVariantSets : IGLTFJsonObject
 {
 	FString Name;
 
 	TArray<FGLTFJsonVariantSet> VariantSets;
 
-	template <class CharType = TCHAR, class PrintPolicy = TPrettyJsonPrintPolicy<CharType>>
-	void WriteObject(TJsonWriter<CharType, PrintPolicy>& JsonWriter, FGLTFJsonExtensions& Extensions) const
+	virtual void WriteObject(IGLTFJsonWriter& Writer) const override
 	{
-		JsonWriter.WriteObjectStart();
-
 		if (!Name.IsEmpty())
 		{
-			JsonWriter.WriteValue(TEXT("name"), Name);
+			Writer.Write(TEXT("name"), Name);
 		}
 
-		FGLTFJsonUtility::WriteObjectArray(JsonWriter, TEXT("variantSets"), VariantSets, Extensions);
-
-		JsonWriter.WriteObjectEnd();
+		Writer.Write(TEXT("variantSets"), VariantSets);
 	}
 };
diff --git a/Engine/Plugins/Enterprise/GLTFExporter/MigratedFromGit/Source/GLTFExporter/Private/Json/GLTFJsonLight.h b/Engine/Plugins/Enterprise/GLTFExporter/MigratedFromGit/Source/GLTFExporter/Private/Json/GLTFJsonLight.h
--- a/Engine/Plugins/Enterprise/GLTFExporter/MigratedFromGit/Source/GLTFExporter/Private/Json/GLTFJsonLight.h
+++ b/Engine/Plugins/Enterprise/GLTFExporter/MigratedFromGit/Source/GLTFExporter/Private/Json/GLTFJsonLight.h
@@ -2,12 +2,11 @@
 
 #pragma once
 
+#include "Json/GLTFJsonObject.h"
 #include "Json/GLTFJsonEnums.h"
 #include "Json/GLTFJsonColor3.h"
-#include "Json/GLTFJsonUtility.h"
-#include "Serialization/JsonSerializer.h"
 
-struct FGLTFJsonSpotLight
+struct FGLTFJsonSpotLight : IGLTFJsonObject
 {
 	float InnerConeAngle;
 	float OuterConeAngle;
@@ -18,26 +17,21 @@ struct FGLTFJsonSpotLight
 	{
 	}
 
-	template <class CharType = TCHAR, class PrintPolicy = TPrettyJsonPrintPolicy<CharType>>
-	void WriteObject(TJsonWriter<CharType, PrintPolicy>& JsonWriter, FGLTFJsonExtensions& Extensions) const
+	virtual void WriteObject(IGLTFJsonWriter& Writer) const override
 	{
-		JsonWriter.WriteObjectStart();
-
 		if (InnerConeAngle != 0)
 		{
-			FGLTFJsonUtility::WriteExactValue(JsonWriter, TEXT("innerConeAngle"), InnerConeAngle);
+			Writer.Write(TEXT("innerConeAngle"), InnerConeAngle);
 		}
 
 		if (OuterConeAngle != PI / 2.0f)
 		{
-			FGLTFJsonUtility::WriteExactValue(JsonWriter, TEXT("outerConeAngle"), OuterConeAngle);
+			Writer.Write(TEXT("outerConeAngle"), OuterConeAngle);
 		}
-
-		JsonWriter.WriteObjectEnd();
 	}
 };
 
-struct FGLTFJsonLight
+struct FGLTFJsonLight : IGLTFJsonObject
 {
 	FString Name;
 
@@ -58,43 +52,36 @@ struct FGLTFJsonLight
 	{
 	}
 
-	template <class CharType = TCHAR, class PrintPolicy = TPrettyJsonPrintPolicy<CharType>>
-	void WriteObject(TJsonWriter<CharType, PrintPolicy>& JsonWriter, FGLTFJsonExtensions& Extensions) const
+	virtual void WriteObject(IGLTFJsonWriter& Writer) const override
 	{
-		JsonWriter.WriteObjectStart();
-
 		if (!Name.IsEmpty())
 		{
-			JsonWriter.WriteValue(TEXT("name"), Name);
+			Writer.Write(TEXT("name"), Name);
 		}
 
-		JsonWriter.WriteValue(TEXT("type"), FGLTFJsonUtility::ToString(Type));
+		Writer.Write(TEXT("type"), Type);
 
 		if (Color != FGLTFJsonColor3::White)
 		{
-			JsonWriter.WriteIdentifierPrefix(TEXT("color"));
-			Color.WriteArray(JsonWriter);
+			Writer.Write(TEXT("color"), Color);
 		}
 
 		if (Intensity != 1)
 		{
-			FGLTFJsonUtility::WriteExactValue(JsonWriter, TEXT("intensity"), Intensity);
+			Writer.Write(TEXT("intensity"), Intensity);
 		}
 
 		if (Type == EGLTFJsonLightType::Point || Type == EGLTFJsonLightType::Spot)
 		{
 			if (Range != 0)
 			{
-				FGLTFJsonUtility::WriteExactValue(JsonWriter, TEXT("range"), Range);
+				Writer.Write(TEXT("range"), Range);
 			}
 
 			if (Type == EGLTFJsonLightType::Spot)
 			{
-				JsonWriter.WriteIdentifierPrefix(TEXT("spot"));
-				Spot.WriteObject(JsonWriter, Extensions);
+				Writer.Write(TEXT("spot"), Spot);
 			}
 		}
-
-		JsonWriter.WriteObjectEnd();
 	}
 };
diff --git a/Engine/Plugins/Enterprise/GLTFExporter/MigratedFromGit/Source/GLTFExporter/Private/Json/GLTFJsonLightMap.h b/Engine/Plugins/Enterprise/GLTFExporter/MigratedFromGit/Source/GLTFExporter/Private/Json/GLTFJsonLightMap.h
--- a/Engine/Plugins/Enterprise/GLTFExporter/MigratedFromGit/Source/GLTFExporter/Private/Json/GLTFJsonLightMap.h
+++ b/Engine/Plugins/Enterprise/GLTFExporter/MigratedFromGit/Source/GLTFExporter/Private/Json/GLTFJsonLightMap.h
@@ -2,10 +2,11 @@
 
 #pragma once
 
+#include "Json/GLTFJsonObject.h"
 #include "Json/GLTFJsonMaterial.h"
 #include "Json/GLTFJsonVector4.h"
 
-struct FGLTFJsonLightMap
+struct FGLTFJsonLightMap : IGLTFJsonObject
 {
 	FString              Name;
 	FGLTFJsonTextureInfo Texture;
@@ -16,35 +17,25 @@ struct FGLTFJsonLightMap
 	FGLTFJsonLightMap()
 		: LightMapScale(FGLTFJsonVector4::One)
 		, LightMapAdd(FGLTFJsonVector4::Zero)
-		, CoordinateScaleBias({ 1, 1, 0, 0 })
+		, CoordinateScaleBias(1, 1, 0, 0)
 	{
 	}
 
-	template <class CharType = TCHAR, class PrintPolicy = TPrettyJsonPrintPolicy<CharType>>
-	void WriteObject(TJsonWriter<CharType, PrintPolicy>& JsonWriter, FGLTFJsonExtensions& Extensions) const
+	virtual void WriteObject(IGLTFJsonWriter& Writer) const override
 	{
-		JsonWriter.WriteObjectStart();
-
 		if (!Name.IsEmpty())
 		{
-			JsonWriter.WriteValue(TEXT("name"), Name);
+			Writer.Write(TEXT("name"), Name);
 		}
 
 		if (Texture.Index != INDEX_NONE)
 		{
-			JsonWriter.WriteIdentifierPrefix(TEXT("texture"));
-			Texture.WriteObject(JsonWriter, Extensions);
+			Writer.Write(TEXT("texture"), Texture);
 		}
 
-		JsonWriter.WriteIdentifierPrefix(TEXT("lightmapScale"));
-		LightMapScale.WriteArray(JsonWriter);
-
-		JsonWriter.WriteIdentifierPrefix(TEXT("lightmapAdd"));
-		LightMapAdd.WriteArray(JsonWriter);
-
-		JsonWriter.WriteIdentifierPrefix(TEXT("coordinateScaleBias"));
-		CoordinateScaleBias.WriteArray(JsonWriter);
+		Writer.Write(TEXT("lightmapScale"), LightMapScale);
+		Writer.Write(TEXT("lightmapAdd"), LightMapAdd);
+		Writer.Write(TEXT("coordinateScaleBias"), CoordinateScaleBias);
 
-		JsonWriter.WriteObjectEnd();
 	}
 };
diff --git a/Engine/Plugins/Enterprise/GLTFExporter/MigratedFromGit/Source/GLTFExporter/Private/Json/GLTFJsonMaterial.h b/Engine/Plugins/Enterprise/GLTFExporter/MigratedFromGit/Source/GLTFExporter/Private/Json/GLTFJsonMaterial.h
--- a/Engine/Plugins/Enterprise/GLTFExporter/MigratedFromGit/Source/GLTFExporter/Private/Json/GLTFJsonMaterial.h
+++ b/Engine/Plugins/Enterprise/GLTFExporter/MigratedFromGit/Source/GLTFExporter/Private/Json/GLTFJsonMaterial.h
@@ -2,14 +2,13 @@
 
 #pragma once
 
+#include "Json/GLTFJsonObject.h"
 #include "Json/GLTFJsonEnums.h"
 #include "Json/GLTFJsonColor3.h"
 #include "Json/GLTFJsonColor4.h"
 #include "Json/GLTFJsonTextureTransform.h"
-#include "Json/GLTFJsonUtility.h"
-#include "Serialization/JsonSerializer.h"
 
-struct FGLTFJsonTextureInfo
+struct FGLTFJsonTextureInfo : IGLTFJsonObject
 {
 	FGLTFJsonTextureIndex Index;
 	int32 TexCoord;
@@ -21,30 +20,21 @@ struct FGLTFJsonTextureInfo
 	{
 	}
 
-	template <class CharType = TCHAR, class PrintPolicy = TPrettyJsonPrintPolicy<CharType>>
-	void WriteObject(TJsonWriter<CharType, PrintPolicy>& JsonWriter, FGLTFJsonExtensions& Extensions) const
+	virtual void WriteObject(IGLTFJsonWriter& Writer) const override
 	{
-		JsonWriter.WriteObjectStart();
-
-		JsonWriter.WriteValue(TEXT("index"), Index);
+		Writer.Write(TEXT("index"), Index);
 
 		if (TexCoord != 0)
 		{
-			JsonWriter.WriteValue(TEXT("texCoord"), TexCoord);
+			Writer.Write(TEXT("texCoord"), TexCoord);
 		}
 
 		if (Transform != FGLTFJsonTextureTransform())
 		{
-			const EGLTFJsonExtension Extension = EGLTFJsonExtension::KHR_TextureTransform;
-			Extensions.Used.Add(Extension);
-
-			JsonWriter.WriteObjectStart(TEXT("extensions"));
-			JsonWriter.WriteIdentifierPrefix(FGLTFJsonUtility::ToString(Extension));
-			Transform.WriteObject(JsonWriter, Extensions);
-			JsonWriter.WriteObjectEnd();
+			Writer.StartExtensions();
+			Writer.Write(EGLTFJsonExtension::KHR_TextureTransform, Transform);
+			Writer.EndExtensions();
 		}
-
-		JsonWriter.WriteObjectEnd();
 	}
 };
 
@@ -57,24 +47,19 @@ struct FGLTFJsonNormalTextureInfo : FGLTFJsonTextureInfo
 	{
 	}
 
-	template <class CharType = TCHAR, class PrintPolicy = TPrettyJsonPrintPolicy<CharType>>
-	void WriteObject(TJsonWriter<CharType, PrintPolicy>& JsonWriter, FGLTFJsonExtensions& Extensions) const
+	virtual void WriteObject(IGLTFJsonWriter& Writer) const override
 	{
-		JsonWriter.WriteObjectStart();
-
-		JsonWriter.WriteValue(TEXT("index"), Index);
+		Writer.Write(TEXT("index"), Index);
 
 		if (TexCoord != 0)
 		{
-			JsonWriter.WriteValue(TEXT("texCoord"), TexCoord);
+			Writer.Write(TEXT("texCoord"), TexCoord);
 		}
 
 		if (Scale != 1)
 		{
-			FGLTFJsonUtility::WriteExactValue(JsonWriter, TEXT("scale"), Scale);
+			Writer.Write(TEXT("scale"), Scale);
 		}
-
-		JsonWriter.WriteObjectEnd();
 	}
 };
 
@@ -87,28 +72,23 @@ struct FGLTFJsonOcclusionTextureInfo : FGLTFJsonTextureInfo
 	{
 	}
 
-	template <class CharType = TCHAR, class PrintPolicy = TPrettyJsonPrintPolicy<CharType>>
-	void WriteObject(TJsonWriter<CharType, PrintPolicy>& JsonWriter, FGLTFJsonExtensions& Extensions) const
+	virtual void WriteObject(IGLTFJsonWriter& Writer) const override
 	{
-		JsonWriter.WriteObjectStart();
-
-		JsonWriter.WriteValue(TEXT("index"), Index);
+		Writer.Write(TEXT("index"), Index);
 
 		if (TexCoord != 0)
 		{
-			JsonWriter.WriteValue(TEXT("texCoord"), TexCoord);
+			Writer.Write(TEXT("texCoord"), TexCoord);
 		}
 
 		if (Strength != 1)
 		{
-			FGLTFJsonUtility::WriteExactValue(JsonWriter, TEXT("strength"), Strength);
+			Writer.Write(TEXT("strength"), Strength);
 		}
-
-		JsonWriter.WriteObjectEnd();
 	}
 };
 
-struct FGLTFJsonPBRMetallicRoughness
+struct FGLTFJsonPBRMetallicRoughness : IGLTFJsonObject
 {
 	FGLTFJsonColor4 BaseColorFactor;
 	FGLTFJsonTextureInfo BaseColorTexture;
@@ -124,44 +104,36 @@ struct FGLTFJsonPBRMetallicRoughness
 	{
 	}
 
-	template <class CharType = TCHAR, class PrintPolicy = TPrettyJsonPrintPolicy<CharType>>
-	void WriteObject(TJsonWriter<CharType, PrintPolicy>& JsonWriter, FGLTFJsonExtensions& Extensions) const
+	virtual void WriteObject(IGLTFJsonWriter& Writer) const override
 	{
-		JsonWriter.WriteObjectStart();
-
 		if (BaseColorFactor != FGLTFJsonColor4::White)
 		{
-			JsonWriter.WriteIdentifierPrefix(TEXT("baseColorFactor"));
-			BaseColorFactor.WriteArray(JsonWriter);
+			Writer.Write(TEXT("baseColorFactor"), BaseColorFactor);
 		}
 
 		if (BaseColorTexture.Index != INDEX_NONE)
 		{
-			JsonWriter.WriteIdentifierPrefix(TEXT("baseColorTexture"));
-			BaseColorTexture.WriteObject(JsonWriter, Extensions);
+			Writer.Write(TEXT("baseColorTexture"), BaseColorTexture);
 		}
 
 		if (MetallicFactor != 1)
 		{
-			FGLTFJsonUtility::WriteExactValue(JsonWriter, TEXT("metallicFactor"), MetallicFactor);
+			Writer.Write(TEXT("metallicFactor"), MetallicFactor);
 		}
 
 		if (RoughnessFactor != 1)
 		{
-			FGLTFJsonUtility::WriteExactValue(JsonWriter, TEXT("roughnessFactor"), RoughnessFactor);
+			Writer.Write(TEXT("roughnessFactor"), RoughnessFactor);
 		}
 
 		if (MetallicRoughnessTexture.Index != INDEX_NONE)
 		{
-			JsonWriter.WriteIdentifierPrefix(TEXT("metallicRoughnessTexture"));
-			MetallicRoughnessTexture.WriteObject(JsonWriter, Extensions);
+			Writer.Write(TEXT("metallicRoughnessTexture"), MetallicRoughnessTexture);
 		}
-
-		JsonWriter.WriteObjectEnd();
 	}
 };
 
-struct FGLTFJsonClearCoatExtension
+struct FGLTFJsonClearCoatExtension : IGLTFJsonObject
 {
 	float ClearCoatFactor;
 	FGLTFJsonTextureInfo ClearCoatTexture;
@@ -177,44 +149,36 @@ struct FGLTFJsonClearCoatExtension
 	{
 	}
 
-	template <class CharType = TCHAR, class PrintPolicy = TPrettyJsonPrintPolicy<CharType>>
-	void WriteObject(TJsonWriter<CharType, PrintPolicy>& JsonWriter, FGLTFJsonExtensions& Extensions) const
+	virtual void WriteObject(IGLTFJsonWriter& Writer) const override
 	{
-		JsonWriter.WriteObjectStart();
-
 		if (ClearCoatFactor != 0)
 		{
-			FGLTFJsonUtility::WriteExactValue(JsonWriter, TEXT("clearcoatFactor"), ClearCoatFactor);
+			Writer.Write(TEXT("clearcoatFactor"), ClearCoatFactor);
 		}
 
 		if (ClearCoatTexture.Index != INDEX_NONE)
 		{
-			JsonWriter.WriteIdentifierPrefix(TEXT("clearcoatTexture"));
-			ClearCoatTexture.WriteObject(JsonWriter, Extensions);
+			Writer.Write(TEXT("clearcoatTexture"), ClearCoatTexture);
 		}
 
 		if (ClearCoatRoughnessFactor != 0)
 		{
-			FGLTFJsonUtility::WriteExactValue(JsonWriter, TEXT("clearcoatRoughnessFactor"), ClearCoatRoughnessFactor);
+			Writer.Write(TEXT("clearcoatRoughnessFactor"), ClearCoatRoughnessFactor);
 		}
 
 		if (ClearCoatRoughnessTexture.Index != INDEX_NONE)
 		{
-			JsonWriter.WriteIdentifierPrefix(TEXT("clearcoatRoughnessTexture"));
-			ClearCoatRoughnessTexture.WriteObject(JsonWriter, Extensions);
+			Writer.Write(TEXT("clearcoatRoughnessTexture"), ClearCoatRoughnessTexture);
 		}
 
 		if (ClearCoatNormalTexture.Index != INDEX_NONE)
 		{
-			JsonWriter.WriteIdentifierPrefix(TEXT("clearcoatNormalTexture"));
-			ClearCoatNormalTexture.WriteObject(JsonWriter, Extensions);
+			Writer.Write(TEXT("clearcoatNormalTexture"), ClearCoatNormalTexture);
 		}
-
-		JsonWriter.WriteObjectEnd();
 	}
 };
 
-struct FGLTFJsonMaterial
+struct FGLTFJsonMaterial : IGLTFJsonObject
 {
 	FString Name;
 
@@ -247,97 +211,76 @@ struct FGLTFJsonMaterial
 	{
 	}
 
-	template <class CharType = TCHAR, class PrintPolicy = TPrettyJsonPrintPolicy<CharType>>
-	void WriteObject(TJsonWriter<CharType, PrintPolicy>& JsonWriter, FGLTFJsonExtensions& Extensions) const
+	virtual void WriteObject(IGLTFJsonWriter& Writer) const override
 	{
-		JsonWriter.WriteObjectStart();
-
 		if (!Name.IsEmpty())
 		{
-			JsonWriter.WriteValue(TEXT("name"), Name);
+			Writer.Write(TEXT("name"), Name);
 		}
 
 		if (ShadingModel != EGLTFJsonShadingModel::None)
 		{
-			JsonWriter.WriteIdentifierPrefix(TEXT("pbrMetallicRoughness"));
-			PBRMetallicRoughness.WriteObject(JsonWriter, Extensions);
+			Writer.Write(TEXT("pbrMetallicRoughness"), PBRMetallicRoughness);
 		}
 
 		if (NormalTexture.Index != INDEX_NONE)
 		{
-			JsonWriter.WriteIdentifierPrefix(TEXT("normalTexture"));
-			NormalTexture.WriteObject(JsonWriter, Extensions);
+			Writer.Write(TEXT("normalTexture"), NormalTexture);
 		}
 
 		if (OcclusionTexture.Index != INDEX_NONE)
 		{
-			JsonWriter.WriteIdentifierPrefix(TEXT("occlusionTexture"));
-			OcclusionTexture.WriteObject(JsonWriter, Extensions);
+			Writer.Write(TEXT("occlusionTexture"), OcclusionTexture);
 		}
 
 		if (EmissiveTexture.Index != INDEX_NONE)
 		{
-			JsonWriter.WriteIdentifierPrefix(TEXT("emissiveTexture"));
-			EmissiveTexture.WriteObject(JsonWriter, Extensions);
+			Writer.Write(TEXT("emissiveTexture"), EmissiveTexture);
 		}
 
 		if (EmissiveFactor != FGLTFJsonColor3::Black)
 		{
-			JsonWriter.WriteIdentifierPrefix(TEXT("emissiveFactor"));
-			EmissiveFactor.WriteArray(JsonWriter);
+			Writer.Write(TEXT("emissiveFactor"), EmissiveFactor);
 		}
 
 		if (AlphaMode != EGLTFJsonAlphaMode::Opaque)
 		{
-			JsonWriter.WriteValue(TEXT("alphaMode"), FGLTFJsonUtility::ToString(AlphaMode));
+			Writer.Write(TEXT("alphaMode"), AlphaMode);
 		}
 
 		if (AlphaMode == EGLTFJsonAlphaMode::Mask && AlphaCutoff != 0.5f)
 		{
-			FGLTFJsonUtility::WriteExactValue(JsonWriter, TEXT("alphaCutoff"), AlphaCutoff);
+			Writer.Write(TEXT("alphaCutoff"), AlphaCutoff);
 		}
 
 		if (DoubleSided)
 		{
-			JsonWriter.WriteValue(TEXT("doubleSided"), DoubleSided);
+			Writer.Write(TEXT("doubleSided"), DoubleSided);
 		}
 
 		if (BlendMode != EGLTFJsonBlendMode::None || ShadingModel == EGLTFJsonShadingModel::Unlit || ShadingModel == EGLTFJsonShadingModel::ClearCoat)
 		{
-			JsonWriter.WriteObjectStart(TEXT("extensions"));
+			Writer.StartExtensions();
 
 			if (BlendMode != EGLTFJsonBlendMode::None)
 			{
-				const EGLTFJsonExtension Extension = EGLTFJsonExtension::EPIC_BlendModes;
-				Extensions.Used.Add(Extension);
-
-				JsonWriter.WriteIdentifierPrefix(FGLTFJsonUtility::ToString(Extension));
-				JsonWriter.WriteObjectStart();
-				JsonWriter.WriteValue(TEXT("blendMode"), FGLTFJsonUtility::ToString(BlendMode));
-				JsonWriter.WriteObjectEnd();
+				Writer.StartExtension(EGLTFJsonExtension::EPIC_BlendModes);
+				Writer.Write(TEXT("blendMode"), BlendMode);
+				Writer.EndExtension();
 			}
 
 			if (ShadingModel == EGLTFJsonShadingModel::Unlit)
 			{
-				const EGLTFJsonExtension Extension = EGLTFJsonExtension::KHR_MaterialsUnlit;
-				Extensions.Used.Add(Extension);
-
-				JsonWriter.WriteIdentifierPrefix(FGLTFJsonUtility::ToString(Extension));
-				JsonWriter.WriteObjectStart(); // Write empty object
-				JsonWriter.WriteObjectEnd();
+				Writer.StartExtension(EGLTFJsonExtension::KHR_MaterialsUnlit);
+				// Write empty object
+				Writer.EndExtension();
 			}
 			else if (ShadingModel == EGLTFJsonShadingModel::ClearCoat)
 			{
-				const EGLTFJsonExtension Extension = EGLTFJsonExtension::KHR_MaterialsClearCoat;
-				Extensions.Used.Add(Extension);
-
-				JsonWriter.WriteIdentifierPrefix(FGLTFJsonUtility::ToString(Extension));
-				ClearCoat.WriteObject(JsonWriter, Extensions);
+				Writer.Write(EGLTFJsonExtension::KHR_MaterialsClearCoat, ClearCoat);
 			}
 
-			JsonWriter.WriteObjectEnd();
+			Writer.EndExtensions();
 		}
-
-		JsonWriter.WriteObjectEnd();
 	}
 };
diff --git a/Engine/Plugins/Enterprise/GLTFExporter/MigratedFromGit/Source/GLTFExporter/Private/Json/GLTFJsonMatrix4.cpp b/Engine/Plugins/Enterprise/GLTFExporter/MigratedFromGit/Source/GLTFExporter/Private/Json/GLTFJsonMatrix4.cpp
--- a/Engine/Plugins/Enterprise/GLTFExporter/MigratedFromGit/Source/GLTFExporter/Private/Json/GLTFJsonMatrix4.cpp
+++ b/Engine/Plugins/Enterprise/GLTFExporter/MigratedFromGit/Source/GLTFExporter/Private/Json/GLTFJsonMatrix4.cpp
@@ -2,10 +2,10 @@
 
 #include "Json/GLTFJsonMatrix4.h"
 
-const FGLTFJsonMatrix4 FGLTFJsonMatrix4::Identity =
+const FGLTFJsonMatrix4 FGLTFJsonMatrix4::Identity(
 {
 	1, 0 ,0, 0,
 	0, 1, 0, 0,
 	0, 0, 1, 0,
 	0, 0, 0, 1
-};
+});
diff --git a/Engine/Plugins/Enterprise/GLTFExporter/MigratedFromGit/Source/GLTFExporter/Private/Json/GLTFJsonMatrix4.h b/Engine/Plugins/Enterprise/GLTFExporter/MigratedFromGit/Source/GLTFExporter/Private/Json/GLTFJsonMatrix4.h
--- a/Engine/Plugins/Enterprise/GLTFExporter/MigratedFromGit/Source/GLTFExporter/Private/Json/GLTFJsonMatrix4.h
+++ b/Engine/Plugins/Enterprise/GLTFExporter/MigratedFromGit/Source/GLTFExporter/Private/Json/GLTFJsonMatrix4.h
@@ -2,38 +2,31 @@
 
 #pragma once
 
-#include "Json/GLTFJsonUtility.h"
-#include "Serialization/JsonSerializer.h"
+#include "Json/GLTFJsonArray.h"
+#include "Converters/GLTFRawTypes.h"
 
-struct FGLTFJsonMatrix4
+struct FGLTFJsonMatrix4 : IGLTFJsonArray, FGLTFRawMatrix4
 {
-	float M[16];
-
 	static const FGLTFJsonMatrix4 Identity;
 
-	template <class CharType = TCHAR, class PrintPolicy = TPrettyJsonPrintPolicy<CharType>>
-	void WriteArray(TJsonWriter<CharType, PrintPolicy>& JsonWriter) const
+	FGLTFJsonMatrix4(const FGLTFRawMatrix4& Raw)
 	{
-		JsonWriter.WriteArrayStart();
-
-		for(int32 I = 0; I < 16; ++I)
-		{
-			FGLTFJsonUtility::WriteExactValue(JsonWriter, M[I]);
-		}
-
-		JsonWriter.WriteArrayEnd();
+		*static_cast<FGLTFRawMatrix4*>(this) = Raw;
 	}
 
-	float& operator()(int32 Row, int32 Col)
+	virtual void WriteArray(IGLTFJsonWriter& Writer) const override
 	{
-		return M[Row * 4 + Col];
+		for (int32 i = 0; i < 16; ++i)
+		{
+			Writer.Write(Cells[i]);
+		}
 	}
 
 	bool operator==(const FGLTFJsonMatrix4& Other) const
 	{
-		for(int32 I = 0; I < 16; ++I)
+		for (int32 i = 0; i < 16; ++i)
 		{
-			if (M[I] != Other.M[I])
+			if (Cells[i] != Other.Cells[i])
 			{
 				return false;
 			}
diff --git a/Engine/Plugins/Enterprise/GLTFExporter/MigratedFromGit/Source/GLTFExporter/Private/Json/GLTFJsonMesh.h b/Engine/Plugins/Enterprise/GLTFExporter/MigratedFromGit/Source/GLTFExporter/Private/Json/GLTFJsonMesh.h
--- a/Engine/Plugins/Enterprise/GLTFExporter/MigratedFromGit/Source/GLTFExporter/Private/Json/GLTFJsonMesh.h
+++ b/Engine/Plugins/Enterprise/GLTFExporter/MigratedFromGit/Source/GLTFExporter/Private/Json/GLTFJsonMesh.h
@@ -2,11 +2,11 @@
 
 #pragma once
 
+#include "Json/GLTFJsonObject.h"
 #include "Json/GLTFJsonEnums.h"
 #include "Json/GLTFJsonIndex.h"
-#include "Serialization/JsonSerializer.h"
 
-struct FGLTFJsonAttributes
+struct FGLTFJsonAttributes : IGLTFJsonObject
 {
 	FGLTFJsonAccessorIndex Position; // always required
 	FGLTFJsonAccessorIndex Color0;
@@ -18,37 +18,32 @@ struct FGLTFJsonAttributes
 	TArray<FGLTFJsonAccessorIndex> Joints;
 	TArray<FGLTFJsonAccessorIndex> Weights;
 
-	template <class CharType = TCHAR, class PrintPolicy = TPrettyJsonPrintPolicy<CharType>>
-	void WriteObject(TJsonWriter<CharType, PrintPolicy>& JsonWriter, FGLTFJsonExtensions& Extensions) const
+	virtual void WriteObject(IGLTFJsonWriter& Writer) const override
 	{
-		JsonWriter.WriteObjectStart();
+		Writer.Write(TEXT("POSITION"), Position);
 
-		JsonWriter.WriteValue(TEXT("POSITION"), Position);
-
-		if (Color0 != INDEX_NONE) JsonWriter.WriteValue(TEXT("COLOR_0"), Color0);
-		if (Normal != INDEX_NONE) JsonWriter.WriteValue(TEXT("NORMAL"), Normal);
-		if (Tangent != INDEX_NONE) JsonWriter.WriteValue(TEXT("TANGENT"), Tangent);
+		if (Color0 != INDEX_NONE) Writer.Write(TEXT("COLOR_0"), Color0);
+		if (Normal != INDEX_NONE) Writer.Write(TEXT("NORMAL"), Normal);
+		if (Tangent != INDEX_NONE) Writer.Write(TEXT("TANGENT"), Tangent);
 
 		for (int32 Index = 0; Index < TexCoords.Num(); ++Index)
 		{
-			JsonWriter.WriteValue(TEXT("TEXCOORD_") + FString::FromInt(Index), TexCoords[Index]);
+			Writer.Write(TEXT("TEXCOORD_") + FString::FromInt(Index), TexCoords[Index]);
 		}
 
 		for (int32 Index = 0; Index < Joints.Num(); ++Index)
 		{
-			JsonWriter.WriteValue(TEXT("JOINTS_") + FString::FromInt(Index), Joints[Index]);
+			Writer.Write(TEXT("JOINTS_") + FString::FromInt(Index), Joints[Index]);
 		}
 
 		for (int32 Index = 0; Index < Weights.Num(); ++Index)
 		{
-			JsonWriter.WriteValue(TEXT("WEIGHTS_") + FString::FromInt(Index), Weights[Index]);
+			Writer.Write(TEXT("WEIGHTS_") + FString::FromInt(Index), Weights[Index]);
 		}
-
-		JsonWriter.WriteObjectEnd();
 	}
 };
 
-struct FGLTFJsonPrimitive
+struct FGLTFJsonPrimitive : IGLTFJsonObject
 {
 	FGLTFJsonAccessorIndex Indices;
 	FGLTFJsonMaterialIndex Material;
@@ -60,51 +55,40 @@ struct FGLTFJsonPrimitive
 	{
 	}
 
-	template <class CharType = TCHAR, class PrintPolicy = TPrettyJsonPrintPolicy<CharType>>
-	void WriteObject(TJsonWriter<CharType, PrintPolicy>& JsonWriter, FGLTFJsonExtensions& Extensions) const
+	virtual void WriteObject(IGLTFJsonWriter& Writer) const override
 	{
-		JsonWriter.WriteObjectStart();
-
-		JsonWriter.WriteIdentifierPrefix(TEXT("attributes"));
-		Attributes.WriteObject(JsonWriter, Extensions);
+		Writer.Write(TEXT("attributes"), Attributes);
 
 		if (Indices != INDEX_NONE)
 		{
-			JsonWriter.WriteValue(TEXT("indices"), Indices);
+			Writer.Write(TEXT("indices"), Indices);
 		}
 
 		if (Material != INDEX_NONE)
 		{
-			JsonWriter.WriteValue(TEXT("material"), Material);
+			Writer.Write(TEXT("material"), Material);
 		}
 
 		if (Mode != EGLTFJsonPrimitiveMode::Triangles)
 		{
-			JsonWriter.WriteValue(TEXT("mode"), FGLTFJsonUtility::ToInteger(Mode));
+			Writer.Write(TEXT("mode"), Mode);
 		}
-
-		JsonWriter.WriteObjectEnd();
 	}
 };
 
-struct FGLTFJsonMesh
+struct FGLTFJsonMesh : IGLTFJsonObject
 {
 	FString Name;
 
 	TArray<FGLTFJsonPrimitive> Primitives;
 
-	template <class CharType = TCHAR, class PrintPolicy = TPrettyJsonPrintPolicy<CharType>>
-	void WriteObject(TJsonWriter<CharType, PrintPolicy>& JsonWriter, FGLTFJsonExtensions& Extensions) const
+	virtual void WriteObject(IGLTFJsonWriter& Writer) const override
 	{
-		JsonWriter.WriteObjectStart();
-
 		if (!Name.IsEmpty())
 		{
-			JsonWriter.WriteValue(TEXT("name"), Name);
+			Writer.Write(TEXT("name"), Name);
 		}
 
-		FGLTFJsonUtility::WriteObjectArray(JsonWriter, TEXT("primitives"), Primitives, Extensions, true);
-
-		JsonWriter.WriteObjectEnd();
+		Writer.Write(TEXT("primitives"), Primitives);
 	}
 };
diff --git a/Engine/Plugins/Enterprise/GLTFExporter/MigratedFromGit/Source/GLTFExporter/Private/Json/GLTFJsonNode.h b/Engine/Plugins/Enterprise/GLTFExporter/MigratedFromGit/Source/GLTFExporter/Private/Json/GLTFJsonNode.h
--- a/Engine/Plugins/Enterprise/GLTFExporter/MigratedFromGit/Source/GLTFExporter/Private/Json/GLTFJsonNode.h
+++ b/Engine/Plugins/Enterprise/GLTFExporter/MigratedFromGit/Source/GLTFExporter/Private/Json/GLTFJsonNode.h
@@ -2,29 +2,22 @@
 
 #pragma once
 
+#include "Json/GLTFJsonObject.h"
 #include "Json/GLTFJsonIndex.h"
 #include "Json/GLTFJsonMatrix4.h"
 #include "Json/GLTFJsonVector3.h"
 #include "Json/GLTFJsonQuaternion.h"
-#include "Serialization/JsonSerializer.h"
 
-struct FGLTFJsonNode
+struct FGLTFJsonNode : IGLTFJsonObject
 {
 	FString Name;
 
 	bool bUseMatrix;
 
-	union
-	{
-		FGLTFJsonMatrix4 Matrix;
-
-		struct
-		{
-			FGLTFJsonVector3    Translation;
-			FGLTFJsonQuaternion Rotation;
-			FGLTFJsonVector3    Scale;
-		};
-	};
+	FGLTFJsonMatrix4    Matrix;
+	FGLTFJsonVector3    Translation;
+	FGLTFJsonQuaternion Rotation;
+	FGLTFJsonVector3    Scale;
 
 	FGLTFJsonCameraIndex    Camera;
 	FGLTFJsonSkinIndex      Skin;
@@ -41,128 +34,105 @@ struct FGLTFJsonNode
 
 	FGLTFJsonNode()
 		: bUseMatrix(false)
+		, Matrix(FGLTFJsonMatrix4::Identity)
 		, Translation(FGLTFJsonVector3::Zero)
 		, Rotation(FGLTFJsonQuaternion::Identity)
 		, Scale(FGLTFJsonVector3::One)
 	{
 	}
 
-	template <class CharType = TCHAR, class PrintPolicy = TPrettyJsonPrintPolicy<CharType>>
-	void WriteObject(TJsonWriter<CharType, PrintPolicy>& JsonWriter, FGLTFJsonExtensions& Extensions) const
+	virtual void WriteObject(IGLTFJsonWriter& Writer) const override
 	{
-		JsonWriter.WriteObjectStart();
-
 		if (!Name.IsEmpty())
 		{
-			JsonWriter.WriteValue(TEXT("name"), Name);
+			Writer.Write(TEXT("name"), Name);
 		}
 
 		if (bUseMatrix)
 		{
 			if (Matrix != FGLTFJsonMatrix4::Identity)
 			{
-				JsonWriter.WriteIdentifierPrefix(TEXT("matrix"));
-				Matrix.WriteArray(JsonWriter);
+				Writer.Write(TEXT("matrix"), Matrix);
 			}
 		}
 		else
 		{
 			if (Translation != FGLTFJsonVector3::Zero)
 			{
-				JsonWriter.WriteIdentifierPrefix(TEXT("translation"));
-				Translation.WriteArray(JsonWriter);
+				Writer.Write(TEXT("translation"), Translation);
 			}
 
 			if (Rotation != FGLTFJsonQuaternion::Identity)
 			{
-				JsonWriter.WriteIdentifierPrefix(TEXT("rotation"));
-				Rotation.WriteArray(JsonWriter);
+				Writer.Write(TEXT("rotation"), Rotation);
 			}
 
 			if (Scale != FGLTFJsonVector3::One)
 			{
-				JsonWriter.WriteIdentifierPrefix(TEXT("scale"));
-				Scale.WriteArray(JsonWriter);
+				Writer.Write(TEXT("scale"), Scale);
 			}
 		}
 
 		if (Camera != INDEX_NONE)
 		{
-			JsonWriter.WriteValue(TEXT("camera"), Camera);
+			Writer.Write(TEXT("camera"), Camera);
 		}
 
 		if (Skin != INDEX_NONE)
 		{
-			JsonWriter.WriteValue(TEXT("skin"), Skin);
+			Writer.Write(TEXT("skin"), Skin);
 		}
 
 		if (Mesh != INDEX_NONE)
 		{
-			JsonWriter.WriteValue(TEXT("mesh"), Mesh);
+			Writer.Write(TEXT("mesh"), Mesh);
 		}
 
 		if (Backdrop != INDEX_NONE || Hotspot != INDEX_NONE || Light != INDEX_NONE || LightMap != INDEX_NONE || SkySphere != INDEX_NONE)
 		{
-			JsonWriter.WriteObjectStart(TEXT("extensions"));
+			Writer.StartExtensions();
 
 			if (Backdrop != INDEX_NONE)
 			{
-				const EGLTFJsonExtension Extension = EGLTFJsonExtension::EPIC_HDRIBackdrops;
-				Extensions.Used.Add(Extension);
-
-				JsonWriter.WriteObjectStart(FGLTFJsonUtility::ToString(Extension));
-				JsonWriter.WriteValue(TEXT("backdrop"), Backdrop);
-				JsonWriter.WriteObjectEnd();
+				Writer.StartExtension(EGLTFJsonExtension::EPIC_HDRIBackdrops);
+				Writer.Write(TEXT("backdrop"), Backdrop);
+				Writer.EndExtension();
 			}
 
 			if (Hotspot != INDEX_NONE)
 			{
-				const EGLTFJsonExtension Extension = EGLTFJsonExtension::EPIC_AnimationHotspots;
-				Extensions.Used.Add(Extension);
-
-				JsonWriter.WriteObjectStart(FGLTFJsonUtility::ToString(Extension));
-				JsonWriter.WriteValue(TEXT("hotspot"), Hotspot);
-				JsonWriter.WriteObjectEnd();
+				Writer.StartExtension(EGLTFJsonExtension::EPIC_AnimationHotspots);
+				Writer.Write(TEXT("hotspot"), Hotspot);
+				Writer.EndExtension();
 			}
 
 			if (Light != INDEX_NONE)
 			{
-				const EGLTFJsonExtension Extension = EGLTFJsonExtension::KHR_LightsPunctual;
-				Extensions.Used.Add(Extension);
-
-				JsonWriter.WriteObjectStart(FGLTFJsonUtility::ToString(Extension));
-				JsonWriter.WriteValue(TEXT("light"), Light);
-				JsonWriter.WriteObjectEnd();
+				Writer.StartExtension(EGLTFJsonExtension::KHR_LightsPunctual);
+				Writer.Write(TEXT("light"), Light);
+				Writer.EndExtension();
 			}
 
 			if (LightMap != INDEX_NONE)
 			{
-				const EGLTFJsonExtension Extension = EGLTFJsonExtension::EPIC_LightmapTextures;
-				Extensions.Used.Add(Extension);
-
-				JsonWriter.WriteObjectStart(FGLTFJsonUtility::ToString(Extension));
-				JsonWriter.WriteValue(TEXT("lightmap"), LightMap);
-				JsonWriter.WriteObjectEnd();
+				Writer.StartExtension(EGLTFJsonExtension::EPIC_LightmapTextures);
+				Writer.Write(TEXT("lightmap"), LightMap);
+				Writer.EndExtension();
 			}
 
 			if (SkySphere != INDEX_NONE)
 			{
-				const EGLTFJsonExtension Extension = EGLTFJsonExtension::EPIC_SkySpheres;
-				Extensions.Used.Add(Extension);
-
-				JsonWriter.WriteObjectStart(FGLTFJsonUtility::ToString(Extension));
-				JsonWriter.WriteValue(TEXT("skySphere"), SkySphere);
-				JsonWriter.WriteObjectEnd();
+				Writer.StartExtension(EGLTFJsonExtension::EPIC_SkySpheres);
+				Writer.Write(TEXT("skySphere"), SkySphere);
+				Writer.EndExtension();
 			}
 
-			JsonWriter.WriteObjectEnd();
+			Writer.EndExtensions();
 		}
 
 		if (Children.Num() > 0)
 		{
-			JsonWriter.WriteValue(TEXT("children"), Children);
+			Writer.Write(TEXT("children"), Children);
 		}
-
-		JsonWriter.WriteObjectEnd();
 	}
 };
diff --git a/Engine/Plugins/Enterprise/GLTFExporter/MigratedFromGit/Source/GLTFExporter/Private/Json/GLTFJsonObject.h b/Engine/Plugins/Enterprise/GLTFExporter/MigratedFromGit/Source/GLTFExporter/Private/Json/GLTFJsonObject.h
new file mode 100644
--- /dev/null
+++ b/Engine/Plugins/Enterprise/GLTFExporter/MigratedFromGit/Source/GLTFExporter/Private/Json/GLTFJsonObject.h
@@ -0,0 +1,18 @@
+// Copyright Epic Games, Inc. All Rights Reserved.
+
+#pragma once
+
+#include "Json/GLTFJsonValue.h"
+#include "Json/GLTFJsonWriter.h"
+
+struct IGLTFJsonObject : IGLTFJsonValue
+{
+	virtual void WriteValue(IGLTFJsonWriter& Writer) const override final
+	{
+		Writer.StartObject();
+		WriteObject(Writer);
+		Writer.EndObject();
+	}
+
+	virtual void WriteObject(IGLTFJsonWriter& Writer) const = 0;
+};
diff --git a/Engine/Plugins/Enterprise/GLTFExporter/MigratedFromGit/Source/GLTFExporter/Private/Json/GLTFJsonQuaternion.h b/Engine/Plugins/Enterprise/GLTFExporter/MigratedFromGit/Source/GLTFExporter/Private/Json/GLTFJsonQuaternion.h
--- a/Engine/Plugins/Enterprise/GLTFExporter/MigratedFromGit/Source/GLTFExporter/Private/Json/GLTFJsonQuaternion.h
+++ b/Engine/Plugins/Enterprise/GLTFExporter/MigratedFromGit/Source/GLTFExporter/Private/Json/GLTFJsonQuaternion.h
@@ -2,10 +2,10 @@
 
 #pragma once
 
-#include "Json/GLTFJsonUtility.h"
-#include "Serialization/JsonSerializer.h"
+#include "Json/GLTFJsonArray.h"
+#include "Converters/GLTFRawTypes.h"
 
-struct FGLTFJsonQuaternion
+struct FGLTFJsonQuaternion : IGLTFJsonArray
 {
 	float X, Y, Z, W;
 
@@ -16,15 +16,17 @@ struct FGLTFJsonQuaternion
 	{
 	}
 
-	template <class CharType = TCHAR, class PrintPolicy = TPrettyJsonPrintPolicy<CharType>>
-	void WriteArray(TJsonWriter<CharType, PrintPolicy>& JsonWriter) const
+	FGLTFJsonQuaternion(const FGLTFRawQuaternion& Raw)
+		: X(Raw.X), Y(Raw.Y), Z(Raw.Z), W(Raw.W)
 	{
-		JsonWriter.WriteArrayStart();
-		FGLTFJsonUtility::WriteExactValue(JsonWriter, X);
-		FGLTFJsonUtility::WriteExactValue(JsonWriter, Y);
-		FGLTFJsonUtility::WriteExactValue(JsonWriter, Z);
-		FGLTFJsonUtility::WriteExactValue(JsonWriter, W);
-		JsonWriter.WriteArrayEnd();
+	}
+
+	virtual void WriteArray(IGLTFJsonWriter& Writer) const override
+	{
+		Writer.Write(X);
+		Writer.Write(Y);
+		Writer.Write(Z);
+		Writer.Write(W);
 	}
 
 	bool operator==(const FGLTFJsonQuaternion& Other) const
diff --git a/Engine/Plugins/Enterprise/GLTFExporter/MigratedFromGit/Source/GLTFExporter/Private/Json/GLTFJsonRoot.h b/Engine/Plugins/Enterprise/GLTFExporter/MigratedFromGit/Source/GLTFExporter/Private/Json/GLTFJsonRoot.h
--- a/Engine/Plugins/Enterprise/GLTFExporter/MigratedFromGit/Source/GLTFExporter/Private/Json/GLTFJsonRoot.h
+++ b/Engine/Plugins/Enterprise/GLTFExporter/MigratedFromGit/Source/GLTFExporter/Private/Json/GLTFJsonRoot.h
@@ -2,9 +2,9 @@
 
 #pragma once
 
+#include "Json/GLTFJsonObject.h"
 #include "Json/GLTFJsonIndex.h"
 #include "Json/GLTFJsonAsset.h"
-#include "Json/GLTFJsonExtensions.h"
 #include "Json/GLTFJsonAccessor.h"
 #include "Json/GLTFJsonAnimation.h"
 #include "Json/GLTFJsonBuffer.h"
@@ -24,9 +24,9 @@
 #include "Json/GLTFJsonLightMap.h"
 #include "Json/GLTFJsonSkySphere.h"
 #include "Json/GLTFJsonLevelVariantSets.h"
-#include "Policies/CondensedJsonPrintPolicy.h"
+#include "Json/GLTFJsonWriter.h"
 
-struct FGLTFJsonRoot
+struct FGLTFJsonRoot : IGLTFJsonObject
 {
 	FGLTFJsonAsset Asset;
 
@@ -54,123 +54,88 @@ struct FGLTFJsonRoot
 	TArray<TUniquePtr<FGLTFJsonSkySphere>>  SkySpheres;
 	TArray<TUniquePtr<FGLTFJsonLevelVariantSets>>  LevelVariantSets;
 
-	template <class CharType = TCHAR, class PrintPolicy = TPrettyJsonPrintPolicy<CharType>>
-	void WriteObject(TJsonWriter<CharType, PrintPolicy>& JsonWriter)
+	virtual void WriteObject(IGLTFJsonWriter& Writer) const override
 	{
-		JsonWriter.WriteObjectStart();
-
-		JsonWriter.WriteIdentifierPrefix(TEXT("asset"));
-		Asset.WriteObject(JsonWriter, Extensions);
+		Writer.Write(TEXT("asset"), Asset);
 
 		if (DefaultScene != INDEX_NONE)
 		{
-			JsonWriter.WriteValue(TEXT("scene"), DefaultScene);
+			Writer.Write(TEXT("scene"), DefaultScene);
 		}
 
-		FGLTFJsonUtility::WriteObjectPtrArray(JsonWriter, TEXT("accessors"), Accessors, Extensions);
-		FGLTFJsonUtility::WriteObjectPtrArray(JsonWriter, TEXT("animations"), Animations, Extensions);
-		FGLTFJsonUtility::WriteObjectPtrArray(JsonWriter, TEXT("buffers"), Buffers, Extensions);
-		FGLTFJsonUtility::WriteObjectPtrArray(JsonWriter, TEXT("bufferViews"), BufferViews, Extensions);
-		FGLTFJsonUtility::WriteObjectPtrArray(JsonWriter, TEXT("cameras"), Cameras, Extensions);
-		FGLTFJsonUtility::WriteObjectPtrArray(JsonWriter, TEXT("images"), Images, Extensions);
-		FGLTFJsonUtility::WriteObjectPtrArray(JsonWriter, TEXT("materials"), Materials, Extensions);
-		FGLTFJsonUtility::WriteObjectPtrArray(JsonWriter, TEXT("meshes"), Meshes, Extensions);
-		FGLTFJsonUtility::WriteObjectPtrArray(JsonWriter, TEXT("nodes"), Nodes, Extensions);
-		FGLTFJsonUtility::WriteObjectPtrArray(JsonWriter, TEXT("samplers"), Samplers, Extensions);
-		FGLTFJsonUtility::WriteObjectPtrArray(JsonWriter, TEXT("scenes"), Scenes, Extensions);
-		FGLTFJsonUtility::WriteObjectPtrArray(JsonWriter, TEXT("skins"), Skins, Extensions);
-		FGLTFJsonUtility::WriteObjectPtrArray(JsonWriter, TEXT("textures"), Textures, Extensions);
+		if (Accessors.Num() > 0) Writer.Write(TEXT("accessors"), Accessors);
+		if (Animations.Num() > 0) Writer.Write(TEXT("animations"), Animations);
+		if (Buffers.Num() > 0) Writer.Write(TEXT("buffers"), Buffers);
+		if (BufferViews.Num() > 0) Writer.Write(TEXT("bufferViews"), BufferViews);
+		if (Cameras.Num() > 0) Writer.Write(TEXT("cameras"), Cameras);
+		if (Images.Num() > 0) Writer.Write(TEXT("images"), Images);
+		if (Materials.Num() > 0) Writer.Write(TEXT("materials"), Materials);
+		if (Meshes.Num() > 0) Writer.Write(TEXT("meshes"), Meshes);
+		if (Nodes.Num() > 0) Writer.Write(TEXT("nodes"), Nodes);
+		if (Samplers.Num() > 0) Writer.Write(TEXT("samplers"), Samplers);
+		if (Scenes.Num() > 0) Writer.Write(TEXT("scenes"), Scenes);
+		if (Skins.Num() > 0) Writer.Write(TEXT("skins"), Skins);
+		if (Textures.Num() > 0) Writer.Write(TEXT("textures"), Textures);
 
 		if (Backdrops.Num() > 0 || Hotspots.Num() > 0 || Lights.Num() > 0 || LightMaps.Num() > 0 || SkySpheres.Num() > 0 || LevelVariantSets.Num() > 0)
 		{
-			JsonWriter.WriteObjectStart(TEXT("extensions"));
+			Writer.StartExtensions();
 
 			if (Backdrops.Num() > 0)
 			{
-				const EGLTFJsonExtension Extension = EGLTFJsonExtension::EPIC_HDRIBackdrops;
-				Extensions.Used.Add(Extension);
-
-				JsonWriter.WriteObjectStart(FGLTFJsonUtility::ToString(Extension));
-				FGLTFJsonUtility::WriteObjectPtrArray(JsonWriter, TEXT("backdrops"), Backdrops, Extensions);
-				JsonWriter.WriteObjectEnd();
+				Writer.StartExtension(EGLTFJsonExtension::EPIC_HDRIBackdrops);
+				Writer.Write(TEXT("backdrops"), Backdrops);
+				Writer.EndExtension();
 			}
 
 			if (Hotspots.Num() > 0)
 			{
-				const EGLTFJsonExtension Extension = EGLTFJsonExtension::EPIC_AnimationHotspots;
-				Extensions.Used.Add(Extension);
-
-				JsonWriter.WriteObjectStart(FGLTFJsonUtility::ToString(Extension));
-				FGLTFJsonUtility::WriteObjectPtrArray(JsonWriter, TEXT("hotspots"), Hotspots, Extensions);
-				JsonWriter.WriteObjectEnd();
+				Writer.StartExtension(EGLTFJsonExtension::EPIC_AnimationHotspots);
+				Writer.Write(TEXT("hotspots"), Hotspots);
+				Writer.EndExtension();
 			}
 
 			if (LevelVariantSets.Num() > 0)
 			{
-				const EGLTFJsonExtension Extension = EGLTFJsonExtension::EPIC_LevelVariantSets;
-				Extensions.Used.Add(Extension);
-
-				JsonWriter.WriteObjectStart(FGLTFJsonUtility::ToString(Extension));
-				FGLTFJsonUtility::WriteObjectPtrArray(JsonWriter, TEXT("levelVariantSets"), LevelVariantSets, Extensions);
-				JsonWriter.WriteObjectEnd();
+				Writer.StartExtension(EGLTFJsonExtension::EPIC_LevelVariantSets);
+				Writer.Write(TEXT("levelVariantSets"), LevelVariantSets);
+				Writer.EndExtension();
 			}
 
 			if (Lights.Num() > 0)
 			{
-				const EGLTFJsonExtension Extension = EGLTFJsonExtension::KHR_LightsPunctual;
-				Extensions.Used.Add(Extension);
-
-				JsonWriter.WriteObjectStart(FGLTFJsonUtility::ToString(Extension));
-				FGLTFJsonUtility::WriteObjectPtrArray(JsonWriter, TEXT("lights"), Lights, Extensions);
-				JsonWriter.WriteObjectEnd();
+				Writer.StartExtension(EGLTFJsonExtension::KHR_LightsPunctual);
+				Writer.Write(TEXT("lights"), Lights);
+				Writer.EndExtension();
 			}
 
 			if (LightMaps.Num() > 0)
 			{
-				const EGLTFJsonExtension Extension = EGLTFJsonExtension::EPIC_LightmapTextures;
-				Extensions.Used.Add(Extension);
-
-				JsonWriter.WriteObjectStart(FGLTFJsonUtility::ToString(Extension));
-				FGLTFJsonUtility::WriteObjectPtrArray(JsonWriter, TEXT("lightmaps"), LightMaps, Extensions);
-				JsonWriter.WriteObjectEnd();
+				Writer.StartExtension(EGLTFJsonExtension::EPIC_LightmapTextures);
+				Writer.Write(TEXT("lightmaps"), LightMaps);
+				Writer.EndExtension();
 			}
 
 			if (SkySpheres.Num() > 0)
 			{
-				const EGLTFJsonExtension Extension = EGLTFJsonExtension::EPIC_SkySpheres;
-				Extensions.Used.Add(Extension);
-
-				JsonWriter.WriteObjectStart(FGLTFJsonUtility::ToString(Extension));
-				FGLTFJsonUtility::WriteObjectPtrArray(JsonWriter, TEXT("skySpheres"), SkySpheres, Extensions);
-				JsonWriter.WriteObjectEnd();
+				Writer.StartExtension(EGLTFJsonExtension::EPIC_SkySpheres);
+				Writer.Write(TEXT("skySpheres"), SkySpheres);
+				Writer.EndExtension();
 			}
 
-			JsonWriter.WriteObjectEnd();
+			Writer.EndExtensions();
 		}
 
-		FGLTFJsonUtility::WriteStringArray(JsonWriter, TEXT("extensionsUsed"), Extensions.Used);
-		FGLTFJsonUtility::WriteStringArray(JsonWriter, TEXT("extensionsRequired"), Extensions.Required);
-
-		JsonWriter.WriteObjectEnd();
+		if (Writer.Extensions.Used.Num() > 0) Writer.Write(TEXT("extensionsUsed"), Writer.Extensions.Used);
+		if (Writer.Extensions.Required.Num() > 0) Writer.Write(TEXT("extensionsRequired"), Writer.Extensions.Required);
 	}
 
-	template <class CharType = TCHAR, class PrintPolicy = TPrettyJsonPrintPolicy<CharType>>
-	void ToJson(FArchive* const Archive)
+	void ToJson(FArchive* const Archive, bool bPrettyJson)
 	{
-		TSharedRef<TJsonWriter<CharType, PrintPolicy>> JsonWriter = TJsonWriterFactory<CharType, PrintPolicy>::Create(Archive);
-		WriteObject(*JsonWriter);
-		JsonWriter->Close();
-	}
-
-	void ToJson(FArchive* const Archive, bool bCondensedPrint)
-	{
-		if (bCondensedPrint)
-		{
-			ToJson<UTF8CHAR, TCondensedJsonPrintPolicy<UTF8CHAR>>(Archive);
-		}
-		else
-		{
-			ToJson<UTF8CHAR, TPrettyJsonPrintPolicy<UTF8CHAR>>(Archive);
-		}
+		TSharedRef<IGLTFJsonWriter> Writer = IGLTFJsonWriter::Create(Archive, bPrettyJson);
+		Writer->Extensions = Extensions;
+		Writer->Write(this);
+		Writer->Close();
+		Extensions = Writer->Extensions; // TODO: fix ugly hack
 	}
 };
diff --git a/Engine/Plugins/Enterprise/GLTFExporter/MigratedFromGit/Source/GLTFExporter/Private/Json/GLTFJsonSampler.h b/Engine/Plugins/Enterprise/GLTFExporter/MigratedFromGit/Source/GLTFExporter/Private/Json/GLTFJsonSampler.h
--- a/Engine/Plugins/Enterprise/GLTFExporter/MigratedFromGit/Source/GLTFExporter/Private/Json/GLTFJsonSampler.h
+++ b/Engine/Plugins/Enterprise/GLTFExporter/MigratedFromGit/Source/GLTFExporter/Private/Json/GLTFJsonSampler.h
@@ -2,11 +2,10 @@
 
 #pragma once
 
+#include "Json/GLTFJsonObject.h"
 #include "Json/GLTFJsonEnums.h"
-#include "Json/GLTFJsonUtility.h"
-#include "Serialization/JsonSerializer.h"
 
-struct FGLTFJsonSampler
+struct FGLTFJsonSampler : IGLTFJsonObject
 {
 	FString Name;
 
@@ -24,36 +23,31 @@ struct FGLTFJsonSampler
 	{
 	}
 
-	template <class CharType = TCHAR, class PrintPolicy = TPrettyJsonPrintPolicy<CharType>>
-	void WriteObject(TJsonWriter<CharType, PrintPolicy>& JsonWriter, FGLTFJsonExtensions& Extensions) const
+	virtual void WriteObject(IGLTFJsonWriter& Writer) const override
 	{
-		JsonWriter.WriteObjectStart();
-
 		if (!Name.IsEmpty())
 		{
-			JsonWriter.WriteValue(TEXT("name"), Name);
+			Writer.Write(TEXT("name"), Name);
 		}
 
 		if (MinFilter != EGLTFJsonTextureFilter::None)
 		{
-			JsonWriter.WriteValue(TEXT("minFilter"), FGLTFJsonUtility::ToInteger(MinFilter));
+			Writer.Write(TEXT("minFilter"), MinFilter);
 		}
 
 		if (MagFilter != EGLTFJsonTextureFilter::None)
 		{
-			JsonWriter.WriteValue(TEXT("magFilter"), FGLTFJsonUtility::ToInteger(MagFilter));
+			Writer.Write(TEXT("magFilter"), MagFilter);
 		}
 
 		if (WrapS != EGLTFJsonTextureWrap::Repeat)
 		{
-			JsonWriter.WriteValue(TEXT("wrapS"), FGLTFJsonUtility::ToInteger(WrapS));
+			Writer.Write(TEXT("wrapS"), WrapS);
 		}
 
 		if (WrapT != EGLTFJsonTextureWrap::Repeat)
 		{
-			JsonWriter.WriteValue(TEXT("wrapT"), FGLTFJsonUtility::ToInteger(WrapT));
+			Writer.Write(TEXT("wrapT"), WrapT);
 		}
-
-		JsonWriter.WriteObjectEnd();
 	}
 };
diff --git a/Engine/Plugins/Enterprise/GLTFExporter/MigratedFromGit/Source/GLTFExporter/Private/Json/GLTFJsonScene.h b/Engine/Plugins/Enterprise/GLTFExporter/MigratedFromGit/Source/GLTFExporter/Private/Json/GLTFJsonScene.h
--- a/Engine/Plugins/Enterprise/GLTFExporter/MigratedFromGit/Source/GLTFExporter/Private/Json/GLTFJsonScene.h
+++ b/Engine/Plugins/Enterprise/GLTFExporter/MigratedFromGit/Source/GLTFExporter/Private/Json/GLTFJsonScene.h
@@ -2,45 +2,37 @@
 
 #pragma once
 
+#include "Json/GLTFJsonObject.h"
 #include "Json/GLTFJsonIndex.h"
-#include "Serialization/JsonSerializer.h"
 
-struct FGLTFJsonScene
+struct FGLTFJsonScene : IGLTFJsonObject
 {
 	FString Name;
 
 	TArray<FGLTFJsonNodeIndex> Nodes;
 	TArray<FGLTFJsonLevelVariantSetsIndex>  LevelVariantSets;
 
-	template <class CharType = TCHAR, class PrintPolicy = TPrettyJsonPrintPolicy<CharType>>
-	void WriteObject(TJsonWriter<CharType, PrintPolicy>& JsonWriter, FGLTFJsonExtensions& Extensions) const
+	virtual void WriteObject(IGLTFJsonWriter& Writer) const override
 	{
-		JsonWriter.WriteObjectStart();
-
 		if (!Name.IsEmpty())
 		{
-			JsonWriter.WriteValue(TEXT("name"), Name);
+			Writer.Write(TEXT("name"), Name);
 		}
 
 		if (Nodes.Num() > 0)
 		{
-			JsonWriter.WriteValue(TEXT("nodes"), Nodes);
+			Writer.Write(TEXT("nodes"), Nodes);
 		}
 
 		if (LevelVariantSets.Num() > 0)
 		{
-			const EGLTFJsonExtension Extension = EGLTFJsonExtension::EPIC_LevelVariantSets;
-			Extensions.Used.Add(Extension);
-
-			JsonWriter.WriteObjectStart(TEXT("extensions"));
-			JsonWriter.WriteObjectStart(FGLTFJsonUtility::ToString(Extension));
+			Writer.StartExtensions();
 
-			JsonWriter.WriteValue(TEXT("levelVariantSets"), LevelVariantSets);
+			Writer.StartExtension(EGLTFJsonExtension::EPIC_LevelVariantSets);
+			Writer.Write(TEXT("levelVariantSets"), LevelVariantSets);
+			Writer.EndExtension();
 
-			JsonWriter.WriteObjectEnd();
-			JsonWriter.WriteObjectEnd();
+			Writer.EndExtensions();
 		}
-
-		JsonWriter.WriteObjectEnd();
 	}
 };
diff --git a/Engine/Plugins/Enterprise/GLTFExporter/MigratedFromGit/Source/GLTFExporter/Private/Json/GLTFJsonSkin.h b/Engine/Plugins/Enterprise/GLTFExporter/MigratedFromGit/Source/GLTFExporter/Private/Json/GLTFJsonSkin.h
--- a/Engine/Plugins/Enterprise/GLTFExporter/MigratedFromGit/Source/GLTFExporter/Private/Json/GLTFJsonSkin.h
+++ b/Engine/Plugins/Enterprise/GLTFExporter/MigratedFromGit/Source/GLTFExporter/Private/Json/GLTFJsonSkin.h
@@ -2,10 +2,10 @@
 
 #pragma once
 
+#include "Json/GLTFJsonObject.h"
 #include "Json/GLTFJsonIndex.h"
-#include "Serialization/JsonSerializer.h"
 
-struct FGLTFJsonSkin
+struct FGLTFJsonSkin : IGLTFJsonObject
 {
 	FString Name;
 
@@ -14,31 +14,26 @@ struct FGLTFJsonSkin
 
 	TArray<FGLTFJsonNodeIndex> Joints;
 
-	template <class CharType = TCHAR, class PrintPolicy = TPrettyJsonPrintPolicy<CharType>>
-	void WriteObject(TJsonWriter<CharType, PrintPolicy>& JsonWriter, FGLTFJsonExtensions& Extensions) const
+	virtual void WriteObject(IGLTFJsonWriter& Writer) const override
 	{
-		JsonWriter.WriteObjectStart();
-
 		if (!Name.IsEmpty())
 		{
-			JsonWriter.WriteValue(TEXT("name"), Name);
+			Writer.Write(TEXT("name"), Name);
 		}
 
 		if (InverseBindMatrices != INDEX_NONE)
 		{
-			JsonWriter.WriteValue(TEXT("inverseBindMatrices"), InverseBindMatrices);
+			Writer.Write(TEXT("inverseBindMatrices"), InverseBindMatrices);
 		}
 
 		if (Skeleton != INDEX_NONE)
 		{
-			JsonWriter.WriteValue(TEXT("skeleton"), Skeleton);
+			Writer.Write(TEXT("skeleton"), Skeleton);
 		}
 
 		if (Joints.Num() > 0)
 		{
-			JsonWriter.WriteValue(TEXT("joints"), Joints);
+			Writer.Write(TEXT("joints"), Joints);
 		}
-
-		JsonWriter.WriteObjectEnd();
 	}
 };
diff --git a/Engine/Plugins/Enterprise/GLTFExporter/MigratedFromGit/Source/GLTFExporter/Private/Json/GLTFJsonSkySphere.h b/Engine/Plugins/Enterprise/GLTFExporter/MigratedFromGit/Source/GLTFExporter/Private/Json/GLTFJsonSkySphere.h
--- a/Engine/Plugins/Enterprise/GLTFExporter/MigratedFromGit/Source/GLTFExporter/Private/Json/GLTFJsonSkySphere.h
+++ b/Engine/Plugins/Enterprise/GLTFExporter/MigratedFromGit/Source/GLTFExporter/Private/Json/GLTFJsonSkySphere.h
@@ -2,12 +2,11 @@
 
 #pragma once
 
+#include "Json/GLTFJsonObject.h"
 #include "Json/GLTFJsonIndex.h"
 #include "Json/GLTFJsonColor4.h"
-#include "Json/GLTFJsonExtensions.h"
-#include "Serialization/JsonSerializer.h"
 
-struct FGLTFJsonSkySphereColorCurve
+struct FGLTFJsonSkySphereColorCurve : IGLTFJsonArray
 {
 	struct FKey
 	{
@@ -22,31 +21,26 @@ struct FGLTFJsonSkySphereColorCurve
 
 	TArray<FComponentCurve> ComponentCurves;
 
-	template <class CharType = TCHAR, class PrintPolicy = TPrettyJsonPrintPolicy<CharType>>
-	void WriteArray(TJsonWriter<CharType, PrintPolicy>& JsonWriter) const
+	virtual void WriteArray(IGLTFJsonWriter& Writer) const override
 	{
-		JsonWriter.WriteArrayStart();
-
 		for (int32 ComponentIndex = 0; ComponentIndex < ComponentCurves.Num(); ++ComponentIndex)
 		{
 			const FComponentCurve& ComponentCurve = ComponentCurves[ComponentIndex];
 
-			JsonWriter.WriteArrayStart();
+			Writer.StartArray();
 
 			for (const FKey& Key: ComponentCurve.Keys)
 			{
-				FGLTFJsonUtility::WriteExactValue(JsonWriter, Key.Time);
-				FGLTFJsonUtility::WriteExactValue(JsonWriter, Key.Value);
+				Writer.Write(Key.Time);
+				Writer.Write(Key.Value);
 			}
 
-			JsonWriter.WriteArrayEnd();
+			Writer.EndArray();
 		}
-
-		JsonWriter.WriteArrayEnd();
 	}
 };
 
-struct FGLTFJsonSkySphere
+struct FGLTFJsonSkySphere : IGLTFJsonObject
 {
 	FString Name;
 
@@ -99,78 +93,63 @@ struct FGLTFJsonSkySphere
 	{
 	}
 
-	template <class CharType = TCHAR, class PrintPolicy = TPrettyJsonPrintPolicy<CharType>>
-	void WriteObject(TJsonWriter<CharType, PrintPolicy>& JsonWriter, FGLTFJsonExtensions& Extensions) const
+	virtual void WriteObject(IGLTFJsonWriter& Writer) const override
 	{
-		JsonWriter.WriteObjectStart();
-
 		if (!Name.IsEmpty())
 		{
-			JsonWriter.WriteValue(TEXT("name"), Name);
+			Writer.Write(TEXT("name"), Name);
 		}
 
-		JsonWriter.WriteValue(TEXT("skySphereMesh"), SkySphereMesh);
-		JsonWriter.WriteValue(TEXT("skyTexture"), SkyTexture);
-		JsonWriter.WriteValue(TEXT("cloudsTexture"), CloudsTexture);
-		JsonWriter.WriteValue(TEXT("starsTexture"), StarsTexture);
+		Writer.Write(TEXT("skySphereMesh"), SkySphereMesh);
+		Writer.Write(TEXT("skyTexture"), SkyTexture);
+		Writer.Write(TEXT("cloudsTexture"), CloudsTexture);
+		Writer.Write(TEXT("starsTexture"), StarsTexture);
 
 		if (DirectionalLight != INDEX_NONE)
 		{
-			JsonWriter.WriteValue(TEXT("directionalLight"), DirectionalLight);
+			Writer.Write(TEXT("directionalLight"), DirectionalLight);
 		}
 
-		FGLTFJsonUtility::WriteExactValue(JsonWriter, TEXT("sunHeight"), SunHeight);
-		FGLTFJsonUtility::WriteExactValue(JsonWriter, TEXT("sunBrightness"), SunBrightness);
-		FGLTFJsonUtility::WriteExactValue(JsonWriter, TEXT("starsBrightness"), StarsBrightness);
-		FGLTFJsonUtility::WriteExactValue(JsonWriter, TEXT("cloudSpeed"), CloudSpeed);
-		FGLTFJsonUtility::WriteExactValue(JsonWriter, TEXT("cloudOpacity"), CloudOpacity);
-		FGLTFJsonUtility::WriteExactValue(JsonWriter, TEXT("horizonFalloff"), HorizonFalloff);
+		Writer.Write(TEXT("sunHeight"), SunHeight);
+		Writer.Write(TEXT("sunBrightness"), SunBrightness);
+		Writer.Write(TEXT("starsBrightness"), StarsBrightness);
+		Writer.Write(TEXT("cloudSpeed"), CloudSpeed);
+		Writer.Write(TEXT("cloudOpacity"), CloudOpacity);
+		Writer.Write(TEXT("horizonFalloff"), HorizonFalloff);
 
-		FGLTFJsonUtility::WriteExactValue(JsonWriter, TEXT("sunRadius"), SunRadius);
-		FGLTFJsonUtility::WriteExactValue(JsonWriter, TEXT("noisePower1"), NoisePower1);
-		FGLTFJsonUtility::WriteExactValue(JsonWriter, TEXT("noisePower2"), NoisePower2);
+		Writer.Write(TEXT("sunRadius"), SunRadius);
+		Writer.Write(TEXT("noisePower1"), NoisePower1);
+		Writer.Write(TEXT("noisePower2"), NoisePower2);
 
-		JsonWriter.WriteValue(TEXT("colorsDeterminedBySunPosition"), bColorsDeterminedBySunPosition);
+		Writer.Write(TEXT("colorsDeterminedBySunPosition"), bColorsDeterminedBySunPosition);
 
-		JsonWriter.WriteIdentifierPrefix(TEXT("zenithColor"));
-		ZenithColor.WriteArray(JsonWriter);
-
-		JsonWriter.WriteIdentifierPrefix(TEXT("horizonColor"));
-		HorizonColor.WriteArray(JsonWriter);
-
-		JsonWriter.WriteIdentifierPrefix(TEXT("cloudColor"));
-		CloudColor.WriteArray(JsonWriter);
+		Writer.Write(TEXT("zenithColor"), ZenithColor);
+		Writer.Write(TEXT("horizonColor"), HorizonColor);
+		Writer.Write(TEXT("cloudColor"), CloudColor);
 
 		if (OverallColor != FGLTFJsonColor4::White)
 		{
-			JsonWriter.WriteIdentifierPrefix(TEXT("overallColor"));
-			OverallColor.WriteArray(JsonWriter);
+			Writer.Write(TEXT("overallColor"), OverallColor);
 		}
 
 		if (ZenithColorCurve.ComponentCurves.Num() >= 3)
 		{
-			JsonWriter.WriteIdentifierPrefix(TEXT("zenithColorCurve"));
-			ZenithColorCurve.WriteArray(JsonWriter);
+			Writer.Write(TEXT("zenithColorCurve"), ZenithColorCurve);
 		}
 
 		if (HorizonColorCurve.ComponentCurves.Num() >= 3)
 		{
-			JsonWriter.WriteIdentifierPrefix(TEXT("horizonColorCurve"));
-			HorizonColorCurve.WriteArray(JsonWriter);
+			Writer.Write(TEXT("horizonColorCurve"), HorizonColorCurve);
 		}
 
 		if (CloudColorCurve.ComponentCurves.Num() >= 3)
 		{
-			JsonWriter.WriteIdentifierPrefix(TEXT("cloudColorCurve"));
-			CloudColorCurve.WriteArray(JsonWriter);
+			Writer.Write(TEXT("cloudColorCurve"), CloudColorCurve);
 		}
 
 		if (Scale != FGLTFJsonVector3::One)
 		{
-			JsonWriter.WriteIdentifierPrefix(TEXT("scale"));
-			Scale.WriteArray(JsonWriter);
+			Writer.Write(TEXT("scale"), Scale);
 		}
-
-		JsonWriter.WriteObjectEnd();
 	}
 };
diff --git a/Engine/Plugins/Enterprise/GLTFExporter/MigratedFromGit/Source/GLTFExporter/Private/Json/GLTFJsonTexture.h b/Engine/Plugins/Enterprise/GLTFExporter/MigratedFromGit/Source/GLTFExporter/Private/Json/GLTFJsonTexture.h
--- a/Engine/Plugins/Enterprise/GLTFExporter/MigratedFromGit/Source/GLTFExporter/Private/Json/GLTFJsonTexture.h
+++ b/Engine/Plugins/Enterprise/GLTFExporter/MigratedFromGit/Source/GLTFExporter/Private/Json/GLTFJsonTexture.h
@@ -2,11 +2,10 @@
 
 #pragma once
 
+#include "Json/GLTFJsonObject.h"
 #include "Json/GLTFJsonIndex.h"
-#include "Json/GLTFJsonUtility.h"
-#include "Serialization/JsonSerializer.h"
 
-struct FGLTFJsonTexture
+struct FGLTFJsonTexture : IGLTFJsonObject
 {
 	FString Name;
 
@@ -21,40 +20,32 @@ struct FGLTFJsonTexture
 	{
 	}
 
-	template <class CharType = TCHAR, class PrintPolicy = TPrettyJsonPrintPolicy<CharType>>
-	void WriteObject(TJsonWriter<CharType, PrintPolicy>& JsonWriter, FGLTFJsonExtensions& Extensions) const
+	virtual void WriteObject(IGLTFJsonWriter& Writer) const override
 	{
-		JsonWriter.WriteObjectStart();
-
 		if (!Name.IsEmpty())
 		{
-			JsonWriter.WriteValue(TEXT("name"), Name);
+			Writer.Write(TEXT("name"), Name);
 		}
 
 		if (Sampler != INDEX_NONE)
 		{
-			JsonWriter.WriteValue(TEXT("sampler"), Sampler);
+			Writer.Write(TEXT("sampler"), Sampler);
 		}
 
 		if (Source != INDEX_NONE)
 		{
-			JsonWriter.WriteValue(TEXT("source"), Source);
+			Writer.Write(TEXT("source"), Source);
 		}
 
 		if (Encoding != EGLTFJsonHDREncoding::None)
 		{
-			const EGLTFJsonExtension Extension = EGLTFJsonExtension::EPIC_TextureHDREncoding;
-			Extensions.Used.Add(Extension);
-
-			JsonWriter.WriteObjectStart(TEXT("extensions"));
-			JsonWriter.WriteObjectStart(FGLTFJsonUtility::ToString(Extension));
+			Writer.StartExtensions();
 
-			JsonWriter.WriteValue(TEXT("encoding"), FGLTFJsonUtility::ToString(Encoding));
+			Writer.StartExtension(EGLTFJsonExtension::EPIC_TextureHDREncoding);
+			Writer.Write(TEXT("encoding"), Encoding);
+			Writer.EndExtension();
 
-			JsonWriter.WriteObjectEnd();
-			JsonWriter.WriteObjectEnd();
+			Writer.EndExtensions();
 		}
-
-		JsonWriter.WriteObjectEnd();
 	}
 };
diff --git a/Engine/Plugins/Enterprise/GLTFExporter/MigratedFromGit/Source/GLTFExporter/Private/Json/GLTFJsonTextureTransform.h b/Engine/Plugins/Enterprise/GLTFExporter/MigratedFromGit/Source/GLTFExporter/Private/Json/GLTFJsonTextureTransform.h
--- a/Engine/Plugins/Enterprise/GLTFExporter/MigratedFromGit/Source/GLTFExporter/Private/Json/GLTFJsonTextureTransform.h
+++ b/Engine/Plugins/Enterprise/GLTFExporter/MigratedFromGit/Source/GLTFExporter/Private/Json/GLTFJsonTextureTransform.h
@@ -2,11 +2,10 @@
 
 #pragma once
 
-#include "Json/GLTFJsonUtility.h"
+#include "Json/GLTFJsonObject.h"
 #include "Json/GLTFJsonVector2.h"
-#include "Serialization/JsonSerializer.h"
 
-struct FGLTFJsonTextureTransform
+struct FGLTFJsonTextureTransform : IGLTFJsonObject
 {
 	FGLTFJsonVector2 Offset;
 	FGLTFJsonVector2 Scale;
@@ -19,30 +18,22 @@ struct FGLTFJsonTextureTransform
 	{
 	}
 
-	template <class CharType = TCHAR, class PrintPolicy = TPrettyJsonPrintPolicy<CharType>>
-	void WriteObject(TJsonWriter<CharType, PrintPolicy>& JsonWriter, FGLTFJsonExtensions& Extensions) const
+	virtual void WriteObject(IGLTFJsonWriter& Writer) const override
 	{
-		JsonWriter.WriteObjectStart();
-
 		if (Offset != FGLTFJsonVector2::Zero)
 		{
-			JsonWriter.WriteIdentifierPrefix(TEXT("offset"));
-			Offset.WriteArray(JsonWriter);
+			Writer.Write(TEXT("offset"), Offset);
 		}
 
 		if (Scale != FGLTFJsonVector2::One)
 		{
-			JsonWriter.WriteIdentifierPrefix(TEXT("scale"));
-			Scale.WriteArray(JsonWriter);
+			Writer.Write(TEXT("scale"), Scale);
 		}
 
 		if (Rotation != 0)
 		{
-			JsonWriter.WriteIdentifierPrefix(TEXT("rotation"));
-			FGLTFJsonUtility::WriteExactValue(JsonWriter, Rotation);
+			Writer.Write(TEXT("rotation"), Rotation);
 		}
-
-		JsonWriter.WriteObjectEnd();
 	}
 
 	bool operator==(const FGLTFJsonTextureTransform& Other) const
diff --git a/Engine/Plugins/Enterprise/GLTFExporter/MigratedFromGit/Source/GLTFExporter/Private/Json/GLTFJsonUtility.h b/Engine/Plugins/Enterprise/GLTFExporter/MigratedFromGit/Source/GLTFExporter/Private/Json/GLTFJsonUtility.h
--- a/Engine/Plugins/Enterprise/GLTFExporter/MigratedFromGit/Source/GLTFExporter/Private/Json/GLTFJsonUtility.h
+++ b/Engine/Plugins/Enterprise/GLTFExporter/MigratedFromGit/Source/GLTFExporter/Private/Json/GLTFJsonUtility.h
@@ -3,19 +3,18 @@
 #pragma once
 
 #include "Json/GLTFJsonEnums.h"
-#include "Json/GLTFJsonExtensions.h"
 
 struct FGLTFJsonUtility
 {
-	template <typename EnumType>
-	static int32 ToInteger(EnumType Value)
+	template <typename EnumType, typename = typename TEnableIf<TIsEnum<EnumType>::Value>::Type>
+	static int32 GetValue(EnumType Enum)
 	{
-		return static_cast<int32>(Value);
+		return static_cast<int32>(Enum);
 	}
 
-	static const TCHAR* ToString(EGLTFJsonExtension Value)
+	static const TCHAR* GetValue(EGLTFJsonExtension Enum)
 	{
-		switch (Value)
+		switch (Enum)
 		{
 			case EGLTFJsonExtension::KHR_LightsPunctual:      return TEXT("KHR_lights_punctual");
 			case EGLTFJsonExtension::KHR_MaterialsClearCoat:  return TEXT("KHR_materials_clearcoat");
@@ -37,9 +36,9 @@ struct FGLTFJsonUtility
 		}
 	}
 
-	static const TCHAR* ToString(EGLTFJsonAlphaMode Value)
+	static const TCHAR* GetValue(EGLTFJsonAlphaMode Enum)
 	{
-		switch (Value)
+		switch (Enum)
 		{
 			case EGLTFJsonAlphaMode::Opaque: return TEXT("OPAQUE");
 			case EGLTFJsonAlphaMode::Blend:  return TEXT("BLEND");
@@ -50,9 +49,9 @@ struct FGLTFJsonUtility
 		}
 	}
 
-	static const TCHAR* ToString(EGLTFJsonBlendMode Value)
+	static const TCHAR* GetValue(EGLTFJsonBlendMode Enum)
 	{
-		switch (Value)
+		switch (Enum)
 		{
 			case EGLTFJsonBlendMode::Additive:       return TEXT("ADDITIVE");
 			case EGLTFJsonBlendMode::Modulate:       return TEXT("MODULATE");
@@ -63,9 +62,9 @@ struct FGLTFJsonUtility
 		}
 	}
 
-	static const TCHAR* ToString(EGLTFJsonMimeType Value)
+	static const TCHAR* GetValue(EGLTFJsonMimeType Enum)
 	{
-		switch (Value)
+		switch (Enum)
 		{
 			case EGLTFJsonMimeType::PNG:  return TEXT("image/png");
 			case EGLTFJsonMimeType::JPEG: return TEXT("image/jpeg");
@@ -75,9 +74,9 @@ struct FGLTFJsonUtility
 		}
 	}
 
-	static const TCHAR* ToString(EGLTFJsonAccessorType Value)
+	static const TCHAR* GetValue(EGLTFJsonAccessorType Enum)
 	{
-		switch (Value)
+		switch (Enum)
 		{
 			case EGLTFJsonAccessorType::Scalar: return TEXT("SCALAR");
 			case EGLTFJsonAccessorType::Vec2:   return TEXT("VEC2");
@@ -92,9 +91,9 @@ struct FGLTFJsonUtility
 		}
 	}
 
-	static const TCHAR* ToString(EGLTFJsonHDREncoding Value)
+	static const TCHAR* GetValue(EGLTFJsonHDREncoding Enum)
 	{
-		switch (Value)
+		switch (Enum)
 		{
 			case EGLTFJsonHDREncoding::RGBM: return TEXT("RGBM");
 			case EGLTFJsonHDREncoding::RGBE: return TEXT("RGBE");
@@ -104,9 +103,9 @@ struct FGLTFJsonUtility
 		}
 	}
 
-	static const TCHAR* ToString(EGLTFJsonCubeFace Value)
+	static const TCHAR* GetValue(EGLTFJsonCubeFace Enum)
 	{
-		switch (Value)
+		switch (Enum)
 		{
 			case EGLTFJsonCubeFace::PosX: return TEXT("PosX");
 			case EGLTFJsonCubeFace::NegX: return TEXT("NegX");
@@ -120,9 +119,9 @@ struct FGLTFJsonUtility
 		}
 	}
 
-	static const TCHAR* ToString(EGLTFJsonCameraType Value)
+	static const TCHAR* GetValue(EGLTFJsonCameraType Enum)
 	{
-		switch (Value)
+		switch (Enum)
 		{
 			case EGLTFJsonCameraType::Perspective:  return TEXT("perspective");
 			case EGLTFJsonCameraType::Orthographic: return TEXT("orthographic");
@@ -132,9 +131,9 @@ struct FGLTFJsonUtility
 		}
 	}
 
-	static const TCHAR* ToString(EGLTFJsonLightType Value)
+	static const TCHAR* GetValue(EGLTFJsonLightType Enum)
 	{
-		switch (Value)
+		switch (Enum)
 		{
 			case EGLTFJsonLightType::Directional: return TEXT("directional");
 			case EGLTFJsonLightType::Point:       return TEXT("point");
@@ -145,9 +144,9 @@ struct FGLTFJsonUtility
 		}
 	}
 
-	static const TCHAR* ToString(EGLTFJsonInterpolation Value)
+	static const TCHAR* GetValue(EGLTFJsonInterpolation Enum)
 	{
-		switch (Value)
+		switch (Enum)
 		{
 			case EGLTFJsonInterpolation::Linear:      return TEXT("LINEAR");
 			case EGLTFJsonInterpolation::Step:        return TEXT("STEP");
@@ -158,9 +157,9 @@ struct FGLTFJsonUtility
 		}
 	}
 
-	static const TCHAR* ToString(EGLTFJsonTargetPath Value)
+	static const TCHAR* GetValue(EGLTFJsonTargetPath Enum)
 	{
-		switch (Value)
+		switch (Enum)
 		{
 			case EGLTFJsonTargetPath::Translation: return TEXT("translation");
 			case EGLTFJsonTargetPath::Rotation:    return TEXT("rotation");
@@ -172,9 +171,9 @@ struct FGLTFJsonUtility
 		}
 	}
 
-	static const TCHAR* ToString(EGLTFJsonCameraControlMode Value)
+	static const TCHAR* GetValue(EGLTFJsonCameraControlMode Enum)
 	{
-		switch (Value)
+		switch (Enum)
 		{
 			case EGLTFJsonCameraControlMode::FreeLook: return TEXT("freeLook");
 			case EGLTFJsonCameraControlMode::Orbital:  return TEXT("orbital");
@@ -183,71 +182,4 @@ struct FGLTFJsonUtility
 				return TEXT("");
 		}
 	}
-
-	template <class CharType = TCHAR, class PrintPolicy = TPrettyJsonPrintPolicy<CharType>>
-	static void WriteExactValue(TJsonWriter<CharType, PrintPolicy>& JsonWriter, float Value)
-	{
-		FString ExactStringRepresentation = FString::Printf(TEXT("%.9g"), Value);
-		JsonWriter.WriteRawJSONValue(ExactStringRepresentation);
-	}
-
-	template <class CharType = TCHAR, class PrintPolicy = TPrettyJsonPrintPolicy<CharType>>
-	static void WriteExactValue(TJsonWriter<CharType, PrintPolicy>& JsonWriter, const FString& Identifier, float Value)
-	{
-		FString ExactStringRepresentation = FString::Printf(TEXT("%.9g"), Value);
-		JsonWriter.WriteRawJSONValue(Identifier, ExactStringRepresentation);
-	}
-
-	template <class WriterType, class ContainerType>
-	static void WriteObjectArray(WriterType& JsonWriter, const FString& Identifier, const ContainerType& Container, FGLTFJsonExtensions& Extensions, bool bWriteIfEmpty = false)
-	{
-		if (Container.Num() > 0 || bWriteIfEmpty)
-		{
-			JsonWriter.WriteArrayStart(Identifier);
-			for (const auto& Element : Container)
-			{
-				Element.WriteObject(JsonWriter, Extensions);
-			}
-			JsonWriter.WriteArrayEnd();
-		}
-	}
-
-	template <class WriterType, class ContainerType>
-	static void WriteObjectPtrArray(WriterType& JsonWriter, const FString& Identifier, const ContainerType& Container, FGLTFJsonExtensions& Extensions, bool bWriteIfEmpty = false)
-	{
-		if (Container.Num() > 0 || bWriteIfEmpty)
-		{
-			JsonWriter.WriteArrayStart(Identifier);
-			for (const auto& Element : Container)
-			{
-				Element->WriteObject(JsonWriter, Extensions);
-			}
-			JsonWriter.WriteArrayEnd();
-		}
-	}
-
-	template <class WriterType, class ContainerType>
-	static void WriteStringArray(WriterType& JsonWriter, const FString& Identifier, const ContainerType& Container, bool bWriteIfEmpty = false)
-	{
-		if (Container.Num() > 0 || bWriteIfEmpty)
-		{
-			JsonWriter.WriteArrayStart(Identifier);
-			for (const auto& Element : Container)
-			{
-				JsonWriter.WriteValue(ToString(Element));
-			}
-			JsonWriter.WriteArrayEnd();
-		}
-	}
-
-	template <class WriterType, class ElementType, SIZE_T Size>
-	static void WriteFixedArray(WriterType& JsonWriter, const FString& Identifier, const ElementType (&Array)[Size])
-	{
-		JsonWriter.WriteArrayStart(Identifier);
-		for (const ElementType& Element : Array)
-		{
-			JsonWriter.WriteValue(Element);
-		}
-		JsonWriter.WriteArrayEnd();
-	}
 };
diff --git a/Engine/Plugins/Enterprise/GLTFExporter/MigratedFromGit/Source/GLTFExporter/Private/Json/GLTFJsonValue.h b/Engine/Plugins/Enterprise/GLTFExporter/MigratedFromGit/Source/GLTFExporter/Private/Json/GLTFJsonValue.h
new file mode 100644
--- /dev/null
+++ b/Engine/Plugins/Enterprise/GLTFExporter/MigratedFromGit/Source/GLTFExporter/Private/Json/GLTFJsonValue.h
@@ -0,0 +1,12 @@
+// Copyright Epic Games, Inc. All Rights Reserved.
+
+#pragma once
+
+class IGLTFJsonWriter;
+
+struct IGLTFJsonValue
+{
+	virtual ~IGLTFJsonValue() = default;
+
+	virtual void WriteValue(IGLTFJsonWriter& Writer) const = 0;
+};
diff --git a/Engine/Plugins/Enterprise/GLTFExporter/MigratedFromGit/Source/GLTFExporter/Private/Json/GLTFJsonVector2.h b/Engine/Plugins/Enterprise/GLTFExporter/MigratedFromGit/Source/GLTFExporter/Private/Json/GLTFJsonVector2.h
--- a/Engine/Plugins/Enterprise/GLTFExporter/MigratedFromGit/Source/GLTFExporter/Private/Json/GLTFJsonVector2.h
+++ b/Engine/Plugins/Enterprise/GLTFExporter/MigratedFromGit/Source/GLTFExporter/Private/Json/GLTFJsonVector2.h
@@ -2,10 +2,10 @@
 
 #pragma once
 
-#include "Json/GLTFJsonUtility.h"
-#include "Serialization/JsonSerializer.h"
+#include "Json/GLTFJsonArray.h"
+#include "Converters/GLTFRawTypes.h"
 
-struct FGLTFJsonVector2
+struct FGLTFJsonVector2 : IGLTFJsonArray
 {
 	float X, Y;
 
@@ -17,13 +17,15 @@ struct FGLTFJsonVector2
 	{
 	}
 
-	template <class CharType = TCHAR, class PrintPolicy = TPrettyJsonPrintPolicy<CharType>>
-	void WriteArray(TJsonWriter<CharType, PrintPolicy>& JsonWriter) const
+	FGLTFJsonVector2(const FGLTFRawVector2& Raw)
+		: X(Raw.X), Y(Raw.Y)
 	{
-		JsonWriter.WriteArrayStart();
-		FGLTFJsonUtility::WriteExactValue(JsonWriter, X);
-		FGLTFJsonUtility::WriteExactValue(JsonWriter, Y);
-		JsonWriter.WriteArrayEnd();
+	}
+
+	virtual void WriteArray(IGLTFJsonWriter& Writer) const override
+	{
+		Writer.Write(X);
+		Writer.Write(Y);
 	}
 
 	bool operator==(const FGLTFJsonVector2& Other) const
diff --git a/Engine/Plugins/Enterprise/GLTFExporter/MigratedFromGit/Source/GLTFExporter/Private/Json/GLTFJsonVector3.h b/Engine/Plugins/Enterprise/GLTFExporter/MigratedFromGit/Source/GLTFExporter/Private/Json/GLTFJsonVector3.h
--- a/Engine/Plugins/Enterprise/GLTFExporter/MigratedFromGit/Source/GLTFExporter/Private/Json/GLTFJsonVector3.h
+++ b/Engine/Plugins/Enterprise/GLTFExporter/MigratedFromGit/Source/GLTFExporter/Private/Json/GLTFJsonVector3.h
@@ -2,11 +2,11 @@
 
 #pragma once
 
+#include "Json/GLTFJsonArray.h"
 #include "Json/GLTFJsonVector2.h"
-#include "Json/GLTFJsonUtility.h"
-#include "Serialization/JsonSerializer.h"
+#include "Converters/GLTFRawTypes.h"
 
-struct FGLTFJsonVector3
+struct FGLTFJsonVector3 : IGLTFJsonArray
 {
 	float X, Y, Z;
 
@@ -23,14 +23,16 @@ struct FGLTFJsonVector3
 	{
 	}
 
-	template <class CharType = TCHAR, class PrintPolicy = TPrettyJsonPrintPolicy<CharType>>
-	void WriteArray(TJsonWriter<CharType, PrintPolicy>& JsonWriter) const
+	FGLTFJsonVector3(const FGLTFRawVector3& Raw)
+		: X(Raw.X), Y(Raw.Y), Z(Raw.Z)
 	{
-		JsonWriter.WriteArrayStart();
-		FGLTFJsonUtility::WriteExactValue(JsonWriter, X);
-		FGLTFJsonUtility::WriteExactValue(JsonWriter, Y);
-		FGLTFJsonUtility::WriteExactValue(JsonWriter, Z);
-		JsonWriter.WriteArrayEnd();
+	}
+
+	virtual void WriteArray(IGLTFJsonWriter& Writer) const override
+	{
+		Writer.Write(X);
+		Writer.Write(Y);
+		Writer.Write(Z);
 	}
 
 	bool operator==(const FGLTFJsonVector3& Other) const
diff --git a/Engine/Plugins/Enterprise/GLTFExporter/MigratedFromGit/Source/GLTFExporter/Private/Json/GLTFJsonVector4.h b/Engine/Plugins/Enterprise/GLTFExporter/MigratedFromGit/Source/GLTFExporter/Private/Json/GLTFJsonVector4.h
--- a/Engine/Plugins/Enterprise/GLTFExporter/MigratedFromGit/Source/GLTFExporter/Private/Json/GLTFJsonVector4.h
+++ b/Engine/Plugins/Enterprise/GLTFExporter/MigratedFromGit/Source/GLTFExporter/Private/Json/GLTFJsonVector4.h
@@ -2,11 +2,10 @@
 
 #pragma once
 
-#include "Json/GLTFJsonVector3.h"
-#include "Json/GLTFJsonUtility.h"
-#include "Serialization/JsonSerializer.h"
+#include "Json/GLTFJsonArray.h"
+#include "Converters/GLTFRawTypes.h"
 
-struct FGLTFJsonVector4
+struct FGLTFJsonVector4 : IGLTFJsonArray
 {
 	float X, Y, Z, W;
 
@@ -18,25 +17,17 @@ struct FGLTFJsonVector4
 	{
 	}
 
-	FGLTFJsonVector4(const FGLTFJsonVector3& Vector3, float W)
-		: X(Vector3.X), Y(Vector3.Y), Z(Vector3.Z), W(W)
+	FGLTFJsonVector4(const FGLTFRawVector4& Raw)
+		: X(Raw.X), Y(Raw.Y), Z(Raw.Z), W(Raw.W)
 	{
 	}
 
-	template <class CharType = TCHAR, class PrintPolicy = TPrettyJsonPrintPolicy<CharType>>
-	void WriteArray(TJsonWriter<CharType, PrintPolicy>& JsonWriter) const
+	virtual void WriteArray(IGLTFJsonWriter& Writer) const override
 	{
-		JsonWriter.WriteArrayStart();
-		FGLTFJsonUtility::WriteExactValue(JsonWriter, X);
-		FGLTFJsonUtility::WriteExactValue(JsonWriter, Y);
-		FGLTFJsonUtility::WriteExactValue(JsonWriter, Z);
-		FGLTFJsonUtility::WriteExactValue(JsonWriter, W);
-		JsonWriter.WriteArrayEnd();
-	}
-
-	operator FGLTFJsonVector3() const
-	{
-		return { X, Y, Z };
+		Writer.Write(X);
+		Writer.Write(Y);
+		Writer.Write(Z);
+		Writer.Write(W);
 	}
 
 	bool operator==(const FGLTFJsonVector4& Other) const
diff --git a/Engine/Plugins/Enterprise/GLTFExporter/MigratedFromGit/Source/GLTFExporter/Private/Json/GLTFJsonWriter.cpp b/Engine/Plugins/Enterprise/GLTFExporter/MigratedFromGit/Source/GLTFExporter/Private/Json/GLTFJsonWriter.cpp
new file mode 100644
--- /dev/null
+++ b/Engine/Plugins/Enterprise/GLTFExporter/MigratedFromGit/Source/GLTFExporter/Private/Json/GLTFJsonWriter.cpp
@@ -0,0 +1,156 @@
+// Copyright Epic Games, Inc. All Rights Reserved.
+
+#include "Json/GLTFJsonWriter.h"
+#include "Serialization/JsonSerializer.h"
+#include "Policies/CondensedJsonPrintPolicy.h"
+
+template <class CharType, class PrintPolicy>
+class TGLTFJsonWriterImpl final : public IGLTFJsonWriter
+{
+public:
+
+	TGLTFJsonWriterImpl(FArchive* const Archive)
+		: JsonWriter(TJsonWriterFactory<CharType, PrintPolicy>::Create(Archive))
+	{
+	}
+
+	virtual void Close() override
+	{
+		JsonWriter->Close();
+	}
+
+	virtual void Write(bool Boolean) override
+	{
+		if (CurrentIdentifier.IsEmpty())
+		{
+			JsonWriter->WriteValue(Boolean);
+		}
+		else
+		{
+			JsonWriter->WriteValue(CurrentIdentifier, Boolean);
+			CurrentIdentifier.Empty();
+		}
+	}
+
+	virtual void Write(int32 Number) override
+	{
+		if (CurrentIdentifier.IsEmpty())
+		{
+			JsonWriter->WriteValue(Number);
+		}
+		else
+		{
+			JsonWriter->WriteValue(CurrentIdentifier, Number);
+			CurrentIdentifier.Empty();
+		}
+	}
+
+	virtual void Write(int64 Number) override
+	{
+		if (CurrentIdentifier.IsEmpty())
+		{
+			JsonWriter->WriteValue(Number);
+		}
+		else
+		{
+			JsonWriter->WriteValue(CurrentIdentifier, Number);
+			CurrentIdentifier.Empty();
+		}
+	}
+
+	virtual void Write(float Number) override
+	{
+		// NOTE: Specifying 9 significant digits, this ensures no precision is lost
+		FString ExactString = FString::Printf(TEXT("%.9g"), Number);
+
+		if (CurrentIdentifier.IsEmpty())
+		{
+			JsonWriter->WriteRawJSONValue(ExactString);
+		}
+		else
+		{
+			JsonWriter->WriteRawJSONValue(CurrentIdentifier, ExactString);
+			CurrentIdentifier.Empty();
+		}
+	}
+
+	virtual void Write(const FString& String) override
+	{
+		if (CurrentIdentifier.IsEmpty())
+		{
+			JsonWriter->WriteValue(String);
+		}
+		else
+		{
+			JsonWriter->WriteValue(CurrentIdentifier, String);
+			CurrentIdentifier.Empty();
+		}
+	}
+
+	virtual void Write(TYPE_OF_NULLPTR) override
+	{
+		if (CurrentIdentifier.IsEmpty())
+		{
+			JsonWriter->WriteValue(nullptr);
+		}
+		else
+		{
+			JsonWriter->WriteValue(CurrentIdentifier, nullptr);
+			CurrentIdentifier.Empty();
+		}
+	}
+
+	virtual void SetIdentifier(const FString& Identifier) override
+	{
+		this->CurrentIdentifier = Identifier;
+	}
+
+	virtual void StartObject() override
+	{
+		if (CurrentIdentifier.IsEmpty())
+		{
+			JsonWriter->WriteObjectStart();
+		}
+		else
+		{
+			JsonWriter->WriteObjectStart(CurrentIdentifier);
+			CurrentIdentifier.Empty();
+		}
+	}
+
+	virtual void EndObject() override
+	{
+		JsonWriter->WriteObjectEnd();
+	}
+
+	virtual void StartArray() override
+	{
+		if (CurrentIdentifier.IsEmpty())
+		{
+			JsonWriter->WriteArrayStart();
+		}
+		else
+		{
+			JsonWriter->WriteArrayStart(CurrentIdentifier);
+			CurrentIdentifier.Empty();
+		}
+	}
+
+	virtual void EndArray() override
+	{
+		JsonWriter->WriteArrayEnd();
+	}
+
+private:
+
+	FString CurrentIdentifier;
+	TSharedRef<TJsonWriter<CharType, PrintPolicy>> JsonWriter;
+};
+
+TSharedRef<IGLTFJsonWriter> IGLTFJsonWriter::Create(FArchive* const Archive, bool bPrettyJson)
+{
+	return MakeShareable(bPrettyJson
+		? static_cast<IGLTFJsonWriter*>(new TGLTFJsonWriterImpl<UTF8CHAR, TPrettyJsonPrintPolicy<UTF8CHAR>>(Archive))
+		: static_cast<IGLTFJsonWriter*>(new TGLTFJsonWriterImpl<UTF8CHAR, TCondensedJsonPrintPolicy<UTF8CHAR>>(Archive))
+	);
+}
diff --git a/Engine/Plugins/Enterprise/GLTFExporter/MigratedFromGit/Source/GLTFExporter/Private/Json/GLTFJsonWriter.h b/Engine/Plugins/Enterprise/GLTFExporter/MigratedFromGit/Source/GLTFExporter/Private/Json/GLTFJsonWriter.h
new file mode 100644
--- /dev/null
+++ b/Engine/Plugins/Enterprise/GLTFExporter/MigratedFromGit/Source/GLTFExporter/Private/Json/GLTFJsonWriter.h
@@ -0,0 +1,254 @@
+// Copyright Epic Games, Inc. All Rights Reserved.
+
+#pragma once
+
+#include "Json/GLTFJsonValue.h"
+#include "Json/GLTFJsonUtility.h"
+#include "Json/GLTFJsonExtensions.h"
+
+class IGLTFJsonWriter
+{
+public:
+
+	FGLTFJsonExtensions Extensions;
+
+	static TSharedRef<IGLTFJsonWriter> Create(FArchive* const Archive, bool bPrettyJson);
+
+	virtual ~IGLTFJsonWriter() = default;
+
+	virtual void Close() = 0;
+
+	virtual void Write(bool Boolean) = 0;
+	virtual void Write(int32 Number) = 0;
+	virtual void Write(int64 Number) = 0;
+	virtual void Write(float Number) = 0;
+	virtual void Write(const FString& String) = 0;
+	virtual void Write(TYPE_OF_NULLPTR) = 0;
+	virtual void SetIdentifier(const FString& Identifier) = 0;
+
+	virtual void StartObject() = 0;
+	virtual void EndObject() = 0;
+	virtual void StartArray() = 0;
+	virtual void EndArray() = 0;
+
+	void Write(const TCHAR* String)
+	{
+		Write(FString(String));
+	}
+
+	void Write(const IGLTFJsonValue& Value)
+	{
+		Value.WriteValue(*this);
+	}
+
+	void Write(const IGLTFJsonValue* Value)
+	{
+		if (Value != nullptr)
+		{
+			Value->WriteValue(*this);
+		}
+		else
+		{
+			Write(nullptr);
+		}
+	}
+
+	template <typename EnumType, typename = typename TEnableIf<TIsEnum<EnumType>::Value>::Type>
+	void Write(EnumType Enum)
+	{
+		Write(FGLTFJsonUtility::GetValue(Enum));
+	}
+
+	template <class ElementType, typename = typename TEnableIf<TNot<TIsSame<ElementType, TCHAR>>::Value>::Type>
+	void Write(const ElementType* Array, SIZE_T ArraySize)
+	{
+		StartArray();
+		for (SIZE_T i = 0; i < ArraySize; ++i)
+		{
+			Write(Array[i]);
+		}
+		EndArray();
+	}
+
+	template <class ElementType, SIZE_T ArraySize, typename = typename TEnableIf<TNot<TIsSame<ElementType, TCHAR>>::Value>::Type>
+	void Write(const ElementType (&Array)[ArraySize])
+	{
+		Write(Array, ArraySize);
+	}
+
+	template <class ElementType>
+	void Write(const TUniquePtr<ElementType>& Pointer)
+	{
+		if (Pointer.IsValid())
+		{
+			Write(*Pointer.Get());
+		}
+		else
+		{
+			Write(nullptr);
+		}
+	}
+
+	template <class ElementType>
+	void Write(const TArray<ElementType>& Array)
+	{
+		StartArray();
+		for (const auto& Element : Array)
+		{
+			Write(Element);
+		}
+		EndArray();
+	}
+
+	template <class ElementType>
+	void Write(const TSet<ElementType>& Set)
+	{
+		StartArray();
+		for (const auto& Element : Set)
+		{
+			Write(Element);
+		}
+		EndArray();
+	}
+
+	void Write(const FString& Identifier, bool Boolean)
+	{
+		SetIdentifier(Identifier);
+		Write(Boolean);
+	}
+
+	void Write(const FString& Identifier, int32 Number)
+	{
+		SetIdentifier(Identifier);
+		Write(Number);
+	}
+
+	void Write(const FString& Identifier, int64 Number)
+	{
+		SetIdentifier(Identifier);
+		Write(Number);
+	}
+
+	void Write(const FString& Identifier, float Number)
+	{
+		SetIdentifier(Identifier);
+		Write(Number);
+	}
+
+	void Write(const FString& Identifier, const FString& String)
+	{
+		SetIdentifier(Identifier);
+		Write(String);
+	}
+
+	void Write(const FString& Identifier, const TCHAR* String)
+	{
+		SetIdentifier(Identifier);
+		Write(String);
+	}
+
+	void Write(const FString& Identifier, const IGLTFJsonValue& Value)
+	{
+		SetIdentifier(Identifier);
+		Write(Value);
+	}
+
+	void Write(const FString& Identifier, const IGLTFJsonValue* Value)
+	{
+		SetIdentifier(Identifier);
+		Write(Value);
+	}
+
+	template <typename EnumType, typename = typename TEnableIf<TIsEnum<EnumType>::Value>::Type>
+	void Write(const FString& Identifier, EnumType Enum)
+	{
+		SetIdentifier(Identifier);
+		Write(Enum);
+	}
+
+	template <class ElementType, typename = typename TEnableIf<TNot<TIsSame<ElementType, TCHAR>>::Value>::Type>
+	void Write(const FString& Identifier, const ElementType* Array, SIZE_T ArraySize)
+	{
+		SetIdentifier(Identifier);
+		Write(Array, ArraySize);
+	}
+
+	template <class ElementType, SIZE_T ArraySize, typename = typename TEnableIf<TNot<TIsSame<ElementType, TCHAR>>::Value>::Type>
+	void Write(const FString& Identifier, const ElementType (&Array)[ArraySize])
+	{
+		SetIdentifier(Identifier);
+		Write(Array);
+	}
+
+	template <class ElementType>
+	void Write(const FString& Identifier, const TUniquePtr<ElementType>& Pointer)
+	{
+		SetIdentifier(Identifier);
+		Write(Pointer);
+	}
+
+	template <class ElementType>
+	void Write(const FString& Identifier, const TArray<ElementType>& Array)
+	{
+		SetIdentifier(Identifier);
+		Write(Array);
+	}
+
+	template <class ElementType>
+	void Write(const FString& Identifier, const TSet<ElementType>& Set)
+	{
+		SetIdentifier(Identifier);
+		Write(Set);
+	}
+
+	void StartObject(const FString& Identifier)
+	{
+		SetIdentifier(Identifier);
+		StartObject();
+	}
+
+	void StartArray(const FString& Identifier)
+	{
+		SetIdentifier(Identifier);
+		StartArray();
+	}
+
+	void StartExtensions()
+	{
+		StartObject(TEXT("extensions"));
+	}
+
+	void EndExtensions()
+	{
+		EndObject();
+	}
+
+	void Write(const EGLTFJsonExtension& Extension, const IGLTFJsonValue& Value, bool bIsRequired = false)
+	{
+		Extensions.Used.Add(Extension);
+
+		if (bIsRequired)
+		{
+			Extensions.Required.Add(Extension);
+		}
+
+		Write(FGLTFJsonUtility::GetValue(Extension), Value);
+	}
+
+	void StartExtension(const EGLTFJsonExtension& Extension, bool bIsRequired = false)
+	{
+		Extensions.Used.Add(Extension);
+
+		if (bIsRequired)
+		{
+			Extensions.Required.Add(Extension);
+		}
+
+		StartObject(FGLTFJsonUtility::GetValue(Extension));
+	}
+
+	void EndExtension()
+	{
+		EndObject();
+	}
+};
diff --git a/Engine/Plugins/Enterprise/GLTFExporter/MigratedFromGit/Source/GLTFExporter/Private/Tasks/GLTFAnimationTasks.cpp b/Engine/Plugins/Enterprise/GLTFExporter/MigratedFromGit/Source/GLTFExporter/Private/Tasks/GLTFAnimationTasks.cpp
--- a/Engine/Plugins/Enterprise/GLTFExporter/MigratedFromGit/Source/GLTFExporter/Private/Tasks/GLTFAnimationTasks.cpp
+++ b/Engine/Plugins/Enterprise/GLTFExporter/MigratedFromGit/Source/GLTFExporter/Private/Tasks/GLTFAnimationTasks.cpp
@@ -96,7 +96,7 @@ void FGLTFAnimSequenceTask::Complete()
 
 		if (KeyPositions.Num() > 0)
 		{
-			TArray<FGLTFJsonVector3> Translations;
+			TArray<FGLTFRawVector3> Translations;
 			Translations.AddUninitialized(KeyPositions.Num());
 
 			for (int32 Key = 0; Key < KeyPositions.Num(); ++Key)
@@ -128,7 +128,7 @@ void FGLTFAnimSequenceTask::Complete()
 
 		if (KeyRotations.Num() > 0)
 		{
-			TArray<FGLTFJsonQuaternion> Rotations;
+			TArray<FGLTFRawQuaternion> Rotations;
 			Rotations.AddUninitialized(KeyRotations.Num());
 
 			for (int32 Key = 0; Key < KeyRotations.Num(); ++Key)
@@ -160,7 +160,7 @@ void FGLTFAnimSequenceTask::Complete()
 
 		if (KeyScales.Num() > 0)
 		{
-			TArray<FGLTFJsonVector3> Scales;
+			TArray<FGLTFRawVector3> Scales;
 			Scales.AddUninitialized(KeyScales.Num());
 
 			for (int32 Key = 0; Key < KeyScales.Num(); ++Key)
@@ -297,7 +297,7 @@ void FGLTFLevelSequenceTask::Complete()
 
 					if (EnumHasAnyFlags(ChannelMask, EMovieSceneTransformChannel::Translation))
 					{
-						TArray<FGLTFJsonVector3> Translations;
+						TArray<FGLTFRawVector3> Translations;
 						Translations.AddUninitialized(FrameCount);
 
 						for (int32 Frame = 0; Frame < FrameCount; ++Frame)
@@ -332,7 +332,7 @@ void FGLTFLevelSequenceTask::Complete()
 
 					if (EnumHasAnyFlags(ChannelMask, EMovieSceneTransformChannel::Rotation))
 					{
-						TArray<FGLTFJsonQuaternion> Rotations;
+						TArray<FGLTFRawQuaternion> Rotations;
 						Rotations.AddUninitialized(FrameCount);
 
 						for (int32 Frame = 0; Frame < FrameCount; ++Frame)
@@ -370,7 +370,7 @@ void FGLTFLevelSequenceTask::Complete()
 
 					if (EnumHasAnyFlags(ChannelMask, EMovieSceneTransformChannel::Scale))
 					{
-						TArray<FGLTFJsonVector3> Scales;
+						TArray<FGLTFRawVector3> Scales;
 						Scales.AddUninitialized(FrameCount);
 
 						for (int32 Frame = 0; Frame < FrameCount; ++Frame)
diff --git a/Engine/Plugins/Enterprise/GLTFExporter/MigratedFromGit/Source/GLTFExporter/Private/Tasks/GLTFTextureTasks.cpp b/Engine/Plugins/Enterprise/GLTFExporter/MigratedFromGit/Source/GLTFExporter/Private/Tasks/GLTFTextureTasks.cpp
--- a/Engine/Plugins/Enterprise/GLTFExporter/MigratedFromGit/Source/GLTFExporter/Private/Tasks/GLTFTextureTasks.cpp
+++ b/Engine/Plugins/Enterprise/GLTFExporter/MigratedFromGit/Source/GLTFExporter/Private/Tasks/GLTFTextureTasks.cpp
@@ -47,7 +47,7 @@ void FGLTFTexture2DTask::Complete()
 void FGLTFTextureCubeTask::Complete()
 {
 	FGLTFJsonTexture& JsonTexture = Builder.GetTexture(TextureIndex);
-	JsonTexture.Name = TextureCube->GetName() + TEXT("_") + FGLTFJsonUtility::ToString(FGLTFConverterUtility::ConvertCubeFace(CubeFace));
+	JsonTexture.Name = TextureCube->GetName() + TEXT("_") + FGLTFJsonUtility::GetValue(FGLTFConverterUtility::ConvertCubeFace(CubeFace));
 
 	// TODO: add optimized "happy path" if cube face doesn't need rotation and has suitable pixel format
 
@@ -115,7 +115,7 @@ void FGLTFTextureRenderTarget2DTask::Complete()
 void FGLTFTextureRenderTargetCubeTask::Complete()
 {
 	FGLTFJsonTexture& JsonTexture = Builder.GetTexture(TextureIndex);
-	JsonTexture.Name = RenderTargetCube->GetName() + TEXT("_") + FGLTFJsonUtility::ToString(FGLTFConverterUtility::ConvertCubeFace(CubeFace));
+	JsonTexture.Name = RenderTargetCube->GetName() + TEXT("_") + FGLTFJsonUtility::GetValue(FGLTFConverterUtility::ConvertCubeFace(CubeFace));
 
 	// TODO: add optimized "happy path" if cube face doesn't need rotation
 
