[GLTFExporter] refactor: simplify all vector/quat/color/matrix types by adding unified core types

commit bc079b20e1aeceabc7cff7da7883dc9378157856
Date: Fri, 11 Jun 2021 02:38:10 +0200

    refactor: simplify all vector/quat/color/matrix types by adding unified core types
    (cherry picked from commit f52f0f672b6837ae6e13fa9cd2fb4bf4d6b0adb7​)

#skipci

diff --git a/Engine/Plugins/Enterprise/GLTFExporter/MigratedFromGit/Source/GLTFExporter/Private/Converters/GLTFAccessorConverters.cpp b/Engine/Plugins/Enterprise/GLTFExporter/MigratedFromGit/Source/GLTFExporter/Private/Converters/GLTFAccessorConverters.cpp
--- a/Engine/Plugins/Enterprise/GLTFExporter/MigratedFromGit/Source/GLTFExporter/Private/Converters/GLTFAccessorConverters.cpp
+++ b/Engine/Plugins/Enterprise/GLTFExporter/MigratedFromGit/Source/GLTFExporter/Private/Converters/GLTFAccessorConverters.cpp
@@ -30,7 +30,7 @@ FGLTFJsonAccessorIndex FGLTFPositionBufferConverter::Convert(const FGLTFMeshSect
 	const TArray<uint32>& IndexMap = MeshSection->IndexMap;
 	const uint32 VertexCount = IndexMap.Num();
 
-	TArray<FGLTFRawVector3> Positions;
+	TArray<FGLTFVector3> Positions;
 	Positions.AddUninitialized(VertexCount);
 
 	for (uint32 VertexIndex = 0; VertexIndex < VertexCount; ++VertexIndex)
@@ -39,34 +39,33 @@ FGLTFJsonAccessorIndex FGLTFPositionBufferConverter::Convert(const FGLTFMeshSect
 		Positions[VertexIndex] = FGLTFConverterUtility::ConvertPosition(VertexBuffer->VertexPosition(MappedVertexIndex), Builder.ExportOptions->ExportUniformScale);
 	}
 
-	// More accurate bounding box if based on raw vertex values
-	FGLTFRawVector3 MinPosition = VertexCount > 0 ? Positions[0] : FGLTFRawVector3(0, 0, 0); // TODO: make static const
-	FGLTFRawVector3 MaxPosition = VertexCount > 0 ? Positions[0] : FGLTFRawVector3(0, 0, 0); // TODO: make static const
-
-	for (uint32 VertexIndex = 1; VertexIndex < VertexCount; ++VertexIndex)
-	{
-		const FGLTFRawVector3& Position = Positions[VertexIndex];
-		MinPosition.X = FMath::Min(MinPosition.X, Position.X);
-		MinPosition.Y = FMath::Min(MinPosition.Y, Position.Y);
-		MinPosition.Z = FMath::Min(MinPosition.Z, Position.Z);
-		MaxPosition.X = FMath::Max(MaxPosition.X, Position.X);
-		MaxPosition.Y = FMath::Max(MaxPosition.Y, Position.Y);
-		MaxPosition.Z = FMath::Max(MaxPosition.Z, Position.Z);
-	}
-
 	FGLTFJsonAccessor JsonAccessor;
 	JsonAccessor.BufferView = Builder.AddBufferView(Positions, EGLTFJsonBufferTarget::ArrayBuffer);
 	JsonAccessor.ComponentType = EGLTFJsonComponentType::F32;
 	JsonAccessor.Count = VertexCount;
 	JsonAccessor.Type = EGLTFJsonAccessorType::Vec3;
 
-	JsonAccessor.MinMaxLength = 3;
-	JsonAccessor.Min[0] = MinPosition.X;
-	JsonAccessor.Min[1] = MinPosition.Y;
-	JsonAccessor.Min[2] = MinPosition.Z;
-	JsonAccessor.Max[0] = MaxPosition.X;
-	JsonAccessor.Max[1] = MaxPosition.Y;
-	JsonAccessor.Max[2] = MaxPosition.Z;
+	if (VertexCount > 0)
+	{
+		// Calculate accurate bounding box based on raw vertex values
+		JsonAccessor.MinMaxLength = 3;
+
+		for (int32 ComponentIndex = 0; ComponentIndex < JsonAccessor.MinMaxLength; ComponentIndex++)
+		{
+			JsonAccessor.Min[ComponentIndex] = Positions[0].Components[ComponentIndex];
+			JsonAccessor.Max[ComponentIndex] = Positions[0].Components[ComponentIndex];
+		}
+
+		for (uint32 VertexIndex = 1; VertexIndex < VertexCount; ++VertexIndex)
+		{
+			const FGLTFVector3& Position = Positions[VertexIndex];
+			for (int32 ComponentIndex = 0; ComponentIndex < JsonAccessor.MinMaxLength; ComponentIndex++)
+			{
+				JsonAccessor.Min[ComponentIndex] = FMath::Min(JsonAccessor.Min[ComponentIndex], Position.Components[ComponentIndex]);
+				JsonAccessor.Max[ComponentIndex] = FMath::Max(JsonAccessor.Max[ComponentIndex], Position.Components[ComponentIndex]);
+			}
+		}
+	}
 
 	return Builder.AddAccessor(JsonAccessor);
 }
@@ -81,7 +80,7 @@ FGLTFJsonAccessorIndex FGLTFColorBufferConverter::Convert(const FGLTFMeshSection
 	const TArray<uint32>& IndexMap = MeshSection->IndexMap;
 	const uint32 VertexCount = IndexMap.Num();
 
-	TArray<FGLTFPackedColor> Colors;
+	TArray<FGLTFUInt8Color4> Colors;
 	Colors.AddUninitialized(VertexCount);
 
 	for (uint32 VertexIndex = 0; VertexIndex < VertexCount; ++VertexIndex)
@@ -120,22 +119,22 @@ FGLTFJsonAccessorIndex FGLTFNormalBufferConverter::Convert(const FGLTFMeshSectio
 		if (bHighPrecision)
 		{
 			ComponentType = EGLTFJsonComponentType::S16;
-			BufferViewIndex = ConvertBufferView<FGLTFPackedVector16, FPackedRGBA16N>(MeshSection, VertexBuffer);
-			Builder.GetBufferView(BufferViewIndex).ByteStride = sizeof(FGLTFPackedVector16);
+			BufferViewIndex = ConvertBufferView<FGLTFInt16Vector4, FPackedRGBA16N>(MeshSection, VertexBuffer);
+			Builder.GetBufferView(BufferViewIndex).ByteStride = sizeof(FGLTFInt16Vector4);
 		}
 		else
 		{
 			ComponentType = EGLTFJsonComponentType::S8;
-			BufferViewIndex = ConvertBufferView<FGLTFPackedVector8, FPackedNormal>(MeshSection, VertexBuffer);
-			Builder.GetBufferView(BufferViewIndex).ByteStride = sizeof(FGLTFPackedVector8);
+			BufferViewIndex = ConvertBufferView<FGLTFInt8Vector4, FPackedNormal>(MeshSection, VertexBuffer);
+			Builder.GetBufferView(BufferViewIndex).ByteStride = sizeof(FGLTFInt8Vector4);
 		}
 	}
 	else
 	{
 		ComponentType = EGLTFJsonComponentType::F32;
 		BufferViewIndex = bHighPrecision
-			? ConvertBufferView<FGLTFRawVector3, FPackedRGBA16N>(MeshSection, VertexBuffer)
-			: ConvertBufferView<FGLTFRawVector3, FPackedNormal>(MeshSection, VertexBuffer);
+			? ConvertBufferView<FGLTFVector3, FPackedRGBA16N>(MeshSection, VertexBuffer)
+			: ConvertBufferView<FGLTFVector3, FPackedNormal>(MeshSection, VertexBuffer);
 	}
 
 	FGLTFJsonAccessor JsonAccessor;
@@ -172,7 +171,7 @@ FGLTFJsonBufferViewIndex FGLTFNormalBufferConverter::ConvertBufferView(const FGL
 		const uint32 MappedVertexIndex = IndexMap[VertexIndex];
 		const FVector SafeNormal = VertexTangents[MappedVertexIndex].TangentZ.ToFVector().GetSafeNormal();
 
-		typedef typename TConditional<TIsSame<DestinationType, FGLTFRawVector3>::Value, FVector, SourceType>::Type IntermediateType;
+		typedef typename TConditional<TIsSame<DestinationType, FGLTFVector3>::Value, FVector, SourceType>::Type IntermediateType;
 		Normals[VertexIndex] = FGLTFConverterUtility::ConvertNormal(IntermediateType(SafeNormal));
 	}
 
@@ -199,20 +198,20 @@ FGLTFJsonAccessorIndex FGLTFTangentBufferConverter::Convert(const FGLTFMeshSecti
 		if (bHighPrecision)
 		{
 			ComponentType = EGLTFJsonComponentType::S16;
-			BufferViewIndex = ConvertBufferView<FGLTFPackedVector16, FPackedRGBA16N>(MeshSection, VertexBuffer);
+			BufferViewIndex = ConvertBufferView<FGLTFInt16Vector4, FPackedRGBA16N>(MeshSection, VertexBuffer);
 		}
 		else
 		{
 			ComponentType = EGLTFJsonComponentType::S8;
-			BufferViewIndex = ConvertBufferView<FGLTFPackedVector8, FPackedNormal>(MeshSection, VertexBuffer);
+			BufferViewIndex = ConvertBufferView<FGLTFInt8Vector4, FPackedNormal>(MeshSection, VertexBuffer);
 		}
 	}
 	else
 	{
 		ComponentType = EGLTFJsonComponentType::F32;
 		BufferViewIndex = bHighPrecision
-			? ConvertBufferView<FGLTFRawVector4, FPackedRGBA16N>(MeshSection, VertexBuffer)
-			: ConvertBufferView<FGLTFRawVector4, FPackedNormal>(MeshSection, VertexBuffer);
+			? ConvertBufferView<FGLTFVector4, FPackedRGBA16N>(MeshSection, VertexBuffer)
+			: ConvertBufferView<FGLTFVector4, FPackedNormal>(MeshSection, VertexBuffer);
 	}
 
 	FGLTFJsonAccessor JsonAccessor;
@@ -249,7 +248,7 @@ FGLTFJsonBufferViewIndex FGLTFTangentBufferConverter::ConvertBufferView(const FG
 		const uint32 MappedVertexIndex = IndexMap[VertexIndex];
 		const FVector SafeTangent = VertexTangents[MappedVertexIndex].TangentX.ToFVector().GetSafeNormal();
 
-		typedef typename TConditional<TIsSame<DestinationType, FGLTFRawVector4>::Value, FVector, SourceType>::Type IntermediateType;
+		typedef typename TConditional<TIsSame<DestinationType, FGLTFVector4>::Value, FVector, SourceType>::Type IntermediateType;
 		Tangents[VertexIndex] = FGLTFConverterUtility::ConvertTangent(IntermediateType(SafeTangent));
 	}
 
@@ -275,7 +274,7 @@ FGLTFJsonAccessorIndex FGLTFUVBufferConverter::Convert(const FGLTFMeshSection* M
 
 	// TODO: report warning or add support for half float precision UVs, i.e. !VertexBuffer->GetUseFullPrecisionUVs()?
 
-	TArray<FGLTFRawVector2> UVs;
+	TArray<FGLTFVector2> UVs;
 	UVs.AddUninitialized(VertexCount);
 
 	for (uint32 VertexIndex = 0; VertexIndex < VertexCount; ++VertexIndex)
diff --git a/Engine/Plugins/Enterprise/GLTFExporter/MigratedFromGit/Source/GLTFExporter/Private/Converters/GLTFConverterUtility.h b/Engine/Plugins/Enterprise/GLTFExporter/MigratedFromGit/Source/GLTFExporter/Private/Converters/GLTFConverterUtility.h
--- a/Engine/Plugins/Enterprise/GLTFExporter/MigratedFromGit/Source/GLTFExporter/Private/Converters/GLTFConverterUtility.h
+++ b/Engine/Plugins/Enterprise/GLTFExporter/MigratedFromGit/Source/GLTFExporter/Private/Converters/GLTFConverterUtility.h
@@ -2,15 +2,11 @@
 
 #pragma once
 
-#include "Converters/GLTFRawTypes.h"
-#include "Converters/GLTFPackedTypes.h"
+#include "Core/GLTFVector.h"
+#include "Core/GLTFColor.h"
+#include "Core/GLTFMatrix.h"
+#include "Core/GLTFQuaternion.h"
 #include "Json/GLTFJsonEnums.h"
-#include "Json/GLTFJsonVector2.h"
-#include "Json/GLTFJsonVector3.h"
-#include "Json/GLTFJsonVector4.h"
-#include "Json/GLTFJsonColor4.h"
-#include "Json/GLTFJsonMatrix4.h"
-#include "Json/GLTFJsonQuaternion.h"
 #include "Engine/EngineTypes.h"
 
 enum class EGLTFCameraControlMode : unsigned char;
@@ -22,61 +18,61 @@ struct FGLTFConverterUtility
 		return Length * ConversionScale;
 	}
 
-	static FGLTFRawVector3 ConvertVector(const FVector& Vector)
+	static FGLTFVector3 ConvertVector(const FVector& Vector)
 	{
 		// UE4 uses a left-handed coordinate system, with Z up.
 		// glTF uses a right-handed coordinate system, with Y up.
 		return { Vector.X, Vector.Z, Vector.Y };
 	}
 
-	static FGLTFRawVector3 ConvertPosition(const FVector& Position, const float ConversionScale)
+	static FGLTFVector3 ConvertPosition(const FVector& Position, const float ConversionScale)
 	{
 		return ConvertVector(Position * ConversionScale);
 	}
 
-	static FGLTFRawVector3 ConvertScale(const FVector& Scale)
+	static FGLTFVector3 ConvertScale(const FVector& Scale)
 	{
 		return ConvertVector(Scale);
 	}
 
-	static FGLTFRawVector3 ConvertNormal(const FVector& Normal)
+	static FGLTFVector3 ConvertNormal(const FVector& Normal)
 	{
 		return ConvertVector(Normal);
 	}
 
-	static FGLTFPackedVector16 ConvertNormal(const FPackedRGBA16N& Normal)
+	static FGLTFInt16Vector4 ConvertNormal(const FPackedRGBA16N& Normal)
 	{
 		return { Normal.X, Normal.Z, Normal.Y, 0 };
 	}
 
-	static FGLTFPackedVector8 ConvertNormal(const FPackedNormal& Normal)
+	static FGLTFInt8Vector4 ConvertNormal(const FPackedNormal& Normal)
 	{
 		return { Normal.Vector.X, Normal.Vector.Z, Normal.Vector.Y };
 	}
 
-	static FGLTFRawVector4 ConvertTangent(const FVector& Tangent)
+	static FGLTFVector4 ConvertTangent(const FVector& Tangent)
 	{
 		// glTF stores tangent as Vec4, with W component indicating handedness of tangent basis.
-		return { ConvertVector(Tangent), 1.0f };
+		return { Tangent.X, Tangent.Z, Tangent.Y, 1.0f };
 	}
 
-	static FGLTFPackedVector16 ConvertTangent(const FPackedRGBA16N& Tangent)
+	static FGLTFInt16Vector4 ConvertTangent(const FPackedRGBA16N& Tangent)
 	{
 		return { Tangent.X, Tangent.Z, Tangent.Y, MAX_int16 /* = 1.0 */ };
 	}
 
-	static FGLTFPackedVector8 ConvertTangent(const FPackedNormal& Tangent)
+	static FGLTFInt8Vector4 ConvertTangent(const FPackedNormal& Tangent)
 	{
 		return { Tangent.Vector.X, Tangent.Vector.Z, Tangent.Vector.Y, MAX_int8 /* = 1.0 */ };
 	}
 
-	static FGLTFRawVector2 ConvertUV(const FVector2D& UV)
+	static FGLTFVector2 ConvertUV(const FVector2D& UV)
 	{
 		// No conversion actually needed, this is primarily for type-safety.
 		return { UV.X, UV.Y };
 	}
 
-	static FGLTFJsonColor4 ConvertColor(const FLinearColor& Color, bool bForceLDR)
+	static FGLTFColor4 ConvertColor(const FLinearColor& Color, bool bForceLDR)
 	{
 		if (bForceLDR)
 		{
@@ -97,14 +93,33 @@ struct FGLTFConverterUtility
 		};
 	}
 
-	static FGLTFPackedColor ConvertColor(const FColor& Color)
+	static FGLTFColor3 ConvertColor3(const FLinearColor& Color, bool bForceLDR)
+	{
+		if (bForceLDR)
+		{
+			return {
+				FMath::Clamp(Color.R, 0.0f, 1.0f),
+				FMath::Clamp(Color.G, 0.0f, 1.0f),
+				FMath::Clamp(Color.B, 0.0f, 1.0f)
+			};
+		}
+
+		// Just make sure its non-negative (which can happen when using MakeFromColorTemperature).
+		return {
+			FMath::Max(Color.R, 0.0f),
+			FMath::Max(Color.G, 0.0f),
+			FMath::Max(Color.B, 0.0f)
+		};
+	}
+
+	static FGLTFUInt8Color4 ConvertColor(const FColor& Color)
 	{
 		// UE4 uses ABGR or ARGB depending on endianness.
 		// glTF always uses RGBA independent of endianness.
 		return { Color.R, Color.G, Color.B, Color.A };
 	}
 
-	static FGLTFRawQuaternion ConvertRotation(const FQuat& Rotation)
+	static FGLTFQuaternion ConvertRotation(const FQuat& Rotation)
 	{
 		// UE4 uses a left-handed coordinate system, with Z up.
 		// glTF uses a right-handed coordinate system, with Y up.
@@ -120,19 +135,19 @@ struct FGLTFConverterUtility
 		// Return the identity quaternion when possible (depending on tolerance)
 		if (Rotation.Equals(FQuat::Identity)) // TODO: doesn't this potentially remove intentional deviance from identity?
 		{
-			return FGLTFRawQuaternion(0, 0, 0, 1); // TODO: make static const
+			return { 0, 0, 0, 1 }; // TODO: make static const
 		}
 
 		const FQuat Normalized = Rotation.GetNormalized();
 		return { -Normalized.X, -Normalized.Z, -Normalized.Y, Normalized.W };
 	}
 
-	static FGLTFRawMatrix4 ConvertMatrix(const FMatrix& Matrix)
+	static FGLTFMatrix4 ConvertMatrix(const FMatrix& Matrix)
 	{
 		// Unreal stores matrix elements in row major order.
 		// glTF stores matrix elements in column major order.
 
-		FGLTFRawMatrix4 Result;
+		FGLTFMatrix4 Result;
 		for (int32 Row = 0; Row < 4; ++Row)
 		{
 			for (int32 Col = 0; Col < 4; ++Col)
@@ -143,7 +158,7 @@ struct FGLTFConverterUtility
 		return Result;
 	}
 
-	static FGLTFRawMatrix4 ConvertTransform(const FTransform& Transform, const float ConversionScale)
+	static FGLTFMatrix4 ConvertTransform(const FTransform& Transform, const float ConversionScale)
 	{
 		const FQuat Rotation = Transform.GetRotation();
 		const FVector Translation = Transform.GetTranslation();
@@ -165,7 +180,7 @@ struct FGLTFConverterUtility
 		return VerticalFOV;
 	}
 
-	static FGLTFJsonQuaternion ConvertCameraDirection()
+	static FGLTFQuaternion ConvertCameraDirection()
 	{
 		// Unreal uses +X axis as camera direction in Unreal coordinates.
 		// glTF uses +Y as camera direction in Unreal coordinates.
@@ -173,7 +188,7 @@ struct FGLTFConverterUtility
 		return ConvertRotation(FRotator(0, 90, 0).Quaternion());
 	}
 
-	static FGLTFJsonQuaternion ConvertLightDirection()
+	static FGLTFQuaternion ConvertLightDirection()
 	{
 		// Unreal uses +X axis as light direction in Unreal coordinates.
 		// glTF uses +Y as light direction in Unreal coordinates.
diff --git a/Engine/Plugins/Enterprise/GLTFExporter/MigratedFromGit/Source/GLTFExporter/Private/Converters/GLTFLightConverters.cpp b/Engine/Plugins/Enterprise/GLTFExporter/MigratedFromGit/Source/GLTFExporter/Private/Converters/GLTFLightConverters.cpp
--- a/Engine/Plugins/Enterprise/GLTFExporter/MigratedFromGit/Source/GLTFExporter/Private/Converters/GLTFLightConverters.cpp
+++ b/Engine/Plugins/Enterprise/GLTFExporter/MigratedFromGit/Source/GLTFExporter/Private/Converters/GLTFLightConverters.cpp
@@ -21,7 +21,7 @@ FGLTFJsonLightIndex FGLTFLightConverter::Convert(const ULightComponent* LightCom
 	Light.Intensity = LightComponent->Intensity;
 
 	const FLinearColor TemperatureColor = LightComponent->bUseTemperature ? FLinearColor::MakeFromColorTemperature(LightComponent->Temperature) : FLinearColor::White;
-	Light.Color = FGLTFConverterUtility::ConvertColor(TemperatureColor * LightComponent->GetLightColor(), Builder.ExportOptions->bStrictCompliance);
+	Light.Color = FGLTFConverterUtility::ConvertColor3(TemperatureColor * LightComponent->GetLightColor(), Builder.ExportOptions->bStrictCompliance);
 
 	if (const UPointLightComponent* PointLightComponent = Cast<UPointLightComponent>(LightComponent))
 	{
diff --git a/Engine/Plugins/Enterprise/GLTFExporter/MigratedFromGit/Source/GLTFExporter/Private/Converters/GLTFPackedTypes.h b/Engine/Plugins/Enterprise/GLTFExporter/MigratedFromGit/Source/GLTFExporter/Private/Converters/GLTFPackedTypes.h
deleted file mode 100644
diff --git a/Engine/Plugins/Enterprise/GLTFExporter/MigratedFromGit/Source/GLTFExporter/Private/Converters/GLTFRawTypes.h b/Engine/Plugins/Enterprise/GLTFExporter/MigratedFromGit/Source/GLTFExporter/Private/Converters/GLTFRawTypes.h
deleted file mode 100644
diff --git a/Engine/Plugins/Enterprise/GLTFExporter/MigratedFromGit/Source/GLTFExporter/Private/Converters/GLTFSkinConverters.cpp b/Engine/Plugins/Enterprise/GLTFExporter/MigratedFromGit/Source/GLTFExporter/Private/Converters/GLTFSkinConverters.cpp
--- a/Engine/Plugins/Enterprise/GLTFExporter/MigratedFromGit/Source/GLTFExporter/Private/Converters/GLTFSkinConverters.cpp
+++ b/Engine/Plugins/Enterprise/GLTFExporter/MigratedFromGit/Source/GLTFExporter/Private/Converters/GLTFSkinConverters.cpp
@@ -25,7 +25,7 @@ FGLTFJsonSkinIndex FGLTFSkinConverter::Convert(FGLTFJsonNodeIndex RootNode, cons
 		Skin.Joints[BoneIndex] = Builder.GetOrAddNode(RootNode, SkeletalMesh, BoneIndex);
 	}
 
-	TArray<FGLTFRawMatrix4> InverseBindMatrices;
+	TArray<FGLTFMatrix4> InverseBindMatrices;
 	InverseBindMatrices.AddUninitialized(BoneCount);
 
 	for (int32 BoneIndex = 0; BoneIndex < BoneCount; ++BoneIndex)
diff --git a/Engine/Plugins/Enterprise/GLTFExporter/MigratedFromGit/Source/GLTFExporter/Private/Converters/GLTFSkySphereConverters.cpp b/Engine/Plugins/Enterprise/GLTFExporter/MigratedFromGit/Source/GLTFExporter/Private/Converters/GLTFSkySphereConverters.cpp
--- a/Engine/Plugins/Enterprise/GLTFExporter/MigratedFromGit/Source/GLTFExporter/Private/Converters/GLTFSkySphereConverters.cpp
+++ b/Engine/Plugins/Enterprise/GLTFExporter/MigratedFromGit/Source/GLTFExporter/Private/Converters/GLTFSkySphereConverters.cpp
@@ -102,7 +102,7 @@ void FGLTFSkySphereConverter::ConvertProperty(const AActor* Actor, const TCHAR*
 	}
 }
 
-void FGLTFSkySphereConverter::ConvertColorProperty(const AActor* Actor, const TCHAR* PropertyName, FGLTFJsonColor4& OutValue) const
+void FGLTFSkySphereConverter::ConvertColorProperty(const AActor* Actor, const TCHAR* PropertyName, FGLTFColor4& OutValue) const
 {
 	check(Actor != nullptr);
 
diff --git a/Engine/Plugins/Enterprise/GLTFExporter/MigratedFromGit/Source/GLTFExporter/Private/Converters/GLTFSkySphereConverters.h b/Engine/Plugins/Enterprise/GLTFExporter/MigratedFromGit/Source/GLTFExporter/Private/Converters/GLTFSkySphereConverters.h
--- a/Engine/Plugins/Enterprise/GLTFExporter/MigratedFromGit/Source/GLTFExporter/Private/Converters/GLTFSkySphereConverters.h
+++ b/Engine/Plugins/Enterprise/GLTFExporter/MigratedFromGit/Source/GLTFExporter/Private/Converters/GLTFSkySphereConverters.h
@@ -2,12 +2,12 @@
 
 #pragma once
 
+#include "Core/GLTFColor.h"
 #include "Json/GLTFJsonIndex.h"
 #include "Converters/GLTFConverter.h"
 #include "Converters/GLTFBuilderContext.h"
 #include "Engine.h"
 
-struct FGLTFJsonColor4;
 struct FGLTFJsonSkySphereColorCurve;
 
 class FGLTFSkySphereConverter final : public FGLTFBuilderContext, public TGLTFConverter<FGLTFJsonSkySphereIndex, const AActor*>
@@ -26,7 +26,7 @@ class FGLTFSkySphereConverter final : public FGLTFBuilderContext, public TGLTFCo
 	template <class ValueType>
 	void ConvertProperty(const AActor* Actor, const TCHAR* PropertyName, ValueType& OutValue) const;
 
-	void ConvertColorProperty(const AActor* Actor, const TCHAR* PropertyName, FGLTFJsonColor4& OutValue) const;
+	void ConvertColorProperty(const AActor* Actor, const TCHAR* PropertyName, FGLTFColor4& OutValue) const;
 	void ConvertColorCurveProperty(const AActor* Actor, const TCHAR* PropertyName, FGLTFJsonSkySphereColorCurve& OutValue) const;
 	void ConvertScalarParameter(const AActor* Actor, const UMaterialInstance* Material, const TCHAR* ParameterName, float& OutValue) const;
 
diff --git a/Engine/Plugins/Enterprise/GLTFExporter/MigratedFromGit/Source/GLTFExporter/Private/Core/GLTFColor.h b/Engine/Plugins/Enterprise/GLTFExporter/MigratedFromGit/Source/GLTFExporter/Private/Core/GLTFColor.h
new file mode 100644
--- /dev/null
+++ b/Engine/Plugins/Enterprise/GLTFExporter/MigratedFromGit/Source/GLTFExporter/Private/Core/GLTFColor.h
@@ -0,0 +1,51 @@
+// Copyright Epic Games, Inc. All Rights Reserved.
+
+#pragma once
+
+#include "Core/GLTFCore.h"
+
+#if PLATFORM_SUPPORTS_PRAGMA_PACK
+	#pragma pack(push, 1)
+#endif
+
+template <typename ComponentType>
+struct TGLTFColor3
+{
+	union
+	{
+		struct
+		{
+			ComponentType R GCC_PACK(1);
+			ComponentType G GCC_PACK(1);
+			ComponentType B GCC_PACK(1);
+		};
+
+		ComponentType Components[3] GCC_PACK(1);
+	};
+};
+
+template <typename ComponentType>
+struct TGLTFColor4
+{
+	union
+	{
+		struct
+		{
+			ComponentType R GCC_PACK(1);
+			ComponentType G GCC_PACK(1);
+			ComponentType B GCC_PACK(1);
+			ComponentType A GCC_PACK(1);
+		};
+
+		ComponentType Components[4] GCC_PACK(1);
+	};
+};
+
+#if PLATFORM_SUPPORTS_PRAGMA_PACK
+	#pragma pack(pop)
+#endif
+
+typedef TGLTFColor3<float> FGLTFColor3;
+typedef TGLTFColor4<float> FGLTFColor4;
+
+typedef TGLTFColor4<uint8> FGLTFUInt8Color4;
diff --git a/Engine/Plugins/Enterprise/GLTFExporter/MigratedFromGit/Source/GLTFExporter/Private/Core/GLTFCore.h b/Engine/Plugins/Enterprise/GLTFExporter/MigratedFromGit/Source/GLTFExporter/Private/Core/GLTFCore.h
new file mode 100644
--- /dev/null
+++ b/Engine/Plugins/Enterprise/GLTFExporter/MigratedFromGit/Source/GLTFExporter/Private/Core/GLTFCore.h
@@ -0,0 +1,5 @@
+// Copyright Epic Games, Inc. All Rights Reserved.
+
+#pragma once
+
+#include "CoreTypes.h"
diff --git a/Engine/Plugins/Enterprise/GLTFExporter/MigratedFromGit/Source/GLTFExporter/Private/Core/GLTFMatrix.h b/Engine/Plugins/Enterprise/GLTFExporter/MigratedFromGit/Source/GLTFExporter/Private/Core/GLTFMatrix.h
new file mode 100644
--- /dev/null
+++ b/Engine/Plugins/Enterprise/GLTFExporter/MigratedFromGit/Source/GLTFExporter/Private/Core/GLTFMatrix.h
@@ -0,0 +1,33 @@
+// Copyright Epic Games, Inc. All Rights Reserved.
+
+#pragma once
+
+#include "Core/GLTFCore.h"
+
+#if PLATFORM_SUPPORTS_PRAGMA_PACK
+	#pragma pack(push, 1)
+#endif
+
+template <typename ElementType, SIZE_T Dimensions>
+struct TGLTFMatrix
+{
+	ElementType Elements[Dimensions * Dimensions] GCC_PACK(1);
+
+	ElementType& operator()(uint32 Column, uint32 Row)
+	{
+		return Elements[Column * Dimensions + Row];
+	}
+
+	const ElementType& operator()(uint32 Column, uint32 Row) const
+	{
+		return Elements[Column * Dimensions + Row];
+	}
+};
+
+#if PLATFORM_SUPPORTS_PRAGMA_PACK
+	#pragma pack(pop)
+#endif
+
+typedef TGLTFMatrix<float, 2> FGLTFMatrix2;
+typedef TGLTFMatrix<float, 3> FGLTFMatrix3;
+typedef TGLTFMatrix<float, 4> FGLTFMatrix4;
diff --git a/Engine/Plugins/Enterprise/GLTFExporter/MigratedFromGit/Source/GLTFExporter/Private/Core/GLTFQuaternion.h b/Engine/Plugins/Enterprise/GLTFExporter/MigratedFromGit/Source/GLTFExporter/Private/Core/GLTFQuaternion.h
new file mode 100644
--- /dev/null
+++ b/Engine/Plugins/Enterprise/GLTFExporter/MigratedFromGit/Source/GLTFExporter/Private/Core/GLTFQuaternion.h
@@ -0,0 +1,32 @@
+// Copyright Epic Games, Inc. All Rights Reserved.
+
+#pragma once
+
+#include "Core/GLTFCore.h"
+
+#if PLATFORM_SUPPORTS_PRAGMA_PACK
+	#pragma pack(push, 1)
+#endif
+
+template <typename ComponentType>
+struct TGLTFQuaternion
+{
+	union
+	{
+		struct
+		{
+			ComponentType X GCC_PACK(1);
+			ComponentType Y GCC_PACK(1);
+			ComponentType Z GCC_PACK(1);
+			ComponentType W GCC_PACK(1);
+		};
+
+		ComponentType Components[4] GCC_PACK(1);
+	};
+};
+
+#if PLATFORM_SUPPORTS_PRAGMA_PACK
+	#pragma pack(pop)
+#endif
+
+typedef TGLTFQuaternion<float> FGLTFQuaternion;
diff --git a/Engine/Plugins/Enterprise/GLTFExporter/MigratedFromGit/Source/GLTFExporter/Private/Core/GLTFVector.h b/Engine/Plugins/Enterprise/GLTFExporter/MigratedFromGit/Source/GLTFExporter/Private/Core/GLTFVector.h
new file mode 100644
--- /dev/null
+++ b/Engine/Plugins/Enterprise/GLTFExporter/MigratedFromGit/Source/GLTFExporter/Private/Core/GLTFVector.h
@@ -0,0 +1,68 @@
+// Copyright Epic Games, Inc. All Rights Reserved.
+
+#pragma once
+
+#include "Core/GLTFCore.h"
+
+#if PLATFORM_SUPPORTS_PRAGMA_PACK
+	#pragma pack(push, 1)
+#endif
+
+template <typename ComponentType>
+struct TGLTFVector2
+{
+	union
+	{
+		struct
+		{
+			ComponentType X GCC_PACK(1);
+			ComponentType Y GCC_PACK(1);
+		};
+
+		ComponentType Components[2] GCC_PACK(1);
+	};
+};
+
+template <typename ComponentType>
+struct TGLTFVector3
+{
+	union
+	{
+		struct
+		{
+			ComponentType X GCC_PACK(1);
+			ComponentType Y GCC_PACK(1);
+			ComponentType Z GCC_PACK(1);
+		};
+
+		ComponentType Components[3] GCC_PACK(1);
+	};
+};
+
+template <typename ComponentType>
+struct TGLTFVector4
+{
+	union
+	{
+		struct
+		{
+			ComponentType X GCC_PACK(1);
+			ComponentType Y GCC_PACK(1);
+			ComponentType Z GCC_PACK(1);
+			ComponentType W GCC_PACK(1);
+		};
+
+		ComponentType Components[4] GCC_PACK(1);
+	};
+};
+
+#if PLATFORM_SUPPORTS_PRAGMA_PACK
+	#pragma pack(pop)
+#endif
+
+typedef TGLTFVector2<float> FGLTFVector2;
+typedef TGLTFVector3<float> FGLTFVector3;
+typedef TGLTFVector4<float> FGLTFVector4;
+
+typedef TGLTFVector4<int8> FGLTFInt8Vector4;
+typedef TGLTFVector4<int16> FGLTFInt16Vector4;
diff --git a/Engine/Plugins/Enterprise/GLTFExporter/MigratedFromGit/Source/GLTFExporter/Private/Json/GLTFJsonBackdrop.h b/Engine/Plugins/Enterprise/GLTFExporter/MigratedFromGit/Source/GLTFExporter/Private/Json/GLTFJsonBackdrop.h
--- a/Engine/Plugins/Enterprise/GLTFExporter/MigratedFromGit/Source/GLTFExporter/Private/Json/GLTFJsonBackdrop.h
+++ b/Engine/Plugins/Enterprise/GLTFExporter/MigratedFromGit/Source/GLTFExporter/Private/Json/GLTFJsonBackdrop.h
@@ -4,7 +4,7 @@
 
 #include "Json/GLTFJsonObject.h"
 #include "Json/GLTFJsonIndex.h"
-#include "Json/GLTFJsonVector3.h"
+#include "Json/GLTFJsonVector.h"
 
 struct FGLTFJsonBackdrop : IGLTFJsonObject
 {
diff --git a/Engine/Plugins/Enterprise/GLTFExporter/MigratedFromGit/Source/GLTFExporter/Private/Json/GLTFJsonCameraControl.h b/Engine/Plugins/Enterprise/GLTFExporter/MigratedFromGit/Source/GLTFExporter/Private/Json/GLTFJsonCameraControl.h
--- a/Engine/Plugins/Enterprise/GLTFExporter/MigratedFromGit/Source/GLTFExporter/Private/Json/GLTFJsonCameraControl.h
+++ b/Engine/Plugins/Enterprise/GLTFExporter/MigratedFromGit/Source/GLTFExporter/Private/Json/GLTFJsonCameraControl.h
@@ -4,7 +4,7 @@
 
 #include "Json/GLTFJsonObject.h"
 #include "Json/GLTFJsonIndex.h"
-#include "Json/GLTFJsonVector4.h"
+#include "Json/GLTFJsonVector.h"
 
 struct FGLTFJsonCameraControl : IGLTFJsonObject
 {
diff --git a/Engine/Plugins/Enterprise/GLTFExporter/MigratedFromGit/Source/GLTFExporter/Private/Json/GLTFJsonColor.cpp b/Engine/Plugins/Enterprise/GLTFExporter/MigratedFromGit/Source/GLTFExporter/Private/Json/GLTFJsonColor.cpp
new file mode 100644
--- /dev/null
+++ b/Engine/Plugins/Enterprise/GLTFExporter/MigratedFromGit/Source/GLTFExporter/Private/Json/GLTFJsonColor.cpp
@@ -0,0 +1,9 @@
+// Copyright Epic Games, Inc. All Rights Reserved.
+
+#include "Json/GLTFJsonColor.h"
+
+const FGLTFJsonColor3 TGLTFJsonColor<FGLTFColor3>::Black({ 0, 0, 0 });
+const FGLTFJsonColor3 TGLTFJsonColor<FGLTFColor3>::White({ 1, 1, 1 });
+
+const FGLTFJsonColor4 TGLTFJsonColor<FGLTFColor4>::Black({ 0, 0, 0, 1 });
+const FGLTFJsonColor4 TGLTFJsonColor<FGLTFColor4>::White({ 1, 1, 1, 1 });
diff --git a/Engine/Plugins/Enterprise/GLTFExporter/MigratedFromGit/Source/GLTFExporter/Private/Json/GLTFJsonColor.h b/Engine/Plugins/Enterprise/GLTFExporter/MigratedFromGit/Source/GLTFExporter/Private/Json/GLTFJsonColor.h
new file mode 100644
--- /dev/null
+++ b/Engine/Plugins/Enterprise/GLTFExporter/MigratedFromGit/Source/GLTFExporter/Private/Json/GLTFJsonColor.h
@@ -0,0 +1,53 @@
+// Copyright Epic Games, Inc. All Rights Reserved.
+
+#pragma once
+
+#include "Core/GLTFColor.h"
+#include "Json/GLTFJsonArray.h"
+
+template <typename BaseType>
+struct TGLTFJsonColor final : BaseType, IGLTFJsonArray
+{
+	static const TGLTFJsonColor Black;
+	static const TGLTFJsonColor White;
+
+	TGLTFJsonColor(const BaseType& Other)
+		: BaseType(Other)
+	{
+	}
+
+	TGLTFJsonColor& operator=(const BaseType& Other)
+	{
+		*static_cast<BaseType*>(this) = Other;
+		return *this;
+	}
+
+	virtual void WriteArray(IGLTFJsonWriter& Writer) const override
+	{
+		for (int32 i = 0; i < GetNum(BaseType::Components); ++i)
+		{
+			Writer.Write(BaseType::Components[i]);
+		}
+	}
+
+	bool operator==(const BaseType& Other) const
+	{
+		for (uint32 i = 0; i < GetNum(BaseType::Components); ++i)
+		{
+			if (BaseType::Components[i] != Other.Components[i])
+			{
+				return false;
+			}
+		}
+
+		return true;
+	}
+
+	bool operator!=(const BaseType& Other) const
+	{
+		return !(*this == Other);
+	}
+};
+
+typedef TGLTFJsonColor<FGLTFColor3> FGLTFJsonColor3;
+typedef TGLTFJsonColor<FGLTFColor4> FGLTFJsonColor4;
diff --git a/Engine/Plugins/Enterprise/GLTFExporter/MigratedFromGit/Source/GLTFExporter/Private/Json/GLTFJsonColor3.cpp b/Engine/Plugins/Enterprise/GLTFExporter/MigratedFromGit/Source/GLTFExporter/Private/Json/GLTFJsonColor3.cpp
deleted file mode 100644
diff --git a/Engine/Plugins/Enterprise/GLTFExporter/MigratedFromGit/Source/GLTFExporter/Private/Json/GLTFJsonColor3.h b/Engine/Plugins/Enterprise/GLTFExporter/MigratedFromGit/Source/GLTFExporter/Private/Json/GLTFJsonColor3.h
deleted file mode 100644
diff --git a/Engine/Plugins/Enterprise/GLTFExporter/MigratedFromGit/Source/GLTFExporter/Private/Json/GLTFJsonColor4.cpp b/Engine/Plugins/Enterprise/GLTFExporter/MigratedFromGit/Source/GLTFExporter/Private/Json/GLTFJsonColor4.cpp
deleted file mode 100644
diff --git a/Engine/Plugins/Enterprise/GLTFExporter/MigratedFromGit/Source/GLTFExporter/Private/Json/GLTFJsonColor4.h b/Engine/Plugins/Enterprise/GLTFExporter/MigratedFromGit/Source/GLTFExporter/Private/Json/GLTFJsonColor4.h
deleted file mode 100644
diff --git a/Engine/Plugins/Enterprise/GLTFExporter/MigratedFromGit/Source/GLTFExporter/Private/Json/GLTFJsonLight.h b/Engine/Plugins/Enterprise/GLTFExporter/MigratedFromGit/Source/GLTFExporter/Private/Json/GLTFJsonLight.h
--- a/Engine/Plugins/Enterprise/GLTFExporter/MigratedFromGit/Source/GLTFExporter/Private/Json/GLTFJsonLight.h
+++ b/Engine/Plugins/Enterprise/GLTFExporter/MigratedFromGit/Source/GLTFExporter/Private/Json/GLTFJsonLight.h
@@ -4,7 +4,7 @@
 
 #include "Json/GLTFJsonObject.h"
 #include "Json/GLTFJsonEnums.h"
-#include "Json/GLTFJsonColor3.h"
+#include "Json/GLTFJsonColor.h"
 
 struct FGLTFJsonSpotLight : IGLTFJsonObject
 {
diff --git a/Engine/Plugins/Enterprise/GLTFExporter/MigratedFromGit/Source/GLTFExporter/Private/Json/GLTFJsonLightMap.h b/Engine/Plugins/Enterprise/GLTFExporter/MigratedFromGit/Source/GLTFExporter/Private/Json/GLTFJsonLightMap.h
--- a/Engine/Plugins/Enterprise/GLTFExporter/MigratedFromGit/Source/GLTFExporter/Private/Json/GLTFJsonLightMap.h
+++ b/Engine/Plugins/Enterprise/GLTFExporter/MigratedFromGit/Source/GLTFExporter/Private/Json/GLTFJsonLightMap.h
@@ -4,7 +4,7 @@
 
 #include "Json/GLTFJsonObject.h"
 #include "Json/GLTFJsonMaterial.h"
-#include "Json/GLTFJsonVector4.h"
+#include "Json/GLTFJsonVector.h"
 
 struct FGLTFJsonLightMap : IGLTFJsonObject
 {
@@ -17,7 +17,7 @@ struct FGLTFJsonLightMap : IGLTFJsonObject
 	FGLTFJsonLightMap()
 		: LightMapScale(FGLTFJsonVector4::One)
 		, LightMapAdd(FGLTFJsonVector4::Zero)
-		, CoordinateScaleBias(1, 1, 0, 0)
+		, CoordinateScaleBias({ 1, 1, 0, 0 })
 	{
 	}
 
diff --git a/Engine/Plugins/Enterprise/GLTFExporter/MigratedFromGit/Source/GLTFExporter/Private/Json/GLTFJsonMaterial.h b/Engine/Plugins/Enterprise/GLTFExporter/MigratedFromGit/Source/GLTFExporter/Private/Json/GLTFJsonMaterial.h
--- a/Engine/Plugins/Enterprise/GLTFExporter/MigratedFromGit/Source/GLTFExporter/Private/Json/GLTFJsonMaterial.h
+++ b/Engine/Plugins/Enterprise/GLTFExporter/MigratedFromGit/Source/GLTFExporter/Private/Json/GLTFJsonMaterial.h
@@ -4,8 +4,7 @@
 
 #include "Json/GLTFJsonObject.h"
 #include "Json/GLTFJsonEnums.h"
-#include "Json/GLTFJsonColor3.h"
-#include "Json/GLTFJsonColor4.h"
+#include "Json/GLTFJsonColor.h"
 #include "Json/GLTFJsonTextureTransform.h"
 
 struct FGLTFJsonTextureInfo : IGLTFJsonObject
diff --git a/Engine/Plugins/Enterprise/GLTFExporter/MigratedFromGit/Source/GLTFExporter/Private/Json/GLTFJsonMatrix.cpp b/Engine/Plugins/Enterprise/GLTFExporter/MigratedFromGit/Source/GLTFExporter/Private/Json/GLTFJsonMatrix.cpp
new file mode 100644
--- /dev/null
+++ b/Engine/Plugins/Enterprise/GLTFExporter/MigratedFromGit/Source/GLTFExporter/Private/Json/GLTFJsonMatrix.cpp
@@ -0,0 +1,24 @@
+// Copyright Epic Games, Inc. All Rights Reserved.
+
+#include "Json/GLTFJsonMatrix.h"
+
+const FGLTFJsonMatrix2 TGLTFJsonMatrix<FGLTFMatrix2>::Identity
+({
+	1, 0,
+	0, 1
+});
+
+const FGLTFJsonMatrix3 TGLTFJsonMatrix<FGLTFMatrix3>::Identity
+({
+	1, 0, 0,
+	0, 1, 0,
+	0, 0, 1
+});
+
+const FGLTFJsonMatrix4 TGLTFJsonMatrix<FGLTFMatrix4>::Identity
+({
+	1, 0, 0, 0,
+	0, 1, 0, 0,
+	0, 0, 1, 0,
+	0, 0, 0, 1
+});
diff --git a/Engine/Plugins/Enterprise/GLTFExporter/MigratedFromGit/Source/GLTFExporter/Private/Json/GLTFJsonMatrix.h b/Engine/Plugins/Enterprise/GLTFExporter/MigratedFromGit/Source/GLTFExporter/Private/Json/GLTFJsonMatrix.h
new file mode 100644
--- /dev/null
+++ b/Engine/Plugins/Enterprise/GLTFExporter/MigratedFromGit/Source/GLTFExporter/Private/Json/GLTFJsonMatrix.h
@@ -0,0 +1,53 @@
+// Copyright Epic Games, Inc. All Rights Reserved.
+
+#pragma once
+
+#include "Core/GLTFMatrix.h"
+#include "Json/GLTFJsonArray.h"
+
+template <typename BaseType>
+struct TGLTFJsonMatrix final : BaseType, IGLTFJsonArray
+{
+	static const TGLTFJsonMatrix Identity;
+
+	TGLTFJsonMatrix(const BaseType& Other)
+		: TGLTFMatrix(Other)
+	{
+	}
+
+	TGLTFJsonMatrix& operator=(const BaseType& Other)
+	{
+		*static_cast<BaseType*>(this) = Other;
+		return *this;
+	}
+
+	virtual void WriteArray(IGLTFJsonWriter& Writer) const override
+	{
+		for (int32 i = 0; i < GetNum(BaseType::Elements); ++i)
+		{
+			Writer.Write(BaseType::Elements[i]);
+		}
+	}
+
+	bool operator==(const BaseType& Other) const
+	{
+		for (uint32 i = 0; i < GetNum(BaseType::Elements); ++i)
+		{
+			if (BaseType::Elements[i] != Other.Elements[i])
+			{
+				return false;
+			}
+		}
+
+		return true;
+	}
+
+	bool operator!=(const BaseType& Other) const
+	{
+		return !(*this == Other);
+	}
+};
+
+typedef TGLTFJsonMatrix<FGLTFMatrix2> FGLTFJsonMatrix2;
+typedef TGLTFJsonMatrix<FGLTFMatrix3> FGLTFJsonMatrix3;
+typedef TGLTFJsonMatrix<FGLTFMatrix4> FGLTFJsonMatrix4;
diff --git a/Engine/Plugins/Enterprise/GLTFExporter/MigratedFromGit/Source/GLTFExporter/Private/Json/GLTFJsonMatrix4.cpp b/Engine/Plugins/Enterprise/GLTFExporter/MigratedFromGit/Source/GLTFExporter/Private/Json/GLTFJsonMatrix4.cpp
deleted file mode 100644
diff --git a/Engine/Plugins/Enterprise/GLTFExporter/MigratedFromGit/Source/GLTFExporter/Private/Json/GLTFJsonMatrix4.h b/Engine/Plugins/Enterprise/GLTFExporter/MigratedFromGit/Source/GLTFExporter/Private/Json/GLTFJsonMatrix4.h
deleted file mode 100644
diff --git a/Engine/Plugins/Enterprise/GLTFExporter/MigratedFromGit/Source/GLTFExporter/Private/Json/GLTFJsonNode.h b/Engine/Plugins/Enterprise/GLTFExporter/MigratedFromGit/Source/GLTFExporter/Private/Json/GLTFJsonNode.h
--- a/Engine/Plugins/Enterprise/GLTFExporter/MigratedFromGit/Source/GLTFExporter/Private/Json/GLTFJsonNode.h
+++ b/Engine/Plugins/Enterprise/GLTFExporter/MigratedFromGit/Source/GLTFExporter/Private/Json/GLTFJsonNode.h
@@ -4,8 +4,8 @@
 
 #include "Json/GLTFJsonObject.h"
 #include "Json/GLTFJsonIndex.h"
-#include "Json/GLTFJsonMatrix4.h"
-#include "Json/GLTFJsonVector3.h"
+#include "Json/GLTFJsonMatrix.h"
+#include "Json/GLTFJsonVector.h"
 #include "Json/GLTFJsonQuaternion.h"
 
 struct FGLTFJsonNode : IGLTFJsonObject
diff --git a/Engine/Plugins/Enterprise/GLTFExporter/MigratedFromGit/Source/GLTFExporter/Private/Json/GLTFJsonQuaternion.cpp b/Engine/Plugins/Enterprise/GLTFExporter/MigratedFromGit/Source/GLTFExporter/Private/Json/GLTFJsonQuaternion.cpp
--- a/Engine/Plugins/Enterprise/GLTFExporter/MigratedFromGit/Source/GLTFExporter/Private/Json/GLTFJsonQuaternion.cpp
+++ b/Engine/Plugins/Enterprise/GLTFExporter/MigratedFromGit/Source/GLTFExporter/Private/Json/GLTFJsonQuaternion.cpp
@@ -2,4 +2,4 @@
 
 #include "Json/GLTFJsonQuaternion.h"
 
-const FGLTFJsonQuaternion FGLTFJsonQuaternion::Identity(0.0f, 0.0f, 0.0f, 1.0f);
+const FGLTFQuaternion FGLTFJsonQuaternion::Identity = { 0, 0, 0, 1 };
diff --git a/Engine/Plugins/Enterprise/GLTFExporter/MigratedFromGit/Source/GLTFExporter/Private/Json/GLTFJsonQuaternion.h b/Engine/Plugins/Enterprise/GLTFExporter/MigratedFromGit/Source/GLTFExporter/Private/Json/GLTFJsonQuaternion.h
--- a/Engine/Plugins/Enterprise/GLTFExporter/MigratedFromGit/Source/GLTFExporter/Private/Json/GLTFJsonQuaternion.h
+++ b/Engine/Plugins/Enterprise/GLTFExporter/MigratedFromGit/Source/GLTFExporter/Private/Json/GLTFJsonQuaternion.h
@@ -2,46 +2,47 @@
 
 #pragma once
 
+#include "Core/GLTFQuaternion.h"
 #include "Json/GLTFJsonArray.h"
-#include "Converters/GLTFRawTypes.h"
 
-struct FGLTFJsonQuaternion : IGLTFJsonArray
+struct FGLTFJsonQuaternion final : TGLTFQuaternion<float>, IGLTFJsonArray
 {
-	float X, Y, Z, W;
+	static const TGLTFQuaternion Identity;
 
-	static const FGLTFJsonQuaternion Identity;
-
-	FGLTFJsonQuaternion(float X, float Y, float Z, float W)
-		: X(X), Y(Y), Z(Z), W(W)
+	FGLTFJsonQuaternion(const TGLTFQuaternion& Other)
+		: TGLTFQuaternion(Other)
 	{
 	}
 
-	FGLTFJsonQuaternion(const FGLTFRawQuaternion& Raw)
-		: X(Raw.X), Y(Raw.Y), Z(Raw.Z), W(Raw.W)
+	FGLTFJsonQuaternion& operator=(const TGLTFQuaternion& Other)
 	{
+		*static_cast<TGLTFQuaternion*>(this) = Other;
+		return *this;
 	}
 
 	virtual void WriteArray(IGLTFJsonWriter& Writer) const override
 	{
-		Writer.Write(X);
-		Writer.Write(Y);
-		Writer.Write(Z);
-		Writer.Write(W);
+		for (int32 i = 0; i < 4; ++i)
+		{
+			Writer.Write(Components[i]);
+		}
 	}
 
-	bool operator==(const FGLTFJsonQuaternion& Other) const
+	bool operator==(const TGLTFQuaternion& Other) const
 	{
-		return X == Other.X
-			&& Y == Other.Y
-			&& Z == Other.Z
-			&& W == Other.W;
+		for (uint32 i = 0; i < GetNum(Components); ++i)
+		{
+			if (Components[i] != Other.Components[i])
+			{
+				return false;
+			}
+		}
+
+		return true;
 	}
 
-	bool operator!=(const FGLTFJsonQuaternion& Other) const
+	bool operator!=(const TGLTFQuaternion& Other) const
 	{
-		return X != Other.X
-			|| Y != Other.Y
-			|| Z != Other.Z
-			|| W != Other.W;
+		return !(*this == Other);
 	}
 };
diff --git a/Engine/Plugins/Enterprise/GLTFExporter/MigratedFromGit/Source/GLTFExporter/Private/Json/GLTFJsonSkySphere.h b/Engine/Plugins/Enterprise/GLTFExporter/MigratedFromGit/Source/GLTFExporter/Private/Json/GLTFJsonSkySphere.h
--- a/Engine/Plugins/Enterprise/GLTFExporter/MigratedFromGit/Source/GLTFExporter/Private/Json/GLTFJsonSkySphere.h
+++ b/Engine/Plugins/Enterprise/GLTFExporter/MigratedFromGit/Source/GLTFExporter/Private/Json/GLTFJsonSkySphere.h
@@ -4,7 +4,7 @@
 
 #include "Json/GLTFJsonObject.h"
 #include "Json/GLTFJsonIndex.h"
-#include "Json/GLTFJsonColor4.h"
+#include "Json/GLTFJsonColor.h"
 
 struct FGLTFJsonSkySphereColorCurve : IGLTFJsonArray
 {
diff --git a/Engine/Plugins/Enterprise/GLTFExporter/MigratedFromGit/Source/GLTFExporter/Private/Json/GLTFJsonTextureTransform.h b/Engine/Plugins/Enterprise/GLTFExporter/MigratedFromGit/Source/GLTFExporter/Private/Json/GLTFJsonTextureTransform.h
--- a/Engine/Plugins/Enterprise/GLTFExporter/MigratedFromGit/Source/GLTFExporter/Private/Json/GLTFJsonTextureTransform.h
+++ b/Engine/Plugins/Enterprise/GLTFExporter/MigratedFromGit/Source/GLTFExporter/Private/Json/GLTFJsonTextureTransform.h
@@ -3,7 +3,7 @@
 #pragma once
 
 #include "Json/GLTFJsonObject.h"
-#include "Json/GLTFJsonVector2.h"
+#include "Json/GLTFJsonVector.h"
 
 struct FGLTFJsonTextureTransform : IGLTFJsonObject
 {
diff --git a/Engine/Plugins/Enterprise/GLTFExporter/MigratedFromGit/Source/GLTFExporter/Private/Json/GLTFJsonVector.cpp b/Engine/Plugins/Enterprise/GLTFExporter/MigratedFromGit/Source/GLTFExporter/Private/Json/GLTFJsonVector.cpp
new file mode 100644
--- /dev/null
+++ b/Engine/Plugins/Enterprise/GLTFExporter/MigratedFromGit/Source/GLTFExporter/Private/Json/GLTFJsonVector.cpp
@@ -0,0 +1,12 @@
+// Copyright Epic Games, Inc. All Rights Reserved.
+
+#include "Json/GLTFJsonVector.h"
+
+const FGLTFJsonVector2 TGLTFJsonVector<FGLTFVector2>::Zero({ 0, 0 });
+const FGLTFJsonVector2 TGLTFJsonVector<FGLTFVector2>::One ({ 1, 1 });
+
+const FGLTFJsonVector3 TGLTFJsonVector<FGLTFVector3>::Zero({ 0, 0, 0 });
+const FGLTFJsonVector3 TGLTFJsonVector<FGLTFVector3>::One ({ 1, 1, 1 });
+
+const FGLTFJsonVector4 TGLTFJsonVector<FGLTFVector4>::Zero({ 0, 0, 0, 0 });
+const FGLTFJsonVector4 TGLTFJsonVector<FGLTFVector4>::One ({ 1, 1, 1, 1 });
diff --git a/Engine/Plugins/Enterprise/GLTFExporter/MigratedFromGit/Source/GLTFExporter/Private/Json/GLTFJsonVector.h b/Engine/Plugins/Enterprise/GLTFExporter/MigratedFromGit/Source/GLTFExporter/Private/Json/GLTFJsonVector.h
new file mode 100644
--- /dev/null
+++ b/Engine/Plugins/Enterprise/GLTFExporter/MigratedFromGit/Source/GLTFExporter/Private/Json/GLTFJsonVector.h
@@ -0,0 +1,54 @@
+// Copyright Epic Games, Inc. All Rights Reserved.
+
+#pragma once
+
+#include "Core/GLTFVector.h"
+#include "Json/GLTFJsonArray.h"
+
+template <typename BaseType>
+struct TGLTFJsonVector final : BaseType, IGLTFJsonArray
+{
+	static const TGLTFJsonVector Zero;
+	static const TGLTFJsonVector One;
+
+	TGLTFJsonVector(const BaseType& Other)
+		: BaseType(Other)
+	{
+	}
+
+	TGLTFJsonVector& operator=(const BaseType& Other)
+	{
+		*static_cast<BaseType*>(this) = Other;
+		return *this;
+	}
+
+	virtual void WriteArray(IGLTFJsonWriter& Writer) const override
+	{
+		for (int32 i = 0; i < GetNum(BaseType::Components); ++i)
+		{
+			Writer.Write(BaseType::Components[i]);
+		}
+	}
+
+	bool operator==(const BaseType& Other) const
+	{
+		for (uint32 i = 0; i < GetNum(BaseType::Components); ++i)
+		{
+			if (BaseType::Components[i] != Other.Components[i])
+			{
+				return false;
+			}
+		}
+
+		return true;
+	}
+
+	bool operator!=(const BaseType& Other) const
+	{
+		return !(*this == Other);
+	}
+};
+
+typedef TGLTFJsonVector<FGLTFVector2> FGLTFJsonVector2;
+typedef TGLTFJsonVector<FGLTFVector3> FGLTFJsonVector3;
+typedef TGLTFJsonVector<FGLTFVector4> FGLTFJsonVector4;
diff --git a/Engine/Plugins/Enterprise/GLTFExporter/MigratedFromGit/Source/GLTFExporter/Private/Json/GLTFJsonVector2.cpp b/Engine/Plugins/Enterprise/GLTFExporter/MigratedFromGit/Source/GLTFExporter/Private/Json/GLTFJsonVector2.cpp
deleted file mode 100644
diff --git a/Engine/Plugins/Enterprise/GLTFExporter/MigratedFromGit/Source/GLTFExporter/Private/Json/GLTFJsonVector2.h b/Engine/Plugins/Enterprise/GLTFExporter/MigratedFromGit/Source/GLTFExporter/Private/Json/GLTFJsonVector2.h
deleted file mode 100644
diff --git a/Engine/Plugins/Enterprise/GLTFExporter/MigratedFromGit/Source/GLTFExporter/Private/Json/GLTFJsonVector3.cpp b/Engine/Plugins/Enterprise/GLTFExporter/MigratedFromGit/Source/GLTFExporter/Private/Json/GLTFJsonVector3.cpp
deleted file mode 100644
diff --git a/Engine/Plugins/Enterprise/GLTFExporter/MigratedFromGit/Source/GLTFExporter/Private/Json/GLTFJsonVector3.h b/Engine/Plugins/Enterprise/GLTFExporter/MigratedFromGit/Source/GLTFExporter/Private/Json/GLTFJsonVector3.h
deleted file mode 100644
diff --git a/Engine/Plugins/Enterprise/GLTFExporter/MigratedFromGit/Source/GLTFExporter/Private/Json/GLTFJsonVector4.cpp b/Engine/Plugins/Enterprise/GLTFExporter/MigratedFromGit/Source/GLTFExporter/Private/Json/GLTFJsonVector4.cpp
deleted file mode 100644
diff --git a/Engine/Plugins/Enterprise/GLTFExporter/MigratedFromGit/Source/GLTFExporter/Private/Json/GLTFJsonVector4.h b/Engine/Plugins/Enterprise/GLTFExporter/MigratedFromGit/Source/GLTFExporter/Private/Json/GLTFJsonVector4.h
deleted file mode 100644
diff --git a/Engine/Plugins/Enterprise/GLTFExporter/MigratedFromGit/Source/GLTFExporter/Private/Tasks/GLTFAnimationTasks.cpp b/Engine/Plugins/Enterprise/GLTFExporter/MigratedFromGit/Source/GLTFExporter/Private/Tasks/GLTFAnimationTasks.cpp
--- a/Engine/Plugins/Enterprise/GLTFExporter/MigratedFromGit/Source/GLTFExporter/Private/Tasks/GLTFAnimationTasks.cpp
+++ b/Engine/Plugins/Enterprise/GLTFExporter/MigratedFromGit/Source/GLTFExporter/Private/Tasks/GLTFAnimationTasks.cpp
@@ -96,7 +96,7 @@ void FGLTFAnimSequenceTask::Complete()
 
 		if (KeyPositions.Num() > 0)
 		{
-			TArray<FGLTFRawVector3> Translations;
+			TArray<FGLTFVector3> Translations;
 			Translations.AddUninitialized(KeyPositions.Num());
 
 			for (int32 Key = 0; Key < KeyPositions.Num(); ++Key)
@@ -128,7 +128,7 @@ void FGLTFAnimSequenceTask::Complete()
 
 		if (KeyRotations.Num() > 0)
 		{
-			TArray<FGLTFRawQuaternion> Rotations;
+			TArray<FGLTFQuaternion> Rotations;
 			Rotations.AddUninitialized(KeyRotations.Num());
 
 			for (int32 Key = 0; Key < KeyRotations.Num(); ++Key)
@@ -160,7 +160,7 @@ void FGLTFAnimSequenceTask::Complete()
 
 		if (KeyScales.Num() > 0)
 		{
-			TArray<FGLTFRawVector3> Scales;
+			TArray<FGLTFVector3> Scales;
 			Scales.AddUninitialized(KeyScales.Num());
 
 			for (int32 Key = 0; Key < KeyScales.Num(); ++Key)
@@ -297,7 +297,7 @@ void FGLTFLevelSequenceTask::Complete()
 
 					if (EnumHasAnyFlags(ChannelMask, EMovieSceneTransformChannel::Translation))
 					{
-						TArray<FGLTFRawVector3> Translations;
+						TArray<FGLTFVector3> Translations;
 						Translations.AddUninitialized(FrameCount);
 
 						for (int32 Frame = 0; Frame < FrameCount; ++Frame)
@@ -332,7 +332,7 @@ void FGLTFLevelSequenceTask::Complete()
 
 					if (EnumHasAnyFlags(ChannelMask, EMovieSceneTransformChannel::Rotation))
 					{
-						TArray<FGLTFRawQuaternion> Rotations;
+						TArray<FGLTFQuaternion> Rotations;
 						Rotations.AddUninitialized(FrameCount);
 
 						for (int32 Frame = 0; Frame < FrameCount; ++Frame)
@@ -370,7 +370,7 @@ void FGLTFLevelSequenceTask::Complete()
 
 					if (EnumHasAnyFlags(ChannelMask, EMovieSceneTransformChannel::Scale))
 					{
-						TArray<FGLTFRawVector3> Scales;
+						TArray<FGLTFVector3> Scales;
 						Scales.AddUninitialized(FrameCount);
 
 						for (int32 Frame = 0; Frame < FrameCount; ++Frame)
diff --git a/Engine/Plugins/Enterprise/GLTFExporter/MigratedFromGit/Source/GLTFExporter/Private/Tasks/GLTFMaterialTasks.cpp b/Engine/Plugins/Enterprise/GLTFExporter/MigratedFromGit/Source/GLTFExporter/Private/Tasks/GLTFMaterialTasks.cpp
--- a/Engine/Plugins/Enterprise/GLTFExporter/MigratedFromGit/Source/GLTFExporter/Private/Tasks/GLTFMaterialTasks.cpp
+++ b/Engine/Plugins/Enterprise/GLTFExporter/MigratedFromGit/Source/GLTFExporter/Private/Tasks/GLTFMaterialTasks.cpp
@@ -697,7 +697,7 @@ bool FGLTFMaterialTask::TryGetEmissive(FGLTFJsonMaterial& JsonMaterial, const FM
 	if (PropertyBakeOutput.bIsConstant)
 	{
 		const FLinearColor EmissiveColor = PropertyBakeOutput.ConstantValue;
-		JsonMaterial.EmissiveFactor = FGLTFConverterUtility::ConvertColor(EmissiveColor * EmissiveScale, Builder.ExportOptions->bStrictCompliance);
+		JsonMaterial.EmissiveFactor = FGLTFConverterUtility::ConvertColor3(EmissiveColor * EmissiveScale, Builder.ExportOptions->bStrictCompliance);
 	}
 	else
 	{
@@ -748,7 +748,7 @@ bool FGLTFMaterialTask::TryGetConstantColor(FGLTFJsonColor3& OutValue, const FMa
 	FLinearColor Value;
 	if (TryGetConstantColor(Value, Property))
 	{
-		OutValue = FGLTFConverterUtility::ConvertColor(Value, Builder.ExportOptions->bStrictCompliance);
+		OutValue = FGLTFConverterUtility::ConvertColor3(Value, Builder.ExportOptions->bStrictCompliance);
 		return true;
 	}
 
@@ -1123,7 +1123,7 @@ bool FGLTFMaterialTask::TryGetBakedMaterialProperty(FGLTFJsonTextureInfo& OutTex
 
 	if (PropertyBakeOutput.bIsConstant)
 	{
-		OutConstant = FGLTFConverterUtility::ConvertColor(PropertyBakeOutput.ConstantValue, Builder.ExportOptions->bStrictCompliance);
+		OutConstant = FGLTFConverterUtility::ConvertColor3(PropertyBakeOutput.ConstantValue, Builder.ExportOptions->bStrictCompliance);
 		return true;
 	}
 
