[GLTFExporter] feat: avoid json serializing vector/color/quaternion/matrix types if nearly equal to default value

commit cf4aa3642bc616d1a4f76175c000085bae9f3007
Date: Fri, 11 Jun 2021 02:40:07 +0200

    feat: avoid json serializing vector/color/quaternion/matrix types if nearly equal to default value
    (cherry picked from commit e8b34be68e2df2baa333ad42d1268e402249a740​)

#skipci

diff --git a/Engine/Plugins/Enterprise/GLTFExporter/MigratedFromGit/Source/GLTFExporter/Private/Converters/GLTFConverterUtility.h b/Engine/Plugins/Enterprise/GLTFExporter/MigratedFromGit/Source/GLTFExporter/Private/Converters/GLTFConverterUtility.h
--- a/Engine/Plugins/Enterprise/GLTFExporter/MigratedFromGit/Source/GLTFExporter/Private/Converters/GLTFConverterUtility.h
+++ b/Engine/Plugins/Enterprise/GLTFExporter/MigratedFromGit/Source/GLTFExporter/Private/Converters/GLTFConverterUtility.h
@@ -129,15 +129,6 @@ struct FGLTFConverterUtility
 		// also, as handedness is changed rotation is inversed - hence negation
 		// therefore glTFRotation = (-qX, -qZ, -qY, qw)
 
-		// Not checking if quaternion is normalized
-		// e.g. some sources use non-unit Quats for rotation tangents
-
-		// Return the identity quaternion when possible (depending on tolerance)
-		if (Rotation.Equals(FQuat::Identity)) // TODO: doesn't this potentially remove intentional deviance from identity?
-		{
-			return { 0, 0, 0, 1 }; // TODO: make static const
-		}
-
 		const FQuat Normalized = Rotation.GetNormalized();
 		return { -Normalized.X, -Normalized.Z, -Normalized.Y, Normalized.W };
 	}
diff --git a/Engine/Plugins/Enterprise/GLTFExporter/MigratedFromGit/Source/GLTFExporter/Private/Json/GLTFJsonColor.h b/Engine/Plugins/Enterprise/GLTFExporter/MigratedFromGit/Source/GLTFExporter/Private/Json/GLTFJsonColor.h
--- a/Engine/Plugins/Enterprise/GLTFExporter/MigratedFromGit/Source/GLTFExporter/Private/Json/GLTFJsonColor.h
+++ b/Engine/Plugins/Enterprise/GLTFExporter/MigratedFromGit/Source/GLTFExporter/Private/Json/GLTFJsonColor.h
@@ -30,11 +30,11 @@ struct TGLTFJsonColor final : BaseType, IGLTFJsonArray
 		}
 	}
 
-	bool operator==(const BaseType& Other) const
+	bool IsNearlyEqual(const BaseType& Other, float Tolerance = KINDA_SMALL_NUMBER) const
 	{
-		for (uint32 i = 0; i < GetNum(BaseType::Components); ++i)
+		for (int32 i = 0; i < GetNum(BaseType::Components); ++i)
 		{
-			if (BaseType::Components[i] != Other.Components[i])
+			if (!FMath::IsNearlyEqual(BaseType::Components[i], Other.Components[i], Tolerance))
 			{
 				return false;
 			}
@@ -42,11 +42,6 @@ struct TGLTFJsonColor final : BaseType, IGLTFJsonArray
 
 		return true;
 	}
-
-	bool operator!=(const BaseType& Other) const
-	{
-		return !(*this == Other);
-	}
 };
 
 typedef TGLTFJsonColor<FGLTFColor3> FGLTFJsonColor3;
diff --git a/Engine/Plugins/Enterprise/GLTFExporter/MigratedFromGit/Source/GLTFExporter/Private/Json/GLTFJsonLight.h b/Engine/Plugins/Enterprise/GLTFExporter/MigratedFromGit/Source/GLTFExporter/Private/Json/GLTFJsonLight.h
--- a/Engine/Plugins/Enterprise/GLTFExporter/MigratedFromGit/Source/GLTFExporter/Private/Json/GLTFJsonLight.h
+++ b/Engine/Plugins/Enterprise/GLTFExporter/MigratedFromGit/Source/GLTFExporter/Private/Json/GLTFJsonLight.h
@@ -61,7 +61,7 @@ struct FGLTFJsonLight : IGLTFJsonObject
 
 		Writer.Write(TEXT("type"), Type);
 
-		if (Color != FGLTFJsonColor3::White)
+		if (!Color.IsNearlyEqual(FGLTFJsonColor3::White))
 		{
 			Writer.Write(TEXT("color"), Color);
 		}
diff --git a/Engine/Plugins/Enterprise/GLTFExporter/MigratedFromGit/Source/GLTFExporter/Private/Json/GLTFJsonMaterial.h b/Engine/Plugins/Enterprise/GLTFExporter/MigratedFromGit/Source/GLTFExporter/Private/Json/GLTFJsonMaterial.h
--- a/Engine/Plugins/Enterprise/GLTFExporter/MigratedFromGit/Source/GLTFExporter/Private/Json/GLTFJsonMaterial.h
+++ b/Engine/Plugins/Enterprise/GLTFExporter/MigratedFromGit/Source/GLTFExporter/Private/Json/GLTFJsonMaterial.h
@@ -28,7 +28,7 @@ struct FGLTFJsonTextureInfo : IGLTFJsonObject
 			Writer.Write(TEXT("texCoord"), TexCoord);
 		}
 
-		if (Transform != FGLTFJsonTextureTransform())
+		if (!Transform.IsDefault())
 		{
 			Writer.StartExtensions();
 			Writer.Write(EGLTFJsonExtension::KHR_TextureTransform, Transform);
@@ -105,7 +105,7 @@ struct FGLTFJsonPBRMetallicRoughness : IGLTFJsonObject
 
 	virtual void WriteObject(IGLTFJsonWriter& Writer) const override
 	{
-		if (BaseColorFactor != FGLTFJsonColor4::White)
+		if (!BaseColorFactor.IsNearlyEqual(FGLTFJsonColor4::White))
 		{
 			Writer.Write(TEXT("baseColorFactor"), BaseColorFactor);
 		}
@@ -237,7 +237,7 @@ struct FGLTFJsonMaterial : IGLTFJsonObject
 			Writer.Write(TEXT("emissiveTexture"), EmissiveTexture);
 		}
 
-		if (EmissiveFactor != FGLTFJsonColor3::Black)
+		if (!EmissiveFactor.IsNearlyEqual(FGLTFJsonColor3::Black))
 		{
 			Writer.Write(TEXT("emissiveFactor"), EmissiveFactor);
 		}
diff --git a/Engine/Plugins/Enterprise/GLTFExporter/MigratedFromGit/Source/GLTFExporter/Private/Json/GLTFJsonMatrix.h b/Engine/Plugins/Enterprise/GLTFExporter/MigratedFromGit/Source/GLTFExporter/Private/Json/GLTFJsonMatrix.h
--- a/Engine/Plugins/Enterprise/GLTFExporter/MigratedFromGit/Source/GLTFExporter/Private/Json/GLTFJsonMatrix.h
+++ b/Engine/Plugins/Enterprise/GLTFExporter/MigratedFromGit/Source/GLTFExporter/Private/Json/GLTFJsonMatrix.h
@@ -29,11 +29,11 @@ struct TGLTFJsonMatrix final : BaseType, IGLTFJsonArray
 		}
 	}
 
-	bool operator==(const BaseType& Other) const
+	bool IsNearlyEqual(const BaseType& Other, float Tolerance = KINDA_SMALL_NUMBER) const
 	{
-		for (uint32 i = 0; i < GetNum(BaseType::Elements); ++i)
+		for (int32 i = 0; i < GetNum(BaseType::Elements); ++i)
 		{
-			if (BaseType::Elements[i] != Other.Elements[i])
+			if (!FMath::IsNearlyEqual(BaseType::Elements[i], Other.Elements[i], Tolerance))
 			{
 				return false;
 			}
@@ -41,11 +41,6 @@ struct TGLTFJsonMatrix final : BaseType, IGLTFJsonArray
 
 		return true;
 	}
-
-	bool operator!=(const BaseType& Other) const
-	{
-		return !(*this == Other);
-	}
 };
 
 typedef TGLTFJsonMatrix<FGLTFMatrix2> FGLTFJsonMatrix2;
diff --git a/Engine/Plugins/Enterprise/GLTFExporter/MigratedFromGit/Source/GLTFExporter/Private/Json/GLTFJsonNode.h b/Engine/Plugins/Enterprise/GLTFExporter/MigratedFromGit/Source/GLTFExporter/Private/Json/GLTFJsonNode.h
--- a/Engine/Plugins/Enterprise/GLTFExporter/MigratedFromGit/Source/GLTFExporter/Private/Json/GLTFJsonNode.h
+++ b/Engine/Plugins/Enterprise/GLTFExporter/MigratedFromGit/Source/GLTFExporter/Private/Json/GLTFJsonNode.h
@@ -50,24 +50,24 @@ struct FGLTFJsonNode : IGLTFJsonObject
 
 		if (bUseMatrix)
 		{
-			if (Matrix != FGLTFJsonMatrix4::Identity)
+			if (!Matrix.IsNearlyEqual(FGLTFJsonMatrix4::Identity))
 			{
 				Writer.Write(TEXT("matrix"), Matrix);
 			}
 		}
 		else
 		{
-			if (Translation != FGLTFJsonVector3::Zero)
+			if (!Translation.IsNearlyEqual(FGLTFJsonVector3::Zero))
 			{
 				Writer.Write(TEXT("translation"), Translation);
 			}
 
-			if (Rotation != FGLTFJsonQuaternion::Identity)
+			if (!Rotation.IsNearlyEqual(FGLTFJsonQuaternion::Identity))
 			{
 				Writer.Write(TEXT("rotation"), Rotation);
 			}
 
-			if (Scale != FGLTFJsonVector3::One)
+			if (!Scale.IsNearlyEqual(FGLTFJsonVector3::One))
 			{
 				Writer.Write(TEXT("scale"), Scale);
 			}
diff --git a/Engine/Plugins/Enterprise/GLTFExporter/MigratedFromGit/Source/GLTFExporter/Private/Json/GLTFJsonQuaternion.h b/Engine/Plugins/Enterprise/GLTFExporter/MigratedFromGit/Source/GLTFExporter/Private/Json/GLTFJsonQuaternion.h
--- a/Engine/Plugins/Enterprise/GLTFExporter/MigratedFromGit/Source/GLTFExporter/Private/Json/GLTFJsonQuaternion.h
+++ b/Engine/Plugins/Enterprise/GLTFExporter/MigratedFromGit/Source/GLTFExporter/Private/Json/GLTFJsonQuaternion.h
@@ -28,21 +28,9 @@ struct FGLTFJsonQuaternion final : TGLTFQuaternion<float>, IGLTFJsonArray
 		}
 	}
 
-	bool operator==(const TGLTFQuaternion& Other) const
+	bool IsNearlyEqual(const TGLTFQuaternion& Other, float Tolerance = KINDA_SMALL_NUMBER) const
 	{
-		for (uint32 i = 0; i < GetNum(Components); ++i)
-		{
-			if (Components[i] != Other.Components[i])
-			{
-				return false;
-			}
-		}
-
-		return true;
-	}
-
-	bool operator!=(const TGLTFQuaternion& Other) const
-	{
-		return !(*this == Other);
+		return (FMath::Abs(X - Other.X) <= Tolerance && FMath::Abs(Y - Other.Y) <= Tolerance && FMath::Abs(Z - Other.Z) <= Tolerance && FMath::Abs(W - Other.W) <= Tolerance)
+			|| (FMath::Abs(X + Other.X) <= Tolerance && FMath::Abs(Y + Other.Y) <= Tolerance && FMath::Abs(Z + Other.Z) <= Tolerance && FMath::Abs(W + Other.W) <= Tolerance);
 	}
 };
diff --git a/Engine/Plugins/Enterprise/GLTFExporter/MigratedFromGit/Source/GLTFExporter/Private/Json/GLTFJsonSkySphere.h b/Engine/Plugins/Enterprise/GLTFExporter/MigratedFromGit/Source/GLTFExporter/Private/Json/GLTFJsonSkySphere.h
--- a/Engine/Plugins/Enterprise/GLTFExporter/MigratedFromGit/Source/GLTFExporter/Private/Json/GLTFJsonSkySphere.h
+++ b/Engine/Plugins/Enterprise/GLTFExporter/MigratedFromGit/Source/GLTFExporter/Private/Json/GLTFJsonSkySphere.h
@@ -127,7 +127,7 @@ struct FGLTFJsonSkySphere : IGLTFJsonObject
 		Writer.Write(TEXT("horizonColor"), HorizonColor);
 		Writer.Write(TEXT("cloudColor"), CloudColor);
 
-		if (OverallColor != FGLTFJsonColor4::White)
+		if (!OverallColor.IsNearlyEqual(FGLTFJsonColor4::White))
 		{
 			Writer.Write(TEXT("overallColor"), OverallColor);
 		}
@@ -147,7 +147,7 @@ struct FGLTFJsonSkySphere : IGLTFJsonObject
 			Writer.Write(TEXT("cloudColorCurve"), CloudColorCurve);
 		}
 
-		if (Scale != FGLTFJsonVector3::One)
+		if (!Scale.IsNearlyEqual(FGLTFJsonVector3::One))
 		{
 			Writer.Write(TEXT("scale"), Scale);
 		}
diff --git a/Engine/Plugins/Enterprise/GLTFExporter/MigratedFromGit/Source/GLTFExporter/Private/Json/GLTFJsonTextureTransform.h b/Engine/Plugins/Enterprise/GLTFExporter/MigratedFromGit/Source/GLTFExporter/Private/Json/GLTFJsonTextureTransform.h
--- a/Engine/Plugins/Enterprise/GLTFExporter/MigratedFromGit/Source/GLTFExporter/Private/Json/GLTFJsonTextureTransform.h
+++ b/Engine/Plugins/Enterprise/GLTFExporter/MigratedFromGit/Source/GLTFExporter/Private/Json/GLTFJsonTextureTransform.h
@@ -20,12 +20,12 @@ struct FGLTFJsonTextureTransform : IGLTFJsonObject
 
 	virtual void WriteObject(IGLTFJsonWriter& Writer) const override
 	{
-		if (Offset != FGLTFJsonVector2::Zero)
+		if (!Offset.IsNearlyEqual(FGLTFJsonVector2::Zero))
 		{
 			Writer.Write(TEXT("offset"), Offset);
 		}
 
-		if (Scale != FGLTFJsonVector2::One)
+		if (!Scale.IsNearlyEqual(FGLTFJsonVector2::One))
 		{
 			Writer.Write(TEXT("scale"), Scale);
 		}
@@ -36,15 +36,22 @@ struct FGLTFJsonTextureTransform : IGLTFJsonObject
 		}
 	}
 
-	bool operator==(const FGLTFJsonTextureTransform& Other) const
+	bool IsNearlyEqual(const FGLTFJsonTextureTransform& Other) const
 	{
-		return Offset == Other.Offset
-			&& Scale == Other.Scale
+		return Offset.IsNearlyEqual(Other.Offset)
+			&& Scale.IsNearlyEqual(Other.Scale)
+			&& FMath::IsNearlyEqual(Rotation, Other.Rotation);
+	}
+
+	bool IsExactEqual(const FGLTFJsonTextureTransform& Other) const
+	{
+		return Offset.X == Other.Offset.X && Offset.Y == Other.Offset.Y
+			&& Scale.X == Other.Scale.X && Scale.Y == Other.Scale.Y
 			&& Rotation == Other.Rotation;
 	}
 
-	bool operator!=(const FGLTFJsonTextureTransform& Other) const
+	bool IsDefault() const
 	{
-		return !(*this == Other);
+		return IsExactEqual(FGLTFJsonTextureTransform());
 	}
 };
diff --git a/Engine/Plugins/Enterprise/GLTFExporter/MigratedFromGit/Source/GLTFExporter/Private/Json/GLTFJsonVector.h b/Engine/Plugins/Enterprise/GLTFExporter/MigratedFromGit/Source/GLTFExporter/Private/Json/GLTFJsonVector.h
--- a/Engine/Plugins/Enterprise/GLTFExporter/MigratedFromGit/Source/GLTFExporter/Private/Json/GLTFJsonVector.h
+++ b/Engine/Plugins/Enterprise/GLTFExporter/MigratedFromGit/Source/GLTFExporter/Private/Json/GLTFJsonVector.h
@@ -30,11 +30,11 @@ struct TGLTFJsonVector final : BaseType, IGLTFJsonArray
 		}
 	}
 
-	bool operator==(const BaseType& Other) const
+	bool IsNearlyEqual(const BaseType& Other, float Tolerance = KINDA_SMALL_NUMBER) const
 	{
-		for (uint32 i = 0; i < GetNum(BaseType::Components); ++i)
+		for (int32 i = 0; i < GetNum(BaseType::Components); ++i)
 		{
-			if (BaseType::Components[i] != Other.Components[i])
+			if (!FMath::IsNearlyEqual(BaseType::Components[i], Other.Components[i], Tolerance))
 			{
 				return false;
 			}
@@ -42,11 +42,6 @@ struct TGLTFJsonVector final : BaseType, IGLTFJsonArray
 
 		return true;
 	}
-
-	bool operator!=(const BaseType& Other) const
-	{
-		return !(*this == Other);
-	}
 };
 
 typedef TGLTFJsonVector<FGLTFVector2> FGLTFJsonVector2;
diff --git a/Engine/Plugins/Enterprise/GLTFExporter/MigratedFromGit/Source/GLTFExporter/Private/Tasks/GLTFMaterialTasks.cpp b/Engine/Plugins/Enterprise/GLTFExporter/MigratedFromGit/Source/GLTFExporter/Private/Tasks/GLTFMaterialTasks.cpp
--- a/Engine/Plugins/Enterprise/GLTFExporter/MigratedFromGit/Source/GLTFExporter/Private/Tasks/GLTFMaterialTasks.cpp
+++ b/Engine/Plugins/Enterprise/GLTFExporter/MigratedFromGit/Source/GLTFExporter/Private/Tasks/GLTFMaterialTasks.cpp
@@ -321,7 +321,7 @@ bool FGLTFMaterialTask::TryGetBaseColorAndOpacity(FGLTFJsonPBRMetallicRoughness&
 		bHasOpacitySourceTexture &&
 		BaseColorTexture == OpacityTexture &&
 		BaseColorTexCoord == OpacityTexCoord &&
-		BaseColorTransform == OpacityTransform)
+		BaseColorTransform.IsExactEqual(OpacityTransform))
 	{
 		OutPBRParams.BaseColorTexture.Index = Builder.GetOrAddTexture(BaseColorTexture);
 		OutPBRParams.BaseColorTexture.TexCoord = BaseColorTexCoord;
@@ -446,7 +446,7 @@ bool FGLTFMaterialTask::TryGetMetallicAndRoughness(FGLTFJsonPBRMetallicRoughness
 		bHasRoughnessSourceTexture &&
 		MetallicTexture == RoughnessTexture &&
 		MetallicTexCoord == RoughnessTexCoord &&
-		MetallicTransform == RoughnessTransform)
+		MetallicTransform.IsExactEqual(RoughnessTransform))
 	{
 		OutPBRParams.MetallicRoughnessTexture.Index = Builder.GetOrAddTexture(MetallicTexture);
 		OutPBRParams.MetallicRoughnessTexture.TexCoord = MetallicTexCoord;
@@ -571,7 +571,7 @@ bool FGLTFMaterialTask::TryGetClearCoatRoughness(FGLTFJsonClearCoatExtension& Ou
 		bHasRoughnessSourceTexture &&
 		IntensityTexture == RoughnessTexture &&
 		IntensityTexCoord == RoughnessTexCoord &&
-		IntensityTransform == RoughnessTransform)
+		IntensityTransform.IsExactEqual(RoughnessTransform))
 	{
 		const FGLTFJsonTextureIndex TextureIndex = Builder.GetOrAddTexture(IntensityTexture);
 		OutExtParams.ClearCoatTexture.Index = TextureIndex;
@@ -1055,7 +1055,7 @@ bool FGLTFMaterialTask::TryGetSourceTexture(const UTexture2D*& OutTexture, int32
 			return false;
 		}
 
-		if (!Builder.ExportOptions->bExportTextureTransforms && OutTransform != FGLTFJsonTextureTransform())
+		if (!Builder.ExportOptions->bExportTextureTransforms && !OutTransform.IsDefault())
 		{
 			Builder.AddWarningMessage(FString::Printf(
 				TEXT("Texture coordinates [%d] in %s for material %s are transformed, but texture transform is disabled by export options"),
@@ -1092,7 +1092,7 @@ bool FGLTFMaterialTask::TryGetSourceTexture(const UTexture2D*& OutTexture, int32
 			return false;
 		}
 
-		if (!Builder.ExportOptions->bExportTextureTransforms && OutTransform != FGLTFJsonTextureTransform())
+		if (!Builder.ExportOptions->bExportTextureTransforms && !OutTransform.IsDefault())
 		{
 			Builder.AddWarningMessage(FString::Printf(
 				TEXT("Texture coordinates [%d] in %s for material %s are transformed, but texture transform is disabled by export options"),
