[GLTFExporter] refactor: remove ugly hack to remember json serialized extensions

commit 2cb909bccba89d8cc6e0a2dc48f710c404b84d30
Date: Sat, 12 Jun 2021 05:17:44 +0200

    refactor: remove ugly hack to remember json serialized extensions
    (cherry picked from commit 91cbe3349bea561add7fe188ea2073a412ce37e2​)

#skipci

diff --git a/Engine/Plugins/Enterprise/GLTFExporter/MigratedFromGit/Source/GLTFExporter/Private/Builders/GLTFJsonBuilder.cpp b/Engine/Plugins/Enterprise/GLTFExporter/MigratedFromGit/Source/GLTFExporter/Private/Builders/GLTFJsonBuilder.cpp
--- a/Engine/Plugins/Enterprise/GLTFExporter/MigratedFromGit/Source/GLTFExporter/Private/Builders/GLTFJsonBuilder.cpp
+++ b/Engine/Plugins/Enterprise/GLTFExporter/MigratedFromGit/Source/GLTFExporter/Private/Builders/GLTFJsonBuilder.cpp
@@ -14,19 +14,16 @@ FGLTFJsonBuilder::FGLTFJsonBuilder(const FString& FilePath, const UGLTFExportOpt
 
 void FGLTFJsonBuilder::WriteJson(FArchive& Archive)
 {
-	JsonRoot.ToJson(&Archive, !bIsGlbFile);
+	JsonRoot.WriteJson(Archive, !bIsGlbFile);
 }
 
 TSet<EGLTFJsonExtension> FGLTFJsonBuilder::GetCustomExtensionsUsed() const
 {
-	const TCHAR CustomPrefix[] = TEXT("EPIC_");
-
 	TSet<EGLTFJsonExtension> CustomExtensions;
 
 	for (EGLTFJsonExtension Extension : JsonRoot.Extensions.Used)
 	{
-		const TCHAR* ExtensionString = FGLTFJsonUtility::GetValue(Extension);
-		if (FCString::Strncmp(ExtensionString, CustomPrefix, (sizeof(CustomPrefix) / sizeof(TCHAR)) - 1) == 0)
+		if (IsCustomExtension(Extension))
 		{
 			CustomExtensions.Add(Extension);
 		}
@@ -278,3 +275,11 @@ FString FGLTFJsonBuilder::GetGeneratorString() const
 		? TEXT(EPIC_PRODUCT_NAME) TEXT(" ") ENGINE_VERSION_STRING TEXT(" ") + PluginDescriptor.FriendlyName + TEXT(" ") + PluginDescriptor.VersionName
 		: TEXT(EPIC_PRODUCT_NAME) TEXT(" ") + PluginDescriptor.FriendlyName;
 }
+
+bool FGLTFJsonBuilder::IsCustomExtension(EGLTFJsonExtension Extension)
+{
+	const TCHAR CustomPrefix[] = TEXT("EPIC_");
+
+	const TCHAR* ExtensionString = FGLTFJsonUtility::GetValue(Extension);
+	return FCString::Strncmp(ExtensionString, CustomPrefix, GetNum(CustomPrefix) - 1) == 0;
+}
diff --git a/Engine/Plugins/Enterprise/GLTFExporter/MigratedFromGit/Source/GLTFExporter/Private/Builders/GLTFJsonBuilder.h b/Engine/Plugins/Enterprise/GLTFExporter/MigratedFromGit/Source/GLTFExporter/Private/Builders/GLTFJsonBuilder.h
--- a/Engine/Plugins/Enterprise/GLTFExporter/MigratedFromGit/Source/GLTFExporter/Private/Builders/GLTFJsonBuilder.h
+++ b/Engine/Plugins/Enterprise/GLTFExporter/MigratedFromGit/Source/GLTFExporter/Private/Builders/GLTFJsonBuilder.h
@@ -70,5 +70,7 @@ private:
 
 	FString GetGeneratorString() const;
 
+	static bool IsCustomExtension(EGLTFJsonExtension Extension);
+
 	FGLTFJsonRoot JsonRoot;
 };
diff --git a/Engine/Plugins/Enterprise/GLTFExporter/MigratedFromGit/Source/GLTFExporter/Private/Json/GLTFJsonRoot.h b/Engine/Plugins/Enterprise/GLTFExporter/MigratedFromGit/Source/GLTFExporter/Private/Json/GLTFJsonRoot.h
--- a/Engine/Plugins/Enterprise/GLTFExporter/MigratedFromGit/Source/GLTFExporter/Private/Json/GLTFJsonRoot.h
+++ b/Engine/Plugins/Enterprise/GLTFExporter/MigratedFromGit/Source/GLTFExporter/Private/Json/GLTFJsonRoot.h
@@ -126,16 +126,21 @@ struct FGLTFJsonRoot : IGLTFJsonObject
 			Writer.EndExtensions();
 		}
 
-		if (Writer.Extensions.Used.Num() > 0) Writer.Write(TEXT("extensionsUsed"), Writer.Extensions.Used);
-		if (Writer.Extensions.Required.Num() > 0) Writer.Write(TEXT("extensionsRequired"), Writer.Extensions.Required);
+		if (Extensions.Used.Num() > 0)
+		{
+			Writer.Write(TEXT("extensionsUsed"), Extensions.Used);
+		}
+
+		if (Extensions.Required.Num() > 0)
+		{
+			Writer.Write(TEXT("extensionsRequired"), Extensions.Required);
+		}
 	}
 
-	void ToJson(FArchive* const Archive, bool bPrettyJson)
+	void WriteJson(FArchive& Archive, bool bPrettyJson)
 	{
-		TSharedRef<IGLTFJsonWriter> Writer = IGLTFJsonWriter::Create(Archive, bPrettyJson);
-		Writer->Extensions = Extensions;
+		TSharedRef<IGLTFJsonWriter> Writer = IGLTFJsonWriter::Create(Archive, bPrettyJson, Extensions);
 		Writer->Write(this);
 		Writer->Close();
-		Extensions = Writer->Extensions; // TODO: fix ugly hack
 	}
 };
diff --git a/Engine/Plugins/Enterprise/GLTFExporter/MigratedFromGit/Source/GLTFExporter/Private/Json/GLTFJsonWriter.cpp b/Engine/Plugins/Enterprise/GLTFExporter/MigratedFromGit/Source/GLTFExporter/Private/Json/GLTFJsonWriter.cpp
--- a/Engine/Plugins/Enterprise/GLTFExporter/MigratedFromGit/Source/GLTFExporter/Private/Json/GLTFJsonWriter.cpp
+++ b/Engine/Plugins/Enterprise/GLTFExporter/MigratedFromGit/Source/GLTFExporter/Private/Json/GLTFJsonWriter.cpp
@@ -9,8 +9,9 @@ class TGLTFJsonWriterImpl final : public IGLTFJsonWriter
 {
 public:
 
-	TGLTFJsonWriterImpl(FArchive* const Archive)
-		: JsonWriter(TJsonWriterFactory<CharType, PrintPolicy>::Create(Archive))
+	TGLTFJsonWriterImpl(FArchive& Archive, FGLTFJsonExtensions& Extensions)
+		: IGLTFJsonWriter(Extensions)
+		, JsonWriter(TJsonWriterFactory<CharType, PrintPolicy>::Create(&Archive))
 	{
 	}
 
@@ -147,10 +148,10 @@ private:
 	TSharedRef<TJsonWriter<CharType, PrintPolicy>> JsonWriter;
 };
 
-TSharedRef<IGLTFJsonWriter> IGLTFJsonWriter::Create(FArchive* const Archive, bool bPrettyJson)
+TSharedRef<IGLTFJsonWriter> IGLTFJsonWriter::Create(FArchive& Archive, bool bPrettyJson, FGLTFJsonExtensions& Extensions)
 {
 	return MakeShareable(bPrettyJson
-		? static_cast<IGLTFJsonWriter*>(new TGLTFJsonWriterImpl<UTF8CHAR, TPrettyJsonPrintPolicy<UTF8CHAR>>(Archive))
-		: static_cast<IGLTFJsonWriter*>(new TGLTFJsonWriterImpl<UTF8CHAR, TCondensedJsonPrintPolicy<UTF8CHAR>>(Archive))
+		? static_cast<IGLTFJsonWriter*>(new TGLTFJsonWriterImpl<UTF8CHAR, TPrettyJsonPrintPolicy<UTF8CHAR>>(Archive, Extensions))
+		: static_cast<IGLTFJsonWriter*>(new TGLTFJsonWriterImpl<UTF8CHAR, TCondensedJsonPrintPolicy<UTF8CHAR>>(Archive, Extensions))
 	);
 }
diff --git a/Engine/Plugins/Enterprise/GLTFExporter/MigratedFromGit/Source/GLTFExporter/Private/Json/GLTFJsonWriter.h b/Engine/Plugins/Enterprise/GLTFExporter/MigratedFromGit/Source/GLTFExporter/Private/Json/GLTFJsonWriter.h
--- a/Engine/Plugins/Enterprise/GLTFExporter/MigratedFromGit/Source/GLTFExporter/Private/Json/GLTFJsonWriter.h
+++ b/Engine/Plugins/Enterprise/GLTFExporter/MigratedFromGit/Source/GLTFExporter/Private/Json/GLTFJsonWriter.h
@@ -10,9 +10,14 @@ class IGLTFJsonWriter
 {
 public:
 
-	FGLTFJsonExtensions Extensions;
+	FGLTFJsonExtensions& Extensions;
 
-	static TSharedRef<IGLTFJsonWriter> Create(FArchive* const Archive, bool bPrettyJson);
+	IGLTFJsonWriter(FGLTFJsonExtensions& Extensions)
+		: Extensions(Extensions)
+	{
+	}
+
+	static TSharedRef<IGLTFJsonWriter> Create(FArchive& Archive, bool bPrettyJson, FGLTFJsonExtensions& Extensions);
 
 	virtual ~IGLTFJsonWriter() = default;
 
