[GLTFExporter] feat: avoid json serializing any float-based property that is nearly equal to default value

commit 69411a0814d6e2878c0ffc18fa60d5d902b5ac81
Date: Sun, 13 Jun 2021 01:27:44 +0200

    feat: avoid json serializing any float-based property that is nearly equal to default value
    (cherry picked from commit 7e3fb446ff35ecb12491aef4e7f8e020859b966c​)

#skipci

diff --git a/Engine/Plugins/Enterprise/GLTFExporter/MigratedFromGit/Source/GLTFExporter/Private/Converters/GLTFCurveUtility.cpp b/Engine/Plugins/Enterprise/GLTFExporter/MigratedFromGit/Source/GLTFExporter/Private/Converters/GLTFCurveUtility.cpp
--- a/Engine/Plugins/Enterprise/GLTFExporter/MigratedFromGit/Source/GLTFExporter/Private/Converters/GLTFCurveUtility.cpp
+++ b/Engine/Plugins/Enterprise/GLTFExporter/MigratedFromGit/Source/GLTFExporter/Private/Converters/GLTFCurveUtility.cpp
@@ -3,15 +3,13 @@
 #include "Converters/GLTFCurveUtility.h"
 #include "Curves/CurveLinearColor.h"
 
-bool FGLTFCurveUtility::HasAnyAdjustment(const UCurveLinearColor& ColorCurve)
+bool FGLTFCurveUtility::HasAnyAdjustment(const UCurveLinearColor& ColorCurve, float Tolerance)
 {
-	const float ErrorTolerance = KINDA_SMALL_NUMBER;
-
-	return !FMath::IsNearlyEqual(ColorCurve.AdjustBrightness, 1.0f, ErrorTolerance) ||
-		!FMath::IsNearlyEqual(ColorCurve.AdjustHue, 0.0f, ErrorTolerance) ||
-		!FMath::IsNearlyEqual(ColorCurve.AdjustSaturation, 1.0f, ErrorTolerance) ||
-		!FMath::IsNearlyEqual(ColorCurve.AdjustVibrance, 0.0f, ErrorTolerance) ||
-		!FMath::IsNearlyEqual(ColorCurve.AdjustBrightnessCurve, 1.0f, ErrorTolerance) ||
-		!FMath::IsNearlyEqual(ColorCurve.AdjustMaxAlpha, 1.0f, ErrorTolerance) ||
-		!FMath::IsNearlyEqual(ColorCurve.AdjustMinAlpha, 0.0f, ErrorTolerance);
+	return !FMath::IsNearlyEqual(ColorCurve.AdjustBrightness,      1, Tolerance)
+		|| !FMath::IsNearlyEqual(ColorCurve.AdjustHue,             0, Tolerance)
+		|| !FMath::IsNearlyEqual(ColorCurve.AdjustSaturation,      1, Tolerance)
+		|| !FMath::IsNearlyEqual(ColorCurve.AdjustVibrance,        0, Tolerance)
+		|| !FMath::IsNearlyEqual(ColorCurve.AdjustBrightnessCurve, 1, Tolerance)
+		|| !FMath::IsNearlyEqual(ColorCurve.AdjustMaxAlpha,        1, Tolerance)
+		|| !FMath::IsNearlyEqual(ColorCurve.AdjustMinAlpha,        0, Tolerance);
 }
diff --git a/Engine/Plugins/Enterprise/GLTFExporter/MigratedFromGit/Source/GLTFExporter/Private/Converters/GLTFCurveUtility.h b/Engine/Plugins/Enterprise/GLTFExporter/MigratedFromGit/Source/GLTFExporter/Private/Converters/GLTFCurveUtility.h
--- a/Engine/Plugins/Enterprise/GLTFExporter/MigratedFromGit/Source/GLTFExporter/Private/Converters/GLTFCurveUtility.h
+++ b/Engine/Plugins/Enterprise/GLTFExporter/MigratedFromGit/Source/GLTFExporter/Private/Converters/GLTFCurveUtility.h
@@ -6,5 +6,5 @@ class UCurveLinearColor;
 
 struct FGLTFCurveUtility
 {
-	static bool HasAnyAdjustment(const UCurveLinearColor& ColorCurve);
+	static bool HasAnyAdjustment(const UCurveLinearColor& ColorCurve, float Tolerance = KINDA_SMALL_NUMBER);
 };
diff --git a/Engine/Plugins/Enterprise/GLTFExporter/MigratedFromGit/Source/GLTFExporter/Private/Json/GLTFJsonAnimationPlayback.h b/Engine/Plugins/Enterprise/GLTFExporter/MigratedFromGit/Source/GLTFExporter/Private/Json/GLTFJsonAnimationPlayback.h
--- a/Engine/Plugins/Enterprise/GLTFExporter/MigratedFromGit/Source/GLTFExporter/Private/Json/GLTFJsonAnimationPlayback.h
+++ b/Engine/Plugins/Enterprise/GLTFExporter/MigratedFromGit/Source/GLTFExporter/Private/Json/GLTFJsonAnimationPlayback.h
@@ -39,12 +39,12 @@ struct FGLTFJsonAnimationPlayback : IGLTFJsonObject
 			Writer.Write(TEXT("autoPlay"), bAutoPlay);
 		}
 
-		if (PlayRate != 1)
+		if (!FMath::IsNearlyEqual(PlayRate, 1))
 		{
 			Writer.Write(TEXT("playRate"), PlayRate);
 		}
 
-		if (StartTime != 0)
+		if (!FMath::IsNearlyEqual(StartTime, 0))
 		{
 			Writer.Write(TEXT("startTime"), StartTime);
 		}
diff --git a/Engine/Plugins/Enterprise/GLTFExporter/MigratedFromGit/Source/GLTFExporter/Private/Json/GLTFJsonCamera.h b/Engine/Plugins/Enterprise/GLTFExporter/MigratedFromGit/Source/GLTFExporter/Private/Json/GLTFJsonCamera.h
--- a/Engine/Plugins/Enterprise/GLTFExporter/MigratedFromGit/Source/GLTFExporter/Private/Json/GLTFJsonCamera.h
+++ b/Engine/Plugins/Enterprise/GLTFExporter/MigratedFromGit/Source/GLTFExporter/Private/Json/GLTFJsonCamera.h
@@ -47,14 +47,14 @@ struct FGLTFJsonPerspective : IGLTFJsonObject
 
 	virtual void WriteObject(IGLTFJsonWriter& Writer) const override
 	{
-		if (AspectRatio != 0)
+		if (!FMath::IsNearlyEqual(AspectRatio, 0))
 		{
 			Writer.Write(TEXT("aspectRatio"), AspectRatio);
 		}
 
 		Writer.Write(TEXT("yfov"), YFov);
 
-		if (ZFar != 0)
+		if (!FMath::IsNearlyEqual(ZFar, 0))
 		{
 			Writer.Write(TEXT("zfar"), ZFar);
 		}
diff --git a/Engine/Plugins/Enterprise/GLTFExporter/MigratedFromGit/Source/GLTFExporter/Private/Json/GLTFJsonCameraControl.h b/Engine/Plugins/Enterprise/GLTFExporter/MigratedFromGit/Source/GLTFExporter/Private/Json/GLTFJsonCameraControl.h
--- a/Engine/Plugins/Enterprise/GLTFExporter/MigratedFromGit/Source/GLTFExporter/Private/Json/GLTFJsonCameraControl.h
+++ b/Engine/Plugins/Enterprise/GLTFExporter/MigratedFromGit/Source/GLTFExporter/Private/Json/GLTFJsonCameraControl.h
@@ -25,9 +25,9 @@ struct FGLTFJsonCameraControl : IGLTFJsonObject
 		: Mode(EGLTFJsonCameraControlMode::FreeLook)
 		, MaxDistance(0)
 		, MinDistance(0)
-		, MaxPitch(90.0f)
-		, MinPitch(-90.0f)
-		, MaxYaw(360.0f)
+		, MaxPitch(90)
+		, MinPitch(-90)
+		, MaxYaw(360)
 		, MinYaw(0)
 		, RotationSensitivity(0)
 		, RotationInertia(0)
@@ -48,22 +48,22 @@ struct FGLTFJsonCameraControl : IGLTFJsonObject
 		Writer.Write(TEXT("maxDistance"), MaxDistance);
 		Writer.Write(TEXT("minDistance"), MinDistance);
 
-		if (!FMath::IsNearlyEqual(MaxPitch, 90.0f))
+		if (!FMath::IsNearlyEqual(MaxPitch, 90))
 		{
 			Writer.Write(TEXT("maxPitch"), MaxPitch);
 		}
 
-		if (!FMath::IsNearlyEqual(MinPitch, -90.0f))
+		if (!FMath::IsNearlyEqual(MinPitch, -90))
 		{
 			Writer.Write(TEXT("minPitch"), MinPitch);
 		}
 
-		if (!FMath::IsNearlyEqual(MaxYaw, 360.0f))
+		if (!FMath::IsNearlyEqual(MaxYaw, 360))
 		{
 			Writer.Write(TEXT("maxYaw"), MaxYaw);
 		}
 
-		if (!FMath::IsNearlyEqual(MinYaw, 0.0f))
+		if (!FMath::IsNearlyEqual(MinYaw, 0))
 		{
 			Writer.Write(TEXT("minYaw"), MinYaw);
 		}
diff --git a/Engine/Plugins/Enterprise/GLTFExporter/MigratedFromGit/Source/GLTFExporter/Private/Json/GLTFJsonLight.h b/Engine/Plugins/Enterprise/GLTFExporter/MigratedFromGit/Source/GLTFExporter/Private/Json/GLTFJsonLight.h
--- a/Engine/Plugins/Enterprise/GLTFExporter/MigratedFromGit/Source/GLTFExporter/Private/Json/GLTFJsonLight.h
+++ b/Engine/Plugins/Enterprise/GLTFExporter/MigratedFromGit/Source/GLTFExporter/Private/Json/GLTFJsonLight.h
@@ -13,18 +13,18 @@ struct FGLTFJsonSpotLight : IGLTFJsonObject
 
 	FGLTFJsonSpotLight()
 		: InnerConeAngle(0)
-		, OuterConeAngle(PI / 2.0f)
+		, OuterConeAngle(HALF_PI)
 	{
 	}
 
 	virtual void WriteObject(IGLTFJsonWriter& Writer) const override
 	{
-		if (InnerConeAngle != 0)
+		if (!FMath::IsNearlyEqual(InnerConeAngle, 0))
 		{
 			Writer.Write(TEXT("innerConeAngle"), InnerConeAngle);
 		}
 
-		if (OuterConeAngle != PI / 2.0f)
+		if (!FMath::IsNearlyEqual(OuterConeAngle, HALF_PI))
 		{
 			Writer.Write(TEXT("outerConeAngle"), OuterConeAngle);
 		}
@@ -66,14 +66,14 @@ struct FGLTFJsonLight : IGLTFJsonObject
 			Writer.Write(TEXT("color"), Color);
 		}
 
-		if (Intensity != 1)
+		if (!FMath::IsNearlyEqual(Intensity, 1))
 		{
 			Writer.Write(TEXT("intensity"), Intensity);
 		}
 
 		if (Type == EGLTFJsonLightType::Point || Type == EGLTFJsonLightType::Spot)
 		{
-			if (Range != 0)
+			if (!FMath::IsNearlyEqual(Range, 0))
 			{
 				Writer.Write(TEXT("range"), Range);
 			}
diff --git a/Engine/Plugins/Enterprise/GLTFExporter/MigratedFromGit/Source/GLTFExporter/Private/Json/GLTFJsonMaterial.h b/Engine/Plugins/Enterprise/GLTFExporter/MigratedFromGit/Source/GLTFExporter/Private/Json/GLTFJsonMaterial.h
--- a/Engine/Plugins/Enterprise/GLTFExporter/MigratedFromGit/Source/GLTFExporter/Private/Json/GLTFJsonMaterial.h
+++ b/Engine/Plugins/Enterprise/GLTFExporter/MigratedFromGit/Source/GLTFExporter/Private/Json/GLTFJsonMaterial.h
@@ -28,7 +28,7 @@ struct FGLTFJsonTextureInfo : IGLTFJsonObject
 			Writer.Write(TEXT("texCoord"), TexCoord);
 		}
 
-		if (!Transform.IsDefault())
+		if (!Transform.IsNearlyDefault())
 		{
 			Writer.StartExtensions();
 			Writer.Write(EGLTFJsonExtension::KHR_TextureTransform, Transform);
@@ -55,7 +55,7 @@ struct FGLTFJsonNormalTextureInfo : FGLTFJsonTextureInfo
 			Writer.Write(TEXT("texCoord"), TexCoord);
 		}
 
-		if (Scale != 1)
+		if (!FMath::IsNearlyEqual(Scale, 1))
 		{
 			Writer.Write(TEXT("scale"), Scale);
 		}
@@ -80,7 +80,7 @@ struct FGLTFJsonOcclusionTextureInfo : FGLTFJsonTextureInfo
 			Writer.Write(TEXT("texCoord"), TexCoord);
 		}
 
-		if (Strength != 1)
+		if (!FMath::IsNearlyEqual(Strength, 1))
 		{
 			Writer.Write(TEXT("strength"), Strength);
 		}
@@ -115,12 +115,12 @@ struct FGLTFJsonPBRMetallicRoughness : IGLTFJsonObject
 			Writer.Write(TEXT("baseColorTexture"), BaseColorTexture);
 		}
 
-		if (MetallicFactor != 1)
+		if (!FMath::IsNearlyEqual(MetallicFactor, 1))
 		{
 			Writer.Write(TEXT("metallicFactor"), MetallicFactor);
 		}
 
-		if (RoughnessFactor != 1)
+		if (!FMath::IsNearlyEqual(RoughnessFactor, 1))
 		{
 			Writer.Write(TEXT("roughnessFactor"), RoughnessFactor);
 		}
@@ -150,7 +150,7 @@ struct FGLTFJsonClearCoatExtension : IGLTFJsonObject
 
 	virtual void WriteObject(IGLTFJsonWriter& Writer) const override
 	{
-		if (ClearCoatFactor != 0)
+		if (!FMath::IsNearlyEqual(ClearCoatFactor, 0))
 		{
 			Writer.Write(TEXT("clearcoatFactor"), ClearCoatFactor);
 		}
@@ -160,7 +160,7 @@ struct FGLTFJsonClearCoatExtension : IGLTFJsonObject
 			Writer.Write(TEXT("clearcoatTexture"), ClearCoatTexture);
 		}
 
-		if (ClearCoatRoughnessFactor != 0)
+		if (!FMath::IsNearlyEqual(ClearCoatRoughnessFactor, 0))
 		{
 			Writer.Write(TEXT("clearcoatRoughnessFactor"), ClearCoatRoughnessFactor);
 		}
@@ -245,11 +245,11 @@ struct FGLTFJsonMaterial : IGLTFJsonObject
 		if (AlphaMode != EGLTFJsonAlphaMode::Opaque)
 		{
 			Writer.Write(TEXT("alphaMode"), AlphaMode);
-		}
 
-		if (AlphaMode == EGLTFJsonAlphaMode::Mask && AlphaCutoff != 0.5f)
-		{
-			Writer.Write(TEXT("alphaCutoff"), AlphaCutoff);
+			if (AlphaMode == EGLTFJsonAlphaMode::Mask && !FMath::IsNearlyEqual(AlphaCutoff, 0.5f))
+			{
+				Writer.Write(TEXT("alphaCutoff"), AlphaCutoff);
+			}
 		}
 
 		if (DoubleSided)
diff --git a/Engine/Plugins/Enterprise/GLTFExporter/MigratedFromGit/Source/GLTFExporter/Private/Json/GLTFJsonTextureTransform.h b/Engine/Plugins/Enterprise/GLTFExporter/MigratedFromGit/Source/GLTFExporter/Private/Json/GLTFJsonTextureTransform.h
--- a/Engine/Plugins/Enterprise/GLTFExporter/MigratedFromGit/Source/GLTFExporter/Private/Json/GLTFJsonTextureTransform.h
+++ b/Engine/Plugins/Enterprise/GLTFExporter/MigratedFromGit/Source/GLTFExporter/Private/Json/GLTFJsonTextureTransform.h
@@ -30,28 +30,33 @@ struct FGLTFJsonTextureTransform : IGLTFJsonObject
 			Writer.Write(TEXT("scale"), Scale);
 		}
 
-		if (Rotation != 0)
+		if (!FMath::IsNearlyEqual(Rotation, 0))
 		{
 			Writer.Write(TEXT("rotation"), Rotation);
 		}
 	}
 
-	bool IsNearlyEqual(const FGLTFJsonTextureTransform& Other) const
+	bool IsNearlyEqual(const FGLTFJsonTextureTransform& Other, float Tolerance = KINDA_SMALL_NUMBER) const
 	{
-		return Offset.IsNearlyEqual(Other.Offset)
-			&& Scale.IsNearlyEqual(Other.Scale)
-			&& FMath::IsNearlyEqual(Rotation, Other.Rotation);
+		return Offset.IsNearlyEqual(Other.Offset, Tolerance)
+			&& Scale.IsNearlyEqual(Other.Scale, Tolerance)
+			&& FMath::IsNearlyEqual(Rotation, Other.Rotation, Tolerance);
 	}
 
-	bool IsExactEqual(const FGLTFJsonTextureTransform& Other) const
+	bool IsExactlyEqual(const FGLTFJsonTextureTransform& Other) const
 	{
 		return Offset.X == Other.Offset.X && Offset.Y == Other.Offset.Y
 			&& Scale.X == Other.Scale.X && Scale.Y == Other.Scale.Y
 			&& Rotation == Other.Rotation;
 	}
 
-	bool IsDefault() const
+	bool IsNearlyDefault(float Tolerance = KINDA_SMALL_NUMBER) const
 	{
-		return IsExactEqual(FGLTFJsonTextureTransform());
+		return IsNearlyEqual(FGLTFJsonTextureTransform(), Tolerance);
+	}
+
+	bool IsExactlyDefault() const
+	{
+		return IsExactlyEqual(FGLTFJsonTextureTransform());
 	}
 };
diff --git a/Engine/Plugins/Enterprise/GLTFExporter/MigratedFromGit/Source/GLTFExporter/Private/Tasks/GLTFMaterialTasks.cpp b/Engine/Plugins/Enterprise/GLTFExporter/MigratedFromGit/Source/GLTFExporter/Private/Tasks/GLTFMaterialTasks.cpp
--- a/Engine/Plugins/Enterprise/GLTFExporter/MigratedFromGit/Source/GLTFExporter/Private/Tasks/GLTFMaterialTasks.cpp
+++ b/Engine/Plugins/Enterprise/GLTFExporter/MigratedFromGit/Source/GLTFExporter/Private/Tasks/GLTFMaterialTasks.cpp
@@ -321,7 +321,7 @@ bool FGLTFMaterialTask::TryGetBaseColorAndOpacity(FGLTFJsonPBRMetallicRoughness&
 		bHasOpacitySourceTexture &&
 		BaseColorTexture == OpacityTexture &&
 		BaseColorTexCoord == OpacityTexCoord &&
-		BaseColorTransform.IsExactEqual(OpacityTransform))
+		BaseColorTransform.IsExactlyEqual(OpacityTransform))
 	{
 		OutPBRParams.BaseColorTexture.Index = Builder.GetOrAddTexture(BaseColorTexture);
 		OutPBRParams.BaseColorTexture.TexCoord = BaseColorTexCoord;
@@ -446,7 +446,7 @@ bool FGLTFMaterialTask::TryGetMetallicAndRoughness(FGLTFJsonPBRMetallicRoughness
 		bHasRoughnessSourceTexture &&
 		MetallicTexture == RoughnessTexture &&
 		MetallicTexCoord == RoughnessTexCoord &&
-		MetallicTransform.IsExactEqual(RoughnessTransform))
+		MetallicTransform.IsExactlyEqual(RoughnessTransform))
 	{
 		OutPBRParams.MetallicRoughnessTexture.Index = Builder.GetOrAddTexture(MetallicTexture);
 		OutPBRParams.MetallicRoughnessTexture.TexCoord = MetallicTexCoord;
@@ -571,7 +571,7 @@ bool FGLTFMaterialTask::TryGetClearCoatRoughness(FGLTFJsonClearCoatExtension& Ou
 		bHasRoughnessSourceTexture &&
 		IntensityTexture == RoughnessTexture &&
 		IntensityTexCoord == RoughnessTexCoord &&
-		IntensityTransform.IsExactEqual(RoughnessTransform))
+		IntensityTransform.IsExactlyEqual(RoughnessTransform))
 	{
 		const FGLTFJsonTextureIndex TextureIndex = Builder.GetOrAddTexture(IntensityTexture);
 		OutExtParams.ClearCoatTexture.Index = TextureIndex;
@@ -1055,7 +1055,7 @@ bool FGLTFMaterialTask::TryGetSourceTexture(const UTexture2D*& OutTexture, int32
 			return false;
 		}
 
-		if (!Builder.ExportOptions->bExportTextureTransforms && !OutTransform.IsDefault())
+		if (!Builder.ExportOptions->bExportTextureTransforms && !OutTransform.IsExactlyDefault())
 		{
 			Builder.AddWarningMessage(FString::Printf(
 				TEXT("Texture coordinates [%d] in %s for material %s are transformed, but texture transform is disabled by export options"),
@@ -1092,7 +1092,7 @@ bool FGLTFMaterialTask::TryGetSourceTexture(const UTexture2D*& OutTexture, int32
 			return false;
 		}
 
-		if (!Builder.ExportOptions->bExportTextureTransforms && !OutTransform.IsDefault())
+		if (!Builder.ExportOptions->bExportTextureTransforms && !OutTransform.IsExactlyDefault())
 		{
 			Builder.AddWarningMessage(FString::Printf(
 				TEXT("Texture coordinates [%d] in %s for material %s are transformed, but texture transform is disabled by export options"),
