[GLTFExporter] feat: add export option bSkipNearDefaultValues to allow configuration of whether float-based json properties should be compared to nearly or exactly default value during serialization

commit 15a8b7834d3a51f01b2dc8780fb3d003085e18e2
Date: Sun, 13 Jun 2021 01:55:49 +0200

    feat: add export option bSkipNearDefaultValues to allow configuration of whether float-based json properties should be compared to nearly or exactly default value during serialization
    (cherry picked from commit 6f4ea36a719cfaca7e8a93e417ee1b3f02c8625d​)

#skipci

diff --git a/Engine/Plugins/Enterprise/GLTFExporter/MigratedFromGit/Docs/export-options-reference.md b/Engine/Plugins/Enterprise/GLTFExporter/MigratedFromGit/Docs/export-options-reference.md
--- a/Engine/Plugins/Enterprise/GLTFExporter/MigratedFromGit/Docs/export-options-reference.md
+++ b/Engine/Plugins/Enterprise/GLTFExporter/MigratedFromGit/Docs/export-options-reference.md
@@ -5,6 +5,7 @@ Option                         | Description
 `Export Uniform Scale`         | Scale factor used for exporting all assets (0.01 by default) for conversion from centimeters (Unreal default) to meters (glTF).
 `Export Preview Mesh`          | If enabled, the preview mesh for a standalone animation or material asset will also be exported.
 `Strict Compliance`            | If enabled, certain values (like HDR colors and light angles) will be truncated during export to strictly conform to the formal glTF specification.
+`Skip Near Default Values`     | If enabled, floating-point-based JSON properties that are nearly equal to their default value will not be exported and thus regarded as exactly default, reducing size of JSON data.
 `Include Generator Version`    | If enabled, version info for Unreal Engine and exporter plugin will be included as metadata in the glTF asset, which is useful when reporting issues.
 `Export Unlit Materials`       | If enabled, materials with shading model unlit will be properly exported. Uses extension KHR_materials_unlit.
 `Export Clear Coat Materials`  | If enabled, materials with shading model clear coat will be properly exported. Uses extension KHR_materials_clearcoat, which is not supported by all glTF viewers.
diff --git a/Engine/Plugins/Enterprise/GLTFExporter/MigratedFromGit/Source/GLTFExporter/Private/Builders/GLTFJsonBuilder.cpp b/Engine/Plugins/Enterprise/GLTFExporter/MigratedFromGit/Source/GLTFExporter/Private/Builders/GLTFJsonBuilder.cpp
--- a/Engine/Plugins/Enterprise/GLTFExporter/MigratedFromGit/Source/GLTFExporter/Private/Builders/GLTFJsonBuilder.cpp
+++ b/Engine/Plugins/Enterprise/GLTFExporter/MigratedFromGit/Source/GLTFExporter/Private/Builders/GLTFJsonBuilder.cpp
@@ -14,7 +14,7 @@ FGLTFJsonBuilder::FGLTFJsonBuilder(const FString& FilePath, const UGLTFExportOpt
 
 void FGLTFJsonBuilder::WriteJson(FArchive& Archive)
 {
-	JsonRoot.WriteJson(Archive, !bIsGlbFile);
+	JsonRoot.WriteJson(Archive, !bIsGlbFile, ExportOptions->bSkipNearDefaultValues ? KINDA_SMALL_NUMBER : 0);
 }
 
 TSet<EGLTFJsonExtension> FGLTFJsonBuilder::GetCustomExtensionsUsed() const
diff --git a/Engine/Plugins/Enterprise/GLTFExporter/MigratedFromGit/Source/GLTFExporter/Private/GLTFExportOptions.cpp b/Engine/Plugins/Enterprise/GLTFExporter/MigratedFromGit/Source/GLTFExporter/Private/GLTFExportOptions.cpp
--- a/Engine/Plugins/Enterprise/GLTFExporter/MigratedFromGit/Source/GLTFExporter/Private/GLTFExportOptions.cpp
+++ b/Engine/Plugins/Enterprise/GLTFExporter/MigratedFromGit/Source/GLTFExporter/Private/GLTFExportOptions.cpp
@@ -15,6 +15,7 @@ void UGLTFExportOptions::ResetToDefault()
 	ExportUniformScale = 0.01;
 	bExportPreviewMesh = true;
 	bStrictCompliance = true;
+	bSkipNearDefaultValues = true;
 	bIncludeGeneratorVersion = true;
 	bExportUnlitMaterials = true;
 	bExportClearCoatMaterials = true;
diff --git a/Engine/Plugins/Enterprise/GLTFExporter/MigratedFromGit/Source/GLTFExporter/Private/Json/GLTFJsonAnimationPlayback.h b/Engine/Plugins/Enterprise/GLTFExporter/MigratedFromGit/Source/GLTFExporter/Private/Json/GLTFJsonAnimationPlayback.h
--- a/Engine/Plugins/Enterprise/GLTFExporter/MigratedFromGit/Source/GLTFExporter/Private/Json/GLTFJsonAnimationPlayback.h
+++ b/Engine/Plugins/Enterprise/GLTFExporter/MigratedFromGit/Source/GLTFExporter/Private/Json/GLTFJsonAnimationPlayback.h
@@ -39,12 +39,12 @@ struct FGLTFJsonAnimationPlayback : IGLTFJsonObject
 			Writer.Write(TEXT("autoPlay"), bAutoPlay);
 		}
 
-		if (!FMath::IsNearlyEqual(PlayRate, 1))
+		if (!FMath::IsNearlyEqual(PlayRate, 1, Writer.DefaultTolerance))
 		{
 			Writer.Write(TEXT("playRate"), PlayRate);
 		}
 
-		if (!FMath::IsNearlyEqual(StartTime, 0))
+		if (!FMath::IsNearlyEqual(StartTime, 0, Writer.DefaultTolerance))
 		{
 			Writer.Write(TEXT("startTime"), StartTime);
 		}
diff --git a/Engine/Plugins/Enterprise/GLTFExporter/MigratedFromGit/Source/GLTFExporter/Private/Json/GLTFJsonCamera.h b/Engine/Plugins/Enterprise/GLTFExporter/MigratedFromGit/Source/GLTFExporter/Private/Json/GLTFJsonCamera.h
--- a/Engine/Plugins/Enterprise/GLTFExporter/MigratedFromGit/Source/GLTFExporter/Private/Json/GLTFJsonCamera.h
+++ b/Engine/Plugins/Enterprise/GLTFExporter/MigratedFromGit/Source/GLTFExporter/Private/Json/GLTFJsonCamera.h
@@ -47,14 +47,14 @@ struct FGLTFJsonPerspective : IGLTFJsonObject
 
 	virtual void WriteObject(IGLTFJsonWriter& Writer) const override
 	{
-		if (!FMath::IsNearlyEqual(AspectRatio, 0))
+		if (!FMath::IsNearlyEqual(AspectRatio, 0, Writer.DefaultTolerance))
 		{
 			Writer.Write(TEXT("aspectRatio"), AspectRatio);
 		}
 
 		Writer.Write(TEXT("yfov"), YFov);
 
-		if (!FMath::IsNearlyEqual(ZFar, 0))
+		if (!FMath::IsNearlyEqual(ZFar, 0, Writer.DefaultTolerance))
 		{
 			Writer.Write(TEXT("zfar"), ZFar);
 		}
diff --git a/Engine/Plugins/Enterprise/GLTFExporter/MigratedFromGit/Source/GLTFExporter/Private/Json/GLTFJsonCameraControl.h b/Engine/Plugins/Enterprise/GLTFExporter/MigratedFromGit/Source/GLTFExporter/Private/Json/GLTFJsonCameraControl.h
--- a/Engine/Plugins/Enterprise/GLTFExporter/MigratedFromGit/Source/GLTFExporter/Private/Json/GLTFJsonCameraControl.h
+++ b/Engine/Plugins/Enterprise/GLTFExporter/MigratedFromGit/Source/GLTFExporter/Private/Json/GLTFJsonCameraControl.h
@@ -48,22 +48,22 @@ struct FGLTFJsonCameraControl : IGLTFJsonObject
 		Writer.Write(TEXT("maxDistance"), MaxDistance);
 		Writer.Write(TEXT("minDistance"), MinDistance);
 
-		if (!FMath::IsNearlyEqual(MaxPitch, 90))
+		if (!FMath::IsNearlyEqual(MaxPitch, 90, Writer.DefaultTolerance))
 		{
 			Writer.Write(TEXT("maxPitch"), MaxPitch);
 		}
 
-		if (!FMath::IsNearlyEqual(MinPitch, -90))
+		if (!FMath::IsNearlyEqual(MinPitch, -90, Writer.DefaultTolerance))
 		{
 			Writer.Write(TEXT("minPitch"), MinPitch);
 		}
 
-		if (!FMath::IsNearlyEqual(MaxYaw, 360))
+		if (!FMath::IsNearlyEqual(MaxYaw, 360, Writer.DefaultTolerance))
 		{
 			Writer.Write(TEXT("maxYaw"), MaxYaw);
 		}
 
-		if (!FMath::IsNearlyEqual(MinYaw, 0))
+		if (!FMath::IsNearlyEqual(MinYaw, 0, Writer.DefaultTolerance))
 		{
 			Writer.Write(TEXT("minYaw"), MinYaw);
 		}
diff --git a/Engine/Plugins/Enterprise/GLTFExporter/MigratedFromGit/Source/GLTFExporter/Private/Json/GLTFJsonLight.h b/Engine/Plugins/Enterprise/GLTFExporter/MigratedFromGit/Source/GLTFExporter/Private/Json/GLTFJsonLight.h
--- a/Engine/Plugins/Enterprise/GLTFExporter/MigratedFromGit/Source/GLTFExporter/Private/Json/GLTFJsonLight.h
+++ b/Engine/Plugins/Enterprise/GLTFExporter/MigratedFromGit/Source/GLTFExporter/Private/Json/GLTFJsonLight.h
@@ -19,12 +19,12 @@ struct FGLTFJsonSpotLight : IGLTFJsonObject
 
 	virtual void WriteObject(IGLTFJsonWriter& Writer) const override
 	{
-		if (!FMath::IsNearlyEqual(InnerConeAngle, 0))
+		if (!FMath::IsNearlyEqual(InnerConeAngle, 0, Writer.DefaultTolerance))
 		{
 			Writer.Write(TEXT("innerConeAngle"), InnerConeAngle);
 		}
 
-		if (!FMath::IsNearlyEqual(OuterConeAngle, HALF_PI))
+		if (!FMath::IsNearlyEqual(OuterConeAngle, HALF_PI, Writer.DefaultTolerance))
 		{
 			Writer.Write(TEXT("outerConeAngle"), OuterConeAngle);
 		}
@@ -61,19 +61,19 @@ struct FGLTFJsonLight : IGLTFJsonObject
 
 		Writer.Write(TEXT("type"), Type);
 
-		if (!Color.IsNearlyEqual(FGLTFJsonColor3::White))
+		if (!Color.IsNearlyEqual(FGLTFJsonColor3::White, Writer.DefaultTolerance))
 		{
 			Writer.Write(TEXT("color"), Color);
 		}
 
-		if (!FMath::IsNearlyEqual(Intensity, 1))
+		if (!FMath::IsNearlyEqual(Intensity, 1, Writer.DefaultTolerance))
 		{
 			Writer.Write(TEXT("intensity"), Intensity);
 		}
 
 		if (Type == EGLTFJsonLightType::Point || Type == EGLTFJsonLightType::Spot)
 		{
-			if (!FMath::IsNearlyEqual(Range, 0))
+			if (!FMath::IsNearlyEqual(Range, 0, Writer.DefaultTolerance))
 			{
 				Writer.Write(TEXT("range"), Range);
 			}
diff --git a/Engine/Plugins/Enterprise/GLTFExporter/MigratedFromGit/Source/GLTFExporter/Private/Json/GLTFJsonMaterial.h b/Engine/Plugins/Enterprise/GLTFExporter/MigratedFromGit/Source/GLTFExporter/Private/Json/GLTFJsonMaterial.h
--- a/Engine/Plugins/Enterprise/GLTFExporter/MigratedFromGit/Source/GLTFExporter/Private/Json/GLTFJsonMaterial.h
+++ b/Engine/Plugins/Enterprise/GLTFExporter/MigratedFromGit/Source/GLTFExporter/Private/Json/GLTFJsonMaterial.h
@@ -28,7 +28,7 @@ struct FGLTFJsonTextureInfo : IGLTFJsonObject
 			Writer.Write(TEXT("texCoord"), TexCoord);
 		}
 
-		if (!Transform.IsNearlyDefault())
+		if (!Transform.IsNearlyDefault(Writer.DefaultTolerance))
 		{
 			Writer.StartExtensions();
 			Writer.Write(EGLTFJsonExtension::KHR_TextureTransform, Transform);
@@ -55,7 +55,7 @@ struct FGLTFJsonNormalTextureInfo : FGLTFJsonTextureInfo
 			Writer.Write(TEXT("texCoord"), TexCoord);
 		}
 
-		if (!FMath::IsNearlyEqual(Scale, 1))
+		if (!FMath::IsNearlyEqual(Scale, 1, Writer.DefaultTolerance))
 		{
 			Writer.Write(TEXT("scale"), Scale);
 		}
@@ -80,7 +80,7 @@ struct FGLTFJsonOcclusionTextureInfo : FGLTFJsonTextureInfo
 			Writer.Write(TEXT("texCoord"), TexCoord);
 		}
 
-		if (!FMath::IsNearlyEqual(Strength, 1))
+		if (!FMath::IsNearlyEqual(Strength, 1, Writer.DefaultTolerance))
 		{
 			Writer.Write(TEXT("strength"), Strength);
 		}
@@ -105,7 +105,7 @@ struct FGLTFJsonPBRMetallicRoughness : IGLTFJsonObject
 
 	virtual void WriteObject(IGLTFJsonWriter& Writer) const override
 	{
-		if (!BaseColorFactor.IsNearlyEqual(FGLTFJsonColor4::White))
+		if (!BaseColorFactor.IsNearlyEqual(FGLTFJsonColor4::White, Writer.DefaultTolerance))
 		{
 			Writer.Write(TEXT("baseColorFactor"), BaseColorFactor);
 		}
@@ -115,12 +115,12 @@ struct FGLTFJsonPBRMetallicRoughness : IGLTFJsonObject
 			Writer.Write(TEXT("baseColorTexture"), BaseColorTexture);
 		}
 
-		if (!FMath::IsNearlyEqual(MetallicFactor, 1))
+		if (!FMath::IsNearlyEqual(MetallicFactor, 1, Writer.DefaultTolerance))
 		{
 			Writer.Write(TEXT("metallicFactor"), MetallicFactor);
 		}
 
-		if (!FMath::IsNearlyEqual(RoughnessFactor, 1))
+		if (!FMath::IsNearlyEqual(RoughnessFactor, 1, Writer.DefaultTolerance))
 		{
 			Writer.Write(TEXT("roughnessFactor"), RoughnessFactor);
 		}
@@ -150,7 +150,7 @@ struct FGLTFJsonClearCoatExtension : IGLTFJsonObject
 
 	virtual void WriteObject(IGLTFJsonWriter& Writer) const override
 	{
-		if (!FMath::IsNearlyEqual(ClearCoatFactor, 0))
+		if (!FMath::IsNearlyEqual(ClearCoatFactor, 0, Writer.DefaultTolerance))
 		{
 			Writer.Write(TEXT("clearcoatFactor"), ClearCoatFactor);
 		}
@@ -160,7 +160,7 @@ struct FGLTFJsonClearCoatExtension : IGLTFJsonObject
 			Writer.Write(TEXT("clearcoatTexture"), ClearCoatTexture);
 		}
 
-		if (!FMath::IsNearlyEqual(ClearCoatRoughnessFactor, 0))
+		if (!FMath::IsNearlyEqual(ClearCoatRoughnessFactor, 0, Writer.DefaultTolerance))
 		{
 			Writer.Write(TEXT("clearcoatRoughnessFactor"), ClearCoatRoughnessFactor);
 		}
@@ -237,7 +237,7 @@ struct FGLTFJsonMaterial : IGLTFJsonObject
 			Writer.Write(TEXT("emissiveTexture"), EmissiveTexture);
 		}
 
-		if (!EmissiveFactor.IsNearlyEqual(FGLTFJsonColor3::Black))
+		if (!EmissiveFactor.IsNearlyEqual(FGLTFJsonColor3::Black, Writer.DefaultTolerance))
 		{
 			Writer.Write(TEXT("emissiveFactor"), EmissiveFactor);
 		}
@@ -246,7 +246,7 @@ struct FGLTFJsonMaterial : IGLTFJsonObject
 		{
 			Writer.Write(TEXT("alphaMode"), AlphaMode);
 
-			if (AlphaMode == EGLTFJsonAlphaMode::Mask && !FMath::IsNearlyEqual(AlphaCutoff, 0.5f))
+			if (AlphaMode == EGLTFJsonAlphaMode::Mask && !FMath::IsNearlyEqual(AlphaCutoff, 0.5f, Writer.DefaultTolerance))
 			{
 				Writer.Write(TEXT("alphaCutoff"), AlphaCutoff);
 			}
diff --git a/Engine/Plugins/Enterprise/GLTFExporter/MigratedFromGit/Source/GLTFExporter/Private/Json/GLTFJsonNode.h b/Engine/Plugins/Enterprise/GLTFExporter/MigratedFromGit/Source/GLTFExporter/Private/Json/GLTFJsonNode.h
--- a/Engine/Plugins/Enterprise/GLTFExporter/MigratedFromGit/Source/GLTFExporter/Private/Json/GLTFJsonNode.h
+++ b/Engine/Plugins/Enterprise/GLTFExporter/MigratedFromGit/Source/GLTFExporter/Private/Json/GLTFJsonNode.h
@@ -50,24 +50,24 @@ struct FGLTFJsonNode : IGLTFJsonObject
 
 		if (bUseMatrix)
 		{
-			if (!Matrix.IsNearlyEqual(FGLTFJsonMatrix4::Identity))
+			if (!Matrix.IsNearlyEqual(FGLTFJsonMatrix4::Identity, Writer.DefaultTolerance))
 			{
 				Writer.Write(TEXT("matrix"), Matrix);
 			}
 		}
 		else
 		{
-			if (!Translation.IsNearlyEqual(FGLTFJsonVector3::Zero))
+			if (!Translation.IsNearlyEqual(FGLTFJsonVector3::Zero, Writer.DefaultTolerance))
 			{
 				Writer.Write(TEXT("translation"), Translation);
 			}
 
-			if (!Rotation.IsNearlyEqual(FGLTFJsonQuaternion::Identity))
+			if (!Rotation.IsNearlyEqual(FGLTFJsonQuaternion::Identity, Writer.DefaultTolerance))
 			{
 				Writer.Write(TEXT("rotation"), Rotation);
 			}
 
-			if (!Scale.IsNearlyEqual(FGLTFJsonVector3::One))
+			if (!Scale.IsNearlyEqual(FGLTFJsonVector3::One, Writer.DefaultTolerance))
 			{
 				Writer.Write(TEXT("scale"), Scale);
 			}
diff --git a/Engine/Plugins/Enterprise/GLTFExporter/MigratedFromGit/Source/GLTFExporter/Private/Json/GLTFJsonRoot.h b/Engine/Plugins/Enterprise/GLTFExporter/MigratedFromGit/Source/GLTFExporter/Private/Json/GLTFJsonRoot.h
--- a/Engine/Plugins/Enterprise/GLTFExporter/MigratedFromGit/Source/GLTFExporter/Private/Json/GLTFJsonRoot.h
+++ b/Engine/Plugins/Enterprise/GLTFExporter/MigratedFromGit/Source/GLTFExporter/Private/Json/GLTFJsonRoot.h
@@ -137,9 +137,10 @@ struct FGLTFJsonRoot : IGLTFJsonObject
 		}
 	}
 
-	void WriteJson(FArchive& Archive, bool bPrettyJson)
+	void WriteJson(FArchive& Archive, bool bPrettyJson, float DefaultTolerance)
 	{
 		TSharedRef<IGLTFJsonWriter> Writer = IGLTFJsonWriter::Create(Archive, bPrettyJson, Extensions);
+		Writer->DefaultTolerance = DefaultTolerance;
 		Writer->Write(this);
 		Writer->Close();
 	}
diff --git a/Engine/Plugins/Enterprise/GLTFExporter/MigratedFromGit/Source/GLTFExporter/Private/Json/GLTFJsonSkySphere.h b/Engine/Plugins/Enterprise/GLTFExporter/MigratedFromGit/Source/GLTFExporter/Private/Json/GLTFJsonSkySphere.h
--- a/Engine/Plugins/Enterprise/GLTFExporter/MigratedFromGit/Source/GLTFExporter/Private/Json/GLTFJsonSkySphere.h
+++ b/Engine/Plugins/Enterprise/GLTFExporter/MigratedFromGit/Source/GLTFExporter/Private/Json/GLTFJsonSkySphere.h
@@ -127,7 +127,7 @@ struct FGLTFJsonSkySphere : IGLTFJsonObject
 		Writer.Write(TEXT("horizonColor"), HorizonColor);
 		Writer.Write(TEXT("cloudColor"), CloudColor);
 
-		if (!OverallColor.IsNearlyEqual(FGLTFJsonColor4::White))
+		if (!OverallColor.IsNearlyEqual(FGLTFJsonColor4::White, Writer.DefaultTolerance))
 		{
 			Writer.Write(TEXT("overallColor"), OverallColor);
 		}
@@ -147,7 +147,7 @@ struct FGLTFJsonSkySphere : IGLTFJsonObject
 			Writer.Write(TEXT("cloudColorCurve"), CloudColorCurve);
 		}
 
-		if (!Scale.IsNearlyEqual(FGLTFJsonVector3::One))
+		if (!Scale.IsNearlyEqual(FGLTFJsonVector3::One, Writer.DefaultTolerance))
 		{
 			Writer.Write(TEXT("scale"), Scale);
 		}
diff --git a/Engine/Plugins/Enterprise/GLTFExporter/MigratedFromGit/Source/GLTFExporter/Private/Json/GLTFJsonTextureTransform.h b/Engine/Plugins/Enterprise/GLTFExporter/MigratedFromGit/Source/GLTFExporter/Private/Json/GLTFJsonTextureTransform.h
--- a/Engine/Plugins/Enterprise/GLTFExporter/MigratedFromGit/Source/GLTFExporter/Private/Json/GLTFJsonTextureTransform.h
+++ b/Engine/Plugins/Enterprise/GLTFExporter/MigratedFromGit/Source/GLTFExporter/Private/Json/GLTFJsonTextureTransform.h
@@ -20,17 +20,17 @@ struct FGLTFJsonTextureTransform : IGLTFJsonObject
 
 	virtual void WriteObject(IGLTFJsonWriter& Writer) const override
 	{
-		if (!Offset.IsNearlyEqual(FGLTFJsonVector2::Zero))
+		if (!Offset.IsNearlyEqual(FGLTFJsonVector2::Zero, Writer.DefaultTolerance))
 		{
 			Writer.Write(TEXT("offset"), Offset);
 		}
 
-		if (!Scale.IsNearlyEqual(FGLTFJsonVector2::One))
+		if (!Scale.IsNearlyEqual(FGLTFJsonVector2::One, Writer.DefaultTolerance))
 		{
 			Writer.Write(TEXT("scale"), Scale);
 		}
 
-		if (!FMath::IsNearlyEqual(Rotation, 0))
+		if (!FMath::IsNearlyEqual(Rotation, 0, Writer.DefaultTolerance))
 		{
 			Writer.Write(TEXT("rotation"), Rotation);
 		}
diff --git a/Engine/Plugins/Enterprise/GLTFExporter/MigratedFromGit/Source/GLTFExporter/Private/Json/GLTFJsonWriter.h b/Engine/Plugins/Enterprise/GLTFExporter/MigratedFromGit/Source/GLTFExporter/Private/Json/GLTFJsonWriter.h
--- a/Engine/Plugins/Enterprise/GLTFExporter/MigratedFromGit/Source/GLTFExporter/Private/Json/GLTFJsonWriter.h
+++ b/Engine/Plugins/Enterprise/GLTFExporter/MigratedFromGit/Source/GLTFExporter/Private/Json/GLTFJsonWriter.h
@@ -10,10 +10,13 @@ class IGLTFJsonWriter
 {
 public:
 
+	float DefaultTolerance;
+
 	FGLTFJsonExtensions& Extensions;
 
 	IGLTFJsonWriter(FGLTFJsonExtensions& Extensions)
-		: Extensions(Extensions)
+		: DefaultTolerance(0)
+		, Extensions(Extensions)
 	{
 	}
 
diff --git a/Engine/Plugins/Enterprise/GLTFExporter/MigratedFromGit/Source/GLTFExporter/Public/GLTFExportOptions.h b/Engine/Plugins/Enterprise/GLTFExporter/MigratedFromGit/Source/GLTFExporter/Public/GLTFExportOptions.h
--- a/Engine/Plugins/Enterprise/GLTFExporter/MigratedFromGit/Source/GLTFExporter/Public/GLTFExportOptions.h
+++ b/Engine/Plugins/Enterprise/GLTFExporter/MigratedFromGit/Source/GLTFExporter/Public/GLTFExportOptions.h
@@ -80,6 +80,10 @@ class GLTFEXPORTER_API UGLTFExportOptions : public UObject
 	UPROPERTY(EditAnywhere, BlueprintReadWrite, Config, Category = General)
 	bool bStrictCompliance;
 
+	/** If enabled, floating-point-based JSON properties that are nearly equal to their default value will not be exported and thus regarded as exactly default, reducing JSON size. */
+	UPROPERTY(EditAnywhere, BlueprintReadWrite, Config, Category = General)
+	bool bSkipNearDefaultValues;
+
 	/** If enabled, version info for Unreal Engine and exporter plugin will be included as metadata in the glTF asset, which is useful when reporting issues. */
 	UPROPERTY(EditAnywhere, BlueprintReadWrite, Config, Category = General)
 	bool bIncludeGeneratorVersion;
