[GLTFExporter] refactor: remove matrix support in node (json object) which its no longer needed

commit 26c8c2551ae44e6be6b50f78b53884593a335f4b
Date: Mon, 14 Jun 2021 01:06:46 +0200

    refactor: remove matrix support in node (json object) which its no longer needed
    (cherry picked from commit f3358907530b7d84213536c744495f8b7ff554fe​)

#skipci

diff --git a/Engine/Plugins/Enterprise/GLTFExporter/MigratedFromGit/Source/GLTFExporter/Private/Json/GLTFJsonNode.h b/Engine/Plugins/Enterprise/GLTFExporter/MigratedFromGit/Source/GLTFExporter/Private/Json/GLTFJsonNode.h
--- a/Engine/Plugins/Enterprise/GLTFExporter/MigratedFromGit/Source/GLTFExporter/Private/Json/GLTFJsonNode.h
+++ b/Engine/Plugins/Enterprise/GLTFExporter/MigratedFromGit/Source/GLTFExporter/Private/Json/GLTFJsonNode.h
@@ -4,7 +4,6 @@
 
 #include "Json/GLTFJsonObject.h"
 #include "Json/GLTFJsonIndex.h"
-#include "Json/GLTFJsonMatrix.h"
 #include "Json/GLTFJsonVector.h"
 #include "Json/GLTFJsonQuaternion.h"
 
@@ -12,9 +11,6 @@ struct FGLTFJsonNode : IGLTFJsonObject
 {
 	FString Name;
 
-	bool bUseMatrix;
-
-	FGLTFJsonMatrix4    Matrix;
 	FGLTFJsonVector3    Translation;
 	FGLTFJsonQuaternion Rotation;
 	FGLTFJsonVector3    Scale;
@@ -33,9 +29,7 @@ struct FGLTFJsonNode : IGLTFJsonObject
 	TArray<FGLTFJsonNodeIndex> Children;
 
 	FGLTFJsonNode()
-		: bUseMatrix(false)
-		, Matrix(FGLTFJsonMatrix4::Identity)
-		, Translation(FGLTFJsonVector3::Zero)
+		: Translation(FGLTFJsonVector3::Zero)
 		, Rotation(FGLTFJsonQuaternion::Identity)
 		, Scale(FGLTFJsonVector3::One)
 	{
@@ -48,29 +42,19 @@ struct FGLTFJsonNode : IGLTFJsonObject
 			Writer.Write(TEXT("name"), Name);
 		}
 
-		if (bUseMatrix)
+		if (!Translation.IsNearlyEqual(FGLTFJsonVector3::Zero, Writer.DefaultTolerance))
 		{
-			if (!Matrix.IsNearlyEqual(FGLTFJsonMatrix4::Identity, Writer.DefaultTolerance))
-			{
-				Writer.Write(TEXT("matrix"), Matrix);
-			}
+			Writer.Write(TEXT("translation"), Translation);
 		}
-		else
-		{
-			if (!Translation.IsNearlyEqual(FGLTFJsonVector3::Zero, Writer.DefaultTolerance))
-			{
-				Writer.Write(TEXT("translation"), Translation);
-			}
 
-			if (!Rotation.IsNearlyEqual(FGLTFJsonQuaternion::Identity, Writer.DefaultTolerance))
-			{
-				Writer.Write(TEXT("rotation"), Rotation);
-			}
+		if (!Rotation.IsNearlyEqual(FGLTFJsonQuaternion::Identity, Writer.DefaultTolerance))
+		{
+			Writer.Write(TEXT("rotation"), Rotation);
+		}
 
-			if (!Scale.IsNearlyEqual(FGLTFJsonVector3::One, Writer.DefaultTolerance))
-			{
-				Writer.Write(TEXT("scale"), Scale);
-			}
+		if (!Scale.IsNearlyEqual(FGLTFJsonVector3::One, Writer.DefaultTolerance))
+		{
+			Writer.Write(TEXT("scale"), Scale);
 		}
 
 		if (Camera != INDEX_NONE)
