[GLTFExporter] refactor: adjust engine-version-specific code snippets to avoid negation

commit 1dcb9e2234d6445097f158a9c722f3ffcde38f66
Date: Fri, 25 Jun 2021 22:26:48 +0200

    refactor: adjust engine-version-specific code snippets to avoid negation
    (cherry picked from commit 071d971835d171426b9824edc0089b455ef443fb​)

#skipci

diff --git a/Engine/Plugins/Enterprise/GLTFExporter/MigratedFromGit/Source/GLTFExporter/Private/Converters/GLTFIndexArray.h b/Engine/Plugins/Enterprise/GLTFExporter/MigratedFromGit/Source/GLTFExporter/Private/Converters/GLTFIndexArray.h
--- a/Engine/Plugins/Enterprise/GLTFExporter/MigratedFromGit/Source/GLTFExporter/Private/Converters/GLTFIndexArray.h
+++ b/Engine/Plugins/Enterprise/GLTFExporter/MigratedFromGit/Source/GLTFExporter/Private/Converters/GLTFIndexArray.h
@@ -6,7 +6,7 @@
 
 typedef TArray<int32> FGLTFIndexArray;
 
-#if !(ENGINE_MAJOR_VERSION == 4 && ENGINE_MINOR_VERSION <= 25)
+#if (ENGINE_MAJOR_VERSION > 4 || ENGINE_MINOR_VERSION >= 26)
 inline uint32 GetTypeHash(const FGLTFIndexArray& IndexArray)
 {
 	uint32 Hash = GetTypeHash(IndexArray.Num());
diff --git a/Engine/Plugins/Enterprise/GLTFExporter/MigratedFromGit/Source/GLTFExporter/Private/Tasks/GLTFAnimationTasks.cpp b/Engine/Plugins/Enterprise/GLTFExporter/MigratedFromGit/Source/GLTFExporter/Private/Tasks/GLTFAnimationTasks.cpp
--- a/Engine/Plugins/Enterprise/GLTFExporter/MigratedFromGit/Source/GLTFExporter/Private/Tasks/GLTFAnimationTasks.cpp
+++ b/Engine/Plugins/Enterprise/GLTFExporter/MigratedFromGit/Source/GLTFExporter/Private/Tasks/GLTFAnimationTasks.cpp
@@ -11,7 +11,7 @@
 #include "Sections/MovieScene3DTransformSection.h"
 #include "Channels/MovieSceneChannelProxy.h"
 
-#if !(ENGINE_MAJOR_VERSION == 4 && ENGINE_MINOR_VERSION <= 25)
+#if (ENGINE_MAJOR_VERSION > 4 || ENGINE_MINOR_VERSION >= 26)
 using namespace UE;
 #endif
 
diff --git a/Engine/Plugins/Enterprise/GLTFExporter/MigratedFromGit/Source/GLTFExporterRuntime/Private/Actors/GLTFHotspotActor.cpp b/Engine/Plugins/Enterprise/GLTFExporter/MigratedFromGit/Source/GLTFExporterRuntime/Private/Actors/GLTFHotspotActor.cpp
--- a/Engine/Plugins/Enterprise/GLTFExporter/MigratedFromGit/Source/GLTFExporterRuntime/Private/Actors/GLTFHotspotActor.cpp
+++ b/Engine/Plugins/Enterprise/GLTFExporter/MigratedFromGit/Source/GLTFExporterRuntime/Private/Actors/GLTFHotspotActor.cpp
@@ -20,7 +20,7 @@
 #include "Slate/SceneViewport.h"
 #endif
 
-#if !(ENGINE_MAJOR_VERSION == 4 && ENGINE_MINOR_VERSION <= 25)
+#if (ENGINE_MAJOR_VERSION > 4 || ENGINE_MINOR_VERSION >= 26)
 using namespace UE;
 #endif
 
diff --git a/Engine/Plugins/Enterprise/GLTFExporter/MigratedFromGit/Source/GLTFMaterialAnalyzer/Private/GLTFProxyMaterialCompiler.h b/Engine/Plugins/Enterprise/GLTFExporter/MigratedFromGit/Source/GLTFMaterialAnalyzer/Private/GLTFProxyMaterialCompiler.h
--- a/Engine/Plugins/Enterprise/GLTFExporter/MigratedFromGit/Source/GLTFMaterialAnalyzer/Private/GLTFProxyMaterialCompiler.h
+++ b/Engine/Plugins/Enterprise/GLTFExporter/MigratedFromGit/Source/GLTFMaterialAnalyzer/Private/GLTFProxyMaterialCompiler.h
@@ -63,7 +63,7 @@ class FGLTFProxyMaterialCompiler : public FProxyMaterialCompiler
 		return EMaterialCompilerType::MaterialProxy;
 	}
 
-#if !(ENGINE_MAJOR_VERSION == 4 && ENGINE_MINOR_VERSION <= 25)
+#if (ENGINE_MAJOR_VERSION > 4 || ENGINE_MINOR_VERSION >= 26)
 	virtual int32 PreSkinVertexOffset() override
 	{
 		return Compiler->PreSkinVertexOffset();
diff --git a/Engine/Plugins/Enterprise/GLTFExporter/MigratedFromGit/Source/GLTFMaterialBaking/Private/ExportMaterialProxy.h b/Engine/Plugins/Enterprise/GLTFExporter/MigratedFromGit/Source/GLTFMaterialBaking/Private/ExportMaterialProxy.h
--- a/Engine/Plugins/Enterprise/GLTFExporter/MigratedFromGit/Source/GLTFMaterialBaking/Private/ExportMaterialProxy.h
+++ b/Engine/Plugins/Enterprise/GLTFExporter/MigratedFromGit/Source/GLTFMaterialBaking/Private/ExportMaterialProxy.h
@@ -226,7 +226,7 @@ struct FExportMaterialCompiler : public FProxyMaterialCompiler
 		return EMaterialCompilerType::MaterialProxy;
 	}
 
-#if !(ENGINE_MAJOR_VERSION == 4 && ENGINE_MINOR_VERSION <= 25)
+#if (ENGINE_MAJOR_VERSION > 4 || ENGINE_MINOR_VERSION >= 26)
 	virtual int32 PreSkinVertexOffset() override
 	{
 		return Compiler->PreSkinVertexOffset();
@@ -251,10 +251,10 @@ public:
 		, ProxyBlendMode(ProxyBlendMode)
 	{
 
-#if (ENGINE_MAJOR_VERSION == 4 && ENGINE_MINOR_VERSION <= 25)
-		SetQualityLevelProperties(EMaterialQualityLevel::High, false, GMaxRHIFeatureLevel);
-#else
+#if (ENGINE_MAJOR_VERSION > 4 || ENGINE_MINOR_VERSION >= 26)
 		SetQualityLevelProperties(GMaxRHIFeatureLevel, EMaterialQualityLevel::High);
+#else
+		SetQualityLevelProperties(EMaterialQualityLevel::High, false, GMaxRHIFeatureLevel);
 #endif
 		Material = InMaterialInterface->GetMaterial();
 		ReferencedTextures = InMaterialInterface->GetReferencedTextures();
