[GLTFExporter] fix: fix build issues for UE 4.27 while keeping backwards compatibility with 4.26 and 4.25

commit 640349bb17563f1ef851b6f0a558b90462c9b459
Date: Fri, 25 Jun 2021 23:28:07 +0200

    fix: fix build issues for UE 4.27 while keeping backwards compatibility with 4.26 and 4.25
    (cherry picked from commit 1e1006a395dd923dfeb4409fd4f777f3bedfb914​)

#skipci

diff --git a/Engine/Plugins/Enterprise/GLTFExporter/MigratedFromGit/Source/GLTFMaterialBaking/Private/ExportMaterialProxy.h b/Engine/Plugins/Enterprise/GLTFExporter/MigratedFromGit/Source/GLTFMaterialBaking/Private/ExportMaterialProxy.h
--- a/Engine/Plugins/Enterprise/GLTFExporter/MigratedFromGit/Source/GLTFMaterialBaking/Private/ExportMaterialProxy.h
+++ b/Engine/Plugins/Enterprise/GLTFExporter/MigratedFromGit/Source/GLTFMaterialBaking/Private/ExportMaterialProxy.h
@@ -273,7 +273,12 @@ public:
 			TArray<FShaderType*> ShaderTypes;
 			TArray<FVertexFactoryType*> VFTypes;
 			TArray<const FShaderPipelineType*> ShaderPipelineTypes;
+
+#if (ENGINE_MAJOR_VERSION > 4 || ENGINE_MINOR_VERSION >= 27)
+			GetDependentShaderAndVFTypes(GMaxRHIShaderPlatform, ResourceId.LayoutParams, ShaderTypes, ShaderPipelineTypes, VFTypes);
+#else
 			GetDependentShaderAndVFTypes(GMaxRHIShaderPlatform, ShaderTypes, ShaderPipelineTypes, VFTypes);
+#endif
 
 			// Overwrite the shader map Id's dependencies with ones that came from the FMaterial actually being compiled (this)
 			// This is necessary as we change FMaterial attributes like GetShadingModels(), which factor into the ShouldCache functions that determine dependent shader types
@@ -296,11 +301,18 @@ public:
 		case MP_OpacityMask: ResourceId.Usage = EMaterialShaderMapUsage::MaterialExportOpacityMask; break;
 		case MP_SubsurfaceColor: ResourceId.Usage = EMaterialShaderMapUsage::MaterialExportSubSurfaceColor; break;
 
+#if (ENGINE_MAJOR_VERSION > 4 || ENGINE_MINOR_VERSION >= 27)
+		case MP_CustomData0: ResourceId.Usage = EMaterialShaderMapUsage::MaterialExportClearCoat; break;
+		case MP_CustomData1: ResourceId.Usage = EMaterialShaderMapUsage::MaterialExportClearCoatRoughness; break;
+		case MP_CustomOutput: ResourceId.Usage = EMaterialShaderMapUsage::MaterialExportCustomOutput; break;
+		case MP_ShadingModel: ResourceId.Usage = static_cast<EMaterialShaderMapUsage::Type>(EMaterialShaderMapUsage::MaterialExportCustomOutput + 1); break;
+#else
 		// NOTE: the following cases for custom data 0/1 and custom output are hacks since a proper solution requires engine changes:
 		case MP_CustomData0: ResourceId.Usage = static_cast<EMaterialShaderMapUsage::Type>(EMaterialShaderMapUsage::DebugViewMode + 1); break;
 		case MP_CustomData1: ResourceId.Usage = static_cast<EMaterialShaderMapUsage::Type>(EMaterialShaderMapUsage::DebugViewMode + 2); break;
 		case MP_ShadingModel: ResourceId.Usage = static_cast<EMaterialShaderMapUsage::Type>(EMaterialShaderMapUsage::DebugViewMode + 3); break;
 		case MP_CustomOutput: ResourceId.Usage = static_cast<EMaterialShaderMapUsage::Type>(EMaterialShaderMapUsage::DebugViewMode + 4); break;
+#endif
 
 		default:
 			ensureMsgf(false, TEXT("ExportMaterial has no usage for property %i.  Will likely reuse the normal rendering shader and crash later with a parameter mismatch"), (int32)InPropertyToCompile);
@@ -346,6 +358,22 @@ public:
 
 	////////////////
 	// FMaterialRenderProxy interface.
+
+#if (ENGINE_MAJOR_VERSION > 4 || ENGINE_MINOR_VERSION >= 27)
+	virtual const FMaterial* GetMaterialNoFallback(ERHIFeatureLevel::Type InFeatureLevel) const override
+	{
+		if (GetRenderingThreadShaderMap())
+		{
+			return this;
+		}
+		return nullptr;
+	}
+
+	virtual const FMaterialRenderProxy* GetFallback(ERHIFeatureLevel::Type InFeatureLevel) const override
+	{
+		return UMaterial::GetDefaultMaterial(MD_Surface)->GetRenderProxy();
+	}
+#else
 	virtual const FMaterial& GetMaterialWithFallback(ERHIFeatureLevel::Type FeatureLevel, const FMaterialRenderProxy*& OutFallbackMaterialRenderProxy) const override
 	{
 		if(GetRenderingThreadShaderMap())
@@ -358,6 +386,7 @@ public:
 			return OutFallbackMaterialRenderProxy->GetMaterialWithFallback(FeatureLevel, OutFallbackMaterialRenderProxy);
 		}
 	}
+#endif
 
 	virtual bool GetVectorValue(const FHashedMaterialParameterInfo& ParameterInfo, FLinearColor* OutValue, const FMaterialRenderContext& Context) const override
 	{
