[GLTFExporter] chore: update project and plugin descriptors to UE 4.27

commit 026edcb3d68a8f96b582135f78b159140ececfde
Date: Fri, 25 Jun 2021 23:30:21 +0200

    chore: update project and plugin descriptors to UE 4.27
    (cherry picked from commit 33c2109431b9a9ca963f42adc2c4024fd4383090​)

#skipci

diff --git a/Engine/Plugins/Enterprise/GLTFExporter/MigratedFromGit/GLTFExporter.uplugin b/Engine/Plugins/Enterprise/GLTFExporter/MigratedFromGit/GLTFExporter.uplugin
--- a/Engine/Plugins/Enterprise/GLTFExporter/MigratedFromGit/GLTFExporter.uplugin
+++ b/Engine/Plugins/Enterprise/GLTFExporter/MigratedFromGit/GLTFExporter.uplugin
@@ -10,7 +10,7 @@
 	"DocsURL": "",
 	"MarketplaceURL": "",
 	"SupportURL": "",
-	"EngineVersion": "4.26.0",
+	"EngineVersion": "4.27.0",
 	"CanContainContent": true,
 	"IsBetaVersion": true,
 	"Modules": [
