[GLTFExporter] fix: fix build warnings for UE 4.27 (not compatible with UE 4.26 or 4.25)

commit 8f591b5ba8f440ad18afa6c1df390b4b8bd0deb8
Date: Fri, 25 Jun 2021 23:32:47 +0200

    fix: fix build warnings for UE 4.27 (not compatible with UE 4.26 or 4.25)
    (cherry picked from commit 0a14177cd5e2cd0d45962bca159db93b791ebbbf​)

#skipci

diff --git a/Engine/Plugins/Enterprise/GLTFExporter/MigratedFromGit/Source/GLTFExporter/Private/Converters/GLTFAnimationConverters.cpp b/Engine/Plugins/Enterprise/GLTFExporter/MigratedFromGit/Source/GLTFExporter/Private/Converters/GLTFAnimationConverters.cpp
--- a/Engine/Plugins/Enterprise/GLTFExporter/MigratedFromGit/Source/GLTFExporter/Private/Converters/GLTFAnimationConverters.cpp
+++ b/Engine/Plugins/Enterprise/GLTFExporter/MigratedFromGit/Source/GLTFExporter/Private/Converters/GLTFAnimationConverters.cpp
@@ -20,7 +20,7 @@ FGLTFJsonAnimationIndex FGLTFAnimationConverter::Convert(FGLTFJsonNodeIndex Root
 		return FGLTFJsonAnimationIndex(INDEX_NONE);
 	}
 
-	if (Skeleton != SkeletalMesh->Skeleton)
+	if (Skeleton != SkeletalMesh->GetSkeleton())
 	{
 		// TODO: report error
 		return FGLTFJsonAnimationIndex(INDEX_NONE);
diff --git a/Engine/Plugins/Enterprise/GLTFExporter/MigratedFromGit/Source/GLTFExporter/Private/Converters/GLTFLightMapConverters.cpp b/Engine/Plugins/Enterprise/GLTFExporter/MigratedFromGit/Source/GLTFExporter/Private/Converters/GLTFLightMapConverters.cpp
--- a/Engine/Plugins/Enterprise/GLTFExporter/MigratedFromGit/Source/GLTFExporter/Private/Converters/GLTFLightMapConverters.cpp
+++ b/Engine/Plugins/Enterprise/GLTFExporter/MigratedFromGit/Source/GLTFExporter/Private/Converters/GLTFLightMapConverters.cpp
@@ -17,7 +17,7 @@ FGLTFJsonLightMapIndex FGLTFLightMapConverter::Convert(const UStaticMeshComponen
 	const int32 LODIndex = FGLTFMeshUtility::GetLOD(StaticMesh, StaticMeshComponent, Builder.ExportOptions->DefaultLevelOfDetail);
 	const FStaticMeshLODResources& LODResources = StaticMesh->GetLODForExport(LODIndex);
 
-	const int32 CoordinateIndex = StaticMesh->LightMapCoordinateIndex;
+	const int32 CoordinateIndex = StaticMesh->GetLightMapCoordinateIndex();
 	if (CoordinateIndex < 0 || CoordinateIndex >= LODResources.GetNumTexCoords())
 	{
 		return FGLTFJsonLightMapIndex(INDEX_NONE);
diff --git a/Engine/Plugins/Enterprise/GLTFExporter/MigratedFromGit/Source/GLTFExporter/Private/Converters/GLTFMeshConverters.cpp b/Engine/Plugins/Enterprise/GLTFExporter/MigratedFromGit/Source/GLTFExporter/Private/Converters/GLTFMeshConverters.cpp
--- a/Engine/Plugins/Enterprise/GLTFExporter/MigratedFromGit/Source/GLTFExporter/Private/Converters/GLTFMeshConverters.cpp
+++ b/Engine/Plugins/Enterprise/GLTFExporter/MigratedFromGit/Source/GLTFExporter/Private/Converters/GLTFMeshConverters.cpp
@@ -14,7 +14,7 @@ void FGLTFStaticMeshConverter::Sanitize(const UStaticMesh*& StaticMesh, const US
 	}
 	else
 	{
-		FGLTFMaterialUtility::ResolveOverrides(Materials, StaticMesh->StaticMaterials);
+		FGLTFMaterialUtility::ResolveOverrides(Materials, StaticMesh->GetStaticMaterials());
 	}
 
 	if (LODIndex < 0)
@@ -53,7 +53,7 @@ void FGLTFSkeletalMeshConverter::Sanitize(const USkeletalMesh*& SkeletalMesh, co
 	}
 	else
 	{
-		FGLTFMaterialUtility::ResolveOverrides(Materials, SkeletalMesh->Materials);
+		FGLTFMaterialUtility::ResolveOverrides(Materials, SkeletalMesh->GetMaterials());
 	}
 
 	if (LODIndex < 0)
diff --git a/Engine/Plugins/Enterprise/GLTFExporter/MigratedFromGit/Source/GLTFExporter/Private/Converters/GLTFMeshData.cpp b/Engine/Plugins/Enterprise/GLTFExporter/MigratedFromGit/Source/GLTFExporter/Private/Converters/GLTFMeshData.cpp
--- a/Engine/Plugins/Enterprise/GLTFExporter/MigratedFromGit/Source/GLTFExporter/Private/Converters/GLTFMeshData.cpp
+++ b/Engine/Plugins/Enterprise/GLTFExporter/MigratedFromGit/Source/GLTFExporter/Private/Converters/GLTFMeshData.cpp
@@ -24,7 +24,7 @@ FGLTFMeshData::FGLTFMeshData(const UStaticMesh* StaticMesh, const UStaticMeshCom
 	}
 
 	const int32 NumTexCoords = StaticMesh->GetLODForExport(LODIndex).VertexBuffers.StaticMeshVertexBuffer.GetNumTexCoords();
-	TexCoord = FMath::Min(StaticMesh->LightMapCoordinateIndex, NumTexCoords - 1);
+	TexCoord = FMath::Min(StaticMesh->GetLightMapCoordinateIndex(), NumTexCoords - 1);
 
 	// TODO: add warning if texture coordinate has overlap
 }
diff --git a/Engine/Plugins/Enterprise/GLTFExporter/MigratedFromGit/Source/GLTFExporter/Private/Converters/GLTFMeshUtility.cpp b/Engine/Plugins/Enterprise/GLTFExporter/MigratedFromGit/Source/GLTFExporter/Private/Converters/GLTFMeshUtility.cpp
--- a/Engine/Plugins/Enterprise/GLTFExporter/MigratedFromGit/Source/GLTFExporter/Private/Converters/GLTFMeshUtility.cpp
+++ b/Engine/Plugins/Enterprise/GLTFExporter/MigratedFromGit/Source/GLTFExporter/Private/Converters/GLTFMeshUtility.cpp
@@ -82,7 +82,7 @@ int32 FGLTFMeshUtility::GetMinimumLOD(const UStaticMesh* StaticMesh, const UStat
 
 	if (StaticMesh != nullptr)
 	{
-		return GetValueForRunningPlatform<int32>(StaticMesh->MinLOD);
+		return GetValueForRunningPlatform<int32>(StaticMesh->GetMinLOD());
 	}
 
 	return -1;
@@ -97,7 +97,7 @@ int32 FGLTFMeshUtility::GetMinimumLOD(const USkeletalMesh* SkeletalMesh, const U
 
 	if (SkeletalMesh != nullptr)
 	{
-		return GetValueForRunningPlatform<int32>(SkeletalMesh->MinLod);
+		return GetValueForRunningPlatform<int32>(SkeletalMesh->GetMinLod());
 	}
 
 	return -1;
diff --git a/Engine/Plugins/Enterprise/GLTFExporter/MigratedFromGit/Source/GLTFExporter/Private/Converters/GLTFNodeConverters.cpp b/Engine/Plugins/Enterprise/GLTFExporter/MigratedFromGit/Source/GLTFExporter/Private/Converters/GLTFNodeConverters.cpp
--- a/Engine/Plugins/Enterprise/GLTFExporter/MigratedFromGit/Source/GLTFExporter/Private/Converters/GLTFNodeConverters.cpp
+++ b/Engine/Plugins/Enterprise/GLTFExporter/MigratedFromGit/Source/GLTFExporter/Private/Converters/GLTFNodeConverters.cpp
@@ -237,12 +237,12 @@ FGLTFJsonNodeIndex FGLTFSkeletalSocketConverter::Convert(FGLTFJsonNodeIndex Root
 		Node.Rotation = FGLTFConverterUtility::ConvertRotation(Socket->RelativeRotation.Quaternion());
 		Node.Scale = FGLTFConverterUtility::ConvertScale(Socket->RelativeScale);
 
-		const int32 ParentBone = SkeletalMesh->RefSkeleton.FindBoneIndex(Socket->BoneName);
+		const int32 ParentBone = SkeletalMesh->GetRefSkeleton().FindBoneIndex(Socket->BoneName);
 		const FGLTFJsonNodeIndex ParentNode = ParentBone != INDEX_NONE ? Builder.GetOrAddNode(RootNode, SkeletalMesh, ParentBone) : RootNode;
 		return Builder.AddChildNode(ParentNode, Node);
 	}
 
-	const int32 BoneIndex = SkeletalMesh->RefSkeleton.FindBoneIndex(SocketName);
+	const int32 BoneIndex = SkeletalMesh->GetRefSkeleton().FindBoneIndex(SocketName);
 	if (BoneIndex != INDEX_NONE)
 	{
 		return Builder.GetOrAddNode(RootNode, SkeletalMesh, BoneIndex);
@@ -256,7 +256,7 @@ FGLTFJsonNodeIndex FGLTFSkeletalBoneConverter::Convert(FGLTFJsonNodeIndex RootNo
 {
 	// TODO: add support for MasterPoseComponent?
 
-	const TArray<FMeshBoneInfo>& BoneInfos = SkeletalMesh->RefSkeleton.GetRefBoneInfo();
+	const TArray<FMeshBoneInfo>& BoneInfos = SkeletalMesh->GetRefSkeleton().GetRefBoneInfo();
 	if (!BoneInfos.IsValidIndex(BoneIndex))
 	{
 		// TODO: report error
@@ -268,7 +268,7 @@ FGLTFJsonNodeIndex FGLTFSkeletalBoneConverter::Convert(FGLTFJsonNodeIndex RootNo
 	FGLTFJsonNode Node;
 	Node.Name = BoneInfo.Name.ToString();
 
-	const TArray<FTransform>& BonePoses = SkeletalMesh->RefSkeleton.GetRefBonePose();
+	const TArray<FTransform>& BonePoses = SkeletalMesh->GetRefSkeleton().GetRefBonePose();
 	if (BonePoses.IsValidIndex(BoneIndex))
 	{
 		// TODO: add warning check for non-uniform scaling
diff --git a/Engine/Plugins/Enterprise/GLTFExporter/MigratedFromGit/Source/GLTFExporter/Private/Converters/GLTFSkinConverters.cpp b/Engine/Plugins/Enterprise/GLTFExporter/MigratedFromGit/Source/GLTFExporter/Private/Converters/GLTFSkinConverters.cpp
--- a/Engine/Plugins/Enterprise/GLTFExporter/MigratedFromGit/Source/GLTFExporter/Private/Converters/GLTFSkinConverters.cpp
+++ b/Engine/Plugins/Enterprise/GLTFExporter/MigratedFromGit/Source/GLTFExporter/Private/Converters/GLTFSkinConverters.cpp
@@ -8,10 +8,10 @@
 FGLTFJsonSkinIndex FGLTFSkinConverter::Convert(FGLTFJsonNodeIndex RootNode, const USkeletalMesh* SkeletalMesh)
 {
 	FGLTFJsonSkin Skin;
-	Skin.Name = SkeletalMesh->Skeleton != nullptr ? SkeletalMesh->Skeleton->GetName() : SkeletalMesh->GetName();
+	Skin.Name = SkeletalMesh->GetSkeleton() != nullptr ? SkeletalMesh->GetSkeleton()->GetName() : SkeletalMesh->GetName();
 	Skin.Skeleton = RootNode;
 
-	const int32 BoneCount = SkeletalMesh->RefSkeleton.GetNum();
+	const int32 BoneCount = SkeletalMesh->GetRefSkeleton().GetNum();
 	if (BoneCount == 0)
 	{
 		// TODO: report warning
@@ -30,7 +30,7 @@ FGLTFJsonSkinIndex FGLTFSkinConverter::Convert(FGLTFJsonNodeIndex RootNode, cons
 
 	for (int32 BoneIndex = 0; BoneIndex < BoneCount; ++BoneIndex)
 	{
-		const FTransform InverseBindTransform = FGLTFBoneUtility::GetBindTransform(SkeletalMesh->RefSkeleton, BoneIndex).Inverse();
+		const FTransform InverseBindTransform = FGLTFBoneUtility::GetBindTransform(SkeletalMesh->GetRefSkeleton(), BoneIndex).Inverse();
 		InverseBindMatrices[BoneIndex] = FGLTFConverterUtility::ConvertTransform(InverseBindTransform, Builder.ExportOptions->ExportUniformScale);
 	}
 
diff --git a/Engine/Plugins/Enterprise/GLTFExporter/MigratedFromGit/Source/GLTFExporter/Private/Tasks/GLTFMeshTasks.cpp b/Engine/Plugins/Enterprise/GLTFExporter/MigratedFromGit/Source/GLTFExporter/Private/Tasks/GLTFMeshTasks.cpp
--- a/Engine/Plugins/Enterprise/GLTFExporter/MigratedFromGit/Source/GLTFExporter/Private/Tasks/GLTFMeshTasks.cpp
+++ b/Engine/Plugins/Enterprise/GLTFExporter/MigratedFromGit/Source/GLTFExporter/Private/Tasks/GLTFMeshTasks.cpp
@@ -131,7 +131,7 @@ void FGLTFStaticMeshTask::Complete()
 
 	ValidateVertexBuffer(Builder, VertexBuffer, *StaticMesh->GetName());
 
-	const int32 MaterialCount = StaticMesh->StaticMaterials.Num();
+	const int32 MaterialCount = StaticMesh->GetStaticMaterials().Num();
 	JsonMesh.Primitives.AddDefaulted(MaterialCount);
 
 	for (int32 MaterialIndex = 0; MaterialIndex < MaterialCount; ++MaterialIndex)
@@ -223,7 +223,7 @@ void FGLTFSkeletalMeshTask::Complete()
 
 	ValidateVertexBuffer(Builder, VertexBuffer, *SkeletalMesh->GetName());
 
-	const uint16 MaterialCount = SkeletalMesh->Materials.Num();
+	const uint16 MaterialCount = SkeletalMesh->GetMaterials().Num();
 	JsonMesh.Primitives.AddDefaulted(MaterialCount);
 
 	for (uint16 MaterialIndex = 0; MaterialIndex < MaterialCount; ++MaterialIndex)
diff --git a/Engine/Plugins/Enterprise/GLTFExporter/MigratedFromGit/Source/GLTFExporterRuntime/Private/Actors/GLTFHotspotActor.cpp b/Engine/Plugins/Enterprise/GLTFExporter/MigratedFromGit/Source/GLTFExporterRuntime/Private/Actors/GLTFHotspotActor.cpp
--- a/Engine/Plugins/Enterprise/GLTFExporter/MigratedFromGit/Source/GLTFExporterRuntime/Private/Actors/GLTFHotspotActor.cpp
+++ b/Engine/Plugins/Enterprise/GLTFExporter/MigratedFromGit/Source/GLTFExporterRuntime/Private/Actors/GLTFHotspotActor.cpp
@@ -456,11 +456,11 @@ void AGLTFHotspotActor::ValidateAnimation()
 	{
 		if (USkeletalMesh* SkeletalMesh = SkeletalMeshActor->GetSkeletalMeshComponent()->SkeletalMesh)
 		{
-			if (SkeletalMesh->Skeleton != AnimationSequence->GetSkeleton())
+			if (SkeletalMesh->GetSkeleton() != AnimationSequence->GetSkeleton())
 			{
-				if (SkeletalMesh->Skeleton != nullptr)
+				if (SkeletalMesh->GetSkeleton() != nullptr)
 				{
-					UE_LOG(LogGLTFHotspot, Warning, TEXT("Animation %s is incompatible with skeleton %s, removing animation from actor."), *AnimationSequence->GetName(), *SkeletalMesh->Skeleton->GetName());
+					UE_LOG(LogGLTFHotspot, Warning, TEXT("Animation %s is incompatible with skeleton %s, removing animation from actor."), *AnimationSequence->GetName(), *SkeletalMesh->GetSkeleton()->GetName());
 				}
 				else
 				{
