[GLTFExporter] chore: bump version to 1.0.6

commit c25891cce03a962da6938eb14fa06ca913696cba
Date: Thu, 19 Aug 2021 23:52:22 +0200

    chore: bump version to 1.0.6
    (cherry picked from commit eb8a7cec83e7d42dbb0c012ba0fdd366ff0654e2​)

#skipci

diff --git a/Engine/Plugins/Enterprise/GLTFExporter/MigratedFromGit/GLTFExporter.uplugin b/Engine/Plugins/Enterprise/GLTFExporter/MigratedFromGit/GLTFExporter.uplugin
--- a/Engine/Plugins/Enterprise/GLTFExporter/MigratedFromGit/GLTFExporter.uplugin
+++ b/Engine/Plugins/Enterprise/GLTFExporter/MigratedFromGit/GLTFExporter.uplugin
@@ -1,7 +1,7 @@
 {
 	"FileVersion": 3,
-	"Version": 105,
-	"VersionName": "1.0.5",
+	"Version": 106,
+	"VersionName": "1.0.6",
 	"FriendlyName": "glTF Exporter",
 	"Description": "An exporter for Khronos glTF 2.0.",
 	"Category": "Exporters",
