[GLTFExporter] refactor: remove some unnecessary hacks/workarounds due to merged improvements in UE 4.27

commit 90fa162fc945da660f73e7d611eba47c4668edbe
Date: Wed, 8 Sep 2021 01:12:00 +0200

    refactor: remove some unnecessary hacks/workarounds due to merged improvements in UE 4.27
    (cherry picked from commit b2edcb8f3e86c936864b29d283ea7d3adbb31a72​)

#skipci

diff --git a/Engine/Plugins/Enterprise/GLTFExporter/MigratedFromGit/Source/GLTFExporter/Private/Converters/GLTFAccessorConverters.cpp b/Engine/Plugins/Enterprise/GLTFExporter/MigratedFromGit/Source/GLTFExporter/Private/Converters/GLTFAccessorConverters.cpp
--- a/Engine/Plugins/Enterprise/GLTFExporter/MigratedFromGit/Source/GLTFExporter/Private/Converters/GLTFAccessorConverters.cpp
+++ b/Engine/Plugins/Enterprise/GLTFExporter/MigratedFromGit/Source/GLTFExporter/Private/Converters/GLTFAccessorConverters.cpp
@@ -2,7 +2,6 @@
 
 #include "Converters/GLTFAccessorConverters.h"
 #include "Converters/GLTFConverterUtility.h"
-#include "Converters/GLTFSkinWeightVertexBufferHack.h"
 #include "Builders/GLTFConvertBuilder.h"
 
 // TODO: Unreal-style implementation of std::conditional to avoid mixing in STL. Should be added to the engine.
@@ -333,8 +332,7 @@ FGLTFJsonAccessorIndex FGLTFBoneIndexBufferConverter::Convert(const FGLTFMeshSec
 
 		for (int32 InfluenceIndex = 0; InfluenceIndex < 4; ++InfluenceIndex)
 		{
-			// TODO: remove hack
-			const uint32 UnmappedBoneIndex = FGLTFSkinWeightVertexBufferHack(VertexBuffer).GetBoneIndex(MappedVertexIndex, InfluenceOffset + InfluenceIndex);
+			const uint32 UnmappedBoneIndex = VertexBuffer->GetBoneIndex(MappedVertexIndex, InfluenceOffset + InfluenceIndex);
 			VertexBones.Index[InfluenceIndex] = static_cast<IndexType>(BoneMap[UnmappedBoneIndex]);
 		}
 	}
@@ -380,8 +378,7 @@ FGLTFJsonAccessorIndex FGLTFBoneWeightBufferConverter::Convert(const FGLTFMeshSe
 
 		for (uint32 InfluenceIndex = 0; InfluenceIndex < 4; ++InfluenceIndex)
 		{
-			// TODO: remove hack
-			const uint8 BoneWeight = FGLTFSkinWeightVertexBufferHack(VertexBuffer).GetBoneWeight(MappedVertexIndex, InfluenceOffset + InfluenceIndex);
+			const uint8 BoneWeight = VertexBuffer->GetBoneWeight(MappedVertexIndex, InfluenceOffset + InfluenceIndex);
 			VertexWeights.Weights[InfluenceIndex] = BoneWeight;
 		}
 	}
diff --git a/Engine/Plugins/Enterprise/GLTFExporter/MigratedFromGit/Source/GLTFExporter/Private/Converters/GLTFLevelVariantSetsConverters.cpp b/Engine/Plugins/Enterprise/GLTFExporter/MigratedFromGit/Source/GLTFExporter/Private/Converters/GLTFLevelVariantSetsConverters.cpp
--- a/Engine/Plugins/Enterprise/GLTFExporter/MigratedFromGit/Source/GLTFExporter/Private/Converters/GLTFLevelVariantSetsConverters.cpp
+++ b/Engine/Plugins/Enterprise/GLTFExporter/MigratedFromGit/Source/GLTFExporter/Private/Converters/GLTFLevelVariantSetsConverters.cpp
@@ -11,19 +11,6 @@
 #include "VariantSet.h"
 #include "Variant.h"
 
-namespace
-{
-	// TODO: replace hack with safe solution by adding proper API access to UPropertyValue
-	class UPropertyValueHack : public UPropertyValue
-	{
-	public:
-		const TArray<FCapturedPropSegment>& GetCapturedPropSegments() const
-		{
-			return CapturedPropSegments;
-		}
-	};
-}
-
 FGLTFJsonLevelVariantSetsIndex FGLTFLevelVariantSetsConverter::Convert(const ULevelVariantSets* LevelVariantSets)
 {
 	FGLTFJsonLevelVariantSets JsonLevelVariantSets;
@@ -216,7 +203,7 @@ bool FGLTFLevelVariantSetsConverter::TryParseMaterialPropertyValue(FGLTFJsonVari
 		return false;
 	}
 
-	const TArray<FCapturedPropSegment>& CapturedPropSegments = reinterpret_cast<UPropertyValueHack*>(MaterialProperty)->GetCapturedPropSegments();
+	const TArray<FCapturedPropSegment>& CapturedPropSegments = MaterialProperty->GetCapturedPropSegments();
 	const int32 NumPropSegments = CapturedPropSegments.Num();
 
 	if (NumPropSegments < 1)
diff --git a/Engine/Plugins/Enterprise/GLTFExporter/MigratedFromGit/Source/GLTFExporter/Private/Converters/GLTFMaterialUtility.cpp b/Engine/Plugins/Enterprise/GLTFExporter/MigratedFromGit/Source/GLTFExporter/Private/Converters/GLTFMaterialUtility.cpp
--- a/Engine/Plugins/Enterprise/GLTFExporter/MigratedFromGit/Source/GLTFExporter/Private/Converters/GLTFMaterialUtility.cpp
+++ b/Engine/Plugins/Enterprise/GLTFExporter/MigratedFromGit/Source/GLTFExporter/Private/Converters/GLTFMaterialUtility.cpp
@@ -28,90 +28,39 @@ bool FGLTFMaterialUtility::IsNormalMap(const FMaterialPropertyEx& Property)
 	return Property == MP_Normal || Property == TEXT("ClearCoatBottomNormal");
 }
 
-FVector4 FGLTFMaterialUtility::GetPropertyDefaultValue(const FMaterialPropertyEx& Property)
+FGuid FGLTFMaterialUtility::GetAttributeID(const FMaterialPropertyEx& Property)
 {
-	// TODO: replace with GMaterialPropertyAttributesMap lookup (when public API available)
-
-	switch (Property.Type)
-	{
-		case MP_EmissiveColor:          return FVector4(0,0,0,0);
-		case MP_Opacity:                return FVector4(1,0,0,0);
-		case MP_OpacityMask:            return FVector4(1,0,0,0);
-		case MP_BaseColor:              return FVector4(0,0,0,0);
-		case MP_Metallic:               return FVector4(0,0,0,0);
-		case MP_Specular:               return FVector4(.5,0,0,0);
-		case MP_Roughness:              return FVector4(.5,0,0,0);
-		case MP_Anisotropy:             return FVector4(0,0,0,0);
-		case MP_Normal:                 return FVector4(0,0,1,0);
-		case MP_Tangent:                return FVector4(1,0,0,0);
-		case MP_WorldPositionOffset:    return FVector4(0,0,0,0);
-		case MP_WorldDisplacement:      return FVector4(0,0,0,0);
-		case MP_TessellationMultiplier: return FVector4(1,0,0,0);
-		case MP_SubsurfaceColor:        return FVector4(1,1,1,0);
-		case MP_CustomData0:            return FVector4(1,0,0,0);
-		case MP_CustomData1:            return FVector4(.1,0,0,0);
-		case MP_AmbientOcclusion:       return FVector4(1,0,0,0);
-		case MP_Refraction:             return FVector4(1,0,0,0);
-		case MP_PixelDepthOffset:       return FVector4(0,0,0,0);
-		case MP_ShadingModel:           return FVector4(0,0,0,0);
-		default:                        break;
-	}
-
-	if (Property >= MP_CustomizedUVs0 && Property <= MP_CustomizedUVs7)
-	{
-		return FVector4(0,0,0,0);
-	}
+	return Property.IsCustomOutput()
+		? FMaterialAttributeDefinitionMap::GetCustomAttributeID(Property.CustomOutput.ToString())
+		: FMaterialAttributeDefinitionMap::GetID(Property.Type);
+}
 
-	if (Property == TEXT("ClearCoatBottomNormal"))
-	{
-		return FVector4(0,0,1,0);
-	}
+FGuid FGLTFMaterialUtility::GetAttributeIDChecked(const FMaterialPropertyEx& Property)
+{
+	const FGuid AttributeID = GetAttributeID(Property);
+	check(AttributeID != FMaterialAttributeDefinitionMap::GetID(MP_MAX));
+	return AttributeID;
+}
 
-	check(false);
-	return FVector4();
+FVector4 FGLTFMaterialUtility::GetPropertyDefaultValue(const FMaterialPropertyEx& Property)
+{
+	return FMaterialAttributeDefinitionMap::GetDefaultValue(GetAttributeIDChecked(Property));
 }
 
 FVector4 FGLTFMaterialUtility::GetPropertyMask(const FMaterialPropertyEx& Property)
 {
-	// TODO: replace with GMaterialPropertyAttributesMap lookup (when public API available)
-
-	switch (Property.Type)
+	const EMaterialValueType ValueType = FMaterialAttributeDefinitionMap::GetValueType(GetAttributeIDChecked(Property));
+	switch (ValueType)
 	{
-		case MP_EmissiveColor:          return FVector4(1,1,1,0);
-		case MP_Opacity:                return FVector4(1,0,0,0);
-		case MP_OpacityMask:            return FVector4(1,0,0,0);
-		case MP_BaseColor:              return FVector4(1,1,1,0);
-		case MP_Metallic:               return FVector4(1,0,0,0);
-		case MP_Specular:               return FVector4(1,0,0,0);
-		case MP_Roughness:              return FVector4(1,0,0,0);
-		case MP_Anisotropy:             return FVector4(1,0,0,0);
-		case MP_Normal:                 return FVector4(1,1,1,0);
-		case MP_Tangent:                return FVector4(1,1,1,0);
-		case MP_WorldPositionOffset:    return FVector4(1,1,1,0);
-		case MP_WorldDisplacement:      return FVector4(1,1,1,0);
-		case MP_TessellationMultiplier: return FVector4(1,0,0,0);
-		case MP_SubsurfaceColor:        return FVector4(1,1,1,0);
-		case MP_CustomData0:            return FVector4(1,0,0,0);
-		case MP_CustomData1:            return FVector4(1,0,0,0);
-		case MP_AmbientOcclusion:       return FVector4(1,0,0,0);
-		case MP_Refraction:             return FVector4(1,1,0,0);
-		case MP_PixelDepthOffset:       return FVector4(1,0,0,0);
-		case MP_ShadingModel:           return FVector4(1,0,0,0);
-		default:                        break;
+		case MCT_Float:
+		case MCT_Float1: return FVector4(1, 0, 0, 0);
+		case MCT_Float2: return FVector4(1, 1, 0, 0);
+		case MCT_Float3: return FVector4(1, 1, 1, 0);
+		case MCT_Float4: return FVector4(1, 1, 1, 1);
+		default:
+			checkNoEntry();
+			return FVector4();
 	}
-
-	if (Property >= MP_CustomizedUVs0 && Property <= MP_CustomizedUVs7)
-	{
-		return FVector4(1,1,0,0);
-	}
-
-	if (Property == TEXT("ClearCoatBottomNormal"))
-	{
-		return FVector4(1,1,1,0);
-	}
-
-	check(false);
-	return FVector4();
 }
 
 const FExpressionInput* FGLTFMaterialUtility::GetInputForProperty(const UMaterialInterface* Material, const FMaterialPropertyEx& Property)
diff --git a/Engine/Plugins/Enterprise/GLTFExporter/MigratedFromGit/Source/GLTFExporter/Private/Converters/GLTFMaterialUtility.h b/Engine/Plugins/Enterprise/GLTFExporter/MigratedFromGit/Source/GLTFExporter/Private/Converters/GLTFMaterialUtility.h
--- a/Engine/Plugins/Enterprise/GLTFExporter/MigratedFromGit/Source/GLTFExporter/Private/Converters/GLTFMaterialUtility.h
+++ b/Engine/Plugins/Enterprise/GLTFExporter/MigratedFromGit/Source/GLTFExporter/Private/Converters/GLTFMaterialUtility.h
@@ -41,6 +41,9 @@ struct FGLTFMaterialUtility
 
 	static bool IsNormalMap(const FMaterialPropertyEx& Property);
 
+	static FGuid GetAttributeID(const FMaterialPropertyEx& Property);
+	static FGuid GetAttributeIDChecked(const FMaterialPropertyEx& Property);
+
 	static FVector4 GetPropertyDefaultValue(const FMaterialPropertyEx& Property);
 	static FVector4 GetPropertyMask(const FMaterialPropertyEx& Property);
 
diff --git a/Engine/Plugins/Enterprise/GLTFExporter/MigratedFromGit/Source/GLTFExporter/Private/Converters/GLTFSkinWeightVertexBufferHack.h b/Engine/Plugins/Enterprise/GLTFExporter/MigratedFromGit/Source/GLTFExporter/Private/Converters/GLTFSkinWeightVertexBufferHack.h
deleted file mode 100644
