[GLTFExporter] fix: perform material analyze on expression-inputs that lack valid expressions

commit 9ea091ed51034a32a0338dc86ec949cc425e4603
Date: Tue, 14 Sep 2021 15:56:23 +0200

    fix: perform material analyze on expression-inputs that lack valid expressions
    when using material attributes, property-based expression-inputs on a material
    may have their expressions set to null in some cases. Compiling and analyzing
    these inputs still work since compilation is performed via the attribute
    node.
    (cherry picked from commit 1da52d3a7083c80144d32954e7044206004cf439​)

#skipci

diff --git a/Engine/Plugins/Enterprise/GLTFExporter/MigratedFromGit/Source/GLTFExporter/Private/Converters/GLTFMaterialUtility.cpp b/Engine/Plugins/Enterprise/GLTFExporter/MigratedFromGit/Source/GLTFExporter/Private/Converters/GLTFMaterialUtility.cpp
--- a/Engine/Plugins/Enterprise/GLTFExporter/MigratedFromGit/Source/GLTFExporter/Private/Converters/GLTFMaterialUtility.cpp
+++ b/Engine/Plugins/Enterprise/GLTFExporter/MigratedFromGit/Source/GLTFExporter/Private/Converters/GLTFMaterialUtility.cpp
@@ -420,8 +420,7 @@ bool FGLTFMaterialUtility::NeedsMeshData(const TArray<const UMaterialInterface*>
 
 void FGLTFMaterialUtility::AnalyzeMaterialProperty(const UMaterialInterface* InMaterial, const FMaterialPropertyEx& InProperty, FGLTFMaterialAnalysis& OutAnalysis)
 {
-	const FExpressionInput* Input = GetInputForProperty(InMaterial, InProperty);
-	if (Input == nullptr || Input->Expression == nullptr)
+	if (GetInputForProperty(InMaterial, InProperty) == nullptr)
 	{
 		OutAnalysis = FGLTFMaterialAnalysis();
 		return;
