[GLTFExporter] feat: expose GLTFExportOptions class as a type that can be used in blueprints

commit 124cc48b16ab8fe562dd0918de66e71a062d2aae
Date: Wed, 15 Sep 2021 21:49:30 +0200

    feat: expose GLTFExportOptions class as a type that can be used in blueprints
    (cherry picked from commit d359c5f26c9667e49d66b801d2fe405f85e4b602​)

#skipci

diff --git a/Engine/Plugins/Enterprise/GLTFExporter/MigratedFromGit/Source/GLTFExporter/Public/GLTFExportOptions.h b/Engine/Plugins/Enterprise/GLTFExporter/MigratedFromGit/Source/GLTFExporter/Public/GLTFExportOptions.h
--- a/Engine/Plugins/Enterprise/GLTFExporter/MigratedFromGit/Source/GLTFExporter/Public/GLTFExportOptions.h
+++ b/Engine/Plugins/Enterprise/GLTFExporter/MigratedFromGit/Source/GLTFExporter/Public/GLTFExportOptions.h
@@ -63,7 +63,7 @@ enum class EGLTFMaterialVariantMode : uint8
 	UseMeshData,
 };
 
-UCLASS(Config=EditorPerProjectUserSettings, HideCategories=(DebugProperty))
+UCLASS(BlueprintType, Config=EditorPerProjectUserSettings, HideCategories=(DebugProperty))
 class GLTFEXPORTER_API UGLTFExportOptions : public UObject
 {
 	GENERATED_UCLASS_BODY()
