[GLTFExporter] fix: attempt to fix strange build issue on some machines regarding template compilation

commit 481d62b0539081d03808e31864d2f0eff7950d71
Date: Thu, 16 Sep 2021 15:03:25 +0200

    fix: attempt to fix strange build issue on some machines regarding template compilation
    (cherry picked from commit 86d13688967134538985c929d8d275cbd42aaebe​)

#skipci

diff --git a/Engine/Plugins/Enterprise/GLTFExporter/MigratedFromGit/Source/GLTFExporter/Private/Converters/GLTFMeshDataConverters.h b/Engine/Plugins/Enterprise/GLTFExporter/MigratedFromGit/Source/GLTFExporter/Private/Converters/GLTFMeshDataConverters.h
--- a/Engine/Plugins/Enterprise/GLTFExporter/MigratedFromGit/Source/GLTFExporter/Private/Converters/GLTFMeshDataConverters.h
+++ b/Engine/Plugins/Enterprise/GLTFExporter/MigratedFromGit/Source/GLTFExporter/Private/Converters/GLTFMeshDataConverters.h
@@ -11,13 +11,13 @@ class TGLTFMeshDataConverter final : public TGLTFConverter<const FGLTFMeshData*,
 {
 	TArray<TUniquePtr<FGLTFMeshData>> Outputs;
 
-	const FGLTFMeshData* Convert(const MeshType* Mesh, const MeshComponentType* MeshComponent, int32 LODIndex)
+	virtual const FGLTFMeshData* Convert(const MeshType* Mesh, const MeshComponentType* MeshComponent, int32 LODIndex) override
 	{
 		TUniquePtr<FGLTFMeshData> Output = MakeUnique<FGLTFMeshData>(Mesh, MeshComponent, LODIndex);
 
 		if (MeshComponent != nullptr)
 		{
-			Output->Parent = GetOrAdd(Mesh, nullptr, LODIndex);
+			Output->Parent = TGLTFConverter::GetOrAdd(Mesh, nullptr, LODIndex);
 		}
 
 		return Outputs.Add_GetRef(MoveTemp(Output)).Get();
diff --git a/Engine/Plugins/Enterprise/GLTFExporter/MigratedFromGit/Source/GLTFExporter/Private/Converters/GLTFMeshSectionConverters.h b/Engine/Plugins/Enterprise/GLTFExporter/MigratedFromGit/Source/GLTFExporter/Private/Converters/GLTFMeshSectionConverters.h
--- a/Engine/Plugins/Enterprise/GLTFExporter/MigratedFromGit/Source/GLTFExporter/Private/Converters/GLTFMeshSectionConverters.h
+++ b/Engine/Plugins/Enterprise/GLTFExporter/MigratedFromGit/Source/GLTFExporter/Private/Converters/GLTFMeshSectionConverters.h
@@ -11,7 +11,7 @@ class TGLTFMeshSectionConverter final : public TGLTFConverter<const FGLTFMeshSec
 {
 	TArray<TUniquePtr<FGLTFMeshSection>> Outputs;
 
-	const FGLTFMeshSection* Convert(const MeshLODType* MeshLOD, FGLTFIndexArray SectionIndices)
+	virtual const FGLTFMeshSection* Convert(const MeshLODType* MeshLOD, FGLTFIndexArray SectionIndices) override
 	{
 		return Outputs.Add_GetRef(MakeUnique<FGLTFMeshSection>(MeshLOD, SectionIndices)).Get();
 	}
diff --git a/Engine/Plugins/Enterprise/GLTFExporter/MigratedFromGit/Source/GLTFExporter/Private/Json/GLTFJsonMatrix.h b/Engine/Plugins/Enterprise/GLTFExporter/MigratedFromGit/Source/GLTFExporter/Private/Json/GLTFJsonMatrix.h
--- a/Engine/Plugins/Enterprise/GLTFExporter/MigratedFromGit/Source/GLTFExporter/Private/Json/GLTFJsonMatrix.h
+++ b/Engine/Plugins/Enterprise/GLTFExporter/MigratedFromGit/Source/GLTFExporter/Private/Json/GLTFJsonMatrix.h
@@ -11,7 +11,7 @@ struct TGLTFJsonMatrix final : BaseType, IGLTFJsonArray
 	static const TGLTFJsonMatrix Identity;
 
 	TGLTFJsonMatrix(const BaseType& Other)
-		: TGLTFMatrix(Other)
+		: BaseType(Other)
 	{
 	}
 
