[GLTFExporter] chore: remove unnecessary dependencies in GLTFExporter plugin module

commit 75afcb376c53b861659a99626d41a050b5eb9a73
Date: Thu, 16 Sep 2021 21:34:51 +0200

    chore: remove unnecessary dependencies in GLTFExporter plugin module
    (cherry picked from commit 48ba375c7cc4627ef0ca78fd0a70510793cc2beb​)

#skipci

diff --git a/Engine/Plugins/Enterprise/GLTFExporter/MigratedFromGit/Source/GLTFExporter/GLTFExporter.Build.cs b/Engine/Plugins/Enterprise/GLTFExporter/MigratedFromGit/Source/GLTFExporter/GLTFExporter.Build.cs
--- a/Engine/Plugins/Enterprise/GLTFExporter/MigratedFromGit/Source/GLTFExporter/GLTFExporter.Build.cs
+++ b/Engine/Plugins/Enterprise/GLTFExporter/MigratedFromGit/Source/GLTFExporter/GLTFExporter.Build.cs
@@ -15,10 +15,8 @@ namespace UnrealBuildTool.Rules
 					"UnrealEd",
 					"MeshDescription",
 					"StaticMeshDescription",
-					"MeshUtilities",
 					"MessageLog",
 					"Json",
-					"MaterialEditor",
 					"Slate",
 					"SlateCore",
 					"Mainframe",
