[GLTFExporter] fix: proper initialization of GLTFHotspotActor member bToggled

commit b4430c544f6e9115beb1abfc961e2e274742be13
Date: Tue, 21 Sep 2021 08:53:29 +0200

    fix: proper initialization of GLTFHotspotActor member bToggled
    (cherry picked from commit f46bf9dfe20fda7789ff2832498a5ef7ca0165e6​)

#skipci

diff --git a/Engine/Plugins/Enterprise/GLTFExporter/MigratedFromGit/Source/GLTFExporterRuntime/Private/Actors/GLTFHotspotActor.cpp b/Engine/Plugins/Enterprise/GLTFExporter/MigratedFromGit/Source/GLTFExporterRuntime/Private/Actors/GLTFHotspotActor.cpp
--- a/Engine/Plugins/Enterprise/GLTFExporter/MigratedFromGit/Source/GLTFExporterRuntime/Private/Actors/GLTFHotspotActor.cpp
+++ b/Engine/Plugins/Enterprise/GLTFExporter/MigratedFromGit/Source/GLTFExporterRuntime/Private/Actors/GLTFHotspotActor.cpp
@@ -49,7 +49,7 @@ AGLTFHotspotActor::AGLTFHotspotActor(const FObjectInitializer& ObjectInitializer
 	DefaultToggledHoveredImage(nullptr),
 	ActiveImage(nullptr),
 	ActiveImageSize(0.0f, 0.0f),
-	bToggled(bToggled),
+	bToggled(false),
 	bIsInteractable(true),
 	RealtimeSecondsWhenLastInSight(0.0f),
 	RealtimeSecondsWhenLastHidden(0.0f)
