[GLTFExporter] refactor: use this-> instead of TGLTFConverter:: in GLTFMeshDataConverter to reference base class member

commit 01532073dd74892551b97ab54bce18c2e0f65f57
Date: Tue, 21 Sep 2021 08:55:38 +0200

    refactor: use this-> instead of TGLTFConverter:: in GLTFMeshDataConverter to reference base class member
    (cherry picked from commit 04b83a8d85b18729f3b21b44ae4ffa309bccae6f​)

#skipci

diff --git a/Engine/Plugins/Enterprise/GLTFExporter/MigratedFromGit/Source/GLTFExporter/Private/Converters/GLTFMeshDataConverters.h b/Engine/Plugins/Enterprise/GLTFExporter/MigratedFromGit/Source/GLTFExporter/Private/Converters/GLTFMeshDataConverters.h
--- a/Engine/Plugins/Enterprise/GLTFExporter/MigratedFromGit/Source/GLTFExporter/Private/Converters/GLTFMeshDataConverters.h
+++ b/Engine/Plugins/Enterprise/GLTFExporter/MigratedFromGit/Source/GLTFExporter/Private/Converters/GLTFMeshDataConverters.h
@@ -17,7 +17,7 @@ class TGLTFMeshDataConverter final : public TGLTFConverter<const FGLTFMeshData*,
 
 		if (MeshComponent != nullptr)
 		{
-			Output->Parent = TGLTFConverter::GetOrAdd(Mesh, nullptr, LODIndex);
+			Output->Parent = this->GetOrAdd(Mesh, nullptr, LODIndex);
 		}
 
 		return Outputs.Add_GetRef(MoveTemp(Output)).Get();
