[GLTFExporter] refactor: avoit forward declaring enums in GLTFMaterialUserData

commit 65b193fb79bb5cf89e12d4ff902695b30300d9d6
Date: Tue, 21 Sep 2021 08:59:37 +0200

    refactor: avoit forward declaring enums in GLTFMaterialUserData
    (cherry picked from commit ab1284a0cad7b885f92b88900849626fabaf52ca​)

#skipci

diff --git a/Engine/Plugins/Enterprise/GLTFExporter/MigratedFromGit/Source/GLTFExporterRuntime/Public/UserData/GLTFMaterialUserData.h b/Engine/Plugins/Enterprise/GLTFExporter/MigratedFromGit/Source/GLTFExporterRuntime/Public/UserData/GLTFMaterialUserData.h
--- a/Engine/Plugins/Enterprise/GLTFExporter/MigratedFromGit/Source/GLTFExporterRuntime/Public/UserData/GLTFMaterialUserData.h
+++ b/Engine/Plugins/Enterprise/GLTFExporter/MigratedFromGit/Source/GLTFExporterRuntime/Public/UserData/GLTFMaterialUserData.h
@@ -2,12 +2,11 @@
 
 #pragma once
 
+#include "Engine/Texture.h"
 #include "Engine/AssetUserData.h"
 #include "GLTFMaterialUserData.generated.h"
 
 class UMaterialInterface;
-enum TextureFilter;
-enum TextureAddress;
 
 UENUM(BlueprintType)
 enum class EGLTFMaterialBakeMode : uint8
