[GLTFExporter] feat: always use LF as line terminator in gltf files to avoid cross platform issues

commit 2b51a7c39001db4dd4ae4a81c3d29a23bfd6efb9
Date: Tue, 21 Sep 2021 10:52:40 +0200

    feat: always use LF as line terminator in gltf files to avoid cross platform issues
    (cherry picked from commit 31add81ecfe0e50371c6860f9b158dc8866309af​)

#skipci

diff --git a/Engine/Plugins/Enterprise/GLTFExporter/MigratedFromGit/Source/GLTFExporter/Private/Json/GLTFJsonWriter.cpp b/Engine/Plugins/Enterprise/GLTFExporter/MigratedFromGit/Source/GLTFExporter/Private/Json/GLTFJsonWriter.cpp
--- a/Engine/Plugins/Enterprise/GLTFExporter/MigratedFromGit/Source/GLTFExporter/Private/Json/GLTFJsonWriter.cpp
+++ b/Engine/Plugins/Enterprise/GLTFExporter/MigratedFromGit/Source/GLTFExporter/Private/Json/GLTFJsonWriter.cpp
@@ -2,6 +2,7 @@
 
 #include "Json/GLTFJsonWriter.h"
 #include "Serialization/JsonSerializer.h"
+#include "Policies/PrettyJsonPrintPolicy.h"
 #include "Policies/CondensedJsonPrintPolicy.h"
 
 template <class CharType, class PrintPolicy>
@@ -148,10 +149,19 @@ private:
 	TSharedRef<TJsonWriter<CharType, PrintPolicy>> JsonWriter;
 };
 
+template <class CharType>
+struct TCrossPlatformPrettyJsonPrintPolicy : TPrettyJsonPrintPolicy<CharType>
+{
+	static void WriteLineTerminator(FArchive* Stream)
+	{
+		TJsonPrintPolicy<CharType>::WriteString(Stream, TEXT("\n")); // don't use system dependent line terminator
+	}
+};
+
 TSharedRef<IGLTFJsonWriter> IGLTFJsonWriter::Create(FArchive& Archive, bool bPrettyJson, FGLTFJsonExtensions& Extensions)
 {
 	return MakeShareable(bPrettyJson
-		? static_cast<IGLTFJsonWriter*>(new TGLTFJsonWriterImpl<UTF8CHAR, TPrettyJsonPrintPolicy<UTF8CHAR>>(Archive, Extensions))
+		? static_cast<IGLTFJsonWriter*>(new TGLTFJsonWriterImpl<UTF8CHAR, TCrossPlatformPrettyJsonPrintPolicy<UTF8CHAR>>(Archive, Extensions))
 		: static_cast<IGLTFJsonWriter*>(new TGLTFJsonWriterImpl<UTF8CHAR, TCondensedJsonPrintPolicy<UTF8CHAR>>(Archive, Extensions))
 	);
 }
