[GLTFExporter] fix: add template<> to avoid template specialization error in some compilers

commit 48c97a444f08127fce89fe2245c842a4cc94fec3
Date: Tue, 21 Sep 2021 11:55:40 +0200

    fix: add template<> to avoid template specialization error in some compilers
    (cherry picked from commit 44c6ed2451d039d402750874fcb0986b81b36134​)

#skipci

diff --git a/Engine/Plugins/Enterprise/GLTFExporter/MigratedFromGit/Source/GLTFExporter/Private/Json/GLTFJsonColor.cpp b/Engine/Plugins/Enterprise/GLTFExporter/MigratedFromGit/Source/GLTFExporter/Private/Json/GLTFJsonColor.cpp
--- a/Engine/Plugins/Enterprise/GLTFExporter/MigratedFromGit/Source/GLTFExporter/Private/Json/GLTFJsonColor.cpp
+++ b/Engine/Plugins/Enterprise/GLTFExporter/MigratedFromGit/Source/GLTFExporter/Private/Json/GLTFJsonColor.cpp
@@ -2,8 +2,8 @@
 
 #include "Json/GLTFJsonColor.h"
 
-const FGLTFJsonColor3 TGLTFJsonColor<FGLTFColor3>::Black({ 0, 0, 0 });
-const FGLTFJsonColor3 TGLTFJsonColor<FGLTFColor3>::White({ 1, 1, 1 });
+template <> const FGLTFJsonColor3 FGLTFJsonColor3::Black({ 0, 0, 0 });
+template <> const FGLTFJsonColor3 FGLTFJsonColor3::White({ 1, 1, 1 });
 
-const FGLTFJsonColor4 TGLTFJsonColor<FGLTFColor4>::Black({ 0, 0, 0, 1 });
-const FGLTFJsonColor4 TGLTFJsonColor<FGLTFColor4>::White({ 1, 1, 1, 1 });
+template <> const FGLTFJsonColor4 FGLTFJsonColor4::Black({ 0, 0, 0, 1 });
+template <> const FGLTFJsonColor4 FGLTFJsonColor4::White({ 1, 1, 1, 1 });
diff --git a/Engine/Plugins/Enterprise/GLTFExporter/MigratedFromGit/Source/GLTFExporter/Private/Json/GLTFJsonMatrix.cpp b/Engine/Plugins/Enterprise/GLTFExporter/MigratedFromGit/Source/GLTFExporter/Private/Json/GLTFJsonMatrix.cpp
--- a/Engine/Plugins/Enterprise/GLTFExporter/MigratedFromGit/Source/GLTFExporter/Private/Json/GLTFJsonMatrix.cpp
+++ b/Engine/Plugins/Enterprise/GLTFExporter/MigratedFromGit/Source/GLTFExporter/Private/Json/GLTFJsonMatrix.cpp
@@ -2,20 +2,20 @@
 
 #include "Json/GLTFJsonMatrix.h"
 
-const FGLTFJsonMatrix2 TGLTFJsonMatrix<FGLTFMatrix2>::Identity
+template <> const FGLTFJsonMatrix2 FGLTFJsonMatrix2::Identity
 ({
 	1, 0,
 	0, 1
 });
 
-const FGLTFJsonMatrix3 TGLTFJsonMatrix<FGLTFMatrix3>::Identity
+template <> const FGLTFJsonMatrix3 FGLTFJsonMatrix3::Identity
 ({
 	1, 0, 0,
 	0, 1, 0,
 	0, 0, 1
 });
 
-const FGLTFJsonMatrix4 TGLTFJsonMatrix<FGLTFMatrix4>::Identity
+template <> const FGLTFJsonMatrix4 FGLTFJsonMatrix4::Identity
 ({
 	1, 0, 0, 0,
 	0, 1, 0, 0,
diff --git a/Engine/Plugins/Enterprise/GLTFExporter/MigratedFromGit/Source/GLTFExporter/Private/Json/GLTFJsonVector.cpp b/Engine/Plugins/Enterprise/GLTFExporter/MigratedFromGit/Source/GLTFExporter/Private/Json/GLTFJsonVector.cpp
--- a/Engine/Plugins/Enterprise/GLTFExporter/MigratedFromGit/Source/GLTFExporter/Private/Json/GLTFJsonVector.cpp
+++ b/Engine/Plugins/Enterprise/GLTFExporter/MigratedFromGit/Source/GLTFExporter/Private/Json/GLTFJsonVector.cpp
@@ -2,11 +2,11 @@
 
 #include "Json/GLTFJsonVector.h"
 
-const FGLTFJsonVector2 TGLTFJsonVector<FGLTFVector2>::Zero({ 0, 0 });
-const FGLTFJsonVector2 TGLTFJsonVector<FGLTFVector2>::One ({ 1, 1 });
+template <> const FGLTFJsonVector2 FGLTFJsonVector2::Zero({ 0, 0 });
+template <> const FGLTFJsonVector2 FGLTFJsonVector2::One ({ 1, 1 });
 
-const FGLTFJsonVector3 TGLTFJsonVector<FGLTFVector3>::Zero({ 0, 0, 0 });
-const FGLTFJsonVector3 TGLTFJsonVector<FGLTFVector3>::One ({ 1, 1, 1 });
+template <> const FGLTFJsonVector3 FGLTFJsonVector3::Zero({ 0, 0, 0 });
+template <> const FGLTFJsonVector3 FGLTFJsonVector3::One ({ 1, 1, 1 });
 
-const FGLTFJsonVector4 TGLTFJsonVector<FGLTFVector4>::Zero({ 0, 0, 0, 0 });
-const FGLTFJsonVector4 TGLTFJsonVector<FGLTFVector4>::One ({ 1, 1, 1, 1 });
+template <> const FGLTFJsonVector4 FGLTFJsonVector4::Zero({ 0, 0, 0, 0 });
+template <> const FGLTFJsonVector4 FGLTFJsonVector4::One ({ 1, 1, 1, 1 });
