[GLTFExporter] feat: add general improvements from UE 4.27's baking module and normalize trivial differences

commit aaeb573db2018151378dfc6e9df233c7ed5c49d5
Date: Thu, 30 Sep 2021 04:58:40 +0200

    feat: add general improvements from UE 4.27's baking module and normalize trivial differences
    (cherry picked from commit 9bbf8c4f62349068de560dc39e219e415de94be3​)

#skipci

diff --git a/Engine/Plugins/Enterprise/GLTFExporter/MigratedFromGit/Source/GLTFMaterialBaking/Private/ExportMaterialProxy.h b/Engine/Plugins/Enterprise/GLTFExporter/MigratedFromGit/Source/GLTFMaterialBaking/Private/ExportMaterialProxy.h
--- a/Engine/Plugins/Enterprise/GLTFExporter/MigratedFromGit/Source/GLTFMaterialBaking/Private/ExportMaterialProxy.h
+++ b/Engine/Plugins/Enterprise/GLTFExporter/MigratedFromGit/Source/GLTFMaterialBaking/Private/ExportMaterialProxy.h
@@ -131,6 +131,18 @@ struct FExportMaterialCompiler : public FProxyMaterialCompiler
 		return Compiler->VertexColor();
 	}
 
+#if (ENGINE_MAJOR_VERSION > 4 || ENGINE_MINOR_VERSION >= 26)
+	virtual int32 PreSkinVertexOffset() override
+	{
+		return Compiler->PreSkinVertexOffset();
+	}
+
+	virtual int32 PostSkinVertexOffset() override
+	{
+		return Compiler->PostSkinVertexOffset();
+	}
+#endif
+
 	virtual int32 PreSkinnedPosition() override
 	{
 		return Compiler->PreSkinnedPosition();
@@ -221,22 +233,7 @@ struct FExportMaterialCompiler : public FProxyMaterialCompiler
 		}
 	}
 
-	virtual EMaterialCompilerType GetCompilerType() const override
-	{
-		return EMaterialCompilerType::MaterialProxy;
-	}
-
-#if (ENGINE_MAJOR_VERSION > 4 || ENGINE_MINOR_VERSION >= 26)
-	virtual int32 PreSkinVertexOffset() override
-	{
-		return Compiler->PreSkinVertexOffset();
-	}
-
-	virtual int32 PostSkinVertexOffset() override
-	{
-		return Compiler->PostSkinVertexOffset();
-	}
-#endif
+	virtual EMaterialCompilerType GetCompilerType() const override { return EMaterialCompilerType::MaterialProxy; }
 };
 
 class FExportMaterialProxy : public FMaterial, public FMaterialRenderProxy
@@ -252,7 +249,7 @@ public:
 	{
 
 #if (ENGINE_MAJOR_VERSION > 4 || ENGINE_MINOR_VERSION >= 26)
-		SetQualityLevelProperties(GMaxRHIFeatureLevel, EMaterialQualityLevel::High);
+		SetQualityLevelProperties(GMaxRHIFeatureLevel);
 #else
 		SetQualityLevelProperties(EMaterialQualityLevel::High, false, GMaxRHIFeatureLevel);
 #endif
@@ -459,7 +456,7 @@ public:
 				// Only return for Opaque and Masked...
 				if (BlendMode == BLEND_Opaque || BlendMode == BLEND_Masked)
 				{
-					return CompileNormalTransform(
+					return CompileNormalEncoding(
 						&ProxyCompiler,
 						MaterialInterface->CompileProperty(&ProxyCompiler, PropertyToCompile, ForceCast_Exact_Replicate));
 				}
@@ -502,10 +499,8 @@ public:
 	}
 	virtual EMaterialDomain GetMaterialDomain() const override
 	{
-		if (Material)
-		{
-			return Material->MaterialDomain;
-		}
+		// Because the baking module applies the material to a plane (or mesh),
+		// it needs to be a surface material.
 		return MD_Surface;
 	}
 	virtual bool IsTwoSided() const  override
@@ -534,7 +529,10 @@ public:
 	}
 	virtual bool IsDeferredDecal() const override
 	{
-		return Material && Material->MaterialDomain == MD_DeferredDecal;
+		// Decals are tricky. Since they mix with the underlying material
+		// and can't be applied to meshes, they can't really be baked 1:1.
+		// Instead we'll just bake them as surface materials.
+		return false;
 	}
 	virtual bool IsVolumetricPrimitive() const override
 	{
@@ -678,25 +676,23 @@ public:
 	}
 
 private:
-	static FGuid GetCustomAttributeID(const FString& AttributeName)
+	int32 CompileInputForCustomOutput(FMaterialCompiler* Compiler, int32 InputIndex, uint32 ForceCastFlags) const
 	{
+#if (ENGINE_MAJOR_VERSION > 4 || ENGINE_MINOR_VERSION >= 27)
+		FGuid AttributeID = FMaterialAttributeDefinitionMap::GetCustomAttributeID(CustomOutputToCompile);
+#else
+		FGuid AttributeID = {};
+
 		TArray<FMaterialCustomOutputAttributeDefintion> CustomAttributes;
 		FMaterialAttributeDefinitionMap::GetCustomAttributeList(CustomAttributes);
-
 		for (auto& Attribute : CustomAttributes)
 		{
-			if (Attribute.AttributeName == AttributeName)
+			if (Attribute.AttributeName == CustomOutputToCompile)
 			{
-				return Attribute.AttributeID;
+				AttributeID = Attribute.AttributeID;
 			}
 		}
-
-		return {};
-	}
-
-	int32 CompileInputForCustomOutput(FMaterialCompiler* Compiler, int32 InputIndex, uint32 ForceCastFlags) const
-	{
-		FGuid AttributeID = GetCustomAttributeID(CustomOutputToCompile);
+#endif
 		check(AttributeID.IsValid());
 
 		UMaterialExpressionCustomOutput* Expression = GetCustomOutputExpressionToCompile();
@@ -706,17 +702,17 @@ private:
 		if (ExpressionInput)
 		{
 			Result = ExpressionInput->Compile(Compiler);
-
-			if (CustomOutputToCompile ==  TEXT("ClearCoatBottomNormal"))
-			{
-				Result = CompileNormalTransform(Compiler, Result);
-			}
 		}
 		else
 		{
 			Result = FMaterialAttributeDefinitionMap::CompileDefaultExpression(Compiler, AttributeID);
 		}
 
+		if (CustomOutputToCompile == TEXT("ClearCoatBottomNormal"))
+		{
+			Result = CompileNormalEncoding(Compiler, Result);
+		}
+
 		if (ForceCastFlags & MFCF_ForceCast)
 		{
 			Result = Compiler->ForceCast(Result, FMaterialAttributeDefinitionMap::GetValueType(AttributeID), ForceCastFlags);
@@ -727,8 +723,7 @@ private:
 
 	UMaterialExpressionCustomOutput* GetCustomOutputExpressionToCompile() const
 	{
-		UMaterial* BaseMaterial = MaterialInterface->GetMaterial();
-		for (UMaterialExpression* Expression : BaseMaterial->Expressions)
+		for (UMaterialExpression* Expression : Material->Expressions)
 		{
 			UMaterialExpressionCustomOutput* CustomOutputExpression = Cast<UMaterialExpressionCustomOutput>(Expression);
 			if (CustomOutputExpression && CustomOutputExpression->GetDisplayName() == CustomOutputToCompile)
@@ -740,6 +735,19 @@ private:
 		return nullptr;
 	}
 
+	int32 CompileNormalEncoding(FMaterialCompiler* Compiler, int32 NormalInput) const
+	{
+		// TODO: make this configurable instead of assuming that we always want tangent-space normals
+		if (!Material->bTangentSpaceNormal)
+		{
+			NormalInput = Compiler->TransformVector(MCB_World, MCB_Tangent, NormalInput);
+		}
+
+		return Compiler->Add(
+			Compiler->Mul(NormalInput, Compiler->Constant(0.5f)), // [-1,1] * 0.5
+			Compiler->Constant(0.5f)); // [-0.5,0.5] + 0.5
+	}
+
 	static int32 CompileShadingModel(FMaterialCompiler* Compiler, int32 Code)
 	{
 		class FMaterialCompilerHack : public FHLSLMaterialTranslator
@@ -760,19 +768,6 @@ private:
 		return Compiler->Div(Code, Compiler->Constant(255)); // [0,255] / 255
 	}
 
-	int32 CompileNormalTransform(FMaterialCompiler* Compiler, int32 ExpressionIndex) const
-	{
-		// TODO: make this configurable instead of assuming that we always want tangent-space normals
-		if (!Material->bTangentSpaceNormal)
-		{
-			ExpressionIndex = Compiler->TransformVector(MCB_World, MCB_Tangent, ExpressionIndex);
-		}
-
-		return Compiler->Add(
-			Compiler->Mul(ExpressionIndex, Compiler->Constant(0.5f)), // [-1,1] * 0.5
-			Compiler->Constant(0.5f)); // [-0.5,0.5] + 0.5
-	}
-
 private:
 	/** The material interface for this proxy */
 	UMaterialInterface* MaterialInterface;
diff --git a/Engine/Plugins/Enterprise/GLTFExporter/MigratedFromGit/Source/GLTFMaterialBaking/Private/MaterialBakingModule.cpp b/Engine/Plugins/Enterprise/GLTFExporter/MigratedFromGit/Source/GLTFMaterialBaking/Private/MaterialBakingModule.cpp
--- a/Engine/Plugins/Enterprise/GLTFExporter/MigratedFromGit/Source/GLTFMaterialBaking/Private/MaterialBakingModule.cpp
+++ b/Engine/Plugins/Enterprise/GLTFExporter/MigratedFromGit/Source/GLTFMaterialBaking/Private/MaterialBakingModule.cpp
@@ -12,6 +12,8 @@
 #include "Async/Async.h"
 #include "Async/ParallelFor.h"
 #include "Materials/MaterialInstance.h"
+#include "Materials/MaterialInstanceConstant.h"
+#include "MaterialEditor/MaterialEditorInstanceConstant.h"
 #include "RenderingThread.h"
 #include "RHISurfaceDataConversion.h"
 #include "Misc/ScopedSlowTask.h"
@@ -234,12 +236,10 @@ namespace FMaterialBakingModuleImpl
 
 void FMaterialBakingModule::StartupModule()
 {
+	bEmissiveHDR = false;
+
 	// Set which properties should enforce gamma correction
-	PerPropertyGamma.Add(MP_Normal);
-	PerPropertyGamma.Add(MP_Opacity);
-	PerPropertyGamma.Add(MP_OpacityMask);
-	PerPropertyGamma.Add(MP_ShadingModel);
-	PerPropertyGamma.Add(TEXT("ClearCoatBottomNormal"));
+	SetLinearBake(false);
 
 	// Set which pixel format should be used for the possible baked out material properties
 	PerPropertyFormat.Add(MP_EmissiveColor, PF_FloatRGBA);
@@ -277,10 +277,9 @@ void FMaterialBakingModule::ShutdownModule()
 
 void FMaterialBakingModule::BakeMaterials(const TArray<FMaterialData*>& MaterialSettings, const TArray<FMeshData*>& MeshSettings, TArray<FBakeOutput>& Output)
 {
-	TArray<FMaterialDataEx> MaterialDataExs;
-	TArray<FMaterialDataEx*> MaterialSettingsEx;
-
 	// Translate old material data to extended types
+	TArray<FMaterialDataEx> MaterialDataExs;
+	MaterialDataExs.Reserve(MaterialSettings.Num());
 	for (const FMaterialData* MaterialData : MaterialSettings)
 	{
 		FMaterialDataEx& MaterialDataEx = MaterialDataExs.AddDefaulted_GetRef();
@@ -292,7 +291,13 @@ void FMaterialBakingModule::BakeMaterials(const TArray<FMaterialData*>& Material
 		{
 			MaterialDataEx.PropertySizes.Add(PropertySizePair.Key, PropertySizePair.Value);
 		}
+	}
 
+	// Build an array of pointers to the extended type
+	TArray<FMaterialDataEx*> MaterialSettingsEx;
+	MaterialSettingsEx.Reserve(MaterialDataExs.Num());
+	for (FMaterialDataEx& MaterialDataEx : MaterialDataExs)
+	{
 		MaterialSettingsEx.Add(&MaterialDataEx);
 	}
 
@@ -300,6 +305,7 @@ void FMaterialBakingModule::BakeMaterials(const TArray<FMaterialData*>& Material
 	BakeMaterials(MaterialSettingsEx, MeshSettings, BakeOutputExs);
 
 	// Translate extended bake output to old types
+	Output.Reserve(BakeOutputExs.Num());
 	for (FBakeOutputEx& BakeOutputEx : BakeOutputExs)
 	{
 		FBakeOutput& BakeOutput = Output.AddDefaulted_GetRef();
@@ -314,8 +320,38 @@ void FMaterialBakingModule::BakeMaterials(const TArray<FMaterialData*>& Material
 		{
 			BakeOutput.PropertyData.Add(PropertyDataPair.Key.Type, MoveTemp(PropertyDataPair.Value));
 		}
+
+		for (TPair<FMaterialPropertyEx, TArray<FFloat16Color>>& PropertyDataPair : BakeOutputEx.HDRPropertyData)
+		{
+			BakeOutput.HDRPropertyData.Add(PropertyDataPair.Key.Type, MoveTemp(PropertyDataPair.Value));
+		}
+	}
+}
+
+#if !(ENGINE_MAJOR_VERSION > 4 || ENGINE_MINOR_VERSION >= 27)
+static inline void ConvertRawR16G16B16A16FDataToFFloat16Color(uint32 Width, uint32 Height, uint8* In, uint32 SrcPitch, FFloat16Color* Out)
+{
+	const uint32 DstPitch = Width * sizeof(FFloat16Color);
+
+	// If source & dest pitch matches, perform a single memcpy.
+	if (DstPitch == SrcPitch)
+	{
+		FPlatformMemory::Memcpy(Out, In, Width * Height * sizeof(FFloat16Color));
+	}
+	else
+	{
+		check(SrcPitch > DstPitch);
+
+		// Need to copy row wise since the Pitch does not match the Width.
+		for (uint32 Y = 0; Y < Height; Y++)
+		{
+			FFloat16Color* SrcPtr = (FFloat16Color*)(In + Y * SrcPitch);
+			FFloat16Color* DestPtr = Out + Y * Width;
+			FMemory::Memcpy(DestPtr, SrcPtr, DstPitch);
+		}
 	}
 }
+#endif
 
 void FMaterialBakingModule::BakeMaterials(const TArray<FMaterialDataEx*>& MaterialSettings, const TArray<FMeshData*>& MeshSettings, TArray<FBakeOutputEx>& Output)
 {
@@ -487,6 +523,10 @@ void FMaterialBakingModule::BakeMaterials(const TArray<FMaterialDataEx*>& Materi
 			{
 				const FMaterialPropertyEx& Property = MaterialPropertiesToBakeOut[PropertyIndex];
 				CurrentOutput.PropertyData.Add(Property);
+				if (bEmissiveHDR && Property == MP_EmissiveColor)
+				{
+					CurrentOutput.HDRPropertyData.Add(Property);
+				}
 			}
 
 			for (int32 PropertyIndex = 0; PropertyIndex < NumPropertiesToRender; ++PropertyIndex)
@@ -501,12 +541,13 @@ void FMaterialBakingModule::BakeMaterials(const TArray<FMaterialDataEx*>& Materi
 				}
 
 				// Lookup gamma and format settings for property, if not found use default values
-				bool InForceLinearGamma = PerPropertyGamma.Contains(Property);
-				EPixelFormat PixelFormat = PerPropertyFormat.Contains(Property) ? PerPropertyFormat[Property] : PF_B8G8R8A8;
+				const EPropertyColorSpace* OverrideColorSpace = PerPropertyColorSpace.Find(Property);
+				const EPropertyColorSpace ColorSpace = OverrideColorSpace ? *OverrideColorSpace : DefaultColorSpace;
+				const EPixelFormat PixelFormat = PerPropertyFormat.Contains(Property) ? PerPropertyFormat[Property] : PF_B8G8R8A8;
 
 				// It is safe to reuse the same render target for each draw pass since they all execute sequentially on the GPU and are copied to staging buffers before
 				// being reused.
-				UTextureRenderTarget2D* RenderTarget = CreateRenderTarget(InForceLinearGamma, PixelFormat, CurrentOutput.PropertySizes[Property]);
+				UTextureRenderTarget2D* RenderTarget = CreateRenderTarget((ColorSpace == EPropertyColorSpace::Linear), PixelFormat, CurrentOutput.PropertySizes[Property]);
 				if (RenderTarget != nullptr)
 				{
 					// Perform everything left of the operation directly on the render thread since we need to modify some RenderItem's properties
@@ -551,7 +592,7 @@ void FMaterialBakingModule::BakeMaterials(const TArray<FMaterialDataEx*>& Materi
 							// Prepare a lambda for final processing that will be executed asynchronously
 							NumTasks++;
 							auto FinalProcessing_AnyThread =
-								[&NumTasks, bSaveIntermediateTextures, CurrentMaterialSettings, &StagingBufferPool, &Output, Property, MaterialIndex](FTexture2DRHIRef& StagingBuffer, void * Data, int32 DataWidth, int32 DataHeight)
+								[&NumTasks, bSaveIntermediateTextures, CurrentMaterialSettings, &StagingBufferPool, &Output, Property, MaterialIndex, bEmissiveHDR = bEmissiveHDR](FTexture2DRHIRef& StagingBuffer, void * Data, int32 DataWidth, int32 DataHeight)
 								{
 									TRACE_CPUPROFILER_EVENT_SCOPE(FinalProcessing)
 
@@ -565,6 +606,13 @@ void FMaterialBakingModule::BakeMaterials(const TArray<FMaterialDataEx*>& Materi
 									{
 										// Only one thread will write to CurrentOutput.EmissiveScale since there can be only one emissive channel property per FBakeOutputEx
 										FMaterialBakingModule::ProcessEmissiveOutput((const FFloat16Color*)Data, DataWidth, OutputSize, OutputColor, CurrentOutput.EmissiveScale);
+
+										if (bEmissiveHDR)
+										{
+											TArray<FFloat16Color>& OutputHDRColor = CurrentOutput.HDRPropertyData[Property];
+											OutputHDRColor.SetNum(OutputSize.X * OutputSize.Y);
+											ConvertRawR16G16B16A16FDataToFFloat16Color(OutputSize.X, OutputSize.Y, (uint8*)Data, DataWidth * sizeof(FFloat16Color), OutputHDRColor.GetData());
+										}
 									}
 									else
 									{
@@ -690,6 +738,32 @@ void FMaterialBakingModule::BakeMaterials(const TArray<FMaterialDataEx*>& Materi
 	}
 }
 
+void FMaterialBakingModule::SetEmissiveHDR(bool bHDR)
+{
+	bEmissiveHDR = bHDR;
+}
+
+void FMaterialBakingModule::SetLinearBake(bool bCorrectLinear)
+{
+	// PerPropertyGamma ultimately sets whether the render target is linear
+	PerPropertyColorSpace.Reset();
+	if (bCorrectLinear)
+	{
+		DefaultColorSpace = EPropertyColorSpace::Linear;
+		PerPropertyColorSpace.Add(MP_BaseColor, EPropertyColorSpace::sRGB);
+		PerPropertyColorSpace.Add(MP_EmissiveColor, EPropertyColorSpace::sRGB);
+		PerPropertyColorSpace.Add(MP_SubsurfaceColor, EPropertyColorSpace::sRGB);
+	}
+	else
+	{
+		DefaultColorSpace = EPropertyColorSpace::sRGB;
+		PerPropertyColorSpace.Add(MP_Normal, EPropertyColorSpace::Linear);
+		PerPropertyColorSpace.Add(MP_Opacity, EPropertyColorSpace::Linear);
+		PerPropertyColorSpace.Add(MP_OpacityMask, EPropertyColorSpace::Linear);
+		PerPropertyColorSpace.Add(TEXT("ClearCoatBottomNormal"), EPropertyColorSpace::Linear);
+	}
+}
+
 void FMaterialBakingModule::CleanupMaterialProxies()
 {
 	for (auto Iterator : MaterialProxyPool)
@@ -864,10 +938,19 @@ void FMaterialBakingModule::OnObjectModified(UObject* Object)
 
 	if (CVarUseMaterialProxyCaching.GetValueOnAnyThread())
 	{
-		if (Object && Object->IsA<UMaterialInterface>())
+		UMaterialInterface* MaterialToInvalidate = Cast<UMaterialInterface>(Object);
+		if (!MaterialToInvalidate)
 		{
-			UMaterialInterface* MaterialToInvalidate = Cast<UMaterialInterface>(Object);
+			// Check to see if the object is a material editor instance constant and if so, retrieve its source instance
+			UMaterialEditorInstanceConstant* EditorInstance = Cast<UMaterialEditorInstanceConstant>(Object);
+			if (EditorInstance && EditorInstance->SourceInstance)
+			{
+				MaterialToInvalidate = EditorInstance->SourceInstance;
+			}
+		}
 
+		if (MaterialToInvalidate)
+		{
 			// Search our proxy pool for materials or material instances that refer to MaterialToInvalidate
 			for (auto It = MaterialProxyPool.CreateIterator(); It; ++It)
 			{
diff --git a/Engine/Plugins/Enterprise/GLTFExporter/MigratedFromGit/Source/GLTFMaterialBaking/Private/MaterialBakingModule.h b/Engine/Plugins/Enterprise/GLTFExporter/MigratedFromGit/Source/GLTFMaterialBaking/Private/MaterialBakingModule.h
--- a/Engine/Plugins/Enterprise/GLTFExporter/MigratedFromGit/Source/GLTFMaterialBaking/Private/MaterialBakingModule.h
+++ b/Engine/Plugins/Enterprise/GLTFExporter/MigratedFromGit/Source/GLTFMaterialBaking/Private/MaterialBakingModule.h
@@ -33,6 +33,12 @@ public:
 	/** Bakes out material properties according to extended MaterialSettings using MeshSettings and stores the output in Output */
 	virtual void BakeMaterials(const TArray<FMaterialDataEx*>& MaterialSettings, const TArray<FMeshData*>& MeshSettings, TArray<FBakeOutputEx>& Output) override;
 
+	/** Outputs true HDR version of emissive color */
+	virtual void SetEmissiveHDR(bool bHDR) override;
+
+	/** Bakes all material properties to linear textures, except for colors */
+	virtual void SetLinearBake(bool bCorrectLinear) override;
+
 protected:
 	/* Creates and adds or reuses a RenderTarget from the pool */
 	UTextureRenderTarget2D* CreateRenderTarget(bool bInForceLinearGamma, EPixelFormat InPixelFormat, const FIntPoint& InTargetSize);
@@ -49,6 +55,12 @@ protected:
 	/** Callback for modified objects which should be removed from MaterialProxyPool in that case */
 	void OnObjectModified(UObject* Object);
 private:
+	enum EPropertyColorSpace
+	{
+		Linear,
+		sRGB,
+	};
+
 	/** Pool of available render targets, cached for re-using on consecutive property rendering */
 	TArray<UTextureRenderTarget2D*> RenderTargetPool;
 
@@ -62,5 +74,9 @@ private:
 	TMap<FMaterialPropertyEx, EPixelFormat> PerPropertyFormat;
 
 	/** Whether or not to enforce gamma correction while baking out specific material properties */
-	TSet<FMaterialPropertyEx> PerPropertyGamma;
+	TMap<FMaterialPropertyEx, EPropertyColorSpace> PerPropertyColorSpace;
+
+	EPropertyColorSpace DefaultColorSpace;
+
+	bool bEmissiveHDR;
 };
diff --git a/Engine/Plugins/Enterprise/GLTFExporter/MigratedFromGit/Source/GLTFMaterialBaking/Private/MaterialRenderItem.cpp b/Engine/Plugins/Enterprise/GLTFExporter/MigratedFromGit/Source/GLTFMaterialBaking/Private/MaterialRenderItem.cpp
--- a/Engine/Plugins/Enterprise/GLTFExporter/MigratedFromGit/Source/GLTFMaterialBaking/Private/MaterialRenderItem.cpp
+++ b/Engine/Plugins/Enterprise/GLTFExporter/MigratedFromGit/Source/GLTFMaterialBaking/Private/MaterialRenderItem.cpp
@@ -27,7 +27,7 @@ FMeshMaterialRenderItem::FMeshMaterialRenderItem(
 
 bool FMeshMaterialRenderItem::Render_RenderThread(FRHICommandListImmediate& RHICmdList, FMeshPassProcessorRenderState& DrawRenderState, const FCanvas* Canvas)
 {
-	checkSlow(ViewFamily && MaterialSettings && MeshSettings && MaterialRenderProxy);
+	checkSlow(ViewFamily && MeshSettings && MaterialRenderProxy);
 	// current render target set for the canvas
 	const FRenderTarget* CanvasRenderTarget = Canvas->GetRenderTarget();
 	const FIntRect ViewRect(FIntPoint(0, 0), CanvasRenderTarget->GetSizeXY());
diff --git a/Engine/Plugins/Enterprise/GLTFExporter/MigratedFromGit/Source/GLTFMaterialBaking/Public/IMaterialBakingModule.h b/Engine/Plugins/Enterprise/GLTFExporter/MigratedFromGit/Source/GLTFMaterialBaking/Public/IMaterialBakingModule.h
--- a/Engine/Plugins/Enterprise/GLTFExporter/MigratedFromGit/Source/GLTFMaterialBaking/Public/IMaterialBakingModule.h
+++ b/Engine/Plugins/Enterprise/GLTFExporter/MigratedFromGit/Source/GLTFMaterialBaking/Public/IMaterialBakingModule.h
@@ -27,4 +27,10 @@ public:
 
 	/** Bakes out material properties according to extended MaterialSettings using MeshSettings and stores the output in Output */
 	virtual void BakeMaterials(const TArray<FMaterialDataEx*>& MaterialSettings, const TArray<FMeshData*>& MeshSettings, TArray<FBakeOutputEx>& Output) = 0;
+
+	/** Outputs true HDR version of emissive color */
+	virtual void SetEmissiveHDR(bool bHDR) = 0;
+
+	/** Bakes all material properties to linear textures, except for colors */
+	virtual void SetLinearBake(bool bCorrectLinear) = 0;
 };
diff --git a/Engine/Plugins/Enterprise/GLTFExporter/MigratedFromGit/Source/GLTFMaterialBaking/Public/MaterialBakingStructures.h b/Engine/Plugins/Enterprise/GLTFExporter/MigratedFromGit/Source/GLTFMaterialBaking/Public/MaterialBakingStructures.h
--- a/Engine/Plugins/Enterprise/GLTFExporter/MigratedFromGit/Source/GLTFMaterialBaking/Public/MaterialBakingStructures.h
+++ b/Engine/Plugins/Enterprise/GLTFExporter/MigratedFromGit/Source/GLTFMaterialBaking/Public/MaterialBakingStructures.h
@@ -101,6 +101,9 @@ struct FBakeOutput
 	/** Contains the resulting texture size for baking out a material's property */
 	TMap<EMaterialProperty, FIntPoint> PropertySizes;
 
+	/** Contains the resulting HDR texture data for baking out a material's property, may be empty */
+	TMap<EMaterialProperty, TArray<FFloat16Color>> HDRPropertyData;
+
 	/** Scale used to allow having wide ranges of emissive values in the source materials, the final proxy material will use this value to scale the emissive texture's pixel values */
 	float EmissiveScale;
 };
@@ -118,6 +121,9 @@ struct FBakeOutputEx
 	/** Contains the resulting texture size for baking out a extened material's property */
 	TMap<FMaterialPropertyEx, FIntPoint> PropertySizes;
 
+	/** Contains the resulting HDR texture data for baking out a material's property, may be empty */
+	TMap<FMaterialPropertyEx, TArray<FFloat16Color>> HDRPropertyData;
+
 	/** Scale used to allow having wide ranges of emissive values in the source materials, the final proxy material will use this value to scale the emissive texture's pixel values */
 	float EmissiveScale;
 };
diff --git a/Engine/Plugins/Enterprise/GLTFExporter/MigratedFromGit/Source/GLTFMaterialBaking/Public/MaterialPropertyEx.h b/Engine/Plugins/Enterprise/GLTFExporter/MigratedFromGit/Source/GLTFMaterialBaking/Public/MaterialPropertyEx.h
--- a/Engine/Plugins/Enterprise/GLTFExporter/MigratedFromGit/Source/GLTFMaterialBaking/Public/MaterialPropertyEx.h
+++ b/Engine/Plugins/Enterprise/GLTFExporter/MigratedFromGit/Source/GLTFMaterialBaking/Public/MaterialPropertyEx.h
@@ -38,55 +38,55 @@ struct FMaterialPropertyEx
 		return !IsCustomOutput() || !Other.IsCustomOutput() ? Type - Other.Type : CustomOutput.Compare(Other.CustomOutput);
 	}
 
-	FString ToString() const
-	{
-		if (!IsCustomOutput())
-		{
-			const UEnum* Enum = StaticEnum<EMaterialProperty>();
-			FName Name = Enum->GetNameByValue(Type);
-			FString TrimmedName = Name.ToString();
-			TrimmedName.RemoveFromStart(TEXT("MP_"), ESearchCase::CaseSensitive);
-			return TrimmedName;
-		}
-
-		return CustomOutput.ToString();
-	}
-
-	bool operator ==(const FMaterialPropertyEx& Other) const
+	FORCEINLINE bool operator ==(const FMaterialPropertyEx& Other) const
 	{
 		return Equals(Other);
 	}
 
-	bool operator !=(const FMaterialPropertyEx& Other) const
+	FORCEINLINE bool operator !=(const FMaterialPropertyEx& Other) const
 	{
 		return !Equals(Other);
 	}
 
-	bool operator >=(const FMaterialPropertyEx& Other) const
+	FORCEINLINE bool operator >=(const FMaterialPropertyEx& Other) const
 	{
 		return Compare(Other) >= 0;
 	}
 
-	bool operator >(const FMaterialPropertyEx& Other) const
+	FORCEINLINE bool operator >(const FMaterialPropertyEx& Other) const
 	{
 		return Compare(Other) > 0;
 	}
 
-	bool operator <=(const FMaterialPropertyEx& Other) const
+	FORCEINLINE bool operator <=(const FMaterialPropertyEx& Other) const
 	{
 		return Compare(Other) <= 0;
 	}
 
-	bool operator <(const FMaterialPropertyEx& Other) const
+	FORCEINLINE bool operator <(const FMaterialPropertyEx& Other) const
 	{
 		return Compare(Other) < 0;
 	}
 
-	friend uint32 GetTypeHash(const FMaterialPropertyEx& Other)
+	friend FORCEINLINE uint32 GetTypeHash(const FMaterialPropertyEx& Other)
 	{
 		return !Other.IsCustomOutput() ? GetTypeHash(Other.Type) : GetTypeHash(Other.CustomOutput);
 	}
 
+	FString ToString() const
+	{
+		if (!IsCustomOutput())
+		{
+			const UEnum* Enum = StaticEnum<EMaterialProperty>();
+			FName Name = Enum->GetNameByValue(Type);
+			FString TrimmedName = Name.ToString();
+			TrimmedName.RemoveFromStart(TEXT("MP_"), ESearchCase::CaseSensitive);
+			return TrimmedName;
+		}
+
+		return CustomOutput.ToString();
+	}
+
 	/** The material property */
 	EMaterialProperty Type;
 
