[GLTFExporter] fix: improve shading model evaluation by ignoring all non-default branches in runtime switches that are statically deterministic

commit 9c14b10b5cecb44d2fc74ca991c3968468870c8b
Date: Thu, 30 Sep 2021 21:51:45 +0200

    fix: improve shading model evaluation by ignoring all non-default branches in runtime switches that are statically deterministic
    (cherry picked from commit 69f8e5524c64b0d651056bd1a99b6258733f17e4​)

#skipci

diff --git a/Engine/Plugins/Enterprise/GLTFExporter/MigratedFromGit/Source/GLTFMaterialAnalyzer/Private/GLTFProxyMaterialCompiler.h b/Engine/Plugins/Enterprise/GLTFExporter/MigratedFromGit/Source/GLTFMaterialAnalyzer/Private/GLTFProxyMaterialCompiler.h
--- a/Engine/Plugins/Enterprise/GLTFExporter/MigratedFromGit/Source/GLTFMaterialAnalyzer/Private/GLTFProxyMaterialCompiler.h
+++ b/Engine/Plugins/Enterprise/GLTFExporter/MigratedFromGit/Source/GLTFMaterialAnalyzer/Private/GLTFProxyMaterialCompiler.h
@@ -8,11 +8,33 @@ class FGLTFProxyMaterialCompiler : public FProxyMaterialCompiler
 {
 	using FProxyMaterialCompiler::FProxyMaterialCompiler;
 
+	virtual int32 GIReplace(int32 Direct, int32 StaticIndirect, int32 DynamicIndirect) override
+	{
+		return Direct; // ignore all non-default branches in this runtime switch
+	}
+
+	virtual int32 ShadowReplace(int32 Default, int32 Shadow) override
+	{
+		return Default; // ignore all non-default branches in this runtime switch
+	}
+
 	virtual int32 RayTracingQualitySwitchReplace(int32 Normal, int32 RayTraced) override
 	{
-		return Normal; // ignore RayTraced branch
+		return Normal; // ignore all non-default branches in this runtime switch
+	}
+
+	virtual int32 VirtualTextureOutputReplace(int32 Default, int32 VirtualTexture) override
+	{
+		return Default; // ignore all non-default branches in this runtime switch
 	}
 
+#if (ENGINE_MAJOR_VERSION > 4 || ENGINE_MINOR_VERSION >= 26)
+	virtual int32 ReflectionCapturePassSwitch(int32 Default, int32 Reflection) override
+	{
+		return Default; // ignore all non-default branches in this runtime switch
+	}
+#endif
+
 	virtual int32 ReflectionAboutCustomWorldNormal(int32 CustomWorldNormal, int32 bNormalizeCustomWorldNormal) override
 	{
 		return Compiler->ReflectionAboutCustomWorldNormal(CustomWorldNormal, bNormalizeCustomWorldNormal);
