[GLTFExporter] refactor: simplify FMaterialPropertyEx to better match engine version

commit 65cdf79b50f5fd46a4ea0594c12004be3365b435
Date: Sun, 10 Oct 2021 02:19:29 +0200

    refactor: simplify FMaterialPropertyEx to better match engine version
    (cherry picked from commit 9e4ef77a4034a571ccf0b7e664f7b0191e1bb204​)

#skipci

diff --git a/Engine/Plugins/Enterprise/GLTFExporter/MigratedFromGit/Source/GLTFExporter/Private/Converters/GLTFMaterialUtility.cpp b/Engine/Plugins/Enterprise/GLTFExporter/MigratedFromGit/Source/GLTFExporter/Private/Converters/GLTFMaterialUtility.cpp
--- a/Engine/Plugins/Enterprise/GLTFExporter/MigratedFromGit/Source/GLTFExporter/Private/Converters/GLTFMaterialUtility.cpp
+++ b/Engine/Plugins/Enterprise/GLTFExporter/MigratedFromGit/Source/GLTFExporter/Private/Converters/GLTFMaterialUtility.cpp
@@ -38,7 +38,7 @@ FGuid FGLTFMaterialUtility::GetAttributeID(const FMaterialPropertyEx& Property)
 FGuid FGLTFMaterialUtility::GetAttributeIDChecked(const FMaterialPropertyEx& Property)
 {
 	const FGuid AttributeID = GetAttributeID(Property);
-	check(AttributeID != FMaterialAttributeDefinitionMap::GetID(MP_MAX));
+	check(AttributeID != FMaterialAttributeDefinitionMap::GetDefaultID());
 	return AttributeID;
 }
 
diff --git a/Engine/Plugins/Enterprise/GLTFExporter/MigratedFromGit/Source/GLTFMaterialBaking/Public/MaterialPropertyEx.h b/Engine/Plugins/Enterprise/GLTFExporter/MigratedFromGit/Source/GLTFMaterialBaking/Public/MaterialPropertyEx.h
--- a/Engine/Plugins/Enterprise/GLTFExporter/MigratedFromGit/Source/GLTFMaterialBaking/Public/MaterialPropertyEx.h
+++ b/Engine/Plugins/Enterprise/GLTFExporter/MigratedFromGit/Source/GLTFMaterialBaking/Public/MaterialPropertyEx.h
@@ -28,44 +28,14 @@ struct FMaterialPropertyEx
 		return Type == MP_CustomOutput;
 	}
 
-	bool Equals(const FMaterialPropertyEx& Other) const
-	{
-		return Type == Other.Type && (!IsCustomOutput() || CustomOutput == Other.CustomOutput);
-	}
-
-	int32 Compare(const FMaterialPropertyEx& Other) const
-	{
-		return !IsCustomOutput() || !Other.IsCustomOutput() ? Type - Other.Type : CustomOutput.Compare(Other.CustomOutput);
-	}
-
 	FORCEINLINE bool operator ==(const FMaterialPropertyEx& Other) const
 	{
-		return Equals(Other);
+		return Type == Other.Type && (!IsCustomOutput() || CustomOutput == Other.CustomOutput);
 	}
 
 	FORCEINLINE bool operator !=(const FMaterialPropertyEx& Other) const
 	{
-		return !Equals(Other);
-	}
-
-	FORCEINLINE bool operator >=(const FMaterialPropertyEx& Other) const
-	{
-		return Compare(Other) >= 0;
-	}
-
-	FORCEINLINE bool operator >(const FMaterialPropertyEx& Other) const
-	{
-		return Compare(Other) > 0;
-	}
-
-	FORCEINLINE bool operator <=(const FMaterialPropertyEx& Other) const
-	{
-		return Compare(Other) <= 0;
-	}
-
-	FORCEINLINE bool operator <(const FMaterialPropertyEx& Other) const
-	{
-		return Compare(Other) < 0;
+		return !(*this == Other);
 	}
 
 	friend FORCEINLINE uint32 GetTypeHash(const FMaterialPropertyEx& Other)
