[GLTFExporter] fix: fix support for baking shading model (see PR https://github.com/EpicGames/UnrealEngine/pull/8497)

commit 699af81511e1250924359eccca38013b0903b1b8
Date: Tue, 12 Oct 2021 22:21:16 +0200

    fix: fix support for baking shading model (see PR https://github.com/EpicGames/UnrealEngine/pull/8497)
    (cherry picked from commit 6719441d69773b844f607d712ba4e217e43ff2c1​)

#skipci

diff --git a/Engine/Plugins/Enterprise/GLTFExporter/MigratedFromGit/Source/GLTFMaterialBaking/Private/ExportMaterialProxy.h b/Engine/Plugins/Enterprise/GLTFExporter/MigratedFromGit/Source/GLTFMaterialBaking/Private/ExportMaterialProxy.h
--- a/Engine/Plugins/Enterprise/GLTFExporter/MigratedFromGit/Source/GLTFMaterialBaking/Private/ExportMaterialProxy.h
+++ b/Engine/Plugins/Enterprise/GLTFExporter/MigratedFromGit/Source/GLTFMaterialBaking/Private/ExportMaterialProxy.h
@@ -307,8 +307,8 @@ public:
 		// NOTE: the following cases for custom data 0/1 and custom output are hacks since a proper solution requires engine changes:
 		case MP_CustomData0: ResourceId.Usage = static_cast<EMaterialShaderMapUsage::Type>(EMaterialShaderMapUsage::DebugViewMode + 1); break;
 		case MP_CustomData1: ResourceId.Usage = static_cast<EMaterialShaderMapUsage::Type>(EMaterialShaderMapUsage::DebugViewMode + 2); break;
-		case MP_ShadingModel: ResourceId.Usage = static_cast<EMaterialShaderMapUsage::Type>(EMaterialShaderMapUsage::DebugViewMode + 3); break;
-		case MP_CustomOutput: ResourceId.Usage = static_cast<EMaterialShaderMapUsage::Type>(EMaterialShaderMapUsage::DebugViewMode + 4); break;
+		case MP_CustomOutput: ResourceId.Usage = static_cast<EMaterialShaderMapUsage::Type>(EMaterialShaderMapUsage::DebugViewMode + 3); break;
+		case MP_ShadingModel: ResourceId.Usage = static_cast<EMaterialShaderMapUsage::Type>(EMaterialShaderMapUsage::DebugViewMode + 4); break;
 #endif
 
 		default:
@@ -462,7 +462,7 @@ public:
 				}
 				break;
 			case MP_ShadingModel:
-				return CompileShadingModel(Compiler, MaterialInterface->CompileProperty(&ProxyCompiler, MP_ShadingModel));
+				return CompileShadingModelEncoding(Compiler, MaterialInterface->CompileProperty(&ProxyCompiler, MP_ShadingModel));
 			case MP_CustomOutput:
 				 // NOTE: Currently we can assume input index is always 0, which it is for all custom outputs that are registered as material attributes
 				return CompileInputForCustomOutput(&ProxyCompiler, 0, ForceCast_Exact_Replicate);
@@ -748,7 +748,7 @@ private:
 			Compiler->Constant(0.5f)); // [-0.5,0.5] + 0.5
 	}
 
-	static int32 CompileShadingModel(FMaterialCompiler* Compiler, int32 Code)
+	static int32 CompileShadingModelEncoding(FMaterialCompiler* Compiler, int32 ShadingModelInput)
 	{
 		class FMaterialCompilerHack : public FHLSLMaterialTranslator
 		{
@@ -756,16 +756,22 @@ private:
 
 		public:
 
-			void SetParameterType(int32 Index, EMaterialValueType Type)
+			int32 ForceCastToFloat(int32 Index)
 			{
-				check(Index >= 0 && Index < CurrentScopeChunks->Num());
-				(*CurrentScopeChunks)[Index].Type = Type;
+				const FString Code = FString::Printf(TEXT("((float)(%s))"), *GetParameterCode(Index));
+				const char* Buffer = static_cast<const char*>(static_cast<const void*>(*Code));
+				const uint64 Hash = CityHash64(Buffer, Code.Len() * sizeof(TCHAR));
+				const int32 CodeIndex = CurrentScopeChunks->Num();
+				new(*CurrentScopeChunks) FShaderCodeChunk(Hash, *Code, TEXT(""), MCT_Float1, true);
+				return CodeIndex;
 			}
 		};
 
-		// NOTE: ugly hack to workaround casting shading model (uint) to float
-		static_cast<FMaterialCompilerHack*>(Compiler)->SetParameterType(Code, MCT_Float1);
-		return Compiler->Div(Code, Compiler->Constant(255)); // [0,255] / 255
+		// [0,MSM_NUM] -> [0,1]
+		return Compiler->Div(
+			// NOTE: ugly hack to workaround casting shading model (uint) to float
+			ShadingModelInput = static_cast<FMaterialCompilerHack*>(Compiler)->ForceCastToFloat(ShadingModelInput),
+			Compiler->Constant(MSM_NUM));
 	}
 
 private:
diff --git a/Engine/Plugins/Enterprise/GLTFExporter/MigratedFromGit/Source/GLTFMaterialBaking/Private/MaterialBakingModule.cpp b/Engine/Plugins/Enterprise/GLTFExporter/MigratedFromGit/Source/GLTFMaterialBaking/Private/MaterialBakingModule.cpp
--- a/Engine/Plugins/Enterprise/GLTFExporter/MigratedFromGit/Source/GLTFMaterialBaking/Private/MaterialBakingModule.cpp
+++ b/Engine/Plugins/Enterprise/GLTFExporter/MigratedFromGit/Source/GLTFMaterialBaking/Private/MaterialBakingModule.cpp
@@ -760,6 +760,7 @@ void FMaterialBakingModule::SetLinearBake(bool bCorrectLinear)
 		PerPropertyColorSpace.Add(MP_Normal, EPropertyColorSpace::Linear);
 		PerPropertyColorSpace.Add(MP_Opacity, EPropertyColorSpace::Linear);
 		PerPropertyColorSpace.Add(MP_OpacityMask, EPropertyColorSpace::Linear);
+		PerPropertyColorSpace.Add(MP_ShadingModel, EPropertyColorSpace::Linear);
 		PerPropertyColorSpace.Add(TEXT("ClearCoatBottomNormal"), EPropertyColorSpace::Linear);
 	}
 }
