[GLTFExporter] feat: add support for material baking refraction property (see PR https://github.com/EpicGames/UnrealEngine/pull/8496)

commit e7b7d16485aa7499514600897419c943942e3f73
Date: Tue, 12 Oct 2021 22:41:56 +0200

    feat: add support for material baking refraction property (see PR https://github.com/EpicGames/UnrealEngine/pull/8496)
    (cherry picked from commit f7846dd5b00de1e976b6169d0f82e995e9546ef6​)

#skipci

diff --git a/Engine/Plugins/Enterprise/GLTFExporter/MigratedFromGit/Source/GLTFMaterialBaking/Private/ExportMaterialProxy.h b/Engine/Plugins/Enterprise/GLTFExporter/MigratedFromGit/Source/GLTFMaterialBaking/Private/ExportMaterialProxy.h
--- a/Engine/Plugins/Enterprise/GLTFExporter/MigratedFromGit/Source/GLTFMaterialBaking/Private/ExportMaterialProxy.h
+++ b/Engine/Plugins/Enterprise/GLTFExporter/MigratedFromGit/Source/GLTFMaterialBaking/Private/ExportMaterialProxy.h
@@ -302,13 +302,15 @@ public:
 		case MP_CustomData0: ResourceId.Usage = EMaterialShaderMapUsage::MaterialExportClearCoat; break;
 		case MP_CustomData1: ResourceId.Usage = EMaterialShaderMapUsage::MaterialExportClearCoatRoughness; break;
 		case MP_CustomOutput: ResourceId.Usage = EMaterialShaderMapUsage::MaterialExportCustomOutput; break;
-		case MP_ShadingModel: ResourceId.Usage = static_cast<EMaterialShaderMapUsage::Type>(EMaterialShaderMapUsage::MaterialExportCustomOutput + 1); break;
+		case MP_Refraction: ResourceId.Usage = static_cast<EMaterialShaderMapUsage::Type>(EMaterialShaderMapUsage::MaterialExportCustomOutput + 1); break;
+		case MP_ShadingModel: ResourceId.Usage = static_cast<EMaterialShaderMapUsage::Type>(EMaterialShaderMapUsage::MaterialExportCustomOutput + 2); break;
 #else
 		// NOTE: the following cases for custom data 0/1 and custom output are hacks since a proper solution requires engine changes:
 		case MP_CustomData0: ResourceId.Usage = static_cast<EMaterialShaderMapUsage::Type>(EMaterialShaderMapUsage::DebugViewMode + 1); break;
 		case MP_CustomData1: ResourceId.Usage = static_cast<EMaterialShaderMapUsage::Type>(EMaterialShaderMapUsage::DebugViewMode + 2); break;
 		case MP_CustomOutput: ResourceId.Usage = static_cast<EMaterialShaderMapUsage::Type>(EMaterialShaderMapUsage::DebugViewMode + 3); break;
-		case MP_ShadingModel: ResourceId.Usage = static_cast<EMaterialShaderMapUsage::Type>(EMaterialShaderMapUsage::DebugViewMode + 4); break;
+		case MP_Refraction: ResourceId.Usage = static_cast<EMaterialShaderMapUsage::Type>(EMaterialShaderMapUsage::DebugViewMode + 4); break;
+		case MP_ShadingModel: ResourceId.Usage = static_cast<EMaterialShaderMapUsage::Type>(EMaterialShaderMapUsage::DebugViewMode + 5); break;
 #endif
 
 		default:
@@ -461,6 +463,15 @@ public:
 						MaterialInterface->CompileProperty(&ProxyCompiler, PropertyToCompile, ForceCast_Exact_Replicate));
 				}
 				break;
+			case MP_Refraction:
+					// Only return for Translucent, Additive, and AlphaComposite...
+					if (BlendMode == BLEND_Translucent || BlendMode == BLEND_Additive || BlendMode == BLEND_AlphaComposite)
+					{
+						return CompileRefractionEncoding(
+							Compiler,
+							MaterialInterface->CompileProperty(&ProxyCompiler, MP_Refraction, ForceCast_Exact_Replicate));
+					}
+				break;
 			case MP_ShadingModel:
 				return CompileShadingModelEncoding(Compiler, MaterialInterface->CompileProperty(&ProxyCompiler, MP_ShadingModel));
 			case MP_CustomOutput:
@@ -748,6 +759,17 @@ private:
 			Compiler->Constant(0.5f)); // [-0.5,0.5] + 0.5
 	}
 
+	static int32 CompileRefractionEncoding(FMaterialCompiler* Compiler, int32 RefractionInput)
+	{
+		// NOTE: avoid getting yellow output (since value type in the attribute definition is incorrectly set to MCT_Float2)
+		RefractionInput = Compiler->ForceCast(RefractionInput, MCT_Float1);
+
+		// [1,Infinity] -> [1,0]
+		return Compiler->Div(
+			Compiler->Constant(1.0f),
+			RefractionInput);
+	}
+
 	static int32 CompileShadingModelEncoding(FMaterialCompiler* Compiler, int32 ShadingModelInput)
 	{
 		class FMaterialCompilerHack : public FHLSLMaterialTranslator
diff --git a/Engine/Plugins/Enterprise/GLTFExporter/MigratedFromGit/Source/GLTFMaterialBaking/Private/MaterialBakingModule.cpp b/Engine/Plugins/Enterprise/GLTFExporter/MigratedFromGit/Source/GLTFMaterialBaking/Private/MaterialBakingModule.cpp
--- a/Engine/Plugins/Enterprise/GLTFExporter/MigratedFromGit/Source/GLTFMaterialBaking/Private/MaterialBakingModule.cpp
+++ b/Engine/Plugins/Enterprise/GLTFExporter/MigratedFromGit/Source/GLTFMaterialBaking/Private/MaterialBakingModule.cpp
@@ -256,6 +256,7 @@ void FMaterialBakingModule::StartupModule()
 	PerPropertyFormat.Add(MP_SubsurfaceColor, PF_B8G8R8A8);
 	PerPropertyFormat.Add(MP_CustomData0, PF_B8G8R8A8);
 	PerPropertyFormat.Add(MP_CustomData1, PF_B8G8R8A8);
+	PerPropertyFormat.Add(MP_Refraction, PF_B8G8R8A8);
 	PerPropertyFormat.Add(MP_ShadingModel, PF_B8G8R8A8);
 	PerPropertyFormat.Add(TEXT("ClearCoatBottomNormal"), PF_B8G8R8A8);
 
@@ -760,6 +761,7 @@ void FMaterialBakingModule::SetLinearBake(bool bCorrectLinear)
 		PerPropertyColorSpace.Add(MP_Normal, EPropertyColorSpace::Linear);
 		PerPropertyColorSpace.Add(MP_Opacity, EPropertyColorSpace::Linear);
 		PerPropertyColorSpace.Add(MP_OpacityMask, EPropertyColorSpace::Linear);
+		PerPropertyColorSpace.Add(MP_Refraction, EPropertyColorSpace::Linear);
 		PerPropertyColorSpace.Add(MP_ShadingModel, EPropertyColorSpace::Linear);
 		PerPropertyColorSpace.Add(TEXT("ClearCoatBottomNormal"), EPropertyColorSpace::Linear);
 	}
