[GLTFExporter] refactor: move FMaterialPropertyEx::ToString to new separate .cpp file (see PR https://github.com/EpicGames/UnrealEngine/pull/8492)

commit 94f95121270e39db69e5ae32dacdf589d764bd74
Date: Tue, 19 Oct 2021 20:36:55 +0200

    refactor: move FMaterialPropertyEx::ToString to new separate .cpp file (see PR https://github.com/EpicGames/UnrealEngine/pull/8492)
    (cherry picked from commit baf927603ba07558728dbef9d7e0a78218ae8823​)

#skipci

diff --git a/Engine/Plugins/Enterprise/GLTFExporter/MigratedFromGit/Source/GLTFMaterialBaking/Private/FMaterialPropertyEx.cpp b/Engine/Plugins/Enterprise/GLTFExporter/MigratedFromGit/Source/GLTFMaterialBaking/Private/FMaterialPropertyEx.cpp
new file mode 100644
--- /dev/null
+++ b/Engine/Plugins/Enterprise/GLTFExporter/MigratedFromGit/Source/GLTFMaterialBaking/Private/FMaterialPropertyEx.cpp
@@ -0,0 +1,18 @@
+// Copyright Epic Games, Inc. All Rights Reserved.
+
+#include "MaterialPropertyEx.h"
+#include "UObject/NameTypes.h"
+
+FString FMaterialPropertyEx::ToString() const
+{
+	if (!IsCustomOutput())
+	{
+		const UEnum* Enum = StaticEnum<EMaterialProperty>();
+		FName Name = Enum->GetNameByValue(Type);
+		FString TrimmedName = Name.ToString();
+		TrimmedName.RemoveFromStart(TEXT("MP_"), ESearchCase::CaseSensitive);
+		return TrimmedName;
+	}
+
+	return CustomOutput.ToString();
+}
diff --git a/Engine/Plugins/Enterprise/GLTFExporter/MigratedFromGit/Source/GLTFMaterialBaking/Public/MaterialPropertyEx.h b/Engine/Plugins/Enterprise/GLTFExporter/MigratedFromGit/Source/GLTFMaterialBaking/Public/MaterialPropertyEx.h
--- a/Engine/Plugins/Enterprise/GLTFExporter/MigratedFromGit/Source/GLTFMaterialBaking/Public/MaterialPropertyEx.h
+++ b/Engine/Plugins/Enterprise/GLTFExporter/MigratedFromGit/Source/GLTFMaterialBaking/Public/MaterialPropertyEx.h
@@ -3,7 +3,6 @@
 #pragma once
 
 #include "SceneTypes.h"
-#include "UObject/NameTypes.h"
 
 /** Structure extending EMaterialProperty to allow detailed information about custom output */
 struct FMaterialPropertyEx
@@ -23,7 +22,7 @@ struct FMaterialPropertyEx
 		, CustomOutput(CustomOutput)
 	{}
 
-	bool IsCustomOutput() const
+	FORCEINLINE bool IsCustomOutput() const
 	{
 		return Type == MP_CustomOutput;
 	}
@@ -43,19 +42,7 @@ struct FMaterialPropertyEx
 		return !Other.IsCustomOutput() ? GetTypeHash(Other.Type) : GetTypeHash(Other.CustomOutput);
 	}
 
-	FString ToString() const
-	{
-		if (!IsCustomOutput())
-		{
-			const UEnum* Enum = StaticEnum<EMaterialProperty>();
-			FName Name = Enum->GetNameByValue(Type);
-			FString TrimmedName = Name.ToString();
-			TrimmedName.RemoveFromStart(TEXT("MP_"), ESearchCase::CaseSensitive);
-			return TrimmedName;
-		}
-
-		return CustomOutput.ToString();
-	}
+	MATERIALBAKING_API FString ToString() const;
 
 	/** The material property */
 	EMaterialProperty Type;
