[GLTFExporter] fix: fix support for material baking subsurface & cloth color property (see PR https://github.com/EpicGames/UnrealEngine/pull/8501)

commit f853156b82760e089553481c5dacde7f52e870ca
Date: Tue, 19 Oct 2021 20:47:48 +0200

    fix: fix support for material baking subsurface & cloth color property (see PR https://github.com/EpicGames/UnrealEngine/pull/8501)
    (cherry picked from commit ba14a259a85de0859f8f3653da193b01532e2a0f​)

#skipci

diff --git a/Engine/Plugins/Enterprise/GLTFExporter/MigratedFromGit/Source/GLTFMaterialBaking/Private/ExportMaterialProxy.h b/Engine/Plugins/Enterprise/GLTFExporter/MigratedFromGit/Source/GLTFMaterialBaking/Private/ExportMaterialProxy.h
--- a/Engine/Plugins/Enterprise/GLTFExporter/MigratedFromGit/Source/GLTFMaterialBaking/Private/ExportMaterialProxy.h
+++ b/Engine/Plugins/Enterprise/GLTFExporter/MigratedFromGit/Source/GLTFMaterialBaking/Private/ExportMaterialProxy.h
@@ -448,6 +448,7 @@ public:
 
 			case MP_Opacity:
 			case MP_OpacityMask:
+			case MP_SubsurfaceColor:
 			case MP_CustomData0:
 			case MP_CustomData1:
 			{
