[GLTFExporter] refactor: make transformation of world-space normals to tangent-space in material baking configurable (see PR https://github.com/EpicGames/UnrealEngine/pull/8499)

commit d49393ce8ef96409df2c942990a91e07de8699fc
Date: Tue, 19 Oct 2021 20:58:55 +0200

    refactor: make transformation of world-space normals to tangent-space in material baking configurable (see PR https://github.com/EpicGames/UnrealEngine/pull/8499)
    (cherry picked from commit 5060c89ed7e1cbcbaa75e0711e8dc21d54fca6b6​)

#skipci

diff --git a/Engine/Plugins/Enterprise/GLTFExporter/MigratedFromGit/Source/GLTFExporter/Private/Converters/GLTFMaterialUtility.cpp b/Engine/Plugins/Enterprise/GLTFExporter/MigratedFromGit/Source/GLTFExporter/Private/Converters/GLTFMaterialUtility.cpp
--- a/Engine/Plugins/Enterprise/GLTFExporter/MigratedFromGit/Source/GLTFExporter/Private/Converters/GLTFMaterialUtility.cpp
+++ b/Engine/Plugins/Enterprise/GLTFExporter/MigratedFromGit/Source/GLTFExporter/Private/Converters/GLTFMaterialUtility.cpp
@@ -168,6 +168,7 @@ FGLTFPropertyBakeOutput FGLTFMaterialUtility::BakeMaterialProperty(const FIntPoi
 	FMaterialDataEx MatSet;
 	MatSet.Material = const_cast<UMaterialInterface*>(Material);
 	MatSet.PropertySizes.Add(Property, OutputSize);
+	MatSet.bTangentSpaceNormal = true;
 
 	TArray<FMeshData*> MeshSettings;
 	TArray<FMaterialDataEx*> MatSettings;
diff --git a/Engine/Plugins/Enterprise/GLTFExporter/MigratedFromGit/Source/GLTFMaterialBaking/Private/ExportMaterialProxy.h b/Engine/Plugins/Enterprise/GLTFExporter/MigratedFromGit/Source/GLTFMaterialBaking/Private/ExportMaterialProxy.h
--- a/Engine/Plugins/Enterprise/GLTFExporter/MigratedFromGit/Source/GLTFMaterialBaking/Private/ExportMaterialProxy.h
+++ b/Engine/Plugins/Enterprise/GLTFExporter/MigratedFromGit/Source/GLTFMaterialBaking/Private/ExportMaterialProxy.h
@@ -239,13 +239,14 @@ struct FExportMaterialCompiler : public FProxyMaterialCompiler
 class FExportMaterialProxy : public FMaterial, public FMaterialRenderProxy
 {
 public:
-	FExportMaterialProxy(UMaterialInterface* InMaterialInterface, EMaterialProperty InPropertyToCompile, const FString& InCustomOutputToCompile = TEXT(""), bool bInSynchronousCompilation = true, EBlendMode ProxyBlendMode = BLEND_Opaque)
+	FExportMaterialProxy(UMaterialInterface* InMaterialInterface, EMaterialProperty InPropertyToCompile, const FString& InCustomOutputToCompile = TEXT(""), bool bInSynchronousCompilation = true, EBlendMode ProxyBlendMode = BLEND_Opaque, bool bTangentSpaceNormal = false)
 		: FMaterial()
 		, MaterialInterface(InMaterialInterface)
 		, PropertyToCompile(InPropertyToCompile)
 		, CustomOutputToCompile(InCustomOutputToCompile)
 		, bSynchronousCompilation(bInSynchronousCompilation)
 		, ProxyBlendMode(ProxyBlendMode)
+		, bTangentSpaceNormal(bTangentSpaceNormal)
 	{
 
 #if (ENGINE_MAJOR_VERSION > 4 || ENGINE_MINOR_VERSION >= 26)
@@ -461,7 +462,7 @@ public:
 				{
 					return CompileNormalEncoding(
 						&ProxyCompiler,
-						MaterialInterface->CompileProperty(&ProxyCompiler, PropertyToCompile, ForceCast_Exact_Replicate));
+						CompileNormalTransform(&ProxyCompiler, MaterialInterface->CompileProperty(&ProxyCompiler, PropertyToCompile, ForceCast_Exact_Replicate)));
 				}
 				break;
 			case MP_Refraction:
@@ -722,7 +723,7 @@ private:
 
 		if (CustomOutputToCompile == TEXT("ClearCoatBottomNormal"))
 		{
-			Result = CompileNormalEncoding(Compiler, Result);
+			Result = CompileNormalEncoding(Compiler, CompileNormalTransform(Compiler, Result));
 		}
 
 		if (ForceCastFlags & MFCF_ForceCast)
@@ -747,14 +748,14 @@ private:
 		return nullptr;
 	}
 
-	int32 CompileNormalEncoding(FMaterialCompiler* Compiler, int32 NormalInput) const
+	int32 CompileNormalTransform(FMaterialCompiler* Compiler, int32 NormalInput) const
 	{
-		// TODO: make this configurable instead of assuming that we always want tangent-space normals
-		if (!Material->bTangentSpaceNormal)
-		{
-			NormalInput = Compiler->TransformVector(MCB_World, MCB_Tangent, NormalInput);
-		}
+		return bTangentSpaceNormal && !Material->bTangentSpaceNormal
+			? Compiler->TransformVector(MCB_World, MCB_Tangent, NormalInput) : NormalInput;
+	}
 
+	int32 CompileNormalEncoding(FMaterialCompiler* Compiler, int32 NormalInput) const
+	{
 		return Compiler->Add(
 			Compiler->Mul(NormalInput, Compiler->Constant(0.5f)), // [-1,1] * 0.5
 			Compiler->Constant(0.5f)); // [-0.5,0.5] + 0.5
@@ -808,5 +809,10 @@ private:
 	FString CustomOutputToCompile;
 	FGuid Id;
 	bool bSynchronousCompilation;
+
+public:
+	/** The blend mode used when baking the proxy material */
 	EBlendMode ProxyBlendMode;
+	/** Whether to transform normals from world-space to tangent-space (does nothing if material already uses tangent-space normals) */
+	bool bTangentSpaceNormal;
 };
diff --git a/Engine/Plugins/Enterprise/GLTFExporter/MigratedFromGit/Source/GLTFMaterialBaking/Private/MaterialBakingModule.cpp b/Engine/Plugins/Enterprise/GLTFExporter/MigratedFromGit/Source/GLTFMaterialBaking/Private/MaterialBakingModule.cpp
--- a/Engine/Plugins/Enterprise/GLTFExporter/MigratedFromGit/Source/GLTFMaterialBaking/Private/MaterialBakingModule.cpp
+++ b/Engine/Plugins/Enterprise/GLTFExporter/MigratedFromGit/Source/GLTFMaterialBaking/Private/MaterialBakingModule.cpp
@@ -287,6 +287,7 @@ void FMaterialBakingModule::BakeMaterials(const TArray<FMaterialData*>& Material
 		MaterialDataEx.Material = MaterialData->Material;
 		MaterialDataEx.bPerformBorderSmear = MaterialData->bPerformBorderSmear;
 		MaterialDataEx.BlendMode = MaterialData->BlendMode;
+		MaterialDataEx.bTangentSpaceNormal = MaterialData->bTangentSpaceNormal;
 
 		for (const TPair<EMaterialProperty, FIntPoint>& PropertySizePair : MaterialData->PropertySizes)
 		{
@@ -469,7 +470,7 @@ void FMaterialBakingModule::BakeMaterials(const TArray<FMaterialDataEx*>& Materi
 			for (TMap<FMaterialPropertyEx, FIntPoint>::TConstIterator PropertySizeIterator = CurrentMaterialSettings->PropertySizes.CreateConstIterator(); PropertySizeIterator; ++PropertySizeIterator)
 			{
 				// They will be stored in the pool and compiled asynchronously
-				CreateMaterialProxy(*CurrentMaterialSettings, PropertySizeIterator.Key());
+				CreateMaterialProxy(CurrentMaterialSettings, PropertySizeIterator.Key());
 			}
 		}
 	}
@@ -533,7 +534,7 @@ void FMaterialBakingModule::BakeMaterials(const TArray<FMaterialDataEx*>& Materi
 			for (int32 PropertyIndex = 0; PropertyIndex < NumPropertiesToRender; ++PropertyIndex)
 			{
 				const FMaterialPropertyEx& Property = MaterialPropertiesToBakeOut[PropertyIndex];
-				FExportMaterialProxy* ExportMaterialProxy = CreateMaterialProxy(*CurrentMaterialSettings, Property);
+				FExportMaterialProxy* ExportMaterialProxy = CreateMaterialProxy(CurrentMaterialSettings, Property);
 
 				if (!ExportMaterialProxy->IsCompilationFinished())
 				{
@@ -815,7 +816,7 @@ UTextureRenderTarget2D* FMaterialBakingModule::CreateRenderTarget(bool bInForceL
 	return RenderTarget;
 }
 
-FExportMaterialProxy* FMaterialBakingModule::CreateMaterialProxy(const FMaterialDataEx& MatSettings, const FMaterialPropertyEx& Property)
+FExportMaterialProxy* FMaterialBakingModule::CreateMaterialProxy(const FMaterialDataEx* MaterialSettings, const FMaterialPropertyEx& Property)
 {
 	TRACE_CPUPROFILER_EVENT_SCOPE(FMaterialBakingModule::CreateMaterialProxy)
 
@@ -823,12 +824,12 @@ FExportMaterialProxy* FMaterialBakingModule::CreateMaterialProxy(const FMaterial
 
 	// Find all pooled material proxy matching this material
 	TArray<FMaterialPoolValue> Entries;
-	MaterialProxyPool.MultiFind(MatSettings.Material, Entries);
+	MaterialProxyPool.MultiFind(MaterialSettings->Material, Entries);
 
 	// Look for the matching property
 	for (FMaterialPoolValue& Entry : Entries)
 	{
-		if (Entry.Key == Property && Entry.Value->GetBlendMode() == MatSettings.BlendMode)
+		if (Entry.Key == Property && Entry.Value->ProxyBlendMode == MaterialSettings->BlendMode && Entry.Value->bTangentSpaceNormal == MaterialSettings->bTangentSpaceNormal)
 		{
 			Proxy = Entry.Value;
 			break;
@@ -838,8 +839,8 @@ FExportMaterialProxy* FMaterialBakingModule::CreateMaterialProxy(const FMaterial
 	// Not found, create a new entry
 	if (Proxy == nullptr)
 	{
-		Proxy = new FExportMaterialProxy(MatSettings.Material, Property.Type, Property.CustomOutput.ToString(), false /* bInSynchronousCompilation */, MatSettings.BlendMode);
-		MaterialProxyPool.Add(MatSettings.Material, FMaterialPoolValue(Property, Proxy));
+		Proxy = new FExportMaterialProxy(MaterialSettings->Material, Property.Type, Property.CustomOutput.ToString(), false /* bInSynchronousCompilation */, MaterialSettings->BlendMode, MaterialSettings->bTangentSpaceNormal);
+		MaterialProxyPool.Add(MaterialSettings->Material, FMaterialPoolValue(Property, Proxy));
 	}
 
 	return Proxy;
diff --git a/Engine/Plugins/Enterprise/GLTFExporter/MigratedFromGit/Source/GLTFMaterialBaking/Private/MaterialBakingModule.h b/Engine/Plugins/Enterprise/GLTFExporter/MigratedFromGit/Source/GLTFMaterialBaking/Private/MaterialBakingModule.h
--- a/Engine/Plugins/Enterprise/GLTFExporter/MigratedFromGit/Source/GLTFMaterialBaking/Private/MaterialBakingModule.h
+++ b/Engine/Plugins/Enterprise/GLTFExporter/MigratedFromGit/Source/GLTFMaterialBaking/Private/MaterialBakingModule.h
@@ -44,7 +44,7 @@ protected:
 	UTextureRenderTarget2D* CreateRenderTarget(bool bInForceLinearGamma, EPixelFormat InPixelFormat, const FIntPoint& InTargetSize);
 
 	/* Creates and adds (or reuses a ExportMaterialProxy from the pool if MaterialBaking.UseMaterialProxyCaching is set to 1) */
-	FExportMaterialProxy* CreateMaterialProxy(const FMaterialDataEx& MatSettings, const FMaterialPropertyEx& Property);
+	FExportMaterialProxy* CreateMaterialProxy(const FMaterialDataEx* MaterialSettings, const FMaterialPropertyEx& Property);
 
 	/** Helper for emissive color conversion to Output */
 	static void ProcessEmissiveOutput(const FFloat16Color* Color16, int32 Color16Pitch, const FIntPoint& OutputSize, TArray<FColor>& Output, float& EmissiveScale);
diff --git a/Engine/Plugins/Enterprise/GLTFExporter/MigratedFromGit/Source/GLTFMaterialBaking/Public/MaterialBakingStructures.h b/Engine/Plugins/Enterprise/GLTFExporter/MigratedFromGit/Source/GLTFMaterialBaking/Public/MaterialBakingStructures.h
--- a/Engine/Plugins/Enterprise/GLTFExporter/MigratedFromGit/Source/GLTFMaterialBaking/Public/MaterialBakingStructures.h
+++ b/Engine/Plugins/Enterprise/GLTFExporter/MigratedFromGit/Source/GLTFMaterialBaking/Public/MaterialBakingStructures.h
@@ -17,6 +17,7 @@ struct FMaterialData
 		: Material(nullptr)
 		, bPerformBorderSmear(true)
 		, BlendMode(BLEND_Opaque)
+		, bTangentSpaceNormal(false)
 	{}
 
 	/** Material to bake out */
@@ -27,6 +28,8 @@ struct FMaterialData
 	bool bPerformBorderSmear;
 	/** Blend mode to use when baking */
 	EBlendMode BlendMode;
+	/** Whether to transform normals from world-space to tangent-space (does nothing if material already uses tangent-space normals) */
+	bool bTangentSpaceNormal;
 };
 
 /** Structure containing extended information about the material and properties which is being baked out */
@@ -36,6 +39,7 @@ struct FMaterialDataEx
 		: Material(nullptr)
 		, bPerformBorderSmear(true)
 		, BlendMode(BLEND_Opaque)
+		, bTangentSpaceNormal(false)
 	{}
 
 	/** Material to bake out */
@@ -46,6 +50,8 @@ struct FMaterialDataEx
 	bool bPerformBorderSmear;
 	/** Blend mode to use when baking */
 	EBlendMode BlendMode;
+	/** Whether to transform normals from world-space to tangent-space (does nothing if material already uses tangent-space normals) */
+	bool bTangentSpaceNormal;
 };
 
 struct FMeshData
