[GLTFExporter] feat: add support for material baking custom output transmittance color (see PR https://github.com/EpicGames/UnrealEngine/pull/8498)

commit 8b127d5fbe5b9b171dba7fea65c386fd411dc89a
Date: Tue, 19 Oct 2021 22:38:49 +0200

    feat: add support for material baking custom output transmittance color (see PR https://github.com/EpicGames/UnrealEngine/pull/8498)
    (cherry picked from commit b9b872e0fdc384f51239af11e4ce9cd8275c3ce1​)

#skipci

diff --git a/Engine/Plugins/Enterprise/GLTFExporter/MigratedFromGit/Source/GLTFMaterialBaking/Private/AddCustomAttributeHack.cpp b/Engine/Plugins/Enterprise/GLTFExporter/MigratedFromGit/Source/GLTFMaterialBaking/Private/AddCustomAttributeHack.cpp
new file mode 100644
--- /dev/null
+++ b/Engine/Plugins/Enterprise/GLTFExporter/MigratedFromGit/Source/GLTFMaterialBaking/Private/AddCustomAttributeHack.cpp
@@ -0,0 +1,35 @@
+// Copyright Epic Games, Inc. All Rights Reserved.
+
+#include "AddCustomAttributeHack.h"
+#include "MaterialShared.h"
+
+// NOTE: to access the raw array of custom attributes we need to use a dummy friend class
+class FMaterialAttributePropertyDetails
+{
+public:
+
+	static TArray<FMaterialCustomOutputAttributeDefintion>& GetCustomAttributes()
+	{
+		return FMaterialAttributeDefinitionMap::GMaterialPropertyAttributesMap.CustomAttributes;
+	}
+};
+
+void AddCustomAttributeHack(const FGuid& AttributeID, const FString& AttributeName, const FString& FunctionName, EMaterialValueType ValueType, const FVector4& DefaultValue, MaterialAttributeBlendFunction BlendFunction)
+{
+	// NOTE: since the linker doesn't have access to the constructor we need a empty buffer
+	uint8 Buffer[sizeof(FMaterialCustomOutputAttributeDefintion)] = { 0 };
+	FMaterialCustomOutputAttributeDefintion& CustomAttribute = *reinterpret_cast<FMaterialCustomOutputAttributeDefintion*>(Buffer);
+
+	CustomAttribute.AttributeID = AttributeID;
+	CustomAttribute.DefaultValue = DefaultValue;
+	CustomAttribute.AttributeName = AttributeName;
+	CustomAttribute.Property = MP_CustomOutput;
+	CustomAttribute.ValueType = ValueType;
+	CustomAttribute.ShaderFrequency = SF_Pixel;
+	CustomAttribute.TexCoordIndex = INDEX_NONE;
+	CustomAttribute.BlendFunction = BlendFunction;
+	CustomAttribute.bIsHidden = false;
+	CustomAttribute.FunctionName = FunctionName;
+
+	FMaterialAttributePropertyDetails::GetCustomAttributes().Add(CustomAttribute);
+}
diff --git a/Engine/Plugins/Enterprise/GLTFExporter/MigratedFromGit/Source/GLTFMaterialBaking/Private/AddCustomAttributeHack.h b/Engine/Plugins/Enterprise/GLTFExporter/MigratedFromGit/Source/GLTFMaterialBaking/Private/AddCustomAttributeHack.h
new file mode 100644
--- /dev/null
+++ b/Engine/Plugins/Enterprise/GLTFExporter/MigratedFromGit/Source/GLTFMaterialBaking/Private/AddCustomAttributeHack.h
@@ -0,0 +1,6 @@
+// Copyright Epic Games, Inc. All Rights Reserved.
+
+#pragma once
+
+// NOTE: ugly hack to register custom attribute without calling FMaterialAttributeDefinitionMap::AddCustomAttribute (which doesn't work after the AttributeDDCString has been initialized)
+void AddCustomAttributeHack(const FGuid& AttributeID, const FString& AttributeName, const FString& FunctionName, EMaterialValueType ValueType, const FVector4& DefaultValue, MaterialAttributeBlendFunction BlendFunction = nullptr);
diff --git a/Engine/Plugins/Enterprise/GLTFExporter/MigratedFromGit/Source/GLTFMaterialBaking/Private/ExportMaterialProxy.h b/Engine/Plugins/Enterprise/GLTFExporter/MigratedFromGit/Source/GLTFMaterialBaking/Private/ExportMaterialProxy.h
--- a/Engine/Plugins/Enterprise/GLTFExporter/MigratedFromGit/Source/GLTFMaterialBaking/Private/ExportMaterialProxy.h
+++ b/Engine/Plugins/Enterprise/GLTFExporter/MigratedFromGit/Source/GLTFMaterialBaking/Private/ExportMaterialProxy.h
@@ -302,7 +302,10 @@ public:
 #if (ENGINE_MAJOR_VERSION > 4 || ENGINE_MINOR_VERSION >= 27)
 		case MP_CustomData0: ResourceId.Usage = EMaterialShaderMapUsage::MaterialExportClearCoat; break;
 		case MP_CustomData1: ResourceId.Usage = EMaterialShaderMapUsage::MaterialExportClearCoatRoughness; break;
-		case MP_CustomOutput: ResourceId.Usage = EMaterialShaderMapUsage::MaterialExportCustomOutput; break;
+		case MP_CustomOutput:
+			ResourceId.Usage = EMaterialShaderMapUsage::MaterialExportCustomOutput;
+			ResourceId.UsageCustomOutput = InCustomOutputToCompile;
+			break;
 		case MP_Refraction: ResourceId.Usage = static_cast<EMaterialShaderMapUsage::Type>(EMaterialShaderMapUsage::MaterialExportCustomOutput + 1); break;
 		case MP_ShadingModel: ResourceId.Usage = static_cast<EMaterialShaderMapUsage::Type>(EMaterialShaderMapUsage::MaterialExportCustomOutput + 2); break;
 #else
@@ -691,24 +694,15 @@ public:
 private:
 	int32 CompileInputForCustomOutput(FMaterialCompiler* Compiler, int32 InputIndex, uint32 ForceCastFlags) const
 	{
-#if (ENGINE_MAJOR_VERSION > 4 || ENGINE_MINOR_VERSION >= 27)
-		FGuid AttributeID = FMaterialAttributeDefinitionMap::GetCustomAttributeID(CustomOutputToCompile);
-#else
-		FGuid AttributeID = {};
-
 		TArray<FMaterialCustomOutputAttributeDefintion> CustomAttributes;
 		FMaterialAttributeDefinitionMap::GetCustomAttributeList(CustomAttributes);
-		for (auto& Attribute : CustomAttributes)
+		FMaterialCustomOutputAttributeDefintion* CustomAttribute = CustomAttributes.FindByPredicate([this](const FMaterialCustomOutputAttributeDefintion& Attribute) -> bool
 		{
-			if (Attribute.AttributeName == CustomOutputToCompile)
-			{
-				AttributeID = Attribute.AttributeID;
-			}
-		}
-#endif
-		check(AttributeID.IsValid());
+			return Attribute.AttributeName == CustomOutputToCompile;
+		});
+		check(CustomAttribute != nullptr);
 
-		UMaterialExpressionCustomOutput* Expression = GetCustomOutputExpressionToCompile();
+		UMaterialExpressionCustomOutput* Expression = GetCustomOutputExpression(CustomAttribute->FunctionName);
 		FExpressionInput* ExpressionInput = Expression ? Expression->GetInput(InputIndex) : nullptr;
 		int32 Result = INDEX_NONE;
 
@@ -718,7 +712,7 @@ private:
 		}
 		else
 		{
-			Result = FMaterialAttributeDefinitionMap::CompileDefaultExpression(Compiler, AttributeID);
+			Result = FMaterialAttributeDefinitionMap::CompileDefaultExpression(Compiler, CustomAttribute->AttributeID);
 		}
 
 		if (CustomOutputToCompile == TEXT("ClearCoatBottomNormal"))
@@ -728,18 +722,18 @@ private:
 
 		if (ForceCastFlags & MFCF_ForceCast)
 		{
-			Result = Compiler->ForceCast(Result, FMaterialAttributeDefinitionMap::GetValueType(AttributeID), ForceCastFlags);
+			Result = Compiler->ForceCast(Result,  CustomAttribute->ValueType, ForceCastFlags);
 		}
 
 		return Result;
 	}
 
-	UMaterialExpressionCustomOutput* GetCustomOutputExpressionToCompile() const
+	UMaterialExpressionCustomOutput* GetCustomOutputExpression(const FString& FunctionName) const
 	{
 		for (UMaterialExpression* Expression : Material->Expressions)
 		{
 			UMaterialExpressionCustomOutput* CustomOutputExpression = Cast<UMaterialExpressionCustomOutput>(Expression);
-			if (CustomOutputExpression && CustomOutputExpression->GetDisplayName() == CustomOutputToCompile)
+			if (CustomOutputExpression && CustomOutputExpression->GetFunctionName() == FunctionName)
 			{
 				return CustomOutputExpression;
 			}
diff --git a/Engine/Plugins/Enterprise/GLTFExporter/MigratedFromGit/Source/GLTFMaterialBaking/Private/MaterialBakingModule.cpp b/Engine/Plugins/Enterprise/GLTFExporter/MigratedFromGit/Source/GLTFMaterialBaking/Private/MaterialBakingModule.cpp
--- a/Engine/Plugins/Enterprise/GLTFExporter/MigratedFromGit/Source/GLTFMaterialBaking/Private/MaterialBakingModule.cpp
+++ b/Engine/Plugins/Enterprise/GLTFExporter/MigratedFromGit/Source/GLTFMaterialBaking/Private/MaterialBakingModule.cpp
@@ -7,8 +7,8 @@
 #include "Interfaces/IMainFrameModule.h"
 #include "UObject/UObjectGlobals.h"
 #include "MaterialBakingStructures.h"
-#include "Framework/Application/SlateApplication.h"
 #include "MaterialBakingHelpers.h"
+#include "AddCustomAttributeHack.h"
 #include "Async/Async.h"
 #include "Async/ParallelFor.h"
 #include "Materials/MaterialInstance.h"
@@ -259,9 +259,13 @@ void FMaterialBakingModule::StartupModule()
 	PerPropertyFormat.Add(MP_Refraction, PF_B8G8R8A8);
 	PerPropertyFormat.Add(MP_ShadingModel, PF_B8G8R8A8);
 	PerPropertyFormat.Add(TEXT("ClearCoatBottomNormal"), PF_B8G8R8A8);
+	PerPropertyFormat.Add(TEXT("TransmittanceColor"), PF_B8G8R8A8);
 
 	// Register callback for modified objects
 	FCoreUObjectDelegates::OnObjectModified.AddRaw(this, &FMaterialBakingModule::OnObjectModified);
+
+	// NOTE: Because the attribute definition of TransmittanceColor is not registered by the engine, we have to do it here
+	AddCustomAttributeHack(FGuid(0xF2D8C70E, 0x42ECA0D1, 0x4652D0AD, 0xB785A065), "TransmittanceColor", "GetThinTranslucentMaterialOutput", MCT_Float3, FVector4(0.5, 0.5, 0.5, 0));
 }
 
 void FMaterialBakingModule::ShutdownModule()
@@ -755,6 +759,7 @@ void FMaterialBakingModule::SetLinearBake(bool bCorrectLinear)
 		PerPropertyColorSpace.Add(MP_BaseColor, EPropertyColorSpace::sRGB);
 		PerPropertyColorSpace.Add(MP_EmissiveColor, EPropertyColorSpace::sRGB);
 		PerPropertyColorSpace.Add(MP_SubsurfaceColor, EPropertyColorSpace::sRGB);
+		PerPropertyColorSpace.Add(TEXT("TransmittanceColor"), EPropertyColorSpace::sRGB);
 	}
 	else
 	{
