[GLTFExporter] feat: add support for material baking expressions that rely on primitive data (see PR https://github.com/EpicGames/UnrealEngine/pull/8526)

commit 45bf3bf9bea771dbac27ab530f78e1ab12fcf301
Date: Tue, 19 Oct 2021 22:54:41 +0200

    feat: add support for material baking expressions that rely on primitive data (see PR https://github.com/EpicGames/UnrealEngine/pull/8526)
    (cherry picked from commit a0c1f9cfade3d218e38a1dc1e18e84390a4e8b01​)

#skipci

diff --git a/Engine/Plugins/Enterprise/GLTFExporter/MigratedFromGit/Source/GLTFExporter/Private/Converters/GLTFMaterialUtility.cpp b/Engine/Plugins/Enterprise/GLTFExporter/MigratedFromGit/Source/GLTFExporter/Private/Converters/GLTFMaterialUtility.cpp
--- a/Engine/Plugins/Enterprise/GLTFExporter/MigratedFromGit/Source/GLTFExporter/Private/Converters/GLTFMaterialUtility.cpp
+++ b/Engine/Plugins/Enterprise/GLTFExporter/MigratedFromGit/Source/GLTFExporter/Private/Converters/GLTFMaterialUtility.cpp
@@ -157,13 +157,17 @@ bool FGLTFMaterialUtility::CombineTextures(TArray<FColor>& OutPixels, const TArr
 	return bReadSuccessful;
 }
 
-FGLTFPropertyBakeOutput FGLTFMaterialUtility::BakeMaterialProperty(const FIntPoint& OutputSize, const FMaterialPropertyEx& Property, const UMaterialInterface* Material, int32 TexCoord, const FMeshDescription* MeshDescription, const FGLTFIndexArray& MeshSectionIndices, bool bCopyAlphaFromRedChannel)
+FGLTFPropertyBakeOutput FGLTFMaterialUtility::BakeMaterialProperty(const FIntPoint& OutputSize, const FMaterialPropertyEx& Property, const UMaterialInterface* Material, int32 TexCoord, const FGLTFMeshData* MeshData, const FGLTFIndexArray& MeshSectionIndices, bool bCopyAlphaFromRedChannel)
 {
 	FMeshData MeshSet;
 	MeshSet.TextureCoordinateBox = { { 0.0f, 0.0f }, { 1.0f, 1.0f } };
 	MeshSet.TextureCoordinateIndex = TexCoord;
-	MeshSet.RawMeshDescription = const_cast<FMeshDescription*>(MeshDescription);
 	MeshSet.MaterialIndices = MeshSectionIndices; // NOTE: MaterialIndices is actually section indices
+	if (MeshData != nullptr)
+	{
+		MeshSet.PrimitiveData = &MeshData->PrimitiveData;
+		MeshSet.RawMeshDescription = const_cast<FMeshDescription*>(&MeshData->Description);
+	}
 
 	FMaterialDataEx MatSet;
 	MatSet.Material = const_cast<UMaterialInterface*>(Material);
diff --git a/Engine/Plugins/Enterprise/GLTFExporter/MigratedFromGit/Source/GLTFExporter/Private/Converters/GLTFMaterialUtility.h b/Engine/Plugins/Enterprise/GLTFExporter/MigratedFromGit/Source/GLTFExporter/Private/Converters/GLTFMaterialUtility.h
--- a/Engine/Plugins/Enterprise/GLTFExporter/MigratedFromGit/Source/GLTFExporter/Private/Converters/GLTFMaterialUtility.h
+++ b/Engine/Plugins/Enterprise/GLTFExporter/MigratedFromGit/Source/GLTFExporter/Private/Converters/GLTFMaterialUtility.h
@@ -60,7 +60,7 @@ struct FGLTFMaterialUtility
 
 	static UTexture2D* CreateTransientTexture(const FGLTFPropertyBakeOutput& PropertyBakeOutput, bool bUseSRGB = false);
 	static bool CombineTextures(TArray<FColor>& OutPixels, const TArray<FGLTFTextureCombineSource>& Sources, const FIntPoint& OutputSize, EPixelFormat OutputPixelFormat);
-	static FGLTFPropertyBakeOutput BakeMaterialProperty(const FIntPoint& OutputSize, const FMaterialPropertyEx& Property, const UMaterialInterface* Material, int32 TexCoord, const FMeshDescription* MeshDescription = nullptr, const FGLTFIndexArray& MeshSectionIndices = {}, bool bCopyAlphaFromRedChannel = false);
+	static FGLTFPropertyBakeOutput BakeMaterialProperty(const FIntPoint& OutputSize, const FMaterialPropertyEx& Property, const UMaterialInterface* Material, int32 TexCoord, const FGLTFMeshData* MeshData = nullptr, const FGLTFIndexArray& MeshSectionIndices = {}, bool bCopyAlphaFromRedChannel = false);
 
 	static FGLTFJsonTextureIndex AddCombinedTexture(FGLTFConvertBuilder& Builder, const TArray<FGLTFTextureCombineSource>& CombineSources, const FIntPoint& TextureSize, bool bIgnoreAlpha, const FString& TextureName, EGLTFJsonTextureFilter MinFilter, EGLTFJsonTextureFilter MagFilter, EGLTFJsonTextureWrap WrapS, EGLTFJsonTextureWrap WrapT);
 	static FGLTFJsonTextureIndex AddTexture(FGLTFConvertBuilder& Builder, const TArray<FColor>& Pixels, const FIntPoint& TextureSize, bool bIgnoreAlpha, bool bIsNormalMap, const FString& TextureName, EGLTFJsonTextureFilter MinFilter, EGLTFJsonTextureFilter MagFilter, EGLTFJsonTextureWrap WrapS, EGLTFJsonTextureWrap WrapT);
diff --git a/Engine/Plugins/Enterprise/GLTFExporter/MigratedFromGit/Source/GLTFExporter/Private/Converters/GLTFMeshData.cpp b/Engine/Plugins/Enterprise/GLTFExporter/MigratedFromGit/Source/GLTFExporter/Private/Converters/GLTFMeshData.cpp
--- a/Engine/Plugins/Enterprise/GLTFExporter/MigratedFromGit/Source/GLTFExporter/Private/Converters/GLTFMeshData.cpp
+++ b/Engine/Plugins/Enterprise/GLTFExporter/MigratedFromGit/Source/GLTFExporter/Private/Converters/GLTFMeshData.cpp
@@ -15,11 +15,13 @@ FGLTFMeshData::FGLTFMeshData(const UStaticMesh* StaticMesh, const UStaticMeshCom
 	if (StaticMeshComponent != nullptr)
 	{
 		Name = FGLTFNameUtility::GetName(StaticMeshComponent);
+		PrimitiveData = { StaticMeshComponent };
 		FMeshMergeHelpers::RetrieveMesh(StaticMeshComponent, LODIndex, Description, true);
 	}
 	else
 	{
 		StaticMesh->GetName(Name);
+		PrimitiveData = { StaticMesh };
 		FMeshMergeHelpers::RetrieveMesh(StaticMesh, LODIndex, Description);
 	}
 
@@ -37,11 +39,13 @@ FGLTFMeshData::FGLTFMeshData(const USkeletalMesh* SkeletalMesh, const USkeletalM
 	if (SkeletalMeshComponent != nullptr)
 	{
 		Name = FGLTFNameUtility::GetName(SkeletalMeshComponent);
+		PrimitiveData = { SkeletalMeshComponent };
 		FMeshMergeHelpers::RetrieveMesh(const_cast<USkeletalMeshComponent*>(SkeletalMeshComponent), LODIndex, Description, true);
 	}
 	else
 	{
 		SkeletalMesh->GetName(Name);
+		PrimitiveData = { SkeletalMesh };
 
 		// NOTE: this is a workaround for the fact that there's no overload for FMeshMergeHelpers::RetrieveMesh
 		// that accepts a USkeletalMesh, only a USkeletalMeshComponent.
diff --git a/Engine/Plugins/Enterprise/GLTFExporter/MigratedFromGit/Source/GLTFExporter/Private/Converters/GLTFMeshData.h b/Engine/Plugins/Enterprise/GLTFExporter/MigratedFromGit/Source/GLTFExporter/Private/Converters/GLTFMeshData.h
--- a/Engine/Plugins/Enterprise/GLTFExporter/MigratedFromGit/Source/GLTFExporter/Private/Converters/GLTFMeshData.h
+++ b/Engine/Plugins/Enterprise/GLTFExporter/MigratedFromGit/Source/GLTFExporter/Private/Converters/GLTFMeshData.h
@@ -4,6 +4,7 @@
 
 #include "CoreMinimal.h"
 #include "MeshDescription.h"
+#include "MaterialBakingStructures.h"
 
 struct FGLTFMeshData
 {
@@ -13,6 +14,7 @@ struct FGLTFMeshData
 	const FGLTFMeshData* GetParent() const;
 
 	FString Name;
+	FPrimitiveData PrimitiveData;
 	FMeshDescription Description;
 	int32 TexCoord;
 
diff --git a/Engine/Plugins/Enterprise/GLTFExporter/MigratedFromGit/Source/GLTFExporter/Private/Tasks/GLTFMaterialTasks.cpp b/Engine/Plugins/Enterprise/GLTFExporter/MigratedFromGit/Source/GLTFExporter/Private/Tasks/GLTFMaterialTasks.cpp
--- a/Engine/Plugins/Enterprise/GLTFExporter/MigratedFromGit/Source/GLTFExporter/Private/Tasks/GLTFMaterialTasks.cpp
+++ b/Engine/Plugins/Enterprise/GLTFExporter/MigratedFromGit/Source/GLTFExporter/Private/Tasks/GLTFMaterialTasks.cpp
@@ -1337,7 +1337,7 @@ FGLTFPropertyBakeOutput FGLTFMaterialTask::BakeMaterialProperty(const FMaterialP
 		Property,
 		Material,
 		OutTexCoord,
-		MeshData != nullptr ? &MeshData->Description : nullptr,
+		MeshData,
 		SectionIndices,
 		bCopyAlphaFromRedChannel);
 
diff --git a/Engine/Plugins/Enterprise/GLTFExporter/MigratedFromGit/Source/GLTFMaterialBaking/Private/ExportMaterialProxy.h b/Engine/Plugins/Enterprise/GLTFExporter/MigratedFromGit/Source/GLTFMaterialBaking/Private/ExportMaterialProxy.h
--- a/Engine/Plugins/Enterprise/GLTFExporter/MigratedFromGit/Source/GLTFMaterialBaking/Private/ExportMaterialProxy.h
+++ b/Engine/Plugins/Enterprise/GLTFExporter/MigratedFromGit/Source/GLTFMaterialBaking/Private/ExportMaterialProxy.h
@@ -47,21 +47,11 @@ struct FExportMaterialCompiler : public FProxyMaterialCompiler
 #endif
 	}
 
-	virtual int32 ObjectWorldPosition() override
-	{
-		return Compiler->ObjectWorldPosition();
-	}
-
 	virtual int32 DistanceCullFade() override
 	{
 		return Compiler->Constant(1.0f);
 	}
 
-	virtual int32 ActorWorldPosition() override
-	{
-		return Compiler->ActorWorldPosition();
-	}
-
 	virtual int32 ParticleRelativeTime() override
 	{
 		return Compiler->Constant(0.0f);
@@ -98,21 +88,6 @@ struct FExportMaterialCompiler : public FProxyMaterialCompiler
 		return Compiler->Constant2(0.0f, 0.0f);
 	}
 
-	virtual int32 ObjectRadius() override
-	{
-		return Compiler->Constant(500);
-	}
-
-	virtual int32 ObjectBounds() override
-	{
-		return Compiler->ObjectBounds();
-	}
-
-	virtual int32 PreSkinnedLocalBounds(int32 OutputIndex) override
-	{
-		return Compiler->PreSkinnedLocalBounds(OutputIndex);
-	}
-
 	virtual int32 CameraVector() override
 	{
 		// NOTE: by using VertexNormal instead of a fixed (world-space) vector, we allow
diff --git a/Engine/Plugins/Enterprise/GLTFExporter/MigratedFromGit/Source/GLTFMaterialBaking/Private/MaterialRenderItem.cpp b/Engine/Plugins/Enterprise/GLTFExporter/MigratedFromGit/Source/GLTFMaterialBaking/Private/MaterialRenderItem.cpp
--- a/Engine/Plugins/Enterprise/GLTFExporter/MigratedFromGit/Source/GLTFMaterialBaking/Private/MaterialRenderItem.cpp
+++ b/Engine/Plugins/Enterprise/GLTFExporter/MigratedFromGit/Source/GLTFMaterialBaking/Private/MaterialRenderItem.cpp
@@ -122,6 +122,38 @@ void FMeshMaterialRenderItem::QueueMaterial(FRHICommandListImmediate& RHICmdList
 
 		DynamicMeshBuilder.GetMeshElement(FMatrix::Identity, MaterialRenderProxy, SDPG_Foreground, true, false, 0, MeshBuilderResources, MeshElement);
 
+		// NOTE: Since we can't pass these parameters during creation, we overwrite them here
+		const FPrimitiveData* PrimitiveData = MeshSettings->PrimitiveData;
+		if (PrimitiveData != nullptr)
+		{
+			FPrimitiveUniformShaderParameters PrimitiveParams = GetPrimitiveUniformShaderParameters(
+				FMatrix::Identity,
+				FMatrix::Identity,
+				PrimitiveData->ActorPosition,
+				PrimitiveData->ObjectBounds,
+				PrimitiveData->ObjectLocalBounds,
+				PrimitiveData->PreSkinnedLocalBounds,
+				false,
+				false,
+				false,
+				false,
+				false,
+				/* bDrawsVelocity = */ false,
+				GetDefaultLightingChannelMask(),
+				0,
+				INDEX_NONE,
+				INDEX_NONE,
+				false,
+				PrimitiveData->CustomPrimitiveData
+			);
+
+			// Overwrite object orientation, since original calculation is derived from LocalToWorld matrix,
+			// which currently needs to be identity matrix to avoid transforming the "2D" vertices used for baking
+			PrimitiveParams.ObjectOrientation = PrimitiveData->ObjectOrientation;
+
+			const_cast<TUniformBuffer<FPrimitiveUniformShaderParameters>*>(MeshElement.Elements[0].PrimitiveUniformBufferResource)->SetContents(PrimitiveParams);
+		}
+
 		check(MeshBuilderResources.IsValidForRendering());
 		bMeshElementDirty = false;
 	}
diff --git a/Engine/Plugins/Enterprise/GLTFExporter/MigratedFromGit/Source/GLTFMaterialBaking/Public/MaterialBakingStructures.h b/Engine/Plugins/Enterprise/GLTFExporter/MigratedFromGit/Source/GLTFMaterialBaking/Public/MaterialBakingStructures.h
--- a/Engine/Plugins/Enterprise/GLTFExporter/MigratedFromGit/Source/GLTFMaterialBaking/Public/MaterialBakingStructures.h
+++ b/Engine/Plugins/Enterprise/GLTFExporter/MigratedFromGit/Source/GLTFMaterialBaking/Public/MaterialBakingStructures.h
@@ -54,10 +54,61 @@ struct FMaterialDataEx
 	bool bTangentSpaceNormal;
 };
 
+/** Structure containing primitive information (regarding a mesh or mesh component) that is accessible through material expressions */
+struct FPrimitiveData
+{
+	FPrimitiveData(const FBoxSphereBounds& LocalBounds = FBoxSphereBounds(ForceInitToZero))
+		: ActorPosition(ForceInitToZero)
+		, ObjectOrientation(FVector::UpVector)
+		, ObjectBounds(LocalBounds)
+		, ObjectLocalBounds(LocalBounds)
+		, PreSkinnedLocalBounds(LocalBounds)
+		, CustomPrimitiveData(nullptr)
+	{}
+
+	FPrimitiveData(const UStaticMesh* StaticMesh)
+		: FPrimitiveData(StaticMesh->GetBounds())
+	{}
+
+	FPrimitiveData(const USkeletalMesh* SkeletalMesh)
+		: FPrimitiveData(SkeletalMesh->GetBounds())
+	{}
+
+	FPrimitiveData(const UMeshComponent* MeshComponent)
+		: ActorPosition(MeshComponent->GetComponentLocation())
+		, ObjectOrientation(MeshComponent->GetRenderMatrix().GetUnitAxis(EAxis::Z))
+		, ObjectBounds(MeshComponent->CalcBounds(MeshComponent->GetComponentTransform()))
+		, ObjectLocalBounds(MeshComponent->CalcLocalBounds())
+		, CustomPrimitiveData(&MeshComponent->GetCustomPrimitiveData())
+	{
+		if (const USkinnedMeshComponent* SkinnedMeshComponent = Cast<USkinnedMeshComponent>(MeshComponent))
+		{
+			SkinnedMeshComponent->GetPreSkinnedLocalBounds(PreSkinnedLocalBounds);
+		}
+		else
+		{
+			PreSkinnedLocalBounds = ObjectLocalBounds;
+		}
+	}
+
+	/** The location of the actor in world-space */
+	FVector ActorPosition;
+	/** The up vector of the object in world-space */
+	FVector ObjectOrientation;
+	/** The bounds of the object in world-space */
+	FBoxSphereBounds ObjectBounds;
+	/** The bounds of the object in local-space */
+	FBoxSphereBounds ObjectLocalBounds;
+	/** The pre-skinning bounds for the object in local-space */
+	FBoxSphereBounds PreSkinnedLocalBounds;
+	/** The custom primitive data for the mesh component */
+	const FCustomPrimitiveData* CustomPrimitiveData;
+};
+
 struct FMeshData
 {
 	FMeshData()
-		: RawMeshDescription(nullptr), Mesh(nullptr), bMirrored(false), VertexColorHash(0), TextureCoordinateIndex(0), LightMapIndex(0), LightMap(nullptr), LightmapResourceCluster(nullptr)
+		: RawMeshDescription(nullptr), Mesh(nullptr), bMirrored(false), VertexColorHash(0), TextureCoordinateIndex(0), LightMapIndex(0), LightMap(nullptr), LightmapResourceCluster(nullptr), PrimitiveData(nullptr)
 	{}
 
 	/** Ptr to raw mesh data to use for baking out the material data, if nullptr a standard quad is used */
@@ -92,6 +143,9 @@ struct FMeshData
 
 	/** Pointer to the LightmapResourceCluster to be passed on the the LightCacheInterface when baking */
 	const FLightmapResourceCluster* LightmapResourceCluster;
+
+	/** Pointer to primitive data that is accessible through material expressions, if nullptr default values are used */
+	const FPrimitiveData* PrimitiveData;
 };
 
 /** Structure containing data being processed while baking out materials*/
