[GLTFExporter] refactor: generalize support for using material baking to efficiently check UV overlap (see PR https://github.com/EpicGames/UnrealEngine/pull/8500)

commit 2b86d53896278c563ec2a3c1d466ddf76d405ea3
Date: Tue, 19 Oct 2021 22:42:06 +0200

    refactor: generalize support for using material baking to efficiently check UV overlap (see PR https://github.com/EpicGames/UnrealEngine/pull/8500)
    (cherry picked from commit 2311d3575e06a5518515e4df86b02815421308ea​)

#skipci

diff --git a/Engine/Plugins/Enterprise/GLTFExporter/MigratedFromGit/Source/GLTFMaterialBaking/Private/ExportMaterialProxy.h b/Engine/Plugins/Enterprise/GLTFExporter/MigratedFromGit/Source/GLTFMaterialBaking/Private/ExportMaterialProxy.h
--- a/Engine/Plugins/Enterprise/GLTFExporter/MigratedFromGit/Source/GLTFMaterialBaking/Private/ExportMaterialProxy.h
+++ b/Engine/Plugins/Enterprise/GLTFExporter/MigratedFromGit/Source/GLTFMaterialBaking/Private/ExportMaterialProxy.h
@@ -541,7 +541,7 @@ public:
 		}
 		return false;
 	}
-	virtual bool IsMasked() const override { return false; }
+	virtual bool IsMasked() const override { return ProxyBlendMode == BLEND_Masked; }
 	virtual enum EBlendMode GetBlendMode() const override { return ProxyBlendMode; }
 	virtual FMaterialShadingModelField GetShadingModels() const override { return MSM_DefaultLit; }
 	virtual bool IsShadingModelFromMaterialExpression() const override { return false; }
diff --git a/Engine/Plugins/Enterprise/GLTFExporter/MigratedFromGit/Source/GLTFMaterialBaking/Private/MaterialBakingHelpers.cpp b/Engine/Plugins/Enterprise/GLTFExporter/MigratedFromGit/Source/GLTFMaterialBaking/Private/MaterialBakingHelpers.cpp
--- a/Engine/Plugins/Enterprise/GLTFExporter/MigratedFromGit/Source/GLTFMaterialBaking/Private/MaterialBakingHelpers.cpp
+++ b/Engine/Plugins/Enterprise/GLTFExporter/MigratedFromGit/Source/GLTFMaterialBaking/Private/MaterialBakingHelpers.cpp
@@ -6,15 +6,13 @@
 
 namespace FMaterialBakingHelpersImpl
 {
-	static FColor BoxBlurSample(FColor* InBMP, int32 X, int32 Y, int32 InImageWidth, int32 InImageHeight)
+	static FColor BoxBlurSample(FColor* InBMP, int32 X, int32 Y, int32 InImageWidth, int32 InImageHeight, uint32 BackgroundMask)
 	{
 		int32 PixelsSampled = 0;
 		int32 CombinedColorR = 0;
 		int32 CombinedColorG = 0;
 		int32 CombinedColorB = 0;
 		int32 CombinedColorA = 0;
-	
-		const uint32 MagentaMask = FColor(255, 0, 255).DWColor();
 
 		const int32 BaseIndex = (Y * InImageWidth) + X;
 	
@@ -25,7 +23,7 @@ namespace FMaterialBakingHelpersImpl
 			{
 				const FColor SampledColor = InBMP[SampleIndex++];
 				// Check if the pixel is a rendered one (not clear color)
-				if (SampledColor.DWColor() != MagentaMask)
+				if (SampledColor.DWColor() != BackgroundMask)
 				{
 					CombinedColorR += SampledColor.R;
 					CombinedColorG += SampledColor.G;
@@ -39,17 +37,14 @@ namespace FMaterialBakingHelpersImpl
 
 		if (PixelsSampled == 0)
 		{
-			const FColor Magenta = FColor(255, 0, 255);
-			return Magenta;
+			return FColor(BackgroundMask);
 		}
 
 		return FColor(CombinedColorR / PixelsSampled, CombinedColorG / PixelsSampled, CombinedColorB / PixelsSampled, CombinedColorA / PixelsSampled);
 	}
 
-	static bool HasBorderingPixel(FColor* InBMP, int32 X, int32 Y, int32 InImageWidth, int32 InImageHeight)
+	static bool HasBorderingPixel(FColor* InBMP, int32 X, int32 Y, int32 InImageWidth, int32 InImageHeight, uint32 BackgroundMask)
 	{
-		const uint32 MagentaMask = FColor(255, 0, 255).DWColor();
-
 		const int32 BaseIndex = Y*InImageWidth + X;
 
 		int32 SampleIndex = BaseIndex - (InImageWidth + 1);
@@ -57,7 +52,7 @@ namespace FMaterialBakingHelpersImpl
 		{
 			for (int32 XI = 0; XI < 3; ++XI)
 			{
-				if (InBMP[SampleIndex++].DWColor() != MagentaMask)
+				if (InBMP[SampleIndex++].DWColor() != BackgroundMask)
 				{
 					return true;
 				}
@@ -67,7 +62,7 @@ namespace FMaterialBakingHelpersImpl
 		return false;
 	}
 
-	void PerformUVBorderSmear(TArray<FColor>& InOutPixels, int32& ImageWidth, int32& ImageHeight, int32 MaxIterations, bool bCanShrink)
+	void PerformUVBorderSmear(TArray<FColor>& InOutPixels, int32& ImageWidth, int32& ImageHeight, int32 MaxIterations, bool bCanShrink, uint32 BackgroundMask)
 	{
 		TRACE_CPUPROFILER_EVENT_SCOPE(FMaterialBakingHelpers::PerformUVBorderSmear)
 
@@ -90,23 +85,23 @@ namespace FMaterialBakingHelpersImpl
 		FColor* Current = InOutPixels.GetData();
 
 		// We can hopefully skip the entire smearing process if there is a single
-		// non-magenta color in the entire array since smearing that completely
+		// non-background color in the entire array since smearing that completely
 		// would lead to a monochome output.
 		if (MaxIterations == -1)
 		{
-			uint32 SingleColor     = MagentaMask;
+			uint32 SingleColor = BackgroundMask;
 			bool   bHasSingleColor = true;
 			for (int32 Index = 0, Num = InOutPixels.Num(); Index < Num; ++Index)
 			{
 				const uint32 Color = Current[Index].DWColor();
-				if (Color != MagentaMask)
+				if (Color != BackgroundMask)
 				{
-					if (SingleColor == MagentaMask)
+					if (SingleColor == BackgroundMask)
 					{
-						// This is the first time we stumble on a color other than magenta, keep it.
+						// This is the first time we stumble on a color other than background color, keep it.
 						SingleColor = Color;
 					}
-					// Compare with the known non-magenta color
+					// Compare with the known non-background color
 					else if (SingleColor != Color)
 					{
 						bHasSingleColor = false;
@@ -139,25 +134,25 @@ namespace FMaterialBakingHelpersImpl
 		ScratchBuffer.SetNumUninitialized(PaddedImageWidth*PaddedImageHeight);
 		FColor* Scratch = ScratchBuffer.GetData();
 
-		// Set top and bottom bordering pixels to magenta
+		// Set top and bottom bordering pixels to background color
 		for (int32 X = 0; X < PaddedImageWidth; ++X)
 		{
-			Scratch[X] = FColor(MagentaMask);
-			Scratch[(PaddedImageHeight - 1)*PaddedImageWidth + X] = FColor(MagentaMask);
+			Scratch[X] = FColor(BackgroundMask);
+			Scratch[(PaddedImageHeight - 1)*PaddedImageWidth + X] = FColor(BackgroundMask);
 		}
 
-		// Set leftg and right border pixels to magenta and copy image data into rows
+		// Set leftg and right border pixels to background color and copy image data into rows
 		for (int32 Y = 1; Y <= ImageHeight; ++Y)
 		{
 			const int32 YOffset = Y * PaddedImageWidth;
-			Scratch[YOffset] = FColor(MagentaMask);
+			Scratch[YOffset] = FColor(BackgroundMask);
 			FMemory::Memcpy(&Scratch[YOffset + 1], &Current[(Y - 1)*ImageWidth], ImageWidth * sizeof(FColor));
-			Scratch[YOffset + ImageWidth + 1] = FColor(MagentaMask);
+			Scratch[YOffset + ImageWidth + 1] = FColor(BackgroundMask);
 		}
 
 		//
-		// Find our initial workset of all rows that have magenta pixels bordering non-magenta pixels.
-		// Also find all rows that have zero magenta pixels - these will never be added to the workset.
+		// Find our initial workset of all rows that have background colored pixels bordering non-background pixels.
+		// Also find all rows that have zero background colored pixels - these will never be added to the workset.
 		//
 
 		TArray<bool> RowCompleted;
@@ -170,28 +165,28 @@ namespace FMaterialBakingHelpersImpl
 		bool bHasAnyData = false;
 		for (int32 Y = 1; Y <= ImageHeight; ++Y)
 		{
-			bool bHasMagenta = false;
-			bool bBordersNonMagenta = false;
+			bool bHasBackground = false;
+			bool bBordersNonBackground = false;
 			const int32 YOffset = Y * PaddedImageWidth;
 			for (int32 X = 1; X <= ImageWidth; ++X)
 			{
-				if (Scratch[YOffset + X].DWColor() == MagentaMask)
+				if (Scratch[YOffset + X].DWColor() == BackgroundMask)
 				{
-					bHasMagenta = true;
-					if (HasBorderingPixel(Scratch, X, Y, PaddedImageWidth, PaddedImageHeight))
+					bHasBackground = true;
+					if (HasBorderingPixel(Scratch, X, Y, PaddedImageWidth, PaddedImageHeight, BackgroundMask))
 					{
-						bBordersNonMagenta = true;
+						bBordersNonBackground = true;
 						break;
 					}
 				}
 			}
 
-			if (!bHasMagenta)
+			if (!bHasBackground)
 			{
 				bHasAnyData = true;
 				RowCompleted[Y] = true;
 			}
-			else if (bBordersNonMagenta)
+			else if (bBordersNonBackground)
 			{
 				bHasAnyData = true;
 				CurrentRowsLeft->Add(Y);
@@ -226,8 +221,8 @@ namespace FMaterialBakingHelpersImpl
 			const int32 NumThreads = FMath::Min(CurrentRowsLeft->Num(), MaxThreads);
 			const int32 LinesPerThread = FMath::CeilToInt((float)CurrentRowsLeft->Num() / (float)NumThreads);
 
-			// split up rows that still have magenta pixels amongst threads
-			ParallelFor(NumThreads, [ImageWidth, PaddedImageWidth, PaddedImageHeight, Current, Scratch, CurrentRowsLeft, LinesPerThread, &RowRemainingPixels, MagentaMask](int32 ThreadIndex)
+			// split up rows that still have background colored pixels amongst threads
+			ParallelFor(NumThreads, [ImageWidth, PaddedImageWidth, PaddedImageHeight, Current, Scratch, CurrentRowsLeft, LinesPerThread, &RowRemainingPixels, BackgroundMask](int32 ThreadIndex)
 			{
 				const int32 StartY = ThreadIndex*LinesPerThread;
 				const int32 EndY = FMath::Min((ThreadIndex + 1) * LinesPerThread, CurrentRowsLeft->Num());
@@ -241,11 +236,11 @@ namespace FMaterialBakingHelpersImpl
 					for (int32 X = 1; X <= ImageWidth; X++)
 					{
 						FColor& Color = Current[PixelIndex++];
-						if (Color.DWColor() == MagentaMask)
+						if (Color.DWColor() == BackgroundMask)
 						{
-							const FColor SampledColor = BoxBlurSample(Scratch, X, Y, PaddedImageWidth, PaddedImageHeight);
+							const FColor SampledColor = BoxBlurSample(Scratch, X, Y, PaddedImageWidth, PaddedImageHeight, BackgroundMask);
 							// If it's a valid pixel
-							if (SampledColor.DWColor() != MagentaMask)
+							if (SampledColor.DWColor() != BackgroundMask)
 							{
 								Color = SampledColor;
 							}
@@ -311,7 +306,7 @@ namespace FMaterialBakingHelpersImpl
 		{
 			for (int32 i = 0; i < InOutPixels.Num(); ++i)
 			{
-				if (InOutPixels[i].DWColor() == MagentaMask)
+				if (InOutPixels[i].DWColor() == BackgroundMask)
 				{
 					InOutPixels[i] = FColor::Black;
 				}
@@ -320,12 +315,14 @@ namespace FMaterialBakingHelpersImpl
 	}
 }
 
-void FMaterialBakingHelpers::PerformUVBorderSmear(TArray<FColor>& InOutPixels, int32 ImageWidth, int32 ImageHeight, int32 MaxIterations)
+void FMaterialBakingHelpers::PerformUVBorderSmear(TArray<FColor>& InOutPixels, int32 ImageWidth, int32 ImageHeight, int32 MaxIterations, FColor BackgroundColor)
 {
-	FMaterialBakingHelpersImpl::PerformUVBorderSmear(InOutPixels, ImageWidth, ImageHeight, MaxIterations, false);
+	const uint32 BackgroundMask = BackgroundColor.DWColor();
+	FMaterialBakingHelpersImpl::PerformUVBorderSmear(InOutPixels, ImageWidth, ImageHeight, MaxIterations, false, BackgroundMask);
 }
 
-void FMaterialBakingHelpers::PerformUVBorderSmearAndShrink(TArray<FColor>& InOutPixels, int32& InOutImageWidth, int32& InOutImageHeight)
+void FMaterialBakingHelpers::PerformUVBorderSmearAndShrink(TArray<FColor>& InOutPixels, int32& InOutImageWidth, int32& InOutImageHeight, FColor BackgroundColor)
 {
-	FMaterialBakingHelpersImpl::PerformUVBorderSmear(InOutPixels, InOutImageWidth, InOutImageHeight, -1, true);
+	const uint32 BackgroundMask = BackgroundColor.DWColor();
+	FMaterialBakingHelpersImpl::PerformUVBorderSmear(InOutPixels, InOutImageWidth, InOutImageHeight, -1, true, BackgroundMask);
 }
\ No newline at end of file
diff --git a/Engine/Plugins/Enterprise/GLTFExporter/MigratedFromGit/Source/GLTFMaterialBaking/Private/MaterialBakingModule.cpp b/Engine/Plugins/Enterprise/GLTFExporter/MigratedFromGit/Source/GLTFMaterialBaking/Private/MaterialBakingModule.cpp
--- a/Engine/Plugins/Enterprise/GLTFExporter/MigratedFromGit/Source/GLTFMaterialBaking/Private/MaterialBakingModule.cpp
+++ b/Engine/Plugins/Enterprise/GLTFExporter/MigratedFromGit/Source/GLTFMaterialBaking/Private/MaterialBakingModule.cpp
@@ -292,6 +292,7 @@ void FMaterialBakingModule::BakeMaterials(const TArray<FMaterialData*>& Material
 		MaterialDataEx.bPerformBorderSmear = MaterialData->bPerformBorderSmear;
 		MaterialDataEx.BlendMode = MaterialData->BlendMode;
 		MaterialDataEx.bTangentSpaceNormal = MaterialData->bTangentSpaceNormal;
+		MaterialDataEx.BackgroundColor = MaterialData->BackgroundColor;
 
 		for (const TPair<EMaterialProperty, FIntPoint>& PropertySizePair : MaterialData->PropertySizes)
 		{
@@ -553,7 +554,7 @@ void FMaterialBakingModule::BakeMaterials(const TArray<FMaterialDataEx*>& Materi
 
 				// It is safe to reuse the same render target for each draw pass since they all execute sequentially on the GPU and are copied to staging buffers before
 				// being reused.
-				UTextureRenderTarget2D* RenderTarget = CreateRenderTarget((ColorSpace == EPropertyColorSpace::Linear), PixelFormat, CurrentOutput.PropertySizes[Property]);
+				UTextureRenderTarget2D* RenderTarget = CreateRenderTarget((ColorSpace == EPropertyColorSpace::Linear), PixelFormat, CurrentOutput.PropertySizes[Property], CurrentMaterialSettings->BackgroundColor);
 				if (RenderTarget != nullptr)
 				{
 					// Perform everything left of the operation directly on the render thread since we need to modify some RenderItem's properties
@@ -611,7 +612,7 @@ void FMaterialBakingModule::BakeMaterials(const TArray<FMaterialDataEx*>& Materi
 									if (Property.Type == MP_EmissiveColor)
 									{
 										// Only one thread will write to CurrentOutput.EmissiveScale since there can be only one emissive channel property per FBakeOutputEx
-										FMaterialBakingModule::ProcessEmissiveOutput((const FFloat16Color*)Data, DataWidth, OutputSize, OutputColor, CurrentOutput.EmissiveScale);
+										FMaterialBakingModule::ProcessEmissiveOutput((const FFloat16Color*)Data, DataWidth, OutputSize, OutputColor, CurrentOutput.EmissiveScale, CurrentMaterialSettings->BackgroundColor);
 
 										if (bEmissiveHDR)
 										{
@@ -634,7 +635,7 @@ void FMaterialBakingModule::BakeMaterials(const TArray<FMaterialDataEx*>& Materi
 									if (CurrentMaterialSettings->bPerformBorderSmear)
 									{
 										// This will resize the output to a single pixel if the result is monochrome.
-										FMaterialBakingHelpers::PerformUVBorderSmearAndShrink(OutputColor, OutputSize.X, OutputSize.Y);
+										FMaterialBakingHelpers::PerformUVBorderSmearAndShrink(OutputColor, OutputSize.X, OutputSize.Y, CurrentMaterialSettings->BackgroundColor);
 									}
 #if WITH_EDITOR
 									// If saving intermediates is turned on
@@ -782,7 +783,7 @@ void FMaterialBakingModule::CleanupMaterialProxies()
 	MaterialProxyPool.Reset();
 }
 
-UTextureRenderTarget2D* FMaterialBakingModule::CreateRenderTarget(bool bInForceLinearGamma, EPixelFormat InPixelFormat, const FIntPoint& InTargetSize)
+UTextureRenderTarget2D* FMaterialBakingModule::CreateRenderTarget(bool bInForceLinearGamma, EPixelFormat InPixelFormat, const FIntPoint& InTargetSize, const FColor& BackgroundColor)
 {
 	TRACE_CPUPROFILER_EVENT_SCOPE(FMaterialBakingModule::CreateRenderTarget)
 
@@ -809,8 +810,7 @@ UTextureRenderTarget2D* FMaterialBakingModule::CreateRenderTarget(bool bInForceL
 		RenderTarget = NewObject<UTextureRenderTarget2D>();
 		check(RenderTarget);
 		RenderTarget->AddToRoot();
-		RenderTarget->ClearColor = FLinearColor(1.0f, 0.0f, 1.0f);
-		RenderTarget->ClearColor.A = 1.0f;
+		RenderTarget->ClearColor = bInForceLinearGamma ? BackgroundColor.ReinterpretAsLinear() : FLinearColor(BackgroundColor);
 		RenderTarget->TargetGamma = 0.0f;
 		RenderTarget->InitCustomFormat(ClampedTargetSize.X, ClampedTargetSize.Y, InPixelFormat, bInForceLinearGamma);
 
@@ -851,7 +851,7 @@ FExportMaterialProxy* FMaterialBakingModule::CreateMaterialProxy(const FMaterial
 	return Proxy;
 }
 
-void FMaterialBakingModule::ProcessEmissiveOutput(const FFloat16Color* Color16, int32 Color16Pitch, const FIntPoint& OutputSize, TArray<FColor>& OutputColor, float& EmissiveScale)
+void FMaterialBakingModule::ProcessEmissiveOutput(const FFloat16Color* Color16, int32 Color16Pitch, const FIntPoint& OutputSize, TArray<FColor>& OutputColor, float& EmissiveScale, const FColor& BackgroundColor)
 {
 	TRACE_CPUPROFILER_EVENT_SCOPE(FMaterialBakingModule::ProcessEmissiveOutput)
 
@@ -863,13 +863,13 @@ void FMaterialBakingModule::ProcessEmissiveOutput(const FFloat16Color* Color16,
 	float* MaxValue = new float[NumThreads];
 	FMemory::Memset(MaxValue, 0, NumThreads * sizeof(MaxValue[0]));
 	const int32 LinesPerThread = FMath::CeilToInt((float)OutputSize.Y / (float)NumThreads);
+	const FFloat16Color BackgroundFloat16 = FFloat16Color(FLinearColor(BackgroundColor)); // Can assume emissive always uses sRGB
 
 	// Find maximum float value across texture
-	ParallelFor(NumThreads, [&Color16, LinesPerThread, MaxValue, OutputSize, Color16Pitch](int32 Index)
+	ParallelFor(NumThreads, [&Color16, LinesPerThread, MaxValue, OutputSize, Color16Pitch, BackgroundFloat16](int32 Index)
 	{
 		const int32 EndY = FMath::Min((Index + 1) * LinesPerThread, OutputSize.Y);			
 		float& CurrentMaxValue = MaxValue[Index];
-		const FFloat16Color MagentaFloat16 = FFloat16Color(FLinearColor(1.0f, 0.0f, 1.0f));
 		for (int32 PixelY = Index * LinesPerThread; PixelY < EndY; ++PixelY)
 		{
 			const int32 SrcYOffset = PixelY * Color16Pitch;
@@ -877,7 +877,7 @@ void FMaterialBakingModule::ProcessEmissiveOutput(const FFloat16Color* Color16,
 			{
 				const FFloat16Color& Pixel16 = Color16[PixelX + SrcYOffset];
 				// Find maximum channel value across texture
-				if (!(Pixel16 == MagentaFloat16))
+				if (!(Pixel16 == BackgroundFloat16))
 				{
 					CurrentMaxValue = FMath::Max(CurrentMaxValue, FMath::Max3(Pixel16.R.GetFloat(), Pixel16.G.GetFloat(), Pixel16.B.GetFloat()));
 				}
@@ -904,10 +904,8 @@ void FMaterialBakingModule::ProcessEmissiveOutput(const FFloat16Color* Color16,
 	// Now convert Float16 to Color using the scale
 	OutputColor.SetNumUninitialized(OutputSize.X * OutputSize.Y);
 	const float Scale = 255.0f / GlobalMaxValue;
-	ParallelFor(NumThreads, [&Color16, LinesPerThread, &OutputColor, OutputSize, Color16Pitch, Scale](int32 Index)
+	ParallelFor(NumThreads, [&Color16, LinesPerThread, &OutputColor, OutputSize, Color16Pitch, Scale, BackgroundFloat16, BackgroundColor](int32 Index)
 	{
-		const FFloat16Color MagentaFloat16 = FFloat16Color(FLinearColor(1.0f, 0.0f, 1.0f));
-
 		const int32 EndY = FMath::Min((Index + 1) * LinesPerThread, OutputSize.Y);
 		for (int32 PixelY = Index * LinesPerThread; PixelY < EndY; ++PixelY)
 		{
@@ -919,11 +917,9 @@ void FMaterialBakingModule::ProcessEmissiveOutput(const FFloat16Color* Color16,
 				const FFloat16Color& Pixel16 = Color16[PixelX + SrcYOffset];
 				FColor& Pixel8 = OutputColor[PixelX + DstYOffset];
 
-				if (Pixel16 == MagentaFloat16)
+				if (Pixel16 == BackgroundFloat16)
 				{
-					Pixel8.R = 255;
-					Pixel8.G = 0;
-					Pixel8.B = 255;
+					Pixel8 = BackgroundColor;
 				}
 				else
 				{
diff --git a/Engine/Plugins/Enterprise/GLTFExporter/MigratedFromGit/Source/GLTFMaterialBaking/Private/MaterialBakingModule.h b/Engine/Plugins/Enterprise/GLTFExporter/MigratedFromGit/Source/GLTFMaterialBaking/Private/MaterialBakingModule.h
--- a/Engine/Plugins/Enterprise/GLTFExporter/MigratedFromGit/Source/GLTFMaterialBaking/Private/MaterialBakingModule.h
+++ b/Engine/Plugins/Enterprise/GLTFExporter/MigratedFromGit/Source/GLTFMaterialBaking/Private/MaterialBakingModule.h
@@ -41,13 +41,13 @@ public:
 
 protected:
 	/* Creates and adds or reuses a RenderTarget from the pool */
-	UTextureRenderTarget2D* CreateRenderTarget(bool bInForceLinearGamma, EPixelFormat InPixelFormat, const FIntPoint& InTargetSize);
+	UTextureRenderTarget2D* CreateRenderTarget(bool bInForceLinearGamma, EPixelFormat InPixelFormat, const FIntPoint& InTargetSize, const FColor& BackgroundColor);
 
 	/* Creates and adds (or reuses a ExportMaterialProxy from the pool if MaterialBaking.UseMaterialProxyCaching is set to 1) */
 	FExportMaterialProxy* CreateMaterialProxy(const FMaterialDataEx* MaterialSettings, const FMaterialPropertyEx& Property);
 
 	/** Helper for emissive color conversion to Output */
-	static void ProcessEmissiveOutput(const FFloat16Color* Color16, int32 Color16Pitch, const FIntPoint& OutputSize, TArray<FColor>& Output, float& EmissiveScale);
+	static void ProcessEmissiveOutput(const FFloat16Color* Color16, int32 Color16Pitch, const FIntPoint& OutputSize, TArray<FColor>& Output, float& EmissiveScale, const FColor& BackgroundColor);
 
 	/** Cleans up all cached material proxies in MaterialProxyPool */
 	void CleanupMaterialProxies();
diff --git a/Engine/Plugins/Enterprise/GLTFExporter/MigratedFromGit/Source/GLTFMaterialBaking/Public/MaterialBakingHelpers.h b/Engine/Plugins/Enterprise/GLTFExporter/MigratedFromGit/Source/GLTFMaterialBaking/Public/MaterialBakingHelpers.h
--- a/Engine/Plugins/Enterprise/GLTFExporter/MigratedFromGit/Source/GLTFMaterialBaking/Public/MaterialBakingHelpers.h
+++ b/Engine/Plugins/Enterprise/GLTFExporter/MigratedFromGit/Source/GLTFMaterialBaking/Public/MaterialBakingHelpers.h
@@ -10,21 +10,23 @@ class MATERIALBAKING_API FMaterialBakingHelpers
 {
 public:
 	/** 
-	 * Applies a box blur to magenta pixels found in given texture represented by InOutPixels using non-magenta pixels, this creates a smear across the magenta/filled pixels.
+	 * Applies a box blur to magenta pixels found in given texture represented by InOutPixels using non-background pixels, this creates a smear across the magenta/filled pixels.
 	 * @param  InOutPixels     The image to apply the smear to
 	 * @param  ImageWidth      The width of the image in pixels
 	 * @param  ImageHeight     The height of the image in pixels
 	 * @param  MaxIterations   The max distance in pixels to smear the edges of the texture. When set to -1 this value is the max of the width and height of the buffer
+	 * @param  BackgroundColor The pixel color considered background in the image, and that can be smeared into
 	 */
-	static void PerformUVBorderSmear(TArray<FColor>& InOutPixels, int32 ImageWidth, int32 ImageHeight, int32 MaxIterations = -1);
+	static void PerformUVBorderSmear(TArray<FColor>& InOutPixels, int32 ImageWidth, int32 ImageHeight, int32 MaxIterations = -1, FColor BackgroundColor = FColor::Magenta);
 
 	/** 
-	 * Applies a box blur to magenta pixels found in given texture represented by InOutPixels using non-magenta pixels, this creates a smear across the magenta/filled pixels.
+	 * Applies a box blur to magenta pixels found in given texture represented by InOutPixels using non-background pixels, this creates a smear across the magenta/filled pixels.
 	 * @param  InOutPixels      The image to apply the smear to
 	 * @param  InOutImageWidth  The width of the image in pixels
 	 * @param  InOutImageHeight The height of the image in pixels
+	 * @param  BackgroundColor  The pixel color considered background in the image, and that can be smeared into
 	 *
 	 * @note   This version is allowed to shrink the array to a single pixel if the entire output is monochrome. Params will be adjusted accordingly when this occurs.
 	 */
-	static void PerformUVBorderSmearAndShrink(TArray<FColor>& InOutPixels, int32& InOutImageWidth, int32& InOutImageHeight);
+	static void PerformUVBorderSmearAndShrink(TArray<FColor>& InOutPixels, int32& InOutImageWidth, int32& InOutImageHeight, FColor BackgroundColor = FColor::Magenta);
 };
\ No newline at end of file
diff --git a/Engine/Plugins/Enterprise/GLTFExporter/MigratedFromGit/Source/GLTFMaterialBaking/Public/MaterialBakingStructures.h b/Engine/Plugins/Enterprise/GLTFExporter/MigratedFromGit/Source/GLTFMaterialBaking/Public/MaterialBakingStructures.h
--- a/Engine/Plugins/Enterprise/GLTFExporter/MigratedFromGit/Source/GLTFMaterialBaking/Public/MaterialBakingStructures.h
+++ b/Engine/Plugins/Enterprise/GLTFExporter/MigratedFromGit/Source/GLTFMaterialBaking/Public/MaterialBakingStructures.h
@@ -17,6 +17,7 @@ struct FMaterialData
 		: Material(nullptr)
 		, bPerformBorderSmear(true)
 		, BlendMode(BLEND_Opaque)
+		, BackgroundColor(FColor::Magenta)
 		, bTangentSpaceNormal(false)
 	{}
 
@@ -26,8 +27,10 @@ struct FMaterialData
 	TMap<EMaterialProperty, FIntPoint> PropertySizes;
 	/** Whether to smear borders after baking */
 	bool bPerformBorderSmear;
-	/** Blend mode to use when baking */
+	/** Blend mode to use when baking, allowing for example detection of overlapping UVs */
 	EBlendMode BlendMode;
+	/** Background color used to initially fill the output texture and used for border smear */
+	FColor BackgroundColor;
 	/** Whether to transform normals from world-space to tangent-space (does nothing if material already uses tangent-space normals) */
 	bool bTangentSpaceNormal;
 };
@@ -39,6 +42,7 @@ struct FMaterialDataEx
 		: Material(nullptr)
 		, bPerformBorderSmear(true)
 		, BlendMode(BLEND_Opaque)
+		, BackgroundColor(FColor::Magenta)
 		, bTangentSpaceNormal(false)
 	{}
 
@@ -48,8 +52,10 @@ struct FMaterialDataEx
 	TMap<FMaterialPropertyEx, FIntPoint> PropertySizes;
 	/** Whether to smear borders after baking */
 	bool bPerformBorderSmear;
-	/** Blend mode to use when baking */
+	/** Blend mode to use when baking, allowing for example detection of overlapping UVs */
 	EBlendMode BlendMode;
+	/** Background color used to initially fill the output texture and used for border smear */
+	FColor BackgroundColor;
 	/** Whether to transform normals from world-space to tangent-space (does nothing if material already uses tangent-space normals) */
 	bool bTangentSpaceNormal;
 };
