[GLTFExporter] feat: improve material analyzer to recognize primitive data material expressions and flagging material as needing mesh data

commit bdddca04687c112d524ab130f7a0f5da92069aa0
Date: Wed, 20 Oct 2021 01:49:18 +0200

    feat: improve material analyzer to recognize primitive data material expressions and flagging material as needing mesh data
    (cherry picked from commit aa637657461b3bc7d6c08abc98b3372188c17d5b​)

#skipci

diff --git a/Engine/Plugins/Enterprise/GLTFExporter/MigratedFromGit/Source/GLTFExporter/Private/Converters/GLTFMaterialUtility.cpp b/Engine/Plugins/Enterprise/GLTFExporter/MigratedFromGit/Source/GLTFExporter/Private/Converters/GLTFMaterialUtility.cpp
--- a/Engine/Plugins/Enterprise/GLTFExporter/MigratedFromGit/Source/GLTFExporter/Private/Converters/GLTFMaterialUtility.cpp
+++ b/Engine/Plugins/Enterprise/GLTFExporter/MigratedFromGit/Source/GLTFExporter/Private/Converters/GLTFMaterialUtility.cpp
@@ -398,16 +398,19 @@ bool FGLTFMaterialUtility::NeedsMeshData(const UMaterialInterface* Material)
 		TEXT("ClearCoatBottomNormal"),
 	};
 
-	bool bRequiresVertexData = false;
+	bool bNeedsMeshData = false;
 	FGLTFMaterialAnalysis Analysis;
 
+	// TODO: optimize baking by seperating need for vertex data and primitive data
+
 	for (const FMaterialPropertyEx& Property: Properties)
 	{
 		AnalyzeMaterialProperty(Material, Property, Analysis);
-		bRequiresVertexData |= Analysis.bRequiresVertexData;
+		bNeedsMeshData |= Analysis.bRequiresVertexData;
+		bNeedsMeshData |= Analysis.bRequiresPrimitiveData;
 	}
 
-	return bRequiresVertexData;
+	return bNeedsMeshData;
 }
 
 bool FGLTFMaterialUtility::NeedsMeshData(const TArray<const UMaterialInterface*>& Materials)
diff --git a/Engine/Plugins/Enterprise/GLTFExporter/MigratedFromGit/Source/GLTFMaterialAnalyzer/Private/GLTFMaterialAnalyzer.cpp b/Engine/Plugins/Enterprise/GLTFExporter/MigratedFromGit/Source/GLTFMaterialAnalyzer/Private/GLTFMaterialAnalyzer.cpp
--- a/Engine/Plugins/Enterprise/GLTFExporter/MigratedFromGit/Source/GLTFMaterialAnalyzer/Private/GLTFMaterialAnalyzer.cpp
+++ b/Engine/Plugins/Enterprise/GLTFExporter/MigratedFromGit/Source/GLTFMaterialAnalyzer/Private/GLTFMaterialAnalyzer.cpp
@@ -120,6 +120,16 @@ int32 UGLTFMaterialAnalyzer::CompilePropertyEx(FMaterialCompiler* Compiler, cons
 		TranslatorHack->bUsesAOMaterialMask ||
 		TranslatorHack->bUsesVertexPosition;
 
+	Analysis->bRequiresPrimitiveData =
+		ProxyCompiler.bUsesActorPosition ||
+		ProxyCompiler.bUsesObjectPosition ||
+		ProxyCompiler.bUsesObjectOrientation ||
+		ProxyCompiler.bUsesObjectRadius ||
+		ProxyCompiler.bUsesObjectBounds ||
+		ProxyCompiler.bUsesObjectLocalBounds ||
+		ProxyCompiler.bUsesPreSkinnedLocalBounds ||
+		ProxyCompiler.bUsesCustomPrimitiveData;
+
 	return Result;
 }
 
diff --git a/Engine/Plugins/Enterprise/GLTFExporter/MigratedFromGit/Source/GLTFMaterialAnalyzer/Private/GLTFProxyMaterialCompiler.h b/Engine/Plugins/Enterprise/GLTFExporter/MigratedFromGit/Source/GLTFMaterialAnalyzer/Private/GLTFProxyMaterialCompiler.h
--- a/Engine/Plugins/Enterprise/GLTFExporter/MigratedFromGit/Source/GLTFMaterialAnalyzer/Private/GLTFProxyMaterialCompiler.h
+++ b/Engine/Plugins/Enterprise/GLTFExporter/MigratedFromGit/Source/GLTFMaterialAnalyzer/Private/GLTFProxyMaterialCompiler.h
@@ -6,8 +6,61 @@
 
 class FGLTFProxyMaterialCompiler : public FProxyMaterialCompiler
 {
+public:
+
+	bool bUsesActorPosition = false;
+	bool bUsesObjectPosition = false;
+	bool bUsesObjectOrientation = false;
+	bool bUsesObjectRadius = false;
+	bool bUsesObjectBounds = false;
+	bool bUsesObjectLocalBounds = false; // TODO: add overridable compiler method for ObjectLocalBounds in engine
+	bool bUsesPreSkinnedLocalBounds = false;
+	bool bUsesCustomPrimitiveData = false;
+
 	using FProxyMaterialCompiler::FProxyMaterialCompiler;
 
+	virtual int32 ActorWorldPosition() override
+	{
+		bUsesActorPosition = true;
+		return Compiler->ActorWorldPosition();
+	}
+
+	virtual int32 ObjectWorldPosition() override
+	{
+		bUsesObjectPosition = true;
+		return Compiler->ObjectWorldPosition();
+	}
+
+	virtual int32 ObjectOrientation() override
+	{
+		bUsesObjectOrientation = true;
+		return Compiler->ObjectOrientation();
+	}
+
+	virtual int32 ObjectRadius() override
+	{
+		bUsesObjectRadius = true;
+		return Compiler->ObjectRadius();
+	}
+
+	virtual int32 ObjectBounds() override
+	{
+		bUsesObjectBounds = true;
+		return Compiler->ObjectBounds();
+	}
+
+	virtual int32 PreSkinnedLocalBounds(int32 OutputIndex) override
+	{
+		bUsesPreSkinnedLocalBounds = true;
+		return Compiler->PreSkinnedLocalBounds(OutputIndex);
+	}
+	
+	virtual int32 CustomPrimitiveData(int32 OutputIndex, EMaterialValueType Type) override
+	{
+		bUsesCustomPrimitiveData = true;
+		return Compiler->CustomPrimitiveData(OutputIndex, Type);
+	}
+
 	virtual int32 GIReplace(int32 Direct, int32 StaticIndirect, int32 DynamicIndirect) override
 	{
 		return Direct; // ignore all non-default branches in this runtime switch
diff --git a/Engine/Plugins/Enterprise/GLTFExporter/MigratedFromGit/Source/GLTFMaterialAnalyzer/Public/GLTFMaterialAnalysis.h b/Engine/Plugins/Enterprise/GLTFExporter/MigratedFromGit/Source/GLTFMaterialAnalyzer/Public/GLTFMaterialAnalysis.h
--- a/Engine/Plugins/Enterprise/GLTFExporter/MigratedFromGit/Source/GLTFMaterialAnalyzer/Public/GLTFMaterialAnalysis.h
+++ b/Engine/Plugins/Enterprise/GLTFExporter/MigratedFromGit/Source/GLTFMaterialAnalyzer/Public/GLTFMaterialAnalysis.h
@@ -8,6 +8,7 @@ struct FGLTFMaterialAnalysis
 {
 	FGLTFMaterialAnalysis()
 		: bRequiresVertexData(false)
+		, bRequiresPrimitiveData(false)
 	{
 	}
 
@@ -22,4 +23,7 @@ struct FGLTFMaterialAnalysis
 
 	/** True if this material reads any vertex data */
 	bool bRequiresVertexData;
+
+	/** True if this material reads any primitive data */
+	bool bRequiresPrimitiveData;
 };
