[GLTFExporter] fix: add support for recognizing custom expression that uses ObjectLocalBounds in material analyzer

commit 0170fc922fecd9d249e1058f83788d4c8b2c4c1c
Date: Wed, 20 Oct 2021 03:17:55 +0200

    fix: add support for recognizing custom expression that uses ObjectLocalBounds in material analyzer
    (cherry picked from commit 3e99b0c29116cb3a22c0968cd006962bcf51ca2c​)

#skipci

diff --git a/Engine/Plugins/Enterprise/GLTFExporter/MigratedFromGit/Source/GLTFMaterialAnalyzer/Private/GLTFProxyMaterialCompiler.h b/Engine/Plugins/Enterprise/GLTFExporter/MigratedFromGit/Source/GLTFMaterialAnalyzer/Private/GLTFProxyMaterialCompiler.h
--- a/Engine/Plugins/Enterprise/GLTFExporter/MigratedFromGit/Source/GLTFMaterialAnalyzer/Private/GLTFProxyMaterialCompiler.h
+++ b/Engine/Plugins/Enterprise/GLTFExporter/MigratedFromGit/Source/GLTFMaterialAnalyzer/Private/GLTFProxyMaterialCompiler.h
@@ -3,6 +3,7 @@
 #pragma once
 
 #include "MaterialCompiler.h"
+#include "Materials/MaterialExpressionCustom.h"
 
 class FGLTFProxyMaterialCompiler : public FProxyMaterialCompiler
 {
@@ -13,7 +14,7 @@ public:
 	bool bUsesObjectOrientation = false;
 	bool bUsesObjectRadius = false;
 	bool bUsesObjectBounds = false;
-	bool bUsesObjectLocalBounds = false; // TODO: add overridable compiler method for ObjectLocalBounds in engine
+	bool bUsesObjectLocalBounds = false;
 	bool bUsesPreSkinnedLocalBounds = false;
 	bool bUsesCustomPrimitiveData = false;
 
@@ -49,12 +50,31 @@ public:
 		return Compiler->ObjectBounds();
 	}
 
+#if (ENGINE_MAJOR_VERSION > 4 || ENGINE_MINOR_VERSION >= 26)
+	virtual int32 CustomExpression(class UMaterialExpressionCustom* Custom, int32 OutputIndex, TArray<int32>& CompiledInputs) override
+#else
+	virtual int32 CustomExpression(class UMaterialExpressionCustom* Custom, TArray<int32>& CompiledInputs) override 
+#endif
+	{
+		// NOTE: because there is no overridable compiler method for ObjectLocalBounds in engine this is our only option currently
+		if (Custom->Code.Contains(TEXT("GetPrimitiveData(Parameters).LocalObjectBounds")))
+		{
+			bUsesObjectLocalBounds = true;
+		}
+
+#if (ENGINE_MAJOR_VERSION > 4 || ENGINE_MINOR_VERSION >= 26)
+		return Compiler->CustomExpression(Custom, OutputIndex, CompiledInputs);
+#else
+		return Compiler->CustomExpression(Custom, CompiledInputs);
+#endif
+	}
+
 	virtual int32 PreSkinnedLocalBounds(int32 OutputIndex) override
 	{
 		bUsesPreSkinnedLocalBounds = true;
 		return Compiler->PreSkinnedLocalBounds(OutputIndex);
 	}
-	
+
 	virtual int32 CustomPrimitiveData(int32 OutputIndex, EMaterialValueType Type) override
 	{
 		bUsesCustomPrimitiveData = true;
