[GLTFExporter] fix: add missing data to bake expressions like PrecomputedAOMask

commit 86dc282fc4d84a5215cc3504e81708cd94c4edaf
Date: Thu, 21 Oct 2021 17:58:00 +0200

    fix: add missing data to bake expressions like PrecomputedAOMask
    (cherry picked from commit 8a67814f5618ee444dc2de5bb4e6dd6174384dbf​)

#skipci

diff --git a/Engine/Plugins/Enterprise/GLTFExporter/MigratedFromGit/Source/GLTFExporter/Private/Converters/GLTFMaterialConverters.cpp b/Engine/Plugins/Enterprise/GLTFExporter/MigratedFromGit/Source/GLTFExporter/Private/Converters/GLTFMaterialConverters.cpp
--- a/Engine/Plugins/Enterprise/GLTFExporter/MigratedFromGit/Source/GLTFExporter/Private/Converters/GLTFMaterialConverters.cpp
+++ b/Engine/Plugins/Enterprise/GLTFExporter/MigratedFromGit/Source/GLTFExporter/Private/Converters/GLTFMaterialConverters.cpp
@@ -19,7 +19,7 @@ void FGLTFMaterialConverter::Sanitize(const UMaterialInterface*& Material, const
 	{
 		const FMeshDescription& MeshDescription = MeshData->GetParent()->Description;
 
-		const float DegenerateUVPercentage = UVDegenerateChecker.GetOrAdd(&MeshDescription, SectionIndices, MeshData->TexCoord);
+		const float DegenerateUVPercentage = UVDegenerateChecker.GetOrAdd(&MeshDescription, SectionIndices, MeshData->BakeUsingTexCoord);
 		if (FMath::IsNearlyEqual(DegenerateUVPercentage, 1))
 		{
 			FString SectionString = TEXT("mesh section");
@@ -30,7 +30,7 @@ void FGLTFMaterialConverter::Sanitize(const UMaterialInterface*& Material, const
 				TEXT("Material %s uses mesh data from %s but the lightmap UVs (channel %d) are nearly 100%% degenerate (in %s). Simple baking will be used as fallback"),
 				*Material->GetName(),
 				*MeshData->GetParent()->Name,
-				MeshData->TexCoord,
+				MeshData->BakeUsingTexCoord,
 				*SectionString));
 
 			MeshData = nullptr;
diff --git a/Engine/Plugins/Enterprise/GLTFExporter/MigratedFromGit/Source/GLTFExporter/Private/Converters/GLTFMaterialUtility.cpp b/Engine/Plugins/Enterprise/GLTFExporter/MigratedFromGit/Source/GLTFExporter/Private/Converters/GLTFMaterialUtility.cpp
--- a/Engine/Plugins/Enterprise/GLTFExporter/MigratedFromGit/Source/GLTFExporter/Private/Converters/GLTFMaterialUtility.cpp
+++ b/Engine/Plugins/Enterprise/GLTFExporter/MigratedFromGit/Source/GLTFExporter/Private/Converters/GLTFMaterialUtility.cpp
@@ -165,8 +165,11 @@ FGLTFPropertyBakeOutput FGLTFMaterialUtility::BakeMaterialProperty(const FIntPoi
 	MeshSet.MaterialIndices = MeshSectionIndices; // NOTE: MaterialIndices is actually section indices
 	if (MeshData != nullptr)
 	{
-		MeshSet.PrimitiveData = &MeshData->PrimitiveData;
 		MeshSet.RawMeshDescription = const_cast<FMeshDescription*>(&MeshData->Description);
+		MeshSet.LightMap = MeshData->LightMap;
+		MeshSet.LightMapIndex = MeshData->LightMapTexCoord;
+		MeshSet.LightmapResourceCluster = MeshData->LightMapResourceCluster;
+		MeshSet.PrimitiveData = &MeshData->PrimitiveData;
 	}
 
 	FMaterialDataEx MatSet;
diff --git a/Engine/Plugins/Enterprise/GLTFExporter/MigratedFromGit/Source/GLTFExporter/Private/Converters/GLTFMeshData.cpp b/Engine/Plugins/Enterprise/GLTFExporter/MigratedFromGit/Source/GLTFExporter/Private/Converters/GLTFMeshData.cpp
--- a/Engine/Plugins/Enterprise/GLTFExporter/MigratedFromGit/Source/GLTFExporter/Private/Converters/GLTFMeshData.cpp
+++ b/Engine/Plugins/Enterprise/GLTFExporter/MigratedFromGit/Source/GLTFExporter/Private/Converters/GLTFMeshData.cpp
@@ -4,11 +4,12 @@
 #include "Converters/GLTFNameUtility.h"
 #include "Converters/GLTFMeshUtility.h"
 #include "StaticMeshAttributes.h"
+#include "Engine/MapBuildDataRegistry.h"
 #include "Developer/MeshMergeUtilities/Private/MeshMergeHelpers.h"
 #include "Rendering/SkeletalMeshRenderData.h"
 
 FGLTFMeshData::FGLTFMeshData(const UStaticMesh* StaticMesh, const UStaticMeshComponent* StaticMeshComponent, int32 LODIndex)
-	: Parent(nullptr)
+	: LightMap(nullptr), LightMapResourceCluster(nullptr), LightMapTexCoord(0), Parent(nullptr)
 {
 	FStaticMeshAttributes(Description).Register();
 
@@ -17,6 +18,17 @@ FGLTFMeshData::FGLTFMeshData(const UStaticMesh* StaticMesh, const UStaticMeshCom
 		Name = FGLTFNameUtility::GetName(StaticMeshComponent);
 		PrimitiveData = { StaticMeshComponent };
 		FMeshMergeHelpers::RetrieveMesh(StaticMeshComponent, LODIndex, Description, true);
+
+		if (StaticMeshComponent->LODData.IsValidIndex(LODIndex))
+		{
+			const FStaticMeshComponentLODInfo& LODData = StaticMeshComponent->LODData[LODIndex];
+			const FMeshMapBuildData* BuildData = StaticMeshComponent->GetMeshMapBuildData(LODData);
+			if (BuildData != nullptr)
+			{
+				LightMap = BuildData->LightMap;
+				LightMapResourceCluster = BuildData->ResourceCluster;
+			}
+		}
 	}
 	else
 	{
@@ -25,14 +37,15 @@ FGLTFMeshData::FGLTFMeshData(const UStaticMesh* StaticMesh, const UStaticMeshCom
 		FMeshMergeHelpers::RetrieveMesh(StaticMesh, LODIndex, Description);
 	}
 
-	const int32 NumTexCoords = StaticMesh->GetLODForExport(LODIndex).VertexBuffers.StaticMeshVertexBuffer.GetNumTexCoords();
-	TexCoord = FMath::Min(StaticMesh->GetLightMapCoordinateIndex(), NumTexCoords - 1);
+	LightMapTexCoord = StaticMesh->GetLightMapCoordinateIndex();
+ 	const int32 NumTexCoords = StaticMesh->GetLODForExport(LODIndex).VertexBuffers.StaticMeshVertexBuffer.GetNumTexCoords();
+	BakeUsingTexCoord = FMath::Min(LightMapTexCoord, NumTexCoords - 1);
 
 	// TODO: add warning if texture coordinate has overlap
 }
 
 FGLTFMeshData::FGLTFMeshData(const USkeletalMesh* SkeletalMesh, const USkeletalMeshComponent* SkeletalMeshComponent, int32 LODIndex)
-	: Parent(nullptr)
+	: LightMap(nullptr), LightMapResourceCluster(nullptr), LightMapTexCoord(0), Parent(nullptr)
 {
 	FStaticMeshAttributes(Description).Register();
 
@@ -74,7 +87,7 @@ FGLTFMeshData::FGLTFMeshData(const USkeletalMesh* SkeletalMesh, const USkeletalM
 
 	// TODO: don't assume last UV channel is non-overlapping
 	const int32 NumTexCoords = SkeletalMesh->GetResourceForRendering()->LODRenderData[LODIndex].StaticVertexBuffers.StaticMeshVertexBuffer.GetNumTexCoords();
-	TexCoord = NumTexCoords - 1;
+	BakeUsingTexCoord = NumTexCoords - 1;
 
 	// TODO: add warning if texture coordinate has overlap
 }
diff --git a/Engine/Plugins/Enterprise/GLTFExporter/MigratedFromGit/Source/GLTFExporter/Private/Converters/GLTFMeshData.h b/Engine/Plugins/Enterprise/GLTFExporter/MigratedFromGit/Source/GLTFExporter/Private/Converters/GLTFMeshData.h
--- a/Engine/Plugins/Enterprise/GLTFExporter/MigratedFromGit/Source/GLTFExporter/Private/Converters/GLTFMeshData.h
+++ b/Engine/Plugins/Enterprise/GLTFExporter/MigratedFromGit/Source/GLTFExporter/Private/Converters/GLTFMeshData.h
@@ -16,7 +16,12 @@ struct FGLTFMeshData
 	FString Name;
 	FPrimitiveData PrimitiveData;
 	FMeshDescription Description;
-	int32 TexCoord;
+
+	FLightMapRef LightMap;
+	const FLightmapResourceCluster* LightMapResourceCluster;
+	int32 LightMapTexCoord;
+
+	int32 BakeUsingTexCoord;
 
 	// TODO: find a better name for referencing the mesh-only data (no component)
 	const FGLTFMeshData* Parent;
diff --git a/Engine/Plugins/Enterprise/GLTFExporter/MigratedFromGit/Source/GLTFExporter/Private/Tasks/GLTFMaterialTasks.cpp b/Engine/Plugins/Enterprise/GLTFExporter/MigratedFromGit/Source/GLTFExporter/Private/Tasks/GLTFMaterialTasks.cpp
--- a/Engine/Plugins/Enterprise/GLTFExporter/MigratedFromGit/Source/GLTFExporter/Private/Tasks/GLTFMaterialTasks.cpp
+++ b/Engine/Plugins/Enterprise/GLTFExporter/MigratedFromGit/Source/GLTFExporter/Private/Tasks/GLTFMaterialTasks.cpp
@@ -165,7 +165,7 @@ void FGLTFMaterialTask::Complete()
 		const FGLTFMeshData* ParentMeshData = MeshData->GetParent();
 
 		const float UVOverlapThreshold = 1.0f / 100.0f;
-		const float UVOverlap = UVOverlapChecker.GetOrAdd(&ParentMeshData->Description, SectionIndices, MeshData->TexCoord);
+		const float UVOverlap = UVOverlapChecker.GetOrAdd(&ParentMeshData->Description, SectionIndices, MeshData->BakeUsingTexCoord);
 
 		if (UVOverlap > UVOverlapThreshold)
 		{
@@ -177,7 +177,7 @@ void FGLTFMaterialTask::Complete()
 				TEXT("Material %s is baked using mesh data from %s but the lightmap UV (channel %d) are overlapping by %.2f%% (in %s) and may produce incorrect results"),
 				*Material->GetName(),
 				*ParentMeshData->Name,
-				MeshData->TexCoord,
+				MeshData->BakeUsingTexCoord,
 				UVOverlap * 100,
 				*SectionString));
 		}
@@ -1326,7 +1326,7 @@ FGLTFPropertyBakeOutput FGLTFMaterialTask::BakeMaterialProperty(const FMaterialP
 	}
 	else
 	{
-		OutTexCoord = MeshData->TexCoord;
+		OutTexCoord = MeshData->BakeUsingTexCoord;
 		MeshDataBakedProperties.Add(Property);
 	}
 
diff --git a/Engine/Plugins/Enterprise/GLTFExporter/MigratedFromGit/Source/GLTFExporter/Private/Tasks/GLTFMeshTasks.cpp b/Engine/Plugins/Enterprise/GLTFExporter/MigratedFromGit/Source/GLTFExporter/Private/Tasks/GLTFMeshTasks.cpp
--- a/Engine/Plugins/Enterprise/GLTFExporter/MigratedFromGit/Source/GLTFExporter/Private/Tasks/GLTFMeshTasks.cpp
+++ b/Engine/Plugins/Enterprise/GLTFExporter/MigratedFromGit/Source/GLTFExporter/Private/Tasks/GLTFMeshTasks.cpp
@@ -122,9 +122,9 @@ void FGLTFStaticMeshTask::Complete()
 			// TODO: report warning in case the mesh actually has data, which means we failed to extract a mesh description.
 			MeshData = nullptr;
 		}
-		else if (MeshData->TexCoord < 0)
+		else if (MeshData->BakeUsingTexCoord < 0)
 		{
-			// TODO: report warning (about missing lightmap texture coordinate for baking with mesh data).
+			// TODO: report warning (about missing texture coordinate for baking with mesh data).
 			MeshData = nullptr;
 		}
 	}
@@ -214,9 +214,9 @@ void FGLTFSkeletalMeshTask::Complete()
 			// TODO: report warning in case the mesh actually has data, which means we failed to extract a mesh description.
 			MeshData = nullptr;
 		}
-		else if (MeshData->TexCoord < 0)
+		else if (MeshData->BakeUsingTexCoord < 0)
 		{
-			// TODO: report warning (about missing non-overlapping texture coordinate for baking with mesh data).
+			// TODO: report warning (about missing texture coordinate for baking with mesh data).
 			MeshData = nullptr;
 		}
 	}
