[GLTFExporter] fix: resolve macos compiler issue regarding braces around initialization

commit 72f5287570f6788e80ebd28b29153552adc3bb29
Date: Fri, 22 Oct 2021 01:54:01 +0200

    fix: resolve macos compiler issue regarding braces around initialization
    (cherry picked from commit 069b9c7acb8cf6a245e0599548417d57ef54b3bc​)

#skipci

diff --git a/Engine/Plugins/Enterprise/GLTFExporter/MigratedFromGit/Source/GLTFExporter/Private/Core/GLTFColor.h b/Engine/Plugins/Enterprise/GLTFExporter/MigratedFromGit/Source/GLTFExporter/Private/Core/GLTFColor.h
--- a/Engine/Plugins/Enterprise/GLTFExporter/MigratedFromGit/Source/GLTFExporter/Private/Core/GLTFColor.h
+++ b/Engine/Plugins/Enterprise/GLTFExporter/MigratedFromGit/Source/GLTFExporter/Private/Core/GLTFColor.h
@@ -22,6 +22,16 @@ struct TGLTFColor3
 
 		ComponentType Components[3] GCC_PACK(1);
 	};
+
+	TGLTFColor3(
+		ComponentType R,
+		ComponentType G,
+		ComponentType B
+	)
+		: R(R)
+		, G(G)
+		, B(B)
+	{}
 };
 
 template <typename ComponentType>
@@ -39,6 +49,18 @@ struct TGLTFColor4
 
 		ComponentType Components[4] GCC_PACK(1);
 	};
+
+	TGLTFColor4(
+		ComponentType R,
+		ComponentType G,
+		ComponentType B,
+		ComponentType A
+	)
+		: R(R)
+		, G(G)
+		, B(B)
+		, A(A)
+	{}
 };
 
 #if PLATFORM_SUPPORTS_PRAGMA_PACK
diff --git a/Engine/Plugins/Enterprise/GLTFExporter/MigratedFromGit/Source/GLTFExporter/Private/Core/GLTFQuaternion.h b/Engine/Plugins/Enterprise/GLTFExporter/MigratedFromGit/Source/GLTFExporter/Private/Core/GLTFQuaternion.h
--- a/Engine/Plugins/Enterprise/GLTFExporter/MigratedFromGit/Source/GLTFExporter/Private/Core/GLTFQuaternion.h
+++ b/Engine/Plugins/Enterprise/GLTFExporter/MigratedFromGit/Source/GLTFExporter/Private/Core/GLTFQuaternion.h
@@ -23,6 +23,18 @@ struct TGLTFQuaternion
 
 		ComponentType Components[4] GCC_PACK(1);
 	};
+
+	TGLTFQuaternion(
+		ComponentType X,
+		ComponentType Y,
+		ComponentType Z,
+		ComponentType W
+	)
+		: X(X)
+		, Y(Y)
+		, Z(Z)
+		, W(W)
+	{}
 };
 
 #if PLATFORM_SUPPORTS_PRAGMA_PACK
diff --git a/Engine/Plugins/Enterprise/GLTFExporter/MigratedFromGit/Source/GLTFExporter/Private/Core/GLTFVector.h b/Engine/Plugins/Enterprise/GLTFExporter/MigratedFromGit/Source/GLTFExporter/Private/Core/GLTFVector.h
--- a/Engine/Plugins/Enterprise/GLTFExporter/MigratedFromGit/Source/GLTFExporter/Private/Core/GLTFVector.h
+++ b/Engine/Plugins/Enterprise/GLTFExporter/MigratedFromGit/Source/GLTFExporter/Private/Core/GLTFVector.h
@@ -21,6 +21,14 @@ struct TGLTFVector2
 
 		ComponentType Components[2] GCC_PACK(1);
 	};
+
+	TGLTFVector2(
+		ComponentType X,
+		ComponentType Y
+	)
+		: X(X)
+		, Y(Y)
+	{}
 };
 
 template <typename ComponentType>
@@ -37,6 +45,16 @@ struct TGLTFVector3
 
 		ComponentType Components[3] GCC_PACK(1);
 	};
+
+	TGLTFVector3(
+		ComponentType X,
+		ComponentType Y,
+		ComponentType Z
+	)
+		: X(X)
+		, Y(Y)
+		, Z(Z)
+	{}
 };
 
 template <typename ComponentType>
@@ -54,6 +72,18 @@ struct TGLTFVector4
 
 		ComponentType Components[4] GCC_PACK(1);
 	};
+
+	TGLTFVector4(
+		ComponentType X,
+		ComponentType Y,
+		ComponentType Z,
+		ComponentType W
+	)
+		: X(X)
+		, Y(Y)
+		, Z(Z)
+		, W(W)
+	{}
 };
 
 #if PLATFORM_SUPPORTS_PRAGMA_PACK
