[GLTFExporter] fix: initialize all components when converting FPackedNormal

commit d23aa93ea6d15f9c2d0debecd0d43595c6f28188
Date: Fri, 22 Oct 2021 01:57:40 +0200

    fix: initialize all components when converting FPackedNormal
    (cherry picked from commit 538b44bda53ec16c1d9f616c73dee14855a2695d​)

#skipci

diff --git a/Engine/Plugins/Enterprise/GLTFExporter/MigratedFromGit/Source/GLTFExporter/Private/Converters/GLTFConverterUtility.h b/Engine/Plugins/Enterprise/GLTFExporter/MigratedFromGit/Source/GLTFExporter/Private/Converters/GLTFConverterUtility.h
--- a/Engine/Plugins/Enterprise/GLTFExporter/MigratedFromGit/Source/GLTFExporter/Private/Converters/GLTFConverterUtility.h
+++ b/Engine/Plugins/Enterprise/GLTFExporter/MigratedFromGit/Source/GLTFExporter/Private/Converters/GLTFConverterUtility.h
@@ -47,7 +47,7 @@ struct FGLTFConverterUtility
 
 	static FGLTFInt8Vector4 ConvertNormal(const FPackedNormal& Normal)
 	{
-		return { Normal.Vector.X, Normal.Vector.Z, Normal.Vector.Y };
+		return { Normal.Vector.X, Normal.Vector.Z, Normal.Vector.Y, 0 };
 	}
 
 	static FGLTFVector4 ConvertTangent(const FVector& Tangent)
