[GLTFExporter] fix: resolve macos compiler issue regarding explicit specialization after instantiation

commit af258449be54c9f6e71566817c8e5eebbedca0ae
Date: Fri, 22 Oct 2021 02:43:47 +0200

    fix: resolve macos compiler issue regarding explicit specialization after instantiation
    (cherry picked from commit f485fa6b5520cb0b744089a72256fb5c6cbbbf38​)

#skipci

diff --git a/Engine/Plugins/Enterprise/GLTFExporter/MigratedFromGit/Source/GLTFExporter/Private/Json/GLTFJsonColor.cpp b/Engine/Plugins/Enterprise/GLTFExporter/MigratedFromGit/Source/GLTFExporter/Private/Json/GLTFJsonColor.cpp
--- a/Engine/Plugins/Enterprise/GLTFExporter/MigratedFromGit/Source/GLTFExporter/Private/Json/GLTFJsonColor.cpp
+++ b/Engine/Plugins/Enterprise/GLTFExporter/MigratedFromGit/Source/GLTFExporter/Private/Json/GLTFJsonColor.cpp
@@ -2,8 +2,8 @@
 
 #include "Json/GLTFJsonColor.h"
 
-template <> const FGLTFJsonColor3 FGLTFJsonColor3::Black({ 0, 0, 0 });
-template <> const FGLTFJsonColor3 FGLTFJsonColor3::White({ 1, 1, 1 });
+const FGLTFJsonColor3 FGLTFJsonColor3::Black({ 0, 0, 0 });
+const FGLTFJsonColor3 FGLTFJsonColor3::White({ 1, 1, 1 });
 
-template <> const FGLTFJsonColor4 FGLTFJsonColor4::Black({ 0, 0, 0, 1 });
-template <> const FGLTFJsonColor4 FGLTFJsonColor4::White({ 1, 1, 1, 1 });
+const FGLTFJsonColor4 FGLTFJsonColor4::Black({ 0, 0, 0, 1 });
+const FGLTFJsonColor4 FGLTFJsonColor4::White({ 1, 1, 1, 1 });
diff --git a/Engine/Plugins/Enterprise/GLTFExporter/MigratedFromGit/Source/GLTFExporter/Private/Json/GLTFJsonColor.h b/Engine/Plugins/Enterprise/GLTFExporter/MigratedFromGit/Source/GLTFExporter/Private/Json/GLTFJsonColor.h
--- a/Engine/Plugins/Enterprise/GLTFExporter/MigratedFromGit/Source/GLTFExporter/Private/Json/GLTFJsonColor.h
+++ b/Engine/Plugins/Enterprise/GLTFExporter/MigratedFromGit/Source/GLTFExporter/Private/Json/GLTFJsonColor.h
@@ -6,11 +6,8 @@
 #include "Json/GLTFJsonArray.h"
 
 template <typename BaseType>
-struct TGLTFJsonColor final : BaseType, IGLTFJsonArray
+struct TGLTFJsonColor : BaseType, IGLTFJsonArray
 {
-	static const TGLTFJsonColor Black;
-	static const TGLTFJsonColor White;
-
 	TGLTFJsonColor(const BaseType& Other)
 		: BaseType(Other)
 	{
@@ -44,5 +41,20 @@ struct TGLTFJsonColor final : BaseType, IGLTFJsonArray
 	}
 };
 
-typedef TGLTFJsonColor<FGLTFColor3> FGLTFJsonColor3;
-typedef TGLTFJsonColor<FGLTFColor4> FGLTFJsonColor4;
+struct FGLTFJsonColor3 : TGLTFJsonColor<FGLTFColor3>
+{
+	static const FGLTFJsonColor3 Black;
+	static const FGLTFJsonColor3 White;
+
+	using TGLTFJsonColor::TGLTFJsonColor;
+	using TGLTFJsonColor::operator=;
+};
+
+struct FGLTFJsonColor4 : TGLTFJsonColor<FGLTFColor4>
+{
+	static const FGLTFJsonColor4 Black;
+	static const FGLTFJsonColor4 White;
+
+	using TGLTFJsonColor::TGLTFJsonColor;
+	using TGLTFJsonColor::operator=;
+};
diff --git a/Engine/Plugins/Enterprise/GLTFExporter/MigratedFromGit/Source/GLTFExporter/Private/Json/GLTFJsonMatrix.cpp b/Engine/Plugins/Enterprise/GLTFExporter/MigratedFromGit/Source/GLTFExporter/Private/Json/GLTFJsonMatrix.cpp
--- a/Engine/Plugins/Enterprise/GLTFExporter/MigratedFromGit/Source/GLTFExporter/Private/Json/GLTFJsonMatrix.cpp
+++ b/Engine/Plugins/Enterprise/GLTFExporter/MigratedFromGit/Source/GLTFExporter/Private/Json/GLTFJsonMatrix.cpp
@@ -2,20 +2,20 @@
 
 #include "Json/GLTFJsonMatrix.h"
 
-template <> const FGLTFJsonMatrix2 FGLTFJsonMatrix2::Identity
+const FGLTFJsonMatrix2 FGLTFJsonMatrix2::Identity
 ({
 	1, 0,
 	0, 1
 });
 
-template <> const FGLTFJsonMatrix3 FGLTFJsonMatrix3::Identity
+const FGLTFJsonMatrix3 FGLTFJsonMatrix3::Identity
 ({
 	1, 0, 0,
 	0, 1, 0,
 	0, 0, 1
 });
 
-template <> const FGLTFJsonMatrix4 FGLTFJsonMatrix4::Identity
+const FGLTFJsonMatrix4 FGLTFJsonMatrix4::Identity
 ({
 	1, 0, 0, 0,
 	0, 1, 0, 0,
diff --git a/Engine/Plugins/Enterprise/GLTFExporter/MigratedFromGit/Source/GLTFExporter/Private/Json/GLTFJsonMatrix.h b/Engine/Plugins/Enterprise/GLTFExporter/MigratedFromGit/Source/GLTFExporter/Private/Json/GLTFJsonMatrix.h
--- a/Engine/Plugins/Enterprise/GLTFExporter/MigratedFromGit/Source/GLTFExporter/Private/Json/GLTFJsonMatrix.h
+++ b/Engine/Plugins/Enterprise/GLTFExporter/MigratedFromGit/Source/GLTFExporter/Private/Json/GLTFJsonMatrix.h
@@ -6,7 +6,7 @@
 #include "Json/GLTFJsonArray.h"
 
 template <typename BaseType>
-struct TGLTFJsonMatrix final : BaseType, IGLTFJsonArray
+struct TGLTFJsonMatrix : BaseType, IGLTFJsonArray
 {
 	static const TGLTFJsonMatrix Identity;
 
@@ -43,6 +43,26 @@ struct TGLTFJsonMatrix final : BaseType, IGLTFJsonArray
 	}
 };
 
-typedef TGLTFJsonMatrix<FGLTFMatrix2> FGLTFJsonMatrix2;
-typedef TGLTFJsonMatrix<FGLTFMatrix3> FGLTFJsonMatrix3;
-typedef TGLTFJsonMatrix<FGLTFMatrix4> FGLTFJsonMatrix4;
+struct FGLTFJsonMatrix2 : TGLTFJsonMatrix<FGLTFMatrix2>
+{
+	static const FGLTFJsonMatrix2 Identity;
+
+	using TGLTFJsonMatrix::TGLTFJsonMatrix;
+	using TGLTFJsonMatrix::operator=;
+};
+
+struct FGLTFJsonMatrix3 : TGLTFJsonMatrix<FGLTFMatrix3>
+{
+	static const FGLTFJsonMatrix3 Identity;
+
+	using TGLTFJsonMatrix::TGLTFJsonMatrix;
+	using TGLTFJsonMatrix::operator=;
+};
+
+struct FGLTFJsonMatrix4 : TGLTFJsonMatrix<FGLTFMatrix4>
+{
+	static const FGLTFJsonMatrix4 Identity;
+
+	using TGLTFJsonMatrix::TGLTFJsonMatrix;
+	using TGLTFJsonMatrix::operator=;
+};
diff --git a/Engine/Plugins/Enterprise/GLTFExporter/MigratedFromGit/Source/GLTFExporter/Private/Json/GLTFJsonQuaternion.h b/Engine/Plugins/Enterprise/GLTFExporter/MigratedFromGit/Source/GLTFExporter/Private/Json/GLTFJsonQuaternion.h
--- a/Engine/Plugins/Enterprise/GLTFExporter/MigratedFromGit/Source/GLTFExporter/Private/Json/GLTFJsonQuaternion.h
+++ b/Engine/Plugins/Enterprise/GLTFExporter/MigratedFromGit/Source/GLTFExporter/Private/Json/GLTFJsonQuaternion.h
@@ -5,7 +5,7 @@
 #include "Core/GLTFQuaternion.h"
 #include "Json/GLTFJsonArray.h"
 
-struct FGLTFJsonQuaternion final : TGLTFQuaternion<float>, IGLTFJsonArray
+struct FGLTFJsonQuaternion : TGLTFQuaternion<float>, IGLTFJsonArray
 {
 	static const TGLTFQuaternion Identity;
 
diff --git a/Engine/Plugins/Enterprise/GLTFExporter/MigratedFromGit/Source/GLTFExporter/Private/Json/GLTFJsonVector.cpp b/Engine/Plugins/Enterprise/GLTFExporter/MigratedFromGit/Source/GLTFExporter/Private/Json/GLTFJsonVector.cpp
--- a/Engine/Plugins/Enterprise/GLTFExporter/MigratedFromGit/Source/GLTFExporter/Private/Json/GLTFJsonVector.cpp
+++ b/Engine/Plugins/Enterprise/GLTFExporter/MigratedFromGit/Source/GLTFExporter/Private/Json/GLTFJsonVector.cpp
@@ -2,11 +2,11 @@
 
 #include "Json/GLTFJsonVector.h"
 
-template <> const FGLTFJsonVector2 FGLTFJsonVector2::Zero({ 0, 0 });
-template <> const FGLTFJsonVector2 FGLTFJsonVector2::One ({ 1, 1 });
+const FGLTFJsonVector2 FGLTFJsonVector2::Zero({ 0, 0 });
+const FGLTFJsonVector2 FGLTFJsonVector2::One ({ 1, 1 });
 
-template <> const FGLTFJsonVector3 FGLTFJsonVector3::Zero({ 0, 0, 0 });
-template <> const FGLTFJsonVector3 FGLTFJsonVector3::One ({ 1, 1, 1 });
+const FGLTFJsonVector3 FGLTFJsonVector3::Zero({ 0, 0, 0 });
+const FGLTFJsonVector3 FGLTFJsonVector3::One ({ 1, 1, 1 });
 
-template <> const FGLTFJsonVector4 FGLTFJsonVector4::Zero({ 0, 0, 0, 0 });
-template <> const FGLTFJsonVector4 FGLTFJsonVector4::One ({ 1, 1, 1, 1 });
+const FGLTFJsonVector4 FGLTFJsonVector4::Zero({ 0, 0, 0, 0 });
+const FGLTFJsonVector4 FGLTFJsonVector4::One ({ 1, 1, 1, 1 });
diff --git a/Engine/Plugins/Enterprise/GLTFExporter/MigratedFromGit/Source/GLTFExporter/Private/Json/GLTFJsonVector.h b/Engine/Plugins/Enterprise/GLTFExporter/MigratedFromGit/Source/GLTFExporter/Private/Json/GLTFJsonVector.h
--- a/Engine/Plugins/Enterprise/GLTFExporter/MigratedFromGit/Source/GLTFExporter/Private/Json/GLTFJsonVector.h
+++ b/Engine/Plugins/Enterprise/GLTFExporter/MigratedFromGit/Source/GLTFExporter/Private/Json/GLTFJsonVector.h
@@ -6,7 +6,7 @@
 #include "Json/GLTFJsonArray.h"
 
 template <typename BaseType>
-struct TGLTFJsonVector final : BaseType, IGLTFJsonArray
+struct TGLTFJsonVector : BaseType, IGLTFJsonArray
 {
 	static const TGLTFJsonVector Zero;
 	static const TGLTFJsonVector One;
@@ -44,6 +44,29 @@ struct TGLTFJsonVector final : BaseType, IGLTFJsonArray
 	}
 };
 
-typedef TGLTFJsonVector<FGLTFVector2> FGLTFJsonVector2;
-typedef TGLTFJsonVector<FGLTFVector3> FGLTFJsonVector3;
-typedef TGLTFJsonVector<FGLTFVector4> FGLTFJsonVector4;
+struct FGLTFJsonVector2 : TGLTFJsonVector<FGLTFVector2>
+{
+	static const FGLTFJsonVector2 Zero;
+	static const FGLTFJsonVector2 One;
+
+	using TGLTFJsonVector::TGLTFJsonVector;
+	using TGLTFJsonVector::operator=;
+};
+
+struct FGLTFJsonVector3 : TGLTFJsonVector<FGLTFVector3>
+{
+	static const FGLTFJsonVector3 Zero;
+	static const FGLTFJsonVector3 One;
+
+	using TGLTFJsonVector::TGLTFJsonVector;
+	using TGLTFJsonVector::operator=;
+};
+
+struct FGLTFJsonVector4 : TGLTFJsonVector<FGLTFVector4>
+{
+	static const FGLTFJsonVector4 Zero;
+	static const FGLTFJsonVector4 One;
+
+	using TGLTFJsonVector::TGLTFJsonVector;
+	using TGLTFJsonVector::operator=;
+};
