[GLTFExporter] chore: whitelisted Mac platform after verifying the exporter works

commit d85c25f3ac9740985cb4d6d278b72b69a86d751d
Date: Fri, 22 Oct 2021 03:24:25 +0200

    chore: whitelisted Mac platform after verifying the exporter works
    (cherry picked from commit 869c29cd42508d72940351e0a1f35a58e615fbf2​)

#skipci

diff --git a/Engine/Plugins/Enterprise/GLTFExporter/MigratedFromGit/GLTFExporter.uplugin b/Engine/Plugins/Enterprise/GLTFExporter/MigratedFromGit/GLTFExporter.uplugin
--- a/Engine/Plugins/Enterprise/GLTFExporter/MigratedFromGit/GLTFExporter.uplugin
+++ b/Engine/Plugins/Enterprise/GLTFExporter/MigratedFromGit/GLTFExporter.uplugin
@@ -18,25 +18,25 @@
 			"Name": "GLTFExporter",
 			"Type": "Editor",
 			"LoadingPhase": "Default",
-			"WhitelistPlatforms": [ "Win64", "Win32" ]
+			"WhitelistPlatforms": [ "Win64", "Win32", "Mac" ]
 		},
 		{
 			"Name": "GLTFExporterRuntime",
 			"Type": "Runtime",
 			"LoadingPhase": "Default",
-			"WhitelistPlatforms": [ "Win64", "Win32" ]
+			"WhitelistPlatforms": [ "Win64", "Win32", "Mac" ]
 		},
 		{
 			"Name": "GLTFMaterialAnalyzer",
 			"Type": "Editor",
 			"LoadingPhase": "Default",
-			"WhitelistPlatforms": [ "Win64", "Win32" ]
+			"WhitelistPlatforms": [ "Win64", "Win32", "Mac" ]
 		},
 		{
 			"Name": "GLTFMaterialBaking",
 			"Type": "Editor",
 			"LoadingPhase": "Default",
-			"WhitelistPlatforms": [ "Win64", "Win32" ]
+			"WhitelistPlatforms": [ "Win64", "Win32", "Mac" ]
 		}
 	],
 	"Plugins": [
