[GLTFExporter] feat: add support for material analyzer to recognize expression LightmapUVs

commit bde776880f159149d3acbb13f55e9042ed0f0c9b
Date: Fri, 22 Oct 2021 05:35:01 +0200

    feat: add support for material analyzer to recognize expression LightmapUVs
    (cherry picked from commit f03bdafc9eebd0d71b6614ae09bbf26633aee511​)

#skipci

diff --git a/Engine/Plugins/Enterprise/GLTFExporter/MigratedFromGit/Source/GLTFMaterialAnalyzer/Private/GLTFMaterialAnalyzer.cpp b/Engine/Plugins/Enterprise/GLTFExporter/MigratedFromGit/Source/GLTFMaterialAnalyzer/Private/GLTFMaterialAnalyzer.cpp
--- a/Engine/Plugins/Enterprise/GLTFExporter/MigratedFromGit/Source/GLTFMaterialAnalyzer/Private/GLTFMaterialAnalyzer.cpp
+++ b/Engine/Plugins/Enterprise/GLTFExporter/MigratedFromGit/Source/GLTFMaterialAnalyzer/Private/GLTFMaterialAnalyzer.cpp
@@ -118,8 +118,12 @@ int32 UGLTFMaterialAnalyzer::CompilePropertyEx(FMaterialCompiler* Compiler, cons
 		TranslatorHack->bUsesTransformVector ||
 		TranslatorHack->bNeedsWorldPositionExcludingShaderOffsets ||
 		TranslatorHack->bUsesAOMaterialMask ||
+		TranslatorHack->bUsesLightmapUVs ||
 		TranslatorHack->bUsesVertexPosition;
 
+	// TODO: need to separate between local data (local bounds, pre-skinned local position & normal, transforming between local-space <-> tangent-space etc)
+	// and between world data (world / actor position, world bounds, ao mask, lightmap uvs, custom primitive data, transforming to and form world-space etc)
+
 	Analysis->bRequiresPrimitiveData =
 		ProxyCompiler.bUsesActorPosition ||
 		ProxyCompiler.bUsesObjectPosition ||
