[GLTFExporter] fix: fix inconsistency between Unreal and glTF when exporting translucent clear-coat materials

commit 691fb402ad93d3a52cfd1e1107b14a96cbc4f7ca
Date: Fri, 22 Oct 2021 05:38:14 +0200

    fix: fix inconsistency between Unreal and glTF when exporting translucent clear-coat materials
    (cherry picked from commit d7f7353a73eba4bb0a6f9fd4a47bf8082cdd434a​)

#skipci

diff --git a/Engine/Plugins/Enterprise/GLTFExporter/MigratedFromGit/Source/GLTFExporter/Private/Tasks/GLTFMaterialTasks.cpp b/Engine/Plugins/Enterprise/GLTFExporter/MigratedFromGit/Source/GLTFExporter/Private/Tasks/GLTFMaterialTasks.cpp
--- a/Engine/Plugins/Enterprise/GLTFExporter/MigratedFromGit/Source/GLTFExporter/Private/Tasks/GLTFMaterialTasks.cpp
+++ b/Engine/Plugins/Enterprise/GLTFExporter/MigratedFromGit/Source/GLTFExporter/Private/Tasks/GLTFMaterialTasks.cpp
@@ -254,6 +254,13 @@ void FGLTFMaterialTask::ConvertShadingModel(EGLTFJsonShadingModel& OutShadingMod
 		}
 	}
 
+	const EBlendMode BlendMode = Material->GetBlendMode();
+	if (ShadingModel == MSM_ClearCoat && BlendMode != BLEND_Opaque && BlendMode != BLEND_Masked)
+	{
+		// NOTE: Unreal seems to disable clear coat when blend mode anything but opaque or masked
+		ShadingModel = MSM_DefaultLit;
+	}
+
 	OutShadingModel = FGLTFConverterUtility::ConvertShadingModel(ShadingModel);
 	if (OutShadingModel == EGLTFJsonShadingModel::None)
 	{
