[GLTFExporter] chore: bump version to 1.1.0

commit bab417d1d220c46c7e47d8e8a91f0a264980098b
Date: Fri, 22 Oct 2021 09:22:10 +0200

    chore: bump version to 1.1.0
    (cherry picked from commit 29989a2471d95deb93da1d4428824973875c5fd4​)

#skipci

diff --git a/Engine/Plugins/Enterprise/GLTFExporter/MigratedFromGit/GLTFExporter.uplugin b/Engine/Plugins/Enterprise/GLTFExporter/MigratedFromGit/GLTFExporter.uplugin
--- a/Engine/Plugins/Enterprise/GLTFExporter/MigratedFromGit/GLTFExporter.uplugin
+++ b/Engine/Plugins/Enterprise/GLTFExporter/MigratedFromGit/GLTFExporter.uplugin
@@ -1,7 +1,7 @@
 {
 	"FileVersion": 3,
-	"Version": 106,
-	"VersionName": "1.0.6",
+	"Version": 110,
+	"VersionName": "1.1.0",
 	"FriendlyName": "glTF Exporter",
 	"Description": "An exporter for Khronos glTF 2.0.",
 	"Category": "Exporters",
