[GLTFExporter] refactor: don't use editor-only functions when converting animations

commit a63dca9f2de3058466fc1f10a6ec81e8bc1da168
Date: Sun, 24 Oct 2021 12:31:52 +0200

    refactor: don't use editor-only functions when converting animations
    (cherry picked from commit f78b0719b19489b26fa4ad6443c1063c0a3f57c5​)

#skipci

diff --git a/Engine/Plugins/Enterprise/GLTFExporter/MigratedFromGit/Source/GLTFExporter/Private/Tasks/GLTFAnimationTasks.cpp b/Engine/Plugins/Enterprise/GLTFExporter/MigratedFromGit/Source/GLTFExporter/Private/Tasks/GLTFAnimationTasks.cpp
--- a/Engine/Plugins/Enterprise/GLTFExporter/MigratedFromGit/Source/GLTFExporter/Private/Tasks/GLTFAnimationTasks.cpp
+++ b/Engine/Plugins/Enterprise/GLTFExporter/MigratedFromGit/Source/GLTFExporter/Private/Tasks/GLTFAnimationTasks.cpp
@@ -25,12 +25,15 @@ void FGLTFAnimSequenceTask::Complete()
 	FGLTFJsonAnimation& JsonAnimation = Builder.GetAnimation(AnimationIndex);
 	JsonAnimation.Name = AnimSequence->GetName() + TEXT("_") + FString::FromInt(AnimationIndex); // Ensure unique name due to limitation in certain gltf viewers
 
+	const float SequenceLength = AnimSequence->SequenceLength;
+	const float FrameLength = FrameCount > 1 ? SequenceLength / (FrameCount - 1) : 0;
+
 	TArray<float> Timestamps;
 	Timestamps.AddUninitialized(FrameCount);
 
 	for (int32 Frame = 0; Frame < FrameCount; ++Frame)
 	{
-		Timestamps[Frame] = AnimSequence->GetTimeAtFrame(Frame);
+		Timestamps[Frame] = FMath::Clamp(Frame * FrameLength, 0.0f, SequenceLength);
 	}
 
 	// TODO: add animation data accessor converters to reuse track information
@@ -55,12 +58,12 @@ void FGLTFAnimSequenceTask::Complete()
 		// TODO: report warning (about unknown interpolation exported as linear)
 	}
 
-	const TArray<FName>& TrackNames = AnimSequence->GetAnimationTrackNames();
-	const int32 TrackCount = TrackNames.Num();
+	const TArray<FRawAnimSequenceTrack>& Tracks = AnimSequence->GetRawAnimationData();
+	const int32 TrackCount = Tracks.Num();
 
 	for (int32 TrackIndex = 0; TrackIndex < TrackCount; ++TrackIndex)
 	{
-		const FRawAnimSequenceTrack& Track = AnimSequence->GetRawAnimationTrack(TrackIndex);
+		const FRawAnimSequenceTrack& Track = Tracks[TrackIndex];
 		const TArray<FVector>& KeyPositions = Track.PosKeys;
 		const TArray<FQuat>& KeyRotations = Track.RotKeys;
 		const TArray<FVector>& KeyScales = Track.ScaleKeys;
