[GLTFExporter] refactor: remove all functions in GLTFFileUtility that are no longer needed

commit f46e82047cdf0979ff3b9426bf00aded1297b257
Date: Sun, 24 Oct 2021 12:29:58 +0200

    refactor: remove all functions in GLTFFileUtility that are no longer needed
    (cherry picked from commit 44c2177d86aabd50bc98630aa9382c5022948437​)

#skipci

diff --git a/Engine/Plugins/Enterprise/GLTFExporter/MigratedFromGit/Source/GLTFExporter/Private/Builders/GLTFFileUtility.cpp b/Engine/Plugins/Enterprise/GLTFExporter/MigratedFromGit/Source/GLTFExporter/Private/Builders/GLTFFileUtility.cpp
--- a/Engine/Plugins/Enterprise/GLTFExporter/MigratedFromGit/Source/GLTFExporter/Private/Builders/GLTFFileUtility.cpp
+++ b/Engine/Plugins/Enterprise/GLTFExporter/MigratedFromGit/Source/GLTFExporter/Private/Builders/GLTFFileUtility.cpp
@@ -1,21 +1,6 @@
 // Copyright Epic Games, Inc. All Rights Reserved.
 
 #include "Builders/GLTFFileUtility.h"
-#include "GLTFExporterModule.h"
-#include "Interfaces/IPluginManager.h"
-#include "Serialization/JsonReader.h"
-#include "Serialization/JsonSerializer.h"
-#include "Misc/FileHelper.h"
-
-#if PLATFORM_LINUX || PLATFORM_MAC
-// TODO: uncomment when GLTFLaunchHelper added for linux and macos
-// #include <sys/stat.h>
-#endif
-
-FString FGLTFFileUtility::GetPluginDir()
-{
-	return IPluginManager::Get().FindPlugin(GLTFEXPORTER_MODULE_NAME)->GetBaseDir();
-}
 
 const TCHAR* FGLTFFileUtility::GetFileExtension(EGLTFJsonMimeType MimeType)
 {
@@ -53,49 +38,3 @@ bool FGLTFFileUtility::IsGlbFile(const FString& Filename)
 	return FPaths::GetExtension(Filename).Equals(TEXT("glb"), ESearchCase::IgnoreCase);
 }
 
-bool FGLTFFileUtility::ReadJsonFile(const FString& FilePath, TSharedPtr<FJsonObject>& JsonObject)
-{
-	FString JsonContent;
-	if (!FFileHelper::LoadFileToString(JsonContent, *FilePath))
-	{
-		return false;
-	}
-
-	const TSharedRef<TJsonReader<>> JsonReader = TJsonReaderFactory<>::Create(JsonContent);
-	return FJsonSerializer::Deserialize(JsonReader, JsonObject) && JsonObject.IsValid();
-}
-
-bool FGLTFFileUtility::WriteJsonFile(const FString& FilePath, const TSharedRef<FJsonObject>& JsonObject)
-{
-	FString JsonContent;
-	const TSharedRef<TJsonWriter<>> JsonWriter = TJsonWriterFactory<>::Create(&JsonContent);
-
-	if (!FJsonSerializer::Serialize(JsonObject, JsonWriter))
-	{
-		return false;
-	}
-
-	return FFileHelper::SaveStringToFile(JsonContent, *FilePath, FFileHelper::EEncodingOptions::ForceUTF8);
-}
-
-bool FGLTFFileUtility::SetExecutable(const TCHAR* Filename, bool bIsExecutable)
-{
-/* TODO: uncomment when GLTFLaunchHelper added for linux and macos
-#if PLATFORM_LINUX || PLATFORM_MAC
-	struct stat FileInfo;
-	FTCHARToUTF8 FilenameUtf8(Filename);
-	bool bSuccess = false;
-	// Set executable permission bit
-	if (stat(FilenameUtf8.Get(), &FileInfo) == 0)
-	{
-		mode_t ExeFlags = S_IXUSR | S_IXGRP | S_IXOTH;
-		FileInfo.st_mode = bIsExecutable ? (FileInfo.st_mode | ExeFlags) : (FileInfo.st_mode & ~ExeFlags);
-		bSuccess = chmod(FilenameUtf8.Get(), FileInfo.st_mode) == 0;
-	}
-	return bSuccess;
-#else
-	return true;
-#endif
-*/
-	return true;
-}
diff --git a/Engine/Plugins/Enterprise/GLTFExporter/MigratedFromGit/Source/GLTFExporter/Private/Builders/GLTFFileUtility.h b/Engine/Plugins/Enterprise/GLTFExporter/MigratedFromGit/Source/GLTFExporter/Private/Builders/GLTFFileUtility.h
--- a/Engine/Plugins/Enterprise/GLTFExporter/MigratedFromGit/Source/GLTFExporter/Private/Builders/GLTFFileUtility.h
+++ b/Engine/Plugins/Enterprise/GLTFExporter/MigratedFromGit/Source/GLTFExporter/Private/Builders/GLTFFileUtility.h
@@ -7,16 +7,9 @@
 
 struct FGLTFFileUtility
 {
-	static FString GetPluginDir();
-
 	static const TCHAR* GetFileExtension(EGLTFJsonMimeType MimeType);
 
 	static FString GetUniqueFilename(const FString& BaseFilename, const FString& FileExtension, const TSet<FString>& UniqueFilenames);
 
 	static bool IsGlbFile(const FString& Filename);
-
-	static bool ReadJsonFile(const FString& FilePath, TSharedPtr<FJsonObject>& JsonObject);
-	static bool WriteJsonFile(const FString& FilePath, const TSharedRef<FJsonObject>& JsonObject);
-
-	static bool SetExecutable(const TCHAR* Filename, bool bIsExecutable);
 };
