[GLTFExporter] feat: skip evaluation of shading model expression if there is only one possible value

commit 56f1f4a894c4e7546b62c08efd7fa0aefdfed2db
Date: Thu, 28 Oct 2021 11:07:12 +0200

    feat: skip evaluation of shading model expression if there is only one possible value
    (cherry picked from commit e2b37ec87c1b036bd2f2fee94ac876187f90ed88​)

#skipci

diff --git a/Engine/Plugins/Enterprise/GLTFExporter/MigratedFromGit/Source/GLTFExporter/Private/Tasks/GLTFMaterialTasks.cpp b/Engine/Plugins/Enterprise/GLTFExporter/MigratedFromGit/Source/GLTFExporter/Private/Tasks/GLTFMaterialTasks.cpp
--- a/Engine/Plugins/Enterprise/GLTFExporter/MigratedFromGit/Source/GLTFExporter/Private/Tasks/GLTFMaterialTasks.cpp
+++ b/Engine/Plugins/Enterprise/GLTFExporter/MigratedFromGit/Source/GLTFExporter/Private/Tasks/GLTFMaterialTasks.cpp
@@ -61,8 +61,8 @@
 	JsonMaterial.AlphaCutoff = Material->GetOpacityMaskClipValue();
 	JsonMaterial.DoubleSided = Material->IsTwoSided();
 
-	ConvertAlphaMode(JsonMaterial.AlphaMode, JsonMaterial.BlendMode);
 	ConvertShadingModel(JsonMaterial.ShadingModel);
+	ConvertAlphaMode(JsonMaterial.AlphaMode, JsonMaterial.BlendMode);
 
 	if (JsonMaterial.ShadingModel != EGLTFJsonShadingModel::None)
 	{
@@ -184,76 +184,62 @@
 	}
 }
 
-void FGLTFMaterialTask::ConvertAlphaMode(EGLTFJsonAlphaMode& OutAlphaMode, EGLTFJsonBlendMode& OutBlendMode) const
+EMaterialShadingModel FGLTFMaterialTask::GetShadingModel() const
 {
-	const EBlendMode BlendMode = Material->GetBlendMode();
+	const FMaterialShadingModelField Possibilites = Material->GetShadingModels();
+	const int32 PossibilitesCount = Possibilites.CountShadingModels();
 
-	OutAlphaMode = FGLTFConverterUtility::ConvertAlphaMode(BlendMode);
-	if (OutAlphaMode == EGLTFJsonAlphaMode::None)
+	if (PossibilitesCount == 0)
 	{
-		OutAlphaMode = EGLTFJsonAlphaMode::Blend;
-		OutBlendMode = EGLTFJsonBlendMode::None;
-
+		EMaterialShadingModel ShadingModel = MSM_DefaultLit;
 		Builder.AddWarningMessage(FString::Printf(
-			TEXT("Unsupported blend mode (%s) in material %s, will export as %s"),
-			*FGLTFNameUtility::GetName(BlendMode),
+			TEXT("No shading model defined for material %s, will export as %s"),
 			*Material->GetName(),
-			*FGLTFNameUtility::GetName(BLEND_Translucent)));
-		return;
+			*FGLTFNameUtility::GetName(ShadingModel)));
+		return ShadingModel;
 	}
 
-	if (OutAlphaMode == EGLTFJsonAlphaMode::Blend)
+	if (PossibilitesCount > 1)
 	{
-		OutBlendMode = FGLTFConverterUtility::ConvertBlendMode(BlendMode);
-		if (OutBlendMode != EGLTFJsonBlendMode::None && !Builder.ExportOptions->bExportExtraBlendModes)
-		{
-			OutBlendMode = EGLTFJsonBlendMode::None;
-
-			Builder.AddWarningMessage(FString::Printf(
-				TEXT("Extra blend mode (%s) in material %s disabled by export options, will export as %s"),
-				*FGLTFNameUtility::GetName(BlendMode),
-				*Material->GetName(),
-				*FGLTFNameUtility::GetName(BLEND_Translucent)));
-		}
-	}
-}
-
-void FGLTFMaterialTask::ConvertShadingModel(EGLTFJsonShadingModel& OutShadingModel) const
-{
-	EMaterialShadingModel ShadingModel;
-	{
-		const FMaterialShadingModelField ShadingModels = Material->GetShadingModels();
 		if (Material->IsShadingModelFromMaterialExpression())
 		{
 			const FMaterialShadingModelField Evaluation = FGLTFMaterialUtility::EvaluateShadingModelExpression(Material);
-			if (Evaluation.CountShadingModels() == 1)
+			const int32 EvaluationCount = Evaluation.CountShadingModels();
+
+			if (EvaluationCount == 0)
 			{
-				ShadingModel = Evaluation.GetFirstShadingModel();
+				EMaterialShadingModel ShadingModel = FGLTFMaterialUtility::GetRichestShadingModel(Possibilites);
+				Builder.AddWarningMessage(FString::Printf(
+					TEXT("Evaluation of shading model expression in material %s returned none, will export as %s"),
+					*Material->GetName(),
+					*FGLTFNameUtility::GetName(ShadingModel)));
+				return ShadingModel;
 			}
-			else if (Evaluation.CountShadingModels() > 1)
+
+			if (EvaluationCount > 1)
 			{
-				ShadingModel = FGLTFMaterialUtility::GetRichestShadingModel(Evaluation);
+				EMaterialShadingModel ShadingModel = FGLTFMaterialUtility::GetRichestShadingModel(Evaluation);
 				Builder.AddWarningMessage(FString::Printf(
 					TEXT("Evaluation of shading model expression in material %s is inconclusive (%s), will export as %s"),
 					*Material->GetName(),
 					*FGLTFMaterialUtility::ShadingModelsToString(Evaluation),
 					*FGLTFNameUtility::GetName(ShadingModel)));
+				return ShadingModel;
 			}
-			else
-			{
-				ShadingModel = FGLTFMaterialUtility::GetRichestShadingModel(ShadingModels);
-				Builder.AddWarningMessage(FString::Printf(
-					TEXT("Evaluation of shading model expression in material %s returned none, will export as %s"),
-					*Material->GetName(),
-					*FGLTFNameUtility::GetName(ShadingModel)));
-			}
+
+			return Evaluation.GetFirstShadingModel();
 		}
-		else
-		{
-			ShadingModel = ShadingModels.GetFirstShadingModel();
-		}
+
+		// we should never end up here
 	}
 
+	return Possibilites.GetFirstShadingModel();
+}
+
+void FGLTFMaterialTask::ConvertShadingModel(EGLTFJsonShadingModel& OutShadingModel) const
+{
+	EMaterialShadingModel ShadingModel = GetShadingModel();
+
 	const EBlendMode BlendMode = Material->GetBlendMode();
 	if (ShadingModel == MSM_ClearCoat && BlendMode != BLEND_Opaque && BlendMode != BLEND_Masked)
 	{
@@ -298,6 +284,40 @@
 	}
 }
 
+void FGLTFMaterialTask::ConvertAlphaMode(EGLTFJsonAlphaMode& OutAlphaMode, EGLTFJsonBlendMode& OutBlendMode) const
+{
+	const EBlendMode BlendMode = Material->GetBlendMode();
+
+	OutAlphaMode = FGLTFConverterUtility::ConvertAlphaMode(BlendMode);
+	if (OutAlphaMode == EGLTFJsonAlphaMode::None)
+	{
+		OutAlphaMode = EGLTFJsonAlphaMode::Blend;
+		OutBlendMode = EGLTFJsonBlendMode::None;
+
+		Builder.AddWarningMessage(FString::Printf(
+			TEXT("Unsupported blend mode (%s) in material %s, will export as %s"),
+			*FGLTFNameUtility::GetName(BlendMode),
+			*Material->GetName(),
+			*FGLTFNameUtility::GetName(BLEND_Translucent)));
+		return;
+	}
+
+	if (OutAlphaMode == EGLTFJsonAlphaMode::Blend)
+	{
+		OutBlendMode = FGLTFConverterUtility::ConvertBlendMode(BlendMode);
+		if (OutBlendMode != EGLTFJsonBlendMode::None && !Builder.ExportOptions->bExportExtraBlendModes)
+		{
+			OutBlendMode = EGLTFJsonBlendMode::None;
+
+			Builder.AddWarningMessage(FString::Printf(
+				TEXT("Extra blend mode (%s) in material %s disabled by export options, will export as %s"),
+				*FGLTFNameUtility::GetName(BlendMode),
+				*Material->GetName(),
+				*FGLTFNameUtility::GetName(BLEND_Translucent)));
+		}
+	}
+}
+
 bool FGLTFMaterialTask::TryGetBaseColorAndOpacity(FGLTFJsonPBRMetallicRoughness& OutPBRParams, const FMaterialPropertyEx& BaseColorProperty, const FMaterialPropertyEx& OpacityProperty)
 {
 	const bool bIsBaseColorConstant = TryGetConstantColor(OutPBRParams.BaseColorFactor, BaseColorProperty);
diff --git a/Engine/Plugins/Enterprise/GLTFExporter/MigratedFromGit/Source/GLTFExporter/Private/Tasks/GLTFMaterialTasks.h b/Engine/Plugins/Enterprise/GLTFExporter/MigratedFromGit/Source/GLTFExporter/Private/Tasks/GLTFMaterialTasks.h
--- a/Engine/Plugins/Enterprise/GLTFExporter/MigratedFromGit/Source/GLTFExporter/Private/Tasks/GLTFMaterialTasks.h
+++ b/Engine/Plugins/Enterprise/GLTFExporter/MigratedFromGit/Source/GLTFExporter/Private/Tasks/GLTFMaterialTasks.h
@@ -44,8 +44,9 @@ private:
 
 	TSet<FMaterialPropertyEx> MeshDataBakedProperties;
 
-	void ConvertAlphaMode(EGLTFJsonAlphaMode& OutAlphaMode, EGLTFJsonBlendMode& OutBlendMode) const;
+	EMaterialShadingModel GetShadingModel() const;
 	void ConvertShadingModel(EGLTFJsonShadingModel& OutShadingModel) const;
+	void ConvertAlphaMode(EGLTFJsonAlphaMode& OutAlphaMode, EGLTFJsonBlendMode& OutBlendMode) const;
 
 	bool TryGetBaseColorAndOpacity(FGLTFJsonPBRMetallicRoughness& OutPBRParams, const FMaterialPropertyEx& BaseColorProperty, const FMaterialPropertyEx& OpacityProperty);
 	bool TryGetMetallicAndRoughness(FGLTFJsonPBRMetallicRoughness& OutPBRParams, const FMaterialPropertyEx& MetallicProperty, const FMaterialPropertyEx& RoughnessProperty);
