[GLTFExporter] fix: add missing check for skipping lightmaps if textures disabled

commit e4f10445d9b9a8e2198a1e4f4dbdb84af18643d8
Date: Thu, 28 Oct 2021 11:36:57 +0200

    fix: add missing check for skipping lightmaps if textures disabled
    (cherry picked from commit 69c987f2b8e60abc9cf558e142fbcc1359f93ffe​)

#skipci

diff --git a/Engine/Plugins/Enterprise/GLTFExporter/MigratedFromGit/Source/GLTFExporter/Private/Converters/GLTFLightMapConverters.cpp b/Engine/Plugins/Enterprise/GLTFExporter/MigratedFromGit/Source/GLTFExporter/Private/Converters/GLTFLightMapConverters.cpp
--- a/Engine/Plugins/Enterprise/GLTFExporter/MigratedFromGit/Source/GLTFExporter/Private/Converters/GLTFLightMapConverters.cpp
+++ b/Engine/Plugins/Enterprise/GLTFExporter/MigratedFromGit/Source/GLTFExporter/Private/Converters/GLTFLightMapConverters.cpp
@@ -7,6 +7,11 @@
 
 FGLTFJsonLightMapIndex FGLTFLightMapConverter::Convert(const UStaticMeshComponent* StaticMeshComponent)
 {
+	if (Builder.ExportOptions->TextureImageFormat == EGLTFTextureImageFormat::None)
+	{
+		return FGLTFJsonLightMapIndex(INDEX_NONE);
+	}
+
 	const UStaticMesh* StaticMesh = StaticMeshComponent->GetStaticMesh();
 
 	if (StaticMesh == nullptr)
@@ -16,14 +21,15 @@ FGLTFJsonLightMapIndex FGLTFLightMapConverter::Convert(const UStaticMeshComponen
 
 	const int32 LODIndex = FGLTFMeshUtility::GetLOD(StaticMesh, StaticMeshComponent, Builder.ExportOptions->DefaultLevelOfDetail);
 	const FStaticMeshLODResources& LODResources = StaticMesh->GetLODForExport(LODIndex);
-
 	const int32 CoordinateIndex = StaticMesh->GetLightMapCoordinateIndex();
+
 	if (CoordinateIndex < 0 || CoordinateIndex >= LODResources.GetNumTexCoords())
 	{
 		return FGLTFJsonLightMapIndex(INDEX_NONE);
 	}
 
-	const int32 LightMapLODIndex = 0;
+	const int32 LightMapLODIndex = 0; // TODO: why is this zero?
+
 	if (!StaticMeshComponent->LODData.IsValidIndex(LightMapLODIndex))
 	{
 		return FGLTFJsonLightMapIndex(INDEX_NONE);
@@ -38,6 +44,7 @@ FGLTFJsonLightMapIndex FGLTFLightMapConverter::Convert(const UStaticMeshComponen
 	}
 
 	const FLightMap2D* LightMap2D = MeshMapBuildData->LightMap->GetLightMap2D();
+
 	if (LightMap2D == nullptr)
 	{
 		return FGLTFJsonLightMapIndex(INDEX_NONE);
