[GLTFExporter] feat: add support for running exporter in GPU-less mode (i.e. -NullRHI) regardless of options and assets

commit 9fa4b3f52679ca0aeb13a33efe971a59f40c013b
Date: Thu, 28 Oct 2021 11:46:26 +0200

    feat: add support for running exporter in GPU-less mode (i.e. -NullRHI) regardless of options and assets
    (cherry picked from commit ab62d474a72602c7624d6b19c540b6fedd9e3bc2​)

#skipci

diff --git a/Engine/Plugins/Enterprise/GLTFExporter/MigratedFromGit/Source/GLTFExporter/Private/Builders/GLTFBuilder.cpp b/Engine/Plugins/Enterprise/GLTFExporter/MigratedFromGit/Source/GLTFExporter/Private/Builders/GLTFBuilder.cpp
--- a/Engine/Plugins/Enterprise/GLTFExporter/MigratedFromGit/Source/GLTFExporter/Private/Builders/GLTFBuilder.cpp
+++ b/Engine/Plugins/Enterprise/GLTFExporter/MigratedFromGit/Source/GLTFExporter/Private/Builders/GLTFBuilder.cpp
@@ -8,7 +8,8 @@ FGLTFBuilder::FGLTFBuilder(const FString& FilePath, const UGLTFExportOptions* Ex
 	: bIsGlbFile(FGLTFFileUtility::IsGlbFile(FilePath))
 	, FilePath(FilePath)
 	, DirPath(FPaths::GetPath(FilePath))
-	, ExportOptions(ExportOptions)
+	, ExportOptions(ValidateExportOptions(ExportOptions))
+	, ExportOptionsGuard(ExportOptions)
 {
 }
 
@@ -82,6 +83,23 @@ bool FGLTFBuilder::ShouldExportLight(EComponentMobility::Type LightMobility) con
 	return EnumHasAllFlags(AllowedMobility, QueriedMobility);
 }
 
+const UGLTFExportOptions* FGLTFBuilder::ValidateExportOptions(const UGLTFExportOptions* ExportOptions)
+{
+	if (!FApp::CanEverRender())
+	{
+		if (ExportOptions->BakeMaterialInputs != EGLTFMaterialBakeMode::Disabled || ExportOptions->TextureImageFormat != EGLTFTextureImageFormat::None)
+		{
+			// TODO: warn the following options requires rendering support and will be overriden
+			UGLTFExportOptions* OverridenOptions = DuplicateObject(ExportOptions, nullptr);
+			OverridenOptions->BakeMaterialInputs = EGLTFMaterialBakeMode::Disabled;
+			OverridenOptions->TextureImageFormat = EGLTFTextureImageFormat::None;
+			return OverridenOptions;
+		}
+	}
+
+	return ExportOptions;
+}
+
 EGLTFSceneMobility FGLTFBuilder::GetSceneMobility(EComponentMobility::Type Mobility)
 {
 	switch (Mobility)
diff --git a/Engine/Plugins/Enterprise/GLTFExporter/MigratedFromGit/Source/GLTFExporter/Private/Builders/GLTFBuilder.h b/Engine/Plugins/Enterprise/GLTFExporter/MigratedFromGit/Source/GLTFExporter/Private/Builders/GLTFBuilder.h
--- a/Engine/Plugins/Enterprise/GLTFExporter/MigratedFromGit/Source/GLTFExporter/Private/Builders/GLTFBuilder.h
+++ b/Engine/Plugins/Enterprise/GLTFExporter/MigratedFromGit/Source/GLTFExporter/Private/Builders/GLTFBuilder.h
@@ -5,6 +5,7 @@
 #include "GLTFExportOptions.h"
 #include "Json/GLTFJsonEnums.h"
 #include "MaterialPropertyEx.h"
+#include "UObject/GCObjectScopeGuard.h"
 
 class FGLTFBuilder
 {
@@ -19,7 +20,7 @@ public:
 	const FString DirPath;
 
 	// TODO: make ExportOptions private and expose each option via getters to ease overriding settings in future
-	const UGLTFExportOptions* const ExportOptions;
+	const UGLTFExportOptions* ExportOptions;
 
 	FIntPoint GetBakeSizeForMaterialProperty(const UMaterialInterface* Material, const FMaterialPropertyEx& Property) const;
 	TextureFilter GetBakeFilterForMaterialProperty(const UMaterialInterface* Material, const FMaterialPropertyEx& Property) const;
@@ -31,6 +32,9 @@ public:
 
 private:
 
+	FGCObjectScopeGuard ExportOptionsGuard;
+	
+	static const UGLTFExportOptions* ValidateExportOptions(const UGLTFExportOptions* ExportOptions);
 	static EGLTFSceneMobility GetSceneMobility(EComponentMobility::Type Mobility);
 	static EGLTFMaterialPropertyGroup GetPropertyGroup(const FMaterialPropertyEx& Property);
 };
diff --git a/Engine/Plugins/Enterprise/GLTFExporter/MigratedFromGit/Source/GLTFExporter/Private/Tasks/GLTFMaterialTasks.cpp b/Engine/Plugins/Enterprise/GLTFExporter/MigratedFromGit/Source/GLTFExporter/Private/Tasks/GLTFMaterialTasks.cpp
--- a/Engine/Plugins/Enterprise/GLTFExporter/MigratedFromGit/Source/GLTFExporter/Private/Tasks/GLTFMaterialTasks.cpp
+++ b/Engine/Plugins/Enterprise/GLTFExporter/MigratedFromGit/Source/GLTFExporter/Private/Tasks/GLTFMaterialTasks.cpp
@@ -201,6 +201,16 @@ EMaterialShadingModel FGLTFMaterialTask::GetShadingModel() const
 
 	if (PossibilitesCount > 1)
 	{
+		if (!FApp::CanEverRender())
+		{
+			EMaterialShadingModel ShadingModel = FGLTFMaterialUtility::GetRichestShadingModel(Possibilites);
+			Builder.AddWarningMessage(FString::Printf(
+				TEXT("Can't evaluate shading model expression in material %s because renderer missing, will export as %s"),
+				*Material->GetName(),
+				*FGLTFNameUtility::GetName(ShadingModel)));
+			return ShadingModel;
+		}
+
 		if (Material->IsShadingModelFromMaterialExpression())
 		{
 			const FMaterialShadingModelField Evaluation = FGLTFMaterialUtility::EvaluateShadingModelExpression(Material);
