[GLTFExporter] refactor: remove ugly hack to add attribute definition for TransmittanceColor

commit 0cdb7b852f8216d2ef45f00c805cc040dc4f38be
Date: Wed, 10 Nov 2021 20:23:51 +0100

    refactor: remove ugly hack to add attribute definition for TransmittanceColor
    (cherry picked from commit fdc47fd6fa042787b0462f386b96b2695b4d72d6​)

#skipci

diff --git a/Engine/Plugins/Enterprise/GLTFExporter/MigratedFromGit/GLTFExporter.uplugin b/Engine/Plugins/Enterprise/GLTFExporter/MigratedFromGit/GLTFExporter.uplugin
--- a/Engine/Plugins/Enterprise/GLTFExporter/MigratedFromGit/GLTFExporter.uplugin
+++ b/Engine/Plugins/Enterprise/GLTFExporter/MigratedFromGit/GLTFExporter.uplugin
@@ -35,7 +35,7 @@
 		{
 			"Name": "GLTFMaterialBaking",
 			"Type": "Editor",
-			"LoadingPhase": "Default",
+			"LoadingPhase": "PostConfigInit",
 			"WhitelistPlatforms": [ "Win64", "Win32", "Mac" ]
 		}
 	],
diff --git a/Engine/Plugins/Enterprise/GLTFExporter/MigratedFromGit/Source/GLTFMaterialBaking/Private/AddCustomAttributeHack.cpp b/Engine/Plugins/Enterprise/GLTFExporter/MigratedFromGit/Source/GLTFMaterialBaking/Private/AddCustomAttributeHack.cpp
deleted file mode 100644
diff --git a/Engine/Plugins/Enterprise/GLTFExporter/MigratedFromGit/Source/GLTFMaterialBaking/Private/AddCustomAttributeHack.h b/Engine/Plugins/Enterprise/GLTFExporter/MigratedFromGit/Source/GLTFMaterialBaking/Private/AddCustomAttributeHack.h
deleted file mode 100644
diff --git a/Engine/Plugins/Enterprise/GLTFExporter/MigratedFromGit/Source/GLTFMaterialBaking/Private/MaterialBakingModule.cpp b/Engine/Plugins/Enterprise/GLTFExporter/MigratedFromGit/Source/GLTFMaterialBaking/Private/MaterialBakingModule.cpp
--- a/Engine/Plugins/Enterprise/GLTFExporter/MigratedFromGit/Source/GLTFMaterialBaking/Private/MaterialBakingModule.cpp
+++ b/Engine/Plugins/Enterprise/GLTFExporter/MigratedFromGit/Source/GLTFMaterialBaking/Private/MaterialBakingModule.cpp
@@ -8,7 +8,6 @@
 #include "UObject/UObjectGlobals.h"
 #include "MaterialBakingStructures.h"
 #include "MaterialBakingHelpers.h"
-#include "AddCustomAttributeHack.h"
 #include "Async/Async.h"
 #include "Async/ParallelFor.h"
 #include "Materials/MaterialInstance.h"
@@ -265,7 +264,7 @@ void FMaterialBakingModule::StartupModule()
 	FCoreUObjectDelegates::OnObjectModified.AddRaw(this, &FMaterialBakingModule::OnObjectModified);
 
 	// NOTE: Because the attribute definition of TransmittanceColor is not registered by the engine, we have to do it here
-	AddCustomAttributeHack(FGuid(0xF2D8C70E, 0x42ECA0D1, 0x4652D0AD, 0xB785A065), "TransmittanceColor", "GetThinTranslucentMaterialOutput", MCT_Float3, FVector4(0.5, 0.5, 0.5, 0));
+	FMaterialAttributeDefinitionMap::AddCustomAttribute(FGuid(0xF2D8C70E, 0x42ECA0D1, 0x4652D0AD, 0xB785A065), "TransmittanceColor", "GetThinTranslucentMaterialOutput", MCT_Float3, FVector4(0.5, 0.5, 0.5, 0));
 }
 
 void FMaterialBakingModule::ShutdownModule()
