[GLTFExporter] Revert "refactor: remove ugly hack to add attribute definition for TransmittanceColor"

commit 0bf758e3d12f9692e0da75df5857afff82b427e1
Date: Mon, 15 Nov 2021 22:30:31 +0100

    Revert "refactor: remove ugly hack to add attribute definition for TransmittanceColor"
    This reverts commit fdc47fd6fa042787b0462f386b96b2695b4d72d6.
    (cherry picked from commit 1f76861086e51688674d2117680d132e9d2bc646​)

#skipci

diff --git a/Engine/Plugins/Enterprise/GLTFExporter/MigratedFromGit/GLTFExporter.uplugin b/Engine/Plugins/Enterprise/GLTFExporter/MigratedFromGit/GLTFExporter.uplugin
--- a/Engine/Plugins/Enterprise/GLTFExporter/MigratedFromGit/GLTFExporter.uplugin
+++ b/Engine/Plugins/Enterprise/GLTFExporter/MigratedFromGit/GLTFExporter.uplugin
@@ -35,7 +35,7 @@
 		{
 			"Name": "GLTFMaterialBaking",
 			"Type": "Editor",
-			"LoadingPhase": "PostConfigInit",
+			"LoadingPhase": "Default",
 			"WhitelistPlatforms": [ "Win64", "Win32", "Mac" ]
 		}
 	],
diff --git a/Engine/Plugins/Enterprise/GLTFExporter/MigratedFromGit/Source/GLTFMaterialBaking/Private/AddCustomAttributeHack.cpp b/Engine/Plugins/Enterprise/GLTFExporter/MigratedFromGit/Source/GLTFMaterialBaking/Private/AddCustomAttributeHack.cpp
new file mode 100644
--- /dev/null
+++ b/Engine/Plugins/Enterprise/GLTFExporter/MigratedFromGit/Source/GLTFMaterialBaking/Private/AddCustomAttributeHack.cpp
@@ -0,0 +1,35 @@
+// Copyright Epic Games, Inc. All Rights Reserved.
+
+#include "AddCustomAttributeHack.h"
+#include "MaterialShared.h"
+
+// NOTE: to access the raw array of custom attributes we need to use a dummy friend class
+class FMaterialAttributePropertyDetails
+{
+public:
+
+	static TArray<FMaterialCustomOutputAttributeDefintion>& GetCustomAttributes()
+	{
+		return FMaterialAttributeDefinitionMap::GMaterialPropertyAttributesMap.CustomAttributes;
+	}
+};
+
+void AddCustomAttributeHack(const FGuid& AttributeID, const FString& AttributeName, const FString& FunctionName, EMaterialValueType ValueType, const FVector4& DefaultValue, MaterialAttributeBlendFunction BlendFunction)
+{
+	// NOTE: since the linker doesn't have access to the constructor we need a empty buffer
+	uint8 Buffer[sizeof(FMaterialCustomOutputAttributeDefintion)] = { 0 };
+	FMaterialCustomOutputAttributeDefintion& CustomAttribute = *reinterpret_cast<FMaterialCustomOutputAttributeDefintion*>(Buffer);
+
+	CustomAttribute.AttributeID = AttributeID;
+	CustomAttribute.DefaultValue = DefaultValue;
+	CustomAttribute.AttributeName = AttributeName;
+	CustomAttribute.Property = MP_CustomOutput;
+	CustomAttribute.ValueType = ValueType;
+	CustomAttribute.ShaderFrequency = SF_Pixel;
+	CustomAttribute.TexCoordIndex = INDEX_NONE;
+	CustomAttribute.BlendFunction = BlendFunction;
+	CustomAttribute.bIsHidden = false;
+	CustomAttribute.FunctionName = FunctionName;
+
+	FMaterialAttributePropertyDetails::GetCustomAttributes().Add(CustomAttribute);
+}
diff --git a/Engine/Plugins/Enterprise/GLTFExporter/MigratedFromGit/Source/GLTFMaterialBaking/Private/AddCustomAttributeHack.h b/Engine/Plugins/Enterprise/GLTFExporter/MigratedFromGit/Source/GLTFMaterialBaking/Private/AddCustomAttributeHack.h
new file mode 100644
--- /dev/null
+++ b/Engine/Plugins/Enterprise/GLTFExporter/MigratedFromGit/Source/GLTFMaterialBaking/Private/AddCustomAttributeHack.h
@@ -0,0 +1,8 @@
+// Copyright Epic Games, Inc. All Rights Reserved.
+
+#pragma once
+
+#include "CoreMinimal.h"
+
+// NOTE: ugly hack to register custom attribute without calling FMaterialAttributeDefinitionMap::AddCustomAttribute (which doesn't work after the AttributeDDCString has been initialized)
+void AddCustomAttributeHack(const FGuid& AttributeID, const FString& AttributeName, const FString& FunctionName, EMaterialValueType ValueType, const FVector4& DefaultValue, MaterialAttributeBlendFunction BlendFunction = nullptr);
diff --git a/Engine/Plugins/Enterprise/GLTFExporter/MigratedFromGit/Source/GLTFMaterialBaking/Private/MaterialBakingModule.cpp b/Engine/Plugins/Enterprise/GLTFExporter/MigratedFromGit/Source/GLTFMaterialBaking/Private/MaterialBakingModule.cpp
--- a/Engine/Plugins/Enterprise/GLTFExporter/MigratedFromGit/Source/GLTFMaterialBaking/Private/MaterialBakingModule.cpp
+++ b/Engine/Plugins/Enterprise/GLTFExporter/MigratedFromGit/Source/GLTFMaterialBaking/Private/MaterialBakingModule.cpp
@@ -8,6 +8,7 @@
 #include "UObject/UObjectGlobals.h"
 #include "MaterialBakingStructures.h"
 #include "MaterialBakingHelpers.h"
+#include "AddCustomAttributeHack.h"
 #include "Async/Async.h"
 #include "Async/ParallelFor.h"
 #include "Materials/MaterialInstance.h"
@@ -264,7 +265,7 @@ void FMaterialBakingModule::StartupModule()
 	FCoreUObjectDelegates::OnObjectModified.AddRaw(this, &FMaterialBakingModule::OnObjectModified);
 
 	// NOTE: Because the attribute definition of TransmittanceColor is not registered by the engine, we have to do it here
-	FMaterialAttributeDefinitionMap::AddCustomAttribute(FGuid(0xF2D8C70E, 0x42ECA0D1, 0x4652D0AD, 0xB785A065), "TransmittanceColor", "GetThinTranslucentMaterialOutput", MCT_Float3, FVector4(0.5, 0.5, 0.5, 0));
+	AddCustomAttributeHack(FGuid(0xF2D8C70E, 0x42ECA0D1, 0x4652D0AD, 0xB785A065), "TransmittanceColor", "GetThinTranslucentMaterialOutput", MCT_Float3, FVector4(0.5, 0.5, 0.5, 0));
 }
 
 void FMaterialBakingModule::ShutdownModule()
