[GLTFExporter] fix: define plugin version info as preprocessors to avoid PluginManager (which isn't available in runtime)

commit 463d5e5e8b17bc1c8064f049b2f09f4a8d54049e
Date: Sun, 31 Oct 2021 01:27:35 +0200

    fix: define plugin version info as preprocessors to avoid PluginManager (which isn't available in runtime)
    (cherry picked from commit f66eb8e4e3f7fc10b034f871c1ebe75aa37df9a8​)

#skipci

diff --git a/Engine/Plugins/Enterprise/GLTFExporter/MigratedFromGit/Source/GLTFExporter/GLTFExporter.Build.cs b/Engine/Plugins/Enterprise/GLTFExporter/MigratedFromGit/Source/GLTFExporter/GLTFExporter.Build.cs
--- a/Engine/Plugins/Enterprise/GLTFExporter/MigratedFromGit/Source/GLTFExporter/GLTFExporter.Build.cs
+++ b/Engine/Plugins/Enterprise/GLTFExporter/MigratedFromGit/Source/GLTFExporter/GLTFExporter.Build.cs
@@ -6,6 +6,17 @@ namespace UnrealBuildTool.Rules
 	{
 		public GLTFExporter(ReadOnlyTargetRules Target) : base(Target)
 		{
+		    // NOTE: ugly hack to access plugin info (should propose change to engine)
+		    var BindingFlags = System.Reflection.BindingFlags.Instance | System.Reflection.BindingFlags.NonPublic;
+		    var FieldInfo = typeof(ModuleRules).GetField("Plugin", BindingFlags);
+			var Plugin = FieldInfo != null ? FieldInfo.GetValue(this) as PluginInfo : null;
+
+			if (Plugin != null && Plugin.Descriptor != null)
+			{
+				PrivateDefinitions.Add("GLTFEXPORTER_FRIENDLY_NAME=TEXT(\"" + Plugin.Descriptor.FriendlyName + "\")");
+				PrivateDefinitions.Add("GLTFEXPORTER_VERSION_NAME=TEXT(\"" + Plugin.Descriptor.VersionName + "\")");
+			}
+
 			PrivateDependencyModuleNames.AddRange(
 				new string[]
 				{
diff --git a/Engine/Plugins/Enterprise/GLTFExporter/MigratedFromGit/Source/GLTFExporter/Private/Builders/GLTFJsonBuilder.cpp b/Engine/Plugins/Enterprise/GLTFExporter/MigratedFromGit/Source/GLTFExporter/Private/Builders/GLTFJsonBuilder.cpp
--- a/Engine/Plugins/Enterprise/GLTFExporter/MigratedFromGit/Source/GLTFExporter/Private/Builders/GLTFJsonBuilder.cpp
+++ b/Engine/Plugins/Enterprise/GLTFExporter/MigratedFromGit/Source/GLTFExporter/Private/Builders/GLTFJsonBuilder.cpp
@@ -1,8 +1,6 @@
 // Copyright Epic Games, Inc. All Rights Reserved.
 
 #include "Builders/GLTFJsonBuilder.h"
-#include "GLTFExporterModule.h"
-#include "Interfaces/IPluginManager.h"
 #include "Runtime/Launch/Resources/Version.h"
 
 FGLTFJsonBuilder::FGLTFJsonBuilder(const FString& FilePath, const UGLTFExportOptions* ExportOptions)
@@ -268,12 +266,9 @@ FGLTFJsonNodeIndex FGLTFJsonBuilder::GetComponentNodeIndex(FGLTFJsonNodeIndex No
 
 FString FGLTFJsonBuilder::GetGeneratorString() const
 {
-	const TSharedPtr<IPlugin> Plugin = IPluginManager::Get().FindPlugin(GLTFEXPORTER_MODULE_NAME);
-	const FPluginDescriptor& PluginDescriptor = Plugin->GetDescriptor();
-
 	return ExportOptions->bIncludeGeneratorVersion
-		? TEXT(EPIC_PRODUCT_NAME) TEXT(" ") ENGINE_VERSION_STRING TEXT(" ") + PluginDescriptor.FriendlyName + TEXT(" ") + PluginDescriptor.VersionName
-		: TEXT(EPIC_PRODUCT_NAME) TEXT(" ") + PluginDescriptor.FriendlyName;
+		? TEXT(EPIC_PRODUCT_NAME) TEXT(" ") ENGINE_VERSION_STRING TEXT(" ") GLTFEXPORTER_FRIENDLY_NAME TEXT(" ") GLTFEXPORTER_VERSION_NAME
+		: TEXT(EPIC_PRODUCT_NAME) TEXT(" ") GLTFEXPORTER_FRIENDLY_NAME;
 }
 
 bool FGLTFJsonBuilder::IsCustomExtension(EGLTFJsonExtension Extension)
diff --git a/Engine/Plugins/Enterprise/GLTFExporter/MigratedFromGit/Source/GLTFExporter/Private/Builders/GLTFMessageBuilder.cpp b/Engine/Plugins/Enterprise/GLTFExporter/MigratedFromGit/Source/GLTFExporter/Private/Builders/GLTFMessageBuilder.cpp
--- a/Engine/Plugins/Enterprise/GLTFExporter/MigratedFromGit/Source/GLTFExporter/Private/Builders/GLTFMessageBuilder.cpp
+++ b/Engine/Plugins/Enterprise/GLTFExporter/MigratedFromGit/Source/GLTFExporter/Private/Builders/GLTFMessageBuilder.cpp
@@ -2,7 +2,6 @@
 
 #include "Builders/GLTFMessageBuilder.h"
 #include "GLTFExporterModule.h"
-#include "Interfaces/IPluginManager.h"
 #include "MessageLogModule.h"
 #include "IMessageLogListing.h"
 
@@ -92,10 +91,7 @@ void FGLTFMessageBuilder::ShowMessages() const
 		FMessageLogModule& MessageLogModule = FModuleManager::LoadModuleChecked<FMessageLogModule>("MessageLog");
 		const TSharedRef<IMessageLogListing> LogListing = MessageLogModule.GetLogListing(GLTFEXPORTER_MODULE_NAME);
 
-		const TSharedPtr<IPlugin> Plugin = IPluginManager::Get().FindPlugin(GLTFEXPORTER_MODULE_NAME);
-		const FPluginDescriptor& PluginDescriptor = Plugin->GetDescriptor();
-
-		LogListing->SetLabel(FText::FromString(PluginDescriptor.FriendlyName));
+		LogListing->SetLabel(FText::FromString(GLTFEXPORTER_FRIENDLY_NAME));
 		LogListing->ClearMessages();
 
 		for (const FLogMessage& Message : Messages)
