[GLTFExporter] feat: continuously log all messages instead of printing them last (to ease debugging)

commit e2ee6d21ef4b0459f97723713600cad128243eee
Date: Sun, 31 Oct 2021 02:30:26 +0200

    feat: continuously log all messages instead of printing them last (to ease debugging)
    (cherry picked from commit a9387767260ca7abcfa6afeed7302b80eda6c2d1​)

#skipci

diff --git a/Engine/Plugins/Enterprise/GLTFExporter/MigratedFromGit/Source/GLTFExporter/Private/Builders/GLTFMessageBuilder.cpp b/Engine/Plugins/Enterprise/GLTFExporter/MigratedFromGit/Source/GLTFExporter/Private/Builders/GLTFMessageBuilder.cpp
--- a/Engine/Plugins/Enterprise/GLTFExporter/MigratedFromGit/Source/GLTFExporter/Private/Builders/GLTFMessageBuilder.cpp
+++ b/Engine/Plugins/Enterprise/GLTFExporter/MigratedFromGit/Source/GLTFExporter/Private/Builders/GLTFMessageBuilder.cpp
@@ -4,139 +4,104 @@
 #include "GLTFExporterModule.h"
 #include "MessageLogModule.h"
 #include "IMessageLogListing.h"
+#include "Misc/FeedbackContext.h"
 
 FGLTFMessageBuilder::FGLTFMessageBuilder(const FString& FilePath, const UGLTFExportOptions* ExportOptions)
 	: FGLTFBuilder(FilePath, ExportOptions)
 {
-}
-
-void FGLTFMessageBuilder::ClearMessages()
-{
-	Messages.Empty();
-}
-
-void FGLTFMessageBuilder::AddMessage(EGLTFMessageSeverity Severity, const FString& Message)
-{
-	Messages.Emplace(Severity, Message);
+	if (!FApp::IsUnattended())
+	{
+		FMessageLogModule& MessageLogModule = FModuleManager::LoadModuleChecked<FMessageLogModule>("MessageLog");
+		LogListing = MessageLogModule.GetLogListing(GLTFEXPORTER_MODULE_NAME);
+		LogListing->SetLabel(FText::FromString(GLTFEXPORTER_FRIENDLY_NAME));
+	}
 }
 
 void FGLTFMessageBuilder::AddInfoMessage(const FString& Message)
 {
-	AddMessage(EGLTFMessageSeverity::Info, Message);
+	Infos.Add(Message);
+	PrintToLog(ELogLevel::Info, Message);
 }
 
 void FGLTFMessageBuilder::AddWarningMessage(const FString& Message)
 {
-	AddMessage(EGLTFMessageSeverity::Warning, Message);
+	Warnings.Add(Message);
+	PrintToLog(ELogLevel::Warning, Message);
 }
 
 void FGLTFMessageBuilder::AddErrorMessage(const FString& Message)
 {
-	AddMessage(EGLTFMessageSeverity::Error, Message);
-}
-
-const TArray<FGLTFMessageBuilder::FLogMessage>& FGLTFMessageBuilder::GetMessages() const
-{
-	return Messages;
-}
-
-TArray<FGLTFMessageBuilder::FLogMessage> FGLTFMessageBuilder::GetMessages(EGLTFMessageSeverity Severity) const
-{
-	return Messages.FilterByPredicate(
-		[Severity](const FLogMessage& LogMessage)
-		{
-			return LogMessage.Key == Severity;
-		});
-}
-
-TArray<FGLTFMessageBuilder::FLogMessage> FGLTFMessageBuilder::GetInfoMessages() const
-{
-	return GetMessages(EGLTFMessageSeverity::Info);
-}
-
-TArray<FGLTFMessageBuilder::FLogMessage> FGLTFMessageBuilder::GetWarningMessages() const
-{
-	return GetMessages(EGLTFMessageSeverity::Warning);
+	Errors.Add(Message);
+	PrintToLog(ELogLevel::Error, Message);
 }
 
-TArray<FGLTFMessageBuilder::FLogMessage> FGLTFMessageBuilder::GetErrorMessages() const
+const TArray<FString>& FGLTFMessageBuilder::GetInfoMessages() const
 {
-	return GetMessages(EGLTFMessageSeverity::Error);
+	return Infos;
 }
 
-int32 FGLTFMessageBuilder::GetMessageCount() const
+const TArray<FString>& FGLTFMessageBuilder::GetWarningMessages() const
 {
-	return Messages.Num();
+	return Warnings;
 }
 
-int32 FGLTFMessageBuilder::GetInfoMessageCount() const
+const TArray<FString>& FGLTFMessageBuilder::GetErrorMessages() const
 {
-	return GetInfoMessages().Num();
+	return Errors;
 }
 
-int32 FGLTFMessageBuilder::GetWarningMessageCount() const
+void FGLTFMessageBuilder::OpenLog() const
 {
-	return GetWarningMessages().Num();
+	if (LogListing != nullptr)
+	{
+		LogListing->Open();
+	}
 }
 
-int32 FGLTFMessageBuilder::GetErrorMessageCount() const
+void FGLTFMessageBuilder::ClearLog()
 {
-	return GetErrorMessages().Num();
-}
+	Infos.Empty();
+	Warnings.Empty();
+	Errors.Empty();
 
-void FGLTFMessageBuilder::ShowMessages() const
-{
-	if (Messages.Num() > 0)
+	if (LogListing != nullptr)
 	{
-		FMessageLogModule& MessageLogModule = FModuleManager::LoadModuleChecked<FMessageLogModule>("MessageLog");
-		const TSharedRef<IMessageLogListing> LogListing = MessageLogModule.GetLogListing(GLTFEXPORTER_MODULE_NAME);
-
-		LogListing->SetLabel(FText::FromString(GLTFEXPORTER_FRIENDLY_NAME));
 		LogListing->ClearMessages();
-
-		for (const FLogMessage& Message : Messages)
-		{
-			LogListing->AddMessage(CreateTokenizedMessage(Message));
-		}
-
-		MessageLogModule.OpenMessageLog(GLTFEXPORTER_MODULE_NAME);
 	}
 }
 
-void FGLTFMessageBuilder::WriteMessagesToConsole() const
+void FGLTFMessageBuilder::PrintToLog(ELogLevel Level, const FString& Message) const
 {
-	for (const FLogMessage& Message : Messages)
-	{
-		WriteMessageToConsole(Message);
-	}
-}
+#if !NO_LOGGING
+	ELogVerbosity::Type Verbosity;
 
-void FGLTFMessageBuilder::WriteMessageToConsole(const FLogMessage& LogMessage)
-{
-	switch (LogMessage.Key)
+	switch (Level)
 	{
-		case EGLTFMessageSeverity::Info:    UE_LOG(LogGLTFExporter, Display, TEXT("%s"), *LogMessage.Value); break;
-		case EGLTFMessageSeverity::Warning: UE_LOG(LogGLTFExporter, Warning, TEXT("%s"), *LogMessage.Value); break;
-		case EGLTFMessageSeverity::Error:   UE_LOG(LogGLTFExporter, Error, TEXT("%s"), *LogMessage.Value); break;
+		case ELogLevel::Info:    Verbosity = ELogVerbosity::Display; break;
+		case ELogLevel::Warning: Verbosity = ELogVerbosity::Warning; break;
+		case ELogLevel::Error:   Verbosity = ELogVerbosity::Error; break;
 		default:
 			checkNoEntry();
-			break;
+			return;
 	}
-}
 
-TSharedRef<FTokenizedMessage> FGLTFMessageBuilder::CreateTokenizedMessage(const FLogMessage& LogMessage)
-{
-	EMessageSeverity::Type MessageSeverity = EMessageSeverity::Type::CriticalError;
+	GWarn->Log(LogGLTFExporter.GetCategoryName(), Verbosity, Message);
+#endif
 
-	switch (LogMessage.Key)
+	if (LogListing != nullptr)
 	{
-		case EGLTFMessageSeverity::Info:    MessageSeverity = EMessageSeverity::Type::Info; break;
-		case EGLTFMessageSeverity::Warning: MessageSeverity = EMessageSeverity::Type::Warning; break;
-		case EGLTFMessageSeverity::Error:   MessageSeverity = EMessageSeverity::Type::Error; break;
-		default:
-			checkNoEntry();
-			break;
-	}
+		EMessageSeverity::Type Severity;
 
-	return FTokenizedMessage::Create(MessageSeverity, FText::FromString(LogMessage.Value));
+		switch (Level)
+		{
+			case ELogLevel::Info:    Severity = EMessageSeverity::Info; break;
+			case ELogLevel::Warning: Severity = EMessageSeverity::Warning; break;
+			case ELogLevel::Error:   Severity = EMessageSeverity::Error; break;
+			default:
+				checkNoEntry();
+				return;
+		}
+
+		LogListing->AddMessage(FTokenizedMessage::Create(Severity, FText::FromString(Message)), false);
+	}
 }
diff --git a/Engine/Plugins/Enterprise/GLTFExporter/MigratedFromGit/Source/GLTFExporter/Private/Builders/GLTFMessageBuilder.h b/Engine/Plugins/Enterprise/GLTFExporter/MigratedFromGit/Source/GLTFExporter/Private/Builders/GLTFMessageBuilder.h
--- a/Engine/Plugins/Enterprise/GLTFExporter/MigratedFromGit/Source/GLTFExporter/Private/Builders/GLTFMessageBuilder.h
+++ b/Engine/Plugins/Enterprise/GLTFExporter/MigratedFromGit/Source/GLTFExporter/Private/Builders/GLTFMessageBuilder.h
@@ -4,12 +4,7 @@
 
 #include "Builders/GLTFBuilder.h"
 
-enum class EGLTFMessageSeverity
-{
-	Info,
-	Warning,
-	Error
-};
+class IMessageLogListing;
 
 class FGLTFMessageBuilder : public FGLTFBuilder
 {
@@ -19,45 +14,36 @@ protected:
 
 public:
 
-	typedef TTuple<EGLTFMessageSeverity, FString> FLogMessage;
-
-	void ClearMessages();
-
-	void AddMessage(EGLTFMessageSeverity Severity, const FString& Message);
-
 	void AddInfoMessage(const FString& Message);
 
 	void AddWarningMessage(const FString& Message);
 
 	void AddErrorMessage(const FString& Message);
 
-	const TArray<FLogMessage>& GetMessages() const;
-
-	TArray<FLogMessage> GetMessages(EGLTFMessageSeverity Severity) const;
+	const TArray<FString>& GetInfoMessages() const;
 
-	TArray<FLogMessage> GetInfoMessages() const;
+	const TArray<FString>& GetWarningMessages() const;
 
-	TArray<FLogMessage> GetWarningMessages() const;
+	const TArray<FString>& GetErrorMessages() const;
 
-	TArray<FLogMessage> GetErrorMessages() const;
+	void OpenLog() const;
 
-	int32 GetMessageCount() const;
-
-	int32 GetInfoMessageCount() const;
-
-	int32 GetWarningMessageCount() const;
-
-	int32 GetErrorMessageCount() const;
-
-	void ShowMessages() const;
-
-	void WriteMessagesToConsole() const;
+	void ClearLog();
 
 private:
 
-	static void WriteMessageToConsole(const FLogMessage& LogMessage);
+	enum class ELogLevel
+	{
+		Info,
+		Warning,
+		Error,
+	};
+
+	void PrintToLog(ELogLevel Level, const FString& Message) const;
 
-	static TSharedRef<FTokenizedMessage> CreateTokenizedMessage(const FLogMessage& LogMessage);
+	TArray<FString> Infos;
+	TArray<FString> Warnings;
+	TArray<FString> Errors;
 
-	TArray<FLogMessage> Messages;
+	TSharedPtr<IMessageLogListing> LogListing;
 };
diff --git a/Engine/Plugins/Enterprise/GLTFExporter/MigratedFromGit/Source/GLTFExporter/Private/Exporters/GLTFExporter.cpp b/Engine/Plugins/Enterprise/GLTFExporter/MigratedFromGit/Source/GLTFExporter/Private/Exporters/GLTFExporter.cpp
--- a/Engine/Plugins/Enterprise/GLTFExporter/MigratedFromGit/Source/GLTFExporter/Private/Exporters/GLTFExporter.cpp
+++ b/Engine/Plugins/Enterprise/GLTFExporter/MigratedFromGit/Source/GLTFExporter/Private/Exporters/GLTFExporter.cpp
@@ -30,10 +30,12 @@ bool UGLTFExporter::ExportBinary(UObject* Object, const TCHAR* Type, FArchive& A
 		return false;
 	}
 
+	FGCObjectScopeGuard OptionsGuard(Options);
+
 	// TODO: add support for UAssetExportTask::IgnoreObjectList?
 
-	FGCObjectScopeGuard OptionsGuard(Options);
 	FGLTFContainerBuilder Builder(CurrentFilename, Options, bSelectedOnly);
+	Builder.ClearLog();
 
 	const bool bSuccess = AddObject(Builder, Object);
 	if (bSuccess)
@@ -43,13 +45,9 @@ bool UGLTFExporter::ExportBinary(UObject* Object, const TCHAR* Type, FArchive& A
 
 	// TODO: should we copy messages to UAssetExportTask::Errors?
 
-	if (FApp::IsUnattended())
-	{
-		Builder.WriteMessagesToConsole();
-	}
-	else
+	if (!FApp::IsUnattended())
 	{
-		Builder.ShowMessages();
+		Builder.OpenLog();
 	}
 
 	return bSuccess;
