[GLTFExporter] refactor: rename info messages to suggestions

commit 81d225c2ee993d7b86dc579f2b9cc17b2981aaab
Date: Sun, 31 Oct 2021 02:41:14 +0200

    refactor: rename info messages to suggestions
    (cherry picked from commit ee17c6360ee4177ce62c7e9dba6c51b7ecf479af​)

#skipci

diff --git a/Engine/Plugins/Enterprise/GLTFExporter/MigratedFromGit/Source/GLTFExporter/Private/Builders/GLTFMessageBuilder.cpp b/Engine/Plugins/Enterprise/GLTFExporter/MigratedFromGit/Source/GLTFExporter/Private/Builders/GLTFMessageBuilder.cpp
--- a/Engine/Plugins/Enterprise/GLTFExporter/MigratedFromGit/Source/GLTFExporter/Private/Builders/GLTFMessageBuilder.cpp
+++ b/Engine/Plugins/Enterprise/GLTFExporter/MigratedFromGit/Source/GLTFExporter/Private/Builders/GLTFMessageBuilder.cpp
@@ -17,10 +17,10 @@ FGLTFMessageBuilder::FGLTFMessageBuilder(const FString& FilePath, const UGLTFExp
 	}
 }
 
-void FGLTFMessageBuilder::AddInfoMessage(const FString& Message)
+void FGLTFMessageBuilder::AddSuggestionMessage(const FString& Message)
 {
-	Infos.Add(Message);
-	PrintToLog(ELogLevel::Info, Message);
+	Suggestions.Add(Message);
+	PrintToLog(ELogLevel::Suggestion, Message);
 }
 
 void FGLTFMessageBuilder::AddWarningMessage(const FString& Message)
@@ -35,9 +35,9 @@ void FGLTFMessageBuilder::AddErrorMessage(const FString& Message)
 	PrintToLog(ELogLevel::Error, Message);
 }
 
-const TArray<FString>& FGLTFMessageBuilder::GetInfoMessages() const
+const TArray<FString>& FGLTFMessageBuilder::GetSuggestionMessages() const
 {
-	return Infos;
+	return Suggestions;
 }
 
 const TArray<FString>& FGLTFMessageBuilder::GetWarningMessages() const
@@ -60,7 +60,7 @@ void FGLTFMessageBuilder::OpenLog() const
 
 void FGLTFMessageBuilder::ClearLog()
 {
-	Infos.Empty();
+	Suggestions.Empty();
 	Warnings.Empty();
 	Errors.Empty();
 
@@ -77,9 +77,9 @@ void FGLTFMessageBuilder::PrintToLog(ELogLevel Level, const FString& Message) co
 
 	switch (Level)
 	{
-		case ELogLevel::Info:    Verbosity = ELogVerbosity::Display; break;
-		case ELogLevel::Warning: Verbosity = ELogVerbosity::Warning; break;
-		case ELogLevel::Error:   Verbosity = ELogVerbosity::Error; break;
+		case ELogLevel::Suggestion: Verbosity = ELogVerbosity::Display; break;
+		case ELogLevel::Warning:    Verbosity = ELogVerbosity::Warning; break;
+		case ELogLevel::Error:      Verbosity = ELogVerbosity::Error; break;
 		default:
 			checkNoEntry();
 			return;
@@ -94,9 +94,9 @@ void FGLTFMessageBuilder::PrintToLog(ELogLevel Level, const FString& Message) co
 
 		switch (Level)
 		{
-			case ELogLevel::Info:    Severity = EMessageSeverity::Info; break;
-			case ELogLevel::Warning: Severity = EMessageSeverity::Warning; break;
-			case ELogLevel::Error:   Severity = EMessageSeverity::Error; break;
+			case ELogLevel::Suggestion: Severity = EMessageSeverity::Info; break;
+			case ELogLevel::Warning:    Severity = EMessageSeverity::Warning; break;
+			case ELogLevel::Error:      Severity = EMessageSeverity::Error; break;
 			default:
 				checkNoEntry();
 				return;
diff --git a/Engine/Plugins/Enterprise/GLTFExporter/MigratedFromGit/Source/GLTFExporter/Private/Builders/GLTFMessageBuilder.h b/Engine/Plugins/Enterprise/GLTFExporter/MigratedFromGit/Source/GLTFExporter/Private/Builders/GLTFMessageBuilder.h
--- a/Engine/Plugins/Enterprise/GLTFExporter/MigratedFromGit/Source/GLTFExporter/Private/Builders/GLTFMessageBuilder.h
+++ b/Engine/Plugins/Enterprise/GLTFExporter/MigratedFromGit/Source/GLTFExporter/Private/Builders/GLTFMessageBuilder.h
@@ -14,13 +14,13 @@ protected:
 
 public:
 
-	void AddInfoMessage(const FString& Message);
+	void AddSuggestionMessage(const FString& Message);
 
 	void AddWarningMessage(const FString& Message);
 
 	void AddErrorMessage(const FString& Message);
 
-	const TArray<FString>& GetInfoMessages() const;
+	const TArray<FString>& GetSuggestionMessages() const;
 
 	const TArray<FString>& GetWarningMessages() const;
 
@@ -34,14 +34,14 @@ private:
 
 	enum class ELogLevel
 	{
-		Info,
+		Suggestion,
 		Warning,
 		Error,
 	};
 
 	void PrintToLog(ELogLevel Level, const FString& Message) const;
 
-	TArray<FString> Infos;
+	TArray<FString> Suggestions;
 	TArray<FString> Warnings;
 	TArray<FString> Errors;
 
