[GLTFExporter] refactor: rename GLTFMessageBuilder to GLTFLogBuilder and rename some of the methods

commit 138761cb8711bff3da6af7be6ee33d80e00b5a74
Date: Sun, 31 Oct 2021 02:45:20 +0200

    refactor: rename GLTFMessageBuilder to GLTFLogBuilder and rename some of the methods
    (cherry picked from commit 82a7b562fef2df0eb571e91319be305ab04bd052​)

#skipci

diff --git a/Engine/Plugins/Enterprise/GLTFExporter/MigratedFromGit/Source/GLTFExporter/Private/Builders/GLTFBufferBuilder.cpp b/Engine/Plugins/Enterprise/GLTFExporter/MigratedFromGit/Source/GLTFExporter/Private/Builders/GLTFBufferBuilder.cpp
--- a/Engine/Plugins/Enterprise/GLTFExporter/MigratedFromGit/Source/GLTFExporter/Private/Builders/GLTFBufferBuilder.cpp
+++ b/Engine/Plugins/Enterprise/GLTFExporter/MigratedFromGit/Source/GLTFExporter/Private/Builders/GLTFBufferBuilder.cpp
@@ -33,7 +33,7 @@ bool FGLTFBufferBuilder::InitializeBuffer()
 		BufferArchive.Reset(IFileManager::Get().CreateFileWriter(*ExternalBinaryPath));
 		if (BufferArchive == nullptr)
 		{
-			AddErrorMessage(FString::Printf(TEXT("Failed to write external binary buffer to file: %s"), *ExternalBinaryPath));
+			LogError(FString::Printf(TEXT("Failed to write external binary buffer to file: %s"), *ExternalBinaryPath));
 			return false;
 		}
 	}
diff --git a/Engine/Plugins/Enterprise/GLTFExporter/MigratedFromGit/Source/GLTFExporter/Private/Builders/GLTFContainerBuilder.cpp b/Engine/Plugins/Enterprise/GLTFExporter/MigratedFromGit/Source/GLTFExporter/Private/Builders/GLTFContainerBuilder.cpp
--- a/Engine/Plugins/Enterprise/GLTFExporter/MigratedFromGit/Source/GLTFExporter/Private/Builders/GLTFContainerBuilder.cpp
+++ b/Engine/Plugins/Enterprise/GLTFExporter/MigratedFromGit/Source/GLTFExporter/Private/Builders/GLTFContainerBuilder.cpp
@@ -24,7 +24,7 @@ void FGLTFContainerBuilder::Write(FArchive& Archive, FFeedbackContext* Context)
 			return FGLTFJsonUtility::GetValue(Extension);
 		});
 
-		AddWarningMessage(FString::Printf(TEXT("Export uses some extensions that may only be supported in Unreal's glTF viewer: %s"), *ExtensionsString));
+		LogWarning(FString::Printf(TEXT("Export uses some extensions that may only be supported in Unreal's glTF viewer: %s"), *ExtensionsString));
 	}
 }
 
diff --git a/Engine/Plugins/Enterprise/GLTFExporter/MigratedFromGit/Source/GLTFExporter/Private/Builders/GLTFImageBuilder.cpp b/Engine/Plugins/Enterprise/GLTFExporter/MigratedFromGit/Source/GLTFExporter/Private/Builders/GLTFImageBuilder.cpp
--- a/Engine/Plugins/Enterprise/GLTFExporter/MigratedFromGit/Source/GLTFExporter/Private/Builders/GLTFImageBuilder.cpp
+++ b/Engine/Plugins/Enterprise/GLTFExporter/MigratedFromGit/Source/GLTFExporter/Private/Builders/GLTFImageBuilder.cpp
@@ -108,7 +108,7 @@ FString FGLTFImageBuilder::SaveImageToFile(const void* CompressedData, int64 Com
 
 	if (!FFileHelper::SaveArrayToFile(ImageData, *ImagePath))
 	{
-		AddErrorMessage(FString::Printf(TEXT("Failed to save image %s to file: %s"), *Name, *ImagePath));
+		LogError(FString::Printf(TEXT("Failed to save image %s to file: %s"), *Name, *ImagePath));
 		return TEXT("");
 	}
 
diff --git a/Engine/Plugins/Enterprise/GLTFExporter/MigratedFromGit/Source/GLTFExporter/Private/Builders/GLTFLogBuilder.cpp b/Engine/Plugins/Enterprise/GLTFExporter/MigratedFromGit/Source/GLTFExporter/Private/Builders/GLTFLogBuilder.cpp
new file mode 100644
--- /dev/null
+++ b/Engine/Plugins/Enterprise/GLTFExporter/MigratedFromGit/Source/GLTFExporter/Private/Builders/GLTFLogBuilder.cpp
@@ -0,0 +1,107 @@
+// Copyright Epic Games, Inc. All Rights Reserved.
+
+#include "Builders/GLTFLogBuilder.h"
+#include "GLTFExporterModule.h"
+#include "MessageLogModule.h"
+#include "IMessageLogListing.h"
+#include "Misc/FeedbackContext.h"
+
+FGLTFLogBuilder::FGLTFLogBuilder(const FString& FilePath, const UGLTFExportOptions* ExportOptions)
+	: FGLTFBuilder(FilePath, ExportOptions)
+{
+	if (!FApp::IsUnattended())
+	{
+		FMessageLogModule& MessageLogModule = FModuleManager::LoadModuleChecked<FMessageLogModule>("MessageLog");
+		LogListing = MessageLogModule.GetLogListing(GLTFEXPORTER_MODULE_NAME);
+		LogListing->SetLabel(FText::FromString(GLTFEXPORTER_FRIENDLY_NAME));
+	}
+}
+
+void FGLTFLogBuilder::LogSuggestion(const FString& Message)
+{
+	Suggestions.Add(Message);
+	PrintToLog(ELogLevel::Suggestion, Message);
+}
+
+void FGLTFLogBuilder::LogWarning(const FString& Message)
+{
+	Warnings.Add(Message);
+	PrintToLog(ELogLevel::Warning, Message);
+}
+
+void FGLTFLogBuilder::LogError(const FString& Message)
+{
+	Errors.Add(Message);
+	PrintToLog(ELogLevel::Error, Message);
+}
+
+const TArray<FString>& FGLTFLogBuilder::GetLoggedSuggestions() const
+{
+	return Suggestions;
+}
+
+const TArray<FString>& FGLTFLogBuilder::GetLoggedWarnings() const
+{
+	return Warnings;
+}
+
+const TArray<FString>& FGLTFLogBuilder::GetLoggedErrors() const
+{
+	return Errors;
+}
+
+void FGLTFLogBuilder::OpenLog() const
+{
+	if (LogListing != nullptr)
+	{
+		LogListing->Open();
+	}
+}
+
+void FGLTFLogBuilder::ClearLog()
+{
+	Suggestions.Empty();
+	Warnings.Empty();
+	Errors.Empty();
+
+	if (LogListing != nullptr)
+	{
+		LogListing->ClearMessages();
+	}
+}
+
+void FGLTFLogBuilder::PrintToLog(ELogLevel Level, const FString& Message) const
+{
+#if !NO_LOGGING
+	ELogVerbosity::Type Verbosity;
+
+	switch (Level)
+	{
+		case ELogLevel::Suggestion: Verbosity = ELogVerbosity::Display; break;
+		case ELogLevel::Warning:    Verbosity = ELogVerbosity::Warning; break;
+		case ELogLevel::Error:      Verbosity = ELogVerbosity::Error; break;
+		default:
+			checkNoEntry();
+			return;
+	}
+
+	GWarn->Log(LogGLTFExporter.GetCategoryName(), Verbosity, Message);
+#endif
+
+	if (LogListing != nullptr)
+	{
+		EMessageSeverity::Type Severity;
+
+		switch (Level)
+		{
+			case ELogLevel::Suggestion: Severity = EMessageSeverity::Info; break;
+			case ELogLevel::Warning:    Severity = EMessageSeverity::Warning; break;
+			case ELogLevel::Error:      Severity = EMessageSeverity::Error; break;
+			default:
+				checkNoEntry();
+				return;
+		}
+
+		LogListing->AddMessage(FTokenizedMessage::Create(Severity, FText::FromString(Message)), false);
+	}
+}
diff --git a/Engine/Plugins/Enterprise/GLTFExporter/MigratedFromGit/Source/GLTFExporter/Private/Builders/GLTFLogBuilder.h b/Engine/Plugins/Enterprise/GLTFExporter/MigratedFromGit/Source/GLTFExporter/Private/Builders/GLTFLogBuilder.h
new file mode 100644
--- /dev/null
+++ b/Engine/Plugins/Enterprise/GLTFExporter/MigratedFromGit/Source/GLTFExporter/Private/Builders/GLTFLogBuilder.h
@@ -0,0 +1,49 @@
+// Copyright Epic Games, Inc. All Rights Reserved.
+
+#pragma once
+
+#include "Builders/GLTFBuilder.h"
+
+class IMessageLogListing;
+
+class FGLTFLogBuilder : public FGLTFBuilder
+{
+protected:
+
+	FGLTFLogBuilder(const FString& FilePath, const UGLTFExportOptions* ExportOptions);
+
+public:
+
+	void LogSuggestion(const FString& Message);
+
+	void LogWarning(const FString& Message);
+
+	void LogError(const FString& Message);
+
+	const TArray<FString>& GetLoggedSuggestions() const;
+
+	const TArray<FString>& GetLoggedWarnings() const;
+
+	const TArray<FString>& GetLoggedErrors() const;
+
+	void OpenLog() const;
+
+	void ClearLog();
+
+private:
+
+	enum class ELogLevel
+	{
+		Suggestion,
+		Warning,
+		Error,
+	};
+
+	void PrintToLog(ELogLevel Level, const FString& Message) const;
+
+	TArray<FString> Suggestions;
+	TArray<FString> Warnings;
+	TArray<FString> Errors;
+
+	TSharedPtr<IMessageLogListing> LogListing;
+};
diff --git a/Engine/Plugins/Enterprise/GLTFExporter/MigratedFromGit/Source/GLTFExporter/Private/Builders/GLTFMessageBuilder.cpp b/Engine/Plugins/Enterprise/GLTFExporter/MigratedFromGit/Source/GLTFExporter/Private/Builders/GLTFMessageBuilder.cpp
deleted file mode 100644
diff --git a/Engine/Plugins/Enterprise/GLTFExporter/MigratedFromGit/Source/GLTFExporter/Private/Builders/GLTFMessageBuilder.h b/Engine/Plugins/Enterprise/GLTFExporter/MigratedFromGit/Source/GLTFExporter/Private/Builders/GLTFMessageBuilder.h
deleted file mode 100644
diff --git a/Engine/Plugins/Enterprise/GLTFExporter/MigratedFromGit/Source/GLTFExporter/Private/Builders/GLTFTaskBuilder.cpp b/Engine/Plugins/Enterprise/GLTFExporter/MigratedFromGit/Source/GLTFExporter/Private/Builders/GLTFTaskBuilder.cpp
--- a/Engine/Plugins/Enterprise/GLTFExporter/MigratedFromGit/Source/GLTFExporter/Private/Builders/GLTFTaskBuilder.cpp
+++ b/Engine/Plugins/Enterprise/GLTFExporter/MigratedFromGit/Source/GLTFExporter/Private/Builders/GLTFTaskBuilder.cpp
@@ -5,7 +5,7 @@
 #include "Misc/ScopedSlowTask.h"
 
 FGLTFTaskBuilder::FGLTFTaskBuilder(const FString& FilePath, const UGLTFExportOptions* ExportOptions)
-	: FGLTFMessageBuilder(FilePath, ExportOptions)
+	: FGLTFLogBuilder(FilePath, ExportOptions)
 	, PriorityIndexLock(INDEX_NONE)
 {
 }
diff --git a/Engine/Plugins/Enterprise/GLTFExporter/MigratedFromGit/Source/GLTFExporter/Private/Builders/GLTFTaskBuilder.h b/Engine/Plugins/Enterprise/GLTFExporter/MigratedFromGit/Source/GLTFExporter/Private/Builders/GLTFTaskBuilder.h
--- a/Engine/Plugins/Enterprise/GLTFExporter/MigratedFromGit/Source/GLTFExporter/Private/Builders/GLTFTaskBuilder.h
+++ b/Engine/Plugins/Enterprise/GLTFExporter/MigratedFromGit/Source/GLTFExporter/Private/Builders/GLTFTaskBuilder.h
@@ -2,10 +2,10 @@
 
 #pragma once
 
-#include "Builders/GLTFMessageBuilder.h"
+#include "Builders/GLTFLogBuilder.h"
 #include "Tasks/GLTFTask.h"
 
-class FGLTFTaskBuilder : public FGLTFMessageBuilder
+class FGLTFTaskBuilder : public FGLTFLogBuilder
 {
 protected:
 
diff --git a/Engine/Plugins/Enterprise/GLTFExporter/MigratedFromGit/Source/GLTFExporter/Private/Converters/GLTFBackdropConverters.cpp b/Engine/Plugins/Enterprise/GLTFExporter/MigratedFromGit/Source/GLTFExporter/Private/Converters/GLTFBackdropConverters.cpp
--- a/Engine/Plugins/Enterprise/GLTFExporter/MigratedFromGit/Source/GLTFExporter/Private/Converters/GLTFBackdropConverters.cpp
+++ b/Engine/Plugins/Enterprise/GLTFExporter/MigratedFromGit/Source/GLTFExporter/Private/Converters/GLTFBackdropConverters.cpp
@@ -34,7 +34,7 @@ FGLTFJsonBackdropIndex FGLTFBackdropConverter::Convert(const AActor* BackdropAct
 	}
 	else
 	{
-		Builder.AddWarningMessage(FString::Printf(TEXT("Failed to export Mesh for HDRIBackdrop %s"), *BackdropActor->GetName()));
+		Builder.LogWarning(FString::Printf(TEXT("Failed to export Mesh for HDRIBackdrop %s"), *BackdropActor->GetName()));
 	}
 
 	const UTextureCube* Cubemap;
@@ -51,7 +51,7 @@ FGLTFJsonBackdropIndex FGLTFBackdropConverter::Convert(const AActor* BackdropAct
 	}
 	else
 	{
-		Builder.AddWarningMessage(FString::Printf(TEXT("Failed to export Cubemap for HDRIBackdrop %s"), *BackdropActor->GetName()));
+		Builder.LogWarning(FString::Printf(TEXT("Failed to export Cubemap for HDRIBackdrop %s"), *BackdropActor->GetName()));
 	}
 
 	float Intensity;
@@ -61,7 +61,7 @@ FGLTFJsonBackdropIndex FGLTFBackdropConverter::Convert(const AActor* BackdropAct
 	}
 	else
 	{
-		Builder.AddWarningMessage(FString::Printf(TEXT("Failed to export Intensity for HDRIBackdrop %s"), *BackdropActor->GetName()));
+		Builder.LogWarning(FString::Printf(TEXT("Failed to export Intensity for HDRIBackdrop %s"), *BackdropActor->GetName()));
 	}
 
 	float Size;
@@ -71,7 +71,7 @@ FGLTFJsonBackdropIndex FGLTFBackdropConverter::Convert(const AActor* BackdropAct
 	}
 	else
 	{
-		Builder.AddWarningMessage(FString::Printf(TEXT("Failed to export Size for HDRIBackdrop %s"), *BackdropActor->GetName()));
+		Builder.LogWarning(FString::Printf(TEXT("Failed to export Size for HDRIBackdrop %s"), *BackdropActor->GetName()));
 	}
 
 	FVector ProjectionCenter;
@@ -81,7 +81,7 @@ FGLTFJsonBackdropIndex FGLTFBackdropConverter::Convert(const AActor* BackdropAct
 	}
 	else
 	{
-		Builder.AddWarningMessage(FString::Printf(TEXT("Failed to export ProjectionCenter for HDRIBackdrop %s"), *BackdropActor->GetName()));
+		Builder.LogWarning(FString::Printf(TEXT("Failed to export ProjectionCenter for HDRIBackdrop %s"), *BackdropActor->GetName()));
 	}
 
 	float LightingDistanceFactor;
@@ -91,7 +91,7 @@ FGLTFJsonBackdropIndex FGLTFBackdropConverter::Convert(const AActor* BackdropAct
 	}
 	else
 	{
-		Builder.AddWarningMessage(FString::Printf(TEXT("Failed to export LightingDistanceFactor for HDRIBackdrop %s"), *BackdropActor->GetName()));
+		Builder.LogWarning(FString::Printf(TEXT("Failed to export LightingDistanceFactor for HDRIBackdrop %s"), *BackdropActor->GetName()));
 	}
 
 	bool UseCameraProjection;
@@ -101,7 +101,7 @@ FGLTFJsonBackdropIndex FGLTFBackdropConverter::Convert(const AActor* BackdropAct
 	}
 	else
 	{
-		Builder.AddWarningMessage(FString::Printf(TEXT("Failed to export UseCameraProjection for HDRIBackdrop %s"), *BackdropActor->GetName()));
+		Builder.LogWarning(FString::Printf(TEXT("Failed to export UseCameraProjection for HDRIBackdrop %s"), *BackdropActor->GetName()));
 	}
 
 	return Builder.AddBackdrop(JsonBackdrop);
diff --git a/Engine/Plugins/Enterprise/GLTFExporter/MigratedFromGit/Source/GLTFExporter/Private/Converters/GLTFCameraConverters.cpp b/Engine/Plugins/Enterprise/GLTFExporter/MigratedFromGit/Source/GLTFExporter/Private/Converters/GLTFCameraConverters.cpp
--- a/Engine/Plugins/Enterprise/GLTFExporter/MigratedFromGit/Source/GLTFExporter/Private/Converters/GLTFCameraConverters.cpp
+++ b/Engine/Plugins/Enterprise/GLTFExporter/MigratedFromGit/Source/GLTFExporter/Private/Converters/GLTFCameraConverters.cpp
@@ -21,7 +21,7 @@ FGLTFJsonCameraIndex FGLTFCameraConverter::Convert(const UCameraComponent* Camer
 		case EGLTFJsonCameraType::Orthographic:
 			if (!DesiredView.bConstrainAspectRatio)
 			{
-				Builder.AddWarningMessage(FString::Printf(TEXT("Aspect ratio for orthographic camera component %s (in actor %s) will be constrainted in glTF"), *CameraComponent->GetName(), *CameraComponent->GetOwner()->GetName()));
+				Builder.LogWarning(FString::Printf(TEXT("Aspect ratio for orthographic camera component %s (in actor %s) will be constrainted in glTF"), *CameraComponent->GetName(), *CameraComponent->GetOwner()->GetName()));
 			}
 			Camera.Orthographic.XMag = FGLTFConverterUtility::ConvertLength(DesiredView.OrthoWidth, ExportScale);
 			Camera.Orthographic.YMag = FGLTFConverterUtility::ConvertLength(DesiredView.OrthoWidth / DesiredView.AspectRatio, ExportScale); // TODO: is this correct?
diff --git a/Engine/Plugins/Enterprise/GLTFExporter/MigratedFromGit/Source/GLTFExporter/Private/Converters/GLTFHotspotConverters.cpp b/Engine/Plugins/Enterprise/GLTFExporter/MigratedFromGit/Source/GLTFExporter/Private/Converters/GLTFHotspotConverters.cpp
--- a/Engine/Plugins/Enterprise/GLTFExporter/MigratedFromGit/Source/GLTFExporter/Private/Converters/GLTFHotspotConverters.cpp
+++ b/Engine/Plugins/Enterprise/GLTFExporter/MigratedFromGit/Source/GLTFExporter/Private/Converters/GLTFHotspotConverters.cpp
@@ -13,13 +13,13 @@ FGLTFJsonHotspotIndex FGLTFHotspotConverter::Convert(const AGLTFHotspotActor* Ho
 	{
 		if (!Builder.ExportOptions->bExportVertexSkinWeights)
 		{
-			Builder.AddWarningMessage(
+			Builder.LogWarning(
 				FString::Printf(TEXT("Can't export animation in hotspot %s because vertex skin weights are disabled by export options"),
 				*JsonHotspot.Name));
 		}
 		else if (!Builder.ExportOptions->bExportAnimationSequences)
 		{
-			Builder.AddWarningMessage(
+			Builder.LogWarning(
 				FString::Printf(TEXT("Can't export animation in hotspot %s because animation sequences are disabled by export options"),
 				*JsonHotspot.Name));
 		}
@@ -49,7 +49,7 @@ FGLTFJsonHotspotIndex FGLTFHotspotConverter::Convert(const AGLTFHotspotActor* Ho
 	{
 		if (!Builder.ExportOptions->bExportLevelSequences)
 		{
-			Builder.AddWarningMessage(
+			Builder.LogWarning(
 				FString::Printf(TEXT("Can't export animation in hotspot %s because level sequences are disabled by export options"),
 				*JsonHotspot.Name));
 		}
diff --git a/Engine/Plugins/Enterprise/GLTFExporter/MigratedFromGit/Source/GLTFExporter/Private/Converters/GLTFLevelVariantSetsConverters.cpp b/Engine/Plugins/Enterprise/GLTFExporter/MigratedFromGit/Source/GLTFExporter/Private/Converters/GLTFLevelVariantSetsConverters.cpp
--- a/Engine/Plugins/Enterprise/GLTFExporter/MigratedFromGit/Source/GLTFExporter/Private/Converters/GLTFLevelVariantSetsConverters.cpp
+++ b/Engine/Plugins/Enterprise/GLTFExporter/MigratedFromGit/Source/GLTFExporter/Private/Converters/GLTFLevelVariantSetsConverters.cpp
@@ -125,7 +125,7 @@ bool FGLTFLevelVariantSetsConverter::TryParseVisibilityPropertyValue(FGLTFJsonVa
 	const USceneComponent* Target = static_cast<USceneComponent*>(Property->GetPropertyParentContainerAddress());
 	if (Target == nullptr)
 	{
-		Builder.AddWarningMessage(FString::Printf(
+		Builder.LogWarning(FString::Printf(
 			TEXT("Target object for property is invalid, the property will be skipped. Context: %s"),
 			*GetLogContext(Property)));
 		return false;
@@ -134,14 +134,14 @@ bool FGLTFLevelVariantSetsConverter::TryParseVisibilityPropertyValue(FGLTFJsonVa
 	const AActor* Owner = Target->GetOwner();
 	if (Owner == nullptr)
 	{
-		Builder.AddWarningMessage(FString::Printf(TEXT("Invalid scene component, the property will be skipped. Context: %s"),
+		Builder.LogWarning(FString::Printf(TEXT("Invalid scene component, the property will be skipped. Context: %s"),
 			*GetLogContext(Property)));
 		return false;
 	}
 
 	if (Builder.bSelectedActorsOnly && !Owner->IsSelected())
 	{
-		Builder.AddWarningMessage(FString::Printf(
+		Builder.LogWarning(FString::Printf(
 			TEXT("Target actor for property is not selected for export, the property will be skipped. Context: %s"),
 			*GetLogContext(Property)));
 		return false;
@@ -151,7 +151,7 @@ bool FGLTFLevelVariantSetsConverter::TryParseVisibilityPropertyValue(FGLTFJsonVa
 
 	if (!TryGetPropertyValue(const_cast<UPropertyValue*>(Property), bIsVisible))
 	{
-		Builder.AddWarningMessage(FString::Printf(
+		Builder.LogWarning(FString::Printf(
 			TEXT("Failed to parse recorded data for property, it will be skipped. Context: %s"),
 			*GetLogContext(Property)));
 		return false;
@@ -172,7 +172,7 @@ bool FGLTFLevelVariantSetsConverter::TryParseMaterialPropertyValue(FGLTFJsonVari
 	UPropertyValueMaterial* MaterialProperty = Cast<UPropertyValueMaterial>(const_cast<UPropertyValue*>(Property));
 	if (MaterialProperty == nullptr)
 	{
-		Builder.AddWarningMessage(FString::Printf(
+		Builder.LogWarning(FString::Printf(
 			TEXT("Material property is invalid, the property will be skipped. Context: %s"),
 			*GetLogContext(Property)));
 		return false;
@@ -181,7 +181,7 @@ bool FGLTFLevelVariantSetsConverter::TryParseMaterialPropertyValue(FGLTFJsonVari
 	const USceneComponent* Target = static_cast<USceneComponent*>(Property->GetPropertyParentContainerAddress());
 	if (Target == nullptr)
 	{
-		Builder.AddWarningMessage(FString::Printf(
+		Builder.LogWarning(FString::Printf(
 			TEXT("Target object for property is invalid, the property will be skipped. Context: %s"),
 			*GetLogContext(Property)));
 		return false;
@@ -190,14 +190,14 @@ bool FGLTFLevelVariantSetsConverter::TryParseMaterialPropertyValue(FGLTFJsonVari
 	const AActor* Owner = Target->GetOwner();
 	if (Owner == nullptr)
 	{
-		Builder.AddWarningMessage(FString::Printf(TEXT("Invalid scene component, the property will be skipped. Context: %s"),
+		Builder.LogWarning(FString::Printf(TEXT("Invalid scene component, the property will be skipped. Context: %s"),
 			*GetLogContext(Property)));
 		return false;
 	}
 
 	if (Builder.bSelectedActorsOnly && !Owner->IsSelected())
 	{
-		Builder.AddWarningMessage(FString::Printf(
+		Builder.LogWarning(FString::Printf(
 			TEXT("Target actor for property is not selected for export, the property will be skipped. Context: %s"),
 			*GetLogContext(Property)));
 		return false;
@@ -208,7 +208,7 @@ bool FGLTFLevelVariantSetsConverter::TryParseMaterialPropertyValue(FGLTFJsonVari
 
 	if (NumPropSegments < 1)
 	{
-		Builder.AddWarningMessage(FString::Printf(
+		Builder.LogWarning(FString::Printf(
 			TEXT("Failed to parse element index to apply the material to, the property will be skipped. Context: %s"),
 			*GetLogContext(Property)));
 		return false;
@@ -275,7 +275,7 @@ bool FGLTFLevelVariantSetsConverter::TryParseStaticMeshPropertyValue(FGLTFJsonVa
 	const USceneComponent* Target = static_cast<USceneComponent*>(Property->GetPropertyParentContainerAddress());
 	if (Target == nullptr)
 	{
-		Builder.AddWarningMessage(FString::Printf(
+		Builder.LogWarning(FString::Printf(
 			TEXT("Target object for property is invalid, the property will be skipped. Context: %s"),
 			*GetLogContext(Property)));
 		return false;
@@ -284,14 +284,14 @@ bool FGLTFLevelVariantSetsConverter::TryParseStaticMeshPropertyValue(FGLTFJsonVa
 	const AActor* Owner = Target->GetOwner();
 	if (Owner == nullptr)
 	{
-		Builder.AddWarningMessage(FString::Printf(TEXT("Invalid scene component, the property will be skipped. Context: %s"),
+		Builder.LogWarning(FString::Printf(TEXT("Invalid scene component, the property will be skipped. Context: %s"),
 			*GetLogContext(Property)));
 		return false;
 	}
 
 	if (Builder.bSelectedActorsOnly && !Owner->IsSelected())
 	{
-		Builder.AddWarningMessage(FString::Printf(
+		Builder.LogWarning(FString::Printf(
 			TEXT("Target actor for property is not selected for export, the property will be skipped. Context: %s"),
 			*GetLogContext(Property)));
 		return false;
@@ -300,7 +300,7 @@ bool FGLTFLevelVariantSetsConverter::TryParseStaticMeshPropertyValue(FGLTFJsonVa
 	const UMeshComponent* MeshComponent = Cast<UMeshComponent>(Target);
 	if (MeshComponent == nullptr)
 	{
-		Builder.AddWarningMessage(FString::Printf(
+		Builder.LogWarning(FString::Printf(
 			TEXT("Target object for property has no mesh-component, the property will be skipped. Context: %s"),
 			*GetLogContext(Property)));
 		return false;
@@ -310,7 +310,7 @@ bool FGLTFLevelVariantSetsConverter::TryParseStaticMeshPropertyValue(FGLTFJsonVa
 
 	if (!TryGetPropertyValue(const_cast<UPropertyValue*>(Property), StaticMesh))
 	{
-		Builder.AddWarningMessage(FString::Printf(
+		Builder.LogWarning(FString::Printf(
 			TEXT("Failed to parse recorded data for property, it will be skipped. Context: %s"),
 			*GetLogContext(Property)));
 		return false;
@@ -332,7 +332,7 @@ bool FGLTFLevelVariantSetsConverter::TryParseSkeletalMeshPropertyValue(FGLTFJson
 	const USceneComponent* Target = static_cast<USceneComponent*>(Property->GetPropertyParentContainerAddress());
 	if (Target == nullptr)
 	{
-		Builder.AddWarningMessage(FString::Printf(
+		Builder.LogWarning(FString::Printf(
 			TEXT("Target object for property is invalid, the property will be skipped. Context: %s"),
 			*GetLogContext(Property)));
 		return false;
@@ -341,14 +341,14 @@ bool FGLTFLevelVariantSetsConverter::TryParseSkeletalMeshPropertyValue(FGLTFJson
 	const AActor* Owner = Target->GetOwner();
 	if (Owner == nullptr)
 	{
-		Builder.AddWarningMessage(FString::Printf(TEXT("Invalid scene component, the property will be skipped. Context: %s"),
+		Builder.LogWarning(FString::Printf(TEXT("Invalid scene component, the property will be skipped. Context: %s"),
 			*GetLogContext(Property)));
 		return false;
 	}
 
 	if (Builder.bSelectedActorsOnly && !Owner->IsSelected())
 	{
-		Builder.AddWarningMessage(FString::Printf(
+		Builder.LogWarning(FString::Printf(
 			TEXT("Target actor for property is not selected for export, the property will be skipped. Context: %s"),
 			*GetLogContext(Property)));
 		return false;
@@ -357,7 +357,7 @@ bool FGLTFLevelVariantSetsConverter::TryParseSkeletalMeshPropertyValue(FGLTFJson
 	const UMeshComponent* MeshComponent = Cast<UMeshComponent>(Target);
 	if (MeshComponent == nullptr)
 	{
-		Builder.AddWarningMessage(FString::Printf(
+		Builder.LogWarning(FString::Printf(
 			TEXT("Target object for property has no mesh-component, the property will be skipped. Context: %s"),
 			*GetLogContext(Property)));
 		return false;
@@ -367,7 +367,7 @@ bool FGLTFLevelVariantSetsConverter::TryParseSkeletalMeshPropertyValue(FGLTFJson
 
 	if (!TryGetPropertyValue(const_cast<UPropertyValue*>(Property), SkeletalMesh) || SkeletalMesh == nullptr)
 	{
-		Builder.AddWarningMessage(FString::Printf(
+		Builder.LogWarning(FString::Printf(
 			TEXT("Failed to parse recorded data for property, it will be skipped. Context: %s"),
 			*GetLogContext(Property)));
 		return false;
diff --git a/Engine/Plugins/Enterprise/GLTFExporter/MigratedFromGit/Source/GLTFExporter/Private/Converters/GLTFMaterialConverters.cpp b/Engine/Plugins/Enterprise/GLTFExporter/MigratedFromGit/Source/GLTFExporter/Private/Converters/GLTFMaterialConverters.cpp
--- a/Engine/Plugins/Enterprise/GLTFExporter/MigratedFromGit/Source/GLTFExporter/Private/Converters/GLTFMaterialConverters.cpp
+++ b/Engine/Plugins/Enterprise/GLTFExporter/MigratedFromGit/Source/GLTFExporter/Private/Converters/GLTFMaterialConverters.cpp
@@ -26,7 +26,7 @@ void FGLTFMaterialConverter::Sanitize(const UMaterialInterface*& Material, const
 			SectionString += SectionIndices.Num() > 1 ? TEXT("s ") : TEXT(" ");
 			SectionString += FString::JoinBy(SectionIndices, TEXT(", "), FString::FromInt);
 
-			Builder.AddWarningMessage(FString::Printf(
+			Builder.LogWarning(FString::Printf(
 				TEXT("Material %s uses mesh data from %s but the lightmap UVs (channel %d) are nearly 100%% degenerate (in %s). Simple baking will be used as fallback"),
 				*Material->GetName(),
 				*MeshData->GetParent()->Name,
diff --git a/Engine/Plugins/Enterprise/GLTFExporter/MigratedFromGit/Source/GLTFExporter/Private/Converters/GLTFNodeConverters.cpp b/Engine/Plugins/Enterprise/GLTFExporter/MigratedFromGit/Source/GLTFExporter/Private/Converters/GLTFNodeConverters.cpp
--- a/Engine/Plugins/Enterprise/GLTFExporter/MigratedFromGit/Source/GLTFExporter/Private/Converters/GLTFNodeConverters.cpp
+++ b/Engine/Plugins/Enterprise/GLTFExporter/MigratedFromGit/Source/GLTFExporter/Private/Converters/GLTFNodeConverters.cpp
@@ -93,7 +93,7 @@ FGLTFJsonNodeIndex FGLTFComponentConverter::Convert(const USceneComponent* Scene
 		const FVector ParentScale = ParentComponent->GetComponentScale();
 		if (!ParentScale.IsUniform())
 		{
-			Builder.AddWarningMessage(
+			Builder.LogWarning(
 				FString::Printf(TEXT("Non-uniform parent scale (%s) for component %s (in actor %s) may be represented differently in glTF"),
 				*ParentScale.ToString(),
 				*SceneComponent->GetName(),
@@ -195,7 +195,7 @@ FGLTFJsonNodeIndex FGLTFComponentSocketConverter::Convert(const USceneComponent*
 
 		// TODO: add support for more socket types
 
-		Builder.AddWarningMessage(
+		Builder.LogWarning(
 			FString::Printf(TEXT("Can't export socket %s because it belongs to an unsupported mesh component %s"),
 			*SocketName.ToString(),
 			*SceneComponent->GetName()));
diff --git a/Engine/Plugins/Enterprise/GLTFExporter/MigratedFromGit/Source/GLTFExporter/Private/Converters/GLTFSkySphereConverters.cpp b/Engine/Plugins/Enterprise/GLTFExporter/MigratedFromGit/Source/GLTFExporter/Private/Converters/GLTFSkySphereConverters.cpp
--- a/Engine/Plugins/Enterprise/GLTFExporter/MigratedFromGit/Source/GLTFExporter/Private/Converters/GLTFSkySphereConverters.cpp
+++ b/Engine/Plugins/Enterprise/GLTFExporter/MigratedFromGit/Source/GLTFExporter/Private/Converters/GLTFSkySphereConverters.cpp
@@ -35,7 +35,7 @@ FGLTFJsonSkySphereIndex FGLTFSkySphereConverter::Convert(const AActor* SkySphere
 	}
 	else
 	{
-		Builder.AddWarningMessage(FString::Printf(TEXT("Failed to export property SkySphereMesh for Sky Sphere %s"), *SkySphereActor->GetName()));
+		Builder.LogWarning(FString::Printf(TEXT("Failed to export property SkySphereMesh for Sky Sphere %s"), *SkySphereActor->GetName()));
 	}
 
 	const ADirectionalLight* DirectionalLight = nullptr;
@@ -45,7 +45,7 @@ FGLTFJsonSkySphereIndex FGLTFSkySphereConverter::Convert(const AActor* SkySphere
 	}
 	else
 	{
-		Builder.AddWarningMessage(FString::Printf(TEXT("Failed to export property Directional light actor for Sky Sphere %s"), *SkySphereActor->GetName()));
+		Builder.LogWarning(FString::Printf(TEXT("Failed to export property Directional light actor for Sky Sphere %s"), *SkySphereActor->GetName()));
 	}
 
 	const UMaterialInstance* SkyMaterial = nullptr;
@@ -63,7 +63,7 @@ FGLTFJsonSkySphereIndex FGLTFSkySphereConverter::Convert(const AActor* SkySphere
 	}
 	else
 	{
-		Builder.AddWarningMessage(FString::Printf(
+		Builder.LogWarning(FString::Printf(
 			TEXT("Failed to export property Sky material for Sky Sphere %s"),
 			*SkySphereActor->GetName()));
 	}
@@ -95,7 +95,7 @@ void FGLTFSkySphereConverter::ConvertProperty(const AActor* Actor, const TCHAR*
 
 	if (!FGLTFActorUtility::TryGetPropertyValue(Actor, PropertyName, OutValue))
 	{
-		Builder.AddWarningMessage(FString::Printf(
+		Builder.LogWarning(FString::Printf(
 			TEXT("Failed to export property %s for Sky Sphere %s"),
 			PropertyName,
 			*Actor->GetName()));
@@ -113,7 +113,7 @@ void FGLTFSkySphereConverter::ConvertColorProperty(const AActor* Actor, const TC
 	}
 	else
 	{
-		Builder.AddWarningMessage(FString::Printf(
+		Builder.LogWarning(FString::Printf(
 			TEXT("Failed to export property %s for Sky Sphere %s"),
 			PropertyName,
 			*Actor->GetName()));
@@ -131,7 +131,7 @@ void FGLTFSkySphereConverter::ConvertColorCurveProperty(const AActor* Actor, con
 	{
 		if (FGLTFCurveUtility::HasAnyAdjustment(*ColorCurve))
 		{
-			Builder.AddWarningMessage(FString::Printf(
+			Builder.LogWarning(FString::Printf(
 				TEXT("Adjustments for property %s in Sky Sphere %s are not supported"),
 				PropertyName,
 				*Actor->GetName()));
@@ -157,7 +157,7 @@ void FGLTFSkySphereConverter::ConvertColorCurveProperty(const AActor* Actor, con
 	}
 	else
 	{
-		Builder.AddWarningMessage(FString::Printf(
+		Builder.LogWarning(FString::Printf(
 			TEXT("Failed to export property %s for Sky Sphere %s"),
 			PropertyName,
 			*Actor->GetName()));
@@ -170,7 +170,7 @@ void FGLTFSkySphereConverter::ConvertScalarParameter(const AActor* Actor, const
 
 	if (!Material->GetScalarParameterValue(ParameterName, OutValue))
 	{
-		Builder.AddWarningMessage(FString::Printf(
+		Builder.LogWarning(FString::Printf(
 			TEXT("Failed to export parameter %s (in material %s) for Sky Sphere %s"),
 			ParameterName,
 			*Material->GetName(),
@@ -201,7 +201,7 @@ void FGLTFSkySphereConverter::ConvertTextureParameter(const AActor* Actor, const
 	}
 	else
 	{
-		Builder.AddWarningMessage(FString::Printf(
+		Builder.LogWarning(FString::Printf(
 			TEXT("Failed to export texture %s (used in material %s) for Sky Sphere %s"),
 			TexturePath,
 			*Material->GetName(),
diff --git a/Engine/Plugins/Enterprise/GLTFExporter/MigratedFromGit/Source/GLTFExporter/Private/Exporters/GLTFAnimSequenceExporter.cpp b/Engine/Plugins/Enterprise/GLTFExporter/MigratedFromGit/Source/GLTFExporter/Private/Exporters/GLTFAnimSequenceExporter.cpp
--- a/Engine/Plugins/Enterprise/GLTFExporter/MigratedFromGit/Source/GLTFExporter/Private/Exporters/GLTFAnimSequenceExporter.cpp
+++ b/Engine/Plugins/Enterprise/GLTFExporter/MigratedFromGit/Source/GLTFExporter/Private/Exporters/GLTFAnimSequenceExporter.cpp
@@ -17,7 +17,7 @@ bool UGLTFAnimSequenceExporter::AddObject(FGLTFContainerBuilder& Builder, const
 
 	if (!Builder.ExportOptions->bExportAnimationSequences)
 	{
-		Builder.AddErrorMessage(
+		Builder.LogError(
 			FString::Printf(TEXT("Failed to export animation sequence %s because animation sequences are disabled by export options"),
 			*AnimSequence->GetName()));
 		return false;
@@ -25,7 +25,7 @@ bool UGLTFAnimSequenceExporter::AddObject(FGLTFContainerBuilder& Builder, const
 
 	if (!Builder.ExportOptions->bExportVertexSkinWeights)
 	{
-		Builder.AddErrorMessage(
+		Builder.LogError(
 			FString::Printf(TEXT("Failed to export animation sequence %s because vertex skin weights are disabled by export options"),
 			*AnimSequence->GetName()));
 		return false;
@@ -34,7 +34,7 @@ bool UGLTFAnimSequenceExporter::AddObject(FGLTFContainerBuilder& Builder, const
 	const USkeletalMesh* SkeletalMesh = FGLTFExporterUtility::GetPreviewMesh(AnimSequence);
 	if (SkeletalMesh == nullptr)
 	{
-		Builder.AddErrorMessage(
+		Builder.LogError(
 			FString::Printf(TEXT("Failed to export animation sequence %s because of missing preview mesh"),
 			*AnimSequence->GetName()));
 		return false;
@@ -47,7 +47,7 @@ bool UGLTFAnimSequenceExporter::AddObject(FGLTFContainerBuilder& Builder, const
 		const FGLTFJsonMeshIndex MeshIndex = Builder.GetOrAddMesh(SkeletalMesh);
 		if (MeshIndex == INDEX_NONE)
 		{
-			Builder.AddErrorMessage(
+			Builder.LogError(
 				FString::Printf(TEXT("Failed to export skeletal mesh %s for animation sequence %s"),
 				*SkeletalMesh->GetName(),
 				*AnimSequence->GetName()));
@@ -62,7 +62,7 @@ bool UGLTFAnimSequenceExporter::AddObject(FGLTFContainerBuilder& Builder, const
 	const FGLTFJsonSkinIndex SkinIndex = Builder.GetOrAddSkin(NodeIndex, SkeletalMesh);
 	if (SkinIndex == INDEX_NONE)
 	{
-		Builder.AddErrorMessage(
+		Builder.LogError(
 			FString::Printf(TEXT("Failed to export bones in skeletal mesh %s for animation sequence %s"),
 			*SkeletalMesh->GetName(),
 			*AnimSequence->GetName()));
@@ -74,7 +74,7 @@ bool UGLTFAnimSequenceExporter::AddObject(FGLTFContainerBuilder& Builder, const
 	const FGLTFJsonAnimationIndex AnimationIndex =  Builder.GetOrAddAnimation(NodeIndex, SkeletalMesh, AnimSequence);
 	if (AnimationIndex == INDEX_NONE)
 	{
-		Builder.AddErrorMessage(
+		Builder.LogError(
 			FString::Printf(TEXT("Failed to export animation sequence %s"),
 			*AnimSequence->GetName()));
 		return false;
diff --git a/Engine/Plugins/Enterprise/GLTFExporter/MigratedFromGit/Source/GLTFExporter/Private/Exporters/GLTFLevelExporter.cpp b/Engine/Plugins/Enterprise/GLTFExporter/MigratedFromGit/Source/GLTFExporter/Private/Exporters/GLTFLevelExporter.cpp
--- a/Engine/Plugins/Enterprise/GLTFExporter/MigratedFromGit/Source/GLTFExporter/Private/Exporters/GLTFLevelExporter.cpp
+++ b/Engine/Plugins/Enterprise/GLTFExporter/MigratedFromGit/Source/GLTFExporter/Private/Exporters/GLTFLevelExporter.cpp
@@ -17,7 +17,7 @@ bool UGLTFLevelExporter::AddObject(FGLTFContainerBuilder& Builder, const UObject
 	const FGLTFJsonSceneIndex SceneIndex = Builder.GetOrAddScene(World);
 	if (SceneIndex == INDEX_NONE)
 	{
-		Builder.AddErrorMessage(FString::Printf(TEXT("Failed to export level %s"), *World->GetName()));
+		Builder.LogError(FString::Printf(TEXT("Failed to export level %s"), *World->GetName()));
 		return false;
 	}
 
diff --git a/Engine/Plugins/Enterprise/GLTFExporter/MigratedFromGit/Source/GLTFExporter/Private/Exporters/GLTFLevelSequenceExporter.cpp b/Engine/Plugins/Enterprise/GLTFExporter/MigratedFromGit/Source/GLTFExporter/Private/Exporters/GLTFLevelSequenceExporter.cpp
--- a/Engine/Plugins/Enterprise/GLTFExporter/MigratedFromGit/Source/GLTFExporter/Private/Exporters/GLTFLevelSequenceExporter.cpp
+++ b/Engine/Plugins/Enterprise/GLTFExporter/MigratedFromGit/Source/GLTFExporter/Private/Exporters/GLTFLevelSequenceExporter.cpp
@@ -17,7 +17,7 @@ bool UGLTFLevelSequenceExporter::AddObject(FGLTFContainerBuilder& Builder, const
 
 	if (!Builder.ExportOptions->bExportLevelSequences)
 	{
-		Builder.AddErrorMessage(
+		Builder.LogError(
 			FString::Printf(TEXT("Failed to export level sequence %s because level sequences are disabled by export options"),
 			*LevelSequence->GetName()));
 		return false;
@@ -26,7 +26,7 @@ bool UGLTFLevelSequenceExporter::AddObject(FGLTFContainerBuilder& Builder, const
 	TArray<UWorld*> Worlds = FGLTFExporterUtility::GetReferencedWorlds(LevelSequence);
 	if (Worlds.Num() == 0)
 	{
-		Builder.AddErrorMessage(
+		Builder.LogError(
 			FString::Printf(TEXT("Failed to export level sequence %s because no level referenced"),
 			*LevelSequence->GetName()));
 		return false;
@@ -34,7 +34,7 @@ bool UGLTFLevelSequenceExporter::AddObject(FGLTFContainerBuilder& Builder, const
 
 	if (Worlds.Num() > 1)
 	{
-		Builder.AddErrorMessage(
+		Builder.LogError(
 			FString::Printf(TEXT("Failed to export level sequence %s because more than one level referenced"),
 			*LevelSequence->GetName()));
 		return false;
@@ -44,7 +44,7 @@ bool UGLTFLevelSequenceExporter::AddObject(FGLTFContainerBuilder& Builder, const
 	const FGLTFJsonSceneIndex SceneIndex = Builder.GetOrAddScene(World);
 	if (SceneIndex == INDEX_NONE)
 	{
-		Builder.AddErrorMessage(
+		Builder.LogError(
 			FString::Printf(TEXT("Failed to export level %s for level sequence %s"),
 			*World->GetName(),
 			*LevelSequence->GetName()));
@@ -54,7 +54,7 @@ bool UGLTFLevelSequenceExporter::AddObject(FGLTFContainerBuilder& Builder, const
 	const FGLTFJsonAnimationIndex AnimationIndex = Builder.GetOrAddAnimation(World->PersistentLevel, LevelSequence);
 	if (AnimationIndex == INDEX_NONE)
 	{
-		Builder.AddErrorMessage(
+		Builder.LogError(
 			FString::Printf(TEXT("Failed to export level sequence %s"),
 			*LevelSequence->GetName()));
 		return false;
diff --git a/Engine/Plugins/Enterprise/GLTFExporter/MigratedFromGit/Source/GLTFExporter/Private/Exporters/GLTFLevelVariantSetsExporter.cpp b/Engine/Plugins/Enterprise/GLTFExporter/MigratedFromGit/Source/GLTFExporter/Private/Exporters/GLTFLevelVariantSetsExporter.cpp
--- a/Engine/Plugins/Enterprise/GLTFExporter/MigratedFromGit/Source/GLTFExporter/Private/Exporters/GLTFLevelVariantSetsExporter.cpp
+++ b/Engine/Plugins/Enterprise/GLTFExporter/MigratedFromGit/Source/GLTFExporter/Private/Exporters/GLTFLevelVariantSetsExporter.cpp
@@ -17,7 +17,7 @@ bool UGLTFLevelVariantSetsExporter::AddObject(FGLTFContainerBuilder& Builder, co
 
 	if (!Builder.ExportOptions->bExportVariantSets)
 	{
-		Builder.AddErrorMessage(
+		Builder.LogError(
 			FString::Printf(TEXT("Failed to export level variant sets %s because variant sets are disabled by export options"),
 			*LevelVariantSets->GetName()));
 		return false;
@@ -26,7 +26,7 @@ bool UGLTFLevelVariantSetsExporter::AddObject(FGLTFContainerBuilder& Builder, co
 	TArray<UWorld*> Worlds = FGLTFExporterUtility::GetReferencedWorlds(LevelVariantSets);
 	if (Worlds.Num() == 0)
 	{
-		Builder.AddErrorMessage(
+		Builder.LogError(
 			FString::Printf(TEXT("Failed to export level variant sets %s because no level referenced"),
 			*LevelVariantSets->GetName()));
 		return false;
@@ -34,7 +34,7 @@ bool UGLTFLevelVariantSetsExporter::AddObject(FGLTFContainerBuilder& Builder, co
 
 	if (Worlds.Num() > 1)
 	{
-		Builder.AddErrorMessage(
+		Builder.LogError(
 			FString::Printf(TEXT("Failed to export level variant sets %s because more than one level referenced"),
 			*LevelVariantSets->GetName()));
 		return false;
@@ -44,7 +44,7 @@ bool UGLTFLevelVariantSetsExporter::AddObject(FGLTFContainerBuilder& Builder, co
 	const FGLTFJsonSceneIndex SceneIndex = Builder.GetOrAddScene(World);
 	if (SceneIndex == INDEX_NONE)
 	{
-		Builder.AddErrorMessage(
+		Builder.LogError(
 			FString::Printf(TEXT("Failed to export level %s for level variant sets %s"),
 			*World->GetName(),
 			*LevelVariantSets->GetName()));
@@ -54,7 +54,7 @@ bool UGLTFLevelVariantSetsExporter::AddObject(FGLTFContainerBuilder& Builder, co
 	const FGLTFJsonLevelVariantSetsIndex LevelVariantSetsIndex = Builder.GetOrAddLevelVariantSets(LevelVariantSets);
 	if (LevelVariantSetsIndex == INDEX_NONE)
 	{
-		Builder.AddErrorMessage(
+		Builder.LogError(
 			FString::Printf(TEXT("Failed to export level variant sets %s"),
 			*LevelVariantSets->GetName()));
 		return false;
diff --git a/Engine/Plugins/Enterprise/GLTFExporter/MigratedFromGit/Source/GLTFExporter/Private/Exporters/GLTFMaterialExporter.cpp b/Engine/Plugins/Enterprise/GLTFExporter/MigratedFromGit/Source/GLTFExporter/Private/Exporters/GLTFMaterialExporter.cpp
--- a/Engine/Plugins/Enterprise/GLTFExporter/MigratedFromGit/Source/GLTFExporter/Private/Exporters/GLTFMaterialExporter.cpp
+++ b/Engine/Plugins/Enterprise/GLTFExporter/MigratedFromGit/Source/GLTFExporter/Private/Exporters/GLTFMaterialExporter.cpp
@@ -23,7 +23,7 @@ bool UGLTFMaterialExporter::AddObject(FGLTFContainerBuilder& Builder, const UObj
 			const FGLTFJsonMeshIndex MeshIndex = Builder.GetOrAddMesh(PreviewMesh, { Material });
 			if (MeshIndex == INDEX_NONE)
 			{
-				Builder.AddErrorMessage(
+				Builder.LogError(
 					FString::Printf(TEXT("Failed to export preview mesh %s for material %s"),
 					*Material->GetName(),
 					*PreviewMesh->GetName()));
@@ -42,7 +42,7 @@ bool UGLTFMaterialExporter::AddObject(FGLTFContainerBuilder& Builder, const UObj
 		}
 		else
 		{
-			Builder.AddErrorMessage(
+			Builder.LogError(
 				FString::Printf(TEXT("Failed to export material %s because of missing preview mesh"),
 				*Material->GetName()));
 			return false;
@@ -53,7 +53,7 @@ bool UGLTFMaterialExporter::AddObject(FGLTFContainerBuilder& Builder, const UObj
 		const FGLTFJsonMaterialIndex MaterialIndex = Builder.GetOrAddMaterial(Material);
 		if (MaterialIndex == INDEX_NONE)
 		{
-			Builder.AddErrorMessage(
+			Builder.LogError(
 				FString::Printf(TEXT("Failed to export material %s"),
 				*Material->GetName()));
 			return false;
diff --git a/Engine/Plugins/Enterprise/GLTFExporter/MigratedFromGit/Source/GLTFExporter/Private/Exporters/GLTFSkeletalMeshExporter.cpp b/Engine/Plugins/Enterprise/GLTFExporter/MigratedFromGit/Source/GLTFExporter/Private/Exporters/GLTFSkeletalMeshExporter.cpp
--- a/Engine/Plugins/Enterprise/GLTFExporter/MigratedFromGit/Source/GLTFExporter/Private/Exporters/GLTFSkeletalMeshExporter.cpp
+++ b/Engine/Plugins/Enterprise/GLTFExporter/MigratedFromGit/Source/GLTFExporter/Private/Exporters/GLTFSkeletalMeshExporter.cpp
@@ -17,7 +17,7 @@ bool UGLTFSkeletalMeshExporter::AddObject(FGLTFContainerBuilder& Builder, const
 	const FGLTFJsonMeshIndex MeshIndex = Builder.GetOrAddMesh(SkeletalMesh);
 	if (MeshIndex == INDEX_NONE)
 	{
-		Builder.AddErrorMessage(FString::Printf(TEXT("Failed to export skeletal mesh %s"), *SkeletalMesh->GetName()));
+		Builder.LogError(FString::Printf(TEXT("Failed to export skeletal mesh %s"), *SkeletalMesh->GetName()));
 		return false;
 	}
 
@@ -30,7 +30,7 @@ bool UGLTFSkeletalMeshExporter::AddObject(FGLTFContainerBuilder& Builder, const
 		const FGLTFJsonSkinIndex SkinIndex = Builder.GetOrAddSkin(NodeIndex, SkeletalMesh);
 		if (SkinIndex == INDEX_NONE)
 		{
-			Builder.AddErrorMessage(FString::Printf(TEXT("Failed to export bones in skeletal mesh %s"), *SkeletalMesh->GetName()));
+			Builder.LogError(FString::Printf(TEXT("Failed to export bones in skeletal mesh %s"), *SkeletalMesh->GetName()));
 			return false;
 		}
 
diff --git a/Engine/Plugins/Enterprise/GLTFExporter/MigratedFromGit/Source/GLTFExporter/Private/Exporters/GLTFStaticMeshExporter.cpp b/Engine/Plugins/Enterprise/GLTFExporter/MigratedFromGit/Source/GLTFExporter/Private/Exporters/GLTFStaticMeshExporter.cpp
--- a/Engine/Plugins/Enterprise/GLTFExporter/MigratedFromGit/Source/GLTFExporter/Private/Exporters/GLTFStaticMeshExporter.cpp
+++ b/Engine/Plugins/Enterprise/GLTFExporter/MigratedFromGit/Source/GLTFExporter/Private/Exporters/GLTFStaticMeshExporter.cpp
@@ -17,7 +17,7 @@ bool UGLTFStaticMeshExporter::AddObject(FGLTFContainerBuilder& Builder, const UO
 	const FGLTFJsonMeshIndex MeshIndex = Builder.GetOrAddMesh(StaticMesh);
 	if (MeshIndex == INDEX_NONE)
 	{
-		Builder.AddErrorMessage(FString::Printf(TEXT("Failed to export static mesh %s"), *StaticMesh->GetName()));
+		Builder.LogError(FString::Printf(TEXT("Failed to export static mesh %s"), *StaticMesh->GetName()));
 		return false;
 	}
 
diff --git a/Engine/Plugins/Enterprise/GLTFExporter/MigratedFromGit/Source/GLTFExporter/Private/Tasks/GLTFMaterialTasks.cpp b/Engine/Plugins/Enterprise/GLTFExporter/MigratedFromGit/Source/GLTFExporter/Private/Tasks/GLTFMaterialTasks.cpp
--- a/Engine/Plugins/Enterprise/GLTFExporter/MigratedFromGit/Source/GLTFExporter/Private/Tasks/GLTFMaterialTasks.cpp
+++ b/Engine/Plugins/Enterprise/GLTFExporter/MigratedFromGit/Source/GLTFExporter/Private/Tasks/GLTFMaterialTasks.cpp
@@ -79,7 +79,7 @@ void FGLTFMaterialTask::Complete()
 				{
 					if (!TryGetBakedMaterialProperty(JsonMaterial.PBRMetallicRoughness.BaseColorTexture, JsonMaterial.PBRMetallicRoughness.BaseColorFactor, BaseColorProperty, TEXT("BaseColor")))
 					{
-						Builder.AddWarningMessage(FString::Printf(TEXT("Failed to export %s for material %s"), *BaseColorProperty.ToString(), *Material->GetName()));
+						Builder.LogWarning(FString::Printf(TEXT("Failed to export %s for material %s"), *BaseColorProperty.ToString(), *Material->GetName()));
 					}
 				}
 			}
@@ -90,7 +90,7 @@ void FGLTFMaterialTask::Complete()
 		{
 			if (!TryGetBaseColorAndOpacity(JsonMaterial.PBRMetallicRoughness, BaseColorProperty, OpacityProperty))
 			{
-				Builder.AddWarningMessage(FString::Printf(TEXT("Failed to export %s and %s for material %s"), *BaseColorProperty.ToString(), *OpacityProperty.ToString(), *Material->GetName()));
+				Builder.LogWarning(FString::Printf(TEXT("Failed to export %s and %s for material %s"), *BaseColorProperty.ToString(), *OpacityProperty.ToString(), *Material->GetName()));
 			}
 		}
 
@@ -101,13 +101,13 @@ void FGLTFMaterialTask::Complete()
 
 			if (!TryGetMetallicAndRoughness(JsonMaterial.PBRMetallicRoughness, MetallicProperty, RoughnessProperty))
 			{
-				Builder.AddWarningMessage(FString::Printf(TEXT("Failed to export %s and %s for material %s"), *MetallicProperty.ToString(), *RoughnessProperty.ToString(), *Material->GetName()));
+				Builder.LogWarning(FString::Printf(TEXT("Failed to export %s and %s for material %s"), *MetallicProperty.ToString(), *RoughnessProperty.ToString(), *Material->GetName()));
 			}
 
 			const FMaterialPropertyEx EmissiveProperty = MP_EmissiveColor;
 			if (!TryGetEmissive(JsonMaterial, EmissiveProperty))
 			{
-				Builder.AddWarningMessage(FString::Printf(TEXT("Failed to export %s for material %s"), *EmissiveProperty.ToString(), *Material->GetName()));
+				Builder.LogWarning(FString::Printf(TEXT("Failed to export %s for material %s"), *EmissiveProperty.ToString(), *Material->GetName()));
 			}
 
 			const FMaterialPropertyEx NormalProperty = JsonMaterial.ShadingModel == EGLTFJsonShadingModel::ClearCoat ? FMaterialPropertyEx(TEXT("ClearCoatBottomNormal")) : FMaterialPropertyEx(MP_Normal);
@@ -117,7 +117,7 @@ void FGLTFMaterialTask::Complete()
 				{
 					if (!TryGetBakedMaterialProperty(JsonMaterial.NormalTexture, NormalProperty, TEXT("Normal")))
 					{
-						Builder.AddWarningMessage(FString::Printf(TEXT("Failed to export %s for material %s"), *NormalProperty.ToString(), *Material->GetName()));
+						Builder.LogWarning(FString::Printf(TEXT("Failed to export %s for material %s"), *NormalProperty.ToString(), *Material->GetName()));
 					}
 				}
 			}
@@ -129,7 +129,7 @@ void FGLTFMaterialTask::Complete()
 				{
 					if (!TryGetBakedMaterialProperty(JsonMaterial.OcclusionTexture, AmbientOcclusionProperty, TEXT("Occlusion"), true))
 					{
-						Builder.AddWarningMessage(FString::Printf(TEXT("Failed to export %s for material %s"), *AmbientOcclusionProperty.ToString(), *Material->GetName()));
+						Builder.LogWarning(FString::Printf(TEXT("Failed to export %s for material %s"), *AmbientOcclusionProperty.ToString(), *Material->GetName()));
 					}
 				}
 			}
@@ -141,7 +141,7 @@ void FGLTFMaterialTask::Complete()
 
 				if (!TryGetClearCoatRoughness(JsonMaterial.ClearCoat, ClearCoatProperty, ClearCoatRoughnessProperty))
 				{
-					Builder.AddWarningMessage(FString::Printf(TEXT("Failed to export %s and %s for material %s"), *ClearCoatProperty.ToString(), *ClearCoatRoughnessProperty.ToString(), *Material->GetName()));
+					Builder.LogWarning(FString::Printf(TEXT("Failed to export %s and %s for material %s"), *ClearCoatProperty.ToString(), *ClearCoatRoughnessProperty.ToString(), *Material->GetName()));
 				}
 
 				const FMaterialPropertyEx ClearCoatNormalProperty = MP_Normal;
@@ -151,7 +151,7 @@ void FGLTFMaterialTask::Complete()
 					{
 						if (!TryGetBakedMaterialProperty(JsonMaterial.ClearCoat.ClearCoatNormalTexture, ClearCoatNormalProperty, TEXT("ClearCoatNormal")))
 						{
-							Builder.AddWarningMessage(FString::Printf(TEXT("Failed to export %s for material %s"), *ClearCoatNormalProperty.ToString(), *Material->GetName()));
+							Builder.LogWarning(FString::Printf(TEXT("Failed to export %s for material %s"), *ClearCoatNormalProperty.ToString(), *Material->GetName()));
 						}
 					}
 				}
@@ -173,7 +173,7 @@ void FGLTFMaterialTask::Complete()
 			SectionString += SectionIndices.Num() > 1 ? TEXT("s ") : TEXT(" ");
 			SectionString += FString::JoinBy(SectionIndices, TEXT(", "), FString::FromInt);
 
-			Builder.AddWarningMessage(FString::Printf(
+			Builder.LogWarning(FString::Printf(
 				TEXT("Material %s is baked using mesh data from %s but the lightmap UV (channel %d) are overlapping by %.2f%% (in %s) and may produce incorrect results"),
 				*Material->GetName(),
 				*ParentMeshData->Name,
@@ -192,7 +192,7 @@ EMaterialShadingModel FGLTFMaterialTask::GetShadingModel() const
 	if (PossibilitesCount == 0)
 	{
 		EMaterialShadingModel ShadingModel = MSM_DefaultLit;
-		Builder.AddWarningMessage(FString::Printf(
+		Builder.LogWarning(FString::Printf(
 			TEXT("No shading model defined for material %s, will export as %s"),
 			*Material->GetName(),
 			*FGLTFNameUtility::GetName(ShadingModel)));
@@ -204,7 +204,7 @@ EMaterialShadingModel FGLTFMaterialTask::GetShadingModel() const
 		if (!FApp::CanEverRender())
 		{
 			EMaterialShadingModel ShadingModel = FGLTFMaterialUtility::GetRichestShadingModel(Possibilites);
-			Builder.AddWarningMessage(FString::Printf(
+			Builder.LogWarning(FString::Printf(
 				TEXT("Can't evaluate shading model expression in material %s because renderer missing, will export as %s"),
 				*Material->GetName(),
 				*FGLTFNameUtility::GetName(ShadingModel)));
@@ -219,7 +219,7 @@ EMaterialShadingModel FGLTFMaterialTask::GetShadingModel() const
 			if (EvaluationCount == 0)
 			{
 				EMaterialShadingModel ShadingModel = FGLTFMaterialUtility::GetRichestShadingModel(Possibilites);
-				Builder.AddWarningMessage(FString::Printf(
+				Builder.LogWarning(FString::Printf(
 					TEXT("Evaluation of shading model expression in material %s returned none, will export as %s"),
 					*Material->GetName(),
 					*FGLTFNameUtility::GetName(ShadingModel)));
@@ -229,7 +229,7 @@ EMaterialShadingModel FGLTFMaterialTask::GetShadingModel() const
 			if (EvaluationCount > 1)
 			{
 				EMaterialShadingModel ShadingModel = FGLTFMaterialUtility::GetRichestShadingModel(Evaluation);
-				Builder.AddWarningMessage(FString::Printf(
+				Builder.LogWarning(FString::Printf(
 					TEXT("Evaluation of shading model expression in material %s is inconclusive (%s), will export as %s"),
 					*Material->GetName(),
 					*FGLTFMaterialUtility::ShadingModelsToString(Evaluation),
@@ -262,7 +262,7 @@ void FGLTFMaterialTask::ConvertShadingModel(EGLTFJsonShadingModel& OutShadingMod
 	{
 		OutShadingModel = EGLTFJsonShadingModel::Default;
 
-		Builder.AddWarningMessage(FString::Printf(
+		Builder.LogWarning(FString::Printf(
 			TEXT("Unsupported shading model (%s) in material %s, will export as %s"),
 			*FGLTFNameUtility::GetName(ShadingModel),
 			*Material->GetName(),
@@ -274,7 +274,7 @@ void FGLTFMaterialTask::ConvertShadingModel(EGLTFJsonShadingModel& OutShadingMod
 	{
 		OutShadingModel = EGLTFJsonShadingModel::Default;
 
-		Builder.AddWarningMessage(FString::Printf(
+		Builder.LogWarning(FString::Printf(
 			TEXT("Shading model (%s) in material %s disabled by export options, will export as %s"),
 			*FGLTFNameUtility::GetName(ShadingModel),
 			*Material->GetName(),
@@ -286,7 +286,7 @@ void FGLTFMaterialTask::ConvertShadingModel(EGLTFJsonShadingModel& OutShadingMod
 	{
 		OutShadingModel = EGLTFJsonShadingModel::Default;
 
-		Builder.AddWarningMessage(FString::Printf(
+		Builder.LogWarning(FString::Printf(
 			TEXT("Shading model (%s) in material %s disabled by export options, will export as %s"),
 			*FGLTFNameUtility::GetName(ShadingModel),
 			*Material->GetName(),
@@ -304,7 +304,7 @@ void FGLTFMaterialTask::ConvertAlphaMode(EGLTFJsonAlphaMode& OutAlphaMode, EGLTF
 		OutAlphaMode = EGLTFJsonAlphaMode::Blend;
 		OutBlendMode = EGLTFJsonBlendMode::None;
 
-		Builder.AddWarningMessage(FString::Printf(
+		Builder.LogWarning(FString::Printf(
 			TEXT("Unsupported blend mode (%s) in material %s, will export as %s"),
 			*FGLTFNameUtility::GetName(BlendMode),
 			*Material->GetName(),
@@ -319,7 +319,7 @@ void FGLTFMaterialTask::ConvertAlphaMode(EGLTFJsonAlphaMode& OutAlphaMode, EGLTF
 		{
 			OutBlendMode = EGLTFJsonBlendMode::None;
 
-			Builder.AddWarningMessage(FString::Printf(
+			Builder.LogWarning(FString::Printf(
 				TEXT("Extra blend mode (%s) in material %s disabled by export options, will export as %s"),
 				*FGLTFNameUtility::GetName(BlendMode),
 				*Material->GetName(),
@@ -368,7 +368,7 @@ bool FGLTFMaterialTask::TryGetBaseColorAndOpacity(FGLTFJsonPBRMetallicRoughness&
 
 	if (Builder.ExportOptions->BakeMaterialInputs == EGLTFMaterialBakeMode::Disabled)
 	{
-		Builder.AddWarningMessage(FString::Printf(
+		Builder.LogWarning(FString::Printf(
 			TEXT("%s and %s for material %s needs to bake, but material baking is disabled by export options"),
 			*BaseColorProperty.ToString(),
 			*OpacityProperty.ToString(),
@@ -493,7 +493,7 @@ bool FGLTFMaterialTask::TryGetMetallicAndRoughness(FGLTFJsonPBRMetallicRoughness
 
 	if (Builder.ExportOptions->BakeMaterialInputs == EGLTFMaterialBakeMode::Disabled)
 	{
-		Builder.AddWarningMessage(FString::Printf(
+		Builder.LogWarning(FString::Printf(
 			TEXT("%s and %s for material %s needs to bake, but material baking is disabled by export options"),
 			*MetallicProperty.ToString(),
 			*RoughnessProperty.ToString(),
@@ -621,7 +621,7 @@ bool FGLTFMaterialTask::TryGetClearCoatRoughness(FGLTFJsonClearCoatExtension& Ou
 
 	if (Builder.ExportOptions->BakeMaterialInputs == EGLTFMaterialBakeMode::Disabled)
 	{
-		Builder.AddWarningMessage(FString::Printf(
+		Builder.LogWarning(FString::Printf(
 			TEXT("%s and %s for material %s needs to bake, but material baking is disabled by export options"),
 			*IntensityProperty.ToString(),
 			*RoughnessProperty.ToString(),
@@ -721,7 +721,7 @@ bool FGLTFMaterialTask::TryGetEmissive(FGLTFJsonMaterial& JsonMaterial, const FM
 
 	if (Builder.ExportOptions->BakeMaterialInputs == EGLTFMaterialBakeMode::Disabled)
 	{
-		Builder.AddWarningMessage(FString::Printf(
+		Builder.LogWarning(FString::Printf(
 			TEXT("%s for material %s needs to bake, but material baking is disabled by export options"),
 			*EmissiveProperty.ToString(),
 			*Material->GetName()));
@@ -1094,7 +1094,7 @@ bool FGLTFMaterialTask::TryGetSourceTexture(const UTexture2D*& OutTexture, int32
 
 		if (!Builder.ExportOptions->bExportTextureTransforms && !OutTransform.IsExactlyDefault())
 		{
-			Builder.AddWarningMessage(FString::Printf(
+			Builder.LogWarning(FString::Printf(
 				TEXT("Texture coordinates [%d] in %s for material %s are transformed, but texture transform is disabled by export options"),
 				OutTexCoord,
 				*Property.ToString(),
@@ -1131,7 +1131,7 @@ bool FGLTFMaterialTask::TryGetSourceTexture(const UTexture2D*& OutTexture, int32
 
 		if (!Builder.ExportOptions->bExportTextureTransforms && !OutTransform.IsExactlyDefault())
 		{
-			Builder.AddWarningMessage(FString::Printf(
+			Builder.LogWarning(FString::Printf(
 				TEXT("Texture coordinates [%d] in %s for material %s are transformed, but texture transform is disabled by export options"),
 				OutTexCoord,
 				*Property.ToString(),
@@ -1149,7 +1149,7 @@ bool FGLTFMaterialTask::TryGetBakedMaterialProperty(FGLTFJsonTextureInfo& OutTex
 {
 	if (Builder.ExportOptions->BakeMaterialInputs == EGLTFMaterialBakeMode::Disabled)
 	{
-		Builder.AddWarningMessage(FString::Printf(
+		Builder.LogWarning(FString::Printf(
 			TEXT("%s for material %s needs to bake, but material baking is disabled by export options"),
 			*Property.ToString(),
 			*Material->GetName()));
@@ -1188,7 +1188,7 @@ bool FGLTFMaterialTask::TryGetBakedMaterialProperty(FGLTFJsonTextureInfo& OutTex
 {
 	if (Builder.ExportOptions->BakeMaterialInputs == EGLTFMaterialBakeMode::Disabled)
 	{
-		Builder.AddWarningMessage(FString::Printf(
+		Builder.LogWarning(FString::Printf(
 			TEXT("%s for material %s needs to bake, but material baking is disabled by export options"),
 			*Property.ToString(),
 			*Material->GetName()));
@@ -1227,7 +1227,7 @@ inline bool FGLTFMaterialTask::TryGetBakedMaterialProperty(FGLTFJsonTextureInfo&
 {
 	if (Builder.ExportOptions->BakeMaterialInputs == EGLTFMaterialBakeMode::Disabled)
 	{
-		Builder.AddWarningMessage(FString::Printf(
+		Builder.LogWarning(FString::Printf(
 			TEXT("%s for material %s needs to bake, but material baking is disabled by export options"),
 			*Property.ToString(),
 			*Material->GetName()));
@@ -1266,7 +1266,7 @@ bool FGLTFMaterialTask::TryGetBakedMaterialProperty(FGLTFJsonTextureInfo& OutTex
 {
 	if (Builder.ExportOptions->BakeMaterialInputs == EGLTFMaterialBakeMode::Disabled)
 	{
-		Builder.AddWarningMessage(FString::Printf(
+		Builder.LogWarning(FString::Printf(
 			TEXT("%s for material %s needs to bake, but material baking is disabled by export options"),
 			*Property.ToString(),
 			*Material->GetName()));
@@ -1348,7 +1348,7 @@ FGLTFPropertyBakeOutput FGLTFMaterialTask::BakeMaterialProperty(const FMaterialP
 
 			if (TexCoords.Num() > 1)
 			{
-				Builder.AddWarningMessage(FString::Printf(
+				Builder.LogWarning(FString::Printf(
 					TEXT("%s for material %s uses multiple texture coordinates (%s), baked texture will be sampled using only the first (%d)"),
 					*Property.ToString(),
 					*Material->GetName(),
diff --git a/Engine/Plugins/Enterprise/GLTFExporter/MigratedFromGit/Source/GLTFExporter/Private/Tasks/GLTFMeshTasks.cpp b/Engine/Plugins/Enterprise/GLTFExporter/MigratedFromGit/Source/GLTFExporter/Private/Tasks/GLTFMeshTasks.cpp
--- a/Engine/Plugins/Enterprise/GLTFExporter/MigratedFromGit/Source/GLTFExporter/Private/Tasks/GLTFMeshTasks.cpp
+++ b/Engine/Plugins/Enterprise/GLTFExporter/MigratedFromGit/Source/GLTFExporter/Private/Tasks/GLTFMeshTasks.cpp
@@ -55,14 +55,14 @@ namespace
 
 		if (bZeroNormals)
 		{
-			Builder.AddWarningMessage(FString::Printf(
+			Builder.LogWarning(FString::Printf(
 				TEXT("Mesh %s has some nearly zero-length normals which can create some issues. Consider checking 'Recompute Normals' in the asset settings"),
 				MeshName));
 		}
 
 		if (bZeroTangents)
 		{
-			Builder.AddWarningMessage(FString::Printf(
+			Builder.LogWarning(FString::Printf(
 				TEXT("Mesh %s has some nearly zero-length tangents which can create some issues. Consider checking 'Recompute Tangents' in the asset settings"),
 				MeshName));
 		}
@@ -97,7 +97,7 @@ void FGLTFStaticMeshTask::Complete()
 
 	if (Builder.ExportOptions->bExportVertexColors && HasVertexColors(ColorBuffer))
 	{
-		Builder.AddWarningMessage(FString::Printf(
+		Builder.LogWarning(FString::Printf(
 			TEXT("Vertex colors in mesh %s will act as a multiplier for base color in glTF, regardless of material, which may produce undesirable results."),
 			*StaticMesh->GetName()));
 	}
@@ -188,7 +188,7 @@ void FGLTFSkeletalMeshTask::Complete()
 
 	if (Builder.ExportOptions->bExportVertexColors && HasVertexColors(ColorBuffer))
 	{
-		Builder.AddWarningMessage(FString::Printf(
+		Builder.LogWarning(FString::Printf(
 			TEXT("Vertex colors in mesh %s will act as a multiplier for base color in glTF, regardless of material, which may produce undesirable results."),
 			*SkeletalMesh->GetName()));
 	}
diff --git a/Engine/Plugins/Enterprise/GLTFExporter/MigratedFromGit/Source/GLTFExporter/Private/Tasks/GLTFTextureTasks.cpp b/Engine/Plugins/Enterprise/GLTFExporter/MigratedFromGit/Source/GLTFExporter/Private/Tasks/GLTFTextureTasks.cpp
--- a/Engine/Plugins/Enterprise/GLTFExporter/MigratedFromGit/Source/GLTFExporter/Private/Tasks/GLTFTextureTasks.cpp
+++ b/Engine/Plugins/Enterprise/GLTFExporter/MigratedFromGit/Source/GLTFExporter/Private/Tasks/GLTFTextureTasks.cpp
@@ -26,7 +26,7 @@ void FGLTFTexture2DTask::Complete()
 	TArray<FColor> Pixels;
 	if (!FGLTFTextureUtility::ReadPixels(RenderTarget, Pixels, JsonTexture.Encoding))
 	{
-		Builder.AddWarningMessage(FString::Printf(TEXT("Failed to read pixels for 2D texture %s"), *JsonTexture.Name));
+		Builder.LogWarning(FString::Printf(TEXT("Failed to read pixels for 2D texture %s"), *JsonTexture.Name));
 		return;
 	}
 
@@ -55,7 +55,7 @@ void FGLTFTextureCubeTask::Complete()
 	const UTexture2D* FaceTexture = FGLTFTextureUtility::CreateTextureFromCubeFace(TextureCube, CubeFace);
 	if (FaceTexture == nullptr)
 	{
-		Builder.AddWarningMessage(FString::Printf(TEXT("Failed to extract cube face %d for cubemap texture %s"), CubeFace, *TextureCube->GetName()));
+		Builder.LogWarning(FString::Printf(TEXT("Failed to extract cube face %d for cubemap texture %s"), CubeFace, *TextureCube->GetName()));
 		return;
 	}
 
@@ -74,7 +74,7 @@ void FGLTFTextureCubeTask::Complete()
 	TArray<FColor> Pixels;
 	if (!FGLTFTextureUtility::ReadPixels(RenderTarget, Pixels, JsonTexture.Encoding))
 	{
-		Builder.AddWarningMessage(FString::Printf(TEXT("Failed to read pixels (cube face %d) for cubemap texture %s"), CubeFace, *TextureCube->GetName()));
+		Builder.LogWarning(FString::Printf(TEXT("Failed to read pixels (cube face %d) for cubemap texture %s"), CubeFace, *TextureCube->GetName()));
 		return;
 	}
 
@@ -101,7 +101,7 @@ void FGLTFTextureRenderTarget2DTask::Complete()
 	TArray<FColor> Pixels;
 	if (!FGLTFTextureUtility::ReadPixels(RenderTarget2D, Pixels, JsonTexture.Encoding))
 	{
-		Builder.AddWarningMessage(FString::Printf(TEXT("Failed to read pixels for 2D render target %s"), *JsonTexture.Name));
+		Builder.LogWarning(FString::Printf(TEXT("Failed to read pixels for 2D render target %s"), *JsonTexture.Name));
 		return;
 	}
 
@@ -122,7 +122,7 @@ void FGLTFTextureRenderTargetCubeTask::Complete()
 	const UTexture2D* FaceTexture = FGLTFTextureUtility::CreateTextureFromCubeFace(RenderTargetCube, CubeFace);
 	if (FaceTexture == nullptr)
 	{
-		Builder.AddWarningMessage(FString::Printf(TEXT("Failed to extract cube face %d for cubemap render target %s"), CubeFace, *RenderTargetCube->GetName()));
+		Builder.LogWarning(FString::Printf(TEXT("Failed to extract cube face %d for cubemap render target %s"), CubeFace, *RenderTargetCube->GetName()));
 		return;
 	}
 
@@ -141,7 +141,7 @@ void FGLTFTextureRenderTargetCubeTask::Complete()
 	TArray<FColor> Pixels;
 	if (!FGLTFTextureUtility::ReadPixels(RenderTarget, Pixels, JsonTexture.Encoding))
 	{
-		Builder.AddWarningMessage(FString::Printf(TEXT("Failed to read pixels (cube face %d) for cubemap render target %s"), CubeFace, *RenderTargetCube->GetName()));
+		Builder.LogWarning(FString::Printf(TEXT("Failed to read pixels (cube face %d) for cubemap render target %s"), CubeFace, *RenderTargetCube->GetName()));
 		return;
 	}
 
@@ -163,14 +163,14 @@ void FGLTFTextureLightMapTask::Complete()
 	FTextureSource& Source = const_cast<FTextureSource&>(LightMap->Source);
 	if (!Source.IsValid())
 	{
-		Builder.AddWarningMessage(FString::Printf(TEXT("Failed to export lightmap texture %s because of missing source data"), *JsonTexture.Name));
+		Builder.LogWarning(FString::Printf(TEXT("Failed to export lightmap texture %s because of missing source data"), *JsonTexture.Name));
 		return;
 	}
 
 	const ETextureSourceFormat SourceFormat = Source.GetFormat();
 	if (SourceFormat != TSF_BGRA8)
 	{
-		Builder.AddWarningMessage(FString::Printf(TEXT("Failed to export lightmap texture %s because of unsupported source format %s"), *JsonTexture.Name, *FGLTFNameUtility::GetName(SourceFormat)));
+		Builder.LogWarning(FString::Printf(TEXT("Failed to export lightmap texture %s because of unsupported source format %s"), *JsonTexture.Name, *FGLTFNameUtility::GetName(SourceFormat)));
 		return;
 	}
 
