[GLTFExporter] feat: export animations using pose data (instead of raw track data) to support compressed data (for runtime) and improve accuracy

commit 1d4b8752b7fa73a688682226e542431022a4b0e7
Date: Tue, 26 Oct 2021 20:36:08 +0200

    feat: export animations using pose data (instead of raw track data) to support compressed data (for runtime) and improve accuracy
    (cherry picked from commit 38c1284b8cea31c5e57e025a6cf3eb2e34b458d8​)

#skipci

diff --git a/Engine/Plugins/Enterprise/GLTFExporter/MigratedFromGit/Source/GLTFExporter/Private/Tasks/GLTFAnimationTasks.cpp b/Engine/Plugins/Enterprise/GLTFExporter/MigratedFromGit/Source/GLTFExporter/Private/Tasks/GLTFAnimationTasks.cpp
--- a/Engine/Plugins/Enterprise/GLTFExporter/MigratedFromGit/Source/GLTFExporter/Private/Tasks/GLTFAnimationTasks.cpp
+++ b/Engine/Plugins/Enterprise/GLTFExporter/MigratedFromGit/Source/GLTFExporter/Private/Tasks/GLTFAnimationTasks.cpp
@@ -25,30 +25,54 @@ void FGLTFAnimSequenceTask::Complete()
 	FGLTFJsonAnimation& JsonAnimation = Builder.GetAnimation(AnimationIndex);
 	JsonAnimation.Name = AnimSequence->GetName() + TEXT("_") + FString::FromInt(AnimationIndex); // Ensure unique name due to limitation in certain gltf viewers
 
-	const float SequenceLength = AnimSequence->SequenceLength;
-	const float FrameLength = FrameCount > 1 ? SequenceLength / (FrameCount - 1) : 0;
-
 	TArray<float> Timestamps;
 	Timestamps.AddUninitialized(FrameCount);
+	{
+		const float SequenceLength = AnimSequence->SequenceLength;
+		const float FrameLength = FrameCount > 1 ? SequenceLength / (FrameCount - 1) : 0;
 
-	for (int32 Frame = 0; Frame < FrameCount; ++Frame)
+		for (int32 Frame = 0; Frame < FrameCount; ++Frame)
+		{
+			Timestamps[Frame] = FMath::Clamp(Frame * FrameLength, 0.0f, SequenceLength);
+		}
+	}
+
+	FBoneContainer BoneContainer;
+	FGLTFBoneUtility::InitializeToSkeleton(BoneContainer, Skeleton);
+
+	TArray<TArray<FTransform>> FrameTransforms;
+	FrameTransforms.AddDefaulted(FrameCount);
 	{
-		Timestamps[Frame] = FMath::Clamp(Frame * FrameLength, 0.0f, SequenceLength);
+		FCompactPose Pose;
+		Pose.SetBoneContainer(&BoneContainer);
+
+		FBlendedCurve Curve;
+		Curve.InitFrom(BoneContainer);
+
+		FStackCustomAttributes Attributes;
+		FAnimationPoseData PoseData(Pose, Curve, Attributes);
+
+		for (int32 Frame = 0; Frame < FrameCount; ++Frame)
+		{
+			const FAnimExtractContext ExtractionContext(Timestamps[Frame]); // TODO: set bExtractRootMotion?
+			AnimSequence->GetBonePose(PoseData, ExtractionContext);
+			Pose.CopyBonesTo(FrameTransforms[Frame]);
+		}
 	}
 
 	// TODO: add animation data accessor converters to reuse track information
 
-	FGLTFJsonAccessor JsonInputAccessor;
-	JsonInputAccessor.BufferView = Builder.AddBufferView(Timestamps);
-	JsonInputAccessor.ComponentType = EGLTFJsonComponentType::F32;
-	JsonInputAccessor.Type = EGLTFJsonAccessorType::Scalar;
-	JsonInputAccessor.MinMaxLength = 1;
-	JsonInputAccessor.Min[0] = 0;
-
-	FBoneContainer BoneContainer;
-	if (Builder.ExportOptions->bRetargetBoneTransforms)
+	FGLTFJsonAccessorIndex InputAccessorIndex;
 	{
-		FGLTFBoneUtility::InitializeToSkeleton(BoneContainer, Skeleton);
+		FGLTFJsonAccessor JsonInputAccessor;
+		JsonInputAccessor.BufferView = Builder.AddBufferView(Timestamps);
+		JsonInputAccessor.ComponentType = EGLTFJsonComponentType::F32;
+		JsonInputAccessor.Type = EGLTFJsonAccessorType::Scalar;
+		JsonInputAccessor.Count = FrameCount;
+		JsonInputAccessor.MinMaxLength = 1;
+		JsonInputAccessor.Min[0] = 0;
+		JsonInputAccessor.Max[0] = Timestamps[FrameCount - 1];
+		InputAccessorIndex = Builder.AddAccessor(JsonInputAccessor);
 	}
 
 	EGLTFJsonInterpolation Interpolation = FGLTFConverterUtility::ConvertInterpolation(AnimSequence->Interpolation);
@@ -58,59 +82,23 @@ void FGLTFAnimSequenceTask::Complete()
 		// TODO: report warning (about unknown interpolation exported as linear)
 	}
 
-	const TArray<FRawAnimSequenceTrack>& Tracks = AnimSequence->GetRawAnimationData();
-	const int32 TrackCount = Tracks.Num();
-
-	for (int32 TrackIndex = 0; TrackIndex < TrackCount; ++TrackIndex)
+	const TArray<FBoneIndexType>& BoneIndices = BoneContainer.GetBoneIndicesArray();
+	for (FBoneIndexType BoneIndex : BoneIndices)
 	{
-		const FRawAnimSequenceTrack& Track = Tracks[TrackIndex];
-		const TArray<FVector>& KeyPositions = Track.PosKeys;
-		const TArray<FQuat>& KeyRotations = Track.RotKeys;
-		const TArray<FVector>& KeyScales = Track.ScaleKeys;
-
-		const int32 MaxKeys = FMath::Max3(KeyPositions.Num(), KeyRotations.Num(), KeyScales.Num());
-		if (MaxKeys == 0)
-		{
-			continue;
-		}
-
-		TArray<FTransform> KeyTransforms;
-		KeyTransforms.AddUninitialized(MaxKeys);
-
-		for (int32 Key = 0; Key < KeyTransforms.Num(); ++Key)
-		{
-			const FVector& KeyPosition = KeyPositions.IsValidIndex(Key) ? KeyPositions[Key] : FVector::ZeroVector;
-			const FQuat& KeyRotation = KeyRotations.IsValidIndex(Key) ? KeyRotations[Key] : FQuat::Identity;
-			const FVector& KeyScale = KeyScales.IsValidIndex(Key) ? KeyScales[Key] : FVector::OneVector;
-			KeyTransforms[Key] = { KeyRotation, KeyPosition, KeyScale };
-		}
-
-		const int32 SkeletonBoneIndex = AnimSequence->GetSkeletonIndexFromRawDataTrackIndex(TrackIndex);
-		const int32 BoneIndex = const_cast<USkeleton*>(Skeleton)->GetMeshBoneIndexFromSkeletonBoneIndex(SkeletalMesh, SkeletonBoneIndex);
 		const FGLTFJsonNodeIndex NodeIndex = Builder.GetOrAddNode(RootNode, SkeletalMesh, BoneIndex);
 
-		if (Builder.ExportOptions->bRetargetBoneTransforms && Skeleton->GetBoneTranslationRetargetingMode(SkeletonBoneIndex) != EBoneTranslationRetargetingMode::Animation)
-		{
-			for (int32 Key = 0; Key < KeyTransforms.Num(); ++Key)
-			{
-				FGLTFBoneUtility::RetargetTransform(AnimSequence, KeyTransforms[Key], SkeletonBoneIndex, BoneIndex, BoneContainer);
-			}
-		}
+		// TODO: detect if a bone has the same transforms across multiple frames (at least if its the same across all frames) and optimize
 
-		if (KeyPositions.Num() > 0)
 		{
 			TArray<FGLTFVector3> Translations;
-			Translations.AddUninitialized(KeyPositions.Num());
+			Translations.AddUninitialized(FrameCount);
 
-			for (int32 Key = 0; Key < KeyPositions.Num(); ++Key)
+			for (int32 Frame = 0; Frame < FrameCount; ++Frame)
 			{
-				const FVector KeyPosition = KeyTransforms[Key].GetTranslation();
-				Translations[Key] = FGLTFConverterUtility::ConvertPosition(KeyPosition, Builder.ExportOptions->ExportUniformScale);
+				const FVector KeyPosition = FrameTransforms[Frame][BoneIndex].GetTranslation();
+				Translations[Frame] = FGLTFConverterUtility::ConvertPosition(KeyPosition, Builder.ExportOptions->ExportUniformScale);
 			}
 
-			JsonInputAccessor.Count = Translations.Num();
-			JsonInputAccessor.Max[0] = Timestamps[Translations.Num() - 1];
-
 			FGLTFJsonAccessor JsonOutputAccessor;
 			JsonOutputAccessor.BufferView = Builder.AddBufferView(Translations);
 			JsonOutputAccessor.ComponentType = EGLTFJsonComponentType::F32;
@@ -118,7 +106,7 @@ void FGLTFAnimSequenceTask::Complete()
 			JsonOutputAccessor.Type = EGLTFJsonAccessorType::Vec3;
 
 			FGLTFJsonAnimationSampler JsonSampler;
-			JsonSampler.Input = Builder.AddAccessor(JsonInputAccessor);
+			JsonSampler.Input = InputAccessorIndex;
 			JsonSampler.Output = Builder.AddAccessor(JsonOutputAccessor);
 			JsonSampler.Interpolation = Interpolation;
 
@@ -129,20 +117,16 @@ void FGLTFAnimSequenceTask::Complete()
 			JsonAnimation.Channels.Add(JsonChannel);
 		}
 
-		if (KeyRotations.Num() > 0)
 		{
 			TArray<FGLTFQuaternion> Rotations;
-			Rotations.AddUninitialized(KeyRotations.Num());
+			Rotations.AddUninitialized(FrameCount);
 
-			for (int32 Key = 0; Key < KeyRotations.Num(); ++Key)
+			for (int32 Frame = 0; Frame < FrameCount; ++Frame)
 			{
-				const FQuat KeyRotation = KeyTransforms[Key].GetRotation();
-				Rotations[Key] = FGLTFConverterUtility::ConvertRotation(KeyRotation);
+				const FQuat KeyRotation = FrameTransforms[Frame][BoneIndex].GetRotation();
+				Rotations[Frame] = FGLTFConverterUtility::ConvertRotation(KeyRotation);
 			}
 
-			JsonInputAccessor.Count = Rotations.Num();
-			JsonInputAccessor.Max[0] = Timestamps[Rotations.Num() - 1];
-
 			FGLTFJsonAccessor JsonOutputAccessor;
 			JsonOutputAccessor.BufferView = Builder.AddBufferView(Rotations);
 			JsonOutputAccessor.ComponentType = EGLTFJsonComponentType::F32;
@@ -150,7 +134,7 @@ void FGLTFAnimSequenceTask::Complete()
 			JsonOutputAccessor.Type = EGLTFJsonAccessorType::Vec4;
 
 			FGLTFJsonAnimationSampler JsonSampler;
-			JsonSampler.Input = Builder.AddAccessor(JsonInputAccessor);
+			JsonSampler.Input = InputAccessorIndex;
 			JsonSampler.Output = Builder.AddAccessor(JsonOutputAccessor);
 			JsonSampler.Interpolation = Interpolation;
 
@@ -161,20 +145,16 @@ void FGLTFAnimSequenceTask::Complete()
 			JsonAnimation.Channels.Add(JsonChannel);
 		}
 
-		if (KeyScales.Num() > 0)
 		{
 			TArray<FGLTFVector3> Scales;
-			Scales.AddUninitialized(KeyScales.Num());
+			Scales.AddUninitialized(FrameCount);
 
-			for (int32 Key = 0; Key < KeyScales.Num(); ++Key)
+			for (int32 Frame = 0; Frame < FrameCount; ++Frame)
 			{
-				const FVector KeyScale = KeyTransforms[Key].GetScale3D();
-				Scales[Key] = FGLTFConverterUtility::ConvertScale(KeyScale);
+				const FVector KeyScale = FrameTransforms[Frame][BoneIndex].GetScale3D();
+				Scales[Frame] = FGLTFConverterUtility::ConvertScale(KeyScale);
 			}
 
-			JsonInputAccessor.Count = Scales.Num();
-			JsonInputAccessor.Max[0] = Timestamps[Scales.Num() - 1];
-
 			FGLTFJsonAccessor JsonOutputAccessor;
 			JsonOutputAccessor.BufferView = Builder.AddBufferView(Scales);
 			JsonOutputAccessor.ComponentType = EGLTFJsonComponentType::F32;
@@ -182,7 +162,7 @@ void FGLTFAnimSequenceTask::Complete()
 			JsonOutputAccessor.Type = EGLTFJsonAccessorType::Vec3;
 
 			FGLTFJsonAnimationSampler JsonSampler;
-			JsonSampler.Input = Builder.AddAccessor(JsonInputAccessor);
+			JsonSampler.Input = InputAccessorIndex;
 			JsonSampler.Output = Builder.AddAccessor(JsonOutputAccessor);
 			JsonSampler.Interpolation = Interpolation;
 
@@ -237,7 +217,7 @@ void FGLTFLevelSequenceTask::Complete()
 	JsonInputAccessor.Count = FrameCount;
 	JsonInputAccessor.MinMaxLength = 1;
 	JsonInputAccessor.Min[0] = 0;
-	JsonInputAccessor.Max[0] = DisplayRate.AsSeconds(FrameCount - 1);
+	JsonInputAccessor.Max[0] = Timestamps[FrameCount - 1];
 	FGLTFJsonAccessorIndex InputAccessorIndex = Builder.AddAccessor(JsonInputAccessor);
 
 	for (const FMovieSceneBinding& Binding : MovieScene->GetBindings())
