[GLTFExporter] refactor!: remove export option bRetargetBoneTransforms since it will always be applied

commit dd228c5004c3da90443fe3f9a6990558a1e2fac5
Date: Tue, 26 Oct 2021 23:42:16 +0200

    refactor!: remove export option bRetargetBoneTransforms since it will always be applied
    (cherry picked from commit c1f82a28afe5aee18f6bcc4db7d861868bd43527​)

#skipci

diff --git a/Engine/Plugins/Enterprise/GLTFExporter/MigratedFromGit/Docs/export-options-reference.md b/Engine/Plugins/Enterprise/GLTFExporter/MigratedFromGit/Docs/export-options-reference.md
--- a/Engine/Plugins/Enterprise/GLTFExporter/MigratedFromGit/Docs/export-options-reference.md
+++ b/Engine/Plugins/Enterprise/GLTFExporter/MigratedFromGit/Docs/export-options-reference.md
@@ -21,7 +21,6 @@ Option                         | Description
 `Use Mesh Quantization`        | If enabled, use quantization for vertex tangents and normals, reducing size. Requires extension KHR_mesh_quantization, which may result in the mesh not loading in some glTF viewers.
 `Export Level Sequences`       | If enabled, export level sequences. Only transform tracks are currently supported. The level sequence will be played at the assigned display rate.
 `Export Animation Sequences`   | If enabled, export single animation asset used by a skeletal mesh component or hotspot actor. Export of vertex skin weights must be enabled.
-`Retarget Bone Transforms`     | If enabled, apply animation retargeting to skeleton bones when exporting an animation sequence.
 `Export Playback Settings`     | If enabled, export play rate, start time, looping, and auto play for an animation or level sequence. Uses extension EPIC_animation_playback, which is supported by Unreal's glTF viewer.
 `Texture Image Format`         | Desired image format used for exported textures.
 `Texture Image Quality`        | Level of compression used for textures exported with lossy image formats, 0 (default) or value between 1 (worst quality, best compression) and 100 (best quality, worst compression).
diff --git a/Engine/Plugins/Enterprise/GLTFExporter/MigratedFromGit/Source/GLTFExporter/Private/Converters/GLTFBoneUtility.cpp b/Engine/Plugins/Enterprise/GLTFExporter/MigratedFromGit/Source/GLTFExporter/Private/Converters/GLTFBoneUtility.cpp
--- a/Engine/Plugins/Enterprise/GLTFExporter/MigratedFromGit/Source/GLTFExporter/Private/Converters/GLTFBoneUtility.cpp
+++ b/Engine/Plugins/Enterprise/GLTFExporter/MigratedFromGit/Source/GLTFExporter/Private/Converters/GLTFBoneUtility.cpp
@@ -33,10 +33,3 @@ void FGLTFBoneUtility::InitializeToSkeleton(FBoneContainer& BoneContainer, const
 	BoneContainer.SetUseRAWData(true);
 	BoneContainer.InitializeTo(RequiredBoneIndices, FCurveEvaluationOption(true), *const_cast<USkeleton*>(Skeleton));
 }
-
-void FGLTFBoneUtility::RetargetTransform(const UAnimSequence* AnimSequence, FTransform& BoneTransform, int32 SkeletonBoneIndex, int32 BoneIndex, const FBoneContainer& RequiredBones)
-{
-	const USkeleton* Skeleton = AnimSequence->GetSkeleton();
-	const FName& RetargetSource = AnimSequence->RetargetSource;
-	FAnimationRuntime::RetargetBoneTransform(Skeleton, RetargetSource, BoneTransform, SkeletonBoneIndex, FCompactPoseBoneIndex(BoneIndex), RequiredBones, false);
-}
diff --git a/Engine/Plugins/Enterprise/GLTFExporter/MigratedFromGit/Source/GLTFExporter/Private/Converters/GLTFBoneUtility.h b/Engine/Plugins/Enterprise/GLTFExporter/MigratedFromGit/Source/GLTFExporter/Private/Converters/GLTFBoneUtility.h
--- a/Engine/Plugins/Enterprise/GLTFExporter/MigratedFromGit/Source/GLTFExporter/Private/Converters/GLTFBoneUtility.h
+++ b/Engine/Plugins/Enterprise/GLTFExporter/MigratedFromGit/Source/GLTFExporter/Private/Converters/GLTFBoneUtility.h
@@ -9,6 +9,4 @@ struct FGLTFBoneUtility
 	static FTransform GetBindTransform(const FReferenceSkeleton& RefSkeleton, int32 BoneIndex);
 
 	static void InitializeToSkeleton(FBoneContainer& BoneContainer, const USkeleton* Skeleton);
-
-	static void RetargetTransform(const UAnimSequence* AnimSequence, FTransform& BoneTransform, int32 SkeletonBoneIndex, int32 BoneIndex, const FBoneContainer& RequiredBones);
 };
diff --git a/Engine/Plugins/Enterprise/GLTFExporter/MigratedFromGit/Source/GLTFExporter/Private/GLTFExportOptions.cpp b/Engine/Plugins/Enterprise/GLTFExporter/MigratedFromGit/Source/GLTFExporter/Private/GLTFExportOptions.cpp
--- a/Engine/Plugins/Enterprise/GLTFExporter/MigratedFromGit/Source/GLTFExporter/Private/GLTFExportOptions.cpp
+++ b/Engine/Plugins/Enterprise/GLTFExporter/MigratedFromGit/Source/GLTFExporter/Private/GLTFExportOptions.cpp
@@ -29,7 +29,6 @@ void UGLTFExportOptions::ResetToDefault()
 	bUseMeshQuantization = false;
 	bExportLevelSequences = true;
 	bExportAnimationSequences = true;
-	bRetargetBoneTransforms = true;
 	bExportPlaybackSettings = false;
 	TextureImageFormat = EGLTFTextureImageFormat::PNG;
 	TextureImageQuality = 0;
diff --git a/Engine/Plugins/Enterprise/GLTFExporter/MigratedFromGit/Source/GLTFExporter/Public/GLTFExportOptions.h b/Engine/Plugins/Enterprise/GLTFExporter/MigratedFromGit/Source/GLTFExporter/Public/GLTFExportOptions.h
--- a/Engine/Plugins/Enterprise/GLTFExporter/MigratedFromGit/Source/GLTFExporter/Public/GLTFExportOptions.h
+++ b/Engine/Plugins/Enterprise/GLTFExporter/MigratedFromGit/Source/GLTFExporter/Public/GLTFExportOptions.h
@@ -144,10 +144,6 @@ class GLTFEXPORTER_API UGLTFExportOptions : public UObject
 	UPROPERTY(EditAnywhere, BlueprintReadWrite, Config, Category = Animation, Meta = (EditCondition = "bExportVertexSkinWeights"))
 	bool bExportAnimationSequences;
 
-	/** If enabled, apply animation retargeting to skeleton bones when exporting an animation sequence. */
-	UPROPERTY(EditAnywhere, BlueprintReadWrite, Config, Category = Animation, Meta = (EditCondition = "bExportVertexSkinWeights && bExportAnimationSequences"))
-	bool bRetargetBoneTransforms;
-
 	/** If enabled, export play rate, start time, looping, and auto play for an animation or level sequence. Uses extension EPIC_animation_playback, which is supported by Unreal's glTF viewer. */
 	UPROPERTY(EditAnywhere, BlueprintReadWrite, Config, Category = Animation)
 	bool bExportPlaybackSettings;
