[GLTFExporter] fix: warn if world missing levels (which can happen in runtime if not fully loaded)

commit 6998103e977adcb17aaa0bc65a401e8c682f12c9
Date: Tue, 26 Oct 2021 23:47:30 +0200

    fix: warn if world missing levels (which can happen in runtime if not fully loaded)
    (cherry picked from commit 676a7d78fcf4453064f0258560e94f2c5cd2a86c​)

#skipci

diff --git a/Engine/Plugins/Enterprise/GLTFExporter/MigratedFromGit/Source/GLTFExporter/Private/Converters/GLTFSceneConverters.cpp b/Engine/Plugins/Enterprise/GLTFExporter/MigratedFromGit/Source/GLTFExporter/Private/Converters/GLTFSceneConverters.cpp
--- a/Engine/Plugins/Enterprise/GLTFExporter/MigratedFromGit/Source/GLTFExporter/Private/Converters/GLTFSceneConverters.cpp
+++ b/Engine/Plugins/Enterprise/GLTFExporter/MigratedFromGit/Source/GLTFExporter/Private/Converters/GLTFSceneConverters.cpp
@@ -11,53 +11,62 @@ FGLTFJsonSceneIndex FGLTFSceneConverter::Convert(const UWorld* World)
 	FGLTFJsonScene Scene;
 	World->GetName(Scene.Name);
 
-	for (int32 Index = 0; Index < World->GetNumLevels(); ++Index)
+	const TArray<ULevel*>& Levels = World->GetLevels();
+	if (Levels.Num() > 0)
 	{
-		ULevel* Level = World->GetLevel(Index);
-		if (Level == nullptr)
+		for (const ULevel* Level : Levels)
 		{
-			continue;
-		}
-
-		// TODO: export Level->Model ?
-
-		for (const AActor* Actor : Level->Actors)
-		{
-			if (Actor == Level->GetDefaultBrush())
+			if (Level == nullptr)
 			{
-				continue; // TODO: can we safely assume no other actor is ever attached to the default brush?
+				continue;
 			}
 
-			if (World->HasDefaultPhysicsVolume() && Actor == World->GetDefaultPhysicsVolume())
-			{
-				continue; // TODO: can we safely assume no other actor is ever attached to the default physics volume?
-			}
+			// TODO: export Level->Model ?
 
-			// TODO: should a LevelVariantSet be exported even if not selected for export?
-			if (const ALevelVariantSetsActor *LevelVariantSetsActor = Cast<ALevelVariantSetsActor>(Actor))
+			for (const AActor* Actor : Level->Actors)
 			{
-				if (Builder.ExportOptions->bExportVariantSets)
+				if (Actor == Level->GetDefaultBrush())
+				{
+					continue; // TODO: can we safely assume no other actor is ever attached to the default brush?
+				}
+
+				if (World->HasDefaultPhysicsVolume() && Actor == World->GetDefaultPhysicsVolume())
 				{
-					if (const ULevelVariantSets* LevelVariantSets = const_cast<ALevelVariantSetsActor*>(LevelVariantSetsActor)->GetLevelVariantSets(true))
+					continue; // TODO: can we safely assume no other actor is ever attached to the default physics volume?
+				}
+
+				// TODO: should a LevelVariantSet be exported even if not selected for export?
+				if (const ALevelVariantSetsActor *LevelVariantSetsActor = Cast<ALevelVariantSetsActor>(Actor))
+				{
+					if (Builder.ExportOptions->bExportVariantSets)
 					{
-						const FGLTFJsonLevelVariantSetsIndex LevelVariantSetsIndex = Builder.GetOrAddLevelVariantSets(LevelVariantSets);
-						if (LevelVariantSetsIndex != INDEX_NONE)
+						if (const ULevelVariantSets* LevelVariantSets = const_cast<ALevelVariantSetsActor*>(LevelVariantSetsActor)->GetLevelVariantSets(true))
 						{
-							Scene.LevelVariantSets.Add(LevelVariantSetsIndex);
+							const FGLTFJsonLevelVariantSetsIndex LevelVariantSetsIndex = Builder.GetOrAddLevelVariantSets(LevelVariantSets);
+							if (LevelVariantSetsIndex != INDEX_NONE)
+							{
+								Scene.LevelVariantSets.Add(LevelVariantSetsIndex);
+							}
 						}
 					}
 				}
-			}
 
-			const FGLTFJsonNodeIndex NodeIndex = Builder.GetOrAddNode(Actor);
-			if (NodeIndex != INDEX_NONE && FGLTFActorUtility::IsRootActor(Actor, Builder.bSelectedActorsOnly))
-			{
-				// TODO: to avoid having to add irrelevant actors/components let GLTFComponentConverter decide and add root nodes to scene.
-				// This change may require node converters to support cyclic calls.
-				Scene.Nodes.Add(NodeIndex);
+				const FGLTFJsonNodeIndex NodeIndex = Builder.GetOrAddNode(Actor);
+				if (NodeIndex != INDEX_NONE && FGLTFActorUtility::IsRootActor(Actor, Builder.bSelectedActorsOnly))
+				{
+					// TODO: to avoid having to add irrelevant actors/components let GLTFComponentConverter decide and add root nodes to scene.
+					// This change may require node converters to support cyclic calls.
+					Scene.Nodes.Add(NodeIndex);
+				}
 			}
 		}
 	}
+	else
+	{
+		Builder.AddWarningMessage(
+			FString::Printf(TEXT("World %s has no levels. Please make sure the world has been fully initialized"),
+			*World->GetName()));
+	}
 
 	return Builder.AddScene(Scene);
 }
