[GLTFExporter] fix: add sanity checks if vertexbuffer missing data

commit 656ed929913608c2e63ee9982c9218dbede1f016
Date: Tue, 26 Oct 2021 23:52:10 +0200

    fix: add sanity checks if vertexbuffer missing data
    (cherry picked from commit d7ef167f644861a9c16d9261c47fa5b2dc583884​)

#skipci

diff --git a/Engine/Plugins/Enterprise/GLTFExporter/MigratedFromGit/Source/GLTFExporter/Private/Converters/GLTFAccessorConverters.cpp b/Engine/Plugins/Enterprise/GLTFExporter/MigratedFromGit/Source/GLTFExporter/Private/Converters/GLTFAccessorConverters.cpp
--- a/Engine/Plugins/Enterprise/GLTFExporter/MigratedFromGit/Source/GLTFExporter/Private/Converters/GLTFAccessorConverters.cpp
+++ b/Engine/Plugins/Enterprise/GLTFExporter/MigratedFromGit/Source/GLTFExporter/Private/Converters/GLTFAccessorConverters.cpp
@@ -26,6 +26,12 @@ FGLTFJsonAccessorIndex FGLTFPositionBufferConverter::Convert(const FGLTFMeshSect
 		return FGLTFJsonAccessorIndex(INDEX_NONE);
 	}
 
+	if (VertexBuffer->GetVertexData() == nullptr)
+	{
+		// TODO: report error
+		return FGLTFJsonAccessorIndex(INDEX_NONE);
+	}
+
 	const TArray<uint32>& IndexMap = MeshSection->IndexMap;
 	const uint32 VertexCount = IndexMap.Num();
 
@@ -76,6 +82,12 @@ FGLTFJsonAccessorIndex FGLTFColorBufferConverter::Convert(const FGLTFMeshSection
 		return FGLTFJsonAccessorIndex(INDEX_NONE);
 	}
 
+	if (VertexBuffer->GetVertexData() == nullptr)
+	{
+		// TODO: report error
+		return FGLTFJsonAccessorIndex(INDEX_NONE);
+	}
+
 	const TArray<uint32>& IndexMap = MeshSection->IndexMap;
 	const uint32 VertexCount = IndexMap.Num();
 
@@ -105,6 +117,13 @@ FGLTFJsonAccessorIndex FGLTFNormalBufferConverter::Convert(const FGLTFMeshSectio
 		return FGLTFJsonAccessorIndex(INDEX_NONE);
 	}
 
+	const void* TangentData = const_cast<FStaticMeshVertexBuffer*>(VertexBuffer)->GetTangentData();
+	if (TangentData == nullptr)
+	{
+		// TODO: report error
+		return FGLTFJsonAccessorIndex(INDEX_NONE);
+	}
+
 	FGLTFJsonBufferViewIndex BufferViewIndex;
 	EGLTFJsonComponentType ComponentType;
 
@@ -118,13 +137,13 @@ FGLTFJsonAccessorIndex FGLTFNormalBufferConverter::Convert(const FGLTFMeshSectio
 		if (bHighPrecision)
 		{
 			ComponentType = EGLTFJsonComponentType::S16;
-			BufferViewIndex = ConvertBufferView<FGLTFInt16Vector4, FPackedRGBA16N>(MeshSection, VertexBuffer);
+			BufferViewIndex = ConvertBufferView<FGLTFInt16Vector4, FPackedRGBA16N>(MeshSection, TangentData);
 			Builder.GetBufferView(BufferViewIndex).ByteStride = sizeof(FGLTFInt16Vector4);
 		}
 		else
 		{
 			ComponentType = EGLTFJsonComponentType::S8;
-			BufferViewIndex = ConvertBufferView<FGLTFInt8Vector4, FPackedNormal>(MeshSection, VertexBuffer);
+			BufferViewIndex = ConvertBufferView<FGLTFInt8Vector4, FPackedNormal>(MeshSection, TangentData);
 			Builder.GetBufferView(BufferViewIndex).ByteStride = sizeof(FGLTFInt8Vector4);
 		}
 	}
@@ -132,8 +151,8 @@ FGLTFJsonAccessorIndex FGLTFNormalBufferConverter::Convert(const FGLTFMeshSectio
 	{
 		ComponentType = EGLTFJsonComponentType::F32;
 		BufferViewIndex = bHighPrecision
-			? ConvertBufferView<FGLTFVector3, FPackedRGBA16N>(MeshSection, VertexBuffer)
-			: ConvertBufferView<FGLTFVector3, FPackedNormal>(MeshSection, VertexBuffer);
+			? ConvertBufferView<FGLTFVector3, FPackedRGBA16N>(MeshSection, TangentData)
+			: ConvertBufferView<FGLTFVector3, FPackedNormal>(MeshSection, TangentData);
 	}
 
 	FGLTFJsonAccessor JsonAccessor;
@@ -147,18 +166,11 @@ FGLTFJsonAccessorIndex FGLTFNormalBufferConverter::Convert(const FGLTFMeshSectio
 }
 
 template <typename DestinationType, typename SourceType>
-FGLTFJsonBufferViewIndex FGLTFNormalBufferConverter::ConvertBufferView(const FGLTFMeshSection* MeshSection, const FStaticMeshVertexBuffer* VertexBuffer)
+FGLTFJsonBufferViewIndex FGLTFNormalBufferConverter::ConvertBufferView(const FGLTFMeshSection* MeshSection, const void* TangentData)
 {
 	const TArray<uint32>& IndexMap = MeshSection->IndexMap;
 	const uint32 VertexCount = IndexMap.Num();
 
-	const void* TangentData = const_cast<FStaticMeshVertexBuffer*>(VertexBuffer)->GetTangentData();
-	if (TangentData == nullptr)
-	{
-		// TODO: report error
-		return FGLTFJsonBufferViewIndex(INDEX_NONE);
-	}
-
 	typedef TStaticMeshVertexTangentDatum<SourceType> VertexTangentType;
 	const VertexTangentType* VertexTangents = static_cast<const VertexTangentType*>(TangentData);
 
@@ -184,6 +196,13 @@ FGLTFJsonAccessorIndex FGLTFTangentBufferConverter::Convert(const FGLTFMeshSecti
 		return FGLTFJsonAccessorIndex(INDEX_NONE);
 	}
 
+	const void* TangentData = const_cast<FStaticMeshVertexBuffer*>(VertexBuffer)->GetTangentData();
+	if (TangentData == nullptr)
+	{
+		// TODO: report error
+		return FGLTFJsonAccessorIndex(INDEX_NONE);
+	}
+
 	FGLTFJsonBufferViewIndex BufferViewIndex;
 	EGLTFJsonComponentType ComponentType;
 
@@ -197,20 +216,20 @@ FGLTFJsonAccessorIndex FGLTFTangentBufferConverter::Convert(const FGLTFMeshSecti
 		if (bHighPrecision)
 		{
 			ComponentType = EGLTFJsonComponentType::S16;
-			BufferViewIndex = ConvertBufferView<FGLTFInt16Vector4, FPackedRGBA16N>(MeshSection, VertexBuffer);
+			BufferViewIndex = ConvertBufferView<FGLTFInt16Vector4, FPackedRGBA16N>(MeshSection, TangentData);
 		}
 		else
 		{
 			ComponentType = EGLTFJsonComponentType::S8;
-			BufferViewIndex = ConvertBufferView<FGLTFInt8Vector4, FPackedNormal>(MeshSection, VertexBuffer);
+			BufferViewIndex = ConvertBufferView<FGLTFInt8Vector4, FPackedNormal>(MeshSection, TangentData);
 		}
 	}
 	else
 	{
 		ComponentType = EGLTFJsonComponentType::F32;
 		BufferViewIndex = bHighPrecision
-			? ConvertBufferView<FGLTFVector4, FPackedRGBA16N>(MeshSection, VertexBuffer)
-			: ConvertBufferView<FGLTFVector4, FPackedNormal>(MeshSection, VertexBuffer);
+			? ConvertBufferView<FGLTFVector4, FPackedRGBA16N>(MeshSection, TangentData)
+			: ConvertBufferView<FGLTFVector4, FPackedNormal>(MeshSection, TangentData);
 	}
 
 	FGLTFJsonAccessor JsonAccessor;
@@ -224,18 +243,11 @@ FGLTFJsonAccessorIndex FGLTFTangentBufferConverter::Convert(const FGLTFMeshSecti
 }
 
 template <typename DestinationType, typename SourceType>
-FGLTFJsonBufferViewIndex FGLTFTangentBufferConverter::ConvertBufferView(const FGLTFMeshSection* MeshSection, const FStaticMeshVertexBuffer* VertexBuffer)
+FGLTFJsonBufferViewIndex FGLTFTangentBufferConverter::ConvertBufferView(const FGLTFMeshSection* MeshSection, const void* TangentData)
 {
 	const TArray<uint32>& IndexMap = MeshSection->IndexMap;
 	const uint32 VertexCount = IndexMap.Num();
 
-	const void* TangentData = const_cast<FStaticMeshVertexBuffer*>(VertexBuffer)->GetTangentData();
-	if (TangentData == nullptr)
-	{
-		// TODO: report error
-		return FGLTFJsonBufferViewIndex(INDEX_NONE);
-	}
-
 	typedef TStaticMeshVertexTangentDatum<SourceType> VertexTangentType;
 	const VertexTangentType* VertexTangents = static_cast<const VertexTangentType*>(TangentData);
 
@@ -261,6 +273,12 @@ FGLTFJsonAccessorIndex FGLTFUVBufferConverter::Convert(const FGLTFMeshSection* M
 		return FGLTFJsonAccessorIndex(INDEX_NONE);
 	}
 
+	if (const_cast<FStaticMeshVertexBuffer*>(VertexBuffer)->GetTexCoordData() == nullptr)
+	{
+		// TODO: report error
+		return FGLTFJsonAccessorIndex(INDEX_NONE);
+	}
+
 	const TArray<uint32>& IndexMap = MeshSection->IndexMap;
 	const uint32 VertexCount = IndexMap.Num();
 
@@ -306,6 +324,12 @@ FGLTFJsonAccessorIndex FGLTFBoneIndexBufferConverter::Convert(const FGLTFMeshSec
 		return FGLTFJsonAccessorIndex(INDEX_NONE);
 	}
 
+	if (VertexBuffer->GetDataVertexBuffer()->GetWeightData() == nullptr)
+	{
+		// TODO: report error
+		return FGLTFJsonAccessorIndex(INDEX_NONE);
+	}
+
 	const uint32 InfluenceCount = VertexBuffer->GetMaxBoneInfluences();
 	if (InfluenceOffset >= InfluenceCount)
 	{
@@ -353,6 +377,12 @@ FGLTFJsonAccessorIndex FGLTFBoneWeightBufferConverter::Convert(const FGLTFMeshSe
 		return FGLTFJsonAccessorIndex(INDEX_NONE);
 	}
 
+	if (VertexBuffer->GetDataVertexBuffer()->GetWeightData() == nullptr)
+	{
+		// TODO: report error
+		return FGLTFJsonAccessorIndex(INDEX_NONE);
+	}
+
 	const uint32 InfluenceCount = VertexBuffer->GetMaxBoneInfluences();
 	if (InfluenceOffset >= InfluenceCount)
 	{
diff --git a/Engine/Plugins/Enterprise/GLTFExporter/MigratedFromGit/Source/GLTFExporter/Private/Converters/GLTFAccessorConverters.h b/Engine/Plugins/Enterprise/GLTFExporter/MigratedFromGit/Source/GLTFExporter/Private/Converters/GLTFAccessorConverters.h
--- a/Engine/Plugins/Enterprise/GLTFExporter/MigratedFromGit/Source/GLTFExporter/Private/Converters/GLTFAccessorConverters.h
+++ b/Engine/Plugins/Enterprise/GLTFExporter/MigratedFromGit/Source/GLTFExporter/Private/Converters/GLTFAccessorConverters.h
@@ -35,7 +35,7 @@ class FGLTFNormalBufferConverter final : public TGLTFAccessorConverter<const FGL
 	virtual FGLTFJsonAccessorIndex Convert(const FGLTFMeshSection* MeshSection, const FStaticMeshVertexBuffer* VertexBuffer) override;
 
 	template <typename DestinationType, typename SourceType>
-	FGLTFJsonBufferViewIndex ConvertBufferView(const FGLTFMeshSection* MeshSection, const FStaticMeshVertexBuffer* VertexBuffer);
+	FGLTFJsonBufferViewIndex ConvertBufferView(const FGLTFMeshSection* MeshSection, const void* TangentData);
 };
 
 class FGLTFTangentBufferConverter final : public TGLTFAccessorConverter<const FGLTFMeshSection*, const FStaticMeshVertexBuffer*>
@@ -45,7 +45,7 @@ class FGLTFTangentBufferConverter final : public TGLTFAccessorConverter<const FG
 	virtual FGLTFJsonAccessorIndex Convert(const FGLTFMeshSection* MeshSection, const FStaticMeshVertexBuffer* VertexBuffer) override;
 
 	template <typename DestinationType, typename SourceType>
-	FGLTFJsonBufferViewIndex ConvertBufferView(const FGLTFMeshSection* MeshSection, const FStaticMeshVertexBuffer* VertexBuffer);
+	FGLTFJsonBufferViewIndex ConvertBufferView(const FGLTFMeshSection* MeshSection, const void* TangentData);
 };
 
 class FGLTFUVBufferConverter final : public TGLTFAccessorConverter<const FGLTFMeshSection*, const FStaticMeshVertexBuffer*, uint32>
diff --git a/Engine/Plugins/Enterprise/GLTFExporter/MigratedFromGit/Source/GLTFExporter/Private/Converters/GLTFMeshSection.cpp b/Engine/Plugins/Enterprise/GLTFExporter/MigratedFromGit/Source/GLTFExporter/Private/Converters/GLTFMeshSection.cpp
--- a/Engine/Plugins/Enterprise/GLTFExporter/MigratedFromGit/Source/GLTFExporter/Private/Converters/GLTFMeshSection.cpp
+++ b/Engine/Plugins/Enterprise/GLTFExporter/MigratedFromGit/Source/GLTFExporter/Private/Converters/GLTFMeshSection.cpp
@@ -7,6 +7,13 @@
 
 FGLTFMeshSection::FGLTFMeshSection(const FStaticMeshLODResources* MeshLOD, const FGLTFIndexArray& SectionIndices)
 {
+	const FRawStaticIndexBuffer& OldIndexBuffer = MeshLOD->IndexBuffer;
+	if (OldIndexBuffer.AccessStream16() == nullptr && OldIndexBuffer.AccessStream32() == nullptr)
+	{
+		// TODO: report error
+		return;
+	}
+
 	const FStaticMeshLODResources::FStaticMeshSectionArray& Sections = MeshLOD->Sections;
 
 	uint32 TotalIndexCount = 0;
@@ -29,7 +36,7 @@ FGLTFMeshSection::FGLTFMeshSection(const FStaticMeshLODResources* MeshLOD, const
 
 		for (uint32 Index = 0; Index < IndexCount; Index++)
 		{
-			const uint32 OldIndex = MeshLOD->IndexBuffer.GetIndex(IndexOffset + Index);
+			const uint32 OldIndex = OldIndexBuffer.GetIndex(IndexOffset + Index);
 			uint32 NewIndex;
 
 			if (uint32* FoundIndex = IndexLookup.Find(OldIndex))
@@ -55,6 +62,13 @@ FGLTFMeshSection::FGLTFMeshSection(const FStaticMeshLODResources* MeshLOD, const
 FGLTFMeshSection::FGLTFMeshSection(const FSkeletalMeshLODRenderData* MeshLOD, const FGLTFIndexArray& SectionIndices)
 	: MaxBoneIndex(0)
 {
+	const FRawStaticIndexBuffer16or32Interface* OldIndexBuffer = MeshLOD->MultiSizeIndexContainer.GetIndexBuffer();
+	if (const_cast<FRawStaticIndexBuffer16or32Interface*>(OldIndexBuffer)->GetPointerTo(0) == nullptr)
+	{
+		// TODO: report error
+		return;
+	}
+
 	const TArray<FSkelMeshRenderSection>& Sections = MeshLOD->RenderSections;
 
 	uint32 TotalIndexCount = 0;
@@ -69,8 +83,6 @@ FGLTFMeshSection::FGLTFMeshSection(const FSkeletalMeshLODRenderData* MeshLOD, co
 
 	TMap<uint32, uint32> IndexLookup;
 
-	const FRawStaticIndexBuffer16or32Interface* OldIndexBuffer = MeshLOD->MultiSizeIndexContainer.GetIndexBuffer();
-
 	for (int32 SectionIndex : SectionIndices)
 	{
 		const FSkelMeshRenderSection& MeshSection = Sections[SectionIndex];
