[GLTFExporter] fix: add checks and error messages for Allow CPU Access when exporting meshes in runtime

commit 63dd2fe9de21b01317fbded5eb72612b62e9fa32
Date: Sun, 31 Oct 2021 02:03:11 +0100

    fix: add checks and error messages for Allow CPU Access when exporting meshes in runtime
    (cherry picked from commit 4799e3c23b446f387e5b243ed16f54a9c1c56d76​)

#skipci

diff --git a/Engine/Plugins/Enterprise/GLTFExporter/MigratedFromGit/Source/GLTFExporter/Private/Converters/GLTFMeshConverters.cpp b/Engine/Plugins/Enterprise/GLTFExporter/MigratedFromGit/Source/GLTFExporter/Private/Converters/GLTFMeshConverters.cpp
--- a/Engine/Plugins/Enterprise/GLTFExporter/MigratedFromGit/Source/GLTFExporter/Private/Converters/GLTFMeshConverters.cpp
+++ b/Engine/Plugins/Enterprise/GLTFExporter/MigratedFromGit/Source/GLTFExporter/Private/Converters/GLTFMeshConverters.cpp
@@ -40,6 +40,16 @@ void FGLTFStaticMeshConverter::Sanitize(const UStaticMesh*& StaticMesh, const US
 
 FGLTFJsonMeshIndex FGLTFStaticMeshConverter::Convert(const UStaticMesh* StaticMesh,  const UStaticMeshComponent* StaticMeshComponent, FGLTFMaterialArray Materials, int32 LODIndex)
 {
+#if !WITH_EDITOR
+	if (!StaticMesh->bAllowCPUAccess)
+	{
+		Builder.LogError(FString::Printf(
+			TEXT("Can't export mesh %s in runtime because Allow CPU Access is not enabled in asset settings"),
+			*StaticMesh->GetName()));
+		return FGLTFJsonMeshIndex(INDEX_NONE);
+	}
+#endif
+
 	const FGLTFJsonMeshIndex MeshIndex = Builder.AddMesh();
 	Builder.SetupTask<FGLTFStaticMeshTask>(Builder, MeshSectionConverter, StaticMesh, StaticMeshComponent, Materials, LODIndex, MeshIndex);
 	return MeshIndex;
@@ -79,6 +89,17 @@ void FGLTFSkeletalMeshConverter::Sanitize(const USkeletalMesh*& SkeletalMesh, co
 
 FGLTFJsonMeshIndex FGLTFSkeletalMeshConverter::Convert(const USkeletalMesh* SkeletalMesh, const USkeletalMeshComponent* SkeletalMeshComponent, FGLTFMaterialArray Materials, int32 LODIndex)
 {
+#if !WITH_EDITOR
+	if (!SkeletalMesh->GetLODInfo(LODIndex)->bAllowCPUAccess)
+	{
+		Builder.LogError(FString::Printf(
+			TEXT("Can't export mesh %s in runtime because Allow CPU Access is not enabled for LOD %d in asset settings"),
+			*SkeletalMesh->GetName(),
+			LODIndex));
+		return FGLTFJsonMeshIndex(INDEX_NONE);
+	}
+#endif
+
 	const FGLTFJsonMeshIndex MeshIndex = Builder.AddMesh();
 	Builder.SetupTask<FGLTFSkeletalMeshTask>(Builder, MeshSectionConverter, SkeletalMesh, SkeletalMeshComponent, Materials, LODIndex, MeshIndex);
 	return MeshIndex;
