[GLTFExporter] fix: add sanity check if cubemap platform data missing

commit 53cad619577df794154cdc59d2d28d54d5f9cc2c
Date: Wed, 27 Oct 2021 00:03:05 +0200

    fix: add sanity check if cubemap platform data missing
    (cherry picked from commit 425206173605978df88ad8a02b86896041050864​)

#skipci

diff --git a/Engine/Plugins/Enterprise/GLTFExporter/MigratedFromGit/Source/GLTFExporter/Private/Converters/GLTFTextureUtility.cpp b/Engine/Plugins/Enterprise/GLTFExporter/MigratedFromGit/Source/GLTFExporter/Private/Converters/GLTFTextureUtility.cpp
--- a/Engine/Plugins/Enterprise/GLTFExporter/MigratedFromGit/Source/GLTFExporter/Private/Converters/GLTFTextureUtility.cpp
+++ b/Engine/Plugins/Enterprise/GLTFExporter/MigratedFromGit/Source/GLTFExporter/Private/Converters/GLTFTextureUtility.cpp
@@ -224,11 +224,15 @@ UTexture2D* FGLTFTextureUtility::CreateTextureFromCubeFace(const UTextureCube* T
 
 	const FByteBulkData& BulkData = TextureCube->PlatformData->Mips[0].BulkData;
 	const int64 MipSize = BulkData.GetBulkDataSize() / 6;
+	UTexture2D* FaceTexture = nullptr;
 
 	const void* MipDataPtr = BulkData.LockReadOnly();
-	const void* FaceDataPtr =  static_cast<const uint8*>(MipDataPtr) + MipSize * CubeFace;
-	UTexture2D* FaceTexture = CreateTransientTexture(FaceDataPtr, MipSize, Size, Format, TextureCube->SRGB);
-	BulkData.Unlock();
+	if (MipDataPtr != nullptr)
+	{
+		const void* FaceDataPtr = static_cast<const uint8*>(MipDataPtr) + MipSize * CubeFace;
+		FaceTexture = CreateTransientTexture(FaceDataPtr, MipSize, Size, Format, TextureCube->SRGB);
+	}
+	const_cast<FByteBulkData&>(BulkData).Unlock();
 
 	return FaceTexture;
 }
