[GLTFExporter] fix: avoid FMaterialPropertyEx dependency in GLTFBuilder to improve runtime compatibility

commit 3756cad8d299821ca150f2cec7112f8c78c4d456
Date: Wed, 27 Oct 2021 01:09:45 +0200

    fix: avoid FMaterialPropertyEx dependency in GLTFBuilder to improve runtime compatibility
    (cherry picked from commit db1d926f4fc62efd02b31f9d536c1c30d25eea74​)

#skipci

diff --git a/Engine/Plugins/Enterprise/GLTFExporter/MigratedFromGit/Source/GLTFExporter/Private/Builders/GLTFBuilder.cpp b/Engine/Plugins/Enterprise/GLTFExporter/MigratedFromGit/Source/GLTFExporter/Private/Builders/GLTFBuilder.cpp
--- a/Engine/Plugins/Enterprise/GLTFExporter/MigratedFromGit/Source/GLTFExporter/Private/Builders/GLTFBuilder.cpp
+++ b/Engine/Plugins/Enterprise/GLTFExporter/MigratedFromGit/Source/GLTFExporter/Private/Builders/GLTFBuilder.cpp
@@ -13,11 +13,9 @@ FGLTFBuilder::FGLTFBuilder(const FString& FilePath, const UGLTFExportOptions* Ex
 {
 }
 
-FIntPoint FGLTFBuilder::GetBakeSizeForMaterialProperty(const UMaterialInterface* Material, const FMaterialPropertyEx& Property) const
+FIntPoint FGLTFBuilder::GetBakeSizeForMaterialProperty(const UMaterialInterface* Material, EGLTFMaterialPropertyGroup PropertyGroup) const
 {
-	const EGLTFMaterialPropertyGroup PropertyGroup = GetPropertyGroup(Property);
 	EGLTFMaterialBakeSizePOT DefaultValue = ExportOptions->DefaultMaterialBakeSize;
-
 	if (const FGLTFOverrideMaterialBakeSettings* BakeSettings = ExportOptions->DefaultInputBakeSettings.Find(PropertyGroup))
 	{
 		if (BakeSettings->bOverrideSize)
@@ -31,11 +29,9 @@ FIntPoint FGLTFBuilder::GetBakeSizeForMaterialProperty(const UMaterialInterface*
 	return { PixelSize, PixelSize };
 }
 
-TextureFilter FGLTFBuilder::GetBakeFilterForMaterialProperty(const UMaterialInterface* Material, const FMaterialPropertyEx& Property) const
+TextureFilter FGLTFBuilder::GetBakeFilterForMaterialProperty(const UMaterialInterface* Material, EGLTFMaterialPropertyGroup PropertyGroup) const
 {
-	const EGLTFMaterialPropertyGroup PropertyGroup = GetPropertyGroup(Property);
 	TextureFilter DefaultValue = ExportOptions->DefaultMaterialBakeFilter;
-
 	if (const FGLTFOverrideMaterialBakeSettings* BakeSettings = ExportOptions->DefaultInputBakeSettings.Find(PropertyGroup))
 	{
 		if (BakeSettings->bOverrideFilter)
@@ -47,11 +43,9 @@ TextureFilter FGLTFBuilder::GetBakeFilterForMaterialProperty(const UMaterialInte
 	return UGLTFMaterialExportOptions::GetBakeFilterForPropertyGroup(Material, PropertyGroup, DefaultValue);
 }
 
-TextureAddress FGLTFBuilder::GetBakeTilingForMaterialProperty(const UMaterialInterface* Material, const FMaterialPropertyEx& Property) const
+TextureAddress FGLTFBuilder::GetBakeTilingForMaterialProperty(const UMaterialInterface* Material, EGLTFMaterialPropertyGroup PropertyGroup) const
 {
-	const EGLTFMaterialPropertyGroup PropertyGroup = GetPropertyGroup(Property);
 	TextureAddress DefaultValue = ExportOptions->DefaultMaterialBakeTiling;
-
 	if (const FGLTFOverrideMaterialBakeSettings* BakeSettings = ExportOptions->DefaultInputBakeSettings.Find(PropertyGroup))
 	{
 		if (BakeSettings->bOverrideTiling)
@@ -112,33 +106,3 @@ EGLTFSceneMobility FGLTFBuilder::GetSceneMobility(EComponentMobility::Type Mobil
 			return EGLTFSceneMobility::None;
 	}
 }
-
-EGLTFMaterialPropertyGroup FGLTFBuilder::GetPropertyGroup(const FMaterialPropertyEx& Property)
-{
-	switch (Property.Type)
-	{
-		case MP_BaseColor:
-		case MP_Opacity:
-		case MP_OpacityMask:
-			return EGLTFMaterialPropertyGroup::BaseColorOpacity;
-		case MP_Metallic:
-		case MP_Roughness:
-			return EGLTFMaterialPropertyGroup::MetallicRoughness;
-		case MP_EmissiveColor:
-			return EGLTFMaterialPropertyGroup::EmissiveColor;
-		case MP_Normal:
-			return EGLTFMaterialPropertyGroup::Normal;
-		case MP_AmbientOcclusion:
-			return EGLTFMaterialPropertyGroup::AmbientOcclusion;
-		case MP_CustomData0:
-		case MP_CustomData1:
-			return EGLTFMaterialPropertyGroup::ClearCoatRoughness;
-		case MP_CustomOutput:
-			if (Property.CustomOutput == TEXT("ClearCoatBottomNormal"))
-			{
-				return EGLTFMaterialPropertyGroup::ClearCoatBottomNormal;
-			}
-		default:
-			return EGLTFMaterialPropertyGroup::None;
-	}
-}
diff --git a/Engine/Plugins/Enterprise/GLTFExporter/MigratedFromGit/Source/GLTFExporter/Private/Builders/GLTFBuilder.h b/Engine/Plugins/Enterprise/GLTFExporter/MigratedFromGit/Source/GLTFExporter/Private/Builders/GLTFBuilder.h
--- a/Engine/Plugins/Enterprise/GLTFExporter/MigratedFromGit/Source/GLTFExporter/Private/Builders/GLTFBuilder.h
+++ b/Engine/Plugins/Enterprise/GLTFExporter/MigratedFromGit/Source/GLTFExporter/Private/Builders/GLTFBuilder.h
@@ -4,7 +4,6 @@
 
 #include "GLTFExportOptions.h"
 #include "Json/GLTFJsonEnums.h"
-#include "MaterialPropertyEx.h"
 #include "UObject/GCObjectScopeGuard.h"
 
 class FGLTFBuilder
@@ -22,9 +21,9 @@ public:
 	// TODO: make ExportOptions private and expose each option via getters to ease overriding settings in future
 	const UGLTFExportOptions* ExportOptions;
 
-	FIntPoint GetBakeSizeForMaterialProperty(const UMaterialInterface* Material, const FMaterialPropertyEx& Property) const;
-	TextureFilter GetBakeFilterForMaterialProperty(const UMaterialInterface* Material, const FMaterialPropertyEx& Property) const;
-	TextureAddress GetBakeTilingForMaterialProperty(const UMaterialInterface* Material, const FMaterialPropertyEx& Property) const;
+	FIntPoint GetBakeSizeForMaterialProperty(const UMaterialInterface* Material, EGLTFMaterialPropertyGroup PropertyGroup) const;
+	TextureFilter GetBakeFilterForMaterialProperty(const UMaterialInterface* Material, EGLTFMaterialPropertyGroup PropertyGroup) const;
+	TextureAddress GetBakeTilingForMaterialProperty(const UMaterialInterface* Material, EGLTFMaterialPropertyGroup PropertyGroup) const;
 
 	EGLTFJsonHDREncoding GetTextureHDREncoding() const;
 
@@ -33,8 +32,7 @@ public:
 private:
 
 	FGCObjectScopeGuard ExportOptionsGuard;
-	
+
 	static const UGLTFExportOptions* ValidateExportOptions(const UGLTFExportOptions* ExportOptions);
 	static EGLTFSceneMobility GetSceneMobility(EComponentMobility::Type Mobility);
-	static EGLTFMaterialPropertyGroup GetPropertyGroup(const FMaterialPropertyEx& Property);
 };
diff --git a/Engine/Plugins/Enterprise/GLTFExporter/MigratedFromGit/Source/GLTFExporter/Private/Tasks/GLTFMaterialTasks.cpp b/Engine/Plugins/Enterprise/GLTFExporter/MigratedFromGit/Source/GLTFExporter/Private/Tasks/GLTFMaterialTasks.cpp
--- a/Engine/Plugins/Enterprise/GLTFExporter/MigratedFromGit/Source/GLTFExporter/Private/Tasks/GLTFMaterialTasks.cpp
+++ b/Engine/Plugins/Enterprise/GLTFExporter/MigratedFromGit/Source/GLTFExporter/Private/Tasks/GLTFMaterialTasks.cpp
@@ -376,13 +376,13 @@ bool FGLTFMaterialTask::TryGetBaseColorAndOpacity(FGLTFJsonPBRMetallicRoughness&
 		return false;
 	}
 
-	FIntPoint TextureSize = Builder.GetBakeSizeForMaterialProperty(Material, BaseColorProperty);
+	FIntPoint TextureSize = Builder.GetBakeSizeForMaterialProperty(Material, GetPropertyGroup(BaseColorProperty));
 
-	const TextureAddress TextureAddress = Builder.GetBakeTilingForMaterialProperty(Material, BaseColorProperty);
+	const TextureAddress TextureAddress = Builder.GetBakeTilingForMaterialProperty(Material, GetPropertyGroup(BaseColorProperty));
 	const EGLTFJsonTextureWrap TextureWrapS = FGLTFConverterUtility::ConvertWrap(TextureAddress);
 	const EGLTFJsonTextureWrap TextureWrapT = FGLTFConverterUtility::ConvertWrap(TextureAddress);
 
-	const TextureFilter TextureFilter = Builder.GetBakeFilterForMaterialProperty(Material, BaseColorProperty);
+	const TextureFilter TextureFilter = Builder.GetBakeFilterForMaterialProperty(Material, GetPropertyGroup(BaseColorProperty));
 	const EGLTFJsonTextureFilter TextureMinFilter = FGLTFConverterUtility::ConvertMinFilter(TextureFilter);
 	const EGLTFJsonTextureFilter TextureMagFilter = FGLTFConverterUtility::ConvertMagFilter(TextureFilter);
 
@@ -502,13 +502,13 @@ bool FGLTFMaterialTask::TryGetMetallicAndRoughness(FGLTFJsonPBRMetallicRoughness
 	}
 
 	// TODO: add support for calculating the ideal resolution to use for baking based on connected (texture) nodes
-	FIntPoint TextureSize = Builder.GetBakeSizeForMaterialProperty(Material, MetallicProperty);
+	FIntPoint TextureSize = Builder.GetBakeSizeForMaterialProperty(Material, GetPropertyGroup(MetallicProperty));
 
-	const TextureAddress TextureAddress = Builder.GetBakeTilingForMaterialProperty(Material, MetallicProperty);
+	const TextureAddress TextureAddress = Builder.GetBakeTilingForMaterialProperty(Material, GetPropertyGroup(MetallicProperty));
 	const EGLTFJsonTextureWrap TextureWrapS = FGLTFConverterUtility::ConvertWrap(TextureAddress);
 	const EGLTFJsonTextureWrap TextureWrapT = FGLTFConverterUtility::ConvertWrap(TextureAddress);
 
-	const TextureFilter TextureFilter = Builder.GetBakeFilterForMaterialProperty(Material, MetallicProperty);
+	const TextureFilter TextureFilter = Builder.GetBakeFilterForMaterialProperty(Material, GetPropertyGroup(MetallicProperty));
 	const EGLTFJsonTextureFilter TextureMinFilter = FGLTFConverterUtility::ConvertMinFilter(TextureFilter);
 	const EGLTFJsonTextureFilter TextureMagFilter = FGLTFConverterUtility::ConvertMagFilter(TextureFilter);
 
@@ -629,13 +629,13 @@ bool FGLTFMaterialTask::TryGetClearCoatRoughness(FGLTFJsonClearCoatExtension& Ou
 		return false;
 	}
 
-	FIntPoint TextureSize = Builder.GetBakeSizeForMaterialProperty(Material, IntensityProperty);
+	FIntPoint TextureSize = Builder.GetBakeSizeForMaterialProperty(Material, GetPropertyGroup(IntensityProperty));
 
-	const TextureAddress TextureAddress = Builder.GetBakeTilingForMaterialProperty(Material, IntensityProperty);
+	const TextureAddress TextureAddress = Builder.GetBakeTilingForMaterialProperty(Material,GetPropertyGroup(IntensityProperty));
 	const EGLTFJsonTextureWrap TextureWrapS = FGLTFConverterUtility::ConvertWrap(TextureAddress);
 	const EGLTFJsonTextureWrap TextureWrapT = FGLTFConverterUtility::ConvertWrap(TextureAddress);
 
-	const TextureFilter TextureFilter = Builder.GetBakeFilterForMaterialProperty(Material, IntensityProperty);
+	const TextureFilter TextureFilter = Builder.GetBakeFilterForMaterialProperty(Material, GetPropertyGroup(IntensityProperty));
 	const EGLTFJsonTextureFilter TextureMinFilter = FGLTFConverterUtility::ConvertMinFilter(TextureFilter);
 	const EGLTFJsonTextureFilter TextureMagFilter = FGLTFConverterUtility::ConvertMagFilter(TextureFilter);
 
@@ -1331,7 +1331,7 @@ bool FGLTFMaterialTask::TryGetBakedMaterialProperty(FGLTFJsonTextureInfo& OutTex
 
 FGLTFPropertyBakeOutput FGLTFMaterialTask::BakeMaterialProperty(const FMaterialPropertyEx& Property, int32& OutTexCoord, bool bCopyAlphaFromRedChannel)
 {
-	const FIntPoint TextureSize = Builder.GetBakeSizeForMaterialProperty(Material, Property);
+	const FIntPoint TextureSize = Builder.GetBakeSizeForMaterialProperty(Material, GetPropertyGroup(Property));
 	return BakeMaterialProperty(Property, OutTexCoord, TextureSize, bCopyAlphaFromRedChannel);
 }
 
@@ -1383,11 +1383,11 @@ FGLTFPropertyBakeOutput FGLTFMaterialTask::BakeMaterialProperty(const FMaterialP
 
 bool FGLTFMaterialTask::StoreBakedPropertyTexture(FGLTFJsonTextureInfo& OutTexInfo, const FGLTFPropertyBakeOutput& PropertyBakeOutput, const FString& PropertyName) const
 {
-	const TextureAddress TextureAddress = Builder.GetBakeTilingForMaterialProperty(Material, PropertyBakeOutput.Property);
+	const TextureAddress TextureAddress = Builder.GetBakeTilingForMaterialProperty(Material, GetPropertyGroup(PropertyBakeOutput.Property));
 	const EGLTFJsonTextureWrap TextureWrapS = FGLTFConverterUtility::ConvertWrap(TextureAddress);
 	const EGLTFJsonTextureWrap TextureWrapT = FGLTFConverterUtility::ConvertWrap(TextureAddress);
 
-	const TextureFilter TextureFilter = Builder.GetBakeFilterForMaterialProperty(Material, PropertyBakeOutput.Property);
+	const TextureFilter TextureFilter = Builder.GetBakeFilterForMaterialProperty(Material, GetPropertyGroup(PropertyBakeOutput.Property));
 	const EGLTFJsonTextureFilter TextureMinFilter = FGLTFConverterUtility::ConvertMinFilter(TextureFilter);
 	const EGLTFJsonTextureFilter TextureMagFilter = FGLTFConverterUtility::ConvertMagFilter(TextureFilter);
 
@@ -1423,3 +1423,33 @@ FString FGLTFMaterialTask::GetBakedTextureName(const FString& PropertyName) cons
 {
 	return GetMaterialName() + TEXT("_") + PropertyName;
 }
+
+EGLTFMaterialPropertyGroup FGLTFMaterialTask::GetPropertyGroup(const FMaterialPropertyEx& Property)
+{
+	switch (Property.Type)
+	{
+		case MP_BaseColor:
+		case MP_Opacity:
+		case MP_OpacityMask:
+			return EGLTFMaterialPropertyGroup::BaseColorOpacity;
+		case MP_Metallic:
+		case MP_Roughness:
+			return EGLTFMaterialPropertyGroup::MetallicRoughness;
+		case MP_EmissiveColor:
+			return EGLTFMaterialPropertyGroup::EmissiveColor;
+		case MP_Normal:
+			return EGLTFMaterialPropertyGroup::Normal;
+		case MP_AmbientOcclusion:
+			return EGLTFMaterialPropertyGroup::AmbientOcclusion;
+		case MP_CustomData0:
+		case MP_CustomData1:
+			return EGLTFMaterialPropertyGroup::ClearCoatRoughness;
+		case MP_CustomOutput:
+			if (Property.CustomOutput == TEXT("ClearCoatBottomNormal"))
+			{
+				return EGLTFMaterialPropertyGroup::ClearCoatBottomNormal;
+			}
+		default:
+			return EGLTFMaterialPropertyGroup::None;
+	}
+}
diff --git a/Engine/Plugins/Enterprise/GLTFExporter/MigratedFromGit/Source/GLTFExporter/Private/Tasks/GLTFMaterialTasks.h b/Engine/Plugins/Enterprise/GLTFExporter/MigratedFromGit/Source/GLTFExporter/Private/Tasks/GLTFMaterialTasks.h
--- a/Engine/Plugins/Enterprise/GLTFExporter/MigratedFromGit/Source/GLTFExporter/Private/Tasks/GLTFMaterialTasks.h
+++ b/Engine/Plugins/Enterprise/GLTFExporter/MigratedFromGit/Source/GLTFExporter/Private/Tasks/GLTFMaterialTasks.h
@@ -74,4 +74,6 @@ private:
 
 	FString GetMaterialName() const;
 	FString GetBakedTextureName(const FString& PropertyName) const;
+
+	static EGLTFMaterialPropertyGroup GetPropertyGroup(const FMaterialPropertyEx& Property);
 };
