[GLTFExporter] fix: use default value for minimum LOD instead of platform-specific to avoid runtime issues

commit b64a95b0fb86f2f9c47534657976ee300bc30707
Date: Wed, 27 Oct 2021 01:11:35 +0200

    fix: use default value for minimum LOD instead of platform-specific to avoid runtime issues
    (cherry picked from commit d46e142d7c98094762e09ef0b80add6d1fb7ee2b​)

#skipci

diff --git a/Engine/Plugins/Enterprise/GLTFExporter/MigratedFromGit/Source/GLTFExporter/Private/Converters/GLTFMeshUtility.cpp b/Engine/Plugins/Enterprise/GLTFExporter/MigratedFromGit/Source/GLTFExporter/Private/Converters/GLTFMeshUtility.cpp
--- a/Engine/Plugins/Enterprise/GLTFExporter/MigratedFromGit/Source/GLTFExporter/Private/Converters/GLTFMeshUtility.cpp
+++ b/Engine/Plugins/Enterprise/GLTFExporter/MigratedFromGit/Source/GLTFExporter/Private/Converters/GLTFMeshUtility.cpp
@@ -2,7 +2,6 @@
 
 #include "Converters/GLTFMeshUtility.h"
 #include "Rendering/SkeletalMeshRenderData.h"
-#include "PlatformInfo.h"
 
 FGLTFIndexArray FGLTFMeshUtility::GetSectionIndices(const FStaticMeshLODResources& MeshLOD, int32 MaterialIndex)
 {
@@ -82,7 +81,7 @@ int32 FGLTFMeshUtility::GetMinimumLOD(const UStaticMesh* StaticMesh, const UStat
 
 	if (StaticMesh != nullptr)
 	{
-		return GetValueForRunningPlatform<int32>(StaticMesh->GetMinLOD());
+		return StaticMesh->GetMinLOD().Default;
 	}
 
 	return -1;
@@ -97,15 +96,8 @@ int32 FGLTFMeshUtility::GetMinimumLOD(const USkeletalMesh* SkeletalMesh, const U
 
 	if (SkeletalMesh != nullptr)
 	{
-		return GetValueForRunningPlatform<int32>(SkeletalMesh->GetMinLod());
+		return SkeletalMesh->GetMinLod().Default;
 	}
 
 	return -1;
 }
-
-template <typename ValueType, typename StructType>
-ValueType FGLTFMeshUtility::GetValueForRunningPlatform(const StructType& Properties)
-{
-	const PlatformInfo::FPlatformInfo& PlatformInfo = GetTargetPlatformManagerRef().GetRunningTargetPlatform()->GetPlatformInfo();
-	return Properties.GetValueForPlatformIdentifiers(PlatformInfo.PlatformGroupName, PlatformInfo.VanillaPlatformName);
-}
diff --git a/Engine/Plugins/Enterprise/GLTFExporter/MigratedFromGit/Source/GLTFExporter/Private/Converters/GLTFMeshUtility.h b/Engine/Plugins/Enterprise/GLTFExporter/MigratedFromGit/Source/GLTFExporter/Private/Converters/GLTFMeshUtility.h
--- a/Engine/Plugins/Enterprise/GLTFExporter/MigratedFromGit/Source/GLTFExporter/Private/Converters/GLTFMeshUtility.h
+++ b/Engine/Plugins/Enterprise/GLTFExporter/MigratedFromGit/Source/GLTFExporter/Private/Converters/GLTFMeshUtility.h
@@ -20,9 +20,4 @@ struct FGLTFMeshUtility
 
 	static int32 GetMinimumLOD(const UStaticMesh* StaticMesh, const UStaticMeshComponent* StaticMeshComponent);
 	static int32 GetMinimumLOD(const USkeletalMesh* SkeletalMesh, const USkeletalMeshComponent* SkeletalMeshComponent);
-
-private:
-
-	template <typename ValueType, typename StructType>
-	static ValueType GetValueForRunningPlatform(const StructType& Properties);
 };
