[GLTFExporter] fix: add WITH_EDITOR preprocessor guards to avoid some (not all) editor dependencies in runtime

commit 791338e00a469b2ab13688dc64a2d71b75ca7c08
Date: Wed, 27 Oct 2021 01:30:24 +0200

    fix: add WITH_EDITOR preprocessor guards to avoid some (not all) editor dependencies in runtime
    (cherry picked from commit 502f92b5fae6a8ae87f79e5df7516cdcf3d9eb51​)

#skipci

diff --git a/Engine/Plugins/Enterprise/GLTFExporter/MigratedFromGit/Source/GLTFExporter/Private/Builders/GLTFLogBuilder.cpp b/Engine/Plugins/Enterprise/GLTFExporter/MigratedFromGit/Source/GLTFExporter/Private/Builders/GLTFLogBuilder.cpp
--- a/Engine/Plugins/Enterprise/GLTFExporter/MigratedFromGit/Source/GLTFExporter/Private/Builders/GLTFLogBuilder.cpp
+++ b/Engine/Plugins/Enterprise/GLTFExporter/MigratedFromGit/Source/GLTFExporter/Private/Builders/GLTFLogBuilder.cpp
@@ -2,19 +2,23 @@
 
 #include "Builders/GLTFLogBuilder.h"
 #include "GLTFExporterModule.h"
+#include "Misc/FeedbackContext.h"
+#if WITH_EDITOR
 #include "MessageLogModule.h"
 #include "IMessageLogListing.h"
-#include "Misc/FeedbackContext.h"
+#endif
 
 FGLTFLogBuilder::FGLTFLogBuilder(const FString& FilePath, const UGLTFExportOptions* ExportOptions)
 	: FGLTFBuilder(FilePath, ExportOptions)
 {
+#if WITH_EDITOR
 	if (!FApp::IsUnattended())
 	{
 		FMessageLogModule& MessageLogModule = FModuleManager::LoadModuleChecked<FMessageLogModule>("MessageLog");
 		LogListing = MessageLogModule.GetLogListing(GLTFEXPORTER_MODULE_NAME);
 		LogListing->SetLabel(FText::FromString(GLTFEXPORTER_FRIENDLY_NAME));
 	}
+#endif
 }
 
 void FGLTFLogBuilder::LogSuggestion(const FString& Message)
@@ -52,10 +56,12 @@ const TArray<FString>& FGLTFLogBuilder::GetLoggedErrors() const
 
 void FGLTFLogBuilder::OpenLog() const
 {
+#if WITH_EDITOR
 	if (LogListing != nullptr)
 	{
 		LogListing->Open();
 	}
+#endif
 }
 
 void FGLTFLogBuilder::ClearLog()
@@ -64,10 +70,12 @@ void FGLTFLogBuilder::ClearLog()
 	Warnings.Empty();
 	Errors.Empty();
 
+#if WITH_EDITOR
 	if (LogListing != nullptr)
 	{
 		LogListing->ClearMessages();
 	}
+#endif
 }
 
 void FGLTFLogBuilder::PrintToLog(ELogLevel Level, const FString& Message) const
@@ -88,6 +96,7 @@ void FGLTFLogBuilder::PrintToLog(ELogLevel Level, const FString& Message) const
 	GWarn->Log(LogGLTFExporter.GetCategoryName(), Verbosity, Message);
 #endif
 
+#if WITH_EDITOR
 	if (LogListing != nullptr)
 	{
 		EMessageSeverity::Type Severity;
@@ -104,4 +113,5 @@ void FGLTFLogBuilder::PrintToLog(ELogLevel Level, const FString& Message) const
 
 		LogListing->AddMessage(FTokenizedMessage::Create(Severity, FText::FromString(Message)), false);
 	}
+#endif
 }
diff --git a/Engine/Plugins/Enterprise/GLTFExporter/MigratedFromGit/Source/GLTFExporter/Private/Builders/GLTFLogBuilder.h b/Engine/Plugins/Enterprise/GLTFExporter/MigratedFromGit/Source/GLTFExporter/Private/Builders/GLTFLogBuilder.h
--- a/Engine/Plugins/Enterprise/GLTFExporter/MigratedFromGit/Source/GLTFExporter/Private/Builders/GLTFLogBuilder.h
+++ b/Engine/Plugins/Enterprise/GLTFExporter/MigratedFromGit/Source/GLTFExporter/Private/Builders/GLTFLogBuilder.h
@@ -4,7 +4,9 @@
 
 #include "Builders/GLTFBuilder.h"
 
+#if WITH_EDITOR
 class IMessageLogListing;
+#endif
 
 class FGLTFLogBuilder : public FGLTFBuilder
 {
@@ -45,5 +47,7 @@ private:
 	TArray<FString> Warnings;
 	TArray<FString> Errors;
 
+#if WITH_EDITOR
 	TSharedPtr<IMessageLogListing> LogListing;
+#endif
 };
diff --git a/Engine/Plugins/Enterprise/GLTFExporter/MigratedFromGit/Source/GLTFExporter/Private/Converters/GLTFMaterialConverters.cpp b/Engine/Plugins/Enterprise/GLTFExporter/MigratedFromGit/Source/GLTFExporter/Private/Converters/GLTFMaterialConverters.cpp
--- a/Engine/Plugins/Enterprise/GLTFExporter/MigratedFromGit/Source/GLTFExporter/Private/Converters/GLTFMaterialConverters.cpp
+++ b/Engine/Plugins/Enterprise/GLTFExporter/MigratedFromGit/Source/GLTFExporter/Private/Converters/GLTFMaterialConverters.cpp
@@ -15,6 +15,7 @@ void FGLTFMaterialConverter::Sanitize(const UMaterialInterface*& Material, const
 		SectionIndices = {};
 	}
 
+#if WITH_EDITOR
 	if (MeshData != nullptr)
 	{
 		const FMeshDescription& MeshDescription = MeshData->GetParent()->Description;
@@ -37,16 +38,19 @@ void FGLTFMaterialConverter::Sanitize(const UMaterialInterface*& Material, const
 			SectionIndices = {};
 		}
 	}
+#endif
 }
 
 FGLTFJsonMaterialIndex FGLTFMaterialConverter::Convert(const UMaterialInterface* Material, const FGLTFMeshData* MeshData, FGLTFIndexArray SectionIndices)
 {
-	if (Material == FGLTFMaterialUtility::GetDefault())
+#if WITH_EDITOR
+	if (Material != FGLTFMaterialUtility::GetDefault())
 	{
-		return FGLTFJsonMaterialIndex(INDEX_NONE); // use default gltf definition
+		const FGLTFJsonMaterialIndex MaterialIndex = Builder.AddMaterial();
+		Builder.SetupTask<FGLTFMaterialTask>(Builder, UVOverlapChecker, Material, MeshData, SectionIndices, MaterialIndex);
+		return MaterialIndex;
 	}
+#endif
 
-	const FGLTFJsonMaterialIndex MaterialIndex = Builder.AddMaterial();
-	Builder.SetupTask<FGLTFMaterialTask>(Builder, UVOverlapChecker, Material, MeshData, SectionIndices, MaterialIndex);
-	return MaterialIndex;
+	return FGLTFJsonMaterialIndex(INDEX_NONE); // use default gltf definition
 }
diff --git a/Engine/Plugins/Enterprise/GLTFExporter/MigratedFromGit/Source/GLTFExporter/Private/Converters/GLTFMaterialConverters.h b/Engine/Plugins/Enterprise/GLTFExporter/MigratedFromGit/Source/GLTFExporter/Private/Converters/GLTFMaterialConverters.h
--- a/Engine/Plugins/Enterprise/GLTFExporter/MigratedFromGit/Source/GLTFExporter/Private/Converters/GLTFMaterialConverters.h
+++ b/Engine/Plugins/Enterprise/GLTFExporter/MigratedFromGit/Source/GLTFExporter/Private/Converters/GLTFMaterialConverters.h
@@ -19,6 +19,8 @@ class FGLTFMaterialConverter final : public FGLTFBuilderContext, public TGLTFCon
 
 	virtual FGLTFJsonMaterialIndex Convert(const UMaterialInterface* Material, const FGLTFMeshData* MeshData, FGLTFIndexArray SectionIndices) override;
 
+#if WITH_EDITOR
 	FGLTFUVOverlapChecker UVOverlapChecker;
 	FGLTFUVDegenerateChecker UVDegenerateChecker;
+#endif
 };
diff --git a/Engine/Plugins/Enterprise/GLTFExporter/MigratedFromGit/Source/GLTFExporter/Private/Converters/GLTFMaterialUtility.cpp b/Engine/Plugins/Enterprise/GLTFExporter/MigratedFromGit/Source/GLTFExporter/Private/Converters/GLTFMaterialUtility.cpp
--- a/Engine/Plugins/Enterprise/GLTFExporter/MigratedFromGit/Source/GLTFExporter/Private/Converters/GLTFMaterialUtility.cpp
+++ b/Engine/Plugins/Enterprise/GLTFExporter/MigratedFromGit/Source/GLTFExporter/Private/Converters/GLTFMaterialUtility.cpp
@@ -2,6 +2,7 @@
 
 #include "Converters/GLTFMaterialUtility.h"
 #include "Converters/GLTFTextureUtility.h"
+#if WITH_EDITOR
 #include "Converters/GLTFNameUtility.h"
 #include "GLTFMaterialAnalyzer.h"
 #include "Engine/TextureRenderTarget2D.h"
@@ -15,6 +16,7 @@
 #include "Materials/MaterialExpressionClearCoatNormalCustomOutput.h"
 #include "Materials/MaterialExpressionTextureCoordinate.h"
 #include "MeshDescription.h"
+#endif
 #include "Misc/DefaultValueHelper.h"
 
 UMaterialInterface* FGLTFMaterialUtility::GetDefault()
@@ -23,6 +25,8 @@ UMaterialInterface* FGLTFMaterialUtility::GetDefault()
 	return DefaultMaterial;
 }
 
+#if WITH_EDITOR
+
 bool FGLTFMaterialUtility::IsNormalMap(const FMaterialPropertyEx& Property)
 {
 	return Property == MP_Normal || Property == TEXT("ClearCoatBottomNormal");
@@ -378,46 +382,62 @@ FString FGLTFMaterialUtility::ShadingModelsToString(const FMaterialShadingModelF
 	return Result;
 }
 
-bool FGLTFMaterialUtility::NeedsMeshData(const UMaterialInterface* Material)
+void FGLTFMaterialUtility::AnalyzeMaterialProperty(const UMaterialInterface* InMaterial, const FMaterialPropertyEx& InProperty, FGLTFMaterialAnalysis& OutAnalysis)
 {
-	if (Material == nullptr)
+	if (GetInputForProperty(InMaterial, InProperty) == nullptr)
 	{
-		return false;
+		OutAnalysis = FGLTFMaterialAnalysis();
+		return;
 	}
 
-	// TODO: only analyze properties that will be needed for this specific material
-	const TArray<FMaterialPropertyEx> Properties =
-	{
-		MP_BaseColor,
-		MP_EmissiveColor,
-		MP_Opacity,
-		MP_OpacityMask,
-		MP_Metallic,
-		MP_Roughness,
-		MP_Normal,
-		MP_AmbientOcclusion,
-		MP_CustomData0,
-		MP_CustomData1,
-		TEXT("ClearCoatBottomNormal"),
-	};
-
-	bool bNeedsMeshData = false;
-	FGLTFMaterialAnalysis Analysis;
+	UGLTFMaterialAnalyzer::AnalyzeMaterialPropertyEx(InMaterial, InProperty.Type, InProperty.CustomOutput.ToString(), OutAnalysis);
+}
 
-	// TODO: optimize baking by seperating need for vertex data and primitive data
+#endif
 
-	for (const FMaterialPropertyEx& Property: Properties)
+bool FGLTFMaterialUtility::NeedsMeshData(const UMaterialInterface* Material)
+{
+#if WITH_EDITOR
+	if (Material != nullptr)
 	{
-		AnalyzeMaterialProperty(Material, Property, Analysis);
-		bNeedsMeshData |= Analysis.bRequiresVertexData;
-		bNeedsMeshData |= Analysis.bRequiresPrimitiveData;
+		// TODO: only analyze properties that will be needed for this specific material
+		const TArray<FMaterialPropertyEx> Properties =
+		{
+			MP_BaseColor,
+			MP_EmissiveColor,
+			MP_Opacity,
+			MP_OpacityMask,
+			MP_Metallic,
+			MP_Roughness,
+			MP_Normal,
+			MP_AmbientOcclusion,
+			MP_CustomData0,
+			MP_CustomData1,
+			TEXT("ClearCoatBottomNormal"),
+		};
+
+		bool bNeedsMeshData = false;
+		FGLTFMaterialAnalysis Analysis;
+
+		// TODO: optimize baking by seperating need for vertex data and primitive data
+
+		for (const FMaterialPropertyEx& Property: Properties)
+		{
+			AnalyzeMaterialProperty(Material, Property, Analysis);
+			bNeedsMeshData |= Analysis.bRequiresVertexData;
+			bNeedsMeshData |= Analysis.bRequiresPrimitiveData;
+		}
+
+		return bNeedsMeshData;
 	}
+#endif
 
-	return bNeedsMeshData;
+	return false;
 }
 
 bool FGLTFMaterialUtility::NeedsMeshData(const TArray<const UMaterialInterface*>& Materials)
 {
+#if WITH_EDITOR
 	for (const UMaterialInterface* Material: Materials)
 	{
 		if (NeedsMeshData(Material))
@@ -425,21 +445,11 @@ bool FGLTFMaterialUtility::NeedsMeshData(const TArray<const UMaterialInterface*>
 			return true;
 		}
 	}
+#endif
 
 	return false;
 }
 
-void FGLTFMaterialUtility::AnalyzeMaterialProperty(const UMaterialInterface* InMaterial, const FMaterialPropertyEx& InProperty, FGLTFMaterialAnalysis& OutAnalysis)
-{
-	if (GetInputForProperty(InMaterial, InProperty) == nullptr)
-	{
-		OutAnalysis = FGLTFMaterialAnalysis();
-		return;
-	}
-
-	UGLTFMaterialAnalyzer::AnalyzeMaterialPropertyEx(InMaterial, InProperty.Type, InProperty.CustomOutput.ToString(), OutAnalysis);
-}
-
 const UMaterialInterface* FGLTFMaterialUtility::GetInterface(const UMaterialInterface* Material)
 {
 	return Material;
diff --git a/Engine/Plugins/Enterprise/GLTFExporter/MigratedFromGit/Source/GLTFExporter/Private/Converters/GLTFMaterialUtility.h b/Engine/Plugins/Enterprise/GLTFExporter/MigratedFromGit/Source/GLTFExporter/Private/Converters/GLTFMaterialUtility.h
--- a/Engine/Plugins/Enterprise/GLTFExporter/MigratedFromGit/Source/GLTFExporter/Private/Converters/GLTFMaterialUtility.h
+++ b/Engine/Plugins/Enterprise/GLTFExporter/MigratedFromGit/Source/GLTFExporter/Private/Converters/GLTFMaterialUtility.h
@@ -39,6 +39,7 @@ struct FGLTFMaterialUtility
 {
 	static UMaterialInterface* GetDefault();
 
+#if WITH_EDITOR
 	static bool IsNormalMap(const FMaterialPropertyEx& Property);
 
 	static FGuid GetAttributeID(const FMaterialPropertyEx& Property);
@@ -77,11 +78,12 @@ struct FGLTFMaterialUtility
 	static EMaterialShadingModel GetRichestShadingModel(const FMaterialShadingModelField& ShadingModels);
 	static FString ShadingModelsToString(const FMaterialShadingModelField& ShadingModels);
 
+	static void AnalyzeMaterialProperty(const UMaterialInterface* Material, const FMaterialPropertyEx& InProperty, FGLTFMaterialAnalysis& OutAnalysis);
+#endif
+
 	static bool NeedsMeshData(const UMaterialInterface* Material);
 	static bool NeedsMeshData(const TArray<const UMaterialInterface*>& Materials);
 
-	static void AnalyzeMaterialProperty(const UMaterialInterface* Material, const FMaterialPropertyEx& InProperty, FGLTFMaterialAnalysis& OutAnalysis);
-
 	static const UMaterialInterface* GetInterface(const UMaterialInterface* Material);
 	static const UMaterialInterface* GetInterface(const FStaticMaterial& Material);
 	static const UMaterialInterface* GetInterface(const FSkeletalMaterial& Material);
diff --git a/Engine/Plugins/Enterprise/GLTFExporter/MigratedFromGit/Source/GLTFExporter/Private/Converters/GLTFMeshData.cpp b/Engine/Plugins/Enterprise/GLTFExporter/MigratedFromGit/Source/GLTFExporter/Private/Converters/GLTFMeshData.cpp
--- a/Engine/Plugins/Enterprise/GLTFExporter/MigratedFromGit/Source/GLTFExporter/Private/Converters/GLTFMeshData.cpp
+++ b/Engine/Plugins/Enterprise/GLTFExporter/MigratedFromGit/Source/GLTFExporter/Private/Converters/GLTFMeshData.cpp
@@ -3,19 +3,31 @@
 #include "Converters/GLTFMeshData.h"
 #include "Converters/GLTFNameUtility.h"
 #include "Converters/GLTFMeshUtility.h"
+
+#if WITH_EDITOR
 #include "StaticMeshAttributes.h"
 #include "Engine/MapBuildDataRegistry.h"
 #include "Developer/MeshMergeUtilities/Private/MeshMergeHelpers.h"
 #include "Rendering/SkeletalMeshRenderData.h"
+#endif
 
 FGLTFMeshData::FGLTFMeshData(const UStaticMesh* StaticMesh, const UStaticMeshComponent* StaticMeshComponent, int32 LODIndex)
-	: LightMap(nullptr), LightMapResourceCluster(nullptr), LightMapTexCoord(0), Parent(nullptr)
+	: Parent(nullptr)
+#if WITH_EDITOR
+	, LightMap(nullptr)
+	, LightMapResourceCluster(nullptr)
+	, LightMapTexCoord(0)
+#endif
 {
+#if WITH_EDITOR
 	FStaticMeshAttributes(Description).Register();
+#endif
 
 	if (StaticMeshComponent != nullptr)
 	{
 		Name = FGLTFNameUtility::GetName(StaticMeshComponent);
+
+#if WITH_EDITOR
 		PrimitiveData = { StaticMeshComponent };
 		FMeshMergeHelpers::RetrieveMesh(StaticMeshComponent, LODIndex, Description, true);
 
@@ -29,35 +41,52 @@ FGLTFMeshData::FGLTFMeshData(const UStaticMesh* StaticMesh, const UStaticMeshCom
 				LightMapResourceCluster = BuildData->ResourceCluster;
 			}
 		}
+#endif
 	}
 	else
 	{
 		StaticMesh->GetName(Name);
+
+#if WITH_EDITOR
 		PrimitiveData = { StaticMesh };
 		FMeshMergeHelpers::RetrieveMesh(StaticMesh, LODIndex, Description);
+#endif
 	}
 
+#if WITH_EDITOR
 	LightMapTexCoord = StaticMesh->GetLightMapCoordinateIndex();
  	const int32 NumTexCoords = StaticMesh->GetLODForExport(LODIndex).VertexBuffers.StaticMeshVertexBuffer.GetNumTexCoords();
 	BakeUsingTexCoord = FMath::Min(LightMapTexCoord, NumTexCoords - 1);
-
-	// TODO: add warning if texture coordinate has overlap
+	// TODO: add warning if texture coordinate for baking has overlap
+#endif
 }
 
 FGLTFMeshData::FGLTFMeshData(const USkeletalMesh* SkeletalMesh, const USkeletalMeshComponent* SkeletalMeshComponent, int32 LODIndex)
-	: LightMap(nullptr), LightMapResourceCluster(nullptr), LightMapTexCoord(0), Parent(nullptr)
+	: Parent(nullptr)
+#if WITH_EDITOR
+	, LightMap(nullptr)
+	, LightMapResourceCluster(nullptr)
+	, LightMapTexCoord(0)
+#endif
 {
+#if WITH_EDITOR
 	FStaticMeshAttributes(Description).Register();
+#endif
 
 	if (SkeletalMeshComponent != nullptr)
 	{
 		Name = FGLTFNameUtility::GetName(SkeletalMeshComponent);
+
+#if WITH_EDITOR
 		PrimitiveData = { SkeletalMeshComponent };
 		FMeshMergeHelpers::RetrieveMesh(const_cast<USkeletalMeshComponent*>(SkeletalMeshComponent), LODIndex, Description, true);
+#endif
 	}
 	else
 	{
 		SkeletalMesh->GetName(Name);
+
+#if WITH_EDITOR
 		PrimitiveData = { SkeletalMesh };
 
 		// NOTE: this is a workaround for the fact that there's no overload for FMeshMergeHelpers::RetrieveMesh
@@ -83,13 +112,15 @@ FGLTFMeshData::FGLTFMeshData(const USkeletalMesh* SkeletalMesh, const USkeletalM
 				World->DestroyActor(Actor, false, false);
 			}
 		}
+#endif
 	}
 
+#if WITH_EDITOR
 	// TODO: don't assume last UV channel is non-overlapping
 	const int32 NumTexCoords = SkeletalMesh->GetResourceForRendering()->LODRenderData[LODIndex].StaticVertexBuffers.StaticMeshVertexBuffer.GetNumTexCoords();
 	BakeUsingTexCoord = NumTexCoords - 1;
-
-	// TODO: add warning if texture coordinate has overlap
+	// TODO: add warning if texture coordinate for baking has overlap
+#endif
 }
 
 const FGLTFMeshData* FGLTFMeshData::GetParent() const
diff --git a/Engine/Plugins/Enterprise/GLTFExporter/MigratedFromGit/Source/GLTFExporter/Private/Converters/GLTFMeshData.h b/Engine/Plugins/Enterprise/GLTFExporter/MigratedFromGit/Source/GLTFExporter/Private/Converters/GLTFMeshData.h
--- a/Engine/Plugins/Enterprise/GLTFExporter/MigratedFromGit/Source/GLTFExporter/Private/Converters/GLTFMeshData.h
+++ b/Engine/Plugins/Enterprise/GLTFExporter/MigratedFromGit/Source/GLTFExporter/Private/Converters/GLTFMeshData.h
@@ -4,7 +4,10 @@
 
 #include "CoreMinimal.h"
 #include "MeshDescription.h"
+
+#if WITH_EDITOR
 #include "MaterialBakingStructures.h"
+#endif
 
 struct FGLTFMeshData
 {
@@ -13,7 +16,12 @@ struct FGLTFMeshData
 
 	const FGLTFMeshData* GetParent() const;
 
+	// TODO: find a better name for referencing the mesh-only data (no component)
+	const FGLTFMeshData* Parent;
+
 	FString Name;
+
+#if WITH_EDITOR
 	FPrimitiveData PrimitiveData;
 	FMeshDescription Description;
 
@@ -22,7 +30,5 @@ struct FGLTFMeshData
 	int32 LightMapTexCoord;
 
 	int32 BakeUsingTexCoord;
-
-	// TODO: find a better name for referencing the mesh-only data (no component)
-	const FGLTFMeshData* Parent;
+#endif
 };
diff --git a/Engine/Plugins/Enterprise/GLTFExporter/MigratedFromGit/Source/GLTFExporter/Private/Converters/GLTFSceneConverters.cpp b/Engine/Plugins/Enterprise/GLTFExporter/MigratedFromGit/Source/GLTFExporter/Private/Converters/GLTFSceneConverters.cpp
--- a/Engine/Plugins/Enterprise/GLTFExporter/MigratedFromGit/Source/GLTFExporter/Private/Converters/GLTFSceneConverters.cpp
+++ b/Engine/Plugins/Enterprise/GLTFExporter/MigratedFromGit/Source/GLTFExporter/Private/Converters/GLTFSceneConverters.cpp
@@ -63,7 +63,7 @@ FGLTFJsonSceneIndex FGLTFSceneConverter::Convert(const UWorld* World)
 	}
 	else
 	{
-		Builder.AddWarningMessage(
+		Builder.LogWarning(
 			FString::Printf(TEXT("World %s has no levels. Please make sure the world has been fully initialized"),
 			*World->GetName()));
 	}
diff --git a/Engine/Plugins/Enterprise/GLTFExporter/MigratedFromGit/Source/GLTFExporter/Private/Converters/GLTFTextureConverters.cpp b/Engine/Plugins/Enterprise/GLTFExporter/MigratedFromGit/Source/GLTFExporter/Private/Converters/GLTFTextureConverters.cpp
--- a/Engine/Plugins/Enterprise/GLTFExporter/MigratedFromGit/Source/GLTFExporter/Private/Converters/GLTFTextureConverters.cpp
+++ b/Engine/Plugins/Enterprise/GLTFExporter/MigratedFromGit/Source/GLTFExporter/Private/Converters/GLTFTextureConverters.cpp
@@ -5,43 +5,43 @@
 
 FGLTFJsonTextureIndex FGLTFTexture2DConverter::Convert(const UTexture2D* Texture2D)
 {
-	if (Builder.ExportOptions->TextureImageFormat == EGLTFTextureImageFormat::None)
+	if (Builder.ExportOptions->TextureImageFormat != EGLTFTextureImageFormat::None)
 	{
-		return FGLTFJsonTextureIndex(INDEX_NONE);
+		const FGLTFJsonTextureIndex TextureIndex = Builder.AddTexture();
+		Builder.SetupTask<FGLTFTexture2DTask>(Builder, Texture2D, TextureIndex);
+		return TextureIndex;
 	}
 
-	const FGLTFJsonTextureIndex TextureIndex = Builder.AddTexture();
-	Builder.SetupTask<FGLTFTexture2DTask>(Builder, Texture2D, TextureIndex);
-	return TextureIndex;
+	return FGLTFJsonTextureIndex(INDEX_NONE);
 }
 
 FGLTFJsonTextureIndex FGLTFTextureCubeConverter::Convert(const UTextureCube* TextureCube, ECubeFace CubeFace)
 {
-	if (Builder.ExportOptions->TextureImageFormat == EGLTFTextureImageFormat::None)
+	if (Builder.ExportOptions->TextureImageFormat != EGLTFTextureImageFormat::None)
 	{
-		return FGLTFJsonTextureIndex(INDEX_NONE);
+		const FGLTFJsonTextureIndex TextureIndex = Builder.AddTexture();
+		Builder.SetupTask<FGLTFTextureCubeTask>(Builder, TextureCube, CubeFace, TextureIndex);
+		return TextureIndex;
 	}
 
-	const FGLTFJsonTextureIndex TextureIndex = Builder.AddTexture();
-	Builder.SetupTask<FGLTFTextureCubeTask>(Builder, TextureCube, CubeFace, TextureIndex);
-	return TextureIndex;
+	return FGLTFJsonTextureIndex(INDEX_NONE);
 }
 
 FGLTFJsonTextureIndex FGLTFTextureRenderTarget2DConverter::Convert(const UTextureRenderTarget2D* RenderTarget2D)
 {
-	if (Builder.ExportOptions->TextureImageFormat == EGLTFTextureImageFormat::None)
+	if (Builder.ExportOptions->TextureImageFormat != EGLTFTextureImageFormat::None)
 	{
-		return FGLTFJsonTextureIndex(INDEX_NONE);
+		const FGLTFJsonTextureIndex TextureIndex = Builder.AddTexture();
+		Builder.SetupTask<FGLTFTextureRenderTarget2DTask>(Builder, RenderTarget2D, TextureIndex);
+		return TextureIndex;
 	}
 
-	const FGLTFJsonTextureIndex TextureIndex = Builder.AddTexture();
-	Builder.SetupTask<FGLTFTextureRenderTarget2DTask>(Builder, RenderTarget2D, TextureIndex);
-	return TextureIndex;
+	return FGLTFJsonTextureIndex(INDEX_NONE);
 }
 
 FGLTFJsonTextureIndex FGLTFTextureRenderTargetCubeConverter::Convert(const UTextureRenderTargetCube* RenderTargetCube, ECubeFace CubeFace)
 {
-	if (Builder.ExportOptions->TextureImageFormat == EGLTFTextureImageFormat::None)
+	if (Builder.ExportOptions->TextureImageFormat != EGLTFTextureImageFormat::None)
 	{
 		return FGLTFJsonTextureIndex(INDEX_NONE);
 	}
@@ -53,12 +53,14 @@ FGLTFJsonTextureIndex FGLTFTextureRenderTargetCubeConverter::Convert(const UText
 
 FGLTFJsonTextureIndex FGLTFTextureLightMapConverter::Convert(const ULightMapTexture2D* LightMap)
 {
-	if (Builder.ExportOptions->TextureImageFormat == EGLTFTextureImageFormat::None)
+#if WITH_EDITOR
+	if (Builder.ExportOptions->TextureImageFormat != EGLTFTextureImageFormat::None)
 	{
-		return FGLTFJsonTextureIndex(INDEX_NONE);
+		const FGLTFJsonTextureIndex TextureIndex = Builder.AddTexture();
+		Builder.SetupTask<FGLTFTextureLightMapTask>(Builder, LightMap, TextureIndex);
+		return TextureIndex;
 	}
+#endif
 
-	const FGLTFJsonTextureIndex TextureIndex = Builder.AddTexture();
-	Builder.SetupTask<FGLTFTextureLightMapTask>(Builder, LightMap, TextureIndex);
-	return TextureIndex;
+	return FGLTFJsonTextureIndex(INDEX_NONE);
 }
diff --git a/Engine/Plugins/Enterprise/GLTFExporter/MigratedFromGit/Source/GLTFExporter/Private/Converters/GLTFTextureUtility.cpp b/Engine/Plugins/Enterprise/GLTFExporter/MigratedFromGit/Source/GLTFExporter/Private/Converters/GLTFTextureUtility.cpp
--- a/Engine/Plugins/Enterprise/GLTFExporter/MigratedFromGit/Source/GLTFExporter/Private/Converters/GLTFTextureUtility.cpp
+++ b/Engine/Plugins/Enterprise/GLTFExporter/MigratedFromGit/Source/GLTFExporter/Private/Converters/GLTFTextureUtility.cpp
@@ -141,8 +141,10 @@ UTexture2D* FGLTFTextureUtility::CreateTransientTexture(const void* RawData, int
 	Texture->PlatformData->Mips[0].BulkData.Unlock();
 
 	Texture->SRGB = bSRGB ? 1 : 0;
+#if WITH_EDITOR
 	Texture->CompressionNone = true;
 	Texture->MipGenSettings = TMGS_NoMipmaps;
+#endif
 
 	Texture->UpdateResource();
 	return Texture;
@@ -352,17 +354,15 @@ bool FGLTFTextureUtility::LoadPlatformData(UTexture2D* Texture)
 		return false;
 	}
 
+#if WITH_EDITOR
 	if (Texture->PlatformData->Mips[0].BulkData.GetBulkDataSize() == 0)
 	{
 		// TODO: is this correct handling?
 		Texture->ForceRebuildPlatformData();
-		if (Texture->PlatformData->Mips[0].BulkData.GetBulkDataSize() == 0)
-		{
-			return false;
-		}
 	}
+#endif
 
-	return true;
+	return Texture->PlatformData->Mips[0].BulkData.GetBulkDataSize() != 0;
 }
 
 bool FGLTFTextureUtility::LoadPlatformData(UTextureCube* TextureCube)
@@ -372,17 +372,15 @@ bool FGLTFTextureUtility::LoadPlatformData(UTextureCube* TextureCube)
 		return false;
 	}
 
+#if WITH_EDITOR
 	if (TextureCube->PlatformData->Mips[0].BulkData.GetBulkDataSize() == 0)
 	{
 		// TODO: is this correct handling?
 		TextureCube->ForceRebuildPlatformData();
-		if (TextureCube->PlatformData->Mips[0].BulkData.GetBulkDataSize() == 0)
-		{
-			return false;
-		}
 	}
+#endif
 
-	return true;
+	return TextureCube->PlatformData->Mips[0].BulkData.GetBulkDataSize() != 0;
 }
 
 void FGLTFTextureUtility::FlipGreenChannel(TArray<FColor>& Pixels)
diff --git a/Engine/Plugins/Enterprise/GLTFExporter/MigratedFromGit/Source/GLTFExporter/Private/Converters/GLTFUVDegenerateChecker.cpp b/Engine/Plugins/Enterprise/GLTFExporter/MigratedFromGit/Source/GLTFExporter/Private/Converters/GLTFUVDegenerateChecker.cpp
--- a/Engine/Plugins/Enterprise/GLTFExporter/MigratedFromGit/Source/GLTFExporter/Private/Converters/GLTFUVDegenerateChecker.cpp
+++ b/Engine/Plugins/Enterprise/GLTFExporter/MigratedFromGit/Source/GLTFExporter/Private/Converters/GLTFUVDegenerateChecker.cpp
@@ -1,5 +1,7 @@
 // Copyright Epic Games, Inc. All Rights Reserved.
 
+#if WITH_EDITOR
+
 #include "Converters/GLTFUVDegenerateChecker.h"
 #include "StaticMeshAttributes.h"
 #include "MeshDescription.h"
@@ -103,3 +105,5 @@ bool FGLTFUVDegenerateChecker::IsDegenerateTriangle(const TStaticArray<FVector,
 	const float DoubleAreaSquared = (AB ^ AC).SizeSquared();
 	return DoubleAreaSquared < 2 * SMALL_NUMBER * SMALL_NUMBER;
 }
+
+#endif
diff --git a/Engine/Plugins/Enterprise/GLTFExporter/MigratedFromGit/Source/GLTFExporter/Private/Converters/GLTFUVDegenerateChecker.h b/Engine/Plugins/Enterprise/GLTFExporter/MigratedFromGit/Source/GLTFExporter/Private/Converters/GLTFUVDegenerateChecker.h
--- a/Engine/Plugins/Enterprise/GLTFExporter/MigratedFromGit/Source/GLTFExporter/Private/Converters/GLTFUVDegenerateChecker.h
+++ b/Engine/Plugins/Enterprise/GLTFExporter/MigratedFromGit/Source/GLTFExporter/Private/Converters/GLTFUVDegenerateChecker.h
@@ -2,6 +2,8 @@
 
 #pragma once
 
+#if WITH_EDITOR
+
 #include "Converters/GLTFConverter.h"
 #include "Converters/GLTFIndexArray.h"
 #include "Engine.h"
@@ -15,3 +17,5 @@ class FGLTFUVDegenerateChecker final : public TGLTFConverter<float, const FMeshD
 	static bool IsDegenerateTriangle(const TStaticArray<FVector2D, 3>& Points);
 	static bool IsDegenerateTriangle(const TStaticArray<FVector, 3>& Points);
 };
+
+#endif
diff --git a/Engine/Plugins/Enterprise/GLTFExporter/MigratedFromGit/Source/GLTFExporter/Private/Converters/GLTFUVOverlapChecker.cpp b/Engine/Plugins/Enterprise/GLTFExporter/MigratedFromGit/Source/GLTFExporter/Private/Converters/GLTFUVOverlapChecker.cpp
--- a/Engine/Plugins/Enterprise/GLTFExporter/MigratedFromGit/Source/GLTFExporter/Private/Converters/GLTFUVOverlapChecker.cpp
+++ b/Engine/Plugins/Enterprise/GLTFExporter/MigratedFromGit/Source/GLTFExporter/Private/Converters/GLTFUVOverlapChecker.cpp
@@ -1,5 +1,7 @@
 // Copyright Epic Games, Inc. All Rights Reserved.
 
+#if WITH_EDITOR
+
 #include "Converters/GLTFUVOverlapChecker.h"
 #include "StaticMeshAttributes.h"
 #include "MeshDescription.h"
@@ -120,3 +122,5 @@ UMaterialInterface* FGLTFUVOverlapChecker::GetMaterial()
 	check(Material);
 	return Material;
 }
+
+#endif
diff --git a/Engine/Plugins/Enterprise/GLTFExporter/MigratedFromGit/Source/GLTFExporter/Private/Converters/GLTFUVOverlapChecker.h b/Engine/Plugins/Enterprise/GLTFExporter/MigratedFromGit/Source/GLTFExporter/Private/Converters/GLTFUVOverlapChecker.h
--- a/Engine/Plugins/Enterprise/GLTFExporter/MigratedFromGit/Source/GLTFExporter/Private/Converters/GLTFUVOverlapChecker.h
+++ b/Engine/Plugins/Enterprise/GLTFExporter/MigratedFromGit/Source/GLTFExporter/Private/Converters/GLTFUVOverlapChecker.h
@@ -2,6 +2,8 @@
 
 #pragma once
 
+#if WITH_EDITOR
+
 #include "Converters/GLTFConverter.h"
 #include "Converters/GLTFIndexArray.h"
 #include "Engine.h"
@@ -14,3 +16,5 @@ class FGLTFUVOverlapChecker final : public TGLTFConverter<float, const FMeshDesc
 
 	static UMaterialInterface* GetMaterial();
 };
+
+#endif
diff --git a/Engine/Plugins/Enterprise/GLTFExporter/MigratedFromGit/Source/GLTFExporter/Private/Exporters/GLTFExporter.cpp b/Engine/Plugins/Enterprise/GLTFExporter/MigratedFromGit/Source/GLTFExporter/Private/Exporters/GLTFExporter.cpp
--- a/Engine/Plugins/Enterprise/GLTFExporter/MigratedFromGit/Source/GLTFExporter/Private/Exporters/GLTFExporter.cpp
+++ b/Engine/Plugins/Enterprise/GLTFExporter/MigratedFromGit/Source/GLTFExporter/Private/Exporters/GLTFExporter.cpp
@@ -78,6 +78,7 @@ UGLTFExportOptions* UGLTFExporter::GetExportOptions()
 		Options = NewObject<UGLTFExportOptions>();
 	}
 
+#if WITH_EDITOR
 	if (GetShowExportOption() && !bAutomatedTask)
 	{
 		bool bExportAll = GetBatchMode();
@@ -95,6 +96,7 @@ UGLTFExportOptions* UGLTFExporter::GetExportOptions()
 		SetShowExportOption(!bExportAll);
 		Options->SaveConfig();
 	}
+#endif
 
 	return Options;
 }
diff --git a/Engine/Plugins/Enterprise/GLTFExporter/MigratedFromGit/Source/GLTFExporter/Private/Tasks/GLTFMaterialTasks.cpp b/Engine/Plugins/Enterprise/GLTFExporter/MigratedFromGit/Source/GLTFExporter/Private/Tasks/GLTFMaterialTasks.cpp
--- a/Engine/Plugins/Enterprise/GLTFExporter/MigratedFromGit/Source/GLTFExporter/Private/Tasks/GLTFMaterialTasks.cpp
+++ b/Engine/Plugins/Enterprise/GLTFExporter/MigratedFromGit/Source/GLTFExporter/Private/Tasks/GLTFMaterialTasks.cpp
@@ -1,5 +1,7 @@
 // Copyright Epic Games, Inc. All Rights Reserved.
 
+#if WITH_EDITOR
+
 #include "Tasks/GLTFMaterialTasks.h"
 #include "Converters/GLTFConverterUtility.h"
 #include "Converters/GLTFNameUtility.h"
@@ -1453,3 +1455,5 @@ EGLTFMaterialPropertyGroup FGLTFMaterialTask::GetPropertyGroup(const FMaterialPr
 			return EGLTFMaterialPropertyGroup::None;
 	}
 }
+
+#endif
diff --git a/Engine/Plugins/Enterprise/GLTFExporter/MigratedFromGit/Source/GLTFExporter/Private/Tasks/GLTFMaterialTasks.h b/Engine/Plugins/Enterprise/GLTFExporter/MigratedFromGit/Source/GLTFExporter/Private/Tasks/GLTFMaterialTasks.h
--- a/Engine/Plugins/Enterprise/GLTFExporter/MigratedFromGit/Source/GLTFExporter/Private/Tasks/GLTFMaterialTasks.h
+++ b/Engine/Plugins/Enterprise/GLTFExporter/MigratedFromGit/Source/GLTFExporter/Private/Tasks/GLTFMaterialTasks.h
@@ -2,6 +2,8 @@
 
 #pragma once
 
+#if WITH_EDITOR
+
 #include "Tasks/GLTFTask.h"
 #include "Builders/GLTFConvertBuilder.h"
 #include "Converters/GLTFUVOverlapChecker.h"
@@ -77,3 +79,5 @@ private:
 
 	static EGLTFMaterialPropertyGroup GetPropertyGroup(const FMaterialPropertyEx& Property);
 };
+
+#endif
diff --git a/Engine/Plugins/Enterprise/GLTFExporter/MigratedFromGit/Source/GLTFExporter/Private/Tasks/GLTFMeshTasks.cpp b/Engine/Plugins/Enterprise/GLTFExporter/MigratedFromGit/Source/GLTFExporter/Private/Tasks/GLTFMeshTasks.cpp
--- a/Engine/Plugins/Enterprise/GLTFExporter/MigratedFromGit/Source/GLTFExporter/Private/Tasks/GLTFMeshTasks.cpp
+++ b/Engine/Plugins/Enterprise/GLTFExporter/MigratedFromGit/Source/GLTFExporter/Private/Tasks/GLTFMeshTasks.cpp
@@ -115,6 +115,7 @@ void FGLTFStaticMeshTask::Complete()
 	const FGLTFMeshData* MeshData = Builder.ExportOptions->BakeMaterialInputs == EGLTFMaterialBakeMode::UseMeshData ?
 		Builder.StaticMeshDataConverter.GetOrAdd(StaticMesh, StaticMeshComponent, LODIndex) : nullptr;
 
+#if WITH_EDITOR
 	if (MeshData != nullptr)
 	{
 		if (MeshData->Description.IsEmpty())
@@ -128,6 +129,7 @@ void FGLTFStaticMeshTask::Complete()
 			MeshData = nullptr;
 		}
 	}
+#endif
 
 	ValidateVertexBuffer(Builder, VertexBuffer, *StaticMesh->GetName());
 
@@ -207,6 +209,7 @@ void FGLTFSkeletalMeshTask::Complete()
 	const FGLTFMeshData* MeshData = Builder.ExportOptions->BakeMaterialInputs == EGLTFMaterialBakeMode::UseMeshData ?
 		Builder.SkeletalMeshDataConverter.GetOrAdd(SkeletalMesh, SkeletalMeshComponent, LODIndex) : nullptr;
 
+#if WITH_EDITOR
 	if (MeshData != nullptr)
 	{
 		if (MeshData->Description.IsEmpty())
@@ -220,6 +223,7 @@ void FGLTFSkeletalMeshTask::Complete()
 			MeshData = nullptr;
 		}
 	}
+#endif
 
 	ValidateVertexBuffer(Builder, VertexBuffer, *SkeletalMesh->GetName());
 
diff --git a/Engine/Plugins/Enterprise/GLTFExporter/MigratedFromGit/Source/GLTFExporter/Private/Tasks/GLTFTextureTasks.cpp b/Engine/Plugins/Enterprise/GLTFExporter/MigratedFromGit/Source/GLTFExporter/Private/Tasks/GLTFTextureTasks.cpp
--- a/Engine/Plugins/Enterprise/GLTFExporter/MigratedFromGit/Source/GLTFExporter/Private/Tasks/GLTFTextureTasks.cpp
+++ b/Engine/Plugins/Enterprise/GLTFExporter/MigratedFromGit/Source/GLTFExporter/Private/Tasks/GLTFTextureTasks.cpp
@@ -152,6 +152,8 @@ void FGLTFTextureRenderTargetCubeTask::Complete()
 	JsonTexture.Sampler = Builder.GetOrAddSampler(RenderTargetCube);
 }
 
+#if WITH_EDITOR
+
 void FGLTFTextureLightMapTask::Complete()
 {
 	FGLTFJsonTexture& JsonTexture = Builder.GetTexture(TextureIndex);
@@ -186,3 +188,5 @@ void FGLTFTextureLightMapTask::Complete()
 
 	JsonTexture.Sampler = Builder.GetOrAddSampler(LightMap);
 }
+
+#endif
diff --git a/Engine/Plugins/Enterprise/GLTFExporter/MigratedFromGit/Source/GLTFExporter/Private/Tasks/GLTFTextureTasks.h b/Engine/Plugins/Enterprise/GLTFExporter/MigratedFromGit/Source/GLTFExporter/Private/Tasks/GLTFTextureTasks.h
--- a/Engine/Plugins/Enterprise/GLTFExporter/MigratedFromGit/Source/GLTFExporter/Private/Tasks/GLTFTextureTasks.h
+++ b/Engine/Plugins/Enterprise/GLTFExporter/MigratedFromGit/Source/GLTFExporter/Private/Tasks/GLTFTextureTasks.h
@@ -114,6 +114,8 @@ private:
 	const FGLTFJsonTextureIndex TextureIndex;
 };
 
+#if WITH_EDITOR
+
 class FGLTFTextureLightMapTask : public FGLTFTask
 {
 public:
@@ -139,3 +141,5 @@ private:
 	const ULightMapTexture2D* LightMap;
 	const FGLTFJsonTextureIndex TextureIndex;
 };
+
+#endif
diff --git a/Engine/Plugins/Enterprise/GLTFExporter/MigratedFromGit/Source/GLTFExporter/Private/UI/GLTFExportOptionsWindow.cpp b/Engine/Plugins/Enterprise/GLTFExporter/MigratedFromGit/Source/GLTFExporter/Private/UI/GLTFExportOptionsWindow.cpp
--- a/Engine/Plugins/Enterprise/GLTFExporter/MigratedFromGit/Source/GLTFExporter/Private/UI/GLTFExportOptionsWindow.cpp
+++ b/Engine/Plugins/Enterprise/GLTFExporter/MigratedFromGit/Source/GLTFExporter/Private/UI/GLTFExportOptionsWindow.cpp
@@ -1,5 +1,7 @@
 // Copyright 1998-2019 Epic Games, Inc. All Rights Reserved.
 
+#if WITH_EDITOR
+
 #include "UI/GLTFExportOptionsWindow.h"
 #include "Modules/ModuleManager.h"
 #include "Widgets/Layout/SBorder.h"
@@ -254,3 +256,5 @@ void SGLTFExportOptionsWindow::ShowDialog(UGLTFExportOptions* ExportOptions, con
 }
 
 #undef LOCTEXT_NAMESPACE
+
+#endif
diff --git a/Engine/Plugins/Enterprise/GLTFExporter/MigratedFromGit/Source/GLTFExporter/Private/UI/GLTFExportOptionsWindow.h b/Engine/Plugins/Enterprise/GLTFExporter/MigratedFromGit/Source/GLTFExporter/Private/UI/GLTFExportOptionsWindow.h
--- a/Engine/Plugins/Enterprise/GLTFExporter/MigratedFromGit/Source/GLTFExporter/Private/UI/GLTFExportOptionsWindow.h
+++ b/Engine/Plugins/Enterprise/GLTFExporter/MigratedFromGit/Source/GLTFExporter/Private/UI/GLTFExportOptionsWindow.h
@@ -2,6 +2,8 @@
 
 #pragma once
 
+#if WITH_EDITOR
+
 #include "GLTFExportOptions.h"
 #include "Widgets/DeclarativeSyntaxSupport.h"
 #include "Input/Reply.h"
@@ -56,3 +58,5 @@ private:
 	bool bShouldExport;
 	bool bShouldExportAll;
 };
+
+#endif
