[GLTFExporter] refactor: use IsBulkDataLoaded instead of GetBulkDataSize() != 0

commit 073132a9dfa1dff787c28927c03aaa87099cf876
Date: Wed, 27 Oct 2021 01:31:45 +0200

    refactor: use IsBulkDataLoaded instead of GetBulkDataSize() != 0
    (cherry picked from commit 61bdd14c7a03e56589c7b837f889a6656cc6b73d​)

#skipci

diff --git a/Engine/Plugins/Enterprise/GLTFExporter/MigratedFromGit/Source/GLTFExporter/Private/Converters/GLTFTextureUtility.cpp b/Engine/Plugins/Enterprise/GLTFExporter/MigratedFromGit/Source/GLTFExporter/Private/Converters/GLTFTextureUtility.cpp
--- a/Engine/Plugins/Enterprise/GLTFExporter/MigratedFromGit/Source/GLTFExporter/Private/Converters/GLTFTextureUtility.cpp
+++ b/Engine/Plugins/Enterprise/GLTFExporter/MigratedFromGit/Source/GLTFExporter/Private/Converters/GLTFTextureUtility.cpp
@@ -355,14 +355,14 @@ bool FGLTFTextureUtility::LoadPlatformData(UTexture2D* Texture)
 	}
 
 #if WITH_EDITOR
-	if (Texture->PlatformData->Mips[0].BulkData.GetBulkDataSize() == 0)
+	if (!Texture->PlatformData->Mips[0].BulkData.IsBulkDataLoaded())
 	{
 		// TODO: is this correct handling?
 		Texture->ForceRebuildPlatformData();
 	}
 #endif
 
-	return Texture->PlatformData->Mips[0].BulkData.GetBulkDataSize() != 0;
+	return Texture->PlatformData->Mips[0].BulkData.IsBulkDataLoaded();
 }
 
 bool FGLTFTextureUtility::LoadPlatformData(UTextureCube* TextureCube)
@@ -373,14 +373,14 @@ bool FGLTFTextureUtility::LoadPlatformData(UTextureCube* TextureCube)
 	}
 
 #if WITH_EDITOR
-	if (TextureCube->PlatformData->Mips[0].BulkData.GetBulkDataSize() == 0)
+	if (!TextureCube->PlatformData->Mips[0].BulkData.IsBulkDataLoaded())
 	{
 		// TODO: is this correct handling?
 		TextureCube->ForceRebuildPlatformData();
 	}
 #endif
 
-	return TextureCube->PlatformData->Mips[0].BulkData.GetBulkDataSize() != 0;
+	return TextureCube->PlatformData->Mips[0].BulkData.IsBulkDataLoaded();
 }
 
 void FGLTFTextureUtility::FlipGreenChannel(TArray<FColor>& Pixels)
