[GLTFExporter] fix: don't rely on thumbnail manager (which isn't available in runtime) to get default material preview mesh

commit 2758924a2f94a238097373c88eb91f2360f2a0b1
Date: Wed, 27 Oct 2021 02:01:49 +0200

    fix: don't rely on thumbnail manager (which isn't available in runtime) to get default material preview mesh
    (cherry picked from commit a25e77da28ec42e84047043e68a78c19b48a8e07​)

#skipci

diff --git a/Engine/Plugins/Enterprise/GLTFExporter/MigratedFromGit/Source/GLTFExporter/Private/Exporters/GLTFExporterUtility.cpp b/Engine/Plugins/Enterprise/GLTFExporter/MigratedFromGit/Source/GLTFExporter/Private/Exporters/GLTFExporterUtility.cpp
--- a/Engine/Plugins/Enterprise/GLTFExporter/MigratedFromGit/Source/GLTFExporter/Private/Exporters/GLTFExporterUtility.cpp
+++ b/Engine/Plugins/Enterprise/GLTFExporter/MigratedFromGit/Source/GLTFExporter/Private/Exporters/GLTFExporterUtility.cpp
@@ -7,8 +7,8 @@
 
 const UStaticMesh* FGLTFExporterUtility::GetPreviewMesh(const UMaterialInterface* Material)
 {
+#if WITH_EDITORONLY_DATA
 	// The following implementation is based of FMaterialInstanceEditor::RefreshPreviewAsset
-
 	const UStaticMesh* PreviewMesh = Cast<UStaticMesh>(Material->PreviewMesh.TryLoad());
 	if (PreviewMesh == nullptr)
 	{
@@ -24,13 +24,13 @@ const UStaticMesh* FGLTFExporterUtility::GetPreviewMesh(const UMaterialInterface
 		}
 	}
 
-	if (PreviewMesh == nullptr)
+	if (PreviewMesh != nullptr)
 	{
-		// Use a default sphere instead if the parent material's preview mesh is also invalid
-		PreviewMesh = GUnrealEd->GetThumbnailManager()->EditorSphere;
+		return PreviewMesh;
 	}
+#endif
 
-	return PreviewMesh;
+	return LoadObject<UStaticMesh>(nullptr, TEXT("/Engine/EditorMeshes/EditorSphere.EditorSphere"));
 }
 
 const USkeletalMesh* FGLTFExporterUtility::GetPreviewMesh(const UAnimSequence* AnimSequence)
