[GLTFExporter] fix: don't use FEditorFileUtils::IsMapPackageAsset

commit c22d1b9f3177980c7c14fbc8775ef5e026090ace
Date: Wed, 27 Oct 2021 02:13:04 +0200

    fix: don't use FEditorFileUtils::IsMapPackageAsset
    (cherry picked from commit 4ac472666a6a0cb7b101b5717a73936355ef1459​)

#skipci

diff --git a/Engine/Plugins/Enterprise/GLTFExporter/MigratedFromGit/Source/GLTFExporter/Private/Exporters/GLTFExporterUtility.cpp b/Engine/Plugins/Enterprise/GLTFExporter/MigratedFromGit/Source/GLTFExporter/Private/Exporters/GLTFExporterUtility.cpp
--- a/Engine/Plugins/Enterprise/GLTFExporter/MigratedFromGit/Source/GLTFExporter/Private/Exporters/GLTFExporterUtility.cpp
+++ b/Engine/Plugins/Enterprise/GLTFExporter/MigratedFromGit/Source/GLTFExporter/Private/Exporters/GLTFExporterUtility.cpp
@@ -3,7 +3,6 @@
 #include "Exporters/GLTFExporterUtility.h"
 #include "AssetRegistryModule.h"
 #include "LevelSequence.h"
-#include "UnrealEd.h"
 
 const UStaticMesh* FGLTFExporterUtility::GetPreviewMesh(const UMaterialInterface* Material)
 {
@@ -50,7 +49,7 @@ const USkeletalMesh* FGLTFExporterUtility::GetPreviewMesh(const UAnimSequence* A
 
 TArray<UWorld*> FGLTFExporterUtility::GetReferencedWorlds(const UObject* Object)
 {
-	FAssetRegistryModule& AssetRegistryModule = FModuleManager::LoadModuleChecked<FAssetRegistryModule>("AssetRegistry");
+	const FAssetRegistryModule& AssetRegistryModule = FModuleManager::LoadModuleChecked<FAssetRegistryModule>("AssetRegistry");
 	const FName OuterPathName = *Object->GetOutermost()->GetPathName();
 
 	TArray<UWorld*> ReferencedWorlds;
@@ -60,13 +59,10 @@ TArray<UWorld*> FGLTFExporterUtility::GetReferencedWorlds(const UObject* Object)
 	for (FAssetIdentifier& AssetDependency : AssetDependencies)
 	{
 		FString PackageName = AssetDependency.PackageName.ToString();
-		if (FEditorFileUtils::IsMapPackageAsset(PackageName))
+		UWorld* World = LoadObject<UWorld>(nullptr, *PackageName, nullptr, LOAD_NoWarn);
+		if (World != nullptr)
 		{
-			UWorld* World = LoadObject<UWorld>(nullptr, *PackageName);
-			if (World != nullptr)
-			{
-				ReferencedWorlds.AddUnique(World);
-			}
+			ReferencedWorlds.AddUnique(World);
 		}
 	}
 
