[GLTFExporter] fix: change GLTFExporter module to runtime and only add editor modules when building for editor

commit cc583e4b1b6e258321acdaee6afd5f47a4c42c7a
Date: Wed, 27 Oct 2021 02:19:12 +0200

    fix: change GLTFExporter module to runtime and only add editor modules when building for editor
    (cherry picked from commit 9bca2a3ae128d292448279015afb84d86addfa15​)

#skipci

diff --git a/Engine/Plugins/Enterprise/GLTFExporter/MigratedFromGit/GLTFExporter.uplugin b/Engine/Plugins/Enterprise/GLTFExporter/MigratedFromGit/GLTFExporter.uplugin
--- a/Engine/Plugins/Enterprise/GLTFExporter/MigratedFromGit/GLTFExporter.uplugin
+++ b/Engine/Plugins/Enterprise/GLTFExporter/MigratedFromGit/GLTFExporter.uplugin
@@ -16,7 +16,7 @@
 	"Modules": [
 		{
 			"Name": "GLTFExporter",
-			"Type": "Editor",
+			"Type": "Runtime",
 			"LoadingPhase": "Default",
 			"WhitelistPlatforms": [ "Win64", "Win32", "Mac" ]
 		},
diff --git a/Engine/Plugins/Enterprise/GLTFExporter/MigratedFromGit/Source/GLTFExporter/GLTFExporter.Build.cs b/Engine/Plugins/Enterprise/GLTFExporter/MigratedFromGit/Source/GLTFExporter/GLTFExporter.Build.cs
--- a/Engine/Plugins/Enterprise/GLTFExporter/MigratedFromGit/Source/GLTFExporter/GLTFExporter.Build.cs
+++ b/Engine/Plugins/Enterprise/GLTFExporter/MigratedFromGit/Source/GLTFExporter/GLTFExporter.Build.cs
@@ -23,17 +23,7 @@ namespace UnrealBuildTool.Rules
 					"Core",
 					"CoreUObject",
 					"Engine",
-					"UnrealEd",
-					"MeshDescription",
-					"StaticMeshDescription",
-					"MessageLog",
 					"Json",
-					"Slate",
-					"SlateCore",
-					"Mainframe",
-					"InputCore",
-					"EditorStyle",
-					"Projects",
 					"RenderCore",
 					"RHI",
 					"DesktopPlatform",
@@ -41,12 +31,31 @@ namespace UnrealBuildTool.Rules
 					"MovieScene",
 					"MovieSceneTracks",
 					"VariantManagerContent",
-					"MeshMergeUtilities",
-					"GLTFMaterialBaking",
-					"GLTFMaterialAnalyzer"
 				}
 				);
 
+			if (Target.bBuildEditor)
+			{
+				PrivateDependencyModuleNames.AddRange(
+					new string[]
+					{
+						"UnrealEd",
+						"MessageLog",
+						"Slate",
+						"SlateCore",
+						"Mainframe",
+						"InputCore",
+						"EditorStyle",
+						"Projects",
+						"MeshMergeUtilities",
+						"MeshDescription",
+						"StaticMeshDescription",
+						"GLTFMaterialBaking",
+						"GLTFMaterialAnalyzer"
+					}
+					);
+			}
+
 			PublicDependencyModuleNames.AddRange(
 				new string[]
 				{
