[GLTFExporter] fix: copy FNormalMapBatchedElementParameters into plugin (and rename to FGLTFNormalMapPreview) to allow use in runtime

commit a9e3d8694838bfaea2163f40a278c9d83599bebf
Date: Wed, 27 Oct 2021 02:28:24 +0200

    fix: copy FNormalMapBatchedElementParameters into plugin (and rename to FGLTFNormalMapPreview) to allow use in runtime
    (cherry picked from commit a4aa65234a91aab29a98a8e78ec7d8aea24576ab​)

#skipci

diff --git a/Engine/Plugins/Enterprise/GLTFExporter/MigratedFromGit/GLTFExporter.uplugin b/Engine/Plugins/Enterprise/GLTFExporter/MigratedFromGit/GLTFExporter.uplugin
--- a/Engine/Plugins/Enterprise/GLTFExporter/MigratedFromGit/GLTFExporter.uplugin
+++ b/Engine/Plugins/Enterprise/GLTFExporter/MigratedFromGit/GLTFExporter.uplugin
@@ -17,7 +17,7 @@
 		{
 			"Name": "GLTFExporter",
 			"Type": "Runtime",
-			"LoadingPhase": "Default",
+			"LoadingPhase": "PostConfigInit",
 			"WhitelistPlatforms": [ "Win64", "Win32", "Mac" ]
 		},
 		{
diff --git a/Engine/Plugins/Enterprise/GLTFExporter/MigratedFromGit/Source/GLTFExporter/Private/Converters/GLTFNormalMapPreview.cpp b/Engine/Plugins/Enterprise/GLTFExporter/MigratedFromGit/Source/GLTFExporter/Private/Converters/GLTFNormalMapPreview.cpp
new file mode 100644
--- /dev/null
+++ b/Engine/Plugins/Enterprise/GLTFExporter/MigratedFromGit/Source/GLTFExporter/Private/Converters/GLTFNormalMapPreview.cpp
@@ -0,0 +1,88 @@
+// Copyright Epic Games, Inc. All Rights Reserved.
+
+/*==============================================================================
+	NormalMapPreview.h: Implementation for previewing normal maps.
+==============================================================================*/
+
+#include "GLTFNormalMapPreview.h"
+#include "Shader.h"
+#include "GlobalShader.h"
+#include "SimpleElementShaders.h"
+#include "ShaderParameterUtils.h"
+#include "PipelineStateCache.h"
+
+/*------------------------------------------------------------------------------
+	Batched element shaders for previewing normal maps.
+------------------------------------------------------------------------------*/
+
+/**
+ * Simple pixel shader that reconstructs a normal for the purposes of visualization.
+ */
+class FGLTFSimpleElementNormalMapPS : public FGlobalShader
+{
+	DECLARE_SHADER_TYPE(FGLTFSimpleElementNormalMapPS,Global);
+public:
+
+	/** Should the shader be cached? Always. */
+	static bool ShouldCompilePermutation(const FGlobalShaderPermutationParameters& Parameters)
+	{
+		return IsPCPlatform(Parameters.Platform);
+	}
+
+	/** Default constructor. */
+	FGLTFSimpleElementNormalMapPS() {}
+
+	/**
+	 * Initialization constructor.
+	 * @param Initializer - Shader initialization container.
+	 */
+	FGLTFSimpleElementNormalMapPS(const ShaderMetaType::CompiledShaderInitializerType& Initializer)
+		: FGlobalShader(Initializer)
+	{
+		Texture.Bind(Initializer.ParameterMap,TEXT("NormalMapTexture"));
+		TextureSampler.Bind(Initializer.ParameterMap,TEXT("NormalMapTextureSampler"));
+	}
+
+	/**
+	 * Set shader parameters.
+	 * @param NormalMapTexture - The normal map texture to sample.
+	 */
+	void SetParameters(FRHICommandList& RHICmdList, const FTexture* NormalMapTexture)
+	{
+		FRHIPixelShader* PixelShaderRHI = RHICmdList.GetBoundPixelShader();
+		SetTextureParameter(RHICmdList, PixelShaderRHI,Texture,TextureSampler,NormalMapTexture);
+	}
+
+private:
+	/** The texture to sample. */
+	LAYOUT_FIELD(FShaderResourceParameter, Texture);
+	LAYOUT_FIELD(FShaderResourceParameter, TextureSampler);
+};
+IMPLEMENT_SHADER_TYPE(,FGLTFSimpleElementNormalMapPS,TEXT("/Engine/Private/SimpleElementNormalMapPixelShader.usf"),TEXT("Main"),SF_Pixel);
+
+/** Binds vertex and pixel shaders for this element */
+void FGLTFNormalMapPreview::BindShaders(
+	FRHICommandList& RHICmdList,
+	FGraphicsPipelineStateInitializer& GraphicsPSOInit,
+	ERHIFeatureLevel::Type InFeatureLevel,
+	const FMatrix& InTransform,
+	const float InGamma,
+	const FMatrix& ColorWeights,
+	const FTexture* Texture)
+{
+	TShaderMapRef<FSimpleElementVS> VertexShader(GetGlobalShaderMap(InFeatureLevel));
+	TShaderMapRef<FGLTFSimpleElementNormalMapPS> PixelShader(GetGlobalShaderMap(InFeatureLevel));
+
+	GraphicsPSOInit.BoundShaderState.VertexDeclarationRHI = GSimpleElementVertexDeclaration.VertexDeclarationRHI;
+	GraphicsPSOInit.BoundShaderState.VertexShaderRHI = VertexShader.GetVertexShader();
+	GraphicsPSOInit.BoundShaderState.PixelShaderRHI = PixelShader.GetPixelShader();
+	GraphicsPSOInit.PrimitiveType = PT_TriangleList;
+
+	GraphicsPSOInit.BlendState = TStaticBlendState<>::GetRHI();
+
+	SetGraphicsPipelineState(RHICmdList, GraphicsPSOInit, EApplyRendertargetOption::ForceApply);
+
+	VertexShader->SetParameters(RHICmdList, InTransform);
+	PixelShader->SetParameters(RHICmdList, Texture);
+}
+
diff --git a/Engine/Plugins/Enterprise/GLTFExporter/MigratedFromGit/Source/GLTFExporter/Private/Converters/GLTFNormalMapPreview.h b/Engine/Plugins/Enterprise/GLTFExporter/MigratedFromGit/Source/GLTFExporter/Private/Converters/GLTFNormalMapPreview.h
new file mode 100644
--- /dev/null
+++ b/Engine/Plugins/Enterprise/GLTFExporter/MigratedFromGit/Source/GLTFExporter/Private/Converters/GLTFNormalMapPreview.h
@@ -0,0 +1,26 @@
+// Copyright Epic Games, Inc. All Rights Reserved.
+
+/*==============================================================================
+	NormalMapPreview.h: Definitions for previewing normal maps.
+==============================================================================*/
+
+#pragma once
+
+#include "CoreMinimal.h"
+#include "BatchedElements.h"
+
+/**
+ * Batched element parameters for previewing normal maps.
+ */
+class FGLTFNormalMapPreview : public FBatchedElementParameters
+{
+	/** Binds vertex and pixel shaders for this element */
+	virtual void BindShaders(
+		FRHICommandList& RHICmdList,
+		FGraphicsPipelineStateInitializer& GraphicsPSOInit,
+		ERHIFeatureLevel::Type InFeatureLevel,
+		const FMatrix& InTransform,
+		const float InGamma,
+		const FMatrix& ColorWeights,
+		const FTexture* Texture) override;
+};
diff --git a/Engine/Plugins/Enterprise/GLTFExporter/MigratedFromGit/Source/GLTFExporter/Private/Converters/GLTFTextureUtility.cpp b/Engine/Plugins/Enterprise/GLTFExporter/MigratedFromGit/Source/GLTFExporter/Private/Converters/GLTFTextureUtility.cpp
--- a/Engine/Plugins/Enterprise/GLTFExporter/MigratedFromGit/Source/GLTFExporter/Private/Converters/GLTFTextureUtility.cpp
+++ b/Engine/Plugins/Enterprise/GLTFExporter/MigratedFromGit/Source/GLTFExporter/Private/Converters/GLTFTextureUtility.cpp
@@ -1,7 +1,7 @@
 // Copyright Epic Games, Inc. All Rights Reserved.
 
 #include "Converters/GLTFTextureUtility.h"
-#include "NormalMapPreview.h"
+#include "GLTFNormalMapPreview.h"
 
 bool FGLTFTextureUtility::IsHDR(EPixelFormat Format)
 {
@@ -141,6 +141,7 @@ UTexture2D* FGLTFTextureUtility::CreateTransientTexture(const void* RawData, int
 	Texture->PlatformData->Mips[0].BulkData.Unlock();
 
 	Texture->SRGB = bSRGB ? 1 : 0;
+	Texture->CompressionSettings = TC_VectorDisplacementmap; // best quality
 #if WITH_EDITOR
 	Texture->CompressionNone = true;
 	Texture->MipGenSettings = TMGS_NoMipmaps;
@@ -183,7 +184,7 @@ bool FGLTFTextureUtility::DrawTexture(UTextureRenderTarget2D* OutTarget, const U
 
 	if (InSource->IsNormalMap())
 	{
-		BatchedElementParameters = new FNormalMapBatchedElementParameters();
+		BatchedElementParameters = new FGLTFNormalMapPreview();
 	}
 
 	FCanvas Canvas(RenderTarget, nullptr, 0.0f, 0.0f, 0.0f, GMaxRHIFeatureLevel);
