[GLTFExporter] fix: fix compilation issues on Win32 Development build configuration

commit 2f3c99053c7ade091f59dff7bd12442925827c42
Date: Wed, 27 Oct 2021 02:41:44 +0200

    fix: fix compilation issues on Win32 Development build configuration
    (cherry picked from commit 0b2aa06912875b9f4b5627bf416649859243621f​)

#skipci

diff --git a/Engine/Plugins/Enterprise/GLTFExporter/MigratedFromGit/Source/GLTFExporter/Private/Json/GLTFJsonColor.h b/Engine/Plugins/Enterprise/GLTFExporter/MigratedFromGit/Source/GLTFExporter/Private/Json/GLTFJsonColor.h
--- a/Engine/Plugins/Enterprise/GLTFExporter/MigratedFromGit/Source/GLTFExporter/Private/Json/GLTFJsonColor.h
+++ b/Engine/Plugins/Enterprise/GLTFExporter/MigratedFromGit/Source/GLTFExporter/Private/Json/GLTFJsonColor.h
@@ -21,7 +21,7 @@ struct TGLTFJsonColor : BaseType, IGLTFJsonArray
 
 	virtual void WriteArray(IGLTFJsonWriter& Writer) const override
 	{
-		for (int32 i = 0; i < GetNum(BaseType::Components); ++i)
+		for (SIZE_T i = 0; i < GetNum(BaseType::Components); ++i)
 		{
 			Writer.Write(BaseType::Components[i]);
 		}
@@ -29,7 +29,7 @@ struct TGLTFJsonColor : BaseType, IGLTFJsonArray
 
 	bool IsNearlyEqual(const BaseType& Other, float Tolerance = KINDA_SMALL_NUMBER) const
 	{
-		for (int32 i = 0; i < GetNum(BaseType::Components); ++i)
+		for (SIZE_T i = 0; i < GetNum(BaseType::Components); ++i)
 		{
 			if (!FMath::IsNearlyEqual(BaseType::Components[i], Other.Components[i], Tolerance))
 			{
diff --git a/Engine/Plugins/Enterprise/GLTFExporter/MigratedFromGit/Source/GLTFExporter/Private/Json/GLTFJsonMatrix.h b/Engine/Plugins/Enterprise/GLTFExporter/MigratedFromGit/Source/GLTFExporter/Private/Json/GLTFJsonMatrix.h
--- a/Engine/Plugins/Enterprise/GLTFExporter/MigratedFromGit/Source/GLTFExporter/Private/Json/GLTFJsonMatrix.h
+++ b/Engine/Plugins/Enterprise/GLTFExporter/MigratedFromGit/Source/GLTFExporter/Private/Json/GLTFJsonMatrix.h
@@ -23,7 +23,7 @@ struct TGLTFJsonMatrix : BaseType, IGLTFJsonArray
 
 	virtual void WriteArray(IGLTFJsonWriter& Writer) const override
 	{
-		for (int32 i = 0; i < GetNum(BaseType::Elements); ++i)
+		for (SIZE_T i = 0; i < GetNum(BaseType::Elements); ++i)
 		{
 			Writer.Write(BaseType::Elements[i]);
 		}
@@ -31,7 +31,7 @@ struct TGLTFJsonMatrix : BaseType, IGLTFJsonArray
 
 	bool IsNearlyEqual(const BaseType& Other, float Tolerance = KINDA_SMALL_NUMBER) const
 	{
-		for (int32 i = 0; i < GetNum(BaseType::Elements); ++i)
+		for (SIZE_T i = 0; i < GetNum(BaseType::Elements); ++i)
 		{
 			if (!FMath::IsNearlyEqual(BaseType::Elements[i], Other.Elements[i], Tolerance))
 			{
diff --git a/Engine/Plugins/Enterprise/GLTFExporter/MigratedFromGit/Source/GLTFExporter/Private/Json/GLTFJsonVector.h b/Engine/Plugins/Enterprise/GLTFExporter/MigratedFromGit/Source/GLTFExporter/Private/Json/GLTFJsonVector.h
--- a/Engine/Plugins/Enterprise/GLTFExporter/MigratedFromGit/Source/GLTFExporter/Private/Json/GLTFJsonVector.h
+++ b/Engine/Plugins/Enterprise/GLTFExporter/MigratedFromGit/Source/GLTFExporter/Private/Json/GLTFJsonVector.h
@@ -24,7 +24,7 @@ struct TGLTFJsonVector : BaseType, IGLTFJsonArray
 
 	virtual void WriteArray(IGLTFJsonWriter& Writer) const override
 	{
-		for (int32 i = 0; i < GetNum(BaseType::Components); ++i)
+		for (SIZE_T i = 0; i < GetNum(BaseType::Components); ++i)
 		{
 			Writer.Write(BaseType::Components[i]);
 		}
@@ -32,7 +32,7 @@ struct TGLTFJsonVector : BaseType, IGLTFJsonArray
 
 	bool IsNearlyEqual(const BaseType& Other, float Tolerance = KINDA_SMALL_NUMBER) const
 	{
-		for (int32 i = 0; i < GetNum(BaseType::Components); ++i)
+		for (SIZE_T i = 0; i < GetNum(BaseType::Components); ++i)
 		{
 			if (!FMath::IsNearlyEqual(BaseType::Components[i], Other.Components[i], Tolerance))
 			{
