[GLTFExporter] fix: remove module dependency DesktopPlatform preventing Shipping builds

commit e0774d4391e59fe994dc8a56439546b8db38c96e
Date: Wed, 27 Oct 2021 02:58:12 +0200

    fix: remove module dependency DesktopPlatform preventing Shipping builds
    (cherry picked from commit f93fc33d001bf610cfcf046256d1e50e47164e1e​)

#skipci

diff --git a/Engine/Plugins/Enterprise/GLTFExporter/MigratedFromGit/Source/GLTFExporter/GLTFExporter.Build.cs b/Engine/Plugins/Enterprise/GLTFExporter/MigratedFromGit/Source/GLTFExporter/GLTFExporter.Build.cs
--- a/Engine/Plugins/Enterprise/GLTFExporter/MigratedFromGit/Source/GLTFExporter/GLTFExporter.Build.cs
+++ b/Engine/Plugins/Enterprise/GLTFExporter/MigratedFromGit/Source/GLTFExporter/GLTFExporter.Build.cs
@@ -26,7 +26,6 @@ namespace UnrealBuildTool.Rules
 					"Json",
 					"RenderCore",
 					"RHI",
-					"DesktopPlatform",
 					"LevelSequence",
 					"MovieScene",
 					"MovieSceneTracks",
