[GLTFExporter] refactor: reduce code duplication in GLTFMeshSection to ease runtime support

commit a2f7d2c3b0110c6cd574d0c113b64301bb1e8a5f
Date: Wed, 10 Nov 2021 20:01:31 +0100

    refactor: reduce code duplication in GLTFMeshSection to ease runtime support
    (cherry picked from commit ecdc44f8833e90ac76ec18e8a02f932ec089b017​)

#skipci

diff --git a/Engine/Plugins/Enterprise/GLTFExporter/MigratedFromGit/Source/GLTFExporter/Private/Converters/GLTFMeshSection.cpp b/Engine/Plugins/Enterprise/GLTFExporter/MigratedFromGit/Source/GLTFExporter/Private/Converters/GLTFMeshSection.cpp
--- a/Engine/Plugins/Enterprise/GLTFExporter/MigratedFromGit/Source/GLTFExporter/Private/Converters/GLTFMeshSection.cpp
+++ b/Engine/Plugins/Enterprise/GLTFExporter/MigratedFromGit/Source/GLTFExporter/Private/Converters/GLTFMeshSection.cpp
@@ -7,72 +7,37 @@
 
 FGLTFMeshSection::FGLTFMeshSection(const FStaticMeshLODResources* MeshLOD, const FGLTFIndexArray& SectionIndices)
 {
-	const FRawStaticIndexBuffer& OldIndexBuffer = MeshLOD->IndexBuffer;
-	if (OldIndexBuffer.AccessStream16() == nullptr && OldIndexBuffer.AccessStream32() == nullptr)
+	const FRawStaticIndexBuffer& SourceBuffer = MeshLOD->IndexBuffer;
+	if (SourceBuffer.Is32Bit())
 	{
-		// TODO: report error
-		return;
+		Init(MeshLOD->Sections, SectionIndices, SourceBuffer.AccessStream32());
 	}
-
-	const FStaticMeshLODResources::FStaticMeshSectionArray& Sections = MeshLOD->Sections;
-
-	uint32 TotalIndexCount = 0;
-	for (int32 SectionIndex : SectionIndices)
+	else
 	{
-		TotalIndexCount += Sections[SectionIndex].NumTriangles * 3;
+		Init(MeshLOD->Sections, SectionIndices, SourceBuffer.AccessStream16());
 	}
-
-	IndexMap.Reserve(TotalIndexCount);
-	IndexBuffer.AddUninitialized(TotalIndexCount);
-	BoneMapLookup.Reserve(TotalIndexCount);
-
-	TMap<uint32, uint32> IndexLookup;
-
-	for (int32 SectionIndex : SectionIndices)
-	{
-		const FStaticMeshSection& MeshSection = Sections[SectionIndex];
-		const uint32 IndexOffset = MeshSection.FirstIndex;
-		const uint32 IndexCount = MeshSection.NumTriangles * 3;
-
-		for (uint32 Index = 0; Index < IndexCount; Index++)
-		{
-			const uint32 OldIndex = OldIndexBuffer.GetIndex(IndexOffset + Index);
-			uint32 NewIndex;
-
-			if (uint32* FoundIndex = IndexLookup.Find(OldIndex))
-			{
-				NewIndex = *FoundIndex;
-			}
-			else
-			{
-				NewIndex = IndexMap.Num();
-				IndexLookup.Add(OldIndex, NewIndex);
-				IndexMap.Add(OldIndex);
-				BoneMapLookup.Add(0);
-			}
-
-			IndexBuffer[Index] = NewIndex;
-		}
-	}
-
-	BoneMaps.Add({});
-	MaxBoneIndex = 0;
 }
 
 FGLTFMeshSection::FGLTFMeshSection(const FSkeletalMeshLODRenderData* MeshLOD, const FGLTFIndexArray& SectionIndices)
-	: MaxBoneIndex(0)
 {
-	const FRawStaticIndexBuffer16or32Interface* OldIndexBuffer = MeshLOD->MultiSizeIndexContainer.GetIndexBuffer();
-	if (const_cast<FRawStaticIndexBuffer16or32Interface*>(OldIndexBuffer)->GetPointerTo(0) == nullptr)
+	const FMultiSizeIndexContainer& SourceBuffer = MeshLOD->MultiSizeIndexContainer;
+	const void* SourceData = const_cast<FRawStaticIndexBuffer16or32Interface*>(SourceBuffer.GetIndexBuffer())->GetPointerTo(0);
+
+	if (SourceBuffer.GetDataTypeSize() != sizeof(uint16))
 	{
-		// TODO: report error
-		return;
+		Init(MeshLOD->RenderSections, SectionIndices, static_cast<const uint32*>(SourceData));
 	}
+	else
+	{
+		Init(MeshLOD->RenderSections, SectionIndices, static_cast<const uint16*>(SourceData));
+	}
+}
 
-	const TArray<FSkelMeshRenderSection>& Sections = MeshLOD->RenderSections;
-
+template <typename IndexType, typename SectionArrayType>
+void FGLTFMeshSection::Init(const SectionArrayType& Sections, const FGLTFIndexArray& SectionIndices, const IndexType* SourceData)
+{
 	uint32 TotalIndexCount = 0;
-	for (int32 SectionIndex : SectionIndices)
+	for (const int32 SectionIndex : SectionIndices)
 	{
 		TotalIndexCount += Sections[SectionIndex].NumTriangles * 3;
 	}
@@ -80,22 +45,25 @@ FGLTFMeshSection::FGLTFMeshSection(const FSkeletalMeshLODRenderData* MeshLOD, co
 	IndexMap.Reserve(TotalIndexCount);
 	IndexBuffer.AddUninitialized(TotalIndexCount);
 	BoneMapLookup.Reserve(TotalIndexCount);
+	MaxBoneIndex = 0;
 
 	TMap<uint32, uint32> IndexLookup;
 
-	for (int32 SectionIndex : SectionIndices)
+	for (const int32 SectionIndex : SectionIndices)
 	{
-		const FSkelMeshRenderSection& MeshSection = Sections[SectionIndex];
-		const uint32 IndexOffset = MeshSection.BaseIndex;
-		const uint32 IndexCount = MeshSection.NumTriangles * 3;
+		const typename SectionArrayType::ElementType& Section = Sections[SectionIndex];
+		const TArray<FBoneIndexType>& BoneMap = GetBoneMap(Section);
+
+		const uint32 IndexOffset = GetIndexOffset(Section);
+		const uint32 IndexCount = Section.NumTriangles * 3;
 		const uint32 BoneMapIndex = BoneMaps.Num();
 
 		for (uint32 Index = 0; Index < IndexCount; Index++)
 		{
-			const uint32 OldIndex = OldIndexBuffer->Get(IndexOffset + Index);
+			const uint32 OldIndex = SourceData[IndexOffset + Index];
 			uint32 NewIndex;
 
-			if (uint32* FoundIndex = IndexLookup.Find(OldIndex))
+			if (const uint32* FoundIndex = IndexLookup.Find(OldIndex))
 			{
 				NewIndex = *FoundIndex;
 			}
@@ -104,17 +72,45 @@ FGLTFMeshSection::FGLTFMeshSection(const FSkeletalMeshLODRenderData* MeshLOD, co
 				NewIndex = IndexMap.Num();
 				IndexLookup.Add(OldIndex, NewIndex);
 				IndexMap.Add(OldIndex);
-				BoneMapLookup.Add(BoneMapIndex);
+
+				if (BoneMap.Num() > 0)
+				{
+					BoneMapLookup.Add(BoneMapIndex);
+				}
 			}
 
 			IndexBuffer[Index] = NewIndex;
 		}
 
-		BoneMaps.Add(MeshSection.BoneMap);
-
-		if (const FBoneIndexType* MaxSectionBoneIndex = Algo::MaxElement(MeshSection.BoneMap))
+		if (BoneMap.Num() > 0)
 		{
-			MaxBoneIndex = FMath::Max(*MaxSectionBoneIndex, MaxBoneIndex);
+			BoneMaps.Add(BoneMap);
+
+			if (const FBoneIndexType* MaxSectionBoneIndex = Algo::MaxElement(BoneMap))
+			{
+				MaxBoneIndex = FMath::Max(*MaxSectionBoneIndex, MaxBoneIndex);
+			}
 		}
 	}
 }
+
+uint32 FGLTFMeshSection::GetIndexOffset(const FStaticMeshSection& Section)
+{
+	return Section.FirstIndex;
+}
+
+uint32 FGLTFMeshSection::GetIndexOffset(const FSkelMeshRenderSection& Section)
+{
+	return Section.BaseIndex;
+}
+
+const TArray<FBoneIndexType>& FGLTFMeshSection::GetBoneMap(const FStaticMeshSection& Section)
+{
+	static const TArray<FBoneIndexType> BoneMap = {};
+	return BoneMap;
+}
+
+const TArray<FBoneIndexType>& FGLTFMeshSection::GetBoneMap(const FSkelMeshRenderSection& Section)
+{
+	return Section.BoneMap;
+}
diff --git a/Engine/Plugins/Enterprise/GLTFExporter/MigratedFromGit/Source/GLTFExporter/Private/Converters/GLTFMeshSection.h b/Engine/Plugins/Enterprise/GLTFExporter/MigratedFromGit/Source/GLTFExporter/Private/Converters/GLTFMeshSection.h
--- a/Engine/Plugins/Enterprise/GLTFExporter/MigratedFromGit/Source/GLTFExporter/Private/Converters/GLTFMeshSection.h
+++ b/Engine/Plugins/Enterprise/GLTFExporter/MigratedFromGit/Source/GLTFExporter/Private/Converters/GLTFMeshSection.h
@@ -16,4 +16,15 @@ struct FGLTFMeshSection
 	TArray<TArray<FBoneIndexType>> BoneMaps;
 	TArray<uint32> BoneMapLookup;
 	FBoneIndexType MaxBoneIndex;
+
+private:
+
+	template <typename IndexType, typename SectionArrayType>
+	void Init(const SectionArrayType& Sections, const FGLTFIndexArray& SectionIndices, const IndexType* SourceData);
+
+	static uint32 GetIndexOffset(const FStaticMeshSection& Section);
+	static uint32 GetIndexOffset(const FSkelMeshRenderSection& Section);
+
+	static const TArray<FBoneIndexType>& GetBoneMap(const FStaticMeshSection& Section);
+	static const TArray<FBoneIndexType>& GetBoneMap(const FSkelMeshRenderSection& Section);
 };
