[GLTFExporter] refactor: add support for creating GLTFBuilder without export options

commit 855e70ef3aca10fe69d74bc0aa2c9c9545039875
Date: Wed, 10 Nov 2021 20:02:16 +0100

    refactor: add support for creating GLTFBuilder without export options
    (cherry picked from commit 1882b2001d24bab60e5408fe6ee654cd6a779c88​)

#skipci

diff --git a/Engine/Plugins/Enterprise/GLTFExporter/MigratedFromGit/Source/GLTFExporter/Private/Builders/GLTFBuilder.cpp b/Engine/Plugins/Enterprise/GLTFExporter/MigratedFromGit/Source/GLTFExporter/Private/Builders/GLTFBuilder.cpp
--- a/Engine/Plugins/Enterprise/GLTFExporter/MigratedFromGit/Source/GLTFExporter/Private/Builders/GLTFBuilder.cpp
+++ b/Engine/Plugins/Enterprise/GLTFExporter/MigratedFromGit/Source/GLTFExporter/Private/Builders/GLTFBuilder.cpp
@@ -8,7 +8,7 @@ FGLTFBuilder::FGLTFBuilder(const FString& FilePath, const UGLTFExportOptions* Ex
 	: bIsGlbFile(FGLTFFileUtility::IsGlbFile(FilePath))
 	, FilePath(FilePath)
 	, DirPath(FPaths::GetPath(FilePath))
-	, ExportOptions(ValidateExportOptions(ExportOptions))
+	, ExportOptions(SanitizeExportOptions(ExportOptions))
 	, ExportOptionsGuard(ExportOptions)
 {
 }
@@ -77,21 +77,28 @@ bool FGLTFBuilder::ShouldExportLight(EComponentMobility::Type LightMobility) con
 	return EnumHasAllFlags(AllowedMobility, QueriedMobility);
 }
 
-const UGLTFExportOptions* FGLTFBuilder::ValidateExportOptions(const UGLTFExportOptions* ExportOptions)
+const UGLTFExportOptions* FGLTFBuilder::SanitizeExportOptions(const UGLTFExportOptions* Options)
 {
+	if (Options == nullptr)
+	{
+		UGLTFExportOptions* NewOptions = NewObject<UGLTFExportOptions>();
+		NewOptions->ResetToDefault();
+		Options = NewOptions;
+	}
+
 	if (!FApp::CanEverRender())
 	{
-		if (ExportOptions->BakeMaterialInputs != EGLTFMaterialBakeMode::Disabled || ExportOptions->TextureImageFormat != EGLTFTextureImageFormat::None)
+		if (Options->BakeMaterialInputs != EGLTFMaterialBakeMode::Disabled || Options->TextureImageFormat != EGLTFTextureImageFormat::None)
 		{
 			// TODO: warn the following options requires rendering support and will be overriden
-			UGLTFExportOptions* OverridenOptions = DuplicateObject(ExportOptions, nullptr);
+			UGLTFExportOptions* OverridenOptions = DuplicateObject(Options, nullptr);
 			OverridenOptions->BakeMaterialInputs = EGLTFMaterialBakeMode::Disabled;
 			OverridenOptions->TextureImageFormat = EGLTFTextureImageFormat::None;
-			return OverridenOptions;
+			Options = OverridenOptions;
 		}
 	}
 
-	return ExportOptions;
+	return Options;
 }
 
 EGLTFSceneMobility FGLTFBuilder::GetSceneMobility(EComponentMobility::Type Mobility)
diff --git a/Engine/Plugins/Enterprise/GLTFExporter/MigratedFromGit/Source/GLTFExporter/Private/Builders/GLTFBuilder.h b/Engine/Plugins/Enterprise/GLTFExporter/MigratedFromGit/Source/GLTFExporter/Private/Builders/GLTFBuilder.h
--- a/Engine/Plugins/Enterprise/GLTFExporter/MigratedFromGit/Source/GLTFExporter/Private/Builders/GLTFBuilder.h
+++ b/Engine/Plugins/Enterprise/GLTFExporter/MigratedFromGit/Source/GLTFExporter/Private/Builders/GLTFBuilder.h
@@ -33,6 +33,6 @@ private:
 
 	FGCObjectScopeGuard ExportOptionsGuard;
 
-	static const UGLTFExportOptions* ValidateExportOptions(const UGLTFExportOptions* ExportOptions);
+	static const UGLTFExportOptions* SanitizeExportOptions(const UGLTFExportOptions* Options);
 	static EGLTFSceneMobility GetSceneMobility(EComponentMobility::Type Mobility);
 };
