[GLTFExporter] refactor: add support for calling GLTFBuilder::CompleteAllTasks without context

commit b079b2aa7e9de65434c19c2dcaf49ae27dc17a72
Date: Wed, 10 Nov 2021 20:07:42 +0100

    refactor: add support for calling GLTFBuilder::CompleteAllTasks without context
    (cherry picked from commit 0ee0d4acff74e6217397a22e9f7126408c2c8899​)

#skipci

diff --git a/Engine/Plugins/Enterprise/GLTFExporter/MigratedFromGit/Source/GLTFExporter/Private/Builders/GLTFTaskBuilder.cpp b/Engine/Plugins/Enterprise/GLTFExporter/MigratedFromGit/Source/GLTFExporter/Private/Builders/GLTFTaskBuilder.cpp
--- a/Engine/Plugins/Enterprise/GLTFExporter/MigratedFromGit/Source/GLTFExporter/Private/Builders/GLTFTaskBuilder.cpp
+++ b/Engine/Plugins/Enterprise/GLTFExporter/MigratedFromGit/Source/GLTFExporter/Private/Builders/GLTFTaskBuilder.cpp
@@ -24,17 +24,27 @@ void FGLTFTaskBuilder::CompleteAllTasks(FFeedbackContext* Context)
 			continue;
 		}
 
-		const FText MessageFormat = GetPriorityMessageFormat(Priority);
-		FScopedSlowTask Progress(Tasks->Num(), FText::Format(MessageFormat, FText()), true, *Context);
-		Progress.MakeDialog();
-
-		for (TUniquePtr<FGLTFTask>& Task : *Tasks)
+		if (Context != nullptr)
 		{
-			const FText Name = FText::FromString(Task->GetName());
-			const FText Message = FText::Format(MessageFormat, Name);
-			Progress.EnterProgressFrame(1, Message);
+			const FText MessageFormat = GetPriorityMessageFormat(Priority);
+			FScopedSlowTask Progress(Tasks->Num(), FText::Format(MessageFormat, FText()), true, *Context);
+			Progress.MakeDialog();
+
+			for (TUniquePtr<FGLTFTask>& Task : *Tasks)
+			{
+				const FText Name = FText::FromString(Task->GetName());
+				const FText Message = FText::Format(MessageFormat, Name);
+				Progress.EnterProgressFrame(1, Message);
 
-			Task->Complete();
+				Task->Complete();
+			}
+		}
+		else
+		{
+			for (TUniquePtr<FGLTFTask>& Task : *Tasks)
+			{
+				Task->Complete();
+			}
 		}
 	}
 
