[GLTFExporter] refactor: add GLTFBufferUtility and GLTFBufferAdapter

commit ad9026f59d06e2918e63ccaa9401e1e7ed701cfb
Date: Wed, 10 Nov 2021 20:08:46 +0100

    refactor: add GLTFBufferUtility and GLTFBufferAdapter
    (cherry picked from commit e22ac67ab666be6ccda96f4ef2ae81a111da6bd7​)

#skipci

diff --git a/Engine/Plugins/Enterprise/GLTFExporter/MigratedFromGit/Source/GLTFExporter/Private/Converters/GLTFBufferAdapter.cpp b/Engine/Plugins/Enterprise/GLTFExporter/MigratedFromGit/Source/GLTFExporter/Private/Converters/GLTFBufferAdapter.cpp
new file mode 100644
--- /dev/null
+++ b/Engine/Plugins/Enterprise/GLTFExporter/MigratedFromGit/Source/GLTFExporter/Private/Converters/GLTFBufferAdapter.cpp
@@ -0,0 +1,101 @@
+// Copyright Epic Games, Inc. All Rights Reserved.
+
+#include "Converters/GLTFBufferAdapter.h"
+#include "Converters/GLTFBufferUtility.h"
+
+class FGLTFBufferAdapterCPU final : public IGLTFBufferAdapter
+{
+public:
+
+	FGLTFBufferAdapterCPU(const void* Data)
+		: Data(Data)
+	{
+	}
+
+	virtual const uint8* GetData() override
+	{
+		return static_cast<const uint8*>(Data);
+	}
+
+	const void* Data;
+};
+
+class FGLTFBufferAdapterGPU final : public IGLTFBufferAdapter
+{
+public:
+
+	FGLTFBufferAdapterGPU(FRHIIndexBuffer* RHIBuffer)
+	{
+		FGLTFBufferUtility::ReadRHIBuffer(RHIBuffer, DataBuffer);
+	}
+
+	FGLTFBufferAdapterGPU(FRHIVertexBuffer* RHIBuffer)
+	{
+		FGLTFBufferUtility::ReadRHIBuffer(RHIBuffer, DataBuffer);
+	}
+
+	virtual const uint8* GetData() override
+	{
+		return DataBuffer.Num() > 0 ? DataBuffer.GetData() : nullptr;
+	}
+
+	TArray<uint8> DataBuffer;
+};
+
+TUniquePtr<IGLTFBufferAdapter> IGLTFBufferAdapter::GetIndices(const FRawStaticIndexBuffer* IndexBuffer)
+{
+	const void* IndexData = IndexBuffer->Is32Bit() ? static_cast<const void*>(IndexBuffer->AccessStream32()) : static_cast<const void*>(IndexBuffer->AccessStream16());
+	if (IndexData != nullptr && (WITH_EDITOR || FGLTFBufferUtility::GetAllowCPUAccess(IndexBuffer))) return MakeUnique<FGLTFBufferAdapterCPU>(IndexData);
+	return MakeUnique<FGLTFBufferAdapterGPU>(IndexBuffer->IndexBufferRHI);
+}
+
+TUniquePtr<IGLTFBufferAdapter> IGLTFBufferAdapter::GetIndices(const FRawStaticIndexBuffer16or32Interface* IndexBuffer)
+{
+	const void* IndexData = IndexBuffer->GetResourceDataSize() > 0 ? const_cast<FRawStaticIndexBuffer16or32Interface*>(IndexBuffer)->GetPointerTo(0) : nullptr;
+	if (IndexData != nullptr && (WITH_EDITOR || IndexBuffer->GetNeedsCPUAccess())) static_cast<TUniquePtr<IGLTFBufferAdapter>>(MakeUnique<FGLTFBufferAdapterCPU>(IndexData));
+	return MakeUnique<FGLTFBufferAdapterGPU>(IndexBuffer->IndexBufferRHI);
+}
+
+TUniquePtr<IGLTFBufferAdapter> IGLTFBufferAdapter::GetPositions(const FPositionVertexBuffer* VertexBuffer)
+{
+	const void* PositionData = VertexBuffer->GetVertexData();
+	if (PositionData != nullptr && (WITH_EDITOR || FGLTFBufferUtility::GetAllowCPUAccess(VertexBuffer))) MakeUnique<FGLTFBufferAdapterCPU>(PositionData);
+	return MakeUnique<FGLTFBufferAdapterGPU>(VertexBuffer->VertexBufferRHI);
+}
+
+TUniquePtr<IGLTFBufferAdapter> IGLTFBufferAdapter::GetColors(const FColorVertexBuffer* VertexBuffer)
+{
+	const void* ColorData = VertexBuffer->GetVertexData();
+	if (ColorData != nullptr && (WITH_EDITOR || FGLTFBufferUtility::GetAllowCPUAccess(VertexBuffer))) MakeUnique<FGLTFBufferAdapterCPU>(ColorData);
+	return MakeUnique<FGLTFBufferAdapterGPU>(VertexBuffer->VertexBufferRHI);
+}
+
+TUniquePtr<IGLTFBufferAdapter> IGLTFBufferAdapter::GetTangents(const FStaticMeshVertexBuffer* VertexBuffer)
+{
+	const void* TangentData = VertexBuffer->GetTangentData();
+	if (TangentData != nullptr && (WITH_EDITOR || const_cast<FStaticMeshVertexBuffer*>(VertexBuffer)->GetAllowCPUAccess())) MakeUnique<FGLTFBufferAdapterCPU>(TangentData);
+	return MakeUnique<FGLTFBufferAdapterGPU>(VertexBuffer->TangentsVertexBuffer.VertexBufferRHI);
+}
+
+TUniquePtr<IGLTFBufferAdapter> IGLTFBufferAdapter::GetUVs(const FStaticMeshVertexBuffer* VertexBuffer)
+{
+	const void* UVData = VertexBuffer->GetTexCoordData();
+	if (UVData != nullptr && (WITH_EDITOR || const_cast<FStaticMeshVertexBuffer*>(VertexBuffer)->GetAllowCPUAccess())) MakeUnique<FGLTFBufferAdapterCPU>(UVData);
+	return MakeUnique<FGLTFBufferAdapterGPU>(VertexBuffer->TexCoordVertexBuffer.VertexBufferRHI);
+}
+
+TUniquePtr<IGLTFBufferAdapter> IGLTFBufferAdapter::GetInfluences(const FSkinWeightVertexBuffer* VertexBuffer)
+{
+	const auto* InfluenceBuffer = VertexBuffer->GetDataVertexBuffer();
+	const void* InfluenceData = InfluenceBuffer->GetWeightData();
+	if (InfluenceData != nullptr && (WITH_EDITOR || VertexBuffer->GetNeedsCPUAccess())) MakeUnique<FGLTFBufferAdapterCPU>(InfluenceData);
+	return MakeUnique<FGLTFBufferAdapterGPU>(InfluenceBuffer->VertexBufferRHI);
+}
+
+TUniquePtr<IGLTFBufferAdapter> IGLTFBufferAdapter::GetLookups(const FSkinWeightVertexBuffer* VertexBuffer)
+{
+	const FSkinWeightLookupVertexBuffer* LookupBuffer = VertexBuffer->GetLookupVertexBuffer();
+	const void* LookupData = FGLTFBufferUtility::GetBufferData(LookupBuffer);
+	if (LookupData != nullptr && (WITH_EDITOR || VertexBuffer->GetNeedsCPUAccess())) MakeUnique<FGLTFBufferAdapterCPU>(LookupData);
+	return MakeUnique<FGLTFBufferAdapterGPU>(LookupBuffer->VertexBufferRHI);
+}
diff --git a/Engine/Plugins/Enterprise/GLTFExporter/MigratedFromGit/Source/GLTFExporter/Private/Converters/GLTFBufferAdapter.h b/Engine/Plugins/Enterprise/GLTFExporter/MigratedFromGit/Source/GLTFExporter/Private/Converters/GLTFBufferAdapter.h
new file mode 100644
--- /dev/null
+++ b/Engine/Plugins/Enterprise/GLTFExporter/MigratedFromGit/Source/GLTFExporter/Private/Converters/GLTFBufferAdapter.h
@@ -0,0 +1,23 @@
+// Copyright Epic Games, Inc. All Rights Reserved.
+
+#pragma once
+
+#include "Engine.h"
+
+class IGLTFBufferAdapter
+{
+public:
+
+	virtual ~IGLTFBufferAdapter() = default;
+
+	virtual const uint8* GetData() = 0;
+
+	static TUniquePtr<IGLTFBufferAdapter> GetIndices(const FRawStaticIndexBuffer* IndexBuffer);
+	static TUniquePtr<IGLTFBufferAdapter> GetIndices(const FRawStaticIndexBuffer16or32Interface* IndexBuffer);
+	static TUniquePtr<IGLTFBufferAdapter> GetPositions(const FPositionVertexBuffer* VertexBuffer);
+	static TUniquePtr<IGLTFBufferAdapter> GetColors(const FColorVertexBuffer* VertexBuffer);
+	static TUniquePtr<IGLTFBufferAdapter> GetTangents(const FStaticMeshVertexBuffer* VertexBuffer);
+	static TUniquePtr<IGLTFBufferAdapter> GetUVs(const FStaticMeshVertexBuffer* VertexBuffer);
+	static TUniquePtr<IGLTFBufferAdapter> GetInfluences(const FSkinWeightVertexBuffer* VertexBuffer);
+	static TUniquePtr<IGLTFBufferAdapter> GetLookups(const FSkinWeightVertexBuffer* VertexBuffer);
+};
diff --git a/Engine/Plugins/Enterprise/GLTFExporter/MigratedFromGit/Source/GLTFExporter/Private/Converters/GLTFBufferUtility.cpp b/Engine/Plugins/Enterprise/GLTFExporter/MigratedFromGit/Source/GLTFExporter/Private/Converters/GLTFBufferUtility.cpp
new file mode 100644
--- /dev/null
+++ b/Engine/Plugins/Enterprise/GLTFExporter/MigratedFromGit/Source/GLTFExporter/Private/Converters/GLTFBufferUtility.cpp
@@ -0,0 +1,107 @@
+// Copyright Epic Games, Inc. All Rights Reserved.
+
+#include "Converters/GLTFBufferUtility.h"
+#include "Rendering/SkeletalMeshRenderData.h"
+
+bool FGLTFBufferUtility::GetAllowCPUAccess(const FRawStaticIndexBuffer* IndexBuffer)
+{
+	struct FRawStaticIndexBufferHack : FIndexBuffer
+	{
+		TResourceArray<uint8, INDEXBUFFER_ALIGNMENT> IndexStorage;
+		int32 CachedNumIndices;
+		bool b32Bit;
+		bool bShouldExpandTo32Bit;
+	};
+
+	static_assert(sizeof(FRawStaticIndexBufferHack) == sizeof(FRawStaticIndexBuffer), "FRawStaticIndexBufferHack memory layout doesn't match FRawStaticIndexBuffer");
+	return reinterpret_cast<const FRawStaticIndexBufferHack*>(IndexBuffer)->IndexStorage.GetAllowCPUAccess();
+}
+
+bool FGLTFBufferUtility::GetAllowCPUAccess(const FPositionVertexBuffer* VertexBuffer)
+{
+	struct FPositionVertexBufferHack : FVertexBuffer
+	{
+		FShaderResourceViewRHIRef PositionComponentSRV;
+		TMemoryImagePtr<FStaticMeshVertexDataInterface> VertexData;
+		uint8* Data;
+		uint32 Stride;
+		uint32 NumVertices;
+		bool bNeedsCPUAccess;
+	};
+
+	static_assert(sizeof(FPositionVertexBufferHack) == sizeof(FPositionVertexBuffer), "FPositionVertexBufferHack memory layout doesn't match FPositionVertexBuffer");
+	const FStaticMeshVertexDataInterface* VertexData = reinterpret_cast<const FPositionVertexBufferHack*>(VertexBuffer)->VertexData;
+	return VertexData != nullptr && VertexData->GetAllowCPUAccess();
+}
+
+bool FGLTFBufferUtility::GetAllowCPUAccess(const FColorVertexBuffer* VertexBuffer)
+{
+	struct FColorVertexBufferHack : FVertexBuffer
+	{
+		FStaticMeshVertexDataInterface* VertexData;
+		FShaderResourceViewRHIRef ColorComponentsSRV;
+		uint8* Data;
+		uint32 Stride;
+		uint32 NumVertices;
+		bool NeedsCPUAccess;
+	};
+
+	static_assert(sizeof(FColorVertexBufferHack) == sizeof(FColorVertexBuffer), "FColorVertexBufferHack memory layout doesn't match FColorVertexBuffer");
+	const FStaticMeshVertexDataInterface* VertexData = reinterpret_cast<const FColorVertexBufferHack*>(VertexBuffer)->VertexData;
+	return VertexData != nullptr && VertexData->GetAllowCPUAccess();
+}
+
+const void* FGLTFBufferUtility::GetBufferData(const FSkinWeightLookupVertexBuffer* VertexBuffer)
+{
+	struct FSkinWeightLookupVertexBufferHack : FVertexBuffer
+	{
+		FShaderResourceViewRHIRef SRVValue;
+		bool bNeedsCPUAccess;
+		FStaticMeshVertexDataInterface* LookupData;
+		uint8* Data;
+		uint32 NumVertices;
+	};
+
+	static_assert(sizeof(FSkinWeightLookupVertexBufferHack) == sizeof(FSkinWeightLookupVertexBuffer), "FSkinWeightLookupVertexBufferHack memory layout doesn't match FSkinWeightLookupVertexBuffer");
+	return reinterpret_cast<const FSkinWeightLookupVertexBufferHack*>(VertexBuffer)->Data;
+}
+
+void FGLTFBufferUtility::ReadRHIBuffer(FRHIVertexBuffer* SourceBuffer, TArray<uint8>& OutData)
+{
+	const uint32 NumBytes = SourceBuffer->GetSize();
+
+	OutData.Empty();
+	OutData.AddUninitialized(NumBytes);
+	void *DstData = OutData.GetData();
+
+	ENQUEUE_RENDER_COMMAND(ReadRHIBuffer)(
+		[SourceBuffer, NumBytes, DstData](FRHICommandListImmediate& RHICmdList)
+		{
+			const void* SrcData = RHICmdList.LockVertexBuffer(SourceBuffer, 0, NumBytes, RLM_ReadOnly);
+			FMemory::Memcpy(DstData, SrcData, NumBytes);
+			RHICmdList.UnlockVertexBuffer(SourceBuffer);
+		}
+	);
+
+	FlushRenderingCommands();
+}
+
+void FGLTFBufferUtility::ReadRHIBuffer(FRHIIndexBuffer* SourceBuffer, TArray<uint8>& OutData)
+{
+	const uint32 NumBytes = SourceBuffer->GetSize();
+
+	OutData.Empty();
+	OutData.AddUninitialized(NumBytes);
+	void *DstData = OutData.GetData();
+
+	ENQUEUE_RENDER_COMMAND(ReadRHIBuffer)(
+		[SourceBuffer, NumBytes, DstData](FRHICommandListImmediate& RHICmdList)
+		{
+			const void* SrcData = RHICmdList.LockIndexBuffer(SourceBuffer, 0, NumBytes, RLM_ReadOnly);
+			FMemory::Memcpy(DstData, SrcData, NumBytes);
+			RHICmdList.UnlockIndexBuffer(SourceBuffer);
+		}
+	);
+
+	FlushRenderingCommands();
+}
diff --git a/Engine/Plugins/Enterprise/GLTFExporter/MigratedFromGit/Source/GLTFExporter/Private/Converters/GLTFBufferUtility.h b/Engine/Plugins/Enterprise/GLTFExporter/MigratedFromGit/Source/GLTFExporter/Private/Converters/GLTFBufferUtility.h
new file mode 100644
--- /dev/null
+++ b/Engine/Plugins/Enterprise/GLTFExporter/MigratedFromGit/Source/GLTFExporter/Private/Converters/GLTFBufferUtility.h
@@ -0,0 +1,17 @@
+// Copyright Epic Games, Inc. All Rights Reserved.
+
+#pragma once
+
+#include "Engine.h"
+
+struct FGLTFBufferUtility
+{
+	static bool GetAllowCPUAccess(const FRawStaticIndexBuffer* IndexBuffer);
+	static bool GetAllowCPUAccess(const FPositionVertexBuffer* VertexBuffer);
+	static bool GetAllowCPUAccess(const FColorVertexBuffer* VertexBuffer);
+
+	static const void* GetBufferData(const FSkinWeightLookupVertexBuffer* VertexBuffer);
+
+	static void ReadRHIBuffer(FRHIIndexBuffer* SourceBuffer, TArray<uint8>& OutData);
+	static void ReadRHIBuffer(FRHIVertexBuffer* SourceBuffer, TArray<uint8>& OutData);
+};
