[GLTFExporter] feat: add support for reading vertex and index buffers from GPU if CPU access not enabled in runtime

commit f225c36a5438c810d8a6e97b532d9cd304301369
Date: Wed, 10 Nov 2021 20:12:50 +0100

    feat: add support for reading vertex and index buffers from GPU if CPU access not enabled in runtime
    (cherry picked from commit 1ce093dcb3f2e7379f0cf2d4674ace0a6120d206​)

#skipci

diff --git a/Engine/Plugins/Enterprise/GLTFExporter/MigratedFromGit/Source/GLTFExporter/Private/Converters/GLTFAccessorConverters.cpp b/Engine/Plugins/Enterprise/GLTFExporter/MigratedFromGit/Source/GLTFExporter/Private/Converters/GLTFAccessorConverters.cpp
--- a/Engine/Plugins/Enterprise/GLTFExporter/MigratedFromGit/Source/GLTFExporter/Private/Converters/GLTFAccessorConverters.cpp
+++ b/Engine/Plugins/Enterprise/GLTFExporter/MigratedFromGit/Source/GLTFExporter/Private/Converters/GLTFAccessorConverters.cpp
@@ -2,6 +2,7 @@
 
 #include "Converters/GLTFAccessorConverters.h"
 #include "Converters/GLTFConverterUtility.h"
+#include "Converters/GLTFBufferAdapter.h"
 #include "Builders/GLTFConvertBuilder.h"
 
 // TODO: Unreal-style implementation of std::conditional to avoid mixing in STL. Should be added to the engine.
@@ -26,14 +27,12 @@ FGLTFJsonAccessorIndex FGLTFPositionBufferConverter::Convert(const FGLTFMeshSect
 		return FGLTFJsonAccessorIndex(INDEX_NONE);
 	}
 
-	if (VertexBuffer->GetVertexData() == nullptr)
-	{
-		// TODO: report error
-		return FGLTFJsonAccessorIndex(INDEX_NONE);
-	}
+	const TUniquePtr<IGLTFBufferAdapter> SourceBuffer = IGLTFBufferAdapter::GetPositions(VertexBuffer);
+	const uint8* SourceData = SourceBuffer->GetData();
 
 	const TArray<uint32>& IndexMap = MeshSection->IndexMap;
 	const uint32 VertexCount = IndexMap.Num();
+	const uint32 Stride = VertexBuffer->GetStride();
 
 	TArray<FGLTFVector3> Positions;
 	Positions.AddUninitialized(VertexCount);
@@ -41,7 +40,8 @@ FGLTFJsonAccessorIndex FGLTFPositionBufferConverter::Convert(const FGLTFMeshSect
 	for (uint32 VertexIndex = 0; VertexIndex < VertexCount; ++VertexIndex)
 	{
 		const uint32 MappedVertexIndex = IndexMap[VertexIndex];
-		Positions[VertexIndex] = FGLTFConverterUtility::ConvertPosition(VertexBuffer->VertexPosition(MappedVertexIndex), Builder.ExportOptions->ExportUniformScale);
+		const FVector& Position = *reinterpret_cast<const FVector*>(SourceData + Stride * MappedVertexIndex);
+		Positions[VertexIndex] = FGLTFConverterUtility::ConvertPosition(Position, Builder.ExportOptions->ExportUniformScale);
 	}
 
 	FGLTFJsonAccessor JsonAccessor;
@@ -82,14 +82,12 @@ FGLTFJsonAccessorIndex FGLTFColorBufferConverter::Convert(const FGLTFMeshSection
 		return FGLTFJsonAccessorIndex(INDEX_NONE);
 	}
 
-	if (VertexBuffer->GetVertexData() == nullptr)
-	{
-		// TODO: report error
-		return FGLTFJsonAccessorIndex(INDEX_NONE);
-	}
+	const TUniquePtr<IGLTFBufferAdapter> SourceBuffer = IGLTFBufferAdapter::GetColors(VertexBuffer);
+	const uint8* SourceData = SourceBuffer->GetData();
 
 	const TArray<uint32>& IndexMap = MeshSection->IndexMap;
 	const uint32 VertexCount = IndexMap.Num();
+	const uint32 Stride = VertexBuffer->GetStride();
 
 	TArray<FGLTFUInt8Color4> Colors;
 	Colors.AddUninitialized(VertexCount);
@@ -97,7 +95,8 @@ FGLTFJsonAccessorIndex FGLTFColorBufferConverter::Convert(const FGLTFMeshSection
 	for (uint32 VertexIndex = 0; VertexIndex < VertexCount; ++VertexIndex)
 	{
 		const uint32 MappedVertexIndex = IndexMap[VertexIndex];
-		Colors[VertexIndex] = FGLTFConverterUtility::ConvertColor(VertexBuffer->VertexColor(MappedVertexIndex));
+		const FColor& Color = *reinterpret_cast<const FColor*>(SourceData + Stride * MappedVertexIndex);
+		Colors[VertexIndex] = FGLTFConverterUtility::ConvertColor(Color);
 	}
 
 	FGLTFJsonAccessor JsonAccessor;
@@ -117,12 +116,8 @@ FGLTFJsonAccessorIndex FGLTFNormalBufferConverter::Convert(const FGLTFMeshSectio
 		return FGLTFJsonAccessorIndex(INDEX_NONE);
 	}
 
-	const void* TangentData = const_cast<FStaticMeshVertexBuffer*>(VertexBuffer)->GetTangentData();
-	if (TangentData == nullptr)
-	{
-		// TODO: report error
-		return FGLTFJsonAccessorIndex(INDEX_NONE);
-	}
+	const TUniquePtr<IGLTFBufferAdapter> SourceBuffer = IGLTFBufferAdapter::GetTangents(VertexBuffer);
+	const uint8* SourceData = SourceBuffer->GetData();
 
 	FGLTFJsonBufferViewIndex BufferViewIndex;
 	EGLTFJsonComponentType ComponentType;
@@ -137,13 +132,13 @@ FGLTFJsonAccessorIndex FGLTFNormalBufferConverter::Convert(const FGLTFMeshSectio
 		if (bHighPrecision)
 		{
 			ComponentType = EGLTFJsonComponentType::S16;
-			BufferViewIndex = ConvertBufferView<FGLTFInt16Vector4, FPackedRGBA16N>(MeshSection, TangentData);
+			BufferViewIndex = ConvertBufferView<FGLTFInt16Vector4, FPackedRGBA16N>(MeshSection, SourceData);
 			Builder.GetBufferView(BufferViewIndex).ByteStride = sizeof(FGLTFInt16Vector4);
 		}
 		else
 		{
 			ComponentType = EGLTFJsonComponentType::S8;
-			BufferViewIndex = ConvertBufferView<FGLTFInt8Vector4, FPackedNormal>(MeshSection, TangentData);
+			BufferViewIndex = ConvertBufferView<FGLTFInt8Vector4, FPackedNormal>(MeshSection, SourceData);
 			Builder.GetBufferView(BufferViewIndex).ByteStride = sizeof(FGLTFInt8Vector4);
 		}
 	}
@@ -151,8 +146,8 @@ FGLTFJsonAccessorIndex FGLTFNormalBufferConverter::Convert(const FGLTFMeshSectio
 	{
 		ComponentType = EGLTFJsonComponentType::F32;
 		BufferViewIndex = bHighPrecision
-			? ConvertBufferView<FGLTFVector3, FPackedRGBA16N>(MeshSection, TangentData)
-			: ConvertBufferView<FGLTFVector3, FPackedNormal>(MeshSection, TangentData);
+			? ConvertBufferView<FGLTFVector3, FPackedRGBA16N>(MeshSection, SourceData)
+			: ConvertBufferView<FGLTFVector3, FPackedNormal>(MeshSection, SourceData);
 	}
 
 	FGLTFJsonAccessor JsonAccessor;
@@ -166,13 +161,13 @@ FGLTFJsonAccessorIndex FGLTFNormalBufferConverter::Convert(const FGLTFMeshSectio
 }
 
 template <typename DestinationType, typename SourceType>
-FGLTFJsonBufferViewIndex FGLTFNormalBufferConverter::ConvertBufferView(const FGLTFMeshSection* MeshSection, const void* TangentData)
+FGLTFJsonBufferViewIndex FGLTFNormalBufferConverter::ConvertBufferView(const FGLTFMeshSection* MeshSection, const void* SourceData)
 {
 	const TArray<uint32>& IndexMap = MeshSection->IndexMap;
 	const uint32 VertexCount = IndexMap.Num();
 
 	typedef TStaticMeshVertexTangentDatum<SourceType> VertexTangentType;
-	const VertexTangentType* VertexTangents = static_cast<const VertexTangentType*>(TangentData);
+	const VertexTangentType* TangentData= static_cast<const VertexTangentType*>(SourceData);
 
 	TArray<DestinationType> Normals;
 	Normals.AddUninitialized(VertexCount);
@@ -180,7 +175,7 @@ FGLTFJsonBufferViewIndex FGLTFNormalBufferConverter::ConvertBufferView(const FGL
 	for (uint32 VertexIndex = 0; VertexIndex < VertexCount; ++VertexIndex)
 	{
 		const uint32 MappedVertexIndex = IndexMap[VertexIndex];
-		const FVector SafeNormal = VertexTangents[MappedVertexIndex].TangentZ.ToFVector().GetSafeNormal();
+		const FVector SafeNormal = TangentData[MappedVertexIndex].TangentZ.ToFVector().GetSafeNormal();
 
 		typedef typename TConditional<TIsSame<DestinationType, FGLTFVector3>::Value, FVector, SourceType>::Type IntermediateType;
 		Normals[VertexIndex] = FGLTFConverterUtility::ConvertNormal(IntermediateType(SafeNormal));
@@ -196,12 +191,8 @@ FGLTFJsonAccessorIndex FGLTFTangentBufferConverter::Convert(const FGLTFMeshSecti
 		return FGLTFJsonAccessorIndex(INDEX_NONE);
 	}
 
-	const void* TangentData = const_cast<FStaticMeshVertexBuffer*>(VertexBuffer)->GetTangentData();
-	if (TangentData == nullptr)
-	{
-		// TODO: report error
-		return FGLTFJsonAccessorIndex(INDEX_NONE);
-	}
+	const TUniquePtr<IGLTFBufferAdapter> SourceBuffer = IGLTFBufferAdapter::GetTangents(VertexBuffer);
+	const uint8* SourceData = SourceBuffer->GetData();
 
 	FGLTFJsonBufferViewIndex BufferViewIndex;
 	EGLTFJsonComponentType ComponentType;
@@ -216,20 +207,20 @@ FGLTFJsonAccessorIndex FGLTFTangentBufferConverter::Convert(const FGLTFMeshSecti
 		if (bHighPrecision)
 		{
 			ComponentType = EGLTFJsonComponentType::S16;
-			BufferViewIndex = ConvertBufferView<FGLTFInt16Vector4, FPackedRGBA16N>(MeshSection, TangentData);
+			BufferViewIndex = ConvertBufferView<FGLTFInt16Vector4, FPackedRGBA16N>(MeshSection, SourceData);
 		}
 		else
 		{
 			ComponentType = EGLTFJsonComponentType::S8;
-			BufferViewIndex = ConvertBufferView<FGLTFInt8Vector4, FPackedNormal>(MeshSection, TangentData);
+			BufferViewIndex = ConvertBufferView<FGLTFInt8Vector4, FPackedNormal>(MeshSection, SourceData);
 		}
 	}
 	else
 	{
 		ComponentType = EGLTFJsonComponentType::F32;
 		BufferViewIndex = bHighPrecision
-			? ConvertBufferView<FGLTFVector4, FPackedRGBA16N>(MeshSection, TangentData)
-			: ConvertBufferView<FGLTFVector4, FPackedNormal>(MeshSection, TangentData);
+			? ConvertBufferView<FGLTFVector4, FPackedRGBA16N>(MeshSection, SourceData)
+			: ConvertBufferView<FGLTFVector4, FPackedNormal>(MeshSection, SourceData);
 	}
 
 	FGLTFJsonAccessor JsonAccessor;
@@ -243,13 +234,13 @@ FGLTFJsonAccessorIndex FGLTFTangentBufferConverter::Convert(const FGLTFMeshSecti
 }
 
 template <typename DestinationType, typename SourceType>
-FGLTFJsonBufferViewIndex FGLTFTangentBufferConverter::ConvertBufferView(const FGLTFMeshSection* MeshSection, const void* TangentData)
+FGLTFJsonBufferViewIndex FGLTFTangentBufferConverter::ConvertBufferView(const FGLTFMeshSection* MeshSection, const void* SourceData)
 {
 	const TArray<uint32>& IndexMap = MeshSection->IndexMap;
 	const uint32 VertexCount = IndexMap.Num();
 
 	typedef TStaticMeshVertexTangentDatum<SourceType> VertexTangentType;
-	const VertexTangentType* VertexTangents = static_cast<const VertexTangentType*>(TangentData);
+	const VertexTangentType* VertexTangents = static_cast<const VertexTangentType*>(SourceData);
 
 	TArray<DestinationType> Tangents;
 	Tangents.AddUninitialized(VertexCount);
@@ -273,15 +264,6 @@ FGLTFJsonAccessorIndex FGLTFUVBufferConverter::Convert(const FGLTFMeshSection* M
 		return FGLTFJsonAccessorIndex(INDEX_NONE);
 	}
 
-	if (const_cast<FStaticMeshVertexBuffer*>(VertexBuffer)->GetTexCoordData() == nullptr)
-	{
-		// TODO: report error
-		return FGLTFJsonAccessorIndex(INDEX_NONE);
-	}
-
-	const TArray<uint32>& IndexMap = MeshSection->IndexMap;
-	const uint32 VertexCount = IndexMap.Num();
-
 	const uint32 UVCount = VertexBuffer->GetNumTexCoords();
 	if (UVIndex >= UVCount)
 	{
@@ -291,13 +273,31 @@ FGLTFJsonAccessorIndex FGLTFUVBufferConverter::Convert(const FGLTFMeshSection* M
 
 	// TODO: report warning or add support for half float precision UVs, i.e. !VertexBuffer->GetUseFullPrecisionUVs()?
 
+	const TUniquePtr<IGLTFBufferAdapter> SourceBuffer = IGLTFBufferAdapter::GetUVs(VertexBuffer);
+	const uint8* SourceData = SourceBuffer->GetData();
+
+	return VertexBuffer->GetUseFullPrecisionUVs()
+		? Convert<FVector2D>(MeshSection, VertexBuffer, UVIndex, SourceData)
+		: Convert<FVector2DHalf>(MeshSection, VertexBuffer, UVIndex, SourceData);
+}
+
+template <typename SourceType>
+FGLTFJsonAccessorIndex FGLTFUVBufferConverter::Convert(const FGLTFMeshSection* MeshSection, const FStaticMeshVertexBuffer* VertexBuffer, uint32 UVIndex, const uint8* SourceData) const
+{
+	const TArray<uint32>& IndexMap = MeshSection->IndexMap;
+	const uint32 VertexCount = IndexMap.Num();
+
 	TArray<FGLTFVector2> UVs;
 	UVs.AddUninitialized(VertexCount);
 
+	const uint32 UVCount = VertexBuffer->GetNumTexCoords();
+	const SourceType* UVData = reinterpret_cast<const SourceType*>(SourceData);
+
 	for (uint32 VertexIndex = 0; VertexIndex < VertexCount; ++VertexIndex)
 	{
 		const uint32 MappedVertexIndex = IndexMap[VertexIndex];
-		UVs[VertexIndex] = FGLTFConverterUtility::ConvertUV(VertexBuffer->GetVertexUV(MappedVertexIndex, UVIndex));
+		const SourceType& UV = UVData[UVCount * MappedVertexIndex + UVIndex];
+		UVs[VertexIndex] = FGLTFConverterUtility::ConvertUV(UV);
 	}
 
 	FGLTFJsonAccessor JsonAccessor;
@@ -310,60 +310,100 @@ FGLTFJsonAccessorIndex FGLTFUVBufferConverter::Convert(const FGLTFMeshSection* M
 }
 
 FGLTFJsonAccessorIndex FGLTFBoneIndexBufferConverter::Convert(const FGLTFMeshSection* MeshSection, const FSkinWeightVertexBuffer* VertexBuffer, uint32 InfluenceOffset)
-{
-	return MeshSection->MaxBoneIndex <= UINT8_MAX
-		? Convert<uint8>(MeshSection, VertexBuffer, InfluenceOffset)
-		: Convert<uint16>(MeshSection, VertexBuffer, InfluenceOffset);
-}
-
-template <typename IndexType>
-FGLTFJsonAccessorIndex FGLTFBoneIndexBufferConverter::Convert(const FGLTFMeshSection* MeshSection, const FSkinWeightVertexBuffer* VertexBuffer, uint32 InfluenceOffset) const
 {
 	if (VertexBuffer == nullptr || VertexBuffer->GetNumVertices() == 0)
 	{
 		return FGLTFJsonAccessorIndex(INDEX_NONE);
 	}
 
-	if (VertexBuffer->GetDataVertexBuffer()->GetWeightData() == nullptr)
+	const TUniquePtr<IGLTFBufferAdapter> SourceBuffer = IGLTFBufferAdapter::GetInfluences(VertexBuffer);
+	const uint8* SourceData = SourceBuffer->GetData();
+
+	return MeshSection->MaxBoneIndex > UINT8_MAX
+		? Convert<uint16>(MeshSection, VertexBuffer, InfluenceOffset, SourceData)
+		: Convert<uint8>(MeshSection, VertexBuffer, InfluenceOffset, SourceData);
+}
+
+template <typename DestinationType>
+FGLTFJsonAccessorIndex FGLTFBoneIndexBufferConverter::Convert(const FGLTFMeshSection* MeshSection, const FSkinWeightVertexBuffer* VertexBuffer, uint32 InfluenceOffset, const uint8* SourceData) const
+{
+	return VertexBuffer->Use16BitBoneIndex()
+		? Convert<DestinationType, uint16>(MeshSection, VertexBuffer, InfluenceOffset, SourceData)
+		: Convert<DestinationType, uint8>(MeshSection, VertexBuffer, InfluenceOffset, SourceData);
+}
+
+template <typename DestinationType, typename SourceType>
+FGLTFJsonAccessorIndex FGLTFBoneIndexBufferConverter::Convert(const FGLTFMeshSection* MeshSection, const FSkinWeightVertexBuffer* VertexBuffer, uint32 InfluenceOffset, const uint8* SourceData) const
+{
+	if (VertexBuffer->GetVariableBonesPerVertex())
 	{
-		// TODO: report error
-		return FGLTFJsonAccessorIndex(INDEX_NONE);
+		const TUniquePtr<IGLTFBufferAdapter> LookupBuffer = IGLTFBufferAdapter::GetLookups(VertexBuffer);
+		const uint32* LookupData = reinterpret_cast<const uint32*>(LookupBuffer->GetData());
+
+		if (LookupData == nullptr)
+		{
+			// TODO: report error
+			return FGLTFJsonAccessorIndex(INDEX_NONE);
+		}
+
+		return Convert<DestinationType, SourceType>(MeshSection, VertexBuffer, InfluenceOffset, SourceData, [LookupData](uint32 VertexIndex, uint32& VertexDataOffset, uint32& VertexInfluenceCount)
+		{
+			const uint32 Value = LookupData[VertexIndex];
+			VertexDataOffset = Value >> 8;
+			VertexInfluenceCount = Value & 0xff;
+		});
 	}
 
-	const uint32 InfluenceCount = VertexBuffer->GetMaxBoneInfluences();
-	if (InfluenceOffset >= InfluenceCount)
+	const uint32 MaxBoneInfluences = VertexBuffer->GetMaxBoneInfluences();
+	return Convert<DestinationType, SourceType>(MeshSection, VertexBuffer, InfluenceOffset, SourceData, [MaxBoneInfluences](uint32 VertexIndex, uint32& VertexDataOffset, uint32& VertexInfluenceCount)
 	{
-		// TODO: report warning
-		return FGLTFJsonAccessorIndex(INDEX_NONE);
-	}
+		VertexDataOffset = (sizeof(SourceType) + sizeof(uint8)) * MaxBoneInfluences * VertexIndex;
+		VertexInfluenceCount = MaxBoneInfluences;
+	});
+}
 
+template <typename DestinationType, typename SourceType, typename CallbackType>
+FGLTFJsonAccessorIndex FGLTFBoneIndexBufferConverter::Convert(const FGLTFMeshSection* MeshSection, const FSkinWeightVertexBuffer* VertexBuffer, uint32 InfluenceOffset, const uint8* SourceData, CallbackType GetVertexInfluenceOffsetCount) const
+{
 	const TArray<uint32>& IndexMap = MeshSection->IndexMap;
 	const uint32 VertexCount = IndexMap.Num();
 
-	struct VertexBoneIndices
+	struct FVertexBoneIndices
 	{
-		IndexType Index[4];
+		DestinationType Index[4];
 	};
 
-	TArray<VertexBoneIndices> BoneIndices;
+	TArray<FVertexBoneIndices> BoneIndices;
 	BoneIndices.AddUninitialized(VertexCount);
 
 	for (uint32 VertexIndex = 0; VertexIndex < VertexCount; ++VertexIndex)
 	{
 		const TArray<FBoneIndexType>& BoneMap = MeshSection->BoneMaps[MeshSection->BoneMapLookup[VertexIndex]];
 		const uint32 MappedVertexIndex = IndexMap[VertexIndex];
-		VertexBoneIndices& VertexBones = BoneIndices[VertexIndex];
+		FVertexBoneIndices& VertexBoneIndices = BoneIndices[VertexIndex];
+
+		uint32 VertexDataOffset;
+		uint32 VertexInfluenceCount;
+		GetVertexInfluenceOffsetCount(MappedVertexIndex, VertexDataOffset, VertexInfluenceCount);
 
-		for (int32 InfluenceIndex = 0; InfluenceIndex < 4; ++InfluenceIndex)
+		const SourceType* VertexBoneIndexData = reinterpret_cast<const SourceType*>(SourceData + VertexDataOffset);
+		const int32 InfluenceCount = FMath::Min(static_cast<int32>(VertexInfluenceCount - InfluenceOffset), 4);
+
+		for (int32 InfluenceIndex = 0; InfluenceIndex < InfluenceCount; ++InfluenceIndex)
 		{
-			const uint32 UnmappedBoneIndex = VertexBuffer->GetBoneIndex(MappedVertexIndex, InfluenceOffset + InfluenceIndex);
-			VertexBones.Index[InfluenceIndex] = static_cast<IndexType>(BoneMap[UnmappedBoneIndex]);
+			const uint32 UnmappedBoneIndex = VertexBoneIndexData[InfluenceOffset + InfluenceIndex];
+			VertexBoneIndices.Index[InfluenceIndex] = static_cast<DestinationType>(BoneMap[UnmappedBoneIndex]);
+		}
+
+		for (int32 InfluenceIndex = InfluenceCount; InfluenceIndex < 4; ++InfluenceIndex)
+		{
+			VertexBoneIndices.Index[InfluenceIndex] = 0;
 		}
 	}
 
 	FGLTFJsonAccessor JsonAccessor;
 	JsonAccessor.BufferView = Builder.AddBufferView(BoneIndices, EGLTFJsonBufferTarget::ArrayBuffer);
-	JsonAccessor.ComponentType = FGLTFConverterUtility::GetComponentType<IndexType>();
+	JsonAccessor.ComponentType = FGLTFConverterUtility::GetComponentType<DestinationType>();
 	JsonAccessor.Count = VertexCount;
 	JsonAccessor.Type = EGLTFJsonAccessorType::Vec4;
 
@@ -377,39 +417,78 @@ FGLTFJsonAccessorIndex FGLTFBoneWeightBufferConverter::Convert(const FGLTFMeshSe
 		return FGLTFJsonAccessorIndex(INDEX_NONE);
 	}
 
-	if (VertexBuffer->GetDataVertexBuffer()->GetWeightData() == nullptr)
+	const TUniquePtr<IGLTFBufferAdapter> SourceBuffer = IGLTFBufferAdapter::GetInfluences(VertexBuffer);
+	const uint8* SourceData = SourceBuffer->GetData();
+
+	return MeshSection->MaxBoneIndex > UINT8_MAX
+		? Convert<uint16>(MeshSection, VertexBuffer, InfluenceOffset, SourceData)
+		: Convert<uint8>(MeshSection, VertexBuffer, InfluenceOffset, SourceData);
+}
+
+template <typename BoneIndexType>
+FGLTFJsonAccessorIndex FGLTFBoneWeightBufferConverter::Convert(const FGLTFMeshSection* MeshSection, const FSkinWeightVertexBuffer* VertexBuffer, uint32 InfluenceOffset, const uint8* SourceData) const
+{
+	if (VertexBuffer->GetVariableBonesPerVertex())
 	{
-		// TODO: report error
-		return FGLTFJsonAccessorIndex(INDEX_NONE);
+		const TUniquePtr<IGLTFBufferAdapter> LookupBuffer = IGLTFBufferAdapter::GetLookups(VertexBuffer);
+		const uint32* LookupData = reinterpret_cast<const uint32*>(LookupBuffer->GetData());
+
+		if (LookupData == nullptr)
+		{
+			// TODO: report error
+			return FGLTFJsonAccessorIndex(INDEX_NONE);
+		}
+
+		return Convert<BoneIndexType>(MeshSection, VertexBuffer, InfluenceOffset, SourceData, [LookupData](uint32 VertexIndex, uint32& VertexDataOffset, uint32& VertexInfluenceCount)
+		{
+			const uint32 Value = LookupData[VertexIndex];
+			VertexDataOffset = Value >> 8;
+			VertexInfluenceCount = Value & 0xff;
+		});
 	}
 
-	const uint32 InfluenceCount = VertexBuffer->GetMaxBoneInfluences();
-	if (InfluenceOffset >= InfluenceCount)
+	const uint32 MaxBoneInfluences = VertexBuffer->GetMaxBoneInfluences();
+	return Convert<BoneIndexType>(MeshSection, VertexBuffer, InfluenceOffset, SourceData, [MaxBoneInfluences](uint32 VertexIndex, uint32& VertexDataOffset, uint32& VertexInfluenceCount)
 	{
-		// TODO: report warning
-		return FGLTFJsonAccessorIndex(INDEX_NONE);
-	}
+		VertexDataOffset = (sizeof(BoneIndexType) + sizeof(uint8)) * MaxBoneInfluences * VertexIndex;
+		VertexInfluenceCount = MaxBoneInfluences;
+	});
+}
 
+template <typename BoneIndexType, typename CallbackType>
+FGLTFJsonAccessorIndex FGLTFBoneWeightBufferConverter::Convert(const FGLTFMeshSection* MeshSection, const FSkinWeightVertexBuffer* VertexBuffer, uint32 InfluenceOffset, const uint8* SourceData, CallbackType GetVertexInfluenceOffsetCount) const
+{
 	const TArray<uint32>& IndexMap = MeshSection->IndexMap;
 	const uint32 VertexCount = IndexMap.Num();
 
-	struct VertexBoneWeights
+	struct FVertexBoneWeights
 	{
 		uint8 Weights[4];
 	};
 
-	TArray<VertexBoneWeights> BoneWeights;
+	TArray<FVertexBoneWeights> BoneWeights;
 	BoneWeights.AddUninitialized(VertexCount);
 
 	for (uint32 VertexIndex = 0; VertexIndex < VertexCount; ++VertexIndex)
 	{
 		const uint32 MappedVertexIndex = IndexMap[VertexIndex];
-		VertexBoneWeights& VertexWeights = BoneWeights[VertexIndex];
+		FVertexBoneWeights& VertexBoneWeights = BoneWeights[VertexIndex];
+
+		uint32 VertexDataOffset;
+		uint32 VertexInfluenceCount;
+		GetVertexInfluenceOffsetCount(MappedVertexIndex, VertexDataOffset, VertexInfluenceCount);
+
+		const uint8* VertexBoneWeightsData = SourceData + VertexDataOffset + sizeof(BoneIndexType) * VertexInfluenceCount;
+		const int32 InfluenceCount = FMath::Min(static_cast<int32>(VertexInfluenceCount - InfluenceOffset), 4);
+
+		for (int32 InfluenceIndex = 0; InfluenceIndex < InfluenceCount; ++InfluenceIndex)
+		{
+			VertexBoneWeights.Weights[InfluenceIndex] = VertexBoneWeightsData[InfluenceOffset + InfluenceIndex];
+		}
 
-		for (uint32 InfluenceIndex = 0; InfluenceIndex < 4; ++InfluenceIndex)
+		for (int32 InfluenceIndex = InfluenceCount; InfluenceIndex < 4; ++InfluenceIndex)
 		{
-			const uint8 BoneWeight = VertexBuffer->GetBoneWeight(MappedVertexIndex, InfluenceOffset + InfluenceIndex);
-			VertexWeights.Weights[InfluenceIndex] = BoneWeight;
+			VertexBoneWeights.Weights[InfluenceIndex] = 0;
 		}
 	}
 
diff --git a/Engine/Plugins/Enterprise/GLTFExporter/MigratedFromGit/Source/GLTFExporter/Private/Converters/GLTFAccessorConverters.h b/Engine/Plugins/Enterprise/GLTFExporter/MigratedFromGit/Source/GLTFExporter/Private/Converters/GLTFAccessorConverters.h
--- a/Engine/Plugins/Enterprise/GLTFExporter/MigratedFromGit/Source/GLTFExporter/Private/Converters/GLTFAccessorConverters.h
+++ b/Engine/Plugins/Enterprise/GLTFExporter/MigratedFromGit/Source/GLTFExporter/Private/Converters/GLTFAccessorConverters.h
@@ -53,6 +53,9 @@ class FGLTFUVBufferConverter final : public TGLTFAccessorConverter<const FGLTFMe
 	using TGLTFAccessorConverter::TGLTFAccessorConverter;
 
 	virtual FGLTFJsonAccessorIndex Convert(const FGLTFMeshSection* MeshSection, const FStaticMeshVertexBuffer* VertexBuffer, uint32 UVIndex) override;
+
+	template <typename SourceType>
+	FGLTFJsonAccessorIndex Convert(const FGLTFMeshSection* MeshSection, const FStaticMeshVertexBuffer* VertexBuffer, uint32 UVIndex, const uint8* SourceData) const;
 };
 
 class FGLTFBoneIndexBufferConverter final : public TGLTFAccessorConverter<const FGLTFMeshSection*, const FSkinWeightVertexBuffer*, uint32>
@@ -61,8 +64,14 @@ class FGLTFBoneIndexBufferConverter final : public TGLTFAccessorConverter<const
 
 	virtual FGLTFJsonAccessorIndex Convert(const FGLTFMeshSection* MeshSection, const FSkinWeightVertexBuffer* VertexBuffer, uint32 InfluenceOffset) override;
 
-	template <typename IndexType>
-	FGLTFJsonAccessorIndex Convert(const FGLTFMeshSection* MeshSection, const FSkinWeightVertexBuffer* VertexBuffer, uint32 InfluenceOffset) const;
+	template <typename DestinationType>
+	FGLTFJsonAccessorIndex Convert(const FGLTFMeshSection* MeshSection, const FSkinWeightVertexBuffer* VertexBuffer, uint32 InfluenceOffset, const uint8* SourceData) const;
+
+	template <typename DestinationType, typename SourceType>
+	FGLTFJsonAccessorIndex Convert(const FGLTFMeshSection* MeshSection, const FSkinWeightVertexBuffer* VertexBuffer, uint32 InfluenceOffset, const uint8* SourceData) const;
+
+	template <typename DestinationType, typename SourceType, typename CallbackType>
+	FGLTFJsonAccessorIndex Convert(const FGLTFMeshSection* MeshSection, const FSkinWeightVertexBuffer* VertexBuffer, uint32 InfluenceOffset, const uint8* SourceData, CallbackType GetVertexInfluenceOffsetCount) const;
 };
 
 class FGLTFBoneWeightBufferConverter final : public TGLTFAccessorConverter<const FGLTFMeshSection*, const FSkinWeightVertexBuffer*, uint32>
@@ -70,6 +79,12 @@ class FGLTFBoneWeightBufferConverter final : public TGLTFAccessorConverter<const
 	using TGLTFAccessorConverter::TGLTFAccessorConverter;
 
 	virtual FGLTFJsonAccessorIndex Convert(const FGLTFMeshSection* MeshSection, const FSkinWeightVertexBuffer* VertexBuffer, uint32 InfluenceOffset) override;
+
+	template <typename BoneIndexType>
+	FGLTFJsonAccessorIndex Convert(const FGLTFMeshSection* MeshSection, const FSkinWeightVertexBuffer* VertexBuffer, uint32 InfluenceOffset, const uint8* SourceData) const;
+
+	template <typename BoneIndexType, typename CallbackType>
+	FGLTFJsonAccessorIndex Convert(const FGLTFMeshSection* MeshSection, const FSkinWeightVertexBuffer* VertexBuffer, uint32 InfluenceOffset, const uint8* SourceData, CallbackType GetVertexInfluenceOffsetCount) const;
 };
 
 class FGLTFIndexBufferConverter final : public TGLTFAccessorConverter<const FGLTFMeshSection*>
diff --git a/Engine/Plugins/Enterprise/GLTFExporter/MigratedFromGit/Source/GLTFExporter/Private/Converters/GLTFConverterUtility.h b/Engine/Plugins/Enterprise/GLTFExporter/MigratedFromGit/Source/GLTFExporter/Private/Converters/GLTFConverterUtility.h
--- a/Engine/Plugins/Enterprise/GLTFExporter/MigratedFromGit/Source/GLTFExporter/Private/Converters/GLTFConverterUtility.h
+++ b/Engine/Plugins/Enterprise/GLTFExporter/MigratedFromGit/Source/GLTFExporter/Private/Converters/GLTFConverterUtility.h
@@ -72,6 +72,11 @@ struct FGLTFConverterUtility
 		return { UV.X, UV.Y };
 	}
 
+	static FGLTFVector2 ConvertUV(const FVector2DHalf& UV)
+	{
+		return ConvertUV(FVector2D(UV));
+	}
+
 	static FGLTFColor4 ConvertColor(const FLinearColor& Color, bool bForceLDR)
 	{
 		if (bForceLDR)
diff --git a/Engine/Plugins/Enterprise/GLTFExporter/MigratedFromGit/Source/GLTFExporter/Private/Converters/GLTFMeshSection.cpp b/Engine/Plugins/Enterprise/GLTFExporter/MigratedFromGit/Source/GLTFExporter/Private/Converters/GLTFMeshSection.cpp
--- a/Engine/Plugins/Enterprise/GLTFExporter/MigratedFromGit/Source/GLTFExporter/Private/Converters/GLTFMeshSection.cpp
+++ b/Engine/Plugins/Enterprise/GLTFExporter/MigratedFromGit/Source/GLTFExporter/Private/Converters/GLTFMeshSection.cpp
@@ -1,35 +1,38 @@
 // Copyright Epic Games, Inc. All Rights Reserved.
 
 #include "Converters/GLTFMeshSection.h"
+#include "Converters/GLTFBufferAdapter.h"
 #include "Rendering/MultiSizeIndexContainer.h"
 #include "Rendering/SkeletalMeshRenderData.h"
 #include "Algo/MaxElement.h"
 
 FGLTFMeshSection::FGLTFMeshSection(const FStaticMeshLODResources* MeshLOD, const FGLTFIndexArray& SectionIndices)
 {
-	const FRawStaticIndexBuffer& SourceBuffer = MeshLOD->IndexBuffer;
-	if (SourceBuffer.Is32Bit())
+	const TUniquePtr<IGLTFBufferAdapter> SourceBuffer = IGLTFBufferAdapter::GetIndices(&MeshLOD->IndexBuffer);
+	const uint8* SourceData = SourceBuffer->GetData();
+
+	if (MeshLOD->IndexBuffer.Is32Bit())
 	{
-		Init(MeshLOD->Sections, SectionIndices, SourceBuffer.AccessStream32());
+		Init(MeshLOD->Sections, SectionIndices, reinterpret_cast<const uint32*>(SourceData));
 	}
 	else
 	{
-		Init(MeshLOD->Sections, SectionIndices, SourceBuffer.AccessStream16());
+		Init(MeshLOD->Sections, SectionIndices, reinterpret_cast<const uint16*>(SourceData));
 	}
 }
 
 FGLTFMeshSection::FGLTFMeshSection(const FSkeletalMeshLODRenderData* MeshLOD, const FGLTFIndexArray& SectionIndices)
 {
-	const FMultiSizeIndexContainer& SourceBuffer = MeshLOD->MultiSizeIndexContainer;
-	const void* SourceData = const_cast<FRawStaticIndexBuffer16or32Interface*>(SourceBuffer.GetIndexBuffer())->GetPointerTo(0);
+	const TUniquePtr<IGLTFBufferAdapter> SourceBuffer = IGLTFBufferAdapter::GetIndices(MeshLOD->MultiSizeIndexContainer.GetIndexBuffer());
+	const uint8* SourceData = SourceBuffer->GetData();
 
-	if (SourceBuffer.GetDataTypeSize() != sizeof(uint16))
+	if (MeshLOD->MultiSizeIndexContainer.GetDataTypeSize() != sizeof(uint16))
 	{
-		Init(MeshLOD->RenderSections, SectionIndices, static_cast<const uint32*>(SourceData));
+		Init(MeshLOD->RenderSections, SectionIndices, reinterpret_cast<const uint32*>(SourceData));
 	}
 	else
 	{
-		Init(MeshLOD->RenderSections, SectionIndices, static_cast<const uint16*>(SourceData));
+		Init(MeshLOD->RenderSections, SectionIndices, reinterpret_cast<const uint16*>(SourceData));
 	}
 }
 
