[GLTFExporter] fix: use GLTFBufferAdapter when checking for vertex colors and validating vertex normals & tangents

commit a68f860ecc8135896b90eac0124a39035b199852
Date: Thu, 11 Nov 2021 19:15:04 +0100

    fix: use GLTFBufferAdapter when checking for vertex colors and validating vertex normals & tangents
    (cherry picked from commit 0fcbb7eb7677c562087e8a6917022be02e93cf4b​)

#skipci

diff --git a/Engine/Plugins/Enterprise/GLTFExporter/MigratedFromGit/Source/GLTFExporter/Private/Tasks/GLTFMeshTasks.cpp b/Engine/Plugins/Enterprise/GLTFExporter/MigratedFromGit/Source/GLTFExporter/Private/Tasks/GLTFMeshTasks.cpp
--- a/Engine/Plugins/Enterprise/GLTFExporter/MigratedFromGit/Source/GLTFExporter/Private/Tasks/GLTFMeshTasks.cpp
+++ b/Engine/Plugins/Enterprise/GLTFExporter/MigratedFromGit/Source/GLTFExporter/Private/Tasks/GLTFMeshTasks.cpp
@@ -1,33 +1,27 @@
 // Copyright Epic Games, Inc. All Rights Reserved.
 
 #include "Tasks/GLTFMeshTasks.h"
-#include "Converters/GLTFConverterUtility.h"
 #include "Converters/GLTFMeshUtility.h"
+#include "Converters/GLTFBufferAdapter.h"
 #include "Builders/GLTFConvertBuilder.h"
 #include "Rendering/SkeletalMeshRenderData.h"
 
 namespace
 {
-	template <typename TangentVectorType>
-	void CheckTangentVectors(const FStaticMeshVertexBuffer* VertexBuffer, bool& bOutZeroNormals, bool& bOutZeroTangents)
+	template <typename VectorType>
+	void CheckTangentVectors(const void* SourceData, uint32 VertexCount, bool& bOutZeroNormals, bool& bOutZeroTangents)
 	{
 		bool bZeroNormals = false;
 		bool bZeroTangents = false;
 
-		typedef TStaticMeshVertexTangentDatum<TangentVectorType> VertexTangentType;
+		typedef TStaticMeshVertexTangentDatum<VectorType> VertexTangentType;
+		const VertexTangentType* VertexTangents = static_cast<const VertexTangentType*>(SourceData);
 
-		const void* TangentData = const_cast<FStaticMeshVertexBuffer*>(VertexBuffer)->GetTangentData();
-		const VertexTangentType* VertexTangents = TangentData != nullptr ? static_cast<const VertexTangentType*>(TangentData) : nullptr;
-
-		if (VertexTangents != nullptr)
+		for (uint32 VertexIndex = 0; VertexIndex < VertexCount; ++VertexIndex)
 		{
-			const uint32 VertexCount = VertexBuffer->GetNumVertices();
-			for (uint32 VertexIndex = 0; VertexIndex < VertexCount; ++VertexIndex)
-			{
-				const VertexTangentType& VertexTangent = VertexTangents[VertexIndex];
-				bZeroNormals |= VertexTangent.TangentZ.ToFVector().IsNearlyZero();
-				bZeroTangents |= VertexTangent.TangentX.ToFVector().IsNearlyZero();
-			}
+			const VertexTangentType& VertexTangent = VertexTangents[VertexIndex];
+			bZeroNormals |= VertexTangent.TangentZ.ToFVector().IsNearlyZero();
+			bZeroTangents |= VertexTangent.TangentX.ToFVector().IsNearlyZero();
 		}
 
 		bOutZeroNormals = bZeroNormals;
@@ -41,16 +35,19 @@ namespace
 			return;
 		}
 
+		const TUniquePtr<IGLTFBufferAdapter> SourceBuffer = IGLTFBufferAdapter::GetTangents(VertexBuffer);
+		const uint8* SourceData = SourceBuffer->GetData();
+		const uint32 VertexCount= VertexBuffer->GetNumVertices();
 		bool bZeroNormals;
 		bool bZeroTangents;
 
 		if (VertexBuffer->GetUseHighPrecisionTangentBasis())
 		{
-			CheckTangentVectors<FPackedRGBA16N>(VertexBuffer, bZeroNormals, bZeroTangents);
+			CheckTangentVectors<FPackedRGBA16N>(SourceData, VertexCount, bZeroNormals, bZeroTangents);
 		}
 		else
 		{
-			CheckTangentVectors<FPackedNormal>(VertexBuffer, bZeroNormals, bZeroTangents);
+			CheckTangentVectors<FPackedNormal>(SourceData, VertexCount, bZeroNormals, bZeroTangents);
 		}
 
 		if (bZeroNormals)
@@ -70,11 +67,14 @@ namespace
 
 	bool HasVertexColors(const FColorVertexBuffer* VertexBuffer)
 	{
+		const TUniquePtr<IGLTFBufferAdapter> SourceBuffer = IGLTFBufferAdapter::GetColors(VertexBuffer);
+		const uint8* SourceData = SourceBuffer->GetData();
 		const uint32 VertexCount = VertexBuffer->GetNumVertices();
+		const uint32 Stride = VertexBuffer->GetStride();
 
 		for (uint32 VertexIndex = 0; VertexIndex < VertexCount; VertexIndex++)
 		{
-			const FColor& Color = VertexBuffer->VertexColor(VertexIndex);
+			const FColor& Color = *reinterpret_cast<const FColor*>(SourceData + Stride * VertexIndex);
 			if (Color != FColor::White)
 			{
 				return true;
