[GLTFExporter] feat: improve detection of blueprints to work in runtime

commit 1e12e0a21c8d125eb1f6bffdd9ab20d2a853feeb
Date: Thu, 11 Nov 2021 19:34:28 +0100

    feat: improve detection of blueprints to work in runtime
    (cherry picked from commit 98671046d1165a1769fd427736ed63f4f858228e​)

#skipci

diff --git a/Engine/Plugins/Enterprise/GLTFExporter/MigratedFromGit/Source/GLTFExporter/Private/Converters/GLTFActorUtility.cpp b/Engine/Plugins/Enterprise/GLTFExporter/MigratedFromGit/Source/GLTFExporter/Private/Converters/GLTFActorUtility.cpp
--- a/Engine/Plugins/Enterprise/GLTFExporter/MigratedFromGit/Source/GLTFExporter/Private/Converters/GLTFActorUtility.cpp
+++ b/Engine/Plugins/Enterprise/GLTFExporter/MigratedFromGit/Source/GLTFExporter/Private/Converters/GLTFActorUtility.cpp
@@ -9,19 +9,25 @@ bool FGLTFActorUtility::IsRootActor(const AActor* Actor, bool bSelectedOnly)
 	return ParentActor == nullptr || (bSelectedOnly && !ParentActor->IsSelected());
 }
 
-UBlueprint* FGLTFActorUtility::GetBlueprintFromActor(const AActor* Actor)
+FString FGLTFActorUtility::GetBlueprintPath(const AActor* Actor)
 {
-	return UBlueprint::GetBlueprintFromClass(Actor->GetClass());
+	UClass* Class = Actor->GetClass();
+	if (Class != nullptr && Class->HasAnyClassFlags(CLASS_CompiledFromBlueprint))
+	{
+		return Class->GetPathName();
+	}
+
+	return TEXT("");
 }
 
-bool FGLTFActorUtility::IsSkySphereBlueprint(const UBlueprint* Blueprint)
+bool FGLTFActorUtility::IsSkySphereBlueprint(const FString& Path)
 {
 	// TODO: what if a blueprint inherits BP_Sky_Sphere?
-	return Blueprint != nullptr && Blueprint->GetPathName().Equals(TEXT("/Engine/EngineSky/BP_Sky_Sphere.BP_Sky_Sphere"));
+	return Path.Equals(TEXT("/Engine/EngineSky/BP_Sky_Sphere.BP_Sky_Sphere_C"));
 }
 
-bool FGLTFActorUtility::IsHDRIBackdropBlueprint(const UBlueprint* Blueprint)
+bool FGLTFActorUtility::IsHDRIBackdropBlueprint(const FString& Path)
 {
 	// TODO: what if a blueprint inherits HDRIBackdrop?
-	return Blueprint != nullptr && Blueprint->GetPathName().Equals(TEXT("/HDRIBackdrop/Blueprints/HDRIBackdrop.HDRIBackdrop"));
+	return Path.Equals(TEXT("/HDRIBackdrop/Blueprints/HDRIBackdrop.HDRIBackdrop_C"));
 }
diff --git a/Engine/Plugins/Enterprise/GLTFExporter/MigratedFromGit/Source/GLTFExporter/Private/Converters/GLTFActorUtility.h b/Engine/Plugins/Enterprise/GLTFExporter/MigratedFromGit/Source/GLTFExporter/Private/Converters/GLTFActorUtility.h
--- a/Engine/Plugins/Enterprise/GLTFExporter/MigratedFromGit/Source/GLTFExporter/Private/Converters/GLTFActorUtility.h
+++ b/Engine/Plugins/Enterprise/GLTFExporter/MigratedFromGit/Source/GLTFExporter/Private/Converters/GLTFActorUtility.h
@@ -8,11 +8,11 @@ struct FGLTFActorUtility
 {
 	static bool IsRootActor(const AActor* Actor, bool bSelectedOnly);
 
-	static UBlueprint* GetBlueprintFromActor(const AActor* Actor);
+	static FString GetBlueprintPath(const AActor* Actor);
 
-	static bool IsSkySphereBlueprint(const UBlueprint* Blueprint);
+	static bool IsSkySphereBlueprint(const FString& Path);
 
-	static bool IsHDRIBackdropBlueprint(const UBlueprint* Blueprint);
+	static bool IsHDRIBackdropBlueprint(const FString& Path);
 
 	template <class ValueType>
 	static bool TryGetPropertyValue(const UObject* Object, const TCHAR* PropertyName, ValueType& Value)
diff --git a/Engine/Plugins/Enterprise/GLTFExporter/MigratedFromGit/Source/GLTFExporter/Private/Converters/GLTFBackdropConverters.cpp b/Engine/Plugins/Enterprise/GLTFExporter/MigratedFromGit/Source/GLTFExporter/Private/Converters/GLTFBackdropConverters.cpp
--- a/Engine/Plugins/Enterprise/GLTFExporter/MigratedFromGit/Source/GLTFExporter/Private/Converters/GLTFBackdropConverters.cpp
+++ b/Engine/Plugins/Enterprise/GLTFExporter/MigratedFromGit/Source/GLTFExporter/Private/Converters/GLTFBackdropConverters.cpp
@@ -8,8 +8,7 @@
 
 FGLTFJsonBackdropIndex FGLTFBackdropConverter::Convert(const AActor* BackdropActor)
 {
-	const UBlueprint* Blueprint = FGLTFActorUtility::GetBlueprintFromActor(BackdropActor);
-	if (!FGLTFActorUtility::IsHDRIBackdropBlueprint(Blueprint))
+	if (!FGLTFActorUtility::IsHDRIBackdropBlueprint(FGLTFActorUtility::GetBlueprintPath(BackdropActor)))
 	{
 		return FGLTFJsonBackdropIndex(INDEX_NONE);
 	}
diff --git a/Engine/Plugins/Enterprise/GLTFExporter/MigratedFromGit/Source/GLTFExporter/Private/Converters/GLTFNodeConverters.cpp b/Engine/Plugins/Enterprise/GLTFExporter/MigratedFromGit/Source/GLTFExporter/Private/Converters/GLTFNodeConverters.cpp
--- a/Engine/Plugins/Enterprise/GLTFExporter/MigratedFromGit/Source/GLTFExporter/Private/Converters/GLTFNodeConverters.cpp
+++ b/Engine/Plugins/Enterprise/GLTFExporter/MigratedFromGit/Source/GLTFExporter/Private/Converters/GLTFNodeConverters.cpp
@@ -18,8 +18,10 @@ FGLTFJsonNodeIndex FGLTFActorConverter::Convert(const AActor* Actor)
 	const USceneComponent* RootComponent = Actor->GetRootComponent();
 	const FGLTFJsonNodeIndex RootNodeIndex = Builder.GetOrAddNode(RootComponent);
 
-	const UBlueprint* Blueprint = FGLTFActorUtility::GetBlueprintFromActor(Actor);
-	if (FGLTFActorUtility::IsSkySphereBlueprint(Blueprint))
+	// TODO: process all components since any component can be attached to any other component in runtime
+
+	const FString BlueprintPath = FGLTFActorUtility::GetBlueprintPath(Actor);
+	if (FGLTFActorUtility::IsSkySphereBlueprint(BlueprintPath))
 	{
 		if (Builder.ExportOptions->bExportSkySpheres)
 		{
@@ -27,7 +29,7 @@ FGLTFJsonNodeIndex FGLTFActorConverter::Convert(const AActor* Actor)
 			RootNode.SkySphere = Builder.GetOrAddSkySphere(Actor);
 		}
 	}
-	else if (FGLTFActorUtility::IsHDRIBackdropBlueprint(Blueprint))
+	else if (FGLTFActorUtility::IsHDRIBackdropBlueprint(BlueprintPath))
 	{
 		if (Builder.ExportOptions->bExportHDRIBackdrops)
 		{
diff --git a/Engine/Plugins/Enterprise/GLTFExporter/MigratedFromGit/Source/GLTFExporter/Private/Converters/GLTFSkySphereConverters.cpp b/Engine/Plugins/Enterprise/GLTFExporter/MigratedFromGit/Source/GLTFExporter/Private/Converters/GLTFSkySphereConverters.cpp
--- a/Engine/Plugins/Enterprise/GLTFExporter/MigratedFromGit/Source/GLTFExporter/Private/Converters/GLTFSkySphereConverters.cpp
+++ b/Engine/Plugins/Enterprise/GLTFExporter/MigratedFromGit/Source/GLTFExporter/Private/Converters/GLTFSkySphereConverters.cpp
@@ -9,8 +9,7 @@
 
 FGLTFJsonSkySphereIndex FGLTFSkySphereConverter::Convert(const AActor* SkySphereActor)
 {
-	const UBlueprint* Blueprint = FGLTFActorUtility::GetBlueprintFromActor(SkySphereActor);
-	if (!FGLTFActorUtility::IsSkySphereBlueprint(Blueprint))
+	if (!FGLTFActorUtility::IsSkySphereBlueprint(FGLTFActorUtility::GetBlueprintPath(SkySphereActor)))
 	{
 		return FGLTFJsonSkySphereIndex(INDEX_NONE);
 	}
