[GLTFExporter] feat: ignore special actors like pawn during export

commit 4b85be3fc1a714549ca1eab1881c7ebf23ede858
Date: Thu, 11 Nov 2021 19:37:08 +0100

    feat: ignore special actors like pawn during export
    (cherry picked from commit cd30be0bbf7dcafbbe4c0519fa8dba3c9d13ab8d​)

#skipci

diff --git a/Engine/Plugins/Enterprise/GLTFExporter/MigratedFromGit/Source/GLTFExporter/Private/Converters/GLTFActorUtility.cpp b/Engine/Plugins/Enterprise/GLTFExporter/MigratedFromGit/Source/GLTFExporter/Private/Converters/GLTFActorUtility.cpp
--- a/Engine/Plugins/Enterprise/GLTFExporter/MigratedFromGit/Source/GLTFExporter/Private/Converters/GLTFActorUtility.cpp
+++ b/Engine/Plugins/Enterprise/GLTFExporter/MigratedFromGit/Source/GLTFExporter/Private/Converters/GLTFActorUtility.cpp
@@ -2,6 +2,8 @@
 
 #include "Converters/GLTFActorUtility.h"
 #include "Engine/Blueprint.h"
+#include "Actors/GLTFHotspotActor.h"
+#include "LevelSequenceActor.h"
 
 bool FGLTFActorUtility::IsRootActor(const AActor* Actor, bool bSelectedOnly)
 {
@@ -9,6 +11,17 @@ bool FGLTFActorUtility::IsRootActor(const AActor* Actor, bool bSelectedOnly)
 	return ParentActor == nullptr || (bSelectedOnly && !ParentActor->IsSelected());
 }
 
+bool FGLTFActorUtility::IsGenericActor(const AActor* Actor)
+{
+	const FString BlueprintPath = GetBlueprintPath(Actor);
+	if (IsSkySphereBlueprint(BlueprintPath)) return false;
+	if (IsHDRIBackdropBlueprint(BlueprintPath)) return false;
+	if (Actor->IsA<ALevelSequenceActor>()) return false;
+	if (Actor->IsA<AGLTFHotspotActor>()) return false;
+	if (Actor->IsA<APawn>()) return false;
+	return true;
+}
+
 FString FGLTFActorUtility::GetBlueprintPath(const AActor* Actor)
 {
 	UClass* Class = Actor->GetClass();
diff --git a/Engine/Plugins/Enterprise/GLTFExporter/MigratedFromGit/Source/GLTFExporter/Private/Converters/GLTFActorUtility.h b/Engine/Plugins/Enterprise/GLTFExporter/MigratedFromGit/Source/GLTFExporter/Private/Converters/GLTFActorUtility.h
--- a/Engine/Plugins/Enterprise/GLTFExporter/MigratedFromGit/Source/GLTFExporter/Private/Converters/GLTFActorUtility.h
+++ b/Engine/Plugins/Enterprise/GLTFExporter/MigratedFromGit/Source/GLTFExporter/Private/Converters/GLTFActorUtility.h
@@ -8,6 +8,8 @@ struct FGLTFActorUtility
 {
 	static bool IsRootActor(const AActor* Actor, bool bSelectedOnly);
 
+	static bool IsGenericActor(const AActor* Actor);
+
 	static FString GetBlueprintPath(const AActor* Actor);
 
 	static bool IsSkySphereBlueprint(const FString& Path);
diff --git a/Engine/Plugins/Enterprise/GLTFExporter/MigratedFromGit/Source/GLTFExporter/Private/Converters/GLTFNodeConverters.cpp b/Engine/Plugins/Enterprise/GLTFExporter/MigratedFromGit/Source/GLTFExporter/Private/Converters/GLTFNodeConverters.cpp
--- a/Engine/Plugins/Enterprise/GLTFExporter/MigratedFromGit/Source/GLTFExporter/Private/Converters/GLTFNodeConverters.cpp
+++ b/Engine/Plugins/Enterprise/GLTFExporter/MigratedFromGit/Source/GLTFExporter/Private/Converters/GLTFNodeConverters.cpp
@@ -117,7 +117,7 @@ FGLTFJsonNodeIndex FGLTFComponentConverter::Convert(const USceneComponent* Scene
 	// TODO: don't export invisible components unless visibility is variable due to variant sets
 
 	// TODO: should hidden in game be configurable like this?
-	if (!SceneComponent->bHiddenInGame || Builder.ExportOptions->bExportHiddenInGame)
+	if ((!SceneComponent->bHiddenInGame || Builder.ExportOptions->bExportHiddenInGame) && FGLTFActorUtility::IsGenericActor(Owner))
 	{
 		if (const UStaticMeshComponent* StaticMeshComponent = Cast<UStaticMeshComponent>(SceneComponent))
 		{
