[GLTFExporter] refactor: merge GLTFExporterRuntime module with GLTFExporter module

commit eab176cd7722a6f7709d92d70d606a93a1793c08
Date: Thu, 11 Nov 2021 19:42:04 +0100

    refactor: merge GLTFExporterRuntime module with GLTFExporter module
    (cherry picked from commit e52d146d63e54999d257a6a9e88835f8d2f78aa2​)

#skipci

diff --git a/Engine/Plugins/Enterprise/GLTFExporter/MigratedFromGit/Config/BaseGLTFExporter.ini b/Engine/Plugins/Enterprise/GLTFExporter/MigratedFromGit/Config/BaseGLTFExporter.ini
new file mode 100644
--- /dev/null
+++ b/Engine/Plugins/Enterprise/GLTFExporter/MigratedFromGit/Config/BaseGLTFExporter.ini
@@ -0,0 +1,2 @@
+[CoreRedirects]
++PackageRedirects=(OldName="/Script/GLTFExporterRuntime",NewName="/Script/GLTFExporter")
diff --git a/Engine/Plugins/Enterprise/GLTFExporter/MigratedFromGit/Config/DefaultGLTFExporter.ini b/Engine/Plugins/Enterprise/GLTFExporter/MigratedFromGit/Config/DefaultGLTFExporter.ini
new file mode 100644
--- /dev/null
+++ b/Engine/Plugins/Enterprise/GLTFExporter/MigratedFromGit/Config/DefaultGLTFExporter.ini
@@ -0,0 +1,2 @@
+[CoreRedirects]
++PackageRedirects=(OldName="/Script/GLTFExporterRuntime",NewName="/Script/GLTFExporter")
diff --git a/Engine/Plugins/Enterprise/GLTFExporter/MigratedFromGit/GLTFExporter.uplugin b/Engine/Plugins/Enterprise/GLTFExporter/MigratedFromGit/GLTFExporter.uplugin
--- a/Engine/Plugins/Enterprise/GLTFExporter/MigratedFromGit/GLTFExporter.uplugin
+++ b/Engine/Plugins/Enterprise/GLTFExporter/MigratedFromGit/GLTFExporter.uplugin
@@ -20,12 +20,6 @@
 			"LoadingPhase": "PostConfigInit",
 			"WhitelistPlatforms": [ "Win64", "Win32", "Mac" ]
 		},
-		{
-			"Name": "GLTFExporterRuntime",
-			"Type": "Runtime",
-			"LoadingPhase": "Default",
-			"WhitelistPlatforms": [ "Win64", "Win32", "Mac" ]
-		},
 		{
 			"Name": "GLTFMaterialAnalyzer",
 			"Type": "Editor",
diff --git a/Engine/Plugins/Enterprise/GLTFExporter/MigratedFromGit/Source/GLTFExporter/GLTFExporter.Build.cs b/Engine/Plugins/Enterprise/GLTFExporter/MigratedFromGit/Source/GLTFExporter/GLTFExporter.Build.cs
--- a/Engine/Plugins/Enterprise/GLTFExporter/MigratedFromGit/Source/GLTFExporter/GLTFExporter.Build.cs
+++ b/Engine/Plugins/Enterprise/GLTFExporter/MigratedFromGit/Source/GLTFExporter/GLTFExporter.Build.cs
@@ -7,8 +7,8 @@ namespace UnrealBuildTool.Rules
 		public GLTFExporter(ReadOnlyTargetRules Target) : base(Target)
 		{
 		    // NOTE: ugly hack to access plugin info (should propose change to engine)
-		    var BindingFlags = System.Reflection.BindingFlags.Instance | System.Reflection.BindingFlags.NonPublic;
-		    var FieldInfo = typeof(ModuleRules).GetField("Plugin", BindingFlags);
+			var BindingFlags = System.Reflection.BindingFlags.Instance | System.Reflection.BindingFlags.NonPublic;
+			var FieldInfo = typeof(ModuleRules).GetField("Plugin", BindingFlags);
 			var Plugin = FieldInfo != null ? FieldInfo.GetValue(this) as PluginInfo : null;
 
 			if (Plugin != null && Plugin.Descriptor != null)
@@ -54,13 +54,6 @@ namespace UnrealBuildTool.Rules
 					}
 					);
 			}
-
-			PublicDependencyModuleNames.AddRange(
-				new string[]
-				{
-					"GLTFExporterRuntime"
-				}
-				);
 		}
 	}
 }
diff --git a/Engine/Plugins/Enterprise/GLTFExporter/MigratedFromGit/Source/GLTFExporter/Private/Actors/GLTFCameraActor.cpp b/Engine/Plugins/Enterprise/GLTFExporter/MigratedFromGit/Source/GLTFExporter/Private/Actors/GLTFCameraActor.cpp
new file mode 100644
--- /dev/null
+++ b/Engine/Plugins/Enterprise/GLTFExporter/MigratedFromGit/Source/GLTFExporter/Private/Actors/GLTFCameraActor.cpp
@@ -0,0 +1,265 @@
+// Copyright Epic Games, Inc. All Rights Reserved.
+
+#include "Actors/GLTFCameraActor.h"
+#include "Components/InputComponent.h"
+
+namespace
+{
+	// Scales to convert from the export-friendly sensitivity-values stored in our properties
+	// to values that we can use when processing axis-values (to get similar results as in the viewer).
+	const float RotationSensitivityScale = 16.667f;
+	const float DollySensitivityScale = 0.1f;
+}
+
+AGLTFCameraActor::AGLTFCameraActor(const FObjectInitializer& ObjectInitializer)
+	: Super(ObjectInitializer)
+	, Mode(EGLTFCameraControlMode::FreeLook)
+	, Target(nullptr)
+	, PitchAngleMin(-90.0f)
+	, PitchAngleMax(90.0f)
+	, YawAngleMin(0.0f)
+	, YawAngleMax(360.0f)
+	, DistanceMin(100.0f)
+	, DistanceMax(1000.0f)
+	, DollyDuration(0.2f)
+	, DollySensitivity(0.5f)
+	, RotationInertia(0.1f)
+	, RotationSensitivity(0.3f)
+	, Distance(0.0f)
+	, Pitch(0.0f)
+	, Yaw(0.0f)
+	, TargetDistance(0.0f)
+	, TargetPitch(0.0f)
+	, TargetYaw(0.0f)
+	, DollyTime(0.0f)
+	, DollyStartDistance(0.0f)
+{
+	PrimaryActorTick.bCanEverTick = true;
+}
+
+#if WITH_EDITOR
+void AGLTFCameraActor::PostEditChangeProperty(FPropertyChangedEvent& PropertyChangedEvent)
+{
+	Super::PostEditChangeProperty(PropertyChangedEvent);
+
+	const FProperty* PropertyThatChanged = PropertyChangedEvent.Property;
+
+	if (PropertyThatChanged)
+	{
+		const FName PropertyName = PropertyThatChanged->GetFName();
+
+		if (PropertyName == GET_MEMBER_NAME_CHECKED(ThisClass, Target))
+		{
+			if (Target == this)
+			{
+				Target = nullptr;
+				// TODO: don't use LogTemp
+				UE_LOG(LogTemp, Warning, TEXT("The camera target must not be the camera's own actor"));
+			}
+		}
+		else if (PropertyName == GET_MEMBER_NAME_CHECKED(ThisClass, PitchAngleMin))
+		{
+			PitchAngleMin = FMath::Clamp(PitchAngleMin, -90.0f, 90.0f);
+		}
+		else if (PropertyName == GET_MEMBER_NAME_CHECKED(ThisClass, PitchAngleMax))
+		{
+			PitchAngleMax = FMath::Clamp(PitchAngleMax, -90.0f, 90.0f);
+		}
+		else if (PropertyName == GET_MEMBER_NAME_CHECKED(ThisClass, YawAngleMin))
+		{
+			YawAngleMin = FMath::Clamp(
+				YawAngleMin,
+				FMath::Max(-360.0f, YawAngleMax - 360.0f),
+				FMath::Min(360.0f, YawAngleMax));
+		}
+		else if (PropertyName == GET_MEMBER_NAME_CHECKED(ThisClass, YawAngleMax))
+		{
+			YawAngleMax = FMath::Clamp(
+				YawAngleMax,
+				FMath::Max(-360.0f, YawAngleMin),
+				FMath::Min(360.0f, YawAngleMin + 360.0f));
+		}
+	}
+}
+#endif // WITH_EDITOR
+
+void AGLTFCameraActor::BeginPlay()
+{
+	Super::BeginPlay();
+
+	if (Mode == EGLTFCameraControlMode::FreeLook)
+	{
+		const FRotator Rotation = GetActorRotation();
+
+		Pitch = ClampPitch(Rotation.Pitch);
+		Yaw = ClampYaw(Rotation.Yaw);
+		TargetPitch = Pitch;
+		TargetYaw = Yaw;
+	}
+	else if (Mode == EGLTFCameraControlMode::Orbital)
+	{
+		const FVector FocusPosition = GetFocusPosition();
+
+		// Ensure that the camera is initially aimed at the focus-position
+		SetActorRotation(GetLookAtRotation(FocusPosition));
+
+		const FVector Position = GetActorLocation();
+		const FRotator Rotation = GetActorRotation();
+
+		// Calculate values based on the current location and orientation
+		Distance = ClampDistance((FocusPosition - Position).Size());
+		Pitch = ClampPitch(Rotation.Pitch);
+		Yaw = ClampYaw(Rotation.Yaw);
+		TargetDistance = Distance;
+		TargetPitch = Pitch;
+		TargetYaw = Yaw;
+	}
+	else
+	{
+		checkNoEntry();
+	}
+
+	if (InputComponent)
+	{
+		InputComponent->BindAxisKey(EKeys::MouseX, this, &AGLTFCameraActor::OnMouseX);
+		InputComponent->BindAxisKey(EKeys::MouseY, this, &AGLTFCameraActor::OnMouseY);
+		InputComponent->BindAxisKey(EKeys::MouseWheelAxis, this, &AGLTFCameraActor::OnMouseWheelAxis);
+	}
+}
+
+void AGLTFCameraActor::Tick(float DeltaSeconds)
+{
+	Super::Tick(DeltaSeconds);
+
+	if (Mode == EGLTFCameraControlMode::FreeLook)
+	{
+		const float Alpha = (RotationInertia == 0.0f) ? 1.0f : FMath::Min(DeltaSeconds / RotationInertia, 1.0f);
+		Yaw = FMath::Lerp(Yaw, TargetYaw, Alpha);
+		Pitch = FMath::Lerp(Pitch, TargetPitch, Alpha);
+
+		SetActorRotation(FQuat::MakeFromEuler(FVector(0.0f, Pitch, Yaw)));
+	}
+	else if (Mode == EGLTFCameraControlMode::Orbital)
+	{
+		if (DollyTime != 0.0f)
+		{
+			DollyTime = FMath::Max(DollyTime - DeltaSeconds, 0.0f);
+			Distance = FMath::InterpEaseInOut(DollyStartDistance, TargetDistance, (DollyDuration - DollyTime) / DollyDuration, 1.2f);
+		}
+
+		const float Alpha = (RotationInertia == 0.0f) ? 1.0f : FMath::Min(DeltaSeconds / RotationInertia, 1.0f);
+		Yaw = FMath::Lerp(Yaw, TargetYaw, Alpha);
+		Pitch = FMath::Lerp(Pitch, TargetPitch, Alpha);
+
+		const FTransform FocusTransform = FTransform(GetFocusPosition());
+		const FTransform DollyTransform = FTransform(-FVector::ForwardVector * Distance);
+		const FTransform RotationTransform = FTransform(FQuat::MakeFromEuler(FVector(0.0f, Pitch, Yaw)));
+		const FTransform ResultTransform = DollyTransform * RotationTransform * FocusTransform;
+
+		SetActorTransform(ResultTransform);
+	}
+}
+
+void AGLTFCameraActor::PreInitializeComponents()
+{
+	AutoReceiveInput = static_cast<EAutoReceiveInput::Type>(GetAutoActivatePlayerIndex() + 1);
+
+	Super::PreInitializeComponents();
+}
+
+void AGLTFCameraActor::PostActorCreated()
+{
+	SetAutoActivateForPlayer(EAutoReceiveInput::Player0);
+}
+
+void AGLTFCameraActor::OnMouseX(float AxisValue)
+{
+	TargetYaw = ClampYaw(TargetYaw + AxisValue * RotationSensitivity * RotationSensitivityScale);
+
+	if (!UsesYawLimits())
+	{
+		// Prevent TargetYaw from reaching excessive values by keeping it in the range (-360, 360)
+		const float ClampedTargetYaw = FMath::Fmod(TargetYaw, 360.f);
+
+		Yaw += ClampedTargetYaw - TargetYaw;
+		TargetYaw = ClampedTargetYaw;
+	}
+}
+
+void AGLTFCameraActor::OnMouseY(float AxisValue)
+{
+	TargetPitch = ClampPitch(TargetPitch + AxisValue * RotationSensitivity * RotationSensitivityScale);
+}
+
+void AGLTFCameraActor::OnMouseWheelAxis(float AxisValue)
+{
+	if (Mode == EGLTFCameraControlMode::Orbital)
+	{
+		if (!FMath::IsNearlyZero(AxisValue))
+		{
+			const float DeltaDistance = -AxisValue * (TargetDistance * DollySensitivity * DollySensitivityScale);
+
+			DollyTime = DollyDuration;
+			TargetDistance = ClampDistance(TargetDistance + DeltaDistance);
+			DollyStartDistance = Distance;
+		}
+	}
+}
+
+float AGLTFCameraActor::ClampDistance(float Value) const
+{
+	return FMath::Clamp(Value, DistanceMin, DistanceMax);
+}
+
+float AGLTFCameraActor::ClampPitch(float Value) const
+{
+	return FMath::Clamp(Value, PitchAngleMin, PitchAngleMax);
+}
+
+float AGLTFCameraActor::ClampYaw(float Value) const
+{
+	return UsesYawLimits() ? FMath::Clamp(Value, YawAngleMin, YawAngleMax) : Value;
+}
+
+void AGLTFCameraActor::RemoveInertia()
+{
+	Yaw = TargetYaw;
+	Pitch = TargetPitch;
+	Distance = TargetDistance;
+}
+
+FRotator AGLTFCameraActor::GetLookAtRotation(const FVector TargetPosition) const
+{
+	const FVector EyePosition = GetActorLocation();
+
+	return FRotationMatrix::MakeFromXZ(TargetPosition - EyePosition, FVector::UpVector).Rotator();
+}
+
+FVector AGLTFCameraActor::GetFocusPosition() const
+{
+	return this->Target != nullptr ? this->Target->GetActorLocation() : FVector::ZeroVector;
+}
+
+bool AGLTFCameraActor::SetAutoActivateForPlayer(const EAutoReceiveInput::Type Player)
+{
+	// TODO: remove hack by adding proper API access to ACameraActor (or by other means)
+	FProperty* Property = GetClass()->FindPropertyByName(TEXT("AutoActivateForPlayer"));
+	if (Property == nullptr)
+	{
+		return false;
+	}
+
+	TEnumAsByte<EAutoReceiveInput::Type>* ValuePtr = Property->ContainerPtrToValuePtr<TEnumAsByte<EAutoReceiveInput::Type>>(this);
+	if (ValuePtr == nullptr)
+	{
+		return false;
+	}
+
+	*ValuePtr = Player;
+	return true;
+}
+
+bool AGLTFCameraActor::UsesYawLimits() const
+{
+	return !FMath::IsNearlyEqual(YawAngleMax - YawAngleMin, 360.0f);
+}
diff --git a/Engine/Plugins/Enterprise/GLTFExporter/MigratedFromGit/Source/GLTFExporter/Private/Actors/GLTFHotspotActor.cpp b/Engine/Plugins/Enterprise/GLTFExporter/MigratedFromGit/Source/GLTFExporter/Private/Actors/GLTFHotspotActor.cpp
new file mode 100644
--- /dev/null
+++ b/Engine/Plugins/Enterprise/GLTFExporter/MigratedFromGit/Source/GLTFExporter/Private/Actors/GLTFHotspotActor.cpp
@@ -0,0 +1,517 @@
+// Copyright Epic Games, Inc. All Rights Reserved.
+
+#include "Actors/GLTFHotspotActor.h"
+#include "Animation/SkeletalMeshActor.h"
+#include "Animation/AnimSequence.h"
+#include "LevelSequenceActor.h"
+#include "LevelSequence.h"
+#include "MovieSceneTimeHelpers.h"
+#include "Components/SphereComponent.h"
+#include "Components/MaterialBillboardComponent.h"
+#include "Kismet/GameplayStatics.h"
+#include "UObject/ConstructorHelpers.h"
+#include "Animation/AnimSingleNodeInstance.h"
+#include "Materials/MaterialInstanceDynamic.h"
+
+#if WITH_EDITOR
+#include "LevelEditor.h"
+#include "Modules/ModuleManager.h"
+#include "SLevelViewport.h"
+#include "Slate/SceneViewport.h"
+#endif
+
+#if (ENGINE_MAJOR_VERSION > 4 || ENGINE_MINOR_VERSION >= 26)
+using namespace UE;
+#endif
+
+DEFINE_LOG_CATEGORY_STATIC(LogGLTFHotspot, Log, All);
+
+namespace
+{
+	const FName NAME_HotspotTag = TEXT("GLTFHotspot");
+	const FName NAME_LevelEditorModule = TEXT("LevelEditor");
+	const FName NAME_SpriteParameter = TEXT("Sprite");
+	const FName NAME_OpacityParameter = TEXT("OpacityMask");
+}
+
+AGLTFHotspotActor::AGLTFHotspotActor(const FObjectInitializer& ObjectInitializer)
+	: Super(ObjectInitializer),
+	Image(nullptr),
+	HoveredImage(nullptr),
+	ToggledImage(nullptr),
+	ToggledHoveredImage(nullptr),
+	BillboardComponent(nullptr),
+	SphereComponent(nullptr),
+	DefaultMaterial(nullptr),
+	DefaultImage(nullptr),
+	DefaultHoveredImage(nullptr),
+	DefaultToggledImage(nullptr),
+	DefaultToggledHoveredImage(nullptr),
+	ActiveImage(nullptr),
+	ActiveImageSize(0.0f, 0.0f),
+	bToggled(false),
+	bIsInteractable(true),
+	RealtimeSecondsWhenLastInSight(0.0f),
+	RealtimeSecondsWhenLastHidden(0.0f)
+{
+	USceneComponent* SceneRootComponent = CreateDefaultSubobject<USceneComponent>(TEXT("Root"));
+	SetRootComponent(SceneRootComponent);
+	AddInstanceComponent(SceneRootComponent);
+	SceneRootComponent->SetMobility(EComponentMobility::Movable);
+
+	BillboardComponent = CreateDefaultSubobject<UMaterialBillboardComponent>(TEXT("MaterialBillboardComponent"));
+	AddInstanceComponent(BillboardComponent);
+	BillboardComponent->SetupAttachment(RootComponent);
+	BillboardComponent->SetMobility(EComponentMobility::Movable);
+
+	SphereComponent = CreateDefaultSubobject<USphereComponent>(TEXT("SphereComponent"));
+	AddInstanceComponent(BillboardComponent);
+	SphereComponent->SetupAttachment(RootComponent);
+	SphereComponent->SetMobility(EComponentMobility::Movable);
+	SphereComponent->ComponentTags.Add(NAME_HotspotTag);
+	SphereComponent->InitSphereRadius(100.0f);
+	SphereComponent->SetVisibility(false);
+
+	// Setup the most minimalistic collision profile for mouse input events
+	SphereComponent->SetCollisionEnabled(ECollisionEnabled::QueryAndPhysics);
+	SphereComponent->SetCollisionResponseToAllChannels(ECR_Ignore);
+	SphereComponent->SetCollisionResponseToChannel(ECC_Visibility, ECR_Block);
+	SphereComponent->SetGenerateOverlapEvents(false);
+
+	// Respond to interactions with the sphere-component
+	SphereComponent->OnBeginCursorOver.AddDynamic(this, &AGLTFHotspotActor::BeginCursorOver);
+	SphereComponent->OnEndCursorOver.AddDynamic(this, &AGLTFHotspotActor::EndCursorOver);
+	SphereComponent->OnClicked.AddDynamic(this, &AGLTFHotspotActor::Clicked);
+
+	struct FConstructorStatics
+	{
+		ConstructorHelpers::FObjectFinder<UMaterial> DefaultMaterial;
+		ConstructorHelpers::FObjectFinder<UMaterial> DefaultIconMaterial;
+		ConstructorHelpers::FObjectFinder<UTexture2D> DefaultImage;
+		ConstructorHelpers::FObjectFinder<UTexture2D> DefaultHoveredImage;
+		ConstructorHelpers::FObjectFinder<UTexture2D> DefaultToggledImage;
+		ConstructorHelpers::FObjectFinder<UTexture2D> DefaultToggledHoveredImage;
+
+		FConstructorStatics()
+			: DefaultMaterial(TEXT("/GLTFExporter/Materials/Hotspot"))
+			, DefaultIconMaterial(TEXT("/GLTFExporter/Materials/HotspotIcon"))
+			, DefaultImage(TEXT("/GLTFExporter/Textures/Hotspots/Image"))
+			, DefaultHoveredImage(TEXT("/GLTFExporter/Textures/Hotspots/HoveredImage"))
+			, DefaultToggledImage(TEXT("/GLTFExporter/Textures/Hotspots/ToggledImage"))
+			, DefaultToggledHoveredImage(TEXT("/GLTFExporter/Textures/Hotspots/ToggledHoveredImage"))
+		{
+		}
+	};
+	static FConstructorStatics ConstructorStatics;
+
+	DefaultMaterial = ConstructorStatics.DefaultMaterial.Object;
+	DefaultIconMaterial = ConstructorStatics.DefaultIconMaterial.Object;
+	DefaultImage = ConstructorStatics.DefaultImage.Object;
+	DefaultHoveredImage = ConstructorStatics.DefaultHoveredImage.Object;
+	DefaultToggledImage = ConstructorStatics.DefaultToggledImage.Object;
+	DefaultToggledHoveredImage = ConstructorStatics.DefaultToggledHoveredImage.Object;
+
+	PrimaryActorTick.bCanEverTick = true;
+	PrimaryActorTick.bStartWithTickEnabled = true;
+
+	LevelSequencePlayer = ObjectInitializer.CreateDefaultSubobject<ULevelSequencePlayer>(this, "LevelSequencePlayer");
+}
+
+
+#if WITH_EDITOR
+void AGLTFHotspotActor::PostEditChangeProperty(FPropertyChangedEvent& PropertyChangedEvent)
+{
+	Super::PostEditChangeProperty(PropertyChangedEvent);
+
+	const FProperty* PropertyThatChanged = PropertyChangedEvent.Property;
+
+	if (PropertyThatChanged)
+	{
+		const FName PropertyFName = PropertyThatChanged->GetFName();
+
+		if (PropertyFName == GET_MEMBER_NAME_CHECKED(ThisClass, Image))
+		{
+			UpdateActiveImageFromState(EGLTFHotspotState::Default);
+		}
+		else if (PropertyFName == GET_MEMBER_NAME_CHECKED(ThisClass, SkeletalMeshActor))
+		{
+			if (SkeletalMeshActor != nullptr && AnimationSequence == nullptr)
+			{
+				AnimationSequence = Cast<UAnimSequence>(SkeletalMeshActor->GetSkeletalMeshComponent()->AnimationData.AnimToPlay);
+			}
+			ValidateAnimation();
+		}
+		else if (PropertyFName == GET_MEMBER_NAME_CHECKED(ThisClass, AnimationSequence))
+		{
+			ValidateAnimation();
+		}
+	}
+}
+#endif // WITH_EDITOR
+
+void AGLTFHotspotActor::PostRegisterAllComponents()
+{
+	Super::PostRegisterAllComponents();
+
+	SetupSpriteElement();
+	UpdateActiveImageFromState(EGLTFHotspotState::Default);
+}
+
+void AGLTFHotspotActor::Tick(float DeltaTime)
+{
+	UWorld* World = GetWorld();
+
+	// TODO: is is safe to assume that we want the first controller for projections?
+	APlayerController* PlayerController = World != nullptr ? World->GetFirstPlayerController() : nullptr;
+	if (PlayerController == nullptr)
+	{
+		return;
+	}
+
+	const FVector ColliderLocation = SphereComponent->GetComponentLocation();
+
+	FVector ColliderScreenLocation;
+	if (!PlayerController->ProjectWorldLocationToScreenWithDistance(ColliderLocation, ColliderScreenLocation))
+	{
+		return;
+	}
+
+	// Update scale of the sphere-component to match the screen-size of the active image
+	{
+		const FVector2D CornerScreenLocation = FVector2D(ColliderScreenLocation) + ActiveImageSize * 0.5f;
+		FVector RayLocation;
+		FVector RayDirection;
+
+		if (PlayerController->DeprojectScreenPositionToWorld(CornerScreenLocation.X, CornerScreenLocation.Y, RayLocation, RayDirection))
+		{
+			const FVector ExtentLocation = RayLocation + RayDirection * ColliderScreenLocation.Z;
+			const float NewSphereRadius = (ExtentLocation - ColliderLocation).Size() / SphereComponent->GetShapeScale();
+			const float OldSphereRadius = SphereComponent->GetUnscaledSphereRadius();
+
+			if (FMath::Abs(NewSphereRadius - OldSphereRadius) > 0.1f)	// TODO: better epsilon?
+			{
+				SphereComponent->SetSphereRadius(NewSphereRadius);
+			}
+		}
+	}
+
+	// Update opacity and interactivity of the hotspot based on if it its occluded by other objects or not
+	{
+		FHitResult HitResult;
+		bool bIsHotspotOccluded = false;
+
+		if (PlayerController->GetHitResultAtScreenPosition(FVector2D(ColliderScreenLocation), ECC_Visibility, false, HitResult))
+		{
+			if (const UPrimitiveComponent* HitComponent = HitResult.GetComponent())
+			{
+				bIsHotspotOccluded = !HitComponent->ComponentTags.Contains(NAME_HotspotTag);
+			}
+		}
+
+		const float CurrentRealtimeSeconds = UGameplayStatics::GetRealTimeSeconds(World);
+		float Opacity;
+
+		if (bIsHotspotOccluded)
+		{
+			RealtimeSecondsWhenLastHidden = CurrentRealtimeSeconds;
+
+			const float HiddenDuration = FMath::Max(RealtimeSecondsWhenLastHidden - RealtimeSecondsWhenLastInSight, 0.0f);
+			const float FadeOutDuration = 0.5f;
+
+			Opacity = 1.0f - FMath::Clamp(HiddenDuration / FadeOutDuration, 0.0f, 1.0f);
+		}
+		else
+		{
+			RealtimeSecondsWhenLastInSight = CurrentRealtimeSeconds;
+
+			const float VisibleDuration = FMath::Max(RealtimeSecondsWhenLastInSight - RealtimeSecondsWhenLastHidden, 0.0f);
+			const float FadeInDuration = 0.25f;
+
+			Opacity = FMath::Clamp(VisibleDuration / FadeInDuration, 0.0f, 1.0f);
+		}
+
+		SetSpriteOpacity(Opacity);
+		bIsInteractable = Opacity >= 0.5f;
+	}
+
+	if (LevelSequencePlayer != nullptr)
+	{
+		LevelSequencePlayer->Update(DeltaTime);
+	}
+}
+
+void AGLTFHotspotActor::EndPlay(const EEndPlayReason::Type EndPlayReason)
+{
+	if (LevelSequencePlayer != nullptr)
+	{
+		// Stop may modify a lot of actor state so it needs to be called
+		// during EndPlay (when Actors + World are still valid) instead
+		// of waiting for the UObject to be destroyed by GC.
+		LevelSequencePlayer->Stop();
+	}
+}
+
+void AGLTFHotspotActor::BeginCursorOver(UPrimitiveComponent* TouchedComponent)
+{
+	if (bIsInteractable)
+	{
+		UpdateActiveImageFromState(bToggled ? EGLTFHotspotState::ToggledHovered :  EGLTFHotspotState::Hovered);
+	}
+}
+
+void AGLTFHotspotActor::EndCursorOver(UPrimitiveComponent* TouchedComponent)
+{
+	UpdateActiveImageFromState(bToggled ? EGLTFHotspotState::Toggled :  EGLTFHotspotState::Default);
+}
+
+void AGLTFHotspotActor::Clicked(UPrimitiveComponent* TouchedComponent, FKey ButtonPressed)
+{
+	if (!bIsInteractable)
+	{
+		return;
+	}
+
+	ToggleAnimation();
+
+	UpdateActiveImageFromState(bToggled ? EGLTFHotspotState::ToggledHovered :  EGLTFHotspotState::Hovered);
+}
+
+void AGLTFHotspotActor::UpdateActiveImageFromState(EGLTFHotspotState State)
+{
+	UTexture2D* NewImage = const_cast<UTexture2D*>(GetImageForState(State));
+
+	UTexture* DefaultTexture;
+	GetSpriteMaterial()->GetTextureParameterDefaultValue(NAME_SpriteParameter, DefaultTexture);
+
+	UTexture* SpriteTexture = NewImage != nullptr ? NewImage : DefaultTexture;
+	GetSpriteMaterial()->SetTextureParameterValue(NAME_SpriteParameter, SpriteTexture);
+
+	ActiveImage = NewImage;
+	ActiveImageSize = { SpriteTexture->GetSurfaceWidth(), SpriteTexture->GetSurfaceHeight() };
+
+	// NOTE: we do this even if size is unchanged since the last update may have failed
+	UpdateSpriteSize();
+}
+
+void AGLTFHotspotActor::SetupSpriteElement() const
+{
+	UMaterialInstanceDynamic* MaterialInstance;
+
+#if WITH_EDITORONLY_DATA
+	if (GetWorld() != nullptr && GetWorld()->WorldType == EWorldType::Editor)
+	{
+		MaterialInstance = UMaterialInstanceDynamic::Create(DefaultIconMaterial, GetTransientPackage());
+	}
+	else
+#endif // WITH_EDITORONLY_DATA
+	{
+		MaterialInstance = UMaterialInstanceDynamic::Create(DefaultMaterial, GetTransientPackage());
+	}
+
+	FMaterialSpriteElement Element;
+	Element.Material = MaterialInstance;
+	Element.bSizeIsInScreenSpace = true;
+	Element.BaseSizeX = 0.1f;
+	Element.BaseSizeY = 0.1f;
+
+	BillboardComponent->SetElements({ Element });
+}
+
+UMaterialInstanceDynamic* AGLTFHotspotActor::GetSpriteMaterial() const
+{
+	return static_cast<UMaterialInstanceDynamic*>(BillboardComponent->GetMaterial(0));
+}
+
+void AGLTFHotspotActor::UpdateSpriteSize()
+{
+	const FIntPoint ViewportSize = GetCurrentViewportSize();
+
+	float BaseSizeX = 0.1f;
+	float BaseSizeY = 0.1f;
+
+	if (ViewportSize.X > 0 && ViewportSize.Y > 0)
+	{
+		BaseSizeX = ActiveImageSize.X / static_cast<float>(ViewportSize.X);
+		BaseSizeY = ActiveImageSize.Y / static_cast<float>(ViewportSize.Y);
+	}
+
+	FMaterialSpriteElement& Element = BillboardComponent->Elements[0];
+
+	if (BaseSizeX != Element.BaseSizeX || BaseSizeY != Element.BaseSizeY)	// TODO: use epsilon for comparison?
+	{
+		Element.BaseSizeX = BaseSizeX;
+		Element.BaseSizeY = BaseSizeY;
+
+		BillboardComponent->MarkRenderStateDirty();
+	}
+}
+
+void AGLTFHotspotActor::SetSpriteOpacity(const float Opacity) const
+{
+	GetSpriteMaterial()->SetScalarParameterValue(NAME_OpacityParameter, Opacity);
+}
+
+FIntPoint AGLTFHotspotActor::GetCurrentViewportSize()
+{
+	// TODO: verify that correct size is calculated in the various play-modes and in the editor
+
+	const FViewport* Viewport = nullptr;
+
+	if (UWorld* World = GetWorld())
+	{
+		if (World->IsGameWorld())
+		{
+			if (UGameViewportClient* GameViewportClient = World->GetGameViewport())
+			{
+				Viewport = GameViewportClient->Viewport;
+			}
+		}
+#if WITH_EDITOR
+		else
+		{
+			if (FModuleManager::Get().IsModuleLoaded(NAME_LevelEditorModule))
+			{
+				FLevelEditorModule& LevelEditorModule = FModuleManager::GetModuleChecked<FLevelEditorModule>(NAME_LevelEditorModule);
+
+				if (const TSharedPtr<SLevelViewport> ActiveLevelViewport = LevelEditorModule.GetFirstActiveLevelViewport())
+				{
+					Viewport = ActiveLevelViewport->GetActiveViewport();
+				}
+			}
+		}
+#endif
+	}
+
+	if (Viewport != nullptr)
+	{
+		if (!Viewport->ViewportResizedEvent.IsBoundToObject(this))
+		{
+			Viewport->ViewportResizedEvent.AddUObject(this, &AGLTFHotspotActor::ViewportResized);
+		}
+
+		return Viewport->GetSizeXY();
+	}
+
+	return FIntPoint::ZeroValue;
+}
+
+void AGLTFHotspotActor::ViewportResized(FViewport*, uint32)
+{
+	UpdateSpriteSize();
+}
+
+void AGLTFHotspotActor::ToggleAnimation()
+{
+	const bool bReverseAnimation = bToggled;
+	const float TargetPlayRate = bToggled ? -1.0f : 1.0f;
+
+	if (SkeletalMeshActor != nullptr && AnimationSequence != nullptr)
+	{
+		USkeletalMeshComponent* SkeletalMeshComponent = SkeletalMeshActor->GetSkeletalMeshComponent();
+		SkeletalMeshComponent->SetAnimationMode(EAnimationMode::Type::AnimationSingleNode);
+		const UAnimSingleNodeInstance* SingleNodeInstance = SkeletalMeshComponent->GetSingleNodeInstance();
+
+		if (SkeletalMeshComponent->IsPlaying() && SingleNodeInstance != nullptr && SingleNodeInstance->GetAnimationAsset() == AnimationSequence)
+		{
+			// If the same animation is already playing, just reverse the play rate for a smooth transition
+			SkeletalMeshComponent->SetPlayRate(TargetPlayRate);
+		}
+		else
+		{
+			SkeletalMeshComponent->SetAnimation(AnimationSequence);
+			SkeletalMeshComponent->SetPlayRate(TargetPlayRate);
+			SkeletalMeshComponent->SetPosition(bReverseAnimation ? AnimationSequence->GetPlayLength() : 0.0f);
+			SkeletalMeshComponent->Play(false);
+		}
+	}
+	else if (LevelSequence != nullptr)
+	{
+		if (LevelSequencePlayer->IsPlaying())
+		{
+			// If the same animation is already playing, just reverse the play rate for a smooth transition
+			LevelSequencePlayer->SetPlayRate(TargetPlayRate);
+		}
+		else
+		{
+			// TODO: is there a simpler way to get the length (in seconds) of a level sequence?
+			const float SequenceLength = LevelSequence->GetMovieScene()->GetTickResolution().AsSeconds(MovieScene::DiscreteSize(LevelSequence->GetMovieScene()->GetPlaybackRange()));
+
+			FMovieSceneSequencePlaybackSettings PlaybackSettings;
+			PlaybackSettings.PlayRate = TargetPlayRate;
+			PlaybackSettings.StartTime = bReverseAnimation ? SequenceLength : 0.0f;
+			PlaybackSettings.bAutoPlay = true;
+			PlaybackSettings.bPauseAtEnd = true;
+			PlaybackSettings.LoopCount.Value = 0;
+			LevelSequencePlayer->Initialize(LevelSequence, GetLevel(), PlaybackSettings, FLevelSequenceCameraSettings());
+			LevelSequencePlayer->Play();
+		}
+	}
+
+	bToggled = !bToggled;
+}
+
+void AGLTFHotspotActor::ValidateAnimation()
+{
+	if (SkeletalMeshActor != nullptr && AnimationSequence != nullptr)
+	{
+		if (USkeletalMesh* SkeletalMesh = SkeletalMeshActor->GetSkeletalMeshComponent()->SkeletalMesh)
+		{
+			if (SkeletalMesh->GetSkeleton() != AnimationSequence->GetSkeleton())
+			{
+				if (SkeletalMesh->GetSkeleton() != nullptr)
+				{
+					UE_LOG(LogGLTFHotspot, Warning, TEXT("Animation %s is incompatible with skeleton %s, removing animation from actor."), *AnimationSequence->GetName(), *SkeletalMesh->GetSkeleton()->GetName());
+				}
+				else
+				{
+					UE_LOG(LogGLTFHotspot, Warning, TEXT("Animation %s is incompatible because mesh %s has no skeleton, removing animation from actor."), *AnimationSequence->GetName(), *SkeletalMesh->GetName());
+				}
+
+				AnimationSequence = nullptr;
+			}
+		}
+	}
+}
+
+const UTexture2D* AGLTFHotspotActor::GetImageForState(EGLTFHotspotState State) const
+{
+	const UTexture2D* CurrentImage = DefaultImage;
+	const UTexture2D* CurrentHoveredImage = DefaultHoveredImage;
+	const UTexture2D* CurrentToggledImage = DefaultToggledImage;
+	const UTexture2D* CurrentToggledHoveredImage = DefaultToggledHoveredImage;
+
+	if (Image != nullptr)
+	{
+		CurrentImage = Image;
+		CurrentHoveredImage = HoveredImage != nullptr ? HoveredImage : CurrentImage;
+		CurrentToggledImage = ToggledImage != nullptr ? ToggledImage : CurrentImage;
+
+		if (ToggledHoveredImage != nullptr)
+		{
+			CurrentToggledHoveredImage = ToggledHoveredImage;
+		}
+		else if (ToggledImage != nullptr)
+		{
+			CurrentToggledHoveredImage = ToggledImage;
+		}
+		else if (HoveredImage != nullptr)
+		{
+			CurrentToggledHoveredImage = HoveredImage;
+		}
+		else
+		{
+			CurrentToggledHoveredImage = Image;
+		}
+	}
+
+	switch (State)
+	{
+		case EGLTFHotspotState::Default:        return CurrentImage;
+		case EGLTFHotspotState::Hovered:        return CurrentHoveredImage;
+		case EGLTFHotspotState::Toggled:        return CurrentToggledImage;
+		case EGLTFHotspotState::ToggledHovered: return CurrentToggledHoveredImage;
+		default:
+			checkNoEntry();
+			return nullptr;
+	}
+}
diff --git a/Engine/Plugins/Enterprise/GLTFExporter/MigratedFromGit/Source/GLTFExporter/Private/UserData/GLTFMaterialUserData.cpp b/Engine/Plugins/Enterprise/GLTFExporter/MigratedFromGit/Source/GLTFExporter/Private/UserData/GLTFMaterialUserData.cpp
new file mode 100644
--- /dev/null
+++ b/Engine/Plugins/Enterprise/GLTFExporter/MigratedFromGit/Source/GLTFExporter/Private/UserData/GLTFMaterialUserData.cpp
@@ -0,0 +1,73 @@
+// Copyright Epic Games, Inc. All Rights Reserved.
+
+#include "UserData/GLTFMaterialUserData.h"
+#include "Materials/MaterialInstance.h"
+#include "Engine/Texture.h"
+
+FGLTFOverrideMaterialBakeSettings::FGLTFOverrideMaterialBakeSettings()
+	: bOverrideSize(false)
+	, Size(EGLTFMaterialBakeSizePOT::POT_1024)
+	, bOverrideFilter(false)
+	, Filter(TF_Trilinear)
+	, bOverrideTiling(false)
+	, Tiling(TA_Wrap)
+{
+}
+
+EGLTFMaterialBakeSizePOT UGLTFMaterialExportOptions::GetBakeSizeForPropertyGroup(const UMaterialInterface* Material, EGLTFMaterialPropertyGroup PropertyGroup, EGLTFMaterialBakeSizePOT DefaultValue)
+{
+	if (const FGLTFOverrideMaterialBakeSettings* BakeSettings = GetBakeSettingsByPredicate(Material, PropertyGroup, [](const FGLTFOverrideMaterialBakeSettings& Settings) { return Settings.bOverrideSize; }))
+	{
+		return BakeSettings->Size;
+	}
+
+	return DefaultValue;
+}
+
+TextureFilter UGLTFMaterialExportOptions::GetBakeFilterForPropertyGroup(const UMaterialInterface* Material, EGLTFMaterialPropertyGroup PropertyGroup, TextureFilter DefaultValue)
+{
+	if (const FGLTFOverrideMaterialBakeSettings* BakeSettings = GetBakeSettingsByPredicate(Material, PropertyGroup, [](const FGLTFOverrideMaterialBakeSettings& Settings) { return Settings.bOverrideFilter; }))
+	{
+		return BakeSettings->Filter;
+	}
+
+	return DefaultValue;
+}
+
+TextureAddress UGLTFMaterialExportOptions::GetBakeTilingForPropertyGroup(const UMaterialInterface* Material, EGLTFMaterialPropertyGroup PropertyGroup, TextureAddress DefaultValue)
+{
+	if (const FGLTFOverrideMaterialBakeSettings* BakeSettings = GetBakeSettingsByPredicate(Material, PropertyGroup, [](const FGLTFOverrideMaterialBakeSettings& Settings) { return Settings.bOverrideTiling; }))
+	{
+		return BakeSettings->Tiling;
+	}
+
+	return DefaultValue;
+}
+
+template <typename Predicate>
+const FGLTFOverrideMaterialBakeSettings* UGLTFMaterialExportOptions::GetBakeSettingsByPredicate(const UMaterialInterface* Material, EGLTFMaterialPropertyGroup PropertyGroup, Predicate Pred)
+{
+	do
+	{
+		if (const UGLTFMaterialExportOptions* UserData = const_cast<UMaterialInterface*>(Material)->GetAssetUserData<UGLTFMaterialExportOptions>())
+		{
+			if (const FGLTFOverrideMaterialBakeSettings* BakeSettings = UserData->Inputs.Find(PropertyGroup))
+			{
+				if (Pred(*BakeSettings))
+				{
+					return BakeSettings;
+				}
+			}
+
+			if (Pred(UserData->Default))
+			{
+				return &UserData->Default;
+			}
+		}
+
+		const UMaterialInstance* MaterialInstance = Cast<UMaterialInstance>(Material);
+		Material = MaterialInstance != nullptr ? MaterialInstance->Parent : nullptr;
+	} while (Material != nullptr);
+
+	return nullptr;
+}
diff --git a/Engine/Plugins/Enterprise/GLTFExporter/MigratedFromGit/Source/GLTFExporter/Public/Actors/GLTFCameraActor.h b/Engine/Plugins/Enterprise/GLTFExporter/MigratedFromGit/Source/GLTFExporter/Public/Actors/GLTFCameraActor.h
new file mode 100644
--- /dev/null
+++ b/Engine/Plugins/Enterprise/GLTFExporter/MigratedFromGit/Source/GLTFExporter/Public/Actors/GLTFCameraActor.h
@@ -0,0 +1,129 @@
+// Copyright Epic Games, Inc. All Rights Reserved.
+
+#pragma once
+
+#include "Camera/CameraActor.h"
+#include "GLTFCameraActor.generated.h"
+
+UENUM(BlueprintType)
+enum class EGLTFCameraControlMode : uint8
+{
+	FreeLook,
+	Orbital
+};
+
+/**
+ * GLTF-compatible camera that will carry over settings and simulate the behavior in the resulting viewer.
+ * Supports two modes:
+ * - FreeLook: Allows the user to rotate the camera without modifying its position.
+ * - Orbital: Focuses one actor in the scene and orbits it through mouse control.
+ */
+UCLASS(BlueprintType, Blueprintable, DisplayName = "GLTF Camera")
+class GLTFEXPORTER_API AGLTFCameraActor : public ACameraActor
+{
+	GENERATED_BODY()
+	//~ Begin UObject Interface
+public:
+	AGLTFCameraActor(const FObjectInitializer& ObjectInitializer);
+#if WITH_EDITOR
+	virtual void PostEditChangeProperty(FPropertyChangedEvent& PropertyChangedEvent) override;
+#endif // WITH_EDITOR
+	//~ End UObject Interface
+
+	//~ Begin AActor Interface
+protected:
+	virtual void BeginPlay() override;
+public:
+	virtual void Tick(float DeltaSeconds) override;
+	virtual void PreInitializeComponents() override;
+	virtual void PostActorCreated() override;
+	//~ End AActor Interface
+
+private:
+	UFUNCTION()
+	void OnMouseX(float AxisValue);
+
+	UFUNCTION()
+	void OnMouseY(float AxisValue);
+
+	UFUNCTION()
+	void OnMouseWheelAxis(float AxisValue);
+
+	float ClampDistance(float Value) const;
+
+	float ClampPitch(float Value) const;
+
+	float ClampYaw(float Value) const;
+
+	void RemoveInertia();
+
+	FRotator GetLookAtRotation(const FVector TargetPosition) const;
+
+	FVector GetFocusPosition() const;
+
+	bool SetAutoActivateForPlayer(const EAutoReceiveInput::Type Player);
+
+public:
+
+	/* Camera mode */
+	UPROPERTY(EditAnywhere, BlueprintReadWrite, Category = "GLTF Camera Actor")
+	EGLTFCameraControlMode Mode;
+
+	/* Actor which the camera will focus on and subsequently orbit when using Orbital mode. */
+	UPROPERTY(EditAnywhere, BlueprintReadWrite, Category = "GLTF Camera Actor", meta=(EditCondition="Mode == EGLTFCameraControlMode::Orbital"))
+	AActor* Target;
+
+	/* Minimum pitch angle (in degrees) for the camera. */
+	UPROPERTY(EditAnywhere, BlueprintReadWrite, Category = "GLTF Camera Actor")
+	float PitchAngleMin;
+
+	/* Maximum pitch angle (in degrees) for the camera. */
+	UPROPERTY(EditAnywhere, BlueprintReadWrite, Category = "GLTF Camera Actor")
+	float PitchAngleMax;
+
+	/* Minimum yaw angle (in degrees) for the camera. */
+	UPROPERTY(EditAnywhere, BlueprintReadWrite, Category = "GLTF Camera Actor")
+	float YawAngleMin;
+
+	/* Maximum yaw angle (in degrees) for the camera. */
+	UPROPERTY(EditAnywhere, BlueprintReadWrite, Category = "GLTF Camera Actor")
+	float YawAngleMax;
+
+	/* Closest distance the camera can approach the focused actor. */
+	UPROPERTY(EditAnywhere, BlueprintReadWrite, Category = "GLTF Camera Actor")
+	float DistanceMin;
+
+	/* Farthest distance the camera can recede from the focused actor. */
+	UPROPERTY(EditAnywhere, BlueprintReadWrite, Category = "GLTF Camera Actor")
+	float DistanceMax;
+
+	/* Duration (in seconds) that it takes the camera to complete a change in distance. */
+	UPROPERTY(EditAnywhere, BlueprintReadOnly, Category = "GLTF Camera Actor")
+	float DollyDuration;
+
+	/* Size of the dolly movement relative to user input. The higher the value, the faster it moves. */
+	UPROPERTY(EditAnywhere, BlueprintReadOnly, Category = "GLTF Camera Actor")
+	float DollySensitivity;
+
+	/* Deceleration that occurs after rotational movement. The higher the value, the longer it takes to settle. */
+	UPROPERTY(EditAnywhere, BlueprintReadOnly, Category = "GLTF Camera Actor")
+	float RotationInertia;
+
+	/* Size of the rotational movement relative to user input. The higher the value, the faster it moves. */
+	UPROPERTY(EditAnywhere, BlueprintReadOnly, Category = "GLTF Camera Actor")
+	float RotationSensitivity;
+
+private:
+	float Distance;
+	float Pitch;
+	float Yaw;
+	float TargetDistance;
+	float TargetPitch;
+	float TargetYaw;
+	float DollyTime;
+	float DollyStartDistance;
+
+public:
+
+	bool UsesYawLimits() const;
+};
diff --git a/Engine/Plugins/Enterprise/GLTFExporter/MigratedFromGit/Source/GLTFExporter/Public/Actors/GLTFHotspotActor.h b/Engine/Plugins/Enterprise/GLTFExporter/MigratedFromGit/Source/GLTFExporter/Public/Actors/GLTFHotspotActor.h
new file mode 100644
--- /dev/null
+++ b/Engine/Plugins/Enterprise/GLTFExporter/MigratedFromGit/Source/GLTFExporter/Public/Actors/GLTFHotspotActor.h
@@ -0,0 +1,152 @@
+// Copyright Epic Games, Inc. All Rights Reserved.
+
+#pragma once
+
+#include "GameFramework/Actor.h"
+#include "GLTFHotspotActor.generated.h"
+
+class ASkeletalMeshActor;
+class UAnimSequence;
+class ULevelSequencePlayer;
+class ULevelSequence;
+class UTexture2D;
+class UTexture;
+class UMaterialInterface;
+class UMaterialBillboardComponent;
+class USphereComponent;
+class UMaterialInstanceDynamic;
+class FViewport;
+
+enum class EGLTFHotspotState : uint8
+{
+	Default,
+	Hovered,
+	Toggled,
+	ToggledHovered
+};
+
+/**
+ * Actor wrapper for the GLTF hotspot component. Appears as a billboard and allows playback of animations when cursor input is enabled.
+ */
+UCLASS(BlueprintType, Blueprintable, HideCategories = (Sprite, Physics, Collision, Navigation), DisplayName = "GLTF Hotspot")
+class GLTFEXPORTER_API AGLTFHotspotActor : public AActor
+{
+	GENERATED_BODY()
+	//~ Begin UObject Interface
+public:
+	AGLTFHotspotActor(const FObjectInitializer& ObjectInitializer);
+#if WITH_EDITOR
+	virtual void PostEditChangeProperty(FPropertyChangedEvent& PropertyChangedEvent) override;
+#endif // WITH_EDITOR
+	//~ End UObject Interface
+
+	//~ Begin AActor Interface
+	virtual void PostRegisterAllComponents() override;
+	virtual void Tick(float DeltaSeconds) override;
+	virtual void EndPlay(const EEndPlayReason::Type EndPlayReason) override;
+	//~ End AActor Interface
+
+private:
+	UFUNCTION()
+	void BeginCursorOver(UPrimitiveComponent* TouchedComponent);
+
+	UFUNCTION()
+	void EndCursorOver(UPrimitiveComponent* TouchedComponent);
+
+	UFUNCTION()
+	void Clicked(UPrimitiveComponent* TouchedComponent, FKey ButtonPressed);
+
+	void UpdateActiveImageFromState(EGLTFHotspotState State);
+
+public:
+
+	/* The skeletal mesh actor that will be animated when the hotspot is clicked. */
+	UPROPERTY(EditAnywhere, BlueprintReadWrite, Category = "Hotspot Animation")
+	ASkeletalMeshActor* SkeletalMeshActor;
+
+	/* The animation sequence that will be played on the skeletal mesh actor. Must be compatible with its skeletal mesh asset. */
+	UPROPERTY(EditAnywhere, BlueprintReadWrite, Category = "Hotspot Animation", meta=(EditCondition="SkeletalMeshActor != nullptr"))
+	UAnimSequence* AnimationSequence;
+
+	/* The level sequence that will be played in the level. */
+	UPROPERTY(EditAnywhere, BlueprintReadWrite, Category = "Hotspot Animation", meta=(EditCondition="SkeletalMeshActor == nullptr"))
+	ULevelSequence* LevelSequence;
+
+	/* The billboard image that will be shown when the hotspot is in an inactive state or one without a specified image. */
+	UPROPERTY(EditAnywhere, BlueprintReadWrite, Category = "Hotspot Appearance")
+	UTexture2D* Image;
+
+	/** The optional billboard image that will be shown when a cursor enters the hotspot. */
+	UPROPERTY(EditAnywhere, BlueprintReadWrite, Category = "Hotspot Appearance", meta=(EditCondition="Image != nullptr"))
+	UTexture2D* HoveredImage;
+
+	/** The optional billboard image that will be shown when the hotspot is toggled by a click. */
+	UPROPERTY(EditAnywhere, BlueprintReadWrite, Category = "Hotspot Appearance", meta=(EditCondition="Image != nullptr"))
+	UTexture2D* ToggledImage;
+
+	/** The optional billboard image that will be shown when the hotspot is toggled by a click and a cursor enters it. */
+	UPROPERTY(EditAnywhere, BlueprintReadWrite, Category = "Hotspot Appearance", meta=(EditCondition="Image != nullptr"))
+	UTexture2D* ToggledHoveredImage;
+
+private:
+
+	UPROPERTY(Transient, DuplicateTransient)
+	UMaterialBillboardComponent* BillboardComponent;
+
+	UPROPERTY(Transient, DuplicateTransient)
+	USphereComponent* SphereComponent;
+
+	UPROPERTY(Transient, DuplicateTransient)
+	UMaterialInterface* DefaultMaterial;
+
+	UPROPERTY(Transient, DuplicateTransient)
+	UTexture2D* DefaultImage;
+
+	UPROPERTY(Transient, DuplicateTransient)
+	UTexture2D* DefaultHoveredImage;
+
+	UPROPERTY(Transient, DuplicateTransient)
+	UTexture2D* DefaultToggledImage;
+
+	UPROPERTY(Transient, DuplicateTransient)
+	UTexture2D* DefaultToggledHoveredImage;
+
+	UPROPERTY(Transient, DuplicateTransient)
+	UMaterialInterface* DefaultIconMaterial;
+
+	UPROPERTY(Transient, DuplicateTransient)
+	const UTexture* ActiveImage;
+
+	UPROPERTY(Instanced, Transient)
+	ULevelSequencePlayer* LevelSequencePlayer;
+
+	FVector2D ActiveImageSize;
+
+	bool bToggled;
+
+	bool bIsInteractable;
+
+	float RealtimeSecondsWhenLastInSight;
+
+	float RealtimeSecondsWhenLastHidden;
+
+	void SetupSpriteElement() const;
+
+	UMaterialInstanceDynamic* GetSpriteMaterial() const;
+
+	void UpdateSpriteSize();
+
+	void SetSpriteOpacity(const float Opacity) const;
+
+	FIntPoint GetCurrentViewportSize();
+
+	void ViewportResized(FViewport*, uint32);
+
+	void ToggleAnimation();
+
+	void ValidateAnimation();
+
+public:
+
+	const UTexture2D* GetImageForState(EGLTFHotspotState State) const;
+};
diff --git a/Engine/Plugins/Enterprise/GLTFExporter/MigratedFromGit/Source/GLTFExporter/Public/UserData/GLTFMaterialUserData.h b/Engine/Plugins/Enterprise/GLTFExporter/MigratedFromGit/Source/GLTFExporter/Public/UserData/GLTFMaterialUserData.h
new file mode 100644
--- /dev/null
+++ b/Engine/Plugins/Enterprise/GLTFExporter/MigratedFromGit/Source/GLTFExporter/Public/UserData/GLTFMaterialUserData.h
@@ -0,0 +1,113 @@
+// Copyright Epic Games, Inc. All Rights Reserved.
+
+#pragma once
+
+#include "Engine/Texture.h"
+#include "Engine/AssetUserData.h"
+#include "GLTFMaterialUserData.generated.h"
+
+class UMaterialInterface;
+
+UENUM(BlueprintType)
+enum class EGLTFMaterialBakeMode : uint8
+{
+	/** Never bake material inputs. */
+	Disabled,
+	/** Only use a simple quad if a material input needs to be baked out. */
+	Simple,
+	/** Allow usage of the mesh data if a material input needs to be baked out with vertex data. */
+	UseMeshData,
+};
+
+UENUM(BlueprintType)
+enum class EGLTFMaterialBakeSizePOT : uint8
+{
+	POT_1 UMETA(DisplayName = "1 x 1"),
+	POT_2 UMETA(DisplayName = "2 x 2"),
+	POT_4 UMETA(DisplayName = "4 x 4"),
+	POT_8 UMETA(DisplayName = "8 x 8"),
+	POT_16 UMETA(DisplayName = "16 x 16"),
+	POT_32 UMETA(DisplayName = "32 x 32"),
+	POT_64 UMETA(DisplayName = "64 x 64"),
+	POT_128 UMETA(DisplayName = "128 x 128"),
+	POT_256 UMETA(DisplayName = "256 x 256"),
+	POT_512 UMETA(DisplayName = "512 x 512"),
+	POT_1024 UMETA(DisplayName = "1024 x 1024"),
+	POT_2048 UMETA(DisplayName = "2048 x 2048"),
+	POT_4096 UMETA(DisplayName = "4096 x 4096"),
+	POT_8192 UMETA(DisplayName = "8192 x 8192")
+};
+
+UENUM(BlueprintType)
+enum class EGLTFMaterialPropertyGroup : uint8
+{
+	None UMETA(DisplayName = "None"),
+
+	BaseColorOpacity UMETA(DisplayName = "Base Color + Opacity (Mask)"),
+	MetallicRoughness UMETA(DisplayName = "Metallic + Roughness"),
+	EmissiveColor UMETA(DisplayName = "Emissive Color"),
+	Normal UMETA(DisplayName = "Normal"),
+	AmbientOcclusion UMETA(DisplayName = "Ambient Occlusion"),
+	ClearCoatRoughness UMETA(DisplayName = "Clear Coat + Clear Coat Roughness"),
+	ClearCoatBottomNormal UMETA(DisplayName = "Clear Coat Bottom Normal"),
+};
+
+USTRUCT(Blueprintable)
+struct FGLTFOverrideMaterialBakeSettings
+{
+	GENERATED_BODY()
+
+	FGLTFOverrideMaterialBakeSettings();
+
+	/** If enabled, default size will be overridden by the corresponding property. */
+	UPROPERTY(EditAnywhere, BlueprintReadWrite, Category = "", meta = (InlineEditConditionToggle))
+	bool bOverrideSize;
+
+	/** Overrides default size of the baked out texture. */
+	UPROPERTY(EditAnywhere, BlueprintReadWrite, Category = "", meta = (EditCondition = "bOverrideSize"))
+	EGLTFMaterialBakeSizePOT Size;
+
+	/** If enabled, default filtering mode will be overridden by the corresponding property. */
+	UPROPERTY(EditAnywhere, BlueprintReadWrite, Category = "", meta = (InlineEditConditionToggle))
+	bool bOverrideFilter;
+
+	/** Overrides the default filtering mode used when sampling the baked out texture. */
+	UPROPERTY(EditAnywhere, BlueprintReadWrite, Category = "", meta = (EditCondition = "bOverrideFilter", ValidEnumValues="TF_Nearest, TF_Bilinear, TF_Trilinear"))
+	TEnumAsByte<TextureFilter> Filter;
+
+	/** If enabled, default addressing mode will be overridden by the corresponding property. */
+	UPROPERTY(EditAnywhere, BlueprintReadWrite, Category = "", meta = (InlineEditConditionToggle))
+	bool bOverrideTiling;
+
+	/** Overrides the default addressing mode used when sampling the baked out texture. */
+	UPROPERTY(EditAnywhere, BlueprintReadWrite, Category = "", meta = (EditCondition = "bOverrideTiling"))
+	TEnumAsByte<TextureAddress> Tiling;
+};
+
+/** glTF-specific user data that can be added to material assets to override export options */
+UCLASS(BlueprintType, meta = (DisplayName = "GLTF Material Export Options"))
+class GLTFEXPORTER_API UGLTFMaterialExportOptions : public UAssetUserData
+{
+	GENERATED_BODY()
+
+public:
+
+	// TODO: add support for overriding more export options
+
+	/** Default bake settings for this material in general. */
+	UPROPERTY(EditAnywhere, BlueprintReadWrite, Category = "Override Bake Settings", meta = (ShowOnlyInnerProperties))
+	FGLTFOverrideMaterialBakeSettings Default;
+
+	/** Input-specific bake settings that override the defaults above. */
+	UPROPERTY(EditAnywhere, BlueprintReadWrite, Category = "Override Bake Settings")
+	TMap<EGLTFMaterialPropertyGroup, FGLTFOverrideMaterialBakeSettings> Inputs;
+
+	static EGLTFMaterialBakeSizePOT GetBakeSizeForPropertyGroup(const UMaterialInterface* Material, EGLTFMaterialPropertyGroup PropertyGroup, EGLTFMaterialBakeSizePOT DefaultValue);
+	static TextureFilter GetBakeFilterForPropertyGroup(const UMaterialInterface* Material, EGLTFMaterialPropertyGroup PropertyGroup, TextureFilter DefaultValue);
+	static TextureAddress GetBakeTilingForPropertyGroup(const UMaterialInterface* Material, EGLTFMaterialPropertyGroup PropertyGroup, TextureAddress DefaultValue);
+
+private:
+
+	template <typename Predicate>
+	static const FGLTFOverrideMaterialBakeSettings* GetBakeSettingsByPredicate(const UMaterialInterface* Material, EGLTFMaterialPropertyGroup PropertyGroup, Predicate Pred);
+};
diff --git a/Engine/Plugins/Enterprise/GLTFExporter/MigratedFromGit/Source/GLTFExporterRuntime/GLTFExporterRuntime.Build.cs b/Engine/Plugins/Enterprise/GLTFExporter/MigratedFromGit/Source/GLTFExporterRuntime/GLTFExporterRuntime.Build.cs
deleted file mode 100644
diff --git a/Engine/Plugins/Enterprise/GLTFExporter/MigratedFromGit/Source/GLTFExporterRuntime/Private/Actors/GLTFCameraActor.cpp b/Engine/Plugins/Enterprise/GLTFExporter/MigratedFromGit/Source/GLTFExporterRuntime/Private/Actors/GLTFCameraActor.cpp
deleted file mode 100644
diff --git a/Engine/Plugins/Enterprise/GLTFExporter/MigratedFromGit/Source/GLTFExporterRuntime/Private/Actors/GLTFHotspotActor.cpp b/Engine/Plugins/Enterprise/GLTFExporter/MigratedFromGit/Source/GLTFExporterRuntime/Private/Actors/GLTFHotspotActor.cpp
deleted file mode 100644
diff --git a/Engine/Plugins/Enterprise/GLTFExporter/MigratedFromGit/Source/GLTFExporterRuntime/Private/GLTFExporterRuntimeModule.cpp b/Engine/Plugins/Enterprise/GLTFExporter/MigratedFromGit/Source/GLTFExporterRuntime/Private/GLTFExporterRuntimeModule.cpp
deleted file mode 100644
diff --git a/Engine/Plugins/Enterprise/GLTFExporter/MigratedFromGit/Source/GLTFExporterRuntime/Private/UserData/GLTFMaterialUserData.cpp b/Engine/Plugins/Enterprise/GLTFExporter/MigratedFromGit/Source/GLTFExporterRuntime/Private/UserData/GLTFMaterialUserData.cpp
deleted file mode 100644
diff --git a/Engine/Plugins/Enterprise/GLTFExporter/MigratedFromGit/Source/GLTFExporterRuntime/Public/Actors/GLTFCameraActor.h b/Engine/Plugins/Enterprise/GLTFExporter/MigratedFromGit/Source/GLTFExporterRuntime/Public/Actors/GLTFCameraActor.h
deleted file mode 100644
diff --git a/Engine/Plugins/Enterprise/GLTFExporter/MigratedFromGit/Source/GLTFExporterRuntime/Public/Actors/GLTFHotspotActor.h b/Engine/Plugins/Enterprise/GLTFExporter/MigratedFromGit/Source/GLTFExporterRuntime/Public/Actors/GLTFHotspotActor.h
deleted file mode 100644
diff --git a/Engine/Plugins/Enterprise/GLTFExporter/MigratedFromGit/Source/GLTFExporterRuntime/Public/GLTFExporterRuntimeModule.h b/Engine/Plugins/Enterprise/GLTFExporter/MigratedFromGit/Source/GLTFExporterRuntime/Public/GLTFExporterRuntimeModule.h
deleted file mode 100644
diff --git a/Engine/Plugins/Enterprise/GLTFExporter/MigratedFromGit/Source/GLTFExporterRuntime/Public/UserData/GLTFMaterialUserData.h b/Engine/Plugins/Enterprise/GLTFExporter/MigratedFromGit/Source/GLTFExporterRuntime/Public/UserData/GLTFMaterialUserData.h
deleted file mode 100644
