[GLTFExporter] feat: change errors about allow cpu access to suggestions

commit b931c1fc4d48add8434fd7fe98561ea825163c46
Date: Mon, 15 Nov 2021 14:24:48 +0100

    feat: change errors about allow cpu access to suggestions
    (cherry picked from commit 23747b76f9702c1cfba83459897fa64c077af5a6​)

#skipci

diff --git a/Engine/Plugins/Enterprise/GLTFExporter/MigratedFromGit/Source/GLTFExporter/Private/Converters/GLTFMeshConverters.cpp b/Engine/Plugins/Enterprise/GLTFExporter/MigratedFromGit/Source/GLTFExporter/Private/Converters/GLTFMeshConverters.cpp
--- a/Engine/Plugins/Enterprise/GLTFExporter/MigratedFromGit/Source/GLTFExporter/Private/Converters/GLTFMeshConverters.cpp
+++ b/Engine/Plugins/Enterprise/GLTFExporter/MigratedFromGit/Source/GLTFExporter/Private/Converters/GLTFMeshConverters.cpp
@@ -43,10 +43,9 @@ FGLTFJsonMeshIndex FGLTFStaticMeshConverter::Convert(const UStaticMesh* StaticMe
 #if !WITH_EDITOR
 	if (!StaticMesh->bAllowCPUAccess)
 	{
-		Builder.LogError(FString::Printf(
-			TEXT("Can't export mesh %s in runtime because Allow CPU Access is not enabled in asset settings"),
+		Builder.LogSuggestion(FString::Printf(
+			TEXT("Export of mesh %s can in runtime be speed-up by checking 'Allow CPU Access' in asset settings"),
 			*StaticMesh->GetName()));
-		return FGLTFJsonMeshIndex(INDEX_NONE);
 	}
 #endif
 
@@ -92,11 +91,10 @@ FGLTFJsonMeshIndex FGLTFSkeletalMeshConverter::Convert(const USkeletalMesh* Skel
 #if !WITH_EDITOR
 	if (!SkeletalMesh->GetLODInfo(LODIndex)->bAllowCPUAccess)
 	{
-		Builder.LogError(FString::Printf(
-			TEXT("Can't export mesh %s in runtime because Allow CPU Access is not enabled for LOD %d in asset settings"),
+		Builder.LogSuggestion(FString::Printf(
+			TEXT("Export of mesh %s (LOD %d) can in runtime be speed-up by checking 'Allow CPU Access' in asset settings"),
 			*SkeletalMesh->GetName(),
 			LODIndex));
-		return FGLTFJsonMeshIndex(INDEX_NONE);
 	}
 #endif
 
diff --git a/Engine/Plugins/Enterprise/GLTFExporter/MigratedFromGit/Source/GLTFExporter/Private/Tasks/GLTFMeshTasks.cpp b/Engine/Plugins/Enterprise/GLTFExporter/MigratedFromGit/Source/GLTFExporter/Private/Tasks/GLTFMeshTasks.cpp
--- a/Engine/Plugins/Enterprise/GLTFExporter/MigratedFromGit/Source/GLTFExporter/Private/Tasks/GLTFMeshTasks.cpp
+++ b/Engine/Plugins/Enterprise/GLTFExporter/MigratedFromGit/Source/GLTFExporter/Private/Tasks/GLTFMeshTasks.cpp
@@ -52,15 +52,15 @@ namespace
 
 		if (bZeroNormals)
 		{
-			Builder.LogWarning(FString::Printf(
-				TEXT("Mesh %s has some nearly zero-length normals which can create some issues. Consider checking 'Recompute Normals' in the asset settings"),
+			Builder.LogSuggestion(FString::Printf(
+				TEXT("Mesh %s has some nearly zero-length normals which may not be supported in some glTF applications. Consider checking 'Recompute Normals' in the asset settings"),
 				MeshName));
 		}
 
 		if (bZeroTangents)
 		{
-			Builder.LogWarning(FString::Printf(
-				TEXT("Mesh %s has some nearly zero-length tangents which can create some issues. Consider checking 'Recompute Tangents' in the asset settings"),
+			Builder.LogSuggestion(FString::Printf(
+				TEXT("Mesh %s has some nearly zero-length tangents which may not be supported in some glTF applications. Consider checking 'Recompute Tangents' in the asset settings"),
 				MeshName));
 		}
 	}
