[GLTFExporter] fix: skip any RHI vertex/index buffer that is invalid

commit 16a675190d689b63319977648aca360091b144e0
Date: Mon, 15 Nov 2021 14:28:13 +0100

    fix: skip any RHI vertex/index buffer that is invalid
    (cherry picked from commit 90173a574af7e64df7aa9a939375c12ab82b1c0f​)

#skipci

diff --git a/Engine/Plugins/Enterprise/GLTFExporter/MigratedFromGit/Source/GLTFExporter/Private/Converters/GLTFAccessorConverters.cpp b/Engine/Plugins/Enterprise/GLTFExporter/MigratedFromGit/Source/GLTFExporter/Private/Converters/GLTFAccessorConverters.cpp
--- a/Engine/Plugins/Enterprise/GLTFExporter/MigratedFromGit/Source/GLTFExporter/Private/Converters/GLTFAccessorConverters.cpp
+++ b/Engine/Plugins/Enterprise/GLTFExporter/MigratedFromGit/Source/GLTFExporter/Private/Converters/GLTFAccessorConverters.cpp
@@ -30,6 +30,12 @@ FGLTFJsonAccessorIndex FGLTFPositionBufferConverter::Convert(const FGLTFMeshSect
 	const TUniquePtr<IGLTFBufferAdapter> SourceBuffer = IGLTFBufferAdapter::GetPositions(VertexBuffer);
 	const uint8* SourceData = SourceBuffer->GetData();
 
+	if (SourceData == nullptr)
+	{
+		// TODO: report error
+		return FGLTFJsonAccessorIndex(INDEX_NONE);
+	}
+
 	const TArray<uint32>& IndexMap = MeshSection->IndexMap;
 	const uint32 VertexCount = IndexMap.Num();
 	const uint32 Stride = VertexBuffer->GetStride();
@@ -85,6 +91,12 @@ FGLTFJsonAccessorIndex FGLTFColorBufferConverter::Convert(const FGLTFMeshSection
 	const TUniquePtr<IGLTFBufferAdapter> SourceBuffer = IGLTFBufferAdapter::GetColors(VertexBuffer);
 	const uint8* SourceData = SourceBuffer->GetData();
 
+	if (SourceData == nullptr)
+	{
+		// TODO: report error
+		return FGLTFJsonAccessorIndex(INDEX_NONE);
+	}
+
 	const TArray<uint32>& IndexMap = MeshSection->IndexMap;
 	const uint32 VertexCount = IndexMap.Num();
 	const uint32 Stride = VertexBuffer->GetStride();
@@ -119,6 +131,12 @@ FGLTFJsonAccessorIndex FGLTFNormalBufferConverter::Convert(const FGLTFMeshSectio
 	const TUniquePtr<IGLTFBufferAdapter> SourceBuffer = IGLTFBufferAdapter::GetTangents(VertexBuffer);
 	const uint8* SourceData = SourceBuffer->GetData();
 
+	if (SourceData == nullptr)
+	{
+		// TODO: report error
+		return FGLTFJsonAccessorIndex(INDEX_NONE);
+	}
+
 	FGLTFJsonBufferViewIndex BufferViewIndex;
 	EGLTFJsonComponentType ComponentType;
 
@@ -194,6 +212,12 @@ FGLTFJsonAccessorIndex FGLTFTangentBufferConverter::Convert(const FGLTFMeshSecti
 	const TUniquePtr<IGLTFBufferAdapter> SourceBuffer = IGLTFBufferAdapter::GetTangents(VertexBuffer);
 	const uint8* SourceData = SourceBuffer->GetData();
 
+	if (SourceData == nullptr)
+	{
+		// TODO: report error
+		return FGLTFJsonAccessorIndex(INDEX_NONE);
+	}
+
 	FGLTFJsonBufferViewIndex BufferViewIndex;
 	EGLTFJsonComponentType ComponentType;
 
@@ -276,6 +300,12 @@ FGLTFJsonAccessorIndex FGLTFUVBufferConverter::Convert(const FGLTFMeshSection* M
 	const TUniquePtr<IGLTFBufferAdapter> SourceBuffer = IGLTFBufferAdapter::GetUVs(VertexBuffer);
 	const uint8* SourceData = SourceBuffer->GetData();
 
+	if (SourceData == nullptr)
+	{
+		// TODO: report error
+		return FGLTFJsonAccessorIndex(INDEX_NONE);
+	}
+
 	return VertexBuffer->GetUseFullPrecisionUVs()
 		? Convert<FVector2D>(MeshSection, VertexBuffer, UVIndex, SourceData)
 		: Convert<FVector2DHalf>(MeshSection, VertexBuffer, UVIndex, SourceData);
@@ -319,6 +349,12 @@ FGLTFJsonAccessorIndex FGLTFBoneIndexBufferConverter::Convert(const FGLTFMeshSec
 	const TUniquePtr<IGLTFBufferAdapter> SourceBuffer = IGLTFBufferAdapter::GetInfluences(VertexBuffer);
 	const uint8* SourceData = SourceBuffer->GetData();
 
+	if (SourceData == nullptr)
+	{
+		// TODO: report error
+		return FGLTFJsonAccessorIndex(INDEX_NONE);
+	}
+
 	return MeshSection->MaxBoneIndex > UINT8_MAX
 		? Convert<uint16>(MeshSection, VertexBuffer, InfluenceOffset, SourceData)
 		: Convert<uint8>(MeshSection, VertexBuffer, InfluenceOffset, SourceData);
@@ -420,6 +456,12 @@ FGLTFJsonAccessorIndex FGLTFBoneWeightBufferConverter::Convert(const FGLTFMeshSe
 	const TUniquePtr<IGLTFBufferAdapter> SourceBuffer = IGLTFBufferAdapter::GetInfluences(VertexBuffer);
 	const uint8* SourceData = SourceBuffer->GetData();
 
+	if (SourceData == nullptr)
+	{
+		// TODO: report error
+		return FGLTFJsonAccessorIndex(INDEX_NONE);
+	}
+
 	return MeshSection->MaxBoneIndex > UINT8_MAX
 		? Convert<uint16>(MeshSection, VertexBuffer, InfluenceOffset, SourceData)
 		: Convert<uint8>(MeshSection, VertexBuffer, InfluenceOffset, SourceData);
diff --git a/Engine/Plugins/Enterprise/GLTFExporter/MigratedFromGit/Source/GLTFExporter/Private/Converters/GLTFBufferUtility.cpp b/Engine/Plugins/Enterprise/GLTFExporter/MigratedFromGit/Source/GLTFExporter/Private/Converters/GLTFBufferUtility.cpp
--- a/Engine/Plugins/Enterprise/GLTFExporter/MigratedFromGit/Source/GLTFExporter/Private/Converters/GLTFBufferUtility.cpp
+++ b/Engine/Plugins/Enterprise/GLTFExporter/MigratedFromGit/Source/GLTFExporter/Private/Converters/GLTFBufferUtility.cpp
@@ -68,9 +68,19 @@ const void* FGLTFBufferUtility::GetBufferData(const FSkinWeightLookupVertexBuffe
 
 void FGLTFBufferUtility::ReadRHIBuffer(FRHIVertexBuffer* SourceBuffer, TArray<uint8>& OutData)
 {
+	OutData.Empty();
+
+	if (SourceBuffer == nullptr)
+	{
+		return;
+	}
+
 	const uint32 NumBytes = SourceBuffer->GetSize();
+	if (NumBytes == 0)
+	{
+		return;
+	}
 
-	OutData.Empty();
 	OutData.AddUninitialized(NumBytes);
 	void *DstData = OutData.GetData();
 
@@ -88,9 +98,19 @@ void FGLTFBufferUtility::ReadRHIBuffer(FRHIVertexBuffer* SourceBuffer, TArray<ui
 
 void FGLTFBufferUtility::ReadRHIBuffer(FRHIIndexBuffer* SourceBuffer, TArray<uint8>& OutData)
 {
+	OutData.Empty();
+
+	if (SourceBuffer == nullptr)
+	{
+		return;
+	}
+
 	const uint32 NumBytes = SourceBuffer->GetSize();
+	if (NumBytes == 0)
+	{
+		return;
+	}
 
-	OutData.Empty();
 	OutData.AddUninitialized(NumBytes);
 	void *DstData = OutData.GetData();
 
diff --git a/Engine/Plugins/Enterprise/GLTFExporter/MigratedFromGit/Source/GLTFExporter/Private/Converters/GLTFMeshSection.cpp b/Engine/Plugins/Enterprise/GLTFExporter/MigratedFromGit/Source/GLTFExporter/Private/Converters/GLTFMeshSection.cpp
--- a/Engine/Plugins/Enterprise/GLTFExporter/MigratedFromGit/Source/GLTFExporter/Private/Converters/GLTFMeshSection.cpp
+++ b/Engine/Plugins/Enterprise/GLTFExporter/MigratedFromGit/Source/GLTFExporter/Private/Converters/GLTFMeshSection.cpp
@@ -7,6 +7,7 @@
 #include "Algo/MaxElement.h"
 
 FGLTFMeshSection::FGLTFMeshSection(const FStaticMeshLODResources* MeshLOD, const FGLTFIndexArray& SectionIndices)
+	: MaxBoneIndex(0)
 {
 	const TUniquePtr<IGLTFBufferAdapter> SourceBuffer = IGLTFBufferAdapter::GetIndices(&MeshLOD->IndexBuffer);
 	const uint8* SourceData = SourceBuffer->GetData();
@@ -22,6 +23,7 @@ FGLTFMeshSection::FGLTFMeshSection(const FStaticMeshLODResources* MeshLOD, const
 }
 
 FGLTFMeshSection::FGLTFMeshSection(const FSkeletalMeshLODRenderData* MeshLOD, const FGLTFIndexArray& SectionIndices)
+	: MaxBoneIndex(0)
 {
 	const TUniquePtr<IGLTFBufferAdapter> SourceBuffer = IGLTFBufferAdapter::GetIndices(MeshLOD->MultiSizeIndexContainer.GetIndexBuffer());
 	const uint8* SourceData = SourceBuffer->GetData();
@@ -39,6 +41,12 @@ FGLTFMeshSection::FGLTFMeshSection(const FSkeletalMeshLODRenderData* MeshLOD, co
 template <typename IndexType, typename SectionArrayType>
 void FGLTFMeshSection::Init(const SectionArrayType& Sections, const FGLTFIndexArray& SectionIndices, const IndexType* SourceData)
 {
+	if (SourceData == nullptr)
+	{
+		// TODO: report error
+		return;
+	}
+
 	uint32 TotalIndexCount = 0;
 	for (const int32 SectionIndex : SectionIndices)
 	{
@@ -48,7 +56,6 @@ void FGLTFMeshSection::Init(const SectionArrayType& Sections, const FGLTFIndexAr
 	IndexMap.Reserve(TotalIndexCount);
 	IndexBuffer.AddUninitialized(TotalIndexCount);
 	BoneMapLookup.Reserve(TotalIndexCount);
-	MaxBoneIndex = 0;
 
 	TMap<uint32, uint32> IndexLookup;
 
diff --git a/Engine/Plugins/Enterprise/GLTFExporter/MigratedFromGit/Source/GLTFExporter/Private/Tasks/GLTFMeshTasks.cpp b/Engine/Plugins/Enterprise/GLTFExporter/MigratedFromGit/Source/GLTFExporter/Private/Tasks/GLTFMeshTasks.cpp
--- a/Engine/Plugins/Enterprise/GLTFExporter/MigratedFromGit/Source/GLTFExporter/Private/Tasks/GLTFMeshTasks.cpp
+++ b/Engine/Plugins/Enterprise/GLTFExporter/MigratedFromGit/Source/GLTFExporter/Private/Tasks/GLTFMeshTasks.cpp
@@ -37,7 +37,13 @@ namespace
 
 		const TUniquePtr<IGLTFBufferAdapter> SourceBuffer = IGLTFBufferAdapter::GetTangents(VertexBuffer);
 		const uint8* SourceData = SourceBuffer->GetData();
-		const uint32 VertexCount= VertexBuffer->GetNumVertices();
+
+		if (SourceData == nullptr)
+		{
+			return;
+		}
+
+		const uint32 VertexCount = VertexBuffer->GetNumVertices();
 		bool bZeroNormals;
 		bool bZeroTangents;
 
@@ -67,8 +73,19 @@ namespace
 
 	bool HasVertexColors(const FColorVertexBuffer* VertexBuffer)
 	{
+		if (VertexBuffer == nullptr)
+		{
+			return false;
+		}
+
 		const TUniquePtr<IGLTFBufferAdapter> SourceBuffer = IGLTFBufferAdapter::GetColors(VertexBuffer);
 		const uint8* SourceData = SourceBuffer->GetData();
+
+		if (SourceData == nullptr)
+		{
+			return false;
+		}
+		
 		const uint32 VertexCount = VertexBuffer->GetNumVertices();
 		const uint32 Stride = VertexBuffer->GetStride();
 
