[GLTFExporter] fix: only read static RHI buffers (since some RHI implementations may crash otherwise)

commit e25c01da0da80f658828124051c42bf83fb94772
Date: Mon, 15 Nov 2021 14:41:52 +0100

    fix: only read static RHI buffers (since some RHI implementations may crash otherwise)
    (cherry picked from commit 2b6c38a614135b1f73044b0c032650c7b84268ba​)

#skipci

diff --git a/Engine/Plugins/Enterprise/GLTFExporter/MigratedFromGit/Source/GLTFExporter/Private/Converters/GLTFBufferUtility.cpp b/Engine/Plugins/Enterprise/GLTFExporter/MigratedFromGit/Source/GLTFExporter/Private/Converters/GLTFBufferUtility.cpp
--- a/Engine/Plugins/Enterprise/GLTFExporter/MigratedFromGit/Source/GLTFExporter/Private/Converters/GLTFBufferUtility.cpp
+++ b/Engine/Plugins/Enterprise/GLTFExporter/MigratedFromGit/Source/GLTFExporter/Private/Converters/GLTFBufferUtility.cpp
@@ -81,6 +81,12 @@ void FGLTFBufferUtility::ReadRHIBuffer(FRHIVertexBuffer* SourceBuffer, TArray<ui
 		return;
 	}
 
+	const uint32 Usage = SourceBuffer->GetUsage();
+	if ((Usage & BUF_Static) == 0)
+	{
+		return; // Some RHI implementations only support reading static buffers
+	}
+
 	OutData.AddUninitialized(NumBytes);
 	void *DstData = OutData.GetData();
 
@@ -111,6 +117,12 @@ void FGLTFBufferUtility::ReadRHIBuffer(FRHIIndexBuffer* SourceBuffer, TArray<uin
 		return;
 	}
 
+	const uint32 Usage = SourceBuffer->GetUsage();
+	if ((Usage & BUF_Static) == 0)
+	{
+		return; // Some RHI implementations only support reading static buffers
+	}
+
 	OutData.AddUninitialized(NumBytes);
 	void *DstData = OutData.GetData();
 
