[GLTFExporter] refactor: rename FGLTFExporterUtility::GetReferencedWorlds to GetAssociatedWorlds

commit 5173ea92011b9896d66427b7f96a4cb9542e2e3f
Date: Mon, 15 Nov 2021 15:12:56 +0100

    refactor: rename FGLTFExporterUtility::GetReferencedWorlds to GetAssociatedWorlds
    (cherry picked from commit 4258b48edb0b6e4afb057ba6eb41bd5e904065d7​)

#skipci

diff --git a/Engine/Plugins/Enterprise/GLTFExporter/MigratedFromGit/Source/GLTFExporter/Private/Exporters/GLTFExporterUtility.cpp b/Engine/Plugins/Enterprise/GLTFExporter/MigratedFromGit/Source/GLTFExporter/Private/Exporters/GLTFExporterUtility.cpp
--- a/Engine/Plugins/Enterprise/GLTFExporter/MigratedFromGit/Source/GLTFExporter/Private/Exporters/GLTFExporterUtility.cpp
+++ b/Engine/Plugins/Enterprise/GLTFExporter/MigratedFromGit/Source/GLTFExporter/Private/Exporters/GLTFExporterUtility.cpp
@@ -1,8 +1,8 @@
 // Copyright Epic Games, Inc. All Rights Reserved.
 
 #include "Exporters/GLTFExporterUtility.h"
-#include "AssetRegistryModule.h"
-#include "LevelSequence.h"
+#include "Materials/MaterialInstance.h"
+#include "AssetRegistry/IAssetRegistry.h"
 
 const UStaticMesh* FGLTFExporterUtility::GetPreviewMesh(const UMaterialInterface* Material)
 {
@@ -47,24 +47,23 @@ const USkeletalMesh* FGLTFExporterUtility::GetPreviewMesh(const UAnimSequence* A
 	return PreviewMesh;
 }
 
-TArray<UWorld*> FGLTFExporterUtility::GetReferencedWorlds(const UObject* Object)
+TArray<UWorld*> FGLTFExporterUtility::GetAssociatedWorlds(const UObject* Object)
 {
-	const FAssetRegistryModule& AssetRegistryModule = FModuleManager::LoadModuleChecked<FAssetRegistryModule>("AssetRegistry");
-	const FName OuterPathName = *Object->GetOutermost()->GetPathName();
+	TArray<UWorld*> Worlds;
+	TArray<FAssetIdentifier> Dependencies;
 
-	TArray<UWorld*> ReferencedWorlds;
-	TArray<FAssetIdentifier> AssetDependencies;
-	AssetRegistryModule.Get().GetDependencies(OuterPathName, AssetDependencies);
+	const FName OuterPathName = *Object->GetOutermost()->GetPathName();
+	IAssetRegistry::GetChecked().GetDependencies(OuterPathName, Dependencies);
 
-	for (FAssetIdentifier& AssetDependency : AssetDependencies)
+	for (FAssetIdentifier& Dependency : Dependencies)
 	{
-		FString PackageName = AssetDependency.PackageName.ToString();
+		FString PackageName = Dependency.PackageName.ToString();
 		UWorld* World = LoadObject<UWorld>(nullptr, *PackageName, nullptr, LOAD_NoWarn);
 		if (World != nullptr)
 		{
-			ReferencedWorlds.AddUnique(World);
+			Worlds.AddUnique(World);
 		}
 	}
 
-	return ReferencedWorlds;
+	return Worlds;
 }
diff --git a/Engine/Plugins/Enterprise/GLTFExporter/MigratedFromGit/Source/GLTFExporter/Private/Exporters/GLTFExporterUtility.h b/Engine/Plugins/Enterprise/GLTFExporter/MigratedFromGit/Source/GLTFExporter/Private/Exporters/GLTFExporterUtility.h
--- a/Engine/Plugins/Enterprise/GLTFExporter/MigratedFromGit/Source/GLTFExporter/Private/Exporters/GLTFExporterUtility.h
+++ b/Engine/Plugins/Enterprise/GLTFExporter/MigratedFromGit/Source/GLTFExporter/Private/Exporters/GLTFExporterUtility.h
@@ -5,12 +5,11 @@
 class UStaticMesh;
 class UMaterialInterface;
 class UAnimSequence;
-class ULevelSequence;
 
 struct FGLTFExporterUtility
 {
 	static const UStaticMesh* GetPreviewMesh(const UMaterialInterface* Material);
 	static const USkeletalMesh* GetPreviewMesh(const UAnimSequence* AnimSequence);
 
-	static TArray<UWorld*> GetReferencedWorlds(const UObject* LevelSequence);
+	static TArray<UWorld*> GetAssociatedWorlds(const UObject* Object);
 };
diff --git a/Engine/Plugins/Enterprise/GLTFExporter/MigratedFromGit/Source/GLTFExporter/Private/Exporters/GLTFLevelSequenceExporter.cpp b/Engine/Plugins/Enterprise/GLTFExporter/MigratedFromGit/Source/GLTFExporter/Private/Exporters/GLTFLevelSequenceExporter.cpp
--- a/Engine/Plugins/Enterprise/GLTFExporter/MigratedFromGit/Source/GLTFExporter/Private/Exporters/GLTFLevelSequenceExporter.cpp
+++ b/Engine/Plugins/Enterprise/GLTFExporter/MigratedFromGit/Source/GLTFExporter/Private/Exporters/GLTFLevelSequenceExporter.cpp
@@ -23,11 +23,11 @@ bool UGLTFLevelSequenceExporter::AddObject(FGLTFContainerBuilder& Builder, const
 		return false;
 	}
 
-	TArray<UWorld*> Worlds = FGLTFExporterUtility::GetReferencedWorlds(LevelSequence);
+	TArray<UWorld*> Worlds = FGLTFExporterUtility::GetAssociatedWorlds(LevelSequence);
 	if (Worlds.Num() == 0)
 	{
 		Builder.LogError(
-			FString::Printf(TEXT("Failed to export level sequence %s because no level referenced"),
+			FString::Printf(TEXT("Failed to export level sequence %s because no associated level"),
 			*LevelSequence->GetName()));
 		return false;
 	}
@@ -35,7 +35,7 @@ bool UGLTFLevelSequenceExporter::AddObject(FGLTFContainerBuilder& Builder, const
 	if (Worlds.Num() > 1)
 	{
 		Builder.LogError(
-			FString::Printf(TEXT("Failed to export level sequence %s because more than one level referenced"),
+			FString::Printf(TEXT("Failed to export level sequence %s because more than one associated level"),
 			*LevelSequence->GetName()));
 		return false;
 	}
diff --git a/Engine/Plugins/Enterprise/GLTFExporter/MigratedFromGit/Source/GLTFExporter/Private/Exporters/GLTFLevelVariantSetsExporter.cpp b/Engine/Plugins/Enterprise/GLTFExporter/MigratedFromGit/Source/GLTFExporter/Private/Exporters/GLTFLevelVariantSetsExporter.cpp
--- a/Engine/Plugins/Enterprise/GLTFExporter/MigratedFromGit/Source/GLTFExporter/Private/Exporters/GLTFLevelVariantSetsExporter.cpp
+++ b/Engine/Plugins/Enterprise/GLTFExporter/MigratedFromGit/Source/GLTFExporter/Private/Exporters/GLTFLevelVariantSetsExporter.cpp
@@ -23,11 +23,11 @@ bool UGLTFLevelVariantSetsExporter::AddObject(FGLTFContainerBuilder& Builder, co
 		return false;
 	}
 
-	TArray<UWorld*> Worlds = FGLTFExporterUtility::GetReferencedWorlds(LevelVariantSets);
+	TArray<UWorld*> Worlds = FGLTFExporterUtility::GetAssociatedWorlds(LevelVariantSets);
 	if (Worlds.Num() == 0)
 	{
 		Builder.LogError(
-			FString::Printf(TEXT("Failed to export level variant sets %s because no level referenced"),
+			FString::Printf(TEXT("Failed to export level variant sets %s because no associated level"),
 			*LevelVariantSets->GetName()));
 		return false;
 	}
@@ -35,7 +35,7 @@ bool UGLTFLevelVariantSetsExporter::AddObject(FGLTFContainerBuilder& Builder, co
 	if (Worlds.Num() > 1)
 	{
 		Builder.LogError(
-			FString::Printf(TEXT("Failed to export level variant sets %s because more than one level referenced"),
+			FString::Printf(TEXT("Failed to export level variant sets %s because more than one associated level"),
 			*LevelVariantSets->GetName()));
 		return false;
 	}
