[GLTFExporter] feat: find compatible preview mesh when exporting animations in runtime

commit 87b6e9b8964d2e826d1389279dd89944795109d7
Date: Mon, 15 Nov 2021 15:18:28 +0100

    feat: find compatible preview mesh when exporting animations in runtime
    (cherry picked from commit 3e01b5e9fc644cdc32df0000363913122570c03b​)

#skipci

diff --git a/Engine/Plugins/Enterprise/GLTFExporter/MigratedFromGit/Source/GLTFExporter/Private/Exporters/GLTFExporterUtility.cpp b/Engine/Plugins/Enterprise/GLTFExporter/MigratedFromGit/Source/GLTFExporter/Private/Exporters/GLTFExporterUtility.cpp
--- a/Engine/Plugins/Enterprise/GLTFExporter/MigratedFromGit/Source/GLTFExporter/Private/Exporters/GLTFExporterUtility.cpp
+++ b/Engine/Plugins/Enterprise/GLTFExporter/MigratedFromGit/Source/GLTFExporter/Private/Exporters/GLTFExporterUtility.cpp
@@ -41,12 +41,37 @@ const USkeletalMesh* FGLTFExporterUtility::GetPreviewMesh(const UAnimSequence* A
 		if (Skeleton != nullptr)
 		{
 			PreviewMesh = Skeleton->GetPreviewMesh();
+			if (PreviewMesh == nullptr)
+			{
+				PreviewMesh = FindCompatibleMesh(Skeleton);
+			}
 		}
 	}
 
 	return PreviewMesh;
 }
 
+const USkeletalMesh* FGLTFExporterUtility::FindCompatibleMesh(const USkeleton *Skeleton)
+{
+	FARFilter Filter;
+	Filter.ClassNames.Add(USkeletalMesh::StaticClass()->GetFName());
+	Filter.TagsAndValues.Add(USkeletalMesh::GetSkeletonMemberName(), FAssetData(Skeleton).GetExportTextName());
+
+	TArray<FAssetData> FilteredAssets;
+	IAssetRegistry::GetChecked().GetAssets(Filter, FilteredAssets);
+
+	for (const FAssetData& Asset : FilteredAssets)
+	{
+		const USkeletalMesh* SkeletalMesh = Cast<USkeletalMesh>(Asset.GetAsset());
+		if (SkeletalMesh != nullptr)
+		{
+			return SkeletalMesh;
+		}
+	}
+
+	return nullptr;
+}
+
 TArray<UWorld*> FGLTFExporterUtility::GetAssociatedWorlds(const UObject* Object)
 {
 	TArray<UWorld*> Worlds;
diff --git a/Engine/Plugins/Enterprise/GLTFExporter/MigratedFromGit/Source/GLTFExporter/Private/Exporters/GLTFExporterUtility.h b/Engine/Plugins/Enterprise/GLTFExporter/MigratedFromGit/Source/GLTFExporter/Private/Exporters/GLTFExporterUtility.h
--- a/Engine/Plugins/Enterprise/GLTFExporter/MigratedFromGit/Source/GLTFExporter/Private/Exporters/GLTFExporterUtility.h
+++ b/Engine/Plugins/Enterprise/GLTFExporter/MigratedFromGit/Source/GLTFExporter/Private/Exporters/GLTFExporterUtility.h
@@ -5,11 +5,13 @@
 class UStaticMesh;
 class UMaterialInterface;
 class UAnimSequence;
+class USkeleton;
 
 struct FGLTFExporterUtility
 {
 	static const UStaticMesh* GetPreviewMesh(const UMaterialInterface* Material);
 	static const USkeletalMesh* GetPreviewMesh(const UAnimSequence* AnimSequence);
+	static const USkeletalMesh* FindCompatibleMesh(const USkeleton* Skeleton);
 
 	static TArray<UWorld*> GetAssociatedWorlds(const UObject* Object);
 };
