[GLTFExporter] feat: improve finding preview mesh when exporting materials

commit 20bd71b236be6808c7923c5f8932dac6cd037c77
Date: Mon, 15 Nov 2021 15:20:09 +0100

    feat: improve finding preview mesh when exporting materials
    (cherry picked from commit c9ff85c363bd531f4e937661d4165ae4df8bc645​)

#skipci

diff --git a/Engine/Plugins/Enterprise/GLTFExporter/MigratedFromGit/Source/GLTFExporter/Private/Exporters/GLTFExporterUtility.cpp b/Engine/Plugins/Enterprise/GLTFExporter/MigratedFromGit/Source/GLTFExporter/Private/Exporters/GLTFExporterUtility.cpp
--- a/Engine/Plugins/Enterprise/GLTFExporter/MigratedFromGit/Source/GLTFExporter/Private/Exporters/GLTFExporterUtility.cpp
+++ b/Engine/Plugins/Enterprise/GLTFExporter/MigratedFromGit/Source/GLTFExporter/Private/Exporters/GLTFExporterUtility.cpp
@@ -7,29 +7,21 @@
 const UStaticMesh* FGLTFExporterUtility::GetPreviewMesh(const UMaterialInterface* Material)
 {
 #if WITH_EDITORONLY_DATA
-	// The following implementation is based of FMaterialInstanceEditor::RefreshPreviewAsset
-	const UStaticMesh* PreviewMesh = Cast<UStaticMesh>(Material->PreviewMesh.TryLoad());
-	if (PreviewMesh == nullptr)
+	do
 	{
-		// Attempt to use the parent material's preview mesh if the instance's preview mesh is invalid
-		const UMaterialInstance* MaterialInstance = Cast<UMaterialInstance>(Material);
-		if (MaterialInstance != nullptr)
+		const UStaticMesh* PreviewMesh = Cast<UStaticMesh>(Material->PreviewMesh.TryLoad());
+		if (PreviewMesh != nullptr)
 		{
-			const UMaterialInterface* ParentMaterial = MaterialInstance->Parent;
-			if (ParentMaterial != nullptr)
-			{
-				PreviewMesh = Cast<UStaticMesh>(ParentMaterial->PreviewMesh.TryLoad());
-			}
+			return PreviewMesh;
 		}
-	}
 
-	if (PreviewMesh != nullptr)
-	{
-		return PreviewMesh;
-	}
+		const UMaterialInstance* MaterialInstance = Cast<UMaterialInstance>(Material);
+		Material = MaterialInstance != nullptr ? MaterialInstance->Parent : nullptr;
+	} while (Material != nullptr);
 #endif
 
-	return LoadObject<UStaticMesh>(nullptr, TEXT("/Engine/EditorMeshes/EditorSphere.EditorSphere"));
+	static const UStaticMesh* DefaultPreviewMesh = LoadObject<UStaticMesh>(nullptr, TEXT("/Engine/EditorMeshes/EditorSphere.EditorSphere"));
+	return DefaultPreviewMesh;
 }
 
 const USkeletalMesh* FGLTFExporterUtility::GetPreviewMesh(const UAnimSequence* AnimSequence)
