[GLTFExporter] fix: don't show message log if no messages

commit 69fb22f453a6d36063a78c2486d7a672ec15d0e9
Date: Mon, 15 Nov 2021 15:29:45 +0100

    fix: don't show message log if no messages
    (cherry picked from commit 9204ae0151fb7b6c8a940bc2f3860ae16dc0ac39​)

#skipci

diff --git a/Engine/Plugins/Enterprise/GLTFExporter/MigratedFromGit/Source/GLTFExporter/Private/Builders/GLTFLogBuilder.cpp b/Engine/Plugins/Enterprise/GLTFExporter/MigratedFromGit/Source/GLTFExporter/Private/Builders/GLTFLogBuilder.cpp
--- a/Engine/Plugins/Enterprise/GLTFExporter/MigratedFromGit/Source/GLTFExporter/Private/Builders/GLTFLogBuilder.cpp
+++ b/Engine/Plugins/Enterprise/GLTFExporter/MigratedFromGit/Source/GLTFExporter/Private/Builders/GLTFLogBuilder.cpp
@@ -54,6 +54,11 @@ const TArray<FString>& FGLTFLogBuilder::GetLoggedErrors() const
 	return Errors;
 }
 
+bool FGLTFLogBuilder::HasLoggedMessages() const
+{
+	return Suggestions.Num() + Warnings.Num() + Errors.Num() > 0;
+}
+
 void FGLTFLogBuilder::OpenLog() const
 {
 #if WITH_EDITOR
diff --git a/Engine/Plugins/Enterprise/GLTFExporter/MigratedFromGit/Source/GLTFExporter/Private/Builders/GLTFLogBuilder.h b/Engine/Plugins/Enterprise/GLTFExporter/MigratedFromGit/Source/GLTFExporter/Private/Builders/GLTFLogBuilder.h
--- a/Engine/Plugins/Enterprise/GLTFExporter/MigratedFromGit/Source/GLTFExporter/Private/Builders/GLTFLogBuilder.h
+++ b/Engine/Plugins/Enterprise/GLTFExporter/MigratedFromGit/Source/GLTFExporter/Private/Builders/GLTFLogBuilder.h
@@ -28,6 +28,8 @@ public:
 
 	const TArray<FString>& GetLoggedErrors() const;
 
+	bool HasLoggedMessages() const;
+
 	void OpenLog() const;
 
 	void ClearLog();
diff --git a/Engine/Plugins/Enterprise/GLTFExporter/MigratedFromGit/Source/GLTFExporter/Private/Exporters/GLTFExporter.cpp b/Engine/Plugins/Enterprise/GLTFExporter/MigratedFromGit/Source/GLTFExporter/Private/Exporters/GLTFExporter.cpp
--- a/Engine/Plugins/Enterprise/GLTFExporter/MigratedFromGit/Source/GLTFExporter/Private/Exporters/GLTFExporter.cpp
+++ b/Engine/Plugins/Enterprise/GLTFExporter/MigratedFromGit/Source/GLTFExporter/Private/Exporters/GLTFExporter.cpp
@@ -45,7 +45,7 @@ bool UGLTFExporter::ExportBinary(UObject* Object, const TCHAR* Type, FArchive& A
 
 	// TODO: should we copy messages to UAssetExportTask::Errors?
 
-	if (!FApp::IsUnattended())
+	if (!FApp::IsUnattended() && Builder.HasLoggedMessages())
 	{
 		Builder.OpenLog();
 	}
