[GLTFExporter] fix: don't ignore default brush (and default physics volume) actor(s) since other actors may be attached to them in runtime

commit b0d83887084746d839b82acbe742383f36a2016b
Date: Mon, 15 Nov 2021 15:44:45 +0100

    fix: don't ignore default brush (and default physics volume) actor(s) since other actors may be attached to them in runtime
    (cherry picked from commit 89c26c2c50687da3b6982c9f10ef53efa70cb6d3​)

#skipci

diff --git a/Engine/Plugins/Enterprise/GLTFExporter/MigratedFromGit/Source/GLTFExporter/Private/Converters/GLTFNodeConverters.cpp b/Engine/Plugins/Enterprise/GLTFExporter/MigratedFromGit/Source/GLTFExporter/Private/Converters/GLTFNodeConverters.cpp
--- a/Engine/Plugins/Enterprise/GLTFExporter/MigratedFromGit/Source/GLTFExporter/Private/Converters/GLTFNodeConverters.cpp
+++ b/Engine/Plugins/Enterprise/GLTFExporter/MigratedFromGit/Source/GLTFExporter/Private/Converters/GLTFNodeConverters.cpp
@@ -54,6 +54,8 @@ FGLTFJsonNodeIndex FGLTFActorConverter::Convert(const AActor* Actor)
 	}
 	else
 	{
+		// TODO: add support for exporting brush geometry?
+
 		// TODO: to reduce number of nodes, only export components that are of interest
 
 		for (const UActorComponent* Component : Actor->GetComponents())
diff --git a/Engine/Plugins/Enterprise/GLTFExporter/MigratedFromGit/Source/GLTFExporter/Private/Converters/GLTFSceneConverters.cpp b/Engine/Plugins/Enterprise/GLTFExporter/MigratedFromGit/Source/GLTFExporter/Private/Converters/GLTFSceneConverters.cpp
--- a/Engine/Plugins/Enterprise/GLTFExporter/MigratedFromGit/Source/GLTFExporter/Private/Converters/GLTFSceneConverters.cpp
+++ b/Engine/Plugins/Enterprise/GLTFExporter/MigratedFromGit/Source/GLTFExporter/Private/Converters/GLTFSceneConverters.cpp
@@ -21,20 +21,10 @@ FGLTFJsonSceneIndex FGLTFSceneConverter::Convert(const UWorld* World)
 				continue;
 			}
 
-			// TODO: export Level->Model ?
+			// TODO: add support for exporting Level->Model?
 
 			for (const AActor* Actor : Level->Actors)
 			{
-				if (Actor == Level->GetDefaultBrush())
-				{
-					continue; // TODO: can we safely assume no other actor is ever attached to the default brush?
-				}
-
-				if (World->HasDefaultPhysicsVolume() && Actor == World->GetDefaultPhysicsVolume())
-				{
-					continue; // TODO: can we safely assume no other actor is ever attached to the default physics volume?
-				}
-
 				// TODO: should a LevelVariantSet be exported even if not selected for export?
 				if (const ALevelVariantSetsActor *LevelVariantSetsActor = Cast<ALevelVariantSetsActor>(Actor))
 				{
