[GLTFExporter] feat: add public API (and blueprint callable) static function UGLTFExporter::ExportToGLTF to allow export of any asset without relying on UExporter::RunAssetExportTask (which doesn't work in runtime currently)

commit 6c098895afdd150b89aa664d2bb9484bef74961c
Date: Mon, 15 Nov 2021 20:51:41 +0100

    feat: add public API (and blueprint callable) static function UGLTFExporter::ExportToGLTF to allow export of any asset without relying on UExporter::RunAssetExportTask (which doesn't work in runtime currently)
    (cherry picked from commit da9600284fc817136a5a62fb47793c6a08047ba7​)

#skipci

diff --git a/Engine/Plugins/Enterprise/GLTFExporter/MigratedFromGit/Source/GLTFExporter/Private/Exporters/GLTFExporter.cpp b/Engine/Plugins/Enterprise/GLTFExporter/MigratedFromGit/Source/GLTFExporter/Private/Exporters/GLTFExporter.cpp
--- a/Engine/Plugins/Enterprise/GLTFExporter/MigratedFromGit/Source/GLTFExporter/Private/Exporters/GLTFExporter.cpp
+++ b/Engine/Plugins/Enterprise/GLTFExporter/MigratedFromGit/Source/GLTFExporter/Private/Exporters/GLTFExporter.cpp
@@ -23,7 +23,7 @@ UGLTFExporter::UGLTFExporter(const FObjectInitializer& ObjectInitializer)
 
 bool UGLTFExporter::ExportBinary(UObject* Object, const TCHAR* Type, FArchive& Archive, FFeedbackContext* Warn, int32 FileIndex, uint32 PortFlags)
 {
-	UGLTFExportOptions* Options = GetExportOptions();
+	const UGLTFExportOptions* Options = GetExportOptions();
 	if (Options == nullptr)
 	{
 		// User cancelled the export
@@ -34,7 +34,7 @@ bool UGLTFExporter::ExportBinary(UObject* Object, const TCHAR* Type, FArchive& A
 
 	// TODO: add support for UAssetExportTask::IgnoreObjectList?
 
-	FGLTFContainerBuilder Builder(CurrentFilename, Options, bSelectedOnly);
+	FGLTFContainerBuilder Builder(GetFilePath(), Options, bSelectedOnly);
 	Builder.ClearLog();
 
 	const bool bSuccess = AddObject(Builder, Object);
@@ -53,12 +53,55 @@ bool UGLTFExporter::ExportBinary(UObject* Object, const TCHAR* Type, FArchive& A
 	return bSuccess;
 }
 
+bool UGLTFExporter::ExportToGLTF(UObject* Object, const FString& Filename, const UGLTFExportOptions* Options, FGLTFExportMessages& OutMessages)
+{
+	if (Object == nullptr)
+	{
+		OutMessages.Errors.Add(TEXT("No object to export"));
+		return false;
+	}
+
+	UGLTFExporter* Exporter = Cast<UGLTFExporter>(FindExporter(Object, TEXT("gltf")));
+	if (Exporter == nullptr)
+	{
+		OutMessages.Errors.Add(FString::Printf(TEXT("Couldn't find exporter for object of type %s"), *Object->GetClass()->GetName()));
+		return false;
+	}
+
+	FGLTFContainerBuilder Builder(Filename, Options, false);
+	bool bSuccess = Exporter->AddObject(Builder, Object);
+
+	OutMessages.Suggestions = Builder.GetLoggedSuggestions();
+	OutMessages.Warnings = Builder.GetLoggedWarnings();
+	OutMessages.Errors = Builder.GetLoggedErrors();
+
+	if (bSuccess)
+	{
+		FBufferArchive Buffer;
+		Builder.Write(Buffer, nullptr);
+
+		if (!FFileHelper::SaveArrayToFile(Buffer, *Filename))
+		{
+			OutMessages.Errors.Add(FString::Printf(TEXT("Couldn't save file: %s"), *Filename));
+			bSuccess = false;
+		}
+	}
+
+	return bSuccess;
+}
+
+bool UGLTFExporter::ExportToGLTF(UObject* Object, const FString& Filename, const UGLTFExportOptions* Options)
+{
+	FGLTFExportMessages Messages;
+	return ExportToGLTF(Object, Filename, Options, Messages);
+}
+
 bool UGLTFExporter::AddObject(FGLTFContainerBuilder& Builder, const UObject* Object)
 {
 	return false;
 }
 
-UGLTFExportOptions* UGLTFExporter::GetExportOptions()
+const UGLTFExportOptions* UGLTFExporter::GetExportOptions()
 {
 	UGLTFExportOptions* Options = nullptr;
 	bool bAutomatedTask = GIsAutomationTesting || FApp::IsUnattended();
@@ -100,3 +143,8 @@ UGLTFExportOptions* UGLTFExporter::GetExportOptions()
 
 	return Options;
 }
+
+FString UGLTFExporter::GetFilePath() const
+{
+	return ExportTask != nullptr ? ExportTask->Filename : CurrentFilename;
+}
diff --git a/Engine/Plugins/Enterprise/GLTFExporter/MigratedFromGit/Source/GLTFExporter/Public/Exporters/GLTFExporter.h b/Engine/Plugins/Enterprise/GLTFExporter/MigratedFromGit/Source/GLTFExporter/Public/Exporters/GLTFExporter.h
--- a/Engine/Plugins/Enterprise/GLTFExporter/MigratedFromGit/Source/GLTFExporter/Public/Exporters/GLTFExporter.h
+++ b/Engine/Plugins/Enterprise/GLTFExporter/MigratedFromGit/Source/GLTFExporter/Public/Exporters/GLTFExporter.h
@@ -9,6 +9,16 @@
 class FGLTFContainerBuilder;
 class UGLTFExportOptions;
 
+USTRUCT(BlueprintType)
+struct FGLTFExportMessages
+{
+	GENERATED_USTRUCT_BODY()
+
+	TArray<FString> Suggestions;
+	TArray<FString> Warnings;
+	TArray<FString> Errors;
+};
+
 UCLASS(Abstract)
 class GLTFEXPORTER_API UGLTFExporter : public UExporter
 {
@@ -20,11 +30,16 @@ public:
 
 	virtual bool ExportBinary(UObject* Object, const TCHAR* Type, FArchive& Archive, FFeedbackContext* Warn, int32 FileIndex, uint32 PortFlags) override;
 
+	UFUNCTION(BlueprintCallable, Category = "Miscellaneous")
+	static bool ExportToGLTF(UObject* Object, const FString& Filename, const UGLTFExportOptions* Options, FGLTFExportMessages& OutMessages);
+	static bool ExportToGLTF(UObject* Object, const FString& Filename, const UGLTFExportOptions* Options = nullptr);
+
 protected:
 
 	virtual bool AddObject(FGLTFContainerBuilder& Builder, const UObject* Object);
 
 private:
 
-	UGLTFExportOptions* GetExportOptions();
+	const UGLTFExportOptions* GetExportOptions();
+	FString GetFilePath() const;
 };
