[GLTFExporter] refactor: update dependency modules in exporter plugin and improve usage of FLevelEditorModule

commit 7ea749630a75e60c3c99a7711c4fa58cde71137b
Date: Mon, 15 Nov 2021 20:55:36 +0100

    refactor: update dependency modules in exporter plugin and improve usage of FLevelEditorModule
    (cherry picked from commit f57c1e0501b70ec358f68d55b39d1ebbb6df8109​)

#skipci

diff --git a/Engine/Plugins/Enterprise/GLTFExporter/MigratedFromGit/Source/GLTFExporter/GLTFExporter.Build.cs b/Engine/Plugins/Enterprise/GLTFExporter/MigratedFromGit/Source/GLTFExporter/GLTFExporter.Build.cs
--- a/Engine/Plugins/Enterprise/GLTFExporter/MigratedFromGit/Source/GLTFExporter/GLTFExporter.Build.cs
+++ b/Engine/Plugins/Enterprise/GLTFExporter/MigratedFromGit/Source/GLTFExporter/GLTFExporter.Build.cs
@@ -6,32 +6,43 @@ namespace UnrealBuildTool.Rules
 	{
 		public GLTFExporter(ReadOnlyTargetRules Target) : base(Target)
 		{
-		    // NOTE: ugly hack to access plugin info (should propose change to engine)
-			var BindingFlags = System.Reflection.BindingFlags.Instance | System.Reflection.BindingFlags.NonPublic;
-			var FieldInfo = typeof(ModuleRules).GetField("Plugin", BindingFlags);
-			var Plugin = FieldInfo != null ? FieldInfo.GetValue(this) as PluginInfo : null;
-
-			if (Plugin != null && Plugin.Descriptor != null)
+			if (IsPlugin)
 			{
-				PrivateDefinitions.Add("GLTFEXPORTER_FRIENDLY_NAME=TEXT(\"" + Plugin.Descriptor.FriendlyName + "\")");
-				PrivateDefinitions.Add("GLTFEXPORTER_VERSION_NAME=TEXT(\"" + Plugin.Descriptor.VersionName + "\")");
+				// NOTE: ugly hack to access plugin info (should propose change to engine)
+				var BindingFlags = System.Reflection.BindingFlags.Instance | System.Reflection.BindingFlags.NonPublic;
+				var FieldInfo = typeof(ModuleRules).GetField("Plugin", BindingFlags);
+				var Plugin = FieldInfo != null ? FieldInfo.GetValue(this) as PluginInfo : null;
+				var PluginDescriptor = Plugin != null ? Plugin.Descriptor : null;
+
+				if (PluginDescriptor != null)
+				{
+					PrivateDefinitions.Add("GLTFEXPORTER_FRIENDLY_NAME=TEXT(\"" + PluginDescriptor.FriendlyName + "\")");
+					PrivateDefinitions.Add("GLTFEXPORTER_VERSION_NAME=TEXT(\"" + PluginDescriptor.VersionName + "\")");
+				}
 			}
 
-			PrivateDependencyModuleNames.AddRange(
+			PublicDependencyModuleNames .AddRange(
 				new string[]
 				{
 					"Core",
 					"CoreUObject",
 					"Engine",
+				}
+			);
+
+			PrivateDependencyModuleNames.AddRange(
+				new string[]
+				{
 					"Json",
 					"RenderCore",
 					"RHI",
+					"ImageWrapper",
 					"LevelSequence",
 					"MovieScene",
 					"MovieSceneTracks",
 					"VariantManagerContent",
 				}
-				);
+			);
 
 			if (Target.bBuildEditor)
 			{
@@ -45,14 +56,15 @@ namespace UnrealBuildTool.Rules
 						"Mainframe",
 						"InputCore",
 						"EditorStyle",
+						"PropertyEditor",
 						"Projects",
 						"MeshMergeUtilities",
 						"MeshDescription",
 						"StaticMeshDescription",
 						"GLTFMaterialBaking",
-						"GLTFMaterialAnalyzer"
+						"GLTFMaterialAnalyzer",
 					}
-					);
+				);
 			}
 		}
 	}
diff --git a/Engine/Plugins/Enterprise/GLTFExporter/MigratedFromGit/Source/GLTFExporter/Private/Actors/GLTFHotspotActor.cpp b/Engine/Plugins/Enterprise/GLTFExporter/MigratedFromGit/Source/GLTFExporter/Private/Actors/GLTFHotspotActor.cpp
--- a/Engine/Plugins/Enterprise/GLTFExporter/MigratedFromGit/Source/GLTFExporter/Private/Actors/GLTFHotspotActor.cpp
+++ b/Engine/Plugins/Enterprise/GLTFExporter/MigratedFromGit/Source/GLTFExporter/Private/Actors/GLTFHotspotActor.cpp
@@ -369,11 +369,10 @@ FIntPoint AGLTFHotspotActor::GetCurrentViewportSize()
 #if WITH_EDITOR
 		else
 		{
-			if (FModuleManager::Get().IsModuleLoaded(NAME_LevelEditorModule))
+			FLevelEditorModule* LevelEditorModule = FModuleManager::GetModulePtr<FLevelEditorModule>(NAME_LevelEditorModule);
+			if (LevelEditorModule != nullptr)
 			{
-				FLevelEditorModule& LevelEditorModule = FModuleManager::GetModuleChecked<FLevelEditorModule>(NAME_LevelEditorModule);
-
-				if (const TSharedPtr<SLevelViewport> ActiveLevelViewport = LevelEditorModule.GetFirstActiveLevelViewport())
+				if (const TSharedPtr<SLevelViewport> ActiveLevelViewport = LevelEditorModule->GetFirstActiveLevelViewport())
 				{
 					Viewport = ActiveLevelViewport->GetActiveViewport();
 				}
diff --git a/Engine/Plugins/Enterprise/GLTFExporter/MigratedFromGit/Source/GLTFMaterialAnalyzer/GLTFMaterialAnalyzer.Build.cs b/Engine/Plugins/Enterprise/GLTFExporter/MigratedFromGit/Source/GLTFMaterialAnalyzer/GLTFMaterialAnalyzer.Build.cs
--- a/Engine/Plugins/Enterprise/GLTFExporter/MigratedFromGit/Source/GLTFMaterialAnalyzer/GLTFMaterialAnalyzer.Build.cs
+++ b/Engine/Plugins/Enterprise/GLTFExporter/MigratedFromGit/Source/GLTFMaterialAnalyzer/GLTFMaterialAnalyzer.Build.cs
@@ -6,14 +6,20 @@ public class GLTFMaterialAnalyzer : ModuleRules
 {
 	public GLTFMaterialAnalyzer(ReadOnlyTargetRules Target) : base(Target)
 	{
-		PrivateDependencyModuleNames.AddRange(
+		PublicDependencyModuleNames .AddRange(
 			new string[]
 			{
 				"Core",
 				"CoreUObject",
 				"Engine",
+			}
+		);
+
+		PrivateDependencyModuleNames.AddRange(
+			new string[]
+			{
 				"RenderCore",
-				"RHI"
+				"RHI",
 			}
 		);
 
diff --git a/Engine/Plugins/Enterprise/GLTFExporter/MigratedFromGit/Source/GLTFMaterialBaking/GLTFMaterialBaking.Build.cs b/Engine/Plugins/Enterprise/GLTFExporter/MigratedFromGit/Source/GLTFMaterialBaking/GLTFMaterialBaking.Build.cs
--- a/Engine/Plugins/Enterprise/GLTFExporter/MigratedFromGit/Source/GLTFMaterialBaking/GLTFMaterialBaking.Build.cs
+++ b/Engine/Plugins/Enterprise/GLTFExporter/MigratedFromGit/Source/GLTFMaterialBaking/GLTFMaterialBaking.Build.cs
@@ -1,35 +1,34 @@
 // Copyright Epic Games, Inc. All Rights Reserved.
 
 using UnrealBuildTool;
-using System.IO;
+
 public class GLTFMaterialBaking : ModuleRules
 {
 	public GLTFMaterialBaking(ReadOnlyTargetRules Target) : base(Target)
 	{
-        PrivateDependencyModuleNames.AddRange(
-			new string [] {
+		PublicDependencyModuleNames .AddRange(
+			new string[]
+			{
 				"Core",
 				"CoreUObject",
 				"Engine",
+			}
+		);
+
+		PrivateDependencyModuleNames.AddRange(
+			new string [] {
 				"RenderCore",
-                "RHI",                
-                "UnrealEd",
-                "MainFrame",
-                "SlateCore",
-                "Slate",
-                "InputCore",
-                "PropertyEditor",
-                "EditorStyle",
-                "Renderer",
-                "MeshDescription",
-				"StaticMeshDescription"
-            }
+				"RHI",
+				"UnrealEd",
+				"MeshDescription",
+				"StaticMeshDescription",
+			}
 		);
 
 		// NOTE: ugly hack to access HLSLMaterialTranslator to bake shading model
-        PrivateIncludePaths.Add(EngineDirectory + "/Source/Runtime/Engine/Private");
+		PrivateIncludePaths.Add(EngineDirectory + "/Source/Runtime/Engine/Private");
 
 		// NOTE: avoid renaming all instaces of MATERIALBAKING_API by redirecting it to GLTFMATERIALBAKING_API
 		PublicDefinitions.Add("MATERIALBAKING_API=GLTFMATERIALBAKING_API");
-    }
+	}
 }
diff --git a/Engine/Plugins/Enterprise/GLTFExporter/MigratedFromGit/Source/GLTFMaterialBaking/Private/MaterialBakingModule.cpp b/Engine/Plugins/Enterprise/GLTFExporter/MigratedFromGit/Source/GLTFMaterialBaking/Private/MaterialBakingModule.cpp
--- a/Engine/Plugins/Enterprise/GLTFExporter/MigratedFromGit/Source/GLTFMaterialBaking/Private/MaterialBakingModule.cpp
+++ b/Engine/Plugins/Enterprise/GLTFExporter/MigratedFromGit/Source/GLTFMaterialBaking/Private/MaterialBakingModule.cpp
@@ -4,7 +4,6 @@
 #include "MaterialRenderItem.h"
 #include "Engine/TextureRenderTarget2D.h"
 #include "ExportMaterialProxy.h"
-#include "Interfaces/IMainFrameModule.h"
 #include "UObject/UObjectGlobals.h"
 #include "MaterialBakingStructures.h"
 #include "MaterialBakingHelpers.h"
