[GLTFExporter] fix: don't reference invalid mesh primitive attributes in gltf mesh json object

commit c8c94e9fe6833e6e2d2b90af7c41a9aadaf5e638
Date: Tue, 16 Nov 2021 11:51:36 +0100

    fix: don't reference invalid mesh primitive attributes in gltf mesh json object
    (cherry picked from commit f90ef0af63ec4c729ee4b35e8d85eb578b6a4ef1​)

#skipci

diff --git a/Engine/Plugins/Enterprise/GLTFExporter/MigratedFromGit/Source/GLTFExporter/Private/Json/GLTFJsonMesh.h b/Engine/Plugins/Enterprise/GLTFExporter/MigratedFromGit/Source/GLTFExporter/Private/Json/GLTFJsonMesh.h
--- a/Engine/Plugins/Enterprise/GLTFExporter/MigratedFromGit/Source/GLTFExporter/Private/Json/GLTFJsonMesh.h
+++ b/Engine/Plugins/Enterprise/GLTFExporter/MigratedFromGit/Source/GLTFExporter/Private/Json/GLTFJsonMesh.h
@@ -8,37 +8,38 @@
 
 struct FGLTFJsonAttributes : IGLTFJsonObject
 {
-	FGLTFJsonAccessorIndex Position; // always required
+	FGLTFJsonAccessorIndex Position;
 	FGLTFJsonAccessorIndex Color0;
 	FGLTFJsonAccessorIndex Normal;
 	FGLTFJsonAccessorIndex Tangent;
-	TArray<FGLTFJsonAccessorIndex> TexCoords;
 
-	// skeletal mesh attributes
+	TArray<FGLTFJsonAccessorIndex> TexCoords;
 	TArray<FGLTFJsonAccessorIndex> Joints;
 	TArray<FGLTFJsonAccessorIndex> Weights;
 
 	virtual void WriteObject(IGLTFJsonWriter& Writer) const override
 	{
-		Writer.Write(TEXT("POSITION"), Position);
-
+		if (Position != INDEX_NONE) Writer.Write(TEXT("POSITION"), Position);
 		if (Color0 != INDEX_NONE) Writer.Write(TEXT("COLOR_0"), Color0);
 		if (Normal != INDEX_NONE) Writer.Write(TEXT("NORMAL"), Normal);
 		if (Tangent != INDEX_NONE) Writer.Write(TEXT("TANGENT"), Tangent);
 
 		for (int32 Index = 0; Index < TexCoords.Num(); ++Index)
 		{
-			Writer.Write(TEXT("TEXCOORD_") + FString::FromInt(Index), TexCoords[Index]);
+			const FGLTFJsonAccessorIndex TexCoord = TexCoords[Index];
+			if (TexCoord != INDEX_NONE) Writer.Write(TEXT("TEXCOORD_") + FString::FromInt(Index), TexCoord);
 		}
 
 		for (int32 Index = 0; Index < Joints.Num(); ++Index)
 		{
-			Writer.Write(TEXT("JOINTS_") + FString::FromInt(Index), Joints[Index]);
+			const FGLTFJsonAccessorIndex Joint = Joints[Index];
+			if (Joint != INDEX_NONE) Writer.Write(TEXT("JOINTS_") + FString::FromInt(Index), Joint);
 		}
 
 		for (int32 Index = 0; Index < Weights.Num(); ++Index)
 		{
-			Writer.Write(TEXT("WEIGHTS_") + FString::FromInt(Index), Weights[Index]);
+			const FGLTFJsonAccessorIndex Weight = Weights[Index];
+			if (Weight != INDEX_NONE) Writer.Write(TEXT("WEIGHTS_") + FString::FromInt(Index), Weight);
 		}
 	}
 };
