[GLTFExporter] fix: add missing reset value for DefaultLevelOfDetail in ResetToDefault

commit 5665d77716ca0416b81107f50471095289815d80
Date: Tue, 16 Nov 2021 12:27:27 +0100

    fix: add missing reset value for DefaultLevelOfDetail in ResetToDefault
    (cherry picked from commit 71986ab4857f63dab998d539bd5bc16faf1c3209​)

#skipci

diff --git a/Engine/Plugins/Enterprise/GLTFExporter/MigratedFromGit/Source/GLTFExporter/Private/GLTFExportOptions.cpp b/Engine/Plugins/Enterprise/GLTFExporter/MigratedFromGit/Source/GLTFExporter/Private/GLTFExportOptions.cpp
--- a/Engine/Plugins/Enterprise/GLTFExporter/MigratedFromGit/Source/GLTFExporter/Private/GLTFExportOptions.cpp
+++ b/Engine/Plugins/Enterprise/GLTFExporter/MigratedFromGit/Source/GLTFExporter/Private/GLTFExportOptions.cpp
@@ -24,6 +24,7 @@ void UGLTFExportOptions::ResetToDefault()
 	DefaultMaterialBakeSize = EGLTFMaterialBakeSizePOT::POT_1024;
 	DefaultMaterialBakeFilter = TF_Trilinear;
 	DefaultMaterialBakeTiling = TA_Wrap;
+	DefaultLevelOfDetail = 0;
 	bExportVertexColors = false;
 	bExportVertexSkinWeights = true;
 	bUseMeshQuantization = false;
