[GLTFExporter] feat: make UGLTFExportOptions::ResetToDefault blueprint callable

commit 41a4e05bcdb9312c6c530bc64c160b36b6b6e63e
Date: Tue, 16 Nov 2021 12:28:52 +0100

    feat: make UGLTFExportOptions::ResetToDefault blueprint callable
    (cherry picked from commit cef52e1af6dafdb3d6b3acdc806fe4ae63f30d0c​)

#skipci

diff --git a/Engine/Plugins/Enterprise/GLTFExporter/MigratedFromGit/Source/GLTFExporter/Public/GLTFExportOptions.h b/Engine/Plugins/Enterprise/GLTFExporter/MigratedFromGit/Source/GLTFExporter/Public/GLTFExportOptions.h
--- a/Engine/Plugins/Enterprise/GLTFExporter/MigratedFromGit/Source/GLTFExporter/Public/GLTFExportOptions.h
+++ b/Engine/Plugins/Enterprise/GLTFExporter/MigratedFromGit/Source/GLTFExporter/Public/GLTFExportOptions.h
@@ -216,5 +216,6 @@ class GLTFEXPORTER_API UGLTFExportOptions : public UObject
 	UPROPERTY(EditAnywhere, BlueprintReadWrite, Config, Category = VariantSets, Meta = (EditCondition = "bExportVariantSets"))
 	bool bExportVisibilityVariants;
 
+	UFUNCTION(BlueprintCallable)
 	void ResetToDefault();
 };
