[GLTFExporter] feat: make GLTFExportMessages properties accessible from blueprint

commit 0345ed19a43f91ecc0ff2314d381dad9dbd59704
Date: Wed, 17 Nov 2021 21:15:03 +0100

    feat: make GLTFExportMessages properties accessible from blueprint
    (cherry picked from commit c49b452630d07a703ce6fd8f3ae7a41de1b3ca92​)

#skipci

diff --git a/Engine/Plugins/Enterprise/GLTFExporter/MigratedFromGit/Source/GLTFExporter/Public/Exporters/GLTFExporter.h b/Engine/Plugins/Enterprise/GLTFExporter/MigratedFromGit/Source/GLTFExporter/Public/Exporters/GLTFExporter.h
--- a/Engine/Plugins/Enterprise/GLTFExporter/MigratedFromGit/Source/GLTFExporter/Public/Exporters/GLTFExporter.h
+++ b/Engine/Plugins/Enterprise/GLTFExporter/MigratedFromGit/Source/GLTFExporter/Public/Exporters/GLTFExporter.h
@@ -14,8 +14,13 @@ struct FGLTFExportMessages
 {
 	GENERATED_USTRUCT_BODY()
 
+	UPROPERTY(EditAnywhere, BlueprintReadWrite)
 	TArray<FString> Suggestions;
+
+	UPROPERTY(EditAnywhere, BlueprintReadWrite)
 	TArray<FString> Warnings;
+
+	UPROPERTY(EditAnywhere, BlueprintReadWrite)
 	TArray<FString> Errors;
 };
 
