[GLTFExporter] feat: don't export editor only actors and components (i.e. sprites etc)

commit c829fb1671c3b3b53efe631a3d2701160e47fb71
Date: Fri, 19 Nov 2021 20:59:14 +0100

    feat: don't export editor only actors and components (i.e. sprites etc)
    (cherry picked from commit 6001319ad8e0b4d3a67a5e1f940c0fa71c77aae8​)

#skipci

diff --git a/Engine/Plugins/Enterprise/GLTFExporter/MigratedFromGit/Source/GLTFExporter/Private/Converters/GLTFNodeConverters.cpp b/Engine/Plugins/Enterprise/GLTFExporter/MigratedFromGit/Source/GLTFExporter/Private/Converters/GLTFNodeConverters.cpp
--- a/Engine/Plugins/Enterprise/GLTFExporter/MigratedFromGit/Source/GLTFExporter/Private/Converters/GLTFNodeConverters.cpp
+++ b/Engine/Plugins/Enterprise/GLTFExporter/MigratedFromGit/Source/GLTFExporter/Private/Converters/GLTFNodeConverters.cpp
@@ -10,6 +10,11 @@
 
 FGLTFJsonNodeIndex FGLTFActorConverter::Convert(const AActor* Actor)
 {
+	if (Actor->bIsEditorOnlyActor)
+	{
+		return FGLTFJsonNodeIndex(INDEX_NONE);
+	}
+
 	if (Builder.bSelectedActorsOnly && !Actor->IsSelected())
 	{
 		return FGLTFJsonNodeIndex(INDEX_NONE);
@@ -73,6 +78,11 @@ FGLTFJsonNodeIndex FGLTFActorConverter::Convert(const AActor* Actor)
 
 FGLTFJsonNodeIndex FGLTFComponentConverter::Convert(const USceneComponent* SceneComponent)
 {
+	if (SceneComponent->IsEditorOnly())
+	{
+		return FGLTFJsonNodeIndex(INDEX_NONE);
+	}
+
 	const AActor* Owner = SceneComponent->GetOwner();
 	if (Owner == nullptr)
 	{
