[GLTFExporter] fix: rename CVars for GLTFMaterialBaking module to avoid conflict with engine module

commit ce8ca0ec3d2ffbc5dbac37321dea7a05b889a7b5
Date: Fri, 19 Nov 2021 22:59:04 +0100

    fix: rename CVars for GLTFMaterialBaking module to avoid conflict with engine module
    (cherry picked from commit 001397eba49d8da7143f1b63cfd5bb98ba9fbe2a​)

#skipci

diff --git a/Engine/Plugins/Enterprise/GLTFExporter/MigratedFromGit/Source/GLTFMaterialBaking/Private/MaterialBakingModule.cpp b/Engine/Plugins/Enterprise/GLTFExporter/MigratedFromGit/Source/GLTFMaterialBaking/Private/MaterialBakingModule.cpp
--- a/Engine/Plugins/Enterprise/GLTFExporter/MigratedFromGit/Source/GLTFMaterialBaking/Private/MaterialBakingModule.cpp
+++ b/Engine/Plugins/Enterprise/GLTFExporter/MigratedFromGit/Source/GLTFMaterialBaking/Private/MaterialBakingModule.cpp
@@ -27,7 +27,7 @@ IMPLEMENT_MODULE(FMaterialBakingModule, GLTFMaterialBaking);
 
 /** Cvars for advanced features */
 static TAutoConsoleVariable<int32> CVarUseMaterialProxyCaching(
-	TEXT("MaterialBaking.UseMaterialProxyCaching"),
+	TEXT("GLTFMaterialBaking.UseMaterialProxyCaching"),
 	1,
 	TEXT("Determines whether or not Material Proxies should be cached to speed up material baking.\n")
 	TEXT("0: Turned Off\n")
@@ -35,7 +35,7 @@ static TAutoConsoleVariable<int32> CVarUseMaterialProxyCaching(
 	ECVF_Default);
 
 static TAutoConsoleVariable<int32> CVarSaveIntermediateTextures(
-	TEXT("MaterialBaking.SaveIntermediateTextures"),
+	TEXT("GLTFMaterialBaking.SaveIntermediateTextures"),
 	0,
 	TEXT("Determines whether or not to save out intermediate BMP images for each flattened material property.\n")
 	TEXT("0: Turned Off\n")
