[GLTFExporter] fix: fix stack memory leak when exporting animations from a commandlet outside the editor

commit 405a83ffface2cf88f0a92a1ebc7a40f6c0b11ab
Date: Sat, 20 Nov 2021 01:50:37 +0100

    fix: fix stack memory leak when exporting animations from a commandlet outside the editor
    (cherry picked from commit 3839aa8f14c88e6e651121100f1c642790804cf2​)

#skipci

diff --git a/Engine/Plugins/Enterprise/GLTFExporter/MigratedFromGit/Source/GLTFExporter/Private/Tasks/GLTFAnimationTasks.cpp b/Engine/Plugins/Enterprise/GLTFExporter/MigratedFromGit/Source/GLTFExporter/Private/Tasks/GLTFAnimationTasks.cpp
--- a/Engine/Plugins/Enterprise/GLTFExporter/MigratedFromGit/Source/GLTFExporter/Private/Tasks/GLTFAnimationTasks.cpp
+++ b/Engine/Plugins/Enterprise/GLTFExporter/MigratedFromGit/Source/GLTFExporter/Private/Tasks/GLTFAnimationTasks.cpp
@@ -43,6 +43,8 @@ void FGLTFAnimSequenceTask::Complete()
 	TArray<TArray<FTransform>> FrameTransforms;
 	FrameTransforms.AddDefaulted(FrameCount);
 	{
+		FMemMark Mark(FMemStack::Get());
+
 		FCompactPose Pose;
 		Pose.SetBoneContainer(&BoneContainer);
 
@@ -58,6 +60,10 @@ void FGLTFAnimSequenceTask::Complete()
 			AnimSequence->GetBonePose(PoseData, ExtractionContext);
 			Pose.CopyBonesTo(FrameTransforms[Frame]);
 		}
+
+		Pose.Empty();
+		Curve.Empty();
+		Attributes.Empty();
 	}
 
 	// TODO: add animation data accessor converters to reuse track information
