[GLTFExporter] fix: add error message if failed to export mesh sections

commit ca71a27243649a62613464de98d7ccd2bbaf8aa7
Date: Sun, 21 Nov 2021 14:07:40 +0100

    fix: add error message if failed to export mesh sections
    (cherry picked from commit 8c7f6470e13c575ed6c7f0f8b4897fd517ed0aa6​)

#skipci

diff --git a/Engine/Plugins/Enterprise/GLTFExporter/MigratedFromGit/Source/GLTFExporter/Private/Tasks/GLTFMeshTasks.cpp b/Engine/Plugins/Enterprise/GLTFExporter/MigratedFromGit/Source/GLTFExporter/Private/Tasks/GLTFMeshTasks.cpp
--- a/Engine/Plugins/Enterprise/GLTFExporter/MigratedFromGit/Source/GLTFExporter/Private/Tasks/GLTFMeshTasks.cpp
+++ b/Engine/Plugins/Enterprise/GLTFExporter/MigratedFromGit/Source/GLTFExporter/Private/Tasks/GLTFMeshTasks.cpp
@@ -85,7 +85,7 @@ namespace
 		{
 			return false;
 		}
-		
+
 		const uint32 VertexCount = VertexBuffer->GetNumVertices();
 		const uint32 Stride = VertexBuffer->GetStride();
 
@@ -108,8 +108,8 @@ void FGLTFStaticMeshTask::Complete()
 	JsonMesh.Name = StaticMeshComponent != nullptr ? FGLTFNameUtility::GetName(StaticMeshComponent) : StaticMesh->GetName();
 
 	const FStaticMeshLODResources& MeshLOD = StaticMesh->GetLODForExport(LODIndex);
-	const FPositionVertexBuffer* PositionBuffer = &MeshLOD.VertexBuffers.PositionVertexBuffer;
-	const FStaticMeshVertexBuffer* VertexBuffer = &MeshLOD.VertexBuffers.StaticMeshVertexBuffer;
+	const FPositionVertexBuffer& PositionBuffer = MeshLOD.VertexBuffers.PositionVertexBuffer;
+	const FStaticMeshVertexBuffer& VertexBuffer = MeshLOD.VertexBuffers.StaticMeshVertexBuffer;
 	const FColorVertexBuffer* ColorBuffer = &MeshLOD.VertexBuffers.ColorVertexBuffer; // TODO: add support for overriding color buffer by component
 
 	if (Builder.ExportOptions->bExportVertexColors && HasVertexColors(ColorBuffer))
@@ -148,9 +148,10 @@ void FGLTFStaticMeshTask::Complete()
 	}
 #endif
 
-	ValidateVertexBuffer(Builder, VertexBuffer, *StaticMesh->GetName());
+	ValidateVertexBuffer(Builder, &VertexBuffer, *StaticMesh->GetName());
 
-	const int32 MaterialCount = StaticMesh->GetStaticMaterials().Num();
+	const TArray<FStaticMaterial>& MaterialSlots = StaticMesh->GetStaticMaterials();
+	const int32 MaterialCount = MaterialSlots.Num();
 	JsonMesh.Primitives.AddDefaulted(MaterialCount);
 
 	for (int32 MaterialIndex = 0; MaterialIndex < MaterialCount; ++MaterialIndex)
@@ -161,10 +162,20 @@ void FGLTFStaticMeshTask::Complete()
 		FGLTFJsonPrimitive& JsonPrimitive = JsonMesh.Primitives[MaterialIndex];
 		JsonPrimitive.Indices = Builder.GetOrAddIndexAccessor(ConvertedSection);
 
-		JsonPrimitive.Attributes.Position = Builder.GetOrAddPositionAccessor(ConvertedSection, PositionBuffer);
+		JsonPrimitive.Attributes.Position = Builder.GetOrAddPositionAccessor(ConvertedSection, &PositionBuffer);
 		if (JsonPrimitive.Attributes.Position == INDEX_NONE)
 		{
-			// TODO: report warning?
+			FString SectionString = TEXT("section");
+			SectionString += SectionIndices.Num() > 1 ? TEXT("s ") : TEXT(" ");
+			SectionString += FString::JoinBy(SectionIndices, TEXT(", "), FString::FromInt);
+
+			Builder.LogError(
+				FString::Printf(TEXT("Failed to export vertex positions related to material slot %d (%s) in static mesh %s (sections %s)"),
+				MaterialIndex,
+				*MaterialSlots[MaterialIndex].MaterialSlotName.ToString(),
+				*StaticMesh->GetName(),
+				*SectionString
+				));
 		}
 
 		if (ColorBuffer != nullptr)
@@ -173,16 +184,16 @@ void FGLTFStaticMeshTask::Complete()
 		}
 
 		// TODO: report warning if both Mesh Quantization (export options) and Use High Precision Tangent Basis (vertex buffer) are disabled
-		JsonPrimitive.Attributes.Normal = Builder.GetOrAddNormalAccessor(ConvertedSection, VertexBuffer);
-		JsonPrimitive.Attributes.Tangent = Builder.GetOrAddTangentAccessor(ConvertedSection, VertexBuffer);
+		JsonPrimitive.Attributes.Normal = Builder.GetOrAddNormalAccessor(ConvertedSection, &VertexBuffer);
+		JsonPrimitive.Attributes.Tangent = Builder.GetOrAddTangentAccessor(ConvertedSection, &VertexBuffer);
 
-		const uint32 UVCount = VertexBuffer->GetNumTexCoords();
+		const uint32 UVCount = VertexBuffer.GetNumTexCoords();
 		// TODO: report warning or option to limit UV channels since most viewers don't support more than 2?
 		JsonPrimitive.Attributes.TexCoords.AddUninitialized(UVCount);
 
 		for (uint32 UVIndex = 0; UVIndex < UVCount; ++UVIndex)
 		{
-			JsonPrimitive.Attributes.TexCoords[UVIndex] = Builder.GetOrAddUVAccessor(ConvertedSection, VertexBuffer, UVIndex);
+			JsonPrimitive.Attributes.TexCoords[UVIndex] = Builder.GetOrAddUVAccessor(ConvertedSection, &VertexBuffer, UVIndex);
 		}
 
 		const UMaterialInterface* Material = Materials[MaterialIndex];
@@ -198,8 +209,8 @@ void FGLTFSkeletalMeshTask::Complete()
 	const FSkeletalMeshRenderData* RenderData = SkeletalMesh->GetResourceForRendering();
 	const FSkeletalMeshLODRenderData& MeshLOD = RenderData->LODRenderData[LODIndex];
 
-	const FPositionVertexBuffer* PositionBuffer = &MeshLOD.StaticVertexBuffers.PositionVertexBuffer;
-	const FStaticMeshVertexBuffer* VertexBuffer = &MeshLOD.StaticVertexBuffers.StaticMeshVertexBuffer;
+	const FPositionVertexBuffer& PositionBuffer = MeshLOD.StaticVertexBuffers.PositionVertexBuffer;
+	const FStaticMeshVertexBuffer& VertexBuffer = MeshLOD.StaticVertexBuffers.StaticMeshVertexBuffer;
 	const FColorVertexBuffer* ColorBuffer = &MeshLOD.StaticVertexBuffers.ColorVertexBuffer; // TODO: add support for overriding color buffer by component
 	const FSkinWeightVertexBuffer* SkinWeightBuffer = MeshLOD.GetSkinWeightVertexBuffer(); // TODO: add support for overriding skin weight buffer by component
 	// TODO: add support for skin weight profiles?
@@ -242,9 +253,10 @@ void FGLTFSkeletalMeshTask::Complete()
 	}
 #endif
 
-	ValidateVertexBuffer(Builder, VertexBuffer, *SkeletalMesh->GetName());
+	ValidateVertexBuffer(Builder, &VertexBuffer, *SkeletalMesh->GetName());
 
-	const uint16 MaterialCount = SkeletalMesh->GetMaterials().Num();
+	const TArray<FSkeletalMaterial>& MaterialSlots = SkeletalMesh->GetMaterials();
+	const uint16 MaterialCount = MaterialSlots.Num();
 	JsonMesh.Primitives.AddDefaulted(MaterialCount);
 
 	for (uint16 MaterialIndex = 0; MaterialIndex < MaterialCount; ++MaterialIndex)
@@ -255,10 +267,20 @@ void FGLTFSkeletalMeshTask::Complete()
 		FGLTFJsonPrimitive& JsonPrimitive = JsonMesh.Primitives[MaterialIndex];
 		JsonPrimitive.Indices = Builder.GetOrAddIndexAccessor(ConvertedSection);
 
-		JsonPrimitive.Attributes.Position = Builder.GetOrAddPositionAccessor(ConvertedSection, PositionBuffer);
+		JsonPrimitive.Attributes.Position = Builder.GetOrAddPositionAccessor(ConvertedSection, &PositionBuffer);
 		if (JsonPrimitive.Attributes.Position == INDEX_NONE)
 		{
-			// TODO: report warning?
+			FString SectionString = TEXT("section");
+			SectionString += SectionIndices.Num() > 1 ? TEXT("s ") : TEXT(" ");
+			SectionString += FString::JoinBy(SectionIndices, TEXT(", "), FString::FromInt);
+
+			Builder.LogError(
+				FString::Printf(TEXT("Failed to export vertex positions related to material slot %d (%s) in skeletal mesh %s (sections %s)"),
+				MaterialIndex,
+				*MaterialSlots[MaterialIndex].MaterialSlotName.ToString(),
+				*SkeletalMesh->GetName(),
+				*SectionString
+				));
 		}
 
 		if (ColorBuffer != nullptr)
@@ -267,16 +289,16 @@ void FGLTFSkeletalMeshTask::Complete()
 		}
 
 		// TODO: report warning if both Mesh Quantization (export options) and Use High Precision Tangent Basis (vertex buffer) are disabled
-		JsonPrimitive.Attributes.Normal = Builder.GetOrAddNormalAccessor(ConvertedSection, VertexBuffer);
-		JsonPrimitive.Attributes.Tangent = Builder.GetOrAddTangentAccessor(ConvertedSection, VertexBuffer);
+		JsonPrimitive.Attributes.Normal = Builder.GetOrAddNormalAccessor(ConvertedSection, &VertexBuffer);
+		JsonPrimitive.Attributes.Tangent = Builder.GetOrAddTangentAccessor(ConvertedSection, &VertexBuffer);
 
-		const uint32 UVCount = VertexBuffer->GetNumTexCoords();
+		const uint32 UVCount = VertexBuffer.GetNumTexCoords();
 		// TODO: report warning or option to limit UV channels since most viewers don't support more than 2?
 		JsonPrimitive.Attributes.TexCoords.AddUninitialized(UVCount);
 
 		for (uint32 UVIndex = 0; UVIndex < UVCount; ++UVIndex)
 		{
-			JsonPrimitive.Attributes.TexCoords[UVIndex] = Builder.GetOrAddUVAccessor(ConvertedSection, VertexBuffer, UVIndex);
+			JsonPrimitive.Attributes.TexCoords[UVIndex] = Builder.GetOrAddUVAccessor(ConvertedSection, &VertexBuffer, UVIndex);
 		}
 
 		if (Builder.ExportOptions->bExportVertexSkinWeights)
