[GLTFExporter] feat: add IMaterialBakingModule::IsLinearBake

commit 3f0c220bc59a55868bfa3bae4a51efc3560f1387
Date: Mon, 22 Nov 2021 17:40:08 +0100

    feat: add IMaterialBakingModule::IsLinearBake
    (cherry picked from commit 32cd3992e3614a1ba243be96e50c8f699515cb2c​)

#skipci

diff --git a/Engine/Plugins/Enterprise/GLTFExporter/MigratedFromGit/Source/GLTFMaterialBaking/Private/MaterialBakingModule.cpp b/Engine/Plugins/Enterprise/GLTFExporter/MigratedFromGit/Source/GLTFMaterialBaking/Private/MaterialBakingModule.cpp
--- a/Engine/Plugins/Enterprise/GLTFExporter/MigratedFromGit/Source/GLTFMaterialBaking/Private/MaterialBakingModule.cpp
+++ b/Engine/Plugins/Enterprise/GLTFExporter/MigratedFromGit/Source/GLTFMaterialBaking/Private/MaterialBakingModule.cpp
@@ -773,6 +773,13 @@ void FMaterialBakingModule::SetLinearBake(bool bCorrectLinear)
 	}
 }
 
+bool FMaterialBakingModule::IsLinearBake(FMaterialPropertyEx Property)
+{
+	const EPropertyColorSpace* OverrideColorSpace = PerPropertyColorSpace.Find(Property);
+	const EPropertyColorSpace ColorSpace = OverrideColorSpace ? *OverrideColorSpace : DefaultColorSpace;
+	return ColorSpace == EPropertyColorSpace::Linear;
+}
+
 void FMaterialBakingModule::CleanupMaterialProxies()
 {
 	for (auto Iterator : MaterialProxyPool)
diff --git a/Engine/Plugins/Enterprise/GLTFExporter/MigratedFromGit/Source/GLTFMaterialBaking/Private/MaterialBakingModule.h b/Engine/Plugins/Enterprise/GLTFExporter/MigratedFromGit/Source/GLTFMaterialBaking/Private/MaterialBakingModule.h
--- a/Engine/Plugins/Enterprise/GLTFExporter/MigratedFromGit/Source/GLTFMaterialBaking/Private/MaterialBakingModule.h
+++ b/Engine/Plugins/Enterprise/GLTFExporter/MigratedFromGit/Source/GLTFMaterialBaking/Private/MaterialBakingModule.h
@@ -39,6 +39,9 @@ public:
 	/** Bakes all material properties to linear textures, except for colors */
 	virtual void SetLinearBake(bool bCorrectLinear) override;
 
+	/** Returns whether a specific material property is baked to a linear texture or not */
+	virtual bool IsLinearBake(FMaterialPropertyEx Property) override;
+
 protected:
 	/* Creates and adds or reuses a RenderTarget from the pool */
 	UTextureRenderTarget2D* CreateRenderTarget(bool bInForceLinearGamma, EPixelFormat InPixelFormat, const FIntPoint& InTargetSize, const FColor& BackgroundColor);
diff --git a/Engine/Plugins/Enterprise/GLTFExporter/MigratedFromGit/Source/GLTFMaterialBaking/Public/IMaterialBakingModule.h b/Engine/Plugins/Enterprise/GLTFExporter/MigratedFromGit/Source/GLTFMaterialBaking/Public/IMaterialBakingModule.h
--- a/Engine/Plugins/Enterprise/GLTFExporter/MigratedFromGit/Source/GLTFMaterialBaking/Public/IMaterialBakingModule.h
+++ b/Engine/Plugins/Enterprise/GLTFExporter/MigratedFromGit/Source/GLTFMaterialBaking/Public/IMaterialBakingModule.h
@@ -4,6 +4,7 @@
 
 #include "CoreMinimal.h"
 #include "Modules/ModuleInterface.h"
+#include "MaterialPropertyEx.h"
 #include "PixelFormat.h"
 #include "SceneTypes.h"
 
@@ -33,4 +34,7 @@ public:
 
 	/** Bakes all material properties to linear textures, except for colors */
 	virtual void SetLinearBake(bool bCorrectLinear) = 0;
+
+	/** Returns whether a specific material property is baked to a linear texture or not */
+	virtual bool IsLinearBake(FMaterialPropertyEx Property) = 0;
 };
