[GLTFExporter] fix: use correct color space when baking material properties into single factor

commit bad34b4bdc406f7ddd4462ecbc917ef05c3a61a2
Date: Mon, 22 Nov 2021 17:55:13 +0100

    fix: use correct color space when baking material properties into single factor
    (cherry picked from commit 3bbc0510b8570d0b1c5dbb73825f08eadf83a2cc​)

#skipci

diff --git a/Engine/Plugins/Enterprise/GLTFExporter/MigratedFromGit/Source/GLTFExporter/Private/Converters/GLTFMaterialUtility.cpp b/Engine/Plugins/Enterprise/GLTFExporter/MigratedFromGit/Source/GLTFExporter/Private/Converters/GLTFMaterialUtility.cpp
--- a/Engine/Plugins/Enterprise/GLTFExporter/MigratedFromGit/Source/GLTFExporter/Private/Converters/GLTFMaterialUtility.cpp
+++ b/Engine/Plugins/Enterprise/GLTFExporter/MigratedFromGit/Source/GLTFExporter/Private/Converters/GLTFMaterialUtility.cpp
@@ -227,10 +227,7 @@ FGLTFPropertyBakeOutput FGLTFMaterialUtility::BakeMaterialProperty(const FIntPoi
 		const FColor& Pixel = BakedPixels[0];
 
 		PropertyBakeOutput.bIsConstant = true;
-
-		// TODO: is the current conversion from sRGB => linear correct?
-		// It seems to give correct results for some properties, but not all.
-		PropertyBakeOutput.ConstantValue = Pixel;
+		PropertyBakeOutput.ConstantValue = Module.IsLinearBake(Property) ? Pixel.ReinterpretAsLinear() : FLinearColor(Pixel);;
 	}
 
 	return PropertyBakeOutput;
