[GLTFExporter] refactor: add GLTFMemoryArchive to replace BufferArchive internally

commit 46ce62d96433728c367b341d50d3ec866d8ca3a2
Date: Thu, 25 Nov 2021 20:02:51 +0100

    refactor: add GLTFMemoryArchive to replace BufferArchive internally
    (cherry picked from commit 3d181f18cf2d9aeaf1f5393a37d8d49974c53f92​)

#skipci

diff --git a/Engine/Plugins/Enterprise/GLTFExporter/MigratedFromGit/Source/GLTFExporter/Private/Builders/GLTFMemoryArchive.h b/Engine/Plugins/Enterprise/GLTFExporter/MigratedFromGit/Source/GLTFExporter/Private/Builders/GLTFMemoryArchive.h
new file mode 100644
--- /dev/null
+++ b/Engine/Plugins/Enterprise/GLTFExporter/MigratedFromGit/Source/GLTFExporter/Private/Builders/GLTFMemoryArchive.h
@@ -0,0 +1,34 @@
+// Copyright Epic Games, Inc. All Rights Reserved.
+
+#pragma once
+
+#include "CoreMinimal.h"
+#include "Serialization/MemoryArchive.h"
+
+class FGLTFMemoryArchive final : public FMemoryArchive, public TArray64<uint8>
+{
+public:
+
+	FGLTFMemoryArchive()
+	{
+		FMemoryArchive::SetIsSaving(true);
+	}
+
+	virtual void Serialize(void* Data, int64 Length) override
+	{
+		check(Length >= 0);
+
+		auto x = Num();
+
+		const int64 NumBytesToAdd = Offset + Length - Num();
+		AddUninitialized(NumBytesToAdd);
+
+		FMemory::Memcpy(GetData() + Offset, Data, Length);
+		Offset += Length;
+	}
+
+	virtual int64 TotalSize() override
+	{
+		return Num();
+	}
+};
