[GLTFExporter] feat: add support for binary and json data larger than 2GB using GLTFMemoryArchive

commit 970e531d5270cf1c6d73ad171b632106d44a33bd
Date: Thu, 25 Nov 2021 20:04:47 +0100

    feat: add support for binary and json data larger than 2GB using GLTFMemoryArchive
    (cherry picked from commit d3f1036b0b10e463addeb3d506c272549ca90ecd​)

#skipci

diff --git a/Engine/Plugins/Enterprise/GLTFExporter/MigratedFromGit/Source/GLTFExporter/Private/Builders/GLTFBufferBuilder.cpp b/Engine/Plugins/Enterprise/GLTFExporter/MigratedFromGit/Source/GLTFExporter/Private/Builders/GLTFBufferBuilder.cpp
--- a/Engine/Plugins/Enterprise/GLTFExporter/MigratedFromGit/Source/GLTFExporter/Private/Builders/GLTFBufferBuilder.cpp
+++ b/Engine/Plugins/Enterprise/GLTFExporter/MigratedFromGit/Source/GLTFExporter/Private/Builders/GLTFBufferBuilder.cpp
@@ -1,6 +1,7 @@
 // Copyright Epic Games, Inc. All Rights Reserved.
 
 #include "Builders/GLTFBufferBuilder.h"
+#include "Builders/GLTFMemoryArchive.h"
 #include "Misc/FileHelper.h"
 #include "Misc/Paths.h"
 
@@ -23,7 +24,7 @@ bool FGLTFBufferBuilder::InitializeBuffer()
 
 	if (bIsGlbFile)
 	{
-		BufferArchive = MakeUnique<FBufferArchive>();
+		BufferArchive = MakeUnique<FGLTFMemoryArchive>();
 	}
 	else
 	{
@@ -42,9 +43,9 @@ bool FGLTFBufferBuilder::InitializeBuffer()
 	return true;
 }
 
-const FBufferArchive* FGLTFBufferBuilder::GetBufferData() const
+const TArray64<uint8>* FGLTFBufferBuilder::GetBufferData() const
 {
-	return bIsGlbFile ? static_cast<FBufferArchive*>(BufferArchive.Get()) : nullptr;
+	return bIsGlbFile ? static_cast<FGLTFMemoryArchive*>(BufferArchive.Get()) : nullptr;
 }
 
 FGLTFJsonBufferViewIndex FGLTFBufferBuilder::AddBufferView(const void* RawData, uint64 ByteLength, EGLTFJsonBufferTarget BufferTarget, uint8 DataAlignment)
diff --git a/Engine/Plugins/Enterprise/GLTFExporter/MigratedFromGit/Source/GLTFExporter/Private/Builders/GLTFBufferBuilder.h b/Engine/Plugins/Enterprise/GLTFExporter/MigratedFromGit/Source/GLTFExporter/Private/Builders/GLTFBufferBuilder.h
--- a/Engine/Plugins/Enterprise/GLTFExporter/MigratedFromGit/Source/GLTFExporter/Private/Builders/GLTFBufferBuilder.h
+++ b/Engine/Plugins/Enterprise/GLTFExporter/MigratedFromGit/Source/GLTFExporter/Private/Builders/GLTFBufferBuilder.h
@@ -12,7 +12,7 @@ protected:
 	FGLTFBufferBuilder(const FString& FilePath, const UGLTFExportOptions* ExportOptions);
 	~FGLTFBufferBuilder();
 
-	const FBufferArchive* GetBufferData() const;
+	const TArray64<uint8>* GetBufferData() const;
 
 public:
 
diff --git a/Engine/Plugins/Enterprise/GLTFExporter/MigratedFromGit/Source/GLTFExporter/Private/Builders/GLTFContainerBuilder.cpp b/Engine/Plugins/Enterprise/GLTFExporter/MigratedFromGit/Source/GLTFExporter/Private/Builders/GLTFContainerBuilder.cpp
--- a/Engine/Plugins/Enterprise/GLTFExporter/MigratedFromGit/Source/GLTFExporter/Private/Builders/GLTFContainerBuilder.cpp
+++ b/Engine/Plugins/Enterprise/GLTFExporter/MigratedFromGit/Source/GLTFExporter/Private/Builders/GLTFContainerBuilder.cpp
@@ -1,6 +1,7 @@
 // Copyright Epic Games, Inc. All Rights Reserved.
 
 #include "Builders/GLTFContainerBuilder.h"
+#include "Builders/GLTFMemoryArchive.h"
 
 void FGLTFContainerBuilder::Write(FArchive& Archive, FFeedbackContext* Context)
 {
@@ -35,20 +36,20 @@ FGLTFContainerBuilder::FGLTFContainerBuilder(const FString& FilePath, const UGLT
 
 void FGLTFContainerBuilder::WriteGlb(FArchive& Archive)
 {
-	FBufferArchive JsonData;
+	FGLTFMemoryArchive JsonData;
 	WriteJson(JsonData);
 
-	const FBufferArchive* BufferData = GetBufferData();
-	WriteGlb(Archive, &JsonData, BufferData);
+	const TArray64<uint8>* BufferData = GetBufferData();
+	WriteGlb(Archive, JsonData, BufferData);
 }
 
-void FGLTFContainerBuilder::WriteGlb(FArchive& Archive, const FBufferArchive* JsonData, const FBufferArchive* BinaryData)
+void FGLTFContainerBuilder::WriteGlb(FArchive& Archive, const TArray64<uint8>& JsonData, const TArray64<uint8>* BinaryData)
 {
-	const uint32 JsonChunkType = 0x4E4F534A; // "JSON" in ASCII
-	const uint32 BinaryChunkType = 0x004E4942; // "BIN" in ASCII
+	constexpr uint32 JsonChunkType = 0x4E4F534A; // "JSON" in ASCII
+	constexpr uint32 BinaryChunkType = 0x004E4942; // "BIN" in ASCII
 	uint32 FileSize =
 		3 * sizeof(uint32) +
-		2 * sizeof(uint32) + GetPaddedChunkSize(JsonData->Num());
+		2 * sizeof(uint32) + GetPaddedChunkSize(JsonData.Num());
 
 	if (BinaryData != nullptr)
 	{
@@ -56,7 +57,7 @@ void FGLTFContainerBuilder::WriteGlb(FArchive& Archive, const FBufferArchive* Js
 	}
 
 	WriteHeader(Archive, FileSize);
-	WriteChunk(Archive, JsonChunkType, *JsonData, 0x20);
+	WriteChunk(Archive, JsonChunkType, JsonData, 0x20);
 
 	if (BinaryData != nullptr)
 	{
@@ -66,15 +67,15 @@ void FGLTFContainerBuilder::WriteGlb(FArchive& Archive, const FBufferArchive* Js
 
 void FGLTFContainerBuilder::WriteHeader(FArchive& Archive, uint32 FileSize)
 {
-	const uint32 FileSignature = 0x46546C67; // "glTF" in ASCII
-	const uint32 FileVersion = 2;
+	constexpr uint32 FileSignature = 0x46546C67; // "glTF" in ASCII
+	constexpr uint32 FileVersion = 2;
 
 	WriteInt(Archive, FileSignature);
 	WriteInt(Archive, FileVersion);
 	WriteInt(Archive, FileSize);
 }
 
-void FGLTFContainerBuilder::WriteChunk(FArchive& Archive, uint32 ChunkType, const TArray<uint8>& ChunkData, uint8 ChunkTrailingByte)
+void FGLTFContainerBuilder::WriteChunk(FArchive& Archive, uint32 ChunkType, const TArray64<uint8>& ChunkData, uint8 ChunkTrailingByte)
 {
 	const uint32 ChunkLength = GetPaddedChunkSize(ChunkData.Num());
 	const uint32 ChunkTrailing = GetTrailingChunkSize(ChunkData.Num());
@@ -90,7 +91,7 @@ void FGLTFContainerBuilder::WriteInt(FArchive& Archive, uint32 Value)
 	Archive.SerializeInt(Value, MAX_uint32);
 }
 
-void FGLTFContainerBuilder::WriteData(FArchive& Archive, const TArray<uint8>& Data)
+void FGLTFContainerBuilder::WriteData(FArchive& Archive, const TArray64<uint8>& Data)
 {
 	Archive.Serialize(const_cast<uint8*>(Data.GetData()), Data.Num());
 }
diff --git a/Engine/Plugins/Enterprise/GLTFExporter/MigratedFromGit/Source/GLTFExporter/Private/Builders/GLTFContainerBuilder.h b/Engine/Plugins/Enterprise/GLTFExporter/MigratedFromGit/Source/GLTFExporter/Private/Builders/GLTFContainerBuilder.h
--- a/Engine/Plugins/Enterprise/GLTFExporter/MigratedFromGit/Source/GLTFExporter/Private/Builders/GLTFContainerBuilder.h
+++ b/Engine/Plugins/Enterprise/GLTFExporter/MigratedFromGit/Source/GLTFExporter/Private/Builders/GLTFContainerBuilder.h
@@ -18,13 +18,13 @@ protected:
 
 private:
 
-	static void WriteGlb(FArchive& Archive, const FBufferArchive* JsonData, const FBufferArchive* BinaryData);
+	static void WriteGlb(FArchive& Archive, const TArray64<uint8>& JsonData, const TArray64<uint8>* BinaryData);
 
 	static void WriteHeader(FArchive& Archive, uint32 FileSize);
-	static void WriteChunk(FArchive& Archive, uint32 ChunkType, const TArray<uint8>& ChunkData, uint8 ChunkTrailingByte);
+	static void WriteChunk(FArchive& Archive, uint32 ChunkType, const TArray64<uint8>& ChunkData, uint8 ChunkTrailingByte);
 
 	static void WriteInt(FArchive& Archive, uint32 Value);
-	static void WriteData(FArchive& Archive, const TArray<uint8>& Data);
+	static void WriteData(FArchive& Archive, const TArray64<uint8>& Data);
 	static void WriteFill(FArchive& Archive, int32 Size, uint8 Value);
 
 	static int32 GetPaddedChunkSize(int32 Size);
