[GLTFExporter] feat: add support for writing files larger than 2GB using GLTFExporter::ExportToGLTF

commit 5a8db71f11a00d3308e552fcd92a3203772dda06
Date: Thu, 25 Nov 2021 20:06:19 +0100

    feat: add support for writing files larger than 2GB using GLTFExporter::ExportToGLTF
    (cherry picked from commit 25f2e7c6748e88142cef39069710b7c1e5f07f3e​)

#skipci

diff --git a/Engine/Plugins/Enterprise/GLTFExporter/MigratedFromGit/Source/GLTFExporter/Private/Exporters/GLTFExporter.cpp b/Engine/Plugins/Enterprise/GLTFExporter/MigratedFromGit/Source/GLTFExporter/Private/Exporters/GLTFExporter.cpp
--- a/Engine/Plugins/Enterprise/GLTFExporter/MigratedFromGit/Source/GLTFExporter/Private/Exporters/GLTFExporter.cpp
+++ b/Engine/Plugins/Enterprise/GLTFExporter/MigratedFromGit/Source/GLTFExporter/Private/Exporters/GLTFExporter.cpp
@@ -69,7 +69,7 @@ bool UGLTFExporter::ExportToGLTF(UObject* Object, const FString& Filename, const
 	}
 
 	FGLTFContainerBuilder Builder(Filename, Options, false);
-	bool bSuccess = Exporter->AddObject(Builder, Object);
+	const bool bSuccess = Exporter->AddObject(Builder, Object);
 
 	OutMessages.Suggestions = Builder.GetLoggedSuggestions();
 	OutMessages.Warnings = Builder.GetLoggedWarnings();
@@ -77,14 +77,15 @@ bool UGLTFExporter::ExportToGLTF(UObject* Object, const FString& Filename, const
 
 	if (bSuccess)
 	{
-		FBufferArchive Buffer;
-		Builder.Write(Buffer, nullptr);
-
-		if (!FFileHelper::SaveArrayToFile(Buffer, *Filename))
+		FArchive* Archive = IFileManager::Get().CreateFileWriter(*Filename);
+		if (Archive == nullptr)
 		{
 			OutMessages.Errors.Add(FString::Printf(TEXT("Couldn't save file: %s"), *Filename));
-			bSuccess = false;
+			return false;
 		}
+
+		Builder.Write(*Archive, nullptr);
+		Archive->Close();
 	}
 
 	return bSuccess;
