[GLTFExporter] fix: allow JPEG export when texture type is None

commit 6cad84e443cee14e6b4a34b3e325d82011ca9171
Date: Tue, 30 Nov 2021 08:11:45 +0100

    fix: allow JPEG export when texture type is None
    the None-type (for general / unspecified textures) cannot be processed the same way as other types since it's represented by 0, and comparing it to NoLossyImageFormatFor using bitwise AND will yield incorrect results.
    Furthermore this type is not part of NoLossyImageFormatFor, so the user has no way of specifying a lossless format for it.
    (cherry picked from commit 6aed4849459aca444117fa664cc3e9a7801c0e9c​)

#skipci

diff --git a/Engine/Plugins/Enterprise/GLTFExporter/MigratedFromGit/Source/GLTFExporter/Private/Builders/GLTFImageBuilder.cpp b/Engine/Plugins/Enterprise/GLTFExporter/MigratedFromGit/Source/GLTFExporter/Private/Builders/GLTFImageBuilder.cpp
--- a/Engine/Plugins/Enterprise/GLTFExporter/MigratedFromGit/Source/GLTFExporter/Private/Builders/GLTFImageBuilder.cpp
+++ b/Engine/Plugins/Enterprise/GLTFExporter/MigratedFromGit/Source/GLTFExporter/Private/Builders/GLTFImageBuilder.cpp
@@ -88,7 +88,7 @@ EGLTFJsonMimeType FGLTFImageBuilder::GetImageFormat(const FColor* Pixels, FIntPo
 
 		case EGLTFTextureImageFormat::JPEG:
 			return
-				!EnumHasAllFlags(static_cast<EGLTFTextureType>(ExportOptions->NoLossyImageFormatFor), Type) &&
+				(Type == EGLTFTextureType::None || !EnumHasAllFlags(static_cast<EGLTFTextureType>(ExportOptions->NoLossyImageFormatFor), Type)) &&
 				(bIgnoreAlpha || FGLTFImageUtility::NoAlphaNeeded(Pixels, Size)) ?
 				EGLTFJsonMimeType::JPEG : EGLTFJsonMimeType::PNG;
 
