[GLTFExporter] feat: treat hidden actors same as components in exported gltf

commit 3afc7d3b83e989f8953c0b4bec09c8ca33a0f3a3
Date: Tue, 9 Nov 2021 16:36:37 +0100

    feat: treat hidden actors same as components in exported gltf
    (cherry picked from commit 60ffe688cc523be4902c4fe0fd589610822e33e3​)

#skipci

diff --git a/Engine/Plugins/Enterprise/GLTFExporter/MigratedFromGit/Source/GLTFExporter/Private/Converters/GLTFNodeConverters.cpp b/Engine/Plugins/Enterprise/GLTFExporter/MigratedFromGit/Source/GLTFExporter/Private/Converters/GLTFNodeConverters.cpp
--- a/Engine/Plugins/Enterprise/GLTFExporter/MigratedFromGit/Source/GLTFExporter/Private/Converters/GLTFNodeConverters.cpp
+++ b/Engine/Plugins/Enterprise/GLTFExporter/MigratedFromGit/Source/GLTFExporter/Private/Converters/GLTFNodeConverters.cpp
@@ -15,7 +15,7 @@ FGLTFJsonNodeIndex FGLTFActorConverter::Convert(const AActor* Actor)
 		return FGLTFJsonNodeIndex(INDEX_NONE);
 	}
 
-	if (Builder.bSelectedActorsOnly && !Actor->IsSelected())
+	if (Builder.bSelectedActorsOnly && !Actor->IsSelected() || !Builder.ExportOptions->bExportHiddenInGame && Actor->IsHidden())
 	{
 		return FGLTFJsonNodeIndex(INDEX_NONE);
 	}
