[GLTFExporter] fix: avoid orphaned nodes if a root actor is hidden, by treating hidden actors exactly like hidden components (and only ignoring their "content" if hidden)

commit 38d55c7cf2af3dff2e471c583803bf581af52c37
Date: Tue, 30 Nov 2021 21:42:03 +0100

    fix: avoid orphaned nodes if a root actor is hidden, by treating hidden actors exactly like hidden components (and only ignoring their "content" if hidden)
    (cherry picked from commit 1652074dbd16b8bbddc70a408c55e33f1405d1aa​)

#skipci

diff --git a/Engine/Plugins/Enterprise/GLTFExporter/MigratedFromGit/Source/GLTFExporter/Private/Converters/GLTFNodeConverters.cpp b/Engine/Plugins/Enterprise/GLTFExporter/MigratedFromGit/Source/GLTFExporter/Private/Converters/GLTFNodeConverters.cpp
--- a/Engine/Plugins/Enterprise/GLTFExporter/MigratedFromGit/Source/GLTFExporter/Private/Converters/GLTFNodeConverters.cpp
+++ b/Engine/Plugins/Enterprise/GLTFExporter/MigratedFromGit/Source/GLTFExporter/Private/Converters/GLTFNodeConverters.cpp
@@ -15,7 +15,7 @@ FGLTFJsonNodeIndex FGLTFActorConverter::Convert(const AActor* Actor)
 		return FGLTFJsonNodeIndex(INDEX_NONE);
 	}
 
-	if (Builder.bSelectedActorsOnly && !Actor->IsSelected() || !Builder.ExportOptions->bExportHiddenInGame && Actor->IsHidden())
+	if (Builder.bSelectedActorsOnly && !Actor->IsSelected())
 	{
 		return FGLTFJsonNodeIndex(INDEX_NONE);
 	}
@@ -129,7 +129,7 @@ FGLTFJsonNodeIndex FGLTFComponentConverter::Convert(const USceneComponent* Scene
 	// TODO: don't export invisible components unless visibility is variable due to variant sets
 
 	// TODO: should hidden in game be configurable like this?
-	if ((!SceneComponent->bHiddenInGame || Builder.ExportOptions->bExportHiddenInGame) && FGLTFActorUtility::IsGenericActor(Owner))
+	if ((Builder.ExportOptions->bExportHiddenInGame || (!SceneComponent->bHiddenInGame && !Owner->IsHidden())) && FGLTFActorUtility::IsGenericActor(Owner))
 	{
 		if (const UStaticMeshComponent* StaticMeshComponent = Cast<UStaticMeshComponent>(SceneComponent))
 		{
