[GLTFExporter] docs: update exporter tooltip to reflect new behavior

commit d759f16196bc1118ce8c876bc79ed5a459052454
Date: Tue, 9 Nov 2021 16:37:54 +0100

    docs: update exporter tooltip to reflect new behavior
    (cherry picked from commit 693ca4e6e0363d05de47db660792a460cd5d822e​)

#skipci

diff --git a/Engine/Plugins/Enterprise/GLTFExporter/MigratedFromGit/Source/GLTFExporter/Public/GLTFExportOptions.h b/Engine/Plugins/Enterprise/GLTFExporter/MigratedFromGit/Source/GLTFExporter/Public/GLTFExportOptions.h
--- a/Engine/Plugins/Enterprise/GLTFExporter/MigratedFromGit/Source/GLTFExporter/Public/GLTFExportOptions.h
+++ b/Engine/Plugins/Enterprise/GLTFExporter/MigratedFromGit/Source/GLTFExporter/Public/GLTFExportOptions.h
@@ -172,7 +172,7 @@ class GLTFEXPORTER_API UGLTFExportOptions : public UObject
 	UPROPERTY(EditAnywhere, BlueprintReadWrite, Config, Category = Texture, Meta = (DisplayName = "Texture HDR Encoding", EditCondition = "TextureImageFormat != EGLTFTextureImageFormat::None"))
 	EGLTFTextureHDREncoding TextureHDREncoding;
 
-	/** If enabled, export components that are flagged as hidden in-game. */
+	/** If enabled, export actors and components that are flagged as hidden in-game. */
 	UPROPERTY(EditAnywhere, BlueprintReadWrite, Config, Category = Scene)
 	bool bExportHiddenInGame;
 
