[GLTFExporter] feat: add support for transforming color space in texture converters

commit 1061150e876e1f1c119d9f9eeacea2a2000ad9b5
Date: Wed, 1 Dec 2021 09:46:44 +0100

    feat: add support for transforming color space in texture converters
    (cherry picked from commit 8af78fd47181bde520fcad237d10ed35439e62ef​)

#skipci

diff --git a/Engine/Plugins/Enterprise/GLTFExporter/MigratedFromGit/Source/GLTFExporter/Private/Builders/GLTFConvertBuilder.cpp b/Engine/Plugins/Enterprise/GLTFExporter/MigratedFromGit/Source/GLTFExporter/Private/Builders/GLTFConvertBuilder.cpp
--- a/Engine/Plugins/Enterprise/GLTFExporter/MigratedFromGit/Source/GLTFExporter/Private/Builders/GLTFConvertBuilder.cpp
+++ b/Engine/Plugins/Enterprise/GLTFExporter/MigratedFromGit/Source/GLTFExporter/Private/Builders/GLTFConvertBuilder.cpp
@@ -161,53 +161,73 @@ FGLTFJsonSamplerIndex FGLTFConvertBuilder::GetOrAddSampler(const UTexture* Textu
 }
 
 FGLTFJsonTextureIndex FGLTFConvertBuilder::GetOrAddTexture(const UTexture2D* Texture)
+{
+	return GetOrAddTexture(Texture, Texture->SRGB);
+}
+
+FGLTFJsonTextureIndex FGLTFConvertBuilder::GetOrAddTexture(const UTextureCube* Texture, ECubeFace CubeFace)
+{
+	return GetOrAddTexture(Texture, CubeFace, Texture->SRGB);
+}
+
+FGLTFJsonTextureIndex FGLTFConvertBuilder::GetOrAddTexture(const UTextureRenderTarget2D* Texture)
+{
+	return GetOrAddTexture(Texture, Texture->SRGB);
+}
+
+FGLTFJsonTextureIndex FGLTFConvertBuilder::GetOrAddTexture(const UTextureRenderTargetCube* Texture, ECubeFace CubeFace)
+{
+	return GetOrAddTexture(Texture, CubeFace, Texture->SRGB);
+}
+
+FGLTFJsonTextureIndex FGLTFConvertBuilder::GetOrAddTexture(const ULightMapTexture2D* Texture)
 {
 	if (Texture == nullptr)
 	{
 		return FGLTFJsonTextureIndex(INDEX_NONE);
 	}
 
-	return Texture2DConverter.GetOrAdd(Texture);
+	return TextureLightMapConverter.GetOrAdd(Texture);
 }
 
-FGLTFJsonTextureIndex FGLTFConvertBuilder::GetOrAddTexture(const UTextureCube* Texture, ECubeFace CubeFace)
+FGLTFJsonTextureIndex FGLTFConvertBuilder::GetOrAddTexture(const UTexture2D* Texture, bool bToSRGB)
 {
 	if (Texture == nullptr)
 	{
 		return FGLTFJsonTextureIndex(INDEX_NONE);
 	}
 
-	return TextureCubeConverter.GetOrAdd(Texture, CubeFace);
+	return Texture2DConverter.GetOrAdd(Texture, bToSRGB);
 }
 
-FGLTFJsonTextureIndex FGLTFConvertBuilder::GetOrAddTexture(const UTextureRenderTarget2D* Texture)
+FGLTFJsonTextureIndex FGLTFConvertBuilder::GetOrAddTexture(const UTextureCube* Texture, ECubeFace CubeFace, bool bToSRGB)
 {
 	if (Texture == nullptr)
 	{
 		return FGLTFJsonTextureIndex(INDEX_NONE);
 	}
 
-	return TextureRenderTarget2DConverter.GetOrAdd(Texture);
+	return TextureCubeConverter.GetOrAdd(Texture, CubeFace, bToSRGB);
 }
 
-FGLTFJsonTextureIndex FGLTFConvertBuilder::GetOrAddTexture(const UTextureRenderTargetCube* Texture, ECubeFace CubeFace)
+FGLTFJsonTextureIndex FGLTFConvertBuilder::GetOrAddTexture(const UTextureRenderTarget2D* Texture, bool bToSRGB)
 {
 	if (Texture == nullptr)
 	{
 		return FGLTFJsonTextureIndex(INDEX_NONE);
 	}
 
-	return TextureRenderTargetCubeConverter.GetOrAdd(Texture, CubeFace);
+	return TextureRenderTarget2DConverter.GetOrAdd(Texture, bToSRGB);
 }
 
-FGLTFJsonTextureIndex FGLTFConvertBuilder::GetOrAddTexture(const ULightMapTexture2D* Texture)
+FGLTFJsonTextureIndex FGLTFConvertBuilder::GetOrAddTexture(const UTextureRenderTargetCube* Texture, ECubeFace CubeFace, bool bToSRGB)
 {
 	if (Texture == nullptr)
 	{
 		return FGLTFJsonTextureIndex(INDEX_NONE);
 	}
 
-	return TextureLightMapConverter.GetOrAdd(Texture);
+	return TextureRenderTargetCubeConverter.GetOrAdd(Texture, CubeFace, bToSRGB);
 }
 
 FGLTFJsonSkinIndex FGLTFConvertBuilder::GetOrAddSkin(FGLTFJsonNodeIndex RootNode, const USkeletalMesh* SkeletalMesh)
diff --git a/Engine/Plugins/Enterprise/GLTFExporter/MigratedFromGit/Source/GLTFExporter/Private/Builders/GLTFConvertBuilder.h b/Engine/Plugins/Enterprise/GLTFExporter/MigratedFromGit/Source/GLTFExporter/Private/Builders/GLTFConvertBuilder.h
--- a/Engine/Plugins/Enterprise/GLTFExporter/MigratedFromGit/Source/GLTFExporter/Private/Builders/GLTFConvertBuilder.h
+++ b/Engine/Plugins/Enterprise/GLTFExporter/MigratedFromGit/Source/GLTFExporter/Private/Builders/GLTFConvertBuilder.h
@@ -57,6 +57,10 @@ public:
 	FGLTFJsonTextureIndex GetOrAddTexture(const UTextureRenderTarget2D* Texture);
 	FGLTFJsonTextureIndex GetOrAddTexture(const UTextureRenderTargetCube* Texture, ECubeFace CubeFace);
 	FGLTFJsonTextureIndex GetOrAddTexture(const ULightMapTexture2D* Texture);
+	FGLTFJsonTextureIndex GetOrAddTexture(const UTexture2D* Texture, bool bToSRGB);
+	FGLTFJsonTextureIndex GetOrAddTexture(const UTextureCube* Texture, ECubeFace CubeFace, bool bToSRGB);
+	FGLTFJsonTextureIndex GetOrAddTexture(const UTextureRenderTarget2D* Texture, bool bToSRGB);
+	FGLTFJsonTextureIndex GetOrAddTexture(const UTextureRenderTargetCube* Texture, ECubeFace CubeFace, bool bToSRGB);
 
 	FGLTFJsonSkinIndex GetOrAddSkin(FGLTFJsonNodeIndex RootNode, const USkeletalMesh* SkeletalMesh);
 	FGLTFJsonSkinIndex GetOrAddSkin(FGLTFJsonNodeIndex RootNode, const USkeletalMeshComponent* SkeletalMeshComponent);
diff --git a/Engine/Plugins/Enterprise/GLTFExporter/MigratedFromGit/Source/GLTFExporter/Private/Converters/GLTFTextureConverters.cpp b/Engine/Plugins/Enterprise/GLTFExporter/MigratedFromGit/Source/GLTFExporter/Private/Converters/GLTFTextureConverters.cpp
--- a/Engine/Plugins/Enterprise/GLTFExporter/MigratedFromGit/Source/GLTFExporter/Private/Converters/GLTFTextureConverters.cpp
+++ b/Engine/Plugins/Enterprise/GLTFExporter/MigratedFromGit/Source/GLTFExporter/Private/Converters/GLTFTextureConverters.cpp
@@ -3,43 +3,43 @@
 #include "Converters/GLTFTextureConverters.h"
 #include "Tasks/GLTFTextureTasks.h"
 
-FGLTFJsonTextureIndex FGLTFTexture2DConverter::Convert(const UTexture2D* Texture2D)
+FGLTFJsonTextureIndex FGLTFTexture2DConverter::Convert(const UTexture2D* Texture2D, bool bToSRGB)
 {
 	if (Builder.ExportOptions->TextureImageFormat != EGLTFTextureImageFormat::None)
 	{
 		const FGLTFJsonTextureIndex TextureIndex = Builder.AddTexture();
-		Builder.SetupTask<FGLTFTexture2DTask>(Builder, Texture2D, TextureIndex);
+		Builder.SetupTask<FGLTFTexture2DTask>(Builder, Texture2D, bToSRGB, TextureIndex);
 		return TextureIndex;
 	}
 
 	return FGLTFJsonTextureIndex(INDEX_NONE);
 }
 
-FGLTFJsonTextureIndex FGLTFTextureCubeConverter::Convert(const UTextureCube* TextureCube, ECubeFace CubeFace)
+FGLTFJsonTextureIndex FGLTFTextureCubeConverter::Convert(const UTextureCube* TextureCube, ECubeFace CubeFace, bool bToSRGB)
 {
 	if (Builder.ExportOptions->TextureImageFormat != EGLTFTextureImageFormat::None)
 	{
 		const FGLTFJsonTextureIndex TextureIndex = Builder.AddTexture();
-		Builder.SetupTask<FGLTFTextureCubeTask>(Builder, TextureCube, CubeFace, TextureIndex);
+		Builder.SetupTask<FGLTFTextureCubeTask>(Builder, TextureCube, CubeFace, bToSRGB, TextureIndex);
 		return TextureIndex;
 	}
 
 	return FGLTFJsonTextureIndex(INDEX_NONE);
 }
 
-FGLTFJsonTextureIndex FGLTFTextureRenderTarget2DConverter::Convert(const UTextureRenderTarget2D* RenderTarget2D)
+FGLTFJsonTextureIndex FGLTFTextureRenderTarget2DConverter::Convert(const UTextureRenderTarget2D* RenderTarget2D, bool bToSRGB)
 {
 	if (Builder.ExportOptions->TextureImageFormat != EGLTFTextureImageFormat::None)
 	{
 		const FGLTFJsonTextureIndex TextureIndex = Builder.AddTexture();
-		Builder.SetupTask<FGLTFTextureRenderTarget2DTask>(Builder, RenderTarget2D, TextureIndex);
+		Builder.SetupTask<FGLTFTextureRenderTarget2DTask>(Builder, RenderTarget2D, bToSRGB, TextureIndex);
 		return TextureIndex;
 	}
 
 	return FGLTFJsonTextureIndex(INDEX_NONE);
 }
 
-FGLTFJsonTextureIndex FGLTFTextureRenderTargetCubeConverter::Convert(const UTextureRenderTargetCube* RenderTargetCube, ECubeFace CubeFace)
+FGLTFJsonTextureIndex FGLTFTextureRenderTargetCubeConverter::Convert(const UTextureRenderTargetCube* RenderTargetCube, ECubeFace CubeFace, bool bToSRGB)
 {
 	if (Builder.ExportOptions->TextureImageFormat != EGLTFTextureImageFormat::None)
 	{
@@ -47,7 +47,7 @@ FGLTFJsonTextureIndex FGLTFTextureRenderTargetCubeConverter::Convert(const UText
 	}
 
 	const FGLTFJsonTextureIndex TextureIndex = Builder.AddTexture();
-	Builder.SetupTask<FGLTFTextureRenderTargetCubeTask>(Builder, RenderTargetCube, CubeFace, TextureIndex);
+	Builder.SetupTask<FGLTFTextureRenderTargetCubeTask>(Builder, RenderTargetCube, CubeFace, bToSRGB, TextureIndex);
 	return TextureIndex;
 }
 
diff --git a/Engine/Plugins/Enterprise/GLTFExporter/MigratedFromGit/Source/GLTFExporter/Private/Converters/GLTFTextureConverters.h b/Engine/Plugins/Enterprise/GLTFExporter/MigratedFromGit/Source/GLTFExporter/Private/Converters/GLTFTextureConverters.h
--- a/Engine/Plugins/Enterprise/GLTFExporter/MigratedFromGit/Source/GLTFExporter/Private/Converters/GLTFTextureConverters.h
+++ b/Engine/Plugins/Enterprise/GLTFExporter/MigratedFromGit/Source/GLTFExporter/Private/Converters/GLTFTextureConverters.h
@@ -13,32 +13,32 @@ class TGLTFTextureConverter : public FGLTFBuilderContext, public TGLTFConverter<
 	using FGLTFBuilderContext::FGLTFBuilderContext;
 };
 
-class FGLTFTexture2DConverter final : public TGLTFTextureConverter<const UTexture2D*>
+class FGLTFTexture2DConverter final : public TGLTFTextureConverter<const UTexture2D*, bool>
 {
 	using TGLTFTextureConverter::TGLTFTextureConverter;
 
-	virtual FGLTFJsonTextureIndex Convert(const UTexture2D* Texture2D) override;
+	virtual FGLTFJsonTextureIndex Convert(const UTexture2D* Texture2D, bool bToSRGB) override;
 };
 
-class FGLTFTextureCubeConverter final : public TGLTFTextureConverter<const UTextureCube*, ECubeFace>
+class FGLTFTextureCubeConverter final : public TGLTFTextureConverter<const UTextureCube*, ECubeFace, bool>
 {
 	using TGLTFTextureConverter::TGLTFTextureConverter;
 
-	virtual FGLTFJsonTextureIndex Convert(const UTextureCube* TextureCube, ECubeFace CubeFace) override;
+	virtual FGLTFJsonTextureIndex Convert(const UTextureCube* TextureCube, ECubeFace CubeFace, bool bToSRGB) override;
 };
 
-class FGLTFTextureRenderTarget2DConverter final : public TGLTFTextureConverter<const UTextureRenderTarget2D*>
+class FGLTFTextureRenderTarget2DConverter final : public TGLTFTextureConverter<const UTextureRenderTarget2D*, bool>
 {
 	using TGLTFTextureConverter::TGLTFTextureConverter;
 
-	virtual FGLTFJsonTextureIndex Convert(const UTextureRenderTarget2D* RenderTarget2D) override;
+	virtual FGLTFJsonTextureIndex Convert(const UTextureRenderTarget2D* RenderTarget2D, bool bToSRGB) override;
 };
 
-class FGLTFTextureRenderTargetCubeConverter final : public TGLTFTextureConverter<const UTextureRenderTargetCube*, ECubeFace>
+class FGLTFTextureRenderTargetCubeConverter final : public TGLTFTextureConverter<const UTextureRenderTargetCube*, ECubeFace, bool>
 {
 	using TGLTFTextureConverter::TGLTFTextureConverter;
 
-	virtual FGLTFJsonTextureIndex Convert(const UTextureRenderTargetCube* RenderTargetCube, ECubeFace CubeFace) override;
+	virtual FGLTFJsonTextureIndex Convert(const UTextureRenderTargetCube* RenderTargetCube, ECubeFace CubeFace, bool bToSRGB) override;
 };
 
 class FGLTFTextureLightMapConverter final : public TGLTFTextureConverter<const ULightMapTexture2D*>
diff --git a/Engine/Plugins/Enterprise/GLTFExporter/MigratedFromGit/Source/GLTFExporter/Private/Converters/GLTFTextureUtility.cpp b/Engine/Plugins/Enterprise/GLTFExporter/MigratedFromGit/Source/GLTFExporter/Private/Converters/GLTFTextureUtility.cpp
--- a/Engine/Plugins/Enterprise/GLTFExporter/MigratedFromGit/Source/GLTFExporter/Private/Converters/GLTFTextureUtility.cpp
+++ b/Engine/Plugins/Enterprise/GLTFExporter/MigratedFromGit/Source/GLTFExporter/Private/Converters/GLTFTextureUtility.cpp
@@ -391,3 +391,26 @@ void FGLTFTextureUtility::FlipGreenChannel(TArray<FColor>& Pixels)
 		Pixel.G = 255 - Pixel.G;
 	}
 }
+
+void FGLTFTextureUtility::TransformColorSpace(TArray<FColor>& Pixels, bool bFromSRGB, bool bToSRGB)
+{
+	if (bFromSRGB == bToSRGB)
+	{
+		return;
+	}
+
+	if (bToSRGB)
+	{
+		for (FColor& Pixel: Pixels)
+		{
+			Pixel = Pixel.ReinterpretAsLinear().ToFColor(true);
+		}
+	}
+	else
+	{
+		for (FColor& Pixel: Pixels)
+		{
+			Pixel = FLinearColor(Pixel).ToFColor(false);
+		}
+	}
+}
diff --git a/Engine/Plugins/Enterprise/GLTFExporter/MigratedFromGit/Source/GLTFExporter/Private/Converters/GLTFTextureUtility.h b/Engine/Plugins/Enterprise/GLTFExporter/MigratedFromGit/Source/GLTFExporter/Private/Converters/GLTFTextureUtility.h
--- a/Engine/Plugins/Enterprise/GLTFExporter/MigratedFromGit/Source/GLTFExporter/Private/Converters/GLTFTextureUtility.h
+++ b/Engine/Plugins/Enterprise/GLTFExporter/MigratedFromGit/Source/GLTFExporter/Private/Converters/GLTFTextureUtility.h
@@ -42,4 +42,5 @@ struct FGLTFTextureUtility
 	static bool LoadPlatformData(UTextureCube* TextureCube);
 
 	static void FlipGreenChannel(TArray<FColor>& Pixels);
+	static void TransformColorSpace(TArray<FColor>& Pixels, bool bFromSRGB, bool bToSRGB);
 };
diff --git a/Engine/Plugins/Enterprise/GLTFExporter/MigratedFromGit/Source/GLTFExporter/Private/Tasks/GLTFTextureTasks.cpp b/Engine/Plugins/Enterprise/GLTFExporter/MigratedFromGit/Source/GLTFExporter/Private/Tasks/GLTFTextureTasks.cpp
--- a/Engine/Plugins/Enterprise/GLTFExporter/MigratedFromGit/Source/GLTFExporter/Private/Tasks/GLTFTextureTasks.cpp
+++ b/Engine/Plugins/Enterprise/GLTFExporter/MigratedFromGit/Source/GLTFExporter/Private/Tasks/GLTFTextureTasks.cpp
@@ -11,6 +11,12 @@ void FGLTFTexture2DTask::Complete()
 	FGLTFJsonTexture& JsonTexture = Builder.GetTexture(TextureIndex);
 	Texture2D->GetName(JsonTexture.Name);
 
+	const bool bFromSRGB = Texture2D->SRGB;
+	if (bFromSRGB != bToSRGB)
+	{
+		JsonTexture.Name += bToSRGB ? TEXT("_sRGB") : TEXT("_Linear");
+	}
+
 	const bool bIsHDR = FGLTFTextureUtility::IsHDR(Texture2D->GetPixelFormat());
 	const FIntPoint Size = FGLTFTextureUtility::GetInGameSize(Texture2D);
 	UTextureRenderTarget2D* RenderTarget = FGLTFTextureUtility::CreateRenderTarget(Size, bIsHDR);
@@ -35,6 +41,8 @@ void FGLTFTexture2DTask::Complete()
 		FGLTFTextureUtility::FlipGreenChannel(Pixels);
 	}
 
+	FGLTFTextureUtility::TransformColorSpace(Pixels, bFromSRGB, bToSRGB);
+
 	const bool bIgnoreAlpha = FGLTFTextureUtility::IsAlphaless(Texture2D->GetPixelFormat());
 	const EGLTFTextureType Type =
 		Texture2D->IsNormalMap() ? EGLTFTextureType::Normalmaps :
@@ -49,6 +57,12 @@ void FGLTFTextureCubeTask::Complete()
 	FGLTFJsonTexture& JsonTexture = Builder.GetTexture(TextureIndex);
 	JsonTexture.Name = TextureCube->GetName() + TEXT("_") + FGLTFJsonUtility::GetValue(FGLTFConverterUtility::ConvertCubeFace(CubeFace));
 
+	const bool bFromSRGB = TextureCube->SRGB;
+	if (bFromSRGB != bToSRGB)
+	{
+		JsonTexture.Name += bToSRGB ? TEXT("_sRGB") : TEXT("_Linear");
+	}
+
 	// TODO: add optimized "happy path" if cube face doesn't need rotation and has suitable pixel format
 
 	// TODO: preserve maximum image quality (avoid compression artifacts) by copying source data (and adjustments) to a temp texture
@@ -78,6 +92,8 @@ void FGLTFTextureCubeTask::Complete()
 		return;
 	}
 
+	FGLTFTextureUtility::TransformColorSpace(Pixels, bFromSRGB, bToSRGB);
+
 	const bool bIgnoreAlpha = FGLTFTextureUtility::IsAlphaless(TextureCube->GetPixelFormat());
 	const EGLTFTextureType Type = bIsHDR ? EGLTFTextureType::HDR : EGLTFTextureType::None;
 
@@ -90,6 +106,12 @@ void FGLTFTextureRenderTarget2DTask::Complete()
 	FGLTFJsonTexture& JsonTexture = Builder.GetTexture(TextureIndex);
 	RenderTarget2D->GetName(JsonTexture.Name);
 
+	const bool bFromSRGB = RenderTarget2D->SRGB;
+	if (bFromSRGB != bToSRGB)
+	{
+		JsonTexture.Name += bToSRGB ? TEXT("_sRGB") : TEXT("_Linear");
+	}
+
 	const bool bIsHDR = FGLTFTextureUtility::IsHDR(RenderTarget2D->GetFormat());
 	const FIntPoint Size = { RenderTarget2D->SizeX, RenderTarget2D->SizeY };
 
@@ -105,6 +127,8 @@ void FGLTFTextureRenderTarget2DTask::Complete()
 		return;
 	}
 
+	FGLTFTextureUtility::TransformColorSpace(Pixels, bFromSRGB, bToSRGB);
+
 	const bool bIgnoreAlpha = FGLTFTextureUtility::IsAlphaless(RenderTarget2D->GetFormat());
 	const EGLTFTextureType Type = bIsHDR ? EGLTFTextureType::HDR : EGLTFTextureType::None;
 
@@ -117,6 +141,12 @@ void FGLTFTextureRenderTargetCubeTask::Complete()
 	FGLTFJsonTexture& JsonTexture = Builder.GetTexture(TextureIndex);
 	JsonTexture.Name = RenderTargetCube->GetName() + TEXT("_") + FGLTFJsonUtility::GetValue(FGLTFConverterUtility::ConvertCubeFace(CubeFace));
 
+	const bool bFromSRGB = RenderTargetCube->SRGB;
+	if (bFromSRGB != bToSRGB)
+	{
+		JsonTexture.Name += bToSRGB ? TEXT("_sRGB") : TEXT("_Linear");
+	}
+
 	// TODO: add optimized "happy path" if cube face doesn't need rotation
 
 	const UTexture2D* FaceTexture = FGLTFTextureUtility::CreateTextureFromCubeFace(RenderTargetCube, CubeFace);
@@ -145,6 +175,8 @@ void FGLTFTextureRenderTargetCubeTask::Complete()
 		return;
 	}
 
+	FGLTFTextureUtility::TransformColorSpace(Pixels, bFromSRGB, bToSRGB);
+
 	const bool bIgnoreAlpha = FGLTFTextureUtility::IsAlphaless(RenderTargetCube->GetFormat());
 	const EGLTFTextureType Type = bIsHDR ? EGLTFTextureType::HDR : EGLTFTextureType::None;
 
diff --git a/Engine/Plugins/Enterprise/GLTFExporter/MigratedFromGit/Source/GLTFExporter/Private/Tasks/GLTFTextureTasks.h b/Engine/Plugins/Enterprise/GLTFExporter/MigratedFromGit/Source/GLTFExporter/Private/Tasks/GLTFTextureTasks.h
--- a/Engine/Plugins/Enterprise/GLTFExporter/MigratedFromGit/Source/GLTFExporter/Private/Tasks/GLTFTextureTasks.h
+++ b/Engine/Plugins/Enterprise/GLTFExporter/MigratedFromGit/Source/GLTFExporter/Private/Tasks/GLTFTextureTasks.h
@@ -10,10 +10,11 @@ class FGLTFTexture2DTask : public FGLTFTask
 {
 public:
 
-	FGLTFTexture2DTask(FGLTFConvertBuilder& Builder, const UTexture2D* Texture2D, FGLTFJsonTextureIndex TextureIndex)
+	FGLTFTexture2DTask(FGLTFConvertBuilder& Builder, const UTexture2D* Texture2D, bool bToSRGB, FGLTFJsonTextureIndex TextureIndex)
 		: FGLTFTask(EGLTFTaskPriority::Texture)
 		, Builder(Builder)
 		, Texture2D(Texture2D)
+		, bToSRGB(bToSRGB)
 		, TextureIndex(TextureIndex)
 	{
 	}
@@ -29,6 +30,7 @@ private:
 
 	FGLTFConvertBuilder& Builder;
 	const UTexture2D* Texture2D;
+	bool bToSRGB;
 	const FGLTFJsonTextureIndex TextureIndex;
 };
 
@@ -36,11 +38,12 @@ class FGLTFTextureCubeTask : public FGLTFTask
 {
 public:
 
-	FGLTFTextureCubeTask(FGLTFConvertBuilder& Builder, const UTextureCube* TextureCube, ECubeFace CubeFace, FGLTFJsonTextureIndex TextureIndex)
+	FGLTFTextureCubeTask(FGLTFConvertBuilder& Builder, const UTextureCube* TextureCube, ECubeFace CubeFace, bool bToSRGB, FGLTFJsonTextureIndex TextureIndex)
 		: FGLTFTask(EGLTFTaskPriority::Texture)
 		, Builder(Builder)
 		, TextureCube(TextureCube)
 		, CubeFace(CubeFace)
+		, bToSRGB(bToSRGB)
 		, TextureIndex(TextureIndex)
 	{
 	}
@@ -57,6 +60,7 @@ private:
 	FGLTFConvertBuilder& Builder;
 	const UTextureCube* TextureCube;
 	ECubeFace CubeFace;
+	bool bToSRGB;
 	const FGLTFJsonTextureIndex TextureIndex;
 };
 
@@ -64,10 +68,11 @@ class FGLTFTextureRenderTarget2DTask : public FGLTFTask
 {
 public:
 
-	FGLTFTextureRenderTarget2DTask(FGLTFConvertBuilder& Builder, const UTextureRenderTarget2D* RenderTarget2D, FGLTFJsonTextureIndex TextureIndex)
+	FGLTFTextureRenderTarget2DTask(FGLTFConvertBuilder& Builder, const UTextureRenderTarget2D* RenderTarget2D, bool bToSRGB, FGLTFJsonTextureIndex TextureIndex)
 		: FGLTFTask(EGLTFTaskPriority::Texture)
 		, Builder(Builder)
 		, RenderTarget2D(RenderTarget2D)
+		, bToSRGB(bToSRGB)
 		, TextureIndex(TextureIndex)
 	{
 	}
@@ -83,6 +88,7 @@ private:
 
 	FGLTFConvertBuilder& Builder;
 	const UTextureRenderTarget2D* RenderTarget2D;
+	bool bToSRGB;
 	const FGLTFJsonTextureIndex TextureIndex;
 };
 
@@ -90,11 +96,12 @@ class FGLTFTextureRenderTargetCubeTask : public FGLTFTask
 {
 public:
 
-	FGLTFTextureRenderTargetCubeTask(FGLTFConvertBuilder& Builder, const UTextureRenderTargetCube* RenderTargetCube, ECubeFace CubeFace, FGLTFJsonTextureIndex TextureIndex)
+	FGLTFTextureRenderTargetCubeTask(FGLTFConvertBuilder& Builder, const UTextureRenderTargetCube* RenderTargetCube, ECubeFace CubeFace, bool bToSRGB, FGLTFJsonTextureIndex TextureIndex)
 		: FGLTFTask(EGLTFTaskPriority::Texture)
 		, Builder(Builder)
 		, RenderTargetCube(RenderTargetCube)
 		, CubeFace(CubeFace)
+		, bToSRGB(bToSRGB)
 		, TextureIndex(TextureIndex)
 	{
 	}
@@ -111,6 +118,7 @@ private:
 	FGLTFConvertBuilder& Builder;
 	const UTextureRenderTargetCube* RenderTargetCube;
 	ECubeFace CubeFace;
+	bool bToSRGB;
 	const FGLTFJsonTextureIndex TextureIndex;
 };
 
