[GLTFExporter] fix: fix incorrect color space when exporting some material properties

commit 1ac521296b5397c57d905e624cbca37a24f56044
Date: Wed, 1 Dec 2021 10:11:02 +0100

    fix: fix incorrect color space when exporting some material properties
    (cherry picked from commit 2c5bc414fab2d7957887be830071e985b0035568​)

#skipci

diff --git a/Engine/Plugins/Enterprise/GLTFExporter/MigratedFromGit/Source/GLTFExporter/Private/Converters/GLTFMaterialUtility.cpp b/Engine/Plugins/Enterprise/GLTFExporter/MigratedFromGit/Source/GLTFExporter/Private/Converters/GLTFMaterialUtility.cpp
--- a/Engine/Plugins/Enterprise/GLTFExporter/MigratedFromGit/Source/GLTFExporter/Private/Converters/GLTFMaterialUtility.cpp
+++ b/Engine/Plugins/Enterprise/GLTFExporter/MigratedFromGit/Source/GLTFExporter/Private/Converters/GLTFMaterialUtility.cpp
@@ -32,6 +32,11 @@ bool FGLTFMaterialUtility::IsNormalMap(const FMaterialPropertyEx& Property)
 	return Property == MP_Normal || Property == TEXT("ClearCoatBottomNormal");
 }
 
+bool FGLTFMaterialUtility::IsSRGB(const FMaterialPropertyEx& Property)
+{
+	return Property == MP_BaseColor || Property == MP_EmissiveColor || Property == MP_SubsurfaceColor || Property == TEXT("TransmittanceColor");
+}
+
 FGuid FGLTFMaterialUtility::GetAttributeID(const FMaterialPropertyEx& Property)
 {
 	return Property.IsCustomOutput()
@@ -189,7 +194,10 @@ FGLTFPropertyBakeOutput FGLTFMaterialUtility::BakeMaterialProperty(const FIntPoi
 	TArray<FBakeOutputEx> BakeOutputs;
 	IMaterialBakingModule& Module = FModuleManager::Get().LoadModuleChecked<IMaterialBakingModule>("GLTFMaterialBaking");
 
+	//Module.SetLinearBake(true);
 	Module.BakeMaterials(MatSettings, MeshSettings, BakeOutputs);
+	const bool bIsLinearBake = Module.IsLinearBake(Property);
+	//Module.SetLinearBake(false);
 
 	FBakeOutputEx& BakeOutput = BakeOutputs[0];
 
@@ -220,14 +228,18 @@ FGLTFPropertyBakeOutput FGLTFMaterialUtility::BakeMaterialProperty(const FIntPoi
 		FGLTFTextureUtility::FlipGreenChannel(BakedPixels);
 	}
 
-	FGLTFPropertyBakeOutput PropertyBakeOutput(Property, PF_B8G8R8A8, BakedPixels, BakedSize, EmissiveScale);
+	bool bFromSRGB = !bIsLinearBake;
+	bool bToSRGB = IsSRGB(Property);
+	FGLTFTextureUtility::TransformColorSpace(BakedPixels, bFromSRGB, bToSRGB);
+
+	FGLTFPropertyBakeOutput PropertyBakeOutput(Property, PF_B8G8R8A8, BakedPixels, BakedSize, EmissiveScale, !bIsLinearBake);
 
 	if (BakedPixels.Num() == 1)
 	{
 		const FColor& Pixel = BakedPixels[0];
 
 		PropertyBakeOutput.bIsConstant = true;
-		PropertyBakeOutput.ConstantValue = Module.IsLinearBake(Property) ? Pixel.ReinterpretAsLinear() : FLinearColor(Pixel);;
+		PropertyBakeOutput.ConstantValue = bToSRGB ? FLinearColor(Pixel) : Pixel.ReinterpretAsLinear();
 	}
 
 	return PropertyBakeOutput;
@@ -267,14 +279,6 @@ FGLTFJsonTextureIndex FGLTFMaterialUtility::AddTexture(FGLTFConvertBuilder& Buil
 	return Builder.AddTexture(JsonTexture);
 }
 
-void FGLTFMaterialUtility::TransformToLinear(TArray<FColor>& InOutPixels)
-{
-	for (FColor& Pixel: InOutPixels)
-	{
-		Pixel = FLinearColor(Pixel).ToFColor(false);
-	}
-}
-
 FLinearColor FGLTFMaterialUtility::GetMask(const FExpressionInput& ExpressionInput)
 {
 	return FLinearColor(ExpressionInput.MaskR, ExpressionInput.MaskG, ExpressionInput.MaskB, ExpressionInput.MaskA);
diff --git a/Engine/Plugins/Enterprise/GLTFExporter/MigratedFromGit/Source/GLTFExporter/Private/Tasks/GLTFMaterialTasks.cpp b/Engine/Plugins/Enterprise/GLTFExporter/MigratedFromGit/Source/GLTFExporter/Private/Tasks/GLTFMaterialTasks.cpp
--- a/Engine/Plugins/Enterprise/GLTFExporter/MigratedFromGit/Source/GLTFExporter/Private/Tasks/GLTFMaterialTasks.cpp
+++ b/Engine/Plugins/Enterprise/GLTFExporter/MigratedFromGit/Source/GLTFExporter/Private/Tasks/GLTFMaterialTasks.cpp
@@ -129,7 +129,7 @@ void FGLTFMaterialTask::Complete()
 			{
 				if (!TryGetSourceTexture(JsonMaterial.OcclusionTexture, AmbientOcclusionProperty, OcclusionInputMasks))
 				{
-					if (!TryGetBakedMaterialProperty(JsonMaterial.OcclusionTexture, AmbientOcclusionProperty, TEXT("Occlusion"), true))
+					if (!TryGetBakedMaterialProperty(JsonMaterial.OcclusionTexture, AmbientOcclusionProperty, TEXT("Occlusion")))
 					{
 						Builder.LogWarning(FString::Printf(TEXT("Failed to export %s for material %s"), *AmbientOcclusionProperty.ToString(), *Material->GetName()));
 					}
@@ -362,7 +362,7 @@ bool FGLTFMaterialTask::TryGetBaseColorAndOpacity(FGLTFJsonPBRMetallicRoughness&
 		BaseColorTexCoord == OpacityTexCoord &&
 		BaseColorTransform.IsExactlyEqual(OpacityTransform))
 	{
-		OutPBRParams.BaseColorTexture.Index = Builder.GetOrAddTexture(BaseColorTexture);
+		OutPBRParams.BaseColorTexture.Index = Builder.GetOrAddTexture(BaseColorTexture, true);
 		OutPBRParams.BaseColorTexture.TexCoord = BaseColorTexCoord;
 		OutPBRParams.BaseColorTexture.Transform = BaseColorTransform;
 		return true;
@@ -487,7 +487,7 @@ bool FGLTFMaterialTask::TryGetMetallicAndRoughness(FGLTFJsonPBRMetallicRoughness
 		MetallicTexCoord == RoughnessTexCoord &&
 		MetallicTransform.IsExactlyEqual(RoughnessTransform))
 	{
-		OutPBRParams.MetallicRoughnessTexture.Index = Builder.GetOrAddTexture(MetallicTexture);
+		OutPBRParams.MetallicRoughnessTexture.Index = Builder.GetOrAddTexture(MetallicTexture, false);
 		OutPBRParams.MetallicRoughnessTexture.TexCoord = MetallicTexCoord;
 		OutPBRParams.MetallicRoughnessTexture.Transform = MetallicTransform;
 		return true;
@@ -549,9 +549,6 @@ bool FGLTFMaterialTask::TryGetMetallicAndRoughness(FGLTFJsonPBRMetallicRoughness
 		return false;
 	}
 
-	FGLTFMaterialUtility::TransformToLinear(MetallicBakeOutput.Pixels);
-	FGLTFMaterialUtility::TransformToLinear(RoughnessBakeOutput.Pixels);
-
 	TextureSize = RoughnessBakeOutput.Size.ComponentMax(MetallicBakeOutput.Size);
 	MetallicTexture = FGLTFMaterialUtility::CreateTransientTexture(MetallicBakeOutput);
 	RoughnessTexture = FGLTFMaterialUtility::CreateTransientTexture(RoughnessBakeOutput);
@@ -612,7 +609,7 @@ bool FGLTFMaterialTask::TryGetClearCoatRoughness(FGLTFJsonClearCoatExtension& Ou
 		IntensityTexCoord == RoughnessTexCoord &&
 		IntensityTransform.IsExactlyEqual(RoughnessTransform))
 	{
-		const FGLTFJsonTextureIndex TextureIndex = Builder.GetOrAddTexture(IntensityTexture);
+		const FGLTFJsonTextureIndex TextureIndex = Builder.GetOrAddTexture(IntensityTexture, false);
 		OutExtParams.ClearCoatTexture.Index = TextureIndex;
 		OutExtParams.ClearCoatTexture.TexCoord = IntensityTexCoord;
 		OutExtParams.ClearCoatRoughnessTexture.Index = TextureIndex;
@@ -1022,7 +1019,7 @@ bool FGLTFMaterialTask::TryGetSourceTexture(FGLTFJsonTextureInfo& OutTexInfo, co
 
 	if (TryGetSourceTexture(Texture, TexCoord, Transform, Property, AllowedMasks))
 	{
-		OutTexInfo.Index = Builder.GetOrAddTexture(Texture);
+		OutTexInfo.Index = Builder.GetOrAddTexture(Texture, FGLTFMaterialUtility::IsSRGB(Property));
 		OutTexInfo.TexCoord = TexCoord;
 		OutTexInfo.Transform = Transform;
 		return true;
@@ -1147,7 +1144,7 @@ bool FGLTFMaterialTask::TryGetSourceTexture(const UTexture2D*& OutTexture, int32
 	return false;
 }
 
-bool FGLTFMaterialTask::TryGetBakedMaterialProperty(FGLTFJsonTextureInfo& OutTexInfo, FGLTFJsonColor3& OutConstant, const FMaterialPropertyEx& Property, const FString& PropertyName, bool bTransformToLinear)
+bool FGLTFMaterialTask::TryGetBakedMaterialProperty(FGLTFJsonTextureInfo& OutTexInfo, FGLTFJsonColor3& OutConstant, const FMaterialPropertyEx& Property, const FString& PropertyName)
 {
 	if (Builder.ExportOptions->BakeMaterialInputs == EGLTFMaterialBakeMode::Disabled)
 	{
@@ -1172,11 +1169,6 @@ bool FGLTFMaterialTask::TryGetBakedMaterialProperty(FGLTFJsonTextureInfo& OutTex
 		return true;
 	}
 
-	if (bTransformToLinear)
-	{
-		FGLTFMaterialUtility::TransformToLinear(PropertyBakeOutput.Pixels);
-	}
-
 	if (StoreBakedPropertyTexture(OutTexInfo, PropertyBakeOutput, PropertyName))
 	{
 		OutConstant = FGLTFJsonColor3::White; // make sure property is not zero
@@ -1186,7 +1178,7 @@ bool FGLTFMaterialTask::TryGetBakedMaterialProperty(FGLTFJsonTextureInfo& OutTex
 	return false;
 }
 
-bool FGLTFMaterialTask::TryGetBakedMaterialProperty(FGLTFJsonTextureInfo& OutTexInfo, FGLTFJsonColor4& OutConstant, const FMaterialPropertyEx& Property, const FString& PropertyName, bool bTransformToLinear)
+bool FGLTFMaterialTask::TryGetBakedMaterialProperty(FGLTFJsonTextureInfo& OutTexInfo, FGLTFJsonColor4& OutConstant, const FMaterialPropertyEx& Property, const FString& PropertyName)
 {
 	if (Builder.ExportOptions->BakeMaterialInputs == EGLTFMaterialBakeMode::Disabled)
 	{
@@ -1211,11 +1203,6 @@ bool FGLTFMaterialTask::TryGetBakedMaterialProperty(FGLTFJsonTextureInfo& OutTex
 		return true;
 	}
 
-	if (bTransformToLinear)
-	{
-		FGLTFMaterialUtility::TransformToLinear(PropertyBakeOutput.Pixels);
-	}
-
 	if (StoreBakedPropertyTexture(OutTexInfo, PropertyBakeOutput, PropertyName))
 	{
 		OutConstant = FGLTFJsonColor4::White; // make sure property is not zero
@@ -1225,7 +1212,7 @@ bool FGLTFMaterialTask::TryGetBakedMaterialProperty(FGLTFJsonTextureInfo& OutTex
 	return false;
 }
 
-inline bool FGLTFMaterialTask::TryGetBakedMaterialProperty(FGLTFJsonTextureInfo& OutTexInfo, float& OutConstant, const FMaterialPropertyEx& Property, const FString& PropertyName, bool bTransformToLinear)
+inline bool FGLTFMaterialTask::TryGetBakedMaterialProperty(FGLTFJsonTextureInfo& OutTexInfo, float& OutConstant, const FMaterialPropertyEx& Property, const FString& PropertyName)
 {
 	if (Builder.ExportOptions->BakeMaterialInputs == EGLTFMaterialBakeMode::Disabled)
 	{
@@ -1250,11 +1237,6 @@ inline bool FGLTFMaterialTask::TryGetBakedMaterialProperty(FGLTFJsonTextureInfo&
 		return true;
 	}
 
-	if (bTransformToLinear)
-	{
-		FGLTFMaterialUtility::TransformToLinear(PropertyBakeOutput.Pixels);
-	}
-
 	if (StoreBakedPropertyTexture(OutTexInfo, PropertyBakeOutput, PropertyName))
 	{
 		OutConstant = 1; // make sure property is not zero
@@ -1264,7 +1246,7 @@ inline bool FGLTFMaterialTask::TryGetBakedMaterialProperty(FGLTFJsonTextureInfo&
 	return false;
 }
 
-bool FGLTFMaterialTask::TryGetBakedMaterialProperty(FGLTFJsonTextureInfo& OutTexInfo, const FMaterialPropertyEx& Property, const FString& PropertyName, bool bTransformToLinear)
+bool FGLTFMaterialTask::TryGetBakedMaterialProperty(FGLTFJsonTextureInfo& OutTexInfo, const FMaterialPropertyEx& Property, const FString& PropertyName)
 {
 	if (Builder.ExportOptions->BakeMaterialInputs == EGLTFMaterialBakeMode::Disabled)
 	{
@@ -1285,11 +1267,6 @@ bool FGLTFMaterialTask::TryGetBakedMaterialProperty(FGLTFJsonTextureInfo& OutTex
 			return true;
 		}
 
-		if (bTransformToLinear)
-		{
-			FGLTFMaterialUtility::TransformToLinear(PropertyBakeOutput.Pixels);
-		}
-
 		return StoreBakedPropertyTexture(OutTexInfo, PropertyBakeOutput, PropertyName);
 	}
 
