[GLTFExporter] refactor: simplify and optimize packing of baked pixels from multiple material properties into same texture

commit 06f98fd4d590f7ff2eb134b32fd7a69c5c0c21ed
Date: Wed, 1 Dec 2021 13:31:28 +0100

    refactor: simplify and optimize packing of baked pixels from multiple material properties into same texture
    (cherry picked from commit 18a9c4858c0eb1b2df1e200fbb002fa4b089dd26​)

#skipci

diff --git a/Engine/Plugins/Enterprise/GLTFExporter/MigratedFromGit/Source/GLTFExporter/Private/Converters/GLTFMaterialUtility.cpp b/Engine/Plugins/Enterprise/GLTFExporter/MigratedFromGit/Source/GLTFExporter/Private/Converters/GLTFMaterialUtility.cpp
--- a/Engine/Plugins/Enterprise/GLTFExporter/MigratedFromGit/Source/GLTFExporter/Private/Converters/GLTFMaterialUtility.cpp
+++ b/Engine/Plugins/Enterprise/GLTFExporter/MigratedFromGit/Source/GLTFExporter/Private/Converters/GLTFMaterialUtility.cpp
@@ -100,73 +100,7 @@ const UMaterialExpressionCustomOutput* FGLTFMaterialUtility::GetCustomOutputByNa
 	return nullptr;
 }
 
-UTexture2D* FGLTFMaterialUtility::CreateTransientTexture(const FGLTFPropertyBakeOutput& PropertyBakeOutput, bool bUseSRGB)
-{
-	return FGLTFTextureUtility::CreateTransientTexture(
-		PropertyBakeOutput.Pixels.GetData(),
-		PropertyBakeOutput.Pixels.Num() * PropertyBakeOutput.Pixels.GetTypeSize(),
-		PropertyBakeOutput.Size,
-		PropertyBakeOutput.PixelFormat,
-		bUseSRGB);
-}
-
-bool FGLTFMaterialUtility::CombineTextures(TArray<FColor>& OutPixels, const TArray<FGLTFTextureCombineSource>& Sources, const FIntPoint& OutputSize, const EPixelFormat OutputPixelFormat)
-{
-	// NOTE: both bForceLinearGamma and TargetGamma=2.2 seem necessary for exported images to match their source data.
-	// It's not entirely clear why gamma must be 2.2 (instead of 0.0) and why bInForceLinearGamma must also be true.
-	const bool bForceLinearGamma = true;
-	const float TargetGamma = 2.2f;
-
-	UTextureRenderTarget2D* RenderTarget2D = NewObject<UTextureRenderTarget2D>();
-
-	RenderTarget2D->AddToRoot();
-	RenderTarget2D->ClearColor = { 0.0f, 0.0f, 0.0f, 0.0f };
-	RenderTarget2D->InitCustomFormat(OutputSize.X, OutputSize.Y, OutputPixelFormat, bForceLinearGamma);
-	RenderTarget2D->TargetGamma = TargetGamma;
-
-	FRenderTarget* RenderTarget = RenderTarget2D->GameThread_GetRenderTargetResource();
-	FCanvas Canvas(RenderTarget, nullptr, 0.0f, 0.0f, 0.0f, GMaxRHIFeatureLevel);
-
-	Canvas.SetRenderTarget_GameThread(RenderTarget);
-	Canvas.Clear({ 0.0f, 0.0f, 0.0f, 0.0f });
-
-	const FVector2D TileSize(OutputSize.X, OutputSize.Y);
-	const FVector2D TilePosition(0.0f, 0.0f);
-
-	for (const FGLTFTextureCombineSource& Source: Sources)
-	{
-		check(Source.Texture);
-
-		TRefCountPtr<FBatchedElementParameters> BatchedElementParameters;
-
-		if (Source.Texture->IsNormalMap())
-		{
-			BatchedElementParameters = new FNormalMapBatchedElementParameters();
-		}
-
-		FCanvasTileItem TileItem(TilePosition, Source.Texture->Resource, TileSize, Source.TintColor);
-
-		TileItem.BatchedElementParameters = BatchedElementParameters;
-		TileItem.BlendMode = Source.BlendMode;
-		TileItem.Draw(&Canvas);
-	}
-
-	Canvas.Flush_GameThread();
-	FlushRenderingCommands();
-	Canvas.SetRenderTarget_GameThread(nullptr);
-	FlushRenderingCommands();
-
-	const bool bReadSuccessful = RenderTarget->ReadPixels(OutPixels);
-
-	// Clean up.
-	RenderTarget2D->ReleaseResource();
-	RenderTarget2D->RemoveFromRoot();
-	RenderTarget2D = nullptr;
-
-	return bReadSuccessful;
-}
-
-FGLTFPropertyBakeOutput FGLTFMaterialUtility::BakeMaterialProperty(const FIntPoint& OutputSize, const FMaterialPropertyEx& Property, const UMaterialInterface* Material, int32 TexCoord, const FGLTFMeshData* MeshData, const FGLTFIndexArray& MeshSectionIndices, bool bCopyAlphaFromRedChannel)
+FGLTFPropertyBakeOutput FGLTFMaterialUtility::BakeMaterialProperty(const FIntPoint& OutputSize, const FMaterialPropertyEx& Property, const UMaterialInterface* Material, int32 TexCoord, const FGLTFMeshData* MeshData, const FGLTFIndexArray& MeshSectionIndices, bool bFillAlpha)
 {
 	FMeshData MeshSet;
 	MeshSet.TextureCoordinateBox = { { 0.0f, 0.0f }, { 1.0f, 1.0f } };
@@ -205,14 +139,7 @@ FGLTFPropertyBakeOutput FGLTFMaterialUtility::BakeMaterialProperty(const FIntPoi
 	const FIntPoint BakedSize = BakeOutput.PropertySizes.FindChecked(Property);
 	const float EmissiveScale = BakeOutput.EmissiveScale;
 
-	if (bCopyAlphaFromRedChannel)
-	{
-		for (FColor& Pixel: BakedPixels)
-		{
-			Pixel.A = Pixel.R;
-		}
-	}
-	else
+	if (bFillAlpha)
 	{
 		// NOTE: alpha is 0 by default after baking a property, but we prefer 255 (1.0).
 		// It makes it easier to view the exported textures.
@@ -245,21 +172,6 @@ FGLTFPropertyBakeOutput FGLTFMaterialUtility::BakeMaterialProperty(const FIntPoi
 	return PropertyBakeOutput;
 }
 
-FGLTFJsonTextureIndex FGLTFMaterialUtility::AddCombinedTexture(FGLTFConvertBuilder& Builder, const TArray<FGLTFTextureCombineSource>& CombineSources, const FIntPoint& TextureSize, bool bIgnoreAlpha, const FString& TextureName, EGLTFJsonTextureFilter MinFilter, EGLTFJsonTextureFilter MagFilter, EGLTFJsonTextureWrap WrapS, EGLTFJsonTextureWrap WrapT)
-{
-	check(CombineSources.Num() > 0);
-
-	TArray<FColor> Pixels;
-	const EPixelFormat PixelFormat = CombineSources[0].Texture->GetPixelFormat(); // TODO: should we really assume pixel format like this?
-
-	if (!CombineTextures(Pixels, CombineSources, TextureSize, PixelFormat))
-	{
-		return FGLTFJsonTextureIndex(INDEX_NONE);
-	}
-
-	return AddTexture(Builder, Pixels, TextureSize, bIgnoreAlpha, false, TextureName, MinFilter, MagFilter, WrapS, WrapT);
-}
-
 FGLTFJsonTextureIndex FGLTFMaterialUtility::AddTexture(FGLTFConvertBuilder& Builder, const TArray<FColor>& Pixels, const FIntPoint& TextureSize, bool bIgnoreAlpha, bool bIsNormalMap, const FString& TextureName, EGLTFJsonTextureFilter MinFilter, EGLTFJsonTextureFilter MagFilter, EGLTFJsonTextureWrap WrapS, EGLTFJsonTextureWrap WrapT)
 {
 	// TODO: maybe we should reuse existing samplers?
diff --git a/Engine/Plugins/Enterprise/GLTFExporter/MigratedFromGit/Source/GLTFExporter/Private/Converters/GLTFMaterialUtility.h b/Engine/Plugins/Enterprise/GLTFExporter/MigratedFromGit/Source/GLTFExporter/Private/Converters/GLTFMaterialUtility.h
--- a/Engine/Plugins/Enterprise/GLTFExporter/MigratedFromGit/Source/GLTFExporter/Private/Converters/GLTFMaterialUtility.h
+++ b/Engine/Plugins/Enterprise/GLTFExporter/MigratedFromGit/Source/GLTFExporter/Private/Converters/GLTFMaterialUtility.h
@@ -60,11 +60,8 @@ struct FGLTFMaterialUtility
 
 	static const UMaterialExpressionCustomOutput* GetCustomOutputByName(const UMaterialInterface* Material, const FString& Name);
 
-	static UTexture2D* CreateTransientTexture(const FGLTFPropertyBakeOutput& PropertyBakeOutput, bool bUseSRGB = false);
-	static bool CombineTextures(TArray<FColor>& OutPixels, const TArray<FGLTFTextureCombineSource>& Sources, const FIntPoint& OutputSize, EPixelFormat OutputPixelFormat);
-	static FGLTFPropertyBakeOutput BakeMaterialProperty(const FIntPoint& OutputSize, const FMaterialPropertyEx& Property, const UMaterialInterface* Material, int32 TexCoord, const FGLTFMeshData* MeshData = nullptr, const FGLTFIndexArray& MeshSectionIndices = {}, bool bCopyAlphaFromRedChannel = false);
+	static FGLTFPropertyBakeOutput BakeMaterialProperty(const FIntPoint& OutputSize, const FMaterialPropertyEx& Property, const UMaterialInterface* Material, int32 TexCoord, const FGLTFMeshData* MeshData = nullptr, const FGLTFIndexArray& MeshSectionIndices = {}, bool bFillAlpha = true);
 
-	static FGLTFJsonTextureIndex AddCombinedTexture(FGLTFConvertBuilder& Builder, const TArray<FGLTFTextureCombineSource>& CombineSources, const FIntPoint& TextureSize, bool bIgnoreAlpha, const FString& TextureName, EGLTFJsonTextureFilter MinFilter, EGLTFJsonTextureFilter MagFilter, EGLTFJsonTextureWrap WrapS, EGLTFJsonTextureWrap WrapT);
 	static FGLTFJsonTextureIndex AddTexture(FGLTFConvertBuilder& Builder, const TArray<FColor>& Pixels, const FIntPoint& TextureSize, bool bIgnoreAlpha, bool bIsNormalMap, const FString& TextureName, EGLTFJsonTextureFilter MinFilter, EGLTFJsonTextureFilter MagFilter, EGLTFJsonTextureWrap WrapS, EGLTFJsonTextureWrap WrapT);
 
 	static FLinearColor GetMask(const FExpressionInput& ExpressionInput);
@@ -78,6 +75,7 @@ struct FGLTFMaterialUtility
 	static FString ShadingModelsToString(const FMaterialShadingModelField& ShadingModels);
 
 	static void AnalyzeMaterialProperty(const UMaterialInterface* Material, const FMaterialPropertyEx& InProperty, FGLTFMaterialAnalysis& OutAnalysis);
+
 #endif
 
 	static bool NeedsMeshData(const UMaterialInterface* Material);
diff --git a/Engine/Plugins/Enterprise/GLTFExporter/MigratedFromGit/Source/GLTFExporter/Private/Tasks/GLTFMaterialTasks.cpp b/Engine/Plugins/Enterprise/GLTFExporter/MigratedFromGit/Source/GLTFExporter/Private/Tasks/GLTFMaterialTasks.cpp
--- a/Engine/Plugins/Enterprise/GLTFExporter/MigratedFromGit/Source/GLTFExporter/Private/Tasks/GLTFMaterialTasks.cpp
+++ b/Engine/Plugins/Enterprise/GLTFExporter/MigratedFromGit/Source/GLTFExporter/Private/Tasks/GLTFMaterialTasks.cpp
@@ -378,7 +378,7 @@ bool FGLTFMaterialTask::TryGetBaseColorAndOpacity(FGLTFJsonPBRMetallicRoughness&
 		return false;
 	}
 
-	FIntPoint TextureSize = Builder.GetBakeSizeForMaterialProperty(Material, GetPropertyGroup(BaseColorProperty));
+	const FIntPoint TextureSize = Builder.GetBakeSizeForMaterialProperty(Material, GetPropertyGroup(BaseColorProperty));
 
 	const TextureAddress TextureAddress = Builder.GetBakeTilingForMaterialProperty(Material, GetPropertyGroup(BaseColorProperty));
 	const EGLTFJsonTextureWrap TextureWrapS = FGLTFConverterUtility::ConvertWrap(TextureAddress);
@@ -388,8 +388,8 @@ bool FGLTFMaterialTask::TryGetBaseColorAndOpacity(FGLTFJsonPBRMetallicRoughness&
 	const EGLTFJsonTextureFilter TextureMinFilter = FGLTFConverterUtility::ConvertMinFilter(TextureFilter);
 	const EGLTFJsonTextureFilter TextureMagFilter = FGLTFConverterUtility::ConvertMagFilter(TextureFilter);
 
-	const FGLTFPropertyBakeOutput BaseColorBakeOutput = BakeMaterialProperty(BaseColorProperty, BaseColorTexCoord, TextureSize);
-	const FGLTFPropertyBakeOutput OpacityBakeOutput = BakeMaterialProperty(OpacityProperty, OpacityTexCoord, TextureSize, true);
+	const FGLTFPropertyBakeOutput BaseColorBakeOutput = BakeMaterialProperty(BaseColorProperty, BaseColorTexCoord, TextureSize, false);
+	const FGLTFPropertyBakeOutput OpacityBakeOutput = BakeMaterialProperty(OpacityProperty, OpacityTexCoord, TextureSize, false);
 	const float BaseColorScale = BaseColorProperty == MP_EmissiveColor ? BaseColorBakeOutput.EmissiveScale : 1;
 
 	// Detect when both baked properties are constants, which means we can avoid exporting a texture
@@ -426,21 +426,22 @@ bool FGLTFMaterialTask::TryGetBaseColorAndOpacity(FGLTFJsonPBRMetallicRoughness&
 		return false;
 	}
 
-	TextureSize = BaseColorBakeOutput.Size.ComponentMax(OpacityBakeOutput.Size);
-	BaseColorTexture = FGLTFMaterialUtility::CreateTransientTexture(BaseColorBakeOutput);
-	OpacityTexture = FGLTFMaterialUtility::CreateTransientTexture(OpacityBakeOutput);
-
-	const TArray<FGLTFTextureCombineSource> CombineSources =
-	{
-		{ OpacityTexture, OpacityMask, SE_BLEND_Opaque },
-		{ BaseColorTexture, BaseColorMask }
-	};
+	TArray<FColor> CombinedPixels;
+	CombinePixels(
+		BaseColorBakeOutput.Pixels,
+		OpacityBakeOutput.Pixels,
+		CombinedPixels,
+		[](const FColor& BaseColor, const FColor& Opacity) -> FColor
+		{
+			return FColor(BaseColor.R, BaseColor.G, BaseColor.B, Opacity.R);
+		});
 
-	const FGLTFJsonTextureIndex TextureIndex = FGLTFMaterialUtility::AddCombinedTexture(
+	const FGLTFJsonTextureIndex TextureIndex = FGLTFMaterialUtility::AddTexture(
 		Builder,
-		CombineSources,
+		CombinedPixels,
 		TextureSize,
 		false,
+		false,
 		GetBakedTextureName(TEXT("BaseColor")),
 		TextureMinFilter,
 		TextureMagFilter,
@@ -504,7 +505,7 @@ bool FGLTFMaterialTask::TryGetMetallicAndRoughness(FGLTFJsonPBRMetallicRoughness
 	}
 
 	// TODO: add support for calculating the ideal resolution to use for baking based on connected (texture) nodes
-	FIntPoint TextureSize = Builder.GetBakeSizeForMaterialProperty(Material, GetPropertyGroup(MetallicProperty));
+	const FIntPoint TextureSize = Builder.GetBakeSizeForMaterialProperty(Material, GetPropertyGroup(MetallicProperty));
 
 	const TextureAddress TextureAddress = Builder.GetBakeTilingForMaterialProperty(Material, GetPropertyGroup(MetallicProperty));
 	const EGLTFJsonTextureWrap TextureWrapS = FGLTFConverterUtility::ConvertWrap(TextureAddress);
@@ -514,8 +515,8 @@ bool FGLTFMaterialTask::TryGetMetallicAndRoughness(FGLTFJsonPBRMetallicRoughness
 	const EGLTFJsonTextureFilter TextureMinFilter = FGLTFConverterUtility::ConvertMinFilter(TextureFilter);
 	const EGLTFJsonTextureFilter TextureMagFilter = FGLTFConverterUtility::ConvertMagFilter(TextureFilter);
 
-	FGLTFPropertyBakeOutput MetallicBakeOutput = BakeMaterialProperty(MetallicProperty, MetallicTexCoord, TextureSize);
-	FGLTFPropertyBakeOutput RoughnessBakeOutput = BakeMaterialProperty(RoughnessProperty, RoughnessTexCoord, TextureSize);
+	FGLTFPropertyBakeOutput MetallicBakeOutput = BakeMaterialProperty(MetallicProperty, MetallicTexCoord, TextureSize, false);
+	FGLTFPropertyBakeOutput RoughnessBakeOutput = BakeMaterialProperty(RoughnessProperty, RoughnessTexCoord, TextureSize, false);
 
 	// Detect when both baked properties are constants, which means we can use factors and avoid exporting a texture
 	if (MetallicBakeOutput.bIsConstant && RoughnessBakeOutput.bIsConstant)
@@ -549,21 +550,22 @@ bool FGLTFMaterialTask::TryGetMetallicAndRoughness(FGLTFJsonPBRMetallicRoughness
 		return false;
 	}
 
-	TextureSize = RoughnessBakeOutput.Size.ComponentMax(MetallicBakeOutput.Size);
-	MetallicTexture = FGLTFMaterialUtility::CreateTransientTexture(MetallicBakeOutput);
-	RoughnessTexture = FGLTFMaterialUtility::CreateTransientTexture(RoughnessBakeOutput);
-
-	const TArray<FGLTFTextureCombineSource> CombineSources =
-	{
-		{ MetallicTexture, MetallicMask + AlphaMask, SE_BLEND_Opaque },
-		{ RoughnessTexture, RoughnessMask }
-	};
+	TArray<FColor> CombinedPixels;
+	CombinePixels(
+		MetallicBakeOutput.Pixels,
+		RoughnessBakeOutput.Pixels,
+		CombinedPixels,
+		[](const FColor& Metallic, const FColor& Roughness) -> FColor
+		{
+			return FColor(0, Roughness.R, Metallic.R);
+		});
 
-	const FGLTFJsonTextureIndex TextureIndex = FGLTFMaterialUtility::AddCombinedTexture(
+	const FGLTFJsonTextureIndex TextureIndex = FGLTFMaterialUtility::AddTexture(
 		Builder,
-		CombineSources,
+		CombinedPixels,
 		TextureSize,
 		true, // NOTE: we can ignore alpha in everything but TryGetBaseColorAndOpacity
+		false,
 		GetBakedTextureName(TEXT("MetallicRoughness")),
 		TextureMinFilter,
 		TextureMagFilter,
@@ -628,7 +630,7 @@ bool FGLTFMaterialTask::TryGetClearCoatRoughness(FGLTFJsonClearCoatExtension& Ou
 		return false;
 	}
 
-	FIntPoint TextureSize = Builder.GetBakeSizeForMaterialProperty(Material, GetPropertyGroup(IntensityProperty));
+	const FIntPoint TextureSize = Builder.GetBakeSizeForMaterialProperty(Material, GetPropertyGroup(IntensityProperty));
 
 	const TextureAddress TextureAddress = Builder.GetBakeTilingForMaterialProperty(Material,GetPropertyGroup(IntensityProperty));
 	const EGLTFJsonTextureWrap TextureWrapS = FGLTFConverterUtility::ConvertWrap(TextureAddress);
@@ -638,8 +640,8 @@ bool FGLTFMaterialTask::TryGetClearCoatRoughness(FGLTFJsonClearCoatExtension& Ou
 	const EGLTFJsonTextureFilter TextureMinFilter = FGLTFConverterUtility::ConvertMinFilter(TextureFilter);
 	const EGLTFJsonTextureFilter TextureMagFilter = FGLTFConverterUtility::ConvertMagFilter(TextureFilter);
 
-	const FGLTFPropertyBakeOutput IntensityBakeOutput = BakeMaterialProperty(IntensityProperty, IntensityTexCoord, TextureSize);
-	const FGLTFPropertyBakeOutput RoughnessBakeOutput = BakeMaterialProperty(RoughnessProperty, RoughnessTexCoord, TextureSize);
+	const FGLTFPropertyBakeOutput IntensityBakeOutput = BakeMaterialProperty(IntensityProperty, IntensityTexCoord, TextureSize, false);
+	const FGLTFPropertyBakeOutput RoughnessBakeOutput = BakeMaterialProperty(RoughnessProperty, RoughnessTexCoord, TextureSize, false);
 
 	// Detect when both baked properties are constants, which means we can use factors and avoid exporting a texture
 	if (IntensityBakeOutput.bIsConstant && RoughnessBakeOutput.bIsConstant)
@@ -673,21 +675,22 @@ bool FGLTFMaterialTask::TryGetClearCoatRoughness(FGLTFJsonClearCoatExtension& Ou
 		return false;
 	}
 
-	TextureSize = RoughnessBakeOutput.Size.ComponentMax(IntensityBakeOutput.Size);
-	IntensityTexture = FGLTFMaterialUtility::CreateTransientTexture(IntensityBakeOutput);
-	RoughnessTexture = FGLTFMaterialUtility::CreateTransientTexture(RoughnessBakeOutput);
-
-	const TArray<FGLTFTextureCombineSource> CombineSources =
-	{
-		{ IntensityTexture, ClearCoatMask + AlphaMask, SE_BLEND_Opaque },
-		{ RoughnessTexture, ClearCoatRoughnessMask }
-	};
+	TArray<FColor> CombinedPixels;
+	CombinePixels(
+		IntensityBakeOutput.Pixels,
+		RoughnessBakeOutput.Pixels,
+		CombinedPixels,
+		[](const FColor& Intensity, const FColor& Roughness) -> FColor
+		{
+			return FColor(Intensity.R, Roughness.R, 0);
+		});
 
-	const FGLTFJsonTextureIndex TextureIndex = FGLTFMaterialUtility::AddCombinedTexture(
+	const FGLTFJsonTextureIndex TextureIndex = FGLTFMaterialUtility::AddTexture(
 		Builder,
-		CombineSources,
+		CombinedPixels,
 		TextureSize,
 		true, // NOTE: we can ignore alpha in everything but TryGetBaseColorAndOpacity
+		false,
 		GetBakedTextureName(TEXT("ClearCoatRoughness")),
 		TextureMinFilter,
 		TextureMagFilter,
@@ -1308,13 +1311,13 @@ bool FGLTFMaterialTask::TryGetBakedMaterialProperty(FGLTFJsonTextureInfo& OutTex
 	return true;
 }
 
-FGLTFPropertyBakeOutput FGLTFMaterialTask::BakeMaterialProperty(const FMaterialPropertyEx& Property, int32& OutTexCoord, bool bCopyAlphaFromRedChannel)
+FGLTFPropertyBakeOutput FGLTFMaterialTask::BakeMaterialProperty(const FMaterialPropertyEx& Property, int32& OutTexCoord)
 {
 	const FIntPoint TextureSize = Builder.GetBakeSizeForMaterialProperty(Material, GetPropertyGroup(Property));
-	return BakeMaterialProperty(Property, OutTexCoord, TextureSize, bCopyAlphaFromRedChannel);
+	return BakeMaterialProperty(Property, OutTexCoord, TextureSize, true);
 }
 
-FGLTFPropertyBakeOutput FGLTFMaterialTask::BakeMaterialProperty(const FMaterialPropertyEx& Property, int32& OutTexCoord, const FIntPoint& TextureSize, bool bCopyAlphaFromRedChannel)
+FGLTFPropertyBakeOutput FGLTFMaterialTask::BakeMaterialProperty(const FMaterialPropertyEx& Property, int32& OutTexCoord, const FIntPoint& TextureSize, bool bFillAlpha)
 {
 	if (MeshData == nullptr)
 	{
@@ -1355,7 +1358,7 @@ FGLTFPropertyBakeOutput FGLTFMaterialTask::BakeMaterialProperty(const FMaterialP
 		OutTexCoord,
 		MeshData,
 		SectionIndices,
-		bCopyAlphaFromRedChannel);
+		bFillAlpha);
 
 	return PropertyBakeOutput;
 }
@@ -1433,4 +1436,40 @@ EGLTFMaterialPropertyGroup FGLTFMaterialTask::GetPropertyGroup(const FMaterialPr
 	}
 }
 
+template <typename CallbackType>
+void FGLTFMaterialTask::CombinePixels(const TArray<FColor>& FirstPixels, const TArray<FColor>& SecondPixels, TArray<FColor>& OutPixels, CallbackType Callback)
+{
+	const int32 Count = FMath::Max(FirstPixels.Num(), SecondPixels.Num());
+	OutPixels.AddUninitialized(Count);
+
+	if (FirstPixels.Num() == 1)
+	{
+		const FColor& FirstPixel = FirstPixels[0];
+		for (int32 Index = 0; Index < Count; ++Index)
+		{
+			const FColor& SecondPixel = SecondPixels[Index];
+			OutPixels[Index] = Callback(FirstPixel, SecondPixel);
+		}
+	}
+	else if (SecondPixels.Num() == 1)
+	{
+		const FColor& SecondPixel = SecondPixels[0];
+		for (int32 Index = 0; Index < Count; ++Index)
+		{
+			const FColor& FirstPixel = FirstPixels[Index];
+			OutPixels[Index] = Callback(FirstPixel, SecondPixel);
+		}
+	}
+	else
+	{
+		check(FirstPixels.Num() == SecondPixels.Num());
+		for (int32 Index = 0; Index < Count; ++Index)
+		{
+			const FColor& FirstPixel = FirstPixels[Index];
+			const FColor& SecondPixel = SecondPixels[Index];
+			OutPixels[Index] = Callback(FirstPixel, SecondPixel);
+		}
+	}
+}
+
 #endif
diff --git a/Engine/Plugins/Enterprise/GLTFExporter/MigratedFromGit/Source/GLTFExporter/Private/Tasks/GLTFMaterialTasks.h b/Engine/Plugins/Enterprise/GLTFExporter/MigratedFromGit/Source/GLTFExporter/Private/Tasks/GLTFMaterialTasks.h
--- a/Engine/Plugins/Enterprise/GLTFExporter/MigratedFromGit/Source/GLTFExporter/Private/Tasks/GLTFMaterialTasks.h
+++ b/Engine/Plugins/Enterprise/GLTFExporter/MigratedFromGit/Source/GLTFExporter/Private/Tasks/GLTFMaterialTasks.h
@@ -69,8 +69,8 @@ private:
 	bool TryGetBakedMaterialProperty(FGLTFJsonTextureInfo& OutTexInfo, float& OutConstant, const FMaterialPropertyEx& Property, const FString& PropertyName);
 	bool TryGetBakedMaterialProperty(FGLTFJsonTextureInfo& OutTexInfo, const FMaterialPropertyEx& Property, const FString& PropertyName);
 
-	FGLTFPropertyBakeOutput BakeMaterialProperty(const FMaterialPropertyEx& Property, int32& OutTexCoord, bool bCopyAlphaFromRedChannel = false);
-	FGLTFPropertyBakeOutput BakeMaterialProperty(const FMaterialPropertyEx& Property, int32& OutTexCoord, const FIntPoint& TextureSize, bool bCopyAlphaFromRedChannel = false);
+	FGLTFPropertyBakeOutput BakeMaterialProperty(const FMaterialPropertyEx& Property, int32& OutTexCoord);
+	FGLTFPropertyBakeOutput BakeMaterialProperty(const FMaterialPropertyEx& Property, int32& OutTexCoord, const FIntPoint& TextureSize, bool bFillAlpha);
 
 	bool StoreBakedPropertyTexture(FGLTFJsonTextureInfo& OutTexInfo, const FGLTFPropertyBakeOutput& PropertyBakeOutput, const FString& PropertyName) const;
 
@@ -78,6 +78,9 @@ private:
 	FString GetBakedTextureName(const FString& PropertyName) const;
 
 	static EGLTFMaterialPropertyGroup GetPropertyGroup(const FMaterialPropertyEx& Property);
+
+	template <typename CallbackType>
+	static void CombinePixels(const TArray<FColor>& FirstPixels, const TArray<FColor>& SecondPixels, TArray<FColor>& OutPixels, CallbackType Callback);
 };
 
 #endif
