[GLTFExporter] fix: add missing happy path for reading vertex buffers from CPU

commit f8b320280c7406bef40f36ab03bd92015dc3ea22
Date: Wed, 1 Dec 2021 15:01:04 +0100

    fix: add missing happy path for reading vertex buffers from CPU
    (cherry picked from commit ca1d169e0fdcdc272ca84f537dc8132fe532b6df​)

#skipci

diff --git a/Engine/Plugins/Enterprise/GLTFExporter/MigratedFromGit/Source/GLTFExporter/Private/Converters/GLTFBufferAdapter.cpp b/Engine/Plugins/Enterprise/GLTFExporter/MigratedFromGit/Source/GLTFExporter/Private/Converters/GLTFBufferAdapter.cpp
--- a/Engine/Plugins/Enterprise/GLTFExporter/MigratedFromGit/Source/GLTFExporter/Private/Converters/GLTFBufferAdapter.cpp
+++ b/Engine/Plugins/Enterprise/GLTFExporter/MigratedFromGit/Source/GLTFExporter/Private/Converters/GLTFBufferAdapter.cpp
@@ -52,35 +52,35 @@ TUniquePtr<IGLTFBufferAdapter> IGLTFBufferAdapter::GetIndices(const FRawStaticIn
 TUniquePtr<IGLTFBufferAdapter> IGLTFBufferAdapter::GetIndices(const FRawStaticIndexBuffer16or32Interface* IndexBuffer)
 {
 	const void* IndexData = IndexBuffer->GetResourceDataSize() > 0 ? const_cast<FRawStaticIndexBuffer16or32Interface*>(IndexBuffer)->GetPointerTo(0) : nullptr;
-	if (IndexData != nullptr && (WITH_EDITOR || IndexBuffer->GetNeedsCPUAccess())) static_cast<TUniquePtr<IGLTFBufferAdapter>>(MakeUnique<FGLTFBufferAdapterCPU>(IndexData));
+	if (IndexData != nullptr && (WITH_EDITOR || IndexBuffer->GetNeedsCPUAccess())) return MakeUnique<FGLTFBufferAdapterCPU>(IndexData);
 	return MakeUnique<FGLTFBufferAdapterGPU>(IndexBuffer->IndexBufferRHI);
 }
 
 TUniquePtr<IGLTFBufferAdapter> IGLTFBufferAdapter::GetPositions(const FPositionVertexBuffer* VertexBuffer)
 {
 	const void* PositionData = VertexBuffer->GetVertexData();
-	if (PositionData != nullptr && (WITH_EDITOR || FGLTFBufferUtility::GetAllowCPUAccess(VertexBuffer))) MakeUnique<FGLTFBufferAdapterCPU>(PositionData);
+	if (PositionData != nullptr && (WITH_EDITOR || FGLTFBufferUtility::GetAllowCPUAccess(VertexBuffer))) return MakeUnique<FGLTFBufferAdapterCPU>(PositionData);
 	return MakeUnique<FGLTFBufferAdapterGPU>(VertexBuffer->VertexBufferRHI);
 }
 
 TUniquePtr<IGLTFBufferAdapter> IGLTFBufferAdapter::GetColors(const FColorVertexBuffer* VertexBuffer)
 {
 	const void* ColorData = VertexBuffer->GetVertexData();
-	if (ColorData != nullptr && (WITH_EDITOR || FGLTFBufferUtility::GetAllowCPUAccess(VertexBuffer))) MakeUnique<FGLTFBufferAdapterCPU>(ColorData);
+	if (ColorData != nullptr && (WITH_EDITOR || FGLTFBufferUtility::GetAllowCPUAccess(VertexBuffer))) return MakeUnique<FGLTFBufferAdapterCPU>(ColorData);
 	return MakeUnique<FGLTFBufferAdapterGPU>(VertexBuffer->VertexBufferRHI);
 }
 
 TUniquePtr<IGLTFBufferAdapter> IGLTFBufferAdapter::GetTangents(const FStaticMeshVertexBuffer* VertexBuffer)
 {
 	const void* TangentData = VertexBuffer->GetTangentData();
-	if (TangentData != nullptr && (WITH_EDITOR || const_cast<FStaticMeshVertexBuffer*>(VertexBuffer)->GetAllowCPUAccess())) MakeUnique<FGLTFBufferAdapterCPU>(TangentData);
+	if (TangentData != nullptr && (WITH_EDITOR || const_cast<FStaticMeshVertexBuffer*>(VertexBuffer)->GetAllowCPUAccess())) return MakeUnique<FGLTFBufferAdapterCPU>(TangentData);
 	return MakeUnique<FGLTFBufferAdapterGPU>(VertexBuffer->TangentsVertexBuffer.VertexBufferRHI);
 }
 
 TUniquePtr<IGLTFBufferAdapter> IGLTFBufferAdapter::GetUVs(const FStaticMeshVertexBuffer* VertexBuffer)
 {
 	const void* UVData = VertexBuffer->GetTexCoordData();
-	if (UVData != nullptr && (WITH_EDITOR || const_cast<FStaticMeshVertexBuffer*>(VertexBuffer)->GetAllowCPUAccess())) MakeUnique<FGLTFBufferAdapterCPU>(UVData);
+	if (UVData != nullptr && (WITH_EDITOR || const_cast<FStaticMeshVertexBuffer*>(VertexBuffer)->GetAllowCPUAccess())) return MakeUnique<FGLTFBufferAdapterCPU>(UVData);
 	return MakeUnique<FGLTFBufferAdapterGPU>(VertexBuffer->TexCoordVertexBuffer.VertexBufferRHI);
 }
 
@@ -88,7 +88,7 @@ TUniquePtr<IGLTFBufferAdapter> IGLTFBufferAdapter::GetInfluences(const FSkinWeig
 {
 	const auto* InfluenceBuffer = VertexBuffer->GetDataVertexBuffer();
 	const void* InfluenceData = InfluenceBuffer->GetWeightData();
-	if (InfluenceData != nullptr && (WITH_EDITOR || VertexBuffer->GetNeedsCPUAccess())) MakeUnique<FGLTFBufferAdapterCPU>(InfluenceData);
+	if (InfluenceData != nullptr && (WITH_EDITOR || VertexBuffer->GetNeedsCPUAccess())) return MakeUnique<FGLTFBufferAdapterCPU>(InfluenceData);
 	return MakeUnique<FGLTFBufferAdapterGPU>(InfluenceBuffer->VertexBufferRHI);
 }
 
@@ -96,6 +96,6 @@ TUniquePtr<IGLTFBufferAdapter> IGLTFBufferAdapter::GetLookups(const FSkinWeightV
 {
 	const FSkinWeightLookupVertexBuffer* LookupBuffer = VertexBuffer->GetLookupVertexBuffer();
 	const void* LookupData = FGLTFBufferUtility::GetBufferData(LookupBuffer);
-	if (LookupData != nullptr && (WITH_EDITOR || VertexBuffer->GetNeedsCPUAccess())) MakeUnique<FGLTFBufferAdapterCPU>(LookupData);
+	if (LookupData != nullptr && (WITH_EDITOR || VertexBuffer->GetNeedsCPUAccess())) return MakeUnique<FGLTFBufferAdapterCPU>(LookupData);
 	return MakeUnique<FGLTFBufferAdapterGPU>(LookupBuffer->VertexBufferRHI);
 }
