[GLTFExporter] chore: add clarifying comments to export-options that use bitmasks

commit 167589d4401a17c573cb119f2c6b9fde8dcc53d3
Date: Thu, 2 Dec 2021 10:15:07 +0100

    chore: add clarifying comments to export-options that use bitmasks
    (cherry picked from commit 42cc14799645fe0383cbebdbb636371f826a2243​)

#skipci

diff --git a/Engine/Plugins/Enterprise/GLTFExporter/MigratedFromGit/Source/GLTFExporter/Public/GLTFExportOptions.h b/Engine/Plugins/Enterprise/GLTFExporter/MigratedFromGit/Source/GLTFExporter/Public/GLTFExportOptions.h
--- a/Engine/Plugins/Enterprise/GLTFExporter/MigratedFromGit/Source/GLTFExporter/Public/GLTFExportOptions.h
+++ b/Engine/Plugins/Enterprise/GLTFExporter/MigratedFromGit/Source/GLTFExporter/Public/GLTFExportOptions.h
@@ -158,7 +158,7 @@ class GLTFEXPORTER_API UGLTFExportOptions : public UObject
 
 	/** Texture types that will always use lossless formats (e.g. PNG) because of sensitivity to compression artifacts. */
 	UPROPERTY(EditAnywhere, BlueprintReadWrite, Config, Category = Texture, Meta = (Bitmask, BitmaskEnum = EGLTFTextureType, EditCondition = "TextureImageFormat == EGLTFTextureImageFormat::JPEG"))
-	int32 NoLossyImageFormatFor;
+	int32 NoLossyImageFormatFor; // Bitmask combined from EGLTFTextureType
 
 	/** If enabled, export UV tiling and un-mirroring settings in a texture coordinate expression node for simple material input expressions. Uses extension KHR_texture_transform. */
 	UPROPERTY(EditAnywhere, BlueprintReadWrite, Config, Category = Texture, Meta = (EditCondition = "TextureImageFormat != EGLTFTextureImageFormat::None"))
@@ -178,7 +178,7 @@ class GLTFEXPORTER_API UGLTFExportOptions : public UObject
 
 	/** Mobility of directional, point, and spot light components that will be exported. Uses extension KHR_lights_punctual. */
 	UPROPERTY(EditAnywhere, BlueprintReadWrite, Config, Category = Scene, Meta = (Bitmask, BitmaskEnum = EGLTFSceneMobility))
-	int32 ExportLights;
+	int32 ExportLights; // Bitmask combined from EGLTFSceneMobility
 
 	/** If enabled, export camera components. */
 	UPROPERTY(EditAnywhere, BlueprintReadWrite, Config, Category = Scene)
