[GLTFExporter] fix: fix incorrect behavior when exporting TextureRenderTargetCube

commit 2e6be77cd10a56c401d39c6ab9e394143192ebb7
Date: Thu, 2 Dec 2021 20:49:56 +0100

    fix: fix incorrect behavior when exporting TextureRenderTargetCube
    (cherry picked from commit 841535a726fcd356b77491d5e2882c447dbd129a​)

#skipci

diff --git a/Engine/Plugins/Enterprise/GLTFExporter/MigratedFromGit/Source/GLTFExporter/Private/Converters/GLTFTextureConverters.cpp b/Engine/Plugins/Enterprise/GLTFExporter/MigratedFromGit/Source/GLTFExporter/Private/Converters/GLTFTextureConverters.cpp
--- a/Engine/Plugins/Enterprise/GLTFExporter/MigratedFromGit/Source/GLTFExporter/Private/Converters/GLTFTextureConverters.cpp
+++ b/Engine/Plugins/Enterprise/GLTFExporter/MigratedFromGit/Source/GLTFExporter/Private/Converters/GLTFTextureConverters.cpp
@@ -43,12 +43,12 @@ FGLTFJsonTextureIndex FGLTFTextureRenderTargetCubeConverter::Convert(const UText
 {
 	if (Builder.ExportOptions->TextureImageFormat != EGLTFTextureImageFormat::None)
 	{
-		return FGLTFJsonTextureIndex(INDEX_NONE);
+		const FGLTFJsonTextureIndex TextureIndex = Builder.AddTexture();
+		Builder.SetupTask<FGLTFTextureRenderTargetCubeTask>(Builder, RenderTargetCube, CubeFace, bToSRGB, TextureIndex);
+		return TextureIndex;
 	}
 
-	const FGLTFJsonTextureIndex TextureIndex = Builder.AddTexture();
-	Builder.SetupTask<FGLTFTextureRenderTargetCubeTask>(Builder, RenderTargetCube, CubeFace, bToSRGB, TextureIndex);
-	return TextureIndex;
+	return FGLTFJsonTextureIndex(INDEX_NONE);
 }
 
 FGLTFJsonTextureIndex FGLTFTextureLightMapConverter::Convert(const ULightMapTexture2D* LightMap)
