[GLTFExporter] refactor: use linear color space directly when material baking non-color properties

commit 857a8b086e4efed94c94e9ce46f2bacaa4dee6c1
Date: Thu, 2 Dec 2021 20:58:58 +0100

    refactor: use linear color space directly when material baking non-color properties
    (cherry picked from commit 7276afa2346ede3d1447a44d86d080c221510ddb​)

#skipci

diff --git a/Engine/Plugins/Enterprise/GLTFExporter/MigratedFromGit/Source/GLTFExporter/Private/Converters/GLTFMaterialUtility.cpp b/Engine/Plugins/Enterprise/GLTFExporter/MigratedFromGit/Source/GLTFExporter/Private/Converters/GLTFMaterialUtility.cpp
--- a/Engine/Plugins/Enterprise/GLTFExporter/MigratedFromGit/Source/GLTFExporter/Private/Converters/GLTFMaterialUtility.cpp
+++ b/Engine/Plugins/Enterprise/GLTFExporter/MigratedFromGit/Source/GLTFExporter/Private/Converters/GLTFMaterialUtility.cpp
@@ -128,10 +128,10 @@ FGLTFPropertyBakeOutput FGLTFMaterialUtility::BakeMaterialProperty(const FIntPoi
 	TArray<FBakeOutputEx> BakeOutputs;
 	IMaterialBakingModule& Module = FModuleManager::Get().LoadModuleChecked<IMaterialBakingModule>("GLTFMaterialBaking");
 
-	//Module.SetLinearBake(true);
+	Module.SetLinearBake(true);
 	Module.BakeMaterials(MatSettings, MeshSettings, BakeOutputs);
 	const bool bIsLinearBake = Module.IsLinearBake(Property);
-	//Module.SetLinearBake(false);
+	Module.SetLinearBake(false);
 
 	FBakeOutputEx& BakeOutput = BakeOutputs[0];
 
