[GLTFExporter] fix: fix incorrect condition when checking for LocalObjectBounds usage

commit 16900c342a80781712713f55190f1c816cd86a1c
Date: Fri, 10 Dec 2021 14:55:29 +0100

    fix: fix incorrect condition when checking for LocalObjectBounds usage
    (cherry picked from commit 85b55b9ad637661e833536bcc1bfff3b34dfb3de​)

#skipci

diff --git a/Engine/Plugins/Enterprise/GLTFExporter/MigratedFromGit/Source/GLTFMaterialAnalyzer/Private/GLTFProxyMaterialCompiler.h b/Engine/Plugins/Enterprise/GLTFExporter/MigratedFromGit/Source/GLTFMaterialAnalyzer/Private/GLTFProxyMaterialCompiler.h
--- a/Engine/Plugins/Enterprise/GLTFExporter/MigratedFromGit/Source/GLTFMaterialAnalyzer/Private/GLTFProxyMaterialCompiler.h
+++ b/Engine/Plugins/Enterprise/GLTFExporter/MigratedFromGit/Source/GLTFMaterialAnalyzer/Private/GLTFProxyMaterialCompiler.h
@@ -57,7 +57,7 @@ public:
 #endif
 	{
 		// NOTE: because there is no overridable compiler method for ObjectLocalBounds in engine this is our only option currently
-		if (Custom->Code.Contains(TEXT("GetPrimitiveData(Parameters).LocalObjectBounds")))
+		if (Custom->Code.Contains(TEXT("GetPrimitiveData(Parameters.PrimitiveId).LocalObjectBounds")))
 		{
 			bUsesObjectLocalBounds = true;
 		}
