[GLTFExporter] fix: add general category for properties in GLTFExportMessages to avoid compiler error

commit ccf86be5b732a6569beb11d37addb7516326291e
Date: Wed, 15 Dec 2021 19:39:57 +0100

    fix: add general category for properties in GLTFExportMessages to avoid compiler error
    (cherry picked from commit caca846537c4daeb32ceee85a0b803b2db8e4dad​)

#skipci

diff --git a/Engine/Plugins/Enterprise/GLTFExporter/MigratedFromGit/Source/GLTFExporter/Public/Exporters/GLTFExporter.h b/Engine/Plugins/Enterprise/GLTFExporter/MigratedFromGit/Source/GLTFExporter/Public/Exporters/GLTFExporter.h
--- a/Engine/Plugins/Enterprise/GLTFExporter/MigratedFromGit/Source/GLTFExporter/Public/Exporters/GLTFExporter.h
+++ b/Engine/Plugins/Enterprise/GLTFExporter/MigratedFromGit/Source/GLTFExporter/Public/Exporters/GLTFExporter.h
@@ -14,13 +14,13 @@ struct FGLTFExportMessages
 {
 	GENERATED_USTRUCT_BODY()
 
-	UPROPERTY(EditAnywhere, BlueprintReadWrite)
+	UPROPERTY(EditAnywhere, BlueprintReadWrite, Category = General)
 	TArray<FString> Suggestions;
 
-	UPROPERTY(EditAnywhere, BlueprintReadWrite)
+	UPROPERTY(EditAnywhere, BlueprintReadWrite, Category = General)
 	TArray<FString> Warnings;
 
-	UPROPERTY(EditAnywhere, BlueprintReadWrite)
+	UPROPERTY(EditAnywhere, BlueprintReadWrite, Category = General)
 	TArray<FString> Errors;
 };
 
