[GLTFExporter] fix: add missing category for ResetToDefault in GLTFExportOptions

commit 48315fd0e9cedf25e104389a9fb8ffa627cd9ea1
Date: Thu, 16 Dec 2021 10:16:53 +0100

    fix: add missing category for ResetToDefault in GLTFExportOptions
    (cherry picked from commit 16be49967cb0fc78db9c7423c1d70fb8c312a5ca​)

#skipci

diff --git a/Engine/Plugins/Enterprise/GLTFExporter/MigratedFromGit/Source/GLTFExporter/Public/GLTFExportOptions.h b/Engine/Plugins/Enterprise/GLTFExporter/MigratedFromGit/Source/GLTFExporter/Public/GLTFExportOptions.h
--- a/Engine/Plugins/Enterprise/GLTFExporter/MigratedFromGit/Source/GLTFExporter/Public/GLTFExportOptions.h
+++ b/Engine/Plugins/Enterprise/GLTFExporter/MigratedFromGit/Source/GLTFExporter/Public/GLTFExportOptions.h
@@ -216,6 +216,6 @@ class GLTFEXPORTER_API UGLTFExportOptions : public UObject
 	UPROPERTY(EditAnywhere, BlueprintReadWrite, Config, Category = VariantSets, Meta = (EditCondition = "bExportVariantSets"))
 	bool bExportVisibilityVariants;
 
-	UFUNCTION(BlueprintCallable)
+	UFUNCTION(BlueprintCallable, Category = General)
 	void ResetToDefault();
 };
