[GLTFExporter] fix: remove all #include "Engine.h" to avoid warning about monolithic headers

commit a76cbfb30bd42d49ad70fe1f618d2c6d84d38f09
Date: Fri, 17 Dec 2021 08:49:54 +0100

    fix: remove all #include "Engine.h" to avoid warning about monolithic headers
    (cherry picked from commit 34ae66991c245a3701bd7554c0d27ab04526d9ab​)

#skipci

diff --git a/Engine/Plugins/Enterprise/GLTFExporter/MigratedFromGit/Source/GLTFExporter/Private/Builders/GLTFConvertBuilder.cpp b/Engine/Plugins/Enterprise/GLTFExporter/MigratedFromGit/Source/GLTFExporter/Private/Builders/GLTFConvertBuilder.cpp
--- a/Engine/Plugins/Enterprise/GLTFExporter/MigratedFromGit/Source/GLTFExporter/Private/Builders/GLTFConvertBuilder.cpp
+++ b/Engine/Plugins/Enterprise/GLTFExporter/MigratedFromGit/Source/GLTFExporter/Private/Builders/GLTFConvertBuilder.cpp
@@ -1,6 +1,10 @@
 // Copyright Epic Games, Inc. All Rights Reserved.
 
 #include "Builders/GLTFConvertBuilder.h"
+#include "Engine/Texture2D.h"
+#include "Engine/TextureCube.h"
+#include "Engine/TextureRenderTarget2D.h"
+#include "Engine/TextureRenderTargetCube.h"
 
 FGLTFConvertBuilder::FGLTFConvertBuilder(const FString& FilePath, const UGLTFExportOptions* ExportOptions, bool bSelectedActorsOnly)
 	: FGLTFImageBuilder(FilePath, ExportOptions)
diff --git a/Engine/Plugins/Enterprise/GLTFExporter/MigratedFromGit/Source/GLTFExporter/Private/Converters/GLTFAccessorConverters.h b/Engine/Plugins/Enterprise/GLTFExporter/MigratedFromGit/Source/GLTFExporter/Private/Converters/GLTFAccessorConverters.h
--- a/Engine/Plugins/Enterprise/GLTFExporter/MigratedFromGit/Source/GLTFExporter/Private/Converters/GLTFAccessorConverters.h
+++ b/Engine/Plugins/Enterprise/GLTFExporter/MigratedFromGit/Source/GLTFExporter/Private/Converters/GLTFAccessorConverters.h
@@ -6,7 +6,6 @@
 #include "Converters/GLTFConverter.h"
 #include "Converters/GLTFBuilderContext.h"
 #include "Converters/GLTFMeshSection.h"
-#include "Engine.h"
 
 template <typename... InputTypes>
 class TGLTFAccessorConverter : public FGLTFBuilderContext, public TGLTFConverter<FGLTFJsonAccessorIndex, InputTypes...>
diff --git a/Engine/Plugins/Enterprise/GLTFExporter/MigratedFromGit/Source/GLTFExporter/Private/Converters/GLTFAnimationConverters.h b/Engine/Plugins/Enterprise/GLTFExporter/MigratedFromGit/Source/GLTFExporter/Private/Converters/GLTFAnimationConverters.h
--- a/Engine/Plugins/Enterprise/GLTFExporter/MigratedFromGit/Source/GLTFExporter/Private/Converters/GLTFAnimationConverters.h
+++ b/Engine/Plugins/Enterprise/GLTFExporter/MigratedFromGit/Source/GLTFExporter/Private/Converters/GLTFAnimationConverters.h
@@ -5,7 +5,6 @@
 #include "Json/GLTFJsonIndex.h"
 #include "Converters/GLTFConverter.h"
 #include "Converters/GLTFBuilderContext.h"
-#include "Engine.h"
 
 class ULevelSequence;
 class ALevelSequenceActor;
diff --git a/Engine/Plugins/Enterprise/GLTFExporter/MigratedFromGit/Source/GLTFExporter/Private/Converters/GLTFBackdropConverters.h b/Engine/Plugins/Enterprise/GLTFExporter/MigratedFromGit/Source/GLTFExporter/Private/Converters/GLTFBackdropConverters.h
--- a/Engine/Plugins/Enterprise/GLTFExporter/MigratedFromGit/Source/GLTFExporter/Private/Converters/GLTFBackdropConverters.h
+++ b/Engine/Plugins/Enterprise/GLTFExporter/MigratedFromGit/Source/GLTFExporter/Private/Converters/GLTFBackdropConverters.h
@@ -5,7 +5,6 @@
 #include "Json/GLTFJsonIndex.h"
 #include "Converters/GLTFConverter.h"
 #include "Converters/GLTFBuilderContext.h"
-#include "Engine.h"
 
 class FGLTFBackdropConverter final : public FGLTFBuilderContext, public TGLTFConverter<FGLTFJsonBackdropIndex, const AActor*>
 {
diff --git a/Engine/Plugins/Enterprise/GLTFExporter/MigratedFromGit/Source/GLTFExporter/Private/Converters/GLTFBoneUtility.cpp b/Engine/Plugins/Enterprise/GLTFExporter/MigratedFromGit/Source/GLTFExporter/Private/Converters/GLTFBoneUtility.cpp
--- a/Engine/Plugins/Enterprise/GLTFExporter/MigratedFromGit/Source/GLTFExporter/Private/Converters/GLTFBoneUtility.cpp
+++ b/Engine/Plugins/Enterprise/GLTFExporter/MigratedFromGit/Source/GLTFExporter/Private/Converters/GLTFBoneUtility.cpp
@@ -1,7 +1,6 @@
 // Copyright Epic Games, Inc. All Rights Reserved.
 
 #include "Converters/GLTFBoneUtility.h"
-#include "Engine.h"
 
 FTransform FGLTFBoneUtility::GetBindTransform(const FReferenceSkeleton& RefSkeleton, int32 BoneIndex)
 {
diff --git a/Engine/Plugins/Enterprise/GLTFExporter/MigratedFromGit/Source/GLTFExporter/Private/Converters/GLTFBufferAdapter.h b/Engine/Plugins/Enterprise/GLTFExporter/MigratedFromGit/Source/GLTFExporter/Private/Converters/GLTFBufferAdapter.h
--- a/Engine/Plugins/Enterprise/GLTFExporter/MigratedFromGit/Source/GLTFExporter/Private/Converters/GLTFBufferAdapter.h
+++ b/Engine/Plugins/Enterprise/GLTFExporter/MigratedFromGit/Source/GLTFExporter/Private/Converters/GLTFBufferAdapter.h
@@ -2,7 +2,7 @@
 
 #pragma once
 
-#include "Engine.h"
+#include "CoreMinimal.h"
 
 class IGLTFBufferAdapter
 {
diff --git a/Engine/Plugins/Enterprise/GLTFExporter/MigratedFromGit/Source/GLTFExporter/Private/Converters/GLTFBufferUtility.h b/Engine/Plugins/Enterprise/GLTFExporter/MigratedFromGit/Source/GLTFExporter/Private/Converters/GLTFBufferUtility.h
--- a/Engine/Plugins/Enterprise/GLTFExporter/MigratedFromGit/Source/GLTFExporter/Private/Converters/GLTFBufferUtility.h
+++ b/Engine/Plugins/Enterprise/GLTFExporter/MigratedFromGit/Source/GLTFExporter/Private/Converters/GLTFBufferUtility.h
@@ -2,7 +2,7 @@
 
 #pragma once
 
-#include "Engine.h"
+#include "CoreMinimal.h"
 
 struct FGLTFBufferUtility
 {
diff --git a/Engine/Plugins/Enterprise/GLTFExporter/MigratedFromGit/Source/GLTFExporter/Private/Converters/GLTFCameraConverters.cpp b/Engine/Plugins/Enterprise/GLTFExporter/MigratedFromGit/Source/GLTFExporter/Private/Converters/GLTFCameraConverters.cpp
--- a/Engine/Plugins/Enterprise/GLTFExporter/MigratedFromGit/Source/GLTFExporter/Private/Converters/GLTFCameraConverters.cpp
+++ b/Engine/Plugins/Enterprise/GLTFExporter/MigratedFromGit/Source/GLTFExporter/Private/Converters/GLTFCameraConverters.cpp
@@ -5,6 +5,7 @@
 #include "Converters/GLTFConverterUtility.h"
 #include "Converters/GLTFNameUtility.h"
 #include "Actors/GLTFCameraActor.h"
+#include "Camera/CameraComponent.h"
 
 FGLTFJsonCameraIndex FGLTFCameraConverter::Convert(const UCameraComponent* CameraComponent)
 {
diff --git a/Engine/Plugins/Enterprise/GLTFExporter/MigratedFromGit/Source/GLTFExporter/Private/Converters/GLTFCameraConverters.h b/Engine/Plugins/Enterprise/GLTFExporter/MigratedFromGit/Source/GLTFExporter/Private/Converters/GLTFCameraConverters.h
--- a/Engine/Plugins/Enterprise/GLTFExporter/MigratedFromGit/Source/GLTFExporter/Private/Converters/GLTFCameraConverters.h
+++ b/Engine/Plugins/Enterprise/GLTFExporter/MigratedFromGit/Source/GLTFExporter/Private/Converters/GLTFCameraConverters.h
@@ -5,7 +5,8 @@
 #include "Json/GLTFJsonIndex.h"
 #include "Converters/GLTFConverter.h"
 #include "Converters/GLTFBuilderContext.h"
-#include "Engine.h"
+
+class UCameraComponent;
 
 class FGLTFCameraConverter final : public FGLTFBuilderContext, public TGLTFConverter<FGLTFJsonCameraIndex, const UCameraComponent*>
 {
diff --git a/Engine/Plugins/Enterprise/GLTFExporter/MigratedFromGit/Source/GLTFExporter/Private/Converters/GLTFLevelVariantSetsConverters.h b/Engine/Plugins/Enterprise/GLTFExporter/MigratedFromGit/Source/GLTFExporter/Private/Converters/GLTFLevelVariantSetsConverters.h
--- a/Engine/Plugins/Enterprise/GLTFExporter/MigratedFromGit/Source/GLTFExporter/Private/Converters/GLTFLevelVariantSetsConverters.h
+++ b/Engine/Plugins/Enterprise/GLTFExporter/MigratedFromGit/Source/GLTFExporter/Private/Converters/GLTFLevelVariantSetsConverters.h
@@ -6,10 +6,12 @@
 #include "Json/GLTFJsonLevelVariantSets.h"
 #include "Converters/GLTFConverter.h"
 #include "Converters/GLTFBuilderContext.h"
-#include "Engine.h"
-#include "Variant.h"
-#include "PropertyValue.h"
-#include "LevelVariantSets.h"
+
+class ULevelVariantSets;
+class UVariantSet;
+class UVariant;
+class UVariantObjectBinding;
+class UPropertyValue;
 
 class FGLTFLevelVariantSetsConverter final : public FGLTFBuilderContext, public TGLTFConverter<FGLTFJsonLevelVariantSetsIndex, const ULevelVariantSets*>
 {
diff --git a/Engine/Plugins/Enterprise/GLTFExporter/MigratedFromGit/Source/GLTFExporter/Private/Converters/GLTFLightConverters.cpp b/Engine/Plugins/Enterprise/GLTFExporter/MigratedFromGit/Source/GLTFExporter/Private/Converters/GLTFLightConverters.cpp
--- a/Engine/Plugins/Enterprise/GLTFExporter/MigratedFromGit/Source/GLTFExporter/Private/Converters/GLTFLightConverters.cpp
+++ b/Engine/Plugins/Enterprise/GLTFExporter/MigratedFromGit/Source/GLTFExporter/Private/Converters/GLTFLightConverters.cpp
@@ -4,6 +4,9 @@
 #include "Builders/GLTFContainerBuilder.h"
 #include "Converters/GLTFConverterUtility.h"
 #include "Converters/GLTFNameUtility.h"
+#include "Components/LightComponent.h"
+#include "Components/PointLightComponent.h"
+#include "Components/SpotLightComponent.h"
 
 FGLTFJsonLightIndex FGLTFLightConverter::Convert(const ULightComponent* LightComponent)
 {
diff --git a/Engine/Plugins/Enterprise/GLTFExporter/MigratedFromGit/Source/GLTFExporter/Private/Converters/GLTFLightConverters.h b/Engine/Plugins/Enterprise/GLTFExporter/MigratedFromGit/Source/GLTFExporter/Private/Converters/GLTFLightConverters.h
--- a/Engine/Plugins/Enterprise/GLTFExporter/MigratedFromGit/Source/GLTFExporter/Private/Converters/GLTFLightConverters.h
+++ b/Engine/Plugins/Enterprise/GLTFExporter/MigratedFromGit/Source/GLTFExporter/Private/Converters/GLTFLightConverters.h
@@ -5,7 +5,6 @@
 #include "Json/GLTFJsonIndex.h"
 #include "Converters/GLTFConverter.h"
 #include "Converters/GLTFBuilderContext.h"
-#include "Engine.h"
 
 class FGLTFLightConverter final : public FGLTFBuilderContext, public TGLTFConverter<FGLTFJsonLightIndex, const ULightComponent*>
 {
diff --git a/Engine/Plugins/Enterprise/GLTFExporter/MigratedFromGit/Source/GLTFExporter/Private/Converters/GLTFLightMapConverters.h b/Engine/Plugins/Enterprise/GLTFExporter/MigratedFromGit/Source/GLTFExporter/Private/Converters/GLTFLightMapConverters.h
--- a/Engine/Plugins/Enterprise/GLTFExporter/MigratedFromGit/Source/GLTFExporter/Private/Converters/GLTFLightMapConverters.h
+++ b/Engine/Plugins/Enterprise/GLTFExporter/MigratedFromGit/Source/GLTFExporter/Private/Converters/GLTFLightMapConverters.h
@@ -5,7 +5,6 @@
 #include "Json/GLTFJsonIndex.h"
 #include "Converters/GLTFConverter.h"
 #include "Converters/GLTFBuilderContext.h"
-#include "Engine.h"
 
 class FGLTFLightMapConverter final : public FGLTFBuilderContext, public TGLTFConverter<FGLTFJsonLightMapIndex, const UStaticMeshComponent*>
 {
diff --git a/Engine/Plugins/Enterprise/GLTFExporter/MigratedFromGit/Source/GLTFExporter/Private/Converters/GLTFMaterialConverters.h b/Engine/Plugins/Enterprise/GLTFExporter/MigratedFromGit/Source/GLTFExporter/Private/Converters/GLTFMaterialConverters.h
--- a/Engine/Plugins/Enterprise/GLTFExporter/MigratedFromGit/Source/GLTFExporter/Private/Converters/GLTFMaterialConverters.h
+++ b/Engine/Plugins/Enterprise/GLTFExporter/MigratedFromGit/Source/GLTFExporter/Private/Converters/GLTFMaterialConverters.h
@@ -9,7 +9,6 @@
 #include "Converters/GLTFIndexArray.h"
 #include "Converters/GLTFUVOverlapChecker.h"
 #include "Converters/GLTFUVDegenerateChecker.h"
-#include "Engine.h"
 
 class FGLTFMaterialConverter final : public FGLTFBuilderContext, public TGLTFConverter<FGLTFJsonMaterialIndex, const UMaterialInterface*, const FGLTFMeshData*, FGLTFIndexArray>
 {
diff --git a/Engine/Plugins/Enterprise/GLTFExporter/MigratedFromGit/Source/GLTFExporter/Private/Converters/GLTFMeshConverters.h b/Engine/Plugins/Enterprise/GLTFExporter/MigratedFromGit/Source/GLTFExporter/Private/Converters/GLTFMeshConverters.h
--- a/Engine/Plugins/Enterprise/GLTFExporter/MigratedFromGit/Source/GLTFExporter/Private/Converters/GLTFMeshConverters.h
+++ b/Engine/Plugins/Enterprise/GLTFExporter/MigratedFromGit/Source/GLTFExporter/Private/Converters/GLTFMeshConverters.h
@@ -8,7 +8,6 @@
 #include "Converters/GLTFMeshSectionConverters.h"
 #include "Converters/GLTFMeshDataConverters.h"
 #include "Converters/GLTFMaterialArray.h"
-#include "Engine.h"
 
 template <typename... InputTypes>
 class TGLTFMeshConverter : public FGLTFBuilderContext, public TGLTFConverter<FGLTFJsonMeshIndex, InputTypes...>
diff --git a/Engine/Plugins/Enterprise/GLTFExporter/MigratedFromGit/Source/GLTFExporter/Private/Converters/GLTFMeshData.cpp b/Engine/Plugins/Enterprise/GLTFExporter/MigratedFromGit/Source/GLTFExporter/Private/Converters/GLTFMeshData.cpp
--- a/Engine/Plugins/Enterprise/GLTFExporter/MigratedFromGit/Source/GLTFExporter/Private/Converters/GLTFMeshData.cpp
+++ b/Engine/Plugins/Enterprise/GLTFExporter/MigratedFromGit/Source/GLTFExporter/Private/Converters/GLTFMeshData.cpp
@@ -2,7 +2,6 @@
 
 #include "Converters/GLTFMeshData.h"
 #include "Converters/GLTFNameUtility.h"
-#include "Converters/GLTFMeshUtility.h"
 
 #if WITH_EDITOR
 #include "StaticMeshAttributes.h"
diff --git a/Engine/Plugins/Enterprise/GLTFExporter/MigratedFromGit/Source/GLTFExporter/Private/Converters/GLTFMeshDataConverters.h b/Engine/Plugins/Enterprise/GLTFExporter/MigratedFromGit/Source/GLTFExporter/Private/Converters/GLTFMeshDataConverters.h
--- a/Engine/Plugins/Enterprise/GLTFExporter/MigratedFromGit/Source/GLTFExporter/Private/Converters/GLTFMeshDataConverters.h
+++ b/Engine/Plugins/Enterprise/GLTFExporter/MigratedFromGit/Source/GLTFExporter/Private/Converters/GLTFMeshDataConverters.h
@@ -4,7 +4,6 @@
 
 #include "Converters/GLTFConverter.h"
 #include "Converters/GLTFMeshData.h"
-#include "Engine.h"
 
 template <typename MeshType, typename MeshComponentType>
 class TGLTFMeshDataConverter final : public TGLTFConverter<const FGLTFMeshData*, const MeshType*, const MeshComponentType*, int32>
diff --git a/Engine/Plugins/Enterprise/GLTFExporter/MigratedFromGit/Source/GLTFExporter/Private/Converters/GLTFMeshUtility.h b/Engine/Plugins/Enterprise/GLTFExporter/MigratedFromGit/Source/GLTFExporter/Private/Converters/GLTFMeshUtility.h
--- a/Engine/Plugins/Enterprise/GLTFExporter/MigratedFromGit/Source/GLTFExporter/Private/Converters/GLTFMeshUtility.h
+++ b/Engine/Plugins/Enterprise/GLTFExporter/MigratedFromGit/Source/GLTFExporter/Private/Converters/GLTFMeshUtility.h
@@ -3,7 +3,6 @@
 #pragma once
 
 #include "Converters/GLTFIndexArray.h"
-#include "Engine.h"
 
 class FSkeletalMeshLODRenderData;
 
diff --git a/Engine/Plugins/Enterprise/GLTFExporter/MigratedFromGit/Source/GLTFExporter/Private/Converters/GLTFNameUtility.cpp b/Engine/Plugins/Enterprise/GLTFExporter/MigratedFromGit/Source/GLTFExporter/Private/Converters/GLTFNameUtility.cpp
--- a/Engine/Plugins/Enterprise/GLTFExporter/MigratedFromGit/Source/GLTFExporter/Private/Converters/GLTFNameUtility.cpp
+++ b/Engine/Plugins/Enterprise/GLTFExporter/MigratedFromGit/Source/GLTFExporter/Private/Converters/GLTFNameUtility.cpp
@@ -1,8 +1,12 @@
 // Copyright Epic Games, Inc. All Rights Reserved.
 
 #include "Converters/GLTFNameUtility.h"
+#include "Engine/StaticMeshActor.h"
 #include "Animation/SkeletalMeshActor.h"
-#include "Engine.h"
+#include "Components/LightComponent.h"
+#include "Engine/Light.h"
+#include "Camera/CameraComponent.h"
+#include "Camera/CameraActor.h"
 
 namespace
 {
diff --git a/Engine/Plugins/Enterprise/GLTFExporter/MigratedFromGit/Source/GLTFExporter/Private/Converters/GLTFNodeConverters.cpp b/Engine/Plugins/Enterprise/GLTFExporter/MigratedFromGit/Source/GLTFExporter/Private/Converters/GLTFNodeConverters.cpp
--- a/Engine/Plugins/Enterprise/GLTFExporter/MigratedFromGit/Source/GLTFExporter/Private/Converters/GLTFNodeConverters.cpp
+++ b/Engine/Plugins/Enterprise/GLTFExporter/MigratedFromGit/Source/GLTFExporter/Private/Converters/GLTFNodeConverters.cpp
@@ -7,6 +7,10 @@
 #include "Converters/GLTFNameUtility.h"
 #include "Actors/GLTFHotspotActor.h"
 #include "LevelSequenceActor.h"
+#include "Camera/CameraComponent.h"
+#include "Components/LightComponent.h"
+#include "Engine/StaticMeshSocket.h"
+#include "Engine/SkeletalMeshSocket.h"
 
 FGLTFJsonNodeIndex FGLTFActorConverter::Convert(const AActor* Actor)
 {
diff --git a/Engine/Plugins/Enterprise/GLTFExporter/MigratedFromGit/Source/GLTFExporter/Private/Converters/GLTFNodeConverters.h b/Engine/Plugins/Enterprise/GLTFExporter/MigratedFromGit/Source/GLTFExporter/Private/Converters/GLTFNodeConverters.h
--- a/Engine/Plugins/Enterprise/GLTFExporter/MigratedFromGit/Source/GLTFExporter/Private/Converters/GLTFNodeConverters.h
+++ b/Engine/Plugins/Enterprise/GLTFExporter/MigratedFromGit/Source/GLTFExporter/Private/Converters/GLTFNodeConverters.h
@@ -5,7 +5,6 @@
 #include "Json/GLTFJsonIndex.h"
 #include "Converters/GLTFConverter.h"
 #include "Converters/GLTFBuilderContext.h"
-#include "Engine.h"
 
 template <typename... InputTypes>
 class TGLTFNodeConverter : public FGLTFBuilderContext, public TGLTFConverter<FGLTFJsonNodeIndex, InputTypes...>
diff --git a/Engine/Plugins/Enterprise/GLTFExporter/MigratedFromGit/Source/GLTFExporter/Private/Converters/GLTFSamplerConverters.h b/Engine/Plugins/Enterprise/GLTFExporter/MigratedFromGit/Source/GLTFExporter/Private/Converters/GLTFSamplerConverters.h
--- a/Engine/Plugins/Enterprise/GLTFExporter/MigratedFromGit/Source/GLTFExporter/Private/Converters/GLTFSamplerConverters.h
+++ b/Engine/Plugins/Enterprise/GLTFExporter/MigratedFromGit/Source/GLTFExporter/Private/Converters/GLTFSamplerConverters.h
@@ -5,7 +5,6 @@
 #include "Json/GLTFJsonIndex.h"
 #include "Converters/GLTFConverter.h"
 #include "Converters/GLTFBuilderContext.h"
-#include "Engine.h"
 
 class FGLTFSamplerConverter final : public FGLTFBuilderContext, public TGLTFConverter<FGLTFJsonSamplerIndex, const UTexture*>
 {
diff --git a/Engine/Plugins/Enterprise/GLTFExporter/MigratedFromGit/Source/GLTFExporter/Private/Converters/GLTFSceneConverters.h b/Engine/Plugins/Enterprise/GLTFExporter/MigratedFromGit/Source/GLTFExporter/Private/Converters/GLTFSceneConverters.h
--- a/Engine/Plugins/Enterprise/GLTFExporter/MigratedFromGit/Source/GLTFExporter/Private/Converters/GLTFSceneConverters.h
+++ b/Engine/Plugins/Enterprise/GLTFExporter/MigratedFromGit/Source/GLTFExporter/Private/Converters/GLTFSceneConverters.h
@@ -5,7 +5,6 @@
 #include "Json/GLTFJsonIndex.h"
 #include "Converters/GLTFConverter.h"
 #include "Converters/GLTFBuilderContext.h"
-#include "Engine.h"
 
 class FGLTFSceneConverter final : public FGLTFBuilderContext, public TGLTFConverter<FGLTFJsonSceneIndex, const UWorld*>
 {
diff --git a/Engine/Plugins/Enterprise/GLTFExporter/MigratedFromGit/Source/GLTFExporter/Private/Converters/GLTFSkinConverters.h b/Engine/Plugins/Enterprise/GLTFExporter/MigratedFromGit/Source/GLTFExporter/Private/Converters/GLTFSkinConverters.h
--- a/Engine/Plugins/Enterprise/GLTFExporter/MigratedFromGit/Source/GLTFExporter/Private/Converters/GLTFSkinConverters.h
+++ b/Engine/Plugins/Enterprise/GLTFExporter/MigratedFromGit/Source/GLTFExporter/Private/Converters/GLTFSkinConverters.h
@@ -5,7 +5,6 @@
 #include "Json/GLTFJsonIndex.h"
 #include "Converters/GLTFConverter.h"
 #include "Converters/GLTFBuilderContext.h"
-#include "Engine.h"
 
 class FGLTFSkinConverter final : public FGLTFBuilderContext, public TGLTFConverter<FGLTFJsonSkinIndex, FGLTFJsonNodeIndex, const USkeletalMesh*>
 {
diff --git a/Engine/Plugins/Enterprise/GLTFExporter/MigratedFromGit/Source/GLTFExporter/Private/Converters/GLTFSkySphereConverters.cpp b/Engine/Plugins/Enterprise/GLTFExporter/MigratedFromGit/Source/GLTFExporter/Private/Converters/GLTFSkySphereConverters.cpp
--- a/Engine/Plugins/Enterprise/GLTFExporter/MigratedFromGit/Source/GLTFExporter/Private/Converters/GLTFSkySphereConverters.cpp
+++ b/Engine/Plugins/Enterprise/GLTFExporter/MigratedFromGit/Source/GLTFExporter/Private/Converters/GLTFSkySphereConverters.cpp
@@ -6,6 +6,9 @@
 #include "Converters/GLTFActorUtility.h"
 #include "Converters/GLTFCurveUtility.h"
 #include "Builders/GLTFContainerBuilder.h"
+#include "Engine/DirectionalLight.h"
+#include "Curves/CurveLinearColor.h"
+#include "Materials/MaterialInstance.h"
 
 FGLTFJsonSkySphereIndex FGLTFSkySphereConverter::Convert(const AActor* SkySphereActor)
 {
diff --git a/Engine/Plugins/Enterprise/GLTFExporter/MigratedFromGit/Source/GLTFExporter/Private/Converters/GLTFSkySphereConverters.h b/Engine/Plugins/Enterprise/GLTFExporter/MigratedFromGit/Source/GLTFExporter/Private/Converters/GLTFSkySphereConverters.h
--- a/Engine/Plugins/Enterprise/GLTFExporter/MigratedFromGit/Source/GLTFExporter/Private/Converters/GLTFSkySphereConverters.h
+++ b/Engine/Plugins/Enterprise/GLTFExporter/MigratedFromGit/Source/GLTFExporter/Private/Converters/GLTFSkySphereConverters.h
@@ -6,7 +6,6 @@
 #include "Json/GLTFJsonIndex.h"
 #include "Converters/GLTFConverter.h"
 #include "Converters/GLTFBuilderContext.h"
-#include "Engine.h"
 
 struct FGLTFJsonSkySphereColorCurve;
 
diff --git a/Engine/Plugins/Enterprise/GLTFExporter/MigratedFromGit/Source/GLTFExporter/Private/Converters/GLTFTextureConverters.h b/Engine/Plugins/Enterprise/GLTFExporter/MigratedFromGit/Source/GLTFExporter/Private/Converters/GLTFTextureConverters.h
--- a/Engine/Plugins/Enterprise/GLTFExporter/MigratedFromGit/Source/GLTFExporter/Private/Converters/GLTFTextureConverters.h
+++ b/Engine/Plugins/Enterprise/GLTFExporter/MigratedFromGit/Source/GLTFExporter/Private/Converters/GLTFTextureConverters.h
@@ -5,7 +5,6 @@
 #include "Json/GLTFJsonIndex.h"
 #include "Converters/GLTFConverter.h"
 #include "Converters/GLTFBuilderContext.h"
-#include "Engine.h"
 
 template <typename... InputTypes>
 class TGLTFTextureConverter : public FGLTFBuilderContext, public TGLTFConverter<FGLTFJsonTextureIndex, InputTypes...>
diff --git a/Engine/Plugins/Enterprise/GLTFExporter/MigratedFromGit/Source/GLTFExporter/Private/Converters/GLTFTextureUtility.cpp b/Engine/Plugins/Enterprise/GLTFExporter/MigratedFromGit/Source/GLTFExporter/Private/Converters/GLTFTextureUtility.cpp
--- a/Engine/Plugins/Enterprise/GLTFExporter/MigratedFromGit/Source/GLTFExporter/Private/Converters/GLTFTextureUtility.cpp
+++ b/Engine/Plugins/Enterprise/GLTFExporter/MigratedFromGit/Source/GLTFExporter/Private/Converters/GLTFTextureUtility.cpp
@@ -1,7 +1,14 @@
 // Copyright Epic Games, Inc. All Rights Reserved.
 
 #include "Converters/GLTFTextureUtility.h"
-#include "GLTFNormalMapPreview.h"
+#include "Converters/GLTFNormalMapPreview.h"
+#include "CanvasItem.h"
+#include "Engine/Texture2D.h"
+#include "Engine/TextureCube.h"
+#include "Engine/TextureRenderTarget2D.h"
+#include "Engine/TextureRenderTargetCube.h"
+#include "DeviceProfiles/DeviceProfile.h"
+#include "DeviceProfiles/DeviceProfileManager.h"
 
 bool FGLTFTextureUtility::IsHDR(EPixelFormat Format)
 {
diff --git a/Engine/Plugins/Enterprise/GLTFExporter/MigratedFromGit/Source/GLTFExporter/Private/Converters/GLTFTextureUtility.h b/Engine/Plugins/Enterprise/GLTFExporter/MigratedFromGit/Source/GLTFExporter/Private/Converters/GLTFTextureUtility.h
--- a/Engine/Plugins/Enterprise/GLTFExporter/MigratedFromGit/Source/GLTFExporter/Private/Converters/GLTFTextureUtility.h
+++ b/Engine/Plugins/Enterprise/GLTFExporter/MigratedFromGit/Source/GLTFExporter/Private/Converters/GLTFTextureUtility.h
@@ -3,7 +3,6 @@
 #pragma once
 
 #include "Json/GLTFJsonEnums.h"
-#include "Engine.h"
 
 struct FGLTFTextureUtility
 {
diff --git a/Engine/Plugins/Enterprise/GLTFExporter/MigratedFromGit/Source/GLTFExporter/Private/Converters/GLTFUVDegenerateChecker.h b/Engine/Plugins/Enterprise/GLTFExporter/MigratedFromGit/Source/GLTFExporter/Private/Converters/GLTFUVDegenerateChecker.h
--- a/Engine/Plugins/Enterprise/GLTFExporter/MigratedFromGit/Source/GLTFExporter/Private/Converters/GLTFUVDegenerateChecker.h
+++ b/Engine/Plugins/Enterprise/GLTFExporter/MigratedFromGit/Source/GLTFExporter/Private/Converters/GLTFUVDegenerateChecker.h
@@ -6,7 +6,6 @@
 
 #include "Converters/GLTFConverter.h"
 #include "Converters/GLTFIndexArray.h"
-#include "Engine.h"
 
 class FGLTFUVDegenerateChecker final : public TGLTFConverter<float, const FMeshDescription*, FGLTFIndexArray, int32>
 {
diff --git a/Engine/Plugins/Enterprise/GLTFExporter/MigratedFromGit/Source/GLTFExporter/Private/Converters/GLTFUVOverlapChecker.h b/Engine/Plugins/Enterprise/GLTFExporter/MigratedFromGit/Source/GLTFExporter/Private/Converters/GLTFUVOverlapChecker.h
--- a/Engine/Plugins/Enterprise/GLTFExporter/MigratedFromGit/Source/GLTFExporter/Private/Converters/GLTFUVOverlapChecker.h
+++ b/Engine/Plugins/Enterprise/GLTFExporter/MigratedFromGit/Source/GLTFExporter/Private/Converters/GLTFUVOverlapChecker.h
@@ -6,7 +6,6 @@
 
 #include "Converters/GLTFConverter.h"
 #include "Converters/GLTFIndexArray.h"
-#include "Engine.h"
 
 class FGLTFUVOverlapChecker final : public TGLTFConverter<float, const FMeshDescription*, FGLTFIndexArray, int32>
 {
diff --git a/Engine/Plugins/Enterprise/GLTFExporter/MigratedFromGit/Source/GLTFExporter/Private/Tasks/GLTFAnimationTasks.cpp b/Engine/Plugins/Enterprise/GLTFExporter/MigratedFromGit/Source/GLTFExporter/Private/Tasks/GLTFAnimationTasks.cpp
--- a/Engine/Plugins/Enterprise/GLTFExporter/MigratedFromGit/Source/GLTFExporter/Private/Tasks/GLTFAnimationTasks.cpp
+++ b/Engine/Plugins/Enterprise/GLTFExporter/MigratedFromGit/Source/GLTFExporter/Private/Tasks/GLTFAnimationTasks.cpp
@@ -7,6 +7,7 @@
 #include "LevelSequence.h"
 #include "LevelSequencePlayer.h"
 #include "MovieSceneTimeHelpers.h"
+#include "Animation/AnimationPoseData.h"
 #include "Tracks/MovieScene3DTransformTrack.h"
 #include "Sections/MovieScene3DTransformSection.h"
 #include "Channels/MovieSceneChannelProxy.h"
diff --git a/Engine/Plugins/Enterprise/GLTFExporter/MigratedFromGit/Source/GLTFExporter/Private/Tasks/GLTFMaterialTasks.h b/Engine/Plugins/Enterprise/GLTFExporter/MigratedFromGit/Source/GLTFExporter/Private/Tasks/GLTFMaterialTasks.h
--- a/Engine/Plugins/Enterprise/GLTFExporter/MigratedFromGit/Source/GLTFExporter/Private/Tasks/GLTFMaterialTasks.h
+++ b/Engine/Plugins/Enterprise/GLTFExporter/MigratedFromGit/Source/GLTFExporter/Private/Tasks/GLTFMaterialTasks.h
@@ -9,7 +9,6 @@
 #include "Converters/GLTFUVOverlapChecker.h"
 #include "Converters/GLTFMeshData.h"
 #include "MaterialPropertyEx.h"
-#include "Engine.h"
 
 struct FGLTFPropertyBakeOutput;
 
diff --git a/Engine/Plugins/Enterprise/GLTFExporter/MigratedFromGit/Source/GLTFExporter/Private/Tasks/GLTFMeshTasks.h b/Engine/Plugins/Enterprise/GLTFExporter/MigratedFromGit/Source/GLTFExporter/Private/Tasks/GLTFMeshTasks.h
--- a/Engine/Plugins/Enterprise/GLTFExporter/MigratedFromGit/Source/GLTFExporter/Private/Tasks/GLTFMeshTasks.h
+++ b/Engine/Plugins/Enterprise/GLTFExporter/MigratedFromGit/Source/GLTFExporter/Private/Tasks/GLTFMeshTasks.h
@@ -7,7 +7,6 @@
 #include "Converters/GLTFMeshSectionConverters.h"
 #include "Converters/GLTFMaterialArray.h"
 #include "Converters/GLTFNameUtility.h"
-#include "Engine.h"
 
 class FGLTFStaticMeshTask : public FGLTFTask
 {
diff --git a/Engine/Plugins/Enterprise/GLTFExporter/MigratedFromGit/Source/GLTFExporter/Private/Tasks/GLTFTextureTasks.cpp b/Engine/Plugins/Enterprise/GLTFExporter/MigratedFromGit/Source/GLTFExporter/Private/Tasks/GLTFTextureTasks.cpp
--- a/Engine/Plugins/Enterprise/GLTFExporter/MigratedFromGit/Source/GLTFExporter/Private/Tasks/GLTFTextureTasks.cpp
+++ b/Engine/Plugins/Enterprise/GLTFExporter/MigratedFromGit/Source/GLTFExporter/Private/Tasks/GLTFTextureTasks.cpp
@@ -5,6 +5,16 @@
 #include "Converters/GLTFConverterUtility.h"
 #include "Converters/GLTFTextureUtility.h"
 #include "Converters/GLTFNameUtility.h"
+#include "Engine/Texture2D.h"
+#include "Engine/TextureCube.h"
+#include "Engine/TextureRenderTarget2D.h"
+#include "Engine/TextureRenderTargetCube.h"
+#include "Engine/LightMapTexture2D.h"
+
+FString FGLTFTexture2DTask::GetName()
+{
+	return Texture2D->GetName();
+}
 
 void FGLTFTexture2DTask::Complete()
 {
@@ -52,6 +62,11 @@ void FGLTFTexture2DTask::Complete()
 	JsonTexture.Sampler = Builder.GetOrAddSampler(Texture2D);
 }
 
+FString FGLTFTextureCubeTask::GetName()
+{
+	return TextureCube->GetName();
+}
+
 void FGLTFTextureCubeTask::Complete()
 {
 	FGLTFJsonTexture& JsonTexture = Builder.GetTexture(TextureIndex);
@@ -101,6 +116,11 @@ void FGLTFTextureCubeTask::Complete()
 	JsonTexture.Sampler = Builder.GetOrAddSampler(TextureCube);
 }
 
+FString FGLTFTextureRenderTarget2DTask::GetName()
+{
+	return RenderTarget2D->GetName();
+}
+
 void FGLTFTextureRenderTarget2DTask::Complete()
 {
 	FGLTFJsonTexture& JsonTexture = Builder.GetTexture(TextureIndex);
@@ -136,6 +156,11 @@ void FGLTFTextureRenderTarget2DTask::Complete()
 	JsonTexture.Sampler = Builder.GetOrAddSampler(RenderTarget2D);
 }
 
+FString FGLTFTextureRenderTargetCubeTask::GetName()
+{
+	return RenderTargetCube->GetName();
+}
+
 void FGLTFTextureRenderTargetCubeTask::Complete()
 {
 	FGLTFJsonTexture& JsonTexture = Builder.GetTexture(TextureIndex);
@@ -186,6 +211,11 @@ void FGLTFTextureRenderTargetCubeTask::Complete()
 
 #if WITH_EDITOR
 
+FString FGLTFTextureLightMapTask::GetName()
+{
+	return LightMap->GetName();
+}
+
 void FGLTFTextureLightMapTask::Complete()
 {
 	FGLTFJsonTexture& JsonTexture = Builder.GetTexture(TextureIndex);
diff --git a/Engine/Plugins/Enterprise/GLTFExporter/MigratedFromGit/Source/GLTFExporter/Private/Tasks/GLTFTextureTasks.h b/Engine/Plugins/Enterprise/GLTFExporter/MigratedFromGit/Source/GLTFExporter/Private/Tasks/GLTFTextureTasks.h
--- a/Engine/Plugins/Enterprise/GLTFExporter/MigratedFromGit/Source/GLTFExporter/Private/Tasks/GLTFTextureTasks.h
+++ b/Engine/Plugins/Enterprise/GLTFExporter/MigratedFromGit/Source/GLTFExporter/Private/Tasks/GLTFTextureTasks.h
@@ -4,7 +4,6 @@
 
 #include "Tasks/GLTFTask.h"
 #include "Builders/GLTFConvertBuilder.h"
-#include "Engine.h"
 
 class FGLTFTexture2DTask : public FGLTFTask
 {
@@ -19,10 +18,7 @@ public:
 	{
 	}
 
-	virtual FString GetName() override
-	{
-		return Texture2D->GetName();
-	}
+	virtual FString GetName() override;
 
 	virtual void Complete() override;
 
@@ -48,10 +44,7 @@ public:
 	{
 	}
 
-	virtual FString GetName() override
-	{
-		return TextureCube->GetName();
-	}
+	virtual FString GetName() override;
 
 	virtual void Complete() override;
 
@@ -77,10 +70,7 @@ public:
 	{
 	}
 
-	virtual FString GetName() override
-	{
-		return RenderTarget2D->GetName();
-	}
+	virtual FString GetName() override;
 
 	virtual void Complete() override;
 
@@ -106,10 +96,7 @@ public:
 	{
 	}
 
-	virtual FString GetName() override
-	{
-		return RenderTargetCube->GetName();
-	}
+	virtual FString GetName() override;
 
 	virtual void Complete() override;
 
@@ -136,10 +123,7 @@ public:
 	{
 	}
 
-	virtual FString GetName() override
-	{
-		return LightMap->GetName();
-	}
+	virtual FString GetName() override;
 
 	virtual void Complete() override;
 
