[GLTFExporter] chore: enable bTreatAsEngineModule in plugin modules

commit 06e754961e3a23623fa3d47e781e017bb6feeada
Date: Fri, 17 Dec 2021 08:57:27 +0100

    chore: enable bTreatAsEngineModule in plugin modules
    (cherry picked from commit 053c0d4a5d751871c1dcdf490d85176109a56e85​)

#skipci

diff --git a/Engine/Plugins/Enterprise/GLTFExporter/MigratedFromGit/Source/GLTFExporter/GLTFExporter.Build.cs b/Engine/Plugins/Enterprise/GLTFExporter/MigratedFromGit/Source/GLTFExporter/GLTFExporter.Build.cs
--- a/Engine/Plugins/Enterprise/GLTFExporter/MigratedFromGit/Source/GLTFExporter/GLTFExporter.Build.cs
+++ b/Engine/Plugins/Enterprise/GLTFExporter/MigratedFromGit/Source/GLTFExporter/GLTFExporter.Build.cs
@@ -1,71 +1,72 @@
 // Copyright Epic Games, Inc. All Rights Reserved.
 
-namespace UnrealBuildTool.Rules
+using UnrealBuildTool;
+
+public class GLTFExporter : ModuleRules
 {
-	public class GLTFExporter : ModuleRules
+	public GLTFExporter(ReadOnlyTargetRules Target) : base(Target)
 	{
-		public GLTFExporter(ReadOnlyTargetRules Target) : base(Target)
+		bTreatAsEngineModule = true;
+
+		if (IsPlugin)
 		{
-			if (IsPlugin)
+			// NOTE: ugly hack to access plugin info (should propose change to engine)
+			var BindingFlags = System.Reflection.BindingFlags.Instance | System.Reflection.BindingFlags.NonPublic;
+			var FieldInfo = typeof(ModuleRules).GetField("Plugin", BindingFlags);
+			var Plugin = FieldInfo != null ? FieldInfo.GetValue(this) as PluginInfo : null;
+			var PluginDescriptor = Plugin != null ? Plugin.Descriptor : null;
+
+			if (PluginDescriptor != null)
 			{
-				// NOTE: ugly hack to access plugin info (should propose change to engine)
-				var BindingFlags = System.Reflection.BindingFlags.Instance | System.Reflection.BindingFlags.NonPublic;
-				var FieldInfo = typeof(ModuleRules).GetField("Plugin", BindingFlags);
-				var Plugin = FieldInfo != null ? FieldInfo.GetValue(this) as PluginInfo : null;
-				var PluginDescriptor = Plugin != null ? Plugin.Descriptor : null;
+				PrivateDefinitions.Add("GLTFEXPORTER_FRIENDLY_NAME=TEXT(\"" + PluginDescriptor.FriendlyName + "\")");
+				PrivateDefinitions.Add("GLTFEXPORTER_VERSION_NAME=TEXT(\"" + PluginDescriptor.VersionName + "\")");
+			}
+		}
 
-				if (PluginDescriptor != null)
-				{
-					PrivateDefinitions.Add("GLTFEXPORTER_FRIENDLY_NAME=TEXT(\"" + PluginDescriptor.FriendlyName + "\")");
-					PrivateDefinitions.Add("GLTFEXPORTER_VERSION_NAME=TEXT(\"" + PluginDescriptor.VersionName + "\")");
-				}
+		PublicDependencyModuleNames .AddRange(
+			new string[]
+			{
+				"Core",
+				"CoreUObject",
+				"Engine",
 			}
+		);
 
-			PublicDependencyModuleNames .AddRange(
-				new string[]
-				{
-					"Core",
-					"CoreUObject",
-					"Engine",
-				}
-			);
+		PrivateDependencyModuleNames.AddRange(
+			new string[]
+			{
+				"Json",
+				"RenderCore",
+				"RHI",
+				"ImageWrapper",
+				"LevelSequence",
+				"MovieScene",
+				"MovieSceneTracks",
+				"VariantManagerContent",
+			}
+		);
 
+		if (Target.bBuildEditor)
+		{
 			PrivateDependencyModuleNames.AddRange(
 				new string[]
 				{
-					"Json",
-					"RenderCore",
-					"RHI",
-					"ImageWrapper",
-					"LevelSequence",
-					"MovieScene",
-					"MovieSceneTracks",
-					"VariantManagerContent",
+					"UnrealEd",
+					"MessageLog",
+					"Slate",
+					"SlateCore",
+					"Mainframe",
+					"InputCore",
+					"EditorStyle",
+					"PropertyEditor",
+					"Projects",
+					"MeshMergeUtilities",
+					"MeshDescription",
+					"StaticMeshDescription",
+					"GLTFMaterialBaking",
+					"GLTFMaterialAnalyzer",
 				}
 			);
-
-			if (Target.bBuildEditor)
-			{
-				PrivateDependencyModuleNames.AddRange(
-					new string[]
-					{
-						"UnrealEd",
-						"MessageLog",
-						"Slate",
-						"SlateCore",
-						"Mainframe",
-						"InputCore",
-						"EditorStyle",
-						"PropertyEditor",
-						"Projects",
-						"MeshMergeUtilities",
-						"MeshDescription",
-						"StaticMeshDescription",
-						"GLTFMaterialBaking",
-						"GLTFMaterialAnalyzer",
-					}
-				);
-			}
 		}
 	}
 }
diff --git a/Engine/Plugins/Enterprise/GLTFExporter/MigratedFromGit/Source/GLTFMaterialAnalyzer/GLTFMaterialAnalyzer.Build.cs b/Engine/Plugins/Enterprise/GLTFExporter/MigratedFromGit/Source/GLTFMaterialAnalyzer/GLTFMaterialAnalyzer.Build.cs
--- a/Engine/Plugins/Enterprise/GLTFExporter/MigratedFromGit/Source/GLTFMaterialAnalyzer/GLTFMaterialAnalyzer.Build.cs
+++ b/Engine/Plugins/Enterprise/GLTFExporter/MigratedFromGit/Source/GLTFMaterialAnalyzer/GLTFMaterialAnalyzer.Build.cs
@@ -6,6 +6,8 @@ public class GLTFMaterialAnalyzer : ModuleRules
 {
 	public GLTFMaterialAnalyzer(ReadOnlyTargetRules Target) : base(Target)
 	{
+		bTreatAsEngineModule = true;
+
 		PublicDependencyModuleNames .AddRange(
 			new string[]
 			{
diff --git a/Engine/Plugins/Enterprise/GLTFExporter/MigratedFromGit/Source/GLTFMaterialBaking/GLTFMaterialBaking.Build.cs b/Engine/Plugins/Enterprise/GLTFExporter/MigratedFromGit/Source/GLTFMaterialBaking/GLTFMaterialBaking.Build.cs
--- a/Engine/Plugins/Enterprise/GLTFExporter/MigratedFromGit/Source/GLTFMaterialBaking/GLTFMaterialBaking.Build.cs
+++ b/Engine/Plugins/Enterprise/GLTFExporter/MigratedFromGit/Source/GLTFMaterialBaking/GLTFMaterialBaking.Build.cs
@@ -6,6 +6,8 @@ public class GLTFMaterialBaking : ModuleRules
 {
 	public GLTFMaterialBaking(ReadOnlyTargetRules Target) : base(Target)
 	{
+		bTreatAsEngineModule = true;
+
 		PublicDependencyModuleNames .AddRange(
 			new string[]
 			{
