[GLTFExporter] chore: use PCHUsageMode.UseExplicitOrSharedPCHs to improve engine plugin compliance

commit 5bb79c5083d5fe927c7f86e118e44e42f01554da
Date: Fri, 17 Dec 2021 22:56:04 +0100

    chore: use PCHUsageMode.UseExplicitOrSharedPCHs to improve engine plugin compliance
    (cherry picked from commit 8d54056b9b89dee5668285581dd5445b6bbad2fb​)

#skipci

diff --git a/Engine/Plugins/Enterprise/GLTFExporter/MigratedFromGit/Source/GLTFExporter/GLTFExporter.Build.cs b/Engine/Plugins/Enterprise/GLTFExporter/MigratedFromGit/Source/GLTFExporter/GLTFExporter.Build.cs
--- a/Engine/Plugins/Enterprise/GLTFExporter/MigratedFromGit/Source/GLTFExporter/GLTFExporter.Build.cs
+++ b/Engine/Plugins/Enterprise/GLTFExporter/MigratedFromGit/Source/GLTFExporter/GLTFExporter.Build.cs
@@ -6,7 +6,8 @@ public class GLTFExporter : ModuleRules
 {
 	public GLTFExporter(ReadOnlyTargetRules Target) : base(Target)
 	{
-		bTreatAsEngineModule = true;
+		PCHUsage = PCHUsageMode.UseExplicitOrSharedPCHs;
+		bTreatAsEngineModule = true; // Only necessary when plugin installed in project
 
 		if (IsPlugin)
 		{
diff --git a/Engine/Plugins/Enterprise/GLTFExporter/MigratedFromGit/Source/GLTFMaterialAnalyzer/GLTFMaterialAnalyzer.Build.cs b/Engine/Plugins/Enterprise/GLTFExporter/MigratedFromGit/Source/GLTFMaterialAnalyzer/GLTFMaterialAnalyzer.Build.cs
--- a/Engine/Plugins/Enterprise/GLTFExporter/MigratedFromGit/Source/GLTFMaterialAnalyzer/GLTFMaterialAnalyzer.Build.cs
+++ b/Engine/Plugins/Enterprise/GLTFExporter/MigratedFromGit/Source/GLTFMaterialAnalyzer/GLTFMaterialAnalyzer.Build.cs
@@ -6,7 +6,8 @@ public class GLTFMaterialAnalyzer : ModuleRules
 {
 	public GLTFMaterialAnalyzer(ReadOnlyTargetRules Target) : base(Target)
 	{
-		bTreatAsEngineModule = true;
+		PCHUsage = PCHUsageMode.UseExplicitOrSharedPCHs;
+		bTreatAsEngineModule = true; // Only necessary when plugin installed in project
 
 		PublicDependencyModuleNames .AddRange(
 			new string[]
diff --git a/Engine/Plugins/Enterprise/GLTFExporter/MigratedFromGit/Source/GLTFMaterialBaking/GLTFMaterialBaking.Build.cs b/Engine/Plugins/Enterprise/GLTFExporter/MigratedFromGit/Source/GLTFMaterialBaking/GLTFMaterialBaking.Build.cs
--- a/Engine/Plugins/Enterprise/GLTFExporter/MigratedFromGit/Source/GLTFMaterialBaking/GLTFMaterialBaking.Build.cs
+++ b/Engine/Plugins/Enterprise/GLTFExporter/MigratedFromGit/Source/GLTFMaterialBaking/GLTFMaterialBaking.Build.cs
@@ -6,7 +6,8 @@ public class GLTFMaterialBaking : ModuleRules
 {
 	public GLTFMaterialBaking(ReadOnlyTargetRules Target) : base(Target)
 	{
-		bTreatAsEngineModule = true;
+		PCHUsage = PCHUsageMode.UseExplicitOrSharedPCHs;
+		bTreatAsEngineModule = true; // Only necessary when plugin installed in project
 
 		PublicDependencyModuleNames .AddRange(
 			new string[]
