[GLTFExporter] refactor: move some utilities from GLTFLevelVariantSetsConverter to GLTFVariantUtility

commit 90a9735be8df0af409f468179fa22803fb55f671
Date: Tue, 28 Dec 2021 03:42:50 +0100

    refactor: move some utilities from GLTFLevelVariantSetsConverter to GLTFVariantUtility
    (cherry picked from commit 729c088778bc9b01a361d0dd17bf25c1c87e55cc​)

#skipci

diff --git a/Engine/Plugins/Enterprise/GLTFExporter/MigratedFromGit/Source/GLTFExporter/Private/Converters/GLTFLevelVariantSetsConverters.cpp b/Engine/Plugins/Enterprise/GLTFExporter/MigratedFromGit/Source/GLTFExporter/Private/Converters/GLTFLevelVariantSetsConverters.cpp
--- a/Engine/Plugins/Enterprise/GLTFExporter/MigratedFromGit/Source/GLTFExporter/Private/Converters/GLTFLevelVariantSetsConverters.cpp
+++ b/Engine/Plugins/Enterprise/GLTFExporter/MigratedFromGit/Source/GLTFExporter/Private/Converters/GLTFLevelVariantSetsConverters.cpp
@@ -1,6 +1,7 @@
 // Copyright Epic Games, Inc. All Rights Reserved.
 
 #include "Converters/GLTFLevelVariantSetsConverters.h"
+#include "Converters/GLTFVariantUtility.h"
 #include "Converters/GLTFMeshUtility.h"
 #include "Builders/GLTFContainerBuilder.h"
 #include "Rendering/SkeletalMeshRenderData.h"
@@ -131,7 +132,7 @@ bool FGLTFLevelVariantSetsConverter::TryParseVisibilityPropertyValue(FGLTFJsonVa
 	{
 		Builder.LogWarning(FString::Printf(
 			TEXT("Target object for property is invalid, the property will be skipped. Context: %s"),
-			*GetLogContext(Property)));
+			*FGLTFVariantUtility::GetLogContext(Property)));
 		return false;
 	}
 
@@ -139,7 +140,7 @@ bool FGLTFLevelVariantSetsConverter::TryParseVisibilityPropertyValue(FGLTFJsonVa
 	if (Owner == nullptr)
 	{
 		Builder.LogWarning(FString::Printf(TEXT("Invalid scene component, the property will be skipped. Context: %s"),
-			*GetLogContext(Property)));
+			*FGLTFVariantUtility::GetLogContext(Property)));
 		return false;
 	}
 
@@ -147,17 +148,17 @@ bool FGLTFLevelVariantSetsConverter::TryParseVisibilityPropertyValue(FGLTFJsonVa
 	{
 		Builder.LogWarning(FString::Printf(
 			TEXT("Target actor for property is not selected for export, the property will be skipped. Context: %s"),
-			*GetLogContext(Property)));
+			*FGLTFVariantUtility::GetLogContext(Property)));
 		return false;
 	}
 
 	bool bIsVisible;
 
-	if (!TryGetPropertyValue(const_cast<UPropertyValue*>(Property), bIsVisible))
+	if (!FGLTFVariantUtility::TryGetPropertyValue(const_cast<UPropertyValue*>(Property), bIsVisible))
 	{
 		Builder.LogWarning(FString::Printf(
 			TEXT("Failed to parse recorded data for property, it will be skipped. Context: %s"),
-			*GetLogContext(Property)));
+			*FGLTFVariantUtility::GetLogContext(Property)));
 		return false;
 	}
 
@@ -178,7 +179,7 @@ bool FGLTFLevelVariantSetsConverter::TryParseMaterialPropertyValue(FGLTFJsonVari
 	{
 		Builder.LogWarning(FString::Printf(
 			TEXT("Material property is invalid, the property will be skipped. Context: %s"),
-			*GetLogContext(Property)));
+			*FGLTFVariantUtility::GetLogContext(Property)));
 		return false;
 	}
 
@@ -187,7 +188,7 @@ bool FGLTFLevelVariantSetsConverter::TryParseMaterialPropertyValue(FGLTFJsonVari
 	{
 		Builder.LogWarning(FString::Printf(
 			TEXT("Target object for property is invalid, the property will be skipped. Context: %s"),
-			*GetLogContext(Property)));
+			*FGLTFVariantUtility::GetLogContext(Property)));
 		return false;
 	}
 
@@ -195,7 +196,7 @@ bool FGLTFLevelVariantSetsConverter::TryParseMaterialPropertyValue(FGLTFJsonVari
 	if (Owner == nullptr)
 	{
 		Builder.LogWarning(FString::Printf(TEXT("Invalid scene component, the property will be skipped. Context: %s"),
-			*GetLogContext(Property)));
+			*FGLTFVariantUtility::GetLogContext(Property)));
 		return false;
 	}
 
@@ -203,7 +204,7 @@ bool FGLTFLevelVariantSetsConverter::TryParseMaterialPropertyValue(FGLTFJsonVari
 	{
 		Builder.LogWarning(FString::Printf(
 			TEXT("Target actor for property is not selected for export, the property will be skipped. Context: %s"),
-			*GetLogContext(Property)));
+			*FGLTFVariantUtility::GetLogContext(Property)));
 		return false;
 	}
 
@@ -214,7 +215,7 @@ bool FGLTFLevelVariantSetsConverter::TryParseMaterialPropertyValue(FGLTFJsonVari
 	{
 		Builder.LogWarning(FString::Printf(
 			TEXT("Failed to parse element index to apply the material to, the property will be skipped. Context: %s"),
-			*GetLogContext(Property)));
+			*FGLTFVariantUtility::GetLogContext(Property)));
 		return false;
 	}
 
@@ -281,7 +282,7 @@ bool FGLTFLevelVariantSetsConverter::TryParseStaticMeshPropertyValue(FGLTFJsonVa
 	{
 		Builder.LogWarning(FString::Printf(
 			TEXT("Target object for property is invalid, the property will be skipped. Context: %s"),
-			*GetLogContext(Property)));
+			*FGLTFVariantUtility::GetLogContext(Property)));
 		return false;
 	}
 
@@ -289,7 +290,7 @@ bool FGLTFLevelVariantSetsConverter::TryParseStaticMeshPropertyValue(FGLTFJsonVa
 	if (Owner == nullptr)
 	{
 		Builder.LogWarning(FString::Printf(TEXT("Invalid scene component, the property will be skipped. Context: %s"),
-			*GetLogContext(Property)));
+			*FGLTFVariantUtility::GetLogContext(Property)));
 		return false;
 	}
 
@@ -297,7 +298,7 @@ bool FGLTFLevelVariantSetsConverter::TryParseStaticMeshPropertyValue(FGLTFJsonVa
 	{
 		Builder.LogWarning(FString::Printf(
 			TEXT("Target actor for property is not selected for export, the property will be skipped. Context: %s"),
-			*GetLogContext(Property)));
+			*FGLTFVariantUtility::GetLogContext(Property)));
 		return false;
 	}
 
@@ -306,17 +307,17 @@ bool FGLTFLevelVariantSetsConverter::TryParseStaticMeshPropertyValue(FGLTFJsonVa
 	{
 		Builder.LogWarning(FString::Printf(
 			TEXT("Target object for property has no mesh-component, the property will be skipped. Context: %s"),
-			*GetLogContext(Property)));
+			*FGLTFVariantUtility::GetLogContext(Property)));
 		return false;
 	}
 
 	const UStaticMesh* StaticMesh;
 
-	if (!TryGetPropertyValue(const_cast<UPropertyValue*>(Property), StaticMesh))
+	if (!FGLTFVariantUtility::TryGetPropertyValue(const_cast<UPropertyValue*>(Property), StaticMesh))
 	{
 		Builder.LogWarning(FString::Printf(
 			TEXT("Failed to parse recorded data for property, it will be skipped. Context: %s"),
-			*GetLogContext(Property)));
+			*FGLTFVariantUtility::GetLogContext(Property)));
 		return false;
 	}
 
@@ -338,7 +339,7 @@ bool FGLTFLevelVariantSetsConverter::TryParseSkeletalMeshPropertyValue(FGLTFJson
 	{
 		Builder.LogWarning(FString::Printf(
 			TEXT("Target object for property is invalid, the property will be skipped. Context: %s"),
-			*GetLogContext(Property)));
+			*FGLTFVariantUtility::GetLogContext(Property)));
 		return false;
 	}
 
@@ -346,7 +347,7 @@ bool FGLTFLevelVariantSetsConverter::TryParseSkeletalMeshPropertyValue(FGLTFJson
 	if (Owner == nullptr)
 	{
 		Builder.LogWarning(FString::Printf(TEXT("Invalid scene component, the property will be skipped. Context: %s"),
-			*GetLogContext(Property)));
+			*FGLTFVariantUtility::GetLogContext(Property)));
 		return false;
 	}
 
@@ -354,7 +355,7 @@ bool FGLTFLevelVariantSetsConverter::TryParseSkeletalMeshPropertyValue(FGLTFJson
 	{
 		Builder.LogWarning(FString::Printf(
 			TEXT("Target actor for property is not selected for export, the property will be skipped. Context: %s"),
-			*GetLogContext(Property)));
+			*FGLTFVariantUtility::GetLogContext(Property)));
 		return false;
 	}
 
@@ -363,17 +364,17 @@ bool FGLTFLevelVariantSetsConverter::TryParseSkeletalMeshPropertyValue(FGLTFJson
 	{
 		Builder.LogWarning(FString::Printf(
 			TEXT("Target object for property has no mesh-component, the property will be skipped. Context: %s"),
-			*GetLogContext(Property)));
+			*FGLTFVariantUtility::GetLogContext(Property)));
 		return false;
 	}
 
 	const USkeletalMesh* SkeletalMesh;
 
-	if (!TryGetPropertyValue(const_cast<UPropertyValue*>(Property), SkeletalMesh) || SkeletalMesh == nullptr)
+	if (!FGLTFVariantUtility::TryGetPropertyValue(const_cast<UPropertyValue*>(Property), SkeletalMesh) || SkeletalMesh == nullptr)
 	{
 		Builder.LogWarning(FString::Printf(
 			TEXT("Failed to parse recorded data for property, it will be skipped. Context: %s"),
-			*GetLogContext(Property)));
+			*FGLTFVariantUtility::GetLogContext(Property)));
 		return false;
 	}
 
@@ -387,44 +388,3 @@ bool FGLTFLevelVariantSetsConverter::TryParseSkeletalMeshPropertyValue(FGLTFJson
 	NodeProperties.Mesh = MeshIndex;
 	return true;
 }
-
-template<typename T>
-bool FGLTFLevelVariantSetsConverter::TryGetPropertyValue(UPropertyValue* Property, T& OutValue) const
-{
-	if (Property == nullptr || !Property->HasRecordedData())
-	{
-		return false;
-	}
-
-	FMemory::Memcpy(&OutValue, Property->GetRecordedData().GetData(), sizeof(T));
-	return true;
-}
-
-FString FGLTFLevelVariantSetsConverter::GetLogContext(const UPropertyValue* Property) const
-{
-	const UVariantObjectBinding* Parent = Property->GetParent();
-	return GetLogContext(Parent) + TEXT("/") + Property->GetFullDisplayString();
-}
-
-FString FGLTFLevelVariantSetsConverter::GetLogContext(const UVariantObjectBinding* Binding) const
-{
-	const UVariant* Parent = const_cast<UVariantObjectBinding*>(Binding)->GetParent();
-	return GetLogContext(Parent) + TEXT("/") + Binding->GetDisplayText().ToString();
-}
-
-FString FGLTFLevelVariantSetsConverter::GetLogContext(const UVariant* Variant) const
-{
-	const UVariantSet* Parent = const_cast<UVariant*>(Variant)->GetParent();
-	return GetLogContext(Parent) + TEXT("/") + Variant->GetDisplayText().ToString();
-}
-
-FString FGLTFLevelVariantSetsConverter::GetLogContext(const UVariantSet* VariantSet) const
-{
-	const ULevelVariantSets* Parent = const_cast<UVariantSet*>(VariantSet)->GetParent();
-	return GetLogContext(Parent) + TEXT("/") + VariantSet->GetDisplayText().ToString();
-}
-
-FString FGLTFLevelVariantSetsConverter::GetLogContext(const ULevelVariantSets* LevelVariantSets) const
-{
-	return LevelVariantSets->GetName();
-}
diff --git a/Engine/Plugins/Enterprise/GLTFExporter/MigratedFromGit/Source/GLTFExporter/Private/Converters/GLTFLevelVariantSetsConverters.h b/Engine/Plugins/Enterprise/GLTFExporter/MigratedFromGit/Source/GLTFExporter/Private/Converters/GLTFLevelVariantSetsConverters.h
--- a/Engine/Plugins/Enterprise/GLTFExporter/MigratedFromGit/Source/GLTFExporter/Private/Converters/GLTFLevelVariantSetsConverters.h
+++ b/Engine/Plugins/Enterprise/GLTFExporter/MigratedFromGit/Source/GLTFExporter/Private/Converters/GLTFLevelVariantSetsConverters.h
@@ -25,13 +25,4 @@ class FGLTFLevelVariantSetsConverter final : public FGLTFBuilderContext, public
 	bool TryParseMaterialPropertyValue(FGLTFJsonVariant& OutVariant, const UPropertyValue* Property) const;
 	bool TryParseStaticMeshPropertyValue(FGLTFJsonVariant& OutVariant, const UPropertyValue* Property) const;
 	bool TryParseSkeletalMeshPropertyValue(FGLTFJsonVariant& OutVariant, const UPropertyValue* Property) const;
-
-	template<typename T>
-	bool TryGetPropertyValue(UPropertyValue* Property, T& OutValue) const;
-
-	FString GetLogContext(const UPropertyValue* Property) const;
-	FString GetLogContext(const UVariantObjectBinding* Binding) const;
-	FString GetLogContext(const UVariant* Variant) const;
-	FString GetLogContext(const UVariantSet* VariantSet) const;
-	FString GetLogContext(const ULevelVariantSets* LevelVariantSets) const;
 };
diff --git a/Engine/Plugins/Enterprise/GLTFExporter/MigratedFromGit/Source/GLTFExporter/Private/Converters/GLTFVariantUtility.cpp b/Engine/Plugins/Enterprise/GLTFExporter/MigratedFromGit/Source/GLTFExporter/Private/Converters/GLTFVariantUtility.cpp
new file mode 100644
--- /dev/null
+++ b/Engine/Plugins/Enterprise/GLTFExporter/MigratedFromGit/Source/GLTFExporter/Private/Converters/GLTFVariantUtility.cpp
@@ -0,0 +1,51 @@
+// Copyright Epic Games, Inc. All Rights Reserved.
+
+#include "Converters/GLTFVariantUtility.h"
+#include "LevelVariantSets.h"
+#include "VariantSet.h"
+#include "Variant.h"
+#include "VariantObjectBinding.h"
+#include "PropertyValue.h"
+
+bool FGLTFVariantUtility::TryGetPropertyValue(UPropertyValue* Property, void* OutData, uint32 OutSize)
+{
+	if (Property == nullptr || !Property->HasRecordedData())
+	{
+		return false;
+	}
+
+	const TArray<uint8>& RecordedData = Property->GetRecordedData();
+	check(OutSize == RecordedData.Num());
+
+	FMemory::Memcpy(OutData, RecordedData.GetData(), OutSize);
+	return true;
+}
+
+FString FGLTFVariantUtility::GetLogContext(const UPropertyValue* Property)
+{
+	const UVariantObjectBinding* Parent = Property->GetParent();
+	return GetLogContext(Parent) + TEXT("/") + Property->GetFullDisplayString();
+}
+
+FString FGLTFVariantUtility::GetLogContext(const UVariantObjectBinding* Binding)
+{
+	const UVariant* Parent = const_cast<UVariantObjectBinding*>(Binding)->GetParent();
+	return GetLogContext(Parent) + TEXT("/") + Binding->GetDisplayText().ToString();
+}
+
+FString FGLTFVariantUtility::GetLogContext(const UVariant* Variant)
+{
+	const UVariantSet* Parent = const_cast<UVariant*>(Variant)->GetParent();
+	return GetLogContext(Parent) + TEXT("/") + Variant->GetDisplayText().ToString();
+}
+
+FString FGLTFVariantUtility::GetLogContext(const UVariantSet* VariantSet)
+{
+	const ULevelVariantSets* Parent = const_cast<UVariantSet*>(VariantSet)->GetParent();
+	return GetLogContext(Parent) + TEXT("/") + VariantSet->GetDisplayText().ToString();
+}
+
+FString FGLTFVariantUtility::GetLogContext(const ULevelVariantSets* LevelVariantSets)
+{
+	return LevelVariantSets->GetName();
+}
diff --git a/Engine/Plugins/Enterprise/GLTFExporter/MigratedFromGit/Source/GLTFExporter/Private/Converters/GLTFVariantUtility.h b/Engine/Plugins/Enterprise/GLTFExporter/MigratedFromGit/Source/GLTFExporter/Private/Converters/GLTFVariantUtility.h
new file mode 100644
--- /dev/null
+++ b/Engine/Plugins/Enterprise/GLTFExporter/MigratedFromGit/Source/GLTFExporter/Private/Converters/GLTFVariantUtility.h
@@ -0,0 +1,27 @@
+// Copyright Epic Games, Inc. All Rights Reserved.
+
+#pragma once
+
+struct FCapturedPropSegment;
+class ULevelVariantSets;
+class UVariantSet;
+class UVariant;
+class UVariantObjectBinding;
+class UPropertyValue;
+
+struct FGLTFVariantUtility
+{
+	template<typename T>
+	static bool TryGetPropertyValue(UPropertyValue* Property, T& OutValue)
+	{
+		return TryGetPropertyValue(Property, &OutValue, sizeof(T));
+	}
+
+	static bool TryGetPropertyValue(UPropertyValue* Property, void* OutData, uint32 OutSize);
+
+	static FString GetLogContext(const UPropertyValue* Property);
+	static FString GetLogContext(const UVariantObjectBinding* Binding);
+	static FString GetLogContext(const UVariant* Variant);
+	static FString GetLogContext(const UVariantSet* VariantSet);
+	static FString GetLogContext(const ULevelVariantSets* LevelVariantSets);
+};
