[GLTFExporter] refactor: add sanity checks in module C# build files to avoid silencing potential issues

commit 6624fb4ea7c7480a77ac05914c12220c378ed242
Date: Tue, 28 Dec 2021 04:05:15 +0100

    refactor: add sanity checks in module C# build files to avoid silencing potential issues
    (cherry picked from commit 93b130c0bfde0e009afc835260a5d8ba4c9717e7​)

#skipci

diff --git a/Engine/Plugins/Enterprise/GLTFExporter/MigratedFromGit/Source/GLTFExporter/GLTFExporter.Build.cs b/Engine/Plugins/Enterprise/GLTFExporter/MigratedFromGit/Source/GLTFExporter/GLTFExporter.Build.cs
--- a/Engine/Plugins/Enterprise/GLTFExporter/MigratedFromGit/Source/GLTFExporter/GLTFExporter.Build.cs
+++ b/Engine/Plugins/Enterprise/GLTFExporter/MigratedFromGit/Source/GLTFExporter/GLTFExporter.Build.cs
@@ -9,19 +9,28 @@ public class GLTFExporter : ModuleRules
 		PCHUsage = PCHUsageMode.UseExplicitOrSharedPCHs;
 		bTreatAsEngineModule = true; // Only necessary when plugin installed in project
 
-		if (IsPlugin)
 		{
 			// NOTE: ugly hack to access plugin info (should propose change to engine)
-			var BindingFlags = System.Reflection.BindingFlags.Instance | System.Reflection.BindingFlags.NonPublic;
-			var FieldInfo = typeof(ModuleRules).GetField("Plugin", BindingFlags);
-			var Plugin = FieldInfo != null ? FieldInfo.GetValue(this) as PluginInfo : null;
-			var PluginDescriptor = Plugin != null ? Plugin.Descriptor : null;
+			var PluginField = typeof(ModuleRules).GetField("Plugin", System.Reflection.BindingFlags.Instance | System.Reflection.BindingFlags.NonPublic);
+			if (PluginField == null)
+			{
+				throw new System.Exception("Missing internal member Plugin");
+			}
+
+			var Plugin = PluginField.GetValue(this) as PluginInfo;
+			if (Plugin == null)
+			{
+				throw new System.Exception("Missing plugin information");
+			}
 
-			if (PluginDescriptor != null)
+			var PluginDescriptor = Plugin.Descriptor;
+			if (PluginDescriptor == null)
 			{
-				PrivateDefinitions.Add("GLTFEXPORTER_FRIENDLY_NAME=TEXT(\"" + PluginDescriptor.FriendlyName + "\")");
-				PrivateDefinitions.Add("GLTFEXPORTER_VERSION_NAME=TEXT(\"" + PluginDescriptor.VersionName + "\")");
+				throw new System.Exception("Missing plugin descriptor");
 			}
+
+			PrivateDefinitions.Add("GLTFEXPORTER_FRIENDLY_NAME=TEXT(\"" + PluginDescriptor.FriendlyName + "\")");
+			PrivateDefinitions.Add("GLTFEXPORTER_VERSION_NAME=TEXT(\"" + PluginDescriptor.VersionName + "\")");
 		}
 
 		PublicDependencyModuleNames .AddRange(
