[GLTFExporter] refactor: unify entire code-base across all supported UE versions (4.27, 4.26, and 4.25), to ease future back-porting of features between branches

commit 685e5d902f4034d968d9179bc54332859d4f1340
Date: Tue, 28 Dec 2021 04:06:34 +0100

    refactor: unify entire code-base across all supported UE versions (4.27, 4.26, and 4.25), to ease future back-porting of features between branches
    (cherry picked from commit 1c7ab27a219826d18750af228ae77e69950488af​)

#skipci

diff --git a/Engine/Plugins/Enterprise/GLTFExporter/MigratedFromGit/Source/GLTFExporter/Private/Actors/GLTFHotspotActor.cpp b/Engine/Plugins/Enterprise/GLTFExporter/MigratedFromGit/Source/GLTFExporter/Private/Actors/GLTFHotspotActor.cpp
--- a/Engine/Plugins/Enterprise/GLTFExporter/MigratedFromGit/Source/GLTFExporter/Private/Actors/GLTFHotspotActor.cpp
+++ b/Engine/Plugins/Enterprise/GLTFExporter/MigratedFromGit/Source/GLTFExporter/Private/Actors/GLTFHotspotActor.cpp
@@ -453,13 +453,19 @@ void AGLTFHotspotActor::ValidateAnimation()
 {
 	if (SkeletalMeshActor != nullptr && AnimationSequence != nullptr)
 	{
-		if (USkeletalMesh* SkeletalMesh = SkeletalMeshActor->GetSkeletalMeshComponent()->SkeletalMesh)
+		if (const USkeletalMesh* SkeletalMesh = SkeletalMeshActor->GetSkeletalMeshComponent()->SkeletalMesh)
 		{
-			if (SkeletalMesh->GetSkeleton() != AnimationSequence->GetSkeleton())
+#if (ENGINE_MAJOR_VERSION > 4 || ENGINE_MINOR_VERSION >= 27)
+			const USkeleton* Skeleton = SkeletalMesh->GetSkeleton();
+#else
+			const USkeleton* Skeleton = SkeletalMesh->Skeleton;
+#endif
+
+			if (Skeleton != AnimationSequence->GetSkeleton())
 			{
-				if (SkeletalMesh->GetSkeleton() != nullptr)
+				if (Skeleton != nullptr)
 				{
-					UE_LOG(LogGLTFHotspot, Warning, TEXT("Animation %s is incompatible with skeleton %s, removing animation from actor."), *AnimationSequence->GetName(), *SkeletalMesh->GetSkeleton()->GetName());
+					UE_LOG(LogGLTFHotspot, Warning, TEXT("Animation %s is incompatible with skeleton %s, removing animation from actor."), *AnimationSequence->GetName(), *Skeleton->GetName());
 				}
 				else
 				{
diff --git a/Engine/Plugins/Enterprise/GLTFExporter/MigratedFromGit/Source/GLTFExporter/Private/Converters/GLTFAnimationConverters.cpp b/Engine/Plugins/Enterprise/GLTFExporter/MigratedFromGit/Source/GLTFExporter/Private/Converters/GLTFAnimationConverters.cpp
--- a/Engine/Plugins/Enterprise/GLTFExporter/MigratedFromGit/Source/GLTFExporter/Private/Converters/GLTFAnimationConverters.cpp
+++ b/Engine/Plugins/Enterprise/GLTFExporter/MigratedFromGit/Source/GLTFExporter/Private/Converters/GLTFAnimationConverters.cpp
@@ -13,14 +13,20 @@ FGLTFJsonAnimationIndex FGLTFAnimationConverter::Convert(FGLTFJsonNodeIndex Root
 		return FGLTFJsonAnimationIndex(INDEX_NONE);
 	}
 
-	const USkeleton* Skeleton = AnimSequence->GetSkeleton();
-	if (Skeleton == nullptr)
+	const USkeleton* AnimSkeleton = AnimSequence->GetSkeleton();
+	if (AnimSkeleton == nullptr)
 	{
 		// TODO: report error
 		return FGLTFJsonAnimationIndex(INDEX_NONE);
 	}
 
-	if (Skeleton != SkeletalMesh->GetSkeleton())
+#if (ENGINE_MAJOR_VERSION > 4 || ENGINE_MINOR_VERSION >= 27)
+	const USkeleton* MeshSkeleton = SkeletalMesh->GetSkeleton();
+#else
+	const USkeleton* MeshSkeleton = SkeletalMesh->Skeleton;
+#endif
+
+	if (AnimSkeleton != MeshSkeleton)
 	{
 		// TODO: report error
 		return FGLTFJsonAnimationIndex(INDEX_NONE);
diff --git a/Engine/Plugins/Enterprise/GLTFExporter/MigratedFromGit/Source/GLTFExporter/Private/Converters/GLTFBufferAdapter.cpp b/Engine/Plugins/Enterprise/GLTFExporter/MigratedFromGit/Source/GLTFExporter/Private/Converters/GLTFBufferAdapter.cpp
--- a/Engine/Plugins/Enterprise/GLTFExporter/MigratedFromGit/Source/GLTFExporter/Private/Converters/GLTFBufferAdapter.cpp
+++ b/Engine/Plugins/Enterprise/GLTFExporter/MigratedFromGit/Source/GLTFExporter/Private/Converters/GLTFBufferAdapter.cpp
@@ -58,28 +58,28 @@ TUniquePtr<IGLTFBufferAdapter> IGLTFBufferAdapter::GetIndices(const FRawStaticIn
 
 TUniquePtr<IGLTFBufferAdapter> IGLTFBufferAdapter::GetPositions(const FPositionVertexBuffer* VertexBuffer)
 {
-	const void* PositionData = VertexBuffer->GetVertexData();
+	const void* PositionData = const_cast<FPositionVertexBuffer*>(VertexBuffer)->GetVertexData();
 	if (PositionData != nullptr && (WITH_EDITOR || FGLTFBufferUtility::GetAllowCPUAccess(VertexBuffer))) return MakeUnique<FGLTFBufferAdapterCPU>(PositionData);
 	return MakeUnique<FGLTFBufferAdapterGPU>(VertexBuffer->VertexBufferRHI);
 }
 
 TUniquePtr<IGLTFBufferAdapter> IGLTFBufferAdapter::GetColors(const FColorVertexBuffer* VertexBuffer)
 {
-	const void* ColorData = VertexBuffer->GetVertexData();
+	const void* ColorData = const_cast<FColorVertexBuffer*>(VertexBuffer)->GetVertexData();
 	if (ColorData != nullptr && (WITH_EDITOR || FGLTFBufferUtility::GetAllowCPUAccess(VertexBuffer))) return MakeUnique<FGLTFBufferAdapterCPU>(ColorData);
 	return MakeUnique<FGLTFBufferAdapterGPU>(VertexBuffer->VertexBufferRHI);
 }
 
 TUniquePtr<IGLTFBufferAdapter> IGLTFBufferAdapter::GetTangents(const FStaticMeshVertexBuffer* VertexBuffer)
 {
-	const void* TangentData = VertexBuffer->GetTangentData();
+	const void* TangentData = const_cast<FStaticMeshVertexBuffer*>(VertexBuffer)->GetTangentData();
 	if (TangentData != nullptr && (WITH_EDITOR || const_cast<FStaticMeshVertexBuffer*>(VertexBuffer)->GetAllowCPUAccess())) return MakeUnique<FGLTFBufferAdapterCPU>(TangentData);
 	return MakeUnique<FGLTFBufferAdapterGPU>(VertexBuffer->TangentsVertexBuffer.VertexBufferRHI);
 }
 
 TUniquePtr<IGLTFBufferAdapter> IGLTFBufferAdapter::GetUVs(const FStaticMeshVertexBuffer* VertexBuffer)
 {
-	const void* UVData = VertexBuffer->GetTexCoordData();
+	const void* UVData = const_cast<FStaticMeshVertexBuffer*>(VertexBuffer)->GetTexCoordData();
 	if (UVData != nullptr && (WITH_EDITOR || const_cast<FStaticMeshVertexBuffer*>(VertexBuffer)->GetAllowCPUAccess())) return MakeUnique<FGLTFBufferAdapterCPU>(UVData);
 	return MakeUnique<FGLTFBufferAdapterGPU>(VertexBuffer->TexCoordVertexBuffer.VertexBufferRHI);
 }
@@ -87,7 +87,7 @@ TUniquePtr<IGLTFBufferAdapter> IGLTFBufferAdapter::GetUVs(const FStaticMeshVerte
 TUniquePtr<IGLTFBufferAdapter> IGLTFBufferAdapter::GetInfluences(const FSkinWeightVertexBuffer* VertexBuffer)
 {
 	const auto* InfluenceBuffer = VertexBuffer->GetDataVertexBuffer();
-	const void* InfluenceData = InfluenceBuffer->GetWeightData();
+	const void* InfluenceData = FGLTFBufferUtility::GetBufferData(InfluenceBuffer);
 	if (InfluenceData != nullptr && (WITH_EDITOR || VertexBuffer->GetNeedsCPUAccess())) return MakeUnique<FGLTFBufferAdapterCPU>(InfluenceData);
 	return MakeUnique<FGLTFBufferAdapterGPU>(InfluenceBuffer->VertexBufferRHI);
 }
diff --git a/Engine/Plugins/Enterprise/GLTFExporter/MigratedFromGit/Source/GLTFExporter/Private/Converters/GLTFBufferUtility.cpp b/Engine/Plugins/Enterprise/GLTFExporter/MigratedFromGit/Source/GLTFExporter/Private/Converters/GLTFBufferUtility.cpp
--- a/Engine/Plugins/Enterprise/GLTFExporter/MigratedFromGit/Source/GLTFExporter/Private/Converters/GLTFBufferUtility.cpp
+++ b/Engine/Plugins/Enterprise/GLTFExporter/MigratedFromGit/Source/GLTFExporter/Private/Converters/GLTFBufferUtility.cpp
@@ -51,6 +51,30 @@ bool FGLTFBufferUtility::GetAllowCPUAccess(const FColorVertexBuffer* VertexBuffe
 	return VertexData != nullptr && VertexData->GetAllowCPUAccess();
 }
 
+const void* FGLTFBufferUtility::GetBufferData(const FSkinWeightDataVertexBuffer* VertexBuffer)
+{
+#if (ENGINE_MAJOR_VERSION > 4 || ENGINE_MINOR_VERSION >= 26)
+	return VertexBuffer->GetWeightData();
+
+#else
+	struct FSkinWeightDataVertexBufferHack : FVertexBuffer
+	{
+		FShaderResourceViewRHIRef SRVValue;
+		bool bNeedsCPUAccess;
+		bool bVariableBonesPerVertex;
+		uint32 MaxBoneInfluences;
+		bool bUse16BitBoneIndex;
+		FStaticMeshVertexDataInterface* WeightData;
+		uint8* Data;
+		uint32 NumVertices;
+		uint32 NumBones;
+	};
+
+	static_assert(sizeof(FSkinWeightDataVertexBufferHack) == sizeof(FSkinWeightDataVertexBuffer), "FSkinWeightDataVertexBufferHack memory layout doesn't match FSkinWeightDataVertexBuffer");
+	return reinterpret_cast<const FSkinWeightDataVertexBufferHack*>(VertexBuffer)->Data;
+#endif
+}
+
 const void* FGLTFBufferUtility::GetBufferData(const FSkinWeightLookupVertexBuffer* VertexBuffer)
 {
 	struct FSkinWeightLookupVertexBufferHack : FVertexBuffer
diff --git a/Engine/Plugins/Enterprise/GLTFExporter/MigratedFromGit/Source/GLTFExporter/Private/Converters/GLTFBufferUtility.h b/Engine/Plugins/Enterprise/GLTFExporter/MigratedFromGit/Source/GLTFExporter/Private/Converters/GLTFBufferUtility.h
--- a/Engine/Plugins/Enterprise/GLTFExporter/MigratedFromGit/Source/GLTFExporter/Private/Converters/GLTFBufferUtility.h
+++ b/Engine/Plugins/Enterprise/GLTFExporter/MigratedFromGit/Source/GLTFExporter/Private/Converters/GLTFBufferUtility.h
@@ -10,6 +10,7 @@ struct FGLTFBufferUtility
 	static bool GetAllowCPUAccess(const FPositionVertexBuffer* VertexBuffer);
 	static bool GetAllowCPUAccess(const FColorVertexBuffer* VertexBuffer);
 
+	static const void* GetBufferData(const FSkinWeightDataVertexBuffer* VertexBuffer);
 	static const void* GetBufferData(const FSkinWeightLookupVertexBuffer* VertexBuffer);
 
 	static void ReadRHIBuffer(FRHIIndexBuffer* SourceBuffer, TArray<uint8>& OutData);
diff --git a/Engine/Plugins/Enterprise/GLTFExporter/MigratedFromGit/Source/GLTFExporter/Private/Converters/GLTFLevelVariantSetsConverters.cpp b/Engine/Plugins/Enterprise/GLTFExporter/MigratedFromGit/Source/GLTFExporter/Private/Converters/GLTFLevelVariantSetsConverters.cpp
--- a/Engine/Plugins/Enterprise/GLTFExporter/MigratedFromGit/Source/GLTFExporter/Private/Converters/GLTFLevelVariantSetsConverters.cpp
+++ b/Engine/Plugins/Enterprise/GLTFExporter/MigratedFromGit/Source/GLTFExporter/Private/Converters/GLTFLevelVariantSetsConverters.cpp
@@ -208,7 +208,7 @@ bool FGLTFLevelVariantSetsConverter::TryParseMaterialPropertyValue(FGLTFJsonVari
 		return false;
 	}
 
-	const TArray<FCapturedPropSegment>& CapturedPropSegments = MaterialProperty->GetCapturedPropSegments();
+	const TArray<FCapturedPropSegment>& CapturedPropSegments = FGLTFVariantUtility::GetCapturedPropSegments(MaterialProperty);
 	const int32 NumPropSegments = CapturedPropSegments.Num();
 
 	if (NumPropSegments < 1)
diff --git a/Engine/Plugins/Enterprise/GLTFExporter/MigratedFromGit/Source/GLTFExporter/Private/Converters/GLTFLightMapConverters.cpp b/Engine/Plugins/Enterprise/GLTFExporter/MigratedFromGit/Source/GLTFExporter/Private/Converters/GLTFLightMapConverters.cpp
--- a/Engine/Plugins/Enterprise/GLTFExporter/MigratedFromGit/Source/GLTFExporter/Private/Converters/GLTFLightMapConverters.cpp
+++ b/Engine/Plugins/Enterprise/GLTFExporter/MigratedFromGit/Source/GLTFExporter/Private/Converters/GLTFLightMapConverters.cpp
@@ -23,7 +23,12 @@ FGLTFJsonLightMapIndex FGLTFLightMapConverter::Convert(const UStaticMeshComponen
 
 	const int32 LODIndex = FGLTFMeshUtility::GetLOD(StaticMesh, StaticMeshComponent, Builder.ExportOptions->DefaultLevelOfDetail);
 	const FStaticMeshLODResources& LODResources = StaticMesh->GetLODForExport(LODIndex);
+
+#if (ENGINE_MAJOR_VERSION > 4 || ENGINE_MINOR_VERSION >= 27)
 	const int32 CoordinateIndex = StaticMesh->GetLightMapCoordinateIndex();
+#else
+	const int32 CoordinateIndex = StaticMesh->LightMapCoordinateIndex;
+#endif
 
 	if (CoordinateIndex < 0 || CoordinateIndex >= LODResources.GetNumTexCoords())
 	{
diff --git a/Engine/Plugins/Enterprise/GLTFExporter/MigratedFromGit/Source/GLTFExporter/Private/Converters/GLTFMaterialUtility.cpp b/Engine/Plugins/Enterprise/GLTFExporter/MigratedFromGit/Source/GLTFExporter/Private/Converters/GLTFMaterialUtility.cpp
--- a/Engine/Plugins/Enterprise/GLTFExporter/MigratedFromGit/Source/GLTFExporter/Private/Converters/GLTFMaterialUtility.cpp
+++ b/Engine/Plugins/Enterprise/GLTFExporter/MigratedFromGit/Source/GLTFExporter/Private/Converters/GLTFMaterialUtility.cpp
@@ -37,6 +37,7 @@ bool FGLTFMaterialUtility::IsSRGB(const FMaterialPropertyEx& Property)
 	return Property == MP_BaseColor || Property == MP_EmissiveColor || Property == MP_SubsurfaceColor || Property == TEXT("TransmittanceColor");
 }
 
+#if (ENGINE_MAJOR_VERSION > 4 || ENGINE_MINOR_VERSION >= 27)
 FGuid FGLTFMaterialUtility::GetAttributeID(const FMaterialPropertyEx& Property)
 {
 	return Property.IsCustomOutput()
@@ -50,16 +51,63 @@ FGuid FGLTFMaterialUtility::GetAttributeIDChecked(const FMaterialPropertyEx& Pro
 	check(AttributeID != FMaterialAttributeDefinitionMap::GetDefaultID());
 	return AttributeID;
 }
+#endif
 
 FVector4 FGLTFMaterialUtility::GetPropertyDefaultValue(const FMaterialPropertyEx& Property)
 {
+#if (ENGINE_MAJOR_VERSION > 4 || ENGINE_MINOR_VERSION >= 27)
 	return FMaterialAttributeDefinitionMap::GetDefaultValue(GetAttributeIDChecked(Property));
+
+#else
+	switch (Property.Type)
+	{
+		case MP_EmissiveColor:          return FVector4(0,0,0,0);
+		case MP_Opacity:                return FVector4(1,0,0,0);
+		case MP_OpacityMask:            return FVector4(1,0,0,0);
+		case MP_BaseColor:              return FVector4(0,0,0,0);
+		case MP_Metallic:               return FVector4(0,0,0,0);
+		case MP_Specular:               return FVector4(.5,0,0,0);
+		case MP_Roughness:              return FVector4(.5,0,0,0);
+		case MP_Anisotropy:             return FVector4(0,0,0,0);
+		case MP_Normal:                 return FVector4(0,0,1,0);
+		case MP_Tangent:                return FVector4(1,0,0,0);
+		case MP_WorldPositionOffset:    return FVector4(0,0,0,0);
+		case MP_WorldDisplacement:      return FVector4(0,0,0,0);
+		case MP_TessellationMultiplier: return FVector4(1,0,0,0);
+		case MP_SubsurfaceColor:        return FVector4(1,1,1,0);
+		case MP_CustomData0:            return FVector4(1,0,0,0);
+		case MP_CustomData1:            return FVector4(.1,0,0,0);
+		case MP_AmbientOcclusion:       return FVector4(1,0,0,0);
+		case MP_Refraction:             return FVector4(1,0,0,0);
+		case MP_PixelDepthOffset:       return FVector4(0,0,0,0);
+		case MP_ShadingModel:           return FVector4(0,0,0,0);
+		default:                        break;
+	}
+
+	if (Property.Type >= MP_CustomizedUVs0 && Property.Type <= MP_CustomizedUVs7)
+	{
+		return FVector4(0,0,0,0);
+	}
+
+	if (Property == TEXT("ClearCoatBottomNormal"))
+	{
+		return FVector4(0,0,1,0);
+	}
+
+	if (Property == TEXT("TransmittanceColor"))
+	{
+		return FVector4(.5,.5,.5,0);
+	}
+
+	check(false);
+	return FVector4();
+#endif
 }
 
 FVector4 FGLTFMaterialUtility::GetPropertyMask(const FMaterialPropertyEx& Property)
 {
-	const EMaterialValueType ValueType = FMaterialAttributeDefinitionMap::GetValueType(GetAttributeIDChecked(Property));
-	switch (ValueType)
+#if (ENGINE_MAJOR_VERSION > 4 || ENGINE_MINOR_VERSION >= 27)
+	switch (FMaterialAttributeDefinitionMap::GetValueType(GetAttributeIDChecked(Property)))
 	{
 		case MCT_Float:
 		case MCT_Float1: return FVector4(1, 0, 0, 0);
@@ -70,6 +118,51 @@ FVector4 FGLTFMaterialUtility::GetPropertyMask(const FMaterialPropertyEx& Proper
 			checkNoEntry();
 			return FVector4();
 	}
+
+#else
+	switch (Property.Type)
+	{
+		case MP_EmissiveColor:          return FVector4(1,1,1,0);
+		case MP_Opacity:                return FVector4(1,0,0,0);
+		case MP_OpacityMask:            return FVector4(1,0,0,0);
+		case MP_BaseColor:              return FVector4(1,1,1,0);
+		case MP_Metallic:               return FVector4(1,0,0,0);
+		case MP_Specular:               return FVector4(1,0,0,0);
+		case MP_Roughness:              return FVector4(1,0,0,0);
+		case MP_Anisotropy:             return FVector4(1,0,0,0);
+		case MP_Normal:                 return FVector4(1,1,1,0);
+		case MP_Tangent:                return FVector4(1,1,1,0);
+		case MP_WorldPositionOffset:    return FVector4(1,1,1,0);
+		case MP_WorldDisplacement:      return FVector4(1,1,1,0);
+		case MP_TessellationMultiplier: return FVector4(1,0,0,0);
+		case MP_SubsurfaceColor:        return FVector4(1,1,1,0);
+		case MP_CustomData0:            return FVector4(1,0,0,0);
+		case MP_CustomData1:            return FVector4(1,0,0,0);
+		case MP_AmbientOcclusion:       return FVector4(1,0,0,0);
+		case MP_Refraction:             return FVector4(1,1,0,0);
+		case MP_PixelDepthOffset:       return FVector4(1,0,0,0);
+		case MP_ShadingModel:           return FVector4(1,0,0,0);
+		default:                        break;
+	}
+
+	if (Property.Type >= MP_CustomizedUVs0 && Property.Type <= MP_CustomizedUVs7)
+	{
+		return FVector4(1,1,0,0);
+	}
+
+	if (Property == TEXT("ClearCoatBottomNormal"))
+	{
+		return FVector4(1,1,1,0);
+	}
+
+	if (Property == TEXT("TransmittanceColor"))
+	{
+		return FVector4(1,1,1,0);
+	}
+
+	check(false);
+	return FVector4();
+#endif
 }
 
 const FExpressionInput* FGLTFMaterialUtility::GetInputForProperty(const UMaterialInterface* Material, const FMaterialPropertyEx& Property)
diff --git a/Engine/Plugins/Enterprise/GLTFExporter/MigratedFromGit/Source/GLTFExporter/Private/Converters/GLTFMaterialUtility.h b/Engine/Plugins/Enterprise/GLTFExporter/MigratedFromGit/Source/GLTFExporter/Private/Converters/GLTFMaterialUtility.h
--- a/Engine/Plugins/Enterprise/GLTFExporter/MigratedFromGit/Source/GLTFExporter/Private/Converters/GLTFMaterialUtility.h
+++ b/Engine/Plugins/Enterprise/GLTFExporter/MigratedFromGit/Source/GLTFExporter/Private/Converters/GLTFMaterialUtility.h
@@ -43,8 +43,10 @@ struct FGLTFMaterialUtility
 	static bool IsNormalMap(const FMaterialPropertyEx& Property);
 	static bool IsSRGB(const FMaterialPropertyEx& Property);
 
+#if (ENGINE_MAJOR_VERSION > 4 || ENGINE_MINOR_VERSION >= 27)
 	static FGuid GetAttributeID(const FMaterialPropertyEx& Property);
 	static FGuid GetAttributeIDChecked(const FMaterialPropertyEx& Property);
+#endif
 
 	static FVector4 GetPropertyDefaultValue(const FMaterialPropertyEx& Property);
 	static FVector4 GetPropertyMask(const FMaterialPropertyEx& Property);
diff --git a/Engine/Plugins/Enterprise/GLTFExporter/MigratedFromGit/Source/GLTFExporter/Private/Converters/GLTFMeshConverters.cpp b/Engine/Plugins/Enterprise/GLTFExporter/MigratedFromGit/Source/GLTFExporter/Private/Converters/GLTFMeshConverters.cpp
--- a/Engine/Plugins/Enterprise/GLTFExporter/MigratedFromGit/Source/GLTFExporter/Private/Converters/GLTFMeshConverters.cpp
+++ b/Engine/Plugins/Enterprise/GLTFExporter/MigratedFromGit/Source/GLTFExporter/Private/Converters/GLTFMeshConverters.cpp
@@ -18,7 +18,11 @@ void FGLTFStaticMeshConverter::Sanitize(const UStaticMesh*& StaticMesh, const US
 	}
 	else
 	{
+#if (ENGINE_MAJOR_VERSION > 4 || ENGINE_MINOR_VERSION >= 27)
 		FGLTFMaterialUtility::ResolveOverrides(Materials, StaticMesh->GetStaticMaterials());
+#else
+		FGLTFMaterialUtility::ResolveOverrides(Materials, StaticMesh->StaticMaterials);
+#endif
 	}
 
 	if (LODIndex < 0)
@@ -66,7 +70,11 @@ void FGLTFSkeletalMeshConverter::Sanitize(const USkeletalMesh*& SkeletalMesh, co
 	}
 	else
 	{
+#if (ENGINE_MAJOR_VERSION > 4 || ENGINE_MINOR_VERSION >= 27)
 		FGLTFMaterialUtility::ResolveOverrides(Materials, SkeletalMesh->GetMaterials());
+#else
+		FGLTFMaterialUtility::ResolveOverrides(Materials, SkeletalMesh->Materials);
+#endif
 	}
 
 	if (LODIndex < 0)
diff --git a/Engine/Plugins/Enterprise/GLTFExporter/MigratedFromGit/Source/GLTFExporter/Private/Converters/GLTFMeshData.cpp b/Engine/Plugins/Enterprise/GLTFExporter/MigratedFromGit/Source/GLTFExporter/Private/Converters/GLTFMeshData.cpp
--- a/Engine/Plugins/Enterprise/GLTFExporter/MigratedFromGit/Source/GLTFExporter/Private/Converters/GLTFMeshData.cpp
+++ b/Engine/Plugins/Enterprise/GLTFExporter/MigratedFromGit/Source/GLTFExporter/Private/Converters/GLTFMeshData.cpp
@@ -53,7 +53,13 @@ FGLTFMeshData::FGLTFMeshData(const UStaticMesh* StaticMesh, const UStaticMeshCom
 	}
 
 #if WITH_EDITOR
+
+#if (ENGINE_MAJOR_VERSION > 4 || ENGINE_MINOR_VERSION >= 27)
 	LightMapTexCoord = StaticMesh->GetLightMapCoordinateIndex();
+#else
+	LightMapTexCoord = StaticMesh->LightMapCoordinateIndex;
+#endif
+
  	const int32 NumTexCoords = StaticMesh->GetLODForExport(LODIndex).VertexBuffers.StaticMeshVertexBuffer.GetNumTexCoords();
 	BakeUsingTexCoord = FMath::Min(LightMapTexCoord, NumTexCoords - 1);
 	// TODO: add warning if texture coordinate for baking has overlap
diff --git a/Engine/Plugins/Enterprise/GLTFExporter/MigratedFromGit/Source/GLTFExporter/Private/Converters/GLTFMeshUtility.cpp b/Engine/Plugins/Enterprise/GLTFExporter/MigratedFromGit/Source/GLTFExporter/Private/Converters/GLTFMeshUtility.cpp
--- a/Engine/Plugins/Enterprise/GLTFExporter/MigratedFromGit/Source/GLTFExporter/Private/Converters/GLTFMeshUtility.cpp
+++ b/Engine/Plugins/Enterprise/GLTFExporter/MigratedFromGit/Source/GLTFExporter/Private/Converters/GLTFMeshUtility.cpp
@@ -81,7 +81,11 @@ int32 FGLTFMeshUtility::GetMinimumLOD(const UStaticMesh* StaticMesh, const UStat
 
 	if (StaticMesh != nullptr)
 	{
+#if (ENGINE_MAJOR_VERSION > 4 || ENGINE_MINOR_VERSION >= 27)
 		return StaticMesh->GetMinLOD().Default;
+#else
+		return StaticMesh->MinLOD.Default;
+#endif
 	}
 
 	return -1;
@@ -96,7 +100,11 @@ int32 FGLTFMeshUtility::GetMinimumLOD(const USkeletalMesh* SkeletalMesh, const U
 
 	if (SkeletalMesh != nullptr)
 	{
+#if (ENGINE_MAJOR_VERSION > 4 || ENGINE_MINOR_VERSION >= 27)
 		return SkeletalMesh->GetMinLod().Default;
+#else
+		return SkeletalMesh->MinLod.Default;
+#endif
 	}
 
 	return -1;
diff --git a/Engine/Plugins/Enterprise/GLTFExporter/MigratedFromGit/Source/GLTFExporter/Private/Converters/GLTFNodeConverters.cpp b/Engine/Plugins/Enterprise/GLTFExporter/MigratedFromGit/Source/GLTFExporter/Private/Converters/GLTFNodeConverters.cpp
--- a/Engine/Plugins/Enterprise/GLTFExporter/MigratedFromGit/Source/GLTFExporter/Private/Converters/GLTFNodeConverters.cpp
+++ b/Engine/Plugins/Enterprise/GLTFExporter/MigratedFromGit/Source/GLTFExporter/Private/Converters/GLTFNodeConverters.cpp
@@ -244,6 +244,12 @@ FGLTFJsonNodeIndex FGLTFStaticSocketConverter::Convert(FGLTFJsonNodeIndex RootNo
 
 FGLTFJsonNodeIndex FGLTFSkeletalSocketConverter::Convert(FGLTFJsonNodeIndex RootNode, const USkeletalMesh* SkeletalMesh, FName SocketName)
 {
+#if (ENGINE_MAJOR_VERSION > 4 || ENGINE_MINOR_VERSION >= 27)
+	const FReferenceSkeleton& RefSkeleton = SkeletalMesh->GetRefSkeleton();
+#else
+	const FReferenceSkeleton& RefSkeleton = SkeletalMesh->RefSkeleton;
+#endif
+
 	const USkeletalMeshSocket* Socket = SkeletalMesh->FindSocket(SocketName);
 	if (Socket != nullptr)
 	{
@@ -255,12 +261,12 @@ FGLTFJsonNodeIndex FGLTFSkeletalSocketConverter::Convert(FGLTFJsonNodeIndex Root
 		Node.Rotation = FGLTFConverterUtility::ConvertRotation(Socket->RelativeRotation.Quaternion());
 		Node.Scale = FGLTFConverterUtility::ConvertScale(Socket->RelativeScale);
 
-		const int32 ParentBone = SkeletalMesh->GetRefSkeleton().FindBoneIndex(Socket->BoneName);
+		const int32 ParentBone = RefSkeleton.FindBoneIndex(Socket->BoneName);
 		const FGLTFJsonNodeIndex ParentNode = ParentBone != INDEX_NONE ? Builder.GetOrAddNode(RootNode, SkeletalMesh, ParentBone) : RootNode;
 		return Builder.AddChildNode(ParentNode, Node);
 	}
 
-	const int32 BoneIndex = SkeletalMesh->GetRefSkeleton().FindBoneIndex(SocketName);
+	const int32 BoneIndex = RefSkeleton.FindBoneIndex(SocketName);
 	if (BoneIndex != INDEX_NONE)
 	{
 		return Builder.GetOrAddNode(RootNode, SkeletalMesh, BoneIndex);
@@ -272,9 +278,15 @@ FGLTFJsonNodeIndex FGLTFSkeletalSocketConverter::Convert(FGLTFJsonNodeIndex Root
 
 FGLTFJsonNodeIndex FGLTFSkeletalBoneConverter::Convert(FGLTFJsonNodeIndex RootNode, const USkeletalMesh* SkeletalMesh, int32 BoneIndex)
 {
+#if (ENGINE_MAJOR_VERSION > 4 || ENGINE_MINOR_VERSION >= 27)
+	const FReferenceSkeleton& RefSkeleton = SkeletalMesh->GetRefSkeleton();
+#else
+	const FReferenceSkeleton& RefSkeleton = SkeletalMesh->RefSkeleton;
+#endif
+
 	// TODO: add support for MasterPoseComponent?
 
-	const TArray<FMeshBoneInfo>& BoneInfos = SkeletalMesh->GetRefSkeleton().GetRefBoneInfo();
+	const TArray<FMeshBoneInfo>& BoneInfos = RefSkeleton.GetRefBoneInfo();
 	if (!BoneInfos.IsValidIndex(BoneIndex))
 	{
 		// TODO: report error
@@ -286,7 +298,7 @@ FGLTFJsonNodeIndex FGLTFSkeletalBoneConverter::Convert(FGLTFJsonNodeIndex RootNo
 	FGLTFJsonNode Node;
 	Node.Name = BoneInfo.Name.ToString();
 
-	const TArray<FTransform>& BonePoses = SkeletalMesh->GetRefSkeleton().GetRefBonePose();
+	const TArray<FTransform>& BonePoses = RefSkeleton.GetRefBonePose();
 	if (BonePoses.IsValidIndex(BoneIndex))
 	{
 		// TODO: add warning check for non-uniform scaling
diff --git a/Engine/Plugins/Enterprise/GLTFExporter/MigratedFromGit/Source/GLTFExporter/Private/Converters/GLTFSkinConverters.cpp b/Engine/Plugins/Enterprise/GLTFExporter/MigratedFromGit/Source/GLTFExporter/Private/Converters/GLTFSkinConverters.cpp
--- a/Engine/Plugins/Enterprise/GLTFExporter/MigratedFromGit/Source/GLTFExporter/Private/Converters/GLTFSkinConverters.cpp
+++ b/Engine/Plugins/Enterprise/GLTFExporter/MigratedFromGit/Source/GLTFExporter/Private/Converters/GLTFSkinConverters.cpp
@@ -7,11 +7,19 @@
 
 FGLTFJsonSkinIndex FGLTFSkinConverter::Convert(FGLTFJsonNodeIndex RootNode, const USkeletalMesh* SkeletalMesh)
 {
+#if (ENGINE_MAJOR_VERSION > 4 || ENGINE_MINOR_VERSION >= 27)
+	const USkeleton* Skeleton = SkeletalMesh->GetSkeleton();
+	const FReferenceSkeleton& RefSkeleton = SkeletalMesh->GetRefSkeleton();
+#else
+	const USkeleton* Skeleton = SkeletalMesh->Skeleton;
+	const FReferenceSkeleton& RefSkeleton = SkeletalMesh->RefSkeleton;
+#endif
+
 	FGLTFJsonSkin Skin;
-	Skin.Name = SkeletalMesh->GetSkeleton() != nullptr ? SkeletalMesh->GetSkeleton()->GetName() : SkeletalMesh->GetName();
+	Skin.Name = Skeleton != nullptr ? Skeleton->GetName() : SkeletalMesh->GetName();
 	Skin.Skeleton = RootNode;
 
-	const int32 BoneCount = SkeletalMesh->GetRefSkeleton().GetNum();
+	const int32 BoneCount = RefSkeleton.GetNum();
 	if (BoneCount == 0)
 	{
 		// TODO: report warning
@@ -30,7 +38,7 @@ FGLTFJsonSkinIndex FGLTFSkinConverter::Convert(FGLTFJsonNodeIndex RootNode, cons
 
 	for (int32 BoneIndex = 0; BoneIndex < BoneCount; ++BoneIndex)
 	{
-		const FTransform InverseBindTransform = FGLTFBoneUtility::GetBindTransform(SkeletalMesh->GetRefSkeleton(), BoneIndex).Inverse();
+		const FTransform InverseBindTransform = FGLTFBoneUtility::GetBindTransform(RefSkeleton, BoneIndex).Inverse();
 		InverseBindMatrices[BoneIndex] = FGLTFConverterUtility::ConvertTransform(InverseBindTransform, Builder.ExportOptions->ExportUniformScale);
 	}
 
diff --git a/Engine/Plugins/Enterprise/GLTFExporter/MigratedFromGit/Source/GLTFExporter/Private/Converters/GLTFVariantUtility.cpp b/Engine/Plugins/Enterprise/GLTFExporter/MigratedFromGit/Source/GLTFExporter/Private/Converters/GLTFVariantUtility.cpp
--- a/Engine/Plugins/Enterprise/GLTFExporter/MigratedFromGit/Source/GLTFExporter/Private/Converters/GLTFVariantUtility.cpp
+++ b/Engine/Plugins/Enterprise/GLTFExporter/MigratedFromGit/Source/GLTFExporter/Private/Converters/GLTFVariantUtility.cpp
@@ -7,6 +7,25 @@
 #include "VariantObjectBinding.h"
 #include "PropertyValue.h"
 
+const TArray<FCapturedPropSegment>& FGLTFVariantUtility::GetCapturedPropSegments(const UPropertyValue* Property)
+{
+#if (ENGINE_MAJOR_VERSION > 4 || ENGINE_MINOR_VERSION >= 26)
+	return Property->GetCapturedPropSegments();
+
+#else
+	class UPropertyValueHack : public UPropertyValue
+	{
+	public:
+		const TArray<FCapturedPropSegment>& GetCapturedPropSegments() const
+		{
+			return CapturedPropSegments;
+		}
+	};
+
+	return static_cast<const UPropertyValueHack*>(Property)->GetCapturedPropSegments();
+#endif
+}
+
 bool FGLTFVariantUtility::TryGetPropertyValue(UPropertyValue* Property, void* OutData, uint32 OutSize)
 {
 	if (Property == nullptr || !Property->HasRecordedData())
diff --git a/Engine/Plugins/Enterprise/GLTFExporter/MigratedFromGit/Source/GLTFExporter/Private/Converters/GLTFVariantUtility.h b/Engine/Plugins/Enterprise/GLTFExporter/MigratedFromGit/Source/GLTFExporter/Private/Converters/GLTFVariantUtility.h
--- a/Engine/Plugins/Enterprise/GLTFExporter/MigratedFromGit/Source/GLTFExporter/Private/Converters/GLTFVariantUtility.h
+++ b/Engine/Plugins/Enterprise/GLTFExporter/MigratedFromGit/Source/GLTFExporter/Private/Converters/GLTFVariantUtility.h
@@ -11,6 +11,8 @@ class UPropertyValue;
 
 struct FGLTFVariantUtility
 {
+	static const TArray<FCapturedPropSegment>& GetCapturedPropSegments(const UPropertyValue* Property);
+
 	template<typename T>
 	static bool TryGetPropertyValue(UPropertyValue* Property, T& OutValue)
 	{
diff --git a/Engine/Plugins/Enterprise/GLTFExporter/MigratedFromGit/Source/GLTFExporter/Private/Exporters/GLTFExporterUtility.cpp b/Engine/Plugins/Enterprise/GLTFExporter/MigratedFromGit/Source/GLTFExporter/Private/Exporters/GLTFExporterUtility.cpp
--- a/Engine/Plugins/Enterprise/GLTFExporter/MigratedFromGit/Source/GLTFExporter/Private/Exporters/GLTFExporterUtility.cpp
+++ b/Engine/Plugins/Enterprise/GLTFExporter/MigratedFromGit/Source/GLTFExporter/Private/Exporters/GLTFExporterUtility.cpp
@@ -2,7 +2,7 @@
 
 #include "Exporters/GLTFExporterUtility.h"
 #include "Materials/MaterialInstance.h"
-#include "AssetRegistry/IAssetRegistry.h"
+#include "AssetRegistryModule.h"
 
 const UStaticMesh* FGLTFExporterUtility::GetPreviewMesh(const UMaterialInterface* Material)
 {
@@ -45,12 +45,18 @@ const USkeletalMesh* FGLTFExporterUtility::GetPreviewMesh(const UAnimSequence* A
 
 const USkeletalMesh* FGLTFExporterUtility::FindCompatibleMesh(const USkeleton *Skeleton)
 {
+#if (ENGINE_MAJOR_VERSION > 4 || ENGINE_MINOR_VERSION >= 27)
+	const FName SkeletonMemberName = USkeletalMesh::GetSkeletonMemberName();
+#else
+	const FName SkeletonMemberName = GET_MEMBER_NAME_CHECKED(USkeletalMesh, Skeleton);
+#endif
+
 	FARFilter Filter;
 	Filter.ClassNames.Add(USkeletalMesh::StaticClass()->GetFName());
-	Filter.TagsAndValues.Add(USkeletalMesh::GetSkeletonMemberName(), FAssetData(Skeleton).GetExportTextName());
+	Filter.TagsAndValues.Add(SkeletonMemberName, FAssetData(Skeleton).GetExportTextName());
 
 	TArray<FAssetData> FilteredAssets;
-	IAssetRegistry::GetChecked().GetAssets(Filter, FilteredAssets);
+	FAssetRegistryModule::GetRegistry().GetAssets(Filter, FilteredAssets);
 
 	for (const FAssetData& Asset : FilteredAssets)
 	{
@@ -70,7 +76,7 @@ TArray<UWorld*> FGLTFExporterUtility::GetAssociatedWorlds(const UObject* Object)
 	TArray<FAssetIdentifier> Dependencies;
 
 	const FName OuterPathName = *Object->GetOutermost()->GetPathName();
-	IAssetRegistry::GetChecked().GetDependencies(OuterPathName, Dependencies);
+	FAssetRegistryModule::GetRegistry().GetDependencies(OuterPathName, Dependencies);
 
 	for (FAssetIdentifier& Dependency : Dependencies)
 	{
diff --git a/Engine/Plugins/Enterprise/GLTFExporter/MigratedFromGit/Source/GLTFExporter/Private/Tasks/GLTFAnimationTasks.cpp b/Engine/Plugins/Enterprise/GLTFExporter/MigratedFromGit/Source/GLTFExporter/Private/Tasks/GLTFAnimationTasks.cpp
--- a/Engine/Plugins/Enterprise/GLTFExporter/MigratedFromGit/Source/GLTFExporter/Private/Tasks/GLTFAnimationTasks.cpp
+++ b/Engine/Plugins/Enterprise/GLTFExporter/MigratedFromGit/Source/GLTFExporter/Private/Tasks/GLTFAnimationTasks.cpp
@@ -7,12 +7,12 @@
 #include "LevelSequence.h"
 #include "LevelSequencePlayer.h"
 #include "MovieSceneTimeHelpers.h"
-#include "Animation/AnimationPoseData.h"
 #include "Tracks/MovieScene3DTransformTrack.h"
 #include "Sections/MovieScene3DTransformSection.h"
 #include "Channels/MovieSceneChannelProxy.h"
 
 #if (ENGINE_MAJOR_VERSION > 4 || ENGINE_MINOR_VERSION >= 26)
+#include "Animation/AnimationPoseData.h"
 using namespace UE;
 #endif
 
@@ -52,19 +52,28 @@ void FGLTFAnimSequenceTask::Complete()
 		FBlendedCurve Curve;
 		Curve.InitFrom(BoneContainer);
 
+#if (ENGINE_MAJOR_VERSION > 4 || ENGINE_MINOR_VERSION >= 26)
 		FStackCustomAttributes Attributes;
 		FAnimationPoseData PoseData(Pose, Curve, Attributes);
+#endif
 
 		for (int32 Frame = 0; Frame < FrameCount; ++Frame)
 		{
 			const FAnimExtractContext ExtractionContext(Timestamps[Frame]); // TODO: set bExtractRootMotion?
+
+#if (ENGINE_MAJOR_VERSION > 4 || ENGINE_MINOR_VERSION >= 26)
 			AnimSequence->GetBonePose(PoseData, ExtractionContext);
+#else
+			AnimSequence->GetBonePose(Pose, Curve, ExtractionContext);
+#endif
 			Pose.CopyBonesTo(FrameTransforms[Frame]);
 		}
 
 		Pose.Empty();
 		Curve.Empty();
+#if (ENGINE_MAJOR_VERSION > 4 || ENGINE_MINOR_VERSION >= 26)
 		Attributes.Empty();
+#endif
 	}
 
 	// TODO: add animation data accessor converters to reuse track information
diff --git a/Engine/Plugins/Enterprise/GLTFExporter/MigratedFromGit/Source/GLTFExporter/Private/Tasks/GLTFMaterialTasks.cpp b/Engine/Plugins/Enterprise/GLTFExporter/MigratedFromGit/Source/GLTFExporter/Private/Tasks/GLTFMaterialTasks.cpp
--- a/Engine/Plugins/Enterprise/GLTFExporter/MigratedFromGit/Source/GLTFExporter/Private/Tasks/GLTFMaterialTasks.cpp
+++ b/Engine/Plugins/Enterprise/GLTFExporter/MigratedFromGit/Source/GLTFExporter/Private/Tasks/GLTFMaterialTasks.cpp
@@ -8,6 +8,7 @@
 #include "Converters/GLTFMaterialUtility.h"
 #include "Builders/GLTFContainerBuilder.h"
 #include "MaterialPropertyEx.h"
+#include "Materials/MaterialInstance.h"
 #include "Materials/MaterialExpressionConstant.h"
 #include "Materials/MaterialExpressionConstant2Vector.h"
 #include "Materials/MaterialExpressionConstant3Vector.h"
diff --git a/Engine/Plugins/Enterprise/GLTFExporter/MigratedFromGit/Source/GLTFExporter/Private/Tasks/GLTFMeshTasks.cpp b/Engine/Plugins/Enterprise/GLTFExporter/MigratedFromGit/Source/GLTFExporter/Private/Tasks/GLTFMeshTasks.cpp
--- a/Engine/Plugins/Enterprise/GLTFExporter/MigratedFromGit/Source/GLTFExporter/Private/Tasks/GLTFMeshTasks.cpp
+++ b/Engine/Plugins/Enterprise/GLTFExporter/MigratedFromGit/Source/GLTFExporter/Private/Tasks/GLTFMeshTasks.cpp
@@ -155,7 +155,12 @@ void FGLTFStaticMeshTask::Complete()
 
 	ValidateVertexBuffer(Builder, &VertexBuffer, *StaticMesh->GetName());
 
+#if (ENGINE_MAJOR_VERSION > 4 || ENGINE_MINOR_VERSION >= 27)
 	const TArray<FStaticMaterial>& MaterialSlots = StaticMesh->GetStaticMaterials();
+#else
+	const TArray<FStaticMaterial>& MaterialSlots = StaticMesh->StaticMaterials;
+#endif
+
 	const int32 MaterialCount = MaterialSlots.Num();
 	JsonMesh.Primitives.AddDefaulted(MaterialCount);
 
@@ -260,7 +265,12 @@ void FGLTFSkeletalMeshTask::Complete()
 
 	ValidateVertexBuffer(Builder, &VertexBuffer, *SkeletalMesh->GetName());
 
+#if (ENGINE_MAJOR_VERSION > 4 || ENGINE_MINOR_VERSION >= 27)
 	const TArray<FSkeletalMaterial>& MaterialSlots = SkeletalMesh->GetMaterials();
+#else
+	const TArray<FSkeletalMaterial>& MaterialSlots = SkeletalMesh->Materials;
+#endif
+
 	const uint16 MaterialCount = MaterialSlots.Num();
 	JsonMesh.Primitives.AddDefaulted(MaterialCount);
 
