[GLTFExporter] refactor: add GLTFMeshUtility::GetMaterials to simplify getting mesh materials independent of UE version

commit 93a1285ac273c23d3519767456d218110887129c
Date: Sat, 1 Jan 2022 21:39:39 +0100

    refactor: add GLTFMeshUtility::GetMaterials to simplify getting mesh materials independent of UE version
    (cherry picked from commit 4d4df058eb09311a1c3b9596b06052b3ad4cd8dc​)

#skipci

diff --git a/Engine/Plugins/Enterprise/GLTFExporter/MigratedFromGit/Source/GLTFExporter/Private/Converters/GLTFMeshConverters.cpp b/Engine/Plugins/Enterprise/GLTFExporter/MigratedFromGit/Source/GLTFExporter/Private/Converters/GLTFMeshConverters.cpp
--- a/Engine/Plugins/Enterprise/GLTFExporter/MigratedFromGit/Source/GLTFExporter/Private/Converters/GLTFMeshConverters.cpp
+++ b/Engine/Plugins/Enterprise/GLTFExporter/MigratedFromGit/Source/GLTFExporter/Private/Converters/GLTFMeshConverters.cpp
@@ -18,11 +18,7 @@ void FGLTFStaticMeshConverter::Sanitize(const UStaticMesh*& StaticMesh, const US
 	}
 	else
 	{
-#if (ENGINE_MAJOR_VERSION > 4 || ENGINE_MINOR_VERSION >= 27)
-		FGLTFMaterialUtility::ResolveOverrides(Materials, StaticMesh->GetStaticMaterials());
-#else
-		FGLTFMaterialUtility::ResolveOverrides(Materials, StaticMesh->StaticMaterials);
-#endif
+		FGLTFMaterialUtility::ResolveOverrides(Materials, FGLTFMeshUtility::GetMaterials(StaticMesh));
 	}
 
 	if (LODIndex < 0)
@@ -70,11 +66,7 @@ void FGLTFSkeletalMeshConverter::Sanitize(const USkeletalMesh*& SkeletalMesh, co
 	}
 	else
 	{
-#if (ENGINE_MAJOR_VERSION > 4 || ENGINE_MINOR_VERSION >= 27)
-		FGLTFMaterialUtility::ResolveOverrides(Materials, SkeletalMesh->GetMaterials());
-#else
-		FGLTFMaterialUtility::ResolveOverrides(Materials, SkeletalMesh->Materials);
-#endif
+		FGLTFMaterialUtility::ResolveOverrides(Materials, FGLTFMeshUtility::GetMaterials(SkeletalMesh));
 	}
 
 	if (LODIndex < 0)
diff --git a/Engine/Plugins/Enterprise/GLTFExporter/MigratedFromGit/Source/GLTFExporter/Private/Converters/GLTFMeshUtility.cpp b/Engine/Plugins/Enterprise/GLTFExporter/MigratedFromGit/Source/GLTFExporter/Private/Converters/GLTFMeshUtility.cpp
--- a/Engine/Plugins/Enterprise/GLTFExporter/MigratedFromGit/Source/GLTFExporter/Private/Converters/GLTFMeshUtility.cpp
+++ b/Engine/Plugins/Enterprise/GLTFExporter/MigratedFromGit/Source/GLTFExporter/Private/Converters/GLTFMeshUtility.cpp
@@ -3,6 +3,24 @@
 #include "Converters/GLTFMeshUtility.h"
 #include "Rendering/SkeletalMeshRenderData.h"
 
+const TArray<FStaticMaterial>& FGLTFMeshUtility::GetMaterials(const UStaticMesh* StaticMesh)
+{
+#if (ENGINE_MAJOR_VERSION > 4 || ENGINE_MINOR_VERSION >= 27)
+	return StaticMesh->GetStaticMaterials();
+#else
+	return StaticMesh->StaticMaterials;
+#endif
+}
+
+const TArray<FSkeletalMaterial>& FGLTFMeshUtility::GetMaterials(const USkeletalMesh* SkeletalMesh)
+{
+#if (ENGINE_MAJOR_VERSION > 4 || ENGINE_MINOR_VERSION >= 27)
+	return SkeletalMesh->GetMaterials();
+#else
+	return SkeletalMesh->Materials;
+#endif
+}
+
 FGLTFIndexArray FGLTFMeshUtility::GetSectionIndices(const FStaticMeshLODResources& MeshLOD, int32 MaterialIndex)
 {
 	const FStaticMeshLODResources::FStaticMeshSectionArray& Sections = MeshLOD.Sections;
diff --git a/Engine/Plugins/Enterprise/GLTFExporter/MigratedFromGit/Source/GLTFExporter/Private/Converters/GLTFMeshUtility.h b/Engine/Plugins/Enterprise/GLTFExporter/MigratedFromGit/Source/GLTFExporter/Private/Converters/GLTFMeshUtility.h
--- a/Engine/Plugins/Enterprise/GLTFExporter/MigratedFromGit/Source/GLTFExporter/Private/Converters/GLTFMeshUtility.h
+++ b/Engine/Plugins/Enterprise/GLTFExporter/MigratedFromGit/Source/GLTFExporter/Private/Converters/GLTFMeshUtility.h
@@ -8,6 +8,9 @@ class FSkeletalMeshLODRenderData;
 
 struct FGLTFMeshUtility
 {
+	static const TArray<FStaticMaterial>& GetMaterials(const UStaticMesh* StaticMesh);
+	static const TArray<FSkeletalMaterial>& GetMaterials(const USkeletalMesh* SkeletalMesh);
+
 	static FGLTFIndexArray GetSectionIndices(const FStaticMeshLODResources& MeshLOD, int32 MaterialIndex);
 	static FGLTFIndexArray GetSectionIndices(const FSkeletalMeshLODRenderData& MeshLOD, int32 MaterialIndex);
 
diff --git a/Engine/Plugins/Enterprise/GLTFExporter/MigratedFromGit/Source/GLTFExporter/Private/Tasks/GLTFMeshTasks.cpp b/Engine/Plugins/Enterprise/GLTFExporter/MigratedFromGit/Source/GLTFExporter/Private/Tasks/GLTFMeshTasks.cpp
--- a/Engine/Plugins/Enterprise/GLTFExporter/MigratedFromGit/Source/GLTFExporter/Private/Tasks/GLTFMeshTasks.cpp
+++ b/Engine/Plugins/Enterprise/GLTFExporter/MigratedFromGit/Source/GLTFExporter/Private/Tasks/GLTFMeshTasks.cpp
@@ -155,12 +155,7 @@ void FGLTFStaticMeshTask::Complete()
 
 	ValidateVertexBuffer(Builder, &VertexBuffer, *StaticMesh->GetName());
 
-#if (ENGINE_MAJOR_VERSION > 4 || ENGINE_MINOR_VERSION >= 27)
-	const TArray<FStaticMaterial>& MaterialSlots = StaticMesh->GetStaticMaterials();
-#else
-	const TArray<FStaticMaterial>& MaterialSlots = StaticMesh->StaticMaterials;
-#endif
-
+	const TArray<FStaticMaterial>& MaterialSlots = FGLTFMeshUtility::GetMaterials(StaticMesh);
 	const int32 MaterialCount = MaterialSlots.Num();
 	JsonMesh.Primitives.AddDefaulted(MaterialCount);
 
@@ -265,12 +260,7 @@ void FGLTFSkeletalMeshTask::Complete()
 
 	ValidateVertexBuffer(Builder, &VertexBuffer, *SkeletalMesh->GetName());
 
-#if (ENGINE_MAJOR_VERSION > 4 || ENGINE_MINOR_VERSION >= 27)
-	const TArray<FSkeletalMaterial>& MaterialSlots = SkeletalMesh->GetMaterials();
-#else
-	const TArray<FSkeletalMaterial>& MaterialSlots = SkeletalMesh->Materials;
-#endif
-
+	const TArray<FSkeletalMaterial>& MaterialSlots = FGLTFMeshUtility::GetMaterials(SkeletalMesh);
 	const uint16 MaterialCount = MaterialSlots.Num();
 	JsonMesh.Primitives.AddDefaulted(MaterialCount);
 
