[GLTFExporter] refactor: add "Epic" to the name of everything related to EPIC_LevelVariantSets

commit 863fdda5d362f24dfe643674e5727cce71b85b27
Date: Thu, 16 Sep 2021 08:29:49 +0200

    refactor: add "Epic" to the name of everything related to EPIC_LevelVariantSets
    this is done to avoid confusion when implementing KHR_materials_variants, since both extensions use similar concepts and types
    (cherry picked from commit 224efeceac732db9a2e437dc5e94dad82563b306​)

#skipci

diff --git a/Engine/Plugins/Enterprise/GLTFExporter/MigratedFromGit/Source/GLTFExporter/Private/Builders/GLTFConvertBuilder.cpp b/Engine/Plugins/Enterprise/GLTFExporter/MigratedFromGit/Source/GLTFExporter/Private/Builders/GLTFConvertBuilder.cpp
--- a/Engine/Plugins/Enterprise/GLTFExporter/MigratedFromGit/Source/GLTFExporter/Private/Builders/GLTFConvertBuilder.cpp
+++ b/Engine/Plugins/Enterprise/GLTFExporter/MigratedFromGit/Source/GLTFExporter/Private/Builders/GLTFConvertBuilder.cpp
@@ -401,14 +401,14 @@ FGLTFJsonBackdropIndex FGLTFConvertBuilder::GetOrAddBackdrop(const AActor* Backd
 	return BackdropConverter.GetOrAdd(BackdropActor);
 }
 
-FGLTFJsonLevelVariantSetsIndex FGLTFConvertBuilder::GetOrAddLevelVariantSets(const ULevelVariantSets* LevelVariantSets)
+FGLTFJsonEpicLevelVariantSetsIndex FGLTFConvertBuilder::GetOrAddEpicLevelVariantSets(const ULevelVariantSets* LevelVariantSets)
 {
 	if (LevelVariantSets == nullptr)
 	{
-		return FGLTFJsonLevelVariantSetsIndex(INDEX_NONE);
+		return FGLTFJsonEpicLevelVariantSetsIndex(INDEX_NONE);
 	}
 
-	return LevelVariantSetsConverter.GetOrAdd(LevelVariantSets);
+	return EpicLevelVariantSetsConverter.GetOrAdd(LevelVariantSets);
 }
 
 FGLTFJsonLightMapIndex FGLTFConvertBuilder::GetOrAddLightMap(const UStaticMeshComponent* StaticMeshComponent)
diff --git a/Engine/Plugins/Enterprise/GLTFExporter/MigratedFromGit/Source/GLTFExporter/Private/Builders/GLTFConvertBuilder.h b/Engine/Plugins/Enterprise/GLTFExporter/MigratedFromGit/Source/GLTFExporter/Private/Builders/GLTFConvertBuilder.h
--- a/Engine/Plugins/Enterprise/GLTFExporter/MigratedFromGit/Source/GLTFExporter/Private/Builders/GLTFConvertBuilder.h
+++ b/Engine/Plugins/Enterprise/GLTFExporter/MigratedFromGit/Source/GLTFExporter/Private/Builders/GLTFConvertBuilder.h
@@ -16,7 +16,7 @@
 #include "Converters/GLTFCameraConverters.h"
 #include "Converters/GLTFLightConverters.h"
 #include "Converters/GLTFBackdropConverters.h"
-#include "Converters/GLTFLevelVariantSetsConverters.h"
+#include "Converters/GLTFEpicLevelVariantSetsConverters.h"
 #include "Converters/GLTFLightMapConverters.h"
 #include "Converters/GLTFHotspotConverters.h"
 #include "Converters/GLTFSkySphereConverters.h"
@@ -80,7 +80,7 @@ public:
 	FGLTFJsonCameraIndex GetOrAddCamera(const UCameraComponent* CameraComponent);
 	FGLTFJsonLightIndex GetOrAddLight(const ULightComponent* LightComponent);
 	FGLTFJsonBackdropIndex GetOrAddBackdrop(const AActor* BackdropActor);
-	FGLTFJsonLevelVariantSetsIndex GetOrAddLevelVariantSets(const ULevelVariantSets* LevelVariantSets);
+	FGLTFJsonEpicLevelVariantSetsIndex GetOrAddEpicLevelVariantSets(const ULevelVariantSets* LevelVariantSets);
 	FGLTFJsonLightMapIndex GetOrAddLightMap(const UStaticMeshComponent* StaticMeshComponent);
 	FGLTFJsonHotspotIndex GetOrAddHotspot(const AGLTFHotspotActor* HotspotActor);
 	FGLTFJsonSkySphereIndex GetOrAddSkySphere(const AActor* SkySphereActor);
@@ -124,7 +124,7 @@ private:
 	FGLTFCameraConverter CameraConverter{ *this };
 	FGLTFLightConverter LightConverter{ *this };
 	FGLTFBackdropConverter BackdropConverter{ *this };
-	FGLTFLevelVariantSetsConverter LevelVariantSetsConverter{ *this };
+	FGLTFEpicLevelVariantSetsConverter EpicLevelVariantSetsConverter{ *this };
 	FGLTFLightMapConverter LightMapConverter{ *this };
 	FGLTFHotspotConverter HotspotConverter{ *this };
 	FGLTFSkySphereConverter SkySphereConverter{ *this };
diff --git a/Engine/Plugins/Enterprise/GLTFExporter/MigratedFromGit/Source/GLTFExporter/Private/Builders/GLTFJsonBuilder.cpp b/Engine/Plugins/Enterprise/GLTFExporter/MigratedFromGit/Source/GLTFExporter/Private/Builders/GLTFJsonBuilder.cpp
--- a/Engine/Plugins/Enterprise/GLTFExporter/MigratedFromGit/Source/GLTFExporter/Private/Builders/GLTFJsonBuilder.cpp
+++ b/Engine/Plugins/Enterprise/GLTFExporter/MigratedFromGit/Source/GLTFExporter/Private/Builders/GLTFJsonBuilder.cpp
@@ -129,9 +129,9 @@ FGLTFJsonSkySphereIndex FGLTFJsonBuilder::AddSkySphere(const FGLTFJsonSkySphere&
 	return FGLTFJsonSkySphereIndex(JsonRoot.SkySpheres.Add(MakeUnique<FGLTFJsonSkySphere>(JsonSkySphere)));
 }
 
-FGLTFJsonLevelVariantSetsIndex FGLTFJsonBuilder::AddLevelVariantSets(const FGLTFJsonLevelVariantSets& JsonLevelVariantSets)
+FGLTFJsonEpicLevelVariantSetsIndex FGLTFJsonBuilder::AddEpicLevelVariantSets(const FGLTFJsonEpicLevelVariantSets& JsonEpicLevelVariantSets)
 {
-	return FGLTFJsonLevelVariantSetsIndex(JsonRoot.LevelVariantSets.Add(MakeUnique<FGLTFJsonLevelVariantSets>(JsonLevelVariantSets)));
+	return FGLTFJsonEpicLevelVariantSetsIndex(JsonRoot.EpicLevelVariantSets.Add(MakeUnique<FGLTFJsonEpicLevelVariantSets>(JsonEpicLevelVariantSets)));
 }
 
 FGLTFJsonNodeIndex FGLTFJsonBuilder::AddChildNode(FGLTFJsonNodeIndex ParentIndex, const FGLTFJsonNode& JsonNode)
@@ -248,9 +248,9 @@ FGLTFJsonSkySphere& FGLTFJsonBuilder::GetSkySphere(FGLTFJsonSkySphereIndex SkySp
 	return *JsonRoot.SkySpheres[SkySphereIndex];
 }
 
-FGLTFJsonLevelVariantSets& FGLTFJsonBuilder::GetLevelVariantSets(FGLTFJsonLevelVariantSetsIndex LevelVariantSetsIndex)
+FGLTFJsonEpicLevelVariantSets& FGLTFJsonBuilder::GetEpicLevelVariantSets(FGLTFJsonEpicLevelVariantSetsIndex EpicLevelVariantSetsIndex)
 {
-	return *JsonRoot.LevelVariantSets[LevelVariantSetsIndex];
+	return *JsonRoot.EpicLevelVariantSets[EpicLevelVariantSetsIndex];
 }
 
 FGLTFJsonNodeIndex FGLTFJsonBuilder::GetComponentNodeIndex(FGLTFJsonNodeIndex NodeIndex)
diff --git a/Engine/Plugins/Enterprise/GLTFExporter/MigratedFromGit/Source/GLTFExporter/Private/Builders/GLTFJsonBuilder.h b/Engine/Plugins/Enterprise/GLTFExporter/MigratedFromGit/Source/GLTFExporter/Private/Builders/GLTFJsonBuilder.h
--- a/Engine/Plugins/Enterprise/GLTFExporter/MigratedFromGit/Source/GLTFExporter/Private/Builders/GLTFJsonBuilder.h
+++ b/Engine/Plugins/Enterprise/GLTFExporter/MigratedFromGit/Source/GLTFExporter/Private/Builders/GLTFJsonBuilder.h
@@ -39,7 +39,7 @@ public:
 	FGLTFJsonLightIndex AddLight(const FGLTFJsonLight& JsonLight = {});
 	FGLTFJsonLightMapIndex AddLightMap(const FGLTFJsonLightMap& JsonLightMap = {});
 	FGLTFJsonSkySphereIndex AddSkySphere(const FGLTFJsonSkySphere& JsonSkySphere = {});
-	FGLTFJsonLevelVariantSetsIndex AddLevelVariantSets(const FGLTFJsonLevelVariantSets& JsonLevelVariantSets = {});
+	FGLTFJsonEpicLevelVariantSetsIndex AddEpicLevelVariantSets(const FGLTFJsonEpicLevelVariantSets& JsonEpicLevelVariantSets = {});
 
 	FGLTFJsonNodeIndex AddChildNode(FGLTFJsonNodeIndex ParentNodeIndex, const FGLTFJsonNode& JsonNode = {});
 	FGLTFJsonNodeIndex AddChildComponentNode(FGLTFJsonNodeIndex ParentNodeIndex, const FGLTFJsonNode& JsonNode = {});
@@ -62,7 +62,7 @@ public:
 	FGLTFJsonLight& GetLight(FGLTFJsonLightIndex LightIndex);
 	FGLTFJsonLightMap& GetLightMap(FGLTFJsonLightMapIndex LightMapIndex);
 	FGLTFJsonSkySphere& GetSkySphere(FGLTFJsonSkySphereIndex SkySphereIndex);
-	FGLTFJsonLevelVariantSets& GetLevelVariantSets(FGLTFJsonLevelVariantSetsIndex LevelVariantSetsIndex);
+	FGLTFJsonEpicLevelVariantSets& GetEpicLevelVariantSets(FGLTFJsonEpicLevelVariantSetsIndex EpicLevelVariantSetsIndex);
 
 	FGLTFJsonNodeIndex GetComponentNodeIndex(FGLTFJsonNodeIndex NodeIndex);
 
diff --git a/Engine/Plugins/Enterprise/GLTFExporter/MigratedFromGit/Source/GLTFExporter/Private/Converters/GLTFEpicLevelVariantSetsConverters.cpp b/Engine/Plugins/Enterprise/GLTFExporter/MigratedFromGit/Source/GLTFExporter/Private/Converters/GLTFEpicLevelVariantSetsConverters.cpp
new file mode 100644
--- /dev/null
+++ b/Engine/Plugins/Enterprise/GLTFExporter/MigratedFromGit/Source/GLTFExporter/Private/Converters/GLTFEpicLevelVariantSetsConverters.cpp
@@ -0,0 +1,390 @@
+// Copyright Epic Games, Inc. All Rights Reserved.
+
+#include "Converters/GLTFEpicLevelVariantSetsConverters.h"
+#include "Converters/GLTFVariantUtility.h"
+#include "Converters/GLTFMeshUtility.h"
+#include "Builders/GLTFContainerBuilder.h"
+#include "Rendering/SkeletalMeshRenderData.h"
+#include "LevelVariantSets.h"
+#include "VariantSet.h"
+#include "Variant.h"
+#include "VariantObjectBinding.h"
+#include "PropertyValue.h"
+#include "PropertyValueMaterial.h"
+#include "Engine/StaticMesh.h"
+#include "Engine/SkeletalMesh.h"
+#include "Components/StaticMeshComponent.h"
+#include "Components/SkeletalMeshComponent.h"
+
+FGLTFJsonEpicLevelVariantSetsIndex FGLTFEpicLevelVariantSetsConverter::Convert(const ULevelVariantSets* LevelVariantSets)
+{
+	FGLTFJsonEpicLevelVariantSets JsonLevelVariantSets;
+	LevelVariantSets->GetName(JsonLevelVariantSets.Name);
+
+	for (const UVariantSet* VariantSet: LevelVariantSets->GetVariantSets())
+	{
+		FGLTFJsonEpicVariantSet JsonVariantSet;
+		JsonVariantSet.Name = VariantSet->GetDisplayText().ToString();
+
+		for (const UVariant* Variant: VariantSet->GetVariants())
+		{
+			FGLTFJsonEpicVariant JsonVariant;
+			if (TryParseVariant(JsonVariant, Variant))
+			{
+				JsonVariantSet.Variants.Add(JsonVariant);
+			}
+		}
+
+		if (JsonVariantSet.Variants.Num() > 0)
+		{
+			JsonLevelVariantSets.VariantSets.Add(JsonVariantSet);
+		}
+	}
+
+	if (JsonLevelVariantSets.VariantSets.Num() == 0)
+	{
+		return FGLTFJsonEpicLevelVariantSetsIndex(INDEX_NONE);
+	}
+
+	return Builder.AddEpicLevelVariantSets(JsonLevelVariantSets);
+}
+
+bool FGLTFEpicLevelVariantSetsConverter::TryParseVariant(FGLTFJsonEpicVariant& OutVariant, const UVariant* Variant) const
+{
+	FGLTFJsonEpicVariant JsonVariant;
+
+	for (const UVariantObjectBinding* Binding: Variant->GetBindings())
+	{
+		TryParseVariantBinding(JsonVariant, Binding);
+	}
+
+	if (JsonVariant.Nodes.Num() == 0)
+	{
+		return false;
+	}
+
+	JsonVariant.Name = Variant->GetDisplayText().ToString();
+	JsonVariant.bIsActive = const_cast<UVariant*>(Variant)->IsActive();
+
+	if (const UTexture2D* Thumbnail = const_cast<UVariant*>(Variant)->GetThumbnail())
+	{
+		// TODO: if thumbnail has generic name "Texture2D", give it a variant-relevant name
+		JsonVariant.Thumbnail = Builder.GetOrAddTexture(Thumbnail);
+	}
+
+	OutVariant = JsonVariant;
+	return true;
+}
+
+bool FGLTFEpicLevelVariantSetsConverter::TryParseVariantBinding(FGLTFJsonEpicVariant& OutVariant, const UVariantObjectBinding* Binding) const
+{
+	bool bHasParsedAnyProperty = false;
+
+	for (UPropertyValue* Property: Binding->GetCapturedProperties())
+	{
+		if (!const_cast<UPropertyValue*>(Property)->Resolve() || !Property->HasRecordedData())
+		{
+			continue;
+		}
+
+		const FName PropertyName = Property->GetPropertyName();
+		const FFieldClass* PropertyClass = Property->GetPropertyClass();
+
+		if (Property->IsA<UPropertyValueMaterial>())
+		{
+			if (Builder.ExportOptions->ExportMaterialVariants != EGLTFMaterialVariantMode::None && TryParseMaterialPropertyValue(OutVariant, Property))
+			{
+				bHasParsedAnyProperty = true;
+			}
+		}
+		else if (PropertyName == TEXT("StaticMesh")) // TODO: should we not also check PropertyClass?
+		{
+			if (Builder.ExportOptions->bExportMeshVariants && TryParseStaticMeshPropertyValue(OutVariant, Property))
+			{
+				bHasParsedAnyProperty = true;
+			}
+		}
+		else if (PropertyName == TEXT("SkeletalMesh")) // TODO: should we not also check PropertyClass?
+		{
+			if (Builder.ExportOptions->bExportMeshVariants && TryParseSkeletalMeshPropertyValue(OutVariant, Property))
+			{
+				bHasParsedAnyProperty = true;
+			}
+		}
+		else if (PropertyName == TEXT("bVisible") && PropertyClass->IsChildOf(FBoolProperty::StaticClass()))
+		{
+			if (Builder.ExportOptions->bExportVisibilityVariants && TryParseVisibilityPropertyValue(OutVariant, Property))
+			{
+				bHasParsedAnyProperty = true;
+			}
+		}
+
+		// TODO: add support for more properties
+	}
+
+	return bHasParsedAnyProperty;
+}
+
+bool FGLTFEpicLevelVariantSetsConverter::TryParseVisibilityPropertyValue(FGLTFJsonEpicVariant& OutVariant, const UPropertyValue* Property) const
+{
+	const USceneComponent* Target = static_cast<USceneComponent*>(Property->GetPropertyParentContainerAddress());
+	if (Target == nullptr)
+	{
+		Builder.LogWarning(FString::Printf(
+			TEXT("Target object for property is invalid, the property will be skipped. Context: %s"),
+			*FGLTFVariantUtility::GetLogContext(Property)));
+		return false;
+	}
+
+	const AActor* Owner = Target->GetOwner();
+	if (Owner == nullptr)
+	{
+		Builder.LogWarning(FString::Printf(TEXT("Invalid scene component, the property will be skipped. Context: %s"),
+			*FGLTFVariantUtility::GetLogContext(Property)));
+		return false;
+	}
+
+	if (Builder.bSelectedActorsOnly && !Owner->IsSelected())
+	{
+		Builder.LogWarning(FString::Printf(
+			TEXT("Target actor for property is not selected for export, the property will be skipped. Context: %s"),
+			*FGLTFVariantUtility::GetLogContext(Property)));
+		return false;
+	}
+
+	bool bIsVisible;
+
+	if (!FGLTFVariantUtility::TryGetPropertyValue(const_cast<UPropertyValue*>(Property), bIsVisible))
+	{
+		Builder.LogWarning(FString::Printf(
+			TEXT("Failed to parse recorded data for property, it will be skipped. Context: %s"),
+			*FGLTFVariantUtility::GetLogContext(Property)));
+		return false;
+	}
+
+	const FGLTFJsonNodeIndex NodeIndex = Builder.GetOrAddNode(Target);
+	const FGLTFJsonNodeIndex ComponentNodeIndex = Builder.GetComponentNodeIndex(NodeIndex);
+
+	FGLTFJsonEpicVariantNodeProperties& NodeProperties = OutVariant.Nodes.FindOrAdd(ComponentNodeIndex);
+
+	NodeProperties.Node = ComponentNodeIndex;
+	NodeProperties.bIsVisible = bIsVisible;
+	return true;
+}
+
+bool FGLTFEpicLevelVariantSetsConverter::TryParseMaterialPropertyValue(FGLTFJsonEpicVariant& OutVariant, const UPropertyValue* Property) const
+{
+	UPropertyValueMaterial* MaterialProperty = Cast<UPropertyValueMaterial>(const_cast<UPropertyValue*>(Property));
+	if (MaterialProperty == nullptr)
+	{
+		Builder.LogWarning(FString::Printf(
+			TEXT("Material property is invalid, the property will be skipped. Context: %s"),
+			*FGLTFVariantUtility::GetLogContext(Property)));
+		return false;
+	}
+
+	const USceneComponent* Target = static_cast<USceneComponent*>(Property->GetPropertyParentContainerAddress());
+	if (Target == nullptr)
+	{
+		Builder.LogWarning(FString::Printf(
+			TEXT("Target object for property is invalid, the property will be skipped. Context: %s"),
+			*FGLTFVariantUtility::GetLogContext(Property)));
+		return false;
+	}
+
+	const AActor* Owner = Target->GetOwner();
+	if (Owner == nullptr)
+	{
+		Builder.LogWarning(FString::Printf(TEXT("Invalid scene component, the property will be skipped. Context: %s"),
+			*FGLTFVariantUtility::GetLogContext(Property)));
+		return false;
+	}
+
+	if (Builder.bSelectedActorsOnly && !Owner->IsSelected())
+	{
+		Builder.LogWarning(FString::Printf(
+			TEXT("Target actor for property is not selected for export, the property will be skipped. Context: %s"),
+			*FGLTFVariantUtility::GetLogContext(Property)));
+		return false;
+	}
+
+	const TArray<FCapturedPropSegment>& CapturedPropSegments = FGLTFVariantUtility::GetCapturedPropSegments(MaterialProperty);
+	const int32 NumPropSegments = CapturedPropSegments.Num();
+
+	if (NumPropSegments < 1)
+	{
+		Builder.LogWarning(FString::Printf(
+			TEXT("Failed to parse element index to apply the material to, the property will be skipped. Context: %s"),
+			*FGLTFVariantUtility::GetLogContext(Property)));
+		return false;
+	}
+
+	// NOTE: UPropertyValueMaterial::GetMaterial does *not* ensure that the recorded data has been loaded,
+	// so we need to call UProperty::GetRecordedData first to make that happen.
+	MaterialProperty->GetRecordedData();
+
+	// TODO: find way to determine whether the material is null because "None" was selected, or because it failed to resolve
+	const UMaterialInterface* Material = MaterialProperty->GetMaterial();
+	const int32 MaterialIndex = CapturedPropSegments[NumPropSegments - 1].PropertyIndex;
+
+	const FGLTFMeshData* MeshData = nullptr;
+	FGLTFIndexArray SectionIndices;
+
+	if (Builder.ExportOptions->ExportMaterialVariants == EGLTFMaterialVariantMode::UseMeshData)
+	{
+		if (Builder.ExportOptions->BakeMaterialInputs == EGLTFMaterialBakeMode::UseMeshData)
+		{
+			const int32 DefaultLODIndex = Builder.ExportOptions->DefaultLevelOfDetail;
+
+			if (const UStaticMeshComponent* StaticMeshComponent = Cast<UStaticMeshComponent>(Target))
+			{
+				const UStaticMesh* StaticMesh = StaticMeshComponent->GetStaticMesh();
+				const int32 LODIndex = FGLTFMeshUtility::GetLOD(StaticMesh, StaticMeshComponent, DefaultLODIndex);
+				const FStaticMeshLODResources& MeshLOD = StaticMesh->GetLODForExport(LODIndex);
+
+				SectionIndices = FGLTFMeshUtility::GetSectionIndices(MeshLOD, MaterialIndex);
+				MeshData = Builder.StaticMeshDataConverter.GetOrAdd(StaticMesh, StaticMeshComponent, LODIndex);
+			}
+			else if (const USkeletalMeshComponent* SkeletalMeshComponent = Cast<USkeletalMeshComponent>(Target))
+			{
+				const USkeletalMesh* SkeletalMesh = SkeletalMeshComponent->SkeletalMesh;
+				const int32 LODIndex = FGLTFMeshUtility::GetLOD(SkeletalMesh, SkeletalMeshComponent, DefaultLODIndex);
+				const FSkeletalMeshRenderData* RenderData = SkeletalMesh->GetResourceForRendering();
+				const FSkeletalMeshLODRenderData& MeshLOD = RenderData->LODRenderData[LODIndex];
+
+				SectionIndices = FGLTFMeshUtility::GetSectionIndices(MeshLOD, MaterialIndex);
+				MeshData = Builder.SkeletalMeshDataConverter.GetOrAdd(SkeletalMesh, SkeletalMeshComponent, LODIndex);
+			}
+		}
+		else
+		{
+			// TODO: report warning (about materials won't be export using mesh data because BakeMaterialInputs is not set to UseMeshData)
+		}
+	}
+
+	FGLTFJsonEpicVariantMaterial VariantMaterial;
+	VariantMaterial.Material = Builder.GetOrAddMaterial(Material, MeshData, SectionIndices);
+	VariantMaterial.Index = MaterialIndex;
+
+	const FGLTFJsonNodeIndex NodeIndex = Builder.GetOrAddNode(Target);
+	const FGLTFJsonNodeIndex ComponentNodeIndex = Builder.GetComponentNodeIndex(NodeIndex);
+	FGLTFJsonEpicVariantNodeProperties& NodeProperties = OutVariant.Nodes.FindOrAdd(ComponentNodeIndex);
+
+	NodeProperties.Node = ComponentNodeIndex;
+	NodeProperties.Materials.Add(VariantMaterial);
+	return true;
+}
+
+bool FGLTFEpicLevelVariantSetsConverter::TryParseStaticMeshPropertyValue(FGLTFJsonEpicVariant& OutVariant, const UPropertyValue* Property) const
+{
+	const USceneComponent* Target = static_cast<USceneComponent*>(Property->GetPropertyParentContainerAddress());
+	if (Target == nullptr)
+	{
+		Builder.LogWarning(FString::Printf(
+			TEXT("Target object for property is invalid, the property will be skipped. Context: %s"),
+			*FGLTFVariantUtility::GetLogContext(Property)));
+		return false;
+	}
+
+	const AActor* Owner = Target->GetOwner();
+	if (Owner == nullptr)
+	{
+		Builder.LogWarning(FString::Printf(TEXT("Invalid scene component, the property will be skipped. Context: %s"),
+			*FGLTFVariantUtility::GetLogContext(Property)));
+		return false;
+	}
+
+	if (Builder.bSelectedActorsOnly && !Owner->IsSelected())
+	{
+		Builder.LogWarning(FString::Printf(
+			TEXT("Target actor for property is not selected for export, the property will be skipped. Context: %s"),
+			*FGLTFVariantUtility::GetLogContext(Property)));
+		return false;
+	}
+
+	const UMeshComponent* MeshComponent = Cast<UMeshComponent>(Target);
+	if (MeshComponent == nullptr)
+	{
+		Builder.LogWarning(FString::Printf(
+			TEXT("Target object for property has no mesh-component, the property will be skipped. Context: %s"),
+			*FGLTFVariantUtility::GetLogContext(Property)));
+		return false;
+	}
+
+	const UStaticMesh* StaticMesh;
+
+	if (!FGLTFVariantUtility::TryGetPropertyValue(const_cast<UPropertyValue*>(Property), StaticMesh))
+	{
+		Builder.LogWarning(FString::Printf(
+			TEXT("Failed to parse recorded data for property, it will be skipped. Context: %s"),
+			*FGLTFVariantUtility::GetLogContext(Property)));
+		return false;
+	}
+
+	const FGLTFMaterialArray OverrideMaterials(MeshComponent->OverrideMaterials);
+	const FGLTFJsonNodeIndex NodeIndex = Builder.GetOrAddNode(Target);
+	const FGLTFJsonNodeIndex ComponentNodeIndex = Builder.GetComponentNodeIndex(NodeIndex);
+	const FGLTFJsonMeshIndex MeshIndex = Builder.GetOrAddMesh(StaticMesh, OverrideMaterials);
+	FGLTFJsonEpicVariantNodeProperties& NodeProperties = OutVariant.Nodes.FindOrAdd(ComponentNodeIndex);
+
+	NodeProperties.Node = ComponentNodeIndex;
+	NodeProperties.Mesh = MeshIndex;
+	return true;
+}
+
+bool FGLTFEpicLevelVariantSetsConverter::TryParseSkeletalMeshPropertyValue(FGLTFJsonEpicVariant& OutVariant, const UPropertyValue* Property) const
+{
+	const USceneComponent* Target = static_cast<USceneComponent*>(Property->GetPropertyParentContainerAddress());
+	if (Target == nullptr)
+	{
+		Builder.LogWarning(FString::Printf(
+			TEXT("Target object for property is invalid, the property will be skipped. Context: %s"),
+			*FGLTFVariantUtility::GetLogContext(Property)));
+		return false;
+	}
+
+	const AActor* Owner = Target->GetOwner();
+	if (Owner == nullptr)
+	{
+		Builder.LogWarning(FString::Printf(TEXT("Invalid scene component, the property will be skipped. Context: %s"),
+			*FGLTFVariantUtility::GetLogContext(Property)));
+		return false;
+	}
+
+	if (Builder.bSelectedActorsOnly && !Owner->IsSelected())
+	{
+		Builder.LogWarning(FString::Printf(
+			TEXT("Target actor for property is not selected for export, the property will be skipped. Context: %s"),
+			*FGLTFVariantUtility::GetLogContext(Property)));
+		return false;
+	}
+
+	const UMeshComponent* MeshComponent = Cast<UMeshComponent>(Target);
+	if (MeshComponent == nullptr)
+	{
+		Builder.LogWarning(FString::Printf(
+			TEXT("Target object for property has no mesh-component, the property will be skipped. Context: %s"),
+			*FGLTFVariantUtility::GetLogContext(Property)));
+		return false;
+	}
+
+	const USkeletalMesh* SkeletalMesh;
+
+	if (!FGLTFVariantUtility::TryGetPropertyValue(const_cast<UPropertyValue*>(Property), SkeletalMesh) || SkeletalMesh == nullptr)
+	{
+		Builder.LogWarning(FString::Printf(
+			TEXT("Failed to parse recorded data for property, it will be skipped. Context: %s"),
+			*FGLTFVariantUtility::GetLogContext(Property)));
+		return false;
+	}
+
+	const FGLTFMaterialArray OverrideMaterials(MeshComponent->OverrideMaterials);
+	const FGLTFJsonNodeIndex NodeIndex = Builder.GetOrAddNode(Target);
+	const FGLTFJsonMeshIndex MeshIndex = Builder.GetOrAddMesh(SkeletalMesh, OverrideMaterials);
+	const FGLTFJsonNodeIndex ComponentNodeIndex = Builder.GetComponentNodeIndex(NodeIndex);
+	FGLTFJsonEpicVariantNodeProperties& NodeProperties = OutVariant.Nodes.FindOrAdd(ComponentNodeIndex);
+
+	NodeProperties.Node = ComponentNodeIndex;
+	NodeProperties.Mesh = MeshIndex;
+	return true;
+}
diff --git a/Engine/Plugins/Enterprise/GLTFExporter/MigratedFromGit/Source/GLTFExporter/Private/Converters/GLTFEpicLevelVariantSetsConverters.h b/Engine/Plugins/Enterprise/GLTFExporter/MigratedFromGit/Source/GLTFExporter/Private/Converters/GLTFEpicLevelVariantSetsConverters.h
new file mode 100644
--- /dev/null
+++ b/Engine/Plugins/Enterprise/GLTFExporter/MigratedFromGit/Source/GLTFExporter/Private/Converters/GLTFEpicLevelVariantSetsConverters.h
@@ -0,0 +1,28 @@
+// Copyright Epic Games, Inc. All Rights Reserved.
+
+#pragma once
+
+#include "Json/GLTFJsonIndex.h"
+#include "Json/GLTFJsonEpicLevelVariantSets.h"
+#include "Converters/GLTFConverter.h"
+#include "Converters/GLTFBuilderContext.h"
+
+class ULevelVariantSets;
+class UVariantSet;
+class UVariant;
+class UVariantObjectBinding;
+class UPropertyValue;
+
+class FGLTFEpicLevelVariantSetsConverter final : public FGLTFBuilderContext, public TGLTFConverter<FGLTFJsonEpicLevelVariantSetsIndex, const ULevelVariantSets*>
+{
+	using FGLTFBuilderContext::FGLTFBuilderContext;
+
+	virtual FGLTFJsonEpicLevelVariantSetsIndex Convert(const ULevelVariantSets* LevelVariantSets) override;
+
+	bool TryParseVariant(FGLTFJsonEpicVariant& OutVariant, const UVariant* Variant) const;
+	bool TryParseVariantBinding(FGLTFJsonEpicVariant& OutVariant, const UVariantObjectBinding* Binding) const;
+	bool TryParseVisibilityPropertyValue(FGLTFJsonEpicVariant& OutVariant, const UPropertyValue* Property) const;
+	bool TryParseMaterialPropertyValue(FGLTFJsonEpicVariant& OutVariant, const UPropertyValue* Property) const;
+	bool TryParseStaticMeshPropertyValue(FGLTFJsonEpicVariant& OutVariant, const UPropertyValue* Property) const;
+	bool TryParseSkeletalMeshPropertyValue(FGLTFJsonEpicVariant& OutVariant, const UPropertyValue* Property) const;
+};
diff --git a/Engine/Plugins/Enterprise/GLTFExporter/MigratedFromGit/Source/GLTFExporter/Private/Converters/GLTFLevelVariantSetsConverters.cpp b/Engine/Plugins/Enterprise/GLTFExporter/MigratedFromGit/Source/GLTFExporter/Private/Converters/GLTFLevelVariantSetsConverters.cpp
deleted file mode 100644
diff --git a/Engine/Plugins/Enterprise/GLTFExporter/MigratedFromGit/Source/GLTFExporter/Private/Converters/GLTFLevelVariantSetsConverters.h b/Engine/Plugins/Enterprise/GLTFExporter/MigratedFromGit/Source/GLTFExporter/Private/Converters/GLTFLevelVariantSetsConverters.h
deleted file mode 100644
diff --git a/Engine/Plugins/Enterprise/GLTFExporter/MigratedFromGit/Source/GLTFExporter/Private/Converters/GLTFSceneConverters.cpp b/Engine/Plugins/Enterprise/GLTFExporter/MigratedFromGit/Source/GLTFExporter/Private/Converters/GLTFSceneConverters.cpp
--- a/Engine/Plugins/Enterprise/GLTFExporter/MigratedFromGit/Source/GLTFExporter/Private/Converters/GLTFSceneConverters.cpp
+++ b/Engine/Plugins/Enterprise/GLTFExporter/MigratedFromGit/Source/GLTFExporter/Private/Converters/GLTFSceneConverters.cpp
@@ -32,10 +32,10 @@ FGLTFJsonSceneIndex FGLTFSceneConverter::Convert(const UWorld* World)
 					{
 						if (const ULevelVariantSets* LevelVariantSets = const_cast<ALevelVariantSetsActor*>(LevelVariantSetsActor)->GetLevelVariantSets(true))
 						{
-							const FGLTFJsonLevelVariantSetsIndex LevelVariantSetsIndex = Builder.GetOrAddLevelVariantSets(LevelVariantSets);
-							if (LevelVariantSetsIndex != INDEX_NONE)
+							const FGLTFJsonEpicLevelVariantSetsIndex EpicLevelVariantSetsIndex = Builder.GetOrAddEpicLevelVariantSets(LevelVariantSets);
+							if (EpicLevelVariantSetsIndex != INDEX_NONE)
 							{
-								Scene.LevelVariantSets.Add(LevelVariantSetsIndex);
+								Scene.EpicLevelVariantSets.Add(EpicLevelVariantSetsIndex);
 							}
 						}
 					}
diff --git a/Engine/Plugins/Enterprise/GLTFExporter/MigratedFromGit/Source/GLTFExporter/Private/Exporters/GLTFLevelVariantSetsExporter.cpp b/Engine/Plugins/Enterprise/GLTFExporter/MigratedFromGit/Source/GLTFExporter/Private/Exporters/GLTFLevelVariantSetsExporter.cpp
--- a/Engine/Plugins/Enterprise/GLTFExporter/MigratedFromGit/Source/GLTFExporter/Private/Exporters/GLTFLevelVariantSetsExporter.cpp
+++ b/Engine/Plugins/Enterprise/GLTFExporter/MigratedFromGit/Source/GLTFExporter/Private/Exporters/GLTFLevelVariantSetsExporter.cpp
@@ -51,8 +51,8 @@ bool UGLTFLevelVariantSetsExporter::AddObject(FGLTFContainerBuilder& Builder, co
 		return false;
 	}
 
-	const FGLTFJsonLevelVariantSetsIndex LevelVariantSetsIndex = Builder.GetOrAddLevelVariantSets(LevelVariantSets);
-	if (LevelVariantSetsIndex == INDEX_NONE)
+	const FGLTFJsonEpicLevelVariantSetsIndex EpicLevelVariantSetsIndex = Builder.GetOrAddEpicLevelVariantSets(LevelVariantSets);
+	if (EpicLevelVariantSetsIndex == INDEX_NONE)
 	{
 		Builder.LogError(
 			FString::Printf(TEXT("Failed to export level variant sets %s"),
@@ -60,7 +60,7 @@ bool UGLTFLevelVariantSetsExporter::AddObject(FGLTFContainerBuilder& Builder, co
 		return false;
 	}
 
-	Builder.GetScene(SceneIndex).LevelVariantSets.AddUnique(LevelVariantSetsIndex);
+	Builder.GetScene(SceneIndex).EpicLevelVariantSets.AddUnique(EpicLevelVariantSetsIndex);
 
 	Builder.DefaultScene = SceneIndex;
 	return true;
diff --git a/Engine/Plugins/Enterprise/GLTFExporter/MigratedFromGit/Source/GLTFExporter/Private/Json/GLTFJsonEpicLevelVariantSets.h b/Engine/Plugins/Enterprise/GLTFExporter/MigratedFromGit/Source/GLTFExporter/Private/Json/GLTFJsonEpicLevelVariantSets.h
new file mode 100644
--- /dev/null
+++ b/Engine/Plugins/Enterprise/GLTFExporter/MigratedFromGit/Source/GLTFExporter/Private/Json/GLTFJsonEpicLevelVariantSets.h
@@ -0,0 +1,123 @@
+// Copyright Epic Games, Inc. All Rights Reserved.
+
+#pragma once
+
+#include "Json/GLTFJsonObject.h"
+#include "Json/GLTFJsonIndex.h"
+#include "Misc/Optional.h"
+
+struct FGLTFJsonEpicVariantMaterial : IGLTFJsonObject
+{
+	FGLTFJsonMaterialIndex Material;
+	int32                  Index;
+
+	FGLTFJsonEpicVariantMaterial()
+		: Material(INDEX_NONE)
+		, Index(INDEX_NONE)
+	{
+	}
+
+	virtual void WriteObject(IGLTFJsonWriter& Writer) const override
+	{
+		Writer.Write(TEXT("material"), Material);
+
+		if (Index != INDEX_NONE)
+		{
+			Writer.Write(TEXT("index"), Index);
+		}
+	}
+};
+
+struct FGLTFJsonEpicVariantNodeProperties : IGLTFJsonObject
+{
+	FGLTFJsonNodeIndex Node;
+	TOptional<bool>    bIsVisible;
+
+	TOptional<FGLTFJsonMeshIndex>        Mesh;
+	TArray<FGLTFJsonEpicVariantMaterial> Materials;
+
+	virtual void WriteObject(IGLTFJsonWriter& Writer) const override
+	{
+		if (Node != INDEX_NONE)
+		{
+			Writer.Write(TEXT("node"), Node);
+		}
+
+		Writer.StartObject(TEXT("properties"));
+
+		if (bIsVisible.IsSet())
+		{
+			Writer.Write(TEXT("visible"), bIsVisible.GetValue());
+		}
+
+		if (Mesh.IsSet())
+		{
+			Writer.Write(TEXT("mesh"), Mesh.GetValue());
+		}
+
+		if (Materials.Num() > 0)
+		{
+			Writer.Write(TEXT("materials"), Materials);
+		}
+
+		Writer.EndObject();
+	}
+};
+
+struct FGLTFJsonEpicVariant : IGLTFJsonObject
+{
+	FString Name;
+	bool    bIsActive;
+
+	FGLTFJsonTextureIndex Thumbnail;
+	TMap<FGLTFJsonNodeIndex, FGLTFJsonEpicVariantNodeProperties> Nodes;
+
+	virtual void WriteObject(IGLTFJsonWriter& Writer) const override
+	{
+		Writer.Write(TEXT("name"), Name);
+		Writer.Write(TEXT("active"), bIsActive);
+
+		if (Thumbnail != INDEX_NONE)
+		{
+			Writer.Write(TEXT("thumbnail"), Thumbnail);
+		}
+
+		TArray<FGLTFJsonEpicVariantNodeProperties> NodeValues;
+		Nodes.GenerateValueArray(NodeValues);
+		Writer.Write(TEXT("nodes"), NodeValues);
+	}
+};
+
+struct FGLTFJsonEpicVariantSet : IGLTFJsonObject
+{
+	FString Name;
+
+	TArray<FGLTFJsonEpicVariant> Variants;
+
+	virtual void WriteObject(IGLTFJsonWriter& Writer) const override
+	{
+		if (!Name.IsEmpty())
+		{
+			Writer.Write(TEXT("name"), Name);
+		}
+
+		Writer.Write(TEXT("variants"), Variants);
+	}
+};
+
+struct FGLTFJsonEpicLevelVariantSets : IGLTFJsonObject
+{
+	FString Name;
+
+	TArray<FGLTFJsonEpicVariantSet> VariantSets;
+
+	virtual void WriteObject(IGLTFJsonWriter& Writer) const override
+	{
+		if (!Name.IsEmpty())
+		{
+			Writer.Write(TEXT("name"), Name);
+		}
+
+		Writer.Write(TEXT("variantSets"), VariantSets);
+	}
+};
diff --git a/Engine/Plugins/Enterprise/GLTFExporter/MigratedFromGit/Source/GLTFExporter/Private/Json/GLTFJsonIndex.h b/Engine/Plugins/Enterprise/GLTFExporter/MigratedFromGit/Source/GLTFExporter/Private/Json/GLTFJsonIndex.h
--- a/Engine/Plugins/Enterprise/GLTFExporter/MigratedFromGit/Source/GLTFExporter/Private/Json/GLTFJsonIndex.h
+++ b/Engine/Plugins/Enterprise/GLTFExporter/MigratedFromGit/Source/GLTFExporter/Private/Json/GLTFJsonIndex.h
@@ -71,4 +71,4 @@ struct FGLTFJsonHotspotIndex : FGLTFJsonIndex<FGLTFJsonHotspotIndex> { using FGL
 struct FGLTFJsonLightIndex : FGLTFJsonIndex<FGLTFJsonLightIndex> { using FGLTFJsonIndex::FGLTFJsonIndex; };
 struct FGLTFJsonLightMapIndex : FGLTFJsonIndex<FGLTFJsonLightMapIndex> { using FGLTFJsonIndex::FGLTFJsonIndex; };
 struct FGLTFJsonSkySphereIndex : FGLTFJsonIndex<FGLTFJsonSkySphereIndex> { using FGLTFJsonIndex::FGLTFJsonIndex; };
-struct FGLTFJsonLevelVariantSetsIndex : FGLTFJsonIndex<FGLTFJsonLevelVariantSetsIndex> { using FGLTFJsonIndex::FGLTFJsonIndex; };
+struct FGLTFJsonEpicLevelVariantSetsIndex : FGLTFJsonIndex<FGLTFJsonEpicLevelVariantSetsIndex> { using FGLTFJsonIndex::FGLTFJsonIndex; };
diff --git a/Engine/Plugins/Enterprise/GLTFExporter/MigratedFromGit/Source/GLTFExporter/Private/Json/GLTFJsonLevelVariantSets.h b/Engine/Plugins/Enterprise/GLTFExporter/MigratedFromGit/Source/GLTFExporter/Private/Json/GLTFJsonLevelVariantSets.h
deleted file mode 100644
diff --git a/Engine/Plugins/Enterprise/GLTFExporter/MigratedFromGit/Source/GLTFExporter/Private/Json/GLTFJsonRoot.h b/Engine/Plugins/Enterprise/GLTFExporter/MigratedFromGit/Source/GLTFExporter/Private/Json/GLTFJsonRoot.h
--- a/Engine/Plugins/Enterprise/GLTFExporter/MigratedFromGit/Source/GLTFExporter/Private/Json/GLTFJsonRoot.h
+++ b/Engine/Plugins/Enterprise/GLTFExporter/MigratedFromGit/Source/GLTFExporter/Private/Json/GLTFJsonRoot.h
@@ -23,7 +23,7 @@
 #include "Json/GLTFJsonLight.h"
 #include "Json/GLTFJsonLightMap.h"
 #include "Json/GLTFJsonSkySphere.h"
-#include "Json/GLTFJsonLevelVariantSets.h"
+#include "Json/GLTFJsonEpicLevelVariantSets.h"
 #include "Json/GLTFJsonWriter.h"
 
 struct FGLTFJsonRoot : IGLTFJsonObject
@@ -52,7 +52,7 @@ struct FGLTFJsonRoot : IGLTFJsonObject
 	TArray<TUniquePtr<FGLTFJsonLight>>      Lights;
 	TArray<TUniquePtr<FGLTFJsonLightMap>>   LightMaps;
 	TArray<TUniquePtr<FGLTFJsonSkySphere>>  SkySpheres;
-	TArray<TUniquePtr<FGLTFJsonLevelVariantSets>>  LevelVariantSets;
+	TArray<TUniquePtr<FGLTFJsonEpicLevelVariantSets>>  EpicLevelVariantSets;
 
 	virtual void WriteObject(IGLTFJsonWriter& Writer) const override
 	{
@@ -77,7 +77,7 @@ struct FGLTFJsonRoot : IGLTFJsonObject
 		if (Skins.Num() > 0) Writer.Write(TEXT("skins"), Skins);
 		if (Textures.Num() > 0) Writer.Write(TEXT("textures"), Textures);
 
-		if (Backdrops.Num() > 0 || Hotspots.Num() > 0 || Lights.Num() > 0 || LightMaps.Num() > 0 || SkySpheres.Num() > 0 || LevelVariantSets.Num() > 0)
+		if (Backdrops.Num() > 0 || Hotspots.Num() > 0 || Lights.Num() > 0 || LightMaps.Num() > 0 || SkySpheres.Num() > 0 || EpicLevelVariantSets.Num() > 0)
 		{
 			Writer.StartExtensions();
 
@@ -95,10 +95,10 @@ struct FGLTFJsonRoot : IGLTFJsonObject
 				Writer.EndExtension();
 			}
 
-			if (LevelVariantSets.Num() > 0)
+			if (EpicLevelVariantSets.Num() > 0)
 			{
 				Writer.StartExtension(EGLTFJsonExtension::EPIC_LevelVariantSets);
-				Writer.Write(TEXT("levelVariantSets"), LevelVariantSets);
+				Writer.Write(TEXT("levelVariantSets"), EpicLevelVariantSets);
 				Writer.EndExtension();
 			}
 
diff --git a/Engine/Plugins/Enterprise/GLTFExporter/MigratedFromGit/Source/GLTFExporter/Private/Json/GLTFJsonScene.h b/Engine/Plugins/Enterprise/GLTFExporter/MigratedFromGit/Source/GLTFExporter/Private/Json/GLTFJsonScene.h
--- a/Engine/Plugins/Enterprise/GLTFExporter/MigratedFromGit/Source/GLTFExporter/Private/Json/GLTFJsonScene.h
+++ b/Engine/Plugins/Enterprise/GLTFExporter/MigratedFromGit/Source/GLTFExporter/Private/Json/GLTFJsonScene.h
@@ -10,7 +10,7 @@ struct FGLTFJsonScene : IGLTFJsonObject
 	FString Name;
 
 	TArray<FGLTFJsonNodeIndex> Nodes;
-	TArray<FGLTFJsonLevelVariantSetsIndex>  LevelVariantSets;
+	TArray<FGLTFJsonEpicLevelVariantSetsIndex>  EpicLevelVariantSets;
 
 	virtual void WriteObject(IGLTFJsonWriter& Writer) const override
 	{
@@ -24,12 +24,12 @@ struct FGLTFJsonScene : IGLTFJsonObject
 			Writer.Write(TEXT("nodes"), Nodes);
 		}
 
-		if (LevelVariantSets.Num() > 0)
+		if (EpicLevelVariantSets.Num() > 0)
 		{
 			Writer.StartExtensions();
 
 			Writer.StartExtension(EGLTFJsonExtension::EPIC_LevelVariantSets);
-			Writer.Write(TEXT("levelVariantSets"), LevelVariantSets);
+			Writer.Write(TEXT("levelVariantSets"), EpicLevelVariantSets);
 			Writer.EndExtension();
 
 			Writer.EndExtensions();
