[GLTFExporter] chore: add json data-structure for KHR_materials_variants

commit 4876ab7fb6ed33a7347a64a9591d1480461ac547
Date: Thu, 16 Sep 2021 09:19:50 +0200

    chore: add json data-structure for KHR_materials_variants
    (cherry picked from commit 45c926bb99568846ff2babb1780efaba0cd31c10​)

#skipci

diff --git a/Engine/Plugins/Enterprise/GLTFExporter/MigratedFromGit/Source/GLTFExporter/Private/Builders/GLTFJsonBuilder.cpp b/Engine/Plugins/Enterprise/GLTFExporter/MigratedFromGit/Source/GLTFExporter/Private/Builders/GLTFJsonBuilder.cpp
--- a/Engine/Plugins/Enterprise/GLTFExporter/MigratedFromGit/Source/GLTFExporter/Private/Builders/GLTFJsonBuilder.cpp
+++ b/Engine/Plugins/Enterprise/GLTFExporter/MigratedFromGit/Source/GLTFExporter/Private/Builders/GLTFJsonBuilder.cpp
@@ -134,6 +134,11 @@ FGLTFJsonEpicLevelVariantSetsIndex FGLTFJsonBuilder::AddEpicLevelVariantSets(con
 	return FGLTFJsonEpicLevelVariantSetsIndex(JsonRoot.EpicLevelVariantSets.Add(MakeUnique<FGLTFJsonEpicLevelVariantSets>(JsonEpicLevelVariantSets)));
 }
 
+FGLTFJsonKhrMaterialVariantIndex FGLTFJsonBuilder::AddKhrMaterialVariant(const FGLTFJsonKhrMaterialVariant& JsonKhrMaterialVariant)
+{
+	return FGLTFJsonKhrMaterialVariantIndex(JsonRoot.KhrMaterialVariants.Add(MakeUnique<FGLTFJsonKhrMaterialVariant>(JsonKhrMaterialVariant)));
+}
+
 FGLTFJsonNodeIndex FGLTFJsonBuilder::AddChildNode(FGLTFJsonNodeIndex ParentIndex, const FGLTFJsonNode& JsonNode)
 {
 	const FGLTFJsonNodeIndex ChildIndex = AddNode(JsonNode);
@@ -253,6 +258,11 @@ FGLTFJsonEpicLevelVariantSets& FGLTFJsonBuilder::GetEpicLevelVariantSets(FGLTFJs
 	return *JsonRoot.EpicLevelVariantSets[EpicLevelVariantSetsIndex];
 }
 
+FGLTFJsonKhrMaterialVariant& FGLTFJsonBuilder::GetKhrMaterialVariant(FGLTFJsonKhrMaterialVariantIndex KhrMaterialVariantIndex)
+{
+	return *JsonRoot.KhrMaterialVariants[KhrMaterialVariantIndex];
+}
+
 FGLTFJsonNodeIndex FGLTFJsonBuilder::GetComponentNodeIndex(FGLTFJsonNodeIndex NodeIndex)
 {
 	if (NodeIndex == INDEX_NONE)
diff --git a/Engine/Plugins/Enterprise/GLTFExporter/MigratedFromGit/Source/GLTFExporter/Private/Builders/GLTFJsonBuilder.h b/Engine/Plugins/Enterprise/GLTFExporter/MigratedFromGit/Source/GLTFExporter/Private/Builders/GLTFJsonBuilder.h
--- a/Engine/Plugins/Enterprise/GLTFExporter/MigratedFromGit/Source/GLTFExporter/Private/Builders/GLTFJsonBuilder.h
+++ b/Engine/Plugins/Enterprise/GLTFExporter/MigratedFromGit/Source/GLTFExporter/Private/Builders/GLTFJsonBuilder.h
@@ -40,6 +40,7 @@ public:
 	FGLTFJsonLightMapIndex AddLightMap(const FGLTFJsonLightMap& JsonLightMap = {});
 	FGLTFJsonSkySphereIndex AddSkySphere(const FGLTFJsonSkySphere& JsonSkySphere = {});
 	FGLTFJsonEpicLevelVariantSetsIndex AddEpicLevelVariantSets(const FGLTFJsonEpicLevelVariantSets& JsonEpicLevelVariantSets = {});
+	FGLTFJsonKhrMaterialVariantIndex AddKhrMaterialVariant(const FGLTFJsonKhrMaterialVariant& JsonKhrMaterialVariant = {});
 
 	FGLTFJsonNodeIndex AddChildNode(FGLTFJsonNodeIndex ParentNodeIndex, const FGLTFJsonNode& JsonNode = {});
 	FGLTFJsonNodeIndex AddChildComponentNode(FGLTFJsonNodeIndex ParentNodeIndex, const FGLTFJsonNode& JsonNode = {});
@@ -63,6 +64,7 @@ public:
 	FGLTFJsonLightMap& GetLightMap(FGLTFJsonLightMapIndex LightMapIndex);
 	FGLTFJsonSkySphere& GetSkySphere(FGLTFJsonSkySphereIndex SkySphereIndex);
 	FGLTFJsonEpicLevelVariantSets& GetEpicLevelVariantSets(FGLTFJsonEpicLevelVariantSetsIndex EpicLevelVariantSetsIndex);
+	FGLTFJsonKhrMaterialVariant& GetKhrMaterialVariant(FGLTFJsonKhrMaterialVariantIndex KhrMaterialVariantIndex);
 
 	FGLTFJsonNodeIndex GetComponentNodeIndex(FGLTFJsonNodeIndex NodeIndex);
 
diff --git a/Engine/Plugins/Enterprise/GLTFExporter/MigratedFromGit/Source/GLTFExporter/Private/Json/GLTFJsonEnums.h b/Engine/Plugins/Enterprise/GLTFExporter/MigratedFromGit/Source/GLTFExporter/Private/Json/GLTFJsonEnums.h
--- a/Engine/Plugins/Enterprise/GLTFExporter/MigratedFromGit/Source/GLTFExporter/Private/Json/GLTFJsonEnums.h
+++ b/Engine/Plugins/Enterprise/GLTFExporter/MigratedFromGit/Source/GLTFExporter/Private/Json/GLTFJsonEnums.h
@@ -9,6 +9,7 @@ enum class EGLTFJsonExtension
 	KHR_LightsPunctual,
 	KHR_MaterialsClearCoat,
 	KHR_MaterialsUnlit,
+	KHR_MaterialsVariants,
 	KHR_MeshQuantization,
 	KHR_TextureTransform,
 	EPIC_AnimationHotspots,
diff --git a/Engine/Plugins/Enterprise/GLTFExporter/MigratedFromGit/Source/GLTFExporter/Private/Json/GLTFJsonIndex.h b/Engine/Plugins/Enterprise/GLTFExporter/MigratedFromGit/Source/GLTFExporter/Private/Json/GLTFJsonIndex.h
--- a/Engine/Plugins/Enterprise/GLTFExporter/MigratedFromGit/Source/GLTFExporter/Private/Json/GLTFJsonIndex.h
+++ b/Engine/Plugins/Enterprise/GLTFExporter/MigratedFromGit/Source/GLTFExporter/Private/Json/GLTFJsonIndex.h
@@ -72,3 +72,4 @@ struct FGLTFJsonLightIndex : FGLTFJsonIndex<FGLTFJsonLightIndex> { using FGLTFJs
 struct FGLTFJsonLightMapIndex : FGLTFJsonIndex<FGLTFJsonLightMapIndex> { using FGLTFJsonIndex::FGLTFJsonIndex; };
 struct FGLTFJsonSkySphereIndex : FGLTFJsonIndex<FGLTFJsonSkySphereIndex> { using FGLTFJsonIndex::FGLTFJsonIndex; };
 struct FGLTFJsonEpicLevelVariantSetsIndex : FGLTFJsonIndex<FGLTFJsonEpicLevelVariantSetsIndex> { using FGLTFJsonIndex::FGLTFJsonIndex; };
+struct FGLTFJsonKhrMaterialVariantIndex : FGLTFJsonIndex<FGLTFJsonKhrMaterialVariantIndex> { using FGLTFJsonIndex::FGLTFJsonIndex; };
diff --git a/Engine/Plugins/Enterprise/GLTFExporter/MigratedFromGit/Source/GLTFExporter/Private/Json/GLTFJsonKhrMaterialVariant.h b/Engine/Plugins/Enterprise/GLTFExporter/MigratedFromGit/Source/GLTFExporter/Private/Json/GLTFJsonKhrMaterialVariant.h
new file mode 100644
--- /dev/null
+++ b/Engine/Plugins/Enterprise/GLTFExporter/MigratedFromGit/Source/GLTFExporter/Private/Json/GLTFJsonKhrMaterialVariant.h
@@ -0,0 +1,28 @@
+// Copyright Epic Games, Inc. All Rights Reserved.
+
+#pragma once
+
+#include "Json/GLTFJsonObject.h"
+#include "Json/GLTFJsonIndex.h"
+
+struct FGLTFJsonKhrMaterialVariantMapping : IGLTFJsonObject
+{
+	FGLTFJsonMaterialIndex                   Material;
+	TArray<FGLTFJsonKhrMaterialVariantIndex> Variants;
+
+	virtual void WriteObject(IGLTFJsonWriter& Writer) const override
+	{
+		Writer.Write(TEXT("material"), Material);
+		Writer.Write(TEXT("variants"), Variants);
+	}
+};
+
+struct FGLTFJsonKhrMaterialVariant : IGLTFJsonObject
+{
+	FString Name;
+
+	virtual void WriteObject(IGLTFJsonWriter& Writer) const override
+	{
+		Writer.Write(TEXT("name"), Name);
+	}
+};
diff --git a/Engine/Plugins/Enterprise/GLTFExporter/MigratedFromGit/Source/GLTFExporter/Private/Json/GLTFJsonMesh.h b/Engine/Plugins/Enterprise/GLTFExporter/MigratedFromGit/Source/GLTFExporter/Private/Json/GLTFJsonMesh.h
--- a/Engine/Plugins/Enterprise/GLTFExporter/MigratedFromGit/Source/GLTFExporter/Private/Json/GLTFJsonMesh.h
+++ b/Engine/Plugins/Enterprise/GLTFExporter/MigratedFromGit/Source/GLTFExporter/Private/Json/GLTFJsonMesh.h
@@ -5,6 +5,7 @@
 #include "Json/GLTFJsonObject.h"
 #include "Json/GLTFJsonEnums.h"
 #include "Json/GLTFJsonIndex.h"
+#include "Json/GLTFJsonKhrMaterialVariant.h"
 
 struct FGLTFJsonAttributes : IGLTFJsonObject
 {
@@ -51,6 +52,8 @@ struct FGLTFJsonPrimitive : IGLTFJsonObject
 	EGLTFJsonPrimitiveMode Mode;
 	FGLTFJsonAttributes    Attributes;
 
+	TArray<FGLTFJsonKhrMaterialVariantMapping> KhrMaterialVariantMappings;
+
 	FGLTFJsonPrimitive()
 		: Mode(EGLTFJsonPrimitiveMode::Triangles)
 	{
@@ -74,6 +77,17 @@ struct FGLTFJsonPrimitive : IGLTFJsonObject
 		{
 			Writer.Write(TEXT("mode"), Mode);
 		}
+
+		if (KhrMaterialVariantMappings.Num() > 0)
+		{
+			Writer.StartExtensions();
+
+			Writer.StartExtension(EGLTFJsonExtension::KHR_MaterialsVariants);
+			Writer.Write(TEXT("mappings"), KhrMaterialVariantMappings);
+			Writer.EndExtension();
+
+			Writer.EndExtensions();
+		}
 	}
 };
 
diff --git a/Engine/Plugins/Enterprise/GLTFExporter/MigratedFromGit/Source/GLTFExporter/Private/Json/GLTFJsonRoot.h b/Engine/Plugins/Enterprise/GLTFExporter/MigratedFromGit/Source/GLTFExporter/Private/Json/GLTFJsonRoot.h
--- a/Engine/Plugins/Enterprise/GLTFExporter/MigratedFromGit/Source/GLTFExporter/Private/Json/GLTFJsonRoot.h
+++ b/Engine/Plugins/Enterprise/GLTFExporter/MigratedFromGit/Source/GLTFExporter/Private/Json/GLTFJsonRoot.h
@@ -24,6 +24,7 @@
 #include "Json/GLTFJsonLightMap.h"
 #include "Json/GLTFJsonSkySphere.h"
 #include "Json/GLTFJsonEpicLevelVariantSets.h"
+#include "Json/GLTFJsonKhrMaterialVariant.h"
 #include "Json/GLTFJsonWriter.h"
 
 struct FGLTFJsonRoot : IGLTFJsonObject
@@ -53,6 +54,7 @@ struct FGLTFJsonRoot : IGLTFJsonObject
 	TArray<TUniquePtr<FGLTFJsonLightMap>>   LightMaps;
 	TArray<TUniquePtr<FGLTFJsonSkySphere>>  SkySpheres;
 	TArray<TUniquePtr<FGLTFJsonEpicLevelVariantSets>>  EpicLevelVariantSets;
+	TArray<TUniquePtr<FGLTFJsonKhrMaterialVariant>>    KhrMaterialVariants;
 
 	virtual void WriteObject(IGLTFJsonWriter& Writer) const override
 	{
@@ -77,7 +79,7 @@ struct FGLTFJsonRoot : IGLTFJsonObject
 		if (Skins.Num() > 0) Writer.Write(TEXT("skins"), Skins);
 		if (Textures.Num() > 0) Writer.Write(TEXT("textures"), Textures);
 
-		if (Backdrops.Num() > 0 || Hotspots.Num() > 0 || Lights.Num() > 0 || LightMaps.Num() > 0 || SkySpheres.Num() > 0 || EpicLevelVariantSets.Num() > 0)
+		if (Backdrops.Num() > 0 || Hotspots.Num() > 0 || Lights.Num() > 0 || LightMaps.Num() > 0 || SkySpheres.Num() > 0 || EpicLevelVariantSets.Num() > 0 || KhrMaterialVariants.Num() > 0)
 		{
 			Writer.StartExtensions();
 
@@ -102,6 +104,13 @@ struct FGLTFJsonRoot : IGLTFJsonObject
 				Writer.EndExtension();
 			}
 
+			if (KhrMaterialVariants.Num() > 0)
+			{
+				Writer.StartExtension(EGLTFJsonExtension::KHR_MaterialsVariants);
+				Writer.Write(TEXT("variants"), KhrMaterialVariants);
+				Writer.EndExtension();
+			}
+
 			if (Lights.Num() > 0)
 			{
 				Writer.StartExtension(EGLTFJsonExtension::KHR_LightsPunctual);
diff --git a/Engine/Plugins/Enterprise/GLTFExporter/MigratedFromGit/Source/GLTFExporter/Private/Json/GLTFJsonUtility.h b/Engine/Plugins/Enterprise/GLTFExporter/MigratedFromGit/Source/GLTFExporter/Private/Json/GLTFJsonUtility.h
--- a/Engine/Plugins/Enterprise/GLTFExporter/MigratedFromGit/Source/GLTFExporter/Private/Json/GLTFJsonUtility.h
+++ b/Engine/Plugins/Enterprise/GLTFExporter/MigratedFromGit/Source/GLTFExporter/Private/Json/GLTFJsonUtility.h
@@ -19,6 +19,7 @@ struct FGLTFJsonUtility
 			case EGLTFJsonExtension::KHR_LightsPunctual:      return TEXT("KHR_lights_punctual");
 			case EGLTFJsonExtension::KHR_MaterialsClearCoat:  return TEXT("KHR_materials_clearcoat");
 			case EGLTFJsonExtension::KHR_MaterialsUnlit:      return TEXT("KHR_materials_unlit");
+			case EGLTFJsonExtension::KHR_MaterialsVariants:   return TEXT("KHR_materials_variants");
 			case EGLTFJsonExtension::KHR_MeshQuantization:    return TEXT("KHR_mesh_quantization");
 			case EGLTFJsonExtension::KHR_TextureTransform:    return TEXT("KHR_texture_transform");
 			case EGLTFJsonExtension::EPIC_AnimationHotspots:  return TEXT("EPIC_animation_hotspots");
