[GLTFExporter] refactor: replace export option bExportVariantSets with VariantSetsMode

commit fa846523e6ba10e2a821623b9fda705db27541f9
Date: Thu, 16 Sep 2021 10:09:04 +0200

    refactor: replace export option bExportVariantSets with VariantSetsMode
    allows the user to select which extension (if any) to use when exporting variant sets
    (cherry picked from commit 9fd9ff4de3e51c9a52db9c620d67ed3accf29baa​)

#skipci

diff --git a/Engine/Plugins/Enterprise/GLTFExporter/MigratedFromGit/Source/GLTFExporter/Private/Converters/GLTFSceneConverters.cpp b/Engine/Plugins/Enterprise/GLTFExporter/MigratedFromGit/Source/GLTFExporter/Private/Converters/GLTFSceneConverters.cpp
--- a/Engine/Plugins/Enterprise/GLTFExporter/MigratedFromGit/Source/GLTFExporter/Private/Converters/GLTFSceneConverters.cpp
+++ b/Engine/Plugins/Enterprise/GLTFExporter/MigratedFromGit/Source/GLTFExporter/Private/Converters/GLTFSceneConverters.cpp
@@ -28,7 +28,7 @@ FGLTFJsonSceneIndex FGLTFSceneConverter::Convert(const UWorld* World)
 				// TODO: should a LevelVariantSet be exported even if not selected for export?
 				if (const ALevelVariantSetsActor *LevelVariantSetsActor = Cast<ALevelVariantSetsActor>(Actor))
 				{
-					if (Builder.ExportOptions->bExportVariantSets)
+					if (Builder.ExportOptions->VariantSetsMode == EGLTFVariantSetsMode::Epic)
 					{
 						if (const ULevelVariantSets* LevelVariantSets = const_cast<ALevelVariantSetsActor*>(LevelVariantSetsActor)->GetLevelVariantSets(true))
 						{
@@ -39,6 +39,10 @@ FGLTFJsonSceneIndex FGLTFSceneConverter::Convert(const UWorld* World)
 							}
 						}
 					}
+					else if (Builder.ExportOptions->VariantSetsMode == EGLTFVariantSetsMode::Khronos)
+					{
+						// TODO: implement
+					}
 				}
 
 				const FGLTFJsonNodeIndex NodeIndex = Builder.GetOrAddNode(Actor);
diff --git a/Engine/Plugins/Enterprise/GLTFExporter/MigratedFromGit/Source/GLTFExporter/Private/Exporters/GLTFLevelVariantSetsExporter.cpp b/Engine/Plugins/Enterprise/GLTFExporter/MigratedFromGit/Source/GLTFExporter/Private/Exporters/GLTFLevelVariantSetsExporter.cpp
--- a/Engine/Plugins/Enterprise/GLTFExporter/MigratedFromGit/Source/GLTFExporter/Private/Exporters/GLTFLevelVariantSetsExporter.cpp
+++ b/Engine/Plugins/Enterprise/GLTFExporter/MigratedFromGit/Source/GLTFExporter/Private/Exporters/GLTFLevelVariantSetsExporter.cpp
@@ -15,7 +15,7 @@ bool UGLTFLevelVariantSetsExporter::AddObject(FGLTFContainerBuilder& Builder, co
 {
 	const ULevelVariantSets* LevelVariantSets = CastChecked<ULevelVariantSets>(Object);
 
-	if (!Builder.ExportOptions->bExportVariantSets)
+	if (Builder.ExportOptions->VariantSetsMode == EGLTFVariantSetsMode::None)
 	{
 		Builder.LogError(
 			FString::Printf(TEXT("Failed to export level variant sets %s because variant sets are disabled by export options"),
@@ -51,16 +51,23 @@ bool UGLTFLevelVariantSetsExporter::AddObject(FGLTFContainerBuilder& Builder, co
 		return false;
 	}
 
-	const FGLTFJsonEpicLevelVariantSetsIndex EpicLevelVariantSetsIndex = Builder.GetOrAddEpicLevelVariantSets(LevelVariantSets);
-	if (EpicLevelVariantSetsIndex == INDEX_NONE)
+	if (Builder.ExportOptions->VariantSetsMode == EGLTFVariantSetsMode::Epic)
 	{
-		Builder.LogError(
-			FString::Printf(TEXT("Failed to export level variant sets %s"),
-			*LevelVariantSets->GetName()));
-		return false;
-	}
+		const FGLTFJsonEpicLevelVariantSetsIndex EpicLevelVariantSetsIndex = Builder.GetOrAddEpicLevelVariantSets(LevelVariantSets);
+		if (EpicLevelVariantSetsIndex == INDEX_NONE)
+		{
+			Builder.LogError(
+				FString::Printf(TEXT("Failed to export level variant sets %s"),
+				*LevelVariantSets->GetName()));
+			return false;
+		}
 
-	Builder.GetScene(SceneIndex).EpicLevelVariantSets.AddUnique(EpicLevelVariantSetsIndex);
+		Builder.GetScene(SceneIndex).EpicLevelVariantSets.AddUnique(EpicLevelVariantSetsIndex);
+	}
+	else if (Builder.ExportOptions->VariantSetsMode == EGLTFVariantSetsMode::Khronos)
+	{
+		// TODO: implement
+	}
 
 	Builder.DefaultScene = SceneIndex;
 	return true;
diff --git a/Engine/Plugins/Enterprise/GLTFExporter/MigratedFromGit/Source/GLTFExporter/Private/GLTFExportOptions.cpp b/Engine/Plugins/Enterprise/GLTFExporter/MigratedFromGit/Source/GLTFExporter/Private/GLTFExportOptions.cpp
--- a/Engine/Plugins/Enterprise/GLTFExporter/MigratedFromGit/Source/GLTFExporter/Private/GLTFExportOptions.cpp
+++ b/Engine/Plugins/Enterprise/GLTFExporter/MigratedFromGit/Source/GLTFExporter/Private/GLTFExportOptions.cpp
@@ -44,7 +44,7 @@ void UGLTFExportOptions::ResetToDefault()
 	bExportAnimationHotspots = false;
 	bExportHDRIBackdrops = false;
 	bExportSkySpheres = false;
-	bExportVariantSets = false;
+	VariantSetsMode = EGLTFVariantSetsMode::None;
 	ExportMaterialVariants = EGLTFMaterialVariantMode::UseMeshData;
 	bExportMeshVariants = true;
 	bExportVisibilityVariants = true;
diff --git a/Engine/Plugins/Enterprise/GLTFExporter/MigratedFromGit/Source/GLTFExporter/Public/GLTFExportOptions.h b/Engine/Plugins/Enterprise/GLTFExporter/MigratedFromGit/Source/GLTFExporter/Public/GLTFExportOptions.h
--- a/Engine/Plugins/Enterprise/GLTFExporter/MigratedFromGit/Source/GLTFExporter/Public/GLTFExportOptions.h
+++ b/Engine/Plugins/Enterprise/GLTFExporter/MigratedFromGit/Source/GLTFExporter/Public/GLTFExportOptions.h
@@ -52,6 +52,17 @@ enum class EGLTFSceneMobility : uint8
 };
 ENUM_CLASS_FLAGS(EGLTFSceneMobility);
 
+UENUM(BlueprintType)
+enum class EGLTFVariantSetsMode : uint8
+{
+	/** Never export variants sets. */
+	None,
+	/** Uses the official extension KHR_materials_variants. Supports material variants only. */
+	Khronos,
+	/** Uses the extension EPIC_level_variant_sets, which is supported by Unreal's glTF viewer. */
+	Epic
+};
+
 UENUM(BlueprintType)
 enum class EGLTFMaterialVariantMode : uint8
 {
@@ -200,20 +211,20 @@ class GLTFEXPORTER_API UGLTFExportOptions : public UObject
 	UPROPERTY(EditAnywhere, BlueprintReadWrite, Config, Category = Scene)
 	bool bExportSkySpheres;
 
-	/** If enabled, export LevelVariantSetsActors. Uses extension EPIC_level_variant_sets, which is supported by Unreal's glTF viewer. */
+	/** Mode determining if and how to export LevelVariantSetsActors. */
 	UPROPERTY(EditAnywhere, BlueprintReadWrite, Config, Category = Scene)
-	bool bExportVariantSets;
+	EGLTFVariantSetsMode VariantSetsMode;
 
 	/** Mode determining if and how to export material variants that change the materials property on a static or skeletal mesh component. */
-	UPROPERTY(EditAnywhere, BlueprintReadWrite, Config, Category = VariantSets, Meta = (EditCondition = "bExportVariantSets"))
+	UPROPERTY(EditAnywhere, BlueprintReadWrite, Config, Category = VariantSets, Meta = (EditCondition = "VariantSetsMode != EGLTFVariantSetsMode::None"))
 	EGLTFMaterialVariantMode ExportMaterialVariants;
 
 	/** If enabled, export variants that change the mesh property on a static or skeletal mesh component. */
-	UPROPERTY(EditAnywhere, BlueprintReadWrite, Config, Category = VariantSets, Meta = (EditCondition = "bExportVariantSets"))
+	UPROPERTY(EditAnywhere, BlueprintReadWrite, Config, Category = VariantSets, Meta = (EditCondition = "VariantSetsMode == EGLTFVariantSetsMode::Epic"))
 	bool bExportMeshVariants;
 
 	/** If enabled, export variants that change the visible property on a scene component. */
-	UPROPERTY(EditAnywhere, BlueprintReadWrite, Config, Category = VariantSets, Meta = (EditCondition = "bExportVariantSets"))
+	UPROPERTY(EditAnywhere, BlueprintReadWrite, Config, Category = VariantSets, Meta = (EditCondition = "VariantSetsMode == EGLTFVariantSetsMode::Epic"))
 	bool bExportVisibilityVariants;
 
 	UFUNCTION(BlueprintCallable, Category = General)
