[GLTFExporter] chore: add GetOrAddMesh method that accepts a UMeshComponent

commit f2ecc819c102a20a5b29bfe57e6f76c00ab4ead9
Date: Fri, 17 Sep 2021 08:54:38 +0200

    chore: add GetOrAddMesh method that accepts a UMeshComponent
    (cherry picked from commit 427508d71f480dfcedd84ad8d6f6725b827dbd35​)

#skipci

diff --git a/Engine/Plugins/Enterprise/GLTFExporter/MigratedFromGit/Source/GLTFExporter/Private/Builders/GLTFConvertBuilder.cpp b/Engine/Plugins/Enterprise/GLTFExporter/MigratedFromGit/Source/GLTFExporter/Private/Builders/GLTFConvertBuilder.cpp
--- a/Engine/Plugins/Enterprise/GLTFExporter/MigratedFromGit/Source/GLTFExporter/Private/Builders/GLTFConvertBuilder.cpp
+++ b/Engine/Plugins/Enterprise/GLTFExporter/MigratedFromGit/Source/GLTFExporter/Private/Builders/GLTFConvertBuilder.cpp
@@ -144,6 +144,22 @@ FGLTFJsonMeshIndex FGLTFConvertBuilder::GetOrAddMesh(const USkeletalMeshComponen
 	return SkeletalMeshConverter.GetOrAdd(SkeletalMesh, SkeletalMeshComponent, Materials, LODIndex);
 }
 
+FGLTFJsonMeshIndex FGLTFConvertBuilder::GetOrAddMesh(const UMeshComponent* MeshComponent, const FGLTFMaterialArray& Materials, int32 LODIndex)
+{
+	if (const UStaticMeshComponent* StaticMeshComponent = Cast<UStaticMeshComponent>(MeshComponent))
+	{
+		return GetOrAddMesh(StaticMeshComponent, Materials, LODIndex);
+	}
+	else if (const USkeletalMeshComponent* SkeletalMeshComponent = Cast<USkeletalMeshComponent>(MeshComponent))
+	{
+		return GetOrAddMesh(SkeletalMeshComponent, Materials, LODIndex);
+	}
+	else
+	{
+		return FGLTFJsonMeshIndex(INDEX_NONE);
+	}
+}
+
 FGLTFJsonMaterialIndex FGLTFConvertBuilder::GetOrAddMaterial(const UMaterialInterface* Material, const FGLTFMeshData* MeshData, const FGLTFIndexArray& SectionIndices)
 {
 	if (Material == nullptr)
diff --git a/Engine/Plugins/Enterprise/GLTFExporter/MigratedFromGit/Source/GLTFExporter/Private/Builders/GLTFConvertBuilder.h b/Engine/Plugins/Enterprise/GLTFExporter/MigratedFromGit/Source/GLTFExporter/Private/Builders/GLTFConvertBuilder.h
--- a/Engine/Plugins/Enterprise/GLTFExporter/MigratedFromGit/Source/GLTFExporter/Private/Builders/GLTFConvertBuilder.h
+++ b/Engine/Plugins/Enterprise/GLTFExporter/MigratedFromGit/Source/GLTFExporter/Private/Builders/GLTFConvertBuilder.h
@@ -48,6 +48,7 @@ public:
 	FGLTFJsonMeshIndex GetOrAddMesh(const UStaticMeshComponent* StaticMeshComponent, const FGLTFMaterialArray& Materials = {}, int32 LODIndex = -1);
 	FGLTFJsonMeshIndex GetOrAddMesh(const USkeletalMesh* SkeletalMesh, const FGLTFMaterialArray& Materials = {}, int32 LODIndex = -1);
 	FGLTFJsonMeshIndex GetOrAddMesh(const USkeletalMeshComponent* SkeletalMeshComponent, const FGLTFMaterialArray& Materials = {}, int32 LODIndex = -1);
+	FGLTFJsonMeshIndex GetOrAddMesh(const UMeshComponent* MeshComponent, const FGLTFMaterialArray& Materials = {}, int32 LODIndex = -1);
 
 	FGLTFJsonMaterialIndex GetOrAddMaterial(const UMaterialInterface* Material, const FGLTFMeshData* MeshData = nullptr, const FGLTFIndexArray& SectionIndices = {});
 
