[GLTFExporter] chore: add a way of retrieving the json-root

commit bcbd5c6dc6f130091c3f400b1f73bec2a40d98b9
Date: Fri, 17 Sep 2021 08:55:59 +0200

    chore: add a way of retrieving the json-root
    (cherry picked from commit efc0f2e4c6d5e6ca822553eaf26af00d42942858​)

#skipci

diff --git a/Engine/Plugins/Enterprise/GLTFExporter/MigratedFromGit/Source/GLTFExporter/Private/Builders/GLTFJsonBuilder.cpp b/Engine/Plugins/Enterprise/GLTFExporter/MigratedFromGit/Source/GLTFExporter/Private/Builders/GLTFJsonBuilder.cpp
--- a/Engine/Plugins/Enterprise/GLTFExporter/MigratedFromGit/Source/GLTFExporter/Private/Builders/GLTFJsonBuilder.cpp
+++ b/Engine/Plugins/Enterprise/GLTFExporter/MigratedFromGit/Source/GLTFExporter/Private/Builders/GLTFJsonBuilder.cpp
@@ -163,6 +163,11 @@ FGLTFJsonNodeIndex FGLTFJsonBuilder::AddChildComponentNode(FGLTFJsonNodeIndex Pa
 	return ChildIndex;
 }
 
+const FGLTFJsonRoot& FGLTFJsonBuilder::GetRoot() const
+{
+	return JsonRoot;
+}
+
 FGLTFJsonAccessor& FGLTFJsonBuilder::GetAccessor(FGLTFJsonAccessorIndex AccessorIndex)
 {
 	return *JsonRoot.Accessors[AccessorIndex];
diff --git a/Engine/Plugins/Enterprise/GLTFExporter/MigratedFromGit/Source/GLTFExporter/Private/Builders/GLTFJsonBuilder.h b/Engine/Plugins/Enterprise/GLTFExporter/MigratedFromGit/Source/GLTFExporter/Private/Builders/GLTFJsonBuilder.h
--- a/Engine/Plugins/Enterprise/GLTFExporter/MigratedFromGit/Source/GLTFExporter/Private/Builders/GLTFJsonBuilder.h
+++ b/Engine/Plugins/Enterprise/GLTFExporter/MigratedFromGit/Source/GLTFExporter/Private/Builders/GLTFJsonBuilder.h
@@ -45,6 +45,7 @@ public:
 	FGLTFJsonNodeIndex AddChildNode(FGLTFJsonNodeIndex ParentNodeIndex, const FGLTFJsonNode& JsonNode = {});
 	FGLTFJsonNodeIndex AddChildComponentNode(FGLTFJsonNodeIndex ParentNodeIndex, const FGLTFJsonNode& JsonNode = {});
 
+	const FGLTFJsonRoot& GetRoot() const;
 	FGLTFJsonAccessor& GetAccessor(FGLTFJsonAccessorIndex AccessorIndex);
 	FGLTFJsonAnimation& GetAnimation(FGLTFJsonAnimationIndex AnimationIndex);
 	FGLTFJsonBuffer& GetBuffer(FGLTFJsonBufferIndex BufferIndex);
