[GLTFExporter] refactor: add default primitives to each mesh in the converter instead of in the task

commit 6911faa97c6ff1b5a6d4e6ce65400a3601598aa7
Date: Fri, 17 Sep 2021 12:56:17 +0200

    refactor: add default primitives to each mesh in the converter instead of in the task
    we will need to modify these primitives from within other converters in order to add data for extensions etc, so they need to exist before the mesh-task is completed
    (cherry picked from commit 332eca3a87dc2a184784ef72f3ffb0d172829ff3​)

#skipci

diff --git a/Engine/Plugins/Enterprise/GLTFExporter/MigratedFromGit/Source/GLTFExporter/Private/Converters/GLTFMeshConverters.cpp b/Engine/Plugins/Enterprise/GLTFExporter/MigratedFromGit/Source/GLTFExporter/Private/Converters/GLTFMeshConverters.cpp
--- a/Engine/Plugins/Enterprise/GLTFExporter/MigratedFromGit/Source/GLTFExporter/Private/Converters/GLTFMeshConverters.cpp
+++ b/Engine/Plugins/Enterprise/GLTFExporter/MigratedFromGit/Source/GLTFExporter/Private/Converters/GLTFMeshConverters.cpp
@@ -53,7 +53,11 @@ FGLTFJsonMeshIndex FGLTFStaticMeshConverter::Convert(const UStaticMesh* StaticMe
 	}
 #endif
 
-	const FGLTFJsonMeshIndex MeshIndex = Builder.AddMesh();
+	FGLTFJsonMesh JsonMesh;
+	const int32 MaterialCount = FGLTFMeshUtility::GetMaterials(StaticMesh).Num();
+	JsonMesh.Primitives.AddDefaulted(MaterialCount);
+
+	const FGLTFJsonMeshIndex MeshIndex = Builder.AddMesh(JsonMesh);
 	Builder.SetupTask<FGLTFStaticMeshTask>(Builder, MeshSectionConverter, StaticMesh, StaticMeshComponent, Materials, LODIndex, MeshIndex);
 	return MeshIndex;
 }
@@ -102,7 +106,11 @@ FGLTFJsonMeshIndex FGLTFSkeletalMeshConverter::Convert(const USkeletalMesh* Skel
 	}
 #endif
 
-	const FGLTFJsonMeshIndex MeshIndex = Builder.AddMesh();
+	FGLTFJsonMesh JsonMesh;
+	const int32 MaterialCount = FGLTFMeshUtility::GetMaterials(SkeletalMesh).Num();
+	JsonMesh.Primitives.AddDefaulted(MaterialCount);
+
+	const FGLTFJsonMeshIndex MeshIndex = Builder.AddMesh(JsonMesh);
 	Builder.SetupTask<FGLTFSkeletalMeshTask>(Builder, MeshSectionConverter, SkeletalMesh, SkeletalMeshComponent, Materials, LODIndex, MeshIndex);
 	return MeshIndex;
 }
diff --git a/Engine/Plugins/Enterprise/GLTFExporter/MigratedFromGit/Source/GLTFExporter/Private/Tasks/GLTFMeshTasks.cpp b/Engine/Plugins/Enterprise/GLTFExporter/MigratedFromGit/Source/GLTFExporter/Private/Tasks/GLTFMeshTasks.cpp
--- a/Engine/Plugins/Enterprise/GLTFExporter/MigratedFromGit/Source/GLTFExporter/Private/Tasks/GLTFMeshTasks.cpp
+++ b/Engine/Plugins/Enterprise/GLTFExporter/MigratedFromGit/Source/GLTFExporter/Private/Tasks/GLTFMeshTasks.cpp
@@ -157,7 +157,6 @@ void FGLTFStaticMeshTask::Complete()
 
 	const TArray<FStaticMaterial>& MaterialSlots = FGLTFMeshUtility::GetMaterials(StaticMesh);
 	const int32 MaterialCount = MaterialSlots.Num();
-	JsonMesh.Primitives.AddDefaulted(MaterialCount);
 
 	for (int32 MaterialIndex = 0; MaterialIndex < MaterialCount; ++MaterialIndex)
 	{
@@ -262,7 +261,6 @@ void FGLTFSkeletalMeshTask::Complete()
 
 	const TArray<FSkeletalMaterial>& MaterialSlots = FGLTFMeshUtility::GetMaterials(SkeletalMesh);
 	const uint16 MaterialCount = MaterialSlots.Num();
-	JsonMesh.Primitives.AddDefaulted(MaterialCount);
 
 	for (uint16 MaterialIndex = 0; MaterialIndex < MaterialCount; ++MaterialIndex)
 	{
