[GLTFExporter] feat: implement support for exporting variants using KHR_materials_variants

commit c975a21ac7daf7cee856df4de7ebf072a0c4a369
Date: Fri, 17 Sep 2021 08:58:34 +0200

    feat: implement support for exporting variants using KHR_materials_variants
    (cherry picked from commit cb566de73332e285db930c70b12c20ec24f3e950​)

#skipci

diff --git a/Engine/Plugins/Enterprise/GLTFExporter/MigratedFromGit/Source/GLTFExporter/Private/Builders/GLTFConvertBuilder.cpp b/Engine/Plugins/Enterprise/GLTFExporter/MigratedFromGit/Source/GLTFExporter/Private/Builders/GLTFConvertBuilder.cpp
--- a/Engine/Plugins/Enterprise/GLTFExporter/MigratedFromGit/Source/GLTFExporter/Private/Builders/GLTFConvertBuilder.cpp
+++ b/Engine/Plugins/Enterprise/GLTFExporter/MigratedFromGit/Source/GLTFExporter/Private/Builders/GLTFConvertBuilder.cpp
@@ -427,6 +427,16 @@ FGLTFJsonEpicLevelVariantSetsIndex FGLTFConvertBuilder::GetOrAddEpicLevelVariant
 	return EpicLevelVariantSetsConverter.GetOrAdd(LevelVariantSets);
 }
 
+FGLTFJsonKhrMaterialVariantIndex FGLTFConvertBuilder::GetOrAddKhrMaterialVariant(const UVariant* Variant)
+{
+	if (Variant == nullptr)
+	{
+		return FGLTFJsonKhrMaterialVariantIndex(INDEX_NONE);
+	}
+
+	return KhrMaterialVariantConverter.GetOrAdd(Variant);
+}
+
 FGLTFJsonLightMapIndex FGLTFConvertBuilder::GetOrAddLightMap(const UStaticMeshComponent* StaticMeshComponent)
 {
 	if (StaticMeshComponent == nullptr)
diff --git a/Engine/Plugins/Enterprise/GLTFExporter/MigratedFromGit/Source/GLTFExporter/Private/Builders/GLTFConvertBuilder.h b/Engine/Plugins/Enterprise/GLTFExporter/MigratedFromGit/Source/GLTFExporter/Private/Builders/GLTFConvertBuilder.h
--- a/Engine/Plugins/Enterprise/GLTFExporter/MigratedFromGit/Source/GLTFExporter/Private/Builders/GLTFConvertBuilder.h
+++ b/Engine/Plugins/Enterprise/GLTFExporter/MigratedFromGit/Source/GLTFExporter/Private/Builders/GLTFConvertBuilder.h
@@ -19,6 +19,7 @@
 #include "Converters/GLTFEpicLevelVariantSetsConverters.h"
 #include "Converters/GLTFLightMapConverters.h"
 #include "Converters/GLTFHotspotConverters.h"
+#include "Converters/GLTFKhrVariantConverters.h"
 #include "Converters/GLTFSkySphereConverters.h"
 
 class FGLTFConvertBuilder : public FGLTFImageBuilder
@@ -82,6 +83,7 @@ public:
 	FGLTFJsonLightIndex GetOrAddLight(const ULightComponent* LightComponent);
 	FGLTFJsonBackdropIndex GetOrAddBackdrop(const AActor* BackdropActor);
 	FGLTFJsonEpicLevelVariantSetsIndex GetOrAddEpicLevelVariantSets(const ULevelVariantSets* LevelVariantSets);
+	FGLTFJsonKhrMaterialVariantIndex GetOrAddKhrMaterialVariant(const UVariant* Variant);
 	FGLTFJsonLightMapIndex GetOrAddLightMap(const UStaticMeshComponent* StaticMeshComponent);
 	FGLTFJsonHotspotIndex GetOrAddHotspot(const AGLTFHotspotActor* HotspotActor);
 	FGLTFJsonSkySphereIndex GetOrAddSkySphere(const AActor* SkySphereActor);
@@ -126,6 +128,7 @@ private:
 	FGLTFLightConverter LightConverter{ *this };
 	FGLTFBackdropConverter BackdropConverter{ *this };
 	FGLTFEpicLevelVariantSetsConverter EpicLevelVariantSetsConverter{ *this };
+	FGLTFKhrMaterialVariantConverter KhrMaterialVariantConverter{ *this };
 	FGLTFLightMapConverter LightMapConverter{ *this };
 	FGLTFHotspotConverter HotspotConverter{ *this };
 	FGLTFSkySphereConverter SkySphereConverter{ *this };
diff --git a/Engine/Plugins/Enterprise/GLTFExporter/MigratedFromGit/Source/GLTFExporter/Private/Converters/GLTFKhrVariantConverters.cpp b/Engine/Plugins/Enterprise/GLTFExporter/MigratedFromGit/Source/GLTFExporter/Private/Converters/GLTFKhrVariantConverters.cpp
new file mode 100644
--- /dev/null
+++ b/Engine/Plugins/Enterprise/GLTFExporter/MigratedFromGit/Source/GLTFExporter/Private/Converters/GLTFKhrVariantConverters.cpp
@@ -0,0 +1,225 @@
+﻿// Copyright Epic Games, Inc. All Rights Reserved.
+
+#include "Converters/GLTFKhrVariantConverters.h"
+#include "Converters/GLTFVariantUtility.h"
+#include "Converters/GLTFMeshUtility.h"
+#include "Builders/GLTFContainerBuilder.h"
+#include "Rendering/SkeletalMeshRenderData.h"
+#include "VariantObjectBinding.h"
+#include "PropertyValueMaterial.h"
+#include "LevelVariantSets.h"
+#include "PropertyValue.h"
+#include "VariantSet.h"
+#include "Variant.h"
+
+FGLTFJsonKhrMaterialVariantIndex FGLTFKhrMaterialVariantConverter::Convert(const UVariant* Variant)
+{
+	if (Variant == nullptr || Builder.ExportOptions->ExportMaterialVariants == EGLTFMaterialVariantMode::None)
+	{
+		return FGLTFJsonKhrMaterialVariantIndex(INDEX_NONE);
+	}
+
+	FGLTFJsonKhrMaterialVariant MaterialVariant;
+	MaterialVariant.Name = Variant->GetDisplayText().ToString();
+
+	typedef TTuple<FGLTFJsonPrimitive*, FGLTFJsonMaterialIndex> TPrimitiveMaterial;
+	TArray<TPrimitiveMaterial> PrimitiveMaterials;
+
+	for (const UVariantObjectBinding* Binding: Variant->GetBindings())
+	{
+		for (const UPropertyValue* Property: Binding->GetCapturedProperties())
+		{
+			if (!const_cast<UPropertyValue*>(Property)->Resolve() || !Property->HasRecordedData())
+			{
+				continue;
+			}
+
+			if (Property->IsA<UPropertyValueMaterial>())
+			{
+				FGLTFJsonPrimitive* Primitive = nullptr;
+				FGLTFJsonMaterialIndex MaterialIndex;
+
+				if (TryParseMaterialProperty(Primitive, MaterialIndex, Property))
+				{
+					PrimitiveMaterials.Add(MakeTuple(Primitive, MaterialIndex));
+				}
+			}
+		}
+	}
+
+	if (PrimitiveMaterials.Num() < 1)
+	{
+		// TODO: add warning and / or allow unused material variants to be added?
+
+		return FGLTFJsonKhrMaterialVariantIndex(INDEX_NONE);
+	}
+
+	const FGLTFJsonKhrMaterialVariantIndex MaterialVariantIndex = Builder.AddKhrMaterialVariant(MaterialVariant);
+
+	for (const TPrimitiveMaterial& PrimitiveMaterial: PrimitiveMaterials)
+	{
+		FGLTFJsonPrimitive* Primitive = PrimitiveMaterial.Key;
+		const FGLTFJsonMaterialIndex MaterialIndex = PrimitiveMaterial.Value;
+
+		FGLTFJsonKhrMaterialVariantMapping* ExistingMapping = Primitive->KhrMaterialVariantMappings.FindByPredicate(
+			[MaterialIndex](const FGLTFJsonKhrMaterialVariantMapping& Mapping)
+			{
+				return Mapping.Material == MaterialIndex;
+			});
+
+		if (ExistingMapping != nullptr)
+		{
+			ExistingMapping->Variants.AddUnique(MaterialVariantIndex);
+		}
+		else
+		{
+			FGLTFJsonKhrMaterialVariantMapping Mapping;
+			Mapping.Material = MaterialIndex;
+			Mapping.Variants.Add(MaterialVariantIndex);
+
+			Primitive->KhrMaterialVariantMappings.Add(Mapping);
+		}
+	}
+
+	return MaterialVariantIndex;
+}
+
+bool FGLTFKhrMaterialVariantConverter::TryParseMaterialProperty(FGLTFJsonPrimitive*& OutPrimitive, FGLTFJsonMaterialIndex& OutMaterialIndex, const UPropertyValue* Property) const
+{
+	UPropertyValueMaterial* MaterialProperty = Cast<UPropertyValueMaterial>(const_cast<UPropertyValue*>(Property));
+	if (MaterialProperty == nullptr)
+	{
+		Builder.LogWarning(FString::Printf(
+			TEXT("Material property is invalid, the property will be skipped. Context: %s"),
+			*GetLogContext(Property)));
+		return false;
+	}
+
+	const USceneComponent* Target = static_cast<USceneComponent*>(Property->GetPropertyParentContainerAddress());
+	if (Target == nullptr)
+	{
+		Builder.LogWarning(FString::Printf(
+			TEXT("Target object for property is invalid, the property will be skipped. Context: %s"),
+			*GetLogContext(Property)));
+		return false;
+	}
+
+	const AActor* Owner = Target->GetOwner();
+	if (Owner == nullptr)
+	{
+		Builder.LogWarning(FString::Printf(TEXT("Invalid scene component, the property will be skipped. Context: %s"),
+			*GetLogContext(Property)));
+		return false;
+	}
+
+	if (Builder.bSelectedActorsOnly && !Owner->IsSelected())
+	{
+		Builder.LogWarning(FString::Printf(
+			TEXT("Target actor for property is not selected for export, the property will be skipped. Context: %s"),
+			*GetLogContext(Property)));
+		return false;
+	}
+
+	const UMeshComponent* MeshComponent = Cast<UMeshComponent>(Target);
+	if (MeshComponent == nullptr)
+	{
+		Builder.LogWarning(FString::Printf(
+			TEXT("Target object for property has no mesh-component, the property will be skipped. Context: %s"),
+			*GetLogContext(Property)));
+		return false;
+	}
+
+	const TArray<FCapturedPropSegment>& CapturedPropSegments = FGLTFVariantUtility::GetCapturedPropSegments(MaterialProperty);
+	const int32 NumPropSegments = CapturedPropSegments.Num();
+
+	if (NumPropSegments < 1)
+	{
+		Builder.LogWarning(FString::Printf(
+			TEXT("Failed to parse element index to apply the material to, the property will be skipped. Context: %s"),
+			*GetLogContext(MaterialProperty)));
+		return false;
+	}
+
+	// NOTE: UPropertyValueMaterial::GetMaterial does *not* ensure that the recorded data has been loaded,
+	// so we need to call UProperty::GetRecordedData first to make that happen.
+	MaterialProperty->GetRecordedData();
+
+	// TODO: find way to determine whether the material is null because "None" was selected, or because it failed to resolve
+	const UMaterialInterface* Material = MaterialProperty->GetMaterial();
+	const int32 MaterialIndex = CapturedPropSegments[NumPropSegments - 1].PropertyIndex;
+	const FGLTFJsonMeshIndex MeshIndex = Builder.GetOrAddMesh(MeshComponent);
+
+	OutPrimitive = &Builder.GetMesh(MeshIndex).Primitives[MaterialIndex];
+	OutMaterialIndex = GetResolvedMaterialForComponent(Material, MaterialIndex, MeshComponent);
+
+	return true;
+}
+
+FGLTFJsonMaterialIndex FGLTFKhrMaterialVariantConverter::GetResolvedMaterialForComponent(const UMaterialInterface* Material, int32 MaterialIndex, const UMeshComponent* MeshComponent) const
+{
+	const FGLTFMeshData* MeshData = nullptr;
+	FGLTFIndexArray SectionIndices;
+
+	if (Builder.ExportOptions->ExportMaterialVariants == EGLTFMaterialVariantMode::UseMeshData)
+	{
+		if (Builder.ExportOptions->BakeMaterialInputs == EGLTFMaterialBakeMode::UseMeshData)
+		{
+			const int32 DefaultLODIndex = Builder.ExportOptions->DefaultLevelOfDetail;
+
+			if (const UStaticMeshComponent* StaticMeshComponent = Cast<UStaticMeshComponent>(MeshComponent))
+			{
+				const UStaticMesh* StaticMesh = StaticMeshComponent->GetStaticMesh();
+				const int32 LODIndex = FGLTFMeshUtility::GetLOD(StaticMesh, StaticMeshComponent, DefaultLODIndex);
+				const FStaticMeshLODResources& MeshLOD = StaticMesh->GetLODForExport(LODIndex);
+
+				SectionIndices = FGLTFMeshUtility::GetSectionIndices(MeshLOD, MaterialIndex);
+				MeshData = Builder.StaticMeshDataConverter.GetOrAdd(StaticMesh, StaticMeshComponent, LODIndex);
+			}
+			else if (const USkeletalMeshComponent* SkeletalMeshComponent = Cast<USkeletalMeshComponent>(MeshComponent))
+			{
+				const USkeletalMesh* SkeletalMesh = SkeletalMeshComponent->SkeletalMesh;
+				const int32 LODIndex = FGLTFMeshUtility::GetLOD(SkeletalMesh, SkeletalMeshComponent, DefaultLODIndex);
+				const FSkeletalMeshRenderData* RenderData = SkeletalMesh->GetResourceForRendering();
+				const FSkeletalMeshLODRenderData& MeshLOD = RenderData->LODRenderData[LODIndex];
+
+				SectionIndices = FGLTFMeshUtility::GetSectionIndices(MeshLOD, MaterialIndex);
+				MeshData = Builder.SkeletalMeshDataConverter.GetOrAdd(SkeletalMesh, SkeletalMeshComponent, LODIndex);
+			}
+		}
+		else
+		{
+			// TODO: report warning (about materials won't be exported using mesh data because BakeMaterialInputs is not set to UseMeshData)
+		}
+	}
+
+	return Builder.GetOrAddMaterial(Material, MeshData, SectionIndices);
+}
+
+FString FGLTFKhrMaterialVariantConverter::GetLogContext(const UPropertyValue* Property) const
+{
+	const UVariantObjectBinding* Parent = Property->GetParent();
+	return GetLogContext(Parent) + TEXT("/") + Property->GetFullDisplayString();
+}
+
+FString FGLTFKhrMaterialVariantConverter::GetLogContext(const UVariantObjectBinding* Binding) const
+{
+	const UVariant* Parent = const_cast<UVariantObjectBinding*>(Binding)->GetParent();
+	return GetLogContext(Parent) + TEXT("/") + Binding->GetDisplayText().ToString();
+}
+
+FString FGLTFKhrMaterialVariantConverter::GetLogContext(const UVariant* Variant) const
+{
+	const UVariantSet* Parent = const_cast<UVariant*>(Variant)->GetParent();
+	return GetLogContext(Parent) + TEXT("/") + Variant->GetDisplayText().ToString();
+}
+
+FString FGLTFKhrMaterialVariantConverter::GetLogContext(const UVariantSet* VariantSet) const
+{
+	const ULevelVariantSets* Parent = const_cast<UVariantSet*>(VariantSet)->GetParent();
+	return GetLogContext(Parent) + TEXT("/") + VariantSet->GetDisplayText().ToString();
+}
+
+FString FGLTFKhrMaterialVariantConverter::GetLogContext(const ULevelVariantSets* LevelVariantSets) const
+{
+	return LevelVariantSets->GetName();
+}
diff --git a/Engine/Plugins/Enterprise/GLTFExporter/MigratedFromGit/Source/GLTFExporter/Private/Converters/GLTFKhrVariantConverters.h b/Engine/Plugins/Enterprise/GLTFExporter/MigratedFromGit/Source/GLTFExporter/Private/Converters/GLTFKhrVariantConverters.h
new file mode 100644
--- /dev/null
+++ b/Engine/Plugins/Enterprise/GLTFExporter/MigratedFromGit/Source/GLTFExporter/Private/Converters/GLTFKhrVariantConverters.h
@@ -0,0 +1,28 @@
+﻿// Copyright Epic Games, Inc. All Rights Reserved.
+
+#pragma once
+
+#include "Json/GLTFJsonIndex.h"
+#include "Json/GLTFJsonMesh.h"
+#include "Converters/GLTFConverter.h"
+#include "Converters/GLTFBuilderContext.h"
+#include "Engine.h"
+#include "Variant.h"
+#include "PropertyValue.h"
+#include "LevelVariantSets.h"
+
+class FGLTFKhrMaterialVariantConverter final : public FGLTFBuilderContext, public TGLTFConverter<FGLTFJsonKhrMaterialVariantIndex, const UVariant*>
+{
+	using FGLTFBuilderContext::FGLTFBuilderContext;
+
+	virtual FGLTFJsonKhrMaterialVariantIndex Convert(const UVariant* Variant) override;
+
+	bool TryParseMaterialProperty(FGLTFJsonPrimitive*& OutPrimitive, FGLTFJsonMaterialIndex& OutMaterialIndex, const UPropertyValue* Property) const;
+	FGLTFJsonMaterialIndex GetResolvedMaterialForComponent(const UMaterialInterface* Material, int32 MaterialIndex, const UMeshComponent* MeshComponent) const;
+
+	FString GetLogContext(const UPropertyValue* Property) const;
+	FString GetLogContext(const UVariantObjectBinding* Binding) const;
+	FString GetLogContext(const UVariant* Variant) const;
+	FString GetLogContext(const UVariantSet* VariantSet) const;
+	FString GetLogContext(const ULevelVariantSets* LevelVariantSets) const;
+};
diff --git a/Engine/Plugins/Enterprise/GLTFExporter/MigratedFromGit/Source/GLTFExporter/Private/Converters/GLTFSceneConverters.cpp b/Engine/Plugins/Enterprise/GLTFExporter/MigratedFromGit/Source/GLTFExporter/Private/Converters/GLTFSceneConverters.cpp
--- a/Engine/Plugins/Enterprise/GLTFExporter/MigratedFromGit/Source/GLTFExporter/Private/Converters/GLTFSceneConverters.cpp
+++ b/Engine/Plugins/Enterprise/GLTFExporter/MigratedFromGit/Source/GLTFExporter/Private/Converters/GLTFSceneConverters.cpp
@@ -5,6 +5,7 @@
 #include "Converters/GLTFConverterUtility.h"
 #include "Converters/GLTFActorUtility.h"
 #include "LevelVariantSetsActor.h"
+#include "VariantSet.h"
 
 FGLTFJsonSceneIndex FGLTFSceneConverter::Convert(const UWorld* World)
 {
@@ -28,9 +29,11 @@ FGLTFJsonSceneIndex FGLTFSceneConverter::Convert(const UWorld* World)
 				// TODO: should a LevelVariantSet be exported even if not selected for export?
 				if (const ALevelVariantSetsActor *LevelVariantSetsActor = Cast<ALevelVariantSetsActor>(Actor))
 				{
+					const ULevelVariantSets* LevelVariantSets = const_cast<ALevelVariantSetsActor*>(LevelVariantSetsActor)->GetLevelVariantSets(true);
+
 					if (Builder.ExportOptions->VariantSetsMode == EGLTFVariantSetsMode::Epic)
 					{
-						if (const ULevelVariantSets* LevelVariantSets = const_cast<ALevelVariantSetsActor*>(LevelVariantSetsActor)->GetLevelVariantSets(true))
+						if (LevelVariantSets != nullptr)
 						{
 							const FGLTFJsonEpicLevelVariantSetsIndex EpicLevelVariantSetsIndex = Builder.GetOrAddEpicLevelVariantSets(LevelVariantSets);
 							if (EpicLevelVariantSetsIndex != INDEX_NONE)
@@ -41,7 +44,16 @@ FGLTFJsonSceneIndex FGLTFSceneConverter::Convert(const UWorld* World)
 					}
 					else if (Builder.ExportOptions->VariantSetsMode == EGLTFVariantSetsMode::Khronos)
 					{
-						// TODO: implement
+						if (LevelVariantSets != nullptr)
+						{
+							for (const UVariantSet* VariantSet: LevelVariantSets->GetVariantSets())
+							{
+								for (const UVariant* Variant: VariantSet->GetVariants())
+								{
+									Builder.GetOrAddKhrMaterialVariant(Variant);
+								}
+							}
+						}
 					}
 				}
 
diff --git a/Engine/Plugins/Enterprise/GLTFExporter/MigratedFromGit/Source/GLTFExporter/Private/Exporters/GLTFLevelVariantSetsExporter.cpp b/Engine/Plugins/Enterprise/GLTFExporter/MigratedFromGit/Source/GLTFExporter/Private/Exporters/GLTFLevelVariantSetsExporter.cpp
--- a/Engine/Plugins/Enterprise/GLTFExporter/MigratedFromGit/Source/GLTFExporter/Private/Exporters/GLTFLevelVariantSetsExporter.cpp
+++ b/Engine/Plugins/Enterprise/GLTFExporter/MigratedFromGit/Source/GLTFExporter/Private/Exporters/GLTFLevelVariantSetsExporter.cpp
@@ -66,7 +66,13 @@ bool UGLTFLevelVariantSetsExporter::AddObject(FGLTFContainerBuilder& Builder, co
 	}
 	else if (Builder.ExportOptions->VariantSetsMode == EGLTFVariantSetsMode::Khronos)
 	{
-		// TODO: implement
+		if (Builder.GetRoot().KhrMaterialVariants.Num() < 1)
+		{
+			Builder.LogError(
+				FString::Printf(TEXT("Failed to export any supported variant from level variant sets %s"),
+				*LevelVariantSets->GetName()));
+			return false;
+		}
 	}
 
 	Builder.DefaultScene = SceneIndex;
