[GLTFExporter] fix: skip variant-properties with invalid material when using KHR_materials_variants

commit 39df35aae204cdd77b1b2bb43c85fa36bca902be
Date: Fri, 17 Sep 2021 13:40:27 +0200

    fix: skip variant-properties with invalid material when using KHR_materials_variants
    (cherry picked from commit 3cc44b271d548bf5ab1d05bb5fa747db378f8314​)

#skipci

diff --git a/Engine/Plugins/Enterprise/GLTFExporter/MigratedFromGit/Source/GLTFExporter/Private/Converters/GLTFKhrVariantConverters.cpp b/Engine/Plugins/Enterprise/GLTFExporter/MigratedFromGit/Source/GLTFExporter/Private/Converters/GLTFKhrVariantConverters.cpp
--- a/Engine/Plugins/Enterprise/GLTFExporter/MigratedFromGit/Source/GLTFExporter/Private/Converters/GLTFKhrVariantConverters.cpp
+++ b/Engine/Plugins/Enterprise/GLTFExporter/MigratedFromGit/Source/GLTFExporter/Private/Converters/GLTFKhrVariantConverters.cpp
@@ -144,8 +144,17 @@ bool FGLTFKhrMaterialVariantConverter::TryParseMaterialProperty(FGLTFJsonPrimiti
 	// so we need to call UProperty::GetRecordedData first to make that happen.
 	MaterialProperty->GetRecordedData();
 
-	// TODO: find way to determine whether the material is null because "None" was selected, or because it failed to resolve
 	const UMaterialInterface* Material = MaterialProperty->GetMaterial();
+	if (Material == nullptr)
+	{
+		// TODO: find way to determine whether the material is null because "None" was selected, or because it failed to resolve
+
+		Builder.LogWarning(FString::Printf(
+			TEXT("No material assigned, the property will be skipped. Context: %s"),
+			*GetLogContext(MaterialProperty)));
+		return false;
+	}
+
 	const int32 MaterialIndex = CapturedPropSegments[NumPropSegments - 1].PropertyIndex;
 	const FGLTFJsonMeshIndex MeshIndex = Builder.GetOrAddMesh(MeshComponent);
 
