[GLTFExporter] refactor: convert level-variant-sets actors prior to other actors of a level

commit 3445d0afbfa1cd3bb1d2e3e11eda0c2bb50317c6
Date: Mon, 20 Sep 2021 14:01:10 +0200

    refactor: convert level-variant-sets actors prior to other actors of a level
    (cherry picked from commit d5f8db8a6dd348d555468dfef776494d18be50ad​)

#skipci

diff --git a/Engine/Plugins/Enterprise/GLTFExporter/MigratedFromGit/Source/GLTFExporter/Private/Converters/GLTFSceneConverters.cpp b/Engine/Plugins/Enterprise/GLTFExporter/MigratedFromGit/Source/GLTFExporter/Private/Converters/GLTFSceneConverters.cpp
--- a/Engine/Plugins/Enterprise/GLTFExporter/MigratedFromGit/Source/GLTFExporter/Private/Converters/GLTFSceneConverters.cpp
+++ b/Engine/Plugins/Enterprise/GLTFExporter/MigratedFromGit/Source/GLTFExporter/Private/Converters/GLTFSceneConverters.cpp
@@ -24,39 +24,45 @@ FGLTFJsonSceneIndex FGLTFSceneConverter::Convert(const UWorld* World)
 
 			// TODO: add support for exporting Level->Model?
 
-			for (const AActor* Actor : Level->Actors)
+			if (Builder.ExportOptions->VariantSetsMode != EGLTFVariantSetsMode::None)
 			{
-				// TODO: should a LevelVariantSet be exported even if not selected for export?
-				if (const ALevelVariantSetsActor *LevelVariantSetsActor = Cast<ALevelVariantSetsActor>(Actor))
+				for (const AActor* Actor : Level->Actors)
 				{
-					const ULevelVariantSets* LevelVariantSets = const_cast<ALevelVariantSetsActor*>(LevelVariantSetsActor)->GetLevelVariantSets(true);
-
-					if (Builder.ExportOptions->VariantSetsMode == EGLTFVariantSetsMode::Epic)
+					// TODO: should a LevelVariantSet be exported even if not selected for export?
+					if (const ALevelVariantSetsActor *LevelVariantSetsActor = Cast<ALevelVariantSetsActor>(Actor))
 					{
-						if (LevelVariantSets != nullptr)
+						const ULevelVariantSets* LevelVariantSets = const_cast<ALevelVariantSetsActor*>(LevelVariantSetsActor)->GetLevelVariantSets(true);
+
+						if (Builder.ExportOptions->VariantSetsMode == EGLTFVariantSetsMode::Epic)
 						{
-							const FGLTFJsonEpicLevelVariantSetsIndex EpicLevelVariantSetsIndex = Builder.GetOrAddEpicLevelVariantSets(LevelVariantSets);
-							if (EpicLevelVariantSetsIndex != INDEX_NONE)
+							if (LevelVariantSets != nullptr)
 							{
-								Scene.EpicLevelVariantSets.Add(EpicLevelVariantSetsIndex);
+								const FGLTFJsonEpicLevelVariantSetsIndex EpicLevelVariantSetsIndex = Builder.GetOrAddEpicLevelVariantSets(LevelVariantSets);
+								if (EpicLevelVariantSetsIndex != INDEX_NONE)
+								{
+									Scene.EpicLevelVariantSets.Add(EpicLevelVariantSetsIndex);
+								}
 							}
 						}
-					}
-					else if (Builder.ExportOptions->VariantSetsMode == EGLTFVariantSetsMode::Khronos)
-					{
-						if (LevelVariantSets != nullptr)
+						else if (Builder.ExportOptions->VariantSetsMode == EGLTFVariantSetsMode::Khronos)
 						{
-							for (const UVariantSet* VariantSet: LevelVariantSets->GetVariantSets())
+							if (LevelVariantSets != nullptr)
 							{
-								for (const UVariant* Variant: VariantSet->GetVariants())
+								for (const UVariantSet* VariantSet: LevelVariantSets->GetVariantSets())
 								{
-									Builder.GetOrAddKhrMaterialVariant(Variant);
+									for (const UVariant* Variant: VariantSet->GetVariants())
+									{
+										Builder.GetOrAddKhrMaterialVariant(Variant);
+									}
 								}
 							}
 						}
 					}
 				}
+			}
 
+			for (const AActor* Actor : Level->Actors)
+			{
 				const FGLTFJsonNodeIndex NodeIndex = Builder.GetOrAddNode(Actor);
 				if (NodeIndex != INDEX_NONE && FGLTFActorUtility::IsRootActor(Actor, Builder.bSelectedActorsOnly))
 				{
