[GLTFExporter] chore: register references used by (converted) variants

commit 34ebfdc51da1650dbf18d67f2c966f43e4958fed
Date: Tue, 21 Sep 2021 08:06:16 +0200

    chore: register references used by (converted) variants
    this allows us to check if certain components or actors are referenced by converted variants during conversion, which will enable us to make smarter choices
    (cherry picked from commit 61e6f6ba30c046d60c94a6c2669d228cc039c83c​)

#skipci

diff --git a/Engine/Plugins/Enterprise/GLTFExporter/MigratedFromGit/Source/GLTFExporter/Private/Builders/GLTFConvertBuilder.h b/Engine/Plugins/Enterprise/GLTFExporter/MigratedFromGit/Source/GLTFExporter/Private/Builders/GLTFConvertBuilder.h
--- a/Engine/Plugins/Enterprise/GLTFExporter/MigratedFromGit/Source/GLTFExporter/Private/Builders/GLTFConvertBuilder.h
+++ b/Engine/Plugins/Enterprise/GLTFExporter/MigratedFromGit/Source/GLTFExporter/Private/Builders/GLTFConvertBuilder.h
@@ -21,6 +21,7 @@
 #include "Converters/GLTFHotspotConverters.h"
 #include "Converters/GLTFKhrVariantConverters.h"
 #include "Converters/GLTFSkySphereConverters.h"
+#include "Converters/GLTFVariantReferenceChecker.h"
 
 class FGLTFConvertBuilder : public FGLTFImageBuilder
 {
@@ -35,6 +36,7 @@ public:
 	// TODO: find a better place for this types of indirect converters
 	FGLTFStaticMeshDataConverter StaticMeshDataConverter;
 	FGLTFSkeletalMeshDataConverter SkeletalMeshDataConverter;
+	FGLTFVariantReferenceChecker VariantReferenceChecker;
 
 	FGLTFJsonAccessorIndex GetOrAddPositionAccessor(const FGLTFMeshSection* MeshSection, const FPositionVertexBuffer* VertexBuffer);
 	FGLTFJsonAccessorIndex GetOrAddColorAccessor(const FGLTFMeshSection* MeshSection, const FColorVertexBuffer* VertexBuffer);
diff --git a/Engine/Plugins/Enterprise/GLTFExporter/MigratedFromGit/Source/GLTFExporter/Private/Converters/GLTFEpicLevelVariantSetsConverters.cpp b/Engine/Plugins/Enterprise/GLTFExporter/MigratedFromGit/Source/GLTFExporter/Private/Converters/GLTFEpicLevelVariantSetsConverters.cpp
--- a/Engine/Plugins/Enterprise/GLTFExporter/MigratedFromGit/Source/GLTFExporter/Private/Converters/GLTFEpicLevelVariantSetsConverters.cpp
+++ b/Engine/Plugins/Enterprise/GLTFExporter/MigratedFromGit/Source/GLTFExporter/Private/Converters/GLTFEpicLevelVariantSetsConverters.cpp
@@ -162,6 +162,8 @@ bool FGLTFEpicLevelVariantSetsConverter::TryParseVisibilityPropertyValue(FGLTFJs
 		return false;
 	}
 
+	Builder.VariantReferenceChecker.GetOrAdd(Property, Owner);
+
 	const FGLTFJsonNodeIndex NodeIndex = Builder.GetOrAddNode(Target);
 	const FGLTFJsonNodeIndex ComponentNodeIndex = Builder.GetComponentNodeIndex(NodeIndex);
 
@@ -262,6 +264,8 @@ bool FGLTFEpicLevelVariantSetsConverter::TryParseMaterialPropertyValue(FGLTFJson
 		}
 	}
 
+	Builder.VariantReferenceChecker.GetOrAdd(Property, Owner);
+
 	FGLTFJsonEpicVariantMaterial VariantMaterial;
 	VariantMaterial.Material = Builder.GetOrAddMaterial(Material, MeshData, SectionIndices);
 	VariantMaterial.Index = MaterialIndex;
@@ -321,6 +325,8 @@ bool FGLTFEpicLevelVariantSetsConverter::TryParseStaticMeshPropertyValue(FGLTFJs
 		return false;
 	}
 
+	Builder.VariantReferenceChecker.GetOrAdd(Property, Owner);
+
 	const FGLTFMaterialArray OverrideMaterials(MeshComponent->OverrideMaterials);
 	const FGLTFJsonNodeIndex NodeIndex = Builder.GetOrAddNode(Target);
 	const FGLTFJsonNodeIndex ComponentNodeIndex = Builder.GetComponentNodeIndex(NodeIndex);
@@ -378,6 +384,8 @@ bool FGLTFEpicLevelVariantSetsConverter::TryParseSkeletalMeshPropertyValue(FGLTF
 		return false;
 	}
 
+	Builder.VariantReferenceChecker.GetOrAdd(Property, Owner);
+
 	const FGLTFMaterialArray OverrideMaterials(MeshComponent->OverrideMaterials);
 	const FGLTFJsonNodeIndex NodeIndex = Builder.GetOrAddNode(Target);
 	const FGLTFJsonMeshIndex MeshIndex = Builder.GetOrAddMesh(SkeletalMesh, OverrideMaterials);
diff --git a/Engine/Plugins/Enterprise/GLTFExporter/MigratedFromGit/Source/GLTFExporter/Private/Converters/GLTFKhrVariantConverters.cpp b/Engine/Plugins/Enterprise/GLTFExporter/MigratedFromGit/Source/GLTFExporter/Private/Converters/GLTFKhrVariantConverters.cpp
--- a/Engine/Plugins/Enterprise/GLTFExporter/MigratedFromGit/Source/GLTFExporter/Private/Converters/GLTFKhrVariantConverters.cpp
+++ b/Engine/Plugins/Enterprise/GLTFExporter/MigratedFromGit/Source/GLTFExporter/Private/Converters/GLTFKhrVariantConverters.cpp
@@ -155,6 +155,9 @@ bool FGLTFKhrMaterialVariantConverter::TryParseMaterialProperty(FGLTFJsonPrimiti
 		return false;
 	}
 
+	Builder.VariantReferenceChecker.GetOrAdd(MaterialProperty, MeshComponent);
+	Builder.VariantReferenceChecker.GetOrAdd(MaterialProperty, Owner);
+
 	const int32 MaterialIndex = CapturedPropSegments[NumPropSegments - 1].PropertyIndex;
 	const FGLTFJsonMeshIndex MeshIndex = Builder.GetOrAddMesh(MeshComponent);
 
diff --git a/Engine/Plugins/Enterprise/GLTFExporter/MigratedFromGit/Source/GLTFExporter/Private/Converters/GLTFVariantReferenceChecker.h b/Engine/Plugins/Enterprise/GLTFExporter/MigratedFromGit/Source/GLTFExporter/Private/Converters/GLTFVariantReferenceChecker.h
new file mode 100644
--- /dev/null
+++ b/Engine/Plugins/Enterprise/GLTFExporter/MigratedFromGit/Source/GLTFExporter/Private/Converters/GLTFVariantReferenceChecker.h
@@ -0,0 +1,34 @@
+﻿// Copyright Epic Games, Inc. All Rights Reserved.
+
+#pragma once
+
+#include "Converters/GLTFConverter.h"
+#include "PropertyValue.h"
+
+class FGLTFVariantReferenceChecker final : public TGLTFConverter<bool, const UPropertyValue*, const UObject*>
+{
+	typedef TTuple<const UPropertyValue*, const UObject*> TPropertyReference;
+
+	TArray<TPropertyReference> References;
+
+	virtual bool Convert(const UPropertyValue* PropertyValue, const UObject* PropertyReference) override
+	{
+		if (PropertyValue == nullptr || PropertyReference == nullptr)
+		{
+			return false;
+		}
+
+		References.Add(TPropertyReference(PropertyValue, PropertyReference));
+		return true;
+	}
+
+public:
+
+	bool IsReferenced(const UObject* Object) const
+	{
+		return References.FindByPredicate([&Object](const TPropertyReference& PropertyReference)
+		{
+			return PropertyReference.Value == Object;
+		}) != nullptr;
+	}
+};
