[GLTFExporter] feat: make all meshes referenced by material-variants unique

commit 92591211f80415d125e3b0c20fa36562afbc9e2b
Date: Tue, 21 Sep 2021 08:07:48 +0200

    feat: make all meshes referenced by material-variants unique
    since KHR_materials_variants stores its mapping on a mesh(-primitive) level, we must ensure that each targeted mesh is unique instead of shared, otherwise more than one mesh-instance in the scene may be affected by the same variant
    (cherry picked from commit 2ffe97bd557775a601ddf6f1640800fcbff38a35​)

#skipci

diff --git a/Engine/Plugins/Enterprise/GLTFExporter/MigratedFromGit/Source/GLTFExporter/Private/Converters/GLTFMeshConverters.cpp b/Engine/Plugins/Enterprise/GLTFExporter/MigratedFromGit/Source/GLTFExporter/Private/Converters/GLTFMeshConverters.cpp
--- a/Engine/Plugins/Enterprise/GLTFExporter/MigratedFromGit/Source/GLTFExporter/Private/Converters/GLTFMeshConverters.cpp
+++ b/Engine/Plugins/Enterprise/GLTFExporter/MigratedFromGit/Source/GLTFExporter/Private/Converters/GLTFMeshConverters.cpp
@@ -32,10 +32,14 @@ void FGLTFStaticMeshConverter::Sanitize(const UStaticMesh*& StaticMesh, const US
 
 	if (StaticMeshComponent != nullptr)
 	{
-		// Only use the component if it's needed for baking, since we would
+		const bool bUsesMeshData = Builder.ExportOptions->BakeMaterialInputs == EGLTFMaterialBakeMode::UseMeshData &&
+			FGLTFMaterialUtility::NeedsMeshData(Materials); // TODO: if this expensive, cache the results for each material
+
+		const bool bIsReferencedByVariant = Builder.VariantReferenceChecker.IsReferenced(StaticMeshComponent); 
+
+		// Only use the component if it's needed for baking or variants, since we would
 		// otherwise export a copy of this mesh for each mesh-component.
-		if (Builder.ExportOptions->BakeMaterialInputs != EGLTFMaterialBakeMode::UseMeshData ||
-			!FGLTFMaterialUtility::NeedsMeshData(Materials)) // TODO: if this expensive, cache the results for each material
+		if (!bUsesMeshData && !bIsReferencedByVariant)
 		{
 			StaticMeshComponent = nullptr;
 		}
@@ -84,10 +88,14 @@ void FGLTFSkeletalMeshConverter::Sanitize(const USkeletalMesh*& SkeletalMesh, co
 
 	if (SkeletalMeshComponent != nullptr)
 	{
-		// Only use the component if it's needed for baking, since we would
+		const bool bUsesMeshData = Builder.ExportOptions->BakeMaterialInputs == EGLTFMaterialBakeMode::UseMeshData &&
+			FGLTFMaterialUtility::NeedsMeshData(Materials); // TODO: if this expensive, cache the results for each material
+
+		const bool bIsReferencedByVariant = Builder.VariantReferenceChecker.IsReferenced(SkeletalMeshComponent); 
+
+		// Only use the component if it's needed for baking or variants, since we would
 		// otherwise export a copy of this mesh for each mesh-component.
-		if (Builder.ExportOptions->BakeMaterialInputs != EGLTFMaterialBakeMode::UseMeshData ||
-			!FGLTFMaterialUtility::NeedsMeshData(Materials)) // TODO: if this expensive, cache the results for each material
+		if (!bUsesMeshData && !bIsReferencedByVariant)
 		{
 			SkeletalMeshComponent = nullptr;
 		}
