[GLTFExporter] chore: add TODO for informing the user about duplicated material variant names

commit 92438de5042c7d2da1143b2ffe8e6180583574d5
Date: Tue, 21 Sep 2021 08:50:04 +0200

    chore: add TODO for informing the user about duplicated material variant names
    (cherry picked from commit 188c2f50778bbdda4ef609fcc8a8702a42345d7f​)

#skipci

diff --git a/Engine/Plugins/Enterprise/GLTFExporter/MigratedFromGit/Source/GLTFExporter/Private/Converters/GLTFKhrVariantConverters.cpp b/Engine/Plugins/Enterprise/GLTFExporter/MigratedFromGit/Source/GLTFExporter/Private/Converters/GLTFKhrVariantConverters.cpp
--- a/Engine/Plugins/Enterprise/GLTFExporter/MigratedFromGit/Source/GLTFExporter/Private/Converters/GLTFKhrVariantConverters.cpp
+++ b/Engine/Plugins/Enterprise/GLTFExporter/MigratedFromGit/Source/GLTFExporter/Private/Converters/GLTFKhrVariantConverters.cpp
@@ -20,6 +20,10 @@ FGLTFJsonKhrMaterialVariantIndex FGLTFKhrMaterialVariantConverter::Convert(const
 	}
 
 	FGLTFJsonKhrMaterialVariant MaterialVariant;
+
+	// TODO: add warning if the variant name is not unique, i.e it's already used?
+	// While material variants are technically allowed to use the same name, it may
+	// cause confusion when trying to select the correct variant in a viewer.
 	MaterialVariant.Name = Variant->GetDisplayText().ToString();
 
 	typedef TTuple<FGLTFJsonPrimitive*, FGLTFJsonMaterialIndex> TPrimitiveMaterial;
