[GLTFExporter] refactor: remove use of "Engine.h" in GLTFKhrVariantConverters to avoid warning regarding monolithic header

commit 16c3188abb1de13fded905d57d38cc0f92d9aeeb
Date: Sat, 1 Jan 2022 21:08:56 +0100

    refactor: remove use of "Engine.h" in GLTFKhrVariantConverters to avoid warning regarding monolithic header
    (cherry picked from commit 334f72edfd5c7e485d1cf8c1ada84f592d93e5d3​)

#skipci

diff --git a/Engine/Plugins/Enterprise/GLTFExporter/MigratedFromGit/Source/GLTFExporter/Private/Converters/GLTFKhrVariantConverters.h b/Engine/Plugins/Enterprise/GLTFExporter/MigratedFromGit/Source/GLTFExporter/Private/Converters/GLTFKhrVariantConverters.h
--- a/Engine/Plugins/Enterprise/GLTFExporter/MigratedFromGit/Source/GLTFExporter/Private/Converters/GLTFKhrVariantConverters.h
+++ b/Engine/Plugins/Enterprise/GLTFExporter/MigratedFromGit/Source/GLTFExporter/Private/Converters/GLTFKhrVariantConverters.h
@@ -6,7 +6,6 @@
 #include "Json/GLTFJsonMesh.h"
 #include "Converters/GLTFConverter.h"
 #include "Converters/GLTFBuilderContext.h"
-#include "Engine.h"
 #include "Variant.h"
 #include "PropertyValue.h"
 #include "LevelVariantSets.h"
