[GLTFExporter] refactor: use GLTFVariantUtility in GLTFKhrVariantConverters

commit 42c00133c5747497a08ac07fc70e4d2f24c58945
Date: Sat, 1 Jan 2022 22:44:16 +0100

    refactor: use GLTFVariantUtility in GLTFKhrVariantConverters
    (cherry picked from commit 10ec43467d4b7f2b6105792df683a7836882f4a7​)

#skipci

diff --git a/Engine/Plugins/Enterprise/GLTFExporter/MigratedFromGit/Source/GLTFExporter/Private/Converters/GLTFKhrVariantConverters.cpp b/Engine/Plugins/Enterprise/GLTFExporter/MigratedFromGit/Source/GLTFExporter/Private/Converters/GLTFKhrVariantConverters.cpp
--- a/Engine/Plugins/Enterprise/GLTFExporter/MigratedFromGit/Source/GLTFExporter/Private/Converters/GLTFKhrVariantConverters.cpp
+++ b/Engine/Plugins/Enterprise/GLTFExporter/MigratedFromGit/Source/GLTFExporter/Private/Converters/GLTFKhrVariantConverters.cpp
@@ -7,9 +7,7 @@
 #include "Rendering/SkeletalMeshRenderData.h"
 #include "VariantObjectBinding.h"
 #include "PropertyValueMaterial.h"
-#include "LevelVariantSets.h"
 #include "PropertyValue.h"
-#include "VariantSet.h"
 #include "Variant.h"
 
 FGLTFJsonKhrMaterialVariantIndex FGLTFKhrMaterialVariantConverter::Convert(const UVariant* Variant)
@@ -95,7 +93,7 @@ bool FGLTFKhrMaterialVariantConverter::TryParseMaterialProperty(FGLTFJsonPrimiti
 	{
 		Builder.LogWarning(FString::Printf(
 			TEXT("Material property is invalid, the property will be skipped. Context: %s"),
-			*GetLogContext(Property)));
+			*FGLTFVariantUtility::GetLogContext(Property)));
 		return false;
 	}
 
@@ -104,7 +102,7 @@ bool FGLTFKhrMaterialVariantConverter::TryParseMaterialProperty(FGLTFJsonPrimiti
 	{
 		Builder.LogWarning(FString::Printf(
 			TEXT("Target object for property is invalid, the property will be skipped. Context: %s"),
-			*GetLogContext(Property)));
+			*FGLTFVariantUtility::GetLogContext(Property)));
 		return false;
 	}
 
@@ -112,7 +110,7 @@ bool FGLTFKhrMaterialVariantConverter::TryParseMaterialProperty(FGLTFJsonPrimiti
 	if (Owner == nullptr)
 	{
 		Builder.LogWarning(FString::Printf(TEXT("Invalid scene component, the property will be skipped. Context: %s"),
-			*GetLogContext(Property)));
+			*FGLTFVariantUtility::GetLogContext(Property)));
 		return false;
 	}
 
@@ -120,7 +118,7 @@ bool FGLTFKhrMaterialVariantConverter::TryParseMaterialProperty(FGLTFJsonPrimiti
 	{
 		Builder.LogWarning(FString::Printf(
 			TEXT("Target actor for property is not selected for export, the property will be skipped. Context: %s"),
-			*GetLogContext(Property)));
+			*FGLTFVariantUtility::GetLogContext(Property)));
 		return false;
 	}
 
@@ -129,7 +127,7 @@ bool FGLTFKhrMaterialVariantConverter::TryParseMaterialProperty(FGLTFJsonPrimiti
 	{
 		Builder.LogWarning(FString::Printf(
 			TEXT("Target object for property has no mesh-component, the property will be skipped. Context: %s"),
-			*GetLogContext(Property)));
+			*FGLTFVariantUtility::GetLogContext(Property)));
 		return false;
 	}
 
@@ -140,7 +138,7 @@ bool FGLTFKhrMaterialVariantConverter::TryParseMaterialProperty(FGLTFJsonPrimiti
 	{
 		Builder.LogWarning(FString::Printf(
 			TEXT("Failed to parse element index to apply the material to, the property will be skipped. Context: %s"),
-			*GetLogContext(MaterialProperty)));
+			*FGLTFVariantUtility::GetLogContext(MaterialProperty)));
 		return false;
 	}
 
@@ -155,7 +153,7 @@ bool FGLTFKhrMaterialVariantConverter::TryParseMaterialProperty(FGLTFJsonPrimiti
 
 		Builder.LogWarning(FString::Printf(
 			TEXT("No material assigned, the property will be skipped. Context: %s"),
-			*GetLogContext(MaterialProperty)));
+			*FGLTFVariantUtility::GetLogContext(MaterialProperty)));
 		return false;
 	}
 
@@ -210,32 +208,3 @@ FGLTFJsonMaterialIndex FGLTFKhrMaterialVariantConverter::GetResolvedMaterialForC
 
 	return Builder.GetOrAddMaterial(Material, MeshData, SectionIndices);
 }
-
-FString FGLTFKhrMaterialVariantConverter::GetLogContext(const UPropertyValue* Property) const
-{
-	const UVariantObjectBinding* Parent = Property->GetParent();
-	return GetLogContext(Parent) + TEXT("/") + Property->GetFullDisplayString();
-}
-
-FString FGLTFKhrMaterialVariantConverter::GetLogContext(const UVariantObjectBinding* Binding) const
-{
-	const UVariant* Parent = const_cast<UVariantObjectBinding*>(Binding)->GetParent();
-	return GetLogContext(Parent) + TEXT("/") + Binding->GetDisplayText().ToString();
-}
-
-FString FGLTFKhrMaterialVariantConverter::GetLogContext(const UVariant* Variant) const
-{
-	const UVariantSet* Parent = const_cast<UVariant*>(Variant)->GetParent();
-	return GetLogContext(Parent) + TEXT("/") + Variant->GetDisplayText().ToString();
-}
-
-FString FGLTFKhrMaterialVariantConverter::GetLogContext(const UVariantSet* VariantSet) const
-{
-	const ULevelVariantSets* Parent = const_cast<UVariantSet*>(VariantSet)->GetParent();
-	return GetLogContext(Parent) + TEXT("/") + VariantSet->GetDisplayText().ToString();
-}
-
-FString FGLTFKhrMaterialVariantConverter::GetLogContext(const ULevelVariantSets* LevelVariantSets) const
-{
-	return LevelVariantSets->GetName();
-}
diff --git a/Engine/Plugins/Enterprise/GLTFExporter/MigratedFromGit/Source/GLTFExporter/Private/Converters/GLTFKhrVariantConverters.h b/Engine/Plugins/Enterprise/GLTFExporter/MigratedFromGit/Source/GLTFExporter/Private/Converters/GLTFKhrVariantConverters.h
--- a/Engine/Plugins/Enterprise/GLTFExporter/MigratedFromGit/Source/GLTFExporter/Private/Converters/GLTFKhrVariantConverters.h
+++ b/Engine/Plugins/Enterprise/GLTFExporter/MigratedFromGit/Source/GLTFExporter/Private/Converters/GLTFKhrVariantConverters.h
@@ -18,10 +18,4 @@ class FGLTFKhrMaterialVariantConverter final : public FGLTFBuilderContext, publi
 
 	bool TryParseMaterialProperty(FGLTFJsonPrimitive*& OutPrimitive, FGLTFJsonMaterialIndex& OutMaterialIndex, const UPropertyValue* Property) const;
 	FGLTFJsonMaterialIndex GetResolvedMaterialForComponent(const UMaterialInterface* Material, int32 MaterialIndex, const UMeshComponent* MeshComponent) const;
-
-	FString GetLogContext(const UPropertyValue* Property) const;
-	FString GetLogContext(const UVariantObjectBinding* Binding) const;
-	FString GetLogContext(const UVariant* Variant) const;
-	FString GetLogContext(const UVariantSet* VariantSet) const;
-	FString GetLogContext(const ULevelVariantSets* LevelVariantSets) const;
 };
