[GLTFExporter] refactor: use texture compression settings (instead of pixel format) to detect HDR, improving maintainability

commit 28621129e7555b4d0c9d98a96b65656eebb373aa
Date: Fri, 7 Jan 2022 07:09:43 +0100

    refactor: use texture compression settings (instead of pixel format) to detect HDR, improving maintainability
    (cherry picked from commit 0652c5f54cc1916e4eeb2b514d8872253a367e44​)

#skipci

diff --git a/Engine/Plugins/Enterprise/GLTFExporter/MigratedFromGit/Source/GLTFExporter/Private/Converters/GLTFTextureUtility.cpp b/Engine/Plugins/Enterprise/GLTFExporter/MigratedFromGit/Source/GLTFExporter/Private/Converters/GLTFTextureUtility.cpp
--- a/Engine/Plugins/Enterprise/GLTFExporter/MigratedFromGit/Source/GLTFExporter/Private/Converters/GLTFTextureUtility.cpp
+++ b/Engine/Plugins/Enterprise/GLTFExporter/MigratedFromGit/Source/GLTFExporter/Private/Converters/GLTFTextureUtility.cpp
@@ -10,30 +10,6 @@
 #include "DeviceProfiles/DeviceProfile.h"
 #include "DeviceProfiles/DeviceProfileManager.h"
 
-bool FGLTFTextureUtility::IsHDR(EPixelFormat Format)
-{
-	switch (Format)
-	{
-		case PF_A32B32G32R32F:
-		case PF_FloatRGB:
-		case PF_FloatRGBA:
-		case PF_R32_FLOAT:
-		case PF_G16R16F:
-		case PF_G16R16F_FILTER:
-		case PF_G32R32F:
-		case PF_R16F:
-		case PF_R16F_FILTER:
-		case PF_FloatR11G11B10:
-		case PF_BC6H:
-		case PF_PLATFORM_HDR_0:
-		case PF_PLATFORM_HDR_1:
-		case PF_PLATFORM_HDR_2:
-			return true;
-		default:
-			return false;
-	}
-}
-
 bool FGLTFTextureUtility::IsAlphaless(EPixelFormat PixelFormat)
 {
 	switch (PixelFormat)
@@ -54,6 +30,19 @@ bool FGLTFTextureUtility::IsAlphaless(EPixelFormat PixelFormat)
 	}
 }
 
+bool FGLTFTextureUtility::IsHDR(const UTexture* Texture)
+{
+	switch (Texture->CompressionSettings)
+	{
+		case TC_HDR:
+		case TC_HDR_Compressed:
+		case TC_HalfFloat:
+			return true;
+		default:
+			return false;
+	}
+}
+
 bool FGLTFTextureUtility::IsCubemap(const UTexture* Texture)
 {
 	return Texture->IsA<UTextureCube>() || Texture->IsA<UTextureRenderTargetCube>();
@@ -258,7 +247,7 @@ UTexture2D* FGLTFTextureUtility::CreateTextureFromCubeFace(const UTextureRenderT
 	const FIntPoint Size(RenderTargetCube->SizeX, RenderTargetCube->SizeX);
 	UTexture2D* FaceTexture;
 
-	if (IsHDR(RenderTargetCube->GetFormat()))
+	if (IsHDR(RenderTargetCube))
 	{
 		TArray<FFloat16Color> Pixels;
 		if (!Resource->ReadPixels(Pixels, FReadSurfaceDataFlags(RCM_UNorm, CubeFace)))
diff --git a/Engine/Plugins/Enterprise/GLTFExporter/MigratedFromGit/Source/GLTFExporter/Private/Converters/GLTFTextureUtility.h b/Engine/Plugins/Enterprise/GLTFExporter/MigratedFromGit/Source/GLTFExporter/Private/Converters/GLTFTextureUtility.h
--- a/Engine/Plugins/Enterprise/GLTFExporter/MigratedFromGit/Source/GLTFExporter/Private/Converters/GLTFTextureUtility.h
+++ b/Engine/Plugins/Enterprise/GLTFExporter/MigratedFromGit/Source/GLTFExporter/Private/Converters/GLTFTextureUtility.h
@@ -6,9 +6,9 @@
 
 struct FGLTFTextureUtility
 {
-	static bool IsHDR(EPixelFormat Format);
 	static bool IsAlphaless(EPixelFormat PixelFormat);
 
+	static bool IsHDR(const UTexture* Texture);
 	static bool IsCubemap(const UTexture* Texture);
 
 	static float GetCubeFaceRotation(ECubeFace CubeFace);
diff --git a/Engine/Plugins/Enterprise/GLTFExporter/MigratedFromGit/Source/GLTFExporter/Private/Tasks/GLTFTextureTasks.cpp b/Engine/Plugins/Enterprise/GLTFExporter/MigratedFromGit/Source/GLTFExporter/Private/Tasks/GLTFTextureTasks.cpp
--- a/Engine/Plugins/Enterprise/GLTFExporter/MigratedFromGit/Source/GLTFExporter/Private/Tasks/GLTFTextureTasks.cpp
+++ b/Engine/Plugins/Enterprise/GLTFExporter/MigratedFromGit/Source/GLTFExporter/Private/Tasks/GLTFTextureTasks.cpp
@@ -27,14 +27,14 @@ void FGLTFTexture2DTask::Complete()
 		JsonTexture.Name += bToSRGB ? TEXT("_sRGB") : TEXT("_Linear");
 	}
 
-	const bool bIsHDR = FGLTFTextureUtility::IsHDR(Texture2D->GetPixelFormat());
+	const bool bIsHDR = FGLTFTextureUtility::IsHDR(Texture2D);
 	const FIntPoint Size = FGLTFTextureUtility::GetInGameSize(Texture2D);
 	UTextureRenderTarget2D* RenderTarget = FGLTFTextureUtility::CreateRenderTarget(Size, bIsHDR);
 
 	// TODO: preserve maximum image quality (avoid compression artifacts) by copying source data (and adjustments) to a temp texture
 	FGLTFTextureUtility::DrawTexture(RenderTarget, Texture2D, FVector2D::ZeroVector, Size);
 
-	if (!Texture2D->IsNormalMap() && bIsHDR)
+	if (bIsHDR)
 	{
 		JsonTexture.Encoding = Builder.GetTextureHDREncoding();
 	}
@@ -88,7 +88,7 @@ void FGLTFTextureCubeTask::Complete()
 		return;
 	}
 
-	const bool bIsHDR = FGLTFTextureUtility::IsHDR(TextureCube->GetPixelFormat());
+	const bool bIsHDR = FGLTFTextureUtility::IsHDR(TextureCube);
 	const FIntPoint Size = { TextureCube->GetSizeX(), TextureCube->GetSizeY() };
 	UTextureRenderTarget2D* RenderTarget = FGLTFTextureUtility::CreateRenderTarget(Size, bIsHDR);
 
@@ -132,7 +132,7 @@ void FGLTFTextureRenderTarget2DTask::Complete()
 		JsonTexture.Name += bToSRGB ? TEXT("_sRGB") : TEXT("_Linear");
 	}
 
-	const bool bIsHDR = FGLTFTextureUtility::IsHDR(RenderTarget2D->GetFormat());
+	const bool bIsHDR = FGLTFTextureUtility::IsHDR(RenderTarget2D);
 	const FIntPoint Size = { RenderTarget2D->SizeX, RenderTarget2D->SizeY };
 
 	if (bIsHDR)
@@ -181,7 +181,7 @@ void FGLTFTextureRenderTargetCubeTask::Complete()
 		return;
 	}
 
-	const bool bIsHDR = FGLTFTextureUtility::IsHDR(RenderTargetCube->GetFormat());
+	const bool bIsHDR = FGLTFTextureUtility::IsHDR(RenderTargetCube);
 	const FIntPoint Size = { RenderTargetCube->SizeX, RenderTargetCube->SizeX };
 	UTextureRenderTarget2D* RenderTarget = FGLTFTextureUtility::CreateRenderTarget(Size, bIsHDR);
 
