[GLTFExporter] refactor: add TODO comments regarding texture sampler type and source

commit 983d6eb49616f187f350f3fc5964789e62b16662
Date: Fri, 7 Jan 2022 07:10:04 +0100

    refactor: add TODO comments regarding texture sampler type and source
    (cherry picked from commit cd9fc4a902d212882c06149242ca523cc650a4f7​)

#skipci

diff --git a/Engine/Plugins/Enterprise/GLTFExporter/MigratedFromGit/Source/GLTFExporter/Private/Tasks/GLTFMaterialTasks.cpp b/Engine/Plugins/Enterprise/GLTFExporter/MigratedFromGit/Source/GLTFExporter/Private/Tasks/GLTFMaterialTasks.cpp
--- a/Engine/Plugins/Enterprise/GLTFExporter/MigratedFromGit/Source/GLTFExporter/Private/Tasks/GLTFMaterialTasks.cpp
+++ b/Engine/Plugins/Enterprise/GLTFExporter/MigratedFromGit/Source/GLTFExporter/Private/Tasks/GLTFMaterialTasks.cpp
@@ -1062,6 +1062,8 @@ bool FGLTFMaterialTask::TryGetSourceTexture(const UTexture2D*& OutTexture, int32
 
 	if (const UMaterialExpressionTextureSampleParameter2D* TextureParameter = ExactCast<UMaterialExpressionTextureSampleParameter2D>(Expression))
 	{
+		// TODO: handle non-default TextureParameter->SamplerType and TextureParameter->SamplerSource
+
 		UTexture* ParameterValue = TextureParameter->Texture;
 
 		const UMaterialInstance* MaterialInstance = Cast<UMaterialInstance>(Material);
@@ -1110,6 +1112,8 @@ bool FGLTFMaterialTask::TryGetSourceTexture(const UTexture2D*& OutTexture, int32
 
 	if (const UMaterialExpressionTextureSample* TextureSampler = ExactCast<UMaterialExpressionTextureSample>(Expression))
 	{
+		// TODO: handle non-default TextureSampler->SamplerType and TextureSampler->SamplerSource
+
 		// TODO: add support for texture object input expression
 		OutTexture = Cast<UTexture2D>(TextureSampler->Texture);
 
