[GLTFExporter] fix: override sRGB color space transformation if texture is to be HDR-encoded

commit 758f03381ce714951223cac8e12ecf64f58e77ee
Date: Fri, 7 Jan 2022 07:09:55 +0100

    fix: override sRGB color space transformation if texture is to be HDR-encoded
    (cherry picked from commit dffc91d0a243b895194632a37481eed081603b1a​)

#skipci

diff --git a/Engine/Plugins/Enterprise/GLTFExporter/MigratedFromGit/Source/GLTFExporter/Private/Converters/GLTFTextureConverters.cpp b/Engine/Plugins/Enterprise/GLTFExporter/MigratedFromGit/Source/GLTFExporter/Private/Converters/GLTFTextureConverters.cpp
--- a/Engine/Plugins/Enterprise/GLTFExporter/MigratedFromGit/Source/GLTFExporter/Private/Converters/GLTFTextureConverters.cpp
+++ b/Engine/Plugins/Enterprise/GLTFExporter/MigratedFromGit/Source/GLTFExporter/Private/Converters/GLTFTextureConverters.cpp
@@ -1,7 +1,20 @@
 // Copyright Epic Games, Inc. All Rights Reserved.
 
 #include "Converters/GLTFTextureConverters.h"
+#include "Converters/GLTFTextureUtility.h"
 #include "Tasks/GLTFTextureTasks.h"
+#include "Engine/Texture2D.h"
+#include "Engine/TextureCube.h"
+#include "Engine/TextureRenderTarget2D.h"
+#include "Engine/TextureRenderTargetCube.h"
+
+void FGLTFTexture2DConverter::Sanitize(const UTexture2D*& Texture2D, bool& bToSRGB)
+{
+	if (FGLTFTextureUtility::IsHDR(Texture2D) && Builder.GetTextureHDREncoding() != EGLTFJsonHDREncoding::None)
+	{
+		bToSRGB = false;
+	}
+}
 
 FGLTFJsonTextureIndex FGLTFTexture2DConverter::Convert(const UTexture2D* Texture2D, bool bToSRGB)
 {
@@ -15,6 +28,14 @@ FGLTFJsonTextureIndex FGLTFTexture2DConverter::Convert(const UTexture2D* Texture
 	return FGLTFJsonTextureIndex(INDEX_NONE);
 }
 
+void FGLTFTextureCubeConverter::Sanitize(const UTextureCube*& TextureCube, ECubeFace& CubeFace, bool& bToSRGB)
+{
+	if (FGLTFTextureUtility::IsHDR(TextureCube) && Builder.GetTextureHDREncoding() != EGLTFJsonHDREncoding::None)
+	{
+		bToSRGB = false;
+	}
+}
+
 FGLTFJsonTextureIndex FGLTFTextureCubeConverter::Convert(const UTextureCube* TextureCube, ECubeFace CubeFace, bool bToSRGB)
 {
 	if (Builder.ExportOptions->TextureImageFormat != EGLTFTextureImageFormat::None)
@@ -27,6 +48,14 @@ FGLTFJsonTextureIndex FGLTFTextureCubeConverter::Convert(const UTextureCube* Tex
 	return FGLTFJsonTextureIndex(INDEX_NONE);
 }
 
+void FGLTFTextureRenderTarget2DConverter::Sanitize(const UTextureRenderTarget2D*& RenderTarget2D, bool& bToSRGB)
+{
+	if (FGLTFTextureUtility::IsHDR(RenderTarget2D) && Builder.GetTextureHDREncoding() != EGLTFJsonHDREncoding::None)
+	{
+		bToSRGB = false;
+	}
+}
+
 FGLTFJsonTextureIndex FGLTFTextureRenderTarget2DConverter::Convert(const UTextureRenderTarget2D* RenderTarget2D, bool bToSRGB)
 {
 	if (Builder.ExportOptions->TextureImageFormat != EGLTFTextureImageFormat::None)
@@ -39,6 +68,14 @@ FGLTFJsonTextureIndex FGLTFTextureRenderTarget2DConverter::Convert(const UTextur
 	return FGLTFJsonTextureIndex(INDEX_NONE);
 }
 
+void FGLTFTextureRenderTargetCubeConverter::Sanitize(const UTextureRenderTargetCube*& RenderTargetCube, ECubeFace& CubeFace, bool& bToSRGB)
+{
+	if (FGLTFTextureUtility::IsHDR(RenderTargetCube) && Builder.GetTextureHDREncoding() != EGLTFJsonHDREncoding::None)
+	{
+		bToSRGB = false;
+	}
+}
+
 FGLTFJsonTextureIndex FGLTFTextureRenderTargetCubeConverter::Convert(const UTextureRenderTargetCube* RenderTargetCube, ECubeFace CubeFace, bool bToSRGB)
 {
 	if (Builder.ExportOptions->TextureImageFormat != EGLTFTextureImageFormat::None)
diff --git a/Engine/Plugins/Enterprise/GLTFExporter/MigratedFromGit/Source/GLTFExporter/Private/Converters/GLTFTextureConverters.h b/Engine/Plugins/Enterprise/GLTFExporter/MigratedFromGit/Source/GLTFExporter/Private/Converters/GLTFTextureConverters.h
--- a/Engine/Plugins/Enterprise/GLTFExporter/MigratedFromGit/Source/GLTFExporter/Private/Converters/GLTFTextureConverters.h
+++ b/Engine/Plugins/Enterprise/GLTFExporter/MigratedFromGit/Source/GLTFExporter/Private/Converters/GLTFTextureConverters.h
@@ -16,6 +16,8 @@ class FGLTFTexture2DConverter final : public TGLTFTextureConverter<const UTextur
 {
 	using TGLTFTextureConverter::TGLTFTextureConverter;
 
+	virtual void Sanitize(const UTexture2D*& Texture2D, bool& bToSRGB) override;
+
 	virtual FGLTFJsonTextureIndex Convert(const UTexture2D* Texture2D, bool bToSRGB) override;
 };
 
@@ -23,6 +25,8 @@ class FGLTFTextureCubeConverter final : public TGLTFTextureConverter<const UText
 {
 	using TGLTFTextureConverter::TGLTFTextureConverter;
 
+	virtual void Sanitize(const UTextureCube*& TextureCube, ECubeFace& CubeFace, bool& bToSRGB) override;
+
 	virtual FGLTFJsonTextureIndex Convert(const UTextureCube* TextureCube, ECubeFace CubeFace, bool bToSRGB) override;
 };
 
@@ -30,6 +34,8 @@ class FGLTFTextureRenderTarget2DConverter final : public TGLTFTextureConverter<c
 {
 	using TGLTFTextureConverter::TGLTFTextureConverter;
 
+	virtual void Sanitize(const UTextureRenderTarget2D*& RenderTarget2D, bool& bToSRGB) override;
+
 	virtual FGLTFJsonTextureIndex Convert(const UTextureRenderTarget2D* RenderTarget2D, bool bToSRGB) override;
 };
 
@@ -37,6 +43,8 @@ class FGLTFTextureRenderTargetCubeConverter final : public TGLTFTextureConverter
 {
 	using TGLTFTextureConverter::TGLTFTextureConverter;
 
+	virtual void Sanitize(const UTextureRenderTargetCube*& RenderTargetCube, ECubeFace& CubeFace, bool& bToSRGB) override;
+
 	virtual FGLTFJsonTextureIndex Convert(const UTextureRenderTargetCube* RenderTargetCube, ECubeFace CubeFace, bool bToSRGB) override;
 };
 
