[GLTFExporter] refactor: revert removal of PluginManager in runtime

commit 063436b32fcb90a664c430e2aa4aa8161022c461
Date: Sat, 8 Jan 2022 06:12:54 +0100

    refactor: revert removal of PluginManager in runtime
    (cherry picked from commit c67a79610352a06c80945cb716fbd9b9ae5adb58​)

#skipci

diff --git a/Engine/Plugins/Enterprise/GLTFExporter/MigratedFromGit/Source/GLTFExporter/GLTFExporter.Build.cs b/Engine/Plugins/Enterprise/GLTFExporter/MigratedFromGit/Source/GLTFExporter/GLTFExporter.Build.cs
--- a/Engine/Plugins/Enterprise/GLTFExporter/MigratedFromGit/Source/GLTFExporter/GLTFExporter.Build.cs
+++ b/Engine/Plugins/Enterprise/GLTFExporter/MigratedFromGit/Source/GLTFExporter/GLTFExporter.Build.cs
@@ -9,30 +9,6 @@ public class GLTFExporter : ModuleRules
 		PCHUsage = PCHUsageMode.UseExplicitOrSharedPCHs;
 		bTreatAsEngineModule = true; // Only necessary when plugin installed in project
 
-		{
-			// NOTE: ugly hack to access plugin info (should propose change to engine)
-			var PluginField = typeof(ModuleRules).GetField("Plugin", System.Reflection.BindingFlags.Instance | System.Reflection.BindingFlags.NonPublic);
-			if (PluginField == null)
-			{
-				throw new System.Exception("Missing internal member Plugin");
-			}
-
-			var Plugin = PluginField.GetValue(this) as PluginInfo;
-			if (Plugin == null)
-			{
-				throw new System.Exception("Missing plugin information");
-			}
-
-			var PluginDescriptor = Plugin.Descriptor;
-			if (PluginDescriptor == null)
-			{
-				throw new System.Exception("Missing plugin descriptor");
-			}
-
-			PrivateDefinitions.Add("GLTFEXPORTER_FRIENDLY_NAME=TEXT(\"" + PluginDescriptor.FriendlyName + "\")");
-			PrivateDefinitions.Add("GLTFEXPORTER_VERSION_NAME=TEXT(\"" + PluginDescriptor.VersionName + "\")");
-		}
-
 		PublicDependencyModuleNames .AddRange(
 			new string[]
 			{
@@ -53,6 +29,7 @@ public class GLTFExporter : ModuleRules
 				"MovieScene",
 				"MovieSceneTracks",
 				"VariantManagerContent",
+				"Projects",
 			}
 		);
 
@@ -69,7 +46,6 @@ public class GLTFExporter : ModuleRules
 					"InputCore",
 					"EditorStyle",
 					"PropertyEditor",
-					"Projects",
 					"MeshMergeUtilities",
 					"MeshDescription",
 					"StaticMeshDescription",
diff --git a/Engine/Plugins/Enterprise/GLTFExporter/MigratedFromGit/Source/GLTFExporter/Private/Builders/GLTFJsonBuilder.cpp b/Engine/Plugins/Enterprise/GLTFExporter/MigratedFromGit/Source/GLTFExporter/Private/Builders/GLTFJsonBuilder.cpp
--- a/Engine/Plugins/Enterprise/GLTFExporter/MigratedFromGit/Source/GLTFExporter/Private/Builders/GLTFJsonBuilder.cpp
+++ b/Engine/Plugins/Enterprise/GLTFExporter/MigratedFromGit/Source/GLTFExporter/Private/Builders/GLTFJsonBuilder.cpp
@@ -1,6 +1,8 @@
 // Copyright Epic Games, Inc. All Rights Reserved.
 
 #include "Builders/GLTFJsonBuilder.h"
+#include "GLTFExporterModule.h"
+#include "Interfaces/IPluginManager.h"
 #include "Runtime/Launch/Resources/Version.h"
 
 FGLTFJsonBuilder::FGLTFJsonBuilder(const FString& FilePath, const UGLTFExportOptions* ExportOptions)
@@ -281,9 +283,12 @@ FGLTFJsonNodeIndex FGLTFJsonBuilder::GetComponentNodeIndex(FGLTFJsonNodeIndex No
 
 FString FGLTFJsonBuilder::GetGeneratorString() const
 {
+	const TSharedPtr<IPlugin> Plugin = IPluginManager::Get().FindPlugin(GLTFEXPORTER_MODULE_NAME);
+	const FPluginDescriptor& PluginDescriptor = Plugin->GetDescriptor();
+
 	return ExportOptions->bIncludeGeneratorVersion
-		? TEXT(EPIC_PRODUCT_NAME) TEXT(" ") ENGINE_VERSION_STRING TEXT(" ") GLTFEXPORTER_FRIENDLY_NAME TEXT(" ") GLTFEXPORTER_VERSION_NAME
-		: TEXT(EPIC_PRODUCT_NAME) TEXT(" ") GLTFEXPORTER_FRIENDLY_NAME;
+		? TEXT(EPIC_PRODUCT_NAME) TEXT(" ") ENGINE_VERSION_STRING TEXT(" ") + PluginDescriptor.FriendlyName + TEXT(" ") + PluginDescriptor.VersionName
+		: TEXT(EPIC_PRODUCT_NAME) TEXT(" ") + PluginDescriptor.FriendlyName;
 }
 
 bool FGLTFJsonBuilder::IsCustomExtension(EGLTFJsonExtension Extension)
diff --git a/Engine/Plugins/Enterprise/GLTFExporter/MigratedFromGit/Source/GLTFExporter/Private/Builders/GLTFLogBuilder.cpp b/Engine/Plugins/Enterprise/GLTFExporter/MigratedFromGit/Source/GLTFExporter/Private/Builders/GLTFLogBuilder.cpp
--- a/Engine/Plugins/Enterprise/GLTFExporter/MigratedFromGit/Source/GLTFExporter/Private/Builders/GLTFLogBuilder.cpp
+++ b/Engine/Plugins/Enterprise/GLTFExporter/MigratedFromGit/Source/GLTFExporter/Private/Builders/GLTFLogBuilder.cpp
@@ -4,6 +4,7 @@
 #include "GLTFExporterModule.h"
 #include "Misc/FeedbackContext.h"
 #if WITH_EDITOR
+#include "Interfaces/IPluginManager.h"
 #include "MessageLogModule.h"
 #include "IMessageLogListing.h"
 #endif
@@ -14,9 +15,12 @@ FGLTFLogBuilder::FGLTFLogBuilder(const FString& FilePath, const UGLTFExportOptio
 #if WITH_EDITOR
 	if (!FApp::IsUnattended())
 	{
+		const TSharedPtr<IPlugin> Plugin = IPluginManager::Get().FindPlugin(GLTFEXPORTER_MODULE_NAME);
+		const FString& PluginFriendlyName = Plugin->GetDescriptor().FriendlyName;
+
 		FMessageLogModule& MessageLogModule = FModuleManager::LoadModuleChecked<FMessageLogModule>("MessageLog");
 		LogListing = MessageLogModule.GetLogListing(GLTFEXPORTER_MODULE_NAME);
-		LogListing->SetLabel(FText::FromString(GLTFEXPORTER_FRIENDLY_NAME));
+		LogListing->SetLabel(FText::FromString(PluginFriendlyName));
 	}
 #endif
 }
