[GLTFExporter] refactor: hide StaticMeshDataConverter and SkeletalMeshDataConverter to ease refactoring GLTFConvertBuilder

commit eeaab5f3943e00a17711c26d66d9af537c172c9e
Date: Mon, 10 Jan 2022 08:11:59 +0100

    refactor: hide StaticMeshDataConverter and SkeletalMeshDataConverter to ease refactoring GLTFConvertBuilder
    (cherry picked from commit 0c7a8ba89e53a25dc23df4eb51e8535b601ceb25​)

#skipci

diff --git a/Engine/Plugins/Enterprise/GLTFExporter/MigratedFromGit/Source/GLTFExporter/Private/Builders/GLTFConvertBuilder.cpp b/Engine/Plugins/Enterprise/GLTFExporter/MigratedFromGit/Source/GLTFExporter/Private/Builders/GLTFConvertBuilder.cpp
--- a/Engine/Plugins/Enterprise/GLTFExporter/MigratedFromGit/Source/GLTFExporter/Private/Builders/GLTFConvertBuilder.cpp
+++ b/Engine/Plugins/Enterprise/GLTFExporter/MigratedFromGit/Source/GLTFExporter/Private/Builders/GLTFConvertBuilder.cpp
@@ -1,6 +1,7 @@
 // Copyright Epic Games, Inc. All Rights Reserved.
 
 #include "Builders/GLTFConvertBuilder.h"
+#include "Converters/GLTFMeshUtility.h"
 #include "Engine/Texture2D.h"
 #include "Engine/TextureCube.h"
 #include "Engine/TextureRenderTarget2D.h"
@@ -102,30 +103,45 @@ FGLTFJsonMeshIndex FGLTFConvertBuilder::GetOrAddMesh(const UStaticMesh* StaticMe
 	return StaticMeshConverter.GetOrAdd(StaticMesh, nullptr, Materials, LODIndex);
 }
 
-FGLTFJsonMeshIndex FGLTFConvertBuilder::GetOrAddMesh(const UStaticMeshComponent* StaticMeshComponent, const FGLTFMaterialArray& Materials, int32 LODIndex)
+FGLTFJsonMeshIndex FGLTFConvertBuilder::GetOrAddMesh(const USkeletalMesh* SkeletalMesh, const FGLTFMaterialArray& Materials, int32 LODIndex)
 {
-	if (StaticMeshComponent == nullptr)
+	if (SkeletalMesh == nullptr)
 	{
 		return FGLTFJsonMeshIndex(INDEX_NONE);
 	}
 
-	const UStaticMesh* StaticMesh = StaticMeshComponent->GetStaticMesh();
-	if (StaticMesh == nullptr)
+	return SkeletalMeshConverter.GetOrAdd(SkeletalMesh, nullptr, Materials, LODIndex);
+}
+
+FGLTFJsonMeshIndex FGLTFConvertBuilder::GetOrAddMesh(const UMeshComponent* MeshComponent, const FGLTFMaterialArray& Materials, int32 LODIndex)
+{
+	if (const UStaticMeshComponent* StaticMeshComponent = Cast<UStaticMeshComponent>(MeshComponent))
 	{
-		return FGLTFJsonMeshIndex(INDEX_NONE);
+		return GetOrAddMesh(StaticMeshComponent, Materials, LODIndex);
 	}
 
-	return StaticMeshConverter.GetOrAdd(StaticMesh, StaticMeshComponent, Materials, LODIndex);
+	if (const USkeletalMeshComponent* SkeletalMeshComponent = Cast<USkeletalMeshComponent>(MeshComponent))
+	{
+		return GetOrAddMesh(SkeletalMeshComponent, Materials, LODIndex);
+	}
+
+	return FGLTFJsonMeshIndex(INDEX_NONE);
 }
 
-FGLTFJsonMeshIndex FGLTFConvertBuilder::GetOrAddMesh(const USkeletalMesh* SkeletalMesh, const FGLTFMaterialArray& Materials, int32 LODIndex)
+FGLTFJsonMeshIndex FGLTFConvertBuilder::GetOrAddMesh(const UStaticMeshComponent* StaticMeshComponent, const FGLTFMaterialArray& Materials, int32 LODIndex)
 {
-	if (SkeletalMesh == nullptr)
+	if (StaticMeshComponent == nullptr)
 	{
 		return FGLTFJsonMeshIndex(INDEX_NONE);
 	}
 
-	return SkeletalMeshConverter.GetOrAdd(SkeletalMesh, nullptr, Materials, LODIndex);
+	const UStaticMesh* StaticMesh = StaticMeshComponent->GetStaticMesh();
+	if (StaticMesh == nullptr)
+	{
+		return FGLTFJsonMeshIndex(INDEX_NONE);
+	}
+
+	return StaticMeshConverter.GetOrAdd(StaticMesh, StaticMeshComponent, Materials, LODIndex);
 }
 
 FGLTFJsonMeshIndex FGLTFConvertBuilder::GetOrAddMesh(const USkeletalMeshComponent* SkeletalMeshComponent, const FGLTFMaterialArray& Materials, int32 LODIndex)
@@ -144,20 +160,59 @@ FGLTFJsonMeshIndex FGLTFConvertBuilder::GetOrAddMesh(const USkeletalMeshComponen
 	return SkeletalMeshConverter.GetOrAdd(SkeletalMesh, SkeletalMeshComponent, Materials, LODIndex);
 }
 
-FGLTFJsonMeshIndex FGLTFConvertBuilder::GetOrAddMesh(const UMeshComponent* MeshComponent, const FGLTFMaterialArray& Materials, int32 LODIndex)
+const FGLTFMeshData* FGLTFConvertBuilder::GetOrAddMeshData(const UStaticMesh* StaticMesh, const UStaticMeshComponent* StaticMeshComponent, int32 LODIndex)
+{
+	return StaticMeshDataConverter.GetOrAdd(StaticMesh, StaticMeshComponent, LODIndex);
+}
+
+const FGLTFMeshData* FGLTFConvertBuilder::GetOrAddMeshData(const USkeletalMesh* SkeletalMesh, const USkeletalMeshComponent* SkeletalMeshComponent, int32 LODIndex)
+{
+	return SkeletalMeshDataConverter.GetOrAdd(SkeletalMesh, SkeletalMeshComponent, LODIndex);
+}
+
+FGLTFJsonMaterialIndex FGLTFConvertBuilder::GetOrAddMaterial(const UMaterialInterface* Material, const UStaticMesh* StaticMesh, int32 LODIndex, int32 MaterialIndex)
+{
+	const FGLTFMeshData* MeshData = GetOrAddMeshData(StaticMesh, nullptr, LODIndex);
+	const FGLTFIndexArray SectionIndices = FGLTFMeshUtility::GetSectionIndices(StaticMesh, MeshData->LODIndex, MaterialIndex);
+	return GetOrAddMaterial(Material, MeshData, SectionIndices);
+}
+
+FGLTFJsonMaterialIndex FGLTFConvertBuilder::GetOrAddMaterial(const UMaterialInterface* Material, const USkeletalMesh* SkeletalMesh, int32 LODIndex, int32 MaterialIndex)
+{
+	const FGLTFMeshData* MeshData = GetOrAddMeshData(SkeletalMesh, nullptr, LODIndex);
+	const FGLTFIndexArray SectionIndices = FGLTFMeshUtility::GetSectionIndices(SkeletalMesh, MeshData->LODIndex, MaterialIndex);
+	return GetOrAddMaterial(Material, MeshData, SectionIndices);
+}
+
+FGLTFJsonMaterialIndex FGLTFConvertBuilder::GetOrAddMaterial(const UMaterialInterface* Material, const UMeshComponent* MeshComponent, int32 LODIndex, int32 MaterialIndex)
 {
 	if (const UStaticMeshComponent* StaticMeshComponent = Cast<UStaticMeshComponent>(MeshComponent))
 	{
-		return GetOrAddMesh(StaticMeshComponent, Materials, LODIndex);
+		return GetOrAddMaterial(Material, StaticMeshComponent, LODIndex, MaterialIndex);
 	}
-	else if (const USkeletalMeshComponent* SkeletalMeshComponent = Cast<USkeletalMeshComponent>(MeshComponent))
+
+	if (const USkeletalMeshComponent* SkeletalMeshComponent = Cast<USkeletalMeshComponent>(MeshComponent))
 	{
-		return GetOrAddMesh(SkeletalMeshComponent, Materials, LODIndex);
-	}
-	else
-	{
-		return FGLTFJsonMeshIndex(INDEX_NONE);
+		return GetOrAddMaterial(Material, SkeletalMeshComponent, LODIndex, MaterialIndex);
 	}
+
+	return FGLTFJsonMaterialIndex(INDEX_NONE);
+}
+
+FGLTFJsonMaterialIndex FGLTFConvertBuilder::GetOrAddMaterial(const UMaterialInterface* Material, const UStaticMeshComponent* StaticMeshComponent, int32 LODIndex, int32 MaterialIndex)
+{
+	const UStaticMesh* StaticMesh = StaticMeshComponent->GetStaticMesh();
+	const FGLTFMeshData* MeshData = GetOrAddMeshData(StaticMesh, StaticMeshComponent, LODIndex);
+	const FGLTFIndexArray SectionIndices = FGLTFMeshUtility::GetSectionIndices(StaticMesh, MeshData->LODIndex, MaterialIndex);
+	return GetOrAddMaterial(Material, MeshData, SectionIndices);
+}
+
+FGLTFJsonMaterialIndex FGLTFConvertBuilder::GetOrAddMaterial(const UMaterialInterface* Material, const USkeletalMeshComponent* SkeletalMeshComponent, int32 LODIndex, int32 MaterialIndex)
+{
+	const USkeletalMesh* SkeletalMesh = SkeletalMeshComponent->SkeletalMesh;
+	const FGLTFMeshData* MeshData = GetOrAddMeshData(SkeletalMesh, SkeletalMeshComponent, LODIndex);
+	const FGLTFIndexArray SectionIndices = FGLTFMeshUtility::GetSectionIndices(SkeletalMesh, MeshData->LODIndex, MaterialIndex);
+	return GetOrAddMaterial(Material, MeshData, SectionIndices);
 }
 
 FGLTFJsonMaterialIndex FGLTFConvertBuilder::GetOrAddMaterial(const UMaterialInterface* Material, const FGLTFMeshData* MeshData, const FGLTFIndexArray& SectionIndices)
diff --git a/Engine/Plugins/Enterprise/GLTFExporter/MigratedFromGit/Source/GLTFExporter/Private/Builders/GLTFConvertBuilder.h b/Engine/Plugins/Enterprise/GLTFExporter/MigratedFromGit/Source/GLTFExporter/Private/Builders/GLTFConvertBuilder.h
--- a/Engine/Plugins/Enterprise/GLTFExporter/MigratedFromGit/Source/GLTFExporter/Private/Builders/GLTFConvertBuilder.h
+++ b/Engine/Plugins/Enterprise/GLTFExporter/MigratedFromGit/Source/GLTFExporter/Private/Builders/GLTFConvertBuilder.h
@@ -34,8 +34,6 @@ public:
 	const bool bSelectedActorsOnly;
 
 	// TODO: find a better place for this types of indirect converters
-	FGLTFStaticMeshDataConverter StaticMeshDataConverter;
-	FGLTFSkeletalMeshDataConverter SkeletalMeshDataConverter;
 	FGLTFVariantReferenceChecker VariantReferenceChecker;
 
 	FGLTFJsonAccessorIndex GetOrAddPositionAccessor(const FGLTFMeshSection* MeshSection, const FPositionVertexBuffer* VertexBuffer);
@@ -48,11 +46,19 @@ public:
 	FGLTFJsonAccessorIndex GetOrAddIndexAccessor(const FGLTFMeshSection* MeshSection);
 
 	FGLTFJsonMeshIndex GetOrAddMesh(const UStaticMesh* StaticMesh, const FGLTFMaterialArray& Materials = {}, int32 LODIndex = -1);
-	FGLTFJsonMeshIndex GetOrAddMesh(const UStaticMeshComponent* StaticMeshComponent, const FGLTFMaterialArray& Materials = {}, int32 LODIndex = -1);
 	FGLTFJsonMeshIndex GetOrAddMesh(const USkeletalMesh* SkeletalMesh, const FGLTFMaterialArray& Materials = {}, int32 LODIndex = -1);
-	FGLTFJsonMeshIndex GetOrAddMesh(const USkeletalMeshComponent* SkeletalMeshComponent, const FGLTFMaterialArray& Materials = {}, int32 LODIndex = -1);
 	FGLTFJsonMeshIndex GetOrAddMesh(const UMeshComponent* MeshComponent, const FGLTFMaterialArray& Materials = {}, int32 LODIndex = -1);
+	FGLTFJsonMeshIndex GetOrAddMesh(const UStaticMeshComponent* StaticMeshComponent, const FGLTFMaterialArray& Materials = {}, int32 LODIndex = -1);
+	FGLTFJsonMeshIndex GetOrAddMesh(const USkeletalMeshComponent* SkeletalMeshComponent, const FGLTFMaterialArray& Materials = {}, int32 LODIndex = -1);
+
+	const FGLTFMeshData* GetOrAddMeshData(const UStaticMesh* StaticMesh, const UStaticMeshComponent* StaticMeshComponent = nullptr, int32 LODIndex = -1);
+	const FGLTFMeshData* GetOrAddMeshData(const USkeletalMesh* SkeletalMesh, const USkeletalMeshComponent* SkeletalMeshComponent = nullptr, int32 LODIndex = -1);
 
+	FGLTFJsonMaterialIndex GetOrAddMaterial(const UMaterialInterface* Material, const UStaticMesh* StaticMesh, int32 LODIndex = -1, int32 MaterialIndex = -1);
+	FGLTFJsonMaterialIndex GetOrAddMaterial(const UMaterialInterface* Material, const USkeletalMesh* SkeletalMesh, int32 LODIndex = -1, int32 MaterialIndex = -1);
+	FGLTFJsonMaterialIndex GetOrAddMaterial(const UMaterialInterface* Material, const UMeshComponent* MeshComponent, int32 LODIndex = -1, int32 MaterialIndex = -1);
+	FGLTFJsonMaterialIndex GetOrAddMaterial(const UMaterialInterface* Material, const UStaticMeshComponent* StaticMeshComponent, int32 LODIndex = -1, int32 MaterialIndex = -1);
+	FGLTFJsonMaterialIndex GetOrAddMaterial(const UMaterialInterface* Material, const USkeletalMeshComponent* SkeletalMeshComponent, int32 LODIndex = -1, int32 MaterialIndex = -1);
 	FGLTFJsonMaterialIndex GetOrAddMaterial(const UMaterialInterface* Material, const FGLTFMeshData* MeshData = nullptr, const FGLTFIndexArray& SectionIndices = {});
 
 	FGLTFJsonSamplerIndex GetOrAddSampler(const UTexture* Texture);
@@ -105,6 +111,9 @@ private:
 	FGLTFSkeletalMeshConverter SkeletalMeshConverter{ *this };
 
 	FGLTFMaterialConverter MaterialConverter{ *this };
+	FGLTFStaticMeshDataConverter StaticMeshDataConverter{ *this };
+	FGLTFSkeletalMeshDataConverter SkeletalMeshDataConverter{ *this };
+
 	FGLTFSamplerConverter SamplerConverter{ *this };
 	FGLTFTexture2DConverter Texture2DConverter{ *this };
 	FGLTFTextureCubeConverter TextureCubeConverter{ *this };
diff --git a/Engine/Plugins/Enterprise/GLTFExporter/MigratedFromGit/Source/GLTFExporter/Private/Converters/GLTFEpicLevelVariantSetsConverters.cpp b/Engine/Plugins/Enterprise/GLTFExporter/MigratedFromGit/Source/GLTFExporter/Private/Converters/GLTFEpicLevelVariantSetsConverters.cpp
--- a/Engine/Plugins/Enterprise/GLTFExporter/MigratedFromGit/Source/GLTFExporter/Private/Converters/GLTFEpicLevelVariantSetsConverters.cpp
+++ b/Engine/Plugins/Enterprise/GLTFExporter/MigratedFromGit/Source/GLTFExporter/Private/Converters/GLTFEpicLevelVariantSetsConverters.cpp
@@ -185,7 +185,7 @@ bool FGLTFEpicLevelVariantSetsConverter::TryParseMaterialPropertyValue(FGLTFJson
 		return false;
 	}
 
-	const USceneComponent* Target = static_cast<USceneComponent*>(Property->GetPropertyParentContainerAddress());
+	const UMeshComponent* Target = static_cast<UMeshComponent*>(Property->GetPropertyParentContainerAddress());
 	if (Target == nullptr)
 	{
 		Builder.LogWarning(FString::Printf(
@@ -229,45 +229,10 @@ bool FGLTFEpicLevelVariantSetsConverter::TryParseMaterialPropertyValue(FGLTFJson
 	const UMaterialInterface* Material = MaterialProperty->GetMaterial();
 	const int32 MaterialIndex = CapturedPropSegments[NumPropSegments - 1].PropertyIndex;
 
-	const FGLTFMeshData* MeshData = nullptr;
-	FGLTFIndexArray SectionIndices;
-
-	if (Builder.ExportOptions->ExportMaterialVariants == EGLTFMaterialVariantMode::UseMeshData)
-	{
-		if (Builder.ExportOptions->BakeMaterialInputs == EGLTFMaterialBakeMode::UseMeshData)
-		{
-			const int32 DefaultLODIndex = Builder.ExportOptions->DefaultLevelOfDetail;
-
-			if (const UStaticMeshComponent* StaticMeshComponent = Cast<UStaticMeshComponent>(Target))
-			{
-				const UStaticMesh* StaticMesh = StaticMeshComponent->GetStaticMesh();
-				const int32 LODIndex = FGLTFMeshUtility::GetLOD(StaticMesh, StaticMeshComponent, DefaultLODIndex);
-				const FStaticMeshLODResources& MeshLOD = StaticMesh->GetLODForExport(LODIndex);
-
-				SectionIndices = FGLTFMeshUtility::GetSectionIndices(MeshLOD, MaterialIndex);
-				MeshData = Builder.StaticMeshDataConverter.GetOrAdd(StaticMesh, StaticMeshComponent, LODIndex);
-			}
-			else if (const USkeletalMeshComponent* SkeletalMeshComponent = Cast<USkeletalMeshComponent>(Target))
-			{
-				const USkeletalMesh* SkeletalMesh = SkeletalMeshComponent->SkeletalMesh;
-				const int32 LODIndex = FGLTFMeshUtility::GetLOD(SkeletalMesh, SkeletalMeshComponent, DefaultLODIndex);
-				const FSkeletalMeshRenderData* RenderData = SkeletalMesh->GetResourceForRendering();
-				const FSkeletalMeshLODRenderData& MeshLOD = RenderData->LODRenderData[LODIndex];
-
-				SectionIndices = FGLTFMeshUtility::GetSectionIndices(MeshLOD, MaterialIndex);
-				MeshData = Builder.SkeletalMeshDataConverter.GetOrAdd(SkeletalMesh, SkeletalMeshComponent, LODIndex);
-			}
-		}
-		else
-		{
-			// TODO: report warning (about materials won't be export using mesh data because BakeMaterialInputs is not set to UseMeshData)
-		}
-	}
-
 	Builder.VariantReferenceChecker.GetOrAdd(Property, Owner);
 
 	FGLTFJsonEpicVariantMaterial VariantMaterial;
-	VariantMaterial.Material = Builder.GetOrAddMaterial(Material, MeshData, SectionIndices);
+	VariantMaterial.Material = FGLTFVariantUtility::GetOrAddMaterial(Builder, Material, Target, MaterialIndex);
 	VariantMaterial.Index = MaterialIndex;
 
 	const FGLTFJsonNodeIndex NodeIndex = Builder.GetOrAddNode(Target);
diff --git a/Engine/Plugins/Enterprise/GLTFExporter/MigratedFromGit/Source/GLTFExporter/Private/Converters/GLTFKhrVariantConverters.cpp b/Engine/Plugins/Enterprise/GLTFExporter/MigratedFromGit/Source/GLTFExporter/Private/Converters/GLTFKhrVariantConverters.cpp
--- a/Engine/Plugins/Enterprise/GLTFExporter/MigratedFromGit/Source/GLTFExporter/Private/Converters/GLTFKhrVariantConverters.cpp
+++ b/Engine/Plugins/Enterprise/GLTFExporter/MigratedFromGit/Source/GLTFExporter/Private/Converters/GLTFKhrVariantConverters.cpp
@@ -164,47 +164,7 @@ bool FGLTFKhrMaterialVariantConverter::TryParseMaterialProperty(FGLTFJsonPrimiti
 	const FGLTFJsonMeshIndex MeshIndex = Builder.GetOrAddMesh(MeshComponent);
 
 	OutPrimitive = &Builder.GetMesh(MeshIndex).Primitives[MaterialIndex];
-	OutMaterialIndex = GetResolvedMaterialForComponent(Material, MaterialIndex, MeshComponent);
+	OutMaterialIndex = FGLTFVariantUtility::GetOrAddMaterial(Builder, Material, MeshComponent, MaterialIndex);
 
 	return true;
 }
-
-FGLTFJsonMaterialIndex FGLTFKhrMaterialVariantConverter::GetResolvedMaterialForComponent(const UMaterialInterface* Material, int32 MaterialIndex, const UMeshComponent* MeshComponent) const
-{
-	const FGLTFMeshData* MeshData = nullptr;
-	FGLTFIndexArray SectionIndices;
-
-	if (Builder.ExportOptions->ExportMaterialVariants == EGLTFMaterialVariantMode::UseMeshData)
-	{
-		if (Builder.ExportOptions->BakeMaterialInputs == EGLTFMaterialBakeMode::UseMeshData)
-		{
-			const int32 DefaultLODIndex = Builder.ExportOptions->DefaultLevelOfDetail;
-
-			if (const UStaticMeshComponent* StaticMeshComponent = Cast<UStaticMeshComponent>(MeshComponent))
-			{
-				const UStaticMesh* StaticMesh = StaticMeshComponent->GetStaticMesh();
-				const int32 LODIndex = FGLTFMeshUtility::GetLOD(StaticMesh, StaticMeshComponent, DefaultLODIndex);
-				const FStaticMeshLODResources& MeshLOD = StaticMesh->GetLODForExport(LODIndex);
-
-				SectionIndices = FGLTFMeshUtility::GetSectionIndices(MeshLOD, MaterialIndex);
-				MeshData = Builder.StaticMeshDataConverter.GetOrAdd(StaticMesh, StaticMeshComponent, LODIndex);
-			}
-			else if (const USkeletalMeshComponent* SkeletalMeshComponent = Cast<USkeletalMeshComponent>(MeshComponent))
-			{
-				const USkeletalMesh* SkeletalMesh = SkeletalMeshComponent->SkeletalMesh;
-				const int32 LODIndex = FGLTFMeshUtility::GetLOD(SkeletalMesh, SkeletalMeshComponent, DefaultLODIndex);
-				const FSkeletalMeshRenderData* RenderData = SkeletalMesh->GetResourceForRendering();
-				const FSkeletalMeshLODRenderData& MeshLOD = RenderData->LODRenderData[LODIndex];
-
-				SectionIndices = FGLTFMeshUtility::GetSectionIndices(MeshLOD, MaterialIndex);
-				MeshData = Builder.SkeletalMeshDataConverter.GetOrAdd(SkeletalMesh, SkeletalMeshComponent, LODIndex);
-			}
-		}
-		else
-		{
-			// TODO: report warning (about materials won't be exported using mesh data because BakeMaterialInputs is not set to UseMeshData)
-		}
-	}
-
-	return Builder.GetOrAddMaterial(Material, MeshData, SectionIndices);
-}
diff --git a/Engine/Plugins/Enterprise/GLTFExporter/MigratedFromGit/Source/GLTFExporter/Private/Converters/GLTFKhrVariantConverters.h b/Engine/Plugins/Enterprise/GLTFExporter/MigratedFromGit/Source/GLTFExporter/Private/Converters/GLTFKhrVariantConverters.h
--- a/Engine/Plugins/Enterprise/GLTFExporter/MigratedFromGit/Source/GLTFExporter/Private/Converters/GLTFKhrVariantConverters.h
+++ b/Engine/Plugins/Enterprise/GLTFExporter/MigratedFromGit/Source/GLTFExporter/Private/Converters/GLTFKhrVariantConverters.h
@@ -17,5 +17,4 @@ class FGLTFKhrMaterialVariantConverter final : public FGLTFBuilderContext, publi
 	virtual FGLTFJsonKhrMaterialVariantIndex Convert(const UVariant* Variant) override;
 
 	bool TryParseMaterialProperty(FGLTFJsonPrimitive*& OutPrimitive, FGLTFJsonMaterialIndex& OutMaterialIndex, const UPropertyValue* Property) const;
-	FGLTFJsonMaterialIndex GetResolvedMaterialForComponent(const UMaterialInterface* Material, int32 MaterialIndex, const UMeshComponent* MeshComponent) const;
 };
diff --git a/Engine/Plugins/Enterprise/GLTFExporter/MigratedFromGit/Source/GLTFExporter/Private/Converters/GLTFLightMapConverters.cpp b/Engine/Plugins/Enterprise/GLTFExporter/MigratedFromGit/Source/GLTFExporter/Private/Converters/GLTFLightMapConverters.cpp
--- a/Engine/Plugins/Enterprise/GLTFExporter/MigratedFromGit/Source/GLTFExporter/Private/Converters/GLTFLightMapConverters.cpp
+++ b/Engine/Plugins/Enterprise/GLTFExporter/MigratedFromGit/Source/GLTFExporter/Private/Converters/GLTFLightMapConverters.cpp
@@ -35,7 +35,7 @@ FGLTFJsonLightMapIndex FGLTFLightMapConverter::Convert(const UStaticMeshComponen
 		return FGLTFJsonLightMapIndex(INDEX_NONE);
 	}
 
-	const int32 LightMapLODIndex = 0; // TODO: why is this zero?
+	constexpr int32 LightMapLODIndex = 0; // TODO: why is this zero?
 
 	if (!StaticMeshComponent->LODData.IsValidIndex(LightMapLODIndex))
 	{
diff --git a/Engine/Plugins/Enterprise/GLTFExporter/MigratedFromGit/Source/GLTFExporter/Private/Converters/GLTFMeshData.cpp b/Engine/Plugins/Enterprise/GLTFExporter/MigratedFromGit/Source/GLTFExporter/Private/Converters/GLTFMeshData.cpp
--- a/Engine/Plugins/Enterprise/GLTFExporter/MigratedFromGit/Source/GLTFExporter/Private/Converters/GLTFMeshData.cpp
+++ b/Engine/Plugins/Enterprise/GLTFExporter/MigratedFromGit/Source/GLTFExporter/Private/Converters/GLTFMeshData.cpp
@@ -12,6 +12,7 @@
 
 FGLTFMeshData::FGLTFMeshData(const UStaticMesh* StaticMesh, const UStaticMeshComponent* StaticMeshComponent, int32 LODIndex)
 	: Parent(nullptr)
+	, LODIndex(LODIndex)
 #if WITH_EDITOR
 	, LightMap(nullptr)
 	, LightMapResourceCluster(nullptr)
@@ -30,9 +31,10 @@ FGLTFMeshData::FGLTFMeshData(const UStaticMesh* StaticMesh, const UStaticMeshCom
 		PrimitiveData = { StaticMeshComponent };
 		FMeshMergeHelpers::RetrieveMesh(StaticMeshComponent, LODIndex, Description, true);
 
-		if (StaticMeshComponent->LODData.IsValidIndex(LODIndex))
+		constexpr int32 LightMapLODIndex = 0; // TODO: why is this zero?
+		if (StaticMeshComponent->LODData.IsValidIndex(LightMapLODIndex))
 		{
-			const FStaticMeshComponentLODInfo& LODData = StaticMeshComponent->LODData[LODIndex];
+			const FStaticMeshComponentLODInfo& LODData = StaticMeshComponent->LODData[LightMapLODIndex];
 			const FMeshMapBuildData* BuildData = StaticMeshComponent->GetMeshMapBuildData(LODData);
 			if (BuildData != nullptr)
 			{
@@ -68,6 +70,7 @@ FGLTFMeshData::FGLTFMeshData(const UStaticMesh* StaticMesh, const UStaticMeshCom
 
 FGLTFMeshData::FGLTFMeshData(const USkeletalMesh* SkeletalMesh, const USkeletalMeshComponent* SkeletalMeshComponent, int32 LODIndex)
 	: Parent(nullptr)
+	, LODIndex(LODIndex)
 #if WITH_EDITOR
 	, LightMap(nullptr)
 	, LightMapResourceCluster(nullptr)
diff --git a/Engine/Plugins/Enterprise/GLTFExporter/MigratedFromGit/Source/GLTFExporter/Private/Converters/GLTFMeshData.h b/Engine/Plugins/Enterprise/GLTFExporter/MigratedFromGit/Source/GLTFExporter/Private/Converters/GLTFMeshData.h
--- a/Engine/Plugins/Enterprise/GLTFExporter/MigratedFromGit/Source/GLTFExporter/Private/Converters/GLTFMeshData.h
+++ b/Engine/Plugins/Enterprise/GLTFExporter/MigratedFromGit/Source/GLTFExporter/Private/Converters/GLTFMeshData.h
@@ -20,6 +20,7 @@ struct FGLTFMeshData
 	const FGLTFMeshData* Parent;
 
 	FString Name;
+	int32 LODIndex;
 
 #if WITH_EDITOR
 	FPrimitiveData PrimitiveData;
diff --git a/Engine/Plugins/Enterprise/GLTFExporter/MigratedFromGit/Source/GLTFExporter/Private/Converters/GLTFMeshDataConverters.cpp b/Engine/Plugins/Enterprise/GLTFExporter/MigratedFromGit/Source/GLTFExporter/Private/Converters/GLTFMeshDataConverters.cpp
new file mode 100644
--- /dev/null
+++ b/Engine/Plugins/Enterprise/GLTFExporter/MigratedFromGit/Source/GLTFExporter/Private/Converters/GLTFMeshDataConverters.cpp
@@ -0,0 +1,34 @@
+// Copyright Epic Games, Inc. All Rights Reserved.
+
+#include "Converters/GLTFMeshDataConverters.h"
+#include "Converters/GLTFMeshUtility.h"
+#include "Builders/GLTFConvertBuilder.h"
+
+template class TGLTFMeshDataConverter<UStaticMesh, UStaticMeshComponent>;
+template class TGLTFMeshDataConverter<USkeletalMesh, USkeletalMeshComponent>;
+
+template <typename MeshType, typename MeshComponentType>
+void TGLTFMeshDataConverter<MeshType, MeshComponentType>::Sanitize(const MeshType*& Mesh, const MeshComponentType*& MeshComponent, int32& LODIndex)
+{
+	if (LODIndex < 0)
+	{
+		LODIndex = FGLTFMeshUtility::GetLOD(Mesh, MeshComponent, Builder.ExportOptions->DefaultLevelOfDetail);
+	}
+	else
+	{
+		LODIndex = FMath::Min(LODIndex, FGLTFMeshUtility::GetMaximumLOD(Mesh));
+	}
+}
+
+template <typename MeshType, typename MeshComponentType>
+const FGLTFMeshData* TGLTFMeshDataConverter<MeshType, MeshComponentType>::Convert(const MeshType* Mesh, const MeshComponentType* MeshComponent, int32 LODIndex)
+{
+	TUniquePtr<FGLTFMeshData> Output = MakeUnique<FGLTFMeshData>(Mesh, MeshComponent, LODIndex);
+
+	if (MeshComponent != nullptr)
+	{
+		Output->Parent = this->GetOrAdd(Mesh, nullptr, LODIndex);
+	}
+
+	return Outputs.Add_GetRef(MoveTemp(Output)).Get();
+}
diff --git a/Engine/Plugins/Enterprise/GLTFExporter/MigratedFromGit/Source/GLTFExporter/Private/Converters/GLTFMeshDataConverters.h b/Engine/Plugins/Enterprise/GLTFExporter/MigratedFromGit/Source/GLTFExporter/Private/Converters/GLTFMeshDataConverters.h
--- a/Engine/Plugins/Enterprise/GLTFExporter/MigratedFromGit/Source/GLTFExporter/Private/Converters/GLTFMeshDataConverters.h
+++ b/Engine/Plugins/Enterprise/GLTFExporter/MigratedFromGit/Source/GLTFExporter/Private/Converters/GLTFMeshDataConverters.h
@@ -4,23 +4,18 @@
 
 #include "Converters/GLTFConverter.h"
 #include "Converters/GLTFMeshData.h"
+#include "Converters/GLTFBuilderContext.h"
 
 template <typename MeshType, typename MeshComponentType>
-class TGLTFMeshDataConverter final : public TGLTFConverter<const FGLTFMeshData*, const MeshType*, const MeshComponentType*, int32>
+class TGLTFMeshDataConverter : public FGLTFBuilderContext, public TGLTFConverter<const FGLTFMeshData*, const MeshType*, const MeshComponentType*, int32>
 {
 	TArray<TUniquePtr<FGLTFMeshData>> Outputs;
 
-	virtual const FGLTFMeshData* Convert(const MeshType* Mesh, const MeshComponentType* MeshComponent, int32 LODIndex) override
-	{
-		TUniquePtr<FGLTFMeshData> Output = MakeUnique<FGLTFMeshData>(Mesh, MeshComponent, LODIndex);
+	using FGLTFBuilderContext::FGLTFBuilderContext;
 
-		if (MeshComponent != nullptr)
-		{
-			Output->Parent = this->GetOrAdd(Mesh, nullptr, LODIndex);
-		}
+	virtual void Sanitize(const MeshType*& Mesh, const MeshComponentType*& MeshComponent, int32& LODIndex) override;
 
-		return Outputs.Add_GetRef(MoveTemp(Output)).Get();
-	}
+	virtual const FGLTFMeshData* Convert(const MeshType* Mesh, const MeshComponentType* MeshComponent, int32 LODIndex) override;
 };
 
 typedef TGLTFMeshDataConverter<UStaticMesh, UStaticMeshComponent> FGLTFStaticMeshDataConverter;
diff --git a/Engine/Plugins/Enterprise/GLTFExporter/MigratedFromGit/Source/GLTFExporter/Private/Converters/GLTFMeshUtility.cpp b/Engine/Plugins/Enterprise/GLTFExporter/MigratedFromGit/Source/GLTFExporter/Private/Converters/GLTFMeshUtility.cpp
--- a/Engine/Plugins/Enterprise/GLTFExporter/MigratedFromGit/Source/GLTFExporter/Private/Converters/GLTFMeshUtility.cpp
+++ b/Engine/Plugins/Enterprise/GLTFExporter/MigratedFromGit/Source/GLTFExporter/Private/Converters/GLTFMeshUtility.cpp
@@ -21,6 +21,29 @@ const TArray<FSkeletalMaterial>& FGLTFMeshUtility::GetMaterials(const USkeletalM
 #endif
 }
 
+FGLTFIndexArray FGLTFMeshUtility::GetSectionIndices(const UStaticMesh* StaticMesh, int32 LODIndex, int32 MaterialIndex)
+{
+	if (StaticMesh == nullptr)
+	{
+		return {};
+	}
+
+	const FStaticMeshLODResources& MeshLOD = StaticMesh->GetLODForExport(LODIndex);
+	return GetSectionIndices(MeshLOD, MaterialIndex);
+}
+
+FGLTFIndexArray FGLTFMeshUtility::GetSectionIndices(const USkeletalMesh* SkeletalMesh, int32 LODIndex, int32 MaterialIndex)
+{
+	if (SkeletalMesh == nullptr)
+	{
+		return {};
+	}
+
+	const FSkeletalMeshRenderData* RenderData = SkeletalMesh->GetResourceForRendering();
+	const FSkeletalMeshLODRenderData& MeshLOD = RenderData->LODRenderData[LODIndex];
+	return GetSectionIndices(MeshLOD, MaterialIndex);
+}
+
 FGLTFIndexArray FGLTFMeshUtility::GetSectionIndices(const FStaticMeshLODResources& MeshLOD, int32 MaterialIndex)
 {
 	const FStaticMeshLODResources::FStaticMeshSectionArray& Sections = MeshLOD.Sections;
diff --git a/Engine/Plugins/Enterprise/GLTFExporter/MigratedFromGit/Source/GLTFExporter/Private/Converters/GLTFMeshUtility.h b/Engine/Plugins/Enterprise/GLTFExporter/MigratedFromGit/Source/GLTFExporter/Private/Converters/GLTFMeshUtility.h
--- a/Engine/Plugins/Enterprise/GLTFExporter/MigratedFromGit/Source/GLTFExporter/Private/Converters/GLTFMeshUtility.h
+++ b/Engine/Plugins/Enterprise/GLTFExporter/MigratedFromGit/Source/GLTFExporter/Private/Converters/GLTFMeshUtility.h
@@ -11,6 +11,9 @@ struct FGLTFMeshUtility
 	static const TArray<FStaticMaterial>& GetMaterials(const UStaticMesh* StaticMesh);
 	static const TArray<FSkeletalMaterial>& GetMaterials(const USkeletalMesh* SkeletalMesh);
 
+	static FGLTFIndexArray GetSectionIndices(const UStaticMesh* StaticMesh, int32 LODIndex, int32 MaterialIndex);
+	static FGLTFIndexArray GetSectionIndices(const USkeletalMesh* SkeletalMesh, int32 LODIndex, int32 MaterialIndex);
+
 	static FGLTFIndexArray GetSectionIndices(const FStaticMeshLODResources& MeshLOD, int32 MaterialIndex);
 	static FGLTFIndexArray GetSectionIndices(const FSkeletalMeshLODRenderData& MeshLOD, int32 MaterialIndex);
 
diff --git a/Engine/Plugins/Enterprise/GLTFExporter/MigratedFromGit/Source/GLTFExporter/Private/Converters/GLTFVariantUtility.cpp b/Engine/Plugins/Enterprise/GLTFExporter/MigratedFromGit/Source/GLTFExporter/Private/Converters/GLTFVariantUtility.cpp
--- a/Engine/Plugins/Enterprise/GLTFExporter/MigratedFromGit/Source/GLTFExporter/Private/Converters/GLTFVariantUtility.cpp
+++ b/Engine/Plugins/Enterprise/GLTFExporter/MigratedFromGit/Source/GLTFExporter/Private/Converters/GLTFVariantUtility.cpp
@@ -1,6 +1,7 @@
 // Copyright Epic Games, Inc. All Rights Reserved.
 
 #include "Converters/GLTFVariantUtility.h"
+#include "Builders/GLTFConvertBuilder.h"
 #include "LevelVariantSets.h"
 #include "VariantSet.h"
 #include "Variant.h"
@@ -68,3 +69,18 @@ FString FGLTFVariantUtility::GetLogContext(const ULevelVariantSets* LevelVariant
 {
 	return LevelVariantSets->GetName();
 }
+
+FGLTFJsonMaterialIndex FGLTFVariantUtility::GetOrAddMaterial(FGLTFConvertBuilder& Builder, const UMaterialInterface* Material, const UMeshComponent* MeshComponent, int32 MaterialIndex)
+{
+	if (Builder.ExportOptions->ExportMaterialVariants == EGLTFMaterialVariantMode::UseMeshData)
+	{
+		if (Builder.ExportOptions->BakeMaterialInputs == EGLTFMaterialBakeMode::UseMeshData)
+		{
+			return Builder.GetOrAddMaterial(Material, MeshComponent, -1, MaterialIndex);
+		}
+
+		// TODO: report warning (about materials won't be export using mesh data because BakeMaterialInputs is not set to UseMeshData)
+	}
+
+	return Builder.GetOrAddMaterial(Material);
+}
diff --git a/Engine/Plugins/Enterprise/GLTFExporter/MigratedFromGit/Source/GLTFExporter/Private/Converters/GLTFVariantUtility.h b/Engine/Plugins/Enterprise/GLTFExporter/MigratedFromGit/Source/GLTFExporter/Private/Converters/GLTFVariantUtility.h
--- a/Engine/Plugins/Enterprise/GLTFExporter/MigratedFromGit/Source/GLTFExporter/Private/Converters/GLTFVariantUtility.h
+++ b/Engine/Plugins/Enterprise/GLTFExporter/MigratedFromGit/Source/GLTFExporter/Private/Converters/GLTFVariantUtility.h
@@ -2,6 +2,9 @@
 
 #pragma once
 
+#include "Json/GLTFJsonIndex.h"
+
+class FGLTFConvertBuilder;
 struct FCapturedPropSegment;
 class ULevelVariantSets;
 class UVariantSet;
@@ -26,4 +29,6 @@ struct FGLTFVariantUtility
 	static FString GetLogContext(const UVariant* Variant);
 	static FString GetLogContext(const UVariantSet* VariantSet);
 	static FString GetLogContext(const ULevelVariantSets* LevelVariantSets);
+
+	static FGLTFJsonMaterialIndex GetOrAddMaterial(FGLTFConvertBuilder& Builder, const UMaterialInterface* Material, const UMeshComponent* MeshComponent, int32 MaterialIndex);
 };
diff --git a/Engine/Plugins/Enterprise/GLTFExporter/MigratedFromGit/Source/GLTFExporter/Private/Tasks/GLTFMeshTasks.cpp b/Engine/Plugins/Enterprise/GLTFExporter/MigratedFromGit/Source/GLTFExporter/Private/Tasks/GLTFMeshTasks.cpp
--- a/Engine/Plugins/Enterprise/GLTFExporter/MigratedFromGit/Source/GLTFExporter/Private/Tasks/GLTFMeshTasks.cpp
+++ b/Engine/Plugins/Enterprise/GLTFExporter/MigratedFromGit/Source/GLTFExporter/Private/Tasks/GLTFMeshTasks.cpp
@@ -135,7 +135,7 @@ void FGLTFStaticMeshTask::Complete()
 	}
 
 	const FGLTFMeshData* MeshData = Builder.ExportOptions->BakeMaterialInputs == EGLTFMaterialBakeMode::UseMeshData ?
-		Builder.StaticMeshDataConverter.GetOrAdd(StaticMesh, StaticMeshComponent, LODIndex) : nullptr;
+		Builder.GetOrAddMeshData(StaticMesh, StaticMeshComponent, LODIndex) : nullptr;
 
 #if WITH_EDITOR
 	if (MeshData != nullptr)
@@ -239,7 +239,7 @@ void FGLTFSkeletalMeshTask::Complete()
 	}
 
 	const FGLTFMeshData* MeshData = Builder.ExportOptions->BakeMaterialInputs == EGLTFMaterialBakeMode::UseMeshData ?
-		Builder.SkeletalMeshDataConverter.GetOrAdd(SkeletalMesh, SkeletalMeshComponent, LODIndex) : nullptr;
+		Builder.GetOrAddMeshData(SkeletalMesh, SkeletalMeshComponent, LODIndex) : nullptr;
 
 #if WITH_EDITOR
 	if (MeshData != nullptr)
