[GLTFExporter] refactor: reduce duplicated code by adding GLTFBuilder::SanitizeLOD

commit 43e540d4ae8feeab978e2dfa2d8dd9c4c2f2ea71
Date: Mon, 10 Jan 2022 08:12:16 +0100

    refactor: reduce duplicated code by adding GLTFBuilder::SanitizeLOD
    (cherry picked from commit c59617895bd63985dd12a8e1b1407e59e79aa649​)

#skipci

diff --git a/Engine/Plugins/Enterprise/GLTFExporter/MigratedFromGit/Source/GLTFExporter/Private/Builders/GLTFBuilder.cpp b/Engine/Plugins/Enterprise/GLTFExporter/MigratedFromGit/Source/GLTFExporter/Private/Builders/GLTFBuilder.cpp
--- a/Engine/Plugins/Enterprise/GLTFExporter/MigratedFromGit/Source/GLTFExporter/Private/Builders/GLTFBuilder.cpp
+++ b/Engine/Plugins/Enterprise/GLTFExporter/MigratedFromGit/Source/GLTFExporter/Private/Builders/GLTFBuilder.cpp
@@ -1,8 +1,9 @@
 // Copyright Epic Games, Inc. All Rights Reserved.
 
 #include "Builders/GLTFBuilder.h"
-#include "UserData/GLTFMaterialUserData.h"
 #include "Builders/GLTFFileUtility.h"
+#include "Converters/GLTFMeshUtility.h"
+#include "UserData/GLTFMaterialUserData.h"
 
 FGLTFBuilder::FGLTFBuilder(const FString& FilePath, const UGLTFExportOptions* ExportOptions)
 	: bIsGlbFile(FGLTFFileUtility::IsGlbFile(FilePath))
@@ -77,6 +78,20 @@ bool FGLTFBuilder::ShouldExportLight(EComponentMobility::Type LightMobility) con
 	return EnumHasAllFlags(AllowedMobility, QueriedMobility);
 }
 
+int32 FGLTFBuilder::SanitizeLOD(const UStaticMesh* StaticMesh, const UStaticMeshComponent* StaticMeshComponent, int32 LODIndex) const
+{
+	return LODIndex > 0
+		? FMath::Min(LODIndex, FGLTFMeshUtility::GetMaximumLOD(StaticMesh))
+		: FGLTFMeshUtility::GetLOD(StaticMesh, StaticMeshComponent, ExportOptions->DefaultLevelOfDetail);
+}
+
+int32 FGLTFBuilder::SanitizeLOD(const USkeletalMesh* SkeletalMesh, const USkeletalMeshComponent* SkeletalMeshComponent, int32 LODIndex) const
+{
+	return LODIndex > 0
+		? FMath::Min(LODIndex, FGLTFMeshUtility::GetMaximumLOD(SkeletalMesh))
+		: FGLTFMeshUtility::GetLOD(SkeletalMesh, SkeletalMeshComponent, ExportOptions->DefaultLevelOfDetail);
+}
+
 const UGLTFExportOptions* FGLTFBuilder::SanitizeExportOptions(const UGLTFExportOptions* Options)
 {
 	if (Options == nullptr)
diff --git a/Engine/Plugins/Enterprise/GLTFExporter/MigratedFromGit/Source/GLTFExporter/Private/Builders/GLTFBuilder.h b/Engine/Plugins/Enterprise/GLTFExporter/MigratedFromGit/Source/GLTFExporter/Private/Builders/GLTFBuilder.h
--- a/Engine/Plugins/Enterprise/GLTFExporter/MigratedFromGit/Source/GLTFExporter/Private/Builders/GLTFBuilder.h
+++ b/Engine/Plugins/Enterprise/GLTFExporter/MigratedFromGit/Source/GLTFExporter/Private/Builders/GLTFBuilder.h
@@ -29,6 +29,9 @@ public:
 
 	bool ShouldExportLight(EComponentMobility::Type LightMobility) const;
 
+	int32 SanitizeLOD(const UStaticMesh* StaticMesh, const UStaticMeshComponent* StaticMeshComponent, int32 LODIndex) const;
+	int32 SanitizeLOD(const USkeletalMesh* SkeletalMesh, const USkeletalMeshComponent* SkeletalMeshComponent, int32 LODIndex) const;
+
 private:
 
 	FGCObjectScopeGuard ExportOptionsGuard;
diff --git a/Engine/Plugins/Enterprise/GLTFExporter/MigratedFromGit/Source/GLTFExporter/Private/Builders/GLTFConvertBuilder.h b/Engine/Plugins/Enterprise/GLTFExporter/MigratedFromGit/Source/GLTFExporter/Private/Builders/GLTFConvertBuilder.h
--- a/Engine/Plugins/Enterprise/GLTFExporter/MigratedFromGit/Source/GLTFExporter/Private/Builders/GLTFConvertBuilder.h
+++ b/Engine/Plugins/Enterprise/GLTFExporter/MigratedFromGit/Source/GLTFExporter/Private/Builders/GLTFConvertBuilder.h
@@ -45,20 +45,20 @@ public:
 	FGLTFJsonAccessorIndex GetOrAddWeightAccessor(const FGLTFMeshSection* MeshSection, const FSkinWeightVertexBuffer* VertexBuffer, int32 InfluenceOffset);
 	FGLTFJsonAccessorIndex GetOrAddIndexAccessor(const FGLTFMeshSection* MeshSection);
 
-	FGLTFJsonMeshIndex GetOrAddMesh(const UStaticMesh* StaticMesh, const FGLTFMaterialArray& Materials = {}, int32 LODIndex = -1);
-	FGLTFJsonMeshIndex GetOrAddMesh(const USkeletalMesh* SkeletalMesh, const FGLTFMaterialArray& Materials = {}, int32 LODIndex = -1);
-	FGLTFJsonMeshIndex GetOrAddMesh(const UMeshComponent* MeshComponent, const FGLTFMaterialArray& Materials = {}, int32 LODIndex = -1);
-	FGLTFJsonMeshIndex GetOrAddMesh(const UStaticMeshComponent* StaticMeshComponent, const FGLTFMaterialArray& Materials = {}, int32 LODIndex = -1);
-	FGLTFJsonMeshIndex GetOrAddMesh(const USkeletalMeshComponent* SkeletalMeshComponent, const FGLTFMaterialArray& Materials = {}, int32 LODIndex = -1);
-
-	const FGLTFMeshData* GetOrAddMeshData(const UStaticMesh* StaticMesh, const UStaticMeshComponent* StaticMeshComponent = nullptr, int32 LODIndex = -1);
-	const FGLTFMeshData* GetOrAddMeshData(const USkeletalMesh* SkeletalMesh, const USkeletalMeshComponent* SkeletalMeshComponent = nullptr, int32 LODIndex = -1);
-
-	FGLTFJsonMaterialIndex GetOrAddMaterial(const UMaterialInterface* Material, const UStaticMesh* StaticMesh, int32 LODIndex = -1, int32 MaterialIndex = -1);
-	FGLTFJsonMaterialIndex GetOrAddMaterial(const UMaterialInterface* Material, const USkeletalMesh* SkeletalMesh, int32 LODIndex = -1, int32 MaterialIndex = -1);
-	FGLTFJsonMaterialIndex GetOrAddMaterial(const UMaterialInterface* Material, const UMeshComponent* MeshComponent, int32 LODIndex = -1, int32 MaterialIndex = -1);
-	FGLTFJsonMaterialIndex GetOrAddMaterial(const UMaterialInterface* Material, const UStaticMeshComponent* StaticMeshComponent, int32 LODIndex = -1, int32 MaterialIndex = -1);
-	FGLTFJsonMaterialIndex GetOrAddMaterial(const UMaterialInterface* Material, const USkeletalMeshComponent* SkeletalMeshComponent, int32 LODIndex = -1, int32 MaterialIndex = -1);
+	FGLTFJsonMeshIndex GetOrAddMesh(const UStaticMesh* StaticMesh, const FGLTFMaterialArray& Materials = {}, int32 LODIndex = INDEX_NONE);
+	FGLTFJsonMeshIndex GetOrAddMesh(const USkeletalMesh* SkeletalMesh, const FGLTFMaterialArray& Materials = {}, int32 LODIndex = INDEX_NONE);
+	FGLTFJsonMeshIndex GetOrAddMesh(const UMeshComponent* MeshComponent, const FGLTFMaterialArray& Materials = {}, int32 LODIndex = INDEX_NONE);
+	FGLTFJsonMeshIndex GetOrAddMesh(const UStaticMeshComponent* StaticMeshComponent, const FGLTFMaterialArray& Materials = {}, int32 LODIndex = INDEX_NONE);
+	FGLTFJsonMeshIndex GetOrAddMesh(const USkeletalMeshComponent* SkeletalMeshComponent, const FGLTFMaterialArray& Materials = {}, int32 LODIndex = INDEX_NONE);
+
+	const FGLTFMeshData* GetOrAddMeshData(const UStaticMesh* StaticMesh, const UStaticMeshComponent* StaticMeshComponent = nullptr, int32 LODIndex = INDEX_NONE);
+	const FGLTFMeshData* GetOrAddMeshData(const USkeletalMesh* SkeletalMesh, const USkeletalMeshComponent* SkeletalMeshComponent = nullptr, int32 LODIndex = INDEX_NONE);
+
+	FGLTFJsonMaterialIndex GetOrAddMaterial(const UMaterialInterface* Material, const UStaticMesh* StaticMesh, int32 LODIndex = INDEX_NONE, int32 MaterialIndex = INDEX_NONE);
+	FGLTFJsonMaterialIndex GetOrAddMaterial(const UMaterialInterface* Material, const USkeletalMesh* SkeletalMesh, int32 LODIndex = INDEX_NONE, int32 MaterialIndex = INDEX_NONE);
+	FGLTFJsonMaterialIndex GetOrAddMaterial(const UMaterialInterface* Material, const UMeshComponent* MeshComponent, int32 LODIndex = INDEX_NONE, int32 MaterialIndex = INDEX_NONE);
+	FGLTFJsonMaterialIndex GetOrAddMaterial(const UMaterialInterface* Material, const UStaticMeshComponent* StaticMeshComponent, int32 LODIndex = INDEX_NONE, int32 MaterialIndex = INDEX_NONE);
+	FGLTFJsonMaterialIndex GetOrAddMaterial(const UMaterialInterface* Material, const USkeletalMeshComponent* SkeletalMeshComponent, int32 LODIndex = INDEX_NONE, int32 MaterialIndex = INDEX_NONE);
 	FGLTFJsonMaterialIndex GetOrAddMaterial(const UMaterialInterface* Material, const FGLTFMeshData* MeshData = nullptr, const FGLTFIndexArray& SectionIndices = {});
 
 	FGLTFJsonSamplerIndex GetOrAddSampler(const UTexture* Texture);
diff --git a/Engine/Plugins/Enterprise/GLTFExporter/MigratedFromGit/Source/GLTFExporter/Private/Converters/GLTFMeshConverters.cpp b/Engine/Plugins/Enterprise/GLTFExporter/MigratedFromGit/Source/GLTFExporter/Private/Converters/GLTFMeshConverters.cpp
--- a/Engine/Plugins/Enterprise/GLTFExporter/MigratedFromGit/Source/GLTFExporter/Private/Converters/GLTFMeshConverters.cpp
+++ b/Engine/Plugins/Enterprise/GLTFExporter/MigratedFromGit/Source/GLTFExporter/Private/Converters/GLTFMeshConverters.cpp
@@ -21,21 +21,14 @@ void FGLTFStaticMeshConverter::Sanitize(const UStaticMesh*& StaticMesh, const US
 		FGLTFMaterialUtility::ResolveOverrides(Materials, FGLTFMeshUtility::GetMaterials(StaticMesh));
 	}
 
-	if (LODIndex < 0)
-	{
-		LODIndex = FGLTFMeshUtility::GetLOD(StaticMesh, StaticMeshComponent, Builder.ExportOptions->DefaultLevelOfDetail);
-	}
-	else
-	{
-		LODIndex = FMath::Min(LODIndex, FGLTFMeshUtility::GetMaximumLOD(StaticMesh));
-	}
+	LODIndex = Builder.SanitizeLOD(StaticMesh, StaticMeshComponent, LODIndex);
 
 	if (StaticMeshComponent != nullptr)
 	{
 		const bool bUsesMeshData = Builder.ExportOptions->BakeMaterialInputs == EGLTFMaterialBakeMode::UseMeshData &&
 			FGLTFMaterialUtility::NeedsMeshData(Materials); // TODO: if this expensive, cache the results for each material
 
-		const bool bIsReferencedByVariant = Builder.VariantReferenceChecker.IsReferenced(StaticMeshComponent); 
+		const bool bIsReferencedByVariant = Builder.VariantReferenceChecker.IsReferenced(StaticMeshComponent);
 
 		// Only use the component if it's needed for baking or variants, since we would
 		// otherwise export a copy of this mesh for each mesh-component.
@@ -77,21 +70,14 @@ void FGLTFSkeletalMeshConverter::Sanitize(const USkeletalMesh*& SkeletalMesh, co
 		FGLTFMaterialUtility::ResolveOverrides(Materials, FGLTFMeshUtility::GetMaterials(SkeletalMesh));
 	}
 
-	if (LODIndex < 0)
-	{
-		LODIndex = FGLTFMeshUtility::GetLOD(SkeletalMesh, SkeletalMeshComponent, Builder.ExportOptions->DefaultLevelOfDetail);
-	}
-	else
-	{
-		LODIndex = FMath::Min(LODIndex, FGLTFMeshUtility::GetMaximumLOD(SkeletalMesh));
-	}
+	LODIndex = Builder.SanitizeLOD(SkeletalMesh, SkeletalMeshComponent, LODIndex);
 
 	if (SkeletalMeshComponent != nullptr)
 	{
 		const bool bUsesMeshData = Builder.ExportOptions->BakeMaterialInputs == EGLTFMaterialBakeMode::UseMeshData &&
 			FGLTFMaterialUtility::NeedsMeshData(Materials); // TODO: if this expensive, cache the results for each material
 
-		const bool bIsReferencedByVariant = Builder.VariantReferenceChecker.IsReferenced(SkeletalMeshComponent); 
+		const bool bIsReferencedByVariant = Builder.VariantReferenceChecker.IsReferenced(SkeletalMeshComponent);
 
 		// Only use the component if it's needed for baking or variants, since we would
 		// otherwise export a copy of this mesh for each mesh-component.
diff --git a/Engine/Plugins/Enterprise/GLTFExporter/MigratedFromGit/Source/GLTFExporter/Private/Converters/GLTFMeshDataConverters.cpp b/Engine/Plugins/Enterprise/GLTFExporter/MigratedFromGit/Source/GLTFExporter/Private/Converters/GLTFMeshDataConverters.cpp
--- a/Engine/Plugins/Enterprise/GLTFExporter/MigratedFromGit/Source/GLTFExporter/Private/Converters/GLTFMeshDataConverters.cpp
+++ b/Engine/Plugins/Enterprise/GLTFExporter/MigratedFromGit/Source/GLTFExporter/Private/Converters/GLTFMeshDataConverters.cpp
@@ -1,7 +1,6 @@
 // Copyright Epic Games, Inc. All Rights Reserved.
 
 #include "Converters/GLTFMeshDataConverters.h"
-#include "Converters/GLTFMeshUtility.h"
 #include "Builders/GLTFConvertBuilder.h"
 
 template class TGLTFMeshDataConverter<UStaticMesh, UStaticMeshComponent>;
@@ -10,14 +9,7 @@ template class TGLTFMeshDataConverter<USkeletalMesh, USkeletalMeshComponent>;
 template <typename MeshType, typename MeshComponentType>
 void TGLTFMeshDataConverter<MeshType, MeshComponentType>::Sanitize(const MeshType*& Mesh, const MeshComponentType*& MeshComponent, int32& LODIndex)
 {
-	if (LODIndex < 0)
-	{
-		LODIndex = FGLTFMeshUtility::GetLOD(Mesh, MeshComponent, Builder.ExportOptions->DefaultLevelOfDetail);
-	}
-	else
-	{
-		LODIndex = FMath::Min(LODIndex, FGLTFMeshUtility::GetMaximumLOD(Mesh));
-	}
+	LODIndex = Builder.SanitizeLOD(Mesh, MeshComponent, LODIndex);
 }
 
 template <typename MeshType, typename MeshComponentType>
