[GLTFExporter] refactor: replace GLTFVariantReferenceChecker with simplified mapping of variant properties

commit 08ba8b46951d16713f5f8d7cbbba0bd7e6e3b20e
Date: Mon, 10 Jan 2022 08:13:14 +0100

    refactor: replace GLTFVariantReferenceChecker with simplified mapping of variant properties
    (cherry picked from commit 11f69588dbe212db74ae875bcb540b8b334873c4​)

#skipci

diff --git a/Engine/Plugins/Enterprise/GLTFExporter/MigratedFromGit/Source/GLTFExporter/Private/Builders/GLTFConvertBuilder.cpp b/Engine/Plugins/Enterprise/GLTFExporter/MigratedFromGit/Source/GLTFExporter/Private/Builders/GLTFConvertBuilder.cpp
--- a/Engine/Plugins/Enterprise/GLTFExporter/MigratedFromGit/Source/GLTFExporter/Private/Builders/GLTFConvertBuilder.cpp
+++ b/Engine/Plugins/Enterprise/GLTFExporter/MigratedFromGit/Source/GLTFExporter/Private/Builders/GLTFConvertBuilder.cpp
@@ -472,26 +472,6 @@ FGLTFJsonBackdropIndex FGLTFConvertBuilder::GetOrAddBackdrop(const AActor* Backd
 	return BackdropConverter.GetOrAdd(BackdropActor);
 }
 
-FGLTFJsonEpicLevelVariantSetsIndex FGLTFConvertBuilder::GetOrAddEpicLevelVariantSets(const ULevelVariantSets* LevelVariantSets)
-{
-	if (LevelVariantSets == nullptr)
-	{
-		return FGLTFJsonEpicLevelVariantSetsIndex(INDEX_NONE);
-	}
-
-	return EpicLevelVariantSetsConverter.GetOrAdd(LevelVariantSets);
-}
-
-FGLTFJsonKhrMaterialVariantIndex FGLTFConvertBuilder::GetOrAddKhrMaterialVariant(const UVariant* Variant)
-{
-	if (Variant == nullptr)
-	{
-		return FGLTFJsonKhrMaterialVariantIndex(INDEX_NONE);
-	}
-
-	return KhrMaterialVariantConverter.GetOrAdd(Variant);
-}
-
 FGLTFJsonLightMapIndex FGLTFConvertBuilder::GetOrAddLightMap(const UStaticMeshComponent* StaticMeshComponent)
 {
 	if (StaticMeshComponent == nullptr)
@@ -521,3 +501,34 @@ FGLTFJsonSkySphereIndex FGLTFConvertBuilder::GetOrAddSkySphere(const AActor* Sky
 
 	return SkySphereConverter.GetOrAdd(SkySphereActor);
 }
+
+FGLTFJsonEpicLevelVariantSetsIndex FGLTFConvertBuilder::GetOrAddEpicLevelVariantSets(const ULevelVariantSets* LevelVariantSets)
+{
+	if (LevelVariantSets == nullptr)
+	{
+		return FGLTFJsonEpicLevelVariantSetsIndex(INDEX_NONE);
+	}
+
+	return EpicLevelVariantSetsConverter.GetOrAdd(LevelVariantSets);
+}
+
+FGLTFJsonKhrMaterialVariantIndex FGLTFConvertBuilder::GetOrAddKhrMaterialVariant(const UVariant* Variant)
+{
+	if (Variant == nullptr)
+	{
+		return FGLTFJsonKhrMaterialVariantIndex(INDEX_NONE);
+	}
+
+	return KhrMaterialVariantConverter.GetOrAdd(Variant);
+}
+
+void FGLTFConvertBuilder::RegisterObjectVariant(const UObject* Object, const UPropertyValue* Property)
+{
+	TArray<const UPropertyValue*>& Variants = ObjectVariants.FindOrAdd(Object);
+	Variants.AddUnique(Property);
+}
+
+const TArray<const UPropertyValue*>* FGLTFConvertBuilder::GetObjectVariants(const UObject* Object) const
+{
+	return ObjectVariants.Find(Object);
+}
diff --git a/Engine/Plugins/Enterprise/GLTFExporter/MigratedFromGit/Source/GLTFExporter/Private/Builders/GLTFConvertBuilder.h b/Engine/Plugins/Enterprise/GLTFExporter/MigratedFromGit/Source/GLTFExporter/Private/Builders/GLTFConvertBuilder.h
--- a/Engine/Plugins/Enterprise/GLTFExporter/MigratedFromGit/Source/GLTFExporter/Private/Builders/GLTFConvertBuilder.h
+++ b/Engine/Plugins/Enterprise/GLTFExporter/MigratedFromGit/Source/GLTFExporter/Private/Builders/GLTFConvertBuilder.h
@@ -21,7 +21,6 @@
 #include "Converters/GLTFHotspotConverters.h"
 #include "Converters/GLTFKhrVariantConverters.h"
 #include "Converters/GLTFSkySphereConverters.h"
-#include "Converters/GLTFVariantReferenceChecker.h"
 
 class FGLTFConvertBuilder : public FGLTFImageBuilder
 {
@@ -33,9 +32,6 @@ public:
 
 	const bool bSelectedActorsOnly;
 
-	// TODO: find a better place for this types of indirect converters
-	FGLTFVariantReferenceChecker VariantReferenceChecker;
-
 	FGLTFJsonAccessorIndex GetOrAddPositionAccessor(const FGLTFMeshSection* MeshSection, const FPositionVertexBuffer* VertexBuffer);
 	FGLTFJsonAccessorIndex GetOrAddColorAccessor(const FGLTFMeshSection* MeshSection, const FColorVertexBuffer* VertexBuffer);
 	FGLTFJsonAccessorIndex GetOrAddNormalAccessor(const FGLTFMeshSection* MeshSection, const FStaticMeshVertexBuffer* VertexBuffer);
@@ -90,12 +86,16 @@ public:
 	FGLTFJsonCameraIndex GetOrAddCamera(const UCameraComponent* CameraComponent);
 	FGLTFJsonLightIndex GetOrAddLight(const ULightComponent* LightComponent);
 	FGLTFJsonBackdropIndex GetOrAddBackdrop(const AActor* BackdropActor);
-	FGLTFJsonEpicLevelVariantSetsIndex GetOrAddEpicLevelVariantSets(const ULevelVariantSets* LevelVariantSets);
-	FGLTFJsonKhrMaterialVariantIndex GetOrAddKhrMaterialVariant(const UVariant* Variant);
 	FGLTFJsonLightMapIndex GetOrAddLightMap(const UStaticMeshComponent* StaticMeshComponent);
 	FGLTFJsonHotspotIndex GetOrAddHotspot(const AGLTFHotspotActor* HotspotActor);
 	FGLTFJsonSkySphereIndex GetOrAddSkySphere(const AActor* SkySphereActor);
 
+	FGLTFJsonEpicLevelVariantSetsIndex GetOrAddEpicLevelVariantSets(const ULevelVariantSets* LevelVariantSets);
+	FGLTFJsonKhrMaterialVariantIndex GetOrAddKhrMaterialVariant(const UVariant* Variant);
+
+	void RegisterObjectVariant(const UObject* Object, const UPropertyValue* Property);
+	const TArray<const UPropertyValue*>* GetObjectVariants(const UObject* Object) const;
+
 private:
 
 	FGLTFPositionBufferConverter PositionBufferConverter{ *this };
@@ -138,9 +138,12 @@ private:
 	FGLTFCameraConverter CameraConverter{ *this };
 	FGLTFLightConverter LightConverter{ *this };
 	FGLTFBackdropConverter BackdropConverter{ *this };
-	FGLTFEpicLevelVariantSetsConverter EpicLevelVariantSetsConverter{ *this };
-	FGLTFKhrMaterialVariantConverter KhrMaterialVariantConverter{ *this };
 	FGLTFLightMapConverter LightMapConverter{ *this };
 	FGLTFHotspotConverter HotspotConverter{ *this };
 	FGLTFSkySphereConverter SkySphereConverter{ *this };
+
+	FGLTFEpicLevelVariantSetsConverter EpicLevelVariantSetsConverter{ *this };
+	FGLTFKhrMaterialVariantConverter KhrMaterialVariantConverter{ *this };
+
+	TMap<const UObject*, TArray<const UPropertyValue*>> ObjectVariants;
 };
diff --git a/Engine/Plugins/Enterprise/GLTFExporter/MigratedFromGit/Source/GLTFExporter/Private/Converters/GLTFEpicLevelVariantSetsConverters.cpp b/Engine/Plugins/Enterprise/GLTFExporter/MigratedFromGit/Source/GLTFExporter/Private/Converters/GLTFEpicLevelVariantSetsConverters.cpp
--- a/Engine/Plugins/Enterprise/GLTFExporter/MigratedFromGit/Source/GLTFExporter/Private/Converters/GLTFEpicLevelVariantSetsConverters.cpp
+++ b/Engine/Plugins/Enterprise/GLTFExporter/MigratedFromGit/Source/GLTFExporter/Private/Converters/GLTFEpicLevelVariantSetsConverters.cpp
@@ -162,7 +162,7 @@ bool FGLTFEpicLevelVariantSetsConverter::TryParseVisibilityPropertyValue(FGLTFJs
 		return false;
 	}
 
-	Builder.VariantReferenceChecker.GetOrAdd(Property, Owner);
+	Builder.RegisterObjectVariant(Owner, Property);
 
 	const FGLTFJsonNodeIndex NodeIndex = Builder.GetOrAddNode(Target);
 	const FGLTFJsonNodeIndex ComponentNodeIndex = Builder.GetComponentNodeIndex(NodeIndex);
@@ -229,7 +229,7 @@ bool FGLTFEpicLevelVariantSetsConverter::TryParseMaterialPropertyValue(FGLTFJson
 	const UMaterialInterface* Material = MaterialProperty->GetMaterial();
 	const int32 MaterialIndex = CapturedPropSegments[NumPropSegments - 1].PropertyIndex;
 
-	Builder.VariantReferenceChecker.GetOrAdd(Property, Owner);
+	Builder.RegisterObjectVariant(Owner, Property);
 
 	FGLTFJsonEpicVariantMaterial VariantMaterial;
 	VariantMaterial.Material = FGLTFVariantUtility::GetOrAddMaterial(Builder, Material, Target, MaterialIndex);
@@ -290,7 +290,7 @@ bool FGLTFEpicLevelVariantSetsConverter::TryParseStaticMeshPropertyValue(FGLTFJs
 		return false;
 	}
 
-	Builder.VariantReferenceChecker.GetOrAdd(Property, Owner);
+	Builder.RegisterObjectVariant(Owner, Property);
 
 	const FGLTFMaterialArray OverrideMaterials(MeshComponent->OverrideMaterials);
 	const FGLTFJsonNodeIndex NodeIndex = Builder.GetOrAddNode(Target);
@@ -349,7 +349,7 @@ bool FGLTFEpicLevelVariantSetsConverter::TryParseSkeletalMeshPropertyValue(FGLTF
 		return false;
 	}
 
-	Builder.VariantReferenceChecker.GetOrAdd(Property, Owner);
+	Builder.RegisterObjectVariant(Owner, Property);
 
 	const FGLTFMaterialArray OverrideMaterials(MeshComponent->OverrideMaterials);
 	const FGLTFJsonNodeIndex NodeIndex = Builder.GetOrAddNode(Target);
diff --git a/Engine/Plugins/Enterprise/GLTFExporter/MigratedFromGit/Source/GLTFExporter/Private/Converters/GLTFKhrVariantConverters.cpp b/Engine/Plugins/Enterprise/GLTFExporter/MigratedFromGit/Source/GLTFExporter/Private/Converters/GLTFKhrVariantConverters.cpp
--- a/Engine/Plugins/Enterprise/GLTFExporter/MigratedFromGit/Source/GLTFExporter/Private/Converters/GLTFKhrVariantConverters.cpp
+++ b/Engine/Plugins/Enterprise/GLTFExporter/MigratedFromGit/Source/GLTFExporter/Private/Converters/GLTFKhrVariantConverters.cpp
@@ -157,8 +157,8 @@ bool FGLTFKhrMaterialVariantConverter::TryParseMaterialProperty(FGLTFJsonPrimiti
 		return false;
 	}
 
-	Builder.VariantReferenceChecker.GetOrAdd(MaterialProperty, MeshComponent);
-	Builder.VariantReferenceChecker.GetOrAdd(MaterialProperty, Owner);
+	Builder.RegisterObjectVariant(MeshComponent, MaterialProperty);
+	Builder.RegisterObjectVariant(Owner, MaterialProperty);
 
 	const int32 MaterialIndex = CapturedPropSegments[NumPropSegments - 1].PropertyIndex;
 	const FGLTFJsonMeshIndex MeshIndex = Builder.GetOrAddMesh(MeshComponent);
diff --git a/Engine/Plugins/Enterprise/GLTFExporter/MigratedFromGit/Source/GLTFExporter/Private/Converters/GLTFMeshConverters.cpp b/Engine/Plugins/Enterprise/GLTFExporter/MigratedFromGit/Source/GLTFExporter/Private/Converters/GLTFMeshConverters.cpp
--- a/Engine/Plugins/Enterprise/GLTFExporter/MigratedFromGit/Source/GLTFExporter/Private/Converters/GLTFMeshConverters.cpp
+++ b/Engine/Plugins/Enterprise/GLTFExporter/MigratedFromGit/Source/GLTFExporter/Private/Converters/GLTFMeshConverters.cpp
@@ -28,7 +28,7 @@ void FGLTFStaticMeshConverter::Sanitize(const UStaticMesh*& StaticMesh, const US
 		const bool bUsesMeshData = Builder.ExportOptions->BakeMaterialInputs == EGLTFMaterialBakeMode::UseMeshData &&
 			FGLTFMaterialUtility::NeedsMeshData(Materials); // TODO: if this expensive, cache the results for each material
 
-		const bool bIsReferencedByVariant = Builder.VariantReferenceChecker.IsReferenced(StaticMeshComponent);
+		const bool bIsReferencedByVariant = Builder.GetObjectVariants(StaticMeshComponent) != nullptr;
 
 		// Only use the component if it's needed for baking or variants, since we would
 		// otherwise export a copy of this mesh for each mesh-component.
@@ -77,7 +77,7 @@ void FGLTFSkeletalMeshConverter::Sanitize(const USkeletalMesh*& SkeletalMesh, co
 		const bool bUsesMeshData = Builder.ExportOptions->BakeMaterialInputs == EGLTFMaterialBakeMode::UseMeshData &&
 			FGLTFMaterialUtility::NeedsMeshData(Materials); // TODO: if this expensive, cache the results for each material
 
-		const bool bIsReferencedByVariant = Builder.VariantReferenceChecker.IsReferenced(SkeletalMeshComponent);
+		const bool bIsReferencedByVariant = Builder.GetObjectVariants(SkeletalMeshComponent) != nullptr;
 
 		// Only use the component if it's needed for baking or variants, since we would
 		// otherwise export a copy of this mesh for each mesh-component.
diff --git a/Engine/Plugins/Enterprise/GLTFExporter/MigratedFromGit/Source/GLTFExporter/Private/Converters/GLTFVariantReferenceChecker.h b/Engine/Plugins/Enterprise/GLTFExporter/MigratedFromGit/Source/GLTFExporter/Private/Converters/GLTFVariantReferenceChecker.h
deleted file mode 100644
