[GLTFExporter] refactor: reduce duplicated code in variant converter and improve warning messages

commit 8c737dc0340a82834140d8b0e50902df39bd21b7
Date: Mon, 10 Jan 2022 08:14:42 +0100

    refactor: reduce duplicated code in variant converter and improve warning messages
    (cherry picked from commit 41a2f97fcd8fb522daba53c83ff1ec85d7790596​)

#skipci

diff --git a/Engine/Plugins/Enterprise/GLTFExporter/MigratedFromGit/Source/GLTFExporter/Private/Converters/GLTFEpicLevelVariantSetsConverters.cpp b/Engine/Plugins/Enterprise/GLTFExporter/MigratedFromGit/Source/GLTFExporter/Private/Converters/GLTFEpicLevelVariantSetsConverters.cpp
--- a/Engine/Plugins/Enterprise/GLTFExporter/MigratedFromGit/Source/GLTFExporter/Private/Converters/GLTFEpicLevelVariantSetsConverters.cpp
+++ b/Engine/Plugins/Enterprise/GLTFExporter/MigratedFromGit/Source/GLTFExporter/Private/Converters/GLTFEpicLevelVariantSetsConverters.cpp
@@ -2,19 +2,14 @@
 
 #include "Converters/GLTFEpicLevelVariantSetsConverters.h"
 #include "Converters/GLTFVariantUtility.h"
-#include "Converters/GLTFMeshUtility.h"
-#include "Builders/GLTFContainerBuilder.h"
-#include "Rendering/SkeletalMeshRenderData.h"
+#include "Builders/GLTFConvertBuilder.h"
 #include "LevelVariantSets.h"
 #include "VariantSet.h"
 #include "Variant.h"
 #include "VariantObjectBinding.h"
 #include "PropertyValue.h"
 #include "PropertyValueMaterial.h"
-#include "Engine/StaticMesh.h"
-#include "Engine/SkeletalMesh.h"
-#include "Components/StaticMeshComponent.h"
-#include "Components/SkeletalMeshComponent.h"
+#include "Components/MeshComponent.h"
 
 FGLTFJsonEpicLevelVariantSetsIndex FGLTFEpicLevelVariantSetsConverter::Convert(const ULevelVariantSets* LevelVariantSets)
 {
@@ -80,7 +75,7 @@ bool FGLTFEpicLevelVariantSetsConverter::TryParseVariantBinding(FGLTFJsonEpicVar
 {
 	bool bHasParsedAnyProperty = false;
 
-	for (UPropertyValue* Property: Binding->GetCapturedProperties())
+	for (const UPropertyValue* Property: Binding->GetCapturedProperties())
 	{
 		if (!const_cast<UPropertyValue*>(Property)->Resolve() || !Property->HasRecordedData())
 		{
@@ -90,23 +85,23 @@ bool FGLTFEpicLevelVariantSetsConverter::TryParseVariantBinding(FGLTFJsonEpicVar
 		const FName PropertyName = Property->GetPropertyName();
 		const FFieldClass* PropertyClass = Property->GetPropertyClass();
 
-		if (Property->IsA<UPropertyValueMaterial>())
+		if (const UPropertyValueMaterial* MaterialProperty = Cast<UPropertyValueMaterial>(Property))
 		{
-			if (Builder.ExportOptions->ExportMaterialVariants != EGLTFMaterialVariantMode::None && TryParseMaterialPropertyValue(OutVariant, Property))
+			if (Builder.ExportOptions->ExportMaterialVariants != EGLTFMaterialVariantMode::None && TryParseMaterialPropertyValue(OutVariant, MaterialProperty))
 			{
 				bHasParsedAnyProperty = true;
 			}
 		}
 		else if (PropertyName == TEXT("StaticMesh")) // TODO: should we not also check PropertyClass?
 		{
-			if (Builder.ExportOptions->bExportMeshVariants && TryParseStaticMeshPropertyValue(OutVariant, Property))
+			if (Builder.ExportOptions->bExportMeshVariants && TryParseMeshPropertyValue<UStaticMesh>(OutVariant, Property))
 			{
 				bHasParsedAnyProperty = true;
 			}
 		}
 		else if (PropertyName == TEXT("SkeletalMesh")) // TODO: should we not also check PropertyClass?
 		{
-			if (Builder.ExportOptions->bExportMeshVariants && TryParseSkeletalMeshPropertyValue(OutVariant, Property))
+			if (Builder.ExportOptions->bExportMeshVariants && TryParseMeshPropertyValue<USkeletalMesh>(OutVariant, Property))
 			{
 				bHasParsedAnyProperty = true;
 			}
@@ -131,7 +126,7 @@ bool FGLTFEpicLevelVariantSetsConverter::TryParseVisibilityPropertyValue(FGLTFJs
 	if (Target == nullptr)
 	{
 		Builder.LogWarning(FString::Printf(
-			TEXT("Target object for property is invalid, the property will be skipped. Context: %s"),
+			TEXT("Property must belong to a scene component, the property will be skipped. Context: %s"),
 			*FGLTFVariantUtility::GetLogContext(Property)));
 		return false;
 	}
@@ -139,7 +134,8 @@ bool FGLTFEpicLevelVariantSetsConverter::TryParseVisibilityPropertyValue(FGLTFJs
 	const AActor* Owner = Target->GetOwner();
 	if (Owner == nullptr)
 	{
-		Builder.LogWarning(FString::Printf(TEXT("Invalid scene component, the property will be skipped. Context: %s"),
+		Builder.LogWarning(FString::Printf(
+			TEXT("Property must belong to an actor, the property will be skipped. Context: %s"),
 			*FGLTFVariantUtility::GetLogContext(Property)));
 		return false;
 	}
@@ -147,7 +143,7 @@ bool FGLTFEpicLevelVariantSetsConverter::TryParseVisibilityPropertyValue(FGLTFJs
 	if (Builder.bSelectedActorsOnly && !Owner->IsSelected())
 	{
 		Builder.LogWarning(FString::Printf(
-			TEXT("Target actor for property is not selected for export, the property will be skipped. Context: %s"),
+			TEXT("Property doesn't belong to an actor selected for export, the property will be skipped. Context: %s"),
 			*FGLTFVariantUtility::GetLogContext(Property)));
 		return false;
 	}
@@ -174,22 +170,13 @@ bool FGLTFEpicLevelVariantSetsConverter::TryParseVisibilityPropertyValue(FGLTFJs
 	return true;
 }
 
-bool FGLTFEpicLevelVariantSetsConverter::TryParseMaterialPropertyValue(FGLTFJsonEpicVariant& OutVariant, const UPropertyValue* Property) const
+bool FGLTFEpicLevelVariantSetsConverter::TryParseMaterialPropertyValue(FGLTFJsonEpicVariant& OutVariant, const UPropertyValueMaterial* Property) const
 {
-	UPropertyValueMaterial* MaterialProperty = Cast<UPropertyValueMaterial>(const_cast<UPropertyValue*>(Property));
-	if (MaterialProperty == nullptr)
-	{
-		Builder.LogWarning(FString::Printf(
-			TEXT("Material property is invalid, the property will be skipped. Context: %s"),
-			*FGLTFVariantUtility::GetLogContext(Property)));
-		return false;
-	}
-
 	const UMeshComponent* Target = static_cast<UMeshComponent*>(Property->GetPropertyParentContainerAddress());
 	if (Target == nullptr)
 	{
 		Builder.LogWarning(FString::Printf(
-			TEXT("Target object for property is invalid, the property will be skipped. Context: %s"),
+			TEXT("Variant property %s must belong to a mesh component, the property will be skipped"),
 			*FGLTFVariantUtility::GetLogContext(Property)));
 		return false;
 	}
@@ -197,7 +184,8 @@ bool FGLTFEpicLevelVariantSetsConverter::TryParseMaterialPropertyValue(FGLTFJson
 	const AActor* Owner = Target->GetOwner();
 	if (Owner == nullptr)
 	{
-		Builder.LogWarning(FString::Printf(TEXT("Invalid scene component, the property will be skipped. Context: %s"),
+		Builder.LogWarning(FString::Printf(
+			TEXT("Variant property %s must belong to an actor, the property will be skipped"),
 			*FGLTFVariantUtility::GetLogContext(Property)));
 		return false;
 	}
@@ -205,28 +193,28 @@ bool FGLTFEpicLevelVariantSetsConverter::TryParseMaterialPropertyValue(FGLTFJson
 	if (Builder.bSelectedActorsOnly && !Owner->IsSelected())
 	{
 		Builder.LogWarning(FString::Printf(
-			TEXT("Target actor for property is not selected for export, the property will be skipped. Context: %s"),
+			TEXT("Variant property %s doesn't belong to an actor selected for export, the property will be skipped"),
 			*FGLTFVariantUtility::GetLogContext(Property)));
 		return false;
 	}
 
-	const TArray<FCapturedPropSegment>& CapturedPropSegments = FGLTFVariantUtility::GetCapturedPropSegments(MaterialProperty);
+	const TArray<FCapturedPropSegment>& CapturedPropSegments = FGLTFVariantUtility::GetCapturedPropSegments(Property);
 	const int32 NumPropSegments = CapturedPropSegments.Num();
 
 	if (NumPropSegments < 1)
 	{
 		Builder.LogWarning(FString::Printf(
-			TEXT("Failed to parse element index to apply the material to, the property will be skipped. Context: %s"),
+			TEXT("Failed to parse material index for variant property %s, the property will be skipped"),
 			*FGLTFVariantUtility::GetLogContext(Property)));
 		return false;
 	}
 
 	// NOTE: UPropertyValueMaterial::GetMaterial does *not* ensure that the recorded data has been loaded,
 	// so we need to call UProperty::GetRecordedData first to make that happen.
-	MaterialProperty->GetRecordedData();
+	const_cast<UPropertyValueMaterial*>(Property)->GetRecordedData();
 
 	// TODO: find way to determine whether the material is null because "None" was selected, or because it failed to resolve
-	const UMaterialInterface* Material = MaterialProperty->GetMaterial();
+	const UMaterialInterface* Material = const_cast<UPropertyValueMaterial*>(Property)->GetMaterial();
 	const int32 MaterialIndex = CapturedPropSegments[NumPropSegments - 1].PropertyIndex;
 
 	Builder.RegisterObjectVariant(Owner, Property);
@@ -244,80 +232,23 @@ bool FGLTFEpicLevelVariantSetsConverter::TryParseMaterialPropertyValue(FGLTFJson
 	return true;
 }
 
-bool FGLTFEpicLevelVariantSetsConverter::TryParseStaticMeshPropertyValue(FGLTFJsonEpicVariant& OutVariant, const UPropertyValue* Property) const
+template <typename MeshType>
+bool FGLTFEpicLevelVariantSetsConverter::TryParseMeshPropertyValue(FGLTFJsonEpicVariant& OutVariant, const UPropertyValue* Property) const
 {
-	const USceneComponent* Target = static_cast<USceneComponent*>(Property->GetPropertyParentContainerAddress());
+	const UMeshComponent* Target = static_cast<UMeshComponent*>(Property->GetPropertyParentContainerAddress());
 	if (Target == nullptr)
 	{
 		Builder.LogWarning(FString::Printf(
-			TEXT("Target object for property is invalid, the property will be skipped. Context: %s"),
+			TEXT("Variant property %s must belong to a mesh component, the property will be skipped"),
 			*FGLTFVariantUtility::GetLogContext(Property)));
 		return false;
 	}
 
 	const AActor* Owner = Target->GetOwner();
 	if (Owner == nullptr)
-	{
-		Builder.LogWarning(FString::Printf(TEXT("Invalid scene component, the property will be skipped. Context: %s"),
-			*FGLTFVariantUtility::GetLogContext(Property)));
-		return false;
-	}
-
-	if (Builder.bSelectedActorsOnly && !Owner->IsSelected())
-	{
-		Builder.LogWarning(FString::Printf(
-			TEXT("Target actor for property is not selected for export, the property will be skipped. Context: %s"),
-			*FGLTFVariantUtility::GetLogContext(Property)));
-		return false;
-	}
-
-	const UMeshComponent* MeshComponent = Cast<UMeshComponent>(Target);
-	if (MeshComponent == nullptr)
-	{
-		Builder.LogWarning(FString::Printf(
-			TEXT("Target object for property has no mesh-component, the property will be skipped. Context: %s"),
-			*FGLTFVariantUtility::GetLogContext(Property)));
-		return false;
-	}
-
-	const UStaticMesh* StaticMesh;
-
-	if (!FGLTFVariantUtility::TryGetPropertyValue(const_cast<UPropertyValue*>(Property), StaticMesh))
-	{
-		Builder.LogWarning(FString::Printf(
-			TEXT("Failed to parse recorded data for property, it will be skipped. Context: %s"),
-			*FGLTFVariantUtility::GetLogContext(Property)));
-		return false;
-	}
-
-	Builder.RegisterObjectVariant(Owner, Property);
-
-	const FGLTFMaterialArray OverrideMaterials(MeshComponent->OverrideMaterials);
-	const FGLTFJsonNodeIndex NodeIndex = Builder.GetOrAddNode(Target);
-	const FGLTFJsonNodeIndex ComponentNodeIndex = Builder.GetComponentNodeIndex(NodeIndex);
-	const FGLTFJsonMeshIndex MeshIndex = Builder.GetOrAddMesh(StaticMesh, OverrideMaterials);
-	FGLTFJsonEpicVariantNodeProperties& NodeProperties = OutVariant.Nodes.FindOrAdd(ComponentNodeIndex);
-
-	NodeProperties.Node = ComponentNodeIndex;
-	NodeProperties.Mesh = MeshIndex;
-	return true;
-}
-
-bool FGLTFEpicLevelVariantSetsConverter::TryParseSkeletalMeshPropertyValue(FGLTFJsonEpicVariant& OutVariant, const UPropertyValue* Property) const
-{
-	const USceneComponent* Target = static_cast<USceneComponent*>(Property->GetPropertyParentContainerAddress());
-	if (Target == nullptr)
 	{
 		Builder.LogWarning(FString::Printf(
-			TEXT("Target object for property is invalid, the property will be skipped. Context: %s"),
-			*FGLTFVariantUtility::GetLogContext(Property)));
-		return false;
-	}
-
-	const AActor* Owner = Target->GetOwner();
-	if (Owner == nullptr)
-	{
-		Builder.LogWarning(FString::Printf(TEXT("Invalid scene component, the property will be skipped. Context: %s"),
+			TEXT("Variant property %s must belong to an actor, the property will be skipped"),
 			*FGLTFVariantUtility::GetLogContext(Property)));
 		return false;
 	}
@@ -325,35 +256,26 @@ bool FGLTFEpicLevelVariantSetsConverter::TryParseSkeletalMeshPropertyValue(FGLTF
 	if (Builder.bSelectedActorsOnly && !Owner->IsSelected())
 	{
 		Builder.LogWarning(FString::Printf(
-			TEXT("Target actor for property is not selected for export, the property will be skipped. Context: %s"),
-			*FGLTFVariantUtility::GetLogContext(Property)));
-		return false;
-	}
-
-	const UMeshComponent* MeshComponent = Cast<UMeshComponent>(Target);
-	if (MeshComponent == nullptr)
-	{
-		Builder.LogWarning(FString::Printf(
-			TEXT("Target object for property has no mesh-component, the property will be skipped. Context: %s"),
+			TEXT("Variant property %s doesn't belong to an actor selected for export, the property will be skipped"),
 			*FGLTFVariantUtility::GetLogContext(Property)));
 		return false;
 	}
 
-	const USkeletalMesh* SkeletalMesh;
+	const MeshType* Mesh;
 
-	if (!FGLTFVariantUtility::TryGetPropertyValue(const_cast<UPropertyValue*>(Property), SkeletalMesh) || SkeletalMesh == nullptr)
+	if (!FGLTFVariantUtility::TryGetPropertyValue(const_cast<UPropertyValue*>(Property), Mesh))
 	{
 		Builder.LogWarning(FString::Printf(
-			TEXT("Failed to parse recorded data for property, it will be skipped. Context: %s"),
+			TEXT("Failed to parse recorded data for variant property %s, it will be skipped"),
 			*FGLTFVariantUtility::GetLogContext(Property)));
 		return false;
 	}
 
 	Builder.RegisterObjectVariant(Owner, Property);
 
-	const FGLTFMaterialArray OverrideMaterials(MeshComponent->OverrideMaterials);
+	const FGLTFMaterialArray OverrideMaterials(Target->OverrideMaterials);
 	const FGLTFJsonNodeIndex NodeIndex = Builder.GetOrAddNode(Target);
-	const FGLTFJsonMeshIndex MeshIndex = Builder.GetOrAddMesh(SkeletalMesh, OverrideMaterials);
+	const FGLTFJsonMeshIndex MeshIndex = Builder.GetOrAddMesh(Mesh, OverrideMaterials);
 	const FGLTFJsonNodeIndex ComponentNodeIndex = Builder.GetComponentNodeIndex(NodeIndex);
 	FGLTFJsonEpicVariantNodeProperties& NodeProperties = OutVariant.Nodes.FindOrAdd(ComponentNodeIndex);
 
diff --git a/Engine/Plugins/Enterprise/GLTFExporter/MigratedFromGit/Source/GLTFExporter/Private/Converters/GLTFEpicLevelVariantSetsConverters.h b/Engine/Plugins/Enterprise/GLTFExporter/MigratedFromGit/Source/GLTFExporter/Private/Converters/GLTFEpicLevelVariantSetsConverters.h
--- a/Engine/Plugins/Enterprise/GLTFExporter/MigratedFromGit/Source/GLTFExporter/Private/Converters/GLTFEpicLevelVariantSetsConverters.h
+++ b/Engine/Plugins/Enterprise/GLTFExporter/MigratedFromGit/Source/GLTFExporter/Private/Converters/GLTFEpicLevelVariantSetsConverters.h
@@ -12,6 +12,7 @@ class UVariantSet;
 class UVariant;
 class UVariantObjectBinding;
 class UPropertyValue;
+class UPropertyValueMaterial;
 
 class FGLTFEpicLevelVariantSetsConverter final : public FGLTFBuilderContext, public TGLTFConverter<FGLTFJsonEpicLevelVariantSetsIndex, const ULevelVariantSets*>
 {
@@ -22,7 +23,8 @@ class FGLTFEpicLevelVariantSetsConverter final : public FGLTFBuilderContext, pub
 	bool TryParseVariant(FGLTFJsonEpicVariant& OutVariant, const UVariant* Variant) const;
 	bool TryParseVariantBinding(FGLTFJsonEpicVariant& OutVariant, const UVariantObjectBinding* Binding) const;
 	bool TryParseVisibilityPropertyValue(FGLTFJsonEpicVariant& OutVariant, const UPropertyValue* Property) const;
-	bool TryParseMaterialPropertyValue(FGLTFJsonEpicVariant& OutVariant, const UPropertyValue* Property) const;
-	bool TryParseStaticMeshPropertyValue(FGLTFJsonEpicVariant& OutVariant, const UPropertyValue* Property) const;
-	bool TryParseSkeletalMeshPropertyValue(FGLTFJsonEpicVariant& OutVariant, const UPropertyValue* Property) const;
+	bool TryParseMaterialPropertyValue(FGLTFJsonEpicVariant& OutVariant, const UPropertyValueMaterial* Property) const;
+
+	template <typename MeshType>
+	bool TryParseMeshPropertyValue(FGLTFJsonEpicVariant& OutVariant, const UPropertyValue* Property) const;
 };
diff --git a/Engine/Plugins/Enterprise/GLTFExporter/MigratedFromGit/Source/GLTFExporter/Private/Converters/GLTFKhrVariantConverters.cpp b/Engine/Plugins/Enterprise/GLTFExporter/MigratedFromGit/Source/GLTFExporter/Private/Converters/GLTFKhrVariantConverters.cpp
--- a/Engine/Plugins/Enterprise/GLTFExporter/MigratedFromGit/Source/GLTFExporter/Private/Converters/GLTFKhrVariantConverters.cpp
+++ b/Engine/Plugins/Enterprise/GLTFExporter/MigratedFromGit/Source/GLTFExporter/Private/Converters/GLTFKhrVariantConverters.cpp
@@ -2,9 +2,7 @@
 
 #include "Converters/GLTFKhrVariantConverters.h"
 #include "Converters/GLTFVariantUtility.h"
-#include "Converters/GLTFMeshUtility.h"
 #include "Builders/GLTFContainerBuilder.h"
-#include "Rendering/SkeletalMeshRenderData.h"
 #include "VariantObjectBinding.h"
 #include "PropertyValueMaterial.h"
 #include "PropertyValue.h"
@@ -36,12 +34,12 @@ FGLTFJsonKhrMaterialVariantIndex FGLTFKhrMaterialVariantConverter::Convert(const
 				continue;
 			}
 
-			if (Property->IsA<UPropertyValueMaterial>())
+			if (const UPropertyValueMaterial* MaterialProperty = Cast<UPropertyValueMaterial>(Property))
 			{
 				FGLTFJsonPrimitive* Primitive = nullptr;
 				FGLTFJsonMaterialIndex MaterialIndex;
 
-				if (TryParseMaterialProperty(Primitive, MaterialIndex, Property))
+				if (TryParseMaterialProperty(Primitive, MaterialIndex, MaterialProperty))
 				{
 					PrimitiveMaterials.Add(MakeTuple(Primitive, MaterialIndex));
 				}
@@ -86,85 +84,68 @@ FGLTFJsonKhrMaterialVariantIndex FGLTFKhrMaterialVariantConverter::Convert(const
 	return MaterialVariantIndex;
 }
 
-bool FGLTFKhrMaterialVariantConverter::TryParseMaterialProperty(FGLTFJsonPrimitive*& OutPrimitive, FGLTFJsonMaterialIndex& OutMaterialIndex, const UPropertyValue* Property) const
+bool FGLTFKhrMaterialVariantConverter::TryParseMaterialProperty(FGLTFJsonPrimitive*& OutPrimitive, FGLTFJsonMaterialIndex& OutMaterialIndex, const UPropertyValueMaterial* Property) const
 {
-	UPropertyValueMaterial* MaterialProperty = Cast<UPropertyValueMaterial>(const_cast<UPropertyValue*>(Property));
-	if (MaterialProperty == nullptr)
-	{
-		Builder.LogWarning(FString::Printf(
-			TEXT("Material property is invalid, the property will be skipped. Context: %s"),
-			*FGLTFVariantUtility::GetLogContext(Property)));
-		return false;
-	}
-
-	const USceneComponent* Target = static_cast<USceneComponent*>(Property->GetPropertyParentContainerAddress());
+	const UMeshComponent* Target = static_cast<UMeshComponent*>(Property->GetPropertyParentContainerAddress());
 	if (Target == nullptr)
 	{
 		Builder.LogWarning(FString::Printf(
-			TEXT("Target object for property is invalid, the property will be skipped. Context: %s"),
+			TEXT("Variant property %s must belong to a mesh component, the property will be skipped"),
 			*FGLTFVariantUtility::GetLogContext(Property)));
 		return false;
 	}
 
 	const AActor* Owner = Target->GetOwner();
 	if (Owner == nullptr)
-	{
-		Builder.LogWarning(FString::Printf(TEXT("Invalid scene component, the property will be skipped. Context: %s"),
-			*FGLTFVariantUtility::GetLogContext(Property)));
-		return false;
-	}
-
-	if (Builder.bSelectedActorsOnly && !Owner->IsSelected())
 	{
 		Builder.LogWarning(FString::Printf(
-			TEXT("Target actor for property is not selected for export, the property will be skipped. Context: %s"),
+			TEXT("Variant property %s must belong to an actor, the property will be skipped"),
 			*FGLTFVariantUtility::GetLogContext(Property)));
 		return false;
 	}
 
-	const UMeshComponent* MeshComponent = Cast<UMeshComponent>(Target);
-	if (MeshComponent == nullptr)
+	if (Builder.bSelectedActorsOnly && !Owner->IsSelected())
 	{
 		Builder.LogWarning(FString::Printf(
-			TEXT("Target object for property has no mesh-component, the property will be skipped. Context: %s"),
+			TEXT("Variant property %s doesn't belong to an actor selected for export, the property will be skipped"),
 			*FGLTFVariantUtility::GetLogContext(Property)));
 		return false;
 	}
 
-	const TArray<FCapturedPropSegment>& CapturedPropSegments = FGLTFVariantUtility::GetCapturedPropSegments(MaterialProperty);
+	const TArray<FCapturedPropSegment>& CapturedPropSegments = FGLTFVariantUtility::GetCapturedPropSegments(Property);
 	const int32 NumPropSegments = CapturedPropSegments.Num();
 
 	if (NumPropSegments < 1)
 	{
 		Builder.LogWarning(FString::Printf(
-			TEXT("Failed to parse element index to apply the material to, the property will be skipped. Context: %s"),
-			*FGLTFVariantUtility::GetLogContext(MaterialProperty)));
+			TEXT("Failed to parse material index for variant property %s, the property will be skipped"),
+			*FGLTFVariantUtility::GetLogContext(Property)));
 		return false;
 	}
 
 	// NOTE: UPropertyValueMaterial::GetMaterial does *not* ensure that the recorded data has been loaded,
 	// so we need to call UProperty::GetRecordedData first to make that happen.
-	MaterialProperty->GetRecordedData();
+	const_cast<UPropertyValueMaterial*>(Property)->GetRecordedData();
 
-	const UMaterialInterface* Material = MaterialProperty->GetMaterial();
+	const UMaterialInterface* Material = const_cast<UPropertyValueMaterial*>(Property)->GetMaterial();
 	if (Material == nullptr)
 	{
 		// TODO: find way to determine whether the material is null because "None" was selected, or because it failed to resolve
 
 		Builder.LogWarning(FString::Printf(
 			TEXT("No material assigned, the property will be skipped. Context: %s"),
-			*FGLTFVariantUtility::GetLogContext(MaterialProperty)));
+			*FGLTFVariantUtility::GetLogContext(Property)));
 		return false;
 	}
 
-	Builder.RegisterObjectVariant(MeshComponent, MaterialProperty);
-	Builder.RegisterObjectVariant(Owner, MaterialProperty);
+	Builder.RegisterObjectVariant(Target, Property);
+	Builder.RegisterObjectVariant(Owner, Property);
 
 	const int32 MaterialIndex = CapturedPropSegments[NumPropSegments - 1].PropertyIndex;
-	const FGLTFJsonMeshIndex MeshIndex = Builder.GetOrAddMesh(MeshComponent);
+	const FGLTFJsonMeshIndex MeshIndex = Builder.GetOrAddMesh(Target);
 
 	OutPrimitive = &Builder.GetMesh(MeshIndex).Primitives[MaterialIndex];
-	OutMaterialIndex = FGLTFVariantUtility::GetOrAddMaterial(Builder, Material, MeshComponent, MaterialIndex);
+	OutMaterialIndex = FGLTFVariantUtility::GetOrAddMaterial(Builder, Material, Target, MaterialIndex);
 
 	return true;
 }
diff --git a/Engine/Plugins/Enterprise/GLTFExporter/MigratedFromGit/Source/GLTFExporter/Private/Converters/GLTFKhrVariantConverters.h b/Engine/Plugins/Enterprise/GLTFExporter/MigratedFromGit/Source/GLTFExporter/Private/Converters/GLTFKhrVariantConverters.h
--- a/Engine/Plugins/Enterprise/GLTFExporter/MigratedFromGit/Source/GLTFExporter/Private/Converters/GLTFKhrVariantConverters.h
+++ b/Engine/Plugins/Enterprise/GLTFExporter/MigratedFromGit/Source/GLTFExporter/Private/Converters/GLTFKhrVariantConverters.h
@@ -6,9 +6,9 @@
 #include "Json/GLTFJsonMesh.h"
 #include "Converters/GLTFConverter.h"
 #include "Converters/GLTFBuilderContext.h"
-#include "Variant.h"
-#include "PropertyValue.h"
-#include "LevelVariantSets.h"
+
+class UVariant;
+class UPropertyValueMaterial;
 
 class FGLTFKhrMaterialVariantConverter final : public FGLTFBuilderContext, public TGLTFConverter<FGLTFJsonKhrMaterialVariantIndex, const UVariant*>
 {
@@ -16,5 +16,5 @@ class FGLTFKhrMaterialVariantConverter final : public FGLTFBuilderContext, publi
 
 	virtual FGLTFJsonKhrMaterialVariantIndex Convert(const UVariant* Variant) override;
 
-	bool TryParseMaterialProperty(FGLTFJsonPrimitive*& OutPrimitive, FGLTFJsonMaterialIndex& OutMaterialIndex, const UPropertyValue* Property) const;
+	bool TryParseMaterialProperty(FGLTFJsonPrimitive*& OutPrimitive, FGLTFJsonMaterialIndex& OutMaterialIndex, const UPropertyValueMaterial* Property) const;
 };
diff --git a/Engine/Plugins/Enterprise/GLTFExporter/MigratedFromGit/Source/GLTFExporter/Private/Converters/GLTFSceneConverters.cpp b/Engine/Plugins/Enterprise/GLTFExporter/MigratedFromGit/Source/GLTFExporter/Private/Converters/GLTFSceneConverters.cpp
--- a/Engine/Plugins/Enterprise/GLTFExporter/MigratedFromGit/Source/GLTFExporter/Private/Converters/GLTFSceneConverters.cpp
+++ b/Engine/Plugins/Enterprise/GLTFExporter/MigratedFromGit/Source/GLTFExporter/Private/Converters/GLTFSceneConverters.cpp
@@ -2,9 +2,9 @@
 
 #include "Converters/GLTFSceneConverters.h"
 #include "Builders/GLTFContainerBuilder.h"
-#include "Converters/GLTFConverterUtility.h"
 #include "Converters/GLTFActorUtility.h"
 #include "LevelVariantSetsActor.h"
+#include "LevelVariantSets.h"
 #include "VariantSet.h"
 
 FGLTFJsonSceneIndex FGLTFSceneConverter::Convert(const UWorld* World)
