[GLTFExporter] refactor: remove final specifier in exporter classes to ease support for custom assets

commit af1c2c04cec2fc32f78ea22fa8bdbb2095600db4
Date: Wed, 12 Jan 2022 18:30:25 +0100

    refactor: remove final specifier in exporter classes to ease support for custom assets
    (cherry picked from commit af7e7a56fc054f99700f40d677dcf00e18536b43​)

#skipci

diff --git a/Engine/Plugins/Enterprise/GLTFExporter/MigratedFromGit/Source/GLTFExporter/Public/Exporters/GLTFAnimSequenceExporter.h b/Engine/Plugins/Enterprise/GLTFExporter/MigratedFromGit/Source/GLTFExporter/Public/Exporters/GLTFAnimSequenceExporter.h
--- a/Engine/Plugins/Enterprise/GLTFExporter/MigratedFromGit/Source/GLTFExporter/Public/Exporters/GLTFAnimSequenceExporter.h
+++ b/Engine/Plugins/Enterprise/GLTFExporter/MigratedFromGit/Source/GLTFExporter/Public/Exporters/GLTFAnimSequenceExporter.h
@@ -6,7 +6,7 @@
 #include "GLTFAnimSequenceExporter.generated.h"
 
 UCLASS()
-class GLTFEXPORTER_API UGLTFAnimSequenceExporter final : public UGLTFExporter
+class GLTFEXPORTER_API UGLTFAnimSequenceExporter : public UGLTFExporter
 {
 public:
 
diff --git a/Engine/Plugins/Enterprise/GLTFExporter/MigratedFromGit/Source/GLTFExporter/Public/Exporters/GLTFLevelExporter.h b/Engine/Plugins/Enterprise/GLTFExporter/MigratedFromGit/Source/GLTFExporter/Public/Exporters/GLTFLevelExporter.h
--- a/Engine/Plugins/Enterprise/GLTFExporter/MigratedFromGit/Source/GLTFExporter/Public/Exporters/GLTFLevelExporter.h
+++ b/Engine/Plugins/Enterprise/GLTFExporter/MigratedFromGit/Source/GLTFExporter/Public/Exporters/GLTFLevelExporter.h
@@ -6,7 +6,7 @@
 #include "GLTFLevelExporter.generated.h"
 
 UCLASS()
-class GLTFEXPORTER_API UGLTFLevelExporter final : public UGLTFExporter
+class GLTFEXPORTER_API UGLTFLevelExporter : public UGLTFExporter
 {
 public:
 
diff --git a/Engine/Plugins/Enterprise/GLTFExporter/MigratedFromGit/Source/GLTFExporter/Public/Exporters/GLTFLevelSequenceExporter.h b/Engine/Plugins/Enterprise/GLTFExporter/MigratedFromGit/Source/GLTFExporter/Public/Exporters/GLTFLevelSequenceExporter.h
--- a/Engine/Plugins/Enterprise/GLTFExporter/MigratedFromGit/Source/GLTFExporter/Public/Exporters/GLTFLevelSequenceExporter.h
+++ b/Engine/Plugins/Enterprise/GLTFExporter/MigratedFromGit/Source/GLTFExporter/Public/Exporters/GLTFLevelSequenceExporter.h
@@ -6,7 +6,7 @@
 #include "GLTFLevelSequenceExporter.generated.h"
 
 UCLASS()
-class GLTFEXPORTER_API UGLTFLevelSequenceExporter final : public UGLTFExporter
+class GLTFEXPORTER_API UGLTFLevelSequenceExporter : public UGLTFExporter
 {
 public:
 
diff --git a/Engine/Plugins/Enterprise/GLTFExporter/MigratedFromGit/Source/GLTFExporter/Public/Exporters/GLTFLevelVariantSetsExporter.h b/Engine/Plugins/Enterprise/GLTFExporter/MigratedFromGit/Source/GLTFExporter/Public/Exporters/GLTFLevelVariantSetsExporter.h
--- a/Engine/Plugins/Enterprise/GLTFExporter/MigratedFromGit/Source/GLTFExporter/Public/Exporters/GLTFLevelVariantSetsExporter.h
+++ b/Engine/Plugins/Enterprise/GLTFExporter/MigratedFromGit/Source/GLTFExporter/Public/Exporters/GLTFLevelVariantSetsExporter.h
@@ -6,7 +6,7 @@
 #include "GLTFLevelVariantSetsExporter.generated.h"
 
 UCLASS()
-class GLTFEXPORTER_API UGLTFLevelVariantSetsExporter final : public UGLTFExporter
+class GLTFEXPORTER_API UGLTFLevelVariantSetsExporter : public UGLTFExporter
 {
 public:
 
diff --git a/Engine/Plugins/Enterprise/GLTFExporter/MigratedFromGit/Source/GLTFExporter/Public/Exporters/GLTFMaterialExporter.h b/Engine/Plugins/Enterprise/GLTFExporter/MigratedFromGit/Source/GLTFExporter/Public/Exporters/GLTFMaterialExporter.h
--- a/Engine/Plugins/Enterprise/GLTFExporter/MigratedFromGit/Source/GLTFExporter/Public/Exporters/GLTFMaterialExporter.h
+++ b/Engine/Plugins/Enterprise/GLTFExporter/MigratedFromGit/Source/GLTFExporter/Public/Exporters/GLTFMaterialExporter.h
@@ -6,7 +6,7 @@
 #include "GLTFMaterialExporter.generated.h"
 
 UCLASS()
-class GLTFEXPORTER_API UGLTFMaterialExporter final : public UGLTFExporter
+class GLTFEXPORTER_API UGLTFMaterialExporter : public UGLTFExporter
 {
 public:
 
diff --git a/Engine/Plugins/Enterprise/GLTFExporter/MigratedFromGit/Source/GLTFExporter/Public/Exporters/GLTFSkeletalMeshExporter.h b/Engine/Plugins/Enterprise/GLTFExporter/MigratedFromGit/Source/GLTFExporter/Public/Exporters/GLTFSkeletalMeshExporter.h
--- a/Engine/Plugins/Enterprise/GLTFExporter/MigratedFromGit/Source/GLTFExporter/Public/Exporters/GLTFSkeletalMeshExporter.h
+++ b/Engine/Plugins/Enterprise/GLTFExporter/MigratedFromGit/Source/GLTFExporter/Public/Exporters/GLTFSkeletalMeshExporter.h
@@ -6,7 +6,7 @@
 #include "GLTFSkeletalMeshExporter.generated.h"
 
 UCLASS()
-class GLTFEXPORTER_API UGLTFSkeletalMeshExporter final : public UGLTFExporter
+class GLTFEXPORTER_API UGLTFSkeletalMeshExporter : public UGLTFExporter
 {
 public:
 
diff --git a/Engine/Plugins/Enterprise/GLTFExporter/MigratedFromGit/Source/GLTFExporter/Public/Exporters/GLTFStaticMeshExporter.h b/Engine/Plugins/Enterprise/GLTFExporter/MigratedFromGit/Source/GLTFExporter/Public/Exporters/GLTFStaticMeshExporter.h
--- a/Engine/Plugins/Enterprise/GLTFExporter/MigratedFromGit/Source/GLTFExporter/Public/Exporters/GLTFStaticMeshExporter.h
+++ b/Engine/Plugins/Enterprise/GLTFExporter/MigratedFromGit/Source/GLTFExporter/Public/Exporters/GLTFStaticMeshExporter.h
@@ -6,7 +6,7 @@
 #include "GLTFStaticMeshExporter.generated.h"
 
 UCLASS()
-class GLTFEXPORTER_API UGLTFStaticMeshExporter final : public UGLTFExporter
+class GLTFEXPORTER_API UGLTFStaticMeshExporter : public UGLTFExporter
 {
 public:
 
