[GLTFExporter] refactor: minimize copying and hashing of input parameters to all converters

commit 4341318b754aa68d45d4636005a6b2888cae48f4
Date: Fri, 14 Jan 2022 08:47:54 +0100

    refactor: minimize copying and hashing of input parameters to all converters
    (cherry picked from commit dcd41ca6f2d091b052ca96a5c2b381c82c08f11d​)

#skipci

diff --git a/Engine/Plugins/Enterprise/GLTFExporter/MigratedFromGit/Source/GLTFExporter/Private/Converters/GLTFConverter.h b/Engine/Plugins/Enterprise/GLTFExporter/MigratedFromGit/Source/GLTFExporter/Private/Converters/GLTFConverter.h
--- a/Engine/Plugins/Enterprise/GLTFExporter/MigratedFromGit/Source/GLTFExporter/Private/Converters/GLTFConverter.h
+++ b/Engine/Plugins/Enterprise/GLTFExporter/MigratedFromGit/Source/GLTFExporter/Private/Converters/GLTFConverter.h
@@ -7,19 +7,21 @@
 template <typename OutputType, typename... InputTypes>
 class TGLTFConverter
 {
-	typedef TTuple<InputTypes...> InputKeyType;
+	typedef TTuple<InputTypes...> InputTupleType;
 
 public:
 
 	virtual ~TGLTFConverter() = default;
 
-	OutputType Get(InputTypes&&... Inputs) const
+	OutputType Get(InputTypes... Inputs) const
 	{
 		Sanitize(Inputs...);
-		const InputKeyType InputKey(Forward<InputTypes>(Inputs)...);
-		if (OutputType* SavedOutput = SavedOutputs.Find(InputKey))
+		InputTupleType InputTuple(MoveTemp(Inputs)...);
+		const uint32 InputHash = GetTypeHash(InputTuple);
+
+		if (OutputType* CachedOutput = CachedOutputs.FindByHash(InputHash, InputTuple))
 		{
-			return *SavedOutput;
+			return *CachedOutput;
 		}
 
 		return {};
@@ -28,25 +30,27 @@ public:
 	OutputType Add(InputTypes... Inputs)
 	{
 		Sanitize(Inputs...);
-		const InputKeyType InputKey(Inputs...);
-		OutputType NewOutput = Convert(Forward<InputTypes>(Inputs)...);
+		InputTupleType InputTuple(MoveTemp(Inputs)...);
+		const uint32 InputHash = GetTypeHash(InputTuple);
 
-		SavedOutputs.Add(InputKey, NewOutput);
+		OutputType NewOutput = InputTuple.ApplyAfter(&TGLTFConverter::Convert, this);
+		CachedOutputs.AddByHash(InputHash, MoveTemp(InputTuple), NewOutput);
 		return NewOutput;
 	}
 
 	OutputType GetOrAdd(InputTypes... Inputs)
 	{
 		Sanitize(Inputs...);
-		const InputKeyType InputKey(Inputs...);
-		if (OutputType* SavedOutput = SavedOutputs.Find(InputKey))
+		InputTupleType InputTuple(MoveTemp(Inputs)...);
+		const uint32 InputHash = GetTypeHash(InputTuple);
+
+		if (OutputType* CachedOutput = CachedOutputs.FindByHash(InputHash, InputTuple))
 		{
-			return *SavedOutput;
+			return *CachedOutput;
 		}
 
-		OutputType NewOutput = Convert(Forward<InputTypes>(Inputs)...);
-
-		SavedOutputs.Add(InputKey, NewOutput);
+		OutputType NewOutput = InputTuple.ApplyAfter(&TGLTFConverter::Convert, this);
+		CachedOutputs.AddByHash(InputHash, MoveTemp(InputTuple), NewOutput);
 		return NewOutput;
 	}
 
@@ -58,5 +62,5 @@ protected:
 
 private:
 
-	TMap<InputKeyType, OutputType> SavedOutputs;
+	TMap<InputTupleType, OutputType> CachedOutputs;
 };
