[GLTFExporter] chore: update export options reference documentation

commit 6e5c8057428da740e7ce448772bb194bf7b9e052
Date: Sun, 16 Jan 2022 22:25:58 +0100

    chore: update export options reference documentation
    (cherry picked from commit a21934aec7b8780192dc9ed06ccd529ef4a2f9c2​)

#skipci

diff --git a/Engine/Plugins/Enterprise/GLTFExporter/MigratedFromGit/Docs/export-options-reference.md b/Engine/Plugins/Enterprise/GLTFExporter/MigratedFromGit/Docs/export-options-reference.md
--- a/Engine/Plugins/Enterprise/GLTFExporter/MigratedFromGit/Docs/export-options-reference.md
+++ b/Engine/Plugins/Enterprise/GLTFExporter/MigratedFromGit/Docs/export-options-reference.md
@@ -28,14 +28,14 @@ Option                         | Description
 `Export Texture Transforms`    | If enabled, export UV tiling and un-mirroring settings in a texture coordinate expression node for simple material input expressions. Uses extension KHR_texture_transform.
 `Export Lightmaps`             | If enabled, export lightmaps (created by Lightmass) when exporting a level. Uses extension EPIC_lightmap_textures, which is supported by Unreal's glTF viewer.
 `Texture HDR Encoding`         | Encoding used to store textures that have pixel colors with more than 8-bit per channel. Uses extension EPIC_texture_hdr_encoding, which is supported by Unreal's glTF viewer.
-`Export Hidden In Game`        | If enabled, export components that are flagged as hidden in-game.
+`Export Hidden In Game`        | If enabled, export actors and components that are flagged as hidden in-game.
 `Export Lights`                | Mobility of directional, point, and spot light components that will be exported. Uses extension KHR_lights_punctual.
 `Export Cameras`               | If enabled, export camera components.
 `Export Camera Controls`       | If enabled, export GLTFCameraActors. Uses extension EPIC_camera_controls, which is supported by Unreal's glTF viewer.
 `Export Animation Hotspots`    | If enabled, export GLTFHotspotActors. Uses extension EPIC_animation_hotspots, which is supported by Unreal's glTF viewer.
 `Export HDRI Backdrops`        | If enabled, export HDRIBackdrop blueprints. Uses extension EPIC_hdri_backdrops, which is supported by Unreal's glTF viewer.
 `Export Sky Spheres`           | If enabled, export SkySphere blueprints. Uses extension EPIC_sky_spheres, which is supported by Unreal's glTF viewer.
-`Export Variant Sets`          | If enabled, export LevelVariantSetsActors. Uses extension EPIC_level_variant_sets, which is supported by Unreal's glTF viewer.
-`Export Material Variants`      | Mode determining if and how to export material variants that change the materials property on a static or skeletal mesh component.
+`Variant Sets Mode`            | Mode determining if and how to export LevelVariantSetsActors.
+`Export Material Variants`     | Mode determining if and how to export material variants that change the materials property on a static or skeletal mesh component.
 `Export Mesh Variants`         | If enabled, export variants that change the mesh property on a static or skeletal mesh component.
 `Export Visibility Variants`   | If enabled, export variants that change the visible property on a scene component.
