[GLTFExporter] refactor: move some common variant checks into utility functions

commit 6909b423caa222201a2c1ffb850ba9b5e94676eb
Date: Tue, 18 Jan 2022 06:42:55 +0100

    refactor: move some common variant checks into utility functions
    (cherry picked from commit 90eb2d6dffa110ec4e994092a86b9eda10f827a5​)

#skipci

diff --git a/Engine/Plugins/Enterprise/GLTFExporter/MigratedFromGit/Source/GLTFExporter/Private/Converters/GLTFEpicLevelVariantSetsConverters.cpp b/Engine/Plugins/Enterprise/GLTFExporter/MigratedFromGit/Source/GLTFExporter/Private/Converters/GLTFEpicLevelVariantSetsConverters.cpp
--- a/Engine/Plugins/Enterprise/GLTFExporter/MigratedFromGit/Source/GLTFExporter/Private/Converters/GLTFEpicLevelVariantSetsConverters.cpp
+++ b/Engine/Plugins/Enterprise/GLTFExporter/MigratedFromGit/Source/GLTFExporter/Private/Converters/GLTFEpicLevelVariantSetsConverters.cpp
@@ -82,9 +82,6 @@ bool FGLTFEpicLevelVariantSetsConverter::TryParseVariantBinding(FGLTFJsonEpicVar
 			continue;
 		}
 
-		const FName PropertyName = Property->GetPropertyName();
-		const FFieldClass* PropertyClass = Property->GetPropertyClass();
-
 		if (const UPropertyValueMaterial* MaterialProperty = Cast<UPropertyValueMaterial>(Property))
 		{
 			if (Builder.ExportOptions->ExportMaterialVariants != EGLTFMaterialVariantMode::None && TryParseMaterialPropertyValue(OutVariant, MaterialProperty))
@@ -92,21 +89,21 @@ bool FGLTFEpicLevelVariantSetsConverter::TryParseVariantBinding(FGLTFJsonEpicVar
 				bHasParsedAnyProperty = true;
 			}
 		}
-		else if (PropertyName == TEXT("StaticMesh")) // TODO: should we not also check PropertyClass?
+		else if (FGLTFVariantUtility::IsStaticMeshProperty(Property))
 		{
 			if (Builder.ExportOptions->bExportMeshVariants && TryParseMeshPropertyValue<UStaticMesh>(OutVariant, Property))
 			{
 				bHasParsedAnyProperty = true;
 			}
 		}
-		else if (PropertyName == TEXT("SkeletalMesh")) // TODO: should we not also check PropertyClass?
+		else if (FGLTFVariantUtility::IsSkeletalMeshProperty(Property))
 		{
 			if (Builder.ExportOptions->bExportMeshVariants && TryParseMeshPropertyValue<USkeletalMesh>(OutVariant, Property))
 			{
 				bHasParsedAnyProperty = true;
 			}
 		}
-		else if (PropertyName == TEXT("bVisible") && PropertyClass->IsChildOf(FBoolProperty::StaticClass()))
+		else if (FGLTFVariantUtility::IsVisibleProperty(Property))
 		{
 			if (Builder.ExportOptions->bExportVisibilityVariants && TryParseVisibilityPropertyValue(OutVariant, Property))
 			{
@@ -158,7 +155,7 @@ bool FGLTFEpicLevelVariantSetsConverter::TryParseVisibilityPropertyValue(FGLTFJs
 		return false;
 	}
 
-	Builder.RegisterObjectVariant(Owner, Property);
+	Builder.RegisterObjectVariant(Owner, Property); // TODO: we should register this on the component
 
 	const FGLTFJsonNodeIndex NodeIndex = Builder.GetOrAddNode(Target);
 	const FGLTFJsonNodeIndex ComponentNodeIndex = Builder.GetComponentNodeIndex(NodeIndex);
@@ -217,7 +214,7 @@ bool FGLTFEpicLevelVariantSetsConverter::TryParseMaterialPropertyValue(FGLTFJson
 	const UMaterialInterface* Material = const_cast<UPropertyValueMaterial*>(Property)->GetMaterial();
 	const int32 MaterialIndex = CapturedPropSegments[NumPropSegments - 1].PropertyIndex;
 
-	Builder.RegisterObjectVariant(Owner, Property);
+	Builder.RegisterObjectVariant(Owner, Property); // TODO: we should register this on the component
 
 	FGLTFJsonEpicVariantMaterial VariantMaterial;
 	VariantMaterial.Material = FGLTFVariantUtility::GetOrAddMaterial(Builder, Material, Target, MaterialIndex);
@@ -271,7 +268,7 @@ bool FGLTFEpicLevelVariantSetsConverter::TryParseMeshPropertyValue(FGLTFJsonEpic
 		return false;
 	}
 
-	Builder.RegisterObjectVariant(Owner, Property);
+	Builder.RegisterObjectVariant(Owner, Property); // TODO: we should register this on the component
 
 	const FGLTFMaterialArray OverrideMaterials(Target->OverrideMaterials);
 	const FGLTFJsonNodeIndex NodeIndex = Builder.GetOrAddNode(Target);
diff --git a/Engine/Plugins/Enterprise/GLTFExporter/MigratedFromGit/Source/GLTFExporter/Private/Converters/GLTFKhrVariantConverters.cpp b/Engine/Plugins/Enterprise/GLTFExporter/MigratedFromGit/Source/GLTFExporter/Private/Converters/GLTFKhrVariantConverters.cpp
--- a/Engine/Plugins/Enterprise/GLTFExporter/MigratedFromGit/Source/GLTFExporter/Private/Converters/GLTFKhrVariantConverters.cpp
+++ b/Engine/Plugins/Enterprise/GLTFExporter/MigratedFromGit/Source/GLTFExporter/Private/Converters/GLTFKhrVariantConverters.cpp
@@ -139,7 +139,7 @@ bool FGLTFKhrMaterialVariantConverter::TryParseMaterialProperty(FGLTFJsonPrimiti
 	}
 
 	Builder.RegisterObjectVariant(Target, Property);
-	Builder.RegisterObjectVariant(Owner, Property);
+	Builder.RegisterObjectVariant(Owner, Property); // TODO: we don't need to register this on the actor
 
 	const int32 MaterialIndex = CapturedPropSegments[NumPropSegments - 1].PropertyIndex;
 	const FGLTFJsonMeshIndex MeshIndex = Builder.GetOrAddMesh(Target);
diff --git a/Engine/Plugins/Enterprise/GLTFExporter/MigratedFromGit/Source/GLTFExporter/Private/Converters/GLTFVariantUtility.cpp b/Engine/Plugins/Enterprise/GLTFExporter/MigratedFromGit/Source/GLTFExporter/Private/Converters/GLTFVariantUtility.cpp
--- a/Engine/Plugins/Enterprise/GLTFExporter/MigratedFromGit/Source/GLTFExporter/Private/Converters/GLTFVariantUtility.cpp
+++ b/Engine/Plugins/Enterprise/GLTFExporter/MigratedFromGit/Source/GLTFExporter/Private/Converters/GLTFVariantUtility.cpp
@@ -27,6 +27,21 @@ const TArray<FCapturedPropSegment>& FGLTFVariantUtility::GetCapturedPropSegments
 #endif
 }
 
+bool FGLTFVariantUtility::IsStaticMeshProperty(const UPropertyValue* Property)
+{
+	return Property->GetPropertyName() == TEXT("StaticMesh") && Property->GetPropertyClass()->IsChildOf(FObjectPropertyBase::StaticClass());
+}
+
+bool FGLTFVariantUtility::IsSkeletalMeshProperty(const UPropertyValue* Property)
+{
+	return Property->GetPropertyName() == TEXT("SkeletalMesh") && Property->GetPropertyClass()->IsChildOf(FObjectPropertyBase::StaticClass());
+}
+
+bool FGLTFVariantUtility::IsVisibleProperty(const UPropertyValue* Property)
+{
+	return Property->GetPropertyName() == TEXT("bVisible") && Property->GetPropertyClass() == FBoolProperty::StaticClass();
+}
+
 bool FGLTFVariantUtility::TryGetPropertyValue(UPropertyValue* Property, void* OutData, uint32 OutSize)
 {
 	if (Property == nullptr || !Property->HasRecordedData())
diff --git a/Engine/Plugins/Enterprise/GLTFExporter/MigratedFromGit/Source/GLTFExporter/Private/Converters/GLTFVariantUtility.h b/Engine/Plugins/Enterprise/GLTFExporter/MigratedFromGit/Source/GLTFExporter/Private/Converters/GLTFVariantUtility.h
--- a/Engine/Plugins/Enterprise/GLTFExporter/MigratedFromGit/Source/GLTFExporter/Private/Converters/GLTFVariantUtility.h
+++ b/Engine/Plugins/Enterprise/GLTFExporter/MigratedFromGit/Source/GLTFExporter/Private/Converters/GLTFVariantUtility.h
@@ -16,6 +16,10 @@ struct FGLTFVariantUtility
 {
 	static const TArray<FCapturedPropSegment>& GetCapturedPropSegments(const UPropertyValue* Property);
 
+	static bool IsStaticMeshProperty(const UPropertyValue* Property);
+	static bool IsSkeletalMeshProperty(const UPropertyValue* Property);
+	static bool IsVisibleProperty(const UPropertyValue* Property);
+
 	template<typename T>
 	static bool TryGetPropertyValue(UPropertyValue* Property, T& OutValue)
 	{
