[GLTFExporter] feat: add SelectedActors parameter to UGLTFExporter::ExportToGLTF, enabling export of a selection in runtime

commit 70341edbc0bd024f09e5a372a2fc2bca85c971f4
Date: Tue, 18 Jan 2022 07:18:17 +0100

    feat: add SelectedActors parameter to UGLTFExporter::ExportToGLTF, enabling export of a selection in runtime
    (cherry picked from commit 4439a310fd11fe236358c86a94a0e9c2fccdbc1d​)

#skipci

diff --git a/Engine/Plugins/Enterprise/GLTFExporter/MigratedFromGit/Source/GLTFExporter/Private/Builders/GLTFContainerBuilder.cpp b/Engine/Plugins/Enterprise/GLTFExporter/MigratedFromGit/Source/GLTFExporter/Private/Builders/GLTFContainerBuilder.cpp
--- a/Engine/Plugins/Enterprise/GLTFExporter/MigratedFromGit/Source/GLTFExporter/Private/Builders/GLTFContainerBuilder.cpp
+++ b/Engine/Plugins/Enterprise/GLTFExporter/MigratedFromGit/Source/GLTFExporter/Private/Builders/GLTFContainerBuilder.cpp
@@ -3,6 +3,11 @@
 #include "Builders/GLTFContainerBuilder.h"
 #include "Builders/GLTFMemoryArchive.h"
 
+FGLTFContainerBuilder::FGLTFContainerBuilder(const FString& FilePath, const UGLTFExportOptions* ExportOptions, const TSet<AActor*>& SelectedActors)
+	: FGLTFConvertBuilder(FilePath, ExportOptions, SelectedActors)
+{
+}
+
 void FGLTFContainerBuilder::Write(FArchive& Archive, FFeedbackContext* Context)
 {
 	CompleteAllTasks(Context);
@@ -29,11 +34,6 @@ void FGLTFContainerBuilder::Write(FArchive& Archive, FFeedbackContext* Context)
 	}
 }
 
-FGLTFContainerBuilder::FGLTFContainerBuilder(const FString& FilePath, const UGLTFExportOptions* ExportOptions, bool bSelectedActorsOnly)
-	: FGLTFConvertBuilder(FilePath, ExportOptions, bSelectedActorsOnly)
-{
-}
-
 void FGLTFContainerBuilder::WriteGlb(FArchive& Archive)
 {
 	FGLTFMemoryArchive JsonData;
diff --git a/Engine/Plugins/Enterprise/GLTFExporter/MigratedFromGit/Source/GLTFExporter/Private/Builders/GLTFContainerBuilder.h b/Engine/Plugins/Enterprise/GLTFExporter/MigratedFromGit/Source/GLTFExporter/Private/Builders/GLTFContainerBuilder.h
--- a/Engine/Plugins/Enterprise/GLTFExporter/MigratedFromGit/Source/GLTFExporter/Private/Builders/GLTFContainerBuilder.h
+++ b/Engine/Plugins/Enterprise/GLTFExporter/MigratedFromGit/Source/GLTFExporter/Private/Builders/GLTFContainerBuilder.h
@@ -8,7 +8,7 @@ class FGLTFContainerBuilder : public FGLTFConvertBuilder
 {
 public:
 
-	FGLTFContainerBuilder(const FString& FilePath, const UGLTFExportOptions* ExportOptions, bool bSelectedActorsOnly);
+	FGLTFContainerBuilder(const FString& FilePath, const UGLTFExportOptions* ExportOptions, const TSet<AActor*>& SelectedActors);
 
 	void Write(FArchive& Archive, FFeedbackContext* Context);
 
diff --git a/Engine/Plugins/Enterprise/GLTFExporter/MigratedFromGit/Source/GLTFExporter/Private/Builders/GLTFConvertBuilder.cpp b/Engine/Plugins/Enterprise/GLTFExporter/MigratedFromGit/Source/GLTFExporter/Private/Builders/GLTFConvertBuilder.cpp
--- a/Engine/Plugins/Enterprise/GLTFExporter/MigratedFromGit/Source/GLTFExporter/Private/Builders/GLTFConvertBuilder.cpp
+++ b/Engine/Plugins/Enterprise/GLTFExporter/MigratedFromGit/Source/GLTFExporter/Private/Builders/GLTFConvertBuilder.cpp
@@ -7,12 +7,23 @@
 #include "Engine/TextureRenderTarget2D.h"
 #include "Engine/TextureRenderTargetCube.h"
 
-FGLTFConvertBuilder::FGLTFConvertBuilder(const FString& FilePath, const UGLTFExportOptions* ExportOptions, bool bSelectedActorsOnly)
+FGLTFConvertBuilder::FGLTFConvertBuilder(const FString& FilePath, const UGLTFExportOptions* ExportOptions, const TSet<AActor*>& SelectedActors)
 	: FGLTFImageBuilder(FilePath, ExportOptions)
-	, bSelectedActorsOnly(bSelectedActorsOnly)
+	, SelectedActors(SelectedActors)
 {
 }
 
+bool FGLTFConvertBuilder::IsSelectedActor(const AActor* Actor) const
+{
+	return SelectedActors.Num() == 0 || SelectedActors.Contains(Actor);
+}
+
+bool FGLTFConvertBuilder::IsRootActor(const AActor* Actor) const
+{
+	const AActor* ParentActor = Actor->GetAttachParentActor();
+	return ParentActor == nullptr || !IsSelectedActor(ParentActor);
+}
+
 FGLTFJsonAccessorIndex FGLTFConvertBuilder::GetOrAddPositionAccessor(const FGLTFMeshSection* MeshSection, const FPositionVertexBuffer* VertexBuffer)
 {
 	if (VertexBuffer == nullptr)
diff --git a/Engine/Plugins/Enterprise/GLTFExporter/MigratedFromGit/Source/GLTFExporter/Private/Builders/GLTFConvertBuilder.h b/Engine/Plugins/Enterprise/GLTFExporter/MigratedFromGit/Source/GLTFExporter/Private/Builders/GLTFConvertBuilder.h
--- a/Engine/Plugins/Enterprise/GLTFExporter/MigratedFromGit/Source/GLTFExporter/Private/Builders/GLTFConvertBuilder.h
+++ b/Engine/Plugins/Enterprise/GLTFExporter/MigratedFromGit/Source/GLTFExporter/Private/Builders/GLTFConvertBuilder.h
@@ -26,11 +26,14 @@ class FGLTFConvertBuilder : public FGLTFImageBuilder
 {
 protected:
 
-	FGLTFConvertBuilder(const FString& FilePath, const UGLTFExportOptions* ExportOptions, bool bSelectedActorsOnly);
+	FGLTFConvertBuilder(const FString& FilePath, const UGLTFExportOptions* ExportOptions, const TSet<AActor*>& SelectedActors);
 
 public:
 
-	const bool bSelectedActorsOnly;
+	const TSet<AActor*> SelectedActors;
+
+	bool IsSelectedActor(const AActor* Object) const;
+	bool IsRootActor(const AActor* Actor) const;
 
 	FGLTFJsonAccessorIndex GetOrAddPositionAccessor(const FGLTFMeshSection* MeshSection, const FPositionVertexBuffer* VertexBuffer);
 	FGLTFJsonAccessorIndex GetOrAddColorAccessor(const FGLTFMeshSection* MeshSection, const FColorVertexBuffer* VertexBuffer);
diff --git a/Engine/Plugins/Enterprise/GLTFExporter/MigratedFromGit/Source/GLTFExporter/Private/Converters/GLTFActorUtility.cpp b/Engine/Plugins/Enterprise/GLTFExporter/MigratedFromGit/Source/GLTFExporter/Private/Converters/GLTFActorUtility.cpp
--- a/Engine/Plugins/Enterprise/GLTFExporter/MigratedFromGit/Source/GLTFExporter/Private/Converters/GLTFActorUtility.cpp
+++ b/Engine/Plugins/Enterprise/GLTFExporter/MigratedFromGit/Source/GLTFExporter/Private/Converters/GLTFActorUtility.cpp
@@ -1,16 +1,9 @@
 // Copyright Epic Games, Inc. All Rights Reserved.
 
 #include "Converters/GLTFActorUtility.h"
-#include "Engine/Blueprint.h"
 #include "Actors/GLTFHotspotActor.h"
 #include "LevelSequenceActor.h"
 
-bool FGLTFActorUtility::IsRootActor(const AActor* Actor, bool bSelectedOnly)
-{
-	const AActor* ParentActor = Actor->GetAttachParentActor();
-	return ParentActor == nullptr || (bSelectedOnly && !ParentActor->IsSelected());
-}
-
 bool FGLTFActorUtility::IsGenericActor(const AActor* Actor)
 {
 	const FString BlueprintPath = GetBlueprintPath(Actor);
diff --git a/Engine/Plugins/Enterprise/GLTFExporter/MigratedFromGit/Source/GLTFExporter/Private/Converters/GLTFActorUtility.h b/Engine/Plugins/Enterprise/GLTFExporter/MigratedFromGit/Source/GLTFExporter/Private/Converters/GLTFActorUtility.h
--- a/Engine/Plugins/Enterprise/GLTFExporter/MigratedFromGit/Source/GLTFExporter/Private/Converters/GLTFActorUtility.h
+++ b/Engine/Plugins/Enterprise/GLTFExporter/MigratedFromGit/Source/GLTFExporter/Private/Converters/GLTFActorUtility.h
@@ -2,12 +2,10 @@
 
 #pragma once
 
-#include "Engine/EngineTypes.h"
+#include "CoreMinimal.h"
 
 struct FGLTFActorUtility
 {
-	static bool IsRootActor(const AActor* Actor, bool bSelectedOnly);
-
 	static bool IsGenericActor(const AActor* Actor);
 
 	static FString GetBlueprintPath(const AActor* Actor);
diff --git a/Engine/Plugins/Enterprise/GLTFExporter/MigratedFromGit/Source/GLTFExporter/Private/Converters/GLTFEpicLevelVariantSetsConverters.cpp b/Engine/Plugins/Enterprise/GLTFExporter/MigratedFromGit/Source/GLTFExporter/Private/Converters/GLTFEpicLevelVariantSetsConverters.cpp
--- a/Engine/Plugins/Enterprise/GLTFExporter/MigratedFromGit/Source/GLTFExporter/Private/Converters/GLTFEpicLevelVariantSetsConverters.cpp
+++ b/Engine/Plugins/Enterprise/GLTFExporter/MigratedFromGit/Source/GLTFExporter/Private/Converters/GLTFEpicLevelVariantSetsConverters.cpp
@@ -137,7 +137,7 @@ bool FGLTFEpicLevelVariantSetsConverter::TryParseVisibilityPropertyValue(FGLTFJs
 		return false;
 	}
 
-	if (Builder.bSelectedActorsOnly && !Owner->IsSelected())
+	if (!Builder.IsSelectedActor(Owner))
 	{
 		Builder.LogWarning(FString::Printf(
 			TEXT("Property doesn't belong to an actor selected for export, the property will be skipped. Context: %s"),
@@ -187,7 +187,7 @@ bool FGLTFEpicLevelVariantSetsConverter::TryParseMaterialPropertyValue(FGLTFJson
 		return false;
 	}
 
-	if (Builder.bSelectedActorsOnly && !Owner->IsSelected())
+	if (!Builder.IsSelectedActor(Owner))
 	{
 		Builder.LogWarning(FString::Printf(
 			TEXT("Variant property %s doesn't belong to an actor selected for export, the property will be skipped"),
@@ -250,7 +250,7 @@ bool FGLTFEpicLevelVariantSetsConverter::TryParseMeshPropertyValue(FGLTFJsonEpic
 		return false;
 	}
 
-	if (Builder.bSelectedActorsOnly && !Owner->IsSelected())
+	if (!Builder.IsSelectedActor(Owner))
 	{
 		Builder.LogWarning(FString::Printf(
 			TEXT("Variant property %s doesn't belong to an actor selected for export, the property will be skipped"),
diff --git a/Engine/Plugins/Enterprise/GLTFExporter/MigratedFromGit/Source/GLTFExporter/Private/Converters/GLTFKhrVariantConverters.cpp b/Engine/Plugins/Enterprise/GLTFExporter/MigratedFromGit/Source/GLTFExporter/Private/Converters/GLTFKhrVariantConverters.cpp
--- a/Engine/Plugins/Enterprise/GLTFExporter/MigratedFromGit/Source/GLTFExporter/Private/Converters/GLTFKhrVariantConverters.cpp
+++ b/Engine/Plugins/Enterprise/GLTFExporter/MigratedFromGit/Source/GLTFExporter/Private/Converters/GLTFKhrVariantConverters.cpp
@@ -104,7 +104,7 @@ bool FGLTFKhrMaterialVariantConverter::TryParseMaterialProperty(FGLTFJsonPrimiti
 		return false;
 	}
 
-	if (Builder.bSelectedActorsOnly && !Owner->IsSelected())
+	if (!Builder.IsSelectedActor(Owner))
 	{
 		Builder.LogWarning(FString::Printf(
 			TEXT("Variant property %s doesn't belong to an actor selected for export, the property will be skipped"),
diff --git a/Engine/Plugins/Enterprise/GLTFExporter/MigratedFromGit/Source/GLTFExporter/Private/Converters/GLTFNodeConverters.cpp b/Engine/Plugins/Enterprise/GLTFExporter/MigratedFromGit/Source/GLTFExporter/Private/Converters/GLTFNodeConverters.cpp
--- a/Engine/Plugins/Enterprise/GLTFExporter/MigratedFromGit/Source/GLTFExporter/Private/Converters/GLTFNodeConverters.cpp
+++ b/Engine/Plugins/Enterprise/GLTFExporter/MigratedFromGit/Source/GLTFExporter/Private/Converters/GLTFNodeConverters.cpp
@@ -19,7 +19,7 @@ FGLTFJsonNodeIndex FGLTFActorConverter::Convert(const AActor* Actor)
 		return FGLTFJsonNodeIndex(INDEX_NONE);
 	}
 
-	if (Builder.bSelectedActorsOnly && !Actor->IsSelected())
+	if (!Builder.IsSelectedActor(Actor))
 	{
 		return FGLTFJsonNodeIndex(INDEX_NONE);
 	}
@@ -94,13 +94,13 @@ FGLTFJsonNodeIndex FGLTFComponentConverter::Convert(const USceneComponent* Scene
 		return FGLTFJsonNodeIndex(INDEX_NONE);
 	}
 
-	if (Builder.bSelectedActorsOnly && !Owner->IsSelected())
+	if (!Builder.IsSelectedActor(Owner))
 	{
 		return FGLTFJsonNodeIndex(INDEX_NONE);
 	}
 
 	const bool bIsRootComponent = Owner->GetRootComponent() == SceneComponent;
-	const bool bIsRootNode = bIsRootComponent && FGLTFActorUtility::IsRootActor(Owner, Builder.bSelectedActorsOnly);
+	const bool bIsRootNode = bIsRootComponent && Builder.IsRootActor(Owner);
 
 	const USceneComponent* ParentComponent = !bIsRootNode ? SceneComponent->GetAttachParent() : nullptr;
 	const FName SocketName = SceneComponent->GetAttachSocketName();
diff --git a/Engine/Plugins/Enterprise/GLTFExporter/MigratedFromGit/Source/GLTFExporter/Private/Converters/GLTFSceneConverters.cpp b/Engine/Plugins/Enterprise/GLTFExporter/MigratedFromGit/Source/GLTFExporter/Private/Converters/GLTFSceneConverters.cpp
--- a/Engine/Plugins/Enterprise/GLTFExporter/MigratedFromGit/Source/GLTFExporter/Private/Converters/GLTFSceneConverters.cpp
+++ b/Engine/Plugins/Enterprise/GLTFExporter/MigratedFromGit/Source/GLTFExporter/Private/Converters/GLTFSceneConverters.cpp
@@ -64,7 +64,7 @@ FGLTFJsonSceneIndex FGLTFSceneConverter::Convert(const UWorld* World)
 			for (const AActor* Actor : Level->Actors)
 			{
 				const FGLTFJsonNodeIndex NodeIndex = Builder.GetOrAddNode(Actor);
-				if (NodeIndex != INDEX_NONE && FGLTFActorUtility::IsRootActor(Actor, Builder.bSelectedActorsOnly))
+				if (NodeIndex != INDEX_NONE && Builder.IsRootActor(Actor))
 				{
 					// TODO: to avoid having to add irrelevant actors/components let GLTFComponentConverter decide and add root nodes to scene.
 					// This change may require node converters to support cyclic calls.
diff --git a/Engine/Plugins/Enterprise/GLTFExporter/MigratedFromGit/Source/GLTFExporter/Private/Exporters/GLTFExporter.cpp b/Engine/Plugins/Enterprise/GLTFExporter/MigratedFromGit/Source/GLTFExporter/Private/Exporters/GLTFExporter.cpp
--- a/Engine/Plugins/Enterprise/GLTFExporter/MigratedFromGit/Source/GLTFExporter/Private/Exporters/GLTFExporter.cpp
+++ b/Engine/Plugins/Enterprise/GLTFExporter/MigratedFromGit/Source/GLTFExporter/Private/Exporters/GLTFExporter.cpp
@@ -6,6 +6,7 @@
 #include "Builders/GLTFContainerBuilder.h"
 #include "UObject/GCObjectScopeGuard.h"
 #include "AssetExportTask.h"
+#include "GLTFExporterUtility.h"
 
 UGLTFExporter::UGLTFExporter(const FObjectInitializer& ObjectInitializer)
 	: Super(ObjectInitializer)
@@ -32,9 +33,15 @@ bool UGLTFExporter::ExportBinary(UObject* Object, const TCHAR* Type, FArchive& A
 
 	FGCObjectScopeGuard OptionsGuard(Options);
 
+	TSet<AActor*> SelectedActors;
+	if (bSelectedOnly)
+	{
+		FGLTFExporterUtility::GetSelectedActors(SelectedActors);
+	}
+
 	// TODO: add support for UAssetExportTask::IgnoreObjectList?
 
-	FGLTFContainerBuilder Builder(GetFilePath(), Options, bSelectedOnly);
+	FGLTFContainerBuilder Builder(GetFilePath(), Options, SelectedActors);
 	Builder.ClearLog();
 
 	const bool bSuccess = AddObject(Builder, Object);
@@ -53,7 +60,7 @@ bool UGLTFExporter::ExportBinary(UObject* Object, const TCHAR* Type, FArchive& A
 	return bSuccess;
 }
 
-bool UGLTFExporter::ExportToGLTF(UObject* Object, const FString& Filename, const UGLTFExportOptions* Options, FGLTFExportMessages& OutMessages)
+bool UGLTFExporter::ExportToGLTF(UObject* Object, const FString& Filename, const UGLTFExportOptions* Options, const TSet<AActor*>& SelectedActors, FGLTFExportMessages& OutMessages)
 {
 	if (Object == nullptr)
 	{
@@ -68,7 +75,7 @@ bool UGLTFExporter::ExportToGLTF(UObject* Object, const FString& Filename, const
 		return false;
 	}
 
-	FGLTFContainerBuilder Builder(Filename, Options, false);
+	FGLTFContainerBuilder Builder(Filename, Options, SelectedActors);
 	const bool bSuccess = Exporter->AddObject(Builder, Object);
 
 	OutMessages.Suggestions = Builder.GetLoggedSuggestions();
@@ -91,10 +98,10 @@ bool UGLTFExporter::ExportToGLTF(UObject* Object, const FString& Filename, const
 	return bSuccess;
 }
 
-bool UGLTFExporter::ExportToGLTF(UObject* Object, const FString& Filename, const UGLTFExportOptions* Options)
+bool UGLTFExporter::ExportToGLTF(UObject* Object, const FString& Filename, const UGLTFExportOptions* Options, const TSet<AActor*>& SelectedActors)
 {
 	FGLTFExportMessages Messages;
-	return ExportToGLTF(Object, Filename, Options, Messages);
+	return ExportToGLTF(Object, Filename, Options, SelectedActors, Messages);
 }
 
 bool UGLTFExporter::AddObject(FGLTFContainerBuilder& Builder, const UObject* Object)
diff --git a/Engine/Plugins/Enterprise/GLTFExporter/MigratedFromGit/Source/GLTFExporter/Private/Exporters/GLTFExporterUtility.cpp b/Engine/Plugins/Enterprise/GLTFExporter/MigratedFromGit/Source/GLTFExporter/Private/Exporters/GLTFExporterUtility.cpp
--- a/Engine/Plugins/Enterprise/GLTFExporter/MigratedFromGit/Source/GLTFExporter/Private/Exporters/GLTFExporterUtility.cpp
+++ b/Engine/Plugins/Enterprise/GLTFExporter/MigratedFromGit/Source/GLTFExporter/Private/Exporters/GLTFExporterUtility.cpp
@@ -4,6 +4,22 @@
 #include "Materials/MaterialInstance.h"
 #include "AssetRegistryModule.h"
 
+void FGLTFExporterUtility::GetSelectedActors(TSet<AActor*>& OutSelectedActors)
+{
+#if WITH_EDITOR
+	const TMap<const UObjectBase*, FBoolAnnotation>& AnnotationMap = ((const FUObjectAnnotationSparse<FBoolAnnotation, true>&)GSelectedActorAnnotation).GetAnnotationMap();
+	OutSelectedActors.Reserve(AnnotationMap.Num());
+
+	for (const TPair<const UObjectBase*, FBoolAnnotation>& Pair : AnnotationMap)
+	{
+		if (Pair.Value.Mark)
+		{
+			OutSelectedActors.Add(static_cast<AActor*>(const_cast<UObjectBase*>(Pair.Key)));
+		}
+	}
+#endif
+}
+
 const UStaticMesh* FGLTFExporterUtility::GetPreviewMesh(const UMaterialInterface* Material)
 {
 #if WITH_EDITORONLY_DATA
diff --git a/Engine/Plugins/Enterprise/GLTFExporter/MigratedFromGit/Source/GLTFExporter/Private/Exporters/GLTFExporterUtility.h b/Engine/Plugins/Enterprise/GLTFExporter/MigratedFromGit/Source/GLTFExporter/Private/Exporters/GLTFExporterUtility.h
--- a/Engine/Plugins/Enterprise/GLTFExporter/MigratedFromGit/Source/GLTFExporter/Private/Exporters/GLTFExporterUtility.h
+++ b/Engine/Plugins/Enterprise/GLTFExporter/MigratedFromGit/Source/GLTFExporter/Private/Exporters/GLTFExporterUtility.h
@@ -9,6 +9,8 @@ class USkeleton;
 
 struct FGLTFExporterUtility
 {
+	static void GetSelectedActors(TSet<AActor*>& OutSelectedActors);
+
 	static const UStaticMesh* GetPreviewMesh(const UMaterialInterface* Material);
 	static const USkeletalMesh* GetPreviewMesh(const UAnimSequence* AnimSequence);
 	static const USkeletalMesh* FindCompatibleMesh(const USkeleton* Skeleton);
diff --git a/Engine/Plugins/Enterprise/GLTFExporter/MigratedFromGit/Source/GLTFExporter/Public/Exporters/GLTFExporter.h b/Engine/Plugins/Enterprise/GLTFExporter/MigratedFromGit/Source/GLTFExporter/Public/Exporters/GLTFExporter.h
--- a/Engine/Plugins/Enterprise/GLTFExporter/MigratedFromGit/Source/GLTFExporter/Public/Exporters/GLTFExporter.h
+++ b/Engine/Plugins/Enterprise/GLTFExporter/MigratedFromGit/Source/GLTFExporter/Public/Exporters/GLTFExporter.h
@@ -36,8 +36,8 @@ public:
 	virtual bool ExportBinary(UObject* Object, const TCHAR* Type, FArchive& Archive, FFeedbackContext* Warn, int32 FileIndex, uint32 PortFlags) override;
 
 	UFUNCTION(BlueprintCallable, Category = "Miscellaneous")
-	static bool ExportToGLTF(UObject* Object, const FString& Filename, const UGLTFExportOptions* Options, FGLTFExportMessages& OutMessages);
-	static bool ExportToGLTF(UObject* Object, const FString& Filename, const UGLTFExportOptions* Options = nullptr);
+	static bool ExportToGLTF(UObject* Object, const FString& Filename, const UGLTFExportOptions* Options, const TSet<AActor*>& SelectedActors, FGLTFExportMessages& OutMessages);
+	static bool ExportToGLTF(UObject* Object, const FString& Filename, const UGLTFExportOptions* Options = nullptr, const TSet<AActor*>& SelectedActors = {});
 
 protected:
 
