[GLTFExporter] refactor: move some animation sequence operations into utility functions

commit d8eb163ac340593c1fda6ac659542d4ebabe4b07
Date: Tue, 18 Jan 2022 14:58:00 +0100

    refactor: move some animation sequence operations into utility functions
    (cherry picked from commit 3274b97d92e98f3e8c93f5abe48ff088aec22932​)

#skipci

diff --git a/Engine/Plugins/Enterprise/GLTFExporter/MigratedFromGit/Source/GLTFExporter/Private/Converters/GLTFBoneUtility.cpp b/Engine/Plugins/Enterprise/GLTFExporter/MigratedFromGit/Source/GLTFExporter/Private/Converters/GLTFBoneUtility.cpp
--- a/Engine/Plugins/Enterprise/GLTFExporter/MigratedFromGit/Source/GLTFExporter/Private/Converters/GLTFBoneUtility.cpp
+++ b/Engine/Plugins/Enterprise/GLTFExporter/MigratedFromGit/Source/GLTFExporter/Private/Converters/GLTFBoneUtility.cpp
@@ -1,6 +1,9 @@
 // Copyright Epic Games, Inc. All Rights Reserved.
 
 #include "Converters/GLTFBoneUtility.h"
+#if (ENGINE_MAJOR_VERSION > 4 || ENGINE_MINOR_VERSION >= 26)
+#include "Animation/AnimationPoseData.h"
+#endif
 
 FTransform FGLTFBoneUtility::GetBindTransform(const FReferenceSkeleton& RefSkeleton, int32 BoneIndex)
 {
@@ -19,16 +22,69 @@ FTransform FGLTFBoneUtility::GetBindTransform(const FReferenceSkeleton& RefSkele
 	return BindTransform;
 }
 
-void FGLTFBoneUtility::InitializeToSkeleton(FBoneContainer& BoneContainer, const USkeleton* Skeleton)
+void FGLTFBoneUtility::GetFrameTimestamps(const UAnimSequence* AnimSequence, TArray<float>& OutFrameTimestamps)
 {
-	TArray<FBoneIndexType> RequiredBoneIndices;
-	RequiredBoneIndices.AddUninitialized(Skeleton->GetReferenceSkeleton().GetNum());
+	const int32 FrameCount = AnimSequence->GetRawNumberOfFrames();
+	OutFrameTimestamps.AddUninitialized(FrameCount);
 
-	for (int32 BoneIndex = 0; BoneIndex < RequiredBoneIndices.Num(); ++BoneIndex)
+	const float SequenceLength = AnimSequence->SequenceLength;
+	const float FrameLength = FrameCount > 1 ? SequenceLength / (FrameCount - 1) : 0;
+
+	for (int32 FrameIndex = 0; FrameIndex < FrameCount; ++FrameIndex)
+	{
+		OutFrameTimestamps[FrameIndex] = FMath::Clamp(FrameIndex * FrameLength, 0.0f, SequenceLength);
+	}
+}
+
+void FGLTFBoneUtility::GetBoneIndices(const USkeleton* Skeleton, TArray<FBoneIndexType>& OutBoneIndices)
+{
+	const int32 BoneCount = Skeleton->GetReferenceSkeleton().GetNum();
+	OutBoneIndices.AddUninitialized(BoneCount);
+
+	for (int32 BoneIndex = 0; BoneIndex < BoneCount; ++BoneIndex)
 	{
-		RequiredBoneIndices[BoneIndex] = BoneIndex;
+		OutBoneIndices[BoneIndex] = BoneIndex;
 	}
+}
+
+void FGLTFBoneUtility::GetBoneTransformsByFrame(const UAnimSequence* AnimSequence, const TArray<float>& FrameTimestamps, const TArray<FBoneIndexType>& BoneIndices, TArray<TArray<FTransform>>& OutBoneTransformsByFrame)
+{
+	FMemMark Mark(FMemStack::Get()); // Make sure to free stack allocations made by FCompactPose, FBlendedCurve, and FStackCustomAttributes when end of scope
 
+	FBoneContainer BoneContainer;
 	BoneContainer.SetUseRAWData(true);
-	BoneContainer.InitializeTo(RequiredBoneIndices, FCurveEvaluationOption(true), *const_cast<USkeleton*>(Skeleton));
+	BoneContainer.InitializeTo(BoneIndices, FCurveEvaluationOption(true), *AnimSequence->GetSkeleton());
+
+	const int32 FrameCount = FrameTimestamps.Num();
+	OutBoneTransformsByFrame.AddDefaulted(FrameCount);
+
+	FCompactPose Pose;
+	Pose.SetBoneContainer(&BoneContainer);
+
+	FBlendedCurve Curve;
+	Curve.InitFrom(BoneContainer);
+
+#if (ENGINE_MAJOR_VERSION > 4 || ENGINE_MINOR_VERSION >= 26)
+	FStackCustomAttributes Attributes;
+	FAnimationPoseData PoseData(Pose, Curve, Attributes);
+#endif
+
+	for (int32 Frame = 0; Frame < FrameCount; ++Frame)
+	{
+		const FAnimExtractContext ExtractionContext(FrameTimestamps[Frame]); // TODO: set bExtractRootMotion?
+#if (ENGINE_MAJOR_VERSION > 4 || ENGINE_MINOR_VERSION >= 26)
+		AnimSequence->GetBonePose(PoseData, ExtractionContext);
+#else
+		AnimSequence->GetBonePose(Pose, Curve, ExtractionContext);
+#endif
+		Pose.CopyBonesTo(OutBoneTransformsByFrame[Frame]);
+	}
+
+	// Clear all stack allocations to allow FMemMark to free them
+	Pose.Empty();
+	Curve.Empty();
+
+#if (ENGINE_MAJOR_VERSION > 4 || ENGINE_MINOR_VERSION >= 26)
+	Attributes.Empty();
+#endif
 }
diff --git a/Engine/Plugins/Enterprise/GLTFExporter/MigratedFromGit/Source/GLTFExporter/Private/Converters/GLTFBoneUtility.h b/Engine/Plugins/Enterprise/GLTFExporter/MigratedFromGit/Source/GLTFExporter/Private/Converters/GLTFBoneUtility.h
--- a/Engine/Plugins/Enterprise/GLTFExporter/MigratedFromGit/Source/GLTFExporter/Private/Converters/GLTFBoneUtility.h
+++ b/Engine/Plugins/Enterprise/GLTFExporter/MigratedFromGit/Source/GLTFExporter/Private/Converters/GLTFBoneUtility.h
@@ -8,5 +8,9 @@ struct FGLTFBoneUtility
 {
 	static FTransform GetBindTransform(const FReferenceSkeleton& RefSkeleton, int32 BoneIndex);
 
-	static void InitializeToSkeleton(FBoneContainer& BoneContainer, const USkeleton* Skeleton);
+	static void GetFrameTimestamps(const UAnimSequence* AnimSequence, TArray<float>& OutFrameTimestamps);
+
+	static void GetBoneIndices(const USkeleton* Skeleton, TArray<FBoneIndexType>& OutBoneIndices);
+
+	static void GetBoneTransformsByFrame(const UAnimSequence* AnimSequence, const TArray<float>& FrameTimestamps, const TArray<FBoneIndexType>& BoneIndices, TArray<TArray<FTransform>>& OutBoneTransformsByFrame);
 };
diff --git a/Engine/Plugins/Enterprise/GLTFExporter/MigratedFromGit/Source/GLTFExporter/Private/Tasks/GLTFAnimationTasks.cpp b/Engine/Plugins/Enterprise/GLTFExporter/MigratedFromGit/Source/GLTFExporter/Private/Tasks/GLTFAnimationTasks.cpp
--- a/Engine/Plugins/Enterprise/GLTFExporter/MigratedFromGit/Source/GLTFExporter/Private/Tasks/GLTFAnimationTasks.cpp
+++ b/Engine/Plugins/Enterprise/GLTFExporter/MigratedFromGit/Source/GLTFExporter/Private/Tasks/GLTFAnimationTasks.cpp
@@ -12,7 +12,6 @@
 #include "Channels/MovieSceneChannelProxy.h"
 
 #if (ENGINE_MAJOR_VERSION > 4 || ENGINE_MINOR_VERSION >= 26)
-#include "Animation/AnimationPoseData.h"
 using namespace UE;
 #endif
 
@@ -20,61 +19,12 @@ void FGLTFAnimSequenceTask::Complete()
 {
 	// TODO: bone transforms should be absolute (not relative) according to gltf spec
 
-	const int32 FrameCount = AnimSequence->GetRawNumberOfFrames();
-	const USkeleton* Skeleton = AnimSequence->GetSkeleton();
-
 	FGLTFJsonAnimation& JsonAnimation = Builder.GetAnimation(AnimationIndex);
 	JsonAnimation.Name = AnimSequence->GetName() + TEXT("_") + FString::FromInt(AnimationIndex); // Ensure unique name due to limitation in certain gltf viewers
 
 	TArray<float> Timestamps;
-	Timestamps.AddUninitialized(FrameCount);
-	{
-		const float SequenceLength = AnimSequence->SequenceLength;
-		const float FrameLength = FrameCount > 1 ? SequenceLength / (FrameCount - 1) : 0;
-
-		for (int32 Frame = 0; Frame < FrameCount; ++Frame)
-		{
-			Timestamps[Frame] = FMath::Clamp(Frame * FrameLength, 0.0f, SequenceLength);
-		}
-	}
-
-	FBoneContainer BoneContainer;
-	FGLTFBoneUtility::InitializeToSkeleton(BoneContainer, Skeleton);
-
-	TArray<TArray<FTransform>> FrameTransforms;
-	FrameTransforms.AddDefaulted(FrameCount);
-	{
-		FMemMark Mark(FMemStack::Get());
-
-		FCompactPose Pose;
-		Pose.SetBoneContainer(&BoneContainer);
-
-		FBlendedCurve Curve;
-		Curve.InitFrom(BoneContainer);
-
-#if (ENGINE_MAJOR_VERSION > 4 || ENGINE_MINOR_VERSION >= 26)
-		FStackCustomAttributes Attributes;
-		FAnimationPoseData PoseData(Pose, Curve, Attributes);
-#endif
-
-		for (int32 Frame = 0; Frame < FrameCount; ++Frame)
-		{
-			const FAnimExtractContext ExtractionContext(Timestamps[Frame]); // TODO: set bExtractRootMotion?
-
-#if (ENGINE_MAJOR_VERSION > 4 || ENGINE_MINOR_VERSION >= 26)
-			AnimSequence->GetBonePose(PoseData, ExtractionContext);
-#else
-			AnimSequence->GetBonePose(Pose, Curve, ExtractionContext);
-#endif
-			Pose.CopyBonesTo(FrameTransforms[Frame]);
-		}
-
-		Pose.Empty();
-		Curve.Empty();
-#if (ENGINE_MAJOR_VERSION > 4 || ENGINE_MINOR_VERSION >= 26)
-		Attributes.Empty();
-#endif
-	}
+	FGLTFBoneUtility::GetFrameTimestamps(AnimSequence, Timestamps);
+	const int32 FrameCount = Timestamps.Num();
 
 	// TODO: add animation data accessor converters to reuse track information
 
@@ -98,7 +48,12 @@ void FGLTFAnimSequenceTask::Complete()
 		// TODO: report warning (about unknown interpolation exported as linear)
 	}
 
-	const TArray<FBoneIndexType>& BoneIndices = BoneContainer.GetBoneIndicesArray();
+	TArray<FBoneIndexType> BoneIndices;
+	FGLTFBoneUtility::GetBoneIndices(AnimSequence->GetSkeleton(), BoneIndices);
+
+	TArray<TArray<FTransform>> FrameTransforms;
+	FGLTFBoneUtility::GetBoneTransformsByFrame(AnimSequence, Timestamps, BoneIndices, FrameTransforms);
+
 	for (FBoneIndexType BoneIndex : BoneIndices)
 	{
 		const FGLTFJsonNodeIndex NodeIndex = Builder.GetOrAddNode(RootNode, SkeletalMesh, BoneIndex);
