[GLTFExporter] fix: export same compressed animation data in both editor and runtime

commit 818b5235c8c3c2c887b68547f64165e69ccd69f0
Date: Tue, 18 Jan 2022 18:29:46 +0100

    fix: export same compressed animation data in both editor and runtime
    (cherry picked from commit faf52f1b5142ef01ca27b5e244c39243063043bf​)

#skipci

diff --git a/Engine/Plugins/Enterprise/GLTFExporter/MigratedFromGit/Source/GLTFExporter/Private/Converters/GLTFBoneUtility.cpp b/Engine/Plugins/Enterprise/GLTFExporter/MigratedFromGit/Source/GLTFExporter/Private/Converters/GLTFBoneUtility.cpp
--- a/Engine/Plugins/Enterprise/GLTFExporter/MigratedFromGit/Source/GLTFExporter/Private/Converters/GLTFBoneUtility.cpp
+++ b/Engine/Plugins/Enterprise/GLTFExporter/MigratedFromGit/Source/GLTFExporter/Private/Converters/GLTFBoneUtility.cpp
@@ -52,7 +52,6 @@ void FGLTFBoneUtility::GetBoneTransformsByFrame(const UAnimSequence* AnimSequenc
 	FMemMark Mark(FMemStack::Get()); // Make sure to free stack allocations made by FCompactPose, FBlendedCurve, and FStackCustomAttributes when end of scope
 
 	FBoneContainer BoneContainer;
-	BoneContainer.SetUseRAWData(true);
 	BoneContainer.InitializeTo(BoneIndices, FCurveEvaluationOption(true), *AnimSequence->GetSkeleton());
 
 	const int32 FrameCount = FrameTimestamps.Num();
