[GLTFExporter] fix: use a custom shader when sampling HDR textures

commit 74c2cb21a43720841eeb1da8b80d723d078cb529
Date: Fri, 7 Jan 2022 11:47:28 +0100

    fix: use a custom shader when sampling HDR textures
    prevents gamma-correction and other transformations from affecting the result
    (cherry picked from commit e27190f8659c82f46c13889d6e4a6ac2f5e607bc​)

#skipci

diff --git a/Engine/Plugins/Enterprise/GLTFExporter/MigratedFromGit/Config/FilterPlugin.ini b/Engine/Plugins/Enterprise/GLTFExporter/MigratedFromGit/Config/FilterPlugin.ini
--- a/Engine/Plugins/Enterprise/GLTFExporter/MigratedFromGit/Config/FilterPlugin.ini
+++ b/Engine/Plugins/Enterprise/GLTFExporter/MigratedFromGit/Config/FilterPlugin.ini
@@ -1,3 +1,4 @@
 [FilterPlugin]
 /Docs/
 /README.md
+/Shaders/*
diff --git a/Engine/Plugins/Enterprise/GLTFExporter/MigratedFromGit/Shaders/Private/SimpleTexture2DPS.usf b/Engine/Plugins/Enterprise/GLTFExporter/MigratedFromGit/Shaders/Private/SimpleTexture2DPS.usf
new file mode 100644
--- /dev/null
+++ b/Engine/Plugins/Enterprise/GLTFExporter/MigratedFromGit/Shaders/Private/SimpleTexture2DPS.usf
@@ -0,0 +1,17 @@
+﻿// Copyright Epic Games, Inc. All Rights Reserved.
+
+#pragma once
+
+#include "/Engine/Private/Common.ush"
+
+Texture2D Texture;
+SamplerState TextureSampler;
+
+void Main(
+	in float2 TextureCoordinate : TEXCOORD0,
+	out float4 OutColor : SV_Target0
+	)
+{
+	const float4 TextureSample = Texture2DSample(Texture, TextureSampler, TextureCoordinate);
+	OutColor = RETURN_COLOR(TextureSample);
+}
diff --git a/Engine/Plugins/Enterprise/GLTFExporter/MigratedFromGit/Source/GLTFExporter/Private/Converters/GLTFSimpleTexture2DPreview.cpp b/Engine/Plugins/Enterprise/GLTFExporter/MigratedFromGit/Source/GLTFExporter/Private/Converters/GLTFSimpleTexture2DPreview.cpp
new file mode 100644
--- /dev/null
+++ b/Engine/Plugins/Enterprise/GLTFExporter/MigratedFromGit/Source/GLTFExporter/Private/Converters/GLTFSimpleTexture2DPreview.cpp
@@ -0,0 +1,78 @@
+﻿// Copyright Epic Games, Inc. All Rights Reserved.
+
+/*==============================================================================
+	NormalMapPreview.h: Implementation for previewing normal maps.
+==============================================================================*/
+
+#include "GLTFSimpleTexture2DPreview.h"
+#include "Shader.h"
+#include "GlobalShader.h"
+#include "SimpleElementShaders.h"
+#include "ShaderParameterUtils.h"
+#include "PipelineStateCache.h"
+
+class FGLTFSimpleTexture2DPS : public FGlobalShader
+{
+	DECLARE_SHADER_TYPE(FGLTFSimpleTexture2DPS,Global);
+public:
+
+	static bool ShouldCompilePermutation(const FGlobalShaderPermutationParameters& Parameters)
+	{
+		return IsPCPlatform(Parameters.Platform);
+	}
+
+	FGLTFSimpleTexture2DPS() {}
+
+	/**
+	 * Initialization constructor.
+	 * @param Initializer - Shader initialization container.
+	 */
+	FGLTFSimpleTexture2DPS(const ShaderMetaType::CompiledShaderInitializerType& Initializer)
+		: FGlobalShader(Initializer)
+	{
+		Texture.Bind(Initializer.ParameterMap,TEXT("Texture"));
+		TextureSampler.Bind(Initializer.ParameterMap,TEXT("TextureSampler"));
+	}
+
+	/**
+	 * Set shader parameters.
+	 * @param SampleTexture - The texture to sample.
+	 */
+	void SetParameters(FRHICommandList& RHICmdList, const FTexture* SampleTexture)
+	{
+		FRHIPixelShader* PixelShaderRHI = RHICmdList.GetBoundPixelShader();
+		SetTextureParameter(RHICmdList, PixelShaderRHI,Texture,TextureSampler,SampleTexture);
+	}
+
+private:
+	/** The texture to sample. */
+	LAYOUT_FIELD(FShaderResourceParameter, Texture);
+	LAYOUT_FIELD(FShaderResourceParameter, TextureSampler);
+};
+IMPLEMENT_SHADER_TYPE(,FGLTFSimpleTexture2DPS,TEXT("/Plugin/GLTFExporter/Private/SimpleTexture2DPS.usf"),TEXT("Main"),SF_Pixel);
+
+/** Binds vertex and pixel shaders for this element */
+void FGLTFSimpleTexture2DPreview::BindShaders(
+	FRHICommandList& RHICmdList,
+	FGraphicsPipelineStateInitializer& GraphicsPSOInit,
+	ERHIFeatureLevel::Type InFeatureLevel,
+	const FMatrix& InTransform,
+	const float InGamma,
+	const FMatrix& ColorWeights,
+	const FTexture* Texture)
+{
+	TShaderMapRef<FSimpleElementVS> VertexShader(GetGlobalShaderMap(InFeatureLevel));
+	TShaderMapRef<FGLTFSimpleTexture2DPS> PixelShader(GetGlobalShaderMap(InFeatureLevel));
+
+	GraphicsPSOInit.BoundShaderState.VertexDeclarationRHI = GSimpleElementVertexDeclaration.VertexDeclarationRHI;
+	GraphicsPSOInit.BoundShaderState.VertexShaderRHI = VertexShader.GetVertexShader();
+	GraphicsPSOInit.BoundShaderState.PixelShaderRHI = PixelShader.GetPixelShader();
+	GraphicsPSOInit.PrimitiveType = PT_TriangleList;
+
+	GraphicsPSOInit.BlendState = TStaticBlendState<>::GetRHI();
+
+	SetGraphicsPipelineState(RHICmdList, GraphicsPSOInit, EApplyRendertargetOption::ForceApply);
+
+	VertexShader->SetParameters(RHICmdList, InTransform);
+	PixelShader->SetParameters(RHICmdList, Texture);
+}
diff --git a/Engine/Plugins/Enterprise/GLTFExporter/MigratedFromGit/Source/GLTFExporter/Private/Converters/GLTFSimpleTexture2DPreview.h b/Engine/Plugins/Enterprise/GLTFExporter/MigratedFromGit/Source/GLTFExporter/Private/Converters/GLTFSimpleTexture2DPreview.h
new file mode 100644
--- /dev/null
+++ b/Engine/Plugins/Enterprise/GLTFExporter/MigratedFromGit/Source/GLTFExporter/Private/Converters/GLTFSimpleTexture2DPreview.h
@@ -0,0 +1,19 @@
+﻿// Copyright Epic Games, Inc. All Rights Reserved.
+
+#pragma once
+
+#include "CoreMinimal.h"
+#include "BatchedElements.h"
+
+class FGLTFSimpleTexture2DPreview : public FBatchedElementParameters
+{
+	/** Binds vertex and pixel shaders for this element */
+	virtual void BindShaders(
+		FRHICommandList& RHICmdList,
+		FGraphicsPipelineStateInitializer& GraphicsPSOInit,
+		ERHIFeatureLevel::Type InFeatureLevel,
+		const FMatrix& InTransform,
+		const float InGamma,
+		const FMatrix& ColorWeights,
+		const FTexture* Texture) override;
+};
diff --git a/Engine/Plugins/Enterprise/GLTFExporter/MigratedFromGit/Source/GLTFExporter/Private/Converters/GLTFTextureUtility.cpp b/Engine/Plugins/Enterprise/GLTFExporter/MigratedFromGit/Source/GLTFExporter/Private/Converters/GLTFTextureUtility.cpp
--- a/Engine/Plugins/Enterprise/GLTFExporter/MigratedFromGit/Source/GLTFExporter/Private/Converters/GLTFTextureUtility.cpp
+++ b/Engine/Plugins/Enterprise/GLTFExporter/MigratedFromGit/Source/GLTFExporter/Private/Converters/GLTFTextureUtility.cpp
@@ -2,6 +2,7 @@
 
 #include "Converters/GLTFTextureUtility.h"
 #include "Converters/GLTFNormalMapPreview.h"
+#include "Converters/GLTFSimpleTexture2DPreview.h"
 #include "CanvasItem.h"
 #include "Engine/Texture2D.h"
 #include "Engine/TextureCube.h"
@@ -182,6 +183,12 @@ bool FGLTFTextureUtility::DrawTexture(UTextureRenderTarget2D* OutTarget, const U
 	{
 		BatchedElementParameters = new FGLTFNormalMapPreview();
 	}
+	else if (IsHDR(InSource))
+	{
+		// NOTE: Simple preview parameters are used to prevent any modifications
+		// such as gamma-correction from being applied during rendering.
+		BatchedElementParameters = new FGLTFSimpleTexture2DPreview();
+	}
 
 	FCanvas Canvas(RenderTarget, nullptr, 0.0f, 0.0f, 0.0f, GMaxRHIFeatureLevel);
 	FCanvasTileItem TileItem(InPosition, InSource->Resource, InSize, FLinearColor::White);
@@ -230,6 +237,11 @@ UTexture2D* FGLTFTextureUtility::CreateTextureFromCubeFace(const UTextureCube* T
 	{
 		const void* FaceDataPtr = static_cast<const uint8*>(MipDataPtr) + MipSize * CubeFace;
 		FaceTexture = CreateTransientTexture(FaceDataPtr, MipSize, Size, Format, TextureCube->SRGB);
+
+		if (IsHDR(TextureCube))
+		{
+			FaceTexture->CompressionSettings = TC_HDR;
+		}
 	}
 	const_cast<FByteBulkData&>(BulkData).Unlock();
 
@@ -256,6 +268,7 @@ UTexture2D* FGLTFTextureUtility::CreateTextureFromCubeFace(const UTextureRenderT
 		}
 
 		FaceTexture = CreateTransientTexture(Pixels.GetData(), Pixels.Num() * sizeof(FFloat16Color), Size, PF_FloatRGBA, false);
+		FaceTexture->CompressionSettings = TC_HDR;
 	}
 	else
 	{
diff --git a/Engine/Plugins/Enterprise/GLTFExporter/MigratedFromGit/Source/GLTFExporter/Private/GLTFExporterModule.cpp b/Engine/Plugins/Enterprise/GLTFExporter/MigratedFromGit/Source/GLTFExporter/Private/GLTFExporterModule.cpp
--- a/Engine/Plugins/Enterprise/GLTFExporter/MigratedFromGit/Source/GLTFExporter/Private/GLTFExporterModule.cpp
+++ b/Engine/Plugins/Enterprise/GLTFExporter/MigratedFromGit/Source/GLTFExporter/Private/GLTFExporterModule.cpp
@@ -4,6 +4,7 @@
 
 #include "CoreMinimal.h"
 #include "Modules/ModuleManager.h"
+#include "Interfaces/IPluginManager.h"
 
 DEFINE_LOG_CATEGORY(LogGLTFExporter);
 
@@ -14,7 +15,11 @@ class FGLTFExporterModule final : public IGLTFExporterModule
 {
 
 public:
-	virtual void StartupModule() override {}
+	virtual void StartupModule() override
+	{
+		FString PluginShaderDir = FPaths::Combine(IPluginManager::Get().FindPlugin(TEXT("GLTFExporter"))->GetBaseDir(), TEXT("Shaders"));
+		AddShaderSourceDirectoryMapping(TEXT("/Plugin/GLTFExporter"), PluginShaderDir);
+	}
 
 	virtual void ShutdownModule() override {}
 };
