[GLTFExporter] fix: use ReadFloat16Pixels when reading pixels from HDR render-target

commit 7ae4cc2105cdeed183646a3e731ca4e0cbbe0f5c
Date: Fri, 7 Jan 2022 11:48:00 +0100

    fix: use ReadFloat16Pixels when reading pixels from HDR render-target
    (cherry picked from commit 15ba589c74eb1ed02fe24583c6f578b59374e126​)

#skipci

diff --git a/Engine/Plugins/Enterprise/GLTFExporter/MigratedFromGit/Source/GLTFExporter/Private/Converters/GLTFTextureUtility.cpp b/Engine/Plugins/Enterprise/GLTFExporter/MigratedFromGit/Source/GLTFExporter/Private/Converters/GLTFTextureUtility.cpp
--- a/Engine/Plugins/Enterprise/GLTFExporter/MigratedFromGit/Source/GLTFExporter/Private/Converters/GLTFTextureUtility.cpp
+++ b/Engine/Plugins/Enterprise/GLTFExporter/MigratedFromGit/Source/GLTFExporter/Private/Converters/GLTFTextureUtility.cpp
@@ -297,8 +297,8 @@ bool FGLTFTextureUtility::ReadPixels(const UTextureRenderTarget2D* InRenderTarge
 		return Resource->ReadPixels(OutPixels);
 	}
 
-	TArray<FLinearColor> HDRPixels;
-	if (!Resource->ReadLinearColorPixels(HDRPixels))
+	TArray<FFloat16Color> HDRPixels;
+	if (!Resource->ReadFloat16Pixels(HDRPixels))
 	{
 		return false;
 	}
@@ -319,13 +319,13 @@ bool FGLTFTextureUtility::ReadPixels(const UTextureRenderTarget2D* InRenderTarge
 	return true;
 }
 
-void FGLTFTextureUtility::EncodeRGBM(const TArray<FLinearColor>& InPixels, TArray<FColor>& OutPixels, float MaxRange)
+void FGLTFTextureUtility::EncodeRGBM(const TArray<FFloat16Color>& InPixels, TArray<FColor>& OutPixels, float MaxRange)
 {
 	OutPixels.AddUninitialized(InPixels.Num());
 
 	for (int32 Index = 0; Index < InPixels.Num(); ++Index)
 	{
-		const FLinearColor& Color = InPixels[Index];
+		const FFloat16Color& Color = InPixels[Index];
 		FLinearColor RGBM;
 
 		RGBM.R = FMath::Sqrt(Color.R);
@@ -347,13 +347,13 @@ void FGLTFTextureUtility::EncodeRGBM(const TArray<FLinearColor>& InPixels, TArra
 	}
 }
 
-void FGLTFTextureUtility::EncodeRGBE(const TArray<FLinearColor>& InPixels, TArray<FColor>& OutPixels)
+void FGLTFTextureUtility::EncodeRGBE(const TArray<FFloat16Color>& InPixels, TArray<FColor>& OutPixels)
 {
 	OutPixels.AddUninitialized(InPixels.Num());
 
 	for (int32 Index = 0; Index < InPixels.Num(); ++Index)
 	{
-		OutPixels[Index] = InPixels[Index].ToRGBE();
+		OutPixels[Index] = FLinearColor(InPixels[Index]).ToRGBE();
 	}
 }
 
diff --git a/Engine/Plugins/Enterprise/GLTFExporter/MigratedFromGit/Source/GLTFExporter/Private/Converters/GLTFTextureUtility.h b/Engine/Plugins/Enterprise/GLTFExporter/MigratedFromGit/Source/GLTFExporter/Private/Converters/GLTFTextureUtility.h
--- a/Engine/Plugins/Enterprise/GLTFExporter/MigratedFromGit/Source/GLTFExporter/Private/Converters/GLTFTextureUtility.h
+++ b/Engine/Plugins/Enterprise/GLTFExporter/MigratedFromGit/Source/GLTFExporter/Private/Converters/GLTFTextureUtility.h
@@ -33,8 +33,8 @@ struct FGLTFTextureUtility
 
 	static bool ReadPixels(const UTextureRenderTarget2D* InRenderTarget, TArray<FColor>& OutPixels, EGLTFJsonHDREncoding Encoding);
 
-	static void EncodeRGBM(const TArray<FLinearColor>& InPixels, TArray<FColor>& OutPixels, float MaxRange = 8);
-	static void EncodeRGBE(const TArray<FLinearColor>& InPixels, TArray<FColor>& OutPixels);
+	static void EncodeRGBM(const TArray<FFloat16Color>& InPixels, TArray<FColor>& OutPixels, float MaxRange = 8);
+	static void EncodeRGBE(const TArray<FFloat16Color>& InPixels, TArray<FColor>& OutPixels);
 
 	// TODO: maybe use template specialization to avoid the need for duplicated functions
 	static bool LoadPlatformData(UTexture2D* Texture);
