[GLTFExporter] refactor: replace GLTFImageBuilder with GLTFImageConverter

commit d61589177799665a82e596a99be7130ca97a5920
Date: Sat, 22 Jan 2022 00:30:26 +0100

    refactor: replace GLTFImageBuilder with GLTFImageConverter
    (cherry picked from commit 86bf1d325dc9e2f958c119160bcd993a1ece0c8d​)

#skipci

diff --git a/Engine/Plugins/Enterprise/GLTFExporter/MigratedFromGit/Source/GLTFExporter/Private/Builders/GLTFBuilder.cpp b/Engine/Plugins/Enterprise/GLTFExporter/MigratedFromGit/Source/GLTFExporter/Private/Builders/GLTFBuilder.cpp
--- a/Engine/Plugins/Enterprise/GLTFExporter/MigratedFromGit/Source/GLTFExporter/Private/Builders/GLTFBuilder.cpp
+++ b/Engine/Plugins/Enterprise/GLTFExporter/MigratedFromGit/Source/GLTFExporter/Private/Builders/GLTFBuilder.cpp
@@ -1,12 +1,11 @@
 // Copyright Epic Games, Inc. All Rights Reserved.
 
 #include "Builders/GLTFBuilder.h"
-#include "Builders/GLTFFileUtility.h"
 #include "Converters/GLTFMeshUtility.h"
 #include "UserData/GLTFMaterialUserData.h"
 
 FGLTFBuilder::FGLTFBuilder(const FString& FilePath, const UGLTFExportOptions* ExportOptions)
-	: bIsGlbFile(FGLTFFileUtility::IsGlbFile(FilePath))
+	: bIsGlbFile(FPaths::GetExtension(FilePath).Equals(TEXT("glb"), ESearchCase::IgnoreCase))
 	, FilePath(FilePath)
 	, DirPath(FPaths::GetPath(FilePath))
 	, ExportOptions(SanitizeExportOptions(ExportOptions))
diff --git a/Engine/Plugins/Enterprise/GLTFExporter/MigratedFromGit/Source/GLTFExporter/Private/Builders/GLTFConvertBuilder.cpp b/Engine/Plugins/Enterprise/GLTFExporter/MigratedFromGit/Source/GLTFExporter/Private/Builders/GLTFConvertBuilder.cpp
--- a/Engine/Plugins/Enterprise/GLTFExporter/MigratedFromGit/Source/GLTFExporter/Private/Builders/GLTFConvertBuilder.cpp
+++ b/Engine/Plugins/Enterprise/GLTFExporter/MigratedFromGit/Source/GLTFExporter/Private/Builders/GLTFConvertBuilder.cpp
@@ -8,7 +8,7 @@
 #include "Engine/TextureRenderTargetCube.h"
 
 FGLTFConvertBuilder::FGLTFConvertBuilder(const FString& FilePath, const UGLTFExportOptions* ExportOptions, const TSet<AActor*>& SelectedActors)
-	: FGLTFImageBuilder(FilePath, ExportOptions)
+	: FGLTFBufferBuilder(FilePath, ExportOptions)
 	, SelectedActors(SelectedActors)
 {
 }
@@ -316,6 +316,11 @@ FGLTFJsonTextureIndex FGLTFConvertBuilder::GetOrAddTexture(const UTextureRenderT
 	return TextureRenderTargetCubeConverter.GetOrAdd(Texture, CubeFace, bToSRGB);
 }
 
+FGLTFJsonImageIndex FGLTFConvertBuilder::GetOrAddImage(TGLTFSharedArray<FColor>& Pixels, FIntPoint Size, bool bIgnoreAlpha, EGLTFTextureType Type, const FString& Name)
+{
+	return ImageConverter.GetOrAdd(Name, Type, bIgnoreAlpha, Size, Pixels);
+}
+
 FGLTFJsonSkinIndex FGLTFConvertBuilder::GetOrAddSkin(FGLTFJsonNodeIndex RootNode, const USkeletalMesh* SkeletalMesh)
 {
 	if (RootNode == INDEX_NONE || SkeletalMesh == nullptr)
diff --git a/Engine/Plugins/Enterprise/GLTFExporter/MigratedFromGit/Source/GLTFExporter/Private/Builders/GLTFConvertBuilder.h b/Engine/Plugins/Enterprise/GLTFExporter/MigratedFromGit/Source/GLTFExporter/Private/Builders/GLTFConvertBuilder.h
--- a/Engine/Plugins/Enterprise/GLTFExporter/MigratedFromGit/Source/GLTFExporter/Private/Builders/GLTFConvertBuilder.h
+++ b/Engine/Plugins/Enterprise/GLTFExporter/MigratedFromGit/Source/GLTFExporter/Private/Builders/GLTFConvertBuilder.h
@@ -2,13 +2,14 @@
 
 #pragma once
 
-#include "Builders/GLTFImageBuilder.h"
+#include "Builders/GLTFBufferBuilder.h"
 #include "Converters/GLTFAccessorConverters.h"
 #include "Converters/GLTFMeshConverters.h"
 #include "Converters/GLTFMeshDataConverters.h"
 #include "Converters/GLTFMaterialConverters.h"
 #include "Converters/GLTFSamplerConverters.h"
 #include "Converters/GLTFTextureConverters.h"
+#include "Converters/GLTFImageConverters.h"
 #include "Converters/GLTFNodeConverters.h"
 #include "Converters/GLTFSkinConverters.h"
 #include "Converters/GLTFAnimationConverters.h"
@@ -22,7 +23,7 @@
 #include "Converters/GLTFKhrVariantConverters.h"
 #include "Converters/GLTFSkySphereConverters.h"
 
-class FGLTFConvertBuilder : public FGLTFImageBuilder
+class FGLTFConvertBuilder : public FGLTFBufferBuilder
 {
 protected:
 
@@ -70,6 +71,7 @@ public:
 	FGLTFJsonTextureIndex GetOrAddTexture(const UTextureCube* Texture, ECubeFace CubeFace, bool bToSRGB);
 	FGLTFJsonTextureIndex GetOrAddTexture(const UTextureRenderTarget2D* Texture, bool bToSRGB);
 	FGLTFJsonTextureIndex GetOrAddTexture(const UTextureRenderTargetCube* Texture, ECubeFace CubeFace, bool bToSRGB);
+	FGLTFJsonImageIndex GetOrAddImage(TGLTFSharedArray<FColor>& Pixels, FIntPoint Size, bool bIgnoreAlpha, EGLTFTextureType Type, const FString& Name);
 
 	FGLTFJsonSkinIndex GetOrAddSkin(FGLTFJsonNodeIndex RootNode, const USkeletalMesh* SkeletalMesh);
 	FGLTFJsonSkinIndex GetOrAddSkin(FGLTFJsonNodeIndex RootNode, const USkeletalMeshComponent* SkeletalMeshComponent);
@@ -123,6 +125,7 @@ private:
 	FGLTFTextureRenderTarget2DConverter TextureRenderTarget2DConverter{ *this };
 	FGLTFTextureRenderTargetCubeConverter TextureRenderTargetCubeConverter{ *this };
 	FGLTFTextureLightMapConverter TextureLightMapConverter{ *this };
+	FGLTFImageConverter ImageConverter{ *this };
 
 	FGLTFSkinConverter SkinConverter{ *this };
 	FGLTFAnimationConverter AnimationConverter{ *this };
diff --git a/Engine/Plugins/Enterprise/GLTFExporter/MigratedFromGit/Source/GLTFExporter/Private/Builders/GLTFFileUtility.cpp b/Engine/Plugins/Enterprise/GLTFExporter/MigratedFromGit/Source/GLTFExporter/Private/Builders/GLTFFileUtility.cpp
deleted file mode 100644
diff --git a/Engine/Plugins/Enterprise/GLTFExporter/MigratedFromGit/Source/GLTFExporter/Private/Builders/GLTFFileUtility.h b/Engine/Plugins/Enterprise/GLTFExporter/MigratedFromGit/Source/GLTFExporter/Private/Builders/GLTFFileUtility.h
deleted file mode 100644
diff --git a/Engine/Plugins/Enterprise/GLTFExporter/MigratedFromGit/Source/GLTFExporter/Private/Builders/GLTFImageBuilder.cpp b/Engine/Plugins/Enterprise/GLTFExporter/MigratedFromGit/Source/GLTFExporter/Private/Builders/GLTFImageBuilder.cpp
deleted file mode 100644
diff --git a/Engine/Plugins/Enterprise/GLTFExporter/MigratedFromGit/Source/GLTFExporter/Private/Builders/GLTFImageBuilder.h b/Engine/Plugins/Enterprise/GLTFExporter/MigratedFromGit/Source/GLTFExporter/Private/Builders/GLTFImageBuilder.h
deleted file mode 100644
diff --git a/Engine/Plugins/Enterprise/GLTFExporter/MigratedFromGit/Source/GLTFExporter/Private/Builders/GLTFImageUtility.cpp b/Engine/Plugins/Enterprise/GLTFExporter/MigratedFromGit/Source/GLTFExporter/Private/Builders/GLTFImageUtility.cpp
deleted file mode 100644
diff --git a/Engine/Plugins/Enterprise/GLTFExporter/MigratedFromGit/Source/GLTFExporter/Private/Builders/GLTFImageUtility.h b/Engine/Plugins/Enterprise/GLTFExporter/MigratedFromGit/Source/GLTFExporter/Private/Builders/GLTFImageUtility.h
deleted file mode 100644
diff --git a/Engine/Plugins/Enterprise/GLTFExporter/MigratedFromGit/Source/GLTFExporter/Private/Converters/GLTFImageConverters.cpp b/Engine/Plugins/Enterprise/GLTFExporter/MigratedFromGit/Source/GLTFExporter/Private/Converters/GLTFImageConverters.cpp
new file mode 100644
--- /dev/null
+++ b/Engine/Plugins/Enterprise/GLTFExporter/MigratedFromGit/Source/GLTFExporter/Private/Converters/GLTFImageConverters.cpp
@@ -0,0 +1,85 @@
+// Copyright Epic Games, Inc. All Rights Reserved.
+
+#include "Converters/GLTFImageConverters.h"
+#include "Converters/GLTFImageUtility.h"
+#include "Builders/GLTFConvertBuilder.h"
+#include "Misc/FileHelper.h"
+
+FGLTFJsonImageIndex FGLTFImageConverter::Convert(TGLTFSuperfluous<FString> Name, EGLTFTextureType Type, bool bIgnoreAlpha, FIntPoint Size, TGLTFSharedArray<FColor> Pixels)
+{
+	TArray64<uint8> CompressedData;
+
+	const EGLTFJsonMimeType MimeType = GetMimeType(Pixels->GetData(), Size, bIgnoreAlpha, Type);
+	switch (MimeType)
+	{
+		case EGLTFJsonMimeType::None:
+			return FGLTFJsonImageIndex(INDEX_NONE);
+
+		case EGLTFJsonMimeType::PNG:
+			FGLTFImageUtility::CompressToPNG(Pixels->GetData(), Size, CompressedData);
+			break;
+
+		case EGLTFJsonMimeType::JPEG:
+			FGLTFImageUtility::CompressToJPEG(Pixels->GetData(), Size, Builder.ExportOptions->TextureImageQuality, CompressedData);
+			break;
+
+		default:
+			checkNoEntry();
+			break;
+	}
+
+	FGLTFJsonImage JsonImage;
+
+	if (Builder.bIsGlbFile)
+	{
+		JsonImage.Name = Name;
+		JsonImage.MimeType = MimeType;
+		JsonImage.BufferView = Builder.AddBufferView(CompressedData.GetData(), CompressedData.Num());
+	}
+	else
+	{
+		JsonImage.Uri = SaveToFile(CompressedData.GetData(), CompressedData.Num(), MimeType, Name);
+	}
+
+	return Builder.AddImage(JsonImage);
+}
+
+EGLTFJsonMimeType FGLTFImageConverter::GetMimeType(const FColor* Pixels, FIntPoint Size, bool bIgnoreAlpha, EGLTFTextureType Type) const
+{
+	switch (Builder.ExportOptions->TextureImageFormat)
+	{
+		case EGLTFTextureImageFormat::None:
+			return EGLTFJsonMimeType::None;
+
+		case EGLTFTextureImageFormat::PNG:
+			return EGLTFJsonMimeType::PNG;
+
+		case EGLTFTextureImageFormat::JPEG:
+			return
+				(Type == EGLTFTextureType::None || !EnumHasAllFlags(static_cast<EGLTFTextureType>(Builder.ExportOptions->NoLossyImageFormatFor), Type)) &&
+				(bIgnoreAlpha || FGLTFImageUtility::NoAlphaNeeded(Pixels, Size)) ?
+				EGLTFJsonMimeType::JPEG : EGLTFJsonMimeType::PNG;
+
+		default:
+			checkNoEntry();
+		return EGLTFJsonMimeType::None;
+	}
+}
+
+FString FGLTFImageConverter::SaveToFile(const void* CompressedData, int64 CompressedByteLength, EGLTFJsonMimeType MimeType, const FString& Name)
+{
+	const TCHAR* Extension = FGLTFImageUtility::GetFileExtension(MimeType);
+	const FString ImageUri = FGLTFImageUtility::GetUniqueFilename(Name, Extension, UniqueImageUris);
+
+	const TArrayView<const uint8> ImageData(static_cast<const uint8*>(CompressedData), CompressedByteLength);
+	const FString ImagePath = FPaths::Combine(Builder.DirPath, ImageUri);
+
+	if (!FFileHelper::SaveArrayToFile(ImageData, *ImagePath))
+	{
+		Builder.LogError(FString::Printf(TEXT("Failed to save image %s to file: %s"), *Name, *ImagePath));
+		return TEXT("");
+	}
+
+	UniqueImageUris.Add(ImageUri);
+	return ImageUri;
+}
diff --git a/Engine/Plugins/Enterprise/GLTFExporter/MigratedFromGit/Source/GLTFExporter/Private/Converters/GLTFImageConverters.h b/Engine/Plugins/Enterprise/GLTFExporter/MigratedFromGit/Source/GLTFExporter/Private/Converters/GLTFImageConverters.h
new file mode 100644
--- /dev/null
+++ b/Engine/Plugins/Enterprise/GLTFExporter/MigratedFromGit/Source/GLTFExporter/Private/Converters/GLTFImageConverters.h
@@ -0,0 +1,24 @@
+// Copyright Epic Games, Inc. All Rights Reserved.
+
+#pragma once
+
+#include "Json/GLTFJsonIndex.h"
+#include "Json/GLTFJsonEnums.h"
+#include "Converters/GLTFConverter.h"
+#include "Converters/GLTFBuilderContext.h"
+#include "Converters/GLTFSuperfluous.h"
+#include "Converters/GLTFSharedArray.h"
+#include "GLTFExportOptions.h"
+
+class FGLTFImageConverter final : public FGLTFBuilderContext, public TGLTFConverter<FGLTFJsonImageIndex, TGLTFSuperfluous<FString>, EGLTFTextureType, bool, FIntPoint, TGLTFSharedArray<FColor>>
+{
+	using FGLTFBuilderContext::FGLTFBuilderContext;
+
+	virtual FGLTFJsonImageIndex Convert(TGLTFSuperfluous<FString> Name, EGLTFTextureType Type, bool bIgnoreAlpha, FIntPoint Size, TGLTFSharedArray<FColor> Pixels) override;
+
+	EGLTFJsonMimeType GetMimeType(const FColor* Pixels, FIntPoint Size, bool bIgnoreAlpha, EGLTFTextureType Type) const;
+
+	FString SaveToFile(const void* CompressedData, int64 CompressedByteLength, EGLTFJsonMimeType MimeType, const FString& Name);
+
+	TSet<FString> UniqueImageUris;
+};
diff --git a/Engine/Plugins/Enterprise/GLTFExporter/MigratedFromGit/Source/GLTFExporter/Private/Converters/GLTFImageUtility.cpp b/Engine/Plugins/Enterprise/GLTFExporter/MigratedFromGit/Source/GLTFExporter/Private/Converters/GLTFImageUtility.cpp
new file mode 100644
--- /dev/null
+++ b/Engine/Plugins/Enterprise/GLTFExporter/MigratedFromGit/Source/GLTFExporter/Private/Converters/GLTFImageUtility.cpp
@@ -0,0 +1,84 @@
+// Copyright Epic Games, Inc. All Rights Reserved.
+
+#include "Converters/GLTFImageUtility.h"
+#include "Modules/ModuleManager.h"
+#include "IImageWrapperModule.h"
+#include "IImageWrapper.h"
+
+const TCHAR* FGLTFImageUtility::GetFileExtension(EGLTFJsonMimeType MimeType)
+{
+	switch (MimeType)
+	{
+		case EGLTFJsonMimeType::PNG:  return TEXT(".png");
+		case EGLTFJsonMimeType::JPEG: return TEXT(".jpg");
+		default:
+			checkNoEntry();
+		return TEXT("");
+	}
+}
+
+FString FGLTFImageUtility::GetUniqueFilename(const FString& BaseFilename, const FString& FileExtension, const TSet<FString>& UniqueFilenames)
+{
+	FString Filename = BaseFilename + FileExtension;
+	if (!UniqueFilenames.Contains(Filename))
+	{
+		return Filename;
+	}
+
+	int32 Suffix = 1;
+	do
+	{
+		Filename = BaseFilename + TEXT('_') + FString::FromInt(Suffix) + FileExtension;
+		Suffix++;
+	}
+	while (UniqueFilenames.Contains(Filename));
+
+	return Filename;
+}
+
+bool FGLTFImageUtility::NoAlphaNeeded(const FColor* Pixels, FIntPoint Size)
+{
+	const int64 Count = Size.X * Size.Y;
+	for (int64 Index = 0; Index < Count; ++Index)
+	{
+		const FColor& Pixel = Pixels[Index];
+		if (Pixel.A < 255)
+		{
+			return false;
+		}
+	}
+
+	return true;
+}
+
+bool FGLTFImageUtility::CompressToPNG(const FColor* InPixels, FIntPoint InSize, TArray64<uint8>& OutCompressedData)
+{
+	return CompressToFormat(InPixels, InSize, EImageFormat::PNG, 0, OutCompressedData);
+}
+
+bool FGLTFImageUtility::CompressToJPEG(const FColor* InPixels, FIntPoint InSize, int32 InCompressionQuality, TArray64<uint8>& OutCompressedData)
+{
+	// NOTE: Tests have shown that explicitly checking for greyscale images (and using EImageFormat::GrayscaleJPEG) results in the same as EImageFormat::JPEG
+	return CompressToFormat(InPixels, InSize, EImageFormat::JPEG, InCompressionQuality, OutCompressedData);
+}
+
+bool FGLTFImageUtility::CompressToFormat(const FColor* InPixels, FIntPoint InSize, EImageFormat InCompressionFormat, int32 InCompressionQuality, TArray64<uint8>& OutCompressedData)
+{
+	const int64 ByteLength = InSize.X * InSize.Y * sizeof(FColor);
+	return CompressToFormat(InPixels, ByteLength, InSize.X, InSize.Y, ERGBFormat::BGRA, 8, InCompressionFormat, InCompressionQuality, OutCompressedData);
+}
+
+bool FGLTFImageUtility::CompressToFormat(const void* InRawData, int64 InRawSize, int32 InWidth, int32 InHeight, ERGBFormat InRGBFormat, int32 InBitDepth, EImageFormat InCompressionFormat, int32 InCompressionQuality, TArray64<uint8>& OutCompressedData)
+{
+	IImageWrapperModule& ImageWrapperModule = FModuleManager::Get().LoadModuleChecked<IImageWrapperModule>(TEXT("ImageWrapper"));
+	const TSharedPtr<IImageWrapper> ImageWrapper = ImageWrapperModule.CreateImageWrapper(InCompressionFormat);
+
+	if (ImageWrapper.IsValid() && ImageWrapper->SetRaw(InRawData, InRawSize, InWidth, InHeight, InRGBFormat, InBitDepth))
+	{
+		const TArray64<uint8>& CompressedData = ImageWrapper->GetCompressed(InCompressionQuality);
+		OutCompressedData = MoveTemp(const_cast<TArray64<uint8>&>(CompressedData)); // safe to do this since we don't keep image wrapper
+		return true;
+	}
+
+	return false;
+}
diff --git a/Engine/Plugins/Enterprise/GLTFExporter/MigratedFromGit/Source/GLTFExporter/Private/Converters/GLTFImageUtility.h b/Engine/Plugins/Enterprise/GLTFExporter/MigratedFromGit/Source/GLTFExporter/Private/Converters/GLTFImageUtility.h
new file mode 100644
--- /dev/null
+++ b/Engine/Plugins/Enterprise/GLTFExporter/MigratedFromGit/Source/GLTFExporter/Private/Converters/GLTFImageUtility.h
@@ -0,0 +1,24 @@
+// Copyright Epic Games, Inc. All Rights Reserved.
+
+#pragma once
+
+#include "CoreMinimal.h"
+#include "Json/GLTFJsonEnums.h"
+
+enum class ERGBFormat : int8;
+enum class EImageFormat : int8;
+
+struct FGLTFImageUtility
+{
+	static const TCHAR* GetFileExtension(EGLTFJsonMimeType MimeType);
+
+	static FString GetUniqueFilename(const FString& BaseFilename, const FString& FileExtension, const TSet<FString>& UniqueFilenames);
+
+	static bool NoAlphaNeeded(const FColor* Pixels, FIntPoint Size);
+
+	static bool CompressToPNG(const FColor* InPixels, FIntPoint InSize, TArray64<uint8>& OutCompressedData);
+	static bool CompressToJPEG(const FColor* InPixels, FIntPoint InSize, int32 InCompressionQuality, TArray64<uint8>& OutCompressedData);
+
+	static bool CompressToFormat(const FColor* InPixels, FIntPoint InSize, EImageFormat InCompressionFormat, int32 InCompressionQuality, TArray64<uint8>& OutCompressedData);
+	static bool CompressToFormat(const void* InRawData, int64 InRawSize, int32 InWidth, int32 InHeight, ERGBFormat InRGBFormat, int32 InBitDepth, EImageFormat InCompressionFormat, int32 InCompressionQuality, TArray64<uint8>& OutCompressedData);
+};
diff --git a/Engine/Plugins/Enterprise/GLTFExporter/MigratedFromGit/Source/GLTFExporter/Private/Converters/GLTFIndexArray.h b/Engine/Plugins/Enterprise/GLTFExporter/MigratedFromGit/Source/GLTFExporter/Private/Converters/GLTFIndexArray.h
--- a/Engine/Plugins/Enterprise/GLTFExporter/MigratedFromGit/Source/GLTFExporter/Private/Converters/GLTFIndexArray.h
+++ b/Engine/Plugins/Enterprise/GLTFExporter/MigratedFromGit/Source/GLTFExporter/Private/Converters/GLTFIndexArray.h
@@ -4,16 +4,20 @@
 
 #include "CoreMinimal.h"
 
-typedef TArray<int32> FGLTFIndexArray;
-
-#if (ENGINE_MAJOR_VERSION > 4 || ENGINE_MINOR_VERSION >= 26)
-inline uint32 GetTypeHash(const FGLTFIndexArray& IndexArray)
+class FGLTFIndexArray : public TArray<int32>
 {
-	uint32 Hash = GetTypeHash(IndexArray.Num());
-	for (int32 Index : IndexArray)
+public:
+
+	using TArray::TArray;
+	using TArray::operator=;
+
+	friend uint32 GetTypeHash(const FGLTFIndexArray& IndexArray)
 	{
-		Hash = HashCombine(Hash, GetTypeHash(Index));
+		uint32 Hash = GetTypeHash(IndexArray.Num());
+		for (const int32 Index : IndexArray)
+		{
+			Hash = HashCombine(Hash, GetTypeHash(Index));
+		}
+		return Hash;
 	}
-	return Hash;
-}
-#endif
+};
diff --git a/Engine/Plugins/Enterprise/GLTFExporter/MigratedFromGit/Source/GLTFExporter/Private/Converters/GLTFMaterialArray.h b/Engine/Plugins/Enterprise/GLTFExporter/MigratedFromGit/Source/GLTFExporter/Private/Converters/GLTFMaterialArray.h
--- a/Engine/Plugins/Enterprise/GLTFExporter/MigratedFromGit/Source/GLTFExporter/Private/Converters/GLTFMaterialArray.h
+++ b/Engine/Plugins/Enterprise/GLTFExporter/MigratedFromGit/Source/GLTFExporter/Private/Converters/GLTFMaterialArray.h
@@ -6,14 +6,20 @@
 
 class UMaterialInterface;
 
-typedef TArray<const UMaterialInterface*> FGLTFMaterialArray;
-
-inline uint32 GetTypeHash(const FGLTFMaterialArray& MaterialArray)
+class FGLTFMaterialArray : public TArray<const UMaterialInterface*>
 {
-	uint32 Hash = GetTypeHash(MaterialArray.Num());
-	for (const UMaterialInterface* Material : MaterialArray)
+public:
+
+	using TArray::TArray;
+	using TArray::operator=;
+
+	friend uint32 GetTypeHash(const FGLTFMaterialArray& MaterialArray)
 	{
-		Hash = HashCombine(Hash, GetTypeHash(Material));
+		uint32 Hash = GetTypeHash(MaterialArray.Num());
+		for (const UMaterialInterface* Material : MaterialArray)
+		{
+			Hash = HashCombine(Hash, GetTypeHash(Material));
+		}
+		return Hash;
 	}
-	return Hash;
-}
+};
diff --git a/Engine/Plugins/Enterprise/GLTFExporter/MigratedFromGit/Source/GLTFExporter/Private/Converters/GLTFMaterialUtility.cpp b/Engine/Plugins/Enterprise/GLTFExporter/MigratedFromGit/Source/GLTFExporter/Private/Converters/GLTFMaterialUtility.cpp
--- a/Engine/Plugins/Enterprise/GLTFExporter/MigratedFromGit/Source/GLTFExporter/Private/Converters/GLTFMaterialUtility.cpp
+++ b/Engine/Plugins/Enterprise/GLTFExporter/MigratedFromGit/Source/GLTFExporter/Private/Converters/GLTFMaterialUtility.cpp
@@ -228,7 +228,7 @@ FGLTFPropertyBakeOutput FGLTFMaterialUtility::BakeMaterialProperty(const FIntPoi
 
 	FBakeOutputEx& BakeOutput = BakeOutputs[0];
 
-	TArray<FColor> BakedPixels = MoveTemp(BakeOutput.PropertyData.FindChecked(Property));
+	TGLTFSharedArray<FColor> BakedPixels = MoveTemp(BakeOutput.PropertyData.FindChecked(Property));
 	const FIntPoint BakedSize = BakeOutput.PropertySizes.FindChecked(Property);
 	const float EmissiveScale = BakeOutput.EmissiveScale;
 
@@ -236,7 +236,7 @@ FGLTFPropertyBakeOutput FGLTFMaterialUtility::BakeMaterialProperty(const FIntPoi
 	{
 		// NOTE: alpha is 0 by default after baking a property, but we prefer 255 (1.0).
 		// It makes it easier to view the exported textures.
-		for (FColor& Pixel: BakedPixels)
+		for (FColor& Pixel: *BakedPixels)
 		{
 			Pixel.A = 255;
 		}
@@ -245,18 +245,18 @@ FGLTFPropertyBakeOutput FGLTFMaterialUtility::BakeMaterialProperty(const FIntPoi
 	if (IsNormalMap(Property))
 	{
 		// Convert normalmaps to use +Y (OpenGL / WebGL standard)
-		FGLTFTextureUtility::FlipGreenChannel(BakedPixels);
+		FGLTFTextureUtility::FlipGreenChannel(*BakedPixels);
 	}
 
 	bool bFromSRGB = !bIsLinearBake;
 	bool bToSRGB = IsSRGB(Property);
-	FGLTFTextureUtility::TransformColorSpace(BakedPixels, bFromSRGB, bToSRGB);
+	FGLTFTextureUtility::TransformColorSpace(*BakedPixels, bFromSRGB, bToSRGB);
 
 	FGLTFPropertyBakeOutput PropertyBakeOutput(Property, PF_B8G8R8A8, BakedPixels, BakedSize, EmissiveScale, !bIsLinearBake);
 
-	if (BakedPixels.Num() == 1)
+	if (BakedPixels->Num() == 1)
 	{
-		const FColor& Pixel = BakedPixels[0];
+		const FColor& Pixel = (*BakedPixels)[0];
 
 		PropertyBakeOutput.bIsConstant = true;
 		PropertyBakeOutput.ConstantValue = bToSRGB ? FLinearColor(Pixel) : Pixel.ReinterpretAsLinear();
@@ -265,7 +265,7 @@ FGLTFPropertyBakeOutput FGLTFMaterialUtility::BakeMaterialProperty(const FIntPoi
 	return PropertyBakeOutput;
 }
 
-FGLTFJsonTextureIndex FGLTFMaterialUtility::AddTexture(FGLTFConvertBuilder& Builder, const TArray<FColor>& Pixels, const FIntPoint& TextureSize, bool bIgnoreAlpha, bool bIsNormalMap, const FString& TextureName, EGLTFJsonTextureFilter MinFilter, EGLTFJsonTextureFilter MagFilter, EGLTFJsonTextureWrap WrapS, EGLTFJsonTextureWrap WrapT)
+FGLTFJsonTextureIndex FGLTFMaterialUtility::AddTexture(FGLTFConvertBuilder& Builder, TGLTFSharedArray<FColor>& Pixels, const FIntPoint& TextureSize, bool bIgnoreAlpha, bool bIsNormalMap, const FString& TextureName, EGLTFJsonTextureFilter MinFilter, EGLTFJsonTextureFilter MagFilter, EGLTFJsonTextureWrap WrapS, EGLTFJsonTextureWrap WrapT)
 {
 	// TODO: maybe we should reuse existing samplers?
 	FGLTFJsonSampler JsonSampler;
@@ -279,7 +279,7 @@ FGLTFJsonTextureIndex FGLTFMaterialUtility::AddTexture(FGLTFConvertBuilder& Buil
 	FGLTFJsonTexture JsonTexture;
 	JsonTexture.Name = TextureName;
 	JsonTexture.Sampler = Builder.AddSampler(JsonSampler);
-	JsonTexture.Source = Builder.AddImage(Pixels, TextureSize, bIgnoreAlpha, bIsNormalMap ? EGLTFTextureType::Normalmaps : EGLTFTextureType::None, TextureName);
+	JsonTexture.Source = Builder.GetOrAddImage(Pixels, TextureSize, bIgnoreAlpha, bIsNormalMap ? EGLTFTextureType::Normalmaps : EGLTFTextureType::None, TextureName);
 
 	return Builder.AddTexture(JsonTexture);
 }
diff --git a/Engine/Plugins/Enterprise/GLTFExporter/MigratedFromGit/Source/GLTFExporter/Private/Converters/GLTFMaterialUtility.h b/Engine/Plugins/Enterprise/GLTFExporter/MigratedFromGit/Source/GLTFExporter/Private/Converters/GLTFMaterialUtility.h
--- a/Engine/Plugins/Enterprise/GLTFExporter/MigratedFromGit/Source/GLTFExporter/Private/Converters/GLTFMaterialUtility.h
+++ b/Engine/Plugins/Enterprise/GLTFExporter/MigratedFromGit/Source/GLTFExporter/Private/Converters/GLTFMaterialUtility.h
@@ -3,6 +3,7 @@
 #pragma once
 
 #include "Builders/GLTFContainerBuilder.h"
+#include "Converters/GLTFSharedArray.h"
 #include "Engine/Texture2D.h"
 
 struct FGLTFMaterialAnalysis;
@@ -10,7 +11,7 @@ struct FMaterialPropertyEx;
 
 struct FGLTFTextureCombineSource
 {
-	constexpr FORCEINLINE FGLTFTextureCombineSource(const UTexture2D* Texture, FLinearColor TintColor = { 1.0f, 1.0f, 1.0f, 1.0f }, ESimpleElementBlendMode BlendMode = SE_BLEND_Additive)
+	constexpr FGLTFTextureCombineSource(const UTexture2D* Texture, FLinearColor TintColor = { 1.0f, 1.0f, 1.0f, 1.0f }, ESimpleElementBlendMode BlendMode = SE_BLEND_Additive)
 		: Texture(Texture), TintColor(TintColor), BlendMode(BlendMode)
 	{}
 
@@ -21,13 +22,13 @@ struct FGLTFTextureCombineSource
 
 struct FGLTFPropertyBakeOutput
 {
-	FORCEINLINE FGLTFPropertyBakeOutput(const FMaterialPropertyEx& Property, EPixelFormat PixelFormat, TArray<FColor>& Pixels, FIntPoint Size, float EmissiveScale, bool bIsSRGB)
+	FGLTFPropertyBakeOutput(const FMaterialPropertyEx& Property, EPixelFormat PixelFormat, const TGLTFSharedArray<FColor>& Pixels, FIntPoint Size, float EmissiveScale, bool bIsSRGB)
 		: Property(Property), PixelFormat(PixelFormat), Pixels(Pixels), Size(Size), EmissiveScale(EmissiveScale), bIsSRGB(bIsSRGB), bIsConstant(false)
 	{}
 
 	const FMaterialPropertyEx& Property;
 	EPixelFormat PixelFormat;
-	TArray<FColor> Pixels;
+	TGLTFSharedArray<FColor> Pixels;
 	FIntPoint Size;
 	float EmissiveScale;
 	bool bIsSRGB;
@@ -64,7 +65,7 @@ struct FGLTFMaterialUtility
 
 	static FGLTFPropertyBakeOutput BakeMaterialProperty(const FIntPoint& OutputSize, const FMaterialPropertyEx& Property, const UMaterialInterface* Material, int32 TexCoord, const FGLTFMeshData* MeshData = nullptr, const FGLTFIndexArray& MeshSectionIndices = {}, bool bFillAlpha = true);
 
-	static FGLTFJsonTextureIndex AddTexture(FGLTFConvertBuilder& Builder, const TArray<FColor>& Pixels, const FIntPoint& TextureSize, bool bIgnoreAlpha, bool bIsNormalMap, const FString& TextureName, EGLTFJsonTextureFilter MinFilter, EGLTFJsonTextureFilter MagFilter, EGLTFJsonTextureWrap WrapS, EGLTFJsonTextureWrap WrapT);
+	static FGLTFJsonTextureIndex AddTexture(FGLTFConvertBuilder& Builder, TGLTFSharedArray<FColor>& Pixels, const FIntPoint& TextureSize, bool bIgnoreAlpha, bool bIsNormalMap, const FString& TextureName, EGLTFJsonTextureFilter MinFilter, EGLTFJsonTextureFilter MagFilter, EGLTFJsonTextureWrap WrapS, EGLTFJsonTextureWrap WrapT);
 
 	static FLinearColor GetMask(const FExpressionInput& ExpressionInput);
 	static uint32 GetMaskComponentCount(const FExpressionInput& ExpressionInput);
diff --git a/Engine/Plugins/Enterprise/GLTFExporter/MigratedFromGit/Source/GLTFExporter/Private/Converters/GLTFSharedArray.h b/Engine/Plugins/Enterprise/GLTFExporter/MigratedFromGit/Source/GLTFExporter/Private/Converters/GLTFSharedArray.h
new file mode 100644
--- /dev/null
+++ b/Engine/Plugins/Enterprise/GLTFExporter/MigratedFromGit/Source/GLTFExporter/Private/Converters/GLTFSharedArray.h
@@ -0,0 +1,76 @@
+// Copyright Epic Games, Inc. All Rights Reserved.
+
+#pragma once
+
+#include "CoreMinimal.h"
+
+template <typename ElementType, typename ArrayType = TArray<ElementType>, ESPMode Mode = ESPMode::Fast>
+class TGLTFSharedArray : public TSharedRef<ArrayType>
+{
+public:
+
+	template <typename... ArgTypes>
+	TGLTFSharedArray(ArgTypes&&... Args)
+		: TSharedRef(MakeShared<ArrayType>(Forward<ArgTypes>(Args)...))
+	{
+	}
+
+	TGLTFSharedArray(TGLTFSharedArray& SharedArray)
+		: TSharedRef(SharedArray)
+	{
+	}
+
+	TGLTFSharedArray(TGLTFSharedArray const& SharedArray)
+		: TSharedRef(SharedArray)
+	{
+	}
+
+	TGLTFSharedArray(TGLTFSharedArray&& SharedArray) noexcept
+		: TSharedRef(SharedArray)
+	{
+	}
+
+	TGLTFSharedArray& operator=(TGLTFSharedArray& Other)
+	{
+		TSharedRef<ArrayType>::operator=(Other);
+		return *this;
+	}
+
+	TGLTFSharedArray& operator=(TGLTFSharedArray const& Other)
+	{
+		TSharedRef<ArrayType>::operator=(Other);
+		return *this;
+	}
+
+	TGLTFSharedArray& operator=(TGLTFSharedArray&& Other) noexcept
+	{
+		TSharedRef<ArrayType>::operator=(Other);
+		return *this;
+	}
+
+	bool operator==(const TGLTFSharedArray& Other) const
+	{
+		return this->Get() == Other.Get();
+	}
+
+	bool operator!=(const TGLTFSharedArray& Other) const
+	{
+		return this->Get() != Other.Get();
+	}
+
+	using TSharedRef<ArrayType>::TSharedRef;
+	using TSharedRef<ArrayType>::operator=;
+
+	friend uint32 GetTypeHash(const TGLTFSharedArray& SharedArray)
+	{
+		const ArrayType& Array = SharedArray.Get();
+		uint32 Hash = GetTypeHash(Array.Num());
+
+		for (const auto& Element : Array)
+		{
+			Hash = HashCombine(Hash, GetTypeHash(Element));
+		}
+
+		return Hash;
+	}
+};
diff --git a/Engine/Plugins/Enterprise/GLTFExporter/MigratedFromGit/Source/GLTFExporter/Private/Converters/GLTFSuperfluous.h b/Engine/Plugins/Enterprise/GLTFExporter/MigratedFromGit/Source/GLTFExporter/Private/Converters/GLTFSuperfluous.h
new file mode 100644
--- /dev/null
+++ b/Engine/Plugins/Enterprise/GLTFExporter/MigratedFromGit/Source/GLTFExporter/Private/Converters/GLTFSuperfluous.h
@@ -0,0 +1,34 @@
+// Copyright Epic Games, Inc. All Rights Reserved.
+
+#pragma once
+
+#include "CoreMinimal.h"
+
+template <typename Type>
+class TGLTFSuperfluous : public Type
+{
+public:
+
+	TGLTFSuperfluous(const Type& Other)
+		: Type(Other)
+	{
+	}
+
+	using Type::Type;
+	using Type::operator=;
+
+	bool operator==(const TGLTFSuperfluous& Other) const
+	{
+		return true;
+	}
+
+	bool operator!=(const TGLTFSuperfluous& Other) const
+	{
+		return false;
+	}
+
+	friend uint32 GetTypeHash(const TGLTFSuperfluous& Superfluous)
+	{
+		return 0;
+	}
+};
diff --git a/Engine/Plugins/Enterprise/GLTFExporter/MigratedFromGit/Source/GLTFExporter/Private/Tasks/GLTFMaterialTasks.cpp b/Engine/Plugins/Enterprise/GLTFExporter/MigratedFromGit/Source/GLTFExporter/Private/Tasks/GLTFMaterialTasks.cpp
--- a/Engine/Plugins/Enterprise/GLTFExporter/MigratedFromGit/Source/GLTFExporter/Private/Tasks/GLTFMaterialTasks.cpp
+++ b/Engine/Plugins/Enterprise/GLTFExporter/MigratedFromGit/Source/GLTFExporter/Private/Tasks/GLTFMaterialTasks.cpp
@@ -427,11 +427,11 @@ bool FGLTFMaterialTask::TryGetBaseColorAndOpacity(FGLTFJsonPBRMetallicRoughness&
 		return false;
 	}
 
-	TArray<FColor> CombinedPixels;
+	TGLTFSharedArray<FColor> CombinedPixels;
 	CombinePixels(
-		BaseColorBakeOutput.Pixels,
-		OpacityBakeOutput.Pixels,
-		CombinedPixels,
+		*BaseColorBakeOutput.Pixels,
+		*OpacityBakeOutput.Pixels,
+		*CombinedPixels,
 		[](const FColor& BaseColor, const FColor& Opacity) -> FColor
 		{
 			return FColor(BaseColor.R, BaseColor.G, BaseColor.B, Opacity.R);
@@ -551,11 +551,11 @@ bool FGLTFMaterialTask::TryGetMetallicAndRoughness(FGLTFJsonPBRMetallicRoughness
 		return false;
 	}
 
-	TArray<FColor> CombinedPixels;
+	TGLTFSharedArray<FColor> CombinedPixels;
 	CombinePixels(
-		MetallicBakeOutput.Pixels,
-		RoughnessBakeOutput.Pixels,
-		CombinedPixels,
+		*MetallicBakeOutput.Pixels,
+		*RoughnessBakeOutput.Pixels,
+		*CombinedPixels,
 		[](const FColor& Metallic, const FColor& Roughness) -> FColor
 		{
 			return FColor(0, Roughness.R, Metallic.R);
@@ -676,11 +676,11 @@ bool FGLTFMaterialTask::TryGetClearCoatRoughness(FGLTFJsonClearCoatExtension& Ou
 		return false;
 	}
 
-	TArray<FColor> CombinedPixels;
+	TGLTFSharedArray<FColor> CombinedPixels;
 	CombinePixels(
-		IntensityBakeOutput.Pixels,
-		RoughnessBakeOutput.Pixels,
-		CombinedPixels,
+		*IntensityBakeOutput.Pixels,
+		*RoughnessBakeOutput.Pixels,
+		*CombinedPixels,
 		[](const FColor& Intensity, const FColor& Roughness) -> FColor
 		{
 			return FColor(Intensity.R, Roughness.R, 0);
@@ -731,7 +731,7 @@ bool FGLTFMaterialTask::TryGetEmissive(FGLTFJsonMaterial& JsonMaterial, const FM
 		return false;
 	}
 
-	const FGLTFPropertyBakeOutput PropertyBakeOutput = BakeMaterialProperty(EmissiveProperty, JsonMaterial.EmissiveTexture.TexCoord);
+	FGLTFPropertyBakeOutput PropertyBakeOutput = BakeMaterialProperty(EmissiveProperty, JsonMaterial.EmissiveTexture.TexCoord);
 	const float EmissiveScale = PropertyBakeOutput.EmissiveScale;
 
 	if (PropertyBakeOutput.bIsConstant)
@@ -1356,7 +1356,7 @@ FGLTFPropertyBakeOutput FGLTFMaterialTask::BakeMaterialProperty(const FMaterialP
 
 	// TODO: add support for calculating the ideal resolution to use for baking based on connected (texture) nodes
 
-	const FGLTFPropertyBakeOutput PropertyBakeOutput = FGLTFMaterialUtility::BakeMaterialProperty(
+	return FGLTFMaterialUtility::BakeMaterialProperty(
 		TextureSize,
 		Property,
 		Material,
@@ -1364,11 +1364,9 @@ FGLTFPropertyBakeOutput FGLTFMaterialTask::BakeMaterialProperty(const FMaterialP
 		MeshData,
 		SectionIndices,
 		bFillAlpha);
-
-	return PropertyBakeOutput;
 }
 
-bool FGLTFMaterialTask::StoreBakedPropertyTexture(FGLTFJsonTextureInfo& OutTexInfo, const FGLTFPropertyBakeOutput& PropertyBakeOutput, const FString& PropertyName) const
+bool FGLTFMaterialTask::StoreBakedPropertyTexture(FGLTFJsonTextureInfo& OutTexInfo, FGLTFPropertyBakeOutput& PropertyBakeOutput, const FString& PropertyName) const
 {
 	const TextureAddress TextureAddress = Builder.GetBakeTilingForMaterialProperty(Material, GetPropertyGroup(PropertyBakeOutput.Property));
 	const EGLTFJsonTextureWrap TextureWrapS = FGLTFConverterUtility::ConvertWrap(TextureAddress);
diff --git a/Engine/Plugins/Enterprise/GLTFExporter/MigratedFromGit/Source/GLTFExporter/Private/Tasks/GLTFMaterialTasks.h b/Engine/Plugins/Enterprise/GLTFExporter/MigratedFromGit/Source/GLTFExporter/Private/Tasks/GLTFMaterialTasks.h
--- a/Engine/Plugins/Enterprise/GLTFExporter/MigratedFromGit/Source/GLTFExporter/Private/Tasks/GLTFMaterialTasks.h
+++ b/Engine/Plugins/Enterprise/GLTFExporter/MigratedFromGit/Source/GLTFExporter/Private/Tasks/GLTFMaterialTasks.h
@@ -71,7 +71,7 @@ private:
 	FGLTFPropertyBakeOutput BakeMaterialProperty(const FMaterialPropertyEx& Property, int32& OutTexCoord);
 	FGLTFPropertyBakeOutput BakeMaterialProperty(const FMaterialPropertyEx& Property, int32& OutTexCoord, const FIntPoint& TextureSize, bool bFillAlpha);
 
-	bool StoreBakedPropertyTexture(FGLTFJsonTextureInfo& OutTexInfo, const FGLTFPropertyBakeOutput& PropertyBakeOutput, const FString& PropertyName) const;
+	bool StoreBakedPropertyTexture(FGLTFJsonTextureInfo& OutTexInfo, FGLTFPropertyBakeOutput& PropertyBakeOutput, const FString& PropertyName) const;
 
 	FString GetMaterialName() const;
 	FString GetBakedTextureName(const FString& PropertyName) const;
diff --git a/Engine/Plugins/Enterprise/GLTFExporter/MigratedFromGit/Source/GLTFExporter/Private/Tasks/GLTFTextureTasks.cpp b/Engine/Plugins/Enterprise/GLTFExporter/MigratedFromGit/Source/GLTFExporter/Private/Tasks/GLTFTextureTasks.cpp
--- a/Engine/Plugins/Enterprise/GLTFExporter/MigratedFromGit/Source/GLTFExporter/Private/Tasks/GLTFTextureTasks.cpp
+++ b/Engine/Plugins/Enterprise/GLTFExporter/MigratedFromGit/Source/GLTFExporter/Private/Tasks/GLTFTextureTasks.cpp
@@ -39,8 +39,8 @@ void FGLTFTexture2DTask::Complete()
 		JsonTexture.Encoding = Builder.GetTextureHDREncoding();
 	}
 
-	TArray<FColor> Pixels;
-	if (!FGLTFTextureUtility::ReadPixels(RenderTarget, Pixels, JsonTexture.Encoding))
+	TGLTFSharedArray<FColor> Pixels;
+	if (!FGLTFTextureUtility::ReadPixels(RenderTarget, *Pixels, JsonTexture.Encoding))
 	{
 		Builder.LogWarning(FString::Printf(TEXT("Failed to read pixels for 2D texture %s"), *JsonTexture.Name));
 		return;
@@ -48,17 +48,17 @@ void FGLTFTexture2DTask::Complete()
 
 	if (Texture2D->IsNormalMap())
 	{
-		FGLTFTextureUtility::FlipGreenChannel(Pixels);
+		FGLTFTextureUtility::FlipGreenChannel(*Pixels);
 	}
 
-	FGLTFTextureUtility::TransformColorSpace(Pixels, bFromSRGB, bToSRGB);
+	FGLTFTextureUtility::TransformColorSpace(*Pixels, bFromSRGB, bToSRGB);
 
 	const bool bIgnoreAlpha = FGLTFTextureUtility::IsAlphaless(Texture2D->GetPixelFormat());
 	const EGLTFTextureType Type =
 		Texture2D->IsNormalMap() ? EGLTFTextureType::Normalmaps :
 		bIsHDR ? EGLTFTextureType::HDR : EGLTFTextureType::None;
 
-	JsonTexture.Source = Builder.AddImage(Pixels, Size, bIgnoreAlpha, Type, JsonTexture.Name);
+	JsonTexture.Source = Builder.GetOrAddImage(Pixels, Size, bIgnoreAlpha, Type, JsonTexture.Name);
 	JsonTexture.Sampler = Builder.GetOrAddSampler(Texture2D);
 }
 
@@ -100,19 +100,19 @@ void FGLTFTextureCubeTask::Complete()
 		JsonTexture.Encoding = Builder.GetTextureHDREncoding();
 	}
 
-	TArray<FColor> Pixels;
-	if (!FGLTFTextureUtility::ReadPixels(RenderTarget, Pixels, JsonTexture.Encoding))
+	TGLTFSharedArray<FColor> Pixels;
+	if (!FGLTFTextureUtility::ReadPixels(RenderTarget, *Pixels, JsonTexture.Encoding))
 	{
 		Builder.LogWarning(FString::Printf(TEXT("Failed to read pixels (cube face %d) for cubemap texture %s"), CubeFace, *TextureCube->GetName()));
 		return;
 	}
 
-	FGLTFTextureUtility::TransformColorSpace(Pixels, bFromSRGB, bToSRGB);
+	FGLTFTextureUtility::TransformColorSpace(*Pixels, bFromSRGB, bToSRGB);
 
 	const bool bIgnoreAlpha = FGLTFTextureUtility::IsAlphaless(TextureCube->GetPixelFormat());
 	const EGLTFTextureType Type = bIsHDR ? EGLTFTextureType::HDR : EGLTFTextureType::None;
 
-	JsonTexture.Source = Builder.AddImage(Pixels, Size, bIgnoreAlpha, Type, JsonTexture.Name);
+	JsonTexture.Source = Builder.GetOrAddImage(Pixels, Size, bIgnoreAlpha, Type, JsonTexture.Name);
 	JsonTexture.Sampler = Builder.GetOrAddSampler(TextureCube);
 }
 
@@ -140,19 +140,19 @@ void FGLTFTextureRenderTarget2DTask::Complete()
 		JsonTexture.Encoding = Builder.GetTextureHDREncoding();
 	}
 
-	TArray<FColor> Pixels;
-	if (!FGLTFTextureUtility::ReadPixels(RenderTarget2D, Pixels, JsonTexture.Encoding))
+	TGLTFSharedArray<FColor> Pixels;
+	if (!FGLTFTextureUtility::ReadPixels(RenderTarget2D, *Pixels, JsonTexture.Encoding))
 	{
 		Builder.LogWarning(FString::Printf(TEXT("Failed to read pixels for 2D render target %s"), *JsonTexture.Name));
 		return;
 	}
 
-	FGLTFTextureUtility::TransformColorSpace(Pixels, bFromSRGB, bToSRGB);
+	FGLTFTextureUtility::TransformColorSpace(*Pixels, bFromSRGB, bToSRGB);
 
 	const bool bIgnoreAlpha = FGLTFTextureUtility::IsAlphaless(RenderTarget2D->GetFormat());
 	const EGLTFTextureType Type = bIsHDR ? EGLTFTextureType::HDR : EGLTFTextureType::None;
 
-	JsonTexture.Source = Builder.AddImage(Pixels, Size, bIgnoreAlpha, Type, JsonTexture.Name);
+	JsonTexture.Source = Builder.GetOrAddImage(Pixels, Size, bIgnoreAlpha, Type, JsonTexture.Name);
 	JsonTexture.Sampler = Builder.GetOrAddSampler(RenderTarget2D);
 }
 
@@ -193,19 +193,19 @@ void FGLTFTextureRenderTargetCubeTask::Complete()
 		JsonTexture.Encoding = Builder.GetTextureHDREncoding();
 	}
 
-	TArray<FColor> Pixels;
-	if (!FGLTFTextureUtility::ReadPixels(RenderTarget, Pixels, JsonTexture.Encoding))
+	TGLTFSharedArray<FColor> Pixels;
+	if (!FGLTFTextureUtility::ReadPixels(RenderTarget, *Pixels, JsonTexture.Encoding))
 	{
 		Builder.LogWarning(FString::Printf(TEXT("Failed to read pixels (cube face %d) for cubemap render target %s"), CubeFace, *RenderTargetCube->GetName()));
 		return;
 	}
 
-	FGLTFTextureUtility::TransformColorSpace(Pixels, bFromSRGB, bToSRGB);
+	FGLTFTextureUtility::TransformColorSpace(*Pixels, bFromSRGB, bToSRGB);
 
 	const bool bIgnoreAlpha = FGLTFTextureUtility::IsAlphaless(RenderTargetCube->GetFormat());
 	const EGLTFTextureType Type = bIsHDR ? EGLTFTextureType::HDR : EGLTFTextureType::None;
 
-	JsonTexture.Source = Builder.AddImage(Pixels, Size, bIgnoreAlpha, Type, JsonTexture.Name);
+	JsonTexture.Source = Builder.GetOrAddImage(Pixels, Size, bIgnoreAlpha, Type, JsonTexture.Name);
 	JsonTexture.Sampler = Builder.GetOrAddSampler(RenderTargetCube);
 }
 
@@ -245,9 +245,10 @@ void FGLTFTextureLightMapTask::Complete()
 	const EGLTFTextureType Type = EGLTFTextureType::Lightmaps;
 
 	const void* RawData = Source.LockMip(0);
-	JsonTexture.Source = Builder.AddImage(static_cast<const FColor*>(RawData), ByteLength, Size, bIgnoreAlpha, Type, JsonTexture.Name);
+	TGLTFSharedArray<FColor> Pixels(static_cast<const FColor*>(RawData), ByteLength / sizeof(FColor));
 	Source.UnlockMip(0);
 
+	JsonTexture.Source = Builder.GetOrAddImage(Pixels, Size, bIgnoreAlpha, Type, JsonTexture.Name);
 	JsonTexture.Sampler = Builder.GetOrAddSampler(LightMap);
 }
 
