[GLTFExporter] refactor: define an interface for each converter to allow future override

commit b82424fcbaa04974d59c182c28d09c22fe0a0e9b
Date: Sat, 22 Jan 2022 05:51:07 +0100

    refactor: define an interface for each converter to allow future override
    (cherry picked from commit 895794dbc0e13cf6be5dd817e23f3551e57db66c​)

#skipci

diff --git a/Engine/Plugins/Enterprise/GLTFExporter/MigratedFromGit/Source/GLTFExporter/Private/Converters/GLTFAccessorConverters.h b/Engine/Plugins/Enterprise/GLTFExporter/MigratedFromGit/Source/GLTFExporter/Private/Converters/GLTFAccessorConverters.h
--- a/Engine/Plugins/Enterprise/GLTFExporter/MigratedFromGit/Source/GLTFExporter/Private/Converters/GLTFAccessorConverters.h
+++ b/Engine/Plugins/Enterprise/GLTFExporter/MigratedFromGit/Source/GLTFExporter/Private/Converters/GLTFAccessorConverters.h
@@ -7,29 +7,32 @@
 #include "Converters/GLTFBuilderContext.h"
 #include "Converters/GLTFMeshSection.h"
 
-template <typename... InputTypes>
-class TGLTFAccessorConverter : public FGLTFBuilderContext, public TGLTFConverter<FGLTFJsonAccessorIndex, InputTypes...>
+typedef TGLTFConverter<FGLTFJsonAccessorIndex, const FGLTFMeshSection*, const FPositionVertexBuffer*> IGLTFPositionBufferConverter;
+typedef TGLTFConverter<FGLTFJsonAccessorIndex, const FGLTFMeshSection*, const FColorVertexBuffer*> IGLTFColorBufferConverter;
+typedef TGLTFConverter<FGLTFJsonAccessorIndex, const FGLTFMeshSection*, const FStaticMeshVertexBuffer*> IGLTFNormalBufferConverter;
+typedef TGLTFConverter<FGLTFJsonAccessorIndex, const FGLTFMeshSection*, const FStaticMeshVertexBuffer*> IGLTFTangentBufferConverter;
+typedef TGLTFConverter<FGLTFJsonAccessorIndex, const FGLTFMeshSection*, const FStaticMeshVertexBuffer*, uint32> IGLTFUVBufferConverter;
+typedef TGLTFConverter<FGLTFJsonAccessorIndex, const FGLTFMeshSection*, const FSkinWeightVertexBuffer*, uint32> IGLTFBoneIndexBufferConverter;
+typedef TGLTFConverter<FGLTFJsonAccessorIndex, const FGLTFMeshSection*, const FSkinWeightVertexBuffer*, uint32> IGLTFBoneWeightBufferConverter;
+typedef TGLTFConverter<FGLTFJsonAccessorIndex, const FGLTFMeshSection*> IGLTFIndexBufferConverter;
+
+class FGLTFPositionBufferConverter final : public FGLTFBuilderContext, public IGLTFPositionBufferConverter
 {
 	using FGLTFBuilderContext::FGLTFBuilderContext;
-};
-
-class FGLTFPositionBufferConverter final : public TGLTFAccessorConverter<const FGLTFMeshSection*, const FPositionVertexBuffer*>
-{
-	using TGLTFAccessorConverter::TGLTFAccessorConverter;
 
 	virtual FGLTFJsonAccessorIndex Convert(const FGLTFMeshSection* MeshSection, const FPositionVertexBuffer* VertexBuffer) override;
 };
 
-class FGLTFColorBufferConverter final : public TGLTFAccessorConverter<const FGLTFMeshSection*, const FColorVertexBuffer*>
+class FGLTFColorBufferConverter final : public FGLTFBuilderContext, public IGLTFColorBufferConverter
 {
-	using TGLTFAccessorConverter::TGLTFAccessorConverter;
+	using FGLTFBuilderContext::FGLTFBuilderContext;
 
 	virtual FGLTFJsonAccessorIndex Convert(const FGLTFMeshSection* MeshSection, const FColorVertexBuffer* VertexBuffer) override;
 };
 
-class FGLTFNormalBufferConverter final : public TGLTFAccessorConverter<const FGLTFMeshSection*, const FStaticMeshVertexBuffer*>
+class FGLTFNormalBufferConverter final : public FGLTFBuilderContext, public IGLTFNormalBufferConverter
 {
-	using TGLTFAccessorConverter::TGLTFAccessorConverter;
+	using FGLTFBuilderContext::FGLTFBuilderContext;
 
 	virtual FGLTFJsonAccessorIndex Convert(const FGLTFMeshSection* MeshSection, const FStaticMeshVertexBuffer* VertexBuffer) override;
 
@@ -37,9 +40,9 @@ class FGLTFNormalBufferConverter final : public TGLTFAccessorConverter<const FGL
 	FGLTFJsonBufferViewIndex ConvertBufferView(const FGLTFMeshSection* MeshSection, const void* TangentData);
 };
 
-class FGLTFTangentBufferConverter final : public TGLTFAccessorConverter<const FGLTFMeshSection*, const FStaticMeshVertexBuffer*>
+class FGLTFTangentBufferConverter final : public FGLTFBuilderContext, public IGLTFTangentBufferConverter
 {
-	using TGLTFAccessorConverter::TGLTFAccessorConverter;
+	using FGLTFBuilderContext::FGLTFBuilderContext;
 
 	virtual FGLTFJsonAccessorIndex Convert(const FGLTFMeshSection* MeshSection, const FStaticMeshVertexBuffer* VertexBuffer) override;
 
@@ -47,9 +50,9 @@ class FGLTFTangentBufferConverter final : public TGLTFAccessorConverter<const FG
 	FGLTFJsonBufferViewIndex ConvertBufferView(const FGLTFMeshSection* MeshSection, const void* TangentData);
 };
 
-class FGLTFUVBufferConverter final : public TGLTFAccessorConverter<const FGLTFMeshSection*, const FStaticMeshVertexBuffer*, uint32>
+class FGLTFUVBufferConverter final : public FGLTFBuilderContext, public IGLTFUVBufferConverter
 {
-	using TGLTFAccessorConverter::TGLTFAccessorConverter;
+	using FGLTFBuilderContext::FGLTFBuilderContext;
 
 	virtual FGLTFJsonAccessorIndex Convert(const FGLTFMeshSection* MeshSection, const FStaticMeshVertexBuffer* VertexBuffer, uint32 UVIndex) override;
 
@@ -57,9 +60,9 @@ class FGLTFUVBufferConverter final : public TGLTFAccessorConverter<const FGLTFMe
 	FGLTFJsonAccessorIndex Convert(const FGLTFMeshSection* MeshSection, const FStaticMeshVertexBuffer* VertexBuffer, uint32 UVIndex, const uint8* SourceData) const;
 };
 
-class FGLTFBoneIndexBufferConverter final : public TGLTFAccessorConverter<const FGLTFMeshSection*, const FSkinWeightVertexBuffer*, uint32>
+class FGLTFBoneIndexBufferConverter final : public FGLTFBuilderContext, public IGLTFBoneIndexBufferConverter
 {
-	using TGLTFAccessorConverter::TGLTFAccessorConverter;
+	using FGLTFBuilderContext::FGLTFBuilderContext;
 
 	virtual FGLTFJsonAccessorIndex Convert(const FGLTFMeshSection* MeshSection, const FSkinWeightVertexBuffer* VertexBuffer, uint32 InfluenceOffset) override;
 
@@ -73,9 +76,9 @@ class FGLTFBoneIndexBufferConverter final : public TGLTFAccessorConverter<const
 	FGLTFJsonAccessorIndex Convert(const FGLTFMeshSection* MeshSection, const FSkinWeightVertexBuffer* VertexBuffer, uint32 InfluenceOffset, const uint8* SourceData, CallbackType GetVertexInfluenceOffsetCount) const;
 };
 
-class FGLTFBoneWeightBufferConverter final : public TGLTFAccessorConverter<const FGLTFMeshSection*, const FSkinWeightVertexBuffer*, uint32>
+class FGLTFBoneWeightBufferConverter final : public FGLTFBuilderContext, public IGLTFBoneWeightBufferConverter
 {
-	using TGLTFAccessorConverter::TGLTFAccessorConverter;
+	using FGLTFBuilderContext::FGLTFBuilderContext;
 
 	virtual FGLTFJsonAccessorIndex Convert(const FGLTFMeshSection* MeshSection, const FSkinWeightVertexBuffer* VertexBuffer, uint32 InfluenceOffset) override;
 
@@ -86,9 +89,9 @@ class FGLTFBoneWeightBufferConverter final : public TGLTFAccessorConverter<const
 	FGLTFJsonAccessorIndex Convert(const FGLTFMeshSection* MeshSection, const FSkinWeightVertexBuffer* VertexBuffer, uint32 InfluenceOffset, const uint8* SourceData, CallbackType GetVertexInfluenceOffsetCount) const;
 };
 
-class FGLTFIndexBufferConverter final : public TGLTFAccessorConverter<const FGLTFMeshSection*>
+class FGLTFIndexBufferConverter final : public FGLTFBuilderContext, public IGLTFIndexBufferConverter
 {
-	using TGLTFAccessorConverter::TGLTFAccessorConverter;
+	using FGLTFBuilderContext::FGLTFBuilderContext;
 
 	virtual FGLTFJsonAccessorIndex Convert(const FGLTFMeshSection* MeshSection) override;
 
diff --git a/Engine/Plugins/Enterprise/GLTFExporter/MigratedFromGit/Source/GLTFExporter/Private/Converters/GLTFAnimationConverters.h b/Engine/Plugins/Enterprise/GLTFExporter/MigratedFromGit/Source/GLTFExporter/Private/Converters/GLTFAnimationConverters.h
--- a/Engine/Plugins/Enterprise/GLTFExporter/MigratedFromGit/Source/GLTFExporter/Private/Converters/GLTFAnimationConverters.h
+++ b/Engine/Plugins/Enterprise/GLTFExporter/MigratedFromGit/Source/GLTFExporter/Private/Converters/GLTFAnimationConverters.h
@@ -9,36 +9,35 @@
 class ULevelSequence;
 class ALevelSequenceActor;
 
-template <typename... InputTypes>
-class TGLTFAnimationConverter : public FGLTFBuilderContext, public TGLTFConverter<FGLTFJsonAnimationIndex, InputTypes...>
-{
-	using FGLTFBuilderContext::FGLTFBuilderContext;
-};
+typedef TGLTFConverter<FGLTFJsonAnimationIndex, FGLTFJsonNodeIndex, const USkeletalMesh*, const UAnimSequence*> IGLTFAnimationConverter;
+typedef TGLTFConverter<FGLTFJsonAnimationIndex, FGLTFJsonNodeIndex, const USkeletalMeshComponent*> IGLTFAnimationDataConverter;
+typedef TGLTFConverter<FGLTFJsonAnimationIndex, const ULevel*, const ULevelSequence*> IGLTFLevelSequenceConverter;
+typedef TGLTFConverter<FGLTFJsonAnimationIndex, const ALevelSequenceActor*> IGLTFLevelSequenceDataConverter;
 
-class FGLTFAnimationConverter final : public TGLTFAnimationConverter<FGLTFJsonNodeIndex, const USkeletalMesh*, const UAnimSequence*>
+class FGLTFAnimationConverter final : public FGLTFBuilderContext, public IGLTFAnimationConverter
 {
-	using TGLTFAnimationConverter::TGLTFAnimationConverter;
+	using FGLTFBuilderContext::FGLTFBuilderContext;
 
 	virtual FGLTFJsonAnimationIndex Convert(FGLTFJsonNodeIndex RootNode, const USkeletalMesh* SkeletalMesh, const UAnimSequence* AnimSequence) override;
 };
 
-class FGLTFAnimationDataConverter final : public TGLTFAnimationConverter<FGLTFJsonNodeIndex, const USkeletalMeshComponent*>
+class FGLTFAnimationDataConverter final : public FGLTFBuilderContext, public IGLTFAnimationDataConverter
 {
-	using TGLTFAnimationConverter::TGLTFAnimationConverter;
+	using FGLTFBuilderContext::FGLTFBuilderContext;
 
 	virtual FGLTFJsonAnimationIndex Convert(FGLTFJsonNodeIndex RootNode, const USkeletalMeshComponent* SkeletalMeshComponent) override;
 };
 
-class FGLTFLevelSequenceConverter final : public TGLTFAnimationConverter<const ULevel*, const ULevelSequence*>
+class FGLTFLevelSequenceConverter final : public FGLTFBuilderContext, public IGLTFLevelSequenceConverter
 {
-	using TGLTFAnimationConverter::TGLTFAnimationConverter;
+	using FGLTFBuilderContext::FGLTFBuilderContext;
 
 	virtual FGLTFJsonAnimationIndex Convert(const ULevel* Level, const ULevelSequence*) override;
 };
 
-class FGLTFLevelSequenceDataConverter final : public TGLTFAnimationConverter<const ALevelSequenceActor*>
+class FGLTFLevelSequenceDataConverter final : public FGLTFBuilderContext, public IGLTFLevelSequenceDataConverter
 {
-	using TGLTFAnimationConverter::TGLTFAnimationConverter;
+	using FGLTFBuilderContext::FGLTFBuilderContext;
 
 	virtual FGLTFJsonAnimationIndex Convert(const ALevelSequenceActor* LevelSequenceActor) override;
 };
diff --git a/Engine/Plugins/Enterprise/GLTFExporter/MigratedFromGit/Source/GLTFExporter/Private/Converters/GLTFBackdropConverters.h b/Engine/Plugins/Enterprise/GLTFExporter/MigratedFromGit/Source/GLTFExporter/Private/Converters/GLTFBackdropConverters.h
--- a/Engine/Plugins/Enterprise/GLTFExporter/MigratedFromGit/Source/GLTFExporter/Private/Converters/GLTFBackdropConverters.h
+++ b/Engine/Plugins/Enterprise/GLTFExporter/MigratedFromGit/Source/GLTFExporter/Private/Converters/GLTFBackdropConverters.h
@@ -6,7 +6,10 @@
 #include "Converters/GLTFConverter.h"
 #include "Converters/GLTFBuilderContext.h"
 
-class FGLTFBackdropConverter final : public FGLTFBuilderContext, public TGLTFConverter<FGLTFJsonBackdropIndex, const AActor*>
+typedef TGLTFConverter<FGLTFJsonBackdropIndex, const AActor*> IGLTFBackdropConverter;
+
+
+class FGLTFBackdropConverter final : public FGLTFBuilderContext, public IGLTFBackdropConverter
 {
 	using FGLTFBuilderContext::FGLTFBuilderContext;
 
diff --git a/Engine/Plugins/Enterprise/GLTFExporter/MigratedFromGit/Source/GLTFExporter/Private/Converters/GLTFCameraConverters.h b/Engine/Plugins/Enterprise/GLTFExporter/MigratedFromGit/Source/GLTFExporter/Private/Converters/GLTFCameraConverters.h
--- a/Engine/Plugins/Enterprise/GLTFExporter/MigratedFromGit/Source/GLTFExporter/Private/Converters/GLTFCameraConverters.h
+++ b/Engine/Plugins/Enterprise/GLTFExporter/MigratedFromGit/Source/GLTFExporter/Private/Converters/GLTFCameraConverters.h
@@ -8,7 +8,9 @@
 
 class UCameraComponent;
 
-class FGLTFCameraConverter final : public FGLTFBuilderContext, public TGLTFConverter<FGLTFJsonCameraIndex, const UCameraComponent*>
+typedef TGLTFConverter<FGLTFJsonCameraIndex, const UCameraComponent*> IGLTFCameraConverter;
+
+class FGLTFCameraConverter final : public FGLTFBuilderContext, public IGLTFCameraConverter
 {
 	using FGLTFBuilderContext::FGLTFBuilderContext;
 
diff --git a/Engine/Plugins/Enterprise/GLTFExporter/MigratedFromGit/Source/GLTFExporter/Private/Converters/GLTFEpicLevelVariantSetsConverters.h b/Engine/Plugins/Enterprise/GLTFExporter/MigratedFromGit/Source/GLTFExporter/Private/Converters/GLTFEpicLevelVariantSetsConverters.h
--- a/Engine/Plugins/Enterprise/GLTFExporter/MigratedFromGit/Source/GLTFExporter/Private/Converters/GLTFEpicLevelVariantSetsConverters.h
+++ b/Engine/Plugins/Enterprise/GLTFExporter/MigratedFromGit/Source/GLTFExporter/Private/Converters/GLTFEpicLevelVariantSetsConverters.h
@@ -14,7 +14,9 @@ class UVariantObjectBinding;
 class UPropertyValue;
 class UPropertyValueMaterial;
 
-class FGLTFEpicLevelVariantSetsConverter final : public FGLTFBuilderContext, public TGLTFConverter<FGLTFJsonEpicLevelVariantSetsIndex, const ULevelVariantSets*>
+typedef TGLTFConverter<FGLTFJsonEpicLevelVariantSetsIndex, const ULevelVariantSets*> IGLTFEpicLevelVariantSetsConverter;
+
+class FGLTFEpicLevelVariantSetsConverter final : public FGLTFBuilderContext, public IGLTFEpicLevelVariantSetsConverter
 {
 	using FGLTFBuilderContext::FGLTFBuilderContext;
 
diff --git a/Engine/Plugins/Enterprise/GLTFExporter/MigratedFromGit/Source/GLTFExporter/Private/Converters/GLTFHotspotConverters.h b/Engine/Plugins/Enterprise/GLTFExporter/MigratedFromGit/Source/GLTFExporter/Private/Converters/GLTFHotspotConverters.h
--- a/Engine/Plugins/Enterprise/GLTFExporter/MigratedFromGit/Source/GLTFExporter/Private/Converters/GLTFHotspotConverters.h
+++ b/Engine/Plugins/Enterprise/GLTFExporter/MigratedFromGit/Source/GLTFExporter/Private/Converters/GLTFHotspotConverters.h
@@ -7,7 +7,9 @@
 #include "Converters/GLTFBuilderContext.h"
 #include "Actors/GLTFHotspotActor.h"
 
-class FGLTFHotspotConverter final : public FGLTFBuilderContext, public TGLTFConverter<FGLTFJsonHotspotIndex, const AGLTFHotspotActor*>
+typedef TGLTFConverter<FGLTFJsonHotspotIndex, const AGLTFHotspotActor*> IGLTFHotspotConverter;
+
+class FGLTFHotspotConverter final : public FGLTFBuilderContext, public IGLTFHotspotConverter
 {
 	using FGLTFBuilderContext::FGLTFBuilderContext;
 
diff --git a/Engine/Plugins/Enterprise/GLTFExporter/MigratedFromGit/Source/GLTFExporter/Private/Converters/GLTFImageConverters.h b/Engine/Plugins/Enterprise/GLTFExporter/MigratedFromGit/Source/GLTFExporter/Private/Converters/GLTFImageConverters.h
--- a/Engine/Plugins/Enterprise/GLTFExporter/MigratedFromGit/Source/GLTFExporter/Private/Converters/GLTFImageConverters.h
+++ b/Engine/Plugins/Enterprise/GLTFExporter/MigratedFromGit/Source/GLTFExporter/Private/Converters/GLTFImageConverters.h
@@ -10,7 +10,9 @@
 #include "Converters/GLTFSharedArray.h"
 #include "GLTFExportOptions.h"
 
-class FGLTFImageConverter final : public FGLTFBuilderContext, public TGLTFConverter<FGLTFJsonImageIndex, TGLTFSuperfluous<FString>, EGLTFTextureType, bool, FIntPoint, TGLTFSharedArray<FColor>>
+typedef TGLTFConverter<FGLTFJsonImageIndex, TGLTFSuperfluous<FString>, EGLTFTextureType, bool, FIntPoint, TGLTFSharedArray<FColor>> IGLTFImageConverter;
+
+class FGLTFImageConverter final : public FGLTFBuilderContext, public IGLTFImageConverter
 {
 	using FGLTFBuilderContext::FGLTFBuilderContext;
 
diff --git a/Engine/Plugins/Enterprise/GLTFExporter/MigratedFromGit/Source/GLTFExporter/Private/Converters/GLTFKhrVariantConverters.h b/Engine/Plugins/Enterprise/GLTFExporter/MigratedFromGit/Source/GLTFExporter/Private/Converters/GLTFKhrVariantConverters.h
--- a/Engine/Plugins/Enterprise/GLTFExporter/MigratedFromGit/Source/GLTFExporter/Private/Converters/GLTFKhrVariantConverters.h
+++ b/Engine/Plugins/Enterprise/GLTFExporter/MigratedFromGit/Source/GLTFExporter/Private/Converters/GLTFKhrVariantConverters.h
@@ -10,7 +10,9 @@
 class UVariant;
 class UPropertyValueMaterial;
 
-class FGLTFKhrMaterialVariantConverter final : public FGLTFBuilderContext, public TGLTFConverter<FGLTFJsonKhrMaterialVariantIndex, const UVariant*>
+typedef TGLTFConverter<FGLTFJsonKhrMaterialVariantIndex, const UVariant*> IGLTFKhrMaterialVariantConverter;
+
+class FGLTFKhrMaterialVariantConverter final : public FGLTFBuilderContext, public IGLTFKhrMaterialVariantConverter
 {
 	using FGLTFBuilderContext::FGLTFBuilderContext;
 
diff --git a/Engine/Plugins/Enterprise/GLTFExporter/MigratedFromGit/Source/GLTFExporter/Private/Converters/GLTFLightConverters.h b/Engine/Plugins/Enterprise/GLTFExporter/MigratedFromGit/Source/GLTFExporter/Private/Converters/GLTFLightConverters.h
--- a/Engine/Plugins/Enterprise/GLTFExporter/MigratedFromGit/Source/GLTFExporter/Private/Converters/GLTFLightConverters.h
+++ b/Engine/Plugins/Enterprise/GLTFExporter/MigratedFromGit/Source/GLTFExporter/Private/Converters/GLTFLightConverters.h
@@ -6,7 +6,9 @@
 #include "Converters/GLTFConverter.h"
 #include "Converters/GLTFBuilderContext.h"
 
-class FGLTFLightConverter final : public FGLTFBuilderContext, public TGLTFConverter<FGLTFJsonLightIndex, const ULightComponent*>
+typedef TGLTFConverter<FGLTFJsonLightIndex, const ULightComponent*> IGLTFLightConverter;
+
+class FGLTFLightConverter final : public FGLTFBuilderContext, public IGLTFLightConverter
 {
 	using FGLTFBuilderContext::FGLTFBuilderContext;
 
diff --git a/Engine/Plugins/Enterprise/GLTFExporter/MigratedFromGit/Source/GLTFExporter/Private/Converters/GLTFLightMapConverters.h b/Engine/Plugins/Enterprise/GLTFExporter/MigratedFromGit/Source/GLTFExporter/Private/Converters/GLTFLightMapConverters.h
--- a/Engine/Plugins/Enterprise/GLTFExporter/MigratedFromGit/Source/GLTFExporter/Private/Converters/GLTFLightMapConverters.h
+++ b/Engine/Plugins/Enterprise/GLTFExporter/MigratedFromGit/Source/GLTFExporter/Private/Converters/GLTFLightMapConverters.h
@@ -6,7 +6,9 @@
 #include "Converters/GLTFConverter.h"
 #include "Converters/GLTFBuilderContext.h"
 
-class FGLTFLightMapConverter final : public FGLTFBuilderContext, public TGLTFConverter<FGLTFJsonLightMapIndex, const UStaticMeshComponent*>
+typedef TGLTFConverter<FGLTFJsonLightMapIndex, const UStaticMeshComponent*> IGLTFLightMapConverter;
+
+class FGLTFLightMapConverter final : public FGLTFBuilderContext, public IGLTFLightMapConverter
 {
 	using FGLTFBuilderContext::FGLTFBuilderContext;
 
diff --git a/Engine/Plugins/Enterprise/GLTFExporter/MigratedFromGit/Source/GLTFExporter/Private/Converters/GLTFMaterialConverters.h b/Engine/Plugins/Enterprise/GLTFExporter/MigratedFromGit/Source/GLTFExporter/Private/Converters/GLTFMaterialConverters.h
--- a/Engine/Plugins/Enterprise/GLTFExporter/MigratedFromGit/Source/GLTFExporter/Private/Converters/GLTFMaterialConverters.h
+++ b/Engine/Plugins/Enterprise/GLTFExporter/MigratedFromGit/Source/GLTFExporter/Private/Converters/GLTFMaterialConverters.h
@@ -10,7 +10,9 @@
 #include "Converters/GLTFUVOverlapChecker.h"
 #include "Converters/GLTFUVDegenerateChecker.h"
 
-class FGLTFMaterialConverter final : public FGLTFBuilderContext, public TGLTFConverter<FGLTFJsonMaterialIndex, const UMaterialInterface*, const FGLTFMeshData*, FGLTFIndexArray>
+typedef TGLTFConverter<FGLTFJsonMaterialIndex, const UMaterialInterface*, const FGLTFMeshData*, FGLTFIndexArray> IGLTFMaterialConverter;
+
+class FGLTFMaterialConverter final : public FGLTFBuilderContext, public IGLTFMaterialConverter
 {
 	using FGLTFBuilderContext::FGLTFBuilderContext;
 
diff --git a/Engine/Plugins/Enterprise/GLTFExporter/MigratedFromGit/Source/GLTFExporter/Private/Converters/GLTFMeshConverters.h b/Engine/Plugins/Enterprise/GLTFExporter/MigratedFromGit/Source/GLTFExporter/Private/Converters/GLTFMeshConverters.h
--- a/Engine/Plugins/Enterprise/GLTFExporter/MigratedFromGit/Source/GLTFExporter/Private/Converters/GLTFMeshConverters.h
+++ b/Engine/Plugins/Enterprise/GLTFExporter/MigratedFromGit/Source/GLTFExporter/Private/Converters/GLTFMeshConverters.h
@@ -9,15 +9,12 @@
 #include "Converters/GLTFMeshDataConverters.h"
 #include "Converters/GLTFMaterialArray.h"
 
-template <typename... InputTypes>
-class TGLTFMeshConverter : public FGLTFBuilderContext, public TGLTFConverter<FGLTFJsonMeshIndex, InputTypes...>
-{
-	using FGLTFBuilderContext::FGLTFBuilderContext;
-};
+typedef TGLTFConverter<FGLTFJsonMeshIndex, const UStaticMesh*, const UStaticMeshComponent*, FGLTFMaterialArray, int32> IGLTFStaticMeshConverter;
+typedef TGLTFConverter<FGLTFJsonMeshIndex, const USkeletalMesh*, const USkeletalMeshComponent*, FGLTFMaterialArray, int32> IGLTFSkeletalMeshConverter;
 
-class FGLTFStaticMeshConverter final : public TGLTFMeshConverter<const UStaticMesh*, const UStaticMeshComponent*, FGLTFMaterialArray, int32>
+class FGLTFStaticMeshConverter final : public FGLTFBuilderContext, public IGLTFStaticMeshConverter
 {
-	using TGLTFMeshConverter::TGLTFMeshConverter;
+	using FGLTFBuilderContext::FGLTFBuilderContext;
 
 	virtual void Sanitize(const UStaticMesh*& StaticMesh, const UStaticMeshComponent*& StaticMeshComponent, FGLTFMaterialArray& Materials, int32& LODIndex) override;
 
@@ -26,9 +23,9 @@ class FGLTFStaticMeshConverter final : public TGLTFMeshConverter<const UStaticMe
 	FGLTFStaticMeshSectionConverter MeshSectionConverter;
 };
 
-class FGLTFSkeletalMeshConverter final : public TGLTFMeshConverter<const USkeletalMesh*, const USkeletalMeshComponent*, FGLTFMaterialArray, int32>
+class FGLTFSkeletalMeshConverter final : public FGLTFBuilderContext, public IGLTFSkeletalMeshConverter
 {
-	using TGLTFMeshConverter::TGLTFMeshConverter;
+	using FGLTFBuilderContext::FGLTFBuilderContext;
 
 	virtual void Sanitize(const USkeletalMesh*& SkeletalMesh, const USkeletalMeshComponent*& SkeletalMeshComponent, FGLTFMaterialArray& Materials, int32& LODIndex) override;
 
diff --git a/Engine/Plugins/Enterprise/GLTFExporter/MigratedFromGit/Source/GLTFExporter/Private/Converters/GLTFNodeConverters.h b/Engine/Plugins/Enterprise/GLTFExporter/MigratedFromGit/Source/GLTFExporter/Private/Converters/GLTFNodeConverters.h
--- a/Engine/Plugins/Enterprise/GLTFExporter/MigratedFromGit/Source/GLTFExporter/Private/Converters/GLTFNodeConverters.h
+++ b/Engine/Plugins/Enterprise/GLTFExporter/MigratedFromGit/Source/GLTFExporter/Private/Converters/GLTFNodeConverters.h
@@ -6,50 +6,51 @@
 #include "Converters/GLTFConverter.h"
 #include "Converters/GLTFBuilderContext.h"
 
-template <typename... InputTypes>
-class TGLTFNodeConverter : public FGLTFBuilderContext, public TGLTFConverter<FGLTFJsonNodeIndex, InputTypes...>
+typedef TGLTFConverter<FGLTFJsonNodeIndex, const AActor*> IGLTFActorConverter;
+typedef TGLTFConverter<FGLTFJsonNodeIndex, const USceneComponent*> IGLTFComponentConverter;
+typedef TGLTFConverter<FGLTFJsonNodeIndex, const USceneComponent*, FName> IGLTFComponentSocketConverter;
+typedef TGLTFConverter<FGLTFJsonNodeIndex, FGLTFJsonNodeIndex, const UStaticMesh*, FName> IGLTFStaticSocketConverter;
+typedef TGLTFConverter<FGLTFJsonNodeIndex, FGLTFJsonNodeIndex, const USkeletalMesh*, FName> IGLTFSkeletalSocketConverter;
+typedef TGLTFConverter<FGLTFJsonNodeIndex, FGLTFJsonNodeIndex, const USkeletalMesh*, int32> IGLTFSkeletalBoneConverter;
+
+class FGLTFActorConverter final : public FGLTFBuilderContext, public IGLTFActorConverter
 {
 	using FGLTFBuilderContext::FGLTFBuilderContext;
-};
-
-class FGLTFActorConverter final : public TGLTFNodeConverter<const AActor*>
-{
-	using TGLTFNodeConverter::TGLTFNodeConverter;
 
 	virtual FGLTFJsonNodeIndex Convert(const AActor* Actor) override;
 };
 
-class FGLTFComponentConverter final : public TGLTFNodeConverter<const USceneComponent*>
+class FGLTFComponentConverter final : public FGLTFBuilderContext, public IGLTFComponentConverter
 {
-	using TGLTFNodeConverter::TGLTFNodeConverter;
+	using FGLTFBuilderContext::FGLTFBuilderContext;
 
 	virtual FGLTFJsonNodeIndex Convert(const USceneComponent* SceneComponent) override;
 };
 
-class FGLTFComponentSocketConverter final : public TGLTFNodeConverter<const USceneComponent*, FName>
+class FGLTFComponentSocketConverter final : public FGLTFBuilderContext, public IGLTFComponentSocketConverter
 {
-	using TGLTFNodeConverter::TGLTFNodeConverter;
+	using FGLTFBuilderContext::FGLTFBuilderContext;
 
 	virtual FGLTFJsonNodeIndex Convert(const USceneComponent* SceneComponent, FName SocketName) override;
 };
 
-class FGLTFStaticSocketConverter final : public TGLTFNodeConverter<FGLTFJsonNodeIndex, const UStaticMesh*, FName>
+class FGLTFStaticSocketConverter final : public FGLTFBuilderContext, public IGLTFStaticSocketConverter
 {
-	using TGLTFNodeConverter::TGLTFNodeConverter;
+	using FGLTFBuilderContext::FGLTFBuilderContext;
 
 	virtual FGLTFJsonNodeIndex Convert(FGLTFJsonNodeIndex RootNode, const UStaticMesh* StaticMesh, FName SocketName) override;
 };
 
-class FGLTFSkeletalSocketConverter final : public TGLTFNodeConverter<FGLTFJsonNodeIndex, const USkeletalMesh*, FName>
+class FGLTFSkeletalSocketConverter final : public FGLTFBuilderContext, public IGLTFSkeletalSocketConverter
 {
-	using TGLTFNodeConverter::TGLTFNodeConverter;
+	using FGLTFBuilderContext::FGLTFBuilderContext;
 
 	virtual FGLTFJsonNodeIndex Convert(FGLTFJsonNodeIndex RootNode, const USkeletalMesh* SkeletalMesh, FName SocketName) override;
 };
 
-class FGLTFSkeletalBoneConverter final : public TGLTFNodeConverter<FGLTFJsonNodeIndex, const USkeletalMesh*, int32>
+class FGLTFSkeletalBoneConverter final : public FGLTFBuilderContext, public IGLTFSkeletalBoneConverter
 {
-	using TGLTFNodeConverter::TGLTFNodeConverter;
+	using FGLTFBuilderContext::FGLTFBuilderContext;
 
 	virtual FGLTFJsonNodeIndex Convert(FGLTFJsonNodeIndex RootNode, const USkeletalMesh* SkeletalMesh, int32 BoneIndex) override;
 };
diff --git a/Engine/Plugins/Enterprise/GLTFExporter/MigratedFromGit/Source/GLTFExporter/Private/Converters/GLTFSamplerConverters.h b/Engine/Plugins/Enterprise/GLTFExporter/MigratedFromGit/Source/GLTFExporter/Private/Converters/GLTFSamplerConverters.h
--- a/Engine/Plugins/Enterprise/GLTFExporter/MigratedFromGit/Source/GLTFExporter/Private/Converters/GLTFSamplerConverters.h
+++ b/Engine/Plugins/Enterprise/GLTFExporter/MigratedFromGit/Source/GLTFExporter/Private/Converters/GLTFSamplerConverters.h
@@ -6,7 +6,9 @@
 #include "Converters/GLTFConverter.h"
 #include "Converters/GLTFBuilderContext.h"
 
-class FGLTFSamplerConverter final : public FGLTFBuilderContext, public TGLTFConverter<FGLTFJsonSamplerIndex, const UTexture*>
+typedef TGLTFConverter<FGLTFJsonSamplerIndex, const UTexture*> IGLTFSamplerConverter;
+
+class FGLTFSamplerConverter final : public FGLTFBuilderContext, public IGLTFSamplerConverter
 {
 	using FGLTFBuilderContext::FGLTFBuilderContext;
 
diff --git a/Engine/Plugins/Enterprise/GLTFExporter/MigratedFromGit/Source/GLTFExporter/Private/Converters/GLTFSceneConverters.h b/Engine/Plugins/Enterprise/GLTFExporter/MigratedFromGit/Source/GLTFExporter/Private/Converters/GLTFSceneConverters.h
--- a/Engine/Plugins/Enterprise/GLTFExporter/MigratedFromGit/Source/GLTFExporter/Private/Converters/GLTFSceneConverters.h
+++ b/Engine/Plugins/Enterprise/GLTFExporter/MigratedFromGit/Source/GLTFExporter/Private/Converters/GLTFSceneConverters.h
@@ -6,7 +6,9 @@
 #include "Converters/GLTFConverter.h"
 #include "Converters/GLTFBuilderContext.h"
 
-class FGLTFSceneConverter final : public FGLTFBuilderContext, public TGLTFConverter<FGLTFJsonSceneIndex, const UWorld*>
+typedef TGLTFConverter<FGLTFJsonSceneIndex, const UWorld*> IGLTFSceneConverter;
+
+class FGLTFSceneConverter final : public FGLTFBuilderContext, public IGLTFSceneConverter
 {
 	using FGLTFBuilderContext::FGLTFBuilderContext;
 
diff --git a/Engine/Plugins/Enterprise/GLTFExporter/MigratedFromGit/Source/GLTFExporter/Private/Converters/GLTFSkinConverters.h b/Engine/Plugins/Enterprise/GLTFExporter/MigratedFromGit/Source/GLTFExporter/Private/Converters/GLTFSkinConverters.h
--- a/Engine/Plugins/Enterprise/GLTFExporter/MigratedFromGit/Source/GLTFExporter/Private/Converters/GLTFSkinConverters.h
+++ b/Engine/Plugins/Enterprise/GLTFExporter/MigratedFromGit/Source/GLTFExporter/Private/Converters/GLTFSkinConverters.h
@@ -6,7 +6,9 @@
 #include "Converters/GLTFConverter.h"
 #include "Converters/GLTFBuilderContext.h"
 
-class FGLTFSkinConverter final : public FGLTFBuilderContext, public TGLTFConverter<FGLTFJsonSkinIndex, FGLTFJsonNodeIndex, const USkeletalMesh*>
+typedef TGLTFConverter<FGLTFJsonSkinIndex, FGLTFJsonNodeIndex, const USkeletalMesh*> IGLTFSkinConverter;
+
+class FGLTFSkinConverter final : public FGLTFBuilderContext, public IGLTFSkinConverter
 {
 	using FGLTFBuilderContext::FGLTFBuilderContext;
 
diff --git a/Engine/Plugins/Enterprise/GLTFExporter/MigratedFromGit/Source/GLTFExporter/Private/Converters/GLTFSkySphereConverters.h b/Engine/Plugins/Enterprise/GLTFExporter/MigratedFromGit/Source/GLTFExporter/Private/Converters/GLTFSkySphereConverters.h
--- a/Engine/Plugins/Enterprise/GLTFExporter/MigratedFromGit/Source/GLTFExporter/Private/Converters/GLTFSkySphereConverters.h
+++ b/Engine/Plugins/Enterprise/GLTFExporter/MigratedFromGit/Source/GLTFExporter/Private/Converters/GLTFSkySphereConverters.h
@@ -9,7 +9,9 @@
 
 struct FGLTFJsonSkySphereColorCurve;
 
-class FGLTFSkySphereConverter final : public FGLTFBuilderContext, public TGLTFConverter<FGLTFJsonSkySphereIndex, const AActor*>
+typedef TGLTFConverter<FGLTFJsonSkySphereIndex, const AActor*> IGLTFSkySphereConverter;
+
+class FGLTFSkySphereConverter final : public FGLTFBuilderContext, public IGLTFSkySphereConverter
 {
 	enum class ESkySphereTextureParameter
 	{
diff --git a/Engine/Plugins/Enterprise/GLTFExporter/MigratedFromGit/Source/GLTFExporter/Private/Converters/GLTFTextureConverters.h b/Engine/Plugins/Enterprise/GLTFExporter/MigratedFromGit/Source/GLTFExporter/Private/Converters/GLTFTextureConverters.h
--- a/Engine/Plugins/Enterprise/GLTFExporter/MigratedFromGit/Source/GLTFExporter/Private/Converters/GLTFTextureConverters.h
+++ b/Engine/Plugins/Enterprise/GLTFExporter/MigratedFromGit/Source/GLTFExporter/Private/Converters/GLTFTextureConverters.h
@@ -6,51 +6,51 @@
 #include "Converters/GLTFConverter.h"
 #include "Converters/GLTFBuilderContext.h"
 
-template <typename... InputTypes>
-class TGLTFTextureConverter : public FGLTFBuilderContext, public TGLTFConverter<FGLTFJsonTextureIndex, InputTypes...>
-{
-	using FGLTFBuilderContext::FGLTFBuilderContext;
-};
+typedef TGLTFConverter<FGLTFJsonTextureIndex, const UTexture2D*, bool> IGLTFTexture2DConverter;
+typedef TGLTFConverter<FGLTFJsonTextureIndex, const UTextureCube*, ECubeFace, bool> IGLTFTextureCubeConverter;
+typedef TGLTFConverter<FGLTFJsonTextureIndex, const UTextureRenderTarget2D*, bool> IGLTFTextureRenderTarget2DConverter;
+typedef TGLTFConverter<FGLTFJsonTextureIndex, const UTextureRenderTargetCube*, ECubeFace, bool> IGLTFTextureRenderTargetCubeConverter;
+typedef TGLTFConverter<FGLTFJsonTextureIndex, const ULightMapTexture2D*> IGLTFTextureLightMapConverter;
 
-class FGLTFTexture2DConverter final : public TGLTFTextureConverter<const UTexture2D*, bool>
+class FGLTFTexture2DConverter final : public FGLTFBuilderContext, public IGLTFTexture2DConverter
 {
-	using TGLTFTextureConverter::TGLTFTextureConverter;
+	using FGLTFBuilderContext::FGLTFBuilderContext;
 
 	virtual void Sanitize(const UTexture2D*& Texture2D, bool& bToSRGB) override;
 
 	virtual FGLTFJsonTextureIndex Convert(const UTexture2D* Texture2D, bool bToSRGB) override;
 };
 
-class FGLTFTextureCubeConverter final : public TGLTFTextureConverter<const UTextureCube*, ECubeFace, bool>
+class FGLTFTextureCubeConverter final : public FGLTFBuilderContext, public IGLTFTextureCubeConverter
 {
-	using TGLTFTextureConverter::TGLTFTextureConverter;
+	using FGLTFBuilderContext::FGLTFBuilderContext;
 
 	virtual void Sanitize(const UTextureCube*& TextureCube, ECubeFace& CubeFace, bool& bToSRGB) override;
 
 	virtual FGLTFJsonTextureIndex Convert(const UTextureCube* TextureCube, ECubeFace CubeFace, bool bToSRGB) override;
 };
 
-class FGLTFTextureRenderTarget2DConverter final : public TGLTFTextureConverter<const UTextureRenderTarget2D*, bool>
+class FGLTFTextureRenderTarget2DConverter final : public FGLTFBuilderContext, public IGLTFTextureRenderTarget2DConverter
 {
-	using TGLTFTextureConverter::TGLTFTextureConverter;
+	using FGLTFBuilderContext::FGLTFBuilderContext;
 
 	virtual void Sanitize(const UTextureRenderTarget2D*& RenderTarget2D, bool& bToSRGB) override;
 
 	virtual FGLTFJsonTextureIndex Convert(const UTextureRenderTarget2D* RenderTarget2D, bool bToSRGB) override;
 };
 
-class FGLTFTextureRenderTargetCubeConverter final : public TGLTFTextureConverter<const UTextureRenderTargetCube*, ECubeFace, bool>
+class FGLTFTextureRenderTargetCubeConverter final : public FGLTFBuilderContext, public IGLTFTextureRenderTargetCubeConverter
 {
-	using TGLTFTextureConverter::TGLTFTextureConverter;
+	using FGLTFBuilderContext::FGLTFBuilderContext;
 
 	virtual void Sanitize(const UTextureRenderTargetCube*& RenderTargetCube, ECubeFace& CubeFace, bool& bToSRGB) override;
 
 	virtual FGLTFJsonTextureIndex Convert(const UTextureRenderTargetCube* RenderTargetCube, ECubeFace CubeFace, bool bToSRGB) override;
 };
 
-class FGLTFTextureLightMapConverter final : public TGLTFTextureConverter<const ULightMapTexture2D*>
+class FGLTFTextureLightMapConverter final : public FGLTFBuilderContext, public IGLTFTextureLightMapConverter
 {
-	using TGLTFTextureConverter::TGLTFTextureConverter;
+	using FGLTFBuilderContext::FGLTFBuilderContext;
 
 	virtual FGLTFJsonTextureIndex Convert(const ULightMapTexture2D* LightMap) override;
 };
