[GLTFExporter] refactor: add ability to override converters in builder

commit 38106e936c1fd31998160e0c5f068ef236b24072
Date: Sat, 22 Jan 2022 06:34:03 +0100

    refactor: add ability to override converters in builder
    (cherry picked from commit 3305003b985f0f899b29b8f421aa7707356e59e9​)

#skipci

diff --git a/Engine/Plugins/Enterprise/GLTFExporter/MigratedFromGit/Source/GLTFExporter/Private/Builders/GLTFConvertBuilder.cpp b/Engine/Plugins/Enterprise/GLTFExporter/MigratedFromGit/Source/GLTFExporter/Private/Builders/GLTFConvertBuilder.cpp
--- a/Engine/Plugins/Enterprise/GLTFExporter/MigratedFromGit/Source/GLTFExporter/Private/Builders/GLTFConvertBuilder.cpp
+++ b/Engine/Plugins/Enterprise/GLTFExporter/MigratedFromGit/Source/GLTFExporter/Private/Builders/GLTFConvertBuilder.cpp
@@ -31,7 +31,7 @@ FGLTFJsonAccessorIndex FGLTFConvertBuilder::GetOrAddPositionAccessor(const FGLTF
 		return FGLTFJsonAccessorIndex(INDEX_NONE);
 	}
 
-	return PositionBufferConverter.GetOrAdd(MeshSection, VertexBuffer);
+	return PositionBufferConverter->GetOrAdd(MeshSection, VertexBuffer);
 }
 
 FGLTFJsonAccessorIndex FGLTFConvertBuilder::GetOrAddColorAccessor(const FGLTFMeshSection* MeshSection, const FColorVertexBuffer* VertexBuffer)
@@ -41,7 +41,7 @@ FGLTFJsonAccessorIndex FGLTFConvertBuilder::GetOrAddColorAccessor(const FGLTFMes
 		return FGLTFJsonAccessorIndex(INDEX_NONE);
 	}
 
-	return ColorBufferConverter.GetOrAdd(MeshSection, VertexBuffer);
+	return ColorBufferConverter->GetOrAdd(MeshSection, VertexBuffer);
 }
 
 FGLTFJsonAccessorIndex FGLTFConvertBuilder::GetOrAddNormalAccessor(const FGLTFMeshSection* MeshSection, const FStaticMeshVertexBuffer* VertexBuffer)
@@ -51,7 +51,7 @@ FGLTFJsonAccessorIndex FGLTFConvertBuilder::GetOrAddNormalAccessor(const FGLTFMe
 		return FGLTFJsonAccessorIndex(INDEX_NONE);
 	}
 
-	return NormalBufferConverter.GetOrAdd(MeshSection, VertexBuffer);
+	return NormalBufferConverter->GetOrAdd(MeshSection, VertexBuffer);
 }
 
 FGLTFJsonAccessorIndex FGLTFConvertBuilder::GetOrAddTangentAccessor(const FGLTFMeshSection* MeshSection, const FStaticMeshVertexBuffer* VertexBuffer)
@@ -61,7 +61,7 @@ FGLTFJsonAccessorIndex FGLTFConvertBuilder::GetOrAddTangentAccessor(const FGLTFM
 		return FGLTFJsonAccessorIndex(INDEX_NONE);
 	}
 
-	return TangentBufferConverter.GetOrAdd(MeshSection, VertexBuffer);
+	return TangentBufferConverter->GetOrAdd(MeshSection, VertexBuffer);
 }
 
 FGLTFJsonAccessorIndex FGLTFConvertBuilder::GetOrAddUVAccessor(const FGLTFMeshSection* MeshSection, const FStaticMeshVertexBuffer* VertexBuffer, int32 UVIndex)
@@ -71,7 +71,7 @@ FGLTFJsonAccessorIndex FGLTFConvertBuilder::GetOrAddUVAccessor(const FGLTFMeshSe
 		return FGLTFJsonAccessorIndex(INDEX_NONE);
 	}
 
-	return UVBufferConverter.GetOrAdd(MeshSection, VertexBuffer, UVIndex);
+	return UVBufferConverter->GetOrAdd(MeshSection, VertexBuffer, UVIndex);
 }
 
 FGLTFJsonAccessorIndex FGLTFConvertBuilder::GetOrAddJointAccessor(const FGLTFMeshSection* MeshSection, const FSkinWeightVertexBuffer* VertexBuffer, int32 InfluenceOffset)
@@ -81,7 +81,7 @@ FGLTFJsonAccessorIndex FGLTFConvertBuilder::GetOrAddJointAccessor(const FGLTFMes
 		return FGLTFJsonAccessorIndex(INDEX_NONE);
 	}
 
-	return BoneIndexBufferConverter.GetOrAdd(MeshSection, VertexBuffer, InfluenceOffset);
+	return BoneIndexBufferConverter->GetOrAdd(MeshSection, VertexBuffer, InfluenceOffset);
 }
 
 FGLTFJsonAccessorIndex FGLTFConvertBuilder::GetOrAddWeightAccessor(const FGLTFMeshSection* MeshSection, const FSkinWeightVertexBuffer* VertexBuffer, int32 InfluenceOffset)
@@ -91,7 +91,7 @@ FGLTFJsonAccessorIndex FGLTFConvertBuilder::GetOrAddWeightAccessor(const FGLTFMe
 		return FGLTFJsonAccessorIndex(INDEX_NONE);
 	}
 
-	return BoneWeightBufferConverter.GetOrAdd(MeshSection, VertexBuffer, InfluenceOffset);
+	return BoneWeightBufferConverter->GetOrAdd(MeshSection, VertexBuffer, InfluenceOffset);
 }
 
 FGLTFJsonAccessorIndex FGLTFConvertBuilder::GetOrAddIndexAccessor(const FGLTFMeshSection* MeshSection)
@@ -101,7 +101,7 @@ FGLTFJsonAccessorIndex FGLTFConvertBuilder::GetOrAddIndexAccessor(const FGLTFMes
 		return FGLTFJsonAccessorIndex(INDEX_NONE);
 	}
 
-	return IndexBufferConverter.GetOrAdd(MeshSection);
+	return IndexBufferConverter->GetOrAdd(MeshSection);
 }
 
 FGLTFJsonMeshIndex FGLTFConvertBuilder::GetOrAddMesh(const UStaticMesh* StaticMesh, const FGLTFMaterialArray& Materials, int32 LODIndex)
@@ -111,7 +111,7 @@ FGLTFJsonMeshIndex FGLTFConvertBuilder::GetOrAddMesh(const UStaticMesh* StaticMe
 		return FGLTFJsonMeshIndex(INDEX_NONE);
 	}
 
-	return StaticMeshConverter.GetOrAdd(StaticMesh, nullptr, Materials, LODIndex);
+	return StaticMeshConverter->GetOrAdd(StaticMesh, nullptr, Materials, LODIndex);
 }
 
 FGLTFJsonMeshIndex FGLTFConvertBuilder::GetOrAddMesh(const USkeletalMesh* SkeletalMesh, const FGLTFMaterialArray& Materials, int32 LODIndex)
@@ -121,7 +121,7 @@ FGLTFJsonMeshIndex FGLTFConvertBuilder::GetOrAddMesh(const USkeletalMesh* Skelet
 		return FGLTFJsonMeshIndex(INDEX_NONE);
 	}
 
-	return SkeletalMeshConverter.GetOrAdd(SkeletalMesh, nullptr, Materials, LODIndex);
+	return SkeletalMeshConverter->GetOrAdd(SkeletalMesh, nullptr, Materials, LODIndex);
 }
 
 FGLTFJsonMeshIndex FGLTFConvertBuilder::GetOrAddMesh(const UMeshComponent* MeshComponent, const FGLTFMaterialArray& Materials, int32 LODIndex)
@@ -152,7 +152,7 @@ FGLTFJsonMeshIndex FGLTFConvertBuilder::GetOrAddMesh(const UStaticMeshComponent*
 		return FGLTFJsonMeshIndex(INDEX_NONE);
 	}
 
-	return StaticMeshConverter.GetOrAdd(StaticMesh, StaticMeshComponent, Materials, LODIndex);
+	return StaticMeshConverter->GetOrAdd(StaticMesh, StaticMeshComponent, Materials, LODIndex);
 }
 
 FGLTFJsonMeshIndex FGLTFConvertBuilder::GetOrAddMesh(const USkeletalMeshComponent* SkeletalMeshComponent, const FGLTFMaterialArray& Materials, int32 LODIndex)
@@ -168,17 +168,17 @@ FGLTFJsonMeshIndex FGLTFConvertBuilder::GetOrAddMesh(const USkeletalMeshComponen
 		return FGLTFJsonMeshIndex(INDEX_NONE);
 	}
 
-	return SkeletalMeshConverter.GetOrAdd(SkeletalMesh, SkeletalMeshComponent, Materials, LODIndex);
+	return SkeletalMeshConverter->GetOrAdd(SkeletalMesh, SkeletalMeshComponent, Materials, LODIndex);
 }
 
 const FGLTFMeshData* FGLTFConvertBuilder::GetOrAddMeshData(const UStaticMesh* StaticMesh, const UStaticMeshComponent* StaticMeshComponent, int32 LODIndex)
 {
-	return StaticMeshDataConverter.GetOrAdd(StaticMesh, StaticMeshComponent, LODIndex);
+	return StaticMeshDataConverter->GetOrAdd(StaticMesh, StaticMeshComponent, LODIndex);
 }
 
 const FGLTFMeshData* FGLTFConvertBuilder::GetOrAddMeshData(const USkeletalMesh* SkeletalMesh, const USkeletalMeshComponent* SkeletalMeshComponent, int32 LODIndex)
 {
-	return SkeletalMeshDataConverter.GetOrAdd(SkeletalMesh, SkeletalMeshComponent, LODIndex);
+	return SkeletalMeshDataConverter->GetOrAdd(SkeletalMesh, SkeletalMeshComponent, LODIndex);
 }
 
 FGLTFJsonMaterialIndex FGLTFConvertBuilder::GetOrAddMaterial(const UMaterialInterface* Material, const UStaticMesh* StaticMesh, int32 LODIndex, int32 MaterialIndex)
@@ -233,7 +233,7 @@ FGLTFJsonMaterialIndex FGLTFConvertBuilder::GetOrAddMaterial(const UMaterialInte
 		return FGLTFJsonMaterialIndex(INDEX_NONE);
 	}
 
-	return MaterialConverter.GetOrAdd(Material, MeshData, SectionIndices);
+	return MaterialConverter->GetOrAdd(Material, MeshData, SectionIndices);
 }
 
 FGLTFJsonSamplerIndex FGLTFConvertBuilder::GetOrAddSampler(const UTexture* Texture)
@@ -243,7 +243,7 @@ FGLTFJsonSamplerIndex FGLTFConvertBuilder::GetOrAddSampler(const UTexture* Textu
 		return FGLTFJsonSamplerIndex(INDEX_NONE);
 	}
 
-	return SamplerConverter.GetOrAdd(Texture);
+	return SamplerConverter->GetOrAdd(Texture);
 }
 
 FGLTFJsonTextureIndex FGLTFConvertBuilder::GetOrAddTexture(const UTexture2D* Texture)
@@ -273,7 +273,7 @@ FGLTFJsonTextureIndex FGLTFConvertBuilder::GetOrAddTexture(const ULightMapTextur
 		return FGLTFJsonTextureIndex(INDEX_NONE);
 	}
 
-	return TextureLightMapConverter.GetOrAdd(Texture);
+	return TextureLightMapConverter->GetOrAdd(Texture);
 }
 
 FGLTFJsonTextureIndex FGLTFConvertBuilder::GetOrAddTexture(const UTexture2D* Texture, bool bToSRGB)
@@ -283,7 +283,7 @@ FGLTFJsonTextureIndex FGLTFConvertBuilder::GetOrAddTexture(const UTexture2D* Tex
 		return FGLTFJsonTextureIndex(INDEX_NONE);
 	}
 
-	return Texture2DConverter.GetOrAdd(Texture, bToSRGB);
+	return Texture2DConverter->GetOrAdd(Texture, bToSRGB);
 }
 
 FGLTFJsonTextureIndex FGLTFConvertBuilder::GetOrAddTexture(const UTextureCube* Texture, ECubeFace CubeFace, bool bToSRGB)
@@ -293,7 +293,7 @@ FGLTFJsonTextureIndex FGLTFConvertBuilder::GetOrAddTexture(const UTextureCube* T
 		return FGLTFJsonTextureIndex(INDEX_NONE);
 	}
 
-	return TextureCubeConverter.GetOrAdd(Texture, CubeFace, bToSRGB);
+	return TextureCubeConverter->GetOrAdd(Texture, CubeFace, bToSRGB);
 }
 
 FGLTFJsonTextureIndex FGLTFConvertBuilder::GetOrAddTexture(const UTextureRenderTarget2D* Texture, bool bToSRGB)
@@ -303,7 +303,7 @@ FGLTFJsonTextureIndex FGLTFConvertBuilder::GetOrAddTexture(const UTextureRenderT
 		return FGLTFJsonTextureIndex(INDEX_NONE);
 	}
 
-	return TextureRenderTarget2DConverter.GetOrAdd(Texture, bToSRGB);
+	return TextureRenderTarget2DConverter->GetOrAdd(Texture, bToSRGB);
 }
 
 FGLTFJsonTextureIndex FGLTFConvertBuilder::GetOrAddTexture(const UTextureRenderTargetCube* Texture, ECubeFace CubeFace, bool bToSRGB)
@@ -313,12 +313,12 @@ FGLTFJsonTextureIndex FGLTFConvertBuilder::GetOrAddTexture(const UTextureRenderT
 		return FGLTFJsonTextureIndex(INDEX_NONE);
 	}
 
-	return TextureRenderTargetCubeConverter.GetOrAdd(Texture, CubeFace, bToSRGB);
+	return TextureRenderTargetCubeConverter->GetOrAdd(Texture, CubeFace, bToSRGB);
 }
 
 FGLTFJsonImageIndex FGLTFConvertBuilder::GetOrAddImage(TGLTFSharedArray<FColor>& Pixels, FIntPoint Size, bool bIgnoreAlpha, EGLTFTextureType Type, const FString& Name)
 {
-	return ImageConverter.GetOrAdd(Name, Type, bIgnoreAlpha, Size, Pixels);
+	return ImageConverter->GetOrAdd(Name, Type, bIgnoreAlpha, Size, Pixels);
 }
 
 FGLTFJsonSkinIndex FGLTFConvertBuilder::GetOrAddSkin(FGLTFJsonNodeIndex RootNode, const USkeletalMesh* SkeletalMesh)
@@ -328,7 +328,7 @@ FGLTFJsonSkinIndex FGLTFConvertBuilder::GetOrAddSkin(FGLTFJsonNodeIndex RootNode
 		return FGLTFJsonSkinIndex(INDEX_NONE);
 	}
 
-	return SkinConverter.GetOrAdd(RootNode, SkeletalMesh);
+	return SkinConverter->GetOrAdd(RootNode, SkeletalMesh);
 }
 
 FGLTFJsonSkinIndex FGLTFConvertBuilder::GetOrAddSkin(FGLTFJsonNodeIndex RootNode, const USkeletalMeshComponent* SkeletalMeshComponent)
@@ -355,7 +355,7 @@ FGLTFJsonAnimationIndex FGLTFConvertBuilder::GetOrAddAnimation(FGLTFJsonNodeInde
 		return FGLTFJsonAnimationIndex(INDEX_NONE);
 	}
 
-	return AnimationConverter.GetOrAdd(RootNode, SkeletalMesh, AnimSequence);
+	return AnimationConverter->GetOrAdd(RootNode, SkeletalMesh, AnimSequence);
 }
 
 FGLTFJsonAnimationIndex FGLTFConvertBuilder::GetOrAddAnimation(FGLTFJsonNodeIndex RootNode, const USkeletalMeshComponent* SkeletalMeshComponent)
@@ -365,7 +365,7 @@ FGLTFJsonAnimationIndex FGLTFConvertBuilder::GetOrAddAnimation(FGLTFJsonNodeInde
 		return FGLTFJsonAnimationIndex(INDEX_NONE);
 	}
 
-	return AnimationDataConverter.GetOrAdd(RootNode, SkeletalMeshComponent);
+	return AnimationDataConverter->GetOrAdd(RootNode, SkeletalMeshComponent);
 }
 
 FGLTFJsonAnimationIndex FGLTFConvertBuilder::GetOrAddAnimation(const ULevel* Level, const ULevelSequence* LevelSequence)
@@ -375,7 +375,7 @@ FGLTFJsonAnimationIndex FGLTFConvertBuilder::GetOrAddAnimation(const ULevel* Lev
 		return FGLTFJsonAnimationIndex(INDEX_NONE);
 	}
 
-	return LevelSequenceConverter.GetOrAdd(Level, LevelSequence);
+	return LevelSequenceConverter->GetOrAdd(Level, LevelSequence);
 }
 
 FGLTFJsonAnimationIndex FGLTFConvertBuilder::GetOrAddAnimation(const ALevelSequenceActor* LevelSequenceActor)
@@ -385,7 +385,7 @@ FGLTFJsonAnimationIndex FGLTFConvertBuilder::GetOrAddAnimation(const ALevelSeque
 		return FGLTFJsonAnimationIndex(INDEX_NONE);
 	}
 
-	return LevelSequenceDataConverter.GetOrAdd(LevelSequenceActor);
+	return LevelSequenceDataConverter->GetOrAdd(LevelSequenceActor);
 }
 
 FGLTFJsonNodeIndex FGLTFConvertBuilder::GetOrAddNode(const AActor* Actor)
@@ -395,7 +395,7 @@ FGLTFJsonNodeIndex FGLTFConvertBuilder::GetOrAddNode(const AActor* Actor)
 		return FGLTFJsonNodeIndex(INDEX_NONE);
 	}
 
-	return ActorConverter.GetOrAdd(Actor);
+	return ActorConverter->GetOrAdd(Actor);
 }
 
 FGLTFJsonNodeIndex FGLTFConvertBuilder::GetOrAddNode(const USceneComponent* SceneComponent)
@@ -405,7 +405,7 @@ FGLTFJsonNodeIndex FGLTFConvertBuilder::GetOrAddNode(const USceneComponent* Scen
 		return FGLTFJsonNodeIndex(INDEX_NONE);
 	}
 
-	return ComponentConverter.GetOrAdd(SceneComponent);
+	return ComponentConverter->GetOrAdd(SceneComponent);
 }
 
 FGLTFJsonNodeIndex FGLTFConvertBuilder::GetOrAddNode(const USceneComponent* SceneComponent, FName SocketName)
@@ -415,7 +415,7 @@ FGLTFJsonNodeIndex FGLTFConvertBuilder::GetOrAddNode(const USceneComponent* Scen
 		return FGLTFJsonNodeIndex(INDEX_NONE);
 	}
 
-	return ComponentSocketConverter.GetOrAdd(SceneComponent, SocketName);
+	return ComponentSocketConverter->GetOrAdd(SceneComponent, SocketName);
 }
 
 FGLTFJsonNodeIndex FGLTFConvertBuilder::GetOrAddNode(FGLTFJsonNodeIndex RootNode, const UStaticMesh* StaticMesh, FName SocketName)
@@ -425,7 +425,7 @@ FGLTFJsonNodeIndex FGLTFConvertBuilder::GetOrAddNode(FGLTFJsonNodeIndex RootNode
 		return FGLTFJsonNodeIndex(INDEX_NONE);
 	}
 
-	return StaticSocketConverter.GetOrAdd(RootNode, StaticMesh, SocketName);
+	return StaticSocketConverter->GetOrAdd(RootNode, StaticMesh, SocketName);
 }
 
 FGLTFJsonNodeIndex FGLTFConvertBuilder::GetOrAddNode(FGLTFJsonNodeIndex RootNode, const USkeletalMesh* SkeletalMesh, FName SocketName)
@@ -435,7 +435,7 @@ FGLTFJsonNodeIndex FGLTFConvertBuilder::GetOrAddNode(FGLTFJsonNodeIndex RootNode
 		return FGLTFJsonNodeIndex(INDEX_NONE);
 	}
 
-	return SkeletalSocketConverter.GetOrAdd(RootNode, SkeletalMesh, SocketName);
+	return SkeletalSocketConverter->GetOrAdd(RootNode, SkeletalMesh, SocketName);
 }
 
 FGLTFJsonNodeIndex FGLTFConvertBuilder::GetOrAddNode(FGLTFJsonNodeIndex RootNode, const USkeletalMesh* SkeletalMesh, int32 BoneIndex)
@@ -445,7 +445,7 @@ FGLTFJsonNodeIndex FGLTFConvertBuilder::GetOrAddNode(FGLTFJsonNodeIndex RootNode
 		return FGLTFJsonNodeIndex(INDEX_NONE);
 	}
 
-	return SkeletalBoneConverter.GetOrAdd(RootNode, SkeletalMesh, BoneIndex);
+	return SkeletalBoneConverter->GetOrAdd(RootNode, SkeletalMesh, BoneIndex);
 }
 
 FGLTFJsonSceneIndex FGLTFConvertBuilder::GetOrAddScene(const UWorld* World)
@@ -455,7 +455,7 @@ FGLTFJsonSceneIndex FGLTFConvertBuilder::GetOrAddScene(const UWorld* World)
 		return FGLTFJsonSceneIndex(INDEX_NONE);
 	}
 
-	return SceneConverter.GetOrAdd(World);
+	return SceneConverter->GetOrAdd(World);
 }
 
 FGLTFJsonCameraIndex FGLTFConvertBuilder::GetOrAddCamera(const UCameraComponent* CameraComponent)
@@ -465,7 +465,7 @@ FGLTFJsonCameraIndex FGLTFConvertBuilder::GetOrAddCamera(const UCameraComponent*
 		return FGLTFJsonCameraIndex(INDEX_NONE);
 	}
 
-	return CameraConverter.GetOrAdd(CameraComponent);
+	return CameraConverter->GetOrAdd(CameraComponent);
 }
 
 FGLTFJsonLightIndex FGLTFConvertBuilder::GetOrAddLight(const ULightComponent* LightComponent)
@@ -475,7 +475,7 @@ FGLTFJsonLightIndex FGLTFConvertBuilder::GetOrAddLight(const ULightComponent* Li
 		return FGLTFJsonLightIndex(INDEX_NONE);
 	}
 
-	return LightConverter.GetOrAdd(LightComponent);
+	return LightConverter->GetOrAdd(LightComponent);
 }
 
 FGLTFJsonBackdropIndex FGLTFConvertBuilder::GetOrAddBackdrop(const AActor* BackdropActor)
@@ -485,7 +485,7 @@ FGLTFJsonBackdropIndex FGLTFConvertBuilder::GetOrAddBackdrop(const AActor* Backd
 		return FGLTFJsonBackdropIndex(INDEX_NONE);
 	}
 
-	return BackdropConverter.GetOrAdd(BackdropActor);
+	return BackdropConverter->GetOrAdd(BackdropActor);
 }
 
 FGLTFJsonLightMapIndex FGLTFConvertBuilder::GetOrAddLightMap(const UStaticMeshComponent* StaticMeshComponent)
@@ -495,7 +495,7 @@ FGLTFJsonLightMapIndex FGLTFConvertBuilder::GetOrAddLightMap(const UStaticMeshCo
 		return FGLTFJsonLightMapIndex(INDEX_NONE);
 	}
 
-	return LightMapConverter.GetOrAdd(StaticMeshComponent);
+	return LightMapConverter->GetOrAdd(StaticMeshComponent);
 }
 
 FGLTFJsonHotspotIndex FGLTFConvertBuilder::GetOrAddHotspot(const AGLTFHotspotActor* HotspotActor)
@@ -505,7 +505,7 @@ FGLTFJsonHotspotIndex FGLTFConvertBuilder::GetOrAddHotspot(const AGLTFHotspotAct
 		return FGLTFJsonHotspotIndex(INDEX_NONE);
 	}
 
-	return HotspotConverter.GetOrAdd(HotspotActor);
+	return HotspotConverter->GetOrAdd(HotspotActor);
 }
 
 FGLTFJsonSkySphereIndex FGLTFConvertBuilder::GetOrAddSkySphere(const AActor* SkySphereActor)
@@ -515,7 +515,7 @@ FGLTFJsonSkySphereIndex FGLTFConvertBuilder::GetOrAddSkySphere(const AActor* Sky
 		return FGLTFJsonSkySphereIndex(INDEX_NONE);
 	}
 
-	return SkySphereConverter.GetOrAdd(SkySphereActor);
+	return SkySphereConverter->GetOrAdd(SkySphereActor);
 }
 
 FGLTFJsonEpicLevelVariantSetsIndex FGLTFConvertBuilder::GetOrAddEpicLevelVariantSets(const ULevelVariantSets* LevelVariantSets)
@@ -525,7 +525,7 @@ FGLTFJsonEpicLevelVariantSetsIndex FGLTFConvertBuilder::GetOrAddEpicLevelVariant
 		return FGLTFJsonEpicLevelVariantSetsIndex(INDEX_NONE);
 	}
 
-	return EpicLevelVariantSetsConverter.GetOrAdd(LevelVariantSets);
+	return EpicLevelVariantSetsConverter->GetOrAdd(LevelVariantSets);
 }
 
 FGLTFJsonKhrMaterialVariantIndex FGLTFConvertBuilder::GetOrAddKhrMaterialVariant(const UVariant* Variant)
@@ -535,7 +535,7 @@ FGLTFJsonKhrMaterialVariantIndex FGLTFConvertBuilder::GetOrAddKhrMaterialVariant
 		return FGLTFJsonKhrMaterialVariantIndex(INDEX_NONE);
 	}
 
-	return KhrMaterialVariantConverter.GetOrAdd(Variant);
+	return KhrMaterialVariantConverter->GetOrAdd(Variant);
 }
 
 void FGLTFConvertBuilder::RegisterObjectVariant(const UObject* Object, const UPropertyValue* Property)
diff --git a/Engine/Plugins/Enterprise/GLTFExporter/MigratedFromGit/Source/GLTFExporter/Private/Builders/GLTFConvertBuilder.h b/Engine/Plugins/Enterprise/GLTFExporter/MigratedFromGit/Source/GLTFExporter/Private/Builders/GLTFConvertBuilder.h
--- a/Engine/Plugins/Enterprise/GLTFExporter/MigratedFromGit/Source/GLTFExporter/Private/Builders/GLTFConvertBuilder.h
+++ b/Engine/Plugins/Enterprise/GLTFExporter/MigratedFromGit/Source/GLTFExporter/Private/Builders/GLTFConvertBuilder.h
@@ -101,55 +101,55 @@ public:
 	void RegisterObjectVariant(const UObject* Object, const UPropertyValue* Property);
 	const TArray<const UPropertyValue*>* GetObjectVariants(const UObject* Object) const;
 
-private:
+	TUniquePtr<IGLTFPositionBufferConverter> PositionBufferConverter = MakeUnique<FGLTFPositionBufferConverter>(*this);
+	TUniquePtr<IGLTFColorBufferConverter> ColorBufferConverter = MakeUnique<FGLTFColorBufferConverter>(*this);
+	TUniquePtr<IGLTFNormalBufferConverter> NormalBufferConverter = MakeUnique<FGLTFNormalBufferConverter>(*this);
+	TUniquePtr<IGLTFTangentBufferConverter> TangentBufferConverter = MakeUnique<FGLTFTangentBufferConverter>(*this);
+	TUniquePtr<IGLTFUVBufferConverter> UVBufferConverter = MakeUnique<FGLTFUVBufferConverter>(*this);
+	TUniquePtr<IGLTFBoneIndexBufferConverter> BoneIndexBufferConverter = MakeUnique<FGLTFBoneIndexBufferConverter>(*this);
+	TUniquePtr<IGLTFBoneWeightBufferConverter> BoneWeightBufferConverter = MakeUnique<FGLTFBoneWeightBufferConverter>(*this);
+	TUniquePtr<IGLTFIndexBufferConverter> IndexBufferConverter = MakeUnique<FGLTFIndexBufferConverter>(*this);
+
+	TUniquePtr<IGLTFStaticMeshConverter> StaticMeshConverter = MakeUnique<FGLTFStaticMeshConverter>(*this);
+	TUniquePtr<IGLTFSkeletalMeshConverter> SkeletalMeshConverter = MakeUnique<FGLTFSkeletalMeshConverter>(*this);
+
+	TUniquePtr<IGLTFMaterialConverter> MaterialConverter = MakeUnique<FGLTFMaterialConverter>(*this);
+	TUniquePtr<IGLTFStaticMeshDataConverter> StaticMeshDataConverter = MakeUnique<FGLTFStaticMeshDataConverter>(*this);
+	TUniquePtr<IGLTFSkeletalMeshDataConverter> SkeletalMeshDataConverter = MakeUnique<FGLTFSkeletalMeshDataConverter>(*this);
+
+	TUniquePtr<IGLTFSamplerConverter> SamplerConverter = MakeUnique<FGLTFSamplerConverter>(*this);
+	TUniquePtr<IGLTFTexture2DConverter> Texture2DConverter = MakeUnique<FGLTFTexture2DConverter>(*this);
+	TUniquePtr<IGLTFTextureCubeConverter> TextureCubeConverter = MakeUnique<FGLTFTextureCubeConverter>(*this);
+	TUniquePtr<IGLTFTextureRenderTarget2DConverter> TextureRenderTarget2DConverter = MakeUnique<FGLTFTextureRenderTarget2DConverter>(*this);
+	TUniquePtr<IGLTFTextureRenderTargetCubeConverter> TextureRenderTargetCubeConverter = MakeUnique<FGLTFTextureRenderTargetCubeConverter>(*this);
+	TUniquePtr<IGLTFTextureLightMapConverter> TextureLightMapConverter = MakeUnique<FGLTFTextureLightMapConverter>(*this);
+	TUniquePtr<IGLTFImageConverter> ImageConverter = MakeUnique<FGLTFImageConverter>(*this);
+
+	TUniquePtr<IGLTFSkinConverter> SkinConverter = MakeUnique<FGLTFSkinConverter>(*this);
+	TUniquePtr<IGLTFAnimationConverter> AnimationConverter = MakeUnique<FGLTFAnimationConverter>(*this);
+	TUniquePtr<IGLTFAnimationDataConverter> AnimationDataConverter = MakeUnique<FGLTFAnimationDataConverter>(*this);
+	TUniquePtr<IGLTFLevelSequenceConverter> LevelSequenceConverter = MakeUnique<FGLTFLevelSequenceConverter>(*this);
+	TUniquePtr<IGLTFLevelSequenceDataConverter> LevelSequenceDataConverter = MakeUnique<FGLTFLevelSequenceDataConverter>(*this);
+
+	TUniquePtr<IGLTFActorConverter> ActorConverter = MakeUnique<FGLTFActorConverter>(*this);
+	TUniquePtr<IGLTFComponentConverter> ComponentConverter = MakeUnique<FGLTFComponentConverter>(*this);
+	TUniquePtr<IGLTFComponentSocketConverter> ComponentSocketConverter = MakeUnique<FGLTFComponentSocketConverter>(*this);
+	TUniquePtr<IGLTFStaticSocketConverter> StaticSocketConverter = MakeUnique<FGLTFStaticSocketConverter>(*this);
+	TUniquePtr<IGLTFSkeletalSocketConverter> SkeletalSocketConverter = MakeUnique<FGLTFSkeletalSocketConverter>(*this);
+	TUniquePtr<IGLTFSkeletalBoneConverter> SkeletalBoneConverter = MakeUnique<FGLTFSkeletalBoneConverter>(*this);
+	TUniquePtr<IGLTFSceneConverter> SceneConverter = MakeUnique<FGLTFSceneConverter>(*this);
+
+	TUniquePtr<IGLTFCameraConverter> CameraConverter = MakeUnique<FGLTFCameraConverter>(*this);
+	TUniquePtr<IGLTFLightConverter> LightConverter = MakeUnique<FGLTFLightConverter>(*this);
+	TUniquePtr<IGLTFBackdropConverter> BackdropConverter = MakeUnique<FGLTFBackdropConverter>(*this);
+	TUniquePtr<IGLTFLightMapConverter> LightMapConverter = MakeUnique<FGLTFLightMapConverter>(*this);
+	TUniquePtr<IGLTFHotspotConverter> HotspotConverter = MakeUnique<FGLTFHotspotConverter>(*this);
+	TUniquePtr<IGLTFSkySphereConverter> SkySphereConverter = MakeUnique<FGLTFSkySphereConverter>(*this);
+
+	TUniquePtr<IGLTFEpicLevelVariantSetsConverter> EpicLevelVariantSetsConverter = MakeUnique<FGLTFEpicLevelVariantSetsConverter>(*this);
+	TUniquePtr<IGLTFKhrMaterialVariantConverter> KhrMaterialVariantConverter = MakeUnique<FGLTFKhrMaterialVariantConverter>(*this);
 
-	FGLTFPositionBufferConverter PositionBufferConverter{ *this };
-	FGLTFColorBufferConverter ColorBufferConverter{ *this };
-	FGLTFNormalBufferConverter NormalBufferConverter{ *this };
-	FGLTFTangentBufferConverter TangentBufferConverter{ *this };
-	FGLTFUVBufferConverter UVBufferConverter{ *this };
-	FGLTFBoneIndexBufferConverter BoneIndexBufferConverter{ *this };
-	FGLTFBoneWeightBufferConverter BoneWeightBufferConverter{ *this };
-	FGLTFIndexBufferConverter IndexBufferConverter{ *this };
-
-	FGLTFStaticMeshConverter StaticMeshConverter{ *this };
-	FGLTFSkeletalMeshConverter SkeletalMeshConverter{ *this };
-
-	FGLTFMaterialConverter MaterialConverter{ *this };
-	FGLTFStaticMeshDataConverter StaticMeshDataConverter{ *this };
-	FGLTFSkeletalMeshDataConverter SkeletalMeshDataConverter{ *this };
-
-	FGLTFSamplerConverter SamplerConverter{ *this };
-	FGLTFTexture2DConverter Texture2DConverter{ *this };
-	FGLTFTextureCubeConverter TextureCubeConverter{ *this };
-	FGLTFTextureRenderTarget2DConverter TextureRenderTarget2DConverter{ *this };
-	FGLTFTextureRenderTargetCubeConverter TextureRenderTargetCubeConverter{ *this };
-	FGLTFTextureLightMapConverter TextureLightMapConverter{ *this };
-	FGLTFImageConverter ImageConverter{ *this };
-
-	FGLTFSkinConverter SkinConverter{ *this };
-	FGLTFAnimationConverter AnimationConverter{ *this };
-	FGLTFAnimationDataConverter AnimationDataConverter{ *this };
-	FGLTFLevelSequenceConverter LevelSequenceConverter{ *this };
-	FGLTFLevelSequenceDataConverter LevelSequenceDataConverter{ *this };
-
-	FGLTFActorConverter ActorConverter{ *this };
-	FGLTFComponentConverter ComponentConverter{ *this };
-	FGLTFComponentSocketConverter ComponentSocketConverter{ *this };
-	FGLTFStaticSocketConverter StaticSocketConverter{ *this };
-	FGLTFSkeletalSocketConverter SkeletalSocketConverter{ *this };
-	FGLTFSkeletalBoneConverter SkeletalBoneConverter{ *this };
-	FGLTFSceneConverter SceneConverter{ *this };
-
-	FGLTFCameraConverter CameraConverter{ *this };
-	FGLTFLightConverter LightConverter{ *this };
-	FGLTFBackdropConverter BackdropConverter{ *this };
-	FGLTFLightMapConverter LightMapConverter{ *this };
-	FGLTFHotspotConverter HotspotConverter{ *this };
-	FGLTFSkySphereConverter SkySphereConverter{ *this };
-
-	FGLTFEpicLevelVariantSetsConverter EpicLevelVariantSetsConverter{ *this };
-	FGLTFKhrMaterialVariantConverter KhrMaterialVariantConverter{ *this };
+private:
 
 	TMap<const UObject*, TArray<const UPropertyValue*>> ObjectVariants;
 };
diff --git a/Engine/Plugins/Enterprise/GLTFExporter/MigratedFromGit/Source/GLTFExporter/Private/Converters/GLTFMeshDataConverters.h b/Engine/Plugins/Enterprise/GLTFExporter/MigratedFromGit/Source/GLTFExporter/Private/Converters/GLTFMeshDataConverters.h
--- a/Engine/Plugins/Enterprise/GLTFExporter/MigratedFromGit/Source/GLTFExporter/Private/Converters/GLTFMeshDataConverters.h
+++ b/Engine/Plugins/Enterprise/GLTFExporter/MigratedFromGit/Source/GLTFExporter/Private/Converters/GLTFMeshDataConverters.h
@@ -6,6 +6,9 @@
 #include "Converters/GLTFMeshData.h"
 #include "Converters/GLTFBuilderContext.h"
 
+typedef TGLTFConverter<const FGLTFMeshData*, const UStaticMesh*, const UStaticMeshComponent*, int32> IGLTFStaticMeshDataConverter;
+typedef TGLTFConverter<const FGLTFMeshData*, const USkeletalMesh*, const USkeletalMeshComponent*, int32> IGLTFSkeletalMeshDataConverter;
+
 template <typename MeshType, typename MeshComponentType>
 class TGLTFMeshDataConverter : public FGLTFBuilderContext, public TGLTFConverter<const FGLTFMeshData*, const MeshType*, const MeshComponentType*, int32>
 {
