[GLTFExporter] feat: add support for material proxies via user data and configurable by export option bExportProxyMaterials

commit 603d401e1f25ee12fcb67da88e3605a36a326e72
Date: Wed, 26 Jan 2022 18:15:29 +0100

    feat: add support for material proxies via user data and configurable by export option bExportProxyMaterials
    (cherry picked from commit d43a60a92bd24574f0d5c07bf39ca7685ce6b53f​)

#skipci

diff --git a/Engine/Plugins/Enterprise/GLTFExporter/MigratedFromGit/Docs/export-options-reference.md b/Engine/Plugins/Enterprise/GLTFExporter/MigratedFromGit/Docs/export-options-reference.md
--- a/Engine/Plugins/Enterprise/GLTFExporter/MigratedFromGit/Docs/export-options-reference.md
+++ b/Engine/Plugins/Enterprise/GLTFExporter/MigratedFromGit/Docs/export-options-reference.md
@@ -7,6 +7,7 @@ Option                         | Description
 `Strict Compliance`            | If enabled, certain values (like HDR colors and light angles) will be truncated during export to strictly conform to the formal glTF specification.
 `Skip Near Default Values`     | If enabled, floating-point-based JSON properties that are nearly equal to their default value will not be exported and thus regarded as exactly default, reducing size of JSON data.
 `Include Generator Version`    | If enabled, version info for Unreal Engine and exporter plugin will be included as metadata in the glTF asset, which is useful when reporting issues.
+`Export Proxy Materials`       | If enabled, materials that have a proxy defined in their user data, will be exported using that proxy instead. Used mainly by pre-baked materials.
 `Export Unlit Materials`       | If enabled, materials with shading model unlit will be properly exported. Uses extension KHR_materials_unlit.
 `Export Clear Coat Materials`  | If enabled, materials with shading model clear coat will be properly exported. Uses extension KHR_materials_clearcoat, which is not supported by all glTF viewers.
 `Export Extra Blend Modes`     | If enabled, materials with blend modes additive, modulate, and alpha composite will be properly exported. Uses extension EPIC_blend_modes, which is supported by Unreal's glTF viewer.
diff --git a/Engine/Plugins/Enterprise/GLTFExporter/MigratedFromGit/Source/GLTFExporter/Private/Builders/GLTFBuilder.cpp b/Engine/Plugins/Enterprise/GLTFExporter/MigratedFromGit/Source/GLTFExporter/Private/Builders/GLTFBuilder.cpp
--- a/Engine/Plugins/Enterprise/GLTFExporter/MigratedFromGit/Source/GLTFExporter/Private/Builders/GLTFBuilder.cpp
+++ b/Engine/Plugins/Enterprise/GLTFExporter/MigratedFromGit/Source/GLTFExporter/Private/Builders/GLTFBuilder.cpp
@@ -13,6 +13,11 @@ FGLTFBuilder::FGLTFBuilder(const FString& FilePath, const UGLTFExportOptions* Ex
 {
 }
 
+const UMaterialInterface* FGLTFBuilder::ResolveProxy(const UMaterialInterface* Material) const
+{
+	return ExportOptions->bExportProxyMaterials ? UGLTFMaterialExportOptions::ResolveProxy(Material) : Material;
+}
+
 FIntPoint FGLTFBuilder::GetBakeSizeForMaterialProperty(const UMaterialInterface* Material, EGLTFMaterialPropertyGroup PropertyGroup) const
 {
 	EGLTFMaterialBakeSizePOT DefaultValue = ExportOptions->DefaultMaterialBakeSize;
@@ -24,6 +29,7 @@ FIntPoint FGLTFBuilder::GetBakeSizeForMaterialProperty(const UMaterialInterface*
 		}
 	}
 
+	// TODO: add option to ignore override
 	const EGLTFMaterialBakeSizePOT Size = UGLTFMaterialExportOptions::GetBakeSizeForPropertyGroup(Material, PropertyGroup, DefaultValue);
 	const int32 PixelSize = 1 << static_cast<uint8>(Size);
 	return { PixelSize, PixelSize };
@@ -40,6 +46,7 @@ TextureFilter FGLTFBuilder::GetBakeFilterForMaterialProperty(const UMaterialInte
 		}
 	}
 
+	// TODO: add option to ignore override
 	return UGLTFMaterialExportOptions::GetBakeFilterForPropertyGroup(Material, PropertyGroup, DefaultValue);
 }
 
@@ -54,6 +61,7 @@ TextureAddress FGLTFBuilder::GetBakeTilingForMaterialProperty(const UMaterialInt
 		}
 	}
 
+	// TODO: add option to ignore override
 	return UGLTFMaterialExportOptions::GetBakeTilingForPropertyGroup(Material, PropertyGroup, DefaultValue);
 }
 
diff --git a/Engine/Plugins/Enterprise/GLTFExporter/MigratedFromGit/Source/GLTFExporter/Private/Builders/GLTFBuilder.h b/Engine/Plugins/Enterprise/GLTFExporter/MigratedFromGit/Source/GLTFExporter/Private/Builders/GLTFBuilder.h
--- a/Engine/Plugins/Enterprise/GLTFExporter/MigratedFromGit/Source/GLTFExporter/Private/Builders/GLTFBuilder.h
+++ b/Engine/Plugins/Enterprise/GLTFExporter/MigratedFromGit/Source/GLTFExporter/Private/Builders/GLTFBuilder.h
@@ -21,6 +21,8 @@ public:
 	// TODO: make ExportOptions private and expose each option via getters to ease overriding settings in future
 	const UGLTFExportOptions* ExportOptions;
 
+	const UMaterialInterface* ResolveProxy(const UMaterialInterface* Material) const;
+
 	FIntPoint GetBakeSizeForMaterialProperty(const UMaterialInterface* Material, EGLTFMaterialPropertyGroup PropertyGroup) const;
 	TextureFilter GetBakeFilterForMaterialProperty(const UMaterialInterface* Material, EGLTFMaterialPropertyGroup PropertyGroup) const;
 	TextureAddress GetBakeTilingForMaterialProperty(const UMaterialInterface* Material, EGLTFMaterialPropertyGroup PropertyGroup) const;
diff --git a/Engine/Plugins/Enterprise/GLTFExporter/MigratedFromGit/Source/GLTFExporter/Private/GLTFExportOptions.cpp b/Engine/Plugins/Enterprise/GLTFExporter/MigratedFromGit/Source/GLTFExporter/Private/GLTFExportOptions.cpp
--- a/Engine/Plugins/Enterprise/GLTFExporter/MigratedFromGit/Source/GLTFExporter/Private/GLTFExportOptions.cpp
+++ b/Engine/Plugins/Enterprise/GLTFExporter/MigratedFromGit/Source/GLTFExporter/Private/GLTFExportOptions.cpp
@@ -17,6 +17,7 @@ void UGLTFExportOptions::ResetToDefault()
 	bStrictCompliance = true;
 	bSkipNearDefaultValues = true;
 	bIncludeGeneratorVersion = true;
+	bExportProxyMaterials = true;
 	bExportUnlitMaterials = true;
 	bExportClearCoatMaterials = true;
 	bExportExtraBlendModes = false;
diff --git a/Engine/Plugins/Enterprise/GLTFExporter/MigratedFromGit/Source/GLTFExporter/Private/Tasks/GLTFMaterialTasks.cpp b/Engine/Plugins/Enterprise/GLTFExporter/MigratedFromGit/Source/GLTFExporter/Private/Tasks/GLTFMaterialTasks.cpp
--- a/Engine/Plugins/Enterprise/GLTFExporter/MigratedFromGit/Source/GLTFExporter/Private/Tasks/GLTFMaterialTasks.cpp
+++ b/Engine/Plugins/Enterprise/GLTFExporter/MigratedFromGit/Source/GLTFExporter/Private/Tasks/GLTFMaterialTasks.cpp
@@ -50,6 +50,8 @@ namespace
 
 void FGLTFMaterialTask::Complete()
 {
+	Material = Builder.ResolveProxy(Material);
+
 	{
 		const UMaterial* ParentMaterial = Material->GetMaterial();
 
diff --git a/Engine/Plugins/Enterprise/GLTFExporter/MigratedFromGit/Source/GLTFExporter/Private/UserData/GLTFMaterialUserData.cpp b/Engine/Plugins/Enterprise/GLTFExporter/MigratedFromGit/Source/GLTFExporter/Private/UserData/GLTFMaterialUserData.cpp
--- a/Engine/Plugins/Enterprise/GLTFExporter/MigratedFromGit/Source/GLTFExporter/Private/UserData/GLTFMaterialUserData.cpp
+++ b/Engine/Plugins/Enterprise/GLTFExporter/MigratedFromGit/Source/GLTFExporter/Private/UserData/GLTFMaterialUserData.cpp
@@ -14,6 +14,35 @@ FGLTFOverrideMaterialBakeSettings::FGLTFOverrideMaterialBakeSettings()
 {
 }
 
+const UMaterialInterface* UGLTFMaterialExportOptions::ResolveProxy(const UMaterialInterface* Material)
+{
+	TSet<const UMaterialInterface*> PreviousProxies;
+
+	for (;;)
+	{
+		const UGLTFMaterialExportOptions* UserData = const_cast<UMaterialInterface*>(Material)->GetAssetUserData<UGLTFMaterialExportOptions>();
+		if (UserData == nullptr)
+		{
+			return Material;
+		}
+
+		const UMaterialInterface* Proxy = UserData->Proxy;
+		if (Proxy == nullptr)
+		{
+			return Material;
+		}
+
+		if (PreviousProxies.Contains(Proxy))
+		{
+			// report warning about cyclic references
+			return Material;
+		}
+
+		PreviousProxies.Add(Proxy);
+		Material = Proxy;
+	}
+}
+
 EGLTFMaterialBakeSizePOT UGLTFMaterialExportOptions::GetBakeSizeForPropertyGroup(const UMaterialInterface* Material, EGLTFMaterialPropertyGroup PropertyGroup, EGLTFMaterialBakeSizePOT DefaultValue)
 {
 	if (const FGLTFOverrideMaterialBakeSettings* BakeSettings = GetBakeSettingsByPredicate(Material, PropertyGroup, [](const FGLTFOverrideMaterialBakeSettings& Settings) { return Settings.bOverrideSize; }))
diff --git a/Engine/Plugins/Enterprise/GLTFExporter/MigratedFromGit/Source/GLTFExporter/Public/GLTFExportOptions.h b/Engine/Plugins/Enterprise/GLTFExporter/MigratedFromGit/Source/GLTFExporter/Public/GLTFExportOptions.h
--- a/Engine/Plugins/Enterprise/GLTFExporter/MigratedFromGit/Source/GLTFExporter/Public/GLTFExportOptions.h
+++ b/Engine/Plugins/Enterprise/GLTFExporter/MigratedFromGit/Source/GLTFExporter/Public/GLTFExportOptions.h
@@ -99,6 +99,10 @@ class GLTFEXPORTER_API UGLTFExportOptions : public UObject
 	UPROPERTY(EditAnywhere, BlueprintReadWrite, Config, Category = General)
 	bool bIncludeGeneratorVersion;
 
+	/** If enabled, materials that have a proxy defined in their user data, will be exported using that proxy instead. Used mainly by pre-baked materials. */
+	UPROPERTY(EditAnywhere, BlueprintReadWrite, Config, Category = Material)
+	bool bExportProxyMaterials;
+
 	/** If enabled, materials with shading model unlit will be properly exported. Uses extension KHR_materials_unlit. */
 	UPROPERTY(EditAnywhere, BlueprintReadWrite, Config, Category = Material)
 	bool bExportUnlitMaterials;
diff --git a/Engine/Plugins/Enterprise/GLTFExporter/MigratedFromGit/Source/GLTFExporter/Public/UserData/GLTFMaterialUserData.h b/Engine/Plugins/Enterprise/GLTFExporter/MigratedFromGit/Source/GLTFExporter/Public/UserData/GLTFMaterialUserData.h
--- a/Engine/Plugins/Enterprise/GLTFExporter/MigratedFromGit/Source/GLTFExporter/Public/UserData/GLTFMaterialUserData.h
+++ b/Engine/Plugins/Enterprise/GLTFExporter/MigratedFromGit/Source/GLTFExporter/Public/UserData/GLTFMaterialUserData.h
@@ -94,6 +94,10 @@ public:
 
 	// TODO: add support for overriding more export options
 
+	/** If assigned, export will use the proxy instead of the original material. */
+	UPROPERTY(EditAnywhere, BlueprintReadWrite, Category = "General")
+	UMaterialInterface* Proxy;
+
 	/** Default bake settings for this material in general. */
 	UPROPERTY(EditAnywhere, BlueprintReadWrite, Category = "Override Bake Settings", meta = (ShowOnlyInnerProperties))
 	FGLTFOverrideMaterialBakeSettings Default;
@@ -102,6 +106,8 @@ public:
 	UPROPERTY(EditAnywhere, BlueprintReadWrite, Category = "Override Bake Settings")
 	TMap<EGLTFMaterialPropertyGroup, FGLTFOverrideMaterialBakeSettings> Inputs;
 
+	static const UMaterialInterface* ResolveProxy(const UMaterialInterface* Material);
+
 	static EGLTFMaterialBakeSizePOT GetBakeSizeForPropertyGroup(const UMaterialInterface* Material, EGLTFMaterialPropertyGroup PropertyGroup, EGLTFMaterialBakeSizePOT DefaultValue);
 	static TextureFilter GetBakeFilterForPropertyGroup(const UMaterialInterface* Material, EGLTFMaterialPropertyGroup PropertyGroup, TextureFilter DefaultValue);
 	static TextureAddress GetBakeTilingForPropertyGroup(const UMaterialInterface* Material, EGLTFMaterialPropertyGroup PropertyGroup, TextureAddress DefaultValue);
