[GLTFExporter] feat: add base materials and material functions for prebaking

commit 52be1a3e07bcef273561aba6b9108394defb17ce
Date: Thu, 27 Jan 2022 18:32:01 +0100

    feat: add base materials and material functions for prebaking
    (cherry picked from commit bc95d0ec091c80ed9a22bfdd62d3a4bbf0782ea9​)

#skipci

diff --git a/Engine/Plugins/Enterprise/GLTFExporter/MigratedFromGit/Source/GLTFExporter/Private/Converters/GLTFMaterialUtility.cpp b/Engine/Plugins/Enterprise/GLTFExporter/MigratedFromGit/Source/GLTFExporter/Private/Converters/GLTFMaterialUtility.cpp
--- a/Engine/Plugins/Enterprise/GLTFExporter/MigratedFromGit/Source/GLTFExporter/Private/Converters/GLTFMaterialUtility.cpp
+++ b/Engine/Plugins/Enterprise/GLTFExporter/MigratedFromGit/Source/GLTFExporter/Private/Converters/GLTFMaterialUtility.cpp
@@ -21,7 +21,7 @@
 
 UMaterialInterface* FGLTFMaterialUtility::GetDefault()
 {
-	static UMaterialInterface* DefaultMaterial = LoadObject<UMaterialInterface>(nullptr, TEXT("/GLTFExporter/Materials/Default.Default"));
+	static UMaterialInterface* DefaultMaterial = LoadObject<UMaterialInterface>(nullptr, TEXT("/GLTFExporter/Materials/GLTFDefault.GLTFDefault"));
 	return DefaultMaterial;
 }
 
