[GLTFExporter] feat: add editor style for exporter with glTF logo icon

commit d0d9d135c6ed706208e53fdfa75f9eb364e1786c
Date: Thu, 27 Jan 2022 18:50:35 +0100

    feat: add editor style for exporter with glTF logo icon
    (cherry picked from commit 03920f6b100d834a06fd78755e2b0fa9d32af126​)

#skipci

diff --git a/Engine/Plugins/Enterprise/GLTFExporter/MigratedFromGit/Source/GLTFExporter/Private/Actions/GLTFEditorStyle.h b/Engine/Plugins/Enterprise/GLTFExporter/MigratedFromGit/Source/GLTFExporter/Private/Actions/GLTFEditorStyle.h
new file mode 100644
--- /dev/null
+++ b/Engine/Plugins/Enterprise/GLTFExporter/MigratedFromGit/Source/GLTFExporter/Private/Actions/GLTFEditorStyle.h
@@ -0,0 +1,44 @@
+// Copyright Epic Games, Inc. All Rights Reserved.
+
+#pragma once
+
+#if WITH_EDITOR
+
+#include "CoreMinimal.h"
+#include "GLTFExporterModule.h"
+#include "Brushes/SlateImageBrush.h"
+#include "Interfaces/IPluginManager.h"
+#include "Styling/SlateStyleRegistry.h"
+#include "Styling/SlateStyle.h"
+
+class FGLTFEditorStyle : public FSlateStyleSet
+{
+public:
+
+	 FGLTFEditorStyle()
+		 : FSlateStyleSet("GLTFEditorStyle")
+	 {
+	 	const TSharedPtr<IPlugin> Plugin = IPluginManager::Get().FindPlugin(GLTFEXPORTER_MODULE_NAME);
+	 	check(Plugin.IsValid());
+
+	 	FSlateStyleSet::SetContentRoot(Plugin->GetContentDir());
+
+	 	const FVector2D Icon16x16(16.0f, 16.0f);
+	 	Set("Icon16x16", new FSlateImageBrush(FSlateStyleSet::RootToContentDir(TEXT("Icon16x16.png")), Icon16x16));
+
+		FSlateStyleRegistry::RegisterSlateStyle(*this);
+	 }
+
+	 virtual ~FGLTFEditorStyle() override
+	 {
+		FSlateStyleRegistry::UnRegisterSlateStyle(*this);
+	 }
+
+	static FGLTFEditorStyle& Get()
+	{
+		static FGLTFEditorStyle Singleton;
+		return Singleton;
+	}
+};
+
+#endif
