[GLTFExporter] feat: add prebake asset actions for all material classes

commit ee55b871745f969762887b3191638352a96f9f96
Date: Thu, 27 Jan 2022 18:50:46 +0100

    feat: add prebake asset actions for all material classes
    (cherry picked from commit f434485aad5a811ed10a4cbadc46d74567f93621​)

#skipci

diff --git a/Engine/Plugins/Enterprise/GLTFExporter/MigratedFromGit/Source/GLTFExporter/GLTFExporter.Build.cs b/Engine/Plugins/Enterprise/GLTFExporter/MigratedFromGit/Source/GLTFExporter/GLTFExporter.Build.cs
--- a/Engine/Plugins/Enterprise/GLTFExporter/MigratedFromGit/Source/GLTFExporter/GLTFExporter.Build.cs
+++ b/Engine/Plugins/Enterprise/GLTFExporter/MigratedFromGit/Source/GLTFExporter/GLTFExporter.Build.cs
@@ -46,6 +46,7 @@ public class GLTFExporter : ModuleRules
 					"InputCore",
 					"EditorStyle",
 					"PropertyEditor",
+					"ToolMenus",
 					"MeshMergeUtilities",
 					"MeshDescription",
 					"StaticMeshDescription",
diff --git a/Engine/Plugins/Enterprise/GLTFExporter/MigratedFromGit/Source/GLTFExporter/Private/Actions/GLTFPrebakeAssetActions.cpp b/Engine/Plugins/Enterprise/GLTFExporter/MigratedFromGit/Source/GLTFExporter/Private/Actions/GLTFPrebakeAssetActions.cpp
new file mode 100644
--- /dev/null
+++ b/Engine/Plugins/Enterprise/GLTFExporter/MigratedFromGit/Source/GLTFExporter/Private/Actions/GLTFPrebakeAssetActions.cpp
@@ -0,0 +1,53 @@
+// Copyright Epic Games, Inc. All Rights Reserved.
+
+#if WITH_EDITOR
+
+#include "Actions/GLTFPrebakeAssetActions.h"
+#include "Actions/GLTFEditorStyle.h"
+#include "ToolMenus.h"
+
+#define LOCTEXT_NAMESPACE "GLTFPrebakeAssetActions"
+
+FGLTFPrebakeAssetActions::FGLTFPrebakeAssetActions(const TSharedRef<IAssetTypeActions>& OriginalActions)
+	: OriginalActions(OriginalActions)
+{
+}
+
+bool FGLTFPrebakeAssetActions::HasActions(const TArray<UObject*>& InObjects) const
+{
+	return true;
+}
+
+void FGLTFPrebakeAssetActions::GetActions(const TArray<UObject*>& InObjects, FToolMenuSection& Section)
+{
+	OriginalActions->GetActions(InObjects, Section);
+	GetPrebakeActions(InObjects, Section);
+}
+
+void FGLTFPrebakeAssetActions::GetPrebakeActions(const TArray<UObject*>& InObjects, FToolMenuSection& Section)
+{
+	const TArray<FWeakObjectPtr> Objects(InObjects);
+
+	Section.AddMenuEntry(
+		"MenuEntry_Prebake",
+		LOCTEXT("MenuEntry_Prebake", "Prebake glTF Export"),
+		LOCTEXT("MenuEntry_PrebakeTooltip", "Creates a prebaked version of this material for glTF export."),
+		FSlateIcon(FGLTFEditorStyle::Get().GetStyleSetName(), "Icon16x16"),
+		FUIAction(FExecuteAction::CreateSP(this, &FGLTFPrebakeAssetActions::OnPrebake, Objects))
+		);
+}
+
+void FGLTFPrebakeAssetActions::OnPrebake(TArray<FWeakObjectPtr> Objects) const
+{
+	for (const FWeakObjectPtr& Object : Objects)
+	{
+		if (UMaterialInterface* Material = Cast<UMaterialInterface>(Object.Get()))
+		{
+			// TODO
+		}
+	}
+}
+
+#undef LOCTEXT_NAMESPACE
+
+#endif
diff --git a/Engine/Plugins/Enterprise/GLTFExporter/MigratedFromGit/Source/GLTFExporter/Private/Actions/GLTFPrebakeAssetActions.h b/Engine/Plugins/Enterprise/GLTFExporter/MigratedFromGit/Source/GLTFExporter/Private/Actions/GLTFPrebakeAssetActions.h
new file mode 100644
--- /dev/null
+++ b/Engine/Plugins/Enterprise/GLTFExporter/MigratedFromGit/Source/GLTFExporter/Private/Actions/GLTFPrebakeAssetActions.h
@@ -0,0 +1,57 @@
+// Copyright Epic Games, Inc. All Rights Reserved.
+
+#pragma once
+
+#if WITH_EDITOR
+
+#include "AssetTypeActions_Base.h"
+
+class FGLTFPrebakeAssetActions : public IAssetTypeActions
+{
+public:
+
+	FGLTFPrebakeAssetActions(const TSharedRef<IAssetTypeActions>& OriginalActions);
+
+	virtual bool HasActions(const TArray<UObject*>& InObjects) const override;
+	virtual void GetActions(const TArray<UObject*>& InObjects, FToolMenuSection& Section) override;
+
+	virtual FText GetName() const override { return OriginalActions->GetName(); }
+	virtual UClass* GetSupportedClass() const override { return OriginalActions->GetSupportedClass(); }
+	virtual FColor GetTypeColor() const override { return OriginalActions->GetTypeColor(); }
+	virtual void GetActions(const TArray<UObject*>& InObjects, FMenuBuilder& MenuBuilder) override { OriginalActions->GetActions(InObjects, MenuBuilder); }
+	virtual void OpenAssetEditor(const TArray<UObject*>& InObjects, TSharedPtr<IToolkitHost> EditWithinLevelEditor) override { OriginalActions->OpenAssetEditor(InObjects, EditWithinLevelEditor); }
+	virtual void AssetsActivated(const TArray<UObject*>& InObjects, EAssetTypeActivationMethod::Type ActivationType) override { OriginalActions->AssetsActivated(InObjects, ActivationType); }
+	virtual bool AssetsActivatedOverride(const TArray<UObject*>& InObjects, EAssetTypeActivationMethod::Type ActivationType) override { return OriginalActions->AssetsActivatedOverride(InObjects, ActivationType); }
+	virtual TArray<FAssetData> GetValidAssetsForPreviewOrEdit(TArrayView<const FAssetData> InAssetDatas, bool bIsPreview) override { return OriginalActions->GetValidAssetsForPreviewOrEdit(InAssetDatas, bIsPreview); }
+	virtual bool CanFilter() override { return OriginalActions->CanFilter(); }
+	virtual bool CanLocalize() const override { return OriginalActions->CanLocalize(); }
+	virtual bool CanMerge() const override { return OriginalActions->CanMerge(); }
+	virtual void Merge( UObject* InObject) override { OriginalActions->Merge(InObject); }
+	virtual void Merge(UObject* BaseAsset, UObject* RemoteAsset, UObject* LocalAsset, const FOnMergeResolved& ResolutionCallback) override { OriginalActions->Merge(BaseAsset, RemoteAsset, LocalAsset, ResolutionCallback); }
+	virtual uint32 GetCategories() override { return OriginalActions->GetCategories(); }
+	virtual FString GetObjectDisplayName(UObject* Object) const override { return OriginalActions->GetObjectDisplayName(Object); }
+	virtual const TArray<FText>& GetSubMenus() const override { return OriginalActions->GetSubMenus(); }
+	virtual bool ShouldForceWorldCentric() override { return OriginalActions->ShouldForceWorldCentric(); }
+	virtual void PerformAssetDiff(UObject* OldAsset, UObject* NewAsset, const FRevisionInfo& OldRevision, const FRevisionInfo& NewRevision) const override { OriginalActions->PerformAssetDiff(OldAsset, NewAsset, OldRevision, NewRevision); }
+	virtual UThumbnailInfo* GetThumbnailInfo(UObject* Asset) const override { return OriginalActions->GetThumbnailInfo(Asset); }
+	virtual EThumbnailPrimType GetDefaultThumbnailPrimitiveType(UObject* Asset) const override { return OriginalActions->GetDefaultThumbnailPrimitiveType(Asset); }
+	virtual TSharedPtr<SWidget> GetThumbnailOverlay(const FAssetData& AssetData) const override { return OriginalActions->GetThumbnailOverlay(AssetData); }
+	virtual FText GetAssetDescription(const FAssetData& AssetData) const override { return OriginalActions->GetAssetDescription(AssetData); }
+	virtual bool IsImportedAsset() const override { return OriginalActions->IsImportedAsset(); }
+	virtual void GetResolvedSourceFilePaths(const TArray<UObject*>& TypeAssets, TArray<FString>& OutSourceFilePaths) const override { OriginalActions->GetResolvedSourceFilePaths(TypeAssets, OutSourceFilePaths); }
+	virtual void GetSourceFileLabels(const TArray<UObject*>& TypeAssets, TArray<FString>& OutSourceFileLabels) const override { OriginalActions->GetSourceFileLabels(TypeAssets, OutSourceFileLabels); }
+	virtual void BuildBackendFilter(FARFilter& InFilter) override { OriginalActions->BuildBackendFilter(InFilter); }
+	virtual FText GetDisplayNameFromAssetData(const FAssetData& AssetData) const override { return OriginalActions->GetDisplayNameFromAssetData(AssetData); }
+	virtual void SetSupported(bool bInSupported) override { OriginalActions->SetSupported(bInSupported); }
+	virtual bool IsSupported() const override { return OriginalActions->IsSupported(); }
+
+private:
+
+	void GetPrebakeActions(const TArray<UObject*>& InObjects, FToolMenuSection& Section);
+
+	void OnPrebake(TArray<FWeakObjectPtr> Objects) const;
+
+	TSharedRef<IAssetTypeActions> OriginalActions;
+};
+
+#endif
diff --git a/Engine/Plugins/Enterprise/GLTFExporter/MigratedFromGit/Source/GLTFExporter/Private/GLTFExporterModule.cpp b/Engine/Plugins/Enterprise/GLTFExporter/MigratedFromGit/Source/GLTFExporter/Private/GLTFExporterModule.cpp
--- a/Engine/Plugins/Enterprise/GLTFExporter/MigratedFromGit/Source/GLTFExporter/Private/GLTFExporterModule.cpp
+++ b/Engine/Plugins/Enterprise/GLTFExporter/MigratedFromGit/Source/GLTFExporter/Private/GLTFExporterModule.cpp
@@ -1,27 +1,72 @@
 // Copyright Epic Games, Inc. All Rights Reserved.
 
 #include "GLTFExporterModule.h"
-
-#include "CoreMinimal.h"
-#include "Modules/ModuleManager.h"
+#include "Actions/GLTFPrebakeAssetActions.h"
 #include "Interfaces/IPluginManager.h"
+#if WITH_EDITOR
+#include "Materials/MaterialInstanceConstant.h"
+#include "Materials/MaterialInstanceDynamic.h"
+#include "AssetToolsModule.h"
+
+class FAssetTypeActions_Base;
+#endif
 
 DEFINE_LOG_CATEGORY(LogGLTFExporter);
 
-/**
- * glTF Exporter module implementation (private)
- */
 class FGLTFExporterModule final : public IGLTFExporterModule
 {
-
 public:
+
 	virtual void StartupModule() override
 	{
-		FString PluginShaderDir = FPaths::Combine(IPluginManager::Get().FindPlugin(TEXT("GLTFExporter"))->GetBaseDir(), TEXT("Shaders"));
+		const FString PluginShaderDir = FPaths::Combine(IPluginManager::Get().FindPlugin(TEXT("GLTFExporter"))->GetBaseDir(), TEXT("Shaders"));
 		AddShaderSourceDirectoryMapping(TEXT("/Plugin/GLTFExporter"), PluginShaderDir);
+
+		FCoreDelegates::OnPostEngineInit.AddRaw(this, &FGLTFExporterModule::PostEngineInit);
 	}
 
-	virtual void ShutdownModule() override {}
+	void PostEngineInit()
+	{
+#if WITH_EDITOR
+		if (FAssetToolsModule::IsModuleLoaded())
+		{
+			IAssetTools& AssetTools = FAssetToolsModule::GetModule().Get();
+
+			// TODO: add support for when GetAssetTypeActionsForClass returns null pointer
+			AssetTypeActionsArray.Add(MakeShared<FGLTFPrebakeAssetActions>(AssetTools.GetAssetTypeActionsForClass(UMaterial::StaticClass()).Pin().ToSharedRef()));
+			AssetTypeActionsArray.Add(MakeShared<FGLTFPrebakeAssetActions>(AssetTools.GetAssetTypeActionsForClass(UMaterialInstanceConstant::StaticClass()).Pin().ToSharedRef()));
+			AssetTypeActionsArray.Add(MakeShared<FGLTFPrebakeAssetActions>(AssetTools.GetAssetTypeActionsForClass(UMaterialInstanceDynamic::StaticClass()).Pin().ToSharedRef()));
+			AssetTypeActionsArray.Add(MakeShared<FGLTFPrebakeAssetActions>(AssetTools.GetAssetTypeActionsForClass(UMaterialInterface::StaticClass()).Pin().ToSharedRef()));
+
+			for (TSharedRef<IAssetTypeActions>& AssetTypeActions : AssetTypeActionsArray)
+			{
+				AssetTools.RegisterAssetTypeActions(AssetTypeActions);
+			}
+		}
+#endif
+	}
+
+	virtual void ShutdownModule() override
+	{
+#if WITH_EDITOR
+		if (FAssetToolsModule::IsModuleLoaded())
+		{
+			IAssetTools& AssetTools = FAssetToolsModule::GetModule().Get();
+			for (TSharedRef<IAssetTypeActions>& AssetTypeActions : AssetTypeActionsArray)
+			{
+				AssetTools.UnregisterAssetTypeActions(AssetTypeActions);
+			}
+		}
+
+		AssetTypeActionsArray.Empty();
+#endif
+	}
+
+private:
+
+#if WITH_EDITOR
+	TArray<TSharedRef<IAssetTypeActions>> AssetTypeActionsArray;
+#endif
 };
 
 IMPLEMENT_MODULE(FGLTFExporterModule, GLTFExporter);
