[GLTFExporter] refactor: move GLTFExportOptions.h to Options/GLTFExportOptions.h

commit 05a296a055b1cb247cd173fc4c4c615fa8fd711b
Date: Thu, 27 Jan 2022 19:06:52 +0100

    refactor: move GLTFExportOptions.h to Options/GLTFExportOptions.h
    (cherry picked from commit 58bb988dae1e2e11485d13409c452dc6bcffb4ca​)

#skipci

diff --git a/Engine/Plugins/Enterprise/GLTFExporter/MigratedFromGit/Source/GLTFExporter/Private/Builders/GLTFBuilder.h b/Engine/Plugins/Enterprise/GLTFExporter/MigratedFromGit/Source/GLTFExporter/Private/Builders/GLTFBuilder.h
--- a/Engine/Plugins/Enterprise/GLTFExporter/MigratedFromGit/Source/GLTFExporter/Private/Builders/GLTFBuilder.h
+++ b/Engine/Plugins/Enterprise/GLTFExporter/MigratedFromGit/Source/GLTFExporter/Private/Builders/GLTFBuilder.h
@@ -2,7 +2,7 @@
 
 #pragma once
 
-#include "GLTFExportOptions.h"
+#include "Options/GLTFExportOptions.h"
 #include "Json/GLTFJsonEnums.h"
 #include "UObject/GCObjectScopeGuard.h"
 
diff --git a/Engine/Plugins/Enterprise/GLTFExporter/MigratedFromGit/Source/GLTFExporter/Private/Converters/GLTFImageConverters.h b/Engine/Plugins/Enterprise/GLTFExporter/MigratedFromGit/Source/GLTFExporter/Private/Converters/GLTFImageConverters.h
--- a/Engine/Plugins/Enterprise/GLTFExporter/MigratedFromGit/Source/GLTFExporter/Private/Converters/GLTFImageConverters.h
+++ b/Engine/Plugins/Enterprise/GLTFExporter/MigratedFromGit/Source/GLTFExporter/Private/Converters/GLTFImageConverters.h
@@ -8,7 +8,7 @@
 #include "Converters/GLTFBuilderContext.h"
 #include "Converters/GLTFSuperfluous.h"
 #include "Converters/GLTFSharedArray.h"
-#include "GLTFExportOptions.h"
+#include "Options/GLTFExportOptions.h"
 
 typedef TGLTFConverter<FGLTFJsonImageIndex, TGLTFSuperfluous<FString>, EGLTFTextureType, bool, FIntPoint, TGLTFSharedArray<FColor>> IGLTFImageConverter;
 
diff --git a/Engine/Plugins/Enterprise/GLTFExporter/MigratedFromGit/Source/GLTFExporter/Private/Exporters/GLTFExporter.cpp b/Engine/Plugins/Enterprise/GLTFExporter/MigratedFromGit/Source/GLTFExporter/Private/Exporters/GLTFExporter.cpp
--- a/Engine/Plugins/Enterprise/GLTFExporter/MigratedFromGit/Source/GLTFExporter/Private/Exporters/GLTFExporter.cpp
+++ b/Engine/Plugins/Enterprise/GLTFExporter/MigratedFromGit/Source/GLTFExporter/Private/Exporters/GLTFExporter.cpp
@@ -1,7 +1,7 @@
 // Copyright Epic Games, Inc. All Rights Reserved.
 
 #include "Exporters/GLTFExporter.h"
-#include "GLTFExportOptions.h"
+#include "Options/GLTFExportOptions.h"
 #include "UI/GLTFExportOptionsWindow.h"
 #include "Builders/GLTFContainerBuilder.h"
 #include "UObject/GCObjectScopeGuard.h"
diff --git a/Engine/Plugins/Enterprise/GLTFExporter/MigratedFromGit/Source/GLTFExporter/Private/GLTFExportOptions.cpp b/Engine/Plugins/Enterprise/GLTFExporter/MigratedFromGit/Source/GLTFExporter/Private/GLTFExportOptions.cpp
deleted file mode 100644
diff --git a/Engine/Plugins/Enterprise/GLTFExporter/MigratedFromGit/Source/GLTFExporter/Private/Options/GLTFExportOptions.cpp b/Engine/Plugins/Enterprise/GLTFExporter/MigratedFromGit/Source/GLTFExporter/Private/Options/GLTFExportOptions.cpp
new file mode 100644
--- /dev/null
+++ b/Engine/Plugins/Enterprise/GLTFExporter/MigratedFromGit/Source/GLTFExporter/Private/Options/GLTFExportOptions.cpp
@@ -0,0 +1,50 @@
+// Copyright Epic Games, Inc. All Rights Reserved.
+
+#include "Options/GLTFExportOptions.h"
+
+UGLTFExportOptions::UGLTFExportOptions(const FObjectInitializer& ObjectInitializer)
+	: Super(ObjectInitializer)
+{
+	ResetToDefault();
+}
+
+void UGLTFExportOptions::ResetToDefault()
+{
+	ExportUniformScale = 0.01;
+	bExportPreviewMesh = true;
+	bStrictCompliance = true;
+	bSkipNearDefaultValues = true;
+	bIncludeGeneratorVersion = true;
+	bExportProxyMaterials = true;
+	bExportUnlitMaterials = true;
+	bExportClearCoatMaterials = true;
+	bExportExtraBlendModes = false;
+	BakeMaterialInputs = EGLTFMaterialBakeMode::UseMeshData;
+	DefaultMaterialBakeSize = EGLTFMaterialBakeSizePOT::POT_1024;
+	DefaultMaterialBakeFilter = TF_Trilinear;
+	DefaultMaterialBakeTiling = TA_Wrap;
+	DefaultLevelOfDetail = 0;
+	bExportVertexColors = false;
+	bExportVertexSkinWeights = true;
+	bUseMeshQuantization = false;
+	bExportLevelSequences = true;
+	bExportAnimationSequences = true;
+	bExportPlaybackSettings = false;
+	TextureImageFormat = EGLTFTextureImageFormat::PNG;
+	TextureImageQuality = 0;
+	NoLossyImageFormatFor = static_cast<int32>(EGLTFTextureType::All);
+	bExportTextureTransforms = true;
+	bExportLightmaps = false;
+	TextureHDREncoding = EGLTFTextureHDREncoding::None;
+	bExportHiddenInGame = false;
+	ExportLights = static_cast<int32>(EGLTFSceneMobility::Stationary | EGLTFSceneMobility::Movable);
+	bExportCameras = true;
+	bExportCameraControls = false;
+	bExportAnimationHotspots = false;
+	bExportHDRIBackdrops = false;
+	bExportSkySpheres = false;
+	VariantSetsMode = EGLTFVariantSetsMode::None;
+	ExportMaterialVariants = EGLTFMaterialVariantMode::UseMeshData;
+	bExportMeshVariants = true;
+	bExportVisibilityVariants = true;
+}
diff --git a/Engine/Plugins/Enterprise/GLTFExporter/MigratedFromGit/Source/GLTFExporter/Private/UI/GLTFExportOptionsWindow.h b/Engine/Plugins/Enterprise/GLTFExporter/MigratedFromGit/Source/GLTFExporter/Private/UI/GLTFExportOptionsWindow.h
--- a/Engine/Plugins/Enterprise/GLTFExporter/MigratedFromGit/Source/GLTFExporter/Private/UI/GLTFExportOptionsWindow.h
+++ b/Engine/Plugins/Enterprise/GLTFExporter/MigratedFromGit/Source/GLTFExporter/Private/UI/GLTFExportOptionsWindow.h
@@ -4,7 +4,7 @@
 
 #if WITH_EDITOR
 
-#include "GLTFExportOptions.h"
+#include "Options/GLTFExportOptions.h"
 #include "Widgets/DeclarativeSyntaxSupport.h"
 #include "Input/Reply.h"
 #include "Widgets/SCompoundWidget.h"
diff --git a/Engine/Plugins/Enterprise/GLTFExporter/MigratedFromGit/Source/GLTFExporter/Public/GLTFExportOptions.h b/Engine/Plugins/Enterprise/GLTFExporter/MigratedFromGit/Source/GLTFExporter/Public/GLTFExportOptions.h
--- a/Engine/Plugins/Enterprise/GLTFExporter/MigratedFromGit/Source/GLTFExporter/Public/GLTFExportOptions.h
+++ b/Engine/Plugins/Enterprise/GLTFExporter/MigratedFromGit/Source/GLTFExporter/Public/GLTFExportOptions.h
@@ -2,235 +2,6 @@
 
 #pragma once
 
-#include "CoreMinimal.h"
-#include "UserData/GLTFMaterialUserData.h"
-#include "GLTFExportOptions.generated.h"
+COMPILE_WARNING("\"GLTFExportOptions.h\" has been moved into subfolder \"Options\", please use #include \"Options/GLTFExportOptions.h\"")
 
-UENUM(BlueprintType)
-enum class EGLTFTextureImageFormat : uint8
-{
-	/** Don't export any textures. */
-	None,
-	/** Always use PNG (lossless compression). */
-	PNG,
-	/** If texture does not have an alpha channel, use JPEG (lossy compression); otherwise fallback to PNG. */
-	JPEG UMETA(DisplayName = "JPEG (if no alpha)")
-};
-
-UENUM(BlueprintType, Meta = (Bitflags, UseEnumValuesAsMaskValuesInEditor = "true"))
-enum class EGLTFTextureType : uint8
-{
-	None = 0 UMETA(Hidden),
-
-	HDR = 1 << 0,
-	Normalmaps = 1 << 1,
-	Lightmaps = 1 << 2,
-
-	All = HDR | Normalmaps | Lightmaps UMETA(Hidden)
-};
-ENUM_CLASS_FLAGS(EGLTFTextureType);
-
-UENUM(BlueprintType)
-enum class EGLTFTextureHDREncoding : uint8
-{
-	/** Clamp HDR colors to standard 8-bit per channel. */
-	None,
-	/** Encode HDR colors to RGBM (will discard alpha). */
-	RGBM
-};
-
-UENUM(BlueprintType, Meta = (Bitflags, UseEnumValuesAsMaskValuesInEditor = "true"))
-enum class EGLTFSceneMobility : uint8
-{
-	None = 0 UMETA(Hidden),
-
-	Static = 1 << 0,
-	Stationary = 1 << 1,
-	Movable = 1 << 2,
-
-	All = Static | Stationary | Movable UMETA(Hidden)
-};
-ENUM_CLASS_FLAGS(EGLTFSceneMobility);
-
-UENUM(BlueprintType)
-enum class EGLTFVariantSetsMode : uint8
-{
-	/** Never export variants sets. */
-	None,
-	/** Uses the official extension KHR_materials_variants. Supports material variants only. */
-	Khronos,
-	/** Uses the extension EPIC_level_variant_sets, which is supported by Unreal's glTF viewer. */
-	Epic
-};
-
-UENUM(BlueprintType)
-enum class EGLTFMaterialVariantMode : uint8
-{
-	/** Never export material variants. */
-	None,
-	/** Export material variants but only use a simple quad if a material input needs to be baked out. */
-	Simple,
-	/** Export material variants and allow usage of the mesh data if a material input needs to be baked out with vertex data. */
-	UseMeshData,
-};
-
-UCLASS(BlueprintType, Config=EditorPerProjectUserSettings, HideCategories=(DebugProperty))
-class GLTFEXPORTER_API UGLTFExportOptions : public UObject
-{
-	GENERATED_UCLASS_BODY()
-
-	/** Scale factor used for exporting all assets (0.01 by default) for conversion from centimeters (Unreal default) to meters (glTF). */
-	UPROPERTY(EditAnywhere, BlueprintReadWrite, Config, Category = General)
-	float ExportUniformScale;
-
-	/** If enabled, the preview mesh for a standalone animation or material asset will also be exported. */
-	UPROPERTY(EditAnywhere, BlueprintReadWrite, Config, Category = General)
-	bool bExportPreviewMesh;
-
-	/** If enabled, certain values (like HDR colors and light angles) will be truncated during export to strictly conform to the formal glTF specification. */
-	UPROPERTY(EditAnywhere, BlueprintReadWrite, Config, Category = General)
-	bool bStrictCompliance;
-
-	/** If enabled, floating-point-based JSON properties that are nearly equal to their default value will not be exported and thus regarded as exactly default, reducing JSON size. */
-	UPROPERTY(EditAnywhere, BlueprintReadWrite, Config, Category = General)
-	bool bSkipNearDefaultValues;
-
-	/** If enabled, version info for Unreal Engine and exporter plugin will be included as metadata in the glTF asset, which is useful when reporting issues. */
-	UPROPERTY(EditAnywhere, BlueprintReadWrite, Config, Category = General)
-	bool bIncludeGeneratorVersion;
-
-	/** If enabled, materials that have a proxy defined in their user data, will be exported using that proxy instead. Used mainly by pre-baked materials. */
-	UPROPERTY(EditAnywhere, BlueprintReadWrite, Config, Category = Material)
-	bool bExportProxyMaterials;
-
-	/** If enabled, materials with shading model unlit will be properly exported. Uses extension KHR_materials_unlit. */
-	UPROPERTY(EditAnywhere, BlueprintReadWrite, Config, Category = Material)
-	bool bExportUnlitMaterials;
-
-	/** If enabled, materials with shading model clear coat will be properly exported. Uses extension KHR_materials_clearcoat, which is not supported by all glTF viewers. */
-	UPROPERTY(EditAnywhere, BlueprintReadWrite, Config, Category = Material)
-	bool bExportClearCoatMaterials;
-
-	/** If enabled, materials with blend modes additive, modulate, and alpha composite will be properly exported. Uses extension EPIC_blend_modes, which is supported by Unreal's glTF viewer. */
-	UPROPERTY(EditAnywhere, BlueprintReadWrite, Config, Category = Material)
-	bool bExportExtraBlendModes;
-
-	/** Bake mode determining if and how a material input is baked out to a texture. Can be overridden by material- and input-specific bake settings, see GLTFMaterialExportOptions. */
-	UPROPERTY(EditAnywhere, BlueprintReadWrite, Config, Category = Material)
-	EGLTFMaterialBakeMode BakeMaterialInputs;
-
-	/** Default size of the baked out texture (containing the material input). Can be overridden by material- and input-specific bake settings, see GLTFMaterialExportOptions. */
-	UPROPERTY(EditAnywhere, BlueprintReadWrite, Config, Category = Material, Meta = (EditCondition = "BakeMaterialInputs != EGLTFMaterialBakeMode::Disabled"))
-	EGLTFMaterialBakeSizePOT DefaultMaterialBakeSize;
-
-	/** Default filtering mode used when sampling the baked out texture. Can be overridden by material- and input-specific bake settings, see GLTFMaterialExportOptions. */
-	UPROPERTY(EditAnywhere, BlueprintReadWrite, Config, Category = Material, Meta = (EditCondition = "BakeMaterialInputs != EGLTFMaterialBakeMode::Disabled", ValidEnumValues="TF_Nearest, TF_Bilinear, TF_Trilinear"))
-	TEnumAsByte<TextureFilter> DefaultMaterialBakeFilter;
-
-	/** Default addressing mode used when sampling the baked out texture. Can be overridden by material- and input-specific bake settings, see GLTFMaterialExportOptions. */
-	UPROPERTY(EditAnywhere, BlueprintReadWrite, Config, Category = Material, Meta = (EditCondition = "BakeMaterialInputs != EGLTFMaterialBakeMode::Disabled"))
-	TEnumAsByte<TextureAddress> DefaultMaterialBakeTiling;
-
-	/** Input-specific default bake settings that override the general defaults above. */
-	UPROPERTY(EditAnywhere, BlueprintReadWrite, Config, Category = Material)
-	TMap<EGLTFMaterialPropertyGroup, FGLTFOverrideMaterialBakeSettings> DefaultInputBakeSettings;
-
-	/** Default LOD level used for exporting a mesh. Can be overridden by component or asset settings (e.g. minimum or forced LOD level). */
-	UPROPERTY(EditAnywhere, BlueprintReadWrite, Config, Category = Mesh, Meta = (ClampMin = "0"))
-	int32 DefaultLevelOfDetail;
-
-	/** If enabled, export vertex color. Not recommended due to vertex colors always being used as a base color multiplier in glTF, regardless of material. Often producing undesirable results. */
-	UPROPERTY(EditAnywhere, BlueprintReadWrite, Config, Category = Mesh)
-	bool bExportVertexColors;
-
-	/** If enabled, export vertex bone weights and indices in skeletal meshes. Necessary for animation sequences. */
-	UPROPERTY(EditAnywhere, BlueprintReadWrite, Config, Category = Mesh)
-	bool bExportVertexSkinWeights;
-
-	/** If enabled, use quantization for vertex tangents and normals, reducing size. Requires extension KHR_mesh_quantization, which may result in the mesh not loading in some glTF viewers. */
-	UPROPERTY(EditAnywhere, BlueprintReadWrite, Config, Category = Mesh)
-	bool bUseMeshQuantization;
-
-	/** If enabled, export level sequences. Only transform tracks are currently supported. The level sequence will be played at the assigned display rate. */
-	UPROPERTY(EditAnywhere, BlueprintReadWrite, Config, Category = Animation)
-	bool bExportLevelSequences;
-
-	/** If enabled, export single animation asset used by a skeletal mesh component or hotspot actor. Export of vertex skin weights must be enabled. */
-	UPROPERTY(EditAnywhere, BlueprintReadWrite, Config, Category = Animation, Meta = (EditCondition = "bExportVertexSkinWeights"))
-	bool bExportAnimationSequences;
-
-	/** If enabled, export play rate, start time, looping, and auto play for an animation or level sequence. Uses extension EPIC_animation_playback, which is supported by Unreal's glTF viewer. */
-	UPROPERTY(EditAnywhere, BlueprintReadWrite, Config, Category = Animation)
-	bool bExportPlaybackSettings;
-
-	/** Desired image format used for exported textures. */
-	UPROPERTY(EditAnywhere, BlueprintReadWrite, Config, Category = Texture)
-	EGLTFTextureImageFormat TextureImageFormat;
-
-	/** Level of compression used for textures exported with lossy image formats, 0 (default) or value between 1 (worst quality, best compression) and 100 (best quality, worst compression). */
-	UPROPERTY(EditAnywhere, BlueprintReadWrite, Config, Category = Texture, Meta = (ClampMin = "0", ClampMax = "100", EditCondition = "TextureImageFormat == EGLTFTextureImageFormat::JPEG"))
-	int32 TextureImageQuality;
-
-	/** Texture types that will always use lossless formats (e.g. PNG) because of sensitivity to compression artifacts. */
-	UPROPERTY(EditAnywhere, BlueprintReadWrite, Config, Category = Texture, Meta = (Bitmask, BitmaskEnum = EGLTFTextureType, EditCondition = "TextureImageFormat == EGLTFTextureImageFormat::JPEG"))
-	int32 NoLossyImageFormatFor; // Bitmask combined from EGLTFTextureType
-
-	/** If enabled, export UV tiling and un-mirroring settings in a texture coordinate expression node for simple material input expressions. Uses extension KHR_texture_transform. */
-	UPROPERTY(EditAnywhere, BlueprintReadWrite, Config, Category = Texture, Meta = (EditCondition = "TextureImageFormat != EGLTFTextureImageFormat::None"))
-	bool bExportTextureTransforms;
-
-	/** If enabled, export lightmaps (created by Lightmass) when exporting a level. Uses extension EPIC_lightmap_textures, which is supported by Unreal's glTF viewer. */
-	UPROPERTY(EditAnywhere, BlueprintReadWrite, Config, Category = Texture, Meta = (EditCondition = "TextureImageFormat != EGLTFTextureImageFormat::None"))
-	bool bExportLightmaps;
-
-	/** Encoding used to store textures that have pixel colors with more than 8-bit per channel. Uses extension EPIC_texture_hdr_encoding, which is supported by Unreal's glTF viewer. */
-	UPROPERTY(EditAnywhere, BlueprintReadWrite, Config, Category = Texture, Meta = (DisplayName = "Texture HDR Encoding", EditCondition = "TextureImageFormat != EGLTFTextureImageFormat::None"))
-	EGLTFTextureHDREncoding TextureHDREncoding;
-
-	/** If enabled, export actors and components that are flagged as hidden in-game. */
-	UPROPERTY(EditAnywhere, BlueprintReadWrite, Config, Category = Scene)
-	bool bExportHiddenInGame;
-
-	/** Mobility of directional, point, and spot light components that will be exported. Uses extension KHR_lights_punctual. */
-	UPROPERTY(EditAnywhere, BlueprintReadWrite, Config, Category = Scene, Meta = (Bitmask, BitmaskEnum = EGLTFSceneMobility))
-	int32 ExportLights; // Bitmask combined from EGLTFSceneMobility
-
-	/** If enabled, export camera components. */
-	UPROPERTY(EditAnywhere, BlueprintReadWrite, Config, Category = Scene)
-	bool bExportCameras;
-
-	/** If enabled, export GLTFCameraActors. Uses extension EPIC_camera_controls, which is supported by Unreal's glTF viewer. */
-	UPROPERTY(EditAnywhere, BlueprintReadWrite, Config, Category = Scene, Meta = (EditCondition = "bExportCameras"))
-	bool bExportCameraControls;
-
-	/** If enabled, export GLTFHotspotActors. Uses extension EPIC_animation_hotspots, which is supported by Unreal's glTF viewer. */
-	UPROPERTY(EditAnywhere, BlueprintReadWrite, Config, Category = Scene)
-	bool bExportAnimationHotspots;
-
-	/** If enabled, export HDRIBackdrop blueprints. Uses extension EPIC_hdri_backdrops, which is supported by Unreal's glTF viewer. */
-	UPROPERTY(EditAnywhere, BlueprintReadWrite, Config, Category = Scene, Meta = (DisplayName = "Export HDRI Backdrops"))
-	bool bExportHDRIBackdrops;
-
-	/** If enabled, export SkySphere blueprints. Uses extension EPIC_sky_spheres, which is supported by Unreal's glTF viewer. */
-	UPROPERTY(EditAnywhere, BlueprintReadWrite, Config, Category = Scene)
-	bool bExportSkySpheres;
-
-	/** Mode determining if and how to export LevelVariantSetsActors. */
-	UPROPERTY(EditAnywhere, BlueprintReadWrite, Config, Category = Scene)
-	EGLTFVariantSetsMode VariantSetsMode;
-
-	/** Mode determining if and how to export material variants that change the materials property on a static or skeletal mesh component. */
-	UPROPERTY(EditAnywhere, BlueprintReadWrite, Config, Category = VariantSets, Meta = (EditCondition = "VariantSetsMode != EGLTFVariantSetsMode::None"))
-	EGLTFMaterialVariantMode ExportMaterialVariants;
-
-	/** If enabled, export variants that change the mesh property on a static or skeletal mesh component. */
-	UPROPERTY(EditAnywhere, BlueprintReadWrite, Config, Category = VariantSets, Meta = (EditCondition = "VariantSetsMode == EGLTFVariantSetsMode::Epic"))
-	bool bExportMeshVariants;
-
-	/** If enabled, export variants that change the visible property on a scene component. */
-	UPROPERTY(EditAnywhere, BlueprintReadWrite, Config, Category = VariantSets, Meta = (EditCondition = "VariantSetsMode == EGLTFVariantSetsMode::Epic"))
-	bool bExportVisibilityVariants;
-
-	UFUNCTION(BlueprintCallable, Category = General)
-	void ResetToDefault();
-};
+#include "Options/GLTFExportOptions.h"
diff --git a/Engine/Plugins/Enterprise/GLTFExporter/MigratedFromGit/Source/GLTFExporter/Public/Options/GLTFExportOptions.h b/Engine/Plugins/Enterprise/GLTFExporter/MigratedFromGit/Source/GLTFExporter/Public/Options/GLTFExportOptions.h
new file mode 100644
--- /dev/null
+++ b/Engine/Plugins/Enterprise/GLTFExporter/MigratedFromGit/Source/GLTFExporter/Public/Options/GLTFExportOptions.h
@@ -0,0 +1,236 @@
+// Copyright Epic Games, Inc. All Rights Reserved.
+
+#pragma once
+
+#include "CoreMinimal.h"
+#include "UserData/GLTFMaterialUserData.h"
+#include "GLTFExportOptions.generated.h"
+
+UENUM(BlueprintType)
+enum class EGLTFTextureImageFormat : uint8
+{
+	/** Don't export any textures. */
+	None,
+	/** Always use PNG (lossless compression). */
+	PNG,
+	/** If texture does not have an alpha channel, use JPEG (lossy compression); otherwise fallback to PNG. */
+	JPEG UMETA(DisplayName = "JPEG (if no alpha)")
+};
+
+UENUM(BlueprintType, Meta = (Bitflags, UseEnumValuesAsMaskValuesInEditor = "true"))
+enum class EGLTFTextureType : uint8
+{
+	None = 0 UMETA(Hidden),
+
+	HDR = 1 << 0,
+	Normalmaps = 1 << 1,
+	Lightmaps = 1 << 2,
+
+	All = HDR | Normalmaps | Lightmaps UMETA(Hidden)
+};
+ENUM_CLASS_FLAGS(EGLTFTextureType);
+
+UENUM(BlueprintType)
+enum class EGLTFTextureHDREncoding : uint8
+{
+	/** Clamp HDR colors to standard 8-bit per channel. */
+	None,
+	/** Encode HDR colors to RGBM (will discard alpha). */
+	RGBM
+};
+
+UENUM(BlueprintType, Meta = (Bitflags, UseEnumValuesAsMaskValuesInEditor = "true"))
+enum class EGLTFSceneMobility : uint8
+{
+	None = 0 UMETA(Hidden),
+
+	Static = 1 << 0,
+	Stationary = 1 << 1,
+	Movable = 1 << 2,
+
+	All = Static | Stationary | Movable UMETA(Hidden)
+};
+ENUM_CLASS_FLAGS(EGLTFSceneMobility);
+
+UENUM(BlueprintType)
+enum class EGLTFVariantSetsMode : uint8
+{
+	/** Never export variants sets. */
+	None,
+	/** Uses the official extension KHR_materials_variants. Supports material variants only. */
+	Khronos,
+	/** Uses the extension EPIC_level_variant_sets, which is supported by Unreal's glTF viewer. */
+	Epic
+};
+
+UENUM(BlueprintType)
+enum class EGLTFMaterialVariantMode : uint8
+{
+	/** Never export material variants. */
+	None,
+	/** Export material variants but only use a simple quad if a material input needs to be baked out. */
+	Simple,
+	/** Export material variants and allow usage of the mesh data if a material input needs to be baked out with vertex data. */
+	UseMeshData,
+};
+
+UCLASS(BlueprintType, Config=EditorPerProjectUserSettings, HideCategories=(DebugProperty))
+class GLTFEXPORTER_API UGLTFExportOptions : public UObject
+{
+	GENERATED_UCLASS_BODY()
+
+	/** Scale factor used for exporting all assets (0.01 by default) for conversion from centimeters (Unreal default) to meters (glTF). */
+	UPROPERTY(EditAnywhere, BlueprintReadWrite, Config, Category = General)
+	float ExportUniformScale;
+
+	/** If enabled, the preview mesh for a standalone animation or material asset will also be exported. */
+	UPROPERTY(EditAnywhere, BlueprintReadWrite, Config, Category = General)
+	bool bExportPreviewMesh;
+
+	/** If enabled, certain values (like HDR colors and light angles) will be truncated during export to strictly conform to the formal glTF specification. */
+	UPROPERTY(EditAnywhere, BlueprintReadWrite, Config, Category = General)
+	bool bStrictCompliance;
+
+	/** If enabled, floating-point-based JSON properties that are nearly equal to their default value will not be exported and thus regarded as exactly default, reducing JSON size. */
+	UPROPERTY(EditAnywhere, BlueprintReadWrite, Config, Category = General)
+	bool bSkipNearDefaultValues;
+
+	/** If enabled, version info for Unreal Engine and exporter plugin will be included as metadata in the glTF asset, which is useful when reporting issues. */
+	UPROPERTY(EditAnywhere, BlueprintReadWrite, Config, Category = General)
+	bool bIncludeGeneratorVersion;
+
+	/** If enabled, materials that have a proxy defined in their user data, will be exported using that proxy instead. Used mainly by pre-baked materials. */
+	UPROPERTY(EditAnywhere, BlueprintReadWrite, Config, Category = Material)
+	bool bExportProxyMaterials;
+
+	/** If enabled, materials with shading model unlit will be properly exported. Uses extension KHR_materials_unlit. */
+	UPROPERTY(EditAnywhere, BlueprintReadWrite, Config, Category = Material)
+	bool bExportUnlitMaterials;
+
+	/** If enabled, materials with shading model clear coat will be properly exported. Uses extension KHR_materials_clearcoat, which is not supported by all glTF viewers. */
+	UPROPERTY(EditAnywhere, BlueprintReadWrite, Config, Category = Material)
+	bool bExportClearCoatMaterials;
+
+	/** If enabled, materials with blend modes additive, modulate, and alpha composite will be properly exported. Uses extension EPIC_blend_modes, which is supported by Unreal's glTF viewer. */
+	UPROPERTY(EditAnywhere, BlueprintReadWrite, Config, Category = Material)
+	bool bExportExtraBlendModes;
+
+	/** Bake mode determining if and how a material input is baked out to a texture. Can be overridden by material- and input-specific bake settings, see GLTFMaterialExportOptions. */
+	UPROPERTY(EditAnywhere, BlueprintReadWrite, Config, Category = Material)
+	EGLTFMaterialBakeMode BakeMaterialInputs;
+
+	/** Default size of the baked out texture (containing the material input). Can be overridden by material- and input-specific bake settings, see GLTFMaterialExportOptions. */
+	UPROPERTY(EditAnywhere, BlueprintReadWrite, Config, Category = Material, Meta = (EditCondition = "BakeMaterialInputs != EGLTFMaterialBakeMode::Disabled"))
+	EGLTFMaterialBakeSizePOT DefaultMaterialBakeSize;
+
+	/** Default filtering mode used when sampling the baked out texture. Can be overridden by material- and input-specific bake settings, see GLTFMaterialExportOptions. */
+	UPROPERTY(EditAnywhere, BlueprintReadWrite, Config, Category = Material, Meta = (EditCondition = "BakeMaterialInputs != EGLTFMaterialBakeMode::Disabled", ValidEnumValues="TF_Nearest, TF_Bilinear, TF_Trilinear"))
+	TEnumAsByte<TextureFilter> DefaultMaterialBakeFilter;
+
+	/** Default addressing mode used when sampling the baked out texture. Can be overridden by material- and input-specific bake settings, see GLTFMaterialExportOptions. */
+	UPROPERTY(EditAnywhere, BlueprintReadWrite, Config, Category = Material, Meta = (EditCondition = "BakeMaterialInputs != EGLTFMaterialBakeMode::Disabled"))
+	TEnumAsByte<TextureAddress> DefaultMaterialBakeTiling;
+
+	/** Input-specific default bake settings that override the general defaults above. */
+	UPROPERTY(EditAnywhere, BlueprintReadWrite, Config, Category = Material)
+	TMap<EGLTFMaterialPropertyGroup, FGLTFOverrideMaterialBakeSettings> DefaultInputBakeSettings;
+
+	/** Default LOD level used for exporting a mesh. Can be overridden by component or asset settings (e.g. minimum or forced LOD level). */
+	UPROPERTY(EditAnywhere, BlueprintReadWrite, Config, Category = Mesh, Meta = (ClampMin = "0"))
+	int32 DefaultLevelOfDetail;
+
+	/** If enabled, export vertex color. Not recommended due to vertex colors always being used as a base color multiplier in glTF, regardless of material. Often producing undesirable results. */
+	UPROPERTY(EditAnywhere, BlueprintReadWrite, Config, Category = Mesh)
+	bool bExportVertexColors;
+
+	/** If enabled, export vertex bone weights and indices in skeletal meshes. Necessary for animation sequences. */
+	UPROPERTY(EditAnywhere, BlueprintReadWrite, Config, Category = Mesh)
+	bool bExportVertexSkinWeights;
+
+	/** If enabled, use quantization for vertex tangents and normals, reducing size. Requires extension KHR_mesh_quantization, which may result in the mesh not loading in some glTF viewers. */
+	UPROPERTY(EditAnywhere, BlueprintReadWrite, Config, Category = Mesh)
+	bool bUseMeshQuantization;
+
+	/** If enabled, export level sequences. Only transform tracks are currently supported. The level sequence will be played at the assigned display rate. */
+	UPROPERTY(EditAnywhere, BlueprintReadWrite, Config, Category = Animation)
+	bool bExportLevelSequences;
+
+	/** If enabled, export single animation asset used by a skeletal mesh component or hotspot actor. Export of vertex skin weights must be enabled. */
+	UPROPERTY(EditAnywhere, BlueprintReadWrite, Config, Category = Animation, Meta = (EditCondition = "bExportVertexSkinWeights"))
+	bool bExportAnimationSequences;
+
+	/** If enabled, export play rate, start time, looping, and auto play for an animation or level sequence. Uses extension EPIC_animation_playback, which is supported by Unreal's glTF viewer. */
+	UPROPERTY(EditAnywhere, BlueprintReadWrite, Config, Category = Animation)
+	bool bExportPlaybackSettings;
+
+	/** Desired image format used for exported textures. */
+	UPROPERTY(EditAnywhere, BlueprintReadWrite, Config, Category = Texture)
+	EGLTFTextureImageFormat TextureImageFormat;
+
+	/** Level of compression used for textures exported with lossy image formats, 0 (default) or value between 1 (worst quality, best compression) and 100 (best quality, worst compression). */
+	UPROPERTY(EditAnywhere, BlueprintReadWrite, Config, Category = Texture, Meta = (ClampMin = "0", ClampMax = "100", EditCondition = "TextureImageFormat == EGLTFTextureImageFormat::JPEG"))
+	int32 TextureImageQuality;
+
+	/** Texture types that will always use lossless formats (e.g. PNG) because of sensitivity to compression artifacts. */
+	UPROPERTY(EditAnywhere, BlueprintReadWrite, Config, Category = Texture, Meta = (Bitmask, BitmaskEnum = EGLTFTextureType, EditCondition = "TextureImageFormat == EGLTFTextureImageFormat::JPEG"))
+	int32 NoLossyImageFormatFor; // Bitmask combined from EGLTFTextureType
+
+	/** If enabled, export UV tiling and un-mirroring settings in a texture coordinate expression node for simple material input expressions. Uses extension KHR_texture_transform. */
+	UPROPERTY(EditAnywhere, BlueprintReadWrite, Config, Category = Texture, Meta = (EditCondition = "TextureImageFormat != EGLTFTextureImageFormat::None"))
+	bool bExportTextureTransforms;
+
+	/** If enabled, export lightmaps (created by Lightmass) when exporting a level. Uses extension EPIC_lightmap_textures, which is supported by Unreal's glTF viewer. */
+	UPROPERTY(EditAnywhere, BlueprintReadWrite, Config, Category = Texture, Meta = (EditCondition = "TextureImageFormat != EGLTFTextureImageFormat::None"))
+	bool bExportLightmaps;
+
+	/** Encoding used to store textures that have pixel colors with more than 8-bit per channel. Uses extension EPIC_texture_hdr_encoding, which is supported by Unreal's glTF viewer. */
+	UPROPERTY(EditAnywhere, BlueprintReadWrite, Config, Category = Texture, Meta = (DisplayName = "Texture HDR Encoding", EditCondition = "TextureImageFormat != EGLTFTextureImageFormat::None"))
+	EGLTFTextureHDREncoding TextureHDREncoding;
+
+	/** If enabled, export actors and components that are flagged as hidden in-game. */
+	UPROPERTY(EditAnywhere, BlueprintReadWrite, Config, Category = Scene)
+	bool bExportHiddenInGame;
+
+	/** Mobility of directional, point, and spot light components that will be exported. Uses extension KHR_lights_punctual. */
+	UPROPERTY(EditAnywhere, BlueprintReadWrite, Config, Category = Scene, Meta = (Bitmask, BitmaskEnum = EGLTFSceneMobility))
+	int32 ExportLights; // Bitmask combined from EGLTFSceneMobility
+
+	/** If enabled, export camera components. */
+	UPROPERTY(EditAnywhere, BlueprintReadWrite, Config, Category = Scene)
+	bool bExportCameras;
+
+	/** If enabled, export GLTFCameraActors. Uses extension EPIC_camera_controls, which is supported by Unreal's glTF viewer. */
+	UPROPERTY(EditAnywhere, BlueprintReadWrite, Config, Category = Scene, Meta = (EditCondition = "bExportCameras"))
+	bool bExportCameraControls;
+
+	/** If enabled, export GLTFHotspotActors. Uses extension EPIC_animation_hotspots, which is supported by Unreal's glTF viewer. */
+	UPROPERTY(EditAnywhere, BlueprintReadWrite, Config, Category = Scene)
+	bool bExportAnimationHotspots;
+
+	/** If enabled, export HDRIBackdrop blueprints. Uses extension EPIC_hdri_backdrops, which is supported by Unreal's glTF viewer. */
+	UPROPERTY(EditAnywhere, BlueprintReadWrite, Config, Category = Scene, Meta = (DisplayName = "Export HDRI Backdrops"))
+	bool bExportHDRIBackdrops;
+
+	/** If enabled, export SkySphere blueprints. Uses extension EPIC_sky_spheres, which is supported by Unreal's glTF viewer. */
+	UPROPERTY(EditAnywhere, BlueprintReadWrite, Config, Category = Scene)
+	bool bExportSkySpheres;
+
+	/** Mode determining if and how to export LevelVariantSetsActors. */
+	UPROPERTY(EditAnywhere, BlueprintReadWrite, Config, Category = Scene)
+	EGLTFVariantSetsMode VariantSetsMode;
+
+	/** Mode determining if and how to export material variants that change the materials property on a static or skeletal mesh component. */
+	UPROPERTY(EditAnywhere, BlueprintReadWrite, Config, Category = VariantSets, Meta = (EditCondition = "VariantSetsMode != EGLTFVariantSetsMode::None"))
+	EGLTFMaterialVariantMode ExportMaterialVariants;
+
+	/** If enabled, export variants that change the mesh property on a static or skeletal mesh component. */
+	UPROPERTY(EditAnywhere, BlueprintReadWrite, Config, Category = VariantSets, Meta = (EditCondition = "VariantSetsMode == EGLTFVariantSetsMode::Epic"))
+	bool bExportMeshVariants;
+
+	/** If enabled, export variants that change the visible property on a scene component. */
+	UPROPERTY(EditAnywhere, BlueprintReadWrite, Config, Category = VariantSets, Meta = (EditCondition = "VariantSetsMode == EGLTFVariantSetsMode::Epic"))
+	bool bExportVisibilityVariants;
+
+	UFUNCTION(BlueprintCallable, Category = General)
+	void ResetToDefault();
+};
