[GLTFExporter] feat: add GLTFPrebakeOptions for prebaking options

commit 68891412be61e8ea8985ab361c241dde21c5847e
Date: Thu, 27 Jan 2022 19:17:23 +0100

    feat: add GLTFPrebakeOptions for prebaking options
    (cherry picked from commit c52149b83942d27b1588e40a09a905e7acc75686​)

#skipci

diff --git a/Engine/Plugins/Enterprise/GLTFExporter/MigratedFromGit/Source/GLTFExporter/Private/Options/GLTFPrebakeOptions.cpp b/Engine/Plugins/Enterprise/GLTFExporter/MigratedFromGit/Source/GLTFExporter/Private/Options/GLTFPrebakeOptions.cpp
new file mode 100644
--- /dev/null
+++ b/Engine/Plugins/Enterprise/GLTFExporter/MigratedFromGit/Source/GLTFExporter/Private/Options/GLTFPrebakeOptions.cpp
@@ -0,0 +1,16 @@
+// Copyright Epic Games, Inc. All Rights Reserved.
+
+#include "Options/GLTFPrebakeOptions.h"
+
+UGLTFPrebakeOptions::UGLTFPrebakeOptions(const FObjectInitializer& ObjectInitializer)
+	: Super(ObjectInitializer)
+{
+	ResetToDefault();
+}
+
+void UGLTFPrebakeOptions::ResetToDefault()
+{
+	DefaultMaterialBakeSize = EGLTFMaterialBakeSizePOT::POT_1024;
+	DefaultMaterialBakeFilter = TF_Trilinear;
+	DefaultMaterialBakeTiling = TA_Wrap;
+}
diff --git a/Engine/Plugins/Enterprise/GLTFExporter/MigratedFromGit/Source/GLTFExporter/Public/Options/GLTFPrebakeOptions.h b/Engine/Plugins/Enterprise/GLTFExporter/MigratedFromGit/Source/GLTFExporter/Public/Options/GLTFPrebakeOptions.h
new file mode 100644
--- /dev/null
+++ b/Engine/Plugins/Enterprise/GLTFExporter/MigratedFromGit/Source/GLTFExporter/Public/Options/GLTFPrebakeOptions.h
@@ -0,0 +1,31 @@
+// Copyright Epic Games, Inc. All Rights Reserved.
+
+#pragma once
+
+#include "Options/GLTFExportOptions.h"
+#include "GLTFPrebakeOptions.generated.h"
+
+UCLASS(BlueprintType, Config=EditorPerProjectUserSettings, HideCategories=(DebugProperty))
+class GLTFEXPORTER_API UGLTFPrebakeOptions : public UObject
+{
+	GENERATED_UCLASS_BODY()
+
+	/** Default size of the baked out texture (containing the material input). Can be overridden by material- and input-specific bake settings, see GLTFMaterialExportOptions. */
+	UPROPERTY(EditAnywhere, BlueprintReadWrite, Config, Category = Material)
+	EGLTFMaterialBakeSizePOT DefaultMaterialBakeSize;
+
+	/** Default filtering mode used when sampling the baked out texture. Can be overridden by material- and input-specific bake settings, see GLTFMaterialExportOptions. */
+	UPROPERTY(EditAnywhere, BlueprintReadWrite, Config, Category = Material, Meta = (ValidEnumValues="TF_Nearest, TF_Bilinear, TF_Trilinear"))
+	TEnumAsByte<TextureFilter> DefaultMaterialBakeFilter;
+
+	/** Default addressing mode used when sampling the baked out texture. Can be overridden by material- and input-specific bake settings, see GLTFMaterialExportOptions. */
+	UPROPERTY(EditAnywhere, BlueprintReadWrite, Config, Category = Material)
+	TEnumAsByte<TextureAddress> DefaultMaterialBakeTiling;
+
+	/** Input-specific default bake settings that override the general defaults above. */
+	UPROPERTY(EditAnywhere, BlueprintReadWrite, Config, Category = Material)
+	TMap<EGLTFMaterialPropertyGroup, FGLTFOverrideMaterialBakeSettings> DefaultInputBakeSettings;
+
+	UFUNCTION(BlueprintCallable, Category = General)
+	void ResetToDefault();
+};
