[GLTFExporter] feat: add UI window for prebaking options

commit 146d583a805c59a68f988366cc65e8fc6ddbc547
Date: Thu, 27 Jan 2022 20:39:31 +0100

    feat: add UI window for prebaking options
    (cherry picked from commit c0d99e4a126d44c5adb404251cf9e77ef3f75643​)

#skipci

diff --git a/Engine/Plugins/Enterprise/GLTFExporter/MigratedFromGit/Source/GLTFExporter/Private/Actions/GLTFPrebakeAssetActions.cpp b/Engine/Plugins/Enterprise/GLTFExporter/MigratedFromGit/Source/GLTFExporter/Private/Actions/GLTFPrebakeAssetActions.cpp
--- a/Engine/Plugins/Enterprise/GLTFExporter/MigratedFromGit/Source/GLTFExporter/Private/Actions/GLTFPrebakeAssetActions.cpp
+++ b/Engine/Plugins/Enterprise/GLTFExporter/MigratedFromGit/Source/GLTFExporter/Private/Actions/GLTFPrebakeAssetActions.cpp
@@ -4,6 +4,8 @@
 
 #include "Actions/GLTFPrebakeAssetActions.h"
 #include "Actions/GLTFEditorStyle.h"
+#include "Options/GLTFPrebakeOptions.h"
+#include "UI/GLTFPrebakeOptionsWindow.h"
 #include "ToolMenus.h"
 
 #define LOCTEXT_NAMESPACE "GLTFPrebakeAssetActions"
@@ -39,6 +41,13 @@ void FGLTFPrebakeAssetActions::GetPrebakeActions(const TArray<UObject*>& InObjec
 
 void FGLTFPrebakeAssetActions::OnPrebake(TArray<FWeakObjectPtr> Objects) const
 {
+	UGLTFPrebakeOptions* Options = NewObject<UGLTFPrebakeOptions>();
+
+	if (!SGLTFPrebakeOptionsWindow::ShowDialog(Options))
+	{
+		return;
+	}
+
 	for (const FWeakObjectPtr& Object : Objects)
 	{
 		if (UMaterialInterface* Material = Cast<UMaterialInterface>(Object.Get()))
diff --git a/Engine/Plugins/Enterprise/GLTFExporter/MigratedFromGit/Source/GLTFExporter/Private/UI/GLTFPrebakeOptionsWindow.cpp b/Engine/Plugins/Enterprise/GLTFExporter/MigratedFromGit/Source/GLTFExporter/Private/UI/GLTFPrebakeOptionsWindow.cpp
new file mode 100644
--- /dev/null
+++ b/Engine/Plugins/Enterprise/GLTFExporter/MigratedFromGit/Source/GLTFExporter/Private/UI/GLTFPrebakeOptionsWindow.cpp
@@ -0,0 +1,137 @@
+// Copyright Epic Games, Inc. All Rights Reserved.
+
+#if WITH_EDITOR
+
+#include "UI/GLTFPrebakeOptionsWindow.h"
+#include "Modules/ModuleManager.h"
+#include "PropertyEditorModule.h"
+#include "IDetailsView.h"
+#include "Interfaces/IMainFrameModule.h"
+#include "Widgets/Layout/SUniformGridPanel.h"
+#include "Widgets/Input/SButton.h"
+
+#define LOCTEXT_NAMESPACE "SGLTFPrebakeOptionsWindow"
+
+SGLTFPrebakeOptionsWindow::SGLTFPrebakeOptionsWindow() : bUserCancelled(true)
+{
+
+}
+
+void SGLTFPrebakeOptionsWindow::Construct(const FArguments& InArgs)
+{
+	WidgetWindow = InArgs._WidgetWindow;
+
+	// Retrieve property editor module and create a SDetailsView
+	FPropertyEditorModule& PropertyEditorModule = FModuleManager::GetModuleChecked<FPropertyEditorModule>("PropertyEditor");
+	FDetailsViewArgs DetailsViewArgs;
+	DetailsViewArgs.bAllowSearch = false;
+	DetailsViewArgs.NameAreaSettings = FDetailsViewArgs::HideNameArea;
+	DetailsViewArgs.bAllowMultipleTopLevelObjects = true;
+
+	DetailsView = PropertyEditorModule.CreateDetailView(DetailsViewArgs);
+
+	// Set provided objects on SDetailsView
+	DetailsView->SetObjects(InArgs._SettingsObjects, true);
+
+	this->ChildSlot
+	[
+		SNew(SVerticalBox)
+
+		+ SVerticalBox::Slot()
+		.Padding(2)
+		.MaxHeight(500.0f)
+		[
+			DetailsView->AsShared()
+		]
+
+		+ SVerticalBox::Slot()
+		.AutoHeight()
+		.HAlign(HAlign_Right)
+		.Padding(2)
+		[
+			SNew(SUniformGridPanel)
+			.SlotPadding(2)
+			+ SUniformGridPanel::Slot(0, 0)
+			[
+				SAssignNew(ConfirmButton, SButton)
+				.HAlign(HAlign_Center)
+				.Text(LOCTEXT("MaterialBakeOptionWindow_Import", "Confirm"))
+				.OnClicked(this, &SGLTFPrebakeOptionsWindow::OnConfirm)
+			]
+			+ SUniformGridPanel::Slot(1, 0)
+			[
+				SNew(SButton)
+				.HAlign(HAlign_Center)
+				.Text(LOCTEXT("MaterialBakeOptionWindow_Cancel", "Cancel"))
+				.ToolTipText(LOCTEXT("MaterialBakeOptionWindow_Cancel_ToolTip", "Cancels baking out Material"))
+				.OnClicked(this, &SGLTFPrebakeOptionsWindow::OnCancel)
+			]
+		]
+	];
+}
+
+FReply SGLTFPrebakeOptionsWindow::OnConfirm()
+{
+	bUserCancelled = false;
+	if (WidgetWindow.IsValid())
+	{
+		WidgetWindow.Pin()->RequestDestroyWindow();
+	}
+	return FReply::Handled();
+}
+
+FReply SGLTFPrebakeOptionsWindow::OnCancel()
+{
+	if (WidgetWindow.IsValid())
+	{
+		WidgetWindow.Pin()->RequestDestroyWindow();
+	}
+	return FReply::Handled();
+}
+
+FReply SGLTFPrebakeOptionsWindow::OnKeyDown(const FGeometry& MyGeometry, const FKeyEvent& InKeyEvent)
+{
+	if (InKeyEvent.GetKey() == EKeys::Escape)
+	{
+		return OnCancel();
+	}
+
+	return FReply::Unhandled();
+}
+
+bool SGLTFPrebakeOptionsWindow::WasUserCancelled()
+{
+	return bUserCancelled;
+}
+
+bool SGLTFPrebakeOptionsWindow::ShowDialog(UGLTFPrebakeOptions* PrebakeOptions)
+{
+	TSharedRef<SWindow> Window = SNew(SWindow)
+		.Title(LOCTEXT("WindowTitle", "Material Baking Options"))
+		.SizingRule(ESizingRule::Autosized);
+
+	TArray<TWeakObjectPtr<UObject>> OptionObjects = { PrebakeOptions };
+	TSharedPtr<SGLTFPrebakeOptionsWindow> Options;
+
+	Window->SetContent
+	(
+		SAssignNew(Options, SGLTFPrebakeOptionsWindow)
+		.WidgetWindow(Window)
+		.SettingsObjects(OptionObjects)
+	);
+
+	TSharedPtr<SWindow> ParentWindow;
+	if (FModuleManager::Get().IsModuleLoaded("MainFrame"))
+	{
+		IMainFrameModule& MainFrame = FModuleManager::LoadModuleChecked<IMainFrameModule>("MainFrame");
+		ParentWindow = MainFrame.GetParentWindow();
+		FSlateApplication::Get().AddModalWindow(Window, ParentWindow, false);
+		return !Options->WasUserCancelled();
+	}
+
+	return false;
+}
+
+#undef LOCTEXT_NAMESPACE //"SGLTFPrebakeOptionsWindow"
+
+#endif
diff --git a/Engine/Plugins/Enterprise/GLTFExporter/MigratedFromGit/Source/GLTFExporter/Private/UI/GLTFPrebakeOptionsWindow.h b/Engine/Plugins/Enterprise/GLTFExporter/MigratedFromGit/Source/GLTFExporter/Private/UI/GLTFPrebakeOptionsWindow.h
new file mode 100644
--- /dev/null
+++ b/Engine/Plugins/Enterprise/GLTFExporter/MigratedFromGit/Source/GLTFExporter/Private/UI/GLTFPrebakeOptionsWindow.h
@@ -0,0 +1,53 @@
+// Copyright Epic Games, Inc. All Rights Reserved.
+
+#pragma once
+
+#if WITH_EDITOR
+
+#include "Options/GLTFPrebakeOptions.h"
+#include "Widgets/SCompoundWidget.h"
+#include "Widgets/DeclarativeSyntaxSupport.h"
+#include "Widgets/SWindow.h"
+
+class SButton;
+
+/** Options window used to populate provided settings objects */
+class SGLTFPrebakeOptionsWindow : public SCompoundWidget
+{
+public:
+
+	SLATE_BEGIN_ARGS(SGLTFPrebakeOptionsWindow)
+	{}
+		SLATE_ARGUMENT(TSharedPtr<SWindow>, WidgetWindow)
+		SLATE_ARGUMENT(TArray<TWeakObjectPtr<UObject>>, SettingsObjects)
+	SLATE_END_ARGS()
+
+	SGLTFPrebakeOptionsWindow();
+	void Construct(const FArguments& InArgs);
+
+	/** Begin SCompoundWidget overrides */
+	virtual bool SupportsKeyboardFocus() const override { return true; }
+	virtual FReply OnKeyDown(const FGeometry& MyGeometry, const FKeyEvent& InKeyEvent) override;
+	/** End SCompoundWidget overrides */
+
+	/** Callbacks used for Confirm and Cancel buttons */
+	FReply OnConfirm();
+	FReply OnCancel();
+
+	/** Returns whether or not the user cancelled the operation */
+	bool WasUserCancelled();
+
+	static bool ShowDialog(UGLTFPrebakeOptions* PrebakeOptions);
+
+private:
+	/** Owning window this widget is part of */
+	TWeakPtr< SWindow > WidgetWindow;
+	/** Whether or not the cancel button was clicked by the user */
+	bool bUserCancelled;
+	/** Detailsview used to display SettingsObjects, and allowing user to change options */
+	TSharedPtr<class IDetailsView> DetailsView;
+	/** Shared ptr to Confirm button */
+	TSharedPtr<SButton> ConfirmButton;
+};
+
+#endif
