[GLTFExporter] refactor: change all builder constructors to public for future material prebaker

commit 6ce47dc22ddb140b2a42646231beb462d5cc6071
Date: Thu, 27 Jan 2022 20:42:39 +0100

    refactor: change all builder constructors to public for future material prebaker
    (cherry picked from commit 1bc2846aa430b51de5bbc92a28aba8c313185908​)

#skipci

diff --git a/Engine/Plugins/Enterprise/GLTFExporter/MigratedFromGit/Source/GLTFExporter/Private/Builders/GLTFBuilder.h b/Engine/Plugins/Enterprise/GLTFExporter/MigratedFromGit/Source/GLTFExporter/Private/Builders/GLTFBuilder.h
--- a/Engine/Plugins/Enterprise/GLTFExporter/MigratedFromGit/Source/GLTFExporter/Private/Builders/GLTFBuilder.h
+++ b/Engine/Plugins/Enterprise/GLTFExporter/MigratedFromGit/Source/GLTFExporter/Private/Builders/GLTFBuilder.h
@@ -8,10 +8,6 @@
 
 class FGLTFBuilder
 {
-protected:
-
-	FGLTFBuilder(const FString& FilePath, const UGLTFExportOptions* ExportOptions);
-
 public:
 
 	const bool bIsGlbFile;
@@ -21,6 +17,8 @@ public:
 	// TODO: make ExportOptions private and expose each option via getters to ease overriding settings in future
 	const UGLTFExportOptions* ExportOptions;
 
+	FGLTFBuilder(const FString& FilePath, const UGLTFExportOptions* ExportOptions);
+
 	const UMaterialInterface* ResolveProxy(const UMaterialInterface* Material) const;
 
 	FIntPoint GetBakeSizeForMaterialProperty(const UMaterialInterface* Material, EGLTFMaterialPropertyGroup PropertyGroup) const;
diff --git a/Engine/Plugins/Enterprise/GLTFExporter/MigratedFromGit/Source/GLTFExporter/Private/Builders/GLTFConvertBuilder.h b/Engine/Plugins/Enterprise/GLTFExporter/MigratedFromGit/Source/GLTFExporter/Private/Builders/GLTFConvertBuilder.h
--- a/Engine/Plugins/Enterprise/GLTFExporter/MigratedFromGit/Source/GLTFExporter/Private/Builders/GLTFConvertBuilder.h
+++ b/Engine/Plugins/Enterprise/GLTFExporter/MigratedFromGit/Source/GLTFExporter/Private/Builders/GLTFConvertBuilder.h
@@ -25,14 +25,12 @@
 
 class FGLTFConvertBuilder : public FGLTFBufferBuilder
 {
-protected:
-
-	FGLTFConvertBuilder(const FString& FilePath, const UGLTFExportOptions* ExportOptions, const TSet<AActor*>& SelectedActors);
-
 public:
 
 	const TSet<AActor*> SelectedActors;
 
+	FGLTFConvertBuilder(const FString& FilePath, const UGLTFExportOptions* ExportOptions, const TSet<AActor*>& SelectedActors = {});
+
 	bool IsSelectedActor(const AActor* Object) const;
 	bool IsRootActor(const AActor* Actor) const;
 
diff --git a/Engine/Plugins/Enterprise/GLTFExporter/MigratedFromGit/Source/GLTFExporter/Private/Builders/GLTFJsonBuilder.h b/Engine/Plugins/Enterprise/GLTFExporter/MigratedFromGit/Source/GLTFExporter/Private/Builders/GLTFJsonBuilder.h
--- a/Engine/Plugins/Enterprise/GLTFExporter/MigratedFromGit/Source/GLTFExporter/Private/Builders/GLTFJsonBuilder.h
+++ b/Engine/Plugins/Enterprise/GLTFExporter/MigratedFromGit/Source/GLTFExporter/Private/Builders/GLTFJsonBuilder.h
@@ -7,7 +7,9 @@
 
 class FGLTFJsonBuilder : public FGLTFTaskBuilder
 {
-protected:
+public:
+
+	FGLTFJsonSceneIndex& DefaultScene;
 
 	FGLTFJsonBuilder(const FString& FilePath, const UGLTFExportOptions* ExportOptions);
 
@@ -15,10 +17,6 @@ protected:
 
 	TSet<EGLTFJsonExtension> GetCustomExtensionsUsed() const;
 
-public:
-
-	FGLTFJsonSceneIndex& DefaultScene;
-
 	void AddExtension(EGLTFJsonExtension Extension, bool bIsRequired = false);
 
 	FGLTFJsonAccessorIndex AddAccessor(const FGLTFJsonAccessor& JsonAccessor = {});
diff --git a/Engine/Plugins/Enterprise/GLTFExporter/MigratedFromGit/Source/GLTFExporter/Private/Builders/GLTFLogBuilder.h b/Engine/Plugins/Enterprise/GLTFExporter/MigratedFromGit/Source/GLTFExporter/Private/Builders/GLTFLogBuilder.h
--- a/Engine/Plugins/Enterprise/GLTFExporter/MigratedFromGit/Source/GLTFExporter/Private/Builders/GLTFLogBuilder.h
+++ b/Engine/Plugins/Enterprise/GLTFExporter/MigratedFromGit/Source/GLTFExporter/Private/Builders/GLTFLogBuilder.h
@@ -10,12 +10,10 @@ class IMessageLogListing;
 
 class FGLTFLogBuilder : public FGLTFBuilder
 {
-protected:
+public:
 
 	FGLTFLogBuilder(const FString& FilePath, const UGLTFExportOptions* ExportOptions);
 
-public:
-
 	void LogSuggestion(const FString& Message);
 
 	void LogWarning(const FString& Message);
diff --git a/Engine/Plugins/Enterprise/GLTFExporter/MigratedFromGit/Source/GLTFExporter/Private/Builders/GLTFTaskBuilder.h b/Engine/Plugins/Enterprise/GLTFExporter/MigratedFromGit/Source/GLTFExporter/Private/Builders/GLTFTaskBuilder.h
--- a/Engine/Plugins/Enterprise/GLTFExporter/MigratedFromGit/Source/GLTFExporter/Private/Builders/GLTFTaskBuilder.h
+++ b/Engine/Plugins/Enterprise/GLTFExporter/MigratedFromGit/Source/GLTFExporter/Private/Builders/GLTFTaskBuilder.h
@@ -7,12 +7,10 @@
 
 class FGLTFTaskBuilder : public FGLTFLogBuilder
 {
-protected:
+public:
 
 	FGLTFTaskBuilder(const FString& FilePath, const UGLTFExportOptions* ExportOptions);
 
-public:
-
 	template <typename TaskType, typename... TaskArgTypes, typename = typename TEnableIf<TIsDerivedFrom<TaskType, FGLTFTask>::Value>::Type>
 	bool SetupTask(TaskArgTypes&&... Args)
 	{
