[GLTFExporter] refactor: add GLTFMaterialUtility::IsPrebaked and GetPrebaked

commit c5f25e60aa269fe5209ecdeb252c0c6168fcb03c
Date: Thu, 27 Jan 2022 20:44:08 +0100

    refactor: add GLTFMaterialUtility::IsPrebaked and GetPrebaked
    (cherry picked from commit d5d505746f10722a0968d2ac124dc20e1ef9e112​)

#skipci

diff --git a/Engine/Plugins/Enterprise/GLTFExporter/MigratedFromGit/Source/GLTFExporter/Private/Converters/GLTFMaterialUtility.cpp b/Engine/Plugins/Enterprise/GLTFExporter/MigratedFromGit/Source/GLTFExporter/Private/Converters/GLTFMaterialUtility.cpp
--- a/Engine/Plugins/Enterprise/GLTFExporter/MigratedFromGit/Source/GLTFExporter/Private/Converters/GLTFMaterialUtility.cpp
+++ b/Engine/Plugins/Enterprise/GLTFExporter/MigratedFromGit/Source/GLTFExporter/Private/Converters/GLTFMaterialUtility.cpp
@@ -7,16 +7,15 @@
 #include "GLTFMaterialAnalyzer.h"
 #include "Engine/TextureRenderTarget2D.h"
 #include "CanvasItem.h"
-#include "CanvasTypes.h"
 #include "Modules/ModuleManager.h"
 #include "GLTFMaterialBaking/Public/IMaterialBakingModule.h"
 #include "GLTFMaterialBaking/Public/MaterialBakingStructures.h"
-#include "NormalMapPreview.h"
 #include "Materials/MaterialExpressionCustomOutput.h"
 #include "Materials/MaterialExpressionClearCoatNormalCustomOutput.h"
 #include "Materials/MaterialExpressionTextureCoordinate.h"
 #include "MeshDescription.h"
 #endif
+#include "GLTFConverterUtility.h"
 #include "Misc/DefaultValueHelper.h"
 
 UMaterialInterface* FGLTFMaterialUtility::GetDefault()
@@ -25,6 +24,27 @@ UMaterialInterface* FGLTFMaterialUtility::GetDefault()
 	return DefaultMaterial;
 }
 
+UMaterialInterface* FGLTFMaterialUtility::GetPrebaked(EGLTFJsonShadingModel ShadingModel)
+{
+	// TODO: don't assume enum values are in this order
+	static UMaterialInterface* BaseMaterials[] = {
+		LoadObject<UMaterialInterface>(nullptr, TEXT("/GLTFExporter/Materials/GLTFDefault.GLTFDefault")),
+		LoadObject<UMaterialInterface>(nullptr, TEXT("/GLTFExporter/Materials/GLTFUnlit.GLTFUnlit")),
+		LoadObject<UMaterialInterface>(nullptr, TEXT("/GLTFExporter/Materials/GLTFClearCoat.GLTFClearCoat")),
+	};
+	return BaseMaterials[static_cast<int32>(ShadingModel)];
+}
+
+bool FGLTFMaterialUtility::IsPrebaked(const UMaterial* Material)
+{
+	return Material->GetPathName().StartsWith(TEXT("/GLTFExporter/Materials/GLTF"));
+}
+
+bool FGLTFMaterialUtility::IsPrebaked(const UMaterialInterface* Material)
+{
+	return IsPrebaked(Material->GetMaterial());
+}
+
 #if WITH_EDITOR
 
 bool FGLTFMaterialUtility::IsNormalMap(const FMaterialPropertyEx& Property)
diff --git a/Engine/Plugins/Enterprise/GLTFExporter/MigratedFromGit/Source/GLTFExporter/Private/Converters/GLTFMaterialUtility.h b/Engine/Plugins/Enterprise/GLTFExporter/MigratedFromGit/Source/GLTFExporter/Private/Converters/GLTFMaterialUtility.h
--- a/Engine/Plugins/Enterprise/GLTFExporter/MigratedFromGit/Source/GLTFExporter/Private/Converters/GLTFMaterialUtility.h
+++ b/Engine/Plugins/Enterprise/GLTFExporter/MigratedFromGit/Source/GLTFExporter/Private/Converters/GLTFMaterialUtility.h
@@ -40,6 +40,11 @@ struct FGLTFMaterialUtility
 {
 	static UMaterialInterface* GetDefault();
 
+	static UMaterialInterface* GetPrebaked(EGLTFJsonShadingModel ShadingModel);
+
+	static bool IsPrebaked(const UMaterial* Material);
+	static bool IsPrebaked(const UMaterialInterface* Material);
+
 #if WITH_EDITOR
 	static bool IsNormalMap(const FMaterialPropertyEx& Property);
 	static bool IsSRGB(const FMaterialPropertyEx& Property);
