[GLTFExporter] refactor: add equality operators for core json math types

commit 9c8f06446e920a418b2e3d12323b504369d88e19
Date: Thu, 27 Jan 2022 20:47:43 +0100

    refactor: add equality operators for core json math types
    (cherry picked from commit 6c0be52ceb10c86d0670a5210e8a99463732eea3​)

#skipci

diff --git a/Engine/Plugins/Enterprise/GLTFExporter/MigratedFromGit/Source/GLTFExporter/Private/Json/GLTFJsonColor.h b/Engine/Plugins/Enterprise/GLTFExporter/MigratedFromGit/Source/GLTFExporter/Private/Json/GLTFJsonColor.h
--- a/Engine/Plugins/Enterprise/GLTFExporter/MigratedFromGit/Source/GLTFExporter/Private/Json/GLTFJsonColor.h
+++ b/Engine/Plugins/Enterprise/GLTFExporter/MigratedFromGit/Source/GLTFExporter/Private/Json/GLTFJsonColor.h
@@ -19,6 +19,24 @@ struct TGLTFJsonColor : BaseType, IGLTFJsonArray
 		return *this;
 	}
 
+	bool operator==(const BaseType& Other) const
+	{
+		for (SIZE_T i = 0; i < GetNum(BaseType::Components); ++i)
+		{
+			if (BaseType::Components[i] != Other.Components[i])
+			{
+				return false;
+			}
+		}
+
+		return true;
+	}
+
+	bool operator!=(const BaseType& Other) const
+	{
+		return !(*this == Other);
+	}
+
 	virtual void WriteArray(IGLTFJsonWriter& Writer) const override
 	{
 		for (SIZE_T i = 0; i < GetNum(BaseType::Components); ++i)
diff --git a/Engine/Plugins/Enterprise/GLTFExporter/MigratedFromGit/Source/GLTFExporter/Private/Json/GLTFJsonMatrix.h b/Engine/Plugins/Enterprise/GLTFExporter/MigratedFromGit/Source/GLTFExporter/Private/Json/GLTFJsonMatrix.h
--- a/Engine/Plugins/Enterprise/GLTFExporter/MigratedFromGit/Source/GLTFExporter/Private/Json/GLTFJsonMatrix.h
+++ b/Engine/Plugins/Enterprise/GLTFExporter/MigratedFromGit/Source/GLTFExporter/Private/Json/GLTFJsonMatrix.h
@@ -21,6 +21,24 @@ struct TGLTFJsonMatrix : BaseType, IGLTFJsonArray
 		return *this;
 	}
 
+	bool operator==(const BaseType& Other) const
+	{
+		for (SIZE_T i = 0; i < GetNum(BaseType::Components); ++i)
+		{
+			if (BaseType::Components[i] != Other.Components[i])
+			{
+				return false;
+			}
+		}
+
+		return true;
+	}
+
+	bool operator!=(const BaseType& Other) const
+	{
+		return !(*this == Other);
+	}
+
 	virtual void WriteArray(IGLTFJsonWriter& Writer) const override
 	{
 		for (SIZE_T i = 0; i < GetNum(BaseType::Elements); ++i)
diff --git a/Engine/Plugins/Enterprise/GLTFExporter/MigratedFromGit/Source/GLTFExporter/Private/Json/GLTFJsonQuaternion.h b/Engine/Plugins/Enterprise/GLTFExporter/MigratedFromGit/Source/GLTFExporter/Private/Json/GLTFJsonQuaternion.h
--- a/Engine/Plugins/Enterprise/GLTFExporter/MigratedFromGit/Source/GLTFExporter/Private/Json/GLTFJsonQuaternion.h
+++ b/Engine/Plugins/Enterprise/GLTFExporter/MigratedFromGit/Source/GLTFExporter/Private/Json/GLTFJsonQuaternion.h
@@ -20,6 +20,24 @@ struct FGLTFJsonQuaternion : TGLTFQuaternion<float>, IGLTFJsonArray
 		return *this;
 	}
 
+	bool operator==(const TGLTFQuaternion& Other) const
+	{
+		for (SIZE_T i = 0; i < GetNum(Components); ++i)
+		{
+			if (Components[i] != Other.Components[i])
+			{
+				return false;
+			}
+		}
+
+		return true;
+	}
+
+	bool operator!=(const TGLTFQuaternion& Other) const
+	{
+		return !(*this == Other);
+	}
+
 	virtual void WriteArray(IGLTFJsonWriter& Writer) const override
 	{
 		for (int32 i = 0; i < 4; ++i)
diff --git a/Engine/Plugins/Enterprise/GLTFExporter/MigratedFromGit/Source/GLTFExporter/Private/Json/GLTFJsonVector.h b/Engine/Plugins/Enterprise/GLTFExporter/MigratedFromGit/Source/GLTFExporter/Private/Json/GLTFJsonVector.h
--- a/Engine/Plugins/Enterprise/GLTFExporter/MigratedFromGit/Source/GLTFExporter/Private/Json/GLTFJsonVector.h
+++ b/Engine/Plugins/Enterprise/GLTFExporter/MigratedFromGit/Source/GLTFExporter/Private/Json/GLTFJsonVector.h
@@ -22,6 +22,24 @@ struct TGLTFJsonVector : BaseType, IGLTFJsonArray
 		return *this;
 	}
 
+	bool operator==(const BaseType& Other) const
+	{
+		for (SIZE_T i = 0; i < GetNum(BaseType::Components); ++i)
+		{
+			if (BaseType::Components[i] != Other.Components[i])
+			{
+				return false;
+			}
+		}
+
+		return true;
+	}
+
+	bool operator!=(const BaseType& Other) const
+	{
+		return !(*this == Other);
+	}
+
 	virtual void WriteArray(IGLTFJsonWriter& Writer) const override
 	{
 		for (SIZE_T i = 0; i < GetNum(BaseType::Components); ++i)
