[GLTFExporter] refactor: add FGLTFConvertBuilder::GetOrAddTexture(const UTexture* Texture)

commit 875bde60d9c30a7b70e6177b071b3a61804c1023
Date: Thu, 27 Jan 2022 20:50:58 +0100

    refactor: add FGLTFConvertBuilder::GetOrAddTexture(const UTexture* Texture)
    (cherry picked from commit a8f0904edb61176af2f94b1e1ea5fee8c6078c5a​)

#skipci

diff --git a/Engine/Plugins/Enterprise/GLTFExporter/MigratedFromGit/Source/GLTFExporter/Private/Builders/GLTFConvertBuilder.cpp b/Engine/Plugins/Enterprise/GLTFExporter/MigratedFromGit/Source/GLTFExporter/Private/Builders/GLTFConvertBuilder.cpp
--- a/Engine/Plugins/Enterprise/GLTFExporter/MigratedFromGit/Source/GLTFExporter/Private/Builders/GLTFConvertBuilder.cpp
+++ b/Engine/Plugins/Enterprise/GLTFExporter/MigratedFromGit/Source/GLTFExporter/Private/Builders/GLTFConvertBuilder.cpp
@@ -246,6 +246,26 @@ FGLTFJsonSamplerIndex FGLTFConvertBuilder::GetOrAddSampler(const UTexture* Textu
 	return SamplerConverter->GetOrAdd(Texture);
 }
 
+FGLTFJsonTextureIndex FGLTFConvertBuilder::GetOrAddTexture(const UTexture* Texture)
+{
+	if (const UTexture2D* Texture2D = Cast<UTexture2D>(Texture))
+	{
+		return GetOrAddTexture(Texture2D);
+	}
+
+	if (const UTextureRenderTarget2D* RenderTarget2D = Cast<UTextureRenderTarget2D>(Texture))
+	{
+		return GetOrAddTexture(RenderTarget2D);
+	}
+
+	if (const ULightMapTexture2D* LightMap = Cast<ULightMapTexture2D>(Texture))
+	{
+		return GetOrAddTexture(LightMap);
+	}
+
+	return FGLTFJsonTextureIndex(INDEX_NONE);
+}
+
 FGLTFJsonTextureIndex FGLTFConvertBuilder::GetOrAddTexture(const UTexture2D* Texture)
 {
 	return GetOrAddTexture(Texture, Texture->SRGB);
diff --git a/Engine/Plugins/Enterprise/GLTFExporter/MigratedFromGit/Source/GLTFExporter/Private/Builders/GLTFConvertBuilder.h b/Engine/Plugins/Enterprise/GLTFExporter/MigratedFromGit/Source/GLTFExporter/Private/Builders/GLTFConvertBuilder.h
--- a/Engine/Plugins/Enterprise/GLTFExporter/MigratedFromGit/Source/GLTFExporter/Private/Builders/GLTFConvertBuilder.h
+++ b/Engine/Plugins/Enterprise/GLTFExporter/MigratedFromGit/Source/GLTFExporter/Private/Builders/GLTFConvertBuilder.h
@@ -60,6 +60,7 @@ public:
 	FGLTFJsonMaterialIndex GetOrAddMaterial(const UMaterialInterface* Material, const FGLTFMeshData* MeshData = nullptr, const FGLTFIndexArray& SectionIndices = {});
 
 	FGLTFJsonSamplerIndex GetOrAddSampler(const UTexture* Texture);
+	FGLTFJsonTextureIndex GetOrAddTexture(const UTexture* Texture);
 	FGLTFJsonTextureIndex GetOrAddTexture(const UTexture2D* Texture);
 	FGLTFJsonTextureIndex GetOrAddTexture(const UTextureCube* Texture, ECubeFace CubeFace);
 	FGLTFJsonTextureIndex GetOrAddTexture(const UTextureRenderTarget2D* Texture);
diff --git a/Engine/Plugins/Enterprise/GLTFExporter/MigratedFromGit/Source/GLTFExporter/Private/Builders/GLTFJsonBuilder.cpp b/Engine/Plugins/Enterprise/GLTFExporter/MigratedFromGit/Source/GLTFExporter/Private/Builders/GLTFJsonBuilder.cpp
--- a/Engine/Plugins/Enterprise/GLTFExporter/MigratedFromGit/Source/GLTFExporter/Private/Builders/GLTFJsonBuilder.cpp
+++ b/Engine/Plugins/Enterprise/GLTFExporter/MigratedFromGit/Source/GLTFExporter/Private/Builders/GLTFJsonBuilder.cpp
@@ -170,107 +170,107 @@ const FGLTFJsonRoot& FGLTFJsonBuilder::GetRoot() const
 	return JsonRoot;
 }
 
-FGLTFJsonAccessor& FGLTFJsonBuilder::GetAccessor(FGLTFJsonAccessorIndex AccessorIndex)
+FGLTFJsonAccessor& FGLTFJsonBuilder::GetAccessor(FGLTFJsonAccessorIndex AccessorIndex) const
 {
 	return *JsonRoot.Accessors[AccessorIndex];
 }
 
-FGLTFJsonAnimation& FGLTFJsonBuilder::GetAnimation(FGLTFJsonAnimationIndex AnimationIndex)
+FGLTFJsonAnimation& FGLTFJsonBuilder::GetAnimation(FGLTFJsonAnimationIndex AnimationIndex) const
 {
 	return *JsonRoot.Animations[AnimationIndex];
 }
 
-FGLTFJsonBuffer& FGLTFJsonBuilder::GetBuffer(FGLTFJsonBufferIndex BufferIndex)
+FGLTFJsonBuffer& FGLTFJsonBuilder::GetBuffer(FGLTFJsonBufferIndex BufferIndex) const
 {
 	return *JsonRoot.Buffers[BufferIndex];
 }
 
-FGLTFJsonBufferView& FGLTFJsonBuilder::GetBufferView(FGLTFJsonBufferViewIndex BufferViewIndex)
+FGLTFJsonBufferView& FGLTFJsonBuilder::GetBufferView(FGLTFJsonBufferViewIndex BufferViewIndex) const
 {
 	return *JsonRoot.BufferViews[BufferViewIndex];
 }
 
-FGLTFJsonCamera& FGLTFJsonBuilder::GetCamera(FGLTFJsonCameraIndex CameraIndex)
+FGLTFJsonCamera& FGLTFJsonBuilder::GetCamera(FGLTFJsonCameraIndex CameraIndex) const
 {
 	return *JsonRoot.Cameras[CameraIndex];
 }
 
-FGLTFJsonImage& FGLTFJsonBuilder::GetImage(FGLTFJsonImageIndex ImageIndex)
+FGLTFJsonImage& FGLTFJsonBuilder::GetImage(FGLTFJsonImageIndex ImageIndex) const
 {
 	return *JsonRoot.Images[ImageIndex];
 }
 
-FGLTFJsonMaterial& FGLTFJsonBuilder::GetMaterial(FGLTFJsonMaterialIndex MaterialIndex)
+FGLTFJsonMaterial& FGLTFJsonBuilder::GetMaterial(FGLTFJsonMaterialIndex MaterialIndex) const
 {
 	return *JsonRoot.Materials[MaterialIndex];
 }
 
-FGLTFJsonMesh& FGLTFJsonBuilder::GetMesh(FGLTFJsonMeshIndex MeshIndex)
+FGLTFJsonMesh& FGLTFJsonBuilder::GetMesh(FGLTFJsonMeshIndex MeshIndex) const
 {
 	return *JsonRoot.Meshes[MeshIndex];
 }
 
-FGLTFJsonNode& FGLTFJsonBuilder::GetNode(FGLTFJsonNodeIndex NodeIndex)
+FGLTFJsonNode& FGLTFJsonBuilder::GetNode(FGLTFJsonNodeIndex NodeIndex) const
 {
 	return *JsonRoot.Nodes[NodeIndex];
 }
 
-FGLTFJsonSampler& FGLTFJsonBuilder::GetSampler(FGLTFJsonSamplerIndex SamplerIndex)
+FGLTFJsonSampler& FGLTFJsonBuilder::GetSampler(FGLTFJsonSamplerIndex SamplerIndex) const
 {
 	return *JsonRoot.Samplers[SamplerIndex];
 }
 
-FGLTFJsonScene& FGLTFJsonBuilder::GetScene(FGLTFJsonSceneIndex SceneIndex)
+FGLTFJsonScene& FGLTFJsonBuilder::GetScene(FGLTFJsonSceneIndex SceneIndex) const
 {
 	return *JsonRoot.Scenes[SceneIndex];
 }
 
-FGLTFJsonSkin& FGLTFJsonBuilder::GetSkin(FGLTFJsonSkinIndex SkinIndex)
+FGLTFJsonSkin& FGLTFJsonBuilder::GetSkin(FGLTFJsonSkinIndex SkinIndex) const
 {
 	return *JsonRoot.Skins[SkinIndex];
 }
 
-FGLTFJsonTexture& FGLTFJsonBuilder::GetTexture(FGLTFJsonTextureIndex TextureIndex)
+FGLTFJsonTexture& FGLTFJsonBuilder::GetTexture(FGLTFJsonTextureIndex TextureIndex) const
 {
 	return *JsonRoot.Textures[TextureIndex];
 }
 
-FGLTFJsonBackdrop& FGLTFJsonBuilder::GetBackdrop(FGLTFJsonBackdropIndex BackdropIndex)
+FGLTFJsonBackdrop& FGLTFJsonBuilder::GetBackdrop(FGLTFJsonBackdropIndex BackdropIndex) const
 {
 	return *JsonRoot.Backdrops[BackdropIndex];
 }
 
-FGLTFJsonHotspot& FGLTFJsonBuilder::GetHotspot(FGLTFJsonHotspotIndex HotspotIndex)
+FGLTFJsonHotspot& FGLTFJsonBuilder::GetHotspot(FGLTFJsonHotspotIndex HotspotIndex) const
 {
 	return *JsonRoot.Hotspots[HotspotIndex];
 }
 
-FGLTFJsonLight& FGLTFJsonBuilder::GetLight(FGLTFJsonLightIndex LightIndex)
+FGLTFJsonLight& FGLTFJsonBuilder::GetLight(FGLTFJsonLightIndex LightIndex) const
 {
 	return *JsonRoot.Lights[LightIndex];
 }
 
-FGLTFJsonLightMap& FGLTFJsonBuilder::GetLightMap(FGLTFJsonLightMapIndex LightMapIndex)
+FGLTFJsonLightMap& FGLTFJsonBuilder::GetLightMap(FGLTFJsonLightMapIndex LightMapIndex) const
 {
 	return *JsonRoot.LightMaps[LightMapIndex];
 }
 
-FGLTFJsonSkySphere& FGLTFJsonBuilder::GetSkySphere(FGLTFJsonSkySphereIndex SkySphereIndex)
+FGLTFJsonSkySphere& FGLTFJsonBuilder::GetSkySphere(FGLTFJsonSkySphereIndex SkySphereIndex) const
 {
 	return *JsonRoot.SkySpheres[SkySphereIndex];
 }
 
-FGLTFJsonEpicLevelVariantSets& FGLTFJsonBuilder::GetEpicLevelVariantSets(FGLTFJsonEpicLevelVariantSetsIndex EpicLevelVariantSetsIndex)
+FGLTFJsonEpicLevelVariantSets& FGLTFJsonBuilder::GetEpicLevelVariantSets(FGLTFJsonEpicLevelVariantSetsIndex EpicLevelVariantSetsIndex) const
 {
 	return *JsonRoot.EpicLevelVariantSets[EpicLevelVariantSetsIndex];
 }
 
-FGLTFJsonKhrMaterialVariant& FGLTFJsonBuilder::GetKhrMaterialVariant(FGLTFJsonKhrMaterialVariantIndex KhrMaterialVariantIndex)
+FGLTFJsonKhrMaterialVariant& FGLTFJsonBuilder::GetKhrMaterialVariant(FGLTFJsonKhrMaterialVariantIndex KhrMaterialVariantIndex) const
 {
 	return *JsonRoot.KhrMaterialVariants[KhrMaterialVariantIndex];
 }
 
-FGLTFJsonNodeIndex FGLTFJsonBuilder::GetComponentNodeIndex(FGLTFJsonNodeIndex NodeIndex)
+FGLTFJsonNodeIndex FGLTFJsonBuilder::GetComponentNodeIndex(FGLTFJsonNodeIndex NodeIndex) const
 {
 	if (NodeIndex == INDEX_NONE)
 	{
diff --git a/Engine/Plugins/Enterprise/GLTFExporter/MigratedFromGit/Source/GLTFExporter/Private/Builders/GLTFJsonBuilder.h b/Engine/Plugins/Enterprise/GLTFExporter/MigratedFromGit/Source/GLTFExporter/Private/Builders/GLTFJsonBuilder.h
--- a/Engine/Plugins/Enterprise/GLTFExporter/MigratedFromGit/Source/GLTFExporter/Private/Builders/GLTFJsonBuilder.h
+++ b/Engine/Plugins/Enterprise/GLTFExporter/MigratedFromGit/Source/GLTFExporter/Private/Builders/GLTFJsonBuilder.h
@@ -44,28 +44,28 @@ public:
 	FGLTFJsonNodeIndex AddChildComponentNode(FGLTFJsonNodeIndex ParentNodeIndex, const FGLTFJsonNode& JsonNode = {});
 
 	const FGLTFJsonRoot& GetRoot() const;
-	FGLTFJsonAccessor& GetAccessor(FGLTFJsonAccessorIndex AccessorIndex);
-	FGLTFJsonAnimation& GetAnimation(FGLTFJsonAnimationIndex AnimationIndex);
-	FGLTFJsonBuffer& GetBuffer(FGLTFJsonBufferIndex BufferIndex);
-	FGLTFJsonBufferView& GetBufferView(FGLTFJsonBufferViewIndex BufferViewIndex);
-	FGLTFJsonCamera& GetCamera(FGLTFJsonCameraIndex CameraIndex);
-	FGLTFJsonImage& GetImage(FGLTFJsonImageIndex ImageIndex);
-	FGLTFJsonMaterial& GetMaterial(FGLTFJsonMaterialIndex MaterialIndex);
-	FGLTFJsonMesh& GetMesh(FGLTFJsonMeshIndex MeshIndex);
-	FGLTFJsonNode& GetNode(FGLTFJsonNodeIndex NodeIndex);
-	FGLTFJsonSampler& GetSampler(FGLTFJsonSamplerIndex SamplerIndex);
-	FGLTFJsonScene& GetScene(FGLTFJsonSceneIndex SceneIndex);
-	FGLTFJsonSkin& GetSkin(FGLTFJsonSkinIndex SkinIndex);
-	FGLTFJsonTexture& GetTexture(FGLTFJsonTextureIndex TextureIndex);
-	FGLTFJsonBackdrop& GetBackdrop(FGLTFJsonBackdropIndex BackdropIndex);
-	FGLTFJsonHotspot& GetHotspot(FGLTFJsonHotspotIndex HotspotIndex);
-	FGLTFJsonLight& GetLight(FGLTFJsonLightIndex LightIndex);
-	FGLTFJsonLightMap& GetLightMap(FGLTFJsonLightMapIndex LightMapIndex);
-	FGLTFJsonSkySphere& GetSkySphere(FGLTFJsonSkySphereIndex SkySphereIndex);
-	FGLTFJsonEpicLevelVariantSets& GetEpicLevelVariantSets(FGLTFJsonEpicLevelVariantSetsIndex EpicLevelVariantSetsIndex);
-	FGLTFJsonKhrMaterialVariant& GetKhrMaterialVariant(FGLTFJsonKhrMaterialVariantIndex KhrMaterialVariantIndex);
-
-	FGLTFJsonNodeIndex GetComponentNodeIndex(FGLTFJsonNodeIndex NodeIndex);
+	FGLTFJsonAccessor& GetAccessor(FGLTFJsonAccessorIndex AccessorIndex) const;
+	FGLTFJsonAnimation& GetAnimation(FGLTFJsonAnimationIndex AnimationIndex) const;
+	FGLTFJsonBuffer& GetBuffer(FGLTFJsonBufferIndex BufferIndex) const;
+	FGLTFJsonBufferView& GetBufferView(FGLTFJsonBufferViewIndex BufferViewIndex) const;
+	FGLTFJsonCamera& GetCamera(FGLTFJsonCameraIndex CameraIndex) const;
+	FGLTFJsonImage& GetImage(FGLTFJsonImageIndex ImageIndex) const;
+	FGLTFJsonMaterial& GetMaterial(FGLTFJsonMaterialIndex MaterialIndex) const;
+	FGLTFJsonMesh& GetMesh(FGLTFJsonMeshIndex MeshIndex) const;
+	FGLTFJsonNode& GetNode(FGLTFJsonNodeIndex NodeIndex) const;
+	FGLTFJsonSampler& GetSampler(FGLTFJsonSamplerIndex SamplerIndex) const;
+	FGLTFJsonScene& GetScene(FGLTFJsonSceneIndex SceneIndex) const;
+	FGLTFJsonSkin& GetSkin(FGLTFJsonSkinIndex SkinIndex) const;
+	FGLTFJsonTexture& GetTexture(FGLTFJsonTextureIndex TextureIndex) const;
+	FGLTFJsonBackdrop& GetBackdrop(FGLTFJsonBackdropIndex BackdropIndex) const;
+	FGLTFJsonHotspot& GetHotspot(FGLTFJsonHotspotIndex HotspotIndex) const;
+	FGLTFJsonLight& GetLight(FGLTFJsonLightIndex LightIndex) const;
+	FGLTFJsonLightMap& GetLightMap(FGLTFJsonLightMapIndex LightMapIndex) const;
+	FGLTFJsonSkySphere& GetSkySphere(FGLTFJsonSkySphereIndex SkySphereIndex) const;
+	FGLTFJsonEpicLevelVariantSets& GetEpicLevelVariantSets(FGLTFJsonEpicLevelVariantSetsIndex EpicLevelVariantSetsIndex) const;
+	FGLTFJsonKhrMaterialVariant& GetKhrMaterialVariant(FGLTFJsonKhrMaterialVariantIndex KhrMaterialVariantIndex) const;
+
+	FGLTFJsonNodeIndex GetComponentNodeIndex(FGLTFJsonNodeIndex NodeIndex) const;
 
 private:
 
