[GLTFExporter] feat: add GLTFMaterialPrebaker for prebaking materials

commit de077be68f68a50d8769314b17f6b7b7b79349ab
Date: Thu, 27 Jan 2022 20:51:43 +0100

    feat: add GLTFMaterialPrebaker for prebaking materials
    (cherry picked from commit 4838a7710dc0ae885e8ff034c603e796e4ecab42​)

#skipci

diff --git a/Engine/Plugins/Enterprise/GLTFExporter/MigratedFromGit/Source/GLTFExporter/GLTFExporter.Build.cs b/Engine/Plugins/Enterprise/GLTFExporter/MigratedFromGit/Source/GLTFExporter/GLTFExporter.Build.cs
--- a/Engine/Plugins/Enterprise/GLTFExporter/MigratedFromGit/Source/GLTFExporter/GLTFExporter.Build.cs
+++ b/Engine/Plugins/Enterprise/GLTFExporter/MigratedFromGit/Source/GLTFExporter/GLTFExporter.Build.cs
@@ -47,6 +47,7 @@ public class GLTFExporter : ModuleRules
 					"EditorStyle",
 					"PropertyEditor",
 					"ToolMenus",
+					"MaterialUtilities",
 					"MeshMergeUtilities",
 					"MeshDescription",
 					"StaticMeshDescription",
diff --git a/Engine/Plugins/Enterprise/GLTFExporter/MigratedFromGit/Source/GLTFExporter/Private/Converters/GLTFMaterialPrebaker.cpp b/Engine/Plugins/Enterprise/GLTFExporter/MigratedFromGit/Source/GLTFExporter/Private/Converters/GLTFMaterialPrebaker.cpp
new file mode 100644
--- /dev/null
+++ b/Engine/Plugins/Enterprise/GLTFExporter/MigratedFromGit/Source/GLTFExporter/Private/Converters/GLTFMaterialPrebaker.cpp
@@ -0,0 +1,318 @@
+// Copyright Epic Games, Inc. All Rights Reserved.
+
+#if WITH_EDITOR
+
+#include "Converters/GLTFMaterialPrebaker.h"
+#include "Converters/GLTFMaterialUtility.h"
+#include "Converters/GLTFImageUtility.h"
+#include "Materials/MaterialInstanceConstant.h"
+#include "MaterialUtilities.h"
+#include "ImageUtils.h"
+
+FGLTFMaterialPrebaker::FGLTFMaterialPrebaker(const UGLTFPrebakeOptions* Options)
+	: Builder(TEXT(""), CreateExportOptions(Options))
+{
+	Builder.ImageConverter = CreateCustomImageConverter();
+}
+
+UMaterialInterface* FGLTFMaterialPrebaker::Prebake(UMaterialInterface* OriginalMaterial)
+{
+	if (FGLTFMaterialUtility::IsPrebaked(OriginalMaterial))
+	{
+		return OriginalMaterial;
+	}
+
+	const FGLTFJsonMaterialIndex MaterialIndex = Builder.GetOrAddMaterial(OriginalMaterial);
+	if (MaterialIndex == INDEX_NONE)
+	{
+		return nullptr;
+	}
+
+	Builder.CompleteAllTasks();
+	const FGLTFJsonMaterial& JsonMaterial = Builder.GetMaterial(MaterialIndex);
+
+	UMaterialInstanceConstant* ProxyMaterial = CreateProxyMaterial(OriginalMaterial, JsonMaterial.ShadingModel);
+	ApplyPrebakedProperties(ProxyMaterial, JsonMaterial);
+
+	return ProxyMaterial;
+}
+
+void FGLTFMaterialPrebaker::ApplyPrebakedProperties(UMaterialInstanceConstant* ProxyMaterial, const FGLTFJsonMaterial& JsonMaterial)
+{
+	ApplyPrebakedProperty(ProxyMaterial, TEXT("Base Color Factor"), JsonMaterial.PBRMetallicRoughness.BaseColorFactor);
+	ApplyPrebakedProperty(ProxyMaterial, TEXT("Base Color"), JsonMaterial.PBRMetallicRoughness.BaseColorTexture);
+
+	if (JsonMaterial.ShadingModel == EGLTFJsonShadingModel::Default || JsonMaterial.ShadingModel == EGLTFJsonShadingModel::ClearCoat)
+	{
+		ApplyPrebakedProperty(ProxyMaterial, TEXT("Emissive Factor"), JsonMaterial.EmissiveFactor);
+		ApplyPrebakedProperty(ProxyMaterial, TEXT("Emissive"), JsonMaterial.EmissiveTexture);
+
+		ApplyPrebakedProperty(ProxyMaterial, TEXT("Metallic Factor"), JsonMaterial.PBRMetallicRoughness.MetallicFactor);
+		ApplyPrebakedProperty(ProxyMaterial, TEXT("Roughness Factor"), JsonMaterial.PBRMetallicRoughness.RoughnessFactor);
+		ApplyPrebakedProperty(ProxyMaterial, TEXT("Metallic Roughness"), JsonMaterial.PBRMetallicRoughness.MetallicRoughnessTexture);
+
+		ApplyPrebakedProperty(ProxyMaterial, TEXT("Normal Scale"), JsonMaterial.NormalTexture.Scale);
+		ApplyPrebakedProperty(ProxyMaterial, TEXT("Normal"), JsonMaterial.NormalTexture, true);
+
+		ApplyPrebakedProperty(ProxyMaterial, TEXT("Occlusion Strength"), JsonMaterial.OcclusionTexture.Strength);
+		ApplyPrebakedProperty(ProxyMaterial, TEXT("Occlusion"), JsonMaterial.OcclusionTexture);
+
+		if (JsonMaterial.ShadingModel == EGLTFJsonShadingModel::ClearCoat)
+		{
+			ApplyPrebakedProperty(ProxyMaterial, TEXT("Clear Coat Factor"), JsonMaterial.ClearCoat.ClearCoatFactor);
+			ApplyPrebakedProperty(ProxyMaterial, TEXT("Clear Coat"), JsonMaterial.ClearCoat.ClearCoatTexture);
+
+			ApplyPrebakedProperty(ProxyMaterial, TEXT("Clear Coat Roughness Factor"), JsonMaterial.ClearCoat.ClearCoatRoughnessFactor);
+			ApplyPrebakedProperty(ProxyMaterial, TEXT("Clear Coat Roughness"), JsonMaterial.ClearCoat.ClearCoatRoughnessTexture);
+
+			ApplyPrebakedProperty(ProxyMaterial, TEXT("Clear Coat Normal Scale"), JsonMaterial.ClearCoat.ClearCoatNormalTexture.Scale);
+			ApplyPrebakedProperty(ProxyMaterial, TEXT("Clear Coat Normal"), JsonMaterial.ClearCoat.ClearCoatNormalTexture, true);
+		}
+	}
+
+	ProxyMaterial->PostEditChange();
+}
+
+void FGLTFMaterialPrebaker::ApplyPrebakedProperty(UMaterialInstanceConstant* ProxyMaterial, const FString& PropertyName, float Scalar)
+{
+	float DefaultValue;
+	if (!ProxyMaterial->GetScalarParameterDefaultValue(*PropertyName, DefaultValue))
+	{
+		// TODO: report error
+		return;
+	}
+
+	if (DefaultValue != Scalar)
+	{
+		ProxyMaterial->SetScalarParameterValueEditorOnly(*PropertyName, Scalar);
+	}
+}
+
+void FGLTFMaterialPrebaker::ApplyPrebakedProperty(UMaterialInstanceConstant* ProxyMaterial, const FString& PropertyName, const FGLTFJsonColor3& Color)
+{
+	FLinearColor DefaultValue;
+	if (!ProxyMaterial->GetVectorParameterDefaultValue(*PropertyName, DefaultValue))
+	{
+		// TODO: report error
+		return;
+	}
+
+	const FLinearColor Value(Color.R, Color.G, Color.B);
+	if (DefaultValue != Value)
+	{
+		ProxyMaterial->SetVectorParameterValueEditorOnly(*PropertyName, Value);
+	}
+}
+
+void FGLTFMaterialPrebaker::ApplyPrebakedProperty(UMaterialInstanceConstant* ProxyMaterial, const FString& PropertyName, const FGLTFJsonColor4& Color)
+{
+	FLinearColor DefaultValue;
+	if (!ProxyMaterial->GetVectorParameterDefaultValue(*PropertyName, DefaultValue))
+	{
+		// TODO: report error
+		return;
+	}
+
+	const FLinearColor Value(Color.R, Color.G, Color.B, Color.A);
+	if (DefaultValue != Value)
+	{
+		ProxyMaterial->SetVectorParameterValueEditorOnly(*PropertyName, Value);
+	}
+}
+
+void FGLTFMaterialPrebaker::ApplyPrebakedProperty(UMaterialInstanceConstant* ProxyMaterial, const FString& PropertyName, const FGLTFJsonTextureInfo& TextureInfo, bool bNormalMap)
+{
+	if (TextureInfo.Index != INDEX_NONE)
+	{
+		const UTexture* Texture = FindOrCreateTexture(TextureInfo.Index, bNormalMap);
+		if (Texture == nullptr)
+		{
+			// TODO: report error
+		}
+
+		ProxyMaterial->SetTextureParameterValueEditorOnly(*(PropertyName + TEXT(" Texture")), const_cast<UTexture*>(Texture));
+	}
+
+	if (TextureInfo.TexCoord != 0)
+	{
+		ProxyMaterial->SetScalarParameterValueEditorOnly(*(PropertyName + TEXT(" UV Index")), TextureInfo.TexCoord);
+	}
+
+	if (TextureInfo.Transform.Offset != FGLTFJsonVector2::Zero)
+	{
+		const FLinearColor Offset(TextureInfo.Transform.Offset.X, TextureInfo.Transform.Offset.Y, 0, 0);
+		ProxyMaterial->SetVectorParameterValueEditorOnly(*(PropertyName + TEXT(" UV Offset")), Offset);
+	}
+
+	if (TextureInfo.Transform.Scale != FGLTFJsonVector2::One)
+	{
+		const FLinearColor Scale(TextureInfo.Transform.Scale.X, TextureInfo.Transform.Scale.Y, 0, 0);
+		ProxyMaterial->SetVectorParameterValueEditorOnly(*(PropertyName + TEXT(" UV Scale")), Scale);
+	}
+
+	if (TextureInfo.Transform.Rotation != 0)
+	{
+		ProxyMaterial->SetScalarParameterValueEditorOnly(*(PropertyName + TEXT(" UV Rotation")), TextureInfo.Transform.Rotation);
+	}
+}
+
+const UTexture2D* FGLTFMaterialPrebaker::FindOrCreateTexture(FGLTFJsonTextureIndex Index, bool bNormalMap)
+{
+	const UTexture2D** FoundPtr = Textures.Find(Index);
+	if (FoundPtr != nullptr)
+	{
+		return *FoundPtr;
+	}
+
+	const FGLTFJsonTexture& JsonTexture = Builder.GetTexture(Index);
+	const FGLTFImageData* ImageData = Images.Find(JsonTexture.Source);
+	if (ImageData == nullptr)
+	{
+		return nullptr;
+	}
+
+	const FGLTFJsonSampler& JsonSampler = Builder.GetSampler(JsonTexture.Sampler);
+	const UTexture2D* Texture = CreateTexture(ImageData, JsonSampler, bNormalMap);
+	Textures.Add(Index, Texture);
+	return Texture;
+}
+
+UTexture2D* FGLTFMaterialPrebaker::CreateTexture(const FGLTFImageData* ImageData, const FGLTFJsonSampler& JsonSampler, bool bNormalMap)
+{
+	const FString PackageName = RootPath / TEXT("T_GLTF_") + ImageData->Filename;
+	UPackage* Package = CreatePackage(*PackageName);
+	Package->FullyLoad();
+	Package->Modify();
+
+	FCreateTexture2DParameters TexParams;
+	TexParams.bUseAlpha = !ImageData->bIgnoreAlpha;
+	TexParams.bSRGB = false;
+	TexParams.bDeferCompression = true;
+	TexParams.SourceGuidHash = FGuid();
+	TexParams.CompressionSettings = bNormalMap ? TC_Normalmap : TC_Default;
+
+	UTexture2D* Texture = FImageUtils::CreateTexture2D(ImageData->Size.X, ImageData->Size.Y, *ImageData->Pixels, Package, FPaths::GetCleanFilename(Package->GetName()), RF_Public | RF_Standalone, TexParams);
+	Texture->Filter = ConvertFilter(JsonSampler.MagFilter);
+	Texture->AddressX = ConvertWrap(JsonSampler.WrapS);
+	Texture->AddressY = ConvertWrap(JsonSampler.WrapT);
+	Texture->LODGroup = TEXTUREGROUP_World;
+
+	Texture->PostEditChange();
+	return Texture;
+}
+
+UMaterialInstanceConstant* FGLTFMaterialPrebaker::CreateProxyMaterial(UMaterialInterface* OriginalMaterial, EGLTFJsonShadingModel ShadingModel)
+{
+	const FString PackageName = RootPath / TEXT("MI_GLTF_") + OriginalMaterial->GetName();
+	UPackage* Package = CreatePackage(*PackageName);
+	Package->FullyLoad();
+	Package->Modify();
+
+	UMaterialInterface* BaseMaterial = FGLTFMaterialUtility::GetPrebaked(ShadingModel);
+	UMaterialInstanceConstant* ProxyMaterial = FMaterialUtilities::CreateInstancedMaterial(BaseMaterial, Package, FPaths::GetCleanFilename(Package->GetName()), RF_Public | RF_Standalone);
+
+	const bool bTwoSided = OriginalMaterial->IsTwoSided();
+	if (bTwoSided != BaseMaterial->IsTwoSided())
+	{
+		ProxyMaterial->BasePropertyOverrides.bOverride_TwoSided = true;
+		ProxyMaterial->BasePropertyOverrides.TwoSided = bTwoSided;
+	}
+
+	const EBlendMode BlendMode = OriginalMaterial->GetBlendMode();
+	if (BlendMode != BaseMaterial->GetBlendMode())
+	{
+		ProxyMaterial->BasePropertyOverrides.bOverride_BlendMode = true;
+		ProxyMaterial->BasePropertyOverrides.BlendMode = BlendMode;
+	}
+
+	const float OpacityMaskClipValue = OriginalMaterial->GetOpacityMaskClipValue();
+	if (OpacityMaskClipValue != BaseMaterial->GetOpacityMaskClipValue())
+	{
+		ProxyMaterial->BasePropertyOverrides.bOverride_OpacityMaskClipValue = true;
+		ProxyMaterial->BasePropertyOverrides.OpacityMaskClipValue = OpacityMaskClipValue;
+	}
+
+	UGLTFMaterialExportOptions* UserData = OriginalMaterial->GetAssetUserData<UGLTFMaterialExportOptions>();
+	if (UserData == nullptr)
+	{
+		UserData = NewObject<UGLTFMaterialExportOptions>();
+		OriginalMaterial->AddAssetUserData(UserData);
+	}
+
+	UserData->Proxy = ProxyMaterial;
+	OriginalMaterial->Modify();
+	return ProxyMaterial;
+}
+
+TUniquePtr<IGLTFImageConverter> FGLTFMaterialPrebaker::CreateCustomImageConverter()
+{
+	class FGLTFCustomImageConverter final: public IGLTFImageConverter
+	{
+	public:
+
+		FGLTFMaterialPrebaker& Prebaker;
+		TSet<FString> UniqueFilenames;
+
+		FGLTFCustomImageConverter(FGLTFMaterialPrebaker& Prebaker)
+			: Prebaker(Prebaker)
+		{
+		}
+
+	protected:
+
+		virtual FGLTFJsonImageIndex Convert(TGLTFSuperfluous<FString> Name, EGLTFTextureType Type, bool bIgnoreAlpha, FIntPoint Size, TGLTFSharedArray<FColor> Pixels) override
+		{
+			const FString Filename = FGLTFImageUtility::GetUniqueFilename(Name, TEXT(""), UniqueFilenames);
+			UniqueFilenames.Add(Filename);
+
+			const FGLTFJsonImageIndex ImageIndex = Prebaker.Builder.AddImage();
+			Prebaker.Images.Add(ImageIndex, { Filename, Type, bIgnoreAlpha, Size, Pixels });
+			return ImageIndex;
+		}
+	};
+
+	return MakeUnique<FGLTFCustomImageConverter>(*this);
+}
+
+UGLTFExportOptions* FGLTFMaterialPrebaker::CreateExportOptions(const UGLTFPrebakeOptions* PrebakeOptions)
+{
+	UGLTFExportOptions* ExportOptions = NewObject<UGLTFExportOptions>();
+	ExportOptions->ResetToDefault();
+	ExportOptions->bExportProxyMaterials = false;
+	ExportOptions->DefaultMaterialBakeSize = PrebakeOptions->DefaultMaterialBakeSize;
+	ExportOptions->DefaultMaterialBakeFilter = PrebakeOptions->DefaultMaterialBakeFilter;
+	ExportOptions->DefaultMaterialBakeTiling = PrebakeOptions->DefaultMaterialBakeTiling;
+	ExportOptions->DefaultInputBakeSettings = PrebakeOptions->DefaultInputBakeSettings;
+	return ExportOptions;
+}
+
+TextureAddress FGLTFMaterialPrebaker::ConvertWrap(EGLTFJsonTextureWrap Wrap)
+{
+	switch (Wrap)
+	{
+		case EGLTFJsonTextureWrap::Repeat:         return TA_Wrap;
+		case EGLTFJsonTextureWrap::MirroredRepeat: return TA_Mirror;
+		case EGLTFJsonTextureWrap::ClampToEdge:    return TA_Clamp;
+		default:
+			checkNoEntry();
+			return TA_MAX;
+	}
+}
+
+TextureFilter FGLTFMaterialPrebaker::ConvertFilter(EGLTFJsonTextureFilter Filter)
+{
+	switch (Filter)
+	{
+		case EGLTFJsonTextureFilter::Nearest:              return TF_Nearest;
+		case EGLTFJsonTextureFilter::NearestMipmapNearest: return TF_Nearest;
+		case EGLTFJsonTextureFilter::LinearMipmapNearest:  return TF_Bilinear;
+		case EGLTFJsonTextureFilter::NearestMipmapLinear:  return TF_Bilinear;
+		case EGLTFJsonTextureFilter::Linear:               return TF_Trilinear;
+		case EGLTFJsonTextureFilter::LinearMipmapLinear:   return TF_Trilinear;
+		default:
+			checkNoEntry();
+			return TF_MAX;
+	}
+}
+
+#endif
diff --git a/Engine/Plugins/Enterprise/GLTFExporter/MigratedFromGit/Source/GLTFExporter/Private/Converters/GLTFMaterialPrebaker.h b/Engine/Plugins/Enterprise/GLTFExporter/MigratedFromGit/Source/GLTFExporter/Private/Converters/GLTFMaterialPrebaker.h
new file mode 100644
--- /dev/null
+++ b/Engine/Plugins/Enterprise/GLTFExporter/MigratedFromGit/Source/GLTFExporter/Private/Converters/GLTFMaterialPrebaker.h
@@ -0,0 +1,59 @@
+// Copyright Epic Games, Inc. All Rights Reserved.
+
+#pragma once
+
+#if WITH_EDITOR
+
+#include "Builders/GLTFConvertBuilder.h"
+#include "Options/GLTFPrebakeOptions.h"
+#include "Options/GLTFExportOptions.h"
+
+class UMaterialInstanceConstant;
+
+class FGLTFMaterialPrebaker
+{
+public:
+
+	FGLTFMaterialPrebaker(const UGLTFPrebakeOptions* Options = nullptr);
+
+	UMaterialInterface* Prebake(UMaterialInterface* OriginalMaterial);
+
+	FString RootPath;
+
+private:
+
+	struct FGLTFImageData
+	{
+		FString Filename;
+		EGLTFTextureType Type;
+		bool bIgnoreAlpha;
+		FIntPoint Size;
+		TGLTFSharedArray<FColor> Pixels;
+	};
+
+	void ApplyPrebakedProperties(UMaterialInstanceConstant* ProxyMaterial, const FGLTFJsonMaterial& JsonMaterial);
+
+	void ApplyPrebakedProperty(UMaterialInstanceConstant* ProxyMaterial, const FString& PropertyName, float Scalar);
+	void ApplyPrebakedProperty(UMaterialInstanceConstant* ProxyMaterial, const FString& PropertyName, const FGLTFJsonColor3& Color);
+	void ApplyPrebakedProperty(UMaterialInstanceConstant* ProxyMaterial, const FString& PropertyName, const FGLTFJsonColor4& Color);
+	void ApplyPrebakedProperty(UMaterialInstanceConstant* ProxyMaterial, const FString& PropertyName, const FGLTFJsonTextureInfo& TextureInfo, bool bNormalMap = false);
+
+	const UTexture2D* FindOrCreateTexture(FGLTFJsonTextureIndex Index, bool bNormalMap);
+	UTexture2D* CreateTexture(const FGLTFImageData* ImageData, const FGLTFJsonSampler& JsonSampler, bool bNormalMap);
+
+	UMaterialInstanceConstant* CreateProxyMaterial(UMaterialInterface* OriginalMaterial, EGLTFJsonShadingModel ShadingModel);
+
+	TUniquePtr<IGLTFImageConverter> CreateCustomImageConverter();
+
+	static UGLTFExportOptions* CreateExportOptions(const UGLTFPrebakeOptions* PrebakeOptions);
+
+	static TextureAddress ConvertWrap(EGLTFJsonTextureWrap Wrap);
+	static TextureFilter ConvertFilter(EGLTFJsonTextureFilter Filter);
+
+	FGLTFConvertBuilder Builder;
+
+	TMap<FGLTFJsonTextureIndex, const UTexture2D*> Textures;
+	TMap<FGLTFJsonImageIndex, FGLTFImageData> Images;
+};
+
+#endif
