[GLTFExporter] feat: enable prebaking via asset actions and add IGLTFExporterModule::PrebakeMaterial

commit fde839fe7b63b52c477cf4df5fd4ac2db0bbd673
Date: Thu, 27 Jan 2022 20:53:04 +0100

    feat: enable prebaking via asset actions and add IGLTFExporterModule::PrebakeMaterial
    (cherry picked from commit ed5c51a8c9e469aae78d03600dec3abae064a8e3​)

#skipci

diff --git a/Engine/Plugins/Enterprise/GLTFExporter/MigratedFromGit/Source/GLTFExporter/Private/Actions/GLTFPrebakeAssetActions.cpp b/Engine/Plugins/Enterprise/GLTFExporter/MigratedFromGit/Source/GLTFExporter/Private/Actions/GLTFPrebakeAssetActions.cpp
--- a/Engine/Plugins/Enterprise/GLTFExporter/MigratedFromGit/Source/GLTFExporter/Private/Actions/GLTFPrebakeAssetActions.cpp
+++ b/Engine/Plugins/Enterprise/GLTFExporter/MigratedFromGit/Source/GLTFExporter/Private/Actions/GLTFPrebakeAssetActions.cpp
@@ -6,6 +6,7 @@
 #include "Actions/GLTFEditorStyle.h"
 #include "Options/GLTFPrebakeOptions.h"
 #include "UI/GLTFPrebakeOptionsWindow.h"
+#include "GLTFExporterModule.h"
 #include "ToolMenus.h"
 
 #define LOCTEXT_NAMESPACE "GLTFPrebakeAssetActions"
@@ -52,7 +53,7 @@ void FGLTFPrebakeAssetActions::OnPrebake(TArray<FWeakObjectPtr> Objects) const
 	{
 		if (UMaterialInterface* Material = Cast<UMaterialInterface>(Object.Get()))
 		{
-			// TODO
+			IGLTFExporterModule::PrebakeMaterial(Material, Options);
 		}
 	}
 }
diff --git a/Engine/Plugins/Enterprise/GLTFExporter/MigratedFromGit/Source/GLTFExporter/Private/GLTFExporterModule.cpp b/Engine/Plugins/Enterprise/GLTFExporter/MigratedFromGit/Source/GLTFExporter/Private/GLTFExporterModule.cpp
--- a/Engine/Plugins/Enterprise/GLTFExporter/MigratedFromGit/Source/GLTFExporter/Private/GLTFExporterModule.cpp
+++ b/Engine/Plugins/Enterprise/GLTFExporter/MigratedFromGit/Source/GLTFExporter/Private/GLTFExporterModule.cpp
@@ -1,6 +1,7 @@
 // Copyright Epic Games, Inc. All Rights Reserved.
 
 #include "GLTFExporterModule.h"
+#include "Converters/GLTFMaterialPrebaker.h"
 #include "Actions/GLTFPrebakeAssetActions.h"
 #include "Interfaces/IPluginManager.h"
 #if WITH_EDITOR
@@ -13,6 +14,17 @@ class FAssetTypeActions_Base;
 
 DEFINE_LOG_CATEGORY(LogGLTFExporter);
 
+UMaterialInterface* IGLTFExporterModule::PrebakeMaterial(UMaterialInterface* Material, const UGLTFPrebakeOptions* Options)
+{
+#if WITH_EDITOR
+	FGLTFMaterialPrebaker Prebaker(Options);
+	Prebaker.RootPath = FPaths::GetPath(Material->GetPathName());
+	return Prebaker.Prebake(Material);
+#else
+	return nullptr;
+#endif
+}
+
 class FGLTFExporterModule final : public IGLTFExporterModule
 {
 public:
diff --git a/Engine/Plugins/Enterprise/GLTFExporter/MigratedFromGit/Source/GLTFExporter/Public/GLTFExporterModule.h b/Engine/Plugins/Enterprise/GLTFExporter/MigratedFromGit/Source/GLTFExporter/Public/GLTFExporterModule.h
--- a/Engine/Plugins/Enterprise/GLTFExporter/MigratedFromGit/Source/GLTFExporter/Public/GLTFExporterModule.h
+++ b/Engine/Plugins/Enterprise/GLTFExporter/MigratedFromGit/Source/GLTFExporter/Public/GLTFExporterModule.h
@@ -10,12 +10,15 @@
 
 DECLARE_LOG_CATEGORY_EXTERN(LogGLTFExporter, Log, All);
 
+class UGLTFPrebakeOptions;
+
 /**
  * The public interface of the GLTFExporter module
  */
 class IGLTFExporterModule : public IModuleInterface
 {
 public:
+
 	/**
 	 * Singleton-like access to IGLTFExporter
 	 *
@@ -35,4 +38,6 @@ public:
 	{
 		return FModuleManager::Get().IsModuleLoaded(GLTFEXPORTER_MODULE_NAME);
 	}
+
+	static UMaterialInterface* PrebakeMaterial(UMaterialInterface* Material, const UGLTFPrebakeOptions* Options);
 };
