[GLTFExporter] feat: add optimized support for exporting prebaked materials

commit a3dbd338d1f2df4fd828d663f0eda97f012a8ea5
Date: Thu, 27 Jan 2022 20:53:38 +0100

    feat: add optimized support for exporting prebaked materials
    (cherry picked from commit a29617278e7687c252e64de2098f750015f5c100​)

#skipci

diff --git a/Engine/Plugins/Enterprise/GLTFExporter/MigratedFromGit/Source/GLTFExporter/Private/Tasks/GLTFMaterialTasks.cpp b/Engine/Plugins/Enterprise/GLTFExporter/MigratedFromGit/Source/GLTFExporter/Private/Tasks/GLTFMaterialTasks.cpp
--- a/Engine/Plugins/Enterprise/GLTFExporter/MigratedFromGit/Source/GLTFExporter/Private/Tasks/GLTFMaterialTasks.cpp
+++ b/Engine/Plugins/Enterprise/GLTFExporter/MigratedFromGit/Source/GLTFExporter/Private/Tasks/GLTFMaterialTasks.cpp
@@ -52,13 +52,10 @@ void FGLTFMaterialTask::Complete()
 {
 	Material = Builder.ResolveProxy(Material);
 
+	const UMaterial* BaseMaterial = Material->GetMaterial();
+	if (BaseMaterial->MaterialDomain != MD_Surface)
 	{
-		const UMaterial* ParentMaterial = Material->GetMaterial();
-
-		if (ParentMaterial->MaterialDomain != MD_Surface)
-		{
-			// TODO: report warning (non-surface materials not supported, will be treated as surface)
-		}
+		// TODO: report warning (non-surface materials not supported, will be treated as surface)
 	}
 
 	FGLTFJsonMaterial& JsonMaterial = Builder.GetMaterial(MaterialIndex);
@@ -69,6 +66,12 @@ void FGLTFMaterialTask::Complete()
 	ConvertShadingModel(JsonMaterial.ShadingModel);
 	ConvertAlphaMode(JsonMaterial.AlphaMode, JsonMaterial.BlendMode);
 
+	if (FGLTFMaterialUtility::IsPrebaked(BaseMaterial))
+	{
+		ApplyPrebakedProperties(JsonMaterial);
+		return;
+	}
+
 	if (JsonMaterial.ShadingModel != EGLTFJsonShadingModel::None)
 	{
 		const FMaterialPropertyEx BaseColorProperty = JsonMaterial.ShadingModel == EGLTFJsonShadingModel::Unlit ? MP_EmissiveColor : MP_BaseColor;
@@ -189,14 +192,106 @@ void FGLTFMaterialTask::Complete()
 	}
 }
 
+void FGLTFMaterialTask::ApplyPrebakedProperties(FGLTFJsonMaterial& OutMaterial) const
+{
+	ApplyPrebakedProperty(TEXT("Base Color Factor"), OutMaterial.PBRMetallicRoughness.BaseColorFactor);
+	ApplyPrebakedProperty(TEXT("Base Color"), OutMaterial.PBRMetallicRoughness.BaseColorTexture);
+
+	if (OutMaterial.ShadingModel == EGLTFJsonShadingModel::Default || OutMaterial.ShadingModel == EGLTFJsonShadingModel::ClearCoat)
+	{
+		ApplyPrebakedProperty(TEXT("Emissive Factor"), OutMaterial.EmissiveFactor);
+		ApplyPrebakedProperty(TEXT("Emissive"), OutMaterial.EmissiveTexture);
+
+		ApplyPrebakedProperty(TEXT("Metallic Factor"), OutMaterial.PBRMetallicRoughness.MetallicFactor);
+		ApplyPrebakedProperty(TEXT("Roughness Factor"), OutMaterial.PBRMetallicRoughness.RoughnessFactor);
+		ApplyPrebakedProperty(TEXT("Metallic Roughness"), OutMaterial.PBRMetallicRoughness.MetallicRoughnessTexture);
+
+		ApplyPrebakedProperty(TEXT("Normal Scale"), OutMaterial.NormalTexture.Scale);
+		ApplyPrebakedProperty(TEXT("Normal"), OutMaterial.NormalTexture);
+
+		ApplyPrebakedProperty(TEXT("Occlusion Strength"), OutMaterial.OcclusionTexture.Strength);
+		ApplyPrebakedProperty(TEXT("Occlusion"), OutMaterial.OcclusionTexture);
+
+		if (OutMaterial.ShadingModel == EGLTFJsonShadingModel::ClearCoat)
+		{
+			ApplyPrebakedProperty(TEXT("Clear Coat Factor"), OutMaterial.ClearCoat.ClearCoatFactor);
+			ApplyPrebakedProperty(TEXT("Clear Coat"), OutMaterial.ClearCoat.ClearCoatTexture);
+
+			ApplyPrebakedProperty(TEXT("Clear Coat Roughness Factor"), OutMaterial.ClearCoat.ClearCoatRoughnessFactor);
+			ApplyPrebakedProperty(TEXT("Clear Coat Roughness"), OutMaterial.ClearCoat.ClearCoatRoughnessTexture);
+
+			ApplyPrebakedProperty(TEXT("Clear Coat Normal Scale"), OutMaterial.ClearCoat.ClearCoatNormalTexture.Scale);
+			ApplyPrebakedProperty(TEXT("Clear Coat Normal"), OutMaterial.ClearCoat.ClearCoatNormalTexture);
+		}
+	}
+}
+
+void FGLTFMaterialTask::ApplyPrebakedProperty(const FString& PropertyName, float& OutValue) const
+{
+	Material->GetScalarParameterValue(*PropertyName, OutValue, true);
+}
+
+void FGLTFMaterialTask::ApplyPrebakedProperty(const FString& PropertyName, FGLTFJsonColor3& OutValue) const
+{
+	FLinearColor Value;
+	if (Material->GetVectorParameterValue(*PropertyName, Value, true))
+	{
+		OutValue = FGLTFConverterUtility::ConvertColor3(Value, false);
+	}
+}
+
+void FGLTFMaterialTask::ApplyPrebakedProperty(const FString& PropertyName, FGLTFJsonColor4& OutValue) const
+{
+	FLinearColor Value;
+	if (Material->GetVectorParameterValue(*PropertyName, Value, true))
+	{
+		OutValue = FGLTFConverterUtility::ConvertColor(Value, false);
+	}
+}
+
+void FGLTFMaterialTask::ApplyPrebakedProperty(const FString& PropertyName, FGLTFJsonTextureInfo& OutValue) const
+{
+	UTexture* Texture;
+	if (Material->GetTextureParameterValue(*(PropertyName + TEXT(" Texture")), Texture, true) || Texture == nullptr)
+	{
+		OutValue.Index = Builder.GetOrAddTexture(Texture);
+	}
+
+	float TexCoord;
+	if (Material->GetScalarParameterValue(*(PropertyName + TEXT(" UV Index")), TexCoord, true))
+	{
+		OutValue.TexCoord = FMath::RoundToInt(TexCoord);
+	}
+
+	FLinearColor Offset;
+	if (Material->GetVectorParameterValue(*(PropertyName + TEXT(" UV Offset")), Offset, true))
+	{
+		OutValue.Transform.Offset.X = Offset.R;
+		OutValue.Transform.Offset.Y = Offset.G;
+	}
+
+	FLinearColor Scale;
+	if (Material->GetVectorParameterValue(*(PropertyName + TEXT(" UV Scale")), Scale, true))
+	{
+		OutValue.Transform.Scale.X = Scale.R;
+		OutValue.Transform.Scale.Y = Scale.G;
+	}
+
+	float Rotation;
+	if (Material->GetScalarParameterValue(*(PropertyName + TEXT(" UV Rotation")), Rotation, true))
+	{
+		OutValue.Transform.Rotation = Rotation;
+	}
+}
+
 EMaterialShadingModel FGLTFMaterialTask::GetShadingModel() const
 {
-	const FMaterialShadingModelField Possibilites = Material->GetShadingModels();
-	const int32 PossibilitesCount = Possibilites.CountShadingModels();
+	const FMaterialShadingModelField Possibilities = Material->GetShadingModels();
+	const int32 PossibilitiesCount = Possibilities.CountShadingModels();
 
-	if (PossibilitesCount == 0)
+	if (PossibilitiesCount == 0)
 	{
-		EMaterialShadingModel ShadingModel = MSM_DefaultLit;
+		const EMaterialShadingModel ShadingModel = MSM_DefaultLit;
 		Builder.LogWarning(FString::Printf(
 			TEXT("No shading model defined for material %s, will export as %s"),
 			*Material->GetName(),
@@ -204,11 +299,11 @@ EMaterialShadingModel FGLTFMaterialTask::GetShadingModel() const
 		return ShadingModel;
 	}
 
-	if (PossibilitesCount > 1)
+	if (PossibilitiesCount > 1)
 	{
 		if (!FApp::CanEverRender())
 		{
-			EMaterialShadingModel ShadingModel = FGLTFMaterialUtility::GetRichestShadingModel(Possibilites);
+			const EMaterialShadingModel ShadingModel = FGLTFMaterialUtility::GetRichestShadingModel(Possibilities);
 			Builder.LogWarning(FString::Printf(
 				TEXT("Can't evaluate shading model expression in material %s because renderer missing, will export as %s"),
 				*Material->GetName(),
@@ -223,7 +318,7 @@ EMaterialShadingModel FGLTFMaterialTask::GetShadingModel() const
 
 			if (EvaluationCount == 0)
 			{
-				EMaterialShadingModel ShadingModel = FGLTFMaterialUtility::GetRichestShadingModel(Possibilites);
+				const EMaterialShadingModel ShadingModel = FGLTFMaterialUtility::GetRichestShadingModel(Possibilities);
 				Builder.LogWarning(FString::Printf(
 					TEXT("Evaluation of shading model expression in material %s returned none, will export as %s"),
 					*Material->GetName(),
@@ -233,7 +328,7 @@ EMaterialShadingModel FGLTFMaterialTask::GetShadingModel() const
 
 			if (EvaluationCount > 1)
 			{
-				EMaterialShadingModel ShadingModel = FGLTFMaterialUtility::GetRichestShadingModel(Evaluation);
+				const EMaterialShadingModel ShadingModel = FGLTFMaterialUtility::GetRichestShadingModel(Evaluation);
 				Builder.LogWarning(FString::Printf(
 					TEXT("Evaluation of shading model expression in material %s is inconclusive (%s), will export as %s"),
 					*Material->GetName(),
@@ -248,7 +343,7 @@ EMaterialShadingModel FGLTFMaterialTask::GetShadingModel() const
 		// we should never end up here
 	}
 
-	return Possibilites.GetFirstShadingModel();
+	return Possibilities.GetFirstShadingModel();
 }
 
 void FGLTFMaterialTask::ConvertShadingModel(EGLTFJsonShadingModel& OutShadingModel) const
diff --git a/Engine/Plugins/Enterprise/GLTFExporter/MigratedFromGit/Source/GLTFExporter/Private/Tasks/GLTFMaterialTasks.h b/Engine/Plugins/Enterprise/GLTFExporter/MigratedFromGit/Source/GLTFExporter/Private/Tasks/GLTFMaterialTasks.h
--- a/Engine/Plugins/Enterprise/GLTFExporter/MigratedFromGit/Source/GLTFExporter/Private/Tasks/GLTFMaterialTasks.h
+++ b/Engine/Plugins/Enterprise/GLTFExporter/MigratedFromGit/Source/GLTFExporter/Private/Tasks/GLTFMaterialTasks.h
@@ -45,6 +45,12 @@ private:
 
 	TSet<FMaterialPropertyEx> MeshDataBakedProperties;
 
+	void ApplyPrebakedProperties(FGLTFJsonMaterial& OutMaterial) const;
+	void ApplyPrebakedProperty(const FString& PropertyName, float& OutValue) const;
+	void ApplyPrebakedProperty(const FString& PropertyName, FGLTFJsonColor3& OutValue) const;
+	void ApplyPrebakedProperty(const FString& PropertyName, FGLTFJsonColor4& OutValue) const;
+	void ApplyPrebakedProperty(const FString& PropertyName, FGLTFJsonTextureInfo& OutValue) const;
+
 	EMaterialShadingModel GetShadingModel() const;
 	void ConvertShadingModel(EGLTFJsonShadingModel& OutShadingModel) const;
 	void ConvertAlphaMode(EGLTFJsonAlphaMode& OutAlphaMode, EGLTFJsonBlendMode& OutBlendMode) const;
