[GLTFExporter] feat: add warning when prebaking material that needs mesh data

commit 636f2b09448868e9eb5ebb154007564a83f00170
Date: Thu, 27 Jan 2022 21:08:11 +0100

    feat: add warning when prebaking material that needs mesh data
    (cherry picked from commit 52008cd00afd7db79b2f2e0a283798a2b0d27274​)

#skipci

diff --git a/Engine/Plugins/Enterprise/GLTFExporter/MigratedFromGit/Source/GLTFExporter/Private/Converters/GLTFMaterialPrebaker.cpp b/Engine/Plugins/Enterprise/GLTFExporter/MigratedFromGit/Source/GLTFExporter/Private/Converters/GLTFMaterialPrebaker.cpp
--- a/Engine/Plugins/Enterprise/GLTFExporter/MigratedFromGit/Source/GLTFExporter/Private/Converters/GLTFMaterialPrebaker.cpp
+++ b/Engine/Plugins/Enterprise/GLTFExporter/MigratedFromGit/Source/GLTFExporter/Private/Converters/GLTFMaterialPrebaker.cpp
@@ -22,6 +22,13 @@ UMaterialInterface* FGLTFMaterialPrebaker::Prebake(UMaterialInterface* OriginalM
 		return OriginalMaterial;
 	}
 
+	if (FGLTFMaterialUtility::NeedsMeshData(OriginalMaterial))
+	{
+		Builder.LogWarning(FString::Printf(
+			TEXT("Material %s uses mesh data but prebaking will only use a simple quad as mesh data currently"),
+			*OriginalMaterial->GetName()));
+	}
+
 	const FGLTFJsonMaterialIndex MaterialIndex = Builder.GetOrAddMaterial(OriginalMaterial);
 	if (MaterialIndex == INDEX_NONE)
 	{
@@ -279,6 +286,7 @@ UGLTFExportOptions* FGLTFMaterialPrebaker::CreateExportOptions(const UGLTFPrebak
 	UGLTFExportOptions* ExportOptions = NewObject<UGLTFExportOptions>();
 	ExportOptions->ResetToDefault();
 	ExportOptions->bExportProxyMaterials = false;
+	ExportOptions->BakeMaterialInputs = EGLTFMaterialBakeMode::Simple;
 	ExportOptions->DefaultMaterialBakeSize = PrebakeOptions->DefaultMaterialBakeSize;
 	ExportOptions->DefaultMaterialBakeFilter = PrebakeOptions->DefaultMaterialBakeFilter;
 	ExportOptions->DefaultMaterialBakeTiling = PrebakeOptions->DefaultMaterialBakeTiling;
