[GLTFExporter] feat: create prebaked assets in subfolder "GLTF"

commit 07961e40f81c6739b1764dad0d44630025153493
Date: Thu, 27 Jan 2022 21:38:17 +0100

    feat: create prebaked assets in subfolder "GLTF"
    (cherry picked from commit 69af1d5596fa9f76cbcabff112b5324ba8204dba​)

#skipci

diff --git a/Engine/Plugins/Enterprise/GLTFExporter/MigratedFromGit/Source/GLTFExporter/Private/Converters/GLTFMaterialPrebaker.cpp b/Engine/Plugins/Enterprise/GLTFExporter/MigratedFromGit/Source/GLTFExporter/Private/Converters/GLTFMaterialPrebaker.cpp
--- a/Engine/Plugins/Enterprise/GLTFExporter/MigratedFromGit/Source/GLTFExporter/Private/Converters/GLTFMaterialPrebaker.cpp
+++ b/Engine/Plugins/Enterprise/GLTFExporter/MigratedFromGit/Source/GLTFExporter/Private/Converters/GLTFMaterialPrebaker.cpp
@@ -5,6 +5,7 @@
 #include "Converters/GLTFMaterialPrebaker.h"
 #include "Converters/GLTFMaterialUtility.h"
 #include "Converters/GLTFImageUtility.h"
+#include "AssetRegistry/AssetRegistryModule.h"
 #include "Materials/MaterialInstanceConstant.h"
 #include "MaterialUtilities.h"
 #include "ImageUtils.h"
@@ -35,6 +36,8 @@ UMaterialInterface* FGLTFMaterialPrebaker::Prebake(UMaterialInterface* OriginalM
 		return nullptr;
 	}
 
+	MakeDirectory(RootPath);
+
 	Builder.CompleteAllTasks();
 	const FGLTFJsonMaterial& JsonMaterial = Builder.GetMaterial(MaterialIndex);
 
@@ -210,7 +213,7 @@ UTexture2D* FGLTFMaterialPrebaker::CreateTexture(const FGLTFImageData* ImageData
 
 UMaterialInstanceConstant* FGLTFMaterialPrebaker::CreateProxyMaterial(UMaterialInterface* OriginalMaterial, EGLTFJsonShadingModel ShadingModel)
 {
-	const FString PackageName = RootPath / TEXT("MI_GLTF_") + OriginalMaterial->GetName();
+	const FString PackageName = RootPath / TEXT("GLTF_") + OriginalMaterial->GetName();
 	UPackage* Package = CreatePackage(*PackageName);
 	Package->FullyLoad();
 	Package->Modify();
@@ -281,6 +284,29 @@ TUniquePtr<IGLTFImageConverter> FGLTFMaterialPrebaker::CreateCustomImageConverte
 	return MakeUnique<FGLTFCustomImageConverter>(*this);
 }
 
+bool FGLTFMaterialPrebaker::MakeDirectory(const FString& PackagePath)
+{
+	const FString DirPath = FPaths::ConvertRelativePathToFull(FPackageName::LongPackageNameToFilename(PackagePath + TEXT("/")));
+	if (DirPath.IsEmpty())
+	{
+		return false;
+	}
+
+	bool bResult = true;
+	if (!IFileManager::Get().DirectoryExists(*DirPath))
+	{
+		bResult = IFileManager::Get().MakeDirectory(*DirPath, true);
+	}
+
+	if (bResult)
+	{
+		FAssetRegistryModule& AssetRegistryModule = FModuleManager::LoadModuleChecked<FAssetRegistryModule>("AssetRegistry");
+		AssetRegistryModule.Get().AddPath(PackagePath);
+	}
+
+	return bResult;
+}
+
 UGLTFExportOptions* FGLTFMaterialPrebaker::CreateExportOptions(const UGLTFPrebakeOptions* PrebakeOptions)
 {
 	UGLTFExportOptions* ExportOptions = NewObject<UGLTFExportOptions>();
diff --git a/Engine/Plugins/Enterprise/GLTFExporter/MigratedFromGit/Source/GLTFExporter/Private/Converters/GLTFMaterialPrebaker.h b/Engine/Plugins/Enterprise/GLTFExporter/MigratedFromGit/Source/GLTFExporter/Private/Converters/GLTFMaterialPrebaker.h
--- a/Engine/Plugins/Enterprise/GLTFExporter/MigratedFromGit/Source/GLTFExporter/Private/Converters/GLTFMaterialPrebaker.h
+++ b/Engine/Plugins/Enterprise/GLTFExporter/MigratedFromGit/Source/GLTFExporter/Private/Converters/GLTFMaterialPrebaker.h
@@ -45,6 +45,8 @@ private:
 
 	TUniquePtr<IGLTFImageConverter> CreateCustomImageConverter();
 
+	static bool MakeDirectory(const FString& String);
+
 	static UGLTFExportOptions* CreateExportOptions(const UGLTFPrebakeOptions* PrebakeOptions);
 
 	static TextureAddress ConvertWrap(EGLTFJsonTextureWrap Wrap);
diff --git a/Engine/Plugins/Enterprise/GLTFExporter/MigratedFromGit/Source/GLTFExporter/Private/GLTFExporterModule.cpp b/Engine/Plugins/Enterprise/GLTFExporter/MigratedFromGit/Source/GLTFExporter/Private/GLTFExporterModule.cpp
--- a/Engine/Plugins/Enterprise/GLTFExporter/MigratedFromGit/Source/GLTFExporter/Private/GLTFExporterModule.cpp
+++ b/Engine/Plugins/Enterprise/GLTFExporter/MigratedFromGit/Source/GLTFExporter/Private/GLTFExporterModule.cpp
@@ -18,7 +18,7 @@ UMaterialInterface* IGLTFExporterModule::PrebakeMaterial(UMaterialInterface* Mat
 {
 #if WITH_EDITOR
 	FGLTFMaterialPrebaker Prebaker(Options);
-	Prebaker.RootPath = FPaths::GetPath(Material->GetPathName());
+	Prebaker.RootPath = FPaths::GetPath(Material->GetPathName()) / TEXT("GLTF");
 	return Prebaker.Prebake(Material);
 #else
 	return nullptr;
