[GLTFExporter] feat: enable partial material converter in runtime (incl applying prebaked properties)

commit ea1c2169271f199b7efbe30359507c951344f4b8
Date: Fri, 28 Jan 2022 16:32:17 +0100

    feat: enable partial material converter in runtime (incl applying prebaked properties)
    (cherry picked from commit ff807230106799c63806ff7244cfdd3434545005​)

#skipci

diff --git a/Engine/Plugins/Enterprise/GLTFExporter/MigratedFromGit/Source/GLTFExporter/Private/Converters/GLTFMaterialConverters.cpp b/Engine/Plugins/Enterprise/GLTFExporter/MigratedFromGit/Source/GLTFExporter/Private/Converters/GLTFMaterialConverters.cpp
--- a/Engine/Plugins/Enterprise/GLTFExporter/MigratedFromGit/Source/GLTFExporter/Private/Converters/GLTFMaterialConverters.cpp
+++ b/Engine/Plugins/Enterprise/GLTFExporter/MigratedFromGit/Source/GLTFExporter/Private/Converters/GLTFMaterialConverters.cpp
@@ -43,14 +43,12 @@ void FGLTFMaterialConverter::Sanitize(const UMaterialInterface*& Material, const
 
 FGLTFJsonMaterialIndex FGLTFMaterialConverter::Convert(const UMaterialInterface* Material, const FGLTFMeshData* MeshData, FGLTFIndexArray SectionIndices)
 {
-#if WITH_EDITOR
 	if (Material != FGLTFMaterialUtility::GetDefault())
 	{
 		const FGLTFJsonMaterialIndex MaterialIndex = Builder.AddMaterial();
 		Builder.SetupTask<FGLTFMaterialTask>(Builder, UVOverlapChecker, Material, MeshData, SectionIndices, MaterialIndex);
 		return MaterialIndex;
 	}
-#endif
 
 	return FGLTFJsonMaterialIndex(INDEX_NONE); // use default gltf definition
 }
diff --git a/Engine/Plugins/Enterprise/GLTFExporter/MigratedFromGit/Source/GLTFExporter/Private/Converters/GLTFMaterialConverters.h b/Engine/Plugins/Enterprise/GLTFExporter/MigratedFromGit/Source/GLTFExporter/Private/Converters/GLTFMaterialConverters.h
--- a/Engine/Plugins/Enterprise/GLTFExporter/MigratedFromGit/Source/GLTFExporter/Private/Converters/GLTFMaterialConverters.h
+++ b/Engine/Plugins/Enterprise/GLTFExporter/MigratedFromGit/Source/GLTFExporter/Private/Converters/GLTFMaterialConverters.h
@@ -20,8 +20,6 @@ class FGLTFMaterialConverter final : public FGLTFBuilderContext, public IGLTFMat
 
 	virtual FGLTFJsonMaterialIndex Convert(const UMaterialInterface* Material, const FGLTFMeshData* MeshData, FGLTFIndexArray SectionIndices) override;
 
-#if WITH_EDITOR
 	FGLTFUVOverlapChecker UVOverlapChecker;
 	FGLTFUVDegenerateChecker UVDegenerateChecker;
-#endif
 };
diff --git a/Engine/Plugins/Enterprise/GLTFExporter/MigratedFromGit/Source/GLTFExporter/Private/Converters/GLTFUVDegenerateChecker.cpp b/Engine/Plugins/Enterprise/GLTFExporter/MigratedFromGit/Source/GLTFExporter/Private/Converters/GLTFUVDegenerateChecker.cpp
--- a/Engine/Plugins/Enterprise/GLTFExporter/MigratedFromGit/Source/GLTFExporter/Private/Converters/GLTFUVDegenerateChecker.cpp
+++ b/Engine/Plugins/Enterprise/GLTFExporter/MigratedFromGit/Source/GLTFExporter/Private/Converters/GLTFUVDegenerateChecker.cpp
@@ -1,13 +1,14 @@
 // Copyright Epic Games, Inc. All Rights Reserved.
 
-#if WITH_EDITOR
-
 #include "Converters/GLTFUVDegenerateChecker.h"
+#if WITH_EDITOR
 #include "StaticMeshAttributes.h"
 #include "MeshDescription.h"
+#endif
 
 void FGLTFUVDegenerateChecker::Sanitize(const FMeshDescription*& Description, FGLTFIndexArray& SectionIndices, int32& TexCoord)
 {
+#if WITH_EDITOR
 	if (Description != nullptr)
 	{
 		const int32 SectionCount = Description->PolygonGroups().Num();
@@ -28,66 +29,68 @@ void FGLTFUVDegenerateChecker::Sanitize(const FMeshDescription*& Description, FG
 			Description = nullptr;
 		}
 	}
+#endif
 }
 
 float FGLTFUVDegenerateChecker::Convert(const FMeshDescription* Description, FGLTFIndexArray SectionIndices, int32 TexCoord)
 {
-	if (Description == nullptr)
+#if WITH_EDITOR
+	if (Description != nullptr)
 	{
-		// TODO: report warning?
-		return -1;
-	}
-
-	const TVertexAttributesConstRef<FVector> Positions =
+		const TVertexAttributesConstRef<FVector> Positions =
 		Description->VertexAttributes().GetAttributesRef<FVector>(MeshAttribute::Vertex::Position);
-	const TVertexInstanceAttributesConstRef<FVector2D> UVs =
-		Description->VertexInstanceAttributes().GetAttributesRef<FVector2D>(MeshAttribute::VertexInstance::TextureCoordinate);
+		const TVertexInstanceAttributesConstRef<FVector2D> UVs =
+			Description->VertexInstanceAttributes().GetAttributesRef<FVector2D>(MeshAttribute::VertexInstance::TextureCoordinate);
 
-	int32 TriangleCount = 0;
-	int32 DegenerateCount = 0;
+		int32 TriangleCount = 0;
+		int32 DegenerateCount = 0;
 
-	for (const int32 SectionIndex : SectionIndices)
-	{
-		for (const FPolygonID PolygonID : Description->GetPolygonGroupPolygons(FPolygonGroupID(SectionIndex)))
+		for (const int32 SectionIndex : SectionIndices)
 		{
-			for (const FTriangleID TriangleID : Description->GetPolygonTriangleIDs(PolygonID))
+			for (const FPolygonID PolygonID : Description->GetPolygonGroupPolygons(FPolygonGroupID(SectionIndex)))
 			{
-				const TStaticArray<FVertexID, 3> TriangleVertexIDs = Description->GetTriangleVertices(TriangleID);
-				TArrayView<const FVertexInstanceID> TriangleVertexInstanceIDs = Description->GetTriangleVertexInstances(TriangleID);
-
-				TStaticArray<FVector, 3> TrianglePositions;
-				TStaticArray<FVector2D, 3> TriangleUVs;
-
-				for (int32 Index = 0; Index < 3; Index++)
+				for (const FTriangleID TriangleID : Description->GetPolygonTriangleIDs(PolygonID))
 				{
-					TrianglePositions[Index] = Positions.Get(TriangleVertexIDs[Index]);
-					TriangleUVs[Index] = UVs.Get(TriangleVertexInstanceIDs[Index], TexCoord);
-				}
+					const TStaticArray<FVertexID, 3> TriangleVertexIDs = Description->GetTriangleVertices(TriangleID);
+					TArrayView<const FVertexInstanceID> TriangleVertexInstanceIDs = Description->GetTriangleVertexInstances(TriangleID);
 
-				if (!IsDegenerateTriangle(TrianglePositions))
-				{
-					TriangleCount++;
+					TStaticArray<FVector, 3> TrianglePositions;
+					TStaticArray<FVector2D, 3> TriangleUVs;
 
-					if (IsDegenerateTriangle(TriangleUVs))
+					for (int32 Index = 0; Index < 3; Index++)
 					{
-						DegenerateCount++;
+						TrianglePositions[Index] = Positions.Get(TriangleVertexIDs[Index]);
+						TriangleUVs[Index] = UVs.Get(TriangleVertexInstanceIDs[Index], TexCoord);
+					}
+
+					if (!IsDegenerateTriangle(TrianglePositions))
+					{
+						TriangleCount++;
+
+						if (IsDegenerateTriangle(TriangleUVs))
+						{
+							DegenerateCount++;
+						}
 					}
 				}
 			}
 		}
-	}
 
-	if (TriangleCount == 0)
-	{
-		return -1;
-	}
+		if (TriangleCount == 0)
+		{
+			return -1;
+		}
 
-	if (TriangleCount == DegenerateCount)
-	{
-		return 1;
+		if (TriangleCount == DegenerateCount)
+		{
+			return 1;
+		}
+
+		return static_cast<float>(DegenerateCount) / static_cast<float>(TriangleCount);
 	}
+#endif
 
-	return static_cast<float>(DegenerateCount) / static_cast<float>(TriangleCount);
+	return -1;
 }
 
 bool FGLTFUVDegenerateChecker::IsDegenerateTriangle(const TStaticArray<FVector2D, 3>& Points)
@@ -105,5 +108,3 @@ bool FGLTFUVDegenerateChecker::IsDegenerateTriangle(const TStaticArray<FVector,
 	const float DoubleAreaSquared = (AB ^ AC).SizeSquared();
 	return DoubleAreaSquared < 2 * SMALL_NUMBER * SMALL_NUMBER;
 }
-
-#endif
diff --git a/Engine/Plugins/Enterprise/GLTFExporter/MigratedFromGit/Source/GLTFExporter/Private/Converters/GLTFUVDegenerateChecker.h b/Engine/Plugins/Enterprise/GLTFExporter/MigratedFromGit/Source/GLTFExporter/Private/Converters/GLTFUVDegenerateChecker.h
--- a/Engine/Plugins/Enterprise/GLTFExporter/MigratedFromGit/Source/GLTFExporter/Private/Converters/GLTFUVDegenerateChecker.h
+++ b/Engine/Plugins/Enterprise/GLTFExporter/MigratedFromGit/Source/GLTFExporter/Private/Converters/GLTFUVDegenerateChecker.h
@@ -2,11 +2,11 @@
 
 #pragma once
 
-#if WITH_EDITOR
-
 #include "Converters/GLTFConverter.h"
 #include "Converters/GLTFIndexArray.h"
 
+struct FMeshDescription;
+
 class FGLTFUVDegenerateChecker final : public TGLTFConverter<float, const FMeshDescription*, FGLTFIndexArray, int32>
 {
 	virtual void Sanitize(const FMeshDescription*& Description, FGLTFIndexArray& SectionIndices, int32& TexCoord) override;
@@ -16,5 +16,3 @@ class FGLTFUVDegenerateChecker final : public TGLTFConverter<float, const FMeshD
 	static bool IsDegenerateTriangle(const TStaticArray<FVector2D, 3>& Points);
 	static bool IsDegenerateTriangle(const TStaticArray<FVector, 3>& Points);
 };
-
-#endif
diff --git a/Engine/Plugins/Enterprise/GLTFExporter/MigratedFromGit/Source/GLTFExporter/Private/Converters/GLTFUVOverlapChecker.cpp b/Engine/Plugins/Enterprise/GLTFExporter/MigratedFromGit/Source/GLTFExporter/Private/Converters/GLTFUVOverlapChecker.cpp
--- a/Engine/Plugins/Enterprise/GLTFExporter/MigratedFromGit/Source/GLTFExporter/Private/Converters/GLTFUVOverlapChecker.cpp
+++ b/Engine/Plugins/Enterprise/GLTFExporter/MigratedFromGit/Source/GLTFExporter/Private/Converters/GLTFUVOverlapChecker.cpp
@@ -1,16 +1,17 @@
 // Copyright Epic Games, Inc. All Rights Reserved.
 
-#if WITH_EDITOR
-
 #include "Converters/GLTFUVOverlapChecker.h"
+#if WITH_EDITOR
 #include "StaticMeshAttributes.h"
 #include "MeshDescription.h"
 #include "Modules/ModuleManager.h"
 #include "IMaterialBakingModule.h"
 #include "MaterialBakingStructures.h"
+#endif
 
 void FGLTFUVOverlapChecker::Sanitize(const FMeshDescription*& Description, FGLTFIndexArray& SectionIndices, int32& TexCoord)
 {
+#if WITH_EDITOR
 	if (Description != nullptr)
 	{
 		const TVertexInstanceAttributesConstRef<FVector2D> VertexInstanceUVs =
@@ -31,89 +32,91 @@ void FGLTFUVOverlapChecker::Sanitize(const FMeshDescription*& Description, FGLTF
 			Description = nullptr;
 		}
 	}
+#endif
 }
 
 float FGLTFUVOverlapChecker::Convert(const FMeshDescription* Description, FGLTFIndexArray SectionIndices, int32 TexCoord)
 {
-	if (Description == nullptr)
+#if WITH_EDITOR
+	if (Description != nullptr)
 	{
-		// TODO: report warning?
-		return -1;
-	}
+		// TODO: investigate if the fixed size is high enough to properly calculate overlap
+		const FIntPoint TextureSize(512, 512);
+		const FMaterialPropertyEx Property = MP_EmissiveColor;
 
-	// TODO: investigate if the fixed size is high enough to properly calculate overlap
-	const FIntPoint TextureSize(512, 512);
-	const FMaterialPropertyEx Property = MP_EmissiveColor;
+		FMeshData MeshSet;
+		MeshSet.TextureCoordinateBox = { { 0.0f, 0.0f }, { 1.0f, 1.0f } };
+		MeshSet.TextureCoordinateIndex = TexCoord;
+		MeshSet.RawMeshDescription = const_cast<FMeshDescription*>(Description);
+		MeshSet.MaterialIndices = SectionIndices; // NOTE: MaterialIndices is actually section indices
 
-	FMeshData MeshSet;
-	MeshSet.TextureCoordinateBox = { { 0.0f, 0.0f }, { 1.0f, 1.0f } };
-	MeshSet.TextureCoordinateIndex = TexCoord;
-	MeshSet.RawMeshDescription = const_cast<FMeshDescription*>(Description);
-	MeshSet.MaterialIndices = SectionIndices; // NOTE: MaterialIndices is actually section indices
+		FMaterialDataEx MatSet;
+		MatSet.Material = GetMaterial();
+		MatSet.PropertySizes.Add(Property, TextureSize);
+		MatSet.BlendMode = MatSet.Material->GetBlendMode();
+		MatSet.bPerformBorderSmear = false;
 
-	FMaterialDataEx MatSet;
-	MatSet.Material = GetMaterial();
-	MatSet.PropertySizes.Add(Property, TextureSize);
-	MatSet.BlendMode = MatSet.Material->GetBlendMode();
-	MatSet.bPerformBorderSmear = false;
+		TArray<FMeshData*> MeshSettings;
+		TArray<FMaterialDataEx*> MatSettings;
+		MeshSettings.Add(&MeshSet);
+		MatSettings.Add(&MatSet);
 
-	TArray<FMeshData*> MeshSettings;
-	TArray<FMaterialDataEx*> MatSettings;
-	MeshSettings.Add(&MeshSet);
-	MatSettings.Add(&MatSet);
+		TArray<FBakeOutputEx> BakeOutputs;
+		IMaterialBakingModule& Module = FModuleManager::Get().LoadModuleChecked<IMaterialBakingModule>("GLTFMaterialBaking");
 
-	TArray<FBakeOutputEx> BakeOutputs;
-	IMaterialBakingModule& Module = FModuleManager::Get().LoadModuleChecked<IMaterialBakingModule>("GLTFMaterialBaking");
+		Module.BakeMaterials(MatSettings, MeshSettings, BakeOutputs);
 
-	Module.BakeMaterials(MatSettings, MeshSettings, BakeOutputs);
+		const FBakeOutputEx& BakeOutput = BakeOutputs[0];
+		const TArray<FColor>& BakedPixels = BakeOutput.PropertyData.FindChecked(Property);
 
-	const FBakeOutputEx& BakeOutput = BakeOutputs[0];
-	const TArray<FColor>& BakedPixels = BakeOutput.PropertyData.FindChecked(Property);
+		if (BakedPixels.Num() <= 0)
+		{
+			return -1;
+		}
 
-	if (BakedPixels.Num() <= 0)
-	{
-		return -1;
-	}
+		// NOTE: the emissive value of each pixel will be incremented by 10 each time a triangle is drawn on it.
+		// Therefore a value of 0.0 indicates an unreferenced pixel, a value of 10.0 indicates a uniquely referenced pixel,
+		// and a value of 20+ indicates an overlapping pixel.
+		// The value may differ somewhat from what was specified in the material due to conversions (color-space, gamma, float => uint etc),
+		// so to be safe we set the limit at 15 instead of 10 to prevent incorrectly flagging some pixels as overlapping.
+		const float EmissiveThreshold = 15.0f;
 
-	// NOTE: the emissive value of each pixel will be incremented by 10 each time a triangle is drawn on it.
-	// Therefore a value of 0.0 indicates an unreferenced pixel, a value of 10.0 indicates a uniquely referenced pixel,
-	// and a value of 20+ indicates an overlapping pixel.
-	// The value may differ somewhat from what was specified in the material due to conversions (color-space, gamma, float => uint etc),
-	// so to be safe we set the limit at 15 instead of 10 to prevent incorrectly flagging some pixels as overlapping.
-	const float EmissiveThreshold = 15.0f;
+		int32 ValidCount = 0;
+		int32 OverlapCount = 0;
 
-	int32 ValidCount = 0;
-	int32 OverlapCount = 0;
+		if (BakeOutput.EmissiveScale > EmissiveThreshold)
+		{
+			// This should never overflow since we know that the quotient will be <= 1.0
+			const uint8 ColorThreshold = FMath::RoundToInt((EmissiveThreshold / BakeOutput.EmissiveScale) * 255.0f);
 
-	if (BakeOutput.EmissiveScale > EmissiveThreshold)
-	{
-		// This should never overflow since we know that the quotient will be <= 1.0
-		const uint8 ColorThreshold = FMath::RoundToInt((EmissiveThreshold / BakeOutput.EmissiveScale) * 255.0f);
+			bool bIsMagenta;
+			bool bIsOverlapping;
 
-		bool bIsMagenta;
-		bool bIsOverlapping;
+			for (const FColor& Pixel: BakedPixels)
+			{
+				bIsMagenta = Pixel == FColor::Magenta;
+				bIsOverlapping = Pixel.G > ColorThreshold;
 
-		for (const FColor& Pixel: BakedPixels)
-		{
-			bIsMagenta = Pixel == FColor::Magenta;
-			bIsOverlapping = Pixel.G > ColorThreshold;
+				ValidCount += !bIsMagenta ? 1 : 0;
+				OverlapCount += !bIsMagenta && bIsOverlapping ? 1 : 0;
+			}
+		}
 
-			ValidCount += !bIsMagenta ? 1 : 0;
-			OverlapCount += !bIsMagenta && bIsOverlapping ? 1 : 0;
+		if (ValidCount == 0)
+		{
+			return -1;
 		}
-	}
 
-	if (ValidCount == 0)
-	{
-		return -1;
-	}
+		if (ValidCount == OverlapCount)
+		{
+			return 1;
+		}
 
-	if (ValidCount == OverlapCount)
-	{
-		return 1;
+		return static_cast<float>(OverlapCount) / static_cast<float>(ValidCount);
 	}
+#endif
 
-	return static_cast<float>(OverlapCount) / static_cast<float>(ValidCount);
+	return -1;
 }
 
 UMaterialInterface* FGLTFUVOverlapChecker::GetMaterial()
@@ -122,5 +125,3 @@ UMaterialInterface* FGLTFUVOverlapChecker::GetMaterial()
 	check(Material);
 	return Material;
 }
-
-#endif
diff --git a/Engine/Plugins/Enterprise/GLTFExporter/MigratedFromGit/Source/GLTFExporter/Private/Converters/GLTFUVOverlapChecker.h b/Engine/Plugins/Enterprise/GLTFExporter/MigratedFromGit/Source/GLTFExporter/Private/Converters/GLTFUVOverlapChecker.h
--- a/Engine/Plugins/Enterprise/GLTFExporter/MigratedFromGit/Source/GLTFExporter/Private/Converters/GLTFUVOverlapChecker.h
+++ b/Engine/Plugins/Enterprise/GLTFExporter/MigratedFromGit/Source/GLTFExporter/Private/Converters/GLTFUVOverlapChecker.h
@@ -2,11 +2,11 @@
 
 #pragma once
 
-#if WITH_EDITOR
-
 #include "Converters/GLTFConverter.h"
 #include "Converters/GLTFIndexArray.h"
 
+struct FMeshDescription;
+
 class FGLTFUVOverlapChecker final : public TGLTFConverter<float, const FMeshDescription*, FGLTFIndexArray, int32>
 {
 	virtual void Sanitize(const FMeshDescription*& Description, FGLTFIndexArray& SectionIndices, int32& TexCoord) override;
@@ -15,5 +15,3 @@ class FGLTFUVOverlapChecker final : public TGLTFConverter<float, const FMeshDesc
 
 	static UMaterialInterface* GetMaterial();
 };
-
-#endif
diff --git a/Engine/Plugins/Enterprise/GLTFExporter/MigratedFromGit/Source/GLTFExporter/Private/Tasks/GLTFMaterialTasks.cpp b/Engine/Plugins/Enterprise/GLTFExporter/MigratedFromGit/Source/GLTFExporter/Private/Tasks/GLTFMaterialTasks.cpp
--- a/Engine/Plugins/Enterprise/GLTFExporter/MigratedFromGit/Source/GLTFExporter/Private/Tasks/GLTFMaterialTasks.cpp
+++ b/Engine/Plugins/Enterprise/GLTFExporter/MigratedFromGit/Source/GLTFExporter/Private/Tasks/GLTFMaterialTasks.cpp
@@ -1,13 +1,11 @@
 // Copyright Epic Games, Inc. All Rights Reserved.
 
-#if WITH_EDITOR
-
 #include "Tasks/GLTFMaterialTasks.h"
 #include "Converters/GLTFConverterUtility.h"
 #include "Converters/GLTFNameUtility.h"
 #include "Converters/GLTFMaterialUtility.h"
 #include "Builders/GLTFContainerBuilder.h"
-#include "MaterialPropertyEx.h"
+#if WITH_EDITOR
 #include "Materials/MaterialInstance.h"
 #include "Materials/MaterialExpressionConstant.h"
 #include "Materials/MaterialExpressionConstant2Vector.h"
@@ -17,6 +15,7 @@
 #include "Materials/MaterialExpressionVectorParameter.h"
 #include "Materials/MaterialExpressionTextureSample.h"
 #include "Materials/MaterialExpressionTextureSampleParameter2D.h"
+#endif
 
 namespace
 {
@@ -72,6 +71,7 @@ void FGLTFMaterialTask::Complete()
 		return;
 	}
 
+#if WITH_EDITOR
 	if (JsonMaterial.ShadingModel != EGLTFJsonShadingModel::None)
 	{
 		const FMaterialPropertyEx BaseColorProperty = JsonMaterial.ShadingModel == EGLTFJsonShadingModel::Unlit ? MP_EmissiveColor : MP_BaseColor;
@@ -190,6 +190,24 @@ void FGLTFMaterialTask::Complete()
 				*SectionString));
 		}
 	}
+#endif
+}
+
+FString FGLTFMaterialTask::GetMaterialName() const
+{
+	FString MaterialName = Material->GetName();
+
+	if (MeshData != nullptr)
+	{
+		MaterialName += TEXT("_") + MeshData->Name;
+	}
+
+	return MaterialName;
+}
+
+FString FGLTFMaterialTask::GetBakedTextureName(const FString& PropertyName) const
+{
+	return GetMaterialName() + TEXT("_") + PropertyName;
 }
 
 void FGLTFMaterialTask::ApplyPrebakedProperties(FGLTFJsonMaterial& OutMaterial) const
@@ -301,6 +319,7 @@ EMaterialShadingModel FGLTFMaterialTask::GetShadingModel() const
 
 	if (PossibilitiesCount > 1)
 	{
+#if WITH_EDITOR
 		if (!FApp::CanEverRender())
 		{
 			const EMaterialShadingModel ShadingModel = FGLTFMaterialUtility::GetRichestShadingModel(Possibilities);
@@ -341,6 +360,14 @@ EMaterialShadingModel FGLTFMaterialTask::GetShadingModel() const
 		}
 
 		// we should never end up here
+#else
+		const EMaterialShadingModel ShadingModel = FGLTFMaterialUtility::GetRichestShadingModel(Possibilities);
+		Builder.LogWarning(FString::Printf(
+			TEXT("Can't evaluate shading model expression in material %s without editor, will export as %s"),
+			*Material->GetName(),
+			*FGLTFNameUtility::GetName(ShadingModel)));
+		return ShadingModel;
+#endif
 	}
 
 	return Possibilities.GetFirstShadingModel();
@@ -428,6 +455,8 @@ void FGLTFMaterialTask::ConvertAlphaMode(EGLTFJsonAlphaMode& OutAlphaMode, EGLTF
 	}
 }
 
+#if WITH_EDITOR
+
 bool FGLTFMaterialTask::TryGetBaseColorAndOpacity(FGLTFJsonPBRMetallicRoughness& OutPBRParams, const FMaterialPropertyEx& BaseColorProperty, const FMaterialPropertyEx& OpacityProperty)
 {
 	const bool bIsBaseColorConstant = TryGetConstantColor(OutPBRParams.BaseColorFactor, BaseColorProperty);
@@ -1489,23 +1518,6 @@ bool FGLTFMaterialTask::StoreBakedPropertyTexture(FGLTFJsonTextureInfo& OutTexIn
 	return true;
 }
 
-FString FGLTFMaterialTask::GetMaterialName() const
-{
-	FString MaterialName = Material->GetName();
-
-	if (MeshData != nullptr)
-	{
-		MaterialName += TEXT("_") + MeshData->Name;
-	}
-
-	return MaterialName;
-}
-
-FString FGLTFMaterialTask::GetBakedTextureName(const FString& PropertyName) const
-{
-	return GetMaterialName() + TEXT("_") + PropertyName;
-}
-
 EGLTFMaterialPropertyGroup FGLTFMaterialTask::GetPropertyGroup(const FMaterialPropertyEx& Property)
 {
 	switch (Property.Type)
diff --git a/Engine/Plugins/Enterprise/GLTFExporter/MigratedFromGit/Source/GLTFExporter/Private/Tasks/GLTFMaterialTasks.h b/Engine/Plugins/Enterprise/GLTFExporter/MigratedFromGit/Source/GLTFExporter/Private/Tasks/GLTFMaterialTasks.h
--- a/Engine/Plugins/Enterprise/GLTFExporter/MigratedFromGit/Source/GLTFExporter/Private/Tasks/GLTFMaterialTasks.h
+++ b/Engine/Plugins/Enterprise/GLTFExporter/MigratedFromGit/Source/GLTFExporter/Private/Tasks/GLTFMaterialTasks.h
@@ -2,13 +2,13 @@
 
 #pragma once
 
-#if WITH_EDITOR
-
 #include "Tasks/GLTFTask.h"
 #include "Builders/GLTFConvertBuilder.h"
 #include "Converters/GLTFUVOverlapChecker.h"
 #include "Converters/GLTFMeshData.h"
+#if WITH_EDITOR
 #include "MaterialPropertyEx.h"
+#endif
 
 struct FGLTFPropertyBakeOutput;
 
@@ -43,7 +43,8 @@ private:
 	const FGLTFIndexArray SectionIndices;
 	const FGLTFJsonMaterialIndex MaterialIndex;
 
-	TSet<FMaterialPropertyEx> MeshDataBakedProperties;
+	FString GetMaterialName() const;
+	FString GetBakedTextureName(const FString& PropertyName) const;
 
 	void ApplyPrebakedProperties(FGLTFJsonMaterial& OutMaterial) const;
 	void ApplyPrebakedProperty(const FString& PropertyName, float& OutValue) const;
@@ -55,6 +56,9 @@ private:
 	void ConvertShadingModel(EGLTFJsonShadingModel& OutShadingModel) const;
 	void ConvertAlphaMode(EGLTFJsonAlphaMode& OutAlphaMode, EGLTFJsonBlendMode& OutBlendMode) const;
 
+#if WITH_EDITOR
+	TSet<FMaterialPropertyEx> MeshDataBakedProperties;
+
 	bool TryGetBaseColorAndOpacity(FGLTFJsonPBRMetallicRoughness& OutPBRParams, const FMaterialPropertyEx& BaseColorProperty, const FMaterialPropertyEx& OpacityProperty);
 	bool TryGetMetallicAndRoughness(FGLTFJsonPBRMetallicRoughness& OutPBRParams, const FMaterialPropertyEx& MetallicProperty, const FMaterialPropertyEx& RoughnessProperty);
 	bool TryGetClearCoatRoughness(FGLTFJsonClearCoatExtension& OutExtParams, const FMaterialPropertyEx& IntensityProperty, const FMaterialPropertyEx& RoughnessProperty);
@@ -79,13 +83,10 @@ private:
 
 	bool StoreBakedPropertyTexture(FGLTFJsonTextureInfo& OutTexInfo, FGLTFPropertyBakeOutput& PropertyBakeOutput, const FString& PropertyName) const;
 
-	FString GetMaterialName() const;
-	FString GetBakedTextureName(const FString& PropertyName) const;
-
 	static EGLTFMaterialPropertyGroup GetPropertyGroup(const FMaterialPropertyEx& Property);
 
 	template <typename CallbackType>
 	static void CombinePixels(const TArray<FColor>& FirstPixels, const TArray<FColor>& SecondPixels, TArray<FColor>& OutPixels, CallbackType Callback);
-};
 
 #endif
+};
