[GLTFExporter] refactor: avoid hardcoding prebaked base material mapping

commit 07722f536ef2cb9e784a293a684be9af4a5a4ee1
Date: Tue, 8 Feb 2022 16:53:15 +0100

    refactor: avoid hardcoding prebaked base material mapping
    (cherry picked from commit fe3934b666d243b3a7eb8a7d97349bfed7627df2​)

#skipci

diff --git a/Engine/Plugins/Enterprise/GLTFExporter/MigratedFromGit/Source/GLTFExporter/Private/Converters/GLTFMaterialUtility.cpp b/Engine/Plugins/Enterprise/GLTFExporter/MigratedFromGit/Source/GLTFExporter/Private/Converters/GLTFMaterialUtility.cpp
--- a/Engine/Plugins/Enterprise/GLTFExporter/MigratedFromGit/Source/GLTFExporter/Private/Converters/GLTFMaterialUtility.cpp
+++ b/Engine/Plugins/Enterprise/GLTFExporter/MigratedFromGit/Source/GLTFExporter/Private/Converters/GLTFMaterialUtility.cpp
@@ -20,19 +20,15 @@
 
 UMaterialInterface* FGLTFMaterialUtility::GetDefault()
 {
-	static UMaterialInterface* DefaultMaterial = LoadObject<UMaterialInterface>(nullptr, TEXT("/GLTFExporter/Materials/GLTFDefault.GLTFDefault"));
+	static UMaterialInterface* DefaultMaterial = GetPrebaked(EGLTFJsonShadingModel::Default);
 	return DefaultMaterial;
 }
 
 UMaterialInterface* FGLTFMaterialUtility::GetPrebaked(EGLTFJsonShadingModel ShadingModel)
 {
-	// TODO: don't assume enum values are in this order
-	static UMaterialInterface* BaseMaterials[] = {
-		LoadObject<UMaterialInterface>(nullptr, TEXT("/GLTFExporter/Materials/GLTFDefault.GLTFDefault")),
-		LoadObject<UMaterialInterface>(nullptr, TEXT("/GLTFExporter/Materials/GLTFUnlit.GLTFUnlit")),
-		LoadObject<UMaterialInterface>(nullptr, TEXT("/GLTFExporter/Materials/GLTFClearCoat.GLTFClearCoat")),
-	};
-	return BaseMaterials[static_cast<int32>(ShadingModel)];
+	const FString Name = FGLTFJsonUtility::GetValue(ShadingModel);
+	const FString Path = TEXT("/GLTFExporter/Materials/GLTF") + Name + TEXT(".GLTF") + Name;
+	return LoadObject<UMaterialInterface>(nullptr, *Path);
 }
 
 bool FGLTFMaterialUtility::IsPrebaked(const UMaterial* Material)
diff --git a/Engine/Plugins/Enterprise/GLTFExporter/MigratedFromGit/Source/GLTFExporter/Private/Json/GLTFJsonUtility.h b/Engine/Plugins/Enterprise/GLTFExporter/MigratedFromGit/Source/GLTFExporter/Private/Json/GLTFJsonUtility.h
--- a/Engine/Plugins/Enterprise/GLTFExporter/MigratedFromGit/Source/GLTFExporter/Private/Json/GLTFJsonUtility.h
+++ b/Engine/Plugins/Enterprise/GLTFExporter/MigratedFromGit/Source/GLTFExporter/Private/Json/GLTFJsonUtility.h
@@ -183,4 +183,17 @@ struct FGLTFJsonUtility
 				return TEXT("");
 		}
 	}
+
+	static const TCHAR* GetValue(EGLTFJsonShadingModel Enum)
+	{
+		switch (Enum)
+		{
+			case EGLTFJsonShadingModel::Default:   return TEXT("Default");
+			case EGLTFJsonShadingModel::Unlit:     return TEXT("Unlit");
+			case EGLTFJsonShadingModel::ClearCoat: return TEXT("ClearCoat");
+			default:
+				checkNoEntry();
+			return TEXT("");
+		}
+	}
 };
