[GLTFExporter] docs: add documentation for new export option bAdjustNormalmaps

commit f4ab81a27705c2afe7d822b0550507ea83e2487b
Date: Tue, 8 Feb 2022 17:00:17 +0100

    docs: add documentation for new export option bAdjustNormalmaps
    (cherry picked from commit 3fc1c3df457af58477880abbe4b70b2cfb238a27​)

#skipci

diff --git a/Engine/Plugins/Enterprise/GLTFExporter/MigratedFromGit/Docs/export-options-reference.md b/Engine/Plugins/Enterprise/GLTFExporter/MigratedFromGit/Docs/export-options-reference.md
--- a/Engine/Plugins/Enterprise/GLTFExporter/MigratedFromGit/Docs/export-options-reference.md
+++ b/Engine/Plugins/Enterprise/GLTFExporter/MigratedFromGit/Docs/export-options-reference.md
@@ -29,6 +29,7 @@ Option                         | Description
 `Export Texture Transforms`    | If enabled, export UV tiling and un-mirroring settings in a texture coordinate expression node for simple material input expressions. Uses extension KHR_texture_transform.
 `Export Lightmaps`             | If enabled, export lightmaps (created by Lightmass) when exporting a level. Uses extension EPIC_lightmap_textures, which is supported by Unreal's glTF viewer.
 `Texture HDR Encoding`         | Encoding used to store textures that have pixel colors with more than 8-bit per channel. Uses extension EPIC_texture_hdr_encoding, which is supported by Unreal's glTF viewer.
+`Adjust Normalmaps`            | If enabled, exported normalmaps will be adjusted from Unreal to glTF convention (i.e. the green channel is flipped).
 `Export Hidden In Game`        | If enabled, export actors and components that are flagged as hidden in-game.
 `Export Lights`                | Mobility of directional, point, and spot light components that will be exported. Uses extension KHR_lights_punctual.
 `Export Cameras`               | If enabled, export camera components.
diff --git a/Engine/Plugins/Enterprise/GLTFExporter/MigratedFromGit/Source/GLTFExporter/Public/Options/GLTFExportOptions.h b/Engine/Plugins/Enterprise/GLTFExporter/MigratedFromGit/Source/GLTFExporter/Public/Options/GLTFExportOptions.h
--- a/Engine/Plugins/Enterprise/GLTFExporter/MigratedFromGit/Source/GLTFExporter/Public/Options/GLTFExportOptions.h
+++ b/Engine/Plugins/Enterprise/GLTFExporter/MigratedFromGit/Source/GLTFExporter/Public/Options/GLTFExportOptions.h
@@ -187,6 +187,7 @@ class GLTFEXPORTER_API UGLTFExportOptions : public UObject
 	UPROPERTY(EditAnywhere, BlueprintReadWrite, Config, Category = Texture, Meta = (DisplayName = "Texture HDR Encoding", EditCondition = "TextureImageFormat != EGLTFTextureImageFormat::None"))
 	EGLTFTextureHDREncoding TextureHDREncoding;
 
+	/** If enabled, exported normalmaps will be adjusted from Unreal to glTF convention (i.e. the green channel is flipped). */
 	UPROPERTY(EditAnywhere, BlueprintReadWrite, Config, Category = Texture, Meta = (EditCondition = "TextureImageFormat != EGLTFTextureImageFormat::None"))
 	bool bAdjustNormalmaps;
 
