[GLTFExporter] fix: add error handling if shading model unsupported by prebaking

commit bc0b630b53252477dce012b17b71ff3b563f5c50
Date: Tue, 8 Feb 2022 17:06:13 +0100

    fix: add error handling if shading model unsupported by prebaking
    (cherry picked from commit 1021bc639f60b851e6b6b2e20ecda852370ae5bc​)

#skipci

diff --git a/Engine/Plugins/Enterprise/GLTFExporter/MigratedFromGit/Source/GLTFExporter/Private/Converters/GLTFMaterialPrebaker.cpp b/Engine/Plugins/Enterprise/GLTFExporter/MigratedFromGit/Source/GLTFExporter/Private/Converters/GLTFMaterialPrebaker.cpp
--- a/Engine/Plugins/Enterprise/GLTFExporter/MigratedFromGit/Source/GLTFExporter/Private/Converters/GLTFMaterialPrebaker.cpp
+++ b/Engine/Plugins/Enterprise/GLTFExporter/MigratedFromGit/Source/GLTFExporter/Private/Converters/GLTFMaterialPrebaker.cpp
@@ -33,6 +33,7 @@ UMaterialInterface* FGLTFMaterialPrebaker::Prebake(UMaterialInterface* OriginalM
 	const FGLTFJsonMaterialIndex MaterialIndex = Builder.GetOrAddMaterial(OriginalMaterial);
 	if (MaterialIndex == INDEX_NONE)
 	{
+		// TODO: report error
 		return nullptr;
 	}
 
@@ -42,7 +43,10 @@ UMaterialInterface* FGLTFMaterialPrebaker::Prebake(UMaterialInterface* OriginalM
 	const FGLTFJsonMaterial& JsonMaterial = Builder.GetMaterial(MaterialIndex);
 
 	UMaterialInstanceConstant* ProxyMaterial = CreateProxyMaterial(OriginalMaterial, JsonMaterial.ShadingModel);
-	ApplyPrebakedProperties(ProxyMaterial, JsonMaterial);
+	if (ProxyMaterial != nullptr)
+	{
+		ApplyPrebakedProperties(ProxyMaterial, JsonMaterial);
+	}
 
 	return ProxyMaterial;
 }
@@ -219,7 +223,17 @@ UMaterialInstanceConstant* FGLTFMaterialPrebaker::CreateProxyMaterial(UMaterialI
 	Package->Modify();
 
 	UMaterialInterface* BaseMaterial = FGLTFMaterialUtility::GetPrebaked(ShadingModel);
-	UMaterialInstanceConstant* ProxyMaterial = FMaterialUtilities::CreateInstancedMaterial(BaseMaterial, Package, FPaths::GetCleanFilename(Package->GetName()), RF_Public | RF_Standalone);
+	if (BaseMaterial == nullptr)
+	{
+		Builder.LogError(FString::Printf(
+			TEXT("Material %s uses a shading model (%s) that doesn't have a prebaked base material"),
+			*OriginalMaterial->GetName(),
+			FGLTFJsonUtility::GetValue(ShadingModel)));
+		return nullptr;
+	}
+
+	const FString BaseName = TEXT("GLTF_") + OriginalMaterial->GetName(); // prefix "M_" is added automatically by CreateInstancedMaterial
+	UMaterialInstanceConstant* ProxyMaterial = FMaterialUtilities::CreateInstancedMaterial(BaseMaterial, Package, BaseName, RF_Public | RF_Standalone);
 
 	const bool bTwoSided = OriginalMaterial->IsTwoSided();
 	if (bTwoSided != BaseMaterial->IsTwoSided())
