[GLTFExporter] fix: correctly handle sRGB and linear samping of non-color textures in prebaked materials

commit 4eb9e5dff081ee933bf55d9e8ca854f245879762
Date: Tue, 8 Feb 2022 17:07:00 +0100

    fix: correctly handle sRGB and linear samping of non-color textures in prebaked materials
    (cherry picked from commit 5b35097464a6e1e9c7d1eb1e578f2fdffc445717​)

#skipci

diff --git a/Engine/Plugins/Enterprise/GLTFExporter/MigratedFromGit/Source/GLTFExporter/Private/Converters/GLTFMaterialPrebaker.cpp b/Engine/Plugins/Enterprise/GLTFExporter/MigratedFromGit/Source/GLTFExporter/Private/Converters/GLTFMaterialPrebaker.cpp
--- a/Engine/Plugins/Enterprise/GLTFExporter/MigratedFromGit/Source/GLTFExporter/Private/Converters/GLTFMaterialPrebaker.cpp
+++ b/Engine/Plugins/Enterprise/GLTFExporter/MigratedFromGit/Source/GLTFExporter/Private/Converters/GLTFMaterialPrebaker.cpp
@@ -54,33 +54,33 @@ UMaterialInterface* FGLTFMaterialPrebaker::Prebake(UMaterialInterface* OriginalM
 void FGLTFMaterialPrebaker::ApplyPrebakedProperties(UMaterialInstanceConstant* ProxyMaterial, const FGLTFJsonMaterial& JsonMaterial)
 {
 	ApplyPrebakedProperty(ProxyMaterial, TEXT("Base Color Factor"), JsonMaterial.PBRMetallicRoughness.BaseColorFactor);
-	ApplyPrebakedProperty(ProxyMaterial, TEXT("Base Color"), JsonMaterial.PBRMetallicRoughness.BaseColorTexture);
+	ApplyPrebakedProperty(ProxyMaterial, TEXT("Base Color"), JsonMaterial.PBRMetallicRoughness.BaseColorTexture, EGLTFMaterialPropertyGroup::BaseColorOpacity);
 
 	if (JsonMaterial.ShadingModel == EGLTFJsonShadingModel::Default || JsonMaterial.ShadingModel == EGLTFJsonShadingModel::ClearCoat)
 	{
 		ApplyPrebakedProperty(ProxyMaterial, TEXT("Emissive Factor"), JsonMaterial.EmissiveFactor);
-		ApplyPrebakedProperty(ProxyMaterial, TEXT("Emissive"), JsonMaterial.EmissiveTexture);
+		ApplyPrebakedProperty(ProxyMaterial, TEXT("Emissive"), JsonMaterial.EmissiveTexture, EGLTFMaterialPropertyGroup::EmissiveColor);
 
 		ApplyPrebakedProperty(ProxyMaterial, TEXT("Metallic Factor"), JsonMaterial.PBRMetallicRoughness.MetallicFactor);
 		ApplyPrebakedProperty(ProxyMaterial, TEXT("Roughness Factor"), JsonMaterial.PBRMetallicRoughness.RoughnessFactor);
-		ApplyPrebakedProperty(ProxyMaterial, TEXT("Metallic Roughness"), JsonMaterial.PBRMetallicRoughness.MetallicRoughnessTexture);
+		ApplyPrebakedProperty(ProxyMaterial, TEXT("Metallic Roughness"), JsonMaterial.PBRMetallicRoughness.MetallicRoughnessTexture, EGLTFMaterialPropertyGroup::MetallicRoughness);
 
 		ApplyPrebakedProperty(ProxyMaterial, TEXT("Normal Scale"), JsonMaterial.NormalTexture.Scale);
-		ApplyPrebakedProperty(ProxyMaterial, TEXT("Normal"), JsonMaterial.NormalTexture, true);
+		ApplyPrebakedProperty(ProxyMaterial, TEXT("Normal"), JsonMaterial.NormalTexture, EGLTFMaterialPropertyGroup::Normal);
 
 		ApplyPrebakedProperty(ProxyMaterial, TEXT("Occlusion Strength"), JsonMaterial.OcclusionTexture.Strength);
-		ApplyPrebakedProperty(ProxyMaterial, TEXT("Occlusion"), JsonMaterial.OcclusionTexture);
+		ApplyPrebakedProperty(ProxyMaterial, TEXT("Occlusion"), JsonMaterial.OcclusionTexture, EGLTFMaterialPropertyGroup::AmbientOcclusion);
 
 		if (JsonMaterial.ShadingModel == EGLTFJsonShadingModel::ClearCoat)
 		{
 			ApplyPrebakedProperty(ProxyMaterial, TEXT("Clear Coat Factor"), JsonMaterial.ClearCoat.ClearCoatFactor);
-			ApplyPrebakedProperty(ProxyMaterial, TEXT("Clear Coat"), JsonMaterial.ClearCoat.ClearCoatTexture);
+			ApplyPrebakedProperty(ProxyMaterial, TEXT("Clear Coat"), JsonMaterial.ClearCoat.ClearCoatTexture, EGLTFMaterialPropertyGroup::ClearCoatRoughness); // TODO: add property group for clear coat intensity only
 
 			ApplyPrebakedProperty(ProxyMaterial, TEXT("Clear Coat Roughness Factor"), JsonMaterial.ClearCoat.ClearCoatRoughnessFactor);
-			ApplyPrebakedProperty(ProxyMaterial, TEXT("Clear Coat Roughness"), JsonMaterial.ClearCoat.ClearCoatRoughnessTexture);
+			ApplyPrebakedProperty(ProxyMaterial, TEXT("Clear Coat Roughness"), JsonMaterial.ClearCoat.ClearCoatRoughnessTexture, EGLTFMaterialPropertyGroup::ClearCoatRoughness);
 
 			ApplyPrebakedProperty(ProxyMaterial, TEXT("Clear Coat Normal Scale"), JsonMaterial.ClearCoat.ClearCoatNormalTexture.Scale);
-			ApplyPrebakedProperty(ProxyMaterial, TEXT("Clear Coat Normal"), JsonMaterial.ClearCoat.ClearCoatNormalTexture, true);
+			ApplyPrebakedProperty(ProxyMaterial, TEXT("Clear Coat Normal"), JsonMaterial.ClearCoat.ClearCoatNormalTexture, EGLTFMaterialPropertyGroup::ClearCoatBottomNormal);
 		}
 	}
 
@@ -134,11 +134,11 @@ void FGLTFMaterialPrebaker::ApplyPrebakedProperty(UMaterialInstanceConstant* Pro
 	}
 }
 
-void FGLTFMaterialPrebaker::ApplyPrebakedProperty(UMaterialInstanceConstant* ProxyMaterial, const FString& PropertyName, const FGLTFJsonTextureInfo& TextureInfo, bool bNormalMap)
+void FGLTFMaterialPrebaker::ApplyPrebakedProperty(UMaterialInstanceConstant* ProxyMaterial, const FString& PropertyName, const FGLTFJsonTextureInfo& TextureInfo, EGLTFMaterialPropertyGroup PropertyGroup)
 {
 	if (TextureInfo.Index != INDEX_NONE)
 	{
-		const UTexture* Texture = FindOrCreateTexture(TextureInfo.Index, bNormalMap);
+		const UTexture* Texture = FindOrCreateTexture(TextureInfo.Index, PropertyGroup);
 		if (Texture == nullptr)
 		{
 			// TODO: report error
@@ -170,8 +170,10 @@ void FGLTFMaterialPrebaker::ApplyPrebakedProperty(UMaterialInstanceConstant* Pro
 	}
 }
 
-const UTexture2D* FGLTFMaterialPrebaker::FindOrCreateTexture(FGLTFJsonTextureIndex Index, bool bNormalMap)
+const UTexture2D* FGLTFMaterialPrebaker::FindOrCreateTexture(FGLTFJsonTextureIndex Index, EGLTFMaterialPropertyGroup PropertyGroup)
 {
+	// TODO: fix potential conflict when same texture used for different material properties that have different encoding (sRGB vs Linear, Normalmap compression etc)
+
 	const UTexture2D** FoundPtr = Textures.Find(Index);
 	if (FoundPtr != nullptr)
 	{
@@ -186,13 +188,16 @@ const UTexture2D* FGLTFMaterialPrebaker::FindOrCreateTexture(FGLTFJsonTextureInd
 	}
 
 	const FGLTFJsonSampler& JsonSampler = Builder.GetSampler(JsonTexture.Sampler);
-	const UTexture2D* Texture = CreateTexture(ImageData, JsonSampler, bNormalMap);
+	const UTexture2D* Texture = CreateTexture(ImageData, JsonSampler, PropertyGroup);
 	Textures.Add(Index, Texture);
 	return Texture;
 }
 
-UTexture2D* FGLTFMaterialPrebaker::CreateTexture(const FGLTFImageData* ImageData, const FGLTFJsonSampler& JsonSampler, bool bNormalMap)
+UTexture2D* FGLTFMaterialPrebaker::CreateTexture(const FGLTFImageData* ImageData, const FGLTFJsonSampler& JsonSampler, EGLTFMaterialPropertyGroup PropertyGroup)
 {
+	const bool bSRGB = PropertyGroup == EGLTFMaterialPropertyGroup::BaseColorOpacity || PropertyGroup == EGLTFMaterialPropertyGroup::EmissiveColor;
+	const bool bNormalMap = PropertyGroup == EGLTFMaterialPropertyGroup::Normal || PropertyGroup == EGLTFMaterialPropertyGroup::ClearCoatBottomNormal;
+
 	const FString PackageName = RootPath / TEXT("T_GLTF_") + ImageData->Filename;
 	UPackage* Package = CreatePackage(*PackageName);
 	Package->FullyLoad();
@@ -200,16 +205,16 @@ UTexture2D* FGLTFMaterialPrebaker::CreateTexture(const FGLTFImageData* ImageData
 
 	FCreateTexture2DParameters TexParams;
 	TexParams.bUseAlpha = !ImageData->bIgnoreAlpha;
-	TexParams.bSRGB = false;
+	TexParams.CompressionSettings = bNormalMap ? TC_Normalmap : TC_Default;
 	TexParams.bDeferCompression = true;
+	TexParams.bSRGB = bSRGB;
+	TexParams.TextureGroup = bNormalMap ? TEXTUREGROUP_WorldNormalMap : TEXTUREGROUP_World;
 	TexParams.SourceGuidHash = FGuid();
-	TexParams.CompressionSettings = bNormalMap ? TC_Normalmap : TC_Default;
 
 	UTexture2D* Texture = FImageUtils::CreateTexture2D(ImageData->Size.X, ImageData->Size.Y, *ImageData->Pixels, Package, FPaths::GetCleanFilename(Package->GetName()), RF_Public | RF_Standalone, TexParams);
 	Texture->Filter = ConvertFilter(JsonSampler.MagFilter);
 	Texture->AddressX = ConvertWrap(JsonSampler.WrapS);
 	Texture->AddressY = ConvertWrap(JsonSampler.WrapT);
-	Texture->LODGroup = TEXTUREGROUP_World;
 
 	Texture->PostEditChange();
 	return Texture;
@@ -217,7 +222,7 @@ UTexture2D* FGLTFMaterialPrebaker::CreateTexture(const FGLTFImageData* ImageData
 
 UMaterialInstanceConstant* FGLTFMaterialPrebaker::CreateProxyMaterial(UMaterialInterface* OriginalMaterial, EGLTFJsonShadingModel ShadingModel)
 {
-	const FString PackageName = RootPath / TEXT("GLTF_") + OriginalMaterial->GetName();
+	const FString PackageName = RootPath / TEXT("M_GLTF_") + OriginalMaterial->GetName();
 	UPackage* Package = CreatePackage(*PackageName);
 	Package->FullyLoad();
 	Package->Modify();
diff --git a/Engine/Plugins/Enterprise/GLTFExporter/MigratedFromGit/Source/GLTFExporter/Private/Converters/GLTFMaterialPrebaker.h b/Engine/Plugins/Enterprise/GLTFExporter/MigratedFromGit/Source/GLTFExporter/Private/Converters/GLTFMaterialPrebaker.h
--- a/Engine/Plugins/Enterprise/GLTFExporter/MigratedFromGit/Source/GLTFExporter/Private/Converters/GLTFMaterialPrebaker.h
+++ b/Engine/Plugins/Enterprise/GLTFExporter/MigratedFromGit/Source/GLTFExporter/Private/Converters/GLTFMaterialPrebaker.h
@@ -36,10 +36,10 @@ private:
 	void ApplyPrebakedProperty(UMaterialInstanceConstant* ProxyMaterial, const FString& PropertyName, float Scalar);
 	void ApplyPrebakedProperty(UMaterialInstanceConstant* ProxyMaterial, const FString& PropertyName, const FGLTFJsonColor3& Color);
 	void ApplyPrebakedProperty(UMaterialInstanceConstant* ProxyMaterial, const FString& PropertyName, const FGLTFJsonColor4& Color);
-	void ApplyPrebakedProperty(UMaterialInstanceConstant* ProxyMaterial, const FString& PropertyName, const FGLTFJsonTextureInfo& TextureInfo, bool bNormalMap = false);
+	void ApplyPrebakedProperty(UMaterialInstanceConstant* ProxyMaterial, const FString& PropertyName, const FGLTFJsonTextureInfo& TextureInfo, EGLTFMaterialPropertyGroup PropertyGroup);
 
-	const UTexture2D* FindOrCreateTexture(FGLTFJsonTextureIndex Index, bool bNormalMap);
-	UTexture2D* CreateTexture(const FGLTFImageData* ImageData, const FGLTFJsonSampler& JsonSampler, bool bNormalMap);
+	const UTexture2D* FindOrCreateTexture(FGLTFJsonTextureIndex Index, EGLTFMaterialPropertyGroup PropertyGroup);
+	UTexture2D* CreateTexture(const FGLTFImageData* ImageData, const FGLTFJsonSampler& JsonSampler, EGLTFMaterialPropertyGroup PropertyGroup);
 
 	UMaterialInstanceConstant* CreateProxyMaterial(UMaterialInterface* OriginalMaterial, EGLTFJsonShadingModel ShadingModel);
 
