[GLTFExporter] refactor: use term material proxy instead of prebake/prebaked/prebaking

commit 73e140570be62e7567e3c8fd3f959af9c4c2e2b9
Date: Fri, 11 Feb 2022 17:27:18 +0100

    refactor: use term material proxy instead of prebake/prebaked/prebaking
    (cherry picked from commit ba2f43ac516ba6f83b3b0c5d1517bc583a933e62​)

#skipci

diff --git a/Engine/Plugins/Enterprise/GLTFExporter/MigratedFromGit/Source/GLTFExporter/Private/Actions/GLTFPrebakeAssetActions.cpp b/Engine/Plugins/Enterprise/GLTFExporter/MigratedFromGit/Source/GLTFExporter/Private/Actions/GLTFPrebakeAssetActions.cpp
deleted file mode 100644
diff --git a/Engine/Plugins/Enterprise/GLTFExporter/MigratedFromGit/Source/GLTFExporter/Private/Actions/GLTFPrebakeAssetActions.h b/Engine/Plugins/Enterprise/GLTFExporter/MigratedFromGit/Source/GLTFExporter/Private/Actions/GLTFPrebakeAssetActions.h
deleted file mode 100644
diff --git a/Engine/Plugins/Enterprise/GLTFExporter/MigratedFromGit/Source/GLTFExporter/Private/Actions/GLTFProxyAssetActions.cpp b/Engine/Plugins/Enterprise/GLTFExporter/MigratedFromGit/Source/GLTFExporter/Private/Actions/GLTFProxyAssetActions.cpp
new file mode 100644
--- /dev/null
+++ b/Engine/Plugins/Enterprise/GLTFExporter/MigratedFromGit/Source/GLTFExporter/Private/Actions/GLTFProxyAssetActions.cpp
@@ -0,0 +1,63 @@
+// Copyright Epic Games, Inc. All Rights Reserved.
+
+#if WITH_EDITOR
+
+#include "Actions/GLTFProxyAssetActions.h"
+#include "Actions/GLTFEditorStyle.h"
+#include "Options/GLTFProxyOptions.h"
+#include "UI/GLTFProxyOptionsWindow.h"
+#include "GLTFExporterModule.h"
+#include "ToolMenus.h"
+
+#define LOCTEXT_NAMESPACE "GLTFProxyAssetActions"
+
+FGLTFProxyAssetActions::FGLTFProxyAssetActions(const TSharedRef<IAssetTypeActions>& OriginalActions)
+	: OriginalActions(OriginalActions)
+{
+}
+
+bool FGLTFProxyAssetActions::HasActions(const TArray<UObject*>& InObjects) const
+{
+	return true;
+}
+
+void FGLTFProxyAssetActions::GetActions(const TArray<UObject*>& InObjects, FToolMenuSection& Section)
+{
+	OriginalActions->GetActions(InObjects, Section);
+	GetProxyActions(InObjects, Section);
+}
+
+void FGLTFProxyAssetActions::GetProxyActions(const TArray<UObject*>& InObjects, FToolMenuSection& Section)
+{
+	const TArray<FWeakObjectPtr> Objects(InObjects);
+
+	Section.AddMenuEntry(
+		"MenuEntry_CreateProxy",
+		LOCTEXT("MenuEntry_CreateProxy", "Create glTF Proxy Material"),
+		LOCTEXT("MenuEntry_CreateProxyTooltip", "Creates a proxy version of this material for glTF export."),
+		FSlateIcon(FGLTFEditorStyle::Get().GetStyleSetName(), "Icon16x16"),
+		FUIAction(FExecuteAction::CreateSP(this, &FGLTFProxyAssetActions::OnCreateProxy, Objects))
+		);
+}
+
+void FGLTFProxyAssetActions::OnCreateProxy(TArray<FWeakObjectPtr> Objects) const
+{
+	UGLTFProxyOptions* Options = NewObject<UGLTFProxyOptions>();
+
+	if (!SGLTFProxyOptionsWindow::ShowDialog(Options))
+	{
+		return;
+	}
+
+	for (const FWeakObjectPtr& Object : Objects)
+	{
+		if (UMaterialInterface* Material = Cast<UMaterialInterface>(Object.Get()))
+		{
+			IGLTFExporterModule::CreateProxyMaterial(Material, Options);
+		}
+	}
+}
+
+#undef LOCTEXT_NAMESPACE
+
+#endif
diff --git a/Engine/Plugins/Enterprise/GLTFExporter/MigratedFromGit/Source/GLTFExporter/Private/Actions/GLTFProxyAssetActions.h b/Engine/Plugins/Enterprise/GLTFExporter/MigratedFromGit/Source/GLTFExporter/Private/Actions/GLTFProxyAssetActions.h
new file mode 100644
--- /dev/null
+++ b/Engine/Plugins/Enterprise/GLTFExporter/MigratedFromGit/Source/GLTFExporter/Private/Actions/GLTFProxyAssetActions.h
@@ -0,0 +1,57 @@
+// Copyright Epic Games, Inc. All Rights Reserved.
+
+#pragma once
+
+#if WITH_EDITOR
+
+#include "AssetTypeActions_Base.h"
+
+class FGLTFProxyAssetActions : public IAssetTypeActions
+{
+public:
+
+	FGLTFProxyAssetActions(const TSharedRef<IAssetTypeActions>& OriginalActions);
+
+	virtual bool HasActions(const TArray<UObject*>& InObjects) const override;
+	virtual void GetActions(const TArray<UObject*>& InObjects, FToolMenuSection& Section) override;
+
+	virtual FText GetName() const override { return OriginalActions->GetName(); }
+	virtual UClass* GetSupportedClass() const override { return OriginalActions->GetSupportedClass(); }
+	virtual FColor GetTypeColor() const override { return OriginalActions->GetTypeColor(); }
+	virtual void GetActions(const TArray<UObject*>& InObjects, FMenuBuilder& MenuBuilder) override { OriginalActions->GetActions(InObjects, MenuBuilder); }
+	virtual void OpenAssetEditor(const TArray<UObject*>& InObjects, TSharedPtr<IToolkitHost> EditWithinLevelEditor) override { OriginalActions->OpenAssetEditor(InObjects, EditWithinLevelEditor); }
+	virtual void AssetsActivated(const TArray<UObject*>& InObjects, EAssetTypeActivationMethod::Type ActivationType) override { OriginalActions->AssetsActivated(InObjects, ActivationType); }
+	virtual bool AssetsActivatedOverride(const TArray<UObject*>& InObjects, EAssetTypeActivationMethod::Type ActivationType) override { return OriginalActions->AssetsActivatedOverride(InObjects, ActivationType); }
+	virtual TArray<FAssetData> GetValidAssetsForPreviewOrEdit(TArrayView<const FAssetData> InAssetDatas, bool bIsPreview) override { return OriginalActions->GetValidAssetsForPreviewOrEdit(InAssetDatas, bIsPreview); }
+	virtual bool CanFilter() override { return OriginalActions->CanFilter(); }
+	virtual bool CanLocalize() const override { return OriginalActions->CanLocalize(); }
+	virtual bool CanMerge() const override { return OriginalActions->CanMerge(); }
+	virtual void Merge( UObject* InObject) override { OriginalActions->Merge(InObject); }
+	virtual void Merge(UObject* BaseAsset, UObject* RemoteAsset, UObject* LocalAsset, const FOnMergeResolved& ResolutionCallback) override { OriginalActions->Merge(BaseAsset, RemoteAsset, LocalAsset, ResolutionCallback); }
+	virtual uint32 GetCategories() override { return OriginalActions->GetCategories(); }
+	virtual FString GetObjectDisplayName(UObject* Object) const override { return OriginalActions->GetObjectDisplayName(Object); }
+	virtual const TArray<FText>& GetSubMenus() const override { return OriginalActions->GetSubMenus(); }
+	virtual bool ShouldForceWorldCentric() override { return OriginalActions->ShouldForceWorldCentric(); }
+	virtual void PerformAssetDiff(UObject* OldAsset, UObject* NewAsset, const FRevisionInfo& OldRevision, const FRevisionInfo& NewRevision) const override { OriginalActions->PerformAssetDiff(OldAsset, NewAsset, OldRevision, NewRevision); }
+	virtual UThumbnailInfo* GetThumbnailInfo(UObject* Asset) const override { return OriginalActions->GetThumbnailInfo(Asset); }
+	virtual EThumbnailPrimType GetDefaultThumbnailPrimitiveType(UObject* Asset) const override { return OriginalActions->GetDefaultThumbnailPrimitiveType(Asset); }
+	virtual TSharedPtr<SWidget> GetThumbnailOverlay(const FAssetData& AssetData) const override { return OriginalActions->GetThumbnailOverlay(AssetData); }
+	virtual FText GetAssetDescription(const FAssetData& AssetData) const override { return OriginalActions->GetAssetDescription(AssetData); }
+	virtual bool IsImportedAsset() const override { return OriginalActions->IsImportedAsset(); }
+	virtual void GetResolvedSourceFilePaths(const TArray<UObject*>& TypeAssets, TArray<FString>& OutSourceFilePaths) const override { OriginalActions->GetResolvedSourceFilePaths(TypeAssets, OutSourceFilePaths); }
+	virtual void GetSourceFileLabels(const TArray<UObject*>& TypeAssets, TArray<FString>& OutSourceFileLabels) const override { OriginalActions->GetSourceFileLabels(TypeAssets, OutSourceFileLabels); }
+	virtual void BuildBackendFilter(FARFilter& InFilter) override { OriginalActions->BuildBackendFilter(InFilter); }
+	virtual FText GetDisplayNameFromAssetData(const FAssetData& AssetData) const override { return OriginalActions->GetDisplayNameFromAssetData(AssetData); }
+	virtual void SetSupported(bool bInSupported) override { OriginalActions->SetSupported(bInSupported); }
+	virtual bool IsSupported() const override { return OriginalActions->IsSupported(); }
+
+private:
+
+	void GetProxyActions(const TArray<UObject*>& InObjects, FToolMenuSection& Section);
+
+	void OnCreateProxy(TArray<FWeakObjectPtr> Objects) const;
+
+	TSharedRef<IAssetTypeActions> OriginalActions;
+};
+
+#endif
diff --git a/Engine/Plugins/Enterprise/GLTFExporter/MigratedFromGit/Source/GLTFExporter/Private/Converters/GLTFBackdropConverters.cpp b/Engine/Plugins/Enterprise/GLTFExporter/MigratedFromGit/Source/GLTFExporter/Private/Converters/GLTFBackdropConverters.cpp
--- a/Engine/Plugins/Enterprise/GLTFExporter/MigratedFromGit/Source/GLTFExporter/Private/Converters/GLTFBackdropConverters.cpp
+++ b/Engine/Plugins/Enterprise/GLTFExporter/MigratedFromGit/Source/GLTFExporter/Private/Converters/GLTFBackdropConverters.cpp
@@ -29,7 +29,7 @@ FGLTFJsonBackdropIndex FGLTFBackdropConverter::Convert(const AActor* BackdropAct
 	const UStaticMesh* Mesh;
 	if (FGLTFActorUtility::TryGetPropertyValue(BackdropActor, TEXT("Mesh"), Mesh))
 	{
-		JsonBackdrop.Mesh = Builder.GetOrAddMesh(Mesh, { FGLTFMaterialUtility::GetDefault() });
+		JsonBackdrop.Mesh = Builder.GetOrAddMesh(Mesh, { FGLTFMaterialUtility::GetDefaultMaterial() });
 	}
 	else
 	{
diff --git a/Engine/Plugins/Enterprise/GLTFExporter/MigratedFromGit/Source/GLTFExporter/Private/Converters/GLTFMaterialConverters.cpp b/Engine/Plugins/Enterprise/GLTFExporter/MigratedFromGit/Source/GLTFExporter/Private/Converters/GLTFMaterialConverters.cpp
--- a/Engine/Plugins/Enterprise/GLTFExporter/MigratedFromGit/Source/GLTFExporter/Private/Converters/GLTFMaterialConverters.cpp
+++ b/Engine/Plugins/Enterprise/GLTFExporter/MigratedFromGit/Source/GLTFExporter/Private/Converters/GLTFMaterialConverters.cpp
@@ -43,7 +43,7 @@ void FGLTFMaterialConverter::Sanitize(const UMaterialInterface*& Material, const
 
 FGLTFJsonMaterialIndex FGLTFMaterialConverter::Convert(const UMaterialInterface* Material, const FGLTFMeshData* MeshData, FGLTFIndexArray SectionIndices)
 {
-	if (Material != FGLTFMaterialUtility::GetDefault())
+	if (Material != FGLTFMaterialUtility::GetDefaultMaterial())
 	{
 		const FGLTFJsonMaterialIndex MaterialIndex = Builder.AddMaterial();
 		Builder.SetupTask<FGLTFMaterialTask>(Builder, UVOverlapChecker, Material, MeshData, SectionIndices, MaterialIndex);
diff --git a/Engine/Plugins/Enterprise/GLTFExporter/MigratedFromGit/Source/GLTFExporter/Private/Converters/GLTFMaterialPrebaker.cpp b/Engine/Plugins/Enterprise/GLTFExporter/MigratedFromGit/Source/GLTFExporter/Private/Converters/GLTFMaterialPrebaker.cpp
deleted file mode 100644
diff --git a/Engine/Plugins/Enterprise/GLTFExporter/MigratedFromGit/Source/GLTFExporter/Private/Converters/GLTFMaterialPrebaker.h b/Engine/Plugins/Enterprise/GLTFExporter/MigratedFromGit/Source/GLTFExporter/Private/Converters/GLTFMaterialPrebaker.h
deleted file mode 100644
diff --git a/Engine/Plugins/Enterprise/GLTFExporter/MigratedFromGit/Source/GLTFExporter/Private/Converters/GLTFMaterialProxyFactory.cpp b/Engine/Plugins/Enterprise/GLTFExporter/MigratedFromGit/Source/GLTFExporter/Private/Converters/GLTFMaterialProxyFactory.cpp
new file mode 100644
--- /dev/null
+++ b/Engine/Plugins/Enterprise/GLTFExporter/MigratedFromGit/Source/GLTFExporter/Private/Converters/GLTFMaterialProxyFactory.cpp
@@ -0,0 +1,372 @@
+// Copyright Epic Games, Inc. All Rights Reserved.
+
+#if WITH_EDITOR
+
+#include "Converters/GLTFMaterialProxyFactory.h"
+#include "Converters/GLTFMaterialUtility.h"
+#include "Converters/GLTFImageUtility.h"
+#include "AssetRegistry/AssetRegistryModule.h"
+#include "Materials/MaterialInstanceConstant.h"
+#include "MaterialUtilities.h"
+#include "ImageUtils.h"
+
+FGLTFMaterialProxyFactory::FGLTFMaterialProxyFactory(const UGLTFProxyOptions* Options)
+	: Builder(TEXT(""), CreateExportOptions(Options))
+{
+	Builder.ImageConverter = CreateCustomImageConverter();
+}
+
+UMaterialInterface* FGLTFMaterialProxyFactory::Create(UMaterialInterface* OriginalMaterial)
+{
+	if (FGLTFMaterialUtility::IsProxyMaterial(OriginalMaterial))
+	{
+		return OriginalMaterial;
+	}
+
+	if (FGLTFMaterialUtility::NeedsMeshData(OriginalMaterial))
+	{
+		Builder.LogWarning(FString::Printf(
+			TEXT("Material %s uses mesh data but prebaking will only use a simple quad as mesh data currently"),
+			*OriginalMaterial->GetName()));
+	}
+
+	const FGLTFJsonMaterialIndex MaterialIndex = Builder.GetOrAddMaterial(OriginalMaterial);
+	if (MaterialIndex == INDEX_NONE)
+	{
+		// TODO: report error
+		return nullptr;
+	}
+
+	MakeDirectory(RootPath);
+
+	Builder.CompleteAllTasks();
+	const FGLTFJsonMaterial& JsonMaterial = Builder.GetMaterial(MaterialIndex);
+
+	UMaterialInstanceConstant* ProxyMaterial = CreateInstancedMaterial(OriginalMaterial, JsonMaterial.ShadingModel);
+	if (ProxyMaterial != nullptr)
+	{
+		SetProxyProperties(ProxyMaterial, JsonMaterial);
+	}
+
+	return ProxyMaterial;
+}
+
+void FGLTFMaterialProxyFactory::SetProxyProperties(UMaterialInstanceConstant* ProxyMaterial, const FGLTFJsonMaterial& JsonMaterial)
+{
+	SetProxyProperty(ProxyMaterial, TEXT("Base Color Factor"), JsonMaterial.PBRMetallicRoughness.BaseColorFactor);
+	SetProxyProperty(ProxyMaterial, TEXT("Base Color"), JsonMaterial.PBRMetallicRoughness.BaseColorTexture, EGLTFMaterialPropertyGroup::BaseColorOpacity);
+
+	if (JsonMaterial.ShadingModel == EGLTFJsonShadingModel::Default || JsonMaterial.ShadingModel == EGLTFJsonShadingModel::ClearCoat)
+	{
+		SetProxyProperty(ProxyMaterial, TEXT("Emissive Factor"), JsonMaterial.EmissiveFactor);
+		SetProxyProperty(ProxyMaterial, TEXT("Emissive"), JsonMaterial.EmissiveTexture, EGLTFMaterialPropertyGroup::EmissiveColor);
+
+		SetProxyProperty(ProxyMaterial, TEXT("Metallic Factor"), JsonMaterial.PBRMetallicRoughness.MetallicFactor);
+		SetProxyProperty(ProxyMaterial, TEXT("Roughness Factor"), JsonMaterial.PBRMetallicRoughness.RoughnessFactor);
+		SetProxyProperty(ProxyMaterial, TEXT("Metallic Roughness"), JsonMaterial.PBRMetallicRoughness.MetallicRoughnessTexture, EGLTFMaterialPropertyGroup::MetallicRoughness);
+
+		SetProxyProperty(ProxyMaterial, TEXT("Normal Scale"), JsonMaterial.NormalTexture.Scale);
+		SetProxyProperty(ProxyMaterial, TEXT("Normal"), JsonMaterial.NormalTexture, EGLTFMaterialPropertyGroup::Normal);
+
+		SetProxyProperty(ProxyMaterial, TEXT("Occlusion Strength"), JsonMaterial.OcclusionTexture.Strength);
+		SetProxyProperty(ProxyMaterial, TEXT("Occlusion"), JsonMaterial.OcclusionTexture, EGLTFMaterialPropertyGroup::AmbientOcclusion);
+
+		if (JsonMaterial.ShadingModel == EGLTFJsonShadingModel::ClearCoat)
+		{
+			SetProxyProperty(ProxyMaterial, TEXT("Clear Coat Factor"), JsonMaterial.ClearCoat.ClearCoatFactor);
+			SetProxyProperty(ProxyMaterial, TEXT("Clear Coat"), JsonMaterial.ClearCoat.ClearCoatTexture, EGLTFMaterialPropertyGroup::ClearCoatRoughness); // TODO: add property group for clear coat intensity only
+
+			SetProxyProperty(ProxyMaterial, TEXT("Clear Coat Roughness Factor"), JsonMaterial.ClearCoat.ClearCoatRoughnessFactor);
+			SetProxyProperty(ProxyMaterial, TEXT("Clear Coat Roughness"), JsonMaterial.ClearCoat.ClearCoatRoughnessTexture, EGLTFMaterialPropertyGroup::ClearCoatRoughness);
+
+			SetProxyProperty(ProxyMaterial, TEXT("Clear Coat Normal Scale"), JsonMaterial.ClearCoat.ClearCoatNormalTexture.Scale);
+			SetProxyProperty(ProxyMaterial, TEXT("Clear Coat Normal"), JsonMaterial.ClearCoat.ClearCoatNormalTexture, EGLTFMaterialPropertyGroup::ClearCoatBottomNormal);
+		}
+	}
+
+	ProxyMaterial->PostEditChange();
+}
+
+void FGLTFMaterialProxyFactory::SetProxyProperty(UMaterialInstanceConstant* ProxyMaterial, const FString& PropertyName, float Scalar)
+{
+	float DefaultValue;
+	if (!ProxyMaterial->GetScalarParameterDefaultValue(*PropertyName, DefaultValue))
+	{
+		// TODO: report error
+		return;
+	}
+
+	if (DefaultValue != Scalar)
+	{
+		ProxyMaterial->SetScalarParameterValueEditorOnly(*PropertyName, Scalar);
+	}
+}
+
+void FGLTFMaterialProxyFactory::SetProxyProperty(UMaterialInstanceConstant* ProxyMaterial, const FString& PropertyName, const FGLTFJsonColor3& Color)
+{
+	FLinearColor DefaultValue;
+	if (!ProxyMaterial->GetVectorParameterDefaultValue(*PropertyName, DefaultValue))
+	{
+		// TODO: report error
+		return;
+	}
+
+	const FLinearColor Value(Color.R, Color.G, Color.B);
+	if (DefaultValue != Value)
+	{
+		ProxyMaterial->SetVectorParameterValueEditorOnly(*PropertyName, Value);
+	}
+}
+
+void FGLTFMaterialProxyFactory::SetProxyProperty(UMaterialInstanceConstant* ProxyMaterial, const FString& PropertyName, const FGLTFJsonColor4& Color)
+{
+	FLinearColor DefaultValue;
+	if (!ProxyMaterial->GetVectorParameterDefaultValue(*PropertyName, DefaultValue))
+	{
+		// TODO: report error
+		return;
+	}
+
+	const FLinearColor Value(Color.R, Color.G, Color.B, Color.A);
+	if (DefaultValue != Value)
+	{
+		ProxyMaterial->SetVectorParameterValueEditorOnly(*PropertyName, Value);
+	}
+}
+
+void FGLTFMaterialProxyFactory::SetProxyProperty(UMaterialInstanceConstant* ProxyMaterial, const FString& PropertyName, const FGLTFJsonTextureInfo& TextureInfo, EGLTFMaterialPropertyGroup PropertyGroup)
+{
+	if (TextureInfo.Index != INDEX_NONE)
+	{
+		const UTexture* Texture = FindOrCreateTexture(TextureInfo.Index, PropertyGroup);
+		if (Texture == nullptr)
+		{
+			// TODO: report error
+		}
+
+		ProxyMaterial->SetTextureParameterValueEditorOnly(*(PropertyName + TEXT(" Texture")), const_cast<UTexture*>(Texture));
+	}
+
+	if (TextureInfo.TexCoord != 0)
+	{
+		ProxyMaterial->SetScalarParameterValueEditorOnly(*(PropertyName + TEXT(" UV Index")), TextureInfo.TexCoord);
+	}
+
+	if (TextureInfo.Transform.Offset != FGLTFJsonVector2::Zero)
+	{
+		const FLinearColor Offset(TextureInfo.Transform.Offset.X, TextureInfo.Transform.Offset.Y, 0, 0);
+		ProxyMaterial->SetVectorParameterValueEditorOnly(*(PropertyName + TEXT(" UV Offset")), Offset);
+	}
+
+	if (TextureInfo.Transform.Scale != FGLTFJsonVector2::One)
+	{
+		const FLinearColor Scale(TextureInfo.Transform.Scale.X, TextureInfo.Transform.Scale.Y, 0, 0);
+		ProxyMaterial->SetVectorParameterValueEditorOnly(*(PropertyName + TEXT(" UV Scale")), Scale);
+	}
+
+	if (TextureInfo.Transform.Rotation != 0)
+	{
+		ProxyMaterial->SetScalarParameterValueEditorOnly(*(PropertyName + TEXT(" UV Rotation")), TextureInfo.Transform.Rotation);
+	}
+}
+
+const UTexture2D* FGLTFMaterialProxyFactory::FindOrCreateTexture(FGLTFJsonTextureIndex Index, EGLTFMaterialPropertyGroup PropertyGroup)
+{
+	// TODO: fix potential conflict when same texture used for different material properties that have different encoding (sRGB vs Linear, Normalmap compression etc)
+
+	const UTexture2D** FoundPtr = Textures.Find(Index);
+	if (FoundPtr != nullptr)
+	{
+		return *FoundPtr;
+	}
+
+	const FGLTFJsonTexture& JsonTexture = Builder.GetTexture(Index);
+	const FGLTFImageData* ImageData = Images.Find(JsonTexture.Source);
+	if (ImageData == nullptr)
+	{
+		return nullptr;
+	}
+
+	const FGLTFJsonSampler& JsonSampler = Builder.GetSampler(JsonTexture.Sampler);
+	const UTexture2D* Texture = CreateTexture(ImageData, JsonSampler, PropertyGroup);
+	Textures.Add(Index, Texture);
+	return Texture;
+}
+
+UTexture2D* FGLTFMaterialProxyFactory::CreateTexture(const FGLTFImageData* ImageData, const FGLTFJsonSampler& JsonSampler, EGLTFMaterialPropertyGroup PropertyGroup)
+{
+	const bool bSRGB = PropertyGroup == EGLTFMaterialPropertyGroup::BaseColorOpacity || PropertyGroup == EGLTFMaterialPropertyGroup::EmissiveColor;
+	const bool bNormalMap = PropertyGroup == EGLTFMaterialPropertyGroup::Normal || PropertyGroup == EGLTFMaterialPropertyGroup::ClearCoatBottomNormal;
+
+	const FString PackageName = RootPath / TEXT("T_GLTF_") + ImageData->Filename;
+	UPackage* Package = CreatePackage(*PackageName);
+	Package->FullyLoad();
+	Package->Modify();
+
+	FCreateTexture2DParameters TexParams;
+	TexParams.bUseAlpha = !ImageData->bIgnoreAlpha;
+	TexParams.CompressionSettings = bNormalMap ? TC_Normalmap : TC_Default;
+	TexParams.bDeferCompression = true;
+	TexParams.bSRGB = bSRGB;
+	TexParams.TextureGroup = bNormalMap ? TEXTUREGROUP_WorldNormalMap : TEXTUREGROUP_World;
+	TexParams.SourceGuidHash = FGuid();
+
+	UTexture2D* Texture = FImageUtils::CreateTexture2D(ImageData->Size.X, ImageData->Size.Y, *ImageData->Pixels, Package, FPaths::GetCleanFilename(Package->GetName()), RF_Public | RF_Standalone, TexParams);
+	Texture->Filter = ConvertFilter(JsonSampler.MagFilter);
+	Texture->AddressX = ConvertWrap(JsonSampler.WrapS);
+	Texture->AddressY = ConvertWrap(JsonSampler.WrapT);
+
+	Texture->PostEditChange();
+	return Texture;
+}
+
+UMaterialInstanceConstant* FGLTFMaterialProxyFactory::CreateInstancedMaterial(UMaterialInterface* OriginalMaterial, EGLTFJsonShadingModel ShadingModel)
+{
+	const FString PackageName = RootPath / TEXT("M_GLTF_") + OriginalMaterial->GetName();
+	UPackage* Package = CreatePackage(*PackageName);
+	Package->FullyLoad();
+	Package->Modify();
+
+	UMaterialInterface* BaseMaterial = FGLTFMaterialUtility::GetProxyBaseMaterial(ShadingModel);
+	if (BaseMaterial == nullptr)
+	{
+		Builder.LogError(FString::Printf(
+			TEXT("Can't create proxy for material %s, because shading model %s has no base material"),
+			*OriginalMaterial->GetName(),
+			FGLTFJsonUtility::GetValue(ShadingModel)));
+		return nullptr;
+	}
+
+	const FString BaseName = TEXT("GLTF_") + OriginalMaterial->GetName(); // prefix "M_" is added automatically by CreateInstancedMaterial
+	UMaterialInstanceConstant* ProxyMaterial = FMaterialUtilities::CreateInstancedMaterial(BaseMaterial, Package, BaseName, RF_Public | RF_Standalone);
+
+	const bool bTwoSided = OriginalMaterial->IsTwoSided();
+	if (bTwoSided != BaseMaterial->IsTwoSided())
+	{
+		ProxyMaterial->BasePropertyOverrides.bOverride_TwoSided = true;
+		ProxyMaterial->BasePropertyOverrides.TwoSided = bTwoSided;
+	}
+
+	const EBlendMode BlendMode = OriginalMaterial->GetBlendMode();
+	if (BlendMode != BaseMaterial->GetBlendMode())
+	{
+		ProxyMaterial->BasePropertyOverrides.bOverride_BlendMode = true;
+		ProxyMaterial->BasePropertyOverrides.BlendMode = BlendMode;
+	}
+
+	const float OpacityMaskClipValue = OriginalMaterial->GetOpacityMaskClipValue();
+	if (OpacityMaskClipValue != BaseMaterial->GetOpacityMaskClipValue())
+	{
+		ProxyMaterial->BasePropertyOverrides.bOverride_OpacityMaskClipValue = true;
+		ProxyMaterial->BasePropertyOverrides.OpacityMaskClipValue = OpacityMaskClipValue;
+	}
+
+	UGLTFMaterialExportOptions* UserData = OriginalMaterial->GetAssetUserData<UGLTFMaterialExportOptions>();
+	if (UserData == nullptr)
+	{
+		UserData = NewObject<UGLTFMaterialExportOptions>();
+		OriginalMaterial->AddAssetUserData(UserData);
+	}
+
+	UserData->Proxy = ProxyMaterial;
+	OriginalMaterial->Modify();
+	return ProxyMaterial;
+}
+
+TUniquePtr<IGLTFImageConverter> FGLTFMaterialProxyFactory::CreateCustomImageConverter()
+{
+	class FGLTFCustomImageConverter final: public IGLTFImageConverter
+	{
+	public:
+
+		FGLTFMaterialProxyFactory& Proxyr;
+		TSet<FString> UniqueFilenames;
+
+		FGLTFCustomImageConverter(FGLTFMaterialProxyFactory& Proxyr)
+			: Proxyr(Proxyr)
+		{
+		}
+
+	protected:
+
+		virtual FGLTFJsonImageIndex Convert(TGLTFSuperfluous<FString> Name, EGLTFTextureType Type, bool bIgnoreAlpha, FIntPoint Size, TGLTFSharedArray<FColor> Pixels) override
+		{
+			const FString Filename = FGLTFImageUtility::GetUniqueFilename(Name, TEXT(""), UniqueFilenames);
+			UniqueFilenames.Add(Filename);
+
+			const FGLTFJsonImageIndex ImageIndex = Proxyr.Builder.AddImage();
+			Proxyr.Images.Add(ImageIndex, { Filename, Type, bIgnoreAlpha, Size, Pixels });
+			return ImageIndex;
+		}
+	};
+
+	return MakeUnique<FGLTFCustomImageConverter>(*this);
+}
+
+bool FGLTFMaterialProxyFactory::MakeDirectory(const FString& PackagePath)
+{
+	const FString DirPath = FPaths::ConvertRelativePathToFull(FPackageName::LongPackageNameToFilename(PackagePath + TEXT("/")));
+	if (DirPath.IsEmpty())
+	{
+		return false;
+	}
+
+	bool bResult = true;
+	if (!IFileManager::Get().DirectoryExists(*DirPath))
+	{
+		bResult = IFileManager::Get().MakeDirectory(*DirPath, true);
+	}
+
+	if (bResult)
+	{
+		const FAssetRegistryModule& AssetRegistryModule = FModuleManager::LoadModuleChecked<FAssetRegistryModule>("AssetRegistry");
+		AssetRegistryModule.Get().AddPath(PackagePath);
+	}
+
+	return bResult;
+}
+
+UGLTFExportOptions* FGLTFMaterialProxyFactory::CreateExportOptions(const UGLTFProxyOptions* ProxyOptions)
+{
+	UGLTFExportOptions* ExportOptions = NewObject<UGLTFExportOptions>();
+	ExportOptions->ResetToDefault();
+	ExportOptions->bExportProxyMaterials = false;
+	ExportOptions->BakeMaterialInputs = EGLTFMaterialBakeMode::Simple;
+	ExportOptions->DefaultMaterialBakeSize = ProxyOptions->DefaultMaterialBakeSize;
+	ExportOptions->DefaultMaterialBakeFilter = ProxyOptions->DefaultMaterialBakeFilter;
+	ExportOptions->DefaultMaterialBakeTiling = ProxyOptions->DefaultMaterialBakeTiling;
+	ExportOptions->DefaultInputBakeSettings = ProxyOptions->DefaultInputBakeSettings;
+	ExportOptions->bAdjustNormalmaps = false;
+	return ExportOptions;
+}
+
+TextureAddress FGLTFMaterialProxyFactory::ConvertWrap(EGLTFJsonTextureWrap Wrap)
+{
+	switch (Wrap)
+	{
+		case EGLTFJsonTextureWrap::Repeat:         return TA_Wrap;
+		case EGLTFJsonTextureWrap::MirroredRepeat: return TA_Mirror;
+		case EGLTFJsonTextureWrap::ClampToEdge:    return TA_Clamp;
+		default:
+			checkNoEntry();
+			return TA_MAX;
+	}
+}
+
+TextureFilter FGLTFMaterialProxyFactory::ConvertFilter(EGLTFJsonTextureFilter Filter)
+{
+	switch (Filter)
+	{
+		case EGLTFJsonTextureFilter::Nearest:              return TF_Nearest;
+		case EGLTFJsonTextureFilter::NearestMipmapNearest: return TF_Nearest;
+		case EGLTFJsonTextureFilter::LinearMipmapNearest:  return TF_Bilinear;
+		case EGLTFJsonTextureFilter::NearestMipmapLinear:  return TF_Bilinear;
+		case EGLTFJsonTextureFilter::Linear:               return TF_Trilinear;
+		case EGLTFJsonTextureFilter::LinearMipmapLinear:   return TF_Trilinear;
+		default:
+			checkNoEntry();
+			return TF_MAX;
+	}
+}
+
+#endif
diff --git a/Engine/Plugins/Enterprise/GLTFExporter/MigratedFromGit/Source/GLTFExporter/Private/Converters/GLTFMaterialProxyFactory.h b/Engine/Plugins/Enterprise/GLTFExporter/MigratedFromGit/Source/GLTFExporter/Private/Converters/GLTFMaterialProxyFactory.h
new file mode 100644
--- /dev/null
+++ b/Engine/Plugins/Enterprise/GLTFExporter/MigratedFromGit/Source/GLTFExporter/Private/Converters/GLTFMaterialProxyFactory.h
@@ -0,0 +1,61 @@
+// Copyright Epic Games, Inc. All Rights Reserved.
+
+#pragma once
+
+#if WITH_EDITOR
+
+#include "Builders/GLTFConvertBuilder.h"
+#include "Options/GLTFProxyOptions.h"
+#include "Options/GLTFExportOptions.h"
+
+class UMaterialInstanceConstant;
+
+class FGLTFMaterialProxyFactory
+{
+public:
+
+	FGLTFMaterialProxyFactory(const UGLTFProxyOptions* Options = nullptr);
+
+	UMaterialInterface* Create(UMaterialInterface* OriginalMaterial);
+
+	FString RootPath;
+
+private:
+
+	struct FGLTFImageData
+	{
+		FString Filename;
+		EGLTFTextureType Type;
+		bool bIgnoreAlpha;
+		FIntPoint Size;
+		TGLTFSharedArray<FColor> Pixels;
+	};
+
+	void SetProxyProperties(UMaterialInstanceConstant* ProxyMaterial, const FGLTFJsonMaterial& JsonMaterial);
+
+	void SetProxyProperty(UMaterialInstanceConstant* ProxyMaterial, const FString& PropertyName, float Scalar);
+	void SetProxyProperty(UMaterialInstanceConstant* ProxyMaterial, const FString& PropertyName, const FGLTFJsonColor3& Color);
+	void SetProxyProperty(UMaterialInstanceConstant* ProxyMaterial, const FString& PropertyName, const FGLTFJsonColor4& Color);
+	void SetProxyProperty(UMaterialInstanceConstant* ProxyMaterial, const FString& PropertyName, const FGLTFJsonTextureInfo& TextureInfo, EGLTFMaterialPropertyGroup PropertyGroup);
+
+	const UTexture2D* FindOrCreateTexture(FGLTFJsonTextureIndex Index, EGLTFMaterialPropertyGroup PropertyGroup);
+	UTexture2D* CreateTexture(const FGLTFImageData* ImageData, const FGLTFJsonSampler& JsonSampler, EGLTFMaterialPropertyGroup PropertyGroup);
+
+	UMaterialInstanceConstant* CreateInstancedMaterial(UMaterialInterface* OriginalMaterial, EGLTFJsonShadingModel ShadingModel);
+
+	TUniquePtr<IGLTFImageConverter> CreateCustomImageConverter();
+
+	static bool MakeDirectory(const FString& String);
+
+	static UGLTFExportOptions* CreateExportOptions(const UGLTFProxyOptions* ProxyOptions);
+
+	static TextureAddress ConvertWrap(EGLTFJsonTextureWrap Wrap);
+	static TextureFilter ConvertFilter(EGLTFJsonTextureFilter Filter);
+
+	FGLTFConvertBuilder Builder;
+
+	TMap<FGLTFJsonTextureIndex, const UTexture2D*> Textures;
+	TMap<FGLTFJsonImageIndex, FGLTFImageData> Images;
+};
+
+#endif
diff --git a/Engine/Plugins/Enterprise/GLTFExporter/MigratedFromGit/Source/GLTFExporter/Private/Converters/GLTFMaterialUtility.cpp b/Engine/Plugins/Enterprise/GLTFExporter/MigratedFromGit/Source/GLTFExporter/Private/Converters/GLTFMaterialUtility.cpp
--- a/Engine/Plugins/Enterprise/GLTFExporter/MigratedFromGit/Source/GLTFExporter/Private/Converters/GLTFMaterialUtility.cpp
+++ b/Engine/Plugins/Enterprise/GLTFExporter/MigratedFromGit/Source/GLTFExporter/Private/Converters/GLTFMaterialUtility.cpp
@@ -18,27 +18,27 @@
 #include "GLTFConverterUtility.h"
 #include "Misc/DefaultValueHelper.h"
 
-UMaterialInterface* FGLTFMaterialUtility::GetDefault()
+UMaterialInterface* FGLTFMaterialUtility::GetDefaultMaterial()
 {
-	static UMaterialInterface* DefaultMaterial = GetPrebaked(EGLTFJsonShadingModel::Default);
+	static UMaterialInterface* DefaultMaterial = GetProxyBaseMaterial(EGLTFJsonShadingModel::Default);
 	return DefaultMaterial;
 }
 
-UMaterialInterface* FGLTFMaterialUtility::GetPrebaked(EGLTFJsonShadingModel ShadingModel)
+UMaterialInterface* FGLTFMaterialUtility::GetProxyBaseMaterial(EGLTFJsonShadingModel ShadingModel)
 {
 	const FString Name = FGLTFJsonUtility::GetValue(ShadingModel);
 	const FString Path = TEXT("/GLTFExporter/Materials/GLTF") + Name + TEXT(".GLTF") + Name;
 	return LoadObject<UMaterialInterface>(nullptr, *Path);
 }
 
-bool FGLTFMaterialUtility::IsPrebaked(const UMaterial* Material)
+bool FGLTFMaterialUtility::IsProxyMaterial(const UMaterial* Material)
 {
 	return Material->GetPathName().StartsWith(TEXT("/GLTFExporter/Materials/GLTF"));
 }
 
-bool FGLTFMaterialUtility::IsPrebaked(const UMaterialInterface* Material)
+bool FGLTFMaterialUtility::IsProxyMaterial(const UMaterialInterface* Material)
 {
-	return IsPrebaked(Material->GetMaterial());
+	return IsProxyMaterial(Material->GetMaterial());
 }
 
 #if WITH_EDITOR
diff --git a/Engine/Plugins/Enterprise/GLTFExporter/MigratedFromGit/Source/GLTFExporter/Private/Converters/GLTFMaterialUtility.h b/Engine/Plugins/Enterprise/GLTFExporter/MigratedFromGit/Source/GLTFExporter/Private/Converters/GLTFMaterialUtility.h
--- a/Engine/Plugins/Enterprise/GLTFExporter/MigratedFromGit/Source/GLTFExporter/Private/Converters/GLTFMaterialUtility.h
+++ b/Engine/Plugins/Enterprise/GLTFExporter/MigratedFromGit/Source/GLTFExporter/Private/Converters/GLTFMaterialUtility.h
@@ -38,12 +38,12 @@ struct FGLTFPropertyBakeOutput
 
 struct FGLTFMaterialUtility
 {
-	static UMaterialInterface* GetDefault();
+	static UMaterialInterface* GetDefaultMaterial();
 
-	static UMaterialInterface* GetPrebaked(EGLTFJsonShadingModel ShadingModel);
+	static UMaterialInterface* GetProxyBaseMaterial(EGLTFJsonShadingModel ShadingModel);
 
-	static bool IsPrebaked(const UMaterial* Material);
-	static bool IsPrebaked(const UMaterialInterface* Material);
+	static bool IsProxyMaterial(const UMaterial* Material);
+	static bool IsProxyMaterial(const UMaterialInterface* Material);
 
 #if WITH_EDITOR
 	static bool IsNormalMap(const FMaterialPropertyEx& Property);
diff --git a/Engine/Plugins/Enterprise/GLTFExporter/MigratedFromGit/Source/GLTFExporter/Private/Converters/GLTFSkySphereConverters.cpp b/Engine/Plugins/Enterprise/GLTFExporter/MigratedFromGit/Source/GLTFExporter/Private/Converters/GLTFSkySphereConverters.cpp
--- a/Engine/Plugins/Enterprise/GLTFExporter/MigratedFromGit/Source/GLTFExporter/Private/Converters/GLTFSkySphereConverters.cpp
+++ b/Engine/Plugins/Enterprise/GLTFExporter/MigratedFromGit/Source/GLTFExporter/Private/Converters/GLTFSkySphereConverters.cpp
@@ -33,7 +33,7 @@ FGLTFJsonSkySphereIndex FGLTFSkySphereConverter::Convert(const AActor* SkySphere
 		const FTransform RelativeTransform = Transform.GetRelativeTransform(ParentTransform);
 
 		JsonSkySphere.Scale = FGLTFConverterUtility::ConvertScale(RelativeTransform.GetScale3D());
-		JsonSkySphere.SkySphereMesh = Builder.GetOrAddMesh(StaticMeshComponent, { FGLTFMaterialUtility::GetDefault() });
+		JsonSkySphere.SkySphereMesh = Builder.GetOrAddMesh(StaticMeshComponent, { FGLTFMaterialUtility::GetDefaultMaterial() });
 	}
 	else
 	{
diff --git a/Engine/Plugins/Enterprise/GLTFExporter/MigratedFromGit/Source/GLTFExporter/Private/GLTFExporterModule.cpp b/Engine/Plugins/Enterprise/GLTFExporter/MigratedFromGit/Source/GLTFExporter/Private/GLTFExporterModule.cpp
--- a/Engine/Plugins/Enterprise/GLTFExporter/MigratedFromGit/Source/GLTFExporter/Private/GLTFExporterModule.cpp
+++ b/Engine/Plugins/Enterprise/GLTFExporter/MigratedFromGit/Source/GLTFExporter/Private/GLTFExporterModule.cpp
@@ -1,8 +1,8 @@
 // Copyright Epic Games, Inc. All Rights Reserved.
 
 #include "GLTFExporterModule.h"
-#include "Converters/GLTFMaterialPrebaker.h"
-#include "Actions/GLTFPrebakeAssetActions.h"
+#include "Converters/GLTFMaterialProxyFactory.h"
+#include "Actions/GLTFProxyAssetActions.h"
 #include "Interfaces/IPluginManager.h"
 #if WITH_EDITOR
 #include "Materials/MaterialInstanceConstant.h"
@@ -14,13 +14,14 @@ class FAssetTypeActions_Base;
 
 DEFINE_LOG_CATEGORY(LogGLTFExporter);
 
-UMaterialInterface* IGLTFExporterModule::PrebakeMaterial(UMaterialInterface* Material, const UGLTFPrebakeOptions* Options)
+UMaterialInterface* IGLTFExporterModule::CreateProxyMaterial(UMaterialInterface* Material, const UGLTFProxyOptions* Options)
 {
 #if WITH_EDITOR
-	FGLTFMaterialPrebaker Prebaker(Options);
-	Prebaker.RootPath = FPaths::GetPath(Material->GetPathName()) / TEXT("GLTF");
-	return Prebaker.Prebake(Material);
+	FGLTFMaterialProxyFactory ProxyFactory(Options);
+	ProxyFactory.RootPath = FPaths::GetPath(Material->GetPathName()) / TEXT("GLTF");
+	return ProxyFactory.Create(Material);
 #else
+	// TODO: report error
 	return nullptr;
 #endif
 }
@@ -45,10 +46,10 @@ public:
 			IAssetTools& AssetTools = FAssetToolsModule::GetModule().Get();
 
 			// TODO: add support for when GetAssetTypeActionsForClass returns null pointer
-			AssetTypeActionsArray.Add(MakeShared<FGLTFPrebakeAssetActions>(AssetTools.GetAssetTypeActionsForClass(UMaterial::StaticClass()).Pin().ToSharedRef()));
-			AssetTypeActionsArray.Add(MakeShared<FGLTFPrebakeAssetActions>(AssetTools.GetAssetTypeActionsForClass(UMaterialInstanceConstant::StaticClass()).Pin().ToSharedRef()));
-			AssetTypeActionsArray.Add(MakeShared<FGLTFPrebakeAssetActions>(AssetTools.GetAssetTypeActionsForClass(UMaterialInstanceDynamic::StaticClass()).Pin().ToSharedRef()));
-			AssetTypeActionsArray.Add(MakeShared<FGLTFPrebakeAssetActions>(AssetTools.GetAssetTypeActionsForClass(UMaterialInterface::StaticClass()).Pin().ToSharedRef()));
+			AssetTypeActionsArray.Add(MakeShared<FGLTFProxyAssetActions>(AssetTools.GetAssetTypeActionsForClass(UMaterial::StaticClass()).Pin().ToSharedRef()));
+			AssetTypeActionsArray.Add(MakeShared<FGLTFProxyAssetActions>(AssetTools.GetAssetTypeActionsForClass(UMaterialInstanceConstant::StaticClass()).Pin().ToSharedRef()));
+			AssetTypeActionsArray.Add(MakeShared<FGLTFProxyAssetActions>(AssetTools.GetAssetTypeActionsForClass(UMaterialInstanceDynamic::StaticClass()).Pin().ToSharedRef()));
+			AssetTypeActionsArray.Add(MakeShared<FGLTFProxyAssetActions>(AssetTools.GetAssetTypeActionsForClass(UMaterialInterface::StaticClass()).Pin().ToSharedRef()));
 
 			for (TSharedRef<IAssetTypeActions>& AssetTypeActions : AssetTypeActionsArray)
 			{
diff --git a/Engine/Plugins/Enterprise/GLTFExporter/MigratedFromGit/Source/GLTFExporter/Private/Options/GLTFPrebakeOptions.cpp b/Engine/Plugins/Enterprise/GLTFExporter/MigratedFromGit/Source/GLTFExporter/Private/Options/GLTFPrebakeOptions.cpp
deleted file mode 100644
diff --git a/Engine/Plugins/Enterprise/GLTFExporter/MigratedFromGit/Source/GLTFExporter/Private/Options/GLTFProxyOptions.cpp b/Engine/Plugins/Enterprise/GLTFExporter/MigratedFromGit/Source/GLTFExporter/Private/Options/GLTFProxyOptions.cpp
new file mode 100644
--- /dev/null
+++ b/Engine/Plugins/Enterprise/GLTFExporter/MigratedFromGit/Source/GLTFExporter/Private/Options/GLTFProxyOptions.cpp
@@ -0,0 +1,16 @@
+// Copyright Epic Games, Inc. All Rights Reserved.
+
+#include "Options/GLTFProxyOptions.h"
+
+UGLTFProxyOptions::UGLTFProxyOptions(const FObjectInitializer& ObjectInitializer)
+	: Super(ObjectInitializer)
+{
+	ResetToDefault();
+}
+
+void UGLTFProxyOptions::ResetToDefault()
+{
+	DefaultMaterialBakeSize = EGLTFMaterialBakeSizePOT::POT_1024;
+	DefaultMaterialBakeFilter = TF_Trilinear;
+	DefaultMaterialBakeTiling = TA_Wrap;
+}
diff --git a/Engine/Plugins/Enterprise/GLTFExporter/MigratedFromGit/Source/GLTFExporter/Private/Tasks/GLTFMaterialTasks.cpp b/Engine/Plugins/Enterprise/GLTFExporter/MigratedFromGit/Source/GLTFExporter/Private/Tasks/GLTFMaterialTasks.cpp
--- a/Engine/Plugins/Enterprise/GLTFExporter/MigratedFromGit/Source/GLTFExporter/Private/Tasks/GLTFMaterialTasks.cpp
+++ b/Engine/Plugins/Enterprise/GLTFExporter/MigratedFromGit/Source/GLTFExporter/Private/Tasks/GLTFMaterialTasks.cpp
@@ -65,9 +65,9 @@ void FGLTFMaterialTask::Complete()
 	ConvertShadingModel(JsonMaterial.ShadingModel);
 	ConvertAlphaMode(JsonMaterial.AlphaMode, JsonMaterial.BlendMode);
 
-	if (FGLTFMaterialUtility::IsPrebaked(BaseMaterial))
+	if (FGLTFMaterialUtility::IsProxyMaterial(BaseMaterial))
 	{
-		ApplyPrebakedProperties(JsonMaterial);
+		GetProxyProperties(JsonMaterial);
 		return;
 	}
 
@@ -210,46 +210,46 @@ FString FGLTFMaterialTask::GetBakedTextureName(const FString& PropertyName) cons
 	return GetMaterialName() + TEXT("_") + PropertyName;
 }
 
-void FGLTFMaterialTask::ApplyPrebakedProperties(FGLTFJsonMaterial& OutMaterial) const
+void FGLTFMaterialTask::GetProxyProperties(FGLTFJsonMaterial& OutMaterial) const
 {
-	ApplyPrebakedProperty(TEXT("Base Color Factor"), OutMaterial.PBRMetallicRoughness.BaseColorFactor);
-	ApplyPrebakedProperty(TEXT("Base Color"), OutMaterial.PBRMetallicRoughness.BaseColorTexture);
+	GetProxyProperty(TEXT("Base Color Factor"), OutMaterial.PBRMetallicRoughness.BaseColorFactor);
+	GetProxyProperty(TEXT("Base Color"), OutMaterial.PBRMetallicRoughness.BaseColorTexture);
 
 	if (OutMaterial.ShadingModel == EGLTFJsonShadingModel::Default || OutMaterial.ShadingModel == EGLTFJsonShadingModel::ClearCoat)
 	{
-		ApplyPrebakedProperty(TEXT("Emissive Factor"), OutMaterial.EmissiveFactor);
-		ApplyPrebakedProperty(TEXT("Emissive"), OutMaterial.EmissiveTexture);
+		GetProxyProperty(TEXT("Emissive Factor"), OutMaterial.EmissiveFactor);
+		GetProxyProperty(TEXT("Emissive"), OutMaterial.EmissiveTexture);
 
-		ApplyPrebakedProperty(TEXT("Metallic Factor"), OutMaterial.PBRMetallicRoughness.MetallicFactor);
-		ApplyPrebakedProperty(TEXT("Roughness Factor"), OutMaterial.PBRMetallicRoughness.RoughnessFactor);
-		ApplyPrebakedProperty(TEXT("Metallic Roughness"), OutMaterial.PBRMetallicRoughness.MetallicRoughnessTexture);
+		GetProxyProperty(TEXT("Metallic Factor"), OutMaterial.PBRMetallicRoughness.MetallicFactor);
+		GetProxyProperty(TEXT("Roughness Factor"), OutMaterial.PBRMetallicRoughness.RoughnessFactor);
+		GetProxyProperty(TEXT("Metallic Roughness"), OutMaterial.PBRMetallicRoughness.MetallicRoughnessTexture);
 
-		ApplyPrebakedProperty(TEXT("Normal Scale"), OutMaterial.NormalTexture.Scale);
-		ApplyPrebakedProperty(TEXT("Normal"), OutMaterial.NormalTexture);
+		GetProxyProperty(TEXT("Normal Scale"), OutMaterial.NormalTexture.Scale);
+		GetProxyProperty(TEXT("Normal"), OutMaterial.NormalTexture);
 
-		ApplyPrebakedProperty(TEXT("Occlusion Strength"), OutMaterial.OcclusionTexture.Strength);
-		ApplyPrebakedProperty(TEXT("Occlusion"), OutMaterial.OcclusionTexture);
+		GetProxyProperty(TEXT("Occlusion Strength"), OutMaterial.OcclusionTexture.Strength);
+		GetProxyProperty(TEXT("Occlusion"), OutMaterial.OcclusionTexture);
 
 		if (OutMaterial.ShadingModel == EGLTFJsonShadingModel::ClearCoat)
 		{
-			ApplyPrebakedProperty(TEXT("Clear Coat Factor"), OutMaterial.ClearCoat.ClearCoatFactor);
-			ApplyPrebakedProperty(TEXT("Clear Coat"), OutMaterial.ClearCoat.ClearCoatTexture);
+			GetProxyProperty(TEXT("Clear Coat Factor"), OutMaterial.ClearCoat.ClearCoatFactor);
+			GetProxyProperty(TEXT("Clear Coat"), OutMaterial.ClearCoat.ClearCoatTexture);
 
-			ApplyPrebakedProperty(TEXT("Clear Coat Roughness Factor"), OutMaterial.ClearCoat.ClearCoatRoughnessFactor);
-			ApplyPrebakedProperty(TEXT("Clear Coat Roughness"), OutMaterial.ClearCoat.ClearCoatRoughnessTexture);
+			GetProxyProperty(TEXT("Clear Coat Roughness Factor"), OutMaterial.ClearCoat.ClearCoatRoughnessFactor);
+			GetProxyProperty(TEXT("Clear Coat Roughness"), OutMaterial.ClearCoat.ClearCoatRoughnessTexture);
 
-			ApplyPrebakedProperty(TEXT("Clear Coat Normal Scale"), OutMaterial.ClearCoat.ClearCoatNormalTexture.Scale);
-			ApplyPrebakedProperty(TEXT("Clear Coat Normal"), OutMaterial.ClearCoat.ClearCoatNormalTexture);
+			GetProxyProperty(TEXT("Clear Coat Normal Scale"), OutMaterial.ClearCoat.ClearCoatNormalTexture.Scale);
+			GetProxyProperty(TEXT("Clear Coat Normal"), OutMaterial.ClearCoat.ClearCoatNormalTexture);
 		}
 	}
 }
 
-void FGLTFMaterialTask::ApplyPrebakedProperty(const FString& PropertyName, float& OutValue) const
+void FGLTFMaterialTask::GetProxyProperty(const FString& PropertyName, float& OutValue) const
 {
 	Material->GetScalarParameterValue(*PropertyName, OutValue, true);
 }
 
-void FGLTFMaterialTask::ApplyPrebakedProperty(const FString& PropertyName, FGLTFJsonColor3& OutValue) const
+void FGLTFMaterialTask::GetProxyProperty(const FString& PropertyName, FGLTFJsonColor3& OutValue) const
 {
 	FLinearColor Value;
 	if (Material->GetVectorParameterValue(*PropertyName, Value, true))
@@ -258,7 +258,7 @@ void FGLTFMaterialTask::ApplyPrebakedProperty(const FString& PropertyName, FGLTF
 	}
 }
 
-void FGLTFMaterialTask::ApplyPrebakedProperty(const FString& PropertyName, FGLTFJsonColor4& OutValue) const
+void FGLTFMaterialTask::GetProxyProperty(const FString& PropertyName, FGLTFJsonColor4& OutValue) const
 {
 	FLinearColor Value;
 	if (Material->GetVectorParameterValue(*PropertyName, Value, true))
@@ -267,7 +267,7 @@ void FGLTFMaterialTask::ApplyPrebakedProperty(const FString& PropertyName, FGLTF
 	}
 }
 
-void FGLTFMaterialTask::ApplyPrebakedProperty(const FString& PropertyName, FGLTFJsonTextureInfo& OutValue) const
+void FGLTFMaterialTask::GetProxyProperty(const FString& PropertyName, FGLTFJsonTextureInfo& OutValue) const
 {
 	UTexture* Texture;
 	if (Material->GetTextureParameterValue(*(PropertyName + TEXT(" Texture")), Texture, true) || Texture == nullptr)
diff --git a/Engine/Plugins/Enterprise/GLTFExporter/MigratedFromGit/Source/GLTFExporter/Private/Tasks/GLTFMaterialTasks.h b/Engine/Plugins/Enterprise/GLTFExporter/MigratedFromGit/Source/GLTFExporter/Private/Tasks/GLTFMaterialTasks.h
--- a/Engine/Plugins/Enterprise/GLTFExporter/MigratedFromGit/Source/GLTFExporter/Private/Tasks/GLTFMaterialTasks.h
+++ b/Engine/Plugins/Enterprise/GLTFExporter/MigratedFromGit/Source/GLTFExporter/Private/Tasks/GLTFMaterialTasks.h
@@ -46,11 +46,11 @@ private:
 	FString GetMaterialName() const;
 	FString GetBakedTextureName(const FString& PropertyName) const;
 
-	void ApplyPrebakedProperties(FGLTFJsonMaterial& OutMaterial) const;
-	void ApplyPrebakedProperty(const FString& PropertyName, float& OutValue) const;
-	void ApplyPrebakedProperty(const FString& PropertyName, FGLTFJsonColor3& OutValue) const;
-	void ApplyPrebakedProperty(const FString& PropertyName, FGLTFJsonColor4& OutValue) const;
-	void ApplyPrebakedProperty(const FString& PropertyName, FGLTFJsonTextureInfo& OutValue) const;
+	void GetProxyProperties(FGLTFJsonMaterial& OutMaterial) const;
+	void GetProxyProperty(const FString& PropertyName, float& OutValue) const;
+	void GetProxyProperty(const FString& PropertyName, FGLTFJsonColor3& OutValue) const;
+	void GetProxyProperty(const FString& PropertyName, FGLTFJsonColor4& OutValue) const;
+	void GetProxyProperty(const FString& PropertyName, FGLTFJsonTextureInfo& OutValue) const;
 
 	EMaterialShadingModel GetShadingModel() const;
 	void ConvertShadingModel(EGLTFJsonShadingModel& OutShadingModel) const;
diff --git a/Engine/Plugins/Enterprise/GLTFExporter/MigratedFromGit/Source/GLTFExporter/Private/UI/GLTFPrebakeOptionsWindow.cpp b/Engine/Plugins/Enterprise/GLTFExporter/MigratedFromGit/Source/GLTFExporter/Private/UI/GLTFPrebakeOptionsWindow.cpp
deleted file mode 100644
diff --git a/Engine/Plugins/Enterprise/GLTFExporter/MigratedFromGit/Source/GLTFExporter/Private/UI/GLTFPrebakeOptionsWindow.h b/Engine/Plugins/Enterprise/GLTFExporter/MigratedFromGit/Source/GLTFExporter/Private/UI/GLTFPrebakeOptionsWindow.h
deleted file mode 100644
diff --git a/Engine/Plugins/Enterprise/GLTFExporter/MigratedFromGit/Source/GLTFExporter/Private/UI/GLTFProxyOptionsWindow.cpp b/Engine/Plugins/Enterprise/GLTFExporter/MigratedFromGit/Source/GLTFExporter/Private/UI/GLTFProxyOptionsWindow.cpp
new file mode 100644
--- /dev/null
+++ b/Engine/Plugins/Enterprise/GLTFExporter/MigratedFromGit/Source/GLTFExporter/Private/UI/GLTFProxyOptionsWindow.cpp
@@ -0,0 +1,137 @@
+// Copyright Epic Games, Inc. All Rights Reserved.
+
+#if WITH_EDITOR
+
+#include "UI/GLTFProxyOptionsWindow.h"
+#include "Modules/ModuleManager.h"
+#include "PropertyEditorModule.h"
+#include "IDetailsView.h"
+#include "Interfaces/IMainFrameModule.h"
+#include "Widgets/Layout/SUniformGridPanel.h"
+#include "Widgets/Input/SButton.h"
+
+#define LOCTEXT_NAMESPACE "SGLTFProxyOptionsWindow"
+
+SGLTFProxyOptionsWindow::SGLTFProxyOptionsWindow() : bUserCancelled(true)
+{
+
+}
+
+void SGLTFProxyOptionsWindow::Construct(const FArguments& InArgs)
+{
+	WidgetWindow = InArgs._WidgetWindow;
+
+	// Retrieve property editor module and create a SDetailsView
+	FPropertyEditorModule& PropertyEditorModule = FModuleManager::GetModuleChecked<FPropertyEditorModule>("PropertyEditor");
+	FDetailsViewArgs DetailsViewArgs;
+	DetailsViewArgs.bAllowSearch = false;
+	DetailsViewArgs.NameAreaSettings = FDetailsViewArgs::HideNameArea;
+	DetailsViewArgs.bAllowMultipleTopLevelObjects = true;
+
+	DetailsView = PropertyEditorModule.CreateDetailView(DetailsViewArgs);
+
+	// Set provided objects on SDetailsView
+	DetailsView->SetObjects(InArgs._SettingsObjects, true);
+
+	this->ChildSlot
+	[
+		SNew(SVerticalBox)
+
+		+ SVerticalBox::Slot()
+		.Padding(2)
+		.MaxHeight(500.0f)
+		[
+			DetailsView->AsShared()
+		]
+
+		+ SVerticalBox::Slot()
+		.AutoHeight()
+		.HAlign(HAlign_Right)
+		.Padding(2)
+		[
+			SNew(SUniformGridPanel)
+			.SlotPadding(2)
+			+ SUniformGridPanel::Slot(0, 0)
+			[
+				SAssignNew(ConfirmButton, SButton)
+				.HAlign(HAlign_Center)
+				.Text(LOCTEXT("MaterialBakeOptionWindow_Import", "Confirm"))
+				.OnClicked(this, &SGLTFProxyOptionsWindow::OnConfirm)
+			]
+			+ SUniformGridPanel::Slot(1, 0)
+			[
+				SNew(SButton)
+				.HAlign(HAlign_Center)
+				.Text(LOCTEXT("MaterialBakeOptionWindow_Cancel", "Cancel"))
+				.ToolTipText(LOCTEXT("MaterialBakeOptionWindow_Cancel_ToolTip", "Cancels baking out Material"))
+				.OnClicked(this, &SGLTFProxyOptionsWindow::OnCancel)
+			]
+		]
+	];
+}
+
+FReply SGLTFProxyOptionsWindow::OnConfirm()
+{
+	bUserCancelled = false;
+	if (WidgetWindow.IsValid())
+	{
+		WidgetWindow.Pin()->RequestDestroyWindow();
+	}
+	return FReply::Handled();
+}
+
+FReply SGLTFProxyOptionsWindow::OnCancel()
+{
+	if (WidgetWindow.IsValid())
+	{
+		WidgetWindow.Pin()->RequestDestroyWindow();
+	}
+	return FReply::Handled();
+}
+
+FReply SGLTFProxyOptionsWindow::OnKeyDown(const FGeometry& MyGeometry, const FKeyEvent& InKeyEvent)
+{
+	if (InKeyEvent.GetKey() == EKeys::Escape)
+	{
+		return OnCancel();
+	}
+
+	return FReply::Unhandled();
+}
+
+bool SGLTFProxyOptionsWindow::WasUserCancelled()
+{
+	return bUserCancelled;
+}
+
+bool SGLTFProxyOptionsWindow::ShowDialog(UGLTFProxyOptions* ProxyOptions)
+{
+	TSharedRef<SWindow> Window = SNew(SWindow)
+		.Title(LOCTEXT("WindowTitle", "Material Baking Options"))
+		.SizingRule(ESizingRule::Autosized);
+
+	TArray<TWeakObjectPtr<UObject>> OptionObjects = { ProxyOptions };
+	TSharedPtr<SGLTFProxyOptionsWindow> Options;
+
+	Window->SetContent
+	(
+		SAssignNew(Options, SGLTFProxyOptionsWindow)
+		.WidgetWindow(Window)
+		.SettingsObjects(OptionObjects)
+	);
+
+	TSharedPtr<SWindow> ParentWindow;
+	if (FModuleManager::Get().IsModuleLoaded("MainFrame"))
+	{
+		IMainFrameModule& MainFrame = FModuleManager::LoadModuleChecked<IMainFrameModule>("MainFrame");
+		ParentWindow = MainFrame.GetParentWindow();
+		FSlateApplication::Get().AddModalWindow(Window, ParentWindow, false);
+		return !Options->WasUserCancelled();
+	}
+
+	return false;
+}
+
+#undef LOCTEXT_NAMESPACE //"SGLTFProxyOptionsWindow"
+
+#endif
diff --git a/Engine/Plugins/Enterprise/GLTFExporter/MigratedFromGit/Source/GLTFExporter/Private/UI/GLTFProxyOptionsWindow.h b/Engine/Plugins/Enterprise/GLTFExporter/MigratedFromGit/Source/GLTFExporter/Private/UI/GLTFProxyOptionsWindow.h
new file mode 100644
--- /dev/null
+++ b/Engine/Plugins/Enterprise/GLTFExporter/MigratedFromGit/Source/GLTFExporter/Private/UI/GLTFProxyOptionsWindow.h
@@ -0,0 +1,53 @@
+// Copyright Epic Games, Inc. All Rights Reserved.
+
+#pragma once
+
+#if WITH_EDITOR
+
+#include "Options/GLTFProxyOptions.h"
+#include "Widgets/SCompoundWidget.h"
+#include "Widgets/DeclarativeSyntaxSupport.h"
+#include "Widgets/SWindow.h"
+
+class SButton;
+
+/** Options window used to populate provided settings objects */
+class SGLTFProxyOptionsWindow : public SCompoundWidget
+{
+public:
+
+	SLATE_BEGIN_ARGS(SGLTFProxyOptionsWindow)
+	{}
+		SLATE_ARGUMENT(TSharedPtr<SWindow>, WidgetWindow)
+		SLATE_ARGUMENT(TArray<TWeakObjectPtr<UObject>>, SettingsObjects)
+	SLATE_END_ARGS()
+
+	SGLTFProxyOptionsWindow();
+	void Construct(const FArguments& InArgs);
+
+	/** Begin SCompoundWidget overrides */
+	virtual bool SupportsKeyboardFocus() const override { return true; }
+	virtual FReply OnKeyDown(const FGeometry& MyGeometry, const FKeyEvent& InKeyEvent) override;
+	/** End SCompoundWidget overrides */
+
+	/** Callbacks used for Confirm and Cancel buttons */
+	FReply OnConfirm();
+	FReply OnCancel();
+
+	/** Returns whether or not the user cancelled the operation */
+	bool WasUserCancelled();
+
+	static bool ShowDialog(UGLTFProxyOptions* ProxyOptions);
+
+private:
+	/** Owning window this widget is part of */
+	TWeakPtr< SWindow > WidgetWindow;
+	/** Whether or not the cancel button was clicked by the user */
+	bool bUserCancelled;
+	/** Detailsview used to display SettingsObjects, and allowing user to change options */
+	TSharedPtr<class IDetailsView> DetailsView;
+	/** Shared ptr to Confirm button */
+	TSharedPtr<SButton> ConfirmButton;
+};
+
+#endif
diff --git a/Engine/Plugins/Enterprise/GLTFExporter/MigratedFromGit/Source/GLTFExporter/Public/GLTFExporterModule.h b/Engine/Plugins/Enterprise/GLTFExporter/MigratedFromGit/Source/GLTFExporter/Public/GLTFExporterModule.h
--- a/Engine/Plugins/Enterprise/GLTFExporter/MigratedFromGit/Source/GLTFExporter/Public/GLTFExporterModule.h
+++ b/Engine/Plugins/Enterprise/GLTFExporter/MigratedFromGit/Source/GLTFExporter/Public/GLTFExporterModule.h
@@ -10,7 +10,7 @@
 
 DECLARE_LOG_CATEGORY_EXTERN(LogGLTFExporter, Log, All);
 
-class UGLTFPrebakeOptions;
+class UGLTFProxyOptions;
 
 /**
  * The public interface of the GLTFExporter module
@@ -39,5 +39,5 @@ public:
 		return FModuleManager::Get().IsModuleLoaded(GLTFEXPORTER_MODULE_NAME);
 	}
 
-	static UMaterialInterface* PrebakeMaterial(UMaterialInterface* Material, const UGLTFPrebakeOptions* Options);
+	static UMaterialInterface* CreateProxyMaterial(UMaterialInterface* Material, const UGLTFProxyOptions* Options);
 };
diff --git a/Engine/Plugins/Enterprise/GLTFExporter/MigratedFromGit/Source/GLTFExporter/Public/Options/GLTFExportOptions.h b/Engine/Plugins/Enterprise/GLTFExporter/MigratedFromGit/Source/GLTFExporter/Public/Options/GLTFExportOptions.h
--- a/Engine/Plugins/Enterprise/GLTFExporter/MigratedFromGit/Source/GLTFExporter/Public/Options/GLTFExportOptions.h
+++ b/Engine/Plugins/Enterprise/GLTFExporter/MigratedFromGit/Source/GLTFExporter/Public/Options/GLTFExportOptions.h
@@ -99,7 +99,7 @@ class GLTFEXPORTER_API UGLTFExportOptions : public UObject
 	UPROPERTY(EditAnywhere, BlueprintReadWrite, Config, Category = General)
 	bool bIncludeGeneratorVersion;
 
-	/** If enabled, materials that have a proxy defined in their user data, will be exported using that proxy instead. Used mainly by pre-baked materials. */
+	/** If enabled, materials that have a proxy defined in their user data, will be exported using that proxy instead. This setting won't affect proxy materials exported or referenced directly. */
 	UPROPERTY(EditAnywhere, BlueprintReadWrite, Config, Category = Material)
 	bool bExportProxyMaterials;
 
diff --git a/Engine/Plugins/Enterprise/GLTFExporter/MigratedFromGit/Source/GLTFExporter/Public/Options/GLTFPrebakeOptions.h b/Engine/Plugins/Enterprise/GLTFExporter/MigratedFromGit/Source/GLTFExporter/Public/Options/GLTFPrebakeOptions.h
deleted file mode 100644
diff --git a/Engine/Plugins/Enterprise/GLTFExporter/MigratedFromGit/Source/GLTFExporter/Public/Options/GLTFProxyOptions.h b/Engine/Plugins/Enterprise/GLTFExporter/MigratedFromGit/Source/GLTFExporter/Public/Options/GLTFProxyOptions.h
new file mode 100644
--- /dev/null
+++ b/Engine/Plugins/Enterprise/GLTFExporter/MigratedFromGit/Source/GLTFExporter/Public/Options/GLTFProxyOptions.h
@@ -0,0 +1,31 @@
+// Copyright Epic Games, Inc. All Rights Reserved.
+
+#pragma once
+
+#include "Options/GLTFExportOptions.h"
+#include "GLTFProxyOptions.generated.h"
+
+UCLASS(BlueprintType, Config=EditorPerProjectUserSettings, HideCategories=(DebugProperty))
+class GLTFEXPORTER_API UGLTFProxyOptions : public UObject
+{
+	GENERATED_UCLASS_BODY()
+
+	/** Default size of the baked out texture (containing the material input). Can be overridden by material- and input-specific bake settings, see GLTFMaterialExportOptions. */
+	UPROPERTY(EditAnywhere, BlueprintReadWrite, Config, Category = Material)
+	EGLTFMaterialBakeSizePOT DefaultMaterialBakeSize;
+
+	/** Default filtering mode used when sampling the baked out texture. Can be overridden by material- and input-specific bake settings, see GLTFMaterialExportOptions. */
+	UPROPERTY(EditAnywhere, BlueprintReadWrite, Config, Category = Material, Meta = (ValidEnumValues="TF_Nearest, TF_Bilinear, TF_Trilinear"))
+	TEnumAsByte<TextureFilter> DefaultMaterialBakeFilter;
+
+	/** Default addressing mode used when sampling the baked out texture. Can be overridden by material- and input-specific bake settings, see GLTFMaterialExportOptions. */
+	UPROPERTY(EditAnywhere, BlueprintReadWrite, Config, Category = Material)
+	TEnumAsByte<TextureAddress> DefaultMaterialBakeTiling;
+
+	/** Input-specific default bake settings that override the general defaults above. */
+	UPROPERTY(EditAnywhere, BlueprintReadWrite, Config, Category = Material)
+	TMap<EGLTFMaterialPropertyGroup, FGLTFOverrideMaterialBakeSettings> DefaultInputBakeSettings;
+
+	UFUNCTION(BlueprintCallable, Category = General)
+	void ResetToDefault();
+};
