[GLTFExporter] fix: address minor warnings from Rider code inspection

commit 6c17e32da9b3989c8be845f95e68f06a23878fe1
Date: Fri, 11 Feb 2022 18:01:57 +0100

    fix: address minor warnings from Rider code inspection
    (cherry picked from commit a7859664576ea05eb78171ba5f7c0ceda0feadef​)

#skipci

diff --git a/Engine/Plugins/Enterprise/GLTFExporter/MigratedFromGit/Source/GLTFExporter/Private/Builders/GLTFMemoryArchive.h b/Engine/Plugins/Enterprise/GLTFExporter/MigratedFromGit/Source/GLTFExporter/Private/Builders/GLTFMemoryArchive.h
--- a/Engine/Plugins/Enterprise/GLTFExporter/MigratedFromGit/Source/GLTFExporter/Private/Builders/GLTFMemoryArchive.h
+++ b/Engine/Plugins/Enterprise/GLTFExporter/MigratedFromGit/Source/GLTFExporter/Private/Builders/GLTFMemoryArchive.h
@@ -18,8 +18,6 @@ public:
 	{
 		check(Length >= 0);
 
-		auto x = Num();
-
 		const int64 NumBytesToAdd = Offset + Length - Num();
 		AddUninitialized(NumBytesToAdd);
 
diff --git a/Engine/Plugins/Enterprise/GLTFExporter/MigratedFromGit/Source/GLTFExporter/Private/Converters/GLTFSceneConverters.cpp b/Engine/Plugins/Enterprise/GLTFExporter/MigratedFromGit/Source/GLTFExporter/Private/Converters/GLTFSceneConverters.cpp
--- a/Engine/Plugins/Enterprise/GLTFExporter/MigratedFromGit/Source/GLTFExporter/Private/Converters/GLTFSceneConverters.cpp
+++ b/Engine/Plugins/Enterprise/GLTFExporter/MigratedFromGit/Source/GLTFExporter/Private/Converters/GLTFSceneConverters.cpp
@@ -2,7 +2,6 @@
 
 #include "Converters/GLTFSceneConverters.h"
 #include "Builders/GLTFContainerBuilder.h"
-#include "Converters/GLTFActorUtility.h"
 #include "LevelVariantSetsActor.h"
 #include "LevelVariantSets.h"
 #include "VariantSet.h"
diff --git a/Engine/Plugins/Enterprise/GLTFExporter/MigratedFromGit/Source/GLTFExporter/Private/Core/GLTFColor.h b/Engine/Plugins/Enterprise/GLTFExporter/MigratedFromGit/Source/GLTFExporter/Private/Core/GLTFColor.h
--- a/Engine/Plugins/Enterprise/GLTFExporter/MigratedFromGit/Source/GLTFExporter/Private/Core/GLTFColor.h
+++ b/Engine/Plugins/Enterprise/GLTFExporter/MigratedFromGit/Source/GLTFExporter/Private/Core/GLTFColor.h
@@ -2,7 +2,7 @@
 
 #pragma once
 
-#include "Core/GLTFCore.h"
+#include "CoreTypes.h"
 
 #if PLATFORM_SUPPORTS_PRAGMA_PACK
 	#pragma pack(push, 1)
diff --git a/Engine/Plugins/Enterprise/GLTFExporter/MigratedFromGit/Source/GLTFExporter/Private/Core/GLTFCore.h b/Engine/Plugins/Enterprise/GLTFExporter/MigratedFromGit/Source/GLTFExporter/Private/Core/GLTFCore.h
deleted file mode 100644
diff --git a/Engine/Plugins/Enterprise/GLTFExporter/MigratedFromGit/Source/GLTFExporter/Private/Core/GLTFMatrix.h b/Engine/Plugins/Enterprise/GLTFExporter/MigratedFromGit/Source/GLTFExporter/Private/Core/GLTFMatrix.h
--- a/Engine/Plugins/Enterprise/GLTFExporter/MigratedFromGit/Source/GLTFExporter/Private/Core/GLTFMatrix.h
+++ b/Engine/Plugins/Enterprise/GLTFExporter/MigratedFromGit/Source/GLTFExporter/Private/Core/GLTFMatrix.h
@@ -2,7 +2,7 @@
 
 #pragma once
 
-#include "Core/GLTFCore.h"
+#include "CoreTypes.h"
 
 #if PLATFORM_SUPPORTS_PRAGMA_PACK
 	#pragma pack(push, 1)
diff --git a/Engine/Plugins/Enterprise/GLTFExporter/MigratedFromGit/Source/GLTFExporter/Private/Core/GLTFQuaternion.h b/Engine/Plugins/Enterprise/GLTFExporter/MigratedFromGit/Source/GLTFExporter/Private/Core/GLTFQuaternion.h
--- a/Engine/Plugins/Enterprise/GLTFExporter/MigratedFromGit/Source/GLTFExporter/Private/Core/GLTFQuaternion.h
+++ b/Engine/Plugins/Enterprise/GLTFExporter/MigratedFromGit/Source/GLTFExporter/Private/Core/GLTFQuaternion.h
@@ -2,7 +2,7 @@
 
 #pragma once
 
-#include "Core/GLTFCore.h"
+#include "CoreTypes.h"
 
 #if PLATFORM_SUPPORTS_PRAGMA_PACK
 	#pragma pack(push, 1)
diff --git a/Engine/Plugins/Enterprise/GLTFExporter/MigratedFromGit/Source/GLTFExporter/Private/Core/GLTFVector.h b/Engine/Plugins/Enterprise/GLTFExporter/MigratedFromGit/Source/GLTFExporter/Private/Core/GLTFVector.h
--- a/Engine/Plugins/Enterprise/GLTFExporter/MigratedFromGit/Source/GLTFExporter/Private/Core/GLTFVector.h
+++ b/Engine/Plugins/Enterprise/GLTFExporter/MigratedFromGit/Source/GLTFExporter/Private/Core/GLTFVector.h
@@ -2,7 +2,7 @@
 
 #pragma once
 
-#include "Core/GLTFCore.h"
+#include "CoreTypes.h"
 
 #if PLATFORM_SUPPORTS_PRAGMA_PACK
 	#pragma pack(push, 1)
diff --git a/Engine/Plugins/Enterprise/GLTFExporter/MigratedFromGit/Source/GLTFMaterialBaking/Private/MaterialBakingHelpers.cpp b/Engine/Plugins/Enterprise/GLTFExporter/MigratedFromGit/Source/GLTFMaterialBaking/Private/MaterialBakingHelpers.cpp
--- a/Engine/Plugins/Enterprise/GLTFExporter/MigratedFromGit/Source/GLTFMaterialBaking/Private/MaterialBakingHelpers.cpp
+++ b/Engine/Plugins/Enterprise/GLTFExporter/MigratedFromGit/Source/GLTFMaterialBaking/Private/MaterialBakingHelpers.cpp
@@ -68,7 +68,6 @@ namespace FMaterialBakingHelpersImpl
 
 		check(InOutPixels.Num() == ImageWidth*ImageHeight);
 
-		const uint32 MagentaMask = FColor(255, 0, 255).DWColor();
 		const int32 PaddedImageWidth = ImageWidth + 2;
 		const int32 PaddedImageHeight = ImageHeight + 2;
 
diff --git a/Engine/Plugins/Enterprise/GLTFExporter/MigratedFromGit/Source/GLTFMaterialBaking/Private/MaterialBakingModule.cpp b/Engine/Plugins/Enterprise/GLTFExporter/MigratedFromGit/Source/GLTFMaterialBaking/Private/MaterialBakingModule.cpp
--- a/Engine/Plugins/Enterprise/GLTFExporter/MigratedFromGit/Source/GLTFMaterialBaking/Private/MaterialBakingModule.cpp
+++ b/Engine/Plugins/Enterprise/GLTFExporter/MigratedFromGit/Source/GLTFMaterialBaking/Private/MaterialBakingModule.cpp
@@ -512,7 +512,6 @@ void FMaterialBakingModule::BakeMaterials(const TArray<FMaterialDataEx*>& Materi
 		}
 
 		const FMaterialDataEx* CurrentMaterialSettings = MaterialSettings[MaterialIndex];
-		const FMeshData* CurrentMeshSettings = MeshSettings[MaterialIndex];
 		FBakeOutputEx& CurrentOutput = Output[MaterialIndex];
 
 		TArray<FMaterialPropertyEx> MaterialPropertiesToBakeOut;
