[GLTFExporter] refactor: separate code in material proxy factory into cleaner functions

commit 63c5d80bd74be69f2a090beba44184084828c3cd
Date: Fri, 11 Feb 2022 19:11:40 +0100

    refactor: separate code in material proxy factory into cleaner functions
    (cherry picked from commit 23db6a9465e78ddbb4cd816bf9697bfc35d833ad​)

#skipci

diff --git a/Engine/Plugins/Enterprise/GLTFExporter/MigratedFromGit/Source/GLTFExporter/Private/Converters/GLTFMaterialProxyFactory.cpp b/Engine/Plugins/Enterprise/GLTFExporter/MigratedFromGit/Source/GLTFExporter/Private/Converters/GLTFMaterialProxyFactory.cpp
--- a/Engine/Plugins/Enterprise/GLTFExporter/MigratedFromGit/Source/GLTFExporter/Private/Converters/GLTFMaterialProxyFactory.cpp
+++ b/Engine/Plugins/Enterprise/GLTFExporter/MigratedFromGit/Source/GLTFExporter/Private/Converters/GLTFMaterialProxyFactory.cpp
@@ -38,19 +38,60 @@
 	}
 
 	MakeDirectory(RootPath);
-
 	Builder.CompleteAllTasks();
+
 	const FGLTFJsonMaterial& JsonMaterial = Builder.GetMaterial(MaterialIndex);
-
 	UMaterialInstanceConstant* ProxyMaterial = CreateInstancedMaterial(OriginalMaterial, JsonMaterial.ShadingModel);
+
 	if (ProxyMaterial != nullptr)
 	{
+		SetUserData(ProxyMaterial, OriginalMaterial);
+		SetBaseProperties(ProxyMaterial, OriginalMaterial);
 		SetProxyProperties(ProxyMaterial, JsonMaterial);
 	}
 
 	return ProxyMaterial;
 }
 
+void FGLTFMaterialProxyFactory::SetUserData(UMaterialInstanceConstant* ProxyMaterial, UMaterialInterface* OriginalMaterial)
+{
+	UGLTFMaterialExportOptions* UserData = OriginalMaterial->GetAssetUserData<UGLTFMaterialExportOptions>();
+	if (UserData == nullptr)
+	{
+		UserData = NewObject<UGLTFMaterialExportOptions>();
+		OriginalMaterial->AddAssetUserData(UserData);
+	}
+
+	UserData->Proxy = ProxyMaterial;
+	OriginalMaterial->Modify();
+}
+
+void FGLTFMaterialProxyFactory::SetBaseProperties(UMaterialInstanceConstant* ProxyMaterial, UMaterialInterface* OriginalMaterial)
+{
+	const UMaterial* BaseMaterial = ProxyMaterial->GetMaterial();
+
+	const bool bTwoSided = OriginalMaterial->IsTwoSided();
+	if (bTwoSided != BaseMaterial->IsTwoSided())
+	{
+		ProxyMaterial->BasePropertyOverrides.bOverride_TwoSided = true;
+		ProxyMaterial->BasePropertyOverrides.TwoSided = bTwoSided;
+	}
+
+	const EBlendMode BlendMode = OriginalMaterial->GetBlendMode();
+	if (BlendMode != BaseMaterial->GetBlendMode())
+	{
+		ProxyMaterial->BasePropertyOverrides.bOverride_BlendMode = true;
+		ProxyMaterial->BasePropertyOverrides.BlendMode = BlendMode;
+	}
+
+	const float OpacityMaskClipValue = OriginalMaterial->GetOpacityMaskClipValue();
+	if (OpacityMaskClipValue != BaseMaterial->GetOpacityMaskClipValue())
+	{
+		ProxyMaterial->BasePropertyOverrides.bOverride_OpacityMaskClipValue = true;
+		ProxyMaterial->BasePropertyOverrides.OpacityMaskClipValue = OpacityMaskClipValue;
+	}
+}
+
 void FGLTFMaterialProxyFactory::SetProxyProperties(UMaterialInstanceConstant* ProxyMaterial, const FGLTFJsonMaterial& JsonMaterial)
 {
 	SetProxyProperty(ProxyMaterial, TEXT("Base Color Factor"), JsonMaterial.PBRMetallicRoughness.BaseColorFactor);
@@ -198,11 +239,6 @@
 	const bool bSRGB = PropertyGroup == EGLTFMaterialPropertyGroup::BaseColorOpacity || PropertyGroup == EGLTFMaterialPropertyGroup::EmissiveColor;
 	const bool bNormalMap = PropertyGroup == EGLTFMaterialPropertyGroup::Normal || PropertyGroup == EGLTFMaterialPropertyGroup::ClearCoatBottomNormal;
 
-	const FString PackageName = RootPath / TEXT("T_GLTF_") + ImageData->Filename;
-	UPackage* Package = CreatePackage(*PackageName);
-	Package->FullyLoad();
-	Package->Modify();
-
 	FCreateTexture2DParameters TexParams;
 	TexParams.bUseAlpha = !ImageData->bIgnoreAlpha;
 	TexParams.CompressionSettings = bNormalMap ? TC_Normalmap : TC_Default;
@@ -211,22 +247,18 @@
 	TexParams.TextureGroup = bNormalMap ? TEXTUREGROUP_WorldNormalMap : TEXTUREGROUP_World;
 	TexParams.SourceGuidHash = FGuid();
 
-	UTexture2D* Texture = FImageUtils::CreateTexture2D(ImageData->Size.X, ImageData->Size.Y, *ImageData->Pixels, Package, FPaths::GetCleanFilename(Package->GetName()), RF_Public | RF_Standalone, TexParams);
+	const FString BaseName = TEXT("T_GLTF_") + ImageData->Filename;
+	UPackage* Package = FindOrCreatePackage(BaseName);
+
+	UTexture2D* Texture = FImageUtils::CreateTexture2D(ImageData->Size.X, ImageData->Size.Y, *ImageData->Pixels, Package, BaseName, RF_Public | RF_Standalone, TexParams);
 	Texture->Filter = ConvertFilter(JsonSampler.MagFilter);
 	Texture->AddressX = ConvertWrap(JsonSampler.WrapS);
 	Texture->AddressY = ConvertWrap(JsonSampler.WrapT);
-
-	Texture->PostEditChange();
 	return Texture;
 }
 
 UMaterialInstanceConstant* FGLTFMaterialProxyFactory::CreateInstancedMaterial(UMaterialInterface* OriginalMaterial, EGLTFJsonShadingModel ShadingModel)
 {
-	const FString PackageName = RootPath / TEXT("M_GLTF_") + OriginalMaterial->GetName();
-	UPackage* Package = CreatePackage(*PackageName);
-	Package->FullyLoad();
-	Package->Modify();
-
 	UMaterialInterface* BaseMaterial = FGLTFMaterialUtility::GetProxyBaseMaterial(ShadingModel);
 	if (BaseMaterial == nullptr)
 	{
@@ -237,40 +269,19 @@
 		return nullptr;
 	}
 
+	UPackage* Package = FindOrCreatePackage(TEXT("M_GLTF_") + OriginalMaterial->GetName());
 	const FString BaseName = TEXT("GLTF_") + OriginalMaterial->GetName(); // prefix "M_" is added automatically by CreateInstancedMaterial
-	UMaterialInstanceConstant* ProxyMaterial = FMaterialUtilities::CreateInstancedMaterial(BaseMaterial, Package, BaseName, RF_Public | RF_Standalone);
-
-	const bool bTwoSided = OriginalMaterial->IsTwoSided();
-	if (bTwoSided != BaseMaterial->IsTwoSided())
-	{
-		ProxyMaterial->BasePropertyOverrides.bOverride_TwoSided = true;
-		ProxyMaterial->BasePropertyOverrides.TwoSided = bTwoSided;
-	}
-
-	const EBlendMode BlendMode = OriginalMaterial->GetBlendMode();
-	if (BlendMode != BaseMaterial->GetBlendMode())
-	{
-		ProxyMaterial->BasePropertyOverrides.bOverride_BlendMode = true;
-		ProxyMaterial->BasePropertyOverrides.BlendMode = BlendMode;
-	}
-
-	const float OpacityMaskClipValue = OriginalMaterial->GetOpacityMaskClipValue();
-	if (OpacityMaskClipValue != BaseMaterial->GetOpacityMaskClipValue())
-	{
-		ProxyMaterial->BasePropertyOverrides.bOverride_OpacityMaskClipValue = true;
-		ProxyMaterial->BasePropertyOverrides.OpacityMaskClipValue = OpacityMaskClipValue;
-	}
-
-	UGLTFMaterialExportOptions* UserData = OriginalMaterial->GetAssetUserData<UGLTFMaterialExportOptions>();
-	if (UserData == nullptr)
-	{
-		UserData = NewObject<UGLTFMaterialExportOptions>();
-		OriginalMaterial->AddAssetUserData(UserData);
-	}
-
-	UserData->Proxy = ProxyMaterial;
-	OriginalMaterial->Modify();
-	return ProxyMaterial;
+	return FMaterialUtilities::CreateInstancedMaterial(BaseMaterial, Package, BaseName, RF_Public | RF_Standalone);
 }
 
+UPackage* FGLTFMaterialProxyFactory::FindOrCreatePackage(const FString& BaseName)
+{
+	const FString PackageName = RootPath / BaseName;
+	UPackage* Package = CreatePackage(*PackageName);
+	Package->FullyLoad();
+	Package->Modify();
+	// TODO: do we need to delete any old assets in the package?
+	return Package;
+}
+
 TUniquePtr<IGLTFImageConverter> FGLTFMaterialProxyFactory::CreateCustomImageConverter()
@@ -279,11 +290,11 @@
 	{
 	public:
 
-		FGLTFMaterialProxyFactory& Proxyr;
+		FGLTFMaterialProxyFactory& Factory;
 		TSet<FString> UniqueFilenames;
 
-		FGLTFCustomImageConverter(FGLTFMaterialProxyFactory& Proxyr)
-			: Proxyr(Proxyr)
+		FGLTFCustomImageConverter(FGLTFMaterialProxyFactory& Factory)
+			: Factory(Factory)
 		{
 		}
 
@@ -294,8 +305,8 @@
 			const FString Filename = FGLTFImageUtility::GetUniqueFilename(Name, TEXT(""), UniqueFilenames);
 			UniqueFilenames.Add(Filename);
 
-			const FGLTFJsonImageIndex ImageIndex = Proxyr.Builder.AddImage();
-			Proxyr.Images.Add(ImageIndex, { Filename, Type, bIgnoreAlpha, Size, Pixels });
+			const FGLTFJsonImageIndex ImageIndex = Factory.Builder.AddImage();
+			Factory.Images.Add(ImageIndex, { Filename, Type, bIgnoreAlpha, Size, Pixels });
 			return ImageIndex;
 		}
 	};
diff --git a/Engine/Plugins/Enterprise/GLTFExporter/MigratedFromGit/Source/GLTFExporter/Private/Converters/GLTFMaterialProxyFactory.h b/Engine/Plugins/Enterprise/GLTFExporter/MigratedFromGit/Source/GLTFExporter/Private/Converters/GLTFMaterialProxyFactory.h
--- a/Engine/Plugins/Enterprise/GLTFExporter/MigratedFromGit/Source/GLTFExporter/Private/Converters/GLTFMaterialProxyFactory.h
+++ b/Engine/Plugins/Enterprise/GLTFExporter/MigratedFromGit/Source/GLTFExporter/Private/Converters/GLTFMaterialProxyFactory.h
@@ -31,6 +31,8 @@ private:
 		TGLTFSharedArray<FColor> Pixels;
 	};
 
+	void SetUserData(UMaterialInstanceConstant* ProxyMaterial, UMaterialInterface* OriginalMaterial);
+	void SetBaseProperties(UMaterialInstanceConstant* ProxyMaterial, UMaterialInterface* OriginalMaterial);
 	void SetProxyProperties(UMaterialInstanceConstant* ProxyMaterial, const FGLTFJsonMaterial& JsonMaterial);
 
 	void SetProxyProperty(UMaterialInstanceConstant* ProxyMaterial, const FString& PropertyName, float Scalar);
@@ -43,6 +45,8 @@ private:
 
 	UMaterialInstanceConstant* CreateInstancedMaterial(UMaterialInterface* OriginalMaterial, EGLTFJsonShadingModel ShadingModel);
 
+	UPackage* FindOrCreatePackage(const FString& BaseName);
+
 	TUniquePtr<IGLTFImageConverter> CreateCustomImageConverter();
 
 	static bool MakeDirectory(const FString& String);
