[GLTFExporter] refactor: rename proxy materials (and related assets) to follow common naming conventions for unreal content

commit 819595bc96299b0780d1b2fdb6b0e29ef26af9bf
Date: Fri, 11 Feb 2022 21:56:35 +0100

    refactor: rename proxy materials (and related assets) to follow common naming conventions for unreal content
    (cherry picked from commit 0dc65545384398add609ead0c86edafda7e17890​)

#skipci

diff --git a/Engine/Plugins/Enterprise/GLTFExporter/MigratedFromGit/Source/GLTFExporter/Private/Converters/GLTFMaterialUtility.cpp b/Engine/Plugins/Enterprise/GLTFExporter/MigratedFromGit/Source/GLTFExporter/Private/Converters/GLTFMaterialUtility.cpp
--- a/Engine/Plugins/Enterprise/GLTFExporter/MigratedFromGit/Source/GLTFExporter/Private/Converters/GLTFMaterialUtility.cpp
+++ b/Engine/Plugins/Enterprise/GLTFExporter/MigratedFromGit/Source/GLTFExporter/Private/Converters/GLTFMaterialUtility.cpp
@@ -15,8 +15,9 @@
 #include "Materials/MaterialExpressionTextureCoordinate.h"
 #include "MeshDescription.h"
 #endif
-#include "GLTFConverterUtility.h"
-#include "Misc/DefaultValueHelper.h"
+
+#define PROXY_MATERIAL_NAME_PREFIX TEXT("M_GLTF_")
+#define PROXY_MATERIAL_ROOT_PATH   TEXT("/GLTFExporter/Materials/Proxy/")
 
 UMaterialInterface* FGLTFMaterialUtility::GetDefaultMaterial()
 {
@@ -27,13 +28,13 @@ UMaterialInterface* FGLTFMaterialUtility::GetDefaultMaterial()
 UMaterialInterface* FGLTFMaterialUtility::GetProxyBaseMaterial(EGLTFJsonShadingModel ShadingModel)
 {
 	const FString Name = FGLTFJsonUtility::GetValue(ShadingModel);
-	const FString Path = TEXT("/GLTFExporter/Materials/GLTF") + Name + TEXT(".GLTF") + Name;
+	const FString Path = PROXY_MATERIAL_ROOT_PATH PROXY_MATERIAL_NAME_PREFIX + Name + TEXT(".") PROXY_MATERIAL_NAME_PREFIX + Name;
 	return LoadObject<UMaterialInterface>(nullptr, *Path);
 }
 
 bool FGLTFMaterialUtility::IsProxyMaterial(const UMaterial* Material)
 {
-	return Material->GetPathName().StartsWith(TEXT("/GLTFExporter/Materials/GLTF"));
+	return Material->GetPathName().StartsWith(PROXY_MATERIAL_ROOT_PATH PROXY_MATERIAL_NAME_PREFIX);
 }
 
 bool FGLTFMaterialUtility::IsProxyMaterial(const UMaterialInterface* Material)
