[GLTFExporter] refactor: use prefix "MI_" (instead of "M_") when creating proxy materials

commit fb2e60adc014170c57c1bc77129e2a25aeee8cc0
Date: Fri, 11 Feb 2022 22:28:19 +0100

    refactor: use prefix "MI_" (instead of "M_") when creating proxy materials
    (cherry picked from commit d68a6e5f820851d866b825cf5bb70e31e303f150​)

#skipci

diff --git a/Engine/Plugins/Enterprise/GLTFExporter/MigratedFromGit/Source/GLTFExporter/Private/Converters/GLTFMaterialProxyFactory.cpp b/Engine/Plugins/Enterprise/GLTFExporter/MigratedFromGit/Source/GLTFExporter/Private/Converters/GLTFMaterialProxyFactory.cpp
--- a/Engine/Plugins/Enterprise/GLTFExporter/MigratedFromGit/Source/GLTFExporter/Private/Converters/GLTFMaterialProxyFactory.cpp
+++ b/Engine/Plugins/Enterprise/GLTFExporter/MigratedFromGit/Source/GLTFExporter/Private/Converters/GLTFMaterialProxyFactory.cpp
@@ -7,7 +7,6 @@
 #include "Converters/GLTFImageUtility.h"
 #include "AssetRegistry/AssetRegistryModule.h"
 #include "Materials/MaterialInstanceConstant.h"
-#include "MaterialUtilities.h"
 #include "ImageUtils.h"
 
 FGLTFMaterialProxyFactory::FGLTFMaterialProxyFactory(const UGLTFProxyOptions* Options)
@@ -269,9 +268,12 @@ UMaterialInstanceConstant* FGLTFMaterialProxyFactory::CreateInstancedMaterial(UM
 		return nullptr;
 	}
 
-	UPackage* Package = FindOrCreatePackage(TEXT("M_GLTF_") + OriginalMaterial->GetName());
-	const FString BaseName = TEXT("GLTF_") + OriginalMaterial->GetName(); // prefix "M_" is added automatically by CreateInstancedMaterial
-	return FMaterialUtilities::CreateInstancedMaterial(BaseMaterial, Package, BaseName, RF_Public | RF_Standalone);
+	const FString BaseName = TEXT("MI_GLTF_") + OriginalMaterial->GetName();
+	UPackage* Package = FindOrCreatePackage(BaseName);
+
+	UMaterialInstanceConstant* MaterialInstance = NewObject<UMaterialInstanceConstant>(Package, *BaseName, RF_Public | RF_Standalone);
+	MaterialInstance->Parent = BaseMaterial;
+	return MaterialInstance;
 }
 
 UPackage* FGLTFMaterialProxyFactory::FindOrCreatePackage(const FString& BaseName)
diff --git a/Engine/Plugins/Enterprise/GLTFExporter/MigratedFromGit/Source/GLTFExporter/Private/Converters/GLTFMaterialUtility.cpp b/Engine/Plugins/Enterprise/GLTFExporter/MigratedFromGit/Source/GLTFExporter/Private/Converters/GLTFMaterialUtility.cpp
--- a/Engine/Plugins/Enterprise/GLTFExporter/MigratedFromGit/Source/GLTFExporter/Private/Converters/GLTFMaterialUtility.cpp
+++ b/Engine/Plugins/Enterprise/GLTFExporter/MigratedFromGit/Source/GLTFExporter/Private/Converters/GLTFMaterialUtility.cpp
@@ -3,6 +3,7 @@
 #include "Converters/GLTFMaterialUtility.h"
 #include "Converters/GLTFTextureUtility.h"
 #include "Converters/GLTFNameUtility.h"
+#include "Misc/DefaultValueHelper.h"
 #if WITH_EDITOR
 #include "GLTFMaterialAnalyzer.h"
 #include "Engine/TextureRenderTarget2D.h"
