[GLTFExporter] fix: improve material override resolution to handle case when mesh component relies on on underlying mesh materials

commit a49f1c77c57b7ed383b252c49c0ad1e5a49f1ed2
Date: Fri, 11 Feb 2022 23:41:53 +0100

    fix: improve material override resolution to handle case when mesh component relies on on underlying mesh materials
    (cherry picked from commit 49765df5c8c44fcfb97e04c0c8890b91f115f3bb​)

#skipci

diff --git a/Engine/Plugins/Enterprise/GLTFExporter/MigratedFromGit/Source/GLTFExporter/Private/Converters/GLTFMaterialUtility.cpp b/Engine/Plugins/Enterprise/GLTFExporter/MigratedFromGit/Source/GLTFExporter/Private/Converters/GLTFMaterialUtility.cpp
--- a/Engine/Plugins/Enterprise/GLTFExporter/MigratedFromGit/Source/GLTFExporter/Private/Converters/GLTFMaterialUtility.cpp
+++ b/Engine/Plugins/Enterprise/GLTFExporter/MigratedFromGit/Source/GLTFExporter/Private/Converters/GLTFMaterialUtility.cpp
@@ -472,53 +472,3 @@ bool FGLTFMaterialUtility::NeedsMeshData(const TArray<const UMaterialInterface*>
 
 	return false;
 }
-
-const UMaterialInterface* FGLTFMaterialUtility::GetInterface(const UMaterialInterface* Material)
-{
-	return Material;
-}
-
-const UMaterialInterface* FGLTFMaterialUtility::GetInterface(const FStaticMaterial& Material)
-{
-	return Material.MaterialInterface;
-}
-
-const UMaterialInterface* FGLTFMaterialUtility::GetInterface(const FSkeletalMaterial& Material)
-{
-	return Material.MaterialInterface;
-}
-
-void FGLTFMaterialUtility::ResolveOverrides(TArray<const UMaterialInterface*>& Overrides, const TArray<UMaterialInterface*>& Defaults)
-{
-	ResolveOverrides<UMaterialInterface*>(Overrides, Defaults);
-}
-
-void FGLTFMaterialUtility::ResolveOverrides(TArray<const UMaterialInterface*>& Overrides, const TArray<FStaticMaterial>& Defaults)
-{
-	ResolveOverrides<FStaticMaterial>(Overrides, Defaults);
-}
-
-void FGLTFMaterialUtility::ResolveOverrides(TArray<const UMaterialInterface*>& Overrides, const TArray<FSkeletalMaterial>& Defaults)
-{
-	ResolveOverrides<FSkeletalMaterial>(Overrides, Defaults);
-}
-
-template <typename MaterialType>
-void FGLTFMaterialUtility::ResolveOverrides(TArray<const UMaterialInterface*>& Overrides, const TArray<MaterialType>& Defaults)
-{
-	const int32 Count = Defaults.Num();
-	Overrides.SetNumZeroed(Count);
-
-	for (int32 Index = 0; Index < Count; ++Index)
-	{
-		const UMaterialInterface*& Element = Overrides[Index];
-		if (Element == nullptr)
-		{
-			Element = GetInterface(Defaults[Index]);
-			if (Element == nullptr)
-			{
-				Element = UMaterial::GetDefaultMaterial(MD_Surface);
-			}
-		}
-	}
-}
diff --git a/Engine/Plugins/Enterprise/GLTFExporter/MigratedFromGit/Source/GLTFExporter/Private/Converters/GLTFMaterialUtility.h b/Engine/Plugins/Enterprise/GLTFExporter/MigratedFromGit/Source/GLTFExporter/Private/Converters/GLTFMaterialUtility.h
--- a/Engine/Plugins/Enterprise/GLTFExporter/MigratedFromGit/Source/GLTFExporter/Private/Converters/GLTFMaterialUtility.h
+++ b/Engine/Plugins/Enterprise/GLTFExporter/MigratedFromGit/Source/GLTFExporter/Private/Converters/GLTFMaterialUtility.h
@@ -88,17 +88,4 @@ struct FGLTFMaterialUtility
 
 	static bool NeedsMeshData(const UMaterialInterface* Material);
 	static bool NeedsMeshData(const TArray<const UMaterialInterface*>& Materials);
-
-	static const UMaterialInterface* GetInterface(const UMaterialInterface* Material);
-	static const UMaterialInterface* GetInterface(const FStaticMaterial& Material);
-	static const UMaterialInterface* GetInterface(const FSkeletalMaterial& Material);
-
-	static void ResolveOverrides(TArray<const UMaterialInterface*>& Overrides, const TArray<UMaterialInterface*>& Defaults);
-	static void ResolveOverrides(TArray<const UMaterialInterface*>& Overrides, const TArray<FStaticMaterial>& Defaults);
-	static void ResolveOverrides(TArray<const UMaterialInterface*>& Overrides, const TArray<FSkeletalMaterial>& Defaults);
-
-private:
-
-	template <typename MaterialType>
-	static void ResolveOverrides(TArray<const UMaterialInterface*>& Overrides, const TArray<MaterialType>& Defaults);
 };
diff --git a/Engine/Plugins/Enterprise/GLTFExporter/MigratedFromGit/Source/GLTFExporter/Private/Converters/GLTFMeshConverters.cpp b/Engine/Plugins/Enterprise/GLTFExporter/MigratedFromGit/Source/GLTFExporter/Private/Converters/GLTFMeshConverters.cpp
--- a/Engine/Plugins/Enterprise/GLTFExporter/MigratedFromGit/Source/GLTFExporter/Private/Converters/GLTFMeshConverters.cpp
+++ b/Engine/Plugins/Enterprise/GLTFExporter/MigratedFromGit/Source/GLTFExporter/Private/Converters/GLTFMeshConverters.cpp
@@ -12,14 +12,7 @@
 
 void FGLTFStaticMeshConverter::Sanitize(const UStaticMesh*& StaticMesh, const UStaticMeshComponent*& StaticMeshComponent, FGLTFMaterialArray& Materials, int32& LODIndex)
 {
-	if (StaticMeshComponent != nullptr)
-	{
-		FGLTFMaterialUtility::ResolveOverrides(Materials, StaticMeshComponent->GetMaterials());
-	}
-	else
-	{
-		FGLTFMaterialUtility::ResolveOverrides(Materials, FGLTFMeshUtility::GetMaterials(StaticMesh));
-	}
+	FGLTFMeshUtility::ResolveMaterials(Materials, StaticMeshComponent, StaticMesh);
 
 	LODIndex = Builder.SanitizeLOD(StaticMesh, StaticMeshComponent, LODIndex);
 
@@ -61,14 +54,7 @@ FGLTFJsonMeshIndex FGLTFStaticMeshConverter::Convert(const UStaticMesh* StaticMe
 
 void FGLTFSkeletalMeshConverter::Sanitize(const USkeletalMesh*& SkeletalMesh, const USkeletalMeshComponent*& SkeletalMeshComponent, FGLTFMaterialArray& Materials, int32& LODIndex)
 {
-	if (SkeletalMeshComponent != nullptr)
-	{
-		FGLTFMaterialUtility::ResolveOverrides(Materials, SkeletalMeshComponent->GetMaterials());
-	}
-	else
-	{
-		FGLTFMaterialUtility::ResolveOverrides(Materials, FGLTFMeshUtility::GetMaterials(SkeletalMesh));
-	}
+	FGLTFMeshUtility::ResolveMaterials(Materials, SkeletalMeshComponent, SkeletalMesh);
 
 	LODIndex = Builder.SanitizeLOD(SkeletalMesh, SkeletalMeshComponent, LODIndex);
 
diff --git a/Engine/Plugins/Enterprise/GLTFExporter/MigratedFromGit/Source/GLTFExporter/Private/Converters/GLTFMeshUtility.cpp b/Engine/Plugins/Enterprise/GLTFExporter/MigratedFromGit/Source/GLTFExporter/Private/Converters/GLTFMeshUtility.cpp
--- a/Engine/Plugins/Enterprise/GLTFExporter/MigratedFromGit/Source/GLTFExporter/Private/Converters/GLTFMeshUtility.cpp
+++ b/Engine/Plugins/Enterprise/GLTFExporter/MigratedFromGit/Source/GLTFExporter/Private/Converters/GLTFMeshUtility.cpp
@@ -3,6 +3,11 @@
 #include "Converters/GLTFMeshUtility.h"
 #include "Rendering/SkeletalMeshRenderData.h"
 
+const UMaterialInterface* FGLTFMeshUtility::GetDefaultMaterial()
+{
+	return UMaterial::GetDefaultMaterial(MD_Surface);
+}
+
 const TArray<FStaticMaterial>& FGLTFMeshUtility::GetMaterials(const UStaticMesh* StaticMesh)
 {
 #if (ENGINE_MAJOR_VERSION > 4 || ENGINE_MINOR_VERSION >= 27)
@@ -21,6 +26,78 @@ const TArray<FSkeletalMaterial>& FGLTFMeshUtility::GetMaterials(const USkeletalM
 #endif
 }
 
+const UMaterialInterface* FGLTFMeshUtility::GetMaterial(const UMaterialInterface* Material)
+{
+	return Material;
+}
+
+const UMaterialInterface* FGLTFMeshUtility::GetMaterial(const FStaticMaterial& Material)
+{
+	return Material.MaterialInterface;
+}
+
+const UMaterialInterface* FGLTFMeshUtility::GetMaterial(const FSkeletalMaterial& Material)
+{
+	return Material.MaterialInterface;
+}
+
+void FGLTFMeshUtility::ResolveMaterials(TArray<const UMaterialInterface*>& Materials, const UStaticMeshComponent* StaticMeshComponent, const UStaticMesh* StaticMesh)
+{
+	if (StaticMeshComponent != nullptr)
+	{
+		ResolveMaterials(Materials, StaticMeshComponent->GetMaterials());
+	}
+
+	if (StaticMesh != nullptr)
+	{
+		ResolveMaterials(Materials, GetMaterials(StaticMesh));
+	}
+
+	ResolveMaterials(Materials, GetDefaultMaterial());
+}
+
+void FGLTFMeshUtility::ResolveMaterials(TArray<const UMaterialInterface*>& Materials, const USkeletalMeshComponent* SkeletalMeshComponent, const USkeletalMesh* SkeletalMesh)
+{
+	if (SkeletalMeshComponent != nullptr)
+	{
+		ResolveMaterials(Materials, SkeletalMeshComponent->GetMaterials());
+	}
+
+	if (SkeletalMesh != nullptr)
+	{
+		ResolveMaterials(Materials, GetMaterials(SkeletalMesh));
+	}
+
+	ResolveMaterials(Materials, GetDefaultMaterial());
+}
+
+template <typename MaterialType>
+void FGLTFMeshUtility::ResolveMaterials(TArray<const UMaterialInterface*>& Materials, const TArray<MaterialType>& Defaults)
+{
+	const int32 Count = Defaults.Num();
+	Materials.SetNumZeroed(Count);
+
+	for (int32 Index = 0; Index < Count; ++Index)
+	{
+		const UMaterialInterface*& Material = Materials[Index];
+		if (Material == nullptr)
+		{
+			Material = GetMaterial(Defaults[Index]);
+		}
+	}
+}
+
+void FGLTFMeshUtility::ResolveMaterials(TArray<const UMaterialInterface*>& Materials, const UMaterialInterface* Default)
+{
+	for (const UMaterialInterface*& Material : Materials)
+	{
+		if (Material == nullptr)
+		{
+			Material = Default;
+		}
+	}
+}
+
 FGLTFIndexArray FGLTFMeshUtility::GetSectionIndices(const UStaticMesh* StaticMesh, int32 LODIndex, int32 MaterialIndex)
 {
 	if (StaticMesh == nullptr)
diff --git a/Engine/Plugins/Enterprise/GLTFExporter/MigratedFromGit/Source/GLTFExporter/Private/Converters/GLTFMeshUtility.h b/Engine/Plugins/Enterprise/GLTFExporter/MigratedFromGit/Source/GLTFExporter/Private/Converters/GLTFMeshUtility.h
--- a/Engine/Plugins/Enterprise/GLTFExporter/MigratedFromGit/Source/GLTFExporter/Private/Converters/GLTFMeshUtility.h
+++ b/Engine/Plugins/Enterprise/GLTFExporter/MigratedFromGit/Source/GLTFExporter/Private/Converters/GLTFMeshUtility.h
@@ -8,9 +8,22 @@ class FSkeletalMeshLODRenderData;
 
 struct FGLTFMeshUtility
 {
+	static const UMaterialInterface* GetDefaultMaterial();
+
 	static const TArray<FStaticMaterial>& GetMaterials(const UStaticMesh* StaticMesh);
 	static const TArray<FSkeletalMaterial>& GetMaterials(const USkeletalMesh* SkeletalMesh);
 
+	static const UMaterialInterface* GetMaterial(const UMaterialInterface* Material);
+	static const UMaterialInterface* GetMaterial(const FStaticMaterial& Material);
+	static const UMaterialInterface* GetMaterial(const FSkeletalMaterial& Material);
+
+	static void ResolveMaterials(TArray<const UMaterialInterface*>& Materials, const UStaticMeshComponent* StaticMeshComponent, const UStaticMesh* StaticMesh);
+	static void ResolveMaterials(TArray<const UMaterialInterface*>& Materials, const USkeletalMeshComponent* SkeletalMeshComponent, const USkeletalMesh* SkeletalMesh);
+
+	template <typename MaterialType>
+	static void ResolveMaterials(TArray<const UMaterialInterface*>& Materials, const TArray<MaterialType>& Defaults);
+	static void ResolveMaterials(TArray<const UMaterialInterface*>& Materials, const UMaterialInterface* Default);
+
 	static FGLTFIndexArray GetSectionIndices(const UStaticMesh* StaticMesh, int32 LODIndex, int32 MaterialIndex);
 	static FGLTFIndexArray GetSectionIndices(const USkeletalMesh* SkeletalMesh, int32 LODIndex, int32 MaterialIndex);
 
