[GLTFExporter] fix: resolve material proxies in mesh converters to avoid baking with mesh data when proxy doesn't need it

commit 04628e1734bcba3298653009042b45455d76a6a6
Date: Fri, 11 Feb 2022 23:43:30 +0100

    fix: resolve material proxies in mesh converters to avoid baking with mesh data when proxy doesn't need it
    (cherry picked from commit 26b67b53241181fcbb0bb501f46eb1bd438461cc​)

#skipci

diff --git a/Engine/Plugins/Enterprise/GLTFExporter/MigratedFromGit/Source/GLTFExporter/Private/Builders/GLTFBuilder.cpp b/Engine/Plugins/Enterprise/GLTFExporter/MigratedFromGit/Source/GLTFExporter/Private/Builders/GLTFBuilder.cpp
--- a/Engine/Plugins/Enterprise/GLTFExporter/MigratedFromGit/Source/GLTFExporter/Private/Builders/GLTFBuilder.cpp
+++ b/Engine/Plugins/Enterprise/GLTFExporter/MigratedFromGit/Source/GLTFExporter/Private/Builders/GLTFBuilder.cpp
@@ -18,6 +18,17 @@ const UMaterialInterface* FGLTFBuilder::ResolveProxy(const UMaterialInterface* M
 	return ExportOptions->bExportProxyMaterials ? UGLTFMaterialExportOptions::ResolveProxy(Material) : Material;
 }
 
+void FGLTFBuilder::ResolveProxies(TArray<const UMaterialInterface*>& Materials) const
+{
+	if (ExportOptions->bExportProxyMaterials)
+	{
+		for (const UMaterialInterface*& Material : Materials)
+		{
+			Material = UGLTFMaterialExportOptions::ResolveProxy(Material);
+		}
+	}
+}
+
 FIntPoint FGLTFBuilder::GetBakeSizeForMaterialProperty(const UMaterialInterface* Material, EGLTFMaterialPropertyGroup PropertyGroup) const
 {
 	EGLTFMaterialBakeSizePOT DefaultValue = ExportOptions->DefaultMaterialBakeSize;
diff --git a/Engine/Plugins/Enterprise/GLTFExporter/MigratedFromGit/Source/GLTFExporter/Private/Builders/GLTFBuilder.h b/Engine/Plugins/Enterprise/GLTFExporter/MigratedFromGit/Source/GLTFExporter/Private/Builders/GLTFBuilder.h
--- a/Engine/Plugins/Enterprise/GLTFExporter/MigratedFromGit/Source/GLTFExporter/Private/Builders/GLTFBuilder.h
+++ b/Engine/Plugins/Enterprise/GLTFExporter/MigratedFromGit/Source/GLTFExporter/Private/Builders/GLTFBuilder.h
@@ -20,6 +20,7 @@ public:
 	FGLTFBuilder(const FString& FilePath, const UGLTFExportOptions* ExportOptions);
 
 	const UMaterialInterface* ResolveProxy(const UMaterialInterface* Material) const;
+	void ResolveProxies(TArray<const UMaterialInterface*>& Materials) const;
 
 	FIntPoint GetBakeSizeForMaterialProperty(const UMaterialInterface* Material, EGLTFMaterialPropertyGroup PropertyGroup) const;
 	TextureFilter GetBakeFilterForMaterialProperty(const UMaterialInterface* Material, EGLTFMaterialPropertyGroup PropertyGroup) const;
diff --git a/Engine/Plugins/Enterprise/GLTFExporter/MigratedFromGit/Source/GLTFExporter/Private/Builders/GLTFConvertBuilder.cpp b/Engine/Plugins/Enterprise/GLTFExporter/MigratedFromGit/Source/GLTFExporter/Private/Builders/GLTFConvertBuilder.cpp
--- a/Engine/Plugins/Enterprise/GLTFExporter/MigratedFromGit/Source/GLTFExporter/Private/Builders/GLTFConvertBuilder.cpp
+++ b/Engine/Plugins/Enterprise/GLTFExporter/MigratedFromGit/Source/GLTFExporter/Private/Builders/GLTFConvertBuilder.cpp
@@ -183,6 +183,7 @@ const FGLTFMeshData* FGLTFConvertBuilder::GetOrAddMeshData(const USkeletalMesh*
 
 FGLTFJsonMaterialIndex FGLTFConvertBuilder::GetOrAddMaterial(const UMaterialInterface* Material, const UStaticMesh* StaticMesh, int32 LODIndex, int32 MaterialIndex)
 {
+	// TODO: optimize by skipping mesh data if material doesn't need it
 	const FGLTFMeshData* MeshData = GetOrAddMeshData(StaticMesh, nullptr, LODIndex);
 	const FGLTFIndexArray SectionIndices = FGLTFMeshUtility::GetSectionIndices(StaticMesh, MeshData->LODIndex, MaterialIndex);
 	return GetOrAddMaterial(Material, MeshData, SectionIndices);
@@ -190,6 +191,7 @@ FGLTFJsonMaterialIndex FGLTFConvertBuilder::GetOrAddMaterial(const UMaterialInte
 
 FGLTFJsonMaterialIndex FGLTFConvertBuilder::GetOrAddMaterial(const UMaterialInterface* Material, const USkeletalMesh* SkeletalMesh, int32 LODIndex, int32 MaterialIndex)
 {
+	// TODO: optimize by skipping mesh data if material doesn't need it
 	const FGLTFMeshData* MeshData = GetOrAddMeshData(SkeletalMesh, nullptr, LODIndex);
 	const FGLTFIndexArray SectionIndices = FGLTFMeshUtility::GetSectionIndices(SkeletalMesh, MeshData->LODIndex, MaterialIndex);
 	return GetOrAddMaterial(Material, MeshData, SectionIndices);
@@ -212,6 +214,7 @@ FGLTFJsonMaterialIndex FGLTFConvertBuilder::GetOrAddMaterial(const UMaterialInte
 
 FGLTFJsonMaterialIndex FGLTFConvertBuilder::GetOrAddMaterial(const UMaterialInterface* Material, const UStaticMeshComponent* StaticMeshComponent, int32 LODIndex, int32 MaterialIndex)
 {
+	// TODO: optimize by skipping mesh data if material doesn't need it
 	const UStaticMesh* StaticMesh = StaticMeshComponent->GetStaticMesh();
 	const FGLTFMeshData* MeshData = GetOrAddMeshData(StaticMesh, StaticMeshComponent, LODIndex);
 	const FGLTFIndexArray SectionIndices = FGLTFMeshUtility::GetSectionIndices(StaticMesh, MeshData->LODIndex, MaterialIndex);
@@ -220,6 +223,7 @@ FGLTFJsonMaterialIndex FGLTFConvertBuilder::GetOrAddMaterial(const UMaterialInte
 
 FGLTFJsonMaterialIndex FGLTFConvertBuilder::GetOrAddMaterial(const UMaterialInterface* Material, const USkeletalMeshComponent* SkeletalMeshComponent, int32 LODIndex, int32 MaterialIndex)
 {
+	// TODO: optimize by skipping mesh data if material doesn't need it
 	const USkeletalMesh* SkeletalMesh = SkeletalMeshComponent->SkeletalMesh;
 	const FGLTFMeshData* MeshData = GetOrAddMeshData(SkeletalMesh, SkeletalMeshComponent, LODIndex);
 	const FGLTFIndexArray SectionIndices = FGLTFMeshUtility::GetSectionIndices(SkeletalMesh, MeshData->LODIndex, MaterialIndex);
diff --git a/Engine/Plugins/Enterprise/GLTFExporter/MigratedFromGit/Source/GLTFExporter/Private/Converters/GLTFMaterialConverters.cpp b/Engine/Plugins/Enterprise/GLTFExporter/MigratedFromGit/Source/GLTFExporter/Private/Converters/GLTFMaterialConverters.cpp
--- a/Engine/Plugins/Enterprise/GLTFExporter/MigratedFromGit/Source/GLTFExporter/Private/Converters/GLTFMaterialConverters.cpp
+++ b/Engine/Plugins/Enterprise/GLTFExporter/MigratedFromGit/Source/GLTFExporter/Private/Converters/GLTFMaterialConverters.cpp
@@ -7,6 +7,8 @@
 
 void FGLTFMaterialConverter::Sanitize(const UMaterialInterface*& Material, const FGLTFMeshData*& MeshData, FGLTFIndexArray& SectionIndices)
 {
+	Material = Builder.ResolveProxy(Material);
+
 	if (MeshData == nullptr ||
 		Builder.ExportOptions->BakeMaterialInputs != EGLTFMaterialBakeMode::UseMeshData ||
 		!FGLTFMaterialUtility::NeedsMeshData(Material))
diff --git a/Engine/Plugins/Enterprise/GLTFExporter/MigratedFromGit/Source/GLTFExporter/Private/Converters/GLTFMaterialUtility.cpp b/Engine/Plugins/Enterprise/GLTFExporter/MigratedFromGit/Source/GLTFExporter/Private/Converters/GLTFMaterialUtility.cpp
--- a/Engine/Plugins/Enterprise/GLTFExporter/MigratedFromGit/Source/GLTFExporter/Private/Converters/GLTFMaterialUtility.cpp
+++ b/Engine/Plugins/Enterprise/GLTFExporter/MigratedFromGit/Source/GLTFExporter/Private/Converters/GLTFMaterialUtility.cpp
@@ -421,7 +421,7 @@ FString FGLTFMaterialUtility::ShadingModelsToString(const FMaterialShadingModelF
 bool FGLTFMaterialUtility::NeedsMeshData(const UMaterialInterface* Material)
 {
 #if WITH_EDITOR
-	if (Material != nullptr)
+	if (Material != nullptr && !IsProxyMaterial(Material))
 	{
 		// TODO: only analyze properties that will be needed for this specific material
 		const TArray<FMaterialPropertyEx> Properties =
@@ -442,7 +442,7 @@ bool FGLTFMaterialUtility::NeedsMeshData(const UMaterialInterface* Material)
 		bool bNeedsMeshData = false;
 		FGLTFMaterialAnalysis Analysis;
 
-		// TODO: optimize baking by seperating need for vertex data and primitive data
+		// TODO: optimize baking by separating need for vertex data and primitive data
 
 		for (const FMaterialPropertyEx& Property: Properties)
 		{
diff --git a/Engine/Plugins/Enterprise/GLTFExporter/MigratedFromGit/Source/GLTFExporter/Private/Converters/GLTFMeshConverters.cpp b/Engine/Plugins/Enterprise/GLTFExporter/MigratedFromGit/Source/GLTFExporter/Private/Converters/GLTFMeshConverters.cpp
--- a/Engine/Plugins/Enterprise/GLTFExporter/MigratedFromGit/Source/GLTFExporter/Private/Converters/GLTFMeshConverters.cpp
+++ b/Engine/Plugins/Enterprise/GLTFExporter/MigratedFromGit/Source/GLTFExporter/Private/Converters/GLTFMeshConverters.cpp
@@ -13,6 +13,7 @@
 void FGLTFStaticMeshConverter::Sanitize(const UStaticMesh*& StaticMesh, const UStaticMeshComponent*& StaticMeshComponent, FGLTFMaterialArray& Materials, int32& LODIndex)
 {
 	FGLTFMeshUtility::ResolveMaterials(Materials, StaticMeshComponent, StaticMesh);
+	Builder.ResolveProxies(Materials);
 
 	LODIndex = Builder.SanitizeLOD(StaticMesh, StaticMeshComponent, LODIndex);
 
@@ -55,6 +56,7 @@ FGLTFJsonMeshIndex FGLTFStaticMeshConverter::Convert(const UStaticMesh* StaticMe
 void FGLTFSkeletalMeshConverter::Sanitize(const USkeletalMesh*& SkeletalMesh, const USkeletalMeshComponent*& SkeletalMeshComponent, FGLTFMaterialArray& Materials, int32& LODIndex)
 {
 	FGLTFMeshUtility::ResolveMaterials(Materials, SkeletalMeshComponent, SkeletalMesh);
+	Builder.ResolveProxies(Materials);
 
 	LODIndex = Builder.SanitizeLOD(SkeletalMesh, SkeletalMeshComponent, LODIndex);
 
diff --git a/Engine/Plugins/Enterprise/GLTFExporter/MigratedFromGit/Source/GLTFExporter/Private/Tasks/GLTFMaterialTasks.cpp b/Engine/Plugins/Enterprise/GLTFExporter/MigratedFromGit/Source/GLTFExporter/Private/Tasks/GLTFMaterialTasks.cpp
--- a/Engine/Plugins/Enterprise/GLTFExporter/MigratedFromGit/Source/GLTFExporter/Private/Tasks/GLTFMaterialTasks.cpp
+++ b/Engine/Plugins/Enterprise/GLTFExporter/MigratedFromGit/Source/GLTFExporter/Private/Tasks/GLTFMaterialTasks.cpp
@@ -49,8 +49,6 @@ namespace
 
 void FGLTFMaterialTask::Complete()
 {
-	Material = Builder.ResolveProxy(Material);
-
 	const UMaterial* BaseMaterial = Material->GetMaterial();
 	if (BaseMaterial->MaterialDomain != MD_Surface)
 	{
diff --git a/Engine/Plugins/Enterprise/GLTFExporter/MigratedFromGit/Source/GLTFExporter/Private/UserData/GLTFMaterialUserData.cpp b/Engine/Plugins/Enterprise/GLTFExporter/MigratedFromGit/Source/GLTFExporter/Private/UserData/GLTFMaterialUserData.cpp
--- a/Engine/Plugins/Enterprise/GLTFExporter/MigratedFromGit/Source/GLTFExporter/Private/UserData/GLTFMaterialUserData.cpp
+++ b/Engine/Plugins/Enterprise/GLTFExporter/MigratedFromGit/Source/GLTFExporter/Private/UserData/GLTFMaterialUserData.cpp
@@ -16,6 +16,11 @@ FGLTFOverrideMaterialBakeSettings::FGLTFOverrideMaterialBakeSettings()
 
 const UMaterialInterface* UGLTFMaterialExportOptions::ResolveProxy(const UMaterialInterface* Material)
 {
+	if (Material == nullptr)
+	{
+		return nullptr;
+	}
+
 	TSet<const UMaterialInterface*> PreviousProxies;
 
 	for (;;)
@@ -76,7 +81,7 @@ TextureAddress UGLTFMaterialExportOptions::GetBakeTilingForPropertyGroup(const U
 template <typename Predicate>
 const FGLTFOverrideMaterialBakeSettings* UGLTFMaterialExportOptions::GetBakeSettingsByPredicate(const UMaterialInterface* Material, EGLTFMaterialPropertyGroup PropertyGroup, Predicate Pred)
 {
-	do
+	while (Material != nullptr)
 	{
 		if (const UGLTFMaterialExportOptions* UserData = const_cast<UMaterialInterface*>(Material)->GetAssetUserData<UGLTFMaterialExportOptions>())
 		{
@@ -96,7 +101,7 @@ const FGLTFOverrideMaterialBakeSettings* UGLTFMaterialExportOptions::GetBakeSett
 
 		const UMaterialInstance* MaterialInstance = Cast<UMaterialInstance>(Material);
 		Material = MaterialInstance != nullptr ? MaterialInstance->Parent : nullptr;
-	} while (Material != nullptr);
+	}
 
 	return nullptr;
 }
