[GLTFExporter] refactor: ensure proxy material parameters are only overridden and exported if non-default

commit 8abdf8f9008072d6f6ce620127224b90d6909cd8
Date: Fri, 11 Feb 2022 23:43:46 +0100

    refactor: ensure proxy material parameters are only overridden and exported if non-default
    (cherry picked from commit 5d2594952c017ace23b3b9e6cd594f3669821c8b​)

#skipci

diff --git a/Engine/Plugins/Enterprise/GLTFExporter/MigratedFromGit/Source/GLTFExporter/Private/Converters/GLTFMaterialProxyFactory.cpp b/Engine/Plugins/Enterprise/GLTFExporter/MigratedFromGit/Source/GLTFExporter/Private/Converters/GLTFMaterialProxyFactory.cpp
--- a/Engine/Plugins/Enterprise/GLTFExporter/MigratedFromGit/Source/GLTFExporter/Private/Converters/GLTFMaterialProxyFactory.cpp
+++ b/Engine/Plugins/Enterprise/GLTFExporter/MigratedFromGit/Source/GLTFExporter/Private/Converters/GLTFMaterialProxyFactory.cpp
@@ -123,98 +123,64 @@ void FGLTFMaterialProxyFactory::SetProxyProperties(UMaterialInstanceConstant* Pr
 			SetProxyProperty(ProxyMaterial, TEXT("Clear Coat Normal"), JsonMaterial.ClearCoat.ClearCoatNormalTexture, EGLTFMaterialPropertyGroup::ClearCoatBottomNormal);
 		}
 	}
-
-	ProxyMaterial->PostEditChange();
 }
 
 void FGLTFMaterialProxyFactory::SetProxyProperty(UMaterialInstanceConstant* ProxyMaterial, const FString& PropertyName, float Scalar)
 {
-	float DefaultValue;
-	if (!ProxyMaterial->GetScalarParameterDefaultValue(*PropertyName, DefaultValue))
-	{
-		// TODO: report error
-		return;
-	}
-
-	if (DefaultValue != Scalar)
-	{
-		ProxyMaterial->SetScalarParameterValueEditorOnly(*PropertyName, Scalar);
-	}
+	FGLTFMaterialUtility::SetNonDefaultParameterValue(ProxyMaterial, PropertyName, Scalar);
 }
 
 void FGLTFMaterialProxyFactory::SetProxyProperty(UMaterialInstanceConstant* ProxyMaterial, const FString& PropertyName, const FGLTFJsonColor3& Color)
 {
-	FLinearColor DefaultValue;
-	if (!ProxyMaterial->GetVectorParameterDefaultValue(*PropertyName, DefaultValue))
-	{
-		// TODO: report error
-		return;
-	}
-
-	const FLinearColor Value(Color.R, Color.G, Color.B);
-	if (DefaultValue != Value)
-	{
-		ProxyMaterial->SetVectorParameterValueEditorOnly(*PropertyName, Value);
-	}
+	FGLTFMaterialUtility::SetNonDefaultParameterValue(ProxyMaterial, PropertyName, FLinearColor(Color.R, Color.G, Color.B));
 }
 
 void FGLTFMaterialProxyFactory::SetProxyProperty(UMaterialInstanceConstant* ProxyMaterial, const FString& PropertyName, const FGLTFJsonColor4& Color)
 {
-	FLinearColor DefaultValue;
-	if (!ProxyMaterial->GetVectorParameterDefaultValue(*PropertyName, DefaultValue))
-	{
-		// TODO: report error
-		return;
-	}
-
-	const FLinearColor Value(Color.R, Color.G, Color.B, Color.A);
-	if (DefaultValue != Value)
-	{
-		ProxyMaterial->SetVectorParameterValueEditorOnly(*PropertyName, Value);
-	}
+	FGLTFMaterialUtility::SetNonDefaultParameterValue(ProxyMaterial, PropertyName, FLinearColor(Color.R, Color.G, Color.B, Color.A));
 }
 
 void FGLTFMaterialProxyFactory::SetProxyProperty(UMaterialInstanceConstant* ProxyMaterial, const FString& PropertyName, const FGLTFJsonTextureInfo& TextureInfo, EGLTFMaterialPropertyGroup PropertyGroup)
 {
-	if (TextureInfo.Index != INDEX_NONE)
+	UTexture* Texture = FindOrCreateTexture(TextureInfo.Index, PropertyGroup);
+	if (Texture == nullptr)
 	{
-		const UTexture* Texture = FindOrCreateTexture(TextureInfo.Index, PropertyGroup);
-		if (Texture == nullptr)
-		{
-			// TODO: report error
-		}
-
-		ProxyMaterial->SetTextureParameterValueEditorOnly(*(PropertyName + TEXT(" Texture")), const_cast<UTexture*>(Texture));
+		return;
 	}
 
+	FGLTFMaterialUtility::SetNonDefaultParameterValue(ProxyMaterial, PropertyName + TEXT(" Texture"), Texture);
+
 	if (TextureInfo.TexCoord != 0)
 	{
-		ProxyMaterial->SetScalarParameterValueEditorOnly(*(PropertyName + TEXT(" UV Index")), TextureInfo.TexCoord);
+		FGLTFMaterialUtility::SetNonDefaultParameterValue(ProxyMaterial, PropertyName + TEXT(" UV Index"), TextureInfo.TexCoord);
 	}
 
 	if (TextureInfo.Transform.Offset != FGLTFJsonVector2::Zero)
 	{
-		const FLinearColor Offset(TextureInfo.Transform.Offset.X, TextureInfo.Transform.Offset.Y, 0, 0);
-		ProxyMaterial->SetVectorParameterValueEditorOnly(*(PropertyName + TEXT(" UV Offset")), Offset);
+		FGLTFMaterialUtility::SetNonDefaultParameterValue(ProxyMaterial, PropertyName + TEXT(" UV Offset"), FLinearColor(TextureInfo.Transform.Offset.X, TextureInfo.Transform.Offset.Y, 0, 0));
 	}
 
 	if (TextureInfo.Transform.Scale != FGLTFJsonVector2::One)
 	{
-		const FLinearColor Scale(TextureInfo.Transform.Scale.X, TextureInfo.Transform.Scale.Y, 0, 0);
-		ProxyMaterial->SetVectorParameterValueEditorOnly(*(PropertyName + TEXT(" UV Scale")), Scale);
+		FGLTFMaterialUtility::SetNonDefaultParameterValue(ProxyMaterial, PropertyName + TEXT(" UV Scale"), FLinearColor(TextureInfo.Transform.Scale.X, TextureInfo.Transform.Scale.Y, 0, 0));
 	}
 
 	if (TextureInfo.Transform.Rotation != 0)
 	{
-		ProxyMaterial->SetScalarParameterValueEditorOnly(*(PropertyName + TEXT(" UV Rotation")), TextureInfo.Transform.Rotation);
+		FGLTFMaterialUtility::SetNonDefaultParameterValue(ProxyMaterial, PropertyName + TEXT(" UV Rotation"), TextureInfo.Transform.Rotation);
 	}
 }
 
-const UTexture2D* FGLTFMaterialProxyFactory::FindOrCreateTexture(FGLTFJsonTextureIndex Index, EGLTFMaterialPropertyGroup PropertyGroup)
+UTexture2D* FGLTFMaterialProxyFactory::FindOrCreateTexture(FGLTFJsonTextureIndex Index, EGLTFMaterialPropertyGroup PropertyGroup)
 {
+	if (Index == INDEX_NONE)
+	{
+		return nullptr;
+	}
+
 	// TODO: fix potential conflict when same texture used for different material properties that have different encoding (sRGB vs Linear, Normalmap compression etc)
 
-	const UTexture2D** FoundPtr = Textures.Find(Index);
+	UTexture2D** FoundPtr = Textures.Find(Index);
 	if (FoundPtr != nullptr)
 	{
 		return *FoundPtr;
@@ -224,11 +190,12 @@ const UTexture2D* FGLTFMaterialProxyFactory::FindOrCreateTexture(FGLTFJsonTextur
 	const FGLTFImageData* ImageData = Images.Find(JsonTexture.Source);
 	if (ImageData == nullptr)
 	{
+		// TODO: report error
 		return nullptr;
 	}
 
 	const FGLTFJsonSampler& JsonSampler = Builder.GetSampler(JsonTexture.Sampler);
-	const UTexture2D* Texture = CreateTexture(ImageData, JsonSampler, PropertyGroup);
+	UTexture2D* Texture = CreateTexture(ImageData, JsonSampler, PropertyGroup);
 	Textures.Add(Index, Texture);
 	return Texture;
 }
diff --git a/Engine/Plugins/Enterprise/GLTFExporter/MigratedFromGit/Source/GLTFExporter/Private/Converters/GLTFMaterialProxyFactory.h b/Engine/Plugins/Enterprise/GLTFExporter/MigratedFromGit/Source/GLTFExporter/Private/Converters/GLTFMaterialProxyFactory.h
--- a/Engine/Plugins/Enterprise/GLTFExporter/MigratedFromGit/Source/GLTFExporter/Private/Converters/GLTFMaterialProxyFactory.h
+++ b/Engine/Plugins/Enterprise/GLTFExporter/MigratedFromGit/Source/GLTFExporter/Private/Converters/GLTFMaterialProxyFactory.h
@@ -40,7 +40,7 @@ private:
 	void SetProxyProperty(UMaterialInstanceConstant* ProxyMaterial, const FString& PropertyName, const FGLTFJsonColor4& Color);
 	void SetProxyProperty(UMaterialInstanceConstant* ProxyMaterial, const FString& PropertyName, const FGLTFJsonTextureInfo& TextureInfo, EGLTFMaterialPropertyGroup PropertyGroup);
 
-	const UTexture2D* FindOrCreateTexture(FGLTFJsonTextureIndex Index, EGLTFMaterialPropertyGroup PropertyGroup);
+	UTexture2D* FindOrCreateTexture(FGLTFJsonTextureIndex Index, EGLTFMaterialPropertyGroup PropertyGroup);
 	UTexture2D* CreateTexture(const FGLTFImageData* ImageData, const FGLTFJsonSampler& JsonSampler, EGLTFMaterialPropertyGroup PropertyGroup);
 
 	UMaterialInstanceConstant* CreateInstancedMaterial(UMaterialInterface* OriginalMaterial, EGLTFJsonShadingModel ShadingModel);
@@ -58,7 +58,7 @@ private:
 
 	FGLTFConvertBuilder Builder;
 
-	TMap<FGLTFJsonTextureIndex, const UTexture2D*> Textures;
+	TMap<FGLTFJsonTextureIndex, UTexture2D*> Textures;
 	TMap<FGLTFJsonImageIndex, FGLTFImageData> Images;
 };
 
diff --git a/Engine/Plugins/Enterprise/GLTFExporter/MigratedFromGit/Source/GLTFExporter/Private/Converters/GLTFMaterialUtility.cpp b/Engine/Plugins/Enterprise/GLTFExporter/MigratedFromGit/Source/GLTFExporter/Private/Converters/GLTFMaterialUtility.cpp
--- a/Engine/Plugins/Enterprise/GLTFExporter/MigratedFromGit/Source/GLTFExporter/Private/Converters/GLTFMaterialUtility.cpp
+++ b/Engine/Plugins/Enterprise/GLTFExporter/MigratedFromGit/Source/GLTFExporter/Private/Converters/GLTFMaterialUtility.cpp
@@ -6,15 +6,13 @@
 #include "Misc/DefaultValueHelper.h"
 #if WITH_EDITOR
 #include "GLTFMaterialAnalyzer.h"
-#include "Engine/TextureRenderTarget2D.h"
-#include "CanvasItem.h"
-#include "Modules/ModuleManager.h"
 #include "GLTFMaterialBaking/Public/IMaterialBakingModule.h"
 #include "GLTFMaterialBaking/Public/MaterialBakingStructures.h"
+#include "Modules/ModuleManager.h"
+#include "Materials/MaterialInstanceConstant.h"
 #include "Materials/MaterialExpressionCustomOutput.h"
 #include "Materials/MaterialExpressionClearCoatNormalCustomOutput.h"
 #include "Materials/MaterialExpressionTextureCoordinate.h"
-#include "MeshDescription.h"
 #endif
 
 #define PROXY_MATERIAL_NAME_PREFIX TEXT("M_GLTF_")
@@ -43,8 +41,122 @@ bool FGLTFMaterialUtility::IsProxyMaterial(const UMaterialInterface* Material)
 	return IsProxyMaterial(Material->GetMaterial());
 }
 
+bool FGLTFMaterialUtility::GetNonDefaultParameterValue(const UMaterialInterface* Material, const FString& PropertyName, float& OutValue)
+{
+	const FHashedMaterialParameterInfo ParameterInfo = *PropertyName;
+
+	float DefaultValue;
+	if (!Material->GetScalarParameterDefaultValue(ParameterInfo, DefaultValue))
+	{
+		// TODO: report error
+		return false;
+	}
+
+	float Value;
+	if (!Material->GetScalarParameterValue(ParameterInfo, Value, true) || Value == DefaultValue)
+	{
+		return false;
+	}
+
+	OutValue = Value;
+	return true;
+}
+
+bool FGLTFMaterialUtility::GetNonDefaultParameterValue(const UMaterialInterface* Material, const FString& PropertyName, FLinearColor& OutValue)
+{
+	const FHashedMaterialParameterInfo ParameterInfo = *PropertyName;
+
+	FLinearColor DefaultValue;
+	if (!Material->GetVectorParameterDefaultValue(ParameterInfo, DefaultValue))
+	{
+		// TODO: report error
+		return false;
+	}
+
+	FLinearColor Value;
+	if (!Material->GetVectorParameterValue(ParameterInfo, Value, true) || Value == DefaultValue)
+	{
+		return false;
+	}
+
+	OutValue = Value;
+	return true;
+}
+
+bool FGLTFMaterialUtility::GetNonDefaultParameterValue(const UMaterialInterface* Material, const FString& PropertyName, UTexture*& OutValue)
+{
+	const FHashedMaterialParameterInfo ParameterInfo = *PropertyName;
+
+	UTexture* DefaultValue;
+	if (!Material->GetTextureParameterDefaultValue(ParameterInfo, DefaultValue))
+	{
+		// TODO: report error
+		return false;
+	}
+
+	UTexture* Value;
+	if (!Material->GetTextureParameterValue(ParameterInfo, Value, true) || Value == DefaultValue)
+	{
+		return false;
+	}
+
+	OutValue = Value;
+	return true;
+}
+
 #if WITH_EDITOR
 
+void FGLTFMaterialUtility::SetNonDefaultParameterValue(UMaterialInstanceConstant* Material, const FString& PropertyName, float Value)
+{
+	const FHashedMaterialParameterInfo ParameterInfo = *PropertyName;
+
+	float DefaultValue;
+	if (!Material->GetScalarParameterDefaultValue(ParameterInfo, DefaultValue))
+	{
+		// TODO: report error
+		return;
+	}
+
+	if (DefaultValue != Value)
+	{
+		Material->SetScalarParameterValueEditorOnly(FMaterialParameterInfo(ParameterInfo), Value);
+	}
+}
+
+void FGLTFMaterialUtility::SetNonDefaultParameterValue(UMaterialInstanceConstant* Material, const FString& PropertyName, const FLinearColor& Value)
+{
+	const FHashedMaterialParameterInfo ParameterInfo = *PropertyName;
+
+	FLinearColor DefaultValue;
+	if (!Material->GetVectorParameterDefaultValue(ParameterInfo, DefaultValue))
+	{
+		// TODO: report error
+		return;
+	}
+
+	if (DefaultValue != Value)
+	{
+		Material->SetVectorParameterValueEditorOnly(FMaterialParameterInfo(ParameterInfo), Value);
+	}
+}
+
+void FGLTFMaterialUtility::SetNonDefaultParameterValue(UMaterialInstanceConstant* Material, const FString& PropertyName, UTexture* Value)
+{
+	const FHashedMaterialParameterInfo ParameterInfo = *PropertyName;
+
+	UTexture* DefaultValue;
+	if (!Material->GetTextureParameterDefaultValue(ParameterInfo, DefaultValue))
+	{
+		// TODO: report error
+		return;
+	}
+
+	if (DefaultValue != Value)
+	{
+		Material->SetTextureParameterValueEditorOnly(FMaterialParameterInfo(ParameterInfo), Value);
+	}
+}
+
 bool FGLTFMaterialUtility::IsNormalMap(const FMaterialPropertyEx& Property)
 {
 	return Property == MP_Normal || Property == TEXT("ClearCoatBottomNormal");
@@ -377,6 +489,7 @@ FMaterialShadingModelField FGLTFMaterialUtility::EvaluateShadingModelExpression(
 
 	return Analysis.ShadingModels;
 }
+
 #endif
 
 EMaterialShadingModel FGLTFMaterialUtility::GetRichestShadingModel(const FMaterialShadingModelField& ShadingModels)
diff --git a/Engine/Plugins/Enterprise/GLTFExporter/MigratedFromGit/Source/GLTFExporter/Private/Converters/GLTFMaterialUtility.h b/Engine/Plugins/Enterprise/GLTFExporter/MigratedFromGit/Source/GLTFExporter/Private/Converters/GLTFMaterialUtility.h
--- a/Engine/Plugins/Enterprise/GLTFExporter/MigratedFromGit/Source/GLTFExporter/Private/Converters/GLTFMaterialUtility.h
+++ b/Engine/Plugins/Enterprise/GLTFExporter/MigratedFromGit/Source/GLTFExporter/Private/Converters/GLTFMaterialUtility.h
@@ -8,6 +8,7 @@
 
 struct FGLTFMaterialAnalysis;
 struct FMaterialPropertyEx;
+class UMaterialInstanceConstant;
 
 struct FGLTFTextureCombineSource
 {
@@ -45,7 +46,15 @@ struct FGLTFMaterialUtility
 	static bool IsProxyMaterial(const UMaterial* Material);
 	static bool IsProxyMaterial(const UMaterialInterface* Material);
 
+	static bool GetNonDefaultParameterValue(const UMaterialInterface* Material, const FString& PropertyName, float& OutValue);
+	static bool GetNonDefaultParameterValue(const UMaterialInterface* Material, const FString& PropertyName, FLinearColor& OutValue);
+	static bool GetNonDefaultParameterValue(const UMaterialInterface* Material, const FString& PropertyName, UTexture*& OutValue);
+
 #if WITH_EDITOR
+	static void SetNonDefaultParameterValue(UMaterialInstanceConstant* Material, const FString& PropertyName, float Value);
+	static void SetNonDefaultParameterValue(UMaterialInstanceConstant* Material, const FString& PropertyName, const FLinearColor& Value);
+	static void SetNonDefaultParameterValue(UMaterialInstanceConstant* Material, const FString& PropertyName, UTexture* Value);
+
 	static bool IsNormalMap(const FMaterialPropertyEx& Property);
 	static bool IsSRGB(const FMaterialPropertyEx& Property);
 
diff --git a/Engine/Plugins/Enterprise/GLTFExporter/MigratedFromGit/Source/GLTFExporter/Private/Tasks/GLTFMaterialTasks.cpp b/Engine/Plugins/Enterprise/GLTFExporter/MigratedFromGit/Source/GLTFExporter/Private/Tasks/GLTFMaterialTasks.cpp
--- a/Engine/Plugins/Enterprise/GLTFExporter/MigratedFromGit/Source/GLTFExporter/Private/Tasks/GLTFMaterialTasks.cpp
+++ b/Engine/Plugins/Enterprise/GLTFExporter/MigratedFromGit/Source/GLTFExporter/Private/Tasks/GLTFMaterialTasks.cpp
@@ -244,13 +244,13 @@ void FGLTFMaterialTask::GetProxyProperties(FGLTFJsonMaterial& OutMaterial) const
 
 void FGLTFMaterialTask::GetProxyProperty(const FString& PropertyName, float& OutValue) const
 {
-	Material->GetScalarParameterValue(*PropertyName, OutValue, true);
+	FGLTFMaterialUtility::GetNonDefaultParameterValue(Material, PropertyName, OutValue);
 }
 
 void FGLTFMaterialTask::GetProxyProperty(const FString& PropertyName, FGLTFJsonColor3& OutValue) const
 {
 	FLinearColor Value;
-	if (Material->GetVectorParameterValue(*PropertyName, Value, true))
+	if (FGLTFMaterialUtility::GetNonDefaultParameterValue(Material, PropertyName, Value))
 	{
 		OutValue = FGLTFConverterUtility::ConvertColor3(Value, false);
 	}
@@ -259,7 +259,7 @@ void FGLTFMaterialTask::GetProxyProperty(const FString& PropertyName, FGLTFJsonC
 void FGLTFMaterialTask::GetProxyProperty(const FString& PropertyName, FGLTFJsonColor4& OutValue) const
 {
 	FLinearColor Value;
-	if (Material->GetVectorParameterValue(*PropertyName, Value, true))
+	if (FGLTFMaterialUtility::GetNonDefaultParameterValue(Material, PropertyName, Value))
 	{
 		OutValue = FGLTFConverterUtility::ConvertColor(Value, false);
 	}
@@ -268,33 +268,35 @@ void FGLTFMaterialTask::GetProxyProperty(const FString& PropertyName, FGLTFJsonC
 void FGLTFMaterialTask::GetProxyProperty(const FString& PropertyName, FGLTFJsonTextureInfo& OutValue) const
 {
 	UTexture* Texture;
-	if (Material->GetTextureParameterValue(*(PropertyName + TEXT(" Texture")), Texture, true) || Texture == nullptr)
+	if (!FGLTFMaterialUtility::GetNonDefaultParameterValue(Material, PropertyName + TEXT(" Texture"), Texture) || Texture == nullptr)
 	{
-		OutValue.Index = Builder.GetOrAddTexture(Texture);
+		return;
 	}
 
+	OutValue.Index = Builder.GetOrAddTexture(Texture);
+
 	float TexCoord;
-	if (Material->GetScalarParameterValue(*(PropertyName + TEXT(" UV Index")), TexCoord, true))
+	if (FGLTFMaterialUtility::GetNonDefaultParameterValue(Material, PropertyName + TEXT(" UV Index"), TexCoord))
 	{
 		OutValue.TexCoord = FMath::RoundToInt(TexCoord);
 	}
 
 	FLinearColor Offset;
-	if (Material->GetVectorParameterValue(*(PropertyName + TEXT(" UV Offset")), Offset, true))
+	if (FGLTFMaterialUtility::GetNonDefaultParameterValue(Material, PropertyName + TEXT(" UV Offset"), Offset))
 	{
 		OutValue.Transform.Offset.X = Offset.R;
 		OutValue.Transform.Offset.Y = Offset.G;
 	}
 
 	FLinearColor Scale;
-	if (Material->GetVectorParameterValue(*(PropertyName + TEXT(" UV Scale")), Scale, true))
+	if (FGLTFMaterialUtility::GetNonDefaultParameterValue(Material, PropertyName + TEXT(" UV Scale"), Scale))
 	{
 		OutValue.Transform.Scale.X = Scale.R;
 		OutValue.Transform.Scale.Y = Scale.G;
 	}
 
 	float Rotation;
-	if (Material->GetScalarParameterValue(*(PropertyName + TEXT(" UV Rotation")), Rotation, true))
+	if (FGLTFMaterialUtility::GetNonDefaultParameterValue(Material, PropertyName + TEXT(" UV Rotation"), Rotation))
 	{
 		OutValue.Transform.Rotation = Rotation;
 	}
