[GLTFExporter] fix: correct tooltip for export option BakeMaterialInputs and update reference documentation

commit 95c660ef5c2a60f5de763634d886d960f3d5dc94
Date: Sat, 12 Feb 2022 01:29:00 +0100

    fix: correct tooltip for export option BakeMaterialInputs and update reference documentation
    (cherry picked from commit 6dcd1194d3edc14247b49e1467b6c1a5c378076a​)

#skipci

diff --git a/Engine/Plugins/Enterprise/GLTFExporter/MigratedFromGit/Docs/export-options-reference.md b/Engine/Plugins/Enterprise/GLTFExporter/MigratedFromGit/Docs/export-options-reference.md
--- a/Engine/Plugins/Enterprise/GLTFExporter/MigratedFromGit/Docs/export-options-reference.md
+++ b/Engine/Plugins/Enterprise/GLTFExporter/MigratedFromGit/Docs/export-options-reference.md
@@ -7,11 +7,11 @@ Option                         | Description
 `Strict Compliance`            | If enabled, certain values (like HDR colors and light angles) will be truncated during export to strictly conform to the formal glTF specification.
 `Skip Near Default Values`     | If enabled, floating-point-based JSON properties that are nearly equal to their default value will not be exported and thus regarded as exactly default, reducing size of JSON data.
 `Include Generator Version`    | If enabled, version info for Unreal Engine and exporter plugin will be included as metadata in the glTF asset, which is useful when reporting issues.
-`Export Proxy Materials`       | If enabled, materials that have a proxy defined in their user data, will be exported using that proxy instead. Used mainly by pre-baked materials.
+`Export Proxy Materials`       | If enabled, materials that have a proxy defined in their user data, will be exported using that proxy instead. This setting won't affect proxy materials exported or referenced directly.
 `Export Unlit Materials`       | If enabled, materials with shading model unlit will be properly exported. Uses extension KHR_materials_unlit.
 `Export Clear Coat Materials`  | If enabled, materials with shading model clear coat will be properly exported. Uses extension KHR_materials_clearcoat, which is not supported by all glTF viewers.
 `Export Extra Blend Modes`     | If enabled, materials with blend modes additive, modulate, and alpha composite will be properly exported. Uses extension EPIC_blend_modes, which is supported by Unreal's glTF viewer.
-`Bake Material Inputs`         | Bake mode determining if and how a material input is baked out to a texture. Can be overridden by material- and input-specific bake settings, see GLTFMaterialExportOptions.
+`Bake Material Inputs`         | Bake mode determining if and how a material input is baked out to a texture. Baking is only used for non-trivial material inputs (i.e. not simple texture or constant expressions).
 `Default Material Bake Size`   | Default size of the baked out texture (containing the material input). Can be overridden by material- and input-specific bake settings, see GLTFMaterialExportOptions.
 `Default Material Bake Filter` | Default filtering mode used when sampling the baked out texture. Can be overridden by material- and input-specific bake settings, see GLTFMaterialExportOptions.
 `Default Material Bake Tiling` | Default addressing mode used when sampling the baked out texture. Can be overridden by material- and input-specific bake settings, see GLTFMaterialExportOptions.
diff --git a/Engine/Plugins/Enterprise/GLTFExporter/MigratedFromGit/Source/GLTFExporter/Public/Options/GLTFExportOptions.h b/Engine/Plugins/Enterprise/GLTFExporter/MigratedFromGit/Source/GLTFExporter/Public/Options/GLTFExportOptions.h
--- a/Engine/Plugins/Enterprise/GLTFExporter/MigratedFromGit/Source/GLTFExporter/Public/Options/GLTFExportOptions.h
+++ b/Engine/Plugins/Enterprise/GLTFExporter/MigratedFromGit/Source/GLTFExporter/Public/Options/GLTFExportOptions.h
@@ -115,7 +115,7 @@ class GLTFEXPORTER_API UGLTFExportOptions : public UObject
 	UPROPERTY(EditAnywhere, BlueprintReadWrite, Config, Category = Material)
 	bool bExportExtraBlendModes;
 
-	/** Bake mode determining if and how a material input is baked out to a texture. Can be overridden by material- and input-specific bake settings, see GLTFMaterialExportOptions. */
+	/** Bake mode determining if and how a material input is baked out to a texture. Baking is only used for non-trivial material inputs (i.e. not simple texture or constant expressions). */
 	UPROPERTY(EditAnywhere, BlueprintReadWrite, Config, Category = Material)
 	EGLTFMaterialBakeMode BakeMaterialInputs;
 
@@ -124,7 +124,7 @@ class GLTFEXPORTER_API UGLTFExportOptions : public UObject
 	EGLTFMaterialBakeSizePOT DefaultMaterialBakeSize;
 
 	/** Default filtering mode used when sampling the baked out texture. Can be overridden by material- and input-specific bake settings, see GLTFMaterialExportOptions. */
-	UPROPERTY(EditAnywhere, BlueprintReadWrite, Config, Category = Material, Meta = (EditCondition = "BakeMaterialInputs != EGLTFMaterialBakeMode::Disabled", ValidEnumValues="TF_Nearest, TF_Bilinear, TF_Trilinear"))
+	UPROPERTY(EditAnywhere, BlueprintReadWrite, Config, Category = Material, Meta = (EditCondition = "BakeMaterialInputs != EGLTFMaterialBakeMode::Disabled", ValidEnumValues = "TF_Nearest, TF_Bilinear, TF_Trilinear"))
 	TEnumAsByte<TextureFilter> DefaultMaterialBakeFilter;
 
 	/** Default addressing mode used when sampling the baked out texture. Can be overridden by material- and input-specific bake settings, see GLTFMaterialExportOptions. */
@@ -132,7 +132,7 @@ class GLTFEXPORTER_API UGLTFExportOptions : public UObject
 	TEnumAsByte<TextureAddress> DefaultMaterialBakeTiling;
 
 	/** Input-specific default bake settings that override the general defaults above. */
-	UPROPERTY(EditAnywhere, BlueprintReadWrite, Config, Category = Material)
+	UPROPERTY(EditAnywhere, BlueprintReadWrite, Config, Category = Material, Meta = (EditCondition = "BakeMaterialInputs != EGLTFMaterialBakeMode::Disabled"))
 	TMap<EGLTFMaterialPropertyGroup, FGLTFOverrideMaterialBakeSettings> DefaultInputBakeSettings;
 
 	/** Default LOD level used for exporting a mesh. Can be overridden by component or asset settings (e.g. minimum or forced LOD level). */
