[GLTFExporter] feat: add option to disable material baking when creating proxy material

commit 6819fa60a77c6ba82de484d6a019b295ac8e8482
Date: Sat, 12 Feb 2022 01:30:53 +0100

    feat: add option to disable material baking when creating proxy material
    (cherry picked from commit 6da6ab6c58db74a6d7e61b2e13e80d3bd27c26c1​)

#skipci

diff --git a/Engine/Plugins/Enterprise/GLTFExporter/MigratedFromGit/Source/GLTFExporter/Private/Converters/GLTFMaterialProxyFactory.cpp b/Engine/Plugins/Enterprise/GLTFExporter/MigratedFromGit/Source/GLTFExporter/Private/Converters/GLTFMaterialProxyFactory.cpp
--- a/Engine/Plugins/Enterprise/GLTFExporter/MigratedFromGit/Source/GLTFExporter/Private/Converters/GLTFMaterialProxyFactory.cpp
+++ b/Engine/Plugins/Enterprise/GLTFExporter/MigratedFromGit/Source/GLTFExporter/Private/Converters/GLTFMaterialProxyFactory.cpp
@@ -311,7 +311,7 @@ UGLTFExportOptions* FGLTFMaterialProxyFactory::CreateExportOptions(const UGLTFPr
 	UGLTFExportOptions* ExportOptions = NewObject<UGLTFExportOptions>();
 	ExportOptions->ResetToDefault();
 	ExportOptions->bExportProxyMaterials = false;
-	ExportOptions->BakeMaterialInputs = EGLTFMaterialBakeMode::Simple;
+	ExportOptions->BakeMaterialInputs = ProxyOptions->bBakeMaterialInputs ? EGLTFMaterialBakeMode::Simple : EGLTFMaterialBakeMode::Disabled;
 	ExportOptions->DefaultMaterialBakeSize = ProxyOptions->DefaultMaterialBakeSize;
 	ExportOptions->DefaultMaterialBakeFilter = ProxyOptions->DefaultMaterialBakeFilter;
 	ExportOptions->DefaultMaterialBakeTiling = ProxyOptions->DefaultMaterialBakeTiling;
diff --git a/Engine/Plugins/Enterprise/GLTFExporter/MigratedFromGit/Source/GLTFExporter/Private/Options/GLTFProxyOptions.cpp b/Engine/Plugins/Enterprise/GLTFExporter/MigratedFromGit/Source/GLTFExporter/Private/Options/GLTFProxyOptions.cpp
--- a/Engine/Plugins/Enterprise/GLTFExporter/MigratedFromGit/Source/GLTFExporter/Private/Options/GLTFProxyOptions.cpp
+++ b/Engine/Plugins/Enterprise/GLTFExporter/MigratedFromGit/Source/GLTFExporter/Private/Options/GLTFProxyOptions.cpp
@@ -10,6 +10,7 @@ UGLTFProxyOptions::UGLTFProxyOptions(const FObjectInitializer& ObjectInitializer
 
 void UGLTFProxyOptions::ResetToDefault()
 {
+	bBakeMaterialInputs = true;
 	DefaultMaterialBakeSize = EGLTFMaterialBakeSizePOT::POT_1024;
 	DefaultMaterialBakeFilter = TF_Trilinear;
 	DefaultMaterialBakeTiling = TA_Wrap;
diff --git a/Engine/Plugins/Enterprise/GLTFExporter/MigratedFromGit/Source/GLTFExporter/Public/Options/GLTFProxyOptions.h b/Engine/Plugins/Enterprise/GLTFExporter/MigratedFromGit/Source/GLTFExporter/Public/Options/GLTFProxyOptions.h
--- a/Engine/Plugins/Enterprise/GLTFExporter/MigratedFromGit/Source/GLTFExporter/Public/Options/GLTFProxyOptions.h
+++ b/Engine/Plugins/Enterprise/GLTFExporter/MigratedFromGit/Source/GLTFExporter/Public/Options/GLTFProxyOptions.h
@@ -10,20 +10,24 @@ class GLTFEXPORTER_API UGLTFProxyOptions : public UObject
 {
 	GENERATED_UCLASS_BODY()
 
-	/** Default size of the baked out texture (containing the material input). Can be overridden by material- and input-specific bake settings, see GLTFMaterialExportOptions. */
+	/** If enabled, a material input may be baked out to a texture (using a simple quad). Baking is only used for non-trivial material inputs (i.e. not simple texture or constant expressions). */
 	UPROPERTY(EditAnywhere, BlueprintReadWrite, Config, Category = Material)
+	bool bBakeMaterialInputs;
+
+	/** Default size of the baked out texture (containing the material input). Can be overridden by material- and input-specific bake settings, see GLTFMaterialExportOptions. */
+	UPROPERTY(EditAnywhere, BlueprintReadWrite, Config, Category = Material, Meta = (EditCondition = "bBakeMaterialInputs"))
 	EGLTFMaterialBakeSizePOT DefaultMaterialBakeSize;
 
 	/** Default filtering mode used when sampling the baked out texture. Can be overridden by material- and input-specific bake settings, see GLTFMaterialExportOptions. */
-	UPROPERTY(EditAnywhere, BlueprintReadWrite, Config, Category = Material, Meta = (ValidEnumValues="TF_Nearest, TF_Bilinear, TF_Trilinear"))
+	UPROPERTY(EditAnywhere, BlueprintReadWrite, Config, Category = Material, Meta = (EditCondition = "bBakeMaterialInputs", ValidEnumValues = "TF_Nearest, TF_Bilinear, TF_Trilinear"))
 	TEnumAsByte<TextureFilter> DefaultMaterialBakeFilter;
 
 	/** Default addressing mode used when sampling the baked out texture. Can be overridden by material- and input-specific bake settings, see GLTFMaterialExportOptions. */
-	UPROPERTY(EditAnywhere, BlueprintReadWrite, Config, Category = Material)
+	UPROPERTY(EditAnywhere, BlueprintReadWrite, Config, Category = Material, Meta = (EditCondition = "bBakeMaterialInputs"))
 	TEnumAsByte<TextureAddress> DefaultMaterialBakeTiling;
 
 	/** Input-specific default bake settings that override the general defaults above. */
-	UPROPERTY(EditAnywhere, BlueprintReadWrite, Config, Category = Material)
+	UPROPERTY(EditAnywhere, BlueprintReadWrite, Config, Category = Material, Meta = (EditCondition = "bBakeMaterialInputs"))
 	TMap<EGLTFMaterialPropertyGroup, FGLTFOverrideMaterialBakeSettings> DefaultInputBakeSettings;
 
 	UFUNCTION(BlueprintCallable, Category = General)
