[GLTFExporter] feat: improve UI window setup and layout for export and proxy options

commit 478f22d14e7c69ba6190978729e60a222cd6c53e
Date: Sat, 12 Feb 2022 01:32:12 +0100

    feat: improve UI window setup and layout for export and proxy options
    (cherry picked from commit eb051ece0d123fa3d398a49ac67cbe2ae91e187b​)

#skipci

diff --git a/Engine/Plugins/Enterprise/GLTFExporter/MigratedFromGit/Source/GLTFExporter/Private/UI/GLTFExportOptionsWindow.cpp b/Engine/Plugins/Enterprise/GLTFExporter/MigratedFromGit/Source/GLTFExporter/Private/UI/GLTFExportOptionsWindow.cpp
--- a/Engine/Plugins/Enterprise/GLTFExporter/MigratedFromGit/Source/GLTFExporter/Private/UI/GLTFExportOptionsWindow.cpp
+++ b/Engine/Plugins/Enterprise/GLTFExporter/MigratedFromGit/Source/GLTFExporter/Private/UI/GLTFExportOptionsWindow.cpp
@@ -16,17 +16,22 @@
 #include "Framework/Application/SlateApplication.h"
 #include "Interfaces/IMainFrameModule.h"
 
-#define LOCTEXT_NAMESPACE "GLTFExportOptions"
+#define LOCTEXT_NAMESPACE "GLTFExportOptionsWindow"
+
+SGLTFExportOptionsWindow::SGLTFExportOptionsWindow()
+	: ExportOptions(nullptr)
+	, bShouldExport(false)
+	, bShouldExportAll(false)
+{
+}
 
 void SGLTFExportOptionsWindow::Construct(const FArguments& InArgs)
 {
 	ExportOptions = InArgs._ExportOptions;
 	WidgetWindow = InArgs._WidgetWindow;
 
-	check(ExportOptions);
-
-	const FText CancelText = InArgs._BatchMode ? LOCTEXT("GLTFExportOptionsWindow_CancelBatch", "Cancel All") : LOCTEXT("GLTFExportOptionsWindow_Cancel", "Cancel");
-	const FText CancelTooltipText = InArgs._BatchMode ? LOCTEXT("GLTFExportOptionsWindow_Cancel_ToolTip_Batch", "Cancel the batch export.") : LOCTEXT("GLTFExportOptionsWindow_Cancel_ToolTip", "Cancel the current glTF export.");
+	const FText CancelText = InArgs._BatchMode ? LOCTEXT("CancelBatch", "Cancel All") : LOCTEXT("Cancel", "Cancel");
+	const FText CancelTooltipText = InArgs._BatchMode ? LOCTEXT("CancelBatch_ToolTip", "Cancel the batch export.") : LOCTEXT("Cancel_ToolTip", "Cancel the current glTF export.");
 
 	TSharedPtr<SBox> HeaderToolBox;
 	TSharedPtr<SHorizontalBox> HeaderButtons;
@@ -86,16 +91,16 @@ void SGLTFExportOptionsWindow::Construct(const FArguments& InArgs)
 				[
 					SNew(SButton)
 					.HAlign(HAlign_Center)
-					.Text(LOCTEXT("GLTFExportOptionsWindow_ExportAll", "Export All"))
-					.ToolTipText(LOCTEXT("GLTExportOptionsWindow_ExportAll_ToolTip", "Export all files with these same settings"))
+					.Text(LOCTEXT("ExportAll", "Export All"))
+					.ToolTipText(LOCTEXT("ExportAll_ToolTip", "Export all files with these same settings"))
 					.Visibility(InArgs._BatchMode ? EVisibility::All : EVisibility::Hidden)
 					.OnClicked(this, &SGLTFExportOptionsWindow::OnExportAll)
 				]
 				+ SUniformGridPanel::Slot(1, 0)
 				[
-					SAssignNew(ImportButton, SButton)
+					SAssignNew(ExportButton, SButton)
 					.HAlign(HAlign_Center)
-					.Text(LOCTEXT("GLTExportOptionsWindow_Export", "Export"))
+					.Text(LOCTEXT("Export", "Export"))
 					.OnClicked(this, &SGLTFExportOptionsWindow::OnExport)
 				]
 				+ SUniformGridPanel::Slot(2, 0)
@@ -135,8 +140,8 @@ void SGLTFExportOptionsWindow::Construct(const FArguments& InArgs)
 					.Padding(FMargin(2.0f, 0.0f))
 					[
 						SNew(SButton)
-						.Text(LOCTEXT("ExportOptionsWindow_ResetOptions", "Reset to Default"))
-						.OnClicked(this, &SGLTFExportOptionsWindow::OnResetToDefaultClick)
+						.Text(LOCTEXT("Reset", "Reset to Default"))
+						.OnClicked(this, &SGLTFExportOptionsWindow::OnReset)
 					]
 				]
 			]
@@ -146,7 +151,7 @@ void SGLTFExportOptionsWindow::Construct(const FArguments& InArgs)
 	DetailsView->SetObject(ExportOptions);
 }
 
-FReply SGLTFExportOptionsWindow::OnResetToDefaultClick() const
+FReply SGLTFExportOptionsWindow::OnReset() const
 {
 	ExportOptions->ResetToDefault();
 	//Refresh the view to make sure the custom UI are updating correctly
@@ -201,37 +206,21 @@ bool SGLTFExportOptionsWindow::ShouldExportAll() const
 	return bShouldExportAll;
 }
 
-SGLTFExportOptionsWindow::SGLTFExportOptionsWindow()
-	: ExportOptions(nullptr)
-	, bShouldExport(false)
-	, bShouldExportAll(false)
-{}
-
 void SGLTFExportOptionsWindow::ShowDialog(UGLTFExportOptions* ExportOptions, const FString& FullPath, bool bBatchMode, bool& bOutOperationCanceled, bool& bOutExportAll)
 {
-	check(ExportOptions != nullptr);
-
 	bOutOperationCanceled = false;
 	bOutExportAll = false;
 
-	TSharedPtr<SWindow> ParentWindow;
-
-	if (FModuleManager::Get().IsModuleLoaded("MainFrame"))
-	{
-		IMainFrameModule& MainFrame = FModuleManager::LoadModuleChecked<IMainFrameModule>("MainFrame");
-		ParentWindow = MainFrame.GetParentWindow();
-	}
-
 	const bool bBinaryFile = FPaths::GetExtension(FullPath, false).Compare(TEXT("glb"), ESearchCase::IgnoreCase) == 0;
 	const FText Title = bBinaryFile ?
-		NSLOCTEXT("UnrealEd", "GLTFExportOptionsTitleBinary", "glTF (Binary) Export Options") :
-		NSLOCTEXT("UnrealEd", "GLTFExportOptionsTitle", "glTF Export Options");
+		LOCTEXT("TitleBinary", "glTF (Binary) Export Options") :
+		LOCTEXT("Title", "glTF Export Options");
 
-	TSharedRef<SWindow> Window = SNew(SWindow)
+	const TSharedRef<SWindow> Window = SNew(SWindow)
 		.Title(Title)
 		.SizingRule(ESizingRule::UserSized)
 		.AutoCenter(EAutoCenter::PrimaryWorkArea)
-		.ClientSize(FVector2D(500, 445));
+		.ClientSize(FVector2D(500, 500));
 
 	TSharedPtr<SGLTFExportOptionsWindow> OptionsWindow;
 	Window->SetContent
@@ -243,6 +232,9 @@ void SGLTFExportOptionsWindow::ShowDialog(UGLTFExportOptions* ExportOptions, con
 		.BatchMode(bBatchMode)
 	);
 
+	IMainFrameModule* MainFrame = FModuleManager::LoadModulePtr<IMainFrameModule>("MainFrame");
+	const TSharedPtr<SWindow> ParentWindow = MainFrame != nullptr ? MainFrame->GetParentWindow() : nullptr;
+
 	FSlateApplication::Get().AddModalWindow(Window, ParentWindow, false);
 
 	if (OptionsWindow->ShouldExport())
diff --git a/Engine/Plugins/Enterprise/GLTFExporter/MigratedFromGit/Source/GLTFExporter/Private/UI/GLTFExportOptionsWindow.h b/Engine/Plugins/Enterprise/GLTFExporter/MigratedFromGit/Source/GLTFExporter/Private/UI/GLTFExportOptionsWindow.h
--- a/Engine/Plugins/Enterprise/GLTFExporter/MigratedFromGit/Source/GLTFExporter/Private/UI/GLTFExportOptionsWindow.h
+++ b/Engine/Plugins/Enterprise/GLTFExporter/MigratedFromGit/Source/GLTFExporter/Private/UI/GLTFExportOptionsWindow.h
@@ -27,34 +27,31 @@ public:
 		SLATE_ARGUMENT( bool, BatchMode )
 	SLATE_END_ARGS()
 
-	void Construct(const FArguments& InArgs);
+	SGLTFExportOptionsWindow();
 
-	virtual bool SupportsKeyboardFocus() const override { return true; }
+	void Construct(const FArguments& InArgs);
 
+	FReply OnReset() const;
 	FReply OnExport();
-
 	FReply OnExportAll();
-
 	FReply OnCancel();
 
+	/* Begin SCompoundWidget overrides */
+	virtual bool SupportsKeyboardFocus() const override { return true; }
 	virtual FReply OnKeyDown(const FGeometry& MyGeometry, const FKeyEvent& InKeyEvent) override;
+	/* End SCompoundWidget overrides */
 
 	bool ShouldExport() const;
-
 	bool ShouldExportAll() const;
 
-	SGLTFExportOptionsWindow();
-
 	static void ShowDialog(UGLTFExportOptions* ExportOptions, const FString& FullPath, bool bBatchMode, bool& bOutOperationCanceled, bool& bOutExportAll);
 
 private:
 
-	FReply OnResetToDefaultClick() const;
-
 	UGLTFExportOptions* ExportOptions;
 	TSharedPtr<class IDetailsView> DetailsView;
 	TWeakPtr<SWindow> WidgetWindow;
-	TSharedPtr<SButton> ImportButton;
+	TSharedPtr<SButton> ExportButton;
 	bool bShouldExport;
 	bool bShouldExportAll;
 };
diff --git a/Engine/Plugins/Enterprise/GLTFExporter/MigratedFromGit/Source/GLTFExporter/Private/UI/GLTFProxyOptionsWindow.cpp b/Engine/Plugins/Enterprise/GLTFExporter/MigratedFromGit/Source/GLTFExporter/Private/UI/GLTFProxyOptionsWindow.cpp
--- a/Engine/Plugins/Enterprise/GLTFExporter/MigratedFromGit/Source/GLTFExporter/Private/UI/GLTFProxyOptionsWindow.cpp
+++ b/Engine/Plugins/Enterprise/GLTFExporter/MigratedFromGit/Source/GLTFExporter/Private/UI/GLTFProxyOptionsWindow.cpp
@@ -10,64 +10,64 @@
 #include "Widgets/Layout/SUniformGridPanel.h"
 #include "Widgets/Input/SButton.h"
 
-#define LOCTEXT_NAMESPACE "SGLTFProxyOptionsWindow"
+#define LOCTEXT_NAMESPACE "GLTFProxyOptionsWindow"
 
-SGLTFProxyOptionsWindow::SGLTFProxyOptionsWindow() : bUserCancelled(true)
+SGLTFProxyOptionsWindow::SGLTFProxyOptionsWindow()
+	: ProxyOptions(nullptr)
+	, bUserCancelled(true)
 {
-
 }
 
 void SGLTFProxyOptionsWindow::Construct(const FArguments& InArgs)
 {
+	ProxyOptions = InArgs._ProxyOptions;
 	WidgetWindow = InArgs._WidgetWindow;
 
-	// Retrieve property editor module and create a SDetailsView
-	FPropertyEditorModule& PropertyEditorModule = FModuleManager::GetModuleChecked<FPropertyEditorModule>("PropertyEditor");
-	FDetailsViewArgs DetailsViewArgs;
-	DetailsViewArgs.bAllowSearch = false;
-	DetailsViewArgs.NameAreaSettings = FDetailsViewArgs::HideNameArea;
-	DetailsViewArgs.bAllowMultipleTopLevelObjects = true;
-
-	DetailsView = PropertyEditorModule.CreateDetailView(DetailsViewArgs);
-
-	// Set provided objects on SDetailsView
-	DetailsView->SetObjects(InArgs._SettingsObjects, true);
-
+	TSharedPtr<SBox> InspectorBox;
 	this->ChildSlot
 	[
-		SNew(SVerticalBox)
-
-		+ SVerticalBox::Slot()
-		.Padding(2)
-		.MaxHeight(500.0f)
-		[
-			DetailsView->AsShared()
-		]
-
-		+ SVerticalBox::Slot()
-		.AutoHeight()
-		.HAlign(HAlign_Right)
-		.Padding(2)
+		SNew(SBox)
 		[
-			SNew(SUniformGridPanel)
-			.SlotPadding(2)
-			+ SUniformGridPanel::Slot(0, 0)
+			SNew(SVerticalBox)
+			+ SVerticalBox::Slot()
+			.FillHeight(1.0f)
+			.Padding(2)
 			[
-				SAssignNew(ConfirmButton, SButton)
-				.HAlign(HAlign_Center)
-				.Text(LOCTEXT("MaterialBakeOptionWindow_Import", "Confirm"))
-				.OnClicked(this, &SGLTFProxyOptionsWindow::OnConfirm)
+				SAssignNew(InspectorBox, SBox)
 			]
-			+ SUniformGridPanel::Slot(1, 0)
+			+ SVerticalBox::Slot()
+			.AutoHeight()
+			.HAlign(HAlign_Right)
+			.Padding(2)
 			[
-				SNew(SButton)
-				.HAlign(HAlign_Center)
-				.Text(LOCTEXT("MaterialBakeOptionWindow_Cancel", "Cancel"))
-				.ToolTipText(LOCTEXT("MaterialBakeOptionWindow_Cancel_ToolTip", "Cancels baking out Material"))
-				.OnClicked(this, &SGLTFProxyOptionsWindow::OnCancel)
+				SNew(SUniformGridPanel)
+				.SlotPadding(2)
+				+ SUniformGridPanel::Slot(0, 0)
+				[
+					SAssignNew(ConfirmButton, SButton)
+					.HAlign(HAlign_Center)
+					.Text(LOCTEXT("Confirm", "Confirm"))
+					.OnClicked(this, &SGLTFProxyOptionsWindow::OnConfirm)
+				]
+				+ SUniformGridPanel::Slot(1, 0)
+				[
+					SNew(SButton)
+					.HAlign(HAlign_Center)
+					.Text(LOCTEXT("Cancel", "Cancel"))
+					.OnClicked(this, &SGLTFProxyOptionsWindow::OnCancel)
+				]
 			]
 		]
 	];
+
+	FPropertyEditorModule& PropertyEditorModule = FModuleManager::GetModuleChecked<FPropertyEditorModule>("PropertyEditor");
+	FDetailsViewArgs DetailsViewArgs;
+	DetailsViewArgs.bAllowSearch = false;
+	DetailsViewArgs.NameAreaSettings = FDetailsViewArgs::HideNameArea;
+
+	DetailsView = PropertyEditorModule.CreateDetailView(DetailsViewArgs);
+	InspectorBox->SetContent(DetailsView->AsShared());
+	DetailsView->SetObject(ProxyOptions);
 }
 
 FReply SGLTFProxyOptionsWindow::OnConfirm()
@@ -82,6 +82,7 @@ FReply SGLTFProxyOptionsWindow::OnConfirm()
 
 FReply SGLTFProxyOptionsWindow::OnCancel()
 {
+	bUserCancelled = true;
 	if (WidgetWindow.IsValid())
 	{
 		WidgetWindow.Pin()->RequestDestroyWindow();
@@ -99,39 +100,30 @@ FReply SGLTFProxyOptionsWindow::OnKeyDown(const FGeometry& MyGeometry, const FKe
 	return FReply::Unhandled();
 }
 
-bool SGLTFProxyOptionsWindow::WasUserCancelled()
-{
-	return bUserCancelled;
-}
-
 bool SGLTFProxyOptionsWindow::ShowDialog(UGLTFProxyOptions* ProxyOptions)
 {
-	TSharedRef<SWindow> Window = SNew(SWindow)
-		.Title(LOCTEXT("WindowTitle", "Material Baking Options"))
-		.SizingRule(ESizingRule::Autosized);
+	const TSharedRef<SWindow> Window = SNew(SWindow)
+		.Title(LOCTEXT("Title", "Create glTF Proxy Material"))
+		.SizingRule(ESizingRule::UserSized)
+		.AutoCenter(EAutoCenter::PrimaryWorkArea)
+		.ClientSize(FVector2D(500, 250));
 
-	TArray<TWeakObjectPtr<UObject>> OptionObjects = { ProxyOptions };
-	TSharedPtr<SGLTFProxyOptionsWindow> Options;
+	TSharedPtr<SGLTFProxyOptionsWindow> OptionsWindow;
 
 	Window->SetContent
 	(
-		SAssignNew(Options, SGLTFProxyOptionsWindow)
+		SAssignNew(OptionsWindow, SGLTFProxyOptionsWindow)
 		.WidgetWindow(Window)
-		.SettingsObjects(OptionObjects)
+		.ProxyOptions(ProxyOptions)
 	);
 
-	TSharedPtr<SWindow> ParentWindow;
-	if (FModuleManager::Get().IsModuleLoaded("MainFrame"))
-	{
-		IMainFrameModule& MainFrame = FModuleManager::LoadModuleChecked<IMainFrameModule>("MainFrame");
-		ParentWindow = MainFrame.GetParentWindow();
-		FSlateApplication::Get().AddModalWindow(Window, ParentWindow, false);
-		return !Options->WasUserCancelled();
-	}
+	IMainFrameModule* MainFrame = FModuleManager::LoadModulePtr<IMainFrameModule>("MainFrame");
+	const TSharedPtr<SWindow> ParentWindow = MainFrame != nullptr ? MainFrame->GetParentWindow() : nullptr;
 
-	return false;
+	FSlateApplication::Get().AddModalWindow(Window, ParentWindow, false);
+	return !OptionsWindow->bUserCancelled;
 }
 
-#undef LOCTEXT_NAMESPACE //"SGLTFProxyOptionsWindow"
+#undef LOCTEXT_NAMESPACE
 
 #endif
diff --git a/Engine/Plugins/Enterprise/GLTFExporter/MigratedFromGit/Source/GLTFExporter/Private/UI/GLTFProxyOptionsWindow.h b/Engine/Plugins/Enterprise/GLTFExporter/MigratedFromGit/Source/GLTFExporter/Private/UI/GLTFProxyOptionsWindow.h
--- a/Engine/Plugins/Enterprise/GLTFExporter/MigratedFromGit/Source/GLTFExporter/Private/UI/GLTFProxyOptionsWindow.h
+++ b/Engine/Plugins/Enterprise/GLTFExporter/MigratedFromGit/Source/GLTFExporter/Private/UI/GLTFProxyOptionsWindow.h
@@ -18,36 +18,31 @@ public:
 
 	SLATE_BEGIN_ARGS(SGLTFProxyOptionsWindow)
 	{}
-		SLATE_ARGUMENT(TSharedPtr<SWindow>, WidgetWindow)
-		SLATE_ARGUMENT(TArray<TWeakObjectPtr<UObject>>, SettingsObjects)
+		SLATE_ARGUMENT( UGLTFProxyOptions*, ProxyOptions )
+		SLATE_ARGUMENT( TSharedPtr<SWindow>, WidgetWindow )
 	SLATE_END_ARGS()
 
 	SGLTFProxyOptionsWindow();
-	void Construct(const FArguments& InArgs);
 
-	/** Begin SCompoundWidget overrides */
-	virtual bool SupportsKeyboardFocus() const override { return true; }
-	virtual FReply OnKeyDown(const FGeometry& MyGeometry, const FKeyEvent& InKeyEvent) override;
-	/** End SCompoundWidget overrides */
+	void Construct(const FArguments& InArgs);
 
-	/** Callbacks used for Confirm and Cancel buttons */
 	FReply OnConfirm();
 	FReply OnCancel();
 
-	/** Returns whether or not the user cancelled the operation */
-	bool WasUserCancelled();
+	/* Begin SCompoundWidget overrides */
+	virtual bool SupportsKeyboardFocus() const override { return true; }
+	virtual FReply OnKeyDown(const FGeometry& MyGeometry, const FKeyEvent& InKeyEvent) override;
+	/* End SCompoundWidget overrides */
 
 	static bool ShowDialog(UGLTFProxyOptions* ProxyOptions);
 
 private:
-	/** Owning window this widget is part of */
-	TWeakPtr< SWindow > WidgetWindow;
-	/** Whether or not the cancel button was clicked by the user */
-	bool bUserCancelled;
-	/** Detailsview used to display SettingsObjects, and allowing user to change options */
+
+	UGLTFProxyOptions* ProxyOptions;
 	TSharedPtr<class IDetailsView> DetailsView;
-	/** Shared ptr to Confirm button */
+	TWeakPtr<SWindow> WidgetWindow;
 	TSharedPtr<SButton> ConfirmButton;
+	bool bUserCancelled;
 };
 
 #endif
