[GLTFExporter] feat: add optional parameter to override root path in public API for creating proxy material programatically

commit 93539c4b17446000245385dafd8ac32140c9edfe
Date: Sat, 12 Feb 2022 01:32:51 +0100

    feat: add optional parameter to override root path in public API for creating proxy material programatically
    (cherry picked from commit 1b5d6d8445af57d475d163344d50f45a2095e1af​)

#skipci

diff --git a/Engine/Plugins/Enterprise/GLTFExporter/MigratedFromGit/Source/GLTFExporter/Private/GLTFExporterModule.cpp b/Engine/Plugins/Enterprise/GLTFExporter/MigratedFromGit/Source/GLTFExporter/Private/GLTFExporterModule.cpp
--- a/Engine/Plugins/Enterprise/GLTFExporter/MigratedFromGit/Source/GLTFExporter/Private/GLTFExporterModule.cpp
+++ b/Engine/Plugins/Enterprise/GLTFExporter/MigratedFromGit/Source/GLTFExporter/Private/GLTFExporterModule.cpp
@@ -14,11 +14,11 @@ class FAssetTypeActions_Base;
 
 DEFINE_LOG_CATEGORY(LogGLTFExporter);
 
-UMaterialInterface* IGLTFExporterModule::CreateProxyMaterial(UMaterialInterface* Material, const UGLTFProxyOptions* Options)
+UMaterialInterface* IGLTFExporterModule::CreateProxyMaterial(UMaterialInterface* Material, const UGLTFProxyOptions* Options, const FString& RootPath)
 {
 #if WITH_EDITOR
 	FGLTFMaterialProxyFactory ProxyFactory(Options);
-	ProxyFactory.RootPath = FPaths::GetPath(Material->GetPathName()) / TEXT("GLTF");
+	ProxyFactory.RootPath = RootPath.IsEmpty() ? FPaths::GetPath(Material->GetPathName()) / TEXT("GLTF") : RootPath;
 	return ProxyFactory.Create(Material);
 #else
 	// TODO: report error
diff --git a/Engine/Plugins/Enterprise/GLTFExporter/MigratedFromGit/Source/GLTFExporter/Public/GLTFExporterModule.h b/Engine/Plugins/Enterprise/GLTFExporter/MigratedFromGit/Source/GLTFExporter/Public/GLTFExporterModule.h
--- a/Engine/Plugins/Enterprise/GLTFExporter/MigratedFromGit/Source/GLTFExporter/Public/GLTFExporterModule.h
+++ b/Engine/Plugins/Enterprise/GLTFExporter/MigratedFromGit/Source/GLTFExporter/Public/GLTFExporterModule.h
@@ -39,5 +39,5 @@ public:
 		return FModuleManager::Get().IsModuleLoaded(GLTFEXPORTER_MODULE_NAME);
 	}
 
-	static UMaterialInterface* CreateProxyMaterial(UMaterialInterface* Material, const UGLTFProxyOptions* Options);
+	static UMaterialInterface* CreateProxyMaterial(UMaterialInterface* Material, const UGLTFProxyOptions* Options, const FString& RootPath = TEXT(""));
 };
