[GLTFExporter] refactor: encapsulate WITH_EDITOR in GLTFBufferUtility to avoid code inspection warnings in Rider

commit 4c6629f07bfebbcf206ef5b5520fcae916289eb9
Date: Sat, 12 Feb 2022 21:39:09 +0100

    refactor: encapsulate WITH_EDITOR in GLTFBufferUtility to avoid code inspection warnings in Rider
    (cherry picked from commit 04f72bcbcc0b9a31bae1267ac52a924cd9555406​)

#skipci

diff --git a/Engine/Plugins/Enterprise/GLTFExporter/MigratedFromGit/Source/GLTFExporter/Private/Converters/GLTFBufferAdapter.cpp b/Engine/Plugins/Enterprise/GLTFExporter/MigratedFromGit/Source/GLTFExporter/Private/Converters/GLTFBufferAdapter.cpp
--- a/Engine/Plugins/Enterprise/GLTFExporter/MigratedFromGit/Source/GLTFExporter/Private/Converters/GLTFBufferAdapter.cpp
+++ b/Engine/Plugins/Enterprise/GLTFExporter/MigratedFromGit/Source/GLTFExporter/Private/Converters/GLTFBufferAdapter.cpp
@@ -44,58 +44,58 @@ public:
 
 TUniquePtr<IGLTFBufferAdapter> IGLTFBufferAdapter::GetIndices(const FRawStaticIndexBuffer* IndexBuffer)
 {
-	const void* IndexData = IndexBuffer->Is32Bit() ? static_cast<const void*>(IndexBuffer->AccessStream32()) : static_cast<const void*>(IndexBuffer->AccessStream16());
-	if (IndexData != nullptr && (WITH_EDITOR || FGLTFBufferUtility::GetAllowCPUAccess(IndexBuffer))) return MakeUnique<FGLTFBufferAdapterCPU>(IndexData);
+	const void* IndexData = FGLTFBufferUtility::GetCPUBuffer(IndexBuffer);
+	if (IndexData != nullptr && FGLTFBufferUtility::HasCPUAccess(IndexBuffer)) return MakeUnique<FGLTFBufferAdapterCPU>(IndexData);
 	return MakeUnique<FGLTFBufferAdapterGPU>(IndexBuffer->IndexBufferRHI);
 }
 
 TUniquePtr<IGLTFBufferAdapter> IGLTFBufferAdapter::GetIndices(const FRawStaticIndexBuffer16or32Interface* IndexBuffer)
 {
-	const void* IndexData = IndexBuffer->GetResourceDataSize() > 0 ? const_cast<FRawStaticIndexBuffer16or32Interface*>(IndexBuffer)->GetPointerTo(0) : nullptr;
-	if (IndexData != nullptr && (WITH_EDITOR || IndexBuffer->GetNeedsCPUAccess())) return MakeUnique<FGLTFBufferAdapterCPU>(IndexData);
+	const void* IndexData =FGLTFBufferUtility::GetCPUBuffer(IndexBuffer);
+	if (IndexData != nullptr && FGLTFBufferUtility::HasCPUAccess(IndexBuffer)) return MakeUnique<FGLTFBufferAdapterCPU>(IndexData);
 	return MakeUnique<FGLTFBufferAdapterGPU>(IndexBuffer->IndexBufferRHI);
 }
 
 TUniquePtr<IGLTFBufferAdapter> IGLTFBufferAdapter::GetPositions(const FPositionVertexBuffer* VertexBuffer)
 {
 	const void* PositionData = const_cast<FPositionVertexBuffer*>(VertexBuffer)->GetVertexData();
-	if (PositionData != nullptr && (WITH_EDITOR || FGLTFBufferUtility::GetAllowCPUAccess(VertexBuffer))) return MakeUnique<FGLTFBufferAdapterCPU>(PositionData);
+	if (PositionData != nullptr && FGLTFBufferUtility::HasCPUAccess(VertexBuffer)) return MakeUnique<FGLTFBufferAdapterCPU>(PositionData);
 	return MakeUnique<FGLTFBufferAdapterGPU>(VertexBuffer->VertexBufferRHI);
 }
 
 TUniquePtr<IGLTFBufferAdapter> IGLTFBufferAdapter::GetColors(const FColorVertexBuffer* VertexBuffer)
 {
 	const void* ColorData = const_cast<FColorVertexBuffer*>(VertexBuffer)->GetVertexData();
-	if (ColorData != nullptr && (WITH_EDITOR || FGLTFBufferUtility::GetAllowCPUAccess(VertexBuffer))) return MakeUnique<FGLTFBufferAdapterCPU>(ColorData);
+	if (ColorData != nullptr && FGLTFBufferUtility::HasCPUAccess(VertexBuffer)) return MakeUnique<FGLTFBufferAdapterCPU>(ColorData);
 	return MakeUnique<FGLTFBufferAdapterGPU>(VertexBuffer->VertexBufferRHI);
 }
 
 TUniquePtr<IGLTFBufferAdapter> IGLTFBufferAdapter::GetTangents(const FStaticMeshVertexBuffer* VertexBuffer)
 {
 	const void* TangentData = const_cast<FStaticMeshVertexBuffer*>(VertexBuffer)->GetTangentData();
-	if (TangentData != nullptr && (WITH_EDITOR || const_cast<FStaticMeshVertexBuffer*>(VertexBuffer)->GetAllowCPUAccess())) return MakeUnique<FGLTFBufferAdapterCPU>(TangentData);
+	if (TangentData != nullptr && FGLTFBufferUtility::HasCPUAccess(VertexBuffer)) return MakeUnique<FGLTFBufferAdapterCPU>(TangentData);
 	return MakeUnique<FGLTFBufferAdapterGPU>(VertexBuffer->TangentsVertexBuffer.VertexBufferRHI);
 }
 
 TUniquePtr<IGLTFBufferAdapter> IGLTFBufferAdapter::GetUVs(const FStaticMeshVertexBuffer* VertexBuffer)
 {
 	const void* UVData = const_cast<FStaticMeshVertexBuffer*>(VertexBuffer)->GetTexCoordData();
-	if (UVData != nullptr && (WITH_EDITOR || const_cast<FStaticMeshVertexBuffer*>(VertexBuffer)->GetAllowCPUAccess())) return MakeUnique<FGLTFBufferAdapterCPU>(UVData);
+	if (UVData != nullptr && FGLTFBufferUtility::HasCPUAccess(VertexBuffer)) return MakeUnique<FGLTFBufferAdapterCPU>(UVData);
 	return MakeUnique<FGLTFBufferAdapterGPU>(VertexBuffer->TexCoordVertexBuffer.VertexBufferRHI);
 }
 
 TUniquePtr<IGLTFBufferAdapter> IGLTFBufferAdapter::GetInfluences(const FSkinWeightVertexBuffer* VertexBuffer)
 {
-	const auto* InfluenceBuffer = VertexBuffer->GetDataVertexBuffer();
-	const void* InfluenceData = FGLTFBufferUtility::GetBufferData(InfluenceBuffer);
-	if (InfluenceData != nullptr && (WITH_EDITOR || VertexBuffer->GetNeedsCPUAccess())) return MakeUnique<FGLTFBufferAdapterCPU>(InfluenceData);
+	const FSkinWeightDataVertexBuffer* InfluenceBuffer = VertexBuffer->GetDataVertexBuffer();
+	const void* InfluenceData = FGLTFBufferUtility::GetCPUBuffer(InfluenceBuffer);
+	if (InfluenceData != nullptr && FGLTFBufferUtility::HasCPUAccess(VertexBuffer)) return MakeUnique<FGLTFBufferAdapterCPU>(InfluenceData);
 	return MakeUnique<FGLTFBufferAdapterGPU>(InfluenceBuffer->VertexBufferRHI);
 }
 
 TUniquePtr<IGLTFBufferAdapter> IGLTFBufferAdapter::GetLookups(const FSkinWeightVertexBuffer* VertexBuffer)
 {
 	const FSkinWeightLookupVertexBuffer* LookupBuffer = VertexBuffer->GetLookupVertexBuffer();
-	const void* LookupData = FGLTFBufferUtility::GetBufferData(LookupBuffer);
-	if (LookupData != nullptr && (WITH_EDITOR || VertexBuffer->GetNeedsCPUAccess())) return MakeUnique<FGLTFBufferAdapterCPU>(LookupData);
+	const void* LookupData = FGLTFBufferUtility::GetCPUBuffer(LookupBuffer);
+	if (LookupData != nullptr && FGLTFBufferUtility::HasCPUAccess(VertexBuffer)) return MakeUnique<FGLTFBufferAdapterCPU>(LookupData);
 	return MakeUnique<FGLTFBufferAdapterGPU>(LookupBuffer->VertexBufferRHI);
 }
diff --git a/Engine/Plugins/Enterprise/GLTFExporter/MigratedFromGit/Source/GLTFExporter/Private/Converters/GLTFBufferUtility.cpp b/Engine/Plugins/Enterprise/GLTFExporter/MigratedFromGit/Source/GLTFExporter/Private/Converters/GLTFBufferUtility.cpp
--- a/Engine/Plugins/Enterprise/GLTFExporter/MigratedFromGit/Source/GLTFExporter/Private/Converters/GLTFBufferUtility.cpp
+++ b/Engine/Plugins/Enterprise/GLTFExporter/MigratedFromGit/Source/GLTFExporter/Private/Converters/GLTFBufferUtility.cpp
@@ -3,8 +3,11 @@
 #include "Converters/GLTFBufferUtility.h"
 #include "Rendering/SkeletalMeshRenderData.h"
 
-bool FGLTFBufferUtility::GetAllowCPUAccess(const FRawStaticIndexBuffer* IndexBuffer)
+bool FGLTFBufferUtility::HasCPUAccess(const FRawStaticIndexBuffer* IndexBuffer)
 {
+#if WITH_EDITOR
+	return true;
+#else
 	struct FRawStaticIndexBufferHack : FIndexBuffer
 	{
 		TResourceArray<uint8, INDEXBUFFER_ALIGNMENT> IndexStorage;
@@ -15,10 +18,23 @@ bool FGLTFBufferUtility::GetAllowCPUAccess(const FRawStaticIndexBuffer* IndexBuf
 
 	static_assert(sizeof(FRawStaticIndexBufferHack) == sizeof(FRawStaticIndexBuffer), "FRawStaticIndexBufferHack memory layout doesn't match FRawStaticIndexBuffer");
 	return reinterpret_cast<const FRawStaticIndexBufferHack*>(IndexBuffer)->IndexStorage.GetAllowCPUAccess();
+#endif
+}
+
+bool FGLTFBufferUtility::HasCPUAccess(const FRawStaticIndexBuffer16or32Interface* IndexBuffer)
+{
+#if WITH_EDITOR
+	return true;
+#else
+	return IndexBuffer->GetNeedsCPUAccess();
+#endif
 }
 
-bool FGLTFBufferUtility::GetAllowCPUAccess(const FPositionVertexBuffer* VertexBuffer)
+bool FGLTFBufferUtility::HasCPUAccess(const FPositionVertexBuffer* VertexBuffer)
 {
+#if WITH_EDITOR
+	return true;
+#else
 	struct FPositionVertexBufferHack : FVertexBuffer
 	{
 		FShaderResourceViewRHIRef PositionComponentSRV;
@@ -32,10 +48,14 @@ bool FGLTFBufferUtility::GetAllowCPUAccess(const FPositionVertexBuffer* VertexBu
 	static_assert(sizeof(FPositionVertexBufferHack) == sizeof(FPositionVertexBuffer), "FPositionVertexBufferHack memory layout doesn't match FPositionVertexBuffer");
 	const FStaticMeshVertexDataInterface* VertexData = reinterpret_cast<const FPositionVertexBufferHack*>(VertexBuffer)->VertexData;
 	return VertexData != nullptr && VertexData->GetAllowCPUAccess();
+#endif
 }
 
-bool FGLTFBufferUtility::GetAllowCPUAccess(const FColorVertexBuffer* VertexBuffer)
+bool FGLTFBufferUtility::HasCPUAccess(const FColorVertexBuffer* VertexBuffer)
 {
+#if WITH_EDITOR
+	return true;
+#else
 	struct FColorVertexBufferHack : FVertexBuffer
 	{
 		FStaticMeshVertexDataInterface* VertexData;
@@ -49,13 +69,41 @@ bool FGLTFBufferUtility::GetAllowCPUAccess(const FColorVertexBuffer* VertexBuffe
 	static_assert(sizeof(FColorVertexBufferHack) == sizeof(FColorVertexBuffer), "FColorVertexBufferHack memory layout doesn't match FColorVertexBuffer");
 	const FStaticMeshVertexDataInterface* VertexData = reinterpret_cast<const FColorVertexBufferHack*>(VertexBuffer)->VertexData;
 	return VertexData != nullptr && VertexData->GetAllowCPUAccess();
+#endif
+}
+
+bool FGLTFBufferUtility::HasCPUAccess(const FStaticMeshVertexBuffer* VertexBuffer)
+{
+#if WITH_EDITOR
+	return true;
+#else
+	return const_cast<FStaticMeshVertexBuffer*>(VertexBuffer)->GetAllowCPUAccess();
+#endif
+}
+
+bool FGLTFBufferUtility::HasCPUAccess(const FSkinWeightVertexBuffer* VertexBuffer)
+{
+#if WITH_EDITOR
+	return true;
+#else
+	return VertexBuffer->GetNeedsCPUAccess();
+#endif
 }
 
-const void* FGLTFBufferUtility::GetBufferData(const FSkinWeightDataVertexBuffer* VertexBuffer)
+const void* FGLTFBufferUtility::GetCPUBuffer(const FRawStaticIndexBuffer* IndexBuffer)
+{
+	return IndexBuffer->Is32Bit() ? static_cast<const void*>(IndexBuffer->AccessStream32()) : static_cast<const void*>(IndexBuffer->AccessStream16());
+}
+
+const void* FGLTFBufferUtility::GetCPUBuffer(const FRawStaticIndexBuffer16or32Interface* IndexBuffer)
+{
+	return IndexBuffer->GetResourceDataSize() > 0 ? const_cast<FRawStaticIndexBuffer16or32Interface*>(IndexBuffer)->GetPointerTo(0) : nullptr;
+}
+
+const void* FGLTFBufferUtility::GetCPUBuffer(const FSkinWeightDataVertexBuffer* VertexBuffer)
 {
 #if (ENGINE_MAJOR_VERSION > 4 || ENGINE_MINOR_VERSION >= 26)
 	return VertexBuffer->GetWeightData();
-
 #else
 	struct FSkinWeightDataVertexBufferHack : FVertexBuffer
 	{
@@ -75,7 +123,7 @@ const void* FGLTFBufferUtility::GetBufferData(const FSkinWeightDataVertexBuffer*
 #endif
 }
 
-const void* FGLTFBufferUtility::GetBufferData(const FSkinWeightLookupVertexBuffer* VertexBuffer)
+const void* FGLTFBufferUtility::GetCPUBuffer(const FSkinWeightLookupVertexBuffer* VertexBuffer)
 {
 	struct FSkinWeightLookupVertexBufferHack : FVertexBuffer
 	{
diff --git a/Engine/Plugins/Enterprise/GLTFExporter/MigratedFromGit/Source/GLTFExporter/Private/Converters/GLTFBufferUtility.h b/Engine/Plugins/Enterprise/GLTFExporter/MigratedFromGit/Source/GLTFExporter/Private/Converters/GLTFBufferUtility.h
--- a/Engine/Plugins/Enterprise/GLTFExporter/MigratedFromGit/Source/GLTFExporter/Private/Converters/GLTFBufferUtility.h
+++ b/Engine/Plugins/Enterprise/GLTFExporter/MigratedFromGit/Source/GLTFExporter/Private/Converters/GLTFBufferUtility.h
@@ -6,12 +6,17 @@
 
 struct FGLTFBufferUtility
 {
-	static bool GetAllowCPUAccess(const FRawStaticIndexBuffer* IndexBuffer);
-	static bool GetAllowCPUAccess(const FPositionVertexBuffer* VertexBuffer);
-	static bool GetAllowCPUAccess(const FColorVertexBuffer* VertexBuffer);
+	static bool HasCPUAccess(const FRawStaticIndexBuffer* IndexBuffer);
+	static bool HasCPUAccess(const FRawStaticIndexBuffer16or32Interface* IndexBuffer);
+	static bool HasCPUAccess(const FPositionVertexBuffer* VertexBuffer);
+	static bool HasCPUAccess(const FColorVertexBuffer* VertexBuffer);
+	static bool HasCPUAccess(const FStaticMeshVertexBuffer* VertexBuffer);
+	static bool HasCPUAccess(const FSkinWeightVertexBuffer* VertexBuffer);
 
-	static const void* GetBufferData(const FSkinWeightDataVertexBuffer* VertexBuffer);
-	static const void* GetBufferData(const FSkinWeightLookupVertexBuffer* VertexBuffer);
+	static const void* GetCPUBuffer(const FRawStaticIndexBuffer* IndexBuffer);
+	static const void* GetCPUBuffer(const FRawStaticIndexBuffer16or32Interface* IndexBuffer);
+	static const void* GetCPUBuffer(const FSkinWeightDataVertexBuffer* VertexBuffer);
+	static const void* GetCPUBuffer(const FSkinWeightLookupVertexBuffer* VertexBuffer);
 
 	static void ReadRHIBuffer(FRHIIndexBuffer* SourceBuffer, TArray<uint8>& OutData);
 	static void ReadRHIBuffer(FRHIVertexBuffer* SourceBuffer, TArray<uint8>& OutData);
