[GLTFExporter] refactor: move "create proxy" icon (16x16) to Resources folder

commit f154286da5d07125958f75e4751b2d6a0a5de7c1
Date: Sat, 12 Feb 2022 22:50:15 +0100

    refactor: move "create proxy" icon (16x16) to Resources folder
    (cherry picked from commit eea2d750dca77bd06bad6017c67afce50cf24921​)

#skipci

diff --git a/Engine/Plugins/Enterprise/GLTFExporter/MigratedFromGit/Source/GLTFExporter/Private/Actions/GLTFEditorStyle.h b/Engine/Plugins/Enterprise/GLTFExporter/MigratedFromGit/Source/GLTFExporter/Private/Actions/GLTFEditorStyle.h
--- a/Engine/Plugins/Enterprise/GLTFExporter/MigratedFromGit/Source/GLTFExporter/Private/Actions/GLTFEditorStyle.h
+++ b/Engine/Plugins/Enterprise/GLTFExporter/MigratedFromGit/Source/GLTFExporter/Private/Actions/GLTFEditorStyle.h
@@ -21,10 +21,10 @@ public:
 	 	const TSharedPtr<IPlugin> Plugin = IPluginManager::Get().FindPlugin(GLTFEXPORTER_MODULE_NAME);
 	 	check(Plugin.IsValid());
 
-	 	FSlateStyleSet::SetContentRoot(Plugin->GetContentDir());
+	 	FSlateStyleSet::SetContentRoot(Plugin->GetBaseDir() / TEXT("Resources"));
 
-	 	const FVector2D Icon16x16(16.0f, 16.0f);
-	 	Set("Icon16x16", new FSlateImageBrush(FSlateStyleSet::RootToContentDir(TEXT("Icon16x16.png")), Icon16x16));
+	 	const FVector2D IconSize(16.0f, 16.0f);
+	 	Set("Icon16", new FSlateImageBrush(FSlateStyleSet::RootToContentDir(TEXT("Icon16.png")), IconSize));
 
 		FSlateStyleRegistry::RegisterSlateStyle(*this);
 	 }
diff --git a/Engine/Plugins/Enterprise/GLTFExporter/MigratedFromGit/Source/GLTFExporter/Private/Actions/GLTFProxyAssetActions.cpp b/Engine/Plugins/Enterprise/GLTFExporter/MigratedFromGit/Source/GLTFExporter/Private/Actions/GLTFProxyAssetActions.cpp
--- a/Engine/Plugins/Enterprise/GLTFExporter/MigratedFromGit/Source/GLTFExporter/Private/Actions/GLTFProxyAssetActions.cpp
+++ b/Engine/Plugins/Enterprise/GLTFExporter/MigratedFromGit/Source/GLTFExporter/Private/Actions/GLTFProxyAssetActions.cpp
@@ -35,7 +35,7 @@ void FGLTFProxyAssetActions::GetProxyActions(const TArray<UObject*>& InObjects,
 		"MenuEntry_CreateProxy",
 		LOCTEXT("MenuEntry_CreateProxy", "Create glTF Proxy Material"),
 		LOCTEXT("MenuEntry_CreateProxyTooltip", "Creates a proxy version of this material for glTF export."),
-		FSlateIcon(FGLTFEditorStyle::Get().GetStyleSetName(), "Icon16x16"),
+		FSlateIcon(FGLTFEditorStyle::Get().GetStyleSetName(), "Icon16"),
 		FUIAction(FExecuteAction::CreateSP(this, &FGLTFProxyAssetActions::OnCreateProxy, Objects))
 		);
 }
