[GLTFExporter] refactor: avoid compiler errors regarding OnPostEngineInit when packaging plugin

commit e8f1ac6ab549197e5244833d0672c1b9555bd79b
Date: Sat, 12 Feb 2022 22:59:43 +0100

    refactor: avoid compiler errors regarding OnPostEngineInit when packaging plugin
    (cherry picked from commit 06bea6567fad99ffdfeb3be3089827922232c12d​)

#skipci

diff --git a/Engine/Plugins/Enterprise/GLTFExporter/MigratedFromGit/Source/GLTFExporter/Private/GLTFExporterModule.cpp b/Engine/Plugins/Enterprise/GLTFExporter/MigratedFromGit/Source/GLTFExporter/Private/GLTFExporterModule.cpp
--- a/Engine/Plugins/Enterprise/GLTFExporter/MigratedFromGit/Source/GLTFExporter/Private/GLTFExporterModule.cpp
+++ b/Engine/Plugins/Enterprise/GLTFExporter/MigratedFromGit/Source/GLTFExporter/Private/GLTFExporterModule.cpp
@@ -35,12 +35,14 @@ public:
 		const FString PluginShaderDir = FPaths::Combine(IPluginManager::Get().FindPlugin(TEXT("GLTFExporter"))->GetBaseDir(), TEXT("Shaders"));
 		AddShaderSourceDirectoryMapping(TEXT("/Plugin/GLTFExporter"), PluginShaderDir);
 
+#if WITH_EDITOR
 		FCoreDelegates::OnPostEngineInit.AddRaw(this, &FGLTFExporterModule::PostEngineInit);
+#endif
 	}
 
+#if WITH_EDITOR
 	void PostEngineInit()
 	{
-#if WITH_EDITOR
 		if (FAssetToolsModule::IsModuleLoaded())
 		{
 			IAssetTools& AssetTools = FAssetToolsModule::GetModule().Get();
@@ -56,8 +58,8 @@ public:
 				AssetTools.RegisterAssetTypeActions(AssetTypeActions);
 			}
 		}
-#endif
 	}
+#endif
 
 	virtual void ShutdownModule() override
 	{
