[GLTFExporter] refactor: avoid compiler warnings "multiple copy constructors specified" and "multiple assignment operators specified" regarding GLTFSharedArray

commit a99a9092bbe41fdcabc46b67777743d503c30154
Date: Sun, 13 Feb 2022 01:38:21 +0100

    refactor: avoid compiler warnings "multiple copy constructors specified" and "multiple assignment operators specified" regarding GLTFSharedArray
    (cherry picked from commit 9b4b6acd4f845695306969e4900c7efd4b495163​)

#skipci

diff --git a/Engine/Plugins/Enterprise/GLTFExporter/MigratedFromGit/Source/GLTFExporter/Private/Converters/GLTFMaterialUtility.cpp b/Engine/Plugins/Enterprise/GLTFExporter/MigratedFromGit/Source/GLTFExporter/Private/Converters/GLTFMaterialUtility.cpp
--- a/Engine/Plugins/Enterprise/GLTFExporter/MigratedFromGit/Source/GLTFExporter/Private/Converters/GLTFMaterialUtility.cpp
+++ b/Engine/Plugins/Enterprise/GLTFExporter/MigratedFromGit/Source/GLTFExporter/Private/Converters/GLTFMaterialUtility.cpp
@@ -358,7 +358,7 @@ FGLTFPropertyBakeOutput FGLTFMaterialUtility::BakeMaterialProperty(const FIntPoi
 
 	FBakeOutputEx& BakeOutput = BakeOutputs[0];
 
-	TGLTFSharedArray<FColor> BakedPixels = MoveTemp(BakeOutput.PropertyData.FindChecked(Property));
+	TGLTFSharedArray<FColor> BakedPixels = MakeShared<TArray<FColor>>(MoveTemp(BakeOutput.PropertyData.FindChecked(Property)));
 	const FIntPoint BakedSize = BakeOutput.PropertySizes.FindChecked(Property);
 	const float EmissiveScale = BakeOutput.EmissiveScale;
 
diff --git a/Engine/Plugins/Enterprise/GLTFExporter/MigratedFromGit/Source/GLTFExporter/Private/Converters/GLTFSharedArray.h b/Engine/Plugins/Enterprise/GLTFExporter/MigratedFromGit/Source/GLTFExporter/Private/Converters/GLTFSharedArray.h
--- a/Engine/Plugins/Enterprise/GLTFExporter/MigratedFromGit/Source/GLTFExporter/Private/Converters/GLTFSharedArray.h
+++ b/Engine/Plugins/Enterprise/GLTFExporter/MigratedFromGit/Source/GLTFExporter/Private/Converters/GLTFSharedArray.h
@@ -5,48 +5,17 @@
 #include "CoreMinimal.h"
 
 template <typename ElementType, typename ArrayType = TArray<ElementType>, ESPMode Mode = ESPMode::Fast>
-class TGLTFSharedArray : public TSharedRef<ArrayType>
+class TGLTFSharedArray : public TSharedRef<ArrayType, Mode>
 {
 public:
 
-	template <typename... ArgTypes>
-	TGLTFSharedArray(ArgTypes&&... Args)
-		: TSharedRef(MakeShared<ArrayType>(Forward<ArgTypes>(Args)...))
+	TGLTFSharedArray()
+		: TSharedRef(MakeShared<ArrayType>())
 	{
 	}
 
-	TGLTFSharedArray(TGLTFSharedArray& SharedArray)
-		: TSharedRef(SharedArray)
-	{
-	}
-
-	TGLTFSharedArray(TGLTFSharedArray const& SharedArray)
-		: TSharedRef(SharedArray)
-	{
-	}
-
-	TGLTFSharedArray(TGLTFSharedArray&& SharedArray) noexcept
-		: TSharedRef(SharedArray)
-	{
-	}
-
-	TGLTFSharedArray& operator=(TGLTFSharedArray& Other)
-	{
-		TSharedRef<ArrayType>::operator=(Other);
-		return *this;
-	}
-
-	TGLTFSharedArray& operator=(TGLTFSharedArray const& Other)
-	{
-		TSharedRef<ArrayType>::operator=(Other);
-		return *this;
-	}
-
-	TGLTFSharedArray& operator=(TGLTFSharedArray&& Other) noexcept
-	{
-		TSharedRef<ArrayType>::operator=(Other);
-		return *this;
-	}
+	using TSharedRef<ArrayType>::TSharedRef;
+	using TSharedRef<ArrayType>::operator=;
 
 	bool operator==(const TGLTFSharedArray& Other) const
 	{
@@ -58,9 +27,6 @@ public:
 		return this->Get() != Other.Get();
 	}
 
-	using TSharedRef<ArrayType>::TSharedRef;
-	using TSharedRef<ArrayType>::operator=;
-
 	friend uint32 GetTypeHash(const TGLTFSharedArray& SharedArray)
 	{
 		const ArrayType& Array = SharedArray.Get();
diff --git a/Engine/Plugins/Enterprise/GLTFExporter/MigratedFromGit/Source/GLTFExporter/Private/Tasks/GLTFTextureTasks.cpp b/Engine/Plugins/Enterprise/GLTFExporter/MigratedFromGit/Source/GLTFExporter/Private/Tasks/GLTFTextureTasks.cpp
--- a/Engine/Plugins/Enterprise/GLTFExporter/MigratedFromGit/Source/GLTFExporter/Private/Tasks/GLTFTextureTasks.cpp
+++ b/Engine/Plugins/Enterprise/GLTFExporter/MigratedFromGit/Source/GLTFExporter/Private/Tasks/GLTFTextureTasks.cpp
@@ -246,7 +246,7 @@ void FGLTFTextureLightMapTask::Complete()
 	const EGLTFTextureType Type = EGLTFTextureType::Lightmaps;
 
 	const void* RawData = Source.LockMip(0);
-	TGLTFSharedArray<FColor> Pixels(static_cast<const FColor*>(RawData), ByteLength / sizeof(FColor));
+	TGLTFSharedArray<FColor> Pixels = MakeShared<TArray<FColor>>(static_cast<const FColor*>(RawData), ByteLength / sizeof(FColor));
 	Source.UnlockMip(0);
 
 	JsonTexture.Source = Builder.GetOrAddImage(Pixels, Size, bIgnoreAlpha, Type, JsonTexture.Name);
