[GLTFExporter] feat: minor optimization for loading default mesh material

commit b37edb3a4aa9bde8b92b5622b93419bb137f67d5
Date: Sun, 13 Feb 2022 03:17:22 +0100

    feat: minor optimization for loading default mesh material
    (cherry picked from commit 6446c6ebfd8c8bd94f435a3813ecea58dc92e284​)

#skipci

diff --git a/Engine/Plugins/Enterprise/GLTFExporter/MigratedFromGit/Source/GLTFExporter/Private/Converters/GLTFMeshUtility.cpp b/Engine/Plugins/Enterprise/GLTFExporter/MigratedFromGit/Source/GLTFExporter/Private/Converters/GLTFMeshUtility.cpp
--- a/Engine/Plugins/Enterprise/GLTFExporter/MigratedFromGit/Source/GLTFExporter/Private/Converters/GLTFMeshUtility.cpp
+++ b/Engine/Plugins/Enterprise/GLTFExporter/MigratedFromGit/Source/GLTFExporter/Private/Converters/GLTFMeshUtility.cpp
@@ -5,7 +5,8 @@
 
 const UMaterialInterface* FGLTFMeshUtility::GetDefaultMaterial()
 {
-	return UMaterial::GetDefaultMaterial(MD_Surface);
+	static UMaterial* DefaultMaterial = UMaterial::GetDefaultMaterial(MD_Surface);
+	return DefaultMaterial;
 }
 
 const TArray<FStaticMaterial>& FGLTFMeshUtility::GetMaterials(const UStaticMesh* StaticMesh)
