[GLTFExporter] fix: make sure default preview mesh for exporting materials (EditorSphere) is bundled in runtime

commit 753092d8f68ea4086d1e685fca2d7484b71ec555
Date: Sun, 13 Feb 2022 05:34:19 +0100

    fix: make sure default preview mesh for exporting materials (EditorSphere) is bundled in runtime
    (cherry picked from commit 8f98b00703344e299f221b63f1fb1e59083ee89c​)

#skipci

diff --git a/Engine/Plugins/Enterprise/GLTFExporter/MigratedFromGit/Source/GLTFExporter/Private/Converters/GLTFMaterialProxyFactory.cpp b/Engine/Plugins/Enterprise/GLTFExporter/MigratedFromGit/Source/GLTFExporter/Private/Converters/GLTFMaterialProxyFactory.cpp
--- a/Engine/Plugins/Enterprise/GLTFExporter/MigratedFromGit/Source/GLTFExporter/Private/Converters/GLTFMaterialProxyFactory.cpp
+++ b/Engine/Plugins/Enterprise/GLTFExporter/MigratedFromGit/Source/GLTFExporter/Private/Converters/GLTFMaterialProxyFactory.cpp
@@ -13,6 +13,7 @@ FGLTFMaterialProxyFactory::FGLTFMaterialProxyFactory(const UGLTFProxyOptions* Op
 	: Builder(TEXT(""), CreateExportOptions(Options))
 {
 	Builder.ImageConverter = CreateCustomImageConverter();
+	// TODO: override texture converters so that we don't create new assets if a texture asset already exists
 }
 
 UMaterialInterface* FGLTFMaterialProxyFactory::Create(UMaterialInterface* OriginalMaterial)
diff --git a/Engine/Plugins/Enterprise/GLTFExporter/MigratedFromGit/Source/GLTFExporter/Private/Exporters/GLTFExporterUtility.cpp b/Engine/Plugins/Enterprise/GLTFExporter/MigratedFromGit/Source/GLTFExporter/Private/Exporters/GLTFExporterUtility.cpp
--- a/Engine/Plugins/Enterprise/GLTFExporter/MigratedFromGit/Source/GLTFExporter/Private/Exporters/GLTFExporterUtility.cpp
+++ b/Engine/Plugins/Enterprise/GLTFExporter/MigratedFromGit/Source/GLTFExporter/Private/Exporters/GLTFExporterUtility.cpp
@@ -36,8 +36,7 @@ const UStaticMesh* FGLTFExporterUtility::GetPreviewMesh(const UMaterialInterface
 	} while (Material != nullptr);
 #endif
 
-	static const UStaticMesh* DefaultPreviewMesh = LoadObject<UStaticMesh>(nullptr, TEXT("/Engine/EditorMeshes/EditorSphere.EditorSphere"));
-	return DefaultPreviewMesh;
+	return nullptr;
 }
 
 const USkeletalMesh* FGLTFExporterUtility::GetPreviewMesh(const UAnimSequence* AnimSequence)
diff --git a/Engine/Plugins/Enterprise/GLTFExporter/MigratedFromGit/Source/GLTFExporter/Private/Exporters/GLTFMaterialExporter.cpp b/Engine/Plugins/Enterprise/GLTFExporter/MigratedFromGit/Source/GLTFExporter/Private/Exporters/GLTFMaterialExporter.cpp
--- a/Engine/Plugins/Enterprise/GLTFExporter/MigratedFromGit/Source/GLTFExporter/Private/Exporters/GLTFMaterialExporter.cpp
+++ b/Engine/Plugins/Enterprise/GLTFExporter/MigratedFromGit/Source/GLTFExporter/Private/Exporters/GLTFMaterialExporter.cpp
@@ -4,60 +4,52 @@
 #include "Exporters/GLTFExporterUtility.h"
 #include "Builders/GLTFContainerBuilder.h"
 #include "Materials/MaterialInterface.h"
+#include "UObject/ConstructorHelpers.h"
 
 UGLTFMaterialExporter::UGLTFMaterialExporter(const FObjectInitializer& ObjectInitializer)
 	: Super(ObjectInitializer)
 {
 	SupportedClass = UMaterialInterface::StaticClass();
+
+	static ConstructorHelpers::FObjectFinder<UStaticMesh> DefaultPreviewMeshFinder(TEXT("/Engine/EditorMeshes/EditorSphere.EditorSphere"));
+	DefaultPreviewMesh = DefaultPreviewMeshFinder.Object;
 }
 
 bool UGLTFMaterialExporter::AddObject(FGLTFContainerBuilder& Builder, const UObject* Object)
 {
 	const UMaterialInterface* Material = CastChecked<UMaterialInterface>(Object);
-	const UStaticMesh* PreviewMesh = Builder.ExportOptions->bExportPreviewMesh ? FGLTFExporterUtility::GetPreviewMesh(Material) : nullptr;
 
 	if (Builder.ExportOptions->bExportPreviewMesh)
 	{
-		if (PreviewMesh != nullptr)
+		const UStaticMesh* PreviewMesh = FGLTFExporterUtility::GetPreviewMesh(Material);
+		if (PreviewMesh == nullptr)
 		{
-			const FGLTFJsonMeshIndex MeshIndex = Builder.GetOrAddMesh(PreviewMesh, { Material });
-			if (MeshIndex == INDEX_NONE)
-			{
-				Builder.LogError(
-					FString::Printf(TEXT("Failed to export preview mesh %s for material %s"),
-					*Material->GetName(),
-					*PreviewMesh->GetName()));
-				return false;
-			}
-
-			FGLTFJsonNode Node;
-			Node.Mesh = MeshIndex;
-			const FGLTFJsonNodeIndex NodeIndex = Builder.AddNode(Node);
-
-			FGLTFJsonScene Scene;
-			Scene.Nodes.Add(NodeIndex);
-			const FGLTFJsonSceneIndex SceneIndex = Builder.AddScene(Scene);
-
-			Builder.DefaultScene = SceneIndex;
+			PreviewMesh = DefaultPreviewMesh;
 		}
-		else
+
+		const FGLTFJsonMeshIndex MeshIndex = Builder.GetOrAddMesh(PreviewMesh, { Material });
+		if (MeshIndex == INDEX_NONE)
 		{
 			Builder.LogError(
-				FString::Printf(TEXT("Failed to export material %s because of missing preview mesh"),
-				*Material->GetName()));
+				FString::Printf(TEXT("Failed to export preview mesh %s for material %s"),
+				*Material->GetName(),
+				*PreviewMesh->GetName()));
 			return false;
 		}
+
+		FGLTFJsonNode Node;
+		Node.Mesh = MeshIndex;
+		const FGLTFJsonNodeIndex NodeIndex = Builder.AddNode(Node);
+
+		FGLTFJsonScene Scene;
+		Scene.Nodes.Add(NodeIndex);
+		const FGLTFJsonSceneIndex SceneIndex = Builder.AddScene(Scene);
+
+		Builder.DefaultScene = SceneIndex;
 	}
 	else
 	{
-		const FGLTFJsonMaterialIndex MaterialIndex = Builder.GetOrAddMaterial(Material);
-		if (MaterialIndex == INDEX_NONE)
-		{
-			Builder.LogError(
-				FString::Printf(TEXT("Failed to export material %s"),
-				*Material->GetName()));
-			return false;
-		}
+		Builder.GetOrAddMaterial(Material);
 	}
 
 	return true;
diff --git a/Engine/Plugins/Enterprise/GLTFExporter/MigratedFromGit/Source/GLTFExporter/Public/Exporters/GLTFMaterialExporter.h b/Engine/Plugins/Enterprise/GLTFExporter/MigratedFromGit/Source/GLTFExporter/Public/Exporters/GLTFMaterialExporter.h
--- a/Engine/Plugins/Enterprise/GLTFExporter/MigratedFromGit/Source/GLTFExporter/Public/Exporters/GLTFMaterialExporter.h
+++ b/Engine/Plugins/Enterprise/GLTFExporter/MigratedFromGit/Source/GLTFExporter/Public/Exporters/GLTFMaterialExporter.h
@@ -15,4 +15,9 @@ public:
 	explicit UGLTFMaterialExporter(const FObjectInitializer& ObjectInitializer = FObjectInitializer());
 
 	virtual bool AddObject(FGLTFContainerBuilder& Builder, const UObject* Object) override;
+
+private:
+
+	UPROPERTY(Transient, DuplicateTransient)
+	UStaticMesh* DefaultPreviewMesh;
 };
