[GLTFExporter] fix: compilation errors on macOS regarding GLTFSharedArray

commit b2384e6ec1ea58dfbd4c211b6b1c1cae38434577
Date: Sun, 13 Feb 2022 06:40:09 +0100

    fix: compilation errors on macOS regarding GLTFSharedArray
    (cherry picked from commit 9c3c738d7112cdbc4c4ceebd14cda40d9cc6d28b​)

#skipci

diff --git a/Engine/Plugins/Enterprise/GLTFExporter/MigratedFromGit/Source/GLTFExporter/Private/Converters/GLTFSharedArray.h b/Engine/Plugins/Enterprise/GLTFExporter/MigratedFromGit/Source/GLTFExporter/Private/Converters/GLTFSharedArray.h
--- a/Engine/Plugins/Enterprise/GLTFExporter/MigratedFromGit/Source/GLTFExporter/Private/Converters/GLTFSharedArray.h
+++ b/Engine/Plugins/Enterprise/GLTFExporter/MigratedFromGit/Source/GLTFExporter/Private/Converters/GLTFSharedArray.h
@@ -10,12 +10,18 @@ class TGLTFSharedArray : public TSharedRef<ArrayType, Mode>
 public:
 
 	TGLTFSharedArray()
-		: TSharedRef(MakeShared<ArrayType>())
+		: TSharedRef<ArrayType, Mode>(MakeShared<ArrayType>())
 	{
 	}
 
-	using TSharedRef<ArrayType>::TSharedRef;
-	using TSharedRef<ArrayType>::operator=;
+	template <typename OtherType, typename = decltype(ImplicitConv<ArrayType*>(static_cast<OtherType*>(nullptr)))>
+	TGLTFSharedArray(const TSharedRef<OtherType, Mode>& SharedRef)
+	: TSharedRef<ArrayType, Mode>(SharedRef)
+	{
+	}
+
+	using TSharedRef<ArrayType, Mode>::TSharedRef;
+	using TSharedRef<ArrayType, Mode>::operator=;
 
 	bool operator==(const TGLTFSharedArray& Other) const
 	{
