[GLTFExporter] chore: add utilities for parsing various enums

commit 697f40bb71da15843ecf5862d1a55b0c073c10c1
Date: Tue, 11 Jan 2022 15:14:03 +0100

    chore: add utilities for parsing various enums
    (cherry picked from commit 7bd094f1f3910db76819e3d46949dab5b6eccbf1​)

#skipci

diff --git a/Engine/Plugins/Enterprise/GLTFExporter/MigratedFromGit/Source/GLTFExporter/Private/Json/GLTFJsonUtility.h b/Engine/Plugins/Enterprise/GLTFExporter/MigratedFromGit/Source/GLTFExporter/Private/Json/GLTFJsonUtility.h
--- a/Engine/Plugins/Enterprise/GLTFExporter/MigratedFromGit/Source/GLTFExporter/Private/Json/GLTFJsonUtility.h
+++ b/Engine/Plugins/Enterprise/GLTFExporter/MigratedFromGit/Source/GLTFExporter/Private/Json/GLTFJsonUtility.h
@@ -196,4 +196,81 @@ struct FGLTFJsonUtility
 				return TEXT("");
 		}
 	}
+
+	static EGLTFJsonComponentType ParseComponentType(int32 ComponentType)
+	{
+		switch (ComponentType)
+		{
+		case EGLTFJsonComponentType::S8:
+		case EGLTFJsonComponentType::U8:
+		case EGLTFJsonComponentType::S16:
+		case EGLTFJsonComponentType::U16:
+		case EGLTFJsonComponentType::S32:
+		case EGLTFJsonComponentType::U32:
+		case EGLTFJsonComponentType::F32:
+			return static_cast<EGLTFJsonComponentType>(ComponentType);
+		default:
+			return EGLTFJsonComponentType::None;
+		}
+	}
+
+	static EGLTFJsonAccessorType ParseAccessorType(const FString& AccessorType)
+	{
+		if (AccessorType.Equals(TEXT("SCALAR")))
+		{
+			return EGLTFJsonAccessorType::Scalar;
+		}
+		if (AccessorType.Equals(TEXT("VEC2")))
+		{
+			return EGLTFJsonAccessorType::Vec2;
+		}
+		if (AccessorType.Equals(TEXT("VEC3")))
+		{
+			return EGLTFJsonAccessorType::Vec3;
+		}
+		if (AccessorType.Equals(TEXT("VEC4")))
+		{
+			return EGLTFJsonAccessorType::Vec4;
+		}
+		if (AccessorType.Equals(TEXT("MAT2")))
+		{
+			return EGLTFJsonAccessorType::Mat2;
+		}
+		if (AccessorType.Equals(TEXT("MAT3")))
+		{
+			return EGLTFJsonAccessorType::Mat3;
+		}
+		if (AccessorType.Equals(TEXT("MAT4")))
+		{
+			return EGLTFJsonAccessorType::Mat4;
+		}
+
+		return EGLTFJsonAccessorType::None;
+	}
+
+	static EGLTFJsonBufferTarget ParseBufferTarget(int32 BufferTarget)
+	{
+		switch (BufferTarget)
+		{
+		case EGLTFJsonBufferTarget::ArrayBuffer:
+		case EGLTFJsonBufferTarget::ElementArrayBuffer:
+			return static_cast<EGLTFJsonBufferTarget>(BufferTarget);
+		default:
+			return EGLTFJsonBufferTarget::None;
+		}
+	}
+
+	static EGLTFJsonMimeType ParseMimeType(const FString& MimeType)
+	{
+		if (MimeType.Equals(TEXT("image/jpeg")))
+		{
+			return EGLTFJsonMimeType::JPEG;
+		}
+		if (MimeType.Equals(TEXT("image/png")))
+		{
+			return EGLTFJsonMimeType::PNG;
+		}
+
+		return EGLTFJsonMimeType::None;
+	}
 };
