[GLTFExporter] feat: make Json types for the exporter public

commit b92ed41b85a0a0cbea472af172bc2a29dcfcf150
Date: Mon, 10 Jan 2022 09:11:29 +0100

    feat: make Json types for the exporter public
    (cherry picked from commit a2f08371dcb14416eef0d733b3bda10392b2dac6​)

#skipci

diff --git a/Engine/Plugins/Enterprise/GLTFExporter/MigratedFromGit/Source/GLTFExporter/Private/Json/GLTFJsonAccessor.h b/Engine/Plugins/Enterprise/GLTFExporter/MigratedFromGit/Source/GLTFExporter/Private/Json/GLTFJsonAccessor.h
deleted file mode 100644
diff --git a/Engine/Plugins/Enterprise/GLTFExporter/MigratedFromGit/Source/GLTFExporter/Private/Json/GLTFJsonAnimation.h b/Engine/Plugins/Enterprise/GLTFExporter/MigratedFromGit/Source/GLTFExporter/Private/Json/GLTFJsonAnimation.h
deleted file mode 100644
diff --git a/Engine/Plugins/Enterprise/GLTFExporter/MigratedFromGit/Source/GLTFExporter/Private/Json/GLTFJsonAnimationPlayback.h b/Engine/Plugins/Enterprise/GLTFExporter/MigratedFromGit/Source/GLTFExporter/Private/Json/GLTFJsonAnimationPlayback.h
deleted file mode 100644
diff --git a/Engine/Plugins/Enterprise/GLTFExporter/MigratedFromGit/Source/GLTFExporter/Private/Json/GLTFJsonArray.h b/Engine/Plugins/Enterprise/GLTFExporter/MigratedFromGit/Source/GLTFExporter/Private/Json/GLTFJsonArray.h
deleted file mode 100644
diff --git a/Engine/Plugins/Enterprise/GLTFExporter/MigratedFromGit/Source/GLTFExporter/Private/Json/GLTFJsonAsset.h b/Engine/Plugins/Enterprise/GLTFExporter/MigratedFromGit/Source/GLTFExporter/Private/Json/GLTFJsonAsset.h
deleted file mode 100644
diff --git a/Engine/Plugins/Enterprise/GLTFExporter/MigratedFromGit/Source/GLTFExporter/Private/Json/GLTFJsonBackdrop.h b/Engine/Plugins/Enterprise/GLTFExporter/MigratedFromGit/Source/GLTFExporter/Private/Json/GLTFJsonBackdrop.h
deleted file mode 100644
diff --git a/Engine/Plugins/Enterprise/GLTFExporter/MigratedFromGit/Source/GLTFExporter/Private/Json/GLTFJsonBuffer.h b/Engine/Plugins/Enterprise/GLTFExporter/MigratedFromGit/Source/GLTFExporter/Private/Json/GLTFJsonBuffer.h
deleted file mode 100644
diff --git a/Engine/Plugins/Enterprise/GLTFExporter/MigratedFromGit/Source/GLTFExporter/Private/Json/GLTFJsonBufferView.h b/Engine/Plugins/Enterprise/GLTFExporter/MigratedFromGit/Source/GLTFExporter/Private/Json/GLTFJsonBufferView.h
deleted file mode 100644
diff --git a/Engine/Plugins/Enterprise/GLTFExporter/MigratedFromGit/Source/GLTFExporter/Private/Json/GLTFJsonCamera.h b/Engine/Plugins/Enterprise/GLTFExporter/MigratedFromGit/Source/GLTFExporter/Private/Json/GLTFJsonCamera.h
deleted file mode 100644
diff --git a/Engine/Plugins/Enterprise/GLTFExporter/MigratedFromGit/Source/GLTFExporter/Private/Json/GLTFJsonCameraControl.h b/Engine/Plugins/Enterprise/GLTFExporter/MigratedFromGit/Source/GLTFExporter/Private/Json/GLTFJsonCameraControl.h
deleted file mode 100644
diff --git a/Engine/Plugins/Enterprise/GLTFExporter/MigratedFromGit/Source/GLTFExporter/Private/Json/GLTFJsonColor.cpp b/Engine/Plugins/Enterprise/GLTFExporter/MigratedFromGit/Source/GLTFExporter/Private/Json/GLTFJsonColor.cpp
deleted file mode 100644
diff --git a/Engine/Plugins/Enterprise/GLTFExporter/MigratedFromGit/Source/GLTFExporter/Private/Json/GLTFJsonColor.h b/Engine/Plugins/Enterprise/GLTFExporter/MigratedFromGit/Source/GLTFExporter/Private/Json/GLTFJsonColor.h
deleted file mode 100644
diff --git a/Engine/Plugins/Enterprise/GLTFExporter/MigratedFromGit/Source/GLTFExporter/Private/Json/GLTFJsonEnums.h b/Engine/Plugins/Enterprise/GLTFExporter/MigratedFromGit/Source/GLTFExporter/Private/Json/GLTFJsonEnums.h
deleted file mode 100644
diff --git a/Engine/Plugins/Enterprise/GLTFExporter/MigratedFromGit/Source/GLTFExporter/Private/Json/GLTFJsonEpicLevelVariantSets.h b/Engine/Plugins/Enterprise/GLTFExporter/MigratedFromGit/Source/GLTFExporter/Private/Json/GLTFJsonEpicLevelVariantSets.h
deleted file mode 100644
diff --git a/Engine/Plugins/Enterprise/GLTFExporter/MigratedFromGit/Source/GLTFExporter/Private/Json/GLTFJsonExtensions.h b/Engine/Plugins/Enterprise/GLTFExporter/MigratedFromGit/Source/GLTFExporter/Private/Json/GLTFJsonExtensions.h
deleted file mode 100644
diff --git a/Engine/Plugins/Enterprise/GLTFExporter/MigratedFromGit/Source/GLTFExporter/Private/Json/GLTFJsonHotspot.h b/Engine/Plugins/Enterprise/GLTFExporter/MigratedFromGit/Source/GLTFExporter/Private/Json/GLTFJsonHotspot.h
deleted file mode 100644
diff --git a/Engine/Plugins/Enterprise/GLTFExporter/MigratedFromGit/Source/GLTFExporter/Private/Json/GLTFJsonImage.h b/Engine/Plugins/Enterprise/GLTFExporter/MigratedFromGit/Source/GLTFExporter/Private/Json/GLTFJsonImage.h
deleted file mode 100644
diff --git a/Engine/Plugins/Enterprise/GLTFExporter/MigratedFromGit/Source/GLTFExporter/Private/Json/GLTFJsonIndex.h b/Engine/Plugins/Enterprise/GLTFExporter/MigratedFromGit/Source/GLTFExporter/Private/Json/GLTFJsonIndex.h
deleted file mode 100644
diff --git a/Engine/Plugins/Enterprise/GLTFExporter/MigratedFromGit/Source/GLTFExporter/Private/Json/GLTFJsonKhrMaterialVariant.h b/Engine/Plugins/Enterprise/GLTFExporter/MigratedFromGit/Source/GLTFExporter/Private/Json/GLTFJsonKhrMaterialVariant.h
deleted file mode 100644
diff --git a/Engine/Plugins/Enterprise/GLTFExporter/MigratedFromGit/Source/GLTFExporter/Private/Json/GLTFJsonLight.h b/Engine/Plugins/Enterprise/GLTFExporter/MigratedFromGit/Source/GLTFExporter/Private/Json/GLTFJsonLight.h
deleted file mode 100644
diff --git a/Engine/Plugins/Enterprise/GLTFExporter/MigratedFromGit/Source/GLTFExporter/Private/Json/GLTFJsonLightMap.h b/Engine/Plugins/Enterprise/GLTFExporter/MigratedFromGit/Source/GLTFExporter/Private/Json/GLTFJsonLightMap.h
deleted file mode 100644
diff --git a/Engine/Plugins/Enterprise/GLTFExporter/MigratedFromGit/Source/GLTFExporter/Private/Json/GLTFJsonMaterial.h b/Engine/Plugins/Enterprise/GLTFExporter/MigratedFromGit/Source/GLTFExporter/Private/Json/GLTFJsonMaterial.h
deleted file mode 100644
diff --git a/Engine/Plugins/Enterprise/GLTFExporter/MigratedFromGit/Source/GLTFExporter/Private/Json/GLTFJsonMatrix.cpp b/Engine/Plugins/Enterprise/GLTFExporter/MigratedFromGit/Source/GLTFExporter/Private/Json/GLTFJsonMatrix.cpp
deleted file mode 100644
diff --git a/Engine/Plugins/Enterprise/GLTFExporter/MigratedFromGit/Source/GLTFExporter/Private/Json/GLTFJsonMatrix.h b/Engine/Plugins/Enterprise/GLTFExporter/MigratedFromGit/Source/GLTFExporter/Private/Json/GLTFJsonMatrix.h
deleted file mode 100644
diff --git a/Engine/Plugins/Enterprise/GLTFExporter/MigratedFromGit/Source/GLTFExporter/Private/Json/GLTFJsonMesh.h b/Engine/Plugins/Enterprise/GLTFExporter/MigratedFromGit/Source/GLTFExporter/Private/Json/GLTFJsonMesh.h
deleted file mode 100644
diff --git a/Engine/Plugins/Enterprise/GLTFExporter/MigratedFromGit/Source/GLTFExporter/Private/Json/GLTFJsonNode.h b/Engine/Plugins/Enterprise/GLTFExporter/MigratedFromGit/Source/GLTFExporter/Private/Json/GLTFJsonNode.h
deleted file mode 100644
diff --git a/Engine/Plugins/Enterprise/GLTFExporter/MigratedFromGit/Source/GLTFExporter/Private/Json/GLTFJsonObject.h b/Engine/Plugins/Enterprise/GLTFExporter/MigratedFromGit/Source/GLTFExporter/Private/Json/GLTFJsonObject.h
deleted file mode 100644
diff --git a/Engine/Plugins/Enterprise/GLTFExporter/MigratedFromGit/Source/GLTFExporter/Private/Json/GLTFJsonQuaternion.cpp b/Engine/Plugins/Enterprise/GLTFExporter/MigratedFromGit/Source/GLTFExporter/Private/Json/GLTFJsonQuaternion.cpp
deleted file mode 100644
diff --git a/Engine/Plugins/Enterprise/GLTFExporter/MigratedFromGit/Source/GLTFExporter/Private/Json/GLTFJsonQuaternion.h b/Engine/Plugins/Enterprise/GLTFExporter/MigratedFromGit/Source/GLTFExporter/Private/Json/GLTFJsonQuaternion.h
deleted file mode 100644
diff --git a/Engine/Plugins/Enterprise/GLTFExporter/MigratedFromGit/Source/GLTFExporter/Private/Json/GLTFJsonRoot.h b/Engine/Plugins/Enterprise/GLTFExporter/MigratedFromGit/Source/GLTFExporter/Private/Json/GLTFJsonRoot.h
deleted file mode 100644
diff --git a/Engine/Plugins/Enterprise/GLTFExporter/MigratedFromGit/Source/GLTFExporter/Private/Json/GLTFJsonSampler.h b/Engine/Plugins/Enterprise/GLTFExporter/MigratedFromGit/Source/GLTFExporter/Private/Json/GLTFJsonSampler.h
deleted file mode 100644
diff --git a/Engine/Plugins/Enterprise/GLTFExporter/MigratedFromGit/Source/GLTFExporter/Private/Json/GLTFJsonScene.h b/Engine/Plugins/Enterprise/GLTFExporter/MigratedFromGit/Source/GLTFExporter/Private/Json/GLTFJsonScene.h
deleted file mode 100644
diff --git a/Engine/Plugins/Enterprise/GLTFExporter/MigratedFromGit/Source/GLTFExporter/Private/Json/GLTFJsonSkin.h b/Engine/Plugins/Enterprise/GLTFExporter/MigratedFromGit/Source/GLTFExporter/Private/Json/GLTFJsonSkin.h
deleted file mode 100644
diff --git a/Engine/Plugins/Enterprise/GLTFExporter/MigratedFromGit/Source/GLTFExporter/Private/Json/GLTFJsonSkySphere.h b/Engine/Plugins/Enterprise/GLTFExporter/MigratedFromGit/Source/GLTFExporter/Private/Json/GLTFJsonSkySphere.h
deleted file mode 100644
diff --git a/Engine/Plugins/Enterprise/GLTFExporter/MigratedFromGit/Source/GLTFExporter/Private/Json/GLTFJsonTexture.h b/Engine/Plugins/Enterprise/GLTFExporter/MigratedFromGit/Source/GLTFExporter/Private/Json/GLTFJsonTexture.h
deleted file mode 100644
diff --git a/Engine/Plugins/Enterprise/GLTFExporter/MigratedFromGit/Source/GLTFExporter/Private/Json/GLTFJsonTextureTransform.h b/Engine/Plugins/Enterprise/GLTFExporter/MigratedFromGit/Source/GLTFExporter/Private/Json/GLTFJsonTextureTransform.h
deleted file mode 100644
diff --git a/Engine/Plugins/Enterprise/GLTFExporter/MigratedFromGit/Source/GLTFExporter/Private/Json/GLTFJsonUtility.h b/Engine/Plugins/Enterprise/GLTFExporter/MigratedFromGit/Source/GLTFExporter/Private/Json/GLTFJsonUtility.h
deleted file mode 100644
diff --git a/Engine/Plugins/Enterprise/GLTFExporter/MigratedFromGit/Source/GLTFExporter/Private/Json/GLTFJsonValue.h b/Engine/Plugins/Enterprise/GLTFExporter/MigratedFromGit/Source/GLTFExporter/Private/Json/GLTFJsonValue.h
deleted file mode 100644
diff --git a/Engine/Plugins/Enterprise/GLTFExporter/MigratedFromGit/Source/GLTFExporter/Private/Json/GLTFJsonVector.cpp b/Engine/Plugins/Enterprise/GLTFExporter/MigratedFromGit/Source/GLTFExporter/Private/Json/GLTFJsonVector.cpp
deleted file mode 100644
diff --git a/Engine/Plugins/Enterprise/GLTFExporter/MigratedFromGit/Source/GLTFExporter/Private/Json/GLTFJsonVector.h b/Engine/Plugins/Enterprise/GLTFExporter/MigratedFromGit/Source/GLTFExporter/Private/Json/GLTFJsonVector.h
deleted file mode 100644
diff --git a/Engine/Plugins/Enterprise/GLTFExporter/MigratedFromGit/Source/GLTFExporter/Private/Json/GLTFJsonWriter.cpp b/Engine/Plugins/Enterprise/GLTFExporter/MigratedFromGit/Source/GLTFExporter/Private/Json/GLTFJsonWriter.cpp
deleted file mode 100644
diff --git a/Engine/Plugins/Enterprise/GLTFExporter/MigratedFromGit/Source/GLTFExporter/Private/Json/GLTFJsonWriter.h b/Engine/Plugins/Enterprise/GLTFExporter/MigratedFromGit/Source/GLTFExporter/Private/Json/GLTFJsonWriter.h
deleted file mode 100644
diff --git a/Engine/Plugins/Enterprise/GLTFExporter/MigratedFromGit/Source/GLTFExporter/Public/Json/GLTFJsonAccessor.h b/Engine/Plugins/Enterprise/GLTFExporter/MigratedFromGit/Source/GLTFExporter/Public/Json/GLTFJsonAccessor.h
new file mode 100644
--- /dev/null
+++ b/Engine/Plugins/Enterprise/GLTFExporter/MigratedFromGit/Source/GLTFExporter/Public/Json/GLTFJsonAccessor.h
@@ -0,0 +1,65 @@
+// Copyright Epic Games, Inc. All Rights Reserved.
+
+#pragma once
+
+#include "Json/GLTFJsonObject.h"
+#include "Json/GLTFJsonEnums.h"
+#include "Json/GLTFJsonIndex.h"
+
+struct FGLTFJsonAccessor : IGLTFJsonObject
+{
+	FString Name;
+
+	FGLTFJsonBufferViewIndex BufferView;
+	int64                    ByteOffset;
+	int32                    Count;
+	EGLTFJsonAccessorType    Type;
+	EGLTFJsonComponentType   ComponentType;
+	bool                     bNormalized;
+
+	int32 MinMaxLength;
+	float Min[16];
+	float Max[16];
+
+	FGLTFJsonAccessor()
+		: ByteOffset(0)
+		, Count(0)
+		, Type(EGLTFJsonAccessorType::None)
+		, ComponentType(EGLTFJsonComponentType::None)
+		, bNormalized(false)
+		, MinMaxLength(0)
+		, Min{0}
+		, Max{0}
+	{
+	}
+
+	virtual void WriteObject(IGLTFJsonWriter& Writer) const override
+	{
+		if (!Name.IsEmpty())
+		{
+			Writer.Write(TEXT("name"), Name);
+		}
+
+		Writer.Write(TEXT("bufferView"), BufferView);
+
+		if (ByteOffset != 0)
+		{
+			Writer.Write(TEXT("byteOffset"), ByteOffset);
+		}
+
+		Writer.Write(TEXT("count"), Count);
+		Writer.Write(TEXT("type"), Type);
+		Writer.Write(TEXT("componentType"), ComponentType);
+
+		if (bNormalized)
+		{
+			Writer.Write(TEXT("normalized"), bNormalized);
+		}
+
+		if (MinMaxLength > 0)
+		{
+			Writer.Write(TEXT("min"), Min, MinMaxLength);
+			Writer.Write(TEXT("max"), Max, MinMaxLength);
+		}
+	}
+};
diff --git a/Engine/Plugins/Enterprise/GLTFExporter/MigratedFromGit/Source/GLTFExporter/Public/Json/GLTFJsonAnimation.h b/Engine/Plugins/Enterprise/GLTFExporter/MigratedFromGit/Source/GLTFExporter/Public/Json/GLTFJsonAnimation.h
new file mode 100644
--- /dev/null
+++ b/Engine/Plugins/Enterprise/GLTFExporter/MigratedFromGit/Source/GLTFExporter/Public/Json/GLTFJsonAnimation.h
@@ -0,0 +1,90 @@
+// Copyright Epic Games, Inc. All Rights Reserved.
+
+#pragma once
+
+#include "Json/GLTFJsonObject.h"
+#include "Json/GLTFJsonEnums.h"
+#include "Json/GLTFJsonIndex.h"
+#include "Json/GLTFJsonAnimationPlayback.h"
+
+struct FGLTFJsonAnimationChannelTarget : IGLTFJsonObject
+{
+	FGLTFJsonNodeIndex Node;
+	EGLTFJsonTargetPath Path;
+
+	FGLTFJsonAnimationChannelTarget()
+		: Path(EGLTFJsonTargetPath::None)
+	{
+	}
+
+	virtual void WriteObject(IGLTFJsonWriter& Writer) const override
+	{
+		Writer.Write(TEXT("node"), Node);
+		Writer.Write(TEXT("path"), Path);
+	}
+};
+
+struct FGLTFJsonAnimationChannel : IGLTFJsonObject
+{
+	FGLTFJsonAnimationSamplerIndex Sampler;
+	FGLTFJsonAnimationChannelTarget Target;
+
+	virtual void WriteObject(IGLTFJsonWriter& Writer) const override
+	{
+		Writer.Write(TEXT("sampler"), Sampler);
+
+		Writer.Write(TEXT("target"), Target);
+	}
+};
+
+struct FGLTFJsonAnimationSampler : IGLTFJsonObject
+{
+	FGLTFJsonAccessorIndex Input;
+	FGLTFJsonAccessorIndex Output;
+
+	EGLTFJsonInterpolation Interpolation;
+
+	FGLTFJsonAnimationSampler()
+		: Interpolation(EGLTFJsonInterpolation::Linear)
+	{
+	}
+
+	virtual void WriteObject(IGLTFJsonWriter& Writer) const override
+	{
+		Writer.Write(TEXT("input"), Input);
+		Writer.Write(TEXT("output"), Output);
+
+		if (Interpolation != EGLTFJsonInterpolation::Linear)
+		{
+			Writer.Write(TEXT("interpolation"), Interpolation);
+		}
+	}
+};
+
+struct FGLTFJsonAnimation : IGLTFJsonObject
+{
+	FString Name;
+
+	TArray<FGLTFJsonAnimationChannel> Channels;
+	TArray<FGLTFJsonAnimationSampler> Samplers;
+
+	FGLTFJsonAnimationPlayback Playback;
+
+	virtual void WriteObject(IGLTFJsonWriter& Writer) const override
+	{
+		if (!Name.IsEmpty())
+		{
+			Writer.Write(TEXT("name"), Name);
+		}
+
+		Writer.Write(TEXT("channels"), Channels);
+		Writer.Write(TEXT("samplers"), Samplers);
+
+		if (Playback != FGLTFJsonAnimationPlayback())
+		{
+			Writer.StartExtensions();
+			Writer.Write(EGLTFJsonExtension::EPIC_AnimationPlayback, Playback);
+			Writer.EndExtensions();
+		}
+	}
+};
diff --git a/Engine/Plugins/Enterprise/GLTFExporter/MigratedFromGit/Source/GLTFExporter/Public/Json/GLTFJsonAnimationPlayback.h b/Engine/Plugins/Enterprise/GLTFExporter/MigratedFromGit/Source/GLTFExporter/Public/Json/GLTFJsonAnimationPlayback.h
new file mode 100644
--- /dev/null
+++ b/Engine/Plugins/Enterprise/GLTFExporter/MigratedFromGit/Source/GLTFExporter/Public/Json/GLTFJsonAnimationPlayback.h
@@ -0,0 +1,65 @@
+// Copyright Epic Games, Inc. All Rights Reserved.
+
+#pragma once
+
+#include "Json/GLTFJsonObject.h"
+
+struct FGLTFJsonAnimationPlayback : IGLTFJsonObject
+{
+	FString Name;
+
+	bool bLoop;
+	bool bAutoPlay;
+
+	float PlayRate;
+	float StartTime;
+
+	FGLTFJsonAnimationPlayback()
+		: bLoop(true)
+		, bAutoPlay(true)
+		, PlayRate(1)
+		, StartTime(0)
+	{
+	}
+
+	virtual void WriteObject(IGLTFJsonWriter& Writer) const override
+	{
+		if (!Name.IsEmpty())
+		{
+			Writer.Write(TEXT("name"), Name);
+		}
+
+		if (bLoop != true)
+		{
+			Writer.Write(TEXT("loop"), bLoop);
+		}
+
+		if (bAutoPlay != true)
+		{
+			Writer.Write(TEXT("autoPlay"), bAutoPlay);
+		}
+
+		if (!FMath::IsNearlyEqual(PlayRate, 1, Writer.DefaultTolerance))
+		{
+			Writer.Write(TEXT("playRate"), PlayRate);
+		}
+
+		if (!FMath::IsNearlyEqual(StartTime, 0, Writer.DefaultTolerance))
+		{
+			Writer.Write(TEXT("startTime"), StartTime);
+		}
+	}
+
+	bool operator==(const FGLTFJsonAnimationPlayback& Other) const
+	{
+		return bLoop == Other.bLoop
+			&& bAutoPlay == Other.bAutoPlay
+			&& PlayRate == Other.PlayRate
+			&& StartTime == Other.StartTime;
+	}
+
+	bool operator!=(const FGLTFJsonAnimationPlayback& Other) const
+	{
+		return !(*this == Other);
+	}
+};
diff --git a/Engine/Plugins/Enterprise/GLTFExporter/MigratedFromGit/Source/GLTFExporter/Public/Json/GLTFJsonArray.h b/Engine/Plugins/Enterprise/GLTFExporter/MigratedFromGit/Source/GLTFExporter/Public/Json/GLTFJsonArray.h
new file mode 100644
--- /dev/null
+++ b/Engine/Plugins/Enterprise/GLTFExporter/MigratedFromGit/Source/GLTFExporter/Public/Json/GLTFJsonArray.h
@@ -0,0 +1,18 @@
+// Copyright Epic Games, Inc. All Rights Reserved.
+
+#pragma once
+
+#include "Json/GLTFJsonValue.h"
+#include "Json/GLTFJsonWriter.h"
+
+struct IGLTFJsonArray : IGLTFJsonValue
+{
+	virtual void WriteValue(IGLTFJsonWriter& Writer) const override final
+	{
+		Writer.StartArray();
+		WriteArray(Writer);
+		Writer.EndArray();
+	}
+
+	virtual void WriteArray(IGLTFJsonWriter& Writer) const = 0;
+};
diff --git a/Engine/Plugins/Enterprise/GLTFExporter/MigratedFromGit/Source/GLTFExporter/Public/Json/GLTFJsonAsset.h b/Engine/Plugins/Enterprise/GLTFExporter/MigratedFromGit/Source/GLTFExporter/Public/Json/GLTFJsonAsset.h
new file mode 100644
--- /dev/null
+++ b/Engine/Plugins/Enterprise/GLTFExporter/MigratedFromGit/Source/GLTFExporter/Public/Json/GLTFJsonAsset.h
@@ -0,0 +1,32 @@
+// Copyright Epic Games, Inc. All Rights Reserved.
+
+#pragma once
+
+#include "Json/GLTFJsonObject.h"
+
+struct FGLTFJsonAsset : IGLTFJsonObject
+{
+	FString Version;
+	FString Generator;
+	FString Copyright;
+
+	FGLTFJsonAsset()
+		: Version(TEXT("2.0"))
+	{
+	}
+
+	virtual void WriteObject(IGLTFJsonWriter& Writer) const override
+	{
+		Writer.Write(TEXT("version"), Version);
+
+		if (!Generator.IsEmpty())
+		{
+			Writer.Write(TEXT("generator"), Generator);
+		}
+
+		if (!Copyright.IsEmpty())
+		{
+			Writer.Write(TEXT("copyright"), Copyright);
+		}
+	}
+};
diff --git a/Engine/Plugins/Enterprise/GLTFExporter/MigratedFromGit/Source/GLTFExporter/Public/Json/GLTFJsonBackdrop.h b/Engine/Plugins/Enterprise/GLTFExporter/MigratedFromGit/Source/GLTFExporter/Public/Json/GLTFJsonBackdrop.h
new file mode 100644
--- /dev/null
+++ b/Engine/Plugins/Enterprise/GLTFExporter/MigratedFromGit/Source/GLTFExporter/Public/Json/GLTFJsonBackdrop.h
@@ -0,0 +1,62 @@
+// Copyright Epic Games, Inc. All Rights Reserved.
+
+#pragma once
+
+#include "Json/GLTFJsonObject.h"
+#include "Json/GLTFJsonIndex.h"
+#include "Json/GLTFJsonVector.h"
+
+struct FGLTFJsonBackdrop : IGLTFJsonObject
+{
+	FString Name;
+
+	FGLTFJsonMeshIndex Mesh;
+	FGLTFJsonTextureIndex Cubemap[6];
+
+	float Intensity;
+	float Size;
+	float Angle;
+
+	FGLTFJsonVector3 ProjectionCenter;
+
+	float LightingDistanceFactor;
+	bool UseCameraProjection;
+
+	FGLTFJsonBackdrop()
+		: Intensity(0)
+		, Size(0)
+		, Angle(0)
+		, ProjectionCenter(FGLTFJsonVector3::Zero)
+		, LightingDistanceFactor(0)
+		, UseCameraProjection(false)
+	{
+	}
+
+	virtual void WriteObject(IGLTFJsonWriter& Writer) const override
+	{
+		if (!Name.IsEmpty())
+		{
+			Writer.Write(TEXT("name"), Name);
+		}
+
+		if (Mesh != INDEX_NONE)
+		{
+			Writer.Write(TEXT("mesh"), Mesh);
+		}
+
+		Writer.Write(TEXT("cubemap"), Cubemap);
+
+		Writer.Write(TEXT("intensity"), Intensity);
+		Writer.Write(TEXT("size"), Size);
+
+		if (!FMath::IsNearlyZero(Angle))
+		{
+			Writer.Write(TEXT("angle"), Angle);
+		}
+
+		Writer.Write(TEXT("projectionCenter"), ProjectionCenter);
+
+		Writer.Write(TEXT("lightingDistanceFactor"), LightingDistanceFactor);
+		Writer.Write(TEXT("useCameraProjection"), UseCameraProjection);
+	}
+};
diff --git a/Engine/Plugins/Enterprise/GLTFExporter/MigratedFromGit/Source/GLTFExporter/Public/Json/GLTFJsonBuffer.h b/Engine/Plugins/Enterprise/GLTFExporter/MigratedFromGit/Source/GLTFExporter/Public/Json/GLTFJsonBuffer.h
new file mode 100644
--- /dev/null
+++ b/Engine/Plugins/Enterprise/GLTFExporter/MigratedFromGit/Source/GLTFExporter/Public/Json/GLTFJsonBuffer.h
@@ -0,0 +1,33 @@
+// Copyright Epic Games, Inc. All Rights Reserved.
+
+#pragma once
+
+#include "Json/GLTFJsonObject.h"
+
+struct FGLTFJsonBuffer : IGLTFJsonObject
+{
+	FString Name;
+
+	FString URI;
+	int64   ByteLength;
+
+	FGLTFJsonBuffer()
+		: ByteLength(0)
+	{
+	}
+
+	virtual void WriteObject(IGLTFJsonWriter& Writer) const override
+	{
+		if (!Name.IsEmpty())
+		{
+			Writer.Write(TEXT("name"), Name);
+		}
+
+		if (!URI.IsEmpty())
+		{
+			Writer.Write(TEXT("uri"), URI);
+		}
+
+		Writer.Write(TEXT("byteLength"), ByteLength);
+	}
+};
diff --git a/Engine/Plugins/Enterprise/GLTFExporter/MigratedFromGit/Source/GLTFExporter/Public/Json/GLTFJsonBufferView.h b/Engine/Plugins/Enterprise/GLTFExporter/MigratedFromGit/Source/GLTFExporter/Public/Json/GLTFJsonBufferView.h
new file mode 100644
--- /dev/null
+++ b/Engine/Plugins/Enterprise/GLTFExporter/MigratedFromGit/Source/GLTFExporter/Public/Json/GLTFJsonBufferView.h
@@ -0,0 +1,55 @@
+// Copyright Epic Games, Inc. All Rights Reserved.
+
+#pragma once
+
+#include "Json/GLTFJsonObject.h"
+#include "Json/GLTFJsonEnums.h"
+#include "Json/GLTFJsonIndex.h"
+
+struct FGLTFJsonBufferView : IGLTFJsonObject
+{
+	FString Name;
+
+	FGLTFJsonBufferIndex Buffer;
+
+	int64 ByteLength;
+	int64 ByteOffset;
+	int32 ByteStride;
+
+	EGLTFJsonBufferTarget Target;
+
+	FGLTFJsonBufferView()
+		: ByteLength(0)
+		, ByteOffset(0)
+		, ByteStride(0)
+		, Target(EGLTFJsonBufferTarget::None)
+	{
+		// check that view fits completely inside the buffer
+	}
+
+	virtual void WriteObject(IGLTFJsonWriter& Writer) const override
+	{
+		if (!Name.IsEmpty())
+		{
+			Writer.Write(TEXT("name"), Name);
+		}
+
+		Writer.Write(TEXT("buffer"), Buffer);
+		Writer.Write(TEXT("byteLength"), ByteLength);
+
+		if (ByteOffset != 0)
+		{
+			Writer.Write(TEXT("byteOffset"), ByteOffset);
+		}
+
+		if (ByteStride != 0)
+		{
+			Writer.Write(TEXT("byteStride"), ByteStride);
+		}
+
+		if (Target != EGLTFJsonBufferTarget::None)
+		{
+			Writer.Write(TEXT("target"), Target);
+		}
+	}
+};
diff --git a/Engine/Plugins/Enterprise/GLTFExporter/MigratedFromGit/Source/GLTFExporter/Public/Json/GLTFJsonCamera.h b/Engine/Plugins/Enterprise/GLTFExporter/MigratedFromGit/Source/GLTFExporter/Public/Json/GLTFJsonCamera.h
new file mode 100644
--- /dev/null
+++ b/Engine/Plugins/Enterprise/GLTFExporter/MigratedFromGit/Source/GLTFExporter/Public/Json/GLTFJsonCamera.h
@@ -0,0 +1,111 @@
+// Copyright Epic Games, Inc. All Rights Reserved.
+
+#pragma once
+
+#include "Json/GLTFJsonObject.h"
+#include "Json/GLTFJsonEnums.h"
+#include "Json/GLTFJsonCameraControl.h"
+
+struct FGLTFJsonOrthographic : IGLTFJsonObject
+{
+	float XMag; // horizontal magnification of the view
+	float YMag; // vertical magnification of the view
+	float ZFar;
+	float ZNear;
+
+	FGLTFJsonOrthographic()
+		: XMag(0)
+		, YMag(0)
+		, ZFar(0)
+		, ZNear(0)
+	{
+	}
+
+	virtual void WriteObject(IGLTFJsonWriter& Writer) const override
+	{
+		Writer.Write(TEXT("xmag"), XMag);
+		Writer.Write(TEXT("ymag"), YMag);
+		Writer.Write(TEXT("zfar"), ZFar);
+		Writer.Write(TEXT("znear"), ZNear);
+	}
+};
+
+struct FGLTFJsonPerspective : IGLTFJsonObject
+{
+	float AspectRatio; // aspect ratio of the field of view
+	float YFov; // vertical field of view in radians
+	float ZFar;
+	float ZNear;
+
+	FGLTFJsonPerspective()
+		: AspectRatio(0)
+		, YFov(0)
+		, ZFar(0)
+		, ZNear(0)
+	{
+	}
+
+	virtual void WriteObject(IGLTFJsonWriter& Writer) const override
+	{
+		if (!FMath::IsNearlyEqual(AspectRatio, 0, Writer.DefaultTolerance))
+		{
+			Writer.Write(TEXT("aspectRatio"), AspectRatio);
+		}
+
+		Writer.Write(TEXT("yfov"), YFov);
+
+		if (!FMath::IsNearlyEqual(ZFar, 0, Writer.DefaultTolerance))
+		{
+			Writer.Write(TEXT("zfar"), ZFar);
+		}
+
+		Writer.Write(TEXT("znear"), ZNear);
+	}
+};
+
+struct FGLTFJsonCamera : IGLTFJsonObject
+{
+	FString Name;
+
+	EGLTFJsonCameraType               Type;
+	TOptional<FGLTFJsonCameraControl> CameraControl;
+
+	FGLTFJsonOrthographic Orthographic;
+	FGLTFJsonPerspective  Perspective;
+
+	FGLTFJsonCamera()
+		: Type(EGLTFJsonCameraType::None)
+	{
+	}
+
+	virtual void WriteObject(IGLTFJsonWriter& Writer) const override
+	{
+		if (!Name.IsEmpty())
+		{
+			Writer.Write(TEXT("name"), Name);
+		}
+
+		Writer.Write(TEXT("type"), Type);
+
+		switch (Type)
+		{
+			case EGLTFJsonCameraType::Orthographic:
+				Writer.Write(TEXT("orthographic"), Orthographic);
+				break;
+
+			case EGLTFJsonCameraType::Perspective:
+				Writer.Write(TEXT("perspective"), Perspective);
+				break;
+
+			default:
+				break;
+		}
+
+		if (CameraControl.IsSet())
+		{
+			Writer.StartExtensions();
+			Writer.Write(EGLTFJsonExtension::EPIC_CameraControls, CameraControl.GetValue());
+			Writer.EndExtensions();
+		}
+	}
+};
diff --git a/Engine/Plugins/Enterprise/GLTFExporter/MigratedFromGit/Source/GLTFExporter/Public/Json/GLTFJsonCameraControl.h b/Engine/Plugins/Enterprise/GLTFExporter/MigratedFromGit/Source/GLTFExporter/Public/Json/GLTFJsonCameraControl.h
new file mode 100644
--- /dev/null
+++ b/Engine/Plugins/Enterprise/GLTFExporter/MigratedFromGit/Source/GLTFExporter/Public/Json/GLTFJsonCameraControl.h
@@ -0,0 +1,76 @@
+// Copyright Epic Games, Inc. All Rights Reserved.
+
+#pragma once
+
+#include "Json/GLTFJsonObject.h"
+#include "Json/GLTFJsonIndex.h"
+#include "Json/GLTFJsonVector.h"
+
+struct FGLTFJsonCameraControl : IGLTFJsonObject
+{
+	EGLTFJsonCameraControlMode Mode;
+	FGLTFJsonNodeIndex Target;
+	float MaxDistance;
+	float MinDistance;
+	float MaxPitch;
+	float MinPitch;
+	float MaxYaw;
+	float MinYaw;
+	float RotationSensitivity;
+	float RotationInertia;
+	float DollySensitivity;
+	float DollyDuration;
+
+	FGLTFJsonCameraControl()
+		: Mode(EGLTFJsonCameraControlMode::FreeLook)
+		, MaxDistance(0)
+		, MinDistance(0)
+		, MaxPitch(90)
+		, MinPitch(-90)
+		, MaxYaw(360)
+		, MinYaw(0)
+		, RotationSensitivity(0)
+		, RotationInertia(0)
+		, DollySensitivity(0)
+		, DollyDuration(0)
+	{
+	}
+
+	virtual void WriteObject(IGLTFJsonWriter& Writer) const override
+	{
+		Writer.Write(TEXT("mode"), Mode);
+
+		if (Target != INDEX_NONE && Mode == EGLTFJsonCameraControlMode::Orbital)
+		{
+			Writer.Write(TEXT("target"), Target);
+		}
+
+		Writer.Write(TEXT("maxDistance"), MaxDistance);
+		Writer.Write(TEXT("minDistance"), MinDistance);
+
+		if (!FMath::IsNearlyEqual(MaxPitch, 90, Writer.DefaultTolerance))
+		{
+			Writer.Write(TEXT("maxPitch"), MaxPitch);
+		}
+
+		if (!FMath::IsNearlyEqual(MinPitch, -90, Writer.DefaultTolerance))
+		{
+			Writer.Write(TEXT("minPitch"), MinPitch);
+		}
+
+		if (!FMath::IsNearlyEqual(MaxYaw, 360, Writer.DefaultTolerance))
+		{
+			Writer.Write(TEXT("maxYaw"), MaxYaw);
+		}
+
+		if (!FMath::IsNearlyEqual(MinYaw, 0, Writer.DefaultTolerance))
+		{
+			Writer.Write(TEXT("minYaw"), MinYaw);
+		}
+
+		Writer.Write(TEXT("rotationSensitivity"), RotationSensitivity);
+		Writer.Write(TEXT("rotationInertia"), RotationInertia);
+		Writer.Write(TEXT("dollySensitivity"), DollySensitivity);
+		Writer.Write(TEXT("dollyDuration"), DollyDuration);
+	}
+};
diff --git a/Engine/Plugins/Enterprise/GLTFExporter/MigratedFromGit/Source/GLTFExporter/Public/Json/GLTFJsonColor.cpp b/Engine/Plugins/Enterprise/GLTFExporter/MigratedFromGit/Source/GLTFExporter/Public/Json/GLTFJsonColor.cpp
new file mode 100644
--- /dev/null
+++ b/Engine/Plugins/Enterprise/GLTFExporter/MigratedFromGit/Source/GLTFExporter/Public/Json/GLTFJsonColor.cpp
@@ -0,0 +1,9 @@
+// Copyright Epic Games, Inc. All Rights Reserved.
+
+#include "Json/GLTFJsonColor.h"
+
+const FGLTFJsonColor3 FGLTFJsonColor3::Black({ 0, 0, 0 });
+const FGLTFJsonColor3 FGLTFJsonColor3::White({ 1, 1, 1 });
+
+const FGLTFJsonColor4 FGLTFJsonColor4::Black({ 0, 0, 0, 1 });
+const FGLTFJsonColor4 FGLTFJsonColor4::White({ 1, 1, 1, 1 });
diff --git a/Engine/Plugins/Enterprise/GLTFExporter/MigratedFromGit/Source/GLTFExporter/Public/Json/GLTFJsonColor.h b/Engine/Plugins/Enterprise/GLTFExporter/MigratedFromGit/Source/GLTFExporter/Public/Json/GLTFJsonColor.h
new file mode 100644
--- /dev/null
+++ b/Engine/Plugins/Enterprise/GLTFExporter/MigratedFromGit/Source/GLTFExporter/Public/Json/GLTFJsonColor.h
@@ -0,0 +1,78 @@
+// Copyright Epic Games, Inc. All Rights Reserved.
+
+#pragma once
+
+#include "Core/GLTFColor.h"
+#include "Json/GLTFJsonArray.h"
+
+template <typename BaseType>
+struct TGLTFJsonColor : BaseType, IGLTFJsonArray
+{
+	TGLTFJsonColor(const BaseType& Other)
+		: BaseType(Other)
+	{
+	}
+
+	TGLTFJsonColor& operator=(const BaseType& Other)
+	{
+		*static_cast<BaseType*>(this) = Other;
+		return *this;
+	}
+
+	bool operator==(const BaseType& Other) const
+	{
+		for (SIZE_T i = 0; i < GetNum(BaseType::Components); ++i)
+		{
+			if (BaseType::Components[i] != Other.Components[i])
+			{
+				return false;
+			}
+		}
+
+		return true;
+	}
+
+	bool operator!=(const BaseType& Other) const
+	{
+		return !(*this == Other);
+	}
+
+	virtual void WriteArray(IGLTFJsonWriter& Writer) const override
+	{
+		for (SIZE_T i = 0; i < GetNum(BaseType::Components); ++i)
+		{
+			Writer.Write(BaseType::Components[i]);
+		}
+	}
+
+	bool IsNearlyEqual(const BaseType& Other, float Tolerance = KINDA_SMALL_NUMBER) const
+	{
+		for (SIZE_T i = 0; i < GetNum(BaseType::Components); ++i)
+		{
+			if (!FMath::IsNearlyEqual(BaseType::Components[i], Other.Components[i], Tolerance))
+			{
+				return false;
+			}
+		}
+
+		return true;
+	}
+};
+
+struct FGLTFJsonColor3 : TGLTFJsonColor<FGLTFColor3>
+{
+	static const FGLTFJsonColor3 Black;
+	static const FGLTFJsonColor3 White;
+
+	using TGLTFJsonColor::TGLTFJsonColor;
+	using TGLTFJsonColor::operator=;
+};
+
+struct FGLTFJsonColor4 : TGLTFJsonColor<FGLTFColor4>
+{
+	static const FGLTFJsonColor4 Black;
+	static const FGLTFJsonColor4 White;
+
+	using TGLTFJsonColor::TGLTFJsonColor;
+	using TGLTFJsonColor::operator=;
+};
diff --git a/Engine/Plugins/Enterprise/GLTFExporter/MigratedFromGit/Source/GLTFExporter/Public/Json/GLTFJsonEnums.h b/Engine/Plugins/Enterprise/GLTFExporter/MigratedFromGit/Source/GLTFExporter/Public/Json/GLTFJsonEnums.h
new file mode 100644
--- /dev/null
+++ b/Engine/Plugins/Enterprise/GLTFExporter/MigratedFromGit/Source/GLTFExporter/Public/Json/GLTFJsonEnums.h
@@ -0,0 +1,178 @@
+// Copyright Epic Games, Inc. All Rights Reserved.
+
+#pragma once
+
+#include "CoreMinimal.h"
+
+enum class EGLTFJsonExtension
+{
+	KHR_LightsPunctual,
+	KHR_MaterialsClearCoat,
+	KHR_MaterialsUnlit,
+	KHR_MaterialsVariants,
+	KHR_MeshQuantization,
+	KHR_TextureTransform,
+	EPIC_AnimationHotspots,
+	EPIC_AnimationPlayback,
+	EPIC_BlendModes,
+	EPIC_CameraControls,
+	EPIC_HDRIBackdrops,
+	EPIC_LevelVariantSets,
+	EPIC_LightmapTextures,
+	EPIC_SkySpheres,
+	EPIC_TextureHDREncoding
+};
+
+enum class EGLTFJsonShadingModel
+{
+	None = -1,
+	Default,
+	Unlit,
+	ClearCoat
+};
+
+enum class EGLTFJsonHDREncoding
+{
+	None = -1,
+	RGBM,
+	RGBE
+};
+
+enum class EGLTFJsonCubeFace
+{
+	None = -1,
+	PosX,
+	NegX,
+	PosY,
+	NegY,
+	PosZ,
+	NegZ
+};
+
+enum class EGLTFJsonAccessorType
+{
+	None = -1,
+	Scalar,
+	Vec2,
+	Vec3,
+	Vec4,
+	Mat2,
+	Mat3,
+	Mat4
+};
+
+enum class EGLTFJsonComponentType
+{
+	None = -1,
+	S8 = 5120, // signed byte
+	U8 = 5121, // unsigned byte
+	S16 = 5122, // signed short
+	U16 = 5123, // unsigned short
+	S32 = 5124, // unused
+	U32 = 5125, // unsigned int -- only valid for indices, not attributes
+	F32 = 5126  // float
+};
+
+enum class EGLTFJsonBufferTarget
+{
+	None = -1,
+	ArrayBuffer = 34962,
+	ElementArrayBuffer = 34963
+};
+
+enum class EGLTFJsonPrimitiveMode
+{
+	// unsupported
+	Points = 0,
+	Lines = 1,
+	LineLoop = 2,
+	LineStrip = 3,
+	// initially supported
+	Triangles = 4,
+	// will be supported prior to release
+	TriangleStrip = 5,
+	TriangleFan = 6
+};
+
+enum class EGLTFJsonAlphaMode
+{
+	None = -1,
+	Opaque,
+	Mask,
+	Blend
+};
+
+enum class EGLTFJsonBlendMode
+{
+	None = -1,
+	Additive,
+	Modulate,
+	AlphaComposite
+};
+
+enum class EGLTFJsonMimeType
+{
+	None = -1,
+	PNG,
+	JPEG
+};
+
+enum class EGLTFJsonTextureFilter
+{
+	None = -1,
+	// valid for Min & Mag
+	Nearest = 9728,
+	Linear = 9729,
+	// valid for Min only
+	NearestMipmapNearest = 9984,
+	LinearMipmapNearest = 9985,
+	NearestMipmapLinear = 9986,
+	LinearMipmapLinear = 9987
+};
+
+enum class EGLTFJsonTextureWrap
+{
+	None = -1,
+	Repeat = 10497,
+	MirroredRepeat = 33648,
+	ClampToEdge = 33071
+};
+
+enum class EGLTFJsonInterpolation
+{
+	None = -1,
+	Linear,
+	Step,
+	CubicSpline,
+};
+
+enum class EGLTFJsonTargetPath
+{
+	None = -1,
+	Translation,
+	Rotation,
+	Scale,
+	Weights
+};
+
+enum class EGLTFJsonCameraType
+{
+	None = -1,
+	Orthographic,
+	Perspective
+};
+
+enum class EGLTFJsonLightType
+{
+	None = -1,
+	Directional,
+	Point,
+	Spot
+};
+
+enum class EGLTFJsonCameraControlMode
+{
+	None = -1,
+	FreeLook,
+	Orbital
+};
diff --git a/Engine/Plugins/Enterprise/GLTFExporter/MigratedFromGit/Source/GLTFExporter/Public/Json/GLTFJsonEpicLevelVariantSets.h b/Engine/Plugins/Enterprise/GLTFExporter/MigratedFromGit/Source/GLTFExporter/Public/Json/GLTFJsonEpicLevelVariantSets.h
new file mode 100644
--- /dev/null
+++ b/Engine/Plugins/Enterprise/GLTFExporter/MigratedFromGit/Source/GLTFExporter/Public/Json/GLTFJsonEpicLevelVariantSets.h
@@ -0,0 +1,123 @@
+// Copyright Epic Games, Inc. All Rights Reserved.
+
+#pragma once
+
+#include "Json/GLTFJsonObject.h"
+#include "Json/GLTFJsonIndex.h"
+#include "Misc/Optional.h"
+
+struct FGLTFJsonEpicVariantMaterial : IGLTFJsonObject
+{
+	FGLTFJsonMaterialIndex Material;
+	int32                  Index;
+
+	FGLTFJsonEpicVariantMaterial()
+		: Material(INDEX_NONE)
+		, Index(INDEX_NONE)
+	{
+	}
+
+	virtual void WriteObject(IGLTFJsonWriter& Writer) const override
+	{
+		Writer.Write(TEXT("material"), Material);
+
+		if (Index != INDEX_NONE)
+		{
+			Writer.Write(TEXT("index"), Index);
+		}
+	}
+};
+
+struct FGLTFJsonEpicVariantNodeProperties : IGLTFJsonObject
+{
+	FGLTFJsonNodeIndex Node;
+	TOptional<bool>    bIsVisible;
+
+	TOptional<FGLTFJsonMeshIndex>        Mesh;
+	TArray<FGLTFJsonEpicVariantMaterial> Materials;
+
+	virtual void WriteObject(IGLTFJsonWriter& Writer) const override
+	{
+		if (Node != INDEX_NONE)
+		{
+			Writer.Write(TEXT("node"), Node);
+		}
+
+		Writer.StartObject(TEXT("properties"));
+
+		if (bIsVisible.IsSet())
+		{
+			Writer.Write(TEXT("visible"), bIsVisible.GetValue());
+		}
+
+		if (Mesh.IsSet())
+		{
+			Writer.Write(TEXT("mesh"), Mesh.GetValue());
+		}
+
+		if (Materials.Num() > 0)
+		{
+			Writer.Write(TEXT("materials"), Materials);
+		}
+
+		Writer.EndObject();
+	}
+};
+
+struct FGLTFJsonEpicVariant : IGLTFJsonObject
+{
+	FString Name;
+	bool    bIsActive;
+
+	FGLTFJsonTextureIndex Thumbnail;
+	TMap<FGLTFJsonNodeIndex, FGLTFJsonEpicVariantNodeProperties> Nodes;
+
+	virtual void WriteObject(IGLTFJsonWriter& Writer) const override
+	{
+		Writer.Write(TEXT("name"), Name);
+		Writer.Write(TEXT("active"), bIsActive);
+
+		if (Thumbnail != INDEX_NONE)
+		{
+			Writer.Write(TEXT("thumbnail"), Thumbnail);
+		}
+
+		TArray<FGLTFJsonEpicVariantNodeProperties> NodeValues;
+		Nodes.GenerateValueArray(NodeValues);
+		Writer.Write(TEXT("nodes"), NodeValues);
+	}
+};
+
+struct FGLTFJsonEpicVariantSet : IGLTFJsonObject
+{
+	FString Name;
+
+	TArray<FGLTFJsonEpicVariant> Variants;
+
+	virtual void WriteObject(IGLTFJsonWriter& Writer) const override
+	{
+		if (!Name.IsEmpty())
+		{
+			Writer.Write(TEXT("name"), Name);
+		}
+
+		Writer.Write(TEXT("variants"), Variants);
+	}
+};
+
+struct FGLTFJsonEpicLevelVariantSets : IGLTFJsonObject
+{
+	FString Name;
+
+	TArray<FGLTFJsonEpicVariantSet> VariantSets;
+
+	virtual void WriteObject(IGLTFJsonWriter& Writer) const override
+	{
+		if (!Name.IsEmpty())
+		{
+			Writer.Write(TEXT("name"), Name);
+		}
+
+		Writer.Write(TEXT("variantSets"), VariantSets);
+	}
+};
diff --git a/Engine/Plugins/Enterprise/GLTFExporter/MigratedFromGit/Source/GLTFExporter/Public/Json/GLTFJsonExtensions.h b/Engine/Plugins/Enterprise/GLTFExporter/MigratedFromGit/Source/GLTFExporter/Public/Json/GLTFJsonExtensions.h
new file mode 100644
--- /dev/null
+++ b/Engine/Plugins/Enterprise/GLTFExporter/MigratedFromGit/Source/GLTFExporter/Public/Json/GLTFJsonExtensions.h
@@ -0,0 +1,11 @@
+// Copyright Epic Games, Inc. All Rights Reserved.
+
+#pragma once
+
+#include "Json/GLTFJsonEnums.h"
+
+struct FGLTFJsonExtensions
+{
+	TSet<EGLTFJsonExtension> Used;
+	TSet<EGLTFJsonExtension> Required;
+};
diff --git a/Engine/Plugins/Enterprise/GLTFExporter/MigratedFromGit/Source/GLTFExporter/Public/Json/GLTFJsonHotspot.h b/Engine/Plugins/Enterprise/GLTFExporter/MigratedFromGit/Source/GLTFExporter/Public/Json/GLTFJsonHotspot.h
new file mode 100644
--- /dev/null
+++ b/Engine/Plugins/Enterprise/GLTFExporter/MigratedFromGit/Source/GLTFExporter/Public/Json/GLTFJsonHotspot.h
@@ -0,0 +1,43 @@
+// Copyright Epic Games, Inc. All Rights Reserved.
+
+#pragma once
+
+#include "Json/GLTFJsonObject.h"
+#include "Json/GLTFJsonIndex.h"
+
+struct FGLTFJsonHotspot : IGLTFJsonObject
+{
+	FString Name;
+
+	FGLTFJsonAnimationIndex Animation;
+	FGLTFJsonTextureIndex   Image;
+	FGLTFJsonTextureIndex   HoveredImage;
+	FGLTFJsonTextureIndex   ToggledImage;
+	FGLTFJsonTextureIndex   ToggledHoveredImage;
+
+	virtual void WriteObject(IGLTFJsonWriter& Writer) const override
+	{
+		if (!Name.IsEmpty())
+		{
+			Writer.Write(TEXT("name"), Name);
+		}
+
+		Writer.Write(TEXT("animation"), Animation);
+		Writer.Write(TEXT("image"), Image);
+
+		if (HoveredImage != INDEX_NONE)
+		{
+			Writer.Write(TEXT("hoveredImage"), HoveredImage);
+		}
+
+		if (ToggledImage != INDEX_NONE)
+		{
+			Writer.Write(TEXT("toggledImage"), ToggledImage);
+		}
+
+		if (ToggledHoveredImage != INDEX_NONE)
+		{
+			Writer.Write(TEXT("toggledHoveredImage"), ToggledHoveredImage);
+		}
+	}
+};
diff --git a/Engine/Plugins/Enterprise/GLTFExporter/MigratedFromGit/Source/GLTFExporter/Public/Json/GLTFJsonImage.h b/Engine/Plugins/Enterprise/GLTFExporter/MigratedFromGit/Source/GLTFExporter/Public/Json/GLTFJsonImage.h
new file mode 100644
--- /dev/null
+++ b/Engine/Plugins/Enterprise/GLTFExporter/MigratedFromGit/Source/GLTFExporter/Public/Json/GLTFJsonImage.h
@@ -0,0 +1,45 @@
+// Copyright Epic Games, Inc. All Rights Reserved.
+
+#pragma once
+
+#include "Json/GLTFJsonObject.h"
+#include "Json/GLTFJsonEnums.h"
+#include "Json/GLTFJsonIndex.h"
+
+struct FGLTFJsonImage : IGLTFJsonObject
+{
+	FString Name;
+	FString Uri;
+
+	EGLTFJsonMimeType MimeType;
+
+	FGLTFJsonBufferViewIndex BufferView;
+
+	FGLTFJsonImage()
+		: MimeType(EGLTFJsonMimeType::None)
+	{
+	}
+
+	virtual void WriteObject(IGLTFJsonWriter& Writer) const override
+	{
+		if (!Name.IsEmpty())
+		{
+			Writer.Write(TEXT("name"), Name);
+		}
+
+		if (!Uri.IsEmpty())
+		{
+			Writer.Write(TEXT("uri"), Uri);
+		}
+
+		if (MimeType != EGLTFJsonMimeType::None)
+		{
+			Writer.Write(TEXT("mimeType"), MimeType);
+		}
+
+		if (BufferView != INDEX_NONE)
+		{
+			Writer.Write(TEXT("bufferView"), BufferView);
+		}
+	}
+};
diff --git a/Engine/Plugins/Enterprise/GLTFExporter/MigratedFromGit/Source/GLTFExporter/Public/Json/GLTFJsonIndex.h b/Engine/Plugins/Enterprise/GLTFExporter/MigratedFromGit/Source/GLTFExporter/Public/Json/GLTFJsonIndex.h
new file mode 100644
--- /dev/null
+++ b/Engine/Plugins/Enterprise/GLTFExporter/MigratedFromGit/Source/GLTFExporter/Public/Json/GLTFJsonIndex.h
@@ -0,0 +1,75 @@
+// Copyright Epic Games, Inc. All Rights Reserved.
+
+#pragma once
+
+#include "CoreMinimal.h"
+
+template <typename DerivedType>
+struct FGLTFJsonIndex
+{
+	FGLTFJsonIndex()
+		: Value(INDEX_NONE)
+	{
+	}
+
+	explicit FGLTFJsonIndex(int32 Value)
+		: Value(Value)
+	{
+	}
+
+	FORCEINLINE operator int32() const
+	{
+		return Value;
+	}
+
+	FORCEINLINE bool operator==(int32 OtherValue) const
+	{
+		return Value == OtherValue;
+	}
+
+	FORCEINLINE bool operator!=(int32 OtherValue) const
+	{
+		return Value != OtherValue;
+	}
+
+	FORCEINLINE bool operator==(const DerivedType& Other) const
+	{
+		return Value == Other.Value;
+	}
+
+	FORCEINLINE bool operator!=(const DerivedType& Other) const
+	{
+		return Value != Other.Value;
+	}
+
+	FORCEINLINE friend uint32 GetTypeHash(const DerivedType& Other)
+	{
+		return GetTypeHash(Other.Value);
+	}
+
+private:
+
+	int32 Value;
+};
+
+struct FGLTFJsonAccessorIndex : FGLTFJsonIndex<FGLTFJsonAccessorIndex> { using FGLTFJsonIndex::FGLTFJsonIndex; };
+struct FGLTFJsonAnimationIndex : FGLTFJsonIndex<FGLTFJsonAnimationIndex> { using FGLTFJsonIndex::FGLTFJsonIndex; };
+struct FGLTFJsonAnimationSamplerIndex : FGLTFJsonIndex<FGLTFJsonAnimationSamplerIndex> { using FGLTFJsonIndex::FGLTFJsonIndex; };
+struct FGLTFJsonBufferIndex : FGLTFJsonIndex<FGLTFJsonBufferIndex> { using FGLTFJsonIndex::FGLTFJsonIndex; };
+struct FGLTFJsonBufferViewIndex : FGLTFJsonIndex<FGLTFJsonBufferViewIndex> { using FGLTFJsonIndex::FGLTFJsonIndex; };
+struct FGLTFJsonCameraIndex : FGLTFJsonIndex<FGLTFJsonCameraIndex> { using FGLTFJsonIndex::FGLTFJsonIndex; };
+struct FGLTFJsonImageIndex : FGLTFJsonIndex<FGLTFJsonImageIndex> { using FGLTFJsonIndex::FGLTFJsonIndex; };
+struct FGLTFJsonMaterialIndex : FGLTFJsonIndex<FGLTFJsonMaterialIndex> { using FGLTFJsonIndex::FGLTFJsonIndex; };
+struct FGLTFJsonMeshIndex : FGLTFJsonIndex<FGLTFJsonMeshIndex> { using FGLTFJsonIndex::FGLTFJsonIndex; };
+struct FGLTFJsonNodeIndex : FGLTFJsonIndex<FGLTFJsonNodeIndex> { using FGLTFJsonIndex::FGLTFJsonIndex; };
+struct FGLTFJsonSamplerIndex : FGLTFJsonIndex<FGLTFJsonSamplerIndex> { using FGLTFJsonIndex::FGLTFJsonIndex; };
+struct FGLTFJsonSceneIndex : FGLTFJsonIndex<FGLTFJsonSceneIndex> { using FGLTFJsonIndex::FGLTFJsonIndex; };
+struct FGLTFJsonSkinIndex : FGLTFJsonIndex<FGLTFJsonSkinIndex> { using FGLTFJsonIndex::FGLTFJsonIndex; };
+struct FGLTFJsonTextureIndex : FGLTFJsonIndex<FGLTFJsonTextureIndex> { using FGLTFJsonIndex::FGLTFJsonIndex; };
+struct FGLTFJsonBackdropIndex : FGLTFJsonIndex<FGLTFJsonBackdropIndex> { using FGLTFJsonIndex::FGLTFJsonIndex; };
+struct FGLTFJsonHotspotIndex : FGLTFJsonIndex<FGLTFJsonHotspotIndex> { using FGLTFJsonIndex::FGLTFJsonIndex; };
+struct FGLTFJsonLightIndex : FGLTFJsonIndex<FGLTFJsonLightIndex> { using FGLTFJsonIndex::FGLTFJsonIndex; };
+struct FGLTFJsonLightMapIndex : FGLTFJsonIndex<FGLTFJsonLightMapIndex> { using FGLTFJsonIndex::FGLTFJsonIndex; };
+struct FGLTFJsonSkySphereIndex : FGLTFJsonIndex<FGLTFJsonSkySphereIndex> { using FGLTFJsonIndex::FGLTFJsonIndex; };
+struct FGLTFJsonEpicLevelVariantSetsIndex : FGLTFJsonIndex<FGLTFJsonEpicLevelVariantSetsIndex> { using FGLTFJsonIndex::FGLTFJsonIndex; };
+struct FGLTFJsonKhrMaterialVariantIndex : FGLTFJsonIndex<FGLTFJsonKhrMaterialVariantIndex> { using FGLTFJsonIndex::FGLTFJsonIndex; };
diff --git a/Engine/Plugins/Enterprise/GLTFExporter/MigratedFromGit/Source/GLTFExporter/Public/Json/GLTFJsonKhrMaterialVariant.h b/Engine/Plugins/Enterprise/GLTFExporter/MigratedFromGit/Source/GLTFExporter/Public/Json/GLTFJsonKhrMaterialVariant.h
new file mode 100644
--- /dev/null
+++ b/Engine/Plugins/Enterprise/GLTFExporter/MigratedFromGit/Source/GLTFExporter/Public/Json/GLTFJsonKhrMaterialVariant.h
@@ -0,0 +1,28 @@
+// Copyright Epic Games, Inc. All Rights Reserved.
+
+#pragma once
+
+#include "Json/GLTFJsonObject.h"
+#include "Json/GLTFJsonIndex.h"
+
+struct FGLTFJsonKhrMaterialVariantMapping : IGLTFJsonObject
+{
+	FGLTFJsonMaterialIndex                   Material;
+	TArray<FGLTFJsonKhrMaterialVariantIndex> Variants;
+
+	virtual void WriteObject(IGLTFJsonWriter& Writer) const override
+	{
+		Writer.Write(TEXT("material"), Material);
+		Writer.Write(TEXT("variants"), Variants);
+	}
+};
+
+struct FGLTFJsonKhrMaterialVariant : IGLTFJsonObject
+{
+	FString Name;
+
+	virtual void WriteObject(IGLTFJsonWriter& Writer) const override
+	{
+		Writer.Write(TEXT("name"), Name);
+	}
+};
diff --git a/Engine/Plugins/Enterprise/GLTFExporter/MigratedFromGit/Source/GLTFExporter/Public/Json/GLTFJsonLight.h b/Engine/Plugins/Enterprise/GLTFExporter/MigratedFromGit/Source/GLTFExporter/Public/Json/GLTFJsonLight.h
new file mode 100644
--- /dev/null
+++ b/Engine/Plugins/Enterprise/GLTFExporter/MigratedFromGit/Source/GLTFExporter/Public/Json/GLTFJsonLight.h
@@ -0,0 +1,87 @@
+// Copyright Epic Games, Inc. All Rights Reserved.
+
+#pragma once
+
+#include "Json/GLTFJsonObject.h"
+#include "Json/GLTFJsonEnums.h"
+#include "Json/GLTFJsonColor.h"
+
+struct FGLTFJsonSpotLight : IGLTFJsonObject
+{
+	float InnerConeAngle;
+	float OuterConeAngle;
+
+	FGLTFJsonSpotLight()
+		: InnerConeAngle(0)
+		, OuterConeAngle(HALF_PI)
+	{
+	}
+
+	virtual void WriteObject(IGLTFJsonWriter& Writer) const override
+	{
+		if (!FMath::IsNearlyEqual(InnerConeAngle, 0, Writer.DefaultTolerance))
+		{
+			Writer.Write(TEXT("innerConeAngle"), InnerConeAngle);
+		}
+
+		if (!FMath::IsNearlyEqual(OuterConeAngle, HALF_PI, Writer.DefaultTolerance))
+		{
+			Writer.Write(TEXT("outerConeAngle"), OuterConeAngle);
+		}
+	}
+};
+
+struct FGLTFJsonLight : IGLTFJsonObject
+{
+	FString Name;
+
+	EGLTFJsonLightType Type;
+
+	FGLTFJsonColor3 Color;
+
+	float Intensity;
+	float Range;
+
+	FGLTFJsonSpotLight Spot;
+
+	FGLTFJsonLight()
+		: Type(EGLTFJsonLightType::None)
+		, Color(FGLTFJsonColor3::White)
+		, Intensity(1)
+		, Range(0)
+	{
+	}
+
+	virtual void WriteObject(IGLTFJsonWriter& Writer) const override
+	{
+		if (!Name.IsEmpty())
+		{
+			Writer.Write(TEXT("name"), Name);
+		}
+
+		Writer.Write(TEXT("type"), Type);
+
+		if (!Color.IsNearlyEqual(FGLTFJsonColor3::White, Writer.DefaultTolerance))
+		{
+			Writer.Write(TEXT("color"), Color);
+		}
+
+		if (!FMath::IsNearlyEqual(Intensity, 1, Writer.DefaultTolerance))
+		{
+			Writer.Write(TEXT("intensity"), Intensity);
+		}
+
+		if (Type == EGLTFJsonLightType::Point || Type == EGLTFJsonLightType::Spot)
+		{
+			if (!FMath::IsNearlyEqual(Range, 0, Writer.DefaultTolerance))
+			{
+				Writer.Write(TEXT("range"), Range);
+			}
+
+			if (Type == EGLTFJsonLightType::Spot)
+			{
+				Writer.Write(TEXT("spot"), Spot);
+			}
+		}
+	}
+};
diff --git a/Engine/Plugins/Enterprise/GLTFExporter/MigratedFromGit/Source/GLTFExporter/Public/Json/GLTFJsonLightMap.h b/Engine/Plugins/Enterprise/GLTFExporter/MigratedFromGit/Source/GLTFExporter/Public/Json/GLTFJsonLightMap.h
new file mode 100644
--- /dev/null
+++ b/Engine/Plugins/Enterprise/GLTFExporter/MigratedFromGit/Source/GLTFExporter/Public/Json/GLTFJsonLightMap.h
@@ -0,0 +1,40 @@
+// Copyright Epic Games, Inc. All Rights Reserved.
+
+#pragma once
+
+#include "Json/GLTFJsonObject.h"
+#include "Json/GLTFJsonMaterial.h"
+#include "Json/GLTFJsonVector.h"
+
+struct FGLTFJsonLightMap : IGLTFJsonObject
+{
+	FString              Name;
+	FGLTFJsonTextureInfo Texture;
+	FGLTFJsonVector4     LightMapScale;
+	FGLTFJsonVector4     LightMapAdd;
+	FGLTFJsonVector4     CoordinateScaleBias;
+
+	FGLTFJsonLightMap()
+		: LightMapScale(FGLTFJsonVector4::One)
+		, LightMapAdd(FGLTFJsonVector4::Zero)
+		, CoordinateScaleBias({ 1, 1, 0, 0 })
+	{
+	}
+
+	virtual void WriteObject(IGLTFJsonWriter& Writer) const override
+	{
+		if (!Name.IsEmpty())
+		{
+			Writer.Write(TEXT("name"), Name);
+		}
+
+		if (Texture.Index != INDEX_NONE)
+		{
+			Writer.Write(TEXT("texture"), Texture);
+		}
+
+		Writer.Write(TEXT("lightmapScale"), LightMapScale);
+		Writer.Write(TEXT("lightmapAdd"), LightMapAdd);
+		Writer.Write(TEXT("coordinateScaleBias"), CoordinateScaleBias);
+	}
+};
diff --git a/Engine/Plugins/Enterprise/GLTFExporter/MigratedFromGit/Source/GLTFExporter/Public/Json/GLTFJsonMaterial.h b/Engine/Plugins/Enterprise/GLTFExporter/MigratedFromGit/Source/GLTFExporter/Public/Json/GLTFJsonMaterial.h
new file mode 100644
--- /dev/null
+++ b/Engine/Plugins/Enterprise/GLTFExporter/MigratedFromGit/Source/GLTFExporter/Public/Json/GLTFJsonMaterial.h
@@ -0,0 +1,285 @@
+// Copyright Epic Games, Inc. All Rights Reserved.
+
+#pragma once
+
+#include "Json/GLTFJsonObject.h"
+#include "Json/GLTFJsonEnums.h"
+#include "Json/GLTFJsonColor.h"
+#include "Json/GLTFJsonTextureTransform.h"
+
+struct FGLTFJsonTextureInfo : IGLTFJsonObject
+{
+	FGLTFJsonTextureIndex Index;
+	int32 TexCoord;
+
+	FGLTFJsonTextureTransform Transform;
+
+	FGLTFJsonTextureInfo()
+		: TexCoord(0)
+	{
+	}
+
+	virtual void WriteObject(IGLTFJsonWriter& Writer) const override
+	{
+		Writer.Write(TEXT("index"), Index);
+
+		if (TexCoord != 0)
+		{
+			Writer.Write(TEXT("texCoord"), TexCoord);
+		}
+
+		if (!Transform.IsNearlyDefault(Writer.DefaultTolerance))
+		{
+			Writer.StartExtensions();
+			Writer.Write(EGLTFJsonExtension::KHR_TextureTransform, Transform);
+			Writer.EndExtensions();
+		}
+	}
+};
+
+struct FGLTFJsonNormalTextureInfo : FGLTFJsonTextureInfo
+{
+	float Scale;
+
+	FGLTFJsonNormalTextureInfo()
+		: Scale(1)
+	{
+	}
+
+	virtual void WriteObject(IGLTFJsonWriter& Writer) const override
+	{
+		Writer.Write(TEXT("index"), Index);
+
+		if (TexCoord != 0)
+		{
+			Writer.Write(TEXT("texCoord"), TexCoord);
+		}
+
+		if (!FMath::IsNearlyEqual(Scale, 1, Writer.DefaultTolerance))
+		{
+			Writer.Write(TEXT("scale"), Scale);
+		}
+	}
+};
+
+struct FGLTFJsonOcclusionTextureInfo : FGLTFJsonTextureInfo
+{
+	float Strength;
+
+	FGLTFJsonOcclusionTextureInfo()
+		: Strength(1)
+	{
+	}
+
+	virtual void WriteObject(IGLTFJsonWriter& Writer) const override
+	{
+		Writer.Write(TEXT("index"), Index);
+
+		if (TexCoord != 0)
+		{
+			Writer.Write(TEXT("texCoord"), TexCoord);
+		}
+
+		if (!FMath::IsNearlyEqual(Strength, 1, Writer.DefaultTolerance))
+		{
+			Writer.Write(TEXT("strength"), Strength);
+		}
+	}
+};
+
+struct FGLTFJsonPBRMetallicRoughness : IGLTFJsonObject
+{
+	FGLTFJsonColor4 BaseColorFactor;
+	FGLTFJsonTextureInfo BaseColorTexture;
+
+	float MetallicFactor;
+	float RoughnessFactor;
+	FGLTFJsonTextureInfo MetallicRoughnessTexture;
+
+	FGLTFJsonPBRMetallicRoughness()
+		: BaseColorFactor(FGLTFJsonColor4::White)
+		, MetallicFactor(1)
+		, RoughnessFactor(1)
+	{
+	}
+
+	virtual void WriteObject(IGLTFJsonWriter& Writer) const override
+	{
+		if (!BaseColorFactor.IsNearlyEqual(FGLTFJsonColor4::White, Writer.DefaultTolerance))
+		{
+			Writer.Write(TEXT("baseColorFactor"), BaseColorFactor);
+		}
+
+		if (BaseColorTexture.Index != INDEX_NONE)
+		{
+			Writer.Write(TEXT("baseColorTexture"), BaseColorTexture);
+		}
+
+		if (!FMath::IsNearlyEqual(MetallicFactor, 1, Writer.DefaultTolerance))
+		{
+			Writer.Write(TEXT("metallicFactor"), MetallicFactor);
+		}
+
+		if (!FMath::IsNearlyEqual(RoughnessFactor, 1, Writer.DefaultTolerance))
+		{
+			Writer.Write(TEXT("roughnessFactor"), RoughnessFactor);
+		}
+
+		if (MetallicRoughnessTexture.Index != INDEX_NONE)
+		{
+			Writer.Write(TEXT("metallicRoughnessTexture"), MetallicRoughnessTexture);
+		}
+	}
+};
+
+struct FGLTFJsonClearCoatExtension : IGLTFJsonObject
+{
+	float ClearCoatFactor;
+	FGLTFJsonTextureInfo ClearCoatTexture;
+
+	float ClearCoatRoughnessFactor;
+	FGLTFJsonTextureInfo ClearCoatRoughnessTexture;
+
+	FGLTFJsonNormalTextureInfo ClearCoatNormalTexture;
+
+	FGLTFJsonClearCoatExtension()
+		: ClearCoatFactor(0)
+		, ClearCoatRoughnessFactor(0)
+	{
+	}
+
+	virtual void WriteObject(IGLTFJsonWriter& Writer) const override
+	{
+		if (!FMath::IsNearlyEqual(ClearCoatFactor, 0, Writer.DefaultTolerance))
+		{
+			Writer.Write(TEXT("clearcoatFactor"), ClearCoatFactor);
+		}
+
+		if (ClearCoatTexture.Index != INDEX_NONE)
+		{
+			Writer.Write(TEXT("clearcoatTexture"), ClearCoatTexture);
+		}
+
+		if (!FMath::IsNearlyEqual(ClearCoatRoughnessFactor, 0, Writer.DefaultTolerance))
+		{
+			Writer.Write(TEXT("clearcoatRoughnessFactor"), ClearCoatRoughnessFactor);
+		}
+
+		if (ClearCoatRoughnessTexture.Index != INDEX_NONE)
+		{
+			Writer.Write(TEXT("clearcoatRoughnessTexture"), ClearCoatRoughnessTexture);
+		}
+
+		if (ClearCoatNormalTexture.Index != INDEX_NONE)
+		{
+			Writer.Write(TEXT("clearcoatNormalTexture"), ClearCoatNormalTexture);
+		}
+	}
+};
+
+struct FGLTFJsonMaterial : IGLTFJsonObject
+{
+	FString Name;
+
+	EGLTFJsonShadingModel ShadingModel;
+
+	FGLTFJsonPBRMetallicRoughness PBRMetallicRoughness;
+
+	FGLTFJsonNormalTextureInfo NormalTexture;
+	FGLTFJsonOcclusionTextureInfo OcclusionTexture;
+
+	FGLTFJsonTextureInfo EmissiveTexture;
+	FGLTFJsonColor3 EmissiveFactor;
+
+	EGLTFJsonAlphaMode AlphaMode;
+	float AlphaCutoff;
+
+	bool DoubleSided;
+
+	EGLTFJsonBlendMode BlendMode;
+
+	FGLTFJsonClearCoatExtension ClearCoat;
+
+	FGLTFJsonMaterial()
+		: ShadingModel(EGLTFJsonShadingModel::Default)
+		, EmissiveFactor(FGLTFJsonColor3::Black)
+		, AlphaMode(EGLTFJsonAlphaMode::Opaque)
+		, AlphaCutoff(0.5f)
+		, DoubleSided(false)
+		, BlendMode(EGLTFJsonBlendMode::None)
+	{
+	}
+
+	virtual void WriteObject(IGLTFJsonWriter& Writer) const override
+	{
+		if (!Name.IsEmpty())
+		{
+			Writer.Write(TEXT("name"), Name);
+		}
+
+		if (ShadingModel != EGLTFJsonShadingModel::None)
+		{
+			Writer.Write(TEXT("pbrMetallicRoughness"), PBRMetallicRoughness);
+		}
+
+		if (NormalTexture.Index != INDEX_NONE)
+		{
+			Writer.Write(TEXT("normalTexture"), NormalTexture);
+		}
+
+		if (OcclusionTexture.Index != INDEX_NONE)
+		{
+			Writer.Write(TEXT("occlusionTexture"), OcclusionTexture);
+		}
+
+		if (EmissiveTexture.Index != INDEX_NONE)
+		{
+			Writer.Write(TEXT("emissiveTexture"), EmissiveTexture);
+		}
+
+		if (!EmissiveFactor.IsNearlyEqual(FGLTFJsonColor3::Black, Writer.DefaultTolerance))
+		{
+			Writer.Write(TEXT("emissiveFactor"), EmissiveFactor);
+		}
+
+		if (AlphaMode != EGLTFJsonAlphaMode::Opaque)
+		{
+			Writer.Write(TEXT("alphaMode"), AlphaMode);
+
+			if (AlphaMode == EGLTFJsonAlphaMode::Mask && !FMath::IsNearlyEqual(AlphaCutoff, 0.5f, Writer.DefaultTolerance))
+			{
+				Writer.Write(TEXT("alphaCutoff"), AlphaCutoff);
+			}
+		}
+
+		if (DoubleSided)
+		{
+			Writer.Write(TEXT("doubleSided"), DoubleSided);
+		}
+
+		if (BlendMode != EGLTFJsonBlendMode::None || ShadingModel == EGLTFJsonShadingModel::Unlit || ShadingModel == EGLTFJsonShadingModel::ClearCoat)
+		{
+			Writer.StartExtensions();
+
+			if (BlendMode != EGLTFJsonBlendMode::None)
+			{
+				Writer.StartExtension(EGLTFJsonExtension::EPIC_BlendModes);
+				Writer.Write(TEXT("blendMode"), BlendMode);
+				Writer.EndExtension();
+			}
+
+			if (ShadingModel == EGLTFJsonShadingModel::Unlit)
+			{
+				Writer.StartExtension(EGLTFJsonExtension::KHR_MaterialsUnlit);
+				// Write empty object
+				Writer.EndExtension();
+			}
+			else if (ShadingModel == EGLTFJsonShadingModel::ClearCoat)
+			{
+				Writer.Write(EGLTFJsonExtension::KHR_MaterialsClearCoat, ClearCoat);
+			}
+
+			Writer.EndExtensions();
+		}
+	}
+};
diff --git a/Engine/Plugins/Enterprise/GLTFExporter/MigratedFromGit/Source/GLTFExporter/Public/Json/GLTFJsonMatrix.cpp b/Engine/Plugins/Enterprise/GLTFExporter/MigratedFromGit/Source/GLTFExporter/Public/Json/GLTFJsonMatrix.cpp
new file mode 100644
--- /dev/null
+++ b/Engine/Plugins/Enterprise/GLTFExporter/MigratedFromGit/Source/GLTFExporter/Public/Json/GLTFJsonMatrix.cpp
@@ -0,0 +1,24 @@
+// Copyright Epic Games, Inc. All Rights Reserved.
+
+#include "Json/GLTFJsonMatrix.h"
+
+const FGLTFJsonMatrix2 FGLTFJsonMatrix2::Identity
+({
+	1, 0,
+	0, 1
+});
+
+const FGLTFJsonMatrix3 FGLTFJsonMatrix3::Identity
+({
+	1, 0, 0,
+	0, 1, 0,
+	0, 0, 1
+});
+
+const FGLTFJsonMatrix4 FGLTFJsonMatrix4::Identity
+({
+	1, 0, 0, 0,
+	0, 1, 0, 0,
+	0, 0, 1, 0,
+	0, 0, 0, 1
+});
diff --git a/Engine/Plugins/Enterprise/GLTFExporter/MigratedFromGit/Source/GLTFExporter/Public/Json/GLTFJsonMatrix.h b/Engine/Plugins/Enterprise/GLTFExporter/MigratedFromGit/Source/GLTFExporter/Public/Json/GLTFJsonMatrix.h
new file mode 100644
--- /dev/null
+++ b/Engine/Plugins/Enterprise/GLTFExporter/MigratedFromGit/Source/GLTFExporter/Public/Json/GLTFJsonMatrix.h
@@ -0,0 +1,86 @@
+// Copyright Epic Games, Inc. All Rights Reserved.
+
+#pragma once
+
+#include "Core/GLTFMatrix.h"
+#include "Json/GLTFJsonArray.h"
+
+template <typename BaseType>
+struct TGLTFJsonMatrix : BaseType, IGLTFJsonArray
+{
+	static const TGLTFJsonMatrix Identity;
+
+	TGLTFJsonMatrix(const BaseType& Other)
+		: BaseType(Other)
+	{
+	}
+
+	TGLTFJsonMatrix& operator=(const BaseType& Other)
+	{
+		*static_cast<BaseType*>(this) = Other;
+		return *this;
+	}
+
+	bool operator==(const BaseType& Other) const
+	{
+		for (SIZE_T i = 0; i < GetNum(BaseType::Components); ++i)
+		{
+			if (BaseType::Components[i] != Other.Components[i])
+			{
+				return false;
+			}
+		}
+
+		return true;
+	}
+
+	bool operator!=(const BaseType& Other) const
+	{
+		return !(*this == Other);
+	}
+
+	virtual void WriteArray(IGLTFJsonWriter& Writer) const override
+	{
+		for (SIZE_T i = 0; i < GetNum(BaseType::Elements); ++i)
+		{
+			Writer.Write(BaseType::Elements[i]);
+		}
+	}
+
+	bool IsNearlyEqual(const BaseType& Other, float Tolerance = KINDA_SMALL_NUMBER) const
+	{
+		for (SIZE_T i = 0; i < GetNum(BaseType::Elements); ++i)
+		{
+			if (!FMath::IsNearlyEqual(BaseType::Elements[i], Other.Elements[i], Tolerance))
+			{
+				return false;
+			}
+		}
+
+		return true;
+	}
+};
+
+struct FGLTFJsonMatrix2 : TGLTFJsonMatrix<FGLTFMatrix2>
+{
+	static const FGLTFJsonMatrix2 Identity;
+
+	using TGLTFJsonMatrix::TGLTFJsonMatrix;
+	using TGLTFJsonMatrix::operator=;
+};
+
+struct FGLTFJsonMatrix3 : TGLTFJsonMatrix<FGLTFMatrix3>
+{
+	static const FGLTFJsonMatrix3 Identity;
+
+	using TGLTFJsonMatrix::TGLTFJsonMatrix;
+	using TGLTFJsonMatrix::operator=;
+};
+
+struct FGLTFJsonMatrix4 : TGLTFJsonMatrix<FGLTFMatrix4>
+{
+	static const FGLTFJsonMatrix4 Identity;
+
+	using TGLTFJsonMatrix::TGLTFJsonMatrix;
+	using TGLTFJsonMatrix::operator=;
+};
diff --git a/Engine/Plugins/Enterprise/GLTFExporter/MigratedFromGit/Source/GLTFExporter/Public/Json/GLTFJsonMesh.h b/Engine/Plugins/Enterprise/GLTFExporter/MigratedFromGit/Source/GLTFExporter/Public/Json/GLTFJsonMesh.h
new file mode 100644
--- /dev/null
+++ b/Engine/Plugins/Enterprise/GLTFExporter/MigratedFromGit/Source/GLTFExporter/Public/Json/GLTFJsonMesh.h
@@ -0,0 +1,109 @@
+// Copyright Epic Games, Inc. All Rights Reserved.
+
+#pragma once
+
+#include "Json/GLTFJsonObject.h"
+#include "Json/GLTFJsonEnums.h"
+#include "Json/GLTFJsonIndex.h"
+#include "Json/GLTFJsonKhrMaterialVariant.h"
+
+struct FGLTFJsonAttributes : IGLTFJsonObject
+{
+	FGLTFJsonAccessorIndex Position;
+	FGLTFJsonAccessorIndex Color0;
+	FGLTFJsonAccessorIndex Normal;
+	FGLTFJsonAccessorIndex Tangent;
+
+	TArray<FGLTFJsonAccessorIndex> TexCoords;
+	TArray<FGLTFJsonAccessorIndex> Joints;
+	TArray<FGLTFJsonAccessorIndex> Weights;
+
+	virtual void WriteObject(IGLTFJsonWriter& Writer) const override
+	{
+		if (Position != INDEX_NONE) Writer.Write(TEXT("POSITION"), Position);
+		if (Color0 != INDEX_NONE) Writer.Write(TEXT("COLOR_0"), Color0);
+		if (Normal != INDEX_NONE) Writer.Write(TEXT("NORMAL"), Normal);
+		if (Tangent != INDEX_NONE) Writer.Write(TEXT("TANGENT"), Tangent);
+
+		for (int32 Index = 0; Index < TexCoords.Num(); ++Index)
+		{
+			const FGLTFJsonAccessorIndex TexCoord = TexCoords[Index];
+			if (TexCoord != INDEX_NONE) Writer.Write(TEXT("TEXCOORD_") + FString::FromInt(Index), TexCoord);
+		}
+
+		for (int32 Index = 0; Index < Joints.Num(); ++Index)
+		{
+			const FGLTFJsonAccessorIndex Joint = Joints[Index];
+			if (Joint != INDEX_NONE) Writer.Write(TEXT("JOINTS_") + FString::FromInt(Index), Joint);
+		}
+
+		for (int32 Index = 0; Index < Weights.Num(); ++Index)
+		{
+			const FGLTFJsonAccessorIndex Weight = Weights[Index];
+			if (Weight != INDEX_NONE) Writer.Write(TEXT("WEIGHTS_") + FString::FromInt(Index), Weight);
+		}
+	}
+};
+
+struct FGLTFJsonPrimitive : IGLTFJsonObject
+{
+	FGLTFJsonAccessorIndex Indices;
+	FGLTFJsonMaterialIndex Material;
+	EGLTFJsonPrimitiveMode Mode;
+	FGLTFJsonAttributes    Attributes;
+
+	TArray<FGLTFJsonKhrMaterialVariantMapping> KhrMaterialVariantMappings;
+
+	FGLTFJsonPrimitive()
+		: Mode(EGLTFJsonPrimitiveMode::Triangles)
+	{
+	}
+
+	virtual void WriteObject(IGLTFJsonWriter& Writer) const override
+	{
+		Writer.Write(TEXT("attributes"), Attributes);
+
+		if (Indices != INDEX_NONE)
+		{
+			Writer.Write(TEXT("indices"), Indices);
+		}
+
+		if (Material != INDEX_NONE)
+		{
+			Writer.Write(TEXT("material"), Material);
+		}
+
+		if (Mode != EGLTFJsonPrimitiveMode::Triangles)
+		{
+			Writer.Write(TEXT("mode"), Mode);
+		}
+
+		if (KhrMaterialVariantMappings.Num() > 0)
+		{
+			Writer.StartExtensions();
+
+			Writer.StartExtension(EGLTFJsonExtension::KHR_MaterialsVariants);
+			Writer.Write(TEXT("mappings"), KhrMaterialVariantMappings);
+			Writer.EndExtension();
+
+			Writer.EndExtensions();
+		}
+	}
+};
+
+struct FGLTFJsonMesh : IGLTFJsonObject
+{
+	FString Name;
+
+	TArray<FGLTFJsonPrimitive> Primitives;
+
+	virtual void WriteObject(IGLTFJsonWriter& Writer) const override
+	{
+		if (!Name.IsEmpty())
+		{
+			Writer.Write(TEXT("name"), Name);
+		}
+
+		Writer.Write(TEXT("primitives"), Primitives);
+	}
+};
diff --git a/Engine/Plugins/Enterprise/GLTFExporter/MigratedFromGit/Source/GLTFExporter/Public/Json/GLTFJsonNode.h b/Engine/Plugins/Enterprise/GLTFExporter/MigratedFromGit/Source/GLTFExporter/Public/Json/GLTFJsonNode.h
new file mode 100644
--- /dev/null
+++ b/Engine/Plugins/Enterprise/GLTFExporter/MigratedFromGit/Source/GLTFExporter/Public/Json/GLTFJsonNode.h
@@ -0,0 +1,122 @@
+// Copyright Epic Games, Inc. All Rights Reserved.
+
+#pragma once
+
+#include "Json/GLTFJsonObject.h"
+#include "Json/GLTFJsonIndex.h"
+#include "Json/GLTFJsonVector.h"
+#include "Json/GLTFJsonQuaternion.h"
+
+struct FGLTFJsonNode : IGLTFJsonObject
+{
+	FString Name;
+
+	FGLTFJsonVector3    Translation;
+	FGLTFJsonQuaternion Rotation;
+	FGLTFJsonVector3    Scale;
+
+	FGLTFJsonCameraIndex    Camera;
+	FGLTFJsonSkinIndex      Skin;
+	FGLTFJsonMeshIndex      Mesh;
+	FGLTFJsonBackdropIndex  Backdrop;
+	FGLTFJsonHotspotIndex   Hotspot;
+	FGLTFJsonLightIndex     Light;
+	FGLTFJsonLightMapIndex  LightMap;
+	FGLTFJsonSkySphereIndex SkySphere;
+
+	FGLTFJsonNodeIndex ComponentNode;
+
+	TArray<FGLTFJsonNodeIndex> Children;
+
+	FGLTFJsonNode()
+		: Translation(FGLTFJsonVector3::Zero)
+		, Rotation(FGLTFJsonQuaternion::Identity)
+		, Scale(FGLTFJsonVector3::One)
+	{
+	}
+
+	virtual void WriteObject(IGLTFJsonWriter& Writer) const override
+	{
+		if (!Name.IsEmpty())
+		{
+			Writer.Write(TEXT("name"), Name);
+		}
+
+		if (!Translation.IsNearlyEqual(FGLTFJsonVector3::Zero, Writer.DefaultTolerance))
+		{
+			Writer.Write(TEXT("translation"), Translation);
+		}
+
+		if (!Rotation.IsNearlyEqual(FGLTFJsonQuaternion::Identity, Writer.DefaultTolerance))
+		{
+			Writer.Write(TEXT("rotation"), Rotation);
+		}
+
+		if (!Scale.IsNearlyEqual(FGLTFJsonVector3::One, Writer.DefaultTolerance))
+		{
+			Writer.Write(TEXT("scale"), Scale);
+		}
+
+		if (Camera != INDEX_NONE)
+		{
+			Writer.Write(TEXT("camera"), Camera);
+		}
+
+		if (Skin != INDEX_NONE)
+		{
+			Writer.Write(TEXT("skin"), Skin);
+		}
+
+		if (Mesh != INDEX_NONE)
+		{
+			Writer.Write(TEXT("mesh"), Mesh);
+		}
+
+		if (Backdrop != INDEX_NONE || Hotspot != INDEX_NONE || Light != INDEX_NONE || LightMap != INDEX_NONE || SkySphere != INDEX_NONE)
+		{
+			Writer.StartExtensions();
+
+			if (Backdrop != INDEX_NONE)
+			{
+				Writer.StartExtension(EGLTFJsonExtension::EPIC_HDRIBackdrops);
+				Writer.Write(TEXT("backdrop"), Backdrop);
+				Writer.EndExtension();
+			}
+
+			if (Hotspot != INDEX_NONE)
+			{
+				Writer.StartExtension(EGLTFJsonExtension::EPIC_AnimationHotspots);
+				Writer.Write(TEXT("hotspot"), Hotspot);
+				Writer.EndExtension();
+			}
+
+			if (Light != INDEX_NONE)
+			{
+				Writer.StartExtension(EGLTFJsonExtension::KHR_LightsPunctual);
+				Writer.Write(TEXT("light"), Light);
+				Writer.EndExtension();
+			}
+
+			if (LightMap != INDEX_NONE)
+			{
+				Writer.StartExtension(EGLTFJsonExtension::EPIC_LightmapTextures);
+				Writer.Write(TEXT("lightmap"), LightMap);
+				Writer.EndExtension();
+			}
+
+			if (SkySphere != INDEX_NONE)
+			{
+				Writer.StartExtension(EGLTFJsonExtension::EPIC_SkySpheres);
+				Writer.Write(TEXT("skySphere"), SkySphere);
+				Writer.EndExtension();
+			}
+
+			Writer.EndExtensions();
+		}
+
+		if (Children.Num() > 0)
+		{
+			Writer.Write(TEXT("children"), Children);
+		}
+	}
+};
diff --git a/Engine/Plugins/Enterprise/GLTFExporter/MigratedFromGit/Source/GLTFExporter/Public/Json/GLTFJsonObject.h b/Engine/Plugins/Enterprise/GLTFExporter/MigratedFromGit/Source/GLTFExporter/Public/Json/GLTFJsonObject.h
new file mode 100644
--- /dev/null
+++ b/Engine/Plugins/Enterprise/GLTFExporter/MigratedFromGit/Source/GLTFExporter/Public/Json/GLTFJsonObject.h
@@ -0,0 +1,18 @@
+// Copyright Epic Games, Inc. All Rights Reserved.
+
+#pragma once
+
+#include "Json/GLTFJsonValue.h"
+#include "Json/GLTFJsonWriter.h"
+
+struct IGLTFJsonObject : IGLTFJsonValue
+{
+	virtual void WriteValue(IGLTFJsonWriter& Writer) const override final
+	{
+		Writer.StartObject();
+		WriteObject(Writer);
+		Writer.EndObject();
+	}
+
+	virtual void WriteObject(IGLTFJsonWriter& Writer) const = 0;
+};
diff --git a/Engine/Plugins/Enterprise/GLTFExporter/MigratedFromGit/Source/GLTFExporter/Public/Json/GLTFJsonQuaternion.cpp b/Engine/Plugins/Enterprise/GLTFExporter/MigratedFromGit/Source/GLTFExporter/Public/Json/GLTFJsonQuaternion.cpp
new file mode 100644
--- /dev/null
+++ b/Engine/Plugins/Enterprise/GLTFExporter/MigratedFromGit/Source/GLTFExporter/Public/Json/GLTFJsonQuaternion.cpp
@@ -0,0 +1,5 @@
+// Copyright Epic Games, Inc. All Rights Reserved.
+
+#include "Json/GLTFJsonQuaternion.h"
+
+const FGLTFQuaternion FGLTFJsonQuaternion::Identity = { 0, 0, 0, 1 };
diff --git a/Engine/Plugins/Enterprise/GLTFExporter/MigratedFromGit/Source/GLTFExporter/Public/Json/GLTFJsonQuaternion.h b/Engine/Plugins/Enterprise/GLTFExporter/MigratedFromGit/Source/GLTFExporter/Public/Json/GLTFJsonQuaternion.h
new file mode 100644
--- /dev/null
+++ b/Engine/Plugins/Enterprise/GLTFExporter/MigratedFromGit/Source/GLTFExporter/Public/Json/GLTFJsonQuaternion.h
@@ -0,0 +1,54 @@
+// Copyright Epic Games, Inc. All Rights Reserved.
+
+#pragma once
+
+#include "Core/GLTFQuaternion.h"
+#include "Json/GLTFJsonArray.h"
+
+struct FGLTFJsonQuaternion : TGLTFQuaternion<float>, IGLTFJsonArray
+{
+	static const TGLTFQuaternion Identity;
+
+	FGLTFJsonQuaternion(const TGLTFQuaternion& Other)
+		: TGLTFQuaternion(Other)
+	{
+	}
+
+	FGLTFJsonQuaternion& operator=(const TGLTFQuaternion& Other)
+	{
+		*static_cast<TGLTFQuaternion*>(this) = Other;
+		return *this;
+	}
+
+	bool operator==(const TGLTFQuaternion& Other) const
+	{
+		for (SIZE_T i = 0; i < GetNum(Components); ++i)
+		{
+			if (Components[i] != Other.Components[i])
+			{
+				return false;
+			}
+		}
+
+		return true;
+	}
+
+	bool operator!=(const TGLTFQuaternion& Other) const
+	{
+		return !(*this == Other);
+	}
+
+	virtual void WriteArray(IGLTFJsonWriter& Writer) const override
+	{
+		for (int32 i = 0; i < 4; ++i)
+		{
+			Writer.Write(Components[i]);
+		}
+	}
+
+	bool IsNearlyEqual(const TGLTFQuaternion& Other, float Tolerance = KINDA_SMALL_NUMBER) const
+	{
+		return (FMath::Abs(X - Other.X) <= Tolerance && FMath::Abs(Y - Other.Y) <= Tolerance && FMath::Abs(Z - Other.Z) <= Tolerance && FMath::Abs(W - Other.W) <= Tolerance)
+			|| (FMath::Abs(X + Other.X) <= Tolerance && FMath::Abs(Y + Other.Y) <= Tolerance && FMath::Abs(Z + Other.Z) <= Tolerance && FMath::Abs(W + Other.W) <= Tolerance);
+	}
+};
diff --git a/Engine/Plugins/Enterprise/GLTFExporter/MigratedFromGit/Source/GLTFExporter/Public/Json/GLTFJsonRoot.h b/Engine/Plugins/Enterprise/GLTFExporter/MigratedFromGit/Source/GLTFExporter/Public/Json/GLTFJsonRoot.h
new file mode 100644
--- /dev/null
+++ b/Engine/Plugins/Enterprise/GLTFExporter/MigratedFromGit/Source/GLTFExporter/Public/Json/GLTFJsonRoot.h
@@ -0,0 +1,156 @@
+// Copyright Epic Games, Inc. All Rights Reserved.
+
+#pragma once
+
+#include "Json/GLTFJsonObject.h"
+#include "Json/GLTFJsonIndex.h"
+#include "Json/GLTFJsonAsset.h"
+#include "Json/GLTFJsonAccessor.h"
+#include "Json/GLTFJsonAnimation.h"
+#include "Json/GLTFJsonBuffer.h"
+#include "Json/GLTFJsonBufferView.h"
+#include "Json/GLTFJsonCamera.h"
+#include "Json/GLTFJsonImage.h"
+#include "Json/GLTFJsonMaterial.h"
+#include "Json/GLTFJsonMesh.h"
+#include "Json/GLTFJsonNode.h"
+#include "Json/GLTFJsonSampler.h"
+#include "Json/GLTFJsonScene.h"
+#include "Json/GLTFJsonSkin.h"
+#include "Json/GLTFJsonTexture.h"
+#include "Json/GLTFJsonBackdrop.h"
+#include "Json/GLTFJsonHotspot.h"
+#include "Json/GLTFJsonLight.h"
+#include "Json/GLTFJsonLightMap.h"
+#include "Json/GLTFJsonSkySphere.h"
+#include "Json/GLTFJsonEpicLevelVariantSets.h"
+#include "Json/GLTFJsonKhrMaterialVariant.h"
+#include "Json/GLTFJsonWriter.h"
+
+struct FGLTFJsonRoot : IGLTFJsonObject
+{
+	FGLTFJsonAsset Asset;
+
+	FGLTFJsonExtensions Extensions;
+
+	FGLTFJsonSceneIndex DefaultScene;
+
+	TArray<TUniquePtr<FGLTFJsonAccessor>>   Accessors;
+	TArray<TUniquePtr<FGLTFJsonAnimation>>  Animations;
+	TArray<TUniquePtr<FGLTFJsonBuffer>>     Buffers;
+	TArray<TUniquePtr<FGLTFJsonBufferView>> BufferViews;
+	TArray<TUniquePtr<FGLTFJsonCamera>>     Cameras;
+	TArray<TUniquePtr<FGLTFJsonMaterial>>   Materials;
+	TArray<TUniquePtr<FGLTFJsonMesh>>       Meshes;
+	TArray<TUniquePtr<FGLTFJsonNode>>       Nodes;
+	TArray<TUniquePtr<FGLTFJsonImage>>      Images;
+	TArray<TUniquePtr<FGLTFJsonSampler>>    Samplers;
+	TArray<TUniquePtr<FGLTFJsonScene>>      Scenes;
+	TArray<TUniquePtr<FGLTFJsonSkin>>       Skins;
+	TArray<TUniquePtr<FGLTFJsonTexture>>    Textures;
+	TArray<TUniquePtr<FGLTFJsonBackdrop>>   Backdrops;
+	TArray<TUniquePtr<FGLTFJsonHotspot>>    Hotspots;
+	TArray<TUniquePtr<FGLTFJsonLight>>      Lights;
+	TArray<TUniquePtr<FGLTFJsonLightMap>>   LightMaps;
+	TArray<TUniquePtr<FGLTFJsonSkySphere>>  SkySpheres;
+	TArray<TUniquePtr<FGLTFJsonEpicLevelVariantSets>>  EpicLevelVariantSets;
+	TArray<TUniquePtr<FGLTFJsonKhrMaterialVariant>>    KhrMaterialVariants;
+
+	virtual void WriteObject(IGLTFJsonWriter& Writer) const override
+	{
+		Writer.Write(TEXT("asset"), Asset);
+
+		if (DefaultScene != INDEX_NONE)
+		{
+			Writer.Write(TEXT("scene"), DefaultScene);
+		}
+
+		if (Accessors.Num() > 0) Writer.Write(TEXT("accessors"), Accessors);
+		if (Animations.Num() > 0) Writer.Write(TEXT("animations"), Animations);
+		if (Buffers.Num() > 0) Writer.Write(TEXT("buffers"), Buffers);
+		if (BufferViews.Num() > 0) Writer.Write(TEXT("bufferViews"), BufferViews);
+		if (Cameras.Num() > 0) Writer.Write(TEXT("cameras"), Cameras);
+		if (Images.Num() > 0) Writer.Write(TEXT("images"), Images);
+		if (Materials.Num() > 0) Writer.Write(TEXT("materials"), Materials);
+		if (Meshes.Num() > 0) Writer.Write(TEXT("meshes"), Meshes);
+		if (Nodes.Num() > 0) Writer.Write(TEXT("nodes"), Nodes);
+		if (Samplers.Num() > 0) Writer.Write(TEXT("samplers"), Samplers);
+		if (Scenes.Num() > 0) Writer.Write(TEXT("scenes"), Scenes);
+		if (Skins.Num() > 0) Writer.Write(TEXT("skins"), Skins);
+		if (Textures.Num() > 0) Writer.Write(TEXT("textures"), Textures);
+
+		if (Backdrops.Num() > 0 || Hotspots.Num() > 0 || Lights.Num() > 0 || LightMaps.Num() > 0 || SkySpheres.Num() > 0 || EpicLevelVariantSets.Num() > 0 || KhrMaterialVariants.Num() > 0)
+		{
+			Writer.StartExtensions();
+
+			if (Backdrops.Num() > 0)
+			{
+				Writer.StartExtension(EGLTFJsonExtension::EPIC_HDRIBackdrops);
+				Writer.Write(TEXT("backdrops"), Backdrops);
+				Writer.EndExtension();
+			}
+
+			if (Hotspots.Num() > 0)
+			{
+				Writer.StartExtension(EGLTFJsonExtension::EPIC_AnimationHotspots);
+				Writer.Write(TEXT("hotspots"), Hotspots);
+				Writer.EndExtension();
+			}
+
+			if (EpicLevelVariantSets.Num() > 0)
+			{
+				Writer.StartExtension(EGLTFJsonExtension::EPIC_LevelVariantSets);
+				Writer.Write(TEXT("levelVariantSets"), EpicLevelVariantSets);
+				Writer.EndExtension();
+			}
+
+			if (KhrMaterialVariants.Num() > 0)
+			{
+				Writer.StartExtension(EGLTFJsonExtension::KHR_MaterialsVariants);
+				Writer.Write(TEXT("variants"), KhrMaterialVariants);
+				Writer.EndExtension();
+			}
+
+			if (Lights.Num() > 0)
+			{
+				Writer.StartExtension(EGLTFJsonExtension::KHR_LightsPunctual);
+				Writer.Write(TEXT("lights"), Lights);
+				Writer.EndExtension();
+			}
+
+			if (LightMaps.Num() > 0)
+			{
+				Writer.StartExtension(EGLTFJsonExtension::EPIC_LightmapTextures);
+				Writer.Write(TEXT("lightmaps"), LightMaps);
+				Writer.EndExtension();
+			}
+
+			if (SkySpheres.Num() > 0)
+			{
+				Writer.StartExtension(EGLTFJsonExtension::EPIC_SkySpheres);
+				Writer.Write(TEXT("skySpheres"), SkySpheres);
+				Writer.EndExtension();
+			}
+
+			Writer.EndExtensions();
+		}
+
+		if (Extensions.Used.Num() > 0)
+		{
+			Writer.Write(TEXT("extensionsUsed"), Extensions.Used);
+		}
+
+		if (Extensions.Required.Num() > 0)
+		{
+			Writer.Write(TEXT("extensionsRequired"), Extensions.Required);
+		}
+	}
+
+	void WriteJson(FArchive& Archive, bool bPrettyJson, float DefaultTolerance)
+	{
+		TSharedRef<IGLTFJsonWriter> Writer = IGLTFJsonWriter::Create(Archive, bPrettyJson, Extensions);
+		Writer->DefaultTolerance = DefaultTolerance;
+		Writer->Write(this);
+		Writer->Close();
+	}
+};
diff --git a/Engine/Plugins/Enterprise/GLTFExporter/MigratedFromGit/Source/GLTFExporter/Public/Json/GLTFJsonSampler.h b/Engine/Plugins/Enterprise/GLTFExporter/MigratedFromGit/Source/GLTFExporter/Public/Json/GLTFJsonSampler.h
new file mode 100644
--- /dev/null
+++ b/Engine/Plugins/Enterprise/GLTFExporter/MigratedFromGit/Source/GLTFExporter/Public/Json/GLTFJsonSampler.h
@@ -0,0 +1,53 @@
+// Copyright Epic Games, Inc. All Rights Reserved.
+
+#pragma once
+
+#include "Json/GLTFJsonObject.h"
+#include "Json/GLTFJsonEnums.h"
+
+struct FGLTFJsonSampler : IGLTFJsonObject
+{
+	FString Name;
+
+	EGLTFJsonTextureFilter MinFilter;
+	EGLTFJsonTextureFilter MagFilter;
+
+	EGLTFJsonTextureWrap WrapS;
+	EGLTFJsonTextureWrap WrapT;
+
+	FGLTFJsonSampler()
+		: MinFilter(EGLTFJsonTextureFilter::None)
+		, MagFilter(EGLTFJsonTextureFilter::None)
+		, WrapS(EGLTFJsonTextureWrap::Repeat)
+		, WrapT(EGLTFJsonTextureWrap::Repeat)
+	{
+	}
+
+	virtual void WriteObject(IGLTFJsonWriter& Writer) const override
+	{
+		if (!Name.IsEmpty())
+		{
+			Writer.Write(TEXT("name"), Name);
+		}
+
+		if (MinFilter != EGLTFJsonTextureFilter::None)
+		{
+			Writer.Write(TEXT("minFilter"), MinFilter);
+		}
+
+		if (MagFilter != EGLTFJsonTextureFilter::None)
+		{
+			Writer.Write(TEXT("magFilter"), MagFilter);
+		}
+
+		if (WrapS != EGLTFJsonTextureWrap::Repeat)
+		{
+			Writer.Write(TEXT("wrapS"), WrapS);
+		}
+
+		if (WrapT != EGLTFJsonTextureWrap::Repeat)
+		{
+			Writer.Write(TEXT("wrapT"), WrapT);
+		}
+	}
+};
diff --git a/Engine/Plugins/Enterprise/GLTFExporter/MigratedFromGit/Source/GLTFExporter/Public/Json/GLTFJsonScene.h b/Engine/Plugins/Enterprise/GLTFExporter/MigratedFromGit/Source/GLTFExporter/Public/Json/GLTFJsonScene.h
new file mode 100644
--- /dev/null
+++ b/Engine/Plugins/Enterprise/GLTFExporter/MigratedFromGit/Source/GLTFExporter/Public/Json/GLTFJsonScene.h
@@ -0,0 +1,38 @@
+// Copyright Epic Games, Inc. All Rights Reserved.
+
+#pragma once
+
+#include "Json/GLTFJsonObject.h"
+#include "Json/GLTFJsonIndex.h"
+
+struct FGLTFJsonScene : IGLTFJsonObject
+{
+	FString Name;
+
+	TArray<FGLTFJsonNodeIndex> Nodes;
+	TArray<FGLTFJsonEpicLevelVariantSetsIndex>  EpicLevelVariantSets;
+
+	virtual void WriteObject(IGLTFJsonWriter& Writer) const override
+	{
+		if (!Name.IsEmpty())
+		{
+			Writer.Write(TEXT("name"), Name);
+		}
+
+		if (Nodes.Num() > 0)
+		{
+			Writer.Write(TEXT("nodes"), Nodes);
+		}
+
+		if (EpicLevelVariantSets.Num() > 0)
+		{
+			Writer.StartExtensions();
+
+			Writer.StartExtension(EGLTFJsonExtension::EPIC_LevelVariantSets);
+			Writer.Write(TEXT("levelVariantSets"), EpicLevelVariantSets);
+			Writer.EndExtension();
+
+			Writer.EndExtensions();
+		}
+	}
+};
diff --git a/Engine/Plugins/Enterprise/GLTFExporter/MigratedFromGit/Source/GLTFExporter/Public/Json/GLTFJsonSkin.h b/Engine/Plugins/Enterprise/GLTFExporter/MigratedFromGit/Source/GLTFExporter/Public/Json/GLTFJsonSkin.h
new file mode 100644
--- /dev/null
+++ b/Engine/Plugins/Enterprise/GLTFExporter/MigratedFromGit/Source/GLTFExporter/Public/Json/GLTFJsonSkin.h
@@ -0,0 +1,39 @@
+// Copyright Epic Games, Inc. All Rights Reserved.
+
+#pragma once
+
+#include "Json/GLTFJsonObject.h"
+#include "Json/GLTFJsonIndex.h"
+
+struct FGLTFJsonSkin : IGLTFJsonObject
+{
+	FString Name;
+
+	FGLTFJsonAccessorIndex InverseBindMatrices;
+	FGLTFJsonNodeIndex Skeleton;
+
+	TArray<FGLTFJsonNodeIndex> Joints;
+
+	virtual void WriteObject(IGLTFJsonWriter& Writer) const override
+	{
+		if (!Name.IsEmpty())
+		{
+			Writer.Write(TEXT("name"), Name);
+		}
+
+		if (InverseBindMatrices != INDEX_NONE)
+		{
+			Writer.Write(TEXT("inverseBindMatrices"), InverseBindMatrices);
+		}
+
+		if (Skeleton != INDEX_NONE)
+		{
+			Writer.Write(TEXT("skeleton"), Skeleton);
+		}
+
+		if (Joints.Num() > 0)
+		{
+			Writer.Write(TEXT("joints"), Joints);
+		}
+	}
+};
diff --git a/Engine/Plugins/Enterprise/GLTFExporter/MigratedFromGit/Source/GLTFExporter/Public/Json/GLTFJsonSkySphere.h b/Engine/Plugins/Enterprise/GLTFExporter/MigratedFromGit/Source/GLTFExporter/Public/Json/GLTFJsonSkySphere.h
new file mode 100644
--- /dev/null
+++ b/Engine/Plugins/Enterprise/GLTFExporter/MigratedFromGit/Source/GLTFExporter/Public/Json/GLTFJsonSkySphere.h
@@ -0,0 +1,155 @@
+// Copyright Epic Games, Inc. All Rights Reserved.
+
+#pragma once
+
+#include "Json/GLTFJsonObject.h"
+#include "Json/GLTFJsonIndex.h"
+#include "Json/GLTFJsonColor.h"
+
+struct FGLTFJsonSkySphereColorCurve : IGLTFJsonArray
+{
+	struct FKey
+	{
+		float Time;
+		float Value;
+	};
+
+	struct FComponentCurve
+	{
+		TArray<FKey> Keys;
+	};
+
+	TArray<FComponentCurve> ComponentCurves;
+
+	virtual void WriteArray(IGLTFJsonWriter& Writer) const override
+	{
+		for (int32 ComponentIndex = 0; ComponentIndex < ComponentCurves.Num(); ++ComponentIndex)
+		{
+			const FComponentCurve& ComponentCurve = ComponentCurves[ComponentIndex];
+
+			Writer.StartArray();
+
+			for (const FKey& Key: ComponentCurve.Keys)
+			{
+				Writer.Write(Key.Time);
+				Writer.Write(Key.Value);
+			}
+
+			Writer.EndArray();
+		}
+	}
+};
+
+struct FGLTFJsonSkySphere : IGLTFJsonObject
+{
+	FString Name;
+
+	FGLTFJsonMeshIndex    SkySphereMesh;
+	FGLTFJsonTextureIndex SkyTexture;
+	FGLTFJsonTextureIndex CloudsTexture;
+	FGLTFJsonTextureIndex StarsTexture;
+	FGLTFJsonNodeIndex    DirectionalLight;
+
+	float SunHeight;
+	float SunBrightness;
+	float StarsBrightness;
+	float CloudSpeed;
+	float CloudOpacity;
+	float HorizonFalloff;
+
+	float SunRadius;
+	float NoisePower1;
+	float NoisePower2;
+
+	bool bColorsDeterminedBySunPosition;
+
+	FGLTFJsonColor4 ZenithColor;
+	FGLTFJsonColor4 HorizonColor;
+	FGLTFJsonColor4 CloudColor;
+	FGLTFJsonColor4 OverallColor;
+
+	FGLTFJsonSkySphereColorCurve ZenithColorCurve;
+	FGLTFJsonSkySphereColorCurve HorizonColorCurve;
+	FGLTFJsonSkySphereColorCurve CloudColorCurve;
+
+	FGLTFJsonVector3 Scale;
+
+	FGLTFJsonSkySphere()
+		: SunHeight(0)
+		, SunBrightness(0)
+		, StarsBrightness(0)
+		, CloudSpeed(0)
+		, CloudOpacity(0)
+		, HorizonFalloff(0)
+		, SunRadius(0)
+		, NoisePower1(0)
+		, NoisePower2(0)
+		, bColorsDeterminedBySunPosition(false)
+		, ZenithColor(FGLTFJsonColor4::White)
+		, HorizonColor(FGLTFJsonColor4::White)
+		, CloudColor(FGLTFJsonColor4::White)
+		, OverallColor(FGLTFJsonColor4::White)
+		, Scale(FGLTFJsonVector3::One)
+	{
+	}
+
+	virtual void WriteObject(IGLTFJsonWriter& Writer) const override
+	{
+		if (!Name.IsEmpty())
+		{
+			Writer.Write(TEXT("name"), Name);
+		}
+
+		Writer.Write(TEXT("skySphereMesh"), SkySphereMesh);
+		Writer.Write(TEXT("skyTexture"), SkyTexture);
+		Writer.Write(TEXT("cloudsTexture"), CloudsTexture);
+		Writer.Write(TEXT("starsTexture"), StarsTexture);
+
+		if (DirectionalLight != INDEX_NONE)
+		{
+			Writer.Write(TEXT("directionalLight"), DirectionalLight);
+		}
+
+		Writer.Write(TEXT("sunHeight"), SunHeight);
+		Writer.Write(TEXT("sunBrightness"), SunBrightness);
+		Writer.Write(TEXT("starsBrightness"), StarsBrightness);
+		Writer.Write(TEXT("cloudSpeed"), CloudSpeed);
+		Writer.Write(TEXT("cloudOpacity"), CloudOpacity);
+		Writer.Write(TEXT("horizonFalloff"), HorizonFalloff);
+
+		Writer.Write(TEXT("sunRadius"), SunRadius);
+		Writer.Write(TEXT("noisePower1"), NoisePower1);
+		Writer.Write(TEXT("noisePower2"), NoisePower2);
+
+		Writer.Write(TEXT("colorsDeterminedBySunPosition"), bColorsDeterminedBySunPosition);
+
+		Writer.Write(TEXT("zenithColor"), ZenithColor);
+		Writer.Write(TEXT("horizonColor"), HorizonColor);
+		Writer.Write(TEXT("cloudColor"), CloudColor);
+
+		if (!OverallColor.IsNearlyEqual(FGLTFJsonColor4::White, Writer.DefaultTolerance))
+		{
+			Writer.Write(TEXT("overallColor"), OverallColor);
+		}
+
+		if (ZenithColorCurve.ComponentCurves.Num() >= 3)
+		{
+			Writer.Write(TEXT("zenithColorCurve"), ZenithColorCurve);
+		}
+
+		if (HorizonColorCurve.ComponentCurves.Num() >= 3)
+		{
+			Writer.Write(TEXT("horizonColorCurve"), HorizonColorCurve);
+		}
+
+		if (CloudColorCurve.ComponentCurves.Num() >= 3)
+		{
+			Writer.Write(TEXT("cloudColorCurve"), CloudColorCurve);
+		}
+
+		if (!Scale.IsNearlyEqual(FGLTFJsonVector3::One, Writer.DefaultTolerance))
+		{
+			Writer.Write(TEXT("scale"), Scale);
+		}
+	}
+};
diff --git a/Engine/Plugins/Enterprise/GLTFExporter/MigratedFromGit/Source/GLTFExporter/Public/Json/GLTFJsonTexture.h b/Engine/Plugins/Enterprise/GLTFExporter/MigratedFromGit/Source/GLTFExporter/Public/Json/GLTFJsonTexture.h
new file mode 100644
--- /dev/null
+++ b/Engine/Plugins/Enterprise/GLTFExporter/MigratedFromGit/Source/GLTFExporter/Public/Json/GLTFJsonTexture.h
@@ -0,0 +1,51 @@
+// Copyright Epic Games, Inc. All Rights Reserved.
+
+#pragma once
+
+#include "Json/GLTFJsonObject.h"
+#include "Json/GLTFJsonIndex.h"
+
+struct FGLTFJsonTexture : IGLTFJsonObject
+{
+	FString Name;
+
+	FGLTFJsonSamplerIndex Sampler;
+
+	FGLTFJsonImageIndex Source;
+
+	EGLTFJsonHDREncoding Encoding;
+
+	FGLTFJsonTexture()
+		: Encoding(EGLTFJsonHDREncoding::None)
+	{
+	}
+
+	virtual void WriteObject(IGLTFJsonWriter& Writer) const override
+	{
+		if (!Name.IsEmpty())
+		{
+			Writer.Write(TEXT("name"), Name);
+		}
+
+		if (Sampler != INDEX_NONE)
+		{
+			Writer.Write(TEXT("sampler"), Sampler);
+		}
+
+		if (Source != INDEX_NONE)
+		{
+			Writer.Write(TEXT("source"), Source);
+		}
+
+		if (Encoding != EGLTFJsonHDREncoding::None)
+		{
+			Writer.StartExtensions();
+
+			Writer.StartExtension(EGLTFJsonExtension::EPIC_TextureHDREncoding);
+			Writer.Write(TEXT("encoding"), Encoding);
+			Writer.EndExtension();
+
+			Writer.EndExtensions();
+		}
+	}
+};
diff --git a/Engine/Plugins/Enterprise/GLTFExporter/MigratedFromGit/Source/GLTFExporter/Public/Json/GLTFJsonTextureTransform.h b/Engine/Plugins/Enterprise/GLTFExporter/MigratedFromGit/Source/GLTFExporter/Public/Json/GLTFJsonTextureTransform.h
new file mode 100644
--- /dev/null
+++ b/Engine/Plugins/Enterprise/GLTFExporter/MigratedFromGit/Source/GLTFExporter/Public/Json/GLTFJsonTextureTransform.h
@@ -0,0 +1,62 @@
+// Copyright Epic Games, Inc. All Rights Reserved.
+
+#pragma once
+
+#include "Json/GLTFJsonObject.h"
+#include "Json/GLTFJsonVector.h"
+
+struct FGLTFJsonTextureTransform : IGLTFJsonObject
+{
+	FGLTFJsonVector2 Offset;
+	FGLTFJsonVector2 Scale;
+	float Rotation;
+
+	FGLTFJsonTextureTransform()
+		: Offset(FGLTFJsonVector2::Zero)
+		, Scale(FGLTFJsonVector2::One)
+		, Rotation(0)
+	{
+	}
+
+	virtual void WriteObject(IGLTFJsonWriter& Writer) const override
+	{
+		if (!Offset.IsNearlyEqual(FGLTFJsonVector2::Zero, Writer.DefaultTolerance))
+		{
+			Writer.Write(TEXT("offset"), Offset);
+		}
+
+		if (!Scale.IsNearlyEqual(FGLTFJsonVector2::One, Writer.DefaultTolerance))
+		{
+			Writer.Write(TEXT("scale"), Scale);
+		}
+
+		if (!FMath::IsNearlyEqual(Rotation, 0, Writer.DefaultTolerance))
+		{
+			Writer.Write(TEXT("rotation"), Rotation);
+		}
+	}
+
+	bool IsNearlyEqual(const FGLTFJsonTextureTransform& Other, float Tolerance = KINDA_SMALL_NUMBER) const
+	{
+		return Offset.IsNearlyEqual(Other.Offset, Tolerance)
+			&& Scale.IsNearlyEqual(Other.Scale, Tolerance)
+			&& FMath::IsNearlyEqual(Rotation, Other.Rotation, Tolerance);
+	}
+
+	bool IsExactlyEqual(const FGLTFJsonTextureTransform& Other) const
+	{
+		return Offset.X == Other.Offset.X && Offset.Y == Other.Offset.Y
+			&& Scale.X == Other.Scale.X && Scale.Y == Other.Scale.Y
+			&& Rotation == Other.Rotation;
+	}
+
+	bool IsNearlyDefault(float Tolerance = KINDA_SMALL_NUMBER) const
+	{
+		return IsNearlyEqual(FGLTFJsonTextureTransform(), Tolerance);
+	}
+
+	bool IsExactlyDefault() const
+	{
+		return IsExactlyEqual(FGLTFJsonTextureTransform());
+	}
+};
diff --git a/Engine/Plugins/Enterprise/GLTFExporter/MigratedFromGit/Source/GLTFExporter/Public/Json/GLTFJsonUtility.h b/Engine/Plugins/Enterprise/GLTFExporter/MigratedFromGit/Source/GLTFExporter/Public/Json/GLTFJsonUtility.h
new file mode 100644
--- /dev/null
+++ b/Engine/Plugins/Enterprise/GLTFExporter/MigratedFromGit/Source/GLTFExporter/Public/Json/GLTFJsonUtility.h
@@ -0,0 +1,276 @@
+// Copyright Epic Games, Inc. All Rights Reserved.
+
+#pragma once
+
+#include "Json/GLTFJsonEnums.h"
+
+struct FGLTFJsonUtility
+{
+	template <typename EnumType, typename = typename TEnableIf<TIsEnum<EnumType>::Value>::Type>
+	static int32 GetValue(EnumType Enum)
+	{
+		return static_cast<int32>(Enum);
+	}
+
+	static const TCHAR* GetValue(EGLTFJsonExtension Enum)
+	{
+		switch (Enum)
+		{
+			case EGLTFJsonExtension::KHR_LightsPunctual:      return TEXT("KHR_lights_punctual");
+			case EGLTFJsonExtension::KHR_MaterialsClearCoat:  return TEXT("KHR_materials_clearcoat");
+			case EGLTFJsonExtension::KHR_MaterialsUnlit:      return TEXT("KHR_materials_unlit");
+			case EGLTFJsonExtension::KHR_MaterialsVariants:   return TEXT("KHR_materials_variants");
+			case EGLTFJsonExtension::KHR_MeshQuantization:    return TEXT("KHR_mesh_quantization");
+			case EGLTFJsonExtension::KHR_TextureTransform:    return TEXT("KHR_texture_transform");
+			case EGLTFJsonExtension::EPIC_AnimationHotspots:  return TEXT("EPIC_animation_hotspots");
+			case EGLTFJsonExtension::EPIC_AnimationPlayback:  return TEXT("EPIC_animation_playback");
+			case EGLTFJsonExtension::EPIC_BlendModes:         return TEXT("EPIC_blend_modes");
+			case EGLTFJsonExtension::EPIC_CameraControls:     return TEXT("EPIC_camera_controls");
+			case EGLTFJsonExtension::EPIC_HDRIBackdrops:      return TEXT("EPIC_hdri_backdrops");
+			case EGLTFJsonExtension::EPIC_LevelVariantSets:   return TEXT("EPIC_level_variant_sets");
+			case EGLTFJsonExtension::EPIC_LightmapTextures:   return TEXT("EPIC_lightmap_textures");
+			case EGLTFJsonExtension::EPIC_SkySpheres:         return TEXT("EPIC_sky_spheres");
+			case EGLTFJsonExtension::EPIC_TextureHDREncoding: return TEXT("EPIC_texture_hdr_encoding");
+			default:
+				checkNoEntry();
+				return TEXT("");
+		}
+	}
+
+	static const TCHAR* GetValue(EGLTFJsonAlphaMode Enum)
+	{
+		switch (Enum)
+		{
+			case EGLTFJsonAlphaMode::Opaque: return TEXT("OPAQUE");
+			case EGLTFJsonAlphaMode::Blend:  return TEXT("BLEND");
+			case EGLTFJsonAlphaMode::Mask:   return TEXT("MASK");
+			default:
+				checkNoEntry();
+				return TEXT("");
+		}
+	}
+
+	static const TCHAR* GetValue(EGLTFJsonBlendMode Enum)
+	{
+		switch (Enum)
+		{
+			case EGLTFJsonBlendMode::Additive:       return TEXT("ADDITIVE");
+			case EGLTFJsonBlendMode::Modulate:       return TEXT("MODULATE");
+			case EGLTFJsonBlendMode::AlphaComposite: return TEXT("ALPHACOMPOSITE");
+			default:
+				checkNoEntry();
+				return TEXT("");
+		}
+	}
+
+	static const TCHAR* GetValue(EGLTFJsonMimeType Enum)
+	{
+		switch (Enum)
+		{
+			case EGLTFJsonMimeType::PNG:  return TEXT("image/png");
+			case EGLTFJsonMimeType::JPEG: return TEXT("image/jpeg");
+			default:
+				checkNoEntry();
+				return TEXT("");
+		}
+	}
+
+	static const TCHAR* GetValue(EGLTFJsonAccessorType Enum)
+	{
+		switch (Enum)
+		{
+			case EGLTFJsonAccessorType::Scalar: return TEXT("SCALAR");
+			case EGLTFJsonAccessorType::Vec2:   return TEXT("VEC2");
+			case EGLTFJsonAccessorType::Vec3:   return TEXT("VEC3");
+			case EGLTFJsonAccessorType::Vec4:   return TEXT("VEC4");
+			case EGLTFJsonAccessorType::Mat2:   return TEXT("MAT2");
+			case EGLTFJsonAccessorType::Mat3:   return TEXT("MAT3");
+			case EGLTFJsonAccessorType::Mat4:   return TEXT("MAT4");
+			default:
+				checkNoEntry();
+				return TEXT("");
+		}
+	}
+
+	static const TCHAR* GetValue(EGLTFJsonHDREncoding Enum)
+	{
+		switch (Enum)
+		{
+			case EGLTFJsonHDREncoding::RGBM: return TEXT("RGBM");
+			case EGLTFJsonHDREncoding::RGBE: return TEXT("RGBE");
+			default:
+				checkNoEntry();
+				return TEXT("");
+		}
+	}
+
+	static const TCHAR* GetValue(EGLTFJsonCubeFace Enum)
+	{
+		switch (Enum)
+		{
+			case EGLTFJsonCubeFace::PosX: return TEXT("PosX");
+			case EGLTFJsonCubeFace::NegX: return TEXT("NegX");
+			case EGLTFJsonCubeFace::PosY: return TEXT("PosY");
+			case EGLTFJsonCubeFace::NegY: return TEXT("NegY");
+			case EGLTFJsonCubeFace::PosZ: return TEXT("PosZ");
+			case EGLTFJsonCubeFace::NegZ: return TEXT("NegZ");
+			default:
+				checkNoEntry();
+				return TEXT("");
+		}
+	}
+
+	static const TCHAR* GetValue(EGLTFJsonCameraType Enum)
+	{
+		switch (Enum)
+		{
+			case EGLTFJsonCameraType::Perspective:  return TEXT("perspective");
+			case EGLTFJsonCameraType::Orthographic: return TEXT("orthographic");
+			default:
+				checkNoEntry();
+				return TEXT("");
+		}
+	}
+
+	static const TCHAR* GetValue(EGLTFJsonLightType Enum)
+	{
+		switch (Enum)
+		{
+			case EGLTFJsonLightType::Directional: return TEXT("directional");
+			case EGLTFJsonLightType::Point:       return TEXT("point");
+			case EGLTFJsonLightType::Spot:        return TEXT("spot");
+			default:
+				checkNoEntry();
+				return TEXT("");
+		}
+	}
+
+	static const TCHAR* GetValue(EGLTFJsonInterpolation Enum)
+	{
+		switch (Enum)
+		{
+			case EGLTFJsonInterpolation::Linear:      return TEXT("LINEAR");
+			case EGLTFJsonInterpolation::Step:        return TEXT("STEP");
+			case EGLTFJsonInterpolation::CubicSpline: return TEXT("CUBICSPLINE");
+			default:
+				checkNoEntry();
+				return TEXT("");
+		}
+	}
+
+	static const TCHAR* GetValue(EGLTFJsonTargetPath Enum)
+	{
+		switch (Enum)
+		{
+			case EGLTFJsonTargetPath::Translation: return TEXT("translation");
+			case EGLTFJsonTargetPath::Rotation:    return TEXT("rotation");
+			case EGLTFJsonTargetPath::Scale:       return TEXT("scale");
+			case EGLTFJsonTargetPath::Weights:     return TEXT("weights");
+			default:
+				checkNoEntry();
+				return TEXT("");
+		}
+	}
+
+	static const TCHAR* GetValue(EGLTFJsonCameraControlMode Enum)
+	{
+		switch (Enum)
+		{
+			case EGLTFJsonCameraControlMode::FreeLook: return TEXT("freeLook");
+			case EGLTFJsonCameraControlMode::Orbital:  return TEXT("orbital");
+			default:
+				checkNoEntry();
+				return TEXT("");
+		}
+	}
+
+	static const TCHAR* GetValue(EGLTFJsonShadingModel Enum)
+	{
+		switch (Enum)
+		{
+			case EGLTFJsonShadingModel::Default:   return TEXT("Default");
+			case EGLTFJsonShadingModel::Unlit:     return TEXT("Unlit");
+			case EGLTFJsonShadingModel::ClearCoat: return TEXT("ClearCoat");
+			default:
+				checkNoEntry();
+				return TEXT("");
+		}
+	}
+
+	static EGLTFJsonComponentType ParseComponentType(int32 ComponentType)
+	{
+		switch (ComponentType)
+		{
+		case EGLTFJsonComponentType::S8:
+		case EGLTFJsonComponentType::U8:
+		case EGLTFJsonComponentType::S16:
+		case EGLTFJsonComponentType::U16:
+		case EGLTFJsonComponentType::S32:
+		case EGLTFJsonComponentType::U32:
+		case EGLTFJsonComponentType::F32:
+			return static_cast<EGLTFJsonComponentType>(ComponentType);
+		default:
+			return EGLTFJsonComponentType::None;
+		}
+	}
+
+	static EGLTFJsonAccessorType ParseAccessorType(const FString& AccessorType)
+	{
+		if (AccessorType.Equals(TEXT("SCALAR")))
+		{
+			return EGLTFJsonAccessorType::Scalar;
+		}
+		if (AccessorType.Equals(TEXT("VEC2")))
+		{
+			return EGLTFJsonAccessorType::Vec2;
+		}
+		if (AccessorType.Equals(TEXT("VEC3")))
+		{
+			return EGLTFJsonAccessorType::Vec3;
+		}
+		if (AccessorType.Equals(TEXT("VEC4")))
+		{
+			return EGLTFJsonAccessorType::Vec4;
+		}
+		if (AccessorType.Equals(TEXT("MAT2")))
+		{
+			return EGLTFJsonAccessorType::Mat2;
+		}
+		if (AccessorType.Equals(TEXT("MAT3")))
+		{
+			return EGLTFJsonAccessorType::Mat3;
+		}
+		if (AccessorType.Equals(TEXT("MAT4")))
+		{
+			return EGLTFJsonAccessorType::Mat4;
+		}
+
+		return EGLTFJsonAccessorType::None;
+	}
+
+	static EGLTFJsonBufferTarget ParseBufferTarget(int32 BufferTarget)
+	{
+		switch (BufferTarget)
+		{
+		case EGLTFJsonBufferTarget::ArrayBuffer:
+		case EGLTFJsonBufferTarget::ElementArrayBuffer:
+			return static_cast<EGLTFJsonBufferTarget>(BufferTarget);
+		default:
+			return EGLTFJsonBufferTarget::None;
+		}
+	}
+
+	static EGLTFJsonMimeType ParseMimeType(const FString& MimeType)
+	{
+		if (MimeType.Equals(TEXT("image/jpeg")))
+		{
+			return EGLTFJsonMimeType::JPEG;
+		}
+		if (MimeType.Equals(TEXT("image/png")))
+		{
+			return EGLTFJsonMimeType::PNG;
+		}
+
+		return EGLTFJsonMimeType::None;
+	}
+};
diff --git a/Engine/Plugins/Enterprise/GLTFExporter/MigratedFromGit/Source/GLTFExporter/Public/Json/GLTFJsonValue.h b/Engine/Plugins/Enterprise/GLTFExporter/MigratedFromGit/Source/GLTFExporter/Public/Json/GLTFJsonValue.h
new file mode 100644
--- /dev/null
+++ b/Engine/Plugins/Enterprise/GLTFExporter/MigratedFromGit/Source/GLTFExporter/Public/Json/GLTFJsonValue.h
@@ -0,0 +1,12 @@
+// Copyright Epic Games, Inc. All Rights Reserved.
+
+#pragma once
+
+class IGLTFJsonWriter;
+
+struct IGLTFJsonValue
+{
+	virtual ~IGLTFJsonValue() = default;
+
+	virtual void WriteValue(IGLTFJsonWriter& Writer) const = 0;
+};
diff --git a/Engine/Plugins/Enterprise/GLTFExporter/MigratedFromGit/Source/GLTFExporter/Public/Json/GLTFJsonVector.cpp b/Engine/Plugins/Enterprise/GLTFExporter/MigratedFromGit/Source/GLTFExporter/Public/Json/GLTFJsonVector.cpp
new file mode 100644
--- /dev/null
+++ b/Engine/Plugins/Enterprise/GLTFExporter/MigratedFromGit/Source/GLTFExporter/Public/Json/GLTFJsonVector.cpp
@@ -0,0 +1,12 @@
+// Copyright Epic Games, Inc. All Rights Reserved.
+
+#include "Json/GLTFJsonVector.h"
+
+const FGLTFJsonVector2 FGLTFJsonVector2::Zero({ 0, 0 });
+const FGLTFJsonVector2 FGLTFJsonVector2::One ({ 1, 1 });
+
+const FGLTFJsonVector3 FGLTFJsonVector3::Zero({ 0, 0, 0 });
+const FGLTFJsonVector3 FGLTFJsonVector3::One ({ 1, 1, 1 });
+
+const FGLTFJsonVector4 FGLTFJsonVector4::Zero({ 0, 0, 0, 0 });
+const FGLTFJsonVector4 FGLTFJsonVector4::One ({ 1, 1, 1, 1 });
diff --git a/Engine/Plugins/Enterprise/GLTFExporter/MigratedFromGit/Source/GLTFExporter/Public/Json/GLTFJsonVector.h b/Engine/Plugins/Enterprise/GLTFExporter/MigratedFromGit/Source/GLTFExporter/Public/Json/GLTFJsonVector.h
new file mode 100644
--- /dev/null
+++ b/Engine/Plugins/Enterprise/GLTFExporter/MigratedFromGit/Source/GLTFExporter/Public/Json/GLTFJsonVector.h
@@ -0,0 +1,90 @@
+// Copyright Epic Games, Inc. All Rights Reserved.
+
+#pragma once
+
+#include "Core/GLTFVector.h"
+#include "Json/GLTFJsonArray.h"
+
+template <typename BaseType>
+struct TGLTFJsonVector : BaseType, IGLTFJsonArray
+{
+	static const TGLTFJsonVector Zero;
+	static const TGLTFJsonVector One;
+
+	TGLTFJsonVector(const BaseType& Other)
+		: BaseType(Other)
+	{
+	}
+
+	TGLTFJsonVector& operator=(const BaseType& Other)
+	{
+		*static_cast<BaseType*>(this) = Other;
+		return *this;
+	}
+
+	bool operator==(const BaseType& Other) const
+	{
+		for (SIZE_T i = 0; i < GetNum(BaseType::Components); ++i)
+		{
+			if (BaseType::Components[i] != Other.Components[i])
+			{
+				return false;
+			}
+		}
+
+		return true;
+	}
+
+	bool operator!=(const BaseType& Other) const
+	{
+		return !(*this == Other);
+	}
+
+	virtual void WriteArray(IGLTFJsonWriter& Writer) const override
+	{
+		for (SIZE_T i = 0; i < GetNum(BaseType::Components); ++i)
+		{
+			Writer.Write(BaseType::Components[i]);
+		}
+	}
+
+	bool IsNearlyEqual(const BaseType& Other, float Tolerance = KINDA_SMALL_NUMBER) const
+	{
+		for (SIZE_T i = 0; i < GetNum(BaseType::Components); ++i)
+		{
+			if (!FMath::IsNearlyEqual(BaseType::Components[i], Other.Components[i], Tolerance))
+			{
+				return false;
+			}
+		}
+
+		return true;
+	}
+};
+
+struct FGLTFJsonVector2 : TGLTFJsonVector<FGLTFVector2>
+{
+	static const FGLTFJsonVector2 Zero;
+	static const FGLTFJsonVector2 One;
+
+	using TGLTFJsonVector::TGLTFJsonVector;
+	using TGLTFJsonVector::operator=;
+};
+
+struct FGLTFJsonVector3 : TGLTFJsonVector<FGLTFVector3>
+{
+	static const FGLTFJsonVector3 Zero;
+	static const FGLTFJsonVector3 One;
+
+	using TGLTFJsonVector::TGLTFJsonVector;
+	using TGLTFJsonVector::operator=;
+};
+
+struct FGLTFJsonVector4 : TGLTFJsonVector<FGLTFVector4>
+{
+	static const FGLTFJsonVector4 Zero;
+	static const FGLTFJsonVector4 One;
+
+	using TGLTFJsonVector::TGLTFJsonVector;
+	using TGLTFJsonVector::operator=;
+};
diff --git a/Engine/Plugins/Enterprise/GLTFExporter/MigratedFromGit/Source/GLTFExporter/Public/Json/GLTFJsonWriter.cpp b/Engine/Plugins/Enterprise/GLTFExporter/MigratedFromGit/Source/GLTFExporter/Public/Json/GLTFJsonWriter.cpp
new file mode 100644
--- /dev/null
+++ b/Engine/Plugins/Enterprise/GLTFExporter/MigratedFromGit/Source/GLTFExporter/Public/Json/GLTFJsonWriter.cpp
@@ -0,0 +1,167 @@
+// Copyright Epic Games, Inc. All Rights Reserved.
+
+#include "Json/GLTFJsonWriter.h"
+#include "Serialization/JsonSerializer.h"
+#include "Policies/PrettyJsonPrintPolicy.h"
+#include "Policies/CondensedJsonPrintPolicy.h"
+
+template <class CharType, class PrintPolicy>
+class TGLTFJsonWriterImpl final : public IGLTFJsonWriter
+{
+public:
+
+	TGLTFJsonWriterImpl(FArchive& Archive, FGLTFJsonExtensions& Extensions)
+		: IGLTFJsonWriter(Extensions)
+		, JsonWriter(TJsonWriterFactory<CharType, PrintPolicy>::Create(&Archive))
+	{
+	}
+
+	virtual void Close() override
+	{
+		JsonWriter->Close();
+	}
+
+	virtual void Write(bool Boolean) override
+	{
+		if (CurrentIdentifier.IsEmpty())
+		{
+			JsonWriter->WriteValue(Boolean);
+		}
+		else
+		{
+			JsonWriter->WriteValue(CurrentIdentifier, Boolean);
+			CurrentIdentifier.Empty();
+		}
+	}
+
+	virtual void Write(int32 Number) override
+	{
+		if (CurrentIdentifier.IsEmpty())
+		{
+			JsonWriter->WriteValue(Number);
+		}
+		else
+		{
+			JsonWriter->WriteValue(CurrentIdentifier, Number);
+			CurrentIdentifier.Empty();
+		}
+	}
+
+	virtual void Write(int64 Number) override
+	{
+		if (CurrentIdentifier.IsEmpty())
+		{
+			JsonWriter->WriteValue(Number);
+		}
+		else
+		{
+			JsonWriter->WriteValue(CurrentIdentifier, Number);
+			CurrentIdentifier.Empty();
+		}
+	}
+
+	virtual void Write(float Number) override
+	{
+		// NOTE: Specifying 9 significant digits, this ensures no precision is lost
+		FString ExactString = FString::Printf(TEXT("%.9g"), Number);
+
+		if (CurrentIdentifier.IsEmpty())
+		{
+			JsonWriter->WriteRawJSONValue(ExactString);
+		}
+		else
+		{
+			JsonWriter->WriteRawJSONValue(CurrentIdentifier, ExactString);
+			CurrentIdentifier.Empty();
+		}
+	}
+
+	virtual void Write(const FString& String) override
+	{
+		if (CurrentIdentifier.IsEmpty())
+		{
+			JsonWriter->WriteValue(String);
+		}
+		else
+		{
+			JsonWriter->WriteValue(CurrentIdentifier, String);
+			CurrentIdentifier.Empty();
+		}
+	}
+
+	virtual void Write(TYPE_OF_NULLPTR) override
+	{
+		if (CurrentIdentifier.IsEmpty())
+		{
+			JsonWriter->WriteValue(nullptr);
+		}
+		else
+		{
+			JsonWriter->WriteValue(CurrentIdentifier, nullptr);
+			CurrentIdentifier.Empty();
+		}
+	}
+
+	virtual void SetIdentifier(const FString& Identifier) override
+	{
+		this->CurrentIdentifier = Identifier;
+	}
+
+	virtual void StartObject() override
+	{
+		if (CurrentIdentifier.IsEmpty())
+		{
+			JsonWriter->WriteObjectStart();
+		}
+		else
+		{
+			JsonWriter->WriteObjectStart(CurrentIdentifier);
+			CurrentIdentifier.Empty();
+		}
+	}
+
+	virtual void EndObject() override
+	{
+		JsonWriter->WriteObjectEnd();
+	}
+
+	virtual void StartArray() override
+	{
+		if (CurrentIdentifier.IsEmpty())
+		{
+			JsonWriter->WriteArrayStart();
+		}
+		else
+		{
+			JsonWriter->WriteArrayStart(CurrentIdentifier);
+			CurrentIdentifier.Empty();
+		}
+	}
+
+	virtual void EndArray() override
+	{
+		JsonWriter->WriteArrayEnd();
+	}
+
+private:
+
+	FString CurrentIdentifier;
+	TSharedRef<TJsonWriter<CharType, PrintPolicy>> JsonWriter;
+};
+
+template <class CharType>
+struct TCrossPlatformPrettyJsonPrintPolicy : TPrettyJsonPrintPolicy<CharType>
+{
+	static void WriteLineTerminator(FArchive* Stream)
+	{
+		TJsonPrintPolicy<CharType>::WriteString(Stream, TEXT("\n")); // don't use system dependent line terminator
+	}
+};
+
+TSharedRef<IGLTFJsonWriter> IGLTFJsonWriter::Create(FArchive& Archive, bool bPrettyJson, FGLTFJsonExtensions& Extensions)
+{
+	return MakeShareable(bPrettyJson
+		? static_cast<IGLTFJsonWriter*>(new TGLTFJsonWriterImpl<UTF8CHAR, TCrossPlatformPrettyJsonPrintPolicy<UTF8CHAR>>(Archive, Extensions))
+		: static_cast<IGLTFJsonWriter*>(new TGLTFJsonWriterImpl<UTF8CHAR, TCondensedJsonPrintPolicy<UTF8CHAR>>(Archive, Extensions))
+	);
+}
diff --git a/Engine/Plugins/Enterprise/GLTFExporter/MigratedFromGit/Source/GLTFExporter/Public/Json/GLTFJsonWriter.h b/Engine/Plugins/Enterprise/GLTFExporter/MigratedFromGit/Source/GLTFExporter/Public/Json/GLTFJsonWriter.h
new file mode 100644
--- /dev/null
+++ b/Engine/Plugins/Enterprise/GLTFExporter/MigratedFromGit/Source/GLTFExporter/Public/Json/GLTFJsonWriter.h
@@ -0,0 +1,262 @@
+// Copyright Epic Games, Inc. All Rights Reserved.
+
+#pragma once
+
+#include "Json/GLTFJsonValue.h"
+#include "Json/GLTFJsonUtility.h"
+#include "Json/GLTFJsonExtensions.h"
+
+class IGLTFJsonWriter
+{
+public:
+
+	float DefaultTolerance;
+
+	FGLTFJsonExtensions& Extensions;
+
+	IGLTFJsonWriter(FGLTFJsonExtensions& Extensions)
+		: DefaultTolerance(0)
+		, Extensions(Extensions)
+	{
+	}
+
+	static TSharedRef<IGLTFJsonWriter> Create(FArchive& Archive, bool bPrettyJson, FGLTFJsonExtensions& Extensions);
+
+	virtual ~IGLTFJsonWriter() = default;
+
+	virtual void Close() = 0;
+
+	virtual void Write(bool Boolean) = 0;
+	virtual void Write(int32 Number) = 0;
+	virtual void Write(int64 Number) = 0;
+	virtual void Write(float Number) = 0;
+	virtual void Write(const FString& String) = 0;
+	virtual void Write(TYPE_OF_NULLPTR) = 0;
+	virtual void SetIdentifier(const FString& Identifier) = 0;
+
+	virtual void StartObject() = 0;
+	virtual void EndObject() = 0;
+	virtual void StartArray() = 0;
+	virtual void EndArray() = 0;
+
+	void Write(const TCHAR* String)
+	{
+		Write(FString(String));
+	}
+
+	void Write(const IGLTFJsonValue& Value)
+	{
+		Value.WriteValue(*this);
+	}
+
+	void Write(const IGLTFJsonValue* Value)
+	{
+		if (Value != nullptr)
+		{
+			Value->WriteValue(*this);
+		}
+		else
+		{
+			Write(nullptr);
+		}
+	}
+
+	template <typename EnumType, typename = typename TEnableIf<TIsEnum<EnumType>::Value>::Type>
+	void Write(EnumType Enum)
+	{
+		Write(FGLTFJsonUtility::GetValue(Enum));
+	}
+
+	template <class ElementType, typename = typename TEnableIf<TNot<TIsSame<ElementType, TCHAR>>::Value>::Type>
+	void Write(const ElementType* Array, SIZE_T ArraySize)
+	{
+		StartArray();
+		for (SIZE_T i = 0; i < ArraySize; ++i)
+		{
+			Write(Array[i]);
+		}
+		EndArray();
+	}
+
+	template <class ElementType, SIZE_T ArraySize, typename = typename TEnableIf<TNot<TIsSame<ElementType, TCHAR>>::Value>::Type>
+	void Write(const ElementType (&Array)[ArraySize])
+	{
+		Write(Array, ArraySize);
+	}
+
+	template <class ElementType>
+	void Write(const TUniquePtr<ElementType>& Pointer)
+	{
+		if (Pointer.IsValid())
+		{
+			Write(*Pointer.Get());
+		}
+		else
+		{
+			Write(nullptr);
+		}
+	}
+
+	template <class ElementType>
+	void Write(const TArray<ElementType>& Array)
+	{
+		StartArray();
+		for (const auto& Element : Array)
+		{
+			Write(Element);
+		}
+		EndArray();
+	}
+
+	template <class ElementType>
+	void Write(const TSet<ElementType>& Set)
+	{
+		StartArray();
+		for (const auto& Element : Set)
+		{
+			Write(Element);
+		}
+		EndArray();
+	}
+
+	void Write(const FString& Identifier, bool Boolean)
+	{
+		SetIdentifier(Identifier);
+		Write(Boolean);
+	}
+
+	void Write(const FString& Identifier, int32 Number)
+	{
+		SetIdentifier(Identifier);
+		Write(Number);
+	}
+
+	void Write(const FString& Identifier, int64 Number)
+	{
+		SetIdentifier(Identifier);
+		Write(Number);
+	}
+
+	void Write(const FString& Identifier, float Number)
+	{
+		SetIdentifier(Identifier);
+		Write(Number);
+	}
+
+	void Write(const FString& Identifier, const FString& String)
+	{
+		SetIdentifier(Identifier);
+		Write(String);
+	}
+
+	void Write(const FString& Identifier, const TCHAR* String)
+	{
+		SetIdentifier(Identifier);
+		Write(String);
+	}
+
+	void Write(const FString& Identifier, const IGLTFJsonValue& Value)
+	{
+		SetIdentifier(Identifier);
+		Write(Value);
+	}
+
+	void Write(const FString& Identifier, const IGLTFJsonValue* Value)
+	{
+		SetIdentifier(Identifier);
+		Write(Value);
+	}
+
+	template <typename EnumType, typename = typename TEnableIf<TIsEnum<EnumType>::Value>::Type>
+	void Write(const FString& Identifier, EnumType Enum)
+	{
+		SetIdentifier(Identifier);
+		Write(Enum);
+	}
+
+	template <class ElementType, typename = typename TEnableIf<TNot<TIsSame<ElementType, TCHAR>>::Value>::Type>
+	void Write(const FString& Identifier, const ElementType* Array, SIZE_T ArraySize)
+	{
+		SetIdentifier(Identifier);
+		Write(Array, ArraySize);
+	}
+
+	template <class ElementType, SIZE_T ArraySize, typename = typename TEnableIf<TNot<TIsSame<ElementType, TCHAR>>::Value>::Type>
+	void Write(const FString& Identifier, const ElementType (&Array)[ArraySize])
+	{
+		SetIdentifier(Identifier);
+		Write(Array);
+	}
+
+	template <class ElementType>
+	void Write(const FString& Identifier, const TUniquePtr<ElementType>& Pointer)
+	{
+		SetIdentifier(Identifier);
+		Write(Pointer);
+	}
+
+	template <class ElementType>
+	void Write(const FString& Identifier, const TArray<ElementType>& Array)
+	{
+		SetIdentifier(Identifier);
+		Write(Array);
+	}
+
+	template <class ElementType>
+	void Write(const FString& Identifier, const TSet<ElementType>& Set)
+	{
+		SetIdentifier(Identifier);
+		Write(Set);
+	}
+
+	void StartObject(const FString& Identifier)
+	{
+		SetIdentifier(Identifier);
+		StartObject();
+	}
+
+	void StartArray(const FString& Identifier)
+	{
+		SetIdentifier(Identifier);
+		StartArray();
+	}
+
+	void StartExtensions()
+	{
+		StartObject(TEXT("extensions"));
+	}
+
+	void EndExtensions()
+	{
+		EndObject();
+	}
+
+	void Write(const EGLTFJsonExtension& Extension, const IGLTFJsonValue& Value, bool bIsRequired = false)
+	{
+		Extensions.Used.Add(Extension);
+
+		if (bIsRequired)
+		{
+			Extensions.Required.Add(Extension);
+		}
+
+		Write(FGLTFJsonUtility::GetValue(Extension), Value);
+	}
+
+	void StartExtension(const EGLTFJsonExtension& Extension, bool bIsRequired = false)
+	{
+		Extensions.Used.Add(Extension);
+
+		if (bIsRequired)
+		{
+			Extensions.Required.Add(Extension);
+		}
+
+		StartObject(FGLTFJsonUtility::GetValue(Extension));
+	}
+
+	void EndExtension()
+	{
+		EndObject();
+	}
+};
