[GLTFExporter] feat: make core types for the exporter public

commit f80a22e64127c4d3f9d239b1afe16f942a187afe
Date: Tue, 11 Jan 2022 15:12:31 +0100

    feat: make core types for the exporter public
    (cherry picked from commit e19d85f0a3196729c1715947b80c6909f361542e​)

#skipci

diff --git a/Engine/Plugins/Enterprise/GLTFExporter/MigratedFromGit/Source/GLTFExporter/Private/Core/GLTFColor.h b/Engine/Plugins/Enterprise/GLTFExporter/MigratedFromGit/Source/GLTFExporter/Private/Core/GLTFColor.h
deleted file mode 100644
diff --git a/Engine/Plugins/Enterprise/GLTFExporter/MigratedFromGit/Source/GLTFExporter/Private/Core/GLTFMatrix.h b/Engine/Plugins/Enterprise/GLTFExporter/MigratedFromGit/Source/GLTFExporter/Private/Core/GLTFMatrix.h
deleted file mode 100644
diff --git a/Engine/Plugins/Enterprise/GLTFExporter/MigratedFromGit/Source/GLTFExporter/Private/Core/GLTFQuaternion.h b/Engine/Plugins/Enterprise/GLTFExporter/MigratedFromGit/Source/GLTFExporter/Private/Core/GLTFQuaternion.h
deleted file mode 100644
diff --git a/Engine/Plugins/Enterprise/GLTFExporter/MigratedFromGit/Source/GLTFExporter/Private/Core/GLTFVector.h b/Engine/Plugins/Enterprise/GLTFExporter/MigratedFromGit/Source/GLTFExporter/Private/Core/GLTFVector.h
deleted file mode 100644
diff --git a/Engine/Plugins/Enterprise/GLTFExporter/MigratedFromGit/Source/GLTFExporter/Public/Core/GLTFColor.h b/Engine/Plugins/Enterprise/GLTFExporter/MigratedFromGit/Source/GLTFExporter/Public/Core/GLTFColor.h
new file mode 100644
--- /dev/null
+++ b/Engine/Plugins/Enterprise/GLTFExporter/MigratedFromGit/Source/GLTFExporter/Public/Core/GLTFColor.h
@@ -0,0 +1,73 @@
+// Copyright Epic Games, Inc. All Rights Reserved.
+
+#pragma once
+
+#include "CoreTypes.h"
+
+#if PLATFORM_SUPPORTS_PRAGMA_PACK
+	#pragma pack(push, 1)
+#endif
+
+template <typename ComponentType>
+struct TGLTFColor3
+{
+	union
+	{
+		struct
+		{
+			ComponentType R GCC_PACK(1);
+			ComponentType G GCC_PACK(1);
+			ComponentType B GCC_PACK(1);
+		};
+
+		ComponentType Components[3] GCC_PACK(1);
+	};
+
+	TGLTFColor3(
+		ComponentType R,
+		ComponentType G,
+		ComponentType B
+	)
+		: R(R)
+		, G(G)
+		, B(B)
+	{}
+};
+
+template <typename ComponentType>
+struct TGLTFColor4
+{
+	union
+	{
+		struct
+		{
+			ComponentType R GCC_PACK(1);
+			ComponentType G GCC_PACK(1);
+			ComponentType B GCC_PACK(1);
+			ComponentType A GCC_PACK(1);
+		};
+
+		ComponentType Components[4] GCC_PACK(1);
+	};
+
+	TGLTFColor4(
+		ComponentType R,
+		ComponentType G,
+		ComponentType B,
+		ComponentType A
+	)
+		: R(R)
+		, G(G)
+		, B(B)
+		, A(A)
+	{}
+};
+
+#if PLATFORM_SUPPORTS_PRAGMA_PACK
+	#pragma pack(pop)
+#endif
+
+typedef TGLTFColor3<float> FGLTFColor3;
+typedef TGLTFColor4<float> FGLTFColor4;
+
+typedef TGLTFColor4<uint8> FGLTFUInt8Color4;
diff --git a/Engine/Plugins/Enterprise/GLTFExporter/MigratedFromGit/Source/GLTFExporter/Public/Core/GLTFMatrix.h b/Engine/Plugins/Enterprise/GLTFExporter/MigratedFromGit/Source/GLTFExporter/Public/Core/GLTFMatrix.h
new file mode 100644
--- /dev/null
+++ b/Engine/Plugins/Enterprise/GLTFExporter/MigratedFromGit/Source/GLTFExporter/Public/Core/GLTFMatrix.h
@@ -0,0 +1,33 @@
+// Copyright Epic Games, Inc. All Rights Reserved.
+
+#pragma once
+
+#include "CoreTypes.h"
+
+#if PLATFORM_SUPPORTS_PRAGMA_PACK
+	#pragma pack(push, 1)
+#endif
+
+template <typename ElementType, SIZE_T Dimensions>
+struct TGLTFMatrix
+{
+	ElementType Elements[Dimensions * Dimensions] GCC_PACK(1);
+
+	ElementType& operator()(uint32 Column, uint32 Row)
+	{
+		return Elements[Column * Dimensions + Row];
+	}
+
+	const ElementType& operator()(uint32 Column, uint32 Row) const
+	{
+		return Elements[Column * Dimensions + Row];
+	}
+};
+
+#if PLATFORM_SUPPORTS_PRAGMA_PACK
+	#pragma pack(pop)
+#endif
+
+typedef TGLTFMatrix<float, 2> FGLTFMatrix2;
+typedef TGLTFMatrix<float, 3> FGLTFMatrix3;
+typedef TGLTFMatrix<float, 4> FGLTFMatrix4;
diff --git a/Engine/Plugins/Enterprise/GLTFExporter/MigratedFromGit/Source/GLTFExporter/Public/Core/GLTFQuaternion.h b/Engine/Plugins/Enterprise/GLTFExporter/MigratedFromGit/Source/GLTFExporter/Public/Core/GLTFQuaternion.h
new file mode 100644
--- /dev/null
+++ b/Engine/Plugins/Enterprise/GLTFExporter/MigratedFromGit/Source/GLTFExporter/Public/Core/GLTFQuaternion.h
@@ -0,0 +1,44 @@
+// Copyright Epic Games, Inc. All Rights Reserved.
+
+#pragma once
+
+#include "CoreTypes.h"
+
+#if PLATFORM_SUPPORTS_PRAGMA_PACK
+	#pragma pack(push, 1)
+#endif
+
+template <typename ComponentType>
+struct TGLTFQuaternion
+{
+	union
+	{
+		struct
+		{
+			ComponentType X GCC_PACK(1);
+			ComponentType Y GCC_PACK(1);
+			ComponentType Z GCC_PACK(1);
+			ComponentType W GCC_PACK(1);
+		};
+
+		ComponentType Components[4] GCC_PACK(1);
+	};
+
+	TGLTFQuaternion(
+		ComponentType X,
+		ComponentType Y,
+		ComponentType Z,
+		ComponentType W
+	)
+		: X(X)
+		, Y(Y)
+		, Z(Z)
+		, W(W)
+	{}
+};
+
+#if PLATFORM_SUPPORTS_PRAGMA_PACK
+	#pragma pack(pop)
+#endif
+
+typedef TGLTFQuaternion<float> FGLTFQuaternion;
diff --git a/Engine/Plugins/Enterprise/GLTFExporter/MigratedFromGit/Source/GLTFExporter/Public/Core/GLTFVector.h b/Engine/Plugins/Enterprise/GLTFExporter/MigratedFromGit/Source/GLTFExporter/Public/Core/GLTFVector.h
new file mode 100644
--- /dev/null
+++ b/Engine/Plugins/Enterprise/GLTFExporter/MigratedFromGit/Source/GLTFExporter/Public/Core/GLTFVector.h
@@ -0,0 +1,98 @@
+// Copyright Epic Games, Inc. All Rights Reserved.
+
+#pragma once
+
+#include "CoreTypes.h"
+
+#if PLATFORM_SUPPORTS_PRAGMA_PACK
+	#pragma pack(push, 1)
+#endif
+
+template <typename ComponentType>
+struct TGLTFVector2
+{
+	union
+	{
+		struct
+		{
+			ComponentType X GCC_PACK(1);
+			ComponentType Y GCC_PACK(1);
+		};
+
+		ComponentType Components[2] GCC_PACK(1);
+	};
+
+	TGLTFVector2(
+		ComponentType X,
+		ComponentType Y
+	)
+		: X(X)
+		, Y(Y)
+	{}
+};
+
+template <typename ComponentType>
+struct TGLTFVector3
+{
+	union
+	{
+		struct
+		{
+			ComponentType X GCC_PACK(1);
+			ComponentType Y GCC_PACK(1);
+			ComponentType Z GCC_PACK(1);
+		};
+
+		ComponentType Components[3] GCC_PACK(1);
+	};
+
+	TGLTFVector3(
+		ComponentType X,
+		ComponentType Y,
+		ComponentType Z
+	)
+		: X(X)
+		, Y(Y)
+		, Z(Z)
+	{}
+};
+
+template <typename ComponentType>
+struct TGLTFVector4
+{
+	union
+	{
+		struct
+		{
+			ComponentType X GCC_PACK(1);
+			ComponentType Y GCC_PACK(1);
+			ComponentType Z GCC_PACK(1);
+			ComponentType W GCC_PACK(1);
+		};
+
+		ComponentType Components[4] GCC_PACK(1);
+	};
+
+	TGLTFVector4(
+		ComponentType X,
+		ComponentType Y,
+		ComponentType Z,
+		ComponentType W
+	)
+		: X(X)
+		, Y(Y)
+		, Z(Z)
+		, W(W)
+	{}
+};
+
+#if PLATFORM_SUPPORTS_PRAGMA_PACK
+	#pragma pack(pop)
+#endif
+
+typedef TGLTFVector2<float> FGLTFVector2;
+typedef TGLTFVector3<float> FGLTFVector3;
+typedef TGLTFVector4<float> FGLTFVector4;
+
+typedef TGLTFVector4<int8> FGLTFInt8Vector4;
+typedef TGLTFVector4<int16> FGLTFInt16Vector4;
