[GLTFExporter] chore: add helpers for accessor components to FGLTFJsonUtility

commit 8ebc593b1419edece034477df325daceab6e23ce
Date: Wed, 19 Jan 2022 08:57:04 +0100

    chore: add helpers for accessor components to FGLTFJsonUtility
    (cherry picked from commit 90d979f613f9fe7b136d60851e2526543ea7301f​)

#skipci

diff --git a/Engine/Plugins/Enterprise/GLTFExporter/MigratedFromGit/Source/GLTFExporter/Public/Json/GLTFJsonUtility.h b/Engine/Plugins/Enterprise/GLTFExporter/MigratedFromGit/Source/GLTFExporter/Public/Json/GLTFJsonUtility.h
--- a/Engine/Plugins/Enterprise/GLTFExporter/MigratedFromGit/Source/GLTFExporter/Public/Json/GLTFJsonUtility.h
+++ b/Engine/Plugins/Enterprise/GLTFExporter/MigratedFromGit/Source/GLTFExporter/Public/Json/GLTFJsonUtility.h
@@ -197,6 +197,40 @@ struct FGLTFJsonUtility
 		}
 	}
 
+	static int32 GetComponentCount(EGLTFJsonAccessorType AccessorType)
+	{
+		switch (AccessorType)
+		{
+		case EGLTFJsonAccessorType::Scalar: return 1;
+		case EGLTFJsonAccessorType::Vec2:   return 2;
+		case EGLTFJsonAccessorType::Vec3:   return 3;
+		case EGLTFJsonAccessorType::Vec4:   return 4;
+		case EGLTFJsonAccessorType::Mat2:   return 4;
+		case EGLTFJsonAccessorType::Mat3:   return 9;
+		case EGLTFJsonAccessorType::Mat4:   return 16;
+		default:                            return 0;
+		}
+	}
+
+	static int32 GetComponentSize(EGLTFJsonComponentType ComponentType)
+	{
+		switch (ComponentType)
+		{
+		case EGLTFJsonComponentType::S8:
+		case EGLTFJsonComponentType::U8:
+			return 1;
+		case EGLTFJsonComponentType::S16:
+		case EGLTFJsonComponentType::U16:
+			return 2;
+		case EGLTFJsonComponentType::S32:
+		case EGLTFJsonComponentType::U32:
+		case EGLTFJsonComponentType::F32:
+			return 4;
+		default:
+			return 0;
+		}
+	}
+
 	static EGLTFJsonComponentType ParseComponentType(int32 ComponentType)
 	{
 		switch (ComponentType)
