[GLTFExporter] refactor: fix issues with implicit conversion of int32 to enum class on macos

commit 0af3f0812716d639e46e71cb21a9aa7a11de80cd
Date: Thu, 20 Jan 2022 10:09:28 +0100

    refactor: fix issues with implicit conversion of int32 to enum class on macos
    (cherry picked from commit 470d980a537e30c70c9b786ff7cf41b8aad5626b​)

#skipci

diff --git a/Engine/Plugins/Enterprise/GLTFExporter/MigratedFromGit/Source/GLTFExporter/Public/Json/GLTFJsonUtility.h b/Engine/Plugins/Enterprise/GLTFExporter/MigratedFromGit/Source/GLTFExporter/Public/Json/GLTFJsonUtility.h
--- a/Engine/Plugins/Enterprise/GLTFExporter/MigratedFromGit/Source/GLTFExporter/Public/Json/GLTFJsonUtility.h
+++ b/Engine/Plugins/Enterprise/GLTFExporter/MigratedFromGit/Source/GLTFExporter/Public/Json/GLTFJsonUtility.h
@@ -231,8 +231,9 @@ struct FGLTFJsonUtility
 		}
 	}
 
-	static EGLTFJsonComponentType ParseComponentType(int32 ComponentType)
+	static EGLTFJsonComponentType ParseComponentType(int32 NumericComponentType)
 	{
+		const EGLTFJsonComponentType ComponentType = static_cast<EGLTFJsonComponentType>(NumericComponentType);
 		switch (ComponentType)
 		{
 		case EGLTFJsonComponentType::S8:
@@ -242,7 +243,7 @@ struct FGLTFJsonUtility
 		case EGLTFJsonComponentType::S32:
 		case EGLTFJsonComponentType::U32:
 		case EGLTFJsonComponentType::F32:
-			return static_cast<EGLTFJsonComponentType>(ComponentType);
+			return ComponentType;
 		default:
 			return EGLTFJsonComponentType::None;
 		}
@@ -282,13 +283,14 @@ struct FGLTFJsonUtility
 		return EGLTFJsonAccessorType::None;
 	}
 
-	static EGLTFJsonBufferTarget ParseBufferTarget(int32 BufferTarget)
+	static EGLTFJsonBufferTarget ParseBufferTarget(int32 NumericBufferTarget)
 	{
+		const EGLTFJsonBufferTarget BufferTarget = static_cast<EGLTFJsonBufferTarget>(NumericBufferTarget);
 		switch (BufferTarget)
 		{
 		case EGLTFJsonBufferTarget::ArrayBuffer:
 		case EGLTFJsonBufferTarget::ElementArrayBuffer:
-			return static_cast<EGLTFJsonBufferTarget>(BufferTarget);
+			return BufferTarget;
 		default:
 			return EGLTFJsonBufferTarget::None;
 		}
