[GLTFExporter] chore: add parsing of animation data to FGLTFJsonReader

commit 076c0d13bc8c832afab97a02007421b6d44dff2c
Date: Thu, 20 Jan 2022 15:21:16 +0100

    chore: add parsing of animation data to FGLTFJsonReader
    (cherry picked from commit 1bff3cc15b9629ff097afc23c898065ec8d80063​)

#skipci

diff --git a/Engine/Plugins/Enterprise/GLTFExporter/MigratedFromGit/Source/GLTFExporter/Public/Json/GLTFJsonUtility.h b/Engine/Plugins/Enterprise/GLTFExporter/MigratedFromGit/Source/GLTFExporter/Public/Json/GLTFJsonUtility.h
--- a/Engine/Plugins/Enterprise/GLTFExporter/MigratedFromGit/Source/GLTFExporter/Public/Json/GLTFJsonUtility.h
+++ b/Engine/Plugins/Enterprise/GLTFExporter/MigratedFromGit/Source/GLTFExporter/Public/Json/GLTFJsonUtility.h
@@ -309,4 +309,44 @@ struct FGLTFJsonUtility
 
 		return EGLTFJsonMimeType::None;
 	}
+
+	static EGLTFJsonTargetPath ParseTargetPath(const FString& TargetPath)
+	{
+		if (TargetPath.Equals(TEXT("translation")))
+		{
+			return EGLTFJsonTargetPath::Translation;
+		}
+		if (TargetPath.Equals(TEXT("rotation")))
+		{
+			return EGLTFJsonTargetPath::Rotation;
+		}
+		if (TargetPath.Equals(TEXT("scale")))
+		{
+			return EGLTFJsonTargetPath::Scale;
+		}
+		if (TargetPath.Equals(TEXT("weights")))
+		{
+			return EGLTFJsonTargetPath::Weights;
+		}
+
+		return EGLTFJsonTargetPath::None;
+	}
+
+	static EGLTFJsonInterpolation ParseInterpolation(const FString& Interpolation)
+	{
+		if (Interpolation.Equals(TEXT("LINEAR")))
+		{
+			return EGLTFJsonInterpolation::Linear;
+		}
+		if (Interpolation.Equals(TEXT("STEP")))
+		{
+			return EGLTFJsonInterpolation::Step;
+		}
+		if (Interpolation.Equals(TEXT("CUBICSPLINE")))
+		{
+			return EGLTFJsonInterpolation::CubicSpline;
+		}
+
+		return EGLTFJsonInterpolation::None;
+	}
 };
