[GLTFExporter] refactor: move implementation files for Json back to the Private folder

commit 0a080496f72bf4333f5eb6a0cd67ae909f541792
Date: Tue, 25 Jan 2022 13:11:20 +0100

    refactor: move implementation files for Json back to the Private folder
    (cherry picked from commit df434deb4b39f71f2fe1e0cb39ed37fb7adeb068​)

#skipci

diff --git a/Engine/Plugins/Enterprise/GLTFExporter/MigratedFromGit/Source/GLTFExporter/Private/Json/GLTFJsonColor.cpp b/Engine/Plugins/Enterprise/GLTFExporter/MigratedFromGit/Source/GLTFExporter/Private/Json/GLTFJsonColor.cpp
new file mode 100644
--- /dev/null
+++ b/Engine/Plugins/Enterprise/GLTFExporter/MigratedFromGit/Source/GLTFExporter/Private/Json/GLTFJsonColor.cpp
@@ -0,0 +1,9 @@
+// Copyright Epic Games, Inc. All Rights Reserved.
+
+#include "Json/GLTFJsonColor.h"
+
+const FGLTFJsonColor3 FGLTFJsonColor3::Black({ 0, 0, 0 });
+const FGLTFJsonColor3 FGLTFJsonColor3::White({ 1, 1, 1 });
+
+const FGLTFJsonColor4 FGLTFJsonColor4::Black({ 0, 0, 0, 1 });
+const FGLTFJsonColor4 FGLTFJsonColor4::White({ 1, 1, 1, 1 });
diff --git a/Engine/Plugins/Enterprise/GLTFExporter/MigratedFromGit/Source/GLTFExporter/Private/Json/GLTFJsonMatrix.cpp b/Engine/Plugins/Enterprise/GLTFExporter/MigratedFromGit/Source/GLTFExporter/Private/Json/GLTFJsonMatrix.cpp
new file mode 100644
--- /dev/null
+++ b/Engine/Plugins/Enterprise/GLTFExporter/MigratedFromGit/Source/GLTFExporter/Private/Json/GLTFJsonMatrix.cpp
@@ -0,0 +1,24 @@
+// Copyright Epic Games, Inc. All Rights Reserved.
+
+#include "Json/GLTFJsonMatrix.h"
+
+const FGLTFJsonMatrix2 FGLTFJsonMatrix2::Identity
+({
+	1, 0,
+	0, 1
+});
+
+const FGLTFJsonMatrix3 FGLTFJsonMatrix3::Identity
+({
+	1, 0, 0,
+	0, 1, 0,
+	0, 0, 1
+});
+
+const FGLTFJsonMatrix4 FGLTFJsonMatrix4::Identity
+({
+	1, 0, 0, 0,
+	0, 1, 0, 0,
+	0, 0, 1, 0,
+	0, 0, 0, 1
+});
diff --git a/Engine/Plugins/Enterprise/GLTFExporter/MigratedFromGit/Source/GLTFExporter/Private/Json/GLTFJsonQuaternion.cpp b/Engine/Plugins/Enterprise/GLTFExporter/MigratedFromGit/Source/GLTFExporter/Private/Json/GLTFJsonQuaternion.cpp
new file mode 100644
--- /dev/null
+++ b/Engine/Plugins/Enterprise/GLTFExporter/MigratedFromGit/Source/GLTFExporter/Private/Json/GLTFJsonQuaternion.cpp
@@ -0,0 +1,5 @@
+// Copyright Epic Games, Inc. All Rights Reserved.
+
+#include "Json/GLTFJsonQuaternion.h"
+
+const FGLTFQuaternion FGLTFJsonQuaternion::Identity = { 0, 0, 0, 1 };
diff --git a/Engine/Plugins/Enterprise/GLTFExporter/MigratedFromGit/Source/GLTFExporter/Private/Json/GLTFJsonVector.cpp b/Engine/Plugins/Enterprise/GLTFExporter/MigratedFromGit/Source/GLTFExporter/Private/Json/GLTFJsonVector.cpp
new file mode 100644
--- /dev/null
+++ b/Engine/Plugins/Enterprise/GLTFExporter/MigratedFromGit/Source/GLTFExporter/Private/Json/GLTFJsonVector.cpp
@@ -0,0 +1,12 @@
+// Copyright Epic Games, Inc. All Rights Reserved.
+
+#include "Json/GLTFJsonVector.h"
+
+const FGLTFJsonVector2 FGLTFJsonVector2::Zero({ 0, 0 });
+const FGLTFJsonVector2 FGLTFJsonVector2::One ({ 1, 1 });
+
+const FGLTFJsonVector3 FGLTFJsonVector3::Zero({ 0, 0, 0 });
+const FGLTFJsonVector3 FGLTFJsonVector3::One ({ 1, 1, 1 });
+
+const FGLTFJsonVector4 FGLTFJsonVector4::Zero({ 0, 0, 0, 0 });
+const FGLTFJsonVector4 FGLTFJsonVector4::One ({ 1, 1, 1, 1 });
diff --git a/Engine/Plugins/Enterprise/GLTFExporter/MigratedFromGit/Source/GLTFExporter/Private/Json/GLTFJsonWriter.cpp b/Engine/Plugins/Enterprise/GLTFExporter/MigratedFromGit/Source/GLTFExporter/Private/Json/GLTFJsonWriter.cpp
new file mode 100644
--- /dev/null
+++ b/Engine/Plugins/Enterprise/GLTFExporter/MigratedFromGit/Source/GLTFExporter/Private/Json/GLTFJsonWriter.cpp
@@ -0,0 +1,167 @@
+// Copyright Epic Games, Inc. All Rights Reserved.
+
+#include "Json/GLTFJsonWriter.h"
+#include "Serialization/JsonSerializer.h"
+#include "Policies/PrettyJsonPrintPolicy.h"
+#include "Policies/CondensedJsonPrintPolicy.h"
+
+template <class CharType, class PrintPolicy>
+class TGLTFJsonWriterImpl final : public IGLTFJsonWriter
+{
+public:
+
+	TGLTFJsonWriterImpl(FArchive& Archive, FGLTFJsonExtensions& Extensions)
+		: IGLTFJsonWriter(Extensions)
+		, JsonWriter(TJsonWriterFactory<CharType, PrintPolicy>::Create(&Archive))
+	{
+	}
+
+	virtual void Close() override
+	{
+		JsonWriter->Close();
+	}
+
+	virtual void Write(bool Boolean) override
+	{
+		if (CurrentIdentifier.IsEmpty())
+		{
+			JsonWriter->WriteValue(Boolean);
+		}
+		else
+		{
+			JsonWriter->WriteValue(CurrentIdentifier, Boolean);
+			CurrentIdentifier.Empty();
+		}
+	}
+
+	virtual void Write(int32 Number) override
+	{
+		if (CurrentIdentifier.IsEmpty())
+		{
+			JsonWriter->WriteValue(Number);
+		}
+		else
+		{
+			JsonWriter->WriteValue(CurrentIdentifier, Number);
+			CurrentIdentifier.Empty();
+		}
+	}
+
+	virtual void Write(int64 Number) override
+	{
+		if (CurrentIdentifier.IsEmpty())
+		{
+			JsonWriter->WriteValue(Number);
+		}
+		else
+		{
+			JsonWriter->WriteValue(CurrentIdentifier, Number);
+			CurrentIdentifier.Empty();
+		}
+	}
+
+	virtual void Write(float Number) override
+	{
+		// NOTE: Specifying 9 significant digits, this ensures no precision is lost
+		FString ExactString = FString::Printf(TEXT("%.9g"), Number);
+
+		if (CurrentIdentifier.IsEmpty())
+		{
+			JsonWriter->WriteRawJSONValue(ExactString);
+		}
+		else
+		{
+			JsonWriter->WriteRawJSONValue(CurrentIdentifier, ExactString);
+			CurrentIdentifier.Empty();
+		}
+	}
+
+	virtual void Write(const FString& String) override
+	{
+		if (CurrentIdentifier.IsEmpty())
+		{
+			JsonWriter->WriteValue(String);
+		}
+		else
+		{
+			JsonWriter->WriteValue(CurrentIdentifier, String);
+			CurrentIdentifier.Empty();
+		}
+	}
+
+	virtual void Write(TYPE_OF_NULLPTR) override
+	{
+		if (CurrentIdentifier.IsEmpty())
+		{
+			JsonWriter->WriteValue(nullptr);
+		}
+		else
+		{
+			JsonWriter->WriteValue(CurrentIdentifier, nullptr);
+			CurrentIdentifier.Empty();
+		}
+	}
+
+	virtual void SetIdentifier(const FString& Identifier) override
+	{
+		this->CurrentIdentifier = Identifier;
+	}
+
+	virtual void StartObject() override
+	{
+		if (CurrentIdentifier.IsEmpty())
+		{
+			JsonWriter->WriteObjectStart();
+		}
+		else
+		{
+			JsonWriter->WriteObjectStart(CurrentIdentifier);
+			CurrentIdentifier.Empty();
+		}
+	}
+
+	virtual void EndObject() override
+	{
+		JsonWriter->WriteObjectEnd();
+	}
+
+	virtual void StartArray() override
+	{
+		if (CurrentIdentifier.IsEmpty())
+		{
+			JsonWriter->WriteArrayStart();
+		}
+		else
+		{
+			JsonWriter->WriteArrayStart(CurrentIdentifier);
+			CurrentIdentifier.Empty();
+		}
+	}
+
+	virtual void EndArray() override
+	{
+		JsonWriter->WriteArrayEnd();
+	}
+
+private:
+
+	FString CurrentIdentifier;
+	TSharedRef<TJsonWriter<CharType, PrintPolicy>> JsonWriter;
+};
+
+template <class CharType>
+struct TCrossPlatformPrettyJsonPrintPolicy : TPrettyJsonPrintPolicy<CharType>
+{
+	static void WriteLineTerminator(FArchive* Stream)
+	{
+		TJsonPrintPolicy<CharType>::WriteString(Stream, TEXT("\n")); // don't use system dependent line terminator
+	}
+};
+
+TSharedRef<IGLTFJsonWriter> IGLTFJsonWriter::Create(FArchive& Archive, bool bPrettyJson, FGLTFJsonExtensions& Extensions)
+{
+	return MakeShareable(bPrettyJson
+		? static_cast<IGLTFJsonWriter*>(new TGLTFJsonWriterImpl<UTF8CHAR, TCrossPlatformPrettyJsonPrintPolicy<UTF8CHAR>>(Archive, Extensions))
+		: static_cast<IGLTFJsonWriter*>(new TGLTFJsonWriterImpl<UTF8CHAR, TCondensedJsonPrintPolicy<UTF8CHAR>>(Archive, Extensions))
+	);
+}
diff --git a/Engine/Plugins/Enterprise/GLTFExporter/MigratedFromGit/Source/GLTFExporter/Public/Json/GLTFJsonColor.cpp b/Engine/Plugins/Enterprise/GLTFExporter/MigratedFromGit/Source/GLTFExporter/Public/Json/GLTFJsonColor.cpp
deleted file mode 100644
diff --git a/Engine/Plugins/Enterprise/GLTFExporter/MigratedFromGit/Source/GLTFExporter/Public/Json/GLTFJsonMatrix.cpp b/Engine/Plugins/Enterprise/GLTFExporter/MigratedFromGit/Source/GLTFExporter/Public/Json/GLTFJsonMatrix.cpp
deleted file mode 100644
diff --git a/Engine/Plugins/Enterprise/GLTFExporter/MigratedFromGit/Source/GLTFExporter/Public/Json/GLTFJsonQuaternion.cpp b/Engine/Plugins/Enterprise/GLTFExporter/MigratedFromGit/Source/GLTFExporter/Public/Json/GLTFJsonQuaternion.cpp
deleted file mode 100644
diff --git a/Engine/Plugins/Enterprise/GLTFExporter/MigratedFromGit/Source/GLTFExporter/Public/Json/GLTFJsonVector.cpp b/Engine/Plugins/Enterprise/GLTFExporter/MigratedFromGit/Source/GLTFExporter/Public/Json/GLTFJsonVector.cpp
deleted file mode 100644
diff --git a/Engine/Plugins/Enterprise/GLTFExporter/MigratedFromGit/Source/GLTFExporter/Public/Json/GLTFJsonWriter.cpp b/Engine/Plugins/Enterprise/GLTFExporter/MigratedFromGit/Source/GLTFExporter/Public/Json/GLTFJsonWriter.cpp
deleted file mode 100644
