[GLTFExporter] refactor: move test plugin specific enum parsing utilities to private class in test plugin

commit f296eaadf75d60500a99918c8be699010729b2bc
Date: Sun, 13 Feb 2022 07:13:40 +0100

    refactor: move test plugin specific enum parsing utilities to private class in test plugin
    (cherry picked from commit a29cbe0c0153a65eb618de76b7f53d13530906ea​)

#skipci

diff --git a/Engine/Plugins/Enterprise/GLTFExporter/MigratedFromGit/Source/GLTFExporter/Public/Json/GLTFJsonUtility.h b/Engine/Plugins/Enterprise/GLTFExporter/MigratedFromGit/Source/GLTFExporter/Public/Json/GLTFJsonUtility.h
--- a/Engine/Plugins/Enterprise/GLTFExporter/MigratedFromGit/Source/GLTFExporter/Public/Json/GLTFJsonUtility.h
+++ b/Engine/Plugins/Enterprise/GLTFExporter/MigratedFromGit/Source/GLTFExporter/Public/Json/GLTFJsonUtility.h
@@ -196,157 +196,4 @@ struct FGLTFJsonUtility
 				return TEXT("");
 		}
 	}
-
-	static int32 GetComponentCount(EGLTFJsonAccessorType AccessorType)
-	{
-		switch (AccessorType)
-		{
-			case EGLTFJsonAccessorType::Scalar: return 1;
-			case EGLTFJsonAccessorType::Vec2:   return 2;
-			case EGLTFJsonAccessorType::Vec3:   return 3;
-			case EGLTFJsonAccessorType::Vec4:   return 4;
-			case EGLTFJsonAccessorType::Mat2:   return 4;
-			case EGLTFJsonAccessorType::Mat3:   return 9;
-			case EGLTFJsonAccessorType::Mat4:   return 16;
-			default:                            return 0;
-		}
-	}
-
-	static int32 GetComponentSize(EGLTFJsonComponentType ComponentType)
-	{
-		switch (ComponentType)
-		{
-			case EGLTFJsonComponentType::S8:
-			case EGLTFJsonComponentType::U8:
-				return 1;
-			case EGLTFJsonComponentType::S16:
-			case EGLTFJsonComponentType::U16:
-				return 2;
-			case EGLTFJsonComponentType::S32:
-			case EGLTFJsonComponentType::U32:
-			case EGLTFJsonComponentType::F32:
-				return 4;
-			default:
-				return 0;
-		}
-	}
-
-	static EGLTFJsonComponentType ParseComponentType(int32 NumericComponentType)
-	{
-		const EGLTFJsonComponentType ComponentType = static_cast<EGLTFJsonComponentType>(NumericComponentType);
-		switch (ComponentType)
-		{
-			case EGLTFJsonComponentType::S8:
-			case EGLTFJsonComponentType::U8:
-			case EGLTFJsonComponentType::S16:
-			case EGLTFJsonComponentType::U16:
-			case EGLTFJsonComponentType::S32:
-			case EGLTFJsonComponentType::U32:
-			case EGLTFJsonComponentType::F32:
-				return ComponentType;
-			default:
-				return EGLTFJsonComponentType::None;
-		}
-	}
-
-	static EGLTFJsonAccessorType ParseAccessorType(const FString& AccessorType)
-	{
-		if (AccessorType.Equals(TEXT("SCALAR")))
-		{
-			return EGLTFJsonAccessorType::Scalar;
-		}
-		if (AccessorType.Equals(TEXT("VEC2")))
-		{
-			return EGLTFJsonAccessorType::Vec2;
-		}
-		if (AccessorType.Equals(TEXT("VEC3")))
-		{
-			return EGLTFJsonAccessorType::Vec3;
-		}
-		if (AccessorType.Equals(TEXT("VEC4")))
-		{
-			return EGLTFJsonAccessorType::Vec4;
-		}
-		if (AccessorType.Equals(TEXT("MAT2")))
-		{
-			return EGLTFJsonAccessorType::Mat2;
-		}
-		if (AccessorType.Equals(TEXT("MAT3")))
-		{
-			return EGLTFJsonAccessorType::Mat3;
-		}
-		if (AccessorType.Equals(TEXT("MAT4")))
-		{
-			return EGLTFJsonAccessorType::Mat4;
-		}
-
-		return EGLTFJsonAccessorType::None;
-	}
-
-	static EGLTFJsonBufferTarget ParseBufferTarget(int32 NumericBufferTarget)
-	{
-		const EGLTFJsonBufferTarget BufferTarget = static_cast<EGLTFJsonBufferTarget>(NumericBufferTarget);
-		switch (BufferTarget)
-		{
-			case EGLTFJsonBufferTarget::ArrayBuffer:
-			case EGLTFJsonBufferTarget::ElementArrayBuffer:
-				return BufferTarget;
-			default:
-				return EGLTFJsonBufferTarget::None;
-		}
-	}
-
-	static EGLTFJsonMimeType ParseMimeType(const FString& MimeType)
-	{
-		if (MimeType.Equals(TEXT("image/jpeg")))
-		{
-			return EGLTFJsonMimeType::JPEG;
-		}
-		if (MimeType.Equals(TEXT("image/png")))
-		{
-			return EGLTFJsonMimeType::PNG;
-		}
-
-		return EGLTFJsonMimeType::None;
-	}
-
-	static EGLTFJsonTargetPath ParseTargetPath(const FString& TargetPath)
-	{
-		if (TargetPath.Equals(TEXT("translation")))
-		{
-			return EGLTFJsonTargetPath::Translation;
-		}
-		if (TargetPath.Equals(TEXT("rotation")))
-		{
-			return EGLTFJsonTargetPath::Rotation;
-		}
-		if (TargetPath.Equals(TEXT("scale")))
-		{
-			return EGLTFJsonTargetPath::Scale;
-		}
-		if (TargetPath.Equals(TEXT("weights")))
-		{
-			return EGLTFJsonTargetPath::Weights;
-		}
-
-		return EGLTFJsonTargetPath::None;
-	}
-
-	static EGLTFJsonInterpolation ParseInterpolation(const FString& Interpolation)
-	{
-		if (Interpolation.Equals(TEXT("LINEAR")))
-		{
-			return EGLTFJsonInterpolation::Linear;
-		}
-		if (Interpolation.Equals(TEXT("STEP")))
-		{
-			return EGLTFJsonInterpolation::Step;
-		}
-		if (Interpolation.Equals(TEXT("CUBICSPLINE")))
-		{
-			return EGLTFJsonInterpolation::CubicSpline;
-		}
-
-		return EGLTFJsonInterpolation::None;
-	}
 };
