[GLTFExporter] refactor: rename enum elements in EGLTFComponentType and remove irrelevant comments in GLTFJsonEnums.h

commit 70e366cd4f007866235f84670264f48cb7e4d9a5
Date: Sun, 13 Feb 2022 09:36:42 +0100

    refactor: rename enum elements in EGLTFComponentType and remove irrelevant comments in GLTFJsonEnums.h
    (cherry picked from commit aa15888209d6c3bd2767b76a57cc7849ff769364​)

#skipci

diff --git a/Engine/Plugins/Enterprise/GLTFExporter/MigratedFromGit/Source/GLTFExporter/Private/Converters/GLTFAccessorConverters.cpp b/Engine/Plugins/Enterprise/GLTFExporter/MigratedFromGit/Source/GLTFExporter/Private/Converters/GLTFAccessorConverters.cpp
--- a/Engine/Plugins/Enterprise/GLTFExporter/MigratedFromGit/Source/GLTFExporter/Private/Converters/GLTFAccessorConverters.cpp
+++ b/Engine/Plugins/Enterprise/GLTFExporter/MigratedFromGit/Source/GLTFExporter/Private/Converters/GLTFAccessorConverters.cpp
@@ -52,7 +52,7 @@ FGLTFJsonAccessorIndex FGLTFPositionBufferConverter::Convert(const FGLTFMeshSect
 
 	FGLTFJsonAccessor JsonAccessor;
 	JsonAccessor.BufferView = Builder.AddBufferView(Positions, EGLTFJsonBufferTarget::ArrayBuffer);
-	JsonAccessor.ComponentType = EGLTFJsonComponentType::F32;
+	JsonAccessor.ComponentType = EGLTFJsonComponentType::Float;
 	JsonAccessor.Count = VertexCount;
 	JsonAccessor.Type = EGLTFJsonAccessorType::Vec3;
 
@@ -113,7 +113,7 @@ FGLTFJsonAccessorIndex FGLTFColorBufferConverter::Convert(const FGLTFMeshSection
 
 	FGLTFJsonAccessor JsonAccessor;
 	JsonAccessor.BufferView = Builder.AddBufferView(Colors, EGLTFJsonBufferTarget::ArrayBuffer);
-	JsonAccessor.ComponentType = EGLTFJsonComponentType::U8;
+	JsonAccessor.ComponentType = EGLTFJsonComponentType::UInt8;
 	JsonAccessor.Count = VertexCount;
 	JsonAccessor.Type = EGLTFJsonAccessorType::Vec4;
 	JsonAccessor.bNormalized = true;
@@ -149,20 +149,20 @@ FGLTFJsonAccessorIndex FGLTFNormalBufferConverter::Convert(const FGLTFMeshSectio
 
 		if (bHighPrecision)
 		{
-			ComponentType = EGLTFJsonComponentType::S16;
+			ComponentType = EGLTFJsonComponentType::Int16;
 			BufferViewIndex = ConvertBufferView<FGLTFInt16Vector4, FPackedRGBA16N>(MeshSection, SourceData);
 			Builder.GetBufferView(BufferViewIndex).ByteStride = sizeof(FGLTFInt16Vector4);
 		}
 		else
 		{
-			ComponentType = EGLTFJsonComponentType::S8;
+			ComponentType = EGLTFJsonComponentType::Int8;
 			BufferViewIndex = ConvertBufferView<FGLTFInt8Vector4, FPackedNormal>(MeshSection, SourceData);
 			Builder.GetBufferView(BufferViewIndex).ByteStride = sizeof(FGLTFInt8Vector4);
 		}
 	}
 	else
 	{
-		ComponentType = EGLTFJsonComponentType::F32;
+		ComponentType = EGLTFJsonComponentType::Float;
 		BufferViewIndex = bHighPrecision
 			? ConvertBufferView<FGLTFVector3, FPackedRGBA16N>(MeshSection, SourceData)
 			: ConvertBufferView<FGLTFVector3, FPackedNormal>(MeshSection, SourceData);
@@ -230,18 +230,18 @@ FGLTFJsonAccessorIndex FGLTFTangentBufferConverter::Convert(const FGLTFMeshSecti
 
 		if (bHighPrecision)
 		{
-			ComponentType = EGLTFJsonComponentType::S16;
+			ComponentType = EGLTFJsonComponentType::Int16;
 			BufferViewIndex = ConvertBufferView<FGLTFInt16Vector4, FPackedRGBA16N>(MeshSection, SourceData);
 		}
 		else
 		{
-			ComponentType = EGLTFJsonComponentType::S8;
+			ComponentType = EGLTFJsonComponentType::Int8;
 			BufferViewIndex = ConvertBufferView<FGLTFInt8Vector4, FPackedNormal>(MeshSection, SourceData);
 		}
 	}
 	else
 	{
-		ComponentType = EGLTFJsonComponentType::F32;
+		ComponentType = EGLTFJsonComponentType::Float;
 		BufferViewIndex = bHighPrecision
 			? ConvertBufferView<FGLTFVector4, FPackedRGBA16N>(MeshSection, SourceData)
 			: ConvertBufferView<FGLTFVector4, FPackedNormal>(MeshSection, SourceData);
@@ -332,7 +332,7 @@ FGLTFJsonAccessorIndex FGLTFUVBufferConverter::Convert(const FGLTFMeshSection* M
 
 	FGLTFJsonAccessor JsonAccessor;
 	JsonAccessor.BufferView = Builder.AddBufferView(UVs, EGLTFJsonBufferTarget::ArrayBuffer);
-	JsonAccessor.ComponentType = EGLTFJsonComponentType::F32;
+	JsonAccessor.ComponentType = EGLTFJsonComponentType::Float;
 	JsonAccessor.Count = VertexCount;
 	JsonAccessor.Type = EGLTFJsonAccessorType::Vec2;
 
@@ -536,7 +536,7 @@ FGLTFJsonAccessorIndex FGLTFBoneWeightBufferConverter::Convert(const FGLTFMeshSe
 
 	FGLTFJsonAccessor JsonAccessor;
 	JsonAccessor.BufferView = Builder.AddBufferView(BoneWeights, EGLTFJsonBufferTarget::ArrayBuffer);
-	JsonAccessor.ComponentType = EGLTFJsonComponentType::U8;
+	JsonAccessor.ComponentType = EGLTFJsonComponentType::UInt8;
 	JsonAccessor.Count = VertexCount;
 	JsonAccessor.Type = EGLTFJsonAccessorType::Vec4;
 	JsonAccessor.bNormalized = true;
diff --git a/Engine/Plugins/Enterprise/GLTFExporter/MigratedFromGit/Source/GLTFExporter/Private/Converters/GLTFConverterUtility.h b/Engine/Plugins/Enterprise/GLTFExporter/MigratedFromGit/Source/GLTFExporter/Private/Converters/GLTFConverterUtility.h
--- a/Engine/Plugins/Enterprise/GLTFExporter/MigratedFromGit/Source/GLTFExporter/Private/Converters/GLTFConverterUtility.h
+++ b/Engine/Plugins/Enterprise/GLTFExporter/MigratedFromGit/Source/GLTFExporter/Private/Converters/GLTFConverterUtility.h
@@ -224,13 +224,13 @@ struct FGLTFConverterUtility
 	template <typename ComponentType>
 	static EGLTFJsonComponentType GetComponentType()
 	{
-		if (TIsSame<ComponentType, int8  >::Value) return EGLTFJsonComponentType::S8;
-		if (TIsSame<ComponentType, uint8 >::Value) return EGLTFJsonComponentType::U8;
-		if (TIsSame<ComponentType, int16 >::Value) return EGLTFJsonComponentType::S16;
-		if (TIsSame<ComponentType, uint16>::Value) return EGLTFJsonComponentType::U16;
-		if (TIsSame<ComponentType, int32 >::Value) return EGLTFJsonComponentType::S32;
-		if (TIsSame<ComponentType, uint32>::Value) return EGLTFJsonComponentType::U32;
-		if (TIsSame<ComponentType, float >::Value) return EGLTFJsonComponentType::F32;
+		if (TIsSame<ComponentType, int8  >::Value) return EGLTFJsonComponentType::Int8;
+		if (TIsSame<ComponentType, uint8 >::Value) return EGLTFJsonComponentType::UInt8;
+		if (TIsSame<ComponentType, int16 >::Value) return EGLTFJsonComponentType::Int16;
+		if (TIsSame<ComponentType, uint16>::Value) return EGLTFJsonComponentType::UInt16;
+		if (TIsSame<ComponentType, int32 >::Value) return EGLTFJsonComponentType::Int32;
+		if (TIsSame<ComponentType, uint32>::Value) return EGLTFJsonComponentType::UInt32;
+		if (TIsSame<ComponentType, float >::Value) return EGLTFJsonComponentType::Float;
 		return EGLTFJsonComponentType::None;
 	}
 };
diff --git a/Engine/Plugins/Enterprise/GLTFExporter/MigratedFromGit/Source/GLTFExporter/Private/Converters/GLTFSkinConverters.cpp b/Engine/Plugins/Enterprise/GLTFExporter/MigratedFromGit/Source/GLTFExporter/Private/Converters/GLTFSkinConverters.cpp
--- a/Engine/Plugins/Enterprise/GLTFExporter/MigratedFromGit/Source/GLTFExporter/Private/Converters/GLTFSkinConverters.cpp
+++ b/Engine/Plugins/Enterprise/GLTFExporter/MigratedFromGit/Source/GLTFExporter/Private/Converters/GLTFSkinConverters.cpp
@@ -44,7 +44,7 @@ FGLTFJsonSkinIndex FGLTFSkinConverter::Convert(FGLTFJsonNodeIndex RootNode, cons
 
 	FGLTFJsonAccessor JsonAccessor;
 	JsonAccessor.BufferView = Builder.AddBufferView(InverseBindMatrices);
-	JsonAccessor.ComponentType = EGLTFJsonComponentType::F32;
+	JsonAccessor.ComponentType = EGLTFJsonComponentType::Float;
 	JsonAccessor.Count = BoneCount;
 	JsonAccessor.Type = EGLTFJsonAccessorType::Mat4;
 
diff --git a/Engine/Plugins/Enterprise/GLTFExporter/MigratedFromGit/Source/GLTFExporter/Private/Tasks/GLTFAnimationTasks.cpp b/Engine/Plugins/Enterprise/GLTFExporter/MigratedFromGit/Source/GLTFExporter/Private/Tasks/GLTFAnimationTasks.cpp
--- a/Engine/Plugins/Enterprise/GLTFExporter/MigratedFromGit/Source/GLTFExporter/Private/Tasks/GLTFAnimationTasks.cpp
+++ b/Engine/Plugins/Enterprise/GLTFExporter/MigratedFromGit/Source/GLTFExporter/Private/Tasks/GLTFAnimationTasks.cpp
@@ -32,7 +32,7 @@ void FGLTFAnimSequenceTask::Complete()
 	{
 		FGLTFJsonAccessor JsonInputAccessor;
 		JsonInputAccessor.BufferView = Builder.AddBufferView(Timestamps);
-		JsonInputAccessor.ComponentType = EGLTFJsonComponentType::F32;
+		JsonInputAccessor.ComponentType = EGLTFJsonComponentType::Float;
 		JsonInputAccessor.Type = EGLTFJsonAccessorType::Scalar;
 		JsonInputAccessor.Count = FrameCount;
 		JsonInputAccessor.MinMaxLength = 1;
@@ -72,7 +72,7 @@ void FGLTFAnimSequenceTask::Complete()
 
 			FGLTFJsonAccessor JsonOutputAccessor;
 			JsonOutputAccessor.BufferView = Builder.AddBufferView(Translations);
-			JsonOutputAccessor.ComponentType = EGLTFJsonComponentType::F32;
+			JsonOutputAccessor.ComponentType = EGLTFJsonComponentType::Float;
 			JsonOutputAccessor.Count = Translations.Num();
 			JsonOutputAccessor.Type = EGLTFJsonAccessorType::Vec3;
 
@@ -100,7 +100,7 @@ void FGLTFAnimSequenceTask::Complete()
 
 			FGLTFJsonAccessor JsonOutputAccessor;
 			JsonOutputAccessor.BufferView = Builder.AddBufferView(Rotations);
-			JsonOutputAccessor.ComponentType = EGLTFJsonComponentType::F32;
+			JsonOutputAccessor.ComponentType = EGLTFJsonComponentType::Float;
 			JsonOutputAccessor.Count = Rotations.Num();
 			JsonOutputAccessor.Type = EGLTFJsonAccessorType::Vec4;
 
@@ -128,7 +128,7 @@ void FGLTFAnimSequenceTask::Complete()
 
 			FGLTFJsonAccessor JsonOutputAccessor;
 			JsonOutputAccessor.BufferView = Builder.AddBufferView(Scales);
-			JsonOutputAccessor.ComponentType = EGLTFJsonComponentType::F32;
+			JsonOutputAccessor.ComponentType = EGLTFJsonComponentType::Float;
 			JsonOutputAccessor.Count = Scales.Num();
 			JsonOutputAccessor.Type = EGLTFJsonAccessorType::Vec3;
 
@@ -183,7 +183,7 @@ void FGLTFLevelSequenceTask::Complete()
 
 	FGLTFJsonAccessor JsonInputAccessor;
 	JsonInputAccessor.BufferView = Builder.AddBufferView(Timestamps);
-	JsonInputAccessor.ComponentType = EGLTFJsonComponentType::F32;
+	JsonInputAccessor.ComponentType = EGLTFJsonComponentType::Float;
 	JsonInputAccessor.Type = EGLTFJsonAccessorType::Scalar;
 	JsonInputAccessor.Count = FrameCount;
 	JsonInputAccessor.MinMaxLength = 1;
@@ -268,7 +268,7 @@ void FGLTFLevelSequenceTask::Complete()
 
 						FGLTFJsonAccessor JsonOutputAccessor;
 						JsonOutputAccessor.BufferView = Builder.AddBufferView(Translations);
-						JsonOutputAccessor.ComponentType = EGLTFJsonComponentType::F32;
+						JsonOutputAccessor.ComponentType = EGLTFJsonComponentType::Float;
 						JsonOutputAccessor.Count = Translations.Num();
 						JsonOutputAccessor.Type = EGLTFJsonAccessorType::Vec3;
 
@@ -306,7 +306,7 @@ void FGLTFLevelSequenceTask::Complete()
 
 						FGLTFJsonAccessor JsonOutputAccessor;
 						JsonOutputAccessor.BufferView = Builder.AddBufferView(Rotations);
-						JsonOutputAccessor.ComponentType = EGLTFJsonComponentType::F32;
+						JsonOutputAccessor.ComponentType = EGLTFJsonComponentType::Float;
 						JsonOutputAccessor.Count = Rotations.Num();
 						JsonOutputAccessor.Type = EGLTFJsonAccessorType::Vec4;
 
@@ -341,7 +341,7 @@ void FGLTFLevelSequenceTask::Complete()
 
 						FGLTFJsonAccessor JsonOutputAccessor;
 						JsonOutputAccessor.BufferView = Builder.AddBufferView(Scales);
-						JsonOutputAccessor.ComponentType = EGLTFJsonComponentType::F32;
+						JsonOutputAccessor.ComponentType = EGLTFJsonComponentType::Float;
 						JsonOutputAccessor.Count = Scales.Num();
 						JsonOutputAccessor.Type = EGLTFJsonAccessorType::Vec3;
 
diff --git a/Engine/Plugins/Enterprise/GLTFExporter/MigratedFromGit/Source/GLTFExporter/Public/Json/GLTFJsonEnums.h b/Engine/Plugins/Enterprise/GLTFExporter/MigratedFromGit/Source/GLTFExporter/Public/Json/GLTFJsonEnums.h
--- a/Engine/Plugins/Enterprise/GLTFExporter/MigratedFromGit/Source/GLTFExporter/Public/Json/GLTFJsonEnums.h
+++ b/Engine/Plugins/Enterprise/GLTFExporter/MigratedFromGit/Source/GLTFExporter/Public/Json/GLTFJsonEnums.h
@@ -64,13 +64,13 @@ enum class EGLTFJsonAccessorType
 enum class EGLTFJsonComponentType
 {
 	None = -1,
-	S8 = 5120, // signed byte
-	U8 = 5121, // unsigned byte
-	S16 = 5122, // signed short
-	U16 = 5123, // unsigned short
-	S32 = 5124, // unused
-	U32 = 5125, // unsigned int -- only valid for indices, not attributes
-	F32 = 5126  // float
+	Int8 = 5120,
+	UInt8 = 5121,
+	Int16 = 5122,
+	UInt16 = 5123,
+	Int32 = 5124, // unused
+	UInt32 = 5125,
+	Float = 5126
 };
 
 enum class EGLTFJsonBufferTarget
@@ -82,14 +82,11 @@ enum class EGLTFJsonBufferTarget
 
 enum class EGLTFJsonPrimitiveMode
 {
-	// unsupported
 	Points = 0,
 	Lines = 1,
 	LineLoop = 2,
 	LineStrip = 3,
-	// initially supported
 	Triangles = 4,
-	// will be supported prior to release
 	TriangleStrip = 5,
 	TriangleFan = 6
 };
@@ -120,10 +117,8 @@ enum class EGLTFJsonMimeType
 enum class EGLTFJsonTextureFilter
 {
 	None = -1,
-	// valid for Min & Mag
 	Nearest = 9728,
 	Linear = 9729,
-	// valid for Min only
 	NearestMipmapNearest = 9984,
 	LinearMipmapNearest = 9985,
 	NearestMipmapLinear = 9986,
