[GLTFExporter] fix: don't export texture transforms for proxy materials if disabled in export options

commit 567b56caccecc0e44b2987a757f8c2b4d26d738b
Date: Sun, 13 Feb 2022 21:18:46 +0100

    fix: don't export texture transforms for proxy materials if disabled in export options
    (cherry picked from commit cc7d64b76eee453a3a3d387157c1f70a349763eb​)

#skipci

diff --git a/Engine/Plugins/Enterprise/GLTFExporter/MigratedFromGit/Source/GLTFExporter/Private/Converters/GLTFMaterialProxyFactory.cpp b/Engine/Plugins/Enterprise/GLTFExporter/MigratedFromGit/Source/GLTFExporter/Private/Converters/GLTFMaterialProxyFactory.cpp
--- a/Engine/Plugins/Enterprise/GLTFExporter/MigratedFromGit/Source/GLTFExporter/Private/Converters/GLTFMaterialProxyFactory.cpp
+++ b/Engine/Plugins/Enterprise/GLTFExporter/MigratedFromGit/Source/GLTFExporter/Private/Converters/GLTFMaterialProxyFactory.cpp
@@ -121,7 +121,7 @@ void FGLTFMaterialProxyFactory::SetProxyProperties(UMaterialInstanceConstant* Pr
 			SetProxyProperty(ProxyMaterial, TEXT("Clear Coat Roughness"), JsonMaterial.ClearCoat.ClearCoatRoughnessTexture, EGLTFMaterialPropertyGroup::ClearCoatRoughness);
 
 			SetProxyProperty(ProxyMaterial, TEXT("Clear Coat Normal Scale"), JsonMaterial.ClearCoat.ClearCoatNormalTexture.Scale);
-			SetProxyProperty(ProxyMaterial, TEXT("Clear Coat Normal"), JsonMaterial.ClearCoat.ClearCoatNormalTexture, EGLTFMaterialPropertyGroup::ClearCoatBottomNormal);
+			SetProxyProperty(ProxyMaterial, TEXT("Clear Coat Normal"), JsonMaterial.ClearCoat.ClearCoatNormalTexture, EGLTFMaterialPropertyGroup::ClearCoatBottomNormal); // TODO: this should actually be top normal
 		}
 	}
 }
diff --git a/Engine/Plugins/Enterprise/GLTFExporter/MigratedFromGit/Source/GLTFExporter/Private/Converters/GLTFNameUtility.h b/Engine/Plugins/Enterprise/GLTFExporter/MigratedFromGit/Source/GLTFExporter/Private/Converters/GLTFNameUtility.h
--- a/Engine/Plugins/Enterprise/GLTFExporter/MigratedFromGit/Source/GLTFExporter/Private/Converters/GLTFNameUtility.h
+++ b/Engine/Plugins/Enterprise/GLTFExporter/MigratedFromGit/Source/GLTFExporter/Private/Converters/GLTFNameUtility.h
@@ -15,8 +15,9 @@ struct FGLTFNameUtility
 	{
 		const UEnum* Enum = StaticEnum<EnumType>();
 		check(Enum != nullptr);
-		const FString DisplayName = Enum->GetDisplayNameTextByValue(Value).ToString();
-		return DisplayName.IsEmpty() ? FString::FromInt(static_cast<int32>(Value)) : DisplayName;
+		const int32 NumericValue = static_cast<int32>(Value);
+		const FString DisplayName = Enum->GetDisplayNameTextByValue(NumericValue).ToString();
+		return DisplayName.IsEmpty() ? FString::FromInt(NumericValue) : DisplayName;
 	}
 
 	static FString GetName(const USceneComponent* Component);
diff --git a/Engine/Plugins/Enterprise/GLTFExporter/MigratedFromGit/Source/GLTFExporter/Private/Tasks/GLTFMaterialTasks.cpp b/Engine/Plugins/Enterprise/GLTFExporter/MigratedFromGit/Source/GLTFExporter/Private/Tasks/GLTFMaterialTasks.cpp
--- a/Engine/Plugins/Enterprise/GLTFExporter/MigratedFromGit/Source/GLTFExporter/Private/Tasks/GLTFMaterialTasks.cpp
+++ b/Engine/Plugins/Enterprise/GLTFExporter/MigratedFromGit/Source/GLTFExporter/Private/Tasks/GLTFMaterialTasks.cpp
@@ -211,33 +211,33 @@ FString FGLTFMaterialTask::GetBakedTextureName(const FString& PropertyName) cons
 void FGLTFMaterialTask::GetProxyProperties(FGLTFJsonMaterial& OutMaterial) const
 {
 	GetProxyProperty(TEXT("Base Color Factor"), OutMaterial.PBRMetallicRoughness.BaseColorFactor);
-	GetProxyProperty(TEXT("Base Color"), OutMaterial.PBRMetallicRoughness.BaseColorTexture);
+	GetProxyProperty(TEXT("Base Color"), OutMaterial.PBRMetallicRoughness.BaseColorTexture, EGLTFMaterialPropertyGroup::BaseColorOpacity);
 
 	if (OutMaterial.ShadingModel == EGLTFJsonShadingModel::Default || OutMaterial.ShadingModel == EGLTFJsonShadingModel::ClearCoat)
 	{
 		GetProxyProperty(TEXT("Emissive Factor"), OutMaterial.EmissiveFactor);
-		GetProxyProperty(TEXT("Emissive"), OutMaterial.EmissiveTexture);
+		GetProxyProperty(TEXT("Emissive"), OutMaterial.EmissiveTexture, EGLTFMaterialPropertyGroup::EmissiveColor);
 
 		GetProxyProperty(TEXT("Metallic Factor"), OutMaterial.PBRMetallicRoughness.MetallicFactor);
 		GetProxyProperty(TEXT("Roughness Factor"), OutMaterial.PBRMetallicRoughness.RoughnessFactor);
-		GetProxyProperty(TEXT("Metallic Roughness"), OutMaterial.PBRMetallicRoughness.MetallicRoughnessTexture);
+		GetProxyProperty(TEXT("Metallic Roughness"), OutMaterial.PBRMetallicRoughness.MetallicRoughnessTexture, EGLTFMaterialPropertyGroup::MetallicRoughness);
 
 		GetProxyProperty(TEXT("Normal Scale"), OutMaterial.NormalTexture.Scale);
-		GetProxyProperty(TEXT("Normal"), OutMaterial.NormalTexture);
+		GetProxyProperty(TEXT("Normal"), OutMaterial.NormalTexture, EGLTFMaterialPropertyGroup::Normal);
 
 		GetProxyProperty(TEXT("Occlusion Strength"), OutMaterial.OcclusionTexture.Strength);
-		GetProxyProperty(TEXT("Occlusion"), OutMaterial.OcclusionTexture);
+		GetProxyProperty(TEXT("Occlusion"), OutMaterial.OcclusionTexture, EGLTFMaterialPropertyGroup::AmbientOcclusion);
 
 		if (OutMaterial.ShadingModel == EGLTFJsonShadingModel::ClearCoat)
 		{
 			GetProxyProperty(TEXT("Clear Coat Factor"), OutMaterial.ClearCoat.ClearCoatFactor);
-			GetProxyProperty(TEXT("Clear Coat"), OutMaterial.ClearCoat.ClearCoatTexture);
+			GetProxyProperty(TEXT("Clear Coat"), OutMaterial.ClearCoat.ClearCoatTexture, EGLTFMaterialPropertyGroup::ClearCoatRoughness); // TODO: add property group for clear coat intensity only
 
 			GetProxyProperty(TEXT("Clear Coat Roughness Factor"), OutMaterial.ClearCoat.ClearCoatRoughnessFactor);
-			GetProxyProperty(TEXT("Clear Coat Roughness"), OutMaterial.ClearCoat.ClearCoatRoughnessTexture);
+			GetProxyProperty(TEXT("Clear Coat Roughness"), OutMaterial.ClearCoat.ClearCoatRoughnessTexture, EGLTFMaterialPropertyGroup::ClearCoatRoughness);
 
 			GetProxyProperty(TEXT("Clear Coat Normal Scale"), OutMaterial.ClearCoat.ClearCoatNormalTexture.Scale);
-			GetProxyProperty(TEXT("Clear Coat Normal"), OutMaterial.ClearCoat.ClearCoatNormalTexture);
+			GetProxyProperty(TEXT("Clear Coat Normal"), OutMaterial.ClearCoat.ClearCoatNormalTexture, EGLTFMaterialPropertyGroup::ClearCoatBottomNormal); // TODO: this should actually be top normal
 		}
 	}
 }
@@ -265,7 +265,7 @@ void FGLTFMaterialTask::GetProxyProperty(const FString& PropertyName, FGLTFJsonC
 	}
 }
 
-void FGLTFMaterialTask::GetProxyProperty(const FString& PropertyName, FGLTFJsonTextureInfo& OutValue) const
+void FGLTFMaterialTask::GetProxyProperty(const FString& PropertyName, FGLTFJsonTextureInfo& OutValue, EGLTFMaterialPropertyGroup PropertyGroup) const
 {
 	UTexture* Texture;
 	if (!FGLTFMaterialUtility::GetNonDefaultParameterValue(Material, PropertyName + TEXT(" Texture"), Texture) || Texture == nullptr)
@@ -300,6 +300,16 @@ void FGLTFMaterialTask::GetProxyProperty(const FString& PropertyName, FGLTFJsonT
 	{
 		OutValue.Transform.Rotation = Rotation;
 	}
+
+	if (!Builder.ExportOptions->bExportTextureTransforms && !OutValue.Transform.IsExactlyDefault())
+	{
+		Builder.LogWarning(FString::Printf(
+			TEXT("Texture coordinates [%d] in %s for material %s are transformed, but texture transform is disabled by export options"),
+			OutValue.TexCoord,
+			*FGLTFNameUtility::GetName(PropertyGroup),
+			*Material->GetName()));
+		OutValue.Transform = {};
+	}
 }
 
 EMaterialShadingModel FGLTFMaterialTask::GetShadingModel() const
diff --git a/Engine/Plugins/Enterprise/GLTFExporter/MigratedFromGit/Source/GLTFExporter/Private/Tasks/GLTFMaterialTasks.h b/Engine/Plugins/Enterprise/GLTFExporter/MigratedFromGit/Source/GLTFExporter/Private/Tasks/GLTFMaterialTasks.h
--- a/Engine/Plugins/Enterprise/GLTFExporter/MigratedFromGit/Source/GLTFExporter/Private/Tasks/GLTFMaterialTasks.h
+++ b/Engine/Plugins/Enterprise/GLTFExporter/MigratedFromGit/Source/GLTFExporter/Private/Tasks/GLTFMaterialTasks.h
@@ -50,7 +50,7 @@ private:
 	void GetProxyProperty(const FString& PropertyName, float& OutValue) const;
 	void GetProxyProperty(const FString& PropertyName, FGLTFJsonColor3& OutValue) const;
 	void GetProxyProperty(const FString& PropertyName, FGLTFJsonColor4& OutValue) const;
-	void GetProxyProperty(const FString& PropertyName, FGLTFJsonTextureInfo& OutValue) const;
+	void GetProxyProperty(const FString& PropertyName, FGLTFJsonTextureInfo& OutValue, EGLTFMaterialPropertyGroup PropertyGroup) const;
 
 	EMaterialShadingModel GetShadingModel() const;
 	void ConvertShadingModel(EGLTFJsonShadingModel& OutShadingModel) const;
diff --git a/Engine/Plugins/Enterprise/GLTFExporter/MigratedFromGit/Source/GLTFExporter/Public/Json/GLTFJsonTextureTransform.h b/Engine/Plugins/Enterprise/GLTFExporter/MigratedFromGit/Source/GLTFExporter/Public/Json/GLTFJsonTextureTransform.h
--- a/Engine/Plugins/Enterprise/GLTFExporter/MigratedFromGit/Source/GLTFExporter/Public/Json/GLTFJsonTextureTransform.h
+++ b/Engine/Plugins/Enterprise/GLTFExporter/MigratedFromGit/Source/GLTFExporter/Public/Json/GLTFJsonTextureTransform.h
@@ -52,11 +52,11 @@ struct FGLTFJsonTextureTransform : IGLTFJsonObject
 
 	bool IsNearlyDefault(float Tolerance = KINDA_SMALL_NUMBER) const
 	{
-		return IsNearlyEqual(FGLTFJsonTextureTransform(), Tolerance);
+		return IsNearlyEqual({}, Tolerance);
 	}
 
 	bool IsExactlyDefault() const
 	{
-		return IsExactlyEqual(FGLTFJsonTextureTransform());
+		return IsExactlyEqual({});
 	}
 };
