[GLTFExporter] feat: add error when exporting material without a proxy in runtime

commit 6c17d1a7214956cf7fa51ecfe10910c8cb69defd
Date: Sun, 13 Feb 2022 21:34:08 +0100

    feat: add error when exporting material without a proxy in runtime
    (cherry picked from commit c9ca1f67f3cee9b570507bad2447721123c75cb4​)

#skipci

diff --git a/Engine/Plugins/Enterprise/GLTFExporter/MigratedFromGit/Source/GLTFExporter/Private/Tasks/GLTFMaterialTasks.cpp b/Engine/Plugins/Enterprise/GLTFExporter/MigratedFromGit/Source/GLTFExporter/Private/Tasks/GLTFMaterialTasks.cpp
--- a/Engine/Plugins/Enterprise/GLTFExporter/MigratedFromGit/Source/GLTFExporter/Private/Tasks/GLTFMaterialTasks.cpp
+++ b/Engine/Plugins/Enterprise/GLTFExporter/MigratedFromGit/Source/GLTFExporter/Private/Tasks/GLTFMaterialTasks.cpp
@@ -188,6 +188,10 @@ void FGLTFMaterialTask::Complete()
 				*SectionString));
 		}
 	}
+#else
+	Builder.LogError(FString::Printf(
+		TEXT("Can't properly export material %s in a runtime environment without a glTF proxy. In the Unreal Editor's Content Browser, please right-click on the material, then select Create glTF Proxy Material"),
+		*Material->GetName()));
 #endif
 }
 
