[GLTFExporter] feat: show log messages after creating proxy materials

commit 314bed50f0031cabb9601721b34d896332f4d172
Date: Sun, 13 Feb 2022 21:58:07 +0100

    feat: show log messages after creating proxy materials
    (cherry picked from commit d011f8d6684d7fc1c11f7468b3cdd8577fe08245​)

#skipci

diff --git a/Engine/Plugins/Enterprise/GLTFExporter/MigratedFromGit/Source/GLTFExporter/Private/Actions/GLTFEditorStyle.h b/Engine/Plugins/Enterprise/GLTFExporter/MigratedFromGit/Source/GLTFExporter/Private/Actions/GLTFEditorStyle.h
--- a/Engine/Plugins/Enterprise/GLTFExporter/MigratedFromGit/Source/GLTFExporter/Private/Actions/GLTFEditorStyle.h
+++ b/Engine/Plugins/Enterprise/GLTFExporter/MigratedFromGit/Source/GLTFExporter/Private/Actions/GLTFEditorStyle.h
@@ -19,12 +19,9 @@ public:
 		 : FSlateStyleSet("GLTFEditorStyle")
 	 {
 	 	const TSharedPtr<IPlugin> Plugin = IPluginManager::Get().FindPlugin(GLTFEXPORTER_MODULE_NAME);
-	 	check(Plugin.IsValid());
-
 	 	FSlateStyleSet::SetContentRoot(Plugin->GetBaseDir() / TEXT("Resources"));
 
-	 	const FVector2D IconSize(16.0f, 16.0f);
-	 	Set("Icon16", new FSlateImageBrush(FSlateStyleSet::RootToContentDir(TEXT("Icon16.png")), IconSize));
+	 	Set("Icon16", new FSlateImageBrush(FSlateStyleSet::RootToContentDir(TEXT("Icon16.png")), FVector2D(16.0f, 16.0f)));
 
 		FSlateStyleRegistry::RegisterSlateStyle(*this);
 	 }
diff --git a/Engine/Plugins/Enterprise/GLTFExporter/MigratedFromGit/Source/GLTFExporter/Private/Actions/GLTFProxyAssetActions.cpp b/Engine/Plugins/Enterprise/GLTFExporter/MigratedFromGit/Source/GLTFExporter/Private/Actions/GLTFProxyAssetActions.cpp
--- a/Engine/Plugins/Enterprise/GLTFExporter/MigratedFromGit/Source/GLTFExporter/Private/Actions/GLTFProxyAssetActions.cpp
+++ b/Engine/Plugins/Enterprise/GLTFExporter/MigratedFromGit/Source/GLTFExporter/Private/Actions/GLTFProxyAssetActions.cpp
@@ -6,7 +6,7 @@
 #include "Actions/GLTFEditorStyle.h"
 #include "Options/GLTFProxyOptions.h"
 #include "UI/GLTFProxyOptionsWindow.h"
-#include "GLTFExporterModule.h"
+#include "Converters/GLTFMaterialProxyFactory.h"
 #include "ToolMenus.h"
 
 #define LOCTEXT_NAMESPACE "GLTFProxyAssetActions"
@@ -49,13 +49,18 @@ void FGLTFProxyAssetActions::OnCreateProxy(TArray<FWeakObjectPtr> Objects) const
 		return;
 	}
 
+	FGLTFMaterialProxyFactory ProxyFactory(Options);
+
 	for (const FWeakObjectPtr& Object : Objects)
 	{
 		if (UMaterialInterface* Material = Cast<UMaterialInterface>(Object.Get()))
 		{
-			IGLTFExporterModule::CreateProxyMaterial(Material, Options);
+			ProxyFactory.RootPath = FPaths::GetPath(Material->GetPathName()) / TEXT("GLTF");
+			ProxyFactory.Create(Material);
 		}
 	}
+
+	ProxyFactory.OpenLog();
 }
 
 #undef LOCTEXT_NAMESPACE
diff --git a/Engine/Plugins/Enterprise/GLTFExporter/MigratedFromGit/Source/GLTFExporter/Private/Converters/GLTFMaterialProxyFactory.cpp b/Engine/Plugins/Enterprise/GLTFExporter/MigratedFromGit/Source/GLTFExporter/Private/Converters/GLTFMaterialProxyFactory.cpp
--- a/Engine/Plugins/Enterprise/GLTFExporter/MigratedFromGit/Source/GLTFExporter/Private/Converters/GLTFMaterialProxyFactory.cpp
+++ b/Engine/Plugins/Enterprise/GLTFExporter/MigratedFromGit/Source/GLTFExporter/Private/Converters/GLTFMaterialProxyFactory.cpp
@@ -53,6 +53,14 @@ UMaterialInterface* FGLTFMaterialProxyFactory::Create(UMaterialInterface* Origin
 	return ProxyMaterial;
 }
 
+void FGLTFMaterialProxyFactory::OpenLog()
+{
+	if (Builder.HasLoggedMessages())
+	{
+		Builder.OpenLog();
+	}
+}
+
 void FGLTFMaterialProxyFactory::SetUserData(UMaterialInstanceConstant* ProxyMaterial, UMaterialInterface* OriginalMaterial)
 {
 	UGLTFMaterialExportOptions* UserData = OriginalMaterial->GetAssetUserData<UGLTFMaterialExportOptions>();
diff --git a/Engine/Plugins/Enterprise/GLTFExporter/MigratedFromGit/Source/GLTFExporter/Private/Converters/GLTFMaterialProxyFactory.h b/Engine/Plugins/Enterprise/GLTFExporter/MigratedFromGit/Source/GLTFExporter/Private/Converters/GLTFMaterialProxyFactory.h
--- a/Engine/Plugins/Enterprise/GLTFExporter/MigratedFromGit/Source/GLTFExporter/Private/Converters/GLTFMaterialProxyFactory.h
+++ b/Engine/Plugins/Enterprise/GLTFExporter/MigratedFromGit/Source/GLTFExporter/Private/Converters/GLTFMaterialProxyFactory.h
@@ -17,6 +17,7 @@ public:
 	FGLTFMaterialProxyFactory(const UGLTFProxyOptions* Options = nullptr);
 
 	UMaterialInterface* Create(UMaterialInterface* OriginalMaterial);
+	void OpenLog();
 
 	FString RootPath;
 
