[GLTFExporter] fix: support extra blend modes in proxy materials (can still be disabled by export options later)

commit 0821263b5996394e4e0b399ca8860985afb6a317
Date: Sun, 13 Feb 2022 22:00:15 +0100

    fix: support extra blend modes in proxy materials (can still be disabled by export options later)
    (cherry picked from commit 00656c05582c57fb8b1857b6b5cd78cf42b8f498​)

#skipci

diff --git a/Engine/Plugins/Enterprise/GLTFExporter/MigratedFromGit/Source/GLTFExporter/Private/Converters/GLTFMaterialProxyFactory.cpp b/Engine/Plugins/Enterprise/GLTFExporter/MigratedFromGit/Source/GLTFExporter/Private/Converters/GLTFMaterialProxyFactory.cpp
--- a/Engine/Plugins/Enterprise/GLTFExporter/MigratedFromGit/Source/GLTFExporter/Private/Converters/GLTFMaterialProxyFactory.cpp
+++ b/Engine/Plugins/Enterprise/GLTFExporter/MigratedFromGit/Source/GLTFExporter/Private/Converters/GLTFMaterialProxyFactory.cpp
@@ -320,6 +320,7 @@ UGLTFExportOptions* FGLTFMaterialProxyFactory::CreateExportOptions(const UGLTFPr
 	UGLTFExportOptions* ExportOptions = NewObject<UGLTFExportOptions>();
 	ExportOptions->ResetToDefault();
 	ExportOptions->bExportProxyMaterials = false;
+	ExportOptions->bExportExtraBlendModes = true;
 	ExportOptions->BakeMaterialInputs = ProxyOptions->bBakeMaterialInputs ? EGLTFMaterialBakeMode::Simple : EGLTFMaterialBakeMode::Disabled;
 	ExportOptions->DefaultMaterialBakeSize = ProxyOptions->DefaultMaterialBakeSize;
 	ExportOptions->DefaultMaterialBakeFilter = ProxyOptions->DefaultMaterialBakeFilter;
