[GLTFExporter] feat: don't create new textures for proxy material if an asset already exists and baking not necessary

commit 105e999a7097be15acdd821d8610c88101ba2e4c
Date: Sun, 13 Feb 2022 22:19:25 +0100

    feat: don't create new textures for proxy material if an asset already exists and baking not necessary
    (cherry picked from commit c8bc0dd4ba4e8880f45b17492173201e957fa834​)

#skipci

diff --git a/Engine/Plugins/Enterprise/GLTFExporter/MigratedFromGit/Source/GLTFExporter/Private/Builders/GLTFConvertBuilder.cpp b/Engine/Plugins/Enterprise/GLTFExporter/MigratedFromGit/Source/GLTFExporter/Private/Builders/GLTFConvertBuilder.cpp
--- a/Engine/Plugins/Enterprise/GLTFExporter/MigratedFromGit/Source/GLTFExporter/Private/Builders/GLTFConvertBuilder.cpp
+++ b/Engine/Plugins/Enterprise/GLTFExporter/MigratedFromGit/Source/GLTFExporter/Private/Builders/GLTFConvertBuilder.cpp
@@ -252,22 +252,7 @@ FGLTFJsonSamplerIndex FGLTFConvertBuilder::GetOrAddSampler(const UTexture* Textu
 
 FGLTFJsonTextureIndex FGLTFConvertBuilder::GetOrAddTexture(const UTexture* Texture)
 {
-	if (const UTexture2D* Texture2D = Cast<UTexture2D>(Texture))
-	{
-		return GetOrAddTexture(Texture2D);
-	}
-
-	if (const UTextureRenderTarget2D* RenderTarget2D = Cast<UTextureRenderTarget2D>(Texture))
-	{
-		return GetOrAddTexture(RenderTarget2D);
-	}
-
-	if (const ULightMapTexture2D* LightMap = Cast<ULightMapTexture2D>(Texture))
-	{
-		return GetOrAddTexture(LightMap);
-	}
-
-	return FGLTFJsonTextureIndex(INDEX_NONE);
+	return GetOrAddTexture(Texture, Texture->SRGB);
 }
 
 FGLTFJsonTextureIndex FGLTFConvertBuilder::GetOrAddTexture(const UTexture2D* Texture)
@@ -300,6 +285,21 @@ FGLTFJsonTextureIndex FGLTFConvertBuilder::GetOrAddTexture(const ULightMapTextur
 	return TextureLightMapConverter->GetOrAdd(Texture);
 }
 
+FGLTFJsonTextureIndex FGLTFConvertBuilder::GetOrAddTexture(const UTexture* Texture, bool bToSRGB)
+{
+	if (const UTexture2D* Texture2D = Cast<UTexture2D>(Texture))
+	{
+		return GetOrAddTexture(Texture2D, bToSRGB);
+	}
+
+	if (const UTextureRenderTarget2D* RenderTarget2D = Cast<UTextureRenderTarget2D>(Texture))
+	{
+		return GetOrAddTexture(RenderTarget2D, bToSRGB);
+	}
+
+	return FGLTFJsonTextureIndex(INDEX_NONE);
+}
+
 FGLTFJsonTextureIndex FGLTFConvertBuilder::GetOrAddTexture(const UTexture2D* Texture, bool bToSRGB)
 {
 	if (Texture == nullptr)
diff --git a/Engine/Plugins/Enterprise/GLTFExporter/MigratedFromGit/Source/GLTFExporter/Private/Builders/GLTFConvertBuilder.h b/Engine/Plugins/Enterprise/GLTFExporter/MigratedFromGit/Source/GLTFExporter/Private/Builders/GLTFConvertBuilder.h
--- a/Engine/Plugins/Enterprise/GLTFExporter/MigratedFromGit/Source/GLTFExporter/Private/Builders/GLTFConvertBuilder.h
+++ b/Engine/Plugins/Enterprise/GLTFExporter/MigratedFromGit/Source/GLTFExporter/Private/Builders/GLTFConvertBuilder.h
@@ -66,6 +66,7 @@ public:
 	FGLTFJsonTextureIndex GetOrAddTexture(const UTextureRenderTarget2D* Texture);
 	FGLTFJsonTextureIndex GetOrAddTexture(const UTextureRenderTargetCube* Texture, ECubeFace CubeFace);
 	FGLTFJsonTextureIndex GetOrAddTexture(const ULightMapTexture2D* Texture);
+	FGLTFJsonTextureIndex GetOrAddTexture(const UTexture* Texture, bool bToSRGB);
 	FGLTFJsonTextureIndex GetOrAddTexture(const UTexture2D* Texture, bool bToSRGB);
 	FGLTFJsonTextureIndex GetOrAddTexture(const UTextureCube* Texture, ECubeFace CubeFace, bool bToSRGB);
 	FGLTFJsonTextureIndex GetOrAddTexture(const UTextureRenderTarget2D* Texture, bool bToSRGB);
diff --git a/Engine/Plugins/Enterprise/GLTFExporter/MigratedFromGit/Source/GLTFExporter/Private/Converters/GLTFMaterialProxyFactory.cpp b/Engine/Plugins/Enterprise/GLTFExporter/MigratedFromGit/Source/GLTFExporter/Private/Converters/GLTFMaterialProxyFactory.cpp
--- a/Engine/Plugins/Enterprise/GLTFExporter/MigratedFromGit/Source/GLTFExporter/Private/Converters/GLTFMaterialProxyFactory.cpp
+++ b/Engine/Plugins/Enterprise/GLTFExporter/MigratedFromGit/Source/GLTFExporter/Private/Converters/GLTFMaterialProxyFactory.cpp
@@ -12,8 +12,8 @@
 FGLTFMaterialProxyFactory::FGLTFMaterialProxyFactory(const UGLTFProxyOptions* Options)
 	: Builder(TEXT(""), CreateExportOptions(Options))
 {
-	Builder.ImageConverter = CreateCustomImageConverter();
-	// TODO: override texture converters so that we don't create new assets if a texture asset already exists
+	Builder.Texture2DConverter = CreateTextureConverter();
+	Builder.ImageConverter = CreateImageConverter();
 }
 
 UMaterialInterface* FGLTFMaterialProxyFactory::Create(UMaterialInterface* OriginalMaterial)
@@ -262,7 +262,38 @@ UPackage* FGLTFMaterialProxyFactory::FindOrCreatePackage(const FString& BaseName
 	return Package;
 }
 
-TUniquePtr<IGLTFImageConverter> FGLTFMaterialProxyFactory::CreateCustomImageConverter()
+TUniquePtr<IGLTFTexture2DConverter> FGLTFMaterialProxyFactory::CreateTextureConverter()
+{
+	class FGLTFCustomTexture2DConverter final: public IGLTFTexture2DConverter
+	{
+	public:
+
+		FGLTFMaterialProxyFactory& Factory;
+
+		FGLTFCustomTexture2DConverter(FGLTFMaterialProxyFactory& Factory)
+			: Factory(Factory)
+		{
+		}
+
+	protected:
+
+		virtual void Sanitize(const UTexture2D*& Texture2D, bool& bToSRGB) override
+		{
+			bToSRGB = false; // ignore
+		}
+
+		virtual FGLTFJsonTextureIndex Convert(const UTexture2D* Texture2D, bool bToSRGB) override
+		{
+			const FGLTFJsonTextureIndex TextureIndex = Factory.Builder.AddTexture();
+			Factory.Textures.Add(TextureIndex, const_cast<UTexture2D*>(Texture2D));
+			return TextureIndex;
+		}
+	};
+
+	return MakeUnique<FGLTFCustomTexture2DConverter>(*this);
+}
+
+TUniquePtr<IGLTFImageConverter> FGLTFMaterialProxyFactory::CreateImageConverter()
 {
 	class FGLTFCustomImageConverter final: public IGLTFImageConverter
 	{
diff --git a/Engine/Plugins/Enterprise/GLTFExporter/MigratedFromGit/Source/GLTFExporter/Private/Converters/GLTFMaterialProxyFactory.h b/Engine/Plugins/Enterprise/GLTFExporter/MigratedFromGit/Source/GLTFExporter/Private/Converters/GLTFMaterialProxyFactory.h
--- a/Engine/Plugins/Enterprise/GLTFExporter/MigratedFromGit/Source/GLTFExporter/Private/Converters/GLTFMaterialProxyFactory.h
+++ b/Engine/Plugins/Enterprise/GLTFExporter/MigratedFromGit/Source/GLTFExporter/Private/Converters/GLTFMaterialProxyFactory.h
@@ -48,7 +48,8 @@ private:
 
 	UPackage* FindOrCreatePackage(const FString& BaseName);
 
-	TUniquePtr<IGLTFImageConverter> CreateCustomImageConverter();
+	TUniquePtr<IGLTFTexture2DConverter> CreateTextureConverter();
+	TUniquePtr<IGLTFImageConverter> CreateImageConverter();
 
 	static bool MakeDirectory(const FString& String);
 
diff --git a/Engine/Plugins/Enterprise/GLTFExporter/MigratedFromGit/Source/GLTFExporter/Private/Converters/GLTFTextureConverters.h b/Engine/Plugins/Enterprise/GLTFExporter/MigratedFromGit/Source/GLTFExporter/Private/Converters/GLTFTextureConverters.h
--- a/Engine/Plugins/Enterprise/GLTFExporter/MigratedFromGit/Source/GLTFExporter/Private/Converters/GLTFTextureConverters.h
+++ b/Engine/Plugins/Enterprise/GLTFExporter/MigratedFromGit/Source/GLTFExporter/Private/Converters/GLTFTextureConverters.h
@@ -6,6 +6,10 @@
 #include "Converters/GLTFConverter.h"
 #include "Converters/GLTFBuilderContext.h"
 
+// TODO: generalize parameter bToSRGB to SamplerType (including normalmap unpacking)
+
+// TODO: remove lightmap-specific converter (since it won't work in runtime)
+
 typedef TGLTFConverter<FGLTFJsonTextureIndex, const UTexture2D*, bool> IGLTFTexture2DConverter;
 typedef TGLTFConverter<FGLTFJsonTextureIndex, const UTextureCube*, ECubeFace, bool> IGLTFTextureCubeConverter;
 typedef TGLTFConverter<FGLTFJsonTextureIndex, const UTextureRenderTarget2D*, bool> IGLTFTextureRenderTarget2DConverter;
diff --git a/Engine/Plugins/Enterprise/GLTFExporter/MigratedFromGit/Source/GLTFExporter/Private/Tasks/GLTFMaterialTasks.cpp b/Engine/Plugins/Enterprise/GLTFExporter/MigratedFromGit/Source/GLTFExporter/Private/Tasks/GLTFMaterialTasks.cpp
--- a/Engine/Plugins/Enterprise/GLTFExporter/MigratedFromGit/Source/GLTFExporter/Private/Tasks/GLTFMaterialTasks.cpp
+++ b/Engine/Plugins/Enterprise/GLTFExporter/MigratedFromGit/Source/GLTFExporter/Private/Tasks/GLTFMaterialTasks.cpp
@@ -277,7 +277,8 @@ void FGLTFMaterialTask::GetProxyProperty(const FString& PropertyName, FGLTFJsonT
 		return;
 	}
 
-	OutValue.Index = Builder.GetOrAddTexture(Texture);
+	const bool bSRGB = PropertyGroup == EGLTFMaterialPropertyGroup::BaseColorOpacity || PropertyGroup == EGLTFMaterialPropertyGroup::EmissiveColor;
+	OutValue.Index = Builder.GetOrAddTexture(Texture, bSRGB);
 
 	float TexCoord;
 	if (FGLTFMaterialUtility::GetNonDefaultParameterValue(Material, PropertyName + TEXT(" UV Index"), TexCoord))
@@ -1216,6 +1217,7 @@ bool FGLTFMaterialTask::TryGetSourceTexture(const UTexture2D*& OutTexture, int32
 			}
 		}
 
+		// TODO: add support for UTextureRenderTarget2D?
 		OutTexture = Cast<UTexture2D>(ParameterValue);
 
 		if (OutTexture == nullptr)
@@ -1255,6 +1257,8 @@ bool FGLTFMaterialTask::TryGetSourceTexture(const UTexture2D*& OutTexture, int32
 		// TODO: handle non-default TextureSampler->SamplerType and TextureSampler->SamplerSource
 
 		// TODO: add support for texture object input expression
+
+		// TODO: add support for UTextureRenderTarget2D?
 		OutTexture = Cast<UTexture2D>(TextureSampler->Texture);
 
 		if (OutTexture == nullptr)
