[GLTFExporter] refactor: remove public API for creating proxy material programmatically until we've settled on the desired behavior

commit b70d2da405e9ee67ebafd740ba585d9304350174
Date: Sun, 13 Feb 2022 22:21:05 +0100

    refactor: remove public API for creating proxy material programmatically until we've settled on the desired behavior
    (cherry picked from commit da0033537ccb51732018ed8322b676994115edbd​)

#skipci

diff --git a/Engine/Plugins/Enterprise/GLTFExporter/MigratedFromGit/Source/GLTFExporter/Private/GLTFExporterModule.cpp b/Engine/Plugins/Enterprise/GLTFExporter/MigratedFromGit/Source/GLTFExporter/Private/GLTFExporterModule.cpp
--- a/Engine/Plugins/Enterprise/GLTFExporter/MigratedFromGit/Source/GLTFExporter/Private/GLTFExporterModule.cpp
+++ b/Engine/Plugins/Enterprise/GLTFExporter/MigratedFromGit/Source/GLTFExporter/Private/GLTFExporterModule.cpp
@@ -1,7 +1,6 @@
 // Copyright Epic Games, Inc. All Rights Reserved.
 
 #include "GLTFExporterModule.h"
-#include "Converters/GLTFMaterialProxyFactory.h"
 #include "Actions/GLTFProxyAssetActions.h"
 #include "Interfaces/IPluginManager.h"
 #if WITH_EDITOR
@@ -14,28 +13,18 @@ class FAssetTypeActions_Base;
 
 DEFINE_LOG_CATEGORY(LogGLTFExporter);
 
-UMaterialInterface* IGLTFExporterModule::CreateProxyMaterial(UMaterialInterface* Material, const UGLTFProxyOptions* Options, const FString& RootPath)
-{
-#if WITH_EDITOR
-	FGLTFMaterialProxyFactory ProxyFactory(Options);
-	ProxyFactory.RootPath = RootPath.IsEmpty() ? FPaths::GetPath(Material->GetPathName()) / TEXT("GLTF") : RootPath;
-	return ProxyFactory.Create(Material);
-#else
-	// TODO: report error
-	return nullptr;
-#endif
-}
-
 class FGLTFExporterModule final : public IGLTFExporterModule
 {
 public:
 
 	virtual void StartupModule() override
 	{
+		// TODO: shaders should be moved into its own module
 		const FString PluginShaderDir = FPaths::Combine(IPluginManager::Get().FindPlugin(TEXT("GLTFExporter"))->GetBaseDir(), TEXT("Shaders"));
 		AddShaderSourceDirectoryMapping(TEXT("/Plugin/GLTFExporter"), PluginShaderDir);
 
 #if WITH_EDITOR
+		// TODO: UI and editor-only functions should be moved (as much as possible) into its own module
 		FCoreDelegates::OnPostEngineInit.AddRaw(this, &FGLTFExporterModule::PostEngineInit);
 #endif
 	}
diff --git a/Engine/Plugins/Enterprise/GLTFExporter/MigratedFromGit/Source/GLTFExporter/Public/GLTFExporterModule.h b/Engine/Plugins/Enterprise/GLTFExporter/MigratedFromGit/Source/GLTFExporter/Public/GLTFExporterModule.h
--- a/Engine/Plugins/Enterprise/GLTFExporter/MigratedFromGit/Source/GLTFExporter/Public/GLTFExporterModule.h
+++ b/Engine/Plugins/Enterprise/GLTFExporter/MigratedFromGit/Source/GLTFExporter/Public/GLTFExporterModule.h
@@ -38,6 +38,4 @@ public:
 	{
 		return FModuleManager::Get().IsModuleLoaded(GLTFEXPORTER_MODULE_NAME);
 	}
-
-	static UMaterialInterface* CreateProxyMaterial(UMaterialInterface* Material, const UGLTFProxyOptions* Options, const FString& RootPath = TEXT(""));
 };
