[GLTFExporter] chore: unmark the plugin as beta

commit dc22e789861a37368854008fac9dc75ec862e83c
Date: Sun, 13 Feb 2022 22:29:27 +0100

    chore: unmark the plugin as beta
    (cherry picked from commit e1b2f70b3a0aa091e5069f7d8a60cc450f639c27​)

#skipci

diff --git a/Engine/Plugins/Enterprise/GLTFExporter/MigratedFromGit/GLTFExporter.uplugin b/Engine/Plugins/Enterprise/GLTFExporter/MigratedFromGit/GLTFExporter.uplugin
--- a/Engine/Plugins/Enterprise/GLTFExporter/MigratedFromGit/GLTFExporter.uplugin
+++ b/Engine/Plugins/Enterprise/GLTFExporter/MigratedFromGit/GLTFExporter.uplugin
@@ -12,7 +12,6 @@
 	"SupportURL": "",
 	"EngineVersion": "4.27.0",
 	"CanContainContent": true,
-	"IsBetaVersion": true,
 	"Modules": [
 		{
 			"Name": "GLTFExporter",
