[GLTFExporter] chore: bump version to 1.2.0

commit 97da050762138c0c05e5f845d9609bfdf1c388ac
Date: Sun, 13 Feb 2022 22:30:38 +0100

    chore: bump version to 1.2.0
    (cherry picked from commit 7c4d4867362a7567222e5203e52fe166b71cd0da​)

#skipci

diff --git a/Engine/Plugins/Enterprise/GLTFExporter/MigratedFromGit/GLTFExporter.uplugin b/Engine/Plugins/Enterprise/GLTFExporter/MigratedFromGit/GLTFExporter.uplugin
--- a/Engine/Plugins/Enterprise/GLTFExporter/MigratedFromGit/GLTFExporter.uplugin
+++ b/Engine/Plugins/Enterprise/GLTFExporter/MigratedFromGit/GLTFExporter.uplugin
@@ -1,7 +1,7 @@
 {
 	"FileVersion": 3,
-	"Version": 110,
-	"VersionName": "1.1.0",
+	"Version": 120,
+	"VersionName": "1.2.0",
 	"FriendlyName": "glTF Exporter",
 	"Description": "An exporter for Khronos glTF 2.0.",
 	"Category": "Exporters",
