[GLTFExporter] fix: fix issue with exporting animation assets that lack compressed data

commit 186a89f7a5a2dc3464a62c2f99b5f94af18511f3
Date: Fri, 18 Feb 2022 18:59:59 +0100

    fix: fix issue with exporting animation assets that lack compressed data
    (cherry picked from commit c9f929d2f12dae4bad19af6097e6a29f602cd4b0​)

#skipci

diff --git a/Engine/Plugins/Enterprise/GLTFExporter/MigratedFromGit/Source/GLTFExporter/Private/Converters/GLTFBoneUtility.cpp b/Engine/Plugins/Enterprise/GLTFExporter/MigratedFromGit/Source/GLTFExporter/Private/Converters/GLTFBoneUtility.cpp
--- a/Engine/Plugins/Enterprise/GLTFExporter/MigratedFromGit/Source/GLTFExporter/Private/Converters/GLTFBoneUtility.cpp
+++ b/Engine/Plugins/Enterprise/GLTFExporter/MigratedFromGit/Source/GLTFExporter/Private/Converters/GLTFBoneUtility.cpp
@@ -54,6 +54,10 @@ void FGLTFBoneUtility::GetBoneTransformsByFrame(const UAnimSequence* AnimSequenc
 	FBoneContainer BoneContainer;
 	BoneContainer.InitializeTo(BoneIndices, FCurveEvaluationOption(true), *AnimSequence->GetSkeleton());
 
+	// NOTE: Some assets may only have raw data, some only compressed data.
+	const bool bHasRawData = AnimSequence->GetRawAnimationData().Num() != 0;
+	BoneContainer.SetUseRAWData(bHasRawData);
+
 	const int32 FrameCount = FrameTimestamps.Num();
 	OutBoneTransformsByFrame.AddDefaulted(FrameCount);
 
