[GLTFExporter] feat: add shader for combining textures (e.g. metallic + roughness and base color + opacity)

commit c121f59a12aca390ebfbe7c5d3a379e3decdc016
Date: Wed, 23 Feb 2022 18:15:35 +0100

    feat: add shader for combining textures (e.g. metallic + roughness and base color + opacity)
    (cherry picked from commit ee1ed2d8b813005e1602cd095e01003f239c743a​)

#skipci

diff --git a/Engine/Plugins/Enterprise/GLTFExporter/MigratedFromGit/Shaders/Private/CombinedTexturePS.usf b/Engine/Plugins/Enterprise/GLTFExporter/MigratedFromGit/Shaders/Private/CombinedTexturePS.usf
new file mode 100644
--- /dev/null
+++ b/Engine/Plugins/Enterprise/GLTFExporter/MigratedFromGit/Shaders/Private/CombinedTexturePS.usf
@@ -0,0 +1,26 @@
+﻿// Copyright Epic Games, Inc. All Rights Reserved.
+
+#pragma once
+
+#include "/Engine/Private/Common.ush"
+
+Texture2D TextureA;
+SamplerState TextureSamplerA;
+
+Texture2D TextureB;
+SamplerState TextureSamplerB;
+
+float4x4 ColorTransformA;
+float4x4 ColorTransformB;
+
+float4 BackgroundColor;
+
+void Main(
+	in float2 TextureCoordinate : TEXCOORD0,
+	out float4 OutColor : SV_Target0
+	)
+{
+	const float4 TextureSampleA = mul(Texture2DSample(TextureA, TextureSamplerA, TextureCoordinate), ColorTransformA);
+	const float4 TextureSampleB = mul(Texture2DSample(TextureB, TextureSamplerB, TextureCoordinate), ColorTransformB);
+	OutColor = RETURN_COLOR(TextureSampleA + TextureSampleB + BackgroundColor);
+}
diff --git a/Engine/Plugins/Enterprise/GLTFExporter/MigratedFromGit/Source/GLTFExporter/Private/Converters/GLTFCombinedTexturePreview.cpp b/Engine/Plugins/Enterprise/GLTFExporter/MigratedFromGit/Source/GLTFExporter/Private/Converters/GLTFCombinedTexturePreview.cpp
new file mode 100644
--- /dev/null
+++ b/Engine/Plugins/Enterprise/GLTFExporter/MigratedFromGit/Source/GLTFExporter/Private/Converters/GLTFCombinedTexturePreview.cpp
@@ -0,0 +1,92 @@
+﻿// Copyright Epic Games, Inc. All Rights Reserved.
+
+/*==============================================================================
+	NormalMapPreview.h: Implementation for previewing normal maps.
+==============================================================================*/
+
+#include "GLTFCombinedTexturePreview.h"
+#include "Shader.h"
+#include "GlobalShader.h"
+#include "SimpleElementShaders.h"
+#include "ShaderParameterUtils.h"
+#include "PipelineStateCache.h"
+
+class FGLTFCombinedTexturePS : public FGlobalShader
+{
+	DECLARE_SHADER_TYPE(FGLTFCombinedTexturePS,Global);
+public:
+
+	static bool ShouldCompilePermutation(const FGlobalShaderPermutationParameters& Parameters)
+	{
+		return IsPCPlatform(Parameters.Platform);
+	}
+
+	FGLTFCombinedTexturePS() {}
+
+	/**
+	 * Initialization constructor.
+	 * @param Initializer - Shader initialization container.
+	 */
+	FGLTFCombinedTexturePS(const ShaderMetaType::CompiledShaderInitializerType& Initializer)
+		: FGlobalShader(Initializer)
+	{
+		TextureA.Bind(Initializer.ParameterMap,TEXT("TextureA"));
+		TextureSamplerA.Bind(Initializer.ParameterMap,TEXT("TextureSamplerA"));
+		TextureB.Bind(Initializer.ParameterMap,TEXT("TextureB"));
+		TextureSamplerB.Bind(Initializer.ParameterMap,TEXT("TextureSamplerB"));
+		ColorTransformA.Bind(Initializer.ParameterMap,TEXT("ColorTransformA"));
+		ColorTransformB.Bind(Initializer.ParameterMap,TEXT("ColorTransformB"));
+		BackgroundColor.Bind(Initializer.ParameterMap,TEXT("BackgroundColor"));
+	}
+
+	/**
+	 * Set shader parameters.
+	 * @param SampleTexture - The texture to sample.
+	 */
+	void SetParameters(FRHICommandList& RHICmdList, const FTexture* InTextureA, const FTexture* InTextureB, const FMatrix& InColorTransformA, const FMatrix& InColorTransformB, const FLinearColor& InBackgroundColor)
+	{
+		FRHIPixelShader* PixelShaderRHI = RHICmdList.GetBoundPixelShader();
+		SetTextureParameter(RHICmdList, PixelShaderRHI,TextureA,TextureSamplerA,InTextureA);
+		SetTextureParameter(RHICmdList, PixelShaderRHI,TextureB,TextureSamplerB,InTextureB);
+		SetShaderValue(RHICmdList, PixelShaderRHI,ColorTransformA,InColorTransformA);
+		SetShaderValue(RHICmdList, PixelShaderRHI,ColorTransformB,InColorTransformB);
+		SetShaderValue(RHICmdList, PixelShaderRHI,BackgroundColor,InBackgroundColor);
+	}
+
+private:
+	/** The texture to sample. */
+	LAYOUT_FIELD(FShaderResourceParameter, TextureA);
+	LAYOUT_FIELD(FShaderResourceParameter, TextureSamplerA);
+	LAYOUT_FIELD(FShaderResourceParameter, TextureB);
+	LAYOUT_FIELD(FShaderResourceParameter, TextureSamplerB);
+	LAYOUT_FIELD(FShaderParameter, ColorTransformA);
+	LAYOUT_FIELD(FShaderParameter, ColorTransformB);
+	LAYOUT_FIELD(FShaderParameter, BackgroundColor);
+};
+IMPLEMENT_SHADER_TYPE(,FGLTFCombinedTexturePS,TEXT("/Plugin/GLTFExporter/Private/CombinedTexturePS.usf"),TEXT("Main"),SF_Pixel);
+
+/** Binds vertex and pixel shaders for this element */
+void FGLTFCombinedTexturePreview::BindShaders(
+	FRHICommandList& RHICmdList,
+	FGraphicsPipelineStateInitializer& GraphicsPSOInit,
+	ERHIFeatureLevel::Type InFeatureLevel,
+	const FMatrix& InTransform,
+	const float InGamma,
+	const FMatrix& ColorWeights,
+	const FTexture* Texture)
+{
+	TShaderMapRef<FSimpleElementVS> VertexShader(GetGlobalShaderMap(InFeatureLevel));
+	TShaderMapRef<FGLTFCombinedTexturePS> PixelShader(GetGlobalShaderMap(InFeatureLevel));
+
+	GraphicsPSOInit.BoundShaderState.VertexDeclarationRHI = GSimpleElementVertexDeclaration.VertexDeclarationRHI;
+	GraphicsPSOInit.BoundShaderState.VertexShaderRHI = VertexShader.GetVertexShader();
+	GraphicsPSOInit.BoundShaderState.PixelShaderRHI = PixelShader.GetPixelShader();
+	GraphicsPSOInit.PrimitiveType = PT_TriangleList;
+
+	GraphicsPSOInit.BlendState = TStaticBlendState<>::GetRHI();
+
+	SetGraphicsPipelineState(RHICmdList, GraphicsPSOInit, EApplyRendertargetOption::ForceApply);
+
+	VertexShader->SetParameters(RHICmdList, InTransform);
+	PixelShader->SetParameters(RHICmdList, TextureA, TextureB, ColorTransformA, ColorTransformB, BackgroundColor);
+}
diff --git a/Engine/Plugins/Enterprise/GLTFExporter/MigratedFromGit/Source/GLTFExporter/Private/Converters/GLTFCombinedTexturePreview.h b/Engine/Plugins/Enterprise/GLTFExporter/MigratedFromGit/Source/GLTFExporter/Private/Converters/GLTFCombinedTexturePreview.h
new file mode 100644
--- /dev/null
+++ b/Engine/Plugins/Enterprise/GLTFExporter/MigratedFromGit/Source/GLTFExporter/Private/Converters/GLTFCombinedTexturePreview.h
@@ -0,0 +1,38 @@
+﻿// Copyright Epic Games, Inc. All Rights Reserved.
+
+#pragma once
+
+#include "CoreMinimal.h"
+#include "BatchedElements.h"
+
+class FGLTFCombinedTexturePreview : public FBatchedElementParameters
+{
+public:
+
+	FGLTFCombinedTexturePreview(const FTexture* TextureA, const FTexture* TextureB, const FMatrix& ColorTransformA, const FMatrix& ColorTransformB, const FLinearColor& BackgroundColor)
+		: TextureA(TextureA)
+		, TextureB(TextureB)
+		, ColorTransformA(ColorTransformA)
+		, ColorTransformB(ColorTransformB)
+		, BackgroundColor(BackgroundColor)
+	{
+	}
+
+private:
+
+	const FTexture* TextureA;
+	const FTexture* TextureB;
+	FMatrix ColorTransformA;
+	FMatrix ColorTransformB;
+	FLinearColor BackgroundColor;
+
+	/** Binds vertex and pixel shaders for this element */
+	virtual void BindShaders(
+		FRHICommandList& RHICmdList,
+		FGraphicsPipelineStateInitializer& GraphicsPSOInit,
+		ERHIFeatureLevel::Type InFeatureLevel,
+		const FMatrix& InTransform,
+		const float InGamma,
+		const FMatrix& ColorWeights,
+		const FTexture* Texture) override;
+};
