[GLTFExporter] feat: add public API class FGLTFTextureUtilities for combining textures

commit 98ee153ffa5ea014ecd1a385652013b57aed98b7
Date: Wed, 23 Feb 2022 18:16:53 +0100

    feat: add public API class FGLTFTextureUtilities for combining textures
    (cherry picked from commit e65473491b7946733fad24bf0b4d13524794af80​)

#skipci

diff --git a/Engine/Plugins/Enterprise/GLTFExporter/MigratedFromGit/Source/GLTFExporter/Private/Utilities/GLTFTextureUtilities.cpp b/Engine/Plugins/Enterprise/GLTFExporter/MigratedFromGit/Source/GLTFExporter/Private/Utilities/GLTFTextureUtilities.cpp
new file mode 100644
--- /dev/null
+++ b/Engine/Plugins/Enterprise/GLTFExporter/MigratedFromGit/Source/GLTFExporter/Private/Utilities/GLTFTextureUtilities.cpp
@@ -0,0 +1,76 @@
+// Copyright Epic Games, Inc. All Rights Reserved.
+
+#include "Utilities/GLTFTextureUtilities.h"
+#include "Converters/GLTFCombinedTexturePreview.h"
+#include "Engine/TextureRenderTarget2D.h"
+#include "CanvasItem.h"
+#include "CanvasTypes.h"
+
+UTextureRenderTarget2D* FGLTFTextureUtilities::CombineBaseColorAndOpacity(const UTexture* BaseColorTexture, const UTexture* OpacityTexture)
+{
+	UTextureRenderTarget2D* RenderTarget = CreateRenderTarget(GetMaxSize(BaseColorTexture, OpacityTexture), TEXT("BaseColorAndOpacity"));
+	RenderTarget->SRGB = BaseColorTexture == nullptr || BaseColorTexture->SRGB; // TODO: is this correct?
+	CombineBaseColorAndOpacity(BaseColorTexture, OpacityTexture, RenderTarget);
+	return RenderTarget;
+}
+
+UTextureRenderTarget2D* FGLTFTextureUtilities::CombineMetallicAndRoughness(const UTexture* MetallicTexture, const UTexture* RoughnessTexture)
+{
+	UTextureRenderTarget2D* RenderTarget = CreateRenderTarget(GetMaxSize(MetallicTexture, RoughnessTexture), TEXT("MetallicAndRoughness"));
+	RenderTarget->SRGB = (MetallicTexture != nullptr && MetallicTexture->SRGB) && (RoughnessTexture != nullptr && RoughnessTexture->SRGB); // TODO: is this correct?
+	CombineMetallicAndRoughness(MetallicTexture, RoughnessTexture, RenderTarget);
+	return RenderTarget;
+}
+
+void FGLTFTextureUtilities::CombineBaseColorAndOpacity(const UTexture* BaseColorTexture, const UTexture* OpacityTexture, UTextureRenderTarget2D* OutputRenderTarget)
+{
+	const FMatrix BaseColorTransform(FPlane(1, 0, 0, 0), FPlane(0, 1, 0, 0), FPlane(0, 0, 1, 0), FPlane(0, 0, 0, 0));
+	const FMatrix OpacityTransform  (FPlane(0, 0, 0, 1), FPlane(0, 0, 0, 0), FPlane(0, 0, 0, 0), FPlane(0, 0, 0, 0));
+	CombineTextures(BaseColorTexture, BaseColorTransform, OpacityTexture, OpacityTransform, FLinearColor::Transparent, OutputRenderTarget);
+}
+
+void FGLTFTextureUtilities::CombineMetallicAndRoughness(const UTexture* MetallicTexture, const UTexture* RoughnessTexture, UTextureRenderTarget2D* OutputRenderTarget)
+{
+	const FMatrix MetallicTransform (FPlane(0, 0, 1, 0), FPlane(0, 0, 0, 0), FPlane(0, 0, 0, 0), FPlane(0, 0, 0, 0));
+	const FMatrix RoughnessTransform(FPlane(0, 1, 0, 0), FPlane(0, 0, 0, 0), FPlane(0, 0, 0, 0), FPlane(0, 0, 0, 0));
+	CombineTextures(MetallicTexture, MetallicTransform, RoughnessTexture, RoughnessTransform, FLinearColor::Black, OutputRenderTarget);
+}
+
+void FGLTFTextureUtilities::CombineTextures(const UTexture* TextureA, const FMatrix& ColorTransformA, const UTexture* TextureB, const FMatrix& ColorTransformB, const FLinearColor& BackgroundColor, UTextureRenderTarget2D* OutRenderTarget)
+{
+	FRenderTarget* RenderTarget = OutRenderTarget->GameThread_GetRenderTargetResource();
+	FCanvas Canvas(RenderTarget, nullptr, 0.0f, 0.0f, 0.0f, GMaxRHIFeatureLevel);
+
+	// TODO: use different white texture depending on if target is SRGB or not
+	const FTexture* TextureResourceA = TextureA != nullptr ? TextureA->Resource : GetWhiteTexture();
+	const FTexture* TextureResourceB = TextureB != nullptr ? TextureB->Resource : GetWhiteTexture();
+
+	FCanvasTileItem TileItem(FVector2D::ZeroVector, FIntPoint(OutRenderTarget->SizeX, OutRenderTarget->SizeY), FLinearColor::White);
+	const TRefCountPtr<FBatchedElementParameters> BatchedElementParameters = new FGLTFCombinedTexturePreview(TextureResourceA, TextureResourceB, ColorTransformA, ColorTransformB, BackgroundColor);
+	TileItem.BatchedElementParameters = BatchedElementParameters;
+	TileItem.Draw(&Canvas);
+
+	Canvas.Flush_GameThread();
+	FlushRenderingCommands();
+}
+
+UTextureRenderTarget2D* FGLTFTextureUtilities::CreateRenderTarget(const FIntPoint& Size, const TCHAR* Name)
+{
+	UTextureRenderTarget2D* RenderTarget = NewObject<UTextureRenderTarget2D>(GetTransientPackage(), Name);
+	RenderTarget->InitCustomFormat(Size.X, Size.Y, PF_B8G8R8A8, true);
+	return RenderTarget;
+}
+
+FIntPoint FGLTFTextureUtilities::GetMaxSize(const UTexture* TextureA, const UTexture* TextureB)
+{
+	const FIntPoint SizeA = (TextureA != nullptr) ? FIntPoint(TextureA->GetSurfaceWidth(), TextureA->GetSurfaceHeight()) : FIntPoint::ZeroValue;
+	const FIntPoint SizeB = (TextureB != nullptr) ? FIntPoint(TextureB->GetSurfaceWidth(), TextureB->GetSurfaceHeight()) : FIntPoint::ZeroValue;
+	return { FMath::Max(SizeA.X, SizeB.X), FMath::Max(SizeA.Y, SizeB.Y) };
+}
+
+FTexture* FGLTFTextureUtilities::GetWhiteTexture()
+{
+	// TODO: explicitly ensure this asset is always cooked
+	static UTexture* Texture = LoadObject<UTexture>(nullptr, TEXT("/GLTFExporter/Textures/Proxy/T_GLTF_DefaultWhite"));
+	return Texture != nullptr ? Texture->Resource : nullptr;
+}
diff --git a/Engine/Plugins/Enterprise/GLTFExporter/MigratedFromGit/Source/GLTFExporter/Public/Utilities/GLTFTextureUtilities.h b/Engine/Plugins/Enterprise/GLTFExporter/MigratedFromGit/Source/GLTFExporter/Public/Utilities/GLTFTextureUtilities.h
new file mode 100644
--- /dev/null
+++ b/Engine/Plugins/Enterprise/GLTFExporter/MigratedFromGit/Source/GLTFExporter/Public/Utilities/GLTFTextureUtilities.h
@@ -0,0 +1,27 @@
+// Copyright Epic Games, Inc. All Rights Reserved.
+
+#pragma once
+
+#include "CoreMinimal.h"
+#include "Engine/TextureRenderTarget2D.h"
+
+class GLTFEXPORTER_API FGLTFTextureUtilities
+{
+public:
+
+	static UTextureRenderTarget2D* CombineBaseColorAndOpacity(const UTexture* BaseColorTexture, const UTexture* OpacityTexture);
+	static UTextureRenderTarget2D* CombineMetallicAndRoughness(const UTexture* MetallicTexture, const UTexture* RoughnessTexture);
+
+	static void CombineBaseColorAndOpacity(const UTexture* BaseColorTexture, const UTexture* OpacityTexture, UTextureRenderTarget2D* OutputRenderTarget);
+	static void CombineMetallicAndRoughness(const UTexture* MetallicTexture, const UTexture* RoughnessTexture, UTextureRenderTarget2D* OutputRenderTarget);
+
+	static void CombineTextures(const UTexture* TextureA, const FMatrix& ColorTransformA, const UTexture* TextureB, const FMatrix& ColorTransformB, const FLinearColor& BackgroundColor, UTextureRenderTarget2D* OutRenderTarget);
+
+	static UTextureRenderTarget2D* CreateRenderTarget(const FIntPoint& Size, const TCHAR* Name);
+
+	static FIntPoint GetMaxSize(const UTexture* TextureA, const UTexture* TextureB);
+
+private:
+
+	static FTexture* GetWhiteTexture();
+};
