[GLTFExporter] feat: add public APIs FGLTFProxyMaterialUtilities and more to ease creating proxy materials in runtime

commit 43e2a550f3c61649c45edb44fa5a4ed9f8519dd0
Date: Wed, 23 Feb 2022 18:27:29 +0100

    feat: add public APIs FGLTFProxyMaterialUtilities and more to ease creating proxy materials in runtime
    (cherry picked from commit af11fe0ebe783fb7e1a73080f3ecbabc8ca63d91​)

#skipci

diff --git a/Engine/Plugins/Enterprise/GLTFExporter/MigratedFromGit/Source/GLTFExporter/Private/Converters/GLTFMaterialProxyFactory.cpp b/Engine/Plugins/Enterprise/GLTFExporter/MigratedFromGit/Source/GLTFExporter/Private/Converters/GLTFMaterialProxyFactory.cpp
--- a/Engine/Plugins/Enterprise/GLTFExporter/MigratedFromGit/Source/GLTFExporter/Private/Converters/GLTFMaterialProxyFactory.cpp
+++ b/Engine/Plugins/Enterprise/GLTFExporter/MigratedFromGit/Source/GLTFExporter/Private/Converters/GLTFMaterialProxyFactory.cpp
@@ -5,8 +5,10 @@
 #include "Converters/GLTFMaterialProxyFactory.h"
 #include "Converters/GLTFMaterialUtility.h"
 #include "Converters/GLTFImageUtility.h"
-#include "AssetRegistry/AssetRegistryModule.h"
+#include "Utilities/GLTFProxyMaterialUtilities.h"
+#include "Materials/GLTFProxyMaterialInfo.h"
 #include "Materials/MaterialInstanceConstant.h"
+#include "AssetRegistry/AssetRegistryModule.h"
 #include "ImageUtils.h"
 
 FGLTFMaterialProxyFactory::FGLTFMaterialProxyFactory(const UGLTFProxyOptions* Options)
@@ -18,7 +20,7 @@ FGLTFMaterialProxyFactory::FGLTFMaterialProxyFactory(const UGLTFProxyOptions* Op
 
 UMaterialInterface* FGLTFMaterialProxyFactory::Create(UMaterialInterface* OriginalMaterial)
 {
-	if (FGLTFMaterialUtility::IsProxyMaterial(OriginalMaterial))
+	if (FGLTFProxyMaterialUtilities::IsProxyMaterial(OriginalMaterial))
 	{
 		return OriginalMaterial;
 	}
@@ -102,85 +104,69 @@ void FGLTFMaterialProxyFactory::SetBaseProperties(UMaterialInstanceConstant* Pro
 
 void FGLTFMaterialProxyFactory::SetProxyProperties(UMaterialInstanceConstant* ProxyMaterial, const FGLTFJsonMaterial& JsonMaterial)
 {
-	SetProxyProperty(ProxyMaterial, TEXT("Base Color Factor"), JsonMaterial.PBRMetallicRoughness.BaseColorFactor);
-	SetProxyProperty(ProxyMaterial, TEXT("Base Color"), JsonMaterial.PBRMetallicRoughness.BaseColorTexture, EGLTFMaterialPropertyGroup::BaseColorOpacity);
+	SetProxyParameter(ProxyMaterial, FGLTFProxyMaterialInfo::BaseColorFactor, JsonMaterial.PBRMetallicRoughness.BaseColorFactor);
+	SetProxyParameter(ProxyMaterial, FGLTFProxyMaterialInfo::BaseColor, JsonMaterial.PBRMetallicRoughness.BaseColorTexture);
 
 	if (JsonMaterial.ShadingModel == EGLTFJsonShadingModel::Default || JsonMaterial.ShadingModel == EGLTFJsonShadingModel::ClearCoat)
 	{
-		SetProxyProperty(ProxyMaterial, TEXT("Emissive Factor"), JsonMaterial.EmissiveFactor);
-		SetProxyProperty(ProxyMaterial, TEXT("Emissive"), JsonMaterial.EmissiveTexture, EGLTFMaterialPropertyGroup::EmissiveColor);
+		SetProxyParameter(ProxyMaterial, FGLTFProxyMaterialInfo::EmissiveFactor, JsonMaterial.EmissiveFactor);
+		SetProxyParameter(ProxyMaterial, FGLTFProxyMaterialInfo::Emissive, JsonMaterial.EmissiveTexture);
 
-		SetProxyProperty(ProxyMaterial, TEXT("Metallic Factor"), JsonMaterial.PBRMetallicRoughness.MetallicFactor);
-		SetProxyProperty(ProxyMaterial, TEXT("Roughness Factor"), JsonMaterial.PBRMetallicRoughness.RoughnessFactor);
-		SetProxyProperty(ProxyMaterial, TEXT("Metallic Roughness"), JsonMaterial.PBRMetallicRoughness.MetallicRoughnessTexture, EGLTFMaterialPropertyGroup::MetallicRoughness);
+		SetProxyParameter(ProxyMaterial, FGLTFProxyMaterialInfo::MetallicFactor, JsonMaterial.PBRMetallicRoughness.MetallicFactor);
+		SetProxyParameter(ProxyMaterial, FGLTFProxyMaterialInfo::RoughnessFactor, JsonMaterial.PBRMetallicRoughness.RoughnessFactor);
+		SetProxyParameter(ProxyMaterial, FGLTFProxyMaterialInfo::MetallicRoughness, JsonMaterial.PBRMetallicRoughness.MetallicRoughnessTexture);
 
-		SetProxyProperty(ProxyMaterial, TEXT("Normal Scale"), JsonMaterial.NormalTexture.Scale);
-		SetProxyProperty(ProxyMaterial, TEXT("Normal"), JsonMaterial.NormalTexture, EGLTFMaterialPropertyGroup::Normal);
+		SetProxyParameter(ProxyMaterial, FGLTFProxyMaterialInfo::NormalScale, JsonMaterial.NormalTexture.Scale);
+		SetProxyParameter(ProxyMaterial, FGLTFProxyMaterialInfo::Normal, JsonMaterial.NormalTexture);
 
-		SetProxyProperty(ProxyMaterial, TEXT("Occlusion Strength"), JsonMaterial.OcclusionTexture.Strength);
-		SetProxyProperty(ProxyMaterial, TEXT("Occlusion"), JsonMaterial.OcclusionTexture, EGLTFMaterialPropertyGroup::AmbientOcclusion);
+		SetProxyParameter(ProxyMaterial, FGLTFProxyMaterialInfo::OcclusionStrength, JsonMaterial.OcclusionTexture.Strength);
+		SetProxyParameter(ProxyMaterial, FGLTFProxyMaterialInfo::Occlusion, JsonMaterial.OcclusionTexture);
 
 		if (JsonMaterial.ShadingModel == EGLTFJsonShadingModel::ClearCoat)
 		{
-			SetProxyProperty(ProxyMaterial, TEXT("Clear Coat Factor"), JsonMaterial.ClearCoat.ClearCoatFactor);
-			SetProxyProperty(ProxyMaterial, TEXT("Clear Coat"), JsonMaterial.ClearCoat.ClearCoatTexture, EGLTFMaterialPropertyGroup::ClearCoatRoughness); // TODO: add property group for clear coat intensity only
+			SetProxyParameter(ProxyMaterial, FGLTFProxyMaterialInfo::ClearCoatFactor, JsonMaterial.ClearCoat.ClearCoatFactor);
+			SetProxyParameter(ProxyMaterial, FGLTFProxyMaterialInfo::ClearCoat, JsonMaterial.ClearCoat.ClearCoatTexture);
 
-			SetProxyProperty(ProxyMaterial, TEXT("Clear Coat Roughness Factor"), JsonMaterial.ClearCoat.ClearCoatRoughnessFactor);
-			SetProxyProperty(ProxyMaterial, TEXT("Clear Coat Roughness"), JsonMaterial.ClearCoat.ClearCoatRoughnessTexture, EGLTFMaterialPropertyGroup::ClearCoatRoughness);
+			SetProxyParameter(ProxyMaterial, FGLTFProxyMaterialInfo::ClearCoatRoughnessFactor, JsonMaterial.ClearCoat.ClearCoatRoughnessFactor);
+			SetProxyParameter(ProxyMaterial, FGLTFProxyMaterialInfo::ClearCoatRoughness, JsonMaterial.ClearCoat.ClearCoatRoughnessTexture);
 
-			SetProxyProperty(ProxyMaterial, TEXT("Clear Coat Normal Scale"), JsonMaterial.ClearCoat.ClearCoatNormalTexture.Scale);
-			SetProxyProperty(ProxyMaterial, TEXT("Clear Coat Normal"), JsonMaterial.ClearCoat.ClearCoatNormalTexture, EGLTFMaterialPropertyGroup::ClearCoatBottomNormal); // TODO: this should actually be top normal
+			SetProxyParameter(ProxyMaterial, FGLTFProxyMaterialInfo::ClearCoatNormalScale, JsonMaterial.ClearCoat.ClearCoatNormalTexture.Scale);
+			SetProxyParameter(ProxyMaterial, FGLTFProxyMaterialInfo::ClearCoatNormal, JsonMaterial.ClearCoat.ClearCoatNormalTexture);
 		}
 	}
 }
 
-void FGLTFMaterialProxyFactory::SetProxyProperty(UMaterialInstanceConstant* ProxyMaterial, const FString& PropertyName, float Scalar)
+void FGLTFMaterialProxyFactory::SetProxyParameter(UMaterialInstanceConstant* ProxyMaterial, const TGLTFProxyMaterialParameterInfo<float>& ParameterInfo, float Scalar)
 {
-	FGLTFMaterialUtility::SetNonDefaultParameterValue(ProxyMaterial, PropertyName, Scalar);
+	ParameterInfo.Set(ProxyMaterial, Scalar, true);
 }
 
-void FGLTFMaterialProxyFactory::SetProxyProperty(UMaterialInstanceConstant* ProxyMaterial, const FString& PropertyName, const FGLTFJsonColor3& Color)
+void FGLTFMaterialProxyFactory::SetProxyParameter(UMaterialInstanceConstant* ProxyMaterial, const TGLTFProxyMaterialParameterInfo<FLinearColor>& ParameterInfo, const FGLTFJsonColor3& Color)
 {
-	FGLTFMaterialUtility::SetNonDefaultParameterValue(ProxyMaterial, PropertyName, FLinearColor(Color.R, Color.G, Color.B));
+	ParameterInfo.Set(ProxyMaterial, FLinearColor(Color.R, Color.G, Color.B), true);
 }
 
-void FGLTFMaterialProxyFactory::SetProxyProperty(UMaterialInstanceConstant* ProxyMaterial, const FString& PropertyName, const FGLTFJsonColor4& Color)
+void FGLTFMaterialProxyFactory::SetProxyParameter(UMaterialInstanceConstant* ProxyMaterial, const TGLTFProxyMaterialParameterInfo<FLinearColor>& ParameterInfo, const FGLTFJsonColor4& Color)
 {
-	FGLTFMaterialUtility::SetNonDefaultParameterValue(ProxyMaterial, PropertyName, FLinearColor(Color.R, Color.G, Color.B, Color.A));
+	ParameterInfo.Set(ProxyMaterial, FLinearColor(Color.R, Color.G, Color.B, Color.A), true);
 }
 
-void FGLTFMaterialProxyFactory::SetProxyProperty(UMaterialInstanceConstant* ProxyMaterial, const FString& PropertyName, const FGLTFJsonTextureInfo& TextureInfo, EGLTFMaterialPropertyGroup PropertyGroup)
+void FGLTFMaterialProxyFactory::SetProxyParameter(UMaterialInstanceConstant* ProxyMaterial, const FGLTFProxyMaterialTextureParameterInfo& ParameterInfo, const FGLTFJsonTextureInfo& TextureInfo)
 {
-	UTexture* Texture = FindOrCreateTexture(TextureInfo.Index, PropertyGroup);
+	UTexture* Texture = FindOrCreateTexture(TextureInfo.Index, ParameterInfo);
 	if (Texture == nullptr)
 	{
 		return;
 	}
 
-	FGLTFMaterialUtility::SetNonDefaultParameterValue(ProxyMaterial, PropertyName + TEXT(" Texture"), Texture);
-
-	if (TextureInfo.TexCoord != 0)
-	{
-		FGLTFMaterialUtility::SetNonDefaultParameterValue(ProxyMaterial, PropertyName + TEXT(" UV Index"), TextureInfo.TexCoord);
-	}
-
-	if (TextureInfo.Transform.Offset != FGLTFJsonVector2::Zero)
-	{
-		FGLTFMaterialUtility::SetNonDefaultParameterValue(ProxyMaterial, PropertyName + TEXT(" UV Offset"), FLinearColor(TextureInfo.Transform.Offset.X, TextureInfo.Transform.Offset.Y, 0, 0));
-	}
-
-	if (TextureInfo.Transform.Scale != FGLTFJsonVector2::One)
-	{
-		FGLTFMaterialUtility::SetNonDefaultParameterValue(ProxyMaterial, PropertyName + TEXT(" UV Scale"), FLinearColor(TextureInfo.Transform.Scale.X, TextureInfo.Transform.Scale.Y, 0, 0));
-	}
-
-	if (TextureInfo.Transform.Rotation != 0)
-	{
-		FGLTFMaterialUtility::SetNonDefaultParameterValue(ProxyMaterial, PropertyName + TEXT(" UV Rotation"), TextureInfo.Transform.Rotation);
-	}
+	ParameterInfo.Texture.Set(ProxyMaterial, Texture, true);
+	ParameterInfo.UVIndex.Set(ProxyMaterial, TextureInfo.TexCoord, true);
+	ParameterInfo.UVOffset.Set(ProxyMaterial, FLinearColor(TextureInfo.Transform.Offset.X, TextureInfo.Transform.Offset.Y, 0, 0), true);
+	ParameterInfo.UVScale.Set(ProxyMaterial, FLinearColor(TextureInfo.Transform.Scale.X, TextureInfo.Transform.Scale.Y, 0, 0), true);
+	ParameterInfo.UVRotation.Set(ProxyMaterial, TextureInfo.Transform.Rotation, true);
 }
 
-UTexture2D* FGLTFMaterialProxyFactory::FindOrCreateTexture(FGLTFJsonTextureIndex Index, EGLTFMaterialPropertyGroup PropertyGroup)
+UTexture2D* FGLTFMaterialProxyFactory::FindOrCreateTexture(FGLTFJsonTextureIndex Index, const FGLTFProxyMaterialTextureParameterInfo& ParameterInfo)
 {
 	if (Index == INDEX_NONE)
 	{
@@ -204,15 +190,15 @@ UTexture2D* FGLTFMaterialProxyFactory::FindOrCreateTexture(FGLTFJsonTextureIndex
 	}
 
 	const FGLTFJsonSampler& JsonSampler = Builder.GetSampler(JsonTexture.Sampler);
-	UTexture2D* Texture = CreateTexture(ImageData, JsonSampler, PropertyGroup);
+	UTexture2D* Texture = CreateTexture(ImageData, JsonSampler, ParameterInfo);
 	Textures.Add(Index, Texture);
 	return Texture;
 }
 
-UTexture2D* FGLTFMaterialProxyFactory::CreateTexture(const FGLTFImageData* ImageData, const FGLTFJsonSampler& JsonSampler, EGLTFMaterialPropertyGroup PropertyGroup)
+UTexture2D* FGLTFMaterialProxyFactory::CreateTexture(const FGLTFImageData* ImageData, const FGLTFJsonSampler& JsonSampler, const FGLTFProxyMaterialTextureParameterInfo& ParameterInfo)
 {
-	const bool bSRGB = PropertyGroup == EGLTFMaterialPropertyGroup::BaseColorOpacity || PropertyGroup == EGLTFMaterialPropertyGroup::EmissiveColor;
-	const bool bNormalMap = PropertyGroup == EGLTFMaterialPropertyGroup::Normal || PropertyGroup == EGLTFMaterialPropertyGroup::ClearCoatBottomNormal;
+	const bool bSRGB = ParameterInfo == FGLTFProxyMaterialInfo::BaseColor || ParameterInfo == FGLTFProxyMaterialInfo::Emissive;
+	const bool bNormalMap = ParameterInfo == FGLTFProxyMaterialInfo::Normal || ParameterInfo == FGLTFProxyMaterialInfo::ClearCoatNormal;
 
 	FCreateTexture2DParameters TexParams;
 	TexParams.bUseAlpha = !ImageData->bIgnoreAlpha;
@@ -234,22 +220,9 @@ UTexture2D* FGLTFMaterialProxyFactory::CreateTexture(const FGLTFImageData* Image
 
 UMaterialInstanceConstant* FGLTFMaterialProxyFactory::CreateInstancedMaterial(UMaterialInterface* OriginalMaterial, EGLTFJsonShadingModel ShadingModel)
 {
-	UMaterialInterface* BaseMaterial = FGLTFMaterialUtility::GetProxyBaseMaterial(ShadingModel);
-	if (BaseMaterial == nullptr)
-	{
-		Builder.LogError(FString::Printf(
-			TEXT("Can't create proxy for material %s, because shading model %s has no base material"),
-			*OriginalMaterial->GetName(),
-			FGLTFJsonUtility::GetValue(ShadingModel)));
-		return nullptr;
-	}
-
 	const FString BaseName = TEXT("MI_GLTF_") + OriginalMaterial->GetName();
 	UPackage* Package = FindOrCreatePackage(BaseName);
-
-	UMaterialInstanceConstant* MaterialInstance = NewObject<UMaterialInstanceConstant>(Package, *BaseName, RF_Public | RF_Standalone);
-	MaterialInstance->Parent = BaseMaterial;
-	return MaterialInstance;
+	return FGLTFProxyMaterialUtilities::CreateProxyMaterial<UMaterialInstanceConstant>(ShadingModel, Package, *BaseName, RF_Public | RF_Standalone);
 }
 
 UPackage* FGLTFMaterialProxyFactory::FindOrCreatePackage(const FString& BaseName)
diff --git a/Engine/Plugins/Enterprise/GLTFExporter/MigratedFromGit/Source/GLTFExporter/Private/Converters/GLTFMaterialProxyFactory.h b/Engine/Plugins/Enterprise/GLTFExporter/MigratedFromGit/Source/GLTFExporter/Private/Converters/GLTFMaterialProxyFactory.h
--- a/Engine/Plugins/Enterprise/GLTFExporter/MigratedFromGit/Source/GLTFExporter/Private/Converters/GLTFMaterialProxyFactory.h
+++ b/Engine/Plugins/Enterprise/GLTFExporter/MigratedFromGit/Source/GLTFExporter/Private/Converters/GLTFMaterialProxyFactory.h
@@ -7,6 +7,7 @@
 #include "Builders/GLTFConvertBuilder.h"
 #include "Options/GLTFProxyOptions.h"
 #include "Options/GLTFExportOptions.h"
+#include "Materials/GLTFProxyMaterialParameterInfo.h"
 
 class UMaterialInstanceConstant;
 
@@ -36,13 +37,13 @@ private:
 	void SetBaseProperties(UMaterialInstanceConstant* ProxyMaterial, UMaterialInterface* OriginalMaterial);
 	void SetProxyProperties(UMaterialInstanceConstant* ProxyMaterial, const FGLTFJsonMaterial& JsonMaterial);
 
-	void SetProxyProperty(UMaterialInstanceConstant* ProxyMaterial, const FString& PropertyName, float Scalar);
-	void SetProxyProperty(UMaterialInstanceConstant* ProxyMaterial, const FString& PropertyName, const FGLTFJsonColor3& Color);
-	void SetProxyProperty(UMaterialInstanceConstant* ProxyMaterial, const FString& PropertyName, const FGLTFJsonColor4& Color);
-	void SetProxyProperty(UMaterialInstanceConstant* ProxyMaterial, const FString& PropertyName, const FGLTFJsonTextureInfo& TextureInfo, EGLTFMaterialPropertyGroup PropertyGroup);
+	void SetProxyParameter(UMaterialInstanceConstant* ProxyMaterial, const TGLTFProxyMaterialParameterInfo<float>& ParameterInfo, float Scalar);
+	void SetProxyParameter(UMaterialInstanceConstant* ProxyMaterial, const TGLTFProxyMaterialParameterInfo<FLinearColor>& ParameterInfo, const FGLTFJsonColor3& Color);
+	void SetProxyParameter(UMaterialInstanceConstant* ProxyMaterial, const TGLTFProxyMaterialParameterInfo<FLinearColor>& ParameterInfo, const FGLTFJsonColor4& Color);
+	void SetProxyParameter(UMaterialInstanceConstant* ProxyMaterial, const FGLTFProxyMaterialTextureParameterInfo& ParameterInfo, const FGLTFJsonTextureInfo& TextureInfo);
 
-	UTexture2D* FindOrCreateTexture(FGLTFJsonTextureIndex Index, EGLTFMaterialPropertyGroup PropertyGroup);
-	UTexture2D* CreateTexture(const FGLTFImageData* ImageData, const FGLTFJsonSampler& JsonSampler, EGLTFMaterialPropertyGroup PropertyGroup);
+	UTexture2D* FindOrCreateTexture(FGLTFJsonTextureIndex Index, const FGLTFProxyMaterialTextureParameterInfo& ParameterInfo);
+	UTexture2D* CreateTexture(const FGLTFImageData* ImageData, const FGLTFJsonSampler& JsonSampler, const FGLTFProxyMaterialTextureParameterInfo& ParameterInfo);
 
 	UMaterialInstanceConstant* CreateInstancedMaterial(UMaterialInterface* OriginalMaterial, EGLTFJsonShadingModel ShadingModel);
 
diff --git a/Engine/Plugins/Enterprise/GLTFExporter/MigratedFromGit/Source/GLTFExporter/Private/Converters/GLTFMaterialUtility.cpp b/Engine/Plugins/Enterprise/GLTFExporter/MigratedFromGit/Source/GLTFExporter/Private/Converters/GLTFMaterialUtility.cpp
--- a/Engine/Plugins/Enterprise/GLTFExporter/MigratedFromGit/Source/GLTFExporter/Private/Converters/GLTFMaterialUtility.cpp
+++ b/Engine/Plugins/Enterprise/GLTFExporter/MigratedFromGit/Source/GLTFExporter/Private/Converters/GLTFMaterialUtility.cpp
@@ -3,6 +3,7 @@
 #include "Converters/GLTFMaterialUtility.h"
 #include "Converters/GLTFTextureUtility.h"
 #include "Converters/GLTFNameUtility.h"
+#include "Utilities/GLTFProxyMaterialUtilities.h"
 #include "Misc/DefaultValueHelper.h"
 #if WITH_EDITOR
 #include "GLTFMaterialAnalyzer.h"
@@ -15,148 +16,14 @@
 #include "Materials/MaterialExpressionTextureCoordinate.h"
 #endif
 
-#define PROXY_MATERIAL_NAME_PREFIX TEXT("M_GLTF_")
-#define PROXY_MATERIAL_ROOT_PATH   TEXT("/GLTFExporter/Materials/Proxy/")
-
 UMaterialInterface* FGLTFMaterialUtility::GetDefaultMaterial()
 {
-	static UMaterialInterface* DefaultMaterial = GetProxyBaseMaterial(EGLTFJsonShadingModel::Default);
+	static UMaterial* DefaultMaterial = FGLTFProxyMaterialUtilities::GetBaseMaterial(EGLTFJsonShadingModel::Default);
 	return DefaultMaterial;
 }
 
-UMaterialInterface* FGLTFMaterialUtility::GetProxyBaseMaterial(EGLTFJsonShadingModel ShadingModel)
-{
-	const FString Name = FGLTFJsonUtility::GetValue(ShadingModel);
-	const FString Path = PROXY_MATERIAL_ROOT_PATH PROXY_MATERIAL_NAME_PREFIX + Name + TEXT(".") PROXY_MATERIAL_NAME_PREFIX + Name;
-	return LoadObject<UMaterialInterface>(nullptr, *Path);
-}
-
-bool FGLTFMaterialUtility::IsProxyMaterial(const UMaterial* Material)
-{
-	return Material->GetPathName().StartsWith(PROXY_MATERIAL_ROOT_PATH PROXY_MATERIAL_NAME_PREFIX);
-}
-
-bool FGLTFMaterialUtility::IsProxyMaterial(const UMaterialInterface* Material)
-{
-	return IsProxyMaterial(Material->GetMaterial());
-}
-
-bool FGLTFMaterialUtility::GetNonDefaultParameterValue(const UMaterialInterface* Material, const FString& PropertyName, float& OutValue)
-{
-	const FHashedMaterialParameterInfo ParameterInfo = *PropertyName;
-
-	float DefaultValue;
-	if (!Material->GetScalarParameterDefaultValue(ParameterInfo, DefaultValue))
-	{
-		// TODO: report error
-		return false;
-	}
-
-	float Value;
-	if (!Material->GetScalarParameterValue(ParameterInfo, Value, true) || Value == DefaultValue)
-	{
-		return false;
-	}
-
-	OutValue = Value;
-	return true;
-}
-
-bool FGLTFMaterialUtility::GetNonDefaultParameterValue(const UMaterialInterface* Material, const FString& PropertyName, FLinearColor& OutValue)
-{
-	const FHashedMaterialParameterInfo ParameterInfo = *PropertyName;
-
-	FLinearColor DefaultValue;
-	if (!Material->GetVectorParameterDefaultValue(ParameterInfo, DefaultValue))
-	{
-		// TODO: report error
-		return false;
-	}
-
-	FLinearColor Value;
-	if (!Material->GetVectorParameterValue(ParameterInfo, Value, true) || Value == DefaultValue)
-	{
-		return false;
-	}
-
-	OutValue = Value;
-	return true;
-}
-
-bool FGLTFMaterialUtility::GetNonDefaultParameterValue(const UMaterialInterface* Material, const FString& PropertyName, UTexture*& OutValue)
-{
-	const FHashedMaterialParameterInfo ParameterInfo = *PropertyName;
-
-	UTexture* DefaultValue;
-	if (!Material->GetTextureParameterDefaultValue(ParameterInfo, DefaultValue))
-	{
-		// TODO: report error
-		return false;
-	}
-
-	UTexture* Value;
-	if (!Material->GetTextureParameterValue(ParameterInfo, Value, true) || Value == DefaultValue)
-	{
-		return false;
-	}
-
-	OutValue = Value;
-	return true;
-}
-
 #if WITH_EDITOR
 
-void FGLTFMaterialUtility::SetNonDefaultParameterValue(UMaterialInstanceConstant* Material, const FString& PropertyName, float Value)
-{
-	const FHashedMaterialParameterInfo ParameterInfo = *PropertyName;
-
-	float DefaultValue;
-	if (!Material->GetScalarParameterDefaultValue(ParameterInfo, DefaultValue))
-	{
-		// TODO: report error
-		return;
-	}
-
-	if (DefaultValue != Value)
-	{
-		Material->SetScalarParameterValueEditorOnly(FMaterialParameterInfo(ParameterInfo), Value);
-	}
-}
-
-void FGLTFMaterialUtility::SetNonDefaultParameterValue(UMaterialInstanceConstant* Material, const FString& PropertyName, const FLinearColor& Value)
-{
-	const FHashedMaterialParameterInfo ParameterInfo = *PropertyName;
-
-	FLinearColor DefaultValue;
-	if (!Material->GetVectorParameterDefaultValue(ParameterInfo, DefaultValue))
-	{
-		// TODO: report error
-		return;
-	}
-
-	if (DefaultValue != Value)
-	{
-		Material->SetVectorParameterValueEditorOnly(FMaterialParameterInfo(ParameterInfo), Value);
-	}
-}
-
-void FGLTFMaterialUtility::SetNonDefaultParameterValue(UMaterialInstanceConstant* Material, const FString& PropertyName, UTexture* Value)
-{
-	const FHashedMaterialParameterInfo ParameterInfo = *PropertyName;
-
-	UTexture* DefaultValue;
-	if (!Material->GetTextureParameterDefaultValue(ParameterInfo, DefaultValue))
-	{
-		// TODO: report error
-		return;
-	}
-
-	if (DefaultValue != Value)
-	{
-		Material->SetTextureParameterValueEditorOnly(FMaterialParameterInfo(ParameterInfo), Value);
-	}
-}
-
 bool FGLTFMaterialUtility::IsNormalMap(const FMaterialPropertyEx& Property)
 {
 	return Property == MP_Normal || Property == TEXT("ClearCoatBottomNormal");
@@ -534,7 +401,7 @@ FString FGLTFMaterialUtility::ShadingModelsToString(const FMaterialShadingModelF
 bool FGLTFMaterialUtility::NeedsMeshData(const UMaterialInterface* Material)
 {
 #if WITH_EDITOR
-	if (Material != nullptr && !IsProxyMaterial(Material))
+	if (Material != nullptr && !FGLTFProxyMaterialUtilities::IsProxyMaterial(Material))
 	{
 		// TODO: only analyze properties that will be needed for this specific material
 		const TArray<FMaterialPropertyEx> Properties =
diff --git a/Engine/Plugins/Enterprise/GLTFExporter/MigratedFromGit/Source/GLTFExporter/Private/Converters/GLTFMaterialUtility.h b/Engine/Plugins/Enterprise/GLTFExporter/MigratedFromGit/Source/GLTFExporter/Private/Converters/GLTFMaterialUtility.h
--- a/Engine/Plugins/Enterprise/GLTFExporter/MigratedFromGit/Source/GLTFExporter/Private/Converters/GLTFMaterialUtility.h
+++ b/Engine/Plugins/Enterprise/GLTFExporter/MigratedFromGit/Source/GLTFExporter/Private/Converters/GLTFMaterialUtility.h
@@ -4,23 +4,11 @@
 
 #include "Builders/GLTFContainerBuilder.h"
 #include "Converters/GLTFSharedArray.h"
-#include "Engine/Texture2D.h"
 
 struct FGLTFMaterialAnalysis;
 struct FMaterialPropertyEx;
 class UMaterialInstanceConstant;
 
-struct FGLTFTextureCombineSource
-{
-	constexpr FGLTFTextureCombineSource(const UTexture2D* Texture, FLinearColor TintColor = { 1.0f, 1.0f, 1.0f, 1.0f }, ESimpleElementBlendMode BlendMode = SE_BLEND_Additive)
-		: Texture(Texture), TintColor(TintColor), BlendMode(BlendMode)
-	{}
-
-	const UTexture2D* Texture;
-	FLinearColor TintColor;
-	ESimpleElementBlendMode BlendMode;
-};
-
 struct FGLTFPropertyBakeOutput
 {
 	FGLTFPropertyBakeOutput(const FMaterialPropertyEx& Property, EPixelFormat PixelFormat, const TGLTFSharedArray<FColor>& Pixels, FIntPoint Size, float EmissiveScale, bool bIsSRGB)
@@ -41,20 +29,7 @@ struct FGLTFMaterialUtility
 {
 	static UMaterialInterface* GetDefaultMaterial();
 
-	static UMaterialInterface* GetProxyBaseMaterial(EGLTFJsonShadingModel ShadingModel);
-
-	static bool IsProxyMaterial(const UMaterial* Material);
-	static bool IsProxyMaterial(const UMaterialInterface* Material);
-
-	static bool GetNonDefaultParameterValue(const UMaterialInterface* Material, const FString& PropertyName, float& OutValue);
-	static bool GetNonDefaultParameterValue(const UMaterialInterface* Material, const FString& PropertyName, FLinearColor& OutValue);
-	static bool GetNonDefaultParameterValue(const UMaterialInterface* Material, const FString& PropertyName, UTexture*& OutValue);
-
 #if WITH_EDITOR
-	static void SetNonDefaultParameterValue(UMaterialInstanceConstant* Material, const FString& PropertyName, float Value);
-	static void SetNonDefaultParameterValue(UMaterialInstanceConstant* Material, const FString& PropertyName, const FLinearColor& Value);
-	static void SetNonDefaultParameterValue(UMaterialInstanceConstant* Material, const FString& PropertyName, UTexture* Value);
-
 	static bool IsNormalMap(const FMaterialPropertyEx& Property);
 	static bool IsSRGB(const FMaterialPropertyEx& Property);
 
@@ -68,7 +43,7 @@ struct FGLTFMaterialUtility
 
 	static const FExpressionInput* GetInputForProperty(const UMaterialInterface* Material, const FMaterialPropertyEx& Property);
 
-	template<class InputType>
+	template <class InputType>
 	static const FMaterialInput<InputType>* GetInputForProperty(const UMaterialInterface* Material, const FMaterialPropertyEx& Property)
 	{
 		const FExpressionInput* ExpressionInput = GetInputForProperty(Material, Property);
diff --git a/Engine/Plugins/Enterprise/GLTFExporter/MigratedFromGit/Source/GLTFExporter/Private/Materials/GLTFProxyMaterialInfo.cpp b/Engine/Plugins/Enterprise/GLTFExporter/MigratedFromGit/Source/GLTFExporter/Private/Materials/GLTFProxyMaterialInfo.cpp
new file mode 100644
--- /dev/null
+++ b/Engine/Plugins/Enterprise/GLTFExporter/MigratedFromGit/Source/GLTFExporter/Private/Materials/GLTFProxyMaterialInfo.cpp
@@ -0,0 +1,28 @@
+// Copyright Epic Games, Inc. All Rights Reserved.
+
+#include "Materials/GLTFProxyMaterialInfo.h"
+
+const FGLTFProxyMaterialTextureParameterInfo FGLTFProxyMaterialInfo::BaseColor = { TEXT("Base Color") };
+const TGLTFProxyMaterialParameterInfo<FLinearColor> FGLTFProxyMaterialInfo::BaseColorFactor = { TEXT("Base Color Factor") };
+
+const FGLTFProxyMaterialTextureParameterInfo FGLTFProxyMaterialInfo::Emissive = { TEXT("Emissive") };
+const TGLTFProxyMaterialParameterInfo<FLinearColor> FGLTFProxyMaterialInfo::EmissiveFactor = { TEXT("Emissive Factor") };
+
+const FGLTFProxyMaterialTextureParameterInfo FGLTFProxyMaterialInfo::MetallicRoughness = { TEXT("Metallic Roughness") };
+const TGLTFProxyMaterialParameterInfo<float> FGLTFProxyMaterialInfo::MetallicFactor = { TEXT("Metallic Factor") };
+const TGLTFProxyMaterialParameterInfo<float> FGLTFProxyMaterialInfo::RoughnessFactor = { TEXT("Roughness Factor") };
+
+const FGLTFProxyMaterialTextureParameterInfo FGLTFProxyMaterialInfo::Normal = { TEXT("Normal") };
+const TGLTFProxyMaterialParameterInfo<float> FGLTFProxyMaterialInfo::NormalScale = { TEXT("Normal Scale") };
+
+const FGLTFProxyMaterialTextureParameterInfo FGLTFProxyMaterialInfo::Occlusion = { TEXT("Occlusion") };
+const TGLTFProxyMaterialParameterInfo<float> FGLTFProxyMaterialInfo::OcclusionStrength = { TEXT("Occlusion Strength") };
+
+const FGLTFProxyMaterialTextureParameterInfo FGLTFProxyMaterialInfo::ClearCoat = { TEXT("Clear Coat") };
+const TGLTFProxyMaterialParameterInfo<float> FGLTFProxyMaterialInfo::ClearCoatFactor = { TEXT("Clear Coat Factor") };
+
+const FGLTFProxyMaterialTextureParameterInfo FGLTFProxyMaterialInfo::ClearCoatRoughness = { TEXT("Clear Coat Roughness") };
+const TGLTFProxyMaterialParameterInfo<float> FGLTFProxyMaterialInfo::ClearCoatRoughnessFactor = { TEXT("Clear Coat Roughness Factor") };
+
+const FGLTFProxyMaterialTextureParameterInfo FGLTFProxyMaterialInfo::ClearCoatNormal = { TEXT("Clear Coat Normal") };
+const TGLTFProxyMaterialParameterInfo<float> FGLTFProxyMaterialInfo::ClearCoatNormalScale = { TEXT("Clear Coat Normal Scale") };
diff --git a/Engine/Plugins/Enterprise/GLTFExporter/MigratedFromGit/Source/GLTFExporter/Private/Tasks/GLTFMaterialTasks.cpp b/Engine/Plugins/Enterprise/GLTFExporter/MigratedFromGit/Source/GLTFExporter/Private/Tasks/GLTFMaterialTasks.cpp
--- a/Engine/Plugins/Enterprise/GLTFExporter/MigratedFromGit/Source/GLTFExporter/Private/Tasks/GLTFMaterialTasks.cpp
+++ b/Engine/Plugins/Enterprise/GLTFExporter/MigratedFromGit/Source/GLTFExporter/Private/Tasks/GLTFMaterialTasks.cpp
@@ -5,6 +5,8 @@
 #include "Converters/GLTFNameUtility.h"
 #include "Converters/GLTFMaterialUtility.h"
 #include "Builders/GLTFContainerBuilder.h"
+#include "Utilities/GLTFProxyMaterialUtilities.h"
+#include "Materials/GLTFProxyMaterialInfo.h"
 #if WITH_EDITOR
 #include "Materials/MaterialInstance.h"
 #include "Materials/MaterialExpressionConstant.h"
@@ -60,12 +62,19 @@ void FGLTFMaterialTask::Complete()
 	JsonMaterial.AlphaCutoff = Material->GetOpacityMaskClipValue();
 	JsonMaterial.DoubleSided = Material->IsTwoSided();
 
+	if (const UMaterialInstance* MaterialInstance = Cast<UMaterialInstance>(Material))
+	{
+		// TODO: remove ugly hack for dynamic instances
+		FGLTFProxyMaterialUtilities::GetOpacityMaskClipValue(MaterialInstance, JsonMaterial.AlphaCutoff);
+		FGLTFProxyMaterialUtilities::GetTwoSided(MaterialInstance, JsonMaterial.DoubleSided);
+	}
+
 	ConvertShadingModel(JsonMaterial.ShadingModel);
 	ConvertAlphaMode(JsonMaterial.AlphaMode, JsonMaterial.BlendMode);
 
-	if (FGLTFMaterialUtility::IsProxyMaterial(BaseMaterial))
+	if (FGLTFProxyMaterialUtilities::IsProxyMaterial(BaseMaterial))
 	{
-		GetProxyProperties(JsonMaterial);
+		GetProxyParameters(JsonMaterial);
 		return;
 	}
 
@@ -212,98 +221,98 @@ FString FGLTFMaterialTask::GetBakedTextureName(const FString& PropertyName) cons
 	return GetMaterialName() + TEXT("_") + PropertyName;
 }
 
-void FGLTFMaterialTask::GetProxyProperties(FGLTFJsonMaterial& OutMaterial) const
+void FGLTFMaterialTask::GetProxyParameters(FGLTFJsonMaterial& OutMaterial) const
 {
-	GetProxyProperty(TEXT("Base Color Factor"), OutMaterial.PBRMetallicRoughness.BaseColorFactor);
-	GetProxyProperty(TEXT("Base Color"), OutMaterial.PBRMetallicRoughness.BaseColorTexture, EGLTFMaterialPropertyGroup::BaseColorOpacity);
+	GetProxyParameter(FGLTFProxyMaterialInfo::BaseColorFactor, OutMaterial.PBRMetallicRoughness.BaseColorFactor);
+	GetProxyParameter(FGLTFProxyMaterialInfo::BaseColor, OutMaterial.PBRMetallicRoughness.BaseColorTexture);
 
 	if (OutMaterial.ShadingModel == EGLTFJsonShadingModel::Default || OutMaterial.ShadingModel == EGLTFJsonShadingModel::ClearCoat)
 	{
-		GetProxyProperty(TEXT("Emissive Factor"), OutMaterial.EmissiveFactor);
-		GetProxyProperty(TEXT("Emissive"), OutMaterial.EmissiveTexture, EGLTFMaterialPropertyGroup::EmissiveColor);
+		GetProxyParameter(FGLTFProxyMaterialInfo::EmissiveFactor, OutMaterial.EmissiveFactor);
+		GetProxyParameter(FGLTFProxyMaterialInfo::Emissive, OutMaterial.EmissiveTexture);
 
-		GetProxyProperty(TEXT("Metallic Factor"), OutMaterial.PBRMetallicRoughness.MetallicFactor);
-		GetProxyProperty(TEXT("Roughness Factor"), OutMaterial.PBRMetallicRoughness.RoughnessFactor);
-		GetProxyProperty(TEXT("Metallic Roughness"), OutMaterial.PBRMetallicRoughness.MetallicRoughnessTexture, EGLTFMaterialPropertyGroup::MetallicRoughness);
+		GetProxyParameter(FGLTFProxyMaterialInfo::MetallicFactor, OutMaterial.PBRMetallicRoughness.MetallicFactor);
+		GetProxyParameter(FGLTFProxyMaterialInfo::RoughnessFactor, OutMaterial.PBRMetallicRoughness.RoughnessFactor);
+		GetProxyParameter(FGLTFProxyMaterialInfo::MetallicRoughness, OutMaterial.PBRMetallicRoughness.MetallicRoughnessTexture);
 
-		GetProxyProperty(TEXT("Normal Scale"), OutMaterial.NormalTexture.Scale);
-		GetProxyProperty(TEXT("Normal"), OutMaterial.NormalTexture, EGLTFMaterialPropertyGroup::Normal);
+		GetProxyParameter(FGLTFProxyMaterialInfo::NormalScale, OutMaterial.NormalTexture.Scale);
+		GetProxyParameter(FGLTFProxyMaterialInfo::Normal, OutMaterial.NormalTexture);
 
-		GetProxyProperty(TEXT("Occlusion Strength"), OutMaterial.OcclusionTexture.Strength);
-		GetProxyProperty(TEXT("Occlusion"), OutMaterial.OcclusionTexture, EGLTFMaterialPropertyGroup::AmbientOcclusion);
+		GetProxyParameter(FGLTFProxyMaterialInfo::OcclusionStrength, OutMaterial.OcclusionTexture.Strength);
+		GetProxyParameter(FGLTFProxyMaterialInfo::Occlusion, OutMaterial.OcclusionTexture);
 
 		if (OutMaterial.ShadingModel == EGLTFJsonShadingModel::ClearCoat)
 		{
-			GetProxyProperty(TEXT("Clear Coat Factor"), OutMaterial.ClearCoat.ClearCoatFactor);
-			GetProxyProperty(TEXT("Clear Coat"), OutMaterial.ClearCoat.ClearCoatTexture, EGLTFMaterialPropertyGroup::ClearCoatRoughness); // TODO: add property group for clear coat intensity only
+			GetProxyParameter(FGLTFProxyMaterialInfo::ClearCoatFactor, OutMaterial.ClearCoat.ClearCoatFactor);
+			GetProxyParameter(FGLTFProxyMaterialInfo::ClearCoat, OutMaterial.ClearCoat.ClearCoatTexture);
 
-			GetProxyProperty(TEXT("Clear Coat Roughness Factor"), OutMaterial.ClearCoat.ClearCoatRoughnessFactor);
-			GetProxyProperty(TEXT("Clear Coat Roughness"), OutMaterial.ClearCoat.ClearCoatRoughnessTexture, EGLTFMaterialPropertyGroup::ClearCoatRoughness);
+			GetProxyParameter(FGLTFProxyMaterialInfo::ClearCoatRoughnessFactor, OutMaterial.ClearCoat.ClearCoatRoughnessFactor);
+			GetProxyParameter(FGLTFProxyMaterialInfo::ClearCoatRoughness, OutMaterial.ClearCoat.ClearCoatRoughnessTexture);
 
-			GetProxyProperty(TEXT("Clear Coat Normal Scale"), OutMaterial.ClearCoat.ClearCoatNormalTexture.Scale);
-			GetProxyProperty(TEXT("Clear Coat Normal"), OutMaterial.ClearCoat.ClearCoatNormalTexture, EGLTFMaterialPropertyGroup::ClearCoatBottomNormal); // TODO: this should actually be top normal
+			GetProxyParameter(FGLTFProxyMaterialInfo::ClearCoatNormalScale, OutMaterial.ClearCoat.ClearCoatNormalTexture.Scale);
+			GetProxyParameter(FGLTFProxyMaterialInfo::ClearCoatNormal, OutMaterial.ClearCoat.ClearCoatNormalTexture);
 		}
 	}
 }
 
-void FGLTFMaterialTask::GetProxyProperty(const FString& PropertyName, float& OutValue) const
+void FGLTFMaterialTask::GetProxyParameter(const TGLTFProxyMaterialParameterInfo<float>& ParameterInfo, float& OutValue) const
 {
-	FGLTFMaterialUtility::GetNonDefaultParameterValue(Material, PropertyName, OutValue);
+	ParameterInfo.Get(Material, OutValue, true);
 }
 
-void FGLTFMaterialTask::GetProxyProperty(const FString& PropertyName, FGLTFJsonColor3& OutValue) const
+void FGLTFMaterialTask::GetProxyParameter(const TGLTFProxyMaterialParameterInfo<FLinearColor>& ParameterInfo, FGLTFJsonColor3& OutValue) const
 {
 	FLinearColor Value;
-	if (FGLTFMaterialUtility::GetNonDefaultParameterValue(Material, PropertyName, Value))
+	if (ParameterInfo.Get(Material, Value, true))
 	{
 		OutValue = FGLTFConverterUtility::ConvertColor3(Value, false);
 	}
 }
 
-void FGLTFMaterialTask::GetProxyProperty(const FString& PropertyName, FGLTFJsonColor4& OutValue) const
+void FGLTFMaterialTask::GetProxyParameter(const TGLTFProxyMaterialParameterInfo<FLinearColor>& ParameterInfo, FGLTFJsonColor4& OutValue) const
 {
 	FLinearColor Value;
-	if (FGLTFMaterialUtility::GetNonDefaultParameterValue(Material, PropertyName, Value))
+	if (ParameterInfo.Get(Material, Value, true))
 	{
 		OutValue = FGLTFConverterUtility::ConvertColor(Value, false);
 	}
 }
 
-void FGLTFMaterialTask::GetProxyProperty(const FString& PropertyName, FGLTFJsonTextureInfo& OutValue, EGLTFMaterialPropertyGroup PropertyGroup) const
+void FGLTFMaterialTask::GetProxyParameter(const FGLTFProxyMaterialTextureParameterInfo& ParameterInfo, FGLTFJsonTextureInfo& OutValue) const
 {
 	UTexture* Texture;
-	if (!FGLTFMaterialUtility::GetNonDefaultParameterValue(Material, PropertyName + TEXT(" Texture"), Texture) || Texture == nullptr)
+	if (!ParameterInfo.Texture.Get(Material, Texture, true) || Texture == nullptr)
 	{
 		return;
 	}
 
-	const bool bSRGB = PropertyGroup == EGLTFMaterialPropertyGroup::BaseColorOpacity || PropertyGroup == EGLTFMaterialPropertyGroup::EmissiveColor;
+	const bool bSRGB = ParameterInfo == FGLTFProxyMaterialInfo::BaseColor || ParameterInfo == FGLTFProxyMaterialInfo::Emissive;
 	OutValue.Index = Builder.GetOrAddTexture(Texture, bSRGB);
 
-	float TexCoord;
-	if (FGLTFMaterialUtility::GetNonDefaultParameterValue(Material, PropertyName + TEXT(" UV Index"), TexCoord))
+	float UVIndex;
+	if (ParameterInfo.UVIndex.Get(Material, UVIndex, true))
 	{
-		OutValue.TexCoord = FMath::RoundToInt(TexCoord);
+		OutValue.TexCoord = FMath::RoundToInt(UVIndex);
 	}
 
-	FLinearColor Offset;
-	if (FGLTFMaterialUtility::GetNonDefaultParameterValue(Material, PropertyName + TEXT(" UV Offset"), Offset))
+	FLinearColor UVOffset;
+	if (ParameterInfo.UVOffset.Get(Material, UVOffset, true))
 	{
-		OutValue.Transform.Offset.X = Offset.R;
-		OutValue.Transform.Offset.Y = Offset.G;
+		OutValue.Transform.Offset.X = UVOffset.R;
+		OutValue.Transform.Offset.Y = UVOffset.G;
 	}
 
-	FLinearColor Scale;
-	if (FGLTFMaterialUtility::GetNonDefaultParameterValue(Material, PropertyName + TEXT(" UV Scale"), Scale))
+	FLinearColor UVScale;
+	if (ParameterInfo.UVScale.Get(Material, UVScale, true))
 	{
-		OutValue.Transform.Scale.X = Scale.R;
-		OutValue.Transform.Scale.Y = Scale.G;
+		OutValue.Transform.Scale.X = UVScale.R;
+		OutValue.Transform.Scale.Y = UVScale.G;
 	}
 
-	float Rotation;
-	if (FGLTFMaterialUtility::GetNonDefaultParameterValue(Material, PropertyName + TEXT(" UV Rotation"), Rotation))
+	float UVRotation;
+	if (ParameterInfo.UVRotation.Get(Material, UVRotation, true))
 	{
-		OutValue.Transform.Rotation = Rotation;
+		OutValue.Transform.Rotation = UVRotation;
 	}
 
 	if (!Builder.ExportOptions->bExportTextureTransforms && !OutValue.Transform.IsExactlyDefault())
@@ -311,7 +320,7 @@ void FGLTFMaterialTask::GetProxyProperty(const FString& PropertyName, FGLTFJsonT
 		Builder.LogWarning(FString::Printf(
 			TEXT("Texture coordinates [%d] in %s for material %s are transformed, but texture transform is disabled by export options"),
 			OutValue.TexCoord,
-			*FGLTFNameUtility::GetName(PropertyGroup),
+			*ParameterInfo.Texture.ToString(),
 			*Material->GetName()));
 		OutValue.Transform = {};
 	}
@@ -438,7 +447,11 @@ void FGLTFMaterialTask::ConvertShadingModel(EGLTFJsonShadingModel& OutShadingMod
 
 void FGLTFMaterialTask::ConvertAlphaMode(EGLTFJsonAlphaMode& OutAlphaMode, EGLTFJsonBlendMode& OutBlendMode) const
 {
-	const EBlendMode BlendMode = Material->GetBlendMode();
+	EBlendMode BlendMode = Material->GetBlendMode();
+	if (const UMaterialInstance* MaterialInstance = Cast<UMaterialInstance>(Material))
+	{
+		FGLTFProxyMaterialUtilities::GetBlendMode(MaterialInstance, BlendMode); // TODO: remove ugly hack for dynamic instances
+	}
 
 	OutAlphaMode = FGLTFConverterUtility::ConvertAlphaMode(BlendMode);
 	if (OutAlphaMode == EGLTFJsonAlphaMode::None)
diff --git a/Engine/Plugins/Enterprise/GLTFExporter/MigratedFromGit/Source/GLTFExporter/Private/Tasks/GLTFMaterialTasks.h b/Engine/Plugins/Enterprise/GLTFExporter/MigratedFromGit/Source/GLTFExporter/Private/Tasks/GLTFMaterialTasks.h
--- a/Engine/Plugins/Enterprise/GLTFExporter/MigratedFromGit/Source/GLTFExporter/Private/Tasks/GLTFMaterialTasks.h
+++ b/Engine/Plugins/Enterprise/GLTFExporter/MigratedFromGit/Source/GLTFExporter/Private/Tasks/GLTFMaterialTasks.h
@@ -6,6 +6,7 @@
 #include "Builders/GLTFConvertBuilder.h"
 #include "Converters/GLTFUVOverlapChecker.h"
 #include "Converters/GLTFMeshData.h"
+#include "Materials/GLTFProxyMaterialParameterInfo.h"
 #if WITH_EDITOR
 #include "MaterialPropertyEx.h"
 #endif
@@ -46,11 +47,11 @@ private:
 	FString GetMaterialName() const;
 	FString GetBakedTextureName(const FString& PropertyName) const;
 
-	void GetProxyProperties(FGLTFJsonMaterial& OutMaterial) const;
-	void GetProxyProperty(const FString& PropertyName, float& OutValue) const;
-	void GetProxyProperty(const FString& PropertyName, FGLTFJsonColor3& OutValue) const;
-	void GetProxyProperty(const FString& PropertyName, FGLTFJsonColor4& OutValue) const;
-	void GetProxyProperty(const FString& PropertyName, FGLTFJsonTextureInfo& OutValue, EGLTFMaterialPropertyGroup PropertyGroup) const;
+	void GetProxyParameters(FGLTFJsonMaterial& OutMaterial) const;
+	void GetProxyParameter(const TGLTFProxyMaterialParameterInfo<float>& ParameterInfo, float& OutValue) const;
+	void GetProxyParameter(const TGLTFProxyMaterialParameterInfo<FLinearColor>& ParameterInfo, FGLTFJsonColor3& OutValue) const;
+	void GetProxyParameter(const TGLTFProxyMaterialParameterInfo<FLinearColor>& ParameterInfo, FGLTFJsonColor4& OutValue) const;
+	void GetProxyParameter(const FGLTFProxyMaterialTextureParameterInfo& ParameterInfo, FGLTFJsonTextureInfo& OutValue) const;
 
 	EMaterialShadingModel GetShadingModel() const;
 	void ConvertShadingModel(EGLTFJsonShadingModel& OutShadingModel) const;
diff --git a/Engine/Plugins/Enterprise/GLTFExporter/MigratedFromGit/Source/GLTFExporter/Private/Utilities/GLTFProxyMaterialUtilities.cpp b/Engine/Plugins/Enterprise/GLTFExporter/MigratedFromGit/Source/GLTFExporter/Private/Utilities/GLTFProxyMaterialUtilities.cpp
new file mode 100644
--- /dev/null
+++ b/Engine/Plugins/Enterprise/GLTFExporter/MigratedFromGit/Source/GLTFExporter/Private/Utilities/GLTFProxyMaterialUtilities.cpp
@@ -0,0 +1,311 @@
+// Copyright Epic Games, Inc. All Rights Reserved.
+
+#include "Utilities/GLTFProxyMaterialUtilities.h"
+#include "Json/GLTFJsonUtility.h"
+#include "UserData/GLTFMaterialUserData.h"
+#include "Materials/MaterialInstanceDynamic.h"
+#include "Materials/MaterialInstanceConstant.h"
+
+#define PROXY_MATERIAL_NAME_PREFIX TEXT("M_GLTF_")
+#define PROXY_MATERIAL_ROOT_PATH   TEXT("/GLTFExporter/Materials/Proxy/")
+
+bool FGLTFProxyMaterialUtilities::IsProxyMaterial(const UMaterial* Material)
+{
+	return Material != nullptr && Material->GetPathName().StartsWith(PROXY_MATERIAL_ROOT_PATH PROXY_MATERIAL_NAME_PREFIX);
+}
+
+bool FGLTFProxyMaterialUtilities::IsProxyMaterial(const UMaterialInterface* Material)
+{
+	return Material != nullptr && IsProxyMaterial(Material->GetMaterial());
+}
+
+UMaterial* FGLTFProxyMaterialUtilities::GetBaseMaterial(EGLTFJsonShadingModel ShadingModel)
+{
+	constexpr int32 NumMaterials = static_cast<int32>(EGLTFJsonShadingModel::NumShadingModels);
+	static UMaterial* Materials[NumMaterials] = { nullptr };
+
+	const int32 Index = static_cast<int32>(ShadingModel);
+	if (Index < 0 || NumMaterials <= Index)
+	{
+		checkNoEntry();
+		return nullptr;
+	}
+
+	UMaterial* Material = Materials[Index];
+	if (Material == nullptr)
+	{
+		const FString Name = FGLTFJsonUtility::GetValue(ShadingModel);
+		const FString Path = PROXY_MATERIAL_ROOT_PATH PROXY_MATERIAL_NAME_PREFIX + Name + TEXT(".") PROXY_MATERIAL_NAME_PREFIX + Name;
+
+		Material = LoadObject<UMaterial>(nullptr, *Path);
+		Materials[Index] = Material;
+	}
+
+	return Material;
+}
+
+UMaterialInterface* FGLTFProxyMaterialUtilities::GetProxyMaterial(UMaterialInterface* OriginalMaterial)
+{
+	UGLTFMaterialExportOptions* UserData = OriginalMaterial->GetAssetUserData<UGLTFMaterialExportOptions>();
+	return UserData != nullptr ? UserData->Proxy : nullptr;
+}
+
+void FGLTFProxyMaterialUtilities::SetProxyMaterial(UMaterialInterface* OriginalMaterial, UMaterialInterface* ProxyMaterial)
+{
+	UGLTFMaterialExportOptions* UserData = OriginalMaterial->GetAssetUserData<UGLTFMaterialExportOptions>();
+	if (UserData == nullptr)
+	{
+		UserData = NewObject<UGLTFMaterialExportOptions>();
+		OriginalMaterial->AddAssetUserData(UserData);
+	}
+
+	UserData->Proxy = ProxyMaterial;
+	OriginalMaterial->Modify();
+}
+
+bool FGLTFProxyMaterialUtilities::GetParameterValue(const UMaterialInterface* Material, const FHashedMaterialParameterInfo& ParameterInfo, float& OutValue, bool NonDefaultOnly)
+{
+	float Value;
+	if (!Material->GetScalarParameterValue(ParameterInfo, Value, NonDefaultOnly))
+	{
+		return false;
+	}
+
+	if (NonDefaultOnly)
+	{
+		float DefaultValue;
+		if (!Material->GetScalarParameterDefaultValue(ParameterInfo, DefaultValue) || DefaultValue == Value)
+		{
+			return false;
+		}
+	}
+
+	OutValue = Value;
+	return true;
+}
+
+bool FGLTFProxyMaterialUtilities::GetParameterValue(const UMaterialInterface* Material, const FHashedMaterialParameterInfo& ParameterInfo, FLinearColor& OutValue, bool NonDefaultOnly)
+{
+	FLinearColor Value;
+	if (!Material->GetVectorParameterValue(ParameterInfo, Value, NonDefaultOnly))
+	{
+		return false;
+	}
+
+	if (NonDefaultOnly)
+	{
+		FLinearColor DefaultValue;
+		if (!Material->GetVectorParameterDefaultValue(ParameterInfo, DefaultValue) || DefaultValue == Value)
+		{
+			return false;
+		}
+	}
+
+	OutValue = Value;
+	return true;
+}
+
+bool FGLTFProxyMaterialUtilities::GetParameterValue(const UMaterialInterface* Material, const FHashedMaterialParameterInfo& ParameterInfo, UTexture*& OutValue, bool NonDefaultOnly)
+{
+	UTexture* Value;
+	if (!Material->GetTextureParameterValue(ParameterInfo, Value, NonDefaultOnly))
+	{
+		return false;
+	}
+
+	if (NonDefaultOnly)
+	{
+		UTexture* DefaultValue;
+		if (!Material->GetTextureParameterDefaultValue(ParameterInfo, DefaultValue) || DefaultValue == Value)
+		{
+			return false;
+		}
+	}
+
+	OutValue = Value;
+	return true;
+}
+
+void FGLTFProxyMaterialUtilities::SetParameterValue(UMaterialInstanceDynamic* Material, const FHashedMaterialParameterInfo& ParameterInfo, float Value, bool NonDefaultOnly)
+{
+	if (NonDefaultOnly)
+	{
+		float DefaultValue;
+		if (!Material->GetScalarParameterDefaultValue(ParameterInfo, DefaultValue) || DefaultValue == Value)
+		{
+			return;
+		}
+	}
+
+	Material->SetScalarParameterValueByInfo(FMaterialParameterInfo(ParameterInfo), Value);
+}
+
+void FGLTFProxyMaterialUtilities::SetParameterValue(UMaterialInstanceDynamic* Material, const FHashedMaterialParameterInfo& ParameterInfo, const FLinearColor& Value, bool NonDefaultOnly)
+{
+	if (NonDefaultOnly)
+	{
+		FLinearColor DefaultValue;
+		if (!Material->GetVectorParameterDefaultValue(ParameterInfo, DefaultValue) || DefaultValue == Value)
+		{
+			return;
+		}
+	}
+
+	Material->SetVectorParameterValueByInfo(FMaterialParameterInfo(ParameterInfo), Value);
+}
+
+void FGLTFProxyMaterialUtilities::SetParameterValue(UMaterialInstanceDynamic* Material, const FHashedMaterialParameterInfo& ParameterInfo, UTexture* Value, bool NonDefaultOnly)
+{
+	if (NonDefaultOnly)
+	{
+		UTexture* DefaultValue;
+		if (!Material->GetTextureParameterDefaultValue(ParameterInfo, DefaultValue) || DefaultValue == Value)
+		{
+			return;
+		}
+	}
+
+	Material->SetTextureParameterValueByInfo(FMaterialParameterInfo(ParameterInfo), Value);
+}
+
+#if WITH_EDITOR
+
+void FGLTFProxyMaterialUtilities::SetParameterValue(UMaterialInstanceConstant* Material, const FHashedMaterialParameterInfo& ParameterInfo, float Value, bool NonDefaultOnly)
+{
+	if (NonDefaultOnly)
+	{
+		float DefaultValue;
+		if (!Material->GetScalarParameterDefaultValue(ParameterInfo, DefaultValue) || DefaultValue == Value)
+		{
+			return;
+		}
+	}
+
+	Material->SetScalarParameterValueEditorOnly(FMaterialParameterInfo(ParameterInfo), Value);
+}
+
+void FGLTFProxyMaterialUtilities::SetParameterValue(UMaterialInstanceConstant* Material, const FHashedMaterialParameterInfo& ParameterInfo, const FLinearColor& Value, bool NonDefaultOnly)
+{
+	if (NonDefaultOnly)
+	{
+		FLinearColor DefaultValue;
+		if (!Material->GetVectorParameterDefaultValue(ParameterInfo, DefaultValue) || DefaultValue == Value)
+		{
+			return;
+		}
+	}
+
+	Material->SetVectorParameterValueEditorOnly(FMaterialParameterInfo(ParameterInfo), Value);
+}
+
+void FGLTFProxyMaterialUtilities::SetParameterValue(UMaterialInstanceConstant* Material, const FHashedMaterialParameterInfo& ParameterInfo, UTexture* Value, bool NonDefaultOnly)
+{
+	if (NonDefaultOnly)
+	{
+		UTexture* DefaultValue;
+		if (!Material->GetTextureParameterDefaultValue(ParameterInfo, DefaultValue) || DefaultValue == Value)
+		{
+			return;
+		}
+	}
+
+	Material->SetTextureParameterValueEditorOnly(FMaterialParameterInfo(ParameterInfo), Value);
+}
+
+#endif
+
+bool FGLTFProxyMaterialUtilities::GetTwoSided(const UMaterialInstance* Material, bool& OutValue, bool NonDefaultOnly)
+{
+	const bool DefaultValue = Material->Parent->IsTwoSided();
+	const bool Value = Material->BasePropertyOverrides.bOverride_TwoSided ? Material->BasePropertyOverrides.TwoSided : DefaultValue;
+
+	if (NonDefaultOnly && Value == DefaultValue)
+	{
+		return false;
+	}
+
+	OutValue = Value;
+	return true;
+}
+
+bool FGLTFProxyMaterialUtilities::GetBlendMode(const UMaterialInstance* Material, EBlendMode& OutValue, bool NonDefaultOnly)
+{
+	const EBlendMode DefaultValue = Material->Parent->GetBlendMode();
+	const EBlendMode Value = Material->BasePropertyOverrides.bOverride_BlendMode ? Material->BasePropertyOverrides.BlendMode : DefaultValue;
+
+	if (NonDefaultOnly && Value == DefaultValue)
+	{
+		return false;
+	}
+
+	OutValue = Value;
+	return true;
+}
+
+bool FGLTFProxyMaterialUtilities::GetOpacityMaskClipValue(const UMaterialInstance* Material, float& OutValue, bool NonDefaultOnly)
+{
+	const float DefaultValue = Material->Parent->GetOpacityMaskClipValue();
+	const float Value = Material->BasePropertyOverrides.bOverride_OpacityMaskClipValue ? Material->BasePropertyOverrides.OpacityMaskClipValue : DefaultValue;
+
+	if (NonDefaultOnly && Value == DefaultValue)
+	{
+		return false;
+	}
+
+	OutValue = Value;
+	return true;
+}
+
+void FGLTFProxyMaterialUtilities::SetTwoSided(UMaterialInstance* Material, bool Value, bool NonDefaultOnly)
+{
+	// TODO: Even if the exporter supports it, UMaterialInstanceDynamic won't actually reflect these overriden base properties in Unreal
+
+	if (NonDefaultOnly)
+	{
+		const bool DefaultValue = Material->Parent->IsTwoSided();
+		if (DefaultValue == Value)
+		{
+			return;
+		}
+	}
+
+	Material->BasePropertyOverrides.bOverride_TwoSided = true;
+	Material->BasePropertyOverrides.TwoSided = Value;
+	Material->TwoSided = Value;
+}
+
+void FGLTFProxyMaterialUtilities::SetBlendMode(UMaterialInstance* Material, EBlendMode Value, bool NonDefaultOnly)
+{
+	// TODO: Even if the exporter supports it, UMaterialInstanceDynamic won't actually reflect these overriden base properties in Unreal
+
+	if (NonDefaultOnly)
+	{
+		const EBlendMode DefaultValue = Material->Parent->GetBlendMode();
+		if (DefaultValue == Value)
+		{
+			return;
+		}
+	}
+
+	Material->BasePropertyOverrides.bOverride_BlendMode = true;
+	Material->BasePropertyOverrides.BlendMode = Value;
+	Material->BlendMode = Value;
+}
+
+void FGLTFProxyMaterialUtilities::SetOpacityMaskClipValue(UMaterialInstance* Material, float Value, bool NonDefaultOnly)
+{
+	// TODO: Even if the exporter supports it, UMaterialInstanceDynamic won't actually reflect these overriden base properties in Unreal
+
+	if (NonDefaultOnly)
+	{
+		const float DefaultValue = Material->Parent->GetOpacityMaskClipValue();
+		if (DefaultValue == Value)
+		{
+			return;
+		}
+	}
+
+	Material->BasePropertyOverrides.bOverride_OpacityMaskClipValue = true;
+	Material->BasePropertyOverrides.OpacityMaskClipValue = Value;
+	Material->OpacityMaskClipValue = Value;
+}
diff --git a/Engine/Plugins/Enterprise/GLTFExporter/MigratedFromGit/Source/GLTFExporter/Public/Json/GLTFJsonEnums.h b/Engine/Plugins/Enterprise/GLTFExporter/MigratedFromGit/Source/GLTFExporter/Public/Json/GLTFJsonEnums.h
--- a/Engine/Plugins/Enterprise/GLTFExporter/MigratedFromGit/Source/GLTFExporter/Public/Json/GLTFJsonEnums.h
+++ b/Engine/Plugins/Enterprise/GLTFExporter/MigratedFromGit/Source/GLTFExporter/Public/Json/GLTFJsonEnums.h
@@ -28,7 +28,9 @@ enum class EGLTFJsonShadingModel
 	None = -1,
 	Default,
 	Unlit,
-	ClearCoat
+	ClearCoat,
+
+	NumShadingModels
 };
 
 enum class EGLTFJsonHDREncoding
diff --git a/Engine/Plugins/Enterprise/GLTFExporter/MigratedFromGit/Source/GLTFExporter/Public/Materials/GLTFProxyMaterialInfo.h b/Engine/Plugins/Enterprise/GLTFExporter/MigratedFromGit/Source/GLTFExporter/Public/Materials/GLTFProxyMaterialInfo.h
new file mode 100644
--- /dev/null
+++ b/Engine/Plugins/Enterprise/GLTFExporter/MigratedFromGit/Source/GLTFExporter/Public/Materials/GLTFProxyMaterialInfo.h
@@ -0,0 +1,35 @@
+// Copyright Epic Games, Inc. All Rights Reserved.
+
+#pragma once
+
+#include "Materials/GLTFProxyMaterialParameterInfo.h"
+
+class GLTFEXPORTER_API FGLTFProxyMaterialInfo
+{
+public:
+
+	static const FGLTFProxyMaterialTextureParameterInfo BaseColor;
+	static const TGLTFProxyMaterialParameterInfo<FLinearColor> BaseColorFactor;
+
+	static const FGLTFProxyMaterialTextureParameterInfo Emissive;
+	static const TGLTFProxyMaterialParameterInfo<FLinearColor> EmissiveFactor;
+
+	static const FGLTFProxyMaterialTextureParameterInfo MetallicRoughness;
+	static const TGLTFProxyMaterialParameterInfo<float> MetallicFactor;
+	static const TGLTFProxyMaterialParameterInfo<float> RoughnessFactor;
+
+	static const FGLTFProxyMaterialTextureParameterInfo Normal;
+	static const TGLTFProxyMaterialParameterInfo<float> NormalScale;
+
+	static const FGLTFProxyMaterialTextureParameterInfo Occlusion;
+	static const TGLTFProxyMaterialParameterInfo<float> OcclusionStrength;
+
+	static const FGLTFProxyMaterialTextureParameterInfo ClearCoat;
+	static const TGLTFProxyMaterialParameterInfo<float> ClearCoatFactor;
+
+	static const FGLTFProxyMaterialTextureParameterInfo ClearCoatRoughness;
+	static const TGLTFProxyMaterialParameterInfo<float> ClearCoatRoughnessFactor;
+
+	static const FGLTFProxyMaterialTextureParameterInfo ClearCoatNormal;
+	static const TGLTFProxyMaterialParameterInfo<float> ClearCoatNormalScale;
+};
diff --git a/Engine/Plugins/Enterprise/GLTFExporter/MigratedFromGit/Source/GLTFExporter/Public/Materials/GLTFProxyMaterialParameterInfo.h b/Engine/Plugins/Enterprise/GLTFExporter/MigratedFromGit/Source/GLTFExporter/Public/Materials/GLTFProxyMaterialParameterInfo.h
new file mode 100644
--- /dev/null
+++ b/Engine/Plugins/Enterprise/GLTFExporter/MigratedFromGit/Source/GLTFExporter/Public/Materials/GLTFProxyMaterialParameterInfo.h
@@ -0,0 +1,84 @@
+// Copyright Epic Games, Inc. All Rights Reserved.
+
+#pragma once
+
+#include "Utilities/GLTFProxyMaterialUtilities.h"
+
+template <typename ParameterType>
+class TGLTFProxyMaterialParameterInfo
+{
+public:
+
+	TGLTFProxyMaterialParameterInfo(const FString& ParameterName)
+		: ParameterInfo(*ParameterName)
+	{
+	}
+
+	bool Get(const UMaterialInterface* Material, ParameterType& OutValue, bool NonDefaultOnly = false) const
+	{
+		return FGLTFProxyMaterialUtilities::GetParameterValue(Material, ParameterInfo, OutValue, NonDefaultOnly);
+	}
+
+	void Set(UMaterialInstanceDynamic* Material, const ParameterType& Value, bool NonDefaultOnly = false) const
+	{
+		FGLTFProxyMaterialUtilities::SetParameterValue(Material, ParameterInfo, Value, NonDefaultOnly);
+	}
+
+#if WITH_EDITOR
+	void Set(UMaterialInstanceConstant* Material, const ParameterType& Value, bool NonDefaultOnly = false) const
+	{
+		FGLTFProxyMaterialUtilities::SetParameterValue(Material, ParameterInfo, Value, NonDefaultOnly);
+	}
+#endif
+
+	FString ToString() const
+	{
+		return ParameterInfo.Name.ToString();
+	}
+
+	bool operator==(const TGLTFProxyMaterialParameterInfo<ParameterType>& Other) const
+	{
+		return ParameterInfo.Name == Other.ParameterInfo.Name
+			&& ParameterInfo.Index == Other.ParameterInfo.Index
+			&& ParameterInfo.Association == Other.ParameterInfo.Association;
+	}
+
+	bool operator!=(const TGLTFProxyMaterialParameterInfo<ParameterType>& Other) const
+	{
+		return !(*this == Other);
+	}
+
+private:
+
+	FHashedMaterialParameterInfo ParameterInfo;
+};
+
+class FGLTFProxyMaterialTextureParameterInfo
+{
+public:
+
+	FGLTFProxyMaterialTextureParameterInfo(const FString& ParameterName)
+		: Texture(ParameterName + TEXT(" Texture"))
+		, UVIndex(ParameterName + TEXT(" UV Index"))
+		, UVOffset(ParameterName + TEXT(" UV Offset"))
+		, UVScale(ParameterName + TEXT(" UV Scale"))
+		, UVRotation(ParameterName + TEXT(" UV Rotation"))
+	{
+	}
+
+	bool operator==(const FGLTFProxyMaterialTextureParameterInfo& Other) const
+	{
+		return Texture == Other.Texture;
+	}
+
+	bool operator!=(const FGLTFProxyMaterialTextureParameterInfo& Other) const
+	{
+		return !(*this == Other);
+	}
+
+	TGLTFProxyMaterialParameterInfo<UTexture*> Texture;
+	TGLTFProxyMaterialParameterInfo<float> UVIndex;
+	TGLTFProxyMaterialParameterInfo<FLinearColor> UVOffset;
+	TGLTFProxyMaterialParameterInfo<FLinearColor> UVScale;
+	TGLTFProxyMaterialParameterInfo<float> UVRotation;
+};
diff --git a/Engine/Plugins/Enterprise/GLTFExporter/MigratedFromGit/Source/GLTFExporter/Public/Utilities/GLTFProxyMaterialUtilities.h b/Engine/Plugins/Enterprise/GLTFExporter/MigratedFromGit/Source/GLTFExporter/Public/Utilities/GLTFProxyMaterialUtilities.h
new file mode 100644
--- /dev/null
+++ b/Engine/Plugins/Enterprise/GLTFExporter/MigratedFromGit/Source/GLTFExporter/Public/Utilities/GLTFProxyMaterialUtilities.h
@@ -0,0 +1,52 @@
+// Copyright Epic Games, Inc. All Rights Reserved.
+
+#pragma once
+
+#include "Json/GLTFJsonEnums.h"
+
+class UMaterialInstance;
+class UMaterialInstanceDynamic;
+class UMaterialInstanceConstant;
+
+class GLTFEXPORTER_API FGLTFProxyMaterialUtilities
+{
+public:
+
+	template <typename MaterialType, typename = typename TEnableIf<TIsDerivedFrom<MaterialType, UMaterialInstance>::Value>::Type>
+	static MaterialType* CreateProxyMaterial(EGLTFJsonShadingModel ShadingModel, UObject* Outer = GetTransientPackage(), FName Name = NAME_None, EObjectFlags Flags = RF_NoFlags)
+	{
+		MaterialType* MaterialInstance = NewObject<MaterialType>(Outer, Name, Flags);
+		MaterialInstance->Parent = GetBaseMaterial(ShadingModel);
+		return MaterialInstance;
+	}
+
+	static bool IsProxyMaterial(const UMaterial* Material);
+	static bool IsProxyMaterial(const UMaterialInterface* Material);
+
+	static UMaterial* GetBaseMaterial(EGLTFJsonShadingModel ShadingModel);
+
+	static UMaterialInterface* GetProxyMaterial(UMaterialInterface* OriginalMaterial);
+	static void SetProxyMaterial(UMaterialInterface* OriginalMaterial, UMaterialInterface* ProxyMaterial);
+
+	static bool GetParameterValue(const UMaterialInterface* Material, const FHashedMaterialParameterInfo& ParameterInfo, float& OutValue, bool NonDefaultOnly = false);
+	static bool GetParameterValue(const UMaterialInterface* Material, const FHashedMaterialParameterInfo& ParameterInfo, FLinearColor& OutValue, bool NonDefaultOnly = false);
+	static bool GetParameterValue(const UMaterialInterface* Material, const FHashedMaterialParameterInfo& ParameterInfo, UTexture*& OutValue, bool NonDefaultOnly = false);
+
+	static void SetParameterValue(UMaterialInstanceDynamic* Material, const FHashedMaterialParameterInfo& ParameterInfo, float Value, bool NonDefaultOnly = false);
+	static void SetParameterValue(UMaterialInstanceDynamic* Material, const FHashedMaterialParameterInfo& ParameterInfo, const FLinearColor& Value, bool NonDefaultOnly = false);
+	static void SetParameterValue(UMaterialInstanceDynamic* Material, const FHashedMaterialParameterInfo& ParameterInfo, UTexture* Value, bool NonDefaultOnly = false);
+
+#if WITH_EDITOR
+	static void SetParameterValue(UMaterialInstanceConstant* Material, const FHashedMaterialParameterInfo& ParameterInfo, float Value, bool NonDefaultOnly = false);
+	static void SetParameterValue(UMaterialInstanceConstant* Material, const FHashedMaterialParameterInfo& ParameterInfo, const FLinearColor& Value, bool NonDefaultOnly = false);
+	static void SetParameterValue(UMaterialInstanceConstant* Material, const FHashedMaterialParameterInfo& ParameterInfo, UTexture* Value, bool NonDefaultOnly = false);
+#endif
+
+	static bool GetTwoSided(const UMaterialInstance* Material, bool& OutValue, bool NonDefaultOnly = false);
+	static bool GetBlendMode(const UMaterialInstance* Material, EBlendMode& OutValue, bool NonDefaultOnly = false);
+	static bool GetOpacityMaskClipValue(const UMaterialInstance* Material, float& OutValue, bool NonDefaultOnly = false);
+
+	static void SetTwoSided(UMaterialInstance* Material, bool Value, bool NonDefaultOnly = false);
+	static void SetBlendMode(UMaterialInstance* Material, EBlendMode Value, bool NonDefaultOnly = false);
+	static void SetOpacityMaskClipValue(UMaterialInstance* Material, float Value, bool NonDefaultOnly = false);
+};
