[GLTFExporter] fix: ensure proxy base materials always are cooked for runtime use

commit 773497a746aa3974ff1f4808d3ac5dea38916ce0
Date: Wed, 23 Feb 2022 18:29:25 +0100

    fix: ensure proxy base materials always are cooked for runtime use
    (cherry picked from commit 368cf1086673362770d295866e614201796497e6​)

#skipci

diff --git a/Engine/Plugins/Enterprise/GLTFExporter/MigratedFromGit/Source/GLTFExporter/Private/Converters/GLTFObjectArrayScopeGuard.h b/Engine/Plugins/Enterprise/GLTFExporter/MigratedFromGit/Source/GLTFExporter/Private/Converters/GLTFObjectArrayScopeGuard.h
new file mode 100644
--- /dev/null
+++ b/Engine/Plugins/Enterprise/GLTFExporter/MigratedFromGit/Source/GLTFExporter/Private/Converters/GLTFObjectArrayScopeGuard.h
@@ -0,0 +1,25 @@
+// Copyright Epic Games, Inc. All Rights Reserved.
+
+#pragma once
+
+#include "CoreMinimal.h"
+
+class FGLTFObjectArrayScopeGuard : public FGCObject, public TArray<UObject*>
+{
+public:
+
+	FGLTFObjectArrayScopeGuard()
+	{
+	}
+
+	using TArray<UObject*>::TArray;
+
+	/** Non-copyable */
+	FGLTFObjectArrayScopeGuard(const FGLTFObjectArrayScopeGuard&) = delete;
+	FGLTFObjectArrayScopeGuard& operator=(const FGLTFObjectArrayScopeGuard&) = delete;
+
+	virtual void AddReferencedObjects(FReferenceCollector& Collector) override
+	{
+		Collector.AddReferencedObjects(*this);
+	}
+};
diff --git a/Engine/Plugins/Enterprise/GLTFExporter/MigratedFromGit/Source/GLTFExporter/Private/GLTFExporterModule.cpp b/Engine/Plugins/Enterprise/GLTFExporter/MigratedFromGit/Source/GLTFExporter/Private/GLTFExporterModule.cpp
--- a/Engine/Plugins/Enterprise/GLTFExporter/MigratedFromGit/Source/GLTFExporter/Private/GLTFExporterModule.cpp
+++ b/Engine/Plugins/Enterprise/GLTFExporter/MigratedFromGit/Source/GLTFExporter/Private/GLTFExporterModule.cpp
@@ -1,6 +1,9 @@
 // Copyright Epic Games, Inc. All Rights Reserved.
 
 #include "GLTFExporterModule.h"
+#include "Json/GLTFJsonEnums.h"
+#include "Converters/GLTFObjectArrayScopeGuard.h"
+#include "Utilities/GLTFProxyMaterialUtilities.h"
 #include "Actions/GLTFProxyAssetActions.h"
 #include "Interfaces/IPluginManager.h"
 #if WITH_EDITOR
@@ -31,6 +34,23 @@ public:
 
 #if WITH_EDITOR
 	void PostEngineInit()
+	{
+		LoadProxyParentsForCook();
+		AddProxyAssetActions();
+	}
+
+	static void LoadProxyParentsForCook()
+	{
+		static FGLTFObjectArrayScopeGuard ScopeGuard;
+
+		constexpr int32 NumMaterials = static_cast<int32>(EGLTFJsonShadingModel::NumShadingModels);
+		for (int32 i = 0; i < NumMaterials; ++i)
+		{
+			ScopeGuard.Add(FGLTFProxyMaterialUtilities::GetBaseMaterial(static_cast<EGLTFJsonShadingModel>(i)));
+		}
+	}
+
+	void AddProxyAssetActions()
 	{
 		if (FAssetToolsModule::IsModuleLoaded())
 		{
@@ -48,11 +68,9 @@ public:
 			}
 		}
 	}
-#endif
 
-	virtual void ShutdownModule() override
+	void RemoveProxyAssetActions()
 	{
-#if WITH_EDITOR
 		if (FAssetToolsModule::IsModuleLoaded())
 		{
 			IAssetTools& AssetTools = FAssetToolsModule::GetModule().Get();
@@ -63,6 +81,13 @@ public:
 		}
 
 		AssetTypeActionsArray.Empty();
+	}
+#endif
+
+	virtual void ShutdownModule() override
+	{
+#if WITH_EDITOR
+		RemoveProxyAssetActions();
 #endif
 	}
 
