[GLTFExporter] feat: add high-level public API FGLTFProxyMaterial to simplify creating proxy materials even more

commit 38cc651561751db39c3768bb21fb19f4ab8b3d98
Date: Wed, 23 Feb 2022 18:30:25 +0100

    feat: add high-level public API FGLTFProxyMaterial to simplify creating proxy materials even more
    (cherry picked from commit 77212113505a02f7bd2a39f6b74821262aa600dc​)

#skipci

diff --git a/Engine/Plugins/Enterprise/GLTFExporter/MigratedFromGit/Source/GLTFExporter/Public/Materials/GLTFProxyMaterial.h b/Engine/Plugins/Enterprise/GLTFExporter/MigratedFromGit/Source/GLTFExporter/Public/Materials/GLTFProxyMaterial.h
new file mode 100644
--- /dev/null
+++ b/Engine/Plugins/Enterprise/GLTFExporter/MigratedFromGit/Source/GLTFExporter/Public/Materials/GLTFProxyMaterial.h
@@ -0,0 +1,73 @@
+// Copyright Epic Games, Inc. All Rights Reserved.
+
+#pragma once
+
+#include "Materials/GLTFProxyMaterialInfo.h"
+#include "Materials/GLTFProxyMaterialParameter.h"
+#include "Utilities/GLTFProxyMaterialUtilities.h"
+#include "UObject/StrongObjectPtr.h"
+
+template <typename MaterialType, typename = typename TEnableIf<TIsDerivedFrom<MaterialType, UMaterialInstance>::Value>::Type>
+class TGLTFProxyMaterial
+{
+public:
+
+	template <typename = typename TEnableIf<TNot<TIsSame<MaterialType, UMaterialInstance>>::Value>::Type>
+	static TGLTFProxyMaterial Create(EGLTFJsonShadingModel ShadingModel, UObject* Outer = GetTransientPackage(), FName Name = NAME_None, EObjectFlags Flags = RF_NoFlags)
+	{
+		return TGLTFProxyMaterial(FGLTFProxyMaterialUtilities::CreateProxyMaterial<MaterialType>(ShadingModel, Outer, Name, Flags));
+	}
+
+	TGLTFProxyMaterial(MaterialType* Material)
+		: Material(Material)
+	{
+		check(FGLTFProxyMaterialUtilities::IsProxyMaterial(Material));
+	}
+
+	MaterialType* GetMaterial() { return Material.Get(); }
+
+	void SetProxy(UMaterialInterface* OriginalMaterial) { FGLTFProxyMaterialUtilities::SetProxyMaterial(OriginalMaterial, GetMaterial()); }
+
+	bool GetTwoSided() { return FGLTFProxyMaterialUtilities::GetTwoSided(GetMaterial()); }
+	EBlendMode GetBlendMode() { return FGLTFProxyMaterialUtilities::GetBlendMode(GetMaterial()); }
+	float GetOpacityMaskClipValue() { return FGLTFProxyMaterialUtilities::GetOpacityMaskClipValue(GetMaterial()); }
+
+	void SetTwoSided(bool Value, bool NonDefaultOnly = false) { FGLTFProxyMaterialUtilities::SetTwoSided(GetMaterial(), Value, NonDefaultOnly); }
+	void SetBlendMode(EBlendMode Value, bool NonDefaultOnly = false) { FGLTFProxyMaterialUtilities::SetBlendMode(GetMaterial(), Value, NonDefaultOnly); }
+	void SetOpacityMaskClipValue(float Value, bool NonDefaultOnly = false) { FGLTFProxyMaterialUtilities::SetOpacityMaskClipValue(GetMaterial(), Value, NonDefaultOnly); }
+
+private:
+
+	TStrongObjectPtr<MaterialType> Material;
+
+public:
+
+	const TGLTFProxyMaterialTextureParameter<MaterialType> BaseColor = { FGLTFProxyMaterialInfo::BaseColor, this->GetMaterial() };
+	const TGLTFProxyMaterialParameter<FLinearColor, MaterialType> BaseColorFactor = { FGLTFProxyMaterialInfo::BaseColorFactor, this->GetMaterial() };
+
+	const TGLTFProxyMaterialTextureParameter<MaterialType> Emissive = { FGLTFProxyMaterialInfo::Emissive, this->GetMaterial() };
+	const TGLTFProxyMaterialParameter<FLinearColor, MaterialType> EmissiveFactor = { FGLTFProxyMaterialInfo::EmissiveFactor, this->GetMaterial() };
+
+	const TGLTFProxyMaterialTextureParameter<MaterialType> MetallicRoughness = { FGLTFProxyMaterialInfo::MetallicRoughness, this->GetMaterial() };
+	const TGLTFProxyMaterialParameter<float, MaterialType> MetallicFactor = { FGLTFProxyMaterialInfo::MetallicFactor, this->GetMaterial() };
+	const TGLTFProxyMaterialParameter<float, MaterialType> RoughnessFactor = { FGLTFProxyMaterialInfo::RoughnessFactor, this->GetMaterial() };
+
+	const TGLTFProxyMaterialTextureParameter<MaterialType> Normal = { FGLTFProxyMaterialInfo::Normal, this->GetMaterial() };
+	const TGLTFProxyMaterialParameter<float, MaterialType> NormalScale = { FGLTFProxyMaterialInfo::NormalScale, this->GetMaterial() };
+
+	const TGLTFProxyMaterialTextureParameter<MaterialType> Occlusion = { FGLTFProxyMaterialInfo::Occlusion, this->GetMaterial() };
+	const TGLTFProxyMaterialParameter<float, MaterialType> OcclusionStrength = { FGLTFProxyMaterialInfo::OcclusionStrength, this->GetMaterial() };
+
+	const TGLTFProxyMaterialTextureParameter<MaterialType> ClearCoat = { FGLTFProxyMaterialInfo::ClearCoat, this->GetMaterial() };
+	const TGLTFProxyMaterialParameter<float, MaterialType> ClearCoatFactor = { FGLTFProxyMaterialInfo::ClearCoatFactor, this->GetMaterial() };
+
+	const TGLTFProxyMaterialTextureParameter<MaterialType> ClearCoatRoughness = { FGLTFProxyMaterialInfo::ClearCoatRoughness, this->GetMaterial() };
+	const TGLTFProxyMaterialParameter<float, MaterialType> ClearCoatRoughnessFactor = { FGLTFProxyMaterialInfo::ClearCoatRoughnessFactor, this->GetMaterial() };
+
+	const TGLTFProxyMaterialTextureParameter<MaterialType> ClearCoatNormal = { FGLTFProxyMaterialInfo::ClearCoatNormal, this->GetMaterial() };
+	const TGLTFProxyMaterialParameter<float, MaterialType> ClearCoatNormalScale = { FGLTFProxyMaterialInfo::ClearCoatNormalScale, this->GetMaterial() };
+};
+
+typedef TGLTFProxyMaterial<UMaterialInstance, void> FGLTFProxyMaterial;
+typedef TGLTFProxyMaterial<UMaterialInstanceDynamic, void> FGLTFProxyMaterialDynamic;
+typedef TGLTFProxyMaterial<UMaterialInstanceConstant, void> FGLTFProxyMaterialConstant;
diff --git a/Engine/Plugins/Enterprise/GLTFExporter/MigratedFromGit/Source/GLTFExporter/Public/Materials/GLTFProxyMaterialParameter.h b/Engine/Plugins/Enterprise/GLTFExporter/MigratedFromGit/Source/GLTFExporter/Public/Materials/GLTFProxyMaterialParameter.h
new file mode 100644
--- /dev/null
+++ b/Engine/Plugins/Enterprise/GLTFExporter/MigratedFromGit/Source/GLTFExporter/Public/Materials/GLTFProxyMaterialParameter.h
@@ -0,0 +1,94 @@
+// Copyright Epic Games, Inc. All Rights Reserved.
+
+#pragma once
+
+#include "Materials/GLTFProxyMaterialParameterInfo.h"
+
+template <typename ParameterType, typename MaterialType, typename = typename TEnableIf<TIsDerivedFrom<MaterialType, UMaterialInterface>::Value>::Type>
+class TGLTFProxyMaterialParameterBase
+{
+public:
+
+	TGLTFProxyMaterialParameterBase(const TGLTFProxyMaterialParameterInfo<ParameterType>& ParameterInfo, MaterialType* Material)
+		: ParameterInfo(ParameterInfo)
+		, Material(Material)
+	{
+	}
+
+	bool Get(ParameterType& OutValue, bool NonDefaultOnly = false) const
+	{
+		return ParameterInfo.Get(Material, OutValue, NonDefaultOnly);
+	}
+
+	ParameterType Get() const
+	{
+		ParameterType Value;
+		if (!ParameterInfo.Get(Material, Value))
+		{
+			checkNoEntry();
+		}
+		return Value;
+	}
+
+protected:
+
+	const TGLTFProxyMaterialParameterInfo<ParameterType> ParameterInfo;
+	MaterialType* Material;
+};
+
+template <typename ParameterType, typename MaterialType, typename = typename TEnableIf<TIsDerivedFrom<MaterialType, UMaterialInterface>::Value>::Type>
+class TGLTFProxyMaterialParameter: public TGLTFProxyMaterialParameterBase<ParameterType, MaterialType, void>
+{
+public:
+
+	using TGLTFProxyMaterialParameterBase<ParameterType, MaterialType, void>::TGLTFProxyMaterialParameterBase;
+};
+
+template <typename ParameterType>
+class TGLTFProxyMaterialParameter<ParameterType, UMaterialInstanceDynamic, void>: public TGLTFProxyMaterialParameterBase<ParameterType, UMaterialInstanceDynamic, void>
+{
+public:
+
+	using TGLTFProxyMaterialParameterBase<ParameterType, UMaterialInstanceDynamic, void>::TGLTFProxyMaterialParameterBase;
+
+	void Set(const ParameterType& Value, bool NonDefaultOnly = false) const
+	{
+		this->ParameterInfo.Set(this->Material, Value, NonDefaultOnly);
+	}
+};
+
+template <typename ParameterType>
+class TGLTFProxyMaterialParameter<ParameterType, UMaterialInstanceConstant, void>: public TGLTFProxyMaterialParameterBase<ParameterType, UMaterialInstanceConstant, void>
+{
+public:
+
+	using TGLTFProxyMaterialParameterBase<ParameterType, UMaterialInstanceConstant, void>::TGLTFProxyMaterialParameterBase;
+
+#if WITH_EDITOR
+	void Set(const ParameterType& Value, bool NonDefaultOnly = false) const
+	{
+		this->ParameterInfo.Set(this->Material, Value, NonDefaultOnly);
+	}
+#endif
+};
+
+template <typename MaterialType, typename = typename TEnableIf<TIsDerivedFrom<MaterialType, UMaterialInstance>::Value>::Type>
+class TGLTFProxyMaterialTextureParameter
+{
+public:
+
+	TGLTFProxyMaterialTextureParameter(const FGLTFProxyMaterialTextureParameterInfo& ParameterInfo, MaterialType* Material)
+		: Texture(ParameterInfo.Texture, Material)
+		, UVIndex(ParameterInfo.UVIndex, Material)
+		, UVOffset(ParameterInfo.UVOffset, Material)
+		, UVScale(ParameterInfo.UVScale, Material)
+		, UVRotation(ParameterInfo.UVRotation, Material)
+	{
+	}
+
+	TGLTFProxyMaterialParameter<UTexture*, MaterialType> Texture;
+	TGLTFProxyMaterialParameter<float, MaterialType> UVIndex;
+	TGLTFProxyMaterialParameter<FLinearColor, MaterialType> UVOffset;
+	TGLTFProxyMaterialParameter<FLinearColor, MaterialType> UVScale;
+	TGLTFProxyMaterialParameter<float, MaterialType> UVRotation;
+};
