[GLTFExporter] refactor: move json member functions into separate cpp files

commit f1f24e77d1c05ccd594a6cfda63f4bd9859dd2e9
Date: Thu, 10 Mar 2022 01:44:14 +0100

    refactor: move json member functions into separate cpp files
    (cherry picked from commit 04016742d6334105d38f7e88ff26a8a00421d5da​)

#skipci

diff --git a/Engine/Plugins/Enterprise/GLTFExporter/MigratedFromGit/Source/GLTFExporter/Private/Json/GLTFJsonAccessor.cpp b/Engine/Plugins/Enterprise/GLTFExporter/MigratedFromGit/Source/GLTFExporter/Private/Json/GLTFJsonAccessor.cpp
new file mode 100644
--- /dev/null
+++ b/Engine/Plugins/Enterprise/GLTFExporter/MigratedFromGit/Source/GLTFExporter/Private/Json/GLTFJsonAccessor.cpp
@@ -0,0 +1,33 @@
+// Copyright Epic Games, Inc. All Rights Reserved.
+
+#include "Json/GLTFJsonAccessor.h"
+
+void FGLTFJsonAccessor::WriteObject(IGLTFJsonWriter& Writer) const
+{
+	if (!Name.IsEmpty())
+	{
+		Writer.Write(TEXT("name"), Name);
+	}
+
+	Writer.Write(TEXT("bufferView"), BufferView);
+
+	if (ByteOffset != 0)
+	{
+		Writer.Write(TEXT("byteOffset"), ByteOffset);
+	}
+
+	Writer.Write(TEXT("count"), Count);
+	Writer.Write(TEXT("type"), Type);
+	Writer.Write(TEXT("componentType"), ComponentType);
+
+	if (bNormalized)
+	{
+		Writer.Write(TEXT("normalized"), bNormalized);
+	}
+
+	if (MinMaxLength > 0)
+	{
+		Writer.Write(TEXT("min"), Min, MinMaxLength);
+		Writer.Write(TEXT("max"), Max, MinMaxLength);
+	}
+}
diff --git a/Engine/Plugins/Enterprise/GLTFExporter/MigratedFromGit/Source/GLTFExporter/Private/Json/GLTFJsonAnimation.cpp b/Engine/Plugins/Enterprise/GLTFExporter/MigratedFromGit/Source/GLTFExporter/Private/Json/GLTFJsonAnimation.cpp
new file mode 100644
--- /dev/null
+++ b/Engine/Plugins/Enterprise/GLTFExporter/MigratedFromGit/Source/GLTFExporter/Private/Json/GLTFJsonAnimation.cpp
@@ -0,0 +1,45 @@
+// Copyright Epic Games, Inc. All Rights Reserved.
+
+#include "Json/GLTFJsonAnimation.h"
+
+void FGLTFJsonAnimationChannelTarget::WriteObject(IGLTFJsonWriter& Writer) const
+{
+	Writer.Write(TEXT("node"), Node);
+	Writer.Write(TEXT("path"), Path);
+}
+
+void FGLTFJsonAnimationChannel::WriteObject(IGLTFJsonWriter& Writer) const
+{
+	Writer.Write(TEXT("sampler"), Sampler);
+
+	Writer.Write(TEXT("target"), Target);
+}
+
+void FGLTFJsonAnimationSampler::WriteObject(IGLTFJsonWriter& Writer) const
+{
+	Writer.Write(TEXT("input"), Input);
+	Writer.Write(TEXT("output"), Output);
+
+	if (Interpolation != EGLTFJsonInterpolation::Linear)
+	{
+		Writer.Write(TEXT("interpolation"), Interpolation);
+	}
+}
+
+void FGLTFJsonAnimation::WriteObject(IGLTFJsonWriter& Writer) const
+{
+	if (!Name.IsEmpty())
+	{
+		Writer.Write(TEXT("name"), Name);
+	}
+
+	Writer.Write(TEXT("channels"), Channels);
+	Writer.Write(TEXT("samplers"), Samplers);
+
+	if (Playback != FGLTFJsonAnimationPlayback())
+	{
+		Writer.StartExtensions();
+		Writer.Write(EGLTFJsonExtension::EPIC_AnimationPlayback, Playback);
+		Writer.EndExtensions();
+	}
+}
diff --git a/Engine/Plugins/Enterprise/GLTFExporter/MigratedFromGit/Source/GLTFExporter/Private/Json/GLTFJsonAnimationPlayback.cpp b/Engine/Plugins/Enterprise/GLTFExporter/MigratedFromGit/Source/GLTFExporter/Private/Json/GLTFJsonAnimationPlayback.cpp
new file mode 100644
--- /dev/null
+++ b/Engine/Plugins/Enterprise/GLTFExporter/MigratedFromGit/Source/GLTFExporter/Private/Json/GLTFJsonAnimationPlayback.cpp
@@ -0,0 +1,44 @@
+// Copyright Epic Games, Inc. All Rights Reserved.
+
+#include "Json/GLTFJsonAnimationPlayback.h"
+
+void FGLTFJsonAnimationPlayback::WriteObject(IGLTFJsonWriter& Writer) const
+{
+	if (!Name.IsEmpty())
+	{
+		Writer.Write(TEXT("name"), Name);
+	}
+
+	if (bLoop != true)
+	{
+		Writer.Write(TEXT("loop"), bLoop);
+	}
+
+	if (bAutoPlay != true)
+	{
+		Writer.Write(TEXT("autoPlay"), bAutoPlay);
+	}
+
+	if (!FMath::IsNearlyEqual(PlayRate, 1, Writer.DefaultTolerance))
+	{
+		Writer.Write(TEXT("playRate"), PlayRate);
+	}
+
+	if (!FMath::IsNearlyEqual(StartTime, 0, Writer.DefaultTolerance))
+	{
+		Writer.Write(TEXT("startTime"), StartTime);
+	}
+}
+
+bool FGLTFJsonAnimationPlayback::operator==(const FGLTFJsonAnimationPlayback& Other) const
+{
+	return bLoop == Other.bLoop
+		&& bAutoPlay == Other.bAutoPlay
+		&& PlayRate == Other.PlayRate
+		&& StartTime == Other.StartTime;
+}
+
+bool FGLTFJsonAnimationPlayback::operator!=(const FGLTFJsonAnimationPlayback& Other) const
+{
+	return !(*this == Other);
+}
diff --git a/Engine/Plugins/Enterprise/GLTFExporter/MigratedFromGit/Source/GLTFExporter/Private/Json/GLTFJsonAsset.cpp b/Engine/Plugins/Enterprise/GLTFExporter/MigratedFromGit/Source/GLTFExporter/Private/Json/GLTFJsonAsset.cpp
new file mode 100644
--- /dev/null
+++ b/Engine/Plugins/Enterprise/GLTFExporter/MigratedFromGit/Source/GLTFExporter/Private/Json/GLTFJsonAsset.cpp
@@ -0,0 +1,18 @@
+// Copyright Epic Games, Inc. All Rights Reserved.
+
+#include "Json/GLTFJsonAsset.h"
+
+void FGLTFJsonAsset::WriteObject(IGLTFJsonWriter& Writer) const
+{
+	Writer.Write(TEXT("version"), Version);
+
+	if (!Generator.IsEmpty())
+	{
+		Writer.Write(TEXT("generator"), Generator);
+	}
+
+	if (!Copyright.IsEmpty())
+	{
+		Writer.Write(TEXT("copyright"), Copyright);
+	}
+}
diff --git a/Engine/Plugins/Enterprise/GLTFExporter/MigratedFromGit/Source/GLTFExporter/Private/Json/GLTFJsonBackdrop.cpp b/Engine/Plugins/Enterprise/GLTFExporter/MigratedFromGit/Source/GLTFExporter/Private/Json/GLTFJsonBackdrop.cpp
new file mode 100644
--- /dev/null
+++ b/Engine/Plugins/Enterprise/GLTFExporter/MigratedFromGit/Source/GLTFExporter/Private/Json/GLTFJsonBackdrop.cpp
@@ -0,0 +1,31 @@
+// Copyright Epic Games, Inc. All Rights Reserved.
+
+#include "Json/GLTFJsonBackdrop.h"
+
+void FGLTFJsonBackdrop::WriteObject(IGLTFJsonWriter& Writer) const
+{
+	if (!Name.IsEmpty())
+	{
+		Writer.Write(TEXT("name"), Name);
+	}
+
+	if (Mesh != INDEX_NONE)
+	{
+		Writer.Write(TEXT("mesh"), Mesh);
+	}
+
+	Writer.Write(TEXT("cubemap"), Cubemap);
+
+	Writer.Write(TEXT("intensity"), Intensity);
+	Writer.Write(TEXT("size"), Size);
+
+	if (!FMath::IsNearlyZero(Angle))
+	{
+		Writer.Write(TEXT("angle"), Angle);
+	}
+
+	Writer.Write(TEXT("projectionCenter"), ProjectionCenter);
+
+	Writer.Write(TEXT("lightingDistanceFactor"), LightingDistanceFactor);
+	Writer.Write(TEXT("useCameraProjection"), UseCameraProjection);
+}
diff --git a/Engine/Plugins/Enterprise/GLTFExporter/MigratedFromGit/Source/GLTFExporter/Private/Json/GLTFJsonBuffer.cpp b/Engine/Plugins/Enterprise/GLTFExporter/MigratedFromGit/Source/GLTFExporter/Private/Json/GLTFJsonBuffer.cpp
new file mode 100644
--- /dev/null
+++ b/Engine/Plugins/Enterprise/GLTFExporter/MigratedFromGit/Source/GLTFExporter/Private/Json/GLTFJsonBuffer.cpp
@@ -0,0 +1,18 @@
+// Copyright Epic Games, Inc. All Rights Reserved.
+
+#include "Json/GLTFJsonBuffer.h"
+
+void FGLTFJsonBuffer::WriteObject(IGLTFJsonWriter& Writer) const
+{
+	if (!Name.IsEmpty())
+	{
+		Writer.Write(TEXT("name"), Name);
+	}
+
+	if (!URI.IsEmpty())
+	{
+		Writer.Write(TEXT("uri"), URI);
+	}
+
+	Writer.Write(TEXT("byteLength"), ByteLength);
+}
diff --git a/Engine/Plugins/Enterprise/GLTFExporter/MigratedFromGit/Source/GLTFExporter/Private/Json/GLTFJsonBufferView.cpp b/Engine/Plugins/Enterprise/GLTFExporter/MigratedFromGit/Source/GLTFExporter/Private/Json/GLTFJsonBufferView.cpp
new file mode 100644
--- /dev/null
+++ b/Engine/Plugins/Enterprise/GLTFExporter/MigratedFromGit/Source/GLTFExporter/Private/Json/GLTFJsonBufferView.cpp
@@ -0,0 +1,29 @@
+// Copyright Epic Games, Inc. All Rights Reserved.
+
+#include "Json/GLTFJsonBufferView.h"
+
+void FGLTFJsonBufferView::WriteObject(IGLTFJsonWriter& Writer) const
+{
+	if (!Name.IsEmpty())
+	{
+		Writer.Write(TEXT("name"), Name);
+	}
+
+	Writer.Write(TEXT("buffer"), Buffer);
+	Writer.Write(TEXT("byteLength"), ByteLength);
+
+	if (ByteOffset != 0)
+	{
+		Writer.Write(TEXT("byteOffset"), ByteOffset);
+	}
+
+	if (ByteStride != 0)
+	{
+		Writer.Write(TEXT("byteStride"), ByteStride);
+	}
+
+	if (Target != EGLTFJsonBufferTarget::None)
+	{
+		Writer.Write(TEXT("target"), Target);
+	}
+}
diff --git a/Engine/Plugins/Enterprise/GLTFExporter/MigratedFromGit/Source/GLTFExporter/Private/Json/GLTFJsonCamera.cpp b/Engine/Plugins/Enterprise/GLTFExporter/MigratedFromGit/Source/GLTFExporter/Private/Json/GLTFJsonCamera.cpp
new file mode 100644
--- /dev/null
+++ b/Engine/Plugins/Enterprise/GLTFExporter/MigratedFromGit/Source/GLTFExporter/Private/Json/GLTFJsonCamera.cpp
@@ -0,0 +1,59 @@
+// Copyright Epic Games, Inc. All Rights Reserved.
+
+#include "Json/GLTFJsonCamera.h"
+
+void FGLTFJsonOrthographic::WriteObject(IGLTFJsonWriter& Writer) const
+{
+	Writer.Write(TEXT("xmag"), XMag);
+	Writer.Write(TEXT("ymag"), YMag);
+	Writer.Write(TEXT("zfar"), ZFar);
+	Writer.Write(TEXT("znear"), ZNear);
+}
+
+void FGLTFJsonPerspective::WriteObject(IGLTFJsonWriter& Writer) const
+{
+	if (!FMath::IsNearlyEqual(AspectRatio, 0, Writer.DefaultTolerance))
+	{
+		Writer.Write(TEXT("aspectRatio"), AspectRatio);
+	}
+
+	Writer.Write(TEXT("yfov"), YFov);
+
+	if (!FMath::IsNearlyEqual(ZFar, 0, Writer.DefaultTolerance))
+	{
+		Writer.Write(TEXT("zfar"), ZFar);
+	}
+
+	Writer.Write(TEXT("znear"), ZNear);
+}
+
+void FGLTFJsonCamera::WriteObject(IGLTFJsonWriter& Writer) const
+{
+	if (!Name.IsEmpty())
+	{
+		Writer.Write(TEXT("name"), Name);
+	}
+
+	Writer.Write(TEXT("type"), Type);
+
+	switch (Type)
+	{
+		case EGLTFJsonCameraType::Orthographic:
+			Writer.Write(TEXT("orthographic"), Orthographic);
+			break;
+
+		case EGLTFJsonCameraType::Perspective:
+			Writer.Write(TEXT("perspective"), Perspective);
+			break;
+
+		default:
+			break;
+	}
+
+	if (CameraControl.IsSet())
+	{
+		Writer.StartExtensions();
+		Writer.Write(EGLTFJsonExtension::EPIC_CameraControls, CameraControl.GetValue());
+		Writer.EndExtensions();
+	}
+}
diff --git a/Engine/Plugins/Enterprise/GLTFExporter/MigratedFromGit/Source/GLTFExporter/Private/Json/GLTFJsonCameraControl.cpp b/Engine/Plugins/Enterprise/GLTFExporter/MigratedFromGit/Source/GLTFExporter/Private/Json/GLTFJsonCameraControl.cpp
new file mode 100644
--- /dev/null
+++ b/Engine/Plugins/Enterprise/GLTFExporter/MigratedFromGit/Source/GLTFExporter/Private/Json/GLTFJsonCameraControl.cpp
@@ -0,0 +1,41 @@
+// Copyright Epic Games, Inc. All Rights Reserved.
+
+#include "Json/GLTFJsonCameraControl.h"
+
+void FGLTFJsonCameraControl::WriteObject(IGLTFJsonWriter& Writer) const
+{
+	Writer.Write(TEXT("mode"), Mode);
+
+	if (Target != INDEX_NONE && Mode == EGLTFJsonCameraControlMode::Orbital)
+	{
+		Writer.Write(TEXT("target"), Target);
+	}
+
+	Writer.Write(TEXT("maxDistance"), MaxDistance);
+	Writer.Write(TEXT("minDistance"), MinDistance);
+
+	if (!FMath::IsNearlyEqual(MaxPitch, 90, Writer.DefaultTolerance))
+	{
+		Writer.Write(TEXT("maxPitch"), MaxPitch);
+	}
+
+	if (!FMath::IsNearlyEqual(MinPitch, -90, Writer.DefaultTolerance))
+	{
+		Writer.Write(TEXT("minPitch"), MinPitch);
+	}
+
+	if (!FMath::IsNearlyEqual(MaxYaw, 360, Writer.DefaultTolerance))
+	{
+		Writer.Write(TEXT("maxYaw"), MaxYaw);
+	}
+
+	if (!FMath::IsNearlyEqual(MinYaw, 0, Writer.DefaultTolerance))
+	{
+		Writer.Write(TEXT("minYaw"), MinYaw);
+	}
+
+	Writer.Write(TEXT("rotationSensitivity"), RotationSensitivity);
+	Writer.Write(TEXT("rotationInertia"), RotationInertia);
+	Writer.Write(TEXT("dollySensitivity"), DollySensitivity);
+	Writer.Write(TEXT("dollyDuration"), DollyDuration);
+}
diff --git a/Engine/Plugins/Enterprise/GLTFExporter/MigratedFromGit/Source/GLTFExporter/Private/Json/GLTFJsonEpicLevelVariantSets.cpp b/Engine/Plugins/Enterprise/GLTFExporter/MigratedFromGit/Source/GLTFExporter/Private/Json/GLTFJsonEpicLevelVariantSets.cpp
new file mode 100644
--- /dev/null
+++ b/Engine/Plugins/Enterprise/GLTFExporter/MigratedFromGit/Source/GLTFExporter/Private/Json/GLTFJsonEpicLevelVariantSets.cpp
@@ -0,0 +1,75 @@
+// Copyright Epic Games, Inc. All Rights Reserved.
+
+#include "Json/GLTFJsonEpicLevelVariantSets.h"
+
+void FGLTFJsonEpicVariantMaterial::WriteObject(IGLTFJsonWriter& Writer) const
+{
+	Writer.Write(TEXT("material"), Material);
+
+	if (Index != INDEX_NONE)
+	{
+		Writer.Write(TEXT("index"), Index);
+	}
+}
+
+void FGLTFJsonEpicVariantNodeProperties::WriteObject(IGLTFJsonWriter& Writer) const
+{
+	if (Node != INDEX_NONE)
+	{
+		Writer.Write(TEXT("node"), Node);
+	}
+
+	Writer.StartObject(TEXT("properties"));
+
+	if (bIsVisible.IsSet())
+	{
+		Writer.Write(TEXT("visible"), bIsVisible.GetValue());
+	}
+
+	if (Mesh.IsSet())
+	{
+		Writer.Write(TEXT("mesh"), Mesh.GetValue());
+	}
+
+	if (Materials.Num() > 0)
+	{
+		Writer.Write(TEXT("materials"), Materials);
+	}
+
+	Writer.EndObject();
+}
+
+void FGLTFJsonEpicVariant::WriteObject(IGLTFJsonWriter& Writer) const
+{
+	Writer.Write(TEXT("name"), Name);
+	Writer.Write(TEXT("active"), bIsActive);
+
+	if (Thumbnail != INDEX_NONE)
+	{
+		Writer.Write(TEXT("thumbnail"), Thumbnail);
+	}
+
+	TArray<FGLTFJsonEpicVariantNodeProperties> NodeValues;
+	Nodes.GenerateValueArray(NodeValues);
+	Writer.Write(TEXT("nodes"), NodeValues);
+}
+
+void FGLTFJsonEpicVariantSet::WriteObject(IGLTFJsonWriter& Writer) const
+{
+	if (!Name.IsEmpty())
+	{
+		Writer.Write(TEXT("name"), Name);
+	}
+
+	Writer.Write(TEXT("variants"), Variants);
+}
+
+void FGLTFJsonEpicLevelVariantSets::WriteObject(IGLTFJsonWriter& Writer) const
+{
+	if (!Name.IsEmpty())
+	{
+		Writer.Write(TEXT("name"), Name);
+	}
+
+	Writer.Write(TEXT("variantSets"), VariantSets);
+}
diff --git a/Engine/Plugins/Enterprise/GLTFExporter/MigratedFromGit/Source/GLTFExporter/Private/Json/GLTFJsonHotspot.cpp b/Engine/Plugins/Enterprise/GLTFExporter/MigratedFromGit/Source/GLTFExporter/Private/Json/GLTFJsonHotspot.cpp
new file mode 100644
--- /dev/null
+++ b/Engine/Plugins/Enterprise/GLTFExporter/MigratedFromGit/Source/GLTFExporter/Private/Json/GLTFJsonHotspot.cpp
@@ -0,0 +1,29 @@
+// Copyright Epic Games, Inc. All Rights Reserved.
+
+#include "Json/GLTFJsonHotspot.h"
+
+void FGLTFJsonHotspot::WriteObject(IGLTFJsonWriter& Writer) const
+{
+	if (!Name.IsEmpty())
+	{
+		Writer.Write(TEXT("name"), Name);
+	}
+
+	Writer.Write(TEXT("animation"), Animation);
+	Writer.Write(TEXT("image"), Image);
+
+	if (HoveredImage != INDEX_NONE)
+	{
+		Writer.Write(TEXT("hoveredImage"), HoveredImage);
+	}
+
+	if (ToggledImage != INDEX_NONE)
+	{
+		Writer.Write(TEXT("toggledImage"), ToggledImage);
+	}
+
+	if (ToggledHoveredImage != INDEX_NONE)
+	{
+		Writer.Write(TEXT("toggledHoveredImage"), ToggledHoveredImage);
+	}
+}
diff --git a/Engine/Plugins/Enterprise/GLTFExporter/MigratedFromGit/Source/GLTFExporter/Private/Json/GLTFJsonImage.cpp b/Engine/Plugins/Enterprise/GLTFExporter/MigratedFromGit/Source/GLTFExporter/Private/Json/GLTFJsonImage.cpp
new file mode 100644
--- /dev/null
+++ b/Engine/Plugins/Enterprise/GLTFExporter/MigratedFromGit/Source/GLTFExporter/Private/Json/GLTFJsonImage.cpp
@@ -0,0 +1,26 @@
+// Copyright Epic Games, Inc. All Rights Reserved.
+
+#include "Json/GLTFJsonImage.h"
+
+void FGLTFJsonImage::WriteObject(IGLTFJsonWriter& Writer) const
+{
+	if (!Name.IsEmpty())
+	{
+		Writer.Write(TEXT("name"), Name);
+	}
+
+	if (!Uri.IsEmpty())
+	{
+		Writer.Write(TEXT("uri"), Uri);
+	}
+
+	if (MimeType != EGLTFJsonMimeType::None)
+	{
+		Writer.Write(TEXT("mimeType"), MimeType);
+	}
+
+	if (BufferView != INDEX_NONE)
+	{
+		Writer.Write(TEXT("bufferView"), BufferView);
+	}
+}
diff --git a/Engine/Plugins/Enterprise/GLTFExporter/MigratedFromGit/Source/GLTFExporter/Private/Json/GLTFJsonKhrMaterialVariant.cpp b/Engine/Plugins/Enterprise/GLTFExporter/MigratedFromGit/Source/GLTFExporter/Private/Json/GLTFJsonKhrMaterialVariant.cpp
new file mode 100644
--- /dev/null
+++ b/Engine/Plugins/Enterprise/GLTFExporter/MigratedFromGit/Source/GLTFExporter/Private/Json/GLTFJsonKhrMaterialVariant.cpp
@@ -0,0 +1,14 @@
+// Copyright Epic Games, Inc. All Rights Reserved.
+
+#include "Json/GLTFJsonKhrMaterialVariant.h"
+
+void FGLTFJsonKhrMaterialVariantMapping::WriteObject(IGLTFJsonWriter& Writer) const
+{
+	Writer.Write(TEXT("material"), Material);
+	Writer.Write(TEXT("variants"), Variants);
+}
+
+void FGLTFJsonKhrMaterialVariant::WriteObject(IGLTFJsonWriter& Writer) const
+{
+	Writer.Write(TEXT("name"), Name);
+}
diff --git a/Engine/Plugins/Enterprise/GLTFExporter/MigratedFromGit/Source/GLTFExporter/Private/Json/GLTFJsonLight.cpp b/Engine/Plugins/Enterprise/GLTFExporter/MigratedFromGit/Source/GLTFExporter/Private/Json/GLTFJsonLight.cpp
new file mode 100644
--- /dev/null
+++ b/Engine/Plugins/Enterprise/GLTFExporter/MigratedFromGit/Source/GLTFExporter/Private/Json/GLTFJsonLight.cpp
@@ -0,0 +1,49 @@
+// Copyright Epic Games, Inc. All Rights Reserved.
+
+#include "Json/GLTFJsonLight.h"
+
+void FGLTFJsonSpotLight::WriteObject(IGLTFJsonWriter& Writer) const
+{
+	if (!FMath::IsNearlyEqual(InnerConeAngle, 0, Writer.DefaultTolerance))
+	{
+		Writer.Write(TEXT("innerConeAngle"), InnerConeAngle);
+	}
+
+	if (!FMath::IsNearlyEqual(OuterConeAngle, HALF_PI, Writer.DefaultTolerance))
+	{
+		Writer.Write(TEXT("outerConeAngle"), OuterConeAngle);
+	}
+}
+
+void FGLTFJsonLight::WriteObject(IGLTFJsonWriter& Writer) const
+{
+	if (!Name.IsEmpty())
+	{
+		Writer.Write(TEXT("name"), Name);
+	}
+
+	Writer.Write(TEXT("type"), Type);
+
+	if (!Color.IsNearlyEqual(FGLTFJsonColor3::White, Writer.DefaultTolerance))
+	{
+		Writer.Write(TEXT("color"), Color);
+	}
+
+	if (!FMath::IsNearlyEqual(Intensity, 1, Writer.DefaultTolerance))
+	{
+		Writer.Write(TEXT("intensity"), Intensity);
+	}
+
+	if (Type == EGLTFJsonLightType::Point || Type == EGLTFJsonLightType::Spot)
+	{
+		if (!FMath::IsNearlyEqual(Range, 0, Writer.DefaultTolerance))
+		{
+			Writer.Write(TEXT("range"), Range);
+		}
+
+		if (Type == EGLTFJsonLightType::Spot)
+		{
+			Writer.Write(TEXT("spot"), Spot);
+		}
+	}
+}
diff --git a/Engine/Plugins/Enterprise/GLTFExporter/MigratedFromGit/Source/GLTFExporter/Private/Json/GLTFJsonLightMap.cpp b/Engine/Plugins/Enterprise/GLTFExporter/MigratedFromGit/Source/GLTFExporter/Private/Json/GLTFJsonLightMap.cpp
new file mode 100644
--- /dev/null
+++ b/Engine/Plugins/Enterprise/GLTFExporter/MigratedFromGit/Source/GLTFExporter/Private/Json/GLTFJsonLightMap.cpp
@@ -0,0 +1,20 @@
+// Copyright Epic Games, Inc. All Rights Reserved.
+
+#include "Json/GLTFJsonLightMap.h"
+
+void FGLTFJsonLightMap::WriteObject(IGLTFJsonWriter& Writer) const
+{
+	if (!Name.IsEmpty())
+	{
+		Writer.Write(TEXT("name"), Name);
+	}
+
+	if (Texture.Index != INDEX_NONE)
+	{
+		Writer.Write(TEXT("texture"), Texture);
+	}
+
+	Writer.Write(TEXT("lightmapScale"), LightMapScale);
+	Writer.Write(TEXT("lightmapAdd"), LightMapAdd);
+	Writer.Write(TEXT("coordinateScaleBias"), CoordinateScaleBias);
+}
diff --git a/Engine/Plugins/Enterprise/GLTFExporter/MigratedFromGit/Source/GLTFExporter/Private/Json/GLTFJsonMaterial.cpp b/Engine/Plugins/Enterprise/GLTFExporter/MigratedFromGit/Source/GLTFExporter/Private/Json/GLTFJsonMaterial.cpp
new file mode 100644
--- /dev/null
+++ b/Engine/Plugins/Enterprise/GLTFExporter/MigratedFromGit/Source/GLTFExporter/Private/Json/GLTFJsonMaterial.cpp
@@ -0,0 +1,179 @@
+// Copyright Epic Games, Inc. All Rights Reserved.
+
+#include "Json/GLTFJsonMaterial.h"
+
+void FGLTFJsonTextureInfo::WriteObject(IGLTFJsonWriter& Writer) const
+{
+	Writer.Write(TEXT("index"), Index);
+
+	if (TexCoord != 0)
+	{
+		Writer.Write(TEXT("texCoord"), TexCoord);
+	}
+
+	if (!Transform.IsNearlyDefault(Writer.DefaultTolerance))
+	{
+		Writer.StartExtensions();
+		Writer.Write(EGLTFJsonExtension::KHR_TextureTransform, Transform);
+		Writer.EndExtensions();
+	}
+}
+
+void FGLTFJsonNormalTextureInfo::WriteObject(IGLTFJsonWriter& Writer) const
+{
+	Writer.Write(TEXT("index"), Index);
+
+	if (TexCoord != 0)
+	{
+		Writer.Write(TEXT("texCoord"), TexCoord);
+	}
+
+	if (!FMath::IsNearlyEqual(Scale, 1, Writer.DefaultTolerance))
+	{
+		Writer.Write(TEXT("scale"), Scale);
+	}
+}
+
+void FGLTFJsonOcclusionTextureInfo::WriteObject(IGLTFJsonWriter& Writer) const
+{
+	Writer.Write(TEXT("index"), Index);
+
+	if (TexCoord != 0)
+	{
+		Writer.Write(TEXT("texCoord"), TexCoord);
+	}
+
+	if (!FMath::IsNearlyEqual(Strength, 1, Writer.DefaultTolerance))
+	{
+		Writer.Write(TEXT("strength"), Strength);
+	}
+}
+
+void FGLTFJsonPBRMetallicRoughness::WriteObject(IGLTFJsonWriter& Writer) const
+{
+	if (!BaseColorFactor.IsNearlyEqual(FGLTFJsonColor4::White, Writer.DefaultTolerance))
+	{
+		Writer.Write(TEXT("baseColorFactor"), BaseColorFactor);
+	}
+
+	if (BaseColorTexture.Index != INDEX_NONE)
+	{
+		Writer.Write(TEXT("baseColorTexture"), BaseColorTexture);
+	}
+
+	if (!FMath::IsNearlyEqual(MetallicFactor, 1, Writer.DefaultTolerance))
+	{
+		Writer.Write(TEXT("metallicFactor"), MetallicFactor);
+	}
+
+	if (!FMath::IsNearlyEqual(RoughnessFactor, 1, Writer.DefaultTolerance))
+	{
+		Writer.Write(TEXT("roughnessFactor"), RoughnessFactor);
+	}
+
+	if (MetallicRoughnessTexture.Index != INDEX_NONE)
+	{
+		Writer.Write(TEXT("metallicRoughnessTexture"), MetallicRoughnessTexture);
+	}
+}
+
+void FGLTFJsonClearCoatExtension::WriteObject(IGLTFJsonWriter& Writer) const
+{
+	if (!FMath::IsNearlyEqual(ClearCoatFactor, 0, Writer.DefaultTolerance))
+	{
+		Writer.Write(TEXT("clearcoatFactor"), ClearCoatFactor);
+	}
+
+	if (ClearCoatTexture.Index != INDEX_NONE)
+	{
+		Writer.Write(TEXT("clearcoatTexture"), ClearCoatTexture);
+	}
+
+	if (!FMath::IsNearlyEqual(ClearCoatRoughnessFactor, 0, Writer.DefaultTolerance))
+	{
+		Writer.Write(TEXT("clearcoatRoughnessFactor"), ClearCoatRoughnessFactor);
+	}
+
+	if (ClearCoatRoughnessTexture.Index != INDEX_NONE)
+	{
+		Writer.Write(TEXT("clearcoatRoughnessTexture"), ClearCoatRoughnessTexture);
+	}
+
+	if (ClearCoatNormalTexture.Index != INDEX_NONE)
+	{
+		Writer.Write(TEXT("clearcoatNormalTexture"), ClearCoatNormalTexture);
+	}
+}
+
+void FGLTFJsonMaterial::WriteObject(IGLTFJsonWriter& Writer) const
+{
+	if (!Name.IsEmpty())
+	{
+		Writer.Write(TEXT("name"), Name);
+	}
+
+	if (ShadingModel != EGLTFJsonShadingModel::None)
+	{
+		Writer.Write(TEXT("pbrMetallicRoughness"), PBRMetallicRoughness);
+	}
+
+	if (NormalTexture.Index != INDEX_NONE)
+	{
+		Writer.Write(TEXT("normalTexture"), NormalTexture);
+	}
+
+	if (OcclusionTexture.Index != INDEX_NONE)
+	{
+		Writer.Write(TEXT("occlusionTexture"), OcclusionTexture);
+	}
+
+	if (EmissiveTexture.Index != INDEX_NONE)
+	{
+		Writer.Write(TEXT("emissiveTexture"), EmissiveTexture);
+	}
+
+	if (!EmissiveFactor.IsNearlyEqual(FGLTFJsonColor3::Black, Writer.DefaultTolerance))
+	{
+		Writer.Write(TEXT("emissiveFactor"), EmissiveFactor);
+	}
+
+	if (AlphaMode != EGLTFJsonAlphaMode::Opaque)
+	{
+		Writer.Write(TEXT("alphaMode"), AlphaMode);
+
+		if (AlphaMode == EGLTFJsonAlphaMode::Mask && !FMath::IsNearlyEqual(AlphaCutoff, 0.5f, Writer.DefaultTolerance))
+		{
+			Writer.Write(TEXT("alphaCutoff"), AlphaCutoff);
+		}
+	}
+
+	if (DoubleSided)
+	{
+		Writer.Write(TEXT("doubleSided"), DoubleSided);
+	}
+
+	if (BlendMode != EGLTFJsonBlendMode::None || ShadingModel == EGLTFJsonShadingModel::Unlit || ShadingModel == EGLTFJsonShadingModel::ClearCoat)
+	{
+		Writer.StartExtensions();
+
+		if (BlendMode != EGLTFJsonBlendMode::None)
+		{
+			Writer.StartExtension(EGLTFJsonExtension::EPIC_BlendModes);
+			Writer.Write(TEXT("blendMode"), BlendMode);
+			Writer.EndExtension();
+		}
+
+		if (ShadingModel == EGLTFJsonShadingModel::Unlit)
+		{
+			Writer.StartExtension(EGLTFJsonExtension::KHR_MaterialsUnlit);
+			// Write empty object
+			Writer.EndExtension();
+		}
+		else if (ShadingModel == EGLTFJsonShadingModel::ClearCoat)
+		{
+			Writer.Write(EGLTFJsonExtension::KHR_MaterialsClearCoat, ClearCoat);
+		}
+
+		Writer.EndExtensions();
+	}
+}
diff --git a/Engine/Plugins/Enterprise/GLTFExporter/MigratedFromGit/Source/GLTFExporter/Private/Json/GLTFJsonMesh.cpp b/Engine/Plugins/Enterprise/GLTFExporter/MigratedFromGit/Source/GLTFExporter/Private/Json/GLTFJsonMesh.cpp
new file mode 100644
--- /dev/null
+++ b/Engine/Plugins/Enterprise/GLTFExporter/MigratedFromGit/Source/GLTFExporter/Private/Json/GLTFJsonMesh.cpp
@@ -0,0 +1,70 @@
+// Copyright Epic Games, Inc. All Rights Reserved.
+
+#include "Json/GLTFJsonMesh.h"
+
+void FGLTFJsonAttributes::WriteObject(IGLTFJsonWriter& Writer) const
+{
+	if (Position != INDEX_NONE) Writer.Write(TEXT("POSITION"), Position);
+	if (Color0 != INDEX_NONE) Writer.Write(TEXT("COLOR_0"), Color0);
+	if (Normal != INDEX_NONE) Writer.Write(TEXT("NORMAL"), Normal);
+	if (Tangent != INDEX_NONE) Writer.Write(TEXT("TANGENT"), Tangent);
+
+	for (int32 Index = 0; Index < TexCoords.Num(); ++Index)
+	{
+		const FGLTFJsonAccessorIndex TexCoord = TexCoords[Index];
+		if (TexCoord != INDEX_NONE) Writer.Write(TEXT("TEXCOORD_") + FString::FromInt(Index), TexCoord);
+	}
+
+	for (int32 Index = 0; Index < Joints.Num(); ++Index)
+	{
+		const FGLTFJsonAccessorIndex Joint = Joints[Index];
+		if (Joint != INDEX_NONE) Writer.Write(TEXT("JOINTS_") + FString::FromInt(Index), Joint);
+	}
+
+	for (int32 Index = 0; Index < Weights.Num(); ++Index)
+	{
+		const FGLTFJsonAccessorIndex Weight = Weights[Index];
+		if (Weight != INDEX_NONE) Writer.Write(TEXT("WEIGHTS_") + FString::FromInt(Index), Weight);
+	}
+}
+
+void FGLTFJsonPrimitive::WriteObject(IGLTFJsonWriter& Writer) const
+{
+	Writer.Write(TEXT("attributes"), Attributes);
+
+	if (Indices != INDEX_NONE)
+	{
+		Writer.Write(TEXT("indices"), Indices);
+	}
+
+	if (Material != INDEX_NONE)
+	{
+		Writer.Write(TEXT("material"), Material);
+	}
+
+	if (Mode != EGLTFJsonPrimitiveMode::Triangles)
+	{
+		Writer.Write(TEXT("mode"), Mode);
+	}
+
+	if (KhrMaterialVariantMappings.Num() > 0)
+	{
+		Writer.StartExtensions();
+
+		Writer.StartExtension(EGLTFJsonExtension::KHR_MaterialsVariants);
+		Writer.Write(TEXT("mappings"), KhrMaterialVariantMappings);
+		Writer.EndExtension();
+
+		Writer.EndExtensions();
+	}
+}
+
+void FGLTFJsonMesh::WriteObject(IGLTFJsonWriter& Writer) const
+{
+	if (!Name.IsEmpty())
+	{
+		Writer.Write(TEXT("name"), Name);
+	}
+
+	Writer.Write(TEXT("primitives"), Primitives);
+}
diff --git a/Engine/Plugins/Enterprise/GLTFExporter/MigratedFromGit/Source/GLTFExporter/Private/Json/GLTFJsonNode.cpp b/Engine/Plugins/Enterprise/GLTFExporter/MigratedFromGit/Source/GLTFExporter/Private/Json/GLTFJsonNode.cpp
new file mode 100644
--- /dev/null
+++ b/Engine/Plugins/Enterprise/GLTFExporter/MigratedFromGit/Source/GLTFExporter/Private/Json/GLTFJsonNode.cpp
@@ -0,0 +1,88 @@
+// Copyright Epic Games, Inc. All Rights Reserved.
+
+#include "Json/GLTFJsonNode.h"
+
+void FGLTFJsonNode::WriteObject(IGLTFJsonWriter& Writer) const
+{
+	if (!Name.IsEmpty())
+	{
+		Writer.Write(TEXT("name"), Name);
+	}
+
+	if (!Translation.IsNearlyEqual(FGLTFJsonVector3::Zero, Writer.DefaultTolerance))
+	{
+		Writer.Write(TEXT("translation"), Translation);
+	}
+
+	if (!Rotation.IsNearlyEqual(FGLTFJsonQuaternion::Identity, Writer.DefaultTolerance))
+	{
+		Writer.Write(TEXT("rotation"), Rotation);
+	}
+
+	if (!Scale.IsNearlyEqual(FGLTFJsonVector3::One, Writer.DefaultTolerance))
+	{
+		Writer.Write(TEXT("scale"), Scale);
+	}
+
+	if (Camera != INDEX_NONE)
+	{
+		Writer.Write(TEXT("camera"), Camera);
+	}
+
+	if (Skin != INDEX_NONE)
+	{
+		Writer.Write(TEXT("skin"), Skin);
+	}
+
+	if (Mesh != INDEX_NONE)
+	{
+		Writer.Write(TEXT("mesh"), Mesh);
+	}
+
+	if (Backdrop != INDEX_NONE || Hotspot != INDEX_NONE || Light != INDEX_NONE || LightMap != INDEX_NONE || SkySphere != INDEX_NONE)
+	{
+		Writer.StartExtensions();
+
+		if (Backdrop != INDEX_NONE)
+		{
+			Writer.StartExtension(EGLTFJsonExtension::EPIC_HDRIBackdrops);
+			Writer.Write(TEXT("backdrop"), Backdrop);
+			Writer.EndExtension();
+		}
+
+		if (Hotspot != INDEX_NONE)
+		{
+			Writer.StartExtension(EGLTFJsonExtension::EPIC_AnimationHotspots);
+			Writer.Write(TEXT("hotspot"), Hotspot);
+			Writer.EndExtension();
+		}
+
+		if (Light != INDEX_NONE)
+		{
+			Writer.StartExtension(EGLTFJsonExtension::KHR_LightsPunctual);
+			Writer.Write(TEXT("light"), Light);
+			Writer.EndExtension();
+		}
+
+		if (LightMap != INDEX_NONE)
+		{
+			Writer.StartExtension(EGLTFJsonExtension::EPIC_LightmapTextures);
+			Writer.Write(TEXT("lightmap"), LightMap);
+			Writer.EndExtension();
+		}
+
+		if (SkySphere != INDEX_NONE)
+		{
+			Writer.StartExtension(EGLTFJsonExtension::EPIC_SkySpheres);
+			Writer.Write(TEXT("skySphere"), SkySphere);
+			Writer.EndExtension();
+		}
+
+		Writer.EndExtensions();
+	}
+
+	if (Children.Num() > 0)
+	{
+		Writer.Write(TEXT("children"), Children);
+	}
+}
diff --git a/Engine/Plugins/Enterprise/GLTFExporter/MigratedFromGit/Source/GLTFExporter/Private/Json/GLTFJsonRoot.cpp b/Engine/Plugins/Enterprise/GLTFExporter/MigratedFromGit/Source/GLTFExporter/Private/Json/GLTFJsonRoot.cpp
new file mode 100644
--- /dev/null
+++ b/Engine/Plugins/Enterprise/GLTFExporter/MigratedFromGit/Source/GLTFExporter/Private/Json/GLTFJsonRoot.cpp
@@ -0,0 +1,101 @@
+// Copyright Epic Games, Inc. All Rights Reserved.
+
+#include "Json/GLTFJsonRoot.h"
+
+void FGLTFJsonRoot::WriteObject(IGLTFJsonWriter& Writer) const
+{
+	Writer.Write(TEXT("asset"), Asset);
+
+	if (DefaultScene != INDEX_NONE)
+	{
+		Writer.Write(TEXT("scene"), DefaultScene);
+	}
+
+	if (Accessors.Num() > 0) Writer.Write(TEXT("accessors"), Accessors);
+	if (Animations.Num() > 0) Writer.Write(TEXT("animations"), Animations);
+	if (Buffers.Num() > 0) Writer.Write(TEXT("buffers"), Buffers);
+	if (BufferViews.Num() > 0) Writer.Write(TEXT("bufferViews"), BufferViews);
+	if (Cameras.Num() > 0) Writer.Write(TEXT("cameras"), Cameras);
+	if (Images.Num() > 0) Writer.Write(TEXT("images"), Images);
+	if (Materials.Num() > 0) Writer.Write(TEXT("materials"), Materials);
+	if (Meshes.Num() > 0) Writer.Write(TEXT("meshes"), Meshes);
+	if (Nodes.Num() > 0) Writer.Write(TEXT("nodes"), Nodes);
+	if (Samplers.Num() > 0) Writer.Write(TEXT("samplers"), Samplers);
+	if (Scenes.Num() > 0) Writer.Write(TEXT("scenes"), Scenes);
+	if (Skins.Num() > 0) Writer.Write(TEXT("skins"), Skins);
+	if (Textures.Num() > 0) Writer.Write(TEXT("textures"), Textures);
+
+	if (Backdrops.Num() > 0 || Hotspots.Num() > 0 || Lights.Num() > 0 || LightMaps.Num() > 0 || SkySpheres.Num() > 0 || EpicLevelVariantSets.Num() > 0 || KhrMaterialVariants.Num() > 0)
+	{
+		Writer.StartExtensions();
+
+		if (Backdrops.Num() > 0)
+		{
+			Writer.StartExtension(EGLTFJsonExtension::EPIC_HDRIBackdrops);
+			Writer.Write(TEXT("backdrops"), Backdrops);
+			Writer.EndExtension();
+		}
+
+		if (Hotspots.Num() > 0)
+		{
+			Writer.StartExtension(EGLTFJsonExtension::EPIC_AnimationHotspots);
+			Writer.Write(TEXT("hotspots"), Hotspots);
+			Writer.EndExtension();
+		}
+
+		if (EpicLevelVariantSets.Num() > 0)
+		{
+			Writer.StartExtension(EGLTFJsonExtension::EPIC_LevelVariantSets);
+			Writer.Write(TEXT("levelVariantSets"), EpicLevelVariantSets);
+			Writer.EndExtension();
+		}
+
+		if (KhrMaterialVariants.Num() > 0)
+		{
+			Writer.StartExtension(EGLTFJsonExtension::KHR_MaterialsVariants);
+			Writer.Write(TEXT("variants"), KhrMaterialVariants);
+			Writer.EndExtension();
+		}
+
+		if (Lights.Num() > 0)
+		{
+			Writer.StartExtension(EGLTFJsonExtension::KHR_LightsPunctual);
+			Writer.Write(TEXT("lights"), Lights);
+			Writer.EndExtension();
+		}
+
+		if (LightMaps.Num() > 0)
+		{
+			Writer.StartExtension(EGLTFJsonExtension::EPIC_LightmapTextures);
+			Writer.Write(TEXT("lightmaps"), LightMaps);
+			Writer.EndExtension();
+		}
+
+		if (SkySpheres.Num() > 0)
+		{
+			Writer.StartExtension(EGLTFJsonExtension::EPIC_SkySpheres);
+			Writer.Write(TEXT("skySpheres"), SkySpheres);
+			Writer.EndExtension();
+		}
+
+		Writer.EndExtensions();
+	}
+
+	if (Extensions.Used.Num() > 0)
+	{
+		Writer.Write(TEXT("extensionsUsed"), Extensions.Used);
+	}
+
+	if (Extensions.Required.Num() > 0)
+	{
+		Writer.Write(TEXT("extensionsRequired"), Extensions.Required);
+	}
+}
+
+void FGLTFJsonRoot::WriteJson(FArchive& Archive, bool bPrettyJson, float DefaultTolerance)
+{
+	const TSharedRef<IGLTFJsonWriter> Writer = IGLTFJsonWriter::Create(Archive, bPrettyJson, Extensions);
+	Writer->DefaultTolerance = DefaultTolerance;
+	Writer->Write(*this);
+	Writer->Close();
+}
diff --git a/Engine/Plugins/Enterprise/GLTFExporter/MigratedFromGit/Source/GLTFExporter/Private/Json/GLTFJsonSampler.cpp b/Engine/Plugins/Enterprise/GLTFExporter/MigratedFromGit/Source/GLTFExporter/Private/Json/GLTFJsonSampler.cpp
new file mode 100644
--- /dev/null
+++ b/Engine/Plugins/Enterprise/GLTFExporter/MigratedFromGit/Source/GLTFExporter/Private/Json/GLTFJsonSampler.cpp
@@ -0,0 +1,31 @@
+// Copyright Epic Games, Inc. All Rights Reserved.
+
+#include "Json/GLTFJsonSampler.h"
+
+void FGLTFJsonSampler::WriteObject(IGLTFJsonWriter& Writer) const
+{
+	if (!Name.IsEmpty())
+	{
+		Writer.Write(TEXT("name"), Name);
+	}
+
+	if (MinFilter != EGLTFJsonTextureFilter::None)
+	{
+		Writer.Write(TEXT("minFilter"), MinFilter);
+	}
+
+	if (MagFilter != EGLTFJsonTextureFilter::None)
+	{
+		Writer.Write(TEXT("magFilter"), MagFilter);
+	}
+
+	if (WrapS != EGLTFJsonTextureWrap::Repeat)
+	{
+		Writer.Write(TEXT("wrapS"), WrapS);
+	}
+
+	if (WrapT != EGLTFJsonTextureWrap::Repeat)
+	{
+		Writer.Write(TEXT("wrapT"), WrapT);
+	}
+}
diff --git a/Engine/Plugins/Enterprise/GLTFExporter/MigratedFromGit/Source/GLTFExporter/Private/Json/GLTFJsonScene.cpp b/Engine/Plugins/Enterprise/GLTFExporter/MigratedFromGit/Source/GLTFExporter/Private/Json/GLTFJsonScene.cpp
new file mode 100644
--- /dev/null
+++ b/Engine/Plugins/Enterprise/GLTFExporter/MigratedFromGit/Source/GLTFExporter/Private/Json/GLTFJsonScene.cpp
@@ -0,0 +1,27 @@
+// Copyright Epic Games, Inc. All Rights Reserved.
+
+#include "Json/GLTFJsonScene.h"
+
+void FGLTFJsonScene::WriteObject(IGLTFJsonWriter& Writer) const
+{
+	if (!Name.IsEmpty())
+	{
+		Writer.Write(TEXT("name"), Name);
+	}
+
+	if (Nodes.Num() > 0)
+	{
+		Writer.Write(TEXT("nodes"), Nodes);
+	}
+
+	if (EpicLevelVariantSets.Num() > 0)
+	{
+		Writer.StartExtensions();
+
+		Writer.StartExtension(EGLTFJsonExtension::EPIC_LevelVariantSets);
+		Writer.Write(TEXT("levelVariantSets"), EpicLevelVariantSets);
+		Writer.EndExtension();
+
+		Writer.EndExtensions();
+	}
+}
diff --git a/Engine/Plugins/Enterprise/GLTFExporter/MigratedFromGit/Source/GLTFExporter/Private/Json/GLTFJsonSkin.cpp b/Engine/Plugins/Enterprise/GLTFExporter/MigratedFromGit/Source/GLTFExporter/Private/Json/GLTFJsonSkin.cpp
new file mode 100644
--- /dev/null
+++ b/Engine/Plugins/Enterprise/GLTFExporter/MigratedFromGit/Source/GLTFExporter/Private/Json/GLTFJsonSkin.cpp
@@ -0,0 +1,26 @@
+// Copyright Epic Games, Inc. All Rights Reserved.
+
+#include "Json/GLTFJsonSkin.h"
+
+void FGLTFJsonSkin::WriteObject(IGLTFJsonWriter& Writer) const
+{
+	if (!Name.IsEmpty())
+	{
+		Writer.Write(TEXT("name"), Name);
+	}
+
+	if (InverseBindMatrices != INDEX_NONE)
+	{
+		Writer.Write(TEXT("inverseBindMatrices"), InverseBindMatrices);
+	}
+
+	if (Skeleton != INDEX_NONE)
+	{
+		Writer.Write(TEXT("skeleton"), Skeleton);
+	}
+
+	if (Joints.Num() > 0)
+	{
+		Writer.Write(TEXT("joints"), Joints);
+	}
+}
diff --git a/Engine/Plugins/Enterprise/GLTFExporter/MigratedFromGit/Source/GLTFExporter/Private/Json/GLTFJsonSkySphere.cpp b/Engine/Plugins/Enterprise/GLTFExporter/MigratedFromGit/Source/GLTFExporter/Private/Json/GLTFJsonSkySphere.cpp
new file mode 100644
--- /dev/null
+++ b/Engine/Plugins/Enterprise/GLTFExporter/MigratedFromGit/Source/GLTFExporter/Private/Json/GLTFJsonSkySphere.cpp
@@ -0,0 +1,81 @@
+// Copyright Epic Games, Inc. All Rights Reserved.
+
+#include "Json/GLTFJsonSkySphere.h"
+
+void FGLTFJsonSkySphereColorCurve::WriteArray(IGLTFJsonWriter& Writer) const
+{
+	for (int32 ComponentIndex = 0; ComponentIndex < ComponentCurves.Num(); ++ComponentIndex)
+	{
+		const FComponentCurve& ComponentCurve = ComponentCurves[ComponentIndex];
+
+		Writer.StartArray();
+
+		for (const FKey& Key: ComponentCurve.Keys)
+		{
+			Writer.Write(Key.Time);
+			Writer.Write(Key.Value);
+		}
+
+		Writer.EndArray();
+	}
+}
+
+void FGLTFJsonSkySphere::WriteObject(IGLTFJsonWriter& Writer) const
+{
+	if (!Name.IsEmpty())
+	{
+		Writer.Write(TEXT("name"), Name);
+	}
+
+	Writer.Write(TEXT("skySphereMesh"), SkySphereMesh);
+	Writer.Write(TEXT("skyTexture"), SkyTexture);
+	Writer.Write(TEXT("cloudsTexture"), CloudsTexture);
+	Writer.Write(TEXT("starsTexture"), StarsTexture);
+
+	if (DirectionalLight != INDEX_NONE)
+	{
+		Writer.Write(TEXT("directionalLight"), DirectionalLight);
+	}
+
+	Writer.Write(TEXT("sunHeight"), SunHeight);
+	Writer.Write(TEXT("sunBrightness"), SunBrightness);
+	Writer.Write(TEXT("starsBrightness"), StarsBrightness);
+	Writer.Write(TEXT("cloudSpeed"), CloudSpeed);
+	Writer.Write(TEXT("cloudOpacity"), CloudOpacity);
+	Writer.Write(TEXT("horizonFalloff"), HorizonFalloff);
+
+	Writer.Write(TEXT("sunRadius"), SunRadius);
+	Writer.Write(TEXT("noisePower1"), NoisePower1);
+	Writer.Write(TEXT("noisePower2"), NoisePower2);
+
+	Writer.Write(TEXT("colorsDeterminedBySunPosition"), bColorsDeterminedBySunPosition);
+
+	Writer.Write(TEXT("zenithColor"), ZenithColor);
+	Writer.Write(TEXT("horizonColor"), HorizonColor);
+	Writer.Write(TEXT("cloudColor"), CloudColor);
+
+	if (!OverallColor.IsNearlyEqual(FGLTFJsonColor4::White, Writer.DefaultTolerance))
+	{
+		Writer.Write(TEXT("overallColor"), OverallColor);
+	}
+
+	if (ZenithColorCurve.ComponentCurves.Num() >= 3)
+	{
+		Writer.Write(TEXT("zenithColorCurve"), ZenithColorCurve);
+	}
+
+	if (HorizonColorCurve.ComponentCurves.Num() >= 3)
+	{
+		Writer.Write(TEXT("horizonColorCurve"), HorizonColorCurve);
+	}
+
+	if (CloudColorCurve.ComponentCurves.Num() >= 3)
+	{
+		Writer.Write(TEXT("cloudColorCurve"), CloudColorCurve);
+	}
+
+	if (!Scale.IsNearlyEqual(FGLTFJsonVector3::One, Writer.DefaultTolerance))
+	{
+		Writer.Write(TEXT("scale"), Scale);
+	}
+}
diff --git a/Engine/Plugins/Enterprise/GLTFExporter/MigratedFromGit/Source/GLTFExporter/Private/Json/GLTFJsonTexture.cpp b/Engine/Plugins/Enterprise/GLTFExporter/MigratedFromGit/Source/GLTFExporter/Private/Json/GLTFJsonTexture.cpp
new file mode 100644
--- /dev/null
+++ b/Engine/Plugins/Enterprise/GLTFExporter/MigratedFromGit/Source/GLTFExporter/Private/Json/GLTFJsonTexture.cpp
@@ -0,0 +1,32 @@
+// Copyright Epic Games, Inc. All Rights Reserved.
+
+#include "Json/GLTFJsonTexture.h"
+
+void FGLTFJsonTexture::WriteObject(IGLTFJsonWriter& Writer) const
+{
+	if (!Name.IsEmpty())
+	{
+		Writer.Write(TEXT("name"), Name);
+	}
+
+	if (Sampler != INDEX_NONE)
+	{
+		Writer.Write(TEXT("sampler"), Sampler);
+	}
+
+	if (Source != INDEX_NONE)
+	{
+		Writer.Write(TEXT("source"), Source);
+	}
+
+	if (Encoding != EGLTFJsonHDREncoding::None)
+	{
+		Writer.StartExtensions();
+
+		Writer.StartExtension(EGLTFJsonExtension::EPIC_TextureHDREncoding);
+		Writer.Write(TEXT("encoding"), Encoding);
+		Writer.EndExtension();
+
+		Writer.EndExtensions();
+	}
+}
diff --git a/Engine/Plugins/Enterprise/GLTFExporter/MigratedFromGit/Source/GLTFExporter/Private/Json/GLTFJsonTextureTransform.cpp b/Engine/Plugins/Enterprise/GLTFExporter/MigratedFromGit/Source/GLTFExporter/Private/Json/GLTFJsonTextureTransform.cpp
new file mode 100644
--- /dev/null
+++ b/Engine/Plugins/Enterprise/GLTFExporter/MigratedFromGit/Source/GLTFExporter/Private/Json/GLTFJsonTextureTransform.cpp
@@ -0,0 +1,45 @@
+// Copyright Epic Games, Inc. All Rights Reserved.
+
+#include "Json/GLTFJsonTextureTransform.h"
+
+void FGLTFJsonTextureTransform::WriteObject(IGLTFJsonWriter& Writer) const
+{
+	if (!Offset.IsNearlyEqual(FGLTFJsonVector2::Zero, Writer.DefaultTolerance))
+	{
+		Writer.Write(TEXT("offset"), Offset);
+	}
+
+	if (!Scale.IsNearlyEqual(FGLTFJsonVector2::One, Writer.DefaultTolerance))
+	{
+		Writer.Write(TEXT("scale"), Scale);
+	}
+
+	if (!FMath::IsNearlyEqual(Rotation, 0, Writer.DefaultTolerance))
+	{
+		Writer.Write(TEXT("rotation"), Rotation);
+	}
+}
+
+bool FGLTFJsonTextureTransform::IsNearlyEqual(const FGLTFJsonTextureTransform& Other, float Tolerance) const
+{
+	return Offset.IsNearlyEqual(Other.Offset, Tolerance)
+		&& Scale.IsNearlyEqual(Other.Scale, Tolerance)
+		&& FMath::IsNearlyEqual(Rotation, Other.Rotation, Tolerance);
+}
+
+bool FGLTFJsonTextureTransform::IsExactlyEqual(const FGLTFJsonTextureTransform& Other) const
+{
+	return Offset.X == Other.Offset.X && Offset.Y == Other.Offset.Y
+		&& Scale.X == Other.Scale.X && Scale.Y == Other.Scale.Y
+		&& Rotation == Other.Rotation;
+}
+
+bool FGLTFJsonTextureTransform::IsNearlyDefault(float Tolerance) const
+{
+	return IsNearlyEqual({}, Tolerance);
+}
+
+bool FGLTFJsonTextureTransform::IsExactlyDefault() const
+{
+	return IsExactlyEqual({});
+}
diff --git a/Engine/Plugins/Enterprise/GLTFExporter/MigratedFromGit/Source/GLTFExporter/Private/Json/GLTFJsonUtility.cpp b/Engine/Plugins/Enterprise/GLTFExporter/MigratedFromGit/Source/GLTFExporter/Private/Json/GLTFJsonUtility.cpp
new file mode 100644
--- /dev/null
+++ b/Engine/Plugins/Enterprise/GLTFExporter/MigratedFromGit/Source/GLTFExporter/Private/Json/GLTFJsonUtility.cpp
@@ -0,0 +1,188 @@
+// Copyright Epic Games, Inc. All Rights Reserved.
+
+#include "Json/GLTFJsonUtility.h"
+
+const TCHAR* FGLTFJsonUtility::GetValue(EGLTFJsonExtension Enum)
+{
+	switch (Enum)
+	{
+		case EGLTFJsonExtension::KHR_LightsPunctual:      return TEXT("KHR_lights_punctual");
+		case EGLTFJsonExtension::KHR_MaterialsClearCoat:  return TEXT("KHR_materials_clearcoat");
+		case EGLTFJsonExtension::KHR_MaterialsUnlit:      return TEXT("KHR_materials_unlit");
+		case EGLTFJsonExtension::KHR_MaterialsVariants:   return TEXT("KHR_materials_variants");
+		case EGLTFJsonExtension::KHR_MeshQuantization:    return TEXT("KHR_mesh_quantization");
+		case EGLTFJsonExtension::KHR_TextureTransform:    return TEXT("KHR_texture_transform");
+		case EGLTFJsonExtension::EPIC_AnimationHotspots:  return TEXT("EPIC_animation_hotspots");
+		case EGLTFJsonExtension::EPIC_AnimationPlayback:  return TEXT("EPIC_animation_playback");
+		case EGLTFJsonExtension::EPIC_BlendModes:         return TEXT("EPIC_blend_modes");
+		case EGLTFJsonExtension::EPIC_CameraControls:     return TEXT("EPIC_camera_controls");
+		case EGLTFJsonExtension::EPIC_HDRIBackdrops:      return TEXT("EPIC_hdri_backdrops");
+		case EGLTFJsonExtension::EPIC_LevelVariantSets:   return TEXT("EPIC_level_variant_sets");
+		case EGLTFJsonExtension::EPIC_LightmapTextures:   return TEXT("EPIC_lightmap_textures");
+		case EGLTFJsonExtension::EPIC_SkySpheres:         return TEXT("EPIC_sky_spheres");
+		case EGLTFJsonExtension::EPIC_TextureHDREncoding: return TEXT("EPIC_texture_hdr_encoding");
+		default:
+			checkNoEntry();
+			return TEXT("");
+	}
+}
+
+const TCHAR* FGLTFJsonUtility::GetValue(EGLTFJsonAlphaMode Enum)
+{
+	switch (Enum)
+	{
+		case EGLTFJsonAlphaMode::Opaque: return TEXT("OPAQUE");
+		case EGLTFJsonAlphaMode::Blend:  return TEXT("BLEND");
+		case EGLTFJsonAlphaMode::Mask:   return TEXT("MASK");
+		default:
+			checkNoEntry();
+			return TEXT("");
+	}
+}
+
+const TCHAR* FGLTFJsonUtility::GetValue(EGLTFJsonBlendMode Enum)
+{
+	switch (Enum)
+	{
+		case EGLTFJsonBlendMode::Additive:       return TEXT("ADDITIVE");
+		case EGLTFJsonBlendMode::Modulate:       return TEXT("MODULATE");
+		case EGLTFJsonBlendMode::AlphaComposite: return TEXT("ALPHACOMPOSITE");
+		default:
+			checkNoEntry();
+			return TEXT("");
+	}
+}
+
+const TCHAR* FGLTFJsonUtility::GetValue(EGLTFJsonMimeType Enum)
+{
+	switch (Enum)
+	{
+		case EGLTFJsonMimeType::PNG:  return TEXT("image/png");
+		case EGLTFJsonMimeType::JPEG: return TEXT("image/jpeg");
+		default:
+			checkNoEntry();
+			return TEXT("");
+	}
+}
+
+const TCHAR* FGLTFJsonUtility::GetValue(EGLTFJsonAccessorType Enum)
+{
+	switch (Enum)
+	{
+		case EGLTFJsonAccessorType::Scalar: return TEXT("SCALAR");
+		case EGLTFJsonAccessorType::Vec2:   return TEXT("VEC2");
+		case EGLTFJsonAccessorType::Vec3:   return TEXT("VEC3");
+		case EGLTFJsonAccessorType::Vec4:   return TEXT("VEC4");
+		case EGLTFJsonAccessorType::Mat2:   return TEXT("MAT2");
+		case EGLTFJsonAccessorType::Mat3:   return TEXT("MAT3");
+		case EGLTFJsonAccessorType::Mat4:   return TEXT("MAT4");
+		default:
+			checkNoEntry();
+			return TEXT("");
+	}
+}
+
+const TCHAR* FGLTFJsonUtility::GetValue(EGLTFJsonHDREncoding Enum)
+{
+	switch (Enum)
+	{
+		case EGLTFJsonHDREncoding::RGBM: return TEXT("RGBM");
+		case EGLTFJsonHDREncoding::RGBE: return TEXT("RGBE");
+		default:
+			checkNoEntry();
+			return TEXT("");
+	}
+}
+
+const TCHAR* FGLTFJsonUtility::GetValue(EGLTFJsonCubeFace Enum)
+{
+	switch (Enum)
+	{
+		case EGLTFJsonCubeFace::PosX: return TEXT("PosX");
+		case EGLTFJsonCubeFace::NegX: return TEXT("NegX");
+		case EGLTFJsonCubeFace::PosY: return TEXT("PosY");
+		case EGLTFJsonCubeFace::NegY: return TEXT("NegY");
+		case EGLTFJsonCubeFace::PosZ: return TEXT("PosZ");
+		case EGLTFJsonCubeFace::NegZ: return TEXT("NegZ");
+		default:
+			checkNoEntry();
+			return TEXT("");
+	}
+}
+
+const TCHAR* FGLTFJsonUtility::GetValue(EGLTFJsonCameraType Enum)
+{
+	switch (Enum)
+	{
+		case EGLTFJsonCameraType::Perspective:  return TEXT("perspective");
+		case EGLTFJsonCameraType::Orthographic: return TEXT("orthographic");
+		default:
+			checkNoEntry();
+			return TEXT("");
+	}
+}
+
+const TCHAR* FGLTFJsonUtility::GetValue(EGLTFJsonLightType Enum)
+{
+	switch (Enum)
+	{
+		case EGLTFJsonLightType::Directional: return TEXT("directional");
+		case EGLTFJsonLightType::Point:       return TEXT("point");
+		case EGLTFJsonLightType::Spot:        return TEXT("spot");
+		default:
+			checkNoEntry();
+			return TEXT("");
+	}
+}
+
+const TCHAR* FGLTFJsonUtility::GetValue(EGLTFJsonInterpolation Enum)
+{
+	switch (Enum)
+	{
+		case EGLTFJsonInterpolation::Linear:      return TEXT("LINEAR");
+		case EGLTFJsonInterpolation::Step:        return TEXT("STEP");
+		case EGLTFJsonInterpolation::CubicSpline: return TEXT("CUBICSPLINE");
+		default:
+			checkNoEntry();
+			return TEXT("");
+	}
+}
+
+const TCHAR* FGLTFJsonUtility::GetValue(EGLTFJsonTargetPath Enum)
+{
+	switch (Enum)
+	{
+		case EGLTFJsonTargetPath::Translation: return TEXT("translation");
+		case EGLTFJsonTargetPath::Rotation:    return TEXT("rotation");
+		case EGLTFJsonTargetPath::Scale:       return TEXT("scale");
+		case EGLTFJsonTargetPath::Weights:     return TEXT("weights");
+		default:
+			checkNoEntry();
+			return TEXT("");
+	}
+}
+
+const TCHAR* FGLTFJsonUtility::GetValue(EGLTFJsonCameraControlMode Enum)
+{
+	switch (Enum)
+	{
+		case EGLTFJsonCameraControlMode::FreeLook: return TEXT("freeLook");
+		case EGLTFJsonCameraControlMode::Orbital:  return TEXT("orbital");
+		default:
+			checkNoEntry();
+			return TEXT("");
+	}
+}
+
+const TCHAR* FGLTFJsonUtility::GetValue(EGLTFJsonShadingModel Enum)
+{
+	switch (Enum)
+	{
+		case EGLTFJsonShadingModel::Default:   return TEXT("Default");
+		case EGLTFJsonShadingModel::Unlit:     return TEXT("Unlit");
+		case EGLTFJsonShadingModel::ClearCoat: return TEXT("ClearCoat");
+		default:
+			checkNoEntry();
+			return TEXT("");
+	}
+}
diff --git a/Engine/Plugins/Enterprise/GLTFExporter/MigratedFromGit/Source/GLTFExporter/Public/Core/GLTFQuaternion.h b/Engine/Plugins/Enterprise/GLTFExporter/MigratedFromGit/Source/GLTFExporter/Public/Core/GLTFQuaternion.h
--- a/Engine/Plugins/Enterprise/GLTFExporter/MigratedFromGit/Source/GLTFExporter/Public/Core/GLTFQuaternion.h
+++ b/Engine/Plugins/Enterprise/GLTFExporter/MigratedFromGit/Source/GLTFExporter/Public/Core/GLTFQuaternion.h
@@ -9,7 +9,7 @@
 #endif
 
 template <typename ComponentType>
-struct TGLTFQuaternion
+struct GLTFEXPORTER_API TGLTFQuaternion
 {
 	union
 	{
diff --git a/Engine/Plugins/Enterprise/GLTFExporter/MigratedFromGit/Source/GLTFExporter/Public/Json/GLTFJsonAccessor.h b/Engine/Plugins/Enterprise/GLTFExporter/MigratedFromGit/Source/GLTFExporter/Public/Json/GLTFJsonAccessor.h
--- a/Engine/Plugins/Enterprise/GLTFExporter/MigratedFromGit/Source/GLTFExporter/Public/Json/GLTFJsonAccessor.h
+++ b/Engine/Plugins/Enterprise/GLTFExporter/MigratedFromGit/Source/GLTFExporter/Public/Json/GLTFJsonAccessor.h
@@ -6,7 +6,7 @@
 #include "Json/GLTFJsonEnums.h"
 #include "Json/GLTFJsonIndex.h"
 
-struct FGLTFJsonAccessor : IGLTFJsonObject
+struct GLTFEXPORTER_API FGLTFJsonAccessor : IGLTFJsonObject
 {
 	FString Name;
 
@@ -33,33 +33,5 @@ struct FGLTFJsonAccessor : IGLTFJsonObject
 	{
 	}
 
-	virtual void WriteObject(IGLTFJsonWriter& Writer) const override
-	{
-		if (!Name.IsEmpty())
-		{
-			Writer.Write(TEXT("name"), Name);
-		}
-
-		Writer.Write(TEXT("bufferView"), BufferView);
-
-		if (ByteOffset != 0)
-		{
-			Writer.Write(TEXT("byteOffset"), ByteOffset);
-		}
-
-		Writer.Write(TEXT("count"), Count);
-		Writer.Write(TEXT("type"), Type);
-		Writer.Write(TEXT("componentType"), ComponentType);
-
-		if (bNormalized)
-		{
-			Writer.Write(TEXT("normalized"), bNormalized);
-		}
-
-		if (MinMaxLength > 0)
-		{
-			Writer.Write(TEXT("min"), Min, MinMaxLength);
-			Writer.Write(TEXT("max"), Max, MinMaxLength);
-		}
-	}
+	virtual void WriteObject(IGLTFJsonWriter& Writer) const override;
 };
diff --git a/Engine/Plugins/Enterprise/GLTFExporter/MigratedFromGit/Source/GLTFExporter/Public/Json/GLTFJsonAnimation.h b/Engine/Plugins/Enterprise/GLTFExporter/MigratedFromGit/Source/GLTFExporter/Public/Json/GLTFJsonAnimation.h
--- a/Engine/Plugins/Enterprise/GLTFExporter/MigratedFromGit/Source/GLTFExporter/Public/Json/GLTFJsonAnimation.h
+++ b/Engine/Plugins/Enterprise/GLTFExporter/MigratedFromGit/Source/GLTFExporter/Public/Json/GLTFJsonAnimation.h
@@ -7,7 +7,7 @@
 #include "Json/GLTFJsonIndex.h"
 #include "Json/GLTFJsonAnimationPlayback.h"
 
-struct FGLTFJsonAnimationChannelTarget : IGLTFJsonObject
+struct GLTFEXPORTER_API FGLTFJsonAnimationChannelTarget : IGLTFJsonObject
 {
 	FGLTFJsonNodeIndex Node;
 	EGLTFJsonTargetPath Path;
@@ -17,27 +17,18 @@ struct FGLTFJsonAnimationChannelTarget : IGLTFJsonObject
 	{
 	}
 
-	virtual void WriteObject(IGLTFJsonWriter& Writer) const override
-	{
-		Writer.Write(TEXT("node"), Node);
-		Writer.Write(TEXT("path"), Path);
-	}
+	virtual void WriteObject(IGLTFJsonWriter& Writer) const override;
 };
 
-struct FGLTFJsonAnimationChannel : IGLTFJsonObject
+struct GLTFEXPORTER_API FGLTFJsonAnimationChannel : IGLTFJsonObject
 {
 	FGLTFJsonAnimationSamplerIndex Sampler;
 	FGLTFJsonAnimationChannelTarget Target;
 
-	virtual void WriteObject(IGLTFJsonWriter& Writer) const override
-	{
-		Writer.Write(TEXT("sampler"), Sampler);
-
-		Writer.Write(TEXT("target"), Target);
-	}
+	virtual void WriteObject(IGLTFJsonWriter& Writer) const override;
 };
 
-struct FGLTFJsonAnimationSampler : IGLTFJsonObject
+struct GLTFEXPORTER_API FGLTFJsonAnimationSampler : IGLTFJsonObject
 {
 	FGLTFJsonAccessorIndex Input;
 	FGLTFJsonAccessorIndex Output;
@@ -49,19 +40,10 @@ struct FGLTFJsonAnimationSampler : IGLTFJsonObject
 	{
 	}
 
-	virtual void WriteObject(IGLTFJsonWriter& Writer) const override
-	{
-		Writer.Write(TEXT("input"), Input);
-		Writer.Write(TEXT("output"), Output);
-
-		if (Interpolation != EGLTFJsonInterpolation::Linear)
-		{
-			Writer.Write(TEXT("interpolation"), Interpolation);
-		}
-	}
+	virtual void WriteObject(IGLTFJsonWriter& Writer) const override;
 };
 
-struct FGLTFJsonAnimation : IGLTFJsonObject
+struct GLTFEXPORTER_API FGLTFJsonAnimation : IGLTFJsonObject
 {
 	FString Name;
 
@@ -70,21 +52,5 @@ struct FGLTFJsonAnimation : IGLTFJsonObject
 
 	FGLTFJsonAnimationPlayback Playback;
 
-	virtual void WriteObject(IGLTFJsonWriter& Writer) const override
-	{
-		if (!Name.IsEmpty())
-		{
-			Writer.Write(TEXT("name"), Name);
-		}
-
-		Writer.Write(TEXT("channels"), Channels);
-		Writer.Write(TEXT("samplers"), Samplers);
-
-		if (Playback != FGLTFJsonAnimationPlayback())
-		{
-			Writer.StartExtensions();
-			Writer.Write(EGLTFJsonExtension::EPIC_AnimationPlayback, Playback);
-			Writer.EndExtensions();
-		}
-	}
+	virtual void WriteObject(IGLTFJsonWriter& Writer) const override;
 };
diff --git a/Engine/Plugins/Enterprise/GLTFExporter/MigratedFromGit/Source/GLTFExporter/Public/Json/GLTFJsonAnimationPlayback.h b/Engine/Plugins/Enterprise/GLTFExporter/MigratedFromGit/Source/GLTFExporter/Public/Json/GLTFJsonAnimationPlayback.h
--- a/Engine/Plugins/Enterprise/GLTFExporter/MigratedFromGit/Source/GLTFExporter/Public/Json/GLTFJsonAnimationPlayback.h
+++ b/Engine/Plugins/Enterprise/GLTFExporter/MigratedFromGit/Source/GLTFExporter/Public/Json/GLTFJsonAnimationPlayback.h
@@ -4,7 +4,7 @@
 
 #include "Json/GLTFJsonObject.h"
 
-struct FGLTFJsonAnimationPlayback : IGLTFJsonObject
+struct GLTFEXPORTER_API FGLTFJsonAnimationPlayback : IGLTFJsonObject
 {
 	FString Name;
 
@@ -22,44 +22,8 @@ struct FGLTFJsonAnimationPlayback : IGLTFJsonObject
 	{
 	}
 
-	virtual void WriteObject(IGLTFJsonWriter& Writer) const override
-	{
-		if (!Name.IsEmpty())
-		{
-			Writer.Write(TEXT("name"), Name);
-		}
-
-		if (bLoop != true)
-		{
-			Writer.Write(TEXT("loop"), bLoop);
-		}
-
-		if (bAutoPlay != true)
-		{
-			Writer.Write(TEXT("autoPlay"), bAutoPlay);
-		}
-
-		if (!FMath::IsNearlyEqual(PlayRate, 1, Writer.DefaultTolerance))
-		{
-			Writer.Write(TEXT("playRate"), PlayRate);
-		}
-
-		if (!FMath::IsNearlyEqual(StartTime, 0, Writer.DefaultTolerance))
-		{
-			Writer.Write(TEXT("startTime"), StartTime);
-		}
-	}
+	virtual void WriteObject(IGLTFJsonWriter& Writer) const override;
 
-	bool operator==(const FGLTFJsonAnimationPlayback& Other) const
-	{
-		return bLoop == Other.bLoop
-			&& bAutoPlay == Other.bAutoPlay
-			&& PlayRate == Other.PlayRate
-			&& StartTime == Other.StartTime;
-	}
-
-	bool operator!=(const FGLTFJsonAnimationPlayback& Other) const
-	{
-		return !(*this == Other);
-	}
+	bool operator==(const FGLTFJsonAnimationPlayback& Other) const;
+	bool operator!=(const FGLTFJsonAnimationPlayback& Other) const;
 };
diff --git a/Engine/Plugins/Enterprise/GLTFExporter/MigratedFromGit/Source/GLTFExporter/Public/Json/GLTFJsonArray.h b/Engine/Plugins/Enterprise/GLTFExporter/MigratedFromGit/Source/GLTFExporter/Public/Json/GLTFJsonArray.h
--- a/Engine/Plugins/Enterprise/GLTFExporter/MigratedFromGit/Source/GLTFExporter/Public/Json/GLTFJsonArray.h
+++ b/Engine/Plugins/Enterprise/GLTFExporter/MigratedFromGit/Source/GLTFExporter/Public/Json/GLTFJsonArray.h
@@ -5,7 +5,7 @@
 #include "Json/GLTFJsonValue.h"
 #include "Json/GLTFJsonWriter.h"
 
-struct IGLTFJsonArray : IGLTFJsonValue
+struct GLTFEXPORTER_API IGLTFJsonArray : IGLTFJsonValue
 {
 	virtual void WriteValue(IGLTFJsonWriter& Writer) const override final
 	{
diff --git a/Engine/Plugins/Enterprise/GLTFExporter/MigratedFromGit/Source/GLTFExporter/Public/Json/GLTFJsonAsset.h b/Engine/Plugins/Enterprise/GLTFExporter/MigratedFromGit/Source/GLTFExporter/Public/Json/GLTFJsonAsset.h
--- a/Engine/Plugins/Enterprise/GLTFExporter/MigratedFromGit/Source/GLTFExporter/Public/Json/GLTFJsonAsset.h
+++ b/Engine/Plugins/Enterprise/GLTFExporter/MigratedFromGit/Source/GLTFExporter/Public/Json/GLTFJsonAsset.h
@@ -4,7 +4,7 @@
 
 #include "Json/GLTFJsonObject.h"
 
-struct FGLTFJsonAsset : IGLTFJsonObject
+struct GLTFEXPORTER_API FGLTFJsonAsset : IGLTFJsonObject
 {
 	FString Version;
 	FString Generator;
@@ -15,18 +15,5 @@ struct FGLTFJsonAsset : IGLTFJsonObject
 	{
 	}
 
-	virtual void WriteObject(IGLTFJsonWriter& Writer) const override
-	{
-		Writer.Write(TEXT("version"), Version);
-
-		if (!Generator.IsEmpty())
-		{
-			Writer.Write(TEXT("generator"), Generator);
-		}
-
-		if (!Copyright.IsEmpty())
-		{
-			Writer.Write(TEXT("copyright"), Copyright);
-		}
-	}
+	virtual void WriteObject(IGLTFJsonWriter& Writer) const override;
 };
diff --git a/Engine/Plugins/Enterprise/GLTFExporter/MigratedFromGit/Source/GLTFExporter/Public/Json/GLTFJsonBackdrop.h b/Engine/Plugins/Enterprise/GLTFExporter/MigratedFromGit/Source/GLTFExporter/Public/Json/GLTFJsonBackdrop.h
--- a/Engine/Plugins/Enterprise/GLTFExporter/MigratedFromGit/Source/GLTFExporter/Public/Json/GLTFJsonBackdrop.h
+++ b/Engine/Plugins/Enterprise/GLTFExporter/MigratedFromGit/Source/GLTFExporter/Public/Json/GLTFJsonBackdrop.h
@@ -6,7 +6,7 @@
 #include "Json/GLTFJsonIndex.h"
 #include "Json/GLTFJsonVector.h"
 
-struct FGLTFJsonBackdrop : IGLTFJsonObject
+struct GLTFEXPORTER_API FGLTFJsonBackdrop : IGLTFJsonObject
 {
 	FString Name;
 
@@ -32,31 +32,5 @@ struct FGLTFJsonBackdrop : IGLTFJsonObject
 	{
 	}
 
-	virtual void WriteObject(IGLTFJsonWriter& Writer) const override
-	{
-		if (!Name.IsEmpty())
-		{
-			Writer.Write(TEXT("name"), Name);
-		}
-
-		if (Mesh != INDEX_NONE)
-		{
-			Writer.Write(TEXT("mesh"), Mesh);
-		}
-
-		Writer.Write(TEXT("cubemap"), Cubemap);
-
-		Writer.Write(TEXT("intensity"), Intensity);
-		Writer.Write(TEXT("size"), Size);
-
-		if (!FMath::IsNearlyZero(Angle))
-		{
-			Writer.Write(TEXT("angle"), Angle);
-		}
-
-		Writer.Write(TEXT("projectionCenter"), ProjectionCenter);
-
-		Writer.Write(TEXT("lightingDistanceFactor"), LightingDistanceFactor);
-		Writer.Write(TEXT("useCameraProjection"), UseCameraProjection);
-	}
+	virtual void WriteObject(IGLTFJsonWriter& Writer) const override;
 };
diff --git a/Engine/Plugins/Enterprise/GLTFExporter/MigratedFromGit/Source/GLTFExporter/Public/Json/GLTFJsonBuffer.h b/Engine/Plugins/Enterprise/GLTFExporter/MigratedFromGit/Source/GLTFExporter/Public/Json/GLTFJsonBuffer.h
--- a/Engine/Plugins/Enterprise/GLTFExporter/MigratedFromGit/Source/GLTFExporter/Public/Json/GLTFJsonBuffer.h
+++ b/Engine/Plugins/Enterprise/GLTFExporter/MigratedFromGit/Source/GLTFExporter/Public/Json/GLTFJsonBuffer.h
@@ -4,7 +4,7 @@
 
 #include "Json/GLTFJsonObject.h"
 
-struct FGLTFJsonBuffer : IGLTFJsonObject
+struct GLTFEXPORTER_API FGLTFJsonBuffer : IGLTFJsonObject
 {
 	FString Name;
 
@@ -16,18 +16,5 @@ struct FGLTFJsonBuffer : IGLTFJsonObject
 	{
 	}
 
-	virtual void WriteObject(IGLTFJsonWriter& Writer) const override
-	{
-		if (!Name.IsEmpty())
-		{
-			Writer.Write(TEXT("name"), Name);
-		}
-
-		if (!URI.IsEmpty())
-		{
-			Writer.Write(TEXT("uri"), URI);
-		}
-
-		Writer.Write(TEXT("byteLength"), ByteLength);
-	}
+	virtual void WriteObject(IGLTFJsonWriter& Writer) const override;
 };
diff --git a/Engine/Plugins/Enterprise/GLTFExporter/MigratedFromGit/Source/GLTFExporter/Public/Json/GLTFJsonBufferView.h b/Engine/Plugins/Enterprise/GLTFExporter/MigratedFromGit/Source/GLTFExporter/Public/Json/GLTFJsonBufferView.h
--- a/Engine/Plugins/Enterprise/GLTFExporter/MigratedFromGit/Source/GLTFExporter/Public/Json/GLTFJsonBufferView.h
+++ b/Engine/Plugins/Enterprise/GLTFExporter/MigratedFromGit/Source/GLTFExporter/Public/Json/GLTFJsonBufferView.h
@@ -6,7 +6,7 @@
 #include "Json/GLTFJsonEnums.h"
 #include "Json/GLTFJsonIndex.h"
 
-struct FGLTFJsonBufferView : IGLTFJsonObject
+struct GLTFEXPORTER_API FGLTFJsonBufferView : IGLTFJsonObject
 {
 	FString Name;
 
@@ -24,32 +24,7 @@ struct FGLTFJsonBufferView : IGLTFJsonObject
 		, ByteStride(0)
 		, Target(EGLTFJsonBufferTarget::None)
 	{
-		// check that view fits completely inside the buffer
 	}
 
-	virtual void WriteObject(IGLTFJsonWriter& Writer) const override
-	{
-		if (!Name.IsEmpty())
-		{
-			Writer.Write(TEXT("name"), Name);
-		}
-
-		Writer.Write(TEXT("buffer"), Buffer);
-		Writer.Write(TEXT("byteLength"), ByteLength);
-
-		if (ByteOffset != 0)
-		{
-			Writer.Write(TEXT("byteOffset"), ByteOffset);
-		}
-
-		if (ByteStride != 0)
-		{
-			Writer.Write(TEXT("byteStride"), ByteStride);
-		}
-
-		if (Target != EGLTFJsonBufferTarget::None)
-		{
-			Writer.Write(TEXT("target"), Target);
-		}
-	}
+	virtual void WriteObject(IGLTFJsonWriter& Writer) const override;
 };
diff --git a/Engine/Plugins/Enterprise/GLTFExporter/MigratedFromGit/Source/GLTFExporter/Public/Json/GLTFJsonCamera.h b/Engine/Plugins/Enterprise/GLTFExporter/MigratedFromGit/Source/GLTFExporter/Public/Json/GLTFJsonCamera.h
--- a/Engine/Plugins/Enterprise/GLTFExporter/MigratedFromGit/Source/GLTFExporter/Public/Json/GLTFJsonCamera.h
+++ b/Engine/Plugins/Enterprise/GLTFExporter/MigratedFromGit/Source/GLTFExporter/Public/Json/GLTFJsonCamera.h
@@ -6,7 +6,7 @@
 #include "Json/GLTFJsonEnums.h"
 #include "Json/GLTFJsonCameraControl.h"
 
-struct FGLTFJsonOrthographic : IGLTFJsonObject
+struct GLTFEXPORTER_API FGLTFJsonOrthographic : IGLTFJsonObject
 {
 	float XMag; // horizontal magnification of the view
 	float YMag; // vertical magnification of the view
@@ -21,16 +21,10 @@ struct FGLTFJsonOrthographic : IGLTFJsonObject
 	{
 	}
 
-	virtual void WriteObject(IGLTFJsonWriter& Writer) const override
-	{
-		Writer.Write(TEXT("xmag"), XMag);
-		Writer.Write(TEXT("ymag"), YMag);
-		Writer.Write(TEXT("zfar"), ZFar);
-		Writer.Write(TEXT("znear"), ZNear);
-	}
+	virtual void WriteObject(IGLTFJsonWriter& Writer) const override;
 };
 
-struct FGLTFJsonPerspective : IGLTFJsonObject
+struct GLTFEXPORTER_API FGLTFJsonPerspective : IGLTFJsonObject
 {
 	float AspectRatio; // aspect ratio of the field of view
 	float YFov; // vertical field of view in radians
@@ -45,25 +39,10 @@ struct FGLTFJsonPerspective : IGLTFJsonObject
 	{
 	}
 
-	virtual void WriteObject(IGLTFJsonWriter& Writer) const override
-	{
-		if (!FMath::IsNearlyEqual(AspectRatio, 0, Writer.DefaultTolerance))
-		{
-			Writer.Write(TEXT("aspectRatio"), AspectRatio);
-		}
-
-		Writer.Write(TEXT("yfov"), YFov);
-
-		if (!FMath::IsNearlyEqual(ZFar, 0, Writer.DefaultTolerance))
-		{
-			Writer.Write(TEXT("zfar"), ZFar);
-		}
-
-		Writer.Write(TEXT("znear"), ZNear);
-	}
+	virtual void WriteObject(IGLTFJsonWriter& Writer) const override;
 };
 
-struct FGLTFJsonCamera : IGLTFJsonObject
+struct GLTFEXPORTER_API FGLTFJsonCamera : IGLTFJsonObject
 {
 	FString Name;
 
@@ -78,34 +57,5 @@ struct FGLTFJsonCamera : IGLTFJsonObject
 	{
 	}
 
-	virtual void WriteObject(IGLTFJsonWriter& Writer) const override
-	{
-		if (!Name.IsEmpty())
-		{
-			Writer.Write(TEXT("name"), Name);
-		}
-
-		Writer.Write(TEXT("type"), Type);
-
-		switch (Type)
-		{
-			case EGLTFJsonCameraType::Orthographic:
-				Writer.Write(TEXT("orthographic"), Orthographic);
-				break;
-
-			case EGLTFJsonCameraType::Perspective:
-				Writer.Write(TEXT("perspective"), Perspective);
-				break;
-
-			default:
-				break;
-		}
-
-		if (CameraControl.IsSet())
-		{
-			Writer.StartExtensions();
-			Writer.Write(EGLTFJsonExtension::EPIC_CameraControls, CameraControl.GetValue());
-			Writer.EndExtensions();
-		}
-	}
+	virtual void WriteObject(IGLTFJsonWriter& Writer) const override;
 };
diff --git a/Engine/Plugins/Enterprise/GLTFExporter/MigratedFromGit/Source/GLTFExporter/Public/Json/GLTFJsonCameraControl.h b/Engine/Plugins/Enterprise/GLTFExporter/MigratedFromGit/Source/GLTFExporter/Public/Json/GLTFJsonCameraControl.h
--- a/Engine/Plugins/Enterprise/GLTFExporter/MigratedFromGit/Source/GLTFExporter/Public/Json/GLTFJsonCameraControl.h
+++ b/Engine/Plugins/Enterprise/GLTFExporter/MigratedFromGit/Source/GLTFExporter/Public/Json/GLTFJsonCameraControl.h
@@ -6,7 +6,7 @@
 #include "Json/GLTFJsonIndex.h"
 #include "Json/GLTFJsonVector.h"
 
-struct FGLTFJsonCameraControl : IGLTFJsonObject
+struct GLTFEXPORTER_API FGLTFJsonCameraControl : IGLTFJsonObject
 {
 	EGLTFJsonCameraControlMode Mode;
 	FGLTFJsonNodeIndex Target;
@@ -36,41 +36,5 @@ struct FGLTFJsonCameraControl : IGLTFJsonObject
 	{
 	}
 
-	virtual void WriteObject(IGLTFJsonWriter& Writer) const override
-	{
-		Writer.Write(TEXT("mode"), Mode);
-
-		if (Target != INDEX_NONE && Mode == EGLTFJsonCameraControlMode::Orbital)
-		{
-			Writer.Write(TEXT("target"), Target);
-		}
-
-		Writer.Write(TEXT("maxDistance"), MaxDistance);
-		Writer.Write(TEXT("minDistance"), MinDistance);
-
-		if (!FMath::IsNearlyEqual(MaxPitch, 90, Writer.DefaultTolerance))
-		{
-			Writer.Write(TEXT("maxPitch"), MaxPitch);
-		}
-
-		if (!FMath::IsNearlyEqual(MinPitch, -90, Writer.DefaultTolerance))
-		{
-			Writer.Write(TEXT("minPitch"), MinPitch);
-		}
-
-		if (!FMath::IsNearlyEqual(MaxYaw, 360, Writer.DefaultTolerance))
-		{
-			Writer.Write(TEXT("maxYaw"), MaxYaw);
-		}
-
-		if (!FMath::IsNearlyEqual(MinYaw, 0, Writer.DefaultTolerance))
-		{
-			Writer.Write(TEXT("minYaw"), MinYaw);
-		}
-
-		Writer.Write(TEXT("rotationSensitivity"), RotationSensitivity);
-		Writer.Write(TEXT("rotationInertia"), RotationInertia);
-		Writer.Write(TEXT("dollySensitivity"), DollySensitivity);
-		Writer.Write(TEXT("dollyDuration"), DollyDuration);
-	}
+	virtual void WriteObject(IGLTFJsonWriter& Writer) const override;
 };
diff --git a/Engine/Plugins/Enterprise/GLTFExporter/MigratedFromGit/Source/GLTFExporter/Public/Json/GLTFJsonColor.h b/Engine/Plugins/Enterprise/GLTFExporter/MigratedFromGit/Source/GLTFExporter/Public/Json/GLTFJsonColor.h
--- a/Engine/Plugins/Enterprise/GLTFExporter/MigratedFromGit/Source/GLTFExporter/Public/Json/GLTFJsonColor.h
+++ b/Engine/Plugins/Enterprise/GLTFExporter/MigratedFromGit/Source/GLTFExporter/Public/Json/GLTFJsonColor.h
@@ -6,7 +6,7 @@
 #include "Json/GLTFJsonArray.h"
 
 template <typename BaseType>
-struct TGLTFJsonColor : BaseType, IGLTFJsonArray
+struct GLTFEXPORTER_API TGLTFJsonColor : BaseType, IGLTFJsonArray
 {
 	TGLTFJsonColor(const BaseType& Other)
 		: BaseType(Other)
@@ -59,7 +59,7 @@ struct TGLTFJsonColor : BaseType, IGLTFJsonArray
 	}
 };
 
-struct FGLTFJsonColor3 : TGLTFJsonColor<FGLTFColor3>
+struct GLTFEXPORTER_API FGLTFJsonColor3 : TGLTFJsonColor<FGLTFColor3>
 {
 	static const FGLTFJsonColor3 Black;
 	static const FGLTFJsonColor3 White;
@@ -68,7 +68,7 @@ struct FGLTFJsonColor3 : TGLTFJsonColor<FGLTFColor3>
 	using TGLTFJsonColor::operator=;
 };
 
-struct FGLTFJsonColor4 : TGLTFJsonColor<FGLTFColor4>
+struct GLTFEXPORTER_API FGLTFJsonColor4 : TGLTFJsonColor<FGLTFColor4>
 {
 	static const FGLTFJsonColor4 Black;
 	static const FGLTFJsonColor4 White;
diff --git a/Engine/Plugins/Enterprise/GLTFExporter/MigratedFromGit/Source/GLTFExporter/Public/Json/GLTFJsonEpicLevelVariantSets.h b/Engine/Plugins/Enterprise/GLTFExporter/MigratedFromGit/Source/GLTFExporter/Public/Json/GLTFJsonEpicLevelVariantSets.h
--- a/Engine/Plugins/Enterprise/GLTFExporter/MigratedFromGit/Source/GLTFExporter/Public/Json/GLTFJsonEpicLevelVariantSets.h
+++ b/Engine/Plugins/Enterprise/GLTFExporter/MigratedFromGit/Source/GLTFExporter/Public/Json/GLTFJsonEpicLevelVariantSets.h
@@ -6,7 +6,7 @@
 #include "Json/GLTFJsonIndex.h"
 #include "Misc/Optional.h"
 
-struct FGLTFJsonEpicVariantMaterial : IGLTFJsonObject
+struct GLTFEXPORTER_API FGLTFJsonEpicVariantMaterial : IGLTFJsonObject
 {
 	FGLTFJsonMaterialIndex Material;
 	int32                  Index;
@@ -17,18 +17,10 @@ struct FGLTFJsonEpicVariantMaterial : IGLTFJsonObject
 	{
 	}
 
-	virtual void WriteObject(IGLTFJsonWriter& Writer) const override
-	{
-		Writer.Write(TEXT("material"), Material);
-
-		if (Index != INDEX_NONE)
-		{
-			Writer.Write(TEXT("index"), Index);
-		}
-	}
+	virtual void WriteObject(IGLTFJsonWriter& Writer) const override;
 };
 
-struct FGLTFJsonEpicVariantNodeProperties : IGLTFJsonObject
+struct GLTFEXPORTER_API FGLTFJsonEpicVariantNodeProperties : IGLTFJsonObject
 {
 	FGLTFJsonNodeIndex Node;
 	TOptional<bool>    bIsVisible;
@@ -36,35 +28,10 @@ struct FGLTFJsonEpicVariantNodeProperties : IGLTFJsonObject
 	TOptional<FGLTFJsonMeshIndex>        Mesh;
 	TArray<FGLTFJsonEpicVariantMaterial> Materials;
 
-	virtual void WriteObject(IGLTFJsonWriter& Writer) const override
-	{
-		if (Node != INDEX_NONE)
-		{
-			Writer.Write(TEXT("node"), Node);
-		}
-
-		Writer.StartObject(TEXT("properties"));
-
-		if (bIsVisible.IsSet())
-		{
-			Writer.Write(TEXT("visible"), bIsVisible.GetValue());
-		}
-
-		if (Mesh.IsSet())
-		{
-			Writer.Write(TEXT("mesh"), Mesh.GetValue());
-		}
-
-		if (Materials.Num() > 0)
-		{
-			Writer.Write(TEXT("materials"), Materials);
-		}
-
-		Writer.EndObject();
-	}
+	virtual void WriteObject(IGLTFJsonWriter& Writer) const override;
 };
 
-struct FGLTFJsonEpicVariant : IGLTFJsonObject
+struct GLTFEXPORTER_API FGLTFJsonEpicVariant : IGLTFJsonObject
 {
 	FString Name;
 	bool    bIsActive;
@@ -72,52 +39,23 @@ struct FGLTFJsonEpicVariant : IGLTFJsonObject
 	FGLTFJsonTextureIndex Thumbnail;
 	TMap<FGLTFJsonNodeIndex, FGLTFJsonEpicVariantNodeProperties> Nodes;
 
-	virtual void WriteObject(IGLTFJsonWriter& Writer) const override
-	{
-		Writer.Write(TEXT("name"), Name);
-		Writer.Write(TEXT("active"), bIsActive);
-
-		if (Thumbnail != INDEX_NONE)
-		{
-			Writer.Write(TEXT("thumbnail"), Thumbnail);
-		}
-
-		TArray<FGLTFJsonEpicVariantNodeProperties> NodeValues;
-		Nodes.GenerateValueArray(NodeValues);
-		Writer.Write(TEXT("nodes"), NodeValues);
-	}
+	virtual void WriteObject(IGLTFJsonWriter& Writer) const override;
 };
 
-struct FGLTFJsonEpicVariantSet : IGLTFJsonObject
+struct GLTFEXPORTER_API FGLTFJsonEpicVariantSet : IGLTFJsonObject
 {
 	FString Name;
 
 	TArray<FGLTFJsonEpicVariant> Variants;
 
-	virtual void WriteObject(IGLTFJsonWriter& Writer) const override
-	{
-		if (!Name.IsEmpty())
-		{
-			Writer.Write(TEXT("name"), Name);
-		}
-
-		Writer.Write(TEXT("variants"), Variants);
-	}
+	virtual void WriteObject(IGLTFJsonWriter& Writer) const override;
 };
 
-struct FGLTFJsonEpicLevelVariantSets : IGLTFJsonObject
+struct GLTFEXPORTER_API FGLTFJsonEpicLevelVariantSets : IGLTFJsonObject
 {
 	FString Name;
 
 	TArray<FGLTFJsonEpicVariantSet> VariantSets;
 
-	virtual void WriteObject(IGLTFJsonWriter& Writer) const override
-	{
-		if (!Name.IsEmpty())
-		{
-			Writer.Write(TEXT("name"), Name);
-		}
-
-		Writer.Write(TEXT("variantSets"), VariantSets);
-	}
+	virtual void WriteObject(IGLTFJsonWriter& Writer) const override;
 };
diff --git a/Engine/Plugins/Enterprise/GLTFExporter/MigratedFromGit/Source/GLTFExporter/Public/Json/GLTFJsonExtensions.h b/Engine/Plugins/Enterprise/GLTFExporter/MigratedFromGit/Source/GLTFExporter/Public/Json/GLTFJsonExtensions.h
--- a/Engine/Plugins/Enterprise/GLTFExporter/MigratedFromGit/Source/GLTFExporter/Public/Json/GLTFJsonExtensions.h
+++ b/Engine/Plugins/Enterprise/GLTFExporter/MigratedFromGit/Source/GLTFExporter/Public/Json/GLTFJsonExtensions.h
@@ -4,7 +4,7 @@
 
 #include "Json/GLTFJsonEnums.h"
 
-struct FGLTFJsonExtensions
+struct GLTFEXPORTER_API FGLTFJsonExtensions
 {
 	TSet<EGLTFJsonExtension> Used;
 	TSet<EGLTFJsonExtension> Required;
diff --git a/Engine/Plugins/Enterprise/GLTFExporter/MigratedFromGit/Source/GLTFExporter/Public/Json/GLTFJsonHotspot.h b/Engine/Plugins/Enterprise/GLTFExporter/MigratedFromGit/Source/GLTFExporter/Public/Json/GLTFJsonHotspot.h
--- a/Engine/Plugins/Enterprise/GLTFExporter/MigratedFromGit/Source/GLTFExporter/Public/Json/GLTFJsonHotspot.h
+++ b/Engine/Plugins/Enterprise/GLTFExporter/MigratedFromGit/Source/GLTFExporter/Public/Json/GLTFJsonHotspot.h
@@ -5,7 +5,7 @@
 #include "Json/GLTFJsonObject.h"
 #include "Json/GLTFJsonIndex.h"
 
-struct FGLTFJsonHotspot : IGLTFJsonObject
+struct GLTFEXPORTER_API FGLTFJsonHotspot : IGLTFJsonObject
 {
 	FString Name;
 
@@ -15,29 +15,5 @@ struct FGLTFJsonHotspot : IGLTFJsonObject
 	FGLTFJsonTextureIndex   ToggledImage;
 	FGLTFJsonTextureIndex   ToggledHoveredImage;
 
-	virtual void WriteObject(IGLTFJsonWriter& Writer) const override
-	{
-		if (!Name.IsEmpty())
-		{
-			Writer.Write(TEXT("name"), Name);
-		}
-
-		Writer.Write(TEXT("animation"), Animation);
-		Writer.Write(TEXT("image"), Image);
-
-		if (HoveredImage != INDEX_NONE)
-		{
-			Writer.Write(TEXT("hoveredImage"), HoveredImage);
-		}
-
-		if (ToggledImage != INDEX_NONE)
-		{
-			Writer.Write(TEXT("toggledImage"), ToggledImage);
-		}
-
-		if (ToggledHoveredImage != INDEX_NONE)
-		{
-			Writer.Write(TEXT("toggledHoveredImage"), ToggledHoveredImage);
-		}
-	}
+	virtual void WriteObject(IGLTFJsonWriter& Writer) const override;
 };
diff --git a/Engine/Plugins/Enterprise/GLTFExporter/MigratedFromGit/Source/GLTFExporter/Public/Json/GLTFJsonImage.h b/Engine/Plugins/Enterprise/GLTFExporter/MigratedFromGit/Source/GLTFExporter/Public/Json/GLTFJsonImage.h
--- a/Engine/Plugins/Enterprise/GLTFExporter/MigratedFromGit/Source/GLTFExporter/Public/Json/GLTFJsonImage.h
+++ b/Engine/Plugins/Enterprise/GLTFExporter/MigratedFromGit/Source/GLTFExporter/Public/Json/GLTFJsonImage.h
@@ -6,7 +6,7 @@
 #include "Json/GLTFJsonEnums.h"
 #include "Json/GLTFJsonIndex.h"
 
-struct FGLTFJsonImage : IGLTFJsonObject
+struct GLTFEXPORTER_API FGLTFJsonImage : IGLTFJsonObject
 {
 	FString Name;
 	FString Uri;
@@ -20,26 +20,5 @@ struct FGLTFJsonImage : IGLTFJsonObject
 	{
 	}
 
-	virtual void WriteObject(IGLTFJsonWriter& Writer) const override
-	{
-		if (!Name.IsEmpty())
-		{
-			Writer.Write(TEXT("name"), Name);
-		}
-
-		if (!Uri.IsEmpty())
-		{
-			Writer.Write(TEXT("uri"), Uri);
-		}
-
-		if (MimeType != EGLTFJsonMimeType::None)
-		{
-			Writer.Write(TEXT("mimeType"), MimeType);
-		}
-
-		if (BufferView != INDEX_NONE)
-		{
-			Writer.Write(TEXT("bufferView"), BufferView);
-		}
-	}
+	virtual void WriteObject(IGLTFJsonWriter& Writer) const override;
 };
diff --git a/Engine/Plugins/Enterprise/GLTFExporter/MigratedFromGit/Source/GLTFExporter/Public/Json/GLTFJsonIndex.h b/Engine/Plugins/Enterprise/GLTFExporter/MigratedFromGit/Source/GLTFExporter/Public/Json/GLTFJsonIndex.h
--- a/Engine/Plugins/Enterprise/GLTFExporter/MigratedFromGit/Source/GLTFExporter/Public/Json/GLTFJsonIndex.h
+++ b/Engine/Plugins/Enterprise/GLTFExporter/MigratedFromGit/Source/GLTFExporter/Public/Json/GLTFJsonIndex.h
@@ -5,7 +5,7 @@
 #include "CoreMinimal.h"
 
 template <typename DerivedType>
-struct FGLTFJsonIndex
+struct GLTFEXPORTER_API FGLTFJsonIndex
 {
 	FGLTFJsonIndex()
 		: Value(INDEX_NONE)
diff --git a/Engine/Plugins/Enterprise/GLTFExporter/MigratedFromGit/Source/GLTFExporter/Public/Json/GLTFJsonKhrMaterialVariant.h b/Engine/Plugins/Enterprise/GLTFExporter/MigratedFromGit/Source/GLTFExporter/Public/Json/GLTFJsonKhrMaterialVariant.h
--- a/Engine/Plugins/Enterprise/GLTFExporter/MigratedFromGit/Source/GLTFExporter/Public/Json/GLTFJsonKhrMaterialVariant.h
+++ b/Engine/Plugins/Enterprise/GLTFExporter/MigratedFromGit/Source/GLTFExporter/Public/Json/GLTFJsonKhrMaterialVariant.h
@@ -5,24 +5,17 @@
 #include "Json/GLTFJsonObject.h"
 #include "Json/GLTFJsonIndex.h"
 
-struct FGLTFJsonKhrMaterialVariantMapping : IGLTFJsonObject
+struct GLTFEXPORTER_API FGLTFJsonKhrMaterialVariantMapping : IGLTFJsonObject
 {
 	FGLTFJsonMaterialIndex                   Material;
 	TArray<FGLTFJsonKhrMaterialVariantIndex> Variants;
 
-	virtual void WriteObject(IGLTFJsonWriter& Writer) const override
-	{
-		Writer.Write(TEXT("material"), Material);
-		Writer.Write(TEXT("variants"), Variants);
-	}
+	virtual void WriteObject(IGLTFJsonWriter& Writer) const override;
 };
 
-struct FGLTFJsonKhrMaterialVariant : IGLTFJsonObject
+struct GLTFEXPORTER_API FGLTFJsonKhrMaterialVariant : IGLTFJsonObject
 {
 	FString Name;
 
-	virtual void WriteObject(IGLTFJsonWriter& Writer) const override
-	{
-		Writer.Write(TEXT("name"), Name);
-	}
+	virtual void WriteObject(IGLTFJsonWriter& Writer) const override;
 };
diff --git a/Engine/Plugins/Enterprise/GLTFExporter/MigratedFromGit/Source/GLTFExporter/Public/Json/GLTFJsonLight.h b/Engine/Plugins/Enterprise/GLTFExporter/MigratedFromGit/Source/GLTFExporter/Public/Json/GLTFJsonLight.h
--- a/Engine/Plugins/Enterprise/GLTFExporter/MigratedFromGit/Source/GLTFExporter/Public/Json/GLTFJsonLight.h
+++ b/Engine/Plugins/Enterprise/GLTFExporter/MigratedFromGit/Source/GLTFExporter/Public/Json/GLTFJsonLight.h
@@ -6,7 +6,7 @@
 #include "Json/GLTFJsonEnums.h"
 #include "Json/GLTFJsonColor.h"
 
-struct FGLTFJsonSpotLight : IGLTFJsonObject
+struct GLTFEXPORTER_API FGLTFJsonSpotLight : IGLTFJsonObject
 {
 	float InnerConeAngle;
 	float OuterConeAngle;
@@ -17,21 +17,10 @@ struct FGLTFJsonSpotLight : IGLTFJsonObject
 	{
 	}
 
-	virtual void WriteObject(IGLTFJsonWriter& Writer) const override
-	{
-		if (!FMath::IsNearlyEqual(InnerConeAngle, 0, Writer.DefaultTolerance))
-		{
-			Writer.Write(TEXT("innerConeAngle"), InnerConeAngle);
-		}
-
-		if (!FMath::IsNearlyEqual(OuterConeAngle, HALF_PI, Writer.DefaultTolerance))
-		{
-			Writer.Write(TEXT("outerConeAngle"), OuterConeAngle);
-		}
-	}
+	virtual void WriteObject(IGLTFJsonWriter& Writer) const override;
 };
 
-struct FGLTFJsonLight : IGLTFJsonObject
+struct GLTFEXPORTER_API FGLTFJsonLight : IGLTFJsonObject
 {
 	FString Name;
 
@@ -52,36 +41,5 @@ struct FGLTFJsonLight : IGLTFJsonObject
 	{
 	}
 
-	virtual void WriteObject(IGLTFJsonWriter& Writer) const override
-	{
-		if (!Name.IsEmpty())
-		{
-			Writer.Write(TEXT("name"), Name);
-		}
-
-		Writer.Write(TEXT("type"), Type);
-
-		if (!Color.IsNearlyEqual(FGLTFJsonColor3::White, Writer.DefaultTolerance))
-		{
-			Writer.Write(TEXT("color"), Color);
-		}
-
-		if (!FMath::IsNearlyEqual(Intensity, 1, Writer.DefaultTolerance))
-		{
-			Writer.Write(TEXT("intensity"), Intensity);
-		}
-
-		if (Type == EGLTFJsonLightType::Point || Type == EGLTFJsonLightType::Spot)
-		{
-			if (!FMath::IsNearlyEqual(Range, 0, Writer.DefaultTolerance))
-			{
-				Writer.Write(TEXT("range"), Range);
-			}
-
-			if (Type == EGLTFJsonLightType::Spot)
-			{
-				Writer.Write(TEXT("spot"), Spot);
-			}
-		}
-	}
+	virtual void WriteObject(IGLTFJsonWriter& Writer) const override;
 };
diff --git a/Engine/Plugins/Enterprise/GLTFExporter/MigratedFromGit/Source/GLTFExporter/Public/Json/GLTFJsonLightMap.h b/Engine/Plugins/Enterprise/GLTFExporter/MigratedFromGit/Source/GLTFExporter/Public/Json/GLTFJsonLightMap.h
--- a/Engine/Plugins/Enterprise/GLTFExporter/MigratedFromGit/Source/GLTFExporter/Public/Json/GLTFJsonLightMap.h
+++ b/Engine/Plugins/Enterprise/GLTFExporter/MigratedFromGit/Source/GLTFExporter/Public/Json/GLTFJsonLightMap.h
@@ -6,7 +6,7 @@
 #include "Json/GLTFJsonMaterial.h"
 #include "Json/GLTFJsonVector.h"
 
-struct FGLTFJsonLightMap : IGLTFJsonObject
+struct GLTFEXPORTER_API FGLTFJsonLightMap : IGLTFJsonObject
 {
 	FString              Name;
 	FGLTFJsonTextureInfo Texture;
@@ -21,20 +21,5 @@ struct FGLTFJsonLightMap : IGLTFJsonObject
 	{
 	}
 
-	virtual void WriteObject(IGLTFJsonWriter& Writer) const override
-	{
-		if (!Name.IsEmpty())
-		{
-			Writer.Write(TEXT("name"), Name);
-		}
-
-		if (Texture.Index != INDEX_NONE)
-		{
-			Writer.Write(TEXT("texture"), Texture);
-		}
-
-		Writer.Write(TEXT("lightmapScale"), LightMapScale);
-		Writer.Write(TEXT("lightmapAdd"), LightMapAdd);
-		Writer.Write(TEXT("coordinateScaleBias"), CoordinateScaleBias);
-	}
+	virtual void WriteObject(IGLTFJsonWriter& Writer) const override;
 };
diff --git a/Engine/Plugins/Enterprise/GLTFExporter/MigratedFromGit/Source/GLTFExporter/Public/Json/GLTFJsonMaterial.h b/Engine/Plugins/Enterprise/GLTFExporter/MigratedFromGit/Source/GLTFExporter/Public/Json/GLTFJsonMaterial.h
--- a/Engine/Plugins/Enterprise/GLTFExporter/MigratedFromGit/Source/GLTFExporter/Public/Json/GLTFJsonMaterial.h
+++ b/Engine/Plugins/Enterprise/GLTFExporter/MigratedFromGit/Source/GLTFExporter/Public/Json/GLTFJsonMaterial.h
@@ -4,10 +4,11 @@
 
 #include "Json/GLTFJsonObject.h"
 #include "Json/GLTFJsonEnums.h"
+#include "Json/GLTFJsonIndex.h"
 #include "Json/GLTFJsonColor.h"
 #include "Json/GLTFJsonTextureTransform.h"
 
-struct FGLTFJsonTextureInfo : IGLTFJsonObject
+struct GLTFEXPORTER_API FGLTFJsonTextureInfo : IGLTFJsonObject
 {
 	FGLTFJsonTextureIndex Index;
 	int32 TexCoord;
@@ -19,25 +20,10 @@ struct FGLTFJsonTextureInfo : IGLTFJsonObject
 	{
 	}
 
-	virtual void WriteObject(IGLTFJsonWriter& Writer) const override
-	{
-		Writer.Write(TEXT("index"), Index);
-
-		if (TexCoord != 0)
-		{
-			Writer.Write(TEXT("texCoord"), TexCoord);
-		}
-
-		if (!Transform.IsNearlyDefault(Writer.DefaultTolerance))
-		{
-			Writer.StartExtensions();
-			Writer.Write(EGLTFJsonExtension::KHR_TextureTransform, Transform);
-			Writer.EndExtensions();
-		}
-	}
+	virtual void WriteObject(IGLTFJsonWriter& Writer) const override;
 };
 
-struct FGLTFJsonNormalTextureInfo : FGLTFJsonTextureInfo
+struct GLTFEXPORTER_API FGLTFJsonNormalTextureInfo : FGLTFJsonTextureInfo
 {
 	float Scale;
 
@@ -46,23 +32,10 @@ struct FGLTFJsonNormalTextureInfo : FGLTFJsonTextureInfo
 	{
 	}
 
-	virtual void WriteObject(IGLTFJsonWriter& Writer) const override
-	{
-		Writer.Write(TEXT("index"), Index);
-
-		if (TexCoord != 0)
-		{
-			Writer.Write(TEXT("texCoord"), TexCoord);
-		}
-
-		if (!FMath::IsNearlyEqual(Scale, 1, Writer.DefaultTolerance))
-		{
-			Writer.Write(TEXT("scale"), Scale);
-		}
-	}
+	virtual void WriteObject(IGLTFJsonWriter& Writer) const override;
 };
 
-struct FGLTFJsonOcclusionTextureInfo : FGLTFJsonTextureInfo
+struct GLTFEXPORTER_API FGLTFJsonOcclusionTextureInfo : FGLTFJsonTextureInfo
 {
 	float Strength;
 
@@ -71,23 +44,10 @@ struct FGLTFJsonOcclusionTextureInfo : FGLTFJsonTextureInfo
 	{
 	}
 
-	virtual void WriteObject(IGLTFJsonWriter& Writer) const override
-	{
-		Writer.Write(TEXT("index"), Index);
-
-		if (TexCoord != 0)
-		{
-			Writer.Write(TEXT("texCoord"), TexCoord);
-		}
-
-		if (!FMath::IsNearlyEqual(Strength, 1, Writer.DefaultTolerance))
-		{
-			Writer.Write(TEXT("strength"), Strength);
-		}
-	}
+	virtual void WriteObject(IGLTFJsonWriter& Writer) const override;
 };
 
-struct FGLTFJsonPBRMetallicRoughness : IGLTFJsonObject
+struct GLTFEXPORTER_API FGLTFJsonPBRMetallicRoughness : IGLTFJsonObject
 {
 	FGLTFJsonColor4 BaseColorFactor;
 	FGLTFJsonTextureInfo BaseColorTexture;
@@ -103,36 +63,10 @@ struct FGLTFJsonPBRMetallicRoughness : IGLTFJsonObject
 	{
 	}
 
-	virtual void WriteObject(IGLTFJsonWriter& Writer) const override
-	{
-		if (!BaseColorFactor.IsNearlyEqual(FGLTFJsonColor4::White, Writer.DefaultTolerance))
-		{
-			Writer.Write(TEXT("baseColorFactor"), BaseColorFactor);
-		}
-
-		if (BaseColorTexture.Index != INDEX_NONE)
-		{
-			Writer.Write(TEXT("baseColorTexture"), BaseColorTexture);
-		}
-
-		if (!FMath::IsNearlyEqual(MetallicFactor, 1, Writer.DefaultTolerance))
-		{
-			Writer.Write(TEXT("metallicFactor"), MetallicFactor);
-		}
-
-		if (!FMath::IsNearlyEqual(RoughnessFactor, 1, Writer.DefaultTolerance))
-		{
-			Writer.Write(TEXT("roughnessFactor"), RoughnessFactor);
-		}
-
-		if (MetallicRoughnessTexture.Index != INDEX_NONE)
-		{
-			Writer.Write(TEXT("metallicRoughnessTexture"), MetallicRoughnessTexture);
-		}
-	}
+	virtual void WriteObject(IGLTFJsonWriter& Writer) const override;
 };
 
-struct FGLTFJsonClearCoatExtension : IGLTFJsonObject
+struct GLTFEXPORTER_API FGLTFJsonClearCoatExtension : IGLTFJsonObject
 {
 	float ClearCoatFactor;
 	FGLTFJsonTextureInfo ClearCoatTexture;
@@ -148,36 +82,10 @@ struct FGLTFJsonClearCoatExtension : IGLTFJsonObject
 	{
 	}
 
-	virtual void WriteObject(IGLTFJsonWriter& Writer) const override
-	{
-		if (!FMath::IsNearlyEqual(ClearCoatFactor, 0, Writer.DefaultTolerance))
-		{
-			Writer.Write(TEXT("clearcoatFactor"), ClearCoatFactor);
-		}
-
-		if (ClearCoatTexture.Index != INDEX_NONE)
-		{
-			Writer.Write(TEXT("clearcoatTexture"), ClearCoatTexture);
-		}
-
-		if (!FMath::IsNearlyEqual(ClearCoatRoughnessFactor, 0, Writer.DefaultTolerance))
-		{
-			Writer.Write(TEXT("clearcoatRoughnessFactor"), ClearCoatRoughnessFactor);
-		}
-
-		if (ClearCoatRoughnessTexture.Index != INDEX_NONE)
-		{
-			Writer.Write(TEXT("clearcoatRoughnessTexture"), ClearCoatRoughnessTexture);
-		}
-
-		if (ClearCoatNormalTexture.Index != INDEX_NONE)
-		{
-			Writer.Write(TEXT("clearcoatNormalTexture"), ClearCoatNormalTexture);
-		}
-	}
+	virtual void WriteObject(IGLTFJsonWriter& Writer) const override;
 };
 
-struct FGLTFJsonMaterial : IGLTFJsonObject
+struct GLTFEXPORTER_API FGLTFJsonMaterial : IGLTFJsonObject
 {
 	FString Name;
 
@@ -210,76 +118,5 @@ struct FGLTFJsonMaterial : IGLTFJsonObject
 	{
 	}
 
-	virtual void WriteObject(IGLTFJsonWriter& Writer) const override
-	{
-		if (!Name.IsEmpty())
-		{
-			Writer.Write(TEXT("name"), Name);
-		}
-
-		if (ShadingModel != EGLTFJsonShadingModel::None)
-		{
-			Writer.Write(TEXT("pbrMetallicRoughness"), PBRMetallicRoughness);
-		}
-
-		if (NormalTexture.Index != INDEX_NONE)
-		{
-			Writer.Write(TEXT("normalTexture"), NormalTexture);
-		}
-
-		if (OcclusionTexture.Index != INDEX_NONE)
-		{
-			Writer.Write(TEXT("occlusionTexture"), OcclusionTexture);
-		}
-
-		if (EmissiveTexture.Index != INDEX_NONE)
-		{
-			Writer.Write(TEXT("emissiveTexture"), EmissiveTexture);
-		}
-
-		if (!EmissiveFactor.IsNearlyEqual(FGLTFJsonColor3::Black, Writer.DefaultTolerance))
-		{
-			Writer.Write(TEXT("emissiveFactor"), EmissiveFactor);
-		}
-
-		if (AlphaMode != EGLTFJsonAlphaMode::Opaque)
-		{
-			Writer.Write(TEXT("alphaMode"), AlphaMode);
-
-			if (AlphaMode == EGLTFJsonAlphaMode::Mask && !FMath::IsNearlyEqual(AlphaCutoff, 0.5f, Writer.DefaultTolerance))
-			{
-				Writer.Write(TEXT("alphaCutoff"), AlphaCutoff);
-			}
-		}
-
-		if (DoubleSided)
-		{
-			Writer.Write(TEXT("doubleSided"), DoubleSided);
-		}
-
-		if (BlendMode != EGLTFJsonBlendMode::None || ShadingModel == EGLTFJsonShadingModel::Unlit || ShadingModel == EGLTFJsonShadingModel::ClearCoat)
-		{
-			Writer.StartExtensions();
-
-			if (BlendMode != EGLTFJsonBlendMode::None)
-			{
-				Writer.StartExtension(EGLTFJsonExtension::EPIC_BlendModes);
-				Writer.Write(TEXT("blendMode"), BlendMode);
-				Writer.EndExtension();
-			}
-
-			if (ShadingModel == EGLTFJsonShadingModel::Unlit)
-			{
-				Writer.StartExtension(EGLTFJsonExtension::KHR_MaterialsUnlit);
-				// Write empty object
-				Writer.EndExtension();
-			}
-			else if (ShadingModel == EGLTFJsonShadingModel::ClearCoat)
-			{
-				Writer.Write(EGLTFJsonExtension::KHR_MaterialsClearCoat, ClearCoat);
-			}
-
-			Writer.EndExtensions();
-		}
-	}
+	virtual void WriteObject(IGLTFJsonWriter& Writer) const override;
 };
diff --git a/Engine/Plugins/Enterprise/GLTFExporter/MigratedFromGit/Source/GLTFExporter/Public/Json/GLTFJsonMatrix.h b/Engine/Plugins/Enterprise/GLTFExporter/MigratedFromGit/Source/GLTFExporter/Public/Json/GLTFJsonMatrix.h
--- a/Engine/Plugins/Enterprise/GLTFExporter/MigratedFromGit/Source/GLTFExporter/Public/Json/GLTFJsonMatrix.h
+++ b/Engine/Plugins/Enterprise/GLTFExporter/MigratedFromGit/Source/GLTFExporter/Public/Json/GLTFJsonMatrix.h
@@ -6,7 +6,7 @@
 #include "Json/GLTFJsonArray.h"
 
 template <typename BaseType>
-struct TGLTFJsonMatrix : BaseType, IGLTFJsonArray
+struct GLTFEXPORTER_API TGLTFJsonMatrix : BaseType, IGLTFJsonArray
 {
 	static const TGLTFJsonMatrix Identity;
 
@@ -61,7 +61,7 @@ struct TGLTFJsonMatrix : BaseType, IGLTFJsonArray
 	}
 };
 
-struct FGLTFJsonMatrix2 : TGLTFJsonMatrix<FGLTFMatrix2>
+struct GLTFEXPORTER_API FGLTFJsonMatrix2 : TGLTFJsonMatrix<FGLTFMatrix2>
 {
 	static const FGLTFJsonMatrix2 Identity;
 
@@ -69,7 +69,7 @@ struct FGLTFJsonMatrix2 : TGLTFJsonMatrix<FGLTFMatrix2>
 	using TGLTFJsonMatrix::operator=;
 };
 
-struct FGLTFJsonMatrix3 : TGLTFJsonMatrix<FGLTFMatrix3>
+struct GLTFEXPORTER_API FGLTFJsonMatrix3 : TGLTFJsonMatrix<FGLTFMatrix3>
 {
 	static const FGLTFJsonMatrix3 Identity;
 
@@ -77,7 +77,7 @@ struct FGLTFJsonMatrix3 : TGLTFJsonMatrix<FGLTFMatrix3>
 	using TGLTFJsonMatrix::operator=;
 };
 
-struct FGLTFJsonMatrix4 : TGLTFJsonMatrix<FGLTFMatrix4>
+struct GLTFEXPORTER_API FGLTFJsonMatrix4 : TGLTFJsonMatrix<FGLTFMatrix4>
 {
 	static const FGLTFJsonMatrix4 Identity;
 
diff --git a/Engine/Plugins/Enterprise/GLTFExporter/MigratedFromGit/Source/GLTFExporter/Public/Json/GLTFJsonMesh.h b/Engine/Plugins/Enterprise/GLTFExporter/MigratedFromGit/Source/GLTFExporter/Public/Json/GLTFJsonMesh.h
--- a/Engine/Plugins/Enterprise/GLTFExporter/MigratedFromGit/Source/GLTFExporter/Public/Json/GLTFJsonMesh.h
+++ b/Engine/Plugins/Enterprise/GLTFExporter/MigratedFromGit/Source/GLTFExporter/Public/Json/GLTFJsonMesh.h
@@ -7,7 +7,7 @@
 #include "Json/GLTFJsonIndex.h"
 #include "Json/GLTFJsonKhrMaterialVariant.h"
 
-struct FGLTFJsonAttributes : IGLTFJsonObject
+struct GLTFEXPORTER_API FGLTFJsonAttributes : IGLTFJsonObject
 {
 	FGLTFJsonAccessorIndex Position;
 	FGLTFJsonAccessorIndex Color0;
@@ -18,34 +18,10 @@ struct FGLTFJsonAttributes : IGLTFJsonObject
 	TArray<FGLTFJsonAccessorIndex> Joints;
 	TArray<FGLTFJsonAccessorIndex> Weights;
 
-	virtual void WriteObject(IGLTFJsonWriter& Writer) const override
-	{
-		if (Position != INDEX_NONE) Writer.Write(TEXT("POSITION"), Position);
-		if (Color0 != INDEX_NONE) Writer.Write(TEXT("COLOR_0"), Color0);
-		if (Normal != INDEX_NONE) Writer.Write(TEXT("NORMAL"), Normal);
-		if (Tangent != INDEX_NONE) Writer.Write(TEXT("TANGENT"), Tangent);
-
-		for (int32 Index = 0; Index < TexCoords.Num(); ++Index)
-		{
-			const FGLTFJsonAccessorIndex TexCoord = TexCoords[Index];
-			if (TexCoord != INDEX_NONE) Writer.Write(TEXT("TEXCOORD_") + FString::FromInt(Index), TexCoord);
-		}
-
-		for (int32 Index = 0; Index < Joints.Num(); ++Index)
-		{
-			const FGLTFJsonAccessorIndex Joint = Joints[Index];
-			if (Joint != INDEX_NONE) Writer.Write(TEXT("JOINTS_") + FString::FromInt(Index), Joint);
-		}
-
-		for (int32 Index = 0; Index < Weights.Num(); ++Index)
-		{
-			const FGLTFJsonAccessorIndex Weight = Weights[Index];
-			if (Weight != INDEX_NONE) Writer.Write(TEXT("WEIGHTS_") + FString::FromInt(Index), Weight);
-		}
-	}
+	virtual void WriteObject(IGLTFJsonWriter& Writer) const override;
 };
 
-struct FGLTFJsonPrimitive : IGLTFJsonObject
+struct GLTFEXPORTER_API FGLTFJsonPrimitive : IGLTFJsonObject
 {
 	FGLTFJsonAccessorIndex Indices;
 	FGLTFJsonMaterialIndex Material;
@@ -59,51 +35,14 @@ struct FGLTFJsonPrimitive : IGLTFJsonObject
 	{
 	}
 
-	virtual void WriteObject(IGLTFJsonWriter& Writer) const override
-	{
-		Writer.Write(TEXT("attributes"), Attributes);
-
-		if (Indices != INDEX_NONE)
-		{
-			Writer.Write(TEXT("indices"), Indices);
-		}
-
-		if (Material != INDEX_NONE)
-		{
-			Writer.Write(TEXT("material"), Material);
-		}
-
-		if (Mode != EGLTFJsonPrimitiveMode::Triangles)
-		{
-			Writer.Write(TEXT("mode"), Mode);
-		}
-
-		if (KhrMaterialVariantMappings.Num() > 0)
-		{
-			Writer.StartExtensions();
-
-			Writer.StartExtension(EGLTFJsonExtension::KHR_MaterialsVariants);
-			Writer.Write(TEXT("mappings"), KhrMaterialVariantMappings);
-			Writer.EndExtension();
-
-			Writer.EndExtensions();
-		}
-	}
+	virtual void WriteObject(IGLTFJsonWriter& Writer) const override;
 };
 
-struct FGLTFJsonMesh : IGLTFJsonObject
+struct GLTFEXPORTER_API FGLTFJsonMesh : IGLTFJsonObject
 {
 	FString Name;
 
 	TArray<FGLTFJsonPrimitive> Primitives;
 
-	virtual void WriteObject(IGLTFJsonWriter& Writer) const override
-	{
-		if (!Name.IsEmpty())
-		{
-			Writer.Write(TEXT("name"), Name);
-		}
-
-		Writer.Write(TEXT("primitives"), Primitives);
-	}
+	virtual void WriteObject(IGLTFJsonWriter& Writer) const override;
 };
diff --git a/Engine/Plugins/Enterprise/GLTFExporter/MigratedFromGit/Source/GLTFExporter/Public/Json/GLTFJsonNode.h b/Engine/Plugins/Enterprise/GLTFExporter/MigratedFromGit/Source/GLTFExporter/Public/Json/GLTFJsonNode.h
--- a/Engine/Plugins/Enterprise/GLTFExporter/MigratedFromGit/Source/GLTFExporter/Public/Json/GLTFJsonNode.h
+++ b/Engine/Plugins/Enterprise/GLTFExporter/MigratedFromGit/Source/GLTFExporter/Public/Json/GLTFJsonNode.h
@@ -7,7 +7,7 @@
 #include "Json/GLTFJsonVector.h"
 #include "Json/GLTFJsonQuaternion.h"
 
-struct FGLTFJsonNode : IGLTFJsonObject
+struct GLTFEXPORTER_API FGLTFJsonNode : IGLTFJsonObject
 {
 	FString Name;
 
@@ -35,88 +35,5 @@ struct FGLTFJsonNode : IGLTFJsonObject
 	{
 	}
 
-	virtual void WriteObject(IGLTFJsonWriter& Writer) const override
-	{
-		if (!Name.IsEmpty())
-		{
-			Writer.Write(TEXT("name"), Name);
-		}
-
-		if (!Translation.IsNearlyEqual(FGLTFJsonVector3::Zero, Writer.DefaultTolerance))
-		{
-			Writer.Write(TEXT("translation"), Translation);
-		}
-
-		if (!Rotation.IsNearlyEqual(FGLTFJsonQuaternion::Identity, Writer.DefaultTolerance))
-		{
-			Writer.Write(TEXT("rotation"), Rotation);
-		}
-
-		if (!Scale.IsNearlyEqual(FGLTFJsonVector3::One, Writer.DefaultTolerance))
-		{
-			Writer.Write(TEXT("scale"), Scale);
-		}
-
-		if (Camera != INDEX_NONE)
-		{
-			Writer.Write(TEXT("camera"), Camera);
-		}
-
-		if (Skin != INDEX_NONE)
-		{
-			Writer.Write(TEXT("skin"), Skin);
-		}
-
-		if (Mesh != INDEX_NONE)
-		{
-			Writer.Write(TEXT("mesh"), Mesh);
-		}
-
-		if (Backdrop != INDEX_NONE || Hotspot != INDEX_NONE || Light != INDEX_NONE || LightMap != INDEX_NONE || SkySphere != INDEX_NONE)
-		{
-			Writer.StartExtensions();
-
-			if (Backdrop != INDEX_NONE)
-			{
-				Writer.StartExtension(EGLTFJsonExtension::EPIC_HDRIBackdrops);
-				Writer.Write(TEXT("backdrop"), Backdrop);
-				Writer.EndExtension();
-			}
-
-			if (Hotspot != INDEX_NONE)
-			{
-				Writer.StartExtension(EGLTFJsonExtension::EPIC_AnimationHotspots);
-				Writer.Write(TEXT("hotspot"), Hotspot);
-				Writer.EndExtension();
-			}
-
-			if (Light != INDEX_NONE)
-			{
-				Writer.StartExtension(EGLTFJsonExtension::KHR_LightsPunctual);
-				Writer.Write(TEXT("light"), Light);
-				Writer.EndExtension();
-			}
-
-			if (LightMap != INDEX_NONE)
-			{
-				Writer.StartExtension(EGLTFJsonExtension::EPIC_LightmapTextures);
-				Writer.Write(TEXT("lightmap"), LightMap);
-				Writer.EndExtension();
-			}
-
-			if (SkySphere != INDEX_NONE)
-			{
-				Writer.StartExtension(EGLTFJsonExtension::EPIC_SkySpheres);
-				Writer.Write(TEXT("skySphere"), SkySphere);
-				Writer.EndExtension();
-			}
-
-			Writer.EndExtensions();
-		}
-
-		if (Children.Num() > 0)
-		{
-			Writer.Write(TEXT("children"), Children);
-		}
-	}
+	virtual void WriteObject(IGLTFJsonWriter& Writer) const override;
 };
diff --git a/Engine/Plugins/Enterprise/GLTFExporter/MigratedFromGit/Source/GLTFExporter/Public/Json/GLTFJsonObject.h b/Engine/Plugins/Enterprise/GLTFExporter/MigratedFromGit/Source/GLTFExporter/Public/Json/GLTFJsonObject.h
--- a/Engine/Plugins/Enterprise/GLTFExporter/MigratedFromGit/Source/GLTFExporter/Public/Json/GLTFJsonObject.h
+++ b/Engine/Plugins/Enterprise/GLTFExporter/MigratedFromGit/Source/GLTFExporter/Public/Json/GLTFJsonObject.h
@@ -5,7 +5,7 @@
 #include "Json/GLTFJsonValue.h"
 #include "Json/GLTFJsonWriter.h"
 
-struct IGLTFJsonObject : IGLTFJsonValue
+struct GLTFEXPORTER_API IGLTFJsonObject : IGLTFJsonValue
 {
 	virtual void WriteValue(IGLTFJsonWriter& Writer) const override final
 	{
diff --git a/Engine/Plugins/Enterprise/GLTFExporter/MigratedFromGit/Source/GLTFExporter/Public/Json/GLTFJsonQuaternion.h b/Engine/Plugins/Enterprise/GLTFExporter/MigratedFromGit/Source/GLTFExporter/Public/Json/GLTFJsonQuaternion.h
--- a/Engine/Plugins/Enterprise/GLTFExporter/MigratedFromGit/Source/GLTFExporter/Public/Json/GLTFJsonQuaternion.h
+++ b/Engine/Plugins/Enterprise/GLTFExporter/MigratedFromGit/Source/GLTFExporter/Public/Json/GLTFJsonQuaternion.h
@@ -5,7 +5,7 @@
 #include "Core/GLTFQuaternion.h"
 #include "Json/GLTFJsonArray.h"
 
-struct FGLTFJsonQuaternion : TGLTFQuaternion<float>, IGLTFJsonArray
+struct GLTFEXPORTER_API FGLTFJsonQuaternion : TGLTFQuaternion<float>, IGLTFJsonArray
 {
 	static const TGLTFQuaternion Identity;
 
diff --git a/Engine/Plugins/Enterprise/GLTFExporter/MigratedFromGit/Source/GLTFExporter/Public/Json/GLTFJsonRoot.h b/Engine/Plugins/Enterprise/GLTFExporter/MigratedFromGit/Source/GLTFExporter/Public/Json/GLTFJsonRoot.h
--- a/Engine/Plugins/Enterprise/GLTFExporter/MigratedFromGit/Source/GLTFExporter/Public/Json/GLTFJsonRoot.h
+++ b/Engine/Plugins/Enterprise/GLTFExporter/MigratedFromGit/Source/GLTFExporter/Public/Json/GLTFJsonRoot.h
@@ -27,7 +27,7 @@
 #include "Json/GLTFJsonKhrMaterialVariant.h"
 #include "Json/GLTFJsonWriter.h"
 
-struct FGLTFJsonRoot : IGLTFJsonObject
+struct GLTFEXPORTER_API FGLTFJsonRoot : IGLTFJsonObject
 {
 	FGLTFJsonAsset Asset;
 
@@ -56,101 +56,14 @@ struct FGLTFJsonRoot : IGLTFJsonObject
 	TArray<TUniquePtr<FGLTFJsonEpicLevelVariantSets>>  EpicLevelVariantSets;
 	TArray<TUniquePtr<FGLTFJsonKhrMaterialVariant>>    KhrMaterialVariants;
 
-	virtual void WriteObject(IGLTFJsonWriter& Writer) const override
-	{
-		Writer.Write(TEXT("asset"), Asset);
+	FGLTFJsonRoot() = default;
+	FGLTFJsonRoot(FGLTFJsonRoot&&) = default;
+	FGLTFJsonRoot &operator=(FGLTFJsonRoot&&) = default;
 
-		if (DefaultScene != INDEX_NONE)
-		{
-			Writer.Write(TEXT("scene"), DefaultScene);
-		}
+	FGLTFJsonRoot(const FGLTFJsonRoot&) = delete;
+	FGLTFJsonRoot &operator=(const FGLTFJsonRoot&) = delete;
 
-		if (Accessors.Num() > 0) Writer.Write(TEXT("accessors"), Accessors);
-		if (Animations.Num() > 0) Writer.Write(TEXT("animations"), Animations);
-		if (Buffers.Num() > 0) Writer.Write(TEXT("buffers"), Buffers);
-		if (BufferViews.Num() > 0) Writer.Write(TEXT("bufferViews"), BufferViews);
-		if (Cameras.Num() > 0) Writer.Write(TEXT("cameras"), Cameras);
-		if (Images.Num() > 0) Writer.Write(TEXT("images"), Images);
-		if (Materials.Num() > 0) Writer.Write(TEXT("materials"), Materials);
-		if (Meshes.Num() > 0) Writer.Write(TEXT("meshes"), Meshes);
-		if (Nodes.Num() > 0) Writer.Write(TEXT("nodes"), Nodes);
-		if (Samplers.Num() > 0) Writer.Write(TEXT("samplers"), Samplers);
-		if (Scenes.Num() > 0) Writer.Write(TEXT("scenes"), Scenes);
-		if (Skins.Num() > 0) Writer.Write(TEXT("skins"), Skins);
-		if (Textures.Num() > 0) Writer.Write(TEXT("textures"), Textures);
+	virtual void WriteObject(IGLTFJsonWriter& Writer) const override;
 
-		if (Backdrops.Num() > 0 || Hotspots.Num() > 0 || Lights.Num() > 0 || LightMaps.Num() > 0 || SkySpheres.Num() > 0 || EpicLevelVariantSets.Num() > 0 || KhrMaterialVariants.Num() > 0)
-		{
-			Writer.StartExtensions();
-
-			if (Backdrops.Num() > 0)
-			{
-				Writer.StartExtension(EGLTFJsonExtension::EPIC_HDRIBackdrops);
-				Writer.Write(TEXT("backdrops"), Backdrops);
-				Writer.EndExtension();
-			}
-
-			if (Hotspots.Num() > 0)
-			{
-				Writer.StartExtension(EGLTFJsonExtension::EPIC_AnimationHotspots);
-				Writer.Write(TEXT("hotspots"), Hotspots);
-				Writer.EndExtension();
-			}
-
-			if (EpicLevelVariantSets.Num() > 0)
-			{
-				Writer.StartExtension(EGLTFJsonExtension::EPIC_LevelVariantSets);
-				Writer.Write(TEXT("levelVariantSets"), EpicLevelVariantSets);
-				Writer.EndExtension();
-			}
-
-			if (KhrMaterialVariants.Num() > 0)
-			{
-				Writer.StartExtension(EGLTFJsonExtension::KHR_MaterialsVariants);
-				Writer.Write(TEXT("variants"), KhrMaterialVariants);
-				Writer.EndExtension();
-			}
-
-			if (Lights.Num() > 0)
-			{
-				Writer.StartExtension(EGLTFJsonExtension::KHR_LightsPunctual);
-				Writer.Write(TEXT("lights"), Lights);
-				Writer.EndExtension();
-			}
-
-			if (LightMaps.Num() > 0)
-			{
-				Writer.StartExtension(EGLTFJsonExtension::EPIC_LightmapTextures);
-				Writer.Write(TEXT("lightmaps"), LightMaps);
-				Writer.EndExtension();
-			}
-
-			if (SkySpheres.Num() > 0)
-			{
-				Writer.StartExtension(EGLTFJsonExtension::EPIC_SkySpheres);
-				Writer.Write(TEXT("skySpheres"), SkySpheres);
-				Writer.EndExtension();
-			}
-
-			Writer.EndExtensions();
-		}
-
-		if (Extensions.Used.Num() > 0)
-		{
-			Writer.Write(TEXT("extensionsUsed"), Extensions.Used);
-		}
-
-		if (Extensions.Required.Num() > 0)
-		{
-			Writer.Write(TEXT("extensionsRequired"), Extensions.Required);
-		}
-	}
-
-	void WriteJson(FArchive& Archive, bool bPrettyJson, float DefaultTolerance)
-	{
-		TSharedRef<IGLTFJsonWriter> Writer = IGLTFJsonWriter::Create(Archive, bPrettyJson, Extensions);
-		Writer->DefaultTolerance = DefaultTolerance;
-		Writer->Write(this);
-		Writer->Close();
-	}
+	void WriteJson(FArchive& Archive, bool bPrettyJson, float DefaultTolerance);
 };
diff --git a/Engine/Plugins/Enterprise/GLTFExporter/MigratedFromGit/Source/GLTFExporter/Public/Json/GLTFJsonSampler.h b/Engine/Plugins/Enterprise/GLTFExporter/MigratedFromGit/Source/GLTFExporter/Public/Json/GLTFJsonSampler.h
--- a/Engine/Plugins/Enterprise/GLTFExporter/MigratedFromGit/Source/GLTFExporter/Public/Json/GLTFJsonSampler.h
+++ b/Engine/Plugins/Enterprise/GLTFExporter/MigratedFromGit/Source/GLTFExporter/Public/Json/GLTFJsonSampler.h
@@ -5,7 +5,7 @@
 #include "Json/GLTFJsonObject.h"
 #include "Json/GLTFJsonEnums.h"
 
-struct FGLTFJsonSampler : IGLTFJsonObject
+struct GLTFEXPORTER_API FGLTFJsonSampler : IGLTFJsonObject
 {
 	FString Name;
 
@@ -23,31 +23,5 @@ struct FGLTFJsonSampler : IGLTFJsonObject
 	{
 	}
 
-	virtual void WriteObject(IGLTFJsonWriter& Writer) const override
-	{
-		if (!Name.IsEmpty())
-		{
-			Writer.Write(TEXT("name"), Name);
-		}
-
-		if (MinFilter != EGLTFJsonTextureFilter::None)
-		{
-			Writer.Write(TEXT("minFilter"), MinFilter);
-		}
-
-		if (MagFilter != EGLTFJsonTextureFilter::None)
-		{
-			Writer.Write(TEXT("magFilter"), MagFilter);
-		}
-
-		if (WrapS != EGLTFJsonTextureWrap::Repeat)
-		{
-			Writer.Write(TEXT("wrapS"), WrapS);
-		}
-
-		if (WrapT != EGLTFJsonTextureWrap::Repeat)
-		{
-			Writer.Write(TEXT("wrapT"), WrapT);
-		}
-	}
+	virtual void WriteObject(IGLTFJsonWriter& Writer) const override;
 };
diff --git a/Engine/Plugins/Enterprise/GLTFExporter/MigratedFromGit/Source/GLTFExporter/Public/Json/GLTFJsonScene.h b/Engine/Plugins/Enterprise/GLTFExporter/MigratedFromGit/Source/GLTFExporter/Public/Json/GLTFJsonScene.h
--- a/Engine/Plugins/Enterprise/GLTFExporter/MigratedFromGit/Source/GLTFExporter/Public/Json/GLTFJsonScene.h
+++ b/Engine/Plugins/Enterprise/GLTFExporter/MigratedFromGit/Source/GLTFExporter/Public/Json/GLTFJsonScene.h
@@ -5,34 +5,12 @@
 #include "Json/GLTFJsonObject.h"
 #include "Json/GLTFJsonIndex.h"
 
-struct FGLTFJsonScene : IGLTFJsonObject
+struct GLTFEXPORTER_API FGLTFJsonScene : IGLTFJsonObject
 {
 	FString Name;
 
 	TArray<FGLTFJsonNodeIndex> Nodes;
 	TArray<FGLTFJsonEpicLevelVariantSetsIndex>  EpicLevelVariantSets;
 
-	virtual void WriteObject(IGLTFJsonWriter& Writer) const override
-	{
-		if (!Name.IsEmpty())
-		{
-			Writer.Write(TEXT("name"), Name);
-		}
-
-		if (Nodes.Num() > 0)
-		{
-			Writer.Write(TEXT("nodes"), Nodes);
-		}
-
-		if (EpicLevelVariantSets.Num() > 0)
-		{
-			Writer.StartExtensions();
-
-			Writer.StartExtension(EGLTFJsonExtension::EPIC_LevelVariantSets);
-			Writer.Write(TEXT("levelVariantSets"), EpicLevelVariantSets);
-			Writer.EndExtension();
-
-			Writer.EndExtensions();
-		}
-	}
+	virtual void WriteObject(IGLTFJsonWriter& Writer) const override;
 };
diff --git a/Engine/Plugins/Enterprise/GLTFExporter/MigratedFromGit/Source/GLTFExporter/Public/Json/GLTFJsonSkin.h b/Engine/Plugins/Enterprise/GLTFExporter/MigratedFromGit/Source/GLTFExporter/Public/Json/GLTFJsonSkin.h
--- a/Engine/Plugins/Enterprise/GLTFExporter/MigratedFromGit/Source/GLTFExporter/Public/Json/GLTFJsonSkin.h
+++ b/Engine/Plugins/Enterprise/GLTFExporter/MigratedFromGit/Source/GLTFExporter/Public/Json/GLTFJsonSkin.h
@@ -5,7 +5,7 @@
 #include "Json/GLTFJsonObject.h"
 #include "Json/GLTFJsonIndex.h"
 
-struct FGLTFJsonSkin : IGLTFJsonObject
+struct GLTFEXPORTER_API FGLTFJsonSkin : IGLTFJsonObject
 {
 	FString Name;
 
@@ -14,26 +14,5 @@ struct FGLTFJsonSkin : IGLTFJsonObject
 
 	TArray<FGLTFJsonNodeIndex> Joints;
 
-	virtual void WriteObject(IGLTFJsonWriter& Writer) const override
-	{
-		if (!Name.IsEmpty())
-		{
-			Writer.Write(TEXT("name"), Name);
-		}
-
-		if (InverseBindMatrices != INDEX_NONE)
-		{
-			Writer.Write(TEXT("inverseBindMatrices"), InverseBindMatrices);
-		}
-
-		if (Skeleton != INDEX_NONE)
-		{
-			Writer.Write(TEXT("skeleton"), Skeleton);
-		}
-
-		if (Joints.Num() > 0)
-		{
-			Writer.Write(TEXT("joints"), Joints);
-		}
-	}
+	virtual void WriteObject(IGLTFJsonWriter& Writer) const override;
 };
diff --git a/Engine/Plugins/Enterprise/GLTFExporter/MigratedFromGit/Source/GLTFExporter/Public/Json/GLTFJsonSkySphere.h b/Engine/Plugins/Enterprise/GLTFExporter/MigratedFromGit/Source/GLTFExporter/Public/Json/GLTFJsonSkySphere.h
--- a/Engine/Plugins/Enterprise/GLTFExporter/MigratedFromGit/Source/GLTFExporter/Public/Json/GLTFJsonSkySphere.h
+++ b/Engine/Plugins/Enterprise/GLTFExporter/MigratedFromGit/Source/GLTFExporter/Public/Json/GLTFJsonSkySphere.h
@@ -5,8 +5,9 @@
 #include "Json/GLTFJsonObject.h"
 #include "Json/GLTFJsonIndex.h"
 #include "Json/GLTFJsonColor.h"
+#include "Json/GLTFJsonVector.h"
 
-struct FGLTFJsonSkySphereColorCurve : IGLTFJsonArray
+struct GLTFEXPORTER_API FGLTFJsonSkySphereColorCurve : IGLTFJsonArray
 {
 	struct FKey
 	{
@@ -21,26 +22,10 @@ struct FGLTFJsonSkySphereColorCurve : IGLTFJsonArray
 
 	TArray<FComponentCurve> ComponentCurves;
 
-	virtual void WriteArray(IGLTFJsonWriter& Writer) const override
-	{
-		for (int32 ComponentIndex = 0; ComponentIndex < ComponentCurves.Num(); ++ComponentIndex)
-		{
-			const FComponentCurve& ComponentCurve = ComponentCurves[ComponentIndex];
-
-			Writer.StartArray();
-
-			for (const FKey& Key: ComponentCurve.Keys)
-			{
-				Writer.Write(Key.Time);
-				Writer.Write(Key.Value);
-			}
-
-			Writer.EndArray();
-		}
-	}
+	virtual void WriteArray(IGLTFJsonWriter& Writer) const override;
 };
 
-struct FGLTFJsonSkySphere : IGLTFJsonObject
+struct GLTFEXPORTER_API FGLTFJsonSkySphere : IGLTFJsonObject
 {
 	FString Name;
 
@@ -93,63 +78,5 @@ struct FGLTFJsonSkySphere : IGLTFJsonObject
 	{
 	}
 
-	virtual void WriteObject(IGLTFJsonWriter& Writer) const override
-	{
-		if (!Name.IsEmpty())
-		{
-			Writer.Write(TEXT("name"), Name);
-		}
-
-		Writer.Write(TEXT("skySphereMesh"), SkySphereMesh);
-		Writer.Write(TEXT("skyTexture"), SkyTexture);
-		Writer.Write(TEXT("cloudsTexture"), CloudsTexture);
-		Writer.Write(TEXT("starsTexture"), StarsTexture);
-
-		if (DirectionalLight != INDEX_NONE)
-		{
-			Writer.Write(TEXT("directionalLight"), DirectionalLight);
-		}
-
-		Writer.Write(TEXT("sunHeight"), SunHeight);
-		Writer.Write(TEXT("sunBrightness"), SunBrightness);
-		Writer.Write(TEXT("starsBrightness"), StarsBrightness);
-		Writer.Write(TEXT("cloudSpeed"), CloudSpeed);
-		Writer.Write(TEXT("cloudOpacity"), CloudOpacity);
-		Writer.Write(TEXT("horizonFalloff"), HorizonFalloff);
-
-		Writer.Write(TEXT("sunRadius"), SunRadius);
-		Writer.Write(TEXT("noisePower1"), NoisePower1);
-		Writer.Write(TEXT("noisePower2"), NoisePower2);
-
-		Writer.Write(TEXT("colorsDeterminedBySunPosition"), bColorsDeterminedBySunPosition);
-
-		Writer.Write(TEXT("zenithColor"), ZenithColor);
-		Writer.Write(TEXT("horizonColor"), HorizonColor);
-		Writer.Write(TEXT("cloudColor"), CloudColor);
-
-		if (!OverallColor.IsNearlyEqual(FGLTFJsonColor4::White, Writer.DefaultTolerance))
-		{
-			Writer.Write(TEXT("overallColor"), OverallColor);
-		}
-
-		if (ZenithColorCurve.ComponentCurves.Num() >= 3)
-		{
-			Writer.Write(TEXT("zenithColorCurve"), ZenithColorCurve);
-		}
-
-		if (HorizonColorCurve.ComponentCurves.Num() >= 3)
-		{
-			Writer.Write(TEXT("horizonColorCurve"), HorizonColorCurve);
-		}
-
-		if (CloudColorCurve.ComponentCurves.Num() >= 3)
-		{
-			Writer.Write(TEXT("cloudColorCurve"), CloudColorCurve);
-		}
-
-		if (!Scale.IsNearlyEqual(FGLTFJsonVector3::One, Writer.DefaultTolerance))
-		{
-			Writer.Write(TEXT("scale"), Scale);
-		}
-	}
+	virtual void WriteObject(IGLTFJsonWriter& Writer) const override;
 };
diff --git a/Engine/Plugins/Enterprise/GLTFExporter/MigratedFromGit/Source/GLTFExporter/Public/Json/GLTFJsonTexture.h b/Engine/Plugins/Enterprise/GLTFExporter/MigratedFromGit/Source/GLTFExporter/Public/Json/GLTFJsonTexture.h
--- a/Engine/Plugins/Enterprise/GLTFExporter/MigratedFromGit/Source/GLTFExporter/Public/Json/GLTFJsonTexture.h
+++ b/Engine/Plugins/Enterprise/GLTFExporter/MigratedFromGit/Source/GLTFExporter/Public/Json/GLTFJsonTexture.h
@@ -5,7 +5,7 @@
 #include "Json/GLTFJsonObject.h"
 #include "Json/GLTFJsonIndex.h"
 
-struct FGLTFJsonTexture : IGLTFJsonObject
+struct GLTFEXPORTER_API FGLTFJsonTexture : IGLTFJsonObject
 {
 	FString Name;
 
@@ -20,32 +20,5 @@ struct FGLTFJsonTexture : IGLTFJsonObject
 	{
 	}
 
-	virtual void WriteObject(IGLTFJsonWriter& Writer) const override
-	{
-		if (!Name.IsEmpty())
-		{
-			Writer.Write(TEXT("name"), Name);
-		}
-
-		if (Sampler != INDEX_NONE)
-		{
-			Writer.Write(TEXT("sampler"), Sampler);
-		}
-
-		if (Source != INDEX_NONE)
-		{
-			Writer.Write(TEXT("source"), Source);
-		}
-
-		if (Encoding != EGLTFJsonHDREncoding::None)
-		{
-			Writer.StartExtensions();
-
-			Writer.StartExtension(EGLTFJsonExtension::EPIC_TextureHDREncoding);
-			Writer.Write(TEXT("encoding"), Encoding);
-			Writer.EndExtension();
-
-			Writer.EndExtensions();
-		}
-	}
+	virtual void WriteObject(IGLTFJsonWriter& Writer) const override;
 };
diff --git a/Engine/Plugins/Enterprise/GLTFExporter/MigratedFromGit/Source/GLTFExporter/Public/Json/GLTFJsonTextureTransform.h b/Engine/Plugins/Enterprise/GLTFExporter/MigratedFromGit/Source/GLTFExporter/Public/Json/GLTFJsonTextureTransform.h
--- a/Engine/Plugins/Enterprise/GLTFExporter/MigratedFromGit/Source/GLTFExporter/Public/Json/GLTFJsonTextureTransform.h
+++ b/Engine/Plugins/Enterprise/GLTFExporter/MigratedFromGit/Source/GLTFExporter/Public/Json/GLTFJsonTextureTransform.h
@@ -5,7 +5,7 @@
 #include "Json/GLTFJsonObject.h"
 #include "Json/GLTFJsonVector.h"
 
-struct FGLTFJsonTextureTransform : IGLTFJsonObject
+struct GLTFEXPORTER_API FGLTFJsonTextureTransform : IGLTFJsonObject
 {
 	FGLTFJsonVector2 Offset;
 	FGLTFJsonVector2 Scale;
@@ -18,45 +18,11 @@ struct FGLTFJsonTextureTransform : IGLTFJsonObject
 	{
 	}
 
-	virtual void WriteObject(IGLTFJsonWriter& Writer) const override
-	{
-		if (!Offset.IsNearlyEqual(FGLTFJsonVector2::Zero, Writer.DefaultTolerance))
-		{
-			Writer.Write(TEXT("offset"), Offset);
-		}
-
-		if (!Scale.IsNearlyEqual(FGLTFJsonVector2::One, Writer.DefaultTolerance))
-		{
-			Writer.Write(TEXT("scale"), Scale);
-		}
-
-		if (!FMath::IsNearlyEqual(Rotation, 0, Writer.DefaultTolerance))
-		{
-			Writer.Write(TEXT("rotation"), Rotation);
-		}
-	}
-
-	bool IsNearlyEqual(const FGLTFJsonTextureTransform& Other, float Tolerance = KINDA_SMALL_NUMBER) const
-	{
-		return Offset.IsNearlyEqual(Other.Offset, Tolerance)
-			&& Scale.IsNearlyEqual(Other.Scale, Tolerance)
-			&& FMath::IsNearlyEqual(Rotation, Other.Rotation, Tolerance);
-	}
+	virtual void WriteObject(IGLTFJsonWriter& Writer) const override;
 
-	bool IsExactlyEqual(const FGLTFJsonTextureTransform& Other) const
-	{
-		return Offset.X == Other.Offset.X && Offset.Y == Other.Offset.Y
-			&& Scale.X == Other.Scale.X && Scale.Y == Other.Scale.Y
-			&& Rotation == Other.Rotation;
-	}
+	bool IsNearlyEqual(const FGLTFJsonTextureTransform& Other, float Tolerance = KINDA_SMALL_NUMBER) const;
+	bool IsExactlyEqual(const FGLTFJsonTextureTransform& Other) const;
 
-	bool IsNearlyDefault(float Tolerance = KINDA_SMALL_NUMBER) const
-	{
-		return IsNearlyEqual({}, Tolerance);
-	}
-
-	bool IsExactlyDefault() const
-	{
-		return IsExactlyEqual({});
-	}
+	bool IsNearlyDefault(float Tolerance = KINDA_SMALL_NUMBER) const;
+	bool IsExactlyDefault() const;
 };
diff --git a/Engine/Plugins/Enterprise/GLTFExporter/MigratedFromGit/Source/GLTFExporter/Public/Json/GLTFJsonUtility.h b/Engine/Plugins/Enterprise/GLTFExporter/MigratedFromGit/Source/GLTFExporter/Public/Json/GLTFJsonUtility.h
--- a/Engine/Plugins/Enterprise/GLTFExporter/MigratedFromGit/Source/GLTFExporter/Public/Json/GLTFJsonUtility.h
+++ b/Engine/Plugins/Enterprise/GLTFExporter/MigratedFromGit/Source/GLTFExporter/Public/Json/GLTFJsonUtility.h
@@ -4,7 +4,7 @@
 
 #include "Json/GLTFJsonEnums.h"
 
-struct FGLTFJsonUtility
+struct GLTFEXPORTER_API FGLTFJsonUtility
 {
 	template <typename EnumType, typename = typename TEnableIf<TIsEnum<EnumType>::Value>::Type>
 	static int32 GetValue(EnumType Enum)
@@ -12,188 +12,17 @@ struct FGLTFJsonUtility
 		return static_cast<int32>(Enum);
 	}
 
-	static const TCHAR* GetValue(EGLTFJsonExtension Enum)
-	{
-		switch (Enum)
-		{
-			case EGLTFJsonExtension::KHR_LightsPunctual:      return TEXT("KHR_lights_punctual");
-			case EGLTFJsonExtension::KHR_MaterialsClearCoat:  return TEXT("KHR_materials_clearcoat");
-			case EGLTFJsonExtension::KHR_MaterialsUnlit:      return TEXT("KHR_materials_unlit");
-			case EGLTFJsonExtension::KHR_MaterialsVariants:   return TEXT("KHR_materials_variants");
-			case EGLTFJsonExtension::KHR_MeshQuantization:    return TEXT("KHR_mesh_quantization");
-			case EGLTFJsonExtension::KHR_TextureTransform:    return TEXT("KHR_texture_transform");
-			case EGLTFJsonExtension::EPIC_AnimationHotspots:  return TEXT("EPIC_animation_hotspots");
-			case EGLTFJsonExtension::EPIC_AnimationPlayback:  return TEXT("EPIC_animation_playback");
-			case EGLTFJsonExtension::EPIC_BlendModes:         return TEXT("EPIC_blend_modes");
-			case EGLTFJsonExtension::EPIC_CameraControls:     return TEXT("EPIC_camera_controls");
-			case EGLTFJsonExtension::EPIC_HDRIBackdrops:      return TEXT("EPIC_hdri_backdrops");
-			case EGLTFJsonExtension::EPIC_LevelVariantSets:   return TEXT("EPIC_level_variant_sets");
-			case EGLTFJsonExtension::EPIC_LightmapTextures:   return TEXT("EPIC_lightmap_textures");
-			case EGLTFJsonExtension::EPIC_SkySpheres:         return TEXT("EPIC_sky_spheres");
-			case EGLTFJsonExtension::EPIC_TextureHDREncoding: return TEXT("EPIC_texture_hdr_encoding");
-			default:
-				checkNoEntry();
-				return TEXT("");
-		}
-	}
-
-	static const TCHAR* GetValue(EGLTFJsonAlphaMode Enum)
-	{
-		switch (Enum)
-		{
-			case EGLTFJsonAlphaMode::Opaque: return TEXT("OPAQUE");
-			case EGLTFJsonAlphaMode::Blend:  return TEXT("BLEND");
-			case EGLTFJsonAlphaMode::Mask:   return TEXT("MASK");
-			default:
-				checkNoEntry();
-				return TEXT("");
-		}
-	}
-
-	static const TCHAR* GetValue(EGLTFJsonBlendMode Enum)
-	{
-		switch (Enum)
-		{
-			case EGLTFJsonBlendMode::Additive:       return TEXT("ADDITIVE");
-			case EGLTFJsonBlendMode::Modulate:       return TEXT("MODULATE");
-			case EGLTFJsonBlendMode::AlphaComposite: return TEXT("ALPHACOMPOSITE");
-			default:
-				checkNoEntry();
-				return TEXT("");
-		}
-	}
-
-	static const TCHAR* GetValue(EGLTFJsonMimeType Enum)
-	{
-		switch (Enum)
-		{
-			case EGLTFJsonMimeType::PNG:  return TEXT("image/png");
-			case EGLTFJsonMimeType::JPEG: return TEXT("image/jpeg");
-			default:
-				checkNoEntry();
-				return TEXT("");
-		}
-	}
-
-	static const TCHAR* GetValue(EGLTFJsonAccessorType Enum)
-	{
-		switch (Enum)
-		{
-			case EGLTFJsonAccessorType::Scalar: return TEXT("SCALAR");
-			case EGLTFJsonAccessorType::Vec2:   return TEXT("VEC2");
-			case EGLTFJsonAccessorType::Vec3:   return TEXT("VEC3");
-			case EGLTFJsonAccessorType::Vec4:   return TEXT("VEC4");
-			case EGLTFJsonAccessorType::Mat2:   return TEXT("MAT2");
-			case EGLTFJsonAccessorType::Mat3:   return TEXT("MAT3");
-			case EGLTFJsonAccessorType::Mat4:   return TEXT("MAT4");
-			default:
-				checkNoEntry();
-				return TEXT("");
-		}
-	}
-
-	static const TCHAR* GetValue(EGLTFJsonHDREncoding Enum)
-	{
-		switch (Enum)
-		{
-			case EGLTFJsonHDREncoding::RGBM: return TEXT("RGBM");
-			case EGLTFJsonHDREncoding::RGBE: return TEXT("RGBE");
-			default:
-				checkNoEntry();
-				return TEXT("");
-		}
-	}
-
-	static const TCHAR* GetValue(EGLTFJsonCubeFace Enum)
-	{
-		switch (Enum)
-		{
-			case EGLTFJsonCubeFace::PosX: return TEXT("PosX");
-			case EGLTFJsonCubeFace::NegX: return TEXT("NegX");
-			case EGLTFJsonCubeFace::PosY: return TEXT("PosY");
-			case EGLTFJsonCubeFace::NegY: return TEXT("NegY");
-			case EGLTFJsonCubeFace::PosZ: return TEXT("PosZ");
-			case EGLTFJsonCubeFace::NegZ: return TEXT("NegZ");
-			default:
-				checkNoEntry();
-				return TEXT("");
-		}
-	}
-
-	static const TCHAR* GetValue(EGLTFJsonCameraType Enum)
-	{
-		switch (Enum)
-		{
-			case EGLTFJsonCameraType::Perspective:  return TEXT("perspective");
-			case EGLTFJsonCameraType::Orthographic: return TEXT("orthographic");
-			default:
-				checkNoEntry();
-				return TEXT("");
-		}
-	}
-
-	static const TCHAR* GetValue(EGLTFJsonLightType Enum)
-	{
-		switch (Enum)
-		{
-			case EGLTFJsonLightType::Directional: return TEXT("directional");
-			case EGLTFJsonLightType::Point:       return TEXT("point");
-			case EGLTFJsonLightType::Spot:        return TEXT("spot");
-			default:
-				checkNoEntry();
-				return TEXT("");
-		}
-	}
-
-	static const TCHAR* GetValue(EGLTFJsonInterpolation Enum)
-	{
-		switch (Enum)
-		{
-			case EGLTFJsonInterpolation::Linear:      return TEXT("LINEAR");
-			case EGLTFJsonInterpolation::Step:        return TEXT("STEP");
-			case EGLTFJsonInterpolation::CubicSpline: return TEXT("CUBICSPLINE");
-			default:
-				checkNoEntry();
-				return TEXT("");
-		}
-	}
-
-	static const TCHAR* GetValue(EGLTFJsonTargetPath Enum)
-	{
-		switch (Enum)
-		{
-			case EGLTFJsonTargetPath::Translation: return TEXT("translation");
-			case EGLTFJsonTargetPath::Rotation:    return TEXT("rotation");
-			case EGLTFJsonTargetPath::Scale:       return TEXT("scale");
-			case EGLTFJsonTargetPath::Weights:     return TEXT("weights");
-			default:
-				checkNoEntry();
-				return TEXT("");
-		}
-	}
-
-	static const TCHAR* GetValue(EGLTFJsonCameraControlMode Enum)
-	{
-		switch (Enum)
-		{
-			case EGLTFJsonCameraControlMode::FreeLook: return TEXT("freeLook");
-			case EGLTFJsonCameraControlMode::Orbital:  return TEXT("orbital");
-			default:
-				checkNoEntry();
-				return TEXT("");
-		}
-	}
-
-	static const TCHAR* GetValue(EGLTFJsonShadingModel Enum)
-	{
-		switch (Enum)
-		{
-			case EGLTFJsonShadingModel::Default:   return TEXT("Default");
-			case EGLTFJsonShadingModel::Unlit:     return TEXT("Unlit");
-			case EGLTFJsonShadingModel::ClearCoat: return TEXT("ClearCoat");
-			default:
-				checkNoEntry();
-				return TEXT("");
-		}
-	}
+	static const TCHAR* GetValue(EGLTFJsonExtension Enum);
+	static const TCHAR* GetValue(EGLTFJsonAlphaMode Enum);
+	static const TCHAR* GetValue(EGLTFJsonBlendMode Enum);
+	static const TCHAR* GetValue(EGLTFJsonMimeType Enum);
+	static const TCHAR* GetValue(EGLTFJsonAccessorType Enum);
+	static const TCHAR* GetValue(EGLTFJsonHDREncoding Enum);
+	static const TCHAR* GetValue(EGLTFJsonCubeFace Enum);
+	static const TCHAR* GetValue(EGLTFJsonCameraType Enum);
+	static const TCHAR* GetValue(EGLTFJsonLightType Enum);
+	static const TCHAR* GetValue(EGLTFJsonInterpolation Enum);
+	static const TCHAR* GetValue(EGLTFJsonTargetPath Enum);
+	static const TCHAR* GetValue(EGLTFJsonCameraControlMode Enum);
+	static const TCHAR* GetValue(EGLTFJsonShadingModel Enum);
 };
diff --git a/Engine/Plugins/Enterprise/GLTFExporter/MigratedFromGit/Source/GLTFExporter/Public/Json/GLTFJsonValue.h b/Engine/Plugins/Enterprise/GLTFExporter/MigratedFromGit/Source/GLTFExporter/Public/Json/GLTFJsonValue.h
--- a/Engine/Plugins/Enterprise/GLTFExporter/MigratedFromGit/Source/GLTFExporter/Public/Json/GLTFJsonValue.h
+++ b/Engine/Plugins/Enterprise/GLTFExporter/MigratedFromGit/Source/GLTFExporter/Public/Json/GLTFJsonValue.h
@@ -4,7 +4,7 @@
 
 class IGLTFJsonWriter;
 
-struct IGLTFJsonValue
+struct GLTFEXPORTER_API IGLTFJsonValue
 {
 	virtual ~IGLTFJsonValue() = default;
 
diff --git a/Engine/Plugins/Enterprise/GLTFExporter/MigratedFromGit/Source/GLTFExporter/Public/Json/GLTFJsonVector.h b/Engine/Plugins/Enterprise/GLTFExporter/MigratedFromGit/Source/GLTFExporter/Public/Json/GLTFJsonVector.h
--- a/Engine/Plugins/Enterprise/GLTFExporter/MigratedFromGit/Source/GLTFExporter/Public/Json/GLTFJsonVector.h
+++ b/Engine/Plugins/Enterprise/GLTFExporter/MigratedFromGit/Source/GLTFExporter/Public/Json/GLTFJsonVector.h
@@ -6,7 +6,7 @@
 #include "Json/GLTFJsonArray.h"
 
 template <typename BaseType>
-struct TGLTFJsonVector : BaseType, IGLTFJsonArray
+struct GLTFEXPORTER_API TGLTFJsonVector : BaseType, IGLTFJsonArray
 {
 	static const TGLTFJsonVector Zero;
 	static const TGLTFJsonVector One;
@@ -62,7 +62,7 @@ struct TGLTFJsonVector : BaseType, IGLTFJsonArray
 	}
 };
 
-struct FGLTFJsonVector2 : TGLTFJsonVector<FGLTFVector2>
+struct GLTFEXPORTER_API FGLTFJsonVector2 : TGLTFJsonVector<FGLTFVector2>
 {
 	static const FGLTFJsonVector2 Zero;
 	static const FGLTFJsonVector2 One;
@@ -71,7 +71,7 @@ struct FGLTFJsonVector2 : TGLTFJsonVector<FGLTFVector2>
 	using TGLTFJsonVector::operator=;
 };
 
-struct FGLTFJsonVector3 : TGLTFJsonVector<FGLTFVector3>
+struct GLTFEXPORTER_API FGLTFJsonVector3 : TGLTFJsonVector<FGLTFVector3>
 {
 	static const FGLTFJsonVector3 Zero;
 	static const FGLTFJsonVector3 One;
@@ -80,7 +80,7 @@ struct FGLTFJsonVector3 : TGLTFJsonVector<FGLTFVector3>
 	using TGLTFJsonVector::operator=;
 };
 
-struct FGLTFJsonVector4 : TGLTFJsonVector<FGLTFVector4>
+struct GLTFEXPORTER_API FGLTFJsonVector4 : TGLTFJsonVector<FGLTFVector4>
 {
 	static const FGLTFJsonVector4 Zero;
 	static const FGLTFJsonVector4 One;
diff --git a/Engine/Plugins/Enterprise/GLTFExporter/MigratedFromGit/Source/GLTFExporter/Public/Json/GLTFJsonWriter.h b/Engine/Plugins/Enterprise/GLTFExporter/MigratedFromGit/Source/GLTFExporter/Public/Json/GLTFJsonWriter.h
--- a/Engine/Plugins/Enterprise/GLTFExporter/MigratedFromGit/Source/GLTFExporter/Public/Json/GLTFJsonWriter.h
+++ b/Engine/Plugins/Enterprise/GLTFExporter/MigratedFromGit/Source/GLTFExporter/Public/Json/GLTFJsonWriter.h
@@ -6,7 +6,7 @@
 #include "Json/GLTFJsonUtility.h"
 #include "Json/GLTFJsonExtensions.h"
 
-class IGLTFJsonWriter
+class GLTFEXPORTER_API IGLTFJsonWriter
 {
 public:
 
