[GLTFExporter] refactor: add IGLTFJsonIndexedObject and TGLTFJsonIndexedObjectArray

commit bf2c65b88c27f0fddf13c2cfd680f3005161c3f3
Date: Sun, 13 Mar 2022 12:35:44 +0100

    refactor: add IGLTFJsonIndexedObject and TGLTFJsonIndexedObjectArray
    (cherry picked from commit 17deb2e6b6ee0b8405446283d6112a62278e4568​)

#skipci

diff --git a/Engine/Plugins/Enterprise/GLTFExporter/MigratedFromGit/Source/GLTFExporter/Private/Builders/GLTFJsonBuilder.cpp b/Engine/Plugins/Enterprise/GLTFExporter/MigratedFromGit/Source/GLTFExporter/Private/Builders/GLTFJsonBuilder.cpp
--- a/Engine/Plugins/Enterprise/GLTFExporter/MigratedFromGit/Source/GLTFExporter/Private/Builders/GLTFJsonBuilder.cpp
+++ b/Engine/Plugins/Enterprise/GLTFExporter/MigratedFromGit/Source/GLTFExporter/Private/Builders/GLTFJsonBuilder.cpp
@@ -43,102 +43,102 @@ void FGLTFJsonBuilder::AddExtension(EGLTFJsonExtension Extension, bool bIsRequir
 
 FGLTFJsonAccessorIndex FGLTFJsonBuilder::AddAccessor(const FGLTFJsonAccessor& JsonAccessor)
 {
-	return FGLTFJsonAccessorIndex(JsonRoot.Accessors.Add(MakeUnique<FGLTFJsonAccessor>(JsonAccessor)));
+	return FGLTFJsonAccessorIndex(JsonRoot.Accessors.Add(JsonAccessor)->Index);
 }
 
 FGLTFJsonAnimationIndex FGLTFJsonBuilder::AddAnimation(const FGLTFJsonAnimation& JsonAnimation)
 {
-	return FGLTFJsonAnimationIndex(JsonRoot.Animations.Add(MakeUnique<FGLTFJsonAnimation>(JsonAnimation)));
+	return FGLTFJsonAnimationIndex(JsonRoot.Animations.Add(JsonAnimation)->Index);
 }
 
 FGLTFJsonBufferIndex FGLTFJsonBuilder::AddBuffer(const FGLTFJsonBuffer& JsonBuffer)
 {
-	return FGLTFJsonBufferIndex(JsonRoot.Buffers.Add(MakeUnique<FGLTFJsonBuffer>(JsonBuffer)));
+	return FGLTFJsonBufferIndex(JsonRoot.Buffers.Add(JsonBuffer)->Index);
 }
 
 FGLTFJsonBufferViewIndex FGLTFJsonBuilder::AddBufferView(const FGLTFJsonBufferView& JsonBufferView)
 {
-	return FGLTFJsonBufferViewIndex(JsonRoot.BufferViews.Add(MakeUnique<FGLTFJsonBufferView>(JsonBufferView)));
+	return FGLTFJsonBufferViewIndex(JsonRoot.BufferViews.Add(JsonBufferView)->Index);
 }
 
 FGLTFJsonCameraIndex FGLTFJsonBuilder::AddCamera(const FGLTFJsonCamera& JsonCamera)
 {
-	return FGLTFJsonCameraIndex(JsonRoot.Cameras.Add(MakeUnique<FGLTFJsonCamera>(JsonCamera)));
+	return FGLTFJsonCameraIndex(JsonRoot.Cameras.Add(JsonCamera)->Index);
 }
 
 FGLTFJsonImageIndex FGLTFJsonBuilder::AddImage(const FGLTFJsonImage& JsonImage)
 {
-	return FGLTFJsonImageIndex(JsonRoot.Images.Add(MakeUnique<FGLTFJsonImage>(JsonImage)));
+	return FGLTFJsonImageIndex(JsonRoot.Images.Add(JsonImage)->Index);
 }
 
 FGLTFJsonMaterialIndex FGLTFJsonBuilder::AddMaterial(const FGLTFJsonMaterial& JsonMaterial)
 {
-	return FGLTFJsonMaterialIndex(JsonRoot.Materials.Add(MakeUnique<FGLTFJsonMaterial>(JsonMaterial)));
+	return FGLTFJsonMaterialIndex(JsonRoot.Materials.Add(JsonMaterial)->Index);
 }
 
 FGLTFJsonMeshIndex FGLTFJsonBuilder::AddMesh(const FGLTFJsonMesh& JsonMesh)
 {
-	return FGLTFJsonMeshIndex(JsonRoot.Meshes.Add(MakeUnique<FGLTFJsonMesh>(JsonMesh)));
+	return FGLTFJsonMeshIndex(JsonRoot.Meshes.Add(JsonMesh)->Index);
 }
 
 FGLTFJsonNodeIndex FGLTFJsonBuilder::AddNode(const FGLTFJsonNode& JsonNode)
 {
-	return FGLTFJsonNodeIndex(JsonRoot.Nodes.Add(MakeUnique<FGLTFJsonNode>(JsonNode)));
+	return FGLTFJsonNodeIndex(JsonRoot.Nodes.Add(JsonNode)->Index);
 }
 
 FGLTFJsonSamplerIndex FGLTFJsonBuilder::AddSampler(const FGLTFJsonSampler& JsonSampler)
 {
-	return FGLTFJsonSamplerIndex(JsonRoot.Samplers.Add(MakeUnique<FGLTFJsonSampler>(JsonSampler)));
+	return FGLTFJsonSamplerIndex(JsonRoot.Samplers.Add(JsonSampler)->Index);
 }
 
 FGLTFJsonSceneIndex FGLTFJsonBuilder::AddScene(const FGLTFJsonScene& JsonScene)
 {
-	return FGLTFJsonSceneIndex(JsonRoot.Scenes.Add(MakeUnique<FGLTFJsonScene>(JsonScene)));
+	return FGLTFJsonSceneIndex(JsonRoot.Scenes.Add(JsonScene)->Index);
 }
 
 FGLTFJsonSkinIndex FGLTFJsonBuilder::AddSkin(const FGLTFJsonSkin& JsonSkin)
 {
-	return FGLTFJsonSkinIndex(JsonRoot.Skins.Add(MakeUnique<FGLTFJsonSkin>(JsonSkin)));
+	return FGLTFJsonSkinIndex(JsonRoot.Skins.Add(JsonSkin)->Index);
 }
 
 FGLTFJsonTextureIndex FGLTFJsonBuilder::AddTexture(const FGLTFJsonTexture& JsonTexture)
 {
-	return FGLTFJsonTextureIndex(JsonRoot.Textures.Add(MakeUnique<FGLTFJsonTexture>(JsonTexture)));
+	return FGLTFJsonTextureIndex(JsonRoot.Textures.Add(JsonTexture)->Index);
 }
 
 FGLTFJsonBackdropIndex FGLTFJsonBuilder::AddBackdrop(const FGLTFJsonBackdrop& JsonBackdrop)
 {
-	return FGLTFJsonBackdropIndex(JsonRoot.Backdrops.Add(MakeUnique<FGLTFJsonBackdrop>(JsonBackdrop)));
+	return FGLTFJsonBackdropIndex(JsonRoot.Backdrops.Add(JsonBackdrop)->Index);
 }
 
 FGLTFJsonHotspotIndex FGLTFJsonBuilder::AddHotspot(const FGLTFJsonHotspot& JsonHotspot)
 {
-	return FGLTFJsonHotspotIndex(JsonRoot.Hotspots.Add(MakeUnique<FGLTFJsonHotspot>(JsonHotspot)));
+	return FGLTFJsonHotspotIndex(JsonRoot.Hotspots.Add(JsonHotspot)->Index);
 }
 
 FGLTFJsonLightIndex FGLTFJsonBuilder::AddLight(const FGLTFJsonLight& JsonLight)
 {
-	return FGLTFJsonLightIndex(JsonRoot.Lights.Add(MakeUnique<FGLTFJsonLight>(JsonLight)));
+	return FGLTFJsonLightIndex(JsonRoot.Lights.Add(JsonLight)->Index);
 }
 
 FGLTFJsonLightMapIndex FGLTFJsonBuilder::AddLightMap(const FGLTFJsonLightMap& JsonLightMap)
 {
-	return FGLTFJsonLightMapIndex(JsonRoot.LightMaps.Add(MakeUnique<FGLTFJsonLightMap>(JsonLightMap)));
+	return FGLTFJsonLightMapIndex(JsonRoot.LightMaps.Add(JsonLightMap)->Index);
 }
 
 FGLTFJsonSkySphereIndex FGLTFJsonBuilder::AddSkySphere(const FGLTFJsonSkySphere& JsonSkySphere)
 {
-	return FGLTFJsonSkySphereIndex(JsonRoot.SkySpheres.Add(MakeUnique<FGLTFJsonSkySphere>(JsonSkySphere)));
+	return FGLTFJsonSkySphereIndex(JsonRoot.SkySpheres.Add(JsonSkySphere)->Index);
 }
 
 FGLTFJsonEpicLevelVariantSetsIndex FGLTFJsonBuilder::AddEpicLevelVariantSets(const FGLTFJsonEpicLevelVariantSets& JsonEpicLevelVariantSets)
 {
-	return FGLTFJsonEpicLevelVariantSetsIndex(JsonRoot.EpicLevelVariantSets.Add(MakeUnique<FGLTFJsonEpicLevelVariantSets>(JsonEpicLevelVariantSets)));
+	return FGLTFJsonEpicLevelVariantSetsIndex(JsonRoot.EpicLevelVariantSets.Add(JsonEpicLevelVariantSets)->Index);
 }
 
 FGLTFJsonKhrMaterialVariantIndex FGLTFJsonBuilder::AddKhrMaterialVariant(const FGLTFJsonKhrMaterialVariant& JsonKhrMaterialVariant)
 {
-	return FGLTFJsonKhrMaterialVariantIndex(JsonRoot.KhrMaterialVariants.Add(MakeUnique<FGLTFJsonKhrMaterialVariant>(JsonKhrMaterialVariant)));
+	return FGLTFJsonKhrMaterialVariantIndex(JsonRoot.KhrMaterialVariants.Add(JsonKhrMaterialVariant)->Index);
 }
 
 FGLTFJsonNodeIndex FGLTFJsonBuilder::AddChildNode(FGLTFJsonNodeIndex ParentIndex, const FGLTFJsonNode& JsonNode)
@@ -170,107 +170,107 @@ const FGLTFJsonRoot& FGLTFJsonBuilder::GetRoot() const
 	return JsonRoot;
 }
 
-FGLTFJsonAccessor& FGLTFJsonBuilder::GetAccessor(FGLTFJsonAccessorIndex AccessorIndex) const
+FGLTFJsonAccessor& FGLTFJsonBuilder::GetAccessor(FGLTFJsonAccessorIndex AccessorIndex)
 {
 	return *JsonRoot.Accessors[AccessorIndex];
 }
 
-FGLTFJsonAnimation& FGLTFJsonBuilder::GetAnimation(FGLTFJsonAnimationIndex AnimationIndex) const
+FGLTFJsonAnimation& FGLTFJsonBuilder::GetAnimation(FGLTFJsonAnimationIndex AnimationIndex)
 {
 	return *JsonRoot.Animations[AnimationIndex];
 }
 
-FGLTFJsonBuffer& FGLTFJsonBuilder::GetBuffer(FGLTFJsonBufferIndex BufferIndex) const
+FGLTFJsonBuffer& FGLTFJsonBuilder::GetBuffer(FGLTFJsonBufferIndex BufferIndex)
 {
 	return *JsonRoot.Buffers[BufferIndex];
 }
 
-FGLTFJsonBufferView& FGLTFJsonBuilder::GetBufferView(FGLTFJsonBufferViewIndex BufferViewIndex) const
+FGLTFJsonBufferView& FGLTFJsonBuilder::GetBufferView(FGLTFJsonBufferViewIndex BufferViewIndex)
 {
 	return *JsonRoot.BufferViews[BufferViewIndex];
 }
 
-FGLTFJsonCamera& FGLTFJsonBuilder::GetCamera(FGLTFJsonCameraIndex CameraIndex) const
+FGLTFJsonCamera& FGLTFJsonBuilder::GetCamera(FGLTFJsonCameraIndex CameraIndex)
 {
 	return *JsonRoot.Cameras[CameraIndex];
 }
 
-FGLTFJsonImage& FGLTFJsonBuilder::GetImage(FGLTFJsonImageIndex ImageIndex) const
+FGLTFJsonImage& FGLTFJsonBuilder::GetImage(FGLTFJsonImageIndex ImageIndex)
 {
 	return *JsonRoot.Images[ImageIndex];
 }
 
-FGLTFJsonMaterial& FGLTFJsonBuilder::GetMaterial(FGLTFJsonMaterialIndex MaterialIndex) const
+FGLTFJsonMaterial& FGLTFJsonBuilder::GetMaterial(FGLTFJsonMaterialIndex MaterialIndex)
 {
 	return *JsonRoot.Materials[MaterialIndex];
 }
 
-FGLTFJsonMesh& FGLTFJsonBuilder::GetMesh(FGLTFJsonMeshIndex MeshIndex) const
+FGLTFJsonMesh& FGLTFJsonBuilder::GetMesh(FGLTFJsonMeshIndex MeshIndex)
 {
 	return *JsonRoot.Meshes[MeshIndex];
 }
 
-FGLTFJsonNode& FGLTFJsonBuilder::GetNode(FGLTFJsonNodeIndex NodeIndex) const
+FGLTFJsonNode& FGLTFJsonBuilder::GetNode(FGLTFJsonNodeIndex NodeIndex)
 {
 	return *JsonRoot.Nodes[NodeIndex];
 }
 
-FGLTFJsonSampler& FGLTFJsonBuilder::GetSampler(FGLTFJsonSamplerIndex SamplerIndex) const
+FGLTFJsonSampler& FGLTFJsonBuilder::GetSampler(FGLTFJsonSamplerIndex SamplerIndex)
 {
 	return *JsonRoot.Samplers[SamplerIndex];
 }
 
-FGLTFJsonScene& FGLTFJsonBuilder::GetScene(FGLTFJsonSceneIndex SceneIndex) const
+FGLTFJsonScene& FGLTFJsonBuilder::GetScene(FGLTFJsonSceneIndex SceneIndex)
 {
 	return *JsonRoot.Scenes[SceneIndex];
 }
 
-FGLTFJsonSkin& FGLTFJsonBuilder::GetSkin(FGLTFJsonSkinIndex SkinIndex) const
+FGLTFJsonSkin& FGLTFJsonBuilder::GetSkin(FGLTFJsonSkinIndex SkinIndex)
 {
 	return *JsonRoot.Skins[SkinIndex];
 }
 
-FGLTFJsonTexture& FGLTFJsonBuilder::GetTexture(FGLTFJsonTextureIndex TextureIndex) const
+FGLTFJsonTexture& FGLTFJsonBuilder::GetTexture(FGLTFJsonTextureIndex TextureIndex)
 {
 	return *JsonRoot.Textures[TextureIndex];
 }
 
-FGLTFJsonBackdrop& FGLTFJsonBuilder::GetBackdrop(FGLTFJsonBackdropIndex BackdropIndex) const
+FGLTFJsonBackdrop& FGLTFJsonBuilder::GetBackdrop(FGLTFJsonBackdropIndex BackdropIndex)
 {
 	return *JsonRoot.Backdrops[BackdropIndex];
 }
 
-FGLTFJsonHotspot& FGLTFJsonBuilder::GetHotspot(FGLTFJsonHotspotIndex HotspotIndex) const
+FGLTFJsonHotspot& FGLTFJsonBuilder::GetHotspot(FGLTFJsonHotspotIndex HotspotIndex)
 {
 	return *JsonRoot.Hotspots[HotspotIndex];
 }
 
-FGLTFJsonLight& FGLTFJsonBuilder::GetLight(FGLTFJsonLightIndex LightIndex) const
+FGLTFJsonLight& FGLTFJsonBuilder::GetLight(FGLTFJsonLightIndex LightIndex)
 {
 	return *JsonRoot.Lights[LightIndex];
 }
 
-FGLTFJsonLightMap& FGLTFJsonBuilder::GetLightMap(FGLTFJsonLightMapIndex LightMapIndex) const
+FGLTFJsonLightMap& FGLTFJsonBuilder::GetLightMap(FGLTFJsonLightMapIndex LightMapIndex)
 {
 	return *JsonRoot.LightMaps[LightMapIndex];
 }
 
-FGLTFJsonSkySphere& FGLTFJsonBuilder::GetSkySphere(FGLTFJsonSkySphereIndex SkySphereIndex) const
+FGLTFJsonSkySphere& FGLTFJsonBuilder::GetSkySphere(FGLTFJsonSkySphereIndex SkySphereIndex)
 {
 	return *JsonRoot.SkySpheres[SkySphereIndex];
 }
 
-FGLTFJsonEpicLevelVariantSets& FGLTFJsonBuilder::GetEpicLevelVariantSets(FGLTFJsonEpicLevelVariantSetsIndex EpicLevelVariantSetsIndex) const
+FGLTFJsonEpicLevelVariantSets& FGLTFJsonBuilder::GetEpicLevelVariantSets(FGLTFJsonEpicLevelVariantSetsIndex EpicLevelVariantSetsIndex)
 {
 	return *JsonRoot.EpicLevelVariantSets[EpicLevelVariantSetsIndex];
 }
 
-FGLTFJsonKhrMaterialVariant& FGLTFJsonBuilder::GetKhrMaterialVariant(FGLTFJsonKhrMaterialVariantIndex KhrMaterialVariantIndex) const
+FGLTFJsonKhrMaterialVariant& FGLTFJsonBuilder::GetKhrMaterialVariant(FGLTFJsonKhrMaterialVariantIndex KhrMaterialVariantIndex)
 {
 	return *JsonRoot.KhrMaterialVariants[KhrMaterialVariantIndex];
 }
 
-FGLTFJsonNodeIndex FGLTFJsonBuilder::GetComponentNodeIndex(FGLTFJsonNodeIndex NodeIndex) const
+FGLTFJsonNodeIndex FGLTFJsonBuilder::GetComponentNodeIndex(FGLTFJsonNodeIndex NodeIndex)
 {
 	if (NodeIndex == INDEX_NONE)
 	{
diff --git a/Engine/Plugins/Enterprise/GLTFExporter/MigratedFromGit/Source/GLTFExporter/Private/Builders/GLTFJsonBuilder.h b/Engine/Plugins/Enterprise/GLTFExporter/MigratedFromGit/Source/GLTFExporter/Private/Builders/GLTFJsonBuilder.h
--- a/Engine/Plugins/Enterprise/GLTFExporter/MigratedFromGit/Source/GLTFExporter/Private/Builders/GLTFJsonBuilder.h
+++ b/Engine/Plugins/Enterprise/GLTFExporter/MigratedFromGit/Source/GLTFExporter/Private/Builders/GLTFJsonBuilder.h
@@ -44,28 +44,28 @@ public:
 	FGLTFJsonNodeIndex AddChildComponentNode(FGLTFJsonNodeIndex ParentNodeIndex, const FGLTFJsonNode& JsonNode = {});
 
 	const FGLTFJsonRoot& GetRoot() const;
-	FGLTFJsonAccessor& GetAccessor(FGLTFJsonAccessorIndex AccessorIndex) const;
-	FGLTFJsonAnimation& GetAnimation(FGLTFJsonAnimationIndex AnimationIndex) const;
-	FGLTFJsonBuffer& GetBuffer(FGLTFJsonBufferIndex BufferIndex) const;
-	FGLTFJsonBufferView& GetBufferView(FGLTFJsonBufferViewIndex BufferViewIndex) const;
-	FGLTFJsonCamera& GetCamera(FGLTFJsonCameraIndex CameraIndex) const;
-	FGLTFJsonImage& GetImage(FGLTFJsonImageIndex ImageIndex) const;
-	FGLTFJsonMaterial& GetMaterial(FGLTFJsonMaterialIndex MaterialIndex) const;
-	FGLTFJsonMesh& GetMesh(FGLTFJsonMeshIndex MeshIndex) const;
-	FGLTFJsonNode& GetNode(FGLTFJsonNodeIndex NodeIndex) const;
-	FGLTFJsonSampler& GetSampler(FGLTFJsonSamplerIndex SamplerIndex) const;
-	FGLTFJsonScene& GetScene(FGLTFJsonSceneIndex SceneIndex) const;
-	FGLTFJsonSkin& GetSkin(FGLTFJsonSkinIndex SkinIndex) const;
-	FGLTFJsonTexture& GetTexture(FGLTFJsonTextureIndex TextureIndex) const;
-	FGLTFJsonBackdrop& GetBackdrop(FGLTFJsonBackdropIndex BackdropIndex) const;
-	FGLTFJsonHotspot& GetHotspot(FGLTFJsonHotspotIndex HotspotIndex) const;
-	FGLTFJsonLight& GetLight(FGLTFJsonLightIndex LightIndex) const;
-	FGLTFJsonLightMap& GetLightMap(FGLTFJsonLightMapIndex LightMapIndex) const;
-	FGLTFJsonSkySphere& GetSkySphere(FGLTFJsonSkySphereIndex SkySphereIndex) const;
-	FGLTFJsonEpicLevelVariantSets& GetEpicLevelVariantSets(FGLTFJsonEpicLevelVariantSetsIndex EpicLevelVariantSetsIndex) const;
-	FGLTFJsonKhrMaterialVariant& GetKhrMaterialVariant(FGLTFJsonKhrMaterialVariantIndex KhrMaterialVariantIndex) const;
-
-	FGLTFJsonNodeIndex GetComponentNodeIndex(FGLTFJsonNodeIndex NodeIndex) const;
+	FGLTFJsonAccessor& GetAccessor(FGLTFJsonAccessorIndex AccessorIndex);
+	FGLTFJsonAnimation& GetAnimation(FGLTFJsonAnimationIndex AnimationIndex);
+	FGLTFJsonBuffer& GetBuffer(FGLTFJsonBufferIndex BufferIndex);
+	FGLTFJsonBufferView& GetBufferView(FGLTFJsonBufferViewIndex BufferViewIndex);
+	FGLTFJsonCamera& GetCamera(FGLTFJsonCameraIndex CameraIndex);
+	FGLTFJsonImage& GetImage(FGLTFJsonImageIndex ImageIndex);
+	FGLTFJsonMaterial& GetMaterial(FGLTFJsonMaterialIndex MaterialIndex);
+	FGLTFJsonMesh& GetMesh(FGLTFJsonMeshIndex MeshIndex);
+	FGLTFJsonNode& GetNode(FGLTFJsonNodeIndex NodeIndex);
+	FGLTFJsonSampler& GetSampler(FGLTFJsonSamplerIndex SamplerIndex);
+	FGLTFJsonScene& GetScene(FGLTFJsonSceneIndex SceneIndex);
+	FGLTFJsonSkin& GetSkin(FGLTFJsonSkinIndex SkinIndex);
+	FGLTFJsonTexture& GetTexture(FGLTFJsonTextureIndex TextureIndex);
+	FGLTFJsonBackdrop& GetBackdrop(FGLTFJsonBackdropIndex BackdropIndex);
+	FGLTFJsonHotspot& GetHotspot(FGLTFJsonHotspotIndex HotspotIndex);
+	FGLTFJsonLight& GetLight(FGLTFJsonLightIndex LightIndex);
+	FGLTFJsonLightMap& GetLightMap(FGLTFJsonLightMapIndex LightMapIndex);
+	FGLTFJsonSkySphere& GetSkySphere(FGLTFJsonSkySphereIndex SkySphereIndex);
+	FGLTFJsonEpicLevelVariantSets& GetEpicLevelVariantSets(FGLTFJsonEpicLevelVariantSetsIndex EpicLevelVariantSetsIndex);
+	FGLTFJsonKhrMaterialVariant& GetKhrMaterialVariant(FGLTFJsonKhrMaterialVariantIndex KhrMaterialVariantIndex);
+
+	FGLTFJsonNodeIndex GetComponentNodeIndex(FGLTFJsonNodeIndex NodeIndex);
 
 private:
 
diff --git a/Engine/Plugins/Enterprise/GLTFExporter/MigratedFromGit/Source/GLTFExporter/Public/Json/GLTFJsonAccessor.h b/Engine/Plugins/Enterprise/GLTFExporter/MigratedFromGit/Source/GLTFExporter/Public/Json/GLTFJsonAccessor.h
--- a/Engine/Plugins/Enterprise/GLTFExporter/MigratedFromGit/Source/GLTFExporter/Public/Json/GLTFJsonAccessor.h
+++ b/Engine/Plugins/Enterprise/GLTFExporter/MigratedFromGit/Source/GLTFExporter/Public/Json/GLTFJsonAccessor.h
@@ -4,7 +4,7 @@
 
 #include "Json/GLTFJsonCore.h"
 
-struct GLTFEXPORTER_API FGLTFJsonAccessor : IGLTFJsonObject
+struct GLTFEXPORTER_API FGLTFJsonAccessor : IGLTFJsonIndexedObject
 {
 	FString Name;
 
@@ -19,8 +19,9 @@ struct GLTFEXPORTER_API FGLTFJsonAccessor : IGLTFJsonObject
 	float Min[16];
 	float Max[16];
 
-	FGLTFJsonAccessor()
-		: ByteOffset(0)
+	FGLTFJsonAccessor(int32 Index = INDEX_NONE)
+		: IGLTFJsonIndexedObject(Index)
+		, ByteOffset(0)
 		, Count(0)
 		, Type(EGLTFJsonAccessorType::None)
 		, ComponentType(EGLTFJsonComponentType::None)
diff --git a/Engine/Plugins/Enterprise/GLTFExporter/MigratedFromGit/Source/GLTFExporter/Public/Json/GLTFJsonAnimation.h b/Engine/Plugins/Enterprise/GLTFExporter/MigratedFromGit/Source/GLTFExporter/Public/Json/GLTFJsonAnimation.h
--- a/Engine/Plugins/Enterprise/GLTFExporter/MigratedFromGit/Source/GLTFExporter/Public/Json/GLTFJsonAnimation.h
+++ b/Engine/Plugins/Enterprise/GLTFExporter/MigratedFromGit/Source/GLTFExporter/Public/Json/GLTFJsonAnimation.h
@@ -41,7 +41,7 @@ struct GLTFEXPORTER_API FGLTFJsonAnimationSampler : IGLTFJsonObject
 	virtual void WriteObject(IGLTFJsonWriter& Writer) const override;
 };
 
-struct GLTFEXPORTER_API FGLTFJsonAnimation : IGLTFJsonObject
+struct GLTFEXPORTER_API FGLTFJsonAnimation : IGLTFJsonIndexedObject
 {
 	FString Name;
 
@@ -50,5 +50,10 @@ struct GLTFEXPORTER_API FGLTFJsonAnimation : IGLTFJsonObject
 
 	FGLTFJsonAnimationPlayback Playback;
 
+	FGLTFJsonAnimation(int32 Index = INDEX_NONE)
+		: IGLTFJsonIndexedObject(Index)
+	{
+	}
+
 	virtual void WriteObject(IGLTFJsonWriter& Writer) const override;
 };
diff --git a/Engine/Plugins/Enterprise/GLTFExporter/MigratedFromGit/Source/GLTFExporter/Public/Json/GLTFJsonArray.h b/Engine/Plugins/Enterprise/GLTFExporter/MigratedFromGit/Source/GLTFExporter/Public/Json/GLTFJsonArray.h
--- a/Engine/Plugins/Enterprise/GLTFExporter/MigratedFromGit/Source/GLTFExporter/Public/Json/GLTFJsonArray.h
+++ b/Engine/Plugins/Enterprise/GLTFExporter/MigratedFromGit/Source/GLTFExporter/Public/Json/GLTFJsonArray.h
@@ -4,6 +4,7 @@
 
 #include "Json/GLTFJsonValue.h"
 #include "Json/GLTFJsonWriter.h"
+#include "Json/GLTFJsonObject.h"
 
 struct GLTFEXPORTER_API IGLTFJsonArray : IGLTFJsonValue
 {
@@ -16,3 +17,124 @@ struct GLTFEXPORTER_API IGLTFJsonArray : IGLTFJsonValue
 
 	virtual void WriteArray(IGLTFJsonWriter& Writer) const = 0;
 };
+
+template <typename ElementType, typename = typename TEnableIf<TIsDerivedFrom<ElementType, IGLTFJsonIndexedObject>::Value>::Type>
+struct TGLTFJsonIndexedObjectArray : IGLTFJsonArray
+{
+	using ArrayType = TArray<TUniquePtr<ElementType>>;
+	using SizeType = typename ArrayType::SizeType;
+
+	template <typename IteratorType, typename ValueType>
+	struct TIterator
+	{
+		TIterator(const IteratorType& Iterator)
+			: Iterator(Iterator)
+		{
+		}
+
+		ValueType operator*() const
+		{
+			return (*Iterator).Get();
+		}
+
+		TIterator& operator++()
+		{
+			++Iterator;
+			return *this;
+		}
+
+		TIterator& operator--()
+		{
+			--Iterator;
+			return *this;
+		}
+
+		friend bool operator!=(const TIterator& A, const TIterator& B)
+		{
+			return A.Iterator != B.Iterator;
+		};
+
+	private:
+
+		IteratorType Iterator;
+	};
+
+	typedef TIterator<typename ArrayType::RangedForIteratorType, ElementType*> IteratorType;
+	typedef TIterator<typename ArrayType::RangedForConstIteratorType, const ElementType*> ConstIteratorType;
+
+	TGLTFJsonIndexedObjectArray() = default;
+	TGLTFJsonIndexedObjectArray(TGLTFJsonIndexedObjectArray&& Other) = default;
+	TGLTFJsonIndexedObjectArray& operator=(TGLTFJsonIndexedObjectArray&& Other) = default;
+
+	TGLTFJsonIndexedObjectArray(const TGLTFJsonIndexedObjectArray&) = delete;
+	TGLTFJsonIndexedObjectArray& operator=(const TGLTFJsonIndexedObjectArray&) = delete;
+
+	ElementType* Add(const ElementType& Temp)
+	{
+		ElementType* Element = new ElementType(Array.Num());
+		*Element = Temp;
+		Element->Index = Array.Num();
+		Array.Add(TUniquePtr<ElementType>(Element));
+		return Element;
+	}
+
+	ElementType* Add()
+	{
+		ElementType* Element = new ElementType(Array.Num());
+		Array.Add(TUniquePtr<ElementType>(Element));
+		return Element;
+	}
+
+	bool IsValidIndex(SizeType Index) const
+	{
+		return Array.IsValidIndex(Index);
+	}
+
+	SizeType Num() const
+	{
+		return Array.Num();
+	}
+
+	bool Contains(const ElementType* Element) const
+	{
+		TUniquePtr<ElementType> TempPtr(Element);
+		const bool Result = Array.Contains(TempPtr);
+		TempPtr.Release();
+		return Result;
+	}
+
+	SizeType Find(const ElementType* Element) const
+	{
+		TUniquePtr<ElementType> TempPtr(Element);
+		const SizeType Result = Array.Find(TempPtr);
+		TempPtr.Release();
+		return Result;
+	}
+
+	ElementType* operator[](SizeType Index)
+	{
+		return Array[Index].Get();
+	}
+
+	const ElementType* operator[](SizeType Index) const
+	{
+		return Array[Index].Get();
+	}
+
+	IteratorType begin() { return IteratorType(Array.begin()); }
+	IteratorType end() { return IteratorType(Array.end()); }
+	ConstIteratorType begin() const { return ConstIteratorType(Array.begin()); }
+	ConstIteratorType end() const { return ConstIteratorType(Array.end()); }
+
+	virtual void WriteArray(IGLTFJsonWriter& Writer) const override
+	{
+		for (const ElementType* Element : *this)
+		{
+			Writer.Write(*Element);
+		}
+	}
+
+private:
+
+	ArrayType Array;
+};
diff --git a/Engine/Plugins/Enterprise/GLTFExporter/MigratedFromGit/Source/GLTFExporter/Public/Json/GLTFJsonBackdrop.h b/Engine/Plugins/Enterprise/GLTFExporter/MigratedFromGit/Source/GLTFExporter/Public/Json/GLTFJsonBackdrop.h
--- a/Engine/Plugins/Enterprise/GLTFExporter/MigratedFromGit/Source/GLTFExporter/Public/Json/GLTFJsonBackdrop.h
+++ b/Engine/Plugins/Enterprise/GLTFExporter/MigratedFromGit/Source/GLTFExporter/Public/Json/GLTFJsonBackdrop.h
@@ -4,7 +4,7 @@
 
 #include "Json/GLTFJsonCore.h"
 
-struct GLTFEXPORTER_API FGLTFJsonBackdrop : IGLTFJsonObject
+struct GLTFEXPORTER_API FGLTFJsonBackdrop : IGLTFJsonIndexedObject
 {
 	FString Name;
 
@@ -20,8 +20,9 @@ struct GLTFEXPORTER_API FGLTFJsonBackdrop : IGLTFJsonObject
 	float LightingDistanceFactor;
 	bool UseCameraProjection;
 
-	FGLTFJsonBackdrop()
-		: Intensity(0)
+	FGLTFJsonBackdrop(int32 Index = INDEX_NONE)
+		: IGLTFJsonIndexedObject(Index)
+		, Intensity(0)
 		, Size(0)
 		, Angle(0)
 		, ProjectionCenter(FGLTFJsonVector3::Zero)
diff --git a/Engine/Plugins/Enterprise/GLTFExporter/MigratedFromGit/Source/GLTFExporter/Public/Json/GLTFJsonBuffer.h b/Engine/Plugins/Enterprise/GLTFExporter/MigratedFromGit/Source/GLTFExporter/Public/Json/GLTFJsonBuffer.h
--- a/Engine/Plugins/Enterprise/GLTFExporter/MigratedFromGit/Source/GLTFExporter/Public/Json/GLTFJsonBuffer.h
+++ b/Engine/Plugins/Enterprise/GLTFExporter/MigratedFromGit/Source/GLTFExporter/Public/Json/GLTFJsonBuffer.h
@@ -4,15 +4,16 @@
 
 #include "Json/GLTFJsonCore.h"
 
-struct GLTFEXPORTER_API FGLTFJsonBuffer : IGLTFJsonObject
+struct GLTFEXPORTER_API FGLTFJsonBuffer : IGLTFJsonIndexedObject
 {
 	FString Name;
 
 	FString URI;
 	int64   ByteLength;
 
-	FGLTFJsonBuffer()
-		: ByteLength(0)
+	FGLTFJsonBuffer(int32 Index = INDEX_NONE)
+		: IGLTFJsonIndexedObject(Index)
+		, ByteLength(0)
 	{
 	}
 
diff --git a/Engine/Plugins/Enterprise/GLTFExporter/MigratedFromGit/Source/GLTFExporter/Public/Json/GLTFJsonBufferView.h b/Engine/Plugins/Enterprise/GLTFExporter/MigratedFromGit/Source/GLTFExporter/Public/Json/GLTFJsonBufferView.h
--- a/Engine/Plugins/Enterprise/GLTFExporter/MigratedFromGit/Source/GLTFExporter/Public/Json/GLTFJsonBufferView.h
+++ b/Engine/Plugins/Enterprise/GLTFExporter/MigratedFromGit/Source/GLTFExporter/Public/Json/GLTFJsonBufferView.h
@@ -4,7 +4,7 @@
 
 #include "Json/GLTFJsonCore.h"
 
-struct GLTFEXPORTER_API FGLTFJsonBufferView : IGLTFJsonObject
+struct GLTFEXPORTER_API FGLTFJsonBufferView : IGLTFJsonIndexedObject
 {
 	FString Name;
 
@@ -16,8 +16,9 @@ struct GLTFEXPORTER_API FGLTFJsonBufferView : IGLTFJsonObject
 
 	EGLTFJsonBufferTarget Target;
 
-	FGLTFJsonBufferView()
-		: ByteLength(0)
+	FGLTFJsonBufferView(int32 Index = INDEX_NONE)
+		: IGLTFJsonIndexedObject(Index)
+		, ByteLength(0)
 		, ByteOffset(0)
 		, ByteStride(0)
 		, Target(EGLTFJsonBufferTarget::None)
diff --git a/Engine/Plugins/Enterprise/GLTFExporter/MigratedFromGit/Source/GLTFExporter/Public/Json/GLTFJsonCamera.h b/Engine/Plugins/Enterprise/GLTFExporter/MigratedFromGit/Source/GLTFExporter/Public/Json/GLTFJsonCamera.h
--- a/Engine/Plugins/Enterprise/GLTFExporter/MigratedFromGit/Source/GLTFExporter/Public/Json/GLTFJsonCamera.h
+++ b/Engine/Plugins/Enterprise/GLTFExporter/MigratedFromGit/Source/GLTFExporter/Public/Json/GLTFJsonCamera.h
@@ -41,7 +41,7 @@ struct GLTFEXPORTER_API FGLTFJsonPerspective : IGLTFJsonObject
 	virtual void WriteObject(IGLTFJsonWriter& Writer) const override;
 };
 
-struct GLTFEXPORTER_API FGLTFJsonCamera : IGLTFJsonObject
+struct GLTFEXPORTER_API FGLTFJsonCamera : IGLTFJsonIndexedObject
 {
 	FString Name;
 
@@ -51,8 +51,9 @@ struct GLTFEXPORTER_API FGLTFJsonCamera : IGLTFJsonObject
 	FGLTFJsonOrthographic Orthographic;
 	FGLTFJsonPerspective  Perspective;
 
-	FGLTFJsonCamera()
-		: Type(EGLTFJsonCameraType::None)
+	FGLTFJsonCamera(int32 Index = INDEX_NONE)
+		: IGLTFJsonIndexedObject(Index)
+		, Type(EGLTFJsonCameraType::None)
 	{
 	}
 
diff --git a/Engine/Plugins/Enterprise/GLTFExporter/MigratedFromGit/Source/GLTFExporter/Public/Json/GLTFJsonEpicLevelVariantSets.h b/Engine/Plugins/Enterprise/GLTFExporter/MigratedFromGit/Source/GLTFExporter/Public/Json/GLTFJsonEpicLevelVariantSets.h
--- a/Engine/Plugins/Enterprise/GLTFExporter/MigratedFromGit/Source/GLTFExporter/Public/Json/GLTFJsonEpicLevelVariantSets.h
+++ b/Engine/Plugins/Enterprise/GLTFExporter/MigratedFromGit/Source/GLTFExporter/Public/Json/GLTFJsonEpicLevelVariantSets.h
@@ -49,11 +49,16 @@ struct GLTFEXPORTER_API FGLTFJsonEpicVariantSet : IGLTFJsonObject
 	virtual void WriteObject(IGLTFJsonWriter& Writer) const override;
 };
 
-struct GLTFEXPORTER_API FGLTFJsonEpicLevelVariantSets : IGLTFJsonObject
+struct GLTFEXPORTER_API FGLTFJsonEpicLevelVariantSets : IGLTFJsonIndexedObject
 {
 	FString Name;
 
 	TArray<FGLTFJsonEpicVariantSet> VariantSets;
 
+	FGLTFJsonEpicLevelVariantSets(int32 Index = INDEX_NONE)
+		: IGLTFJsonIndexedObject(Index)
+	{
+	}
+
 	virtual void WriteObject(IGLTFJsonWriter& Writer) const override;
 };
diff --git a/Engine/Plugins/Enterprise/GLTFExporter/MigratedFromGit/Source/GLTFExporter/Public/Json/GLTFJsonHotspot.h b/Engine/Plugins/Enterprise/GLTFExporter/MigratedFromGit/Source/GLTFExporter/Public/Json/GLTFJsonHotspot.h
--- a/Engine/Plugins/Enterprise/GLTFExporter/MigratedFromGit/Source/GLTFExporter/Public/Json/GLTFJsonHotspot.h
+++ b/Engine/Plugins/Enterprise/GLTFExporter/MigratedFromGit/Source/GLTFExporter/Public/Json/GLTFJsonHotspot.h
@@ -4,7 +4,7 @@
 
 #include "Json/GLTFJsonCore.h"
 
-struct GLTFEXPORTER_API FGLTFJsonHotspot : IGLTFJsonObject
+struct GLTFEXPORTER_API FGLTFJsonHotspot : IGLTFJsonIndexedObject
 {
 	FString Name;
 
@@ -14,5 +14,10 @@ struct GLTFEXPORTER_API FGLTFJsonHotspot : IGLTFJsonObject
 	FGLTFJsonTextureIndex   ToggledImage;
 	FGLTFJsonTextureIndex   ToggledHoveredImage;
 
+	FGLTFJsonHotspot(int32 Index = INDEX_NONE)
+		: IGLTFJsonIndexedObject(Index)
+	{
+	}
+
 	virtual void WriteObject(IGLTFJsonWriter& Writer) const override;
 };
diff --git a/Engine/Plugins/Enterprise/GLTFExporter/MigratedFromGit/Source/GLTFExporter/Public/Json/GLTFJsonImage.h b/Engine/Plugins/Enterprise/GLTFExporter/MigratedFromGit/Source/GLTFExporter/Public/Json/GLTFJsonImage.h
--- a/Engine/Plugins/Enterprise/GLTFExporter/MigratedFromGit/Source/GLTFExporter/Public/Json/GLTFJsonImage.h
+++ b/Engine/Plugins/Enterprise/GLTFExporter/MigratedFromGit/Source/GLTFExporter/Public/Json/GLTFJsonImage.h
@@ -4,7 +4,7 @@
 
 #include "Json/GLTFJsonCore.h"
 
-struct GLTFEXPORTER_API FGLTFJsonImage : IGLTFJsonObject
+struct GLTFEXPORTER_API FGLTFJsonImage : IGLTFJsonIndexedObject
 {
 	FString Name;
 	FString Uri;
@@ -13,8 +13,9 @@ struct GLTFEXPORTER_API FGLTFJsonImage : IGLTFJsonObject
 
 	FGLTFJsonBufferViewIndex BufferView;
 
-	FGLTFJsonImage()
-		: MimeType(EGLTFJsonMimeType::None)
+	FGLTFJsonImage(int32 Index = INDEX_NONE)
+		: IGLTFJsonIndexedObject(Index)
+		, MimeType(EGLTFJsonMimeType::None)
 	{
 	}
 
diff --git a/Engine/Plugins/Enterprise/GLTFExporter/MigratedFromGit/Source/GLTFExporter/Public/Json/GLTFJsonIndex.h b/Engine/Plugins/Enterprise/GLTFExporter/MigratedFromGit/Source/GLTFExporter/Public/Json/GLTFJsonIndex.h
--- a/Engine/Plugins/Enterprise/GLTFExporter/MigratedFromGit/Source/GLTFExporter/Public/Json/GLTFJsonIndex.h
+++ b/Engine/Plugins/Enterprise/GLTFExporter/MigratedFromGit/Source/GLTFExporter/Public/Json/GLTFJsonIndex.h
@@ -12,7 +12,7 @@ struct GLTFEXPORTER_API FGLTFJsonIndex
 	{
 	}
 
-	explicit FGLTFJsonIndex(int32 Value)
+	FGLTFJsonIndex(int32 Value)
 		: Value(Value)
 	{
 	}
diff --git a/Engine/Plugins/Enterprise/GLTFExporter/MigratedFromGit/Source/GLTFExporter/Public/Json/GLTFJsonKhrMaterialVariant.h b/Engine/Plugins/Enterprise/GLTFExporter/MigratedFromGit/Source/GLTFExporter/Public/Json/GLTFJsonKhrMaterialVariant.h
--- a/Engine/Plugins/Enterprise/GLTFExporter/MigratedFromGit/Source/GLTFExporter/Public/Json/GLTFJsonKhrMaterialVariant.h
+++ b/Engine/Plugins/Enterprise/GLTFExporter/MigratedFromGit/Source/GLTFExporter/Public/Json/GLTFJsonKhrMaterialVariant.h
@@ -12,9 +12,14 @@ struct GLTFEXPORTER_API FGLTFJsonKhrMaterialVariantMapping : IGLTFJsonObject
 	virtual void WriteObject(IGLTFJsonWriter& Writer) const override;
 };
 
-struct GLTFEXPORTER_API FGLTFJsonKhrMaterialVariant : IGLTFJsonObject
+struct GLTFEXPORTER_API FGLTFJsonKhrMaterialVariant : IGLTFJsonIndexedObject
 {
 	FString Name;
 
+	FGLTFJsonKhrMaterialVariant(int32 Index = INDEX_NONE)
+		: IGLTFJsonIndexedObject(Index)
+	{
+	}
+
 	virtual void WriteObject(IGLTFJsonWriter& Writer) const override;
 };
diff --git a/Engine/Plugins/Enterprise/GLTFExporter/MigratedFromGit/Source/GLTFExporter/Public/Json/GLTFJsonLight.h b/Engine/Plugins/Enterprise/GLTFExporter/MigratedFromGit/Source/GLTFExporter/Public/Json/GLTFJsonLight.h
--- a/Engine/Plugins/Enterprise/GLTFExporter/MigratedFromGit/Source/GLTFExporter/Public/Json/GLTFJsonLight.h
+++ b/Engine/Plugins/Enterprise/GLTFExporter/MigratedFromGit/Source/GLTFExporter/Public/Json/GLTFJsonLight.h
@@ -18,7 +18,7 @@ struct GLTFEXPORTER_API FGLTFJsonSpotLight : IGLTFJsonObject
 	virtual void WriteObject(IGLTFJsonWriter& Writer) const override;
 };
 
-struct GLTFEXPORTER_API FGLTFJsonLight : IGLTFJsonObject
+struct GLTFEXPORTER_API FGLTFJsonLight : IGLTFJsonIndexedObject
 {
 	FString Name;
 
@@ -31,8 +31,9 @@ struct GLTFEXPORTER_API FGLTFJsonLight : IGLTFJsonObject
 
 	FGLTFJsonSpotLight Spot;
 
-	FGLTFJsonLight()
-		: Type(EGLTFJsonLightType::None)
+	FGLTFJsonLight(int32 Index = INDEX_NONE)
+		: IGLTFJsonIndexedObject(Index)
+		, Type(EGLTFJsonLightType::None)
 		, Color(FGLTFJsonColor3::White)
 		, Intensity(1)
 		, Range(0)
diff --git a/Engine/Plugins/Enterprise/GLTFExporter/MigratedFromGit/Source/GLTFExporter/Public/Json/GLTFJsonLightMap.h b/Engine/Plugins/Enterprise/GLTFExporter/MigratedFromGit/Source/GLTFExporter/Public/Json/GLTFJsonLightMap.h
--- a/Engine/Plugins/Enterprise/GLTFExporter/MigratedFromGit/Source/GLTFExporter/Public/Json/GLTFJsonLightMap.h
+++ b/Engine/Plugins/Enterprise/GLTFExporter/MigratedFromGit/Source/GLTFExporter/Public/Json/GLTFJsonLightMap.h
@@ -5,7 +5,7 @@
 #include "Json/GLTFJsonCore.h"
 #include "Json/GLTFJsonMaterial.h"
 
-struct GLTFEXPORTER_API FGLTFJsonLightMap : IGLTFJsonObject
+struct GLTFEXPORTER_API FGLTFJsonLightMap : IGLTFJsonIndexedObject
 {
 	FString              Name;
 	FGLTFJsonTextureInfo Texture;
@@ -13,8 +13,9 @@ struct GLTFEXPORTER_API FGLTFJsonLightMap : IGLTFJsonObject
 	FGLTFJsonVector4     LightMapAdd;
 	FGLTFJsonVector4     CoordinateScaleBias;
 
-	FGLTFJsonLightMap()
-		: LightMapScale(FGLTFJsonVector4::One)
+	FGLTFJsonLightMap(int32 Index = INDEX_NONE)
+		: IGLTFJsonIndexedObject(Index)
+		, LightMapScale(FGLTFJsonVector4::One)
 		, LightMapAdd(FGLTFJsonVector4::Zero)
 		, CoordinateScaleBias({ 1, 1, 0, 0 })
 	{
diff --git a/Engine/Plugins/Enterprise/GLTFExporter/MigratedFromGit/Source/GLTFExporter/Public/Json/GLTFJsonMaterial.h b/Engine/Plugins/Enterprise/GLTFExporter/MigratedFromGit/Source/GLTFExporter/Public/Json/GLTFJsonMaterial.h
--- a/Engine/Plugins/Enterprise/GLTFExporter/MigratedFromGit/Source/GLTFExporter/Public/Json/GLTFJsonMaterial.h
+++ b/Engine/Plugins/Enterprise/GLTFExporter/MigratedFromGit/Source/GLTFExporter/Public/Json/GLTFJsonMaterial.h
@@ -82,7 +82,7 @@ struct GLTFEXPORTER_API FGLTFJsonClearCoatExtension : IGLTFJsonObject
 	virtual void WriteObject(IGLTFJsonWriter& Writer) const override;
 };
 
-struct GLTFEXPORTER_API FGLTFJsonMaterial : IGLTFJsonObject
+struct GLTFEXPORTER_API FGLTFJsonMaterial : IGLTFJsonIndexedObject
 {
 	FString Name;
 
@@ -105,8 +105,9 @@ struct GLTFEXPORTER_API FGLTFJsonMaterial : IGLTFJsonObject
 
 	FGLTFJsonClearCoatExtension ClearCoat;
 
-	FGLTFJsonMaterial()
-		: ShadingModel(EGLTFJsonShadingModel::Default)
+	FGLTFJsonMaterial(int32 Index = INDEX_NONE)
+		: IGLTFJsonIndexedObject(Index)
+		, ShadingModel(EGLTFJsonShadingModel::Default)
 		, EmissiveFactor(FGLTFJsonColor3::Black)
 		, AlphaMode(EGLTFJsonAlphaMode::Opaque)
 		, AlphaCutoff(0.5f)
diff --git a/Engine/Plugins/Enterprise/GLTFExporter/MigratedFromGit/Source/GLTFExporter/Public/Json/GLTFJsonMesh.h b/Engine/Plugins/Enterprise/GLTFExporter/MigratedFromGit/Source/GLTFExporter/Public/Json/GLTFJsonMesh.h
--- a/Engine/Plugins/Enterprise/GLTFExporter/MigratedFromGit/Source/GLTFExporter/Public/Json/GLTFJsonMesh.h
+++ b/Engine/Plugins/Enterprise/GLTFExporter/MigratedFromGit/Source/GLTFExporter/Public/Json/GLTFJsonMesh.h
@@ -36,11 +36,16 @@ struct GLTFEXPORTER_API FGLTFJsonPrimitive : IGLTFJsonObject
 	virtual void WriteObject(IGLTFJsonWriter& Writer) const override;
 };
 
-struct GLTFEXPORTER_API FGLTFJsonMesh : IGLTFJsonObject
+struct GLTFEXPORTER_API FGLTFJsonMesh : IGLTFJsonIndexedObject
 {
 	FString Name;
 
 	TArray<FGLTFJsonPrimitive> Primitives;
 
+	FGLTFJsonMesh(int32 Index = INDEX_NONE)
+		: IGLTFJsonIndexedObject(Index)
+	{
+	}
+
 	virtual void WriteObject(IGLTFJsonWriter& Writer) const override;
 };
diff --git a/Engine/Plugins/Enterprise/GLTFExporter/MigratedFromGit/Source/GLTFExporter/Public/Json/GLTFJsonNode.h b/Engine/Plugins/Enterprise/GLTFExporter/MigratedFromGit/Source/GLTFExporter/Public/Json/GLTFJsonNode.h
--- a/Engine/Plugins/Enterprise/GLTFExporter/MigratedFromGit/Source/GLTFExporter/Public/Json/GLTFJsonNode.h
+++ b/Engine/Plugins/Enterprise/GLTFExporter/MigratedFromGit/Source/GLTFExporter/Public/Json/GLTFJsonNode.h
@@ -4,7 +4,7 @@
 
 #include "Json/GLTFJsonCore.h"
 
-struct GLTFEXPORTER_API FGLTFJsonNode : IGLTFJsonObject
+struct GLTFEXPORTER_API FGLTFJsonNode : IGLTFJsonIndexedObject
 {
 	FString Name;
 
@@ -25,8 +25,9 @@ struct GLTFEXPORTER_API FGLTFJsonNode : IGLTFJsonObject
 
 	TArray<FGLTFJsonNodeIndex> Children;
 
-	FGLTFJsonNode()
-		: Translation(FGLTFJsonVector3::Zero)
+	FGLTFJsonNode(int32 Index = INDEX_NONE)
+		: IGLTFJsonIndexedObject(Index)
+		, Translation(FGLTFJsonVector3::Zero)
 		, Rotation(FGLTFJsonQuaternion::Identity)
 		, Scale(FGLTFJsonVector3::One)
 	{
diff --git a/Engine/Plugins/Enterprise/GLTFExporter/MigratedFromGit/Source/GLTFExporter/Public/Json/GLTFJsonObject.h b/Engine/Plugins/Enterprise/GLTFExporter/MigratedFromGit/Source/GLTFExporter/Public/Json/GLTFJsonObject.h
--- a/Engine/Plugins/Enterprise/GLTFExporter/MigratedFromGit/Source/GLTFExporter/Public/Json/GLTFJsonObject.h
+++ b/Engine/Plugins/Enterprise/GLTFExporter/MigratedFromGit/Source/GLTFExporter/Public/Json/GLTFJsonObject.h
@@ -16,3 +16,15 @@ struct GLTFEXPORTER_API IGLTFJsonObject : IGLTFJsonValue
 
 	virtual void WriteObject(IGLTFJsonWriter& Writer) const = 0;
 };
+
+struct GLTFEXPORTER_API IGLTFJsonIndexedObject : IGLTFJsonObject
+{
+	int32 Index;
+
+protected:
+
+	IGLTFJsonIndexedObject(int32 Index)
+		: Index(Index)
+	{
+	}
+};
diff --git a/Engine/Plugins/Enterprise/GLTFExporter/MigratedFromGit/Source/GLTFExporter/Public/Json/GLTFJsonRoot.h b/Engine/Plugins/Enterprise/GLTFExporter/MigratedFromGit/Source/GLTFExporter/Public/Json/GLTFJsonRoot.h
--- a/Engine/Plugins/Enterprise/GLTFExporter/MigratedFromGit/Source/GLTFExporter/Public/Json/GLTFJsonRoot.h
+++ b/Engine/Plugins/Enterprise/GLTFExporter/MigratedFromGit/Source/GLTFExporter/Public/Json/GLTFJsonRoot.h
@@ -33,33 +33,33 @@ struct GLTFEXPORTER_API FGLTFJsonRoot : IGLTFJsonObject
 
 	FGLTFJsonSceneIndex DefaultScene;
 
-	TArray<TUniquePtr<FGLTFJsonAccessor>>   Accessors;
-	TArray<TUniquePtr<FGLTFJsonAnimation>>  Animations;
-	TArray<TUniquePtr<FGLTFJsonBuffer>>     Buffers;
-	TArray<TUniquePtr<FGLTFJsonBufferView>> BufferViews;
-	TArray<TUniquePtr<FGLTFJsonCamera>>     Cameras;
-	TArray<TUniquePtr<FGLTFJsonMaterial>>   Materials;
-	TArray<TUniquePtr<FGLTFJsonMesh>>       Meshes;
-	TArray<TUniquePtr<FGLTFJsonNode>>       Nodes;
-	TArray<TUniquePtr<FGLTFJsonImage>>      Images;
-	TArray<TUniquePtr<FGLTFJsonSampler>>    Samplers;
-	TArray<TUniquePtr<FGLTFJsonScene>>      Scenes;
-	TArray<TUniquePtr<FGLTFJsonSkin>>       Skins;
-	TArray<TUniquePtr<FGLTFJsonTexture>>    Textures;
-	TArray<TUniquePtr<FGLTFJsonBackdrop>>   Backdrops;
-	TArray<TUniquePtr<FGLTFJsonHotspot>>    Hotspots;
-	TArray<TUniquePtr<FGLTFJsonLight>>      Lights;
-	TArray<TUniquePtr<FGLTFJsonLightMap>>   LightMaps;
-	TArray<TUniquePtr<FGLTFJsonSkySphere>>  SkySpheres;
-	TArray<TUniquePtr<FGLTFJsonEpicLevelVariantSets>>  EpicLevelVariantSets;
-	TArray<TUniquePtr<FGLTFJsonKhrMaterialVariant>>    KhrMaterialVariants;
+	TGLTFJsonIndexedObjectArray<FGLTFJsonAccessor>   Accessors;
+	TGLTFJsonIndexedObjectArray<FGLTFJsonAnimation>  Animations;
+	TGLTFJsonIndexedObjectArray<FGLTFJsonBuffer>     Buffers;
+	TGLTFJsonIndexedObjectArray<FGLTFJsonBufferView> BufferViews;
+	TGLTFJsonIndexedObjectArray<FGLTFJsonCamera>     Cameras;
+	TGLTFJsonIndexedObjectArray<FGLTFJsonMaterial>   Materials;
+	TGLTFJsonIndexedObjectArray<FGLTFJsonMesh>       Meshes;
+	TGLTFJsonIndexedObjectArray<FGLTFJsonNode>       Nodes;
+	TGLTFJsonIndexedObjectArray<FGLTFJsonImage>      Images;
+	TGLTFJsonIndexedObjectArray<FGLTFJsonSampler>    Samplers;
+	TGLTFJsonIndexedObjectArray<FGLTFJsonScene>      Scenes;
+	TGLTFJsonIndexedObjectArray<FGLTFJsonSkin>       Skins;
+	TGLTFJsonIndexedObjectArray<FGLTFJsonTexture>    Textures;
+	TGLTFJsonIndexedObjectArray<FGLTFJsonBackdrop>   Backdrops;
+	TGLTFJsonIndexedObjectArray<FGLTFJsonHotspot>    Hotspots;
+	TGLTFJsonIndexedObjectArray<FGLTFJsonLight>      Lights;
+	TGLTFJsonIndexedObjectArray<FGLTFJsonLightMap>   LightMaps;
+	TGLTFJsonIndexedObjectArray<FGLTFJsonSkySphere>  SkySpheres;
+	TGLTFJsonIndexedObjectArray<FGLTFJsonEpicLevelVariantSets> EpicLevelVariantSets;
+	TGLTFJsonIndexedObjectArray<FGLTFJsonKhrMaterialVariant>   KhrMaterialVariants;
 
 	FGLTFJsonRoot() = default;
 	FGLTFJsonRoot(FGLTFJsonRoot&&) = default;
-	FGLTFJsonRoot &operator=(FGLTFJsonRoot&&) = default;
+	FGLTFJsonRoot& operator=(FGLTFJsonRoot&&) = default;
 
 	FGLTFJsonRoot(const FGLTFJsonRoot&) = delete;
-	FGLTFJsonRoot &operator=(const FGLTFJsonRoot&) = delete;
+	FGLTFJsonRoot& operator=(const FGLTFJsonRoot&) = delete;
 
 	virtual void WriteObject(IGLTFJsonWriter& Writer) const override;
 
diff --git a/Engine/Plugins/Enterprise/GLTFExporter/MigratedFromGit/Source/GLTFExporter/Public/Json/GLTFJsonSampler.h b/Engine/Plugins/Enterprise/GLTFExporter/MigratedFromGit/Source/GLTFExporter/Public/Json/GLTFJsonSampler.h
--- a/Engine/Plugins/Enterprise/GLTFExporter/MigratedFromGit/Source/GLTFExporter/Public/Json/GLTFJsonSampler.h
+++ b/Engine/Plugins/Enterprise/GLTFExporter/MigratedFromGit/Source/GLTFExporter/Public/Json/GLTFJsonSampler.h
@@ -4,7 +4,7 @@
 
 #include "Json/GLTFJsonCore.h"
 
-struct GLTFEXPORTER_API FGLTFJsonSampler : IGLTFJsonObject
+struct GLTFEXPORTER_API FGLTFJsonSampler : IGLTFJsonIndexedObject
 {
 	FString Name;
 
@@ -14,8 +14,9 @@ struct GLTFEXPORTER_API FGLTFJsonSampler : IGLTFJsonObject
 	EGLTFJsonTextureWrap WrapS;
 	EGLTFJsonTextureWrap WrapT;
 
-	FGLTFJsonSampler()
-		: MinFilter(EGLTFJsonTextureFilter::None)
+	FGLTFJsonSampler(int32 Index = INDEX_NONE)
+		: IGLTFJsonIndexedObject(Index)
+		, MinFilter(EGLTFJsonTextureFilter::None)
 		, MagFilter(EGLTFJsonTextureFilter::None)
 		, WrapS(EGLTFJsonTextureWrap::Repeat)
 		, WrapT(EGLTFJsonTextureWrap::Repeat)
diff --git a/Engine/Plugins/Enterprise/GLTFExporter/MigratedFromGit/Source/GLTFExporter/Public/Json/GLTFJsonScene.h b/Engine/Plugins/Enterprise/GLTFExporter/MigratedFromGit/Source/GLTFExporter/Public/Json/GLTFJsonScene.h
--- a/Engine/Plugins/Enterprise/GLTFExporter/MigratedFromGit/Source/GLTFExporter/Public/Json/GLTFJsonScene.h
+++ b/Engine/Plugins/Enterprise/GLTFExporter/MigratedFromGit/Source/GLTFExporter/Public/Json/GLTFJsonScene.h
@@ -4,12 +4,17 @@
 
 #include "Json/GLTFJsonCore.h"
 
-struct GLTFEXPORTER_API FGLTFJsonScene : IGLTFJsonObject
+struct GLTFEXPORTER_API FGLTFJsonScene : IGLTFJsonIndexedObject
 {
 	FString Name;
 
 	TArray<FGLTFJsonNodeIndex> Nodes;
 	TArray<FGLTFJsonEpicLevelVariantSetsIndex>  EpicLevelVariantSets;
 
+	FGLTFJsonScene(int32 Index = INDEX_NONE)
+		: IGLTFJsonIndexedObject(Index)
+	{
+	}
+
 	virtual void WriteObject(IGLTFJsonWriter& Writer) const override;
 };
diff --git a/Engine/Plugins/Enterprise/GLTFExporter/MigratedFromGit/Source/GLTFExporter/Public/Json/GLTFJsonSkin.h b/Engine/Plugins/Enterprise/GLTFExporter/MigratedFromGit/Source/GLTFExporter/Public/Json/GLTFJsonSkin.h
--- a/Engine/Plugins/Enterprise/GLTFExporter/MigratedFromGit/Source/GLTFExporter/Public/Json/GLTFJsonSkin.h
+++ b/Engine/Plugins/Enterprise/GLTFExporter/MigratedFromGit/Source/GLTFExporter/Public/Json/GLTFJsonSkin.h
@@ -4,7 +4,7 @@
 
 #include "Json/GLTFJsonCore.h"
 
-struct GLTFEXPORTER_API FGLTFJsonSkin : IGLTFJsonObject
+struct GLTFEXPORTER_API FGLTFJsonSkin : IGLTFJsonIndexedObject
 {
 	FString Name;
 
@@ -13,5 +13,10 @@ struct GLTFEXPORTER_API FGLTFJsonSkin : IGLTFJsonObject
 
 	TArray<FGLTFJsonNodeIndex> Joints;
 
+	FGLTFJsonSkin(int32 Index = INDEX_NONE)
+		: IGLTFJsonIndexedObject(Index)
+	{
+	}
+
 	virtual void WriteObject(IGLTFJsonWriter& Writer) const override;
 };
diff --git a/Engine/Plugins/Enterprise/GLTFExporter/MigratedFromGit/Source/GLTFExporter/Public/Json/GLTFJsonSkySphere.h b/Engine/Plugins/Enterprise/GLTFExporter/MigratedFromGit/Source/GLTFExporter/Public/Json/GLTFJsonSkySphere.h
--- a/Engine/Plugins/Enterprise/GLTFExporter/MigratedFromGit/Source/GLTFExporter/Public/Json/GLTFJsonSkySphere.h
+++ b/Engine/Plugins/Enterprise/GLTFExporter/MigratedFromGit/Source/GLTFExporter/Public/Json/GLTFJsonSkySphere.h
@@ -22,7 +22,7 @@ struct GLTFEXPORTER_API FGLTFJsonSkySphereColorCurve : IGLTFJsonArray
 	virtual void WriteArray(IGLTFJsonWriter& Writer) const override;
 };
 
-struct GLTFEXPORTER_API FGLTFJsonSkySphere : IGLTFJsonObject
+struct GLTFEXPORTER_API FGLTFJsonSkySphere : IGLTFJsonIndexedObject
 {
 	FString Name;
 
@@ -56,8 +56,9 @@ struct GLTFEXPORTER_API FGLTFJsonSkySphere : IGLTFJsonObject
 
 	FGLTFJsonVector3 Scale;
 
-	FGLTFJsonSkySphere()
-		: SunHeight(0)
+	FGLTFJsonSkySphere(int32 Index = INDEX_NONE)
+		: IGLTFJsonIndexedObject(Index)
+		, SunHeight(0)
 		, SunBrightness(0)
 		, StarsBrightness(0)
 		, CloudSpeed(0)
diff --git a/Engine/Plugins/Enterprise/GLTFExporter/MigratedFromGit/Source/GLTFExporter/Public/Json/GLTFJsonTexture.h b/Engine/Plugins/Enterprise/GLTFExporter/MigratedFromGit/Source/GLTFExporter/Public/Json/GLTFJsonTexture.h
--- a/Engine/Plugins/Enterprise/GLTFExporter/MigratedFromGit/Source/GLTFExporter/Public/Json/GLTFJsonTexture.h
+++ b/Engine/Plugins/Enterprise/GLTFExporter/MigratedFromGit/Source/GLTFExporter/Public/Json/GLTFJsonTexture.h
@@ -4,7 +4,7 @@
 
 #include "Json/GLTFJsonCore.h"
 
-struct GLTFEXPORTER_API FGLTFJsonTexture : IGLTFJsonObject
+struct GLTFEXPORTER_API FGLTFJsonTexture : IGLTFJsonIndexedObject
 {
 	FString Name;
 
@@ -14,8 +14,9 @@ struct GLTFEXPORTER_API FGLTFJsonTexture : IGLTFJsonObject
 
 	EGLTFJsonHDREncoding Encoding;
 
-	FGLTFJsonTexture()
-		: Encoding(EGLTFJsonHDREncoding::None)
+	FGLTFJsonTexture(int32 Index = INDEX_NONE)
+		: IGLTFJsonIndexedObject(Index)
+		, Encoding(EGLTFJsonHDREncoding::None)
 	{
 	}
 
