[GLTFExporter] refactor: replace indices with pointers in conversion tasks

commit 3b75848908b7b9931f94447b2a84dea6a23d9338
Date: Sun, 13 Mar 2022 13:15:12 +0100

    refactor: replace indices with pointers in conversion tasks
    (cherry picked from commit fd22e986056bfaa4700a8ae74b06fe29fec13a37​)

#skipci

diff --git a/Engine/Plugins/Enterprise/GLTFExporter/MigratedFromGit/Source/GLTFExporter/Private/Builders/GLTFJsonBuilder.cpp b/Engine/Plugins/Enterprise/GLTFExporter/MigratedFromGit/Source/GLTFExporter/Private/Builders/GLTFJsonBuilder.cpp
--- a/Engine/Plugins/Enterprise/GLTFExporter/MigratedFromGit/Source/GLTFExporter/Private/Builders/GLTFJsonBuilder.cpp
+++ b/Engine/Plugins/Enterprise/GLTFExporter/MigratedFromGit/Source/GLTFExporter/Private/Builders/GLTFJsonBuilder.cpp
@@ -41,6 +41,130 @@ void FGLTFJsonBuilder::AddExtension(EGLTFJsonExtension Extension, bool bIsRequir
 	}
 }
 
+FGLTFJsonAccessor* FGLTFJsonBuilder::AddAccessor()
+{
+	return JsonRoot.Accessors.Add();
+}
+
+FGLTFJsonAnimation* FGLTFJsonBuilder::AddAnimation()
+{
+	return JsonRoot.Animations.Add();
+}
+
+FGLTFJsonBuffer* FGLTFJsonBuilder::AddBuffer()
+{
+	return JsonRoot.Buffers.Add();
+}
+
+FGLTFJsonBufferView* FGLTFJsonBuilder::AddBufferView()
+{
+	return JsonRoot.BufferViews.Add();
+}
+
+FGLTFJsonCamera* FGLTFJsonBuilder::AddCamera()
+{
+	return JsonRoot.Cameras.Add();
+}
+
+FGLTFJsonImage* FGLTFJsonBuilder::AddImage()
+{
+	return JsonRoot.Images.Add();
+}
+
+FGLTFJsonMaterial* FGLTFJsonBuilder::AddMaterial()
+{
+	return JsonRoot.Materials.Add();
+}
+
+FGLTFJsonMesh* FGLTFJsonBuilder::AddMesh()
+{
+	return JsonRoot.Meshes.Add();
+}
+
+FGLTFJsonNode* FGLTFJsonBuilder::AddNode()
+{
+	return JsonRoot.Nodes.Add();
+}
+
+FGLTFJsonSampler* FGLTFJsonBuilder::AddSampler()
+{
+	return JsonRoot.Samplers.Add();
+}
+
+FGLTFJsonScene* FGLTFJsonBuilder::AddScene()
+{
+	return JsonRoot.Scenes.Add();
+}
+
+FGLTFJsonSkin* FGLTFJsonBuilder::AddSkin()
+{
+	return JsonRoot.Skins.Add();
+}
+
+FGLTFJsonTexture* FGLTFJsonBuilder::AddTexture()
+{
+	return JsonRoot.Textures.Add();
+}
+
+FGLTFJsonBackdrop* FGLTFJsonBuilder::AddBackdrop()
+{
+	return JsonRoot.Backdrops.Add();
+}
+
+FGLTFJsonHotspot* FGLTFJsonBuilder::AddHotspot()
+{
+	return JsonRoot.Hotspots.Add();
+}
+
+FGLTFJsonLight* FGLTFJsonBuilder::AddLight()
+{
+	return JsonRoot.Lights.Add();
+}
+
+FGLTFJsonLightMap* FGLTFJsonBuilder::AddLightMap()
+{
+	return JsonRoot.LightMaps.Add();
+}
+
+FGLTFJsonSkySphere* FGLTFJsonBuilder::AddSkySphere()
+{
+	return JsonRoot.SkySpheres.Add();
+}
+
+FGLTFJsonEpicLevelVariantSets* FGLTFJsonBuilder::AddEpicLevelVariantSets()
+{
+	return JsonRoot.EpicLevelVariantSets.Add();
+}
+
+FGLTFJsonKhrMaterialVariant* FGLTFJsonBuilder::AddKhrMaterialVariant()
+{
+	return JsonRoot.KhrMaterialVariants.Add();
+}
+
+FGLTFJsonNode* FGLTFJsonBuilder::AddChildNode(FGLTFJsonNodeIndex ParentIndex)
+{
+	FGLTFJsonNode* Child = AddNode();
+
+	if (ParentIndex != INDEX_NONE)
+	{
+		GetNode(ParentIndex).Children.Add(Child->Index);
+	}
+
+	return Child;
+}
+
+FGLTFJsonNode* FGLTFJsonBuilder::AddChildComponentNode(FGLTFJsonNodeIndex ParentIndex)
+{
+	FGLTFJsonNode* Child = AddChildNode(ParentIndex);
+
+	if (ParentIndex != INDEX_NONE)
+	{
+		GetNode(ParentIndex).ComponentNode = Child->Index;
+	}
+
+	return Child;
+}
+
 FGLTFJsonAccessorIndex FGLTFJsonBuilder::AddAccessor(const FGLTFJsonAccessor& JsonAccessor)
 {
 	return FGLTFJsonAccessorIndex(JsonRoot.Accessors.Add(JsonAccessor)->Index);
diff --git a/Engine/Plugins/Enterprise/GLTFExporter/MigratedFromGit/Source/GLTFExporter/Private/Builders/GLTFJsonBuilder.h b/Engine/Plugins/Enterprise/GLTFExporter/MigratedFromGit/Source/GLTFExporter/Private/Builders/GLTFJsonBuilder.h
--- a/Engine/Plugins/Enterprise/GLTFExporter/MigratedFromGit/Source/GLTFExporter/Private/Builders/GLTFJsonBuilder.h
+++ b/Engine/Plugins/Enterprise/GLTFExporter/MigratedFromGit/Source/GLTFExporter/Private/Builders/GLTFJsonBuilder.h
@@ -19,29 +19,53 @@ public:
 
 	void AddExtension(EGLTFJsonExtension Extension, bool bIsRequired = false);
 
-	FGLTFJsonAccessorIndex AddAccessor(const FGLTFJsonAccessor& JsonAccessor = {});
-	FGLTFJsonAnimationIndex AddAnimation(const FGLTFJsonAnimation& JsonAnimation = {});
-	FGLTFJsonBufferIndex AddBuffer(const FGLTFJsonBuffer& JsonBuffer = {});
-	FGLTFJsonBufferViewIndex AddBufferView(const FGLTFJsonBufferView& JsonBufferView = {});
-	FGLTFJsonCameraIndex AddCamera(const FGLTFJsonCamera& JsonCamera = {});
-	FGLTFJsonImageIndex AddImage(const FGLTFJsonImage& JsonImage = {});
-	FGLTFJsonMaterialIndex AddMaterial(const FGLTFJsonMaterial& JsonMaterial = {});
-	FGLTFJsonMeshIndex AddMesh(const FGLTFJsonMesh& JsonMesh = {});
-	FGLTFJsonNodeIndex AddNode(const FGLTFJsonNode& JsonNode = {});
-	FGLTFJsonSamplerIndex AddSampler(const FGLTFJsonSampler& JsonSampler = {});
-	FGLTFJsonSceneIndex AddScene(const FGLTFJsonScene& JsonScene = {});
-	FGLTFJsonSkinIndex AddSkin(const FGLTFJsonSkin& JsonSkin = {});
-	FGLTFJsonTextureIndex AddTexture(const FGLTFJsonTexture& JsonTexture = {});
-	FGLTFJsonBackdropIndex AddBackdrop(const FGLTFJsonBackdrop& JsonBackdrop = {});
-	FGLTFJsonHotspotIndex AddHotspot(const FGLTFJsonHotspot& JsonHotspot = {});
-	FGLTFJsonLightIndex AddLight(const FGLTFJsonLight& JsonLight = {});
-	FGLTFJsonLightMapIndex AddLightMap(const FGLTFJsonLightMap& JsonLightMap = {});
-	FGLTFJsonSkySphereIndex AddSkySphere(const FGLTFJsonSkySphere& JsonSkySphere = {});
-	FGLTFJsonEpicLevelVariantSetsIndex AddEpicLevelVariantSets(const FGLTFJsonEpicLevelVariantSets& JsonEpicLevelVariantSets = {});
-	FGLTFJsonKhrMaterialVariantIndex AddKhrMaterialVariant(const FGLTFJsonKhrMaterialVariant& JsonKhrMaterialVariant = {});
-
-	FGLTFJsonNodeIndex AddChildNode(FGLTFJsonNodeIndex ParentNodeIndex, const FGLTFJsonNode& JsonNode = {});
-	FGLTFJsonNodeIndex AddChildComponentNode(FGLTFJsonNodeIndex ParentNodeIndex, const FGLTFJsonNode& JsonNode = {});
+	FGLTFJsonAccessor* AddAccessor();
+	FGLTFJsonAnimation* AddAnimation();
+	FGLTFJsonBuffer* AddBuffer();
+	FGLTFJsonBufferView* AddBufferView();
+	FGLTFJsonCamera* AddCamera();
+	FGLTFJsonImage* AddImage();
+	FGLTFJsonMaterial* AddMaterial();
+	FGLTFJsonMesh* AddMesh();
+	FGLTFJsonNode* AddNode();
+	FGLTFJsonSampler* AddSampler();
+	FGLTFJsonScene* AddScene();
+	FGLTFJsonSkin* AddSkin();
+	FGLTFJsonTexture* AddTexture();
+	FGLTFJsonBackdrop* AddBackdrop();
+	FGLTFJsonHotspot* AddHotspot();
+	FGLTFJsonLight* AddLight();
+	FGLTFJsonLightMap* AddLightMap();
+	FGLTFJsonSkySphere* AddSkySphere();
+	FGLTFJsonEpicLevelVariantSets* AddEpicLevelVariantSets();
+	FGLTFJsonKhrMaterialVariant* AddKhrMaterialVariant();
+
+	FGLTFJsonNode* AddChildNode(FGLTFJsonNodeIndex ParentNodeIndex);
+	FGLTFJsonNode* AddChildComponentNode(FGLTFJsonNodeIndex ParentNodeIndex);
+
+	FGLTFJsonAccessorIndex AddAccessor(const FGLTFJsonAccessor& JsonAccessor);
+	FGLTFJsonAnimationIndex AddAnimation(const FGLTFJsonAnimation& JsonAnimation);
+	FGLTFJsonBufferIndex AddBuffer(const FGLTFJsonBuffer& JsonBuffer);
+	FGLTFJsonBufferViewIndex AddBufferView(const FGLTFJsonBufferView& JsonBufferView);
+	FGLTFJsonCameraIndex AddCamera(const FGLTFJsonCamera& JsonCamera);
+	FGLTFJsonImageIndex AddImage(const FGLTFJsonImage& JsonImage);
+	FGLTFJsonMaterialIndex AddMaterial(const FGLTFJsonMaterial& JsonMaterial);
+	FGLTFJsonMeshIndex AddMesh(const FGLTFJsonMesh& JsonMesh);
+	FGLTFJsonNodeIndex AddNode(const FGLTFJsonNode& JsonNode);
+	FGLTFJsonSamplerIndex AddSampler(const FGLTFJsonSampler& JsonSampler);
+	FGLTFJsonSceneIndex AddScene(const FGLTFJsonScene& JsonScene);
+	FGLTFJsonSkinIndex AddSkin(const FGLTFJsonSkin& JsonSkin);
+	FGLTFJsonTextureIndex AddTexture(const FGLTFJsonTexture& JsonTexture);
+	FGLTFJsonBackdropIndex AddBackdrop(const FGLTFJsonBackdrop& JsonBackdrop);
+	FGLTFJsonHotspotIndex AddHotspot(const FGLTFJsonHotspot& JsonHotspot);
+	FGLTFJsonLightIndex AddLight(const FGLTFJsonLight& JsonLight);
+	FGLTFJsonLightMapIndex AddLightMap(const FGLTFJsonLightMap& JsonLightMap);
+	FGLTFJsonSkySphereIndex AddSkySphere(const FGLTFJsonSkySphere& JsonSkySphere);
+	FGLTFJsonEpicLevelVariantSetsIndex AddEpicLevelVariantSets(const FGLTFJsonEpicLevelVariantSets& JsonEpicLevelVariantSets);
+	FGLTFJsonKhrMaterialVariantIndex AddKhrMaterialVariant(const FGLTFJsonKhrMaterialVariant& JsonKhrMaterialVariant);
+
+	FGLTFJsonNodeIndex AddChildNode(FGLTFJsonNodeIndex ParentNodeIndex, const FGLTFJsonNode& JsonNode);
+	FGLTFJsonNodeIndex AddChildComponentNode(FGLTFJsonNodeIndex ParentNodeIndex, const FGLTFJsonNode& JsonNode);
 
 	const FGLTFJsonRoot& GetRoot() const;
 	FGLTFJsonAccessor& GetAccessor(FGLTFJsonAccessorIndex AccessorIndex);
diff --git a/Engine/Plugins/Enterprise/GLTFExporter/MigratedFromGit/Source/GLTFExporter/Private/Converters/GLTFAnimationConverters.cpp b/Engine/Plugins/Enterprise/GLTFExporter/MigratedFromGit/Source/GLTFExporter/Private/Converters/GLTFAnimationConverters.cpp
--- a/Engine/Plugins/Enterprise/GLTFExporter/MigratedFromGit/Source/GLTFExporter/Private/Converters/GLTFAnimationConverters.cpp
+++ b/Engine/Plugins/Enterprise/GLTFExporter/MigratedFromGit/Source/GLTFExporter/Private/Converters/GLTFAnimationConverters.cpp
@@ -32,9 +32,9 @@ FGLTFJsonAnimationIndex FGLTFAnimationConverter::Convert(FGLTFJsonNodeIndex Root
 		return FGLTFJsonAnimationIndex(INDEX_NONE);
 	}
 
-	const FGLTFJsonAnimationIndex AnimationIndex = Builder.AddAnimation();
-	Builder.SetupTask<FGLTFAnimSequenceTask>(Builder, RootNode, SkeletalMesh, AnimSequence, AnimationIndex);
-	return AnimationIndex;
+	FGLTFJsonAnimation* JsonAnimation = Builder.AddAnimation();
+	Builder.SetupTask<FGLTFAnimSequenceTask>(Builder, RootNode, SkeletalMesh, AnimSequence, JsonAnimation);
+	return JsonAnimation->Index;
 }
 
 FGLTFJsonAnimationIndex FGLTFAnimationDataConverter::Convert(FGLTFJsonNodeIndex RootNode, const USkeletalMeshComponent* SkeletalMeshComponent)
@@ -64,9 +64,9 @@ FGLTFJsonAnimationIndex FGLTFAnimationDataConverter::Convert(FGLTFJsonNodeIndex
 
 FGLTFJsonAnimationIndex FGLTFLevelSequenceConverter::Convert(const ULevel* Level, const ULevelSequence* LevelSequence)
 {
-	const FGLTFJsonAnimationIndex AnimationIndex = Builder.AddAnimation();
-	Builder.SetupTask<FGLTFLevelSequenceTask>(Builder, Level, LevelSequence, AnimationIndex);
-	return AnimationIndex;
+	FGLTFJsonAnimation* JsonAnimation = Builder.AddAnimation();
+	Builder.SetupTask<FGLTFLevelSequenceTask>(Builder, Level, LevelSequence, JsonAnimation);
+	return JsonAnimation->Index;
 }
 
 FGLTFJsonAnimationIndex FGLTFLevelSequenceDataConverter::Convert(const ALevelSequenceActor* LevelSequenceActor)
diff --git a/Engine/Plugins/Enterprise/GLTFExporter/MigratedFromGit/Source/GLTFExporter/Private/Converters/GLTFMaterialConverters.cpp b/Engine/Plugins/Enterprise/GLTFExporter/MigratedFromGit/Source/GLTFExporter/Private/Converters/GLTFMaterialConverters.cpp
--- a/Engine/Plugins/Enterprise/GLTFExporter/MigratedFromGit/Source/GLTFExporter/Private/Converters/GLTFMaterialConverters.cpp
+++ b/Engine/Plugins/Enterprise/GLTFExporter/MigratedFromGit/Source/GLTFExporter/Private/Converters/GLTFMaterialConverters.cpp
@@ -47,9 +47,9 @@ FGLTFJsonMaterialIndex FGLTFMaterialConverter::Convert(const UMaterialInterface*
 {
 	if (Material != FGLTFMaterialUtility::GetDefaultMaterial())
 	{
-		const FGLTFJsonMaterialIndex MaterialIndex = Builder.AddMaterial();
-		Builder.SetupTask<FGLTFMaterialTask>(Builder, UVOverlapChecker, Material, MeshData, SectionIndices, MaterialIndex);
-		return MaterialIndex;
+		FGLTFJsonMaterial* JsonMaterial = Builder.AddMaterial();
+		Builder.SetupTask<FGLTFMaterialTask>(Builder, UVOverlapChecker, Material, MeshData, SectionIndices, JsonMaterial);
+		return JsonMaterial->Index;
 	}
 
 	return FGLTFJsonMaterialIndex(INDEX_NONE); // use default gltf definition
diff --git a/Engine/Plugins/Enterprise/GLTFExporter/MigratedFromGit/Source/GLTFExporter/Private/Converters/GLTFMaterialProxyFactory.cpp b/Engine/Plugins/Enterprise/GLTFExporter/MigratedFromGit/Source/GLTFExporter/Private/Converters/GLTFMaterialProxyFactory.cpp
--- a/Engine/Plugins/Enterprise/GLTFExporter/MigratedFromGit/Source/GLTFExporter/Private/Converters/GLTFMaterialProxyFactory.cpp
+++ b/Engine/Plugins/Enterprise/GLTFExporter/MigratedFromGit/Source/GLTFExporter/Private/Converters/GLTFMaterialProxyFactory.cpp
@@ -257,7 +257,7 @@ TUniquePtr<IGLTFTexture2DConverter> FGLTFMaterialProxyFactory::CreateTextureConv
 
 		virtual FGLTFJsonTextureIndex Convert(const UTexture2D* Texture2D, bool bToSRGB) override
 		{
-			const FGLTFJsonTextureIndex TextureIndex = Factory.Builder.AddTexture();
+			const FGLTFJsonTextureIndex TextureIndex = Factory.Builder.AddTexture({});
 			Factory.Textures.Add(TextureIndex, const_cast<UTexture2D*>(Texture2D));
 			return TextureIndex;
 		}
@@ -287,7 +287,7 @@ TUniquePtr<IGLTFImageConverter> FGLTFMaterialProxyFactory::CreateImageConverter(
 			const FString Filename = FGLTFImageUtility::GetUniqueFilename(Name, TEXT(""), UniqueFilenames);
 			UniqueFilenames.Add(Filename);
 
-			const FGLTFJsonImageIndex ImageIndex = Factory.Builder.AddImage();
+			const FGLTFJsonImageIndex ImageIndex = Factory.Builder.AddImage({});
 			Factory.Images.Add(ImageIndex, { Filename, Type, bIgnoreAlpha, Size, Pixels });
 			return ImageIndex;
 		}
diff --git a/Engine/Plugins/Enterprise/GLTFExporter/MigratedFromGit/Source/GLTFExporter/Private/Converters/GLTFMeshConverters.cpp b/Engine/Plugins/Enterprise/GLTFExporter/MigratedFromGit/Source/GLTFExporter/Private/Converters/GLTFMeshConverters.cpp
--- a/Engine/Plugins/Enterprise/GLTFExporter/MigratedFromGit/Source/GLTFExporter/Private/Converters/GLTFMeshConverters.cpp
+++ b/Engine/Plugins/Enterprise/GLTFExporter/MigratedFromGit/Source/GLTFExporter/Private/Converters/GLTFMeshConverters.cpp
@@ -44,13 +44,12 @@ FGLTFJsonMeshIndex FGLTFStaticMeshConverter::Convert(const UStaticMesh* StaticMe
 	}
 #endif
 
-	FGLTFJsonMesh JsonMesh;
 	const int32 MaterialCount = FGLTFMeshUtility::GetMaterials(StaticMesh).Num();
-	JsonMesh.Primitives.AddDefaulted(MaterialCount);
+	FGLTFJsonMesh* JsonMesh = Builder.AddMesh();
+	JsonMesh->Primitives.AddDefaulted(MaterialCount);
 
-	const FGLTFJsonMeshIndex MeshIndex = Builder.AddMesh(JsonMesh);
-	Builder.SetupTask<FGLTFStaticMeshTask>(Builder, MeshSectionConverter, StaticMesh, StaticMeshComponent, Materials, LODIndex, MeshIndex);
-	return MeshIndex;
+	Builder.SetupTask<FGLTFStaticMeshTask>(Builder, MeshSectionConverter, StaticMesh, StaticMeshComponent, Materials, LODIndex, JsonMesh);
+	return JsonMesh->Index;
 }
 
 void FGLTFSkeletalMeshConverter::Sanitize(const USkeletalMesh*& SkeletalMesh, const USkeletalMeshComponent*& SkeletalMeshComponent, FGLTFMaterialArray& Materials, int32& LODIndex)
@@ -88,11 +87,10 @@ FGLTFJsonMeshIndex FGLTFSkeletalMeshConverter::Convert(const USkeletalMesh* Skel
 	}
 #endif
 
-	FGLTFJsonMesh JsonMesh;
 	const int32 MaterialCount = FGLTFMeshUtility::GetMaterials(SkeletalMesh).Num();
-	JsonMesh.Primitives.AddDefaulted(MaterialCount);
+	FGLTFJsonMesh* JsonMesh = Builder.AddMesh();
+	JsonMesh->Primitives.AddDefaulted(MaterialCount);
 
-	const FGLTFJsonMeshIndex MeshIndex = Builder.AddMesh(JsonMesh);
-	Builder.SetupTask<FGLTFSkeletalMeshTask>(Builder, MeshSectionConverter, SkeletalMesh, SkeletalMeshComponent, Materials, LODIndex, MeshIndex);
-	return MeshIndex;
+	Builder.SetupTask<FGLTFSkeletalMeshTask>(Builder, MeshSectionConverter, SkeletalMesh, SkeletalMeshComponent, Materials, LODIndex, JsonMesh);
+	return JsonMesh->Index;
 }
diff --git a/Engine/Plugins/Enterprise/GLTFExporter/MigratedFromGit/Source/GLTFExporter/Private/Converters/GLTFNodeConverters.cpp b/Engine/Plugins/Enterprise/GLTFExporter/MigratedFromGit/Source/GLTFExporter/Private/Converters/GLTFNodeConverters.cpp
--- a/Engine/Plugins/Enterprise/GLTFExporter/MigratedFromGit/Source/GLTFExporter/Private/Converters/GLTFNodeConverters.cpp
+++ b/Engine/Plugins/Enterprise/GLTFExporter/MigratedFromGit/Source/GLTFExporter/Private/Converters/GLTFNodeConverters.cpp
@@ -123,12 +123,11 @@ FGLTFJsonNodeIndex FGLTFComponentConverter::Convert(const USceneComponent* Scene
 	const FTransform ParentTransform = ParentComponent != nullptr ? ParentComponent->GetSocketTransform(SocketName) : FTransform::Identity;
 	const FTransform RelativeTransform = bIsRootNode ? Transform : Transform.GetRelativeTransform(ParentTransform);
 
-	const FGLTFJsonNodeIndex NodeIndex = Builder.AddChildNode(ParentNodeIndex);
-	FGLTFJsonNode& Node = Builder.GetNode(NodeIndex);
-	Node.Name = FGLTFNameUtility::GetName(SceneComponent);
-	Node.Translation = FGLTFConverterUtility::ConvertPosition(RelativeTransform.GetTranslation(), Builder.ExportOptions->ExportUniformScale);
-	Node.Rotation = FGLTFConverterUtility::ConvertRotation(RelativeTransform.GetRotation());
-	Node.Scale = FGLTFConverterUtility::ConvertScale(RelativeTransform.GetScale3D());
+	FGLTFJsonNode* Node = Builder.AddChildNode(ParentNodeIndex);
+	Node->Name = FGLTFNameUtility::GetName(SceneComponent);
+	Node->Translation = FGLTFConverterUtility::ConvertPosition(RelativeTransform.GetTranslation(), Builder.ExportOptions->ExportUniformScale);
+	Node->Rotation = FGLTFConverterUtility::ConvertRotation(RelativeTransform.GetRotation());
+	Node->Scale = FGLTFConverterUtility::ConvertScale(RelativeTransform.GetScale3D());
 
 	// TODO: don't export invisible components unless visibility is variable due to variant sets
 
@@ -137,28 +136,28 @@ FGLTFJsonNodeIndex FGLTFComponentConverter::Convert(const USceneComponent* Scene
 	{
 		if (const UStaticMeshComponent* StaticMeshComponent = Cast<UStaticMeshComponent>(SceneComponent))
 		{
-			Node.Mesh = Builder.GetOrAddMesh(StaticMeshComponent);
+			Node->Mesh = Builder.GetOrAddMesh(StaticMeshComponent);
 
 			if (Builder.ExportOptions->bExportLightmaps)
 			{
-				Node.LightMap = Builder.GetOrAddLightMap(StaticMeshComponent);
+				Node->LightMap = Builder.GetOrAddLightMap(StaticMeshComponent);
 			}
 		}
 		else if (const USkeletalMeshComponent* SkeletalMeshComponent = Cast<USkeletalMeshComponent>(SceneComponent))
 		{
-			Node.Mesh = Builder.GetOrAddMesh(SkeletalMeshComponent);
+			Node->Mesh = Builder.GetOrAddMesh(SkeletalMeshComponent);
 
 			if (Builder.ExportOptions->bExportVertexSkinWeights)
 			{
 				// TODO: remove need for NodeIndex by adding support for cyclic calls in converter
-				const FGLTFJsonSkinIndex SkinIndex = Builder.GetOrAddSkin(NodeIndex, SkeletalMeshComponent);
+				const FGLTFJsonSkinIndex SkinIndex = Builder.GetOrAddSkin(Node->Index, SkeletalMeshComponent);
 				if (SkinIndex != INDEX_NONE)
 				{
-					Builder.GetNode(NodeIndex).Skin = SkinIndex;
+					Node->Skin = SkinIndex;
 
 					if (Builder.ExportOptions->bExportAnimationSequences)
 					{
-						Builder.GetOrAddAnimation(NodeIndex, SkeletalMeshComponent);
+						Builder.GetOrAddAnimation(Node->Index, SkeletalMeshComponent);
 					}
 				}
 			}
@@ -168,11 +167,10 @@ FGLTFJsonNodeIndex FGLTFComponentConverter::Convert(const USceneComponent* Scene
 			if (Builder.ExportOptions->bExportCameras)
 			{
 				// TODO: conversion of camera direction should be done in separate converter
-				FGLTFJsonNode CameraNode;
-				CameraNode.Name = FGLTFNameUtility::GetName(CameraComponent);
-				CameraNode.Rotation = FGLTFConverterUtility::ConvertCameraDirection();
-				CameraNode.Camera = Builder.GetOrAddCamera(CameraComponent);
-				Builder.AddChildComponentNode(NodeIndex, CameraNode);
+				FGLTFJsonNode* CameraNode = Builder.AddChildComponentNode(Node->Index);
+				CameraNode->Name = FGLTFNameUtility::GetName(CameraComponent);
+				CameraNode->Rotation = FGLTFConverterUtility::ConvertCameraDirection();
+				CameraNode->Camera = Builder.GetOrAddCamera(CameraComponent);
 			}
 		}
 		else if (const ULightComponent* LightComponent = Cast<ULightComponent>(SceneComponent))
@@ -180,16 +178,15 @@ FGLTFJsonNodeIndex FGLTFComponentConverter::Convert(const USceneComponent* Scene
 			if (Builder.ShouldExportLight(LightComponent->Mobility))
 			{
 				// TODO: conversion of light direction should be done in separate converter
-				FGLTFJsonNode LightNode;
-				LightNode.Name = FGLTFNameUtility::GetName(LightComponent);
-				LightNode.Rotation = FGLTFConverterUtility::ConvertLightDirection();
-				LightNode.Light = Builder.GetOrAddLight(LightComponent);
-				Builder.AddChildComponentNode(NodeIndex, LightNode);
+				FGLTFJsonNode* LightNode = Builder.AddChildComponentNode(Node->Index);
+				LightNode->Name = FGLTFNameUtility::GetName(LightComponent);
+				LightNode->Rotation = FGLTFConverterUtility::ConvertLightDirection();
+				LightNode->Light = Builder.GetOrAddLight(LightComponent);
 			}
 		}
 	}
 
-	return NodeIndex;
+	return Node->Index;
 }
 
 FGLTFJsonNodeIndex FGLTFComponentSocketConverter::Convert(const USceneComponent* SceneComponent, FName SocketName)
diff --git a/Engine/Plugins/Enterprise/GLTFExporter/MigratedFromGit/Source/GLTFExporter/Private/Converters/GLTFTextureConverters.cpp b/Engine/Plugins/Enterprise/GLTFExporter/MigratedFromGit/Source/GLTFExporter/Private/Converters/GLTFTextureConverters.cpp
--- a/Engine/Plugins/Enterprise/GLTFExporter/MigratedFromGit/Source/GLTFExporter/Private/Converters/GLTFTextureConverters.cpp
+++ b/Engine/Plugins/Enterprise/GLTFExporter/MigratedFromGit/Source/GLTFExporter/Private/Converters/GLTFTextureConverters.cpp
@@ -20,9 +20,9 @@ FGLTFJsonTextureIndex FGLTFTexture2DConverter::Convert(const UTexture2D* Texture
 {
 	if (Builder.ExportOptions->TextureImageFormat != EGLTFTextureImageFormat::None)
 	{
-		const FGLTFJsonTextureIndex TextureIndex = Builder.AddTexture();
-		Builder.SetupTask<FGLTFTexture2DTask>(Builder, Texture2D, bToSRGB, TextureIndex);
-		return TextureIndex;
+		FGLTFJsonTexture* JsonTexture = Builder.AddTexture();
+		Builder.SetupTask<FGLTFTexture2DTask>(Builder, Texture2D, bToSRGB, JsonTexture);
+		return JsonTexture->Index;;
 	}
 
 	return FGLTFJsonTextureIndex(INDEX_NONE);
@@ -40,9 +40,9 @@ FGLTFJsonTextureIndex FGLTFTextureCubeConverter::Convert(const UTextureCube* Tex
 {
 	if (Builder.ExportOptions->TextureImageFormat != EGLTFTextureImageFormat::None)
 	{
-		const FGLTFJsonTextureIndex TextureIndex = Builder.AddTexture();
-		Builder.SetupTask<FGLTFTextureCubeTask>(Builder, TextureCube, CubeFace, bToSRGB, TextureIndex);
-		return TextureIndex;
+		FGLTFJsonTexture* JsonTexture = Builder.AddTexture();
+		Builder.SetupTask<FGLTFTextureCubeTask>(Builder, TextureCube, CubeFace, bToSRGB, JsonTexture);
+		return JsonTexture->Index;;
 	}
 
 	return FGLTFJsonTextureIndex(INDEX_NONE);
@@ -60,9 +60,9 @@ FGLTFJsonTextureIndex FGLTFTextureRenderTarget2DConverter::Convert(const UTextur
 {
 	if (Builder.ExportOptions->TextureImageFormat != EGLTFTextureImageFormat::None)
 	{
-		const FGLTFJsonTextureIndex TextureIndex = Builder.AddTexture();
-		Builder.SetupTask<FGLTFTextureRenderTarget2DTask>(Builder, RenderTarget2D, bToSRGB, TextureIndex);
-		return TextureIndex;
+		FGLTFJsonTexture* JsonTexture = Builder.AddTexture();
+		Builder.SetupTask<FGLTFTextureRenderTarget2DTask>(Builder, RenderTarget2D, bToSRGB, JsonTexture);
+		return JsonTexture->Index;;
 	}
 
 	return FGLTFJsonTextureIndex(INDEX_NONE);
@@ -80,9 +80,9 @@ FGLTFJsonTextureIndex FGLTFTextureRenderTargetCubeConverter::Convert(const UText
 {
 	if (Builder.ExportOptions->TextureImageFormat != EGLTFTextureImageFormat::None)
 	{
-		const FGLTFJsonTextureIndex TextureIndex = Builder.AddTexture();
-		Builder.SetupTask<FGLTFTextureRenderTargetCubeTask>(Builder, RenderTargetCube, CubeFace, bToSRGB, TextureIndex);
-		return TextureIndex;
+		FGLTFJsonTexture* JsonTexture = Builder.AddTexture();
+		Builder.SetupTask<FGLTFTextureRenderTargetCubeTask>(Builder, RenderTargetCube, CubeFace, bToSRGB, JsonTexture);
+		return JsonTexture->Index;;
 	}
 
 	return FGLTFJsonTextureIndex(INDEX_NONE);
@@ -93,9 +93,9 @@ FGLTFJsonTextureIndex FGLTFTextureLightMapConverter::Convert(const ULightMapText
 #if WITH_EDITOR
 	if (Builder.ExportOptions->TextureImageFormat != EGLTFTextureImageFormat::None)
 	{
-		const FGLTFJsonTextureIndex TextureIndex = Builder.AddTexture();
-		Builder.SetupTask<FGLTFTextureLightMapTask>(Builder, LightMap, TextureIndex);
-		return TextureIndex;
+		FGLTFJsonTexture* JsonTexture = Builder.AddTexture();
+		Builder.SetupTask<FGLTFTextureLightMapTask>(Builder, LightMap, JsonTexture);
+		return JsonTexture->Index;;
 	}
 #endif
 
diff --git a/Engine/Plugins/Enterprise/GLTFExporter/MigratedFromGit/Source/GLTFExporter/Private/Tasks/GLTFAnimationTasks.cpp b/Engine/Plugins/Enterprise/GLTFExporter/MigratedFromGit/Source/GLTFExporter/Private/Tasks/GLTFAnimationTasks.cpp
--- a/Engine/Plugins/Enterprise/GLTFExporter/MigratedFromGit/Source/GLTFExporter/Private/Tasks/GLTFAnimationTasks.cpp
+++ b/Engine/Plugins/Enterprise/GLTFExporter/MigratedFromGit/Source/GLTFExporter/Private/Tasks/GLTFAnimationTasks.cpp
@@ -19,8 +19,7 @@ void FGLTFAnimSequenceTask::Complete()
 {
 	// TODO: bone transforms should be absolute (not relative) according to gltf spec
 
-	FGLTFJsonAnimation& JsonAnimation = Builder.GetAnimation(AnimationIndex);
-	JsonAnimation.Name = AnimSequence->GetName() + TEXT("_") + FString::FromInt(AnimationIndex); // Ensure unique name due to limitation in certain gltf viewers
+	JsonAnimation->Name = AnimSequence->GetName() + TEXT("_") + FString::FromInt(JsonAnimation->Index); // Ensure unique name due to limitation in certain gltf viewers
 
 	TArray<float> Timestamps;
 	FGLTFBoneUtility::GetFrameTimestamps(AnimSequence, Timestamps);
@@ -82,10 +81,10 @@ void FGLTFAnimSequenceTask::Complete()
 			JsonSampler.Interpolation = Interpolation;
 
 			FGLTFJsonAnimationChannel JsonChannel;
-			JsonChannel.Sampler = FGLTFJsonAnimationSamplerIndex(JsonAnimation.Samplers.Add(JsonSampler));
+			JsonChannel.Sampler = FGLTFJsonAnimationSamplerIndex(JsonAnimation->Samplers.Add(JsonSampler));
 			JsonChannel.Target.Path = EGLTFJsonTargetPath::Translation;
 			JsonChannel.Target.Node = NodeIndex;
-			JsonAnimation.Channels.Add(JsonChannel);
+			JsonAnimation->Channels.Add(JsonChannel);
 		}
 
 		{
@@ -110,10 +109,10 @@ void FGLTFAnimSequenceTask::Complete()
 			JsonSampler.Interpolation = Interpolation;
 
 			FGLTFJsonAnimationChannel JsonChannel;
-			JsonChannel.Sampler = FGLTFJsonAnimationSamplerIndex(JsonAnimation.Samplers.Add(JsonSampler));
+			JsonChannel.Sampler = FGLTFJsonAnimationSamplerIndex(JsonAnimation->Samplers.Add(JsonSampler));
 			JsonChannel.Target.Path = EGLTFJsonTargetPath::Rotation;
 			JsonChannel.Target.Node = NodeIndex;
-			JsonAnimation.Channels.Add(JsonChannel);
+			JsonAnimation->Channels.Add(JsonChannel);
 		}
 
 		{
@@ -138,10 +137,10 @@ void FGLTFAnimSequenceTask::Complete()
 			JsonSampler.Interpolation = Interpolation;
 
 			FGLTFJsonAnimationChannel JsonChannel;
-			JsonChannel.Sampler = FGLTFJsonAnimationSamplerIndex(JsonAnimation.Samplers.Add(JsonSampler));
+			JsonChannel.Sampler = FGLTFJsonAnimationSamplerIndex(JsonAnimation->Samplers.Add(JsonSampler));
 			JsonChannel.Target.Path = EGLTFJsonTargetPath::Scale;
 			JsonChannel.Target.Node = NodeIndex;
-			JsonAnimation.Channels.Add(JsonChannel);
+			JsonAnimation->Channels.Add(JsonChannel);
 		}
 	}
 }
@@ -165,8 +164,7 @@ void FGLTFLevelSequenceTask::Complete()
 	int32 FrameOffset = FFrameRate::TransformTime(FFrameTime(MovieScene::DiscreteInclusiveLower(PlaybackRange)), TickResolution, DisplayRate).RoundToFrame().Value;
 	int32 FrameCount = FFrameRate::TransformTime(FFrameTime(FFrameNumber(MovieScene::DiscreteSize(PlaybackRange))), TickResolution, DisplayRate).RoundToFrame().Value + 1;
 
-	FGLTFJsonAnimation& JsonAnimation = Builder.GetAnimation(AnimationIndex);
-	JsonAnimation.Name = Sequence->GetName() + TEXT("_") + FString::FromInt(AnimationIndex); // Ensure unique name due to limitation in certain gltf viewers
+	JsonAnimation->Name = Sequence->GetName() + TEXT("_") + FString::FromInt(JsonAnimation->Index); // Ensure unique name due to limitation in certain gltf viewers
 
 	TArray<float> Timestamps;
 	TArray<FFrameTime> FrameTimes;
@@ -278,10 +276,10 @@ void FGLTFLevelSequenceTask::Complete()
 						JsonSampler.Interpolation = EGLTFJsonInterpolation::Linear;
 
 						FGLTFJsonAnimationChannel JsonChannel;
-						JsonChannel.Sampler = FGLTFJsonAnimationSamplerIndex(JsonAnimation.Samplers.Add(JsonSampler));
+						JsonChannel.Sampler = FGLTFJsonAnimationSamplerIndex(JsonAnimation->Samplers.Add(JsonSampler));
 						JsonChannel.Target.Path = EGLTFJsonTargetPath::Translation;
 						JsonChannel.Target.Node = NodeIndex;
-						JsonAnimation.Channels.Add(JsonChannel);
+						JsonAnimation->Channels.Add(JsonChannel);
 					}
 
 					if (EnumHasAnyFlags(ChannelMask, EMovieSceneTransformChannel::Rotation))
@@ -316,10 +314,10 @@ void FGLTFLevelSequenceTask::Complete()
 						JsonSampler.Interpolation = EGLTFJsonInterpolation::Linear;
 
 						FGLTFJsonAnimationChannel JsonChannel;
-						JsonChannel.Sampler = FGLTFJsonAnimationSamplerIndex(JsonAnimation.Samplers.Add(JsonSampler));
+						JsonChannel.Sampler = FGLTFJsonAnimationSamplerIndex(JsonAnimation->Samplers.Add(JsonSampler));
 						JsonChannel.Target.Path = EGLTFJsonTargetPath::Rotation;
 						JsonChannel.Target.Node = NodeIndex;
-						JsonAnimation.Channels.Add(JsonChannel);
+						JsonAnimation->Channels.Add(JsonChannel);
 					}
 
 					if (EnumHasAnyFlags(ChannelMask, EMovieSceneTransformChannel::Scale))
@@ -351,10 +349,10 @@ void FGLTFLevelSequenceTask::Complete()
 						JsonSampler.Interpolation = EGLTFJsonInterpolation::Linear;
 
 						FGLTFJsonAnimationChannel JsonChannel;
-						JsonChannel.Sampler = FGLTFJsonAnimationSamplerIndex(JsonAnimation.Samplers.Add(JsonSampler));
+						JsonChannel.Sampler = FGLTFJsonAnimationSamplerIndex(JsonAnimation->Samplers.Add(JsonSampler));
 						JsonChannel.Target.Path = EGLTFJsonTargetPath::Scale;
 						JsonChannel.Target.Node = NodeIndex;
-						JsonAnimation.Channels.Add(JsonChannel);
+						JsonAnimation->Channels.Add(JsonChannel);
 					}
 				}
 			}
diff --git a/Engine/Plugins/Enterprise/GLTFExporter/MigratedFromGit/Source/GLTFExporter/Private/Tasks/GLTFAnimationTasks.h b/Engine/Plugins/Enterprise/GLTFExporter/MigratedFromGit/Source/GLTFExporter/Private/Tasks/GLTFAnimationTasks.h
--- a/Engine/Plugins/Enterprise/GLTFExporter/MigratedFromGit/Source/GLTFExporter/Private/Tasks/GLTFAnimationTasks.h
+++ b/Engine/Plugins/Enterprise/GLTFExporter/MigratedFromGit/Source/GLTFExporter/Private/Tasks/GLTFAnimationTasks.h
@@ -11,13 +11,13 @@ class FGLTFAnimSequenceTask : public FGLTFTask
 {
 public:
 
-	FGLTFAnimSequenceTask(FGLTFConvertBuilder& Builder, FGLTFJsonNodeIndex RootNode, const USkeletalMesh* SkeletalMesh, const UAnimSequence* AnimSequence, FGLTFJsonAnimationIndex AnimationIndex)
+	FGLTFAnimSequenceTask(FGLTFConvertBuilder& Builder, FGLTFJsonNodeIndex RootNode, const USkeletalMesh* SkeletalMesh, const UAnimSequence* AnimSequence, FGLTFJsonAnimation* JsonAnimation)
 		: FGLTFTask(EGLTFTaskPriority::Animation)
 		, Builder(Builder)
 		, RootNode(RootNode)
 		, SkeletalMesh(SkeletalMesh)
 		, AnimSequence(AnimSequence)
-		, AnimationIndex(AnimationIndex)
+		, JsonAnimation(JsonAnimation)
 	{
 	}
 
@@ -34,19 +34,19 @@ private:
 	FGLTFJsonNodeIndex RootNode;
 	const USkeletalMesh* SkeletalMesh;
 	const UAnimSequence* AnimSequence;
-	const FGLTFJsonAnimationIndex AnimationIndex;
+	FGLTFJsonAnimation* JsonAnimation;
 };
 
 class FGLTFLevelSequenceTask : public FGLTFTask
 {
 public:
 
-	FGLTFLevelSequenceTask(FGLTFConvertBuilder& Builder, const ULevel* Level, const ULevelSequence* LevelSequence, FGLTFJsonAnimationIndex AnimationIndex)
+	FGLTFLevelSequenceTask(FGLTFConvertBuilder& Builder, const ULevel* Level, const ULevelSequence* LevelSequence, FGLTFJsonAnimation* JsonAnimation)
 		: FGLTFTask(EGLTFTaskPriority::Animation)
 		, Builder(Builder)
 		, Level(Level)
 		, LevelSequence(LevelSequence)
-		, AnimationIndex(AnimationIndex)
+		, JsonAnimation(JsonAnimation)
 	{
 	}
 
@@ -62,5 +62,5 @@ private:
 	FGLTFConvertBuilder& Builder;
 	const ULevel* Level;
 	const ULevelSequence* LevelSequence;
-	const FGLTFJsonAnimationIndex AnimationIndex;
+	FGLTFJsonAnimation* JsonAnimation;
 };
diff --git a/Engine/Plugins/Enterprise/GLTFExporter/MigratedFromGit/Source/GLTFExporter/Private/Tasks/GLTFMaterialTasks.cpp b/Engine/Plugins/Enterprise/GLTFExporter/MigratedFromGit/Source/GLTFExporter/Private/Tasks/GLTFMaterialTasks.cpp
--- a/Engine/Plugins/Enterprise/GLTFExporter/MigratedFromGit/Source/GLTFExporter/Private/Tasks/GLTFMaterialTasks.cpp
+++ b/Engine/Plugins/Enterprise/GLTFExporter/MigratedFromGit/Source/GLTFExporter/Private/Tasks/GLTFMaterialTasks.cpp
@@ -57,80 +57,79 @@ void FGLTFMaterialTask::Complete()
 		// TODO: report warning (non-surface materials not supported, will be treated as surface)
 	}
 
-	FGLTFJsonMaterial& JsonMaterial = Builder.GetMaterial(MaterialIndex);
-	JsonMaterial.Name = GetMaterialName();
-	JsonMaterial.AlphaCutoff = Material->GetOpacityMaskClipValue();
-	JsonMaterial.DoubleSided = Material->IsTwoSided();
+	JsonMaterial->Name = GetMaterialName();
+	JsonMaterial->AlphaCutoff = Material->GetOpacityMaskClipValue();
+	JsonMaterial->DoubleSided = Material->IsTwoSided();
 
 	if (const UMaterialInstance* MaterialInstance = Cast<UMaterialInstance>(Material))
 	{
 		// TODO: remove ugly hack for dynamic instances
-		FGLTFProxyMaterialUtilities::GetOpacityMaskClipValue(MaterialInstance, JsonMaterial.AlphaCutoff);
-		FGLTFProxyMaterialUtilities::GetTwoSided(MaterialInstance, JsonMaterial.DoubleSided);
+		FGLTFProxyMaterialUtilities::GetOpacityMaskClipValue(MaterialInstance, JsonMaterial->AlphaCutoff);
+		FGLTFProxyMaterialUtilities::GetTwoSided(MaterialInstance, JsonMaterial->DoubleSided);
 	}
 
-	ConvertShadingModel(JsonMaterial.ShadingModel);
-	ConvertAlphaMode(JsonMaterial.AlphaMode, JsonMaterial.BlendMode);
+	ConvertShadingModel(JsonMaterial->ShadingModel);
+	ConvertAlphaMode(JsonMaterial->AlphaMode, JsonMaterial->BlendMode);
 
 	if (FGLTFProxyMaterialUtilities::IsProxyMaterial(BaseMaterial))
 	{
-		GetProxyParameters(JsonMaterial);
+		GetProxyParameters(*JsonMaterial);
 		return;
 	}
 
 #if WITH_EDITOR
-	if (JsonMaterial.ShadingModel != EGLTFJsonShadingModel::None)
+	if (JsonMaterial->ShadingModel != EGLTFJsonShadingModel::None)
 	{
-		const FMaterialPropertyEx BaseColorProperty = JsonMaterial.ShadingModel == EGLTFJsonShadingModel::Unlit ? MP_EmissiveColor : MP_BaseColor;
-		const FMaterialPropertyEx OpacityProperty = JsonMaterial.AlphaMode == EGLTFJsonAlphaMode::Mask ? MP_OpacityMask : MP_Opacity;
+		const FMaterialPropertyEx BaseColorProperty = JsonMaterial->ShadingModel == EGLTFJsonShadingModel::Unlit ? MP_EmissiveColor : MP_BaseColor;
+		const FMaterialPropertyEx OpacityProperty = JsonMaterial->AlphaMode == EGLTFJsonAlphaMode::Mask ? MP_OpacityMask : MP_Opacity;
 
 		// TODO: check if a property is active before trying to get it (i.e. Material->IsPropertyActive)
 
-		if (JsonMaterial.AlphaMode == EGLTFJsonAlphaMode::Opaque)
+		if (JsonMaterial->AlphaMode == EGLTFJsonAlphaMode::Opaque)
 		{
-			if (!TryGetConstantColor(JsonMaterial.PBRMetallicRoughness.BaseColorFactor, BaseColorProperty))
+			if (!TryGetConstantColor(JsonMaterial->PBRMetallicRoughness.BaseColorFactor, BaseColorProperty))
 			{
-				if (!TryGetSourceTexture(JsonMaterial.PBRMetallicRoughness.BaseColorTexture, BaseColorProperty, DefaultColorInputMasks))
+				if (!TryGetSourceTexture(JsonMaterial->PBRMetallicRoughness.BaseColorTexture, BaseColorProperty, DefaultColorInputMasks))
 				{
-					if (!TryGetBakedMaterialProperty(JsonMaterial.PBRMetallicRoughness.BaseColorTexture, JsonMaterial.PBRMetallicRoughness.BaseColorFactor, BaseColorProperty, TEXT("BaseColor")))
+					if (!TryGetBakedMaterialProperty(JsonMaterial->PBRMetallicRoughness.BaseColorTexture, JsonMaterial->PBRMetallicRoughness.BaseColorFactor, BaseColorProperty, TEXT("BaseColor")))
 					{
 						Builder.LogWarning(FString::Printf(TEXT("Failed to export %s for material %s"), *BaseColorProperty.ToString(), *Material->GetName()));
 					}
 				}
 			}
 
-			JsonMaterial.PBRMetallicRoughness.BaseColorFactor.A = 1.0f; // make sure base color is opaque
+			JsonMaterial->PBRMetallicRoughness.BaseColorFactor.A = 1.0f; // make sure base color is opaque
 		}
 		else
 		{
-			if (!TryGetBaseColorAndOpacity(JsonMaterial.PBRMetallicRoughness, BaseColorProperty, OpacityProperty))
+			if (!TryGetBaseColorAndOpacity(JsonMaterial->PBRMetallicRoughness, BaseColorProperty, OpacityProperty))
 			{
 				Builder.LogWarning(FString::Printf(TEXT("Failed to export %s and %s for material %s"), *BaseColorProperty.ToString(), *OpacityProperty.ToString(), *Material->GetName()));
 			}
 		}
 
-		if (JsonMaterial.ShadingModel == EGLTFJsonShadingModel::Default || JsonMaterial.ShadingModel == EGLTFJsonShadingModel::ClearCoat)
+		if (JsonMaterial->ShadingModel == EGLTFJsonShadingModel::Default || JsonMaterial->ShadingModel == EGLTFJsonShadingModel::ClearCoat)
 		{
 			const FMaterialPropertyEx MetallicProperty = MP_Metallic;
 			const FMaterialPropertyEx RoughnessProperty = MP_Roughness;
 
-			if (!TryGetMetallicAndRoughness(JsonMaterial.PBRMetallicRoughness, MetallicProperty, RoughnessProperty))
+			if (!TryGetMetallicAndRoughness(JsonMaterial->PBRMetallicRoughness, MetallicProperty, RoughnessProperty))
 			{
 				Builder.LogWarning(FString::Printf(TEXT("Failed to export %s and %s for material %s"), *MetallicProperty.ToString(), *RoughnessProperty.ToString(), *Material->GetName()));
 			}
 
 			const FMaterialPropertyEx EmissiveProperty = MP_EmissiveColor;
-			if (!TryGetEmissive(JsonMaterial, EmissiveProperty))
+			if (!TryGetEmissive(*JsonMaterial, EmissiveProperty))
 			{
 				Builder.LogWarning(FString::Printf(TEXT("Failed to export %s for material %s"), *EmissiveProperty.ToString(), *Material->GetName()));
 			}
 
-			const FMaterialPropertyEx NormalProperty = JsonMaterial.ShadingModel == EGLTFJsonShadingModel::ClearCoat ? FMaterialPropertyEx(TEXT("ClearCoatBottomNormal")) : FMaterialPropertyEx(MP_Normal);
+			const FMaterialPropertyEx NormalProperty = JsonMaterial->ShadingModel == EGLTFJsonShadingModel::ClearCoat ? FMaterialPropertyEx(TEXT("ClearCoatBottomNormal")) : FMaterialPropertyEx(MP_Normal);
 			if (IsPropertyNonDefault(NormalProperty))
 			{
-				if (!TryGetSourceTexture(JsonMaterial.NormalTexture, NormalProperty, DefaultColorInputMasks))
+				if (!TryGetSourceTexture(JsonMaterial->NormalTexture, NormalProperty, DefaultColorInputMasks))
 				{
-					if (!TryGetBakedMaterialProperty(JsonMaterial.NormalTexture, NormalProperty, TEXT("Normal")))
+					if (!TryGetBakedMaterialProperty(JsonMaterial->NormalTexture, NormalProperty, TEXT("Normal")))
 					{
 						Builder.LogWarning(FString::Printf(TEXT("Failed to export %s for material %s"), *NormalProperty.ToString(), *Material->GetName()));
 					}
@@ -140,21 +139,21 @@ void FGLTFMaterialTask::Complete()
 			const FMaterialPropertyEx AmbientOcclusionProperty = MP_AmbientOcclusion;
 			if (IsPropertyNonDefault(AmbientOcclusionProperty))
 			{
-				if (!TryGetSourceTexture(JsonMaterial.OcclusionTexture, AmbientOcclusionProperty, OcclusionInputMasks))
+				if (!TryGetSourceTexture(JsonMaterial->OcclusionTexture, AmbientOcclusionProperty, OcclusionInputMasks))
 				{
-					if (!TryGetBakedMaterialProperty(JsonMaterial.OcclusionTexture, AmbientOcclusionProperty, TEXT("Occlusion")))
+					if (!TryGetBakedMaterialProperty(JsonMaterial->OcclusionTexture, AmbientOcclusionProperty, TEXT("Occlusion")))
 					{
 						Builder.LogWarning(FString::Printf(TEXT("Failed to export %s for material %s"), *AmbientOcclusionProperty.ToString(), *Material->GetName()));
 					}
 				}
 			}
 
-			if (JsonMaterial.ShadingModel == EGLTFJsonShadingModel::ClearCoat)
+			if (JsonMaterial->ShadingModel == EGLTFJsonShadingModel::ClearCoat)
 			{
 				const FMaterialPropertyEx ClearCoatProperty = MP_CustomData0;
 				const FMaterialPropertyEx ClearCoatRoughnessProperty = MP_CustomData1;
 
-				if (!TryGetClearCoatRoughness(JsonMaterial.ClearCoat, ClearCoatProperty, ClearCoatRoughnessProperty))
+				if (!TryGetClearCoatRoughness(JsonMaterial->ClearCoat, ClearCoatProperty, ClearCoatRoughnessProperty))
 				{
 					Builder.LogWarning(FString::Printf(TEXT("Failed to export %s and %s for material %s"), *ClearCoatProperty.ToString(), *ClearCoatRoughnessProperty.ToString(), *Material->GetName()));
 				}
@@ -162,9 +161,9 @@ void FGLTFMaterialTask::Complete()
 				const FMaterialPropertyEx ClearCoatNormalProperty = MP_Normal;
 				if (IsPropertyNonDefault(ClearCoatNormalProperty))
 				{
-					if (!TryGetSourceTexture(JsonMaterial.ClearCoat.ClearCoatNormalTexture, ClearCoatNormalProperty, DefaultColorInputMasks))
+					if (!TryGetSourceTexture(JsonMaterial->ClearCoat.ClearCoatNormalTexture, ClearCoatNormalProperty, DefaultColorInputMasks))
 					{
-						if (!TryGetBakedMaterialProperty(JsonMaterial.ClearCoat.ClearCoatNormalTexture, ClearCoatNormalProperty, TEXT("ClearCoatNormal")))
+						if (!TryGetBakedMaterialProperty(JsonMaterial->ClearCoat.ClearCoatNormalTexture, ClearCoatNormalProperty, TEXT("ClearCoatNormal")))
 						{
 							Builder.LogWarning(FString::Printf(TEXT("Failed to export %s for material %s"), *ClearCoatNormalProperty.ToString(), *Material->GetName()));
 						}
@@ -860,19 +859,19 @@ bool FGLTFMaterialTask::TryGetClearCoatRoughness(FGLTFJsonClearCoatExtension& Ou
 	return true;
 }
 
-bool FGLTFMaterialTask::TryGetEmissive(FGLTFJsonMaterial& JsonMaterial, const FMaterialPropertyEx& EmissiveProperty)
+bool FGLTFMaterialTask::TryGetEmissive(FGLTFJsonMaterial& OutMaterial, const FMaterialPropertyEx& EmissiveProperty)
 {
 	// TODO: right now we allow EmissiveFactor to be > 1.0 to support very bright emission, although it's not valid according to the glTF standard.
 	// We may want to change this behaviour and store factors above 1.0 using a custom extension instead.
 
-	if (TryGetConstantColor(JsonMaterial.EmissiveFactor, MP_EmissiveColor))
+	if (TryGetConstantColor(OutMaterial.EmissiveFactor, MP_EmissiveColor))
 	{
 		return true;
 	}
 
-	if (TryGetSourceTexture(JsonMaterial.EmissiveTexture, EmissiveProperty, DefaultColorInputMasks))
+	if (TryGetSourceTexture(OutMaterial.EmissiveTexture, EmissiveProperty, DefaultColorInputMasks))
 	{
-		JsonMaterial.EmissiveFactor = FGLTFJsonColor3::White;	// make sure texture is not multiplied with black
+		OutMaterial.EmissiveFactor = FGLTFJsonColor3::White;	// make sure texture is not multiplied with black
 		return true;
 	}
 
@@ -885,28 +884,28 @@ bool FGLTFMaterialTask::TryGetEmissive(FGLTFJsonMaterial& JsonMaterial, const FM
 		return false;
 	}
 
-	FGLTFPropertyBakeOutput PropertyBakeOutput = BakeMaterialProperty(EmissiveProperty, JsonMaterial.EmissiveTexture.TexCoord);
+	FGLTFPropertyBakeOutput PropertyBakeOutput = BakeMaterialProperty(EmissiveProperty, OutMaterial.EmissiveTexture.TexCoord);
 	const float EmissiveScale = PropertyBakeOutput.EmissiveScale;
 
 	if (PropertyBakeOutput.bIsConstant)
 	{
 		const FLinearColor EmissiveColor = PropertyBakeOutput.ConstantValue;
-		JsonMaterial.EmissiveFactor = FGLTFConverterUtility::ConvertColor3(EmissiveColor * EmissiveScale, Builder.ExportOptions->bStrictCompliance);
+		OutMaterial.EmissiveFactor = FGLTFConverterUtility::ConvertColor3(EmissiveColor * EmissiveScale, Builder.ExportOptions->bStrictCompliance);
 	}
 	else
 	{
 		if (Builder.ExportOptions->TextureImageFormat == EGLTFTextureImageFormat::None)
 		{
-			JsonMaterial.EmissiveTexture.Index = FGLTFJsonTextureIndex(INDEX_NONE);
+			OutMaterial.EmissiveTexture.Index = FGLTFJsonTextureIndex(INDEX_NONE);
 			return true;
 		}
 
-		if (!StoreBakedPropertyTexture(JsonMaterial.EmissiveTexture, PropertyBakeOutput, TEXT("Emissive")))
+		if (!StoreBakedPropertyTexture(OutMaterial.EmissiveTexture, PropertyBakeOutput, TEXT("Emissive")))
 		{
 			return false;
 		}
 
-		JsonMaterial.EmissiveFactor = { EmissiveScale, EmissiveScale, EmissiveScale };
+		OutMaterial.EmissiveFactor = { EmissiveScale, EmissiveScale, EmissiveScale };
 	}
 
 	return true;
diff --git a/Engine/Plugins/Enterprise/GLTFExporter/MigratedFromGit/Source/GLTFExporter/Private/Tasks/GLTFMaterialTasks.h b/Engine/Plugins/Enterprise/GLTFExporter/MigratedFromGit/Source/GLTFExporter/Private/Tasks/GLTFMaterialTasks.h
--- a/Engine/Plugins/Enterprise/GLTFExporter/MigratedFromGit/Source/GLTFExporter/Private/Tasks/GLTFMaterialTasks.h
+++ b/Engine/Plugins/Enterprise/GLTFExporter/MigratedFromGit/Source/GLTFExporter/Private/Tasks/GLTFMaterialTasks.h
@@ -17,14 +17,14 @@ class FGLTFMaterialTask : public FGLTFTask
 {
 public:
 
-	FGLTFMaterialTask(FGLTFConvertBuilder& Builder, FGLTFUVOverlapChecker& UVOverlapChecker, const UMaterialInterface* Material, const FGLTFMeshData* MeshData, FGLTFIndexArray SectionIndices, FGLTFJsonMaterialIndex MaterialIndex)
+	FGLTFMaterialTask(FGLTFConvertBuilder& Builder, FGLTFUVOverlapChecker& UVOverlapChecker, const UMaterialInterface* Material, const FGLTFMeshData* MeshData, FGLTFIndexArray SectionIndices, FGLTFJsonMaterial* JsonMaterial)
 		: FGLTFTask(EGLTFTaskPriority::Material)
 		, Builder(Builder)
 		, UVOverlapChecker(UVOverlapChecker)
 		, Material(Material)
 		, MeshData(MeshData)
 		, SectionIndices(SectionIndices)
-		, MaterialIndex(MaterialIndex)
+		, JsonMaterial(JsonMaterial)
 	{
 	}
 
@@ -42,7 +42,7 @@ private:
 	const UMaterialInterface* Material;
 	const FGLTFMeshData* MeshData;
 	const FGLTFIndexArray SectionIndices;
-	const FGLTFJsonMaterialIndex MaterialIndex;
+	FGLTFJsonMaterial* JsonMaterial;
 
 	FString GetMaterialName() const;
 	FString GetBakedTextureName(const FString& PropertyName) const;
@@ -63,7 +63,7 @@ private:
 	bool TryGetBaseColorAndOpacity(FGLTFJsonPBRMetallicRoughness& OutPBRParams, const FMaterialPropertyEx& BaseColorProperty, const FMaterialPropertyEx& OpacityProperty);
 	bool TryGetMetallicAndRoughness(FGLTFJsonPBRMetallicRoughness& OutPBRParams, const FMaterialPropertyEx& MetallicProperty, const FMaterialPropertyEx& RoughnessProperty);
 	bool TryGetClearCoatRoughness(FGLTFJsonClearCoatExtension& OutExtParams, const FMaterialPropertyEx& IntensityProperty, const FMaterialPropertyEx& RoughnessProperty);
-	bool TryGetEmissive(FGLTFJsonMaterial& JsonMaterial, const FMaterialPropertyEx& EmissiveProperty);
+	bool TryGetEmissive(FGLTFJsonMaterial& OutMaterial, const FMaterialPropertyEx& EmissiveProperty);
 
 	bool IsPropertyNonDefault(const FMaterialPropertyEx& Property) const;
 	bool TryGetConstantColor(FGLTFJsonColor3& OutValue, const FMaterialPropertyEx& Property) const;
diff --git a/Engine/Plugins/Enterprise/GLTFExporter/MigratedFromGit/Source/GLTFExporter/Private/Tasks/GLTFMeshTasks.cpp b/Engine/Plugins/Enterprise/GLTFExporter/MigratedFromGit/Source/GLTFExporter/Private/Tasks/GLTFMeshTasks.cpp
--- a/Engine/Plugins/Enterprise/GLTFExporter/MigratedFromGit/Source/GLTFExporter/Private/Tasks/GLTFMeshTasks.cpp
+++ b/Engine/Plugins/Enterprise/GLTFExporter/MigratedFromGit/Source/GLTFExporter/Private/Tasks/GLTFMeshTasks.cpp
@@ -109,8 +109,7 @@ namespace
 
 void FGLTFStaticMeshTask::Complete()
 {
-	FGLTFJsonMesh& JsonMesh = Builder.GetMesh(MeshIndex);
-	JsonMesh.Name = StaticMeshComponent != nullptr ? FGLTFNameUtility::GetName(StaticMeshComponent) : StaticMesh->GetName();
+	JsonMesh->Name = StaticMeshComponent != nullptr ? FGLTFNameUtility::GetName(StaticMeshComponent) : StaticMesh->GetName();
 
 	const FStaticMeshLODResources& MeshLOD = StaticMesh->GetLODForExport(LODIndex);
 	const FPositionVertexBuffer& PositionBuffer = MeshLOD.VertexBuffers.PositionVertexBuffer;
@@ -163,7 +162,7 @@ void FGLTFStaticMeshTask::Complete()
 		const FGLTFIndexArray SectionIndices = FGLTFMeshUtility::GetSectionIndices(MeshLOD, MaterialIndex);
 		const FGLTFMeshSection* ConvertedSection = MeshSectionConverter.GetOrAdd(&MeshLOD, SectionIndices);
 
-		FGLTFJsonPrimitive& JsonPrimitive = JsonMesh.Primitives[MaterialIndex];
+		FGLTFJsonPrimitive& JsonPrimitive = JsonMesh->Primitives[MaterialIndex];
 		JsonPrimitive.Indices = Builder.GetOrAddIndexAccessor(ConvertedSection);
 
 		JsonPrimitive.Attributes.Position = Builder.GetOrAddPositionAccessor(ConvertedSection, &PositionBuffer);
@@ -207,8 +206,7 @@ void FGLTFStaticMeshTask::Complete()
 
 void FGLTFSkeletalMeshTask::Complete()
 {
-	FGLTFJsonMesh& JsonMesh = Builder.GetMesh(MeshIndex);
-	JsonMesh.Name = SkeletalMeshComponent != nullptr ? FGLTFNameUtility::GetName(SkeletalMeshComponent) : SkeletalMesh->GetName();
+	JsonMesh->Name = SkeletalMeshComponent != nullptr ? FGLTFNameUtility::GetName(SkeletalMeshComponent) : SkeletalMesh->GetName();
 
 	const FSkeletalMeshRenderData* RenderData = SkeletalMesh->GetResourceForRendering();
 	const FSkeletalMeshLODRenderData& MeshLOD = RenderData->LODRenderData[LODIndex];
@@ -267,7 +265,7 @@ void FGLTFSkeletalMeshTask::Complete()
 		const FGLTFIndexArray SectionIndices = FGLTFMeshUtility::GetSectionIndices(MeshLOD, MaterialIndex);
 		const FGLTFMeshSection* ConvertedSection = MeshSectionConverter.GetOrAdd(&MeshLOD, SectionIndices);
 
-		FGLTFJsonPrimitive& JsonPrimitive = JsonMesh.Primitives[MaterialIndex];
+		FGLTFJsonPrimitive& JsonPrimitive = JsonMesh->Primitives[MaterialIndex];
 		JsonPrimitive.Indices = Builder.GetOrAddIndexAccessor(ConvertedSection);
 
 		JsonPrimitive.Attributes.Position = Builder.GetOrAddPositionAccessor(ConvertedSection, &PositionBuffer);
diff --git a/Engine/Plugins/Enterprise/GLTFExporter/MigratedFromGit/Source/GLTFExporter/Private/Tasks/GLTFMeshTasks.h b/Engine/Plugins/Enterprise/GLTFExporter/MigratedFromGit/Source/GLTFExporter/Private/Tasks/GLTFMeshTasks.h
--- a/Engine/Plugins/Enterprise/GLTFExporter/MigratedFromGit/Source/GLTFExporter/Private/Tasks/GLTFMeshTasks.h
+++ b/Engine/Plugins/Enterprise/GLTFExporter/MigratedFromGit/Source/GLTFExporter/Private/Tasks/GLTFMeshTasks.h
@@ -12,7 +12,7 @@ class FGLTFStaticMeshTask : public FGLTFTask
 {
 public:
 
-	FGLTFStaticMeshTask(FGLTFConvertBuilder& Builder, FGLTFStaticMeshSectionConverter& MeshSectionConverter, const UStaticMesh* StaticMesh, const UStaticMeshComponent* StaticMeshComponent, FGLTFMaterialArray Materials, int32 LODIndex, FGLTFJsonMeshIndex MeshIndex)
+	FGLTFStaticMeshTask(FGLTFConvertBuilder& Builder, FGLTFStaticMeshSectionConverter& MeshSectionConverter, const UStaticMesh* StaticMesh, const UStaticMeshComponent* StaticMeshComponent, FGLTFMaterialArray Materials, int32 LODIndex, FGLTFJsonMesh* JsonMesh)
 		: FGLTFTask(EGLTFTaskPriority::Mesh)
 		, Builder(Builder)
 		, MeshSectionConverter(MeshSectionConverter)
@@ -20,7 +20,7 @@ public:
 		, StaticMeshComponent(StaticMeshComponent)
 		, Materials(Materials)
 		, LODIndex(LODIndex)
-		, MeshIndex(MeshIndex)
+		, JsonMesh(JsonMesh)
 	{
 	}
 
@@ -39,14 +39,14 @@ private:
 	const UStaticMeshComponent* StaticMeshComponent;
 	const FGLTFMaterialArray Materials;
 	const int32 LODIndex;
-	const FGLTFJsonMeshIndex MeshIndex;
+	FGLTFJsonMesh* JsonMesh;
 };
 
 class FGLTFSkeletalMeshTask : public FGLTFTask
 {
 public:
 
-	FGLTFSkeletalMeshTask(FGLTFConvertBuilder& Builder, FGLTFSkeletalMeshSectionConverter& MeshSectionConverter, const USkeletalMesh* SkeletalMesh, const USkeletalMeshComponent* SkeletalMeshComponent, FGLTFMaterialArray Materials, int32 LODIndex, FGLTFJsonMeshIndex MeshIndex)
+	FGLTFSkeletalMeshTask(FGLTFConvertBuilder& Builder, FGLTFSkeletalMeshSectionConverter& MeshSectionConverter, const USkeletalMesh* SkeletalMesh, const USkeletalMeshComponent* SkeletalMeshComponent, FGLTFMaterialArray Materials, int32 LODIndex, FGLTFJsonMesh* JsonMesh)
 		: FGLTFTask(EGLTFTaskPriority::Mesh)
 		, Builder(Builder)
 		, MeshSectionConverter(MeshSectionConverter)
@@ -54,7 +54,7 @@ public:
 		, SkeletalMeshComponent(SkeletalMeshComponent)
 		, Materials(Materials)
 		, LODIndex(LODIndex)
-		, MeshIndex(MeshIndex)
+		, JsonMesh(JsonMesh)
 	{
 	}
 
@@ -73,5 +73,5 @@ private:
 	const USkeletalMeshComponent* SkeletalMeshComponent;
 	const FGLTFMaterialArray Materials;
 	const int32 LODIndex;
-	const FGLTFJsonMeshIndex MeshIndex;
+	FGLTFJsonMesh* JsonMesh;
 };
diff --git a/Engine/Plugins/Enterprise/GLTFExporter/MigratedFromGit/Source/GLTFExporter/Private/Tasks/GLTFTextureTasks.cpp b/Engine/Plugins/Enterprise/GLTFExporter/MigratedFromGit/Source/GLTFExporter/Private/Tasks/GLTFTextureTasks.cpp
--- a/Engine/Plugins/Enterprise/GLTFExporter/MigratedFromGit/Source/GLTFExporter/Private/Tasks/GLTFTextureTasks.cpp
+++ b/Engine/Plugins/Enterprise/GLTFExporter/MigratedFromGit/Source/GLTFExporter/Private/Tasks/GLTFTextureTasks.cpp
@@ -18,13 +18,12 @@ FString FGLTFTexture2DTask::GetName()
 
 void FGLTFTexture2DTask::Complete()
 {
-	FGLTFJsonTexture& JsonTexture = Builder.GetTexture(TextureIndex);
-	Texture2D->GetName(JsonTexture.Name);
+	Texture2D->GetName(JsonTexture->Name);
 
 	const bool bFromSRGB = Texture2D->SRGB;
 	if (bFromSRGB != bToSRGB)
 	{
-		JsonTexture.Name += bToSRGB ? TEXT("_sRGB") : TEXT("_Linear");
+		JsonTexture->Name += bToSRGB ? TEXT("_sRGB") : TEXT("_Linear");
 	}
 
 	const bool bIsHDR = FGLTFTextureUtility::IsHDR(Texture2D);
@@ -36,13 +35,13 @@ void FGLTFTexture2DTask::Complete()
 
 	if (bIsHDR)
 	{
-		JsonTexture.Encoding = Builder.GetTextureHDREncoding();
+		JsonTexture->Encoding = Builder.GetTextureHDREncoding();
 	}
 
 	TGLTFSharedArray<FColor> Pixels;
-	if (!FGLTFTextureUtility::ReadPixels(RenderTarget, *Pixels, JsonTexture.Encoding))
+	if (!FGLTFTextureUtility::ReadPixels(RenderTarget, *Pixels, JsonTexture->Encoding))
 	{
-		Builder.LogWarning(FString::Printf(TEXT("Failed to read pixels for 2D texture %s"), *JsonTexture.Name));
+		Builder.LogWarning(FString::Printf(TEXT("Failed to read pixels for 2D texture %s"), *JsonTexture->Name));
 		return;
 	}
 
@@ -59,8 +58,8 @@ void FGLTFTexture2DTask::Complete()
 		Texture2D->IsNormalMap() ? EGLTFTextureType::Normalmaps :
 		bIsHDR ? EGLTFTextureType::HDR : EGLTFTextureType::None;
 
-	JsonTexture.Source = Builder.GetOrAddImage(Pixels, Size, bIgnoreAlpha, Type, JsonTexture.Name);
-	JsonTexture.Sampler = Builder.GetOrAddSampler(Texture2D);
+	JsonTexture->Source = Builder.GetOrAddImage(Pixels, Size, bIgnoreAlpha, Type, JsonTexture->Name);
+	JsonTexture->Sampler = Builder.GetOrAddSampler(Texture2D);
 }
 
 FString FGLTFTextureCubeTask::GetName()
@@ -70,13 +69,12 @@ FString FGLTFTextureCubeTask::GetName()
 
 void FGLTFTextureCubeTask::Complete()
 {
-	FGLTFJsonTexture& JsonTexture = Builder.GetTexture(TextureIndex);
-	JsonTexture.Name = TextureCube->GetName() + TEXT("_") + FGLTFJsonUtility::GetValue(FGLTFConverterUtility::ConvertCubeFace(CubeFace));
+	JsonTexture->Name = TextureCube->GetName() + TEXT("_") + FGLTFJsonUtility::GetValue(FGLTFConverterUtility::ConvertCubeFace(CubeFace));
 
 	const bool bFromSRGB = TextureCube->SRGB;
 	if (bFromSRGB != bToSRGB)
 	{
-		JsonTexture.Name += bToSRGB ? TEXT("_sRGB") : TEXT("_Linear");
+		JsonTexture->Name += bToSRGB ? TEXT("_sRGB") : TEXT("_Linear");
 	}
 
 	// TODO: add optimized "happy path" if cube face doesn't need rotation and has suitable pixel format
@@ -98,11 +96,11 @@ void FGLTFTextureCubeTask::Complete()
 
 	if (bIsHDR)
 	{
-		JsonTexture.Encoding = Builder.GetTextureHDREncoding();
+		JsonTexture->Encoding = Builder.GetTextureHDREncoding();
 	}
 
 	TGLTFSharedArray<FColor> Pixels;
-	if (!FGLTFTextureUtility::ReadPixels(RenderTarget, *Pixels, JsonTexture.Encoding))
+	if (!FGLTFTextureUtility::ReadPixels(RenderTarget, *Pixels, JsonTexture->Encoding))
 	{
 		Builder.LogWarning(FString::Printf(TEXT("Failed to read pixels (cube face %d) for cubemap texture %s"), CubeFace, *TextureCube->GetName()));
 		return;
@@ -113,8 +111,8 @@ void FGLTFTextureCubeTask::Complete()
 	const bool bIgnoreAlpha = FGLTFTextureUtility::IsAlphaless(TextureCube->GetPixelFormat());
 	const EGLTFTextureType Type = bIsHDR ? EGLTFTextureType::HDR : EGLTFTextureType::None;
 
-	JsonTexture.Source = Builder.GetOrAddImage(Pixels, Size, bIgnoreAlpha, Type, JsonTexture.Name);
-	JsonTexture.Sampler = Builder.GetOrAddSampler(TextureCube);
+	JsonTexture->Source = Builder.GetOrAddImage(Pixels, Size, bIgnoreAlpha, Type, JsonTexture->Name);
+	JsonTexture->Sampler = Builder.GetOrAddSampler(TextureCube);
 }
 
 FString FGLTFTextureRenderTarget2DTask::GetName()
@@ -124,13 +122,12 @@ FString FGLTFTextureRenderTarget2DTask::GetName()
 
 void FGLTFTextureRenderTarget2DTask::Complete()
 {
-	FGLTFJsonTexture& JsonTexture = Builder.GetTexture(TextureIndex);
-	RenderTarget2D->GetName(JsonTexture.Name);
+	RenderTarget2D->GetName(JsonTexture->Name);
 
 	const bool bFromSRGB = RenderTarget2D->SRGB;
 	if (bFromSRGB != bToSRGB)
 	{
-		JsonTexture.Name += bToSRGB ? TEXT("_sRGB") : TEXT("_Linear");
+		JsonTexture->Name += bToSRGB ? TEXT("_sRGB") : TEXT("_Linear");
 	}
 
 	const bool bIsHDR = FGLTFTextureUtility::IsHDR(RenderTarget2D);
@@ -138,13 +135,13 @@ void FGLTFTextureRenderTarget2DTask::Complete()
 
 	if (bIsHDR)
 	{
-		JsonTexture.Encoding = Builder.GetTextureHDREncoding();
+		JsonTexture->Encoding = Builder.GetTextureHDREncoding();
 	}
 
 	TGLTFSharedArray<FColor> Pixels;
-	if (!FGLTFTextureUtility::ReadPixels(RenderTarget2D, *Pixels, JsonTexture.Encoding))
+	if (!FGLTFTextureUtility::ReadPixels(RenderTarget2D, *Pixels, JsonTexture->Encoding))
 	{
-		Builder.LogWarning(FString::Printf(TEXT("Failed to read pixels for 2D render target %s"), *JsonTexture.Name));
+		Builder.LogWarning(FString::Printf(TEXT("Failed to read pixels for 2D render target %s"), *JsonTexture->Name));
 		return;
 	}
 
@@ -153,8 +150,8 @@ void FGLTFTextureRenderTarget2DTask::Complete()
 	const bool bIgnoreAlpha = FGLTFTextureUtility::IsAlphaless(RenderTarget2D->GetFormat());
 	const EGLTFTextureType Type = bIsHDR ? EGLTFTextureType::HDR : EGLTFTextureType::None;
 
-	JsonTexture.Source = Builder.GetOrAddImage(Pixels, Size, bIgnoreAlpha, Type, JsonTexture.Name);
-	JsonTexture.Sampler = Builder.GetOrAddSampler(RenderTarget2D);
+	JsonTexture->Source = Builder.GetOrAddImage(Pixels, Size, bIgnoreAlpha, Type, JsonTexture->Name);
+	JsonTexture->Sampler = Builder.GetOrAddSampler(RenderTarget2D);
 }
 
 FString FGLTFTextureRenderTargetCubeTask::GetName()
@@ -164,13 +161,12 @@ FString FGLTFTextureRenderTargetCubeTask::GetName()
 
 void FGLTFTextureRenderTargetCubeTask::Complete()
 {
-	FGLTFJsonTexture& JsonTexture = Builder.GetTexture(TextureIndex);
-	JsonTexture.Name = RenderTargetCube->GetName() + TEXT("_") + FGLTFJsonUtility::GetValue(FGLTFConverterUtility::ConvertCubeFace(CubeFace));
+	JsonTexture->Name = RenderTargetCube->GetName() + TEXT("_") + FGLTFJsonUtility::GetValue(FGLTFConverterUtility::ConvertCubeFace(CubeFace));
 
 	const bool bFromSRGB = RenderTargetCube->SRGB;
 	if (bFromSRGB != bToSRGB)
 	{
-		JsonTexture.Name += bToSRGB ? TEXT("_sRGB") : TEXT("_Linear");
+		JsonTexture->Name += bToSRGB ? TEXT("_sRGB") : TEXT("_Linear");
 	}
 
 	// TODO: add optimized "happy path" if cube face doesn't need rotation
@@ -191,11 +187,11 @@ void FGLTFTextureRenderTargetCubeTask::Complete()
 
 	if (bIsHDR)
 	{
-		JsonTexture.Encoding = Builder.GetTextureHDREncoding();
+		JsonTexture->Encoding = Builder.GetTextureHDREncoding();
 	}
 
 	TGLTFSharedArray<FColor> Pixels;
-	if (!FGLTFTextureUtility::ReadPixels(RenderTarget, *Pixels, JsonTexture.Encoding))
+	if (!FGLTFTextureUtility::ReadPixels(RenderTarget, *Pixels, JsonTexture->Encoding))
 	{
 		Builder.LogWarning(FString::Printf(TEXT("Failed to read pixels (cube face %d) for cubemap render target %s"), CubeFace, *RenderTargetCube->GetName()));
 		return;
@@ -206,8 +202,8 @@ void FGLTFTextureRenderTargetCubeTask::Complete()
 	const bool bIgnoreAlpha = FGLTFTextureUtility::IsAlphaless(RenderTargetCube->GetFormat());
 	const EGLTFTextureType Type = bIsHDR ? EGLTFTextureType::HDR : EGLTFTextureType::None;
 
-	JsonTexture.Source = Builder.GetOrAddImage(Pixels, Size, bIgnoreAlpha, Type, JsonTexture.Name);
-	JsonTexture.Sampler = Builder.GetOrAddSampler(RenderTargetCube);
+	JsonTexture->Source = Builder.GetOrAddImage(Pixels, Size, bIgnoreAlpha, Type, JsonTexture->Name);
+	JsonTexture->Sampler = Builder.GetOrAddSampler(RenderTargetCube);
 }
 
 #if WITH_EDITOR
@@ -219,8 +215,7 @@ FString FGLTFTextureLightMapTask::GetName()
 
 void FGLTFTextureLightMapTask::Complete()
 {
-	FGLTFJsonTexture& JsonTexture = Builder.GetTexture(TextureIndex);
-	LightMap->GetName(JsonTexture.Name);
+	LightMap->GetName(JsonTexture->Name);
 
 	// NOTE: export of lightmaps via source data is used to work around issues with
 	// quality-loss due to incorrect gamma transformation when rendering to a canvas.
@@ -228,14 +223,14 @@ void FGLTFTextureLightMapTask::Complete()
 	FTextureSource& Source = const_cast<FTextureSource&>(LightMap->Source);
 	if (!Source.IsValid())
 	{
-		Builder.LogWarning(FString::Printf(TEXT("Failed to export lightmap texture %s because of missing source data"), *JsonTexture.Name));
+		Builder.LogWarning(FString::Printf(TEXT("Failed to export lightmap texture %s because of missing source data"), *JsonTexture->Name));
 		return;
 	}
 
 	const ETextureSourceFormat SourceFormat = Source.GetFormat();
 	if (SourceFormat != TSF_BGRA8)
 	{
-		Builder.LogWarning(FString::Printf(TEXT("Failed to export lightmap texture %s because of unsupported source format %s"), *JsonTexture.Name, *FGLTFNameUtility::GetName(SourceFormat)));
+		Builder.LogWarning(FString::Printf(TEXT("Failed to export lightmap texture %s because of unsupported source format %s"), *JsonTexture->Name, *FGLTFNameUtility::GetName(SourceFormat)));
 		return;
 	}
 
@@ -249,8 +244,8 @@ void FGLTFTextureLightMapTask::Complete()
 	TGLTFSharedArray<FColor> Pixels = MakeShared<TArray<FColor>>(static_cast<const FColor*>(RawData), ByteLength / sizeof(FColor));
 	Source.UnlockMip(0);
 
-	JsonTexture.Source = Builder.GetOrAddImage(Pixels, Size, bIgnoreAlpha, Type, JsonTexture.Name);
-	JsonTexture.Sampler = Builder.GetOrAddSampler(LightMap);
+	JsonTexture->Source = Builder.GetOrAddImage(Pixels, Size, bIgnoreAlpha, Type, JsonTexture->Name);
+	JsonTexture->Sampler = Builder.GetOrAddSampler(LightMap);
 }
 
 #endif
diff --git a/Engine/Plugins/Enterprise/GLTFExporter/MigratedFromGit/Source/GLTFExporter/Private/Tasks/GLTFTextureTasks.h b/Engine/Plugins/Enterprise/GLTFExporter/MigratedFromGit/Source/GLTFExporter/Private/Tasks/GLTFTextureTasks.h
--- a/Engine/Plugins/Enterprise/GLTFExporter/MigratedFromGit/Source/GLTFExporter/Private/Tasks/GLTFTextureTasks.h
+++ b/Engine/Plugins/Enterprise/GLTFExporter/MigratedFromGit/Source/GLTFExporter/Private/Tasks/GLTFTextureTasks.h
@@ -9,12 +9,12 @@ class FGLTFTexture2DTask : public FGLTFTask
 {
 public:
 
-	FGLTFTexture2DTask(FGLTFConvertBuilder& Builder, const UTexture2D* Texture2D, bool bToSRGB, FGLTFJsonTextureIndex TextureIndex)
+	FGLTFTexture2DTask(FGLTFConvertBuilder& Builder, const UTexture2D* Texture2D, bool bToSRGB, FGLTFJsonTexture* JsonTexture)
 		: FGLTFTask(EGLTFTaskPriority::Texture)
 		, Builder(Builder)
 		, Texture2D(Texture2D)
 		, bToSRGB(bToSRGB)
-		, TextureIndex(TextureIndex)
+		, JsonTexture(JsonTexture)
 	{
 	}
 
@@ -27,20 +27,20 @@ private:
 	FGLTFConvertBuilder& Builder;
 	const UTexture2D* Texture2D;
 	bool bToSRGB;
-	const FGLTFJsonTextureIndex TextureIndex;
+	FGLTFJsonTexture* JsonTexture;
 };
 
 class FGLTFTextureCubeTask : public FGLTFTask
 {
 public:
 
-	FGLTFTextureCubeTask(FGLTFConvertBuilder& Builder, const UTextureCube* TextureCube, ECubeFace CubeFace, bool bToSRGB, FGLTFJsonTextureIndex TextureIndex)
+	FGLTFTextureCubeTask(FGLTFConvertBuilder& Builder, const UTextureCube* TextureCube, ECubeFace CubeFace, bool bToSRGB, FGLTFJsonTexture* JsonTexture)
 		: FGLTFTask(EGLTFTaskPriority::Texture)
 		, Builder(Builder)
 		, TextureCube(TextureCube)
 		, CubeFace(CubeFace)
 		, bToSRGB(bToSRGB)
-		, TextureIndex(TextureIndex)
+		, JsonTexture(JsonTexture)
 	{
 	}
 
@@ -54,19 +54,19 @@ private:
 	const UTextureCube* TextureCube;
 	ECubeFace CubeFace;
 	bool bToSRGB;
-	const FGLTFJsonTextureIndex TextureIndex;
+	FGLTFJsonTexture* JsonTexture;
 };
 
 class FGLTFTextureRenderTarget2DTask : public FGLTFTask
 {
 public:
 
-	FGLTFTextureRenderTarget2DTask(FGLTFConvertBuilder& Builder, const UTextureRenderTarget2D* RenderTarget2D, bool bToSRGB, FGLTFJsonTextureIndex TextureIndex)
+	FGLTFTextureRenderTarget2DTask(FGLTFConvertBuilder& Builder, const UTextureRenderTarget2D* RenderTarget2D, bool bToSRGB, FGLTFJsonTexture* JsonTexture)
 		: FGLTFTask(EGLTFTaskPriority::Texture)
 		, Builder(Builder)
 		, RenderTarget2D(RenderTarget2D)
 		, bToSRGB(bToSRGB)
-		, TextureIndex(TextureIndex)
+		, JsonTexture(JsonTexture)
 	{
 	}
 
@@ -79,20 +79,20 @@ private:
 	FGLTFConvertBuilder& Builder;
 	const UTextureRenderTarget2D* RenderTarget2D;
 	bool bToSRGB;
-	const FGLTFJsonTextureIndex TextureIndex;
+	FGLTFJsonTexture* JsonTexture;
 };
 
 class FGLTFTextureRenderTargetCubeTask : public FGLTFTask
 {
 public:
 
-	FGLTFTextureRenderTargetCubeTask(FGLTFConvertBuilder& Builder, const UTextureRenderTargetCube* RenderTargetCube, ECubeFace CubeFace, bool bToSRGB, FGLTFJsonTextureIndex TextureIndex)
+	FGLTFTextureRenderTargetCubeTask(FGLTFConvertBuilder& Builder, const UTextureRenderTargetCube* RenderTargetCube, ECubeFace CubeFace, bool bToSRGB, FGLTFJsonTexture* JsonTexture)
 		: FGLTFTask(EGLTFTaskPriority::Texture)
 		, Builder(Builder)
 		, RenderTargetCube(RenderTargetCube)
 		, CubeFace(CubeFace)
 		, bToSRGB(bToSRGB)
-		, TextureIndex(TextureIndex)
+		, JsonTexture(JsonTexture)
 	{
 	}
 
@@ -106,7 +106,7 @@ private:
 	const UTextureRenderTargetCube* RenderTargetCube;
 	ECubeFace CubeFace;
 	bool bToSRGB;
-	const FGLTFJsonTextureIndex TextureIndex;
+	FGLTFJsonTexture* JsonTexture;
 };
 
 #if WITH_EDITOR
@@ -115,11 +115,11 @@ class FGLTFTextureLightMapTask : public FGLTFTask
 {
 public:
 
-	FGLTFTextureLightMapTask(FGLTFConvertBuilder& Builder, const ULightMapTexture2D* LightMap, FGLTFJsonTextureIndex TextureIndex)
+	FGLTFTextureLightMapTask(FGLTFConvertBuilder& Builder, const ULightMapTexture2D* LightMap, FGLTFJsonTexture* JsonTexture)
 		: FGLTFTask(EGLTFTaskPriority::Texture)
 		, Builder(Builder)
 		, LightMap(LightMap)
-		, TextureIndex(TextureIndex)
+		, JsonTexture(JsonTexture)
 	{
 	}
 
@@ -131,7 +131,7 @@ private:
 
 	FGLTFConvertBuilder& Builder;
 	const ULightMapTexture2D* LightMap;
-	const FGLTFJsonTextureIndex TextureIndex;
+	FGLTFJsonTexture* JsonTexture;
 };
 
 #endif
