[GLTFExporter] refactor: add separate cpp files for IGLTFJsonIndexedObject and TGLTFJsonIndexedObjectArray

commit ae8cbb09dcab04b5b63de4f42cc70ed62fac7be4
Date: Sun, 13 Mar 2022 14:52:49 +0100

    refactor: add separate cpp files for IGLTFJsonIndexedObject and TGLTFJsonIndexedObjectArray
    (cherry picked from commit c8f2d3ea768b2cb4ad7401c095939339873ce911​)

#skipci

diff --git a/Engine/Plugins/Enterprise/GLTFExporter/MigratedFromGit/Source/GLTFExporter/Private/Json/GLTFJsonArray.cpp b/Engine/Plugins/Enterprise/GLTFExporter/MigratedFromGit/Source/GLTFExporter/Private/Json/GLTFJsonArray.cpp
new file mode 100644
--- /dev/null
+++ b/Engine/Plugins/Enterprise/GLTFExporter/MigratedFromGit/Source/GLTFExporter/Private/Json/GLTFJsonArray.cpp
@@ -0,0 +1,11 @@
+// Copyright Epic Games, Inc. All Rights Reserved.
+
+#include "Json/GLTFJsonArray.h"
+#include "Json/GLTFJsonWriter.h"
+
+void IGLTFJsonArray::WriteValue(IGLTFJsonWriter& Writer) const
+{
+	Writer.StartArray();
+	WriteArray(Writer);
+	Writer.EndArray();
+}
diff --git a/Engine/Plugins/Enterprise/GLTFExporter/MigratedFromGit/Source/GLTFExporter/Private/Json/GLTFJsonObject.cpp b/Engine/Plugins/Enterprise/GLTFExporter/MigratedFromGit/Source/GLTFExporter/Private/Json/GLTFJsonObject.cpp
new file mode 100644
--- /dev/null
+++ b/Engine/Plugins/Enterprise/GLTFExporter/MigratedFromGit/Source/GLTFExporter/Private/Json/GLTFJsonObject.cpp
@@ -0,0 +1,11 @@
+// Copyright Epic Games, Inc. All Rights Reserved.
+
+#include "Json/GLTFJsonObject.h"
+#include "Json/GLTFJsonWriter.h"
+
+void IGLTFJsonObject::WriteValue(IGLTFJsonWriter& Writer) const
+{
+	Writer.StartObject();
+	WriteObject(Writer);
+	Writer.EndObject();
+}
diff --git a/Engine/Plugins/Enterprise/GLTFExporter/MigratedFromGit/Source/GLTFExporter/Public/Core/GLTFQuaternion.h b/Engine/Plugins/Enterprise/GLTFExporter/MigratedFromGit/Source/GLTFExporter/Public/Core/GLTFQuaternion.h
--- a/Engine/Plugins/Enterprise/GLTFExporter/MigratedFromGit/Source/GLTFExporter/Public/Core/GLTFQuaternion.h
+++ b/Engine/Plugins/Enterprise/GLTFExporter/MigratedFromGit/Source/GLTFExporter/Public/Core/GLTFQuaternion.h
@@ -9,7 +9,7 @@
 #endif
 
 template <typename ComponentType>
-struct GLTFEXPORTER_API TGLTFQuaternion
+struct TGLTFQuaternion
 {
 	union
 	{
diff --git a/Engine/Plugins/Enterprise/GLTFExporter/MigratedFromGit/Source/GLTFExporter/Public/Json/GLTFJsonArray.h b/Engine/Plugins/Enterprise/GLTFExporter/MigratedFromGit/Source/GLTFExporter/Public/Json/GLTFJsonArray.h
--- a/Engine/Plugins/Enterprise/GLTFExporter/MigratedFromGit/Source/GLTFExporter/Public/Json/GLTFJsonArray.h
+++ b/Engine/Plugins/Enterprise/GLTFExporter/MigratedFromGit/Source/GLTFExporter/Public/Json/GLTFJsonArray.h
@@ -3,138 +3,10 @@
 #pragma once
 
 #include "Json/GLTFJsonValue.h"
-#include "Json/GLTFJsonWriter.h"
-#include "Json/GLTFJsonObject.h"
 
 struct GLTFEXPORTER_API IGLTFJsonArray : IGLTFJsonValue
 {
-	virtual void WriteValue(IGLTFJsonWriter& Writer) const override final
-	{
-		Writer.StartArray();
-		WriteArray(Writer);
-		Writer.EndArray();
-	}
+	virtual void WriteValue(IGLTFJsonWriter& Writer) const override final;
 
 	virtual void WriteArray(IGLTFJsonWriter& Writer) const = 0;
 };
-
-template <typename ElementType, typename = typename TEnableIf<TIsDerivedFrom<ElementType, IGLTFJsonIndexedObject>::Value>::Type>
-struct TGLTFJsonIndexedObjectArray : IGLTFJsonArray
-{
-	using ArrayType = TArray<TUniquePtr<ElementType>>;
-	using SizeType = typename ArrayType::SizeType;
-
-	template <typename IteratorType, typename ValueType>
-	struct TIterator
-	{
-		TIterator(const IteratorType& Iterator)
-			: Iterator(Iterator)
-		{
-		}
-
-		ValueType operator*() const
-		{
-			return (*Iterator).Get();
-		}
-
-		TIterator& operator++()
-		{
-			++Iterator;
-			return *this;
-		}
-
-		TIterator& operator--()
-		{
-			--Iterator;
-			return *this;
-		}
-
-		friend bool operator!=(const TIterator& A, const TIterator& B)
-		{
-			return A.Iterator != B.Iterator;
-		};
-
-	private:
-
-		IteratorType Iterator;
-	};
-
-	typedef TIterator<typename ArrayType::RangedForIteratorType, ElementType*> IteratorType;
-	typedef TIterator<typename ArrayType::RangedForConstIteratorType, const ElementType*> ConstIteratorType;
-
-	TGLTFJsonIndexedObjectArray() = default;
-	TGLTFJsonIndexedObjectArray(TGLTFJsonIndexedObjectArray&& Other) = default;
-	TGLTFJsonIndexedObjectArray& operator=(TGLTFJsonIndexedObjectArray&& Other) = default;
-
-	TGLTFJsonIndexedObjectArray(const TGLTFJsonIndexedObjectArray&) = delete;
-	TGLTFJsonIndexedObjectArray& operator=(const TGLTFJsonIndexedObjectArray&) = delete;
-
-	ElementType* Add(const ElementType& Temp)
-	{
-		ElementType* Element = new ElementType(Array.Num());
-		*Element = Temp;
-		Element->Index = Array.Num();
-		Array.Add(TUniquePtr<ElementType>(Element));
-		return Element;
-	}
-
-	ElementType* Add()
-	{
-		ElementType* Element = new ElementType(Array.Num());
-		Array.Add(TUniquePtr<ElementType>(Element));
-		return Element;
-	}
-
-	bool IsValidIndex(SizeType Index) const
-	{
-		return Array.IsValidIndex(Index);
-	}
-
-	SizeType Num() const
-	{
-		return Array.Num();
-	}
-
-	bool Contains(const ElementType* Element) const
-	{
-		TUniquePtr<ElementType> TempPtr(Element);
-		const bool Result = Array.Contains(TempPtr);
-		TempPtr.Release();
-		return Result;
-	}
-
-	SizeType Find(const ElementType* Element) const
-	{
-		TUniquePtr<ElementType> TempPtr(Element);
-		const SizeType Result = Array.Find(TempPtr);
-		TempPtr.Release();
-		return Result;
-	}
-
-	ElementType* operator[](SizeType Index)
-	{
-		return Array[Index].Get();
-	}
-
-	const ElementType* operator[](SizeType Index) const
-	{
-		return Array[Index].Get();
-	}
-
-	IteratorType begin() { return IteratorType(Array.begin()); }
-	IteratorType end() { return IteratorType(Array.end()); }
-	ConstIteratorType begin() const { return ConstIteratorType(Array.begin()); }
-	ConstIteratorType end() const { return ConstIteratorType(Array.end()); }
-
-	virtual void WriteArray(IGLTFJsonWriter& Writer) const override
-	{
-		for (const ElementType* Element : *this)
-		{
-			Writer.Write(*Element);
-		}
-	}
-
-private:
-
-	ArrayType Array;
-};
diff --git a/Engine/Plugins/Enterprise/GLTFExporter/MigratedFromGit/Source/GLTFExporter/Public/Json/GLTFJsonColor.h b/Engine/Plugins/Enterprise/GLTFExporter/MigratedFromGit/Source/GLTFExporter/Public/Json/GLTFJsonColor.h
--- a/Engine/Plugins/Enterprise/GLTFExporter/MigratedFromGit/Source/GLTFExporter/Public/Json/GLTFJsonColor.h
+++ b/Engine/Plugins/Enterprise/GLTFExporter/MigratedFromGit/Source/GLTFExporter/Public/Json/GLTFJsonColor.h
@@ -4,9 +4,10 @@
 
 #include "Core/GLTFColor.h"
 #include "Json/GLTFJsonArray.h"
+#include "Json/GLTFJsonWriter.h"
 
 template <typename BaseType>
-struct GLTFEXPORTER_API TGLTFJsonColor : BaseType, IGLTFJsonArray
+struct TGLTFJsonColor : BaseType, IGLTFJsonArray
 {
 	TGLTFJsonColor(const BaseType& Other)
 		: BaseType(Other)
diff --git a/Engine/Plugins/Enterprise/GLTFExporter/MigratedFromGit/Source/GLTFExporter/Public/Json/GLTFJsonCore.h b/Engine/Plugins/Enterprise/GLTFExporter/MigratedFromGit/Source/GLTFExporter/Public/Json/GLTFJsonCore.h
--- a/Engine/Plugins/Enterprise/GLTFExporter/MigratedFromGit/Source/GLTFExporter/Public/Json/GLTFJsonCore.h
+++ b/Engine/Plugins/Enterprise/GLTFExporter/MigratedFromGit/Source/GLTFExporter/Public/Json/GLTFJsonCore.h
@@ -5,9 +5,12 @@
 #include "Json/GLTFJsonEnums.h"
 #include "Json/GLTFJsonIndex.h"
 #include "Json/GLTFJsonValue.h"
-#include "Json/GLTFJsonObject.h"
 #include "Json/GLTFJsonArray.h"
+#include "Json/GLTFJsonObject.h"
+#include "Json/GLTFJsonIndexedObject.h"
+#include "Json/GLTFJsonIndexedObjectArray.h"
 #include "Json/GLTFJsonColor.h"
 #include "Json/GLTFJsonMatrix.h"
 #include "Json/GLTFJsonQuaternion.h"
 #include "Json/GLTFJsonVector.h"
+#include "Json/GLTFJsonWriter.h"
diff --git a/Engine/Plugins/Enterprise/GLTFExporter/MigratedFromGit/Source/GLTFExporter/Public/Json/GLTFJsonExtensions.h b/Engine/Plugins/Enterprise/GLTFExporter/MigratedFromGit/Source/GLTFExporter/Public/Json/GLTFJsonExtensions.h
--- a/Engine/Plugins/Enterprise/GLTFExporter/MigratedFromGit/Source/GLTFExporter/Public/Json/GLTFJsonExtensions.h
+++ b/Engine/Plugins/Enterprise/GLTFExporter/MigratedFromGit/Source/GLTFExporter/Public/Json/GLTFJsonExtensions.h
@@ -2,7 +2,8 @@
 
 #pragma once
 
-#include "Json/GLTFJsonCore.h"
+#include "CoreMinimal.h"
+#include "Json/GLTFJsonEnums.h"
 
 struct GLTFEXPORTER_API FGLTFJsonExtensions
 {
diff --git a/Engine/Plugins/Enterprise/GLTFExporter/MigratedFromGit/Source/GLTFExporter/Public/Json/GLTFJsonIndex.h b/Engine/Plugins/Enterprise/GLTFExporter/MigratedFromGit/Source/GLTFExporter/Public/Json/GLTFJsonIndex.h
--- a/Engine/Plugins/Enterprise/GLTFExporter/MigratedFromGit/Source/GLTFExporter/Public/Json/GLTFJsonIndex.h
+++ b/Engine/Plugins/Enterprise/GLTFExporter/MigratedFromGit/Source/GLTFExporter/Public/Json/GLTFJsonIndex.h
@@ -5,7 +5,7 @@
 #include "CoreMinimal.h"
 
 template <typename DerivedType>
-struct GLTFEXPORTER_API FGLTFJsonIndex
+struct FGLTFJsonIndex
 {
 	FGLTFJsonIndex()
 		: Value(INDEX_NONE)
diff --git a/Engine/Plugins/Enterprise/GLTFExporter/MigratedFromGit/Source/GLTFExporter/Public/Json/GLTFJsonIndexedObject.h b/Engine/Plugins/Enterprise/GLTFExporter/MigratedFromGit/Source/GLTFExporter/Public/Json/GLTFJsonIndexedObject.h
new file mode 100644
--- /dev/null
+++ b/Engine/Plugins/Enterprise/GLTFExporter/MigratedFromGit/Source/GLTFExporter/Public/Json/GLTFJsonIndexedObject.h
@@ -0,0 +1,17 @@
+// Copyright Epic Games, Inc. All Rights Reserved.
+
+#pragma once
+
+#include "Json/GLTFJsonObject.h"
+
+struct GLTFEXPORTER_API IGLTFJsonIndexedObject : IGLTFJsonObject
+{
+	int32 Index;
+
+protected:
+
+	IGLTFJsonIndexedObject(int32 Index)
+		: Index(Index)
+	{
+	}
+};
diff --git a/Engine/Plugins/Enterprise/GLTFExporter/MigratedFromGit/Source/GLTFExporter/Public/Json/GLTFJsonIndexedObjectArray.h b/Engine/Plugins/Enterprise/GLTFExporter/MigratedFromGit/Source/GLTFExporter/Public/Json/GLTFJsonIndexedObjectArray.h
new file mode 100644
--- /dev/null
+++ b/Engine/Plugins/Enterprise/GLTFExporter/MigratedFromGit/Source/GLTFExporter/Public/Json/GLTFJsonIndexedObjectArray.h
@@ -0,0 +1,128 @@
+// Copyright Epic Games, Inc. All Rights Reserved.
+
+#pragma once
+
+#include "Json/GLTFJsonArray.h"
+#include "Json/GLTFJsonWriter.h"
+#include "Json/GLTFJsonIndexedObject.h"
+
+template <typename ElementType, typename = typename TEnableIf<TIsDerivedFrom<ElementType, IGLTFJsonIndexedObject>::Value>::Type>
+struct TGLTFJsonIndexedObjectArray : IGLTFJsonArray
+{
+	using ArrayType = TArray<TUniquePtr<ElementType>>;
+	using SizeType = typename ArrayType::SizeType;
+
+	template <typename IteratorType, typename ValueType>
+	struct TIterator
+	{
+		TIterator(const IteratorType& Iterator)
+			: Iterator(Iterator)
+		{
+		}
+
+		ValueType operator*() const
+		{
+			return (*Iterator).Get();
+		}
+
+		TIterator& operator++()
+		{
+			++Iterator;
+			return *this;
+		}
+
+		TIterator& operator--()
+		{
+			--Iterator;
+			return *this;
+		}
+
+		friend bool operator!=(const TIterator& A, const TIterator& B)
+		{
+			return A.Iterator != B.Iterator;
+		};
+
+	private:
+
+		IteratorType Iterator;
+	};
+
+	typedef TIterator<typename ArrayType::RangedForIteratorType, ElementType*> IteratorType;
+	typedef TIterator<typename ArrayType::RangedForConstIteratorType, const ElementType*> ConstIteratorType;
+
+	TGLTFJsonIndexedObjectArray() = default;
+	TGLTFJsonIndexedObjectArray(TGLTFJsonIndexedObjectArray&& Other) = default;
+	TGLTFJsonIndexedObjectArray& operator=(TGLTFJsonIndexedObjectArray&& Other) = default;
+
+	TGLTFJsonIndexedObjectArray(const TGLTFJsonIndexedObjectArray&) = delete;
+	TGLTFJsonIndexedObjectArray& operator=(const TGLTFJsonIndexedObjectArray&) = delete;
+
+	ElementType* Add(const ElementType& Temp)
+	{
+		ElementType* Element = new ElementType(Array.Num());
+		*Element = Temp;
+		Element->Index = Array.Num();
+		Array.Add(TUniquePtr<ElementType>(Element));
+		return Element;
+	}
+
+	ElementType* Add()
+	{
+		ElementType* Element = new ElementType(Array.Num());
+		Array.Add(TUniquePtr<ElementType>(Element));
+		return Element;
+	}
+
+	bool IsValidIndex(SizeType Index) const
+	{
+		return Array.IsValidIndex(Index);
+	}
+
+	SizeType Num() const
+	{
+		return Array.Num();
+	}
+
+	bool Contains(const ElementType* Element) const
+	{
+		TUniquePtr<ElementType> TempPtr(Element);
+		const bool Result = Array.Contains(TempPtr);
+		TempPtr.Release();
+		return Result;
+	}
+
+	SizeType Find(const ElementType* Element) const
+	{
+		TUniquePtr<ElementType> TempPtr(Element);
+		const SizeType Result = Array.Find(TempPtr);
+		TempPtr.Release();
+		return Result;
+	}
+
+	ElementType* operator[](SizeType Index)
+	{
+		return Array[Index].Get();
+	}
+
+	const ElementType* operator[](SizeType Index) const
+	{
+		return Array[Index].Get();
+	}
+
+	IteratorType begin() { return IteratorType(Array.begin()); }
+	IteratorType end() { return IteratorType(Array.end()); }
+	ConstIteratorType begin() const { return ConstIteratorType(Array.begin()); }
+	ConstIteratorType end() const { return ConstIteratorType(Array.end()); }
+
+	virtual void WriteArray(IGLTFJsonWriter& Writer) const override
+	{
+		for (const ElementType* Element : *this)
+		{
+			Writer.Write(*Element);
+		}
+	}
+
+private:
+
+	ArrayType Array;
+};
diff --git a/Engine/Plugins/Enterprise/GLTFExporter/MigratedFromGit/Source/GLTFExporter/Public/Json/GLTFJsonMatrix.h b/Engine/Plugins/Enterprise/GLTFExporter/MigratedFromGit/Source/GLTFExporter/Public/Json/GLTFJsonMatrix.h
--- a/Engine/Plugins/Enterprise/GLTFExporter/MigratedFromGit/Source/GLTFExporter/Public/Json/GLTFJsonMatrix.h
+++ b/Engine/Plugins/Enterprise/GLTFExporter/MigratedFromGit/Source/GLTFExporter/Public/Json/GLTFJsonMatrix.h
@@ -4,12 +4,11 @@
 
 #include "Core/GLTFMatrix.h"
 #include "Json/GLTFJsonArray.h"
+#include "Json/GLTFJsonWriter.h"
 
 template <typename BaseType>
-struct GLTFEXPORTER_API TGLTFJsonMatrix : BaseType, IGLTFJsonArray
+struct TGLTFJsonMatrix : BaseType, IGLTFJsonArray
 {
-	static const TGLTFJsonMatrix Identity;
-
 	TGLTFJsonMatrix(const BaseType& Other)
 		: BaseType(Other)
 	{
@@ -23,9 +22,9 @@ struct GLTFEXPORTER_API TGLTFJsonMatrix : BaseType, IGLTFJsonArray
 
 	bool operator==(const BaseType& Other) const
 	{
-		for (SIZE_T i = 0; i < GetNum(BaseType::Components); ++i)
+		for (SIZE_T i = 0; i < GetNum(BaseType::Elements); ++i)
 		{
-			if (BaseType::Components[i] != Other.Components[i])
+			if (BaseType::Elements[i] != Other.Elements[i])
 			{
 				return false;
 			}
diff --git a/Engine/Plugins/Enterprise/GLTFExporter/MigratedFromGit/Source/GLTFExporter/Public/Json/GLTFJsonObject.h b/Engine/Plugins/Enterprise/GLTFExporter/MigratedFromGit/Source/GLTFExporter/Public/Json/GLTFJsonObject.h
--- a/Engine/Plugins/Enterprise/GLTFExporter/MigratedFromGit/Source/GLTFExporter/Public/Json/GLTFJsonObject.h
+++ b/Engine/Plugins/Enterprise/GLTFExporter/MigratedFromGit/Source/GLTFExporter/Public/Json/GLTFJsonObject.h
@@ -3,28 +3,10 @@
 #pragma once
 
 #include "Json/GLTFJsonValue.h"
-#include "Json/GLTFJsonWriter.h"
 
 struct GLTFEXPORTER_API IGLTFJsonObject : IGLTFJsonValue
 {
-	virtual void WriteValue(IGLTFJsonWriter& Writer) const override final
-	{
-		Writer.StartObject();
-		WriteObject(Writer);
-		Writer.EndObject();
-	}
+	virtual void WriteValue(IGLTFJsonWriter& Writer) const override final;
 
 	virtual void WriteObject(IGLTFJsonWriter& Writer) const = 0;
 };
-
-struct GLTFEXPORTER_API IGLTFJsonIndexedObject : IGLTFJsonObject
-{
-	int32 Index;
-
-protected:
-
-	IGLTFJsonIndexedObject(int32 Index)
-		: Index(Index)
-	{
-	}
-};
diff --git a/Engine/Plugins/Enterprise/GLTFExporter/MigratedFromGit/Source/GLTFExporter/Public/Json/GLTFJsonQuaternion.h b/Engine/Plugins/Enterprise/GLTFExporter/MigratedFromGit/Source/GLTFExporter/Public/Json/GLTFJsonQuaternion.h
--- a/Engine/Plugins/Enterprise/GLTFExporter/MigratedFromGit/Source/GLTFExporter/Public/Json/GLTFJsonQuaternion.h
+++ b/Engine/Plugins/Enterprise/GLTFExporter/MigratedFromGit/Source/GLTFExporter/Public/Json/GLTFJsonQuaternion.h
@@ -4,6 +4,7 @@
 
 #include "Core/GLTFQuaternion.h"
 #include "Json/GLTFJsonArray.h"
+#include "Json/GLTFJsonWriter.h"
 
 struct GLTFEXPORTER_API FGLTFJsonQuaternion : TGLTFQuaternion<float>, IGLTFJsonArray
 {
diff --git a/Engine/Plugins/Enterprise/GLTFExporter/MigratedFromGit/Source/GLTFExporter/Public/Json/GLTFJsonUtility.h b/Engine/Plugins/Enterprise/GLTFExporter/MigratedFromGit/Source/GLTFExporter/Public/Json/GLTFJsonUtility.h
--- a/Engine/Plugins/Enterprise/GLTFExporter/MigratedFromGit/Source/GLTFExporter/Public/Json/GLTFJsonUtility.h
+++ b/Engine/Plugins/Enterprise/GLTFExporter/MigratedFromGit/Source/GLTFExporter/Public/Json/GLTFJsonUtility.h
@@ -2,7 +2,8 @@
 
 #pragma once
 
-#include "Json/GLTFJsonCore.h"
+#include "CoreMinimal.h"
+#include "Json/GLTFJsonEnums.h"
 
 struct GLTFEXPORTER_API FGLTFJsonUtility
 {
diff --git a/Engine/Plugins/Enterprise/GLTFExporter/MigratedFromGit/Source/GLTFExporter/Public/Json/GLTFJsonVector.h b/Engine/Plugins/Enterprise/GLTFExporter/MigratedFromGit/Source/GLTFExporter/Public/Json/GLTFJsonVector.h
--- a/Engine/Plugins/Enterprise/GLTFExporter/MigratedFromGit/Source/GLTFExporter/Public/Json/GLTFJsonVector.h
+++ b/Engine/Plugins/Enterprise/GLTFExporter/MigratedFromGit/Source/GLTFExporter/Public/Json/GLTFJsonVector.h
@@ -4,13 +4,11 @@
 
 #include "Core/GLTFVector.h"
 #include "Json/GLTFJsonArray.h"
+#include "Json/GLTFJsonWriter.h"
 
 template <typename BaseType>
-struct GLTFEXPORTER_API TGLTFJsonVector : BaseType, IGLTFJsonArray
+struct TGLTFJsonVector : BaseType, IGLTFJsonArray
 {
-	static const TGLTFJsonVector Zero;
-	static const TGLTFJsonVector One;
-
 	TGLTFJsonVector(const BaseType& Other)
 		: BaseType(Other)
 	{
diff --git a/Engine/Plugins/Enterprise/GLTFExporter/MigratedFromGit/Source/GLTFExporter/Public/Json/GLTFJsonWriter.h b/Engine/Plugins/Enterprise/GLTFExporter/MigratedFromGit/Source/GLTFExporter/Public/Json/GLTFJsonWriter.h
--- a/Engine/Plugins/Enterprise/GLTFExporter/MigratedFromGit/Source/GLTFExporter/Public/Json/GLTFJsonWriter.h
+++ b/Engine/Plugins/Enterprise/GLTFExporter/MigratedFromGit/Source/GLTFExporter/Public/Json/GLTFJsonWriter.h
@@ -3,6 +3,7 @@
 #pragma once
 
 #include "Json/GLTFJsonValue.h"
+#include "Json/GLTFJsonIndexedObject.h"
 #include "Json/GLTFJsonUtility.h"
 #include "Json/GLTFJsonExtensions.h"
 
@@ -49,15 +50,15 @@ public:
 		Value.WriteValue(*this);
 	}
 
-	void Write(const IGLTFJsonValue* Value)
+	void Write(const IGLTFJsonIndexedObject* Object)
 	{
-		if (Value != nullptr)
+		if (Object != nullptr)
 		{
-			Value->WriteValue(*this);
+			Write(Object->Index);
 		}
 		else
 		{
-			Write(nullptr);
+			Write(INDEX_NONE);
 		}
 	}
 
@@ -161,7 +162,7 @@ public:
 		Write(Value);
 	}
 
-	void Write(const FString& Identifier, const IGLTFJsonValue* Value)
+	void Write(const FString& Identifier, const IGLTFJsonIndexedObject* Value)
 	{
 		SetIdentifier(Identifier);
 		Write(Value);
