[GLTFExporter] refactor: replace indices with pointers in json data structures

commit 62243e780a01d0812b52f5b43c6899aa5fb37770
Date: Sun, 13 Mar 2022 17:20:40 +0100

    refactor: replace indices with pointers in json data structures
    (cherry picked from commit 97b871667e94968217b0c95609fed4812f678278​)

#skipci

diff --git a/Engine/Plugins/Enterprise/GLTFExporter/MigratedFromGit/Source/GLTFExporter/Private/Builders/GLTFBufferBuilder.cpp b/Engine/Plugins/Enterprise/GLTFExporter/MigratedFromGit/Source/GLTFExporter/Private/Builders/GLTFBufferBuilder.cpp
--- a/Engine/Plugins/Enterprise/GLTFExporter/MigratedFromGit/Source/GLTFExporter/Private/Builders/GLTFBufferBuilder.cpp
+++ b/Engine/Plugins/Enterprise/GLTFExporter/MigratedFromGit/Source/GLTFExporter/Private/Builders/GLTFBufferBuilder.cpp
@@ -48,12 +48,12 @@ const TArray64<uint8>* FGLTFBufferBuilder::GetBufferData() const
 	return bIsGlbFile ? static_cast<FGLTFMemoryArchive*>(BufferArchive.Get()) : nullptr;
 }
 
-FGLTFJsonBufferViewIndex FGLTFBufferBuilder::AddBufferView(const void* RawData, uint64 ByteLength, EGLTFJsonBufferTarget BufferTarget, uint8 DataAlignment)
+FGLTFJsonBufferView* FGLTFBufferBuilder::AddBufferView(const void* RawData, uint64 ByteLength, EGLTFJsonBufferTarget BufferTarget, uint8 DataAlignment)
 {
 	if (BufferArchive == nullptr && !InitializeBuffer())
 	{
 		// TODO: report error
-		return FGLTFJsonBufferViewIndex(INDEX_NONE);
+		return nullptr;
 	}
 
 	uint64 ByteOffset = BufferArchive->Tell();
@@ -71,11 +71,11 @@ FGLTFJsonBufferViewIndex FGLTFBufferBuilder::AddBufferView(const void* RawData,
 	FGLTFJsonBuffer& JsonBuffer = GetBuffer(BufferIndex);
 	JsonBuffer.ByteLength = BufferArchive->Tell();
 
-	FGLTFJsonBufferView JsonBufferView;
-	JsonBufferView.Buffer = BufferIndex;
-	JsonBufferView.ByteOffset = ByteOffset;
-	JsonBufferView.ByteLength = ByteLength;
-	JsonBufferView.Target = BufferTarget;
+	FGLTFJsonBufferView* JsonBufferView = FGLTFJsonBuilder::AddBufferView();
+	JsonBufferView->Buffer = BufferIndex;
+	JsonBufferView->ByteOffset = ByteOffset;
+	JsonBufferView->ByteLength = ByteLength;
+	JsonBufferView->Target = BufferTarget;
 
-	return FGLTFJsonBuilder::AddBufferView(JsonBufferView);
+	return JsonBufferView;
 }
diff --git a/Engine/Plugins/Enterprise/GLTFExporter/MigratedFromGit/Source/GLTFExporter/Private/Builders/GLTFBufferBuilder.h b/Engine/Plugins/Enterprise/GLTFExporter/MigratedFromGit/Source/GLTFExporter/Private/Builders/GLTFBufferBuilder.h
--- a/Engine/Plugins/Enterprise/GLTFExporter/MigratedFromGit/Source/GLTFExporter/Private/Builders/GLTFBufferBuilder.h
+++ b/Engine/Plugins/Enterprise/GLTFExporter/MigratedFromGit/Source/GLTFExporter/Private/Builders/GLTFBufferBuilder.h
@@ -16,10 +16,10 @@ protected:
 
 public:
 
-	FGLTFJsonBufferViewIndex AddBufferView(const void* RawData, uint64 ByteLength, EGLTFJsonBufferTarget BufferTarget = EGLTFJsonBufferTarget::None, uint8 DataAlignment = 4);
+	FGLTFJsonBufferView* AddBufferView(const void* RawData, uint64 ByteLength, EGLTFJsonBufferTarget BufferTarget = EGLTFJsonBufferTarget::None, uint8 DataAlignment = 4);
 
 	template <class ElementType>
-	FGLTFJsonBufferViewIndex AddBufferView(const TArray<ElementType>& Array, EGLTFJsonBufferTarget BufferTarget = EGLTFJsonBufferTarget::None, uint8 DataAlignment = 4)
+	FGLTFJsonBufferView* AddBufferView(const TArray<ElementType>& Array, EGLTFJsonBufferTarget BufferTarget = EGLTFJsonBufferTarget::None, uint8 DataAlignment = 4)
 	{
 		return AddBufferView(Array.GetData(), Array.Num() * sizeof(ElementType), BufferTarget, DataAlignment);
 	}
@@ -28,6 +28,6 @@ private:
 
 	bool InitializeBuffer();
 
-	FGLTFJsonBufferIndex BufferIndex;
+	FGLTFJsonBuffer* BufferIndex;
 	TUniquePtr<FArchive> BufferArchive;
 };
diff --git a/Engine/Plugins/Enterprise/GLTFExporter/MigratedFromGit/Source/GLTFExporter/Private/Builders/GLTFConvertBuilder.cpp b/Engine/Plugins/Enterprise/GLTFExporter/MigratedFromGit/Source/GLTFExporter/Private/Builders/GLTFConvertBuilder.cpp
--- a/Engine/Plugins/Enterprise/GLTFExporter/MigratedFromGit/Source/GLTFExporter/Private/Builders/GLTFConvertBuilder.cpp
+++ b/Engine/Plugins/Enterprise/GLTFExporter/MigratedFromGit/Source/GLTFExporter/Private/Builders/GLTFConvertBuilder.cpp
@@ -24,107 +24,107 @@ bool FGLTFConvertBuilder::IsRootActor(const AActor* Actor) const
 	return ParentActor == nullptr || !IsSelectedActor(ParentActor);
 }
 
-FGLTFJsonAccessorIndex FGLTFConvertBuilder::GetOrAddPositionAccessor(const FGLTFMeshSection* MeshSection, const FPositionVertexBuffer* VertexBuffer)
+FGLTFJsonAccessor* FGLTFConvertBuilder::GetOrAddPositionAccessor(const FGLTFMeshSection* MeshSection, const FPositionVertexBuffer* VertexBuffer)
 {
 	if (VertexBuffer == nullptr)
 	{
-		return FGLTFJsonAccessorIndex(INDEX_NONE);
+		return nullptr;
 	}
 
 	return PositionBufferConverter->GetOrAdd(MeshSection, VertexBuffer);
 }
 
-FGLTFJsonAccessorIndex FGLTFConvertBuilder::GetOrAddColorAccessor(const FGLTFMeshSection* MeshSection, const FColorVertexBuffer* VertexBuffer)
+FGLTFJsonAccessor* FGLTFConvertBuilder::GetOrAddColorAccessor(const FGLTFMeshSection* MeshSection, const FColorVertexBuffer* VertexBuffer)
 {
 	if (VertexBuffer == nullptr)
 	{
-		return FGLTFJsonAccessorIndex(INDEX_NONE);
+		return nullptr;
 	}
 
 	return ColorBufferConverter->GetOrAdd(MeshSection, VertexBuffer);
 }
 
-FGLTFJsonAccessorIndex FGLTFConvertBuilder::GetOrAddNormalAccessor(const FGLTFMeshSection* MeshSection, const FStaticMeshVertexBuffer* VertexBuffer)
+FGLTFJsonAccessor* FGLTFConvertBuilder::GetOrAddNormalAccessor(const FGLTFMeshSection* MeshSection, const FStaticMeshVertexBuffer* VertexBuffer)
 {
 	if (VertexBuffer == nullptr)
 	{
-		return FGLTFJsonAccessorIndex(INDEX_NONE);
+		return nullptr;
 	}
 
 	return NormalBufferConverter->GetOrAdd(MeshSection, VertexBuffer);
 }
 
-FGLTFJsonAccessorIndex FGLTFConvertBuilder::GetOrAddTangentAccessor(const FGLTFMeshSection* MeshSection, const FStaticMeshVertexBuffer* VertexBuffer)
+FGLTFJsonAccessor* FGLTFConvertBuilder::GetOrAddTangentAccessor(const FGLTFMeshSection* MeshSection, const FStaticMeshVertexBuffer* VertexBuffer)
 {
 	if (VertexBuffer == nullptr)
 	{
-		return FGLTFJsonAccessorIndex(INDEX_NONE);
+		return nullptr;
 	}
 
 	return TangentBufferConverter->GetOrAdd(MeshSection, VertexBuffer);
 }
 
-FGLTFJsonAccessorIndex FGLTFConvertBuilder::GetOrAddUVAccessor(const FGLTFMeshSection* MeshSection, const FStaticMeshVertexBuffer* VertexBuffer, int32 UVIndex)
+FGLTFJsonAccessor* FGLTFConvertBuilder::GetOrAddUVAccessor(const FGLTFMeshSection* MeshSection, const FStaticMeshVertexBuffer* VertexBuffer, int32 UVIndex)
 {
 	if (VertexBuffer == nullptr)
 	{
-		return FGLTFJsonAccessorIndex(INDEX_NONE);
+		return nullptr;
 	}
 
 	return UVBufferConverter->GetOrAdd(MeshSection, VertexBuffer, UVIndex);
 }
 
-FGLTFJsonAccessorIndex FGLTFConvertBuilder::GetOrAddJointAccessor(const FGLTFMeshSection* MeshSection, const FSkinWeightVertexBuffer* VertexBuffer, int32 InfluenceOffset)
+FGLTFJsonAccessor* FGLTFConvertBuilder::GetOrAddJointAccessor(const FGLTFMeshSection* MeshSection, const FSkinWeightVertexBuffer* VertexBuffer, int32 InfluenceOffset)
 {
 	if (VertexBuffer == nullptr)
 	{
-		return FGLTFJsonAccessorIndex(INDEX_NONE);
+		return nullptr;
 	}
 
 	return BoneIndexBufferConverter->GetOrAdd(MeshSection, VertexBuffer, InfluenceOffset);
 }
 
-FGLTFJsonAccessorIndex FGLTFConvertBuilder::GetOrAddWeightAccessor(const FGLTFMeshSection* MeshSection, const FSkinWeightVertexBuffer* VertexBuffer, int32 InfluenceOffset)
+FGLTFJsonAccessor* FGLTFConvertBuilder::GetOrAddWeightAccessor(const FGLTFMeshSection* MeshSection, const FSkinWeightVertexBuffer* VertexBuffer, int32 InfluenceOffset)
 {
 	if (VertexBuffer == nullptr)
 	{
-		return FGLTFJsonAccessorIndex(INDEX_NONE);
+		return nullptr;
 	}
 
 	return BoneWeightBufferConverter->GetOrAdd(MeshSection, VertexBuffer, InfluenceOffset);
 }
 
-FGLTFJsonAccessorIndex FGLTFConvertBuilder::GetOrAddIndexAccessor(const FGLTFMeshSection* MeshSection)
+FGLTFJsonAccessor* FGLTFConvertBuilder::GetOrAddIndexAccessor(const FGLTFMeshSection* MeshSection)
 {
 	if (MeshSection == nullptr)
 	{
-		return FGLTFJsonAccessorIndex(INDEX_NONE);
+		return nullptr;
 	}
 
 	return IndexBufferConverter->GetOrAdd(MeshSection);
 }
 
-FGLTFJsonMeshIndex FGLTFConvertBuilder::GetOrAddMesh(const UStaticMesh* StaticMesh, const FGLTFMaterialArray& Materials, int32 LODIndex)
+FGLTFJsonMesh* FGLTFConvertBuilder::GetOrAddMesh(const UStaticMesh* StaticMesh, const FGLTFMaterialArray& Materials, int32 LODIndex)
 {
 	if (StaticMesh == nullptr)
 	{
-		return FGLTFJsonMeshIndex(INDEX_NONE);
+		return nullptr;
 	}
 
 	return StaticMeshConverter->GetOrAdd(StaticMesh, nullptr, Materials, LODIndex);
 }
 
-FGLTFJsonMeshIndex FGLTFConvertBuilder::GetOrAddMesh(const USkeletalMesh* SkeletalMesh, const FGLTFMaterialArray& Materials, int32 LODIndex)
+FGLTFJsonMesh* FGLTFConvertBuilder::GetOrAddMesh(const USkeletalMesh* SkeletalMesh, const FGLTFMaterialArray& Materials, int32 LODIndex)
 {
 	if (SkeletalMesh == nullptr)
 	{
-		return FGLTFJsonMeshIndex(INDEX_NONE);
+		return nullptr;
 	}
 
 	return SkeletalMeshConverter->GetOrAdd(SkeletalMesh, nullptr, Materials, LODIndex);
 }
 
-FGLTFJsonMeshIndex FGLTFConvertBuilder::GetOrAddMesh(const UMeshComponent* MeshComponent, const FGLTFMaterialArray& Materials, int32 LODIndex)
+FGLTFJsonMesh* FGLTFConvertBuilder::GetOrAddMesh(const UMeshComponent* MeshComponent, const FGLTFMaterialArray& Materials, int32 LODIndex)
 {
 	if (const UStaticMeshComponent* StaticMeshComponent = Cast<UStaticMeshComponent>(MeshComponent))
 	{
@@ -136,36 +136,36 @@ FGLTFJsonMeshIndex FGLTFConvertBuilder::GetOrAddMesh(const UMeshComponent* MeshC
 		return GetOrAddMesh(SkeletalMeshComponent, Materials, LODIndex);
 	}
 
-	return FGLTFJsonMeshIndex(INDEX_NONE);
+	return nullptr;
 }
 
-FGLTFJsonMeshIndex FGLTFConvertBuilder::GetOrAddMesh(const UStaticMeshComponent* StaticMeshComponent, const FGLTFMaterialArray& Materials, int32 LODIndex)
+FGLTFJsonMesh* FGLTFConvertBuilder::GetOrAddMesh(const UStaticMeshComponent* StaticMeshComponent, const FGLTFMaterialArray& Materials, int32 LODIndex)
 {
 	if (StaticMeshComponent == nullptr)
 	{
-		return FGLTFJsonMeshIndex(INDEX_NONE);
+		return nullptr;
 	}
 
 	const UStaticMesh* StaticMesh = StaticMeshComponent->GetStaticMesh();
 	if (StaticMesh == nullptr)
 	{
-		return FGLTFJsonMeshIndex(INDEX_NONE);
+		return nullptr;
 	}
 
 	return StaticMeshConverter->GetOrAdd(StaticMesh, StaticMeshComponent, Materials, LODIndex);
 }
 
-FGLTFJsonMeshIndex FGLTFConvertBuilder::GetOrAddMesh(const USkeletalMeshComponent* SkeletalMeshComponent, const FGLTFMaterialArray& Materials, int32 LODIndex)
+FGLTFJsonMesh* FGLTFConvertBuilder::GetOrAddMesh(const USkeletalMeshComponent* SkeletalMeshComponent, const FGLTFMaterialArray& Materials, int32 LODIndex)
 {
 	if (SkeletalMeshComponent == nullptr)
 	{
-		return FGLTFJsonMeshIndex(INDEX_NONE);
+		return nullptr;
 	}
 
 	const USkeletalMesh* SkeletalMesh = SkeletalMeshComponent->SkeletalMesh;
 	if (SkeletalMesh == nullptr)
 	{
-		return FGLTFJsonMeshIndex(INDEX_NONE);
+		return nullptr;
 	}
 
 	return SkeletalMeshConverter->GetOrAdd(SkeletalMesh, SkeletalMeshComponent, Materials, LODIndex);
@@ -181,7 +181,7 @@ const FGLTFMeshData* FGLTFConvertBuilder::GetOrAddMeshData(const USkeletalMesh*
 	return SkeletalMeshDataConverter->GetOrAdd(SkeletalMesh, SkeletalMeshComponent, LODIndex);
 }
 
-FGLTFJsonMaterialIndex FGLTFConvertBuilder::GetOrAddMaterial(const UMaterialInterface* Material, const UStaticMesh* StaticMesh, int32 LODIndex, int32 MaterialIndex)
+FGLTFJsonMaterial* FGLTFConvertBuilder::GetOrAddMaterial(const UMaterialInterface* Material, const UStaticMesh* StaticMesh, int32 LODIndex, int32 MaterialIndex)
 {
 	// TODO: optimize by skipping mesh data if material doesn't need it
 	const FGLTFMeshData* MeshData = GetOrAddMeshData(StaticMesh, nullptr, LODIndex);
@@ -189,7 +189,7 @@ FGLTFJsonMaterialIndex FGLTFConvertBuilder::GetOrAddMaterial(const UMaterialInte
 	return GetOrAddMaterial(Material, MeshData, SectionIndices);
 }
 
-FGLTFJsonMaterialIndex FGLTFConvertBuilder::GetOrAddMaterial(const UMaterialInterface* Material, const USkeletalMesh* SkeletalMesh, int32 LODIndex, int32 MaterialIndex)
+FGLTFJsonMaterial* FGLTFConvertBuilder::GetOrAddMaterial(const UMaterialInterface* Material, const USkeletalMesh* SkeletalMesh, int32 LODIndex, int32 MaterialIndex)
 {
 	// TODO: optimize by skipping mesh data if material doesn't need it
 	const FGLTFMeshData* MeshData = GetOrAddMeshData(SkeletalMesh, nullptr, LODIndex);
@@ -197,7 +197,7 @@ FGLTFJsonMaterialIndex FGLTFConvertBuilder::GetOrAddMaterial(const UMaterialInte
 	return GetOrAddMaterial(Material, MeshData, SectionIndices);
 }
 
-FGLTFJsonMaterialIndex FGLTFConvertBuilder::GetOrAddMaterial(const UMaterialInterface* Material, const UMeshComponent* MeshComponent, int32 LODIndex, int32 MaterialIndex)
+FGLTFJsonMaterial* FGLTFConvertBuilder::GetOrAddMaterial(const UMaterialInterface* Material, const UMeshComponent* MeshComponent, int32 LODIndex, int32 MaterialIndex)
 {
 	if (const UStaticMeshComponent* StaticMeshComponent = Cast<UStaticMeshComponent>(MeshComponent))
 	{
@@ -209,10 +209,10 @@ FGLTFJsonMaterialIndex FGLTFConvertBuilder::GetOrAddMaterial(const UMaterialInte
 		return GetOrAddMaterial(Material, SkeletalMeshComponent, LODIndex, MaterialIndex);
 	}
 
-	return FGLTFJsonMaterialIndex(INDEX_NONE);
+	return nullptr;
 }
 
-FGLTFJsonMaterialIndex FGLTFConvertBuilder::GetOrAddMaterial(const UMaterialInterface* Material, const UStaticMeshComponent* StaticMeshComponent, int32 LODIndex, int32 MaterialIndex)
+FGLTFJsonMaterial* FGLTFConvertBuilder::GetOrAddMaterial(const UMaterialInterface* Material, const UStaticMeshComponent* StaticMeshComponent, int32 LODIndex, int32 MaterialIndex)
 {
 	// TODO: optimize by skipping mesh data if material doesn't need it
 	const UStaticMesh* StaticMesh = StaticMeshComponent->GetStaticMesh();
@@ -221,7 +221,7 @@ FGLTFJsonMaterialIndex FGLTFConvertBuilder::GetOrAddMaterial(const UMaterialInte
 	return GetOrAddMaterial(Material, MeshData, SectionIndices);
 }
 
-FGLTFJsonMaterialIndex FGLTFConvertBuilder::GetOrAddMaterial(const UMaterialInterface* Material, const USkeletalMeshComponent* SkeletalMeshComponent, int32 LODIndex, int32 MaterialIndex)
+FGLTFJsonMaterial* FGLTFConvertBuilder::GetOrAddMaterial(const UMaterialInterface* Material, const USkeletalMeshComponent* SkeletalMeshComponent, int32 LODIndex, int32 MaterialIndex)
 {
 	// TODO: optimize by skipping mesh data if material doesn't need it
 	const USkeletalMesh* SkeletalMesh = SkeletalMeshComponent->SkeletalMesh;
@@ -230,62 +230,62 @@ FGLTFJsonMaterialIndex FGLTFConvertBuilder::GetOrAddMaterial(const UMaterialInte
 	return GetOrAddMaterial(Material, MeshData, SectionIndices);
 }
 
-FGLTFJsonMaterialIndex FGLTFConvertBuilder::GetOrAddMaterial(const UMaterialInterface* Material, const FGLTFMeshData* MeshData, const FGLTFIndexArray& SectionIndices)
+FGLTFJsonMaterial* FGLTFConvertBuilder::GetOrAddMaterial(const UMaterialInterface* Material, const FGLTFMeshData* MeshData, const FGLTFIndexArray& SectionIndices)
 {
 	if (Material == nullptr)
 	{
-		return FGLTFJsonMaterialIndex(INDEX_NONE);
+		return nullptr;
 	}
 
 	return MaterialConverter->GetOrAdd(Material, MeshData, SectionIndices);
 }
 
-FGLTFJsonSamplerIndex FGLTFConvertBuilder::GetOrAddSampler(const UTexture* Texture)
+FGLTFJsonSampler* FGLTFConvertBuilder::GetOrAddSampler(const UTexture* Texture)
 {
 	if (Texture == nullptr)
 	{
-		return FGLTFJsonSamplerIndex(INDEX_NONE);
+		return nullptr;
 	}
 
 	return SamplerConverter->GetOrAdd(Texture);
 }
 
-FGLTFJsonTextureIndex FGLTFConvertBuilder::GetOrAddTexture(const UTexture* Texture)
+FGLTFJsonTexture* FGLTFConvertBuilder::GetOrAddTexture(const UTexture* Texture)
 {
 	return GetOrAddTexture(Texture, Texture->SRGB);
 }
 
-FGLTFJsonTextureIndex FGLTFConvertBuilder::GetOrAddTexture(const UTexture2D* Texture)
+FGLTFJsonTexture* FGLTFConvertBuilder::GetOrAddTexture(const UTexture2D* Texture)
 {
 	return GetOrAddTexture(Texture, Texture->SRGB);
 }
 
-FGLTFJsonTextureIndex FGLTFConvertBuilder::GetOrAddTexture(const UTextureCube* Texture, ECubeFace CubeFace)
+FGLTFJsonTexture* FGLTFConvertBuilder::GetOrAddTexture(const UTextureCube* Texture, ECubeFace CubeFace)
 {
 	return GetOrAddTexture(Texture, CubeFace, Texture->SRGB);
 }
 
-FGLTFJsonTextureIndex FGLTFConvertBuilder::GetOrAddTexture(const UTextureRenderTarget2D* Texture)
+FGLTFJsonTexture* FGLTFConvertBuilder::GetOrAddTexture(const UTextureRenderTarget2D* Texture)
 {
 	return GetOrAddTexture(Texture, Texture->SRGB);
 }
 
-FGLTFJsonTextureIndex FGLTFConvertBuilder::GetOrAddTexture(const UTextureRenderTargetCube* Texture, ECubeFace CubeFace)
+FGLTFJsonTexture* FGLTFConvertBuilder::GetOrAddTexture(const UTextureRenderTargetCube* Texture, ECubeFace CubeFace)
 {
 	return GetOrAddTexture(Texture, CubeFace, Texture->SRGB);
 }
 
-FGLTFJsonTextureIndex FGLTFConvertBuilder::GetOrAddTexture(const ULightMapTexture2D* Texture)
+FGLTFJsonTexture* FGLTFConvertBuilder::GetOrAddTexture(const ULightMapTexture2D* Texture)
 {
 	if (Texture == nullptr)
 	{
-		return FGLTFJsonTextureIndex(INDEX_NONE);
+		return nullptr;
 	}
 
 	return TextureLightMapConverter->GetOrAdd(Texture);
 }
 
-FGLTFJsonTextureIndex FGLTFConvertBuilder::GetOrAddTexture(const UTexture* Texture, bool bToSRGB)
+FGLTFJsonTexture* FGLTFConvertBuilder::GetOrAddTexture(const UTexture* Texture, bool bToSRGB)
 {
 	if (const UTexture2D* Texture2D = Cast<UTexture2D>(Texture))
 	{
@@ -297,266 +297,266 @@ FGLTFJsonTextureIndex FGLTFConvertBuilder::GetOrAddTexture(const UTexture* Textu
 		return GetOrAddTexture(RenderTarget2D, bToSRGB);
 	}
 
-	return FGLTFJsonTextureIndex(INDEX_NONE);
+	return nullptr;
 }
 
-FGLTFJsonTextureIndex FGLTFConvertBuilder::GetOrAddTexture(const UTexture2D* Texture, bool bToSRGB)
+FGLTFJsonTexture* FGLTFConvertBuilder::GetOrAddTexture(const UTexture2D* Texture, bool bToSRGB)
 {
 	if (Texture == nullptr)
 	{
-		return FGLTFJsonTextureIndex(INDEX_NONE);
+		return nullptr;
 	}
 
 	return Texture2DConverter->GetOrAdd(Texture, bToSRGB);
 }
 
-FGLTFJsonTextureIndex FGLTFConvertBuilder::GetOrAddTexture(const UTextureCube* Texture, ECubeFace CubeFace, bool bToSRGB)
+FGLTFJsonTexture* FGLTFConvertBuilder::GetOrAddTexture(const UTextureCube* Texture, ECubeFace CubeFace, bool bToSRGB)
 {
 	if (Texture == nullptr)
 	{
-		return FGLTFJsonTextureIndex(INDEX_NONE);
+		return nullptr;
 	}
 
 	return TextureCubeConverter->GetOrAdd(Texture, CubeFace, bToSRGB);
 }
 
-FGLTFJsonTextureIndex FGLTFConvertBuilder::GetOrAddTexture(const UTextureRenderTarget2D* Texture, bool bToSRGB)
+FGLTFJsonTexture* FGLTFConvertBuilder::GetOrAddTexture(const UTextureRenderTarget2D* Texture, bool bToSRGB)
 {
 	if (Texture == nullptr)
 	{
-		return FGLTFJsonTextureIndex(INDEX_NONE);
+		return nullptr;
 	}
 
 	return TextureRenderTarget2DConverter->GetOrAdd(Texture, bToSRGB);
 }
 
-FGLTFJsonTextureIndex FGLTFConvertBuilder::GetOrAddTexture(const UTextureRenderTargetCube* Texture, ECubeFace CubeFace, bool bToSRGB)
+FGLTFJsonTexture* FGLTFConvertBuilder::GetOrAddTexture(const UTextureRenderTargetCube* Texture, ECubeFace CubeFace, bool bToSRGB)
 {
 	if (Texture == nullptr)
 	{
-		return FGLTFJsonTextureIndex(INDEX_NONE);
+		return nullptr;
 	}
 
 	return TextureRenderTargetCubeConverter->GetOrAdd(Texture, CubeFace, bToSRGB);
 }
 
-FGLTFJsonImageIndex FGLTFConvertBuilder::GetOrAddImage(TGLTFSharedArray<FColor>& Pixels, FIntPoint Size, bool bIgnoreAlpha, EGLTFTextureType Type, const FString& Name)
+FGLTFJsonImage* FGLTFConvertBuilder::GetOrAddImage(TGLTFSharedArray<FColor>& Pixels, FIntPoint Size, bool bIgnoreAlpha, EGLTFTextureType Type, const FString& Name)
 {
 	return ImageConverter->GetOrAdd(Name, Type, bIgnoreAlpha, Size, Pixels);
 }
 
-FGLTFJsonSkinIndex FGLTFConvertBuilder::GetOrAddSkin(FGLTFJsonNodeIndex RootNode, const USkeletalMesh* SkeletalMesh)
+FGLTFJsonSkin* FGLTFConvertBuilder::GetOrAddSkin(FGLTFJsonNode* RootNode, const USkeletalMesh* SkeletalMesh)
 {
-	if (RootNode == INDEX_NONE || SkeletalMesh == nullptr)
+	if (RootNode == nullptr || SkeletalMesh == nullptr)
 	{
-		return FGLTFJsonSkinIndex(INDEX_NONE);
+		return nullptr;
 	}
 
 	return SkinConverter->GetOrAdd(RootNode, SkeletalMesh);
 }
 
-FGLTFJsonSkinIndex FGLTFConvertBuilder::GetOrAddSkin(FGLTFJsonNodeIndex RootNode, const USkeletalMeshComponent* SkeletalMeshComponent)
+FGLTFJsonSkin* FGLTFConvertBuilder::GetOrAddSkin(FGLTFJsonNode* RootNode, const USkeletalMeshComponent* SkeletalMeshComponent)
 {
-	if (RootNode == INDEX_NONE || SkeletalMeshComponent == nullptr)
+	if (RootNode == nullptr || SkeletalMeshComponent == nullptr)
 	{
-		return FGLTFJsonSkinIndex(INDEX_NONE);
+		return nullptr;
 	}
 
 	const USkeletalMesh* SkeletalMesh = SkeletalMeshComponent->SkeletalMesh;
 
 	if (SkeletalMesh == nullptr)
 	{
-		return FGLTFJsonSkinIndex(INDEX_NONE);
+		return nullptr;
 	}
 
 	return GetOrAddSkin(RootNode, SkeletalMesh);
 }
 
-FGLTFJsonAnimationIndex FGLTFConvertBuilder::GetOrAddAnimation(FGLTFJsonNodeIndex RootNode, const USkeletalMesh* SkeletalMesh, const UAnimSequence* AnimSequence)
+FGLTFJsonAnimation* FGLTFConvertBuilder::GetOrAddAnimation(FGLTFJsonNode* RootNode, const USkeletalMesh* SkeletalMesh, const UAnimSequence* AnimSequence)
 {
-	if (RootNode == INDEX_NONE || SkeletalMesh == nullptr || AnimSequence == nullptr)
+	if (RootNode == nullptr || SkeletalMesh == nullptr || AnimSequence == nullptr)
 	{
-		return FGLTFJsonAnimationIndex(INDEX_NONE);
+		return nullptr;
 	}
 
 	return AnimationConverter->GetOrAdd(RootNode, SkeletalMesh, AnimSequence);
 }
 
-FGLTFJsonAnimationIndex FGLTFConvertBuilder::GetOrAddAnimation(FGLTFJsonNodeIndex RootNode, const USkeletalMeshComponent* SkeletalMeshComponent)
+FGLTFJsonAnimation* FGLTFConvertBuilder::GetOrAddAnimation(FGLTFJsonNode* RootNode, const USkeletalMeshComponent* SkeletalMeshComponent)
 {
-	if (RootNode == INDEX_NONE || SkeletalMeshComponent == nullptr)
+	if (RootNode == nullptr || SkeletalMeshComponent == nullptr)
 	{
-		return FGLTFJsonAnimationIndex(INDEX_NONE);
+		return nullptr;
 	}
 
 	return AnimationDataConverter->GetOrAdd(RootNode, SkeletalMeshComponent);
 }
 
-FGLTFJsonAnimationIndex FGLTFConvertBuilder::GetOrAddAnimation(const ULevel* Level, const ULevelSequence* LevelSequence)
+FGLTFJsonAnimation* FGLTFConvertBuilder::GetOrAddAnimation(const ULevel* Level, const ULevelSequence* LevelSequence)
 {
 	if (Level == nullptr || LevelSequence == nullptr)
 	{
-		return FGLTFJsonAnimationIndex(INDEX_NONE);
+		return nullptr;
 	}
 
 	return LevelSequenceConverter->GetOrAdd(Level, LevelSequence);
 }
 
-FGLTFJsonAnimationIndex FGLTFConvertBuilder::GetOrAddAnimation(const ALevelSequenceActor* LevelSequenceActor)
+FGLTFJsonAnimation* FGLTFConvertBuilder::GetOrAddAnimation(const ALevelSequenceActor* LevelSequenceActor)
 {
 	if (LevelSequenceActor == nullptr)
 	{
-		return FGLTFJsonAnimationIndex(INDEX_NONE);
+		return nullptr;
 	}
 
 	return LevelSequenceDataConverter->GetOrAdd(LevelSequenceActor);
 }
 
-FGLTFJsonNodeIndex FGLTFConvertBuilder::GetOrAddNode(const AActor* Actor)
+FGLTFJsonNode* FGLTFConvertBuilder::GetOrAddNode(const AActor* Actor)
 {
 	if (Actor == nullptr)
 	{
-		return FGLTFJsonNodeIndex(INDEX_NONE);
+		return nullptr;
 	}
 
 	return ActorConverter->GetOrAdd(Actor);
 }
 
-FGLTFJsonNodeIndex FGLTFConvertBuilder::GetOrAddNode(const USceneComponent* SceneComponent)
+FGLTFJsonNode* FGLTFConvertBuilder::GetOrAddNode(const USceneComponent* SceneComponent)
 {
 	if (SceneComponent == nullptr)
 	{
-		return FGLTFJsonNodeIndex(INDEX_NONE);
+		return nullptr;
 	}
 
 	return ComponentConverter->GetOrAdd(SceneComponent);
 }
 
-FGLTFJsonNodeIndex FGLTFConvertBuilder::GetOrAddNode(const USceneComponent* SceneComponent, FName SocketName)
+FGLTFJsonNode* FGLTFConvertBuilder::GetOrAddNode(const USceneComponent* SceneComponent, FName SocketName)
 {
 	if (SceneComponent == nullptr)
 	{
-		return FGLTFJsonNodeIndex(INDEX_NONE);
+		return nullptr;
 	}
 
 	return ComponentSocketConverter->GetOrAdd(SceneComponent, SocketName);
 }
 
-FGLTFJsonNodeIndex FGLTFConvertBuilder::GetOrAddNode(FGLTFJsonNodeIndex RootNode, const UStaticMesh* StaticMesh, FName SocketName)
+FGLTFJsonNode* FGLTFConvertBuilder::GetOrAddNode(FGLTFJsonNode* RootNode, const UStaticMesh* StaticMesh, FName SocketName)
 {
-	if (RootNode == INDEX_NONE || StaticMesh == nullptr || SocketName == NAME_None)
+	if (RootNode == nullptr || StaticMesh == nullptr || SocketName == NAME_None)
 	{
-		return FGLTFJsonNodeIndex(INDEX_NONE);
+		return nullptr;
 	}
 
 	return StaticSocketConverter->GetOrAdd(RootNode, StaticMesh, SocketName);
 }
 
-FGLTFJsonNodeIndex FGLTFConvertBuilder::GetOrAddNode(FGLTFJsonNodeIndex RootNode, const USkeletalMesh* SkeletalMesh, FName SocketName)
+FGLTFJsonNode* FGLTFConvertBuilder::GetOrAddNode(FGLTFJsonNode* RootNode, const USkeletalMesh* SkeletalMesh, FName SocketName)
 {
-	if (RootNode == INDEX_NONE || SkeletalMesh == nullptr || SocketName == NAME_None)
+	if (RootNode == nullptr || SkeletalMesh == nullptr || SocketName == NAME_None)
 	{
-		return FGLTFJsonNodeIndex(INDEX_NONE);
+		return nullptr;
 	}
 
 	return SkeletalSocketConverter->GetOrAdd(RootNode, SkeletalMesh, SocketName);
 }
 
-FGLTFJsonNodeIndex FGLTFConvertBuilder::GetOrAddNode(FGLTFJsonNodeIndex RootNode, const USkeletalMesh* SkeletalMesh, int32 BoneIndex)
+FGLTFJsonNode* FGLTFConvertBuilder::GetOrAddNode(FGLTFJsonNode* RootNode, const USkeletalMesh* SkeletalMesh, int32 BoneIndex)
 {
-	if (RootNode == INDEX_NONE || SkeletalMesh == nullptr || BoneIndex == INDEX_NONE)
+	if (RootNode == nullptr || SkeletalMesh == nullptr || BoneIndex == INDEX_NONE)
 	{
-		return FGLTFJsonNodeIndex(INDEX_NONE);
+		return nullptr;
 	}
 
 	return SkeletalBoneConverter->GetOrAdd(RootNode, SkeletalMesh, BoneIndex);
 }
 
-FGLTFJsonSceneIndex FGLTFConvertBuilder::GetOrAddScene(const UWorld* World)
+FGLTFJsonScene* FGLTFConvertBuilder::GetOrAddScene(const UWorld* World)
 {
 	if (World == nullptr)
 	{
-		return FGLTFJsonSceneIndex(INDEX_NONE);
+		return nullptr;
 	}
 
 	return SceneConverter->GetOrAdd(World);
 }
 
-FGLTFJsonCameraIndex FGLTFConvertBuilder::GetOrAddCamera(const UCameraComponent* CameraComponent)
+FGLTFJsonCamera* FGLTFConvertBuilder::GetOrAddCamera(const UCameraComponent* CameraComponent)
 {
 	if (CameraComponent == nullptr)
 	{
-		return FGLTFJsonCameraIndex(INDEX_NONE);
+		return nullptr;
 	}
 
 	return CameraConverter->GetOrAdd(CameraComponent);
 }
 
-FGLTFJsonLightIndex FGLTFConvertBuilder::GetOrAddLight(const ULightComponent* LightComponent)
+FGLTFJsonLight* FGLTFConvertBuilder::GetOrAddLight(const ULightComponent* LightComponent)
 {
 	if (LightComponent == nullptr)
 	{
-		return FGLTFJsonLightIndex(INDEX_NONE);
+		return nullptr;
 	}
 
 	return LightConverter->GetOrAdd(LightComponent);
 }
 
-FGLTFJsonBackdropIndex FGLTFConvertBuilder::GetOrAddBackdrop(const AActor* BackdropActor)
+FGLTFJsonBackdrop* FGLTFConvertBuilder::GetOrAddBackdrop(const AActor* BackdropActor)
 {
 	if (BackdropActor == nullptr)
 	{
-		return FGLTFJsonBackdropIndex(INDEX_NONE);
+		return nullptr;
 	}
 
 	return BackdropConverter->GetOrAdd(BackdropActor);
 }
 
-FGLTFJsonLightMapIndex FGLTFConvertBuilder::GetOrAddLightMap(const UStaticMeshComponent* StaticMeshComponent)
+FGLTFJsonLightMap* FGLTFConvertBuilder::GetOrAddLightMap(const UStaticMeshComponent* StaticMeshComponent)
 {
 	if (StaticMeshComponent == nullptr)
 	{
-		return FGLTFJsonLightMapIndex(INDEX_NONE);
+		return nullptr;
 	}
 
 	return LightMapConverter->GetOrAdd(StaticMeshComponent);
 }
 
-FGLTFJsonHotspotIndex FGLTFConvertBuilder::GetOrAddHotspot(const AGLTFHotspotActor* HotspotActor)
+FGLTFJsonHotspot* FGLTFConvertBuilder::GetOrAddHotspot(const AGLTFHotspotActor* HotspotActor)
 {
 	if (HotspotActor == nullptr)
 	{
-		return FGLTFJsonHotspotIndex(INDEX_NONE);
+		return nullptr;
 	}
 
 	return HotspotConverter->GetOrAdd(HotspotActor);
 }
 
-FGLTFJsonSkySphereIndex FGLTFConvertBuilder::GetOrAddSkySphere(const AActor* SkySphereActor)
+FGLTFJsonSkySphere* FGLTFConvertBuilder::GetOrAddSkySphere(const AActor* SkySphereActor)
 {
 	if (SkySphereActor == nullptr)
 	{
-		return FGLTFJsonSkySphereIndex(INDEX_NONE);
+		return nullptr;
 	}
 
 	return SkySphereConverter->GetOrAdd(SkySphereActor);
 }
 
-FGLTFJsonEpicLevelVariantSetsIndex FGLTFConvertBuilder::GetOrAddEpicLevelVariantSets(const ULevelVariantSets* LevelVariantSets)
+FGLTFJsonEpicLevelVariantSets* FGLTFConvertBuilder::GetOrAddEpicLevelVariantSets(const ULevelVariantSets* LevelVariantSets)
 {
 	if (LevelVariantSets == nullptr)
 	{
-		return FGLTFJsonEpicLevelVariantSetsIndex(INDEX_NONE);
+		return nullptr;
 	}
 
 	return EpicLevelVariantSetsConverter->GetOrAdd(LevelVariantSets);
 }
 
-FGLTFJsonKhrMaterialVariantIndex FGLTFConvertBuilder::GetOrAddKhrMaterialVariant(const UVariant* Variant)
+FGLTFJsonKhrMaterialVariant* FGLTFConvertBuilder::GetOrAddKhrMaterialVariant(const UVariant* Variant)
 {
 	if (Variant == nullptr)
 	{
-		return FGLTFJsonKhrMaterialVariantIndex(INDEX_NONE);
+		return nullptr;
 	}
 
 	return KhrMaterialVariantConverter->GetOrAdd(Variant);
diff --git a/Engine/Plugins/Enterprise/GLTFExporter/MigratedFromGit/Source/GLTFExporter/Private/Builders/GLTFConvertBuilder.h b/Engine/Plugins/Enterprise/GLTFExporter/MigratedFromGit/Source/GLTFExporter/Private/Builders/GLTFConvertBuilder.h
--- a/Engine/Plugins/Enterprise/GLTFExporter/MigratedFromGit/Source/GLTFExporter/Private/Builders/GLTFConvertBuilder.h
+++ b/Engine/Plugins/Enterprise/GLTFExporter/MigratedFromGit/Source/GLTFExporter/Private/Builders/GLTFConvertBuilder.h
@@ -34,69 +34,69 @@ public:
 	bool IsSelectedActor(const AActor* Object) const;
 	bool IsRootActor(const AActor* Actor) const;
 
-	FGLTFJsonAccessorIndex GetOrAddPositionAccessor(const FGLTFMeshSection* MeshSection, const FPositionVertexBuffer* VertexBuffer);
-	FGLTFJsonAccessorIndex GetOrAddColorAccessor(const FGLTFMeshSection* MeshSection, const FColorVertexBuffer* VertexBuffer);
-	FGLTFJsonAccessorIndex GetOrAddNormalAccessor(const FGLTFMeshSection* MeshSection, const FStaticMeshVertexBuffer* VertexBuffer);
-	FGLTFJsonAccessorIndex GetOrAddTangentAccessor(const FGLTFMeshSection* MeshSection, const FStaticMeshVertexBuffer* VertexBuffer);
-	FGLTFJsonAccessorIndex GetOrAddUVAccessor(const FGLTFMeshSection* MeshSection, const FStaticMeshVertexBuffer* VertexBuffer, int32 UVIndex);
-	FGLTFJsonAccessorIndex GetOrAddJointAccessor(const FGLTFMeshSection* MeshSection, const FSkinWeightVertexBuffer* VertexBuffer, int32 InfluenceOffset);
-	FGLTFJsonAccessorIndex GetOrAddWeightAccessor(const FGLTFMeshSection* MeshSection, const FSkinWeightVertexBuffer* VertexBuffer, int32 InfluenceOffset);
-	FGLTFJsonAccessorIndex GetOrAddIndexAccessor(const FGLTFMeshSection* MeshSection);
-
-	FGLTFJsonMeshIndex GetOrAddMesh(const UStaticMesh* StaticMesh, const FGLTFMaterialArray& Materials = {}, int32 LODIndex = INDEX_NONE);
-	FGLTFJsonMeshIndex GetOrAddMesh(const USkeletalMesh* SkeletalMesh, const FGLTFMaterialArray& Materials = {}, int32 LODIndex = INDEX_NONE);
-	FGLTFJsonMeshIndex GetOrAddMesh(const UMeshComponent* MeshComponent, const FGLTFMaterialArray& Materials = {}, int32 LODIndex = INDEX_NONE);
-	FGLTFJsonMeshIndex GetOrAddMesh(const UStaticMeshComponent* StaticMeshComponent, const FGLTFMaterialArray& Materials = {}, int32 LODIndex = INDEX_NONE);
-	FGLTFJsonMeshIndex GetOrAddMesh(const USkeletalMeshComponent* SkeletalMeshComponent, const FGLTFMaterialArray& Materials = {}, int32 LODIndex = INDEX_NONE);
+	FGLTFJsonAccessor* GetOrAddPositionAccessor(const FGLTFMeshSection* MeshSection, const FPositionVertexBuffer* VertexBuffer);
+	FGLTFJsonAccessor* GetOrAddColorAccessor(const FGLTFMeshSection* MeshSection, const FColorVertexBuffer* VertexBuffer);
+	FGLTFJsonAccessor* GetOrAddNormalAccessor(const FGLTFMeshSection* MeshSection, const FStaticMeshVertexBuffer* VertexBuffer);
+	FGLTFJsonAccessor* GetOrAddTangentAccessor(const FGLTFMeshSection* MeshSection, const FStaticMeshVertexBuffer* VertexBuffer);
+	FGLTFJsonAccessor* GetOrAddUVAccessor(const FGLTFMeshSection* MeshSection, const FStaticMeshVertexBuffer* VertexBuffer, int32 UVIndex);
+	FGLTFJsonAccessor* GetOrAddJointAccessor(const FGLTFMeshSection* MeshSection, const FSkinWeightVertexBuffer* VertexBuffer, int32 InfluenceOffset);
+	FGLTFJsonAccessor* GetOrAddWeightAccessor(const FGLTFMeshSection* MeshSection, const FSkinWeightVertexBuffer* VertexBuffer, int32 InfluenceOffset);
+	FGLTFJsonAccessor* GetOrAddIndexAccessor(const FGLTFMeshSection* MeshSection);
+
+	FGLTFJsonMesh* GetOrAddMesh(const UStaticMesh* StaticMesh, const FGLTFMaterialArray& Materials = {}, int32 LODIndex = INDEX_NONE);
+	FGLTFJsonMesh* GetOrAddMesh(const USkeletalMesh* SkeletalMesh, const FGLTFMaterialArray& Materials = {}, int32 LODIndex = INDEX_NONE);
+	FGLTFJsonMesh* GetOrAddMesh(const UMeshComponent* MeshComponent, const FGLTFMaterialArray& Materials = {}, int32 LODIndex = INDEX_NONE);
+	FGLTFJsonMesh* GetOrAddMesh(const UStaticMeshComponent* StaticMeshComponent, const FGLTFMaterialArray& Materials = {}, int32 LODIndex = INDEX_NONE);
+	FGLTFJsonMesh* GetOrAddMesh(const USkeletalMeshComponent* SkeletalMeshComponent, const FGLTFMaterialArray& Materials = {}, int32 LODIndex = INDEX_NONE);
 
 	const FGLTFMeshData* GetOrAddMeshData(const UStaticMesh* StaticMesh, const UStaticMeshComponent* StaticMeshComponent = nullptr, int32 LODIndex = INDEX_NONE);
 	const FGLTFMeshData* GetOrAddMeshData(const USkeletalMesh* SkeletalMesh, const USkeletalMeshComponent* SkeletalMeshComponent = nullptr, int32 LODIndex = INDEX_NONE);
 
-	FGLTFJsonMaterialIndex GetOrAddMaterial(const UMaterialInterface* Material, const UStaticMesh* StaticMesh, int32 LODIndex = INDEX_NONE, int32 MaterialIndex = INDEX_NONE);
-	FGLTFJsonMaterialIndex GetOrAddMaterial(const UMaterialInterface* Material, const USkeletalMesh* SkeletalMesh, int32 LODIndex = INDEX_NONE, int32 MaterialIndex = INDEX_NONE);
-	FGLTFJsonMaterialIndex GetOrAddMaterial(const UMaterialInterface* Material, const UMeshComponent* MeshComponent, int32 LODIndex = INDEX_NONE, int32 MaterialIndex = INDEX_NONE);
-	FGLTFJsonMaterialIndex GetOrAddMaterial(const UMaterialInterface* Material, const UStaticMeshComponent* StaticMeshComponent, int32 LODIndex = INDEX_NONE, int32 MaterialIndex = INDEX_NONE);
-	FGLTFJsonMaterialIndex GetOrAddMaterial(const UMaterialInterface* Material, const USkeletalMeshComponent* SkeletalMeshComponent, int32 LODIndex = INDEX_NONE, int32 MaterialIndex = INDEX_NONE);
-	FGLTFJsonMaterialIndex GetOrAddMaterial(const UMaterialInterface* Material, const FGLTFMeshData* MeshData = nullptr, const FGLTFIndexArray& SectionIndices = {});
-
-	FGLTFJsonSamplerIndex GetOrAddSampler(const UTexture* Texture);
-	FGLTFJsonTextureIndex GetOrAddTexture(const UTexture* Texture);
-	FGLTFJsonTextureIndex GetOrAddTexture(const UTexture2D* Texture);
-	FGLTFJsonTextureIndex GetOrAddTexture(const UTextureCube* Texture, ECubeFace CubeFace);
-	FGLTFJsonTextureIndex GetOrAddTexture(const UTextureRenderTarget2D* Texture);
-	FGLTFJsonTextureIndex GetOrAddTexture(const UTextureRenderTargetCube* Texture, ECubeFace CubeFace);
-	FGLTFJsonTextureIndex GetOrAddTexture(const ULightMapTexture2D* Texture);
-	FGLTFJsonTextureIndex GetOrAddTexture(const UTexture* Texture, bool bToSRGB);
-	FGLTFJsonTextureIndex GetOrAddTexture(const UTexture2D* Texture, bool bToSRGB);
-	FGLTFJsonTextureIndex GetOrAddTexture(const UTextureCube* Texture, ECubeFace CubeFace, bool bToSRGB);
-	FGLTFJsonTextureIndex GetOrAddTexture(const UTextureRenderTarget2D* Texture, bool bToSRGB);
-	FGLTFJsonTextureIndex GetOrAddTexture(const UTextureRenderTargetCube* Texture, ECubeFace CubeFace, bool bToSRGB);
-	FGLTFJsonImageIndex GetOrAddImage(TGLTFSharedArray<FColor>& Pixels, FIntPoint Size, bool bIgnoreAlpha, EGLTFTextureType Type, const FString& Name);
-
-	FGLTFJsonSkinIndex GetOrAddSkin(FGLTFJsonNodeIndex RootNode, const USkeletalMesh* SkeletalMesh);
-	FGLTFJsonSkinIndex GetOrAddSkin(FGLTFJsonNodeIndex RootNode, const USkeletalMeshComponent* SkeletalMeshComponent);
-	FGLTFJsonAnimationIndex GetOrAddAnimation(FGLTFJsonNodeIndex RootNode, const USkeletalMesh* SkeletalMesh, const UAnimSequence* AnimSequence);
-	FGLTFJsonAnimationIndex GetOrAddAnimation(FGLTFJsonNodeIndex RootNode, const USkeletalMeshComponent* SkeletalMeshComponent);
-	FGLTFJsonAnimationIndex GetOrAddAnimation(const ULevel* Level, const ULevelSequence* LevelSequence);
-	FGLTFJsonAnimationIndex GetOrAddAnimation(const ALevelSequenceActor* LevelSequenceActor);
-
-	FGLTFJsonNodeIndex GetOrAddNode(const AActor* Actor);
-	FGLTFJsonNodeIndex GetOrAddNode(const USceneComponent* SceneComponent);
-	FGLTFJsonNodeIndex GetOrAddNode(const USceneComponent* SceneComponent, FName SocketName);
-	FGLTFJsonNodeIndex GetOrAddNode(FGLTFJsonNodeIndex RootNode, const UStaticMesh* StaticMesh, FName SocketName);
-	FGLTFJsonNodeIndex GetOrAddNode(FGLTFJsonNodeIndex RootNode, const USkeletalMesh* SkeletalMesh, FName SocketName);
-	FGLTFJsonNodeIndex GetOrAddNode(FGLTFJsonNodeIndex RootNode, const USkeletalMesh* SkeletalMesh, int32 BoneIndex);
-	FGLTFJsonSceneIndex GetOrAddScene(const UWorld* World);
-
-	FGLTFJsonCameraIndex GetOrAddCamera(const UCameraComponent* CameraComponent);
-	FGLTFJsonLightIndex GetOrAddLight(const ULightComponent* LightComponent);
-	FGLTFJsonBackdropIndex GetOrAddBackdrop(const AActor* BackdropActor);
-	FGLTFJsonLightMapIndex GetOrAddLightMap(const UStaticMeshComponent* StaticMeshComponent);
-	FGLTFJsonHotspotIndex GetOrAddHotspot(const AGLTFHotspotActor* HotspotActor);
-	FGLTFJsonSkySphereIndex GetOrAddSkySphere(const AActor* SkySphereActor);
-
-	FGLTFJsonEpicLevelVariantSetsIndex GetOrAddEpicLevelVariantSets(const ULevelVariantSets* LevelVariantSets);
-	FGLTFJsonKhrMaterialVariantIndex GetOrAddKhrMaterialVariant(const UVariant* Variant);
+	FGLTFJsonMaterial* GetOrAddMaterial(const UMaterialInterface* Material, const UStaticMesh* StaticMesh, int32 LODIndex = INDEX_NONE, int32 MaterialIndex = INDEX_NONE);
+	FGLTFJsonMaterial* GetOrAddMaterial(const UMaterialInterface* Material, const USkeletalMesh* SkeletalMesh, int32 LODIndex = INDEX_NONE, int32 MaterialIndex = INDEX_NONE);
+	FGLTFJsonMaterial* GetOrAddMaterial(const UMaterialInterface* Material, const UMeshComponent* MeshComponent, int32 LODIndex = INDEX_NONE, int32 MaterialIndex = INDEX_NONE);
+	FGLTFJsonMaterial* GetOrAddMaterial(const UMaterialInterface* Material, const UStaticMeshComponent* StaticMeshComponent, int32 LODIndex = INDEX_NONE, int32 MaterialIndex = INDEX_NONE);
+	FGLTFJsonMaterial* GetOrAddMaterial(const UMaterialInterface* Material, const USkeletalMeshComponent* SkeletalMeshComponent, int32 LODIndex = INDEX_NONE, int32 MaterialIndex = INDEX_NONE);
+	FGLTFJsonMaterial* GetOrAddMaterial(const UMaterialInterface* Material, const FGLTFMeshData* MeshData = nullptr, const FGLTFIndexArray& SectionIndices = {});
+
+	FGLTFJsonSampler* GetOrAddSampler(const UTexture* Texture);
+	FGLTFJsonTexture* GetOrAddTexture(const UTexture* Texture);
+	FGLTFJsonTexture* GetOrAddTexture(const UTexture2D* Texture);
+	FGLTFJsonTexture* GetOrAddTexture(const UTextureCube* Texture, ECubeFace CubeFace);
+	FGLTFJsonTexture* GetOrAddTexture(const UTextureRenderTarget2D* Texture);
+	FGLTFJsonTexture* GetOrAddTexture(const UTextureRenderTargetCube* Texture, ECubeFace CubeFace);
+	FGLTFJsonTexture* GetOrAddTexture(const ULightMapTexture2D* Texture);
+	FGLTFJsonTexture* GetOrAddTexture(const UTexture* Texture, bool bToSRGB);
+	FGLTFJsonTexture* GetOrAddTexture(const UTexture2D* Texture, bool bToSRGB);
+	FGLTFJsonTexture* GetOrAddTexture(const UTextureCube* Texture, ECubeFace CubeFace, bool bToSRGB);
+	FGLTFJsonTexture* GetOrAddTexture(const UTextureRenderTarget2D* Texture, bool bToSRGB);
+	FGLTFJsonTexture* GetOrAddTexture(const UTextureRenderTargetCube* Texture, ECubeFace CubeFace, bool bToSRGB);
+	FGLTFJsonImage* GetOrAddImage(TGLTFSharedArray<FColor>& Pixels, FIntPoint Size, bool bIgnoreAlpha, EGLTFTextureType Type, const FString& Name);
+
+	FGLTFJsonSkin* GetOrAddSkin(FGLTFJsonNode* RootNode, const USkeletalMesh* SkeletalMesh);
+	FGLTFJsonSkin* GetOrAddSkin(FGLTFJsonNode* RootNode, const USkeletalMeshComponent* SkeletalMeshComponent);
+	FGLTFJsonAnimation* GetOrAddAnimation(FGLTFJsonNode* RootNode, const USkeletalMesh* SkeletalMesh, const UAnimSequence* AnimSequence);
+	FGLTFJsonAnimation* GetOrAddAnimation(FGLTFJsonNode* RootNode, const USkeletalMeshComponent* SkeletalMeshComponent);
+	FGLTFJsonAnimation* GetOrAddAnimation(const ULevel* Level, const ULevelSequence* LevelSequence);
+	FGLTFJsonAnimation* GetOrAddAnimation(const ALevelSequenceActor* LevelSequenceActor);
+
+	FGLTFJsonNode* GetOrAddNode(const AActor* Actor);
+	FGLTFJsonNode* GetOrAddNode(const USceneComponent* SceneComponent);
+	FGLTFJsonNode* GetOrAddNode(const USceneComponent* SceneComponent, FName SocketName);
+	FGLTFJsonNode* GetOrAddNode(FGLTFJsonNode* RootNode, const UStaticMesh* StaticMesh, FName SocketName);
+	FGLTFJsonNode* GetOrAddNode(FGLTFJsonNode* RootNode, const USkeletalMesh* SkeletalMesh, FName SocketName);
+	FGLTFJsonNode* GetOrAddNode(FGLTFJsonNode* RootNode, const USkeletalMesh* SkeletalMesh, int32 BoneIndex);
+	FGLTFJsonScene* GetOrAddScene(const UWorld* World);
+
+	FGLTFJsonCamera* GetOrAddCamera(const UCameraComponent* CameraComponent);
+	FGLTFJsonLight* GetOrAddLight(const ULightComponent* LightComponent);
+	FGLTFJsonBackdrop* GetOrAddBackdrop(const AActor* BackdropActor);
+	FGLTFJsonLightMap* GetOrAddLightMap(const UStaticMeshComponent* StaticMeshComponent);
+	FGLTFJsonHotspot* GetOrAddHotspot(const AGLTFHotspotActor* HotspotActor);
+	FGLTFJsonSkySphere* GetOrAddSkySphere(const AActor* SkySphereActor);
+
+	FGLTFJsonEpicLevelVariantSets* GetOrAddEpicLevelVariantSets(const ULevelVariantSets* LevelVariantSets);
+	FGLTFJsonKhrMaterialVariant* GetOrAddKhrMaterialVariant(const UVariant* Variant);
 
 	void RegisterObjectVariant(const UObject* Object, const UPropertyValue* Property);
 	const TArray<const UPropertyValue*>* GetObjectVariants(const UObject* Object) const;
diff --git a/Engine/Plugins/Enterprise/GLTFExporter/MigratedFromGit/Source/GLTFExporter/Private/Builders/GLTFJsonBuilder.cpp b/Engine/Plugins/Enterprise/GLTFExporter/MigratedFromGit/Source/GLTFExporter/Private/Builders/GLTFJsonBuilder.cpp
--- a/Engine/Plugins/Enterprise/GLTFExporter/MigratedFromGit/Source/GLTFExporter/Private/Builders/GLTFJsonBuilder.cpp
+++ b/Engine/Plugins/Enterprise/GLTFExporter/MigratedFromGit/Source/GLTFExporter/Private/Builders/GLTFJsonBuilder.cpp
@@ -141,152 +141,147 @@ FGLTFJsonKhrMaterialVariant* FGLTFJsonBuilder::AddKhrMaterialVariant()
 	return JsonRoot.KhrMaterialVariants.Add();
 }
 
-FGLTFJsonNode* FGLTFJsonBuilder::AddChildNode(FGLTFJsonNodeIndex ParentIndex)
+FGLTFJsonNode* FGLTFJsonBuilder::AddChildNode(FGLTFJsonNode* ParentNode)
 {
-	FGLTFJsonNode* Child = AddNode();
+	FGLTFJsonNode* ChildNode = AddNode();
 
-	if (ParentIndex != INDEX_NONE)
+	if (ParentNode != nullptr)
 	{
-		GetNode(ParentIndex).Children.Add(Child->Index);
+		ParentNode->Children.Add(ChildNode);
 	}
 
-	return Child;
+	return ChildNode;
 }
 
-FGLTFJsonNode* FGLTFJsonBuilder::AddChildComponentNode(FGLTFJsonNodeIndex ParentIndex)
+FGLTFJsonNode* FGLTFJsonBuilder::AddChildComponentNode(FGLTFJsonNode* ParentNode)
 {
-	FGLTFJsonNode* Child = AddChildNode(ParentIndex);
+	FGLTFJsonNode* ChildNode = AddChildNode(ParentNode);
 
-	if (ParentIndex != INDEX_NONE)
+	if (ParentNode != nullptr)
 	{
-		GetNode(ParentIndex).ComponentNode = Child->Index;
+		ParentNode->ComponentNode = ChildNode;
 	}
 
-	return Child;
+	return ChildNode;
 }
 
-FGLTFJsonAccessorIndex FGLTFJsonBuilder::AddAccessor(const FGLTFJsonAccessor& JsonAccessor)
+FGLTFJsonAccessor* FGLTFJsonBuilder::AddAccessor(const FGLTFJsonAccessor& JsonAccessor)
 {
-	return FGLTFJsonAccessorIndex(JsonRoot.Accessors.Add(JsonAccessor)->Index);
+	return JsonRoot.Accessors.Add(JsonAccessor);
 }
 
-FGLTFJsonAnimationIndex FGLTFJsonBuilder::AddAnimation(const FGLTFJsonAnimation& JsonAnimation)
+FGLTFJsonBuffer* FGLTFJsonBuilder::AddBuffer(const FGLTFJsonBuffer& JsonBuffer)
 {
-	return FGLTFJsonAnimationIndex(JsonRoot.Animations.Add(JsonAnimation)->Index);
+	return JsonRoot.Buffers.Add(JsonBuffer);
 }
 
-FGLTFJsonBufferIndex FGLTFJsonBuilder::AddBuffer(const FGLTFJsonBuffer& JsonBuffer)
+FGLTFJsonBufferView* FGLTFJsonBuilder::AddBufferView(const FGLTFJsonBufferView& JsonBufferView)
 {
-	return FGLTFJsonBufferIndex(JsonRoot.Buffers.Add(JsonBuffer)->Index);
+	return JsonRoot.BufferViews.Add(JsonBufferView);
 }
 
-FGLTFJsonBufferViewIndex FGLTFJsonBuilder::AddBufferView(const FGLTFJsonBufferView& JsonBufferView)
+FGLTFJsonCamera* FGLTFJsonBuilder::AddCamera(const FGLTFJsonCamera& JsonCamera)
 {
-	return FGLTFJsonBufferViewIndex(JsonRoot.BufferViews.Add(JsonBufferView)->Index);
+	return JsonRoot.Cameras.Add(JsonCamera);
 }
 
-FGLTFJsonCameraIndex FGLTFJsonBuilder::AddCamera(const FGLTFJsonCamera& JsonCamera)
+FGLTFJsonImage* FGLTFJsonBuilder::AddImage(const FGLTFJsonImage& JsonImage)
 {
-	return FGLTFJsonCameraIndex(JsonRoot.Cameras.Add(JsonCamera)->Index);
+	return JsonRoot.Images.Add(JsonImage);
 }
 
-FGLTFJsonImageIndex FGLTFJsonBuilder::AddImage(const FGLTFJsonImage& JsonImage)
+FGLTFJsonMaterial* FGLTFJsonBuilder::AddMaterial(const FGLTFJsonMaterial& JsonMaterial)
 {
-	return FGLTFJsonImageIndex(JsonRoot.Images.Add(JsonImage)->Index);
+	return JsonRoot.Materials.Add(JsonMaterial);
 }
 
-FGLTFJsonMaterialIndex FGLTFJsonBuilder::AddMaterial(const FGLTFJsonMaterial& JsonMaterial)
+FGLTFJsonMesh* FGLTFJsonBuilder::AddMesh(const FGLTFJsonMesh& JsonMesh)
 {
-	return FGLTFJsonMaterialIndex(JsonRoot.Materials.Add(JsonMaterial)->Index);
+	return JsonRoot.Meshes.Add(JsonMesh);
 }
 
-FGLTFJsonMeshIndex FGLTFJsonBuilder::AddMesh(const FGLTFJsonMesh& JsonMesh)
+FGLTFJsonNode* FGLTFJsonBuilder::AddNode(const FGLTFJsonNode& JsonNode)
 {
-	return FGLTFJsonMeshIndex(JsonRoot.Meshes.Add(JsonMesh)->Index);
+	return JsonRoot.Nodes.Add(JsonNode);
 }
 
-FGLTFJsonNodeIndex FGLTFJsonBuilder::AddNode(const FGLTFJsonNode& JsonNode)
+FGLTFJsonSampler* FGLTFJsonBuilder::AddSampler(const FGLTFJsonSampler& JsonSampler)
 {
-	return FGLTFJsonNodeIndex(JsonRoot.Nodes.Add(JsonNode)->Index);
+	return JsonRoot.Samplers.Add(JsonSampler);
 }
 
-FGLTFJsonSamplerIndex FGLTFJsonBuilder::AddSampler(const FGLTFJsonSampler& JsonSampler)
+FGLTFJsonScene* FGLTFJsonBuilder::AddScene(const FGLTFJsonScene& JsonScene)
 {
-	return FGLTFJsonSamplerIndex(JsonRoot.Samplers.Add(JsonSampler)->Index);
+	return JsonRoot.Scenes.Add(JsonScene);
 }
 
-FGLTFJsonSceneIndex FGLTFJsonBuilder::AddScene(const FGLTFJsonScene& JsonScene)
+FGLTFJsonSkin* FGLTFJsonBuilder::AddSkin(const FGLTFJsonSkin& JsonSkin)
 {
-	return FGLTFJsonSceneIndex(JsonRoot.Scenes.Add(JsonScene)->Index);
+	return JsonRoot.Skins.Add(JsonSkin);
 }
 
-FGLTFJsonSkinIndex FGLTFJsonBuilder::AddSkin(const FGLTFJsonSkin& JsonSkin)
+FGLTFJsonTexture* FGLTFJsonBuilder::AddTexture(const FGLTFJsonTexture& JsonTexture)
 {
-	return FGLTFJsonSkinIndex(JsonRoot.Skins.Add(JsonSkin)->Index);
+	return JsonRoot.Textures.Add(JsonTexture);
 }
 
-FGLTFJsonTextureIndex FGLTFJsonBuilder::AddTexture(const FGLTFJsonTexture& JsonTexture)
+FGLTFJsonBackdrop* FGLTFJsonBuilder::AddBackdrop(const FGLTFJsonBackdrop& JsonBackdrop)
 {
-	return FGLTFJsonTextureIndex(JsonRoot.Textures.Add(JsonTexture)->Index);
+	return JsonRoot.Backdrops.Add(JsonBackdrop);
 }
 
-FGLTFJsonBackdropIndex FGLTFJsonBuilder::AddBackdrop(const FGLTFJsonBackdrop& JsonBackdrop)
+FGLTFJsonHotspot* FGLTFJsonBuilder::AddHotspot(const FGLTFJsonHotspot& JsonHotspot)
 {
-	return FGLTFJsonBackdropIndex(JsonRoot.Backdrops.Add(JsonBackdrop)->Index);
+	return JsonRoot.Hotspots.Add(JsonHotspot);
 }
 
-FGLTFJsonHotspotIndex FGLTFJsonBuilder::AddHotspot(const FGLTFJsonHotspot& JsonHotspot)
+FGLTFJsonLight* FGLTFJsonBuilder::AddLight(const FGLTFJsonLight& JsonLight)
 {
-	return FGLTFJsonHotspotIndex(JsonRoot.Hotspots.Add(JsonHotspot)->Index);
+	return JsonRoot.Lights.Add(JsonLight);
 }
 
-FGLTFJsonLightIndex FGLTFJsonBuilder::AddLight(const FGLTFJsonLight& JsonLight)
+FGLTFJsonLightMap* FGLTFJsonBuilder::AddLightMap(const FGLTFJsonLightMap& JsonLightMap)
 {
-	return FGLTFJsonLightIndex(JsonRoot.Lights.Add(JsonLight)->Index);
+	return JsonRoot.LightMaps.Add(JsonLightMap);
 }
 
-FGLTFJsonLightMapIndex FGLTFJsonBuilder::AddLightMap(const FGLTFJsonLightMap& JsonLightMap)
+FGLTFJsonSkySphere* FGLTFJsonBuilder::AddSkySphere(const FGLTFJsonSkySphere& JsonSkySphere)
 {
-	return FGLTFJsonLightMapIndex(JsonRoot.LightMaps.Add(JsonLightMap)->Index);
+	return JsonRoot.SkySpheres.Add(JsonSkySphere);
 }
 
-FGLTFJsonSkySphereIndex FGLTFJsonBuilder::AddSkySphere(const FGLTFJsonSkySphere& JsonSkySphere)
+FGLTFJsonEpicLevelVariantSets* FGLTFJsonBuilder::AddEpicLevelVariantSets(const FGLTFJsonEpicLevelVariantSets& JsonEpicLevelVariantSets)
 {
-	return FGLTFJsonSkySphereIndex(JsonRoot.SkySpheres.Add(JsonSkySphere)->Index);
+	return JsonRoot.EpicLevelVariantSets.Add(JsonEpicLevelVariantSets);
 }
 
-FGLTFJsonEpicLevelVariantSetsIndex FGLTFJsonBuilder::AddEpicLevelVariantSets(const FGLTFJsonEpicLevelVariantSets& JsonEpicLevelVariantSets)
+FGLTFJsonKhrMaterialVariant* FGLTFJsonBuilder::AddKhrMaterialVariant(const FGLTFJsonKhrMaterialVariant& JsonKhrMaterialVariant)
 {
-	return FGLTFJsonEpicLevelVariantSetsIndex(JsonRoot.EpicLevelVariantSets.Add(JsonEpicLevelVariantSets)->Index);
+	return JsonRoot.KhrMaterialVariants.Add(JsonKhrMaterialVariant);
 }
 
-FGLTFJsonKhrMaterialVariantIndex FGLTFJsonBuilder::AddKhrMaterialVariant(const FGLTFJsonKhrMaterialVariant& JsonKhrMaterialVariant)
+FGLTFJsonNode* FGLTFJsonBuilder::AddChildNode(FGLTFJsonNode* Parent, const FGLTFJsonNode& JsonNode)
 {
-	return FGLTFJsonKhrMaterialVariantIndex(JsonRoot.KhrMaterialVariants.Add(JsonKhrMaterialVariant)->Index);
-}
-
-FGLTFJsonNodeIndex FGLTFJsonBuilder::AddChildNode(FGLTFJsonNodeIndex ParentIndex, const FGLTFJsonNode& JsonNode)
-{
-	const FGLTFJsonNodeIndex ChildIndex = AddNode(JsonNode);
+	FGLTFJsonNode* ChildNode = AddNode(JsonNode);
 
-	if (ParentIndex != INDEX_NONE)
+	if (Parent != nullptr)
 	{
-		GetNode(ParentIndex).Children.Add(ChildIndex);
+		Parent->Children.Add(ChildNode);
 	}
 
-	return ChildIndex;
+	return ChildNode;
 }
 
-FGLTFJsonNodeIndex FGLTFJsonBuilder::AddChildComponentNode(FGLTFJsonNodeIndex ParentIndex, const FGLTFJsonNode& JsonNode)
+FGLTFJsonNode* FGLTFJsonBuilder::AddChildComponentNode(FGLTFJsonNode* Parent, const FGLTFJsonNode& JsonNode)
 {
-	const FGLTFJsonNodeIndex ChildIndex = AddChildNode(ParentIndex, JsonNode);
+	FGLTFJsonNode* ChildNode = AddChildNode(Parent, JsonNode);
 
-	if (ParentIndex != INDEX_NONE)
+	if (Parent != nullptr)
 	{
-		GetNode(ParentIndex).ComponentNode = ChildIndex;
+		Parent->ComponentNode = ChildNode;
 	}
 
-	return ChildIndex;
+	return ChildNode;
 }
 
 const FGLTFJsonRoot& FGLTFJsonBuilder::GetRoot() const
@@ -294,115 +289,114 @@ const FGLTFJsonRoot& FGLTFJsonBuilder::GetRoot() const
 	return JsonRoot;
 }
 
-FGLTFJsonAccessor& FGLTFJsonBuilder::GetAccessor(FGLTFJsonAccessorIndex AccessorIndex)
+FGLTFJsonAccessor& FGLTFJsonBuilder::GetAccessor(FGLTFJsonAccessor* AccessorIndex)
 {
-	return *JsonRoot.Accessors[AccessorIndex];
+	return *AccessorIndex;
 }
 
-FGLTFJsonAnimation& FGLTFJsonBuilder::GetAnimation(FGLTFJsonAnimationIndex AnimationIndex)
+FGLTFJsonAnimation& FGLTFJsonBuilder::GetAnimation(FGLTFJsonAnimation* AnimationIndex)
 {
-	return *JsonRoot.Animations[AnimationIndex];
+	return *AnimationIndex;
 }
 
-FGLTFJsonBuffer& FGLTFJsonBuilder::GetBuffer(FGLTFJsonBufferIndex BufferIndex)
+FGLTFJsonBuffer& FGLTFJsonBuilder::GetBuffer(FGLTFJsonBuffer* BufferIndex)
 {
-	return *JsonRoot.Buffers[BufferIndex];
+	return *BufferIndex;
 }
 
-FGLTFJsonBufferView& FGLTFJsonBuilder::GetBufferView(FGLTFJsonBufferViewIndex BufferViewIndex)
+FGLTFJsonBufferView& FGLTFJsonBuilder::GetBufferView(FGLTFJsonBufferView* BufferViewIndex)
 {
-	return *JsonRoot.BufferViews[BufferViewIndex];
+	return *BufferViewIndex;
 }
 
-FGLTFJsonCamera& FGLTFJsonBuilder::GetCamera(FGLTFJsonCameraIndex CameraIndex)
+FGLTFJsonCamera& FGLTFJsonBuilder::GetCamera(FGLTFJsonCamera* CameraIndex)
 {
-	return *JsonRoot.Cameras[CameraIndex];
+	return *CameraIndex;
 }
 
-FGLTFJsonImage& FGLTFJsonBuilder::GetImage(FGLTFJsonImageIndex ImageIndex)
+FGLTFJsonImage& FGLTFJsonBuilder::GetImage(FGLTFJsonImage* ImageIndex)
 {
-	return *JsonRoot.Images[ImageIndex];
+	return *ImageIndex;
 }
 
-FGLTFJsonMaterial& FGLTFJsonBuilder::GetMaterial(FGLTFJsonMaterialIndex MaterialIndex)
+FGLTFJsonMaterial& FGLTFJsonBuilder::GetMaterial(FGLTFJsonMaterial* MaterialIndex)
 {
-	return *JsonRoot.Materials[MaterialIndex];
+	return *MaterialIndex;
 }
 
-FGLTFJsonMesh& FGLTFJsonBuilder::GetMesh(FGLTFJsonMeshIndex MeshIndex)
+FGLTFJsonMesh& FGLTFJsonBuilder::GetMesh(FGLTFJsonMesh* MeshIndex)
 {
-	return *JsonRoot.Meshes[MeshIndex];
+	return *MeshIndex;
 }
 
-FGLTFJsonNode& FGLTFJsonBuilder::GetNode(FGLTFJsonNodeIndex NodeIndex)
+FGLTFJsonNode& FGLTFJsonBuilder::GetNode(FGLTFJsonNode* NodeIndex)
 {
-	return *JsonRoot.Nodes[NodeIndex];
+	return *NodeIndex;
 }
 
-FGLTFJsonSampler& FGLTFJsonBuilder::GetSampler(FGLTFJsonSamplerIndex SamplerIndex)
+FGLTFJsonSampler& FGLTFJsonBuilder::GetSampler(FGLTFJsonSampler* SamplerIndex)
 {
-	return *JsonRoot.Samplers[SamplerIndex];
+	return *SamplerIndex;
 }
 
-FGLTFJsonScene& FGLTFJsonBuilder::GetScene(FGLTFJsonSceneIndex SceneIndex)
+FGLTFJsonScene& FGLTFJsonBuilder::GetScene(FGLTFJsonScene* SceneIndex)
 {
-	return *JsonRoot.Scenes[SceneIndex];
+	return *SceneIndex;
 }
 
-FGLTFJsonSkin& FGLTFJsonBuilder::GetSkin(FGLTFJsonSkinIndex SkinIndex)
+FGLTFJsonSkin& FGLTFJsonBuilder::GetSkin(FGLTFJsonSkin* SkinIndex)
 {
-	return *JsonRoot.Skins[SkinIndex];
+	return *SkinIndex;
 }
 
-FGLTFJsonTexture& FGLTFJsonBuilder::GetTexture(FGLTFJsonTextureIndex TextureIndex)
+FGLTFJsonTexture& FGLTFJsonBuilder::GetTexture(FGLTFJsonTexture* TextureIndex)
 {
-	return *JsonRoot.Textures[TextureIndex];
+	return *TextureIndex;
 }
 
-FGLTFJsonBackdrop& FGLTFJsonBuilder::GetBackdrop(FGLTFJsonBackdropIndex BackdropIndex)
+FGLTFJsonBackdrop& FGLTFJsonBuilder::GetBackdrop(FGLTFJsonBackdrop* BackdropIndex)
 {
-	return *JsonRoot.Backdrops[BackdropIndex];
+	return *BackdropIndex;
 }
 
-FGLTFJsonHotspot& FGLTFJsonBuilder::GetHotspot(FGLTFJsonHotspotIndex HotspotIndex)
+FGLTFJsonHotspot& FGLTFJsonBuilder::GetHotspot(FGLTFJsonHotspot* HotspotIndex)
 {
-	return *JsonRoot.Hotspots[HotspotIndex];
+	return *HotspotIndex;
 }
 
-FGLTFJsonLight& FGLTFJsonBuilder::GetLight(FGLTFJsonLightIndex LightIndex)
+FGLTFJsonLight& FGLTFJsonBuilder::GetLight(FGLTFJsonLight* LightIndex)
 {
-	return *JsonRoot.Lights[LightIndex];
+	return *LightIndex;
 }
 
-FGLTFJsonLightMap& FGLTFJsonBuilder::GetLightMap(FGLTFJsonLightMapIndex LightMapIndex)
+FGLTFJsonLightMap& FGLTFJsonBuilder::GetLightMap(FGLTFJsonLightMap* LightMapIndex)
 {
-	return *JsonRoot.LightMaps[LightMapIndex];
+	return *LightMapIndex;
 }
 
-FGLTFJsonSkySphere& FGLTFJsonBuilder::GetSkySphere(FGLTFJsonSkySphereIndex SkySphereIndex)
+FGLTFJsonSkySphere& FGLTFJsonBuilder::GetSkySphere(FGLTFJsonSkySphere* SkySphereIndex)
 {
-	return *JsonRoot.SkySpheres[SkySphereIndex];
+	return *SkySphereIndex;
 }
 
-FGLTFJsonEpicLevelVariantSets& FGLTFJsonBuilder::GetEpicLevelVariantSets(FGLTFJsonEpicLevelVariantSetsIndex EpicLevelVariantSetsIndex)
+FGLTFJsonEpicLevelVariantSets& FGLTFJsonBuilder::GetEpicLevelVariantSets(FGLTFJsonEpicLevelVariantSets* EpicLevelVariantSetsIndex)
 {
-	return *JsonRoot.EpicLevelVariantSets[EpicLevelVariantSetsIndex];
+	return *EpicLevelVariantSetsIndex;
 }
 
-FGLTFJsonKhrMaterialVariant& FGLTFJsonBuilder::GetKhrMaterialVariant(FGLTFJsonKhrMaterialVariantIndex KhrMaterialVariantIndex)
+FGLTFJsonKhrMaterialVariant& FGLTFJsonBuilder::GetKhrMaterialVariant(FGLTFJsonKhrMaterialVariant* KhrMaterialVariantIndex)
 {
-	return *JsonRoot.KhrMaterialVariants[KhrMaterialVariantIndex];
+	return *KhrMaterialVariantIndex;
 }
 
-FGLTFJsonNodeIndex FGLTFJsonBuilder::GetComponentNodeIndex(FGLTFJsonNodeIndex NodeIndex)
+FGLTFJsonNode* FGLTFJsonBuilder::GetComponentNode(FGLTFJsonNode* Node)
 {
-	if (NodeIndex == INDEX_NONE)
+	if (Node == nullptr)
 	{
-		return FGLTFJsonNodeIndex(INDEX_NONE);
+		return nullptr;
 	}
 
-	const FGLTFJsonNode& Node = GetNode(NodeIndex);
-	return Node.ComponentNode != INDEX_NONE ? Node.ComponentNode : NodeIndex;
+	return Node->ComponentNode != nullptr ? Node->ComponentNode : Node;
 }
 
 FString FGLTFJsonBuilder::GetGeneratorString() const
diff --git a/Engine/Plugins/Enterprise/GLTFExporter/MigratedFromGit/Source/GLTFExporter/Private/Builders/GLTFJsonBuilder.h b/Engine/Plugins/Enterprise/GLTFExporter/MigratedFromGit/Source/GLTFExporter/Private/Builders/GLTFJsonBuilder.h
--- a/Engine/Plugins/Enterprise/GLTFExporter/MigratedFromGit/Source/GLTFExporter/Private/Builders/GLTFJsonBuilder.h
+++ b/Engine/Plugins/Enterprise/GLTFExporter/MigratedFromGit/Source/GLTFExporter/Private/Builders/GLTFJsonBuilder.h
@@ -9,7 +9,7 @@ class FGLTFJsonBuilder : public FGLTFTaskBuilder
 {
 public:
 
-	FGLTFJsonSceneIndex& DefaultScene;
+	FGLTFJsonScene*& DefaultScene;
 
 	FGLTFJsonBuilder(const FString& FilePath, const UGLTFExportOptions* ExportOptions);
 
@@ -40,56 +40,55 @@ public:
 	FGLTFJsonEpicLevelVariantSets* AddEpicLevelVariantSets();
 	FGLTFJsonKhrMaterialVariant* AddKhrMaterialVariant();
 
-	FGLTFJsonNode* AddChildNode(FGLTFJsonNodeIndex ParentNodeIndex);
-	FGLTFJsonNode* AddChildComponentNode(FGLTFJsonNodeIndex ParentNodeIndex);
-
-	FGLTFJsonAccessorIndex AddAccessor(const FGLTFJsonAccessor& JsonAccessor);
-	FGLTFJsonAnimationIndex AddAnimation(const FGLTFJsonAnimation& JsonAnimation);
-	FGLTFJsonBufferIndex AddBuffer(const FGLTFJsonBuffer& JsonBuffer);
-	FGLTFJsonBufferViewIndex AddBufferView(const FGLTFJsonBufferView& JsonBufferView);
-	FGLTFJsonCameraIndex AddCamera(const FGLTFJsonCamera& JsonCamera);
-	FGLTFJsonImageIndex AddImage(const FGLTFJsonImage& JsonImage);
-	FGLTFJsonMaterialIndex AddMaterial(const FGLTFJsonMaterial& JsonMaterial);
-	FGLTFJsonMeshIndex AddMesh(const FGLTFJsonMesh& JsonMesh);
-	FGLTFJsonNodeIndex AddNode(const FGLTFJsonNode& JsonNode);
-	FGLTFJsonSamplerIndex AddSampler(const FGLTFJsonSampler& JsonSampler);
-	FGLTFJsonSceneIndex AddScene(const FGLTFJsonScene& JsonScene);
-	FGLTFJsonSkinIndex AddSkin(const FGLTFJsonSkin& JsonSkin);
-	FGLTFJsonTextureIndex AddTexture(const FGLTFJsonTexture& JsonTexture);
-	FGLTFJsonBackdropIndex AddBackdrop(const FGLTFJsonBackdrop& JsonBackdrop);
-	FGLTFJsonHotspotIndex AddHotspot(const FGLTFJsonHotspot& JsonHotspot);
-	FGLTFJsonLightIndex AddLight(const FGLTFJsonLight& JsonLight);
-	FGLTFJsonLightMapIndex AddLightMap(const FGLTFJsonLightMap& JsonLightMap);
-	FGLTFJsonSkySphereIndex AddSkySphere(const FGLTFJsonSkySphere& JsonSkySphere);
-	FGLTFJsonEpicLevelVariantSetsIndex AddEpicLevelVariantSets(const FGLTFJsonEpicLevelVariantSets& JsonEpicLevelVariantSets);
-	FGLTFJsonKhrMaterialVariantIndex AddKhrMaterialVariant(const FGLTFJsonKhrMaterialVariant& JsonKhrMaterialVariant);
-
-	FGLTFJsonNodeIndex AddChildNode(FGLTFJsonNodeIndex ParentNodeIndex, const FGLTFJsonNode& JsonNode);
-	FGLTFJsonNodeIndex AddChildComponentNode(FGLTFJsonNodeIndex ParentNodeIndex, const FGLTFJsonNode& JsonNode);
+	FGLTFJsonNode* AddChildNode(FGLTFJsonNode* ParentNode);
+	FGLTFJsonNode* AddChildComponentNode(FGLTFJsonNode* ParentNode);
+
+	FGLTFJsonAccessor* AddAccessor(const FGLTFJsonAccessor& JsonAccessor);
+	FGLTFJsonBuffer* AddBuffer(const FGLTFJsonBuffer& JsonBuffer);
+	FGLTFJsonBufferView* AddBufferView(const FGLTFJsonBufferView& JsonBufferView);
+	FGLTFJsonCamera* AddCamera(const FGLTFJsonCamera& JsonCamera);
+	FGLTFJsonImage* AddImage(const FGLTFJsonImage& JsonImage);
+	FGLTFJsonMaterial* AddMaterial(const FGLTFJsonMaterial& JsonMaterial);
+	FGLTFJsonMesh* AddMesh(const FGLTFJsonMesh& JsonMesh);
+	FGLTFJsonNode* AddNode(const FGLTFJsonNode& JsonNode);
+	FGLTFJsonSampler* AddSampler(const FGLTFJsonSampler& JsonSampler);
+	FGLTFJsonScene* AddScene(const FGLTFJsonScene& JsonScene);
+	FGLTFJsonSkin* AddSkin(const FGLTFJsonSkin& JsonSkin);
+	FGLTFJsonTexture* AddTexture(const FGLTFJsonTexture& JsonTexture);
+	FGLTFJsonBackdrop* AddBackdrop(const FGLTFJsonBackdrop& JsonBackdrop);
+	FGLTFJsonHotspot* AddHotspot(const FGLTFJsonHotspot& JsonHotspot);
+	FGLTFJsonLight* AddLight(const FGLTFJsonLight& JsonLight);
+	FGLTFJsonLightMap* AddLightMap(const FGLTFJsonLightMap& JsonLightMap);
+	FGLTFJsonSkySphere* AddSkySphere(const FGLTFJsonSkySphere& JsonSkySphere);
+	FGLTFJsonEpicLevelVariantSets* AddEpicLevelVariantSets(const FGLTFJsonEpicLevelVariantSets& JsonEpicLevelVariantSets);
+	FGLTFJsonKhrMaterialVariant* AddKhrMaterialVariant(const FGLTFJsonKhrMaterialVariant& JsonKhrMaterialVariant);
+
+	FGLTFJsonNode* AddChildNode(FGLTFJsonNode* ParentNode, const FGLTFJsonNode& JsonNode);
+	FGLTFJsonNode* AddChildComponentNode(FGLTFJsonNode* ParentNode, const FGLTFJsonNode& JsonNode);
 
 	const FGLTFJsonRoot& GetRoot() const;
-	FGLTFJsonAccessor& GetAccessor(FGLTFJsonAccessorIndex AccessorIndex);
-	FGLTFJsonAnimation& GetAnimation(FGLTFJsonAnimationIndex AnimationIndex);
-	FGLTFJsonBuffer& GetBuffer(FGLTFJsonBufferIndex BufferIndex);
-	FGLTFJsonBufferView& GetBufferView(FGLTFJsonBufferViewIndex BufferViewIndex);
-	FGLTFJsonCamera& GetCamera(FGLTFJsonCameraIndex CameraIndex);
-	FGLTFJsonImage& GetImage(FGLTFJsonImageIndex ImageIndex);
-	FGLTFJsonMaterial& GetMaterial(FGLTFJsonMaterialIndex MaterialIndex);
-	FGLTFJsonMesh& GetMesh(FGLTFJsonMeshIndex MeshIndex);
-	FGLTFJsonNode& GetNode(FGLTFJsonNodeIndex NodeIndex);
-	FGLTFJsonSampler& GetSampler(FGLTFJsonSamplerIndex SamplerIndex);
-	FGLTFJsonScene& GetScene(FGLTFJsonSceneIndex SceneIndex);
-	FGLTFJsonSkin& GetSkin(FGLTFJsonSkinIndex SkinIndex);
-	FGLTFJsonTexture& GetTexture(FGLTFJsonTextureIndex TextureIndex);
-	FGLTFJsonBackdrop& GetBackdrop(FGLTFJsonBackdropIndex BackdropIndex);
-	FGLTFJsonHotspot& GetHotspot(FGLTFJsonHotspotIndex HotspotIndex);
-	FGLTFJsonLight& GetLight(FGLTFJsonLightIndex LightIndex);
-	FGLTFJsonLightMap& GetLightMap(FGLTFJsonLightMapIndex LightMapIndex);
-	FGLTFJsonSkySphere& GetSkySphere(FGLTFJsonSkySphereIndex SkySphereIndex);
-	FGLTFJsonEpicLevelVariantSets& GetEpicLevelVariantSets(FGLTFJsonEpicLevelVariantSetsIndex EpicLevelVariantSetsIndex);
-	FGLTFJsonKhrMaterialVariant& GetKhrMaterialVariant(FGLTFJsonKhrMaterialVariantIndex KhrMaterialVariantIndex);
-
-	FGLTFJsonNodeIndex GetComponentNodeIndex(FGLTFJsonNodeIndex NodeIndex);
+	FGLTFJsonAccessor& GetAccessor(FGLTFJsonAccessor* AccessorIndex);
+	FGLTFJsonAnimation& GetAnimation(FGLTFJsonAnimation* AnimationIndex);
+	FGLTFJsonBuffer& GetBuffer(FGLTFJsonBuffer* BufferIndex);
+	FGLTFJsonBufferView& GetBufferView(FGLTFJsonBufferView* BufferViewIndex);
+	FGLTFJsonCamera& GetCamera(FGLTFJsonCamera* CameraIndex);
+	FGLTFJsonImage& GetImage(FGLTFJsonImage* ImageIndex);
+	FGLTFJsonMaterial& GetMaterial(FGLTFJsonMaterial* MaterialIndex);
+	FGLTFJsonMesh& GetMesh(FGLTFJsonMesh* MeshIndex);
+	FGLTFJsonNode& GetNode(FGLTFJsonNode* NodeIndex);
+	FGLTFJsonSampler& GetSampler(FGLTFJsonSampler* SamplerIndex);
+	FGLTFJsonScene& GetScene(FGLTFJsonScene* SceneIndex);
+	FGLTFJsonSkin& GetSkin(FGLTFJsonSkin* SkinIndex);
+	FGLTFJsonTexture& GetTexture(FGLTFJsonTexture* TextureIndex);
+	FGLTFJsonBackdrop& GetBackdrop(FGLTFJsonBackdrop* BackdropIndex);
+	FGLTFJsonHotspot& GetHotspot(FGLTFJsonHotspot* HotspotIndex);
+	FGLTFJsonLight& GetLight(FGLTFJsonLight* LightIndex);
+	FGLTFJsonLightMap& GetLightMap(FGLTFJsonLightMap* LightMapIndex);
+	FGLTFJsonSkySphere& GetSkySphere(FGLTFJsonSkySphere* SkySphereIndex);
+	FGLTFJsonEpicLevelVariantSets& GetEpicLevelVariantSets(FGLTFJsonEpicLevelVariantSets* EpicLevelVariantSetsIndex);
+	FGLTFJsonKhrMaterialVariant& GetKhrMaterialVariant(FGLTFJsonKhrMaterialVariant* KhrMaterialVariantIndex);
+
+	FGLTFJsonNode* GetComponentNode(FGLTFJsonNode* Node);
 
 private:
 
diff --git a/Engine/Plugins/Enterprise/GLTFExporter/MigratedFromGit/Source/GLTFExporter/Private/Converters/GLTFAccessorConverters.cpp b/Engine/Plugins/Enterprise/GLTFExporter/MigratedFromGit/Source/GLTFExporter/Private/Converters/GLTFAccessorConverters.cpp
--- a/Engine/Plugins/Enterprise/GLTFExporter/MigratedFromGit/Source/GLTFExporter/Private/Converters/GLTFAccessorConverters.cpp
+++ b/Engine/Plugins/Enterprise/GLTFExporter/MigratedFromGit/Source/GLTFExporter/Private/Converters/GLTFAccessorConverters.cpp
@@ -20,11 +20,11 @@ public:
 	typedef TypeIfTrue Type;
 };
 
-FGLTFJsonAccessorIndex FGLTFPositionBufferConverter::Convert(const FGLTFMeshSection* MeshSection, const FPositionVertexBuffer* VertexBuffer)
+FGLTFJsonAccessor* FGLTFPositionBufferConverter::Convert(const FGLTFMeshSection* MeshSection, const FPositionVertexBuffer* VertexBuffer)
 {
 	if (VertexBuffer == nullptr || VertexBuffer->GetNumVertices() == 0)
 	{
-		return FGLTFJsonAccessorIndex(INDEX_NONE);
+		return nullptr;
 	}
 
 	const TUniquePtr<IGLTFBufferAdapter> SourceBuffer = IGLTFBufferAdapter::GetPositions(VertexBuffer);
@@ -33,7 +33,7 @@ FGLTFJsonAccessorIndex FGLTFPositionBufferConverter::Convert(const FGLTFMeshSect
 	if (SourceData == nullptr)
 	{
 		// TODO: report error
-		return FGLTFJsonAccessorIndex(INDEX_NONE);
+		return nullptr;
 	}
 
 	const TArray<uint32>& IndexMap = MeshSection->IndexMap;
@@ -81,11 +81,11 @@ FGLTFJsonAccessorIndex FGLTFPositionBufferConverter::Convert(const FGLTFMeshSect
 	return Builder.AddAccessor(JsonAccessor);
 }
 
-FGLTFJsonAccessorIndex FGLTFColorBufferConverter::Convert(const FGLTFMeshSection* MeshSection, const FColorVertexBuffer* VertexBuffer)
+FGLTFJsonAccessor* FGLTFColorBufferConverter::Convert(const FGLTFMeshSection* MeshSection, const FColorVertexBuffer* VertexBuffer)
 {
 	if (VertexBuffer == nullptr || VertexBuffer->GetNumVertices() == 0)
 	{
-		return FGLTFJsonAccessorIndex(INDEX_NONE);
+		return nullptr;
 	}
 
 	const TUniquePtr<IGLTFBufferAdapter> SourceBuffer = IGLTFBufferAdapter::GetColors(VertexBuffer);
@@ -94,7 +94,7 @@ FGLTFJsonAccessorIndex FGLTFColorBufferConverter::Convert(const FGLTFMeshSection
 	if (SourceData == nullptr)
 	{
 		// TODO: report error
-		return FGLTFJsonAccessorIndex(INDEX_NONE);
+		return nullptr;
 	}
 
 	const TArray<uint32>& IndexMap = MeshSection->IndexMap;
@@ -121,11 +121,11 @@ FGLTFJsonAccessorIndex FGLTFColorBufferConverter::Convert(const FGLTFMeshSection
 	return Builder.AddAccessor(JsonAccessor);
 }
 
-FGLTFJsonAccessorIndex FGLTFNormalBufferConverter::Convert(const FGLTFMeshSection* MeshSection, const FStaticMeshVertexBuffer* VertexBuffer)
+FGLTFJsonAccessor* FGLTFNormalBufferConverter::Convert(const FGLTFMeshSection* MeshSection, const FStaticMeshVertexBuffer* VertexBuffer)
 {
 	if (VertexBuffer == nullptr || VertexBuffer->GetNumVertices() == 0)
 	{
-		return FGLTFJsonAccessorIndex(INDEX_NONE);
+		return nullptr;
 	}
 
 	const TUniquePtr<IGLTFBufferAdapter> SourceBuffer = IGLTFBufferAdapter::GetTangents(VertexBuffer);
@@ -134,10 +134,10 @@ FGLTFJsonAccessorIndex FGLTFNormalBufferConverter::Convert(const FGLTFMeshSectio
 	if (SourceData == nullptr)
 	{
 		// TODO: report error
-		return FGLTFJsonAccessorIndex(INDEX_NONE);
+		return nullptr;
 	}
 
-	FGLTFJsonBufferViewIndex BufferViewIndex;
+	FGLTFJsonBufferView* BufferViewIndex;
 	EGLTFJsonComponentType ComponentType;
 
 	const bool bMeshQuantization = Builder.ExportOptions->bUseMeshQuantization;
@@ -179,7 +179,7 @@ FGLTFJsonAccessorIndex FGLTFNormalBufferConverter::Convert(const FGLTFMeshSectio
 }
 
 template <typename DestinationType, typename SourceType>
-FGLTFJsonBufferViewIndex FGLTFNormalBufferConverter::ConvertBufferView(const FGLTFMeshSection* MeshSection, const void* SourceData)
+FGLTFJsonBufferView* FGLTFNormalBufferConverter::ConvertBufferView(const FGLTFMeshSection* MeshSection, const void* SourceData)
 {
 	const TArray<uint32>& IndexMap = MeshSection->IndexMap;
 	const uint32 VertexCount = IndexMap.Num();
@@ -202,11 +202,11 @@ FGLTFJsonBufferViewIndex FGLTFNormalBufferConverter::ConvertBufferView(const FGL
 	return Builder.AddBufferView(Normals, EGLTFJsonBufferTarget::ArrayBuffer);
 }
 
-FGLTFJsonAccessorIndex FGLTFTangentBufferConverter::Convert(const FGLTFMeshSection* MeshSection, const FStaticMeshVertexBuffer* VertexBuffer)
+FGLTFJsonAccessor* FGLTFTangentBufferConverter::Convert(const FGLTFMeshSection* MeshSection, const FStaticMeshVertexBuffer* VertexBuffer)
 {
 	if (VertexBuffer == nullptr || VertexBuffer->GetNumVertices() == 0)
 	{
-		return FGLTFJsonAccessorIndex(INDEX_NONE);
+		return nullptr;
 	}
 
 	const TUniquePtr<IGLTFBufferAdapter> SourceBuffer = IGLTFBufferAdapter::GetTangents(VertexBuffer);
@@ -215,10 +215,10 @@ FGLTFJsonAccessorIndex FGLTFTangentBufferConverter::Convert(const FGLTFMeshSecti
 	if (SourceData == nullptr)
 	{
 		// TODO: report error
-		return FGLTFJsonAccessorIndex(INDEX_NONE);
+		return nullptr;
 	}
 
-	FGLTFJsonBufferViewIndex BufferViewIndex;
+	FGLTFJsonBufferView* BufferViewIndex;
 	EGLTFJsonComponentType ComponentType;
 
 	const bool bMeshQuantization = Builder.ExportOptions->bUseMeshQuantization;
@@ -258,7 +258,7 @@ FGLTFJsonAccessorIndex FGLTFTangentBufferConverter::Convert(const FGLTFMeshSecti
 }
 
 template <typename DestinationType, typename SourceType>
-FGLTFJsonBufferViewIndex FGLTFTangentBufferConverter::ConvertBufferView(const FGLTFMeshSection* MeshSection, const void* SourceData)
+FGLTFJsonBufferView* FGLTFTangentBufferConverter::ConvertBufferView(const FGLTFMeshSection* MeshSection, const void* SourceData)
 {
 	const TArray<uint32>& IndexMap = MeshSection->IndexMap;
 	const uint32 VertexCount = IndexMap.Num();
@@ -281,18 +281,18 @@ FGLTFJsonBufferViewIndex FGLTFTangentBufferConverter::ConvertBufferView(const FG
 	return Builder.AddBufferView(Tangents, EGLTFJsonBufferTarget::ArrayBuffer);
 }
 
-FGLTFJsonAccessorIndex FGLTFUVBufferConverter::Convert(const FGLTFMeshSection* MeshSection, const FStaticMeshVertexBuffer* VertexBuffer, uint32 UVIndex)
+FGLTFJsonAccessor* FGLTFUVBufferConverter::Convert(const FGLTFMeshSection* MeshSection, const FStaticMeshVertexBuffer* VertexBuffer, uint32 UVIndex)
 {
 	if (VertexBuffer == nullptr || VertexBuffer->GetNumVertices() == 0)
 	{
-		return FGLTFJsonAccessorIndex(INDEX_NONE);
+		return nullptr;
 	}
 
 	const uint32 UVCount = VertexBuffer->GetNumTexCoords();
 	if (UVIndex >= UVCount)
 	{
 		// TODO: report warning
-		return FGLTFJsonAccessorIndex(INDEX_NONE);
+		return nullptr;
 	}
 
 	// TODO: report warning or add support for half float precision UVs, i.e. !VertexBuffer->GetUseFullPrecisionUVs()?
@@ -303,7 +303,7 @@ FGLTFJsonAccessorIndex FGLTFUVBufferConverter::Convert(const FGLTFMeshSection* M
 	if (SourceData == nullptr)
 	{
 		// TODO: report error
-		return FGLTFJsonAccessorIndex(INDEX_NONE);
+		return nullptr;
 	}
 
 	return VertexBuffer->GetUseFullPrecisionUVs()
@@ -312,7 +312,7 @@ FGLTFJsonAccessorIndex FGLTFUVBufferConverter::Convert(const FGLTFMeshSection* M
 }
 
 template <typename SourceType>
-FGLTFJsonAccessorIndex FGLTFUVBufferConverter::Convert(const FGLTFMeshSection* MeshSection, const FStaticMeshVertexBuffer* VertexBuffer, uint32 UVIndex, const uint8* SourceData) const
+FGLTFJsonAccessor* FGLTFUVBufferConverter::Convert(const FGLTFMeshSection* MeshSection, const FStaticMeshVertexBuffer* VertexBuffer, uint32 UVIndex, const uint8* SourceData) const
 {
 	const TArray<uint32>& IndexMap = MeshSection->IndexMap;
 	const uint32 VertexCount = IndexMap.Num();
@@ -339,11 +339,11 @@ FGLTFJsonAccessorIndex FGLTFUVBufferConverter::Convert(const FGLTFMeshSection* M
 	return Builder.AddAccessor(JsonAccessor);
 }
 
-FGLTFJsonAccessorIndex FGLTFBoneIndexBufferConverter::Convert(const FGLTFMeshSection* MeshSection, const FSkinWeightVertexBuffer* VertexBuffer, uint32 InfluenceOffset)
+FGLTFJsonAccessor* FGLTFBoneIndexBufferConverter::Convert(const FGLTFMeshSection* MeshSection, const FSkinWeightVertexBuffer* VertexBuffer, uint32 InfluenceOffset)
 {
 	if (VertexBuffer == nullptr || VertexBuffer->GetNumVertices() == 0)
 	{
-		return FGLTFJsonAccessorIndex(INDEX_NONE);
+		return nullptr;
 	}
 
 	const TUniquePtr<IGLTFBufferAdapter> SourceBuffer = IGLTFBufferAdapter::GetInfluences(VertexBuffer);
@@ -352,7 +352,7 @@ FGLTFJsonAccessorIndex FGLTFBoneIndexBufferConverter::Convert(const FGLTFMeshSec
 	if (SourceData == nullptr)
 	{
 		// TODO: report error
-		return FGLTFJsonAccessorIndex(INDEX_NONE);
+		return nullptr;
 	}
 
 	return MeshSection->MaxBoneIndex > UINT8_MAX
@@ -361,7 +361,7 @@ FGLTFJsonAccessorIndex FGLTFBoneIndexBufferConverter::Convert(const FGLTFMeshSec
 }
 
 template <typename DestinationType>
-FGLTFJsonAccessorIndex FGLTFBoneIndexBufferConverter::Convert(const FGLTFMeshSection* MeshSection, const FSkinWeightVertexBuffer* VertexBuffer, uint32 InfluenceOffset, const uint8* SourceData) const
+FGLTFJsonAccessor* FGLTFBoneIndexBufferConverter::Convert(const FGLTFMeshSection* MeshSection, const FSkinWeightVertexBuffer* VertexBuffer, uint32 InfluenceOffset, const uint8* SourceData) const
 {
 	return VertexBuffer->Use16BitBoneIndex()
 		? Convert<DestinationType, uint16>(MeshSection, VertexBuffer, InfluenceOffset, SourceData)
@@ -369,7 +369,7 @@ FGLTFJsonAccessorIndex FGLTFBoneIndexBufferConverter::Convert(const FGLTFMeshSec
 }
 
 template <typename DestinationType, typename SourceType>
-FGLTFJsonAccessorIndex FGLTFBoneIndexBufferConverter::Convert(const FGLTFMeshSection* MeshSection, const FSkinWeightVertexBuffer* VertexBuffer, uint32 InfluenceOffset, const uint8* SourceData) const
+FGLTFJsonAccessor* FGLTFBoneIndexBufferConverter::Convert(const FGLTFMeshSection* MeshSection, const FSkinWeightVertexBuffer* VertexBuffer, uint32 InfluenceOffset, const uint8* SourceData) const
 {
 	if (VertexBuffer->GetVariableBonesPerVertex())
 	{
@@ -379,7 +379,7 @@ FGLTFJsonAccessorIndex FGLTFBoneIndexBufferConverter::Convert(const FGLTFMeshSec
 		if (LookupData == nullptr)
 		{
 			// TODO: report error
-			return FGLTFJsonAccessorIndex(INDEX_NONE);
+			return nullptr;
 		}
 
 		return Convert<DestinationType, SourceType>(MeshSection, VertexBuffer, InfluenceOffset, SourceData, [LookupData](uint32 VertexIndex, uint32& VertexDataOffset, uint32& VertexInfluenceCount)
@@ -399,7 +399,7 @@ FGLTFJsonAccessorIndex FGLTFBoneIndexBufferConverter::Convert(const FGLTFMeshSec
 }
 
 template <typename DestinationType, typename SourceType, typename CallbackType>
-FGLTFJsonAccessorIndex FGLTFBoneIndexBufferConverter::Convert(const FGLTFMeshSection* MeshSection, const FSkinWeightVertexBuffer* VertexBuffer, uint32 InfluenceOffset, const uint8* SourceData, CallbackType GetVertexInfluenceOffsetCount) const
+FGLTFJsonAccessor* FGLTFBoneIndexBufferConverter::Convert(const FGLTFMeshSection* MeshSection, const FSkinWeightVertexBuffer* VertexBuffer, uint32 InfluenceOffset, const uint8* SourceData, CallbackType GetVertexInfluenceOffsetCount) const
 {
 	const TArray<uint32>& IndexMap = MeshSection->IndexMap;
 	const uint32 VertexCount = IndexMap.Num();
@@ -446,11 +446,11 @@ FGLTFJsonAccessorIndex FGLTFBoneIndexBufferConverter::Convert(const FGLTFMeshSec
 	return Builder.AddAccessor(JsonAccessor);
 }
 
-FGLTFJsonAccessorIndex FGLTFBoneWeightBufferConverter::Convert(const FGLTFMeshSection* MeshSection, const FSkinWeightVertexBuffer* VertexBuffer, uint32 InfluenceOffset)
+FGLTFJsonAccessor* FGLTFBoneWeightBufferConverter::Convert(const FGLTFMeshSection* MeshSection, const FSkinWeightVertexBuffer* VertexBuffer, uint32 InfluenceOffset)
 {
 	if (VertexBuffer == nullptr || VertexBuffer->GetNumVertices() == 0)
 	{
-		return FGLTFJsonAccessorIndex(INDEX_NONE);
+		return nullptr;
 	}
 
 	const TUniquePtr<IGLTFBufferAdapter> SourceBuffer = IGLTFBufferAdapter::GetInfluences(VertexBuffer);
@@ -459,7 +459,7 @@ FGLTFJsonAccessorIndex FGLTFBoneWeightBufferConverter::Convert(const FGLTFMeshSe
 	if (SourceData == nullptr)
 	{
 		// TODO: report error
-		return FGLTFJsonAccessorIndex(INDEX_NONE);
+		return nullptr;
 	}
 
 	return MeshSection->MaxBoneIndex > UINT8_MAX
@@ -468,7 +468,7 @@ FGLTFJsonAccessorIndex FGLTFBoneWeightBufferConverter::Convert(const FGLTFMeshSe
 }
 
 template <typename BoneIndexType>
-FGLTFJsonAccessorIndex FGLTFBoneWeightBufferConverter::Convert(const FGLTFMeshSection* MeshSection, const FSkinWeightVertexBuffer* VertexBuffer, uint32 InfluenceOffset, const uint8* SourceData) const
+FGLTFJsonAccessor* FGLTFBoneWeightBufferConverter::Convert(const FGLTFMeshSection* MeshSection, const FSkinWeightVertexBuffer* VertexBuffer, uint32 InfluenceOffset, const uint8* SourceData) const
 {
 	if (VertexBuffer->GetVariableBonesPerVertex())
 	{
@@ -478,7 +478,7 @@ FGLTFJsonAccessorIndex FGLTFBoneWeightBufferConverter::Convert(const FGLTFMeshSe
 		if (LookupData == nullptr)
 		{
 			// TODO: report error
-			return FGLTFJsonAccessorIndex(INDEX_NONE);
+			return nullptr;
 		}
 
 		return Convert<BoneIndexType>(MeshSection, VertexBuffer, InfluenceOffset, SourceData, [LookupData](uint32 VertexIndex, uint32& VertexDataOffset, uint32& VertexInfluenceCount)
@@ -498,7 +498,7 @@ FGLTFJsonAccessorIndex FGLTFBoneWeightBufferConverter::Convert(const FGLTFMeshSe
 }
 
 template <typename BoneIndexType, typename CallbackType>
-FGLTFJsonAccessorIndex FGLTFBoneWeightBufferConverter::Convert(const FGLTFMeshSection* MeshSection, const FSkinWeightVertexBuffer* VertexBuffer, uint32 InfluenceOffset, const uint8* SourceData, CallbackType GetVertexInfluenceOffsetCount) const
+FGLTFJsonAccessor* FGLTFBoneWeightBufferConverter::Convert(const FGLTFMeshSection* MeshSection, const FSkinWeightVertexBuffer* VertexBuffer, uint32 InfluenceOffset, const uint8* SourceData, CallbackType GetVertexInfluenceOffsetCount) const
 {
 	const TArray<uint32>& IndexMap = MeshSection->IndexMap;
 	const uint32 VertexCount = IndexMap.Num();
@@ -544,7 +544,7 @@ FGLTFJsonAccessorIndex FGLTFBoneWeightBufferConverter::Convert(const FGLTFMeshSe
 	return Builder.AddAccessor(JsonAccessor);
 }
 
-FGLTFJsonAccessorIndex FGLTFIndexBufferConverter::Convert(const FGLTFMeshSection* MeshSection)
+FGLTFJsonAccessor* FGLTFIndexBufferConverter::Convert(const FGLTFMeshSection* MeshSection)
 {
 	const uint32 MaxVertexIndex = MeshSection->IndexMap.Num() - 1;
 	if (MaxVertexIndex <= UINT8_MAX) return Convert<uint8>(MeshSection);
@@ -553,13 +553,13 @@ FGLTFJsonAccessorIndex FGLTFIndexBufferConverter::Convert(const FGLTFMeshSection
 }
 
 template <typename IndexType>
-FGLTFJsonAccessorIndex FGLTFIndexBufferConverter::Convert(const FGLTFMeshSection* MeshSection) const
+FGLTFJsonAccessor* FGLTFIndexBufferConverter::Convert(const FGLTFMeshSection* MeshSection) const
 {
 	const TArray<uint32>& IndexBuffer = MeshSection->IndexBuffer;
 	const uint32 IndexCount = IndexBuffer.Num();
 	if (IndexCount == 0)
 	{
-		return FGLTFJsonAccessorIndex(INDEX_NONE);
+		return nullptr;
 	}
 
 	TArray<IndexType> Indices;
diff --git a/Engine/Plugins/Enterprise/GLTFExporter/MigratedFromGit/Source/GLTFExporter/Private/Converters/GLTFAccessorConverters.h b/Engine/Plugins/Enterprise/GLTFExporter/MigratedFromGit/Source/GLTFExporter/Private/Converters/GLTFAccessorConverters.h
--- a/Engine/Plugins/Enterprise/GLTFExporter/MigratedFromGit/Source/GLTFExporter/Private/Converters/GLTFAccessorConverters.h
+++ b/Engine/Plugins/Enterprise/GLTFExporter/MigratedFromGit/Source/GLTFExporter/Private/Converters/GLTFAccessorConverters.h
@@ -7,94 +7,94 @@
 #include "Converters/GLTFBuilderContext.h"
 #include "Converters/GLTFMeshSection.h"
 
-typedef TGLTFConverter<FGLTFJsonAccessorIndex, const FGLTFMeshSection*, const FPositionVertexBuffer*> IGLTFPositionBufferConverter;
-typedef TGLTFConverter<FGLTFJsonAccessorIndex, const FGLTFMeshSection*, const FColorVertexBuffer*> IGLTFColorBufferConverter;
-typedef TGLTFConverter<FGLTFJsonAccessorIndex, const FGLTFMeshSection*, const FStaticMeshVertexBuffer*> IGLTFNormalBufferConverter;
-typedef TGLTFConverter<FGLTFJsonAccessorIndex, const FGLTFMeshSection*, const FStaticMeshVertexBuffer*> IGLTFTangentBufferConverter;
-typedef TGLTFConverter<FGLTFJsonAccessorIndex, const FGLTFMeshSection*, const FStaticMeshVertexBuffer*, uint32> IGLTFUVBufferConverter;
-typedef TGLTFConverter<FGLTFJsonAccessorIndex, const FGLTFMeshSection*, const FSkinWeightVertexBuffer*, uint32> IGLTFBoneIndexBufferConverter;
-typedef TGLTFConverter<FGLTFJsonAccessorIndex, const FGLTFMeshSection*, const FSkinWeightVertexBuffer*, uint32> IGLTFBoneWeightBufferConverter;
-typedef TGLTFConverter<FGLTFJsonAccessorIndex, const FGLTFMeshSection*> IGLTFIndexBufferConverter;
+typedef TGLTFConverter<FGLTFJsonAccessor*, const FGLTFMeshSection*, const FPositionVertexBuffer*> IGLTFPositionBufferConverter;
+typedef TGLTFConverter<FGLTFJsonAccessor*, const FGLTFMeshSection*, const FColorVertexBuffer*> IGLTFColorBufferConverter;
+typedef TGLTFConverter<FGLTFJsonAccessor*, const FGLTFMeshSection*, const FStaticMeshVertexBuffer*> IGLTFNormalBufferConverter;
+typedef TGLTFConverter<FGLTFJsonAccessor*, const FGLTFMeshSection*, const FStaticMeshVertexBuffer*> IGLTFTangentBufferConverter;
+typedef TGLTFConverter<FGLTFJsonAccessor*, const FGLTFMeshSection*, const FStaticMeshVertexBuffer*, uint32> IGLTFUVBufferConverter;
+typedef TGLTFConverter<FGLTFJsonAccessor*, const FGLTFMeshSection*, const FSkinWeightVertexBuffer*, uint32> IGLTFBoneIndexBufferConverter;
+typedef TGLTFConverter<FGLTFJsonAccessor*, const FGLTFMeshSection*, const FSkinWeightVertexBuffer*, uint32> IGLTFBoneWeightBufferConverter;
+typedef TGLTFConverter<FGLTFJsonAccessor*, const FGLTFMeshSection*> IGLTFIndexBufferConverter;
 
 class FGLTFPositionBufferConverter final : public FGLTFBuilderContext, public IGLTFPositionBufferConverter
 {
 	using FGLTFBuilderContext::FGLTFBuilderContext;
 
-	virtual FGLTFJsonAccessorIndex Convert(const FGLTFMeshSection* MeshSection, const FPositionVertexBuffer* VertexBuffer) override;
+	virtual FGLTFJsonAccessor* Convert(const FGLTFMeshSection* MeshSection, const FPositionVertexBuffer* VertexBuffer) override;
 };
 
 class FGLTFColorBufferConverter final : public FGLTFBuilderContext, public IGLTFColorBufferConverter
 {
 	using FGLTFBuilderContext::FGLTFBuilderContext;
 
-	virtual FGLTFJsonAccessorIndex Convert(const FGLTFMeshSection* MeshSection, const FColorVertexBuffer* VertexBuffer) override;
+	virtual FGLTFJsonAccessor* Convert(const FGLTFMeshSection* MeshSection, const FColorVertexBuffer* VertexBuffer) override;
 };
 
 class FGLTFNormalBufferConverter final : public FGLTFBuilderContext, public IGLTFNormalBufferConverter
 {
 	using FGLTFBuilderContext::FGLTFBuilderContext;
 
-	virtual FGLTFJsonAccessorIndex Convert(const FGLTFMeshSection* MeshSection, const FStaticMeshVertexBuffer* VertexBuffer) override;
+	virtual FGLTFJsonAccessor* Convert(const FGLTFMeshSection* MeshSection, const FStaticMeshVertexBuffer* VertexBuffer) override;
 
 	template <typename DestinationType, typename SourceType>
-	FGLTFJsonBufferViewIndex ConvertBufferView(const FGLTFMeshSection* MeshSection, const void* TangentData);
+	FGLTFJsonBufferView* ConvertBufferView(const FGLTFMeshSection* MeshSection, const void* TangentData);
 };
 
 class FGLTFTangentBufferConverter final : public FGLTFBuilderContext, public IGLTFTangentBufferConverter
 {
 	using FGLTFBuilderContext::FGLTFBuilderContext;
 
-	virtual FGLTFJsonAccessorIndex Convert(const FGLTFMeshSection* MeshSection, const FStaticMeshVertexBuffer* VertexBuffer) override;
+	virtual FGLTFJsonAccessor* Convert(const FGLTFMeshSection* MeshSection, const FStaticMeshVertexBuffer* VertexBuffer) override;
 
 	template <typename DestinationType, typename SourceType>
-	FGLTFJsonBufferViewIndex ConvertBufferView(const FGLTFMeshSection* MeshSection, const void* TangentData);
+	FGLTFJsonBufferView* ConvertBufferView(const FGLTFMeshSection* MeshSection, const void* TangentData);
 };
 
 class FGLTFUVBufferConverter final : public FGLTFBuilderContext, public IGLTFUVBufferConverter
 {
 	using FGLTFBuilderContext::FGLTFBuilderContext;
 
-	virtual FGLTFJsonAccessorIndex Convert(const FGLTFMeshSection* MeshSection, const FStaticMeshVertexBuffer* VertexBuffer, uint32 UVIndex) override;
+	virtual FGLTFJsonAccessor* Convert(const FGLTFMeshSection* MeshSection, const FStaticMeshVertexBuffer* VertexBuffer, uint32 UVIndex) override;
 
 	template <typename SourceType>
-	FGLTFJsonAccessorIndex Convert(const FGLTFMeshSection* MeshSection, const FStaticMeshVertexBuffer* VertexBuffer, uint32 UVIndex, const uint8* SourceData) const;
+	FGLTFJsonAccessor* Convert(const FGLTFMeshSection* MeshSection, const FStaticMeshVertexBuffer* VertexBuffer, uint32 UVIndex, const uint8* SourceData) const;
 };
 
 class FGLTFBoneIndexBufferConverter final : public FGLTFBuilderContext, public IGLTFBoneIndexBufferConverter
 {
 	using FGLTFBuilderContext::FGLTFBuilderContext;
 
-	virtual FGLTFJsonAccessorIndex Convert(const FGLTFMeshSection* MeshSection, const FSkinWeightVertexBuffer* VertexBuffer, uint32 InfluenceOffset) override;
+	virtual FGLTFJsonAccessor* Convert(const FGLTFMeshSection* MeshSection, const FSkinWeightVertexBuffer* VertexBuffer, uint32 InfluenceOffset) override;
 
 	template <typename DestinationType>
-	FGLTFJsonAccessorIndex Convert(const FGLTFMeshSection* MeshSection, const FSkinWeightVertexBuffer* VertexBuffer, uint32 InfluenceOffset, const uint8* SourceData) const;
+	FGLTFJsonAccessor* Convert(const FGLTFMeshSection* MeshSection, const FSkinWeightVertexBuffer* VertexBuffer, uint32 InfluenceOffset, const uint8* SourceData) const;
 
 	template <typename DestinationType, typename SourceType>
-	FGLTFJsonAccessorIndex Convert(const FGLTFMeshSection* MeshSection, const FSkinWeightVertexBuffer* VertexBuffer, uint32 InfluenceOffset, const uint8* SourceData) const;
+	FGLTFJsonAccessor* Convert(const FGLTFMeshSection* MeshSection, const FSkinWeightVertexBuffer* VertexBuffer, uint32 InfluenceOffset, const uint8* SourceData) const;
 
 	template <typename DestinationType, typename SourceType, typename CallbackType>
-	FGLTFJsonAccessorIndex Convert(const FGLTFMeshSection* MeshSection, const FSkinWeightVertexBuffer* VertexBuffer, uint32 InfluenceOffset, const uint8* SourceData, CallbackType GetVertexInfluenceOffsetCount) const;
+	FGLTFJsonAccessor* Convert(const FGLTFMeshSection* MeshSection, const FSkinWeightVertexBuffer* VertexBuffer, uint32 InfluenceOffset, const uint8* SourceData, CallbackType GetVertexInfluenceOffsetCount) const;
 };
 
 class FGLTFBoneWeightBufferConverter final : public FGLTFBuilderContext, public IGLTFBoneWeightBufferConverter
 {
 	using FGLTFBuilderContext::FGLTFBuilderContext;
 
-	virtual FGLTFJsonAccessorIndex Convert(const FGLTFMeshSection* MeshSection, const FSkinWeightVertexBuffer* VertexBuffer, uint32 InfluenceOffset) override;
+	virtual FGLTFJsonAccessor* Convert(const FGLTFMeshSection* MeshSection, const FSkinWeightVertexBuffer* VertexBuffer, uint32 InfluenceOffset) override;
 
 	template <typename BoneIndexType>
-	FGLTFJsonAccessorIndex Convert(const FGLTFMeshSection* MeshSection, const FSkinWeightVertexBuffer* VertexBuffer, uint32 InfluenceOffset, const uint8* SourceData) const;
+	FGLTFJsonAccessor* Convert(const FGLTFMeshSection* MeshSection, const FSkinWeightVertexBuffer* VertexBuffer, uint32 InfluenceOffset, const uint8* SourceData) const;
 
 	template <typename BoneIndexType, typename CallbackType>
-	FGLTFJsonAccessorIndex Convert(const FGLTFMeshSection* MeshSection, const FSkinWeightVertexBuffer* VertexBuffer, uint32 InfluenceOffset, const uint8* SourceData, CallbackType GetVertexInfluenceOffsetCount) const;
+	FGLTFJsonAccessor* Convert(const FGLTFMeshSection* MeshSection, const FSkinWeightVertexBuffer* VertexBuffer, uint32 InfluenceOffset, const uint8* SourceData, CallbackType GetVertexInfluenceOffsetCount) const;
 };
 
 class FGLTFIndexBufferConverter final : public FGLTFBuilderContext, public IGLTFIndexBufferConverter
 {
 	using FGLTFBuilderContext::FGLTFBuilderContext;
 
-	virtual FGLTFJsonAccessorIndex Convert(const FGLTFMeshSection* MeshSection) override;
+	virtual FGLTFJsonAccessor* Convert(const FGLTFMeshSection* MeshSection) override;
 
 	template <typename IndexType>
-	FGLTFJsonAccessorIndex Convert(const FGLTFMeshSection* MeshSection) const;
+	FGLTFJsonAccessor* Convert(const FGLTFMeshSection* MeshSection) const;
 };
diff --git a/Engine/Plugins/Enterprise/GLTFExporter/MigratedFromGit/Source/GLTFExporter/Private/Converters/GLTFAnimationConverters.cpp b/Engine/Plugins/Enterprise/GLTFExporter/MigratedFromGit/Source/GLTFExporter/Private/Converters/GLTFAnimationConverters.cpp
--- a/Engine/Plugins/Enterprise/GLTFExporter/MigratedFromGit/Source/GLTFExporter/Private/Converters/GLTFAnimationConverters.cpp
+++ b/Engine/Plugins/Enterprise/GLTFExporter/MigratedFromGit/Source/GLTFExporter/Private/Converters/GLTFAnimationConverters.cpp
@@ -5,19 +5,19 @@
 #include "Tasks/GLTFAnimationTasks.h"
 #include "Animation/AnimSequence.h"
 
-FGLTFJsonAnimationIndex FGLTFAnimationConverter::Convert(FGLTFJsonNodeIndex RootNode, const USkeletalMesh* SkeletalMesh, const UAnimSequence* AnimSequence)
+FGLTFJsonAnimation* FGLTFAnimationConverter::Convert(FGLTFJsonNode* RootNode, const USkeletalMesh* SkeletalMesh, const UAnimSequence* AnimSequence)
 {
 	if (AnimSequence->GetRawNumberOfFrames() < 0)
 	{
 		// TODO: report warning
-		return FGLTFJsonAnimationIndex(INDEX_NONE);
+		return nullptr;
 	}
 
 	const USkeleton* AnimSkeleton = AnimSequence->GetSkeleton();
 	if (AnimSkeleton == nullptr)
 	{
 		// TODO: report error
-		return FGLTFJsonAnimationIndex(INDEX_NONE);
+		return nullptr;
 	}
 
 #if (ENGINE_MAJOR_VERSION > 4 || ENGINE_MINOR_VERSION >= 27)
@@ -29,29 +29,28 @@ FGLTFJsonAnimationIndex FGLTFAnimationConverter::Convert(FGLTFJsonNodeIndex Root
 	if (AnimSkeleton != MeshSkeleton)
 	{
 		// TODO: report error
-		return FGLTFJsonAnimationIndex(INDEX_NONE);
+		return nullptr;
 	}
 
 	FGLTFJsonAnimation* JsonAnimation = Builder.AddAnimation();
 	Builder.SetupTask<FGLTFAnimSequenceTask>(Builder, RootNode, SkeletalMesh, AnimSequence, JsonAnimation);
-	return JsonAnimation->Index;
+	return JsonAnimation;
 }
 
-FGLTFJsonAnimationIndex FGLTFAnimationDataConverter::Convert(FGLTFJsonNodeIndex RootNode, const USkeletalMeshComponent* SkeletalMeshComponent)
+FGLTFJsonAnimation* FGLTFAnimationDataConverter::Convert(FGLTFJsonNode* RootNode, const USkeletalMeshComponent* SkeletalMeshComponent)
 {
 	const USkeletalMesh* SkeletalMesh = SkeletalMeshComponent->SkeletalMesh;
 	const UAnimSequence* AnimSequence = Cast<UAnimSequence>(SkeletalMeshComponent->AnimationData.AnimToPlay);
 
 	if (SkeletalMesh == nullptr || AnimSequence == nullptr || SkeletalMeshComponent->GetAnimationMode() != EAnimationMode::AnimationSingleNode)
 	{
-		return FGLTFJsonAnimationIndex(INDEX_NONE);
+		return nullptr;
 	}
 
-	const FGLTFJsonAnimationIndex AnimationIndex = Builder.GetOrAddAnimation(RootNode, SkeletalMesh, AnimSequence);
-	if (AnimationIndex != INDEX_NONE && Builder.ExportOptions->bExportPlaybackSettings)
+	FGLTFJsonAnimation* JsonAnimation = Builder.GetOrAddAnimation(RootNode, SkeletalMesh, AnimSequence);
+	if (JsonAnimation != nullptr && Builder.ExportOptions->bExportPlaybackSettings)
 	{
-		FGLTFJsonAnimation& JsonAnimation = Builder.GetAnimation(AnimationIndex);
-		FGLTFJsonAnimationPlayback& JsonPlayback = JsonAnimation.Playback;
+		FGLTFJsonAnimationPlayback& JsonPlayback = JsonAnimation->Playback;
 
 		JsonPlayback.bLoop = SkeletalMeshComponent->AnimationData.bSavedLooping;
 		JsonPlayback.bAutoPlay = SkeletalMeshComponent->AnimationData.bSavedPlaying;
@@ -59,31 +58,30 @@ FGLTFJsonAnimationIndex FGLTFAnimationDataConverter::Convert(FGLTFJsonNodeIndex
 		JsonPlayback.StartTime = SkeletalMeshComponent->AnimationData.SavedPosition;
 	}
 
-	return AnimationIndex;
+	return JsonAnimation;
 }
 
-FGLTFJsonAnimationIndex FGLTFLevelSequenceConverter::Convert(const ULevel* Level, const ULevelSequence* LevelSequence)
+FGLTFJsonAnimation* FGLTFLevelSequenceConverter::Convert(const ULevel* Level, const ULevelSequence* LevelSequence)
 {
 	FGLTFJsonAnimation* JsonAnimation = Builder.AddAnimation();
 	Builder.SetupTask<FGLTFLevelSequenceTask>(Builder, Level, LevelSequence, JsonAnimation);
-	return JsonAnimation->Index;
+	return JsonAnimation;
 }
 
-FGLTFJsonAnimationIndex FGLTFLevelSequenceDataConverter::Convert(const ALevelSequenceActor* LevelSequenceActor)
+FGLTFJsonAnimation* FGLTFLevelSequenceDataConverter::Convert(const ALevelSequenceActor* LevelSequenceActor)
 {
 	const ULevel* Level = LevelSequenceActor->GetLevel();
 	const ULevelSequence* LevelSequence = LevelSequenceActor->LoadSequence();
 
 	if (Level == nullptr || LevelSequence == nullptr)
 	{
-		return FGLTFJsonAnimationIndex(INDEX_NONE);
+		return nullptr;
 	}
 
-	const FGLTFJsonAnimationIndex AnimationIndex = Builder.GetOrAddAnimation(Level, LevelSequence);
-	if (AnimationIndex != INDEX_NONE && Builder.ExportOptions->bExportPlaybackSettings)
+	FGLTFJsonAnimation* JsonAnimation = Builder.GetOrAddAnimation(Level, LevelSequence);
+	if (JsonAnimation != nullptr && Builder.ExportOptions->bExportPlaybackSettings)
 	{
-		FGLTFJsonAnimation& JsonAnimation = Builder.GetAnimation(AnimationIndex);
-		FGLTFJsonAnimationPlayback& JsonPlayback = JsonAnimation.Playback;
+		FGLTFJsonAnimationPlayback& JsonPlayback = JsonAnimation->Playback;
 
 		// TODO: report warning if loop count is not 0 or -1 (infinite)
 		JsonPlayback.bLoop = LevelSequenceActor->PlaybackSettings.LoopCount.Value != 0;
@@ -92,5 +90,5 @@ FGLTFJsonAnimationIndex FGLTFLevelSequenceDataConverter::Convert(const ALevelSeq
 		JsonPlayback.StartTime = LevelSequenceActor->PlaybackSettings.StartTime;
 	}
 
-	return AnimationIndex;
+	return JsonAnimation;
 }
diff --git a/Engine/Plugins/Enterprise/GLTFExporter/MigratedFromGit/Source/GLTFExporter/Private/Converters/GLTFAnimationConverters.h b/Engine/Plugins/Enterprise/GLTFExporter/MigratedFromGit/Source/GLTFExporter/Private/Converters/GLTFAnimationConverters.h
--- a/Engine/Plugins/Enterprise/GLTFExporter/MigratedFromGit/Source/GLTFExporter/Private/Converters/GLTFAnimationConverters.h
+++ b/Engine/Plugins/Enterprise/GLTFExporter/MigratedFromGit/Source/GLTFExporter/Private/Converters/GLTFAnimationConverters.h
@@ -9,35 +9,35 @@
 class ULevelSequence;
 class ALevelSequenceActor;
 
-typedef TGLTFConverter<FGLTFJsonAnimationIndex, FGLTFJsonNodeIndex, const USkeletalMesh*, const UAnimSequence*> IGLTFAnimationConverter;
-typedef TGLTFConverter<FGLTFJsonAnimationIndex, FGLTFJsonNodeIndex, const USkeletalMeshComponent*> IGLTFAnimationDataConverter;
-typedef TGLTFConverter<FGLTFJsonAnimationIndex, const ULevel*, const ULevelSequence*> IGLTFLevelSequenceConverter;
-typedef TGLTFConverter<FGLTFJsonAnimationIndex, const ALevelSequenceActor*> IGLTFLevelSequenceDataConverter;
+typedef TGLTFConverter<FGLTFJsonAnimation*, FGLTFJsonNode*, const USkeletalMesh*, const UAnimSequence*> IGLTFAnimationConverter;
+typedef TGLTFConverter<FGLTFJsonAnimation*, FGLTFJsonNode*, const USkeletalMeshComponent*> IGLTFAnimationDataConverter;
+typedef TGLTFConverter<FGLTFJsonAnimation*, const ULevel*, const ULevelSequence*> IGLTFLevelSequenceConverter;
+typedef TGLTFConverter<FGLTFJsonAnimation*, const ALevelSequenceActor*> IGLTFLevelSequenceDataConverter;
 
 class FGLTFAnimationConverter final : public FGLTFBuilderContext, public IGLTFAnimationConverter
 {
 	using FGLTFBuilderContext::FGLTFBuilderContext;
 
-	virtual FGLTFJsonAnimationIndex Convert(FGLTFJsonNodeIndex RootNode, const USkeletalMesh* SkeletalMesh, const UAnimSequence* AnimSequence) override;
+	virtual FGLTFJsonAnimation* Convert(FGLTFJsonNode* RootNode, const USkeletalMesh* SkeletalMesh, const UAnimSequence* AnimSequence) override;
 };
 
 class FGLTFAnimationDataConverter final : public FGLTFBuilderContext, public IGLTFAnimationDataConverter
 {
 	using FGLTFBuilderContext::FGLTFBuilderContext;
 
-	virtual FGLTFJsonAnimationIndex Convert(FGLTFJsonNodeIndex RootNode, const USkeletalMeshComponent* SkeletalMeshComponent) override;
+	virtual FGLTFJsonAnimation* Convert(FGLTFJsonNode* RootNode, const USkeletalMeshComponent* SkeletalMeshComponent) override;
 };
 
 class FGLTFLevelSequenceConverter final : public FGLTFBuilderContext, public IGLTFLevelSequenceConverter
 {
 	using FGLTFBuilderContext::FGLTFBuilderContext;
 
-	virtual FGLTFJsonAnimationIndex Convert(const ULevel* Level, const ULevelSequence*) override;
+	virtual FGLTFJsonAnimation* Convert(const ULevel* Level, const ULevelSequence*) override;
 };
 
 class FGLTFLevelSequenceDataConverter final : public FGLTFBuilderContext, public IGLTFLevelSequenceDataConverter
 {
 	using FGLTFBuilderContext::FGLTFBuilderContext;
 
-	virtual FGLTFJsonAnimationIndex Convert(const ALevelSequenceActor* LevelSequenceActor) override;
+	virtual FGLTFJsonAnimation* Convert(const ALevelSequenceActor* LevelSequenceActor) override;
 };
diff --git a/Engine/Plugins/Enterprise/GLTFExporter/MigratedFromGit/Source/GLTFExporter/Private/Converters/GLTFBackdropConverters.cpp b/Engine/Plugins/Enterprise/GLTFExporter/MigratedFromGit/Source/GLTFExporter/Private/Converters/GLTFBackdropConverters.cpp
--- a/Engine/Plugins/Enterprise/GLTFExporter/MigratedFromGit/Source/GLTFExporter/Private/Converters/GLTFBackdropConverters.cpp
+++ b/Engine/Plugins/Enterprise/GLTFExporter/MigratedFromGit/Source/GLTFExporter/Private/Converters/GLTFBackdropConverters.cpp
@@ -6,11 +6,11 @@
 #include "Converters/GLTFActorUtility.h"
 #include "Builders/GLTFContainerBuilder.h"
 
-FGLTFJsonBackdropIndex FGLTFBackdropConverter::Convert(const AActor* BackdropActor)
+FGLTFJsonBackdrop* FGLTFBackdropConverter::Convert(const AActor* BackdropActor)
 {
 	if (!FGLTFActorUtility::IsHDRIBackdropBlueprint(FGLTFActorUtility::GetBlueprintPath(BackdropActor)))
 	{
-		return FGLTFJsonBackdropIndex(INDEX_NONE);
+		return nullptr;
 	}
 
 	FGLTFJsonBackdrop JsonBackdrop;
diff --git a/Engine/Plugins/Enterprise/GLTFExporter/MigratedFromGit/Source/GLTFExporter/Private/Converters/GLTFBackdropConverters.h b/Engine/Plugins/Enterprise/GLTFExporter/MigratedFromGit/Source/GLTFExporter/Private/Converters/GLTFBackdropConverters.h
--- a/Engine/Plugins/Enterprise/GLTFExporter/MigratedFromGit/Source/GLTFExporter/Private/Converters/GLTFBackdropConverters.h
+++ b/Engine/Plugins/Enterprise/GLTFExporter/MigratedFromGit/Source/GLTFExporter/Private/Converters/GLTFBackdropConverters.h
@@ -6,12 +6,12 @@
 #include "Converters/GLTFConverter.h"
 #include "Converters/GLTFBuilderContext.h"
 
-typedef TGLTFConverter<FGLTFJsonBackdropIndex, const AActor*> IGLTFBackdropConverter;
+typedef TGLTFConverter<FGLTFJsonBackdrop*, const AActor*> IGLTFBackdropConverter;
 
 
 class FGLTFBackdropConverter final : public FGLTFBuilderContext, public IGLTFBackdropConverter
 {
 	using FGLTFBuilderContext::FGLTFBuilderContext;
 
-	virtual FGLTFJsonBackdropIndex Convert(const AActor* BackdropActor) override;
+	virtual FGLTFJsonBackdrop* Convert(const AActor* BackdropActor) override;
 };
diff --git a/Engine/Plugins/Enterprise/GLTFExporter/MigratedFromGit/Source/GLTFExporter/Private/Converters/GLTFCameraConverters.cpp b/Engine/Plugins/Enterprise/GLTFExporter/MigratedFromGit/Source/GLTFExporter/Private/Converters/GLTFCameraConverters.cpp
--- a/Engine/Plugins/Enterprise/GLTFExporter/MigratedFromGit/Source/GLTFExporter/Private/Converters/GLTFCameraConverters.cpp
+++ b/Engine/Plugins/Enterprise/GLTFExporter/MigratedFromGit/Source/GLTFExporter/Private/Converters/GLTFCameraConverters.cpp
@@ -7,7 +7,7 @@
 #include "Actors/GLTFCameraActor.h"
 #include "Camera/CameraComponent.h"
 
-FGLTFJsonCameraIndex FGLTFCameraConverter::Convert(const UCameraComponent* CameraComponent)
+FGLTFJsonCamera* FGLTFCameraConverter::Convert(const UCameraComponent* CameraComponent)
 {
 	FGLTFJsonCamera Camera;
 	Camera.Name = FGLTFNameUtility::GetName(CameraComponent);
@@ -43,7 +43,7 @@ FGLTFJsonCameraIndex FGLTFCameraConverter::Convert(const UCameraComponent* Camer
 
 		case EGLTFJsonCameraType::None:
 			// TODO: report error (unsupported camera type)
-			return FGLTFJsonCameraIndex(INDEX_NONE);
+			return nullptr;
 
 		default:
 			checkNoEntry();
diff --git a/Engine/Plugins/Enterprise/GLTFExporter/MigratedFromGit/Source/GLTFExporter/Private/Converters/GLTFCameraConverters.h b/Engine/Plugins/Enterprise/GLTFExporter/MigratedFromGit/Source/GLTFExporter/Private/Converters/GLTFCameraConverters.h
--- a/Engine/Plugins/Enterprise/GLTFExporter/MigratedFromGit/Source/GLTFExporter/Private/Converters/GLTFCameraConverters.h
+++ b/Engine/Plugins/Enterprise/GLTFExporter/MigratedFromGit/Source/GLTFExporter/Private/Converters/GLTFCameraConverters.h
@@ -8,11 +8,11 @@
 
 class UCameraComponent;
 
-typedef TGLTFConverter<FGLTFJsonCameraIndex, const UCameraComponent*> IGLTFCameraConverter;
+typedef TGLTFConverter<FGLTFJsonCamera*, const UCameraComponent*> IGLTFCameraConverter;
 
 class FGLTFCameraConverter final : public FGLTFBuilderContext, public IGLTFCameraConverter
 {
 	using FGLTFBuilderContext::FGLTFBuilderContext;
 
-	virtual FGLTFJsonCameraIndex Convert(const UCameraComponent* CameraComponent) override;
+	virtual FGLTFJsonCamera* Convert(const UCameraComponent* CameraComponent) override;
 };
diff --git a/Engine/Plugins/Enterprise/GLTFExporter/MigratedFromGit/Source/GLTFExporter/Private/Converters/GLTFEpicLevelVariantSetsConverters.cpp b/Engine/Plugins/Enterprise/GLTFExporter/MigratedFromGit/Source/GLTFExporter/Private/Converters/GLTFEpicLevelVariantSetsConverters.cpp
--- a/Engine/Plugins/Enterprise/GLTFExporter/MigratedFromGit/Source/GLTFExporter/Private/Converters/GLTFEpicLevelVariantSetsConverters.cpp
+++ b/Engine/Plugins/Enterprise/GLTFExporter/MigratedFromGit/Source/GLTFExporter/Private/Converters/GLTFEpicLevelVariantSetsConverters.cpp
@@ -11,7 +11,7 @@
 #include "PropertyValueMaterial.h"
 #include "Components/MeshComponent.h"
 
-FGLTFJsonEpicLevelVariantSetsIndex FGLTFEpicLevelVariantSetsConverter::Convert(const ULevelVariantSets* LevelVariantSets)
+FGLTFJsonEpicLevelVariantSets* FGLTFEpicLevelVariantSetsConverter::Convert(const ULevelVariantSets* LevelVariantSets)
 {
 	FGLTFJsonEpicLevelVariantSets JsonLevelVariantSets;
 	LevelVariantSets->GetName(JsonLevelVariantSets.Name);
@@ -38,7 +38,7 @@ FGLTFJsonEpicLevelVariantSetsIndex FGLTFEpicLevelVariantSetsConverter::Convert(c
 
 	if (JsonLevelVariantSets.VariantSets.Num() == 0)
 	{
-		return FGLTFJsonEpicLevelVariantSetsIndex(INDEX_NONE);
+		return nullptr;
 	}
 
 	return Builder.AddEpicLevelVariantSets(JsonLevelVariantSets);
@@ -157,8 +157,8 @@ bool FGLTFEpicLevelVariantSetsConverter::TryParseVisibilityPropertyValue(FGLTFJs
 
 	Builder.RegisterObjectVariant(Owner, Property); // TODO: we should register this on the component
 
-	const FGLTFJsonNodeIndex NodeIndex = Builder.GetOrAddNode(Target);
-	const FGLTFJsonNodeIndex ComponentNodeIndex = Builder.GetComponentNodeIndex(NodeIndex);
+	FGLTFJsonNode* NodeIndex = Builder.GetOrAddNode(Target);
+	FGLTFJsonNode* ComponentNodeIndex = Builder.GetComponentNode(NodeIndex);
 
 	FGLTFJsonEpicVariantNodeProperties& NodeProperties = OutVariant.Nodes.FindOrAdd(ComponentNodeIndex);
 
@@ -220,8 +220,8 @@ bool FGLTFEpicLevelVariantSetsConverter::TryParseMaterialPropertyValue(FGLTFJson
 	VariantMaterial.Material = FGLTFVariantUtility::GetOrAddMaterial(Builder, Material, Target, MaterialIndex);
 	VariantMaterial.Index = MaterialIndex;
 
-	const FGLTFJsonNodeIndex NodeIndex = Builder.GetOrAddNode(Target);
-	const FGLTFJsonNodeIndex ComponentNodeIndex = Builder.GetComponentNodeIndex(NodeIndex);
+	FGLTFJsonNode* NodeIndex = Builder.GetOrAddNode(Target);
+	FGLTFJsonNode* ComponentNodeIndex = Builder.GetComponentNode(NodeIndex);
 	FGLTFJsonEpicVariantNodeProperties& NodeProperties = OutVariant.Nodes.FindOrAdd(ComponentNodeIndex);
 
 	NodeProperties.Node = ComponentNodeIndex;
@@ -271,9 +271,9 @@ bool FGLTFEpicLevelVariantSetsConverter::TryParseMeshPropertyValue(FGLTFJsonEpic
 	Builder.RegisterObjectVariant(Owner, Property); // TODO: we should register this on the component
 
 	const FGLTFMaterialArray OverrideMaterials(Target->OverrideMaterials);
-	const FGLTFJsonNodeIndex NodeIndex = Builder.GetOrAddNode(Target);
-	const FGLTFJsonMeshIndex MeshIndex = Builder.GetOrAddMesh(Mesh, OverrideMaterials);
-	const FGLTFJsonNodeIndex ComponentNodeIndex = Builder.GetComponentNodeIndex(NodeIndex);
+	FGLTFJsonNode* NodeIndex = Builder.GetOrAddNode(Target);
+	FGLTFJsonMesh* MeshIndex = Builder.GetOrAddMesh(Mesh, OverrideMaterials);
+	FGLTFJsonNode* ComponentNodeIndex = Builder.GetComponentNode(NodeIndex);
 	FGLTFJsonEpicVariantNodeProperties& NodeProperties = OutVariant.Nodes.FindOrAdd(ComponentNodeIndex);
 
 	NodeProperties.Node = ComponentNodeIndex;
diff --git a/Engine/Plugins/Enterprise/GLTFExporter/MigratedFromGit/Source/GLTFExporter/Private/Converters/GLTFEpicLevelVariantSetsConverters.h b/Engine/Plugins/Enterprise/GLTFExporter/MigratedFromGit/Source/GLTFExporter/Private/Converters/GLTFEpicLevelVariantSetsConverters.h
--- a/Engine/Plugins/Enterprise/GLTFExporter/MigratedFromGit/Source/GLTFExporter/Private/Converters/GLTFEpicLevelVariantSetsConverters.h
+++ b/Engine/Plugins/Enterprise/GLTFExporter/MigratedFromGit/Source/GLTFExporter/Private/Converters/GLTFEpicLevelVariantSetsConverters.h
@@ -14,13 +14,13 @@ class UVariantObjectBinding;
 class UPropertyValue;
 class UPropertyValueMaterial;
 
-typedef TGLTFConverter<FGLTFJsonEpicLevelVariantSetsIndex, const ULevelVariantSets*> IGLTFEpicLevelVariantSetsConverter;
+typedef TGLTFConverter<FGLTFJsonEpicLevelVariantSets*, const ULevelVariantSets*> IGLTFEpicLevelVariantSetsConverter;
 
 class FGLTFEpicLevelVariantSetsConverter final : public FGLTFBuilderContext, public IGLTFEpicLevelVariantSetsConverter
 {
 	using FGLTFBuilderContext::FGLTFBuilderContext;
 
-	virtual FGLTFJsonEpicLevelVariantSetsIndex Convert(const ULevelVariantSets* LevelVariantSets) override;
+	virtual FGLTFJsonEpicLevelVariantSets* Convert(const ULevelVariantSets* LevelVariantSets) override;
 
 	bool TryParseVariant(FGLTFJsonEpicVariant& OutVariant, const UVariant* Variant) const;
 	bool TryParseVariantBinding(FGLTFJsonEpicVariant& OutVariant, const UVariantObjectBinding* Binding) const;
diff --git a/Engine/Plugins/Enterprise/GLTFExporter/MigratedFromGit/Source/GLTFExporter/Private/Converters/GLTFHotspotConverters.cpp b/Engine/Plugins/Enterprise/GLTFExporter/MigratedFromGit/Source/GLTFExporter/Private/Converters/GLTFHotspotConverters.cpp
--- a/Engine/Plugins/Enterprise/GLTFExporter/MigratedFromGit/Source/GLTFExporter/Private/Converters/GLTFHotspotConverters.cpp
+++ b/Engine/Plugins/Enterprise/GLTFExporter/MigratedFromGit/Source/GLTFExporter/Private/Converters/GLTFHotspotConverters.cpp
@@ -4,7 +4,7 @@
 #include "Builders/GLTFContainerBuilder.h"
 #include "Animation/SkeletalMeshActor.h"
 
-FGLTFJsonHotspotIndex FGLTFHotspotConverter::Convert(const AGLTFHotspotActor* HotspotActor)
+FGLTFJsonHotspot* FGLTFHotspotConverter::Convert(const AGLTFHotspotActor* HotspotActor)
 {
 	FGLTFJsonHotspot JsonHotspot;
 	HotspotActor->GetName(JsonHotspot.Name);
@@ -25,7 +25,7 @@ FGLTFJsonHotspotIndex FGLTFHotspotConverter::Convert(const AGLTFHotspotActor* Ho
 		}
 		else
 		{
-			const FGLTFJsonNodeIndex RootNode = Builder.GetOrAddNode(SkeletalMeshActor);
+			FGLTFJsonNode* RootNode = Builder.GetOrAddNode(SkeletalMeshActor);
 
 			const USkeletalMeshComponent* SkeletalMeshComponent = SkeletalMeshActor->GetSkeletalMeshComponent();
 			if (const USkeletalMesh* SkeletalMesh = SkeletalMeshComponent->SkeletalMesh)
diff --git a/Engine/Plugins/Enterprise/GLTFExporter/MigratedFromGit/Source/GLTFExporter/Private/Converters/GLTFHotspotConverters.h b/Engine/Plugins/Enterprise/GLTFExporter/MigratedFromGit/Source/GLTFExporter/Private/Converters/GLTFHotspotConverters.h
--- a/Engine/Plugins/Enterprise/GLTFExporter/MigratedFromGit/Source/GLTFExporter/Private/Converters/GLTFHotspotConverters.h
+++ b/Engine/Plugins/Enterprise/GLTFExporter/MigratedFromGit/Source/GLTFExporter/Private/Converters/GLTFHotspotConverters.h
@@ -7,11 +7,11 @@
 #include "Converters/GLTFBuilderContext.h"
 #include "Actors/GLTFHotspotActor.h"
 
-typedef TGLTFConverter<FGLTFJsonHotspotIndex, const AGLTFHotspotActor*> IGLTFHotspotConverter;
+typedef TGLTFConverter<FGLTFJsonHotspot*, const AGLTFHotspotActor*> IGLTFHotspotConverter;
 
 class FGLTFHotspotConverter final : public FGLTFBuilderContext, public IGLTFHotspotConverter
 {
 	using FGLTFBuilderContext::FGLTFBuilderContext;
 
-	virtual FGLTFJsonHotspotIndex Convert(const AGLTFHotspotActor* HotspotActor) override;
+	virtual FGLTFJsonHotspot* Convert(const AGLTFHotspotActor* HotspotActor) override;
 };
diff --git a/Engine/Plugins/Enterprise/GLTFExporter/MigratedFromGit/Source/GLTFExporter/Private/Converters/GLTFImageConverters.cpp b/Engine/Plugins/Enterprise/GLTFExporter/MigratedFromGit/Source/GLTFExporter/Private/Converters/GLTFImageConverters.cpp
--- a/Engine/Plugins/Enterprise/GLTFExporter/MigratedFromGit/Source/GLTFExporter/Private/Converters/GLTFImageConverters.cpp
+++ b/Engine/Plugins/Enterprise/GLTFExporter/MigratedFromGit/Source/GLTFExporter/Private/Converters/GLTFImageConverters.cpp
@@ -5,7 +5,7 @@
 #include "Builders/GLTFConvertBuilder.h"
 #include "Misc/FileHelper.h"
 
-FGLTFJsonImageIndex FGLTFImageConverter::Convert(TGLTFSuperfluous<FString> Name, EGLTFTextureType Type, bool bIgnoreAlpha, FIntPoint Size, TGLTFSharedArray<FColor> Pixels)
+FGLTFJsonImage* FGLTFImageConverter::Convert(TGLTFSuperfluous<FString> Name, EGLTFTextureType Type, bool bIgnoreAlpha, FIntPoint Size, TGLTFSharedArray<FColor> Pixels)
 {
 	TArray64<uint8> CompressedData;
 
@@ -13,7 +13,7 @@ FGLTFJsonImageIndex FGLTFImageConverter::Convert(TGLTFSuperfluous<FString> Name,
 	switch (MimeType)
 	{
 		case EGLTFJsonMimeType::None:
-			return FGLTFJsonImageIndex(INDEX_NONE);
+			return nullptr;
 
 		case EGLTFJsonMimeType::PNG:
 			FGLTFImageUtility::CompressToPNG(Pixels->GetData(), Size, CompressedData);
diff --git a/Engine/Plugins/Enterprise/GLTFExporter/MigratedFromGit/Source/GLTFExporter/Private/Converters/GLTFImageConverters.h b/Engine/Plugins/Enterprise/GLTFExporter/MigratedFromGit/Source/GLTFExporter/Private/Converters/GLTFImageConverters.h
--- a/Engine/Plugins/Enterprise/GLTFExporter/MigratedFromGit/Source/GLTFExporter/Private/Converters/GLTFImageConverters.h
+++ b/Engine/Plugins/Enterprise/GLTFExporter/MigratedFromGit/Source/GLTFExporter/Private/Converters/GLTFImageConverters.h
@@ -9,13 +9,13 @@
 #include "Converters/GLTFSharedArray.h"
 #include "Options/GLTFExportOptions.h"
 
-typedef TGLTFConverter<FGLTFJsonImageIndex, TGLTFSuperfluous<FString>, EGLTFTextureType, bool, FIntPoint, TGLTFSharedArray<FColor>> IGLTFImageConverter;
+typedef TGLTFConverter<FGLTFJsonImage*, TGLTFSuperfluous<FString>, EGLTFTextureType, bool, FIntPoint, TGLTFSharedArray<FColor>> IGLTFImageConverter;
 
 class FGLTFImageConverter final : public FGLTFBuilderContext, public IGLTFImageConverter
 {
 	using FGLTFBuilderContext::FGLTFBuilderContext;
 
-	virtual FGLTFJsonImageIndex Convert(TGLTFSuperfluous<FString> Name, EGLTFTextureType Type, bool bIgnoreAlpha, FIntPoint Size, TGLTFSharedArray<FColor> Pixels) override;
+	virtual FGLTFJsonImage* Convert(TGLTFSuperfluous<FString> Name, EGLTFTextureType Type, bool bIgnoreAlpha, FIntPoint Size, TGLTFSharedArray<FColor> Pixels) override;
 
 	EGLTFJsonMimeType GetMimeType(const FColor* Pixels, FIntPoint Size, bool bIgnoreAlpha, EGLTFTextureType Type) const;
 
diff --git a/Engine/Plugins/Enterprise/GLTFExporter/MigratedFromGit/Source/GLTFExporter/Private/Converters/GLTFKhrVariantConverters.cpp b/Engine/Plugins/Enterprise/GLTFExporter/MigratedFromGit/Source/GLTFExporter/Private/Converters/GLTFKhrVariantConverters.cpp
--- a/Engine/Plugins/Enterprise/GLTFExporter/MigratedFromGit/Source/GLTFExporter/Private/Converters/GLTFKhrVariantConverters.cpp
+++ b/Engine/Plugins/Enterprise/GLTFExporter/MigratedFromGit/Source/GLTFExporter/Private/Converters/GLTFKhrVariantConverters.cpp
@@ -8,11 +8,11 @@
 #include "PropertyValue.h"
 #include "Variant.h"
 
-FGLTFJsonKhrMaterialVariantIndex FGLTFKhrMaterialVariantConverter::Convert(const UVariant* Variant)
+FGLTFJsonKhrMaterialVariant* FGLTFKhrMaterialVariantConverter::Convert(const UVariant* Variant)
 {
 	if (Variant == nullptr || Builder.ExportOptions->ExportMaterialVariants == EGLTFMaterialVariantMode::None)
 	{
-		return FGLTFJsonKhrMaterialVariantIndex(INDEX_NONE);
+		return nullptr;
 	}
 
 	FGLTFJsonKhrMaterialVariant MaterialVariant;
@@ -22,7 +22,7 @@ FGLTFJsonKhrMaterialVariantIndex FGLTFKhrMaterialVariantConverter::Convert(const
 	// cause confusion when trying to select the correct variant in a viewer.
 	MaterialVariant.Name = Variant->GetDisplayText().ToString();
 
-	typedef TTuple<FGLTFJsonPrimitive*, FGLTFJsonMaterialIndex> TPrimitiveMaterial;
+	typedef TTuple<FGLTFJsonPrimitive*, FGLTFJsonMaterial*> TPrimitiveMaterial;
 	TArray<TPrimitiveMaterial> PrimitiveMaterials;
 
 	for (const UVariantObjectBinding* Binding: Variant->GetBindings())
@@ -37,7 +37,7 @@ FGLTFJsonKhrMaterialVariantIndex FGLTFKhrMaterialVariantConverter::Convert(const
 			if (const UPropertyValueMaterial* MaterialProperty = Cast<UPropertyValueMaterial>(Property))
 			{
 				FGLTFJsonPrimitive* Primitive = nullptr;
-				FGLTFJsonMaterialIndex MaterialIndex;
+				FGLTFJsonMaterial* MaterialIndex;
 
 				if (TryParseMaterialProperty(Primitive, MaterialIndex, MaterialProperty))
 				{
@@ -51,15 +51,15 @@ FGLTFJsonKhrMaterialVariantIndex FGLTFKhrMaterialVariantConverter::Convert(const
 	{
 		// TODO: add warning and / or allow unused material variants to be added?
 
-		return FGLTFJsonKhrMaterialVariantIndex(INDEX_NONE);
+		return nullptr;
 	}
 
-	const FGLTFJsonKhrMaterialVariantIndex MaterialVariantIndex = Builder.AddKhrMaterialVariant(MaterialVariant);
+	FGLTFJsonKhrMaterialVariant* MaterialVariantIndex = Builder.AddKhrMaterialVariant(MaterialVariant);
 
 	for (const TPrimitiveMaterial& PrimitiveMaterial: PrimitiveMaterials)
 	{
 		FGLTFJsonPrimitive* Primitive = PrimitiveMaterial.Key;
-		const FGLTFJsonMaterialIndex MaterialIndex = PrimitiveMaterial.Value;
+		FGLTFJsonMaterial* MaterialIndex = PrimitiveMaterial.Value;
 
 		FGLTFJsonKhrMaterialVariantMapping* ExistingMapping = Primitive->KhrMaterialVariantMappings.FindByPredicate(
 			[MaterialIndex](const FGLTFJsonKhrMaterialVariantMapping& Mapping)
@@ -84,7 +84,7 @@ FGLTFJsonKhrMaterialVariantIndex FGLTFKhrMaterialVariantConverter::Convert(const
 	return MaterialVariantIndex;
 }
 
-bool FGLTFKhrMaterialVariantConverter::TryParseMaterialProperty(FGLTFJsonPrimitive*& OutPrimitive, FGLTFJsonMaterialIndex& OutMaterialIndex, const UPropertyValueMaterial* Property) const
+bool FGLTFKhrMaterialVariantConverter::TryParseMaterialProperty(FGLTFJsonPrimitive*& OutPrimitive, FGLTFJsonMaterial*& OutMaterialIndex, const UPropertyValueMaterial* Property) const
 {
 	const UMeshComponent* Target = static_cast<UMeshComponent*>(Property->GetPropertyParentContainerAddress());
 	if (Target == nullptr)
@@ -142,7 +142,7 @@ bool FGLTFKhrMaterialVariantConverter::TryParseMaterialProperty(FGLTFJsonPrimiti
 	Builder.RegisterObjectVariant(Owner, Property); // TODO: we don't need to register this on the actor
 
 	const int32 MaterialIndex = CapturedPropSegments[NumPropSegments - 1].PropertyIndex;
-	const FGLTFJsonMeshIndex MeshIndex = Builder.GetOrAddMesh(Target);
+	FGLTFJsonMesh* MeshIndex = Builder.GetOrAddMesh(Target);
 
 	OutPrimitive = &Builder.GetMesh(MeshIndex).Primitives[MaterialIndex];
 	OutMaterialIndex = FGLTFVariantUtility::GetOrAddMaterial(Builder, Material, Target, MaterialIndex);
diff --git a/Engine/Plugins/Enterprise/GLTFExporter/MigratedFromGit/Source/GLTFExporter/Private/Converters/GLTFKhrVariantConverters.h b/Engine/Plugins/Enterprise/GLTFExporter/MigratedFromGit/Source/GLTFExporter/Private/Converters/GLTFKhrVariantConverters.h
--- a/Engine/Plugins/Enterprise/GLTFExporter/MigratedFromGit/Source/GLTFExporter/Private/Converters/GLTFKhrVariantConverters.h
+++ b/Engine/Plugins/Enterprise/GLTFExporter/MigratedFromGit/Source/GLTFExporter/Private/Converters/GLTFKhrVariantConverters.h
@@ -10,13 +10,13 @@
 class UVariant;
 class UPropertyValueMaterial;
 
-typedef TGLTFConverter<FGLTFJsonKhrMaterialVariantIndex, const UVariant*> IGLTFKhrMaterialVariantConverter;
+typedef TGLTFConverter<FGLTFJsonKhrMaterialVariant*, const UVariant*> IGLTFKhrMaterialVariantConverter;
 
 class FGLTFKhrMaterialVariantConverter final : public FGLTFBuilderContext, public IGLTFKhrMaterialVariantConverter
 {
 	using FGLTFBuilderContext::FGLTFBuilderContext;
 
-	virtual FGLTFJsonKhrMaterialVariantIndex Convert(const UVariant* Variant) override;
+	virtual FGLTFJsonKhrMaterialVariant* Convert(const UVariant* Variant) override;
 
-	bool TryParseMaterialProperty(FGLTFJsonPrimitive*& OutPrimitive, FGLTFJsonMaterialIndex& OutMaterialIndex, const UPropertyValueMaterial* Property) const;
+	bool TryParseMaterialProperty(FGLTFJsonPrimitive*& OutPrimitive, FGLTFJsonMaterial*& OutMaterialIndex, const UPropertyValueMaterial* Property) const;
 };
diff --git a/Engine/Plugins/Enterprise/GLTFExporter/MigratedFromGit/Source/GLTFExporter/Private/Converters/GLTFLightConverters.cpp b/Engine/Plugins/Enterprise/GLTFExporter/MigratedFromGit/Source/GLTFExporter/Private/Converters/GLTFLightConverters.cpp
--- a/Engine/Plugins/Enterprise/GLTFExporter/MigratedFromGit/Source/GLTFExporter/Private/Converters/GLTFLightConverters.cpp
+++ b/Engine/Plugins/Enterprise/GLTFExporter/MigratedFromGit/Source/GLTFExporter/Private/Converters/GLTFLightConverters.cpp
@@ -8,7 +8,7 @@
 #include "Components/PointLightComponent.h"
 #include "Components/SpotLightComponent.h"
 
-FGLTFJsonLightIndex FGLTFLightConverter::Convert(const ULightComponent* LightComponent)
+FGLTFJsonLight* FGLTFLightConverter::Convert(const ULightComponent* LightComponent)
 {
 	FGLTFJsonLight Light;
 
@@ -18,7 +18,7 @@ FGLTFJsonLightIndex FGLTFLightConverter::Convert(const ULightComponent* LightCom
 	if (Light.Type == EGLTFJsonLightType::None)
 	{
 		// TODO: report error (unsupported light component type)
-		return FGLTFJsonLightIndex(INDEX_NONE);
+		return nullptr;
 	}
 
 	Light.Intensity = LightComponent->Intensity;
diff --git a/Engine/Plugins/Enterprise/GLTFExporter/MigratedFromGit/Source/GLTFExporter/Private/Converters/GLTFLightConverters.h b/Engine/Plugins/Enterprise/GLTFExporter/MigratedFromGit/Source/GLTFExporter/Private/Converters/GLTFLightConverters.h
--- a/Engine/Plugins/Enterprise/GLTFExporter/MigratedFromGit/Source/GLTFExporter/Private/Converters/GLTFLightConverters.h
+++ b/Engine/Plugins/Enterprise/GLTFExporter/MigratedFromGit/Source/GLTFExporter/Private/Converters/GLTFLightConverters.h
@@ -6,11 +6,11 @@
 #include "Converters/GLTFConverter.h"
 #include "Converters/GLTFBuilderContext.h"
 
-typedef TGLTFConverter<FGLTFJsonLightIndex, const ULightComponent*> IGLTFLightConverter;
+typedef TGLTFConverter<FGLTFJsonLight*, const ULightComponent*> IGLTFLightConverter;
 
 class FGLTFLightConverter final : public FGLTFBuilderContext, public IGLTFLightConverter
 {
 	using FGLTFBuilderContext::FGLTFBuilderContext;
 
-	virtual FGLTFJsonLightIndex Convert(const ULightComponent* LightComponent) override;
+	virtual FGLTFJsonLight* Convert(const ULightComponent* LightComponent) override;
 };
diff --git a/Engine/Plugins/Enterprise/GLTFExporter/MigratedFromGit/Source/GLTFExporter/Private/Converters/GLTFLightMapConverters.cpp b/Engine/Plugins/Enterprise/GLTFExporter/MigratedFromGit/Source/GLTFExporter/Private/Converters/GLTFLightMapConverters.cpp
--- a/Engine/Plugins/Enterprise/GLTFExporter/MigratedFromGit/Source/GLTFExporter/Private/Converters/GLTFLightMapConverters.cpp
+++ b/Engine/Plugins/Enterprise/GLTFExporter/MigratedFromGit/Source/GLTFExporter/Private/Converters/GLTFLightMapConverters.cpp
@@ -7,18 +7,18 @@
 #include "Engine/StaticMesh.h"
 #include "Engine/MapBuildDataRegistry.h"
 
-FGLTFJsonLightMapIndex FGLTFLightMapConverter::Convert(const UStaticMeshComponent* StaticMeshComponent)
+FGLTFJsonLightMap* FGLTFLightMapConverter::Convert(const UStaticMeshComponent* StaticMeshComponent)
 {
 	if (Builder.ExportOptions->TextureImageFormat == EGLTFTextureImageFormat::None)
 	{
-		return FGLTFJsonLightMapIndex(INDEX_NONE);
+		return nullptr;
 	}
 
 	const UStaticMesh* StaticMesh = StaticMeshComponent->GetStaticMesh();
 
 	if (StaticMesh == nullptr)
 	{
-		return FGLTFJsonLightMapIndex(INDEX_NONE);
+		return nullptr;
 	}
 
 	const int32 LODIndex = FGLTFMeshUtility::GetLOD(StaticMesh, StaticMeshComponent, Builder.ExportOptions->DefaultLevelOfDetail);
@@ -32,14 +32,14 @@ FGLTFJsonLightMapIndex FGLTFLightMapConverter::Convert(const UStaticMeshComponen
 
 	if (CoordinateIndex < 0 || CoordinateIndex >= LODResources.GetNumTexCoords())
 	{
-		return FGLTFJsonLightMapIndex(INDEX_NONE);
+		return nullptr;
 	}
 
 	constexpr int32 LightMapLODIndex = 0; // TODO: why is this zero?
 
 	if (!StaticMeshComponent->LODData.IsValidIndex(LightMapLODIndex))
 	{
-		return FGLTFJsonLightMapIndex(INDEX_NONE);
+		return nullptr;
 	}
 
 	const FStaticMeshComponentLODInfo& ComponentLODInfo = StaticMeshComponent->LODData[LightMapLODIndex];
@@ -47,23 +47,23 @@ FGLTFJsonLightMapIndex FGLTFLightMapConverter::Convert(const UStaticMeshComponen
 
 	if (MeshMapBuildData == nullptr || MeshMapBuildData->LightMap == nullptr)
 	{
-		return FGLTFJsonLightMapIndex(INDEX_NONE);
+		return nullptr;
 	}
 
 	const FLightMap2D* LightMap2D = MeshMapBuildData->LightMap->GetLightMap2D();
 
 	if (LightMap2D == nullptr)
 	{
-		return FGLTFJsonLightMapIndex(INDEX_NONE);
+		return nullptr;
 	}
 
 	const FLightMapInteraction LightMapInteraction = LightMap2D->GetInteraction(GMaxRHIFeatureLevel);
 	const ULightMapTexture2D* Texture = LightMapInteraction.GetTexture(true);
-	const FGLTFJsonTextureIndex TextureIndex = Builder.GetOrAddTexture(Texture);
+	FGLTFJsonTexture* TextureIndex = Builder.GetOrAddTexture(Texture);
 
-	if (TextureIndex == INDEX_NONE)
+	if (TextureIndex == nullptr)
 	{
-		return FGLTFJsonLightMapIndex(INDEX_NONE);
+		return nullptr;
 	}
 
 	const FVector2D& CoordinateBias = LightMap2D->GetCoordinateBias();
diff --git a/Engine/Plugins/Enterprise/GLTFExporter/MigratedFromGit/Source/GLTFExporter/Private/Converters/GLTFLightMapConverters.h b/Engine/Plugins/Enterprise/GLTFExporter/MigratedFromGit/Source/GLTFExporter/Private/Converters/GLTFLightMapConverters.h
--- a/Engine/Plugins/Enterprise/GLTFExporter/MigratedFromGit/Source/GLTFExporter/Private/Converters/GLTFLightMapConverters.h
+++ b/Engine/Plugins/Enterprise/GLTFExporter/MigratedFromGit/Source/GLTFExporter/Private/Converters/GLTFLightMapConverters.h
@@ -6,11 +6,11 @@
 #include "Converters/GLTFConverter.h"
 #include "Converters/GLTFBuilderContext.h"
 
-typedef TGLTFConverter<FGLTFJsonLightMapIndex, const UStaticMeshComponent*> IGLTFLightMapConverter;
+typedef TGLTFConverter<FGLTFJsonLightMap*, const UStaticMeshComponent*> IGLTFLightMapConverter;
 
 class FGLTFLightMapConverter final : public FGLTFBuilderContext, public IGLTFLightMapConverter
 {
 	using FGLTFBuilderContext::FGLTFBuilderContext;
 
-	virtual FGLTFJsonLightMapIndex Convert(const UStaticMeshComponent* StaticMeshComponent) override;
+	virtual FGLTFJsonLightMap* Convert(const UStaticMeshComponent* StaticMeshComponent) override;
 };
diff --git a/Engine/Plugins/Enterprise/GLTFExporter/MigratedFromGit/Source/GLTFExporter/Private/Converters/GLTFMaterialConverters.cpp b/Engine/Plugins/Enterprise/GLTFExporter/MigratedFromGit/Source/GLTFExporter/Private/Converters/GLTFMaterialConverters.cpp
--- a/Engine/Plugins/Enterprise/GLTFExporter/MigratedFromGit/Source/GLTFExporter/Private/Converters/GLTFMaterialConverters.cpp
+++ b/Engine/Plugins/Enterprise/GLTFExporter/MigratedFromGit/Source/GLTFExporter/Private/Converters/GLTFMaterialConverters.cpp
@@ -43,14 +43,14 @@ void FGLTFMaterialConverter::Sanitize(const UMaterialInterface*& Material, const
 #endif
 }
 
-FGLTFJsonMaterialIndex FGLTFMaterialConverter::Convert(const UMaterialInterface* Material, const FGLTFMeshData* MeshData, FGLTFIndexArray SectionIndices)
+FGLTFJsonMaterial* FGLTFMaterialConverter::Convert(const UMaterialInterface* Material, const FGLTFMeshData* MeshData, FGLTFIndexArray SectionIndices)
 {
 	if (Material != FGLTFMaterialUtility::GetDefaultMaterial())
 	{
 		FGLTFJsonMaterial* JsonMaterial = Builder.AddMaterial();
 		Builder.SetupTask<FGLTFMaterialTask>(Builder, UVOverlapChecker, Material, MeshData, SectionIndices, JsonMaterial);
-		return JsonMaterial->Index;
+		return JsonMaterial;
 	}
 
-	return FGLTFJsonMaterialIndex(INDEX_NONE); // use default gltf definition
+	return nullptr; // use default gltf definition
 }
diff --git a/Engine/Plugins/Enterprise/GLTFExporter/MigratedFromGit/Source/GLTFExporter/Private/Converters/GLTFMaterialConverters.h b/Engine/Plugins/Enterprise/GLTFExporter/MigratedFromGit/Source/GLTFExporter/Private/Converters/GLTFMaterialConverters.h
--- a/Engine/Plugins/Enterprise/GLTFExporter/MigratedFromGit/Source/GLTFExporter/Private/Converters/GLTFMaterialConverters.h
+++ b/Engine/Plugins/Enterprise/GLTFExporter/MigratedFromGit/Source/GLTFExporter/Private/Converters/GLTFMaterialConverters.h
@@ -10,7 +10,7 @@
 #include "Converters/GLTFUVOverlapChecker.h"
 #include "Converters/GLTFUVDegenerateChecker.h"
 
-typedef TGLTFConverter<FGLTFJsonMaterialIndex, const UMaterialInterface*, const FGLTFMeshData*, FGLTFIndexArray> IGLTFMaterialConverter;
+typedef TGLTFConverter<FGLTFJsonMaterial*, const UMaterialInterface*, const FGLTFMeshData*, FGLTFIndexArray> IGLTFMaterialConverter;
 
 class FGLTFMaterialConverter final : public FGLTFBuilderContext, public IGLTFMaterialConverter
 {
@@ -18,7 +18,7 @@ class FGLTFMaterialConverter final : public FGLTFBuilderContext, public IGLTFMat
 
 	virtual void Sanitize(const UMaterialInterface*& Material, const FGLTFMeshData*& MeshData, FGLTFIndexArray& SectionIndices) override;
 
-	virtual FGLTFJsonMaterialIndex Convert(const UMaterialInterface* Material, const FGLTFMeshData* MeshData, FGLTFIndexArray SectionIndices) override;
+	virtual FGLTFJsonMaterial* Convert(const UMaterialInterface* Material, const FGLTFMeshData* MeshData, FGLTFIndexArray SectionIndices) override;
 
 	FGLTFUVOverlapChecker UVOverlapChecker;
 	FGLTFUVDegenerateChecker UVDegenerateChecker;
diff --git a/Engine/Plugins/Enterprise/GLTFExporter/MigratedFromGit/Source/GLTFExporter/Private/Converters/GLTFMaterialProxyFactory.cpp b/Engine/Plugins/Enterprise/GLTFExporter/MigratedFromGit/Source/GLTFExporter/Private/Converters/GLTFMaterialProxyFactory.cpp
--- a/Engine/Plugins/Enterprise/GLTFExporter/MigratedFromGit/Source/GLTFExporter/Private/Converters/GLTFMaterialProxyFactory.cpp
+++ b/Engine/Plugins/Enterprise/GLTFExporter/MigratedFromGit/Source/GLTFExporter/Private/Converters/GLTFMaterialProxyFactory.cpp
@@ -32,8 +32,8 @@ UMaterialInterface* FGLTFMaterialProxyFactory::Create(UMaterialInterface* Origin
 			*OriginalMaterial->GetName()));
 	}
 
-	const FGLTFJsonMaterialIndex MaterialIndex = Builder.GetOrAddMaterial(OriginalMaterial);
-	if (MaterialIndex == INDEX_NONE)
+	FGLTFJsonMaterial* MaterialIndex = Builder.GetOrAddMaterial(OriginalMaterial);
+	if (MaterialIndex == nullptr)
 	{
 		// TODO: report error
 		return nullptr;
@@ -166,9 +166,9 @@ void FGLTFMaterialProxyFactory::SetProxyParameter(UMaterialInstanceConstant* Pro
 	ParameterInfo.UVRotation.Set(ProxyMaterial, TextureInfo.Transform.Rotation, true);
 }
 
-UTexture2D* FGLTFMaterialProxyFactory::FindOrCreateTexture(FGLTFJsonTextureIndex Index, const FGLTFProxyMaterialTextureParameterInfo& ParameterInfo)
+UTexture2D* FGLTFMaterialProxyFactory::FindOrCreateTexture(FGLTFJsonTexture* Index, const FGLTFProxyMaterialTextureParameterInfo& ParameterInfo)
 {
-	if (Index == INDEX_NONE)
+	if (Index == nullptr)
 	{
 		return nullptr;
 	}
@@ -255,9 +255,9 @@ TUniquePtr<IGLTFTexture2DConverter> FGLTFMaterialProxyFactory::CreateTextureConv
 			bToSRGB = false; // ignore
 		}
 
-		virtual FGLTFJsonTextureIndex Convert(const UTexture2D* Texture2D, bool bToSRGB) override
+		virtual FGLTFJsonTexture* Convert(const UTexture2D* Texture2D, bool bToSRGB) override
 		{
-			const FGLTFJsonTextureIndex TextureIndex = Factory.Builder.AddTexture({});
+			FGLTFJsonTexture* TextureIndex = Factory.Builder.AddTexture();
 			Factory.Textures.Add(TextureIndex, const_cast<UTexture2D*>(Texture2D));
 			return TextureIndex;
 		}
@@ -282,12 +282,12 @@ TUniquePtr<IGLTFImageConverter> FGLTFMaterialProxyFactory::CreateImageConverter(
 
 	protected:
 
-		virtual FGLTFJsonImageIndex Convert(TGLTFSuperfluous<FString> Name, EGLTFTextureType Type, bool bIgnoreAlpha, FIntPoint Size, TGLTFSharedArray<FColor> Pixels) override
+		virtual FGLTFJsonImage* Convert(TGLTFSuperfluous<FString> Name, EGLTFTextureType Type, bool bIgnoreAlpha, FIntPoint Size, TGLTFSharedArray<FColor> Pixels) override
 		{
 			const FString Filename = FGLTFImageUtility::GetUniqueFilename(Name, TEXT(""), UniqueFilenames);
 			UniqueFilenames.Add(Filename);
 
-			const FGLTFJsonImageIndex ImageIndex = Factory.Builder.AddImage({});
+			FGLTFJsonImage* ImageIndex = Factory.Builder.AddImage();
 			Factory.Images.Add(ImageIndex, { Filename, Type, bIgnoreAlpha, Size, Pixels });
 			return ImageIndex;
 		}
diff --git a/Engine/Plugins/Enterprise/GLTFExporter/MigratedFromGit/Source/GLTFExporter/Private/Converters/GLTFMaterialProxyFactory.h b/Engine/Plugins/Enterprise/GLTFExporter/MigratedFromGit/Source/GLTFExporter/Private/Converters/GLTFMaterialProxyFactory.h
--- a/Engine/Plugins/Enterprise/GLTFExporter/MigratedFromGit/Source/GLTFExporter/Private/Converters/GLTFMaterialProxyFactory.h
+++ b/Engine/Plugins/Enterprise/GLTFExporter/MigratedFromGit/Source/GLTFExporter/Private/Converters/GLTFMaterialProxyFactory.h
@@ -42,7 +42,7 @@ private:
 	void SetProxyParameter(UMaterialInstanceConstant* ProxyMaterial, const TGLTFProxyMaterialParameterInfo<FLinearColor>& ParameterInfo, const FGLTFJsonColor4& Color);
 	void SetProxyParameter(UMaterialInstanceConstant* ProxyMaterial, const FGLTFProxyMaterialTextureParameterInfo& ParameterInfo, const FGLTFJsonTextureInfo& TextureInfo);
 
-	UTexture2D* FindOrCreateTexture(FGLTFJsonTextureIndex Index, const FGLTFProxyMaterialTextureParameterInfo& ParameterInfo);
+	UTexture2D* FindOrCreateTexture(FGLTFJsonTexture* Index, const FGLTFProxyMaterialTextureParameterInfo& ParameterInfo);
 	UTexture2D* CreateTexture(const FGLTFImageData* ImageData, const FGLTFJsonSampler& JsonSampler, const FGLTFProxyMaterialTextureParameterInfo& ParameterInfo);
 
 	UMaterialInstanceConstant* CreateInstancedMaterial(UMaterialInterface* OriginalMaterial, EGLTFJsonShadingModel ShadingModel);
@@ -61,8 +61,8 @@ private:
 
 	FGLTFConvertBuilder Builder;
 
-	TMap<FGLTFJsonTextureIndex, UTexture2D*> Textures;
-	TMap<FGLTFJsonImageIndex, FGLTFImageData> Images;
+	TMap<FGLTFJsonTexture*, UTexture2D*> Textures;
+	TMap<FGLTFJsonImage*, FGLTFImageData> Images;
 };
 
 #endif
diff --git a/Engine/Plugins/Enterprise/GLTFExporter/MigratedFromGit/Source/GLTFExporter/Private/Converters/GLTFMaterialUtility.cpp b/Engine/Plugins/Enterprise/GLTFExporter/MigratedFromGit/Source/GLTFExporter/Private/Converters/GLTFMaterialUtility.cpp
--- a/Engine/Plugins/Enterprise/GLTFExporter/MigratedFromGit/Source/GLTFExporter/Private/Converters/GLTFMaterialUtility.cpp
+++ b/Engine/Plugins/Enterprise/GLTFExporter/MigratedFromGit/Source/GLTFExporter/Private/Converters/GLTFMaterialUtility.cpp
@@ -262,7 +262,7 @@ FGLTFPropertyBakeOutput FGLTFMaterialUtility::BakeMaterialProperty(const FIntPoi
 	return PropertyBakeOutput;
 }
 
-FGLTFJsonTextureIndex FGLTFMaterialUtility::AddTexture(FGLTFConvertBuilder& Builder, TGLTFSharedArray<FColor>& Pixels, const FIntPoint& TextureSize, bool bIgnoreAlpha, bool bIsNormalMap, const FString& TextureName, EGLTFJsonTextureFilter MinFilter, EGLTFJsonTextureFilter MagFilter, EGLTFJsonTextureWrap WrapS, EGLTFJsonTextureWrap WrapT)
+FGLTFJsonTexture* FGLTFMaterialUtility::AddTexture(FGLTFConvertBuilder& Builder, TGLTFSharedArray<FColor>& Pixels, const FIntPoint& TextureSize, bool bIgnoreAlpha, bool bIsNormalMap, const FString& TextureName, EGLTFJsonTextureFilter MinFilter, EGLTFJsonTextureFilter MagFilter, EGLTFJsonTextureWrap WrapS, EGLTFJsonTextureWrap WrapT)
 {
 	// TODO: maybe we should reuse existing samplers?
 	FGLTFJsonSampler JsonSampler;
diff --git a/Engine/Plugins/Enterprise/GLTFExporter/MigratedFromGit/Source/GLTFExporter/Private/Converters/GLTFMaterialUtility.h b/Engine/Plugins/Enterprise/GLTFExporter/MigratedFromGit/Source/GLTFExporter/Private/Converters/GLTFMaterialUtility.h
--- a/Engine/Plugins/Enterprise/GLTFExporter/MigratedFromGit/Source/GLTFExporter/Private/Converters/GLTFMaterialUtility.h
+++ b/Engine/Plugins/Enterprise/GLTFExporter/MigratedFromGit/Source/GLTFExporter/Private/Converters/GLTFMaterialUtility.h
@@ -54,7 +54,7 @@ struct FGLTFMaterialUtility
 
 	static FGLTFPropertyBakeOutput BakeMaterialProperty(const FIntPoint& OutputSize, const FMaterialPropertyEx& Property, const UMaterialInterface* Material, int32 TexCoord, const FGLTFMeshData* MeshData = nullptr, const FGLTFIndexArray& MeshSectionIndices = {}, bool bFillAlpha = true, bool bAdjustNormalmaps = true);
 
-	static FGLTFJsonTextureIndex AddTexture(FGLTFConvertBuilder& Builder, TGLTFSharedArray<FColor>& Pixels, const FIntPoint& TextureSize, bool bIgnoreAlpha, bool bIsNormalMap, const FString& TextureName, EGLTFJsonTextureFilter MinFilter, EGLTFJsonTextureFilter MagFilter, EGLTFJsonTextureWrap WrapS, EGLTFJsonTextureWrap WrapT);
+	static FGLTFJsonTexture* AddTexture(FGLTFConvertBuilder& Builder, TGLTFSharedArray<FColor>& Pixels, const FIntPoint& TextureSize, bool bIgnoreAlpha, bool bIsNormalMap, const FString& TextureName, EGLTFJsonTextureFilter MinFilter, EGLTFJsonTextureFilter MagFilter, EGLTFJsonTextureWrap WrapS, EGLTFJsonTextureWrap WrapT);
 
 	static FLinearColor GetMask(const FExpressionInput& ExpressionInput);
 	static uint32 GetMaskComponentCount(const FExpressionInput& ExpressionInput);
diff --git a/Engine/Plugins/Enterprise/GLTFExporter/MigratedFromGit/Source/GLTFExporter/Private/Converters/GLTFMeshConverters.cpp b/Engine/Plugins/Enterprise/GLTFExporter/MigratedFromGit/Source/GLTFExporter/Private/Converters/GLTFMeshConverters.cpp
--- a/Engine/Plugins/Enterprise/GLTFExporter/MigratedFromGit/Source/GLTFExporter/Private/Converters/GLTFMeshConverters.cpp
+++ b/Engine/Plugins/Enterprise/GLTFExporter/MigratedFromGit/Source/GLTFExporter/Private/Converters/GLTFMeshConverters.cpp
@@ -33,7 +33,7 @@ void FGLTFStaticMeshConverter::Sanitize(const UStaticMesh*& StaticMesh, const US
 	}
 }
 
-FGLTFJsonMeshIndex FGLTFStaticMeshConverter::Convert(const UStaticMesh* StaticMesh,  const UStaticMeshComponent* StaticMeshComponent, FGLTFMaterialArray Materials, int32 LODIndex)
+FGLTFJsonMesh* FGLTFStaticMeshConverter::Convert(const UStaticMesh* StaticMesh,  const UStaticMeshComponent* StaticMeshComponent, FGLTFMaterialArray Materials, int32 LODIndex)
 {
 #if !WITH_EDITOR
 	if (!StaticMesh->bAllowCPUAccess)
@@ -49,7 +49,7 @@ FGLTFJsonMeshIndex FGLTFStaticMeshConverter::Convert(const UStaticMesh* StaticMe
 	JsonMesh->Primitives.AddDefaulted(MaterialCount);
 
 	Builder.SetupTask<FGLTFStaticMeshTask>(Builder, MeshSectionConverter, StaticMesh, StaticMeshComponent, Materials, LODIndex, JsonMesh);
-	return JsonMesh->Index;
+	return JsonMesh;
 }
 
 void FGLTFSkeletalMeshConverter::Sanitize(const USkeletalMesh*& SkeletalMesh, const USkeletalMeshComponent*& SkeletalMeshComponent, FGLTFMaterialArray& Materials, int32& LODIndex)
@@ -75,7 +75,7 @@ void FGLTFSkeletalMeshConverter::Sanitize(const USkeletalMesh*& SkeletalMesh, co
 	}
 }
 
-FGLTFJsonMeshIndex FGLTFSkeletalMeshConverter::Convert(const USkeletalMesh* SkeletalMesh, const USkeletalMeshComponent* SkeletalMeshComponent, FGLTFMaterialArray Materials, int32 LODIndex)
+FGLTFJsonMesh* FGLTFSkeletalMeshConverter::Convert(const USkeletalMesh* SkeletalMesh, const USkeletalMeshComponent* SkeletalMeshComponent, FGLTFMaterialArray Materials, int32 LODIndex)
 {
 #if !WITH_EDITOR
 	if (!SkeletalMesh->GetLODInfo(LODIndex)->bAllowCPUAccess)
@@ -92,5 +92,5 @@ FGLTFJsonMeshIndex FGLTFSkeletalMeshConverter::Convert(const USkeletalMesh* Skel
 	JsonMesh->Primitives.AddDefaulted(MaterialCount);
 
 	Builder.SetupTask<FGLTFSkeletalMeshTask>(Builder, MeshSectionConverter, SkeletalMesh, SkeletalMeshComponent, Materials, LODIndex, JsonMesh);
-	return JsonMesh->Index;
+	return JsonMesh;
 }
diff --git a/Engine/Plugins/Enterprise/GLTFExporter/MigratedFromGit/Source/GLTFExporter/Private/Converters/GLTFMeshConverters.h b/Engine/Plugins/Enterprise/GLTFExporter/MigratedFromGit/Source/GLTFExporter/Private/Converters/GLTFMeshConverters.h
--- a/Engine/Plugins/Enterprise/GLTFExporter/MigratedFromGit/Source/GLTFExporter/Private/Converters/GLTFMeshConverters.h
+++ b/Engine/Plugins/Enterprise/GLTFExporter/MigratedFromGit/Source/GLTFExporter/Private/Converters/GLTFMeshConverters.h
@@ -9,8 +9,8 @@
 #include "Converters/GLTFMeshDataConverters.h"
 #include "Converters/GLTFMaterialArray.h"
 
-typedef TGLTFConverter<FGLTFJsonMeshIndex, const UStaticMesh*, const UStaticMeshComponent*, FGLTFMaterialArray, int32> IGLTFStaticMeshConverter;
-typedef TGLTFConverter<FGLTFJsonMeshIndex, const USkeletalMesh*, const USkeletalMeshComponent*, FGLTFMaterialArray, int32> IGLTFSkeletalMeshConverter;
+typedef TGLTFConverter<FGLTFJsonMesh*, const UStaticMesh*, const UStaticMeshComponent*, FGLTFMaterialArray, int32> IGLTFStaticMeshConverter;
+typedef TGLTFConverter<FGLTFJsonMesh*, const USkeletalMesh*, const USkeletalMeshComponent*, FGLTFMaterialArray, int32> IGLTFSkeletalMeshConverter;
 
 class FGLTFStaticMeshConverter final : public FGLTFBuilderContext, public IGLTFStaticMeshConverter
 {
@@ -18,7 +18,7 @@ class FGLTFStaticMeshConverter final : public FGLTFBuilderContext, public IGLTFS
 
 	virtual void Sanitize(const UStaticMesh*& StaticMesh, const UStaticMeshComponent*& StaticMeshComponent, FGLTFMaterialArray& Materials, int32& LODIndex) override;
 
-	virtual FGLTFJsonMeshIndex Convert(const UStaticMesh* StaticMesh, const UStaticMeshComponent* StaticMeshComponent, FGLTFMaterialArray Materials, int32 LODIndex) override;
+	virtual FGLTFJsonMesh* Convert(const UStaticMesh* StaticMesh, const UStaticMeshComponent* StaticMeshComponent, FGLTFMaterialArray Materials, int32 LODIndex) override;
 
 	FGLTFStaticMeshSectionConverter MeshSectionConverter;
 };
@@ -29,7 +29,7 @@ class FGLTFSkeletalMeshConverter final : public FGLTFBuilderContext, public IGLT
 
 	virtual void Sanitize(const USkeletalMesh*& SkeletalMesh, const USkeletalMeshComponent*& SkeletalMeshComponent, FGLTFMaterialArray& Materials, int32& LODIndex) override;
 
-	virtual FGLTFJsonMeshIndex Convert(const USkeletalMesh* SkeletalMesh, const USkeletalMeshComponent* SkeletalMeshComponent, FGLTFMaterialArray Materials, int32 LODIndex) override;
+	virtual FGLTFJsonMesh* Convert(const USkeletalMesh* SkeletalMesh, const USkeletalMeshComponent* SkeletalMeshComponent, FGLTFMaterialArray Materials, int32 LODIndex) override;
 
 	FGLTFSkeletalMeshSectionConverter MeshSectionConverter;
 };
diff --git a/Engine/Plugins/Enterprise/GLTFExporter/MigratedFromGit/Source/GLTFExporter/Private/Converters/GLTFNodeConverters.cpp b/Engine/Plugins/Enterprise/GLTFExporter/MigratedFromGit/Source/GLTFExporter/Private/Converters/GLTFNodeConverters.cpp
--- a/Engine/Plugins/Enterprise/GLTFExporter/MigratedFromGit/Source/GLTFExporter/Private/Converters/GLTFNodeConverters.cpp
+++ b/Engine/Plugins/Enterprise/GLTFExporter/MigratedFromGit/Source/GLTFExporter/Private/Converters/GLTFNodeConverters.cpp
@@ -12,20 +12,20 @@
 #include "Engine/StaticMeshSocket.h"
 #include "Engine/SkeletalMeshSocket.h"
 
-FGLTFJsonNodeIndex FGLTFActorConverter::Convert(const AActor* Actor)
+FGLTFJsonNode* FGLTFActorConverter::Convert(const AActor* Actor)
 {
 	if (Actor->bIsEditorOnlyActor)
 	{
-		return FGLTFJsonNodeIndex(INDEX_NONE);
+		return nullptr;
 	}
 
 	if (!Builder.IsSelectedActor(Actor))
 	{
-		return FGLTFJsonNodeIndex(INDEX_NONE);
+		return nullptr;
 	}
 
 	const USceneComponent* RootComponent = Actor->GetRootComponent();
-	const FGLTFJsonNodeIndex RootNodeIndex = Builder.GetOrAddNode(RootComponent);
+	FGLTFJsonNode* RootNodeIndex = Builder.GetOrAddNode(RootComponent);
 
 	// TODO: process all components since any component can be attached to any other component in runtime
 
@@ -80,23 +80,23 @@ FGLTFJsonNodeIndex FGLTFActorConverter::Convert(const AActor* Actor)
 	return RootNodeIndex;
 }
 
-FGLTFJsonNodeIndex FGLTFComponentConverter::Convert(const USceneComponent* SceneComponent)
+FGLTFJsonNode* FGLTFComponentConverter::Convert(const USceneComponent* SceneComponent)
 {
 	if (SceneComponent->IsEditorOnly())
 	{
-		return FGLTFJsonNodeIndex(INDEX_NONE);
+		return nullptr;
 	}
 
 	const AActor* Owner = SceneComponent->GetOwner();
 	if (Owner == nullptr)
 	{
 		// TODO: report error (invalid scene component)
-		return FGLTFJsonNodeIndex(INDEX_NONE);
+		return nullptr;
 	}
 
 	if (!Builder.IsSelectedActor(Owner))
 	{
-		return FGLTFJsonNodeIndex(INDEX_NONE);
+		return nullptr;
 	}
 
 	const bool bIsRootComponent = Owner->GetRootComponent() == SceneComponent;
@@ -104,7 +104,7 @@ FGLTFJsonNodeIndex FGLTFComponentConverter::Convert(const USceneComponent* Scene
 
 	const USceneComponent* ParentComponent = !bIsRootNode ? SceneComponent->GetAttachParent() : nullptr;
 	const FName SocketName = SceneComponent->GetAttachSocketName();
-	const FGLTFJsonNodeIndex ParentNodeIndex = Builder.GetOrAddNode(ParentComponent, SocketName);
+	FGLTFJsonNode* ParentNodeIndex = Builder.GetOrAddNode(ParentComponent, SocketName);
 
 	if (ParentComponent != nullptr && !SceneComponent->IsUsingAbsoluteScale())
 	{
@@ -150,14 +150,12 @@ FGLTFJsonNodeIndex FGLTFComponentConverter::Convert(const USceneComponent* Scene
 			if (Builder.ExportOptions->bExportVertexSkinWeights)
 			{
 				// TODO: remove need for NodeIndex by adding support for cyclic calls in converter
-				const FGLTFJsonSkinIndex SkinIndex = Builder.GetOrAddSkin(Node->Index, SkeletalMeshComponent);
-				if (SkinIndex != INDEX_NONE)
+				Node->Skin = Builder.GetOrAddSkin(Node, SkeletalMeshComponent);
+				if (Node->Skin != nullptr)
 				{
-					Node->Skin = SkinIndex;
-
 					if (Builder.ExportOptions->bExportAnimationSequences)
 					{
-						Builder.GetOrAddAnimation(Node->Index, SkeletalMeshComponent);
+						Builder.GetOrAddAnimation(Node, SkeletalMeshComponent);
 					}
 				}
 			}
@@ -167,7 +165,7 @@ FGLTFJsonNodeIndex FGLTFComponentConverter::Convert(const USceneComponent* Scene
 			if (Builder.ExportOptions->bExportCameras)
 			{
 				// TODO: conversion of camera direction should be done in separate converter
-				FGLTFJsonNode* CameraNode = Builder.AddChildComponentNode(Node->Index);
+				FGLTFJsonNode* CameraNode = Builder.AddChildComponentNode(Node);
 				CameraNode->Name = FGLTFNameUtility::GetName(CameraComponent);
 				CameraNode->Rotation = FGLTFConverterUtility::ConvertCameraDirection();
 				CameraNode->Camera = Builder.GetOrAddCamera(CameraComponent);
@@ -178,7 +176,7 @@ FGLTFJsonNodeIndex FGLTFComponentConverter::Convert(const USceneComponent* Scene
 			if (Builder.ShouldExportLight(LightComponent->Mobility))
 			{
 				// TODO: conversion of light direction should be done in separate converter
-				FGLTFJsonNode* LightNode = Builder.AddChildComponentNode(Node->Index);
+				FGLTFJsonNode* LightNode = Builder.AddChildComponentNode(Node);
 				LightNode->Name = FGLTFNameUtility::GetName(LightComponent);
 				LightNode->Rotation = FGLTFConverterUtility::ConvertLightDirection();
 				LightNode->Light = Builder.GetOrAddLight(LightComponent);
@@ -186,26 +184,26 @@ FGLTFJsonNodeIndex FGLTFComponentConverter::Convert(const USceneComponent* Scene
 		}
 	}
 
-	return Node->Index;
+	return Node;
 }
 
-FGLTFJsonNodeIndex FGLTFComponentSocketConverter::Convert(const USceneComponent* SceneComponent, FName SocketName)
+FGLTFJsonNode* FGLTFComponentSocketConverter::Convert(const USceneComponent* SceneComponent, FName SocketName)
 {
-	const FGLTFJsonNodeIndex NodeIndex = Builder.GetOrAddNode(SceneComponent);
+	FGLTFJsonNode* JsonNode = Builder.GetOrAddNode(SceneComponent);
 
 	if (SocketName != NAME_None)
 	{
 		if (const UStaticMeshComponent* StaticMeshComponent = Cast<UStaticMeshComponent>(SceneComponent))
 		{
 			const UStaticMesh* StaticMesh = StaticMeshComponent->GetStaticMesh();
-			return Builder.GetOrAddNode(NodeIndex, StaticMesh, SocketName);
+			return Builder.GetOrAddNode(JsonNode, StaticMesh, SocketName);
 		}
 
 		if (const USkinnedMeshComponent* SkinnedMeshComponent = Cast<USkinnedMeshComponent>(SceneComponent))
 		{
 			// TODO: add support for SocketOverrideLookup?
 			const USkeletalMesh* SkeletalMesh = SkinnedMeshComponent->SkeletalMesh;
-			return Builder.GetOrAddNode(NodeIndex, SkeletalMesh, SocketName);
+			return Builder.GetOrAddNode(JsonNode, SkeletalMesh, SocketName);
 		}
 
 		// TODO: add support for more socket types
@@ -216,16 +214,16 @@ FGLTFJsonNodeIndex FGLTFComponentSocketConverter::Convert(const USceneComponent*
 			*SceneComponent->GetName()));
 	}
 
-	return NodeIndex;
+	return JsonNode;
 }
 
-FGLTFJsonNodeIndex FGLTFStaticSocketConverter::Convert(FGLTFJsonNodeIndex RootNode, const UStaticMesh* StaticMesh, FName SocketName)
+FGLTFJsonNode* FGLTFStaticSocketConverter::Convert(FGLTFJsonNode* RootNode, const UStaticMesh* StaticMesh, FName SocketName)
 {
 	const UStaticMeshSocket* Socket = StaticMesh->FindSocket(SocketName);
 	if (Socket == nullptr)
 	{
 		// TODO: report error
-		return FGLTFJsonNodeIndex(INDEX_NONE);
+		return nullptr;
 	}
 
 	FGLTFJsonNode Node;
@@ -239,7 +237,7 @@ FGLTFJsonNodeIndex FGLTFStaticSocketConverter::Convert(FGLTFJsonNodeIndex RootNo
 	return Builder.AddChildNode(RootNode, Node);
 }
 
-FGLTFJsonNodeIndex FGLTFSkeletalSocketConverter::Convert(FGLTFJsonNodeIndex RootNode, const USkeletalMesh* SkeletalMesh, FName SocketName)
+FGLTFJsonNode* FGLTFSkeletalSocketConverter::Convert(FGLTFJsonNode* RootNode, const USkeletalMesh* SkeletalMesh, FName SocketName)
 {
 #if (ENGINE_MAJOR_VERSION > 4 || ENGINE_MINOR_VERSION >= 27)
 	const FReferenceSkeleton& RefSkeleton = SkeletalMesh->GetRefSkeleton();
@@ -259,7 +257,7 @@ FGLTFJsonNodeIndex FGLTFSkeletalSocketConverter::Convert(FGLTFJsonNodeIndex Root
 		Node.Scale = FGLTFConverterUtility::ConvertScale(Socket->RelativeScale);
 
 		const int32 ParentBone = RefSkeleton.FindBoneIndex(Socket->BoneName);
-		const FGLTFJsonNodeIndex ParentNode = ParentBone != INDEX_NONE ? Builder.GetOrAddNode(RootNode, SkeletalMesh, ParentBone) : RootNode;
+		FGLTFJsonNode* ParentNode = ParentBone != INDEX_NONE ? Builder.GetOrAddNode(RootNode, SkeletalMesh, ParentBone) : RootNode;
 		return Builder.AddChildNode(ParentNode, Node);
 	}
 
@@ -270,10 +268,10 @@ FGLTFJsonNodeIndex FGLTFSkeletalSocketConverter::Convert(FGLTFJsonNodeIndex Root
 	}
 
 	// TODO: report error
-	return FGLTFJsonNodeIndex(INDEX_NONE);
+	return nullptr;
 }
 
-FGLTFJsonNodeIndex FGLTFSkeletalBoneConverter::Convert(FGLTFJsonNodeIndex RootNode, const USkeletalMesh* SkeletalMesh, int32 BoneIndex)
+FGLTFJsonNode* FGLTFSkeletalBoneConverter::Convert(FGLTFJsonNode* RootNode, const USkeletalMesh* SkeletalMesh, int32 BoneIndex)
 {
 #if (ENGINE_MAJOR_VERSION > 4 || ENGINE_MINOR_VERSION >= 27)
 	const FReferenceSkeleton& RefSkeleton = SkeletalMesh->GetRefSkeleton();
@@ -287,7 +285,7 @@ FGLTFJsonNodeIndex FGLTFSkeletalBoneConverter::Convert(FGLTFJsonNodeIndex RootNo
 	if (!BoneInfos.IsValidIndex(BoneIndex))
 	{
 		// TODO: report error
-		return FGLTFJsonNodeIndex(INDEX_NONE);
+		return nullptr;
 	}
 
 	const FMeshBoneInfo& BoneInfo = BoneInfos[BoneIndex];
@@ -309,6 +307,6 @@ FGLTFJsonNodeIndex FGLTFSkeletalBoneConverter::Convert(FGLTFJsonNodeIndex RootNo
 		// TODO: report error
 	}
 
-	const FGLTFJsonNodeIndex ParentNode = BoneInfo.ParentIndex != INDEX_NONE ? Builder.GetOrAddNode(RootNode, SkeletalMesh, BoneInfo.ParentIndex) : RootNode;
+	FGLTFJsonNode* ParentNode = BoneInfo.ParentIndex != INDEX_NONE ? Builder.GetOrAddNode(RootNode, SkeletalMesh, BoneInfo.ParentIndex) : RootNode;
 	return Builder.AddChildNode(ParentNode, Node);
 }
diff --git a/Engine/Plugins/Enterprise/GLTFExporter/MigratedFromGit/Source/GLTFExporter/Private/Converters/GLTFNodeConverters.h b/Engine/Plugins/Enterprise/GLTFExporter/MigratedFromGit/Source/GLTFExporter/Private/Converters/GLTFNodeConverters.h
--- a/Engine/Plugins/Enterprise/GLTFExporter/MigratedFromGit/Source/GLTFExporter/Private/Converters/GLTFNodeConverters.h
+++ b/Engine/Plugins/Enterprise/GLTFExporter/MigratedFromGit/Source/GLTFExporter/Private/Converters/GLTFNodeConverters.h
@@ -6,51 +6,51 @@
 #include "Converters/GLTFConverter.h"
 #include "Converters/GLTFBuilderContext.h"
 
-typedef TGLTFConverter<FGLTFJsonNodeIndex, const AActor*> IGLTFActorConverter;
-typedef TGLTFConverter<FGLTFJsonNodeIndex, const USceneComponent*> IGLTFComponentConverter;
-typedef TGLTFConverter<FGLTFJsonNodeIndex, const USceneComponent*, FName> IGLTFComponentSocketConverter;
-typedef TGLTFConverter<FGLTFJsonNodeIndex, FGLTFJsonNodeIndex, const UStaticMesh*, FName> IGLTFStaticSocketConverter;
-typedef TGLTFConverter<FGLTFJsonNodeIndex, FGLTFJsonNodeIndex, const USkeletalMesh*, FName> IGLTFSkeletalSocketConverter;
-typedef TGLTFConverter<FGLTFJsonNodeIndex, FGLTFJsonNodeIndex, const USkeletalMesh*, int32> IGLTFSkeletalBoneConverter;
+typedef TGLTFConverter<FGLTFJsonNode*, const AActor*> IGLTFActorConverter;
+typedef TGLTFConverter<FGLTFJsonNode*, const USceneComponent*> IGLTFComponentConverter;
+typedef TGLTFConverter<FGLTFJsonNode*, const USceneComponent*, FName> IGLTFComponentSocketConverter;
+typedef TGLTFConverter<FGLTFJsonNode*, FGLTFJsonNode*, const UStaticMesh*, FName> IGLTFStaticSocketConverter;
+typedef TGLTFConverter<FGLTFJsonNode*, FGLTFJsonNode*, const USkeletalMesh*, FName> IGLTFSkeletalSocketConverter;
+typedef TGLTFConverter<FGLTFJsonNode*, FGLTFJsonNode*, const USkeletalMesh*, int32> IGLTFSkeletalBoneConverter;
 
 class FGLTFActorConverter final : public FGLTFBuilderContext, public IGLTFActorConverter
 {
 	using FGLTFBuilderContext::FGLTFBuilderContext;
 
-	virtual FGLTFJsonNodeIndex Convert(const AActor* Actor) override;
+	virtual FGLTFJsonNode* Convert(const AActor* Actor) override;
 };
 
 class FGLTFComponentConverter final : public FGLTFBuilderContext, public IGLTFComponentConverter
 {
 	using FGLTFBuilderContext::FGLTFBuilderContext;
 
-	virtual FGLTFJsonNodeIndex Convert(const USceneComponent* SceneComponent) override;
+	virtual FGLTFJsonNode* Convert(const USceneComponent* SceneComponent) override;
 };
 
 class FGLTFComponentSocketConverter final : public FGLTFBuilderContext, public IGLTFComponentSocketConverter
 {
 	using FGLTFBuilderContext::FGLTFBuilderContext;
 
-	virtual FGLTFJsonNodeIndex Convert(const USceneComponent* SceneComponent, FName SocketName) override;
+	virtual FGLTFJsonNode* Convert(const USceneComponent* SceneComponent, FName SocketName) override;
 };
 
 class FGLTFStaticSocketConverter final : public FGLTFBuilderContext, public IGLTFStaticSocketConverter
 {
 	using FGLTFBuilderContext::FGLTFBuilderContext;
 
-	virtual FGLTFJsonNodeIndex Convert(FGLTFJsonNodeIndex RootNode, const UStaticMesh* StaticMesh, FName SocketName) override;
+	virtual FGLTFJsonNode* Convert(FGLTFJsonNode* RootNode, const UStaticMesh* StaticMesh, FName SocketName) override;
 };
 
 class FGLTFSkeletalSocketConverter final : public FGLTFBuilderContext, public IGLTFSkeletalSocketConverter
 {
 	using FGLTFBuilderContext::FGLTFBuilderContext;
 
-	virtual FGLTFJsonNodeIndex Convert(FGLTFJsonNodeIndex RootNode, const USkeletalMesh* SkeletalMesh, FName SocketName) override;
+	virtual FGLTFJsonNode* Convert(FGLTFJsonNode* RootNode, const USkeletalMesh* SkeletalMesh, FName SocketName) override;
 };
 
 class FGLTFSkeletalBoneConverter final : public FGLTFBuilderContext, public IGLTFSkeletalBoneConverter
 {
 	using FGLTFBuilderContext::FGLTFBuilderContext;
 
-	virtual FGLTFJsonNodeIndex Convert(FGLTFJsonNodeIndex RootNode, const USkeletalMesh* SkeletalMesh, int32 BoneIndex) override;
+	virtual FGLTFJsonNode* Convert(FGLTFJsonNode* RootNode, const USkeletalMesh* SkeletalMesh, int32 BoneIndex) override;
 };
diff --git a/Engine/Plugins/Enterprise/GLTFExporter/MigratedFromGit/Source/GLTFExporter/Private/Converters/GLTFSamplerConverters.cpp b/Engine/Plugins/Enterprise/GLTFExporter/MigratedFromGit/Source/GLTFExporter/Private/Converters/GLTFSamplerConverters.cpp
--- a/Engine/Plugins/Enterprise/GLTFExporter/MigratedFromGit/Source/GLTFExporter/Private/Converters/GLTFSamplerConverters.cpp
+++ b/Engine/Plugins/Enterprise/GLTFExporter/MigratedFromGit/Source/GLTFExporter/Private/Converters/GLTFSamplerConverters.cpp
@@ -5,7 +5,7 @@
 #include "Converters/GLTFConverterUtility.h"
 #include "Converters/GLTFTextureUtility.h"
 
-FGLTFJsonSamplerIndex FGLTFSamplerConverter::Convert(const UTexture* Texture)
+FGLTFJsonSampler* FGLTFSamplerConverter::Convert(const UTexture* Texture)
 {
 	// TODO: reuse existing samplers, don't create a unique sampler per texture
 
diff --git a/Engine/Plugins/Enterprise/GLTFExporter/MigratedFromGit/Source/GLTFExporter/Private/Converters/GLTFSamplerConverters.h b/Engine/Plugins/Enterprise/GLTFExporter/MigratedFromGit/Source/GLTFExporter/Private/Converters/GLTFSamplerConverters.h
--- a/Engine/Plugins/Enterprise/GLTFExporter/MigratedFromGit/Source/GLTFExporter/Private/Converters/GLTFSamplerConverters.h
+++ b/Engine/Plugins/Enterprise/GLTFExporter/MigratedFromGit/Source/GLTFExporter/Private/Converters/GLTFSamplerConverters.h
@@ -6,11 +6,11 @@
 #include "Converters/GLTFConverter.h"
 #include "Converters/GLTFBuilderContext.h"
 
-typedef TGLTFConverter<FGLTFJsonSamplerIndex, const UTexture*> IGLTFSamplerConverter;
+typedef TGLTFConverter<FGLTFJsonSampler*, const UTexture*> IGLTFSamplerConverter;
 
 class FGLTFSamplerConverter final : public FGLTFBuilderContext, public IGLTFSamplerConverter
 {
 	using FGLTFBuilderContext::FGLTFBuilderContext;
 
-	virtual FGLTFJsonSamplerIndex Convert(const UTexture* Texture) override;
+	virtual FGLTFJsonSampler* Convert(const UTexture* Texture) override;
 };
diff --git a/Engine/Plugins/Enterprise/GLTFExporter/MigratedFromGit/Source/GLTFExporter/Private/Converters/GLTFSceneConverters.cpp b/Engine/Plugins/Enterprise/GLTFExporter/MigratedFromGit/Source/GLTFExporter/Private/Converters/GLTFSceneConverters.cpp
--- a/Engine/Plugins/Enterprise/GLTFExporter/MigratedFromGit/Source/GLTFExporter/Private/Converters/GLTFSceneConverters.cpp
+++ b/Engine/Plugins/Enterprise/GLTFExporter/MigratedFromGit/Source/GLTFExporter/Private/Converters/GLTFSceneConverters.cpp
@@ -6,7 +6,7 @@
 #include "LevelVariantSets.h"
 #include "VariantSet.h"
 
-FGLTFJsonSceneIndex FGLTFSceneConverter::Convert(const UWorld* World)
+FGLTFJsonScene* FGLTFSceneConverter::Convert(const UWorld* World)
 {
 	FGLTFJsonScene Scene;
 	World->GetName(Scene.Name);
@@ -36,10 +36,10 @@ FGLTFJsonSceneIndex FGLTFSceneConverter::Convert(const UWorld* World)
 						{
 							if (LevelVariantSets != nullptr)
 							{
-								const FGLTFJsonEpicLevelVariantSetsIndex EpicLevelVariantSetsIndex = Builder.GetOrAddEpicLevelVariantSets(LevelVariantSets);
-								if (EpicLevelVariantSetsIndex != INDEX_NONE)
+								FGLTFJsonEpicLevelVariantSets* JsonEpicLevelVariantSets = Builder.GetOrAddEpicLevelVariantSets(LevelVariantSets);
+								if (JsonEpicLevelVariantSets != nullptr)
 								{
-									Scene.EpicLevelVariantSets.Add(EpicLevelVariantSetsIndex);
+									Scene.EpicLevelVariantSets.Add(JsonEpicLevelVariantSets);
 								}
 							}
 						}
@@ -62,12 +62,12 @@ FGLTFJsonSceneIndex FGLTFSceneConverter::Convert(const UWorld* World)
 
 			for (const AActor* Actor : Level->Actors)
 			{
-				const FGLTFJsonNodeIndex NodeIndex = Builder.GetOrAddNode(Actor);
-				if (NodeIndex != INDEX_NONE && Builder.IsRootActor(Actor))
+				FGLTFJsonNode* JsonNode = Builder.GetOrAddNode(Actor);
+				if (JsonNode != nullptr && Builder.IsRootActor(Actor))
 				{
 					// TODO: to avoid having to add irrelevant actors/components let GLTFComponentConverter decide and add root nodes to scene.
 					// This change may require node converters to support cyclic calls.
-					Scene.Nodes.Add(NodeIndex);
+					Scene.Nodes.Add(JsonNode);
 				}
 			}
 		}
diff --git a/Engine/Plugins/Enterprise/GLTFExporter/MigratedFromGit/Source/GLTFExporter/Private/Converters/GLTFSceneConverters.h b/Engine/Plugins/Enterprise/GLTFExporter/MigratedFromGit/Source/GLTFExporter/Private/Converters/GLTFSceneConverters.h
--- a/Engine/Plugins/Enterprise/GLTFExporter/MigratedFromGit/Source/GLTFExporter/Private/Converters/GLTFSceneConverters.h
+++ b/Engine/Plugins/Enterprise/GLTFExporter/MigratedFromGit/Source/GLTFExporter/Private/Converters/GLTFSceneConverters.h
@@ -6,11 +6,11 @@
 #include "Converters/GLTFConverter.h"
 #include "Converters/GLTFBuilderContext.h"
 
-typedef TGLTFConverter<FGLTFJsonSceneIndex, const UWorld*> IGLTFSceneConverter;
+typedef TGLTFConverter<FGLTFJsonScene*, const UWorld*> IGLTFSceneConverter;
 
 class FGLTFSceneConverter final : public FGLTFBuilderContext, public IGLTFSceneConverter
 {
 	using FGLTFBuilderContext::FGLTFBuilderContext;
 
-	virtual FGLTFJsonSceneIndex Convert(const UWorld* Level) override;
+	virtual FGLTFJsonScene* Convert(const UWorld* Level) override;
 };
diff --git a/Engine/Plugins/Enterprise/GLTFExporter/MigratedFromGit/Source/GLTFExporter/Private/Converters/GLTFSkinConverters.cpp b/Engine/Plugins/Enterprise/GLTFExporter/MigratedFromGit/Source/GLTFExporter/Private/Converters/GLTFSkinConverters.cpp
--- a/Engine/Plugins/Enterprise/GLTFExporter/MigratedFromGit/Source/GLTFExporter/Private/Converters/GLTFSkinConverters.cpp
+++ b/Engine/Plugins/Enterprise/GLTFExporter/MigratedFromGit/Source/GLTFExporter/Private/Converters/GLTFSkinConverters.cpp
@@ -5,7 +5,7 @@
 #include "Converters/GLTFBoneUtility.h"
 #include "Builders/GLTFConvertBuilder.h"
 
-FGLTFJsonSkinIndex FGLTFSkinConverter::Convert(FGLTFJsonNodeIndex RootNode, const USkeletalMesh* SkeletalMesh)
+FGLTFJsonSkin* FGLTFSkinConverter::Convert(FGLTFJsonNode* RootNode, const USkeletalMesh* SkeletalMesh)
 {
 #if (ENGINE_MAJOR_VERSION > 4 || ENGINE_MINOR_VERSION >= 27)
 	const USkeleton* Skeleton = SkeletalMesh->GetSkeleton();
@@ -23,7 +23,7 @@ FGLTFJsonSkinIndex FGLTFSkinConverter::Convert(FGLTFJsonNodeIndex RootNode, cons
 	if (BoneCount == 0)
 	{
 		// TODO: report warning
-		return FGLTFJsonSkinIndex(INDEX_NONE);
+		return nullptr;
 	}
 
 	Skin.Joints.AddUninitialized(BoneCount);
diff --git a/Engine/Plugins/Enterprise/GLTFExporter/MigratedFromGit/Source/GLTFExporter/Private/Converters/GLTFSkinConverters.h b/Engine/Plugins/Enterprise/GLTFExporter/MigratedFromGit/Source/GLTFExporter/Private/Converters/GLTFSkinConverters.h
--- a/Engine/Plugins/Enterprise/GLTFExporter/MigratedFromGit/Source/GLTFExporter/Private/Converters/GLTFSkinConverters.h
+++ b/Engine/Plugins/Enterprise/GLTFExporter/MigratedFromGit/Source/GLTFExporter/Private/Converters/GLTFSkinConverters.h
@@ -6,11 +6,11 @@
 #include "Converters/GLTFConverter.h"
 #include "Converters/GLTFBuilderContext.h"
 
-typedef TGLTFConverter<FGLTFJsonSkinIndex, FGLTFJsonNodeIndex, const USkeletalMesh*> IGLTFSkinConverter;
+typedef TGLTFConverter<FGLTFJsonSkin*, FGLTFJsonNode*, const USkeletalMesh*> IGLTFSkinConverter;
 
 class FGLTFSkinConverter final : public FGLTFBuilderContext, public IGLTFSkinConverter
 {
 	using FGLTFBuilderContext::FGLTFBuilderContext;
 
-	virtual FGLTFJsonSkinIndex Convert(FGLTFJsonNodeIndex RootNode, const USkeletalMesh* SkeletalMesh) override;
+	virtual FGLTFJsonSkin* Convert(FGLTFJsonNode* RootNode, const USkeletalMesh* SkeletalMesh) override;
 };
diff --git a/Engine/Plugins/Enterprise/GLTFExporter/MigratedFromGit/Source/GLTFExporter/Private/Converters/GLTFSkySphereConverters.cpp b/Engine/Plugins/Enterprise/GLTFExporter/MigratedFromGit/Source/GLTFExporter/Private/Converters/GLTFSkySphereConverters.cpp
--- a/Engine/Plugins/Enterprise/GLTFExporter/MigratedFromGit/Source/GLTFExporter/Private/Converters/GLTFSkySphereConverters.cpp
+++ b/Engine/Plugins/Enterprise/GLTFExporter/MigratedFromGit/Source/GLTFExporter/Private/Converters/GLTFSkySphereConverters.cpp
@@ -10,11 +10,11 @@
 #include "Curves/CurveLinearColor.h"
 #include "Materials/MaterialInstance.h"
 
-FGLTFJsonSkySphereIndex FGLTFSkySphereConverter::Convert(const AActor* SkySphereActor)
+FGLTFJsonSkySphere* FGLTFSkySphereConverter::Convert(const AActor* SkySphereActor)
 {
 	if (!FGLTFActorUtility::IsSkySphereBlueprint(FGLTFActorUtility::GetBlueprintPath(SkySphereActor)))
 	{
-		return FGLTFJsonSkySphereIndex(INDEX_NONE);
+		return nullptr;
 	}
 
 	FGLTFJsonSkySphere JsonSkySphere;
@@ -180,7 +180,7 @@ void FGLTFSkySphereConverter::ConvertScalarParameter(const AActor* Actor, const
 	}
 }
 
-void FGLTFSkySphereConverter::ConvertTextureParameter(const AActor* Actor, const UMaterialInstance* Material, const ESkySphereTextureParameter Parameter, FGLTFJsonTextureIndex& OutValue) const
+void FGLTFSkySphereConverter::ConvertTextureParameter(const AActor* Actor, const UMaterialInstance* Material, const ESkySphereTextureParameter Parameter, FGLTFJsonTexture*& OutValue) const
 {
 	check(Actor != nullptr && Material != nullptr);
 
diff --git a/Engine/Plugins/Enterprise/GLTFExporter/MigratedFromGit/Source/GLTFExporter/Private/Converters/GLTFSkySphereConverters.h b/Engine/Plugins/Enterprise/GLTFExporter/MigratedFromGit/Source/GLTFExporter/Private/Converters/GLTFSkySphereConverters.h
--- a/Engine/Plugins/Enterprise/GLTFExporter/MigratedFromGit/Source/GLTFExporter/Private/Converters/GLTFSkySphereConverters.h
+++ b/Engine/Plugins/Enterprise/GLTFExporter/MigratedFromGit/Source/GLTFExporter/Private/Converters/GLTFSkySphereConverters.h
@@ -9,7 +9,7 @@
 
 struct FGLTFJsonSkySphereColorCurve;
 
-typedef TGLTFConverter<FGLTFJsonSkySphereIndex, const AActor*> IGLTFSkySphereConverter;
+typedef TGLTFConverter<FGLTFJsonSkySphere*, const AActor*> IGLTFSkySphereConverter;
 
 class FGLTFSkySphereConverter final : public FGLTFBuilderContext, public IGLTFSkySphereConverter
 {
@@ -22,7 +22,7 @@ class FGLTFSkySphereConverter final : public FGLTFBuilderContext, public IGLTFSk
 
 	using FGLTFBuilderContext::FGLTFBuilderContext;
 
-	virtual FGLTFJsonSkySphereIndex Convert(const AActor* SkySphereActor) override;
+	virtual FGLTFJsonSkySphere* Convert(const AActor* SkySphereActor) override;
 
 	template <class ValueType>
 	void ConvertProperty(const AActor* Actor, const TCHAR* PropertyName, ValueType& OutValue) const;
@@ -31,5 +31,5 @@ class FGLTFSkySphereConverter final : public FGLTFBuilderContext, public IGLTFSk
 	void ConvertColorCurveProperty(const AActor* Actor, const TCHAR* PropertyName, FGLTFJsonSkySphereColorCurve& OutValue) const;
 	void ConvertScalarParameter(const AActor* Actor, const UMaterialInstance* Material, const TCHAR* ParameterName, float& OutValue) const;
 
-	void ConvertTextureParameter(const AActor* Actor, const UMaterialInstance* Material, const ESkySphereTextureParameter Parameter, FGLTFJsonTextureIndex& OutValue) const;
+	void ConvertTextureParameter(const AActor* Actor, const UMaterialInstance* Material, const ESkySphereTextureParameter Parameter, FGLTFJsonTexture*& OutValue) const;
 };
diff --git a/Engine/Plugins/Enterprise/GLTFExporter/MigratedFromGit/Source/GLTFExporter/Private/Converters/GLTFTextureConverters.cpp b/Engine/Plugins/Enterprise/GLTFExporter/MigratedFromGit/Source/GLTFExporter/Private/Converters/GLTFTextureConverters.cpp
--- a/Engine/Plugins/Enterprise/GLTFExporter/MigratedFromGit/Source/GLTFExporter/Private/Converters/GLTFTextureConverters.cpp
+++ b/Engine/Plugins/Enterprise/GLTFExporter/MigratedFromGit/Source/GLTFExporter/Private/Converters/GLTFTextureConverters.cpp
@@ -16,16 +16,16 @@ void FGLTFTexture2DConverter::Sanitize(const UTexture2D*& Texture2D, bool& bToSR
 	}
 }
 
-FGLTFJsonTextureIndex FGLTFTexture2DConverter::Convert(const UTexture2D* Texture2D, bool bToSRGB)
+FGLTFJsonTexture* FGLTFTexture2DConverter::Convert(const UTexture2D* Texture2D, bool bToSRGB)
 {
 	if (Builder.ExportOptions->TextureImageFormat != EGLTFTextureImageFormat::None)
 	{
 		FGLTFJsonTexture* JsonTexture = Builder.AddTexture();
 		Builder.SetupTask<FGLTFTexture2DTask>(Builder, Texture2D, bToSRGB, JsonTexture);
-		return JsonTexture->Index;;
+		return JsonTexture;
 	}
 
-	return FGLTFJsonTextureIndex(INDEX_NONE);
+	return nullptr;
 }
 
 void FGLTFTextureCubeConverter::Sanitize(const UTextureCube*& TextureCube, ECubeFace& CubeFace, bool& bToSRGB)
@@ -36,16 +36,16 @@ void FGLTFTextureCubeConverter::Sanitize(const UTextureCube*& TextureCube, ECube
 	}
 }
 
-FGLTFJsonTextureIndex FGLTFTextureCubeConverter::Convert(const UTextureCube* TextureCube, ECubeFace CubeFace, bool bToSRGB)
+FGLTFJsonTexture* FGLTFTextureCubeConverter::Convert(const UTextureCube* TextureCube, ECubeFace CubeFace, bool bToSRGB)
 {
 	if (Builder.ExportOptions->TextureImageFormat != EGLTFTextureImageFormat::None)
 	{
 		FGLTFJsonTexture* JsonTexture = Builder.AddTexture();
 		Builder.SetupTask<FGLTFTextureCubeTask>(Builder, TextureCube, CubeFace, bToSRGB, JsonTexture);
-		return JsonTexture->Index;;
+		return JsonTexture;
 	}
 
-	return FGLTFJsonTextureIndex(INDEX_NONE);
+	return nullptr;
 }
 
 void FGLTFTextureRenderTarget2DConverter::Sanitize(const UTextureRenderTarget2D*& RenderTarget2D, bool& bToSRGB)
@@ -56,16 +56,16 @@ void FGLTFTextureRenderTarget2DConverter::Sanitize(const UTextureRenderTarget2D*
 	}
 }
 
-FGLTFJsonTextureIndex FGLTFTextureRenderTarget2DConverter::Convert(const UTextureRenderTarget2D* RenderTarget2D, bool bToSRGB)
+FGLTFJsonTexture* FGLTFTextureRenderTarget2DConverter::Convert(const UTextureRenderTarget2D* RenderTarget2D, bool bToSRGB)
 {
 	if (Builder.ExportOptions->TextureImageFormat != EGLTFTextureImageFormat::None)
 	{
 		FGLTFJsonTexture* JsonTexture = Builder.AddTexture();
 		Builder.SetupTask<FGLTFTextureRenderTarget2DTask>(Builder, RenderTarget2D, bToSRGB, JsonTexture);
-		return JsonTexture->Index;;
+		return JsonTexture;
 	}
 
-	return FGLTFJsonTextureIndex(INDEX_NONE);
+	return nullptr;
 }
 
 void FGLTFTextureRenderTargetCubeConverter::Sanitize(const UTextureRenderTargetCube*& RenderTargetCube, ECubeFace& CubeFace, bool& bToSRGB)
@@ -76,28 +76,28 @@ void FGLTFTextureRenderTargetCubeConverter::Sanitize(const UTextureRenderTargetC
 	}
 }
 
-FGLTFJsonTextureIndex FGLTFTextureRenderTargetCubeConverter::Convert(const UTextureRenderTargetCube* RenderTargetCube, ECubeFace CubeFace, bool bToSRGB)
+FGLTFJsonTexture* FGLTFTextureRenderTargetCubeConverter::Convert(const UTextureRenderTargetCube* RenderTargetCube, ECubeFace CubeFace, bool bToSRGB)
 {
 	if (Builder.ExportOptions->TextureImageFormat != EGLTFTextureImageFormat::None)
 	{
 		FGLTFJsonTexture* JsonTexture = Builder.AddTexture();
 		Builder.SetupTask<FGLTFTextureRenderTargetCubeTask>(Builder, RenderTargetCube, CubeFace, bToSRGB, JsonTexture);
-		return JsonTexture->Index;;
+		return JsonTexture;
 	}
 
-	return FGLTFJsonTextureIndex(INDEX_NONE);
+	return nullptr;
 }
 
-FGLTFJsonTextureIndex FGLTFTextureLightMapConverter::Convert(const ULightMapTexture2D* LightMap)
+FGLTFJsonTexture* FGLTFTextureLightMapConverter::Convert(const ULightMapTexture2D* LightMap)
 {
 #if WITH_EDITOR
 	if (Builder.ExportOptions->TextureImageFormat != EGLTFTextureImageFormat::None)
 	{
 		FGLTFJsonTexture* JsonTexture = Builder.AddTexture();
 		Builder.SetupTask<FGLTFTextureLightMapTask>(Builder, LightMap, JsonTexture);
-		return JsonTexture->Index;;
+		return JsonTexture;
 	}
 #endif
 
-	return FGLTFJsonTextureIndex(INDEX_NONE);
+	return nullptr;
 }
diff --git a/Engine/Plugins/Enterprise/GLTFExporter/MigratedFromGit/Source/GLTFExporter/Private/Converters/GLTFTextureConverters.h b/Engine/Plugins/Enterprise/GLTFExporter/MigratedFromGit/Source/GLTFExporter/Private/Converters/GLTFTextureConverters.h
--- a/Engine/Plugins/Enterprise/GLTFExporter/MigratedFromGit/Source/GLTFExporter/Private/Converters/GLTFTextureConverters.h
+++ b/Engine/Plugins/Enterprise/GLTFExporter/MigratedFromGit/Source/GLTFExporter/Private/Converters/GLTFTextureConverters.h
@@ -10,11 +10,11 @@
 
 // TODO: remove lightmap-specific converter (since it won't work in runtime)
 
-typedef TGLTFConverter<FGLTFJsonTextureIndex, const UTexture2D*, bool> IGLTFTexture2DConverter;
-typedef TGLTFConverter<FGLTFJsonTextureIndex, const UTextureCube*, ECubeFace, bool> IGLTFTextureCubeConverter;
-typedef TGLTFConverter<FGLTFJsonTextureIndex, const UTextureRenderTarget2D*, bool> IGLTFTextureRenderTarget2DConverter;
-typedef TGLTFConverter<FGLTFJsonTextureIndex, const UTextureRenderTargetCube*, ECubeFace, bool> IGLTFTextureRenderTargetCubeConverter;
-typedef TGLTFConverter<FGLTFJsonTextureIndex, const ULightMapTexture2D*> IGLTFTextureLightMapConverter;
+typedef TGLTFConverter<FGLTFJsonTexture*, const UTexture2D*, bool> IGLTFTexture2DConverter;
+typedef TGLTFConverter<FGLTFJsonTexture*, const UTextureCube*, ECubeFace, bool> IGLTFTextureCubeConverter;
+typedef TGLTFConverter<FGLTFJsonTexture*, const UTextureRenderTarget2D*, bool> IGLTFTextureRenderTarget2DConverter;
+typedef TGLTFConverter<FGLTFJsonTexture*, const UTextureRenderTargetCube*, ECubeFace, bool> IGLTFTextureRenderTargetCubeConverter;
+typedef TGLTFConverter<FGLTFJsonTexture*, const ULightMapTexture2D*> IGLTFTextureLightMapConverter;
 
 class FGLTFTexture2DConverter final : public FGLTFBuilderContext, public IGLTFTexture2DConverter
 {
@@ -22,7 +22,7 @@ class FGLTFTexture2DConverter final : public FGLTFBuilderContext, public IGLTFTe
 
 	virtual void Sanitize(const UTexture2D*& Texture2D, bool& bToSRGB) override;
 
-	virtual FGLTFJsonTextureIndex Convert(const UTexture2D* Texture2D, bool bToSRGB) override;
+	virtual FGLTFJsonTexture* Convert(const UTexture2D* Texture2D, bool bToSRGB) override;
 };
 
 class FGLTFTextureCubeConverter final : public FGLTFBuilderContext, public IGLTFTextureCubeConverter
@@ -31,7 +31,7 @@ class FGLTFTextureCubeConverter final : public FGLTFBuilderContext, public IGLTF
 
 	virtual void Sanitize(const UTextureCube*& TextureCube, ECubeFace& CubeFace, bool& bToSRGB) override;
 
-	virtual FGLTFJsonTextureIndex Convert(const UTextureCube* TextureCube, ECubeFace CubeFace, bool bToSRGB) override;
+	virtual FGLTFJsonTexture* Convert(const UTextureCube* TextureCube, ECubeFace CubeFace, bool bToSRGB) override;
 };
 
 class FGLTFTextureRenderTarget2DConverter final : public FGLTFBuilderContext, public IGLTFTextureRenderTarget2DConverter
@@ -40,7 +40,7 @@ class FGLTFTextureRenderTarget2DConverter final : public FGLTFBuilderContext, pu
 
 	virtual void Sanitize(const UTextureRenderTarget2D*& RenderTarget2D, bool& bToSRGB) override;
 
-	virtual FGLTFJsonTextureIndex Convert(const UTextureRenderTarget2D* RenderTarget2D, bool bToSRGB) override;
+	virtual FGLTFJsonTexture* Convert(const UTextureRenderTarget2D* RenderTarget2D, bool bToSRGB) override;
 };
 
 class FGLTFTextureRenderTargetCubeConverter final : public FGLTFBuilderContext, public IGLTFTextureRenderTargetCubeConverter
@@ -49,12 +49,12 @@ class FGLTFTextureRenderTargetCubeConverter final : public FGLTFBuilderContext,
 
 	virtual void Sanitize(const UTextureRenderTargetCube*& RenderTargetCube, ECubeFace& CubeFace, bool& bToSRGB) override;
 
-	virtual FGLTFJsonTextureIndex Convert(const UTextureRenderTargetCube* RenderTargetCube, ECubeFace CubeFace, bool bToSRGB) override;
+	virtual FGLTFJsonTexture* Convert(const UTextureRenderTargetCube* RenderTargetCube, ECubeFace CubeFace, bool bToSRGB) override;
 };
 
 class FGLTFTextureLightMapConverter final : public FGLTFBuilderContext, public IGLTFTextureLightMapConverter
 {
 	using FGLTFBuilderContext::FGLTFBuilderContext;
 
-	virtual FGLTFJsonTextureIndex Convert(const ULightMapTexture2D* LightMap) override;
+	virtual FGLTFJsonTexture* Convert(const ULightMapTexture2D* LightMap) override;
 };
diff --git a/Engine/Plugins/Enterprise/GLTFExporter/MigratedFromGit/Source/GLTFExporter/Private/Converters/GLTFVariantUtility.cpp b/Engine/Plugins/Enterprise/GLTFExporter/MigratedFromGit/Source/GLTFExporter/Private/Converters/GLTFVariantUtility.cpp
--- a/Engine/Plugins/Enterprise/GLTFExporter/MigratedFromGit/Source/GLTFExporter/Private/Converters/GLTFVariantUtility.cpp
+++ b/Engine/Plugins/Enterprise/GLTFExporter/MigratedFromGit/Source/GLTFExporter/Private/Converters/GLTFVariantUtility.cpp
@@ -85,7 +85,7 @@ FString FGLTFVariantUtility::GetLogContext(const ULevelVariantSets* LevelVariant
 	return LevelVariantSets->GetName();
 }
 
-FGLTFJsonMaterialIndex FGLTFVariantUtility::GetOrAddMaterial(FGLTFConvertBuilder& Builder, const UMaterialInterface* Material, const UMeshComponent* MeshComponent, int32 MaterialIndex)
+FGLTFJsonMaterial* FGLTFVariantUtility::GetOrAddMaterial(FGLTFConvertBuilder& Builder, const UMaterialInterface* Material, const UMeshComponent* MeshComponent, int32 MaterialIndex)
 {
 	if (Builder.ExportOptions->ExportMaterialVariants == EGLTFMaterialVariantMode::UseMeshData)
 	{
diff --git a/Engine/Plugins/Enterprise/GLTFExporter/MigratedFromGit/Source/GLTFExporter/Private/Converters/GLTFVariantUtility.h b/Engine/Plugins/Enterprise/GLTFExporter/MigratedFromGit/Source/GLTFExporter/Private/Converters/GLTFVariantUtility.h
--- a/Engine/Plugins/Enterprise/GLTFExporter/MigratedFromGit/Source/GLTFExporter/Private/Converters/GLTFVariantUtility.h
+++ b/Engine/Plugins/Enterprise/GLTFExporter/MigratedFromGit/Source/GLTFExporter/Private/Converters/GLTFVariantUtility.h
@@ -34,5 +34,5 @@ struct FGLTFVariantUtility
 	static FString GetLogContext(const UVariantSet* VariantSet);
 	static FString GetLogContext(const ULevelVariantSets* LevelVariantSets);
 
-	static FGLTFJsonMaterialIndex GetOrAddMaterial(FGLTFConvertBuilder& Builder, const UMaterialInterface* Material, const UMeshComponent* MeshComponent, int32 MaterialIndex);
+	static FGLTFJsonMaterial* GetOrAddMaterial(FGLTFConvertBuilder& Builder, const UMaterialInterface* Material, const UMeshComponent* MeshComponent, int32 MaterialIndex);
 };
diff --git a/Engine/Plugins/Enterprise/GLTFExporter/MigratedFromGit/Source/GLTFExporter/Private/Exporters/GLTFAnimSequenceExporter.cpp b/Engine/Plugins/Enterprise/GLTFExporter/MigratedFromGit/Source/GLTFExporter/Private/Exporters/GLTFAnimSequenceExporter.cpp
--- a/Engine/Plugins/Enterprise/GLTFExporter/MigratedFromGit/Source/GLTFExporter/Private/Exporters/GLTFAnimSequenceExporter.cpp
+++ b/Engine/Plugins/Enterprise/GLTFExporter/MigratedFromGit/Source/GLTFExporter/Private/Exporters/GLTFAnimSequenceExporter.cpp
@@ -44,8 +44,8 @@ bool UGLTFAnimSequenceExporter::AddObject(FGLTFContainerBuilder& Builder, const
 
 	if (Builder.ExportOptions->bExportPreviewMesh)
 	{
-		const FGLTFJsonMeshIndex MeshIndex = Builder.GetOrAddMesh(SkeletalMesh);
-		if (MeshIndex == INDEX_NONE)
+		FGLTFJsonMesh* MeshIndex = Builder.GetOrAddMesh(SkeletalMesh);
+		if (MeshIndex == nullptr)
 		{
 			Builder.LogError(
 				FString::Printf(TEXT("Failed to export skeletal mesh %s for animation sequence %s"),
@@ -57,10 +57,10 @@ bool UGLTFAnimSequenceExporter::AddObject(FGLTFContainerBuilder& Builder, const
 		Node.Mesh = MeshIndex;
 	}
 
-	const FGLTFJsonNodeIndex NodeIndex = Builder.AddNode(Node);
+	FGLTFJsonNode* NodeIndex = Builder.AddNode(Node);
 
-	const FGLTFJsonSkinIndex SkinIndex = Builder.GetOrAddSkin(NodeIndex, SkeletalMesh);
-	if (SkinIndex == INDEX_NONE)
+	FGLTFJsonSkin* SkinIndex = Builder.GetOrAddSkin(NodeIndex, SkeletalMesh);
+	if (SkinIndex == nullptr)
 	{
 		Builder.LogError(
 			FString::Printf(TEXT("Failed to export bones in skeletal mesh %s for animation sequence %s"),
@@ -71,8 +71,8 @@ bool UGLTFAnimSequenceExporter::AddObject(FGLTFContainerBuilder& Builder, const
 
 	Builder.GetNode(NodeIndex).Skin = SkinIndex;
 
-	const FGLTFJsonAnimationIndex AnimationIndex =  Builder.GetOrAddAnimation(NodeIndex, SkeletalMesh, AnimSequence);
-	if (AnimationIndex == INDEX_NONE)
+	FGLTFJsonAnimation* AnimationIndex =  Builder.GetOrAddAnimation(NodeIndex, SkeletalMesh, AnimSequence);
+	if (AnimationIndex == nullptr)
 	{
 		Builder.LogError(
 			FString::Printf(TEXT("Failed to export animation sequence %s"),
@@ -82,7 +82,7 @@ bool UGLTFAnimSequenceExporter::AddObject(FGLTFContainerBuilder& Builder, const
 
 	FGLTFJsonScene Scene;
 	Scene.Nodes.Add(NodeIndex);
-	const FGLTFJsonSceneIndex SceneIndex = Builder.AddScene(Scene);
+	FGLTFJsonScene* SceneIndex = Builder.AddScene(Scene);
 
 	Builder.DefaultScene = SceneIndex;
 	return true;
diff --git a/Engine/Plugins/Enterprise/GLTFExporter/MigratedFromGit/Source/GLTFExporter/Private/Exporters/GLTFLevelExporter.cpp b/Engine/Plugins/Enterprise/GLTFExporter/MigratedFromGit/Source/GLTFExporter/Private/Exporters/GLTFLevelExporter.cpp
--- a/Engine/Plugins/Enterprise/GLTFExporter/MigratedFromGit/Source/GLTFExporter/Private/Exporters/GLTFLevelExporter.cpp
+++ b/Engine/Plugins/Enterprise/GLTFExporter/MigratedFromGit/Source/GLTFExporter/Private/Exporters/GLTFLevelExporter.cpp
@@ -14,8 +14,8 @@ bool UGLTFLevelExporter::AddObject(FGLTFContainerBuilder& Builder, const UObject
 {
 	const UWorld* World = CastChecked<UWorld>(Object);
 
-	const FGLTFJsonSceneIndex SceneIndex = Builder.GetOrAddScene(World);
-	if (SceneIndex == INDEX_NONE)
+	FGLTFJsonScene* SceneIndex = Builder.GetOrAddScene(World);
+	if (SceneIndex == nullptr)
 	{
 		Builder.LogError(FString::Printf(TEXT("Failed to export level %s"), *World->GetName()));
 		return false;
diff --git a/Engine/Plugins/Enterprise/GLTFExporter/MigratedFromGit/Source/GLTFExporter/Private/Exporters/GLTFLevelSequenceExporter.cpp b/Engine/Plugins/Enterprise/GLTFExporter/MigratedFromGit/Source/GLTFExporter/Private/Exporters/GLTFLevelSequenceExporter.cpp
--- a/Engine/Plugins/Enterprise/GLTFExporter/MigratedFromGit/Source/GLTFExporter/Private/Exporters/GLTFLevelSequenceExporter.cpp
+++ b/Engine/Plugins/Enterprise/GLTFExporter/MigratedFromGit/Source/GLTFExporter/Private/Exporters/GLTFLevelSequenceExporter.cpp
@@ -41,8 +41,8 @@ bool UGLTFLevelSequenceExporter::AddObject(FGLTFContainerBuilder& Builder, const
 	}
 
 	UWorld* World = Worlds[0];
-	const FGLTFJsonSceneIndex SceneIndex = Builder.GetOrAddScene(World);
-	if (SceneIndex == INDEX_NONE)
+	FGLTFJsonScene* SceneIndex = Builder.GetOrAddScene(World);
+	if (SceneIndex == nullptr)
 	{
 		Builder.LogError(
 			FString::Printf(TEXT("Failed to export level %s for level sequence %s"),
@@ -51,8 +51,8 @@ bool UGLTFLevelSequenceExporter::AddObject(FGLTFContainerBuilder& Builder, const
 		return false;
 	}
 
-	const FGLTFJsonAnimationIndex AnimationIndex = Builder.GetOrAddAnimation(World->PersistentLevel, LevelSequence);
-	if (AnimationIndex == INDEX_NONE)
+	FGLTFJsonAnimation* AnimationIndex = Builder.GetOrAddAnimation(World->PersistentLevel, LevelSequence);
+	if (AnimationIndex == nullptr)
 	{
 		Builder.LogError(
 			FString::Printf(TEXT("Failed to export level sequence %s"),
diff --git a/Engine/Plugins/Enterprise/GLTFExporter/MigratedFromGit/Source/GLTFExporter/Private/Exporters/GLTFLevelVariantSetsExporter.cpp b/Engine/Plugins/Enterprise/GLTFExporter/MigratedFromGit/Source/GLTFExporter/Private/Exporters/GLTFLevelVariantSetsExporter.cpp
--- a/Engine/Plugins/Enterprise/GLTFExporter/MigratedFromGit/Source/GLTFExporter/Private/Exporters/GLTFLevelVariantSetsExporter.cpp
+++ b/Engine/Plugins/Enterprise/GLTFExporter/MigratedFromGit/Source/GLTFExporter/Private/Exporters/GLTFLevelVariantSetsExporter.cpp
@@ -41,8 +41,8 @@ bool UGLTFLevelVariantSetsExporter::AddObject(FGLTFContainerBuilder& Builder, co
 	}
 
 	UWorld* World = Worlds[0];
-	const FGLTFJsonSceneIndex SceneIndex = Builder.GetOrAddScene(World);
-	if (SceneIndex == INDEX_NONE)
+	FGLTFJsonScene* SceneIndex = Builder.GetOrAddScene(World);
+	if (SceneIndex == nullptr)
 	{
 		Builder.LogError(
 			FString::Printf(TEXT("Failed to export level %s for level variant sets %s"),
@@ -53,8 +53,8 @@ bool UGLTFLevelVariantSetsExporter::AddObject(FGLTFContainerBuilder& Builder, co
 
 	if (Builder.ExportOptions->VariantSetsMode == EGLTFVariantSetsMode::Epic)
 	{
-		const FGLTFJsonEpicLevelVariantSetsIndex EpicLevelVariantSetsIndex = Builder.GetOrAddEpicLevelVariantSets(LevelVariantSets);
-		if (EpicLevelVariantSetsIndex == INDEX_NONE)
+		FGLTFJsonEpicLevelVariantSets* EpicLevelVariantSetsIndex = Builder.GetOrAddEpicLevelVariantSets(LevelVariantSets);
+		if (EpicLevelVariantSetsIndex == nullptr)
 		{
 			Builder.LogError(
 				FString::Printf(TEXT("Failed to export level variant sets %s"),
diff --git a/Engine/Plugins/Enterprise/GLTFExporter/MigratedFromGit/Source/GLTFExporter/Private/Exporters/GLTFMaterialExporter.cpp b/Engine/Plugins/Enterprise/GLTFExporter/MigratedFromGit/Source/GLTFExporter/Private/Exporters/GLTFMaterialExporter.cpp
--- a/Engine/Plugins/Enterprise/GLTFExporter/MigratedFromGit/Source/GLTFExporter/Private/Exporters/GLTFMaterialExporter.cpp
+++ b/Engine/Plugins/Enterprise/GLTFExporter/MigratedFromGit/Source/GLTFExporter/Private/Exporters/GLTFMaterialExporter.cpp
@@ -27,8 +27,8 @@ bool UGLTFMaterialExporter::AddObject(FGLTFContainerBuilder& Builder, const UObj
 			PreviewMesh = DefaultPreviewMesh;
 		}
 
-		const FGLTFJsonMeshIndex MeshIndex = Builder.GetOrAddMesh(PreviewMesh, { Material });
-		if (MeshIndex == INDEX_NONE)
+		FGLTFJsonMesh* MeshIndex = Builder.GetOrAddMesh(PreviewMesh, { Material });
+		if (MeshIndex == nullptr)
 		{
 			Builder.LogError(
 				FString::Printf(TEXT("Failed to export preview mesh %s for material %s"),
@@ -39,11 +39,11 @@ bool UGLTFMaterialExporter::AddObject(FGLTFContainerBuilder& Builder, const UObj
 
 		FGLTFJsonNode Node;
 		Node.Mesh = MeshIndex;
-		const FGLTFJsonNodeIndex NodeIndex = Builder.AddNode(Node);
+		FGLTFJsonNode* NodeIndex = Builder.AddNode(Node);
 
 		FGLTFJsonScene Scene;
 		Scene.Nodes.Add(NodeIndex);
-		const FGLTFJsonSceneIndex SceneIndex = Builder.AddScene(Scene);
+		FGLTFJsonScene* SceneIndex = Builder.AddScene(Scene);
 
 		Builder.DefaultScene = SceneIndex;
 	}
diff --git a/Engine/Plugins/Enterprise/GLTFExporter/MigratedFromGit/Source/GLTFExporter/Private/Exporters/GLTFSkeletalMeshExporter.cpp b/Engine/Plugins/Enterprise/GLTFExporter/MigratedFromGit/Source/GLTFExporter/Private/Exporters/GLTFSkeletalMeshExporter.cpp
--- a/Engine/Plugins/Enterprise/GLTFExporter/MigratedFromGit/Source/GLTFExporter/Private/Exporters/GLTFSkeletalMeshExporter.cpp
+++ b/Engine/Plugins/Enterprise/GLTFExporter/MigratedFromGit/Source/GLTFExporter/Private/Exporters/GLTFSkeletalMeshExporter.cpp
@@ -14,8 +14,8 @@ bool UGLTFSkeletalMeshExporter::AddObject(FGLTFContainerBuilder& Builder, const
 {
 	const USkeletalMesh* SkeletalMesh = CastChecked<USkeletalMesh>(Object);
 
-	const FGLTFJsonMeshIndex MeshIndex = Builder.GetOrAddMesh(SkeletalMesh);
-	if (MeshIndex == INDEX_NONE)
+	FGLTFJsonMesh* MeshIndex = Builder.GetOrAddMesh(SkeletalMesh);
+	if (MeshIndex == nullptr)
 	{
 		Builder.LogError(FString::Printf(TEXT("Failed to export skeletal mesh %s"), *SkeletalMesh->GetName()));
 		return false;
@@ -23,12 +23,12 @@ bool UGLTFSkeletalMeshExporter::AddObject(FGLTFContainerBuilder& Builder, const
 
 	FGLTFJsonNode Node;
 	Node.Mesh = MeshIndex;
-	const FGLTFJsonNodeIndex NodeIndex = Builder.AddNode(Node);
+	FGLTFJsonNode* NodeIndex = Builder.AddNode(Node);
 
 	if (Builder.ExportOptions->bExportVertexSkinWeights)
 	{
-		const FGLTFJsonSkinIndex SkinIndex = Builder.GetOrAddSkin(NodeIndex, SkeletalMesh);
-		if (SkinIndex == INDEX_NONE)
+		FGLTFJsonSkin* SkinIndex = Builder.GetOrAddSkin(NodeIndex, SkeletalMesh);
+		if (SkinIndex == nullptr)
 		{
 			Builder.LogError(FString::Printf(TEXT("Failed to export bones in skeletal mesh %s"), *SkeletalMesh->GetName()));
 			return false;
@@ -39,7 +39,7 @@ bool UGLTFSkeletalMeshExporter::AddObject(FGLTFContainerBuilder& Builder, const
 
 	FGLTFJsonScene Scene;
 	Scene.Nodes.Add(NodeIndex);
-	const FGLTFJsonSceneIndex SceneIndex = Builder.AddScene(Scene);
+	FGLTFJsonScene* SceneIndex = Builder.AddScene(Scene);
 
 	Builder.DefaultScene = SceneIndex;
 	return true;
diff --git a/Engine/Plugins/Enterprise/GLTFExporter/MigratedFromGit/Source/GLTFExporter/Private/Exporters/GLTFStaticMeshExporter.cpp b/Engine/Plugins/Enterprise/GLTFExporter/MigratedFromGit/Source/GLTFExporter/Private/Exporters/GLTFStaticMeshExporter.cpp
--- a/Engine/Plugins/Enterprise/GLTFExporter/MigratedFromGit/Source/GLTFExporter/Private/Exporters/GLTFStaticMeshExporter.cpp
+++ b/Engine/Plugins/Enterprise/GLTFExporter/MigratedFromGit/Source/GLTFExporter/Private/Exporters/GLTFStaticMeshExporter.cpp
@@ -14,8 +14,8 @@ bool UGLTFStaticMeshExporter::AddObject(FGLTFContainerBuilder& Builder, const UO
 {
 	const UStaticMesh* StaticMesh = CastChecked<UStaticMesh>(Object);
 
-	const FGLTFJsonMeshIndex MeshIndex = Builder.GetOrAddMesh(StaticMesh);
-	if (MeshIndex == INDEX_NONE)
+	FGLTFJsonMesh* MeshIndex = Builder.GetOrAddMesh(StaticMesh);
+	if (MeshIndex == nullptr)
 	{
 		Builder.LogError(FString::Printf(TEXT("Failed to export static mesh %s"), *StaticMesh->GetName()));
 		return false;
@@ -23,11 +23,11 @@ bool UGLTFStaticMeshExporter::AddObject(FGLTFContainerBuilder& Builder, const UO
 
 	FGLTFJsonNode Node;
 	Node.Mesh = MeshIndex;
-	const FGLTFJsonNodeIndex NodeIndex = Builder.AddNode(Node);
+	FGLTFJsonNode* NodeIndex = Builder.AddNode(Node);
 
 	FGLTFJsonScene Scene;
 	Scene.Nodes.Add(NodeIndex);
-	const FGLTFJsonSceneIndex SceneIndex = Builder.AddScene(Scene);
+	FGLTFJsonScene* SceneIndex = Builder.AddScene(Scene);
 
 	Builder.DefaultScene = SceneIndex;
 	return true;
diff --git a/Engine/Plugins/Enterprise/GLTFExporter/MigratedFromGit/Source/GLTFExporter/Private/Json/GLTFJsonBackdrop.cpp b/Engine/Plugins/Enterprise/GLTFExporter/MigratedFromGit/Source/GLTFExporter/Private/Json/GLTFJsonBackdrop.cpp
--- a/Engine/Plugins/Enterprise/GLTFExporter/MigratedFromGit/Source/GLTFExporter/Private/Json/GLTFJsonBackdrop.cpp
+++ b/Engine/Plugins/Enterprise/GLTFExporter/MigratedFromGit/Source/GLTFExporter/Private/Json/GLTFJsonBackdrop.cpp
@@ -1,6 +1,8 @@
 // Copyright Epic Games, Inc. All Rights Reserved.
 
 #include "Json/GLTFJsonBackdrop.h"
+#include "Json/GLTFJsonTexture.h"
+#include "Json/GLTFJsonMesh.h"
 
 void FGLTFJsonBackdrop::WriteObject(IGLTFJsonWriter& Writer) const
 {
@@ -9,7 +11,7 @@ void FGLTFJsonBackdrop::WriteObject(IGLTFJsonWriter& Writer) const
 		Writer.Write(TEXT("name"), Name);
 	}
 
-	if (Mesh != INDEX_NONE)
+	if (Mesh != nullptr)
 	{
 		Writer.Write(TEXT("mesh"), Mesh);
 	}
diff --git a/Engine/Plugins/Enterprise/GLTFExporter/MigratedFromGit/Source/GLTFExporter/Private/Json/GLTFJsonCameraControl.cpp b/Engine/Plugins/Enterprise/GLTFExporter/MigratedFromGit/Source/GLTFExporter/Private/Json/GLTFJsonCameraControl.cpp
--- a/Engine/Plugins/Enterprise/GLTFExporter/MigratedFromGit/Source/GLTFExporter/Private/Json/GLTFJsonCameraControl.cpp
+++ b/Engine/Plugins/Enterprise/GLTFExporter/MigratedFromGit/Source/GLTFExporter/Private/Json/GLTFJsonCameraControl.cpp
@@ -1,12 +1,13 @@
 // Copyright Epic Games, Inc. All Rights Reserved.
 
 #include "Json/GLTFJsonCameraControl.h"
+#include "Json/GLTFJsonNode.h"
 
 void FGLTFJsonCameraControl::WriteObject(IGLTFJsonWriter& Writer) const
 {
 	Writer.Write(TEXT("mode"), Mode);
 
-	if (Target != INDEX_NONE && Mode == EGLTFJsonCameraControlMode::Orbital)
+	if (Target != nullptr && Mode == EGLTFJsonCameraControlMode::Orbital)
 	{
 		Writer.Write(TEXT("target"), Target);
 	}
diff --git a/Engine/Plugins/Enterprise/GLTFExporter/MigratedFromGit/Source/GLTFExporter/Private/Json/GLTFJsonEpicLevelVariantSets.cpp b/Engine/Plugins/Enterprise/GLTFExporter/MigratedFromGit/Source/GLTFExporter/Private/Json/GLTFJsonEpicLevelVariantSets.cpp
--- a/Engine/Plugins/Enterprise/GLTFExporter/MigratedFromGit/Source/GLTFExporter/Private/Json/GLTFJsonEpicLevelVariantSets.cpp
+++ b/Engine/Plugins/Enterprise/GLTFExporter/MigratedFromGit/Source/GLTFExporter/Private/Json/GLTFJsonEpicLevelVariantSets.cpp
@@ -1,6 +1,10 @@
 // Copyright Epic Games, Inc. All Rights Reserved.
 
 #include "Json/GLTFJsonEpicLevelVariantSets.h"
+#include "Json/GLTFJsonMaterial.h"
+#include "Json/GLTFJsonTexture.h"
+#include "Json/GLTFJsonMesh.h"
+#include "Json/GLTFJsonNode.h"
 
 void FGLTFJsonEpicVariantMaterial::WriteObject(IGLTFJsonWriter& Writer) const
 {
@@ -14,10 +18,7 @@ void FGLTFJsonEpicVariantMaterial::WriteObject(IGLTFJsonWriter& Writer) const
 
 void FGLTFJsonEpicVariantNodeProperties::WriteObject(IGLTFJsonWriter& Writer) const
 {
-	if (Node != INDEX_NONE)
-	{
-		Writer.Write(TEXT("node"), Node);
-	}
+	Writer.Write(TEXT("node"), Node);
 
 	Writer.StartObject(TEXT("properties"));
 
@@ -44,7 +45,7 @@ void FGLTFJsonEpicVariant::WriteObject(IGLTFJsonWriter& Writer) const
 	Writer.Write(TEXT("name"), Name);
 	Writer.Write(TEXT("active"), bIsActive);
 
-	if (Thumbnail != INDEX_NONE)
+	if (Thumbnail != nullptr)
 	{
 		Writer.Write(TEXT("thumbnail"), Thumbnail);
 	}
diff --git a/Engine/Plugins/Enterprise/GLTFExporter/MigratedFromGit/Source/GLTFExporter/Private/Json/GLTFJsonHotspot.cpp b/Engine/Plugins/Enterprise/GLTFExporter/MigratedFromGit/Source/GLTFExporter/Private/Json/GLTFJsonHotspot.cpp
--- a/Engine/Plugins/Enterprise/GLTFExporter/MigratedFromGit/Source/GLTFExporter/Private/Json/GLTFJsonHotspot.cpp
+++ b/Engine/Plugins/Enterprise/GLTFExporter/MigratedFromGit/Source/GLTFExporter/Private/Json/GLTFJsonHotspot.cpp
@@ -1,6 +1,8 @@
 // Copyright Epic Games, Inc. All Rights Reserved.
 
 #include "Json/GLTFJsonHotspot.h"
+#include "Json/GLTFJsonAnimation.h"
+#include "Json/GLTFJsonTexture.h"
 
 void FGLTFJsonHotspot::WriteObject(IGLTFJsonWriter& Writer) const
 {
@@ -12,17 +14,17 @@ void FGLTFJsonHotspot::WriteObject(IGLTFJsonWriter& Writer) const
 	Writer.Write(TEXT("animation"), Animation);
 	Writer.Write(TEXT("image"), Image);
 
-	if (HoveredImage != INDEX_NONE)
+	if (HoveredImage != nullptr)
 	{
 		Writer.Write(TEXT("hoveredImage"), HoveredImage);
 	}
 
-	if (ToggledImage != INDEX_NONE)
+	if (ToggledImage != nullptr)
 	{
 		Writer.Write(TEXT("toggledImage"), ToggledImage);
 	}
 
-	if (ToggledHoveredImage != INDEX_NONE)
+	if (ToggledHoveredImage != nullptr)
 	{
 		Writer.Write(TEXT("toggledHoveredImage"), ToggledHoveredImage);
 	}
diff --git a/Engine/Plugins/Enterprise/GLTFExporter/MigratedFromGit/Source/GLTFExporter/Private/Json/GLTFJsonImage.cpp b/Engine/Plugins/Enterprise/GLTFExporter/MigratedFromGit/Source/GLTFExporter/Private/Json/GLTFJsonImage.cpp
--- a/Engine/Plugins/Enterprise/GLTFExporter/MigratedFromGit/Source/GLTFExporter/Private/Json/GLTFJsonImage.cpp
+++ b/Engine/Plugins/Enterprise/GLTFExporter/MigratedFromGit/Source/GLTFExporter/Private/Json/GLTFJsonImage.cpp
@@ -1,6 +1,7 @@
 // Copyright Epic Games, Inc. All Rights Reserved.
 
 #include "Json/GLTFJsonImage.h"
+#include "Json/GLTFJsonBufferView.h"
 
 void FGLTFJsonImage::WriteObject(IGLTFJsonWriter& Writer) const
 {
@@ -19,7 +20,7 @@ void FGLTFJsonImage::WriteObject(IGLTFJsonWriter& Writer) const
 		Writer.Write(TEXT("mimeType"), MimeType);
 	}
 
-	if (BufferView != INDEX_NONE)
+	if (BufferView != nullptr)
 	{
 		Writer.Write(TEXT("bufferView"), BufferView);
 	}
diff --git a/Engine/Plugins/Enterprise/GLTFExporter/MigratedFromGit/Source/GLTFExporter/Private/Json/GLTFJsonLightMap.cpp b/Engine/Plugins/Enterprise/GLTFExporter/MigratedFromGit/Source/GLTFExporter/Private/Json/GLTFJsonLightMap.cpp
--- a/Engine/Plugins/Enterprise/GLTFExporter/MigratedFromGit/Source/GLTFExporter/Private/Json/GLTFJsonLightMap.cpp
+++ b/Engine/Plugins/Enterprise/GLTFExporter/MigratedFromGit/Source/GLTFExporter/Private/Json/GLTFJsonLightMap.cpp
@@ -9,7 +9,7 @@ void FGLTFJsonLightMap::WriteObject(IGLTFJsonWriter& Writer) const
 		Writer.Write(TEXT("name"), Name);
 	}
 
-	if (Texture.Index != INDEX_NONE)
+	if (Texture.Index != nullptr)
 	{
 		Writer.Write(TEXT("texture"), Texture);
 	}
diff --git a/Engine/Plugins/Enterprise/GLTFExporter/MigratedFromGit/Source/GLTFExporter/Private/Json/GLTFJsonMaterial.cpp b/Engine/Plugins/Enterprise/GLTFExporter/MigratedFromGit/Source/GLTFExporter/Private/Json/GLTFJsonMaterial.cpp
--- a/Engine/Plugins/Enterprise/GLTFExporter/MigratedFromGit/Source/GLTFExporter/Private/Json/GLTFJsonMaterial.cpp
+++ b/Engine/Plugins/Enterprise/GLTFExporter/MigratedFromGit/Source/GLTFExporter/Private/Json/GLTFJsonMaterial.cpp
@@ -1,6 +1,7 @@
 // Copyright Epic Games, Inc. All Rights Reserved.
 
 #include "Json/GLTFJsonMaterial.h"
+#include "Json/GLTFJsonTexture.h"
 
 void FGLTFJsonTextureInfo::WriteObject(IGLTFJsonWriter& Writer) const
 {
@@ -56,7 +57,7 @@ void FGLTFJsonPBRMetallicRoughness::WriteObject(IGLTFJsonWriter& Writer) const
 		Writer.Write(TEXT("baseColorFactor"), BaseColorFactor);
 	}
 
-	if (BaseColorTexture.Index != INDEX_NONE)
+	if (BaseColorTexture.Index != nullptr)
 	{
 		Writer.Write(TEXT("baseColorTexture"), BaseColorTexture);
 	}
@@ -71,7 +72,7 @@ void FGLTFJsonPBRMetallicRoughness::WriteObject(IGLTFJsonWriter& Writer) const
 		Writer.Write(TEXT("roughnessFactor"), RoughnessFactor);
 	}
 
-	if (MetallicRoughnessTexture.Index != INDEX_NONE)
+	if (MetallicRoughnessTexture.Index != nullptr)
 	{
 		Writer.Write(TEXT("metallicRoughnessTexture"), MetallicRoughnessTexture);
 	}
@@ -84,7 +85,7 @@ void FGLTFJsonClearCoatExtension::WriteObject(IGLTFJsonWriter& Writer) const
 		Writer.Write(TEXT("clearcoatFactor"), ClearCoatFactor);
 	}
 
-	if (ClearCoatTexture.Index != INDEX_NONE)
+	if (ClearCoatTexture.Index != nullptr)
 	{
 		Writer.Write(TEXT("clearcoatTexture"), ClearCoatTexture);
 	}
@@ -94,12 +95,12 @@ void FGLTFJsonClearCoatExtension::WriteObject(IGLTFJsonWriter& Writer) const
 		Writer.Write(TEXT("clearcoatRoughnessFactor"), ClearCoatRoughnessFactor);
 	}
 
-	if (ClearCoatRoughnessTexture.Index != INDEX_NONE)
+	if (ClearCoatRoughnessTexture.Index != nullptr)
 	{
 		Writer.Write(TEXT("clearcoatRoughnessTexture"), ClearCoatRoughnessTexture);
 	}
 
-	if (ClearCoatNormalTexture.Index != INDEX_NONE)
+	if (ClearCoatNormalTexture.Index != nullptr)
 	{
 		Writer.Write(TEXT("clearcoatNormalTexture"), ClearCoatNormalTexture);
 	}
@@ -117,17 +118,17 @@ void FGLTFJsonMaterial::WriteObject(IGLTFJsonWriter& Writer) const
 		Writer.Write(TEXT("pbrMetallicRoughness"), PBRMetallicRoughness);
 	}
 
-	if (NormalTexture.Index != INDEX_NONE)
+	if (NormalTexture.Index != nullptr)
 	{
 		Writer.Write(TEXT("normalTexture"), NormalTexture);
 	}
 
-	if (OcclusionTexture.Index != INDEX_NONE)
+	if (OcclusionTexture.Index != nullptr)
 	{
 		Writer.Write(TEXT("occlusionTexture"), OcclusionTexture);
 	}
 
-	if (EmissiveTexture.Index != INDEX_NONE)
+	if (EmissiveTexture.Index != nullptr)
 	{
 		Writer.Write(TEXT("emissiveTexture"), EmissiveTexture);
 	}
diff --git a/Engine/Plugins/Enterprise/GLTFExporter/MigratedFromGit/Source/GLTFExporter/Private/Json/GLTFJsonMesh.cpp b/Engine/Plugins/Enterprise/GLTFExporter/MigratedFromGit/Source/GLTFExporter/Private/Json/GLTFJsonMesh.cpp
--- a/Engine/Plugins/Enterprise/GLTFExporter/MigratedFromGit/Source/GLTFExporter/Private/Json/GLTFJsonMesh.cpp
+++ b/Engine/Plugins/Enterprise/GLTFExporter/MigratedFromGit/Source/GLTFExporter/Private/Json/GLTFJsonMesh.cpp
@@ -1,30 +1,32 @@
 // Copyright Epic Games, Inc. All Rights Reserved.
 
 #include "Json/GLTFJsonMesh.h"
+#include "Json/GLTFJsonAccessor.h"
+#include "Json/GLTFJsonMaterial.h"
 
 void FGLTFJsonAttributes::WriteObject(IGLTFJsonWriter& Writer) const
 {
-	if (Position != INDEX_NONE) Writer.Write(TEXT("POSITION"), Position);
-	if (Color0 != INDEX_NONE) Writer.Write(TEXT("COLOR_0"), Color0);
-	if (Normal != INDEX_NONE) Writer.Write(TEXT("NORMAL"), Normal);
-	if (Tangent != INDEX_NONE) Writer.Write(TEXT("TANGENT"), Tangent);
+	if (Position != nullptr) Writer.Write(TEXT("POSITION"), Position);
+	if (Color0 != nullptr) Writer.Write(TEXT("COLOR_0"), Color0);
+	if (Normal != nullptr) Writer.Write(TEXT("NORMAL"), Normal);
+	if (Tangent != nullptr) Writer.Write(TEXT("TANGENT"), Tangent);
 
 	for (int32 Index = 0; Index < TexCoords.Num(); ++Index)
 	{
-		const FGLTFJsonAccessorIndex TexCoord = TexCoords[Index];
-		if (TexCoord != INDEX_NONE) Writer.Write(TEXT("TEXCOORD_") + FString::FromInt(Index), TexCoord);
+		const FGLTFJsonAccessor* TexCoord = TexCoords[Index];
+		if (TexCoord != nullptr) Writer.Write(TEXT("TEXCOORD_") + FString::FromInt(Index), TexCoord);
 	}
 
 	for (int32 Index = 0; Index < Joints.Num(); ++Index)
 	{
-		const FGLTFJsonAccessorIndex Joint = Joints[Index];
-		if (Joint != INDEX_NONE) Writer.Write(TEXT("JOINTS_") + FString::FromInt(Index), Joint);
+		const FGLTFJsonAccessor* Joint = Joints[Index];
+		if (Joint != nullptr) Writer.Write(TEXT("JOINTS_") + FString::FromInt(Index), Joint);
 	}
 
 	for (int32 Index = 0; Index < Weights.Num(); ++Index)
 	{
-		const FGLTFJsonAccessorIndex Weight = Weights[Index];
-		if (Weight != INDEX_NONE) Writer.Write(TEXT("WEIGHTS_") + FString::FromInt(Index), Weight);
+		const FGLTFJsonAccessor* Weight = Weights[Index];
+		if (Weight != nullptr) Writer.Write(TEXT("WEIGHTS_") + FString::FromInt(Index), Weight);
 	}
 }
 
@@ -32,12 +34,12 @@ void FGLTFJsonPrimitive::WriteObject(IGLTFJsonWriter& Writer) const
 {
 	Writer.Write(TEXT("attributes"), Attributes);
 
-	if (Indices != INDEX_NONE)
+	if (Indices != nullptr)
 	{
 		Writer.Write(TEXT("indices"), Indices);
 	}
 
-	if (Material != INDEX_NONE)
+	if (Material != nullptr)
 	{
 		Writer.Write(TEXT("material"), Material);
 	}
diff --git a/Engine/Plugins/Enterprise/GLTFExporter/MigratedFromGit/Source/GLTFExporter/Private/Json/GLTFJsonNode.cpp b/Engine/Plugins/Enterprise/GLTFExporter/MigratedFromGit/Source/GLTFExporter/Private/Json/GLTFJsonNode.cpp
--- a/Engine/Plugins/Enterprise/GLTFExporter/MigratedFromGit/Source/GLTFExporter/Private/Json/GLTFJsonNode.cpp
+++ b/Engine/Plugins/Enterprise/GLTFExporter/MigratedFromGit/Source/GLTFExporter/Private/Json/GLTFJsonNode.cpp
@@ -1,6 +1,14 @@
 // Copyright Epic Games, Inc. All Rights Reserved.
 
 #include "Json/GLTFJsonNode.h"
+#include "Json/GLTFJsonCamera.h"
+#include "Json/GLTFJsonSkin.h"
+#include "Json/GLTFJsonMesh.h"
+#include "Json/GLTFJsonBackdrop.h"
+#include "Json/GLTFJsonHotspot.h"
+#include "Json/GLTFJsonLight.h"
+#include "Json/GLTFJsonLightMap.h"
+#include "Json/GLTFJsonSkySphere.h"
 
 void FGLTFJsonNode::WriteObject(IGLTFJsonWriter& Writer) const
 {
@@ -24,54 +32,54 @@ void FGLTFJsonNode::WriteObject(IGLTFJsonWriter& Writer) const
 		Writer.Write(TEXT("scale"), Scale);
 	}
 
-	if (Camera != INDEX_NONE)
+	if (Camera != nullptr)
 	{
 		Writer.Write(TEXT("camera"), Camera);
 	}
 
-	if (Skin != INDEX_NONE)
+	if (Skin != nullptr)
 	{
 		Writer.Write(TEXT("skin"), Skin);
 	}
 
-	if (Mesh != INDEX_NONE)
+	if (Mesh != nullptr)
 	{
 		Writer.Write(TEXT("mesh"), Mesh);
 	}
 
-	if (Backdrop != INDEX_NONE || Hotspot != INDEX_NONE || Light != INDEX_NONE || LightMap != INDEX_NONE || SkySphere != INDEX_NONE)
+	if (Backdrop != nullptr || Hotspot != nullptr || Light != nullptr || LightMap != nullptr || SkySphere != nullptr)
 	{
 		Writer.StartExtensions();
 
-		if (Backdrop != INDEX_NONE)
+		if (Backdrop != nullptr)
 		{
 			Writer.StartExtension(EGLTFJsonExtension::EPIC_HDRIBackdrops);
 			Writer.Write(TEXT("backdrop"), Backdrop);
 			Writer.EndExtension();
 		}
 
-		if (Hotspot != INDEX_NONE)
+		if (Hotspot != nullptr)
 		{
 			Writer.StartExtension(EGLTFJsonExtension::EPIC_AnimationHotspots);
 			Writer.Write(TEXT("hotspot"), Hotspot);
 			Writer.EndExtension();
 		}
 
-		if (Light != INDEX_NONE)
+		if (Light != nullptr)
 		{
 			Writer.StartExtension(EGLTFJsonExtension::KHR_LightsPunctual);
 			Writer.Write(TEXT("light"), Light);
 			Writer.EndExtension();
 		}
 
-		if (LightMap != INDEX_NONE)
+		if (LightMap != nullptr)
 		{
 			Writer.StartExtension(EGLTFJsonExtension::EPIC_LightmapTextures);
 			Writer.Write(TEXT("lightmap"), LightMap);
 			Writer.EndExtension();
 		}
 
-		if (SkySphere != INDEX_NONE)
+		if (SkySphere != nullptr)
 		{
 			Writer.StartExtension(EGLTFJsonExtension::EPIC_SkySpheres);
 			Writer.Write(TEXT("skySphere"), SkySphere);
diff --git a/Engine/Plugins/Enterprise/GLTFExporter/MigratedFromGit/Source/GLTFExporter/Private/Json/GLTFJsonRoot.cpp b/Engine/Plugins/Enterprise/GLTFExporter/MigratedFromGit/Source/GLTFExporter/Private/Json/GLTFJsonRoot.cpp
--- a/Engine/Plugins/Enterprise/GLTFExporter/MigratedFromGit/Source/GLTFExporter/Private/Json/GLTFJsonRoot.cpp
+++ b/Engine/Plugins/Enterprise/GLTFExporter/MigratedFromGit/Source/GLTFExporter/Private/Json/GLTFJsonRoot.cpp
@@ -6,7 +6,7 @@ void FGLTFJsonRoot::WriteObject(IGLTFJsonWriter& Writer) const
 {
 	Writer.Write(TEXT("asset"), Asset);
 
-	if (DefaultScene != INDEX_NONE)
+	if (DefaultScene != nullptr)
 	{
 		Writer.Write(TEXT("scene"), DefaultScene);
 	}
diff --git a/Engine/Plugins/Enterprise/GLTFExporter/MigratedFromGit/Source/GLTFExporter/Private/Json/GLTFJsonSkin.cpp b/Engine/Plugins/Enterprise/GLTFExporter/MigratedFromGit/Source/GLTFExporter/Private/Json/GLTFJsonSkin.cpp
--- a/Engine/Plugins/Enterprise/GLTFExporter/MigratedFromGit/Source/GLTFExporter/Private/Json/GLTFJsonSkin.cpp
+++ b/Engine/Plugins/Enterprise/GLTFExporter/MigratedFromGit/Source/GLTFExporter/Private/Json/GLTFJsonSkin.cpp
@@ -1,6 +1,8 @@
 // Copyright Epic Games, Inc. All Rights Reserved.
 
 #include "Json/GLTFJsonSkin.h"
+#include "Json/GLTFJsonAccessor.h"
+#include "Json/GLTFJsonNode.h"
 
 void FGLTFJsonSkin::WriteObject(IGLTFJsonWriter& Writer) const
 {
@@ -9,12 +11,12 @@ void FGLTFJsonSkin::WriteObject(IGLTFJsonWriter& Writer) const
 		Writer.Write(TEXT("name"), Name);
 	}
 
-	if (InverseBindMatrices != INDEX_NONE)
+	if (InverseBindMatrices != nullptr)
 	{
 		Writer.Write(TEXT("inverseBindMatrices"), InverseBindMatrices);
 	}
 
-	if (Skeleton != INDEX_NONE)
+	if (Skeleton != nullptr)
 	{
 		Writer.Write(TEXT("skeleton"), Skeleton);
 	}
diff --git a/Engine/Plugins/Enterprise/GLTFExporter/MigratedFromGit/Source/GLTFExporter/Private/Json/GLTFJsonSkySphere.cpp b/Engine/Plugins/Enterprise/GLTFExporter/MigratedFromGit/Source/GLTFExporter/Private/Json/GLTFJsonSkySphere.cpp
--- a/Engine/Plugins/Enterprise/GLTFExporter/MigratedFromGit/Source/GLTFExporter/Private/Json/GLTFJsonSkySphere.cpp
+++ b/Engine/Plugins/Enterprise/GLTFExporter/MigratedFromGit/Source/GLTFExporter/Private/Json/GLTFJsonSkySphere.cpp
@@ -1,6 +1,9 @@
 // Copyright Epic Games, Inc. All Rights Reserved.
 
 #include "Json/GLTFJsonSkySphere.h"
+#include "Json/GLTFJsonMesh.h"
+#include "Json/GLTFJsonTexture.h"
+#include "Json/GLTFJsonNode.h"
 
 void FGLTFJsonSkySphereColorCurve::WriteArray(IGLTFJsonWriter& Writer) const
 {
@@ -32,7 +35,7 @@ void FGLTFJsonSkySphere::WriteObject(IGLTFJsonWriter& Writer) const
 	Writer.Write(TEXT("cloudsTexture"), CloudsTexture);
 	Writer.Write(TEXT("starsTexture"), StarsTexture);
 
-	if (DirectionalLight != INDEX_NONE)
+	if (DirectionalLight != nullptr)
 	{
 		Writer.Write(TEXT("directionalLight"), DirectionalLight);
 	}
diff --git a/Engine/Plugins/Enterprise/GLTFExporter/MigratedFromGit/Source/GLTFExporter/Private/Json/GLTFJsonTexture.cpp b/Engine/Plugins/Enterprise/GLTFExporter/MigratedFromGit/Source/GLTFExporter/Private/Json/GLTFJsonTexture.cpp
--- a/Engine/Plugins/Enterprise/GLTFExporter/MigratedFromGit/Source/GLTFExporter/Private/Json/GLTFJsonTexture.cpp
+++ b/Engine/Plugins/Enterprise/GLTFExporter/MigratedFromGit/Source/GLTFExporter/Private/Json/GLTFJsonTexture.cpp
@@ -1,6 +1,8 @@
 // Copyright Epic Games, Inc. All Rights Reserved.
 
 #include "Json/GLTFJsonTexture.h"
+#include "Json/GLTFJsonSampler.h"
+#include "Json/GLTFJsonImage.h"
 
 void FGLTFJsonTexture::WriteObject(IGLTFJsonWriter& Writer) const
 {
@@ -9,12 +11,12 @@ void FGLTFJsonTexture::WriteObject(IGLTFJsonWriter& Writer) const
 		Writer.Write(TEXT("name"), Name);
 	}
 
-	if (Sampler != INDEX_NONE)
+	if (Sampler != nullptr)
 	{
 		Writer.Write(TEXT("sampler"), Sampler);
 	}
 
-	if (Source != INDEX_NONE)
+	if (Source != nullptr)
 	{
 		Writer.Write(TEXT("source"), Source);
 	}
diff --git a/Engine/Plugins/Enterprise/GLTFExporter/MigratedFromGit/Source/GLTFExporter/Private/Tasks/GLTFAnimationTasks.cpp b/Engine/Plugins/Enterprise/GLTFExporter/MigratedFromGit/Source/GLTFExporter/Private/Tasks/GLTFAnimationTasks.cpp
--- a/Engine/Plugins/Enterprise/GLTFExporter/MigratedFromGit/Source/GLTFExporter/Private/Tasks/GLTFAnimationTasks.cpp
+++ b/Engine/Plugins/Enterprise/GLTFExporter/MigratedFromGit/Source/GLTFExporter/Private/Tasks/GLTFAnimationTasks.cpp
@@ -27,7 +27,7 @@ void FGLTFAnimSequenceTask::Complete()
 
 	// TODO: add animation data accessor converters to reuse track information
 
-	FGLTFJsonAccessorIndex InputAccessorIndex;
+	FGLTFJsonAccessor* InputAccessorIndex;
 	{
 		FGLTFJsonAccessor JsonInputAccessor;
 		JsonInputAccessor.BufferView = Builder.AddBufferView(Timestamps);
@@ -55,7 +55,7 @@ void FGLTFAnimSequenceTask::Complete()
 
 	for (FBoneIndexType BoneIndex : BoneIndices)
 	{
-		const FGLTFJsonNodeIndex NodeIndex = Builder.GetOrAddNode(RootNode, SkeletalMesh, BoneIndex);
+		FGLTFJsonNode* NodeIndex = Builder.GetOrAddNode(RootNode, SkeletalMesh, BoneIndex);
 
 		// TODO: detect if a bone has the same transforms across multiple frames (at least if its the same across all frames) and optimize
 
@@ -81,7 +81,7 @@ void FGLTFAnimSequenceTask::Complete()
 			JsonSampler.Interpolation = Interpolation;
 
 			FGLTFJsonAnimationChannel JsonChannel;
-			JsonChannel.Sampler = FGLTFJsonAnimationSamplerIndex(JsonAnimation->Samplers.Add(JsonSampler));
+			JsonChannel.Sampler = JsonAnimation->Samplers.Add(JsonSampler);
 			JsonChannel.Target.Path = EGLTFJsonTargetPath::Translation;
 			JsonChannel.Target.Node = NodeIndex;
 			JsonAnimation->Channels.Add(JsonChannel);
@@ -109,7 +109,7 @@ void FGLTFAnimSequenceTask::Complete()
 			JsonSampler.Interpolation = Interpolation;
 
 			FGLTFJsonAnimationChannel JsonChannel;
-			JsonChannel.Sampler = FGLTFJsonAnimationSamplerIndex(JsonAnimation->Samplers.Add(JsonSampler));
+			JsonChannel.Sampler = JsonAnimation->Samplers.Add(JsonSampler);
 			JsonChannel.Target.Path = EGLTFJsonTargetPath::Rotation;
 			JsonChannel.Target.Node = NodeIndex;
 			JsonAnimation->Channels.Add(JsonChannel);
@@ -137,7 +137,7 @@ void FGLTFAnimSequenceTask::Complete()
 			JsonSampler.Interpolation = Interpolation;
 
 			FGLTFJsonAnimationChannel JsonChannel;
-			JsonChannel.Sampler = FGLTFJsonAnimationSamplerIndex(JsonAnimation->Samplers.Add(JsonSampler));
+			JsonChannel.Sampler = JsonAnimation->Samplers.Add(JsonSampler);
 			JsonChannel.Target.Path = EGLTFJsonTargetPath::Scale;
 			JsonChannel.Target.Node = NodeIndex;
 			JsonAnimation->Channels.Add(JsonChannel);
@@ -187,13 +187,13 @@ void FGLTFLevelSequenceTask::Complete()
 	JsonInputAccessor.MinMaxLength = 1;
 	JsonInputAccessor.Min[0] = 0;
 	JsonInputAccessor.Max[0] = Timestamps[FrameCount - 1];
-	FGLTFJsonAccessorIndex InputAccessorIndex = Builder.AddAccessor(JsonInputAccessor);
+	FGLTFJsonAccessor* InputAccessorIndex = Builder.AddAccessor(JsonInputAccessor);
 
 	for (const FMovieSceneBinding& Binding : MovieScene->GetBindings())
 	{
 		for (TWeakObjectPtr<UObject> Object : Player->FindBoundObjects(Binding.GetObjectGuid(), MovieSceneSequenceID::Root))
 		{
-			FGLTFJsonNodeIndex NodeIndex;
+			FGLTFJsonNode* NodeIndex;
 
 			if (AActor* Actor = Cast<AActor>(Object.Get()))
 			{
@@ -209,7 +209,7 @@ void FGLTFLevelSequenceTask::Complete()
 				continue;
 			}
 
-			if (NodeIndex == INDEX_NONE)
+			if (NodeIndex == nullptr)
 			{
 				// TODO: report warning
 				continue;
@@ -276,7 +276,7 @@ void FGLTFLevelSequenceTask::Complete()
 						JsonSampler.Interpolation = EGLTFJsonInterpolation::Linear;
 
 						FGLTFJsonAnimationChannel JsonChannel;
-						JsonChannel.Sampler = FGLTFJsonAnimationSamplerIndex(JsonAnimation->Samplers.Add(JsonSampler));
+						JsonChannel.Sampler = JsonAnimation->Samplers.Add(JsonSampler);
 						JsonChannel.Target.Path = EGLTFJsonTargetPath::Translation;
 						JsonChannel.Target.Node = NodeIndex;
 						JsonAnimation->Channels.Add(JsonChannel);
@@ -314,7 +314,7 @@ void FGLTFLevelSequenceTask::Complete()
 						JsonSampler.Interpolation = EGLTFJsonInterpolation::Linear;
 
 						FGLTFJsonAnimationChannel JsonChannel;
-						JsonChannel.Sampler = FGLTFJsonAnimationSamplerIndex(JsonAnimation->Samplers.Add(JsonSampler));
+						JsonChannel.Sampler = JsonAnimation->Samplers.Add(JsonSampler);
 						JsonChannel.Target.Path = EGLTFJsonTargetPath::Rotation;
 						JsonChannel.Target.Node = NodeIndex;
 						JsonAnimation->Channels.Add(JsonChannel);
@@ -349,7 +349,7 @@ void FGLTFLevelSequenceTask::Complete()
 						JsonSampler.Interpolation = EGLTFJsonInterpolation::Linear;
 
 						FGLTFJsonAnimationChannel JsonChannel;
-						JsonChannel.Sampler = FGLTFJsonAnimationSamplerIndex(JsonAnimation->Samplers.Add(JsonSampler));
+						JsonChannel.Sampler = JsonAnimation->Samplers.Add(JsonSampler);
 						JsonChannel.Target.Path = EGLTFJsonTargetPath::Scale;
 						JsonChannel.Target.Node = NodeIndex;
 						JsonAnimation->Channels.Add(JsonChannel);
diff --git a/Engine/Plugins/Enterprise/GLTFExporter/MigratedFromGit/Source/GLTFExporter/Private/Tasks/GLTFAnimationTasks.h b/Engine/Plugins/Enterprise/GLTFExporter/MigratedFromGit/Source/GLTFExporter/Private/Tasks/GLTFAnimationTasks.h
--- a/Engine/Plugins/Enterprise/GLTFExporter/MigratedFromGit/Source/GLTFExporter/Private/Tasks/GLTFAnimationTasks.h
+++ b/Engine/Plugins/Enterprise/GLTFExporter/MigratedFromGit/Source/GLTFExporter/Private/Tasks/GLTFAnimationTasks.h
@@ -11,7 +11,7 @@ class FGLTFAnimSequenceTask : public FGLTFTask
 {
 public:
 
-	FGLTFAnimSequenceTask(FGLTFConvertBuilder& Builder, FGLTFJsonNodeIndex RootNode, const USkeletalMesh* SkeletalMesh, const UAnimSequence* AnimSequence, FGLTFJsonAnimation* JsonAnimation)
+	FGLTFAnimSequenceTask(FGLTFConvertBuilder& Builder,  FGLTFJsonNode* RootNode, const USkeletalMesh* SkeletalMesh, const UAnimSequence* AnimSequence, FGLTFJsonAnimation* JsonAnimation)
 		: FGLTFTask(EGLTFTaskPriority::Animation)
 		, Builder(Builder)
 		, RootNode(RootNode)
@@ -31,7 +31,7 @@ public:
 private:
 
 	FGLTFConvertBuilder& Builder;
-	FGLTFJsonNodeIndex RootNode;
+	FGLTFJsonNode* RootNode;
 	const USkeletalMesh* SkeletalMesh;
 	const UAnimSequence* AnimSequence;
 	FGLTFJsonAnimation* JsonAnimation;
diff --git a/Engine/Plugins/Enterprise/GLTFExporter/MigratedFromGit/Source/GLTFExporter/Private/Tasks/GLTFMaterialTasks.cpp b/Engine/Plugins/Enterprise/GLTFExporter/MigratedFromGit/Source/GLTFExporter/Private/Tasks/GLTFMaterialTasks.cpp
--- a/Engine/Plugins/Enterprise/GLTFExporter/MigratedFromGit/Source/GLTFExporter/Private/Tasks/GLTFMaterialTasks.cpp
+++ b/Engine/Plugins/Enterprise/GLTFExporter/MigratedFromGit/Source/GLTFExporter/Private/Tasks/GLTFMaterialTasks.cpp
@@ -590,7 +590,7 @@ bool FGLTFMaterialTask::TryGetBaseColorAndOpacity(FGLTFJsonPBRMetallicRoughness&
 			return FColor(BaseColor.R, BaseColor.G, BaseColor.B, Opacity.R);
 		});
 
-	const FGLTFJsonTextureIndex TextureIndex = FGLTFMaterialUtility::AddTexture(
+	FGLTFJsonTexture* TextureIndex = FGLTFMaterialUtility::AddTexture(
 		Builder,
 		CombinedPixels,
 		TextureSize,
@@ -714,7 +714,7 @@ bool FGLTFMaterialTask::TryGetMetallicAndRoughness(FGLTFJsonPBRMetallicRoughness
 			return FColor(0, Roughness.R, Metallic.R);
 		});
 
-	const FGLTFJsonTextureIndex TextureIndex = FGLTFMaterialUtility::AddTexture(
+	FGLTFJsonTexture* TextureIndex = FGLTFMaterialUtility::AddTexture(
 		Builder,
 		CombinedPixels,
 		TextureSize,
@@ -765,7 +765,7 @@ bool FGLTFMaterialTask::TryGetClearCoatRoughness(FGLTFJsonClearCoatExtension& Ou
 		IntensityTexCoord == RoughnessTexCoord &&
 		IntensityTransform.IsExactlyEqual(RoughnessTransform))
 	{
-		const FGLTFJsonTextureIndex TextureIndex = Builder.GetOrAddTexture(IntensityTexture, false);
+		FGLTFJsonTexture* TextureIndex = Builder.GetOrAddTexture(IntensityTexture, false);
 		OutExtParams.ClearCoatTexture.Index = TextureIndex;
 		OutExtParams.ClearCoatTexture.TexCoord = IntensityTexCoord;
 		OutExtParams.ClearCoatRoughnessTexture.Index = TextureIndex;
@@ -839,7 +839,7 @@ bool FGLTFMaterialTask::TryGetClearCoatRoughness(FGLTFJsonClearCoatExtension& Ou
 			return FColor(Intensity.R, Roughness.R, 0);
 		});
 
-	const FGLTFJsonTextureIndex TextureIndex = FGLTFMaterialUtility::AddTexture(
+	FGLTFJsonTexture* TextureIndex = FGLTFMaterialUtility::AddTexture(
 		Builder,
 		CombinedPixels,
 		TextureSize,
@@ -896,7 +896,7 @@ bool FGLTFMaterialTask::TryGetEmissive(FGLTFJsonMaterial& OutMaterial, const FMa
 	{
 		if (Builder.ExportOptions->TextureImageFormat == EGLTFTextureImageFormat::None)
 		{
-			OutMaterial.EmissiveTexture.Index = FGLTFJsonTextureIndex(INDEX_NONE);
+			OutMaterial.EmissiveTexture.Index = nullptr;
 			return true;
 		}
 
@@ -1329,7 +1329,7 @@ bool FGLTFMaterialTask::TryGetBakedMaterialProperty(FGLTFJsonTextureInfo& OutTex
 
 	if (Builder.ExportOptions->TextureImageFormat == EGLTFTextureImageFormat::None)
 	{
-		OutTexInfo.Index = FGLTFJsonTextureIndex(INDEX_NONE);
+		OutTexInfo.Index = nullptr;
 		return true;
 	}
 
@@ -1363,7 +1363,7 @@ bool FGLTFMaterialTask::TryGetBakedMaterialProperty(FGLTFJsonTextureInfo& OutTex
 
 	if (Builder.ExportOptions->TextureImageFormat == EGLTFTextureImageFormat::None)
 	{
-		OutTexInfo.Index = FGLTFJsonTextureIndex(INDEX_NONE);
+		OutTexInfo.Index = nullptr;
 		return true;
 	}
 
@@ -1397,7 +1397,7 @@ inline bool FGLTFMaterialTask::TryGetBakedMaterialProperty(FGLTFJsonTextureInfo&
 
 	if (Builder.ExportOptions->TextureImageFormat == EGLTFTextureImageFormat::None)
 	{
-		OutTexInfo.Index = FGLTFJsonTextureIndex(INDEX_NONE);
+		OutTexInfo.Index = nullptr;
 		return true;
 	}
 
@@ -1427,7 +1427,7 @@ bool FGLTFMaterialTask::TryGetBakedMaterialProperty(FGLTFJsonTextureInfo& OutTex
 	{
 		if (Builder.ExportOptions->TextureImageFormat == EGLTFTextureImageFormat::None)
 		{
-			OutTexInfo.Index = FGLTFJsonTextureIndex(INDEX_NONE);
+			OutTexInfo.Index = nullptr;
 			return true;
 		}
 
@@ -1438,7 +1438,7 @@ bool FGLTFMaterialTask::TryGetBakedMaterialProperty(FGLTFJsonTextureInfo& OutTex
 	if (MaskedConstant == FGLTFMaterialUtility::GetPropertyDefaultValue(Property))
 	{
 		// Constant value is the same as the property's default so we can set gltf to default.
-		OutTexInfo.Index = FGLTFJsonTextureIndex(INDEX_NONE);
+		OutTexInfo.Index = nullptr;
 		return true;
 	}
 
@@ -1447,7 +1447,7 @@ bool FGLTFMaterialTask::TryGetBakedMaterialProperty(FGLTFJsonTextureInfo& OutTex
 		// TODO: In some cases baking normal can result in constant vector that differs slight from default (i.e 0,0,1).
 		// Yet often, when looking at such a material, it should be exactly default. Needs further investigation.
 		// Maybe because of incorrect sRGB conversion? For now, assume a constant normal is always default.
-		OutTexInfo.Index = FGLTFJsonTextureIndex(INDEX_NONE);
+		OutTexInfo.Index = nullptr;
 		return true;
 	}
 
@@ -1456,7 +1456,7 @@ bool FGLTFMaterialTask::TryGetBakedMaterialProperty(FGLTFJsonTextureInfo& OutTex
 	// even though the same material when set to opaque will properly bake AmbientOcclusion to a texture.
 	// For now, create a 1x1 texture with the constant value.
 
-	const FGLTFJsonTextureIndex TextureIndex = FGLTFMaterialUtility::AddTexture(
+	FGLTFJsonTexture* TextureIndex = FGLTFMaterialUtility::AddTexture(
 		Builder,
 		PropertyBakeOutput.Pixels,
 		PropertyBakeOutput.Size,
@@ -1533,7 +1533,7 @@ bool FGLTFMaterialTask::StoreBakedPropertyTexture(FGLTFJsonTextureInfo& OutTexIn
 	const EGLTFJsonTextureFilter TextureMinFilter = FGLTFConverterUtility::ConvertMinFilter(TextureFilter);
 	const EGLTFJsonTextureFilter TextureMagFilter = FGLTFConverterUtility::ConvertMagFilter(TextureFilter);
 
-	const FGLTFJsonTextureIndex TextureIndex = FGLTFMaterialUtility::AddTexture(
+	FGLTFJsonTexture* TextureIndex = FGLTFMaterialUtility::AddTexture(
 		Builder,
 		PropertyBakeOutput.Pixels,
 		PropertyBakeOutput.Size,
diff --git a/Engine/Plugins/Enterprise/GLTFExporter/MigratedFromGit/Source/GLTFExporter/Private/Tasks/GLTFMeshTasks.cpp b/Engine/Plugins/Enterprise/GLTFExporter/MigratedFromGit/Source/GLTFExporter/Private/Tasks/GLTFMeshTasks.cpp
--- a/Engine/Plugins/Enterprise/GLTFExporter/MigratedFromGit/Source/GLTFExporter/Private/Tasks/GLTFMeshTasks.cpp
+++ b/Engine/Plugins/Enterprise/GLTFExporter/MigratedFromGit/Source/GLTFExporter/Private/Tasks/GLTFMeshTasks.cpp
@@ -166,7 +166,7 @@ void FGLTFStaticMeshTask::Complete()
 		JsonPrimitive.Indices = Builder.GetOrAddIndexAccessor(ConvertedSection);
 
 		JsonPrimitive.Attributes.Position = Builder.GetOrAddPositionAccessor(ConvertedSection, &PositionBuffer);
-		if (JsonPrimitive.Attributes.Position == INDEX_NONE)
+		if (JsonPrimitive.Attributes.Position == nullptr)
 		{
 			FString SectionString = TEXT("section");
 			SectionString += SectionIndices.Num() > 1 ? TEXT("s ") : TEXT(" ");
@@ -269,7 +269,7 @@ void FGLTFSkeletalMeshTask::Complete()
 		JsonPrimitive.Indices = Builder.GetOrAddIndexAccessor(ConvertedSection);
 
 		JsonPrimitive.Attributes.Position = Builder.GetOrAddPositionAccessor(ConvertedSection, &PositionBuffer);
-		if (JsonPrimitive.Attributes.Position == INDEX_NONE)
+		if (JsonPrimitive.Attributes.Position == nullptr)
 		{
 			FString SectionString = TEXT("section");
 			SectionString += SectionIndices.Num() > 1 ? TEXT("s ") : TEXT(" ");
diff --git a/Engine/Plugins/Enterprise/GLTFExporter/MigratedFromGit/Source/GLTFExporter/Public/Json/GLTFJsonAccessor.h b/Engine/Plugins/Enterprise/GLTFExporter/MigratedFromGit/Source/GLTFExporter/Public/Json/GLTFJsonAccessor.h
--- a/Engine/Plugins/Enterprise/GLTFExporter/MigratedFromGit/Source/GLTFExporter/Public/Json/GLTFJsonAccessor.h
+++ b/Engine/Plugins/Enterprise/GLTFExporter/MigratedFromGit/Source/GLTFExporter/Public/Json/GLTFJsonAccessor.h
@@ -8,12 +8,12 @@ struct GLTFEXPORTER_API FGLTFJsonAccessor : IGLTFJsonIndexedObject
 {
 	FString Name;
 
-	FGLTFJsonBufferViewIndex BufferView;
-	int64                    ByteOffset;
-	int32                    Count;
-	EGLTFJsonAccessorType    Type;
-	EGLTFJsonComponentType   ComponentType;
-	bool                     bNormalized;
+	FGLTFJsonBufferView*   BufferView;
+	int64                  ByteOffset;
+	int32                  Count;
+	EGLTFJsonAccessorType  Type;
+	EGLTFJsonComponentType ComponentType;
+	bool                   bNormalized;
 
 	int32 MinMaxLength;
 	float Min[16];
@@ -21,6 +21,7 @@ struct GLTFEXPORTER_API FGLTFJsonAccessor : IGLTFJsonIndexedObject
 
 	FGLTFJsonAccessor(int32 Index = INDEX_NONE)
 		: IGLTFJsonIndexedObject(Index)
+		, BufferView(nullptr)
 		, ByteOffset(0)
 		, Count(0)
 		, Type(EGLTFJsonAccessorType::None)
diff --git a/Engine/Plugins/Enterprise/GLTFExporter/MigratedFromGit/Source/GLTFExporter/Public/Json/GLTFJsonAnimation.h b/Engine/Plugins/Enterprise/GLTFExporter/MigratedFromGit/Source/GLTFExporter/Public/Json/GLTFJsonAnimation.h
--- a/Engine/Plugins/Enterprise/GLTFExporter/MigratedFromGit/Source/GLTFExporter/Public/Json/GLTFJsonAnimation.h
+++ b/Engine/Plugins/Enterprise/GLTFExporter/MigratedFromGit/Source/GLTFExporter/Public/Json/GLTFJsonAnimation.h
@@ -7,11 +7,12 @@
 
 struct GLTFEXPORTER_API FGLTFJsonAnimationChannelTarget : IGLTFJsonObject
 {
-	FGLTFJsonNodeIndex Node;
+	FGLTFJsonNode* Node;
 	EGLTFJsonTargetPath Path;
 
 	FGLTFJsonAnimationChannelTarget()
-		: Path(EGLTFJsonTargetPath::None)
+		: Node(nullptr)
+		, Path(EGLTFJsonTargetPath::None)
 	{
 	}
 
@@ -20,21 +21,29 @@ struct GLTFEXPORTER_API FGLTFJsonAnimationChannelTarget : IGLTFJsonObject
 
 struct GLTFEXPORTER_API FGLTFJsonAnimationChannel : IGLTFJsonObject
 {
-	FGLTFJsonAnimationSamplerIndex Sampler;
+	FGLTFJsonAnimationSampler* Sampler;
 	FGLTFJsonAnimationChannelTarget Target;
 
+	FGLTFJsonAnimationChannel()
+		: Sampler(nullptr)
+	{
+	}
+
 	virtual void WriteObject(IGLTFJsonWriter& Writer) const override;
 };
 
-struct GLTFEXPORTER_API FGLTFJsonAnimationSampler : IGLTFJsonObject
+struct GLTFEXPORTER_API FGLTFJsonAnimationSampler : IGLTFJsonIndexedObject
 {
-	FGLTFJsonAccessorIndex Input;
-	FGLTFJsonAccessorIndex Output;
+	FGLTFJsonAccessor* Input;
+	FGLTFJsonAccessor* Output;
 
 	EGLTFJsonInterpolation Interpolation;
 
-	FGLTFJsonAnimationSampler()
-		: Interpolation(EGLTFJsonInterpolation::Linear)
+	FGLTFJsonAnimationSampler(int32 Index = INDEX_NONE)
+		: IGLTFJsonIndexedObject(Index)
+		, Input(nullptr)
+		, Output(nullptr)
+		, Interpolation(EGLTFJsonInterpolation::Linear)
 	{
 	}
 
@@ -46,7 +55,7 @@ struct GLTFEXPORTER_API FGLTFJsonAnimation : IGLTFJsonIndexedObject
 	FString Name;
 
 	TArray<FGLTFJsonAnimationChannel> Channels;
-	TArray<FGLTFJsonAnimationSampler> Samplers;
+	TGLTFJsonIndexedObjectArray<FGLTFJsonAnimationSampler> Samplers;
 
 	FGLTFJsonAnimationPlayback Playback;
 
diff --git a/Engine/Plugins/Enterprise/GLTFExporter/MigratedFromGit/Source/GLTFExporter/Public/Json/GLTFJsonBackdrop.h b/Engine/Plugins/Enterprise/GLTFExporter/MigratedFromGit/Source/GLTFExporter/Public/Json/GLTFJsonBackdrop.h
--- a/Engine/Plugins/Enterprise/GLTFExporter/MigratedFromGit/Source/GLTFExporter/Public/Json/GLTFJsonBackdrop.h
+++ b/Engine/Plugins/Enterprise/GLTFExporter/MigratedFromGit/Source/GLTFExporter/Public/Json/GLTFJsonBackdrop.h
@@ -8,8 +8,8 @@ struct GLTFEXPORTER_API FGLTFJsonBackdrop : IGLTFJsonIndexedObject
 {
 	FString Name;
 
-	FGLTFJsonMeshIndex Mesh;
-	FGLTFJsonTextureIndex Cubemap[6];
+	FGLTFJsonMesh* Mesh;
+	FGLTFJsonTexture* Cubemap[6];
 
 	float Intensity;
 	float Size;
@@ -22,6 +22,8 @@ struct GLTFEXPORTER_API FGLTFJsonBackdrop : IGLTFJsonIndexedObject
 
 	FGLTFJsonBackdrop(int32 Index = INDEX_NONE)
 		: IGLTFJsonIndexedObject(Index)
+		, Mesh(nullptr)
+		, Cubemap{nullptr}
 		, Intensity(0)
 		, Size(0)
 		, Angle(0)
diff --git a/Engine/Plugins/Enterprise/GLTFExporter/MigratedFromGit/Source/GLTFExporter/Public/Json/GLTFJsonBufferView.h b/Engine/Plugins/Enterprise/GLTFExporter/MigratedFromGit/Source/GLTFExporter/Public/Json/GLTFJsonBufferView.h
--- a/Engine/Plugins/Enterprise/GLTFExporter/MigratedFromGit/Source/GLTFExporter/Public/Json/GLTFJsonBufferView.h
+++ b/Engine/Plugins/Enterprise/GLTFExporter/MigratedFromGit/Source/GLTFExporter/Public/Json/GLTFJsonBufferView.h
@@ -8,7 +8,7 @@ struct GLTFEXPORTER_API FGLTFJsonBufferView : IGLTFJsonIndexedObject
 {
 	FString Name;
 
-	FGLTFJsonBufferIndex Buffer;
+	FGLTFJsonBuffer* Buffer;
 
 	int64 ByteLength;
 	int64 ByteOffset;
@@ -18,6 +18,7 @@ struct GLTFEXPORTER_API FGLTFJsonBufferView : IGLTFJsonIndexedObject
 
 	FGLTFJsonBufferView(int32 Index = INDEX_NONE)
 		: IGLTFJsonIndexedObject(Index)
+		, Buffer(nullptr)
 		, ByteLength(0)
 		, ByteOffset(0)
 		, ByteStride(0)
diff --git a/Engine/Plugins/Enterprise/GLTFExporter/MigratedFromGit/Source/GLTFExporter/Public/Json/GLTFJsonCameraControl.h b/Engine/Plugins/Enterprise/GLTFExporter/MigratedFromGit/Source/GLTFExporter/Public/Json/GLTFJsonCameraControl.h
--- a/Engine/Plugins/Enterprise/GLTFExporter/MigratedFromGit/Source/GLTFExporter/Public/Json/GLTFJsonCameraControl.h
+++ b/Engine/Plugins/Enterprise/GLTFExporter/MigratedFromGit/Source/GLTFExporter/Public/Json/GLTFJsonCameraControl.h
@@ -7,7 +7,8 @@
 struct GLTFEXPORTER_API FGLTFJsonCameraControl : IGLTFJsonObject
 {
 	EGLTFJsonCameraControlMode Mode;
-	FGLTFJsonNodeIndex Target;
+	FGLTFJsonNode* Target;
+
 	float MaxDistance;
 	float MinDistance;
 	float MaxPitch;
@@ -21,6 +22,7 @@ struct GLTFEXPORTER_API FGLTFJsonCameraControl : IGLTFJsonObject
 
 	FGLTFJsonCameraControl()
 		: Mode(EGLTFJsonCameraControlMode::FreeLook)
+		, Target(nullptr)
 		, MaxDistance(0)
 		, MinDistance(0)
 		, MaxPitch(90)
diff --git a/Engine/Plugins/Enterprise/GLTFExporter/MigratedFromGit/Source/GLTFExporter/Public/Json/GLTFJsonCore.h b/Engine/Plugins/Enterprise/GLTFExporter/MigratedFromGit/Source/GLTFExporter/Public/Json/GLTFJsonCore.h
--- a/Engine/Plugins/Enterprise/GLTFExporter/MigratedFromGit/Source/GLTFExporter/Public/Json/GLTFJsonCore.h
+++ b/Engine/Plugins/Enterprise/GLTFExporter/MigratedFromGit/Source/GLTFExporter/Public/Json/GLTFJsonCore.h
@@ -3,7 +3,6 @@
 #pragma once
 
 #include "Json/GLTFJsonEnums.h"
-#include "Json/GLTFJsonIndex.h"
 #include "Json/GLTFJsonValue.h"
 #include "Json/GLTFJsonArray.h"
 #include "Json/GLTFJsonObject.h"
@@ -14,3 +13,25 @@
 #include "Json/GLTFJsonQuaternion.h"
 #include "Json/GLTFJsonVector.h"
 #include "Json/GLTFJsonWriter.h"
+
+struct FGLTFJsonAccessor;
+struct FGLTFJsonAnimation;
+struct FGLTFJsonAnimationSampler;
+struct FGLTFJsonBuffer;
+struct FGLTFJsonBufferView;
+struct FGLTFJsonCamera;
+struct FGLTFJsonImage;
+struct FGLTFJsonMaterial;
+struct FGLTFJsonMesh;
+struct FGLTFJsonNode;
+struct FGLTFJsonSampler;
+struct FGLTFJsonScene;
+struct FGLTFJsonSkin;
+struct FGLTFJsonTexture;
+struct FGLTFJsonBackdrop;
+struct FGLTFJsonHotspot;
+struct FGLTFJsonLight;
+struct FGLTFJsonLightMap;
+struct FGLTFJsonSkySphere;
+struct FGLTFJsonEpicLevelVariantSets;
+struct FGLTFJsonKhrMaterialVariant;
diff --git a/Engine/Plugins/Enterprise/GLTFExporter/MigratedFromGit/Source/GLTFExporter/Public/Json/GLTFJsonEpicLevelVariantSets.h b/Engine/Plugins/Enterprise/GLTFExporter/MigratedFromGit/Source/GLTFExporter/Public/Json/GLTFJsonEpicLevelVariantSets.h
--- a/Engine/Plugins/Enterprise/GLTFExporter/MigratedFromGit/Source/GLTFExporter/Public/Json/GLTFJsonEpicLevelVariantSets.h
+++ b/Engine/Plugins/Enterprise/GLTFExporter/MigratedFromGit/Source/GLTFExporter/Public/Json/GLTFJsonEpicLevelVariantSets.h
@@ -6,11 +6,11 @@
 
 struct GLTFEXPORTER_API FGLTFJsonEpicVariantMaterial : IGLTFJsonObject
 {
-	FGLTFJsonMaterialIndex Material;
-	int32                  Index;
+	FGLTFJsonMaterial* Material;
+	int32 Index;
 
 	FGLTFJsonEpicVariantMaterial()
-		: Material(INDEX_NONE)
+		: Material(nullptr)
 		, Index(INDEX_NONE)
 	{
 	}
@@ -20,22 +20,33 @@ struct GLTFEXPORTER_API FGLTFJsonEpicVariantMaterial : IGLTFJsonObject
 
 struct GLTFEXPORTER_API FGLTFJsonEpicVariantNodeProperties : IGLTFJsonObject
 {
-	FGLTFJsonNodeIndex Node;
-	TOptional<bool>    bIsVisible;
+	FGLTFJsonNode* Node;
+	TOptional<bool> bIsVisible;
 
-	TOptional<FGLTFJsonMeshIndex>        Mesh;
+	TOptional<FGLTFJsonMesh*> Mesh;
 	TArray<FGLTFJsonEpicVariantMaterial> Materials;
 
+	FGLTFJsonEpicVariantNodeProperties()
+		: Node(nullptr)
+	{
+	}
+
 	virtual void WriteObject(IGLTFJsonWriter& Writer) const override;
 };
 
 struct GLTFEXPORTER_API FGLTFJsonEpicVariant : IGLTFJsonObject
 {
 	FString Name;
-	bool    bIsActive;
+	bool bIsActive;
 
-	FGLTFJsonTextureIndex Thumbnail;
-	TMap<FGLTFJsonNodeIndex, FGLTFJsonEpicVariantNodeProperties> Nodes;
+	FGLTFJsonTexture* Thumbnail;
+	TMap<FGLTFJsonNode*, FGLTFJsonEpicVariantNodeProperties> Nodes;
+
+	FGLTFJsonEpicVariant()
+		: bIsActive(false)
+		, Thumbnail(nullptr)
+	{
+	}
 
 	virtual void WriteObject(IGLTFJsonWriter& Writer) const override;
 };
diff --git a/Engine/Plugins/Enterprise/GLTFExporter/MigratedFromGit/Source/GLTFExporter/Public/Json/GLTFJsonHotspot.h b/Engine/Plugins/Enterprise/GLTFExporter/MigratedFromGit/Source/GLTFExporter/Public/Json/GLTFJsonHotspot.h
--- a/Engine/Plugins/Enterprise/GLTFExporter/MigratedFromGit/Source/GLTFExporter/Public/Json/GLTFJsonHotspot.h
+++ b/Engine/Plugins/Enterprise/GLTFExporter/MigratedFromGit/Source/GLTFExporter/Public/Json/GLTFJsonHotspot.h
@@ -8,14 +8,19 @@ struct GLTFEXPORTER_API FGLTFJsonHotspot : IGLTFJsonIndexedObject
 {
 	FString Name;
 
-	FGLTFJsonAnimationIndex Animation;
-	FGLTFJsonTextureIndex   Image;
-	FGLTFJsonTextureIndex   HoveredImage;
-	FGLTFJsonTextureIndex   ToggledImage;
-	FGLTFJsonTextureIndex   ToggledHoveredImage;
+	FGLTFJsonAnimation* Animation;
+	FGLTFJsonTexture*   Image;
+	FGLTFJsonTexture*   HoveredImage;
+	FGLTFJsonTexture*   ToggledImage;
+	FGLTFJsonTexture*   ToggledHoveredImage;
 
 	FGLTFJsonHotspot(int32 Index = INDEX_NONE)
 		: IGLTFJsonIndexedObject(Index)
+		, Animation(nullptr)
+		, Image(nullptr)
+		, HoveredImage(nullptr)
+		, ToggledImage(nullptr)
+		, ToggledHoveredImage(nullptr)
 	{
 	}
 
diff --git a/Engine/Plugins/Enterprise/GLTFExporter/MigratedFromGit/Source/GLTFExporter/Public/Json/GLTFJsonImage.h b/Engine/Plugins/Enterprise/GLTFExporter/MigratedFromGit/Source/GLTFExporter/Public/Json/GLTFJsonImage.h
--- a/Engine/Plugins/Enterprise/GLTFExporter/MigratedFromGit/Source/GLTFExporter/Public/Json/GLTFJsonImage.h
+++ b/Engine/Plugins/Enterprise/GLTFExporter/MigratedFromGit/Source/GLTFExporter/Public/Json/GLTFJsonImage.h
@@ -11,11 +11,12 @@ struct GLTFEXPORTER_API FGLTFJsonImage : IGLTFJsonIndexedObject
 
 	EGLTFJsonMimeType MimeType;
 
-	FGLTFJsonBufferViewIndex BufferView;
+	FGLTFJsonBufferView* BufferView;
 
 	FGLTFJsonImage(int32 Index = INDEX_NONE)
 		: IGLTFJsonIndexedObject(Index)
 		, MimeType(EGLTFJsonMimeType::None)
+		, BufferView(nullptr)
 	{
 	}
 
diff --git a/Engine/Plugins/Enterprise/GLTFExporter/MigratedFromGit/Source/GLTFExporter/Public/Json/GLTFJsonIndex.h b/Engine/Plugins/Enterprise/GLTFExporter/MigratedFromGit/Source/GLTFExporter/Public/Json/GLTFJsonIndex.h
deleted file mode 100644
diff --git a/Engine/Plugins/Enterprise/GLTFExporter/MigratedFromGit/Source/GLTFExporter/Public/Json/GLTFJsonKhrMaterialVariant.h b/Engine/Plugins/Enterprise/GLTFExporter/MigratedFromGit/Source/GLTFExporter/Public/Json/GLTFJsonKhrMaterialVariant.h
--- a/Engine/Plugins/Enterprise/GLTFExporter/MigratedFromGit/Source/GLTFExporter/Public/Json/GLTFJsonKhrMaterialVariant.h
+++ b/Engine/Plugins/Enterprise/GLTFExporter/MigratedFromGit/Source/GLTFExporter/Public/Json/GLTFJsonKhrMaterialVariant.h
@@ -6,8 +6,13 @@
 
 struct GLTFEXPORTER_API FGLTFJsonKhrMaterialVariantMapping : IGLTFJsonObject
 {
-	FGLTFJsonMaterialIndex                   Material;
-	TArray<FGLTFJsonKhrMaterialVariantIndex> Variants;
+	FGLTFJsonMaterial* Material;
+	TArray<FGLTFJsonKhrMaterialVariant*> Variants;
+
+	FGLTFJsonKhrMaterialVariantMapping()
+		: Material(nullptr)
+	{
+	}
 
 	virtual void WriteObject(IGLTFJsonWriter& Writer) const override;
 };
diff --git a/Engine/Plugins/Enterprise/GLTFExporter/MigratedFromGit/Source/GLTFExporter/Public/Json/GLTFJsonMaterial.h b/Engine/Plugins/Enterprise/GLTFExporter/MigratedFromGit/Source/GLTFExporter/Public/Json/GLTFJsonMaterial.h
--- a/Engine/Plugins/Enterprise/GLTFExporter/MigratedFromGit/Source/GLTFExporter/Public/Json/GLTFJsonMaterial.h
+++ b/Engine/Plugins/Enterprise/GLTFExporter/MigratedFromGit/Source/GLTFExporter/Public/Json/GLTFJsonMaterial.h
@@ -7,13 +7,14 @@
 
 struct GLTFEXPORTER_API FGLTFJsonTextureInfo : IGLTFJsonObject
 {
-	FGLTFJsonTextureIndex Index;
+	FGLTFJsonTexture* Index;
 	int32 TexCoord;
 
 	FGLTFJsonTextureTransform Transform;
 
 	FGLTFJsonTextureInfo()
-		: TexCoord(0)
+		: Index(nullptr)
+		, TexCoord(0)
 	{
 	}
 
diff --git a/Engine/Plugins/Enterprise/GLTFExporter/MigratedFromGit/Source/GLTFExporter/Public/Json/GLTFJsonMesh.h b/Engine/Plugins/Enterprise/GLTFExporter/MigratedFromGit/Source/GLTFExporter/Public/Json/GLTFJsonMesh.h
--- a/Engine/Plugins/Enterprise/GLTFExporter/MigratedFromGit/Source/GLTFExporter/Public/Json/GLTFJsonMesh.h
+++ b/Engine/Plugins/Enterprise/GLTFExporter/MigratedFromGit/Source/GLTFExporter/Public/Json/GLTFJsonMesh.h
@@ -7,29 +7,39 @@
 
 struct GLTFEXPORTER_API FGLTFJsonAttributes : IGLTFJsonObject
 {
-	FGLTFJsonAccessorIndex Position;
-	FGLTFJsonAccessorIndex Color0;
-	FGLTFJsonAccessorIndex Normal;
-	FGLTFJsonAccessorIndex Tangent;
-
-	TArray<FGLTFJsonAccessorIndex> TexCoords;
-	TArray<FGLTFJsonAccessorIndex> Joints;
-	TArray<FGLTFJsonAccessorIndex> Weights;
+	FGLTFJsonAccessor* Position;
+	FGLTFJsonAccessor* Color0;
+	FGLTFJsonAccessor* Normal;
+	FGLTFJsonAccessor* Tangent;
+
+	TArray<FGLTFJsonAccessor*> TexCoords;
+	TArray<FGLTFJsonAccessor*> Joints;
+	TArray<FGLTFJsonAccessor*> Weights;
+
+	FGLTFJsonAttributes()
+		: Position(nullptr)
+		, Color0(nullptr)
+		, Normal(nullptr)
+		, Tangent(nullptr)
+	{
+	}
 
 	virtual void WriteObject(IGLTFJsonWriter& Writer) const override;
 };
 
 struct GLTFEXPORTER_API FGLTFJsonPrimitive : IGLTFJsonObject
 {
-	FGLTFJsonAccessorIndex Indices;
-	FGLTFJsonMaterialIndex Material;
+	FGLTFJsonAttributes Attributes;
+	FGLTFJsonAccessor* Indices;
+	FGLTFJsonMaterial* Material;
 	EGLTFJsonPrimitiveMode Mode;
-	FGLTFJsonAttributes    Attributes;
 
 	TArray<FGLTFJsonKhrMaterialVariantMapping> KhrMaterialVariantMappings;
 
 	FGLTFJsonPrimitive()
-		: Mode(EGLTFJsonPrimitiveMode::Triangles)
+		: Indices(nullptr)
+		, Material(nullptr)
+		, Mode(EGLTFJsonPrimitiveMode::Triangles)
 	{
 	}
 
diff --git a/Engine/Plugins/Enterprise/GLTFExporter/MigratedFromGit/Source/GLTFExporter/Public/Json/GLTFJsonNode.h b/Engine/Plugins/Enterprise/GLTFExporter/MigratedFromGit/Source/GLTFExporter/Public/Json/GLTFJsonNode.h
--- a/Engine/Plugins/Enterprise/GLTFExporter/MigratedFromGit/Source/GLTFExporter/Public/Json/GLTFJsonNode.h
+++ b/Engine/Plugins/Enterprise/GLTFExporter/MigratedFromGit/Source/GLTFExporter/Public/Json/GLTFJsonNode.h
@@ -12,24 +12,33 @@ struct GLTFEXPORTER_API FGLTFJsonNode : IGLTFJsonIndexedObject
 	FGLTFJsonQuaternion Rotation;
 	FGLTFJsonVector3    Scale;
 
-	FGLTFJsonCameraIndex    Camera;
-	FGLTFJsonSkinIndex      Skin;
-	FGLTFJsonMeshIndex      Mesh;
-	FGLTFJsonBackdropIndex  Backdrop;
-	FGLTFJsonHotspotIndex   Hotspot;
-	FGLTFJsonLightIndex     Light;
-	FGLTFJsonLightMapIndex  LightMap;
-	FGLTFJsonSkySphereIndex SkySphere;
+	FGLTFJsonCamera*    Camera;
+	FGLTFJsonSkin*      Skin;
+	FGLTFJsonMesh*      Mesh;
+	FGLTFJsonBackdrop*  Backdrop;
+	FGLTFJsonHotspot*   Hotspot;
+	FGLTFJsonLight*     Light;
+	FGLTFJsonLightMap*  LightMap;
+	FGLTFJsonSkySphere* SkySphere;
 
-	FGLTFJsonNodeIndex ComponentNode;
+	FGLTFJsonNode* ComponentNode;
 
-	TArray<FGLTFJsonNodeIndex> Children;
+	TArray<FGLTFJsonNode*> Children;
 
 	FGLTFJsonNode(int32 Index = INDEX_NONE)
 		: IGLTFJsonIndexedObject(Index)
 		, Translation(FGLTFJsonVector3::Zero)
 		, Rotation(FGLTFJsonQuaternion::Identity)
 		, Scale(FGLTFJsonVector3::One)
+		, Camera(nullptr)
+		, Skin(nullptr)
+		, Mesh(nullptr)
+		, Backdrop(nullptr)
+		, Hotspot(nullptr)
+		, Light(nullptr)
+		, LightMap(nullptr)
+		, SkySphere(nullptr)
+		, ComponentNode(nullptr)
 	{
 	}
 
diff --git a/Engine/Plugins/Enterprise/GLTFExporter/MigratedFromGit/Source/GLTFExporter/Public/Json/GLTFJsonRoot.h b/Engine/Plugins/Enterprise/GLTFExporter/MigratedFromGit/Source/GLTFExporter/Public/Json/GLTFJsonRoot.h
--- a/Engine/Plugins/Enterprise/GLTFExporter/MigratedFromGit/Source/GLTFExporter/Public/Json/GLTFJsonRoot.h
+++ b/Engine/Plugins/Enterprise/GLTFExporter/MigratedFromGit/Source/GLTFExporter/Public/Json/GLTFJsonRoot.h
@@ -31,7 +31,7 @@ struct GLTFEXPORTER_API FGLTFJsonRoot : IGLTFJsonObject
 
 	FGLTFJsonExtensions Extensions;
 
-	FGLTFJsonSceneIndex DefaultScene;
+	FGLTFJsonScene* DefaultScene;
 
 	TGLTFJsonIndexedObjectArray<FGLTFJsonAccessor>   Accessors;
 	TGLTFJsonIndexedObjectArray<FGLTFJsonAnimation>  Animations;
@@ -54,7 +54,11 @@ struct GLTFEXPORTER_API FGLTFJsonRoot : IGLTFJsonObject
 	TGLTFJsonIndexedObjectArray<FGLTFJsonEpicLevelVariantSets> EpicLevelVariantSets;
 	TGLTFJsonIndexedObjectArray<FGLTFJsonKhrMaterialVariant>   KhrMaterialVariants;
 
-	FGLTFJsonRoot() = default;
+	FGLTFJsonRoot()
+		: DefaultScene(nullptr)
+	{
+	}
+
 	FGLTFJsonRoot(FGLTFJsonRoot&&) = default;
 	FGLTFJsonRoot& operator=(FGLTFJsonRoot&&) = default;
 
diff --git a/Engine/Plugins/Enterprise/GLTFExporter/MigratedFromGit/Source/GLTFExporter/Public/Json/GLTFJsonScene.h b/Engine/Plugins/Enterprise/GLTFExporter/MigratedFromGit/Source/GLTFExporter/Public/Json/GLTFJsonScene.h
--- a/Engine/Plugins/Enterprise/GLTFExporter/MigratedFromGit/Source/GLTFExporter/Public/Json/GLTFJsonScene.h
+++ b/Engine/Plugins/Enterprise/GLTFExporter/MigratedFromGit/Source/GLTFExporter/Public/Json/GLTFJsonScene.h
@@ -8,8 +8,8 @@ struct GLTFEXPORTER_API FGLTFJsonScene : IGLTFJsonIndexedObject
 {
 	FString Name;
 
-	TArray<FGLTFJsonNodeIndex> Nodes;
-	TArray<FGLTFJsonEpicLevelVariantSetsIndex>  EpicLevelVariantSets;
+	TArray<FGLTFJsonNode*> Nodes;
+	TArray<FGLTFJsonEpicLevelVariantSets*>  EpicLevelVariantSets;
 
 	FGLTFJsonScene(int32 Index = INDEX_NONE)
 		: IGLTFJsonIndexedObject(Index)
diff --git a/Engine/Plugins/Enterprise/GLTFExporter/MigratedFromGit/Source/GLTFExporter/Public/Json/GLTFJsonSkin.h b/Engine/Plugins/Enterprise/GLTFExporter/MigratedFromGit/Source/GLTFExporter/Public/Json/GLTFJsonSkin.h
--- a/Engine/Plugins/Enterprise/GLTFExporter/MigratedFromGit/Source/GLTFExporter/Public/Json/GLTFJsonSkin.h
+++ b/Engine/Plugins/Enterprise/GLTFExporter/MigratedFromGit/Source/GLTFExporter/Public/Json/GLTFJsonSkin.h
@@ -8,13 +8,15 @@ struct GLTFEXPORTER_API FGLTFJsonSkin : IGLTFJsonIndexedObject
 {
 	FString Name;
 
-	FGLTFJsonAccessorIndex InverseBindMatrices;
-	FGLTFJsonNodeIndex Skeleton;
+	FGLTFJsonAccessor* InverseBindMatrices;
+	FGLTFJsonNode* Skeleton;
 
-	TArray<FGLTFJsonNodeIndex> Joints;
+	TArray<FGLTFJsonNode*> Joints;
 
 	FGLTFJsonSkin(int32 Index = INDEX_NONE)
 		: IGLTFJsonIndexedObject(Index)
+		, InverseBindMatrices(nullptr)
+		, Skeleton(nullptr)
 	{
 	}
 
diff --git a/Engine/Plugins/Enterprise/GLTFExporter/MigratedFromGit/Source/GLTFExporter/Public/Json/GLTFJsonSkySphere.h b/Engine/Plugins/Enterprise/GLTFExporter/MigratedFromGit/Source/GLTFExporter/Public/Json/GLTFJsonSkySphere.h
--- a/Engine/Plugins/Enterprise/GLTFExporter/MigratedFromGit/Source/GLTFExporter/Public/Json/GLTFJsonSkySphere.h
+++ b/Engine/Plugins/Enterprise/GLTFExporter/MigratedFromGit/Source/GLTFExporter/Public/Json/GLTFJsonSkySphere.h
@@ -26,11 +26,11 @@ struct GLTFEXPORTER_API FGLTFJsonSkySphere : IGLTFJsonIndexedObject
 {
 	FString Name;
 
-	FGLTFJsonMeshIndex    SkySphereMesh;
-	FGLTFJsonTextureIndex SkyTexture;
-	FGLTFJsonTextureIndex CloudsTexture;
-	FGLTFJsonTextureIndex StarsTexture;
-	FGLTFJsonNodeIndex    DirectionalLight;
+	FGLTFJsonMesh*    SkySphereMesh;
+	FGLTFJsonTexture* SkyTexture;
+	FGLTFJsonTexture* CloudsTexture;
+	FGLTFJsonTexture* StarsTexture;
+	FGLTFJsonNode*    DirectionalLight;
 
 	float SunHeight;
 	float SunBrightness;
@@ -58,6 +58,11 @@ struct GLTFEXPORTER_API FGLTFJsonSkySphere : IGLTFJsonIndexedObject
 
 	FGLTFJsonSkySphere(int32 Index = INDEX_NONE)
 		: IGLTFJsonIndexedObject(Index)
+		, SkySphereMesh(nullptr)
+		, SkyTexture(nullptr)
+		, CloudsTexture(nullptr)
+		, StarsTexture(nullptr)
+		, DirectionalLight(nullptr)
 		, SunHeight(0)
 		, SunBrightness(0)
 		, StarsBrightness(0)
diff --git a/Engine/Plugins/Enterprise/GLTFExporter/MigratedFromGit/Source/GLTFExporter/Public/Json/GLTFJsonTexture.h b/Engine/Plugins/Enterprise/GLTFExporter/MigratedFromGit/Source/GLTFExporter/Public/Json/GLTFJsonTexture.h
--- a/Engine/Plugins/Enterprise/GLTFExporter/MigratedFromGit/Source/GLTFExporter/Public/Json/GLTFJsonTexture.h
+++ b/Engine/Plugins/Enterprise/GLTFExporter/MigratedFromGit/Source/GLTFExporter/Public/Json/GLTFJsonTexture.h
@@ -8,14 +8,16 @@ struct GLTFEXPORTER_API FGLTFJsonTexture : IGLTFJsonIndexedObject
 {
 	FString Name;
 
-	FGLTFJsonSamplerIndex Sampler;
+	FGLTFJsonSampler* Sampler;
 
-	FGLTFJsonImageIndex Source;
+	FGLTFJsonImage* Source;
 
 	EGLTFJsonHDREncoding Encoding;
 
 	FGLTFJsonTexture(int32 Index = INDEX_NONE)
 		: IGLTFJsonIndexedObject(Index)
+		, Sampler(nullptr)
+		, Source(nullptr)
 		, Encoding(EGLTFJsonHDREncoding::None)
 	{
 	}
