[GLTFExporter] refactor: remove all old json data structure adders in builder

commit c663a91e8244cb9d9c472dc99d9cbfa4a8973488
Date: Sun, 13 Mar 2022 17:52:20 +0100

    refactor: remove all old json data structure adders in builder
    (cherry picked from commit a99ff683d0a0b6964cc225ecde2e10198bd51dbb​)

#skipci

diff --git a/Engine/Plugins/Enterprise/GLTFExporter/MigratedFromGit/Source/GLTFExporter/Private/Builders/GLTFBufferBuilder.cpp b/Engine/Plugins/Enterprise/GLTFExporter/MigratedFromGit/Source/GLTFExporter/Private/Builders/GLTFBufferBuilder.cpp
--- a/Engine/Plugins/Enterprise/GLTFExporter/MigratedFromGit/Source/GLTFExporter/Private/Builders/GLTFBufferBuilder.cpp
+++ b/Engine/Plugins/Enterprise/GLTFExporter/MigratedFromGit/Source/GLTFExporter/Private/Builders/GLTFBufferBuilder.cpp
@@ -20,7 +20,7 @@ FGLTFBufferBuilder::~FGLTFBufferBuilder()
 
 bool FGLTFBufferBuilder::InitializeBuffer()
 {
-	FGLTFJsonBuffer JsonBuffer;
+	JsonBuffer = FGLTFJsonBuilder::AddBuffer();
 
 	if (bIsGlbFile)
 	{
@@ -29,7 +29,7 @@ bool FGLTFBufferBuilder::InitializeBuffer()
 	else
 	{
 		const FString ExternalBinaryPath = FPaths::ChangeExtension(FilePath, TEXT(".bin"));
-		JsonBuffer.URI = FPaths::GetCleanFilename(ExternalBinaryPath);
+		JsonBuffer->URI = FPaths::GetCleanFilename(ExternalBinaryPath);
 
 		BufferArchive.Reset(IFileManager::Get().CreateFileWriter(*ExternalBinaryPath));
 		if (BufferArchive == nullptr)
@@ -39,7 +39,6 @@ bool FGLTFBufferBuilder::InitializeBuffer()
 		}
 	}
 
-	BufferIndex = AddBuffer(JsonBuffer);
 	return true;
 }
 
@@ -67,12 +66,10 @@ FGLTFJsonBufferView* FGLTFBufferBuilder::AddBufferView(const void* RawData, uint
 	}
 
 	BufferArchive->Serialize(const_cast<void*>(RawData), ByteLength);
-
-	FGLTFJsonBuffer& JsonBuffer = GetBuffer(BufferIndex);
-	JsonBuffer.ByteLength = BufferArchive->Tell();
+	JsonBuffer->ByteLength = BufferArchive->Tell();
 
 	FGLTFJsonBufferView* JsonBufferView = FGLTFJsonBuilder::AddBufferView();
-	JsonBufferView->Buffer = BufferIndex;
+	JsonBufferView->Buffer = JsonBuffer;
 	JsonBufferView->ByteOffset = ByteOffset;
 	JsonBufferView->ByteLength = ByteLength;
 	JsonBufferView->Target = BufferTarget;
diff --git a/Engine/Plugins/Enterprise/GLTFExporter/MigratedFromGit/Source/GLTFExporter/Private/Builders/GLTFBufferBuilder.h b/Engine/Plugins/Enterprise/GLTFExporter/MigratedFromGit/Source/GLTFExporter/Private/Builders/GLTFBufferBuilder.h
--- a/Engine/Plugins/Enterprise/GLTFExporter/MigratedFromGit/Source/GLTFExporter/Private/Builders/GLTFBufferBuilder.h
+++ b/Engine/Plugins/Enterprise/GLTFExporter/MigratedFromGit/Source/GLTFExporter/Private/Builders/GLTFBufferBuilder.h
@@ -28,6 +28,6 @@ private:
 
 	bool InitializeBuffer();
 
-	FGLTFJsonBuffer* BufferIndex;
+	FGLTFJsonBuffer* JsonBuffer;
 	TUniquePtr<FArchive> BufferArchive;
 };
diff --git a/Engine/Plugins/Enterprise/GLTFExporter/MigratedFromGit/Source/GLTFExporter/Private/Builders/GLTFJsonBuilder.cpp b/Engine/Plugins/Enterprise/GLTFExporter/MigratedFromGit/Source/GLTFExporter/Private/Builders/GLTFJsonBuilder.cpp
--- a/Engine/Plugins/Enterprise/GLTFExporter/MigratedFromGit/Source/GLTFExporter/Private/Builders/GLTFJsonBuilder.cpp
+++ b/Engine/Plugins/Enterprise/GLTFExporter/MigratedFromGit/Source/GLTFExporter/Private/Builders/GLTFJsonBuilder.cpp
@@ -141,144 +141,31 @@ FGLTFJsonKhrMaterialVariant* FGLTFJsonBuilder::AddKhrMaterialVariant()
 	return JsonRoot.KhrMaterialVariants.Add();
 }
 
-FGLTFJsonNode* FGLTFJsonBuilder::AddChildNode(FGLTFJsonNode* ParentNode)
+FGLTFJsonNode* FGLTFJsonBuilder::AddChildNode(FGLTFJsonNode* ParentNode, FGLTFJsonNode* ChildNode)
 {
-	FGLTFJsonNode* ChildNode = AddNode();
-
-	if (ParentNode != nullptr)
+	if (ChildNode == nullptr)
 	{
-		ParentNode->Children.Add(ChildNode);
+		ChildNode = AddNode();
 	}
 
-	return ChildNode;
-}
-
-FGLTFJsonNode* FGLTFJsonBuilder::AddChildComponentNode(FGLTFJsonNode* ParentNode)
-{
-	FGLTFJsonNode* ChildNode = AddChildNode(ParentNode);
-
 	if (ParentNode != nullptr)
 	{
-		ParentNode->ComponentNode = ChildNode;
+		ParentNode->Children.Add(ChildNode);
 	}
 
 	return ChildNode;
 }
 
-FGLTFJsonAccessor* FGLTFJsonBuilder::AddAccessor(const FGLTFJsonAccessor& JsonAccessor)
-{
-	return JsonRoot.Accessors.Add(JsonAccessor);
-}
-
-FGLTFJsonBuffer* FGLTFJsonBuilder::AddBuffer(const FGLTFJsonBuffer& JsonBuffer)
-{
-	return JsonRoot.Buffers.Add(JsonBuffer);
-}
-
-FGLTFJsonBufferView* FGLTFJsonBuilder::AddBufferView(const FGLTFJsonBufferView& JsonBufferView)
-{
-	return JsonRoot.BufferViews.Add(JsonBufferView);
-}
-
-FGLTFJsonCamera* FGLTFJsonBuilder::AddCamera(const FGLTFJsonCamera& JsonCamera)
+FGLTFJsonNode* FGLTFJsonBuilder::AddChildComponentNode(FGLTFJsonNode* ParentNode, FGLTFJsonNode* ChildNode)
 {
-	return JsonRoot.Cameras.Add(JsonCamera);
-}
-
-FGLTFJsonImage* FGLTFJsonBuilder::AddImage(const FGLTFJsonImage& JsonImage)
-{
-	return JsonRoot.Images.Add(JsonImage);
-}
-
-FGLTFJsonMaterial* FGLTFJsonBuilder::AddMaterial(const FGLTFJsonMaterial& JsonMaterial)
-{
-	return JsonRoot.Materials.Add(JsonMaterial);
-}
-
-FGLTFJsonMesh* FGLTFJsonBuilder::AddMesh(const FGLTFJsonMesh& JsonMesh)
-{
-	return JsonRoot.Meshes.Add(JsonMesh);
-}
-
-FGLTFJsonNode* FGLTFJsonBuilder::AddNode(const FGLTFJsonNode& JsonNode)
-{
-	return JsonRoot.Nodes.Add(JsonNode);
-}
-
-FGLTFJsonSampler* FGLTFJsonBuilder::AddSampler(const FGLTFJsonSampler& JsonSampler)
-{
-	return JsonRoot.Samplers.Add(JsonSampler);
-}
-
-FGLTFJsonScene* FGLTFJsonBuilder::AddScene(const FGLTFJsonScene& JsonScene)
-{
-	return JsonRoot.Scenes.Add(JsonScene);
-}
-
-FGLTFJsonSkin* FGLTFJsonBuilder::AddSkin(const FGLTFJsonSkin& JsonSkin)
-{
-	return JsonRoot.Skins.Add(JsonSkin);
-}
-
-FGLTFJsonTexture* FGLTFJsonBuilder::AddTexture(const FGLTFJsonTexture& JsonTexture)
-{
-	return JsonRoot.Textures.Add(JsonTexture);
-}
-
-FGLTFJsonBackdrop* FGLTFJsonBuilder::AddBackdrop(const FGLTFJsonBackdrop& JsonBackdrop)
-{
-	return JsonRoot.Backdrops.Add(JsonBackdrop);
-}
-
-FGLTFJsonHotspot* FGLTFJsonBuilder::AddHotspot(const FGLTFJsonHotspot& JsonHotspot)
-{
-	return JsonRoot.Hotspots.Add(JsonHotspot);
-}
-
-FGLTFJsonLight* FGLTFJsonBuilder::AddLight(const FGLTFJsonLight& JsonLight)
-{
-	return JsonRoot.Lights.Add(JsonLight);
-}
-
-FGLTFJsonLightMap* FGLTFJsonBuilder::AddLightMap(const FGLTFJsonLightMap& JsonLightMap)
-{
-	return JsonRoot.LightMaps.Add(JsonLightMap);
-}
-
-FGLTFJsonSkySphere* FGLTFJsonBuilder::AddSkySphere(const FGLTFJsonSkySphere& JsonSkySphere)
-{
-	return JsonRoot.SkySpheres.Add(JsonSkySphere);
-}
-
-FGLTFJsonEpicLevelVariantSets* FGLTFJsonBuilder::AddEpicLevelVariantSets(const FGLTFJsonEpicLevelVariantSets& JsonEpicLevelVariantSets)
-{
-	return JsonRoot.EpicLevelVariantSets.Add(JsonEpicLevelVariantSets);
-}
-
-FGLTFJsonKhrMaterialVariant* FGLTFJsonBuilder::AddKhrMaterialVariant(const FGLTFJsonKhrMaterialVariant& JsonKhrMaterialVariant)
-{
-	return JsonRoot.KhrMaterialVariants.Add(JsonKhrMaterialVariant);
-}
-
-FGLTFJsonNode* FGLTFJsonBuilder::AddChildNode(FGLTFJsonNode* Parent, const FGLTFJsonNode& JsonNode)
-{
-	FGLTFJsonNode* ChildNode = AddNode(JsonNode);
-
-	if (Parent != nullptr)
+	if (ChildNode == nullptr)
 	{
-		Parent->Children.Add(ChildNode);
+		ChildNode = AddChildNode(ParentNode, ChildNode);
 	}
 
-	return ChildNode;
-}
-
-FGLTFJsonNode* FGLTFJsonBuilder::AddChildComponentNode(FGLTFJsonNode* Parent, const FGLTFJsonNode& JsonNode)
-{
-	FGLTFJsonNode* ChildNode = AddChildNode(Parent, JsonNode);
-
-	if (Parent != nullptr)
+	if (ParentNode != nullptr)
 	{
-		Parent->ComponentNode = ChildNode;
+		ParentNode->ComponentNode = ChildNode;
 	}
 
 	return ChildNode;
diff --git a/Engine/Plugins/Enterprise/GLTFExporter/MigratedFromGit/Source/GLTFExporter/Private/Builders/GLTFJsonBuilder.h b/Engine/Plugins/Enterprise/GLTFExporter/MigratedFromGit/Source/GLTFExporter/Private/Builders/GLTFJsonBuilder.h
--- a/Engine/Plugins/Enterprise/GLTFExporter/MigratedFromGit/Source/GLTFExporter/Private/Builders/GLTFJsonBuilder.h
+++ b/Engine/Plugins/Enterprise/GLTFExporter/MigratedFromGit/Source/GLTFExporter/Private/Builders/GLTFJsonBuilder.h
@@ -40,31 +40,8 @@ public:
 	FGLTFJsonEpicLevelVariantSets* AddEpicLevelVariantSets();
 	FGLTFJsonKhrMaterialVariant* AddKhrMaterialVariant();
 
-	FGLTFJsonNode* AddChildNode(FGLTFJsonNode* ParentNode);
-	FGLTFJsonNode* AddChildComponentNode(FGLTFJsonNode* ParentNode);
-
-	FGLTFJsonAccessor* AddAccessor(const FGLTFJsonAccessor& JsonAccessor);
-	FGLTFJsonBuffer* AddBuffer(const FGLTFJsonBuffer& JsonBuffer);
-	FGLTFJsonBufferView* AddBufferView(const FGLTFJsonBufferView& JsonBufferView);
-	FGLTFJsonCamera* AddCamera(const FGLTFJsonCamera& JsonCamera);
-	FGLTFJsonImage* AddImage(const FGLTFJsonImage& JsonImage);
-	FGLTFJsonMaterial* AddMaterial(const FGLTFJsonMaterial& JsonMaterial);
-	FGLTFJsonMesh* AddMesh(const FGLTFJsonMesh& JsonMesh);
-	FGLTFJsonNode* AddNode(const FGLTFJsonNode& JsonNode);
-	FGLTFJsonSampler* AddSampler(const FGLTFJsonSampler& JsonSampler);
-	FGLTFJsonScene* AddScene(const FGLTFJsonScene& JsonScene);
-	FGLTFJsonSkin* AddSkin(const FGLTFJsonSkin& JsonSkin);
-	FGLTFJsonTexture* AddTexture(const FGLTFJsonTexture& JsonTexture);
-	FGLTFJsonBackdrop* AddBackdrop(const FGLTFJsonBackdrop& JsonBackdrop);
-	FGLTFJsonHotspot* AddHotspot(const FGLTFJsonHotspot& JsonHotspot);
-	FGLTFJsonLight* AddLight(const FGLTFJsonLight& JsonLight);
-	FGLTFJsonLightMap* AddLightMap(const FGLTFJsonLightMap& JsonLightMap);
-	FGLTFJsonSkySphere* AddSkySphere(const FGLTFJsonSkySphere& JsonSkySphere);
-	FGLTFJsonEpicLevelVariantSets* AddEpicLevelVariantSets(const FGLTFJsonEpicLevelVariantSets& JsonEpicLevelVariantSets);
-	FGLTFJsonKhrMaterialVariant* AddKhrMaterialVariant(const FGLTFJsonKhrMaterialVariant& JsonKhrMaterialVariant);
-
-	FGLTFJsonNode* AddChildNode(FGLTFJsonNode* ParentNode, const FGLTFJsonNode& JsonNode);
-	FGLTFJsonNode* AddChildComponentNode(FGLTFJsonNode* ParentNode, const FGLTFJsonNode& JsonNode);
+	FGLTFJsonNode* AddChildNode(FGLTFJsonNode* ParentNode, FGLTFJsonNode* ChildNode = nullptr);
+	FGLTFJsonNode* AddChildComponentNode(FGLTFJsonNode* ParentNodeParentNode, FGLTFJsonNode* ChildNode = nullptr);
 
 	const FGLTFJsonRoot& GetRoot() const;
 	FGLTFJsonAccessor& GetAccessor(FGLTFJsonAccessor* AccessorIndex);
diff --git a/Engine/Plugins/Enterprise/GLTFExporter/MigratedFromGit/Source/GLTFExporter/Private/Converters/GLTFAccessorConverters.cpp b/Engine/Plugins/Enterprise/GLTFExporter/MigratedFromGit/Source/GLTFExporter/Private/Converters/GLTFAccessorConverters.cpp
--- a/Engine/Plugins/Enterprise/GLTFExporter/MigratedFromGit/Source/GLTFExporter/Private/Converters/GLTFAccessorConverters.cpp
+++ b/Engine/Plugins/Enterprise/GLTFExporter/MigratedFromGit/Source/GLTFExporter/Private/Converters/GLTFAccessorConverters.cpp
@@ -50,35 +50,35 @@ FGLTFJsonAccessor* FGLTFPositionBufferConverter::Convert(const FGLTFMeshSection*
 		Positions[VertexIndex] = FGLTFConverterUtility::ConvertPosition(Position, Builder.ExportOptions->ExportUniformScale);
 	}
 
-	FGLTFJsonAccessor JsonAccessor;
-	JsonAccessor.BufferView = Builder.AddBufferView(Positions, EGLTFJsonBufferTarget::ArrayBuffer);
-	JsonAccessor.ComponentType = EGLTFJsonComponentType::Float;
-	JsonAccessor.Count = VertexCount;
-	JsonAccessor.Type = EGLTFJsonAccessorType::Vec3;
+	FGLTFJsonAccessor* JsonAccessor = Builder.AddAccessor();
+	JsonAccessor->BufferView = Builder.AddBufferView(Positions, EGLTFJsonBufferTarget::ArrayBuffer);
+	JsonAccessor->ComponentType = EGLTFJsonComponentType::Float;
+	JsonAccessor->Count = VertexCount;
+	JsonAccessor->Type = EGLTFJsonAccessorType::Vec3;
 
 	if (VertexCount > 0)
 	{
 		// Calculate accurate bounding box based on raw vertex values
-		JsonAccessor.MinMaxLength = 3;
+		JsonAccessor->MinMaxLength = 3;
 
-		for (int32 ComponentIndex = 0; ComponentIndex < JsonAccessor.MinMaxLength; ComponentIndex++)
+		for (int32 ComponentIndex = 0; ComponentIndex < JsonAccessor->MinMaxLength; ComponentIndex++)
 		{
-			JsonAccessor.Min[ComponentIndex] = Positions[0].Components[ComponentIndex];
-			JsonAccessor.Max[ComponentIndex] = Positions[0].Components[ComponentIndex];
+			JsonAccessor->Min[ComponentIndex] = Positions[0].Components[ComponentIndex];
+			JsonAccessor->Max[ComponentIndex] = Positions[0].Components[ComponentIndex];
 		}
 
 		for (uint32 VertexIndex = 1; VertexIndex < VertexCount; ++VertexIndex)
 		{
 			const FGLTFVector3& Position = Positions[VertexIndex];
-			for (int32 ComponentIndex = 0; ComponentIndex < JsonAccessor.MinMaxLength; ComponentIndex++)
+			for (int32 ComponentIndex = 0; ComponentIndex < JsonAccessor->MinMaxLength; ComponentIndex++)
 			{
-				JsonAccessor.Min[ComponentIndex] = FMath::Min(JsonAccessor.Min[ComponentIndex], Position.Components[ComponentIndex]);
-				JsonAccessor.Max[ComponentIndex] = FMath::Max(JsonAccessor.Max[ComponentIndex], Position.Components[ComponentIndex]);
+				JsonAccessor->Min[ComponentIndex] = FMath::Min(JsonAccessor->Min[ComponentIndex], Position.Components[ComponentIndex]);
+				JsonAccessor->Max[ComponentIndex] = FMath::Max(JsonAccessor->Max[ComponentIndex], Position.Components[ComponentIndex]);
 			}
 		}
 	}
 
-	return Builder.AddAccessor(JsonAccessor);
+	return JsonAccessor;
 }
 
 FGLTFJsonAccessor* FGLTFColorBufferConverter::Convert(const FGLTFMeshSection* MeshSection, const FColorVertexBuffer* VertexBuffer)
@@ -111,14 +111,14 @@ FGLTFJsonAccessor* FGLTFColorBufferConverter::Convert(const FGLTFMeshSection* Me
 		Colors[VertexIndex] = FGLTFConverterUtility::ConvertColor(Color);
 	}
 
-	FGLTFJsonAccessor JsonAccessor;
-	JsonAccessor.BufferView = Builder.AddBufferView(Colors, EGLTFJsonBufferTarget::ArrayBuffer);
-	JsonAccessor.ComponentType = EGLTFJsonComponentType::UInt8;
-	JsonAccessor.Count = VertexCount;
-	JsonAccessor.Type = EGLTFJsonAccessorType::Vec4;
-	JsonAccessor.bNormalized = true;
+	FGLTFJsonAccessor* JsonAccessor = Builder.AddAccessor();
+	JsonAccessor->BufferView = Builder.AddBufferView(Colors, EGLTFJsonBufferTarget::ArrayBuffer);
+	JsonAccessor->ComponentType = EGLTFJsonComponentType::UInt8;
+	JsonAccessor->Count = VertexCount;
+	JsonAccessor->Type = EGLTFJsonAccessorType::Vec4;
+	JsonAccessor->bNormalized = true;
 
-	return Builder.AddAccessor(JsonAccessor);
+	return JsonAccessor;
 }
 
 FGLTFJsonAccessor* FGLTFNormalBufferConverter::Convert(const FGLTFMeshSection* MeshSection, const FStaticMeshVertexBuffer* VertexBuffer)
@@ -168,14 +168,14 @@ FGLTFJsonAccessor* FGLTFNormalBufferConverter::Convert(const FGLTFMeshSection* M
 			: ConvertBufferView<FGLTFVector3, FPackedNormal>(MeshSection, SourceData);
 	}
 
-	FGLTFJsonAccessor JsonAccessor;
-	JsonAccessor.BufferView = BufferViewIndex;
-	JsonAccessor.ComponentType = ComponentType;
-	JsonAccessor.Count = MeshSection->IndexMap.Num();
-	JsonAccessor.Type = EGLTFJsonAccessorType::Vec3;
-	JsonAccessor.bNormalized = bMeshQuantization;
+	FGLTFJsonAccessor* JsonAccessor = Builder.AddAccessor();
+	JsonAccessor->BufferView = BufferViewIndex;
+	JsonAccessor->ComponentType = ComponentType;
+	JsonAccessor->Count = MeshSection->IndexMap.Num();
+	JsonAccessor->Type = EGLTFJsonAccessorType::Vec3;
+	JsonAccessor->bNormalized = bMeshQuantization;
 
-	return Builder.AddAccessor(JsonAccessor);
+	return JsonAccessor;
 }
 
 template <typename DestinationType, typename SourceType>
@@ -247,14 +247,14 @@ FGLTFJsonAccessor* FGLTFTangentBufferConverter::Convert(const FGLTFMeshSection*
 			: ConvertBufferView<FGLTFVector4, FPackedNormal>(MeshSection, SourceData);
 	}
 
-	FGLTFJsonAccessor JsonAccessor;
-	JsonAccessor.BufferView = BufferViewIndex;
-	JsonAccessor.ComponentType = ComponentType;
-	JsonAccessor.Count = MeshSection->IndexMap.Num();
-	JsonAccessor.Type = EGLTFJsonAccessorType::Vec4;
-	JsonAccessor.bNormalized = bMeshQuantization;
+	FGLTFJsonAccessor* JsonAccessor = Builder.AddAccessor();
+	JsonAccessor->BufferView = BufferViewIndex;
+	JsonAccessor->ComponentType = ComponentType;
+	JsonAccessor->Count = MeshSection->IndexMap.Num();
+	JsonAccessor->Type = EGLTFJsonAccessorType::Vec4;
+	JsonAccessor->bNormalized = bMeshQuantization;
 
-	return Builder.AddAccessor(JsonAccessor);
+	return JsonAccessor;
 }
 
 template <typename DestinationType, typename SourceType>
@@ -330,13 +330,13 @@ FGLTFJsonAccessor* FGLTFUVBufferConverter::Convert(const FGLTFMeshSection* MeshS
 		UVs[VertexIndex] = FGLTFConverterUtility::ConvertUV(UV);
 	}
 
-	FGLTFJsonAccessor JsonAccessor;
-	JsonAccessor.BufferView = Builder.AddBufferView(UVs, EGLTFJsonBufferTarget::ArrayBuffer);
-	JsonAccessor.ComponentType = EGLTFJsonComponentType::Float;
-	JsonAccessor.Count = VertexCount;
-	JsonAccessor.Type = EGLTFJsonAccessorType::Vec2;
+	FGLTFJsonAccessor* JsonAccessor = Builder.AddAccessor();
+	JsonAccessor->BufferView = Builder.AddBufferView(UVs, EGLTFJsonBufferTarget::ArrayBuffer);
+	JsonAccessor->ComponentType = EGLTFJsonComponentType::Float;
+	JsonAccessor->Count = VertexCount;
+	JsonAccessor->Type = EGLTFJsonAccessorType::Vec2;
 
-	return Builder.AddAccessor(JsonAccessor);
+	return JsonAccessor;
 }
 
 FGLTFJsonAccessor* FGLTFBoneIndexBufferConverter::Convert(const FGLTFMeshSection* MeshSection, const FSkinWeightVertexBuffer* VertexBuffer, uint32 InfluenceOffset)
@@ -437,13 +437,13 @@ FGLTFJsonAccessor* FGLTFBoneIndexBufferConverter::Convert(const FGLTFMeshSection
 		}
 	}
 
-	FGLTFJsonAccessor JsonAccessor;
-	JsonAccessor.BufferView = Builder.AddBufferView(BoneIndices, EGLTFJsonBufferTarget::ArrayBuffer);
-	JsonAccessor.ComponentType = FGLTFConverterUtility::GetComponentType<DestinationType>();
-	JsonAccessor.Count = VertexCount;
-	JsonAccessor.Type = EGLTFJsonAccessorType::Vec4;
+	FGLTFJsonAccessor* JsonAccessor = Builder.AddAccessor();
+	JsonAccessor->BufferView = Builder.AddBufferView(BoneIndices, EGLTFJsonBufferTarget::ArrayBuffer);
+	JsonAccessor->ComponentType = FGLTFConverterUtility::GetComponentType<DestinationType>();
+	JsonAccessor->Count = VertexCount;
+	JsonAccessor->Type = EGLTFJsonAccessorType::Vec4;
 
-	return Builder.AddAccessor(JsonAccessor);
+	return JsonAccessor;
 }
 
 FGLTFJsonAccessor* FGLTFBoneWeightBufferConverter::Convert(const FGLTFMeshSection* MeshSection, const FSkinWeightVertexBuffer* VertexBuffer, uint32 InfluenceOffset)
@@ -534,14 +534,14 @@ FGLTFJsonAccessor* FGLTFBoneWeightBufferConverter::Convert(const FGLTFMeshSectio
 		}
 	}
 
-	FGLTFJsonAccessor JsonAccessor;
-	JsonAccessor.BufferView = Builder.AddBufferView(BoneWeights, EGLTFJsonBufferTarget::ArrayBuffer);
-	JsonAccessor.ComponentType = EGLTFJsonComponentType::UInt8;
-	JsonAccessor.Count = VertexCount;
-	JsonAccessor.Type = EGLTFJsonAccessorType::Vec4;
-	JsonAccessor.bNormalized = true;
+	FGLTFJsonAccessor* JsonAccessor = Builder.AddAccessor();
+	JsonAccessor->BufferView = Builder.AddBufferView(BoneWeights, EGLTFJsonBufferTarget::ArrayBuffer);
+	JsonAccessor->ComponentType = EGLTFJsonComponentType::UInt8;
+	JsonAccessor->Count = VertexCount;
+	JsonAccessor->Type = EGLTFJsonAccessorType::Vec4;
+	JsonAccessor->bNormalized = true;
 
-	return Builder.AddAccessor(JsonAccessor);
+	return JsonAccessor;
 }
 
 FGLTFJsonAccessor* FGLTFIndexBufferConverter::Convert(const FGLTFMeshSection* MeshSection)
@@ -570,11 +570,11 @@ FGLTFJsonAccessor* FGLTFIndexBufferConverter::Convert(const FGLTFMeshSection* Me
 		Indices[Index] = static_cast<IndexType>(IndexBuffer[Index]);
 	}
 
-	FGLTFJsonAccessor JsonAccessor;
-	JsonAccessor.BufferView = Builder.AddBufferView(Indices, EGLTFJsonBufferTarget::ElementArrayBuffer, sizeof(IndexType));
-	JsonAccessor.ComponentType = FGLTFConverterUtility::GetComponentType<IndexType>();
-	JsonAccessor.Count = IndexCount;
-	JsonAccessor.Type = EGLTFJsonAccessorType::Scalar;
+	FGLTFJsonAccessor* JsonAccessor = Builder.AddAccessor();
+	JsonAccessor->BufferView = Builder.AddBufferView(Indices, EGLTFJsonBufferTarget::ElementArrayBuffer, sizeof(IndexType));
+	JsonAccessor->ComponentType = FGLTFConverterUtility::GetComponentType<IndexType>();
+	JsonAccessor->Count = IndexCount;
+	JsonAccessor->Type = EGLTFJsonAccessorType::Scalar;
 
-	return Builder.AddAccessor(JsonAccessor);
+	return JsonAccessor;
 }
diff --git a/Engine/Plugins/Enterprise/GLTFExporter/MigratedFromGit/Source/GLTFExporter/Private/Converters/GLTFBackdropConverters.cpp b/Engine/Plugins/Enterprise/GLTFExporter/MigratedFromGit/Source/GLTFExporter/Private/Converters/GLTFBackdropConverters.cpp
--- a/Engine/Plugins/Enterprise/GLTFExporter/MigratedFromGit/Source/GLTFExporter/Private/Converters/GLTFBackdropConverters.cpp
+++ b/Engine/Plugins/Enterprise/GLTFExporter/MigratedFromGit/Source/GLTFExporter/Private/Converters/GLTFBackdropConverters.cpp
@@ -13,8 +13,8 @@ FGLTFJsonBackdrop* FGLTFBackdropConverter::Convert(const AActor* BackdropActor)
 		return nullptr;
 	}
 
-	FGLTFJsonBackdrop JsonBackdrop;
-	BackdropActor->GetName(JsonBackdrop.Name);
+	FGLTFJsonBackdrop* JsonBackdrop = Builder.AddBackdrop();
+	BackdropActor->GetName(JsonBackdrop->Name);
 
 	const USceneComponent* SceneComponent = BackdropActor->GetRootComponent();
 	const FRotator Rotation = SceneComponent->GetComponentRotation();
@@ -24,12 +24,12 @@ FGLTFJsonBackdrop* FGLTFBackdropConverter::Convert(const AActor* BackdropActor)
 	// The end result is that the map rotation is always 2*PI more than the actual yaw of the scene component.
 
 	// TODO: remove the 2*PI multiplier if HDRI_Attributes is updated to match map rotation with yaw
-	JsonBackdrop.Angle = FRotator::ClampAxis(Rotation.Yaw * 2.0f * PI);
+	JsonBackdrop->Angle = FRotator::ClampAxis(Rotation.Yaw * 2.0f * PI);
 
 	const UStaticMesh* Mesh;
 	if (FGLTFActorUtility::TryGetPropertyValue(BackdropActor, TEXT("Mesh"), Mesh))
 	{
-		JsonBackdrop.Mesh = Builder.GetOrAddMesh(Mesh, { FGLTFMaterialUtility::GetDefaultMaterial() });
+		JsonBackdrop->Mesh = Builder.GetOrAddMesh(Mesh, { FGLTFMaterialUtility::GetDefaultMaterial() });
 	}
 	else
 	{
@@ -45,7 +45,7 @@ FGLTFJsonBackdrop* FGLTFBackdropConverter::Convert(const AActor* BackdropActor)
 		{
 			const ECubeFace CubeFace = static_cast<ECubeFace>(CubeFaceIndex);
 			const EGLTFJsonCubeFace JsonCubeFace = FGLTFConverterUtility::ConvertCubeFace(CubeFace);
-			JsonBackdrop.Cubemap[static_cast<int32>(JsonCubeFace)] = Builder.GetOrAddTexture(Cubemap, CubeFace);
+			JsonBackdrop->Cubemap[static_cast<int32>(JsonCubeFace)] = Builder.GetOrAddTexture(Cubemap, CubeFace);
 		}
 	}
 	else
@@ -56,7 +56,7 @@ FGLTFJsonBackdrop* FGLTFBackdropConverter::Convert(const AActor* BackdropActor)
 	float Intensity;
 	if (FGLTFActorUtility::TryGetPropertyValue(BackdropActor, TEXT("Intensity"), Intensity))
 	{
-		JsonBackdrop.Intensity = Intensity;
+		JsonBackdrop->Intensity = Intensity;
 	}
 	else
 	{
@@ -66,7 +66,7 @@ FGLTFJsonBackdrop* FGLTFBackdropConverter::Convert(const AActor* BackdropActor)
 	float Size;
 	if (FGLTFActorUtility::TryGetPropertyValue(BackdropActor, TEXT("Size"), Size))
 	{
-		JsonBackdrop.Size = Size;
+		JsonBackdrop->Size = Size;
 	}
 	else
 	{
@@ -76,7 +76,7 @@ FGLTFJsonBackdrop* FGLTFBackdropConverter::Convert(const AActor* BackdropActor)
 	FVector ProjectionCenter;
 	if (FGLTFActorUtility::TryGetPropertyValue(BackdropActor, TEXT("ProjectionCenter"), ProjectionCenter))
 	{
-		JsonBackdrop.ProjectionCenter = FGLTFConverterUtility::ConvertPosition(ProjectionCenter, Builder.ExportOptions->ExportUniformScale);
+		JsonBackdrop->ProjectionCenter = FGLTFConverterUtility::ConvertPosition(ProjectionCenter, Builder.ExportOptions->ExportUniformScale);
 	}
 	else
 	{
@@ -86,7 +86,7 @@ FGLTFJsonBackdrop* FGLTFBackdropConverter::Convert(const AActor* BackdropActor)
 	float LightingDistanceFactor;
 	if (FGLTFActorUtility::TryGetPropertyValue(BackdropActor, TEXT("LightingDistanceFactor"), LightingDistanceFactor))
 	{
-		JsonBackdrop.LightingDistanceFactor = LightingDistanceFactor;
+		JsonBackdrop->LightingDistanceFactor = LightingDistanceFactor;
 	}
 	else
 	{
@@ -96,12 +96,12 @@ FGLTFJsonBackdrop* FGLTFBackdropConverter::Convert(const AActor* BackdropActor)
 	bool UseCameraProjection;
 	if (FGLTFActorUtility::TryGetPropertyValue(BackdropActor, TEXT("UseCameraProjection"), UseCameraProjection))
 	{
-		JsonBackdrop.UseCameraProjection = UseCameraProjection;
+		JsonBackdrop->UseCameraProjection = UseCameraProjection;
 	}
 	else
 	{
 		Builder.LogWarning(FString::Printf(TEXT("Failed to export UseCameraProjection for HDRIBackdrop %s"), *BackdropActor->GetName()));
 	}
 
-	return Builder.AddBackdrop(JsonBackdrop);
+	return JsonBackdrop;
 }
diff --git a/Engine/Plugins/Enterprise/GLTFExporter/MigratedFromGit/Source/GLTFExporter/Private/Converters/GLTFCameraConverters.cpp b/Engine/Plugins/Enterprise/GLTFExporter/MigratedFromGit/Source/GLTFExporter/Private/Converters/GLTFCameraConverters.cpp
--- a/Engine/Plugins/Enterprise/GLTFExporter/MigratedFromGit/Source/GLTFExporter/Private/Converters/GLTFCameraConverters.cpp
+++ b/Engine/Plugins/Enterprise/GLTFExporter/MigratedFromGit/Source/GLTFExporter/Private/Converters/GLTFCameraConverters.cpp
@@ -9,36 +9,36 @@
 
 FGLTFJsonCamera* FGLTFCameraConverter::Convert(const UCameraComponent* CameraComponent)
 {
-	FGLTFJsonCamera Camera;
-	Camera.Name = FGLTFNameUtility::GetName(CameraComponent);
-	Camera.Type = FGLTFConverterUtility::ConvertCameraType(CameraComponent->ProjectionMode);
+	FGLTFJsonCamera* JsonCamera = Builder.AddCamera();
+	JsonCamera->Name = FGLTFNameUtility::GetName(CameraComponent);
+	JsonCamera->Type = FGLTFConverterUtility::ConvertCameraType(CameraComponent->ProjectionMode);
 
 	FMinimalViewInfo DesiredView;
 	const_cast<UCameraComponent*>(CameraComponent)->GetCameraView(0, DesiredView);
 	const float ExportScale = Builder.ExportOptions->ExportUniformScale;
 
-	switch (Camera.Type)
+	switch (JsonCamera->Type)
 	{
 		case EGLTFJsonCameraType::Orthographic:
 			if (!DesiredView.bConstrainAspectRatio)
 			{
 				Builder.LogWarning(FString::Printf(TEXT("Aspect ratio for orthographic camera component %s (in actor %s) will be constrainted in glTF"), *CameraComponent->GetName(), *CameraComponent->GetOwner()->GetName()));
 			}
-			Camera.Orthographic.XMag = FGLTFConverterUtility::ConvertLength(DesiredView.OrthoWidth, ExportScale);
-			Camera.Orthographic.YMag = FGLTFConverterUtility::ConvertLength(DesiredView.OrthoWidth / DesiredView.AspectRatio, ExportScale); // TODO: is this correct?
-			Camera.Orthographic.ZFar = FGLTFConverterUtility::ConvertLength(DesiredView.OrthoFarClipPlane, ExportScale);
-			Camera.Orthographic.ZNear = FGLTFConverterUtility::ConvertLength(DesiredView.OrthoNearClipPlane, ExportScale);
+			JsonCamera->Orthographic.XMag = FGLTFConverterUtility::ConvertLength(DesiredView.OrthoWidth, ExportScale);
+			JsonCamera->Orthographic.YMag = FGLTFConverterUtility::ConvertLength(DesiredView.OrthoWidth / DesiredView.AspectRatio, ExportScale); // TODO: is this correct?
+			JsonCamera->Orthographic.ZFar = FGLTFConverterUtility::ConvertLength(DesiredView.OrthoFarClipPlane, ExportScale);
+			JsonCamera->Orthographic.ZNear = FGLTFConverterUtility::ConvertLength(DesiredView.OrthoNearClipPlane, ExportScale);
 			break;
 
 		case EGLTFJsonCameraType::Perspective:
 			if (DesiredView.bConstrainAspectRatio)
 			{
-				Camera.Perspective.AspectRatio = DesiredView.AspectRatio;
+				JsonCamera->Perspective.AspectRatio = DesiredView.AspectRatio;
 			}
-			Camera.Perspective.YFov = FGLTFConverterUtility::ConvertFieldOfView(DesiredView.FOV, DesiredView.AspectRatio);
+			JsonCamera->Perspective.YFov = FGLTFConverterUtility::ConvertFieldOfView(DesiredView.FOV, DesiredView.AspectRatio);
 			// NOTE: even thought ZFar is optional, if we don't set it, then most gltf viewers won't handle it well.
-			Camera.Perspective.ZFar = FGLTFConverterUtility::ConvertLength(WORLD_MAX, ExportScale); // TODO: Unreal doesn't have max draw distance per view?
-			Camera.Perspective.ZNear = FGLTFConverterUtility::ConvertLength(GNearClippingPlane, ExportScale);
+			JsonCamera->Perspective.ZFar = FGLTFConverterUtility::ConvertLength(WORLD_MAX, ExportScale); // TODO: Unreal doesn't have max draw distance per view?
+			JsonCamera->Perspective.ZNear = FGLTFConverterUtility::ConvertLength(GNearClippingPlane, ExportScale);
 			break;
 
 		case EGLTFJsonCameraType::None:
@@ -85,9 +85,9 @@ FGLTFJsonCamera* FGLTFCameraConverter::Convert(const UCameraComponent* CameraCom
 			CameraControl.DollySensitivity = CameraActor->DollySensitivity;
 			CameraControl.DollyDuration = CameraActor->DollyDuration;
 
-			Camera.CameraControl = CameraControl;
+			JsonCamera->CameraControl = CameraControl;
 		}
 	}
 
-	return Builder.AddCamera(Camera);
+	return JsonCamera;
 }
diff --git a/Engine/Plugins/Enterprise/GLTFExporter/MigratedFromGit/Source/GLTFExporter/Private/Converters/GLTFEpicLevelVariantSetsConverters.cpp b/Engine/Plugins/Enterprise/GLTFExporter/MigratedFromGit/Source/GLTFExporter/Private/Converters/GLTFEpicLevelVariantSetsConverters.cpp
--- a/Engine/Plugins/Enterprise/GLTFExporter/MigratedFromGit/Source/GLTFExporter/Private/Converters/GLTFEpicLevelVariantSetsConverters.cpp
+++ b/Engine/Plugins/Enterprise/GLTFExporter/MigratedFromGit/Source/GLTFExporter/Private/Converters/GLTFEpicLevelVariantSetsConverters.cpp
@@ -13,8 +13,8 @@
 
 FGLTFJsonEpicLevelVariantSets* FGLTFEpicLevelVariantSetsConverter::Convert(const ULevelVariantSets* LevelVariantSets)
 {
-	FGLTFJsonEpicLevelVariantSets JsonLevelVariantSets;
-	LevelVariantSets->GetName(JsonLevelVariantSets.Name);
+	FGLTFJsonEpicLevelVariantSets* JsonLevelVariantSets = Builder.AddEpicLevelVariantSets();
+	LevelVariantSets->GetName(JsonLevelVariantSets->Name);
 
 	for (const UVariantSet* VariantSet: LevelVariantSets->GetVariantSets())
 	{
@@ -32,16 +32,16 @@ FGLTFJsonEpicLevelVariantSets* FGLTFEpicLevelVariantSetsConverter::Convert(const
 
 		if (JsonVariantSet.Variants.Num() > 0)
 		{
-			JsonLevelVariantSets.VariantSets.Add(JsonVariantSet);
+			JsonLevelVariantSets->VariantSets.Add(JsonVariantSet);
 		}
 	}
 
-	if (JsonLevelVariantSets.VariantSets.Num() == 0)
+	if (JsonLevelVariantSets->VariantSets.Num() == 0)
 	{
 		return nullptr;
 	}
 
-	return Builder.AddEpicLevelVariantSets(JsonLevelVariantSets);
+	return JsonLevelVariantSets;
 }
 
 bool FGLTFEpicLevelVariantSetsConverter::TryParseVariant(FGLTFJsonEpicVariant& OutVariant, const UVariant* Variant) const
diff --git a/Engine/Plugins/Enterprise/GLTFExporter/MigratedFromGit/Source/GLTFExporter/Private/Converters/GLTFHotspotConverters.cpp b/Engine/Plugins/Enterprise/GLTFExporter/MigratedFromGit/Source/GLTFExporter/Private/Converters/GLTFHotspotConverters.cpp
--- a/Engine/Plugins/Enterprise/GLTFExporter/MigratedFromGit/Source/GLTFExporter/Private/Converters/GLTFHotspotConverters.cpp
+++ b/Engine/Plugins/Enterprise/GLTFExporter/MigratedFromGit/Source/GLTFExporter/Private/Converters/GLTFHotspotConverters.cpp
@@ -6,8 +6,8 @@
 
 FGLTFJsonHotspot* FGLTFHotspotConverter::Convert(const AGLTFHotspotActor* HotspotActor)
 {
-	FGLTFJsonHotspot JsonHotspot;
-	HotspotActor->GetName(JsonHotspot.Name);
+	FGLTFJsonHotspot* JsonHotspot = Builder.AddHotspot();
+	HotspotActor->GetName(JsonHotspot->Name);
 
 	if (const ASkeletalMeshActor* SkeletalMeshActor = HotspotActor->SkeletalMeshActor)
 	{
@@ -15,13 +15,13 @@ FGLTFJsonHotspot* FGLTFHotspotConverter::Convert(const AGLTFHotspotActor* Hotspo
 		{
 			Builder.LogWarning(
 				FString::Printf(TEXT("Can't export animation in hotspot %s because vertex skin weights are disabled by export options"),
-				*JsonHotspot.Name));
+				*JsonHotspot->Name));
 		}
 		else if (!Builder.ExportOptions->bExportAnimationSequences)
 		{
 			Builder.LogWarning(
 				FString::Printf(TEXT("Can't export animation in hotspot %s because animation sequences are disabled by export options"),
-				*JsonHotspot.Name));
+				*JsonHotspot->Name));
 		}
 		else
 		{
@@ -32,7 +32,7 @@ FGLTFJsonHotspot* FGLTFHotspotConverter::Convert(const AGLTFHotspotActor* Hotspo
 			{
 				if (const UAnimSequence* AnimSequence = HotspotActor->AnimationSequence)
 				{
-					JsonHotspot.Animation = Builder.GetOrAddAnimation(RootNode, SkeletalMesh, AnimSequence);
+					JsonHotspot->Animation = Builder.GetOrAddAnimation(RootNode, SkeletalMesh, AnimSequence);
 				}
 				else
 				{
@@ -51,11 +51,11 @@ FGLTFJsonHotspot* FGLTFHotspotConverter::Convert(const AGLTFHotspotActor* Hotspo
 		{
 			Builder.LogWarning(
 				FString::Printf(TEXT("Can't export animation in hotspot %s because level sequences are disabled by export options"),
-				*JsonHotspot.Name));
+				*JsonHotspot->Name));
 		}
 		else
 		{
-			JsonHotspot.Animation = Builder.GetOrAddAnimation(HotspotActor->GetLevel(), LevelSequence);
+			JsonHotspot->Animation = Builder.GetOrAddAnimation(HotspotActor->GetLevel(), LevelSequence);
 		}
 	}
 	else
@@ -63,10 +63,10 @@ FGLTFJsonHotspot* FGLTFHotspotConverter::Convert(const AGLTFHotspotActor* Hotspo
 		// TODO: report warning
 	}
 
-	JsonHotspot.Image = Builder.GetOrAddTexture(HotspotActor->GetImageForState(EGLTFHotspotState::Default));
-	JsonHotspot.HoveredImage = Builder.GetOrAddTexture(HotspotActor->GetImageForState(EGLTFHotspotState::Hovered));
-	JsonHotspot.ToggledImage = Builder.GetOrAddTexture(HotspotActor->GetImageForState(EGLTFHotspotState::Toggled));
-	JsonHotspot.ToggledHoveredImage = Builder.GetOrAddTexture(HotspotActor->GetImageForState(EGLTFHotspotState::ToggledHovered));
+	JsonHotspot->Image = Builder.GetOrAddTexture(HotspotActor->GetImageForState(EGLTFHotspotState::Default));
+	JsonHotspot->HoveredImage = Builder.GetOrAddTexture(HotspotActor->GetImageForState(EGLTFHotspotState::Hovered));
+	JsonHotspot->ToggledImage = Builder.GetOrAddTexture(HotspotActor->GetImageForState(EGLTFHotspotState::Toggled));
+	JsonHotspot->ToggledHoveredImage = Builder.GetOrAddTexture(HotspotActor->GetImageForState(EGLTFHotspotState::ToggledHovered));
 
-	return Builder.AddHotspot(JsonHotspot);
+	return JsonHotspot;
 }
diff --git a/Engine/Plugins/Enterprise/GLTFExporter/MigratedFromGit/Source/GLTFExporter/Private/Converters/GLTFImageConverters.cpp b/Engine/Plugins/Enterprise/GLTFExporter/MigratedFromGit/Source/GLTFExporter/Private/Converters/GLTFImageConverters.cpp
--- a/Engine/Plugins/Enterprise/GLTFExporter/MigratedFromGit/Source/GLTFExporter/Private/Converters/GLTFImageConverters.cpp
+++ b/Engine/Plugins/Enterprise/GLTFExporter/MigratedFromGit/Source/GLTFExporter/Private/Converters/GLTFImageConverters.cpp
@@ -28,20 +28,20 @@ FGLTFJsonImage* FGLTFImageConverter::Convert(TGLTFSuperfluous<FString> Name, EGL
 			break;
 	}
 
-	FGLTFJsonImage JsonImage;
+	FGLTFJsonImage* JsonImage = Builder.AddImage();
 
 	if (Builder.bIsGlbFile)
 	{
-		JsonImage.Name = Name;
-		JsonImage.MimeType = MimeType;
-		JsonImage.BufferView = Builder.AddBufferView(CompressedData.GetData(), CompressedData.Num());
+		JsonImage->Name = Name;
+		JsonImage->MimeType = MimeType;
+		JsonImage->BufferView = Builder.AddBufferView(CompressedData.GetData(), CompressedData.Num());
 	}
 	else
 	{
-		JsonImage.Uri = SaveToFile(CompressedData.GetData(), CompressedData.Num(), MimeType, Name);
+		JsonImage->Uri = SaveToFile(CompressedData.GetData(), CompressedData.Num(), MimeType, Name);
 	}
 
-	return Builder.AddImage(JsonImage);
+	return JsonImage;
 }
 
 EGLTFJsonMimeType FGLTFImageConverter::GetMimeType(const FColor* Pixels, FIntPoint Size, bool bIgnoreAlpha, EGLTFTextureType Type) const
diff --git a/Engine/Plugins/Enterprise/GLTFExporter/MigratedFromGit/Source/GLTFExporter/Private/Converters/GLTFKhrVariantConverters.cpp b/Engine/Plugins/Enterprise/GLTFExporter/MigratedFromGit/Source/GLTFExporter/Private/Converters/GLTFKhrVariantConverters.cpp
--- a/Engine/Plugins/Enterprise/GLTFExporter/MigratedFromGit/Source/GLTFExporter/Private/Converters/GLTFKhrVariantConverters.cpp
+++ b/Engine/Plugins/Enterprise/GLTFExporter/MigratedFromGit/Source/GLTFExporter/Private/Converters/GLTFKhrVariantConverters.cpp
@@ -15,12 +15,13 @@ FGLTFJsonKhrMaterialVariant* FGLTFKhrMaterialVariantConverter::Convert(const UVa
 		return nullptr;
 	}
 
-	FGLTFJsonKhrMaterialVariant MaterialVariant;
+
+	FGLTFJsonKhrMaterialVariant* MaterialVariant = Builder.AddKhrMaterialVariant();
 
 	// TODO: add warning if the variant name is not unique, i.e it's already used?
 	// While material variants are technically allowed to use the same name, it may
 	// cause confusion when trying to select the correct variant in a viewer.
-	MaterialVariant.Name = Variant->GetDisplayText().ToString();
+	MaterialVariant->Name = Variant->GetDisplayText().ToString();
 
 	typedef TTuple<FGLTFJsonPrimitive*, FGLTFJsonMaterial*> TPrimitiveMaterial;
 	TArray<TPrimitiveMaterial> PrimitiveMaterials;
@@ -54,8 +55,6 @@ FGLTFJsonKhrMaterialVariant* FGLTFKhrMaterialVariantConverter::Convert(const UVa
 		return nullptr;
 	}
 
-	FGLTFJsonKhrMaterialVariant* MaterialVariantIndex = Builder.AddKhrMaterialVariant(MaterialVariant);
-
 	for (const TPrimitiveMaterial& PrimitiveMaterial: PrimitiveMaterials)
 	{
 		FGLTFJsonPrimitive* Primitive = PrimitiveMaterial.Key;
@@ -69,19 +68,19 @@ FGLTFJsonKhrMaterialVariant* FGLTFKhrMaterialVariantConverter::Convert(const UVa
 
 		if (ExistingMapping != nullptr)
 		{
-			ExistingMapping->Variants.AddUnique(MaterialVariantIndex);
+			ExistingMapping->Variants.AddUnique(MaterialVariant);
 		}
 		else
 		{
 			FGLTFJsonKhrMaterialVariantMapping Mapping;
 			Mapping.Material = MaterialIndex;
-			Mapping.Variants.Add(MaterialVariantIndex);
+			Mapping.Variants.Add(MaterialVariant);
 
 			Primitive->KhrMaterialVariantMappings.Add(Mapping);
 		}
 	}
 
-	return MaterialVariantIndex;
+	return MaterialVariant;
 }
 
 bool FGLTFKhrMaterialVariantConverter::TryParseMaterialProperty(FGLTFJsonPrimitive*& OutPrimitive, FGLTFJsonMaterial*& OutMaterialIndex, const UPropertyValueMaterial* Property) const
diff --git a/Engine/Plugins/Enterprise/GLTFExporter/MigratedFromGit/Source/GLTFExporter/Private/Converters/GLTFLightConverters.cpp b/Engine/Plugins/Enterprise/GLTFExporter/MigratedFromGit/Source/GLTFExporter/Private/Converters/GLTFLightConverters.cpp
--- a/Engine/Plugins/Enterprise/GLTFExporter/MigratedFromGit/Source/GLTFExporter/Private/Converters/GLTFLightConverters.cpp
+++ b/Engine/Plugins/Enterprise/GLTFExporter/MigratedFromGit/Source/GLTFExporter/Private/Converters/GLTFLightConverters.cpp
@@ -10,38 +10,38 @@
 
 FGLTFJsonLight* FGLTFLightConverter::Convert(const ULightComponent* LightComponent)
 {
-	FGLTFJsonLight Light;
+	FGLTFJsonLight* Light = Builder.AddLight();
 
-	Light.Name = FGLTFNameUtility::GetName(LightComponent);
-	Light.Type = FGLTFConverterUtility::ConvertLightType(LightComponent->GetLightType());
+	Light->Name = FGLTFNameUtility::GetName(LightComponent);
+	Light->Type = FGLTFConverterUtility::ConvertLightType(LightComponent->GetLightType());
 
-	if (Light.Type == EGLTFJsonLightType::None)
+	if (Light->Type == EGLTFJsonLightType::None)
 	{
 		// TODO: report error (unsupported light component type)
 		return nullptr;
 	}
 
-	Light.Intensity = LightComponent->Intensity;
+	Light->Intensity = LightComponent->Intensity;
 
 	const FLinearColor TemperatureColor = LightComponent->bUseTemperature ? FLinearColor::MakeFromColorTemperature(LightComponent->Temperature) : FLinearColor::White;
-	Light.Color = FGLTFConverterUtility::ConvertColor3(TemperatureColor * LightComponent->GetLightColor(), Builder.ExportOptions->bStrictCompliance);
+	Light->Color = FGLTFConverterUtility::ConvertColor3(TemperatureColor * LightComponent->GetLightColor(), Builder.ExportOptions->bStrictCompliance);
 
 	if (const UPointLightComponent* PointLightComponent = Cast<UPointLightComponent>(LightComponent))
 	{
-		Light.Range = FGLTFConverterUtility::ConvertLength(PointLightComponent->AttenuationRadius, Builder.ExportOptions->ExportUniformScale);
+		Light->Range = FGLTFConverterUtility::ConvertLength(PointLightComponent->AttenuationRadius, Builder.ExportOptions->ExportUniformScale);
 	}
 
 	if (const USpotLightComponent* SpotLightComponent = Cast<USpotLightComponent>(LightComponent))
 	{
-		Light.Spot.InnerConeAngle = FGLTFConverterUtility::ConvertLightAngle(SpotLightComponent->InnerConeAngle);
-		Light.Spot.OuterConeAngle = FGLTFConverterUtility::ConvertLightAngle(SpotLightComponent->OuterConeAngle);
+		Light->Spot.InnerConeAngle = FGLTFConverterUtility::ConvertLightAngle(SpotLightComponent->InnerConeAngle);
+		Light->Spot.OuterConeAngle = FGLTFConverterUtility::ConvertLightAngle(SpotLightComponent->OuterConeAngle);
 
 		if (Builder.ExportOptions->bStrictCompliance)
 		{
-			Light.Spot.InnerConeAngle = FMath::Clamp(Light.Spot.InnerConeAngle, 0.0f, nextafterf(Light.Spot.OuterConeAngle, 0.0f));
-			Light.Spot.OuterConeAngle = FMath::Clamp(Light.Spot.OuterConeAngle, nextafterf(Light.Spot.InnerConeAngle, HALF_PI), HALF_PI);
+			Light->Spot.InnerConeAngle = FMath::Clamp(Light->Spot.InnerConeAngle, 0.0f, nextafterf(Light->Spot.OuterConeAngle, 0.0f));
+			Light->Spot.OuterConeAngle = FMath::Clamp(Light->Spot.OuterConeAngle, nextafterf(Light->Spot.InnerConeAngle, HALF_PI), HALF_PI);
 		}
 	}
 
-	return Builder.AddLight(Light);
+	return Light;
 }
diff --git a/Engine/Plugins/Enterprise/GLTFExporter/MigratedFromGit/Source/GLTFExporter/Private/Converters/GLTFLightMapConverters.cpp b/Engine/Plugins/Enterprise/GLTFExporter/MigratedFromGit/Source/GLTFExporter/Private/Converters/GLTFLightMapConverters.cpp
--- a/Engine/Plugins/Enterprise/GLTFExporter/MigratedFromGit/Source/GLTFExporter/Private/Converters/GLTFLightMapConverters.cpp
+++ b/Engine/Plugins/Enterprise/GLTFExporter/MigratedFromGit/Source/GLTFExporter/Private/Converters/GLTFLightMapConverters.cpp
@@ -71,13 +71,13 @@ FGLTFJsonLightMap* FGLTFLightMapConverter::Convert(const UStaticMeshComponent* S
 	const FVector4& LightMapAdd = LightMapInteraction.GetAddArray()[0];
 	const FVector4& LightMapScale = LightMapInteraction.GetScaleArray()[0];
 
-	FGLTFJsonLightMap JsonLightMap;
-	StaticMeshComponent->GetName(JsonLightMap.Name); // TODO: use better name (similar to light and camera)
-	JsonLightMap.Texture.Index = TextureIndex;
-	JsonLightMap.Texture.TexCoord = CoordinateIndex;
-	JsonLightMap.LightMapScale = { LightMapScale.X, LightMapScale.Y, LightMapScale.Z, LightMapScale.W };
-	JsonLightMap.LightMapAdd = { LightMapAdd.X, LightMapAdd.Y, LightMapAdd.Z, LightMapAdd.W };
-	JsonLightMap.CoordinateScaleBias = { CoordinateScale.X, CoordinateScale.Y, CoordinateBias.X, CoordinateBias.Y };
-
-	return Builder.AddLightMap(JsonLightMap);
+	FGLTFJsonLightMap* JsonLightMap = Builder.AddLightMap();
+	StaticMeshComponent->GetName(JsonLightMap->Name); // TODO: use better name (similar to light and camera)
+	JsonLightMap->Texture.Index = TextureIndex;
+	JsonLightMap->Texture.TexCoord = CoordinateIndex;
+	JsonLightMap->LightMapScale = { LightMapScale.X, LightMapScale.Y, LightMapScale.Z, LightMapScale.W };
+	JsonLightMap->LightMapAdd = { LightMapAdd.X, LightMapAdd.Y, LightMapAdd.Z, LightMapAdd.W };
+	JsonLightMap->CoordinateScaleBias = { CoordinateScale.X, CoordinateScale.Y, CoordinateBias.X, CoordinateBias.Y };
+
+	return JsonLightMap;
 }
diff --git a/Engine/Plugins/Enterprise/GLTFExporter/MigratedFromGit/Source/GLTFExporter/Private/Converters/GLTFMaterialUtility.cpp b/Engine/Plugins/Enterprise/GLTFExporter/MigratedFromGit/Source/GLTFExporter/Private/Converters/GLTFMaterialUtility.cpp
--- a/Engine/Plugins/Enterprise/GLTFExporter/MigratedFromGit/Source/GLTFExporter/Private/Converters/GLTFMaterialUtility.cpp
+++ b/Engine/Plugins/Enterprise/GLTFExporter/MigratedFromGit/Source/GLTFExporter/Private/Converters/GLTFMaterialUtility.cpp
@@ -265,20 +265,20 @@ FGLTFPropertyBakeOutput FGLTFMaterialUtility::BakeMaterialProperty(const FIntPoi
 FGLTFJsonTexture* FGLTFMaterialUtility::AddTexture(FGLTFConvertBuilder& Builder, TGLTFSharedArray<FColor>& Pixels, const FIntPoint& TextureSize, bool bIgnoreAlpha, bool bIsNormalMap, const FString& TextureName, EGLTFJsonTextureFilter MinFilter, EGLTFJsonTextureFilter MagFilter, EGLTFJsonTextureWrap WrapS, EGLTFJsonTextureWrap WrapT)
 {
 	// TODO: maybe we should reuse existing samplers?
-	FGLTFJsonSampler JsonSampler;
-	JsonSampler.Name = TextureName;
-	JsonSampler.MinFilter = MinFilter;
-	JsonSampler.MagFilter = MagFilter;
-	JsonSampler.WrapS = WrapS;
-	JsonSampler.WrapT = WrapT;
+	FGLTFJsonSampler* JsonSampler = Builder.AddSampler();
+	JsonSampler->Name = TextureName;
+	JsonSampler->MinFilter = MinFilter;
+	JsonSampler->MagFilter = MagFilter;
+	JsonSampler->WrapS = WrapS;
+	JsonSampler->WrapT = WrapT;
 
 	// TODO: reuse same texture index when image is the same
-	FGLTFJsonTexture JsonTexture;
-	JsonTexture.Name = TextureName;
-	JsonTexture.Sampler = Builder.AddSampler(JsonSampler);
-	JsonTexture.Source = Builder.GetOrAddImage(Pixels, TextureSize, bIgnoreAlpha, bIsNormalMap ? EGLTFTextureType::Normalmaps : EGLTFTextureType::None, TextureName);
+	FGLTFJsonTexture* JsonTexture = Builder.AddTexture();
+	JsonTexture->Name = TextureName;
+	JsonTexture->Sampler = JsonSampler;
+	JsonTexture->Source = Builder.GetOrAddImage(Pixels, TextureSize, bIgnoreAlpha, bIsNormalMap ? EGLTFTextureType::Normalmaps : EGLTFTextureType::None, TextureName);
 
-	return Builder.AddTexture(JsonTexture);
+	return JsonTexture;
 }
 
 FLinearColor FGLTFMaterialUtility::GetMask(const FExpressionInput& ExpressionInput)
diff --git a/Engine/Plugins/Enterprise/GLTFExporter/MigratedFromGit/Source/GLTFExporter/Private/Converters/GLTFNodeConverters.cpp b/Engine/Plugins/Enterprise/GLTFExporter/MigratedFromGit/Source/GLTFExporter/Private/Converters/GLTFNodeConverters.cpp
--- a/Engine/Plugins/Enterprise/GLTFExporter/MigratedFromGit/Source/GLTFExporter/Private/Converters/GLTFNodeConverters.cpp
+++ b/Engine/Plugins/Enterprise/GLTFExporter/MigratedFromGit/Source/GLTFExporter/Private/Converters/GLTFNodeConverters.cpp
@@ -25,7 +25,7 @@ FGLTFJsonNode* FGLTFActorConverter::Convert(const AActor* Actor)
 	}
 
 	const USceneComponent* RootComponent = Actor->GetRootComponent();
-	FGLTFJsonNode* RootNodeIndex = Builder.GetOrAddNode(RootComponent);
+	FGLTFJsonNode* RootNode = Builder.GetOrAddNode(RootComponent);
 
 	// TODO: process all components since any component can be attached to any other component in runtime
 
@@ -34,16 +34,14 @@ FGLTFJsonNode* FGLTFActorConverter::Convert(const AActor* Actor)
 	{
 		if (Builder.ExportOptions->bExportSkySpheres)
 		{
-			FGLTFJsonNode& RootNode = Builder.GetNode(RootNodeIndex);
-			RootNode.SkySphere = Builder.GetOrAddSkySphere(Actor);
+			RootNode->SkySphere = Builder.GetOrAddSkySphere(Actor);
 		}
 	}
 	else if (FGLTFActorUtility::IsHDRIBackdropBlueprint(BlueprintPath))
 	{
 		if (Builder.ExportOptions->bExportHDRIBackdrops)
 		{
-			FGLTFJsonNode& RootNode = Builder.GetNode(RootNodeIndex);
-			RootNode.Backdrop = Builder.GetOrAddBackdrop(Actor);
+			RootNode->Backdrop = Builder.GetOrAddBackdrop(Actor);
 		}
 	}
 	else if (const ALevelSequenceActor* LevelSequenceActor = Cast<ALevelSequenceActor>(Actor))
@@ -57,8 +55,7 @@ FGLTFJsonNode* FGLTFActorConverter::Convert(const AActor* Actor)
 	{
 		if (Builder.ExportOptions->bExportAnimationHotspots)
 		{
-			FGLTFJsonNode& RootNode = Builder.GetNode(RootNodeIndex);
-			RootNode.Hotspot = Builder.GetOrAddHotspot(HotspotActor);
+			RootNode->Hotspot = Builder.GetOrAddHotspot(HotspotActor);
 		}
 	}
 	else
@@ -77,7 +74,7 @@ FGLTFJsonNode* FGLTFActorConverter::Convert(const AActor* Actor)
 		}
 	}
 
-	return RootNodeIndex;
+	return RootNode;
 }
 
 FGLTFJsonNode* FGLTFComponentConverter::Convert(const USceneComponent* SceneComponent)
@@ -226,13 +223,13 @@ FGLTFJsonNode* FGLTFStaticSocketConverter::Convert(FGLTFJsonNode* RootNode, cons
 		return nullptr;
 	}
 
-	FGLTFJsonNode Node;
-	Node.Name = SocketName.ToString();
+	FGLTFJsonNode* Node = Builder.AddNode();
+	Node->Name = SocketName.ToString();
 
 	// TODO: add warning check for non-uniform scaling
-	Node.Translation = FGLTFConverterUtility::ConvertPosition(Socket->RelativeLocation, Builder.ExportOptions->ExportUniformScale);
-	Node.Rotation = FGLTFConverterUtility::ConvertRotation(Socket->RelativeRotation.Quaternion());
-	Node.Scale = FGLTFConverterUtility::ConvertScale(Socket->RelativeScale);
+	Node->Translation = FGLTFConverterUtility::ConvertPosition(Socket->RelativeLocation, Builder.ExportOptions->ExportUniformScale);
+	Node->Rotation = FGLTFConverterUtility::ConvertRotation(Socket->RelativeRotation.Quaternion());
+	Node->Scale = FGLTFConverterUtility::ConvertScale(Socket->RelativeScale);
 
 	return Builder.AddChildNode(RootNode, Node);
 }
@@ -248,13 +245,13 @@ FGLTFJsonNode* FGLTFSkeletalSocketConverter::Convert(FGLTFJsonNode* RootNode, co
 	const USkeletalMeshSocket* Socket = SkeletalMesh->FindSocket(SocketName);
 	if (Socket != nullptr)
 	{
-		FGLTFJsonNode Node;
-		Node.Name = SocketName.ToString();
+		FGLTFJsonNode* Node = Builder.AddNode();
+		Node->Name = SocketName.ToString();
 
 		// TODO: add warning check for non-uniform scaling
-		Node.Translation = FGLTFConverterUtility::ConvertPosition(Socket->RelativeLocation, Builder.ExportOptions->ExportUniformScale);
-		Node.Rotation = FGLTFConverterUtility::ConvertRotation(Socket->RelativeRotation.Quaternion());
-		Node.Scale = FGLTFConverterUtility::ConvertScale(Socket->RelativeScale);
+		Node->Translation = FGLTFConverterUtility::ConvertPosition(Socket->RelativeLocation, Builder.ExportOptions->ExportUniformScale);
+		Node->Rotation = FGLTFConverterUtility::ConvertRotation(Socket->RelativeRotation.Quaternion());
+		Node->Scale = FGLTFConverterUtility::ConvertScale(Socket->RelativeScale);
 
 		const int32 ParentBone = RefSkeleton.FindBoneIndex(Socket->BoneName);
 		FGLTFJsonNode* ParentNode = ParentBone != INDEX_NONE ? Builder.GetOrAddNode(RootNode, SkeletalMesh, ParentBone) : RootNode;
@@ -290,17 +287,17 @@ FGLTFJsonNode* FGLTFSkeletalBoneConverter::Convert(FGLTFJsonNode* RootNode, cons
 
 	const FMeshBoneInfo& BoneInfo = BoneInfos[BoneIndex];
 
-	FGLTFJsonNode Node;
-	Node.Name = BoneInfo.Name.ToString();
+	FGLTFJsonNode* Node = Builder.AddNode();
+	Node->Name = BoneInfo.Name.ToString();
 
 	const TArray<FTransform>& BonePoses = RefSkeleton.GetRefBonePose();
 	if (BonePoses.IsValidIndex(BoneIndex))
 	{
 		// TODO: add warning check for non-uniform scaling
 		const FTransform& BonePose = BonePoses[BoneIndex];
-		Node.Translation = FGLTFConverterUtility::ConvertPosition(BonePose.GetTranslation(), Builder.ExportOptions->ExportUniformScale);
-		Node.Rotation = FGLTFConverterUtility::ConvertRotation(BonePose.GetRotation());
-		Node.Scale = FGLTFConverterUtility::ConvertScale(BonePose.GetScale3D());
+		Node->Translation = FGLTFConverterUtility::ConvertPosition(BonePose.GetTranslation(), Builder.ExportOptions->ExportUniformScale);
+		Node->Rotation = FGLTFConverterUtility::ConvertRotation(BonePose.GetRotation());
+		Node->Scale = FGLTFConverterUtility::ConvertScale(BonePose.GetScale3D());
 	}
 	else
 	{
diff --git a/Engine/Plugins/Enterprise/GLTFExporter/MigratedFromGit/Source/GLTFExporter/Private/Converters/GLTFSamplerConverters.cpp b/Engine/Plugins/Enterprise/GLTFExporter/MigratedFromGit/Source/GLTFExporter/Private/Converters/GLTFSamplerConverters.cpp
--- a/Engine/Plugins/Enterprise/GLTFExporter/MigratedFromGit/Source/GLTFExporter/Private/Converters/GLTFSamplerConverters.cpp
+++ b/Engine/Plugins/Enterprise/GLTFExporter/MigratedFromGit/Source/GLTFExporter/Private/Converters/GLTFSamplerConverters.cpp
@@ -9,34 +9,34 @@ FGLTFJsonSampler* FGLTFSamplerConverter::Convert(const UTexture* Texture)
 {
 	// TODO: reuse existing samplers, don't create a unique sampler per texture
 
-	FGLTFJsonSampler JsonSampler;
-	Texture->GetName(JsonSampler.Name);
+	FGLTFJsonSampler* JsonSampler = Builder.AddSampler();
+	Texture->GetName(JsonSampler->Name);
 
 	if (Texture->IsA<ULightMapTexture2D>())
 	{
 		// Override default filter settings for LightMap textures (which otherwise is "nearest")
-		JsonSampler.MinFilter = EGLTFJsonTextureFilter::LinearMipmapLinear;
-		JsonSampler.MagFilter = EGLTFJsonTextureFilter::Linear;
+		JsonSampler->MinFilter = EGLTFJsonTextureFilter::LinearMipmapLinear;
+		JsonSampler->MagFilter = EGLTFJsonTextureFilter::Linear;
 	}
 	else
 	{
-		JsonSampler.MinFilter = FGLTFConverterUtility::ConvertMinFilter(Texture->Filter, Texture->LODGroup);
-		JsonSampler.MagFilter = FGLTFConverterUtility::ConvertMagFilter(Texture->Filter, Texture->LODGroup);
+		JsonSampler->MinFilter = FGLTFConverterUtility::ConvertMinFilter(Texture->Filter, Texture->LODGroup);
+		JsonSampler->MagFilter = FGLTFConverterUtility::ConvertMagFilter(Texture->Filter, Texture->LODGroup);
 	}
 
 	if (FGLTFTextureUtility::IsCubemap(Texture))
 	{
-		JsonSampler.WrapS = EGLTFJsonTextureWrap::ClampToEdge;
-		JsonSampler.WrapT = EGLTFJsonTextureWrap::ClampToEdge;
+		JsonSampler->WrapS = EGLTFJsonTextureWrap::ClampToEdge;
+		JsonSampler->WrapT = EGLTFJsonTextureWrap::ClampToEdge;
 	}
 	else
 	{
 		const TTuple<TextureAddress, TextureAddress> AddressXY = FGLTFTextureUtility::GetAddressXY(Texture);
 		// TODO: report error if AddressX == TA_MAX or AddressY == TA_MAX
 
-		JsonSampler.WrapS = FGLTFConverterUtility::ConvertWrap(AddressXY.Get<0>());
-		JsonSampler.WrapT = FGLTFConverterUtility::ConvertWrap(AddressXY.Get<1>());
+		JsonSampler->WrapS = FGLTFConverterUtility::ConvertWrap(AddressXY.Get<0>());
+		JsonSampler->WrapT = FGLTFConverterUtility::ConvertWrap(AddressXY.Get<1>());
 	}
 
-	return Builder.AddSampler(JsonSampler);
+	return JsonSampler;
 }
diff --git a/Engine/Plugins/Enterprise/GLTFExporter/MigratedFromGit/Source/GLTFExporter/Private/Converters/GLTFSceneConverters.cpp b/Engine/Plugins/Enterprise/GLTFExporter/MigratedFromGit/Source/GLTFExporter/Private/Converters/GLTFSceneConverters.cpp
--- a/Engine/Plugins/Enterprise/GLTFExporter/MigratedFromGit/Source/GLTFExporter/Private/Converters/GLTFSceneConverters.cpp
+++ b/Engine/Plugins/Enterprise/GLTFExporter/MigratedFromGit/Source/GLTFExporter/Private/Converters/GLTFSceneConverters.cpp
@@ -8,8 +8,8 @@
 
 FGLTFJsonScene* FGLTFSceneConverter::Convert(const UWorld* World)
 {
-	FGLTFJsonScene Scene;
-	World->GetName(Scene.Name);
+	FGLTFJsonScene* Scene = Builder.AddScene();
+	World->GetName(Scene->Name);
 
 	const TArray<ULevel*>& Levels = World->GetLevels();
 	if (Levels.Num() > 0)
@@ -39,7 +39,7 @@ FGLTFJsonScene* FGLTFSceneConverter::Convert(const UWorld* World)
 								FGLTFJsonEpicLevelVariantSets* JsonEpicLevelVariantSets = Builder.GetOrAddEpicLevelVariantSets(LevelVariantSets);
 								if (JsonEpicLevelVariantSets != nullptr)
 								{
-									Scene.EpicLevelVariantSets.Add(JsonEpicLevelVariantSets);
+									Scene->EpicLevelVariantSets.Add(JsonEpicLevelVariantSets);
 								}
 							}
 						}
@@ -65,9 +65,9 @@ FGLTFJsonScene* FGLTFSceneConverter::Convert(const UWorld* World)
 				FGLTFJsonNode* JsonNode = Builder.GetOrAddNode(Actor);
 				if (JsonNode != nullptr && Builder.IsRootActor(Actor))
 				{
-					// TODO: to avoid having to add irrelevant actors/components let GLTFComponentConverter decide and add root nodes to scene.
+					// TODO: to avoid having to add irrelevant actors/components let GLTFComponentConverter decide and add root nodes to Scene->
 					// This change may require node converters to support cyclic calls.
-					Scene.Nodes.Add(JsonNode);
+					Scene->Nodes.Add(JsonNode);
 				}
 			}
 		}
@@ -79,5 +79,5 @@ FGLTFJsonScene* FGLTFSceneConverter::Convert(const UWorld* World)
 			*World->GetName()));
 	}
 
-	return Builder.AddScene(Scene);
+	return Scene;
 }
diff --git a/Engine/Plugins/Enterprise/GLTFExporter/MigratedFromGit/Source/GLTFExporter/Private/Converters/GLTFSkinConverters.cpp b/Engine/Plugins/Enterprise/GLTFExporter/MigratedFromGit/Source/GLTFExporter/Private/Converters/GLTFSkinConverters.cpp
--- a/Engine/Plugins/Enterprise/GLTFExporter/MigratedFromGit/Source/GLTFExporter/Private/Converters/GLTFSkinConverters.cpp
+++ b/Engine/Plugins/Enterprise/GLTFExporter/MigratedFromGit/Source/GLTFExporter/Private/Converters/GLTFSkinConverters.cpp
@@ -15,9 +15,9 @@ FGLTFJsonSkin* FGLTFSkinConverter::Convert(FGLTFJsonNode* RootNode, const USkele
 	const FReferenceSkeleton& RefSkeleton = SkeletalMesh->RefSkeleton;
 #endif
 
-	FGLTFJsonSkin Skin;
-	Skin.Name = Skeleton != nullptr ? Skeleton->GetName() : SkeletalMesh->GetName();
-	Skin.Skeleton = RootNode;
+	FGLTFJsonSkin* JsonSkin = Builder.AddSkin();
+	JsonSkin->Name = Skeleton != nullptr ? Skeleton->GetName() : SkeletalMesh->GetName();
+	JsonSkin->Skeleton = RootNode;
 
 	const int32 BoneCount = RefSkeleton.GetNum();
 	if (BoneCount == 0)
@@ -26,11 +26,11 @@ FGLTFJsonSkin* FGLTFSkinConverter::Convert(FGLTFJsonNode* RootNode, const USkele
 		return nullptr;
 	}
 
-	Skin.Joints.AddUninitialized(BoneCount);
+	JsonSkin->Joints.AddUninitialized(BoneCount);
 
 	for (int32 BoneIndex = 0; BoneIndex < BoneCount; ++BoneIndex)
 	{
-		Skin.Joints[BoneIndex] = Builder.GetOrAddNode(RootNode, SkeletalMesh, BoneIndex);
+		JsonSkin->Joints[BoneIndex] = Builder.GetOrAddNode(RootNode, SkeletalMesh, BoneIndex);
 	}
 
 	TArray<FGLTFMatrix4> InverseBindMatrices;
@@ -42,13 +42,12 @@ FGLTFJsonSkin* FGLTFSkinConverter::Convert(FGLTFJsonNode* RootNode, const USkele
 		InverseBindMatrices[BoneIndex] = FGLTFConverterUtility::ConvertTransform(InverseBindTransform, Builder.ExportOptions->ExportUniformScale);
 	}
 
-	FGLTFJsonAccessor JsonAccessor;
-	JsonAccessor.BufferView = Builder.AddBufferView(InverseBindMatrices);
-	JsonAccessor.ComponentType = EGLTFJsonComponentType::Float;
-	JsonAccessor.Count = BoneCount;
-	JsonAccessor.Type = EGLTFJsonAccessorType::Mat4;
+	FGLTFJsonAccessor* JsonBindMatricesAccessor = Builder.AddAccessor();
+	JsonBindMatricesAccessor->BufferView = Builder.AddBufferView(InverseBindMatrices);
+	JsonBindMatricesAccessor->ComponentType = EGLTFJsonComponentType::Float;
+	JsonBindMatricesAccessor->Count = BoneCount;
+	JsonBindMatricesAccessor->Type = EGLTFJsonAccessorType::Mat4;
 
-	Skin.InverseBindMatrices = Builder.AddAccessor(JsonAccessor);
-
-	return Builder.AddSkin(Skin);
+	JsonSkin->InverseBindMatrices = JsonBindMatricesAccessor;
+	return JsonSkin;
 }
diff --git a/Engine/Plugins/Enterprise/GLTFExporter/MigratedFromGit/Source/GLTFExporter/Private/Converters/GLTFSkySphereConverters.cpp b/Engine/Plugins/Enterprise/GLTFExporter/MigratedFromGit/Source/GLTFExporter/Private/Converters/GLTFSkySphereConverters.cpp
--- a/Engine/Plugins/Enterprise/GLTFExporter/MigratedFromGit/Source/GLTFExporter/Private/Converters/GLTFSkySphereConverters.cpp
+++ b/Engine/Plugins/Enterprise/GLTFExporter/MigratedFromGit/Source/GLTFExporter/Private/Converters/GLTFSkySphereConverters.cpp
@@ -17,8 +17,8 @@ FGLTFJsonSkySphere* FGLTFSkySphereConverter::Convert(const AActor* SkySphereActo
 		return nullptr;
 	}
 
-	FGLTFJsonSkySphere JsonSkySphere;
-	SkySphereActor->GetName(JsonSkySphere.Name);
+	FGLTFJsonSkySphere* JsonSkySphere = Builder.AddSkySphere();
+	SkySphereActor->GetName(JsonSkySphere->Name);
 
 	const UStaticMeshComponent* StaticMeshComponent = nullptr;
 	FGLTFActorUtility::TryGetPropertyValue(SkySphereActor, TEXT("SkySphereMesh"), StaticMeshComponent);
@@ -32,8 +32,8 @@ FGLTFJsonSkySphere* FGLTFSkySphereConverter::Convert(const AActor* SkySphereActo
 		const FTransform ParentTransform = ParentComponent->GetSocketTransform(SocketName);
 		const FTransform RelativeTransform = Transform.GetRelativeTransform(ParentTransform);
 
-		JsonSkySphere.Scale = FGLTFConverterUtility::ConvertScale(RelativeTransform.GetScale3D());
-		JsonSkySphere.SkySphereMesh = Builder.GetOrAddMesh(StaticMeshComponent, { FGLTFMaterialUtility::GetDefaultMaterial() });
+		JsonSkySphere->Scale = FGLTFConverterUtility::ConvertScale(RelativeTransform.GetScale3D());
+		JsonSkySphere->SkySphereMesh = Builder.GetOrAddMesh(StaticMeshComponent, { FGLTFMaterialUtility::GetDefaultMaterial() });
 	}
 	else
 	{
@@ -43,7 +43,7 @@ FGLTFJsonSkySphere* FGLTFSkySphereConverter::Convert(const AActor* SkySphereActo
 	const ADirectionalLight* DirectionalLight = nullptr;
 	if (FGLTFActorUtility::TryGetPropertyValue(SkySphereActor, TEXT("Directional light actor"), DirectionalLight))
 	{
-		JsonSkySphere.DirectionalLight = Builder.GetOrAddNode(DirectionalLight);
+		JsonSkySphere->DirectionalLight = Builder.GetOrAddNode(DirectionalLight);
 	}
 	else
 	{
@@ -55,13 +55,13 @@ FGLTFJsonSkySphere* FGLTFSkySphereConverter::Convert(const AActor* SkySphereActo
 
 	if (SkyMaterial != nullptr)
 	{
-		ConvertScalarParameter(SkySphereActor, SkyMaterial, TEXT("Sun Radius"), JsonSkySphere.SunRadius);
-		ConvertScalarParameter(SkySphereActor, SkyMaterial, TEXT("NoisePower1"), JsonSkySphere.NoisePower1);
-		ConvertScalarParameter(SkySphereActor, SkyMaterial, TEXT("NoisePower2"), JsonSkySphere.NoisePower2);
+		ConvertScalarParameter(SkySphereActor, SkyMaterial, TEXT("Sun Radius"), JsonSkySphere->SunRadius);
+		ConvertScalarParameter(SkySphereActor, SkyMaterial, TEXT("NoisePower1"), JsonSkySphere->NoisePower1);
+		ConvertScalarParameter(SkySphereActor, SkyMaterial, TEXT("NoisePower2"), JsonSkySphere->NoisePower2);
 
-		ConvertTextureParameter(SkySphereActor, SkyMaterial, ESkySphereTextureParameter::SkyTexture, JsonSkySphere.SkyTexture);
-		ConvertTextureParameter(SkySphereActor, SkyMaterial, ESkySphereTextureParameter::CloudsTexture, JsonSkySphere.CloudsTexture);
-		ConvertTextureParameter(SkySphereActor, SkyMaterial, ESkySphereTextureParameter::StarsTexture, JsonSkySphere.StarsTexture);
+		ConvertTextureParameter(SkySphereActor, SkyMaterial, ESkySphereTextureParameter::SkyTexture, JsonSkySphere->SkyTexture);
+		ConvertTextureParameter(SkySphereActor, SkyMaterial, ESkySphereTextureParameter::CloudsTexture, JsonSkySphere->CloudsTexture);
+		ConvertTextureParameter(SkySphereActor, SkyMaterial, ESkySphereTextureParameter::StarsTexture, JsonSkySphere->StarsTexture);
 	}
 	else
 	{
@@ -70,24 +70,24 @@ FGLTFJsonSkySphere* FGLTFSkySphereConverter::Convert(const AActor* SkySphereActo
 			*SkySphereActor->GetName()));
 	}
 
-	ConvertProperty(SkySphereActor, TEXT("Sun height"), JsonSkySphere.SunHeight);
-	ConvertProperty(SkySphereActor, TEXT("Sun brightness"), JsonSkySphere.SunBrightness);
-	ConvertProperty(SkySphereActor, TEXT("Stars brightness"), JsonSkySphere.StarsBrightness);
-	ConvertProperty(SkySphereActor, TEXT("Cloud speed"), JsonSkySphere.CloudSpeed);
-	ConvertProperty(SkySphereActor, TEXT("Cloud opacity"), JsonSkySphere.CloudOpacity);
-	ConvertProperty(SkySphereActor, TEXT("Horizon Falloff"), JsonSkySphere.HorizonFalloff);
-	ConvertProperty(SkySphereActor, TEXT("Colors determined by sun position"), JsonSkySphere.bColorsDeterminedBySunPosition);
+	ConvertProperty(SkySphereActor, TEXT("Sun height"), JsonSkySphere->SunHeight);
+	ConvertProperty(SkySphereActor, TEXT("Sun brightness"), JsonSkySphere->SunBrightness);
+	ConvertProperty(SkySphereActor, TEXT("Stars brightness"), JsonSkySphere->StarsBrightness);
+	ConvertProperty(SkySphereActor, TEXT("Cloud speed"), JsonSkySphere->CloudSpeed);
+	ConvertProperty(SkySphereActor, TEXT("Cloud opacity"), JsonSkySphere->CloudOpacity);
+	ConvertProperty(SkySphereActor, TEXT("Horizon Falloff"), JsonSkySphere->HorizonFalloff);
+	ConvertProperty(SkySphereActor, TEXT("Colors determined by sun position"), JsonSkySphere->bColorsDeterminedBySunPosition);
 
-	ConvertColorProperty(SkySphereActor, TEXT("Zenith Color"), JsonSkySphere.ZenithColor);
-	ConvertColorProperty(SkySphereActor, TEXT("Horizon color"), JsonSkySphere.HorizonColor);
-	ConvertColorProperty(SkySphereActor, TEXT("Cloud color"), JsonSkySphere.CloudColor);
-	ConvertColorProperty(SkySphereActor, TEXT("Overall Color"), JsonSkySphere.OverallColor);
+	ConvertColorProperty(SkySphereActor, TEXT("Zenith Color"), JsonSkySphere->ZenithColor);
+	ConvertColorProperty(SkySphereActor, TEXT("Horizon color"), JsonSkySphere->HorizonColor);
+	ConvertColorProperty(SkySphereActor, TEXT("Cloud color"), JsonSkySphere->CloudColor);
+	ConvertColorProperty(SkySphereActor, TEXT("Overall Color"), JsonSkySphere->OverallColor);
 
-	ConvertColorCurveProperty(SkySphereActor, TEXT("Zenith color curve"), JsonSkySphere.ZenithColorCurve);
-	ConvertColorCurveProperty(SkySphereActor, TEXT("Horizon color curve"), JsonSkySphere.HorizonColorCurve);
-	ConvertColorCurveProperty(SkySphereActor, TEXT("Cloud color curve"), JsonSkySphere.CloudColorCurve);
+	ConvertColorCurveProperty(SkySphereActor, TEXT("Zenith color curve"), JsonSkySphere->ZenithColorCurve);
+	ConvertColorCurveProperty(SkySphereActor, TEXT("Horizon color curve"), JsonSkySphere->HorizonColorCurve);
+	ConvertColorCurveProperty(SkySphereActor, TEXT("Cloud color curve"), JsonSkySphere->CloudColorCurve);
 
-	return Builder.AddSkySphere(JsonSkySphere);
+	return JsonSkySphere;
 }
 
 template <class ValueType>
diff --git a/Engine/Plugins/Enterprise/GLTFExporter/MigratedFromGit/Source/GLTFExporter/Private/Exporters/GLTFAnimSequenceExporter.cpp b/Engine/Plugins/Enterprise/GLTFExporter/MigratedFromGit/Source/GLTFExporter/Private/Exporters/GLTFAnimSequenceExporter.cpp
--- a/Engine/Plugins/Enterprise/GLTFExporter/MigratedFromGit/Source/GLTFExporter/Private/Exporters/GLTFAnimSequenceExporter.cpp
+++ b/Engine/Plugins/Enterprise/GLTFExporter/MigratedFromGit/Source/GLTFExporter/Private/Exporters/GLTFAnimSequenceExporter.cpp
@@ -40,7 +40,7 @@ bool UGLTFAnimSequenceExporter::AddObject(FGLTFContainerBuilder& Builder, const
 		return false;
 	}
 
-	FGLTFJsonNode Node;
+	FGLTFJsonNode* Node = Builder.AddNode();
 
 	if (Builder.ExportOptions->bExportPreviewMesh)
 	{
@@ -54,12 +54,11 @@ bool UGLTFAnimSequenceExporter::AddObject(FGLTFContainerBuilder& Builder, const
 			return false;
 		}
 
-		Node.Mesh = MeshIndex;
+		Node->Mesh = MeshIndex;
 	}
 
-	FGLTFJsonNode* NodeIndex = Builder.AddNode(Node);
 
-	FGLTFJsonSkin* SkinIndex = Builder.GetOrAddSkin(NodeIndex, SkeletalMesh);
+	FGLTFJsonSkin* SkinIndex = Builder.GetOrAddSkin(Node, SkeletalMesh);
 	if (SkinIndex == nullptr)
 	{
 		Builder.LogError(
@@ -69,9 +68,9 @@ bool UGLTFAnimSequenceExporter::AddObject(FGLTFContainerBuilder& Builder, const
 		return false;
 	}
 
-	Builder.GetNode(NodeIndex).Skin = SkinIndex;
+	Node->Skin = SkinIndex;
 
-	FGLTFJsonAnimation* AnimationIndex =  Builder.GetOrAddAnimation(NodeIndex, SkeletalMesh, AnimSequence);
+	FGLTFJsonAnimation* AnimationIndex =  Builder.GetOrAddAnimation(Node, SkeletalMesh, AnimSequence);
 	if (AnimationIndex == nullptr)
 	{
 		Builder.LogError(
@@ -80,10 +79,9 @@ bool UGLTFAnimSequenceExporter::AddObject(FGLTFContainerBuilder& Builder, const
 		return false;
 	}
 
-	FGLTFJsonScene Scene;
-	Scene.Nodes.Add(NodeIndex);
-	FGLTFJsonScene* SceneIndex = Builder.AddScene(Scene);
+	FGLTFJsonScene* Scene = Builder.AddScene();
+	Scene->Nodes.Add(Node);
 
-	Builder.DefaultScene = SceneIndex;
+	Builder.DefaultScene = Scene;
 	return true;
 }
diff --git a/Engine/Plugins/Enterprise/GLTFExporter/MigratedFromGit/Source/GLTFExporter/Private/Exporters/GLTFMaterialExporter.cpp b/Engine/Plugins/Enterprise/GLTFExporter/MigratedFromGit/Source/GLTFExporter/Private/Exporters/GLTFMaterialExporter.cpp
--- a/Engine/Plugins/Enterprise/GLTFExporter/MigratedFromGit/Source/GLTFExporter/Private/Exporters/GLTFMaterialExporter.cpp
+++ b/Engine/Plugins/Enterprise/GLTFExporter/MigratedFromGit/Source/GLTFExporter/Private/Exporters/GLTFMaterialExporter.cpp
@@ -37,15 +37,13 @@ bool UGLTFMaterialExporter::AddObject(FGLTFContainerBuilder& Builder, const UObj
 			return false;
 		}
 
-		FGLTFJsonNode Node;
-		Node.Mesh = MeshIndex;
-		FGLTFJsonNode* NodeIndex = Builder.AddNode(Node);
+		FGLTFJsonNode* Node = Builder.AddNode();
+		Node->Mesh = MeshIndex;
 
-		FGLTFJsonScene Scene;
-		Scene.Nodes.Add(NodeIndex);
-		FGLTFJsonScene* SceneIndex = Builder.AddScene(Scene);
+		FGLTFJsonScene* Scene = Builder.AddScene();
+		Scene->Nodes.Add(Node);
 
-		Builder.DefaultScene = SceneIndex;
+		Builder.DefaultScene = Scene;
 	}
 	else
 	{
diff --git a/Engine/Plugins/Enterprise/GLTFExporter/MigratedFromGit/Source/GLTFExporter/Private/Exporters/GLTFSkeletalMeshExporter.cpp b/Engine/Plugins/Enterprise/GLTFExporter/MigratedFromGit/Source/GLTFExporter/Private/Exporters/GLTFSkeletalMeshExporter.cpp
--- a/Engine/Plugins/Enterprise/GLTFExporter/MigratedFromGit/Source/GLTFExporter/Private/Exporters/GLTFSkeletalMeshExporter.cpp
+++ b/Engine/Plugins/Enterprise/GLTFExporter/MigratedFromGit/Source/GLTFExporter/Private/Exporters/GLTFSkeletalMeshExporter.cpp
@@ -21,26 +21,24 @@ bool UGLTFSkeletalMeshExporter::AddObject(FGLTFContainerBuilder& Builder, const
 		return false;
 	}
 
-	FGLTFJsonNode Node;
-	Node.Mesh = MeshIndex;
-	FGLTFJsonNode* NodeIndex = Builder.AddNode(Node);
+	FGLTFJsonNode* Node = Builder.AddNode();
+	Node->Mesh = MeshIndex;
 
 	if (Builder.ExportOptions->bExportVertexSkinWeights)
 	{
-		FGLTFJsonSkin* SkinIndex = Builder.GetOrAddSkin(NodeIndex, SkeletalMesh);
+		FGLTFJsonSkin* SkinIndex = Builder.GetOrAddSkin(Node, SkeletalMesh);
 		if (SkinIndex == nullptr)
 		{
 			Builder.LogError(FString::Printf(TEXT("Failed to export bones in skeletal mesh %s"), *SkeletalMesh->GetName()));
 			return false;
 		}
 
-		Builder.GetNode(NodeIndex).Skin = SkinIndex;
+		Node->Skin = SkinIndex;
 	}
 
-	FGLTFJsonScene Scene;
-	Scene.Nodes.Add(NodeIndex);
-	FGLTFJsonScene* SceneIndex = Builder.AddScene(Scene);
+	FGLTFJsonScene* Scene = Builder.AddScene();
+	Scene->Nodes.Add(Node);
 
-	Builder.DefaultScene = SceneIndex;
+	Builder.DefaultScene = Scene;
 	return true;
 }
diff --git a/Engine/Plugins/Enterprise/GLTFExporter/MigratedFromGit/Source/GLTFExporter/Private/Exporters/GLTFStaticMeshExporter.cpp b/Engine/Plugins/Enterprise/GLTFExporter/MigratedFromGit/Source/GLTFExporter/Private/Exporters/GLTFStaticMeshExporter.cpp
--- a/Engine/Plugins/Enterprise/GLTFExporter/MigratedFromGit/Source/GLTFExporter/Private/Exporters/GLTFStaticMeshExporter.cpp
+++ b/Engine/Plugins/Enterprise/GLTFExporter/MigratedFromGit/Source/GLTFExporter/Private/Exporters/GLTFStaticMeshExporter.cpp
@@ -21,14 +21,12 @@ bool UGLTFStaticMeshExporter::AddObject(FGLTFContainerBuilder& Builder, const UO
 		return false;
 	}
 
-	FGLTFJsonNode Node;
-	Node.Mesh = MeshIndex;
-	FGLTFJsonNode* NodeIndex = Builder.AddNode(Node);
+	FGLTFJsonNode* Node = Builder.AddNode();
+	Node->Mesh = MeshIndex;
 
-	FGLTFJsonScene Scene;
-	Scene.Nodes.Add(NodeIndex);
-	FGLTFJsonScene* SceneIndex = Builder.AddScene(Scene);
+	FGLTFJsonScene* Scene = Builder.AddScene();
+	Scene->Nodes.Add(Node);
 
-	Builder.DefaultScene = SceneIndex;
+	Builder.DefaultScene = Scene;
 	return true;
 }
diff --git a/Engine/Plugins/Enterprise/GLTFExporter/MigratedFromGit/Source/GLTFExporter/Private/Tasks/GLTFAnimationTasks.cpp b/Engine/Plugins/Enterprise/GLTFExporter/MigratedFromGit/Source/GLTFExporter/Private/Tasks/GLTFAnimationTasks.cpp
--- a/Engine/Plugins/Enterprise/GLTFExporter/MigratedFromGit/Source/GLTFExporter/Private/Tasks/GLTFAnimationTasks.cpp
+++ b/Engine/Plugins/Enterprise/GLTFExporter/MigratedFromGit/Source/GLTFExporter/Private/Tasks/GLTFAnimationTasks.cpp
@@ -27,18 +27,14 @@ void FGLTFAnimSequenceTask::Complete()
 
 	// TODO: add animation data accessor converters to reuse track information
 
-	FGLTFJsonAccessor* InputAccessorIndex;
-	{
-		FGLTFJsonAccessor JsonInputAccessor;
-		JsonInputAccessor.BufferView = Builder.AddBufferView(Timestamps);
-		JsonInputAccessor.ComponentType = EGLTFJsonComponentType::Float;
-		JsonInputAccessor.Type = EGLTFJsonAccessorType::Scalar;
-		JsonInputAccessor.Count = FrameCount;
-		JsonInputAccessor.MinMaxLength = 1;
-		JsonInputAccessor.Min[0] = 0;
-		JsonInputAccessor.Max[0] = Timestamps[FrameCount - 1];
-		InputAccessorIndex = Builder.AddAccessor(JsonInputAccessor);
-	}
+	FGLTFJsonAccessor* JsonInputAccessor = Builder.AddAccessor();
+	JsonInputAccessor->BufferView = Builder.AddBufferView(Timestamps);
+	JsonInputAccessor->ComponentType = EGLTFJsonComponentType::Float;
+	JsonInputAccessor->Type = EGLTFJsonAccessorType::Scalar;
+	JsonInputAccessor->Count = FrameCount;
+	JsonInputAccessor->MinMaxLength = 1;
+	JsonInputAccessor->Min[0] = 0;
+	JsonInputAccessor->Max[0] = Timestamps[FrameCount - 1];
 
 	EGLTFJsonInterpolation Interpolation = FGLTFConverterUtility::ConvertInterpolation(AnimSequence->Interpolation);
 	if (Interpolation == EGLTFJsonInterpolation::None)
@@ -69,19 +65,19 @@ void FGLTFAnimSequenceTask::Complete()
 				Translations[Frame] = FGLTFConverterUtility::ConvertPosition(KeyPosition, Builder.ExportOptions->ExportUniformScale);
 			}
 
-			FGLTFJsonAccessor JsonOutputAccessor;
-			JsonOutputAccessor.BufferView = Builder.AddBufferView(Translations);
-			JsonOutputAccessor.ComponentType = EGLTFJsonComponentType::Float;
-			JsonOutputAccessor.Count = Translations.Num();
-			JsonOutputAccessor.Type = EGLTFJsonAccessorType::Vec3;
+			FGLTFJsonAccessor* JsonOutputAccessor = Builder.AddAccessor();
+			JsonOutputAccessor->BufferView = Builder.AddBufferView(Translations);
+			JsonOutputAccessor->ComponentType = EGLTFJsonComponentType::Float;
+			JsonOutputAccessor->Count = Translations.Num();
+			JsonOutputAccessor->Type = EGLTFJsonAccessorType::Vec3;
 
-			FGLTFJsonAnimationSampler JsonSampler;
-			JsonSampler.Input = InputAccessorIndex;
-			JsonSampler.Output = Builder.AddAccessor(JsonOutputAccessor);
-			JsonSampler.Interpolation = Interpolation;
+			FGLTFJsonAnimationSampler* JsonSampler = JsonAnimation->Samplers.Add();
+			JsonSampler->Input = JsonInputAccessor;
+			JsonSampler->Output = JsonOutputAccessor;
+			JsonSampler->Interpolation = Interpolation;
 
 			FGLTFJsonAnimationChannel JsonChannel;
-			JsonChannel.Sampler = JsonAnimation->Samplers.Add(JsonSampler);
+			JsonChannel.Sampler = JsonSampler;
 			JsonChannel.Target.Path = EGLTFJsonTargetPath::Translation;
 			JsonChannel.Target.Node = NodeIndex;
 			JsonAnimation->Channels.Add(JsonChannel);
@@ -97,19 +93,19 @@ void FGLTFAnimSequenceTask::Complete()
 				Rotations[Frame] = FGLTFConverterUtility::ConvertRotation(KeyRotation);
 			}
 
-			FGLTFJsonAccessor JsonOutputAccessor;
-			JsonOutputAccessor.BufferView = Builder.AddBufferView(Rotations);
-			JsonOutputAccessor.ComponentType = EGLTFJsonComponentType::Float;
-			JsonOutputAccessor.Count = Rotations.Num();
-			JsonOutputAccessor.Type = EGLTFJsonAccessorType::Vec4;
+			FGLTFJsonAccessor* JsonOutputAccessor = Builder.AddAccessor();
+			JsonOutputAccessor->BufferView = Builder.AddBufferView(Rotations);
+			JsonOutputAccessor->ComponentType = EGLTFJsonComponentType::Float;
+			JsonOutputAccessor->Count = Rotations.Num();
+			JsonOutputAccessor->Type = EGLTFJsonAccessorType::Vec4;
 
-			FGLTFJsonAnimationSampler JsonSampler;
-			JsonSampler.Input = InputAccessorIndex;
-			JsonSampler.Output = Builder.AddAccessor(JsonOutputAccessor);
-			JsonSampler.Interpolation = Interpolation;
+			FGLTFJsonAnimationSampler* JsonSampler = JsonAnimation->Samplers.Add();
+			JsonSampler->Input = JsonInputAccessor;
+			JsonSampler->Output = JsonOutputAccessor;
+			JsonSampler->Interpolation = Interpolation;
 
 			FGLTFJsonAnimationChannel JsonChannel;
-			JsonChannel.Sampler = JsonAnimation->Samplers.Add(JsonSampler);
+			JsonChannel.Sampler = JsonSampler;
 			JsonChannel.Target.Path = EGLTFJsonTargetPath::Rotation;
 			JsonChannel.Target.Node = NodeIndex;
 			JsonAnimation->Channels.Add(JsonChannel);
@@ -125,19 +121,19 @@ void FGLTFAnimSequenceTask::Complete()
 				Scales[Frame] = FGLTFConverterUtility::ConvertScale(KeyScale);
 			}
 
-			FGLTFJsonAccessor JsonOutputAccessor;
-			JsonOutputAccessor.BufferView = Builder.AddBufferView(Scales);
-			JsonOutputAccessor.ComponentType = EGLTFJsonComponentType::Float;
-			JsonOutputAccessor.Count = Scales.Num();
-			JsonOutputAccessor.Type = EGLTFJsonAccessorType::Vec3;
+			FGLTFJsonAccessor* JsonOutputAccessor = Builder.AddAccessor();
+			JsonOutputAccessor->BufferView = Builder.AddBufferView(Scales);
+			JsonOutputAccessor->ComponentType = EGLTFJsonComponentType::Float;
+			JsonOutputAccessor->Count = Scales.Num();
+			JsonOutputAccessor->Type = EGLTFJsonAccessorType::Vec3;
 
-			FGLTFJsonAnimationSampler JsonSampler;
-			JsonSampler.Input = InputAccessorIndex;
-			JsonSampler.Output = Builder.AddAccessor(JsonOutputAccessor);
-			JsonSampler.Interpolation = Interpolation;
+			FGLTFJsonAnimationSampler* JsonSampler = JsonAnimation->Samplers.Add();
+			JsonSampler->Input = JsonInputAccessor;
+			JsonSampler->Output = JsonOutputAccessor;
+			JsonSampler->Interpolation = Interpolation;
 
 			FGLTFJsonAnimationChannel JsonChannel;
-			JsonChannel.Sampler = JsonAnimation->Samplers.Add(JsonSampler);
+			JsonChannel.Sampler = JsonSampler;
 			JsonChannel.Target.Path = EGLTFJsonTargetPath::Scale;
 			JsonChannel.Target.Node = NodeIndex;
 			JsonAnimation->Channels.Add(JsonChannel);
@@ -179,15 +175,14 @@ void FGLTFLevelSequenceTask::Complete()
 
 	// TODO: add animation data accessor converters to reuse track information
 
-	FGLTFJsonAccessor JsonInputAccessor;
-	JsonInputAccessor.BufferView = Builder.AddBufferView(Timestamps);
-	JsonInputAccessor.ComponentType = EGLTFJsonComponentType::Float;
-	JsonInputAccessor.Type = EGLTFJsonAccessorType::Scalar;
-	JsonInputAccessor.Count = FrameCount;
-	JsonInputAccessor.MinMaxLength = 1;
-	JsonInputAccessor.Min[0] = 0;
-	JsonInputAccessor.Max[0] = Timestamps[FrameCount - 1];
-	FGLTFJsonAccessor* InputAccessorIndex = Builder.AddAccessor(JsonInputAccessor);
+	FGLTFJsonAccessor* JsonInputAccessor = Builder.AddAccessor();
+	JsonInputAccessor->BufferView = Builder.AddBufferView(Timestamps);
+	JsonInputAccessor->ComponentType = EGLTFJsonComponentType::Float;
+	JsonInputAccessor->Type = EGLTFJsonAccessorType::Scalar;
+	JsonInputAccessor->Count = FrameCount;
+	JsonInputAccessor->MinMaxLength = 1;
+	JsonInputAccessor->Min[0] = 0;
+	JsonInputAccessor->Max[0] = Timestamps[FrameCount - 1];
 
 	for (const FMovieSceneBinding& Binding : MovieScene->GetBindings())
 	{
@@ -264,19 +259,19 @@ void FGLTFLevelSequenceTask::Complete()
 							Translations[Frame] = FGLTFConverterUtility::ConvertPosition(Translation, Builder.ExportOptions->ExportUniformScale);
 						}
 
-						FGLTFJsonAccessor JsonOutputAccessor;
-						JsonOutputAccessor.BufferView = Builder.AddBufferView(Translations);
-						JsonOutputAccessor.ComponentType = EGLTFJsonComponentType::Float;
-						JsonOutputAccessor.Count = Translations.Num();
-						JsonOutputAccessor.Type = EGLTFJsonAccessorType::Vec3;
+						FGLTFJsonAccessor* JsonOutputAccessor = Builder.AddAccessor();
+						JsonOutputAccessor->BufferView = Builder.AddBufferView(Translations);
+						JsonOutputAccessor->ComponentType = EGLTFJsonComponentType::Float;
+						JsonOutputAccessor->Count = Translations.Num();
+						JsonOutputAccessor->Type = EGLTFJsonAccessorType::Vec3;
 
-						FGLTFJsonAnimationSampler JsonSampler;
-						JsonSampler.Input = InputAccessorIndex;
-						JsonSampler.Output = Builder.AddAccessor(JsonOutputAccessor);
-						JsonSampler.Interpolation = EGLTFJsonInterpolation::Linear;
+						FGLTFJsonAnimationSampler* JsonSampler = JsonAnimation->Samplers.Add();
+						JsonSampler->Input = JsonInputAccessor;
+						JsonSampler->Output = JsonOutputAccessor;
+						JsonSampler->Interpolation = EGLTFJsonInterpolation::Linear;
 
 						FGLTFJsonAnimationChannel JsonChannel;
-						JsonChannel.Sampler = JsonAnimation->Samplers.Add(JsonSampler);
+						JsonChannel.Sampler = JsonSampler;
 						JsonChannel.Target.Path = EGLTFJsonTargetPath::Translation;
 						JsonChannel.Target.Node = NodeIndex;
 						JsonAnimation->Channels.Add(JsonChannel);
@@ -299,22 +294,22 @@ void FGLTFLevelSequenceTask::Complete()
 							Rotations[Frame] = FGLTFConverterUtility::ConvertRotation(Rotator.Quaternion());
 						}
 
-						JsonInputAccessor.Count = Rotations.Num();
-						JsonInputAccessor.Max[0] = Timestamps[Rotations.Num() - 1];
+						JsonInputAccessor->Count = Rotations.Num();
+						JsonInputAccessor->Max[0] = Timestamps[Rotations.Num() - 1];
 
-						FGLTFJsonAccessor JsonOutputAccessor;
-						JsonOutputAccessor.BufferView = Builder.AddBufferView(Rotations);
-						JsonOutputAccessor.ComponentType = EGLTFJsonComponentType::Float;
-						JsonOutputAccessor.Count = Rotations.Num();
-						JsonOutputAccessor.Type = EGLTFJsonAccessorType::Vec4;
+						FGLTFJsonAccessor* JsonOutputAccessor = Builder.AddAccessor();
+						JsonOutputAccessor->BufferView = Builder.AddBufferView(Rotations);
+						JsonOutputAccessor->ComponentType = EGLTFJsonComponentType::Float;
+						JsonOutputAccessor->Count = Rotations.Num();
+						JsonOutputAccessor->Type = EGLTFJsonAccessorType::Vec4;
 
-						FGLTFJsonAnimationSampler JsonSampler;
-						JsonSampler.Input = InputAccessorIndex;
-						JsonSampler.Output = Builder.AddAccessor(JsonOutputAccessor);
-						JsonSampler.Interpolation = EGLTFJsonInterpolation::Linear;
+						FGLTFJsonAnimationSampler* JsonSampler = JsonAnimation->Samplers.Add();
+						JsonSampler->Input = JsonInputAccessor;
+						JsonSampler->Output = JsonOutputAccessor;
+						JsonSampler->Interpolation = EGLTFJsonInterpolation::Linear;
 
 						FGLTFJsonAnimationChannel JsonChannel;
-						JsonChannel.Sampler = JsonAnimation->Samplers.Add(JsonSampler);
+						JsonChannel.Sampler = JsonSampler;
 						JsonChannel.Target.Path = EGLTFJsonTargetPath::Rotation;
 						JsonChannel.Target.Node = NodeIndex;
 						JsonAnimation->Channels.Add(JsonChannel);
@@ -337,19 +332,19 @@ void FGLTFLevelSequenceTask::Complete()
 							Scales[Frame] = FGLTFConverterUtility::ConvertScale(Scale);
 						}
 
-						FGLTFJsonAccessor JsonOutputAccessor;
-						JsonOutputAccessor.BufferView = Builder.AddBufferView(Scales);
-						JsonOutputAccessor.ComponentType = EGLTFJsonComponentType::Float;
-						JsonOutputAccessor.Count = Scales.Num();
-						JsonOutputAccessor.Type = EGLTFJsonAccessorType::Vec3;
+						FGLTFJsonAccessor* JsonOutputAccessor = Builder.AddAccessor();
+						JsonOutputAccessor->BufferView = Builder.AddBufferView(Scales);
+						JsonOutputAccessor->ComponentType = EGLTFJsonComponentType::Float;
+						JsonOutputAccessor->Count = Scales.Num();
+						JsonOutputAccessor->Type = EGLTFJsonAccessorType::Vec3;
 
-						FGLTFJsonAnimationSampler JsonSampler;
-						JsonSampler.Input = InputAccessorIndex;
-						JsonSampler.Output = Builder.AddAccessor(JsonOutputAccessor);
-						JsonSampler.Interpolation = EGLTFJsonInterpolation::Linear;
+						FGLTFJsonAnimationSampler* JsonSampler = JsonAnimation->Samplers.Add();
+						JsonSampler->Input = JsonInputAccessor;
+						JsonSampler->Output = JsonOutputAccessor;
+						JsonSampler->Interpolation = EGLTFJsonInterpolation::Linear;
 
 						FGLTFJsonAnimationChannel JsonChannel;
-						JsonChannel.Sampler = JsonAnimation->Samplers.Add(JsonSampler);
+						JsonChannel.Sampler = JsonSampler;
 						JsonChannel.Target.Path = EGLTFJsonTargetPath::Scale;
 						JsonChannel.Target.Node = NodeIndex;
 						JsonAnimation->Channels.Add(JsonChannel);
diff --git a/Engine/Plugins/Enterprise/GLTFExporter/MigratedFromGit/Source/GLTFExporter/Public/Json/GLTFJsonIndexedObjectArray.h b/Engine/Plugins/Enterprise/GLTFExporter/MigratedFromGit/Source/GLTFExporter/Public/Json/GLTFJsonIndexedObjectArray.h
--- a/Engine/Plugins/Enterprise/GLTFExporter/MigratedFromGit/Source/GLTFExporter/Public/Json/GLTFJsonIndexedObjectArray.h
+++ b/Engine/Plugins/Enterprise/GLTFExporter/MigratedFromGit/Source/GLTFExporter/Public/Json/GLTFJsonIndexedObjectArray.h
@@ -57,15 +57,6 @@ struct TGLTFJsonIndexedObjectArray : IGLTFJsonArray
 	TGLTFJsonIndexedObjectArray(const TGLTFJsonIndexedObjectArray&) = delete;
 	TGLTFJsonIndexedObjectArray& operator=(const TGLTFJsonIndexedObjectArray&) = delete;
 
-	ElementType* Add(const ElementType& Temp)
-	{
-		ElementType* Element = new ElementType(Array.Num());
-		*Element = Temp;
-		Element->Index = Array.Num();
-		Array.Add(TUniquePtr<ElementType>(Element));
-		return Element;
-	}
-
 	ElementType* Add()
 	{
 		ElementType* Element = new ElementType(Array.Num());
