[GLTFExporter] refactor: make all IGLTFJsonIndexedObject-inherited constructors only accessible to TGLTFJsonIndexedObjectArray<*>

commit 33347654eb44d603214b6f1e3c7058fab34c3809
Date: Sun, 13 Mar 2022 18:08:08 +0100

    refactor: make all IGLTFJsonIndexedObject-inherited constructors only accessible to TGLTFJsonIndexedObjectArray<*>
    (cherry picked from commit 1a5cfaf84bc16c459a67ac3d06b7fdb12ffb7832​)

#skipci

diff --git a/Engine/Plugins/Enterprise/GLTFExporter/MigratedFromGit/Source/GLTFExporter/Public/Json/GLTFJsonAccessor.h b/Engine/Plugins/Enterprise/GLTFExporter/MigratedFromGit/Source/GLTFExporter/Public/Json/GLTFJsonAccessor.h
--- a/Engine/Plugins/Enterprise/GLTFExporter/MigratedFromGit/Source/GLTFExporter/Public/Json/GLTFJsonAccessor.h
+++ b/Engine/Plugins/Enterprise/GLTFExporter/MigratedFromGit/Source/GLTFExporter/Public/Json/GLTFJsonAccessor.h
@@ -19,7 +19,13 @@ struct GLTFEXPORTER_API FGLTFJsonAccessor : IGLTFJsonIndexedObject
 	float Min[16];
 	float Max[16];
 
-	FGLTFJsonAccessor(int32 Index = INDEX_NONE)
+	virtual void WriteObject(IGLTFJsonWriter& Writer) const override;
+
+protected:
+
+	friend TGLTFJsonIndexedObjectArray<FGLTFJsonAccessor, void>;
+
+	FGLTFJsonAccessor(int32 Index)
 		: IGLTFJsonIndexedObject(Index)
 		, BufferView(nullptr)
 		, ByteOffset(0)
@@ -29,9 +35,7 @@ struct GLTFEXPORTER_API FGLTFJsonAccessor : IGLTFJsonIndexedObject
 		, bNormalized(false)
 		, MinMaxLength(0)
 		, Min{0}
-		, Max{0}
+	, Max{0}
 	{
 	}
-
-	virtual void WriteObject(IGLTFJsonWriter& Writer) const override;
 };
diff --git a/Engine/Plugins/Enterprise/GLTFExporter/MigratedFromGit/Source/GLTFExporter/Public/Json/GLTFJsonAnimation.h b/Engine/Plugins/Enterprise/GLTFExporter/MigratedFromGit/Source/GLTFExporter/Public/Json/GLTFJsonAnimation.h
--- a/Engine/Plugins/Enterprise/GLTFExporter/MigratedFromGit/Source/GLTFExporter/Public/Json/GLTFJsonAnimation.h
+++ b/Engine/Plugins/Enterprise/GLTFExporter/MigratedFromGit/Source/GLTFExporter/Public/Json/GLTFJsonAnimation.h
@@ -39,15 +39,19 @@ struct GLTFEXPORTER_API FGLTFJsonAnimationSampler : IGLTFJsonIndexedObject
 
 	EGLTFJsonInterpolation Interpolation;
 
-	FGLTFJsonAnimationSampler(int32 Index = INDEX_NONE)
+	virtual void WriteObject(IGLTFJsonWriter& Writer) const override;
+
+protected:
+
+	friend TGLTFJsonIndexedObjectArray<FGLTFJsonAnimationSampler, void>;
+
+	FGLTFJsonAnimationSampler(int32 Index)
 		: IGLTFJsonIndexedObject(Index)
 		, Input(nullptr)
 		, Output(nullptr)
 		, Interpolation(EGLTFJsonInterpolation::Linear)
 	{
 	}
-
-	virtual void WriteObject(IGLTFJsonWriter& Writer) const override;
 };
 
 struct GLTFEXPORTER_API FGLTFJsonAnimation : IGLTFJsonIndexedObject
@@ -59,10 +63,14 @@ struct GLTFEXPORTER_API FGLTFJsonAnimation : IGLTFJsonIndexedObject
 
 	FGLTFJsonAnimationPlayback Playback;
 
-	FGLTFJsonAnimation(int32 Index = INDEX_NONE)
+	virtual void WriteObject(IGLTFJsonWriter& Writer) const override;
+
+protected:
+
+	friend TGLTFJsonIndexedObjectArray<FGLTFJsonAnimation, void>;
+
+	FGLTFJsonAnimation(int32 Index)
 		: IGLTFJsonIndexedObject(Index)
 	{
 	}
-
-	virtual void WriteObject(IGLTFJsonWriter& Writer) const override;
 };
diff --git a/Engine/Plugins/Enterprise/GLTFExporter/MigratedFromGit/Source/GLTFExporter/Public/Json/GLTFJsonBackdrop.h b/Engine/Plugins/Enterprise/GLTFExporter/MigratedFromGit/Source/GLTFExporter/Public/Json/GLTFJsonBackdrop.h
--- a/Engine/Plugins/Enterprise/GLTFExporter/MigratedFromGit/Source/GLTFExporter/Public/Json/GLTFJsonBackdrop.h
+++ b/Engine/Plugins/Enterprise/GLTFExporter/MigratedFromGit/Source/GLTFExporter/Public/Json/GLTFJsonBackdrop.h
@@ -20,7 +20,13 @@ struct GLTFEXPORTER_API FGLTFJsonBackdrop : IGLTFJsonIndexedObject
 	float LightingDistanceFactor;
 	bool UseCameraProjection;
 
-	FGLTFJsonBackdrop(int32 Index = INDEX_NONE)
+	virtual void WriteObject(IGLTFJsonWriter& Writer) const override;
+
+protected:
+
+	friend TGLTFJsonIndexedObjectArray<FGLTFJsonBackdrop, void>;
+
+	FGLTFJsonBackdrop(int32 Index)
 		: IGLTFJsonIndexedObject(Index)
 		, Mesh(nullptr)
 		, Cubemap{nullptr}
@@ -32,6 +38,4 @@ struct GLTFEXPORTER_API FGLTFJsonBackdrop : IGLTFJsonIndexedObject
 		, UseCameraProjection(false)
 	{
 	}
-
-	virtual void WriteObject(IGLTFJsonWriter& Writer) const override;
 };
diff --git a/Engine/Plugins/Enterprise/GLTFExporter/MigratedFromGit/Source/GLTFExporter/Public/Json/GLTFJsonBuffer.h b/Engine/Plugins/Enterprise/GLTFExporter/MigratedFromGit/Source/GLTFExporter/Public/Json/GLTFJsonBuffer.h
--- a/Engine/Plugins/Enterprise/GLTFExporter/MigratedFromGit/Source/GLTFExporter/Public/Json/GLTFJsonBuffer.h
+++ b/Engine/Plugins/Enterprise/GLTFExporter/MigratedFromGit/Source/GLTFExporter/Public/Json/GLTFJsonBuffer.h
@@ -11,11 +11,15 @@ struct GLTFEXPORTER_API FGLTFJsonBuffer : IGLTFJsonIndexedObject
 	FString URI;
 	int64   ByteLength;
 
-	FGLTFJsonBuffer(int32 Index = INDEX_NONE)
+	virtual void WriteObject(IGLTFJsonWriter& Writer) const override;
+
+protected:
+
+	friend TGLTFJsonIndexedObjectArray<FGLTFJsonBuffer, void>;
+
+	FGLTFJsonBuffer(int32 Index)
 		: IGLTFJsonIndexedObject(Index)
 		, ByteLength(0)
 	{
 	}
-
-	virtual void WriteObject(IGLTFJsonWriter& Writer) const override;
 };
diff --git a/Engine/Plugins/Enterprise/GLTFExporter/MigratedFromGit/Source/GLTFExporter/Public/Json/GLTFJsonBufferView.h b/Engine/Plugins/Enterprise/GLTFExporter/MigratedFromGit/Source/GLTFExporter/Public/Json/GLTFJsonBufferView.h
--- a/Engine/Plugins/Enterprise/GLTFExporter/MigratedFromGit/Source/GLTFExporter/Public/Json/GLTFJsonBufferView.h
+++ b/Engine/Plugins/Enterprise/GLTFExporter/MigratedFromGit/Source/GLTFExporter/Public/Json/GLTFJsonBufferView.h
@@ -16,7 +16,13 @@ struct GLTFEXPORTER_API FGLTFJsonBufferView : IGLTFJsonIndexedObject
 
 	EGLTFJsonBufferTarget Target;
 
-	FGLTFJsonBufferView(int32 Index = INDEX_NONE)
+	virtual void WriteObject(IGLTFJsonWriter& Writer) const override;
+
+protected:
+
+	friend TGLTFJsonIndexedObjectArray<FGLTFJsonBufferView, void>;
+
+	FGLTFJsonBufferView(int32 Index)
 		: IGLTFJsonIndexedObject(Index)
 		, Buffer(nullptr)
 		, ByteLength(0)
@@ -25,6 +31,4 @@ struct GLTFEXPORTER_API FGLTFJsonBufferView : IGLTFJsonIndexedObject
 		, Target(EGLTFJsonBufferTarget::None)
 	{
 	}
-
-	virtual void WriteObject(IGLTFJsonWriter& Writer) const override;
 };
diff --git a/Engine/Plugins/Enterprise/GLTFExporter/MigratedFromGit/Source/GLTFExporter/Public/Json/GLTFJsonCamera.h b/Engine/Plugins/Enterprise/GLTFExporter/MigratedFromGit/Source/GLTFExporter/Public/Json/GLTFJsonCamera.h
--- a/Engine/Plugins/Enterprise/GLTFExporter/MigratedFromGit/Source/GLTFExporter/Public/Json/GLTFJsonCamera.h
+++ b/Engine/Plugins/Enterprise/GLTFExporter/MigratedFromGit/Source/GLTFExporter/Public/Json/GLTFJsonCamera.h
@@ -51,11 +51,15 @@ struct GLTFEXPORTER_API FGLTFJsonCamera : IGLTFJsonIndexedObject
 	FGLTFJsonOrthographic Orthographic;
 	FGLTFJsonPerspective  Perspective;
 
-	FGLTFJsonCamera(int32 Index = INDEX_NONE)
+	virtual void WriteObject(IGLTFJsonWriter& Writer) const override;
+
+protected:
+
+	friend TGLTFJsonIndexedObjectArray<FGLTFJsonCamera, void>;
+
+	FGLTFJsonCamera(int32 Index)
 		: IGLTFJsonIndexedObject(Index)
 		, Type(EGLTFJsonCameraType::None)
 	{
 	}
-
-	virtual void WriteObject(IGLTFJsonWriter& Writer) const override;
 };
diff --git a/Engine/Plugins/Enterprise/GLTFExporter/MigratedFromGit/Source/GLTFExporter/Public/Json/GLTFJsonEpicLevelVariantSets.h b/Engine/Plugins/Enterprise/GLTFExporter/MigratedFromGit/Source/GLTFExporter/Public/Json/GLTFJsonEpicLevelVariantSets.h
--- a/Engine/Plugins/Enterprise/GLTFExporter/MigratedFromGit/Source/GLTFExporter/Public/Json/GLTFJsonEpicLevelVariantSets.h
+++ b/Engine/Plugins/Enterprise/GLTFExporter/MigratedFromGit/Source/GLTFExporter/Public/Json/GLTFJsonEpicLevelVariantSets.h
@@ -66,10 +66,14 @@ struct GLTFEXPORTER_API FGLTFJsonEpicLevelVariantSets : IGLTFJsonIndexedObject
 
 	TArray<FGLTFJsonEpicVariantSet> VariantSets;
 
-	FGLTFJsonEpicLevelVariantSets(int32 Index = INDEX_NONE)
+	virtual void WriteObject(IGLTFJsonWriter& Writer) const override;
+
+protected:
+
+	friend TGLTFJsonIndexedObjectArray<FGLTFJsonEpicLevelVariantSets, void>;
+
+	FGLTFJsonEpicLevelVariantSets(int32 Index)
 		: IGLTFJsonIndexedObject(Index)
 	{
 	}
-
-	virtual void WriteObject(IGLTFJsonWriter& Writer) const override;
 };
diff --git a/Engine/Plugins/Enterprise/GLTFExporter/MigratedFromGit/Source/GLTFExporter/Public/Json/GLTFJsonHotspot.h b/Engine/Plugins/Enterprise/GLTFExporter/MigratedFromGit/Source/GLTFExporter/Public/Json/GLTFJsonHotspot.h
--- a/Engine/Plugins/Enterprise/GLTFExporter/MigratedFromGit/Source/GLTFExporter/Public/Json/GLTFJsonHotspot.h
+++ b/Engine/Plugins/Enterprise/GLTFExporter/MigratedFromGit/Source/GLTFExporter/Public/Json/GLTFJsonHotspot.h
@@ -14,7 +14,13 @@ struct GLTFEXPORTER_API FGLTFJsonHotspot : IGLTFJsonIndexedObject
 	FGLTFJsonTexture*   ToggledImage;
 	FGLTFJsonTexture*   ToggledHoveredImage;
 
-	FGLTFJsonHotspot(int32 Index = INDEX_NONE)
+	virtual void WriteObject(IGLTFJsonWriter& Writer) const override;
+
+protected:
+
+	friend TGLTFJsonIndexedObjectArray<FGLTFJsonHotspot, void>;
+
+	FGLTFJsonHotspot(int32 Index)
 		: IGLTFJsonIndexedObject(Index)
 		, Animation(nullptr)
 		, Image(nullptr)
@@ -23,6 +29,4 @@ struct GLTFEXPORTER_API FGLTFJsonHotspot : IGLTFJsonIndexedObject
 		, ToggledHoveredImage(nullptr)
 	{
 	}
-
-	virtual void WriteObject(IGLTFJsonWriter& Writer) const override;
 };
diff --git a/Engine/Plugins/Enterprise/GLTFExporter/MigratedFromGit/Source/GLTFExporter/Public/Json/GLTFJsonImage.h b/Engine/Plugins/Enterprise/GLTFExporter/MigratedFromGit/Source/GLTFExporter/Public/Json/GLTFJsonImage.h
--- a/Engine/Plugins/Enterprise/GLTFExporter/MigratedFromGit/Source/GLTFExporter/Public/Json/GLTFJsonImage.h
+++ b/Engine/Plugins/Enterprise/GLTFExporter/MigratedFromGit/Source/GLTFExporter/Public/Json/GLTFJsonImage.h
@@ -13,12 +13,16 @@ struct GLTFEXPORTER_API FGLTFJsonImage : IGLTFJsonIndexedObject
 
 	FGLTFJsonBufferView* BufferView;
 
-	FGLTFJsonImage(int32 Index = INDEX_NONE)
+	virtual void WriteObject(IGLTFJsonWriter& Writer) const override;
+
+protected:
+
+	friend TGLTFJsonIndexedObjectArray<FGLTFJsonImage, void>;
+
+	FGLTFJsonImage(int32 Index)
 		: IGLTFJsonIndexedObject(Index)
 		, MimeType(EGLTFJsonMimeType::None)
 		, BufferView(nullptr)
 	{
 	}
-
-	virtual void WriteObject(IGLTFJsonWriter& Writer) const override;
 };
diff --git a/Engine/Plugins/Enterprise/GLTFExporter/MigratedFromGit/Source/GLTFExporter/Public/Json/GLTFJsonKhrMaterialVariant.h b/Engine/Plugins/Enterprise/GLTFExporter/MigratedFromGit/Source/GLTFExporter/Public/Json/GLTFJsonKhrMaterialVariant.h
--- a/Engine/Plugins/Enterprise/GLTFExporter/MigratedFromGit/Source/GLTFExporter/Public/Json/GLTFJsonKhrMaterialVariant.h
+++ b/Engine/Plugins/Enterprise/GLTFExporter/MigratedFromGit/Source/GLTFExporter/Public/Json/GLTFJsonKhrMaterialVariant.h
@@ -21,10 +21,14 @@ struct GLTFEXPORTER_API FGLTFJsonKhrMaterialVariant : IGLTFJsonIndexedObject
 {
 	FString Name;
 
-	FGLTFJsonKhrMaterialVariant(int32 Index = INDEX_NONE)
+	virtual void WriteObject(IGLTFJsonWriter& Writer) const override;
+
+protected:
+
+	friend TGLTFJsonIndexedObjectArray<FGLTFJsonKhrMaterialVariant, void>;
+
+	FGLTFJsonKhrMaterialVariant(int32 Index)
 		: IGLTFJsonIndexedObject(Index)
 	{
 	}
-
-	virtual void WriteObject(IGLTFJsonWriter& Writer) const override;
 };
diff --git a/Engine/Plugins/Enterprise/GLTFExporter/MigratedFromGit/Source/GLTFExporter/Public/Json/GLTFJsonLight.h b/Engine/Plugins/Enterprise/GLTFExporter/MigratedFromGit/Source/GLTFExporter/Public/Json/GLTFJsonLight.h
--- a/Engine/Plugins/Enterprise/GLTFExporter/MigratedFromGit/Source/GLTFExporter/Public/Json/GLTFJsonLight.h
+++ b/Engine/Plugins/Enterprise/GLTFExporter/MigratedFromGit/Source/GLTFExporter/Public/Json/GLTFJsonLight.h
@@ -31,7 +31,13 @@ struct GLTFEXPORTER_API FGLTFJsonLight : IGLTFJsonIndexedObject
 
 	FGLTFJsonSpotLight Spot;
 
-	FGLTFJsonLight(int32 Index = INDEX_NONE)
+	virtual void WriteObject(IGLTFJsonWriter& Writer) const override;
+
+protected:
+
+	friend TGLTFJsonIndexedObjectArray<FGLTFJsonLight, void>;
+
+	FGLTFJsonLight(int32 Index)
 		: IGLTFJsonIndexedObject(Index)
 		, Type(EGLTFJsonLightType::None)
 		, Color(FGLTFJsonColor3::White)
@@ -39,6 +45,4 @@ struct GLTFEXPORTER_API FGLTFJsonLight : IGLTFJsonIndexedObject
 		, Range(0)
 	{
 	}
-
-	virtual void WriteObject(IGLTFJsonWriter& Writer) const override;
 };
diff --git a/Engine/Plugins/Enterprise/GLTFExporter/MigratedFromGit/Source/GLTFExporter/Public/Json/GLTFJsonLightMap.h b/Engine/Plugins/Enterprise/GLTFExporter/MigratedFromGit/Source/GLTFExporter/Public/Json/GLTFJsonLightMap.h
--- a/Engine/Plugins/Enterprise/GLTFExporter/MigratedFromGit/Source/GLTFExporter/Public/Json/GLTFJsonLightMap.h
+++ b/Engine/Plugins/Enterprise/GLTFExporter/MigratedFromGit/Source/GLTFExporter/Public/Json/GLTFJsonLightMap.h
@@ -13,13 +13,17 @@ struct GLTFEXPORTER_API FGLTFJsonLightMap : IGLTFJsonIndexedObject
 	FGLTFJsonVector4     LightMapAdd;
 	FGLTFJsonVector4     CoordinateScaleBias;
 
-	FGLTFJsonLightMap(int32 Index = INDEX_NONE)
+	virtual void WriteObject(IGLTFJsonWriter& Writer) const override;
+
+protected:
+
+	friend TGLTFJsonIndexedObjectArray<FGLTFJsonLightMap, void>;
+
+	FGLTFJsonLightMap(int32 Index)
 		: IGLTFJsonIndexedObject(Index)
 		, LightMapScale(FGLTFJsonVector4::One)
 		, LightMapAdd(FGLTFJsonVector4::Zero)
 		, CoordinateScaleBias({ 1, 1, 0, 0 })
 	{
 	}
-
-	virtual void WriteObject(IGLTFJsonWriter& Writer) const override;
 };
diff --git a/Engine/Plugins/Enterprise/GLTFExporter/MigratedFromGit/Source/GLTFExporter/Public/Json/GLTFJsonMaterial.h b/Engine/Plugins/Enterprise/GLTFExporter/MigratedFromGit/Source/GLTFExporter/Public/Json/GLTFJsonMaterial.h
--- a/Engine/Plugins/Enterprise/GLTFExporter/MigratedFromGit/Source/GLTFExporter/Public/Json/GLTFJsonMaterial.h
+++ b/Engine/Plugins/Enterprise/GLTFExporter/MigratedFromGit/Source/GLTFExporter/Public/Json/GLTFJsonMaterial.h
@@ -106,7 +106,13 @@ struct GLTFEXPORTER_API FGLTFJsonMaterial : IGLTFJsonIndexedObject
 
 	FGLTFJsonClearCoatExtension ClearCoat;
 
-	FGLTFJsonMaterial(int32 Index = INDEX_NONE)
+	virtual void WriteObject(IGLTFJsonWriter& Writer) const override;
+
+protected:
+
+	friend TGLTFJsonIndexedObjectArray<FGLTFJsonMaterial, void>;
+
+	FGLTFJsonMaterial(int32 Index)
 		: IGLTFJsonIndexedObject(Index)
 		, ShadingModel(EGLTFJsonShadingModel::Default)
 		, EmissiveFactor(FGLTFJsonColor3::Black)
@@ -116,6 +122,4 @@ struct GLTFEXPORTER_API FGLTFJsonMaterial : IGLTFJsonIndexedObject
 		, BlendMode(EGLTFJsonBlendMode::None)
 	{
 	}
-
-	virtual void WriteObject(IGLTFJsonWriter& Writer) const override;
 };
diff --git a/Engine/Plugins/Enterprise/GLTFExporter/MigratedFromGit/Source/GLTFExporter/Public/Json/GLTFJsonMesh.h b/Engine/Plugins/Enterprise/GLTFExporter/MigratedFromGit/Source/GLTFExporter/Public/Json/GLTFJsonMesh.h
--- a/Engine/Plugins/Enterprise/GLTFExporter/MigratedFromGit/Source/GLTFExporter/Public/Json/GLTFJsonMesh.h
+++ b/Engine/Plugins/Enterprise/GLTFExporter/MigratedFromGit/Source/GLTFExporter/Public/Json/GLTFJsonMesh.h
@@ -52,10 +52,14 @@ struct GLTFEXPORTER_API FGLTFJsonMesh : IGLTFJsonIndexedObject
 
 	TArray<FGLTFJsonPrimitive> Primitives;
 
-	FGLTFJsonMesh(int32 Index = INDEX_NONE)
+	virtual void WriteObject(IGLTFJsonWriter& Writer) const override;
+
+protected:
+
+	friend TGLTFJsonIndexedObjectArray<FGLTFJsonMesh, void>;
+
+	FGLTFJsonMesh(int32 Index)
 		: IGLTFJsonIndexedObject(Index)
 	{
 	}
-
-	virtual void WriteObject(IGLTFJsonWriter& Writer) const override;
 };
diff --git a/Engine/Plugins/Enterprise/GLTFExporter/MigratedFromGit/Source/GLTFExporter/Public/Json/GLTFJsonNode.h b/Engine/Plugins/Enterprise/GLTFExporter/MigratedFromGit/Source/GLTFExporter/Public/Json/GLTFJsonNode.h
--- a/Engine/Plugins/Enterprise/GLTFExporter/MigratedFromGit/Source/GLTFExporter/Public/Json/GLTFJsonNode.h
+++ b/Engine/Plugins/Enterprise/GLTFExporter/MigratedFromGit/Source/GLTFExporter/Public/Json/GLTFJsonNode.h
@@ -25,7 +25,13 @@ struct GLTFEXPORTER_API FGLTFJsonNode : IGLTFJsonIndexedObject
 
 	TArray<FGLTFJsonNode*> Children;
 
-	FGLTFJsonNode(int32 Index = INDEX_NONE)
+	virtual void WriteObject(IGLTFJsonWriter& Writer) const override;
+
+protected:
+
+	friend TGLTFJsonIndexedObjectArray<FGLTFJsonNode, void>;
+
+	FGLTFJsonNode(int32 Index)
 		: IGLTFJsonIndexedObject(Index)
 		, Translation(FGLTFJsonVector3::Zero)
 		, Rotation(FGLTFJsonQuaternion::Identity)
@@ -41,6 +47,4 @@ struct GLTFEXPORTER_API FGLTFJsonNode : IGLTFJsonIndexedObject
 		, ComponentNode(nullptr)
 	{
 	}
-
-	virtual void WriteObject(IGLTFJsonWriter& Writer) const override;
 };
diff --git a/Engine/Plugins/Enterprise/GLTFExporter/MigratedFromGit/Source/GLTFExporter/Public/Json/GLTFJsonSampler.h b/Engine/Plugins/Enterprise/GLTFExporter/MigratedFromGit/Source/GLTFExporter/Public/Json/GLTFJsonSampler.h
--- a/Engine/Plugins/Enterprise/GLTFExporter/MigratedFromGit/Source/GLTFExporter/Public/Json/GLTFJsonSampler.h
+++ b/Engine/Plugins/Enterprise/GLTFExporter/MigratedFromGit/Source/GLTFExporter/Public/Json/GLTFJsonSampler.h
@@ -14,7 +14,13 @@ struct GLTFEXPORTER_API FGLTFJsonSampler : IGLTFJsonIndexedObject
 	EGLTFJsonTextureWrap WrapS;
 	EGLTFJsonTextureWrap WrapT;
 
-	FGLTFJsonSampler(int32 Index = INDEX_NONE)
+	virtual void WriteObject(IGLTFJsonWriter& Writer) const override;
+
+protected:
+
+	friend TGLTFJsonIndexedObjectArray<FGLTFJsonSampler, void>;
+
+	FGLTFJsonSampler(int32 Index)
 		: IGLTFJsonIndexedObject(Index)
 		, MinFilter(EGLTFJsonTextureFilter::None)
 		, MagFilter(EGLTFJsonTextureFilter::None)
@@ -22,6 +28,4 @@ struct GLTFEXPORTER_API FGLTFJsonSampler : IGLTFJsonIndexedObject
 		, WrapT(EGLTFJsonTextureWrap::Repeat)
 	{
 	}
-
-	virtual void WriteObject(IGLTFJsonWriter& Writer) const override;
 };
diff --git a/Engine/Plugins/Enterprise/GLTFExporter/MigratedFromGit/Source/GLTFExporter/Public/Json/GLTFJsonScene.h b/Engine/Plugins/Enterprise/GLTFExporter/MigratedFromGit/Source/GLTFExporter/Public/Json/GLTFJsonScene.h
--- a/Engine/Plugins/Enterprise/GLTFExporter/MigratedFromGit/Source/GLTFExporter/Public/Json/GLTFJsonScene.h
+++ b/Engine/Plugins/Enterprise/GLTFExporter/MigratedFromGit/Source/GLTFExporter/Public/Json/GLTFJsonScene.h
@@ -11,10 +11,14 @@ struct GLTFEXPORTER_API FGLTFJsonScene : IGLTFJsonIndexedObject
 	TArray<FGLTFJsonNode*> Nodes;
 	TArray<FGLTFJsonEpicLevelVariantSets*>  EpicLevelVariantSets;
 
-	FGLTFJsonScene(int32 Index = INDEX_NONE)
+	virtual void WriteObject(IGLTFJsonWriter& Writer) const override;
+
+protected:
+
+	friend TGLTFJsonIndexedObjectArray<FGLTFJsonScene, void>;
+
+	FGLTFJsonScene(int32 Index)
 		: IGLTFJsonIndexedObject(Index)
 	{
 	}
-
-	virtual void WriteObject(IGLTFJsonWriter& Writer) const override;
 };
diff --git a/Engine/Plugins/Enterprise/GLTFExporter/MigratedFromGit/Source/GLTFExporter/Public/Json/GLTFJsonSkin.h b/Engine/Plugins/Enterprise/GLTFExporter/MigratedFromGit/Source/GLTFExporter/Public/Json/GLTFJsonSkin.h
--- a/Engine/Plugins/Enterprise/GLTFExporter/MigratedFromGit/Source/GLTFExporter/Public/Json/GLTFJsonSkin.h
+++ b/Engine/Plugins/Enterprise/GLTFExporter/MigratedFromGit/Source/GLTFExporter/Public/Json/GLTFJsonSkin.h
@@ -13,12 +13,16 @@ struct GLTFEXPORTER_API FGLTFJsonSkin : IGLTFJsonIndexedObject
 
 	TArray<FGLTFJsonNode*> Joints;
 
-	FGLTFJsonSkin(int32 Index = INDEX_NONE)
+	virtual void WriteObject(IGLTFJsonWriter& Writer) const override;
+
+protected:
+
+	friend TGLTFJsonIndexedObjectArray<FGLTFJsonSkin, void>;
+
+	FGLTFJsonSkin(int32 Index)
 		: IGLTFJsonIndexedObject(Index)
 		, InverseBindMatrices(nullptr)
 		, Skeleton(nullptr)
 	{
 	}
-
-	virtual void WriteObject(IGLTFJsonWriter& Writer) const override;
 };
diff --git a/Engine/Plugins/Enterprise/GLTFExporter/MigratedFromGit/Source/GLTFExporter/Public/Json/GLTFJsonSkySphere.h b/Engine/Plugins/Enterprise/GLTFExporter/MigratedFromGit/Source/GLTFExporter/Public/Json/GLTFJsonSkySphere.h
--- a/Engine/Plugins/Enterprise/GLTFExporter/MigratedFromGit/Source/GLTFExporter/Public/Json/GLTFJsonSkySphere.h
+++ b/Engine/Plugins/Enterprise/GLTFExporter/MigratedFromGit/Source/GLTFExporter/Public/Json/GLTFJsonSkySphere.h
@@ -56,7 +56,13 @@ struct GLTFEXPORTER_API FGLTFJsonSkySphere : IGLTFJsonIndexedObject
 
 	FGLTFJsonVector3 Scale;
 
-	FGLTFJsonSkySphere(int32 Index = INDEX_NONE)
+	virtual void WriteObject(IGLTFJsonWriter& Writer) const override;
+
+protected:
+
+	friend TGLTFJsonIndexedObjectArray<FGLTFJsonSkySphere, void>;
+
+	FGLTFJsonSkySphere(int32 Index)
 		: IGLTFJsonIndexedObject(Index)
 		, SkySphereMesh(nullptr)
 		, SkyTexture(nullptr)
@@ -80,6 +86,4 @@ struct GLTFEXPORTER_API FGLTFJsonSkySphere : IGLTFJsonIndexedObject
 		, Scale(FGLTFJsonVector3::One)
 	{
 	}
-
-	virtual void WriteObject(IGLTFJsonWriter& Writer) const override;
 };
diff --git a/Engine/Plugins/Enterprise/GLTFExporter/MigratedFromGit/Source/GLTFExporter/Public/Json/GLTFJsonTexture.h b/Engine/Plugins/Enterprise/GLTFExporter/MigratedFromGit/Source/GLTFExporter/Public/Json/GLTFJsonTexture.h
--- a/Engine/Plugins/Enterprise/GLTFExporter/MigratedFromGit/Source/GLTFExporter/Public/Json/GLTFJsonTexture.h
+++ b/Engine/Plugins/Enterprise/GLTFExporter/MigratedFromGit/Source/GLTFExporter/Public/Json/GLTFJsonTexture.h
@@ -14,13 +14,17 @@ struct GLTFEXPORTER_API FGLTFJsonTexture : IGLTFJsonIndexedObject
 
 	EGLTFJsonHDREncoding Encoding;
 
-	FGLTFJsonTexture(int32 Index = INDEX_NONE)
+	virtual void WriteObject(IGLTFJsonWriter& Writer) const override;
+
+protected:
+
+	friend TGLTFJsonIndexedObjectArray<FGLTFJsonTexture, void>;
+
+	FGLTFJsonTexture(int32 Index)
 		: IGLTFJsonIndexedObject(Index)
 		, Sampler(nullptr)
 		, Source(nullptr)
 		, Encoding(EGLTFJsonHDREncoding::None)
 	{
 	}
-
-	virtual void WriteObject(IGLTFJsonWriter& Writer) const override;
 };
