[GLTFExporter] refactor: remove json data getters in builder

commit 25b923510020749af11699d7ac793a646d040bc1
Date: Sun, 13 Mar 2022 18:37:08 +0100

    refactor: remove json data getters in builder
    (cherry picked from commit d86e973aa5693ca10975aa67b03981661653e972​)

#skipci

diff --git a/Engine/Plugins/Enterprise/GLTFExporter/MigratedFromGit/Source/GLTFExporter/Private/Builders/GLTFJsonBuilder.cpp b/Engine/Plugins/Enterprise/GLTFExporter/MigratedFromGit/Source/GLTFExporter/Private/Builders/GLTFJsonBuilder.cpp
--- a/Engine/Plugins/Enterprise/GLTFExporter/MigratedFromGit/Source/GLTFExporter/Private/Builders/GLTFJsonBuilder.cpp
+++ b/Engine/Plugins/Enterprise/GLTFExporter/MigratedFromGit/Source/GLTFExporter/Private/Builders/GLTFJsonBuilder.cpp
@@ -176,106 +176,6 @@ const FGLTFJsonRoot& FGLTFJsonBuilder::GetRoot() const
 	return JsonRoot;
 }
 
-FGLTFJsonAccessor& FGLTFJsonBuilder::GetAccessor(FGLTFJsonAccessor* AccessorIndex)
-{
-	return *AccessorIndex;
-}
-
-FGLTFJsonAnimation& FGLTFJsonBuilder::GetAnimation(FGLTFJsonAnimation* AnimationIndex)
-{
-	return *AnimationIndex;
-}
-
-FGLTFJsonBuffer& FGLTFJsonBuilder::GetBuffer(FGLTFJsonBuffer* BufferIndex)
-{
-	return *BufferIndex;
-}
-
-FGLTFJsonBufferView& FGLTFJsonBuilder::GetBufferView(FGLTFJsonBufferView* BufferViewIndex)
-{
-	return *BufferViewIndex;
-}
-
-FGLTFJsonCamera& FGLTFJsonBuilder::GetCamera(FGLTFJsonCamera* CameraIndex)
-{
-	return *CameraIndex;
-}
-
-FGLTFJsonImage& FGLTFJsonBuilder::GetImage(FGLTFJsonImage* ImageIndex)
-{
-	return *ImageIndex;
-}
-
-FGLTFJsonMaterial& FGLTFJsonBuilder::GetMaterial(FGLTFJsonMaterial* MaterialIndex)
-{
-	return *MaterialIndex;
-}
-
-FGLTFJsonMesh& FGLTFJsonBuilder::GetMesh(FGLTFJsonMesh* MeshIndex)
-{
-	return *MeshIndex;
-}
-
-FGLTFJsonNode& FGLTFJsonBuilder::GetNode(FGLTFJsonNode* NodeIndex)
-{
-	return *NodeIndex;
-}
-
-FGLTFJsonSampler& FGLTFJsonBuilder::GetSampler(FGLTFJsonSampler* SamplerIndex)
-{
-	return *SamplerIndex;
-}
-
-FGLTFJsonScene& FGLTFJsonBuilder::GetScene(FGLTFJsonScene* SceneIndex)
-{
-	return *SceneIndex;
-}
-
-FGLTFJsonSkin& FGLTFJsonBuilder::GetSkin(FGLTFJsonSkin* SkinIndex)
-{
-	return *SkinIndex;
-}
-
-FGLTFJsonTexture& FGLTFJsonBuilder::GetTexture(FGLTFJsonTexture* TextureIndex)
-{
-	return *TextureIndex;
-}
-
-FGLTFJsonBackdrop& FGLTFJsonBuilder::GetBackdrop(FGLTFJsonBackdrop* BackdropIndex)
-{
-	return *BackdropIndex;
-}
-
-FGLTFJsonHotspot& FGLTFJsonBuilder::GetHotspot(FGLTFJsonHotspot* HotspotIndex)
-{
-	return *HotspotIndex;
-}
-
-FGLTFJsonLight& FGLTFJsonBuilder::GetLight(FGLTFJsonLight* LightIndex)
-{
-	return *LightIndex;
-}
-
-FGLTFJsonLightMap& FGLTFJsonBuilder::GetLightMap(FGLTFJsonLightMap* LightMapIndex)
-{
-	return *LightMapIndex;
-}
-
-FGLTFJsonSkySphere& FGLTFJsonBuilder::GetSkySphere(FGLTFJsonSkySphere* SkySphereIndex)
-{
-	return *SkySphereIndex;
-}
-
-FGLTFJsonEpicLevelVariantSets& FGLTFJsonBuilder::GetEpicLevelVariantSets(FGLTFJsonEpicLevelVariantSets* EpicLevelVariantSetsIndex)
-{
-	return *EpicLevelVariantSetsIndex;
-}
-
-FGLTFJsonKhrMaterialVariant& FGLTFJsonBuilder::GetKhrMaterialVariant(FGLTFJsonKhrMaterialVariant* KhrMaterialVariantIndex)
-{
-	return *KhrMaterialVariantIndex;
-}
-
 FGLTFJsonNode* FGLTFJsonBuilder::GetComponentNode(FGLTFJsonNode* Node)
 {
 	if (Node == nullptr)
diff --git a/Engine/Plugins/Enterprise/GLTFExporter/MigratedFromGit/Source/GLTFExporter/Private/Builders/GLTFJsonBuilder.h b/Engine/Plugins/Enterprise/GLTFExporter/MigratedFromGit/Source/GLTFExporter/Private/Builders/GLTFJsonBuilder.h
--- a/Engine/Plugins/Enterprise/GLTFExporter/MigratedFromGit/Source/GLTFExporter/Private/Builders/GLTFJsonBuilder.h
+++ b/Engine/Plugins/Enterprise/GLTFExporter/MigratedFromGit/Source/GLTFExporter/Private/Builders/GLTFJsonBuilder.h
@@ -44,26 +44,6 @@ public:
 	FGLTFJsonNode* AddChildComponentNode(FGLTFJsonNode* ParentNodeParentNode, FGLTFJsonNode* ChildNode = nullptr);
 
 	const FGLTFJsonRoot& GetRoot() const;
-	FGLTFJsonAccessor& GetAccessor(FGLTFJsonAccessor* AccessorIndex);
-	FGLTFJsonAnimation& GetAnimation(FGLTFJsonAnimation* AnimationIndex);
-	FGLTFJsonBuffer& GetBuffer(FGLTFJsonBuffer* BufferIndex);
-	FGLTFJsonBufferView& GetBufferView(FGLTFJsonBufferView* BufferViewIndex);
-	FGLTFJsonCamera& GetCamera(FGLTFJsonCamera* CameraIndex);
-	FGLTFJsonImage& GetImage(FGLTFJsonImage* ImageIndex);
-	FGLTFJsonMaterial& GetMaterial(FGLTFJsonMaterial* MaterialIndex);
-	FGLTFJsonMesh& GetMesh(FGLTFJsonMesh* MeshIndex);
-	FGLTFJsonNode& GetNode(FGLTFJsonNode* NodeIndex);
-	FGLTFJsonSampler& GetSampler(FGLTFJsonSampler* SamplerIndex);
-	FGLTFJsonScene& GetScene(FGLTFJsonScene* SceneIndex);
-	FGLTFJsonSkin& GetSkin(FGLTFJsonSkin* SkinIndex);
-	FGLTFJsonTexture& GetTexture(FGLTFJsonTexture* TextureIndex);
-	FGLTFJsonBackdrop& GetBackdrop(FGLTFJsonBackdrop* BackdropIndex);
-	FGLTFJsonHotspot& GetHotspot(FGLTFJsonHotspot* HotspotIndex);
-	FGLTFJsonLight& GetLight(FGLTFJsonLight* LightIndex);
-	FGLTFJsonLightMap& GetLightMap(FGLTFJsonLightMap* LightMapIndex);
-	FGLTFJsonSkySphere& GetSkySphere(FGLTFJsonSkySphere* SkySphereIndex);
-	FGLTFJsonEpicLevelVariantSets& GetEpicLevelVariantSets(FGLTFJsonEpicLevelVariantSets* EpicLevelVariantSetsIndex);
-	FGLTFJsonKhrMaterialVariant& GetKhrMaterialVariant(FGLTFJsonKhrMaterialVariant* KhrMaterialVariantIndex);
 
 	FGLTFJsonNode* GetComponentNode(FGLTFJsonNode* Node);
 
diff --git a/Engine/Plugins/Enterprise/GLTFExporter/MigratedFromGit/Source/GLTFExporter/Private/Converters/GLTFAccessorConverters.cpp b/Engine/Plugins/Enterprise/GLTFExporter/MigratedFromGit/Source/GLTFExporter/Private/Converters/GLTFAccessorConverters.cpp
--- a/Engine/Plugins/Enterprise/GLTFExporter/MigratedFromGit/Source/GLTFExporter/Private/Converters/GLTFAccessorConverters.cpp
+++ b/Engine/Plugins/Enterprise/GLTFExporter/MigratedFromGit/Source/GLTFExporter/Private/Converters/GLTFAccessorConverters.cpp
@@ -137,7 +137,7 @@ FGLTFJsonAccessor* FGLTFNormalBufferConverter::Convert(const FGLTFMeshSection* M
 		return nullptr;
 	}
 
-	FGLTFJsonBufferView* BufferViewIndex;
+	FGLTFJsonBufferView* BufferView;
 	EGLTFJsonComponentType ComponentType;
 
 	const bool bMeshQuantization = Builder.ExportOptions->bUseMeshQuantization;
@@ -150,26 +150,26 @@ FGLTFJsonAccessor* FGLTFNormalBufferConverter::Convert(const FGLTFMeshSection* M
 		if (bHighPrecision)
 		{
 			ComponentType = EGLTFJsonComponentType::Int16;
-			BufferViewIndex = ConvertBufferView<FGLTFInt16Vector4, FPackedRGBA16N>(MeshSection, SourceData);
-			Builder.GetBufferView(BufferViewIndex).ByteStride = sizeof(FGLTFInt16Vector4);
+			BufferView = ConvertBufferView<FGLTFInt16Vector4, FPackedRGBA16N>(MeshSection, SourceData);
+			BufferView->ByteStride = sizeof(FGLTFInt16Vector4);
 		}
 		else
 		{
 			ComponentType = EGLTFJsonComponentType::Int8;
-			BufferViewIndex = ConvertBufferView<FGLTFInt8Vector4, FPackedNormal>(MeshSection, SourceData);
-			Builder.GetBufferView(BufferViewIndex).ByteStride = sizeof(FGLTFInt8Vector4);
+			BufferView = ConvertBufferView<FGLTFInt8Vector4, FPackedNormal>(MeshSection, SourceData);
+			BufferView->ByteStride = sizeof(FGLTFInt8Vector4);
 		}
 	}
 	else
 	{
 		ComponentType = EGLTFJsonComponentType::Float;
-		BufferViewIndex = bHighPrecision
+		BufferView = bHighPrecision
 			? ConvertBufferView<FGLTFVector3, FPackedRGBA16N>(MeshSection, SourceData)
 			: ConvertBufferView<FGLTFVector3, FPackedNormal>(MeshSection, SourceData);
 	}
 
 	FGLTFJsonAccessor* JsonAccessor = Builder.AddAccessor();
-	JsonAccessor->BufferView = BufferViewIndex;
+	JsonAccessor->BufferView = BufferView;
 	JsonAccessor->ComponentType = ComponentType;
 	JsonAccessor->Count = MeshSection->IndexMap.Num();
 	JsonAccessor->Type = EGLTFJsonAccessorType::Vec3;
@@ -218,7 +218,7 @@ FGLTFJsonAccessor* FGLTFTangentBufferConverter::Convert(const FGLTFMeshSection*
 		return nullptr;
 	}
 
-	FGLTFJsonBufferView* BufferViewIndex;
+	FGLTFJsonBufferView* BufferView;
 	EGLTFJsonComponentType ComponentType;
 
 	const bool bMeshQuantization = Builder.ExportOptions->bUseMeshQuantization;
@@ -231,24 +231,24 @@ FGLTFJsonAccessor* FGLTFTangentBufferConverter::Convert(const FGLTFMeshSection*
 		if (bHighPrecision)
 		{
 			ComponentType = EGLTFJsonComponentType::Int16;
-			BufferViewIndex = ConvertBufferView<FGLTFInt16Vector4, FPackedRGBA16N>(MeshSection, SourceData);
+			BufferView = ConvertBufferView<FGLTFInt16Vector4, FPackedRGBA16N>(MeshSection, SourceData);
 		}
 		else
 		{
 			ComponentType = EGLTFJsonComponentType::Int8;
-			BufferViewIndex = ConvertBufferView<FGLTFInt8Vector4, FPackedNormal>(MeshSection, SourceData);
+			BufferView = ConvertBufferView<FGLTFInt8Vector4, FPackedNormal>(MeshSection, SourceData);
 		}
 	}
 	else
 	{
 		ComponentType = EGLTFJsonComponentType::Float;
-		BufferViewIndex = bHighPrecision
+		BufferView = bHighPrecision
 			? ConvertBufferView<FGLTFVector4, FPackedRGBA16N>(MeshSection, SourceData)
 			: ConvertBufferView<FGLTFVector4, FPackedNormal>(MeshSection, SourceData);
 	}
 
 	FGLTFJsonAccessor* JsonAccessor = Builder.AddAccessor();
-	JsonAccessor->BufferView = BufferViewIndex;
+	JsonAccessor->BufferView = BufferView;
 	JsonAccessor->ComponentType = ComponentType;
 	JsonAccessor->Count = MeshSection->IndexMap.Num();
 	JsonAccessor->Type = EGLTFJsonAccessorType::Vec4;
diff --git a/Engine/Plugins/Enterprise/GLTFExporter/MigratedFromGit/Source/GLTFExporter/Private/Converters/GLTFEpicLevelVariantSetsConverters.cpp b/Engine/Plugins/Enterprise/GLTFExporter/MigratedFromGit/Source/GLTFExporter/Private/Converters/GLTFEpicLevelVariantSetsConverters.cpp
--- a/Engine/Plugins/Enterprise/GLTFExporter/MigratedFromGit/Source/GLTFExporter/Private/Converters/GLTFEpicLevelVariantSetsConverters.cpp
+++ b/Engine/Plugins/Enterprise/GLTFExporter/MigratedFromGit/Source/GLTFExporter/Private/Converters/GLTFEpicLevelVariantSetsConverters.cpp
@@ -157,12 +157,12 @@ bool FGLTFEpicLevelVariantSetsConverter::TryParseVisibilityPropertyValue(FGLTFJs
 
 	Builder.RegisterObjectVariant(Owner, Property); // TODO: we should register this on the component
 
-	FGLTFJsonNode* NodeIndex = Builder.GetOrAddNode(Target);
-	FGLTFJsonNode* ComponentNodeIndex = Builder.GetComponentNode(NodeIndex);
+	FGLTFJsonNode* Node = Builder.GetOrAddNode(Target);
+	FGLTFJsonNode* ComponentNode = Builder.GetComponentNode(Node);
 
-	FGLTFJsonEpicVariantNodeProperties& NodeProperties = OutVariant.Nodes.FindOrAdd(ComponentNodeIndex);
+	FGLTFJsonEpicVariantNodeProperties& NodeProperties = OutVariant.Nodes.FindOrAdd(ComponentNode);
 
-	NodeProperties.Node = ComponentNodeIndex;
+	NodeProperties.Node = ComponentNode;
 	NodeProperties.bIsVisible = bIsVisible;
 	return true;
 }
@@ -220,11 +220,11 @@ bool FGLTFEpicLevelVariantSetsConverter::TryParseMaterialPropertyValue(FGLTFJson
 	VariantMaterial.Material = FGLTFVariantUtility::GetOrAddMaterial(Builder, Material, Target, MaterialIndex);
 	VariantMaterial.Index = MaterialIndex;
 
-	FGLTFJsonNode* NodeIndex = Builder.GetOrAddNode(Target);
-	FGLTFJsonNode* ComponentNodeIndex = Builder.GetComponentNode(NodeIndex);
-	FGLTFJsonEpicVariantNodeProperties& NodeProperties = OutVariant.Nodes.FindOrAdd(ComponentNodeIndex);
+	FGLTFJsonNode* Node = Builder.GetOrAddNode(Target);
+	FGLTFJsonNode* ComponentNode = Builder.GetComponentNode(Node);
+	FGLTFJsonEpicVariantNodeProperties& NodeProperties = OutVariant.Nodes.FindOrAdd(ComponentNode);
 
-	NodeProperties.Node = ComponentNodeIndex;
+	NodeProperties.Node = ComponentNode;
 	NodeProperties.Materials.Add(VariantMaterial);
 	return true;
 }
@@ -271,12 +271,12 @@ bool FGLTFEpicLevelVariantSetsConverter::TryParseMeshPropertyValue(FGLTFJsonEpic
 	Builder.RegisterObjectVariant(Owner, Property); // TODO: we should register this on the component
 
 	const FGLTFMaterialArray OverrideMaterials(Target->OverrideMaterials);
-	FGLTFJsonNode* NodeIndex = Builder.GetOrAddNode(Target);
-	FGLTFJsonMesh* MeshIndex = Builder.GetOrAddMesh(Mesh, OverrideMaterials);
-	FGLTFJsonNode* ComponentNodeIndex = Builder.GetComponentNode(NodeIndex);
-	FGLTFJsonEpicVariantNodeProperties& NodeProperties = OutVariant.Nodes.FindOrAdd(ComponentNodeIndex);
+	FGLTFJsonNode* JsonNode = Builder.GetOrAddNode(Target);
+	FGLTFJsonNode* JsonComponentNode = Builder.GetComponentNode(JsonNode);
+	FGLTFJsonEpicVariantNodeProperties& NodeProperties = OutVariant.Nodes.FindOrAdd(JsonComponentNode);
+	FGLTFJsonMesh* JsonMesh = Builder.GetOrAddMesh(Mesh, OverrideMaterials);
 
-	NodeProperties.Node = ComponentNodeIndex;
-	NodeProperties.Mesh = MeshIndex;
+	NodeProperties.Node = JsonComponentNode;
+	NodeProperties.Mesh = JsonMesh;
 	return true;
 }
diff --git a/Engine/Plugins/Enterprise/GLTFExporter/MigratedFromGit/Source/GLTFExporter/Private/Converters/GLTFKhrVariantConverters.cpp b/Engine/Plugins/Enterprise/GLTFExporter/MigratedFromGit/Source/GLTFExporter/Private/Converters/GLTFKhrVariantConverters.cpp
--- a/Engine/Plugins/Enterprise/GLTFExporter/MigratedFromGit/Source/GLTFExporter/Private/Converters/GLTFKhrVariantConverters.cpp
+++ b/Engine/Plugins/Enterprise/GLTFExporter/MigratedFromGit/Source/GLTFExporter/Private/Converters/GLTFKhrVariantConverters.cpp
@@ -38,11 +38,11 @@ FGLTFJsonKhrMaterialVariant* FGLTFKhrMaterialVariantConverter::Convert(const UVa
 			if (const UPropertyValueMaterial* MaterialProperty = Cast<UPropertyValueMaterial>(Property))
 			{
 				FGLTFJsonPrimitive* Primitive = nullptr;
-				FGLTFJsonMaterial* MaterialIndex;
+				FGLTFJsonMaterial* JsonMaterial;
 
-				if (TryParseMaterialProperty(Primitive, MaterialIndex, MaterialProperty))
+				if (TryParseMaterialProperty(Primitive, JsonMaterial, MaterialProperty))
 				{
-					PrimitiveMaterials.Add(MakeTuple(Primitive, MaterialIndex));
+					PrimitiveMaterials.Add(MakeTuple(Primitive, JsonMaterial));
 				}
 			}
 		}
@@ -58,12 +58,12 @@ FGLTFJsonKhrMaterialVariant* FGLTFKhrMaterialVariantConverter::Convert(const UVa
 	for (const TPrimitiveMaterial& PrimitiveMaterial: PrimitiveMaterials)
 	{
 		FGLTFJsonPrimitive* Primitive = PrimitiveMaterial.Key;
-		FGLTFJsonMaterial* MaterialIndex = PrimitiveMaterial.Value;
+		FGLTFJsonMaterial* Material = PrimitiveMaterial.Value;
 
 		FGLTFJsonKhrMaterialVariantMapping* ExistingMapping = Primitive->KhrMaterialVariantMappings.FindByPredicate(
-			[MaterialIndex](const FGLTFJsonKhrMaterialVariantMapping& Mapping)
+			[Material](const FGLTFJsonKhrMaterialVariantMapping& Mapping)
 			{
-				return Mapping.Material == MaterialIndex;
+				return Mapping.Material == Material;
 			});
 
 		if (ExistingMapping != nullptr)
@@ -73,7 +73,7 @@ FGLTFJsonKhrMaterialVariant* FGLTFKhrMaterialVariantConverter::Convert(const UVa
 		else
 		{
 			FGLTFJsonKhrMaterialVariantMapping Mapping;
-			Mapping.Material = MaterialIndex;
+			Mapping.Material = Material;
 			Mapping.Variants.Add(MaterialVariant);
 
 			Primitive->KhrMaterialVariantMappings.Add(Mapping);
@@ -83,7 +83,7 @@ FGLTFJsonKhrMaterialVariant* FGLTFKhrMaterialVariantConverter::Convert(const UVa
 	return MaterialVariant;
 }
 
-bool FGLTFKhrMaterialVariantConverter::TryParseMaterialProperty(FGLTFJsonPrimitive*& OutPrimitive, FGLTFJsonMaterial*& OutMaterialIndex, const UPropertyValueMaterial* Property) const
+bool FGLTFKhrMaterialVariantConverter::TryParseMaterialProperty(FGLTFJsonPrimitive*& OutPrimitive, FGLTFJsonMaterial*& OutMaterial, const UPropertyValueMaterial* Property) const
 {
 	const UMeshComponent* Target = static_cast<UMeshComponent*>(Property->GetPropertyParentContainerAddress());
 	if (Target == nullptr)
@@ -141,10 +141,10 @@ bool FGLTFKhrMaterialVariantConverter::TryParseMaterialProperty(FGLTFJsonPrimiti
 	Builder.RegisterObjectVariant(Owner, Property); // TODO: we don't need to register this on the actor
 
 	const int32 MaterialIndex = CapturedPropSegments[NumPropSegments - 1].PropertyIndex;
-	FGLTFJsonMesh* MeshIndex = Builder.GetOrAddMesh(Target);
+	FGLTFJsonMesh* JsonMesh = Builder.GetOrAddMesh(Target);
 
-	OutPrimitive = &Builder.GetMesh(MeshIndex).Primitives[MaterialIndex];
-	OutMaterialIndex = FGLTFVariantUtility::GetOrAddMaterial(Builder, Material, Target, MaterialIndex);
+	OutPrimitive = &JsonMesh->Primitives[MaterialIndex];
+	OutMaterial = FGLTFVariantUtility::GetOrAddMaterial(Builder, Material, Target, MaterialIndex);
 
 	return true;
 }
diff --git a/Engine/Plugins/Enterprise/GLTFExporter/MigratedFromGit/Source/GLTFExporter/Private/Converters/GLTFKhrVariantConverters.h b/Engine/Plugins/Enterprise/GLTFExporter/MigratedFromGit/Source/GLTFExporter/Private/Converters/GLTFKhrVariantConverters.h
--- a/Engine/Plugins/Enterprise/GLTFExporter/MigratedFromGit/Source/GLTFExporter/Private/Converters/GLTFKhrVariantConverters.h
+++ b/Engine/Plugins/Enterprise/GLTFExporter/MigratedFromGit/Source/GLTFExporter/Private/Converters/GLTFKhrVariantConverters.h
@@ -18,5 +18,5 @@ class FGLTFKhrMaterialVariantConverter final : public FGLTFBuilderContext, publi
 
 	virtual FGLTFJsonKhrMaterialVariant* Convert(const UVariant* Variant) override;
 
-	bool TryParseMaterialProperty(FGLTFJsonPrimitive*& OutPrimitive, FGLTFJsonMaterial*& OutMaterialIndex, const UPropertyValueMaterial* Property) const;
+	bool TryParseMaterialProperty(FGLTFJsonPrimitive*& OutPrimitive, FGLTFJsonMaterial*& OutMaterial, const UPropertyValueMaterial* Property) const;
 };
diff --git a/Engine/Plugins/Enterprise/GLTFExporter/MigratedFromGit/Source/GLTFExporter/Private/Converters/GLTFLightMapConverters.cpp b/Engine/Plugins/Enterprise/GLTFExporter/MigratedFromGit/Source/GLTFExporter/Private/Converters/GLTFLightMapConverters.cpp
--- a/Engine/Plugins/Enterprise/GLTFExporter/MigratedFromGit/Source/GLTFExporter/Private/Converters/GLTFLightMapConverters.cpp
+++ b/Engine/Plugins/Enterprise/GLTFExporter/MigratedFromGit/Source/GLTFExporter/Private/Converters/GLTFLightMapConverters.cpp
@@ -59,9 +59,9 @@ FGLTFJsonLightMap* FGLTFLightMapConverter::Convert(const UStaticMeshComponent* S
 
 	const FLightMapInteraction LightMapInteraction = LightMap2D->GetInteraction(GMaxRHIFeatureLevel);
 	const ULightMapTexture2D* Texture = LightMapInteraction.GetTexture(true);
-	FGLTFJsonTexture* TextureIndex = Builder.GetOrAddTexture(Texture);
+	FGLTFJsonTexture* JsonTexture = Builder.GetOrAddTexture(Texture);
 
-	if (TextureIndex == nullptr)
+	if (JsonTexture == nullptr)
 	{
 		return nullptr;
 	}
@@ -73,7 +73,7 @@ FGLTFJsonLightMap* FGLTFLightMapConverter::Convert(const UStaticMeshComponent* S
 
 	FGLTFJsonLightMap* JsonLightMap = Builder.AddLightMap();
 	StaticMeshComponent->GetName(JsonLightMap->Name); // TODO: use better name (similar to light and camera)
-	JsonLightMap->Texture.Index = TextureIndex;
+	JsonLightMap->Texture.Index = JsonTexture;
 	JsonLightMap->Texture.TexCoord = CoordinateIndex;
 	JsonLightMap->LightMapScale = { LightMapScale.X, LightMapScale.Y, LightMapScale.Z, LightMapScale.W };
 	JsonLightMap->LightMapAdd = { LightMapAdd.X, LightMapAdd.Y, LightMapAdd.Z, LightMapAdd.W };
diff --git a/Engine/Plugins/Enterprise/GLTFExporter/MigratedFromGit/Source/GLTFExporter/Private/Converters/GLTFMaterialProxyFactory.cpp b/Engine/Plugins/Enterprise/GLTFExporter/MigratedFromGit/Source/GLTFExporter/Private/Converters/GLTFMaterialProxyFactory.cpp
--- a/Engine/Plugins/Enterprise/GLTFExporter/MigratedFromGit/Source/GLTFExporter/Private/Converters/GLTFMaterialProxyFactory.cpp
+++ b/Engine/Plugins/Enterprise/GLTFExporter/MigratedFromGit/Source/GLTFExporter/Private/Converters/GLTFMaterialProxyFactory.cpp
@@ -32,8 +32,8 @@ UMaterialInterface* FGLTFMaterialProxyFactory::Create(UMaterialInterface* Origin
 			*OriginalMaterial->GetName()));
 	}
 
-	FGLTFJsonMaterial* MaterialIndex = Builder.GetOrAddMaterial(OriginalMaterial);
-	if (MaterialIndex == nullptr)
+	FGLTFJsonMaterial* JsonMaterial = Builder.GetOrAddMaterial(OriginalMaterial);
+	if (JsonMaterial == nullptr)
 	{
 		// TODO: report error
 		return nullptr;
@@ -42,14 +42,12 @@ UMaterialInterface* FGLTFMaterialProxyFactory::Create(UMaterialInterface* Origin
 	MakeDirectory(RootPath);
 	Builder.CompleteAllTasks();
 
-	const FGLTFJsonMaterial& JsonMaterial = Builder.GetMaterial(MaterialIndex);
-	UMaterialInstanceConstant* ProxyMaterial = CreateInstancedMaterial(OriginalMaterial, JsonMaterial.ShadingModel);
-
+	UMaterialInstanceConstant* ProxyMaterial = CreateInstancedMaterial(OriginalMaterial, JsonMaterial->ShadingModel);
 	if (ProxyMaterial != nullptr)
 	{
 		SetUserData(ProxyMaterial, OriginalMaterial);
 		SetBaseProperties(ProxyMaterial, OriginalMaterial);
-		SetProxyProperties(ProxyMaterial, JsonMaterial);
+		SetProxyProperties(ProxyMaterial, *JsonMaterial);
 	}
 
 	return ProxyMaterial;
@@ -166,32 +164,31 @@ void FGLTFMaterialProxyFactory::SetProxyParameter(UMaterialInstanceConstant* Pro
 	ParameterInfo.UVRotation.Set(ProxyMaterial, TextureInfo.Transform.Rotation, true);
 }
 
-UTexture2D* FGLTFMaterialProxyFactory::FindOrCreateTexture(FGLTFJsonTexture* Index, const FGLTFProxyMaterialTextureParameterInfo& ParameterInfo)
+UTexture2D* FGLTFMaterialProxyFactory::FindOrCreateTexture(FGLTFJsonTexture* JsonTexture, const FGLTFProxyMaterialTextureParameterInfo& ParameterInfo)
 {
-	if (Index == nullptr)
+	if (JsonTexture == nullptr)
 	{
 		return nullptr;
 	}
 
 	// TODO: fix potential conflict when same texture used for different material properties that have different encoding (sRGB vs Linear, Normalmap compression etc)
 
-	UTexture2D** FoundPtr = Textures.Find(Index);
+	UTexture2D** FoundPtr = Textures.Find(JsonTexture);
 	if (FoundPtr != nullptr)
 	{
 		return *FoundPtr;
 	}
 
-	const FGLTFJsonTexture& JsonTexture = Builder.GetTexture(Index);
-	const FGLTFImageData* ImageData = Images.Find(JsonTexture.Source);
+	const FGLTFImageData* ImageData = Images.Find(JsonTexture->Source);
 	if (ImageData == nullptr)
 	{
 		// TODO: report error
 		return nullptr;
 	}
 
-	const FGLTFJsonSampler& JsonSampler = Builder.GetSampler(JsonTexture.Sampler);
+	const FGLTFJsonSampler& JsonSampler = *JsonTexture->Sampler;
 	UTexture2D* Texture = CreateTexture(ImageData, JsonSampler, ParameterInfo);
-	Textures.Add(Index, Texture);
+	Textures.Add(JsonTexture, Texture);
 	return Texture;
 }
 
@@ -257,9 +254,9 @@ TUniquePtr<IGLTFTexture2DConverter> FGLTFMaterialProxyFactory::CreateTextureConv
 
 		virtual FGLTFJsonTexture* Convert(const UTexture2D* Texture2D, bool bToSRGB) override
 		{
-			FGLTFJsonTexture* TextureIndex = Factory.Builder.AddTexture();
-			Factory.Textures.Add(TextureIndex, const_cast<UTexture2D*>(Texture2D));
-			return TextureIndex;
+			FGLTFJsonTexture* Texture = Factory.Builder.AddTexture();
+			Factory.Textures.Add(Texture, const_cast<UTexture2D*>(Texture2D));
+			return Texture;
 		}
 	};
 
@@ -287,9 +284,9 @@ TUniquePtr<IGLTFImageConverter> FGLTFMaterialProxyFactory::CreateImageConverter(
 			const FString Filename = FGLTFImageUtility::GetUniqueFilename(Name, TEXT(""), UniqueFilenames);
 			UniqueFilenames.Add(Filename);
 
-			FGLTFJsonImage* ImageIndex = Factory.Builder.AddImage();
-			Factory.Images.Add(ImageIndex, { Filename, Type, bIgnoreAlpha, Size, Pixels });
-			return ImageIndex;
+			FGLTFJsonImage* Image = Factory.Builder.AddImage();
+			Factory.Images.Add(Image, { Filename, Type, bIgnoreAlpha, Size, Pixels });
+			return Image;
 		}
 	};
 
diff --git a/Engine/Plugins/Enterprise/GLTFExporter/MigratedFromGit/Source/GLTFExporter/Private/Converters/GLTFMaterialProxyFactory.h b/Engine/Plugins/Enterprise/GLTFExporter/MigratedFromGit/Source/GLTFExporter/Private/Converters/GLTFMaterialProxyFactory.h
--- a/Engine/Plugins/Enterprise/GLTFExporter/MigratedFromGit/Source/GLTFExporter/Private/Converters/GLTFMaterialProxyFactory.h
+++ b/Engine/Plugins/Enterprise/GLTFExporter/MigratedFromGit/Source/GLTFExporter/Private/Converters/GLTFMaterialProxyFactory.h
@@ -42,7 +42,7 @@ private:
 	void SetProxyParameter(UMaterialInstanceConstant* ProxyMaterial, const TGLTFProxyMaterialParameterInfo<FLinearColor>& ParameterInfo, const FGLTFJsonColor4& Color);
 	void SetProxyParameter(UMaterialInstanceConstant* ProxyMaterial, const FGLTFProxyMaterialTextureParameterInfo& ParameterInfo, const FGLTFJsonTextureInfo& TextureInfo);
 
-	UTexture2D* FindOrCreateTexture(FGLTFJsonTexture* Index, const FGLTFProxyMaterialTextureParameterInfo& ParameterInfo);
+	UTexture2D* FindOrCreateTexture(FGLTFJsonTexture* JsonTexture, const FGLTFProxyMaterialTextureParameterInfo& ParameterInfo);
 	UTexture2D* CreateTexture(const FGLTFImageData* ImageData, const FGLTFJsonSampler& JsonSampler, const FGLTFProxyMaterialTextureParameterInfo& ParameterInfo);
 
 	UMaterialInstanceConstant* CreateInstancedMaterial(UMaterialInterface* OriginalMaterial, EGLTFJsonShadingModel ShadingModel);
diff --git a/Engine/Plugins/Enterprise/GLTFExporter/MigratedFromGit/Source/GLTFExporter/Private/Converters/GLTFNodeConverters.cpp b/Engine/Plugins/Enterprise/GLTFExporter/MigratedFromGit/Source/GLTFExporter/Private/Converters/GLTFNodeConverters.cpp
--- a/Engine/Plugins/Enterprise/GLTFExporter/MigratedFromGit/Source/GLTFExporter/Private/Converters/GLTFNodeConverters.cpp
+++ b/Engine/Plugins/Enterprise/GLTFExporter/MigratedFromGit/Source/GLTFExporter/Private/Converters/GLTFNodeConverters.cpp
@@ -101,7 +101,7 @@ FGLTFJsonNode* FGLTFComponentConverter::Convert(const USceneComponent* SceneComp
 
 	const USceneComponent* ParentComponent = !bIsRootNode ? SceneComponent->GetAttachParent() : nullptr;
 	const FName SocketName = SceneComponent->GetAttachSocketName();
-	FGLTFJsonNode* ParentNodeIndex = Builder.GetOrAddNode(ParentComponent, SocketName);
+	FGLTFJsonNode* ParentNode = Builder.GetOrAddNode(ParentComponent, SocketName);
 
 	if (ParentComponent != nullptr && !SceneComponent->IsUsingAbsoluteScale())
 	{
@@ -120,7 +120,7 @@ FGLTFJsonNode* FGLTFComponentConverter::Convert(const USceneComponent* SceneComp
 	const FTransform ParentTransform = ParentComponent != nullptr ? ParentComponent->GetSocketTransform(SocketName) : FTransform::Identity;
 	const FTransform RelativeTransform = bIsRootNode ? Transform : Transform.GetRelativeTransform(ParentTransform);
 
-	FGLTFJsonNode* Node = Builder.AddChildNode(ParentNodeIndex);
+	FGLTFJsonNode* Node = Builder.AddChildNode(ParentNode);
 	Node->Name = FGLTFNameUtility::GetName(SceneComponent);
 	Node->Translation = FGLTFConverterUtility::ConvertPosition(RelativeTransform.GetTranslation(), Builder.ExportOptions->ExportUniformScale);
 	Node->Rotation = FGLTFConverterUtility::ConvertRotation(RelativeTransform.GetRotation());
@@ -146,7 +146,6 @@ FGLTFJsonNode* FGLTFComponentConverter::Convert(const USceneComponent* SceneComp
 
 			if (Builder.ExportOptions->bExportVertexSkinWeights)
 			{
-				// TODO: remove need for NodeIndex by adding support for cyclic calls in converter
 				Node->Skin = Builder.GetOrAddSkin(Node, SkeletalMeshComponent);
 				if (Node->Skin != nullptr)
 				{
diff --git a/Engine/Plugins/Enterprise/GLTFExporter/MigratedFromGit/Source/GLTFExporter/Private/Exporters/GLTFAnimSequenceExporter.cpp b/Engine/Plugins/Enterprise/GLTFExporter/MigratedFromGit/Source/GLTFExporter/Private/Exporters/GLTFAnimSequenceExporter.cpp
--- a/Engine/Plugins/Enterprise/GLTFExporter/MigratedFromGit/Source/GLTFExporter/Private/Exporters/GLTFAnimSequenceExporter.cpp
+++ b/Engine/Plugins/Enterprise/GLTFExporter/MigratedFromGit/Source/GLTFExporter/Private/Exporters/GLTFAnimSequenceExporter.cpp
@@ -44,8 +44,8 @@ bool UGLTFAnimSequenceExporter::AddObject(FGLTFContainerBuilder& Builder, const
 
 	if (Builder.ExportOptions->bExportPreviewMesh)
 	{
-		FGLTFJsonMesh* MeshIndex = Builder.GetOrAddMesh(SkeletalMesh);
-		if (MeshIndex == nullptr)
+		Node->Mesh = Builder.GetOrAddMesh(SkeletalMesh);
+		if (Node->Mesh == nullptr)
 		{
 			Builder.LogError(
 				FString::Printf(TEXT("Failed to export skeletal mesh %s for animation sequence %s"),
@@ -53,13 +53,11 @@ bool UGLTFAnimSequenceExporter::AddObject(FGLTFContainerBuilder& Builder, const
 				*AnimSequence->GetName()));
 			return false;
 		}
-
-		Node->Mesh = MeshIndex;
 	}
 
 
-	FGLTFJsonSkin* SkinIndex = Builder.GetOrAddSkin(Node, SkeletalMesh);
-	if (SkinIndex == nullptr)
+	Node->Skin = Builder.GetOrAddSkin(Node, SkeletalMesh);
+	if (Node->Skin == nullptr)
 	{
 		Builder.LogError(
 			FString::Printf(TEXT("Failed to export bones in skeletal mesh %s for animation sequence %s"),
@@ -68,10 +66,8 @@ bool UGLTFAnimSequenceExporter::AddObject(FGLTFContainerBuilder& Builder, const
 		return false;
 	}
 
-	Node->Skin = SkinIndex;
-
-	FGLTFJsonAnimation* AnimationIndex =  Builder.GetOrAddAnimation(Node, SkeletalMesh, AnimSequence);
-	if (AnimationIndex == nullptr)
+	FGLTFJsonAnimation* Animation = Builder.GetOrAddAnimation(Node, SkeletalMesh, AnimSequence);
+	if (Animation == nullptr)
 	{
 		Builder.LogError(
 			FString::Printf(TEXT("Failed to export animation sequence %s"),
diff --git a/Engine/Plugins/Enterprise/GLTFExporter/MigratedFromGit/Source/GLTFExporter/Private/Exporters/GLTFLevelExporter.cpp b/Engine/Plugins/Enterprise/GLTFExporter/MigratedFromGit/Source/GLTFExporter/Private/Exporters/GLTFLevelExporter.cpp
--- a/Engine/Plugins/Enterprise/GLTFExporter/MigratedFromGit/Source/GLTFExporter/Private/Exporters/GLTFLevelExporter.cpp
+++ b/Engine/Plugins/Enterprise/GLTFExporter/MigratedFromGit/Source/GLTFExporter/Private/Exporters/GLTFLevelExporter.cpp
@@ -14,13 +14,13 @@ bool UGLTFLevelExporter::AddObject(FGLTFContainerBuilder& Builder, const UObject
 {
 	const UWorld* World = CastChecked<UWorld>(Object);
 
-	FGLTFJsonScene* SceneIndex = Builder.GetOrAddScene(World);
-	if (SceneIndex == nullptr)
+	FGLTFJsonScene* Scene = Builder.GetOrAddScene(World);
+	if (Scene == nullptr)
 	{
 		Builder.LogError(FString::Printf(TEXT("Failed to export level %s"), *World->GetName()));
 		return false;
 	}
 
-	Builder.DefaultScene = SceneIndex;
+	Builder.DefaultScene = Scene;
 	return true;
 }
diff --git a/Engine/Plugins/Enterprise/GLTFExporter/MigratedFromGit/Source/GLTFExporter/Private/Exporters/GLTFLevelSequenceExporter.cpp b/Engine/Plugins/Enterprise/GLTFExporter/MigratedFromGit/Source/GLTFExporter/Private/Exporters/GLTFLevelSequenceExporter.cpp
--- a/Engine/Plugins/Enterprise/GLTFExporter/MigratedFromGit/Source/GLTFExporter/Private/Exporters/GLTFLevelSequenceExporter.cpp
+++ b/Engine/Plugins/Enterprise/GLTFExporter/MigratedFromGit/Source/GLTFExporter/Private/Exporters/GLTFLevelSequenceExporter.cpp
@@ -40,9 +40,9 @@ bool UGLTFLevelSequenceExporter::AddObject(FGLTFContainerBuilder& Builder, const
 		return false;
 	}
 
-	UWorld* World = Worlds[0];
-	FGLTFJsonScene* SceneIndex = Builder.GetOrAddScene(World);
-	if (SceneIndex == nullptr)
+	const UWorld* World = Worlds[0];
+	FGLTFJsonScene* Scene = Builder.GetOrAddScene(World);
+	if (Scene == nullptr)
 	{
 		Builder.LogError(
 			FString::Printf(TEXT("Failed to export level %s for level sequence %s"),
@@ -51,8 +51,8 @@ bool UGLTFLevelSequenceExporter::AddObject(FGLTFContainerBuilder& Builder, const
 		return false;
 	}
 
-	FGLTFJsonAnimation* AnimationIndex = Builder.GetOrAddAnimation(World->PersistentLevel, LevelSequence);
-	if (AnimationIndex == nullptr)
+	FGLTFJsonAnimation* Animation = Builder.GetOrAddAnimation(World->PersistentLevel, LevelSequence);
+	if (Animation == nullptr)
 	{
 		Builder.LogError(
 			FString::Printf(TEXT("Failed to export level sequence %s"),
@@ -60,5 +60,6 @@ bool UGLTFLevelSequenceExporter::AddObject(FGLTFContainerBuilder& Builder, const
 		return false;
 	}
 
+	Builder.DefaultScene = Scene;
 	return true;
 }
diff --git a/Engine/Plugins/Enterprise/GLTFExporter/MigratedFromGit/Source/GLTFExporter/Private/Exporters/GLTFLevelVariantSetsExporter.cpp b/Engine/Plugins/Enterprise/GLTFExporter/MigratedFromGit/Source/GLTFExporter/Private/Exporters/GLTFLevelVariantSetsExporter.cpp
--- a/Engine/Plugins/Enterprise/GLTFExporter/MigratedFromGit/Source/GLTFExporter/Private/Exporters/GLTFLevelVariantSetsExporter.cpp
+++ b/Engine/Plugins/Enterprise/GLTFExporter/MigratedFromGit/Source/GLTFExporter/Private/Exporters/GLTFLevelVariantSetsExporter.cpp
@@ -40,9 +40,9 @@ bool UGLTFLevelVariantSetsExporter::AddObject(FGLTFContainerBuilder& Builder, co
 		return false;
 	}
 
-	UWorld* World = Worlds[0];
-	FGLTFJsonScene* SceneIndex = Builder.GetOrAddScene(World);
-	if (SceneIndex == nullptr)
+	const UWorld* World = Worlds[0];
+	FGLTFJsonScene* Scene = Builder.GetOrAddScene(World);
+	if (Scene == nullptr)
 	{
 		Builder.LogError(
 			FString::Printf(TEXT("Failed to export level %s for level variant sets %s"),
@@ -53,8 +53,8 @@ bool UGLTFLevelVariantSetsExporter::AddObject(FGLTFContainerBuilder& Builder, co
 
 	if (Builder.ExportOptions->VariantSetsMode == EGLTFVariantSetsMode::Epic)
 	{
-		FGLTFJsonEpicLevelVariantSets* EpicLevelVariantSetsIndex = Builder.GetOrAddEpicLevelVariantSets(LevelVariantSets);
-		if (EpicLevelVariantSetsIndex == nullptr)
+		FGLTFJsonEpicLevelVariantSets* EpicLevelVariantSets = Builder.GetOrAddEpicLevelVariantSets(LevelVariantSets);
+		if (EpicLevelVariantSets == nullptr)
 		{
 			Builder.LogError(
 				FString::Printf(TEXT("Failed to export level variant sets %s"),
@@ -62,7 +62,7 @@ bool UGLTFLevelVariantSetsExporter::AddObject(FGLTFContainerBuilder& Builder, co
 			return false;
 		}
 
-		Builder.GetScene(SceneIndex).EpicLevelVariantSets.AddUnique(EpicLevelVariantSetsIndex);
+		Scene->EpicLevelVariantSets.AddUnique(EpicLevelVariantSets);
 	}
 	else if (Builder.ExportOptions->VariantSetsMode == EGLTFVariantSetsMode::Khronos)
 	{
@@ -75,6 +75,6 @@ bool UGLTFLevelVariantSetsExporter::AddObject(FGLTFContainerBuilder& Builder, co
 		}
 	}
 
-	Builder.DefaultScene = SceneIndex;
+	Builder.DefaultScene = Scene;
 	return true;
 }
diff --git a/Engine/Plugins/Enterprise/GLTFExporter/MigratedFromGit/Source/GLTFExporter/Private/Exporters/GLTFMaterialExporter.cpp b/Engine/Plugins/Enterprise/GLTFExporter/MigratedFromGit/Source/GLTFExporter/Private/Exporters/GLTFMaterialExporter.cpp
--- a/Engine/Plugins/Enterprise/GLTFExporter/MigratedFromGit/Source/GLTFExporter/Private/Exporters/GLTFMaterialExporter.cpp
+++ b/Engine/Plugins/Enterprise/GLTFExporter/MigratedFromGit/Source/GLTFExporter/Private/Exporters/GLTFMaterialExporter.cpp
@@ -27,8 +27,8 @@ bool UGLTFMaterialExporter::AddObject(FGLTFContainerBuilder& Builder, const UObj
 			PreviewMesh = DefaultPreviewMesh;
 		}
 
-		FGLTFJsonMesh* MeshIndex = Builder.GetOrAddMesh(PreviewMesh, { Material });
-		if (MeshIndex == nullptr)
+		FGLTFJsonMesh* Mesh = Builder.GetOrAddMesh(PreviewMesh, { Material });
+		if (Mesh == nullptr)
 		{
 			Builder.LogError(
 				FString::Printf(TEXT("Failed to export preview mesh %s for material %s"),
@@ -38,7 +38,7 @@ bool UGLTFMaterialExporter::AddObject(FGLTFContainerBuilder& Builder, const UObj
 		}
 
 		FGLTFJsonNode* Node = Builder.AddNode();
-		Node->Mesh = MeshIndex;
+		Node->Mesh = Mesh;
 
 		FGLTFJsonScene* Scene = Builder.AddScene();
 		Scene->Nodes.Add(Node);
diff --git a/Engine/Plugins/Enterprise/GLTFExporter/MigratedFromGit/Source/GLTFExporter/Private/Exporters/GLTFSkeletalMeshExporter.cpp b/Engine/Plugins/Enterprise/GLTFExporter/MigratedFromGit/Source/GLTFExporter/Private/Exporters/GLTFSkeletalMeshExporter.cpp
--- a/Engine/Plugins/Enterprise/GLTFExporter/MigratedFromGit/Source/GLTFExporter/Private/Exporters/GLTFSkeletalMeshExporter.cpp
+++ b/Engine/Plugins/Enterprise/GLTFExporter/MigratedFromGit/Source/GLTFExporter/Private/Exporters/GLTFSkeletalMeshExporter.cpp
@@ -14,26 +14,24 @@ bool UGLTFSkeletalMeshExporter::AddObject(FGLTFContainerBuilder& Builder, const
 {
 	const USkeletalMesh* SkeletalMesh = CastChecked<USkeletalMesh>(Object);
 
-	FGLTFJsonMesh* MeshIndex = Builder.GetOrAddMesh(SkeletalMesh);
-	if (MeshIndex == nullptr)
+	FGLTFJsonMesh* Mesh = Builder.GetOrAddMesh(SkeletalMesh);
+	if (Mesh == nullptr)
 	{
 		Builder.LogError(FString::Printf(TEXT("Failed to export skeletal mesh %s"), *SkeletalMesh->GetName()));
 		return false;
 	}
 
 	FGLTFJsonNode* Node = Builder.AddNode();
-	Node->Mesh = MeshIndex;
+	Node->Mesh = Mesh;
 
 	if (Builder.ExportOptions->bExportVertexSkinWeights)
 	{
-		FGLTFJsonSkin* SkinIndex = Builder.GetOrAddSkin(Node, SkeletalMesh);
-		if (SkinIndex == nullptr)
+		Node->Skin = Builder.GetOrAddSkin(Node, SkeletalMesh);
+		if (Node->Skin == nullptr)
 		{
 			Builder.LogError(FString::Printf(TEXT("Failed to export bones in skeletal mesh %s"), *SkeletalMesh->GetName()));
 			return false;
 		}
-
-		Node->Skin = SkinIndex;
 	}
 
 	FGLTFJsonScene* Scene = Builder.AddScene();
diff --git a/Engine/Plugins/Enterprise/GLTFExporter/MigratedFromGit/Source/GLTFExporter/Private/Exporters/GLTFStaticMeshExporter.cpp b/Engine/Plugins/Enterprise/GLTFExporter/MigratedFromGit/Source/GLTFExporter/Private/Exporters/GLTFStaticMeshExporter.cpp
--- a/Engine/Plugins/Enterprise/GLTFExporter/MigratedFromGit/Source/GLTFExporter/Private/Exporters/GLTFStaticMeshExporter.cpp
+++ b/Engine/Plugins/Enterprise/GLTFExporter/MigratedFromGit/Source/GLTFExporter/Private/Exporters/GLTFStaticMeshExporter.cpp
@@ -14,15 +14,15 @@ bool UGLTFStaticMeshExporter::AddObject(FGLTFContainerBuilder& Builder, const UO
 {
 	const UStaticMesh* StaticMesh = CastChecked<UStaticMesh>(Object);
 
-	FGLTFJsonMesh* MeshIndex = Builder.GetOrAddMesh(StaticMesh);
-	if (MeshIndex == nullptr)
+	FGLTFJsonMesh* Mesh = Builder.GetOrAddMesh(StaticMesh);
+	if (Mesh == nullptr)
 	{
 		Builder.LogError(FString::Printf(TEXT("Failed to export static mesh %s"), *StaticMesh->GetName()));
 		return false;
 	}
 
 	FGLTFJsonNode* Node = Builder.AddNode();
-	Node->Mesh = MeshIndex;
+	Node->Mesh = Mesh;
 
 	FGLTFJsonScene* Scene = Builder.AddScene();
 	Scene->Nodes.Add(Node);
diff --git a/Engine/Plugins/Enterprise/GLTFExporter/MigratedFromGit/Source/GLTFExporter/Private/Tasks/GLTFAnimationTasks.cpp b/Engine/Plugins/Enterprise/GLTFExporter/MigratedFromGit/Source/GLTFExporter/Private/Tasks/GLTFAnimationTasks.cpp
--- a/Engine/Plugins/Enterprise/GLTFExporter/MigratedFromGit/Source/GLTFExporter/Private/Tasks/GLTFAnimationTasks.cpp
+++ b/Engine/Plugins/Enterprise/GLTFExporter/MigratedFromGit/Source/GLTFExporter/Private/Tasks/GLTFAnimationTasks.cpp
@@ -49,9 +49,9 @@ void FGLTFAnimSequenceTask::Complete()
 	TArray<TArray<FTransform>> FrameTransforms;
 	FGLTFBoneUtility::GetBoneTransformsByFrame(AnimSequence, Timestamps, BoneIndices, FrameTransforms);
 
-	for (FBoneIndexType BoneIndex : BoneIndices)
+	for (const FBoneIndexType BoneIndex : BoneIndices)
 	{
-		FGLTFJsonNode* NodeIndex = Builder.GetOrAddNode(RootNode, SkeletalMesh, BoneIndex);
+		FGLTFJsonNode* Node = Builder.GetOrAddNode(RootNode, SkeletalMesh, BoneIndex);
 
 		// TODO: detect if a bone has the same transforms across multiple frames (at least if its the same across all frames) and optimize
 
@@ -79,7 +79,7 @@ void FGLTFAnimSequenceTask::Complete()
 			FGLTFJsonAnimationChannel JsonChannel;
 			JsonChannel.Sampler = JsonSampler;
 			JsonChannel.Target.Path = EGLTFJsonTargetPath::Translation;
-			JsonChannel.Target.Node = NodeIndex;
+			JsonChannel.Target.Node = Node;
 			JsonAnimation->Channels.Add(JsonChannel);
 		}
 
@@ -107,7 +107,7 @@ void FGLTFAnimSequenceTask::Complete()
 			FGLTFJsonAnimationChannel JsonChannel;
 			JsonChannel.Sampler = JsonSampler;
 			JsonChannel.Target.Path = EGLTFJsonTargetPath::Rotation;
-			JsonChannel.Target.Node = NodeIndex;
+			JsonChannel.Target.Node = Node;
 			JsonAnimation->Channels.Add(JsonChannel);
 		}
 
@@ -135,7 +135,7 @@ void FGLTFAnimSequenceTask::Complete()
 			FGLTFJsonAnimationChannel JsonChannel;
 			JsonChannel.Sampler = JsonSampler;
 			JsonChannel.Target.Path = EGLTFJsonTargetPath::Scale;
-			JsonChannel.Target.Node = NodeIndex;
+			JsonChannel.Target.Node = Node;
 			JsonAnimation->Channels.Add(JsonChannel);
 		}
 	}
@@ -188,15 +188,15 @@ void FGLTFLevelSequenceTask::Complete()
 	{
 		for (TWeakObjectPtr<UObject> Object : Player->FindBoundObjects(Binding.GetObjectGuid(), MovieSceneSequenceID::Root))
 		{
-			FGLTFJsonNode* NodeIndex;
+			FGLTFJsonNode* Node;
 
 			if (AActor* Actor = Cast<AActor>(Object.Get()))
 			{
-				NodeIndex = Builder.GetOrAddNode(Actor);
+				Node = Builder.GetOrAddNode(Actor);
 			}
 			else if (USceneComponent* SceneComponent = Cast<USceneComponent>(Object.Get()))
 			{
-				NodeIndex = Builder.GetOrAddNode(SceneComponent);
+				Node = Builder.GetOrAddNode(SceneComponent);
 			}
 			else
 			{
@@ -204,7 +204,7 @@ void FGLTFLevelSequenceTask::Complete()
 				continue;
 			}
 
-			if (NodeIndex == nullptr)
+			if (Node == nullptr)
 			{
 				// TODO: report warning
 				continue;
@@ -273,7 +273,7 @@ void FGLTFLevelSequenceTask::Complete()
 						FGLTFJsonAnimationChannel JsonChannel;
 						JsonChannel.Sampler = JsonSampler;
 						JsonChannel.Target.Path = EGLTFJsonTargetPath::Translation;
-						JsonChannel.Target.Node = NodeIndex;
+						JsonChannel.Target.Node = Node;
 						JsonAnimation->Channels.Add(JsonChannel);
 					}
 
@@ -311,7 +311,7 @@ void FGLTFLevelSequenceTask::Complete()
 						FGLTFJsonAnimationChannel JsonChannel;
 						JsonChannel.Sampler = JsonSampler;
 						JsonChannel.Target.Path = EGLTFJsonTargetPath::Rotation;
-						JsonChannel.Target.Node = NodeIndex;
+						JsonChannel.Target.Node = Node;
 						JsonAnimation->Channels.Add(JsonChannel);
 					}
 
@@ -346,7 +346,7 @@ void FGLTFLevelSequenceTask::Complete()
 						FGLTFJsonAnimationChannel JsonChannel;
 						JsonChannel.Sampler = JsonSampler;
 						JsonChannel.Target.Path = EGLTFJsonTargetPath::Scale;
-						JsonChannel.Target.Node = NodeIndex;
+						JsonChannel.Target.Node = Node;
 						JsonAnimation->Channels.Add(JsonChannel);
 					}
 				}
diff --git a/Engine/Plugins/Enterprise/GLTFExporter/MigratedFromGit/Source/GLTFExporter/Private/Tasks/GLTFMaterialTasks.cpp b/Engine/Plugins/Enterprise/GLTFExporter/MigratedFromGit/Source/GLTFExporter/Private/Tasks/GLTFMaterialTasks.cpp
--- a/Engine/Plugins/Enterprise/GLTFExporter/MigratedFromGit/Source/GLTFExporter/Private/Tasks/GLTFMaterialTasks.cpp
+++ b/Engine/Plugins/Enterprise/GLTFExporter/MigratedFromGit/Source/GLTFExporter/Private/Tasks/GLTFMaterialTasks.cpp
@@ -590,7 +590,7 @@ bool FGLTFMaterialTask::TryGetBaseColorAndOpacity(FGLTFJsonPBRMetallicRoughness&
 			return FColor(BaseColor.R, BaseColor.G, BaseColor.B, Opacity.R);
 		});
 
-	FGLTFJsonTexture* TextureIndex = FGLTFMaterialUtility::AddTexture(
+	FGLTFJsonTexture* Texture = FGLTFMaterialUtility::AddTexture(
 		Builder,
 		CombinedPixels,
 		TextureSize,
@@ -603,7 +603,7 @@ bool FGLTFMaterialTask::TryGetBaseColorAndOpacity(FGLTFJsonPBRMetallicRoughness&
 		TextureWrapT);
 
 	OutPBRParams.BaseColorTexture.TexCoord = TexCoord;
-	OutPBRParams.BaseColorTexture.Index = TextureIndex;
+	OutPBRParams.BaseColorTexture.Index = Texture;
 	OutPBRParams.BaseColorFactor = FGLTFConverterUtility::ConvertColor({ BaseColorScale, BaseColorScale, BaseColorScale }, Builder.ExportOptions->bStrictCompliance);
 
 	return true;
@@ -714,7 +714,7 @@ bool FGLTFMaterialTask::TryGetMetallicAndRoughness(FGLTFJsonPBRMetallicRoughness
 			return FColor(0, Roughness.R, Metallic.R);
 		});
 
-	FGLTFJsonTexture* TextureIndex = FGLTFMaterialUtility::AddTexture(
+	FGLTFJsonTexture* Texture = FGLTFMaterialUtility::AddTexture(
 		Builder,
 		CombinedPixels,
 		TextureSize,
@@ -727,7 +727,7 @@ bool FGLTFMaterialTask::TryGetMetallicAndRoughness(FGLTFJsonPBRMetallicRoughness
 		TextureWrapT);
 
 	OutPBRParams.MetallicRoughnessTexture.TexCoord = TexCoord;
-	OutPBRParams.MetallicRoughnessTexture.Index = TextureIndex;
+	OutPBRParams.MetallicRoughnessTexture.Index = Texture;
 
 	return true;
 }
@@ -765,10 +765,10 @@ bool FGLTFMaterialTask::TryGetClearCoatRoughness(FGLTFJsonClearCoatExtension& Ou
 		IntensityTexCoord == RoughnessTexCoord &&
 		IntensityTransform.IsExactlyEqual(RoughnessTransform))
 	{
-		FGLTFJsonTexture* TextureIndex = Builder.GetOrAddTexture(IntensityTexture, false);
-		OutExtParams.ClearCoatTexture.Index = TextureIndex;
+		FGLTFJsonTexture* Texture = Builder.GetOrAddTexture(IntensityTexture, false);
+		OutExtParams.ClearCoatTexture.Index = Texture;
 		OutExtParams.ClearCoatTexture.TexCoord = IntensityTexCoord;
-		OutExtParams.ClearCoatRoughnessTexture.Index = TextureIndex;
+		OutExtParams.ClearCoatRoughnessTexture.Index = Texture;
 		OutExtParams.ClearCoatRoughnessTexture.TexCoord = IntensityTexCoord;
 		OutExtParams.ClearCoatRoughnessTexture.Transform = IntensityTransform;
 		return true;
@@ -839,7 +839,7 @@ bool FGLTFMaterialTask::TryGetClearCoatRoughness(FGLTFJsonClearCoatExtension& Ou
 			return FColor(Intensity.R, Roughness.R, 0);
 		});
 
-	FGLTFJsonTexture* TextureIndex = FGLTFMaterialUtility::AddTexture(
+	FGLTFJsonTexture* Texture = FGLTFMaterialUtility::AddTexture(
 		Builder,
 		CombinedPixels,
 		TextureSize,
@@ -851,9 +851,9 @@ bool FGLTFMaterialTask::TryGetClearCoatRoughness(FGLTFJsonClearCoatExtension& Ou
 		TextureWrapS,
 		TextureWrapT);
 
-	OutExtParams.ClearCoatTexture.Index = TextureIndex;
+	OutExtParams.ClearCoatTexture.Index = Texture;
 	OutExtParams.ClearCoatTexture.TexCoord = TexCoord;
-	OutExtParams.ClearCoatRoughnessTexture.Index = TextureIndex;
+	OutExtParams.ClearCoatRoughnessTexture.Index = Texture;
 	OutExtParams.ClearCoatRoughnessTexture.TexCoord = TexCoord;
 
 	return true;
@@ -1456,7 +1456,7 @@ bool FGLTFMaterialTask::TryGetBakedMaterialProperty(FGLTFJsonTextureInfo& OutTex
 	// even though the same material when set to opaque will properly bake AmbientOcclusion to a texture.
 	// For now, create a 1x1 texture with the constant value.
 
-	FGLTFJsonTexture* TextureIndex = FGLTFMaterialUtility::AddTexture(
+	FGLTFJsonTexture* Texture = FGLTFMaterialUtility::AddTexture(
 		Builder,
 		PropertyBakeOutput.Pixels,
 		PropertyBakeOutput.Size,
@@ -1468,7 +1468,7 @@ bool FGLTFMaterialTask::TryGetBakedMaterialProperty(FGLTFJsonTextureInfo& OutTex
 		EGLTFJsonTextureWrap::ClampToEdge,
 		EGLTFJsonTextureWrap::ClampToEdge);
 
-	OutTexInfo.Index = TextureIndex;
+	OutTexInfo.Index = Texture;
 	return true;
 }
 
@@ -1533,7 +1533,7 @@ bool FGLTFMaterialTask::StoreBakedPropertyTexture(FGLTFJsonTextureInfo& OutTexIn
 	const EGLTFJsonTextureFilter TextureMinFilter = FGLTFConverterUtility::ConvertMinFilter(TextureFilter);
 	const EGLTFJsonTextureFilter TextureMagFilter = FGLTFConverterUtility::ConvertMagFilter(TextureFilter);
 
-	FGLTFJsonTexture* TextureIndex = FGLTFMaterialUtility::AddTexture(
+	FGLTFJsonTexture* Texture = FGLTFMaterialUtility::AddTexture(
 		Builder,
 		PropertyBakeOutput.Pixels,
 		PropertyBakeOutput.Size,
@@ -1545,7 +1545,7 @@ bool FGLTFMaterialTask::StoreBakedPropertyTexture(FGLTFJsonTextureInfo& OutTexIn
 		TextureWrapS,
 		TextureWrapT);
 
-	OutTexInfo.Index = TextureIndex;
+	OutTexInfo.Index = Texture;
 	return true;
 }
 
