[GLTFExporter] refactor: remove AddChildNode and AddChildComponentNode in Builder

commit e76c5070c217b97f942451918391797836d37812
Date: Sun, 13 Mar 2022 18:57:27 +0100

    refactor: remove AddChildNode and AddChildComponentNode in Builder
    (cherry picked from commit df1ca9864987a3272b70c1ea244757db76382f32​)

#skipci

diff --git a/Engine/Plugins/Enterprise/GLTFExporter/MigratedFromGit/Source/GLTFExporter/Private/Builders/GLTFJsonBuilder.cpp b/Engine/Plugins/Enterprise/GLTFExporter/MigratedFromGit/Source/GLTFExporter/Private/Builders/GLTFJsonBuilder.cpp
--- a/Engine/Plugins/Enterprise/GLTFExporter/MigratedFromGit/Source/GLTFExporter/Private/Builders/GLTFJsonBuilder.cpp
+++ b/Engine/Plugins/Enterprise/GLTFExporter/MigratedFromGit/Source/GLTFExporter/Private/Builders/GLTFJsonBuilder.cpp
@@ -141,36 +141,6 @@ FGLTFJsonKhrMaterialVariant* FGLTFJsonBuilder::AddKhrMaterialVariant()
 	return JsonRoot.KhrMaterialVariants.Add();
 }
 
-FGLTFJsonNode* FGLTFJsonBuilder::AddChildNode(FGLTFJsonNode* ParentNode, FGLTFJsonNode* ChildNode)
-{
-	if (ChildNode == nullptr)
-	{
-		ChildNode = AddNode();
-	}
-
-	if (ParentNode != nullptr)
-	{
-		ParentNode->Children.Add(ChildNode);
-	}
-
-	return ChildNode;
-}
-
-FGLTFJsonNode* FGLTFJsonBuilder::AddChildComponentNode(FGLTFJsonNode* ParentNode, FGLTFJsonNode* ChildNode)
-{
-	if (ChildNode == nullptr)
-	{
-		ChildNode = AddChildNode(ParentNode, ChildNode);
-	}
-
-	if (ParentNode != nullptr)
-	{
-		ParentNode->ComponentNode = ChildNode;
-	}
-
-	return ChildNode;
-}
-
 const FGLTFJsonRoot& FGLTFJsonBuilder::GetRoot() const
 {
 	return JsonRoot;
diff --git a/Engine/Plugins/Enterprise/GLTFExporter/MigratedFromGit/Source/GLTFExporter/Private/Builders/GLTFJsonBuilder.h b/Engine/Plugins/Enterprise/GLTFExporter/MigratedFromGit/Source/GLTFExporter/Private/Builders/GLTFJsonBuilder.h
--- a/Engine/Plugins/Enterprise/GLTFExporter/MigratedFromGit/Source/GLTFExporter/Private/Builders/GLTFJsonBuilder.h
+++ b/Engine/Plugins/Enterprise/GLTFExporter/MigratedFromGit/Source/GLTFExporter/Private/Builders/GLTFJsonBuilder.h
@@ -40,9 +40,6 @@ public:
 	FGLTFJsonEpicLevelVariantSets* AddEpicLevelVariantSets();
 	FGLTFJsonKhrMaterialVariant* AddKhrMaterialVariant();
 
-	FGLTFJsonNode* AddChildNode(FGLTFJsonNode* ParentNode, FGLTFJsonNode* ChildNode = nullptr);
-	FGLTFJsonNode* AddChildComponentNode(FGLTFJsonNode* ParentNodeParentNode, FGLTFJsonNode* ChildNode = nullptr);
-
 	const FGLTFJsonRoot& GetRoot() const;
 
 	FGLTFJsonNode* GetComponentNode(FGLTFJsonNode* Node);
diff --git a/Engine/Plugins/Enterprise/GLTFExporter/MigratedFromGit/Source/GLTFExporter/Private/Converters/GLTFNodeConverters.cpp b/Engine/Plugins/Enterprise/GLTFExporter/MigratedFromGit/Source/GLTFExporter/Private/Converters/GLTFNodeConverters.cpp
--- a/Engine/Plugins/Enterprise/GLTFExporter/MigratedFromGit/Source/GLTFExporter/Private/Converters/GLTFNodeConverters.cpp
+++ b/Engine/Plugins/Enterprise/GLTFExporter/MigratedFromGit/Source/GLTFExporter/Private/Converters/GLTFNodeConverters.cpp
@@ -120,12 +120,17 @@ FGLTFJsonNode* FGLTFComponentConverter::Convert(const USceneComponent* SceneComp
 	const FTransform ParentTransform = ParentComponent != nullptr ? ParentComponent->GetSocketTransform(SocketName) : FTransform::Identity;
 	const FTransform RelativeTransform = bIsRootNode ? Transform : Transform.GetRelativeTransform(ParentTransform);
 
-	FGLTFJsonNode* Node = Builder.AddChildNode(ParentNode);
+	FGLTFJsonNode* Node = Builder.AddNode();
 	Node->Name = FGLTFNameUtility::GetName(SceneComponent);
 	Node->Translation = FGLTFConverterUtility::ConvertPosition(RelativeTransform.GetTranslation(), Builder.ExportOptions->ExportUniformScale);
 	Node->Rotation = FGLTFConverterUtility::ConvertRotation(RelativeTransform.GetRotation());
 	Node->Scale = FGLTFConverterUtility::ConvertScale(RelativeTransform.GetScale3D());
 
+	if (ParentNode != nullptr)
+	{
+		ParentNode->Children.Add(Node);
+	}
+
 	// TODO: don't export invisible components unless visibility is variable due to variant sets
 
 	// TODO: should hidden in game be configurable like this?
@@ -161,10 +166,12 @@ FGLTFJsonNode* FGLTFComponentConverter::Convert(const USceneComponent* SceneComp
 			if (Builder.ExportOptions->bExportCameras)
 			{
 				// TODO: conversion of camera direction should be done in separate converter
-				FGLTFJsonNode* CameraNode = Builder.AddChildComponentNode(Node);
+				FGLTFJsonNode* CameraNode = Builder.AddNode();
 				CameraNode->Name = FGLTFNameUtility::GetName(CameraComponent);
 				CameraNode->Rotation = FGLTFConverterUtility::ConvertCameraDirection();
 				CameraNode->Camera = Builder.GetOrAddCamera(CameraComponent);
+				Node->Children.Add(CameraNode);
+				Node->ComponentNode = CameraNode;
 			}
 		}
 		else if (const ULightComponent* LightComponent = Cast<ULightComponent>(SceneComponent))
@@ -172,10 +179,12 @@ FGLTFJsonNode* FGLTFComponentConverter::Convert(const USceneComponent* SceneComp
 			if (Builder.ShouldExportLight(LightComponent->Mobility))
 			{
 				// TODO: conversion of light direction should be done in separate converter
-				FGLTFJsonNode* LightNode = Builder.AddChildComponentNode(Node);
+				FGLTFJsonNode* LightNode = Builder.AddNode();
 				LightNode->Name = FGLTFNameUtility::GetName(LightComponent);
 				LightNode->Rotation = FGLTFConverterUtility::ConvertLightDirection();
 				LightNode->Light = Builder.GetOrAddLight(LightComponent);
+				Node->Children.Add(LightNode);
+				Node->ComponentNode = LightNode;
 			}
 		}
 	}
@@ -230,7 +239,8 @@ FGLTFJsonNode* FGLTFStaticSocketConverter::Convert(FGLTFJsonNode* RootNode, cons
 	Node->Rotation = FGLTFConverterUtility::ConvertRotation(Socket->RelativeRotation.Quaternion());
 	Node->Scale = FGLTFConverterUtility::ConvertScale(Socket->RelativeScale);
 
-	return Builder.AddChildNode(RootNode, Node);
+	RootNode->Children.Add(Node);
+	return Node;
 }
 
 FGLTFJsonNode* FGLTFSkeletalSocketConverter::Convert(FGLTFJsonNode* RootNode, const USkeletalMesh* SkeletalMesh, FName SocketName)
@@ -254,7 +264,9 @@ FGLTFJsonNode* FGLTFSkeletalSocketConverter::Convert(FGLTFJsonNode* RootNode, co
 
 		const int32 ParentBone = RefSkeleton.FindBoneIndex(Socket->BoneName);
 		FGLTFJsonNode* ParentNode = ParentBone != INDEX_NONE ? Builder.GetOrAddNode(RootNode, SkeletalMesh, ParentBone) : RootNode;
-		return Builder.AddChildNode(ParentNode, Node);
+
+		ParentNode->Children.Add(Node);
+		return Node;
 	}
 
 	const int32 BoneIndex = RefSkeleton.FindBoneIndex(SocketName);
@@ -304,5 +316,6 @@ FGLTFJsonNode* FGLTFSkeletalBoneConverter::Convert(FGLTFJsonNode* RootNode, cons
 	}
 
 	FGLTFJsonNode* ParentNode = BoneInfo.ParentIndex != INDEX_NONE ? Builder.GetOrAddNode(RootNode, SkeletalMesh, BoneInfo.ParentIndex) : RootNode;
-	return Builder.AddChildNode(ParentNode, Node);
+	ParentNode->Children.Add(Node);
+	return Node;
 }
