[GLTFExporter] refactor: remove ComponentNode from FGLTFJsonNode

commit b2596c8344073ffa3b697d301f5ae8ad1825dd12
Date: Sun, 13 Mar 2022 19:05:02 +0100

    refactor: remove ComponentNode from FGLTFJsonNode
    (cherry picked from commit 211b81070a527836cf419e3ad539261f8399cfcf​)

#skipci

diff --git a/Engine/Plugins/Enterprise/GLTFExporter/MigratedFromGit/Source/GLTFExporter/Private/Builders/GLTFJsonBuilder.cpp b/Engine/Plugins/Enterprise/GLTFExporter/MigratedFromGit/Source/GLTFExporter/Private/Builders/GLTFJsonBuilder.cpp
--- a/Engine/Plugins/Enterprise/GLTFExporter/MigratedFromGit/Source/GLTFExporter/Private/Builders/GLTFJsonBuilder.cpp
+++ b/Engine/Plugins/Enterprise/GLTFExporter/MigratedFromGit/Source/GLTFExporter/Private/Builders/GLTFJsonBuilder.cpp
@@ -146,16 +146,6 @@ const FGLTFJsonRoot& FGLTFJsonBuilder::GetRoot() const
 	return JsonRoot;
 }
 
-FGLTFJsonNode* FGLTFJsonBuilder::GetComponentNode(FGLTFJsonNode* Node)
-{
-	if (Node == nullptr)
-	{
-		return nullptr;
-	}
-
-	return Node->ComponentNode != nullptr ? Node->ComponentNode : Node;
-}
-
 FString FGLTFJsonBuilder::GetGeneratorString() const
 {
 	const TSharedPtr<IPlugin> Plugin = IPluginManager::Get().FindPlugin(GLTFEXPORTER_MODULE_NAME);
diff --git a/Engine/Plugins/Enterprise/GLTFExporter/MigratedFromGit/Source/GLTFExporter/Private/Builders/GLTFJsonBuilder.h b/Engine/Plugins/Enterprise/GLTFExporter/MigratedFromGit/Source/GLTFExporter/Private/Builders/GLTFJsonBuilder.h
--- a/Engine/Plugins/Enterprise/GLTFExporter/MigratedFromGit/Source/GLTFExporter/Private/Builders/GLTFJsonBuilder.h
+++ b/Engine/Plugins/Enterprise/GLTFExporter/MigratedFromGit/Source/GLTFExporter/Private/Builders/GLTFJsonBuilder.h
@@ -42,8 +42,6 @@ public:
 
 	const FGLTFJsonRoot& GetRoot() const;
 
-	FGLTFJsonNode* GetComponentNode(FGLTFJsonNode* Node);
-
 private:
 
 	FString GetGeneratorString() const;
diff --git a/Engine/Plugins/Enterprise/GLTFExporter/MigratedFromGit/Source/GLTFExporter/Private/Converters/GLTFEpicLevelVariantSetsConverters.cpp b/Engine/Plugins/Enterprise/GLTFExporter/MigratedFromGit/Source/GLTFExporter/Private/Converters/GLTFEpicLevelVariantSetsConverters.cpp
--- a/Engine/Plugins/Enterprise/GLTFExporter/MigratedFromGit/Source/GLTFExporter/Private/Converters/GLTFEpicLevelVariantSetsConverters.cpp
+++ b/Engine/Plugins/Enterprise/GLTFExporter/MigratedFromGit/Source/GLTFExporter/Private/Converters/GLTFEpicLevelVariantSetsConverters.cpp
@@ -10,6 +10,8 @@
 #include "PropertyValue.h"
 #include "PropertyValueMaterial.h"
 #include "Components/MeshComponent.h"
+#include "Components/LightComponent.h"
+#include "Camera/CameraComponent.h"
 
 FGLTFJsonEpicLevelVariantSets* FGLTFEpicLevelVariantSetsConverter::Convert(const ULevelVariantSets* LevelVariantSets)
 {
@@ -158,11 +160,21 @@ bool FGLTFEpicLevelVariantSetsConverter::TryParseVisibilityPropertyValue(FGLTFJs
 	Builder.RegisterObjectVariant(Owner, Property); // TODO: we should register this on the component
 
 	FGLTFJsonNode* Node = Builder.GetOrAddNode(Target);
-	FGLTFJsonNode* ComponentNode = Builder.GetComponentNode(Node);
+	if (Target->IsA<UCameraComponent>() || Target->IsA<ULightComponent>())
+	{
+		// NOTE: If target component is a camera or light, we know that the target gltf node is actually the first child.
+		if (Node->Children.Num() <= 0)
+		{
+			// TODO: report error
+			return false;
+		}
+
+		Node = Node->Children[0];
+	}
 
-	FGLTFJsonEpicVariantNodeProperties& NodeProperties = OutVariant.Nodes.FindOrAdd(ComponentNode);
+	FGLTFJsonEpicVariantNodeProperties& NodeProperties = OutVariant.Nodes.FindOrAdd(Node);
 
-	NodeProperties.Node = ComponentNode;
+	NodeProperties.Node = Node;
 	NodeProperties.bIsVisible = bIsVisible;
 	return true;
 }
@@ -221,10 +233,9 @@ bool FGLTFEpicLevelVariantSetsConverter::TryParseMaterialPropertyValue(FGLTFJson
 	VariantMaterial.Index = MaterialIndex;
 
 	FGLTFJsonNode* Node = Builder.GetOrAddNode(Target);
-	FGLTFJsonNode* ComponentNode = Builder.GetComponentNode(Node);
-	FGLTFJsonEpicVariantNodeProperties& NodeProperties = OutVariant.Nodes.FindOrAdd(ComponentNode);
+	FGLTFJsonEpicVariantNodeProperties& NodeProperties = OutVariant.Nodes.FindOrAdd(Node);
 
-	NodeProperties.Node = ComponentNode;
+	NodeProperties.Node = Node;
 	NodeProperties.Materials.Add(VariantMaterial);
 	return true;
 }
@@ -272,11 +283,10 @@ bool FGLTFEpicLevelVariantSetsConverter::TryParseMeshPropertyValue(FGLTFJsonEpic
 
 	const FGLTFMaterialArray OverrideMaterials(Target->OverrideMaterials);
 	FGLTFJsonNode* JsonNode = Builder.GetOrAddNode(Target);
-	FGLTFJsonNode* JsonComponentNode = Builder.GetComponentNode(JsonNode);
-	FGLTFJsonEpicVariantNodeProperties& NodeProperties = OutVariant.Nodes.FindOrAdd(JsonComponentNode);
+	FGLTFJsonEpicVariantNodeProperties& NodeProperties = OutVariant.Nodes.FindOrAdd(JsonNode);
 	FGLTFJsonMesh* JsonMesh = Builder.GetOrAddMesh(Mesh, OverrideMaterials);
 
-	NodeProperties.Node = JsonComponentNode;
+	NodeProperties.Node = JsonNode;
 	NodeProperties.Mesh = JsonMesh;
 	return true;
 }
diff --git a/Engine/Plugins/Enterprise/GLTFExporter/MigratedFromGit/Source/GLTFExporter/Private/Converters/GLTFNodeConverters.cpp b/Engine/Plugins/Enterprise/GLTFExporter/MigratedFromGit/Source/GLTFExporter/Private/Converters/GLTFNodeConverters.cpp
--- a/Engine/Plugins/Enterprise/GLTFExporter/MigratedFromGit/Source/GLTFExporter/Private/Converters/GLTFNodeConverters.cpp
+++ b/Engine/Plugins/Enterprise/GLTFExporter/MigratedFromGit/Source/GLTFExporter/Private/Converters/GLTFNodeConverters.cpp
@@ -171,7 +171,6 @@ FGLTFJsonNode* FGLTFComponentConverter::Convert(const USceneComponent* SceneComp
 				CameraNode->Rotation = FGLTFConverterUtility::ConvertCameraDirection();
 				CameraNode->Camera = Builder.GetOrAddCamera(CameraComponent);
 				Node->Children.Add(CameraNode);
-				Node->ComponentNode = CameraNode;
 			}
 		}
 		else if (const ULightComponent* LightComponent = Cast<ULightComponent>(SceneComponent))
@@ -184,7 +183,6 @@ FGLTFJsonNode* FGLTFComponentConverter::Convert(const USceneComponent* SceneComp
 				LightNode->Rotation = FGLTFConverterUtility::ConvertLightDirection();
 				LightNode->Light = Builder.GetOrAddLight(LightComponent);
 				Node->Children.Add(LightNode);
-				Node->ComponentNode = LightNode;
 			}
 		}
 	}
diff --git a/Engine/Plugins/Enterprise/GLTFExporter/MigratedFromGit/Source/GLTFExporter/Public/Json/GLTFJsonNode.h b/Engine/Plugins/Enterprise/GLTFExporter/MigratedFromGit/Source/GLTFExporter/Public/Json/GLTFJsonNode.h
--- a/Engine/Plugins/Enterprise/GLTFExporter/MigratedFromGit/Source/GLTFExporter/Public/Json/GLTFJsonNode.h
+++ b/Engine/Plugins/Enterprise/GLTFExporter/MigratedFromGit/Source/GLTFExporter/Public/Json/GLTFJsonNode.h
@@ -21,8 +21,6 @@ struct GLTFEXPORTER_API FGLTFJsonNode : IGLTFJsonIndexedObject
 	FGLTFJsonLightMap*  LightMap;
 	FGLTFJsonSkySphere* SkySphere;
 
-	FGLTFJsonNode* ComponentNode;
-
 	TArray<FGLTFJsonNode*> Children;
 
 	virtual void WriteObject(IGLTFJsonWriter& Writer) const override;
@@ -44,7 +42,6 @@ protected:
 		, Light(nullptr)
 		, LightMap(nullptr)
 		, SkySphere(nullptr)
-		, ComponentNode(nullptr)
 	{
 	}
 };
