[GLTFExporter] refactor: remove FGLTFBinaryHashKey

commit b454cf451ab4fb6ba50738f324e1c5d3077238dd
Date: Sun, 13 Mar 2022 19:33:44 +0100

    refactor: remove FGLTFBinaryHashKey
    (cherry picked from commit a19720ca8590a59befcce19d14a1a1be5c6b263a​)

#skipci

diff --git a/Engine/Plugins/Enterprise/GLTFExporter/MigratedFromGit/Source/GLTFExporter/Private/Builders/GLTFBinaryHashKey.h b/Engine/Plugins/Enterprise/GLTFExporter/MigratedFromGit/Source/GLTFExporter/Private/Builders/GLTFBinaryHashKey.h
deleted file mode 100644
diff --git a/Engine/Plugins/Enterprise/GLTFExporter/MigratedFromGit/Source/GLTFExporter/Private/Builders/GLTFJsonBuilder.h b/Engine/Plugins/Enterprise/GLTFExporter/MigratedFromGit/Source/GLTFExporter/Private/Builders/GLTFJsonBuilder.h
--- a/Engine/Plugins/Enterprise/GLTFExporter/MigratedFromGit/Source/GLTFExporter/Private/Builders/GLTFJsonBuilder.h
+++ b/Engine/Plugins/Enterprise/GLTFExporter/MigratedFromGit/Source/GLTFExporter/Private/Builders/GLTFJsonBuilder.h
@@ -15,8 +15,6 @@ public:
 
 	void WriteJson(FArchive& Archive);
 
-	TSet<EGLTFJsonExtension> GetCustomExtensionsUsed() const;
-
 	void AddExtension(EGLTFJsonExtension Extension, bool bIsRequired = false);
 
 	FGLTFJsonAccessor* AddAccessor();
@@ -42,6 +40,10 @@ public:
 
 	const FGLTFJsonRoot& GetRoot() const;
 
+protected:
+
+	TSet<EGLTFJsonExtension> GetCustomExtensionsUsed() const;
+
 private:
 
 	FString GetGeneratorString() const;
