[GLTFExporter] refactor: rename GetOrAdd* to AddUnique* in Builder

commit 9b3c70bb058f225d9ed28f8e97240fa4be36ed67
Date: Sun, 13 Mar 2022 19:35:17 +0100

    refactor: rename GetOrAdd* to AddUnique* in Builder
    (cherry picked from commit dafa59462c60381d601654dd64b979a8363d743f​)

#skipci

diff --git a/Engine/Plugins/Enterprise/GLTFExporter/MigratedFromGit/Source/GLTFExporter/Private/Builders/GLTFConvertBuilder.cpp b/Engine/Plugins/Enterprise/GLTFExporter/MigratedFromGit/Source/GLTFExporter/Private/Builders/GLTFConvertBuilder.cpp
--- a/Engine/Plugins/Enterprise/GLTFExporter/MigratedFromGit/Source/GLTFExporter/Private/Builders/GLTFConvertBuilder.cpp
+++ b/Engine/Plugins/Enterprise/GLTFExporter/MigratedFromGit/Source/GLTFExporter/Private/Builders/GLTFConvertBuilder.cpp
@@ -24,7 +24,7 @@ bool FGLTFConvertBuilder::IsRootActor(const AActor* Actor) const
 	return ParentActor == nullptr || !IsSelectedActor(ParentActor);
 }
 
-FGLTFJsonAccessor* FGLTFConvertBuilder::GetOrAddPositionAccessor(const FGLTFMeshSection* MeshSection, const FPositionVertexBuffer* VertexBuffer)
+FGLTFJsonAccessor* FGLTFConvertBuilder::AddUniquePositionAccessor(const FGLTFMeshSection* MeshSection, const FPositionVertexBuffer* VertexBuffer)
 {
 	if (VertexBuffer == nullptr)
 	{
@@ -34,7 +34,7 @@ FGLTFJsonAccessor* FGLTFConvertBuilder::GetOrAddPositionAccessor(const FGLTFMesh
 	return PositionBufferConverter->GetOrAdd(MeshSection, VertexBuffer);
 }
 
-FGLTFJsonAccessor* FGLTFConvertBuilder::GetOrAddColorAccessor(const FGLTFMeshSection* MeshSection, const FColorVertexBuffer* VertexBuffer)
+FGLTFJsonAccessor* FGLTFConvertBuilder::AddUniqueColorAccessor(const FGLTFMeshSection* MeshSection, const FColorVertexBuffer* VertexBuffer)
 {
 	if (VertexBuffer == nullptr)
 	{
@@ -44,7 +44,7 @@ FGLTFJsonAccessor* FGLTFConvertBuilder::GetOrAddColorAccessor(const FGLTFMeshSec
 	return ColorBufferConverter->GetOrAdd(MeshSection, VertexBuffer);
 }
 
-FGLTFJsonAccessor* FGLTFConvertBuilder::GetOrAddNormalAccessor(const FGLTFMeshSection* MeshSection, const FStaticMeshVertexBuffer* VertexBuffer)
+FGLTFJsonAccessor* FGLTFConvertBuilder::AddUniqueNormalAccessor(const FGLTFMeshSection* MeshSection, const FStaticMeshVertexBuffer* VertexBuffer)
 {
 	if (VertexBuffer == nullptr)
 	{
@@ -54,7 +54,7 @@ FGLTFJsonAccessor* FGLTFConvertBuilder::GetOrAddNormalAccessor(const FGLTFMeshSe
 	return NormalBufferConverter->GetOrAdd(MeshSection, VertexBuffer);
 }
 
-FGLTFJsonAccessor* FGLTFConvertBuilder::GetOrAddTangentAccessor(const FGLTFMeshSection* MeshSection, const FStaticMeshVertexBuffer* VertexBuffer)
+FGLTFJsonAccessor* FGLTFConvertBuilder::AddUniqueTangentAccessor(const FGLTFMeshSection* MeshSection, const FStaticMeshVertexBuffer* VertexBuffer)
 {
 	if (VertexBuffer == nullptr)
 	{
@@ -64,7 +64,7 @@ FGLTFJsonAccessor* FGLTFConvertBuilder::GetOrAddTangentAccessor(const FGLTFMeshS
 	return TangentBufferConverter->GetOrAdd(MeshSection, VertexBuffer);
 }
 
-FGLTFJsonAccessor* FGLTFConvertBuilder::GetOrAddUVAccessor(const FGLTFMeshSection* MeshSection, const FStaticMeshVertexBuffer* VertexBuffer, int32 UVIndex)
+FGLTFJsonAccessor* FGLTFConvertBuilder::AddUniqueUVAccessor(const FGLTFMeshSection* MeshSection, const FStaticMeshVertexBuffer* VertexBuffer, int32 UVIndex)
 {
 	if (VertexBuffer == nullptr)
 	{
@@ -74,7 +74,7 @@ FGLTFJsonAccessor* FGLTFConvertBuilder::GetOrAddUVAccessor(const FGLTFMeshSectio
 	return UVBufferConverter->GetOrAdd(MeshSection, VertexBuffer, UVIndex);
 }
 
-FGLTFJsonAccessor* FGLTFConvertBuilder::GetOrAddJointAccessor(const FGLTFMeshSection* MeshSection, const FSkinWeightVertexBuffer* VertexBuffer, int32 InfluenceOffset)
+FGLTFJsonAccessor* FGLTFConvertBuilder::AddUniqueJointAccessor(const FGLTFMeshSection* MeshSection, const FSkinWeightVertexBuffer* VertexBuffer, int32 InfluenceOffset)
 {
 	if (VertexBuffer == nullptr)
 	{
@@ -84,7 +84,7 @@ FGLTFJsonAccessor* FGLTFConvertBuilder::GetOrAddJointAccessor(const FGLTFMeshSec
 	return BoneIndexBufferConverter->GetOrAdd(MeshSection, VertexBuffer, InfluenceOffset);
 }
 
-FGLTFJsonAccessor* FGLTFConvertBuilder::GetOrAddWeightAccessor(const FGLTFMeshSection* MeshSection, const FSkinWeightVertexBuffer* VertexBuffer, int32 InfluenceOffset)
+FGLTFJsonAccessor* FGLTFConvertBuilder::AddUniqueWeightAccessor(const FGLTFMeshSection* MeshSection, const FSkinWeightVertexBuffer* VertexBuffer, int32 InfluenceOffset)
 {
 	if (VertexBuffer == nullptr)
 	{
@@ -94,7 +94,7 @@ FGLTFJsonAccessor* FGLTFConvertBuilder::GetOrAddWeightAccessor(const FGLTFMeshSe
 	return BoneWeightBufferConverter->GetOrAdd(MeshSection, VertexBuffer, InfluenceOffset);
 }
 
-FGLTFJsonAccessor* FGLTFConvertBuilder::GetOrAddIndexAccessor(const FGLTFMeshSection* MeshSection)
+FGLTFJsonAccessor* FGLTFConvertBuilder::AddUniqueIndexAccessor(const FGLTFMeshSection* MeshSection)
 {
 	if (MeshSection == nullptr)
 	{
@@ -104,7 +104,7 @@ FGLTFJsonAccessor* FGLTFConvertBuilder::GetOrAddIndexAccessor(const FGLTFMeshSec
 	return IndexBufferConverter->GetOrAdd(MeshSection);
 }
 
-FGLTFJsonMesh* FGLTFConvertBuilder::GetOrAddMesh(const UStaticMesh* StaticMesh, const FGLTFMaterialArray& Materials, int32 LODIndex)
+FGLTFJsonMesh* FGLTFConvertBuilder::AddUniqueMesh(const UStaticMesh* StaticMesh, const FGLTFMaterialArray& Materials, int32 LODIndex)
 {
 	if (StaticMesh == nullptr)
 	{
@@ -114,7 +114,7 @@ FGLTFJsonMesh* FGLTFConvertBuilder::GetOrAddMesh(const UStaticMesh* StaticMesh,
 	return StaticMeshConverter->GetOrAdd(StaticMesh, nullptr, Materials, LODIndex);
 }
 
-FGLTFJsonMesh* FGLTFConvertBuilder::GetOrAddMesh(const USkeletalMesh* SkeletalMesh, const FGLTFMaterialArray& Materials, int32 LODIndex)
+FGLTFJsonMesh* FGLTFConvertBuilder::AddUniqueMesh(const USkeletalMesh* SkeletalMesh, const FGLTFMaterialArray& Materials, int32 LODIndex)
 {
 	if (SkeletalMesh == nullptr)
 	{
@@ -124,22 +124,22 @@ FGLTFJsonMesh* FGLTFConvertBuilder::GetOrAddMesh(const USkeletalMesh* SkeletalMe
 	return SkeletalMeshConverter->GetOrAdd(SkeletalMesh, nullptr, Materials, LODIndex);
 }
 
-FGLTFJsonMesh* FGLTFConvertBuilder::GetOrAddMesh(const UMeshComponent* MeshComponent, const FGLTFMaterialArray& Materials, int32 LODIndex)
+FGLTFJsonMesh* FGLTFConvertBuilder::AddUniqueMesh(const UMeshComponent* MeshComponent, const FGLTFMaterialArray& Materials, int32 LODIndex)
 {
 	if (const UStaticMeshComponent* StaticMeshComponent = Cast<UStaticMeshComponent>(MeshComponent))
 	{
-		return GetOrAddMesh(StaticMeshComponent, Materials, LODIndex);
+		return AddUniqueMesh(StaticMeshComponent, Materials, LODIndex);
 	}
 
 	if (const USkeletalMeshComponent* SkeletalMeshComponent = Cast<USkeletalMeshComponent>(MeshComponent))
 	{
-		return GetOrAddMesh(SkeletalMeshComponent, Materials, LODIndex);
+		return AddUniqueMesh(SkeletalMeshComponent, Materials, LODIndex);
 	}
 
 	return nullptr;
 }
 
-FGLTFJsonMesh* FGLTFConvertBuilder::GetOrAddMesh(const UStaticMeshComponent* StaticMeshComponent, const FGLTFMaterialArray& Materials, int32 LODIndex)
+FGLTFJsonMesh* FGLTFConvertBuilder::AddUniqueMesh(const UStaticMeshComponent* StaticMeshComponent, const FGLTFMaterialArray& Materials, int32 LODIndex)
 {
 	if (StaticMeshComponent == nullptr)
 	{
@@ -155,7 +155,7 @@ FGLTFJsonMesh* FGLTFConvertBuilder::GetOrAddMesh(const UStaticMeshComponent* Sta
 	return StaticMeshConverter->GetOrAdd(StaticMesh, StaticMeshComponent, Materials, LODIndex);
 }
 
-FGLTFJsonMesh* FGLTFConvertBuilder::GetOrAddMesh(const USkeletalMeshComponent* SkeletalMeshComponent, const FGLTFMaterialArray& Materials, int32 LODIndex)
+FGLTFJsonMesh* FGLTFConvertBuilder::AddUniqueMesh(const USkeletalMeshComponent* SkeletalMeshComponent, const FGLTFMaterialArray& Materials, int32 LODIndex)
 {
 	if (SkeletalMeshComponent == nullptr)
 	{
@@ -171,66 +171,66 @@ FGLTFJsonMesh* FGLTFConvertBuilder::GetOrAddMesh(const USkeletalMeshComponent* S
 	return SkeletalMeshConverter->GetOrAdd(SkeletalMesh, SkeletalMeshComponent, Materials, LODIndex);
 }
 
-const FGLTFMeshData* FGLTFConvertBuilder::GetOrAddMeshData(const UStaticMesh* StaticMesh, const UStaticMeshComponent* StaticMeshComponent, int32 LODIndex)
+const FGLTFMeshData* FGLTFConvertBuilder::AddUniqueMeshData(const UStaticMesh* StaticMesh, const UStaticMeshComponent* StaticMeshComponent, int32 LODIndex)
 {
 	return StaticMeshDataConverter->GetOrAdd(StaticMesh, StaticMeshComponent, LODIndex);
 }
 
-const FGLTFMeshData* FGLTFConvertBuilder::GetOrAddMeshData(const USkeletalMesh* SkeletalMesh, const USkeletalMeshComponent* SkeletalMeshComponent, int32 LODIndex)
+const FGLTFMeshData* FGLTFConvertBuilder::AddUniqueMeshData(const USkeletalMesh* SkeletalMesh, const USkeletalMeshComponent* SkeletalMeshComponent, int32 LODIndex)
 {
 	return SkeletalMeshDataConverter->GetOrAdd(SkeletalMesh, SkeletalMeshComponent, LODIndex);
 }
 
-FGLTFJsonMaterial* FGLTFConvertBuilder::GetOrAddMaterial(const UMaterialInterface* Material, const UStaticMesh* StaticMesh, int32 LODIndex, int32 MaterialIndex)
+FGLTFJsonMaterial* FGLTFConvertBuilder::AddUniqueMaterial(const UMaterialInterface* Material, const UStaticMesh* StaticMesh, int32 LODIndex, int32 MaterialIndex)
 {
 	// TODO: optimize by skipping mesh data if material doesn't need it
-	const FGLTFMeshData* MeshData = GetOrAddMeshData(StaticMesh, nullptr, LODIndex);
+	const FGLTFMeshData* MeshData = AddUniqueMeshData(StaticMesh, nullptr, LODIndex);
 	const FGLTFIndexArray SectionIndices = FGLTFMeshUtility::GetSectionIndices(StaticMesh, MeshData->LODIndex, MaterialIndex);
-	return GetOrAddMaterial(Material, MeshData, SectionIndices);
+	return AddUniqueMaterial(Material, MeshData, SectionIndices);
 }
 
-FGLTFJsonMaterial* FGLTFConvertBuilder::GetOrAddMaterial(const UMaterialInterface* Material, const USkeletalMesh* SkeletalMesh, int32 LODIndex, int32 MaterialIndex)
+FGLTFJsonMaterial* FGLTFConvertBuilder::AddUniqueMaterial(const UMaterialInterface* Material, const USkeletalMesh* SkeletalMesh, int32 LODIndex, int32 MaterialIndex)
 {
 	// TODO: optimize by skipping mesh data if material doesn't need it
-	const FGLTFMeshData* MeshData = GetOrAddMeshData(SkeletalMesh, nullptr, LODIndex);
+	const FGLTFMeshData* MeshData = AddUniqueMeshData(SkeletalMesh, nullptr, LODIndex);
 	const FGLTFIndexArray SectionIndices = FGLTFMeshUtility::GetSectionIndices(SkeletalMesh, MeshData->LODIndex, MaterialIndex);
-	return GetOrAddMaterial(Material, MeshData, SectionIndices);
+	return AddUniqueMaterial(Material, MeshData, SectionIndices);
 }
 
-FGLTFJsonMaterial* FGLTFConvertBuilder::GetOrAddMaterial(const UMaterialInterface* Material, const UMeshComponent* MeshComponent, int32 LODIndex, int32 MaterialIndex)
+FGLTFJsonMaterial* FGLTFConvertBuilder::AddUniqueMaterial(const UMaterialInterface* Material, const UMeshComponent* MeshComponent, int32 LODIndex, int32 MaterialIndex)
 {
 	if (const UStaticMeshComponent* StaticMeshComponent = Cast<UStaticMeshComponent>(MeshComponent))
 	{
-		return GetOrAddMaterial(Material, StaticMeshComponent, LODIndex, MaterialIndex);
+		return AddUniqueMaterial(Material, StaticMeshComponent, LODIndex, MaterialIndex);
 	}
 
 	if (const USkeletalMeshComponent* SkeletalMeshComponent = Cast<USkeletalMeshComponent>(MeshComponent))
 	{
-		return GetOrAddMaterial(Material, SkeletalMeshComponent, LODIndex, MaterialIndex);
+		return AddUniqueMaterial(Material, SkeletalMeshComponent, LODIndex, MaterialIndex);
 	}
 
 	return nullptr;
 }
 
-FGLTFJsonMaterial* FGLTFConvertBuilder::GetOrAddMaterial(const UMaterialInterface* Material, const UStaticMeshComponent* StaticMeshComponent, int32 LODIndex, int32 MaterialIndex)
+FGLTFJsonMaterial* FGLTFConvertBuilder::AddUniqueMaterial(const UMaterialInterface* Material, const UStaticMeshComponent* StaticMeshComponent, int32 LODIndex, int32 MaterialIndex)
 {
 	// TODO: optimize by skipping mesh data if material doesn't need it
 	const UStaticMesh* StaticMesh = StaticMeshComponent->GetStaticMesh();
-	const FGLTFMeshData* MeshData = GetOrAddMeshData(StaticMesh, StaticMeshComponent, LODIndex);
+	const FGLTFMeshData* MeshData = AddUniqueMeshData(StaticMesh, StaticMeshComponent, LODIndex);
 	const FGLTFIndexArray SectionIndices = FGLTFMeshUtility::GetSectionIndices(StaticMesh, MeshData->LODIndex, MaterialIndex);
-	return GetOrAddMaterial(Material, MeshData, SectionIndices);
+	return AddUniqueMaterial(Material, MeshData, SectionIndices);
 }
 
-FGLTFJsonMaterial* FGLTFConvertBuilder::GetOrAddMaterial(const UMaterialInterface* Material, const USkeletalMeshComponent* SkeletalMeshComponent, int32 LODIndex, int32 MaterialIndex)
+FGLTFJsonMaterial* FGLTFConvertBuilder::AddUniqueMaterial(const UMaterialInterface* Material, const USkeletalMeshComponent* SkeletalMeshComponent, int32 LODIndex, int32 MaterialIndex)
 {
 	// TODO: optimize by skipping mesh data if material doesn't need it
 	const USkeletalMesh* SkeletalMesh = SkeletalMeshComponent->SkeletalMesh;
-	const FGLTFMeshData* MeshData = GetOrAddMeshData(SkeletalMesh, SkeletalMeshComponent, LODIndex);
+	const FGLTFMeshData* MeshData = AddUniqueMeshData(SkeletalMesh, SkeletalMeshComponent, LODIndex);
 	const FGLTFIndexArray SectionIndices = FGLTFMeshUtility::GetSectionIndices(SkeletalMesh, MeshData->LODIndex, MaterialIndex);
-	return GetOrAddMaterial(Material, MeshData, SectionIndices);
+	return AddUniqueMaterial(Material, MeshData, SectionIndices);
 }
 
-FGLTFJsonMaterial* FGLTFConvertBuilder::GetOrAddMaterial(const UMaterialInterface* Material, const FGLTFMeshData* MeshData, const FGLTFIndexArray& SectionIndices)
+FGLTFJsonMaterial* FGLTFConvertBuilder::AddUniqueMaterial(const UMaterialInterface* Material, const FGLTFMeshData* MeshData, const FGLTFIndexArray& SectionIndices)
 {
 	if (Material == nullptr)
 	{
@@ -240,7 +240,7 @@ FGLTFJsonMaterial* FGLTFConvertBuilder::GetOrAddMaterial(const UMaterialInterfac
 	return MaterialConverter->GetOrAdd(Material, MeshData, SectionIndices);
 }
 
-FGLTFJsonSampler* FGLTFConvertBuilder::GetOrAddSampler(const UTexture* Texture)
+FGLTFJsonSampler* FGLTFConvertBuilder::AddUniqueSampler(const UTexture* Texture)
 {
 	if (Texture == nullptr)
 	{
@@ -250,32 +250,32 @@ FGLTFJsonSampler* FGLTFConvertBuilder::GetOrAddSampler(const UTexture* Texture)
 	return SamplerConverter->GetOrAdd(Texture);
 }
 
-FGLTFJsonTexture* FGLTFConvertBuilder::GetOrAddTexture(const UTexture* Texture)
+FGLTFJsonTexture* FGLTFConvertBuilder::AddUniqueTexture(const UTexture* Texture)
 {
-	return GetOrAddTexture(Texture, Texture->SRGB);
+	return AddUniqueTexture(Texture, Texture->SRGB);
 }
 
-FGLTFJsonTexture* FGLTFConvertBuilder::GetOrAddTexture(const UTexture2D* Texture)
+FGLTFJsonTexture* FGLTFConvertBuilder::AddUniqueTexture(const UTexture2D* Texture)
 {
-	return GetOrAddTexture(Texture, Texture->SRGB);
+	return AddUniqueTexture(Texture, Texture->SRGB);
 }
 
-FGLTFJsonTexture* FGLTFConvertBuilder::GetOrAddTexture(const UTextureCube* Texture, ECubeFace CubeFace)
+FGLTFJsonTexture* FGLTFConvertBuilder::AddUniqueTexture(const UTextureCube* Texture, ECubeFace CubeFace)
 {
-	return GetOrAddTexture(Texture, CubeFace, Texture->SRGB);
+	return AddUniqueTexture(Texture, CubeFace, Texture->SRGB);
 }
 
-FGLTFJsonTexture* FGLTFConvertBuilder::GetOrAddTexture(const UTextureRenderTarget2D* Texture)
+FGLTFJsonTexture* FGLTFConvertBuilder::AddUniqueTexture(const UTextureRenderTarget2D* Texture)
 {
-	return GetOrAddTexture(Texture, Texture->SRGB);
+	return AddUniqueTexture(Texture, Texture->SRGB);
 }
 
-FGLTFJsonTexture* FGLTFConvertBuilder::GetOrAddTexture(const UTextureRenderTargetCube* Texture, ECubeFace CubeFace)
+FGLTFJsonTexture* FGLTFConvertBuilder::AddUniqueTexture(const UTextureRenderTargetCube* Texture, ECubeFace CubeFace)
 {
-	return GetOrAddTexture(Texture, CubeFace, Texture->SRGB);
+	return AddUniqueTexture(Texture, CubeFace, Texture->SRGB);
 }
 
-FGLTFJsonTexture* FGLTFConvertBuilder::GetOrAddTexture(const ULightMapTexture2D* Texture)
+FGLTFJsonTexture* FGLTFConvertBuilder::AddUniqueTexture(const ULightMapTexture2D* Texture)
 {
 	if (Texture == nullptr)
 	{
@@ -285,22 +285,22 @@ FGLTFJsonTexture* FGLTFConvertBuilder::GetOrAddTexture(const ULightMapTexture2D*
 	return TextureLightMapConverter->GetOrAdd(Texture);
 }
 
-FGLTFJsonTexture* FGLTFConvertBuilder::GetOrAddTexture(const UTexture* Texture, bool bToSRGB)
+FGLTFJsonTexture* FGLTFConvertBuilder::AddUniqueTexture(const UTexture* Texture, bool bToSRGB)
 {
 	if (const UTexture2D* Texture2D = Cast<UTexture2D>(Texture))
 	{
-		return GetOrAddTexture(Texture2D, bToSRGB);
+		return AddUniqueTexture(Texture2D, bToSRGB);
 	}
 
 	if (const UTextureRenderTarget2D* RenderTarget2D = Cast<UTextureRenderTarget2D>(Texture))
 	{
-		return GetOrAddTexture(RenderTarget2D, bToSRGB);
+		return AddUniqueTexture(RenderTarget2D, bToSRGB);
 	}
 
 	return nullptr;
 }
 
-FGLTFJsonTexture* FGLTFConvertBuilder::GetOrAddTexture(const UTexture2D* Texture, bool bToSRGB)
+FGLTFJsonTexture* FGLTFConvertBuilder::AddUniqueTexture(const UTexture2D* Texture, bool bToSRGB)
 {
 	if (Texture == nullptr)
 	{
@@ -310,7 +310,7 @@ FGLTFJsonTexture* FGLTFConvertBuilder::GetOrAddTexture(const UTexture2D* Texture
 	return Texture2DConverter->GetOrAdd(Texture, bToSRGB);
 }
 
-FGLTFJsonTexture* FGLTFConvertBuilder::GetOrAddTexture(const UTextureCube* Texture, ECubeFace CubeFace, bool bToSRGB)
+FGLTFJsonTexture* FGLTFConvertBuilder::AddUniqueTexture(const UTextureCube* Texture, ECubeFace CubeFace, bool bToSRGB)
 {
 	if (Texture == nullptr)
 	{
@@ -320,7 +320,7 @@ FGLTFJsonTexture* FGLTFConvertBuilder::GetOrAddTexture(const UTextureCube* Textu
 	return TextureCubeConverter->GetOrAdd(Texture, CubeFace, bToSRGB);
 }
 
-FGLTFJsonTexture* FGLTFConvertBuilder::GetOrAddTexture(const UTextureRenderTarget2D* Texture, bool bToSRGB)
+FGLTFJsonTexture* FGLTFConvertBuilder::AddUniqueTexture(const UTextureRenderTarget2D* Texture, bool bToSRGB)
 {
 	if (Texture == nullptr)
 	{
@@ -330,7 +330,7 @@ FGLTFJsonTexture* FGLTFConvertBuilder::GetOrAddTexture(const UTextureRenderTarge
 	return TextureRenderTarget2DConverter->GetOrAdd(Texture, bToSRGB);
 }
 
-FGLTFJsonTexture* FGLTFConvertBuilder::GetOrAddTexture(const UTextureRenderTargetCube* Texture, ECubeFace CubeFace, bool bToSRGB)
+FGLTFJsonTexture* FGLTFConvertBuilder::AddUniqueTexture(const UTextureRenderTargetCube* Texture, ECubeFace CubeFace, bool bToSRGB)
 {
 	if (Texture == nullptr)
 	{
@@ -340,12 +340,12 @@ FGLTFJsonTexture* FGLTFConvertBuilder::GetOrAddTexture(const UTextureRenderTarge
 	return TextureRenderTargetCubeConverter->GetOrAdd(Texture, CubeFace, bToSRGB);
 }
 
-FGLTFJsonImage* FGLTFConvertBuilder::GetOrAddImage(TGLTFSharedArray<FColor>& Pixels, FIntPoint Size, bool bIgnoreAlpha, EGLTFTextureType Type, const FString& Name)
+FGLTFJsonImage* FGLTFConvertBuilder::AddUniqueImage(TGLTFSharedArray<FColor>& Pixels, FIntPoint Size, bool bIgnoreAlpha, EGLTFTextureType Type, const FString& Name)
 {
 	return ImageConverter->GetOrAdd(Name, Type, bIgnoreAlpha, Size, Pixels);
 }
 
-FGLTFJsonSkin* FGLTFConvertBuilder::GetOrAddSkin(FGLTFJsonNode* RootNode, const USkeletalMesh* SkeletalMesh)
+FGLTFJsonSkin* FGLTFConvertBuilder::AddUniqueSkin(FGLTFJsonNode* RootNode, const USkeletalMesh* SkeletalMesh)
 {
 	if (RootNode == nullptr || SkeletalMesh == nullptr)
 	{
@@ -355,7 +355,7 @@ FGLTFJsonSkin* FGLTFConvertBuilder::GetOrAddSkin(FGLTFJsonNode* RootNode, const
 	return SkinConverter->GetOrAdd(RootNode, SkeletalMesh);
 }
 
-FGLTFJsonSkin* FGLTFConvertBuilder::GetOrAddSkin(FGLTFJsonNode* RootNode, const USkeletalMeshComponent* SkeletalMeshComponent)
+FGLTFJsonSkin* FGLTFConvertBuilder::AddUniqueSkin(FGLTFJsonNode* RootNode, const USkeletalMeshComponent* SkeletalMeshComponent)
 {
 	if (RootNode == nullptr || SkeletalMeshComponent == nullptr)
 	{
@@ -369,10 +369,10 @@ FGLTFJsonSkin* FGLTFConvertBuilder::GetOrAddSkin(FGLTFJsonNode* RootNode, const
 		return nullptr;
 	}
 
-	return GetOrAddSkin(RootNode, SkeletalMesh);
+	return AddUniqueSkin(RootNode, SkeletalMesh);
 }
 
-FGLTFJsonAnimation* FGLTFConvertBuilder::GetOrAddAnimation(FGLTFJsonNode* RootNode, const USkeletalMesh* SkeletalMesh, const UAnimSequence* AnimSequence)
+FGLTFJsonAnimation* FGLTFConvertBuilder::AddUniqueAnimation(FGLTFJsonNode* RootNode, const USkeletalMesh* SkeletalMesh, const UAnimSequence* AnimSequence)
 {
 	if (RootNode == nullptr || SkeletalMesh == nullptr || AnimSequence == nullptr)
 	{
@@ -382,7 +382,7 @@ FGLTFJsonAnimation* FGLTFConvertBuilder::GetOrAddAnimation(FGLTFJsonNode* RootNo
 	return AnimationConverter->GetOrAdd(RootNode, SkeletalMesh, AnimSequence);
 }
 
-FGLTFJsonAnimation* FGLTFConvertBuilder::GetOrAddAnimation(FGLTFJsonNode* RootNode, const USkeletalMeshComponent* SkeletalMeshComponent)
+FGLTFJsonAnimation* FGLTFConvertBuilder::AddUniqueAnimation(FGLTFJsonNode* RootNode, const USkeletalMeshComponent* SkeletalMeshComponent)
 {
 	if (RootNode == nullptr || SkeletalMeshComponent == nullptr)
 	{
@@ -392,7 +392,7 @@ FGLTFJsonAnimation* FGLTFConvertBuilder::GetOrAddAnimation(FGLTFJsonNode* RootNo
 	return AnimationDataConverter->GetOrAdd(RootNode, SkeletalMeshComponent);
 }
 
-FGLTFJsonAnimation* FGLTFConvertBuilder::GetOrAddAnimation(const ULevel* Level, const ULevelSequence* LevelSequence)
+FGLTFJsonAnimation* FGLTFConvertBuilder::AddUniqueAnimation(const ULevel* Level, const ULevelSequence* LevelSequence)
 {
 	if (Level == nullptr || LevelSequence == nullptr)
 	{
@@ -402,7 +402,7 @@ FGLTFJsonAnimation* FGLTFConvertBuilder::GetOrAddAnimation(const ULevel* Level,
 	return LevelSequenceConverter->GetOrAdd(Level, LevelSequence);
 }
 
-FGLTFJsonAnimation* FGLTFConvertBuilder::GetOrAddAnimation(const ALevelSequenceActor* LevelSequenceActor)
+FGLTFJsonAnimation* FGLTFConvertBuilder::AddUniqueAnimation(const ALevelSequenceActor* LevelSequenceActor)
 {
 	if (LevelSequenceActor == nullptr)
 	{
@@ -412,7 +412,7 @@ FGLTFJsonAnimation* FGLTFConvertBuilder::GetOrAddAnimation(const ALevelSequenceA
 	return LevelSequenceDataConverter->GetOrAdd(LevelSequenceActor);
 }
 
-FGLTFJsonNode* FGLTFConvertBuilder::GetOrAddNode(const AActor* Actor)
+FGLTFJsonNode* FGLTFConvertBuilder::AddUniqueNode(const AActor* Actor)
 {
 	if (Actor == nullptr)
 	{
@@ -422,7 +422,7 @@ FGLTFJsonNode* FGLTFConvertBuilder::GetOrAddNode(const AActor* Actor)
 	return ActorConverter->GetOrAdd(Actor);
 }
 
-FGLTFJsonNode* FGLTFConvertBuilder::GetOrAddNode(const USceneComponent* SceneComponent)
+FGLTFJsonNode* FGLTFConvertBuilder::AddUniqueNode(const USceneComponent* SceneComponent)
 {
 	if (SceneComponent == nullptr)
 	{
@@ -432,7 +432,7 @@ FGLTFJsonNode* FGLTFConvertBuilder::GetOrAddNode(const USceneComponent* SceneCom
 	return ComponentConverter->GetOrAdd(SceneComponent);
 }
 
-FGLTFJsonNode* FGLTFConvertBuilder::GetOrAddNode(const USceneComponent* SceneComponent, FName SocketName)
+FGLTFJsonNode* FGLTFConvertBuilder::AddUniqueNode(const USceneComponent* SceneComponent, FName SocketName)
 {
 	if (SceneComponent == nullptr)
 	{
@@ -442,7 +442,7 @@ FGLTFJsonNode* FGLTFConvertBuilder::GetOrAddNode(const USceneComponent* SceneCom
 	return ComponentSocketConverter->GetOrAdd(SceneComponent, SocketName);
 }
 
-FGLTFJsonNode* FGLTFConvertBuilder::GetOrAddNode(FGLTFJsonNode* RootNode, const UStaticMesh* StaticMesh, FName SocketName)
+FGLTFJsonNode* FGLTFConvertBuilder::AddUniqueNode(FGLTFJsonNode* RootNode, const UStaticMesh* StaticMesh, FName SocketName)
 {
 	if (RootNode == nullptr || StaticMesh == nullptr || SocketName == NAME_None)
 	{
@@ -452,7 +452,7 @@ FGLTFJsonNode* FGLTFConvertBuilder::GetOrAddNode(FGLTFJsonNode* RootNode, const
 	return StaticSocketConverter->GetOrAdd(RootNode, StaticMesh, SocketName);
 }
 
-FGLTFJsonNode* FGLTFConvertBuilder::GetOrAddNode(FGLTFJsonNode* RootNode, const USkeletalMesh* SkeletalMesh, FName SocketName)
+FGLTFJsonNode* FGLTFConvertBuilder::AddUniqueNode(FGLTFJsonNode* RootNode, const USkeletalMesh* SkeletalMesh, FName SocketName)
 {
 	if (RootNode == nullptr || SkeletalMesh == nullptr || SocketName == NAME_None)
 	{
@@ -462,7 +462,7 @@ FGLTFJsonNode* FGLTFConvertBuilder::GetOrAddNode(FGLTFJsonNode* RootNode, const
 	return SkeletalSocketConverter->GetOrAdd(RootNode, SkeletalMesh, SocketName);
 }
 
-FGLTFJsonNode* FGLTFConvertBuilder::GetOrAddNode(FGLTFJsonNode* RootNode, const USkeletalMesh* SkeletalMesh, int32 BoneIndex)
+FGLTFJsonNode* FGLTFConvertBuilder::AddUniqueNode(FGLTFJsonNode* RootNode, const USkeletalMesh* SkeletalMesh, int32 BoneIndex)
 {
 	if (RootNode == nullptr || SkeletalMesh == nullptr || BoneIndex == INDEX_NONE)
 	{
@@ -472,7 +472,7 @@ FGLTFJsonNode* FGLTFConvertBuilder::GetOrAddNode(FGLTFJsonNode* RootNode, const
 	return SkeletalBoneConverter->GetOrAdd(RootNode, SkeletalMesh, BoneIndex);
 }
 
-FGLTFJsonScene* FGLTFConvertBuilder::GetOrAddScene(const UWorld* World)
+FGLTFJsonScene* FGLTFConvertBuilder::AddUniqueScene(const UWorld* World)
 {
 	if (World == nullptr)
 	{
@@ -482,7 +482,7 @@ FGLTFJsonScene* FGLTFConvertBuilder::GetOrAddScene(const UWorld* World)
 	return SceneConverter->GetOrAdd(World);
 }
 
-FGLTFJsonCamera* FGLTFConvertBuilder::GetOrAddCamera(const UCameraComponent* CameraComponent)
+FGLTFJsonCamera* FGLTFConvertBuilder::AddUniqueCamera(const UCameraComponent* CameraComponent)
 {
 	if (CameraComponent == nullptr)
 	{
@@ -492,7 +492,7 @@ FGLTFJsonCamera* FGLTFConvertBuilder::GetOrAddCamera(const UCameraComponent* Cam
 	return CameraConverter->GetOrAdd(CameraComponent);
 }
 
-FGLTFJsonLight* FGLTFConvertBuilder::GetOrAddLight(const ULightComponent* LightComponent)
+FGLTFJsonLight* FGLTFConvertBuilder::AddUniqueLight(const ULightComponent* LightComponent)
 {
 	if (LightComponent == nullptr)
 	{
@@ -502,7 +502,7 @@ FGLTFJsonLight* FGLTFConvertBuilder::GetOrAddLight(const ULightComponent* LightC
 	return LightConverter->GetOrAdd(LightComponent);
 }
 
-FGLTFJsonBackdrop* FGLTFConvertBuilder::GetOrAddBackdrop(const AActor* BackdropActor)
+FGLTFJsonBackdrop* FGLTFConvertBuilder::AddUniqueBackdrop(const AActor* BackdropActor)
 {
 	if (BackdropActor == nullptr)
 	{
@@ -512,7 +512,7 @@ FGLTFJsonBackdrop* FGLTFConvertBuilder::GetOrAddBackdrop(const AActor* BackdropA
 	return BackdropConverter->GetOrAdd(BackdropActor);
 }
 
-FGLTFJsonLightMap* FGLTFConvertBuilder::GetOrAddLightMap(const UStaticMeshComponent* StaticMeshComponent)
+FGLTFJsonLightMap* FGLTFConvertBuilder::AddUniqueLightMap(const UStaticMeshComponent* StaticMeshComponent)
 {
 	if (StaticMeshComponent == nullptr)
 	{
@@ -522,7 +522,7 @@ FGLTFJsonLightMap* FGLTFConvertBuilder::GetOrAddLightMap(const UStaticMeshCompon
 	return LightMapConverter->GetOrAdd(StaticMeshComponent);
 }
 
-FGLTFJsonHotspot* FGLTFConvertBuilder::GetOrAddHotspot(const AGLTFHotspotActor* HotspotActor)
+FGLTFJsonHotspot* FGLTFConvertBuilder::AddUniqueHotspot(const AGLTFHotspotActor* HotspotActor)
 {
 	if (HotspotActor == nullptr)
 	{
@@ -532,7 +532,7 @@ FGLTFJsonHotspot* FGLTFConvertBuilder::GetOrAddHotspot(const AGLTFHotspotActor*
 	return HotspotConverter->GetOrAdd(HotspotActor);
 }
 
-FGLTFJsonSkySphere* FGLTFConvertBuilder::GetOrAddSkySphere(const AActor* SkySphereActor)
+FGLTFJsonSkySphere* FGLTFConvertBuilder::AddUniqueSkySphere(const AActor* SkySphereActor)
 {
 	if (SkySphereActor == nullptr)
 	{
@@ -542,7 +542,7 @@ FGLTFJsonSkySphere* FGLTFConvertBuilder::GetOrAddSkySphere(const AActor* SkySphe
 	return SkySphereConverter->GetOrAdd(SkySphereActor);
 }
 
-FGLTFJsonEpicLevelVariantSets* FGLTFConvertBuilder::GetOrAddEpicLevelVariantSets(const ULevelVariantSets* LevelVariantSets)
+FGLTFJsonEpicLevelVariantSets* FGLTFConvertBuilder::AddUniqueEpicLevelVariantSets(const ULevelVariantSets* LevelVariantSets)
 {
 	if (LevelVariantSets == nullptr)
 	{
@@ -552,7 +552,7 @@ FGLTFJsonEpicLevelVariantSets* FGLTFConvertBuilder::GetOrAddEpicLevelVariantSets
 	return EpicLevelVariantSetsConverter->GetOrAdd(LevelVariantSets);
 }
 
-FGLTFJsonKhrMaterialVariant* FGLTFConvertBuilder::GetOrAddKhrMaterialVariant(const UVariant* Variant)
+FGLTFJsonKhrMaterialVariant* FGLTFConvertBuilder::AddUniqueKhrMaterialVariant(const UVariant* Variant)
 {
 	if (Variant == nullptr)
 	{
diff --git a/Engine/Plugins/Enterprise/GLTFExporter/MigratedFromGit/Source/GLTFExporter/Private/Builders/GLTFConvertBuilder.h b/Engine/Plugins/Enterprise/GLTFExporter/MigratedFromGit/Source/GLTFExporter/Private/Builders/GLTFConvertBuilder.h
--- a/Engine/Plugins/Enterprise/GLTFExporter/MigratedFromGit/Source/GLTFExporter/Private/Builders/GLTFConvertBuilder.h
+++ b/Engine/Plugins/Enterprise/GLTFExporter/MigratedFromGit/Source/GLTFExporter/Private/Builders/GLTFConvertBuilder.h
@@ -34,69 +34,69 @@ public:
 	bool IsSelectedActor(const AActor* Object) const;
 	bool IsRootActor(const AActor* Actor) const;
 
-	FGLTFJsonAccessor* GetOrAddPositionAccessor(const FGLTFMeshSection* MeshSection, const FPositionVertexBuffer* VertexBuffer);
-	FGLTFJsonAccessor* GetOrAddColorAccessor(const FGLTFMeshSection* MeshSection, const FColorVertexBuffer* VertexBuffer);
-	FGLTFJsonAccessor* GetOrAddNormalAccessor(const FGLTFMeshSection* MeshSection, const FStaticMeshVertexBuffer* VertexBuffer);
-	FGLTFJsonAccessor* GetOrAddTangentAccessor(const FGLTFMeshSection* MeshSection, const FStaticMeshVertexBuffer* VertexBuffer);
-	FGLTFJsonAccessor* GetOrAddUVAccessor(const FGLTFMeshSection* MeshSection, const FStaticMeshVertexBuffer* VertexBuffer, int32 UVIndex);
-	FGLTFJsonAccessor* GetOrAddJointAccessor(const FGLTFMeshSection* MeshSection, const FSkinWeightVertexBuffer* VertexBuffer, int32 InfluenceOffset);
-	FGLTFJsonAccessor* GetOrAddWeightAccessor(const FGLTFMeshSection* MeshSection, const FSkinWeightVertexBuffer* VertexBuffer, int32 InfluenceOffset);
-	FGLTFJsonAccessor* GetOrAddIndexAccessor(const FGLTFMeshSection* MeshSection);
-
-	FGLTFJsonMesh* GetOrAddMesh(const UStaticMesh* StaticMesh, const FGLTFMaterialArray& Materials = {}, int32 LODIndex = INDEX_NONE);
-	FGLTFJsonMesh* GetOrAddMesh(const USkeletalMesh* SkeletalMesh, const FGLTFMaterialArray& Materials = {}, int32 LODIndex = INDEX_NONE);
-	FGLTFJsonMesh* GetOrAddMesh(const UMeshComponent* MeshComponent, const FGLTFMaterialArray& Materials = {}, int32 LODIndex = INDEX_NONE);
-	FGLTFJsonMesh* GetOrAddMesh(const UStaticMeshComponent* StaticMeshComponent, const FGLTFMaterialArray& Materials = {}, int32 LODIndex = INDEX_NONE);
-	FGLTFJsonMesh* GetOrAddMesh(const USkeletalMeshComponent* SkeletalMeshComponent, const FGLTFMaterialArray& Materials = {}, int32 LODIndex = INDEX_NONE);
-
-	const FGLTFMeshData* GetOrAddMeshData(const UStaticMesh* StaticMesh, const UStaticMeshComponent* StaticMeshComponent = nullptr, int32 LODIndex = INDEX_NONE);
-	const FGLTFMeshData* GetOrAddMeshData(const USkeletalMesh* SkeletalMesh, const USkeletalMeshComponent* SkeletalMeshComponent = nullptr, int32 LODIndex = INDEX_NONE);
-
-	FGLTFJsonMaterial* GetOrAddMaterial(const UMaterialInterface* Material, const UStaticMesh* StaticMesh, int32 LODIndex = INDEX_NONE, int32 MaterialIndex = INDEX_NONE);
-	FGLTFJsonMaterial* GetOrAddMaterial(const UMaterialInterface* Material, const USkeletalMesh* SkeletalMesh, int32 LODIndex = INDEX_NONE, int32 MaterialIndex = INDEX_NONE);
-	FGLTFJsonMaterial* GetOrAddMaterial(const UMaterialInterface* Material, const UMeshComponent* MeshComponent, int32 LODIndex = INDEX_NONE, int32 MaterialIndex = INDEX_NONE);
-	FGLTFJsonMaterial* GetOrAddMaterial(const UMaterialInterface* Material, const UStaticMeshComponent* StaticMeshComponent, int32 LODIndex = INDEX_NONE, int32 MaterialIndex = INDEX_NONE);
-	FGLTFJsonMaterial* GetOrAddMaterial(const UMaterialInterface* Material, const USkeletalMeshComponent* SkeletalMeshComponent, int32 LODIndex = INDEX_NONE, int32 MaterialIndex = INDEX_NONE);
-	FGLTFJsonMaterial* GetOrAddMaterial(const UMaterialInterface* Material, const FGLTFMeshData* MeshData = nullptr, const FGLTFIndexArray& SectionIndices = {});
-
-	FGLTFJsonSampler* GetOrAddSampler(const UTexture* Texture);
-	FGLTFJsonTexture* GetOrAddTexture(const UTexture* Texture);
-	FGLTFJsonTexture* GetOrAddTexture(const UTexture2D* Texture);
-	FGLTFJsonTexture* GetOrAddTexture(const UTextureCube* Texture, ECubeFace CubeFace);
-	FGLTFJsonTexture* GetOrAddTexture(const UTextureRenderTarget2D* Texture);
-	FGLTFJsonTexture* GetOrAddTexture(const UTextureRenderTargetCube* Texture, ECubeFace CubeFace);
-	FGLTFJsonTexture* GetOrAddTexture(const ULightMapTexture2D* Texture);
-	FGLTFJsonTexture* GetOrAddTexture(const UTexture* Texture, bool bToSRGB);
-	FGLTFJsonTexture* GetOrAddTexture(const UTexture2D* Texture, bool bToSRGB);
-	FGLTFJsonTexture* GetOrAddTexture(const UTextureCube* Texture, ECubeFace CubeFace, bool bToSRGB);
-	FGLTFJsonTexture* GetOrAddTexture(const UTextureRenderTarget2D* Texture, bool bToSRGB);
-	FGLTFJsonTexture* GetOrAddTexture(const UTextureRenderTargetCube* Texture, ECubeFace CubeFace, bool bToSRGB);
-	FGLTFJsonImage* GetOrAddImage(TGLTFSharedArray<FColor>& Pixels, FIntPoint Size, bool bIgnoreAlpha, EGLTFTextureType Type, const FString& Name);
-
-	FGLTFJsonSkin* GetOrAddSkin(FGLTFJsonNode* RootNode, const USkeletalMesh* SkeletalMesh);
-	FGLTFJsonSkin* GetOrAddSkin(FGLTFJsonNode* RootNode, const USkeletalMeshComponent* SkeletalMeshComponent);
-	FGLTFJsonAnimation* GetOrAddAnimation(FGLTFJsonNode* RootNode, const USkeletalMesh* SkeletalMesh, const UAnimSequence* AnimSequence);
-	FGLTFJsonAnimation* GetOrAddAnimation(FGLTFJsonNode* RootNode, const USkeletalMeshComponent* SkeletalMeshComponent);
-	FGLTFJsonAnimation* GetOrAddAnimation(const ULevel* Level, const ULevelSequence* LevelSequence);
-	FGLTFJsonAnimation* GetOrAddAnimation(const ALevelSequenceActor* LevelSequenceActor);
-
-	FGLTFJsonNode* GetOrAddNode(const AActor* Actor);
-	FGLTFJsonNode* GetOrAddNode(const USceneComponent* SceneComponent);
-	FGLTFJsonNode* GetOrAddNode(const USceneComponent* SceneComponent, FName SocketName);
-	FGLTFJsonNode* GetOrAddNode(FGLTFJsonNode* RootNode, const UStaticMesh* StaticMesh, FName SocketName);
-	FGLTFJsonNode* GetOrAddNode(FGLTFJsonNode* RootNode, const USkeletalMesh* SkeletalMesh, FName SocketName);
-	FGLTFJsonNode* GetOrAddNode(FGLTFJsonNode* RootNode, const USkeletalMesh* SkeletalMesh, int32 BoneIndex);
-	FGLTFJsonScene* GetOrAddScene(const UWorld* World);
-
-	FGLTFJsonCamera* GetOrAddCamera(const UCameraComponent* CameraComponent);
-	FGLTFJsonLight* GetOrAddLight(const ULightComponent* LightComponent);
-	FGLTFJsonBackdrop* GetOrAddBackdrop(const AActor* BackdropActor);
-	FGLTFJsonLightMap* GetOrAddLightMap(const UStaticMeshComponent* StaticMeshComponent);
-	FGLTFJsonHotspot* GetOrAddHotspot(const AGLTFHotspotActor* HotspotActor);
-	FGLTFJsonSkySphere* GetOrAddSkySphere(const AActor* SkySphereActor);
-
-	FGLTFJsonEpicLevelVariantSets* GetOrAddEpicLevelVariantSets(const ULevelVariantSets* LevelVariantSets);
-	FGLTFJsonKhrMaterialVariant* GetOrAddKhrMaterialVariant(const UVariant* Variant);
+	FGLTFJsonAccessor* AddUniquePositionAccessor(const FGLTFMeshSection* MeshSection, const FPositionVertexBuffer* VertexBuffer);
+	FGLTFJsonAccessor* AddUniqueColorAccessor(const FGLTFMeshSection* MeshSection, const FColorVertexBuffer* VertexBuffer);
+	FGLTFJsonAccessor* AddUniqueNormalAccessor(const FGLTFMeshSection* MeshSection, const FStaticMeshVertexBuffer* VertexBuffer);
+	FGLTFJsonAccessor* AddUniqueTangentAccessor(const FGLTFMeshSection* MeshSection, const FStaticMeshVertexBuffer* VertexBuffer);
+	FGLTFJsonAccessor* AddUniqueUVAccessor(const FGLTFMeshSection* MeshSection, const FStaticMeshVertexBuffer* VertexBuffer, int32 UVIndex);
+	FGLTFJsonAccessor* AddUniqueJointAccessor(const FGLTFMeshSection* MeshSection, const FSkinWeightVertexBuffer* VertexBuffer, int32 InfluenceOffset);
+	FGLTFJsonAccessor* AddUniqueWeightAccessor(const FGLTFMeshSection* MeshSection, const FSkinWeightVertexBuffer* VertexBuffer, int32 InfluenceOffset);
+	FGLTFJsonAccessor* AddUniqueIndexAccessor(const FGLTFMeshSection* MeshSection);
+
+	FGLTFJsonMesh* AddUniqueMesh(const UStaticMesh* StaticMesh, const FGLTFMaterialArray& Materials = {}, int32 LODIndex = INDEX_NONE);
+	FGLTFJsonMesh* AddUniqueMesh(const USkeletalMesh* SkeletalMesh, const FGLTFMaterialArray& Materials = {}, int32 LODIndex = INDEX_NONE);
+	FGLTFJsonMesh* AddUniqueMesh(const UMeshComponent* MeshComponent, const FGLTFMaterialArray& Materials = {}, int32 LODIndex = INDEX_NONE);
+	FGLTFJsonMesh* AddUniqueMesh(const UStaticMeshComponent* StaticMeshComponent, const FGLTFMaterialArray& Materials = {}, int32 LODIndex = INDEX_NONE);
+	FGLTFJsonMesh* AddUniqueMesh(const USkeletalMeshComponent* SkeletalMeshComponent, const FGLTFMaterialArray& Materials = {}, int32 LODIndex = INDEX_NONE);
+
+	const FGLTFMeshData* AddUniqueMeshData(const UStaticMesh* StaticMesh, const UStaticMeshComponent* StaticMeshComponent = nullptr, int32 LODIndex = INDEX_NONE);
+	const FGLTFMeshData* AddUniqueMeshData(const USkeletalMesh* SkeletalMesh, const USkeletalMeshComponent* SkeletalMeshComponent = nullptr, int32 LODIndex = INDEX_NONE);
+
+	FGLTFJsonMaterial* AddUniqueMaterial(const UMaterialInterface* Material, const UStaticMesh* StaticMesh, int32 LODIndex = INDEX_NONE, int32 MaterialIndex = INDEX_NONE);
+	FGLTFJsonMaterial* AddUniqueMaterial(const UMaterialInterface* Material, const USkeletalMesh* SkeletalMesh, int32 LODIndex = INDEX_NONE, int32 MaterialIndex = INDEX_NONE);
+	FGLTFJsonMaterial* AddUniqueMaterial(const UMaterialInterface* Material, const UMeshComponent* MeshComponent, int32 LODIndex = INDEX_NONE, int32 MaterialIndex = INDEX_NONE);
+	FGLTFJsonMaterial* AddUniqueMaterial(const UMaterialInterface* Material, const UStaticMeshComponent* StaticMeshComponent, int32 LODIndex = INDEX_NONE, int32 MaterialIndex = INDEX_NONE);
+	FGLTFJsonMaterial* AddUniqueMaterial(const UMaterialInterface* Material, const USkeletalMeshComponent* SkeletalMeshComponent, int32 LODIndex = INDEX_NONE, int32 MaterialIndex = INDEX_NONE);
+	FGLTFJsonMaterial* AddUniqueMaterial(const UMaterialInterface* Material, const FGLTFMeshData* MeshData = nullptr, const FGLTFIndexArray& SectionIndices = {});
+
+	FGLTFJsonSampler* AddUniqueSampler(const UTexture* Texture);
+	FGLTFJsonTexture* AddUniqueTexture(const UTexture* Texture);
+	FGLTFJsonTexture* AddUniqueTexture(const UTexture2D* Texture);
+	FGLTFJsonTexture* AddUniqueTexture(const UTextureCube* Texture, ECubeFace CubeFace);
+	FGLTFJsonTexture* AddUniqueTexture(const UTextureRenderTarget2D* Texture);
+	FGLTFJsonTexture* AddUniqueTexture(const UTextureRenderTargetCube* Texture, ECubeFace CubeFace);
+	FGLTFJsonTexture* AddUniqueTexture(const ULightMapTexture2D* Texture);
+	FGLTFJsonTexture* AddUniqueTexture(const UTexture* Texture, bool bToSRGB);
+	FGLTFJsonTexture* AddUniqueTexture(const UTexture2D* Texture, bool bToSRGB);
+	FGLTFJsonTexture* AddUniqueTexture(const UTextureCube* Texture, ECubeFace CubeFace, bool bToSRGB);
+	FGLTFJsonTexture* AddUniqueTexture(const UTextureRenderTarget2D* Texture, bool bToSRGB);
+	FGLTFJsonTexture* AddUniqueTexture(const UTextureRenderTargetCube* Texture, ECubeFace CubeFace, bool bToSRGB);
+	FGLTFJsonImage* AddUniqueImage(TGLTFSharedArray<FColor>& Pixels, FIntPoint Size, bool bIgnoreAlpha, EGLTFTextureType Type, const FString& Name);
+
+	FGLTFJsonSkin* AddUniqueSkin(FGLTFJsonNode* RootNode, const USkeletalMesh* SkeletalMesh);
+	FGLTFJsonSkin* AddUniqueSkin(FGLTFJsonNode* RootNode, const USkeletalMeshComponent* SkeletalMeshComponent);
+	FGLTFJsonAnimation* AddUniqueAnimation(FGLTFJsonNode* RootNode, const USkeletalMesh* SkeletalMesh, const UAnimSequence* AnimSequence);
+	FGLTFJsonAnimation* AddUniqueAnimation(FGLTFJsonNode* RootNode, const USkeletalMeshComponent* SkeletalMeshComponent);
+	FGLTFJsonAnimation* AddUniqueAnimation(const ULevel* Level, const ULevelSequence* LevelSequence);
+	FGLTFJsonAnimation* AddUniqueAnimation(const ALevelSequenceActor* LevelSequenceActor);
+
+	FGLTFJsonNode* AddUniqueNode(const AActor* Actor);
+	FGLTFJsonNode* AddUniqueNode(const USceneComponent* SceneComponent);
+	FGLTFJsonNode* AddUniqueNode(const USceneComponent* SceneComponent, FName SocketName);
+	FGLTFJsonNode* AddUniqueNode(FGLTFJsonNode* RootNode, const UStaticMesh* StaticMesh, FName SocketName);
+	FGLTFJsonNode* AddUniqueNode(FGLTFJsonNode* RootNode, const USkeletalMesh* SkeletalMesh, FName SocketName);
+	FGLTFJsonNode* AddUniqueNode(FGLTFJsonNode* RootNode, const USkeletalMesh* SkeletalMesh, int32 BoneIndex);
+	FGLTFJsonScene* AddUniqueScene(const UWorld* World);
+
+	FGLTFJsonCamera* AddUniqueCamera(const UCameraComponent* CameraComponent);
+	FGLTFJsonLight* AddUniqueLight(const ULightComponent* LightComponent);
+	FGLTFJsonBackdrop* AddUniqueBackdrop(const AActor* BackdropActor);
+	FGLTFJsonLightMap* AddUniqueLightMap(const UStaticMeshComponent* StaticMeshComponent);
+	FGLTFJsonHotspot* AddUniqueHotspot(const AGLTFHotspotActor* HotspotActor);
+	FGLTFJsonSkySphere* AddUniqueSkySphere(const AActor* SkySphereActor);
+
+	FGLTFJsonEpicLevelVariantSets* AddUniqueEpicLevelVariantSets(const ULevelVariantSets* LevelVariantSets);
+	FGLTFJsonKhrMaterialVariant* AddUniqueKhrMaterialVariant(const UVariant* Variant);
 
 	void RegisterObjectVariant(const UObject* Object, const UPropertyValue* Property);
 	const TArray<const UPropertyValue*>* GetObjectVariants(const UObject* Object) const;
diff --git a/Engine/Plugins/Enterprise/GLTFExporter/MigratedFromGit/Source/GLTFExporter/Private/Converters/GLTFAnimationConverters.cpp b/Engine/Plugins/Enterprise/GLTFExporter/MigratedFromGit/Source/GLTFExporter/Private/Converters/GLTFAnimationConverters.cpp
--- a/Engine/Plugins/Enterprise/GLTFExporter/MigratedFromGit/Source/GLTFExporter/Private/Converters/GLTFAnimationConverters.cpp
+++ b/Engine/Plugins/Enterprise/GLTFExporter/MigratedFromGit/Source/GLTFExporter/Private/Converters/GLTFAnimationConverters.cpp
@@ -47,7 +47,7 @@ FGLTFJsonAnimation* FGLTFAnimationDataConverter::Convert(FGLTFJsonNode* RootNode
 		return nullptr;
 	}
 
-	FGLTFJsonAnimation* JsonAnimation = Builder.GetOrAddAnimation(RootNode, SkeletalMesh, AnimSequence);
+	FGLTFJsonAnimation* JsonAnimation = Builder.AddUniqueAnimation(RootNode, SkeletalMesh, AnimSequence);
 	if (JsonAnimation != nullptr && Builder.ExportOptions->bExportPlaybackSettings)
 	{
 		FGLTFJsonAnimationPlayback& JsonPlayback = JsonAnimation->Playback;
@@ -78,7 +78,7 @@ FGLTFJsonAnimation* FGLTFLevelSequenceDataConverter::Convert(const ALevelSequenc
 		return nullptr;
 	}
 
-	FGLTFJsonAnimation* JsonAnimation = Builder.GetOrAddAnimation(Level, LevelSequence);
+	FGLTFJsonAnimation* JsonAnimation = Builder.AddUniqueAnimation(Level, LevelSequence);
 	if (JsonAnimation != nullptr && Builder.ExportOptions->bExportPlaybackSettings)
 	{
 		FGLTFJsonAnimationPlayback& JsonPlayback = JsonAnimation->Playback;
diff --git a/Engine/Plugins/Enterprise/GLTFExporter/MigratedFromGit/Source/GLTFExporter/Private/Converters/GLTFBackdropConverters.cpp b/Engine/Plugins/Enterprise/GLTFExporter/MigratedFromGit/Source/GLTFExporter/Private/Converters/GLTFBackdropConverters.cpp
--- a/Engine/Plugins/Enterprise/GLTFExporter/MigratedFromGit/Source/GLTFExporter/Private/Converters/GLTFBackdropConverters.cpp
+++ b/Engine/Plugins/Enterprise/GLTFExporter/MigratedFromGit/Source/GLTFExporter/Private/Converters/GLTFBackdropConverters.cpp
@@ -29,7 +29,7 @@ FGLTFJsonBackdrop* FGLTFBackdropConverter::Convert(const AActor* BackdropActor)
 	const UStaticMesh* Mesh;
 	if (FGLTFActorUtility::TryGetPropertyValue(BackdropActor, TEXT("Mesh"), Mesh))
 	{
-		JsonBackdrop->Mesh = Builder.GetOrAddMesh(Mesh, { FGLTFMaterialUtility::GetDefaultMaterial() });
+		JsonBackdrop->Mesh = Builder.AddUniqueMesh(Mesh, { FGLTFMaterialUtility::GetDefaultMaterial() });
 	}
 	else
 	{
@@ -45,7 +45,7 @@ FGLTFJsonBackdrop* FGLTFBackdropConverter::Convert(const AActor* BackdropActor)
 		{
 			const ECubeFace CubeFace = static_cast<ECubeFace>(CubeFaceIndex);
 			const EGLTFJsonCubeFace JsonCubeFace = FGLTFConverterUtility::ConvertCubeFace(CubeFace);
-			JsonBackdrop->Cubemap[static_cast<int32>(JsonCubeFace)] = Builder.GetOrAddTexture(Cubemap, CubeFace);
+			JsonBackdrop->Cubemap[static_cast<int32>(JsonCubeFace)] = Builder.AddUniqueTexture(Cubemap, CubeFace);
 		}
 	}
 	else
diff --git a/Engine/Plugins/Enterprise/GLTFExporter/MigratedFromGit/Source/GLTFExporter/Private/Converters/GLTFCameraConverters.cpp b/Engine/Plugins/Enterprise/GLTFExporter/MigratedFromGit/Source/GLTFExporter/Private/Converters/GLTFCameraConverters.cpp
--- a/Engine/Plugins/Enterprise/GLTFExporter/MigratedFromGit/Source/GLTFExporter/Private/Converters/GLTFCameraConverters.cpp
+++ b/Engine/Plugins/Enterprise/GLTFExporter/MigratedFromGit/Source/GLTFExporter/Private/Converters/GLTFCameraConverters.cpp
@@ -59,7 +59,7 @@ FGLTFJsonCamera* FGLTFCameraConverter::Convert(const UCameraComponent* CameraCom
 		{
 			FGLTFJsonCameraControl CameraControl;
 			CameraControl.Mode = FGLTFConverterUtility::ConvertCameraControlMode(CameraActor->Mode);
-			CameraControl.Target = Builder.GetOrAddNode(CameraActor->Target);
+			CameraControl.Target = Builder.AddUniqueNode(CameraActor->Target);
 			CameraControl.MaxDistance = FGLTFConverterUtility::ConvertLength(CameraActor->DistanceMax, ExportScale);
 			CameraControl.MinDistance = FGLTFConverterUtility::ConvertLength(CameraActor->DistanceMin, ExportScale);
 			CameraControl.MaxPitch = CameraActor->PitchAngleMax;
diff --git a/Engine/Plugins/Enterprise/GLTFExporter/MigratedFromGit/Source/GLTFExporter/Private/Converters/GLTFConverter.h b/Engine/Plugins/Enterprise/GLTFExporter/MigratedFromGit/Source/GLTFExporter/Private/Converters/GLTFConverter.h
--- a/Engine/Plugins/Enterprise/GLTFExporter/MigratedFromGit/Source/GLTFExporter/Private/Converters/GLTFConverter.h
+++ b/Engine/Plugins/Enterprise/GLTFExporter/MigratedFromGit/Source/GLTFExporter/Private/Converters/GLTFConverter.h
@@ -27,17 +27,6 @@ public:
 		return {};
 	}
 
-	OutputType Add(InputTypes... Inputs)
-	{
-		Sanitize(Inputs...);
-		InputTupleType InputTuple(MoveTemp(Inputs)...);
-		const uint32 InputHash = GetTypeHash(InputTuple);
-
-		OutputType NewOutput = InputTuple.ApplyAfter(&TGLTFConverter::Convert, this);
-		CachedOutputs.AddByHash(InputHash, MoveTemp(InputTuple), NewOutput);
-		return NewOutput;
-	}
-
 	OutputType GetOrAdd(InputTypes... Inputs)
 	{
 		Sanitize(Inputs...);
diff --git a/Engine/Plugins/Enterprise/GLTFExporter/MigratedFromGit/Source/GLTFExporter/Private/Converters/GLTFEpicLevelVariantSetsConverters.cpp b/Engine/Plugins/Enterprise/GLTFExporter/MigratedFromGit/Source/GLTFExporter/Private/Converters/GLTFEpicLevelVariantSetsConverters.cpp
--- a/Engine/Plugins/Enterprise/GLTFExporter/MigratedFromGit/Source/GLTFExporter/Private/Converters/GLTFEpicLevelVariantSetsConverters.cpp
+++ b/Engine/Plugins/Enterprise/GLTFExporter/MigratedFromGit/Source/GLTFExporter/Private/Converters/GLTFEpicLevelVariantSetsConverters.cpp
@@ -66,7 +66,7 @@ bool FGLTFEpicLevelVariantSetsConverter::TryParseVariant(FGLTFJsonEpicVariant& O
 	if (const UTexture2D* Thumbnail = const_cast<UVariant*>(Variant)->GetThumbnail())
 	{
 		// TODO: if thumbnail has generic name "Texture2D", give it a variant-relevant name
-		JsonVariant.Thumbnail = Builder.GetOrAddTexture(Thumbnail);
+		JsonVariant.Thumbnail = Builder.AddUniqueTexture(Thumbnail);
 	}
 
 	OutVariant = JsonVariant;
@@ -159,7 +159,7 @@ bool FGLTFEpicLevelVariantSetsConverter::TryParseVisibilityPropertyValue(FGLTFJs
 
 	Builder.RegisterObjectVariant(Owner, Property); // TODO: we should register this on the component
 
-	FGLTFJsonNode* Node = Builder.GetOrAddNode(Target);
+	FGLTFJsonNode* Node = Builder.AddUniqueNode(Target);
 	if (Target->IsA<UCameraComponent>() || Target->IsA<ULightComponent>())
 	{
 		// NOTE: If target component is a camera or light, we know that the target gltf node is actually the first child.
@@ -229,10 +229,10 @@ bool FGLTFEpicLevelVariantSetsConverter::TryParseMaterialPropertyValue(FGLTFJson
 	Builder.RegisterObjectVariant(Owner, Property); // TODO: we should register this on the component
 
 	FGLTFJsonEpicVariantMaterial VariantMaterial;
-	VariantMaterial.Material = FGLTFVariantUtility::GetOrAddMaterial(Builder, Material, Target, MaterialIndex);
+	VariantMaterial.Material = FGLTFVariantUtility::AddUniqueMaterial(Builder, Material, Target, MaterialIndex);
 	VariantMaterial.Index = MaterialIndex;
 
-	FGLTFJsonNode* Node = Builder.GetOrAddNode(Target);
+	FGLTFJsonNode* Node = Builder.AddUniqueNode(Target);
 	FGLTFJsonEpicVariantNodeProperties& NodeProperties = OutVariant.Nodes.FindOrAdd(Node);
 
 	NodeProperties.Node = Node;
@@ -282,9 +282,9 @@ bool FGLTFEpicLevelVariantSetsConverter::TryParseMeshPropertyValue(FGLTFJsonEpic
 	Builder.RegisterObjectVariant(Owner, Property); // TODO: we should register this on the component
 
 	const FGLTFMaterialArray OverrideMaterials(Target->OverrideMaterials);
-	FGLTFJsonNode* JsonNode = Builder.GetOrAddNode(Target);
+	FGLTFJsonNode* JsonNode = Builder.AddUniqueNode(Target);
 	FGLTFJsonEpicVariantNodeProperties& NodeProperties = OutVariant.Nodes.FindOrAdd(JsonNode);
-	FGLTFJsonMesh* JsonMesh = Builder.GetOrAddMesh(Mesh, OverrideMaterials);
+	FGLTFJsonMesh* JsonMesh = Builder.AddUniqueMesh(Mesh, OverrideMaterials);
 
 	NodeProperties.Node = JsonNode;
 	NodeProperties.Mesh = JsonMesh;
diff --git a/Engine/Plugins/Enterprise/GLTFExporter/MigratedFromGit/Source/GLTFExporter/Private/Converters/GLTFHotspotConverters.cpp b/Engine/Plugins/Enterprise/GLTFExporter/MigratedFromGit/Source/GLTFExporter/Private/Converters/GLTFHotspotConverters.cpp
--- a/Engine/Plugins/Enterprise/GLTFExporter/MigratedFromGit/Source/GLTFExporter/Private/Converters/GLTFHotspotConverters.cpp
+++ b/Engine/Plugins/Enterprise/GLTFExporter/MigratedFromGit/Source/GLTFExporter/Private/Converters/GLTFHotspotConverters.cpp
@@ -25,14 +25,14 @@ FGLTFJsonHotspot* FGLTFHotspotConverter::Convert(const AGLTFHotspotActor* Hotspo
 		}
 		else
 		{
-			FGLTFJsonNode* RootNode = Builder.GetOrAddNode(SkeletalMeshActor);
+			FGLTFJsonNode* RootNode = Builder.AddUniqueNode(SkeletalMeshActor);
 
 			const USkeletalMeshComponent* SkeletalMeshComponent = SkeletalMeshActor->GetSkeletalMeshComponent();
 			if (const USkeletalMesh* SkeletalMesh = SkeletalMeshComponent->SkeletalMesh)
 			{
 				if (const UAnimSequence* AnimSequence = HotspotActor->AnimationSequence)
 				{
-					JsonHotspot->Animation = Builder.GetOrAddAnimation(RootNode, SkeletalMesh, AnimSequence);
+					JsonHotspot->Animation = Builder.AddUniqueAnimation(RootNode, SkeletalMesh, AnimSequence);
 				}
 				else
 				{
@@ -55,7 +55,7 @@ FGLTFJsonHotspot* FGLTFHotspotConverter::Convert(const AGLTFHotspotActor* Hotspo
 		}
 		else
 		{
-			JsonHotspot->Animation = Builder.GetOrAddAnimation(HotspotActor->GetLevel(), LevelSequence);
+			JsonHotspot->Animation = Builder.AddUniqueAnimation(HotspotActor->GetLevel(), LevelSequence);
 		}
 	}
 	else
@@ -63,10 +63,10 @@ FGLTFJsonHotspot* FGLTFHotspotConverter::Convert(const AGLTFHotspotActor* Hotspo
 		// TODO: report warning
 	}
 
-	JsonHotspot->Image = Builder.GetOrAddTexture(HotspotActor->GetImageForState(EGLTFHotspotState::Default));
-	JsonHotspot->HoveredImage = Builder.GetOrAddTexture(HotspotActor->GetImageForState(EGLTFHotspotState::Hovered));
-	JsonHotspot->ToggledImage = Builder.GetOrAddTexture(HotspotActor->GetImageForState(EGLTFHotspotState::Toggled));
-	JsonHotspot->ToggledHoveredImage = Builder.GetOrAddTexture(HotspotActor->GetImageForState(EGLTFHotspotState::ToggledHovered));
+	JsonHotspot->Image = Builder.AddUniqueTexture(HotspotActor->GetImageForState(EGLTFHotspotState::Default));
+	JsonHotspot->HoveredImage = Builder.AddUniqueTexture(HotspotActor->GetImageForState(EGLTFHotspotState::Hovered));
+	JsonHotspot->ToggledImage = Builder.AddUniqueTexture(HotspotActor->GetImageForState(EGLTFHotspotState::Toggled));
+	JsonHotspot->ToggledHoveredImage = Builder.AddUniqueTexture(HotspotActor->GetImageForState(EGLTFHotspotState::ToggledHovered));
 
 	return JsonHotspot;
 }
diff --git a/Engine/Plugins/Enterprise/GLTFExporter/MigratedFromGit/Source/GLTFExporter/Private/Converters/GLTFKhrVariantConverters.cpp b/Engine/Plugins/Enterprise/GLTFExporter/MigratedFromGit/Source/GLTFExporter/Private/Converters/GLTFKhrVariantConverters.cpp
--- a/Engine/Plugins/Enterprise/GLTFExporter/MigratedFromGit/Source/GLTFExporter/Private/Converters/GLTFKhrVariantConverters.cpp
+++ b/Engine/Plugins/Enterprise/GLTFExporter/MigratedFromGit/Source/GLTFExporter/Private/Converters/GLTFKhrVariantConverters.cpp
@@ -141,10 +141,10 @@ bool FGLTFKhrMaterialVariantConverter::TryParseMaterialProperty(FGLTFJsonPrimiti
 	Builder.RegisterObjectVariant(Owner, Property); // TODO: we don't need to register this on the actor
 
 	const int32 MaterialIndex = CapturedPropSegments[NumPropSegments - 1].PropertyIndex;
-	FGLTFJsonMesh* JsonMesh = Builder.GetOrAddMesh(Target);
+	FGLTFJsonMesh* JsonMesh = Builder.AddUniqueMesh(Target);
 
 	OutPrimitive = &JsonMesh->Primitives[MaterialIndex];
-	OutMaterial = FGLTFVariantUtility::GetOrAddMaterial(Builder, Material, Target, MaterialIndex);
+	OutMaterial = FGLTFVariantUtility::AddUniqueMaterial(Builder, Material, Target, MaterialIndex);
 
 	return true;
 }
diff --git a/Engine/Plugins/Enterprise/GLTFExporter/MigratedFromGit/Source/GLTFExporter/Private/Converters/GLTFLightMapConverters.cpp b/Engine/Plugins/Enterprise/GLTFExporter/MigratedFromGit/Source/GLTFExporter/Private/Converters/GLTFLightMapConverters.cpp
--- a/Engine/Plugins/Enterprise/GLTFExporter/MigratedFromGit/Source/GLTFExporter/Private/Converters/GLTFLightMapConverters.cpp
+++ b/Engine/Plugins/Enterprise/GLTFExporter/MigratedFromGit/Source/GLTFExporter/Private/Converters/GLTFLightMapConverters.cpp
@@ -59,7 +59,7 @@ FGLTFJsonLightMap* FGLTFLightMapConverter::Convert(const UStaticMeshComponent* S
 
 	const FLightMapInteraction LightMapInteraction = LightMap2D->GetInteraction(GMaxRHIFeatureLevel);
 	const ULightMapTexture2D* Texture = LightMapInteraction.GetTexture(true);
-	FGLTFJsonTexture* JsonTexture = Builder.GetOrAddTexture(Texture);
+	FGLTFJsonTexture* JsonTexture = Builder.AddUniqueTexture(Texture);
 
 	if (JsonTexture == nullptr)
 	{
diff --git a/Engine/Plugins/Enterprise/GLTFExporter/MigratedFromGit/Source/GLTFExporter/Private/Converters/GLTFMaterialProxyFactory.cpp b/Engine/Plugins/Enterprise/GLTFExporter/MigratedFromGit/Source/GLTFExporter/Private/Converters/GLTFMaterialProxyFactory.cpp
--- a/Engine/Plugins/Enterprise/GLTFExporter/MigratedFromGit/Source/GLTFExporter/Private/Converters/GLTFMaterialProxyFactory.cpp
+++ b/Engine/Plugins/Enterprise/GLTFExporter/MigratedFromGit/Source/GLTFExporter/Private/Converters/GLTFMaterialProxyFactory.cpp
@@ -32,7 +32,7 @@ UMaterialInterface* FGLTFMaterialProxyFactory::Create(UMaterialInterface* Origin
 			*OriginalMaterial->GetName()));
 	}
 
-	FGLTFJsonMaterial* JsonMaterial = Builder.GetOrAddMaterial(OriginalMaterial);
+	FGLTFJsonMaterial* JsonMaterial = Builder.AddUniqueMaterial(OriginalMaterial);
 	if (JsonMaterial == nullptr)
 	{
 		// TODO: report error
diff --git a/Engine/Plugins/Enterprise/GLTFExporter/MigratedFromGit/Source/GLTFExporter/Private/Converters/GLTFMaterialUtility.cpp b/Engine/Plugins/Enterprise/GLTFExporter/MigratedFromGit/Source/GLTFExporter/Private/Converters/GLTFMaterialUtility.cpp
--- a/Engine/Plugins/Enterprise/GLTFExporter/MigratedFromGit/Source/GLTFExporter/Private/Converters/GLTFMaterialUtility.cpp
+++ b/Engine/Plugins/Enterprise/GLTFExporter/MigratedFromGit/Source/GLTFExporter/Private/Converters/GLTFMaterialUtility.cpp
@@ -276,7 +276,7 @@ FGLTFJsonTexture* FGLTFMaterialUtility::AddTexture(FGLTFConvertBuilder& Builder,
 	FGLTFJsonTexture* JsonTexture = Builder.AddTexture();
 	JsonTexture->Name = TextureName;
 	JsonTexture->Sampler = JsonSampler;
-	JsonTexture->Source = Builder.GetOrAddImage(Pixels, TextureSize, bIgnoreAlpha, bIsNormalMap ? EGLTFTextureType::Normalmaps : EGLTFTextureType::None, TextureName);
+	JsonTexture->Source = Builder.AddUniqueImage(Pixels, TextureSize, bIgnoreAlpha, bIsNormalMap ? EGLTFTextureType::Normalmaps : EGLTFTextureType::None, TextureName);
 
 	return JsonTexture;
 }
diff --git a/Engine/Plugins/Enterprise/GLTFExporter/MigratedFromGit/Source/GLTFExporter/Private/Converters/GLTFNodeConverters.cpp b/Engine/Plugins/Enterprise/GLTFExporter/MigratedFromGit/Source/GLTFExporter/Private/Converters/GLTFNodeConverters.cpp
--- a/Engine/Plugins/Enterprise/GLTFExporter/MigratedFromGit/Source/GLTFExporter/Private/Converters/GLTFNodeConverters.cpp
+++ b/Engine/Plugins/Enterprise/GLTFExporter/MigratedFromGit/Source/GLTFExporter/Private/Converters/GLTFNodeConverters.cpp
@@ -25,7 +25,7 @@ FGLTFJsonNode* FGLTFActorConverter::Convert(const AActor* Actor)
 	}
 
 	const USceneComponent* RootComponent = Actor->GetRootComponent();
-	FGLTFJsonNode* RootNode = Builder.GetOrAddNode(RootComponent);
+	FGLTFJsonNode* RootNode = Builder.AddUniqueNode(RootComponent);
 
 	// TODO: process all components since any component can be attached to any other component in runtime
 
@@ -34,28 +34,28 @@ FGLTFJsonNode* FGLTFActorConverter::Convert(const AActor* Actor)
 	{
 		if (Builder.ExportOptions->bExportSkySpheres)
 		{
-			RootNode->SkySphere = Builder.GetOrAddSkySphere(Actor);
+			RootNode->SkySphere = Builder.AddUniqueSkySphere(Actor);
 		}
 	}
 	else if (FGLTFActorUtility::IsHDRIBackdropBlueprint(BlueprintPath))
 	{
 		if (Builder.ExportOptions->bExportHDRIBackdrops)
 		{
-			RootNode->Backdrop = Builder.GetOrAddBackdrop(Actor);
+			RootNode->Backdrop = Builder.AddUniqueBackdrop(Actor);
 		}
 	}
 	else if (const ALevelSequenceActor* LevelSequenceActor = Cast<ALevelSequenceActor>(Actor))
 	{
 		if (Builder.ExportOptions->bExportLevelSequences)
 		{
-			Builder.GetOrAddAnimation(LevelSequenceActor);
+			Builder.AddUniqueAnimation(LevelSequenceActor);
 		}
 	}
 	else if (const AGLTFHotspotActor* HotspotActor = Cast<AGLTFHotspotActor>(Actor))
 	{
 		if (Builder.ExportOptions->bExportAnimationHotspots)
 		{
-			RootNode->Hotspot = Builder.GetOrAddHotspot(HotspotActor);
+			RootNode->Hotspot = Builder.AddUniqueHotspot(HotspotActor);
 		}
 	}
 	else
@@ -69,7 +69,7 @@ FGLTFJsonNode* FGLTFActorConverter::Convert(const AActor* Actor)
 			const USceneComponent* SceneComponent = Cast<USceneComponent>(Component);
 			if (SceneComponent != nullptr)
 			{
-				Builder.GetOrAddNode(SceneComponent);
+				Builder.AddUniqueNode(SceneComponent);
 			}
 		}
 	}
@@ -101,7 +101,7 @@ FGLTFJsonNode* FGLTFComponentConverter::Convert(const USceneComponent* SceneComp
 
 	const USceneComponent* ParentComponent = !bIsRootNode ? SceneComponent->GetAttachParent() : nullptr;
 	const FName SocketName = SceneComponent->GetAttachSocketName();
-	FGLTFJsonNode* ParentNode = Builder.GetOrAddNode(ParentComponent, SocketName);
+	FGLTFJsonNode* ParentNode = Builder.AddUniqueNode(ParentComponent, SocketName);
 
 	if (ParentComponent != nullptr && !SceneComponent->IsUsingAbsoluteScale())
 	{
@@ -138,25 +138,25 @@ FGLTFJsonNode* FGLTFComponentConverter::Convert(const USceneComponent* SceneComp
 	{
 		if (const UStaticMeshComponent* StaticMeshComponent = Cast<UStaticMeshComponent>(SceneComponent))
 		{
-			Node->Mesh = Builder.GetOrAddMesh(StaticMeshComponent);
+			Node->Mesh = Builder.AddUniqueMesh(StaticMeshComponent);
 
 			if (Builder.ExportOptions->bExportLightmaps)
 			{
-				Node->LightMap = Builder.GetOrAddLightMap(StaticMeshComponent);
+				Node->LightMap = Builder.AddUniqueLightMap(StaticMeshComponent);
 			}
 		}
 		else if (const USkeletalMeshComponent* SkeletalMeshComponent = Cast<USkeletalMeshComponent>(SceneComponent))
 		{
-			Node->Mesh = Builder.GetOrAddMesh(SkeletalMeshComponent);
+			Node->Mesh = Builder.AddUniqueMesh(SkeletalMeshComponent);
 
 			if (Builder.ExportOptions->bExportVertexSkinWeights)
 			{
-				Node->Skin = Builder.GetOrAddSkin(Node, SkeletalMeshComponent);
+				Node->Skin = Builder.AddUniqueSkin(Node, SkeletalMeshComponent);
 				if (Node->Skin != nullptr)
 				{
 					if (Builder.ExportOptions->bExportAnimationSequences)
 					{
-						Builder.GetOrAddAnimation(Node, SkeletalMeshComponent);
+						Builder.AddUniqueAnimation(Node, SkeletalMeshComponent);
 					}
 				}
 			}
@@ -169,7 +169,7 @@ FGLTFJsonNode* FGLTFComponentConverter::Convert(const USceneComponent* SceneComp
 				FGLTFJsonNode* CameraNode = Builder.AddNode();
 				CameraNode->Name = FGLTFNameUtility::GetName(CameraComponent);
 				CameraNode->Rotation = FGLTFConverterUtility::ConvertCameraDirection();
-				CameraNode->Camera = Builder.GetOrAddCamera(CameraComponent);
+				CameraNode->Camera = Builder.AddUniqueCamera(CameraComponent);
 				Node->Children.Add(CameraNode);
 			}
 		}
@@ -181,7 +181,7 @@ FGLTFJsonNode* FGLTFComponentConverter::Convert(const USceneComponent* SceneComp
 				FGLTFJsonNode* LightNode = Builder.AddNode();
 				LightNode->Name = FGLTFNameUtility::GetName(LightComponent);
 				LightNode->Rotation = FGLTFConverterUtility::ConvertLightDirection();
-				LightNode->Light = Builder.GetOrAddLight(LightComponent);
+				LightNode->Light = Builder.AddUniqueLight(LightComponent);
 				Node->Children.Add(LightNode);
 			}
 		}
@@ -192,21 +192,21 @@ FGLTFJsonNode* FGLTFComponentConverter::Convert(const USceneComponent* SceneComp
 
 FGLTFJsonNode* FGLTFComponentSocketConverter::Convert(const USceneComponent* SceneComponent, FName SocketName)
 {
-	FGLTFJsonNode* JsonNode = Builder.GetOrAddNode(SceneComponent);
+	FGLTFJsonNode* JsonNode = Builder.AddUniqueNode(SceneComponent);
 
 	if (SocketName != NAME_None)
 	{
 		if (const UStaticMeshComponent* StaticMeshComponent = Cast<UStaticMeshComponent>(SceneComponent))
 		{
 			const UStaticMesh* StaticMesh = StaticMeshComponent->GetStaticMesh();
-			return Builder.GetOrAddNode(JsonNode, StaticMesh, SocketName);
+			return Builder.AddUniqueNode(JsonNode, StaticMesh, SocketName);
 		}
 
 		if (const USkinnedMeshComponent* SkinnedMeshComponent = Cast<USkinnedMeshComponent>(SceneComponent))
 		{
 			// TODO: add support for SocketOverrideLookup?
 			const USkeletalMesh* SkeletalMesh = SkinnedMeshComponent->SkeletalMesh;
-			return Builder.GetOrAddNode(JsonNode, SkeletalMesh, SocketName);
+			return Builder.AddUniqueNode(JsonNode, SkeletalMesh, SocketName);
 		}
 
 		// TODO: add support for more socket types
@@ -261,7 +261,7 @@ FGLTFJsonNode* FGLTFSkeletalSocketConverter::Convert(FGLTFJsonNode* RootNode, co
 		Node->Scale = FGLTFConverterUtility::ConvertScale(Socket->RelativeScale);
 
 		const int32 ParentBone = RefSkeleton.FindBoneIndex(Socket->BoneName);
-		FGLTFJsonNode* ParentNode = ParentBone != INDEX_NONE ? Builder.GetOrAddNode(RootNode, SkeletalMesh, ParentBone) : RootNode;
+		FGLTFJsonNode* ParentNode = ParentBone != INDEX_NONE ? Builder.AddUniqueNode(RootNode, SkeletalMesh, ParentBone) : RootNode;
 
 		ParentNode->Children.Add(Node);
 		return Node;
@@ -270,7 +270,7 @@ FGLTFJsonNode* FGLTFSkeletalSocketConverter::Convert(FGLTFJsonNode* RootNode, co
 	const int32 BoneIndex = RefSkeleton.FindBoneIndex(SocketName);
 	if (BoneIndex != INDEX_NONE)
 	{
-		return Builder.GetOrAddNode(RootNode, SkeletalMesh, BoneIndex);
+		return Builder.AddUniqueNode(RootNode, SkeletalMesh, BoneIndex);
 	}
 
 	// TODO: report error
@@ -313,7 +313,7 @@ FGLTFJsonNode* FGLTFSkeletalBoneConverter::Convert(FGLTFJsonNode* RootNode, cons
 		// TODO: report error
 	}
 
-	FGLTFJsonNode* ParentNode = BoneInfo.ParentIndex != INDEX_NONE ? Builder.GetOrAddNode(RootNode, SkeletalMesh, BoneInfo.ParentIndex) : RootNode;
+	FGLTFJsonNode* ParentNode = BoneInfo.ParentIndex != INDEX_NONE ? Builder.AddUniqueNode(RootNode, SkeletalMesh, BoneInfo.ParentIndex) : RootNode;
 	ParentNode->Children.Add(Node);
 	return Node;
 }
diff --git a/Engine/Plugins/Enterprise/GLTFExporter/MigratedFromGit/Source/GLTFExporter/Private/Converters/GLTFSceneConverters.cpp b/Engine/Plugins/Enterprise/GLTFExporter/MigratedFromGit/Source/GLTFExporter/Private/Converters/GLTFSceneConverters.cpp
--- a/Engine/Plugins/Enterprise/GLTFExporter/MigratedFromGit/Source/GLTFExporter/Private/Converters/GLTFSceneConverters.cpp
+++ b/Engine/Plugins/Enterprise/GLTFExporter/MigratedFromGit/Source/GLTFExporter/Private/Converters/GLTFSceneConverters.cpp
@@ -36,7 +36,7 @@ FGLTFJsonScene* FGLTFSceneConverter::Convert(const UWorld* World)
 						{
 							if (LevelVariantSets != nullptr)
 							{
-								FGLTFJsonEpicLevelVariantSets* JsonEpicLevelVariantSets = Builder.GetOrAddEpicLevelVariantSets(LevelVariantSets);
+								FGLTFJsonEpicLevelVariantSets* JsonEpicLevelVariantSets = Builder.AddUniqueEpicLevelVariantSets(LevelVariantSets);
 								if (JsonEpicLevelVariantSets != nullptr)
 								{
 									Scene->EpicLevelVariantSets.Add(JsonEpicLevelVariantSets);
@@ -51,7 +51,7 @@ FGLTFJsonScene* FGLTFSceneConverter::Convert(const UWorld* World)
 								{
 									for (const UVariant* Variant: VariantSet->GetVariants())
 									{
-										Builder.GetOrAddKhrMaterialVariant(Variant);
+										Builder.AddUniqueKhrMaterialVariant(Variant);
 									}
 								}
 							}
@@ -62,7 +62,7 @@ FGLTFJsonScene* FGLTFSceneConverter::Convert(const UWorld* World)
 
 			for (const AActor* Actor : Level->Actors)
 			{
-				FGLTFJsonNode* JsonNode = Builder.GetOrAddNode(Actor);
+				FGLTFJsonNode* JsonNode = Builder.AddUniqueNode(Actor);
 				if (JsonNode != nullptr && Builder.IsRootActor(Actor))
 				{
 					// TODO: to avoid having to add irrelevant actors/components let GLTFComponentConverter decide and add root nodes to Scene->
diff --git a/Engine/Plugins/Enterprise/GLTFExporter/MigratedFromGit/Source/GLTFExporter/Private/Converters/GLTFSkinConverters.cpp b/Engine/Plugins/Enterprise/GLTFExporter/MigratedFromGit/Source/GLTFExporter/Private/Converters/GLTFSkinConverters.cpp
--- a/Engine/Plugins/Enterprise/GLTFExporter/MigratedFromGit/Source/GLTFExporter/Private/Converters/GLTFSkinConverters.cpp
+++ b/Engine/Plugins/Enterprise/GLTFExporter/MigratedFromGit/Source/GLTFExporter/Private/Converters/GLTFSkinConverters.cpp
@@ -30,7 +30,7 @@ FGLTFJsonSkin* FGLTFSkinConverter::Convert(FGLTFJsonNode* RootNode, const USkele
 
 	for (int32 BoneIndex = 0; BoneIndex < BoneCount; ++BoneIndex)
 	{
-		JsonSkin->Joints[BoneIndex] = Builder.GetOrAddNode(RootNode, SkeletalMesh, BoneIndex);
+		JsonSkin->Joints[BoneIndex] = Builder.AddUniqueNode(RootNode, SkeletalMesh, BoneIndex);
 	}
 
 	TArray<FGLTFMatrix4> InverseBindMatrices;
diff --git a/Engine/Plugins/Enterprise/GLTFExporter/MigratedFromGit/Source/GLTFExporter/Private/Converters/GLTFSkySphereConverters.cpp b/Engine/Plugins/Enterprise/GLTFExporter/MigratedFromGit/Source/GLTFExporter/Private/Converters/GLTFSkySphereConverters.cpp
--- a/Engine/Plugins/Enterprise/GLTFExporter/MigratedFromGit/Source/GLTFExporter/Private/Converters/GLTFSkySphereConverters.cpp
+++ b/Engine/Plugins/Enterprise/GLTFExporter/MigratedFromGit/Source/GLTFExporter/Private/Converters/GLTFSkySphereConverters.cpp
@@ -33,7 +33,7 @@ FGLTFJsonSkySphere* FGLTFSkySphereConverter::Convert(const AActor* SkySphereActo
 		const FTransform RelativeTransform = Transform.GetRelativeTransform(ParentTransform);
 
 		JsonSkySphere->Scale = FGLTFConverterUtility::ConvertScale(RelativeTransform.GetScale3D());
-		JsonSkySphere->SkySphereMesh = Builder.GetOrAddMesh(StaticMeshComponent, { FGLTFMaterialUtility::GetDefaultMaterial() });
+		JsonSkySphere->SkySphereMesh = Builder.AddUniqueMesh(StaticMeshComponent, { FGLTFMaterialUtility::GetDefaultMaterial() });
 	}
 	else
 	{
@@ -43,7 +43,7 @@ FGLTFJsonSkySphere* FGLTFSkySphereConverter::Convert(const AActor* SkySphereActo
 	const ADirectionalLight* DirectionalLight = nullptr;
 	if (FGLTFActorUtility::TryGetPropertyValue(SkySphereActor, TEXT("Directional light actor"), DirectionalLight))
 	{
-		JsonSkySphere->DirectionalLight = Builder.GetOrAddNode(DirectionalLight);
+		JsonSkySphere->DirectionalLight = Builder.AddUniqueNode(DirectionalLight);
 	}
 	else
 	{
@@ -199,7 +199,7 @@ void FGLTFSkySphereConverter::ConvertTextureParameter(const AActor* Actor, const
 	const UTexture2D* Texture2D = TexturePath != nullptr ? LoadObject<UTexture2D>(nullptr, TexturePath) : nullptr;
 	if (Texture2D != nullptr)
 	{
-		OutValue = Builder.GetOrAddTexture(Texture2D);
+		OutValue = Builder.AddUniqueTexture(Texture2D);
 	}
 	else
 	{
diff --git a/Engine/Plugins/Enterprise/GLTFExporter/MigratedFromGit/Source/GLTFExporter/Private/Converters/GLTFVariantUtility.cpp b/Engine/Plugins/Enterprise/GLTFExporter/MigratedFromGit/Source/GLTFExporter/Private/Converters/GLTFVariantUtility.cpp
--- a/Engine/Plugins/Enterprise/GLTFExporter/MigratedFromGit/Source/GLTFExporter/Private/Converters/GLTFVariantUtility.cpp
+++ b/Engine/Plugins/Enterprise/GLTFExporter/MigratedFromGit/Source/GLTFExporter/Private/Converters/GLTFVariantUtility.cpp
@@ -85,17 +85,17 @@ FString FGLTFVariantUtility::GetLogContext(const ULevelVariantSets* LevelVariant
 	return LevelVariantSets->GetName();
 }
 
-FGLTFJsonMaterial* FGLTFVariantUtility::GetOrAddMaterial(FGLTFConvertBuilder& Builder, const UMaterialInterface* Material, const UMeshComponent* MeshComponent, int32 MaterialIndex)
+FGLTFJsonMaterial* FGLTFVariantUtility::AddUniqueMaterial(FGLTFConvertBuilder& Builder, const UMaterialInterface* Material, const UMeshComponent* MeshComponent, int32 MaterialIndex)
 {
 	if (Builder.ExportOptions->ExportMaterialVariants == EGLTFMaterialVariantMode::UseMeshData)
 	{
 		if (Builder.ExportOptions->BakeMaterialInputs == EGLTFMaterialBakeMode::UseMeshData)
 		{
-			return Builder.GetOrAddMaterial(Material, MeshComponent, -1, MaterialIndex);
+			return Builder.AddUniqueMaterial(Material, MeshComponent, -1, MaterialIndex);
 		}
 
 		// TODO: report warning (about materials won't be export using mesh data because BakeMaterialInputs is not set to UseMeshData)
 	}
 
-	return Builder.GetOrAddMaterial(Material);
+	return Builder.AddUniqueMaterial(Material);
 }
diff --git a/Engine/Plugins/Enterprise/GLTFExporter/MigratedFromGit/Source/GLTFExporter/Private/Converters/GLTFVariantUtility.h b/Engine/Plugins/Enterprise/GLTFExporter/MigratedFromGit/Source/GLTFExporter/Private/Converters/GLTFVariantUtility.h
--- a/Engine/Plugins/Enterprise/GLTFExporter/MigratedFromGit/Source/GLTFExporter/Private/Converters/GLTFVariantUtility.h
+++ b/Engine/Plugins/Enterprise/GLTFExporter/MigratedFromGit/Source/GLTFExporter/Private/Converters/GLTFVariantUtility.h
@@ -34,5 +34,5 @@ struct FGLTFVariantUtility
 	static FString GetLogContext(const UVariantSet* VariantSet);
 	static FString GetLogContext(const ULevelVariantSets* LevelVariantSets);
 
-	static FGLTFJsonMaterial* GetOrAddMaterial(FGLTFConvertBuilder& Builder, const UMaterialInterface* Material, const UMeshComponent* MeshComponent, int32 MaterialIndex);
+	static FGLTFJsonMaterial* AddUniqueMaterial(FGLTFConvertBuilder& Builder, const UMaterialInterface* Material, const UMeshComponent* MeshComponent, int32 MaterialIndex);
 };
diff --git a/Engine/Plugins/Enterprise/GLTFExporter/MigratedFromGit/Source/GLTFExporter/Private/Exporters/GLTFAnimSequenceExporter.cpp b/Engine/Plugins/Enterprise/GLTFExporter/MigratedFromGit/Source/GLTFExporter/Private/Exporters/GLTFAnimSequenceExporter.cpp
--- a/Engine/Plugins/Enterprise/GLTFExporter/MigratedFromGit/Source/GLTFExporter/Private/Exporters/GLTFAnimSequenceExporter.cpp
+++ b/Engine/Plugins/Enterprise/GLTFExporter/MigratedFromGit/Source/GLTFExporter/Private/Exporters/GLTFAnimSequenceExporter.cpp
@@ -44,7 +44,7 @@ bool UGLTFAnimSequenceExporter::AddObject(FGLTFContainerBuilder& Builder, const
 
 	if (Builder.ExportOptions->bExportPreviewMesh)
 	{
-		Node->Mesh = Builder.GetOrAddMesh(SkeletalMesh);
+		Node->Mesh = Builder.AddUniqueMesh(SkeletalMesh);
 		if (Node->Mesh == nullptr)
 		{
 			Builder.LogError(
@@ -56,7 +56,7 @@ bool UGLTFAnimSequenceExporter::AddObject(FGLTFContainerBuilder& Builder, const
 	}
 
 
-	Node->Skin = Builder.GetOrAddSkin(Node, SkeletalMesh);
+	Node->Skin = Builder.AddUniqueSkin(Node, SkeletalMesh);
 	if (Node->Skin == nullptr)
 	{
 		Builder.LogError(
@@ -66,7 +66,7 @@ bool UGLTFAnimSequenceExporter::AddObject(FGLTFContainerBuilder& Builder, const
 		return false;
 	}
 
-	FGLTFJsonAnimation* Animation = Builder.GetOrAddAnimation(Node, SkeletalMesh, AnimSequence);
+	FGLTFJsonAnimation* Animation = Builder.AddUniqueAnimation(Node, SkeletalMesh, AnimSequence);
 	if (Animation == nullptr)
 	{
 		Builder.LogError(
diff --git a/Engine/Plugins/Enterprise/GLTFExporter/MigratedFromGit/Source/GLTFExporter/Private/Exporters/GLTFLevelExporter.cpp b/Engine/Plugins/Enterprise/GLTFExporter/MigratedFromGit/Source/GLTFExporter/Private/Exporters/GLTFLevelExporter.cpp
--- a/Engine/Plugins/Enterprise/GLTFExporter/MigratedFromGit/Source/GLTFExporter/Private/Exporters/GLTFLevelExporter.cpp
+++ b/Engine/Plugins/Enterprise/GLTFExporter/MigratedFromGit/Source/GLTFExporter/Private/Exporters/GLTFLevelExporter.cpp
@@ -14,7 +14,7 @@ bool UGLTFLevelExporter::AddObject(FGLTFContainerBuilder& Builder, const UObject
 {
 	const UWorld* World = CastChecked<UWorld>(Object);
 
-	FGLTFJsonScene* Scene = Builder.GetOrAddScene(World);
+	FGLTFJsonScene* Scene = Builder.AddUniqueScene(World);
 	if (Scene == nullptr)
 	{
 		Builder.LogError(FString::Printf(TEXT("Failed to export level %s"), *World->GetName()));
diff --git a/Engine/Plugins/Enterprise/GLTFExporter/MigratedFromGit/Source/GLTFExporter/Private/Exporters/GLTFLevelSequenceExporter.cpp b/Engine/Plugins/Enterprise/GLTFExporter/MigratedFromGit/Source/GLTFExporter/Private/Exporters/GLTFLevelSequenceExporter.cpp
--- a/Engine/Plugins/Enterprise/GLTFExporter/MigratedFromGit/Source/GLTFExporter/Private/Exporters/GLTFLevelSequenceExporter.cpp
+++ b/Engine/Plugins/Enterprise/GLTFExporter/MigratedFromGit/Source/GLTFExporter/Private/Exporters/GLTFLevelSequenceExporter.cpp
@@ -41,7 +41,7 @@ bool UGLTFLevelSequenceExporter::AddObject(FGLTFContainerBuilder& Builder, const
 	}
 
 	const UWorld* World = Worlds[0];
-	FGLTFJsonScene* Scene = Builder.GetOrAddScene(World);
+	FGLTFJsonScene* Scene = Builder.AddUniqueScene(World);
 	if (Scene == nullptr)
 	{
 		Builder.LogError(
@@ -51,7 +51,7 @@ bool UGLTFLevelSequenceExporter::AddObject(FGLTFContainerBuilder& Builder, const
 		return false;
 	}
 
-	FGLTFJsonAnimation* Animation = Builder.GetOrAddAnimation(World->PersistentLevel, LevelSequence);
+	FGLTFJsonAnimation* Animation = Builder.AddUniqueAnimation(World->PersistentLevel, LevelSequence);
 	if (Animation == nullptr)
 	{
 		Builder.LogError(
diff --git a/Engine/Plugins/Enterprise/GLTFExporter/MigratedFromGit/Source/GLTFExporter/Private/Exporters/GLTFLevelVariantSetsExporter.cpp b/Engine/Plugins/Enterprise/GLTFExporter/MigratedFromGit/Source/GLTFExporter/Private/Exporters/GLTFLevelVariantSetsExporter.cpp
--- a/Engine/Plugins/Enterprise/GLTFExporter/MigratedFromGit/Source/GLTFExporter/Private/Exporters/GLTFLevelVariantSetsExporter.cpp
+++ b/Engine/Plugins/Enterprise/GLTFExporter/MigratedFromGit/Source/GLTFExporter/Private/Exporters/GLTFLevelVariantSetsExporter.cpp
@@ -41,7 +41,7 @@ bool UGLTFLevelVariantSetsExporter::AddObject(FGLTFContainerBuilder& Builder, co
 	}
 
 	const UWorld* World = Worlds[0];
-	FGLTFJsonScene* Scene = Builder.GetOrAddScene(World);
+	FGLTFJsonScene* Scene = Builder.AddUniqueScene(World);
 	if (Scene == nullptr)
 	{
 		Builder.LogError(
@@ -53,7 +53,7 @@ bool UGLTFLevelVariantSetsExporter::AddObject(FGLTFContainerBuilder& Builder, co
 
 	if (Builder.ExportOptions->VariantSetsMode == EGLTFVariantSetsMode::Epic)
 	{
-		FGLTFJsonEpicLevelVariantSets* EpicLevelVariantSets = Builder.GetOrAddEpicLevelVariantSets(LevelVariantSets);
+		FGLTFJsonEpicLevelVariantSets* EpicLevelVariantSets = Builder.AddUniqueEpicLevelVariantSets(LevelVariantSets);
 		if (EpicLevelVariantSets == nullptr)
 		{
 			Builder.LogError(
diff --git a/Engine/Plugins/Enterprise/GLTFExporter/MigratedFromGit/Source/GLTFExporter/Private/Exporters/GLTFMaterialExporter.cpp b/Engine/Plugins/Enterprise/GLTFExporter/MigratedFromGit/Source/GLTFExporter/Private/Exporters/GLTFMaterialExporter.cpp
--- a/Engine/Plugins/Enterprise/GLTFExporter/MigratedFromGit/Source/GLTFExporter/Private/Exporters/GLTFMaterialExporter.cpp
+++ b/Engine/Plugins/Enterprise/GLTFExporter/MigratedFromGit/Source/GLTFExporter/Private/Exporters/GLTFMaterialExporter.cpp
@@ -27,7 +27,7 @@ bool UGLTFMaterialExporter::AddObject(FGLTFContainerBuilder& Builder, const UObj
 			PreviewMesh = DefaultPreviewMesh;
 		}
 
-		FGLTFJsonMesh* Mesh = Builder.GetOrAddMesh(PreviewMesh, { Material });
+		FGLTFJsonMesh* Mesh = Builder.AddUniqueMesh(PreviewMesh, { Material });
 		if (Mesh == nullptr)
 		{
 			Builder.LogError(
@@ -47,7 +47,7 @@ bool UGLTFMaterialExporter::AddObject(FGLTFContainerBuilder& Builder, const UObj
 	}
 	else
 	{
-		Builder.GetOrAddMaterial(Material);
+		Builder.AddUniqueMaterial(Material);
 	}
 
 	return true;
diff --git a/Engine/Plugins/Enterprise/GLTFExporter/MigratedFromGit/Source/GLTFExporter/Private/Exporters/GLTFSkeletalMeshExporter.cpp b/Engine/Plugins/Enterprise/GLTFExporter/MigratedFromGit/Source/GLTFExporter/Private/Exporters/GLTFSkeletalMeshExporter.cpp
--- a/Engine/Plugins/Enterprise/GLTFExporter/MigratedFromGit/Source/GLTFExporter/Private/Exporters/GLTFSkeletalMeshExporter.cpp
+++ b/Engine/Plugins/Enterprise/GLTFExporter/MigratedFromGit/Source/GLTFExporter/Private/Exporters/GLTFSkeletalMeshExporter.cpp
@@ -14,7 +14,7 @@ bool UGLTFSkeletalMeshExporter::AddObject(FGLTFContainerBuilder& Builder, const
 {
 	const USkeletalMesh* SkeletalMesh = CastChecked<USkeletalMesh>(Object);
 
-	FGLTFJsonMesh* Mesh = Builder.GetOrAddMesh(SkeletalMesh);
+	FGLTFJsonMesh* Mesh = Builder.AddUniqueMesh(SkeletalMesh);
 	if (Mesh == nullptr)
 	{
 		Builder.LogError(FString::Printf(TEXT("Failed to export skeletal mesh %s"), *SkeletalMesh->GetName()));
@@ -26,7 +26,7 @@ bool UGLTFSkeletalMeshExporter::AddObject(FGLTFContainerBuilder& Builder, const
 
 	if (Builder.ExportOptions->bExportVertexSkinWeights)
 	{
-		Node->Skin = Builder.GetOrAddSkin(Node, SkeletalMesh);
+		Node->Skin = Builder.AddUniqueSkin(Node, SkeletalMesh);
 		if (Node->Skin == nullptr)
 		{
 			Builder.LogError(FString::Printf(TEXT("Failed to export bones in skeletal mesh %s"), *SkeletalMesh->GetName()));
diff --git a/Engine/Plugins/Enterprise/GLTFExporter/MigratedFromGit/Source/GLTFExporter/Private/Exporters/GLTFStaticMeshExporter.cpp b/Engine/Plugins/Enterprise/GLTFExporter/MigratedFromGit/Source/GLTFExporter/Private/Exporters/GLTFStaticMeshExporter.cpp
--- a/Engine/Plugins/Enterprise/GLTFExporter/MigratedFromGit/Source/GLTFExporter/Private/Exporters/GLTFStaticMeshExporter.cpp
+++ b/Engine/Plugins/Enterprise/GLTFExporter/MigratedFromGit/Source/GLTFExporter/Private/Exporters/GLTFStaticMeshExporter.cpp
@@ -14,7 +14,7 @@ bool UGLTFStaticMeshExporter::AddObject(FGLTFContainerBuilder& Builder, const UO
 {
 	const UStaticMesh* StaticMesh = CastChecked<UStaticMesh>(Object);
 
-	FGLTFJsonMesh* Mesh = Builder.GetOrAddMesh(StaticMesh);
+	FGLTFJsonMesh* Mesh = Builder.AddUniqueMesh(StaticMesh);
 	if (Mesh == nullptr)
 	{
 		Builder.LogError(FString::Printf(TEXT("Failed to export static mesh %s"), *StaticMesh->GetName()));
diff --git a/Engine/Plugins/Enterprise/GLTFExporter/MigratedFromGit/Source/GLTFExporter/Private/Tasks/GLTFAnimationTasks.cpp b/Engine/Plugins/Enterprise/GLTFExporter/MigratedFromGit/Source/GLTFExporter/Private/Tasks/GLTFAnimationTasks.cpp
--- a/Engine/Plugins/Enterprise/GLTFExporter/MigratedFromGit/Source/GLTFExporter/Private/Tasks/GLTFAnimationTasks.cpp
+++ b/Engine/Plugins/Enterprise/GLTFExporter/MigratedFromGit/Source/GLTFExporter/Private/Tasks/GLTFAnimationTasks.cpp
@@ -51,7 +51,7 @@ void FGLTFAnimSequenceTask::Complete()
 
 	for (const FBoneIndexType BoneIndex : BoneIndices)
 	{
-		FGLTFJsonNode* Node = Builder.GetOrAddNode(RootNode, SkeletalMesh, BoneIndex);
+		FGLTFJsonNode* Node = Builder.AddUniqueNode(RootNode, SkeletalMesh, BoneIndex);
 
 		// TODO: detect if a bone has the same transforms across multiple frames (at least if its the same across all frames) and optimize
 
@@ -192,11 +192,11 @@ void FGLTFLevelSequenceTask::Complete()
 
 			if (AActor* Actor = Cast<AActor>(Object.Get()))
 			{
-				Node = Builder.GetOrAddNode(Actor);
+				Node = Builder.AddUniqueNode(Actor);
 			}
 			else if (USceneComponent* SceneComponent = Cast<USceneComponent>(Object.Get()))
 			{
-				Node = Builder.GetOrAddNode(SceneComponent);
+				Node = Builder.AddUniqueNode(SceneComponent);
 			}
 			else
 			{
diff --git a/Engine/Plugins/Enterprise/GLTFExporter/MigratedFromGit/Source/GLTFExporter/Private/Tasks/GLTFMaterialTasks.cpp b/Engine/Plugins/Enterprise/GLTFExporter/MigratedFromGit/Source/GLTFExporter/Private/Tasks/GLTFMaterialTasks.cpp
--- a/Engine/Plugins/Enterprise/GLTFExporter/MigratedFromGit/Source/GLTFExporter/Private/Tasks/GLTFMaterialTasks.cpp
+++ b/Engine/Plugins/Enterprise/GLTFExporter/MigratedFromGit/Source/GLTFExporter/Private/Tasks/GLTFMaterialTasks.cpp
@@ -286,7 +286,7 @@ void FGLTFMaterialTask::GetProxyParameter(const FGLTFProxyMaterialTextureParamet
 	}
 
 	const bool bSRGB = ParameterInfo == FGLTFProxyMaterialInfo::BaseColor || ParameterInfo == FGLTFProxyMaterialInfo::Emissive;
-	OutValue.Index = Builder.GetOrAddTexture(Texture, bSRGB);
+	OutValue.Index = Builder.AddUniqueTexture(Texture, bSRGB);
 
 	float UVIndex;
 	if (ParameterInfo.UVIndex.Get(Material, UVIndex, true))
@@ -516,7 +516,7 @@ bool FGLTFMaterialTask::TryGetBaseColorAndOpacity(FGLTFJsonPBRMetallicRoughness&
 		BaseColorTexCoord == OpacityTexCoord &&
 		BaseColorTransform.IsExactlyEqual(OpacityTransform))
 	{
-		OutPBRParams.BaseColorTexture.Index = Builder.GetOrAddTexture(BaseColorTexture, true);
+		OutPBRParams.BaseColorTexture.Index = Builder.AddUniqueTexture(BaseColorTexture, true);
 		OutPBRParams.BaseColorTexture.TexCoord = BaseColorTexCoord;
 		OutPBRParams.BaseColorTexture.Transform = BaseColorTransform;
 		return true;
@@ -642,7 +642,7 @@ bool FGLTFMaterialTask::TryGetMetallicAndRoughness(FGLTFJsonPBRMetallicRoughness
 		MetallicTexCoord == RoughnessTexCoord &&
 		MetallicTransform.IsExactlyEqual(RoughnessTransform))
 	{
-		OutPBRParams.MetallicRoughnessTexture.Index = Builder.GetOrAddTexture(MetallicTexture, false);
+		OutPBRParams.MetallicRoughnessTexture.Index = Builder.AddUniqueTexture(MetallicTexture, false);
 		OutPBRParams.MetallicRoughnessTexture.TexCoord = MetallicTexCoord;
 		OutPBRParams.MetallicRoughnessTexture.Transform = MetallicTransform;
 		return true;
@@ -765,7 +765,7 @@ bool FGLTFMaterialTask::TryGetClearCoatRoughness(FGLTFJsonClearCoatExtension& Ou
 		IntensityTexCoord == RoughnessTexCoord &&
 		IntensityTransform.IsExactlyEqual(RoughnessTransform))
 	{
-		FGLTFJsonTexture* Texture = Builder.GetOrAddTexture(IntensityTexture, false);
+		FGLTFJsonTexture* Texture = Builder.AddUniqueTexture(IntensityTexture, false);
 		OutExtParams.ClearCoatTexture.Index = Texture;
 		OutExtParams.ClearCoatTexture.TexCoord = IntensityTexCoord;
 		OutExtParams.ClearCoatRoughnessTexture.Index = Texture;
@@ -1176,7 +1176,7 @@ bool FGLTFMaterialTask::TryGetSourceTexture(FGLTFJsonTextureInfo& OutTexInfo, co
 
 	if (TryGetSourceTexture(Texture, TexCoord, Transform, Property, AllowedMasks))
 	{
-		OutTexInfo.Index = Builder.GetOrAddTexture(Texture, FGLTFMaterialUtility::IsSRGB(Property));
+		OutTexInfo.Index = Builder.AddUniqueTexture(Texture, FGLTFMaterialUtility::IsSRGB(Property));
 		OutTexInfo.TexCoord = TexCoord;
 		OutTexInfo.Transform = Transform;
 		return true;
diff --git a/Engine/Plugins/Enterprise/GLTFExporter/MigratedFromGit/Source/GLTFExporter/Private/Tasks/GLTFMeshTasks.cpp b/Engine/Plugins/Enterprise/GLTFExporter/MigratedFromGit/Source/GLTFExporter/Private/Tasks/GLTFMeshTasks.cpp
--- a/Engine/Plugins/Enterprise/GLTFExporter/MigratedFromGit/Source/GLTFExporter/Private/Tasks/GLTFMeshTasks.cpp
+++ b/Engine/Plugins/Enterprise/GLTFExporter/MigratedFromGit/Source/GLTFExporter/Private/Tasks/GLTFMeshTasks.cpp
@@ -134,7 +134,7 @@ void FGLTFStaticMeshTask::Complete()
 	}
 
 	const FGLTFMeshData* MeshData = Builder.ExportOptions->BakeMaterialInputs == EGLTFMaterialBakeMode::UseMeshData ?
-		Builder.GetOrAddMeshData(StaticMesh, StaticMeshComponent, LODIndex) : nullptr;
+		Builder.AddUniqueMeshData(StaticMesh, StaticMeshComponent, LODIndex) : nullptr;
 
 #if WITH_EDITOR
 	if (MeshData != nullptr)
@@ -163,9 +163,9 @@ void FGLTFStaticMeshTask::Complete()
 		const FGLTFMeshSection* ConvertedSection = MeshSectionConverter.GetOrAdd(&MeshLOD, SectionIndices);
 
 		FGLTFJsonPrimitive& JsonPrimitive = JsonMesh->Primitives[MaterialIndex];
-		JsonPrimitive.Indices = Builder.GetOrAddIndexAccessor(ConvertedSection);
+		JsonPrimitive.Indices = Builder.AddUniqueIndexAccessor(ConvertedSection);
 
-		JsonPrimitive.Attributes.Position = Builder.GetOrAddPositionAccessor(ConvertedSection, &PositionBuffer);
+		JsonPrimitive.Attributes.Position = Builder.AddUniquePositionAccessor(ConvertedSection, &PositionBuffer);
 		if (JsonPrimitive.Attributes.Position == nullptr)
 		{
 			FString SectionString = TEXT("section");
@@ -183,12 +183,12 @@ void FGLTFStaticMeshTask::Complete()
 
 		if (ColorBuffer != nullptr)
 		{
-			JsonPrimitive.Attributes.Color0 = Builder.GetOrAddColorAccessor(ConvertedSection, ColorBuffer);
+			JsonPrimitive.Attributes.Color0 = Builder.AddUniqueColorAccessor(ConvertedSection, ColorBuffer);
 		}
 
 		// TODO: report warning if both Mesh Quantization (export options) and Use High Precision Tangent Basis (vertex buffer) are disabled
-		JsonPrimitive.Attributes.Normal = Builder.GetOrAddNormalAccessor(ConvertedSection, &VertexBuffer);
-		JsonPrimitive.Attributes.Tangent = Builder.GetOrAddTangentAccessor(ConvertedSection, &VertexBuffer);
+		JsonPrimitive.Attributes.Normal = Builder.AddUniqueNormalAccessor(ConvertedSection, &VertexBuffer);
+		JsonPrimitive.Attributes.Tangent = Builder.AddUniqueTangentAccessor(ConvertedSection, &VertexBuffer);
 
 		const uint32 UVCount = VertexBuffer.GetNumTexCoords();
 		// TODO: report warning or option to limit UV channels since most viewers don't support more than 2?
@@ -196,11 +196,11 @@ void FGLTFStaticMeshTask::Complete()
 
 		for (uint32 UVIndex = 0; UVIndex < UVCount; ++UVIndex)
 		{
-			JsonPrimitive.Attributes.TexCoords[UVIndex] = Builder.GetOrAddUVAccessor(ConvertedSection, &VertexBuffer, UVIndex);
+			JsonPrimitive.Attributes.TexCoords[UVIndex] = Builder.AddUniqueUVAccessor(ConvertedSection, &VertexBuffer, UVIndex);
 		}
 
 		const UMaterialInterface* Material = Materials[MaterialIndex];
-		JsonPrimitive.Material =  Builder.GetOrAddMaterial(Material, MeshData, SectionIndices);
+		JsonPrimitive.Material =  Builder.AddUniqueMaterial(Material, MeshData, SectionIndices);
 	}
 }
 
@@ -237,7 +237,7 @@ void FGLTFSkeletalMeshTask::Complete()
 	}
 
 	const FGLTFMeshData* MeshData = Builder.ExportOptions->BakeMaterialInputs == EGLTFMaterialBakeMode::UseMeshData ?
-		Builder.GetOrAddMeshData(SkeletalMesh, SkeletalMeshComponent, LODIndex) : nullptr;
+		Builder.AddUniqueMeshData(SkeletalMesh, SkeletalMeshComponent, LODIndex) : nullptr;
 
 #if WITH_EDITOR
 	if (MeshData != nullptr)
@@ -266,9 +266,9 @@ void FGLTFSkeletalMeshTask::Complete()
 		const FGLTFMeshSection* ConvertedSection = MeshSectionConverter.GetOrAdd(&MeshLOD, SectionIndices);
 
 		FGLTFJsonPrimitive& JsonPrimitive = JsonMesh->Primitives[MaterialIndex];
-		JsonPrimitive.Indices = Builder.GetOrAddIndexAccessor(ConvertedSection);
+		JsonPrimitive.Indices = Builder.AddUniqueIndexAccessor(ConvertedSection);
 
-		JsonPrimitive.Attributes.Position = Builder.GetOrAddPositionAccessor(ConvertedSection, &PositionBuffer);
+		JsonPrimitive.Attributes.Position = Builder.AddUniquePositionAccessor(ConvertedSection, &PositionBuffer);
 		if (JsonPrimitive.Attributes.Position == nullptr)
 		{
 			FString SectionString = TEXT("section");
@@ -286,12 +286,12 @@ void FGLTFSkeletalMeshTask::Complete()
 
 		if (ColorBuffer != nullptr)
 		{
-			JsonPrimitive.Attributes.Color0 = Builder.GetOrAddColorAccessor(ConvertedSection, ColorBuffer);
+			JsonPrimitive.Attributes.Color0 = Builder.AddUniqueColorAccessor(ConvertedSection, ColorBuffer);
 		}
 
 		// TODO: report warning if both Mesh Quantization (export options) and Use High Precision Tangent Basis (vertex buffer) are disabled
-		JsonPrimitive.Attributes.Normal = Builder.GetOrAddNormalAccessor(ConvertedSection, &VertexBuffer);
-		JsonPrimitive.Attributes.Tangent = Builder.GetOrAddTangentAccessor(ConvertedSection, &VertexBuffer);
+		JsonPrimitive.Attributes.Normal = Builder.AddUniqueNormalAccessor(ConvertedSection, &VertexBuffer);
+		JsonPrimitive.Attributes.Tangent = Builder.AddUniqueTangentAccessor(ConvertedSection, &VertexBuffer);
 
 		const uint32 UVCount = VertexBuffer.GetNumTexCoords();
 		// TODO: report warning or option to limit UV channels since most viewers don't support more than 2?
@@ -299,7 +299,7 @@ void FGLTFSkeletalMeshTask::Complete()
 
 		for (uint32 UVIndex = 0; UVIndex < UVCount; ++UVIndex)
 		{
-			JsonPrimitive.Attributes.TexCoords[UVIndex] = Builder.GetOrAddUVAccessor(ConvertedSection, &VertexBuffer, UVIndex);
+			JsonPrimitive.Attributes.TexCoords[UVIndex] = Builder.AddUniqueUVAccessor(ConvertedSection, &VertexBuffer, UVIndex);
 		}
 
 		if (Builder.ExportOptions->bExportVertexSkinWeights)
@@ -311,12 +311,12 @@ void FGLTFSkeletalMeshTask::Complete()
 
 			for (uint32 GroupIndex = 0; GroupIndex < GroupCount; ++GroupIndex)
 			{
-				JsonPrimitive.Attributes.Joints[GroupIndex] = Builder.GetOrAddJointAccessor(ConvertedSection, SkinWeightBuffer, GroupIndex * 4);
-				JsonPrimitive.Attributes.Weights[GroupIndex] = Builder.GetOrAddWeightAccessor(ConvertedSection, SkinWeightBuffer, GroupIndex * 4);
+				JsonPrimitive.Attributes.Joints[GroupIndex] = Builder.AddUniqueJointAccessor(ConvertedSection, SkinWeightBuffer, GroupIndex * 4);
+				JsonPrimitive.Attributes.Weights[GroupIndex] = Builder.AddUniqueWeightAccessor(ConvertedSection, SkinWeightBuffer, GroupIndex * 4);
 			}
 		}
 
 		const UMaterialInterface* Material = Materials[MaterialIndex];
-		JsonPrimitive.Material =  Builder.GetOrAddMaterial(Material, MeshData, SectionIndices);
+		JsonPrimitive.Material =  Builder.AddUniqueMaterial(Material, MeshData, SectionIndices);
 	}
 }
diff --git a/Engine/Plugins/Enterprise/GLTFExporter/MigratedFromGit/Source/GLTFExporter/Private/Tasks/GLTFTextureTasks.cpp b/Engine/Plugins/Enterprise/GLTFExporter/MigratedFromGit/Source/GLTFExporter/Private/Tasks/GLTFTextureTasks.cpp
--- a/Engine/Plugins/Enterprise/GLTFExporter/MigratedFromGit/Source/GLTFExporter/Private/Tasks/GLTFTextureTasks.cpp
+++ b/Engine/Plugins/Enterprise/GLTFExporter/MigratedFromGit/Source/GLTFExporter/Private/Tasks/GLTFTextureTasks.cpp
@@ -58,8 +58,8 @@ void FGLTFTexture2DTask::Complete()
 		Texture2D->IsNormalMap() ? EGLTFTextureType::Normalmaps :
 		bIsHDR ? EGLTFTextureType::HDR : EGLTFTextureType::None;
 
-	JsonTexture->Source = Builder.GetOrAddImage(Pixels, Size, bIgnoreAlpha, Type, JsonTexture->Name);
-	JsonTexture->Sampler = Builder.GetOrAddSampler(Texture2D);
+	JsonTexture->Source = Builder.AddUniqueImage(Pixels, Size, bIgnoreAlpha, Type, JsonTexture->Name);
+	JsonTexture->Sampler = Builder.AddUniqueSampler(Texture2D);
 }
 
 FString FGLTFTextureCubeTask::GetName()
@@ -111,8 +111,8 @@ void FGLTFTextureCubeTask::Complete()
 	const bool bIgnoreAlpha = FGLTFTextureUtility::IsAlphaless(TextureCube->GetPixelFormat());
 	const EGLTFTextureType Type = bIsHDR ? EGLTFTextureType::HDR : EGLTFTextureType::None;
 
-	JsonTexture->Source = Builder.GetOrAddImage(Pixels, Size, bIgnoreAlpha, Type, JsonTexture->Name);
-	JsonTexture->Sampler = Builder.GetOrAddSampler(TextureCube);
+	JsonTexture->Source = Builder.AddUniqueImage(Pixels, Size, bIgnoreAlpha, Type, JsonTexture->Name);
+	JsonTexture->Sampler = Builder.AddUniqueSampler(TextureCube);
 }
 
 FString FGLTFTextureRenderTarget2DTask::GetName()
@@ -150,8 +150,8 @@ void FGLTFTextureRenderTarget2DTask::Complete()
 	const bool bIgnoreAlpha = FGLTFTextureUtility::IsAlphaless(RenderTarget2D->GetFormat());
 	const EGLTFTextureType Type = bIsHDR ? EGLTFTextureType::HDR : EGLTFTextureType::None;
 
-	JsonTexture->Source = Builder.GetOrAddImage(Pixels, Size, bIgnoreAlpha, Type, JsonTexture->Name);
-	JsonTexture->Sampler = Builder.GetOrAddSampler(RenderTarget2D);
+	JsonTexture->Source = Builder.AddUniqueImage(Pixels, Size, bIgnoreAlpha, Type, JsonTexture->Name);
+	JsonTexture->Sampler = Builder.AddUniqueSampler(RenderTarget2D);
 }
 
 FString FGLTFTextureRenderTargetCubeTask::GetName()
@@ -202,8 +202,8 @@ void FGLTFTextureRenderTargetCubeTask::Complete()
 	const bool bIgnoreAlpha = FGLTFTextureUtility::IsAlphaless(RenderTargetCube->GetFormat());
 	const EGLTFTextureType Type = bIsHDR ? EGLTFTextureType::HDR : EGLTFTextureType::None;
 
-	JsonTexture->Source = Builder.GetOrAddImage(Pixels, Size, bIgnoreAlpha, Type, JsonTexture->Name);
-	JsonTexture->Sampler = Builder.GetOrAddSampler(RenderTargetCube);
+	JsonTexture->Source = Builder.AddUniqueImage(Pixels, Size, bIgnoreAlpha, Type, JsonTexture->Name);
+	JsonTexture->Sampler = Builder.AddUniqueSampler(RenderTargetCube);
 }
 
 #if WITH_EDITOR
@@ -244,8 +244,8 @@ void FGLTFTextureLightMapTask::Complete()
 	TGLTFSharedArray<FColor> Pixels = MakeShared<TArray<FColor>>(static_cast<const FColor*>(RawData), ByteLength / sizeof(FColor));
 	Source.UnlockMip(0);
 
-	JsonTexture->Source = Builder.GetOrAddImage(Pixels, Size, bIgnoreAlpha, Type, JsonTexture->Name);
-	JsonTexture->Sampler = Builder.GetOrAddSampler(LightMap);
+	JsonTexture->Source = Builder.AddUniqueImage(Pixels, Size, bIgnoreAlpha, Type, JsonTexture->Name);
+	JsonTexture->Sampler = Builder.AddUniqueSampler(LightMap);
 }
 
 #endif
