[GLTFExporter] chore: add custom shader for rendering individual cubemap faces

commit 756292aeae8b23d8a8c8ff7e8b0fc907dbc3d41f
Date: Wed, 2 Feb 2022 15:18:51 +0100

    chore: add custom shader for rendering individual cubemap faces
    (cherry picked from commit bb3d1ec76e6fb3b70e59b16acab822b6f1b6f3a2​)

#skipci

diff --git a/Engine/Plugins/Enterprise/GLTFExporter/MigratedFromGit/Shaders/Private/CubemapFacePS.usf b/Engine/Plugins/Enterprise/GLTFExporter/MigratedFromGit/Shaders/Private/CubemapFacePS.usf
new file mode 100644
--- /dev/null
+++ b/Engine/Plugins/Enterprise/GLTFExporter/MigratedFromGit/Shaders/Private/CubemapFacePS.usf
@@ -0,0 +1,60 @@
+﻿// Copyright Epic Games, Inc. All Rights Reserved.
+
+#pragma once
+
+#include "/Engine/Private/Common.ush"
+
+static const float3x3 FaceTransforms[6] =
+{
+	// PosX (1, -v, -u)
+	{
+		 0.0,  0.0, -1.0,
+		 0.0, -1.0,  0.0,
+		 1.0,  0.0,  0.0,
+	},
+	// NegX (-1, -v, u)
+	{
+		 0.0,  0.0,  1.0,
+		 0.0, -1.0,  0.0,
+		-1.0,  0.0,  0.0,
+	},
+	// PosY (u, 1, v)
+	{
+		1.0,  0.0,  0.0,
+		0.0,  0.0,  1.0,
+		0.0,  1.0,  0.0,
+   },
+	// NegY (u, -1, -v)
+	{
+		 1.0,  0.0,  0.0,
+		 0.0,  0.0, -1.0,
+		 0.0, -1.0,  0.0,
+	},
+	// PosZ (u, -v, 1)
+	{
+		 1.0,  0.0,  0.0,
+		 0.0, -1.0,  0.0,
+		 0.0,  0.0,  1.0,
+	},
+	// NegZ (-u, -v, -1)
+	{
+		-1.0,  0.0,  0.0,
+		 0.0, -1.0,  0.0,
+		 0.0,  0.0, -1.0,
+	},
+};
+
+TextureCube InTextureCube;
+SamplerState InTextureCubeSampler;
+int InCubeFaceIndex;
+
+void Main(
+	in float2 TextureCoordinate : TEXCOORD0,
+	out float4 OutColor : SV_Target0
+	)
+{
+	const float3 InputVec = float3(TextureCoordinate.xy * 2. - 1., 1.);
+	const float3 SampleVec = mul(InputVec, FaceTransforms[InCubeFaceIndex]);
+
+	OutColor = TextureCubeSample(InTextureCube, InTextureCubeSampler, SampleVec);
+}
diff --git a/Engine/Plugins/Enterprise/GLTFExporter/MigratedFromGit/Source/GLTFExporter/Private/Converters/GLTFCubemapFacePreview.cpp b/Engine/Plugins/Enterprise/GLTFExporter/MigratedFromGit/Source/GLTFExporter/Private/Converters/GLTFCubemapFacePreview.cpp
new file mode 100644
--- /dev/null
+++ b/Engine/Plugins/Enterprise/GLTFExporter/MigratedFromGit/Source/GLTFExporter/Private/Converters/GLTFCubemapFacePreview.cpp
@@ -0,0 +1,83 @@
+﻿// Copyright Epic Games, Inc. All Rights Reserved.
+
+#include "GLTFCubemapFacePreview.h"
+#include "Shader.h"
+#include "GlobalShader.h"
+#include "SimpleElementShaders.h"
+#include "ShaderParameterUtils.h"
+#include "PipelineStateCache.h"
+
+class FGLTFCubemapFacePreviewPS : public FGlobalShader
+{
+	DECLARE_SHADER_TYPE(FGLTFCubemapFacePreviewPS, Global);
+public:
+
+	static bool ShouldCompilePermutation(const FGlobalShaderPermutationParameters& Parameters)
+	{
+		return IsPCPlatform(Parameters.Platform);
+	}
+
+	FGLTFCubemapFacePreviewPS() {}
+
+	/**
+	 * Initialization constructor.
+	 * @param Initializer - Shader initialization container.
+	 */
+	FGLTFCubemapFacePreviewPS(const ShaderMetaType::CompiledShaderInitializerType& Initializer)
+		: FGlobalShader(Initializer)
+	{
+		InTextureCube.Bind(Initializer.ParameterMap,TEXT("InTextureCube"));
+		InTextureCubeSampler.Bind(Initializer.ParameterMap,TEXT("InTextureCubeSampler"));
+		InCubeFaceIndex.Bind(Initializer.ParameterMap,TEXT("InCubeFaceIndex"));
+	}
+	
+	/**
+	 * Set shader parameters.
+	 * @param TextureCube - The cubemap to render.
+	 * @param CubeFaceIndex - The cubemap face to render (0-5).
+	 */
+	void SetParameters(FRHICommandList& RHICmdList, FRHITexture& TextureCube, int32 CubeFaceIndex)
+	{
+		FRHIPixelShader* PixelShaderRHI = RHICmdList.GetBoundPixelShader();
+		SetTextureParameter(
+			RHICmdList,
+			PixelShaderRHI,
+			InTextureCube,
+			InTextureCubeSampler,
+			TStaticSamplerState<SF_Trilinear, AM_Clamp, AM_Clamp, AM_Clamp>::GetRHI(),
+			&TextureCube);
+		SetShaderValue(RHICmdList, PixelShaderRHI, InCubeFaceIndex, CubeFaceIndex);
+	}
+
+private:
+	/** The texture to sample. */
+	LAYOUT_FIELD(FShaderResourceParameter, InTextureCube);
+	LAYOUT_FIELD(FShaderResourceParameter, InTextureCubeSampler);
+	LAYOUT_FIELD(FShaderParameter, InCubeFaceIndex);
+};
+IMPLEMENT_SHADER_TYPE(,FGLTFCubemapFacePreviewPS,TEXT("/Plugin/GLTFExporter/Private/CubemapFacePS.usf"),TEXT("Main"),SF_Pixel);
+
+/** Binds vertex and pixel shaders for this element */
+void FGLTFCubemapFacePreview::BindShaders(
+	FRHICommandList& RHICmdList,
+	FGraphicsPipelineStateInitializer& GraphicsPSOInit,
+	ERHIFeatureLevel::Type InFeatureLevel,
+	const FMatrix& InTransform,
+	const float InGamma,
+	const FMatrix& ColorWeights,
+	const FTexture* Texture)
+{
+	const TShaderMapRef<FSimpleElementVS> VertexShader(GetGlobalShaderMap(InFeatureLevel));
+	const TShaderMapRef<FGLTFCubemapFacePreviewPS> PixelShader(GetGlobalShaderMap(InFeatureLevel));
+
+	GraphicsPSOInit.BoundShaderState.VertexDeclarationRHI = GSimpleElementVertexDeclaration.VertexDeclarationRHI;
+	GraphicsPSOInit.BoundShaderState.VertexShaderRHI = VertexShader.GetVertexShader();
+	GraphicsPSOInit.BoundShaderState.PixelShaderRHI = PixelShader.GetPixelShader();
+	GraphicsPSOInit.PrimitiveType = PT_TriangleList;
+	GraphicsPSOInit.BlendState = TStaticBlendState<>::GetRHI();
+
+	SetGraphicsPipelineState(RHICmdList, GraphicsPSOInit, EApplyRendertargetOption::ForceApply);
+
+	VertexShader->SetParameters(RHICmdList, InTransform);
+	PixelShader->SetParameters(RHICmdList, *TextureCube, CubeFaceIndex);
+}
diff --git a/Engine/Plugins/Enterprise/GLTFExporter/MigratedFromGit/Source/GLTFExporter/Private/Converters/GLTFCubemapFacePreview.h b/Engine/Plugins/Enterprise/GLTFExporter/MigratedFromGit/Source/GLTFExporter/Private/Converters/GLTFCubemapFacePreview.h
new file mode 100644
--- /dev/null
+++ b/Engine/Plugins/Enterprise/GLTFExporter/MigratedFromGit/Source/GLTFExporter/Private/Converters/GLTFCubemapFacePreview.h
@@ -0,0 +1,33 @@
+﻿// Copyright Epic Games, Inc. All Rights Reserved.
+
+#pragma once
+
+#include "CoreMinimal.h"
+#include "BatchedElements.h"
+
+class FGLTFCubemapFacePreview : public FBatchedElementParameters
+{
+
+public:
+
+	FGLTFCubemapFacePreview(FRHITexture* TextureCube, int32 CubeFaceIndex)
+		: TextureCube(TextureCube)
+		, CubeFaceIndex(CubeFaceIndex)
+	{
+	}
+
+	/** Binds vertex and pixel shaders for this element */
+	virtual void BindShaders(
+		FRHICommandList& RHICmdList,
+		FGraphicsPipelineStateInitializer& GraphicsPSOInit,
+		ERHIFeatureLevel::Type InFeatureLevel,
+		const FMatrix& InTransform,
+		const float InGamma,
+		const FMatrix& ColorWeights,
+		const FTexture* Texture) override;
+
+private:
+
+	FRHITexture* TextureCube;
+	int32 CubeFaceIndex;
+};
