[GLTFExporter] feat: use a shader instead of platform data when exporting cubemap faces to textures

commit edb47c304cae6c6625e2b92ae5f8a8b6e120d230
Date: Wed, 2 Feb 2022 15:19:32 +0100

    feat: use a shader instead of platform data when exporting cubemap faces to textures
    (cherry picked from commit 05ef0050019482b8cb4ed3aa1d9287028c313421​)

#skipci

diff --git a/Engine/Plugins/Enterprise/GLTFExporter/MigratedFromGit/Source/GLTFExporter/Private/Converters/GLTFTextureUtility.cpp b/Engine/Plugins/Enterprise/GLTFExporter/MigratedFromGit/Source/GLTFExporter/Private/Converters/GLTFTextureUtility.cpp
--- a/Engine/Plugins/Enterprise/GLTFExporter/MigratedFromGit/Source/GLTFExporter/Private/Converters/GLTFTextureUtility.cpp
+++ b/Engine/Plugins/Enterprise/GLTFExporter/MigratedFromGit/Source/GLTFExporter/Private/Converters/GLTFTextureUtility.cpp
@@ -3,6 +3,7 @@
 #include "Converters/GLTFTextureUtility.h"
 #include "Converters/GLTFNormalMapPreview.h"
 #include "Converters/GLTFSimpleTexture2DPreview.h"
+#include "Converters/GLTFCubemapFacePreview.h"
 #include "CanvasItem.h"
 #include "Engine/Texture2D.h"
 #include "Engine/TextureCube.h"
@@ -220,30 +221,56 @@ bool FGLTFTextureUtility::RotateTexture(UTextureRenderTarget2D* OutTarget, const
 
 UTexture2D* FGLTFTextureUtility::CreateTextureFromCubeFace(const UTextureCube* TextureCube, ECubeFace CubeFace)
 {
-	const FIntPoint Size(TextureCube->GetSizeX(), TextureCube->GetSizeY());
-	const EPixelFormat Format = TextureCube->GetPixelFormat();
+	check(TextureCube != nullptr &&
+		TextureCube->TextureReference.IsInitialized() &&
+		TextureCube->TextureReference.TextureReferenceRHI.IsValid());
 
-	if (!LoadPlatformData(const_cast<UTextureCube*>(TextureCube)))
-	{
-		return nullptr;
-	}
+	// Fully stream in the texture before drawing it.
+	const_cast<UTextureCube*>(TextureCube)->SetForceMipLevelsToBeResident(30.0f, true);
+	const_cast<UTextureCube*>(TextureCube)->WaitForStreaming();
+
+	const FIntPoint Size = GetInGameSize(TextureCube);
+	const bool bIsHDR = IsHDR(TextureCube);
+
+	UTextureRenderTarget2D* RenderTarget = CreateRenderTarget(Size, bIsHDR);
+	FTextureRenderTargetResource* RenderTargetResource = RenderTarget->GameThread_GetRenderTargetResource();
+	const FTextureReferenceRHIRef& RHIRef = TextureCube->TextureReference.TextureReferenceRHI;
+
+	FCanvas Canvas(RenderTargetResource, nullptr, 0.0f, 0.0f, 0.0f, GMaxRHIFeatureLevel);
+	FCanvasTileItem TileItem(FVector2D::ZeroVector, TextureCube->Resource, FLinearColor::White);
+	TileItem.BatchedElementParameters = new FGLTFCubemapFacePreview(const_cast<FTextureReferenceRHIRef&>(RHIRef),CubeFace);
+	TileItem.Draw(&Canvas);
 
-	const FByteBulkData& BulkData = TextureCube->PlatformData->Mips[0].BulkData;
-	const int64 MipSize = BulkData.GetBulkDataSize() / 6;
-	UTexture2D* FaceTexture = nullptr;
+	Canvas.Flush_GameThread();
+	FlushRenderingCommands();
+	Canvas.SetRenderTarget_GameThread(nullptr);
+	FlushRenderingCommands();
+
+	UTexture2D* FaceTexture;
+	FTextureRenderTarget2DResource* Resource = static_cast<FTextureRenderTarget2DResource*>(RenderTarget->Resource);
+	check(Resource != nullptr);
 
-	const void* MipDataPtr = BulkData.LockReadOnly();
-	if (MipDataPtr != nullptr)
+	if (bIsHDR)
 	{
-		const void* FaceDataPtr = static_cast<const uint8*>(MipDataPtr) + MipSize * CubeFace;
-		FaceTexture = CreateTransientTexture(FaceDataPtr, MipSize, Size, Format, TextureCube->SRGB);
+		TArray<FFloat16Color> Pixels;
+		if (!Resource->ReadFloat16Pixels(Pixels, CubeFace))
+		{
+			return nullptr;
+		}
 
-		if (IsHDR(TextureCube))
+		FaceTexture = CreateTransientTexture(Pixels.GetData(), Pixels.Num() * sizeof(FFloat16Color), Size, PF_FloatRGBA, false);
+		FaceTexture->CompressionSettings = TC_HDR;
+	}
+	else
+	{
+		TArray<FColor> Pixels;
+		if (!Resource->ReadPixels(Pixels, FReadSurfaceDataFlags(RCM_UNorm, CubeFace)))
 		{
-			FaceTexture->CompressionSettings = TC_HDR;
+			return nullptr;
 		}
+
+		FaceTexture = CreateTransientTexture(Pixels.GetData(), Pixels.Num() * sizeof(FColor), Size, PF_B8G8R8A8, false);
 	}
-	const_cast<FByteBulkData&>(BulkData).Unlock();
 
 	return FaceTexture;
 }
