[GLTFExporter] refactor: use an enum for specifying the cubemap face to render

commit 69dd5fcc9425afcfa5f631b699352750fde5e520
Date: Wed, 2 Feb 2022 15:53:23 +0100

    refactor: use an enum for specifying the cubemap face to render
    (cherry picked from commit 8c7e87be1aedc2e642127246fab2e43857a3f954​)

#skipci

diff --git a/Engine/Plugins/Enterprise/GLTFExporter/MigratedFromGit/Source/GLTFExporter/Private/Converters/GLTFCubemapFacePreview.cpp b/Engine/Plugins/Enterprise/GLTFExporter/MigratedFromGit/Source/GLTFExporter/Private/Converters/GLTFCubemapFacePreview.cpp
--- a/Engine/Plugins/Enterprise/GLTFExporter/MigratedFromGit/Source/GLTFExporter/Private/Converters/GLTFCubemapFacePreview.cpp
+++ b/Engine/Plugins/Enterprise/GLTFExporter/MigratedFromGit/Source/GLTFExporter/Private/Converters/GLTFCubemapFacePreview.cpp
@@ -34,9 +34,9 @@ public:
 	/**
 	 * Set shader parameters.
 	 * @param TextureCube - The cubemap to render.
-	 * @param CubeFaceIndex - The cubemap face to render (0-5).
+	 * @param CubeFace - The cubemap face to render.
 	 */
-	void SetParameters(FRHICommandList& RHICmdList, FRHITexture& TextureCube, int32 CubeFaceIndex)
+	void SetParameters(FRHICommandList& RHICmdList, FRHITexture& TextureCube, ECubeFace CubeFace)
 	{
 		FRHIPixelShader* PixelShaderRHI = RHICmdList.GetBoundPixelShader();
 		SetTextureParameter(
@@ -46,7 +46,7 @@ public:
 			InTextureCubeSampler,
 			TStaticSamplerState<SF_Trilinear, AM_Clamp, AM_Clamp, AM_Clamp>::GetRHI(),
 			&TextureCube);
-		SetShaderValue(RHICmdList, PixelShaderRHI, InCubeFaceIndex, CubeFaceIndex);
+		SetShaderValue(RHICmdList, PixelShaderRHI, InCubeFaceIndex, CubeFace);
 	}
 
 private:
@@ -79,5 +79,5 @@ void FGLTFCubemapFacePreview::BindShaders(
 	SetGraphicsPipelineState(RHICmdList, GraphicsPSOInit, EApplyRendertargetOption::ForceApply);
 
 	VertexShader->SetParameters(RHICmdList, InTransform);
-	PixelShader->SetParameters(RHICmdList, *TextureCube, CubeFaceIndex);
+	PixelShader->SetParameters(RHICmdList, *TextureCube, CubeFace);
 }
diff --git a/Engine/Plugins/Enterprise/GLTFExporter/MigratedFromGit/Source/GLTFExporter/Private/Converters/GLTFCubemapFacePreview.h b/Engine/Plugins/Enterprise/GLTFExporter/MigratedFromGit/Source/GLTFExporter/Private/Converters/GLTFCubemapFacePreview.h
--- a/Engine/Plugins/Enterprise/GLTFExporter/MigratedFromGit/Source/GLTFExporter/Private/Converters/GLTFCubemapFacePreview.h
+++ b/Engine/Plugins/Enterprise/GLTFExporter/MigratedFromGit/Source/GLTFExporter/Private/Converters/GLTFCubemapFacePreview.h
@@ -10,9 +10,9 @@ class FGLTFCubemapFacePreview : public FBatchedElementParameters
 
 public:
 
-	FGLTFCubemapFacePreview(FRHITexture* TextureCube, int32 CubeFaceIndex)
+	FGLTFCubemapFacePreview(FRHITexture* TextureCube, ECubeFace CubeFace)
 		: TextureCube(TextureCube)
-		, CubeFaceIndex(CubeFaceIndex)
+		, CubeFace(CubeFace)
 	{
 	}
 
@@ -29,5 +29,5 @@ public:
 private:
 
 	FRHITexture* TextureCube;
-	int32 CubeFaceIndex;
+	ECubeFace CubeFace;
 };
