[GLTFExporter] refactor: remove the need to pass an RHI-texture to FGLTFCubemapFacePreview

commit 203a114aec3d656824ab6a319626eec6c825db3c
Date: Wed, 2 Mar 2022 12:33:34 +0100

    refactor: remove the need to pass an RHI-texture to FGLTFCubemapFacePreview
    (cherry picked from commit c3be697b3c4aa19b9b6564a9d5ca614ae0682d62​)

#skipci

diff --git a/Engine/Plugins/Enterprise/GLTFExporter/MigratedFromGit/Source/GLTFExporter/Private/Converters/GLTFCubemapFacePreview.cpp b/Engine/Plugins/Enterprise/GLTFExporter/MigratedFromGit/Source/GLTFExporter/Private/Converters/GLTFCubemapFacePreview.cpp
--- a/Engine/Plugins/Enterprise/GLTFExporter/MigratedFromGit/Source/GLTFExporter/Private/Converters/GLTFCubemapFacePreview.cpp
+++ b/Engine/Plugins/Enterprise/GLTFExporter/MigratedFromGit/Source/GLTFExporter/Private/Converters/GLTFCubemapFacePreview.cpp
@@ -36,7 +36,7 @@ public:
 	 * @param TextureCube - The cubemap to render.
 	 * @param CubeFace - The cubemap face to render.
 	 */
-	void SetParameters(FRHICommandList& RHICmdList, FRHITexture& TextureCube, ECubeFace CubeFace)
+	void SetParameters(FRHICommandList& RHICmdList, const FTexture* TextureCube, ECubeFace CubeFace) const
 	{
 		FRHIPixelShader* PixelShaderRHI = RHICmdList.GetBoundPixelShader();
 		SetTextureParameter(
@@ -44,8 +44,7 @@ public:
 			PixelShaderRHI,
 			InTextureCube,
 			InTextureCubeSampler,
-			TStaticSamplerState<SF_Trilinear, AM_Clamp, AM_Clamp, AM_Clamp>::GetRHI(),
-			&TextureCube);
+			TextureCube);
 		SetShaderValue(RHICmdList, PixelShaderRHI, InCubeFaceIndex, CubeFace);
 	}
 
@@ -79,5 +78,5 @@ void FGLTFCubemapFacePreview::BindShaders(
 	SetGraphicsPipelineState(RHICmdList, GraphicsPSOInit, EApplyRendertargetOption::ForceApply);
 
 	VertexShader->SetParameters(RHICmdList, InTransform);
-	PixelShader->SetParameters(RHICmdList, *TextureCube, CubeFace);
+	PixelShader->SetParameters(RHICmdList, Texture, CubeFace);
 }
diff --git a/Engine/Plugins/Enterprise/GLTFExporter/MigratedFromGit/Source/GLTFExporter/Private/Converters/GLTFCubemapFacePreview.h b/Engine/Plugins/Enterprise/GLTFExporter/MigratedFromGit/Source/GLTFExporter/Private/Converters/GLTFCubemapFacePreview.h
--- a/Engine/Plugins/Enterprise/GLTFExporter/MigratedFromGit/Source/GLTFExporter/Private/Converters/GLTFCubemapFacePreview.h
+++ b/Engine/Plugins/Enterprise/GLTFExporter/MigratedFromGit/Source/GLTFExporter/Private/Converters/GLTFCubemapFacePreview.h
@@ -10,9 +10,8 @@ class FGLTFCubemapFacePreview : public FBatchedElementParameters
 
 public:
 
-	FGLTFCubemapFacePreview(FRHITexture* TextureCube, ECubeFace CubeFace)
-		: TextureCube(TextureCube)
-		, CubeFace(CubeFace)
+	FGLTFCubemapFacePreview(ECubeFace CubeFace)
+		: CubeFace(CubeFace)
 	{
 	}
 
@@ -28,6 +27,5 @@ public:
 
 private:
 
-	FRHITexture* TextureCube;
 	ECubeFace CubeFace;
 };
diff --git a/Engine/Plugins/Enterprise/GLTFExporter/MigratedFromGit/Source/GLTFExporter/Private/Converters/GLTFTextureUtility.cpp b/Engine/Plugins/Enterprise/GLTFExporter/MigratedFromGit/Source/GLTFExporter/Private/Converters/GLTFTextureUtility.cpp
--- a/Engine/Plugins/Enterprise/GLTFExporter/MigratedFromGit/Source/GLTFExporter/Private/Converters/GLTFTextureUtility.cpp
+++ b/Engine/Plugins/Enterprise/GLTFExporter/MigratedFromGit/Source/GLTFExporter/Private/Converters/GLTFTextureUtility.cpp
@@ -234,11 +234,10 @@ UTexture2D* FGLTFTextureUtility::CreateTextureFromCubeFace(const UTextureCube* T
 
 	UTextureRenderTarget2D* RenderTarget = CreateRenderTarget(Size, bIsHDR);
 	FTextureRenderTargetResource* RenderTargetResource = RenderTarget->GameThread_GetRenderTargetResource();
-	const FTextureReferenceRHIRef& RHIRef = TextureCube->TextureReference.TextureReferenceRHI;
 
 	FCanvas Canvas(RenderTargetResource, nullptr, 0.0f, 0.0f, 0.0f, GMaxRHIFeatureLevel);
 	FCanvasTileItem TileItem(FVector2D::ZeroVector, TextureCube->Resource, FLinearColor::White);
-	TileItem.BatchedElementParameters = new FGLTFCubemapFacePreview(const_cast<FTextureReferenceRHIRef&>(RHIRef),CubeFace);
+	TileItem.BatchedElementParameters = new FGLTFCubemapFacePreview(CubeFace);
 	TileItem.Draw(&Canvas);
 
 	Canvas.Flush_GameThread();
