[GLTFExporter] feat: add support for running the exporter plugin on Linux

commit 9ff3f65067f570f8a9806925e0c4fa0f9687dbbb
Date: Wed, 23 Feb 2022 09:00:04 +0100

    feat: add support for running the exporter plugin on Linux
    (cherry picked from commit 2698dc90ad844446f82fb1acbccd282eebd0a414​)

#skipci

diff --git a/Engine/Plugins/Enterprise/GLTFExporter/MigratedFromGit/GLTFExporter.uplugin b/Engine/Plugins/Enterprise/GLTFExporter/MigratedFromGit/GLTFExporter.uplugin
--- a/Engine/Plugins/Enterprise/GLTFExporter/MigratedFromGit/GLTFExporter.uplugin
+++ b/Engine/Plugins/Enterprise/GLTFExporter/MigratedFromGit/GLTFExporter.uplugin
@@ -17,19 +17,19 @@
 			"Name": "GLTFExporter",
 			"Type": "Runtime",
 			"LoadingPhase": "PostConfigInit",
-			"WhitelistPlatforms": [ "Win64", "Win32", "Mac" ]
+			"WhitelistPlatforms": [ "Win64", "Win32", "Mac", "Linux" ]
 		},
 		{
 			"Name": "GLTFMaterialAnalyzer",
 			"Type": "Editor",
 			"LoadingPhase": "Default",
-			"WhitelistPlatforms": [ "Win64", "Win32", "Mac" ]
+			"WhitelistPlatforms": [ "Win64", "Win32", "Mac", "Linux" ]
 		},
 		{
 			"Name": "GLTFMaterialBaking",
 			"Type": "Editor",
 			"LoadingPhase": "Default",
-			"WhitelistPlatforms": [ "Win64", "Win32", "Mac" ]
+			"WhitelistPlatforms": [ "Win64", "Win32", "Mac", "Linux" ]
 		}
 	],
 	"Plugins": [
