[GLTFExporter] refactor: use hardcoded API identifier for the custom baking module

commit ece7af9fabbb55cd40e5e351facb03cf8f728020
Date: Fri, 25 Feb 2022 14:06:35 +0100

    refactor: use hardcoded API identifier for the custom baking module
    (cherry picked from commit dc672087fcda0eaff5b30515525d41eca5837ebe​)

#skipci

diff --git a/Engine/Plugins/Enterprise/GLTFExporter/MigratedFromGit/Source/GLTFMaterialBaking/GLTFMaterialBaking.Build.cs b/Engine/Plugins/Enterprise/GLTFExporter/MigratedFromGit/Source/GLTFMaterialBaking/GLTFMaterialBaking.Build.cs
--- a/Engine/Plugins/Enterprise/GLTFExporter/MigratedFromGit/Source/GLTFMaterialBaking/GLTFMaterialBaking.Build.cs
+++ b/Engine/Plugins/Enterprise/GLTFExporter/MigratedFromGit/Source/GLTFMaterialBaking/GLTFMaterialBaking.Build.cs
@@ -30,8 +30,5 @@ public class GLTFMaterialBaking : ModuleRules
 
 		// NOTE: ugly hack to access HLSLMaterialTranslator to bake shading model
 		PrivateIncludePaths.Add(EngineDirectory + "/Source/Runtime/Engine/Private");
-
-		// NOTE: avoid renaming all instaces of MATERIALBAKING_API by redirecting it to GLTFMATERIALBAKING_API
-		PublicDefinitions.Add("MATERIALBAKING_API=GLTFMATERIALBAKING_API");
 	}
 }
diff --git a/Engine/Plugins/Enterprise/GLTFExporter/MigratedFromGit/Source/GLTFMaterialBaking/Private/MaterialBakingModule.h b/Engine/Plugins/Enterprise/GLTFExporter/MigratedFromGit/Source/GLTFMaterialBaking/Private/MaterialBakingModule.h
--- a/Engine/Plugins/Enterprise/GLTFExporter/MigratedFromGit/Source/GLTFMaterialBaking/Private/MaterialBakingModule.h
+++ b/Engine/Plugins/Enterprise/GLTFExporter/MigratedFromGit/Source/GLTFMaterialBaking/Private/MaterialBakingModule.h
@@ -19,7 +19,7 @@ struct FMaterialData;
 struct FMeshData;
 struct FBakeOutput;
 
-class MATERIALBAKING_API FMaterialBakingModule : public IMaterialBakingModule
+class GLTFMATERIALBAKING_API FMaterialBakingModule : public IMaterialBakingModule
 {
 public:
 	/** IModuleInterface overrides begin */
diff --git a/Engine/Plugins/Enterprise/GLTFExporter/MigratedFromGit/Source/GLTFMaterialBaking/Public/IMaterialBakingModule.h b/Engine/Plugins/Enterprise/GLTFExporter/MigratedFromGit/Source/GLTFMaterialBaking/Public/IMaterialBakingModule.h
--- a/Engine/Plugins/Enterprise/GLTFExporter/MigratedFromGit/Source/GLTFMaterialBaking/Public/IMaterialBakingModule.h
+++ b/Engine/Plugins/Enterprise/GLTFExporter/MigratedFromGit/Source/GLTFMaterialBaking/Public/IMaterialBakingModule.h
@@ -20,7 +20,7 @@ struct FMeshData;
 struct FBakeOutput;
 struct FBakeOutputEx;
 
-class MATERIALBAKING_API IMaterialBakingModule : public IModuleInterface
+class GLTFMATERIALBAKING_API IMaterialBakingModule : public IModuleInterface
 {
 public:
 	/** Bakes out material properties according to MaterialSettings using MeshSettings and stores the output in Output */
diff --git a/Engine/Plugins/Enterprise/GLTFExporter/MigratedFromGit/Source/GLTFMaterialBaking/Public/MaterialBakingHelpers.h b/Engine/Plugins/Enterprise/GLTFExporter/MigratedFromGit/Source/GLTFMaterialBaking/Public/MaterialBakingHelpers.h
--- a/Engine/Plugins/Enterprise/GLTFExporter/MigratedFromGit/Source/GLTFMaterialBaking/Public/MaterialBakingHelpers.h
+++ b/Engine/Plugins/Enterprise/GLTFExporter/MigratedFromGit/Source/GLTFMaterialBaking/Public/MaterialBakingHelpers.h
@@ -6,7 +6,7 @@
 
 struct FColor;
 
-class MATERIALBAKING_API FMaterialBakingHelpers
+class GLTFMATERIALBAKING_API FMaterialBakingHelpers
 {
 public:
 	/** 
diff --git a/Engine/Plugins/Enterprise/GLTFExporter/MigratedFromGit/Source/GLTFMaterialBaking/Public/MaterialPropertyEx.h b/Engine/Plugins/Enterprise/GLTFExporter/MigratedFromGit/Source/GLTFMaterialBaking/Public/MaterialPropertyEx.h
--- a/Engine/Plugins/Enterprise/GLTFExporter/MigratedFromGit/Source/GLTFMaterialBaking/Public/MaterialPropertyEx.h
+++ b/Engine/Plugins/Enterprise/GLTFExporter/MigratedFromGit/Source/GLTFMaterialBaking/Public/MaterialPropertyEx.h
@@ -42,7 +42,7 @@ struct FMaterialPropertyEx
 		return !Other.IsCustomOutput() ? GetTypeHash(Other.Type) : GetTypeHash(Other.CustomOutput);
 	}
 
-	MATERIALBAKING_API FString ToString() const;
+	GLTFMATERIALBAKING_API FString ToString() const;
 
 	/** The material property */
 	EMaterialProperty Type;
