[GLTFExporter] refactor: use GLTF prefix for structures in the custom baking module

commit d68d4b35ad55dd070c736a9fc1ad8a3c079add12
Date: Fri, 25 Feb 2022 14:15:48 +0100

    refactor: use GLTF prefix for structures in the custom baking module
    (cherry picked from commit ab695fac0c8f3f75ac8a9f90ae32822358aa9e0d​)

#skipci

diff --git a/Engine/Plugins/Enterprise/GLTFExporter/MigratedFromGit/Source/GLTFExporter/Private/Converters/GLTFMaterialUtility.cpp b/Engine/Plugins/Enterprise/GLTFExporter/MigratedFromGit/Source/GLTFExporter/Private/Converters/GLTFMaterialUtility.cpp
--- a/Engine/Plugins/Enterprise/GLTFExporter/MigratedFromGit/Source/GLTFExporter/Private/Converters/GLTFMaterialUtility.cpp
+++ b/Engine/Plugins/Enterprise/GLTFExporter/MigratedFromGit/Source/GLTFExporter/Private/Converters/GLTFMaterialUtility.cpp
@@ -7,7 +7,7 @@
 #include "Misc/DefaultValueHelper.h"
 #if WITH_EDITOR
 #include "GLTFMaterialAnalyzer.h"
-#include "GLTFMaterialBaking/Public/IMaterialBakingModule.h"
+#include "GLTFMaterialBaking/Public/IGLTFMaterialBakingModule.h"
 #include "GLTFMaterialBaking/Public/MaterialBakingStructures.h"
 #include "Modules/ModuleManager.h"
 #include "Materials/MaterialInstanceConstant.h"
@@ -24,25 +24,25 @@ UMaterialInterface* FGLTFMaterialUtility::GetDefaultMaterial()
 
 #if WITH_EDITOR
 
-bool FGLTFMaterialUtility::IsNormalMap(const FMaterialPropertyEx& Property)
+bool FGLTFMaterialUtility::IsNormalMap(const FGLTFMaterialPropertyEx& Property)
 {
 	return Property == MP_Normal || Property == TEXT("ClearCoatBottomNormal");
 }
 
-bool FGLTFMaterialUtility::IsSRGB(const FMaterialPropertyEx& Property)
+bool FGLTFMaterialUtility::IsSRGB(const FGLTFMaterialPropertyEx& Property)
 {
 	return Property == MP_BaseColor || Property == MP_EmissiveColor || Property == MP_SubsurfaceColor || Property == TEXT("TransmittanceColor");
 }
 
 #if (ENGINE_MAJOR_VERSION > 4 || ENGINE_MINOR_VERSION >= 27)
-FGuid FGLTFMaterialUtility::GetAttributeID(const FMaterialPropertyEx& Property)
+FGuid FGLTFMaterialUtility::GetAttributeID(const FGLTFMaterialPropertyEx& Property)
 {
 	return Property.IsCustomOutput()
 		? FMaterialAttributeDefinitionMap::GetCustomAttributeID(Property.CustomOutput.ToString())
 		: FMaterialAttributeDefinitionMap::GetID(Property.Type);
 }
 
-FGuid FGLTFMaterialUtility::GetAttributeIDChecked(const FMaterialPropertyEx& Property)
+FGuid FGLTFMaterialUtility::GetAttributeIDChecked(const FGLTFMaterialPropertyEx& Property)
 {
 	const FGuid AttributeID = GetAttributeID(Property);
 	check(AttributeID != FMaterialAttributeDefinitionMap::GetDefaultID());
@@ -50,7 +50,7 @@ FGuid FGLTFMaterialUtility::GetAttributeIDChecked(const FMaterialPropertyEx& Pro
 }
 #endif
 
-FVector4 FGLTFMaterialUtility::GetPropertyDefaultValue(const FMaterialPropertyEx& Property)
+FVector4 FGLTFMaterialUtility::GetPropertyDefaultValue(const FGLTFMaterialPropertyEx& Property)
 {
 #if (ENGINE_MAJOR_VERSION > 4 || ENGINE_MINOR_VERSION >= 27)
 	return FMaterialAttributeDefinitionMap::GetDefaultValue(GetAttributeIDChecked(Property));
@@ -101,7 +101,7 @@ FVector4 FGLTFMaterialUtility::GetPropertyDefaultValue(const FMaterialPropertyEx
 #endif
 }
 
-FVector4 FGLTFMaterialUtility::GetPropertyMask(const FMaterialPropertyEx& Property)
+FVector4 FGLTFMaterialUtility::GetPropertyMask(const FGLTFMaterialPropertyEx& Property)
 {
 #if (ENGINE_MAJOR_VERSION > 4 || ENGINE_MINOR_VERSION >= 27)
 	switch (FMaterialAttributeDefinitionMap::GetValueType(GetAttributeIDChecked(Property)))
@@ -162,7 +162,7 @@ FVector4 FGLTFMaterialUtility::GetPropertyMask(const FMaterialPropertyEx& Proper
 #endif
 }
 
-const FExpressionInput* FGLTFMaterialUtility::GetInputForProperty(const UMaterialInterface* Material, const FMaterialPropertyEx& Property)
+const FExpressionInput* FGLTFMaterialUtility::GetInputForProperty(const UMaterialInterface* Material, const FGLTFMaterialPropertyEx& Property)
 {
 	if (Property.IsCustomOutput())
 	{
@@ -190,9 +190,9 @@ const UMaterialExpressionCustomOutput* FGLTFMaterialUtility::GetCustomOutputByNa
 	return nullptr;
 }
 
-FGLTFPropertyBakeOutput FGLTFMaterialUtility::BakeMaterialProperty(const FIntPoint& OutputSize, const FMaterialPropertyEx& Property, const UMaterialInterface* Material, int32 TexCoord, const FGLTFMeshData* MeshData, const FGLTFIndexArray& MeshSectionIndices, bool bFillAlpha, bool bAdjustNormalmaps)
+FGLTFPropertyBakeOutput FGLTFMaterialUtility::BakeMaterialProperty(const FIntPoint& OutputSize, const FGLTFMaterialPropertyEx& Property, const UMaterialInterface* Material, int32 TexCoord, const FGLTFMeshData* MeshData, const FGLTFIndexArray& MeshSectionIndices, bool bFillAlpha, bool bAdjustNormalmaps)
 {
-	FMeshData MeshSet;
+	FGLTFMeshRenderData MeshSet;
 	MeshSet.TextureCoordinateBox = { { 0.0f, 0.0f }, { 1.0f, 1.0f } };
 	MeshSet.TextureCoordinateIndex = TexCoord;
 	MeshSet.MaterialIndices = MeshSectionIndices; // NOTE: MaterialIndices is actually section indices
@@ -205,25 +205,25 @@ FGLTFPropertyBakeOutput FGLTFMaterialUtility::BakeMaterialProperty(const FIntPoi
 		MeshSet.PrimitiveData = &MeshData->PrimitiveData;
 	}
 
-	FMaterialDataEx MatSet;
+	FGLTFMaterialDataEx MatSet;
 	MatSet.Material = const_cast<UMaterialInterface*>(Material);
 	MatSet.PropertySizes.Add(Property, OutputSize);
 	MatSet.bTangentSpaceNormal = true;
 
-	TArray<FMeshData*> MeshSettings;
-	TArray<FMaterialDataEx*> MatSettings;
+	TArray<FGLTFMeshRenderData*> MeshSettings;
+	TArray<FGLTFMaterialDataEx*> MatSettings;
 	MeshSettings.Add(&MeshSet);
 	MatSettings.Add(&MatSet);
 
-	TArray<FBakeOutputEx> BakeOutputs;
-	IMaterialBakingModule& Module = FModuleManager::Get().LoadModuleChecked<IMaterialBakingModule>("GLTFMaterialBaking");
+	TArray<FGLTFBakeOutputEx> BakeOutputs;
+	IGLTFMaterialBakingModule& Module = FModuleManager::Get().LoadModuleChecked<IGLTFMaterialBakingModule>("GLTFMaterialBaking");
 
 	Module.SetLinearBake(true);
 	Module.BakeMaterials(MatSettings, MeshSettings, BakeOutputs);
 	const bool bIsLinearBake = Module.IsLinearBake(Property);
 	Module.SetLinearBake(false);
 
-	FBakeOutputEx& BakeOutput = BakeOutputs[0];
+	FGLTFBakeOutputEx& BakeOutput = BakeOutputs[0];
 
 	TGLTFSharedArray<FColor> BakedPixels = MakeShared<TArray<FColor>>(MoveTemp(BakeOutput.PropertyData.FindChecked(Property)));
 	const FIntPoint BakedSize = BakeOutput.PropertySizes.FindChecked(Property);
@@ -317,7 +317,7 @@ bool FGLTFMaterialUtility::TryGetTextureCoordinateIndex(const UMaterialExpressio
 	return false;
 }
 
-void FGLTFMaterialUtility::GetAllTextureCoordinateIndices(const UMaterialInterface* InMaterial, const FMaterialPropertyEx& InProperty, FGLTFIndexArray& OutTexCoords)
+void FGLTFMaterialUtility::GetAllTextureCoordinateIndices(const UMaterialInterface* InMaterial, const FGLTFMaterialPropertyEx& InProperty, FGLTFIndexArray& OutTexCoords)
 {
 	FGLTFMaterialAnalysis Analysis;
 	AnalyzeMaterialProperty(InMaterial, InProperty, Analysis);
@@ -332,7 +332,7 @@ void FGLTFMaterialUtility::GetAllTextureCoordinateIndices(const UMaterialInterfa
 	}
 }
 
-void FGLTFMaterialUtility::AnalyzeMaterialProperty(const UMaterialInterface* InMaterial, const FMaterialPropertyEx& InProperty, FGLTFMaterialAnalysis& OutAnalysis)
+void FGLTFMaterialUtility::AnalyzeMaterialProperty(const UMaterialInterface* InMaterial, const FGLTFMaterialPropertyEx& InProperty, FGLTFMaterialAnalysis& OutAnalysis)
 {
 	if (GetInputForProperty(InMaterial, InProperty) == nullptr)
 	{
@@ -404,7 +404,7 @@ bool FGLTFMaterialUtility::NeedsMeshData(const UMaterialInterface* Material)
 	if (Material != nullptr && !FGLTFProxyMaterialUtilities::IsProxyMaterial(Material))
 	{
 		// TODO: only analyze properties that will be needed for this specific material
-		const TArray<FMaterialPropertyEx> Properties =
+		const TArray<FGLTFMaterialPropertyEx> Properties =
 		{
 			MP_BaseColor,
 			MP_EmissiveColor,
@@ -424,7 +424,7 @@ bool FGLTFMaterialUtility::NeedsMeshData(const UMaterialInterface* Material)
 
 		// TODO: optimize baking by separating need for vertex data and primitive data
 
-		for (const FMaterialPropertyEx& Property: Properties)
+		for (const FGLTFMaterialPropertyEx& Property: Properties)
 		{
 			AnalyzeMaterialProperty(Material, Property, Analysis);
 			bNeedsMeshData |= Analysis.bRequiresVertexData;
diff --git a/Engine/Plugins/Enterprise/GLTFExporter/MigratedFromGit/Source/GLTFExporter/Private/Converters/GLTFMaterialUtility.h b/Engine/Plugins/Enterprise/GLTFExporter/MigratedFromGit/Source/GLTFExporter/Private/Converters/GLTFMaterialUtility.h
--- a/Engine/Plugins/Enterprise/GLTFExporter/MigratedFromGit/Source/GLTFExporter/Private/Converters/GLTFMaterialUtility.h
+++ b/Engine/Plugins/Enterprise/GLTFExporter/MigratedFromGit/Source/GLTFExporter/Private/Converters/GLTFMaterialUtility.h
@@ -6,16 +6,16 @@
 #include "Converters/GLTFSharedArray.h"
 
 struct FGLTFMaterialAnalysis;
-struct FMaterialPropertyEx;
+struct FGLTFMaterialPropertyEx;
 class UMaterialInstanceConstant;
 
 struct FGLTFPropertyBakeOutput
 {
-	FGLTFPropertyBakeOutput(const FMaterialPropertyEx& Property, EPixelFormat PixelFormat, const TGLTFSharedArray<FColor>& Pixels, FIntPoint Size, float EmissiveScale, bool bIsSRGB)
+	FGLTFPropertyBakeOutput(const FGLTFMaterialPropertyEx& Property, EPixelFormat PixelFormat, const TGLTFSharedArray<FColor>& Pixels, FIntPoint Size, float EmissiveScale, bool bIsSRGB)
 		: Property(Property), PixelFormat(PixelFormat), Pixels(Pixels), Size(Size), EmissiveScale(EmissiveScale), bIsSRGB(bIsSRGB), bIsConstant(false)
 	{}
 
-	const FMaterialPropertyEx& Property;
+	const FGLTFMaterialPropertyEx& Property;
 	EPixelFormat PixelFormat;
 	TGLTFSharedArray<FColor> Pixels;
 	FIntPoint Size;
@@ -30,21 +30,21 @@ struct FGLTFMaterialUtility
 	static UMaterialInterface* GetDefaultMaterial();
 
 #if WITH_EDITOR
-	static bool IsNormalMap(const FMaterialPropertyEx& Property);
-	static bool IsSRGB(const FMaterialPropertyEx& Property);
+	static bool IsNormalMap(const FGLTFMaterialPropertyEx& Property);
+	static bool IsSRGB(const FGLTFMaterialPropertyEx& Property);
 
 #if (ENGINE_MAJOR_VERSION > 4 || ENGINE_MINOR_VERSION >= 27)
-	static FGuid GetAttributeID(const FMaterialPropertyEx& Property);
-	static FGuid GetAttributeIDChecked(const FMaterialPropertyEx& Property);
+	static FGuid GetAttributeID(const FGLTFMaterialPropertyEx& Property);
+	static FGuid GetAttributeIDChecked(const FGLTFMaterialPropertyEx& Property);
 #endif
 
-	static FVector4 GetPropertyDefaultValue(const FMaterialPropertyEx& Property);
-	static FVector4 GetPropertyMask(const FMaterialPropertyEx& Property);
+	static FVector4 GetPropertyDefaultValue(const FGLTFMaterialPropertyEx& Property);
+	static FVector4 GetPropertyMask(const FGLTFMaterialPropertyEx& Property);
 
-	static const FExpressionInput* GetInputForProperty(const UMaterialInterface* Material, const FMaterialPropertyEx& Property);
+	static const FExpressionInput* GetInputForProperty(const UMaterialInterface* Material, const FGLTFMaterialPropertyEx& Property);
 
 	template <class InputType>
-	static const FMaterialInput<InputType>* GetInputForProperty(const UMaterialInterface* Material, const FMaterialPropertyEx& Property)
+	static const FMaterialInput<InputType>* GetInputForProperty(const UMaterialInterface* Material, const FGLTFMaterialPropertyEx& Property)
 	{
 		const FExpressionInput* ExpressionInput = GetInputForProperty(Material, Property);
 		return static_cast<const FMaterialInput<InputType>*>(ExpressionInput);
@@ -52,7 +52,7 @@ struct FGLTFMaterialUtility
 
 	static const UMaterialExpressionCustomOutput* GetCustomOutputByName(const UMaterialInterface* Material, const FString& Name);
 
-	static FGLTFPropertyBakeOutput BakeMaterialProperty(const FIntPoint& OutputSize, const FMaterialPropertyEx& Property, const UMaterialInterface* Material, int32 TexCoord, const FGLTFMeshData* MeshData = nullptr, const FGLTFIndexArray& MeshSectionIndices = {}, bool bFillAlpha = true, bool bAdjustNormalmaps = true);
+	static FGLTFPropertyBakeOutput BakeMaterialProperty(const FIntPoint& OutputSize, const FGLTFMaterialPropertyEx& Property, const UMaterialInterface* Material, int32 TexCoord, const FGLTFMeshData* MeshData = nullptr, const FGLTFIndexArray& MeshSectionIndices = {}, bool bFillAlpha = true, bool bAdjustNormalmaps = true);
 
 	static FGLTFJsonTexture* AddTexture(FGLTFConvertBuilder& Builder, TGLTFSharedArray<FColor>& Pixels, const FIntPoint& TextureSize, bool bIgnoreAlpha, bool bIsNormalMap, const FString& TextureName, EGLTFJsonTextureFilter MinFilter, EGLTFJsonTextureFilter MagFilter, EGLTFJsonTextureWrap WrapS, EGLTFJsonTextureWrap WrapT);
 
@@ -60,9 +60,9 @@ struct FGLTFMaterialUtility
 	static uint32 GetMaskComponentCount(const FExpressionInput& ExpressionInput);
 
 	static bool TryGetTextureCoordinateIndex(const UMaterialExpressionTextureSample* TextureSampler, int32& TexCoord, FGLTFJsonTextureTransform& Transform);
-	static void GetAllTextureCoordinateIndices(const UMaterialInterface* InMaterial, const FMaterialPropertyEx& InProperty, FGLTFIndexArray& OutTexCoords);
+	static void GetAllTextureCoordinateIndices(const UMaterialInterface* InMaterial, const FGLTFMaterialPropertyEx& InProperty, FGLTFIndexArray& OutTexCoords);
 
-	static void AnalyzeMaterialProperty(const UMaterialInterface* Material, const FMaterialPropertyEx& InProperty, FGLTFMaterialAnalysis& OutAnalysis);
+	static void AnalyzeMaterialProperty(const UMaterialInterface* Material, const FGLTFMaterialPropertyEx& InProperty, FGLTFMaterialAnalysis& OutAnalysis);
 
 	static FMaterialShadingModelField EvaluateShadingModelExpression(const UMaterialInterface* Material);
 #endif
diff --git a/Engine/Plugins/Enterprise/GLTFExporter/MigratedFromGit/Source/GLTFExporter/Private/Converters/GLTFMeshData.h b/Engine/Plugins/Enterprise/GLTFExporter/MigratedFromGit/Source/GLTFExporter/Private/Converters/GLTFMeshData.h
--- a/Engine/Plugins/Enterprise/GLTFExporter/MigratedFromGit/Source/GLTFExporter/Private/Converters/GLTFMeshData.h
+++ b/Engine/Plugins/Enterprise/GLTFExporter/MigratedFromGit/Source/GLTFExporter/Private/Converters/GLTFMeshData.h
@@ -23,7 +23,7 @@ struct FGLTFMeshData
 	int32 LODIndex;
 
 #if WITH_EDITOR
-	FPrimitiveData PrimitiveData;
+	FGLTFPrimitiveData PrimitiveData;
 	FMeshDescription Description;
 
 	FLightMapRef LightMap;
diff --git a/Engine/Plugins/Enterprise/GLTFExporter/MigratedFromGit/Source/GLTFExporter/Private/Converters/GLTFUVOverlapChecker.cpp b/Engine/Plugins/Enterprise/GLTFExporter/MigratedFromGit/Source/GLTFExporter/Private/Converters/GLTFUVOverlapChecker.cpp
--- a/Engine/Plugins/Enterprise/GLTFExporter/MigratedFromGit/Source/GLTFExporter/Private/Converters/GLTFUVOverlapChecker.cpp
+++ b/Engine/Plugins/Enterprise/GLTFExporter/MigratedFromGit/Source/GLTFExporter/Private/Converters/GLTFUVOverlapChecker.cpp
@@ -5,7 +5,7 @@
 #include "StaticMeshAttributes.h"
 #include "MeshDescription.h"
 #include "Modules/ModuleManager.h"
-#include "IMaterialBakingModule.h"
+#include "IGLTFMaterialBakingModule.h"
 #include "MaterialBakingStructures.h"
 #endif
 
@@ -42,31 +42,31 @@ float FGLTFUVOverlapChecker::Convert(const FMeshDescription* Description, FGLTFI
 	{
 		// TODO: investigate if the fixed size is high enough to properly calculate overlap
 		const FIntPoint TextureSize(512, 512);
-		const FMaterialPropertyEx Property = MP_EmissiveColor;
+		const FGLTFMaterialPropertyEx Property = MP_EmissiveColor;
 
-		FMeshData MeshSet;
+		FGLTFMeshRenderData MeshSet;
 		MeshSet.TextureCoordinateBox = { { 0.0f, 0.0f }, { 1.0f, 1.0f } };
 		MeshSet.TextureCoordinateIndex = TexCoord;
 		MeshSet.RawMeshDescription = const_cast<FMeshDescription*>(Description);
 		MeshSet.MaterialIndices = SectionIndices; // NOTE: MaterialIndices is actually section indices
 
-		FMaterialDataEx MatSet;
+		FGLTFMaterialDataEx MatSet;
 		MatSet.Material = GetMaterial();
 		MatSet.PropertySizes.Add(Property, TextureSize);
 		MatSet.BlendMode = MatSet.Material->GetBlendMode();
 		MatSet.bPerformBorderSmear = false;
 
-		TArray<FMeshData*> MeshSettings;
-		TArray<FMaterialDataEx*> MatSettings;
+		TArray<FGLTFMeshRenderData*> MeshSettings;
+		TArray<FGLTFMaterialDataEx*> MatSettings;
 		MeshSettings.Add(&MeshSet);
 		MatSettings.Add(&MatSet);
 
-		TArray<FBakeOutputEx> BakeOutputs;
-		IMaterialBakingModule& Module = FModuleManager::Get().LoadModuleChecked<IMaterialBakingModule>("GLTFMaterialBaking");
+		TArray<FGLTFBakeOutputEx> BakeOutputs;
+		IGLTFMaterialBakingModule& Module = FModuleManager::Get().LoadModuleChecked<IGLTFMaterialBakingModule>("GLTFMaterialBaking");
 
 		Module.BakeMaterials(MatSettings, MeshSettings, BakeOutputs);
 
-		const FBakeOutputEx& BakeOutput = BakeOutputs[0];
+		const FGLTFBakeOutputEx& BakeOutput = BakeOutputs[0];
 		const TArray<FColor>& BakedPixels = BakeOutput.PropertyData.FindChecked(Property);
 
 		if (BakedPixels.Num() <= 0)
diff --git a/Engine/Plugins/Enterprise/GLTFExporter/MigratedFromGit/Source/GLTFExporter/Private/Tasks/GLTFMaterialTasks.cpp b/Engine/Plugins/Enterprise/GLTFExporter/MigratedFromGit/Source/GLTFExporter/Private/Tasks/GLTFMaterialTasks.cpp
--- a/Engine/Plugins/Enterprise/GLTFExporter/MigratedFromGit/Source/GLTFExporter/Private/Tasks/GLTFMaterialTasks.cpp
+++ b/Engine/Plugins/Enterprise/GLTFExporter/MigratedFromGit/Source/GLTFExporter/Private/Tasks/GLTFMaterialTasks.cpp
@@ -80,8 +80,8 @@ void FGLTFMaterialTask::Complete()
 #if WITH_EDITOR
 	if (JsonMaterial->ShadingModel != EGLTFJsonShadingModel::None)
 	{
-		const FMaterialPropertyEx BaseColorProperty = JsonMaterial->ShadingModel == EGLTFJsonShadingModel::Unlit ? MP_EmissiveColor : MP_BaseColor;
-		const FMaterialPropertyEx OpacityProperty = JsonMaterial->AlphaMode == EGLTFJsonAlphaMode::Mask ? MP_OpacityMask : MP_Opacity;
+		const FGLTFMaterialPropertyEx BaseColorProperty = JsonMaterial->ShadingModel == EGLTFJsonShadingModel::Unlit ? MP_EmissiveColor : MP_BaseColor;
+		const FGLTFMaterialPropertyEx OpacityProperty = JsonMaterial->AlphaMode == EGLTFJsonAlphaMode::Mask ? MP_OpacityMask : MP_Opacity;
 
 		// TODO: check if a property is active before trying to get it (i.e. Material->IsPropertyActive)
 
@@ -110,21 +110,21 @@ void FGLTFMaterialTask::Complete()
 
 		if (JsonMaterial->ShadingModel == EGLTFJsonShadingModel::Default || JsonMaterial->ShadingModel == EGLTFJsonShadingModel::ClearCoat)
 		{
-			const FMaterialPropertyEx MetallicProperty = MP_Metallic;
-			const FMaterialPropertyEx RoughnessProperty = MP_Roughness;
+			const FGLTFMaterialPropertyEx MetallicProperty = MP_Metallic;
+			const FGLTFMaterialPropertyEx RoughnessProperty = MP_Roughness;
 
 			if (!TryGetMetallicAndRoughness(JsonMaterial->PBRMetallicRoughness, MetallicProperty, RoughnessProperty))
 			{
 				Builder.LogWarning(FString::Printf(TEXT("Failed to export %s and %s for material %s"), *MetallicProperty.ToString(), *RoughnessProperty.ToString(), *Material->GetName()));
 			}
 
-			const FMaterialPropertyEx EmissiveProperty = MP_EmissiveColor;
+			const FGLTFMaterialPropertyEx EmissiveProperty = MP_EmissiveColor;
 			if (!TryGetEmissive(*JsonMaterial, EmissiveProperty))
 			{
 				Builder.LogWarning(FString::Printf(TEXT("Failed to export %s for material %s"), *EmissiveProperty.ToString(), *Material->GetName()));
 			}
 
-			const FMaterialPropertyEx NormalProperty = JsonMaterial->ShadingModel == EGLTFJsonShadingModel::ClearCoat ? FMaterialPropertyEx(TEXT("ClearCoatBottomNormal")) : FMaterialPropertyEx(MP_Normal);
+			const FGLTFMaterialPropertyEx NormalProperty = JsonMaterial->ShadingModel == EGLTFJsonShadingModel::ClearCoat ? FGLTFMaterialPropertyEx(TEXT("ClearCoatBottomNormal")) : FGLTFMaterialPropertyEx(MP_Normal);
 			if (IsPropertyNonDefault(NormalProperty))
 			{
 				if (!TryGetSourceTexture(JsonMaterial->NormalTexture, NormalProperty, DefaultColorInputMasks))
@@ -136,7 +136,7 @@ void FGLTFMaterialTask::Complete()
 				}
 			}
 
-			const FMaterialPropertyEx AmbientOcclusionProperty = MP_AmbientOcclusion;
+			const FGLTFMaterialPropertyEx AmbientOcclusionProperty = MP_AmbientOcclusion;
 			if (IsPropertyNonDefault(AmbientOcclusionProperty))
 			{
 				if (!TryGetSourceTexture(JsonMaterial->OcclusionTexture, AmbientOcclusionProperty, OcclusionInputMasks))
@@ -150,15 +150,15 @@ void FGLTFMaterialTask::Complete()
 
 			if (JsonMaterial->ShadingModel == EGLTFJsonShadingModel::ClearCoat)
 			{
-				const FMaterialPropertyEx ClearCoatProperty = MP_CustomData0;
-				const FMaterialPropertyEx ClearCoatRoughnessProperty = MP_CustomData1;
+				const FGLTFMaterialPropertyEx ClearCoatProperty = MP_CustomData0;
+				const FGLTFMaterialPropertyEx ClearCoatRoughnessProperty = MP_CustomData1;
 
 				if (!TryGetClearCoatRoughness(JsonMaterial->ClearCoat, ClearCoatProperty, ClearCoatRoughnessProperty))
 				{
 					Builder.LogWarning(FString::Printf(TEXT("Failed to export %s and %s for material %s"), *ClearCoatProperty.ToString(), *ClearCoatRoughnessProperty.ToString(), *Material->GetName()));
 				}
 
-				const FMaterialPropertyEx ClearCoatNormalProperty = MP_Normal;
+				const FGLTFMaterialPropertyEx ClearCoatNormalProperty = MP_Normal;
 				if (IsPropertyNonDefault(ClearCoatNormalProperty))
 				{
 					if (!TryGetSourceTexture(JsonMaterial->ClearCoat.ClearCoatNormalTexture, ClearCoatNormalProperty, DefaultColorInputMasks))
@@ -484,7 +484,7 @@ void FGLTFMaterialTask::ConvertAlphaMode(EGLTFJsonAlphaMode& OutAlphaMode, EGLTF
 
 #if WITH_EDITOR
 
-bool FGLTFMaterialTask::TryGetBaseColorAndOpacity(FGLTFJsonPBRMetallicRoughness& OutPBRParams, const FMaterialPropertyEx& BaseColorProperty, const FMaterialPropertyEx& OpacityProperty)
+bool FGLTFMaterialTask::TryGetBaseColorAndOpacity(FGLTFJsonPBRMetallicRoughness& OutPBRParams, const FGLTFMaterialPropertyEx& BaseColorProperty, const FGLTFMaterialPropertyEx& OpacityProperty)
 {
 	const bool bIsBaseColorConstant = TryGetConstantColor(OutPBRParams.BaseColorFactor, BaseColorProperty);
 	const bool bIsOpacityConstant = TryGetConstantScalar(OutPBRParams.BaseColorFactor.A, OpacityProperty);
@@ -609,7 +609,7 @@ bool FGLTFMaterialTask::TryGetBaseColorAndOpacity(FGLTFJsonPBRMetallicRoughness&
 	return true;
 }
 
-bool FGLTFMaterialTask::TryGetMetallicAndRoughness(FGLTFJsonPBRMetallicRoughness& OutPBRParams, const FMaterialPropertyEx& MetallicProperty, const FMaterialPropertyEx& RoughnessProperty)
+bool FGLTFMaterialTask::TryGetMetallicAndRoughness(FGLTFJsonPBRMetallicRoughness& OutPBRParams, const FGLTFMaterialPropertyEx& MetallicProperty, const FGLTFMaterialPropertyEx& RoughnessProperty)
 {
 	const bool bIsMetallicConstant = TryGetConstantScalar(OutPBRParams.MetallicFactor, MetallicProperty);
 	const bool bIsRoughnessConstant = TryGetConstantScalar(OutPBRParams.RoughnessFactor, RoughnessProperty);
@@ -732,7 +732,7 @@ bool FGLTFMaterialTask::TryGetMetallicAndRoughness(FGLTFJsonPBRMetallicRoughness
 	return true;
 }
 
-bool FGLTFMaterialTask::TryGetClearCoatRoughness(FGLTFJsonClearCoatExtension& OutExtParams, const FMaterialPropertyEx& IntensityProperty, const FMaterialPropertyEx& RoughnessProperty)
+bool FGLTFMaterialTask::TryGetClearCoatRoughness(FGLTFJsonClearCoatExtension& OutExtParams, const FGLTFMaterialPropertyEx& IntensityProperty, const FGLTFMaterialPropertyEx& RoughnessProperty)
 {
 	const bool bIsIntensityConstant = TryGetConstantScalar(OutExtParams.ClearCoatFactor, IntensityProperty);
 	const bool bIsRoughnessConstant = TryGetConstantScalar(OutExtParams.ClearCoatRoughnessFactor, RoughnessProperty);
@@ -859,7 +859,7 @@ bool FGLTFMaterialTask::TryGetClearCoatRoughness(FGLTFJsonClearCoatExtension& Ou
 	return true;
 }
 
-bool FGLTFMaterialTask::TryGetEmissive(FGLTFJsonMaterial& OutMaterial, const FMaterialPropertyEx& EmissiveProperty)
+bool FGLTFMaterialTask::TryGetEmissive(FGLTFJsonMaterial& OutMaterial, const FGLTFMaterialPropertyEx& EmissiveProperty)
 {
 	// TODO: right now we allow EmissiveFactor to be > 1.0 to support very bright emission, although it's not valid according to the glTF standard.
 	// We may want to change this behaviour and store factors above 1.0 using a custom extension instead.
@@ -911,7 +911,7 @@ bool FGLTFMaterialTask::TryGetEmissive(FGLTFJsonMaterial& OutMaterial, const FMa
 	return true;
 }
 
-bool FGLTFMaterialTask::IsPropertyNonDefault(const FMaterialPropertyEx& Property) const
+bool FGLTFMaterialTask::IsPropertyNonDefault(const FGLTFMaterialPropertyEx& Property) const
 {
 	const bool bUseMaterialAttributes = Material->GetMaterial()->bUseMaterialAttributes;
 	if (bUseMaterialAttributes)
@@ -936,7 +936,7 @@ bool FGLTFMaterialTask::IsPropertyNonDefault(const FMaterialPropertyEx& Property
 	return true;
 }
 
-bool FGLTFMaterialTask::TryGetConstantColor(FGLTFJsonColor3& OutValue, const FMaterialPropertyEx& Property) const
+bool FGLTFMaterialTask::TryGetConstantColor(FGLTFJsonColor3& OutValue, const FGLTFMaterialPropertyEx& Property) const
 {
 	FLinearColor Value;
 	if (TryGetConstantColor(Value, Property))
@@ -948,7 +948,7 @@ bool FGLTFMaterialTask::TryGetConstantColor(FGLTFJsonColor3& OutValue, const FMa
 	return false;
 }
 
-bool FGLTFMaterialTask::TryGetConstantColor(FGLTFJsonColor4& OutValue, const FMaterialPropertyEx& Property) const
+bool FGLTFMaterialTask::TryGetConstantColor(FGLTFJsonColor4& OutValue, const FGLTFMaterialPropertyEx& Property) const
 {
 	FLinearColor Value;
 	if (TryGetConstantColor(Value, Property))
@@ -960,7 +960,7 @@ bool FGLTFMaterialTask::TryGetConstantColor(FGLTFJsonColor4& OutValue, const FMa
 	return false;
 }
 
-bool FGLTFMaterialTask::TryGetConstantColor(FLinearColor& OutValue, const FMaterialPropertyEx& Property) const
+bool FGLTFMaterialTask::TryGetConstantColor(FLinearColor& OutValue, const FGLTFMaterialPropertyEx& Property) const
 {
 	const bool bUseMaterialAttributes = Material->GetMaterial()->bUseMaterialAttributes;
 	if (bUseMaterialAttributes)
@@ -1067,7 +1067,7 @@ bool FGLTFMaterialTask::TryGetConstantColor(FLinearColor& OutValue, const FMater
 	return false;
 }
 
-bool FGLTFMaterialTask::TryGetConstantScalar(float& OutValue, const FMaterialPropertyEx& Property) const
+bool FGLTFMaterialTask::TryGetConstantScalar(float& OutValue, const FGLTFMaterialPropertyEx& Property) const
 {
 	const bool bUseMaterialAttributes = Material->GetMaterial()->bUseMaterialAttributes;
 	if (bUseMaterialAttributes)
@@ -1168,7 +1168,7 @@ bool FGLTFMaterialTask::TryGetConstantScalar(float& OutValue, const FMaterialPro
 	return false;
 }
 
-bool FGLTFMaterialTask::TryGetSourceTexture(FGLTFJsonTextureInfo& OutTexInfo, const FMaterialPropertyEx& Property, const TArray<FLinearColor>& AllowedMasks) const
+bool FGLTFMaterialTask::TryGetSourceTexture(FGLTFJsonTextureInfo& OutTexInfo, const FGLTFMaterialPropertyEx& Property, const TArray<FLinearColor>& AllowedMasks) const
 {
 	const UTexture2D* Texture;
 	int32 TexCoord;
@@ -1185,7 +1185,7 @@ bool FGLTFMaterialTask::TryGetSourceTexture(FGLTFJsonTextureInfo& OutTexInfo, co
 	return false;
 }
 
-bool FGLTFMaterialTask::TryGetSourceTexture(const UTexture2D*& OutTexture, int32& OutTexCoord, FGLTFJsonTextureTransform& OutTransform, const FMaterialPropertyEx& Property, const TArray<FLinearColor>& AllowedMasks) const
+bool FGLTFMaterialTask::TryGetSourceTexture(const UTexture2D*& OutTexture, int32& OutTexCoord, FGLTFJsonTextureTransform& OutTransform, const FGLTFMaterialPropertyEx& Property, const TArray<FLinearColor>& AllowedMasks) const
 {
 	if (Builder.ExportOptions->TextureImageFormat == EGLTFTextureImageFormat::None)
 	{
@@ -1308,7 +1308,7 @@ bool FGLTFMaterialTask::TryGetSourceTexture(const UTexture2D*& OutTexture, int32
 	return false;
 }
 
-bool FGLTFMaterialTask::TryGetBakedMaterialProperty(FGLTFJsonTextureInfo& OutTexInfo, FGLTFJsonColor3& OutConstant, const FMaterialPropertyEx& Property, const FString& PropertyName)
+bool FGLTFMaterialTask::TryGetBakedMaterialProperty(FGLTFJsonTextureInfo& OutTexInfo, FGLTFJsonColor3& OutConstant, const FGLTFMaterialPropertyEx& Property, const FString& PropertyName)
 {
 	if (Builder.ExportOptions->BakeMaterialInputs == EGLTFMaterialBakeMode::Disabled)
 	{
@@ -1342,7 +1342,7 @@ bool FGLTFMaterialTask::TryGetBakedMaterialProperty(FGLTFJsonTextureInfo& OutTex
 	return false;
 }
 
-bool FGLTFMaterialTask::TryGetBakedMaterialProperty(FGLTFJsonTextureInfo& OutTexInfo, FGLTFJsonColor4& OutConstant, const FMaterialPropertyEx& Property, const FString& PropertyName)
+bool FGLTFMaterialTask::TryGetBakedMaterialProperty(FGLTFJsonTextureInfo& OutTexInfo, FGLTFJsonColor4& OutConstant, const FGLTFMaterialPropertyEx& Property, const FString& PropertyName)
 {
 	if (Builder.ExportOptions->BakeMaterialInputs == EGLTFMaterialBakeMode::Disabled)
 	{
@@ -1376,7 +1376,7 @@ bool FGLTFMaterialTask::TryGetBakedMaterialProperty(FGLTFJsonTextureInfo& OutTex
 	return false;
 }
 
-inline bool FGLTFMaterialTask::TryGetBakedMaterialProperty(FGLTFJsonTextureInfo& OutTexInfo, float& OutConstant, const FMaterialPropertyEx& Property, const FString& PropertyName)
+inline bool FGLTFMaterialTask::TryGetBakedMaterialProperty(FGLTFJsonTextureInfo& OutTexInfo, float& OutConstant, const FGLTFMaterialPropertyEx& Property, const FString& PropertyName)
 {
 	if (Builder.ExportOptions->BakeMaterialInputs == EGLTFMaterialBakeMode::Disabled)
 	{
@@ -1410,7 +1410,7 @@ inline bool FGLTFMaterialTask::TryGetBakedMaterialProperty(FGLTFJsonTextureInfo&
 	return false;
 }
 
-bool FGLTFMaterialTask::TryGetBakedMaterialProperty(FGLTFJsonTextureInfo& OutTexInfo, const FMaterialPropertyEx& Property, const FString& PropertyName)
+bool FGLTFMaterialTask::TryGetBakedMaterialProperty(FGLTFJsonTextureInfo& OutTexInfo, const FGLTFMaterialPropertyEx& Property, const FString& PropertyName)
 {
 	if (Builder.ExportOptions->BakeMaterialInputs == EGLTFMaterialBakeMode::Disabled)
 	{
@@ -1472,13 +1472,13 @@ bool FGLTFMaterialTask::TryGetBakedMaterialProperty(FGLTFJsonTextureInfo& OutTex
 	return true;
 }
 
-FGLTFPropertyBakeOutput FGLTFMaterialTask::BakeMaterialProperty(const FMaterialPropertyEx& Property, int32& OutTexCoord)
+FGLTFPropertyBakeOutput FGLTFMaterialTask::BakeMaterialProperty(const FGLTFMaterialPropertyEx& Property, int32& OutTexCoord)
 {
 	const FIntPoint TextureSize = Builder.GetBakeSizeForMaterialProperty(Material, GetPropertyGroup(Property));
 	return BakeMaterialProperty(Property, OutTexCoord, TextureSize, true);
 }
 
-FGLTFPropertyBakeOutput FGLTFMaterialTask::BakeMaterialProperty(const FMaterialPropertyEx& Property, int32& OutTexCoord, const FIntPoint& TextureSize, bool bFillAlpha)
+FGLTFPropertyBakeOutput FGLTFMaterialTask::BakeMaterialProperty(const FGLTFMaterialPropertyEx& Property, int32& OutTexCoord, const FIntPoint& TextureSize, bool bFillAlpha)
 {
 	if (MeshData == nullptr)
 	{
@@ -1549,7 +1549,7 @@ bool FGLTFMaterialTask::StoreBakedPropertyTexture(FGLTFJsonTextureInfo& OutTexIn
 	return true;
 }
 
-EGLTFMaterialPropertyGroup FGLTFMaterialTask::GetPropertyGroup(const FMaterialPropertyEx& Property)
+EGLTFMaterialPropertyGroup FGLTFMaterialTask::GetPropertyGroup(const FGLTFMaterialPropertyEx& Property)
 {
 	switch (Property.Type)
 	{
diff --git a/Engine/Plugins/Enterprise/GLTFExporter/MigratedFromGit/Source/GLTFExporter/Private/Tasks/GLTFMaterialTasks.h b/Engine/Plugins/Enterprise/GLTFExporter/MigratedFromGit/Source/GLTFExporter/Private/Tasks/GLTFMaterialTasks.h
--- a/Engine/Plugins/Enterprise/GLTFExporter/MigratedFromGit/Source/GLTFExporter/Private/Tasks/GLTFMaterialTasks.h
+++ b/Engine/Plugins/Enterprise/GLTFExporter/MigratedFromGit/Source/GLTFExporter/Private/Tasks/GLTFMaterialTasks.h
@@ -8,7 +8,7 @@
 #include "Converters/GLTFMeshData.h"
 #include "Materials/GLTFProxyMaterialParameterInfo.h"
 #if WITH_EDITOR
-#include "MaterialPropertyEx.h"
+#include "GLTFMaterialPropertyEx.h"
 #endif
 
 struct FGLTFPropertyBakeOutput;
@@ -58,33 +58,33 @@ private:
 	void ConvertAlphaMode(EGLTFJsonAlphaMode& OutAlphaMode, EGLTFJsonBlendMode& OutBlendMode) const;
 
 #if WITH_EDITOR
-	TSet<FMaterialPropertyEx> MeshDataBakedProperties;
+	TSet<FGLTFMaterialPropertyEx> MeshDataBakedProperties;
 
-	bool TryGetBaseColorAndOpacity(FGLTFJsonPBRMetallicRoughness& OutPBRParams, const FMaterialPropertyEx& BaseColorProperty, const FMaterialPropertyEx& OpacityProperty);
-	bool TryGetMetallicAndRoughness(FGLTFJsonPBRMetallicRoughness& OutPBRParams, const FMaterialPropertyEx& MetallicProperty, const FMaterialPropertyEx& RoughnessProperty);
-	bool TryGetClearCoatRoughness(FGLTFJsonClearCoatExtension& OutExtParams, const FMaterialPropertyEx& IntensityProperty, const FMaterialPropertyEx& RoughnessProperty);
-	bool TryGetEmissive(FGLTFJsonMaterial& OutMaterial, const FMaterialPropertyEx& EmissiveProperty);
+	bool TryGetBaseColorAndOpacity(FGLTFJsonPBRMetallicRoughness& OutPBRParams, const FGLTFMaterialPropertyEx& BaseColorProperty, const FGLTFMaterialPropertyEx& OpacityProperty);
+	bool TryGetMetallicAndRoughness(FGLTFJsonPBRMetallicRoughness& OutPBRParams, const FGLTFMaterialPropertyEx& MetallicProperty, const FGLTFMaterialPropertyEx& RoughnessProperty);
+	bool TryGetClearCoatRoughness(FGLTFJsonClearCoatExtension& OutExtParams, const FGLTFMaterialPropertyEx& IntensityProperty, const FGLTFMaterialPropertyEx& RoughnessProperty);
+	bool TryGetEmissive(FGLTFJsonMaterial& OutMaterial, const FGLTFMaterialPropertyEx& EmissiveProperty);
 
-	bool IsPropertyNonDefault(const FMaterialPropertyEx& Property) const;
-	bool TryGetConstantColor(FGLTFJsonColor3& OutValue, const FMaterialPropertyEx& Property) const;
-	bool TryGetConstantColor(FGLTFJsonColor4& OutValue, const FMaterialPropertyEx& Property) const;
-	bool TryGetConstantColor(FLinearColor& OutValue, const FMaterialPropertyEx& Property) const;
-	bool TryGetConstantScalar(float& OutValue, const FMaterialPropertyEx& Property) const;
+	bool IsPropertyNonDefault(const FGLTFMaterialPropertyEx& Property) const;
+	bool TryGetConstantColor(FGLTFJsonColor3& OutValue, const FGLTFMaterialPropertyEx& Property) const;
+	bool TryGetConstantColor(FGLTFJsonColor4& OutValue, const FGLTFMaterialPropertyEx& Property) const;
+	bool TryGetConstantColor(FLinearColor& OutValue, const FGLTFMaterialPropertyEx& Property) const;
+	bool TryGetConstantScalar(float& OutValue, const FGLTFMaterialPropertyEx& Property) const;
 
-	bool TryGetSourceTexture(FGLTFJsonTextureInfo& OutTexInfo, const FMaterialPropertyEx& Property, const TArray<FLinearColor>& AllowedMasks = {}) const;
-	bool TryGetSourceTexture(const UTexture2D*& OutTexture, int32& OutTexCoord, FGLTFJsonTextureTransform& OutTransform, const FMaterialPropertyEx& Property, const TArray<FLinearColor>& AllowedMasks = {}) const;
+	bool TryGetSourceTexture(FGLTFJsonTextureInfo& OutTexInfo, const FGLTFMaterialPropertyEx& Property, const TArray<FLinearColor>& AllowedMasks = {}) const;
+	bool TryGetSourceTexture(const UTexture2D*& OutTexture, int32& OutTexCoord, FGLTFJsonTextureTransform& OutTransform, const FGLTFMaterialPropertyEx& Property, const TArray<FLinearColor>& AllowedMasks = {}) const;
 
-	bool TryGetBakedMaterialProperty(FGLTFJsonTextureInfo& OutTexInfo, FGLTFJsonColor3& OutConstant, const FMaterialPropertyEx& Property, const FString& PropertyName);
-	bool TryGetBakedMaterialProperty(FGLTFJsonTextureInfo& OutTexInfo, FGLTFJsonColor4& OutConstant, const FMaterialPropertyEx& Property, const FString& PropertyName);
-	bool TryGetBakedMaterialProperty(FGLTFJsonTextureInfo& OutTexInfo, float& OutConstant, const FMaterialPropertyEx& Property, const FString& PropertyName);
-	bool TryGetBakedMaterialProperty(FGLTFJsonTextureInfo& OutTexInfo, const FMaterialPropertyEx& Property, const FString& PropertyName);
+	bool TryGetBakedMaterialProperty(FGLTFJsonTextureInfo& OutTexInfo, FGLTFJsonColor3& OutConstant, const FGLTFMaterialPropertyEx& Property, const FString& PropertyName);
+	bool TryGetBakedMaterialProperty(FGLTFJsonTextureInfo& OutTexInfo, FGLTFJsonColor4& OutConstant, const FGLTFMaterialPropertyEx& Property, const FString& PropertyName);
+	bool TryGetBakedMaterialProperty(FGLTFJsonTextureInfo& OutTexInfo, float& OutConstant, const FGLTFMaterialPropertyEx& Property, const FString& PropertyName);
+	bool TryGetBakedMaterialProperty(FGLTFJsonTextureInfo& OutTexInfo, const FGLTFMaterialPropertyEx& Property, const FString& PropertyName);
 
-	FGLTFPropertyBakeOutput BakeMaterialProperty(const FMaterialPropertyEx& Property, int32& OutTexCoord);
-	FGLTFPropertyBakeOutput BakeMaterialProperty(const FMaterialPropertyEx& Property, int32& OutTexCoord, const FIntPoint& TextureSize, bool bFillAlpha);
+	FGLTFPropertyBakeOutput BakeMaterialProperty(const FGLTFMaterialPropertyEx& Property, int32& OutTexCoord);
+	FGLTFPropertyBakeOutput BakeMaterialProperty(const FGLTFMaterialPropertyEx& Property, int32& OutTexCoord, const FIntPoint& TextureSize, bool bFillAlpha);
 
 	bool StoreBakedPropertyTexture(FGLTFJsonTextureInfo& OutTexInfo, FGLTFPropertyBakeOutput& PropertyBakeOutput, const FString& PropertyName) const;
 
-	static EGLTFMaterialPropertyGroup GetPropertyGroup(const FMaterialPropertyEx& Property);
+	static EGLTFMaterialPropertyGroup GetPropertyGroup(const FGLTFMaterialPropertyEx& Property);
 
 	template <typename CallbackType>
 	static void CombinePixels(const TArray<FColor>& FirstPixels, const TArray<FColor>& SecondPixels, TArray<FColor>& OutPixels, CallbackType Callback);
diff --git a/Engine/Plugins/Enterprise/GLTFExporter/MigratedFromGit/Source/GLTFMaterialBaking/Private/ExportMaterialProxy.h b/Engine/Plugins/Enterprise/GLTFExporter/MigratedFromGit/Source/GLTFMaterialBaking/Private/ExportMaterialProxy.h
deleted file mode 100644
diff --git a/Engine/Plugins/Enterprise/GLTFExporter/MigratedFromGit/Source/GLTFMaterialBaking/Private/FMaterialPropertyEx.cpp b/Engine/Plugins/Enterprise/GLTFExporter/MigratedFromGit/Source/GLTFMaterialBaking/Private/FMaterialPropertyEx.cpp
deleted file mode 100644
diff --git a/Engine/Plugins/Enterprise/GLTFExporter/MigratedFromGit/Source/GLTFMaterialBaking/Private/GLTFExportMaterialProxy.h b/Engine/Plugins/Enterprise/GLTFExporter/MigratedFromGit/Source/GLTFMaterialBaking/Private/GLTFExportMaterialProxy.h
new file mode 100644
--- /dev/null
+++ b/Engine/Plugins/Enterprise/GLTFExporter/MigratedFromGit/Source/GLTFMaterialBaking/Private/GLTFExportMaterialProxy.h
@@ -0,0 +1,783 @@
+// Copyright Epic Games, Inc. All Rights Reserved.
+
+#pragma once
+
+#include "MaterialShared.h"
+#include "MaterialCompiler.h"
+#include "Materials/MaterialParameterCollection.h"
+
+#include "Engine/TextureLODSettings.h"
+#include "Engine/Texture2D.h"
+#include "Engine/Texture.h"
+#include "Engine/TextureCube.h"
+#include "Engine/Texture2DArray.h"
+
+#include "DeviceProfiles/DeviceProfileManager.h"
+#include "DeviceProfiles/DeviceProfile.h"
+#include "Materials/MaterialInterface.h"
+#include "SceneTypes.h"
+#include "Materials/Material.h"
+#include "Materials/MaterialExpressionCustomOutput.h"
+
+#include "Materials/HLSLMaterialTranslator.h"
+
+struct FGLTFExportMaterialCompiler : public FProxyMaterialCompiler
+{
+	FGLTFExportMaterialCompiler(FMaterialCompiler* InCompiler) :
+		FProxyMaterialCompiler(InCompiler)
+	{}
+
+	// gets value stored by SetMaterialProperty()
+	virtual EShaderFrequency GetCurrentShaderFrequency() const override
+	{
+		return SF_Pixel;
+	}
+
+	virtual FMaterialShadingModelField GetMaterialShadingModels() const override
+	{
+		return MSM_MAX;
+	}
+
+	virtual int32 WorldPosition(EWorldPositionIncludedOffsets WorldPositionIncludedOffsets) override
+	{
+#if WITH_EDITOR
+		return Compiler->MaterialBakingWorldPosition();
+#else
+		return Compiler->WorldPosition(WorldPositionIncludedOffsets);
+#endif
+	}
+
+	virtual int32 DistanceCullFade() override
+	{
+		return Compiler->Constant(1.0f);
+	}
+
+	virtual int32 ParticleRelativeTime() override
+	{
+		return Compiler->Constant(0.0f);
+	}
+
+	virtual int32 ParticleMotionBlurFade() override
+	{
+		return Compiler->Constant(1.0f);
+	}
+
+	virtual int32 PixelNormalWS() override
+	{
+		// Current returning vertex normal since pixel normal will contain incorrect data (normal calculated from uv data used as vertex positions to render out the material)
+		return Compiler->VertexNormal();
+	}
+
+	virtual int32 ParticleRandom() override
+	{
+		return Compiler->Constant(0.0f);
+	}
+
+	virtual int32 ParticleDirection() override
+	{
+		return Compiler->Constant3(0.0f, 0.0f, 0.0f);
+	}
+
+	virtual int32 ParticleSpeed() override
+	{
+		return Compiler->Constant(0.0f);
+	}
+
+	virtual int32 ParticleSize() override
+	{
+		return Compiler->Constant2(0.0f, 0.0f);
+	}
+
+	virtual int32 CameraVector() override
+	{
+		// NOTE: by using VertexNormal instead of a fixed (world-space) vector, we allow
+		// more correct baking of materials that use the dot-product between the vertex normal
+		// and the camera vector for effects such as fresnel.
+		return Compiler->VertexNormal();
+	}
+
+	virtual int32 ReflectionAboutCustomWorldNormal(int32 CustomWorldNormal, int32 bNormalizeCustomWorldNormal) override
+	{
+		if (CustomWorldNormal == INDEX_NONE)
+		{
+			return INDEX_NONE;
+		}
+
+		int32 N = CustomWorldNormal;
+		int32 C = CameraVector();
+
+		if (bNormalizeCustomWorldNormal)
+		{
+			// N = N / sqrt(dot(N, N))
+			N = Compiler->Div(N, Compiler->SquareRoot(Compiler->Dot(N, N)));
+		}
+
+		// return 2 * dot(N, C) * N - C
+		return Compiler->Sub(Compiler->Mul(Compiler->Constant(2.0f), Compiler->Mul(Compiler->Dot(N, C), N)), C);
+	}
+
+#if (ENGINE_MAJOR_VERSION > 4 || ENGINE_MINOR_VERSION >= 26)
+	virtual int32 PreSkinVertexOffset() override
+	{
+		return Compiler->PreSkinVertexOffset();
+	}
+
+	virtual int32 PostSkinVertexOffset() override
+	{
+		return Compiler->PostSkinVertexOffset();
+	}
+#endif
+
+	virtual int32 PreSkinnedPosition() override
+	{
+		return Compiler->PreSkinnedPosition();
+	}
+
+	virtual int32 PreSkinnedNormal() override
+	{
+		return Compiler->PreSkinnedNormal();
+	}
+
+	virtual int32 VertexInterpolator(uint32 InterpolatorIndex) override
+	{
+		return Compiler->VertexInterpolator(InterpolatorIndex);
+	}
+
+	virtual int32 ReflectionVector() override
+	{
+		// Because camera vector is identical to normal vector we can work out that reflection vector will also be the same
+		return Compiler->VertexNormal();
+	}
+
+#if WITH_EDITOR
+	virtual int32 MaterialBakingWorldPosition() override
+	{
+		return Compiler->MaterialBakingWorldPosition();
+	}
+#endif
+
+	virtual int32 AccessCollectionParameter(UMaterialParameterCollection* ParameterCollection, int32 ParameterIndex, int32 ComponentIndex) override
+	{
+		if (!ParameterCollection || ParameterIndex == -1)
+		{
+			return INDEX_NONE;
+		}
+
+		// Collect names of all parameters
+		TArray<FName> ParameterNames;
+		ParameterCollection->GetParameterNames(ParameterNames, /*bVectorParameters=*/ false);
+		int32 NumScalarParameters = ParameterNames.Num();
+		ParameterCollection->GetParameterNames(ParameterNames, /*bVectorParameters=*/ true);
+
+		// Find a parameter corresponding to ParameterIndex/ComponentIndex pair
+		int32 Index;
+		for (Index = 0; Index < ParameterNames.Num(); Index++)
+		{
+			FGuid ParameterId = ParameterCollection->GetParameterId(ParameterNames[Index]);
+			int32 CheckParameterIndex, CheckComponentIndex;
+			ParameterCollection->GetParameterIndex(ParameterId, CheckParameterIndex, CheckComponentIndex);
+			if (CheckParameterIndex == ParameterIndex && CheckComponentIndex == ComponentIndex)
+			{
+				// Found
+				break;
+			}
+		}
+		if (Index >= ParameterNames.Num())
+		{
+			// Not found, should not happen
+			return INDEX_NONE;
+		}
+
+		// Create code for parameter
+		if (Index < NumScalarParameters)
+		{
+			const FCollectionScalarParameter* ScalarParameter = ParameterCollection->GetScalarParameterByName(ParameterNames[Index]);
+			check(ScalarParameter);
+			return Constant(ScalarParameter->DefaultValue);
+		}
+		else
+		{
+			const FCollectionVectorParameter* VectorParameter = ParameterCollection->GetVectorParameterByName(ParameterNames[Index]);
+			check(VectorParameter);
+			const FLinearColor& Color = VectorParameter->DefaultValue;
+			return Constant4(Color.R, Color.G, Color.B, Color.A);
+		}
+	}
+
+	virtual EMaterialCompilerType GetCompilerType() const override { return EMaterialCompilerType::MaterialProxy; }
+};
+
+class FGLTFExportMaterialProxy : public FMaterial, public FMaterialRenderProxy
+{
+public:
+	FGLTFExportMaterialProxy(UMaterialInterface* InMaterialInterface, EMaterialProperty InPropertyToCompile, const FString& InCustomOutputToCompile = TEXT(""), bool bInSynchronousCompilation = true, EBlendMode ProxyBlendMode = BLEND_Opaque, bool bTangentSpaceNormal = false)
+		: FMaterial()
+		, MaterialInterface(InMaterialInterface)
+		, PropertyToCompile(InPropertyToCompile)
+		, CustomOutputToCompile(InCustomOutputToCompile)
+		, bSynchronousCompilation(bInSynchronousCompilation)
+		, ProxyBlendMode(ProxyBlendMode)
+		, bTangentSpaceNormal(bTangentSpaceNormal)
+	{
+
+#if (ENGINE_MAJOR_VERSION > 4 || ENGINE_MINOR_VERSION >= 26)
+		SetQualityLevelProperties(GMaxRHIFeatureLevel);
+#else
+		SetQualityLevelProperties(EMaterialQualityLevel::High, false, GMaxRHIFeatureLevel);
+#endif
+		Material = InMaterialInterface->GetMaterial();
+		ReferencedTextures = InMaterialInterface->GetReferencedTextures();
+
+		const FMaterialResource* Resource = InMaterialInterface->GetMaterialResource(GMaxRHIFeatureLevel);
+
+		FMaterialShaderMapId ResourceId;
+		Resource->GetShaderMapId(GMaxRHIShaderPlatform, nullptr, ResourceId);
+
+		// Our Id must be the same as BaseMaterialId for the shader compiler
+		// to be able to set back GameThreadShaderMap after async compilation.
+		Id = ResourceId.BaseMaterialId;
+
+
+		{
+			TArray<FShaderType*> ShaderTypes;
+			TArray<FVertexFactoryType*> VFTypes;
+			TArray<const FShaderPipelineType*> ShaderPipelineTypes;
+
+#if (ENGINE_MAJOR_VERSION > 4 || ENGINE_MINOR_VERSION >= 27)
+			GetDependentShaderAndVFTypes(GMaxRHIShaderPlatform, ResourceId.LayoutParams, ShaderTypes, ShaderPipelineTypes, VFTypes);
+#else
+			GetDependentShaderAndVFTypes(GMaxRHIShaderPlatform, ShaderTypes, ShaderPipelineTypes, VFTypes);
+#endif
+
+			// Overwrite the shader map Id's dependencies with ones that came from the FMaterial actually being compiled (this)
+			// This is necessary as we change FMaterial attributes like GetShadingModels(), which factor into the ShouldCache functions that determine dependent shader types
+			ResourceId.SetShaderDependencies(ShaderTypes, ShaderPipelineTypes, VFTypes, GMaxRHIShaderPlatform);
+		}
+
+		// Override with a special usage so we won't re-use the shader map used by the material for rendering
+		switch (InPropertyToCompile)
+		{
+		case MP_BaseColor: ResourceId.Usage = EMaterialShaderMapUsage::MaterialExportBaseColor; break;
+		case MP_Specular: ResourceId.Usage = EMaterialShaderMapUsage::MaterialExportSpecular; break;
+		case MP_Normal: ResourceId.Usage = EMaterialShaderMapUsage::MaterialExportNormal; break;
+		case MP_Tangent: ResourceId.Usage = EMaterialShaderMapUsage::MaterialExportTangent; break;
+		case MP_Metallic: ResourceId.Usage = EMaterialShaderMapUsage::MaterialExportMetallic; break;
+		case MP_Roughness: ResourceId.Usage = EMaterialShaderMapUsage::MaterialExportRoughness; break;
+		case MP_Anisotropy: ResourceId.Usage = EMaterialShaderMapUsage::MaterialExportAnisotropy; break;
+		case MP_AmbientOcclusion: ResourceId.Usage = EMaterialShaderMapUsage::MaterialExportAO; break;
+		case MP_EmissiveColor: ResourceId.Usage = EMaterialShaderMapUsage::MaterialExportEmissive; break;
+		case MP_Opacity: ResourceId.Usage = EMaterialShaderMapUsage::MaterialExportOpacity; break;
+		case MP_OpacityMask: ResourceId.Usage = EMaterialShaderMapUsage::MaterialExportOpacityMask; break;
+		case MP_SubsurfaceColor: ResourceId.Usage = EMaterialShaderMapUsage::MaterialExportSubSurfaceColor; break;
+
+#if (ENGINE_MAJOR_VERSION > 4 || ENGINE_MINOR_VERSION >= 27)
+		case MP_CustomData0: ResourceId.Usage = EMaterialShaderMapUsage::MaterialExportClearCoat; break;
+		case MP_CustomData1: ResourceId.Usage = EMaterialShaderMapUsage::MaterialExportClearCoatRoughness; break;
+		case MP_CustomOutput:
+			ResourceId.Usage = EMaterialShaderMapUsage::MaterialExportCustomOutput;
+			ResourceId.UsageCustomOutput = InCustomOutputToCompile;
+			break;
+		case MP_Refraction: ResourceId.Usage = static_cast<EMaterialShaderMapUsage::Type>(EMaterialShaderMapUsage::MaterialExportCustomOutput + 1); break;
+		case MP_ShadingModel: ResourceId.Usage = static_cast<EMaterialShaderMapUsage::Type>(EMaterialShaderMapUsage::MaterialExportCustomOutput + 2); break;
+#else
+		// NOTE: the following cases for custom data 0/1 and custom output are hacks since a proper solution requires engine changes:
+		case MP_CustomData0: ResourceId.Usage = static_cast<EMaterialShaderMapUsage::Type>(EMaterialShaderMapUsage::DebugViewMode + 1); break;
+		case MP_CustomData1: ResourceId.Usage = static_cast<EMaterialShaderMapUsage::Type>(EMaterialShaderMapUsage::DebugViewMode + 2); break;
+		case MP_CustomOutput: ResourceId.Usage = static_cast<EMaterialShaderMapUsage::Type>(EMaterialShaderMapUsage::DebugViewMode + 3); break;
+		case MP_Refraction: ResourceId.Usage = static_cast<EMaterialShaderMapUsage::Type>(EMaterialShaderMapUsage::DebugViewMode + 4); break;
+		case MP_ShadingModel: ResourceId.Usage = static_cast<EMaterialShaderMapUsage::Type>(EMaterialShaderMapUsage::DebugViewMode + 5); break;
+#endif
+
+		default:
+			ensureMsgf(false, TEXT("ExportMaterial has no usage for property %i.  Will likely reuse the normal rendering shader and crash later with a parameter mismatch"), (int32)InPropertyToCompile);
+			break;
+		};
+
+		CacheShaders(ResourceId, GMaxRHIShaderPlatform);
+	}
+
+	/** This override is required otherwise the shaders aren't ready for use when the surface is rendered resulting in a blank image */
+	virtual bool RequiresSynchronousCompilation() const override { return bSynchronousCompilation; };
+
+	/**
+	* Should the shader for this material with the given platform, shader type and vertex
+	* factory type combination be compiled
+	*
+	* @param Platform		The platform currently being compiled for
+	* @param ShaderType	Which shader is being compiled
+	* @param VertexFactory	Which vertex factory is being compiled (can be NULL)
+	*
+	* @return true if the shader should be compiled
+	*/
+	virtual bool ShouldCache(EShaderPlatform Platform, const FShaderType* ShaderType, const FVertexFactoryType* VertexFactoryType) const override
+	{
+		const bool bCorrectVertexFactory = VertexFactoryType == FindVertexFactoryType(FName(TEXT("FLocalVertexFactory"), FNAME_Find));
+		const bool bPCPlatform = !IsConsolePlatform(Platform);
+		const bool bCorrectFrequency = ShaderType->GetFrequency() == SF_Vertex || ShaderType->GetFrequency() == SF_Pixel;
+		return bCorrectVertexFactory && bPCPlatform && bCorrectFrequency;
+	}
+
+	virtual TArrayView<UObject* const> GetReferencedTextures() const override
+	{
+		return ReferencedTextures;
+	}
+
+	virtual void GetStaticParameterSet(EShaderPlatform Platform, FStaticParameterSet& OutSet) const override
+	{
+		if (const FMaterialResource* Resource = MaterialInterface->GetMaterialResource(GMaxRHIFeatureLevel))
+		{
+			Resource->GetStaticParameterSet(Platform, OutSet);
+		}
+	}
+
+	////////////////
+	// FMaterialRenderProxy interface.
+
+#if (ENGINE_MAJOR_VERSION > 4 || ENGINE_MINOR_VERSION >= 27)
+	virtual const FMaterial* GetMaterialNoFallback(ERHIFeatureLevel::Type InFeatureLevel) const override
+	{
+		if (GetRenderingThreadShaderMap())
+		{
+			return this;
+		}
+		return nullptr;
+	}
+
+	virtual const FMaterialRenderProxy* GetFallback(ERHIFeatureLevel::Type InFeatureLevel) const override
+	{
+		return UMaterial::GetDefaultMaterial(MD_Surface)->GetRenderProxy();
+	}
+#else
+	virtual const FMaterial& GetMaterialWithFallback(ERHIFeatureLevel::Type FeatureLevel, const FMaterialRenderProxy*& OutFallbackMaterialRenderProxy) const override
+	{
+		if(GetRenderingThreadShaderMap())
+		{
+			return *this;
+		}
+		else
+		{
+			OutFallbackMaterialRenderProxy = UMaterial::GetDefaultMaterial(MD_Surface)->GetRenderProxy();
+			return OutFallbackMaterialRenderProxy->GetMaterialWithFallback(FeatureLevel, OutFallbackMaterialRenderProxy);
+		}
+	}
+#endif
+
+	virtual bool GetVectorValue(const FHashedMaterialParameterInfo& ParameterInfo, FLinearColor* OutValue, const FMaterialRenderContext& Context) const override
+	{
+		return MaterialInterface->GetRenderProxy()->GetVectorValue(ParameterInfo, OutValue, Context);
+	}
+
+	virtual bool GetScalarValue(const FHashedMaterialParameterInfo& ParameterInfo, float* OutValue, const FMaterialRenderContext& Context) const override
+	{
+		return MaterialInterface->GetRenderProxy()->GetScalarValue(ParameterInfo, OutValue, Context);
+	}
+
+	virtual bool GetTextureValue(const FHashedMaterialParameterInfo& ParameterInfo, const UTexture** OutValue, const FMaterialRenderContext& Context) const override
+	{
+		return MaterialInterface->GetRenderProxy()->GetTextureValue(ParameterInfo, OutValue, Context);
+	}
+
+	virtual bool GetTextureValue(const FHashedMaterialParameterInfo& ParameterInfo, const URuntimeVirtualTexture** OutValue, const FMaterialRenderContext& Context) const override
+	{
+		return MaterialInterface->GetRenderProxy()->GetTextureValue(ParameterInfo, OutValue, Context);
+	}
+
+	// Material properties.
+	/** Entry point for compiling a specific material property.  This must call SetMaterialProperty. */
+	virtual int32 CompilePropertyAndSetMaterialProperty(EMaterialProperty Property, FMaterialCompiler* Compiler, EShaderFrequency OverrideShaderFrequency, bool bUsePreviousFrameTime) const override
+	{
+		// needs to be called in this function!!
+		Compiler->SetMaterialProperty(Property, OverrideShaderFrequency, bUsePreviousFrameTime);
+		const int32 Ret = CompilePropertyAndSetMaterialPropertyWithoutCast(Property, Compiler);
+		return Compiler->ForceCast(Ret, FMaterialAttributeDefinitionMap::GetValueType(Property), MFCF_ExactMatch | MFCF_ReplicateValue);
+	}
+
+	/** helper for CompilePropertyAndSetMaterialProperty() */
+	int32 CompilePropertyAndSetMaterialPropertyWithoutCast(EMaterialProperty Property, FMaterialCompiler* Compiler) const
+	{
+		if (Property == MP_EmissiveColor)
+		{
+			const EBlendMode BlendMode = MaterialInterface->GetBlendMode();
+			FGLTFExportMaterialCompiler ProxyCompiler(Compiler);
+			const uint32 ForceCast_Exact_Replicate = MFCF_ForceCast | MFCF_ExactMatch | MFCF_ReplicateValue;
+
+			switch (PropertyToCompile)
+			{
+			case MP_EmissiveColor:
+				// Emissive is ALWAYS returned...
+				return MaterialInterface->CompileProperty(&ProxyCompiler, MP_EmissiveColor, ForceCast_Exact_Replicate);
+			case MP_BaseColor:
+				return MaterialInterface->CompileProperty(&ProxyCompiler, MP_BaseColor, ForceCast_Exact_Replicate);
+				break;
+			case MP_Specular:
+			case MP_Roughness:
+			case MP_Anisotropy:
+			case MP_Metallic:
+			case MP_AmbientOcclusion:
+				// Only return for Opaque and Masked...
+				if (BlendMode == BLEND_Opaque || BlendMode == BLEND_Masked)
+				{
+					return MaterialInterface->CompileProperty(&ProxyCompiler, PropertyToCompile, ForceCast_Exact_Replicate);
+				}
+				break;
+
+			case MP_Opacity:
+			case MP_OpacityMask:
+			case MP_SubsurfaceColor:
+			case MP_CustomData0:
+			case MP_CustomData1:
+			{
+				return MaterialInterface->CompileProperty(&ProxyCompiler, PropertyToCompile, ForceCast_Exact_Replicate);
+			}
+			case MP_Normal:
+			case MP_Tangent:
+				// Only return for Opaque and Masked...
+				if (BlendMode == BLEND_Opaque || BlendMode == BLEND_Masked)
+				{
+					return CompileNormalEncoding(
+						&ProxyCompiler,
+						CompileNormalTransform(&ProxyCompiler, MaterialInterface->CompileProperty(&ProxyCompiler, PropertyToCompile, ForceCast_Exact_Replicate)));
+				}
+				break;
+			case MP_Refraction:
+					// Only return for Translucent, Additive, and AlphaComposite...
+					if (BlendMode == BLEND_Translucent || BlendMode == BLEND_Additive || BlendMode == BLEND_AlphaComposite)
+					{
+						return CompileRefractionEncoding(
+							Compiler,
+							MaterialInterface->CompileProperty(&ProxyCompiler, MP_Refraction, ForceCast_Exact_Replicate));
+					}
+				break;
+			case MP_ShadingModel:
+				return CompileShadingModelEncoding(Compiler, MaterialInterface->CompileProperty(&ProxyCompiler, MP_ShadingModel));
+			case MP_CustomOutput:
+				 // NOTE: Currently we can assume input index is always 0, which it is for all custom outputs that are registered as material attributes
+				return CompileInputForCustomOutput(&ProxyCompiler, 0, ForceCast_Exact_Replicate);
+			default:
+				return Compiler->Constant(1.0f);
+			}
+
+			return Compiler->Constant(0.0f);
+		}
+		else if (Property == MP_WorldPositionOffset)
+		{
+			//This property MUST return 0 as a default or during the process of rendering textures out for lightmass to use, pixels will be off by 1.
+			return Compiler->Constant(0.0f);
+		}
+		else if (Property >= MP_CustomizedUVs0 && Property <= MP_CustomizedUVs7)
+		{
+			// Pass through customized UVs
+			return MaterialInterface->CompileProperty(Compiler, Property);
+		}
+		else if (Property == MP_ShadingModel)
+		{
+			return MaterialInterface->CompileProperty(Compiler, MP_ShadingModel);
+
+		}
+		else
+		{
+			return Compiler->Constant(1.0f);
+		}
+	}
+
+	virtual FString GetMaterialUsageDescription() const override
+	{
+		return FString::Printf(TEXT("MaterialBaking_%s"), MaterialInterface ? *MaterialInterface->GetName() : TEXT("NULL"));
+	}
+	virtual EMaterialDomain GetMaterialDomain() const override
+	{
+		// Because the baking module applies the material to a plane (or mesh),
+		// it needs to be a surface material.
+		return MD_Surface;
+	}
+	virtual bool IsTwoSided() const  override
+	{
+		if (MaterialInterface)
+		{
+			return MaterialInterface->IsTwoSided();
+		}
+		return false;
+	}
+	virtual bool IsDitheredLODTransition() const  override
+	{
+		if (MaterialInterface)
+		{
+			return MaterialInterface->IsDitheredLODTransition();
+		}
+		return false;
+	}
+	virtual bool IsLightFunction() const override
+	{
+		if (Material)
+		{
+			return (Material->MaterialDomain == MD_LightFunction);
+		}
+		return false;
+	}
+	virtual bool IsDeferredDecal() const override
+	{
+		// Decals are tricky. Since they mix with the underlying material
+		// and can't be applied to meshes, they can't really be baked 1:1.
+		// Instead we'll just bake them as surface materials.
+		return false;
+	}
+	virtual bool IsVolumetricPrimitive() const override
+	{
+		return Material && Material->MaterialDomain == MD_Volume;
+	}
+	virtual bool IsSpecialEngineMaterial() const override
+	{
+		if (Material)
+		{
+			return (Material->bUsedAsSpecialEngineMaterial == 1);
+		}
+		return false;
+	}
+	virtual bool IsWireframe() const override
+	{
+		if (Material)
+		{
+			return (Material->Wireframe == 1);
+		}
+		return false;
+	}
+	virtual bool IsMasked() const override { return ProxyBlendMode == BLEND_Masked; }
+	virtual enum EBlendMode GetBlendMode() const override { return ProxyBlendMode; }
+	virtual FMaterialShadingModelField GetShadingModels() const override { return MSM_DefaultLit; }
+	virtual bool IsShadingModelFromMaterialExpression() const override { return false; }
+	virtual float GetOpacityMaskClipValue() const override { return 0.5f; }
+	virtual bool GetCastDynamicShadowAsMasked() const override { return false; }
+	virtual FString GetFriendlyName() const override { return FString::Printf(TEXT("FExportMaterialRenderer %s"), MaterialInterface ? *MaterialInterface->GetName() : TEXT("NULL")); }
+	/**
+	* Should shaders compiled for this material be saved to disk?
+	*/
+	virtual bool IsPersistent() const override { return true; }
+	virtual FGuid GetMaterialId() const override { return Id; }
+
+	virtual UMaterialInterface* GetMaterialInterface() const override
+	{
+		return MaterialInterface;
+	}
+
+	friend FArchive& operator<< (FArchive& Ar, FGLTFExportMaterialProxy& V)
+	{
+		return Ar << V.MaterialInterface;
+	}
+
+	/**
+	* Iterate through all textures used by the material and return the maximum texture resolution used
+	* (ideally this could be made dependent of the material property)
+	*
+	* @param MaterialInterface The material to scan for texture size
+	*
+	* @return Size (width and height)
+	*/
+	FIntPoint FindMaxTextureSize(UMaterialInterface* InMaterialInterface, FIntPoint MinimumSize = FIntPoint(1, 1)) const
+	{
+		// static lod settings so that we only initialize them once
+		const UTextureLODSettings* GameTextureLODSettings = UDeviceProfileManager::Get().GetActiveProfile()->GetTextureLODSettings();
+
+		TArray<UTexture*> MaterialTextures;
+		InMaterialInterface->GetUsedTextures(MaterialTextures, EMaterialQualityLevel::Num, false, GMaxRHIFeatureLevel, false);
+
+		// find the largest texture in the list (applying it's LOD bias)
+		FIntPoint MaxSize = MinimumSize;
+		for (int32 TexIndex = 0; TexIndex < MaterialTextures.Num(); TexIndex++)
+		{
+			const UTexture* Texture = MaterialTextures[TexIndex];
+
+			if (Texture == NULL)
+			{
+				continue;
+			}
+
+			// get the max size of the texture
+			const FIntPoint LocalSize = [&]()
+			{
+				if (Texture->IsA(UTexture2D::StaticClass()))
+				{
+					UTexture2D* Tex2D = (UTexture2D*)Texture;
+					return FIntPoint(Tex2D->GetSizeX(), Tex2D->GetSizeY());
+				}
+				else if (Texture->IsA(UTextureCube::StaticClass()))
+				{
+					UTextureCube* TexCube = (UTextureCube*)Texture;
+					return FIntPoint(TexCube->GetSizeX(), TexCube->GetSizeY());
+				}
+				else if (Texture->IsA(UTexture2DArray::StaticClass())) 
+				{
+					UTexture2DArray* TexArray = (UTexture2DArray*)Texture;
+					return FIntPoint(TexArray->GetSizeX(), TexArray->GetSizeY());
+				}
+				return FIntPoint(0, 0);
+			}();
+			
+			// bias the texture size based on LOD group
+			const int32 LocalBias = GameTextureLODSettings->CalculateLODBias(Texture);
+			MaxSize.X = FMath::Max(LocalSize.X >> LocalBias, MaxSize.X);
+			MaxSize.Y = FMath::Max(LocalSize.Y >> LocalBias, MaxSize.Y);
+		}
+
+		return MaxSize;
+	}
+
+	static bool WillFillData(EBlendMode InBlendMode, EMaterialProperty InMaterialProperty)
+	{
+		if (InMaterialProperty == MP_EmissiveColor)
+		{
+			return true;
+		}
+
+		switch (InBlendMode)
+		{
+		case BLEND_Opaque:
+		{
+			switch (InMaterialProperty)
+			{
+			case MP_BaseColor:				return true;
+			case MP_Specular:				return true;
+			case MP_Normal:					return true;
+			case MP_Tangent:				return true;
+			case MP_Metallic:				return true;
+			case MP_Roughness:				return true;
+			case MP_Anisotropy:				return true;
+			case MP_AmbientOcclusion:		return true;
+			}
+		}
+		break;
+		}
+		return false;
+	}
+
+	virtual bool IsUsedWithStaticLighting() const override
+	{
+		return true; 
+	}
+
+	virtual void GatherExpressionsForCustomInterpolators(TArray<UMaterialExpression*>& OutExpressions) const override
+	{
+		if(Material)
+		{
+			Material->GetAllExpressionsForCustomInterpolators(OutExpressions);
+		}
+	}
+
+private:
+	int32 CompileInputForCustomOutput(FMaterialCompiler* Compiler, int32 InputIndex, uint32 ForceCastFlags) const
+	{
+		TArray<FMaterialCustomOutputAttributeDefintion> CustomAttributes;
+		FMaterialAttributeDefinitionMap::GetCustomAttributeList(CustomAttributes);
+		FMaterialCustomOutputAttributeDefintion* CustomAttribute = CustomAttributes.FindByPredicate([this](const FMaterialCustomOutputAttributeDefintion& Attribute) -> bool
+		{
+			return Attribute.AttributeName == CustomOutputToCompile;
+		});
+		check(CustomAttribute != nullptr);
+
+		UMaterialExpressionCustomOutput* Expression = GetCustomOutputExpression(CustomAttribute->FunctionName);
+		FExpressionInput* ExpressionInput = Expression ? Expression->GetInput(InputIndex) : nullptr;
+		int32 Result = INDEX_NONE;
+
+		if (ExpressionInput)
+		{
+			Result = ExpressionInput->Compile(Compiler);
+		}
+		else
+		{
+			Result = FMaterialAttributeDefinitionMap::CompileDefaultExpression(Compiler, CustomAttribute->AttributeID);
+		}
+
+		if (CustomOutputToCompile == TEXT("ClearCoatBottomNormal"))
+		{
+			Result = CompileNormalEncoding(Compiler, CompileNormalTransform(Compiler, Result));
+		}
+
+		if (ForceCastFlags & MFCF_ForceCast)
+		{
+			Result = Compiler->ForceCast(Result,  CustomAttribute->ValueType, ForceCastFlags);
+		}
+
+		return Result;
+	}
+
+	UMaterialExpressionCustomOutput* GetCustomOutputExpression(const FString& FunctionName) const
+	{
+		for (UMaterialExpression* Expression : Material->Expressions)
+		{
+			UMaterialExpressionCustomOutput* CustomOutputExpression = Cast<UMaterialExpressionCustomOutput>(Expression);
+			if (CustomOutputExpression && CustomOutputExpression->GetFunctionName() == FunctionName)
+			{
+				return CustomOutputExpression;
+			}
+		}
+
+		return nullptr;
+	}
+
+	int32 CompileNormalTransform(FMaterialCompiler* Compiler, int32 NormalInput) const
+	{
+		return bTangentSpaceNormal && !Material->bTangentSpaceNormal
+			? Compiler->TransformVector(MCB_World, MCB_Tangent, NormalInput) : NormalInput;
+	}
+
+	int32 CompileNormalEncoding(FMaterialCompiler* Compiler, int32 NormalInput) const
+	{
+		return Compiler->Add(
+			Compiler->Mul(NormalInput, Compiler->Constant(0.5f)), // [-1,1] * 0.5
+			Compiler->Constant(0.5f)); // [-0.5,0.5] + 0.5
+	}
+
+	static int32 CompileRefractionEncoding(FMaterialCompiler* Compiler, int32 RefractionInput)
+	{
+		// NOTE: avoid getting yellow output (since value type in the attribute definition is incorrectly set to MCT_Float2)
+		RefractionInput = Compiler->ForceCast(RefractionInput, MCT_Float1);
+
+		// [1,Infinity] -> [1,0]
+		return Compiler->Div(
+			Compiler->Constant(1.0f),
+			RefractionInput);
+	}
+
+	static int32 CompileShadingModelEncoding(FMaterialCompiler* Compiler, int32 ShadingModelInput)
+	{
+		class FMaterialCompilerHack : public FHLSLMaterialTranslator
+		{
+			using FHLSLMaterialTranslator::FHLSLMaterialTranslator;
+
+		public:
+
+			int32 ForceCastToFloat(int32 Index)
+			{
+				const FString Code = FString::Printf(TEXT("((float)(%s))"), *GetParameterCode(Index));
+				const char* Buffer = static_cast<const char*>(static_cast<const void*>(*Code));
+				const uint64 Hash = CityHash64(Buffer, Code.Len() * sizeof(TCHAR));
+				const int32 CodeIndex = CurrentScopeChunks->Num();
+				new(*CurrentScopeChunks) FShaderCodeChunk(Hash, *Code, TEXT(""), MCT_Float1, true);
+				return CodeIndex;
+			}
+		};
+
+		// [0,MSM_NUM] -> [0,1]
+		return Compiler->Div(
+			// NOTE: ugly hack to workaround casting shading model (uint) to float
+			ShadingModelInput = static_cast<FMaterialCompilerHack*>(Compiler)->ForceCastToFloat(ShadingModelInput),
+			Compiler->Constant(MSM_NUM));
+	}
+
+private:
+	/** The material interface for this proxy */
+	UMaterialInterface* MaterialInterface;
+	UMaterial* Material;
+	TArray<UObject*> ReferencedTextures;
+	/** The property to compile for rendering the sample */
+	EMaterialProperty PropertyToCompile;
+	/** The name of the specific custom output to compile for rendering the sample. Only used if PropertyToCompile is MP_CustomOutput */
+	FString CustomOutputToCompile;
+	FGuid Id;
+	bool bSynchronousCompilation;
+
+public:
+	/** The blend mode used when baking the proxy material */
+	EBlendMode ProxyBlendMode;
+	/** Whether to transform normals from world-space to tangent-space (does nothing if material already uses tangent-space normals) */
+	bool bTangentSpaceNormal;
+};
diff --git a/Engine/Plugins/Enterprise/GLTFExporter/MigratedFromGit/Source/GLTFMaterialBaking/Private/GLTFMaterialBakingHelpers.cpp b/Engine/Plugins/Enterprise/GLTFExporter/MigratedFromGit/Source/GLTFMaterialBaking/Private/GLTFMaterialBakingHelpers.cpp
new file mode 100644
--- /dev/null
+++ b/Engine/Plugins/Enterprise/GLTFExporter/MigratedFromGit/Source/GLTFMaterialBaking/Private/GLTFMaterialBakingHelpers.cpp
@@ -0,0 +1,327 @@
+// Copyright Epic Games, Inc. All Rights Reserved.
+
+#include "GLTFMaterialBakingHelpers.h"
+#include "Math/Color.h"
+#include "Async/ParallelFor.h"
+
+namespace FMaterialBakingHelpersImpl
+{
+	static FColor BoxBlurSample(FColor* InBMP, int32 X, int32 Y, int32 InImageWidth, int32 InImageHeight, uint32 BackgroundMask)
+	{
+		int32 PixelsSampled = 0;
+		int32 CombinedColorR = 0;
+		int32 CombinedColorG = 0;
+		int32 CombinedColorB = 0;
+		int32 CombinedColorA = 0;
+
+		const int32 BaseIndex = (Y * InImageWidth) + X;
+	
+		int32 SampleIndex = BaseIndex - (InImageWidth + 1);
+		for (int32 YI = 0; YI < 3; ++YI)
+		{
+			for (int32 XI = 0; XI < 3; ++XI)
+			{
+				const FColor SampledColor = InBMP[SampleIndex++];
+				// Check if the pixel is a rendered one (not clear color)
+				if (SampledColor.DWColor() != BackgroundMask)
+				{
+					CombinedColorR += SampledColor.R;
+					CombinedColorG += SampledColor.G;
+					CombinedColorB += SampledColor.B;
+					CombinedColorA += SampledColor.A;
+					++PixelsSampled;
+				}
+			}
+			SampleIndex += InImageWidth - 3;
+		}
+
+		if (PixelsSampled == 0)
+		{
+			return FColor(BackgroundMask);
+		}
+
+		return FColor(CombinedColorR / PixelsSampled, CombinedColorG / PixelsSampled, CombinedColorB / PixelsSampled, CombinedColorA / PixelsSampled);
+	}
+
+	static bool HasBorderingPixel(FColor* InBMP, int32 X, int32 Y, int32 InImageWidth, int32 InImageHeight, uint32 BackgroundMask)
+	{
+		const int32 BaseIndex = Y*InImageWidth + X;
+
+		int32 SampleIndex = BaseIndex - (InImageWidth + 1);
+		for (int32 YI = 0; YI < 3; ++YI)
+		{
+			for (int32 XI = 0; XI < 3; ++XI)
+			{
+				if (InBMP[SampleIndex++].DWColor() != BackgroundMask)
+				{
+					return true;
+				}
+			}
+			SampleIndex += InImageWidth - 3;
+		}
+		return false;
+	}
+
+	void PerformUVBorderSmear(TArray<FColor>& InOutPixels, int32& ImageWidth, int32& ImageHeight, int32 MaxIterations, bool bCanShrink, uint32 BackgroundMask)
+	{
+		TRACE_CPUPROFILER_EVENT_SCOPE(FMaterialBakingHelpers::PerformUVBorderSmear)
+
+		check(InOutPixels.Num() == ImageWidth*ImageHeight);
+
+		const int32 PaddedImageWidth = ImageWidth + 2;
+		const int32 PaddedImageHeight = ImageHeight + 2;
+
+		TArray<uint32> RowsSwap1;
+		TArray<uint32> RowsSwap2;
+		TArray<uint32>* CurrentRowsLeft = &RowsSwap1;
+		TArray<uint32>* NextRowsLeft = &RowsSwap2;
+
+		//
+		// Create a copy of the image with a 1 pixel magenta border around the edge.
+		// This avoids needing to bounds check in inner loops.
+		//
+
+		FColor* Current = InOutPixels.GetData();
+
+		// We can hopefully skip the entire smearing process if there is a single
+		// non-background color in the entire array since smearing that completely
+		// would lead to a monochome output.
+		if (MaxIterations == -1)
+		{
+			uint32 SingleColor = BackgroundMask;
+			bool   bHasSingleColor = true;
+			for (int32 Index = 0, Num = InOutPixels.Num(); Index < Num; ++Index)
+			{
+				const uint32 Color = Current[Index].DWColor();
+				if (Color != BackgroundMask)
+				{
+					if (SingleColor == BackgroundMask)
+					{
+						// This is the first time we stumble on a color other than background color, keep it.
+						SingleColor = Color;
+					}
+					// Compare with the known non-background color
+					else if (SingleColor != Color)
+					{
+						bHasSingleColor = false;
+						break;
+					}
+				}
+			}
+
+			if (bHasSingleColor)
+			{
+				// Shrink to single texel when allowed
+				if (bCanShrink)
+				{
+					InOutPixels.SetNum(1);
+					ImageWidth = ImageHeight = 1;
+					Current = InOutPixels.GetData();
+				}
+
+				FColor Color(SingleColor);
+				for (int32 Index = 0, Num = InOutPixels.Num(); Index < Num; ++Index)
+				{
+					Current[Index] = Color;
+				}
+
+				return;
+			}
+		}
+
+		TArray<FColor> ScratchBuffer;
+		ScratchBuffer.SetNumUninitialized(PaddedImageWidth*PaddedImageHeight);
+		FColor* Scratch = ScratchBuffer.GetData();
+
+		// Set top and bottom bordering pixels to background color
+		for (int32 X = 0; X < PaddedImageWidth; ++X)
+		{
+			Scratch[X] = FColor(BackgroundMask);
+			Scratch[(PaddedImageHeight - 1)*PaddedImageWidth + X] = FColor(BackgroundMask);
+		}
+
+		// Set leftg and right border pixels to background color and copy image data into rows
+		for (int32 Y = 1; Y <= ImageHeight; ++Y)
+		{
+			const int32 YOffset = Y * PaddedImageWidth;
+			Scratch[YOffset] = FColor(BackgroundMask);
+			FMemory::Memcpy(&Scratch[YOffset + 1], &Current[(Y - 1)*ImageWidth], ImageWidth * sizeof(FColor));
+			Scratch[YOffset + ImageWidth + 1] = FColor(BackgroundMask);
+		}
+
+		//
+		// Find our initial workset of all rows that have background colored pixels bordering non-background pixels.
+		// Also find all rows that have zero background colored pixels - these will never be added to the workset.
+		//
+
+		TArray<bool> RowCompleted;
+		RowCompleted.SetNumZeroed(PaddedImageHeight);
+
+		// Set border rows to start completed
+		RowCompleted[0] = true;
+		RowCompleted[PaddedImageHeight - 1] = true;
+
+		bool bHasAnyData = false;
+		for (int32 Y = 1; Y <= ImageHeight; ++Y)
+		{
+			bool bHasBackground = false;
+			bool bBordersNonBackground = false;
+			const int32 YOffset = Y * PaddedImageWidth;
+			for (int32 X = 1; X <= ImageWidth; ++X)
+			{
+				if (Scratch[YOffset + X].DWColor() == BackgroundMask)
+				{
+					bHasBackground = true;
+					if (HasBorderingPixel(Scratch, X, Y, PaddedImageWidth, PaddedImageHeight, BackgroundMask))
+					{
+						bBordersNonBackground = true;
+						break;
+					}
+				}
+			}
+
+			if (!bHasBackground)
+			{
+				bHasAnyData = true;
+				RowCompleted[Y] = true;
+			}
+			else if (bBordersNonBackground)
+			{
+				bHasAnyData = true;
+				CurrentRowsLeft->Add(Y);
+			}
+		}
+
+		// early-out if no valid pixels were found.
+		if (!bHasAnyData)
+		{
+			FMemory::Memzero(InOutPixels.GetData(), ImageWidth * ImageHeight * sizeof(FColor));
+			return;
+		}
+
+		// Early out if all pixels were valid.
+		if (CurrentRowsLeft->Num() == 0)
+		{
+			return;
+		}
+
+		TArray<uint32> RowRemainingPixels;
+		RowRemainingPixels.SetNumZeroed(PaddedImageHeight);
+
+		const int32 MaxThreads = FPlatformProcess::SupportsMultithreading() ? FPlatformMisc::NumberOfCoresIncludingHyperthreads() : 1;
+
+		//
+		// Iteratively smear until all rows are filled.
+		//
+
+		int32 LoopCount = 0;
+		while (CurrentRowsLeft->Num() && (MaxIterations <= 0 || LoopCount <= MaxIterations))
+		{
+			const int32 NumThreads = FMath::Min(CurrentRowsLeft->Num(), MaxThreads);
+			const int32 LinesPerThread = FMath::CeilToInt((float)CurrentRowsLeft->Num() / (float)NumThreads);
+
+			// split up rows that still have background colored pixels amongst threads
+			ParallelFor(NumThreads, [ImageWidth, PaddedImageWidth, PaddedImageHeight, Current, Scratch, CurrentRowsLeft, LinesPerThread, &RowRemainingPixels, BackgroundMask](int32 ThreadIndex)
+			{
+				const int32 StartY = ThreadIndex*LinesPerThread;
+				const int32 EndY = FMath::Min((ThreadIndex + 1) * LinesPerThread, CurrentRowsLeft->Num());
+
+				for (int32 Index = StartY; Index < EndY; Index++)
+				{
+					const int32 Y = (*CurrentRowsLeft)[Index];
+					RowRemainingPixels[Index] = 0;
+
+					int32 PixelIndex = (Y - 1) * ImageWidth;
+					for (int32 X = 1; X <= ImageWidth; X++)
+					{
+						FColor& Color = Current[PixelIndex++];
+						if (Color.DWColor() == BackgroundMask)
+						{
+							const FColor SampledColor = BoxBlurSample(Scratch, X, Y, PaddedImageWidth, PaddedImageHeight, BackgroundMask);
+							// If it's a valid pixel
+							if (SampledColor.DWColor() != BackgroundMask)
+							{
+								Color = SampledColor;
+							}
+							else
+							{
+								RowRemainingPixels[Index]++;
+							}
+						}
+					}
+				}
+			});
+
+			// Mark all completed rows and copy data between buffers
+			const int32 RowMemorySize = ImageWidth * sizeof(FColor);
+			for (int32 RowIndex = 0; RowIndex < CurrentRowsLeft->Num(); ++RowIndex)
+			{
+				const int32 Y = (*CurrentRowsLeft)[RowIndex];
+				if (!RowRemainingPixels[RowIndex])
+				{
+					RowCompleted[Y] = true;
+				}
+				FMemory::Memcpy(&Scratch[Y * PaddedImageWidth + 1], &Current[(Y - 1) * ImageWidth], RowMemorySize);
+			}
+
+			// Find the next set of rows to process.
+			// This will be our current rows that haven't completed, plus any bordering rows that haven't been marked completed.
+			NextRowsLeft->SetNum(0, false);
+			{
+				int32 PreviousY = -1;
+				for (int32 RowIndex = 0; RowIndex < CurrentRowsLeft->Num(); ++RowIndex)
+				{
+					const int32 Y = (*CurrentRowsLeft)[RowIndex];
+					if (!RowCompleted[Y - 1] && PreviousY < Y - 1)
+					{
+						PreviousY = Y - 1;
+						NextRowsLeft->Add(Y - 1);
+					}
+					if (!RowCompleted[Y] && RowRemainingPixels[RowIndex] && PreviousY < Y)
+					{
+						PreviousY = Y;
+						NextRowsLeft->Add(Y);
+					}
+					if (!RowCompleted[Y + 1])
+					{
+						PreviousY = Y + 1;
+						NextRowsLeft->Add(Y + 1);
+					}
+				}
+			}
+
+			// Swap rows left buffers
+			{			
+				TArray<uint32>* Temp = NextRowsLeft;
+				NextRowsLeft = CurrentRowsLeft;
+				CurrentRowsLeft = Temp;
+			}
+
+			LoopCount++;
+		}
+
+		// If we finished before replacing all pixels, replace the remaining pixels with black.
+		if (CurrentRowsLeft->Num() > 0)
+		{
+			for (int32 i = 0; i < InOutPixels.Num(); ++i)
+			{
+				if (InOutPixels[i].DWColor() == BackgroundMask)
+				{
+					InOutPixels[i] = FColor::Black;
+				}
+			}
+		}
+	}
+}
+
+void FGLTFMaterialBakingHelpers::PerformUVBorderSmear(TArray<FColor>& InOutPixels, int32 ImageWidth, int32 ImageHeight, int32 MaxIterations, FColor BackgroundColor)
+{
+	const uint32 BackgroundMask = BackgroundColor.DWColor();
+	FMaterialBakingHelpersImpl::PerformUVBorderSmear(InOutPixels, ImageWidth, ImageHeight, MaxIterations, false, BackgroundMask);
+}
+
+void FGLTFMaterialBakingHelpers::PerformUVBorderSmearAndShrink(TArray<FColor>& InOutPixels, int32& InOutImageWidth, int32& InOutImageHeight, FColor BackgroundColor)
+{
+	const uint32 BackgroundMask = BackgroundColor.DWColor();
+	FMaterialBakingHelpersImpl::PerformUVBorderSmear(InOutPixels, InOutImageWidth, InOutImageHeight, -1, true, BackgroundMask);
+}
\ No newline at end of file
diff --git a/Engine/Plugins/Enterprise/GLTFExporter/MigratedFromGit/Source/GLTFMaterialBaking/Private/GLTFMaterialBakingModule.cpp b/Engine/Plugins/Enterprise/GLTFExporter/MigratedFromGit/Source/GLTFMaterialBaking/Private/GLTFMaterialBakingModule.cpp
new file mode 100644
--- /dev/null
+++ b/Engine/Plugins/Enterprise/GLTFExporter/MigratedFromGit/Source/GLTFMaterialBaking/Private/GLTFMaterialBakingModule.cpp
@@ -0,0 +1,1005 @@
+// Copyright Epic Games, Inc. All Rights Reserved.
+
+#include "GLTFMaterialBakingModule.h"
+#include "MaterialRenderItem.h"
+#include "Engine/TextureRenderTarget2D.h"
+#include "GLTFExportMaterialProxy.h"
+#include "UObject/UObjectGlobals.h"
+#include "MaterialBakingStructures.h"
+#include "GLTFMaterialBakingHelpers.h"
+#include "AddCustomAttributeHack.h"
+#include "Async/Async.h"
+#include "Async/ParallelFor.h"
+#include "Materials/MaterialInstance.h"
+#include "Materials/MaterialInstanceConstant.h"
+#include "MaterialEditor/MaterialEditorInstanceConstant.h"
+#include "RenderingThread.h"
+#include "RHISurfaceDataConversion.h"
+#include "Misc/ScopedSlowTask.h"
+#include "MeshDescription.h"
+#if WITH_EDITOR
+#include "Misc/FileHelper.h"
+#endif
+
+IMPLEMENT_MODULE(FGLTFMaterialBakingModule, GLTFMaterialBaking);
+
+#define LOCTEXT_NAMESPACE "MaterialBakingModule"
+
+/** Cvars for advanced features */
+static TAutoConsoleVariable<int32> CVarUseMaterialProxyCaching(
+	TEXT("GLTFMaterialBaking.UseMaterialProxyCaching"),
+	1,
+	TEXT("Determines whether or not Material Proxies should be cached to speed up material baking.\n")
+	TEXT("0: Turned Off\n")
+	TEXT("1: Turned On"),	
+	ECVF_Default);
+
+static TAutoConsoleVariable<int32> CVarSaveIntermediateTextures(
+	TEXT("GLTFMaterialBaking.SaveIntermediateTextures"),
+	0,
+	TEXT("Determines whether or not to save out intermediate BMP images for each flattened material property.\n")
+	TEXT("0: Turned Off\n")
+	TEXT("1: Turned On"),
+	ECVF_Default);
+
+namespace FMaterialBakingModuleImpl
+{
+	// Custom dynamic mesh allocator specifically tailored for Material Baking.
+	// This will always reuse the same couple buffers, so searching linearly is not a problem.
+	class FMaterialBakingDynamicMeshBufferAllocator : public FDynamicMeshBufferAllocator
+	{
+		// This must be smaller than the large allocation blocks on Windows 10 which is currently ~508K.
+		// Large allocations uses VirtualAlloc directly without any kind of buffering before
+		// releasing pages to the kernel, so it causes lots of soft page fault when
+		// memory is first initialized.
+		const uint32 SmallestPooledBufferSize = 256*1024;
+
+		TArray<FIndexBufferRHIRef>  IndexBuffers;
+		TArray<FVertexBufferRHIRef> VertexBuffers;
+
+		template <typename RefType>
+		RefType GetSmallestFit(uint32 SizeInBytes, TArray<RefType>& Array)
+		{
+			uint32 SmallestFitIndex = UINT32_MAX;
+			uint32 SmallestFitSize  = UINT32_MAX;
+			for (int32 Index = 0; Index < Array.Num(); ++Index)
+			{
+				uint32 Size = Array[Index]->GetSize();
+				if (Size >= SizeInBytes && (SmallestFitIndex == UINT32_MAX || Size < SmallestFitSize))
+				{
+					SmallestFitIndex = Index;
+					SmallestFitSize  = Size;
+				}
+			}
+
+			RefType Ref;
+			// Do not reuse the smallest fit if it's a lot bigger than what we requested
+			if (SmallestFitIndex != UINT32_MAX && SmallestFitSize < SizeInBytes*2)
+			{
+				Ref = Array[SmallestFitIndex];
+				Array.RemoveAtSwap(SmallestFitIndex);
+			}
+
+			return Ref;
+		}
+
+		virtual FIndexBufferRHIRef AllocIndexBuffer(uint32 NumElements) override
+		{
+			uint32 BufferSize = GetIndexBufferSize(NumElements);
+			if (BufferSize > SmallestPooledBufferSize)
+			{
+				FIndexBufferRHIRef Ref = GetSmallestFit(GetIndexBufferSize(NumElements), IndexBuffers);
+				if (Ref.IsValid())
+				{
+					return Ref;
+				}
+			}
+
+			return FDynamicMeshBufferAllocator::AllocIndexBuffer(NumElements);
+		}
+
+		virtual void ReleaseIndexBuffer(FIndexBufferRHIRef& IndexBufferRHI) override
+		{
+			if (IndexBufferRHI->GetSize() > SmallestPooledBufferSize)
+			{
+				IndexBuffers.Add(MoveTemp(IndexBufferRHI));
+			}
+
+			IndexBufferRHI = nullptr;
+		}
+
+		virtual FVertexBufferRHIRef AllocVertexBuffer(uint32 Stride, uint32 NumElements) override
+		{
+			uint32 BufferSize = GetVertexBufferSize(Stride, NumElements);
+			if (BufferSize > SmallestPooledBufferSize)
+			{
+				FVertexBufferRHIRef Ref = GetSmallestFit(BufferSize, VertexBuffers);
+				if (Ref.IsValid())
+				{
+					return Ref;
+				}
+			}
+
+			return FDynamicMeshBufferAllocator::AllocVertexBuffer(Stride, NumElements);
+		}
+
+		virtual void ReleaseVertexBuffer(FVertexBufferRHIRef& VertexBufferRHI) override
+		{
+			if (VertexBufferRHI->GetSize() > SmallestPooledBufferSize)
+			{
+				VertexBuffers.Add(MoveTemp(VertexBufferRHI));
+			}
+
+			VertexBufferRHI = nullptr;
+		}
+	};
+
+	class FStagingBufferPool
+	{
+	public:
+		FTexture2DRHIRef CreateStagingBuffer_RenderThread(FRHICommandListImmediate& RHICmdList, int32 Width, int32 Height, EPixelFormat Format)
+		{
+			TRACE_CPUPROFILER_EVENT_SCOPE(CreateStagingBuffer_RenderThread)
+
+			auto StagingBufferPredicate = 
+				[Width, Height, Format](const FTexture2DRHIRef& Texture2DRHIRef)
+				{
+					return Texture2DRHIRef->GetSizeX() == Width && Texture2DRHIRef->GetSizeY() == Height && Texture2DRHIRef->GetFormat() == Format;
+				};
+
+			// Process any staging buffers available for unmapping
+			{
+				TArray<FTexture2DRHIRef> ToUnmapLocal;
+				{
+					FScopeLock Lock(&ToUnmapLock);
+					ToUnmapLocal = MoveTemp(ToUnmap);
+				}
+
+				for (int32 Index = 0, Num = ToUnmapLocal.Num(); Index < Num; ++Index)
+				{
+					RHICmdList.UnmapStagingSurface(ToUnmapLocal[Index]);
+					Pool.Add(MoveTemp(ToUnmapLocal[Index]));
+				}
+			}
+
+			// Find any pooled staging buffer with suitable properties.
+			int32 Index = Pool.IndexOfByPredicate(StagingBufferPredicate);
+
+			if (Index != -1)
+			{
+				FTexture2DRHIRef StagingBuffer = MoveTemp(Pool[Index]);
+				Pool.RemoveAtSwap(Index);
+				return StagingBuffer;
+			}
+
+			TRACE_CPUPROFILER_EVENT_SCOPE(RHICreateTexture2D)
+			FRHIResourceCreateInfo CreateInfo;
+			return RHICreateTexture2D(Width, Height, Format, 1, 1, TexCreate_CPUReadback, CreateInfo);
+		}
+
+		void ReleaseStagingBufferForUnmap_AnyThread(FTexture2DRHIRef& Texture2DRHIRef)
+		{
+			TRACE_CPUPROFILER_EVENT_SCOPE(ReleaseStagingBufferForUnmap_AnyThread)
+			FScopeLock Lock(&ToUnmapLock);
+			ToUnmap.Emplace(MoveTemp(Texture2DRHIRef));
+		}
+
+		void Clear_RenderThread(FRHICommandListImmediate& RHICmdList)
+		{
+			TRACE_CPUPROFILER_EVENT_SCOPE(Clear_RenderThread)
+			for (FTexture2DRHIRef& StagingSurface : ToUnmap)
+			{
+				RHICmdList.UnmapStagingSurface(StagingSurface);
+			}
+
+			ToUnmap.Empty();
+			Pool.Empty();
+		}
+
+		~FStagingBufferPool()
+		{
+			check(Pool.Num() == 0);
+		}
+
+	private:
+		TArray<FTexture2DRHIRef> Pool;
+
+		// Not contented enough to warrant the use of lockless structures.
+		FCriticalSection         ToUnmapLock;
+		TArray<FTexture2DRHIRef> ToUnmap;
+	};
+
+	struct FRenderItemKey
+	{
+		const FGLTFMeshRenderData* RenderData;
+		const FIntPoint  RenderSize;
+
+		FRenderItemKey(const FGLTFMeshRenderData* InRenderData, const FIntPoint& InRenderSize)
+			: RenderData(InRenderData)
+			, RenderSize(InRenderSize)
+		{
+		}
+
+		bool operator == (const FRenderItemKey& Other) const
+		{
+			return RenderData == Other.RenderData &&
+				RenderSize == Other.RenderSize;
+		}
+	};
+
+	uint32 GetTypeHash(const FRenderItemKey& Key)
+	{
+		return HashCombine(GetTypeHash(Key.RenderData), GetTypeHash(Key.RenderSize));
+	}
+}
+
+void FGLTFMaterialBakingModule::StartupModule()
+{
+	bEmissiveHDR = false;
+
+	// Set which properties should enforce gamma correction
+	SetLinearBake(false);
+
+	// Set which pixel format should be used for the possible baked out material properties
+	PerPropertyFormat.Add(MP_EmissiveColor, PF_FloatRGBA);
+	PerPropertyFormat.Add(MP_Opacity, PF_B8G8R8A8);
+	PerPropertyFormat.Add(MP_OpacityMask, PF_B8G8R8A8);
+	PerPropertyFormat.Add(MP_BaseColor, PF_B8G8R8A8);
+	PerPropertyFormat.Add(MP_Metallic, PF_B8G8R8A8);
+	PerPropertyFormat.Add(MP_Specular, PF_B8G8R8A8);
+	PerPropertyFormat.Add(MP_Roughness, PF_B8G8R8A8);
+	PerPropertyFormat.Add(MP_Anisotropy, PF_B8G8R8A8);
+	PerPropertyFormat.Add(MP_Normal, PF_B8G8R8A8);
+	PerPropertyFormat.Add(MP_Tangent, PF_B8G8R8A8);
+	PerPropertyFormat.Add(MP_AmbientOcclusion, PF_B8G8R8A8);
+	PerPropertyFormat.Add(MP_SubsurfaceColor, PF_B8G8R8A8);
+	PerPropertyFormat.Add(MP_CustomData0, PF_B8G8R8A8);
+	PerPropertyFormat.Add(MP_CustomData1, PF_B8G8R8A8);
+	PerPropertyFormat.Add(MP_Refraction, PF_B8G8R8A8);
+	PerPropertyFormat.Add(MP_ShadingModel, PF_B8G8R8A8);
+	PerPropertyFormat.Add(TEXT("ClearCoatBottomNormal"), PF_B8G8R8A8);
+	PerPropertyFormat.Add(TEXT("TransmittanceColor"), PF_B8G8R8A8);
+
+	// Register callback for modified objects
+	FCoreUObjectDelegates::OnObjectModified.AddRaw(this, &FGLTFMaterialBakingModule::OnObjectModified);
+
+	// NOTE: Because the attribute definition of TransmittanceColor is not registered by the engine, we have to do it here
+	AddCustomAttributeHack(FGuid(0xF2D8C70E, 0x42ECA0D1, 0x4652D0AD, 0xB785A065), "TransmittanceColor", "GetThinTranslucentMaterialOutput", MCT_Float3, FVector4(0.5, 0.5, 0.5, 0));
+}
+
+void FGLTFMaterialBakingModule::ShutdownModule()
+{
+	// Unregister customization and callback
+	FPropertyEditorModule* PropertyEditorModule = FModuleManager::GetModulePtr<FPropertyEditorModule>("PropertyEditor");
+
+	if (PropertyEditorModule)
+	{
+		PropertyEditorModule->UnregisterCustomPropertyTypeLayout(TEXT("PropertyEntry"));
+	}
+	FCoreUObjectDelegates::OnObjectModified.RemoveAll(this);
+}
+
+void FGLTFMaterialBakingModule::BakeMaterials(const TArray<FGLTFMaterialData*>& MaterialSettings, const TArray<FGLTFMeshRenderData*>& MeshSettings, TArray<FGLTFBakeOutput>& Output)
+{
+	// Translate old material data to extended types
+	TArray<FGLTFMaterialDataEx> MaterialDataExs;
+	MaterialDataExs.Reserve(MaterialSettings.Num());
+	for (const FGLTFMaterialData* MaterialData : MaterialSettings)
+	{
+		FGLTFMaterialDataEx& MaterialDataEx = MaterialDataExs.AddDefaulted_GetRef();
+		MaterialDataEx.Material = MaterialData->Material;
+		MaterialDataEx.bPerformBorderSmear = MaterialData->bPerformBorderSmear;
+		MaterialDataEx.BlendMode = MaterialData->BlendMode;
+		MaterialDataEx.bTangentSpaceNormal = MaterialData->bTangentSpaceNormal;
+		MaterialDataEx.BackgroundColor = MaterialData->BackgroundColor;
+
+		for (const TPair<EMaterialProperty, FIntPoint>& PropertySizePair : MaterialData->PropertySizes)
+		{
+			MaterialDataEx.PropertySizes.Add(PropertySizePair.Key, PropertySizePair.Value);
+		}
+	}
+
+	// Build an array of pointers to the extended type
+	TArray<FGLTFMaterialDataEx*> MaterialSettingsEx;
+	MaterialSettingsEx.Reserve(MaterialDataExs.Num());
+	for (FGLTFMaterialDataEx& MaterialDataEx : MaterialDataExs)
+	{
+		MaterialSettingsEx.Add(&MaterialDataEx);
+	}
+
+	TArray<FGLTFBakeOutputEx> BakeOutputExs;
+	BakeMaterials(MaterialSettingsEx, MeshSettings, BakeOutputExs);
+
+	// Translate extended bake output to old types
+	Output.Reserve(BakeOutputExs.Num());
+	for (FGLTFBakeOutputEx& BakeOutputEx : BakeOutputExs)
+	{
+		FGLTFBakeOutput& BakeOutput = Output.AddDefaulted_GetRef();
+		BakeOutput.EmissiveScale = BakeOutputEx.EmissiveScale;
+
+		for (TPair<FGLTFMaterialPropertyEx, FIntPoint>& PropertySizePair : BakeOutputEx.PropertySizes)
+		{
+			BakeOutput.PropertySizes.Add(PropertySizePair.Key.Type, PropertySizePair.Value);
+		}
+
+		for (TPair<FGLTFMaterialPropertyEx, TArray<FColor>>& PropertyDataPair : BakeOutputEx.PropertyData)
+		{
+			BakeOutput.PropertyData.Add(PropertyDataPair.Key.Type, MoveTemp(PropertyDataPair.Value));
+		}
+
+		for (TPair<FGLTFMaterialPropertyEx, TArray<FFloat16Color>>& PropertyDataPair : BakeOutputEx.HDRPropertyData)
+		{
+			BakeOutput.HDRPropertyData.Add(PropertyDataPair.Key.Type, MoveTemp(PropertyDataPair.Value));
+		}
+	}
+}
+
+#if !(ENGINE_MAJOR_VERSION > 4 || ENGINE_MINOR_VERSION >= 27)
+static inline void ConvertRawR16G16B16A16FDataToFFloat16Color(uint32 Width, uint32 Height, uint8* In, uint32 SrcPitch, FFloat16Color* Out)
+{
+	const uint32 DstPitch = Width * sizeof(FFloat16Color);
+
+	// If source & dest pitch matches, perform a single memcpy.
+	if (DstPitch == SrcPitch)
+	{
+		FPlatformMemory::Memcpy(Out, In, Width * Height * sizeof(FFloat16Color));
+	}
+	else
+	{
+		check(SrcPitch > DstPitch);
+
+		// Need to copy row wise since the Pitch does not match the Width.
+		for (uint32 Y = 0; Y < Height; Y++)
+		{
+			FFloat16Color* SrcPtr = (FFloat16Color*)(In + Y * SrcPitch);
+			FFloat16Color* DestPtr = Out + Y * Width;
+			FMemory::Memcpy(DestPtr, SrcPtr, DstPitch);
+		}
+	}
+}
+#endif
+
+void FGLTFMaterialBakingModule::BakeMaterials(const TArray<FGLTFMaterialDataEx*>& MaterialSettings, const TArray<FGLTFMeshRenderData*>& MeshSettings, TArray<FGLTFBakeOutputEx>& Output)
+{
+	TRACE_CPUPROFILER_EVENT_SCOPE(FMaterialBakingModule::BakeMaterials)
+
+	checkf(MaterialSettings.Num() == MeshSettings.Num(), TEXT("Number of material settings does not match that of MeshSettings"));
+	const int32 NumMaterials = MaterialSettings.Num();
+	const bool bSaveIntermediateTextures = CVarSaveIntermediateTextures.GetValueOnAnyThread() == 1;
+
+	using namespace FMaterialBakingModuleImpl;
+	FMaterialBakingDynamicMeshBufferAllocator MaterialBakingDynamicMeshBufferAllocator;
+
+	FScopedSlowTask Progress(NumMaterials, LOCTEXT("BakeMaterials", "Baking Materials..."), true );
+	Progress.MakeDialog(true);
+
+	TArray<uint32> ProcessingOrder;
+	ProcessingOrder.Reserve(MeshSettings.Num());
+	for (int32 Index = 0; Index < MeshSettings.Num(); ++Index)
+	{
+		ProcessingOrder.Add(Index);
+	}
+
+	// Start with the biggest mesh first so we can always reuse the same vertex/index buffers.
+	// This will decrease the number of allocations backed by newly allocated memory from the OS,
+	// which will reduce soft page faults while copying into that memory.
+	// Soft page faults are now incredibly expensive on Windows 10.
+	Algo::SortBy(
+		ProcessingOrder,
+		[&MeshSettings](const uint32 Index){ return MeshSettings[Index]->RawMeshDescription ? MeshSettings[Index]->RawMeshDescription->Vertices().Num() : 0; },
+		TGreater<>()
+	);
+
+	Output.SetNum(NumMaterials);
+
+	struct FPipelineContext
+	{
+		typedef TFunction<void (FRHICommandListImmediate& RHICmdList)> FReadCommand;
+		FReadCommand ReadCommand;
+	};
+
+	// Distance between the command sent to rendering and the GPU read-back of the result
+	// to minimize sync time waiting on GPU.
+	const int32 PipelineDepth = 16;
+	int32 PipelineIndex = 0;
+	FPipelineContext PipelineContext[PipelineDepth];
+
+	// This will create and prepare FMeshMaterialRenderItem for each property sizes we're going to need
+	auto PrepareRenderItems_AnyThread =
+		[&](int32 MaterialIndex)
+	{
+		TRACE_CPUPROFILER_EVENT_SCOPE(PrepareRenderItems);
+		
+		TMap<FMaterialBakingModuleImpl::FRenderItemKey, FGLTFMeshMaterialRenderItem*>* RenderItems = new TMap<FRenderItemKey, FGLTFMeshMaterialRenderItem *>();
+		const FGLTFMaterialDataEx* CurrentMaterialSettings = MaterialSettings[MaterialIndex];
+		const FGLTFMeshRenderData* CurrentMeshSettings = MeshSettings[MaterialIndex];
+
+		for (TMap<FGLTFMaterialPropertyEx, FIntPoint>::TConstIterator PropertySizeIterator = CurrentMaterialSettings->PropertySizes.CreateConstIterator(); PropertySizeIterator; ++PropertySizeIterator)
+		{
+			FRenderItemKey RenderItemKey(CurrentMeshSettings, PropertySizeIterator.Value());
+			if (RenderItems->Find(RenderItemKey) == nullptr)
+			{
+				RenderItems->Add(RenderItemKey, new FGLTFMeshMaterialRenderItem(PropertySizeIterator.Value(), CurrentMeshSettings, &MaterialBakingDynamicMeshBufferAllocator));
+			}
+		}
+
+		return RenderItems;
+	};
+
+	// We reuse the pipeline depth to prepare render items in advance to avoid stalling the game thread
+	int NextRenderItem = 0;
+	TFuture<TMap<FRenderItemKey, FGLTFMeshMaterialRenderItem*>*> PreparedRenderItems[PipelineDepth];
+	for (; NextRenderItem < NumMaterials && NextRenderItem < PipelineDepth; ++NextRenderItem)
+	{
+		PreparedRenderItems[NextRenderItem] = 
+			Async(
+				EAsyncExecution::ThreadPool,
+				[&PrepareRenderItems_AnyThread, &ProcessingOrder, NextRenderItem]()
+				{
+					return PrepareRenderItems_AnyThread(ProcessingOrder[NextRenderItem]);
+				}
+			);
+	}
+
+	// Create all material proxies right away to start compiling shaders asynchronously and avoid stalling the baking process as much as possible
+	{
+		TRACE_CPUPROFILER_EVENT_SCOPE(CreateMaterialProxies)
+
+		for (int32 Index = 0; Index < NumMaterials; ++Index)
+		{
+			int32 MaterialIndex = ProcessingOrder[Index];
+			const FGLTFMaterialDataEx* CurrentMaterialSettings = MaterialSettings[MaterialIndex];
+
+			TArray<UTexture*> MaterialTextures;
+			CurrentMaterialSettings->Material->GetUsedTextures(MaterialTextures, EMaterialQualityLevel::Num, true, GMaxRHIFeatureLevel, true);
+
+			// Force load materials used by the current material
+			{
+				TRACE_CPUPROFILER_EVENT_SCOPE(LoadTexturesForMaterial)
+
+				for (UTexture* Texture : MaterialTextures)
+				{
+					if (Texture != NULL)
+					{
+						UTexture2D* Texture2D = Cast<UTexture2D>(Texture);
+						if (Texture2D)
+						{
+							Texture2D->SetForceMipLevelsToBeResident(30.0f);
+							Texture2D->WaitForStreaming();
+						}
+					}
+				}
+			}
+
+			for (TMap<FGLTFMaterialPropertyEx, FIntPoint>::TConstIterator PropertySizeIterator = CurrentMaterialSettings->PropertySizes.CreateConstIterator(); PropertySizeIterator; ++PropertySizeIterator)
+			{
+				// They will be stored in the pool and compiled asynchronously
+				CreateMaterialProxy(CurrentMaterialSettings, PropertySizeIterator.Key());
+			}
+		}
+	}
+
+	TAtomic<uint32>    NumTasks(0);
+	FStagingBufferPool StagingBufferPool;
+
+	for (int32 Index = 0; Index < NumMaterials; ++Index)
+	{
+		TRACE_CPUPROFILER_EVENT_SCOPE(BakeOneMaterial)
+
+		Progress.EnterProgressFrame(1.0f, FText::Format(LOCTEXT("BakingMaterial", "Baking Material {0}/{1}"), Index, NumMaterials));
+
+		int32 MaterialIndex = ProcessingOrder[Index];
+		TMap<FRenderItemKey, FGLTFMeshMaterialRenderItem*>* RenderItems;
+
+		{
+			TRACE_CPUPROFILER_EVENT_SCOPE(WaitOnPreparedRenderItems)
+			RenderItems = PreparedRenderItems[Index % PipelineDepth].Get();
+		}
+
+		// Prepare the next render item in advance
+		if (NextRenderItem < NumMaterials)
+		{
+			check((NextRenderItem % PipelineDepth) == (Index % PipelineDepth));
+			PreparedRenderItems[NextRenderItem % PipelineDepth] = 
+				Async(
+					EAsyncExecution::ThreadPool,
+					[&PrepareRenderItems_AnyThread, NextMaterialIndex = ProcessingOrder[NextRenderItem]]()
+					{
+						return PrepareRenderItems_AnyThread(NextMaterialIndex);
+					}
+				);
+			NextRenderItem++;
+		}
+
+		const FGLTFMaterialDataEx* CurrentMaterialSettings = MaterialSettings[MaterialIndex];
+		FGLTFBakeOutputEx& CurrentOutput = Output[MaterialIndex];
+
+		TArray<FGLTFMaterialPropertyEx> MaterialPropertiesToBakeOut;
+		CurrentMaterialSettings->PropertySizes.GenerateKeyArray(MaterialPropertiesToBakeOut);
+
+		const int32 NumPropertiesToRender = MaterialPropertiesToBakeOut.Num();
+		if (NumPropertiesToRender > 0)
+		{
+			TRACE_CPUPROFILER_EVENT_SCOPE(RenderOneMaterial)
+
+			// Ensure data in memory will not change place passed this point to avoid race conditions
+			CurrentOutput.PropertySizes = CurrentMaterialSettings->PropertySizes;
+			for (int32 PropertyIndex = 0; PropertyIndex < NumPropertiesToRender; ++PropertyIndex)
+			{
+				const FGLTFMaterialPropertyEx& Property = MaterialPropertiesToBakeOut[PropertyIndex];
+				CurrentOutput.PropertyData.Add(Property);
+				if (bEmissiveHDR && Property == MP_EmissiveColor)
+				{
+					CurrentOutput.HDRPropertyData.Add(Property);
+				}
+			}
+
+			for (int32 PropertyIndex = 0; PropertyIndex < NumPropertiesToRender; ++PropertyIndex)
+			{
+				const FGLTFMaterialPropertyEx& Property = MaterialPropertiesToBakeOut[PropertyIndex];
+				FGLTFExportMaterialProxy* ExportMaterialProxy = CreateMaterialProxy(CurrentMaterialSettings, Property);
+
+				if (!ExportMaterialProxy->IsCompilationFinished())
+				{
+					TRACE_CPUPROFILER_EVENT_SCOPE(WaitForMaterialProxyCompilation)
+					ExportMaterialProxy->FinishCompilation();
+				}
+
+				// Lookup gamma and format settings for property, if not found use default values
+				const EPropertyColorSpace* OverrideColorSpace = PerPropertyColorSpace.Find(Property);
+				const EPropertyColorSpace ColorSpace = OverrideColorSpace ? *OverrideColorSpace : DefaultColorSpace;
+				const EPixelFormat PixelFormat = PerPropertyFormat.Contains(Property) ? PerPropertyFormat[Property] : PF_B8G8R8A8;
+
+				// It is safe to reuse the same render target for each draw pass since they all execute sequentially on the GPU and are copied to staging buffers before
+				// being reused.
+				UTextureRenderTarget2D* RenderTarget = CreateRenderTarget((ColorSpace == EPropertyColorSpace::Linear), PixelFormat, CurrentOutput.PropertySizes[Property], CurrentMaterialSettings->BackgroundColor);
+				if (RenderTarget != nullptr)
+				{
+					// Perform everything left of the operation directly on the render thread since we need to modify some RenderItem's properties
+					// for each render pass and we can't do that without costly synchronization (flush) between the game thread and render thread.
+					// Everything slow to execute has already been prepared on the game thread anyway.
+					ENQUEUE_RENDER_COMMAND(RenderOneMaterial)(
+						[this, RenderItems, RenderTarget, Property, ExportMaterialProxy, &PipelineContext, PipelineIndex, &StagingBufferPool, &NumTasks, bSaveIntermediateTextures, &MaterialSettings, &MeshSettings, MaterialIndex, &Output](FRHICommandListImmediate& RHICmdList)
+						{
+							const FGLTFMaterialDataEx* CurrentMaterialSettings = MaterialSettings[MaterialIndex];
+							const FGLTFMeshRenderData* CurrentMeshSettings = MeshSettings[MaterialIndex];
+
+							FGLTFMeshMaterialRenderItem& RenderItem = *RenderItems->FindChecked(FRenderItemKey(CurrentMeshSettings, FIntPoint(RenderTarget->GetSurfaceWidth(), RenderTarget->GetSurfaceHeight())));
+
+							FSceneViewFamily ViewFamily(FSceneViewFamily::ConstructionValues(RenderTarget->GetRenderTargetResource(), nullptr,
+								FEngineShowFlags(ESFIM_Game))
+								.SetWorldTimes(0.0f, 0.0f, 0.0f)
+								.SetGammaCorrection(RenderTarget->GetRenderTargetResource()->GetDisplayGamma()));
+
+							RenderItem.MaterialRenderProxy = ExportMaterialProxy;
+							RenderItem.ViewFamily = &ViewFamily;
+
+							FTextureRenderTargetResource* RenderTargetResource = RenderTarget->GetRenderTargetResource();
+							FCanvas Canvas(RenderTargetResource, nullptr, FApp::GetCurrentTime() - GStartTime, FApp::GetDeltaTime(), FApp::GetCurrentTime() - GStartTime, GMaxRHIFeatureLevel);
+							Canvas.SetAllowedModes(FCanvas::Allow_Flush);
+							Canvas.SetRenderTargetRect(FIntRect(0, 0, RenderTarget->GetSurfaceWidth(), RenderTarget->GetSurfaceHeight()));
+							Canvas.SetBaseTransform(Canvas.CalcBaseTransform2D(RenderTarget->GetSurfaceWidth(), RenderTarget->GetSurfaceHeight()));
+
+							// Do rendering
+							Canvas.Clear(RenderTarget->ClearColor);
+							FCanvas::FCanvasSortElement& SortElement = Canvas.GetSortElement(Canvas.TopDepthSortKey());
+							SortElement.RenderBatchArray.Add(&RenderItem);
+							Canvas.Flush_RenderThread(RHICmdList);
+							SortElement.RenderBatchArray.Empty();
+
+							FTexture2DRHIRef StagingBufferRef = StagingBufferPool.CreateStagingBuffer_RenderThread(RHICmdList, RenderTargetResource->GetSizeX(), RenderTargetResource->GetSizeY(), RenderTarget->GetFormat());
+							FGPUFenceRHIRef GPUFence = RHICreateGPUFence(TEXT("MaterialBackingFence"));
+
+							FResolveRect Rect(0, 0, RenderTargetResource->GetSizeX(), RenderTargetResource->GetSizeY());
+							RHICmdList.CopyToResolveTarget(RenderTargetResource->GetRenderTargetTexture(), StagingBufferRef, FResolveParams(Rect));	
+							RHICmdList.WriteGPUFence(GPUFence);
+
+							// Prepare a lambda for final processing that will be executed asynchronously
+							NumTasks++;
+							auto FinalProcessing_AnyThread =
+								[&NumTasks, bSaveIntermediateTextures, CurrentMaterialSettings, &StagingBufferPool, &Output, Property, MaterialIndex, bEmissiveHDR = bEmissiveHDR](FTexture2DRHIRef& StagingBuffer, void * Data, int32 DataWidth, int32 DataHeight)
+								{
+									TRACE_CPUPROFILER_EVENT_SCOPE(FinalProcessing)
+
+									FGLTFBakeOutputEx& CurrentOutput  = Output[MaterialIndex];
+									TArray<FColor>& OutputColor = CurrentOutput.PropertyData[Property];
+									FIntPoint& OutputSize       = CurrentOutput.PropertySizes[Property];
+
+									OutputColor.SetNum(OutputSize.X * OutputSize.Y);
+
+									if (Property.Type == MP_EmissiveColor)
+									{
+										// Only one thread will write to CurrentOutput.EmissiveScale since there can be only one emissive channel property per FBakeOutputEx
+										FGLTFMaterialBakingModule::ProcessEmissiveOutput((const FFloat16Color*)Data, DataWidth, OutputSize, OutputColor, CurrentOutput.EmissiveScale, CurrentMaterialSettings->BackgroundColor);
+
+										if (bEmissiveHDR)
+										{
+											TArray<FFloat16Color>& OutputHDRColor = CurrentOutput.HDRPropertyData[Property];
+											OutputHDRColor.SetNum(OutputSize.X * OutputSize.Y);
+											ConvertRawR16G16B16A16FDataToFFloat16Color(OutputSize.X, OutputSize.Y, (uint8*)Data, DataWidth * sizeof(FFloat16Color), OutputHDRColor.GetData());
+										}
+									}
+									else
+									{
+										TRACE_CPUPROFILER_EVENT_SCOPE(ConvertRawB8G8R8A8DataToFColor)
+										
+										check(StagingBuffer->GetFormat() == PF_B8G8R8A8);
+										ConvertRawB8G8R8A8DataToFColor(OutputSize.X, OutputSize.Y, (uint8*)Data, DataWidth * sizeof(FColor), OutputColor.GetData());
+									}
+
+									// We can't unmap ourself since we're not on the render thread
+									StagingBufferPool.ReleaseStagingBufferForUnmap_AnyThread(StagingBuffer);
+
+									if (CurrentMaterialSettings->bPerformBorderSmear)
+									{
+										// This will resize the output to a single pixel if the result is monochrome.
+										FGLTFMaterialBakingHelpers::PerformUVBorderSmearAndShrink(OutputColor, OutputSize.X, OutputSize.Y, CurrentMaterialSettings->BackgroundColor);
+									}
+#if WITH_EDITOR
+									// If saving intermediates is turned on
+									if (bSaveIntermediateTextures)
+									{
+										TRACE_CPUPROFILER_EVENT_SCOPE(SaveIntermediateTextures)
+										FString TrimmedPropertyName = Property.ToString();
+
+										const FString DirectoryPath = FPaths::ConvertRelativePathToFull(FPaths::ProjectIntermediateDir() + TEXT("MaterialBaking/"));
+										FString FilenameString = FString::Printf(TEXT("%s%s-%d-%s.bmp"), *DirectoryPath, *CurrentMaterialSettings->Material->GetName(), MaterialIndex, *TrimmedPropertyName);
+										FFileHelper::CreateBitmap(*FilenameString, CurrentOutput.PropertySizes[Property].X, CurrentOutput.PropertySizes[Property].Y, CurrentOutput.PropertyData[Property].GetData());
+									}
+#endif // WITH_EDITOR
+									NumTasks--;
+								};
+
+							// Run previous command if we're going to overwrite it meaning pipeline depth has been reached
+							if (PipelineContext[PipelineIndex].ReadCommand)
+							{
+								PipelineContext[PipelineIndex].ReadCommand(RHICmdList);
+							}
+
+							// Generate a texture reading command that will be executed once it reaches the end of the pipeline
+							PipelineContext[PipelineIndex].ReadCommand =
+								[FinalProcessing_AnyThread, StagingBufferRef = MoveTemp(StagingBufferRef), GPUFence = MoveTemp(GPUFence)](FRHICommandListImmediate& RHICmdList) mutable
+								{
+									TRACE_CPUPROFILER_EVENT_SCOPE(MapAndEnqueue)
+
+									void * Data = nullptr;
+									int32 Width; int32 Height;
+									RHICmdList.MapStagingSurface(StagingBufferRef, GPUFence.GetReference(), Data, Width, Height);
+
+									// Schedule the copy and processing on another thread to free up the render thread as much as possible
+									Async(
+										EAsyncExecution::ThreadPool,
+										[FinalProcessing_AnyThread, Data, Width, Height, StagingBufferRef = MoveTemp(StagingBufferRef)]() mutable
+										{
+											FinalProcessing_AnyThread(StagingBufferRef, Data, Width, Height);
+										}
+									);
+								};
+						}
+					);
+
+					PipelineIndex = (PipelineIndex + 1) % PipelineDepth;
+				}
+			}
+
+		}
+
+		// Destroying Render Items must happen on the render thread to ensure
+		// they are not used anymore.
+		ENQUEUE_RENDER_COMMAND(DestroyRenderItems)(
+			[RenderItems](FRHICommandListImmediate& RHICmdList)
+			{
+				for (auto RenderItem : (*RenderItems))
+				{
+					delete RenderItem.Value;
+				}
+
+				delete RenderItems;
+			}
+		);
+	}
+
+	ENQUEUE_RENDER_COMMAND(ProcessRemainingReads)(
+		[&PipelineContext, PipelineDepth, PipelineIndex](FRHICommandListImmediate& RHICmdList)
+		{
+			// Enqueue remaining reads
+			for (int32 Index = 0; Index < PipelineDepth; Index++)
+			{
+				int32 LocalPipelineIndex = (PipelineIndex + Index) % PipelineDepth;
+
+				if (PipelineContext[LocalPipelineIndex].ReadCommand)
+				{
+					PipelineContext[LocalPipelineIndex].ReadCommand(RHICmdList);
+				}
+			}
+		}
+	);
+
+	// Wait until every tasks have been queued so that NumTasks is only decreasing
+	FlushRenderingCommands();
+
+	// Wait for any remaining final processing tasks
+	while (NumTasks.Load(EMemoryOrder::Relaxed) > 0)
+	{
+		FPlatformProcess::Sleep(0.1f);
+	}
+
+	// Wait for all tasks to have been processed before clearing the staging buffers
+	FlushRenderingCommands();
+
+	ENQUEUE_RENDER_COMMAND(ClearStagingBufferPool)(
+		[&StagingBufferPool](FRHICommandListImmediate& RHICmdList)
+		{
+			StagingBufferPool.Clear_RenderThread(RHICmdList);
+		}
+	);
+
+	// Wait for StagingBufferPool clear to have executed before exiting the function
+	FlushRenderingCommands();
+
+	if (!CVarUseMaterialProxyCaching.GetValueOnAnyThread())
+	{
+		CleanupMaterialProxies();
+	}
+}
+
+void FGLTFMaterialBakingModule::SetEmissiveHDR(bool bHDR)
+{
+	bEmissiveHDR = bHDR;
+}
+
+void FGLTFMaterialBakingModule::SetLinearBake(bool bCorrectLinear)
+{
+	// PerPropertyGamma ultimately sets whether the render target is linear
+	PerPropertyColorSpace.Reset();
+	if (bCorrectLinear)
+	{
+		DefaultColorSpace = EPropertyColorSpace::Linear;
+		PerPropertyColorSpace.Add(MP_BaseColor, EPropertyColorSpace::sRGB);
+		PerPropertyColorSpace.Add(MP_EmissiveColor, EPropertyColorSpace::sRGB);
+		PerPropertyColorSpace.Add(MP_SubsurfaceColor, EPropertyColorSpace::sRGB);
+		PerPropertyColorSpace.Add(TEXT("TransmittanceColor"), EPropertyColorSpace::sRGB);
+	}
+	else
+	{
+		DefaultColorSpace = EPropertyColorSpace::sRGB;
+		PerPropertyColorSpace.Add(MP_Normal, EPropertyColorSpace::Linear);
+		PerPropertyColorSpace.Add(MP_Opacity, EPropertyColorSpace::Linear);
+		PerPropertyColorSpace.Add(MP_OpacityMask, EPropertyColorSpace::Linear);
+		PerPropertyColorSpace.Add(MP_Refraction, EPropertyColorSpace::Linear);
+		PerPropertyColorSpace.Add(MP_ShadingModel, EPropertyColorSpace::Linear);
+		PerPropertyColorSpace.Add(TEXT("ClearCoatBottomNormal"), EPropertyColorSpace::Linear);
+	}
+}
+
+bool FGLTFMaterialBakingModule::IsLinearBake(FGLTFMaterialPropertyEx Property)
+{
+	const EPropertyColorSpace* OverrideColorSpace = PerPropertyColorSpace.Find(Property);
+	const EPropertyColorSpace ColorSpace = OverrideColorSpace ? *OverrideColorSpace : DefaultColorSpace;
+	return ColorSpace == EPropertyColorSpace::Linear;
+}
+
+void FGLTFMaterialBakingModule::CleanupMaterialProxies()
+{
+	for (auto Iterator : MaterialProxyPool)
+	{
+		delete Iterator.Value.Value;
+	}
+	MaterialProxyPool.Reset();
+}
+
+UTextureRenderTarget2D* FGLTFMaterialBakingModule::CreateRenderTarget(bool bInForceLinearGamma, EPixelFormat InPixelFormat, const FIntPoint& InTargetSize, const FColor& BackgroundColor)
+{
+	TRACE_CPUPROFILER_EVENT_SCOPE(FMaterialBakingModule::CreateRenderTarget)
+
+	UTextureRenderTarget2D* RenderTarget = nullptr;
+	const int32 MaxTextureSize = 1 << (MAX_TEXTURE_MIP_COUNT - 1); // Don't use GetMax2DTextureDimension() as this is for the RHI only.
+	const FIntPoint ClampedTargetSize(FMath::Clamp(InTargetSize.X, 1, MaxTextureSize), FMath::Clamp(InTargetSize.Y, 1, MaxTextureSize));
+	auto RenderTargetComparison = [bInForceLinearGamma, InPixelFormat, ClampedTargetSize](const UTextureRenderTarget2D* CompareRenderTarget) -> bool
+	{
+		return (CompareRenderTarget->SizeX == ClampedTargetSize.X && CompareRenderTarget->SizeY == ClampedTargetSize.Y && CompareRenderTarget->OverrideFormat == InPixelFormat && CompareRenderTarget->bForceLinearGamma == bInForceLinearGamma);
+	};
+
+	// Find any pooled render target with suitable properties.
+	UTextureRenderTarget2D** FindResult = RenderTargetPool.FindByPredicate(RenderTargetComparison);
+	
+	if (FindResult)
+	{
+		RenderTarget = *FindResult;
+	}
+	else
+	{
+		TRACE_CPUPROFILER_EVENT_SCOPE(CreateNewRenderTarget)
+
+		// Not found - create a new one.
+		RenderTarget = NewObject<UTextureRenderTarget2D>();
+		check(RenderTarget);
+		RenderTarget->AddToRoot();
+		RenderTarget->ClearColor = bInForceLinearGamma ? BackgroundColor.ReinterpretAsLinear() : FLinearColor(BackgroundColor);
+		RenderTarget->TargetGamma = 0.0f;
+		RenderTarget->InitCustomFormat(ClampedTargetSize.X, ClampedTargetSize.Y, InPixelFormat, bInForceLinearGamma);
+
+		RenderTargetPool.Add(RenderTarget);
+	}
+
+	checkf(RenderTarget != nullptr, TEXT("Unable to create or find valid render target"));
+	return RenderTarget;
+}
+
+FGLTFExportMaterialProxy* FGLTFMaterialBakingModule::CreateMaterialProxy(const FGLTFMaterialDataEx* MaterialSettings, const FGLTFMaterialPropertyEx& Property)
+{
+	TRACE_CPUPROFILER_EVENT_SCOPE(FMaterialBakingModule::CreateMaterialProxy)
+
+	FGLTFExportMaterialProxy* Proxy = nullptr;
+
+	// Find all pooled material proxy matching this material
+	TArray<FMaterialPoolValue> Entries;
+	MaterialProxyPool.MultiFind(MaterialSettings->Material, Entries);
+
+	// Look for the matching property
+	for (FMaterialPoolValue& Entry : Entries)
+	{
+		if (Entry.Key == Property && Entry.Value->ProxyBlendMode == MaterialSettings->BlendMode && Entry.Value->bTangentSpaceNormal == MaterialSettings->bTangentSpaceNormal)
+		{
+			Proxy = Entry.Value;
+			break;
+		}
+	}
+
+	// Not found, create a new entry
+	if (Proxy == nullptr)
+	{
+		Proxy = new FGLTFExportMaterialProxy(MaterialSettings->Material, Property.Type, Property.CustomOutput.ToString(), false /* bInSynchronousCompilation */, MaterialSettings->BlendMode, MaterialSettings->bTangentSpaceNormal);
+		MaterialProxyPool.Add(MaterialSettings->Material, FMaterialPoolValue(Property, Proxy));
+	}
+
+	return Proxy;
+}
+
+void FGLTFMaterialBakingModule::ProcessEmissiveOutput(const FFloat16Color* Color16, int32 Color16Pitch, const FIntPoint& OutputSize, TArray<FColor>& OutputColor, float& EmissiveScale, const FColor& BackgroundColor)
+{
+	TRACE_CPUPROFILER_EVENT_SCOPE(FMaterialBakingModule::ProcessEmissiveOutput)
+
+	const int32 NumThreads = [&]()
+	{
+		return FPlatformProcess::SupportsMultithreading() ? FPlatformMisc::NumberOfCores() : 1;
+	}();
+
+	float* MaxValue = new float[NumThreads];
+	FMemory::Memset(MaxValue, 0, NumThreads * sizeof(MaxValue[0]));
+	const int32 LinesPerThread = FMath::CeilToInt((float)OutputSize.Y / (float)NumThreads);
+	const FFloat16Color BackgroundFloat16 = FFloat16Color(FLinearColor(BackgroundColor)); // Can assume emissive always uses sRGB
+
+	// Find maximum float value across texture
+	ParallelFor(NumThreads, [&Color16, LinesPerThread, MaxValue, OutputSize, Color16Pitch, BackgroundFloat16](int32 Index)
+	{
+		const int32 EndY = FMath::Min((Index + 1) * LinesPerThread, OutputSize.Y);			
+		float& CurrentMaxValue = MaxValue[Index];
+		for (int32 PixelY = Index * LinesPerThread; PixelY < EndY; ++PixelY)
+		{
+			const int32 SrcYOffset = PixelY * Color16Pitch;
+			for (int32 PixelX = 0; PixelX < OutputSize.X; PixelX++)
+			{
+				const FFloat16Color& Pixel16 = Color16[PixelX + SrcYOffset];
+				// Find maximum channel value across texture
+				if (!(Pixel16 == BackgroundFloat16))
+				{
+					CurrentMaxValue = FMath::Max(CurrentMaxValue, FMath::Max3(Pixel16.R.GetFloat(), Pixel16.G.GetFloat(), Pixel16.B.GetFloat()));
+				}
+			}
+		}
+	});
+
+	const float GlobalMaxValue = [&MaxValue, NumThreads]
+	{
+		float TempValue = 0.0f;
+		for (int32 ThreadIndex = 0; ThreadIndex < NumThreads; ++ThreadIndex)
+		{
+			TempValue = FMath::Max(TempValue, MaxValue[ThreadIndex]);
+		}
+
+		return TempValue;
+	}();
+		
+	if (GlobalMaxValue <= 0.01f)
+	{
+		// Black emissive, drop it			
+	}
+
+	// Now convert Float16 to Color using the scale
+	OutputColor.SetNumUninitialized(OutputSize.X * OutputSize.Y);
+	const float Scale = 255.0f / GlobalMaxValue;
+	ParallelFor(NumThreads, [&Color16, LinesPerThread, &OutputColor, OutputSize, Color16Pitch, Scale, BackgroundFloat16, BackgroundColor](int32 Index)
+	{
+		const int32 EndY = FMath::Min((Index + 1) * LinesPerThread, OutputSize.Y);
+		for (int32 PixelY = Index * LinesPerThread; PixelY < EndY; ++PixelY)
+		{
+			const int32 SrcYOffset = PixelY * Color16Pitch;
+			const int32 DstYOffset = PixelY * OutputSize.X;
+
+			for (int32 PixelX = 0; PixelX < OutputSize.X; PixelX++)
+			{
+				const FFloat16Color& Pixel16 = Color16[PixelX + SrcYOffset];
+				FColor& Pixel8 = OutputColor[PixelX + DstYOffset];
+
+				if (Pixel16 == BackgroundFloat16)
+				{
+					Pixel8 = BackgroundColor;
+				}
+				else
+				{
+					Pixel8.R = (uint8)FMath::RoundToInt(Pixel16.R.GetFloat() * Scale);
+					Pixel8.G = (uint8)FMath::RoundToInt(Pixel16.G.GetFloat() * Scale);
+					Pixel8.B = (uint8)FMath::RoundToInt(Pixel16.B.GetFloat() * Scale);
+				}
+					
+				Pixel8.A = 255;
+			}
+		}
+	});
+
+	// This scale will be used in the proxy material to get the original range of emissive values outside of 0-1
+	EmissiveScale = GlobalMaxValue;
+}
+
+void FGLTFMaterialBakingModule::OnObjectModified(UObject* Object)
+{
+	TRACE_CPUPROFILER_EVENT_SCOPE(FMaterialBakingModule::OnObjectModified)
+
+	if (CVarUseMaterialProxyCaching.GetValueOnAnyThread())
+	{
+		UMaterialInterface* MaterialToInvalidate = Cast<UMaterialInterface>(Object);
+		if (!MaterialToInvalidate)
+		{
+			// Check to see if the object is a material editor instance constant and if so, retrieve its source instance
+			UMaterialEditorInstanceConstant* EditorInstance = Cast<UMaterialEditorInstanceConstant>(Object);
+			if (EditorInstance && EditorInstance->SourceInstance)
+			{
+				MaterialToInvalidate = EditorInstance->SourceInstance;
+			}
+		}
+
+		if (MaterialToInvalidate)
+		{
+			// Search our proxy pool for materials or material instances that refer to MaterialToInvalidate
+			for (auto It = MaterialProxyPool.CreateIterator(); It; ++It)
+			{
+				TWeakObjectPtr<UMaterialInterface> PoolMaterialPtr = It.Key();
+
+				// Remove stale entries from the pool
+				bool bMustDelete = PoolMaterialPtr.IsValid();
+				if (!bMustDelete)
+				{
+					bMustDelete = PoolMaterialPtr == MaterialToInvalidate;
+				}
+
+				// No match - Test the MaterialInstance hierarchy
+				if (!bMustDelete)
+				{
+					UMaterialInstance* MaterialInstance = Cast<UMaterialInstance>(PoolMaterialPtr);
+					while (!bMustDelete && MaterialInstance && MaterialInstance->Parent != nullptr)
+					{
+						bMustDelete = MaterialInstance->Parent == MaterialToInvalidate;
+						MaterialInstance = Cast<UMaterialInstance>(MaterialInstance->Parent);
+					}
+				}
+
+				// We have a match, remove the entry from our pool
+				if (bMustDelete)
+				{
+					FGLTFExportMaterialProxy* Proxy = It.Value().Value;
+
+					ENQUEUE_RENDER_COMMAND(DeleteCachedMaterialProxy)(
+						[Proxy](FRHICommandListImmediate& RHICmdList)
+						{
+							delete Proxy;
+						});
+
+					It.RemoveCurrent();
+				}
+			}
+		}
+	}
+}
+
+#undef LOCTEXT_NAMESPACE //"MaterialBakingModule"
diff --git a/Engine/Plugins/Enterprise/GLTFExporter/MigratedFromGit/Source/GLTFMaterialBaking/Private/GLTFMaterialBakingModule.h b/Engine/Plugins/Enterprise/GLTFExporter/MigratedFromGit/Source/GLTFMaterialBaking/Private/GLTFMaterialBakingModule.h
new file mode 100644
--- /dev/null
+++ b/Engine/Plugins/Enterprise/GLTFExporter/MigratedFromGit/Source/GLTFMaterialBaking/Private/GLTFMaterialBakingModule.h
@@ -0,0 +1,85 @@
+// Copyright Epic Games, Inc. All Rights Reserved.
+
+#pragma once
+
+#include "CoreMinimal.h"
+#include "Modules/ModuleInterface.h"
+#include "PixelFormat.h"
+#include "SceneTypes.h"
+#include "IGLTFMaterialBakingModule.h"
+#include "GLTFMaterialPropertyEx.h"
+
+class UTextureRenderTarget2D;
+class UMaterialInterface;
+
+class FGLTFExportMaterialProxy;
+class FTextureRenderTargetResource;
+
+struct FGLTFMaterialData; 
+struct FGLTFMeshRenderData;
+struct FGLTFBakeOutput;
+
+class GLTFMATERIALBAKING_API FGLTFMaterialBakingModule : public IGLTFMaterialBakingModule
+{
+public:
+	/** IModuleInterface overrides begin */
+	virtual void StartupModule();
+	virtual void ShutdownModule();
+	/** IModuleInterface overrides end */
+
+	/** Bakes out material properties according to MaterialSettings using MeshSettings and stores the output in Output */
+	virtual void BakeMaterials(const TArray<FGLTFMaterialData*>& MaterialSettings, const TArray<FGLTFMeshRenderData*>& MeshSettings, TArray<FGLTFBakeOutput>& Output) override;
+
+	/** Bakes out material properties according to extended MaterialSettings using MeshSettings and stores the output in Output */
+	virtual void BakeMaterials(const TArray<FGLTFMaterialDataEx*>& MaterialSettings, const TArray<FGLTFMeshRenderData*>& MeshSettings, TArray<FGLTFBakeOutputEx>& Output) override;
+
+	/** Outputs true HDR version of emissive color */
+	virtual void SetEmissiveHDR(bool bHDR) override;
+
+	/** Bakes all material properties to linear textures, except for colors */
+	virtual void SetLinearBake(bool bCorrectLinear) override;
+
+	/** Returns whether a specific material property is baked to a linear texture or not */
+	virtual bool IsLinearBake(FGLTFMaterialPropertyEx Property) override;
+
+protected:
+	/* Creates and adds or reuses a RenderTarget from the pool */
+	UTextureRenderTarget2D* CreateRenderTarget(bool bInForceLinearGamma, EPixelFormat InPixelFormat, const FIntPoint& InTargetSize, const FColor& BackgroundColor);
+
+	/* Creates and adds (or reuses a ExportMaterialProxy from the pool if MaterialBaking.UseMaterialProxyCaching is set to 1) */
+	FGLTFExportMaterialProxy* CreateMaterialProxy(const FGLTFMaterialDataEx* MaterialSettings, const FGLTFMaterialPropertyEx& Property);
+
+	/** Helper for emissive color conversion to Output */
+	static void ProcessEmissiveOutput(const FFloat16Color* Color16, int32 Color16Pitch, const FIntPoint& OutputSize, TArray<FColor>& Output, float& EmissiveScale, const FColor& BackgroundColor);
+
+	/** Cleans up all cached material proxies in MaterialProxyPool */
+	void CleanupMaterialProxies();
+
+	/** Callback for modified objects which should be removed from MaterialProxyPool in that case */
+	void OnObjectModified(UObject* Object);
+private:
+	enum EPropertyColorSpace
+	{
+		Linear,
+		sRGB,
+	};
+
+	/** Pool of available render targets, cached for re-using on consecutive property rendering */
+	TArray<UTextureRenderTarget2D*> RenderTargetPool;
+
+	/** Pool of cached material proxies to optimize material baking workflow, stays resident when MaterialBaking.UseMaterialProxyCaching is set to 1 */
+	typedef TWeakObjectPtr<UMaterialInterface>				FMaterialPoolKey;
+	typedef TPair<FGLTFMaterialPropertyEx, FGLTFExportMaterialProxy*> FMaterialPoolValue;
+	typedef TMultiMap<FMaterialPoolKey, FMaterialPoolValue, FDefaultSetAllocator, TWeakObjectPtrMapKeyFuncs<FMaterialPoolKey, FMaterialPoolValue, true /*bInAllowDuplicateKeys*/>> FMaterialPoolMap;
+	FMaterialPoolMap MaterialProxyPool;
+
+	/** Pixel formats to use for baking out specific material properties */
+	TMap<FGLTFMaterialPropertyEx, EPixelFormat> PerPropertyFormat;
+
+	/** Whether or not to enforce gamma correction while baking out specific material properties */
+	TMap<FGLTFMaterialPropertyEx, EPropertyColorSpace> PerPropertyColorSpace;
+
+	EPropertyColorSpace DefaultColorSpace;
+
+	bool bEmissiveHDR;
+};
diff --git a/Engine/Plugins/Enterprise/GLTFExporter/MigratedFromGit/Source/GLTFMaterialBaking/Private/GLTFMaterialPropertyEx.cpp b/Engine/Plugins/Enterprise/GLTFExporter/MigratedFromGit/Source/GLTFMaterialBaking/Private/GLTFMaterialPropertyEx.cpp
new file mode 100644
--- /dev/null
+++ b/Engine/Plugins/Enterprise/GLTFExporter/MigratedFromGit/Source/GLTFMaterialBaking/Private/GLTFMaterialPropertyEx.cpp
@@ -0,0 +1,18 @@
+// Copyright Epic Games, Inc. All Rights Reserved.
+
+#include "GLTFMaterialPropertyEx.h"
+#include "UObject/NameTypes.h"
+
+FString FGLTFMaterialPropertyEx::ToString() const
+{
+	if (!IsCustomOutput())
+	{
+		const UEnum* Enum = StaticEnum<EMaterialProperty>();
+		FName Name = Enum->GetNameByValue(Type);
+		FString TrimmedName = Name.ToString();
+		TrimmedName.RemoveFromStart(TEXT("MP_"), ESearchCase::CaseSensitive);
+		return TrimmedName;
+	}
+
+	return CustomOutput.ToString();
+}
diff --git a/Engine/Plugins/Enterprise/GLTFExporter/MigratedFromGit/Source/GLTFMaterialBaking/Private/MaterialBakingHelpers.cpp b/Engine/Plugins/Enterprise/GLTFExporter/MigratedFromGit/Source/GLTFMaterialBaking/Private/MaterialBakingHelpers.cpp
deleted file mode 100644
diff --git a/Engine/Plugins/Enterprise/GLTFExporter/MigratedFromGit/Source/GLTFMaterialBaking/Private/MaterialBakingModule.cpp b/Engine/Plugins/Enterprise/GLTFExporter/MigratedFromGit/Source/GLTFMaterialBaking/Private/MaterialBakingModule.cpp
deleted file mode 100644
diff --git a/Engine/Plugins/Enterprise/GLTFExporter/MigratedFromGit/Source/GLTFMaterialBaking/Private/MaterialBakingModule.h b/Engine/Plugins/Enterprise/GLTFExporter/MigratedFromGit/Source/GLTFMaterialBaking/Private/MaterialBakingModule.h
deleted file mode 100644
diff --git a/Engine/Plugins/Enterprise/GLTFExporter/MigratedFromGit/Source/GLTFMaterialBaking/Private/MaterialRenderItem.cpp b/Engine/Plugins/Enterprise/GLTFExporter/MigratedFromGit/Source/GLTFMaterialBaking/Private/MaterialRenderItem.cpp
--- a/Engine/Plugins/Enterprise/GLTFExporter/MigratedFromGit/Source/GLTFMaterialBaking/Private/MaterialRenderItem.cpp
+++ b/Engine/Plugins/Enterprise/GLTFExporter/MigratedFromGit/Source/GLTFMaterialBaking/Private/MaterialRenderItem.cpp
@@ -10,9 +10,9 @@
 
 #define SHOW_WIREFRAME_MESH 0
 
-FMeshMaterialRenderItem::FMeshMaterialRenderItem(
+FGLTFMeshMaterialRenderItem::FGLTFMeshMaterialRenderItem(
 	const FIntPoint& InTextureSize,
-	const FMeshData* InMeshSettings,
+	const FGLTFMeshRenderData* InMeshSettings,
 	FDynamicMeshBufferAllocator* InDynamicMeshBufferAllocator)
 	: MeshSettings(InMeshSettings)
 	, TextureSize(InTextureSize)
@@ -22,10 +22,10 @@ FMeshMaterialRenderItem::FMeshMaterialRenderItem(
 	, DynamicMeshBufferAllocator(InDynamicMeshBufferAllocator)
 {
 	GenerateRenderData();
-	LCI = new FMeshRenderInfo(InMeshSettings->LightMap, nullptr, nullptr, InMeshSettings->LightmapResourceCluster);
+	LCI = new FGLTFMeshRenderInfo(InMeshSettings->LightMap, nullptr, nullptr, InMeshSettings->LightmapResourceCluster);
 }
 
-bool FMeshMaterialRenderItem::Render_RenderThread(FRHICommandListImmediate& RHICmdList, FMeshPassProcessorRenderState& DrawRenderState, const FCanvas* Canvas)
+bool FGLTFMeshMaterialRenderItem::Render_RenderThread(FRHICommandListImmediate& RHICmdList, FMeshPassProcessorRenderState& DrawRenderState, const FCanvas* Canvas)
 {
 	checkSlow(ViewFamily && MeshSettings && MaterialRenderProxy);
 	// current render target set for the canvas
@@ -63,7 +63,7 @@ bool FMeshMaterialRenderItem::Render_RenderThread(FRHICommandListImmediate& RHIC
 	return true;
 }
 
-bool FMeshMaterialRenderItem::Render_GameThread(const FCanvas* Canvas, FRenderThreadScope& RenderScope)
+bool FGLTFMeshMaterialRenderItem::Render_GameThread(const FCanvas* Canvas, FRenderThreadScope& RenderScope)
 {
 	RenderScope.EnqueueRenderCommand(
 		[this, Canvas](FRHICommandListImmediate& RHICmdList)
@@ -77,7 +77,7 @@ bool FMeshMaterialRenderItem::Render_GameThread(const FCanvas* Canvas, FRenderTh
 	return true;
 }
 
-void FMeshMaterialRenderItem::GenerateRenderData()
+void FGLTFMeshMaterialRenderItem::GenerateRenderData()
 {
 	TRACE_CPUPROFILER_EVENT_SCOPE(FMeshMaterialRenderItem::GenerateRenderData)
 
@@ -98,7 +98,7 @@ void FMeshMaterialRenderItem::GenerateRenderData()
 	bMeshElementDirty = true;
 }
 
-FMeshMaterialRenderItem::~FMeshMaterialRenderItem()
+FGLTFMeshMaterialRenderItem::~FGLTFMeshMaterialRenderItem()
 {
 	// Send the release of the buffers to the render thread
 	ENQUEUE_RENDER_COMMAND(ReleaseResources)(
@@ -106,7 +106,7 @@ FMeshMaterialRenderItem::~FMeshMaterialRenderItem()
 	);
 }
 
-void FMeshMaterialRenderItem::QueueMaterial(FRHICommandListImmediate& RHICmdList, FMeshPassProcessorRenderState& DrawRenderState, const FSceneView* View)
+void FGLTFMeshMaterialRenderItem::QueueMaterial(FRHICommandListImmediate& RHICmdList, FMeshPassProcessorRenderState& DrawRenderState, const FSceneView* View)
 {
 	TRACE_CPUPROFILER_EVENT_SCOPE(FMeshMaterialRenderItem::QueueMaterial)
 
@@ -123,7 +123,7 @@ void FMeshMaterialRenderItem::QueueMaterial(FRHICommandListImmediate& RHICmdList
 		DynamicMeshBuilder.GetMeshElement(FMatrix::Identity, MaterialRenderProxy, SDPG_Foreground, true, false, 0, MeshBuilderResources, MeshElement);
 
 		// NOTE: Since we can't pass these parameters during creation, we overwrite them here
-		const FPrimitiveData* PrimitiveData = MeshSettings->PrimitiveData;
+		const FGLTFPrimitiveData* PrimitiveData = MeshSettings->PrimitiveData;
 		if (PrimitiveData != nullptr)
 		{
 			FPrimitiveUniformShaderParameters PrimitiveParams = GetPrimitiveUniformShaderParameters(
@@ -178,7 +178,7 @@ void FMeshMaterialRenderItem::QueueMaterial(FRHICommandListImmediate& RHICmdList
 	GetRendererModule().DrawTileMesh(RHICmdList, DrawRenderState, *View, MeshElement, false /*bIsHitTesting*/, FHitProxyId());
 }
 
-void FMeshMaterialRenderItem::PopulateWithQuadData()
+void FGLTFMeshMaterialRenderItem::PopulateWithQuadData()
 {
 	Vertices.Empty(4);
 	Indices.Empty(6);
@@ -211,7 +211,7 @@ void FMeshMaterialRenderItem::PopulateWithQuadData()
 	Indices.Append(TriangleIndices, 6);
 }
 
-void FMeshMaterialRenderItem::PopulateWithMeshData()
+void FGLTFMeshMaterialRenderItem::PopulateWithMeshData()
 {
 	const FMeshDescription* RawMesh = MeshSettings->RawMeshDescription;
 
diff --git a/Engine/Plugins/Enterprise/GLTFExporter/MigratedFromGit/Source/GLTFMaterialBaking/Private/MaterialRenderItem.h b/Engine/Plugins/Enterprise/GLTFExporter/MigratedFromGit/Source/GLTFMaterialBaking/Private/MaterialRenderItem.h
--- a/Engine/Plugins/Enterprise/GLTFExporter/MigratedFromGit/Source/GLTFMaterialBaking/Private/MaterialRenderItem.h
+++ b/Engine/Plugins/Enterprise/GLTFExporter/MigratedFromGit/Source/GLTFMaterialBaking/Private/MaterialRenderItem.h
@@ -10,17 +10,17 @@ class FSceneViewFamily;
 class FMaterialRenderProxy;
 class FSceneView;
 class FRHICommandListImmediate;
-struct FMaterialData;
-struct FMeshData;
+struct FGLTFMaterialData;
+struct FGLTFMeshRenderData;
 struct FMeshPassProcessorRenderState;
 
 // This will hold onto the resources until not needed anymore.
 // Move constructor makes it easier to send the destruction to another thread (render thread).
-class FMeshBuilderResources : public FMeshBuilderOneFrameResources
+class FGLTFMeshBuilderResources : public FMeshBuilderOneFrameResources
 {
 public:
-	FMeshBuilderResources() = default;
-	FMeshBuilderResources(FMeshBuilderResources&& Other)
+	FGLTFMeshBuilderResources() = default;
+	FGLTFMeshBuilderResources(FGLTFMeshBuilderResources&& Other)
 	{
 		if (this != &Other)
 		{
@@ -46,11 +46,11 @@ public:
 	}
 };
 
-class FMeshMaterialRenderItem : public FCanvasBaseRenderItem
+class FGLTFMeshMaterialRenderItem : public FCanvasBaseRenderItem
 {
 public:
-	FMeshMaterialRenderItem(const FIntPoint& InTextureSize, const FMeshData* InMeshSettings, FDynamicMeshBufferAllocator* InDynamicMeshBufferAllocator = nullptr);
-	virtual ~FMeshMaterialRenderItem();
+	FGLTFMeshMaterialRenderItem(const FIntPoint& InTextureSize, const FGLTFMeshRenderData* InMeshSettings, FDynamicMeshBufferAllocator* InDynamicMeshBufferAllocator = nullptr);
+	virtual ~FGLTFMeshMaterialRenderItem();
 
 	/** Begin FCanvasBaseRenderItem overrides */
 	virtual bool Render_RenderThread(FRHICommandListImmediate& RHICmdList, FMeshPassProcessorRenderState& DrawRenderState, const FCanvas* Canvas) final;
@@ -68,7 +68,7 @@ protected:
 	void PopulateWithMeshData();
 public:
 	/** Mesh settings to use while baking out the material */
-	const FMeshData* MeshSettings;
+	const FGLTFMeshRenderData* MeshSettings;
 	/** The texture size to use while baking */
 	FIntPoint TextureSize;
 	/** Material render proxy (material/shader) to use while baking */
@@ -83,6 +83,6 @@ public:
 private:
 	FMeshBatch MeshElement;
 	bool bMeshElementDirty;
-	FMeshBuilderResources MeshBuilderResources;
+	FGLTFMeshBuilderResources MeshBuilderResources;
 	FDynamicMeshBufferAllocator* DynamicMeshBufferAllocator;
 };
diff --git a/Engine/Plugins/Enterprise/GLTFExporter/MigratedFromGit/Source/GLTFMaterialBaking/Private/MaterialRenderItemData.h b/Engine/Plugins/Enterprise/GLTFExporter/MigratedFromGit/Source/GLTFMaterialBaking/Private/MaterialRenderItemData.h
--- a/Engine/Plugins/Enterprise/GLTFExporter/MigratedFromGit/Source/GLTFMaterialBaking/Private/MaterialRenderItemData.h
+++ b/Engine/Plugins/Enterprise/GLTFExporter/MigratedFromGit/Source/GLTFMaterialBaking/Private/MaterialRenderItemData.h
@@ -8,10 +8,10 @@
 #include "SceneManagement.h"
 
 /** Simple implementation for light cache to simulated lightmap behaviour (used for accessing prebaked ambient occlusion values) */
-class FMeshRenderInfo : public FLightCacheInterface
+class FGLTFMeshRenderInfo : public FLightCacheInterface
 {
 public:
-	FMeshRenderInfo(const FLightMap* InLightMap, const FShadowMap* InShadowMap, FUniformBufferRHIRef Buffer, const FLightmapResourceCluster* InLightmapResourceCluster)
+	FGLTFMeshRenderInfo(const FLightMap* InLightMap, const FShadowMap* InShadowMap, FUniformBufferRHIRef Buffer, const FLightmapResourceCluster* InLightmapResourceCluster)
 		: FLightCacheInterface()
 	{
 		SetLightMap(InLightMap);
diff --git a/Engine/Plugins/Enterprise/GLTFExporter/MigratedFromGit/Source/GLTFMaterialBaking/Public/GLTFMaterialBakingHelpers.h b/Engine/Plugins/Enterprise/GLTFExporter/MigratedFromGit/Source/GLTFMaterialBaking/Public/GLTFMaterialBakingHelpers.h
new file mode 100644
--- /dev/null
+++ b/Engine/Plugins/Enterprise/GLTFExporter/MigratedFromGit/Source/GLTFMaterialBaking/Public/GLTFMaterialBakingHelpers.h
@@ -0,0 +1,32 @@
+// Copyright Epic Games, Inc. All Rights Reserved.
+
+#pragma once
+
+#include "Containers/Array.h"
+
+struct FColor;
+
+class GLTFMATERIALBAKING_API FGLTFMaterialBakingHelpers
+{
+public:
+	/** 
+	 * Applies a box blur to magenta pixels found in given texture represented by InOutPixels using non-background pixels, this creates a smear across the magenta/filled pixels.
+	 * @param  InOutPixels     The image to apply the smear to
+	 * @param  ImageWidth      The width of the image in pixels
+	 * @param  ImageHeight     The height of the image in pixels
+	 * @param  MaxIterations   The max distance in pixels to smear the edges of the texture. When set to -1 this value is the max of the width and height of the buffer
+	 * @param  BackgroundColor The pixel color considered background in the image, and that can be smeared into
+	 */
+	static void PerformUVBorderSmear(TArray<FColor>& InOutPixels, int32 ImageWidth, int32 ImageHeight, int32 MaxIterations = -1, FColor BackgroundColor = FColor::Magenta);
+
+	/** 
+	 * Applies a box blur to magenta pixels found in given texture represented by InOutPixels using non-background pixels, this creates a smear across the magenta/filled pixels.
+	 * @param  InOutPixels      The image to apply the smear to
+	 * @param  InOutImageWidth  The width of the image in pixels
+	 * @param  InOutImageHeight The height of the image in pixels
+	 * @param  BackgroundColor  The pixel color considered background in the image, and that can be smeared into
+	 *
+	 * @note   This version is allowed to shrink the array to a single pixel if the entire output is monochrome. Params will be adjusted accordingly when this occurs.
+	 */
+	static void PerformUVBorderSmearAndShrink(TArray<FColor>& InOutPixels, int32& InOutImageWidth, int32& InOutImageHeight, FColor BackgroundColor = FColor::Magenta);
+};
\ No newline at end of file
diff --git a/Engine/Plugins/Enterprise/GLTFExporter/MigratedFromGit/Source/GLTFMaterialBaking/Public/GLTFMaterialPropertyEx.h b/Engine/Plugins/Enterprise/GLTFExporter/MigratedFromGit/Source/GLTFMaterialBaking/Public/GLTFMaterialPropertyEx.h
new file mode 100644
--- /dev/null
+++ b/Engine/Plugins/Enterprise/GLTFExporter/MigratedFromGit/Source/GLTFMaterialBaking/Public/GLTFMaterialPropertyEx.h
@@ -0,0 +1,52 @@
+// Copyright Epic Games, Inc. All Rights Reserved.
+
+#pragma once
+
+#include "SceneTypes.h"
+
+/** Structure extending EMaterialProperty to allow detailed information about custom output */
+struct FGLTFMaterialPropertyEx
+{
+	FGLTFMaterialPropertyEx(EMaterialProperty Type = MP_MAX, const FName& CustomOutput = NAME_None)
+		: Type(Type)
+		, CustomOutput(CustomOutput)
+	{}
+
+	FGLTFMaterialPropertyEx(const FName& CustomOutput)
+		: Type(MP_CustomOutput)
+		, CustomOutput(CustomOutput)
+	{}
+
+	FGLTFMaterialPropertyEx(const TCHAR* CustomOutput)
+		: Type(MP_CustomOutput)
+		, CustomOutput(CustomOutput)
+	{}
+
+	FORCEINLINE bool IsCustomOutput() const
+	{
+		return Type == MP_CustomOutput;
+	}
+
+	FORCEINLINE bool operator ==(const FGLTFMaterialPropertyEx& Other) const
+	{
+		return Type == Other.Type && (!IsCustomOutput() || CustomOutput == Other.CustomOutput);
+	}
+
+	FORCEINLINE bool operator !=(const FGLTFMaterialPropertyEx& Other) const
+	{
+		return !(*this == Other);
+	}
+
+	friend FORCEINLINE uint32 GetTypeHash(const FGLTFMaterialPropertyEx& Other)
+	{
+		return !Other.IsCustomOutput() ? GetTypeHash(Other.Type) : GetTypeHash(Other.CustomOutput);
+	}
+
+	GLTFMATERIALBAKING_API FString ToString() const;
+
+	/** The material property */
+	EMaterialProperty Type;
+
+	/** The name of a specific custom output. Only used if property is MP_CustomOutput */
+	FName CustomOutput;
+};
diff --git a/Engine/Plugins/Enterprise/GLTFExporter/MigratedFromGit/Source/GLTFMaterialBaking/Public/IGLTFMaterialBakingModule.h b/Engine/Plugins/Enterprise/GLTFExporter/MigratedFromGit/Source/GLTFMaterialBaking/Public/IGLTFMaterialBakingModule.h
new file mode 100644
--- /dev/null
+++ b/Engine/Plugins/Enterprise/GLTFExporter/MigratedFromGit/Source/GLTFMaterialBaking/Public/IGLTFMaterialBakingModule.h
@@ -0,0 +1,40 @@
+// Copyright Epic Games, Inc. All Rights Reserved.
+
+#pragma once
+
+#include "CoreMinimal.h"
+#include "Modules/ModuleInterface.h"
+#include "GLTFMaterialPropertyEx.h"
+#include "PixelFormat.h"
+#include "SceneTypes.h"
+
+class UTextureRenderTarget2D;
+class UMaterialInterface;
+
+class FGLTFExportMaterialProxy;
+class FTextureRenderTargetResource;
+
+struct FGLTFMaterialData; 
+struct FGLTFMaterialDataEx;
+struct FGLTFMeshRenderData;
+struct FGLTFBakeOutput;
+struct FGLTFBakeOutputEx;
+
+class GLTFMATERIALBAKING_API IGLTFMaterialBakingModule : public IModuleInterface
+{
+public:
+	/** Bakes out material properties according to MaterialSettings using MeshSettings and stores the output in Output */
+	virtual void BakeMaterials(const TArray<FGLTFMaterialData*>& MaterialSettings, const TArray<FGLTFMeshRenderData*>& MeshSettings, TArray<FGLTFBakeOutput>& Output) = 0;
+
+	/** Bakes out material properties according to extended MaterialSettings using MeshSettings and stores the output in Output */
+	virtual void BakeMaterials(const TArray<FGLTFMaterialDataEx*>& MaterialSettings, const TArray<FGLTFMeshRenderData*>& MeshSettings, TArray<FGLTFBakeOutputEx>& Output) = 0;
+
+	/** Outputs true HDR version of emissive color */
+	virtual void SetEmissiveHDR(bool bHDR) = 0;
+
+	/** Bakes all material properties to linear textures, except for colors */
+	virtual void SetLinearBake(bool bCorrectLinear) = 0;
+
+	/** Returns whether a specific material property is baked to a linear texture or not */
+	virtual bool IsLinearBake(FGLTFMaterialPropertyEx Property) = 0;
+};
diff --git a/Engine/Plugins/Enterprise/GLTFExporter/MigratedFromGit/Source/GLTFMaterialBaking/Public/IMaterialBakingModule.h b/Engine/Plugins/Enterprise/GLTFExporter/MigratedFromGit/Source/GLTFMaterialBaking/Public/IMaterialBakingModule.h
deleted file mode 100644
diff --git a/Engine/Plugins/Enterprise/GLTFExporter/MigratedFromGit/Source/GLTFMaterialBaking/Public/MaterialBakingHelpers.h b/Engine/Plugins/Enterprise/GLTFExporter/MigratedFromGit/Source/GLTFMaterialBaking/Public/MaterialBakingHelpers.h
deleted file mode 100644
diff --git a/Engine/Plugins/Enterprise/GLTFExporter/MigratedFromGit/Source/GLTFMaterialBaking/Public/MaterialBakingStructures.h b/Engine/Plugins/Enterprise/GLTFExporter/MigratedFromGit/Source/GLTFMaterialBaking/Public/MaterialBakingStructures.h
--- a/Engine/Plugins/Enterprise/GLTFExporter/MigratedFromGit/Source/GLTFMaterialBaking/Public/MaterialBakingStructures.h
+++ b/Engine/Plugins/Enterprise/GLTFExporter/MigratedFromGit/Source/GLTFMaterialBaking/Public/MaterialBakingStructures.h
@@ -5,15 +5,15 @@
 #include "Math/IntPoint.h"
 #include "SceneTypes.h"
 #include "LightMap.h"
-#include "MaterialPropertyEx.h"
+#include "GLTFMaterialPropertyEx.h"
 
 class UMaterialInterface;
 struct FMeshDescription;
 
 /** Structure containing information about the material which is being baked out */
-struct FMaterialData
+struct FGLTFMaterialData
 {
-	FMaterialData()
+	FGLTFMaterialData()
 		: Material(nullptr)
 		, bPerformBorderSmear(true)
 		, BlendMode(BLEND_Opaque)
@@ -36,9 +36,9 @@ struct FMaterialData
 };
 
 /** Structure containing extended information about the material and properties which is being baked out */
-struct FMaterialDataEx
+struct FGLTFMaterialDataEx
 {
-	FMaterialDataEx()
+	FGLTFMaterialDataEx()
 		: Material(nullptr)
 		, bPerformBorderSmear(true)
 		, BlendMode(BLEND_Opaque)
@@ -49,7 +49,7 @@ struct FMaterialDataEx
 	/** Material to bake out */
 	UMaterialInterface* Material;
 	/** Extended properties and the texture size at which they should be baked out */
-	TMap<FMaterialPropertyEx, FIntPoint> PropertySizes;
+	TMap<FGLTFMaterialPropertyEx, FIntPoint> PropertySizes;
 	/** Whether to smear borders after baking */
 	bool bPerformBorderSmear;
 	/** Blend mode to use when baking, allowing for example detection of overlapping UVs */
@@ -61,9 +61,9 @@ struct FMaterialDataEx
 };
 
 /** Structure containing primitive information (regarding a mesh or mesh component) that is accessible through material expressions */
-struct FPrimitiveData
+struct FGLTFPrimitiveData
 {
-	FPrimitiveData(const FBoxSphereBounds& LocalBounds = FBoxSphereBounds(ForceInitToZero))
+	FGLTFPrimitiveData(const FBoxSphereBounds& LocalBounds = FBoxSphereBounds(ForceInitToZero))
 		: ActorPosition(ForceInitToZero)
 		, ObjectOrientation(FVector::UpVector)
 		, ObjectBounds(LocalBounds)
@@ -72,15 +72,15 @@ struct FPrimitiveData
 		, CustomPrimitiveData(nullptr)
 	{}
 
-	FPrimitiveData(const UStaticMesh* StaticMesh)
-		: FPrimitiveData(StaticMesh->GetBounds())
+	FGLTFPrimitiveData(const UStaticMesh* StaticMesh)
+		: FGLTFPrimitiveData(StaticMesh->GetBounds())
 	{}
 
-	FPrimitiveData(const USkeletalMesh* SkeletalMesh)
-		: FPrimitiveData(SkeletalMesh->GetBounds())
+	FGLTFPrimitiveData(const USkeletalMesh* SkeletalMesh)
+		: FGLTFPrimitiveData(SkeletalMesh->GetBounds())
 	{}
 
-	FPrimitiveData(const UMeshComponent* MeshComponent)
+	FGLTFPrimitiveData(const UMeshComponent* MeshComponent)
 		: ActorPosition(MeshComponent->GetComponentLocation())
 		, ObjectOrientation(MeshComponent->GetRenderMatrix().GetUnitAxis(EAxis::Z))
 		, ObjectBounds(MeshComponent->CalcBounds(MeshComponent->GetComponentTransform()))
@@ -111,9 +111,9 @@ struct FPrimitiveData
 	const FCustomPrimitiveData* CustomPrimitiveData;
 };
 
-struct FMeshData
+struct FGLTFMeshRenderData
 {
-	FMeshData()
+	FGLTFMeshRenderData()
 		: RawMeshDescription(nullptr), Mesh(nullptr), bMirrored(false), VertexColorHash(0), TextureCoordinateIndex(0), LightMapIndex(0), LightMap(nullptr), LightmapResourceCluster(nullptr), PrimitiveData(nullptr)
 	{}
 
@@ -151,13 +151,13 @@ struct FMeshData
 	const FLightmapResourceCluster* LightmapResourceCluster;
 
 	/** Pointer to primitive data that is accessible through material expressions, if nullptr default values are used */
-	const FPrimitiveData* PrimitiveData;
+	const FGLTFPrimitiveData* PrimitiveData;
 };
 
 /** Structure containing data being processed while baking out materials*/
-struct FBakeOutput
+struct FGLTFBakeOutput
 {
-	FBakeOutput()
+	FGLTFBakeOutput()
 		: EmissiveScale(1.0f)
 	{}
 
@@ -175,20 +175,20 @@ struct FBakeOutput
 };
 
 /** Structure containing extended data being processed while baking out materials*/
-struct FBakeOutputEx
+struct FGLTFBakeOutputEx
 {
-	FBakeOutputEx()
+	FGLTFBakeOutputEx()
 		: EmissiveScale(1.0f)
 	{}
 
 	/** Contains the resulting texture data for baking out a extened material's property */
-	TMap<FMaterialPropertyEx, TArray<FColor>> PropertyData;
+	TMap<FGLTFMaterialPropertyEx, TArray<FColor>> PropertyData;
 
 	/** Contains the resulting texture size for baking out a extened material's property */
-	TMap<FMaterialPropertyEx, FIntPoint> PropertySizes;
+	TMap<FGLTFMaterialPropertyEx, FIntPoint> PropertySizes;
 
 	/** Contains the resulting HDR texture data for baking out a material's property, may be empty */
-	TMap<FMaterialPropertyEx, TArray<FFloat16Color>> HDRPropertyData;
+	TMap<FGLTFMaterialPropertyEx, TArray<FFloat16Color>> HDRPropertyData;
 
 	/** Scale used to allow having wide ranges of emissive values in the source materials, the final proxy material will use this value to scale the emissive texture's pixel values */
 	float EmissiveScale;
diff --git a/Engine/Plugins/Enterprise/GLTFExporter/MigratedFromGit/Source/GLTFMaterialBaking/Public/MaterialPropertyEx.h b/Engine/Plugins/Enterprise/GLTFExporter/MigratedFromGit/Source/GLTFMaterialBaking/Public/MaterialPropertyEx.h
deleted file mode 100644
