[GLTFExporter] refactor: fix issues reported by clang on linux that prevents a successful build

commit 5401368732f1a626bb5d351c5373bb7ac52d0d0d
Date: Fri, 25 Feb 2022 15:07:05 +0100

    refactor: fix issues reported by clang on linux that prevents a successful build
    (cherry picked from commit 1c23c35151b4ab3e52a4199a6c06459804a98929​)

#skipci

diff --git a/Engine/Plugins/Enterprise/GLTFExporter/MigratedFromGit/Source/GLTFExporter/Private/Utilities/GLTFProxyMaterialUtilities.cpp b/Engine/Plugins/Enterprise/GLTFExporter/MigratedFromGit/Source/GLTFExporter/Private/Utilities/GLTFProxyMaterialUtilities.cpp
--- a/Engine/Plugins/Enterprise/GLTFExporter/MigratedFromGit/Source/GLTFExporter/Private/Utilities/GLTFProxyMaterialUtilities.cpp
+++ b/Engine/Plugins/Enterprise/GLTFExporter/MigratedFromGit/Source/GLTFExporter/Private/Utilities/GLTFProxyMaterialUtilities.cpp
@@ -231,7 +231,7 @@ bool FGLTFProxyMaterialUtilities::GetTwoSided(const UMaterialInstance* Material,
 bool FGLTFProxyMaterialUtilities::GetBlendMode(const UMaterialInstance* Material, EBlendMode& OutValue, bool NonDefaultOnly)
 {
 	const EBlendMode DefaultValue = Material->Parent->GetBlendMode();
-	const EBlendMode Value = Material->BasePropertyOverrides.bOverride_BlendMode ? Material->BasePropertyOverrides.BlendMode : DefaultValue;
+	const EBlendMode Value = Material->BasePropertyOverrides.bOverride_BlendMode ? Material->BasePropertyOverrides.BlendMode.GetValue() : DefaultValue;
 
 	if (NonDefaultOnly && Value == DefaultValue)
 	{
