[GLTFExporter] fix: different behavior on Vulkan when reading pixels from an HDR cubemap

commit a5d5e7518efc245cf1302614b1044dbce09073d4
Date: Mon, 28 Feb 2022 09:47:26 +0100

    fix: different behavior on Vulkan when reading pixels from an HDR cubemap
    when using ReadPixels to read floating point values, Vulkan disregards the bLinearToGamma flag in FReadSurfaceDataFlags and never performs any gamma-correction.
    The current fix is to disable automatic gamma-correction when reading pixels and instead make sure that our manual gamma-correction is performed afterwards.
    (cherry picked from commit d9dcf827d8163bbd85077acb4a85ee4d1c17c11f​)

#skipci

diff --git a/Engine/Plugins/Enterprise/GLTFExporter/MigratedFromGit/Source/GLTFExporter/Private/Converters/GLTFTextureUtility.cpp b/Engine/Plugins/Enterprise/GLTFExporter/MigratedFromGit/Source/GLTFExporter/Private/Converters/GLTFTextureUtility.cpp
--- a/Engine/Plugins/Enterprise/GLTFExporter/MigratedFromGit/Source/GLTFExporter/Private/Converters/GLTFTextureUtility.cpp
+++ b/Engine/Plugins/Enterprise/GLTFExporter/MigratedFromGit/Source/GLTFExporter/Private/Converters/GLTFTextureUtility.cpp
@@ -320,7 +320,9 @@ bool FGLTFTextureUtility::ReadPixels(const UTextureRenderTarget2D* InRenderTarge
 
 	if (Encoding == EGLTFJsonHDREncoding::None)
 	{
-		return Resource->ReadPixels(OutPixels);
+		FReadSurfaceDataFlags ReadSurfaceDataFlags(RCM_UNorm, CubeFace_MAX);
+		ReadSurfaceDataFlags.SetLinearToGamma(false);
+		return Resource->ReadPixels(OutPixels, ReadSurfaceDataFlags);
 	}
 
 	TArray<FFloat16Color> HDRPixels;
diff --git a/Engine/Plugins/Enterprise/GLTFExporter/MigratedFromGit/Source/GLTFExporter/Private/Tasks/GLTFTextureTasks.cpp b/Engine/Plugins/Enterprise/GLTFExporter/MigratedFromGit/Source/GLTFExporter/Private/Tasks/GLTFTextureTasks.cpp
--- a/Engine/Plugins/Enterprise/GLTFExporter/MigratedFromGit/Source/GLTFExporter/Private/Tasks/GLTFTextureTasks.cpp
+++ b/Engine/Plugins/Enterprise/GLTFExporter/MigratedFromGit/Source/GLTFExporter/Private/Tasks/GLTFTextureTasks.cpp
@@ -97,6 +97,10 @@ void FGLTFTextureCubeTask::Complete()
 	if (bIsHDR)
 	{
 		JsonTexture->Encoding = Builder.GetTextureHDREncoding();
+		if (JsonTexture->Encoding == EGLTFJsonHDREncoding::None)
+		{
+			bToSRGB = true;
+		}
 	}
 
 	TGLTFSharedArray<FColor> Pixels;
