[GLTFExporter] refactor: add FGLTFFileBuilder and change behaviour for writing external files

commit a5820041ef6078afbda932d3d8d44d33b944f66f
Date: Sun, 13 Mar 2022 21:24:17 +0100

    refactor: add FGLTFFileBuilder and change behaviour for writing external files
    (cherry picked from commit 088b59ce5b3ccbb8758d8480c4b0309d21adf056​)

#skipci

diff --git a/Engine/Plugins/Enterprise/GLTFExporter/MigratedFromGit/Source/GLTFExporter/Private/Builders/GLTFBufferBuilder.cpp b/Engine/Plugins/Enterprise/GLTFExporter/MigratedFromGit/Source/GLTFExporter/Private/Builders/GLTFBufferBuilder.cpp
--- a/Engine/Plugins/Enterprise/GLTFExporter/MigratedFromGit/Source/GLTFExporter/Private/Builders/GLTFBufferBuilder.cpp
+++ b/Engine/Plugins/Enterprise/GLTFExporter/MigratedFromGit/Source/GLTFExporter/Private/Builders/GLTFBufferBuilder.cpp
@@ -5,8 +5,8 @@
 #include "Misc/FileHelper.h"
 #include "Misc/Paths.h"
 
-FGLTFBufferBuilder::FGLTFBufferBuilder(const FString& FilePath, const UGLTFExportOptions* ExportOptions)
-	: FGLTFJsonBuilder(FilePath, ExportOptions)
+FGLTFBufferBuilder::FGLTFBufferBuilder(const FString& FileName, const UGLTFExportOptions* ExportOptions)
+	: FGLTFJsonBuilder(FileName, ExportOptions)
 {
 }
 
@@ -18,44 +18,33 @@ FGLTFBufferBuilder::~FGLTFBufferBuilder()
 	}
 }
 
-bool FGLTFBufferBuilder::InitializeBuffer()
+void FGLTFBufferBuilder::InitializeBuffer()
 {
-	JsonBuffer = FGLTFJsonBuilder::AddBuffer();
-
-	if (bIsGlbFile)
+	if (BufferArchive != nullptr)
 	{
-		BufferArchive = MakeUnique<FGLTFMemoryArchive>();
+		return;
 	}
-	else
-	{
-		const FString ExternalBinaryPath = FPaths::ChangeExtension(FilePath, TEXT(".bin"));
-		JsonBuffer->URI = FPaths::GetCleanFilename(ExternalBinaryPath);
 
-		BufferArchive.Reset(IFileManager::Get().CreateFileWriter(*ExternalBinaryPath));
-		if (BufferArchive == nullptr)
-		{
-			LogError(FString::Printf(TEXT("Failed to write external binary buffer to file: %s"), *ExternalBinaryPath));
-			return false;
-		}
-	}
+	JsonBuffer = AddBuffer();
+	BufferArchive = MakeShared<FGLTFMemoryArchive>();
 
-	return true;
+	if (!bIsGLB)
+	{
+		const FString URI = FPaths::ChangeExtension(FileName, TEXT(".bin"));
+		JsonBuffer->URI = AddExternalFile(URI, BufferArchive);
+	}
 }
 
-const TArray64<uint8>* FGLTFBufferBuilder::GetBufferData() const
+const FGLTFMemoryArchive* FGLTFBufferBuilder::GetBufferData() const
 {
-	return bIsGlbFile ? static_cast<FGLTFMemoryArchive*>(BufferArchive.Get()) : nullptr;
+	return BufferArchive.Get();
 }
 
 FGLTFJsonBufferView* FGLTFBufferBuilder::AddBufferView(const void* RawData, uint64 ByteLength, EGLTFJsonBufferTarget BufferTarget, uint8 DataAlignment)
 {
-	if (BufferArchive == nullptr && !InitializeBuffer())
-	{
-		// TODO: report error
-		return nullptr;
-	}
+	InitializeBuffer();
 
-	uint64 ByteOffset = BufferArchive->Tell();
+	uint64 ByteOffset = BufferArchive->Num();
 
 	// Data offset must be a multiple of the size of the glTF component type (given by ByteAlignment).
 	const int64 Padding = (DataAlignment - (ByteOffset % DataAlignment)) % DataAlignment;
diff --git a/Engine/Plugins/Enterprise/GLTFExporter/MigratedFromGit/Source/GLTFExporter/Private/Builders/GLTFBufferBuilder.h b/Engine/Plugins/Enterprise/GLTFExporter/MigratedFromGit/Source/GLTFExporter/Private/Builders/GLTFBufferBuilder.h
--- a/Engine/Plugins/Enterprise/GLTFExporter/MigratedFromGit/Source/GLTFExporter/Private/Builders/GLTFBufferBuilder.h
+++ b/Engine/Plugins/Enterprise/GLTFExporter/MigratedFromGit/Source/GLTFExporter/Private/Builders/GLTFBufferBuilder.h
@@ -9,10 +9,10 @@ class FGLTFBufferBuilder : public FGLTFJsonBuilder
 {
 protected:
 
-	FGLTFBufferBuilder(const FString& FilePath, const UGLTFExportOptions* ExportOptions);
+	FGLTFBufferBuilder(const FString& FileName, const UGLTFExportOptions* ExportOptions);
 	~FGLTFBufferBuilder();
 
-	const TArray64<uint8>* GetBufferData() const;
+	const FGLTFMemoryArchive* GetBufferData() const;
 
 public:
 
@@ -26,8 +26,8 @@ public:
 
 private:
 
-	bool InitializeBuffer();
+	void InitializeBuffer();
 
 	FGLTFJsonBuffer* JsonBuffer;
-	TUniquePtr<FArchive> BufferArchive;
+	TSharedPtr<FGLTFMemoryArchive> BufferArchive;
 };
diff --git a/Engine/Plugins/Enterprise/GLTFExporter/MigratedFromGit/Source/GLTFExporter/Private/Builders/GLTFBuilder.cpp b/Engine/Plugins/Enterprise/GLTFExporter/MigratedFromGit/Source/GLTFExporter/Private/Builders/GLTFBuilder.cpp
--- a/Engine/Plugins/Enterprise/GLTFExporter/MigratedFromGit/Source/GLTFExporter/Private/Builders/GLTFBuilder.cpp
+++ b/Engine/Plugins/Enterprise/GLTFExporter/MigratedFromGit/Source/GLTFExporter/Private/Builders/GLTFBuilder.cpp
@@ -4,10 +4,9 @@
 #include "Converters/GLTFMeshUtility.h"
 #include "UserData/GLTFMaterialUserData.h"
 
-FGLTFBuilder::FGLTFBuilder(const FString& FilePath, const UGLTFExportOptions* ExportOptions)
-	: bIsGlbFile(FPaths::GetExtension(FilePath).Equals(TEXT("glb"), ESearchCase::IgnoreCase))
-	, FilePath(FilePath)
-	, DirPath(FPaths::GetPath(FilePath))
+FGLTFBuilder::FGLTFBuilder(const FString& FileName, const UGLTFExportOptions* ExportOptions)
+	: FileName(FPaths::GetCleanFilename(FileName))
+	, bIsGLB(FileName.EndsWith(TEXT(".glb"), ESearchCase::IgnoreCase))
 	, ExportOptions(SanitizeExportOptions(ExportOptions))
 	, ExportOptionsGuard(ExportOptions)
 {
diff --git a/Engine/Plugins/Enterprise/GLTFExporter/MigratedFromGit/Source/GLTFExporter/Private/Builders/GLTFBuilder.h b/Engine/Plugins/Enterprise/GLTFExporter/MigratedFromGit/Source/GLTFExporter/Private/Builders/GLTFBuilder.h
--- a/Engine/Plugins/Enterprise/GLTFExporter/MigratedFromGit/Source/GLTFExporter/Private/Builders/GLTFBuilder.h
+++ b/Engine/Plugins/Enterprise/GLTFExporter/MigratedFromGit/Source/GLTFExporter/Private/Builders/GLTFBuilder.h
@@ -10,14 +10,12 @@ class FGLTFBuilder
 {
 public:
 
-	const bool bIsGlbFile;
-	const FString FilePath;
-	const FString DirPath;
+	const FString FileName;
+	const bool bIsGLB;
 
-	// TODO: make ExportOptions private and expose each option via getters to ease overriding settings in future
-	const UGLTFExportOptions* ExportOptions;
+	const UGLTFExportOptions* ExportOptions; // TODO: make ExportOptions private and expose each option via getters to ease overriding settings in future
 
-	FGLTFBuilder(const FString& FilePath, const UGLTFExportOptions* ExportOptions);
+	FGLTFBuilder(const FString& FileName, const UGLTFExportOptions* ExportOptions);
 
 	const UMaterialInterface* ResolveProxy(const UMaterialInterface* Material) const;
 	void ResolveProxies(TArray<const UMaterialInterface*>& Materials) const;
diff --git a/Engine/Plugins/Enterprise/GLTFExporter/MigratedFromGit/Source/GLTFExporter/Private/Builders/GLTFContainerBuilder.cpp b/Engine/Plugins/Enterprise/GLTFExporter/MigratedFromGit/Source/GLTFExporter/Private/Builders/GLTFContainerBuilder.cpp
--- a/Engine/Plugins/Enterprise/GLTFExporter/MigratedFromGit/Source/GLTFExporter/Private/Builders/GLTFContainerBuilder.cpp
+++ b/Engine/Plugins/Enterprise/GLTFExporter/MigratedFromGit/Source/GLTFExporter/Private/Builders/GLTFContainerBuilder.cpp
@@ -3,16 +3,14 @@
 #include "Builders/GLTFContainerBuilder.h"
 #include "Builders/GLTFMemoryArchive.h"
 
-FGLTFContainerBuilder::FGLTFContainerBuilder(const FString& FilePath, const UGLTFExportOptions* ExportOptions, const TSet<AActor*>& SelectedActors)
-	: FGLTFConvertBuilder(FilePath, ExportOptions, SelectedActors)
+FGLTFContainerBuilder::FGLTFContainerBuilder(const FString& FileName, const UGLTFExportOptions* ExportOptions, const TSet<AActor*>& SelectedActors)
+	: FGLTFConvertBuilder(FileName, ExportOptions, SelectedActors)
 {
 }
 
-void FGLTFContainerBuilder::Write(FArchive& Archive, FFeedbackContext* Context)
+void FGLTFContainerBuilder::WriteToArchive(FArchive& Archive)
 {
-	CompleteAllTasks(Context);
-
-	if (bIsGlbFile)
+	if (bIsGLB)
 	{
 		WriteGlb(Archive);
 	}
@@ -34,6 +32,20 @@ void FGLTFContainerBuilder::Write(FArchive& Archive, FFeedbackContext* Context)
 	}
 }
 
+bool FGLTFContainerBuilder::WriteAllFiles(const FString& DirPath, bool bOverwrite)
+{
+	FGLTFMemoryArchive Archive;
+	WriteToArchive(Archive);
+
+	const FString FilePath = FPaths::Combine(DirPath, FileName);
+	if (!SaveToFile(FilePath, Archive, bOverwrite))
+	{
+		return false;
+	}
+
+	return WriteExternalFiles(DirPath, bOverwrite);
+}
+
 void FGLTFContainerBuilder::WriteGlb(FArchive& Archive)
 {
 	FGLTFMemoryArchive JsonData;
diff --git a/Engine/Plugins/Enterprise/GLTFExporter/MigratedFromGit/Source/GLTFExporter/Private/Builders/GLTFContainerBuilder.h b/Engine/Plugins/Enterprise/GLTFExporter/MigratedFromGit/Source/GLTFExporter/Private/Builders/GLTFContainerBuilder.h
--- a/Engine/Plugins/Enterprise/GLTFExporter/MigratedFromGit/Source/GLTFExporter/Private/Builders/GLTFContainerBuilder.h
+++ b/Engine/Plugins/Enterprise/GLTFExporter/MigratedFromGit/Source/GLTFExporter/Private/Builders/GLTFContainerBuilder.h
@@ -8,9 +8,11 @@ class FGLTFContainerBuilder : public FGLTFConvertBuilder
 {
 public:
 
-	FGLTFContainerBuilder(const FString& FilePath, const UGLTFExportOptions* ExportOptions, const TSet<AActor*>& SelectedActors);
+	FGLTFContainerBuilder(const FString& FileName, const UGLTFExportOptions* ExportOptions, const TSet<AActor*>& SelectedActors);
 
-	void Write(FArchive& Archive, FFeedbackContext* Context);
+	void WriteToArchive(FArchive& Archive);
+
+	bool WriteAllFiles(const FString& DirPath, bool bOverwrite = true);
 
 protected:
 
diff --git a/Engine/Plugins/Enterprise/GLTFExporter/MigratedFromGit/Source/GLTFExporter/Private/Builders/GLTFConvertBuilder.cpp b/Engine/Plugins/Enterprise/GLTFExporter/MigratedFromGit/Source/GLTFExporter/Private/Builders/GLTFConvertBuilder.cpp
--- a/Engine/Plugins/Enterprise/GLTFExporter/MigratedFromGit/Source/GLTFExporter/Private/Builders/GLTFConvertBuilder.cpp
+++ b/Engine/Plugins/Enterprise/GLTFExporter/MigratedFromGit/Source/GLTFExporter/Private/Builders/GLTFConvertBuilder.cpp
@@ -7,8 +7,8 @@
 #include "Engine/TextureRenderTarget2D.h"
 #include "Engine/TextureRenderTargetCube.h"
 
-FGLTFConvertBuilder::FGLTFConvertBuilder(const FString& FilePath, const UGLTFExportOptions* ExportOptions, const TSet<AActor*>& SelectedActors)
-	: FGLTFBufferBuilder(FilePath, ExportOptions)
+FGLTFConvertBuilder::FGLTFConvertBuilder(const FString& FileName, const UGLTFExportOptions* ExportOptions, const TSet<AActor*>& SelectedActors)
+	: FGLTFBufferBuilder(FileName, ExportOptions)
 	, SelectedActors(SelectedActors)
 {
 }
diff --git a/Engine/Plugins/Enterprise/GLTFExporter/MigratedFromGit/Source/GLTFExporter/Private/Builders/GLTFConvertBuilder.h b/Engine/Plugins/Enterprise/GLTFExporter/MigratedFromGit/Source/GLTFExporter/Private/Builders/GLTFConvertBuilder.h
--- a/Engine/Plugins/Enterprise/GLTFExporter/MigratedFromGit/Source/GLTFExporter/Private/Builders/GLTFConvertBuilder.h
+++ b/Engine/Plugins/Enterprise/GLTFExporter/MigratedFromGit/Source/GLTFExporter/Private/Builders/GLTFConvertBuilder.h
@@ -29,7 +29,7 @@ public:
 
 	const TSet<AActor*> SelectedActors;
 
-	FGLTFConvertBuilder(const FString& FilePath, const UGLTFExportOptions* ExportOptions, const TSet<AActor*>& SelectedActors = {});
+	FGLTFConvertBuilder(const FString& FileName, const UGLTFExportOptions* ExportOptions, const TSet<AActor*>& SelectedActors = {});
 
 	bool IsSelectedActor(const AActor* Object) const;
 	bool IsRootActor(const AActor* Actor) const;
diff --git a/Engine/Plugins/Enterprise/GLTFExporter/MigratedFromGit/Source/GLTFExporter/Private/Builders/GLTFFileBuilder.cpp b/Engine/Plugins/Enterprise/GLTFExporter/MigratedFromGit/Source/GLTFExporter/Private/Builders/GLTFFileBuilder.cpp
new file mode 100644
--- /dev/null
+++ b/Engine/Plugins/Enterprise/GLTFExporter/MigratedFromGit/Source/GLTFExporter/Private/Builders/GLTFFileBuilder.cpp
@@ -0,0 +1,72 @@
+// Copyright Epic Games, Inc. All Rights Reserved.
+
+#include "Builders/GLTFFileBuilder.h"
+#include "Misc/FileHelper.h"
+
+FGLTFFileBuilder::FGLTFFileBuilder(const FString& FileName, const UGLTFExportOptions* ExportOptions)
+	: FGLTFTaskBuilder(FileName, ExportOptions)
+{
+}
+
+FString FGLTFFileBuilder::AddExternalFile(const FString& URI, const TSharedPtr<FGLTFMemoryArchive>& Archive)
+{
+	const FString UnqiueURI = GetUniqueURI(URI);
+	ExternalFiles.Add(UnqiueURI, Archive);
+	return UnqiueURI;
+}
+
+const TMap<FString, TSharedPtr<FGLTFMemoryArchive>>& FGLTFFileBuilder::GetExternalFiles() const
+{
+	return ExternalFiles;
+}
+
+bool FGLTFFileBuilder::WriteExternalFiles(const FString& DirPath, bool bOverwrite)
+{
+	for (const TPair<FString, TSharedPtr<FGLTFMemoryArchive>>& ExternalFile : ExternalFiles)
+	{
+		const FString FilePath = FPaths::Combine(DirPath, *ExternalFile.Key);
+		const TArray64<uint8>& FileData = *ExternalFile.Value;
+
+		if (!SaveToFile(FilePath, FileData, bOverwrite))
+		{
+			return false;
+		}
+	}
+
+	return true;
+}
+
+FString FGLTFFileBuilder::GetUniqueURI(const FString& URI) const
+{
+	if (!ExternalFiles.Contains(URI))
+	{
+		return URI;
+	}
+
+	const FString BaseFilename = FPaths::GetBaseFilename(URI);
+	const FString FileExtension = FPaths::GetExtension(URI, true);
+	FString UnqiueURI;
+
+	int32 Suffix = 1;
+	do
+	{
+		UnqiueURI = BaseFilename + TEXT('_') + FString::FromInt(Suffix) + FileExtension;
+		Suffix++;
+	}
+	while (ExternalFiles.Contains(UnqiueURI));
+
+	return UnqiueURI;
+}
+
+bool FGLTFFileBuilder::SaveToFile(const FString& FilePath, const TArray64<uint8>& FileData, bool bOverwrite)
+{
+	const uint32 WriteFlags = bOverwrite ? FILEWRITE_None : FILEWRITE_NoReplaceExisting;
+
+	if (!FFileHelper::SaveArrayToFile(FileData, *FilePath, &IFileManager::Get(), WriteFlags))
+	{
+		LogError(FString::Printf(TEXT("Failed to save file: %s"), *FilePath));
+		return false;
+	}
+
+	return true;
+}
diff --git a/Engine/Plugins/Enterprise/GLTFExporter/MigratedFromGit/Source/GLTFExporter/Private/Builders/GLTFFileBuilder.h b/Engine/Plugins/Enterprise/GLTFExporter/MigratedFromGit/Source/GLTFExporter/Private/Builders/GLTFFileBuilder.h
new file mode 100644
--- /dev/null
+++ b/Engine/Plugins/Enterprise/GLTFExporter/MigratedFromGit/Source/GLTFExporter/Private/Builders/GLTFFileBuilder.h
@@ -0,0 +1,29 @@
+// Copyright Epic Games, Inc. All Rights Reserved.
+
+#pragma once
+
+#include "Builders/GLTFTaskBuilder.h"
+#include "Builders/GLTFMemoryArchive.h"
+
+class FGLTFFileBuilder : public FGLTFTaskBuilder
+{
+public:
+
+	FGLTFFileBuilder(const FString& FileName, const UGLTFExportOptions* ExportOptions);
+
+	FString AddExternalFile(const FString& URI, const TSharedPtr<FGLTFMemoryArchive>& Archive = MakeShared<FGLTFMemoryArchive>());
+
+	const TMap<FString, TSharedPtr<FGLTFMemoryArchive>>& GetExternalFiles() const;
+
+	bool WriteExternalFiles(const FString& DirPath, bool bOverwrite = true);
+
+private:
+
+	TMap<FString, TSharedPtr<FGLTFMemoryArchive>> ExternalFiles;
+
+	FString GetUniqueURI(const FString& URI) const;
+
+protected:
+
+	bool SaveToFile(const FString& FilePath, const TArray64<uint8>& FileData, bool bOverwrite);
+};
diff --git a/Engine/Plugins/Enterprise/GLTFExporter/MigratedFromGit/Source/GLTFExporter/Private/Builders/GLTFJsonBuilder.cpp b/Engine/Plugins/Enterprise/GLTFExporter/MigratedFromGit/Source/GLTFExporter/Private/Builders/GLTFJsonBuilder.cpp
--- a/Engine/Plugins/Enterprise/GLTFExporter/MigratedFromGit/Source/GLTFExporter/Private/Builders/GLTFJsonBuilder.cpp
+++ b/Engine/Plugins/Enterprise/GLTFExporter/MigratedFromGit/Source/GLTFExporter/Private/Builders/GLTFJsonBuilder.cpp
@@ -5,8 +5,8 @@
 #include "Interfaces/IPluginManager.h"
 #include "Runtime/Launch/Resources/Version.h"
 
-FGLTFJsonBuilder::FGLTFJsonBuilder(const FString& FilePath, const UGLTFExportOptions* ExportOptions)
-	: FGLTFTaskBuilder(FilePath, ExportOptions)
+FGLTFJsonBuilder::FGLTFJsonBuilder(const FString& FileName, const UGLTFExportOptions* ExportOptions)
+	: FGLTFFileBuilder(FileName, ExportOptions)
 	, DefaultScene(JsonRoot.DefaultScene)
 {
 	JsonRoot.Asset.Generator = GetGeneratorString();
@@ -14,7 +14,7 @@ FGLTFJsonBuilder::FGLTFJsonBuilder(const FString& FilePath, const UGLTFExportOpt
 
 void FGLTFJsonBuilder::WriteJson(FArchive& Archive)
 {
-	JsonRoot.WriteJson(Archive, !bIsGlbFile, ExportOptions->bSkipNearDefaultValues ? KINDA_SMALL_NUMBER : 0);
+	JsonRoot.WriteJson(Archive, !bIsGLB, ExportOptions->bSkipNearDefaultValues ? KINDA_SMALL_NUMBER : 0);
 }
 
 TSet<EGLTFJsonExtension> FGLTFJsonBuilder::GetCustomExtensionsUsed() const
diff --git a/Engine/Plugins/Enterprise/GLTFExporter/MigratedFromGit/Source/GLTFExporter/Private/Builders/GLTFJsonBuilder.h b/Engine/Plugins/Enterprise/GLTFExporter/MigratedFromGit/Source/GLTFExporter/Private/Builders/GLTFJsonBuilder.h
--- a/Engine/Plugins/Enterprise/GLTFExporter/MigratedFromGit/Source/GLTFExporter/Private/Builders/GLTFJsonBuilder.h
+++ b/Engine/Plugins/Enterprise/GLTFExporter/MigratedFromGit/Source/GLTFExporter/Private/Builders/GLTFJsonBuilder.h
@@ -3,15 +3,15 @@
 #pragma once
 
 #include "Json/GLTFJsonRoot.h"
-#include "Builders/GLTFTaskBuilder.h"
+#include "Builders/GLTFFileBuilder.h"
 
-class FGLTFJsonBuilder : public FGLTFTaskBuilder
+class FGLTFJsonBuilder : public FGLTFFileBuilder
 {
 public:
 
 	FGLTFJsonScene*& DefaultScene;
 
-	FGLTFJsonBuilder(const FString& FilePath, const UGLTFExportOptions* ExportOptions);
+	FGLTFJsonBuilder(const FString& FileName, const UGLTFExportOptions* ExportOptions);
 
 	void WriteJson(FArchive& Archive);
 
diff --git a/Engine/Plugins/Enterprise/GLTFExporter/MigratedFromGit/Source/GLTFExporter/Private/Builders/GLTFLogBuilder.cpp b/Engine/Plugins/Enterprise/GLTFExporter/MigratedFromGit/Source/GLTFExporter/Private/Builders/GLTFLogBuilder.cpp
--- a/Engine/Plugins/Enterprise/GLTFExporter/MigratedFromGit/Source/GLTFExporter/Private/Builders/GLTFLogBuilder.cpp
+++ b/Engine/Plugins/Enterprise/GLTFExporter/MigratedFromGit/Source/GLTFExporter/Private/Builders/GLTFLogBuilder.cpp
@@ -9,8 +9,8 @@
 #include "IMessageLogListing.h"
 #endif
 
-FGLTFLogBuilder::FGLTFLogBuilder(const FString& FilePath, const UGLTFExportOptions* ExportOptions)
-	: FGLTFBuilder(FilePath, ExportOptions)
+FGLTFLogBuilder::FGLTFLogBuilder(const FString& FileName, const UGLTFExportOptions* ExportOptions)
+	: FGLTFBuilder(FileName, ExportOptions)
 {
 #if WITH_EDITOR
 	if (!FApp::IsUnattended())
diff --git a/Engine/Plugins/Enterprise/GLTFExporter/MigratedFromGit/Source/GLTFExporter/Private/Builders/GLTFLogBuilder.h b/Engine/Plugins/Enterprise/GLTFExporter/MigratedFromGit/Source/GLTFExporter/Private/Builders/GLTFLogBuilder.h
--- a/Engine/Plugins/Enterprise/GLTFExporter/MigratedFromGit/Source/GLTFExporter/Private/Builders/GLTFLogBuilder.h
+++ b/Engine/Plugins/Enterprise/GLTFExporter/MigratedFromGit/Source/GLTFExporter/Private/Builders/GLTFLogBuilder.h
@@ -12,7 +12,7 @@ class FGLTFLogBuilder : public FGLTFBuilder
 {
 public:
 
-	FGLTFLogBuilder(const FString& FilePath, const UGLTFExportOptions* ExportOptions);
+	FGLTFLogBuilder(const FString& FileName, const UGLTFExportOptions* ExportOptions);
 
 	void LogSuggestion(const FString& Message);
 
diff --git a/Engine/Plugins/Enterprise/GLTFExporter/MigratedFromGit/Source/GLTFExporter/Private/Builders/GLTFTaskBuilder.cpp b/Engine/Plugins/Enterprise/GLTFExporter/MigratedFromGit/Source/GLTFExporter/Private/Builders/GLTFTaskBuilder.cpp
--- a/Engine/Plugins/Enterprise/GLTFExporter/MigratedFromGit/Source/GLTFExporter/Private/Builders/GLTFTaskBuilder.cpp
+++ b/Engine/Plugins/Enterprise/GLTFExporter/MigratedFromGit/Source/GLTFExporter/Private/Builders/GLTFTaskBuilder.cpp
@@ -4,8 +4,8 @@
 #include "Misc/FeedbackContext.h"
 #include "Misc/ScopedSlowTask.h"
 
-FGLTFTaskBuilder::FGLTFTaskBuilder(const FString& FilePath, const UGLTFExportOptions* ExportOptions)
-	: FGLTFLogBuilder(FilePath, ExportOptions)
+FGLTFTaskBuilder::FGLTFTaskBuilder(const FString& FileName, const UGLTFExportOptions* ExportOptions)
+	: FGLTFLogBuilder(FileName, ExportOptions)
 	, PriorityIndexLock(INDEX_NONE)
 {
 }
diff --git a/Engine/Plugins/Enterprise/GLTFExporter/MigratedFromGit/Source/GLTFExporter/Private/Builders/GLTFTaskBuilder.h b/Engine/Plugins/Enterprise/GLTFExporter/MigratedFromGit/Source/GLTFExporter/Private/Builders/GLTFTaskBuilder.h
--- a/Engine/Plugins/Enterprise/GLTFExporter/MigratedFromGit/Source/GLTFExporter/Private/Builders/GLTFTaskBuilder.h
+++ b/Engine/Plugins/Enterprise/GLTFExporter/MigratedFromGit/Source/GLTFExporter/Private/Builders/GLTFTaskBuilder.h
@@ -9,7 +9,7 @@ class FGLTFTaskBuilder : public FGLTFLogBuilder
 {
 public:
 
-	FGLTFTaskBuilder(const FString& FilePath, const UGLTFExportOptions* ExportOptions);
+	FGLTFTaskBuilder(const FString& FileName, const UGLTFExportOptions* ExportOptions);
 
 	template <typename TaskType, typename... TaskArgTypes, typename = typename TEnableIf<TIsDerivedFrom<TaskType, FGLTFTask>::Value>::Type>
 	bool SetupTask(TaskArgTypes&&... Args)
diff --git a/Engine/Plugins/Enterprise/GLTFExporter/MigratedFromGit/Source/GLTFExporter/Private/Converters/GLTFImageConverters.cpp b/Engine/Plugins/Enterprise/GLTFExporter/MigratedFromGit/Source/GLTFExporter/Private/Converters/GLTFImageConverters.cpp
--- a/Engine/Plugins/Enterprise/GLTFExporter/MigratedFromGit/Source/GLTFExporter/Private/Converters/GLTFImageConverters.cpp
+++ b/Engine/Plugins/Enterprise/GLTFExporter/MigratedFromGit/Source/GLTFExporter/Private/Converters/GLTFImageConverters.cpp
@@ -7,7 +7,7 @@
 
 FGLTFJsonImage* FGLTFImageConverter::Convert(TGLTFSuperfluous<FString> Name, EGLTFTextureType Type, bool bIgnoreAlpha, FIntPoint Size, TGLTFSharedArray<FColor> Pixels)
 {
-	TArray64<uint8> CompressedData;
+	const TSharedPtr<FGLTFMemoryArchive> CompressedData = MakeShared<FGLTFMemoryArchive>();
 
 	const EGLTFJsonMimeType MimeType = GetMimeType(Pixels->GetData(), Size, bIgnoreAlpha, Type);
 	switch (MimeType)
@@ -16,11 +16,11 @@ FGLTFJsonImage* FGLTFImageConverter::Convert(TGLTFSuperfluous<FString> Name, EGL
 			return nullptr;
 
 		case EGLTFJsonMimeType::PNG:
-			FGLTFImageUtility::CompressToPNG(Pixels->GetData(), Size, CompressedData);
+			FGLTFImageUtility::CompressToPNG(Pixels->GetData(), Size, *CompressedData);
 			break;
 
 		case EGLTFJsonMimeType::JPEG:
-			FGLTFImageUtility::CompressToJPEG(Pixels->GetData(), Size, Builder.ExportOptions->TextureImageQuality, CompressedData);
+			FGLTFImageUtility::CompressToJPEG(Pixels->GetData(), Size, Builder.ExportOptions->TextureImageQuality, *CompressedData);
 			break;
 
 		default:
@@ -30,15 +30,16 @@ FGLTFJsonImage* FGLTFImageConverter::Convert(TGLTFSuperfluous<FString> Name, EGL
 
 	FGLTFJsonImage* JsonImage = Builder.AddImage();
 
-	if (Builder.bIsGlbFile)
+	if (Builder.bIsGLB)
 	{
 		JsonImage->Name = Name;
 		JsonImage->MimeType = MimeType;
-		JsonImage->BufferView = Builder.AddBufferView(CompressedData.GetData(), CompressedData.Num());
+		JsonImage->BufferView = Builder.AddBufferView(CompressedData->GetData(), CompressedData->Num());
 	}
 	else
 	{
-		JsonImage->Uri = SaveToFile(CompressedData.GetData(), CompressedData.Num(), MimeType, Name);
+		const TCHAR* Extension = FGLTFImageUtility::GetFileExtension(MimeType);
+		JsonImage->URI = Builder.AddExternalFile(Name + Extension, CompressedData);
 	}
 
 	return JsonImage;
@@ -65,21 +66,3 @@ EGLTFJsonMimeType FGLTFImageConverter::GetMimeType(const FColor* Pixels, FIntPoi
 		return EGLTFJsonMimeType::None;
 	}
 }
-
-FString FGLTFImageConverter::SaveToFile(const void* CompressedData, int64 CompressedByteLength, EGLTFJsonMimeType MimeType, const FString& Name)
-{
-	const TCHAR* Extension = FGLTFImageUtility::GetFileExtension(MimeType);
-	const FString ImageUri = FGLTFImageUtility::GetUniqueFilename(Name, Extension, UniqueImageUris);
-
-	const TArrayView<const uint8> ImageData(static_cast<const uint8*>(CompressedData), CompressedByteLength);
-	const FString ImagePath = FPaths::Combine(Builder.DirPath, ImageUri);
-
-	if (!FFileHelper::SaveArrayToFile(ImageData, *ImagePath))
-	{
-		Builder.LogError(FString::Printf(TEXT("Failed to save image %s to file: %s"), *Name, *ImagePath));
-		return TEXT("");
-	}
-
-	UniqueImageUris.Add(ImageUri);
-	return ImageUri;
-}
diff --git a/Engine/Plugins/Enterprise/GLTFExporter/MigratedFromGit/Source/GLTFExporter/Private/Converters/GLTFImageConverters.h b/Engine/Plugins/Enterprise/GLTFExporter/MigratedFromGit/Source/GLTFExporter/Private/Converters/GLTFImageConverters.h
--- a/Engine/Plugins/Enterprise/GLTFExporter/MigratedFromGit/Source/GLTFExporter/Private/Converters/GLTFImageConverters.h
+++ b/Engine/Plugins/Enterprise/GLTFExporter/MigratedFromGit/Source/GLTFExporter/Private/Converters/GLTFImageConverters.h
@@ -18,8 +18,4 @@ class FGLTFImageConverter final : public FGLTFBuilderContext, public IGLTFImageC
 	virtual FGLTFJsonImage* Convert(TGLTFSuperfluous<FString> Name, EGLTFTextureType Type, bool bIgnoreAlpha, FIntPoint Size, TGLTFSharedArray<FColor> Pixels) override;
 
 	EGLTFJsonMimeType GetMimeType(const FColor* Pixels, FIntPoint Size, bool bIgnoreAlpha, EGLTFTextureType Type) const;
-
-	FString SaveToFile(const void* CompressedData, int64 CompressedByteLength, EGLTFJsonMimeType MimeType, const FString& Name);
-
-	TSet<FString> UniqueImageUris;
 };
diff --git a/Engine/Plugins/Enterprise/GLTFExporter/MigratedFromGit/Source/GLTFExporter/Private/Exporters/GLTFExporter.cpp b/Engine/Plugins/Enterprise/GLTFExporter/MigratedFromGit/Source/GLTFExporter/Private/Exporters/GLTFExporter.cpp
--- a/Engine/Plugins/Enterprise/GLTFExporter/MigratedFromGit/Source/GLTFExporter/Private/Exporters/GLTFExporter.cpp
+++ b/Engine/Plugins/Enterprise/GLTFExporter/MigratedFromGit/Source/GLTFExporter/Private/Exporters/GLTFExporter.cpp
@@ -41,13 +41,19 @@ bool UGLTFExporter::ExportBinary(UObject* Object, const TCHAR* Type, FArchive& A
 
 	// TODO: add support for UAssetExportTask::IgnoreObjectList?
 
-	FGLTFContainerBuilder Builder(GetFilePath(), Options, SelectedActors);
+	const FString FilePath = GetFilePath();
+	const FString FileName = FPaths::GetCleanFilename(FilePath);
+	const FString DirPath = FPaths::GetPath(FilePath);
+
+	FGLTFContainerBuilder Builder(FileName, Options, SelectedActors);
 	Builder.ClearLog();
 
-	const bool bSuccess = AddObject(Builder, Object);
+	bool bSuccess = AddObject(Builder, Object);
 	if (bSuccess)
 	{
-		Builder.Write(Archive, Warn);
+		Builder.CompleteAllTasks(Warn);
+		Builder.WriteToArchive(Archive);
+		bSuccess = Builder.WriteExternalFiles(DirPath);
 	}
 
 	// TODO: should we copy messages to UAssetExportTask::Errors?
@@ -60,7 +66,7 @@ bool UGLTFExporter::ExportBinary(UObject* Object, const TCHAR* Type, FArchive& A
 	return bSuccess;
 }
 
-bool UGLTFExporter::ExportToGLTF(UObject* Object, const FString& Filename, const UGLTFExportOptions* Options, const TSet<AActor*>& SelectedActors, FGLTFExportMessages& OutMessages)
+bool UGLTFExporter::ExportToGLTF(UObject* Object, const FString& FilePath, const UGLTFExportOptions* Options, const TSet<AActor*>& SelectedActors, FGLTFExportMessages& OutMessages)
 {
 	if (Object == nullptr)
 	{
@@ -75,8 +81,11 @@ bool UGLTFExporter::ExportToGLTF(UObject* Object, const FString& Filename, const
 		return false;
 	}
 
-	FGLTFContainerBuilder Builder(Filename, Options, SelectedActors);
-	const bool bSuccess = Exporter->AddObject(Builder, Object);
+	const FString FileName = FPaths::GetCleanFilename(FilePath);
+	const FString DirPath = FPaths::GetPath(FilePath);
+
+	FGLTFContainerBuilder Builder(FileName, Options, SelectedActors);
+	bool bSuccess = Exporter->AddObject(Builder, Object);
 
 	OutMessages.Suggestions = Builder.GetLoggedSuggestions();
 	OutMessages.Warnings = Builder.GetLoggedWarnings();
@@ -84,24 +93,17 @@ bool UGLTFExporter::ExportToGLTF(UObject* Object, const FString& Filename, const
 
 	if (bSuccess)
 	{
-		FArchive* Archive = IFileManager::Get().CreateFileWriter(*Filename);
-		if (Archive == nullptr)
-		{
-			OutMessages.Errors.Add(FString::Printf(TEXT("Couldn't save file: %s"), *Filename));
-			return false;
-		}
-
-		Builder.Write(*Archive, nullptr);
-		Archive->Close();
+		Builder.CompleteAllTasks(nullptr);
+		bSuccess = Builder.WriteAllFiles(DirPath);
 	}
 
 	return bSuccess;
 }
 
-bool UGLTFExporter::ExportToGLTF(UObject* Object, const FString& Filename, const UGLTFExportOptions* Options, const TSet<AActor*>& SelectedActors)
+bool UGLTFExporter::ExportToGLTF(UObject* Object, const FString& FilePath, const UGLTFExportOptions* Options, const TSet<AActor*>& SelectedActors)
 {
 	FGLTFExportMessages Messages;
-	return ExportToGLTF(Object, Filename, Options, SelectedActors, Messages);
+	return ExportToGLTF(Object, FilePath, Options, SelectedActors, Messages);
 }
 
 bool UGLTFExporter::AddObject(FGLTFContainerBuilder& Builder, const UObject* Object)
diff --git a/Engine/Plugins/Enterprise/GLTFExporter/MigratedFromGit/Source/GLTFExporter/Private/Json/GLTFJsonImage.cpp b/Engine/Plugins/Enterprise/GLTFExporter/MigratedFromGit/Source/GLTFExporter/Private/Json/GLTFJsonImage.cpp
--- a/Engine/Plugins/Enterprise/GLTFExporter/MigratedFromGit/Source/GLTFExporter/Private/Json/GLTFJsonImage.cpp
+++ b/Engine/Plugins/Enterprise/GLTFExporter/MigratedFromGit/Source/GLTFExporter/Private/Json/GLTFJsonImage.cpp
@@ -10,9 +10,9 @@ void FGLTFJsonImage::WriteObject(IGLTFJsonWriter& Writer) const
 		Writer.Write(TEXT("name"), Name);
 	}
 
-	if (!Uri.IsEmpty())
+	if (!URI.IsEmpty())
 	{
-		Writer.Write(TEXT("uri"), Uri);
+		Writer.Write(TEXT("uri"), URI);
 	}
 
 	if (MimeType != EGLTFJsonMimeType::None)
diff --git a/Engine/Plugins/Enterprise/GLTFExporter/MigratedFromGit/Source/GLTFExporter/Public/Exporters/GLTFExporter.h b/Engine/Plugins/Enterprise/GLTFExporter/MigratedFromGit/Source/GLTFExporter/Public/Exporters/GLTFExporter.h
--- a/Engine/Plugins/Enterprise/GLTFExporter/MigratedFromGit/Source/GLTFExporter/Public/Exporters/GLTFExporter.h
+++ b/Engine/Plugins/Enterprise/GLTFExporter/MigratedFromGit/Source/GLTFExporter/Public/Exporters/GLTFExporter.h
@@ -36,8 +36,8 @@ public:
 	virtual bool ExportBinary(UObject* Object, const TCHAR* Type, FArchive& Archive, FFeedbackContext* Warn, int32 FileIndex, uint32 PortFlags) override;
 
 	UFUNCTION(BlueprintCallable, Category = "Miscellaneous")
-	static bool ExportToGLTF(UObject* Object, const FString& Filename, const UGLTFExportOptions* Options, const TSet<AActor*>& SelectedActors, FGLTFExportMessages& OutMessages);
-	static bool ExportToGLTF(UObject* Object, const FString& Filename, const UGLTFExportOptions* Options = nullptr, const TSet<AActor*>& SelectedActors = {});
+	static bool ExportToGLTF(UObject* Object, const FString& FilePath, const UGLTFExportOptions* Options, const TSet<AActor*>& SelectedActors, FGLTFExportMessages& OutMessages);
+	static bool ExportToGLTF(UObject* Object, const FString& FilePath, const UGLTFExportOptions* Options = nullptr, const TSet<AActor*>& SelectedActors = {});
 
 protected:
 
diff --git a/Engine/Plugins/Enterprise/GLTFExporter/MigratedFromGit/Source/GLTFExporter/Public/Json/GLTFJsonImage.h b/Engine/Plugins/Enterprise/GLTFExporter/MigratedFromGit/Source/GLTFExporter/Public/Json/GLTFJsonImage.h
--- a/Engine/Plugins/Enterprise/GLTFExporter/MigratedFromGit/Source/GLTFExporter/Public/Json/GLTFJsonImage.h
+++ b/Engine/Plugins/Enterprise/GLTFExporter/MigratedFromGit/Source/GLTFExporter/Public/Json/GLTFJsonImage.h
@@ -7,7 +7,7 @@
 struct GLTFEXPORTER_API FGLTFJsonImage : IGLTFJsonIndexedObject
 {
 	FString Name;
-	FString Uri;
+	FString URI;
 
 	EGLTFJsonMimeType MimeType;
 
