[GLTFExporter] refactor: rename FGLTFTask to FGLTFDelayedTask and rename some function members in FGLTFTaskBuilder

commit f61d2a55c99c0917f02930b92928d4ef2a409e20
Date: Sun, 13 Mar 2022 21:46:26 +0100

    refactor: rename FGLTFTask to FGLTFDelayedTask and rename some function members in FGLTFTaskBuilder
    (cherry picked from commit f614df401b955bc53bed8532f375e394cf0d9ec3​)

#skipci

diff --git a/Engine/Plugins/Enterprise/GLTFExporter/MigratedFromGit/Source/GLTFExporter/Private/Builders/GLTFTaskBuilder.cpp b/Engine/Plugins/Enterprise/GLTFExporter/MigratedFromGit/Source/GLTFExporter/Private/Builders/GLTFTaskBuilder.cpp
--- a/Engine/Plugins/Enterprise/GLTFExporter/MigratedFromGit/Source/GLTFExporter/Private/Builders/GLTFTaskBuilder.cpp
+++ b/Engine/Plugins/Enterprise/GLTFExporter/MigratedFromGit/Source/GLTFExporter/Private/Builders/GLTFTaskBuilder.cpp
@@ -10,7 +10,7 @@ FGLTFTaskBuilder::FGLTFTaskBuilder(const FString& FileName, const UGLTFExportOpt
 {
 }
 
-void FGLTFTaskBuilder::CompleteAllTasks(FFeedbackContext* Context)
+void FGLTFTaskBuilder::ProcessSlowTasks(FFeedbackContext* Context)
 {
 	const int32 PriorityCount = static_cast<int32>(EGLTFTaskPriority::MAX);
 	for (int32 PriorityIndex = 0; PriorityIndex < PriorityCount; PriorityIndex++)
@@ -18,7 +18,7 @@ void FGLTFTaskBuilder::CompleteAllTasks(FFeedbackContext* Context)
 		PriorityIndexLock = PriorityIndex;
 		const EGLTFTaskPriority Priority = static_cast<EGLTFTaskPriority>(PriorityIndex);
 
-		TArray<TUniquePtr<FGLTFTask>>* Tasks = TasksByPriority.Find(Priority);
+		TArray<TUniquePtr<FGLTFDelayedTask>>* Tasks = TasksByPriority.Find(Priority);
 		if (Tasks == nullptr)
 		{
 			continue;
@@ -30,20 +30,20 @@ void FGLTFTaskBuilder::CompleteAllTasks(FFeedbackContext* Context)
 			FScopedSlowTask Progress(Tasks->Num(), FText::Format(MessageFormat, FText()), true, *Context);
 			Progress.MakeDialog();
 
-			for (TUniquePtr<FGLTFTask>& Task : *Tasks)
+			for (TUniquePtr<FGLTFDelayedTask>& Task : *Tasks)
 			{
 				const FText Name = FText::FromString(Task->GetName());
 				const FText Message = FText::Format(MessageFormat, Name);
 				Progress.EnterProgressFrame(1, Message);
 
-				Task->Complete();
+				Task->Process();
 			}
 		}
 		else
 		{
-			for (TUniquePtr<FGLTFTask>& Task : *Tasks)
+			for (TUniquePtr<FGLTFDelayedTask>& Task : *Tasks)
 			{
-				Task->Complete();
+				Task->Process();
 			}
 		}
 	}
@@ -52,7 +52,7 @@ void FGLTFTaskBuilder::CompleteAllTasks(FFeedbackContext* Context)
 	PriorityIndexLock = INDEX_NONE;
 }
 
-bool FGLTFTaskBuilder::SetupTask(TUniquePtr<FGLTFTask> Task)
+bool FGLTFTaskBuilder::ScheduleSlowTask(TUniquePtr<FGLTFDelayedTask> Task)
 {
 	const EGLTFTaskPriority Priority = Task->Priority;
 	if (static_cast<int32>(Priority) >= static_cast<int32>(EGLTFTaskPriority::MAX))
diff --git a/Engine/Plugins/Enterprise/GLTFExporter/MigratedFromGit/Source/GLTFExporter/Private/Builders/GLTFTaskBuilder.h b/Engine/Plugins/Enterprise/GLTFExporter/MigratedFromGit/Source/GLTFExporter/Private/Builders/GLTFTaskBuilder.h
--- a/Engine/Plugins/Enterprise/GLTFExporter/MigratedFromGit/Source/GLTFExporter/Private/Builders/GLTFTaskBuilder.h
+++ b/Engine/Plugins/Enterprise/GLTFExporter/MigratedFromGit/Source/GLTFExporter/Private/Builders/GLTFTaskBuilder.h
@@ -3,7 +3,7 @@
 #pragma once
 
 #include "Builders/GLTFLogBuilder.h"
-#include "Tasks/GLTFTask.h"
+#include "Tasks/GLTFDelayedTask.h"
 
 class FGLTFTaskBuilder : public FGLTFLogBuilder
 {
@@ -11,26 +11,26 @@ public:
 
 	FGLTFTaskBuilder(const FString& FileName, const UGLTFExportOptions* ExportOptions);
 
-	template <typename TaskType, typename... TaskArgTypes, typename = typename TEnableIf<TIsDerivedFrom<TaskType, FGLTFTask>::Value>::Type>
-	bool SetupTask(TaskArgTypes&&... Args)
+	template <typename TaskType, typename... TaskArgTypes, typename = typename TEnableIf<TIsDerivedFrom<TaskType, FGLTFDelayedTask>::Value>::Type>
+	bool ScheduleSlowTask(TaskArgTypes&&... Args)
 	{
-		return SetupTask(MakeUnique<TaskType>(Forward<TaskArgTypes>(Args)...));
+		return ScheduleSlowTask(MakeUnique<TaskType>(Forward<TaskArgTypes>(Args)...));
 	}
 
-	template <typename TaskType, typename = typename TEnableIf<TIsDerivedFrom<TaskType, FGLTFTask>::Value>::Type>
-	bool SetupTask(TUniquePtr<TaskType> Task)
+	template <typename TaskType, typename = typename TEnableIf<TIsDerivedFrom<TaskType, FGLTFDelayedTask>::Value>::Type>
+	bool ScheduleSlowTask(TUniquePtr<TaskType> Task)
 	{
-		return SetupTask(TUniquePtr<FGLTFTask>(Task.Release()));
+		return ScheduleSlowTask(TUniquePtr<FGLTFDelayedTask>(Task.Release()));
 	}
 
-	bool SetupTask(TUniquePtr<FGLTFTask> Task);
+	bool ScheduleSlowTask(TUniquePtr<FGLTFDelayedTask> Task);
 
-	void CompleteAllTasks(FFeedbackContext* Context = GWarn);
+	void ProcessSlowTasks(FFeedbackContext* Context = GWarn);
 
 private:
 
 	static FText GetPriorityMessageFormat(EGLTFTaskPriority Priority);
 
 	int32 PriorityIndexLock;
-	TMap<EGLTFTaskPriority, TArray<TUniquePtr<FGLTFTask>>> TasksByPriority;
+	TMap<EGLTFTaskPriority, TArray<TUniquePtr<FGLTFDelayedTask>>> TasksByPriority;
 };
diff --git a/Engine/Plugins/Enterprise/GLTFExporter/MigratedFromGit/Source/GLTFExporter/Private/Converters/GLTFAnimationConverters.cpp b/Engine/Plugins/Enterprise/GLTFExporter/MigratedFromGit/Source/GLTFExporter/Private/Converters/GLTFAnimationConverters.cpp
--- a/Engine/Plugins/Enterprise/GLTFExporter/MigratedFromGit/Source/GLTFExporter/Private/Converters/GLTFAnimationConverters.cpp
+++ b/Engine/Plugins/Enterprise/GLTFExporter/MigratedFromGit/Source/GLTFExporter/Private/Converters/GLTFAnimationConverters.cpp
@@ -2,7 +2,7 @@
 
 #include "Converters/GLTFAnimationConverters.h"
 #include "Builders/GLTFConvertBuilder.h"
-#include "Tasks/GLTFAnimationTasks.h"
+#include "Tasks/GLTFDelayedAnimationTasks.h"
 #include "Animation/AnimSequence.h"
 
 FGLTFJsonAnimation* FGLTFAnimationConverter::Convert(FGLTFJsonNode* RootNode, const USkeletalMesh* SkeletalMesh, const UAnimSequence* AnimSequence)
@@ -33,7 +33,7 @@ FGLTFJsonAnimation* FGLTFAnimationConverter::Convert(FGLTFJsonNode* RootNode, co
 	}
 
 	FGLTFJsonAnimation* JsonAnimation = Builder.AddAnimation();
-	Builder.SetupTask<FGLTFAnimSequenceTask>(Builder, RootNode, SkeletalMesh, AnimSequence, JsonAnimation);
+	Builder.ScheduleSlowTask<FGLTFDelayedAnimSequenceTask>(Builder, RootNode, SkeletalMesh, AnimSequence, JsonAnimation);
 	return JsonAnimation;
 }
 
@@ -64,7 +64,7 @@ FGLTFJsonAnimation* FGLTFAnimationDataConverter::Convert(FGLTFJsonNode* RootNode
 FGLTFJsonAnimation* FGLTFLevelSequenceConverter::Convert(const ULevel* Level, const ULevelSequence* LevelSequence)
 {
 	FGLTFJsonAnimation* JsonAnimation = Builder.AddAnimation();
-	Builder.SetupTask<FGLTFLevelSequenceTask>(Builder, Level, LevelSequence, JsonAnimation);
+	Builder.ScheduleSlowTask<FGLTFDelayedLevelSequenceTask>(Builder, Level, LevelSequence, JsonAnimation);
 	return JsonAnimation;
 }
 
diff --git a/Engine/Plugins/Enterprise/GLTFExporter/MigratedFromGit/Source/GLTFExporter/Private/Converters/GLTFMaterialConverters.cpp b/Engine/Plugins/Enterprise/GLTFExporter/MigratedFromGit/Source/GLTFExporter/Private/Converters/GLTFMaterialConverters.cpp
--- a/Engine/Plugins/Enterprise/GLTFExporter/MigratedFromGit/Source/GLTFExporter/Private/Converters/GLTFMaterialConverters.cpp
+++ b/Engine/Plugins/Enterprise/GLTFExporter/MigratedFromGit/Source/GLTFExporter/Private/Converters/GLTFMaterialConverters.cpp
@@ -3,7 +3,7 @@
 #include "Converters/GLTFMaterialConverters.h"
 #include "Converters/GLTFMaterialUtility.h"
 #include "Builders/GLTFConvertBuilder.h"
-#include "Tasks/GLTFMaterialTasks.h"
+#include "Tasks/GLTFDelayedMaterialTasks.h"
 
 void FGLTFMaterialConverter::Sanitize(const UMaterialInterface*& Material, const FGLTFMeshData*& MeshData, FGLTFIndexArray& SectionIndices)
 {
@@ -48,7 +48,7 @@ FGLTFJsonMaterial* FGLTFMaterialConverter::Convert(const UMaterialInterface* Mat
 	if (Material != FGLTFMaterialUtility::GetDefaultMaterial())
 	{
 		FGLTFJsonMaterial* JsonMaterial = Builder.AddMaterial();
-		Builder.SetupTask<FGLTFMaterialTask>(Builder, UVOverlapChecker, Material, MeshData, SectionIndices, JsonMaterial);
+		Builder.ScheduleSlowTask<FGLTFDelayedMaterialTask>(Builder, UVOverlapChecker, Material, MeshData, SectionIndices, JsonMaterial);
 		return JsonMaterial;
 	}
 
diff --git a/Engine/Plugins/Enterprise/GLTFExporter/MigratedFromGit/Source/GLTFExporter/Private/Converters/GLTFMaterialProxyFactory.cpp b/Engine/Plugins/Enterprise/GLTFExporter/MigratedFromGit/Source/GLTFExporter/Private/Converters/GLTFMaterialProxyFactory.cpp
--- a/Engine/Plugins/Enterprise/GLTFExporter/MigratedFromGit/Source/GLTFExporter/Private/Converters/GLTFMaterialProxyFactory.cpp
+++ b/Engine/Plugins/Enterprise/GLTFExporter/MigratedFromGit/Source/GLTFExporter/Private/Converters/GLTFMaterialProxyFactory.cpp
@@ -40,7 +40,7 @@ UMaterialInterface* FGLTFMaterialProxyFactory::Create(UMaterialInterface* Origin
 	}
 
 	MakeDirectory(RootPath);
-	Builder.CompleteAllTasks();
+	Builder.ProcessSlowTasks();
 
 	UMaterialInstanceConstant* ProxyMaterial = CreateInstancedMaterial(OriginalMaterial, JsonMaterial->ShadingModel);
 	if (ProxyMaterial != nullptr)
diff --git a/Engine/Plugins/Enterprise/GLTFExporter/MigratedFromGit/Source/GLTFExporter/Private/Converters/GLTFMeshConverters.cpp b/Engine/Plugins/Enterprise/GLTFExporter/MigratedFromGit/Source/GLTFExporter/Private/Converters/GLTFMeshConverters.cpp
--- a/Engine/Plugins/Enterprise/GLTFExporter/MigratedFromGit/Source/GLTFExporter/Private/Converters/GLTFMeshConverters.cpp
+++ b/Engine/Plugins/Enterprise/GLTFExporter/MigratedFromGit/Source/GLTFExporter/Private/Converters/GLTFMeshConverters.cpp
@@ -4,7 +4,7 @@
 #include "Converters/GLTFMeshUtility.h"
 #include "Converters/GLTFMaterialUtility.h"
 #include "Builders/GLTFConvertBuilder.h"
-#include "Tasks/GLTFMeshTasks.h"
+#include "Tasks/GLTFDelayedMeshTasks.h"
 #include "Engine/StaticMesh.h"
 #include "Engine/SkeletalMesh.h"
 #include "Components/StaticMeshComponent.h"
@@ -48,7 +48,7 @@ FGLTFJsonMesh* FGLTFStaticMeshConverter::Convert(const UStaticMesh* StaticMesh,
 	FGLTFJsonMesh* JsonMesh = Builder.AddMesh();
 	JsonMesh->Primitives.AddDefaulted(MaterialCount);
 
-	Builder.SetupTask<FGLTFStaticMeshTask>(Builder, MeshSectionConverter, StaticMesh, StaticMeshComponent, Materials, LODIndex, JsonMesh);
+	Builder.ScheduleSlowTask<FGLTFDelayedStaticMeshTask>(Builder, MeshSectionConverter, StaticMesh, StaticMeshComponent, Materials, LODIndex, JsonMesh);
 	return JsonMesh;
 }
 
@@ -91,6 +91,6 @@ FGLTFJsonMesh* FGLTFSkeletalMeshConverter::Convert(const USkeletalMesh* Skeletal
 	FGLTFJsonMesh* JsonMesh = Builder.AddMesh();
 	JsonMesh->Primitives.AddDefaulted(MaterialCount);
 
-	Builder.SetupTask<FGLTFSkeletalMeshTask>(Builder, MeshSectionConverter, SkeletalMesh, SkeletalMeshComponent, Materials, LODIndex, JsonMesh);
+	Builder.ScheduleSlowTask<FGLTFDelayedSkeletalMeshTask>(Builder, MeshSectionConverter, SkeletalMesh, SkeletalMeshComponent, Materials, LODIndex, JsonMesh);
 	return JsonMesh;
 }
diff --git a/Engine/Plugins/Enterprise/GLTFExporter/MigratedFromGit/Source/GLTFExporter/Private/Converters/GLTFTextureConverters.cpp b/Engine/Plugins/Enterprise/GLTFExporter/MigratedFromGit/Source/GLTFExporter/Private/Converters/GLTFTextureConverters.cpp
--- a/Engine/Plugins/Enterprise/GLTFExporter/MigratedFromGit/Source/GLTFExporter/Private/Converters/GLTFTextureConverters.cpp
+++ b/Engine/Plugins/Enterprise/GLTFExporter/MigratedFromGit/Source/GLTFExporter/Private/Converters/GLTFTextureConverters.cpp
@@ -2,7 +2,7 @@
 
 #include "Converters/GLTFTextureConverters.h"
 #include "Converters/GLTFTextureUtility.h"
-#include "Tasks/GLTFTextureTasks.h"
+#include "Tasks/GLTFDelayedTextureTasks.h"
 #include "Engine/Texture2D.h"
 #include "Engine/TextureCube.h"
 #include "Engine/TextureRenderTarget2D.h"
@@ -21,7 +21,7 @@ FGLTFJsonTexture* FGLTFTexture2DConverter::Convert(const UTexture2D* Texture2D,
 	if (Builder.ExportOptions->TextureImageFormat != EGLTFTextureImageFormat::None)
 	{
 		FGLTFJsonTexture* JsonTexture = Builder.AddTexture();
-		Builder.SetupTask<FGLTFTexture2DTask>(Builder, Texture2D, bToSRGB, JsonTexture);
+		Builder.ScheduleSlowTask<FGLTFDelayedTexture2DTask>(Builder, Texture2D, bToSRGB, JsonTexture);
 		return JsonTexture;
 	}
 
@@ -41,7 +41,7 @@ FGLTFJsonTexture* FGLTFTextureCubeConverter::Convert(const UTextureCube* Texture
 	if (Builder.ExportOptions->TextureImageFormat != EGLTFTextureImageFormat::None)
 	{
 		FGLTFJsonTexture* JsonTexture = Builder.AddTexture();
-		Builder.SetupTask<FGLTFTextureCubeTask>(Builder, TextureCube, CubeFace, bToSRGB, JsonTexture);
+		Builder.ScheduleSlowTask<FGLTFDelayedTextureCubeTask>(Builder, TextureCube, CubeFace, bToSRGB, JsonTexture);
 		return JsonTexture;
 	}
 
@@ -61,7 +61,7 @@ FGLTFJsonTexture* FGLTFTextureRenderTarget2DConverter::Convert(const UTextureRen
 	if (Builder.ExportOptions->TextureImageFormat != EGLTFTextureImageFormat::None)
 	{
 		FGLTFJsonTexture* JsonTexture = Builder.AddTexture();
-		Builder.SetupTask<FGLTFTextureRenderTarget2DTask>(Builder, RenderTarget2D, bToSRGB, JsonTexture);
+		Builder.ScheduleSlowTask<FGLTFDelayedTextureRenderTarget2DTask>(Builder, RenderTarget2D, bToSRGB, JsonTexture);
 		return JsonTexture;
 	}
 
@@ -81,7 +81,7 @@ FGLTFJsonTexture* FGLTFTextureRenderTargetCubeConverter::Convert(const UTextureR
 	if (Builder.ExportOptions->TextureImageFormat != EGLTFTextureImageFormat::None)
 	{
 		FGLTFJsonTexture* JsonTexture = Builder.AddTexture();
-		Builder.SetupTask<FGLTFTextureRenderTargetCubeTask>(Builder, RenderTargetCube, CubeFace, bToSRGB, JsonTexture);
+		Builder.ScheduleSlowTask<FGLTFDelayedTextureRenderTargetCubeTask>(Builder, RenderTargetCube, CubeFace, bToSRGB, JsonTexture);
 		return JsonTexture;
 	}
 
@@ -94,7 +94,7 @@ FGLTFJsonTexture* FGLTFTextureLightMapConverter::Convert(const ULightMapTexture2
 	if (Builder.ExportOptions->TextureImageFormat != EGLTFTextureImageFormat::None)
 	{
 		FGLTFJsonTexture* JsonTexture = Builder.AddTexture();
-		Builder.SetupTask<FGLTFTextureLightMapTask>(Builder, LightMap, JsonTexture);
+		Builder.ScheduleSlowTask<FGLTFDelayedTextureLightMapTask>(Builder, LightMap, JsonTexture);
 		return JsonTexture;
 	}
 #endif
diff --git a/Engine/Plugins/Enterprise/GLTFExporter/MigratedFromGit/Source/GLTFExporter/Private/Exporters/GLTFExporter.cpp b/Engine/Plugins/Enterprise/GLTFExporter/MigratedFromGit/Source/GLTFExporter/Private/Exporters/GLTFExporter.cpp
--- a/Engine/Plugins/Enterprise/GLTFExporter/MigratedFromGit/Source/GLTFExporter/Private/Exporters/GLTFExporter.cpp
+++ b/Engine/Plugins/Enterprise/GLTFExporter/MigratedFromGit/Source/GLTFExporter/Private/Exporters/GLTFExporter.cpp
@@ -51,7 +51,7 @@ bool UGLTFExporter::ExportBinary(UObject* Object, const TCHAR* Type, FArchive& A
 	bool bSuccess = AddObject(Builder, Object);
 	if (bSuccess)
 	{
-		Builder.CompleteAllTasks(Warn);
+		Builder.ProcessSlowTasks(Warn);
 		Builder.WriteToArchive(Archive);
 		bSuccess = Builder.WriteExternalFiles(DirPath);
 	}
@@ -93,7 +93,7 @@ bool UGLTFExporter::ExportToGLTF(UObject* Object, const FString& FilePath, const
 
 	if (bSuccess)
 	{
-		Builder.CompleteAllTasks(nullptr);
+		Builder.ProcessSlowTasks(nullptr);
 		bSuccess = Builder.WriteAllFiles(DirPath);
 	}
 
diff --git a/Engine/Plugins/Enterprise/GLTFExporter/MigratedFromGit/Source/GLTFExporter/Private/Tasks/GLTFAnimationTasks.cpp b/Engine/Plugins/Enterprise/GLTFExporter/MigratedFromGit/Source/GLTFExporter/Private/Tasks/GLTFAnimationTasks.cpp
deleted file mode 100644
diff --git a/Engine/Plugins/Enterprise/GLTFExporter/MigratedFromGit/Source/GLTFExporter/Private/Tasks/GLTFAnimationTasks.h b/Engine/Plugins/Enterprise/GLTFExporter/MigratedFromGit/Source/GLTFExporter/Private/Tasks/GLTFAnimationTasks.h
deleted file mode 100644
diff --git a/Engine/Plugins/Enterprise/GLTFExporter/MigratedFromGit/Source/GLTFExporter/Private/Tasks/GLTFDelayedAnimationTasks.cpp b/Engine/Plugins/Enterprise/GLTFExporter/MigratedFromGit/Source/GLTFExporter/Private/Tasks/GLTFDelayedAnimationTasks.cpp
new file mode 100644
--- /dev/null
+++ b/Engine/Plugins/Enterprise/GLTFExporter/MigratedFromGit/Source/GLTFExporter/Private/Tasks/GLTFDelayedAnimationTasks.cpp
@@ -0,0 +1,358 @@
+// Copyright Epic Games, Inc. All Rights Reserved.
+
+#include "Tasks/GLTFDelayedAnimationTasks.h"
+#include "Builders/GLTFContainerBuilder.h"
+#include "Converters/GLTFConverterUtility.h"
+#include "Converters/GLTFBoneUtility.h"
+#include "LevelSequence.h"
+#include "LevelSequencePlayer.h"
+#include "MovieSceneTimeHelpers.h"
+#include "Tracks/MovieScene3DTransformTrack.h"
+#include "Sections/MovieScene3DTransformSection.h"
+#include "Channels/MovieSceneChannelProxy.h"
+
+#if (ENGINE_MAJOR_VERSION > 4 || ENGINE_MINOR_VERSION >= 26)
+using namespace UE;
+#endif
+
+void FGLTFDelayedAnimSequenceTask::Process()
+{
+	// TODO: bone transforms should be absolute (not relative) according to gltf spec
+
+	JsonAnimation->Name = AnimSequence->GetName() + TEXT("_") + FString::FromInt(JsonAnimation->Index); // Ensure unique name due to limitation in certain gltf viewers
+
+	TArray<float> Timestamps;
+	FGLTFBoneUtility::GetFrameTimestamps(AnimSequence, Timestamps);
+	const int32 FrameCount = Timestamps.Num();
+
+	// TODO: add animation data accessor converters to reuse track information
+
+	FGLTFJsonAccessor* JsonInputAccessor = Builder.AddAccessor();
+	JsonInputAccessor->BufferView = Builder.AddBufferView(Timestamps);
+	JsonInputAccessor->ComponentType = EGLTFJsonComponentType::Float;
+	JsonInputAccessor->Type = EGLTFJsonAccessorType::Scalar;
+	JsonInputAccessor->Count = FrameCount;
+	JsonInputAccessor->MinMaxLength = 1;
+	JsonInputAccessor->Min[0] = 0;
+	JsonInputAccessor->Max[0] = Timestamps[FrameCount - 1];
+
+	EGLTFJsonInterpolation Interpolation = FGLTFConverterUtility::ConvertInterpolation(AnimSequence->Interpolation);
+	if (Interpolation == EGLTFJsonInterpolation::None)
+	{
+		Interpolation = EGLTFJsonInterpolation::Linear;
+		// TODO: report warning (about unknown interpolation exported as linear)
+	}
+
+	TArray<FBoneIndexType> BoneIndices;
+	FGLTFBoneUtility::GetBoneIndices(AnimSequence->GetSkeleton(), BoneIndices);
+
+	TArray<TArray<FTransform>> FrameTransforms;
+	FGLTFBoneUtility::GetBoneTransformsByFrame(AnimSequence, Timestamps, BoneIndices, FrameTransforms);
+
+	for (const FBoneIndexType BoneIndex : BoneIndices)
+	{
+		FGLTFJsonNode* Node = Builder.AddUniqueNode(RootNode, SkeletalMesh, BoneIndex);
+
+		// TODO: detect if a bone has the same transforms across multiple frames (at least if its the same across all frames) and optimize
+
+		{
+			TArray<FGLTFVector3> Translations;
+			Translations.AddUninitialized(FrameCount);
+
+			for (int32 Frame = 0; Frame < FrameCount; ++Frame)
+			{
+				const FVector KeyPosition = FrameTransforms[Frame][BoneIndex].GetTranslation();
+				Translations[Frame] = FGLTFConverterUtility::ConvertPosition(KeyPosition, Builder.ExportOptions->ExportUniformScale);
+			}
+
+			FGLTFJsonAccessor* JsonOutputAccessor = Builder.AddAccessor();
+			JsonOutputAccessor->BufferView = Builder.AddBufferView(Translations);
+			JsonOutputAccessor->ComponentType = EGLTFJsonComponentType::Float;
+			JsonOutputAccessor->Count = Translations.Num();
+			JsonOutputAccessor->Type = EGLTFJsonAccessorType::Vec3;
+
+			FGLTFJsonAnimationSampler* JsonSampler = JsonAnimation->Samplers.Add();
+			JsonSampler->Input = JsonInputAccessor;
+			JsonSampler->Output = JsonOutputAccessor;
+			JsonSampler->Interpolation = Interpolation;
+
+			FGLTFJsonAnimationChannel JsonChannel;
+			JsonChannel.Sampler = JsonSampler;
+			JsonChannel.Target.Path = EGLTFJsonTargetPath::Translation;
+			JsonChannel.Target.Node = Node;
+			JsonAnimation->Channels.Add(JsonChannel);
+		}
+
+		{
+			TArray<FGLTFQuaternion> Rotations;
+			Rotations.AddUninitialized(FrameCount);
+
+			for (int32 Frame = 0; Frame < FrameCount; ++Frame)
+			{
+				const FQuat KeyRotation = FrameTransforms[Frame][BoneIndex].GetRotation();
+				Rotations[Frame] = FGLTFConverterUtility::ConvertRotation(KeyRotation);
+			}
+
+			FGLTFJsonAccessor* JsonOutputAccessor = Builder.AddAccessor();
+			JsonOutputAccessor->BufferView = Builder.AddBufferView(Rotations);
+			JsonOutputAccessor->ComponentType = EGLTFJsonComponentType::Float;
+			JsonOutputAccessor->Count = Rotations.Num();
+			JsonOutputAccessor->Type = EGLTFJsonAccessorType::Vec4;
+
+			FGLTFJsonAnimationSampler* JsonSampler = JsonAnimation->Samplers.Add();
+			JsonSampler->Input = JsonInputAccessor;
+			JsonSampler->Output = JsonOutputAccessor;
+			JsonSampler->Interpolation = Interpolation;
+
+			FGLTFJsonAnimationChannel JsonChannel;
+			JsonChannel.Sampler = JsonSampler;
+			JsonChannel.Target.Path = EGLTFJsonTargetPath::Rotation;
+			JsonChannel.Target.Node = Node;
+			JsonAnimation->Channels.Add(JsonChannel);
+		}
+
+		{
+			TArray<FGLTFVector3> Scales;
+			Scales.AddUninitialized(FrameCount);
+
+			for (int32 Frame = 0; Frame < FrameCount; ++Frame)
+			{
+				const FVector KeyScale = FrameTransforms[Frame][BoneIndex].GetScale3D();
+				Scales[Frame] = FGLTFConverterUtility::ConvertScale(KeyScale);
+			}
+
+			FGLTFJsonAccessor* JsonOutputAccessor = Builder.AddAccessor();
+			JsonOutputAccessor->BufferView = Builder.AddBufferView(Scales);
+			JsonOutputAccessor->ComponentType = EGLTFJsonComponentType::Float;
+			JsonOutputAccessor->Count = Scales.Num();
+			JsonOutputAccessor->Type = EGLTFJsonAccessorType::Vec3;
+
+			FGLTFJsonAnimationSampler* JsonSampler = JsonAnimation->Samplers.Add();
+			JsonSampler->Input = JsonInputAccessor;
+			JsonSampler->Output = JsonOutputAccessor;
+			JsonSampler->Interpolation = Interpolation;
+
+			FGLTFJsonAnimationChannel JsonChannel;
+			JsonChannel.Sampler = JsonSampler;
+			JsonChannel.Target.Path = EGLTFJsonTargetPath::Scale;
+			JsonChannel.Target.Node = Node;
+			JsonAnimation->Channels.Add(JsonChannel);
+		}
+	}
+}
+
+void FGLTFDelayedLevelSequenceTask::Process()
+{
+	ULevelSequence* Sequence = const_cast<ULevelSequence*>(LevelSequence);
+	UMovieScene* MovieScene = Sequence->GetMovieScene();
+
+	ALevelSequenceActor* LevelSequenceActor;
+	UWorld* World = Level->GetWorld();
+	ULevelSequencePlayer* Player = ULevelSequencePlayer::CreateLevelSequencePlayer(World, Sequence, FMovieSceneSequencePlaybackSettings(), LevelSequenceActor);
+
+	Player->Initialize(Sequence, LevelSequenceActor->GetLevel(), FMovieSceneSequencePlaybackSettings(), FLevelSequenceCameraSettings());
+	Player->State.AssignSequence(MovieSceneSequenceID::Root, *Sequence, *Player);
+
+	FFrameRate TickResolution = MovieScene->GetTickResolution();
+	FFrameRate DisplayRate = MovieScene->GetDisplayRate(); // TODO: add option to switch between DisplayRate and TickResolution?
+	TRange<FFrameNumber> PlaybackRange = MovieScene->GetPlaybackRange();
+
+	int32 FrameOffset = FFrameRate::TransformTime(FFrameTime(MovieScene::DiscreteInclusiveLower(PlaybackRange)), TickResolution, DisplayRate).RoundToFrame().Value;
+	int32 FrameCount = FFrameRate::TransformTime(FFrameTime(FFrameNumber(MovieScene::DiscreteSize(PlaybackRange))), TickResolution, DisplayRate).RoundToFrame().Value + 1;
+
+	JsonAnimation->Name = Sequence->GetName() + TEXT("_") + FString::FromInt(JsonAnimation->Index); // Ensure unique name due to limitation in certain gltf viewers
+
+	TArray<float> Timestamps;
+	TArray<FFrameTime> FrameTimes;
+	Timestamps.AddUninitialized(FrameCount);
+	FrameTimes.AddUninitialized(FrameCount);
+
+	for (int32 Frame = 0; Frame < FrameCount; ++Frame)
+	{
+		Timestamps[Frame] = DisplayRate.AsSeconds(Frame);
+		FrameTimes[Frame] = FFrameRate::TransformTime(FFrameTime(FrameOffset + Frame), DisplayRate, TickResolution);
+	}
+
+	// TODO: add animation data accessor converters to reuse track information
+
+	FGLTFJsonAccessor* JsonInputAccessor = Builder.AddAccessor();
+	JsonInputAccessor->BufferView = Builder.AddBufferView(Timestamps);
+	JsonInputAccessor->ComponentType = EGLTFJsonComponentType::Float;
+	JsonInputAccessor->Type = EGLTFJsonAccessorType::Scalar;
+	JsonInputAccessor->Count = FrameCount;
+	JsonInputAccessor->MinMaxLength = 1;
+	JsonInputAccessor->Min[0] = 0;
+	JsonInputAccessor->Max[0] = Timestamps[FrameCount - 1];
+
+	for (const FMovieSceneBinding& Binding : MovieScene->GetBindings())
+	{
+		for (TWeakObjectPtr<UObject> Object : Player->FindBoundObjects(Binding.GetObjectGuid(), MovieSceneSequenceID::Root))
+		{
+			FGLTFJsonNode* Node;
+
+			if (AActor* Actor = Cast<AActor>(Object.Get()))
+			{
+				Node = Builder.AddUniqueNode(Actor);
+			}
+			else if (USceneComponent* SceneComponent = Cast<USceneComponent>(Object.Get()))
+			{
+				Node = Builder.AddUniqueNode(SceneComponent);
+			}
+			else
+			{
+				// TODO: report warning
+				continue;
+			}
+
+			if (Node == nullptr)
+			{
+				// TODO: report warning
+				continue;
+			}
+
+			for (UMovieSceneTrack* Track : Binding.GetTracks())
+			{
+				UMovieScene3DTransformTrack* TransformTrack = Cast<UMovieScene3DTransformTrack>(Track);
+				if (TransformTrack == nullptr)
+				{
+					// TODO: report warning
+					continue;
+				}
+
+				for (UMovieSceneSection* Section : TransformTrack->GetAllSections())
+				{
+					// TODO: add support for Section->GetCompletionMode() (i.e. when finished)
+
+					UMovieScene3DTransformSection* TransformSection = Cast<UMovieScene3DTransformSection>(Section);
+					if (TransformSection == nullptr)
+					{
+						// TODO: report warning
+						continue;
+					}
+
+					FOptionalMovieSceneBlendType BlendType = TransformSection->GetBlendType();
+					if (BlendType.IsValid() && BlendType.Get() != EMovieSceneBlendType::Absolute)
+					{
+						// TODO: report warning
+						continue;
+					}
+
+					// TODO: do we need to account for TransformSection->GetOffsetTime() ?
+
+					TArrayView<FMovieSceneFloatChannel*> Channels = TransformSection->GetChannelProxy().GetChannels<FMovieSceneFloatChannel>();
+					EMovieSceneTransformChannel ChannelMask = TransformSection->GetMask().GetChannels();
+
+					if (EnumHasAnyFlags(ChannelMask, EMovieSceneTransformChannel::Translation))
+					{
+						TArray<FGLTFVector3> Translations;
+						Translations.AddUninitialized(FrameCount);
+
+						for (int32 Frame = 0; Frame < FrameCount; ++Frame)
+						{
+							FFrameTime& FrameTime = FrameTimes[Frame];
+
+							FVector Translation = FVector::ZeroVector;
+							Channels[0]->Evaluate(FrameTime, Translation.X);
+							Channels[1]->Evaluate(FrameTime, Translation.Y);
+							Channels[2]->Evaluate(FrameTime, Translation.Z);
+
+							Translations[Frame] = FGLTFConverterUtility::ConvertPosition(Translation, Builder.ExportOptions->ExportUniformScale);
+						}
+
+						FGLTFJsonAccessor* JsonOutputAccessor = Builder.AddAccessor();
+						JsonOutputAccessor->BufferView = Builder.AddBufferView(Translations);
+						JsonOutputAccessor->ComponentType = EGLTFJsonComponentType::Float;
+						JsonOutputAccessor->Count = Translations.Num();
+						JsonOutputAccessor->Type = EGLTFJsonAccessorType::Vec3;
+
+						FGLTFJsonAnimationSampler* JsonSampler = JsonAnimation->Samplers.Add();
+						JsonSampler->Input = JsonInputAccessor;
+						JsonSampler->Output = JsonOutputAccessor;
+						JsonSampler->Interpolation = EGLTFJsonInterpolation::Linear;
+
+						FGLTFJsonAnimationChannel JsonChannel;
+						JsonChannel.Sampler = JsonSampler;
+						JsonChannel.Target.Path = EGLTFJsonTargetPath::Translation;
+						JsonChannel.Target.Node = Node;
+						JsonAnimation->Channels.Add(JsonChannel);
+					}
+
+					if (EnumHasAnyFlags(ChannelMask, EMovieSceneTransformChannel::Rotation))
+					{
+						TArray<FGLTFQuaternion> Rotations;
+						Rotations.AddUninitialized(FrameCount);
+
+						for (int32 Frame = 0; Frame < FrameCount; ++Frame)
+						{
+							FFrameTime& FrameTime = FrameTimes[Frame];
+
+							FRotator Rotator = FRotator::ZeroRotator;
+							Channels[3]->Evaluate(FrameTime, Rotator.Roll);
+							Channels[4]->Evaluate(FrameTime, Rotator.Pitch);
+							Channels[5]->Evaluate(FrameTime, Rotator.Yaw);
+
+							Rotations[Frame] = FGLTFConverterUtility::ConvertRotation(Rotator.Quaternion());
+						}
+
+						JsonInputAccessor->Count = Rotations.Num();
+						JsonInputAccessor->Max[0] = Timestamps[Rotations.Num() - 1];
+
+						FGLTFJsonAccessor* JsonOutputAccessor = Builder.AddAccessor();
+						JsonOutputAccessor->BufferView = Builder.AddBufferView(Rotations);
+						JsonOutputAccessor->ComponentType = EGLTFJsonComponentType::Float;
+						JsonOutputAccessor->Count = Rotations.Num();
+						JsonOutputAccessor->Type = EGLTFJsonAccessorType::Vec4;
+
+						FGLTFJsonAnimationSampler* JsonSampler = JsonAnimation->Samplers.Add();
+						JsonSampler->Input = JsonInputAccessor;
+						JsonSampler->Output = JsonOutputAccessor;
+						JsonSampler->Interpolation = EGLTFJsonInterpolation::Linear;
+
+						FGLTFJsonAnimationChannel JsonChannel;
+						JsonChannel.Sampler = JsonSampler;
+						JsonChannel.Target.Path = EGLTFJsonTargetPath::Rotation;
+						JsonChannel.Target.Node = Node;
+						JsonAnimation->Channels.Add(JsonChannel);
+					}
+
+					if (EnumHasAnyFlags(ChannelMask, EMovieSceneTransformChannel::Scale))
+					{
+						TArray<FGLTFVector3> Scales;
+						Scales.AddUninitialized(FrameCount);
+
+						for (int32 Frame = 0; Frame < FrameCount; ++Frame)
+						{
+							FFrameTime& FrameTime = FrameTimes[Frame];
+
+							FVector Scale = FVector::OneVector;
+							Channels[6]->Evaluate(FrameTime, Scale.X);
+							Channels[7]->Evaluate(FrameTime, Scale.Y);
+							Channels[8]->Evaluate(FrameTime, Scale.Z);
+
+							Scales[Frame] = FGLTFConverterUtility::ConvertScale(Scale);
+						}
+
+						FGLTFJsonAccessor* JsonOutputAccessor = Builder.AddAccessor();
+						JsonOutputAccessor->BufferView = Builder.AddBufferView(Scales);
+						JsonOutputAccessor->ComponentType = EGLTFJsonComponentType::Float;
+						JsonOutputAccessor->Count = Scales.Num();
+						JsonOutputAccessor->Type = EGLTFJsonAccessorType::Vec3;
+
+						FGLTFJsonAnimationSampler* JsonSampler = JsonAnimation->Samplers.Add();
+						JsonSampler->Input = JsonInputAccessor;
+						JsonSampler->Output = JsonOutputAccessor;
+						JsonSampler->Interpolation = EGLTFJsonInterpolation::Linear;
+
+						FGLTFJsonAnimationChannel JsonChannel;
+						JsonChannel.Sampler = JsonSampler;
+						JsonChannel.Target.Path = EGLTFJsonTargetPath::Scale;
+						JsonChannel.Target.Node = Node;
+						JsonAnimation->Channels.Add(JsonChannel);
+					}
+				}
+			}
+		}
+	}
+
+	World->DestroyActor(LevelSequenceActor);
+}
diff --git a/Engine/Plugins/Enterprise/GLTFExporter/MigratedFromGit/Source/GLTFExporter/Private/Tasks/GLTFDelayedAnimationTasks.h b/Engine/Plugins/Enterprise/GLTFExporter/MigratedFromGit/Source/GLTFExporter/Private/Tasks/GLTFDelayedAnimationTasks.h
new file mode 100644
--- /dev/null
+++ b/Engine/Plugins/Enterprise/GLTFExporter/MigratedFromGit/Source/GLTFExporter/Private/Tasks/GLTFDelayedAnimationTasks.h
@@ -0,0 +1,66 @@
+// Copyright Epic Games, Inc. All Rights Reserved.
+
+#pragma once
+
+#include "Tasks/GLTFDelayedTask.h"
+#include "Builders/GLTFConvertBuilder.h"
+#include "LevelSequenceActor.h"
+#include "LevelSequence.h"
+
+class FGLTFDelayedAnimSequenceTask : public FGLTFDelayedTask
+{
+public:
+
+	FGLTFDelayedAnimSequenceTask(FGLTFConvertBuilder& Builder,  FGLTFJsonNode* RootNode, const USkeletalMesh* SkeletalMesh, const UAnimSequence* AnimSequence, FGLTFJsonAnimation* JsonAnimation)
+		: FGLTFDelayedTask(EGLTFTaskPriority::Animation)
+		, Builder(Builder)
+		, RootNode(RootNode)
+		, SkeletalMesh(SkeletalMesh)
+		, AnimSequence(AnimSequence)
+		, JsonAnimation(JsonAnimation)
+	{
+	}
+
+	virtual FString GetName() override
+	{
+		return AnimSequence->GetName();
+	}
+
+	virtual void Process() override;
+
+private:
+
+	FGLTFConvertBuilder& Builder;
+	FGLTFJsonNode* RootNode;
+	const USkeletalMesh* SkeletalMesh;
+	const UAnimSequence* AnimSequence;
+	FGLTFJsonAnimation* JsonAnimation;
+};
+
+class FGLTFDelayedLevelSequenceTask : public FGLTFDelayedTask
+{
+public:
+
+	FGLTFDelayedLevelSequenceTask(FGLTFConvertBuilder& Builder, const ULevel* Level, const ULevelSequence* LevelSequence, FGLTFJsonAnimation* JsonAnimation)
+		: FGLTFDelayedTask(EGLTFTaskPriority::Animation)
+		, Builder(Builder)
+		, Level(Level)
+		, LevelSequence(LevelSequence)
+		, JsonAnimation(JsonAnimation)
+	{
+	}
+
+	virtual FString GetName() override
+	{
+		return LevelSequence->GetName();
+	}
+
+	virtual void Process() override;
+
+private:
+
+	FGLTFConvertBuilder& Builder;
+	const ULevel* Level;
+	const ULevelSequence* LevelSequence;
+	FGLTFJsonAnimation* JsonAnimation;
+};
diff --git a/Engine/Plugins/Enterprise/GLTFExporter/MigratedFromGit/Source/GLTFExporter/Private/Tasks/GLTFDelayedMaterialTasks.cpp b/Engine/Plugins/Enterprise/GLTFExporter/MigratedFromGit/Source/GLTFExporter/Private/Tasks/GLTFDelayedMaterialTasks.cpp
new file mode 100644
--- /dev/null
+++ b/Engine/Plugins/Enterprise/GLTFExporter/MigratedFromGit/Source/GLTFExporter/Private/Tasks/GLTFDelayedMaterialTasks.cpp
@@ -0,0 +1,1618 @@
+// Copyright Epic Games, Inc. All Rights Reserved.
+
+#include "Tasks/GLTFDelayedMaterialTasks.h"
+#include "Converters/GLTFConverterUtility.h"
+#include "Converters/GLTFNameUtility.h"
+#include "Converters/GLTFMaterialUtility.h"
+#include "Builders/GLTFContainerBuilder.h"
+#include "Utilities/GLTFProxyMaterialUtilities.h"
+#include "Materials/GLTFProxyMaterialInfo.h"
+#if WITH_EDITOR
+#include "Materials/MaterialInstance.h"
+#include "Materials/MaterialExpressionConstant.h"
+#include "Materials/MaterialExpressionConstant2Vector.h"
+#include "Materials/MaterialExpressionConstant3Vector.h"
+#include "Materials/MaterialExpressionConstant4Vector.h"
+#include "Materials/MaterialExpressionScalarParameter.h"
+#include "Materials/MaterialExpressionVectorParameter.h"
+#include "Materials/MaterialExpressionTextureSample.h"
+#include "Materials/MaterialExpressionTextureSampleParameter2D.h"
+#endif
+
+namespace
+{
+	// Component masks
+	const FLinearColor RedMask(1.0f, 0.0f, 0.0f, 0.0f);
+	const FLinearColor GreenMask(0.0f, 1.0f, 0.0f, 0.0f);
+	const FLinearColor BlueMask(0.0f, 0.0f, 1.0f, 0.0f);
+	const FLinearColor AlphaMask(0.0f, 0.0f, 0.0f, 1.0f);
+	const FLinearColor RgbMask = RedMask + GreenMask + BlueMask;
+	const FLinearColor RgbaMask = RgbMask + AlphaMask;
+
+	// Property-specific component masks
+	const FLinearColor BaseColorMask = RgbMask;
+	const FLinearColor OpacityMask = AlphaMask;
+	const FLinearColor MetallicMask = BlueMask;
+	const FLinearColor RoughnessMask = GreenMask;
+	const FLinearColor OcclusionMask = RedMask;
+	const FLinearColor ClearCoatMask = RedMask;
+	const FLinearColor ClearCoatRoughnessMask = GreenMask;
+
+	// Ideal masks for texture-inputs (doesn't require baking)
+	const TArray<FLinearColor> DefaultColorInputMasks = { RgbMask, RgbaMask };
+	const TArray<FLinearColor> BaseColorInputMasks = { BaseColorMask };
+	const TArray<FLinearColor> OpacityInputMasks = { OpacityMask };
+	const TArray<FLinearColor> MetallicInputMasks = { MetallicMask };
+	const TArray<FLinearColor> RoughnessInputMasks = { RoughnessMask };
+	const TArray<FLinearColor> OcclusionInputMasks = { OcclusionMask };
+	const TArray<FLinearColor> ClearCoatInputMasks = { ClearCoatMask };
+	const TArray<FLinearColor> ClearCoatRoughnessInputMasks = { ClearCoatRoughnessMask };
+}
+
+void FGLTFDelayedMaterialTask::Process()
+{
+	const UMaterial* BaseMaterial = Material->GetMaterial();
+	if (BaseMaterial->MaterialDomain != MD_Surface)
+	{
+		// TODO: report warning (non-surface materials not supported, will be treated as surface)
+	}
+
+	JsonMaterial->Name = GetMaterialName();
+	JsonMaterial->AlphaCutoff = Material->GetOpacityMaskClipValue();
+	JsonMaterial->DoubleSided = Material->IsTwoSided();
+
+	if (const UMaterialInstance* MaterialInstance = Cast<UMaterialInstance>(Material))
+	{
+		// TODO: remove ugly hack for dynamic instances
+		FGLTFProxyMaterialUtilities::GetOpacityMaskClipValue(MaterialInstance, JsonMaterial->AlphaCutoff);
+		FGLTFProxyMaterialUtilities::GetTwoSided(MaterialInstance, JsonMaterial->DoubleSided);
+	}
+
+	ConvertShadingModel(JsonMaterial->ShadingModel);
+	ConvertAlphaMode(JsonMaterial->AlphaMode, JsonMaterial->BlendMode);
+
+	if (FGLTFProxyMaterialUtilities::IsProxyMaterial(BaseMaterial))
+	{
+		GetProxyParameters(*JsonMaterial);
+		return;
+	}
+
+#if WITH_EDITOR
+	if (JsonMaterial->ShadingModel != EGLTFJsonShadingModel::None)
+	{
+		const FGLTFMaterialPropertyEx BaseColorProperty = JsonMaterial->ShadingModel == EGLTFJsonShadingModel::Unlit ? MP_EmissiveColor : MP_BaseColor;
+		const FGLTFMaterialPropertyEx OpacityProperty = JsonMaterial->AlphaMode == EGLTFJsonAlphaMode::Mask ? MP_OpacityMask : MP_Opacity;
+
+		// TODO: check if a property is active before trying to get it (i.e. Material->IsPropertyActive)
+
+		if (JsonMaterial->AlphaMode == EGLTFJsonAlphaMode::Opaque)
+		{
+			if (!TryGetConstantColor(JsonMaterial->PBRMetallicRoughness.BaseColorFactor, BaseColorProperty))
+			{
+				if (!TryGetSourceTexture(JsonMaterial->PBRMetallicRoughness.BaseColorTexture, BaseColorProperty, DefaultColorInputMasks))
+				{
+					if (!TryGetBakedMaterialProperty(JsonMaterial->PBRMetallicRoughness.BaseColorTexture, JsonMaterial->PBRMetallicRoughness.BaseColorFactor, BaseColorProperty, TEXT("BaseColor")))
+					{
+						Builder.LogWarning(FString::Printf(TEXT("Failed to export %s for material %s"), *BaseColorProperty.ToString(), *Material->GetName()));
+					}
+				}
+			}
+
+			JsonMaterial->PBRMetallicRoughness.BaseColorFactor.A = 1.0f; // make sure base color is opaque
+		}
+		else
+		{
+			if (!TryGetBaseColorAndOpacity(JsonMaterial->PBRMetallicRoughness, BaseColorProperty, OpacityProperty))
+			{
+				Builder.LogWarning(FString::Printf(TEXT("Failed to export %s and %s for material %s"), *BaseColorProperty.ToString(), *OpacityProperty.ToString(), *Material->GetName()));
+			}
+		}
+
+		if (JsonMaterial->ShadingModel == EGLTFJsonShadingModel::Default || JsonMaterial->ShadingModel == EGLTFJsonShadingModel::ClearCoat)
+		{
+			const FGLTFMaterialPropertyEx MetallicProperty = MP_Metallic;
+			const FGLTFMaterialPropertyEx RoughnessProperty = MP_Roughness;
+
+			if (!TryGetMetallicAndRoughness(JsonMaterial->PBRMetallicRoughness, MetallicProperty, RoughnessProperty))
+			{
+				Builder.LogWarning(FString::Printf(TEXT("Failed to export %s and %s for material %s"), *MetallicProperty.ToString(), *RoughnessProperty.ToString(), *Material->GetName()));
+			}
+
+			const FGLTFMaterialPropertyEx EmissiveProperty = MP_EmissiveColor;
+			if (!TryGetEmissive(*JsonMaterial, EmissiveProperty))
+			{
+				Builder.LogWarning(FString::Printf(TEXT("Failed to export %s for material %s"), *EmissiveProperty.ToString(), *Material->GetName()));
+			}
+
+			const FGLTFMaterialPropertyEx NormalProperty = JsonMaterial->ShadingModel == EGLTFJsonShadingModel::ClearCoat ? FGLTFMaterialPropertyEx(TEXT("ClearCoatBottomNormal")) : FGLTFMaterialPropertyEx(MP_Normal);
+			if (IsPropertyNonDefault(NormalProperty))
+			{
+				if (!TryGetSourceTexture(JsonMaterial->NormalTexture, NormalProperty, DefaultColorInputMasks))
+				{
+					if (!TryGetBakedMaterialProperty(JsonMaterial->NormalTexture, NormalProperty, TEXT("Normal")))
+					{
+						Builder.LogWarning(FString::Printf(TEXT("Failed to export %s for material %s"), *NormalProperty.ToString(), *Material->GetName()));
+					}
+				}
+			}
+
+			const FGLTFMaterialPropertyEx AmbientOcclusionProperty = MP_AmbientOcclusion;
+			if (IsPropertyNonDefault(AmbientOcclusionProperty))
+			{
+				if (!TryGetSourceTexture(JsonMaterial->OcclusionTexture, AmbientOcclusionProperty, OcclusionInputMasks))
+				{
+					if (!TryGetBakedMaterialProperty(JsonMaterial->OcclusionTexture, AmbientOcclusionProperty, TEXT("Occlusion")))
+					{
+						Builder.LogWarning(FString::Printf(TEXT("Failed to export %s for material %s"), *AmbientOcclusionProperty.ToString(), *Material->GetName()));
+					}
+				}
+			}
+
+			if (JsonMaterial->ShadingModel == EGLTFJsonShadingModel::ClearCoat)
+			{
+				const FGLTFMaterialPropertyEx ClearCoatProperty = MP_CustomData0;
+				const FGLTFMaterialPropertyEx ClearCoatRoughnessProperty = MP_CustomData1;
+
+				if (!TryGetClearCoatRoughness(JsonMaterial->ClearCoat, ClearCoatProperty, ClearCoatRoughnessProperty))
+				{
+					Builder.LogWarning(FString::Printf(TEXT("Failed to export %s and %s for material %s"), *ClearCoatProperty.ToString(), *ClearCoatRoughnessProperty.ToString(), *Material->GetName()));
+				}
+
+				const FGLTFMaterialPropertyEx ClearCoatNormalProperty = MP_Normal;
+				if (IsPropertyNonDefault(ClearCoatNormalProperty))
+				{
+					if (!TryGetSourceTexture(JsonMaterial->ClearCoat.ClearCoatNormalTexture, ClearCoatNormalProperty, DefaultColorInputMasks))
+					{
+						if (!TryGetBakedMaterialProperty(JsonMaterial->ClearCoat.ClearCoatNormalTexture, ClearCoatNormalProperty, TEXT("ClearCoatNormal")))
+						{
+							Builder.LogWarning(FString::Printf(TEXT("Failed to export %s for material %s"), *ClearCoatNormalProperty.ToString(), *Material->GetName()));
+						}
+					}
+				}
+			}
+		}
+	}
+
+	// Warn if properties baked using mesh data are using overlapping UVs
+	if (MeshData != nullptr && MeshDataBakedProperties.Num() > 0)
+	{
+		const FGLTFMeshData* ParentMeshData = MeshData->GetParent();
+
+		const float UVOverlapThreshold = 1.0f / 100.0f;
+		const float UVOverlap = UVOverlapChecker.GetOrAdd(&ParentMeshData->Description, SectionIndices, MeshData->BakeUsingTexCoord);
+
+		if (UVOverlap > UVOverlapThreshold)
+		{
+			FString SectionString = TEXT("mesh section");
+			SectionString += SectionIndices.Num() > 1 ? TEXT("s ") : TEXT(" ");
+			SectionString += FString::JoinBy(SectionIndices, TEXT(", "), FString::FromInt);
+
+			Builder.LogWarning(FString::Printf(
+				TEXT("Material %s is baked using mesh data from %s but the lightmap UV (channel %d) are overlapping by %.2f%% (in %s) and may produce incorrect results"),
+				*Material->GetName(),
+				*ParentMeshData->Name,
+				MeshData->BakeUsingTexCoord,
+				UVOverlap * 100,
+				*SectionString));
+		}
+	}
+#else
+	Builder.LogError(FString::Printf(
+		TEXT("Can't properly export material %s in a runtime environment without a glTF proxy. In the Unreal Editor's Content Browser, please right-click on the material, then select Create glTF Proxy Material"),
+		*Material->GetName()));
+#endif
+}
+
+FString FGLTFDelayedMaterialTask::GetMaterialName() const
+{
+	FString MaterialName = Material->GetName();
+
+	if (MeshData != nullptr)
+	{
+		MaterialName += TEXT("_") + MeshData->Name;
+	}
+
+	return MaterialName;
+}
+
+FString FGLTFDelayedMaterialTask::GetBakedTextureName(const FString& PropertyName) const
+{
+	return GetMaterialName() + TEXT("_") + PropertyName;
+}
+
+void FGLTFDelayedMaterialTask::GetProxyParameters(FGLTFJsonMaterial& OutMaterial) const
+{
+	GetProxyParameter(FGLTFProxyMaterialInfo::BaseColorFactor, OutMaterial.PBRMetallicRoughness.BaseColorFactor);
+	GetProxyParameter(FGLTFProxyMaterialInfo::BaseColor, OutMaterial.PBRMetallicRoughness.BaseColorTexture);
+
+	if (OutMaterial.ShadingModel == EGLTFJsonShadingModel::Default || OutMaterial.ShadingModel == EGLTFJsonShadingModel::ClearCoat)
+	{
+		GetProxyParameter(FGLTFProxyMaterialInfo::EmissiveFactor, OutMaterial.EmissiveFactor);
+		GetProxyParameter(FGLTFProxyMaterialInfo::Emissive, OutMaterial.EmissiveTexture);
+
+		GetProxyParameter(FGLTFProxyMaterialInfo::MetallicFactor, OutMaterial.PBRMetallicRoughness.MetallicFactor);
+		GetProxyParameter(FGLTFProxyMaterialInfo::RoughnessFactor, OutMaterial.PBRMetallicRoughness.RoughnessFactor);
+		GetProxyParameter(FGLTFProxyMaterialInfo::MetallicRoughness, OutMaterial.PBRMetallicRoughness.MetallicRoughnessTexture);
+
+		GetProxyParameter(FGLTFProxyMaterialInfo::NormalScale, OutMaterial.NormalTexture.Scale);
+		GetProxyParameter(FGLTFProxyMaterialInfo::Normal, OutMaterial.NormalTexture);
+
+		GetProxyParameter(FGLTFProxyMaterialInfo::OcclusionStrength, OutMaterial.OcclusionTexture.Strength);
+		GetProxyParameter(FGLTFProxyMaterialInfo::Occlusion, OutMaterial.OcclusionTexture);
+
+		if (OutMaterial.ShadingModel == EGLTFJsonShadingModel::ClearCoat)
+		{
+			GetProxyParameter(FGLTFProxyMaterialInfo::ClearCoatFactor, OutMaterial.ClearCoat.ClearCoatFactor);
+			GetProxyParameter(FGLTFProxyMaterialInfo::ClearCoat, OutMaterial.ClearCoat.ClearCoatTexture);
+
+			GetProxyParameter(FGLTFProxyMaterialInfo::ClearCoatRoughnessFactor, OutMaterial.ClearCoat.ClearCoatRoughnessFactor);
+			GetProxyParameter(FGLTFProxyMaterialInfo::ClearCoatRoughness, OutMaterial.ClearCoat.ClearCoatRoughnessTexture);
+
+			GetProxyParameter(FGLTFProxyMaterialInfo::ClearCoatNormalScale, OutMaterial.ClearCoat.ClearCoatNormalTexture.Scale);
+			GetProxyParameter(FGLTFProxyMaterialInfo::ClearCoatNormal, OutMaterial.ClearCoat.ClearCoatNormalTexture);
+		}
+	}
+}
+
+void FGLTFDelayedMaterialTask::GetProxyParameter(const TGLTFProxyMaterialParameterInfo<float>& ParameterInfo, float& OutValue) const
+{
+	ParameterInfo.Get(Material, OutValue, true);
+}
+
+void FGLTFDelayedMaterialTask::GetProxyParameter(const TGLTFProxyMaterialParameterInfo<FLinearColor>& ParameterInfo, FGLTFJsonColor3& OutValue) const
+{
+	FLinearColor Value;
+	if (ParameterInfo.Get(Material, Value, true))
+	{
+		OutValue = FGLTFConverterUtility::ConvertColor3(Value, false);
+	}
+}
+
+void FGLTFDelayedMaterialTask::GetProxyParameter(const TGLTFProxyMaterialParameterInfo<FLinearColor>& ParameterInfo, FGLTFJsonColor4& OutValue) const
+{
+	FLinearColor Value;
+	if (ParameterInfo.Get(Material, Value, true))
+	{
+		OutValue = FGLTFConverterUtility::ConvertColor(Value, false);
+	}
+}
+
+void FGLTFDelayedMaterialTask::GetProxyParameter(const FGLTFProxyMaterialTextureParameterInfo& ParameterInfo, FGLTFJsonTextureInfo& OutValue) const
+{
+	UTexture* Texture;
+	if (!ParameterInfo.Texture.Get(Material, Texture, true) || Texture == nullptr)
+	{
+		return;
+	}
+
+	const bool bSRGB = ParameterInfo == FGLTFProxyMaterialInfo::BaseColor || ParameterInfo == FGLTFProxyMaterialInfo::Emissive;
+	OutValue.Index = Builder.AddUniqueTexture(Texture, bSRGB);
+
+	float UVIndex;
+	if (ParameterInfo.UVIndex.Get(Material, UVIndex, true))
+	{
+		OutValue.TexCoord = FMath::RoundToInt(UVIndex);
+	}
+
+	FLinearColor UVOffset;
+	if (ParameterInfo.UVOffset.Get(Material, UVOffset, true))
+	{
+		OutValue.Transform.Offset.X = UVOffset.R;
+		OutValue.Transform.Offset.Y = UVOffset.G;
+	}
+
+	FLinearColor UVScale;
+	if (ParameterInfo.UVScale.Get(Material, UVScale, true))
+	{
+		OutValue.Transform.Scale.X = UVScale.R;
+		OutValue.Transform.Scale.Y = UVScale.G;
+	}
+
+	float UVRotation;
+	if (ParameterInfo.UVRotation.Get(Material, UVRotation, true))
+	{
+		OutValue.Transform.Rotation = UVRotation;
+	}
+
+	if (!Builder.ExportOptions->bExportTextureTransforms && !OutValue.Transform.IsExactlyDefault())
+	{
+		Builder.LogWarning(FString::Printf(
+			TEXT("Texture coordinates [%d] in %s for material %s are transformed, but texture transform is disabled by export options"),
+			OutValue.TexCoord,
+			*ParameterInfo.Texture.ToString(),
+			*Material->GetName()));
+		OutValue.Transform = {};
+	}
+}
+
+EMaterialShadingModel FGLTFDelayedMaterialTask::GetShadingModel() const
+{
+	const FMaterialShadingModelField Possibilities = Material->GetShadingModels();
+	const int32 PossibilitiesCount = Possibilities.CountShadingModels();
+
+	if (PossibilitiesCount == 0)
+	{
+		const EMaterialShadingModel ShadingModel = MSM_DefaultLit;
+		Builder.LogWarning(FString::Printf(
+			TEXT("No shading model defined for material %s, will export as %s"),
+			*Material->GetName(),
+			*FGLTFNameUtility::GetName(ShadingModel)));
+		return ShadingModel;
+	}
+
+	if (PossibilitiesCount > 1)
+	{
+#if WITH_EDITOR
+		if (!FApp::CanEverRender())
+		{
+			const EMaterialShadingModel ShadingModel = FGLTFMaterialUtility::GetRichestShadingModel(Possibilities);
+			Builder.LogWarning(FString::Printf(
+				TEXT("Can't evaluate shading model expression in material %s because renderer missing, will export as %s"),
+				*Material->GetName(),
+				*FGLTFNameUtility::GetName(ShadingModel)));
+			return ShadingModel;
+		}
+
+		if (Material->IsShadingModelFromMaterialExpression())
+		{
+			const FMaterialShadingModelField Evaluation = FGLTFMaterialUtility::EvaluateShadingModelExpression(Material);
+			const int32 EvaluationCount = Evaluation.CountShadingModels();
+
+			if (EvaluationCount == 0)
+			{
+				const EMaterialShadingModel ShadingModel = FGLTFMaterialUtility::GetRichestShadingModel(Possibilities);
+				Builder.LogWarning(FString::Printf(
+					TEXT("Evaluation of shading model expression in material %s returned none, will export as %s"),
+					*Material->GetName(),
+					*FGLTFNameUtility::GetName(ShadingModel)));
+				return ShadingModel;
+			}
+
+			if (EvaluationCount > 1)
+			{
+				const EMaterialShadingModel ShadingModel = FGLTFMaterialUtility::GetRichestShadingModel(Evaluation);
+				Builder.LogWarning(FString::Printf(
+					TEXT("Evaluation of shading model expression in material %s is inconclusive (%s), will export as %s"),
+					*Material->GetName(),
+					*FGLTFMaterialUtility::ShadingModelsToString(Evaluation),
+					*FGLTFNameUtility::GetName(ShadingModel)));
+				return ShadingModel;
+			}
+
+			return Evaluation.GetFirstShadingModel();
+		}
+
+		// we should never end up here
+#else
+		const EMaterialShadingModel ShadingModel = FGLTFMaterialUtility::GetRichestShadingModel(Possibilities);
+		Builder.LogWarning(FString::Printf(
+			TEXT("Can't evaluate shading model expression in material %s without editor, will export as %s"),
+			*Material->GetName(),
+			*FGLTFNameUtility::GetName(ShadingModel)));
+		return ShadingModel;
+#endif
+	}
+
+	return Possibilities.GetFirstShadingModel();
+}
+
+void FGLTFDelayedMaterialTask::ConvertShadingModel(EGLTFJsonShadingModel& OutShadingModel) const
+{
+	EMaterialShadingModel ShadingModel = GetShadingModel();
+
+	const EBlendMode BlendMode = Material->GetBlendMode();
+	if (ShadingModel == MSM_ClearCoat && BlendMode != BLEND_Opaque && BlendMode != BLEND_Masked)
+	{
+		// NOTE: Unreal seems to disable clear coat when blend mode anything but opaque or masked
+		ShadingModel = MSM_DefaultLit;
+	}
+
+	OutShadingModel = FGLTFConverterUtility::ConvertShadingModel(ShadingModel);
+	if (OutShadingModel == EGLTFJsonShadingModel::None)
+	{
+		OutShadingModel = EGLTFJsonShadingModel::Default;
+
+		Builder.LogWarning(FString::Printf(
+			TEXT("Unsupported shading model (%s) in material %s, will export as %s"),
+			*FGLTFNameUtility::GetName(ShadingModel),
+			*Material->GetName(),
+			*FGLTFNameUtility::GetName(MSM_DefaultLit)));
+		return;
+	}
+
+	if (OutShadingModel == EGLTFJsonShadingModel::Unlit && !Builder.ExportOptions->bExportUnlitMaterials)
+	{
+		OutShadingModel = EGLTFJsonShadingModel::Default;
+
+		Builder.LogWarning(FString::Printf(
+			TEXT("Shading model (%s) in material %s disabled by export options, will export as %s"),
+			*FGLTFNameUtility::GetName(ShadingModel),
+			*Material->GetName(),
+			*FGLTFNameUtility::GetName(MSM_DefaultLit)));
+		return;
+	}
+
+	if (OutShadingModel == EGLTFJsonShadingModel::ClearCoat && !Builder.ExportOptions->bExportClearCoatMaterials)
+	{
+		OutShadingModel = EGLTFJsonShadingModel::Default;
+
+		Builder.LogWarning(FString::Printf(
+			TEXT("Shading model (%s) in material %s disabled by export options, will export as %s"),
+			*FGLTFNameUtility::GetName(ShadingModel),
+			*Material->GetName(),
+			*FGLTFNameUtility::GetName(MSM_DefaultLit)));
+	}
+}
+
+void FGLTFDelayedMaterialTask::ConvertAlphaMode(EGLTFJsonAlphaMode& OutAlphaMode, EGLTFJsonBlendMode& OutBlendMode) const
+{
+	EBlendMode BlendMode = Material->GetBlendMode();
+	if (const UMaterialInstance* MaterialInstance = Cast<UMaterialInstance>(Material))
+	{
+		FGLTFProxyMaterialUtilities::GetBlendMode(MaterialInstance, BlendMode); // TODO: remove ugly hack for dynamic instances
+	}
+
+	OutAlphaMode = FGLTFConverterUtility::ConvertAlphaMode(BlendMode);
+	if (OutAlphaMode == EGLTFJsonAlphaMode::None)
+	{
+		OutAlphaMode = EGLTFJsonAlphaMode::Blend;
+		OutBlendMode = EGLTFJsonBlendMode::None;
+
+		Builder.LogWarning(FString::Printf(
+			TEXT("Unsupported blend mode (%s) in material %s, will export as %s"),
+			*FGLTFNameUtility::GetName(BlendMode),
+			*Material->GetName(),
+			*FGLTFNameUtility::GetName(BLEND_Translucent)));
+		return;
+	}
+
+	if (OutAlphaMode == EGLTFJsonAlphaMode::Blend)
+	{
+		OutBlendMode = FGLTFConverterUtility::ConvertBlendMode(BlendMode);
+		if (OutBlendMode != EGLTFJsonBlendMode::None && !Builder.ExportOptions->bExportExtraBlendModes)
+		{
+			OutBlendMode = EGLTFJsonBlendMode::None;
+
+			Builder.LogWarning(FString::Printf(
+				TEXT("Extra blend mode (%s) in material %s disabled by export options, will export as %s"),
+				*FGLTFNameUtility::GetName(BlendMode),
+				*Material->GetName(),
+				*FGLTFNameUtility::GetName(BLEND_Translucent)));
+		}
+	}
+}
+
+#if WITH_EDITOR
+
+bool FGLTFDelayedMaterialTask::TryGetBaseColorAndOpacity(FGLTFJsonPBRMetallicRoughness& OutPBRParams, const FGLTFMaterialPropertyEx& BaseColorProperty, const FGLTFMaterialPropertyEx& OpacityProperty)
+{
+	const bool bIsBaseColorConstant = TryGetConstantColor(OutPBRParams.BaseColorFactor, BaseColorProperty);
+	const bool bIsOpacityConstant = TryGetConstantScalar(OutPBRParams.BaseColorFactor.A, OpacityProperty);
+
+	if (bIsBaseColorConstant && bIsOpacityConstant)
+	{
+		return true;
+	}
+
+	// NOTE: since we always bake the properties (for now) when at least property is non-const, we need
+	// to reset the constant factors to their defaults. Otherwise the baked value of a constant property
+	// would be scaled with the factor, i.e a double scaling.
+	OutPBRParams.BaseColorFactor = { 1.0f, 1.0f, 1.0f, 1.0f };
+
+	const UTexture2D* BaseColorTexture;
+	const UTexture2D* OpacityTexture;
+	int32 BaseColorTexCoord;
+	int32 OpacityTexCoord;
+	FGLTFJsonTextureTransform BaseColorTransform;
+	FGLTFJsonTextureTransform OpacityTransform;
+
+	const bool bHasBaseColorSourceTexture = TryGetSourceTexture(BaseColorTexture, BaseColorTexCoord, BaseColorTransform, BaseColorProperty, BaseColorInputMasks);
+	const bool bHasOpacitySourceTexture = TryGetSourceTexture(OpacityTexture, OpacityTexCoord, OpacityTransform, OpacityProperty, OpacityInputMasks);
+
+	// Detect the "happy path" where both inputs share the same texture and are correctly masked.
+	if (bHasBaseColorSourceTexture &&
+		bHasOpacitySourceTexture &&
+		BaseColorTexture == OpacityTexture &&
+		BaseColorTexCoord == OpacityTexCoord &&
+		BaseColorTransform.IsExactlyEqual(OpacityTransform))
+	{
+		OutPBRParams.BaseColorTexture.Index = Builder.AddUniqueTexture(BaseColorTexture, true);
+		OutPBRParams.BaseColorTexture.TexCoord = BaseColorTexCoord;
+		OutPBRParams.BaseColorTexture.Transform = BaseColorTransform;
+		return true;
+	}
+
+	if (Builder.ExportOptions->BakeMaterialInputs == EGLTFMaterialBakeMode::Disabled)
+	{
+		Builder.LogWarning(FString::Printf(
+			TEXT("%s and %s for material %s needs to bake, but material baking is disabled by export options"),
+			*BaseColorProperty.ToString(),
+			*OpacityProperty.ToString(),
+			*Material->GetName()));
+		return false;
+	}
+
+	const FIntPoint TextureSize = Builder.GetBakeSizeForMaterialProperty(Material, GetPropertyGroup(BaseColorProperty));
+
+	const TextureAddress TextureAddress = Builder.GetBakeTilingForMaterialProperty(Material, GetPropertyGroup(BaseColorProperty));
+	const EGLTFJsonTextureWrap TextureWrapS = FGLTFConverterUtility::ConvertWrap(TextureAddress);
+	const EGLTFJsonTextureWrap TextureWrapT = FGLTFConverterUtility::ConvertWrap(TextureAddress);
+
+	const TextureFilter TextureFilter = Builder.GetBakeFilterForMaterialProperty(Material, GetPropertyGroup(BaseColorProperty));
+	const EGLTFJsonTextureFilter TextureMinFilter = FGLTFConverterUtility::ConvertMinFilter(TextureFilter);
+	const EGLTFJsonTextureFilter TextureMagFilter = FGLTFConverterUtility::ConvertMagFilter(TextureFilter);
+
+	const FGLTFPropertyBakeOutput BaseColorBakeOutput = BakeMaterialProperty(BaseColorProperty, BaseColorTexCoord, TextureSize, false);
+	const FGLTFPropertyBakeOutput OpacityBakeOutput = BakeMaterialProperty(OpacityProperty, OpacityTexCoord, TextureSize, false);
+	const float BaseColorScale = BaseColorProperty == MP_EmissiveColor ? BaseColorBakeOutput.EmissiveScale : 1;
+
+	// Detect when both baked properties are constants, which means we can avoid exporting a texture
+	if (BaseColorBakeOutput.bIsConstant && OpacityBakeOutput.bIsConstant)
+	{
+		FLinearColor BaseColorFactor(BaseColorBakeOutput.ConstantValue * BaseColorScale);
+		BaseColorFactor.A = OpacityBakeOutput.ConstantValue.A;
+
+		OutPBRParams.BaseColorFactor = FGLTFConverterUtility::ConvertColor(BaseColorFactor, Builder.ExportOptions->bStrictCompliance);
+		return true;
+	}
+
+	if (Builder.ExportOptions->TextureImageFormat == EGLTFTextureImageFormat::None)
+	{
+		return true;
+	}
+
+	int32 TexCoord;
+	if (BaseColorBakeOutput.bIsConstant)
+	{
+		TexCoord = OpacityTexCoord;
+	}
+	else if (OpacityBakeOutput.bIsConstant)
+	{
+		TexCoord = BaseColorTexCoord;
+	}
+	else if (BaseColorTexCoord == OpacityTexCoord)
+	{
+		TexCoord = BaseColorTexCoord;
+	}
+	else
+	{
+		// TODO: report error (texture coordinate conflict)
+		return false;
+	}
+
+	TGLTFSharedArray<FColor> CombinedPixels;
+	CombinePixels(
+		*BaseColorBakeOutput.Pixels,
+		*OpacityBakeOutput.Pixels,
+		*CombinedPixels,
+		[](const FColor& BaseColor, const FColor& Opacity) -> FColor
+		{
+			return FColor(BaseColor.R, BaseColor.G, BaseColor.B, Opacity.R);
+		});
+
+	FGLTFJsonTexture* Texture = FGLTFMaterialUtility::AddTexture(
+		Builder,
+		CombinedPixels,
+		TextureSize,
+		false,
+		false,
+		GetBakedTextureName(TEXT("BaseColor")),
+		TextureMinFilter,
+		TextureMagFilter,
+		TextureWrapS,
+		TextureWrapT);
+
+	OutPBRParams.BaseColorTexture.TexCoord = TexCoord;
+	OutPBRParams.BaseColorTexture.Index = Texture;
+	OutPBRParams.BaseColorFactor = FGLTFConverterUtility::ConvertColor({ BaseColorScale, BaseColorScale, BaseColorScale }, Builder.ExportOptions->bStrictCompliance);
+
+	return true;
+}
+
+bool FGLTFDelayedMaterialTask::TryGetMetallicAndRoughness(FGLTFJsonPBRMetallicRoughness& OutPBRParams, const FGLTFMaterialPropertyEx& MetallicProperty, const FGLTFMaterialPropertyEx& RoughnessProperty)
+{
+	const bool bIsMetallicConstant = TryGetConstantScalar(OutPBRParams.MetallicFactor, MetallicProperty);
+	const bool bIsRoughnessConstant = TryGetConstantScalar(OutPBRParams.RoughnessFactor, RoughnessProperty);
+
+	if (bIsMetallicConstant && bIsRoughnessConstant)
+	{
+		return true;
+	}
+
+	// NOTE: since we always bake the properties (for now) when at least one property is non-const, we need
+	// to reset the constant factors to their defaults. Otherwise the baked value of a constant property
+	// would be scaled with the factor, i.e a double scaling.
+	OutPBRParams.MetallicFactor = 1.0f;
+	OutPBRParams.RoughnessFactor = 1.0f;
+
+	const UTexture2D* MetallicTexture;
+	const UTexture2D* RoughnessTexture;
+	int32 MetallicTexCoord;
+	int32 RoughnessTexCoord;
+	FGLTFJsonTextureTransform MetallicTransform;
+	FGLTFJsonTextureTransform RoughnessTransform;
+
+	const bool bHasMetallicSourceTexture = TryGetSourceTexture(MetallicTexture, MetallicTexCoord, MetallicTransform, MetallicProperty, MetallicInputMasks);
+	const bool bHasRoughnessSourceTexture = TryGetSourceTexture(RoughnessTexture, RoughnessTexCoord, RoughnessTransform, RoughnessProperty, RoughnessInputMasks);
+
+	// Detect the "happy path" where both inputs share the same texture and are correctly masked.
+	if (bHasMetallicSourceTexture &&
+		bHasRoughnessSourceTexture &&
+		MetallicTexture == RoughnessTexture &&
+		MetallicTexCoord == RoughnessTexCoord &&
+		MetallicTransform.IsExactlyEqual(RoughnessTransform))
+	{
+		OutPBRParams.MetallicRoughnessTexture.Index = Builder.AddUniqueTexture(MetallicTexture, false);
+		OutPBRParams.MetallicRoughnessTexture.TexCoord = MetallicTexCoord;
+		OutPBRParams.MetallicRoughnessTexture.Transform = MetallicTransform;
+		return true;
+	}
+
+	if (Builder.ExportOptions->BakeMaterialInputs == EGLTFMaterialBakeMode::Disabled)
+	{
+		Builder.LogWarning(FString::Printf(
+			TEXT("%s and %s for material %s needs to bake, but material baking is disabled by export options"),
+			*MetallicProperty.ToString(),
+			*RoughnessProperty.ToString(),
+			*Material->GetName()));
+		return false;
+	}
+
+	// TODO: add support for calculating the ideal resolution to use for baking based on connected (texture) nodes
+	const FIntPoint TextureSize = Builder.GetBakeSizeForMaterialProperty(Material, GetPropertyGroup(MetallicProperty));
+
+	const TextureAddress TextureAddress = Builder.GetBakeTilingForMaterialProperty(Material, GetPropertyGroup(MetallicProperty));
+	const EGLTFJsonTextureWrap TextureWrapS = FGLTFConverterUtility::ConvertWrap(TextureAddress);
+	const EGLTFJsonTextureWrap TextureWrapT = FGLTFConverterUtility::ConvertWrap(TextureAddress);
+
+	const TextureFilter TextureFilter = Builder.GetBakeFilterForMaterialProperty(Material, GetPropertyGroup(MetallicProperty));
+	const EGLTFJsonTextureFilter TextureMinFilter = FGLTFConverterUtility::ConvertMinFilter(TextureFilter);
+	const EGLTFJsonTextureFilter TextureMagFilter = FGLTFConverterUtility::ConvertMagFilter(TextureFilter);
+
+	FGLTFPropertyBakeOutput MetallicBakeOutput = BakeMaterialProperty(MetallicProperty, MetallicTexCoord, TextureSize, false);
+	FGLTFPropertyBakeOutput RoughnessBakeOutput = BakeMaterialProperty(RoughnessProperty, RoughnessTexCoord, TextureSize, false);
+
+	// Detect when both baked properties are constants, which means we can use factors and avoid exporting a texture
+	if (MetallicBakeOutput.bIsConstant && RoughnessBakeOutput.bIsConstant)
+	{
+		OutPBRParams.MetallicFactor = MetallicBakeOutput.ConstantValue.R;
+		OutPBRParams.RoughnessFactor = RoughnessBakeOutput.ConstantValue.R;
+		return true;
+	}
+
+	if (Builder.ExportOptions->TextureImageFormat == EGLTFTextureImageFormat::None)
+	{
+		return true;
+	}
+
+	int32 TexCoord;
+	if (MetallicBakeOutput.bIsConstant)
+	{
+		TexCoord = RoughnessTexCoord;
+	}
+	else if (MetallicBakeOutput.bIsConstant)
+	{
+		TexCoord = MetallicTexCoord;
+	}
+	else if (MetallicTexCoord == RoughnessTexCoord)
+	{
+		TexCoord = MetallicTexCoord;
+	}
+	else
+	{
+		// TODO: report error (texture coordinate conflict)
+		return false;
+	}
+
+	TGLTFSharedArray<FColor> CombinedPixels;
+	CombinePixels(
+		*MetallicBakeOutput.Pixels,
+		*RoughnessBakeOutput.Pixels,
+		*CombinedPixels,
+		[](const FColor& Metallic, const FColor& Roughness) -> FColor
+		{
+			return FColor(0, Roughness.R, Metallic.R);
+		});
+
+	FGLTFJsonTexture* Texture = FGLTFMaterialUtility::AddTexture(
+		Builder,
+		CombinedPixels,
+		TextureSize,
+		true, // NOTE: we can ignore alpha in everything but TryGetBaseColorAndOpacity
+		false,
+		GetBakedTextureName(TEXT("MetallicRoughness")),
+		TextureMinFilter,
+		TextureMagFilter,
+		TextureWrapS,
+		TextureWrapT);
+
+	OutPBRParams.MetallicRoughnessTexture.TexCoord = TexCoord;
+	OutPBRParams.MetallicRoughnessTexture.Index = Texture;
+
+	return true;
+}
+
+bool FGLTFDelayedMaterialTask::TryGetClearCoatRoughness(FGLTFJsonClearCoatExtension& OutExtParams, const FGLTFMaterialPropertyEx& IntensityProperty, const FGLTFMaterialPropertyEx& RoughnessProperty)
+{
+	const bool bIsIntensityConstant = TryGetConstantScalar(OutExtParams.ClearCoatFactor, IntensityProperty);
+	const bool bIsRoughnessConstant = TryGetConstantScalar(OutExtParams.ClearCoatRoughnessFactor, RoughnessProperty);
+
+	if (bIsIntensityConstant && bIsRoughnessConstant)
+	{
+		return true;
+	}
+
+	// NOTE: since we always bake the properties (for now) when at least one property is non-const, we need
+	// to reset the constant factors to their defaults. Otherwise the baked value of a constant property
+	// would be scaled with the factor, i.e a double scaling.
+	OutExtParams.ClearCoatFactor = 1.0f;
+	OutExtParams.ClearCoatRoughnessFactor = 1.0f;
+
+	const UTexture2D* IntensityTexture;
+	const UTexture2D* RoughnessTexture;
+	int32 IntensityTexCoord;
+	int32 RoughnessTexCoord;
+	FGLTFJsonTextureTransform IntensityTransform;
+	FGLTFJsonTextureTransform RoughnessTransform;
+
+	const bool bHasIntensitySourceTexture = TryGetSourceTexture(IntensityTexture, IntensityTexCoord, IntensityTransform, IntensityProperty, ClearCoatInputMasks);
+	const bool bHasRoughnessSourceTexture = TryGetSourceTexture(RoughnessTexture, RoughnessTexCoord, RoughnessTransform, RoughnessProperty, ClearCoatRoughnessInputMasks);
+
+	// Detect the "happy path" where both inputs share the same texture and are correctly masked.
+	if (bHasIntensitySourceTexture &&
+		bHasRoughnessSourceTexture &&
+		IntensityTexture == RoughnessTexture &&
+		IntensityTexCoord == RoughnessTexCoord &&
+		IntensityTransform.IsExactlyEqual(RoughnessTransform))
+	{
+		FGLTFJsonTexture* Texture = Builder.AddUniqueTexture(IntensityTexture, false);
+		OutExtParams.ClearCoatTexture.Index = Texture;
+		OutExtParams.ClearCoatTexture.TexCoord = IntensityTexCoord;
+		OutExtParams.ClearCoatRoughnessTexture.Index = Texture;
+		OutExtParams.ClearCoatRoughnessTexture.TexCoord = IntensityTexCoord;
+		OutExtParams.ClearCoatRoughnessTexture.Transform = IntensityTransform;
+		return true;
+	}
+
+	if (Builder.ExportOptions->BakeMaterialInputs == EGLTFMaterialBakeMode::Disabled)
+	{
+		Builder.LogWarning(FString::Printf(
+			TEXT("%s and %s for material %s needs to bake, but material baking is disabled by export options"),
+			*IntensityProperty.ToString(),
+			*RoughnessProperty.ToString(),
+			*Material->GetName()));
+		return false;
+	}
+
+	const FIntPoint TextureSize = Builder.GetBakeSizeForMaterialProperty(Material, GetPropertyGroup(IntensityProperty));
+
+	const TextureAddress TextureAddress = Builder.GetBakeTilingForMaterialProperty(Material,GetPropertyGroup(IntensityProperty));
+	const EGLTFJsonTextureWrap TextureWrapS = FGLTFConverterUtility::ConvertWrap(TextureAddress);
+	const EGLTFJsonTextureWrap TextureWrapT = FGLTFConverterUtility::ConvertWrap(TextureAddress);
+
+	const TextureFilter TextureFilter = Builder.GetBakeFilterForMaterialProperty(Material, GetPropertyGroup(IntensityProperty));
+	const EGLTFJsonTextureFilter TextureMinFilter = FGLTFConverterUtility::ConvertMinFilter(TextureFilter);
+	const EGLTFJsonTextureFilter TextureMagFilter = FGLTFConverterUtility::ConvertMagFilter(TextureFilter);
+
+	const FGLTFPropertyBakeOutput IntensityBakeOutput = BakeMaterialProperty(IntensityProperty, IntensityTexCoord, TextureSize, false);
+	const FGLTFPropertyBakeOutput RoughnessBakeOutput = BakeMaterialProperty(RoughnessProperty, RoughnessTexCoord, TextureSize, false);
+
+	// Detect when both baked properties are constants, which means we can use factors and avoid exporting a texture
+	if (IntensityBakeOutput.bIsConstant && RoughnessBakeOutput.bIsConstant)
+	{
+		OutExtParams.ClearCoatFactor = IntensityBakeOutput.ConstantValue.R;
+		OutExtParams.ClearCoatRoughnessFactor = RoughnessBakeOutput.ConstantValue.R;
+		return true;
+	}
+
+	if (Builder.ExportOptions->TextureImageFormat == EGLTFTextureImageFormat::None)
+	{
+		return true;
+	}
+
+	int32 TexCoord;
+	if (IntensityBakeOutput.bIsConstant)
+	{
+		TexCoord = RoughnessTexCoord;
+	}
+	else if (IntensityBakeOutput.bIsConstant)
+	{
+		TexCoord = IntensityTexCoord;
+	}
+	else if (IntensityTexCoord == RoughnessTexCoord)
+	{
+		TexCoord = IntensityTexCoord;
+	}
+	else
+	{
+		// TODO: report error (texture coordinate conflict)
+		return false;
+	}
+
+	TGLTFSharedArray<FColor> CombinedPixels;
+	CombinePixels(
+		*IntensityBakeOutput.Pixels,
+		*RoughnessBakeOutput.Pixels,
+		*CombinedPixels,
+		[](const FColor& Intensity, const FColor& Roughness) -> FColor
+		{
+			return FColor(Intensity.R, Roughness.R, 0);
+		});
+
+	FGLTFJsonTexture* Texture = FGLTFMaterialUtility::AddTexture(
+		Builder,
+		CombinedPixels,
+		TextureSize,
+		true, // NOTE: we can ignore alpha in everything but TryGetBaseColorAndOpacity
+		false,
+		GetBakedTextureName(TEXT("ClearCoatRoughness")),
+		TextureMinFilter,
+		TextureMagFilter,
+		TextureWrapS,
+		TextureWrapT);
+
+	OutExtParams.ClearCoatTexture.Index = Texture;
+	OutExtParams.ClearCoatTexture.TexCoord = TexCoord;
+	OutExtParams.ClearCoatRoughnessTexture.Index = Texture;
+	OutExtParams.ClearCoatRoughnessTexture.TexCoord = TexCoord;
+
+	return true;
+}
+
+bool FGLTFDelayedMaterialTask::TryGetEmissive(FGLTFJsonMaterial& OutMaterial, const FGLTFMaterialPropertyEx& EmissiveProperty)
+{
+	// TODO: right now we allow EmissiveFactor to be > 1.0 to support very bright emission, although it's not valid according to the glTF standard.
+	// We may want to change this behaviour and store factors above 1.0 using a custom extension instead.
+
+	if (TryGetConstantColor(OutMaterial.EmissiveFactor, MP_EmissiveColor))
+	{
+		return true;
+	}
+
+	if (TryGetSourceTexture(OutMaterial.EmissiveTexture, EmissiveProperty, DefaultColorInputMasks))
+	{
+		OutMaterial.EmissiveFactor = FGLTFJsonColor3::White;	// make sure texture is not multiplied with black
+		return true;
+	}
+
+	if (Builder.ExportOptions->BakeMaterialInputs == EGLTFMaterialBakeMode::Disabled)
+	{
+		Builder.LogWarning(FString::Printf(
+			TEXT("%s for material %s needs to bake, but material baking is disabled by export options"),
+			*EmissiveProperty.ToString(),
+			*Material->GetName()));
+		return false;
+	}
+
+	FGLTFPropertyBakeOutput PropertyBakeOutput = BakeMaterialProperty(EmissiveProperty, OutMaterial.EmissiveTexture.TexCoord);
+	const float EmissiveScale = PropertyBakeOutput.EmissiveScale;
+
+	if (PropertyBakeOutput.bIsConstant)
+	{
+		const FLinearColor EmissiveColor = PropertyBakeOutput.ConstantValue;
+		OutMaterial.EmissiveFactor = FGLTFConverterUtility::ConvertColor3(EmissiveColor * EmissiveScale, Builder.ExportOptions->bStrictCompliance);
+	}
+	else
+	{
+		if (Builder.ExportOptions->TextureImageFormat == EGLTFTextureImageFormat::None)
+		{
+			OutMaterial.EmissiveTexture.Index = nullptr;
+			return true;
+		}
+
+		if (!StoreBakedPropertyTexture(OutMaterial.EmissiveTexture, PropertyBakeOutput, TEXT("Emissive")))
+		{
+			return false;
+		}
+
+		OutMaterial.EmissiveFactor = { EmissiveScale, EmissiveScale, EmissiveScale };
+	}
+
+	return true;
+}
+
+bool FGLTFDelayedMaterialTask::IsPropertyNonDefault(const FGLTFMaterialPropertyEx& Property) const
+{
+	const bool bUseMaterialAttributes = Material->GetMaterial()->bUseMaterialAttributes;
+	if (bUseMaterialAttributes)
+	{
+		// TODO: check if attribute property connected, i.e. Material->GetMaterial()->MaterialAttributes.IsConnected(Property)
+		return true;
+	}
+
+	const FExpressionInput* MaterialInput = FGLTFMaterialUtility::GetInputForProperty(Material, Property);
+	if (MaterialInput == nullptr)
+	{
+		// TODO: report error
+		return false;
+	}
+
+	const UMaterialExpression* Expression = MaterialInput->Expression;
+	if (Expression == nullptr)
+	{
+		return false;
+	}
+
+	return true;
+}
+
+bool FGLTFDelayedMaterialTask::TryGetConstantColor(FGLTFJsonColor3& OutValue, const FGLTFMaterialPropertyEx& Property) const
+{
+	FLinearColor Value;
+	if (TryGetConstantColor(Value, Property))
+	{
+		OutValue = FGLTFConverterUtility::ConvertColor3(Value, Builder.ExportOptions->bStrictCompliance);
+		return true;
+	}
+
+	return false;
+}
+
+bool FGLTFDelayedMaterialTask::TryGetConstantColor(FGLTFJsonColor4& OutValue, const FGLTFMaterialPropertyEx& Property) const
+{
+	FLinearColor Value;
+	if (TryGetConstantColor(Value, Property))
+	{
+		OutValue = FGLTFConverterUtility::ConvertColor(Value, Builder.ExportOptions->bStrictCompliance);
+		return true;
+	}
+
+	return false;
+}
+
+bool FGLTFDelayedMaterialTask::TryGetConstantColor(FLinearColor& OutValue, const FGLTFMaterialPropertyEx& Property) const
+{
+	const bool bUseMaterialAttributes = Material->GetMaterial()->bUseMaterialAttributes;
+	if (bUseMaterialAttributes)
+	{
+		// TODO: check if attribute property connected, i.e. Material->GetMaterial()->MaterialAttributes.IsConnected(Property)
+		return false;
+	}
+
+	const FMaterialInput<FColor>* MaterialInput = FGLTFMaterialUtility::GetInputForProperty<FColor>(Material, Property);
+	if (MaterialInput == nullptr)
+	{
+		// TODO: report error
+		return false;
+	}
+
+	if (MaterialInput->UseConstant)
+	{
+		OutValue = { MaterialInput->Constant };
+		return true;
+	}
+
+	const UMaterialExpression* Expression = MaterialInput->Expression;
+	if (Expression == nullptr)
+	{
+		OutValue = FLinearColor(FGLTFMaterialUtility::GetPropertyDefaultValue(Property));
+		return true;
+	}
+
+	if (const UMaterialExpressionVectorParameter* VectorParameter = ExactCast<UMaterialExpressionVectorParameter>(Expression))
+	{
+		FLinearColor Value = VectorParameter->DefaultValue;
+
+		const UMaterialInstance* MaterialInstance = Cast<UMaterialInstance>(Material);
+		if (MaterialInstance != nullptr)
+		{
+			const FHashedMaterialParameterInfo ParameterInfo(VectorParameter->GetParameterName());
+			if (!MaterialInstance->GetVectorParameterValue(ParameterInfo, Value))
+			{
+				// TODO: how to handle this?
+			}
+		}
+
+		const uint32 MaskComponentCount = FGLTFMaterialUtility::GetMaskComponentCount(*MaterialInput);
+
+		if (MaskComponentCount > 0)
+		{
+			const FLinearColor Mask = FGLTFMaterialUtility::GetMask(*MaterialInput);
+
+			Value *= Mask;
+
+			if (MaskComponentCount == 1)
+			{
+				const float ComponentValue = Value.R + Value.G + Value.B + Value.A;
+				Value = { ComponentValue, ComponentValue, ComponentValue, ComponentValue };
+			}
+		}
+
+		OutValue = Value;
+		return true;
+	}
+
+	if (const UMaterialExpressionScalarParameter* ScalarParameter = ExactCast<UMaterialExpressionScalarParameter>(Expression))
+	{
+		float Value = ScalarParameter->DefaultValue;
+
+		const UMaterialInstance* MaterialInstance = Cast<UMaterialInstance>(Material);
+		if (MaterialInstance != nullptr)
+		{
+			const FHashedMaterialParameterInfo ParameterInfo(ScalarParameter->GetParameterName());
+			if (!MaterialInstance->GetScalarParameterValue(ParameterInfo, Value))
+			{
+				// TODO: how to handle this?
+			}
+		}
+
+		OutValue = { Value, Value, Value, Value };
+		return true;
+	}
+
+	if (const UMaterialExpressionConstant4Vector* Constant4Vector = ExactCast<UMaterialExpressionConstant4Vector>(Expression))
+	{
+		OutValue = Constant4Vector->Constant;
+		return true;
+	}
+
+	if (const UMaterialExpressionConstant3Vector* Constant3Vector = ExactCast<UMaterialExpressionConstant3Vector>(Expression))
+	{
+		OutValue = Constant3Vector->Constant;
+		return true;
+	}
+
+	if (const UMaterialExpressionConstant2Vector* Constant2Vector = ExactCast<UMaterialExpressionConstant2Vector>(Expression))
+	{
+		OutValue = { Constant2Vector->R, Constant2Vector->G, 0, 0 };
+		return true;
+	}
+
+	if (const UMaterialExpressionConstant* Constant = ExactCast<UMaterialExpressionConstant>(Expression))
+	{
+		OutValue = { Constant->R, Constant->R, Constant->R, Constant->R };
+		return true;
+	}
+
+	return false;
+}
+
+bool FGLTFDelayedMaterialTask::TryGetConstantScalar(float& OutValue, const FGLTFMaterialPropertyEx& Property) const
+{
+	const bool bUseMaterialAttributes = Material->GetMaterial()->bUseMaterialAttributes;
+	if (bUseMaterialAttributes)
+	{
+		// TODO: check if attribute property connected, i.e. Material->GetMaterial()->MaterialAttributes.IsConnected(Property)
+		return false;
+	}
+
+	const FMaterialInput<float>* MaterialInput = FGLTFMaterialUtility::GetInputForProperty<float>(Material, Property);
+	if (MaterialInput == nullptr)
+	{
+		// TODO: report error
+		return false;
+	}
+
+	if (MaterialInput->UseConstant)
+	{
+		OutValue = MaterialInput->Constant;
+		return true;
+	}
+
+	const UMaterialExpression* Expression = MaterialInput->Expression;
+	if (Expression == nullptr)
+	{
+		OutValue = FGLTFMaterialUtility::GetPropertyDefaultValue(Property).X;
+		return true;
+	}
+
+	if (const UMaterialExpressionVectorParameter* VectorParameter = ExactCast<UMaterialExpressionVectorParameter>(Expression))
+	{
+		FLinearColor Value = VectorParameter->DefaultValue;
+
+		const UMaterialInstance* MaterialInstance = Cast<UMaterialInstance>(Material);
+		if (MaterialInstance != nullptr)
+		{
+			const FHashedMaterialParameterInfo ParameterInfo(VectorParameter->GetParameterName());
+			if (!MaterialInstance->GetVectorParameterValue(ParameterInfo, Value))
+			{
+				// TODO: how to handle this?
+			}
+		}
+
+		const uint32 MaskComponentCount = FGLTFMaterialUtility::GetMaskComponentCount(*MaterialInput);
+
+		if (MaskComponentCount > 0)
+		{
+			const FLinearColor Mask = FGLTFMaterialUtility::GetMask(*MaterialInput);
+			Value *= Mask;
+		}
+
+		// TODO: is this a correct assumption, that the max component should be used as value?
+		OutValue = Value.GetMax();
+		return true;
+	}
+
+	if (const UMaterialExpressionScalarParameter* ScalarParameter = ExactCast<UMaterialExpressionScalarParameter>(Expression))
+	{
+		float Value = ScalarParameter->DefaultValue;
+
+		const UMaterialInstance* MaterialInstance = Cast<UMaterialInstance>(Material);
+		if (MaterialInstance != nullptr)
+		{
+			const FHashedMaterialParameterInfo ParameterInfo(ScalarParameter->GetParameterName());
+			if (!MaterialInstance->GetScalarParameterValue(ParameterInfo, Value))
+			{
+				// TODO: how to handle this?
+			}
+		}
+
+		OutValue = Value;
+		return true;
+	}
+
+	if (const UMaterialExpressionConstant4Vector* Constant4Vector = ExactCast<UMaterialExpressionConstant4Vector>(Expression))
+	{
+		OutValue = Constant4Vector->Constant.R;
+		return true;
+	}
+
+	if (const UMaterialExpressionConstant3Vector* Constant3Vector = ExactCast<UMaterialExpressionConstant3Vector>(Expression))
+	{
+		OutValue = Constant3Vector->Constant.R;
+		return true;
+	}
+
+	if (const UMaterialExpressionConstant2Vector* Constant2Vector = ExactCast<UMaterialExpressionConstant2Vector>(Expression))
+	{
+		OutValue = Constant2Vector->R;
+		return true;
+	}
+
+	if (const UMaterialExpressionConstant* Constant = ExactCast<UMaterialExpressionConstant>(Expression))
+	{
+		OutValue = Constant->R;
+		return true;
+	}
+
+	return false;
+}
+
+bool FGLTFDelayedMaterialTask::TryGetSourceTexture(FGLTFJsonTextureInfo& OutTexInfo, const FGLTFMaterialPropertyEx& Property, const TArray<FLinearColor>& AllowedMasks) const
+{
+	const UTexture2D* Texture;
+	int32 TexCoord;
+	FGLTFJsonTextureTransform Transform;
+
+	if (TryGetSourceTexture(Texture, TexCoord, Transform, Property, AllowedMasks))
+	{
+		OutTexInfo.Index = Builder.AddUniqueTexture(Texture, FGLTFMaterialUtility::IsSRGB(Property));
+		OutTexInfo.TexCoord = TexCoord;
+		OutTexInfo.Transform = Transform;
+		return true;
+	}
+
+	return false;
+}
+
+bool FGLTFDelayedMaterialTask::TryGetSourceTexture(const UTexture2D*& OutTexture, int32& OutTexCoord, FGLTFJsonTextureTransform& OutTransform, const FGLTFMaterialPropertyEx& Property, const TArray<FLinearColor>& AllowedMasks) const
+{
+	if (Builder.ExportOptions->TextureImageFormat == EGLTFTextureImageFormat::None)
+	{
+		return false;
+	}
+
+	const FExpressionInput* MaterialInput = FGLTFMaterialUtility::GetInputForProperty(Material, Property);
+	if (MaterialInput == nullptr)
+	{
+		// TODO: report error
+		return false;
+	}
+
+	const UMaterialExpression* Expression = MaterialInput->Expression;
+	if (Expression == nullptr)
+	{
+		return false;
+	}
+
+	const FLinearColor InputMask = FGLTFMaterialUtility::GetMask(*MaterialInput);
+	if (AllowedMasks.Num() > 0 && !AllowedMasks.Contains(InputMask))
+	{
+		return false;
+	}
+
+	// TODO: add support or warning for texture sampler settings that override texture asset addressing (i.e. wrap, clamp etc)?
+
+	if (const UMaterialExpressionTextureSampleParameter2D* TextureParameter = ExactCast<UMaterialExpressionTextureSampleParameter2D>(Expression))
+	{
+		// TODO: handle non-default TextureParameter->SamplerType and TextureParameter->SamplerSource
+
+		UTexture* ParameterValue = TextureParameter->Texture;
+
+		const UMaterialInstance* MaterialInstance = Cast<UMaterialInstance>(Material);
+		if (MaterialInstance != nullptr)
+		{
+			const FHashedMaterialParameterInfo ParameterInfo(TextureParameter->GetParameterName());
+			if (!MaterialInstance->GetTextureParameterValue(ParameterInfo, ParameterValue))
+			{
+				// TODO: how to handle this?
+			}
+		}
+
+		// TODO: add support for UTextureRenderTarget2D?
+		OutTexture = Cast<UTexture2D>(ParameterValue);
+
+		if (OutTexture == nullptr)
+		{
+			if (ParameterValue == nullptr)
+			{
+				// TODO: report error (no texture parameter assigned)
+			}
+			else
+			{
+				// TODO: report error (incorrect texture type)
+			}
+			return false;
+		}
+
+		if (!FGLTFMaterialUtility::TryGetTextureCoordinateIndex(TextureParameter, OutTexCoord, OutTransform))
+		{
+			// TODO: report error (failed to identify texture coordinate index)
+			return false;
+		}
+
+		if (!Builder.ExportOptions->bExportTextureTransforms && !OutTransform.IsExactlyDefault())
+		{
+			Builder.LogWarning(FString::Printf(
+				TEXT("Texture coordinates [%d] in %s for material %s are transformed, but texture transform is disabled by export options"),
+				OutTexCoord,
+				*Property.ToString(),
+				*Material->GetName()));
+			OutTransform = {};
+		}
+
+		return true;
+	}
+
+	if (const UMaterialExpressionTextureSample* TextureSampler = ExactCast<UMaterialExpressionTextureSample>(Expression))
+	{
+		// TODO: handle non-default TextureSampler->SamplerType and TextureSampler->SamplerSource
+
+		// TODO: add support for texture object input expression
+
+		// TODO: add support for UTextureRenderTarget2D?
+		OutTexture = Cast<UTexture2D>(TextureSampler->Texture);
+
+		if (OutTexture == nullptr)
+		{
+			if (TextureSampler->Texture == nullptr)
+			{
+				// TODO: report error (no texture sample assigned)
+			}
+			else
+			{
+				// TODO: report error (incorrect texture type)
+			}
+			return false;
+		}
+
+		if (!FGLTFMaterialUtility::TryGetTextureCoordinateIndex(TextureSampler, OutTexCoord, OutTransform))
+		{
+			// TODO: report error (failed to identify texture coordinate index)
+			return false;
+		}
+
+		if (!Builder.ExportOptions->bExportTextureTransforms && !OutTransform.IsExactlyDefault())
+		{
+			Builder.LogWarning(FString::Printf(
+				TEXT("Texture coordinates [%d] in %s for material %s are transformed, but texture transform is disabled by export options"),
+				OutTexCoord,
+				*Property.ToString(),
+				*Material->GetName()));
+			OutTransform = {};
+		}
+
+		return true;
+	}
+
+	return false;
+}
+
+bool FGLTFDelayedMaterialTask::TryGetBakedMaterialProperty(FGLTFJsonTextureInfo& OutTexInfo, FGLTFJsonColor3& OutConstant, const FGLTFMaterialPropertyEx& Property, const FString& PropertyName)
+{
+	if (Builder.ExportOptions->BakeMaterialInputs == EGLTFMaterialBakeMode::Disabled)
+	{
+		Builder.LogWarning(FString::Printf(
+			TEXT("%s for material %s needs to bake, but material baking is disabled by export options"),
+			*Property.ToString(),
+			*Material->GetName()));
+		return false;
+	}
+
+	FGLTFPropertyBakeOutput PropertyBakeOutput = BakeMaterialProperty(Property, OutTexInfo.TexCoord);
+
+	if (PropertyBakeOutput.bIsConstant)
+	{
+		OutConstant = FGLTFConverterUtility::ConvertColor3(PropertyBakeOutput.ConstantValue, Builder.ExportOptions->bStrictCompliance);
+		return true;
+	}
+
+	if (Builder.ExportOptions->TextureImageFormat == EGLTFTextureImageFormat::None)
+	{
+		OutTexInfo.Index = nullptr;
+		return true;
+	}
+
+	if (StoreBakedPropertyTexture(OutTexInfo, PropertyBakeOutput, PropertyName))
+	{
+		OutConstant = FGLTFJsonColor3::White; // make sure property is not zero
+		return true;
+	}
+
+	return false;
+}
+
+bool FGLTFDelayedMaterialTask::TryGetBakedMaterialProperty(FGLTFJsonTextureInfo& OutTexInfo, FGLTFJsonColor4& OutConstant, const FGLTFMaterialPropertyEx& Property, const FString& PropertyName)
+{
+	if (Builder.ExportOptions->BakeMaterialInputs == EGLTFMaterialBakeMode::Disabled)
+	{
+		Builder.LogWarning(FString::Printf(
+			TEXT("%s for material %s needs to bake, but material baking is disabled by export options"),
+			*Property.ToString(),
+			*Material->GetName()));
+		return false;
+	}
+
+	FGLTFPropertyBakeOutput PropertyBakeOutput = BakeMaterialProperty(Property, OutTexInfo.TexCoord);
+
+	if (PropertyBakeOutput.bIsConstant)
+	{
+		OutConstant = FGLTFConverterUtility::ConvertColor(PropertyBakeOutput.ConstantValue, Builder.ExportOptions->bStrictCompliance);
+		return true;
+	}
+
+	if (Builder.ExportOptions->TextureImageFormat == EGLTFTextureImageFormat::None)
+	{
+		OutTexInfo.Index = nullptr;
+		return true;
+	}
+
+	if (StoreBakedPropertyTexture(OutTexInfo, PropertyBakeOutput, PropertyName))
+	{
+		OutConstant = FGLTFJsonColor4::White; // make sure property is not zero
+		return true;
+	}
+
+	return false;
+}
+
+inline bool FGLTFDelayedMaterialTask::TryGetBakedMaterialProperty(FGLTFJsonTextureInfo& OutTexInfo, float& OutConstant, const FGLTFMaterialPropertyEx& Property, const FString& PropertyName)
+{
+	if (Builder.ExportOptions->BakeMaterialInputs == EGLTFMaterialBakeMode::Disabled)
+	{
+		Builder.LogWarning(FString::Printf(
+			TEXT("%s for material %s needs to bake, but material baking is disabled by export options"),
+			*Property.ToString(),
+			*Material->GetName()));
+		return false;
+	}
+
+	FGLTFPropertyBakeOutput PropertyBakeOutput = BakeMaterialProperty(Property, OutTexInfo.TexCoord);
+
+	if (PropertyBakeOutput.bIsConstant)
+	{
+		OutConstant = PropertyBakeOutput.ConstantValue.R;
+		return true;
+	}
+
+	if (Builder.ExportOptions->TextureImageFormat == EGLTFTextureImageFormat::None)
+	{
+		OutTexInfo.Index = nullptr;
+		return true;
+	}
+
+	if (StoreBakedPropertyTexture(OutTexInfo, PropertyBakeOutput, PropertyName))
+	{
+		OutConstant = 1; // make sure property is not zero
+		return true;
+	}
+
+	return false;
+}
+
+bool FGLTFDelayedMaterialTask::TryGetBakedMaterialProperty(FGLTFJsonTextureInfo& OutTexInfo, const FGLTFMaterialPropertyEx& Property, const FString& PropertyName)
+{
+	if (Builder.ExportOptions->BakeMaterialInputs == EGLTFMaterialBakeMode::Disabled)
+	{
+		Builder.LogWarning(FString::Printf(
+			TEXT("%s for material %s needs to bake, but material baking is disabled by export options"),
+			*Property.ToString(),
+			*Material->GetName()));
+		return false;
+	}
+
+	FGLTFPropertyBakeOutput PropertyBakeOutput = BakeMaterialProperty(Property, OutTexInfo.TexCoord);
+
+	if (!PropertyBakeOutput.bIsConstant)
+	{
+		if (Builder.ExportOptions->TextureImageFormat == EGLTFTextureImageFormat::None)
+		{
+			OutTexInfo.Index = nullptr;
+			return true;
+		}
+
+		return StoreBakedPropertyTexture(OutTexInfo, PropertyBakeOutput, PropertyName);
+	}
+
+	const FVector4 MaskedConstant = FVector4(PropertyBakeOutput.ConstantValue) * FGLTFMaterialUtility::GetPropertyMask(Property);
+	if (MaskedConstant == FGLTFMaterialUtility::GetPropertyDefaultValue(Property))
+	{
+		// Constant value is the same as the property's default so we can set gltf to default.
+		OutTexInfo.Index = nullptr;
+		return true;
+	}
+
+	if (FGLTFMaterialUtility::IsNormalMap(Property))
+	{
+		// TODO: In some cases baking normal can result in constant vector that differs slight from default (i.e 0,0,1).
+		// Yet often, when looking at such a material, it should be exactly default. Needs further investigation.
+		// Maybe because of incorrect sRGB conversion? For now, assume a constant normal is always default.
+		OutTexInfo.Index = nullptr;
+		return true;
+	}
+
+	// TODO: let function fail and investigate why in some cases a constant baking result is returned for a property
+	// that is non-constant. This happens (for example) when baking AmbientOcclusion for a translucent material,
+	// even though the same material when set to opaque will properly bake AmbientOcclusion to a texture.
+	// For now, create a 1x1 texture with the constant value.
+
+	FGLTFJsonTexture* Texture = FGLTFMaterialUtility::AddTexture(
+		Builder,
+		PropertyBakeOutput.Pixels,
+		PropertyBakeOutput.Size,
+		true, // NOTE: we can ignore alpha in everything but TryGetBaseColorAndOpacity
+		false, // Normal and ClearCoatBottomNormal are handled above
+		GetBakedTextureName(PropertyName),
+		EGLTFJsonTextureFilter::Nearest,
+		EGLTFJsonTextureFilter::Nearest,
+		EGLTFJsonTextureWrap::ClampToEdge,
+		EGLTFJsonTextureWrap::ClampToEdge);
+
+	OutTexInfo.Index = Texture;
+	return true;
+}
+
+FGLTFPropertyBakeOutput FGLTFDelayedMaterialTask::BakeMaterialProperty(const FGLTFMaterialPropertyEx& Property, int32& OutTexCoord)
+{
+	const FIntPoint TextureSize = Builder.GetBakeSizeForMaterialProperty(Material, GetPropertyGroup(Property));
+	return BakeMaterialProperty(Property, OutTexCoord, TextureSize, true);
+}
+
+FGLTFPropertyBakeOutput FGLTFDelayedMaterialTask::BakeMaterialProperty(const FGLTFMaterialPropertyEx& Property, int32& OutTexCoord, const FIntPoint& TextureSize, bool bFillAlpha)
+{
+	if (MeshData == nullptr)
+	{
+		FGLTFIndexArray TexCoords;
+		FGLTFMaterialUtility::GetAllTextureCoordinateIndices(Material, Property, TexCoords);
+
+		if (TexCoords.Num() > 0)
+		{
+			OutTexCoord = TexCoords[0];
+
+			if (TexCoords.Num() > 1)
+			{
+				Builder.LogWarning(FString::Printf(
+					TEXT("%s for material %s uses multiple texture coordinates (%s), baked texture will be sampled using only the first (%d)"),
+					*Property.ToString(),
+					*Material->GetName(),
+					*FString::JoinBy(TexCoords, TEXT(", "), FString::FromInt),
+					OutTexCoord));
+			}
+		}
+		else
+		{
+			OutTexCoord = 0;
+		}
+	}
+	else
+	{
+		OutTexCoord = MeshData->BakeUsingTexCoord;
+		MeshDataBakedProperties.Add(Property);
+	}
+
+	// TODO: add support for calculating the ideal resolution to use for baking based on connected (texture) nodes
+
+	return FGLTFMaterialUtility::BakeMaterialProperty(
+		TextureSize,
+		Property,
+		Material,
+		OutTexCoord,
+		MeshData,
+		SectionIndices,
+		bFillAlpha,
+		Builder.ExportOptions->bAdjustNormalmaps);
+}
+
+bool FGLTFDelayedMaterialTask::StoreBakedPropertyTexture(FGLTFJsonTextureInfo& OutTexInfo, FGLTFPropertyBakeOutput& PropertyBakeOutput, const FString& PropertyName) const
+{
+	const TextureAddress TextureAddress = Builder.GetBakeTilingForMaterialProperty(Material, GetPropertyGroup(PropertyBakeOutput.Property));
+	const EGLTFJsonTextureWrap TextureWrapS = FGLTFConverterUtility::ConvertWrap(TextureAddress);
+	const EGLTFJsonTextureWrap TextureWrapT = FGLTFConverterUtility::ConvertWrap(TextureAddress);
+
+	const TextureFilter TextureFilter = Builder.GetBakeFilterForMaterialProperty(Material, GetPropertyGroup(PropertyBakeOutput.Property));
+	const EGLTFJsonTextureFilter TextureMinFilter = FGLTFConverterUtility::ConvertMinFilter(TextureFilter);
+	const EGLTFJsonTextureFilter TextureMagFilter = FGLTFConverterUtility::ConvertMagFilter(TextureFilter);
+
+	FGLTFJsonTexture* Texture = FGLTFMaterialUtility::AddTexture(
+		Builder,
+		PropertyBakeOutput.Pixels,
+		PropertyBakeOutput.Size,
+		true, // NOTE: we can ignore alpha in everything but TryGetBaseColorAndOpacity
+		FGLTFMaterialUtility::IsNormalMap(PropertyBakeOutput.Property),
+		GetBakedTextureName(PropertyName),
+		TextureMinFilter,
+		TextureMagFilter,
+		TextureWrapS,
+		TextureWrapT);
+
+	OutTexInfo.Index = Texture;
+	return true;
+}
+
+EGLTFMaterialPropertyGroup FGLTFDelayedMaterialTask::GetPropertyGroup(const FGLTFMaterialPropertyEx& Property)
+{
+	switch (Property.Type)
+	{
+		case MP_BaseColor:
+		case MP_Opacity:
+		case MP_OpacityMask:
+			return EGLTFMaterialPropertyGroup::BaseColorOpacity;
+		case MP_Metallic:
+		case MP_Roughness:
+			return EGLTFMaterialPropertyGroup::MetallicRoughness;
+		case MP_EmissiveColor:
+			return EGLTFMaterialPropertyGroup::EmissiveColor;
+		case MP_Normal:
+			return EGLTFMaterialPropertyGroup::Normal;
+		case MP_AmbientOcclusion:
+			return EGLTFMaterialPropertyGroup::AmbientOcclusion;
+		case MP_CustomData0:
+		case MP_CustomData1:
+			return EGLTFMaterialPropertyGroup::ClearCoatRoughness;
+		case MP_CustomOutput:
+			if (Property.CustomOutput == TEXT("ClearCoatBottomNormal"))
+			{
+				return EGLTFMaterialPropertyGroup::ClearCoatBottomNormal;
+			}
+		default:
+			return EGLTFMaterialPropertyGroup::None;
+	}
+}
+
+template <typename CallbackType>
+void FGLTFDelayedMaterialTask::CombinePixels(const TArray<FColor>& FirstPixels, const TArray<FColor>& SecondPixels, TArray<FColor>& OutPixels, CallbackType Callback)
+{
+	const int32 Count = FMath::Max(FirstPixels.Num(), SecondPixels.Num());
+	OutPixels.AddUninitialized(Count);
+
+	if (FirstPixels.Num() == 1)
+	{
+		const FColor& FirstPixel = FirstPixels[0];
+		for (int32 Index = 0; Index < Count; ++Index)
+		{
+			const FColor& SecondPixel = SecondPixels[Index];
+			OutPixels[Index] = Callback(FirstPixel, SecondPixel);
+		}
+	}
+	else if (SecondPixels.Num() == 1)
+	{
+		const FColor& SecondPixel = SecondPixels[0];
+		for (int32 Index = 0; Index < Count; ++Index)
+		{
+			const FColor& FirstPixel = FirstPixels[Index];
+			OutPixels[Index] = Callback(FirstPixel, SecondPixel);
+		}
+	}
+	else
+	{
+		check(FirstPixels.Num() == SecondPixels.Num());
+		for (int32 Index = 0; Index < Count; ++Index)
+		{
+			const FColor& FirstPixel = FirstPixels[Index];
+			const FColor& SecondPixel = SecondPixels[Index];
+			OutPixels[Index] = Callback(FirstPixel, SecondPixel);
+		}
+	}
+}
+
+#endif
diff --git a/Engine/Plugins/Enterprise/GLTFExporter/MigratedFromGit/Source/GLTFExporter/Private/Tasks/GLTFDelayedMaterialTasks.h b/Engine/Plugins/Enterprise/GLTFExporter/MigratedFromGit/Source/GLTFExporter/Private/Tasks/GLTFDelayedMaterialTasks.h
new file mode 100644
--- /dev/null
+++ b/Engine/Plugins/Enterprise/GLTFExporter/MigratedFromGit/Source/GLTFExporter/Private/Tasks/GLTFDelayedMaterialTasks.h
@@ -0,0 +1,93 @@
+// Copyright Epic Games, Inc. All Rights Reserved.
+
+#pragma once
+
+#include "Tasks/GLTFDelayedTask.h"
+#include "Builders/GLTFConvertBuilder.h"
+#include "Converters/GLTFUVOverlapChecker.h"
+#include "Converters/GLTFMeshData.h"
+#include "Materials/GLTFProxyMaterialParameterInfo.h"
+#if WITH_EDITOR
+#include "GLTFMaterialPropertyEx.h"
+#endif
+
+struct FGLTFPropertyBakeOutput;
+
+class FGLTFDelayedMaterialTask : public FGLTFDelayedTask
+{
+public:
+
+	FGLTFDelayedMaterialTask(FGLTFConvertBuilder& Builder, FGLTFUVOverlapChecker& UVOverlapChecker, const UMaterialInterface* Material, const FGLTFMeshData* MeshData, FGLTFIndexArray SectionIndices, FGLTFJsonMaterial* JsonMaterial)
+		: FGLTFDelayedTask(EGLTFTaskPriority::Material)
+		, Builder(Builder)
+		, UVOverlapChecker(UVOverlapChecker)
+		, Material(Material)
+		, MeshData(MeshData)
+		, SectionIndices(SectionIndices)
+		, JsonMaterial(JsonMaterial)
+	{
+	}
+
+	virtual FString GetName() override
+	{
+		return Material->GetName();
+	}
+
+	virtual void Process() override;
+
+private:
+
+	FGLTFConvertBuilder& Builder;
+	FGLTFUVOverlapChecker& UVOverlapChecker;
+	const UMaterialInterface* Material;
+	const FGLTFMeshData* MeshData;
+	const FGLTFIndexArray SectionIndices;
+	FGLTFJsonMaterial* JsonMaterial;
+
+	FString GetMaterialName() const;
+	FString GetBakedTextureName(const FString& PropertyName) const;
+
+	void GetProxyParameters(FGLTFJsonMaterial& OutMaterial) const;
+	void GetProxyParameter(const TGLTFProxyMaterialParameterInfo<float>& ParameterInfo, float& OutValue) const;
+	void GetProxyParameter(const TGLTFProxyMaterialParameterInfo<FLinearColor>& ParameterInfo, FGLTFJsonColor3& OutValue) const;
+	void GetProxyParameter(const TGLTFProxyMaterialParameterInfo<FLinearColor>& ParameterInfo, FGLTFJsonColor4& OutValue) const;
+	void GetProxyParameter(const FGLTFProxyMaterialTextureParameterInfo& ParameterInfo, FGLTFJsonTextureInfo& OutValue) const;
+
+	EMaterialShadingModel GetShadingModel() const;
+	void ConvertShadingModel(EGLTFJsonShadingModel& OutShadingModel) const;
+	void ConvertAlphaMode(EGLTFJsonAlphaMode& OutAlphaMode, EGLTFJsonBlendMode& OutBlendMode) const;
+
+#if WITH_EDITOR
+	TSet<FGLTFMaterialPropertyEx> MeshDataBakedProperties;
+
+	bool TryGetBaseColorAndOpacity(FGLTFJsonPBRMetallicRoughness& OutPBRParams, const FGLTFMaterialPropertyEx& BaseColorProperty, const FGLTFMaterialPropertyEx& OpacityProperty);
+	bool TryGetMetallicAndRoughness(FGLTFJsonPBRMetallicRoughness& OutPBRParams, const FGLTFMaterialPropertyEx& MetallicProperty, const FGLTFMaterialPropertyEx& RoughnessProperty);
+	bool TryGetClearCoatRoughness(FGLTFJsonClearCoatExtension& OutExtParams, const FGLTFMaterialPropertyEx& IntensityProperty, const FGLTFMaterialPropertyEx& RoughnessProperty);
+	bool TryGetEmissive(FGLTFJsonMaterial& OutMaterial, const FGLTFMaterialPropertyEx& EmissiveProperty);
+
+	bool IsPropertyNonDefault(const FGLTFMaterialPropertyEx& Property) const;
+	bool TryGetConstantColor(FGLTFJsonColor3& OutValue, const FGLTFMaterialPropertyEx& Property) const;
+	bool TryGetConstantColor(FGLTFJsonColor4& OutValue, const FGLTFMaterialPropertyEx& Property) const;
+	bool TryGetConstantColor(FLinearColor& OutValue, const FGLTFMaterialPropertyEx& Property) const;
+	bool TryGetConstantScalar(float& OutValue, const FGLTFMaterialPropertyEx& Property) const;
+
+	bool TryGetSourceTexture(FGLTFJsonTextureInfo& OutTexInfo, const FGLTFMaterialPropertyEx& Property, const TArray<FLinearColor>& AllowedMasks = {}) const;
+	bool TryGetSourceTexture(const UTexture2D*& OutTexture, int32& OutTexCoord, FGLTFJsonTextureTransform& OutTransform, const FGLTFMaterialPropertyEx& Property, const TArray<FLinearColor>& AllowedMasks = {}) const;
+
+	bool TryGetBakedMaterialProperty(FGLTFJsonTextureInfo& OutTexInfo, FGLTFJsonColor3& OutConstant, const FGLTFMaterialPropertyEx& Property, const FString& PropertyName);
+	bool TryGetBakedMaterialProperty(FGLTFJsonTextureInfo& OutTexInfo, FGLTFJsonColor4& OutConstant, const FGLTFMaterialPropertyEx& Property, const FString& PropertyName);
+	bool TryGetBakedMaterialProperty(FGLTFJsonTextureInfo& OutTexInfo, float& OutConstant, const FGLTFMaterialPropertyEx& Property, const FString& PropertyName);
+	bool TryGetBakedMaterialProperty(FGLTFJsonTextureInfo& OutTexInfo, const FGLTFMaterialPropertyEx& Property, const FString& PropertyName);
+
+	FGLTFPropertyBakeOutput BakeMaterialProperty(const FGLTFMaterialPropertyEx& Property, int32& OutTexCoord);
+	FGLTFPropertyBakeOutput BakeMaterialProperty(const FGLTFMaterialPropertyEx& Property, int32& OutTexCoord, const FIntPoint& TextureSize, bool bFillAlpha);
+
+	bool StoreBakedPropertyTexture(FGLTFJsonTextureInfo& OutTexInfo, FGLTFPropertyBakeOutput& PropertyBakeOutput, const FString& PropertyName) const;
+
+	static EGLTFMaterialPropertyGroup GetPropertyGroup(const FGLTFMaterialPropertyEx& Property);
+
+	template <typename CallbackType>
+	static void CombinePixels(const TArray<FColor>& FirstPixels, const TArray<FColor>& SecondPixels, TArray<FColor>& OutPixels, CallbackType Callback);
+
+#endif
+};
diff --git a/Engine/Plugins/Enterprise/GLTFExporter/MigratedFromGit/Source/GLTFExporter/Private/Tasks/GLTFDelayedMeshTasks.cpp b/Engine/Plugins/Enterprise/GLTFExporter/MigratedFromGit/Source/GLTFExporter/Private/Tasks/GLTFDelayedMeshTasks.cpp
new file mode 100644
--- /dev/null
+++ b/Engine/Plugins/Enterprise/GLTFExporter/MigratedFromGit/Source/GLTFExporter/Private/Tasks/GLTFDelayedMeshTasks.cpp
@@ -0,0 +1,322 @@
+// Copyright Epic Games, Inc. All Rights Reserved.
+
+#include "Tasks/GLTFDelayedMeshTasks.h"
+#include "Converters/GLTFMeshUtility.h"
+#include "Converters/GLTFBufferAdapter.h"
+#include "Builders/GLTFConvertBuilder.h"
+#include "Engine/StaticMesh.h"
+#include "Engine/SkeletalMesh.h"
+#include "Components/StaticMeshComponent.h"
+#include "Components/SkeletalMeshComponent.h"
+#include "StaticMeshResources.h"
+#include "Rendering/SkeletalMeshRenderData.h"
+
+namespace
+{
+	template <typename VectorType>
+	void CheckTangentVectors(const void* SourceData, uint32 VertexCount, bool& bOutZeroNormals, bool& bOutZeroTangents)
+	{
+		bool bZeroNormals = false;
+		bool bZeroTangents = false;
+
+		typedef TStaticMeshVertexTangentDatum<VectorType> VertexTangentType;
+		const VertexTangentType* VertexTangents = static_cast<const VertexTangentType*>(SourceData);
+
+		for (uint32 VertexIndex = 0; VertexIndex < VertexCount; ++VertexIndex)
+		{
+			const VertexTangentType& VertexTangent = VertexTangents[VertexIndex];
+			bZeroNormals |= VertexTangent.TangentZ.ToFVector().IsNearlyZero();
+			bZeroTangents |= VertexTangent.TangentX.ToFVector().IsNearlyZero();
+		}
+
+		bOutZeroNormals = bZeroNormals;
+		bOutZeroTangents = bZeroTangents;
+	}
+
+	void ValidateVertexBuffer(FGLTFConvertBuilder& Builder, const FStaticMeshVertexBuffer* VertexBuffer, const TCHAR* MeshName)
+	{
+		if (VertexBuffer == nullptr)
+		{
+			return;
+		}
+
+		const TUniquePtr<IGLTFBufferAdapter> SourceBuffer = IGLTFBufferAdapter::GetTangents(VertexBuffer);
+		const uint8* SourceData = SourceBuffer->GetData();
+
+		if (SourceData == nullptr)
+		{
+			return;
+		}
+
+		const uint32 VertexCount = VertexBuffer->GetNumVertices();
+		bool bZeroNormals;
+		bool bZeroTangents;
+
+		if (VertexBuffer->GetUseHighPrecisionTangentBasis())
+		{
+			CheckTangentVectors<FPackedRGBA16N>(SourceData, VertexCount, bZeroNormals, bZeroTangents);
+		}
+		else
+		{
+			CheckTangentVectors<FPackedNormal>(SourceData, VertexCount, bZeroNormals, bZeroTangents);
+		}
+
+		if (bZeroNormals)
+		{
+			Builder.LogSuggestion(FString::Printf(
+				TEXT("Mesh %s has some nearly zero-length normals which may not be supported in some glTF applications. Consider checking 'Recompute Normals' in the asset settings"),
+				MeshName));
+		}
+
+		if (bZeroTangents)
+		{
+			Builder.LogSuggestion(FString::Printf(
+				TEXT("Mesh %s has some nearly zero-length tangents which may not be supported in some glTF applications. Consider checking 'Recompute Tangents' in the asset settings"),
+				MeshName));
+		}
+	}
+
+	bool HasVertexColors(const FColorVertexBuffer* VertexBuffer)
+	{
+		if (VertexBuffer == nullptr)
+		{
+			return false;
+		}
+
+		const TUniquePtr<IGLTFBufferAdapter> SourceBuffer = IGLTFBufferAdapter::GetColors(VertexBuffer);
+		const uint8* SourceData = SourceBuffer->GetData();
+
+		if (SourceData == nullptr)
+		{
+			return false;
+		}
+
+		const uint32 VertexCount = VertexBuffer->GetNumVertices();
+		const uint32 Stride = VertexBuffer->GetStride();
+
+		for (uint32 VertexIndex = 0; VertexIndex < VertexCount; VertexIndex++)
+		{
+			const FColor& Color = *reinterpret_cast<const FColor*>(SourceData + Stride * VertexIndex);
+			if (Color != FColor::White)
+			{
+				return true;
+			}
+		}
+
+		return false;
+	}
+}
+
+void FGLTFDelayedStaticMeshTask::Process()
+{
+	JsonMesh->Name = StaticMeshComponent != nullptr ? FGLTFNameUtility::GetName(StaticMeshComponent) : StaticMesh->GetName();
+
+	const FStaticMeshLODResources& MeshLOD = StaticMesh->GetLODForExport(LODIndex);
+	const FPositionVertexBuffer& PositionBuffer = MeshLOD.VertexBuffers.PositionVertexBuffer;
+	const FStaticMeshVertexBuffer& VertexBuffer = MeshLOD.VertexBuffers.StaticMeshVertexBuffer;
+	const FColorVertexBuffer* ColorBuffer = &MeshLOD.VertexBuffers.ColorVertexBuffer; // TODO: add support for overriding color buffer by component
+
+	if (Builder.ExportOptions->bExportVertexColors && HasVertexColors(ColorBuffer))
+	{
+		Builder.LogWarning(FString::Printf(
+			TEXT("Vertex colors in mesh %s will act as a multiplier for base color in glTF, regardless of material, which may produce undesirable results."),
+			*StaticMesh->GetName()));
+	}
+	else
+	{
+		ColorBuffer = nullptr;
+	}
+
+	if (StaticMeshComponent != nullptr && StaticMeshComponent->LODData.IsValidIndex(LODIndex))
+	{
+		const FStaticMeshComponentLODInfo& LODInfo = StaticMeshComponent->LODData[LODIndex];
+		ColorBuffer = LODInfo.OverrideVertexColors != nullptr ? LODInfo.OverrideVertexColors : ColorBuffer;
+	}
+
+	const FGLTFMeshData* MeshData = Builder.ExportOptions->BakeMaterialInputs == EGLTFMaterialBakeMode::UseMeshData ?
+		Builder.AddUniqueMeshData(StaticMesh, StaticMeshComponent, LODIndex) : nullptr;
+
+#if WITH_EDITOR
+	if (MeshData != nullptr)
+	{
+		if (MeshData->Description.IsEmpty())
+		{
+			// TODO: report warning in case the mesh actually has data, which means we failed to extract a mesh description.
+			MeshData = nullptr;
+		}
+		else if (MeshData->BakeUsingTexCoord < 0)
+		{
+			// TODO: report warning (about missing texture coordinate for baking with mesh data).
+			MeshData = nullptr;
+		}
+	}
+#endif
+
+	ValidateVertexBuffer(Builder, &VertexBuffer, *StaticMesh->GetName());
+
+	const TArray<FStaticMaterial>& MaterialSlots = FGLTFMeshUtility::GetMaterials(StaticMesh);
+	const int32 MaterialCount = MaterialSlots.Num();
+
+	for (int32 MaterialIndex = 0; MaterialIndex < MaterialCount; ++MaterialIndex)
+	{
+		const FGLTFIndexArray SectionIndices = FGLTFMeshUtility::GetSectionIndices(MeshLOD, MaterialIndex);
+		const FGLTFMeshSection* ConvertedSection = MeshSectionConverter.GetOrAdd(&MeshLOD, SectionIndices);
+
+		FGLTFJsonPrimitive& JsonPrimitive = JsonMesh->Primitives[MaterialIndex];
+		JsonPrimitive.Indices = Builder.AddUniqueIndexAccessor(ConvertedSection);
+
+		JsonPrimitive.Attributes.Position = Builder.AddUniquePositionAccessor(ConvertedSection, &PositionBuffer);
+		if (JsonPrimitive.Attributes.Position == nullptr)
+		{
+			FString SectionString = TEXT("section");
+			SectionString += SectionIndices.Num() > 1 ? TEXT("s ") : TEXT(" ");
+			SectionString += FString::JoinBy(SectionIndices, TEXT(", "), FString::FromInt);
+
+			Builder.LogError(
+				FString::Printf(TEXT("Failed to export vertex positions related to material slot %d (%s) in static mesh %s (sections %s)"),
+				MaterialIndex,
+				*MaterialSlots[MaterialIndex].MaterialSlotName.ToString(),
+				*StaticMesh->GetName(),
+				*SectionString
+				));
+		}
+
+		if (ColorBuffer != nullptr)
+		{
+			JsonPrimitive.Attributes.Color0 = Builder.AddUniqueColorAccessor(ConvertedSection, ColorBuffer);
+		}
+
+		// TODO: report warning if both Mesh Quantization (export options) and Use High Precision Tangent Basis (vertex buffer) are disabled
+		JsonPrimitive.Attributes.Normal = Builder.AddUniqueNormalAccessor(ConvertedSection, &VertexBuffer);
+		JsonPrimitive.Attributes.Tangent = Builder.AddUniqueTangentAccessor(ConvertedSection, &VertexBuffer);
+
+		const uint32 UVCount = VertexBuffer.GetNumTexCoords();
+		// TODO: report warning or option to limit UV channels since most viewers don't support more than 2?
+		JsonPrimitive.Attributes.TexCoords.AddUninitialized(UVCount);
+
+		for (uint32 UVIndex = 0; UVIndex < UVCount; ++UVIndex)
+		{
+			JsonPrimitive.Attributes.TexCoords[UVIndex] = Builder.AddUniqueUVAccessor(ConvertedSection, &VertexBuffer, UVIndex);
+		}
+
+		const UMaterialInterface* Material = Materials[MaterialIndex];
+		JsonPrimitive.Material =  Builder.AddUniqueMaterial(Material, MeshData, SectionIndices);
+	}
+}
+
+void FGLTFDelayedSkeletalMeshTask::Process()
+{
+	JsonMesh->Name = SkeletalMeshComponent != nullptr ? FGLTFNameUtility::GetName(SkeletalMeshComponent) : SkeletalMesh->GetName();
+
+	const FSkeletalMeshRenderData* RenderData = SkeletalMesh->GetResourceForRendering();
+	const FSkeletalMeshLODRenderData& MeshLOD = RenderData->LODRenderData[LODIndex];
+
+	const FPositionVertexBuffer& PositionBuffer = MeshLOD.StaticVertexBuffers.PositionVertexBuffer;
+	const FStaticMeshVertexBuffer& VertexBuffer = MeshLOD.StaticVertexBuffers.StaticMeshVertexBuffer;
+	const FColorVertexBuffer* ColorBuffer = &MeshLOD.StaticVertexBuffers.ColorVertexBuffer; // TODO: add support for overriding color buffer by component
+	const FSkinWeightVertexBuffer* SkinWeightBuffer = MeshLOD.GetSkinWeightVertexBuffer(); // TODO: add support for overriding skin weight buffer by component
+	// TODO: add support for skin weight profiles?
+	// TODO: add support for morph targets
+
+	if (Builder.ExportOptions->bExportVertexColors && HasVertexColors(ColorBuffer))
+	{
+		Builder.LogWarning(FString::Printf(
+			TEXT("Vertex colors in mesh %s will act as a multiplier for base color in glTF, regardless of material, which may produce undesirable results."),
+			*SkeletalMesh->GetName()));
+	}
+	else
+	{
+		ColorBuffer = nullptr;
+	}
+
+	if (SkeletalMeshComponent != nullptr && SkeletalMeshComponent->LODInfo.IsValidIndex(LODIndex))
+	{
+		const FSkelMeshComponentLODInfo& LODInfo = SkeletalMeshComponent->LODInfo[LODIndex];
+		ColorBuffer = LODInfo.OverrideVertexColors != nullptr ? LODInfo.OverrideVertexColors : ColorBuffer;
+		SkinWeightBuffer = LODInfo.OverrideSkinWeights != nullptr ? LODInfo.OverrideSkinWeights : SkinWeightBuffer;
+	}
+
+	const FGLTFMeshData* MeshData = Builder.ExportOptions->BakeMaterialInputs == EGLTFMaterialBakeMode::UseMeshData ?
+		Builder.AddUniqueMeshData(SkeletalMesh, SkeletalMeshComponent, LODIndex) : nullptr;
+
+#if WITH_EDITOR
+	if (MeshData != nullptr)
+	{
+		if (MeshData->Description.IsEmpty())
+		{
+			// TODO: report warning in case the mesh actually has data, which means we failed to extract a mesh description.
+			MeshData = nullptr;
+		}
+		else if (MeshData->BakeUsingTexCoord < 0)
+		{
+			// TODO: report warning (about missing texture coordinate for baking with mesh data).
+			MeshData = nullptr;
+		}
+	}
+#endif
+
+	ValidateVertexBuffer(Builder, &VertexBuffer, *SkeletalMesh->GetName());
+
+	const TArray<FSkeletalMaterial>& MaterialSlots = FGLTFMeshUtility::GetMaterials(SkeletalMesh);
+	const uint16 MaterialCount = MaterialSlots.Num();
+
+	for (uint16 MaterialIndex = 0; MaterialIndex < MaterialCount; ++MaterialIndex)
+	{
+		const FGLTFIndexArray SectionIndices = FGLTFMeshUtility::GetSectionIndices(MeshLOD, MaterialIndex);
+		const FGLTFMeshSection* ConvertedSection = MeshSectionConverter.GetOrAdd(&MeshLOD, SectionIndices);
+
+		FGLTFJsonPrimitive& JsonPrimitive = JsonMesh->Primitives[MaterialIndex];
+		JsonPrimitive.Indices = Builder.AddUniqueIndexAccessor(ConvertedSection);
+
+		JsonPrimitive.Attributes.Position = Builder.AddUniquePositionAccessor(ConvertedSection, &PositionBuffer);
+		if (JsonPrimitive.Attributes.Position == nullptr)
+		{
+			FString SectionString = TEXT("section");
+			SectionString += SectionIndices.Num() > 1 ? TEXT("s ") : TEXT(" ");
+			SectionString += FString::JoinBy(SectionIndices, TEXT(", "), FString::FromInt);
+
+			Builder.LogError(
+				FString::Printf(TEXT("Failed to export vertex positions related to material slot %d (%s) in skeletal mesh %s (sections %s)"),
+				MaterialIndex,
+				*MaterialSlots[MaterialIndex].MaterialSlotName.ToString(),
+				*SkeletalMesh->GetName(),
+				*SectionString
+				));
+		}
+
+		if (ColorBuffer != nullptr)
+		{
+			JsonPrimitive.Attributes.Color0 = Builder.AddUniqueColorAccessor(ConvertedSection, ColorBuffer);
+		}
+
+		// TODO: report warning if both Mesh Quantization (export options) and Use High Precision Tangent Basis (vertex buffer) are disabled
+		JsonPrimitive.Attributes.Normal = Builder.AddUniqueNormalAccessor(ConvertedSection, &VertexBuffer);
+		JsonPrimitive.Attributes.Tangent = Builder.AddUniqueTangentAccessor(ConvertedSection, &VertexBuffer);
+
+		const uint32 UVCount = VertexBuffer.GetNumTexCoords();
+		// TODO: report warning or option to limit UV channels since most viewers don't support more than 2?
+		JsonPrimitive.Attributes.TexCoords.AddUninitialized(UVCount);
+
+		for (uint32 UVIndex = 0; UVIndex < UVCount; ++UVIndex)
+		{
+			JsonPrimitive.Attributes.TexCoords[UVIndex] = Builder.AddUniqueUVAccessor(ConvertedSection, &VertexBuffer, UVIndex);
+		}
+
+		if (Builder.ExportOptions->bExportVertexSkinWeights)
+		{
+			const uint32 GroupCount = (SkinWeightBuffer->GetMaxBoneInfluences() + 3) / 4;
+			// TODO: report warning or option to limit groups (of joints and weights) since most viewers don't support more than one?
+			JsonPrimitive.Attributes.Joints.AddUninitialized(GroupCount);
+			JsonPrimitive.Attributes.Weights.AddUninitialized(GroupCount);
+
+			for (uint32 GroupIndex = 0; GroupIndex < GroupCount; ++GroupIndex)
+			{
+				JsonPrimitive.Attributes.Joints[GroupIndex] = Builder.AddUniqueJointAccessor(ConvertedSection, SkinWeightBuffer, GroupIndex * 4);
+				JsonPrimitive.Attributes.Weights[GroupIndex] = Builder.AddUniqueWeightAccessor(ConvertedSection, SkinWeightBuffer, GroupIndex * 4);
+			}
+		}
+
+		const UMaterialInterface* Material = Materials[MaterialIndex];
+		JsonPrimitive.Material =  Builder.AddUniqueMaterial(Material, MeshData, SectionIndices);
+	}
+}
diff --git a/Engine/Plugins/Enterprise/GLTFExporter/MigratedFromGit/Source/GLTFExporter/Private/Tasks/GLTFDelayedMeshTasks.h b/Engine/Plugins/Enterprise/GLTFExporter/MigratedFromGit/Source/GLTFExporter/Private/Tasks/GLTFDelayedMeshTasks.h
new file mode 100644
--- /dev/null
+++ b/Engine/Plugins/Enterprise/GLTFExporter/MigratedFromGit/Source/GLTFExporter/Private/Tasks/GLTFDelayedMeshTasks.h
@@ -0,0 +1,77 @@
+// Copyright Epic Games, Inc. All Rights Reserved.
+
+#pragma once
+
+#include "Tasks/GLTFDelayedTask.h"
+#include "Builders/GLTFConvertBuilder.h"
+#include "Converters/GLTFMeshSectionConverters.h"
+#include "Converters/GLTFMaterialArray.h"
+#include "Converters/GLTFNameUtility.h"
+
+class FGLTFDelayedStaticMeshTask : public FGLTFDelayedTask
+{
+public:
+
+	FGLTFDelayedStaticMeshTask(FGLTFConvertBuilder& Builder, FGLTFStaticMeshSectionConverter& MeshSectionConverter, const UStaticMesh* StaticMesh, const UStaticMeshComponent* StaticMeshComponent, FGLTFMaterialArray Materials, int32 LODIndex, FGLTFJsonMesh* JsonMesh)
+		: FGLTFDelayedTask(EGLTFTaskPriority::Mesh)
+		, Builder(Builder)
+		, MeshSectionConverter(MeshSectionConverter)
+		, StaticMesh(StaticMesh)
+		, StaticMeshComponent(StaticMeshComponent)
+		, Materials(Materials)
+		, LODIndex(LODIndex)
+		, JsonMesh(JsonMesh)
+	{
+	}
+
+	virtual FString GetName() override
+	{
+		return StaticMeshComponent != nullptr ? FGLTFNameUtility::GetName(StaticMeshComponent) : StaticMesh->GetName();
+	}
+
+	virtual void Process() override;
+
+private:
+
+	FGLTFConvertBuilder& Builder;
+	FGLTFStaticMeshSectionConverter& MeshSectionConverter;
+	const UStaticMesh* StaticMesh;
+	const UStaticMeshComponent* StaticMeshComponent;
+	const FGLTFMaterialArray Materials;
+	const int32 LODIndex;
+	FGLTFJsonMesh* JsonMesh;
+};
+
+class FGLTFDelayedSkeletalMeshTask : public FGLTFDelayedTask
+{
+public:
+
+	FGLTFDelayedSkeletalMeshTask(FGLTFConvertBuilder& Builder, FGLTFSkeletalMeshSectionConverter& MeshSectionConverter, const USkeletalMesh* SkeletalMesh, const USkeletalMeshComponent* SkeletalMeshComponent, FGLTFMaterialArray Materials, int32 LODIndex, FGLTFJsonMesh* JsonMesh)
+		: FGLTFDelayedTask(EGLTFTaskPriority::Mesh)
+		, Builder(Builder)
+		, MeshSectionConverter(MeshSectionConverter)
+		, SkeletalMesh(SkeletalMesh)
+		, SkeletalMeshComponent(SkeletalMeshComponent)
+		, Materials(Materials)
+		, LODIndex(LODIndex)
+		, JsonMesh(JsonMesh)
+	{
+	}
+
+	virtual FString GetName() override
+	{
+		return SkeletalMeshComponent != nullptr ? FGLTFNameUtility::GetName(SkeletalMeshComponent) : SkeletalMesh->GetName();
+	}
+
+	virtual void Process() override;
+
+private:
+
+	FGLTFConvertBuilder& Builder;
+	FGLTFSkeletalMeshSectionConverter& MeshSectionConverter;
+	const USkeletalMesh* SkeletalMesh;
+	const USkeletalMeshComponent* SkeletalMeshComponent;
+	const FGLTFMaterialArray Materials;
+	const int32 LODIndex;
+	FGLTFJsonMesh* JsonMesh;
+};
diff --git a/Engine/Plugins/Enterprise/GLTFExporter/MigratedFromGit/Source/GLTFExporter/Private/Tasks/GLTFDelayedTask.h b/Engine/Plugins/Enterprise/GLTFExporter/MigratedFromGit/Source/GLTFExporter/Private/Tasks/GLTFDelayedTask.h
new file mode 100644
--- /dev/null
+++ b/Engine/Plugins/Enterprise/GLTFExporter/MigratedFromGit/Source/GLTFExporter/Private/Tasks/GLTFDelayedTask.h
@@ -0,0 +1,32 @@
+// Copyright Epic Games, Inc. All Rights Reserved.
+
+#pragma once
+
+#include "CoreMinimal.h"
+
+enum class EGLTFTaskPriority : uint8
+{
+	Animation,
+	Mesh,
+	Material,
+	Texture,
+	MAX
+};
+
+class FGLTFDelayedTask
+{
+public:
+
+	const EGLTFTaskPriority Priority;
+
+	FGLTFDelayedTask(EGLTFTaskPriority Priority)
+		: Priority(Priority)
+	{
+	}
+
+	virtual ~FGLTFDelayedTask() = default;
+
+	virtual FString GetName() = 0;
+
+	virtual void Process() = 0;
+};
diff --git a/Engine/Plugins/Enterprise/GLTFExporter/MigratedFromGit/Source/GLTFExporter/Private/Tasks/GLTFDelayedTextureTasks.cpp b/Engine/Plugins/Enterprise/GLTFExporter/MigratedFromGit/Source/GLTFExporter/Private/Tasks/GLTFDelayedTextureTasks.cpp
new file mode 100644
--- /dev/null
+++ b/Engine/Plugins/Enterprise/GLTFExporter/MigratedFromGit/Source/GLTFExporter/Private/Tasks/GLTFDelayedTextureTasks.cpp
@@ -0,0 +1,255 @@
+// Copyright Epic Games, Inc. All Rights Reserved.
+
+#include "Tasks/GLTFDelayedTextureTasks.h"
+#include "Builders/GLTFContainerBuilder.h"
+#include "Converters/GLTFConverterUtility.h"
+#include "Converters/GLTFTextureUtility.h"
+#include "Converters/GLTFNameUtility.h"
+#include "Engine/Texture2D.h"
+#include "Engine/TextureCube.h"
+#include "Engine/TextureRenderTarget2D.h"
+#include "Engine/TextureRenderTargetCube.h"
+#include "Engine/LightMapTexture2D.h"
+
+FString FGLTFDelayedTexture2DTask::GetName()
+{
+	return Texture2D->GetName();
+}
+
+void FGLTFDelayedTexture2DTask::Process()
+{
+	Texture2D->GetName(JsonTexture->Name);
+
+	const bool bFromSRGB = Texture2D->SRGB;
+	if (bFromSRGB != bToSRGB)
+	{
+		JsonTexture->Name += bToSRGB ? TEXT("_sRGB") : TEXT("_Linear");
+	}
+
+	const bool bIsHDR = FGLTFTextureUtility::IsHDR(Texture2D);
+	const FIntPoint Size = FGLTFTextureUtility::GetInGameSize(Texture2D);
+	UTextureRenderTarget2D* RenderTarget = FGLTFTextureUtility::CreateRenderTarget(Size, bIsHDR);
+
+	// TODO: preserve maximum image quality (avoid compression artifacts) by copying source data (and adjustments) to a temp texture
+	FGLTFTextureUtility::DrawTexture(RenderTarget, Texture2D, FVector2D::ZeroVector, Size);
+
+	if (bIsHDR)
+	{
+		JsonTexture->Encoding = Builder.GetTextureHDREncoding();
+	}
+
+	TGLTFSharedArray<FColor> Pixels;
+	if (!FGLTFTextureUtility::ReadPixels(RenderTarget, *Pixels, JsonTexture->Encoding))
+	{
+		Builder.LogWarning(FString::Printf(TEXT("Failed to read pixels for 2D texture %s"), *JsonTexture->Name));
+		return;
+	}
+
+	if (Builder.ExportOptions->bAdjustNormalmaps && Texture2D->IsNormalMap())
+	{
+		// TODO: add support for adjusting normals in GLTFNormalMapPreview instead
+		FGLTFTextureUtility::FlipGreenChannel(*Pixels);
+	}
+
+	FGLTFTextureUtility::TransformColorSpace(*Pixels, bFromSRGB, bToSRGB);
+
+	const bool bIgnoreAlpha = FGLTFTextureUtility::IsAlphaless(Texture2D->GetPixelFormat());
+	const EGLTFTextureType Type =
+		Texture2D->IsNormalMap() ? EGLTFTextureType::Normalmaps :
+		bIsHDR ? EGLTFTextureType::HDR : EGLTFTextureType::None;
+
+	JsonTexture->Source = Builder.AddUniqueImage(Pixels, Size, bIgnoreAlpha, Type, JsonTexture->Name);
+	JsonTexture->Sampler = Builder.AddUniqueSampler(Texture2D);
+}
+
+FString FGLTFDelayedTextureCubeTask::GetName()
+{
+	return TextureCube->GetName();
+}
+
+void FGLTFDelayedTextureCubeTask::Process()
+{
+	JsonTexture->Name = TextureCube->GetName() + TEXT("_") + FGLTFJsonUtility::GetValue(FGLTFConverterUtility::ConvertCubeFace(CubeFace));
+
+	const bool bFromSRGB = TextureCube->SRGB;
+	if (bFromSRGB != bToSRGB)
+	{
+		JsonTexture->Name += bToSRGB ? TEXT("_sRGB") : TEXT("_Linear");
+	}
+
+	// TODO: add optimized "happy path" if cube face doesn't need rotation and has suitable pixel format
+
+	// TODO: preserve maximum image quality (avoid compression artifacts) by copying source data (and adjustments) to a temp texture
+	const UTexture2D* FaceTexture = FGLTFTextureUtility::CreateTextureFromCubeFace(TextureCube, CubeFace);
+	if (FaceTexture == nullptr)
+	{
+		Builder.LogWarning(FString::Printf(TEXT("Failed to extract cube face %d for cubemap texture %s"), CubeFace, *TextureCube->GetName()));
+		return;
+	}
+
+	const bool bIsHDR = FGLTFTextureUtility::IsHDR(TextureCube);
+	const FIntPoint Size = { TextureCube->GetSizeX(), TextureCube->GetSizeY() };
+	UTextureRenderTarget2D* RenderTarget = FGLTFTextureUtility::CreateRenderTarget(Size, bIsHDR);
+
+	const float FaceRotation = FGLTFTextureUtility::GetCubeFaceRotation(CubeFace);
+	FGLTFTextureUtility::RotateTexture(RenderTarget, FaceTexture, FVector2D::ZeroVector, Size, FaceRotation);
+
+	if (bIsHDR)
+	{
+		JsonTexture->Encoding = Builder.GetTextureHDREncoding();
+		if (JsonTexture->Encoding == EGLTFJsonHDREncoding::None)
+		{
+			bToSRGB = true;
+		}
+	}
+
+	TGLTFSharedArray<FColor> Pixels;
+	if (!FGLTFTextureUtility::ReadPixels(RenderTarget, *Pixels, JsonTexture->Encoding))
+	{
+		Builder.LogWarning(FString::Printf(TEXT("Failed to read pixels (cube face %d) for cubemap texture %s"), CubeFace, *TextureCube->GetName()));
+		return;
+	}
+
+	FGLTFTextureUtility::TransformColorSpace(*Pixels, bFromSRGB, bToSRGB);
+
+	const bool bIgnoreAlpha = FGLTFTextureUtility::IsAlphaless(TextureCube->GetPixelFormat());
+	const EGLTFTextureType Type = bIsHDR ? EGLTFTextureType::HDR : EGLTFTextureType::None;
+
+	JsonTexture->Source = Builder.AddUniqueImage(Pixels, Size, bIgnoreAlpha, Type, JsonTexture->Name);
+	JsonTexture->Sampler = Builder.AddUniqueSampler(TextureCube);
+}
+
+FString FGLTFDelayedTextureRenderTarget2DTask::GetName()
+{
+	return RenderTarget2D->GetName();
+}
+
+void FGLTFDelayedTextureRenderTarget2DTask::Process()
+{
+	RenderTarget2D->GetName(JsonTexture->Name);
+
+	const bool bFromSRGB = RenderTarget2D->SRGB;
+	if (bFromSRGB != bToSRGB)
+	{
+		JsonTexture->Name += bToSRGB ? TEXT("_sRGB") : TEXT("_Linear");
+	}
+
+	const bool bIsHDR = FGLTFTextureUtility::IsHDR(RenderTarget2D);
+	const FIntPoint Size = { RenderTarget2D->SizeX, RenderTarget2D->SizeY };
+
+	if (bIsHDR)
+	{
+		JsonTexture->Encoding = Builder.GetTextureHDREncoding();
+	}
+
+	TGLTFSharedArray<FColor> Pixels;
+	if (!FGLTFTextureUtility::ReadPixels(RenderTarget2D, *Pixels, JsonTexture->Encoding))
+	{
+		Builder.LogWarning(FString::Printf(TEXT("Failed to read pixels for 2D render target %s"), *JsonTexture->Name));
+		return;
+	}
+
+	FGLTFTextureUtility::TransformColorSpace(*Pixels, bFromSRGB, bToSRGB);
+
+	const bool bIgnoreAlpha = FGLTFTextureUtility::IsAlphaless(RenderTarget2D->GetFormat());
+	const EGLTFTextureType Type = bIsHDR ? EGLTFTextureType::HDR : EGLTFTextureType::None;
+
+	JsonTexture->Source = Builder.AddUniqueImage(Pixels, Size, bIgnoreAlpha, Type, JsonTexture->Name);
+	JsonTexture->Sampler = Builder.AddUniqueSampler(RenderTarget2D);
+}
+
+FString FGLTFDelayedTextureRenderTargetCubeTask::GetName()
+{
+	return RenderTargetCube->GetName();
+}
+
+void FGLTFDelayedTextureRenderTargetCubeTask::Process()
+{
+	JsonTexture->Name = RenderTargetCube->GetName() + TEXT("_") + FGLTFJsonUtility::GetValue(FGLTFConverterUtility::ConvertCubeFace(CubeFace));
+
+	const bool bFromSRGB = RenderTargetCube->SRGB;
+	if (bFromSRGB != bToSRGB)
+	{
+		JsonTexture->Name += bToSRGB ? TEXT("_sRGB") : TEXT("_Linear");
+	}
+
+	// TODO: add optimized "happy path" if cube face doesn't need rotation
+
+	const UTexture2D* FaceTexture = FGLTFTextureUtility::CreateTextureFromCubeFace(RenderTargetCube, CubeFace);
+	if (FaceTexture == nullptr)
+	{
+		Builder.LogWarning(FString::Printf(TEXT("Failed to extract cube face %d for cubemap render target %s"), CubeFace, *RenderTargetCube->GetName()));
+		return;
+	}
+
+	const bool bIsHDR = FGLTFTextureUtility::IsHDR(RenderTargetCube);
+	const FIntPoint Size = { RenderTargetCube->SizeX, RenderTargetCube->SizeX };
+	UTextureRenderTarget2D* RenderTarget = FGLTFTextureUtility::CreateRenderTarget(Size, bIsHDR);
+
+	const float FaceRotation = FGLTFTextureUtility::GetCubeFaceRotation(CubeFace);
+	FGLTFTextureUtility::RotateTexture(RenderTarget, FaceTexture, FVector2D::ZeroVector, Size, FaceRotation);
+
+	if (bIsHDR)
+	{
+		JsonTexture->Encoding = Builder.GetTextureHDREncoding();
+	}
+
+	TGLTFSharedArray<FColor> Pixels;
+	if (!FGLTFTextureUtility::ReadPixels(RenderTarget, *Pixels, JsonTexture->Encoding))
+	{
+		Builder.LogWarning(FString::Printf(TEXT("Failed to read pixels (cube face %d) for cubemap render target %s"), CubeFace, *RenderTargetCube->GetName()));
+		return;
+	}
+
+	FGLTFTextureUtility::TransformColorSpace(*Pixels, bFromSRGB, bToSRGB);
+
+	const bool bIgnoreAlpha = FGLTFTextureUtility::IsAlphaless(RenderTargetCube->GetFormat());
+	const EGLTFTextureType Type = bIsHDR ? EGLTFTextureType::HDR : EGLTFTextureType::None;
+
+	JsonTexture->Source = Builder.AddUniqueImage(Pixels, Size, bIgnoreAlpha, Type, JsonTexture->Name);
+	JsonTexture->Sampler = Builder.AddUniqueSampler(RenderTargetCube);
+}
+
+#if WITH_EDITOR
+
+FString FGLTFDelayedTextureLightMapTask::GetName()
+{
+	return LightMap->GetName();
+}
+
+void FGLTFDelayedTextureLightMapTask::Process()
+{
+	LightMap->GetName(JsonTexture->Name);
+
+	// NOTE: export of lightmaps via source data is used to work around issues with
+	// quality-loss due to incorrect gamma transformation when rendering to a canvas.
+
+	FTextureSource& Source = const_cast<FTextureSource&>(LightMap->Source);
+	if (!Source.IsValid())
+	{
+		Builder.LogWarning(FString::Printf(TEXT("Failed to export lightmap texture %s because of missing source data"), *JsonTexture->Name));
+		return;
+	}
+
+	const ETextureSourceFormat SourceFormat = Source.GetFormat();
+	if (SourceFormat != TSF_BGRA8)
+	{
+		Builder.LogWarning(FString::Printf(TEXT("Failed to export lightmap texture %s because of unsupported source format %s"), *JsonTexture->Name, *FGLTFNameUtility::GetName(SourceFormat)));
+		return;
+	}
+
+	const FIntPoint Size = { Source.GetSizeX(), Source.GetSizeY() };
+	const int64 ByteLength = Source.CalcMipSize(0);
+
+	const bool bIgnoreAlpha = false;
+	const EGLTFTextureType Type = EGLTFTextureType::Lightmaps;
+
+	const void* RawData = Source.LockMip(0);
+	TGLTFSharedArray<FColor> Pixels = MakeShared<TArray<FColor>>(static_cast<const FColor*>(RawData), ByteLength / sizeof(FColor));
+	Source.UnlockMip(0);
+
+	JsonTexture->Source = Builder.AddUniqueImage(Pixels, Size, bIgnoreAlpha, Type, JsonTexture->Name);
+	JsonTexture->Sampler = Builder.AddUniqueSampler(LightMap);
+}
+
+#endif
diff --git a/Engine/Plugins/Enterprise/GLTFExporter/MigratedFromGit/Source/GLTFExporter/Private/Tasks/GLTFDelayedTextureTasks.h b/Engine/Plugins/Enterprise/GLTFExporter/MigratedFromGit/Source/GLTFExporter/Private/Tasks/GLTFDelayedTextureTasks.h
new file mode 100644
--- /dev/null
+++ b/Engine/Plugins/Enterprise/GLTFExporter/MigratedFromGit/Source/GLTFExporter/Private/Tasks/GLTFDelayedTextureTasks.h
@@ -0,0 +1,137 @@
+// Copyright Epic Games, Inc. All Rights Reserved.
+
+#pragma once
+
+#include "Tasks/GLTFDelayedTask.h"
+#include "Builders/GLTFConvertBuilder.h"
+
+class FGLTFDelayedTexture2DTask : public FGLTFDelayedTask
+{
+public:
+
+	FGLTFDelayedTexture2DTask(FGLTFConvertBuilder& Builder, const UTexture2D* Texture2D, bool bToSRGB, FGLTFJsonTexture* JsonTexture)
+		: FGLTFDelayedTask(EGLTFTaskPriority::Texture)
+		, Builder(Builder)
+		, Texture2D(Texture2D)
+		, bToSRGB(bToSRGB)
+		, JsonTexture(JsonTexture)
+	{
+	}
+
+	virtual FString GetName() override;
+
+	virtual void Process() override;
+
+private:
+
+	FGLTFConvertBuilder& Builder;
+	const UTexture2D* Texture2D;
+	bool bToSRGB;
+	FGLTFJsonTexture* JsonTexture;
+};
+
+class FGLTFDelayedTextureCubeTask : public FGLTFDelayedTask
+{
+public:
+
+	FGLTFDelayedTextureCubeTask(FGLTFConvertBuilder& Builder, const UTextureCube* TextureCube, ECubeFace CubeFace, bool bToSRGB, FGLTFJsonTexture* JsonTexture)
+		: FGLTFDelayedTask(EGLTFTaskPriority::Texture)
+		, Builder(Builder)
+		, TextureCube(TextureCube)
+		, CubeFace(CubeFace)
+		, bToSRGB(bToSRGB)
+		, JsonTexture(JsonTexture)
+	{
+	}
+
+	virtual FString GetName() override;
+
+	virtual void Process() override;
+
+private:
+
+	FGLTFConvertBuilder& Builder;
+	const UTextureCube* TextureCube;
+	ECubeFace CubeFace;
+	bool bToSRGB;
+	FGLTFJsonTexture* JsonTexture;
+};
+
+class FGLTFDelayedTextureRenderTarget2DTask : public FGLTFDelayedTask
+{
+public:
+
+	FGLTFDelayedTextureRenderTarget2DTask(FGLTFConvertBuilder& Builder, const UTextureRenderTarget2D* RenderTarget2D, bool bToSRGB, FGLTFJsonTexture* JsonTexture)
+		: FGLTFDelayedTask(EGLTFTaskPriority::Texture)
+		, Builder(Builder)
+		, RenderTarget2D(RenderTarget2D)
+		, bToSRGB(bToSRGB)
+		, JsonTexture(JsonTexture)
+	{
+	}
+
+	virtual FString GetName() override;
+
+	virtual void Process() override;
+
+private:
+
+	FGLTFConvertBuilder& Builder;
+	const UTextureRenderTarget2D* RenderTarget2D;
+	bool bToSRGB;
+	FGLTFJsonTexture* JsonTexture;
+};
+
+class FGLTFDelayedTextureRenderTargetCubeTask : public FGLTFDelayedTask
+{
+public:
+
+	FGLTFDelayedTextureRenderTargetCubeTask(FGLTFConvertBuilder& Builder, const UTextureRenderTargetCube* RenderTargetCube, ECubeFace CubeFace, bool bToSRGB, FGLTFJsonTexture* JsonTexture)
+		: FGLTFDelayedTask(EGLTFTaskPriority::Texture)
+		, Builder(Builder)
+		, RenderTargetCube(RenderTargetCube)
+		, CubeFace(CubeFace)
+		, bToSRGB(bToSRGB)
+		, JsonTexture(JsonTexture)
+	{
+	}
+
+	virtual FString GetName() override;
+
+	virtual void Process() override;
+
+private:
+
+	FGLTFConvertBuilder& Builder;
+	const UTextureRenderTargetCube* RenderTargetCube;
+	ECubeFace CubeFace;
+	bool bToSRGB;
+	FGLTFJsonTexture* JsonTexture;
+};
+
+#if WITH_EDITOR
+
+class FGLTFDelayedTextureLightMapTask : public FGLTFDelayedTask
+{
+public:
+
+	FGLTFDelayedTextureLightMapTask(FGLTFConvertBuilder& Builder, const ULightMapTexture2D* LightMap, FGLTFJsonTexture* JsonTexture)
+		: FGLTFDelayedTask(EGLTFTaskPriority::Texture)
+		, Builder(Builder)
+		, LightMap(LightMap)
+		, JsonTexture(JsonTexture)
+	{
+	}
+
+	virtual FString GetName() override;
+
+	virtual void Process() override;
+
+private:
+
+	FGLTFConvertBuilder& Builder;
+	const ULightMapTexture2D* LightMap;
+	FGLTFJsonTexture* JsonTexture;
+};
+
+#endif
diff --git a/Engine/Plugins/Enterprise/GLTFExporter/MigratedFromGit/Source/GLTFExporter/Private/Tasks/GLTFMaterialTasks.cpp b/Engine/Plugins/Enterprise/GLTFExporter/MigratedFromGit/Source/GLTFExporter/Private/Tasks/GLTFMaterialTasks.cpp
deleted file mode 100644
diff --git a/Engine/Plugins/Enterprise/GLTFExporter/MigratedFromGit/Source/GLTFExporter/Private/Tasks/GLTFMaterialTasks.h b/Engine/Plugins/Enterprise/GLTFExporter/MigratedFromGit/Source/GLTFExporter/Private/Tasks/GLTFMaterialTasks.h
deleted file mode 100644
diff --git a/Engine/Plugins/Enterprise/GLTFExporter/MigratedFromGit/Source/GLTFExporter/Private/Tasks/GLTFMeshTasks.cpp b/Engine/Plugins/Enterprise/GLTFExporter/MigratedFromGit/Source/GLTFExporter/Private/Tasks/GLTFMeshTasks.cpp
deleted file mode 100644
diff --git a/Engine/Plugins/Enterprise/GLTFExporter/MigratedFromGit/Source/GLTFExporter/Private/Tasks/GLTFMeshTasks.h b/Engine/Plugins/Enterprise/GLTFExporter/MigratedFromGit/Source/GLTFExporter/Private/Tasks/GLTFMeshTasks.h
deleted file mode 100644
diff --git a/Engine/Plugins/Enterprise/GLTFExporter/MigratedFromGit/Source/GLTFExporter/Private/Tasks/GLTFTask.h b/Engine/Plugins/Enterprise/GLTFExporter/MigratedFromGit/Source/GLTFExporter/Private/Tasks/GLTFTask.h
deleted file mode 100644
diff --git a/Engine/Plugins/Enterprise/GLTFExporter/MigratedFromGit/Source/GLTFExporter/Private/Tasks/GLTFTextureTasks.cpp b/Engine/Plugins/Enterprise/GLTFExporter/MigratedFromGit/Source/GLTFExporter/Private/Tasks/GLTFTextureTasks.cpp
deleted file mode 100644
diff --git a/Engine/Plugins/Enterprise/GLTFExporter/MigratedFromGit/Source/GLTFExporter/Private/Tasks/GLTFTextureTasks.h b/Engine/Plugins/Enterprise/GLTFExporter/MigratedFromGit/Source/GLTFExporter/Private/Tasks/GLTFTextureTasks.h
deleted file mode 100644
