[GLTFExporter] refactor: rename MaterialBakingStructures.h to GLTFMaterialBakingStructures.h

commit 577b86beeddbb3141261ad0294adee567ab3f922
Date: Sun, 13 Mar 2022 22:27:31 +0100

    refactor: rename MaterialBakingStructures.h to GLTFMaterialBakingStructures.h
    (cherry picked from commit d58b1927769c6f1c44aa10df82e5ea5c0253daa6​)

#skipci

diff --git a/Engine/Plugins/Enterprise/GLTFExporter/MigratedFromGit/Source/GLTFExporter/Private/Converters/GLTFMaterialUtility.cpp b/Engine/Plugins/Enterprise/GLTFExporter/MigratedFromGit/Source/GLTFExporter/Private/Converters/GLTFMaterialUtility.cpp
--- a/Engine/Plugins/Enterprise/GLTFExporter/MigratedFromGit/Source/GLTFExporter/Private/Converters/GLTFMaterialUtility.cpp
+++ b/Engine/Plugins/Enterprise/GLTFExporter/MigratedFromGit/Source/GLTFExporter/Private/Converters/GLTFMaterialUtility.cpp
@@ -7,8 +7,8 @@
 #include "Misc/DefaultValueHelper.h"
 #if WITH_EDITOR
 #include "GLTFMaterialAnalyzer.h"
-#include "GLTFMaterialBaking/Public/IGLTFMaterialBakingModule.h"
-#include "GLTFMaterialBaking/Public/MaterialBakingStructures.h"
+#include "IGLTFMaterialBakingModule.h"
+#include "GLTFMaterialBakingStructures.h"
 #include "Modules/ModuleManager.h"
 #include "Materials/MaterialInstanceConstant.h"
 #include "Materials/MaterialExpressionCustomOutput.h"
diff --git a/Engine/Plugins/Enterprise/GLTFExporter/MigratedFromGit/Source/GLTFExporter/Private/Converters/GLTFMeshData.h b/Engine/Plugins/Enterprise/GLTFExporter/MigratedFromGit/Source/GLTFExporter/Private/Converters/GLTFMeshData.h
--- a/Engine/Plugins/Enterprise/GLTFExporter/MigratedFromGit/Source/GLTFExporter/Private/Converters/GLTFMeshData.h
+++ b/Engine/Plugins/Enterprise/GLTFExporter/MigratedFromGit/Source/GLTFExporter/Private/Converters/GLTFMeshData.h
@@ -6,7 +6,7 @@
 #include "MeshDescription.h"
 
 #if WITH_EDITOR
-#include "MaterialBakingStructures.h"
+#include "GLTFMaterialBakingStructures.h"
 #endif
 
 struct FGLTFMeshData
diff --git a/Engine/Plugins/Enterprise/GLTFExporter/MigratedFromGit/Source/GLTFExporter/Private/Converters/GLTFUVOverlapChecker.cpp b/Engine/Plugins/Enterprise/GLTFExporter/MigratedFromGit/Source/GLTFExporter/Private/Converters/GLTFUVOverlapChecker.cpp
--- a/Engine/Plugins/Enterprise/GLTFExporter/MigratedFromGit/Source/GLTFExporter/Private/Converters/GLTFUVOverlapChecker.cpp
+++ b/Engine/Plugins/Enterprise/GLTFExporter/MigratedFromGit/Source/GLTFExporter/Private/Converters/GLTFUVOverlapChecker.cpp
@@ -6,7 +6,7 @@
 #include "MeshDescription.h"
 #include "Modules/ModuleManager.h"
 #include "IGLTFMaterialBakingModule.h"
-#include "MaterialBakingStructures.h"
+#include "GLTFMaterialBakingStructures.h"
 #endif
 
 void FGLTFUVOverlapChecker::Sanitize(const FMeshDescription*& Description, FGLTFIndexArray& SectionIndices, int32& TexCoord)
diff --git a/Engine/Plugins/Enterprise/GLTFExporter/MigratedFromGit/Source/GLTFMaterialBaking/Private/GLTFMaterialBakingModule.cpp b/Engine/Plugins/Enterprise/GLTFExporter/MigratedFromGit/Source/GLTFMaterialBaking/Private/GLTFMaterialBakingModule.cpp
--- a/Engine/Plugins/Enterprise/GLTFExporter/MigratedFromGit/Source/GLTFMaterialBaking/Private/GLTFMaterialBakingModule.cpp
+++ b/Engine/Plugins/Enterprise/GLTFExporter/MigratedFromGit/Source/GLTFMaterialBaking/Private/GLTFMaterialBakingModule.cpp
@@ -5,7 +5,7 @@
 #include "Engine/TextureRenderTarget2D.h"
 #include "GLTFExportMaterialProxy.h"
 #include "UObject/UObjectGlobals.h"
-#include "MaterialBakingStructures.h"
+#include "GLTFMaterialBakingStructures.h"
 #include "GLTFMaterialBakingHelpers.h"
 #include "AddCustomAttributeHack.h"
 #include "Async/Async.h"
diff --git a/Engine/Plugins/Enterprise/GLTFExporter/MigratedFromGit/Source/GLTFMaterialBaking/Private/MaterialRenderItem.cpp b/Engine/Plugins/Enterprise/GLTFExporter/MigratedFromGit/Source/GLTFMaterialBaking/Private/MaterialRenderItem.cpp
--- a/Engine/Plugins/Enterprise/GLTFExporter/MigratedFromGit/Source/GLTFMaterialBaking/Private/MaterialRenderItem.cpp
+++ b/Engine/Plugins/Enterprise/GLTFExporter/MigratedFromGit/Source/GLTFMaterialBaking/Private/MaterialRenderItem.cpp
@@ -1,7 +1,7 @@
 // Copyright Epic Games, Inc. All Rights Reserved.
 
 #include "MaterialRenderItem.h"
-#include "MaterialBakingStructures.h"
+#include "GLTFMaterialBakingStructures.h"
 
 #include "EngineModule.h"
 #include "StaticMeshAttributes.h"
diff --git a/Engine/Plugins/Enterprise/GLTFExporter/MigratedFromGit/Source/GLTFMaterialBaking/Public/GLTFMaterialBakingStructures.h b/Engine/Plugins/Enterprise/GLTFExporter/MigratedFromGit/Source/GLTFMaterialBaking/Public/GLTFMaterialBakingStructures.h
new file mode 100644
--- /dev/null
+++ b/Engine/Plugins/Enterprise/GLTFExporter/MigratedFromGit/Source/GLTFMaterialBaking/Public/GLTFMaterialBakingStructures.h
@@ -0,0 +1,195 @@
+// Copyright Epic Games, Inc. All Rights Reserved.
+
+#pragma once
+
+#include "Math/IntPoint.h"
+#include "SceneTypes.h"
+#include "LightMap.h"
+#include "GLTFMaterialPropertyEx.h"
+
+class UMaterialInterface;
+struct FMeshDescription;
+
+/** Structure containing information about the material which is being baked out */
+struct FGLTFMaterialData
+{
+	FGLTFMaterialData()
+		: Material(nullptr)
+		, bPerformBorderSmear(true)
+		, BlendMode(BLEND_Opaque)
+		, BackgroundColor(FColor::Magenta)
+		, bTangentSpaceNormal(false)
+	{}
+
+	/** Material to bake out */
+	UMaterialInterface* Material;
+	/** Properties and the texture size at which they should be baked out */
+	TMap<EMaterialProperty, FIntPoint> PropertySizes;
+	/** Whether to smear borders after baking */
+	bool bPerformBorderSmear;
+	/** Blend mode to use when baking, allowing for example detection of overlapping UVs */
+	EBlendMode BlendMode;
+	/** Background color used to initially fill the output texture and used for border smear */
+	FColor BackgroundColor;
+	/** Whether to transform normals from world-space to tangent-space (does nothing if material already uses tangent-space normals) */
+	bool bTangentSpaceNormal;
+};
+
+/** Structure containing extended information about the material and properties which is being baked out */
+struct FGLTFMaterialDataEx
+{
+	FGLTFMaterialDataEx()
+		: Material(nullptr)
+		, bPerformBorderSmear(true)
+		, BlendMode(BLEND_Opaque)
+		, BackgroundColor(FColor::Magenta)
+		, bTangentSpaceNormal(false)
+	{}
+
+	/** Material to bake out */
+	UMaterialInterface* Material;
+	/** Extended properties and the texture size at which they should be baked out */
+	TMap<FGLTFMaterialPropertyEx, FIntPoint> PropertySizes;
+	/** Whether to smear borders after baking */
+	bool bPerformBorderSmear;
+	/** Blend mode to use when baking, allowing for example detection of overlapping UVs */
+	EBlendMode BlendMode;
+	/** Background color used to initially fill the output texture and used for border smear */
+	FColor BackgroundColor;
+	/** Whether to transform normals from world-space to tangent-space (does nothing if material already uses tangent-space normals) */
+	bool bTangentSpaceNormal;
+};
+
+/** Structure containing primitive information (regarding a mesh or mesh component) that is accessible through material expressions */
+struct FGLTFPrimitiveData
+{
+	FGLTFPrimitiveData(const FBoxSphereBounds& LocalBounds = FBoxSphereBounds(ForceInitToZero))
+		: ActorPosition(ForceInitToZero)
+		, ObjectOrientation(FVector::UpVector)
+		, ObjectBounds(LocalBounds)
+		, ObjectLocalBounds(LocalBounds)
+		, PreSkinnedLocalBounds(LocalBounds)
+		, CustomPrimitiveData(nullptr)
+	{}
+
+	FGLTFPrimitiveData(const UStaticMesh* StaticMesh)
+		: FGLTFPrimitiveData(StaticMesh->GetBounds())
+	{}
+
+	FGLTFPrimitiveData(const USkeletalMesh* SkeletalMesh)
+		: FGLTFPrimitiveData(SkeletalMesh->GetBounds())
+	{}
+
+	FGLTFPrimitiveData(const UMeshComponent* MeshComponent)
+		: ActorPosition(MeshComponent->GetComponentLocation())
+		, ObjectOrientation(MeshComponent->GetRenderMatrix().GetUnitAxis(EAxis::Z))
+		, ObjectBounds(MeshComponent->CalcBounds(MeshComponent->GetComponentTransform()))
+		, ObjectLocalBounds(MeshComponent->CalcLocalBounds())
+		, CustomPrimitiveData(&MeshComponent->GetCustomPrimitiveData())
+	{
+		if (const USkinnedMeshComponent* SkinnedMeshComponent = Cast<USkinnedMeshComponent>(MeshComponent))
+		{
+			SkinnedMeshComponent->GetPreSkinnedLocalBounds(PreSkinnedLocalBounds);
+		}
+		else
+		{
+			PreSkinnedLocalBounds = ObjectLocalBounds;
+		}
+	}
+
+	/** The location of the actor in world-space */
+	FVector ActorPosition;
+	/** The up vector of the object in world-space */
+	FVector ObjectOrientation;
+	/** The bounds of the object in world-space */
+	FBoxSphereBounds ObjectBounds;
+	/** The bounds of the object in local-space */
+	FBoxSphereBounds ObjectLocalBounds;
+	/** The pre-skinning bounds for the object in local-space */
+	FBoxSphereBounds PreSkinnedLocalBounds;
+	/** The custom primitive data for the mesh component */
+	const FCustomPrimitiveData* CustomPrimitiveData;
+};
+
+struct FGLTFMeshRenderData
+{
+	FGLTFMeshRenderData()
+		: RawMeshDescription(nullptr), Mesh(nullptr), bMirrored(false), VertexColorHash(0), TextureCoordinateIndex(0), LightMapIndex(0), LightMap(nullptr), LightmapResourceCluster(nullptr), PrimitiveData(nullptr)
+	{}
+
+	/** Ptr to raw mesh data to use for baking out the material data, if nullptr a standard quad is used */
+	FMeshDescription* RawMeshDescription;
+
+	/** Ptr to original static mesh this mesh data came from */
+	UStaticMesh* Mesh;
+
+	/** Transform determinant used to detect mirroring */
+	bool bMirrored;
+
+	/** A hash of the vertex color buffer for the rawmesh */
+	uint32 VertexColorHash;
+
+	/** Material indices to test the Raw Mesh data against, ensuring we only bake out triangles which use the currently baked out material */
+	TArray<int32> MaterialIndices;
+
+	/** Set of custom texture coordinates which ensure that the material is baked out with unique/non-overlapping positions */
+	TArray<FVector2D> CustomTextureCoordinates;
+
+	/** Box which's space contains the UV coordinates used to bake out the material */
+	FBox2D TextureCoordinateBox;
+	
+	/** Specific texture coordinate index to use as texture coordinates to bake out the material (is overruled if CustomTextureCoordinates contains any data) */
+	int32 TextureCoordinateIndex;
+
+	/** Light map index used to retrieve the light-map UVs from RawMesh */
+	int32 LightMapIndex;
+
+	/** Reference to the lightmap texture part of the level in the currently being baked out mesh instance data is resident */
+	FLightMapRef LightMap;
+
+	/** Pointer to the LightmapResourceCluster to be passed on the the LightCacheInterface when baking */
+	const FLightmapResourceCluster* LightmapResourceCluster;
+
+	/** Pointer to primitive data that is accessible through material expressions, if nullptr default values are used */
+	const FGLTFPrimitiveData* PrimitiveData;
+};
+
+/** Structure containing data being processed while baking out materials*/
+struct FGLTFBakeOutput
+{
+	FGLTFBakeOutput()
+		: EmissiveScale(1.0f)
+	{}
+
+	/** Contains the resulting texture data for baking out a material's property */
+	TMap<EMaterialProperty, TArray<FColor>> PropertyData;
+
+	/** Contains the resulting texture size for baking out a material's property */
+	TMap<EMaterialProperty, FIntPoint> PropertySizes;
+
+	/** Contains the resulting HDR texture data for baking out a material's property, may be empty */
+	TMap<EMaterialProperty, TArray<FFloat16Color>> HDRPropertyData;
+
+	/** Scale used to allow having wide ranges of emissive values in the source materials, the final proxy material will use this value to scale the emissive texture's pixel values */
+	float EmissiveScale;
+};
+
+/** Structure containing extended data being processed while baking out materials*/
+struct FGLTFBakeOutputEx
+{
+	FGLTFBakeOutputEx()
+		: EmissiveScale(1.0f)
+	{}
+
+	/** Contains the resulting texture data for baking out a extened material's property */
+	TMap<FGLTFMaterialPropertyEx, TArray<FColor>> PropertyData;
+
+	/** Contains the resulting texture size for baking out a extened material's property */
+	TMap<FGLTFMaterialPropertyEx, FIntPoint> PropertySizes;
+
+	/** Contains the resulting HDR texture data for baking out a material's property, may be empty */
+	TMap<FGLTFMaterialPropertyEx, TArray<FFloat16Color>> HDRPropertyData;
+
+	/** Scale used to allow having wide ranges of emissive values in the source materials, the final proxy material will use this value to scale the emissive texture's pixel values */
+	float EmissiveScale;
+};
diff --git a/Engine/Plugins/Enterprise/GLTFExporter/MigratedFromGit/Source/GLTFMaterialBaking/Public/MaterialBakingStructures.h b/Engine/Plugins/Enterprise/GLTFExporter/MigratedFromGit/Source/GLTFMaterialBaking/Public/MaterialBakingStructures.h
deleted file mode 100644
