[GLTFExporter] refactor: make all necessary headers public to expose builder api

commit 49ccfc51b5aca14684001c9b93c920672b236f4e
Date: Sun, 13 Mar 2022 22:33:54 +0100

    refactor: make all necessary headers public to expose builder api
    (cherry picked from commit e1e7c266a7be10405be447b4eb852e992e8f3e51​)

#skipci

diff --git a/Engine/Plugins/Enterprise/GLTFExporter/MigratedFromGit/Source/GLTFExporter/Private/Builders/GLTFBufferBuilder.h b/Engine/Plugins/Enterprise/GLTFExporter/MigratedFromGit/Source/GLTFExporter/Private/Builders/GLTFBufferBuilder.h
deleted file mode 100644
diff --git a/Engine/Plugins/Enterprise/GLTFExporter/MigratedFromGit/Source/GLTFExporter/Private/Builders/GLTFBuilder.h b/Engine/Plugins/Enterprise/GLTFExporter/MigratedFromGit/Source/GLTFExporter/Private/Builders/GLTFBuilder.h
deleted file mode 100644
diff --git a/Engine/Plugins/Enterprise/GLTFExporter/MigratedFromGit/Source/GLTFExporter/Private/Builders/GLTFContainerBuilder.h b/Engine/Plugins/Enterprise/GLTFExporter/MigratedFromGit/Source/GLTFExporter/Private/Builders/GLTFContainerBuilder.h
deleted file mode 100644
diff --git a/Engine/Plugins/Enterprise/GLTFExporter/MigratedFromGit/Source/GLTFExporter/Private/Builders/GLTFConvertBuilder.h b/Engine/Plugins/Enterprise/GLTFExporter/MigratedFromGit/Source/GLTFExporter/Private/Builders/GLTFConvertBuilder.h
deleted file mode 100644
diff --git a/Engine/Plugins/Enterprise/GLTFExporter/MigratedFromGit/Source/GLTFExporter/Private/Builders/GLTFFileBuilder.h b/Engine/Plugins/Enterprise/GLTFExporter/MigratedFromGit/Source/GLTFExporter/Private/Builders/GLTFFileBuilder.h
deleted file mode 100644
diff --git a/Engine/Plugins/Enterprise/GLTFExporter/MigratedFromGit/Source/GLTFExporter/Private/Builders/GLTFJsonBuilder.h b/Engine/Plugins/Enterprise/GLTFExporter/MigratedFromGit/Source/GLTFExporter/Private/Builders/GLTFJsonBuilder.h
deleted file mode 100644
diff --git a/Engine/Plugins/Enterprise/GLTFExporter/MigratedFromGit/Source/GLTFExporter/Private/Builders/GLTFLogBuilder.h b/Engine/Plugins/Enterprise/GLTFExporter/MigratedFromGit/Source/GLTFExporter/Private/Builders/GLTFLogBuilder.h
deleted file mode 100644
diff --git a/Engine/Plugins/Enterprise/GLTFExporter/MigratedFromGit/Source/GLTFExporter/Private/Builders/GLTFMemoryArchive.h b/Engine/Plugins/Enterprise/GLTFExporter/MigratedFromGit/Source/GLTFExporter/Private/Builders/GLTFMemoryArchive.h
deleted file mode 100644
diff --git a/Engine/Plugins/Enterprise/GLTFExporter/MigratedFromGit/Source/GLTFExporter/Private/Builders/GLTFTaskBuilder.h b/Engine/Plugins/Enterprise/GLTFExporter/MigratedFromGit/Source/GLTFExporter/Private/Builders/GLTFTaskBuilder.h
deleted file mode 100644
diff --git a/Engine/Plugins/Enterprise/GLTFExporter/MigratedFromGit/Source/GLTFExporter/Private/Converters/GLTFAccessorConverters.h b/Engine/Plugins/Enterprise/GLTFExporter/MigratedFromGit/Source/GLTFExporter/Private/Converters/GLTFAccessorConverters.h
deleted file mode 100644
diff --git a/Engine/Plugins/Enterprise/GLTFExporter/MigratedFromGit/Source/GLTFExporter/Private/Converters/GLTFAnimationConverters.h b/Engine/Plugins/Enterprise/GLTFExporter/MigratedFromGit/Source/GLTFExporter/Private/Converters/GLTFAnimationConverters.h
deleted file mode 100644
diff --git a/Engine/Plugins/Enterprise/GLTFExporter/MigratedFromGit/Source/GLTFExporter/Private/Converters/GLTFBackdropConverters.h b/Engine/Plugins/Enterprise/GLTFExporter/MigratedFromGit/Source/GLTFExporter/Private/Converters/GLTFBackdropConverters.h
deleted file mode 100644
diff --git a/Engine/Plugins/Enterprise/GLTFExporter/MigratedFromGit/Source/GLTFExporter/Private/Converters/GLTFBuilderContext.h b/Engine/Plugins/Enterprise/GLTFExporter/MigratedFromGit/Source/GLTFExporter/Private/Converters/GLTFBuilderContext.h
deleted file mode 100644
diff --git a/Engine/Plugins/Enterprise/GLTFExporter/MigratedFromGit/Source/GLTFExporter/Private/Converters/GLTFCameraConverters.h b/Engine/Plugins/Enterprise/GLTFExporter/MigratedFromGit/Source/GLTFExporter/Private/Converters/GLTFCameraConverters.h
deleted file mode 100644
diff --git a/Engine/Plugins/Enterprise/GLTFExporter/MigratedFromGit/Source/GLTFExporter/Private/Converters/GLTFConverter.h b/Engine/Plugins/Enterprise/GLTFExporter/MigratedFromGit/Source/GLTFExporter/Private/Converters/GLTFConverter.h
deleted file mode 100644
diff --git a/Engine/Plugins/Enterprise/GLTFExporter/MigratedFromGit/Source/GLTFExporter/Private/Converters/GLTFEpicLevelVariantSetsConverters.h b/Engine/Plugins/Enterprise/GLTFExporter/MigratedFromGit/Source/GLTFExporter/Private/Converters/GLTFEpicLevelVariantSetsConverters.h
deleted file mode 100644
diff --git a/Engine/Plugins/Enterprise/GLTFExporter/MigratedFromGit/Source/GLTFExporter/Private/Converters/GLTFHotspotConverters.h b/Engine/Plugins/Enterprise/GLTFExporter/MigratedFromGit/Source/GLTFExporter/Private/Converters/GLTFHotspotConverters.h
deleted file mode 100644
diff --git a/Engine/Plugins/Enterprise/GLTFExporter/MigratedFromGit/Source/GLTFExporter/Private/Converters/GLTFImageConverters.h b/Engine/Plugins/Enterprise/GLTFExporter/MigratedFromGit/Source/GLTFExporter/Private/Converters/GLTFImageConverters.h
deleted file mode 100644
diff --git a/Engine/Plugins/Enterprise/GLTFExporter/MigratedFromGit/Source/GLTFExporter/Private/Converters/GLTFIndexArray.h b/Engine/Plugins/Enterprise/GLTFExporter/MigratedFromGit/Source/GLTFExporter/Private/Converters/GLTFIndexArray.h
deleted file mode 100644
diff --git a/Engine/Plugins/Enterprise/GLTFExporter/MigratedFromGit/Source/GLTFExporter/Private/Converters/GLTFKhrVariantConverters.h b/Engine/Plugins/Enterprise/GLTFExporter/MigratedFromGit/Source/GLTFExporter/Private/Converters/GLTFKhrVariantConverters.h
deleted file mode 100644
diff --git a/Engine/Plugins/Enterprise/GLTFExporter/MigratedFromGit/Source/GLTFExporter/Private/Converters/GLTFLightConverters.h b/Engine/Plugins/Enterprise/GLTFExporter/MigratedFromGit/Source/GLTFExporter/Private/Converters/GLTFLightConverters.h
deleted file mode 100644
diff --git a/Engine/Plugins/Enterprise/GLTFExporter/MigratedFromGit/Source/GLTFExporter/Private/Converters/GLTFLightMapConverters.h b/Engine/Plugins/Enterprise/GLTFExporter/MigratedFromGit/Source/GLTFExporter/Private/Converters/GLTFLightMapConverters.h
deleted file mode 100644
diff --git a/Engine/Plugins/Enterprise/GLTFExporter/MigratedFromGit/Source/GLTFExporter/Private/Converters/GLTFMaterialArray.h b/Engine/Plugins/Enterprise/GLTFExporter/MigratedFromGit/Source/GLTFExporter/Private/Converters/GLTFMaterialArray.h
deleted file mode 100644
diff --git a/Engine/Plugins/Enterprise/GLTFExporter/MigratedFromGit/Source/GLTFExporter/Private/Converters/GLTFMaterialConverters.h b/Engine/Plugins/Enterprise/GLTFExporter/MigratedFromGit/Source/GLTFExporter/Private/Converters/GLTFMaterialConverters.h
deleted file mode 100644
diff --git a/Engine/Plugins/Enterprise/GLTFExporter/MigratedFromGit/Source/GLTFExporter/Private/Converters/GLTFMeshConverters.h b/Engine/Plugins/Enterprise/GLTFExporter/MigratedFromGit/Source/GLTFExporter/Private/Converters/GLTFMeshConverters.h
deleted file mode 100644
diff --git a/Engine/Plugins/Enterprise/GLTFExporter/MigratedFromGit/Source/GLTFExporter/Private/Converters/GLTFMeshData.h b/Engine/Plugins/Enterprise/GLTFExporter/MigratedFromGit/Source/GLTFExporter/Private/Converters/GLTFMeshData.h
deleted file mode 100644
diff --git a/Engine/Plugins/Enterprise/GLTFExporter/MigratedFromGit/Source/GLTFExporter/Private/Converters/GLTFMeshDataConverters.h b/Engine/Plugins/Enterprise/GLTFExporter/MigratedFromGit/Source/GLTFExporter/Private/Converters/GLTFMeshDataConverters.h
deleted file mode 100644
diff --git a/Engine/Plugins/Enterprise/GLTFExporter/MigratedFromGit/Source/GLTFExporter/Private/Converters/GLTFMeshSection.h b/Engine/Plugins/Enterprise/GLTFExporter/MigratedFromGit/Source/GLTFExporter/Private/Converters/GLTFMeshSection.h
deleted file mode 100644
diff --git a/Engine/Plugins/Enterprise/GLTFExporter/MigratedFromGit/Source/GLTFExporter/Private/Converters/GLTFMeshSectionConverters.h b/Engine/Plugins/Enterprise/GLTFExporter/MigratedFromGit/Source/GLTFExporter/Private/Converters/GLTFMeshSectionConverters.h
deleted file mode 100644
diff --git a/Engine/Plugins/Enterprise/GLTFExporter/MigratedFromGit/Source/GLTFExporter/Private/Converters/GLTFNodeConverters.h b/Engine/Plugins/Enterprise/GLTFExporter/MigratedFromGit/Source/GLTFExporter/Private/Converters/GLTFNodeConverters.h
deleted file mode 100644
diff --git a/Engine/Plugins/Enterprise/GLTFExporter/MigratedFromGit/Source/GLTFExporter/Private/Converters/GLTFSamplerConverters.h b/Engine/Plugins/Enterprise/GLTFExporter/MigratedFromGit/Source/GLTFExporter/Private/Converters/GLTFSamplerConverters.h
deleted file mode 100644
diff --git a/Engine/Plugins/Enterprise/GLTFExporter/MigratedFromGit/Source/GLTFExporter/Private/Converters/GLTFSceneConverters.h b/Engine/Plugins/Enterprise/GLTFExporter/MigratedFromGit/Source/GLTFExporter/Private/Converters/GLTFSceneConverters.h
deleted file mode 100644
diff --git a/Engine/Plugins/Enterprise/GLTFExporter/MigratedFromGit/Source/GLTFExporter/Private/Converters/GLTFSharedArray.h b/Engine/Plugins/Enterprise/GLTFExporter/MigratedFromGit/Source/GLTFExporter/Private/Converters/GLTFSharedArray.h
deleted file mode 100644
diff --git a/Engine/Plugins/Enterprise/GLTFExporter/MigratedFromGit/Source/GLTFExporter/Private/Converters/GLTFSkinConverters.h b/Engine/Plugins/Enterprise/GLTFExporter/MigratedFromGit/Source/GLTFExporter/Private/Converters/GLTFSkinConverters.h
deleted file mode 100644
diff --git a/Engine/Plugins/Enterprise/GLTFExporter/MigratedFromGit/Source/GLTFExporter/Private/Converters/GLTFSkySphereConverters.h b/Engine/Plugins/Enterprise/GLTFExporter/MigratedFromGit/Source/GLTFExporter/Private/Converters/GLTFSkySphereConverters.h
deleted file mode 100644
diff --git a/Engine/Plugins/Enterprise/GLTFExporter/MigratedFromGit/Source/GLTFExporter/Private/Converters/GLTFSuperfluous.h b/Engine/Plugins/Enterprise/GLTFExporter/MigratedFromGit/Source/GLTFExporter/Private/Converters/GLTFSuperfluous.h
deleted file mode 100644
diff --git a/Engine/Plugins/Enterprise/GLTFExporter/MigratedFromGit/Source/GLTFExporter/Private/Converters/GLTFTextureConverters.h b/Engine/Plugins/Enterprise/GLTFExporter/MigratedFromGit/Source/GLTFExporter/Private/Converters/GLTFTextureConverters.h
deleted file mode 100644
diff --git a/Engine/Plugins/Enterprise/GLTFExporter/MigratedFromGit/Source/GLTFExporter/Private/Converters/GLTFUVDegenerateChecker.h b/Engine/Plugins/Enterprise/GLTFExporter/MigratedFromGit/Source/GLTFExporter/Private/Converters/GLTFUVDegenerateChecker.h
deleted file mode 100644
diff --git a/Engine/Plugins/Enterprise/GLTFExporter/MigratedFromGit/Source/GLTFExporter/Private/Converters/GLTFUVOverlapChecker.h b/Engine/Plugins/Enterprise/GLTFExporter/MigratedFromGit/Source/GLTFExporter/Private/Converters/GLTFUVOverlapChecker.h
deleted file mode 100644
diff --git a/Engine/Plugins/Enterprise/GLTFExporter/MigratedFromGit/Source/GLTFExporter/Private/Tasks/GLTFDelayedTask.h b/Engine/Plugins/Enterprise/GLTFExporter/MigratedFromGit/Source/GLTFExporter/Private/Tasks/GLTFDelayedTask.h
deleted file mode 100644
diff --git a/Engine/Plugins/Enterprise/GLTFExporter/MigratedFromGit/Source/GLTFExporter/Public/Builders/GLTFBufferBuilder.h b/Engine/Plugins/Enterprise/GLTFExporter/MigratedFromGit/Source/GLTFExporter/Public/Builders/GLTFBufferBuilder.h
new file mode 100644
--- /dev/null
+++ b/Engine/Plugins/Enterprise/GLTFExporter/MigratedFromGit/Source/GLTFExporter/Public/Builders/GLTFBufferBuilder.h
@@ -0,0 +1,32 @@
+// Copyright Epic Games, Inc. All Rights Reserved.
+
+#pragma once
+
+#include "Builders/GLTFJsonBuilder.h"
+
+class FGLTFBufferBuilder : public FGLTFJsonBuilder
+{
+protected:
+
+	FGLTFBufferBuilder(const FString& FileName, const UGLTFExportOptions* ExportOptions);
+	~FGLTFBufferBuilder();
+
+	const FGLTFMemoryArchive* GetBufferData() const;
+
+public:
+
+	FGLTFJsonBufferView* AddBufferView(const void* RawData, uint64 ByteLength, EGLTFJsonBufferTarget BufferTarget = EGLTFJsonBufferTarget::None, uint8 DataAlignment = 4);
+
+	template <class ElementType>
+	FGLTFJsonBufferView* AddBufferView(const TArray<ElementType>& Array, EGLTFJsonBufferTarget BufferTarget = EGLTFJsonBufferTarget::None, uint8 DataAlignment = 4)
+	{
+		return AddBufferView(Array.GetData(), Array.Num() * sizeof(ElementType), BufferTarget, DataAlignment);
+	}
+
+private:
+
+	void InitializeBuffer();
+
+	FGLTFJsonBuffer* JsonBuffer;
+	TSharedPtr<FGLTFMemoryArchive> BufferArchive;
+};
diff --git a/Engine/Plugins/Enterprise/GLTFExporter/MigratedFromGit/Source/GLTFExporter/Public/Builders/GLTFBuilder.h b/Engine/Plugins/Enterprise/GLTFExporter/MigratedFromGit/Source/GLTFExporter/Public/Builders/GLTFBuilder.h
new file mode 100644
--- /dev/null
+++ b/Engine/Plugins/Enterprise/GLTFExporter/MigratedFromGit/Source/GLTFExporter/Public/Builders/GLTFBuilder.h
@@ -0,0 +1,40 @@
+// Copyright Epic Games, Inc. All Rights Reserved.
+
+#pragma once
+
+#include "Options/GLTFExportOptions.h"
+#include "Json/GLTFJsonEnums.h"
+#include "UObject/GCObjectScopeGuard.h"
+
+class FGLTFBuilder
+{
+public:
+
+	const FString FileName;
+	const bool bIsGLB;
+
+	const UGLTFExportOptions* ExportOptions; // TODO: make ExportOptions private and expose each option via getters to ease overriding settings in future
+
+	FGLTFBuilder(const FString& FileName, const UGLTFExportOptions* ExportOptions);
+
+	const UMaterialInterface* ResolveProxy(const UMaterialInterface* Material) const;
+	void ResolveProxies(TArray<const UMaterialInterface*>& Materials) const;
+
+	FIntPoint GetBakeSizeForMaterialProperty(const UMaterialInterface* Material, EGLTFMaterialPropertyGroup PropertyGroup) const;
+	TextureFilter GetBakeFilterForMaterialProperty(const UMaterialInterface* Material, EGLTFMaterialPropertyGroup PropertyGroup) const;
+	TextureAddress GetBakeTilingForMaterialProperty(const UMaterialInterface* Material, EGLTFMaterialPropertyGroup PropertyGroup) const;
+
+	EGLTFJsonHDREncoding GetTextureHDREncoding() const;
+
+	bool ShouldExportLight(EComponentMobility::Type LightMobility) const;
+
+	int32 SanitizeLOD(const UStaticMesh* StaticMesh, const UStaticMeshComponent* StaticMeshComponent, int32 LODIndex) const;
+	int32 SanitizeLOD(const USkeletalMesh* SkeletalMesh, const USkeletalMeshComponent* SkeletalMeshComponent, int32 LODIndex) const;
+
+private:
+
+	FGCObjectScopeGuard ExportOptionsGuard;
+
+	static const UGLTFExportOptions* SanitizeExportOptions(const UGLTFExportOptions* Options);
+	static EGLTFSceneMobility GetSceneMobility(EComponentMobility::Type Mobility);
+};
diff --git a/Engine/Plugins/Enterprise/GLTFExporter/MigratedFromGit/Source/GLTFExporter/Public/Builders/GLTFContainerBuilder.h b/Engine/Plugins/Enterprise/GLTFExporter/MigratedFromGit/Source/GLTFExporter/Public/Builders/GLTFContainerBuilder.h
new file mode 100644
--- /dev/null
+++ b/Engine/Plugins/Enterprise/GLTFExporter/MigratedFromGit/Source/GLTFExporter/Public/Builders/GLTFContainerBuilder.h
@@ -0,0 +1,34 @@
+// Copyright Epic Games, Inc. All Rights Reserved.
+
+#pragma once
+
+#include "Builders/GLTFConvertBuilder.h"
+
+class FGLTFContainerBuilder : public FGLTFConvertBuilder
+{
+public:
+
+	FGLTFContainerBuilder(const FString& FileName, const UGLTFExportOptions* ExportOptions, const TSet<AActor*>& SelectedActors);
+
+	void WriteToArchive(FArchive& Archive);
+
+	bool WriteAllFiles(const FString& DirPath, bool bOverwrite = true);
+
+protected:
+
+	void WriteGlb(FArchive& Archive);
+
+private:
+
+	static void WriteGlb(FArchive& Archive, const TArray64<uint8>& JsonData, const TArray64<uint8>* BinaryData);
+
+	static void WriteHeader(FArchive& Archive, uint32 FileSize);
+	static void WriteChunk(FArchive& Archive, uint32 ChunkType, const TArray64<uint8>& ChunkData, uint8 ChunkTrailingByte);
+
+	static void WriteInt(FArchive& Archive, uint32 Value);
+	static void WriteData(FArchive& Archive, const TArray64<uint8>& Data);
+	static void WriteFill(FArchive& Archive, int32 Size, uint8 Value);
+
+	static int32 GetPaddedChunkSize(int32 Size);
+	static int32 GetTrailingChunkSize(int32 Size);
+};
diff --git a/Engine/Plugins/Enterprise/GLTFExporter/MigratedFromGit/Source/GLTFExporter/Public/Builders/GLTFConvertBuilder.h b/Engine/Plugins/Enterprise/GLTFExporter/MigratedFromGit/Source/GLTFExporter/Public/Builders/GLTFConvertBuilder.h
new file mode 100644
--- /dev/null
+++ b/Engine/Plugins/Enterprise/GLTFExporter/MigratedFromGit/Source/GLTFExporter/Public/Builders/GLTFConvertBuilder.h
@@ -0,0 +1,155 @@
+// Copyright Epic Games, Inc. All Rights Reserved.
+
+#pragma once
+
+#include "Builders/GLTFBufferBuilder.h"
+#include "Converters/GLTFAccessorConverters.h"
+#include "Converters/GLTFMeshConverters.h"
+#include "Converters/GLTFMeshDataConverters.h"
+#include "Converters/GLTFMaterialConverters.h"
+#include "Converters/GLTFSamplerConverters.h"
+#include "Converters/GLTFTextureConverters.h"
+#include "Converters/GLTFImageConverters.h"
+#include "Converters/GLTFNodeConverters.h"
+#include "Converters/GLTFSkinConverters.h"
+#include "Converters/GLTFAnimationConverters.h"
+#include "Converters/GLTFSceneConverters.h"
+#include "Converters/GLTFCameraConverters.h"
+#include "Converters/GLTFLightConverters.h"
+#include "Converters/GLTFBackdropConverters.h"
+#include "Converters/GLTFEpicLevelVariantSetsConverters.h"
+#include "Converters/GLTFLightMapConverters.h"
+#include "Converters/GLTFHotspotConverters.h"
+#include "Converters/GLTFKhrVariantConverters.h"
+#include "Converters/GLTFSkySphereConverters.h"
+
+class FGLTFConvertBuilder : public FGLTFBufferBuilder
+{
+public:
+
+	const TSet<AActor*> SelectedActors;
+
+	FGLTFConvertBuilder(const FString& FileName, const UGLTFExportOptions* ExportOptions, const TSet<AActor*>& SelectedActors = {});
+
+	bool IsSelectedActor(const AActor* Object) const;
+	bool IsRootActor(const AActor* Actor) const;
+
+	FGLTFJsonAccessor* AddUniquePositionAccessor(const FGLTFMeshSection* MeshSection, const FPositionVertexBuffer* VertexBuffer);
+	FGLTFJsonAccessor* AddUniqueColorAccessor(const FGLTFMeshSection* MeshSection, const FColorVertexBuffer* VertexBuffer);
+	FGLTFJsonAccessor* AddUniqueNormalAccessor(const FGLTFMeshSection* MeshSection, const FStaticMeshVertexBuffer* VertexBuffer);
+	FGLTFJsonAccessor* AddUniqueTangentAccessor(const FGLTFMeshSection* MeshSection, const FStaticMeshVertexBuffer* VertexBuffer);
+	FGLTFJsonAccessor* AddUniqueUVAccessor(const FGLTFMeshSection* MeshSection, const FStaticMeshVertexBuffer* VertexBuffer, int32 UVIndex);
+	FGLTFJsonAccessor* AddUniqueJointAccessor(const FGLTFMeshSection* MeshSection, const FSkinWeightVertexBuffer* VertexBuffer, int32 InfluenceOffset);
+	FGLTFJsonAccessor* AddUniqueWeightAccessor(const FGLTFMeshSection* MeshSection, const FSkinWeightVertexBuffer* VertexBuffer, int32 InfluenceOffset);
+	FGLTFJsonAccessor* AddUniqueIndexAccessor(const FGLTFMeshSection* MeshSection);
+
+	FGLTFJsonMesh* AddUniqueMesh(const UStaticMesh* StaticMesh, const FGLTFMaterialArray& Materials = {}, int32 LODIndex = INDEX_NONE);
+	FGLTFJsonMesh* AddUniqueMesh(const USkeletalMesh* SkeletalMesh, const FGLTFMaterialArray& Materials = {}, int32 LODIndex = INDEX_NONE);
+	FGLTFJsonMesh* AddUniqueMesh(const UMeshComponent* MeshComponent, const FGLTFMaterialArray& Materials = {}, int32 LODIndex = INDEX_NONE);
+	FGLTFJsonMesh* AddUniqueMesh(const UStaticMeshComponent* StaticMeshComponent, const FGLTFMaterialArray& Materials = {}, int32 LODIndex = INDEX_NONE);
+	FGLTFJsonMesh* AddUniqueMesh(const USkeletalMeshComponent* SkeletalMeshComponent, const FGLTFMaterialArray& Materials = {}, int32 LODIndex = INDEX_NONE);
+
+	const FGLTFMeshData* AddUniqueMeshData(const UStaticMesh* StaticMesh, const UStaticMeshComponent* StaticMeshComponent = nullptr, int32 LODIndex = INDEX_NONE);
+	const FGLTFMeshData* AddUniqueMeshData(const USkeletalMesh* SkeletalMesh, const USkeletalMeshComponent* SkeletalMeshComponent = nullptr, int32 LODIndex = INDEX_NONE);
+
+	FGLTFJsonMaterial* AddUniqueMaterial(const UMaterialInterface* Material, const UStaticMesh* StaticMesh, int32 LODIndex = INDEX_NONE, int32 MaterialIndex = INDEX_NONE);
+	FGLTFJsonMaterial* AddUniqueMaterial(const UMaterialInterface* Material, const USkeletalMesh* SkeletalMesh, int32 LODIndex = INDEX_NONE, int32 MaterialIndex = INDEX_NONE);
+	FGLTFJsonMaterial* AddUniqueMaterial(const UMaterialInterface* Material, const UMeshComponent* MeshComponent, int32 LODIndex = INDEX_NONE, int32 MaterialIndex = INDEX_NONE);
+	FGLTFJsonMaterial* AddUniqueMaterial(const UMaterialInterface* Material, const UStaticMeshComponent* StaticMeshComponent, int32 LODIndex = INDEX_NONE, int32 MaterialIndex = INDEX_NONE);
+	FGLTFJsonMaterial* AddUniqueMaterial(const UMaterialInterface* Material, const USkeletalMeshComponent* SkeletalMeshComponent, int32 LODIndex = INDEX_NONE, int32 MaterialIndex = INDEX_NONE);
+	FGLTFJsonMaterial* AddUniqueMaterial(const UMaterialInterface* Material, const FGLTFMeshData* MeshData = nullptr, const FGLTFIndexArray& SectionIndices = {});
+
+	FGLTFJsonSampler* AddUniqueSampler(const UTexture* Texture);
+	FGLTFJsonTexture* AddUniqueTexture(const UTexture* Texture);
+	FGLTFJsonTexture* AddUniqueTexture(const UTexture2D* Texture);
+	FGLTFJsonTexture* AddUniqueTexture(const UTextureCube* Texture, ECubeFace CubeFace);
+	FGLTFJsonTexture* AddUniqueTexture(const UTextureRenderTarget2D* Texture);
+	FGLTFJsonTexture* AddUniqueTexture(const UTextureRenderTargetCube* Texture, ECubeFace CubeFace);
+	FGLTFJsonTexture* AddUniqueTexture(const ULightMapTexture2D* Texture);
+	FGLTFJsonTexture* AddUniqueTexture(const UTexture* Texture, bool bToSRGB);
+	FGLTFJsonTexture* AddUniqueTexture(const UTexture2D* Texture, bool bToSRGB);
+	FGLTFJsonTexture* AddUniqueTexture(const UTextureCube* Texture, ECubeFace CubeFace, bool bToSRGB);
+	FGLTFJsonTexture* AddUniqueTexture(const UTextureRenderTarget2D* Texture, bool bToSRGB);
+	FGLTFJsonTexture* AddUniqueTexture(const UTextureRenderTargetCube* Texture, ECubeFace CubeFace, bool bToSRGB);
+	FGLTFJsonImage* AddUniqueImage(TGLTFSharedArray<FColor>& Pixels, FIntPoint Size, bool bIgnoreAlpha, EGLTFTextureType Type, const FString& Name);
+
+	FGLTFJsonSkin* AddUniqueSkin(FGLTFJsonNode* RootNode, const USkeletalMesh* SkeletalMesh);
+	FGLTFJsonSkin* AddUniqueSkin(FGLTFJsonNode* RootNode, const USkeletalMeshComponent* SkeletalMeshComponent);
+	FGLTFJsonAnimation* AddUniqueAnimation(FGLTFJsonNode* RootNode, const USkeletalMesh* SkeletalMesh, const UAnimSequence* AnimSequence);
+	FGLTFJsonAnimation* AddUniqueAnimation(FGLTFJsonNode* RootNode, const USkeletalMeshComponent* SkeletalMeshComponent);
+	FGLTFJsonAnimation* AddUniqueAnimation(const ULevel* Level, const ULevelSequence* LevelSequence);
+	FGLTFJsonAnimation* AddUniqueAnimation(const ALevelSequenceActor* LevelSequenceActor);
+
+	FGLTFJsonNode* AddUniqueNode(const AActor* Actor);
+	FGLTFJsonNode* AddUniqueNode(const USceneComponent* SceneComponent);
+	FGLTFJsonNode* AddUniqueNode(const USceneComponent* SceneComponent, FName SocketName);
+	FGLTFJsonNode* AddUniqueNode(FGLTFJsonNode* RootNode, const UStaticMesh* StaticMesh, FName SocketName);
+	FGLTFJsonNode* AddUniqueNode(FGLTFJsonNode* RootNode, const USkeletalMesh* SkeletalMesh, FName SocketName);
+	FGLTFJsonNode* AddUniqueNode(FGLTFJsonNode* RootNode, const USkeletalMesh* SkeletalMesh, int32 BoneIndex);
+	FGLTFJsonScene* AddUniqueScene(const UWorld* World);
+
+	FGLTFJsonCamera* AddUniqueCamera(const UCameraComponent* CameraComponent);
+	FGLTFJsonLight* AddUniqueLight(const ULightComponent* LightComponent);
+	FGLTFJsonBackdrop* AddUniqueBackdrop(const AActor* BackdropActor);
+	FGLTFJsonLightMap* AddUniqueLightMap(const UStaticMeshComponent* StaticMeshComponent);
+	FGLTFJsonHotspot* AddUniqueHotspot(const AGLTFHotspotActor* HotspotActor);
+	FGLTFJsonSkySphere* AddUniqueSkySphere(const AActor* SkySphereActor);
+
+	FGLTFJsonEpicLevelVariantSets* AddUniqueEpicLevelVariantSets(const ULevelVariantSets* LevelVariantSets);
+	FGLTFJsonKhrMaterialVariant* AddUniqueKhrMaterialVariant(const UVariant* Variant);
+
+	void RegisterObjectVariant(const UObject* Object, const UPropertyValue* Property);
+	const TArray<const UPropertyValue*>* GetObjectVariants(const UObject* Object) const;
+
+	TUniquePtr<IGLTFPositionBufferConverter> PositionBufferConverter = MakeUnique<FGLTFPositionBufferConverter>(*this);
+	TUniquePtr<IGLTFColorBufferConverter> ColorBufferConverter = MakeUnique<FGLTFColorBufferConverter>(*this);
+	TUniquePtr<IGLTFNormalBufferConverter> NormalBufferConverter = MakeUnique<FGLTFNormalBufferConverter>(*this);
+	TUniquePtr<IGLTFTangentBufferConverter> TangentBufferConverter = MakeUnique<FGLTFTangentBufferConverter>(*this);
+	TUniquePtr<IGLTFUVBufferConverter> UVBufferConverter = MakeUnique<FGLTFUVBufferConverter>(*this);
+	TUniquePtr<IGLTFBoneIndexBufferConverter> BoneIndexBufferConverter = MakeUnique<FGLTFBoneIndexBufferConverter>(*this);
+	TUniquePtr<IGLTFBoneWeightBufferConverter> BoneWeightBufferConverter = MakeUnique<FGLTFBoneWeightBufferConverter>(*this);
+	TUniquePtr<IGLTFIndexBufferConverter> IndexBufferConverter = MakeUnique<FGLTFIndexBufferConverter>(*this);
+
+	TUniquePtr<IGLTFStaticMeshConverter> StaticMeshConverter = MakeUnique<FGLTFStaticMeshConverter>(*this);
+	TUniquePtr<IGLTFSkeletalMeshConverter> SkeletalMeshConverter = MakeUnique<FGLTFSkeletalMeshConverter>(*this);
+
+	TUniquePtr<IGLTFMaterialConverter> MaterialConverter = MakeUnique<FGLTFMaterialConverter>(*this);
+	TUniquePtr<IGLTFStaticMeshDataConverter> StaticMeshDataConverter = MakeUnique<FGLTFStaticMeshDataConverter>(*this);
+	TUniquePtr<IGLTFSkeletalMeshDataConverter> SkeletalMeshDataConverter = MakeUnique<FGLTFSkeletalMeshDataConverter>(*this);
+
+	TUniquePtr<IGLTFSamplerConverter> SamplerConverter = MakeUnique<FGLTFSamplerConverter>(*this);
+	TUniquePtr<IGLTFTexture2DConverter> Texture2DConverter = MakeUnique<FGLTFTexture2DConverter>(*this);
+	TUniquePtr<IGLTFTextureCubeConverter> TextureCubeConverter = MakeUnique<FGLTFTextureCubeConverter>(*this);
+	TUniquePtr<IGLTFTextureRenderTarget2DConverter> TextureRenderTarget2DConverter = MakeUnique<FGLTFTextureRenderTarget2DConverter>(*this);
+	TUniquePtr<IGLTFTextureRenderTargetCubeConverter> TextureRenderTargetCubeConverter = MakeUnique<FGLTFTextureRenderTargetCubeConverter>(*this);
+	TUniquePtr<IGLTFTextureLightMapConverter> TextureLightMapConverter = MakeUnique<FGLTFTextureLightMapConverter>(*this);
+	TUniquePtr<IGLTFImageConverter> ImageConverter = MakeUnique<FGLTFImageConverter>(*this);
+
+	TUniquePtr<IGLTFSkinConverter> SkinConverter = MakeUnique<FGLTFSkinConverter>(*this);
+	TUniquePtr<IGLTFAnimationConverter> AnimationConverter = MakeUnique<FGLTFAnimationConverter>(*this);
+	TUniquePtr<IGLTFAnimationDataConverter> AnimationDataConverter = MakeUnique<FGLTFAnimationDataConverter>(*this);
+	TUniquePtr<IGLTFLevelSequenceConverter> LevelSequenceConverter = MakeUnique<FGLTFLevelSequenceConverter>(*this);
+	TUniquePtr<IGLTFLevelSequenceDataConverter> LevelSequenceDataConverter = MakeUnique<FGLTFLevelSequenceDataConverter>(*this);
+
+	TUniquePtr<IGLTFActorConverter> ActorConverter = MakeUnique<FGLTFActorConverter>(*this);
+	TUniquePtr<IGLTFComponentConverter> ComponentConverter = MakeUnique<FGLTFComponentConverter>(*this);
+	TUniquePtr<IGLTFComponentSocketConverter> ComponentSocketConverter = MakeUnique<FGLTFComponentSocketConverter>(*this);
+	TUniquePtr<IGLTFStaticSocketConverter> StaticSocketConverter = MakeUnique<FGLTFStaticSocketConverter>(*this);
+	TUniquePtr<IGLTFSkeletalSocketConverter> SkeletalSocketConverter = MakeUnique<FGLTFSkeletalSocketConverter>(*this);
+	TUniquePtr<IGLTFSkeletalBoneConverter> SkeletalBoneConverter = MakeUnique<FGLTFSkeletalBoneConverter>(*this);
+	TUniquePtr<IGLTFSceneConverter> SceneConverter = MakeUnique<FGLTFSceneConverter>(*this);
+
+	TUniquePtr<IGLTFCameraConverter> CameraConverter = MakeUnique<FGLTFCameraConverter>(*this);
+	TUniquePtr<IGLTFLightConverter> LightConverter = MakeUnique<FGLTFLightConverter>(*this);
+	TUniquePtr<IGLTFBackdropConverter> BackdropConverter = MakeUnique<FGLTFBackdropConverter>(*this);
+	TUniquePtr<IGLTFLightMapConverter> LightMapConverter = MakeUnique<FGLTFLightMapConverter>(*this);
+	TUniquePtr<IGLTFHotspotConverter> HotspotConverter = MakeUnique<FGLTFHotspotConverter>(*this);
+	TUniquePtr<IGLTFSkySphereConverter> SkySphereConverter = MakeUnique<FGLTFSkySphereConverter>(*this);
+
+	TUniquePtr<IGLTFEpicLevelVariantSetsConverter> EpicLevelVariantSetsConverter = MakeUnique<FGLTFEpicLevelVariantSetsConverter>(*this);
+	TUniquePtr<IGLTFKhrMaterialVariantConverter> KhrMaterialVariantConverter = MakeUnique<FGLTFKhrMaterialVariantConverter>(*this);
+
+private:
+
+	TMap<const UObject*, TArray<const UPropertyValue*>> ObjectVariants;
+};
diff --git a/Engine/Plugins/Enterprise/GLTFExporter/MigratedFromGit/Source/GLTFExporter/Public/Builders/GLTFFileBuilder.h b/Engine/Plugins/Enterprise/GLTFExporter/MigratedFromGit/Source/GLTFExporter/Public/Builders/GLTFFileBuilder.h
new file mode 100644
--- /dev/null
+++ b/Engine/Plugins/Enterprise/GLTFExporter/MigratedFromGit/Source/GLTFExporter/Public/Builders/GLTFFileBuilder.h
@@ -0,0 +1,29 @@
+// Copyright Epic Games, Inc. All Rights Reserved.
+
+#pragma once
+
+#include "Builders/GLTFTaskBuilder.h"
+#include "Builders/GLTFMemoryArchive.h"
+
+class FGLTFFileBuilder : public FGLTFTaskBuilder
+{
+public:
+
+	FGLTFFileBuilder(const FString& FileName, const UGLTFExportOptions* ExportOptions);
+
+	FString AddExternalFile(const FString& URI, const TSharedPtr<FGLTFMemoryArchive>& Archive = MakeShared<FGLTFMemoryArchive>());
+
+	const TMap<FString, TSharedPtr<FGLTFMemoryArchive>>& GetExternalFiles() const;
+
+	bool WriteExternalFiles(const FString& DirPath, bool bOverwrite = true);
+
+private:
+
+	TMap<FString, TSharedPtr<FGLTFMemoryArchive>> ExternalFiles;
+
+	FString GetUniqueURI(const FString& URI) const;
+
+protected:
+
+	bool SaveToFile(const FString& FilePath, const TArray64<uint8>& FileData, bool bOverwrite);
+};
diff --git a/Engine/Plugins/Enterprise/GLTFExporter/MigratedFromGit/Source/GLTFExporter/Public/Builders/GLTFJsonBuilder.h b/Engine/Plugins/Enterprise/GLTFExporter/MigratedFromGit/Source/GLTFExporter/Public/Builders/GLTFJsonBuilder.h
new file mode 100644
--- /dev/null
+++ b/Engine/Plugins/Enterprise/GLTFExporter/MigratedFromGit/Source/GLTFExporter/Public/Builders/GLTFJsonBuilder.h
@@ -0,0 +1,54 @@
+// Copyright Epic Games, Inc. All Rights Reserved.
+
+#pragma once
+
+#include "Json/GLTFJsonRoot.h"
+#include "Builders/GLTFFileBuilder.h"
+
+class FGLTFJsonBuilder : public FGLTFFileBuilder
+{
+public:
+
+	FGLTFJsonScene*& DefaultScene;
+
+	FGLTFJsonBuilder(const FString& FileName, const UGLTFExportOptions* ExportOptions);
+
+	void WriteJson(FArchive& Archive);
+
+	void AddExtension(EGLTFJsonExtension Extension, bool bIsRequired = false);
+
+	FGLTFJsonAccessor* AddAccessor();
+	FGLTFJsonAnimation* AddAnimation();
+	FGLTFJsonBuffer* AddBuffer();
+	FGLTFJsonBufferView* AddBufferView();
+	FGLTFJsonCamera* AddCamera();
+	FGLTFJsonImage* AddImage();
+	FGLTFJsonMaterial* AddMaterial();
+	FGLTFJsonMesh* AddMesh();
+	FGLTFJsonNode* AddNode();
+	FGLTFJsonSampler* AddSampler();
+	FGLTFJsonScene* AddScene();
+	FGLTFJsonSkin* AddSkin();
+	FGLTFJsonTexture* AddTexture();
+	FGLTFJsonBackdrop* AddBackdrop();
+	FGLTFJsonHotspot* AddHotspot();
+	FGLTFJsonLight* AddLight();
+	FGLTFJsonLightMap* AddLightMap();
+	FGLTFJsonSkySphere* AddSkySphere();
+	FGLTFJsonEpicLevelVariantSets* AddEpicLevelVariantSets();
+	FGLTFJsonKhrMaterialVariant* AddKhrMaterialVariant();
+
+	const FGLTFJsonRoot& GetRoot() const;
+
+protected:
+
+	TSet<EGLTFJsonExtension> GetCustomExtensionsUsed() const;
+
+private:
+
+	FString GetGeneratorString() const;
+
+	static bool IsCustomExtension(EGLTFJsonExtension Extension);
+
+	FGLTFJsonRoot JsonRoot;
+};
diff --git a/Engine/Plugins/Enterprise/GLTFExporter/MigratedFromGit/Source/GLTFExporter/Public/Builders/GLTFLogBuilder.h b/Engine/Plugins/Enterprise/GLTFExporter/MigratedFromGit/Source/GLTFExporter/Public/Builders/GLTFLogBuilder.h
new file mode 100644
--- /dev/null
+++ b/Engine/Plugins/Enterprise/GLTFExporter/MigratedFromGit/Source/GLTFExporter/Public/Builders/GLTFLogBuilder.h
@@ -0,0 +1,53 @@
+// Copyright Epic Games, Inc. All Rights Reserved.
+
+#pragma once
+
+#include "Builders/GLTFBuilder.h"
+
+#if WITH_EDITOR
+class IMessageLogListing;
+#endif
+
+class FGLTFLogBuilder : public FGLTFBuilder
+{
+public:
+
+	FGLTFLogBuilder(const FString& FileName, const UGLTFExportOptions* ExportOptions);
+
+	void LogSuggestion(const FString& Message);
+
+	void LogWarning(const FString& Message);
+
+	void LogError(const FString& Message);
+
+	const TArray<FString>& GetLoggedSuggestions() const;
+
+	const TArray<FString>& GetLoggedWarnings() const;
+
+	const TArray<FString>& GetLoggedErrors() const;
+
+	bool HasLoggedMessages() const;
+
+	void OpenLog() const;
+
+	void ClearLog();
+
+private:
+
+	enum class ELogLevel
+	{
+		Suggestion,
+		Warning,
+		Error,
+	};
+
+	void PrintToLog(ELogLevel Level, const FString& Message) const;
+
+	TArray<FString> Suggestions;
+	TArray<FString> Warnings;
+	TArray<FString> Errors;
+
+#if WITH_EDITOR
+	TSharedPtr<IMessageLogListing> LogListing;
+#endif
+};
diff --git a/Engine/Plugins/Enterprise/GLTFExporter/MigratedFromGit/Source/GLTFExporter/Public/Builders/GLTFMemoryArchive.h b/Engine/Plugins/Enterprise/GLTFExporter/MigratedFromGit/Source/GLTFExporter/Public/Builders/GLTFMemoryArchive.h
new file mode 100644
--- /dev/null
+++ b/Engine/Plugins/Enterprise/GLTFExporter/MigratedFromGit/Source/GLTFExporter/Public/Builders/GLTFMemoryArchive.h
@@ -0,0 +1,32 @@
+// Copyright Epic Games, Inc. All Rights Reserved.
+
+#pragma once
+
+#include "CoreMinimal.h"
+#include "Serialization/MemoryArchive.h"
+
+class FGLTFMemoryArchive final : public FMemoryArchive, public TArray64<uint8>
+{
+public:
+
+	FGLTFMemoryArchive()
+	{
+		FMemoryArchive::SetIsSaving(true);
+	}
+
+	virtual void Serialize(void* Data, int64 Length) override
+	{
+		check(Length >= 0);
+
+		const int64 NumBytesToAdd = Offset + Length - Num();
+		AddUninitialized(NumBytesToAdd);
+
+		FMemory::Memcpy(GetData() + Offset, Data, Length);
+		Offset += Length;
+	}
+
+	virtual int64 TotalSize() override
+	{
+		return Num();
+	}
+};
diff --git a/Engine/Plugins/Enterprise/GLTFExporter/MigratedFromGit/Source/GLTFExporter/Public/Builders/GLTFTaskBuilder.h b/Engine/Plugins/Enterprise/GLTFExporter/MigratedFromGit/Source/GLTFExporter/Public/Builders/GLTFTaskBuilder.h
new file mode 100644
--- /dev/null
+++ b/Engine/Plugins/Enterprise/GLTFExporter/MigratedFromGit/Source/GLTFExporter/Public/Builders/GLTFTaskBuilder.h
@@ -0,0 +1,36 @@
+// Copyright Epic Games, Inc. All Rights Reserved.
+
+#pragma once
+
+#include "Builders/GLTFLogBuilder.h"
+#include "Tasks/GLTFDelayedTask.h"
+
+class FGLTFTaskBuilder : public FGLTFLogBuilder
+{
+public:
+
+	FGLTFTaskBuilder(const FString& FileName, const UGLTFExportOptions* ExportOptions);
+
+	template <typename TaskType, typename... TaskArgTypes, typename = typename TEnableIf<TIsDerivedFrom<TaskType, FGLTFDelayedTask>::Value>::Type>
+	bool ScheduleSlowTask(TaskArgTypes&&... Args)
+	{
+		return ScheduleSlowTask(MakeUnique<TaskType>(Forward<TaskArgTypes>(Args)...));
+	}
+
+	template <typename TaskType, typename = typename TEnableIf<TIsDerivedFrom<TaskType, FGLTFDelayedTask>::Value>::Type>
+	bool ScheduleSlowTask(TUniquePtr<TaskType> Task)
+	{
+		return ScheduleSlowTask(TUniquePtr<FGLTFDelayedTask>(Task.Release()));
+	}
+
+	bool ScheduleSlowTask(TUniquePtr<FGLTFDelayedTask> Task);
+
+	void ProcessSlowTasks(FFeedbackContext* Context = GWarn);
+
+private:
+
+	static FText GetPriorityMessageFormat(EGLTFTaskPriority Priority);
+
+	int32 PriorityIndexLock;
+	TMap<EGLTFTaskPriority, TArray<TUniquePtr<FGLTFDelayedTask>>> TasksByPriority;
+};
diff --git a/Engine/Plugins/Enterprise/GLTFExporter/MigratedFromGit/Source/GLTFExporter/Public/Converters/GLTFAccessorConverters.h b/Engine/Plugins/Enterprise/GLTFExporter/MigratedFromGit/Source/GLTFExporter/Public/Converters/GLTFAccessorConverters.h
new file mode 100644
--- /dev/null
+++ b/Engine/Plugins/Enterprise/GLTFExporter/MigratedFromGit/Source/GLTFExporter/Public/Converters/GLTFAccessorConverters.h
@@ -0,0 +1,100 @@
+// Copyright Epic Games, Inc. All Rights Reserved.
+
+#pragma once
+
+#include "Json/GLTFJsonCore.h"
+#include "Converters/GLTFConverter.h"
+#include "Converters/GLTFBuilderContext.h"
+#include "Converters/GLTFMeshSection.h"
+
+typedef TGLTFConverter<FGLTFJsonAccessor*, const FGLTFMeshSection*, const FPositionVertexBuffer*> IGLTFPositionBufferConverter;
+typedef TGLTFConverter<FGLTFJsonAccessor*, const FGLTFMeshSection*, const FColorVertexBuffer*> IGLTFColorBufferConverter;
+typedef TGLTFConverter<FGLTFJsonAccessor*, const FGLTFMeshSection*, const FStaticMeshVertexBuffer*> IGLTFNormalBufferConverter;
+typedef TGLTFConverter<FGLTFJsonAccessor*, const FGLTFMeshSection*, const FStaticMeshVertexBuffer*> IGLTFTangentBufferConverter;
+typedef TGLTFConverter<FGLTFJsonAccessor*, const FGLTFMeshSection*, const FStaticMeshVertexBuffer*, uint32> IGLTFUVBufferConverter;
+typedef TGLTFConverter<FGLTFJsonAccessor*, const FGLTFMeshSection*, const FSkinWeightVertexBuffer*, uint32> IGLTFBoneIndexBufferConverter;
+typedef TGLTFConverter<FGLTFJsonAccessor*, const FGLTFMeshSection*, const FSkinWeightVertexBuffer*, uint32> IGLTFBoneWeightBufferConverter;
+typedef TGLTFConverter<FGLTFJsonAccessor*, const FGLTFMeshSection*> IGLTFIndexBufferConverter;
+
+class FGLTFPositionBufferConverter final : public FGLTFBuilderContext, public IGLTFPositionBufferConverter
+{
+	using FGLTFBuilderContext::FGLTFBuilderContext;
+
+	virtual FGLTFJsonAccessor* Convert(const FGLTFMeshSection* MeshSection, const FPositionVertexBuffer* VertexBuffer) override;
+};
+
+class FGLTFColorBufferConverter final : public FGLTFBuilderContext, public IGLTFColorBufferConverter
+{
+	using FGLTFBuilderContext::FGLTFBuilderContext;
+
+	virtual FGLTFJsonAccessor* Convert(const FGLTFMeshSection* MeshSection, const FColorVertexBuffer* VertexBuffer) override;
+};
+
+class FGLTFNormalBufferConverter final : public FGLTFBuilderContext, public IGLTFNormalBufferConverter
+{
+	using FGLTFBuilderContext::FGLTFBuilderContext;
+
+	virtual FGLTFJsonAccessor* Convert(const FGLTFMeshSection* MeshSection, const FStaticMeshVertexBuffer* VertexBuffer) override;
+
+	template <typename DestinationType, typename SourceType>
+	FGLTFJsonBufferView* ConvertBufferView(const FGLTFMeshSection* MeshSection, const void* TangentData);
+};
+
+class FGLTFTangentBufferConverter final : public FGLTFBuilderContext, public IGLTFTangentBufferConverter
+{
+	using FGLTFBuilderContext::FGLTFBuilderContext;
+
+	virtual FGLTFJsonAccessor* Convert(const FGLTFMeshSection* MeshSection, const FStaticMeshVertexBuffer* VertexBuffer) override;
+
+	template <typename DestinationType, typename SourceType>
+	FGLTFJsonBufferView* ConvertBufferView(const FGLTFMeshSection* MeshSection, const void* TangentData);
+};
+
+class FGLTFUVBufferConverter final : public FGLTFBuilderContext, public IGLTFUVBufferConverter
+{
+	using FGLTFBuilderContext::FGLTFBuilderContext;
+
+	virtual FGLTFJsonAccessor* Convert(const FGLTFMeshSection* MeshSection, const FStaticMeshVertexBuffer* VertexBuffer, uint32 UVIndex) override;
+
+	template <typename SourceType>
+	FGLTFJsonAccessor* Convert(const FGLTFMeshSection* MeshSection, const FStaticMeshVertexBuffer* VertexBuffer, uint32 UVIndex, const uint8* SourceData) const;
+};
+
+class FGLTFBoneIndexBufferConverter final : public FGLTFBuilderContext, public IGLTFBoneIndexBufferConverter
+{
+	using FGLTFBuilderContext::FGLTFBuilderContext;
+
+	virtual FGLTFJsonAccessor* Convert(const FGLTFMeshSection* MeshSection, const FSkinWeightVertexBuffer* VertexBuffer, uint32 InfluenceOffset) override;
+
+	template <typename DestinationType>
+	FGLTFJsonAccessor* Convert(const FGLTFMeshSection* MeshSection, const FSkinWeightVertexBuffer* VertexBuffer, uint32 InfluenceOffset, const uint8* SourceData) const;
+
+	template <typename DestinationType, typename SourceType>
+	FGLTFJsonAccessor* Convert(const FGLTFMeshSection* MeshSection, const FSkinWeightVertexBuffer* VertexBuffer, uint32 InfluenceOffset, const uint8* SourceData) const;
+
+	template <typename DestinationType, typename SourceType, typename CallbackType>
+	FGLTFJsonAccessor* Convert(const FGLTFMeshSection* MeshSection, const FSkinWeightVertexBuffer* VertexBuffer, uint32 InfluenceOffset, const uint8* SourceData, CallbackType GetVertexInfluenceOffsetCount) const;
+};
+
+class FGLTFBoneWeightBufferConverter final : public FGLTFBuilderContext, public IGLTFBoneWeightBufferConverter
+{
+	using FGLTFBuilderContext::FGLTFBuilderContext;
+
+	virtual FGLTFJsonAccessor* Convert(const FGLTFMeshSection* MeshSection, const FSkinWeightVertexBuffer* VertexBuffer, uint32 InfluenceOffset) override;
+
+	template <typename BoneIndexType>
+	FGLTFJsonAccessor* Convert(const FGLTFMeshSection* MeshSection, const FSkinWeightVertexBuffer* VertexBuffer, uint32 InfluenceOffset, const uint8* SourceData) const;
+
+	template <typename BoneIndexType, typename CallbackType>
+	FGLTFJsonAccessor* Convert(const FGLTFMeshSection* MeshSection, const FSkinWeightVertexBuffer* VertexBuffer, uint32 InfluenceOffset, const uint8* SourceData, CallbackType GetVertexInfluenceOffsetCount) const;
+};
+
+class FGLTFIndexBufferConverter final : public FGLTFBuilderContext, public IGLTFIndexBufferConverter
+{
+	using FGLTFBuilderContext::FGLTFBuilderContext;
+
+	virtual FGLTFJsonAccessor* Convert(const FGLTFMeshSection* MeshSection) override;
+
+	template <typename IndexType>
+	FGLTFJsonAccessor* Convert(const FGLTFMeshSection* MeshSection) const;
+};
diff --git a/Engine/Plugins/Enterprise/GLTFExporter/MigratedFromGit/Source/GLTFExporter/Public/Converters/GLTFAnimationConverters.h b/Engine/Plugins/Enterprise/GLTFExporter/MigratedFromGit/Source/GLTFExporter/Public/Converters/GLTFAnimationConverters.h
new file mode 100644
--- /dev/null
+++ b/Engine/Plugins/Enterprise/GLTFExporter/MigratedFromGit/Source/GLTFExporter/Public/Converters/GLTFAnimationConverters.h
@@ -0,0 +1,43 @@
+// Copyright Epic Games, Inc. All Rights Reserved.
+
+#pragma once
+
+#include "Json/GLTFJsonCore.h"
+#include "Converters/GLTFConverter.h"
+#include "Converters/GLTFBuilderContext.h"
+
+class ULevelSequence;
+class ALevelSequenceActor;
+
+typedef TGLTFConverter<FGLTFJsonAnimation*, FGLTFJsonNode*, const USkeletalMesh*, const UAnimSequence*> IGLTFAnimationConverter;
+typedef TGLTFConverter<FGLTFJsonAnimation*, FGLTFJsonNode*, const USkeletalMeshComponent*> IGLTFAnimationDataConverter;
+typedef TGLTFConverter<FGLTFJsonAnimation*, const ULevel*, const ULevelSequence*> IGLTFLevelSequenceConverter;
+typedef TGLTFConverter<FGLTFJsonAnimation*, const ALevelSequenceActor*> IGLTFLevelSequenceDataConverter;
+
+class FGLTFAnimationConverter final : public FGLTFBuilderContext, public IGLTFAnimationConverter
+{
+	using FGLTFBuilderContext::FGLTFBuilderContext;
+
+	virtual FGLTFJsonAnimation* Convert(FGLTFJsonNode* RootNode, const USkeletalMesh* SkeletalMesh, const UAnimSequence* AnimSequence) override;
+};
+
+class FGLTFAnimationDataConverter final : public FGLTFBuilderContext, public IGLTFAnimationDataConverter
+{
+	using FGLTFBuilderContext::FGLTFBuilderContext;
+
+	virtual FGLTFJsonAnimation* Convert(FGLTFJsonNode* RootNode, const USkeletalMeshComponent* SkeletalMeshComponent) override;
+};
+
+class FGLTFLevelSequenceConverter final : public FGLTFBuilderContext, public IGLTFLevelSequenceConverter
+{
+	using FGLTFBuilderContext::FGLTFBuilderContext;
+
+	virtual FGLTFJsonAnimation* Convert(const ULevel* Level, const ULevelSequence*) override;
+};
+
+class FGLTFLevelSequenceDataConverter final : public FGLTFBuilderContext, public IGLTFLevelSequenceDataConverter
+{
+	using FGLTFBuilderContext::FGLTFBuilderContext;
+
+	virtual FGLTFJsonAnimation* Convert(const ALevelSequenceActor* LevelSequenceActor) override;
+};
diff --git a/Engine/Plugins/Enterprise/GLTFExporter/MigratedFromGit/Source/GLTFExporter/Public/Converters/GLTFBackdropConverters.h b/Engine/Plugins/Enterprise/GLTFExporter/MigratedFromGit/Source/GLTFExporter/Public/Converters/GLTFBackdropConverters.h
new file mode 100644
--- /dev/null
+++ b/Engine/Plugins/Enterprise/GLTFExporter/MigratedFromGit/Source/GLTFExporter/Public/Converters/GLTFBackdropConverters.h
@@ -0,0 +1,17 @@
+// Copyright Epic Games, Inc. All Rights Reserved.
+
+#pragma once
+
+#include "Json/GLTFJsonCore.h"
+#include "Converters/GLTFConverter.h"
+#include "Converters/GLTFBuilderContext.h"
+
+typedef TGLTFConverter<FGLTFJsonBackdrop*, const AActor*> IGLTFBackdropConverter;
+
+
+class FGLTFBackdropConverter final : public FGLTFBuilderContext, public IGLTFBackdropConverter
+{
+	using FGLTFBuilderContext::FGLTFBuilderContext;
+
+	virtual FGLTFJsonBackdrop* Convert(const AActor* BackdropActor) override;
+};
diff --git a/Engine/Plugins/Enterprise/GLTFExporter/MigratedFromGit/Source/GLTFExporter/Public/Converters/GLTFBuilderContext.h b/Engine/Plugins/Enterprise/GLTFExporter/MigratedFromGit/Source/GLTFExporter/Public/Converters/GLTFBuilderContext.h
new file mode 100644
--- /dev/null
+++ b/Engine/Plugins/Enterprise/GLTFExporter/MigratedFromGit/Source/GLTFExporter/Public/Converters/GLTFBuilderContext.h
@@ -0,0 +1,19 @@
+// Copyright Epic Games, Inc. All Rights Reserved.
+
+#pragma once
+
+class FGLTFConvertBuilder;
+
+class FGLTFBuilderContext
+{
+public:
+
+	FGLTFBuilderContext(FGLTFConvertBuilder& Builder)
+		: Builder(Builder)
+	{
+	}
+
+protected:
+
+	FGLTFConvertBuilder& Builder;
+};
diff --git a/Engine/Plugins/Enterprise/GLTFExporter/MigratedFromGit/Source/GLTFExporter/Public/Converters/GLTFCameraConverters.h b/Engine/Plugins/Enterprise/GLTFExporter/MigratedFromGit/Source/GLTFExporter/Public/Converters/GLTFCameraConverters.h
new file mode 100644
--- /dev/null
+++ b/Engine/Plugins/Enterprise/GLTFExporter/MigratedFromGit/Source/GLTFExporter/Public/Converters/GLTFCameraConverters.h
@@ -0,0 +1,18 @@
+// Copyright Epic Games, Inc. All Rights Reserved.
+
+#pragma once
+
+#include "Json/GLTFJsonCore.h"
+#include "Converters/GLTFConverter.h"
+#include "Converters/GLTFBuilderContext.h"
+
+class UCameraComponent;
+
+typedef TGLTFConverter<FGLTFJsonCamera*, const UCameraComponent*> IGLTFCameraConverter;
+
+class FGLTFCameraConverter final : public FGLTFBuilderContext, public IGLTFCameraConverter
+{
+	using FGLTFBuilderContext::FGLTFBuilderContext;
+
+	virtual FGLTFJsonCamera* Convert(const UCameraComponent* CameraComponent) override;
+};
diff --git a/Engine/Plugins/Enterprise/GLTFExporter/MigratedFromGit/Source/GLTFExporter/Public/Converters/GLTFConverter.h b/Engine/Plugins/Enterprise/GLTFExporter/MigratedFromGit/Source/GLTFExporter/Public/Converters/GLTFConverter.h
new file mode 100644
--- /dev/null
+++ b/Engine/Plugins/Enterprise/GLTFExporter/MigratedFromGit/Source/GLTFExporter/Public/Converters/GLTFConverter.h
@@ -0,0 +1,55 @@
+// Copyright Epic Games, Inc. All Rights Reserved.
+
+#pragma once
+
+#include "CoreMinimal.h"
+
+template <typename OutputType, typename... InputTypes>
+class TGLTFConverter
+{
+	typedef TTuple<InputTypes...> InputTupleType;
+
+public:
+
+	virtual ~TGLTFConverter() = default;
+
+	OutputType Get(InputTypes... Inputs) const
+	{
+		Sanitize(Inputs...);
+		InputTupleType InputTuple(MoveTemp(Inputs)...);
+		const uint32 InputHash = GetTypeHash(InputTuple);
+
+		if (OutputType* CachedOutput = CachedOutputs.FindByHash(InputHash, InputTuple))
+		{
+			return *CachedOutput;
+		}
+
+		return {};
+	}
+
+	OutputType GetOrAdd(InputTypes... Inputs)
+	{
+		Sanitize(Inputs...);
+		InputTupleType InputTuple(MoveTemp(Inputs)...);
+		const uint32 InputHash = GetTypeHash(InputTuple);
+
+		if (OutputType* CachedOutput = CachedOutputs.FindByHash(InputHash, InputTuple))
+		{
+			return *CachedOutput;
+		}
+
+		OutputType NewOutput = InputTuple.ApplyAfter(&TGLTFConverter::Convert, this);
+		CachedOutputs.AddByHash(InputHash, MoveTemp(InputTuple), NewOutput);
+		return NewOutput;
+	}
+
+protected:
+
+	virtual void Sanitize(InputTypes&... Inputs) { }
+
+	virtual OutputType Convert(InputTypes... Inputs) = 0;
+
+private:
+
+	TMap<InputTupleType, OutputType> CachedOutputs;
+};
diff --git a/Engine/Plugins/Enterprise/GLTFExporter/MigratedFromGit/Source/GLTFExporter/Public/Converters/GLTFEpicLevelVariantSetsConverters.h b/Engine/Plugins/Enterprise/GLTFExporter/MigratedFromGit/Source/GLTFExporter/Public/Converters/GLTFEpicLevelVariantSetsConverters.h
new file mode 100644
--- /dev/null
+++ b/Engine/Plugins/Enterprise/GLTFExporter/MigratedFromGit/Source/GLTFExporter/Public/Converters/GLTFEpicLevelVariantSetsConverters.h
@@ -0,0 +1,32 @@
+// Copyright Epic Games, Inc. All Rights Reserved.
+
+#pragma once
+
+#include "Json/GLTFJsonCore.h"
+#include "Json/GLTFJsonEpicLevelVariantSets.h"
+#include "Converters/GLTFConverter.h"
+#include "Converters/GLTFBuilderContext.h"
+
+class ULevelVariantSets;
+class UVariantSet;
+class UVariant;
+class UVariantObjectBinding;
+class UPropertyValue;
+class UPropertyValueMaterial;
+
+typedef TGLTFConverter<FGLTFJsonEpicLevelVariantSets*, const ULevelVariantSets*> IGLTFEpicLevelVariantSetsConverter;
+
+class FGLTFEpicLevelVariantSetsConverter final : public FGLTFBuilderContext, public IGLTFEpicLevelVariantSetsConverter
+{
+	using FGLTFBuilderContext::FGLTFBuilderContext;
+
+	virtual FGLTFJsonEpicLevelVariantSets* Convert(const ULevelVariantSets* LevelVariantSets) override;
+
+	bool TryParseVariant(FGLTFJsonEpicVariant& OutVariant, const UVariant* Variant) const;
+	bool TryParseVariantBinding(FGLTFJsonEpicVariant& OutVariant, const UVariantObjectBinding* Binding) const;
+	bool TryParseVisibilityPropertyValue(FGLTFJsonEpicVariant& OutVariant, const UPropertyValue* Property) const;
+	bool TryParseMaterialPropertyValue(FGLTFJsonEpicVariant& OutVariant, const UPropertyValueMaterial* Property) const;
+
+	template <typename MeshType>
+	bool TryParseMeshPropertyValue(FGLTFJsonEpicVariant& OutVariant, const UPropertyValue* Property) const;
+};
diff --git a/Engine/Plugins/Enterprise/GLTFExporter/MigratedFromGit/Source/GLTFExporter/Public/Converters/GLTFHotspotConverters.h b/Engine/Plugins/Enterprise/GLTFExporter/MigratedFromGit/Source/GLTFExporter/Public/Converters/GLTFHotspotConverters.h
new file mode 100644
--- /dev/null
+++ b/Engine/Plugins/Enterprise/GLTFExporter/MigratedFromGit/Source/GLTFExporter/Public/Converters/GLTFHotspotConverters.h
@@ -0,0 +1,17 @@
+// Copyright Epic Games, Inc. All Rights Reserved.
+
+#pragma once
+
+#include "Json/GLTFJsonCore.h"
+#include "Converters/GLTFConverter.h"
+#include "Converters/GLTFBuilderContext.h"
+#include "Actors/GLTFHotspotActor.h"
+
+typedef TGLTFConverter<FGLTFJsonHotspot*, const AGLTFHotspotActor*> IGLTFHotspotConverter;
+
+class FGLTFHotspotConverter final : public FGLTFBuilderContext, public IGLTFHotspotConverter
+{
+	using FGLTFBuilderContext::FGLTFBuilderContext;
+
+	virtual FGLTFJsonHotspot* Convert(const AGLTFHotspotActor* HotspotActor) override;
+};
diff --git a/Engine/Plugins/Enterprise/GLTFExporter/MigratedFromGit/Source/GLTFExporter/Public/Converters/GLTFImageConverters.h b/Engine/Plugins/Enterprise/GLTFExporter/MigratedFromGit/Source/GLTFExporter/Public/Converters/GLTFImageConverters.h
new file mode 100644
--- /dev/null
+++ b/Engine/Plugins/Enterprise/GLTFExporter/MigratedFromGit/Source/GLTFExporter/Public/Converters/GLTFImageConverters.h
@@ -0,0 +1,21 @@
+// Copyright Epic Games, Inc. All Rights Reserved.
+
+#pragma once
+
+#include "Json/GLTFJsonCore.h"
+#include "Converters/GLTFConverter.h"
+#include "Converters/GLTFBuilderContext.h"
+#include "Converters/GLTFSuperfluous.h"
+#include "Converters/GLTFSharedArray.h"
+#include "Options/GLTFExportOptions.h"
+
+typedef TGLTFConverter<FGLTFJsonImage*, TGLTFSuperfluous<FString>, EGLTFTextureType, bool, FIntPoint, TGLTFSharedArray<FColor>> IGLTFImageConverter;
+
+class FGLTFImageConverter final : public FGLTFBuilderContext, public IGLTFImageConverter
+{
+	using FGLTFBuilderContext::FGLTFBuilderContext;
+
+	virtual FGLTFJsonImage* Convert(TGLTFSuperfluous<FString> Name, EGLTFTextureType Type, bool bIgnoreAlpha, FIntPoint Size, TGLTFSharedArray<FColor> Pixels) override;
+
+	EGLTFJsonMimeType GetMimeType(const FColor* Pixels, FIntPoint Size, bool bIgnoreAlpha, EGLTFTextureType Type) const;
+};
diff --git a/Engine/Plugins/Enterprise/GLTFExporter/MigratedFromGit/Source/GLTFExporter/Public/Converters/GLTFIndexArray.h b/Engine/Plugins/Enterprise/GLTFExporter/MigratedFromGit/Source/GLTFExporter/Public/Converters/GLTFIndexArray.h
new file mode 100644
--- /dev/null
+++ b/Engine/Plugins/Enterprise/GLTFExporter/MigratedFromGit/Source/GLTFExporter/Public/Converters/GLTFIndexArray.h
@@ -0,0 +1,23 @@
+// Copyright Epic Games, Inc. All Rights Reserved.
+
+#pragma once
+
+#include "CoreMinimal.h"
+
+class FGLTFIndexArray : public TArray<int32>
+{
+public:
+
+	using TArray::TArray;
+	using TArray::operator=;
+
+	friend uint32 GetTypeHash(const FGLTFIndexArray& IndexArray)
+	{
+		uint32 Hash = GetTypeHash(IndexArray.Num());
+		for (const int32 Index : IndexArray)
+		{
+			Hash = HashCombine(Hash, GetTypeHash(Index));
+		}
+		return Hash;
+	}
+};
diff --git a/Engine/Plugins/Enterprise/GLTFExporter/MigratedFromGit/Source/GLTFExporter/Public/Converters/GLTFKhrVariantConverters.h b/Engine/Plugins/Enterprise/GLTFExporter/MigratedFromGit/Source/GLTFExporter/Public/Converters/GLTFKhrVariantConverters.h
new file mode 100644
--- /dev/null
+++ b/Engine/Plugins/Enterprise/GLTFExporter/MigratedFromGit/Source/GLTFExporter/Public/Converters/GLTFKhrVariantConverters.h
@@ -0,0 +1,22 @@
+﻿// Copyright Epic Games, Inc. All Rights Reserved.
+
+#pragma once
+
+#include "Json/GLTFJsonCore.h"
+#include "Json/GLTFJsonMesh.h"
+#include "Converters/GLTFConverter.h"
+#include "Converters/GLTFBuilderContext.h"
+
+class UVariant;
+class UPropertyValueMaterial;
+
+typedef TGLTFConverter<FGLTFJsonKhrMaterialVariant*, const UVariant*> IGLTFKhrMaterialVariantConverter;
+
+class FGLTFKhrMaterialVariantConverter final : public FGLTFBuilderContext, public IGLTFKhrMaterialVariantConverter
+{
+	using FGLTFBuilderContext::FGLTFBuilderContext;
+
+	virtual FGLTFJsonKhrMaterialVariant* Convert(const UVariant* Variant) override;
+
+	bool TryParseMaterialProperty(FGLTFJsonPrimitive*& OutPrimitive, FGLTFJsonMaterial*& OutMaterial, const UPropertyValueMaterial* Property) const;
+};
diff --git a/Engine/Plugins/Enterprise/GLTFExporter/MigratedFromGit/Source/GLTFExporter/Public/Converters/GLTFLightConverters.h b/Engine/Plugins/Enterprise/GLTFExporter/MigratedFromGit/Source/GLTFExporter/Public/Converters/GLTFLightConverters.h
new file mode 100644
--- /dev/null
+++ b/Engine/Plugins/Enterprise/GLTFExporter/MigratedFromGit/Source/GLTFExporter/Public/Converters/GLTFLightConverters.h
@@ -0,0 +1,16 @@
+// Copyright Epic Games, Inc. All Rights Reserved.
+
+#pragma once
+
+#include "Json/GLTFJsonCore.h"
+#include "Converters/GLTFConverter.h"
+#include "Converters/GLTFBuilderContext.h"
+
+typedef TGLTFConverter<FGLTFJsonLight*, const ULightComponent*> IGLTFLightConverter;
+
+class FGLTFLightConverter final : public FGLTFBuilderContext, public IGLTFLightConverter
+{
+	using FGLTFBuilderContext::FGLTFBuilderContext;
+
+	virtual FGLTFJsonLight* Convert(const ULightComponent* LightComponent) override;
+};
diff --git a/Engine/Plugins/Enterprise/GLTFExporter/MigratedFromGit/Source/GLTFExporter/Public/Converters/GLTFLightMapConverters.h b/Engine/Plugins/Enterprise/GLTFExporter/MigratedFromGit/Source/GLTFExporter/Public/Converters/GLTFLightMapConverters.h
new file mode 100644
--- /dev/null
+++ b/Engine/Plugins/Enterprise/GLTFExporter/MigratedFromGit/Source/GLTFExporter/Public/Converters/GLTFLightMapConverters.h
@@ -0,0 +1,16 @@
+// Copyright Epic Games, Inc. All Rights Reserved.
+
+#pragma once
+
+#include "Json/GLTFJsonCore.h"
+#include "Converters/GLTFConverter.h"
+#include "Converters/GLTFBuilderContext.h"
+
+typedef TGLTFConverter<FGLTFJsonLightMap*, const UStaticMeshComponent*> IGLTFLightMapConverter;
+
+class FGLTFLightMapConverter final : public FGLTFBuilderContext, public IGLTFLightMapConverter
+{
+	using FGLTFBuilderContext::FGLTFBuilderContext;
+
+	virtual FGLTFJsonLightMap* Convert(const UStaticMeshComponent* StaticMeshComponent) override;
+};
diff --git a/Engine/Plugins/Enterprise/GLTFExporter/MigratedFromGit/Source/GLTFExporter/Public/Converters/GLTFMaterialArray.h b/Engine/Plugins/Enterprise/GLTFExporter/MigratedFromGit/Source/GLTFExporter/Public/Converters/GLTFMaterialArray.h
new file mode 100644
--- /dev/null
+++ b/Engine/Plugins/Enterprise/GLTFExporter/MigratedFromGit/Source/GLTFExporter/Public/Converters/GLTFMaterialArray.h
@@ -0,0 +1,25 @@
+// Copyright Epic Games, Inc. All Rights Reserved.
+
+#pragma once
+
+#include "CoreMinimal.h"
+
+class UMaterialInterface;
+
+class FGLTFMaterialArray : public TArray<const UMaterialInterface*>
+{
+public:
+
+	using TArray::TArray;
+	using TArray::operator=;
+
+	friend uint32 GetTypeHash(const FGLTFMaterialArray& MaterialArray)
+	{
+		uint32 Hash = GetTypeHash(MaterialArray.Num());
+		for (const UMaterialInterface* Material : MaterialArray)
+		{
+			Hash = HashCombine(Hash, GetTypeHash(Material));
+		}
+		return Hash;
+	}
+};
diff --git a/Engine/Plugins/Enterprise/GLTFExporter/MigratedFromGit/Source/GLTFExporter/Public/Converters/GLTFMaterialConverters.h b/Engine/Plugins/Enterprise/GLTFExporter/MigratedFromGit/Source/GLTFExporter/Public/Converters/GLTFMaterialConverters.h
new file mode 100644
--- /dev/null
+++ b/Engine/Plugins/Enterprise/GLTFExporter/MigratedFromGit/Source/GLTFExporter/Public/Converters/GLTFMaterialConverters.h
@@ -0,0 +1,25 @@
+// Copyright Epic Games, Inc. All Rights Reserved.
+
+#pragma once
+
+#include "Json/GLTFJsonCore.h"
+#include "Converters/GLTFConverter.h"
+#include "Converters/GLTFBuilderContext.h"
+#include "Converters/GLTFMeshData.h"
+#include "Converters/GLTFIndexArray.h"
+#include "Converters/GLTFUVOverlapChecker.h"
+#include "Converters/GLTFUVDegenerateChecker.h"
+
+typedef TGLTFConverter<FGLTFJsonMaterial*, const UMaterialInterface*, const FGLTFMeshData*, FGLTFIndexArray> IGLTFMaterialConverter;
+
+class FGLTFMaterialConverter final : public FGLTFBuilderContext, public IGLTFMaterialConverter
+{
+	using FGLTFBuilderContext::FGLTFBuilderContext;
+
+	virtual void Sanitize(const UMaterialInterface*& Material, const FGLTFMeshData*& MeshData, FGLTFIndexArray& SectionIndices) override;
+
+	virtual FGLTFJsonMaterial* Convert(const UMaterialInterface* Material, const FGLTFMeshData* MeshData, FGLTFIndexArray SectionIndices) override;
+
+	FGLTFUVOverlapChecker UVOverlapChecker;
+	FGLTFUVDegenerateChecker UVDegenerateChecker;
+};
diff --git a/Engine/Plugins/Enterprise/GLTFExporter/MigratedFromGit/Source/GLTFExporter/Public/Converters/GLTFMeshConverters.h b/Engine/Plugins/Enterprise/GLTFExporter/MigratedFromGit/Source/GLTFExporter/Public/Converters/GLTFMeshConverters.h
new file mode 100644
--- /dev/null
+++ b/Engine/Plugins/Enterprise/GLTFExporter/MigratedFromGit/Source/GLTFExporter/Public/Converters/GLTFMeshConverters.h
@@ -0,0 +1,35 @@
+// Copyright Epic Games, Inc. All Rights Reserved.
+
+#pragma once
+
+#include "Json/GLTFJsonCore.h"
+#include "Converters/GLTFConverter.h"
+#include "Converters/GLTFBuilderContext.h"
+#include "Converters/GLTFMeshSectionConverters.h"
+#include "Converters/GLTFMeshDataConverters.h"
+#include "Converters/GLTFMaterialArray.h"
+
+typedef TGLTFConverter<FGLTFJsonMesh*, const UStaticMesh*, const UStaticMeshComponent*, FGLTFMaterialArray, int32> IGLTFStaticMeshConverter;
+typedef TGLTFConverter<FGLTFJsonMesh*, const USkeletalMesh*, const USkeletalMeshComponent*, FGLTFMaterialArray, int32> IGLTFSkeletalMeshConverter;
+
+class FGLTFStaticMeshConverter final : public FGLTFBuilderContext, public IGLTFStaticMeshConverter
+{
+	using FGLTFBuilderContext::FGLTFBuilderContext;
+
+	virtual void Sanitize(const UStaticMesh*& StaticMesh, const UStaticMeshComponent*& StaticMeshComponent, FGLTFMaterialArray& Materials, int32& LODIndex) override;
+
+	virtual FGLTFJsonMesh* Convert(const UStaticMesh* StaticMesh, const UStaticMeshComponent* StaticMeshComponent, FGLTFMaterialArray Materials, int32 LODIndex) override;
+
+	FGLTFStaticMeshSectionConverter MeshSectionConverter;
+};
+
+class FGLTFSkeletalMeshConverter final : public FGLTFBuilderContext, public IGLTFSkeletalMeshConverter
+{
+	using FGLTFBuilderContext::FGLTFBuilderContext;
+
+	virtual void Sanitize(const USkeletalMesh*& SkeletalMesh, const USkeletalMeshComponent*& SkeletalMeshComponent, FGLTFMaterialArray& Materials, int32& LODIndex) override;
+
+	virtual FGLTFJsonMesh* Convert(const USkeletalMesh* SkeletalMesh, const USkeletalMeshComponent* SkeletalMeshComponent, FGLTFMaterialArray Materials, int32 LODIndex) override;
+
+	FGLTFSkeletalMeshSectionConverter MeshSectionConverter;
+};
diff --git a/Engine/Plugins/Enterprise/GLTFExporter/MigratedFromGit/Source/GLTFExporter/Public/Converters/GLTFMeshData.h b/Engine/Plugins/Enterprise/GLTFExporter/MigratedFromGit/Source/GLTFExporter/Public/Converters/GLTFMeshData.h
new file mode 100644
--- /dev/null
+++ b/Engine/Plugins/Enterprise/GLTFExporter/MigratedFromGit/Source/GLTFExporter/Public/Converters/GLTFMeshData.h
@@ -0,0 +1,35 @@
+// Copyright Epic Games, Inc. All Rights Reserved.
+
+#pragma once
+
+#include "CoreMinimal.h"
+#include "MeshDescription.h"
+
+#if WITH_EDITOR
+#include "GLTFMaterialBakingStructures.h"
+#endif
+
+struct FGLTFMeshData
+{
+	FGLTFMeshData(const UStaticMesh* StaticMesh, const UStaticMeshComponent* StaticMeshComponent, int32 LODIndex);
+	FGLTFMeshData(const USkeletalMesh* SkeletalMesh, const USkeletalMeshComponent* SkeletalMeshComponent, int32 LODIndex);
+
+	const FGLTFMeshData* GetParent() const;
+
+	// TODO: find a better name for referencing the mesh-only data (no component)
+	const FGLTFMeshData* Parent;
+
+	FString Name;
+	int32 LODIndex;
+
+#if WITH_EDITOR
+	FGLTFPrimitiveData PrimitiveData;
+	FMeshDescription Description;
+
+	FLightMapRef LightMap;
+	const FLightmapResourceCluster* LightMapResourceCluster;
+	int32 LightMapTexCoord;
+
+	int32 BakeUsingTexCoord;
+#endif
+};
diff --git a/Engine/Plugins/Enterprise/GLTFExporter/MigratedFromGit/Source/GLTFExporter/Public/Converters/GLTFMeshDataConverters.h b/Engine/Plugins/Enterprise/GLTFExporter/MigratedFromGit/Source/GLTFExporter/Public/Converters/GLTFMeshDataConverters.h
new file mode 100644
--- /dev/null
+++ b/Engine/Plugins/Enterprise/GLTFExporter/MigratedFromGit/Source/GLTFExporter/Public/Converters/GLTFMeshDataConverters.h
@@ -0,0 +1,25 @@
+// Copyright Epic Games, Inc. All Rights Reserved.
+
+#pragma once
+
+#include "Converters/GLTFConverter.h"
+#include "Converters/GLTFMeshData.h"
+#include "Converters/GLTFBuilderContext.h"
+
+typedef TGLTFConverter<const FGLTFMeshData*, const UStaticMesh*, const UStaticMeshComponent*, int32> IGLTFStaticMeshDataConverter;
+typedef TGLTFConverter<const FGLTFMeshData*, const USkeletalMesh*, const USkeletalMeshComponent*, int32> IGLTFSkeletalMeshDataConverter;
+
+template <typename MeshType, typename MeshComponentType>
+class TGLTFMeshDataConverter : public FGLTFBuilderContext, public TGLTFConverter<const FGLTFMeshData*, const MeshType*, const MeshComponentType*, int32>
+{
+	TArray<TUniquePtr<FGLTFMeshData>> Outputs;
+
+	using FGLTFBuilderContext::FGLTFBuilderContext;
+
+	virtual void Sanitize(const MeshType*& Mesh, const MeshComponentType*& MeshComponent, int32& LODIndex) override;
+
+	virtual const FGLTFMeshData* Convert(const MeshType* Mesh, const MeshComponentType* MeshComponent, int32 LODIndex) override;
+};
+
+typedef TGLTFMeshDataConverter<UStaticMesh, UStaticMeshComponent> FGLTFStaticMeshDataConverter;
+typedef TGLTFMeshDataConverter<USkeletalMesh, USkeletalMeshComponent> FGLTFSkeletalMeshDataConverter;
diff --git a/Engine/Plugins/Enterprise/GLTFExporter/MigratedFromGit/Source/GLTFExporter/Public/Converters/GLTFMeshSection.h b/Engine/Plugins/Enterprise/GLTFExporter/MigratedFromGit/Source/GLTFExporter/Public/Converters/GLTFMeshSection.h
new file mode 100644
--- /dev/null
+++ b/Engine/Plugins/Enterprise/GLTFExporter/MigratedFromGit/Source/GLTFExporter/Public/Converters/GLTFMeshSection.h
@@ -0,0 +1,30 @@
+// Copyright Epic Games, Inc. All Rights Reserved.
+
+#pragma once
+
+#include "CoreMinimal.h"
+#include "Converters/GLTFIndexArray.h"
+
+struct FGLTFMeshSection
+{
+	FGLTFMeshSection(const FStaticMeshLODResources* MeshLOD, const FGLTFIndexArray& SectionIndices);
+	FGLTFMeshSection(const FSkeletalMeshLODRenderData* MeshLOD, const FGLTFIndexArray& SectionIndices);
+
+	TArray<uint32> IndexMap;
+	TArray<uint32> IndexBuffer;
+
+	TArray<TArray<FBoneIndexType>> BoneMaps;
+	TArray<uint32> BoneMapLookup;
+	FBoneIndexType MaxBoneIndex;
+
+private:
+
+	template <typename IndexType, typename SectionArrayType>
+	void Init(const SectionArrayType& Sections, const FGLTFIndexArray& SectionIndices, const IndexType* SourceData);
+
+	static uint32 GetIndexOffset(const FStaticMeshSection& Section);
+	static uint32 GetIndexOffset(const FSkelMeshRenderSection& Section);
+
+	static const TArray<FBoneIndexType>& GetBoneMap(const FStaticMeshSection& Section);
+	static const TArray<FBoneIndexType>& GetBoneMap(const FSkelMeshRenderSection& Section);
+};
diff --git a/Engine/Plugins/Enterprise/GLTFExporter/MigratedFromGit/Source/GLTFExporter/Public/Converters/GLTFMeshSectionConverters.h b/Engine/Plugins/Enterprise/GLTFExporter/MigratedFromGit/Source/GLTFExporter/Public/Converters/GLTFMeshSectionConverters.h
new file mode 100644
--- /dev/null
+++ b/Engine/Plugins/Enterprise/GLTFExporter/MigratedFromGit/Source/GLTFExporter/Public/Converters/GLTFMeshSectionConverters.h
@@ -0,0 +1,21 @@
+// Copyright Epic Games, Inc. All Rights Reserved.
+
+#pragma once
+
+#include "Converters/GLTFConverter.h"
+#include "Converters/GLTFMeshSection.h"
+#include "Converters/GLTFIndexArray.h"
+
+template <typename MeshLODType>
+class TGLTFMeshSectionConverter final : public TGLTFConverter<const FGLTFMeshSection*, const MeshLODType*, FGLTFIndexArray>
+{
+	TArray<TUniquePtr<FGLTFMeshSection>> Outputs;
+
+	virtual const FGLTFMeshSection* Convert(const MeshLODType* MeshLOD, FGLTFIndexArray SectionIndices) override
+	{
+		return Outputs.Add_GetRef(MakeUnique<FGLTFMeshSection>(MeshLOD, SectionIndices)).Get();
+	}
+};
+
+typedef TGLTFMeshSectionConverter<FStaticMeshLODResources> FGLTFStaticMeshSectionConverter;
+typedef TGLTFMeshSectionConverter<FSkeletalMeshLODRenderData> FGLTFSkeletalMeshSectionConverter;
diff --git a/Engine/Plugins/Enterprise/GLTFExporter/MigratedFromGit/Source/GLTFExporter/Public/Converters/GLTFNodeConverters.h b/Engine/Plugins/Enterprise/GLTFExporter/MigratedFromGit/Source/GLTFExporter/Public/Converters/GLTFNodeConverters.h
new file mode 100644
--- /dev/null
+++ b/Engine/Plugins/Enterprise/GLTFExporter/MigratedFromGit/Source/GLTFExporter/Public/Converters/GLTFNodeConverters.h
@@ -0,0 +1,56 @@
+// Copyright Epic Games, Inc. All Rights Reserved.
+
+#pragma once
+
+#include "Json/GLTFJsonCore.h"
+#include "Converters/GLTFConverter.h"
+#include "Converters/GLTFBuilderContext.h"
+
+typedef TGLTFConverter<FGLTFJsonNode*, const AActor*> IGLTFActorConverter;
+typedef TGLTFConverter<FGLTFJsonNode*, const USceneComponent*> IGLTFComponentConverter;
+typedef TGLTFConverter<FGLTFJsonNode*, const USceneComponent*, FName> IGLTFComponentSocketConverter;
+typedef TGLTFConverter<FGLTFJsonNode*, FGLTFJsonNode*, const UStaticMesh*, FName> IGLTFStaticSocketConverter;
+typedef TGLTFConverter<FGLTFJsonNode*, FGLTFJsonNode*, const USkeletalMesh*, FName> IGLTFSkeletalSocketConverter;
+typedef TGLTFConverter<FGLTFJsonNode*, FGLTFJsonNode*, const USkeletalMesh*, int32> IGLTFSkeletalBoneConverter;
+
+class FGLTFActorConverter final : public FGLTFBuilderContext, public IGLTFActorConverter
+{
+	using FGLTFBuilderContext::FGLTFBuilderContext;
+
+	virtual FGLTFJsonNode* Convert(const AActor* Actor) override;
+};
+
+class FGLTFComponentConverter final : public FGLTFBuilderContext, public IGLTFComponentConverter
+{
+	using FGLTFBuilderContext::FGLTFBuilderContext;
+
+	virtual FGLTFJsonNode* Convert(const USceneComponent* SceneComponent) override;
+};
+
+class FGLTFComponentSocketConverter final : public FGLTFBuilderContext, public IGLTFComponentSocketConverter
+{
+	using FGLTFBuilderContext::FGLTFBuilderContext;
+
+	virtual FGLTFJsonNode* Convert(const USceneComponent* SceneComponent, FName SocketName) override;
+};
+
+class FGLTFStaticSocketConverter final : public FGLTFBuilderContext, public IGLTFStaticSocketConverter
+{
+	using FGLTFBuilderContext::FGLTFBuilderContext;
+
+	virtual FGLTFJsonNode* Convert(FGLTFJsonNode* RootNode, const UStaticMesh* StaticMesh, FName SocketName) override;
+};
+
+class FGLTFSkeletalSocketConverter final : public FGLTFBuilderContext, public IGLTFSkeletalSocketConverter
+{
+	using FGLTFBuilderContext::FGLTFBuilderContext;
+
+	virtual FGLTFJsonNode* Convert(FGLTFJsonNode* RootNode, const USkeletalMesh* SkeletalMesh, FName SocketName) override;
+};
+
+class FGLTFSkeletalBoneConverter final : public FGLTFBuilderContext, public IGLTFSkeletalBoneConverter
+{
+	using FGLTFBuilderContext::FGLTFBuilderContext;
+
+	virtual FGLTFJsonNode* Convert(FGLTFJsonNode* RootNode, const USkeletalMesh* SkeletalMesh, int32 BoneIndex) override;
+};
diff --git a/Engine/Plugins/Enterprise/GLTFExporter/MigratedFromGit/Source/GLTFExporter/Public/Converters/GLTFSamplerConverters.h b/Engine/Plugins/Enterprise/GLTFExporter/MigratedFromGit/Source/GLTFExporter/Public/Converters/GLTFSamplerConverters.h
new file mode 100644
--- /dev/null
+++ b/Engine/Plugins/Enterprise/GLTFExporter/MigratedFromGit/Source/GLTFExporter/Public/Converters/GLTFSamplerConverters.h
@@ -0,0 +1,16 @@
+// Copyright Epic Games, Inc. All Rights Reserved.
+
+#pragma once
+
+#include "Json/GLTFJsonCore.h"
+#include "Converters/GLTFConverter.h"
+#include "Converters/GLTFBuilderContext.h"
+
+typedef TGLTFConverter<FGLTFJsonSampler*, const UTexture*> IGLTFSamplerConverter;
+
+class FGLTFSamplerConverter final : public FGLTFBuilderContext, public IGLTFSamplerConverter
+{
+	using FGLTFBuilderContext::FGLTFBuilderContext;
+
+	virtual FGLTFJsonSampler* Convert(const UTexture* Texture) override;
+};
diff --git a/Engine/Plugins/Enterprise/GLTFExporter/MigratedFromGit/Source/GLTFExporter/Public/Converters/GLTFSceneConverters.h b/Engine/Plugins/Enterprise/GLTFExporter/MigratedFromGit/Source/GLTFExporter/Public/Converters/GLTFSceneConverters.h
new file mode 100644
--- /dev/null
+++ b/Engine/Plugins/Enterprise/GLTFExporter/MigratedFromGit/Source/GLTFExporter/Public/Converters/GLTFSceneConverters.h
@@ -0,0 +1,16 @@
+// Copyright Epic Games, Inc. All Rights Reserved.
+
+#pragma once
+
+#include "Json/GLTFJsonCore.h"
+#include "Converters/GLTFConverter.h"
+#include "Converters/GLTFBuilderContext.h"
+
+typedef TGLTFConverter<FGLTFJsonScene*, const UWorld*> IGLTFSceneConverter;
+
+class FGLTFSceneConverter final : public FGLTFBuilderContext, public IGLTFSceneConverter
+{
+	using FGLTFBuilderContext::FGLTFBuilderContext;
+
+	virtual FGLTFJsonScene* Convert(const UWorld* Level) override;
+};
diff --git a/Engine/Plugins/Enterprise/GLTFExporter/MigratedFromGit/Source/GLTFExporter/Public/Converters/GLTFSharedArray.h b/Engine/Plugins/Enterprise/GLTFExporter/MigratedFromGit/Source/GLTFExporter/Public/Converters/GLTFSharedArray.h
new file mode 100644
--- /dev/null
+++ b/Engine/Plugins/Enterprise/GLTFExporter/MigratedFromGit/Source/GLTFExporter/Public/Converters/GLTFSharedArray.h
@@ -0,0 +1,48 @@
+// Copyright Epic Games, Inc. All Rights Reserved.
+
+#pragma once
+
+#include "CoreMinimal.h"
+
+template <typename ElementType, typename ArrayType = TArray<ElementType>, ESPMode Mode = ESPMode::Fast>
+class TGLTFSharedArray : public TSharedRef<ArrayType, Mode>
+{
+public:
+
+	TGLTFSharedArray()
+		: TSharedRef<ArrayType, Mode>(MakeShared<ArrayType>())
+	{
+	}
+
+	template <typename OtherType, typename = decltype(ImplicitConv<ArrayType*>(static_cast<OtherType*>(nullptr)))>
+	TGLTFSharedArray(const TSharedRef<OtherType, Mode>& SharedRef)
+	: TSharedRef<ArrayType, Mode>(SharedRef)
+	{
+	}
+
+	using TSharedRef<ArrayType, Mode>::TSharedRef;
+	using TSharedRef<ArrayType, Mode>::operator=;
+
+	bool operator==(const TGLTFSharedArray& Other) const
+	{
+		return this->Get() == Other.Get();
+	}
+
+	bool operator!=(const TGLTFSharedArray& Other) const
+	{
+		return this->Get() != Other.Get();
+	}
+
+	friend uint32 GetTypeHash(const TGLTFSharedArray& SharedArray)
+	{
+		const ArrayType& Array = SharedArray.Get();
+		uint32 Hash = GetTypeHash(Array.Num());
+
+		for (const auto& Element : Array)
+		{
+			Hash = HashCombine(Hash, GetTypeHash(Element));
+		}
+
+		return Hash;
+	}
+};
diff --git a/Engine/Plugins/Enterprise/GLTFExporter/MigratedFromGit/Source/GLTFExporter/Public/Converters/GLTFSkinConverters.h b/Engine/Plugins/Enterprise/GLTFExporter/MigratedFromGit/Source/GLTFExporter/Public/Converters/GLTFSkinConverters.h
new file mode 100644
--- /dev/null
+++ b/Engine/Plugins/Enterprise/GLTFExporter/MigratedFromGit/Source/GLTFExporter/Public/Converters/GLTFSkinConverters.h
@@ -0,0 +1,16 @@
+// Copyright Epic Games, Inc. All Rights Reserved.
+
+#pragma once
+
+#include "Json/GLTFJsonCore.h"
+#include "Converters/GLTFConverter.h"
+#include "Converters/GLTFBuilderContext.h"
+
+typedef TGLTFConverter<FGLTFJsonSkin*, FGLTFJsonNode*, const USkeletalMesh*> IGLTFSkinConverter;
+
+class FGLTFSkinConverter final : public FGLTFBuilderContext, public IGLTFSkinConverter
+{
+	using FGLTFBuilderContext::FGLTFBuilderContext;
+
+	virtual FGLTFJsonSkin* Convert(FGLTFJsonNode* RootNode, const USkeletalMesh* SkeletalMesh) override;
+};
diff --git a/Engine/Plugins/Enterprise/GLTFExporter/MigratedFromGit/Source/GLTFExporter/Public/Converters/GLTFSkySphereConverters.h b/Engine/Plugins/Enterprise/GLTFExporter/MigratedFromGit/Source/GLTFExporter/Public/Converters/GLTFSkySphereConverters.h
new file mode 100644
--- /dev/null
+++ b/Engine/Plugins/Enterprise/GLTFExporter/MigratedFromGit/Source/GLTFExporter/Public/Converters/GLTFSkySphereConverters.h
@@ -0,0 +1,35 @@
+// Copyright Epic Games, Inc. All Rights Reserved.
+
+#pragma once
+
+#include "Core/GLTFColor.h"
+#include "Json/GLTFJsonCore.h"
+#include "Converters/GLTFConverter.h"
+#include "Converters/GLTFBuilderContext.h"
+
+struct FGLTFJsonSkySphereColorCurve;
+
+typedef TGLTFConverter<FGLTFJsonSkySphere*, const AActor*> IGLTFSkySphereConverter;
+
+class FGLTFSkySphereConverter final : public FGLTFBuilderContext, public IGLTFSkySphereConverter
+{
+	enum class ESkySphereTextureParameter
+	{
+		SkyTexture,
+		CloudsTexture,
+		StarsTexture
+	};
+
+	using FGLTFBuilderContext::FGLTFBuilderContext;
+
+	virtual FGLTFJsonSkySphere* Convert(const AActor* SkySphereActor) override;
+
+	template <class ValueType>
+	void ConvertProperty(const AActor* Actor, const TCHAR* PropertyName, ValueType& OutValue) const;
+
+	void ConvertColorProperty(const AActor* Actor, const TCHAR* PropertyName, FGLTFColor4& OutValue) const;
+	void ConvertColorCurveProperty(const AActor* Actor, const TCHAR* PropertyName, FGLTFJsonSkySphereColorCurve& OutValue) const;
+	void ConvertScalarParameter(const AActor* Actor, const UMaterialInstance* Material, const TCHAR* ParameterName, float& OutValue) const;
+
+	void ConvertTextureParameter(const AActor* Actor, const UMaterialInstance* Material, const ESkySphereTextureParameter Parameter, FGLTFJsonTexture*& OutValue) const;
+};
diff --git a/Engine/Plugins/Enterprise/GLTFExporter/MigratedFromGit/Source/GLTFExporter/Public/Converters/GLTFSuperfluous.h b/Engine/Plugins/Enterprise/GLTFExporter/MigratedFromGit/Source/GLTFExporter/Public/Converters/GLTFSuperfluous.h
new file mode 100644
--- /dev/null
+++ b/Engine/Plugins/Enterprise/GLTFExporter/MigratedFromGit/Source/GLTFExporter/Public/Converters/GLTFSuperfluous.h
@@ -0,0 +1,34 @@
+// Copyright Epic Games, Inc. All Rights Reserved.
+
+#pragma once
+
+#include "CoreMinimal.h"
+
+template <typename Type>
+class TGLTFSuperfluous : public Type
+{
+public:
+
+	TGLTFSuperfluous(const Type& Other)
+		: Type(Other)
+	{
+	}
+
+	using Type::Type;
+	using Type::operator=;
+
+	bool operator==(const TGLTFSuperfluous& Other) const
+	{
+		return true;
+	}
+
+	bool operator!=(const TGLTFSuperfluous& Other) const
+	{
+		return false;
+	}
+
+	friend uint32 GetTypeHash(const TGLTFSuperfluous& Superfluous)
+	{
+		return 0;
+	}
+};
diff --git a/Engine/Plugins/Enterprise/GLTFExporter/MigratedFromGit/Source/GLTFExporter/Public/Converters/GLTFTextureConverters.h b/Engine/Plugins/Enterprise/GLTFExporter/MigratedFromGit/Source/GLTFExporter/Public/Converters/GLTFTextureConverters.h
new file mode 100644
--- /dev/null
+++ b/Engine/Plugins/Enterprise/GLTFExporter/MigratedFromGit/Source/GLTFExporter/Public/Converters/GLTFTextureConverters.h
@@ -0,0 +1,60 @@
+// Copyright Epic Games, Inc. All Rights Reserved.
+
+#pragma once
+
+#include "Json/GLTFJsonCore.h"
+#include "Converters/GLTFConverter.h"
+#include "Converters/GLTFBuilderContext.h"
+
+// TODO: generalize parameter bToSRGB to SamplerType (including normalmap unpacking)
+
+// TODO: remove lightmap-specific converter (since it won't work in runtime)
+
+typedef TGLTFConverter<FGLTFJsonTexture*, const UTexture2D*, bool> IGLTFTexture2DConverter;
+typedef TGLTFConverter<FGLTFJsonTexture*, const UTextureCube*, ECubeFace, bool> IGLTFTextureCubeConverter;
+typedef TGLTFConverter<FGLTFJsonTexture*, const UTextureRenderTarget2D*, bool> IGLTFTextureRenderTarget2DConverter;
+typedef TGLTFConverter<FGLTFJsonTexture*, const UTextureRenderTargetCube*, ECubeFace, bool> IGLTFTextureRenderTargetCubeConverter;
+typedef TGLTFConverter<FGLTFJsonTexture*, const ULightMapTexture2D*> IGLTFTextureLightMapConverter;
+
+class FGLTFTexture2DConverter final : public FGLTFBuilderContext, public IGLTFTexture2DConverter
+{
+	using FGLTFBuilderContext::FGLTFBuilderContext;
+
+	virtual void Sanitize(const UTexture2D*& Texture2D, bool& bToSRGB) override;
+
+	virtual FGLTFJsonTexture* Convert(const UTexture2D* Texture2D, bool bToSRGB) override;
+};
+
+class FGLTFTextureCubeConverter final : public FGLTFBuilderContext, public IGLTFTextureCubeConverter
+{
+	using FGLTFBuilderContext::FGLTFBuilderContext;
+
+	virtual void Sanitize(const UTextureCube*& TextureCube, ECubeFace& CubeFace, bool& bToSRGB) override;
+
+	virtual FGLTFJsonTexture* Convert(const UTextureCube* TextureCube, ECubeFace CubeFace, bool bToSRGB) override;
+};
+
+class FGLTFTextureRenderTarget2DConverter final : public FGLTFBuilderContext, public IGLTFTextureRenderTarget2DConverter
+{
+	using FGLTFBuilderContext::FGLTFBuilderContext;
+
+	virtual void Sanitize(const UTextureRenderTarget2D*& RenderTarget2D, bool& bToSRGB) override;
+
+	virtual FGLTFJsonTexture* Convert(const UTextureRenderTarget2D* RenderTarget2D, bool bToSRGB) override;
+};
+
+class FGLTFTextureRenderTargetCubeConverter final : public FGLTFBuilderContext, public IGLTFTextureRenderTargetCubeConverter
+{
+	using FGLTFBuilderContext::FGLTFBuilderContext;
+
+	virtual void Sanitize(const UTextureRenderTargetCube*& RenderTargetCube, ECubeFace& CubeFace, bool& bToSRGB) override;
+
+	virtual FGLTFJsonTexture* Convert(const UTextureRenderTargetCube* RenderTargetCube, ECubeFace CubeFace, bool bToSRGB) override;
+};
+
+class FGLTFTextureLightMapConverter final : public FGLTFBuilderContext, public IGLTFTextureLightMapConverter
+{
+	using FGLTFBuilderContext::FGLTFBuilderContext;
+
+	virtual FGLTFJsonTexture* Convert(const ULightMapTexture2D* LightMap) override;
+};
diff --git a/Engine/Plugins/Enterprise/GLTFExporter/MigratedFromGit/Source/GLTFExporter/Public/Converters/GLTFUVDegenerateChecker.h b/Engine/Plugins/Enterprise/GLTFExporter/MigratedFromGit/Source/GLTFExporter/Public/Converters/GLTFUVDegenerateChecker.h
new file mode 100644
--- /dev/null
+++ b/Engine/Plugins/Enterprise/GLTFExporter/MigratedFromGit/Source/GLTFExporter/Public/Converters/GLTFUVDegenerateChecker.h
@@ -0,0 +1,18 @@
+// Copyright Epic Games, Inc. All Rights Reserved.
+
+#pragma once
+
+#include "Converters/GLTFConverter.h"
+#include "Converters/GLTFIndexArray.h"
+
+struct FMeshDescription;
+
+class FGLTFUVDegenerateChecker final : public TGLTFConverter<float, const FMeshDescription*, FGLTFIndexArray, int32>
+{
+	virtual void Sanitize(const FMeshDescription*& Description, FGLTFIndexArray& SectionIndices, int32& TexCoord) override;
+
+	virtual float Convert(const FMeshDescription* Description, FGLTFIndexArray SectionIndices, int32 TexCoord) override;
+
+	static bool IsDegenerateTriangle(const TStaticArray<FVector2D, 3>& Points);
+	static bool IsDegenerateTriangle(const TStaticArray<FVector, 3>& Points);
+};
diff --git a/Engine/Plugins/Enterprise/GLTFExporter/MigratedFromGit/Source/GLTFExporter/Public/Converters/GLTFUVOverlapChecker.h b/Engine/Plugins/Enterprise/GLTFExporter/MigratedFromGit/Source/GLTFExporter/Public/Converters/GLTFUVOverlapChecker.h
new file mode 100644
--- /dev/null
+++ b/Engine/Plugins/Enterprise/GLTFExporter/MigratedFromGit/Source/GLTFExporter/Public/Converters/GLTFUVOverlapChecker.h
@@ -0,0 +1,17 @@
+// Copyright Epic Games, Inc. All Rights Reserved.
+
+#pragma once
+
+#include "Converters/GLTFConverter.h"
+#include "Converters/GLTFIndexArray.h"
+
+struct FMeshDescription;
+
+class FGLTFUVOverlapChecker final : public TGLTFConverter<float, const FMeshDescription*, FGLTFIndexArray, int32>
+{
+	virtual void Sanitize(const FMeshDescription*& Description, FGLTFIndexArray& SectionIndices, int32& TexCoord) override;
+
+	virtual float Convert(const FMeshDescription* Description, FGLTFIndexArray SectionIndices, int32 TexCoord) override;
+
+	static UMaterialInterface* GetMaterial();
+};
diff --git a/Engine/Plugins/Enterprise/GLTFExporter/MigratedFromGit/Source/GLTFExporter/Public/Tasks/GLTFDelayedTask.h b/Engine/Plugins/Enterprise/GLTFExporter/MigratedFromGit/Source/GLTFExporter/Public/Tasks/GLTFDelayedTask.h
new file mode 100644
--- /dev/null
+++ b/Engine/Plugins/Enterprise/GLTFExporter/MigratedFromGit/Source/GLTFExporter/Public/Tasks/GLTFDelayedTask.h
@@ -0,0 +1,32 @@
+// Copyright Epic Games, Inc. All Rights Reserved.
+
+#pragma once
+
+#include "CoreMinimal.h"
+
+enum class EGLTFTaskPriority : uint8
+{
+	Animation,
+	Mesh,
+	Material,
+	Texture,
+	MAX
+};
+
+class FGLTFDelayedTask
+{
+public:
+
+	const EGLTFTaskPriority Priority;
+
+	FGLTFDelayedTask(EGLTFTaskPriority Priority)
+		: Priority(Priority)
+	{
+	}
+
+	virtual ~FGLTFDelayedTask() = default;
+
+	virtual FString GetName() = 0;
+
+	virtual void Process() = 0;
+};
