[GLTFExporter] refactor: add GLTFMaterialBaking as public dependency in GLTFExporter module

commit 7100fb4dd39a66828f7f95950f35957fddbf96de
Date: Sun, 13 Mar 2022 22:34:52 +0100

    refactor: add GLTFMaterialBaking as public dependency in GLTFExporter module
    (cherry picked from commit 786dfb59f3bdf26b95715a4ed1fb13583638f8f5​)

#skipci

diff --git a/Engine/Plugins/Enterprise/GLTFExporter/MigratedFromGit/Source/GLTFExporter/GLTFExporter.Build.cs b/Engine/Plugins/Enterprise/GLTFExporter/MigratedFromGit/Source/GLTFExporter/GLTFExporter.Build.cs
--- a/Engine/Plugins/Enterprise/GLTFExporter/MigratedFromGit/Source/GLTFExporter/GLTFExporter.Build.cs
+++ b/Engine/Plugins/Enterprise/GLTFExporter/MigratedFromGit/Source/GLTFExporter/GLTFExporter.Build.cs
@@ -35,6 +35,14 @@ public class GLTFExporter : ModuleRules
 
 		if (Target.bBuildEditor)
 		{
+			// TODO: remove this when we no longer need to include GLTFMaterialBakingStructures.h in GLTFMeshData.h
+			PublicDependencyModuleNames.AddRange(
+				new string[]
+				{
+					"GLTFMaterialBaking",
+				}
+			);
+
 			PrivateDependencyModuleNames.AddRange(
 				new string[]
 				{
@@ -51,7 +59,6 @@ public class GLTFExporter : ModuleRules
 					"MeshMergeUtilities",
 					"MeshDescription",
 					"StaticMeshDescription",
-					"GLTFMaterialBaking",
 					"GLTFMaterialAnalyzer",
 				}
 			);
diff --git a/Engine/Plugins/Enterprise/GLTFExporter/MigratedFromGit/Source/GLTFExporter/Public/Converters/GLTFMeshData.h b/Engine/Plugins/Enterprise/GLTFExporter/MigratedFromGit/Source/GLTFExporter/Public/Converters/GLTFMeshData.h
--- a/Engine/Plugins/Enterprise/GLTFExporter/MigratedFromGit/Source/GLTFExporter/Public/Converters/GLTFMeshData.h
+++ b/Engine/Plugins/Enterprise/GLTFExporter/MigratedFromGit/Source/GLTFExporter/Public/Converters/GLTFMeshData.h
@@ -6,6 +6,7 @@
 #include "MeshDescription.h"
 
 #if WITH_EDITOR
+// TODO: remove dependency to GLTFMaterialBaking in public header
 #include "GLTFMaterialBakingStructures.h"
 #endif
 
