[GLTFExporter] refactor: add GLTFEXPORTER_API to all new public classes

commit fe481cd993ffc05015021499689ef8626c7717ef
Date: Sun, 13 Mar 2022 23:02:17 +0100

    refactor: add GLTFEXPORTER_API to all new public classes
    (cherry picked from commit cce3d1e5e7e26ba0e48861e9ad41520e93fa928e​)

#skipci

diff --git a/Engine/Plugins/Enterprise/GLTFExporter/MigratedFromGit/Source/GLTFExporter/Public/Builders/GLTFBufferBuilder.h b/Engine/Plugins/Enterprise/GLTFExporter/MigratedFromGit/Source/GLTFExporter/Public/Builders/GLTFBufferBuilder.h
--- a/Engine/Plugins/Enterprise/GLTFExporter/MigratedFromGit/Source/GLTFExporter/Public/Builders/GLTFBufferBuilder.h
+++ b/Engine/Plugins/Enterprise/GLTFExporter/MigratedFromGit/Source/GLTFExporter/Public/Builders/GLTFBufferBuilder.h
@@ -4,7 +4,7 @@
 
 #include "Builders/GLTFJsonBuilder.h"
 
-class FGLTFBufferBuilder : public FGLTFJsonBuilder
+class GLTFEXPORTER_API FGLTFBufferBuilder : public FGLTFJsonBuilder
 {
 protected:
 
diff --git a/Engine/Plugins/Enterprise/GLTFExporter/MigratedFromGit/Source/GLTFExporter/Public/Builders/GLTFBuilder.h b/Engine/Plugins/Enterprise/GLTFExporter/MigratedFromGit/Source/GLTFExporter/Public/Builders/GLTFBuilder.h
--- a/Engine/Plugins/Enterprise/GLTFExporter/MigratedFromGit/Source/GLTFExporter/Public/Builders/GLTFBuilder.h
+++ b/Engine/Plugins/Enterprise/GLTFExporter/MigratedFromGit/Source/GLTFExporter/Public/Builders/GLTFBuilder.h
@@ -6,7 +6,7 @@
 #include "Json/GLTFJsonEnums.h"
 #include "UObject/GCObjectScopeGuard.h"
 
-class FGLTFBuilder
+class GLTFEXPORTER_API FGLTFBuilder
 {
 public:
 
diff --git a/Engine/Plugins/Enterprise/GLTFExporter/MigratedFromGit/Source/GLTFExporter/Public/Builders/GLTFContainerBuilder.h b/Engine/Plugins/Enterprise/GLTFExporter/MigratedFromGit/Source/GLTFExporter/Public/Builders/GLTFContainerBuilder.h
--- a/Engine/Plugins/Enterprise/GLTFExporter/MigratedFromGit/Source/GLTFExporter/Public/Builders/GLTFContainerBuilder.h
+++ b/Engine/Plugins/Enterprise/GLTFExporter/MigratedFromGit/Source/GLTFExporter/Public/Builders/GLTFContainerBuilder.h
@@ -4,7 +4,7 @@
 
 #include "Builders/GLTFConvertBuilder.h"
 
-class FGLTFContainerBuilder : public FGLTFConvertBuilder
+class GLTFEXPORTER_API FGLTFContainerBuilder : public FGLTFConvertBuilder
 {
 public:
 
diff --git a/Engine/Plugins/Enterprise/GLTFExporter/MigratedFromGit/Source/GLTFExporter/Public/Builders/GLTFConvertBuilder.h b/Engine/Plugins/Enterprise/GLTFExporter/MigratedFromGit/Source/GLTFExporter/Public/Builders/GLTFConvertBuilder.h
--- a/Engine/Plugins/Enterprise/GLTFExporter/MigratedFromGit/Source/GLTFExporter/Public/Builders/GLTFConvertBuilder.h
+++ b/Engine/Plugins/Enterprise/GLTFExporter/MigratedFromGit/Source/GLTFExporter/Public/Builders/GLTFConvertBuilder.h
@@ -23,7 +23,7 @@
 #include "Converters/GLTFKhrVariantConverters.h"
 #include "Converters/GLTFSkySphereConverters.h"
 
-class FGLTFConvertBuilder : public FGLTFBufferBuilder
+class GLTFEXPORTER_API FGLTFConvertBuilder : public FGLTFBufferBuilder
 {
 public:
 
diff --git a/Engine/Plugins/Enterprise/GLTFExporter/MigratedFromGit/Source/GLTFExporter/Public/Builders/GLTFFileBuilder.h b/Engine/Plugins/Enterprise/GLTFExporter/MigratedFromGit/Source/GLTFExporter/Public/Builders/GLTFFileBuilder.h
--- a/Engine/Plugins/Enterprise/GLTFExporter/MigratedFromGit/Source/GLTFExporter/Public/Builders/GLTFFileBuilder.h
+++ b/Engine/Plugins/Enterprise/GLTFExporter/MigratedFromGit/Source/GLTFExporter/Public/Builders/GLTFFileBuilder.h
@@ -5,7 +5,7 @@
 #include "Builders/GLTFTaskBuilder.h"
 #include "Builders/GLTFMemoryArchive.h"
 
-class FGLTFFileBuilder : public FGLTFTaskBuilder
+class GLTFEXPORTER_API FGLTFFileBuilder : public FGLTFTaskBuilder
 {
 public:
 
diff --git a/Engine/Plugins/Enterprise/GLTFExporter/MigratedFromGit/Source/GLTFExporter/Public/Builders/GLTFJsonBuilder.h b/Engine/Plugins/Enterprise/GLTFExporter/MigratedFromGit/Source/GLTFExporter/Public/Builders/GLTFJsonBuilder.h
--- a/Engine/Plugins/Enterprise/GLTFExporter/MigratedFromGit/Source/GLTFExporter/Public/Builders/GLTFJsonBuilder.h
+++ b/Engine/Plugins/Enterprise/GLTFExporter/MigratedFromGit/Source/GLTFExporter/Public/Builders/GLTFJsonBuilder.h
@@ -5,7 +5,7 @@
 #include "Json/GLTFJsonRoot.h"
 #include "Builders/GLTFFileBuilder.h"
 
-class FGLTFJsonBuilder : public FGLTFFileBuilder
+class GLTFEXPORTER_API FGLTFJsonBuilder : public FGLTFFileBuilder
 {
 public:
 
diff --git a/Engine/Plugins/Enterprise/GLTFExporter/MigratedFromGit/Source/GLTFExporter/Public/Builders/GLTFLogBuilder.h b/Engine/Plugins/Enterprise/GLTFExporter/MigratedFromGit/Source/GLTFExporter/Public/Builders/GLTFLogBuilder.h
--- a/Engine/Plugins/Enterprise/GLTFExporter/MigratedFromGit/Source/GLTFExporter/Public/Builders/GLTFLogBuilder.h
+++ b/Engine/Plugins/Enterprise/GLTFExporter/MigratedFromGit/Source/GLTFExporter/Public/Builders/GLTFLogBuilder.h
@@ -8,7 +8,7 @@
 class IMessageLogListing;
 #endif
 
-class FGLTFLogBuilder : public FGLTFBuilder
+class GLTFEXPORTER_API FGLTFLogBuilder : public FGLTFBuilder
 {
 public:
 
diff --git a/Engine/Plugins/Enterprise/GLTFExporter/MigratedFromGit/Source/GLTFExporter/Public/Builders/GLTFMemoryArchive.h b/Engine/Plugins/Enterprise/GLTFExporter/MigratedFromGit/Source/GLTFExporter/Public/Builders/GLTFMemoryArchive.h
--- a/Engine/Plugins/Enterprise/GLTFExporter/MigratedFromGit/Source/GLTFExporter/Public/Builders/GLTFMemoryArchive.h
+++ b/Engine/Plugins/Enterprise/GLTFExporter/MigratedFromGit/Source/GLTFExporter/Public/Builders/GLTFMemoryArchive.h
@@ -5,7 +5,7 @@
 #include "CoreMinimal.h"
 #include "Serialization/MemoryArchive.h"
 
-class FGLTFMemoryArchive final : public FMemoryArchive, public TArray64<uint8>
+class GLTFEXPORTER_API FGLTFMemoryArchive final : public FMemoryArchive, public TArray64<uint8>
 {
 public:
 
diff --git a/Engine/Plugins/Enterprise/GLTFExporter/MigratedFromGit/Source/GLTFExporter/Public/Builders/GLTFTaskBuilder.h b/Engine/Plugins/Enterprise/GLTFExporter/MigratedFromGit/Source/GLTFExporter/Public/Builders/GLTFTaskBuilder.h
--- a/Engine/Plugins/Enterprise/GLTFExporter/MigratedFromGit/Source/GLTFExporter/Public/Builders/GLTFTaskBuilder.h
+++ b/Engine/Plugins/Enterprise/GLTFExporter/MigratedFromGit/Source/GLTFExporter/Public/Builders/GLTFTaskBuilder.h
@@ -5,7 +5,7 @@
 #include "Builders/GLTFLogBuilder.h"
 #include "Tasks/GLTFDelayedTask.h"
 
-class FGLTFTaskBuilder : public FGLTFLogBuilder
+class GLTFEXPORTER_API FGLTFTaskBuilder : public FGLTFLogBuilder
 {
 public:
 
diff --git a/Engine/Plugins/Enterprise/GLTFExporter/MigratedFromGit/Source/GLTFExporter/Public/Converters/GLTFAccessorConverters.h b/Engine/Plugins/Enterprise/GLTFExporter/MigratedFromGit/Source/GLTFExporter/Public/Converters/GLTFAccessorConverters.h
--- a/Engine/Plugins/Enterprise/GLTFExporter/MigratedFromGit/Source/GLTFExporter/Public/Converters/GLTFAccessorConverters.h
+++ b/Engine/Plugins/Enterprise/GLTFExporter/MigratedFromGit/Source/GLTFExporter/Public/Converters/GLTFAccessorConverters.h
@@ -16,21 +16,21 @@ typedef TGLTFConverter<FGLTFJsonAccessor*, const FGLTFMeshSection*, const FSkinW
 typedef TGLTFConverter<FGLTFJsonAccessor*, const FGLTFMeshSection*, const FSkinWeightVertexBuffer*, uint32> IGLTFBoneWeightBufferConverter;
 typedef TGLTFConverter<FGLTFJsonAccessor*, const FGLTFMeshSection*> IGLTFIndexBufferConverter;
 
-class FGLTFPositionBufferConverter final : public FGLTFBuilderContext, public IGLTFPositionBufferConverter
+class GLTFEXPORTER_API FGLTFPositionBufferConverter final : public FGLTFBuilderContext, public IGLTFPositionBufferConverter
 {
 	using FGLTFBuilderContext::FGLTFBuilderContext;
 
 	virtual FGLTFJsonAccessor* Convert(const FGLTFMeshSection* MeshSection, const FPositionVertexBuffer* VertexBuffer) override;
 };
 
-class FGLTFColorBufferConverter final : public FGLTFBuilderContext, public IGLTFColorBufferConverter
+class GLTFEXPORTER_API FGLTFColorBufferConverter final : public FGLTFBuilderContext, public IGLTFColorBufferConverter
 {
 	using FGLTFBuilderContext::FGLTFBuilderContext;
 
 	virtual FGLTFJsonAccessor* Convert(const FGLTFMeshSection* MeshSection, const FColorVertexBuffer* VertexBuffer) override;
 };
 
-class FGLTFNormalBufferConverter final : public FGLTFBuilderContext, public IGLTFNormalBufferConverter
+class GLTFEXPORTER_API FGLTFNormalBufferConverter final : public FGLTFBuilderContext, public IGLTFNormalBufferConverter
 {
 	using FGLTFBuilderContext::FGLTFBuilderContext;
 
@@ -40,7 +40,7 @@ class FGLTFNormalBufferConverter final : public FGLTFBuilderContext, public IGLT
 	FGLTFJsonBufferView* ConvertBufferView(const FGLTFMeshSection* MeshSection, const void* TangentData);
 };
 
-class FGLTFTangentBufferConverter final : public FGLTFBuilderContext, public IGLTFTangentBufferConverter
+class GLTFEXPORTER_API FGLTFTangentBufferConverter final : public FGLTFBuilderContext, public IGLTFTangentBufferConverter
 {
 	using FGLTFBuilderContext::FGLTFBuilderContext;
 
@@ -50,7 +50,7 @@ class FGLTFTangentBufferConverter final : public FGLTFBuilderContext, public IGL
 	FGLTFJsonBufferView* ConvertBufferView(const FGLTFMeshSection* MeshSection, const void* TangentData);
 };
 
-class FGLTFUVBufferConverter final : public FGLTFBuilderContext, public IGLTFUVBufferConverter
+class GLTFEXPORTER_API FGLTFUVBufferConverter final : public FGLTFBuilderContext, public IGLTFUVBufferConverter
 {
 	using FGLTFBuilderContext::FGLTFBuilderContext;
 
@@ -60,7 +60,7 @@ class FGLTFUVBufferConverter final : public FGLTFBuilderContext, public IGLTFUVB
 	FGLTFJsonAccessor* Convert(const FGLTFMeshSection* MeshSection, const FStaticMeshVertexBuffer* VertexBuffer, uint32 UVIndex, const uint8* SourceData) const;
 };
 
-class FGLTFBoneIndexBufferConverter final : public FGLTFBuilderContext, public IGLTFBoneIndexBufferConverter
+class GLTFEXPORTER_API FGLTFBoneIndexBufferConverter final : public FGLTFBuilderContext, public IGLTFBoneIndexBufferConverter
 {
 	using FGLTFBuilderContext::FGLTFBuilderContext;
 
@@ -76,7 +76,7 @@ class FGLTFBoneIndexBufferConverter final : public FGLTFBuilderContext, public I
 	FGLTFJsonAccessor* Convert(const FGLTFMeshSection* MeshSection, const FSkinWeightVertexBuffer* VertexBuffer, uint32 InfluenceOffset, const uint8* SourceData, CallbackType GetVertexInfluenceOffsetCount) const;
 };
 
-class FGLTFBoneWeightBufferConverter final : public FGLTFBuilderContext, public IGLTFBoneWeightBufferConverter
+class GLTFEXPORTER_API FGLTFBoneWeightBufferConverter final : public FGLTFBuilderContext, public IGLTFBoneWeightBufferConverter
 {
 	using FGLTFBuilderContext::FGLTFBuilderContext;
 
@@ -89,7 +89,7 @@ class FGLTFBoneWeightBufferConverter final : public FGLTFBuilderContext, public
 	FGLTFJsonAccessor* Convert(const FGLTFMeshSection* MeshSection, const FSkinWeightVertexBuffer* VertexBuffer, uint32 InfluenceOffset, const uint8* SourceData, CallbackType GetVertexInfluenceOffsetCount) const;
 };
 
-class FGLTFIndexBufferConverter final : public FGLTFBuilderContext, public IGLTFIndexBufferConverter
+class GLTFEXPORTER_API FGLTFIndexBufferConverter final : public FGLTFBuilderContext, public IGLTFIndexBufferConverter
 {
 	using FGLTFBuilderContext::FGLTFBuilderContext;
 
diff --git a/Engine/Plugins/Enterprise/GLTFExporter/MigratedFromGit/Source/GLTFExporter/Public/Converters/GLTFAnimationConverters.h b/Engine/Plugins/Enterprise/GLTFExporter/MigratedFromGit/Source/GLTFExporter/Public/Converters/GLTFAnimationConverters.h
--- a/Engine/Plugins/Enterprise/GLTFExporter/MigratedFromGit/Source/GLTFExporter/Public/Converters/GLTFAnimationConverters.h
+++ b/Engine/Plugins/Enterprise/GLTFExporter/MigratedFromGit/Source/GLTFExporter/Public/Converters/GLTFAnimationConverters.h
@@ -14,28 +14,28 @@ typedef TGLTFConverter<FGLTFJsonAnimation*, FGLTFJsonNode*, const USkeletalMeshC
 typedef TGLTFConverter<FGLTFJsonAnimation*, const ULevel*, const ULevelSequence*> IGLTFLevelSequenceConverter;
 typedef TGLTFConverter<FGLTFJsonAnimation*, const ALevelSequenceActor*> IGLTFLevelSequenceDataConverter;
 
-class FGLTFAnimationConverter final : public FGLTFBuilderContext, public IGLTFAnimationConverter
+class GLTFEXPORTER_API FGLTFAnimationConverter final : public FGLTFBuilderContext, public IGLTFAnimationConverter
 {
 	using FGLTFBuilderContext::FGLTFBuilderContext;
 
 	virtual FGLTFJsonAnimation* Convert(FGLTFJsonNode* RootNode, const USkeletalMesh* SkeletalMesh, const UAnimSequence* AnimSequence) override;
 };
 
-class FGLTFAnimationDataConverter final : public FGLTFBuilderContext, public IGLTFAnimationDataConverter
+class GLTFEXPORTER_API FGLTFAnimationDataConverter final : public FGLTFBuilderContext, public IGLTFAnimationDataConverter
 {
 	using FGLTFBuilderContext::FGLTFBuilderContext;
 
 	virtual FGLTFJsonAnimation* Convert(FGLTFJsonNode* RootNode, const USkeletalMeshComponent* SkeletalMeshComponent) override;
 };
 
-class FGLTFLevelSequenceConverter final : public FGLTFBuilderContext, public IGLTFLevelSequenceConverter
+class GLTFEXPORTER_API FGLTFLevelSequenceConverter final : public FGLTFBuilderContext, public IGLTFLevelSequenceConverter
 {
 	using FGLTFBuilderContext::FGLTFBuilderContext;
 
 	virtual FGLTFJsonAnimation* Convert(const ULevel* Level, const ULevelSequence*) override;
 };
 
-class FGLTFLevelSequenceDataConverter final : public FGLTFBuilderContext, public IGLTFLevelSequenceDataConverter
+class GLTFEXPORTER_API FGLTFLevelSequenceDataConverter final : public FGLTFBuilderContext, public IGLTFLevelSequenceDataConverter
 {
 	using FGLTFBuilderContext::FGLTFBuilderContext;
 
diff --git a/Engine/Plugins/Enterprise/GLTFExporter/MigratedFromGit/Source/GLTFExporter/Public/Converters/GLTFBackdropConverters.h b/Engine/Plugins/Enterprise/GLTFExporter/MigratedFromGit/Source/GLTFExporter/Public/Converters/GLTFBackdropConverters.h
--- a/Engine/Plugins/Enterprise/GLTFExporter/MigratedFromGit/Source/GLTFExporter/Public/Converters/GLTFBackdropConverters.h
+++ b/Engine/Plugins/Enterprise/GLTFExporter/MigratedFromGit/Source/GLTFExporter/Public/Converters/GLTFBackdropConverters.h
@@ -9,7 +9,7 @@
 typedef TGLTFConverter<FGLTFJsonBackdrop*, const AActor*> IGLTFBackdropConverter;
 
 
-class FGLTFBackdropConverter final : public FGLTFBuilderContext, public IGLTFBackdropConverter
+class GLTFEXPORTER_API FGLTFBackdropConverter final : public FGLTFBuilderContext, public IGLTFBackdropConverter
 {
 	using FGLTFBuilderContext::FGLTFBuilderContext;
 
diff --git a/Engine/Plugins/Enterprise/GLTFExporter/MigratedFromGit/Source/GLTFExporter/Public/Converters/GLTFBuilderContext.h b/Engine/Plugins/Enterprise/GLTFExporter/MigratedFromGit/Source/GLTFExporter/Public/Converters/GLTFBuilderContext.h
--- a/Engine/Plugins/Enterprise/GLTFExporter/MigratedFromGit/Source/GLTFExporter/Public/Converters/GLTFBuilderContext.h
+++ b/Engine/Plugins/Enterprise/GLTFExporter/MigratedFromGit/Source/GLTFExporter/Public/Converters/GLTFBuilderContext.h
@@ -4,7 +4,7 @@
 
 class FGLTFConvertBuilder;
 
-class FGLTFBuilderContext
+class GLTFEXPORTER_API FGLTFBuilderContext
 {
 public:
 
diff --git a/Engine/Plugins/Enterprise/GLTFExporter/MigratedFromGit/Source/GLTFExporter/Public/Converters/GLTFCameraConverters.h b/Engine/Plugins/Enterprise/GLTFExporter/MigratedFromGit/Source/GLTFExporter/Public/Converters/GLTFCameraConverters.h
--- a/Engine/Plugins/Enterprise/GLTFExporter/MigratedFromGit/Source/GLTFExporter/Public/Converters/GLTFCameraConverters.h
+++ b/Engine/Plugins/Enterprise/GLTFExporter/MigratedFromGit/Source/GLTFExporter/Public/Converters/GLTFCameraConverters.h
@@ -10,7 +10,7 @@ class UCameraComponent;
 
 typedef TGLTFConverter<FGLTFJsonCamera*, const UCameraComponent*> IGLTFCameraConverter;
 
-class FGLTFCameraConverter final : public FGLTFBuilderContext, public IGLTFCameraConverter
+class GLTFEXPORTER_API FGLTFCameraConverter final : public FGLTFBuilderContext, public IGLTFCameraConverter
 {
 	using FGLTFBuilderContext::FGLTFBuilderContext;
 
diff --git a/Engine/Plugins/Enterprise/GLTFExporter/MigratedFromGit/Source/GLTFExporter/Public/Converters/GLTFConverter.h b/Engine/Plugins/Enterprise/GLTFExporter/MigratedFromGit/Source/GLTFExporter/Public/Converters/GLTFConverter.h
--- a/Engine/Plugins/Enterprise/GLTFExporter/MigratedFromGit/Source/GLTFExporter/Public/Converters/GLTFConverter.h
+++ b/Engine/Plugins/Enterprise/GLTFExporter/MigratedFromGit/Source/GLTFExporter/Public/Converters/GLTFConverter.h
@@ -13,7 +13,7 @@ public:
 
 	virtual ~TGLTFConverter() = default;
 
-	OutputType Get(InputTypes... Inputs) const
+	OutputType Get(InputTypes... Inputs)
 	{
 		Sanitize(Inputs...);
 		InputTupleType InputTuple(MoveTemp(Inputs)...);
diff --git a/Engine/Plugins/Enterprise/GLTFExporter/MigratedFromGit/Source/GLTFExporter/Public/Converters/GLTFEpicLevelVariantSetsConverters.h b/Engine/Plugins/Enterprise/GLTFExporter/MigratedFromGit/Source/GLTFExporter/Public/Converters/GLTFEpicLevelVariantSetsConverters.h
--- a/Engine/Plugins/Enterprise/GLTFExporter/MigratedFromGit/Source/GLTFExporter/Public/Converters/GLTFEpicLevelVariantSetsConverters.h
+++ b/Engine/Plugins/Enterprise/GLTFExporter/MigratedFromGit/Source/GLTFExporter/Public/Converters/GLTFEpicLevelVariantSetsConverters.h
@@ -16,7 +16,7 @@ class UPropertyValueMaterial;
 
 typedef TGLTFConverter<FGLTFJsonEpicLevelVariantSets*, const ULevelVariantSets*> IGLTFEpicLevelVariantSetsConverter;
 
-class FGLTFEpicLevelVariantSetsConverter final : public FGLTFBuilderContext, public IGLTFEpicLevelVariantSetsConverter
+class GLTFEXPORTER_API FGLTFEpicLevelVariantSetsConverter final : public FGLTFBuilderContext, public IGLTFEpicLevelVariantSetsConverter
 {
 	using FGLTFBuilderContext::FGLTFBuilderContext;
 
diff --git a/Engine/Plugins/Enterprise/GLTFExporter/MigratedFromGit/Source/GLTFExporter/Public/Converters/GLTFHotspotConverters.h b/Engine/Plugins/Enterprise/GLTFExporter/MigratedFromGit/Source/GLTFExporter/Public/Converters/GLTFHotspotConverters.h
--- a/Engine/Plugins/Enterprise/GLTFExporter/MigratedFromGit/Source/GLTFExporter/Public/Converters/GLTFHotspotConverters.h
+++ b/Engine/Plugins/Enterprise/GLTFExporter/MigratedFromGit/Source/GLTFExporter/Public/Converters/GLTFHotspotConverters.h
@@ -9,7 +9,7 @@
 
 typedef TGLTFConverter<FGLTFJsonHotspot*, const AGLTFHotspotActor*> IGLTFHotspotConverter;
 
-class FGLTFHotspotConverter final : public FGLTFBuilderContext, public IGLTFHotspotConverter
+class GLTFEXPORTER_API FGLTFHotspotConverter final : public FGLTFBuilderContext, public IGLTFHotspotConverter
 {
 	using FGLTFBuilderContext::FGLTFBuilderContext;
 
diff --git a/Engine/Plugins/Enterprise/GLTFExporter/MigratedFromGit/Source/GLTFExporter/Public/Converters/GLTFImageConverters.h b/Engine/Plugins/Enterprise/GLTFExporter/MigratedFromGit/Source/GLTFExporter/Public/Converters/GLTFImageConverters.h
--- a/Engine/Plugins/Enterprise/GLTFExporter/MigratedFromGit/Source/GLTFExporter/Public/Converters/GLTFImageConverters.h
+++ b/Engine/Plugins/Enterprise/GLTFExporter/MigratedFromGit/Source/GLTFExporter/Public/Converters/GLTFImageConverters.h
@@ -11,7 +11,7 @@
 
 typedef TGLTFConverter<FGLTFJsonImage*, TGLTFSuperfluous<FString>, EGLTFTextureType, bool, FIntPoint, TGLTFSharedArray<FColor>> IGLTFImageConverter;
 
-class FGLTFImageConverter final : public FGLTFBuilderContext, public IGLTFImageConverter
+class GLTFEXPORTER_API FGLTFImageConverter final : public FGLTFBuilderContext, public IGLTFImageConverter
 {
 	using FGLTFBuilderContext::FGLTFBuilderContext;
 
diff --git a/Engine/Plugins/Enterprise/GLTFExporter/MigratedFromGit/Source/GLTFExporter/Public/Converters/GLTFIndexArray.h b/Engine/Plugins/Enterprise/GLTFExporter/MigratedFromGit/Source/GLTFExporter/Public/Converters/GLTFIndexArray.h
--- a/Engine/Plugins/Enterprise/GLTFExporter/MigratedFromGit/Source/GLTFExporter/Public/Converters/GLTFIndexArray.h
+++ b/Engine/Plugins/Enterprise/GLTFExporter/MigratedFromGit/Source/GLTFExporter/Public/Converters/GLTFIndexArray.h
@@ -4,6 +4,7 @@
 
 #include "CoreMinimal.h"
 
+// TODO: should we add GLTFEXPORTER_API ?
 class FGLTFIndexArray : public TArray<int32>
 {
 public:
diff --git a/Engine/Plugins/Enterprise/GLTFExporter/MigratedFromGit/Source/GLTFExporter/Public/Converters/GLTFKhrVariantConverters.h b/Engine/Plugins/Enterprise/GLTFExporter/MigratedFromGit/Source/GLTFExporter/Public/Converters/GLTFKhrVariantConverters.h
--- a/Engine/Plugins/Enterprise/GLTFExporter/MigratedFromGit/Source/GLTFExporter/Public/Converters/GLTFKhrVariantConverters.h
+++ b/Engine/Plugins/Enterprise/GLTFExporter/MigratedFromGit/Source/GLTFExporter/Public/Converters/GLTFKhrVariantConverters.h
@@ -12,7 +12,7 @@ class UPropertyValueMaterial;
 
 typedef TGLTFConverter<FGLTFJsonKhrMaterialVariant*, const UVariant*> IGLTFKhrMaterialVariantConverter;
 
-class FGLTFKhrMaterialVariantConverter final : public FGLTFBuilderContext, public IGLTFKhrMaterialVariantConverter
+class GLTFEXPORTER_API FGLTFKhrMaterialVariantConverter final : public FGLTFBuilderContext, public IGLTFKhrMaterialVariantConverter
 {
 	using FGLTFBuilderContext::FGLTFBuilderContext;
 
diff --git a/Engine/Plugins/Enterprise/GLTFExporter/MigratedFromGit/Source/GLTFExporter/Public/Converters/GLTFLightConverters.h b/Engine/Plugins/Enterprise/GLTFExporter/MigratedFromGit/Source/GLTFExporter/Public/Converters/GLTFLightConverters.h
--- a/Engine/Plugins/Enterprise/GLTFExporter/MigratedFromGit/Source/GLTFExporter/Public/Converters/GLTFLightConverters.h
+++ b/Engine/Plugins/Enterprise/GLTFExporter/MigratedFromGit/Source/GLTFExporter/Public/Converters/GLTFLightConverters.h
@@ -8,7 +8,7 @@
 
 typedef TGLTFConverter<FGLTFJsonLight*, const ULightComponent*> IGLTFLightConverter;
 
-class FGLTFLightConverter final : public FGLTFBuilderContext, public IGLTFLightConverter
+class GLTFEXPORTER_API FGLTFLightConverter final : public FGLTFBuilderContext, public IGLTFLightConverter
 {
 	using FGLTFBuilderContext::FGLTFBuilderContext;
 
diff --git a/Engine/Plugins/Enterprise/GLTFExporter/MigratedFromGit/Source/GLTFExporter/Public/Converters/GLTFLightMapConverters.h b/Engine/Plugins/Enterprise/GLTFExporter/MigratedFromGit/Source/GLTFExporter/Public/Converters/GLTFLightMapConverters.h
--- a/Engine/Plugins/Enterprise/GLTFExporter/MigratedFromGit/Source/GLTFExporter/Public/Converters/GLTFLightMapConverters.h
+++ b/Engine/Plugins/Enterprise/GLTFExporter/MigratedFromGit/Source/GLTFExporter/Public/Converters/GLTFLightMapConverters.h
@@ -8,7 +8,7 @@
 
 typedef TGLTFConverter<FGLTFJsonLightMap*, const UStaticMeshComponent*> IGLTFLightMapConverter;
 
-class FGLTFLightMapConverter final : public FGLTFBuilderContext, public IGLTFLightMapConverter
+class GLTFEXPORTER_API FGLTFLightMapConverter final : public FGLTFBuilderContext, public IGLTFLightMapConverter
 {
 	using FGLTFBuilderContext::FGLTFBuilderContext;
 
diff --git a/Engine/Plugins/Enterprise/GLTFExporter/MigratedFromGit/Source/GLTFExporter/Public/Converters/GLTFMaterialArray.h b/Engine/Plugins/Enterprise/GLTFExporter/MigratedFromGit/Source/GLTFExporter/Public/Converters/GLTFMaterialArray.h
--- a/Engine/Plugins/Enterprise/GLTFExporter/MigratedFromGit/Source/GLTFExporter/Public/Converters/GLTFMaterialArray.h
+++ b/Engine/Plugins/Enterprise/GLTFExporter/MigratedFromGit/Source/GLTFExporter/Public/Converters/GLTFMaterialArray.h
@@ -6,6 +6,7 @@
 
 class UMaterialInterface;
 
+// TODO: should we add GLTFEXPORTER_API ?
 class FGLTFMaterialArray : public TArray<const UMaterialInterface*>
 {
 public:
diff --git a/Engine/Plugins/Enterprise/GLTFExporter/MigratedFromGit/Source/GLTFExporter/Public/Converters/GLTFMaterialConverters.h b/Engine/Plugins/Enterprise/GLTFExporter/MigratedFromGit/Source/GLTFExporter/Public/Converters/GLTFMaterialConverters.h
--- a/Engine/Plugins/Enterprise/GLTFExporter/MigratedFromGit/Source/GLTFExporter/Public/Converters/GLTFMaterialConverters.h
+++ b/Engine/Plugins/Enterprise/GLTFExporter/MigratedFromGit/Source/GLTFExporter/Public/Converters/GLTFMaterialConverters.h
@@ -12,7 +12,7 @@
 
 typedef TGLTFConverter<FGLTFJsonMaterial*, const UMaterialInterface*, const FGLTFMeshData*, FGLTFIndexArray> IGLTFMaterialConverter;
 
-class FGLTFMaterialConverter final : public FGLTFBuilderContext, public IGLTFMaterialConverter
+class GLTFEXPORTER_API FGLTFMaterialConverter final : public FGLTFBuilderContext, public IGLTFMaterialConverter
 {
 	using FGLTFBuilderContext::FGLTFBuilderContext;
 
diff --git a/Engine/Plugins/Enterprise/GLTFExporter/MigratedFromGit/Source/GLTFExporter/Public/Converters/GLTFMeshConverters.h b/Engine/Plugins/Enterprise/GLTFExporter/MigratedFromGit/Source/GLTFExporter/Public/Converters/GLTFMeshConverters.h
--- a/Engine/Plugins/Enterprise/GLTFExporter/MigratedFromGit/Source/GLTFExporter/Public/Converters/GLTFMeshConverters.h
+++ b/Engine/Plugins/Enterprise/GLTFExporter/MigratedFromGit/Source/GLTFExporter/Public/Converters/GLTFMeshConverters.h
@@ -6,13 +6,12 @@
 #include "Converters/GLTFConverter.h"
 #include "Converters/GLTFBuilderContext.h"
 #include "Converters/GLTFMeshSectionConverters.h"
-#include "Converters/GLTFMeshDataConverters.h"
 #include "Converters/GLTFMaterialArray.h"
 
 typedef TGLTFConverter<FGLTFJsonMesh*, const UStaticMesh*, const UStaticMeshComponent*, FGLTFMaterialArray, int32> IGLTFStaticMeshConverter;
 typedef TGLTFConverter<FGLTFJsonMesh*, const USkeletalMesh*, const USkeletalMeshComponent*, FGLTFMaterialArray, int32> IGLTFSkeletalMeshConverter;
 
-class FGLTFStaticMeshConverter final : public FGLTFBuilderContext, public IGLTFStaticMeshConverter
+class GLTFEXPORTER_API FGLTFStaticMeshConverter final : public FGLTFBuilderContext, public IGLTFStaticMeshConverter
 {
 	using FGLTFBuilderContext::FGLTFBuilderContext;
 
@@ -23,7 +22,7 @@ class FGLTFStaticMeshConverter final : public FGLTFBuilderContext, public IGLTFS
 	FGLTFStaticMeshSectionConverter MeshSectionConverter;
 };
 
-class FGLTFSkeletalMeshConverter final : public FGLTFBuilderContext, public IGLTFSkeletalMeshConverter
+class GLTFEXPORTER_API FGLTFSkeletalMeshConverter final : public FGLTFBuilderContext, public IGLTFSkeletalMeshConverter
 {
 	using FGLTFBuilderContext::FGLTFBuilderContext;
 
diff --git a/Engine/Plugins/Enterprise/GLTFExporter/MigratedFromGit/Source/GLTFExporter/Public/Converters/GLTFMeshData.h b/Engine/Plugins/Enterprise/GLTFExporter/MigratedFromGit/Source/GLTFExporter/Public/Converters/GLTFMeshData.h
--- a/Engine/Plugins/Enterprise/GLTFExporter/MigratedFromGit/Source/GLTFExporter/Public/Converters/GLTFMeshData.h
+++ b/Engine/Plugins/Enterprise/GLTFExporter/MigratedFromGit/Source/GLTFExporter/Public/Converters/GLTFMeshData.h
@@ -10,7 +10,7 @@
 #include "GLTFMaterialBakingStructures.h"
 #endif
 
-struct FGLTFMeshData
+struct GLTFEXPORTER_API FGLTFMeshData
 {
 	FGLTFMeshData(const UStaticMesh* StaticMesh, const UStaticMeshComponent* StaticMeshComponent, int32 LODIndex);
 	FGLTFMeshData(const USkeletalMesh* SkeletalMesh, const USkeletalMeshComponent* SkeletalMeshComponent, int32 LODIndex);
diff --git a/Engine/Plugins/Enterprise/GLTFExporter/MigratedFromGit/Source/GLTFExporter/Public/Converters/GLTFMeshDataConverters.h b/Engine/Plugins/Enterprise/GLTFExporter/MigratedFromGit/Source/GLTFExporter/Public/Converters/GLTFMeshDataConverters.h
--- a/Engine/Plugins/Enterprise/GLTFExporter/MigratedFromGit/Source/GLTFExporter/Public/Converters/GLTFMeshDataConverters.h
+++ b/Engine/Plugins/Enterprise/GLTFExporter/MigratedFromGit/Source/GLTFExporter/Public/Converters/GLTFMeshDataConverters.h
@@ -12,10 +12,21 @@ typedef TGLTFConverter<const FGLTFMeshData*, const USkeletalMesh*, const USkelet
 template <typename MeshType, typename MeshComponentType>
 class TGLTFMeshDataConverter : public FGLTFBuilderContext, public TGLTFConverter<const FGLTFMeshData*, const MeshType*, const MeshComponentType*, int32>
 {
-	TArray<TUniquePtr<FGLTFMeshData>> Outputs;
+public:
 
 	using FGLTFBuilderContext::FGLTFBuilderContext;
 
+	TGLTFMeshDataConverter(TGLTFMeshDataConverter&&) = default;
+	TGLTFMeshDataConverter& operator=(TGLTFMeshDataConverter&&) = default;
+
+	TGLTFMeshDataConverter(const TGLTFMeshDataConverter&) = delete;
+	TGLTFMeshDataConverter& operator=(const TGLTFMeshDataConverter&) = delete;
+
+private:
+
+	// TODO: standardize converters that are responsible for deleting their output
+	TArray<TUniquePtr<FGLTFMeshData>> Outputs;
+
 	virtual void Sanitize(const MeshType*& Mesh, const MeshComponentType*& MeshComponent, int32& LODIndex) override;
 
 	virtual const FGLTFMeshData* Convert(const MeshType* Mesh, const MeshComponentType* MeshComponent, int32 LODIndex) override;
diff --git a/Engine/Plugins/Enterprise/GLTFExporter/MigratedFromGit/Source/GLTFExporter/Public/Converters/GLTFMeshSection.h b/Engine/Plugins/Enterprise/GLTFExporter/MigratedFromGit/Source/GLTFExporter/Public/Converters/GLTFMeshSection.h
--- a/Engine/Plugins/Enterprise/GLTFExporter/MigratedFromGit/Source/GLTFExporter/Public/Converters/GLTFMeshSection.h
+++ b/Engine/Plugins/Enterprise/GLTFExporter/MigratedFromGit/Source/GLTFExporter/Public/Converters/GLTFMeshSection.h
@@ -5,7 +5,7 @@
 #include "CoreMinimal.h"
 #include "Converters/GLTFIndexArray.h"
 
-struct FGLTFMeshSection
+struct GLTFEXPORTER_API FGLTFMeshSection
 {
 	FGLTFMeshSection(const FStaticMeshLODResources* MeshLOD, const FGLTFIndexArray& SectionIndices);
 	FGLTFMeshSection(const FSkeletalMeshLODRenderData* MeshLOD, const FGLTFIndexArray& SectionIndices);
diff --git a/Engine/Plugins/Enterprise/GLTFExporter/MigratedFromGit/Source/GLTFExporter/Public/Converters/GLTFMeshSectionConverters.h b/Engine/Plugins/Enterprise/GLTFExporter/MigratedFromGit/Source/GLTFExporter/Public/Converters/GLTFMeshSectionConverters.h
--- a/Engine/Plugins/Enterprise/GLTFExporter/MigratedFromGit/Source/GLTFExporter/Public/Converters/GLTFMeshSectionConverters.h
+++ b/Engine/Plugins/Enterprise/GLTFExporter/MigratedFromGit/Source/GLTFExporter/Public/Converters/GLTFMeshSectionConverters.h
@@ -9,6 +9,17 @@
 template <typename MeshLODType>
 class TGLTFMeshSectionConverter final : public TGLTFConverter<const FGLTFMeshSection*, const MeshLODType*, FGLTFIndexArray>
 {
+public:
+
+	TGLTFMeshSectionConverter() = default;
+	TGLTFMeshSectionConverter(TGLTFMeshSectionConverter&&) = default;
+	TGLTFMeshSectionConverter& operator=(TGLTFMeshSectionConverter&&) = default;
+
+	TGLTFMeshSectionConverter(const TGLTFMeshSectionConverter&) = delete;
+	TGLTFMeshSectionConverter& operator=(const TGLTFMeshSectionConverter&) = delete;
+
+private:
+
 	TArray<TUniquePtr<FGLTFMeshSection>> Outputs;
 
 	virtual const FGLTFMeshSection* Convert(const MeshLODType* MeshLOD, FGLTFIndexArray SectionIndices) override
diff --git a/Engine/Plugins/Enterprise/GLTFExporter/MigratedFromGit/Source/GLTFExporter/Public/Converters/GLTFNodeConverters.h b/Engine/Plugins/Enterprise/GLTFExporter/MigratedFromGit/Source/GLTFExporter/Public/Converters/GLTFNodeConverters.h
--- a/Engine/Plugins/Enterprise/GLTFExporter/MigratedFromGit/Source/GLTFExporter/Public/Converters/GLTFNodeConverters.h
+++ b/Engine/Plugins/Enterprise/GLTFExporter/MigratedFromGit/Source/GLTFExporter/Public/Converters/GLTFNodeConverters.h
@@ -13,42 +13,42 @@ typedef TGLTFConverter<FGLTFJsonNode*, FGLTFJsonNode*, const UStaticMesh*, FName
 typedef TGLTFConverter<FGLTFJsonNode*, FGLTFJsonNode*, const USkeletalMesh*, FName> IGLTFSkeletalSocketConverter;
 typedef TGLTFConverter<FGLTFJsonNode*, FGLTFJsonNode*, const USkeletalMesh*, int32> IGLTFSkeletalBoneConverter;
 
-class FGLTFActorConverter final : public FGLTFBuilderContext, public IGLTFActorConverter
+class GLTFEXPORTER_API FGLTFActorConverter final : public FGLTFBuilderContext, public IGLTFActorConverter
 {
 	using FGLTFBuilderContext::FGLTFBuilderContext;
 
 	virtual FGLTFJsonNode* Convert(const AActor* Actor) override;
 };
 
-class FGLTFComponentConverter final : public FGLTFBuilderContext, public IGLTFComponentConverter
+class GLTFEXPORTER_API FGLTFComponentConverter final : public FGLTFBuilderContext, public IGLTFComponentConverter
 {
 	using FGLTFBuilderContext::FGLTFBuilderContext;
 
 	virtual FGLTFJsonNode* Convert(const USceneComponent* SceneComponent) override;
 };
 
-class FGLTFComponentSocketConverter final : public FGLTFBuilderContext, public IGLTFComponentSocketConverter
+class GLTFEXPORTER_API FGLTFComponentSocketConverter final : public FGLTFBuilderContext, public IGLTFComponentSocketConverter
 {
 	using FGLTFBuilderContext::FGLTFBuilderContext;
 
 	virtual FGLTFJsonNode* Convert(const USceneComponent* SceneComponent, FName SocketName) override;
 };
 
-class FGLTFStaticSocketConverter final : public FGLTFBuilderContext, public IGLTFStaticSocketConverter
+class GLTFEXPORTER_API FGLTFStaticSocketConverter final : public FGLTFBuilderContext, public IGLTFStaticSocketConverter
 {
 	using FGLTFBuilderContext::FGLTFBuilderContext;
 
 	virtual FGLTFJsonNode* Convert(FGLTFJsonNode* RootNode, const UStaticMesh* StaticMesh, FName SocketName) override;
 };
 
-class FGLTFSkeletalSocketConverter final : public FGLTFBuilderContext, public IGLTFSkeletalSocketConverter
+class GLTFEXPORTER_API FGLTFSkeletalSocketConverter final : public FGLTFBuilderContext, public IGLTFSkeletalSocketConverter
 {
 	using FGLTFBuilderContext::FGLTFBuilderContext;
 
 	virtual FGLTFJsonNode* Convert(FGLTFJsonNode* RootNode, const USkeletalMesh* SkeletalMesh, FName SocketName) override;
 };
 
-class FGLTFSkeletalBoneConverter final : public FGLTFBuilderContext, public IGLTFSkeletalBoneConverter
+class GLTFEXPORTER_API FGLTFSkeletalBoneConverter final : public FGLTFBuilderContext, public IGLTFSkeletalBoneConverter
 {
 	using FGLTFBuilderContext::FGLTFBuilderContext;
 
diff --git a/Engine/Plugins/Enterprise/GLTFExporter/MigratedFromGit/Source/GLTFExporter/Public/Converters/GLTFSamplerConverters.h b/Engine/Plugins/Enterprise/GLTFExporter/MigratedFromGit/Source/GLTFExporter/Public/Converters/GLTFSamplerConverters.h
--- a/Engine/Plugins/Enterprise/GLTFExporter/MigratedFromGit/Source/GLTFExporter/Public/Converters/GLTFSamplerConverters.h
+++ b/Engine/Plugins/Enterprise/GLTFExporter/MigratedFromGit/Source/GLTFExporter/Public/Converters/GLTFSamplerConverters.h
@@ -8,7 +8,7 @@
 
 typedef TGLTFConverter<FGLTFJsonSampler*, const UTexture*> IGLTFSamplerConverter;
 
-class FGLTFSamplerConverter final : public FGLTFBuilderContext, public IGLTFSamplerConverter
+class GLTFEXPORTER_API FGLTFSamplerConverter final : public FGLTFBuilderContext, public IGLTFSamplerConverter
 {
 	using FGLTFBuilderContext::FGLTFBuilderContext;
 
diff --git a/Engine/Plugins/Enterprise/GLTFExporter/MigratedFromGit/Source/GLTFExporter/Public/Converters/GLTFSceneConverters.h b/Engine/Plugins/Enterprise/GLTFExporter/MigratedFromGit/Source/GLTFExporter/Public/Converters/GLTFSceneConverters.h
--- a/Engine/Plugins/Enterprise/GLTFExporter/MigratedFromGit/Source/GLTFExporter/Public/Converters/GLTFSceneConverters.h
+++ b/Engine/Plugins/Enterprise/GLTFExporter/MigratedFromGit/Source/GLTFExporter/Public/Converters/GLTFSceneConverters.h
@@ -8,7 +8,7 @@
 
 typedef TGLTFConverter<FGLTFJsonScene*, const UWorld*> IGLTFSceneConverter;
 
-class FGLTFSceneConverter final : public FGLTFBuilderContext, public IGLTFSceneConverter
+class GLTFEXPORTER_API FGLTFSceneConverter final : public FGLTFBuilderContext, public IGLTFSceneConverter
 {
 	using FGLTFBuilderContext::FGLTFBuilderContext;
 
diff --git a/Engine/Plugins/Enterprise/GLTFExporter/MigratedFromGit/Source/GLTFExporter/Public/Converters/GLTFSkinConverters.h b/Engine/Plugins/Enterprise/GLTFExporter/MigratedFromGit/Source/GLTFExporter/Public/Converters/GLTFSkinConverters.h
--- a/Engine/Plugins/Enterprise/GLTFExporter/MigratedFromGit/Source/GLTFExporter/Public/Converters/GLTFSkinConverters.h
+++ b/Engine/Plugins/Enterprise/GLTFExporter/MigratedFromGit/Source/GLTFExporter/Public/Converters/GLTFSkinConverters.h
@@ -8,7 +8,7 @@
 
 typedef TGLTFConverter<FGLTFJsonSkin*, FGLTFJsonNode*, const USkeletalMesh*> IGLTFSkinConverter;
 
-class FGLTFSkinConverter final : public FGLTFBuilderContext, public IGLTFSkinConverter
+class GLTFEXPORTER_API FGLTFSkinConverter final : public FGLTFBuilderContext, public IGLTFSkinConverter
 {
 	using FGLTFBuilderContext::FGLTFBuilderContext;
 
diff --git a/Engine/Plugins/Enterprise/GLTFExporter/MigratedFromGit/Source/GLTFExporter/Public/Converters/GLTFSkySphereConverters.h b/Engine/Plugins/Enterprise/GLTFExporter/MigratedFromGit/Source/GLTFExporter/Public/Converters/GLTFSkySphereConverters.h
--- a/Engine/Plugins/Enterprise/GLTFExporter/MigratedFromGit/Source/GLTFExporter/Public/Converters/GLTFSkySphereConverters.h
+++ b/Engine/Plugins/Enterprise/GLTFExporter/MigratedFromGit/Source/GLTFExporter/Public/Converters/GLTFSkySphereConverters.h
@@ -11,7 +11,7 @@ struct FGLTFJsonSkySphereColorCurve;
 
 typedef TGLTFConverter<FGLTFJsonSkySphere*, const AActor*> IGLTFSkySphereConverter;
 
-class FGLTFSkySphereConverter final : public FGLTFBuilderContext, public IGLTFSkySphereConverter
+class GLTFEXPORTER_API FGLTFSkySphereConverter final : public FGLTFBuilderContext, public IGLTFSkySphereConverter
 {
 	enum class ESkySphereTextureParameter
 	{
diff --git a/Engine/Plugins/Enterprise/GLTFExporter/MigratedFromGit/Source/GLTFExporter/Public/Converters/GLTFTextureConverters.h b/Engine/Plugins/Enterprise/GLTFExporter/MigratedFromGit/Source/GLTFExporter/Public/Converters/GLTFTextureConverters.h
--- a/Engine/Plugins/Enterprise/GLTFExporter/MigratedFromGit/Source/GLTFExporter/Public/Converters/GLTFTextureConverters.h
+++ b/Engine/Plugins/Enterprise/GLTFExporter/MigratedFromGit/Source/GLTFExporter/Public/Converters/GLTFTextureConverters.h
@@ -16,7 +16,7 @@ typedef TGLTFConverter<FGLTFJsonTexture*, const UTextureRenderTarget2D*, bool> I
 typedef TGLTFConverter<FGLTFJsonTexture*, const UTextureRenderTargetCube*, ECubeFace, bool> IGLTFTextureRenderTargetCubeConverter;
 typedef TGLTFConverter<FGLTFJsonTexture*, const ULightMapTexture2D*> IGLTFTextureLightMapConverter;
 
-class FGLTFTexture2DConverter final : public FGLTFBuilderContext, public IGLTFTexture2DConverter
+class GLTFEXPORTER_API FGLTFTexture2DConverter final : public FGLTFBuilderContext, public IGLTFTexture2DConverter
 {
 	using FGLTFBuilderContext::FGLTFBuilderContext;
 
@@ -25,7 +25,7 @@ class FGLTFTexture2DConverter final : public FGLTFBuilderContext, public IGLTFTe
 	virtual FGLTFJsonTexture* Convert(const UTexture2D* Texture2D, bool bToSRGB) override;
 };
 
-class FGLTFTextureCubeConverter final : public FGLTFBuilderContext, public IGLTFTextureCubeConverter
+class GLTFEXPORTER_API FGLTFTextureCubeConverter final : public FGLTFBuilderContext, public IGLTFTextureCubeConverter
 {
 	using FGLTFBuilderContext::FGLTFBuilderContext;
 
@@ -34,7 +34,7 @@ class FGLTFTextureCubeConverter final : public FGLTFBuilderContext, public IGLTF
 	virtual FGLTFJsonTexture* Convert(const UTextureCube* TextureCube, ECubeFace CubeFace, bool bToSRGB) override;
 };
 
-class FGLTFTextureRenderTarget2DConverter final : public FGLTFBuilderContext, public IGLTFTextureRenderTarget2DConverter
+class GLTFEXPORTER_API FGLTFTextureRenderTarget2DConverter final : public FGLTFBuilderContext, public IGLTFTextureRenderTarget2DConverter
 {
 	using FGLTFBuilderContext::FGLTFBuilderContext;
 
@@ -43,7 +43,7 @@ class FGLTFTextureRenderTarget2DConverter final : public FGLTFBuilderContext, pu
 	virtual FGLTFJsonTexture* Convert(const UTextureRenderTarget2D* RenderTarget2D, bool bToSRGB) override;
 };
 
-class FGLTFTextureRenderTargetCubeConverter final : public FGLTFBuilderContext, public IGLTFTextureRenderTargetCubeConverter
+class GLTFEXPORTER_API FGLTFTextureRenderTargetCubeConverter final : public FGLTFBuilderContext, public IGLTFTextureRenderTargetCubeConverter
 {
 	using FGLTFBuilderContext::FGLTFBuilderContext;
 
@@ -52,7 +52,7 @@ class FGLTFTextureRenderTargetCubeConverter final : public FGLTFBuilderContext,
 	virtual FGLTFJsonTexture* Convert(const UTextureRenderTargetCube* RenderTargetCube, ECubeFace CubeFace, bool bToSRGB) override;
 };
 
-class FGLTFTextureLightMapConverter final : public FGLTFBuilderContext, public IGLTFTextureLightMapConverter
+class GLTFEXPORTER_API FGLTFTextureLightMapConverter final : public FGLTFBuilderContext, public IGLTFTextureLightMapConverter
 {
 	using FGLTFBuilderContext::FGLTFBuilderContext;
 
diff --git a/Engine/Plugins/Enterprise/GLTFExporter/MigratedFromGit/Source/GLTFExporter/Public/Converters/GLTFUVDegenerateChecker.h b/Engine/Plugins/Enterprise/GLTFExporter/MigratedFromGit/Source/GLTFExporter/Public/Converters/GLTFUVDegenerateChecker.h
--- a/Engine/Plugins/Enterprise/GLTFExporter/MigratedFromGit/Source/GLTFExporter/Public/Converters/GLTFUVDegenerateChecker.h
+++ b/Engine/Plugins/Enterprise/GLTFExporter/MigratedFromGit/Source/GLTFExporter/Public/Converters/GLTFUVDegenerateChecker.h
@@ -7,7 +7,7 @@
 
 struct FMeshDescription;
 
-class FGLTFUVDegenerateChecker final : public TGLTFConverter<float, const FMeshDescription*, FGLTFIndexArray, int32>
+class GLTFEXPORTER_API FGLTFUVDegenerateChecker final : public TGLTFConverter<float, const FMeshDescription*, FGLTFIndexArray, int32>
 {
 	virtual void Sanitize(const FMeshDescription*& Description, FGLTFIndexArray& SectionIndices, int32& TexCoord) override;
 
diff --git a/Engine/Plugins/Enterprise/GLTFExporter/MigratedFromGit/Source/GLTFExporter/Public/Converters/GLTFUVOverlapChecker.h b/Engine/Plugins/Enterprise/GLTFExporter/MigratedFromGit/Source/GLTFExporter/Public/Converters/GLTFUVOverlapChecker.h
--- a/Engine/Plugins/Enterprise/GLTFExporter/MigratedFromGit/Source/GLTFExporter/Public/Converters/GLTFUVOverlapChecker.h
+++ b/Engine/Plugins/Enterprise/GLTFExporter/MigratedFromGit/Source/GLTFExporter/Public/Converters/GLTFUVOverlapChecker.h
@@ -7,7 +7,7 @@
 
 struct FMeshDescription;
 
-class FGLTFUVOverlapChecker final : public TGLTFConverter<float, const FMeshDescription*, FGLTFIndexArray, int32>
+class GLTFEXPORTER_API FGLTFUVOverlapChecker final : public TGLTFConverter<float, const FMeshDescription*, FGLTFIndexArray, int32>
 {
 	virtual void Sanitize(const FMeshDescription*& Description, FGLTFIndexArray& SectionIndices, int32& TexCoord) override;
 
diff --git a/Engine/Plugins/Enterprise/GLTFExporter/MigratedFromGit/Source/GLTFExporter/Public/Tasks/GLTFDelayedTask.h b/Engine/Plugins/Enterprise/GLTFExporter/MigratedFromGit/Source/GLTFExporter/Public/Tasks/GLTFDelayedTask.h
--- a/Engine/Plugins/Enterprise/GLTFExporter/MigratedFromGit/Source/GLTFExporter/Public/Tasks/GLTFDelayedTask.h
+++ b/Engine/Plugins/Enterprise/GLTFExporter/MigratedFromGit/Source/GLTFExporter/Public/Tasks/GLTFDelayedTask.h
@@ -13,7 +13,7 @@ enum class EGLTFTaskPriority : uint8
 	MAX
 };
 
-class FGLTFDelayedTask
+class GLTFEXPORTER_API FGLTFDelayedTask
 {
 public:
 
