[GLTFExporter] refactor: improve builder api to get all outputed filenames

commit 18a8da329476cbfba38f9a256a2b6ad5009ab440
Date: Sun, 13 Mar 2022 23:59:57 +0100

    refactor: improve builder api to get all outputed filenames
    (cherry picked from commit 2ccbdb415f3d1331fee3cde5fa2fc39df4631463​)

#skipci

diff --git a/Engine/Plugins/Enterprise/GLTFExporter/MigratedFromGit/Source/GLTFExporter/Private/Builders/GLTFBuilder.cpp b/Engine/Plugins/Enterprise/GLTFExporter/MigratedFromGit/Source/GLTFExporter/Private/Builders/GLTFBuilder.cpp
--- a/Engine/Plugins/Enterprise/GLTFExporter/MigratedFromGit/Source/GLTFExporter/Private/Builders/GLTFBuilder.cpp
+++ b/Engine/Plugins/Enterprise/GLTFExporter/MigratedFromGit/Source/GLTFExporter/Private/Builders/GLTFBuilder.cpp
@@ -6,9 +6,9 @@
 
 FGLTFBuilder::FGLTFBuilder(const FString& FileName, const UGLTFExportOptions* ExportOptions)
 	: FileName(FPaths::GetCleanFilename(FileName))
-	, bIsGLB(FileName.EndsWith(TEXT(".glb"), ESearchCase::IgnoreCase))
+	, bIsGLB(this->FileName.EndsWith(TEXT(".glb"), ESearchCase::IgnoreCase))
 	, ExportOptions(SanitizeExportOptions(ExportOptions))
-	, ExportOptionsGuard(ExportOptions)
+	, ExportOptionsGuard(this->ExportOptions)
 {
 }
 
diff --git a/Engine/Plugins/Enterprise/GLTFExporter/MigratedFromGit/Source/GLTFExporter/Private/Builders/GLTFContainerBuilder.cpp b/Engine/Plugins/Enterprise/GLTFExporter/MigratedFromGit/Source/GLTFExporter/Private/Builders/GLTFContainerBuilder.cpp
--- a/Engine/Plugins/Enterprise/GLTFExporter/MigratedFromGit/Source/GLTFExporter/Private/Builders/GLTFContainerBuilder.cpp
+++ b/Engine/Plugins/Enterprise/GLTFExporter/MigratedFromGit/Source/GLTFExporter/Private/Builders/GLTFContainerBuilder.cpp
@@ -8,15 +8,17 @@ FGLTFContainerBuilder::FGLTFContainerBuilder(const FString& FileName, const UGLT
 {
 }
 
-void FGLTFContainerBuilder::WriteToArchive(FArchive& Archive)
+void FGLTFContainerBuilder::WriteInternalArchive(FArchive& Archive)
 {
+	ProcessSlowTasks();
+
 	if (bIsGLB)
 	{
-		WriteGlb(Archive);
+		WriteGlbArchive(Archive);
 	}
 	else
 	{
-		WriteJson(Archive);
+		WriteJsonArchive(Archive);
 	}
 
 	const TSet<EGLTFJsonExtension> CustomExtensions = GetCustomExtensionsUsed();
@@ -32,24 +34,41 @@ void FGLTFContainerBuilder::WriteToArchive(FArchive& Archive)
 	}
 }
 
-bool FGLTFContainerBuilder::WriteAllFiles(const FString& DirPath, bool bOverwrite)
+bool FGLTFContainerBuilder::WriteAllFiles(const FString& DirPath, uint32 WriteFlags)
 {
 	FGLTFMemoryArchive Archive;
-	WriteToArchive(Archive);
+	WriteInternalArchive(Archive);
 
 	const FString FilePath = FPaths::Combine(DirPath, FileName);
-	if (!SaveToFile(FilePath, Archive, bOverwrite))
+	if (!SaveToFile(FilePath, Archive, WriteFlags))
 	{
 		return false;
 	}
 
-	return WriteExternalFiles(DirPath, bOverwrite);
+	return WriteExternalFiles(DirPath, WriteFlags);
+}
+
+bool FGLTFContainerBuilder::WriteAllFiles(const FString& DirPath, TArray<FString>& OutFilePaths, uint32 WriteFlags)
+{
+	if (!WriteAllFiles(DirPath, WriteFlags))
+	{
+		return false;
+	}
+
+	GetAllFiles(OutFilePaths, DirPath);
+	return true;
+}
+
+void FGLTFContainerBuilder::GetAllFiles(TArray<FString>& OutFilePaths, const FString& DirPath) const
+{
+	OutFilePaths.Add(FPaths::Combine(DirPath, FileName));
+	GetExternalFiles(OutFilePaths, DirPath);
 }
 
-void FGLTFContainerBuilder::WriteGlb(FArchive& Archive)
+void FGLTFContainerBuilder::WriteGlbArchive(FArchive& Archive)
 {
 	FGLTFMemoryArchive JsonData;
-	WriteJson(JsonData);
+	WriteJsonArchive(JsonData);
 
 	const TArray64<uint8>* BufferData = GetBufferData();
 	WriteGlb(Archive, JsonData, BufferData);
diff --git a/Engine/Plugins/Enterprise/GLTFExporter/MigratedFromGit/Source/GLTFExporter/Private/Builders/GLTFFileBuilder.cpp b/Engine/Plugins/Enterprise/GLTFExporter/MigratedFromGit/Source/GLTFExporter/Private/Builders/GLTFFileBuilder.cpp
--- a/Engine/Plugins/Enterprise/GLTFExporter/MigratedFromGit/Source/GLTFExporter/Private/Builders/GLTFFileBuilder.cpp
+++ b/Engine/Plugins/Enterprise/GLTFExporter/MigratedFromGit/Source/GLTFExporter/Private/Builders/GLTFFileBuilder.cpp
@@ -11,23 +11,32 @@ FGLTFFileBuilder::FGLTFFileBuilder(const FString& FileName, const UGLTFExportOpt
 FString FGLTFFileBuilder::AddExternalFile(const FString& URI, const TSharedPtr<FGLTFMemoryArchive>& Archive)
 {
 	const FString UnqiueURI = GetUniqueURI(URI);
-	ExternalFiles.Add(UnqiueURI, Archive);
+	ExternalArchives.Add(UnqiueURI, Archive);
 	return UnqiueURI;
 }
 
-const TMap<FString, TSharedPtr<FGLTFMemoryArchive>>& FGLTFFileBuilder::GetExternalFiles() const
+void FGLTFFileBuilder::GetExternalFiles(TArray<FString>& OutFilePaths, const FString& DirPath) const
 {
-	return ExternalFiles;
+	for (const TPair<FString, TSharedPtr<FGLTFMemoryArchive>>& ExternalFile : ExternalArchives)
+	{
+		const FString FilePath = FPaths::Combine(DirPath, *ExternalFile.Key);
+		OutFilePaths.Add(FilePath);
+	}
 }
 
-bool FGLTFFileBuilder::WriteExternalFiles(const FString& DirPath, bool bOverwrite)
+const TMap<FString, TSharedPtr<FGLTFMemoryArchive>>& FGLTFFileBuilder::GetExternalArchives() const
 {
-	for (const TPair<FString, TSharedPtr<FGLTFMemoryArchive>>& ExternalFile : ExternalFiles)
+	return ExternalArchives;
+}
+
+bool FGLTFFileBuilder::WriteExternalFiles(const FString& DirPath, uint32 WriteFlags)
+{
+	for (const TPair<FString, TSharedPtr<FGLTFMemoryArchive>>& ExternalFile : ExternalArchives)
 	{
 		const FString FilePath = FPaths::Combine(DirPath, *ExternalFile.Key);
 		const TArray64<uint8>& FileData = *ExternalFile.Value;
 
-		if (!SaveToFile(FilePath, FileData, bOverwrite))
+		if (!SaveToFile(FilePath, FileData, WriteFlags))
 		{
 			return false;
 		}
@@ -38,7 +47,7 @@ bool FGLTFFileBuilder::WriteExternalFiles(const FString& DirPath, bool bOverwrit
 
 FString FGLTFFileBuilder::GetUniqueURI(const FString& URI) const
 {
-	if (!ExternalFiles.Contains(URI))
+	if (!ExternalArchives.Contains(URI))
 	{
 		return URI;
 	}
@@ -53,15 +62,13 @@ FString FGLTFFileBuilder::GetUniqueURI(const FString& URI) const
 		UnqiueURI = BaseFilename + TEXT('_') + FString::FromInt(Suffix) + FileExtension;
 		Suffix++;
 	}
-	while (ExternalFiles.Contains(UnqiueURI));
+	while (ExternalArchives.Contains(UnqiueURI));
 
 	return UnqiueURI;
 }
 
-bool FGLTFFileBuilder::SaveToFile(const FString& FilePath, const TArray64<uint8>& FileData, bool bOverwrite)
+bool FGLTFFileBuilder::SaveToFile(const FString& FilePath, const TArray64<uint8>& FileData, uint32 WriteFlags)
 {
-	const uint32 WriteFlags = bOverwrite ? FILEWRITE_None : FILEWRITE_NoReplaceExisting;
-
 	if (!FFileHelper::SaveArrayToFile(FileData, *FilePath, &IFileManager::Get(), WriteFlags))
 	{
 		LogError(FString::Printf(TEXT("Failed to save file: %s"), *FilePath));
diff --git a/Engine/Plugins/Enterprise/GLTFExporter/MigratedFromGit/Source/GLTFExporter/Private/Builders/GLTFJsonBuilder.cpp b/Engine/Plugins/Enterprise/GLTFExporter/MigratedFromGit/Source/GLTFExporter/Private/Builders/GLTFJsonBuilder.cpp
--- a/Engine/Plugins/Enterprise/GLTFExporter/MigratedFromGit/Source/GLTFExporter/Private/Builders/GLTFJsonBuilder.cpp
+++ b/Engine/Plugins/Enterprise/GLTFExporter/MigratedFromGit/Source/GLTFExporter/Private/Builders/GLTFJsonBuilder.cpp
@@ -12,7 +12,7 @@ FGLTFJsonBuilder::FGLTFJsonBuilder(const FString& FileName, const UGLTFExportOpt
 	JsonRoot.Asset.Generator = GetGeneratorString();
 }
 
-void FGLTFJsonBuilder::WriteJson(FArchive& Archive)
+void FGLTFJsonBuilder::WriteJsonArchive(FArchive& Archive)
 {
 	JsonRoot.WriteJson(Archive, !bIsGLB, ExportOptions->bSkipNearDefaultValues ? KINDA_SMALL_NUMBER : 0);
 }
diff --git a/Engine/Plugins/Enterprise/GLTFExporter/MigratedFromGit/Source/GLTFExporter/Private/Exporters/GLTFExporter.cpp b/Engine/Plugins/Enterprise/GLTFExporter/MigratedFromGit/Source/GLTFExporter/Private/Exporters/GLTFExporter.cpp
--- a/Engine/Plugins/Enterprise/GLTFExporter/MigratedFromGit/Source/GLTFExporter/Private/Exporters/GLTFExporter.cpp
+++ b/Engine/Plugins/Enterprise/GLTFExporter/MigratedFromGit/Source/GLTFExporter/Private/Exporters/GLTFExporter.cpp
@@ -52,7 +52,7 @@ bool UGLTFExporter::ExportBinary(UObject* Object, const TCHAR* Type, FArchive& A
 	if (bSuccess)
 	{
 		Builder.ProcessSlowTasks(Warn);
-		Builder.WriteToArchive(Archive);
+		Builder.WriteInternalArchive(Archive);
 		bSuccess = Builder.WriteExternalFiles(DirPath);
 	}
 
@@ -93,7 +93,7 @@ bool UGLTFExporter::ExportToGLTF(UObject* Object, const FString& FilePath, const
 
 	if (bSuccess)
 	{
-		Builder.ProcessSlowTasks(nullptr);
+		Builder.ProcessSlowTasks();
 		bSuccess = Builder.WriteAllFiles(DirPath);
 	}
 
diff --git a/Engine/Plugins/Enterprise/GLTFExporter/MigratedFromGit/Source/GLTFExporter/Public/Builders/GLTFBufferBuilder.h b/Engine/Plugins/Enterprise/GLTFExporter/MigratedFromGit/Source/GLTFExporter/Public/Builders/GLTFBufferBuilder.h
--- a/Engine/Plugins/Enterprise/GLTFExporter/MigratedFromGit/Source/GLTFExporter/Public/Builders/GLTFBufferBuilder.h
+++ b/Engine/Plugins/Enterprise/GLTFExporter/MigratedFromGit/Source/GLTFExporter/Public/Builders/GLTFBufferBuilder.h
@@ -8,7 +8,7 @@ class GLTFEXPORTER_API FGLTFBufferBuilder : public FGLTFJsonBuilder
 {
 protected:
 
-	FGLTFBufferBuilder(const FString& FileName, const UGLTFExportOptions* ExportOptions);
+	FGLTFBufferBuilder(const FString& FileName, const UGLTFExportOptions* ExportOptions = nullptr);
 	~FGLTFBufferBuilder();
 
 	const FGLTFMemoryArchive* GetBufferData() const;
diff --git a/Engine/Plugins/Enterprise/GLTFExporter/MigratedFromGit/Source/GLTFExporter/Public/Builders/GLTFBuilder.h b/Engine/Plugins/Enterprise/GLTFExporter/MigratedFromGit/Source/GLTFExporter/Public/Builders/GLTFBuilder.h
--- a/Engine/Plugins/Enterprise/GLTFExporter/MigratedFromGit/Source/GLTFExporter/Public/Builders/GLTFBuilder.h
+++ b/Engine/Plugins/Enterprise/GLTFExporter/MigratedFromGit/Source/GLTFExporter/Public/Builders/GLTFBuilder.h
@@ -15,7 +15,7 @@ public:
 
 	const UGLTFExportOptions* ExportOptions; // TODO: make ExportOptions private and expose each option via getters to ease overriding settings in future
 
-	FGLTFBuilder(const FString& FileName, const UGLTFExportOptions* ExportOptions);
+	FGLTFBuilder(const FString& FileName, const UGLTFExportOptions* ExportOptions = nullptr);
 
 	const UMaterialInterface* ResolveProxy(const UMaterialInterface* Material) const;
 	void ResolveProxies(TArray<const UMaterialInterface*>& Materials) const;
diff --git a/Engine/Plugins/Enterprise/GLTFExporter/MigratedFromGit/Source/GLTFExporter/Public/Builders/GLTFContainerBuilder.h b/Engine/Plugins/Enterprise/GLTFExporter/MigratedFromGit/Source/GLTFExporter/Public/Builders/GLTFContainerBuilder.h
--- a/Engine/Plugins/Enterprise/GLTFExporter/MigratedFromGit/Source/GLTFExporter/Public/Builders/GLTFContainerBuilder.h
+++ b/Engine/Plugins/Enterprise/GLTFExporter/MigratedFromGit/Source/GLTFExporter/Public/Builders/GLTFContainerBuilder.h
@@ -8,15 +8,18 @@ class GLTFEXPORTER_API FGLTFContainerBuilder : public FGLTFConvertBuilder
 {
 public:
 
-	FGLTFContainerBuilder(const FString& FileName, const UGLTFExportOptions* ExportOptions, const TSet<AActor*>& SelectedActors);
+	FGLTFContainerBuilder(const FString& FileName, const UGLTFExportOptions* ExportOptions = nullptr, const TSet<AActor*>& SelectedActors = {});
 
-	void WriteToArchive(FArchive& Archive);
+	void WriteInternalArchive(FArchive& Archive);
 
-	bool WriteAllFiles(const FString& DirPath, bool bOverwrite = true);
+	bool WriteAllFiles(const FString& DirPath, uint32 WriteFlags = 0);
+	bool WriteAllFiles(const FString& DirPath, TArray<FString>& OutFilePaths, uint32 WriteFlags = 0);
+
+	void GetAllFiles(TArray<FString>& OutFilePaths, const FString& DirPath = TEXT("")) const;
 
 protected:
 
-	void WriteGlb(FArchive& Archive);
+	void WriteGlbArchive(FArchive& Archive);
 
 private:
 
diff --git a/Engine/Plugins/Enterprise/GLTFExporter/MigratedFromGit/Source/GLTFExporter/Public/Builders/GLTFConvertBuilder.h b/Engine/Plugins/Enterprise/GLTFExporter/MigratedFromGit/Source/GLTFExporter/Public/Builders/GLTFConvertBuilder.h
--- a/Engine/Plugins/Enterprise/GLTFExporter/MigratedFromGit/Source/GLTFExporter/Public/Builders/GLTFConvertBuilder.h
+++ b/Engine/Plugins/Enterprise/GLTFExporter/MigratedFromGit/Source/GLTFExporter/Public/Builders/GLTFConvertBuilder.h
@@ -29,7 +29,7 @@ public:
 
 	const TSet<AActor*> SelectedActors;
 
-	FGLTFConvertBuilder(const FString& FileName, const UGLTFExportOptions* ExportOptions, const TSet<AActor*>& SelectedActors = {});
+	FGLTFConvertBuilder(const FString& FileName, const UGLTFExportOptions* ExportOptions = nullptr, const TSet<AActor*>& SelectedActors = {});
 
 	bool IsSelectedActor(const AActor* Object) const;
 	bool IsRootActor(const AActor* Actor) const;
diff --git a/Engine/Plugins/Enterprise/GLTFExporter/MigratedFromGit/Source/GLTFExporter/Public/Builders/GLTFFileBuilder.h b/Engine/Plugins/Enterprise/GLTFExporter/MigratedFromGit/Source/GLTFExporter/Public/Builders/GLTFFileBuilder.h
--- a/Engine/Plugins/Enterprise/GLTFExporter/MigratedFromGit/Source/GLTFExporter/Public/Builders/GLTFFileBuilder.h
+++ b/Engine/Plugins/Enterprise/GLTFExporter/MigratedFromGit/Source/GLTFExporter/Public/Builders/GLTFFileBuilder.h
@@ -9,21 +9,23 @@ class GLTFEXPORTER_API FGLTFFileBuilder : public FGLTFTaskBuilder
 {
 public:
 
-	FGLTFFileBuilder(const FString& FileName, const UGLTFExportOptions* ExportOptions);
+	FGLTFFileBuilder(const FString& FileName, const UGLTFExportOptions* ExportOptions = nullptr);
 
 	FString AddExternalFile(const FString& URI, const TSharedPtr<FGLTFMemoryArchive>& Archive = MakeShared<FGLTFMemoryArchive>());
 
-	const TMap<FString, TSharedPtr<FGLTFMemoryArchive>>& GetExternalFiles() const;
+	void GetExternalFiles(TArray<FString>& OutFilePaths, const FString& DirPath = TEXT("")) const;
 
-	bool WriteExternalFiles(const FString& DirPath, bool bOverwrite = true);
+	const TMap<FString, TSharedPtr<FGLTFMemoryArchive>>& GetExternalArchives() const;
+
+	bool WriteExternalFiles(const FString& DirPath, uint32 WriteFlags = 0);
 
 private:
 
-	TMap<FString, TSharedPtr<FGLTFMemoryArchive>> ExternalFiles;
+	TMap<FString, TSharedPtr<FGLTFMemoryArchive>> ExternalArchives;
 
 	FString GetUniqueURI(const FString& URI) const;
 
 protected:
 
-	bool SaveToFile(const FString& FilePath, const TArray64<uint8>& FileData, bool bOverwrite);
+	bool SaveToFile(const FString& FilePath, const TArray64<uint8>& FileData, uint32 WriteFlags);
 };
diff --git a/Engine/Plugins/Enterprise/GLTFExporter/MigratedFromGit/Source/GLTFExporter/Public/Builders/GLTFJsonBuilder.h b/Engine/Plugins/Enterprise/GLTFExporter/MigratedFromGit/Source/GLTFExporter/Public/Builders/GLTFJsonBuilder.h
--- a/Engine/Plugins/Enterprise/GLTFExporter/MigratedFromGit/Source/GLTFExporter/Public/Builders/GLTFJsonBuilder.h
+++ b/Engine/Plugins/Enterprise/GLTFExporter/MigratedFromGit/Source/GLTFExporter/Public/Builders/GLTFJsonBuilder.h
@@ -11,9 +11,7 @@ public:
 
 	FGLTFJsonScene*& DefaultScene;
 
-	FGLTFJsonBuilder(const FString& FileName, const UGLTFExportOptions* ExportOptions);
-
-	void WriteJson(FArchive& Archive);
+	FGLTFJsonBuilder(const FString& FileName, const UGLTFExportOptions* ExportOptions = nullptr);
 
 	void AddExtension(EGLTFJsonExtension Extension, bool bIsRequired = false);
 
@@ -42,6 +40,8 @@ public:
 
 protected:
 
+	void WriteJsonArchive(FArchive& Archive);
+
 	TSet<EGLTFJsonExtension> GetCustomExtensionsUsed() const;
 
 private:
diff --git a/Engine/Plugins/Enterprise/GLTFExporter/MigratedFromGit/Source/GLTFExporter/Public/Builders/GLTFLogBuilder.h b/Engine/Plugins/Enterprise/GLTFExporter/MigratedFromGit/Source/GLTFExporter/Public/Builders/GLTFLogBuilder.h
--- a/Engine/Plugins/Enterprise/GLTFExporter/MigratedFromGit/Source/GLTFExporter/Public/Builders/GLTFLogBuilder.h
+++ b/Engine/Plugins/Enterprise/GLTFExporter/MigratedFromGit/Source/GLTFExporter/Public/Builders/GLTFLogBuilder.h
@@ -12,7 +12,7 @@ class GLTFEXPORTER_API FGLTFLogBuilder : public FGLTFBuilder
 {
 public:
 
-	FGLTFLogBuilder(const FString& FileName, const UGLTFExportOptions* ExportOptions);
+	FGLTFLogBuilder(const FString& FileName, const UGLTFExportOptions* ExportOptions = nullptr);
 
 	void LogSuggestion(const FString& Message);
 
diff --git a/Engine/Plugins/Enterprise/GLTFExporter/MigratedFromGit/Source/GLTFExporter/Public/Builders/GLTFTaskBuilder.h b/Engine/Plugins/Enterprise/GLTFExporter/MigratedFromGit/Source/GLTFExporter/Public/Builders/GLTFTaskBuilder.h
--- a/Engine/Plugins/Enterprise/GLTFExporter/MigratedFromGit/Source/GLTFExporter/Public/Builders/GLTFTaskBuilder.h
+++ b/Engine/Plugins/Enterprise/GLTFExporter/MigratedFromGit/Source/GLTFExporter/Public/Builders/GLTFTaskBuilder.h
@@ -9,7 +9,7 @@ class GLTFEXPORTER_API FGLTFTaskBuilder : public FGLTFLogBuilder
 {
 public:
 
-	FGLTFTaskBuilder(const FString& FileName, const UGLTFExportOptions* ExportOptions);
+	FGLTFTaskBuilder(const FString& FileName, const UGLTFExportOptions* ExportOptions = nullptr);
 
 	template <typename TaskType, typename... TaskArgTypes, typename = typename TEnableIf<TIsDerivedFrom<TaskType, FGLTFDelayedTask>::Value>::Type>
 	bool ScheduleSlowTask(TaskArgTypes&&... Args)
@@ -25,7 +25,7 @@ public:
 
 	bool ScheduleSlowTask(TUniquePtr<FGLTFDelayedTask> Task);
 
-	void ProcessSlowTasks(FFeedbackContext* Context = GWarn);
+	void ProcessSlowTasks(FFeedbackContext* Context = nullptr);
 
 private:
 
