[GLTFExporter] refactor: rename FGLTFConverterUtility to FGLTFCoreUtilities and make it public

commit 88c9255f8a6e7d429712bfe8884f6afd4d7115eb
Date: Mon, 14 Mar 2022 00:48:08 +0100

    refactor: rename FGLTFConverterUtility to FGLTFCoreUtilities and make it public
    (cherry picked from commit e86c06731613a4c2918907a53ce3e1e0f4c03efe​)

#skipci

diff --git a/Engine/Plugins/Enterprise/GLTFExporter/MigratedFromGit/Source/GLTFExporter/Private/Builders/GLTFContainerBuilder.cpp b/Engine/Plugins/Enterprise/GLTFExporter/MigratedFromGit/Source/GLTFExporter/Private/Builders/GLTFContainerBuilder.cpp
--- a/Engine/Plugins/Enterprise/GLTFExporter/MigratedFromGit/Source/GLTFExporter/Private/Builders/GLTFContainerBuilder.cpp
+++ b/Engine/Plugins/Enterprise/GLTFExporter/MigratedFromGit/Source/GLTFExporter/Private/Builders/GLTFContainerBuilder.cpp
@@ -27,7 +27,7 @@ void FGLTFContainerBuilder::WriteInternalArchive(FArchive& Archive)
 		const FString ExtensionsString = FString::JoinBy(CustomExtensions, TEXT(", "),
 			[](EGLTFJsonExtension Extension)
 		{
-			return FGLTFJsonUtility::GetValue(Extension);
+			return FGLTFJsonUtilities::GetValue(Extension);
 		});
 
 		LogWarning(FString::Printf(TEXT("Export uses some extensions that may only be supported in Unreal's glTF viewer: %s"), *ExtensionsString));
diff --git a/Engine/Plugins/Enterprise/GLTFExporter/MigratedFromGit/Source/GLTFExporter/Private/Builders/GLTFJsonBuilder.cpp b/Engine/Plugins/Enterprise/GLTFExporter/MigratedFromGit/Source/GLTFExporter/Private/Builders/GLTFJsonBuilder.cpp
--- a/Engine/Plugins/Enterprise/GLTFExporter/MigratedFromGit/Source/GLTFExporter/Private/Builders/GLTFJsonBuilder.cpp
+++ b/Engine/Plugins/Enterprise/GLTFExporter/MigratedFromGit/Source/GLTFExporter/Private/Builders/GLTFJsonBuilder.cpp
@@ -160,6 +160,6 @@ bool FGLTFJsonBuilder::IsCustomExtension(EGLTFJsonExtension Extension)
 {
 	const TCHAR CustomPrefix[] = TEXT("EPIC_");
 
-	const TCHAR* ExtensionString = FGLTFJsonUtility::GetValue(Extension);
+	const TCHAR* ExtensionString = FGLTFJsonUtilities::GetValue(Extension);
 	return FCString::Strncmp(ExtensionString, CustomPrefix, GetNum(CustomPrefix) - 1) == 0;
 }
diff --git a/Engine/Plugins/Enterprise/GLTFExporter/MigratedFromGit/Source/GLTFExporter/Private/Converters/GLTFAccessorConverters.cpp b/Engine/Plugins/Enterprise/GLTFExporter/MigratedFromGit/Source/GLTFExporter/Private/Converters/GLTFAccessorConverters.cpp
--- a/Engine/Plugins/Enterprise/GLTFExporter/MigratedFromGit/Source/GLTFExporter/Private/Converters/GLTFAccessorConverters.cpp
+++ b/Engine/Plugins/Enterprise/GLTFExporter/MigratedFromGit/Source/GLTFExporter/Private/Converters/GLTFAccessorConverters.cpp
@@ -1,7 +1,7 @@
 // Copyright Epic Games, Inc. All Rights Reserved.
 
 #include "Converters/GLTFAccessorConverters.h"
-#include "Converters/GLTFConverterUtility.h"
+#include "Utilities/GLTFCoreUtilities.h"
 #include "Converters/GLTFBufferAdapter.h"
 #include "Builders/GLTFConvertBuilder.h"
 
@@ -47,7 +47,7 @@ FGLTFJsonAccessor* FGLTFPositionBufferConverter::Convert(const FGLTFMeshSection*
 	{
 		const uint32 MappedVertexIndex = IndexMap[VertexIndex];
 		const FVector& Position = *reinterpret_cast<const FVector*>(SourceData + Stride * MappedVertexIndex);
-		Positions[VertexIndex] = FGLTFConverterUtility::ConvertPosition(Position, Builder.ExportOptions->ExportUniformScale);
+		Positions[VertexIndex] = FGLTFCoreUtilities::ConvertPosition(Position, Builder.ExportOptions->ExportUniformScale);
 	}
 
 	FGLTFJsonAccessor* JsonAccessor = Builder.AddAccessor();
@@ -108,7 +108,7 @@ FGLTFJsonAccessor* FGLTFColorBufferConverter::Convert(const FGLTFMeshSection* Me
 	{
 		const uint32 MappedVertexIndex = IndexMap[VertexIndex];
 		const FColor& Color = *reinterpret_cast<const FColor*>(SourceData + Stride * MappedVertexIndex);
-		Colors[VertexIndex] = FGLTFConverterUtility::ConvertColor(Color);
+		Colors[VertexIndex] = FGLTFCoreUtilities::ConvertColor(Color);
 	}
 
 	FGLTFJsonAccessor* JsonAccessor = Builder.AddAccessor();
@@ -196,7 +196,7 @@ FGLTFJsonBufferView* FGLTFNormalBufferConverter::ConvertBufferView(const FGLTFMe
 		const FVector SafeNormal = TangentData[MappedVertexIndex].TangentZ.ToFVector().GetSafeNormal();
 
 		typedef typename TConditional<TIsSame<DestinationType, FGLTFVector3>::Value, FVector, SourceType>::Type IntermediateType;
-		Normals[VertexIndex] = FGLTFConverterUtility::ConvertNormal(IntermediateType(SafeNormal));
+		Normals[VertexIndex] = FGLTFCoreUtilities::ConvertNormal(IntermediateType(SafeNormal));
 	}
 
 	return Builder.AddBufferView(Normals, EGLTFJsonBufferTarget::ArrayBuffer);
@@ -275,7 +275,7 @@ FGLTFJsonBufferView* FGLTFTangentBufferConverter::ConvertBufferView(const FGLTFM
 		const FVector SafeTangent = VertexTangents[MappedVertexIndex].TangentX.ToFVector().GetSafeNormal();
 
 		typedef typename TConditional<TIsSame<DestinationType, FGLTFVector4>::Value, FVector, SourceType>::Type IntermediateType;
-		Tangents[VertexIndex] = FGLTFConverterUtility::ConvertTangent(IntermediateType(SafeTangent));
+		Tangents[VertexIndex] = FGLTFCoreUtilities::ConvertTangent(IntermediateType(SafeTangent));
 	}
 
 	return Builder.AddBufferView(Tangents, EGLTFJsonBufferTarget::ArrayBuffer);
@@ -327,7 +327,7 @@ FGLTFJsonAccessor* FGLTFUVBufferConverter::Convert(const FGLTFMeshSection* MeshS
 	{
 		const uint32 MappedVertexIndex = IndexMap[VertexIndex];
 		const SourceType& UV = UVData[UVCount * MappedVertexIndex + UVIndex];
-		UVs[VertexIndex] = FGLTFConverterUtility::ConvertUV(UV);
+		UVs[VertexIndex] = FGLTFCoreUtilities::ConvertUV(UV);
 	}
 
 	FGLTFJsonAccessor* JsonAccessor = Builder.AddAccessor();
@@ -439,7 +439,7 @@ FGLTFJsonAccessor* FGLTFBoneIndexBufferConverter::Convert(const FGLTFMeshSection
 
 	FGLTFJsonAccessor* JsonAccessor = Builder.AddAccessor();
 	JsonAccessor->BufferView = Builder.AddBufferView(BoneIndices, EGLTFJsonBufferTarget::ArrayBuffer);
-	JsonAccessor->ComponentType = FGLTFConverterUtility::GetComponentType<DestinationType>();
+	JsonAccessor->ComponentType = FGLTFCoreUtilities::GetComponentType<DestinationType>();
 	JsonAccessor->Count = VertexCount;
 	JsonAccessor->Type = EGLTFJsonAccessorType::Vec4;
 
@@ -572,7 +572,7 @@ FGLTFJsonAccessor* FGLTFIndexBufferConverter::Convert(const FGLTFMeshSection* Me
 
 	FGLTFJsonAccessor* JsonAccessor = Builder.AddAccessor();
 	JsonAccessor->BufferView = Builder.AddBufferView(Indices, EGLTFJsonBufferTarget::ElementArrayBuffer, sizeof(IndexType));
-	JsonAccessor->ComponentType = FGLTFConverterUtility::GetComponentType<IndexType>();
+	JsonAccessor->ComponentType = FGLTFCoreUtilities::GetComponentType<IndexType>();
 	JsonAccessor->Count = IndexCount;
 	JsonAccessor->Type = EGLTFJsonAccessorType::Scalar;
 
diff --git a/Engine/Plugins/Enterprise/GLTFExporter/MigratedFromGit/Source/GLTFExporter/Private/Converters/GLTFBackdropConverters.cpp b/Engine/Plugins/Enterprise/GLTFExporter/MigratedFromGit/Source/GLTFExporter/Private/Converters/GLTFBackdropConverters.cpp
--- a/Engine/Plugins/Enterprise/GLTFExporter/MigratedFromGit/Source/GLTFExporter/Private/Converters/GLTFBackdropConverters.cpp
+++ b/Engine/Plugins/Enterprise/GLTFExporter/MigratedFromGit/Source/GLTFExporter/Private/Converters/GLTFBackdropConverters.cpp
@@ -1,7 +1,7 @@
 // Copyright Epic Games, Inc. All Rights Reserved.
 
 #include "Converters/GLTFBackdropConverters.h"
-#include "Converters/GLTFConverterUtility.h"
+#include "Utilities/GLTFCoreUtilities.h"
 #include "Converters/GLTFMaterialUtility.h"
 #include "Converters/GLTFActorUtility.h"
 #include "Builders/GLTFContainerBuilder.h"
@@ -44,7 +44,7 @@ FGLTFJsonBackdrop* FGLTFBackdropConverter::Convert(const AActor* BackdropActor)
 		for (int32 CubeFaceIndex = 0; CubeFaceIndex < CubeFace_MAX; ++CubeFaceIndex)
 		{
 			const ECubeFace CubeFace = static_cast<ECubeFace>(CubeFaceIndex);
-			const EGLTFJsonCubeFace JsonCubeFace = FGLTFConverterUtility::ConvertCubeFace(CubeFace);
+			const EGLTFJsonCubeFace JsonCubeFace = FGLTFCoreUtilities::ConvertCubeFace(CubeFace);
 			JsonBackdrop->Cubemap[static_cast<int32>(JsonCubeFace)] = Builder.AddUniqueTexture(Cubemap, CubeFace);
 		}
 	}
@@ -76,7 +76,7 @@ FGLTFJsonBackdrop* FGLTFBackdropConverter::Convert(const AActor* BackdropActor)
 	FVector ProjectionCenter;
 	if (FGLTFActorUtility::TryGetPropertyValue(BackdropActor, TEXT("ProjectionCenter"), ProjectionCenter))
 	{
-		JsonBackdrop->ProjectionCenter = FGLTFConverterUtility::ConvertPosition(ProjectionCenter, Builder.ExportOptions->ExportUniformScale);
+		JsonBackdrop->ProjectionCenter = FGLTFCoreUtilities::ConvertPosition(ProjectionCenter, Builder.ExportOptions->ExportUniformScale);
 	}
 	else
 	{
diff --git a/Engine/Plugins/Enterprise/GLTFExporter/MigratedFromGit/Source/GLTFExporter/Private/Converters/GLTFCameraConverters.cpp b/Engine/Plugins/Enterprise/GLTFExporter/MigratedFromGit/Source/GLTFExporter/Private/Converters/GLTFCameraConverters.cpp
--- a/Engine/Plugins/Enterprise/GLTFExporter/MigratedFromGit/Source/GLTFExporter/Private/Converters/GLTFCameraConverters.cpp
+++ b/Engine/Plugins/Enterprise/GLTFExporter/MigratedFromGit/Source/GLTFExporter/Private/Converters/GLTFCameraConverters.cpp
@@ -2,7 +2,7 @@
 
 #include "Converters/GLTFCameraConverters.h"
 #include "Builders/GLTFContainerBuilder.h"
-#include "Converters/GLTFConverterUtility.h"
+#include "Utilities/GLTFCoreUtilities.h"
 #include "Converters/GLTFNameUtility.h"
 #include "Actors/GLTFCameraActor.h"
 #include "Camera/CameraComponent.h"
@@ -11,7 +11,7 @@ FGLTFJsonCamera* FGLTFCameraConverter::Convert(const UCameraComponent* CameraCom
 {
 	FGLTFJsonCamera* JsonCamera = Builder.AddCamera();
 	JsonCamera->Name = FGLTFNameUtility::GetName(CameraComponent);
-	JsonCamera->Type = FGLTFConverterUtility::ConvertCameraType(CameraComponent->ProjectionMode);
+	JsonCamera->Type = FGLTFCoreUtilities::ConvertCameraType(CameraComponent->ProjectionMode);
 
 	FMinimalViewInfo DesiredView;
 	const_cast<UCameraComponent*>(CameraComponent)->GetCameraView(0, DesiredView);
@@ -24,10 +24,10 @@ FGLTFJsonCamera* FGLTFCameraConverter::Convert(const UCameraComponent* CameraCom
 			{
 				Builder.LogWarning(FString::Printf(TEXT("Aspect ratio for orthographic camera component %s (in actor %s) will be constrainted in glTF"), *CameraComponent->GetName(), *CameraComponent->GetOwner()->GetName()));
 			}
-			JsonCamera->Orthographic.XMag = FGLTFConverterUtility::ConvertLength(DesiredView.OrthoWidth, ExportScale);
-			JsonCamera->Orthographic.YMag = FGLTFConverterUtility::ConvertLength(DesiredView.OrthoWidth / DesiredView.AspectRatio, ExportScale); // TODO: is this correct?
-			JsonCamera->Orthographic.ZFar = FGLTFConverterUtility::ConvertLength(DesiredView.OrthoFarClipPlane, ExportScale);
-			JsonCamera->Orthographic.ZNear = FGLTFConverterUtility::ConvertLength(DesiredView.OrthoNearClipPlane, ExportScale);
+			JsonCamera->Orthographic.XMag = FGLTFCoreUtilities::ConvertLength(DesiredView.OrthoWidth, ExportScale);
+			JsonCamera->Orthographic.YMag = FGLTFCoreUtilities::ConvertLength(DesiredView.OrthoWidth / DesiredView.AspectRatio, ExportScale); // TODO: is this correct?
+			JsonCamera->Orthographic.ZFar = FGLTFCoreUtilities::ConvertLength(DesiredView.OrthoFarClipPlane, ExportScale);
+			JsonCamera->Orthographic.ZNear = FGLTFCoreUtilities::ConvertLength(DesiredView.OrthoNearClipPlane, ExportScale);
 			break;
 
 		case EGLTFJsonCameraType::Perspective:
@@ -35,10 +35,10 @@ FGLTFJsonCamera* FGLTFCameraConverter::Convert(const UCameraComponent* CameraCom
 			{
 				JsonCamera->Perspective.AspectRatio = DesiredView.AspectRatio;
 			}
-			JsonCamera->Perspective.YFov = FGLTFConverterUtility::ConvertFieldOfView(DesiredView.FOV, DesiredView.AspectRatio);
+			JsonCamera->Perspective.YFov = FGLTFCoreUtilities::ConvertFieldOfView(DesiredView.FOV, DesiredView.AspectRatio);
 			// NOTE: even thought ZFar is optional, if we don't set it, then most gltf viewers won't handle it well.
-			JsonCamera->Perspective.ZFar = FGLTFConverterUtility::ConvertLength(WORLD_MAX, ExportScale); // TODO: Unreal doesn't have max draw distance per view?
-			JsonCamera->Perspective.ZNear = FGLTFConverterUtility::ConvertLength(GNearClippingPlane, ExportScale);
+			JsonCamera->Perspective.ZFar = FGLTFCoreUtilities::ConvertLength(WORLD_MAX, ExportScale); // TODO: Unreal doesn't have max draw distance per view?
+			JsonCamera->Perspective.ZNear = FGLTFCoreUtilities::ConvertLength(GNearClippingPlane, ExportScale);
 			break;
 
 		case EGLTFJsonCameraType::None:
@@ -58,10 +58,10 @@ FGLTFJsonCamera* FGLTFCameraConverter::Convert(const UCameraComponent* CameraCom
 		if (Builder.ExportOptions->bExportCameraControls)
 		{
 			FGLTFJsonCameraControl CameraControl;
-			CameraControl.Mode = FGLTFConverterUtility::ConvertCameraControlMode(CameraActor->Mode);
+			CameraControl.Mode = FGLTFCoreUtilities::ConvertCameraControlMode(CameraActor->Mode);
 			CameraControl.Target = Builder.AddUniqueNode(CameraActor->Target);
-			CameraControl.MaxDistance = FGLTFConverterUtility::ConvertLength(CameraActor->DistanceMax, ExportScale);
-			CameraControl.MinDistance = FGLTFConverterUtility::ConvertLength(CameraActor->DistanceMin, ExportScale);
+			CameraControl.MaxDistance = FGLTFCoreUtilities::ConvertLength(CameraActor->DistanceMax, ExportScale);
+			CameraControl.MinDistance = FGLTFCoreUtilities::ConvertLength(CameraActor->DistanceMin, ExportScale);
 			CameraControl.MaxPitch = CameraActor->PitchAngleMax;
 			CameraControl.MinPitch = CameraActor->PitchAngleMin;
 
diff --git a/Engine/Plugins/Enterprise/GLTFExporter/MigratedFromGit/Source/GLTFExporter/Private/Converters/GLTFConverterUtility.cpp b/Engine/Plugins/Enterprise/GLTFExporter/MigratedFromGit/Source/GLTFExporter/Private/Converters/GLTFConverterUtility.cpp
deleted file mode 100644
diff --git a/Engine/Plugins/Enterprise/GLTFExporter/MigratedFromGit/Source/GLTFExporter/Private/Converters/GLTFConverterUtility.h b/Engine/Plugins/Enterprise/GLTFExporter/MigratedFromGit/Source/GLTFExporter/Private/Converters/GLTFConverterUtility.h
deleted file mode 100644
diff --git a/Engine/Plugins/Enterprise/GLTFExporter/MigratedFromGit/Source/GLTFExporter/Private/Converters/GLTFLightConverters.cpp b/Engine/Plugins/Enterprise/GLTFExporter/MigratedFromGit/Source/GLTFExporter/Private/Converters/GLTFLightConverters.cpp
--- a/Engine/Plugins/Enterprise/GLTFExporter/MigratedFromGit/Source/GLTFExporter/Private/Converters/GLTFLightConverters.cpp
+++ b/Engine/Plugins/Enterprise/GLTFExporter/MigratedFromGit/Source/GLTFExporter/Private/Converters/GLTFLightConverters.cpp
@@ -2,7 +2,7 @@
 
 #include "Converters/GLTFLightConverters.h"
 #include "Builders/GLTFContainerBuilder.h"
-#include "Converters/GLTFConverterUtility.h"
+#include "Utilities/GLTFCoreUtilities.h"
 #include "Converters/GLTFNameUtility.h"
 #include "Components/LightComponent.h"
 #include "Components/PointLightComponent.h"
@@ -13,7 +13,7 @@ FGLTFJsonLight* FGLTFLightConverter::Convert(const ULightComponent* LightCompone
 	FGLTFJsonLight* Light = Builder.AddLight();
 
 	Light->Name = FGLTFNameUtility::GetName(LightComponent);
-	Light->Type = FGLTFConverterUtility::ConvertLightType(LightComponent->GetLightType());
+	Light->Type = FGLTFCoreUtilities::ConvertLightType(LightComponent->GetLightType());
 
 	if (Light->Type == EGLTFJsonLightType::None)
 	{
@@ -24,17 +24,17 @@ FGLTFJsonLight* FGLTFLightConverter::Convert(const ULightComponent* LightCompone
 	Light->Intensity = LightComponent->Intensity;
 
 	const FLinearColor TemperatureColor = LightComponent->bUseTemperature ? FLinearColor::MakeFromColorTemperature(LightComponent->Temperature) : FLinearColor::White;
-	Light->Color = FGLTFConverterUtility::ConvertColor3(TemperatureColor * LightComponent->GetLightColor(), Builder.ExportOptions->bStrictCompliance);
+	Light->Color = FGLTFCoreUtilities::ConvertColor3(TemperatureColor * LightComponent->GetLightColor(), Builder.ExportOptions->bStrictCompliance);
 
 	if (const UPointLightComponent* PointLightComponent = Cast<UPointLightComponent>(LightComponent))
 	{
-		Light->Range = FGLTFConverterUtility::ConvertLength(PointLightComponent->AttenuationRadius, Builder.ExportOptions->ExportUniformScale);
+		Light->Range = FGLTFCoreUtilities::ConvertLength(PointLightComponent->AttenuationRadius, Builder.ExportOptions->ExportUniformScale);
 	}
 
 	if (const USpotLightComponent* SpotLightComponent = Cast<USpotLightComponent>(LightComponent))
 	{
-		Light->Spot.InnerConeAngle = FGLTFConverterUtility::ConvertLightAngle(SpotLightComponent->InnerConeAngle);
-		Light->Spot.OuterConeAngle = FGLTFConverterUtility::ConvertLightAngle(SpotLightComponent->OuterConeAngle);
+		Light->Spot.InnerConeAngle = FGLTFCoreUtilities::ConvertLightAngle(SpotLightComponent->InnerConeAngle);
+		Light->Spot.OuterConeAngle = FGLTFCoreUtilities::ConvertLightAngle(SpotLightComponent->OuterConeAngle);
 
 		if (Builder.ExportOptions->bStrictCompliance)
 		{
diff --git a/Engine/Plugins/Enterprise/GLTFExporter/MigratedFromGit/Source/GLTFExporter/Private/Converters/GLTFNodeConverters.cpp b/Engine/Plugins/Enterprise/GLTFExporter/MigratedFromGit/Source/GLTFExporter/Private/Converters/GLTFNodeConverters.cpp
--- a/Engine/Plugins/Enterprise/GLTFExporter/MigratedFromGit/Source/GLTFExporter/Private/Converters/GLTFNodeConverters.cpp
+++ b/Engine/Plugins/Enterprise/GLTFExporter/MigratedFromGit/Source/GLTFExporter/Private/Converters/GLTFNodeConverters.cpp
@@ -2,7 +2,7 @@
 
 #include "Converters/GLTFNodeConverters.h"
 #include "Builders/GLTFContainerBuilder.h"
-#include "Converters/GLTFConverterUtility.h"
+#include "Utilities/GLTFCoreUtilities.h"
 #include "Converters/GLTFActorUtility.h"
 #include "Converters/GLTFNameUtility.h"
 #include "Actors/GLTFHotspotActor.h"
@@ -122,9 +122,9 @@ FGLTFJsonNode* FGLTFComponentConverter::Convert(const USceneComponent* SceneComp
 
 	FGLTFJsonNode* Node = Builder.AddNode();
 	Node->Name = FGLTFNameUtility::GetName(SceneComponent);
-	Node->Translation = FGLTFConverterUtility::ConvertPosition(RelativeTransform.GetTranslation(), Builder.ExportOptions->ExportUniformScale);
-	Node->Rotation = FGLTFConverterUtility::ConvertRotation(RelativeTransform.GetRotation());
-	Node->Scale = FGLTFConverterUtility::ConvertScale(RelativeTransform.GetScale3D());
+	Node->Translation = FGLTFCoreUtilities::ConvertPosition(RelativeTransform.GetTranslation(), Builder.ExportOptions->ExportUniformScale);
+	Node->Rotation = FGLTFCoreUtilities::ConvertRotation(RelativeTransform.GetRotation());
+	Node->Scale = FGLTFCoreUtilities::ConvertScale(RelativeTransform.GetScale3D());
 
 	if (ParentNode != nullptr)
 	{
@@ -168,7 +168,7 @@ FGLTFJsonNode* FGLTFComponentConverter::Convert(const USceneComponent* SceneComp
 				// TODO: conversion of camera direction should be done in separate converter
 				FGLTFJsonNode* CameraNode = Builder.AddNode();
 				CameraNode->Name = FGLTFNameUtility::GetName(CameraComponent);
-				CameraNode->Rotation = FGLTFConverterUtility::ConvertCameraDirection();
+				CameraNode->Rotation = FGLTFCoreUtilities::GetLocalCameraRotation();
 				CameraNode->Camera = Builder.AddUniqueCamera(CameraComponent);
 				Node->Children.Add(CameraNode);
 			}
@@ -180,7 +180,7 @@ FGLTFJsonNode* FGLTFComponentConverter::Convert(const USceneComponent* SceneComp
 				// TODO: conversion of light direction should be done in separate converter
 				FGLTFJsonNode* LightNode = Builder.AddNode();
 				LightNode->Name = FGLTFNameUtility::GetName(LightComponent);
-				LightNode->Rotation = FGLTFConverterUtility::ConvertLightDirection();
+				LightNode->Rotation = FGLTFCoreUtilities::GetLocalLightRotation();
 				LightNode->Light = Builder.AddUniqueLight(LightComponent);
 				Node->Children.Add(LightNode);
 			}
@@ -233,9 +233,9 @@ FGLTFJsonNode* FGLTFStaticSocketConverter::Convert(FGLTFJsonNode* RootNode, cons
 	Node->Name = SocketName.ToString();
 
 	// TODO: add warning check for non-uniform scaling
-	Node->Translation = FGLTFConverterUtility::ConvertPosition(Socket->RelativeLocation, Builder.ExportOptions->ExportUniformScale);
-	Node->Rotation = FGLTFConverterUtility::ConvertRotation(Socket->RelativeRotation.Quaternion());
-	Node->Scale = FGLTFConverterUtility::ConvertScale(Socket->RelativeScale);
+	Node->Translation = FGLTFCoreUtilities::ConvertPosition(Socket->RelativeLocation, Builder.ExportOptions->ExportUniformScale);
+	Node->Rotation = FGLTFCoreUtilities::ConvertRotation(Socket->RelativeRotation.Quaternion());
+	Node->Scale = FGLTFCoreUtilities::ConvertScale(Socket->RelativeScale);
 
 	RootNode->Children.Add(Node);
 	return Node;
@@ -256,9 +256,9 @@ FGLTFJsonNode* FGLTFSkeletalSocketConverter::Convert(FGLTFJsonNode* RootNode, co
 		Node->Name = SocketName.ToString();
 
 		// TODO: add warning check for non-uniform scaling
-		Node->Translation = FGLTFConverterUtility::ConvertPosition(Socket->RelativeLocation, Builder.ExportOptions->ExportUniformScale);
-		Node->Rotation = FGLTFConverterUtility::ConvertRotation(Socket->RelativeRotation.Quaternion());
-		Node->Scale = FGLTFConverterUtility::ConvertScale(Socket->RelativeScale);
+		Node->Translation = FGLTFCoreUtilities::ConvertPosition(Socket->RelativeLocation, Builder.ExportOptions->ExportUniformScale);
+		Node->Rotation = FGLTFCoreUtilities::ConvertRotation(Socket->RelativeRotation.Quaternion());
+		Node->Scale = FGLTFCoreUtilities::ConvertScale(Socket->RelativeScale);
 
 		const int32 ParentBone = RefSkeleton.FindBoneIndex(Socket->BoneName);
 		FGLTFJsonNode* ParentNode = ParentBone != INDEX_NONE ? Builder.AddUniqueNode(RootNode, SkeletalMesh, ParentBone) : RootNode;
@@ -304,9 +304,9 @@ FGLTFJsonNode* FGLTFSkeletalBoneConverter::Convert(FGLTFJsonNode* RootNode, cons
 	{
 		// TODO: add warning check for non-uniform scaling
 		const FTransform& BonePose = BonePoses[BoneIndex];
-		Node->Translation = FGLTFConverterUtility::ConvertPosition(BonePose.GetTranslation(), Builder.ExportOptions->ExportUniformScale);
-		Node->Rotation = FGLTFConverterUtility::ConvertRotation(BonePose.GetRotation());
-		Node->Scale = FGLTFConverterUtility::ConvertScale(BonePose.GetScale3D());
+		Node->Translation = FGLTFCoreUtilities::ConvertPosition(BonePose.GetTranslation(), Builder.ExportOptions->ExportUniformScale);
+		Node->Rotation = FGLTFCoreUtilities::ConvertRotation(BonePose.GetRotation());
+		Node->Scale = FGLTFCoreUtilities::ConvertScale(BonePose.GetScale3D());
 	}
 	else
 	{
diff --git a/Engine/Plugins/Enterprise/GLTFExporter/MigratedFromGit/Source/GLTFExporter/Private/Converters/GLTFSamplerConverters.cpp b/Engine/Plugins/Enterprise/GLTFExporter/MigratedFromGit/Source/GLTFExporter/Private/Converters/GLTFSamplerConverters.cpp
--- a/Engine/Plugins/Enterprise/GLTFExporter/MigratedFromGit/Source/GLTFExporter/Private/Converters/GLTFSamplerConverters.cpp
+++ b/Engine/Plugins/Enterprise/GLTFExporter/MigratedFromGit/Source/GLTFExporter/Private/Converters/GLTFSamplerConverters.cpp
@@ -2,7 +2,7 @@
 
 #include "Converters/GLTFSamplerConverters.h"
 #include "Builders/GLTFContainerBuilder.h"
-#include "Converters/GLTFConverterUtility.h"
+#include "Utilities/GLTFCoreUtilities.h"
 #include "Converters/GLTFTextureUtility.h"
 
 FGLTFJsonSampler* FGLTFSamplerConverter::Convert(const UTexture* Texture)
@@ -20,8 +20,8 @@ FGLTFJsonSampler* FGLTFSamplerConverter::Convert(const UTexture* Texture)
 	}
 	else
 	{
-		JsonSampler->MinFilter = FGLTFConverterUtility::ConvertMinFilter(Texture->Filter, Texture->LODGroup);
-		JsonSampler->MagFilter = FGLTFConverterUtility::ConvertMagFilter(Texture->Filter, Texture->LODGroup);
+		JsonSampler->MinFilter = FGLTFCoreUtilities::ConvertMinFilter(Texture->Filter, Texture->LODGroup);
+		JsonSampler->MagFilter = FGLTFCoreUtilities::ConvertMagFilter(Texture->Filter, Texture->LODGroup);
 	}
 
 	if (FGLTFTextureUtility::IsCubemap(Texture))
@@ -34,8 +34,8 @@ FGLTFJsonSampler* FGLTFSamplerConverter::Convert(const UTexture* Texture)
 		const TTuple<TextureAddress, TextureAddress> AddressXY = FGLTFTextureUtility::GetAddressXY(Texture);
 		// TODO: report error if AddressX == TA_MAX or AddressY == TA_MAX
 
-		JsonSampler->WrapS = FGLTFConverterUtility::ConvertWrap(AddressXY.Get<0>());
-		JsonSampler->WrapT = FGLTFConverterUtility::ConvertWrap(AddressXY.Get<1>());
+		JsonSampler->WrapS = FGLTFCoreUtilities::ConvertWrap(AddressXY.Get<0>());
+		JsonSampler->WrapT = FGLTFCoreUtilities::ConvertWrap(AddressXY.Get<1>());
 	}
 
 	return JsonSampler;
diff --git a/Engine/Plugins/Enterprise/GLTFExporter/MigratedFromGit/Source/GLTFExporter/Private/Converters/GLTFSkinConverters.cpp b/Engine/Plugins/Enterprise/GLTFExporter/MigratedFromGit/Source/GLTFExporter/Private/Converters/GLTFSkinConverters.cpp
--- a/Engine/Plugins/Enterprise/GLTFExporter/MigratedFromGit/Source/GLTFExporter/Private/Converters/GLTFSkinConverters.cpp
+++ b/Engine/Plugins/Enterprise/GLTFExporter/MigratedFromGit/Source/GLTFExporter/Private/Converters/GLTFSkinConverters.cpp
@@ -1,7 +1,7 @@
 // Copyright Epic Games, Inc. All Rights Reserved.
 
 #include "Converters/GLTFSkinConverters.h"
-#include "Converters/GLTFConverterUtility.h"
+#include "Utilities/GLTFCoreUtilities.h"
 #include "Converters/GLTFBoneUtility.h"
 #include "Builders/GLTFConvertBuilder.h"
 
@@ -39,7 +39,7 @@ FGLTFJsonSkin* FGLTFSkinConverter::Convert(FGLTFJsonNode* RootNode, const USkele
 	for (int32 BoneIndex = 0; BoneIndex < BoneCount; ++BoneIndex)
 	{
 		const FTransform InverseBindTransform = FGLTFBoneUtility::GetBindTransform(RefSkeleton, BoneIndex).Inverse();
-		InverseBindMatrices[BoneIndex] = FGLTFConverterUtility::ConvertTransform(InverseBindTransform, Builder.ExportOptions->ExportUniformScale);
+		InverseBindMatrices[BoneIndex] = FGLTFCoreUtilities::ConvertTransform(InverseBindTransform, Builder.ExportOptions->ExportUniformScale);
 	}
 
 	FGLTFJsonAccessor* JsonBindMatricesAccessor = Builder.AddAccessor();
diff --git a/Engine/Plugins/Enterprise/GLTFExporter/MigratedFromGit/Source/GLTFExporter/Private/Converters/GLTFSkySphereConverters.cpp b/Engine/Plugins/Enterprise/GLTFExporter/MigratedFromGit/Source/GLTFExporter/Private/Converters/GLTFSkySphereConverters.cpp
--- a/Engine/Plugins/Enterprise/GLTFExporter/MigratedFromGit/Source/GLTFExporter/Private/Converters/GLTFSkySphereConverters.cpp
+++ b/Engine/Plugins/Enterprise/GLTFExporter/MigratedFromGit/Source/GLTFExporter/Private/Converters/GLTFSkySphereConverters.cpp
@@ -1,7 +1,7 @@
 // Copyright Epic Games, Inc. All Rights Reserved.
 
 #include "Converters/GLTFSkySphereConverters.h"
-#include "Converters/GLTFConverterUtility.h"
+#include "Utilities/GLTFCoreUtilities.h"
 #include "Converters/GLTFMaterialUtility.h"
 #include "Converters/GLTFActorUtility.h"
 #include "Converters/GLTFCurveUtility.h"
@@ -32,7 +32,7 @@ FGLTFJsonSkySphere* FGLTFSkySphereConverter::Convert(const AActor* SkySphereActo
 		const FTransform ParentTransform = ParentComponent->GetSocketTransform(SocketName);
 		const FTransform RelativeTransform = Transform.GetRelativeTransform(ParentTransform);
 
-		JsonSkySphere->Scale = FGLTFConverterUtility::ConvertScale(RelativeTransform.GetScale3D());
+		JsonSkySphere->Scale = FGLTFCoreUtilities::ConvertScale(RelativeTransform.GetScale3D());
 		JsonSkySphere->SkySphereMesh = Builder.AddUniqueMesh(StaticMeshComponent, { FGLTFMaterialUtility::GetDefaultMaterial() });
 	}
 	else
@@ -111,7 +111,7 @@ void FGLTFSkySphereConverter::ConvertColorProperty(const AActor* Actor, const TC
 	FLinearColor LinearColor;
 	if (FGLTFActorUtility::TryGetPropertyValue(Actor, PropertyName, LinearColor))
 	{
-		OutValue = FGLTFConverterUtility::ConvertColor(LinearColor, Builder.ExportOptions->bStrictCompliance);
+		OutValue = FGLTFCoreUtilities::ConvertColor(LinearColor, Builder.ExportOptions->bStrictCompliance);
 	}
 	else
 	{
diff --git a/Engine/Plugins/Enterprise/GLTFExporter/MigratedFromGit/Source/GLTFExporter/Private/Json/GLTFJsonUtility.cpp b/Engine/Plugins/Enterprise/GLTFExporter/MigratedFromGit/Source/GLTFExporter/Private/Json/GLTFJsonUtility.cpp
deleted file mode 100644
diff --git a/Engine/Plugins/Enterprise/GLTFExporter/MigratedFromGit/Source/GLTFExporter/Private/Tasks/GLTFDelayedAnimationTasks.cpp b/Engine/Plugins/Enterprise/GLTFExporter/MigratedFromGit/Source/GLTFExporter/Private/Tasks/GLTFDelayedAnimationTasks.cpp
--- a/Engine/Plugins/Enterprise/GLTFExporter/MigratedFromGit/Source/GLTFExporter/Private/Tasks/GLTFDelayedAnimationTasks.cpp
+++ b/Engine/Plugins/Enterprise/GLTFExporter/MigratedFromGit/Source/GLTFExporter/Private/Tasks/GLTFDelayedAnimationTasks.cpp
@@ -2,7 +2,7 @@
 
 #include "Tasks/GLTFDelayedAnimationTasks.h"
 #include "Builders/GLTFContainerBuilder.h"
-#include "Converters/GLTFConverterUtility.h"
+#include "Utilities/GLTFCoreUtilities.h"
 #include "Converters/GLTFBoneUtility.h"
 #include "LevelSequence.h"
 #include "LevelSequencePlayer.h"
@@ -36,7 +36,7 @@ void FGLTFDelayedAnimSequenceTask::Process()
 	JsonInputAccessor->Min[0] = 0;
 	JsonInputAccessor->Max[0] = Timestamps[FrameCount - 1];
 
-	EGLTFJsonInterpolation Interpolation = FGLTFConverterUtility::ConvertInterpolation(AnimSequence->Interpolation);
+	EGLTFJsonInterpolation Interpolation = FGLTFCoreUtilities::ConvertInterpolation(AnimSequence->Interpolation);
 	if (Interpolation == EGLTFJsonInterpolation::None)
 	{
 		Interpolation = EGLTFJsonInterpolation::Linear;
@@ -62,7 +62,7 @@ void FGLTFDelayedAnimSequenceTask::Process()
 			for (int32 Frame = 0; Frame < FrameCount; ++Frame)
 			{
 				const FVector KeyPosition = FrameTransforms[Frame][BoneIndex].GetTranslation();
-				Translations[Frame] = FGLTFConverterUtility::ConvertPosition(KeyPosition, Builder.ExportOptions->ExportUniformScale);
+				Translations[Frame] = FGLTFCoreUtilities::ConvertPosition(KeyPosition, Builder.ExportOptions->ExportUniformScale);
 			}
 
 			FGLTFJsonAccessor* JsonOutputAccessor = Builder.AddAccessor();
@@ -90,7 +90,7 @@ void FGLTFDelayedAnimSequenceTask::Process()
 			for (int32 Frame = 0; Frame < FrameCount; ++Frame)
 			{
 				const FQuat KeyRotation = FrameTransforms[Frame][BoneIndex].GetRotation();
-				Rotations[Frame] = FGLTFConverterUtility::ConvertRotation(KeyRotation);
+				Rotations[Frame] = FGLTFCoreUtilities::ConvertRotation(KeyRotation);
 			}
 
 			FGLTFJsonAccessor* JsonOutputAccessor = Builder.AddAccessor();
@@ -118,7 +118,7 @@ void FGLTFDelayedAnimSequenceTask::Process()
 			for (int32 Frame = 0; Frame < FrameCount; ++Frame)
 			{
 				const FVector KeyScale = FrameTransforms[Frame][BoneIndex].GetScale3D();
-				Scales[Frame] = FGLTFConverterUtility::ConvertScale(KeyScale);
+				Scales[Frame] = FGLTFCoreUtilities::ConvertScale(KeyScale);
 			}
 
 			FGLTFJsonAccessor* JsonOutputAccessor = Builder.AddAccessor();
@@ -256,7 +256,7 @@ void FGLTFDelayedLevelSequenceTask::Process()
 							Channels[1]->Evaluate(FrameTime, Translation.Y);
 							Channels[2]->Evaluate(FrameTime, Translation.Z);
 
-							Translations[Frame] = FGLTFConverterUtility::ConvertPosition(Translation, Builder.ExportOptions->ExportUniformScale);
+							Translations[Frame] = FGLTFCoreUtilities::ConvertPosition(Translation, Builder.ExportOptions->ExportUniformScale);
 						}
 
 						FGLTFJsonAccessor* JsonOutputAccessor = Builder.AddAccessor();
@@ -291,7 +291,7 @@ void FGLTFDelayedLevelSequenceTask::Process()
 							Channels[4]->Evaluate(FrameTime, Rotator.Pitch);
 							Channels[5]->Evaluate(FrameTime, Rotator.Yaw);
 
-							Rotations[Frame] = FGLTFConverterUtility::ConvertRotation(Rotator.Quaternion());
+							Rotations[Frame] = FGLTFCoreUtilities::ConvertRotation(Rotator.Quaternion());
 						}
 
 						JsonInputAccessor->Count = Rotations.Num();
@@ -329,7 +329,7 @@ void FGLTFDelayedLevelSequenceTask::Process()
 							Channels[7]->Evaluate(FrameTime, Scale.Y);
 							Channels[8]->Evaluate(FrameTime, Scale.Z);
 
-							Scales[Frame] = FGLTFConverterUtility::ConvertScale(Scale);
+							Scales[Frame] = FGLTFCoreUtilities::ConvertScale(Scale);
 						}
 
 						FGLTFJsonAccessor* JsonOutputAccessor = Builder.AddAccessor();
diff --git a/Engine/Plugins/Enterprise/GLTFExporter/MigratedFromGit/Source/GLTFExporter/Private/Tasks/GLTFDelayedMaterialTasks.cpp b/Engine/Plugins/Enterprise/GLTFExporter/MigratedFromGit/Source/GLTFExporter/Private/Tasks/GLTFDelayedMaterialTasks.cpp
--- a/Engine/Plugins/Enterprise/GLTFExporter/MigratedFromGit/Source/GLTFExporter/Private/Tasks/GLTFDelayedMaterialTasks.cpp
+++ b/Engine/Plugins/Enterprise/GLTFExporter/MigratedFromGit/Source/GLTFExporter/Private/Tasks/GLTFDelayedMaterialTasks.cpp
@@ -1,7 +1,7 @@
 // Copyright Epic Games, Inc. All Rights Reserved.
 
 #include "Tasks/GLTFDelayedMaterialTasks.h"
-#include "Converters/GLTFConverterUtility.h"
+#include "Utilities/GLTFCoreUtilities.h"
 #include "Converters/GLTFNameUtility.h"
 #include "Converters/GLTFMaterialUtility.h"
 #include "Builders/GLTFContainerBuilder.h"
@@ -264,7 +264,7 @@ void FGLTFDelayedMaterialTask::GetProxyParameter(const TGLTFProxyMaterialParamet
 	FLinearColor Value;
 	if (ParameterInfo.Get(Material, Value, true))
 	{
-		OutValue = FGLTFConverterUtility::ConvertColor3(Value, false);
+		OutValue = FGLTFCoreUtilities::ConvertColor3(Value, false);
 	}
 }
 
@@ -273,7 +273,7 @@ void FGLTFDelayedMaterialTask::GetProxyParameter(const TGLTFProxyMaterialParamet
 	FLinearColor Value;
 	if (ParameterInfo.Get(Material, Value, true))
 	{
-		OutValue = FGLTFConverterUtility::ConvertColor(Value, false);
+		OutValue = FGLTFCoreUtilities::ConvertColor(Value, false);
 	}
 }
 
@@ -407,7 +407,7 @@ void FGLTFDelayedMaterialTask::ConvertShadingModel(EGLTFJsonShadingModel& OutSha
 		ShadingModel = MSM_DefaultLit;
 	}
 
-	OutShadingModel = FGLTFConverterUtility::ConvertShadingModel(ShadingModel);
+	OutShadingModel = FGLTFCoreUtilities::ConvertShadingModel(ShadingModel);
 	if (OutShadingModel == EGLTFJsonShadingModel::None)
 	{
 		OutShadingModel = EGLTFJsonShadingModel::Default;
@@ -452,7 +452,7 @@ void FGLTFDelayedMaterialTask::ConvertAlphaMode(EGLTFJsonAlphaMode& OutAlphaMode
 		FGLTFProxyMaterialUtilities::GetBlendMode(MaterialInstance, BlendMode); // TODO: remove ugly hack for dynamic instances
 	}
 
-	OutAlphaMode = FGLTFConverterUtility::ConvertAlphaMode(BlendMode);
+	OutAlphaMode = FGLTFCoreUtilities::ConvertAlphaMode(BlendMode);
 	if (OutAlphaMode == EGLTFJsonAlphaMode::None)
 	{
 		OutAlphaMode = EGLTFJsonAlphaMode::Blend;
@@ -468,7 +468,7 @@ void FGLTFDelayedMaterialTask::ConvertAlphaMode(EGLTFJsonAlphaMode& OutAlphaMode
 
 	if (OutAlphaMode == EGLTFJsonAlphaMode::Blend)
 	{
-		OutBlendMode = FGLTFConverterUtility::ConvertBlendMode(BlendMode);
+		OutBlendMode = FGLTFCoreUtilities::ConvertBlendMode(BlendMode);
 		if (OutBlendMode != EGLTFJsonBlendMode::None && !Builder.ExportOptions->bExportExtraBlendModes)
 		{
 			OutBlendMode = EGLTFJsonBlendMode::None;
@@ -535,12 +535,12 @@ bool FGLTFDelayedMaterialTask::TryGetBaseColorAndOpacity(FGLTFJsonPBRMetallicRou
 	const FIntPoint TextureSize = Builder.GetBakeSizeForMaterialProperty(Material, GetPropertyGroup(BaseColorProperty));
 
 	const TextureAddress TextureAddress = Builder.GetBakeTilingForMaterialProperty(Material, GetPropertyGroup(BaseColorProperty));
-	const EGLTFJsonTextureWrap TextureWrapS = FGLTFConverterUtility::ConvertWrap(TextureAddress);
-	const EGLTFJsonTextureWrap TextureWrapT = FGLTFConverterUtility::ConvertWrap(TextureAddress);
+	const EGLTFJsonTextureWrap TextureWrapS = FGLTFCoreUtilities::ConvertWrap(TextureAddress);
+	const EGLTFJsonTextureWrap TextureWrapT = FGLTFCoreUtilities::ConvertWrap(TextureAddress);
 
 	const TextureFilter TextureFilter = Builder.GetBakeFilterForMaterialProperty(Material, GetPropertyGroup(BaseColorProperty));
-	const EGLTFJsonTextureFilter TextureMinFilter = FGLTFConverterUtility::ConvertMinFilter(TextureFilter);
-	const EGLTFJsonTextureFilter TextureMagFilter = FGLTFConverterUtility::ConvertMagFilter(TextureFilter);
+	const EGLTFJsonTextureFilter TextureMinFilter = FGLTFCoreUtilities::ConvertMinFilter(TextureFilter);
+	const EGLTFJsonTextureFilter TextureMagFilter = FGLTFCoreUtilities::ConvertMagFilter(TextureFilter);
 
 	const FGLTFPropertyBakeOutput BaseColorBakeOutput = BakeMaterialProperty(BaseColorProperty, BaseColorTexCoord, TextureSize, false);
 	const FGLTFPropertyBakeOutput OpacityBakeOutput = BakeMaterialProperty(OpacityProperty, OpacityTexCoord, TextureSize, false);
@@ -552,7 +552,7 @@ bool FGLTFDelayedMaterialTask::TryGetBaseColorAndOpacity(FGLTFJsonPBRMetallicRou
 		FLinearColor BaseColorFactor(BaseColorBakeOutput.ConstantValue * BaseColorScale);
 		BaseColorFactor.A = OpacityBakeOutput.ConstantValue.A;
 
-		OutPBRParams.BaseColorFactor = FGLTFConverterUtility::ConvertColor(BaseColorFactor, Builder.ExportOptions->bStrictCompliance);
+		OutPBRParams.BaseColorFactor = FGLTFCoreUtilities::ConvertColor(BaseColorFactor, Builder.ExportOptions->bStrictCompliance);
 		return true;
 	}
 
@@ -604,7 +604,7 @@ bool FGLTFDelayedMaterialTask::TryGetBaseColorAndOpacity(FGLTFJsonPBRMetallicRou
 
 	OutPBRParams.BaseColorTexture.TexCoord = TexCoord;
 	OutPBRParams.BaseColorTexture.Index = Texture;
-	OutPBRParams.BaseColorFactor = FGLTFConverterUtility::ConvertColor({ BaseColorScale, BaseColorScale, BaseColorScale }, Builder.ExportOptions->bStrictCompliance);
+	OutPBRParams.BaseColorFactor = FGLTFCoreUtilities::ConvertColor({ BaseColorScale, BaseColorScale, BaseColorScale }, Builder.ExportOptions->bStrictCompliance);
 
 	return true;
 }
@@ -662,12 +662,12 @@ bool FGLTFDelayedMaterialTask::TryGetMetallicAndRoughness(FGLTFJsonPBRMetallicRo
 	const FIntPoint TextureSize = Builder.GetBakeSizeForMaterialProperty(Material, GetPropertyGroup(MetallicProperty));
 
 	const TextureAddress TextureAddress = Builder.GetBakeTilingForMaterialProperty(Material, GetPropertyGroup(MetallicProperty));
-	const EGLTFJsonTextureWrap TextureWrapS = FGLTFConverterUtility::ConvertWrap(TextureAddress);
-	const EGLTFJsonTextureWrap TextureWrapT = FGLTFConverterUtility::ConvertWrap(TextureAddress);
+	const EGLTFJsonTextureWrap TextureWrapS = FGLTFCoreUtilities::ConvertWrap(TextureAddress);
+	const EGLTFJsonTextureWrap TextureWrapT = FGLTFCoreUtilities::ConvertWrap(TextureAddress);
 
 	const TextureFilter TextureFilter = Builder.GetBakeFilterForMaterialProperty(Material, GetPropertyGroup(MetallicProperty));
-	const EGLTFJsonTextureFilter TextureMinFilter = FGLTFConverterUtility::ConvertMinFilter(TextureFilter);
-	const EGLTFJsonTextureFilter TextureMagFilter = FGLTFConverterUtility::ConvertMagFilter(TextureFilter);
+	const EGLTFJsonTextureFilter TextureMinFilter = FGLTFCoreUtilities::ConvertMinFilter(TextureFilter);
+	const EGLTFJsonTextureFilter TextureMagFilter = FGLTFCoreUtilities::ConvertMagFilter(TextureFilter);
 
 	FGLTFPropertyBakeOutput MetallicBakeOutput = BakeMaterialProperty(MetallicProperty, MetallicTexCoord, TextureSize, false);
 	FGLTFPropertyBakeOutput RoughnessBakeOutput = BakeMaterialProperty(RoughnessProperty, RoughnessTexCoord, TextureSize, false);
@@ -787,12 +787,12 @@ bool FGLTFDelayedMaterialTask::TryGetClearCoatRoughness(FGLTFJsonClearCoatExtens
 	const FIntPoint TextureSize = Builder.GetBakeSizeForMaterialProperty(Material, GetPropertyGroup(IntensityProperty));
 
 	const TextureAddress TextureAddress = Builder.GetBakeTilingForMaterialProperty(Material,GetPropertyGroup(IntensityProperty));
-	const EGLTFJsonTextureWrap TextureWrapS = FGLTFConverterUtility::ConvertWrap(TextureAddress);
-	const EGLTFJsonTextureWrap TextureWrapT = FGLTFConverterUtility::ConvertWrap(TextureAddress);
+	const EGLTFJsonTextureWrap TextureWrapS = FGLTFCoreUtilities::ConvertWrap(TextureAddress);
+	const EGLTFJsonTextureWrap TextureWrapT = FGLTFCoreUtilities::ConvertWrap(TextureAddress);
 
 	const TextureFilter TextureFilter = Builder.GetBakeFilterForMaterialProperty(Material, GetPropertyGroup(IntensityProperty));
-	const EGLTFJsonTextureFilter TextureMinFilter = FGLTFConverterUtility::ConvertMinFilter(TextureFilter);
-	const EGLTFJsonTextureFilter TextureMagFilter = FGLTFConverterUtility::ConvertMagFilter(TextureFilter);
+	const EGLTFJsonTextureFilter TextureMinFilter = FGLTFCoreUtilities::ConvertMinFilter(TextureFilter);
+	const EGLTFJsonTextureFilter TextureMagFilter = FGLTFCoreUtilities::ConvertMagFilter(TextureFilter);
 
 	const FGLTFPropertyBakeOutput IntensityBakeOutput = BakeMaterialProperty(IntensityProperty, IntensityTexCoord, TextureSize, false);
 	const FGLTFPropertyBakeOutput RoughnessBakeOutput = BakeMaterialProperty(RoughnessProperty, RoughnessTexCoord, TextureSize, false);
@@ -890,7 +890,7 @@ bool FGLTFDelayedMaterialTask::TryGetEmissive(FGLTFJsonMaterial& OutMaterial, co
 	if (PropertyBakeOutput.bIsConstant)
 	{
 		const FLinearColor EmissiveColor = PropertyBakeOutput.ConstantValue;
-		OutMaterial.EmissiveFactor = FGLTFConverterUtility::ConvertColor3(EmissiveColor * EmissiveScale, Builder.ExportOptions->bStrictCompliance);
+		OutMaterial.EmissiveFactor = FGLTFCoreUtilities::ConvertColor3(EmissiveColor * EmissiveScale, Builder.ExportOptions->bStrictCompliance);
 	}
 	else
 	{
@@ -941,7 +941,7 @@ bool FGLTFDelayedMaterialTask::TryGetConstantColor(FGLTFJsonColor3& OutValue, co
 	FLinearColor Value;
 	if (TryGetConstantColor(Value, Property))
 	{
-		OutValue = FGLTFConverterUtility::ConvertColor3(Value, Builder.ExportOptions->bStrictCompliance);
+		OutValue = FGLTFCoreUtilities::ConvertColor3(Value, Builder.ExportOptions->bStrictCompliance);
 		return true;
 	}
 
@@ -953,7 +953,7 @@ bool FGLTFDelayedMaterialTask::TryGetConstantColor(FGLTFJsonColor4& OutValue, co
 	FLinearColor Value;
 	if (TryGetConstantColor(Value, Property))
 	{
-		OutValue = FGLTFConverterUtility::ConvertColor(Value, Builder.ExportOptions->bStrictCompliance);
+		OutValue = FGLTFCoreUtilities::ConvertColor(Value, Builder.ExportOptions->bStrictCompliance);
 		return true;
 	}
 
@@ -1323,7 +1323,7 @@ bool FGLTFDelayedMaterialTask::TryGetBakedMaterialProperty(FGLTFJsonTextureInfo&
 
 	if (PropertyBakeOutput.bIsConstant)
 	{
-		OutConstant = FGLTFConverterUtility::ConvertColor3(PropertyBakeOutput.ConstantValue, Builder.ExportOptions->bStrictCompliance);
+		OutConstant = FGLTFCoreUtilities::ConvertColor3(PropertyBakeOutput.ConstantValue, Builder.ExportOptions->bStrictCompliance);
 		return true;
 	}
 
@@ -1357,7 +1357,7 @@ bool FGLTFDelayedMaterialTask::TryGetBakedMaterialProperty(FGLTFJsonTextureInfo&
 
 	if (PropertyBakeOutput.bIsConstant)
 	{
-		OutConstant = FGLTFConverterUtility::ConvertColor(PropertyBakeOutput.ConstantValue, Builder.ExportOptions->bStrictCompliance);
+		OutConstant = FGLTFCoreUtilities::ConvertColor(PropertyBakeOutput.ConstantValue, Builder.ExportOptions->bStrictCompliance);
 		return true;
 	}
 
@@ -1526,12 +1526,12 @@ FGLTFPropertyBakeOutput FGLTFDelayedMaterialTask::BakeMaterialProperty(const FGL
 bool FGLTFDelayedMaterialTask::StoreBakedPropertyTexture(FGLTFJsonTextureInfo& OutTexInfo, FGLTFPropertyBakeOutput& PropertyBakeOutput, const FString& PropertyName) const
 {
 	const TextureAddress TextureAddress = Builder.GetBakeTilingForMaterialProperty(Material, GetPropertyGroup(PropertyBakeOutput.Property));
-	const EGLTFJsonTextureWrap TextureWrapS = FGLTFConverterUtility::ConvertWrap(TextureAddress);
-	const EGLTFJsonTextureWrap TextureWrapT = FGLTFConverterUtility::ConvertWrap(TextureAddress);
+	const EGLTFJsonTextureWrap TextureWrapS = FGLTFCoreUtilities::ConvertWrap(TextureAddress);
+	const EGLTFJsonTextureWrap TextureWrapT = FGLTFCoreUtilities::ConvertWrap(TextureAddress);
 
 	const TextureFilter TextureFilter = Builder.GetBakeFilterForMaterialProperty(Material, GetPropertyGroup(PropertyBakeOutput.Property));
-	const EGLTFJsonTextureFilter TextureMinFilter = FGLTFConverterUtility::ConvertMinFilter(TextureFilter);
-	const EGLTFJsonTextureFilter TextureMagFilter = FGLTFConverterUtility::ConvertMagFilter(TextureFilter);
+	const EGLTFJsonTextureFilter TextureMinFilter = FGLTFCoreUtilities::ConvertMinFilter(TextureFilter);
+	const EGLTFJsonTextureFilter TextureMagFilter = FGLTFCoreUtilities::ConvertMagFilter(TextureFilter);
 
 	FGLTFJsonTexture* Texture = FGLTFMaterialUtility::AddTexture(
 		Builder,
diff --git a/Engine/Plugins/Enterprise/GLTFExporter/MigratedFromGit/Source/GLTFExporter/Private/Tasks/GLTFDelayedTextureTasks.cpp b/Engine/Plugins/Enterprise/GLTFExporter/MigratedFromGit/Source/GLTFExporter/Private/Tasks/GLTFDelayedTextureTasks.cpp
--- a/Engine/Plugins/Enterprise/GLTFExporter/MigratedFromGit/Source/GLTFExporter/Private/Tasks/GLTFDelayedTextureTasks.cpp
+++ b/Engine/Plugins/Enterprise/GLTFExporter/MigratedFromGit/Source/GLTFExporter/Private/Tasks/GLTFDelayedTextureTasks.cpp
@@ -2,7 +2,7 @@
 
 #include "Tasks/GLTFDelayedTextureTasks.h"
 #include "Builders/GLTFContainerBuilder.h"
-#include "Converters/GLTFConverterUtility.h"
+#include "Utilities/GLTFCoreUtilities.h"
 #include "Converters/GLTFTextureUtility.h"
 #include "Converters/GLTFNameUtility.h"
 #include "Engine/Texture2D.h"
@@ -69,7 +69,7 @@ FString FGLTFDelayedTextureCubeTask::GetName()
 
 void FGLTFDelayedTextureCubeTask::Process()
 {
-	JsonTexture->Name = TextureCube->GetName() + TEXT("_") + FGLTFJsonUtility::GetValue(FGLTFConverterUtility::ConvertCubeFace(CubeFace));
+	JsonTexture->Name = TextureCube->GetName() + TEXT("_") + FGLTFJsonUtilities::GetValue(FGLTFCoreUtilities::ConvertCubeFace(CubeFace));
 
 	const bool bFromSRGB = TextureCube->SRGB;
 	if (bFromSRGB != bToSRGB)
@@ -165,7 +165,7 @@ FString FGLTFDelayedTextureRenderTargetCubeTask::GetName()
 
 void FGLTFDelayedTextureRenderTargetCubeTask::Process()
 {
-	JsonTexture->Name = RenderTargetCube->GetName() + TEXT("_") + FGLTFJsonUtility::GetValue(FGLTFConverterUtility::ConvertCubeFace(CubeFace));
+	JsonTexture->Name = RenderTargetCube->GetName() + TEXT("_") + FGLTFJsonUtilities::GetValue(FGLTFCoreUtilities::ConvertCubeFace(CubeFace));
 
 	const bool bFromSRGB = RenderTargetCube->SRGB;
 	if (bFromSRGB != bToSRGB)
diff --git a/Engine/Plugins/Enterprise/GLTFExporter/MigratedFromGit/Source/GLTFExporter/Private/Utilities/GLTFCoreUtilities.cpp b/Engine/Plugins/Enterprise/GLTFExporter/MigratedFromGit/Source/GLTFExporter/Private/Utilities/GLTFCoreUtilities.cpp
new file mode 100644
--- /dev/null
+++ b/Engine/Plugins/Enterprise/GLTFExporter/MigratedFromGit/Source/GLTFExporter/Private/Utilities/GLTFCoreUtilities.cpp
@@ -0,0 +1,332 @@
+// Copyright Epic Games, Inc. All Rights Reserved.
+
+#include "Utilities/GLTFCoreUtilities.h"
+#include "Converters/GLTFTextureUtility.h"
+#include "Actors/GLTFCameraActor.h"
+
+float FGLTFCoreUtilities::ConvertLength(const float Length, const float ConversionScale)
+{
+	return Length * ConversionScale;
+}
+
+FGLTFVector3 FGLTFCoreUtilities::ConvertVector(const FVector& Vector)
+{
+	// UE4 uses a left-handed coordinate system, with Z up.
+	// glTF uses a right-handed coordinate system, with Y up.
+	return { Vector.X, Vector.Z, Vector.Y };
+}
+
+FGLTFVector3 FGLTFCoreUtilities::ConvertPosition(const FVector& Position, const float ConversionScale)
+{
+	return ConvertVector(Position * ConversionScale);
+}
+
+FGLTFVector3 FGLTFCoreUtilities::ConvertScale(const FVector& Scale)
+{
+	return ConvertVector(Scale);
+}
+
+FGLTFVector3 FGLTFCoreUtilities::ConvertNormal(const FVector& Normal)
+{
+	return ConvertVector(Normal);
+}
+
+FGLTFInt16Vector4 FGLTFCoreUtilities::ConvertNormal(const FPackedRGBA16N& Normal)
+{
+	return { Normal.X, Normal.Z, Normal.Y, 0 };
+}
+
+FGLTFInt8Vector4 FGLTFCoreUtilities::ConvertNormal(const FPackedNormal& Normal)
+{
+	return { Normal.Vector.X, Normal.Vector.Z, Normal.Vector.Y, 0 };
+}
+
+FGLTFVector4 FGLTFCoreUtilities::ConvertTangent(const FVector& Tangent)
+{
+	// glTF stores tangent as Vec4, with W component indicating handedness of tangent basis.
+	return { Tangent.X, Tangent.Z, Tangent.Y, 1.0f };
+}
+
+FGLTFInt16Vector4 FGLTFCoreUtilities::ConvertTangent(const FPackedRGBA16N& Tangent)
+{
+	return { Tangent.X, Tangent.Z, Tangent.Y, MAX_int16 /* = 1.0 */ };
+}
+
+FGLTFInt8Vector4 FGLTFCoreUtilities::ConvertTangent(const FPackedNormal& Tangent)
+{
+	return { Tangent.Vector.X, Tangent.Vector.Z, Tangent.Vector.Y, MAX_int8 /* = 1.0 */ };
+}
+
+FGLTFVector2 FGLTFCoreUtilities::ConvertUV(const FVector2D& UV)
+{
+	// No conversion actually needed, this is primarily for type-safety.
+	return { UV.X, UV.Y };
+}
+
+FGLTFVector2 FGLTFCoreUtilities::ConvertUV(const FVector2DHalf& UV)
+{
+	return ConvertUV(FVector2D(UV));
+}
+
+FGLTFColor4 FGLTFCoreUtilities::ConvertColor(const FLinearColor& Color, bool bForceLDR)
+{
+	if (bForceLDR)
+	{
+		return {
+			FMath::Clamp(Color.R, 0.0f, 1.0f),
+			FMath::Clamp(Color.G, 0.0f, 1.0f),
+			FMath::Clamp(Color.B, 0.0f, 1.0f),
+			FMath::Clamp(Color.A, 0.0f, 1.0f)
+		};
+	}
+
+	// Just make sure its non-negative (which can happen when using MakeFromColorTemperature).
+	return {
+		FMath::Max(Color.R, 0.0f),
+		FMath::Max(Color.G, 0.0f),
+		FMath::Max(Color.B, 0.0f),
+		FMath::Max(Color.A, 0.0f)
+	};
+}
+
+FGLTFColor3 FGLTFCoreUtilities::ConvertColor3(const FLinearColor& Color, bool bForceLDR)
+{
+	if (bForceLDR)
+	{
+		return {
+			FMath::Clamp(Color.R, 0.0f, 1.0f),
+			FMath::Clamp(Color.G, 0.0f, 1.0f),
+			FMath::Clamp(Color.B, 0.0f, 1.0f)
+		};
+	}
+
+	// Just make sure its non-negative (which can happen when using MakeFromColorTemperature).
+	return {
+		FMath::Max(Color.R, 0.0f),
+		FMath::Max(Color.G, 0.0f),
+		FMath::Max(Color.B, 0.0f)
+	};
+}
+
+FGLTFUInt8Color4 FGLTFCoreUtilities::ConvertColor(const FColor& Color)
+{
+	// UE4 uses ABGR or ARGB depending on endianness.
+	// glTF always uses RGBA independent of endianness.
+	return { Color.R, Color.G, Color.B, Color.A };
+}
+
+FGLTFQuaternion FGLTFCoreUtilities::ConvertRotation(const FRotator& Rotation)
+{
+	return ConvertRotation(Rotation.Quaternion());
+}
+
+FGLTFQuaternion FGLTFCoreUtilities::ConvertRotation(const FQuat& Rotation)
+{
+	// UE4 uses a left-handed coordinate system, with Z up.
+	// glTF uses a right-handed coordinate system, with Y up.
+	// Rotation = (qX, qY, qZ, qW) = (sin(angle/2) * aX, sin(angle/2) * aY, sin(angle/2) * aZ, cons(angle/2))
+	// where (aX, aY, aZ) - rotation axis, angle - rotation angle
+	// Y swapped with Z between these coordinate systems
+	// also, as handedness is changed rotation is inversed - hence negation
+	// therefore glTFRotation = (-qX, -qZ, -qY, qw)
+
+	const FQuat Normalized = Rotation.GetNormalized();
+	return { -Normalized.X, -Normalized.Z, -Normalized.Y, Normalized.W };
+}
+
+FGLTFMatrix4 FGLTFCoreUtilities::ConvertMatrix(const FMatrix& Matrix)
+{
+	// Unreal stores matrix elements in row major order.
+	// glTF stores matrix elements in column major order.
+
+	FGLTFMatrix4 Result;
+	for (int32 Row = 0; Row < 4; ++Row)
+	{
+		for (int32 Col = 0; Col < 4; ++Col)
+		{
+			Result(Row, Col) = Matrix.M[Row][Col];
+		}
+	}
+	return Result;
+}
+
+FGLTFMatrix4 FGLTFCoreUtilities::ConvertTransform(const FTransform& Transform, const float ConversionScale)
+{
+	const FQuat Rotation = Transform.GetRotation();
+	const FVector Translation = Transform.GetTranslation();
+	const FVector Scale = Transform.GetScale3D();
+
+	const FQuat ConvertedRotation(-Rotation.X, -Rotation.Z, -Rotation.Y, Rotation.W);
+	const FVector ConvertedTranslation = FVector(Translation.X, Translation.Z, Translation.Y) * ConversionScale;
+	const FVector ConvertedScale(Scale.X, Scale.Z, Scale.Y);
+
+	const FTransform ConvertedTransform(ConvertedRotation, ConvertedTranslation, ConvertedScale);
+	const FMatrix Matrix = ConvertedTransform.ToMatrixWithScale(); // TODO: should it be ToMatrixNoScale()?
+	return ConvertMatrix(Matrix);
+}
+
+float FGLTFCoreUtilities::ConvertFieldOfView(float FOVInDegrees, float AspectRatio)
+{
+	const float HorizontalFOV = FMath::DegreesToRadians(FOVInDegrees);
+	const float VerticalFOV = 2 * FMath::Atan(FMath::Tan(HorizontalFOV / 2) / AspectRatio);
+	return VerticalFOV;
+}
+
+float FGLTFCoreUtilities::ConvertLightAngle(const float Angle)
+{
+	// Unreal uses degrees.
+	// glTF uses radians.
+	return FMath::DegreesToRadians(Angle);
+}
+
+FGLTFQuaternion FGLTFCoreUtilities::GetLocalCameraRotation()
+{
+	// Unreal uses +X axis as camera direction in Unreal coordinates.
+	// glTF uses +Y as camera direction in Unreal coordinates.
+
+	static FGLTFQuaternion Rotation = ConvertRotation(FRotator(0, 90, 0));
+	return Rotation;
+}
+
+FGLTFQuaternion FGLTFCoreUtilities::GetLocalLightRotation()
+{
+	// Unreal uses +X axis as light direction in Unreal coordinates.
+	// glTF uses +Y as light direction in Unreal coordinates.
+
+	static FGLTFQuaternion Rotation = ConvertRotation(FRotator(0, 90, 0));
+	return Rotation;
+}
+
+EGLTFJsonCameraType FGLTFCoreUtilities::ConvertCameraType(ECameraProjectionMode::Type ProjectionMode)
+{
+	switch (ProjectionMode)
+	{
+		case ECameraProjectionMode::Perspective:  return EGLTFJsonCameraType::Perspective;
+		case ECameraProjectionMode::Orthographic: return EGLTFJsonCameraType::Orthographic;
+		default:                                  return EGLTFJsonCameraType::None;
+	}
+}
+
+EGLTFJsonLightType FGLTFCoreUtilities::ConvertLightType(ELightComponentType ComponentType)
+{
+	switch (ComponentType)
+	{
+		case LightType_Directional: return EGLTFJsonLightType::Directional;
+		case LightType_Point:       return EGLTFJsonLightType::Point;
+		case LightType_Spot:        return EGLTFJsonLightType::Spot;
+		default:                    return EGLTFJsonLightType::None;
+	}
+}
+
+EGLTFJsonInterpolation FGLTFCoreUtilities::ConvertInterpolation(const EAnimInterpolationType Type)
+{
+	switch (Type)
+	{
+		case EAnimInterpolationType::Linear: return EGLTFJsonInterpolation::Linear;
+		case EAnimInterpolationType::Step:   return EGLTFJsonInterpolation::Step;
+		default:                             return EGLTFJsonInterpolation::None;
+	}
+}
+
+EGLTFJsonShadingModel FGLTFCoreUtilities::ConvertShadingModel(EMaterialShadingModel ShadingModel)
+{
+	switch (ShadingModel)
+	{
+		case MSM_Unlit:      return EGLTFJsonShadingModel::Unlit;
+		case MSM_DefaultLit: return EGLTFJsonShadingModel::Default;
+		case MSM_ClearCoat:  return EGLTFJsonShadingModel::ClearCoat;
+		default:             return EGLTFJsonShadingModel::None;
+	}
+}
+
+EGLTFJsonAlphaMode FGLTFCoreUtilities::ConvertAlphaMode(EBlendMode Mode)
+{
+	switch (Mode)
+	{
+		case BLEND_Opaque:         return EGLTFJsonAlphaMode::Opaque;
+		case BLEND_Masked:         return EGLTFJsonAlphaMode::Mask;
+		case BLEND_Translucent:    return EGLTFJsonAlphaMode::Blend;
+		case BLEND_Additive:       return EGLTFJsonAlphaMode::Blend;
+		case BLEND_Modulate:       return EGLTFJsonAlphaMode::Blend;
+		case BLEND_AlphaComposite: return EGLTFJsonAlphaMode::Blend;
+		default:                   return EGLTFJsonAlphaMode::None;
+	}
+}
+
+EGLTFJsonBlendMode FGLTFCoreUtilities::ConvertBlendMode(EBlendMode Mode)
+{
+	switch (Mode)
+	{
+		case BLEND_Additive:       return EGLTFJsonBlendMode::Additive;
+		case BLEND_Modulate:       return EGLTFJsonBlendMode::Modulate;
+		case BLEND_AlphaComposite: return EGLTFJsonBlendMode::AlphaComposite;
+		default:                   return EGLTFJsonBlendMode::None;
+	}
+}
+
+EGLTFJsonTextureWrap FGLTFCoreUtilities::ConvertWrap(TextureAddress Address)
+{
+	switch (Address)
+	{
+		case TA_Wrap:   return EGLTFJsonTextureWrap::Repeat;
+		case TA_Mirror: return EGLTFJsonTextureWrap::MirroredRepeat;
+		case TA_Clamp:  return EGLTFJsonTextureWrap::ClampToEdge;
+		default:        return EGLTFJsonTextureWrap::None; // TODO: add error handling in callers
+	}
+}
+
+EGLTFJsonTextureFilter FGLTFCoreUtilities::ConvertMinFilter(TextureFilter Filter)
+{
+	switch (Filter)
+	{
+		case TF_Nearest:   return EGLTFJsonTextureFilter::NearestMipmapNearest;
+		case TF_Bilinear:  return EGLTFJsonTextureFilter::LinearMipmapNearest;
+		case TF_Trilinear: return EGLTFJsonTextureFilter::LinearMipmapLinear;
+		default:           return EGLTFJsonTextureFilter::None;
+	}
+}
+
+EGLTFJsonTextureFilter FGLTFCoreUtilities::ConvertMagFilter(TextureFilter Filter)
+{
+	switch (Filter)
+	{
+		case TF_Nearest:   return EGLTFJsonTextureFilter::Nearest;
+		case TF_Bilinear:  return EGLTFJsonTextureFilter::Linear;
+		case TF_Trilinear: return EGLTFJsonTextureFilter::Linear;
+		default:           return EGLTFJsonTextureFilter::None;
+	}
+}
+
+EGLTFJsonTextureFilter FGLTFCoreUtilities::ConvertMinFilter(TextureFilter Filter, TextureGroup LODGroup)
+{
+	return ConvertMinFilter(Filter == TF_Default ? FGLTFTextureUtility::GetDefaultFilter(LODGroup) : Filter);
+}
+
+EGLTFJsonTextureFilter FGLTFCoreUtilities::ConvertMagFilter(TextureFilter Filter, TextureGroup LODGroup)
+{
+	return ConvertMagFilter(Filter == TF_Default ? FGLTFTextureUtility::GetDefaultFilter(LODGroup) : Filter);
+}
+
+EGLTFJsonCubeFace FGLTFCoreUtilities::ConvertCubeFace(ECubeFace CubeFace)
+{
+	switch (CubeFace)
+	{
+		case CubeFace_PosX:	return EGLTFJsonCubeFace::NegX;
+		case CubeFace_NegX:	return EGLTFJsonCubeFace::PosX;
+		case CubeFace_PosY:	return EGLTFJsonCubeFace::PosZ;
+		case CubeFace_NegY:	return EGLTFJsonCubeFace::NegZ;
+		case CubeFace_PosZ:	return EGLTFJsonCubeFace::PosY;
+		case CubeFace_NegZ:	return EGLTFJsonCubeFace::NegY;
+		default:            return EGLTFJsonCubeFace::None; // TODO: add error handling in callers
+	}
+}
+
+EGLTFJsonCameraControlMode FGLTFCoreUtilities::ConvertCameraControlMode(EGLTFCameraControlMode CameraMode)
+{
+	switch (CameraMode)
+	{
+		case EGLTFCameraControlMode::FreeLook: return EGLTFJsonCameraControlMode::FreeLook;
+		case EGLTFCameraControlMode::Orbital:  return EGLTFJsonCameraControlMode::Orbital;
+		default:                               return EGLTFJsonCameraControlMode::None; // TODO: add error handling in callers
+	}
+}
diff --git a/Engine/Plugins/Enterprise/GLTFExporter/MigratedFromGit/Source/GLTFExporter/Private/Utilities/GLTFJsonUtilities.cpp b/Engine/Plugins/Enterprise/GLTFExporter/MigratedFromGit/Source/GLTFExporter/Private/Utilities/GLTFJsonUtilities.cpp
new file mode 100644
--- /dev/null
+++ b/Engine/Plugins/Enterprise/GLTFExporter/MigratedFromGit/Source/GLTFExporter/Private/Utilities/GLTFJsonUtilities.cpp
@@ -0,0 +1,188 @@
+// Copyright Epic Games, Inc. All Rights Reserved.
+
+#include "Utilities/GLTFJsonUtilities.h"
+
+const TCHAR* FGLTFJsonUtilities::GetValue(EGLTFJsonExtension Enum)
+{
+	switch (Enum)
+	{
+		case EGLTFJsonExtension::KHR_LightsPunctual:      return TEXT("KHR_lights_punctual");
+		case EGLTFJsonExtension::KHR_MaterialsClearCoat:  return TEXT("KHR_materials_clearcoat");
+		case EGLTFJsonExtension::KHR_MaterialsUnlit:      return TEXT("KHR_materials_unlit");
+		case EGLTFJsonExtension::KHR_MaterialsVariants:   return TEXT("KHR_materials_variants");
+		case EGLTFJsonExtension::KHR_MeshQuantization:    return TEXT("KHR_mesh_quantization");
+		case EGLTFJsonExtension::KHR_TextureTransform:    return TEXT("KHR_texture_transform");
+		case EGLTFJsonExtension::EPIC_AnimationHotspots:  return TEXT("EPIC_animation_hotspots");
+		case EGLTFJsonExtension::EPIC_AnimationPlayback:  return TEXT("EPIC_animation_playback");
+		case EGLTFJsonExtension::EPIC_BlendModes:         return TEXT("EPIC_blend_modes");
+		case EGLTFJsonExtension::EPIC_CameraControls:     return TEXT("EPIC_camera_controls");
+		case EGLTFJsonExtension::EPIC_HDRIBackdrops:      return TEXT("EPIC_hdri_backdrops");
+		case EGLTFJsonExtension::EPIC_LevelVariantSets:   return TEXT("EPIC_level_variant_sets");
+		case EGLTFJsonExtension::EPIC_LightmapTextures:   return TEXT("EPIC_lightmap_textures");
+		case EGLTFJsonExtension::EPIC_SkySpheres:         return TEXT("EPIC_sky_spheres");
+		case EGLTFJsonExtension::EPIC_TextureHDREncoding: return TEXT("EPIC_texture_hdr_encoding");
+		default:
+			checkNoEntry();
+			return TEXT("");
+	}
+}
+
+const TCHAR* FGLTFJsonUtilities::GetValue(EGLTFJsonAlphaMode Enum)
+{
+	switch (Enum)
+	{
+		case EGLTFJsonAlphaMode::Opaque: return TEXT("OPAQUE");
+		case EGLTFJsonAlphaMode::Blend:  return TEXT("BLEND");
+		case EGLTFJsonAlphaMode::Mask:   return TEXT("MASK");
+		default:
+			checkNoEntry();
+			return TEXT("");
+	}
+}
+
+const TCHAR* FGLTFJsonUtilities::GetValue(EGLTFJsonBlendMode Enum)
+{
+	switch (Enum)
+	{
+		case EGLTFJsonBlendMode::Additive:       return TEXT("ADDITIVE");
+		case EGLTFJsonBlendMode::Modulate:       return TEXT("MODULATE");
+		case EGLTFJsonBlendMode::AlphaComposite: return TEXT("ALPHACOMPOSITE");
+		default:
+			checkNoEntry();
+			return TEXT("");
+	}
+}
+
+const TCHAR* FGLTFJsonUtilities::GetValue(EGLTFJsonMimeType Enum)
+{
+	switch (Enum)
+	{
+		case EGLTFJsonMimeType::PNG:  return TEXT("image/png");
+		case EGLTFJsonMimeType::JPEG: return TEXT("image/jpeg");
+		default:
+			checkNoEntry();
+			return TEXT("");
+	}
+}
+
+const TCHAR* FGLTFJsonUtilities::GetValue(EGLTFJsonAccessorType Enum)
+{
+	switch (Enum)
+	{
+		case EGLTFJsonAccessorType::Scalar: return TEXT("SCALAR");
+		case EGLTFJsonAccessorType::Vec2:   return TEXT("VEC2");
+		case EGLTFJsonAccessorType::Vec3:   return TEXT("VEC3");
+		case EGLTFJsonAccessorType::Vec4:   return TEXT("VEC4");
+		case EGLTFJsonAccessorType::Mat2:   return TEXT("MAT2");
+		case EGLTFJsonAccessorType::Mat3:   return TEXT("MAT3");
+		case EGLTFJsonAccessorType::Mat4:   return TEXT("MAT4");
+		default:
+			checkNoEntry();
+			return TEXT("");
+	}
+}
+
+const TCHAR* FGLTFJsonUtilities::GetValue(EGLTFJsonHDREncoding Enum)
+{
+	switch (Enum)
+	{
+		case EGLTFJsonHDREncoding::RGBM: return TEXT("RGBM");
+		case EGLTFJsonHDREncoding::RGBE: return TEXT("RGBE");
+		default:
+			checkNoEntry();
+			return TEXT("");
+	}
+}
+
+const TCHAR* FGLTFJsonUtilities::GetValue(EGLTFJsonCubeFace Enum)
+{
+	switch (Enum)
+	{
+		case EGLTFJsonCubeFace::PosX: return TEXT("PosX");
+		case EGLTFJsonCubeFace::NegX: return TEXT("NegX");
+		case EGLTFJsonCubeFace::PosY: return TEXT("PosY");
+		case EGLTFJsonCubeFace::NegY: return TEXT("NegY");
+		case EGLTFJsonCubeFace::PosZ: return TEXT("PosZ");
+		case EGLTFJsonCubeFace::NegZ: return TEXT("NegZ");
+		default:
+			checkNoEntry();
+			return TEXT("");
+	}
+}
+
+const TCHAR* FGLTFJsonUtilities::GetValue(EGLTFJsonCameraType Enum)
+{
+	switch (Enum)
+	{
+		case EGLTFJsonCameraType::Perspective:  return TEXT("perspective");
+		case EGLTFJsonCameraType::Orthographic: return TEXT("orthographic");
+		default:
+			checkNoEntry();
+			return TEXT("");
+	}
+}
+
+const TCHAR* FGLTFJsonUtilities::GetValue(EGLTFJsonLightType Enum)
+{
+	switch (Enum)
+	{
+		case EGLTFJsonLightType::Directional: return TEXT("directional");
+		case EGLTFJsonLightType::Point:       return TEXT("point");
+		case EGLTFJsonLightType::Spot:        return TEXT("spot");
+		default:
+			checkNoEntry();
+			return TEXT("");
+	}
+}
+
+const TCHAR* FGLTFJsonUtilities::GetValue(EGLTFJsonInterpolation Enum)
+{
+	switch (Enum)
+	{
+		case EGLTFJsonInterpolation::Linear:      return TEXT("LINEAR");
+		case EGLTFJsonInterpolation::Step:        return TEXT("STEP");
+		case EGLTFJsonInterpolation::CubicSpline: return TEXT("CUBICSPLINE");
+		default:
+			checkNoEntry();
+			return TEXT("");
+	}
+}
+
+const TCHAR* FGLTFJsonUtilities::GetValue(EGLTFJsonTargetPath Enum)
+{
+	switch (Enum)
+	{
+		case EGLTFJsonTargetPath::Translation: return TEXT("translation");
+		case EGLTFJsonTargetPath::Rotation:    return TEXT("rotation");
+		case EGLTFJsonTargetPath::Scale:       return TEXT("scale");
+		case EGLTFJsonTargetPath::Weights:     return TEXT("weights");
+		default:
+			checkNoEntry();
+			return TEXT("");
+	}
+}
+
+const TCHAR* FGLTFJsonUtilities::GetValue(EGLTFJsonCameraControlMode Enum)
+{
+	switch (Enum)
+	{
+		case EGLTFJsonCameraControlMode::FreeLook: return TEXT("freeLook");
+		case EGLTFJsonCameraControlMode::Orbital:  return TEXT("orbital");
+		default:
+			checkNoEntry();
+			return TEXT("");
+	}
+}
+
+const TCHAR* FGLTFJsonUtilities::GetValue(EGLTFJsonShadingModel Enum)
+{
+	switch (Enum)
+	{
+		case EGLTFJsonShadingModel::Default:   return TEXT("Default");
+		case EGLTFJsonShadingModel::Unlit:     return TEXT("Unlit");
+		case EGLTFJsonShadingModel::ClearCoat: return TEXT("ClearCoat");
+		default:
+			checkNoEntry();
+			return TEXT("");
+	}
+}
diff --git a/Engine/Plugins/Enterprise/GLTFExporter/MigratedFromGit/Source/GLTFExporter/Private/Utilities/GLTFProxyMaterialUtilities.cpp b/Engine/Plugins/Enterprise/GLTFExporter/MigratedFromGit/Source/GLTFExporter/Private/Utilities/GLTFProxyMaterialUtilities.cpp
--- a/Engine/Plugins/Enterprise/GLTFExporter/MigratedFromGit/Source/GLTFExporter/Private/Utilities/GLTFProxyMaterialUtilities.cpp
+++ b/Engine/Plugins/Enterprise/GLTFExporter/MigratedFromGit/Source/GLTFExporter/Private/Utilities/GLTFProxyMaterialUtilities.cpp
@@ -1,7 +1,7 @@
 // Copyright Epic Games, Inc. All Rights Reserved.
 
 #include "Utilities/GLTFProxyMaterialUtilities.h"
-#include "Json/GLTFJsonUtility.h"
+#include "Utilities/GLTFJsonUtilities.h"
 #include "UserData/GLTFMaterialUserData.h"
 #include "Materials/MaterialInstanceDynamic.h"
 #include "Materials/MaterialInstanceConstant.h"
@@ -34,7 +34,7 @@ UMaterial* FGLTFProxyMaterialUtilities::GetBaseMaterial(EGLTFJsonShadingModel Sh
 	UMaterial* Material = Materials[Index];
 	if (Material == nullptr)
 	{
-		const FString Name = FGLTFJsonUtility::GetValue(ShadingModel);
+		const FString Name = FGLTFJsonUtilities::GetValue(ShadingModel);
 		const FString Path = PROXY_MATERIAL_ROOT_PATH PROXY_MATERIAL_NAME_PREFIX + Name + TEXT(".") PROXY_MATERIAL_NAME_PREFIX + Name;
 
 		Material = LoadObject<UMaterial>(nullptr, *Path);
diff --git a/Engine/Plugins/Enterprise/GLTFExporter/MigratedFromGit/Source/GLTFExporter/Public/Json/GLTFJsonUtility.h b/Engine/Plugins/Enterprise/GLTFExporter/MigratedFromGit/Source/GLTFExporter/Public/Json/GLTFJsonUtility.h
deleted file mode 100644
diff --git a/Engine/Plugins/Enterprise/GLTFExporter/MigratedFromGit/Source/GLTFExporter/Public/Json/GLTFJsonWriter.h b/Engine/Plugins/Enterprise/GLTFExporter/MigratedFromGit/Source/GLTFExporter/Public/Json/GLTFJsonWriter.h
--- a/Engine/Plugins/Enterprise/GLTFExporter/MigratedFromGit/Source/GLTFExporter/Public/Json/GLTFJsonWriter.h
+++ b/Engine/Plugins/Enterprise/GLTFExporter/MigratedFromGit/Source/GLTFExporter/Public/Json/GLTFJsonWriter.h
@@ -4,8 +4,8 @@
 
 #include "Json/GLTFJsonValue.h"
 #include "Json/GLTFJsonIndexedObject.h"
-#include "Json/GLTFJsonUtility.h"
 #include "Json/GLTFJsonExtensions.h"
+#include "Utilities/GLTFJsonUtilities.h"
 
 class GLTFEXPORTER_API IGLTFJsonWriter
 {
@@ -65,7 +65,7 @@ public:
 	template <typename EnumType, typename = typename TEnableIf<TIsEnum<EnumType>::Value>::Type>
 	void Write(EnumType Enum)
 	{
-		Write(FGLTFJsonUtility::GetValue(Enum));
+		Write(FGLTFJsonUtilities::GetValue(Enum));
 	}
 
 	template <class ElementType, typename = typename TEnableIf<TNot<TIsSame<ElementType, TCHAR>>::Value>::Type>
@@ -241,7 +241,7 @@ public:
 			Extensions.Required.Add(Extension);
 		}
 
-		Write(FGLTFJsonUtility::GetValue(Extension), Value);
+		Write(FGLTFJsonUtilities::GetValue(Extension), Value);
 	}
 
 	void StartExtension(const EGLTFJsonExtension& Extension, bool bIsRequired = false)
@@ -253,7 +253,7 @@ public:
 			Extensions.Required.Add(Extension);
 		}
 
-		StartObject(FGLTFJsonUtility::GetValue(Extension));
+		StartObject(FGLTFJsonUtilities::GetValue(Extension));
 	}
 
 	void EndExtension()
diff --git a/Engine/Plugins/Enterprise/GLTFExporter/MigratedFromGit/Source/GLTFExporter/Public/Utilities/GLTFCoreUtilities.h b/Engine/Plugins/Enterprise/GLTFExporter/MigratedFromGit/Source/GLTFExporter/Public/Utilities/GLTFCoreUtilities.h
new file mode 100644
--- /dev/null
+++ b/Engine/Plugins/Enterprise/GLTFExporter/MigratedFromGit/Source/GLTFExporter/Public/Utilities/GLTFCoreUtilities.h
@@ -0,0 +1,87 @@
+// Copyright Epic Games, Inc. All Rights Reserved.
+
+#pragma once
+
+#include "Core/GLTFVector.h"
+#include "Core/GLTFColor.h"
+#include "Core/GLTFMatrix.h"
+#include "Core/GLTFQuaternion.h"
+#include "Json/GLTFJsonEnums.h"
+#include "Engine/EngineTypes.h"
+
+enum class EGLTFCameraControlMode : unsigned char;
+
+struct GLTFEXPORTER_API FGLTFCoreUtilities
+{
+	static float ConvertLength(const float Length, const float ConversionScale = 0.01);
+
+	static FGLTFVector3 ConvertVector(const FVector& Vector);
+
+	static FGLTFVector3 ConvertPosition(const FVector& Position, const float ConversionScale = 0.01);
+
+	static FGLTFVector3 ConvertScale(const FVector& Scale);
+
+	static FGLTFVector3 ConvertNormal(const FVector& Normal);
+	static FGLTFInt16Vector4 ConvertNormal(const FPackedRGBA16N& Normal);
+	static FGLTFInt8Vector4 ConvertNormal(const FPackedNormal& Normal);
+
+	static FGLTFVector4 ConvertTangent(const FVector& Tangent);
+	static FGLTFInt16Vector4 ConvertTangent(const FPackedRGBA16N& Tangent);
+	static FGLTFInt8Vector4 ConvertTangent(const FPackedNormal& Tangent);
+
+	static FGLTFVector2 ConvertUV(const FVector2D& UV);
+	static FGLTFVector2 ConvertUV(const FVector2DHalf& UV);
+
+	static FGLTFColor4 ConvertColor(const FLinearColor& Color, bool bForceLDR);
+	static FGLTFColor3 ConvertColor3(const FLinearColor& Color, bool bForceLDR);
+	static FGLTFUInt8Color4 ConvertColor(const FColor& Color);
+
+	static FGLTFQuaternion ConvertRotation(const FRotator& Rotation);
+	static FGLTFQuaternion ConvertRotation(const FQuat& Rotation);
+
+	static FGLTFMatrix4 ConvertMatrix(const FMatrix& Matrix);
+
+	static FGLTFMatrix4 ConvertTransform(const FTransform& Transform, const float ConversionScale = 0.01);
+
+	static float ConvertFieldOfView(float FOVInDegrees, float AspectRatio);
+
+	static float ConvertLightAngle(const float Angle);
+
+	static FGLTFQuaternion GetLocalCameraRotation();
+	static FGLTFQuaternion GetLocalLightRotation();
+
+	static EGLTFJsonCameraType ConvertCameraType(ECameraProjectionMode::Type ProjectionMode);
+	static EGLTFJsonLightType ConvertLightType(ELightComponentType ComponentType);
+
+	static EGLTFJsonInterpolation ConvertInterpolation(const EAnimInterpolationType Type);
+
+	static EGLTFJsonShadingModel ConvertShadingModel(EMaterialShadingModel ShadingModel);
+
+	static EGLTFJsonAlphaMode ConvertAlphaMode(EBlendMode Mode);
+	static EGLTFJsonBlendMode ConvertBlendMode(EBlendMode Mode);
+
+	static EGLTFJsonTextureWrap ConvertWrap(TextureAddress Address);
+
+	static EGLTFJsonTextureFilter ConvertMinFilter(TextureFilter Filter);
+	static EGLTFJsonTextureFilter ConvertMagFilter(TextureFilter Filter);
+
+	static EGLTFJsonTextureFilter ConvertMinFilter(TextureFilter Filter, TextureGroup LODGroup);
+	static EGLTFJsonTextureFilter ConvertMagFilter(TextureFilter Filter, TextureGroup LODGroup);
+
+	static EGLTFJsonCubeFace ConvertCubeFace(ECubeFace CubeFace);
+
+	static EGLTFJsonCameraControlMode ConvertCameraControlMode(EGLTFCameraControlMode CameraMode);
+
+	template <typename ComponentType>
+	static EGLTFJsonComponentType GetComponentType()
+	{
+		if (TIsSame<ComponentType, int8  >::Value) return EGLTFJsonComponentType::Int8;
+		if (TIsSame<ComponentType, uint8 >::Value) return EGLTFJsonComponentType::UInt8;
+		if (TIsSame<ComponentType, int16 >::Value) return EGLTFJsonComponentType::Int16;
+		if (TIsSame<ComponentType, uint16>::Value) return EGLTFJsonComponentType::UInt16;
+		if (TIsSame<ComponentType, int32 >::Value) return EGLTFJsonComponentType::Int32;
+		if (TIsSame<ComponentType, uint32>::Value) return EGLTFJsonComponentType::UInt32;
+		if (TIsSame<ComponentType, float >::Value) return EGLTFJsonComponentType::Float;
+		return EGLTFJsonComponentType::None;
+	}
+};
diff --git a/Engine/Plugins/Enterprise/GLTFExporter/MigratedFromGit/Source/GLTFExporter/Public/Utilities/GLTFJsonUtilities.h b/Engine/Plugins/Enterprise/GLTFExporter/MigratedFromGit/Source/GLTFExporter/Public/Utilities/GLTFJsonUtilities.h
new file mode 100644
--- /dev/null
+++ b/Engine/Plugins/Enterprise/GLTFExporter/MigratedFromGit/Source/GLTFExporter/Public/Utilities/GLTFJsonUtilities.h
@@ -0,0 +1,29 @@
+// Copyright Epic Games, Inc. All Rights Reserved.
+
+#pragma once
+
+#include "CoreMinimal.h"
+#include "Json/GLTFJsonEnums.h"
+
+struct GLTFEXPORTER_API FGLTFJsonUtilities
+{
+	template <typename EnumType, typename = typename TEnableIf<TIsEnum<EnumType>::Value>::Type>
+	static int32 GetValue(EnumType Enum)
+	{
+		return static_cast<int32>(Enum);
+	}
+
+	static const TCHAR* GetValue(EGLTFJsonExtension Enum);
+	static const TCHAR* GetValue(EGLTFJsonAlphaMode Enum);
+	static const TCHAR* GetValue(EGLTFJsonBlendMode Enum);
+	static const TCHAR* GetValue(EGLTFJsonMimeType Enum);
+	static const TCHAR* GetValue(EGLTFJsonAccessorType Enum);
+	static const TCHAR* GetValue(EGLTFJsonHDREncoding Enum);
+	static const TCHAR* GetValue(EGLTFJsonCubeFace Enum);
+	static const TCHAR* GetValue(EGLTFJsonCameraType Enum);
+	static const TCHAR* GetValue(EGLTFJsonLightType Enum);
+	static const TCHAR* GetValue(EGLTFJsonInterpolation Enum);
+	static const TCHAR* GetValue(EGLTFJsonTargetPath Enum);
+	static const TCHAR* GetValue(EGLTFJsonCameraControlMode Enum);
+	static const TCHAR* GetValue(EGLTFJsonShadingModel Enum);
+};
