[GLTFExporter] fix: properly handle level sequence channel masks and add checks to make sure data matches mask (needed for UE5)

commit aee8478541bdd60091eb05ccd386cf80f389517e
Date: Thu, 17 Mar 2022 16:47:01 +0100

    fix: properly handle level sequence channel masks and add checks to make sure data matches mask (needed for UE5)
    (cherry picked from commit 5f9ec85c3df30c039db446266c5c263ae1791ece​)

#skipci

diff --git a/Engine/Plugins/Enterprise/GLTFExporter/MigratedFromGit/Source/GLTFExporter/Private/Tasks/GLTFDelayedAnimationTasks.cpp b/Engine/Plugins/Enterprise/GLTFExporter/MigratedFromGit/Source/GLTFExporter/Private/Tasks/GLTFDelayedAnimationTasks.cpp
--- a/Engine/Plugins/Enterprise/GLTFExporter/MigratedFromGit/Source/GLTFExporter/Private/Tasks/GLTFDelayedAnimationTasks.cpp
+++ b/Engine/Plugins/Enterprise/GLTFExporter/MigratedFromGit/Source/GLTFExporter/Private/Tasks/GLTFDelayedAnimationTasks.cpp
@@ -189,14 +189,17 @@ void FGLTFDelayedLevelSequenceTask::Process()
 		for (TWeakObjectPtr<UObject> Object : Player->FindBoundObjects(Binding.GetObjectGuid(), MovieSceneSequenceID::Root))
 		{
 			FGLTFJsonNode* Node;
+			FTransform RelativeTransform;
 
 			if (AActor* Actor = Cast<AActor>(Object.Get()))
 			{
 				Node = Builder.AddUniqueNode(Actor);
+				RelativeTransform = Actor->GetRootComponent()->GetRelativeTransform();
 			}
 			else if (USceneComponent* SceneComponent = Cast<USceneComponent>(Object.Get()))
 			{
 				Node = Builder.AddUniqueNode(SceneComponent);
+				RelativeTransform = SceneComponent->GetRelativeTransform();
 			}
 			else
 			{
@@ -244,6 +247,21 @@ void FGLTFDelayedLevelSequenceTask::Process()
 
 					if (EnumHasAnyFlags(ChannelMask, EMovieSceneTransformChannel::Translation))
 					{
+						constexpr int32 IndexX = 0;
+						constexpr int32 IndexY = 1;
+						constexpr int32 IndexZ = 2;
+
+						const bool MaskX = EnumHasAllFlags(ChannelMask, EMovieSceneTransformChannel::TranslationX) && Channels.IsValidIndex(IndexX);
+						const bool MaskY = EnumHasAllFlags(ChannelMask, EMovieSceneTransformChannel::TranslationY) && Channels.IsValidIndex(IndexY);
+						const bool MaskZ = EnumHasAllFlags(ChannelMask, EMovieSceneTransformChannel::TranslationZ) && Channels.IsValidIndex(IndexZ);
+
+						if (!(MaskX | MaskY | MaskZ))
+						{
+							// TODO: report warning
+							continue;
+						}
+
+						FVector DefaultTranslation = RelativeTransform.GetTranslation();
 						TArray<FGLTFVector3> Translations;
 						Translations.AddUninitialized(FrameCount);
 
@@ -251,10 +269,10 @@ void FGLTFDelayedLevelSequenceTask::Process()
 						{
 							FFrameTime& FrameTime = FrameTimes[Frame];
 
-							FVector Translation = FVector::ZeroVector;
-							Channels[0]->Evaluate(FrameTime, Translation.X);
-							Channels[1]->Evaluate(FrameTime, Translation.Y);
-							Channels[2]->Evaluate(FrameTime, Translation.Z);
+							FVector Translation = DefaultTranslation;
+							if (MaskX) Channels[IndexX]->Evaluate(FrameTime, Translation.X);
+							if (MaskY) Channels[IndexY]->Evaluate(FrameTime, Translation.Y);
+							if (MaskZ) Channels[IndexZ]->Evaluate(FrameTime, Translation.Z);
 
 							Translations[Frame] = FGLTFCoreUtilities::ConvertPosition(Translation, Builder.ExportOptions->ExportUniformScale);
 						}
@@ -279,6 +297,21 @@ void FGLTFDelayedLevelSequenceTask::Process()
 
 					if (EnumHasAnyFlags(ChannelMask, EMovieSceneTransformChannel::Rotation))
 					{
+						constexpr int32 IndexX = 3;
+						constexpr int32 IndexY = 4;
+						constexpr int32 IndexZ = 5;
+
+						const bool MaskX = EnumHasAllFlags(ChannelMask, EMovieSceneTransformChannel::RotationX) && Channels.IsValidIndex(IndexX);
+						const bool MaskY = EnumHasAllFlags(ChannelMask, EMovieSceneTransformChannel::RotationY) && Channels.IsValidIndex(IndexY);
+						const bool MaskZ = EnumHasAllFlags(ChannelMask, EMovieSceneTransformChannel::RotationZ) && Channels.IsValidIndex(IndexZ);
+
+						if (!(MaskX | MaskY | MaskZ))
+						{
+							// TODO: report warning
+							continue;
+						}
+
+						FRotator DefaultRotator = RelativeTransform.GetRotation().Rotator();
 						TArray<FGLTFQuaternion> Rotations;
 						Rotations.AddUninitialized(FrameCount);
 
@@ -286,10 +319,10 @@ void FGLTFDelayedLevelSequenceTask::Process()
 						{
 							FFrameTime& FrameTime = FrameTimes[Frame];
 
-							FRotator Rotator = FRotator::ZeroRotator;
-							Channels[3]->Evaluate(FrameTime, Rotator.Roll);
-							Channels[4]->Evaluate(FrameTime, Rotator.Pitch);
-							Channels[5]->Evaluate(FrameTime, Rotator.Yaw);
+							FRotator Rotator = DefaultRotator;
+							if (MaskX) Channels[IndexX]->Evaluate(FrameTime, Rotator.Roll);
+							if (MaskY) Channels[IndexY]->Evaluate(FrameTime, Rotator.Pitch);
+							if (MaskZ) Channels[IndexZ]->Evaluate(FrameTime, Rotator.Yaw);
 
 							Rotations[Frame] = FGLTFCoreUtilities::ConvertRotation(Rotator.Quaternion());
 						}
@@ -317,6 +350,21 @@ void FGLTFDelayedLevelSequenceTask::Process()
 
 					if (EnumHasAnyFlags(ChannelMask, EMovieSceneTransformChannel::Scale))
 					{
+						constexpr int32 IndexX = 6;
+						constexpr int32 IndexY = 7;
+						constexpr int32 IndexZ = 8;
+
+						const bool MaskX = EnumHasAllFlags(ChannelMask, EMovieSceneTransformChannel::ScaleX) && Channels.IsValidIndex(IndexX);
+						const bool MaskY = EnumHasAllFlags(ChannelMask, EMovieSceneTransformChannel::ScaleY) && Channels.IsValidIndex(IndexY);
+						const bool MaskZ = EnumHasAllFlags(ChannelMask, EMovieSceneTransformChannel::ScaleZ) && Channels.IsValidIndex(IndexZ);
+
+						if (!(MaskX | MaskY | MaskZ))
+						{
+							// TODO: report warning
+							continue;
+						}
+
+						FVector DefaultScale = RelativeTransform.GetScale3D();
 						TArray<FGLTFVector3> Scales;
 						Scales.AddUninitialized(FrameCount);
 
@@ -324,10 +372,10 @@ void FGLTFDelayedLevelSequenceTask::Process()
 						{
 							FFrameTime& FrameTime = FrameTimes[Frame];
 
-							FVector Scale = FVector::OneVector;
-							Channels[6]->Evaluate(FrameTime, Scale.X);
-							Channels[7]->Evaluate(FrameTime, Scale.Y);
-							Channels[8]->Evaluate(FrameTime, Scale.Z);
+							FVector Scale = DefaultScale;
+							if (MaskX) Channels[IndexX]->Evaluate(FrameTime, Scale.X);
+							if (MaskY) Channels[IndexY]->Evaluate(FrameTime, Scale.Y);
+							if (MaskZ) Channels[IndexZ]->Evaluate(FrameTime, Scale.Z);
 
 							Scales[Frame] = FGLTFCoreUtilities::ConvertScale(Scale);
 						}
