[GLTFExporter] refactor: add GetReferencerName to FGLTFObjectArrayScopeGuard (needed for UE5)

commit 30a3b6a4b9daf3e74764a116520381978663c640
Date: Thu, 17 Mar 2022 17:33:11 +0100

    refactor: add GetReferencerName to FGLTFObjectArrayScopeGuard (needed for UE5)
    (cherry picked from commit b4cfd497868fe66a2d7b329cda0c4e4d4fa57e46​)

#skipci

diff --git a/Engine/Plugins/Enterprise/GLTFExporter/MigratedFromGit/Source/GLTFExporter/Private/Converters/GLTFObjectArrayScopeGuard.h b/Engine/Plugins/Enterprise/GLTFExporter/MigratedFromGit/Source/GLTFExporter/Private/Converters/GLTFObjectArrayScopeGuard.h
--- a/Engine/Plugins/Enterprise/GLTFExporter/MigratedFromGit/Source/GLTFExporter/Private/Converters/GLTFObjectArrayScopeGuard.h
+++ b/Engine/Plugins/Enterprise/GLTFExporter/MigratedFromGit/Source/GLTFExporter/Private/Converters/GLTFObjectArrayScopeGuard.h
@@ -22,4 +22,9 @@ public:
 	{
 		Collector.AddReferencedObjects(*this);
 	}
+
+	virtual FString GetReferencerName() const override
+	{
+		return TEXT("FGLTFObjectArrayScopeGuard");
+	}
 };
