[GLTFExporter] refactor: use GEditor->GetSelectedActors() instead of GSelectedActorAnnotation (needed for UE5)

commit 7ce4e559171c2f78bff55e8cf822548495bb9dc2
Date: Thu, 17 Mar 2022 17:36:10 +0100

    refactor: use GEditor->GetSelectedActors() instead of GSelectedActorAnnotation (needed for UE5)
    (cherry picked from commit 06afb1453428bbfbdac7cb868220e5abe688c03c​)

#skipci

diff --git a/Engine/Plugins/Enterprise/GLTFExporter/MigratedFromGit/Source/GLTFExporter/Private/Exporters/GLTFExporterUtility.cpp b/Engine/Plugins/Enterprise/GLTFExporter/MigratedFromGit/Source/GLTFExporter/Private/Exporters/GLTFExporterUtility.cpp
--- a/Engine/Plugins/Enterprise/GLTFExporter/MigratedFromGit/Source/GLTFExporter/Private/Exporters/GLTFExporterUtility.cpp
+++ b/Engine/Plugins/Enterprise/GLTFExporter/MigratedFromGit/Source/GLTFExporter/Private/Exporters/GLTFExporterUtility.cpp
@@ -3,18 +3,17 @@
 #include "Exporters/GLTFExporterUtility.h"
 #include "Materials/MaterialInstance.h"
 #include "AssetRegistryModule.h"
+#include "Engine/Selection.h"
 
 void FGLTFExporterUtility::GetSelectedActors(TSet<AActor*>& OutSelectedActors)
 {
 #if WITH_EDITOR
-	const TMap<const UObjectBase*, FBoolAnnotation>& AnnotationMap = ((const FUObjectAnnotationSparse<FBoolAnnotation, true>&)GSelectedActorAnnotation).GetAnnotationMap();
-	OutSelectedActors.Reserve(AnnotationMap.Num());
-
-	for (const TPair<const UObjectBase*, FBoolAnnotation>& Pair : AnnotationMap)
+	USelection* Selection = GEditor->GetSelectedActors();
+	for (FSelectionIterator SelectedObject(*Selection); SelectedObject; ++SelectedObject)
 	{
-		if (Pair.Value.Mark)
+		if (AActor* SelectedActor = Cast<AActor>(*SelectedObject))
 		{
-			OutSelectedActors.Add(static_cast<AActor*>(const_cast<UObjectBase*>(Pair.Key)));
+			OutSelectedActors.Add(SelectedActor);
 		}
 	}
 #endif
