[GLTFExporter] refactor: minor tweaks in FGLTFTextureUtility (needed for UE5)

commit f49da322c3ebe747f050120216b52d194aa9598a
Date: Thu, 17 Mar 2022 17:39:45 +0100

    refactor: minor tweaks in FGLTFTextureUtility (needed for UE5)
    (cherry picked from commit 60034c57783e962262d60968fe7d76e3c339a9cd​)

#skipci

diff --git a/Engine/Plugins/Enterprise/GLTFExporter/MigratedFromGit/Source/GLTFExporter/Private/Converters/GLTFTextureUtility.cpp b/Engine/Plugins/Enterprise/GLTFExporter/MigratedFromGit/Source/GLTFExporter/Private/Converters/GLTFTextureUtility.cpp
--- a/Engine/Plugins/Enterprise/GLTFExporter/MigratedFromGit/Source/GLTFExporter/Private/Converters/GLTFTextureUtility.cpp
+++ b/Engine/Plugins/Enterprise/GLTFExporter/MigratedFromGit/Source/GLTFExporter/Private/Converters/GLTFTextureUtility.cpp
@@ -5,6 +5,7 @@
 #include "Converters/GLTFSimpleTexture2DPreview.h"
 #include "Converters/GLTFCubemapFacePreview.h"
 #include "CanvasItem.h"
+#include "CanvasTypes.h"
 #include "Engine/Texture2D.h"
 #include "Engine/TextureCube.h"
 #include "Engine/TextureRenderTarget2D.h"
@@ -356,9 +357,9 @@ void FGLTFTextureUtility::EncodeRGBM(const TArray<FFloat16Color>& InPixels, TArr
 		const FFloat16Color& Color = InPixels[Index];
 		FLinearColor RGBM;
 
-		RGBM.R = FMath::Sqrt(Color.R);
-		RGBM.G = FMath::Sqrt(Color.G);
-		RGBM.B = FMath::Sqrt(Color.B);
+		RGBM.R = FMath::Sqrt(Color.R.GetFloat());
+		RGBM.G = FMath::Sqrt(Color.G.GetFloat());
+		RGBM.B = FMath::Sqrt(Color.B.GetFloat());
 
 		RGBM.R /= MaxRange;
 		RGBM.G /= MaxRange;
