[GLTFExporter] fix: fix incorrect casing in module dependency name (needed for UE5)

commit 289bd6fa3c8886524e4541ca53ca2651271edd71
Date: Thu, 17 Mar 2022 17:42:05 +0100

    fix: fix incorrect casing in module dependency name (needed for UE5)
    (cherry picked from commit bb182f48bbc1ca86cad74193faaec118cebb45fa​)

#skipci

diff --git a/Engine/Plugins/Enterprise/GLTFExporter/MigratedFromGit/Source/GLTFExporter/GLTFExporter.Build.cs b/Engine/Plugins/Enterprise/GLTFExporter/MigratedFromGit/Source/GLTFExporter/GLTFExporter.Build.cs
--- a/Engine/Plugins/Enterprise/GLTFExporter/MigratedFromGit/Source/GLTFExporter/GLTFExporter.Build.cs
+++ b/Engine/Plugins/Enterprise/GLTFExporter/MigratedFromGit/Source/GLTFExporter/GLTFExporter.Build.cs
@@ -50,7 +50,7 @@ public class GLTFExporter : ModuleRules
 					"MessageLog",
 					"Slate",
 					"SlateCore",
-					"Mainframe",
+					"MainFrame",
 					"InputCore",
 					"EditorStyle",
 					"PropertyEditor",
