[GLTFExporter] fix: add custom json-print-policies to properly handle UTF8 encoding (needed for UE5)

commit f3b75bc29d651c64d025b177afcd02a6dd3c2f55
Date: Thu, 17 Mar 2022 18:07:15 +0100

    fix: add custom json-print-policies to properly handle UTF8 encoding (needed for UE5)
    (cherry picked from commit 5966102b92d3367284f4aa37a81dc03523973103​)

#skipci

diff --git a/Engine/Plugins/Enterprise/GLTFExporter/MigratedFromGit/Source/GLTFExporter/Private/Json/GLTFJsonWriter.cpp b/Engine/Plugins/Enterprise/GLTFExporter/MigratedFromGit/Source/GLTFExporter/Private/Json/GLTFJsonWriter.cpp
--- a/Engine/Plugins/Enterprise/GLTFExporter/MigratedFromGit/Source/GLTFExporter/Private/Json/GLTFJsonWriter.cpp
+++ b/Engine/Plugins/Enterprise/GLTFExporter/MigratedFromGit/Source/GLTFExporter/Private/Json/GLTFJsonWriter.cpp
@@ -149,19 +149,63 @@ private:
 	TSharedRef<TJsonWriter<CharType, PrintPolicy>> JsonWriter;
 };
 
-template <class CharType>
-struct TCrossPlatformPrettyJsonPrintPolicy : TPrettyJsonPrintPolicy<CharType>
+struct TGLTFJsonPrintPolicyUTF8
+{
+	static void WriteChar(FArchive* Stream, UTF8CHAR Char)
+	{
+		Stream->Serialize(&Char, sizeof(Char));
+	}
+
+	static void WriteChar(FArchive* Stream, TCHAR Char)
+	{
+		ANSICHAR Buffer[8] = { 0 };
+		const int32 Length = FTCHARToUTF8_Convert::Utf8FromCodepoint(Char, Buffer, UE_ARRAY_COUNT(Buffer));
+		Stream->Serialize(&Char, Length);
+	}
+
+	static void WriteString(FArchive* Stream, const FString& String)
+	{
+		const TCHAR* CharPtr = *String;
+
+		for (int32 CharIndex = 0; CharIndex < String.Len(); ++CharIndex, ++CharPtr)
+		{
+			WriteChar(Stream, *CharPtr);
+		}
+	}
+};
+
+struct TGLTFCondensedJsonPrintPolicyUTF8 : TGLTFJsonPrintPolicyUTF8
+{
+	static void WriteLineTerminator(FArchive* Stream) {}
+	static void WriteTabs(FArchive* Stream, int32 Count) {}
+	static void WriteSpace(FArchive* Stream) {}
+};
+
+struct TGLTFPrettyJsonPrintPolicyUTF8 : TGLTFJsonPrintPolicyUTF8
 {
 	static void WriteLineTerminator(FArchive* Stream)
 	{
-		TJsonPrintPolicy<CharType>::WriteString(Stream, TEXT("\n")); // don't use system dependent line terminator
+		WriteChar(Stream, static_cast<UTF8CHAR>('\n')); // don't use system dependent line terminator
+	}
+
+	static void WriteTabs( FArchive* Stream, int32 Count )
+	{
+		for (int32 i = 0; i < Count; ++i)
+		{
+			WriteChar(Stream, static_cast<UTF8CHAR>('\t'));
+		}
+	}
+
+	static void WriteSpace( FArchive* Stream )
+	{
+		WriteChar(Stream, static_cast<UTF8CHAR>(' '));
 	}
 };
 
 TSharedRef<IGLTFJsonWriter> IGLTFJsonWriter::Create(FArchive& Archive, bool bPrettyJson, FGLTFJsonExtensions& Extensions)
 {
 	return MakeShareable(bPrettyJson
-		? static_cast<IGLTFJsonWriter*>(new TGLTFJsonWriterImpl<UTF8CHAR, TCrossPlatformPrettyJsonPrintPolicy<UTF8CHAR>>(Archive, Extensions))
-		: static_cast<IGLTFJsonWriter*>(new TGLTFJsonWriterImpl<UTF8CHAR, TCondensedJsonPrintPolicy<UTF8CHAR>>(Archive, Extensions))
+		? static_cast<IGLTFJsonWriter*>(new TGLTFJsonWriterImpl<UTF8CHAR, TGLTFPrettyJsonPrintPolicyUTF8>(Archive, Extensions))
+		: static_cast<IGLTFJsonWriter*>(new TGLTFJsonWriterImpl<UTF8CHAR, TGLTFCondensedJsonPrintPolicyUTF8>(Archive, Extensions))
 	);
 }
