[GLTFExporter] chore: update engine version in plugin and project descriptors

commit 8ea80ceeb2271a6416fecb8a5c845ab33ef753f3
Date: Thu, 17 Mar 2022 17:06:36 +0100

    chore: update engine version in plugin and project descriptors
    (cherry picked from commit 4612d5c7ba3e32b28670c981cc1b1cf6718f9fb4​)

#skipci

diff --git a/Engine/Plugins/Enterprise/GLTFExporter/MigratedFromGit/GLTFExporter.uplugin b/Engine/Plugins/Enterprise/GLTFExporter/MigratedFromGit/GLTFExporter.uplugin
--- a/Engine/Plugins/Enterprise/GLTFExporter/MigratedFromGit/GLTFExporter.uplugin
+++ b/Engine/Plugins/Enterprise/GLTFExporter/MigratedFromGit/GLTFExporter.uplugin
@@ -10,7 +10,7 @@
 	"DocsURL": "",
 	"MarketplaceURL": "",
 	"SupportURL": "",
-	"EngineVersion": "4.27.0",
+	"EngineVersion": "5.0.0",
 	"CanContainContent": true,
 	"Modules": [
 		{
