[GLTFExporter] chore: remove Win32 from plugin descriptor since UE5 no longer supports that platform

commit 2da2cf9e083275011df51e282324ee5265514133
Date: Thu, 17 Mar 2022 17:11:16 +0100

    chore: remove Win32 from plugin descriptor since UE5 no longer supports that platform
    (cherry picked from commit b5e644f2e305273aada1cdcc90e69cf890f0fb6c​)

#skipci

diff --git a/Engine/Plugins/Enterprise/GLTFExporter/MigratedFromGit/GLTFExporter.uplugin b/Engine/Plugins/Enterprise/GLTFExporter/MigratedFromGit/GLTFExporter.uplugin
--- a/Engine/Plugins/Enterprise/GLTFExporter/MigratedFromGit/GLTFExporter.uplugin
+++ b/Engine/Plugins/Enterprise/GLTFExporter/MigratedFromGit/GLTFExporter.uplugin
@@ -17,19 +17,19 @@
 			"Name": "GLTFExporter",
 			"Type": "Runtime",
 			"LoadingPhase": "PostConfigInit",
-			"WhitelistPlatforms": [ "Win64", "Win32", "Mac", "Linux" ]
+			"PlatformAllowList": [ "Win64", "Mac", "Linux" ]
 		},
 		{
 			"Name": "GLTFMaterialAnalyzer",
 			"Type": "Editor",
 			"LoadingPhase": "Default",
-			"WhitelistPlatforms": [ "Win64", "Win32", "Mac", "Linux" ]
+			"PlatformAllowList": [ "Win64", "Mac", "Linux" ]
 		},
 		{
 			"Name": "GLTFMaterialBaking",
 			"Type": "Editor",
 			"LoadingPhase": "Default",
-			"WhitelistPlatforms": [ "Win64", "Win32", "Mac", "Linux" ]
+			"PlatformAllowList": [ "Win64", "Mac", "Linux" ]
 		}
 	],
 	"Plugins": [
