[GLTFExporter] feat: merge in UE5 changes for baking module into plugin fork of module

commit feacaec230eb5264c200a78b7417250901d3568c
Date: Thu, 17 Mar 2022 17:14:09 +0100

    feat: merge in UE5 changes for baking module into plugin fork of module
    (cherry picked from commit 02af8ebf1d88be7c338fe80771aa4d12ae64a504​)

#skipci

diff --git a/Engine/Plugins/Enterprise/GLTFExporter/MigratedFromGit/Source/GLTFExporter/Private/Converters/GLTFMaterialUtility.cpp b/Engine/Plugins/Enterprise/GLTFExporter/MigratedFromGit/Source/GLTFExporter/Private/Converters/GLTFMaterialUtility.cpp
--- a/Engine/Plugins/Enterprise/GLTFExporter/MigratedFromGit/Source/GLTFExporter/Private/Converters/GLTFMaterialUtility.cpp
+++ b/Engine/Plugins/Enterprise/GLTFExporter/MigratedFromGit/Source/GLTFExporter/Private/Converters/GLTFMaterialUtility.cpp
@@ -198,7 +198,7 @@ FGLTFPropertyBakeOutput FGLTFMaterialUtility::BakeMaterialProperty(const FIntPoi
 	MeshSet.MaterialIndices = MeshSectionIndices; // NOTE: MaterialIndices is actually section indices
 	if (MeshData != nullptr)
 	{
-		MeshSet.RawMeshDescription = const_cast<FMeshDescription*>(&MeshData->Description);
+		MeshSet.MeshDescription = const_cast<FMeshDescription*>(&MeshData->Description);
 		MeshSet.LightMap = MeshData->LightMap;
 		MeshSet.LightMapIndex = MeshData->LightMapTexCoord;
 		MeshSet.LightmapResourceCluster = MeshData->LightMapResourceCluster;
diff --git a/Engine/Plugins/Enterprise/GLTFExporter/MigratedFromGit/Source/GLTFExporter/Private/Converters/GLTFUVOverlapChecker.cpp b/Engine/Plugins/Enterprise/GLTFExporter/MigratedFromGit/Source/GLTFExporter/Private/Converters/GLTFUVOverlapChecker.cpp
--- a/Engine/Plugins/Enterprise/GLTFExporter/MigratedFromGit/Source/GLTFExporter/Private/Converters/GLTFUVOverlapChecker.cpp
+++ b/Engine/Plugins/Enterprise/GLTFExporter/MigratedFromGit/Source/GLTFExporter/Private/Converters/GLTFUVOverlapChecker.cpp
@@ -47,7 +47,7 @@ float FGLTFUVOverlapChecker::Convert(const FMeshDescription* Description, FGLTFI
 		FGLTFMeshRenderData MeshSet;
 		MeshSet.TextureCoordinateBox = { { 0.0f, 0.0f }, { 1.0f, 1.0f } };
 		MeshSet.TextureCoordinateIndex = TexCoord;
-		MeshSet.RawMeshDescription = const_cast<FMeshDescription*>(Description);
+		MeshSet.MeshDescription = const_cast<FMeshDescription*>(Description);
 		MeshSet.MaterialIndices = SectionIndices; // NOTE: MaterialIndices is actually section indices
 
 		FGLTFMaterialDataEx MatSet;
diff --git a/Engine/Plugins/Enterprise/GLTFExporter/MigratedFromGit/Source/GLTFMaterialBaking/Private/GLTFExportMaterialProxy.h b/Engine/Plugins/Enterprise/GLTFExporter/MigratedFromGit/Source/GLTFMaterialBaking/Private/GLTFExportMaterialProxy.h
--- a/Engine/Plugins/Enterprise/GLTFExporter/MigratedFromGit/Source/GLTFMaterialBaking/Private/GLTFExportMaterialProxy.h
+++ b/Engine/Plugins/Enterprise/GLTFExporter/MigratedFromGit/Source/GLTFMaterialBaking/Private/GLTFExportMaterialProxy.h
@@ -4,6 +4,7 @@
 
 #include "MaterialShared.h"
 #include "MaterialCompiler.h"
+#include "TextureCompiler.h"
 #include "Materials/MaterialParameterCollection.h"
 
 #include "Engine/TextureLODSettings.h"
@@ -38,6 +39,11 @@ struct FGLTFExportMaterialCompiler : public FProxyMaterialCompiler
 		return MSM_MAX;
 	}
 
+	virtual FMaterialShadingModelField GetCompiledShadingModels() const override
+	{
+		return MSM_MAX;
+	}
+
 	virtual int32 WorldPosition(EWorldPositionIncludedOffsets WorldPositionIncludedOffsets) override
 	{
 #if WITH_EDITOR
@@ -90,9 +96,7 @@ struct FGLTFExportMaterialCompiler : public FProxyMaterialCompiler
 
 	virtual int32 CameraVector() override
 	{
-		// NOTE: by using VertexNormal instead of a fixed (world-space) vector, we allow
-		// more correct baking of materials that use the dot-product between the vertex normal
-		// and the camera vector for effects such as fresnel.
+		// By returning vertex normal instead of a constant vector (like up), we ensure materials (with fresnel for example) are more correctly baked using custom mesh data.
 		return Compiler->VertexNormal();
 	}
 
@@ -101,22 +105,22 @@ struct FGLTFExportMaterialCompiler : public FProxyMaterialCompiler
 		if (CustomWorldNormal == INDEX_NONE)
 		{
 			return INDEX_NONE;
-		}
+	}
 
 		int32 N = CustomWorldNormal;
 		int32 C = CameraVector();
 
 		if (bNormalizeCustomWorldNormal)
-		{
+	{
 			// N = N / sqrt(dot(N, N))
 			N = Compiler->Div(N, Compiler->SquareRoot(Compiler->Dot(N, N)));
-		}
+	}
 
 		// return 2 * dot(N, C) * N - C
 		return Compiler->Sub(Compiler->Mul(Compiler->Constant(2.0f), Compiler->Mul(Compiler->Dot(N, C), N)), C);
 	}
 
-#if (ENGINE_MAJOR_VERSION > 4 || ENGINE_MINOR_VERSION >= 26)
+#if (ENGINE_MAJOR_VERSION == 4 && ENGINE_MINOR_VERSION >= 26)
 	virtual int32 PreSkinVertexOffset() override
 	{
 		return Compiler->PreSkinVertexOffset();
@@ -212,6 +216,7 @@ class FGLTFExportMaterialProxy : public FMaterial, public FMaterialRenderProxy
 public:
 	FGLTFExportMaterialProxy(UMaterialInterface* InMaterialInterface, EMaterialProperty InPropertyToCompile, const FString& InCustomOutputToCompile = TEXT(""), bool bInSynchronousCompilation = true, EBlendMode ProxyBlendMode = BLEND_Opaque, bool bTangentSpaceNormal = false)
 		: FMaterial()
+		, FMaterialRenderProxy(GetPathNameSafe(InMaterialInterface->GetMaterial()))
 		, MaterialInterface(InMaterialInterface)
 		, PropertyToCompile(InPropertyToCompile)
 		, CustomOutputToCompile(InCustomOutputToCompile)
@@ -317,7 +322,7 @@ public:
 		return bCorrectVertexFactory && bPCPlatform && bCorrectFrequency;
 	}
 
-	virtual TArrayView<UObject* const> GetReferencedTextures() const override
+	virtual TArrayView<const TObjectPtr<UObject>> GetReferencedTextures() const override
 	{
 		return ReferencedTextures;
 	}
@@ -362,24 +367,9 @@ public:
 	}
 #endif
 
-	virtual bool GetVectorValue(const FHashedMaterialParameterInfo& ParameterInfo, FLinearColor* OutValue, const FMaterialRenderContext& Context) const override
-	{
-		return MaterialInterface->GetRenderProxy()->GetVectorValue(ParameterInfo, OutValue, Context);
-	}
-
-	virtual bool GetScalarValue(const FHashedMaterialParameterInfo& ParameterInfo, float* OutValue, const FMaterialRenderContext& Context) const override
+	virtual bool GetParameterValue(EMaterialParameterType Type, const FHashedMaterialParameterInfo& ParameterInfo, FMaterialParameterValue& OutValue, const FMaterialRenderContext& Context) const override
 	{
-		return MaterialInterface->GetRenderProxy()->GetScalarValue(ParameterInfo, OutValue, Context);
-	}
-
-	virtual bool GetTextureValue(const FHashedMaterialParameterInfo& ParameterInfo, const UTexture** OutValue, const FMaterialRenderContext& Context) const override
-	{
-		return MaterialInterface->GetRenderProxy()->GetTextureValue(ParameterInfo, OutValue, Context);
-	}
-
-	virtual bool GetTextureValue(const FHashedMaterialParameterInfo& ParameterInfo, const URuntimeVirtualTexture** OutValue, const FMaterialRenderContext& Context) const override
-	{
-		return MaterialInterface->GetRenderProxy()->GetTextureValue(ParameterInfo, OutValue, Context);
+		return MaterialInterface->GetRenderProxy()->GetParameterValue(Type, ParameterInfo, OutValue, Context);
 	}
 
 	// Material properties.
@@ -423,9 +413,9 @@ public:
 
 			case MP_Opacity:
 			case MP_OpacityMask:
-			case MP_SubsurfaceColor:
 			case MP_CustomData0:
 			case MP_CustomData1:
+			case MP_SubsurfaceColor:
 			{
 				return MaterialInterface->CompileProperty(&ProxyCompiler, PropertyToCompile, ForceCast_Exact_Replicate);
 			}
@@ -435,7 +425,7 @@ public:
 				if (BlendMode == BLEND_Opaque || BlendMode == BLEND_Masked)
 				{
 					return CompileNormalEncoding(
-						&ProxyCompiler,
+						Compiler,
 						CompileNormalTransform(&ProxyCompiler, MaterialInterface->CompileProperty(&ProxyCompiler, PropertyToCompile, ForceCast_Exact_Replicate)));
 				}
 				break;
@@ -473,6 +463,11 @@ public:
 		{
 			return MaterialInterface->CompileProperty(Compiler, MP_ShadingModel);
 
+		}
+		else if (Property == MP_FrontMaterial)
+		{
+			return MaterialInterface->CompileProperty(Compiler, MP_FrontMaterial);
+
 		}
 		else
 		{
@@ -564,94 +559,9 @@ public:
 		return Ar << V.MaterialInterface;
 	}
 
-	/**
-	* Iterate through all textures used by the material and return the maximum texture resolution used
-	* (ideally this could be made dependent of the material property)
-	*
-	* @param MaterialInterface The material to scan for texture size
-	*
-	* @return Size (width and height)
-	*/
-	FIntPoint FindMaxTextureSize(UMaterialInterface* InMaterialInterface, FIntPoint MinimumSize = FIntPoint(1, 1)) const
-	{
-		// static lod settings so that we only initialize them once
-		const UTextureLODSettings* GameTextureLODSettings = UDeviceProfileManager::Get().GetActiveProfile()->GetTextureLODSettings();
-
-		TArray<UTexture*> MaterialTextures;
-		InMaterialInterface->GetUsedTextures(MaterialTextures, EMaterialQualityLevel::Num, false, GMaxRHIFeatureLevel, false);
-
-		// find the largest texture in the list (applying it's LOD bias)
-		FIntPoint MaxSize = MinimumSize;
-		for (int32 TexIndex = 0; TexIndex < MaterialTextures.Num(); TexIndex++)
-		{
-			const UTexture* Texture = MaterialTextures[TexIndex];
-
-			if (Texture == NULL)
-			{
-				continue;
-			}
-
-			// get the max size of the texture
-			const FIntPoint LocalSize = [&]()
-			{
-				if (Texture->IsA(UTexture2D::StaticClass()))
-				{
-					UTexture2D* Tex2D = (UTexture2D*)Texture;
-					return FIntPoint(Tex2D->GetSizeX(), Tex2D->GetSizeY());
-				}
-				else if (Texture->IsA(UTextureCube::StaticClass()))
-				{
-					UTextureCube* TexCube = (UTextureCube*)Texture;
-					return FIntPoint(TexCube->GetSizeX(), TexCube->GetSizeY());
-				}
-				else if (Texture->IsA(UTexture2DArray::StaticClass())) 
-				{
-					UTexture2DArray* TexArray = (UTexture2DArray*)Texture;
-					return FIntPoint(TexArray->GetSizeX(), TexArray->GetSizeY());
-				}
-				return FIntPoint(0, 0);
-			}();
-			
-			// bias the texture size based on LOD group
-			const int32 LocalBias = GameTextureLODSettings->CalculateLODBias(Texture);
-			MaxSize.X = FMath::Max(LocalSize.X >> LocalBias, MaxSize.X);
-			MaxSize.Y = FMath::Max(LocalSize.Y >> LocalBias, MaxSize.Y);
-		}
-
-		return MaxSize;
-	}
-
-	static bool WillFillData(EBlendMode InBlendMode, EMaterialProperty InMaterialProperty)
-	{
-		if (InMaterialProperty == MP_EmissiveColor)
-		{
-			return true;
-		}
-
-		switch (InBlendMode)
-		{
-		case BLEND_Opaque:
-		{
-			switch (InMaterialProperty)
-			{
-			case MP_BaseColor:				return true;
-			case MP_Specular:				return true;
-			case MP_Normal:					return true;
-			case MP_Tangent:				return true;
-			case MP_Metallic:				return true;
-			case MP_Roughness:				return true;
-			case MP_Anisotropy:				return true;
-			case MP_AmbientOcclusion:		return true;
-			}
-		}
-		break;
-		}
-		return false;
-	}
-
 	virtual bool IsUsedWithStaticLighting() const override
 	{
-		return true; 
+		return true;
 	}
 
 	virtual void GatherExpressionsForCustomInterpolators(TArray<UMaterialExpression*>& OutExpressions) const override
@@ -751,7 +661,7 @@ private:
 				const char* Buffer = static_cast<const char*>(static_cast<const void*>(*Code));
 				const uint64 Hash = CityHash64(Buffer, Code.Len() * sizeof(TCHAR));
 				const int32 CodeIndex = CurrentScopeChunks->Num();
-				new(*CurrentScopeChunks) FShaderCodeChunk(Hash, *Code, TEXT(""), MCT_Float1, true);
+				new(*CurrentScopeChunks) FShaderCodeChunk(Hash, *Code, *Code, TEXT(""), MCT_Float1, EDerivativeStatus::NotAware, true);
 				return CodeIndex;
 			}
 		};
@@ -767,7 +677,7 @@ private:
 	/** The material interface for this proxy */
 	UMaterialInterface* MaterialInterface;
 	UMaterial* Material;
-	TArray<UObject*> ReferencedTextures;
+	TArray<TObjectPtr<UObject>> ReferencedTextures;
 	/** The property to compile for rendering the sample */
 	EMaterialProperty PropertyToCompile;
 	/** The name of the specific custom output to compile for rendering the sample. Only used if PropertyToCompile is MP_CustomOutput */
diff --git a/Engine/Plugins/Enterprise/GLTFExporter/MigratedFromGit/Source/GLTFMaterialBaking/Private/GLTFMaterialBakingHelpers.cpp b/Engine/Plugins/Enterprise/GLTFExporter/MigratedFromGit/Source/GLTFMaterialBaking/Private/GLTFMaterialBakingHelpers.cpp
--- a/Engine/Plugins/Enterprise/GLTFExporter/MigratedFromGit/Source/GLTFMaterialBaking/Private/GLTFMaterialBakingHelpers.cpp
+++ b/Engine/Plugins/Enterprise/GLTFExporter/MigratedFromGit/Source/GLTFMaterialBaking/Private/GLTFMaterialBakingHelpers.cpp
@@ -15,7 +15,7 @@ namespace FMaterialBakingHelpersImpl
 		int32 CombinedColorA = 0;
 
 		const int32 BaseIndex = (Y * InImageWidth) + X;
-	
+
 		int32 SampleIndex = BaseIndex - (InImageWidth + 1);
 		for (int32 YI = 0; YI < 3; ++YI)
 		{
@@ -62,12 +62,57 @@ namespace FMaterialBakingHelpersImpl
 		return false;
 	}
 
-	void PerformUVBorderSmear(TArray<FColor>& InOutPixels, int32& ImageWidth, int32& ImageHeight, int32 MaxIterations, bool bCanShrink, uint32 BackgroundMask)
+	static void PerformShrinking(TArray<FColor>& InOutPixels, int32& ImageWidth, int32& ImageHeight, uint32 BackgroundMask)
+	{
+		FColor* Current = InOutPixels.GetData();
+
+		// We can hopefully skip the entire smearing process if there is a single
+		// non-background color in the entire array since smearing that completely
+		// would lead to a monochome output.
+		uint32 SingleColor = BackgroundMask;
+		bool   bHasSingleColor = true;
+		for (int32 Index = 0, Num = InOutPixels.Num(); Index < Num; ++Index)
+		{
+			const uint32 Color = Current[Index].DWColor();
+			if (Color != BackgroundMask)
+			{
+				if (SingleColor == BackgroundMask)
+				{
+					// This is the first time we stumble on a color other than background color, keep it.
+					SingleColor = Color;
+				}
+				// Compare with the known non-background color
+				else if (SingleColor != Color)
+				{
+					bHasSingleColor = false;
+					break;
+				}
+			}
+		}
+
+		// Shrink to a single texel
+		if (bHasSingleColor)
+		{
+			InOutPixels.SetNum(1);
+			ImageWidth = ImageHeight = 1;
+			Current = InOutPixels.GetData();
+
+			InOutPixels[0] = FColor(SingleColor);
+		}
+	}
+
+	static void PerformUVBorderSmear(TArray<FColor>& InOutPixels, int32& ImageWidth, int32& ImageHeight, int32 MaxIterations, uint32 BackgroundMask)
 	{
 		TRACE_CPUPROFILER_EVENT_SCOPE(FMaterialBakingHelpers::PerformUVBorderSmear)
 
 		check(InOutPixels.Num() == ImageWidth*ImageHeight);
 
+		// No smearing needed for a single pixel
+		if (ImageWidth * ImageHeight == 1)
+		{
+			return;
+		}
+
 		const int32 PaddedImageWidth = ImageWidth + 2;
 		const int32 PaddedImageHeight = ImageHeight + 2;
 
@@ -77,58 +122,10 @@ namespace FMaterialBakingHelpersImpl
 		TArray<uint32>* NextRowsLeft = &RowsSwap2;
 
 		//
-		// Create a copy of the image with a 1 pixel magenta border around the edge.
+		// Create a copy of the image with a 1 pixel background color border around the edge.
 		// This avoids needing to bounds check in inner loops.
 		//
 
-		FColor* Current = InOutPixels.GetData();
-
-		// We can hopefully skip the entire smearing process if there is a single
-		// non-background color in the entire array since smearing that completely
-		// would lead to a monochome output.
-		if (MaxIterations == -1)
-		{
-			uint32 SingleColor = BackgroundMask;
-			bool   bHasSingleColor = true;
-			for (int32 Index = 0, Num = InOutPixels.Num(); Index < Num; ++Index)
-			{
-				const uint32 Color = Current[Index].DWColor();
-				if (Color != BackgroundMask)
-				{
-					if (SingleColor == BackgroundMask)
-					{
-						// This is the first time we stumble on a color other than background color, keep it.
-						SingleColor = Color;
-					}
-					// Compare with the known non-background color
-					else if (SingleColor != Color)
-					{
-						bHasSingleColor = false;
-						break;
-					}
-				}
-			}
-
-			if (bHasSingleColor)
-			{
-				// Shrink to single texel when allowed
-				if (bCanShrink)
-				{
-					InOutPixels.SetNum(1);
-					ImageWidth = ImageHeight = 1;
-					Current = InOutPixels.GetData();
-				}
-
-				FColor Color(SingleColor);
-				for (int32 Index = 0, Num = InOutPixels.Num(); Index < Num; ++Index)
-				{
-					Current[Index] = Color;
-				}
-
-				return;
-			}
-		}
-
 		TArray<FColor> ScratchBuffer;
 		ScratchBuffer.SetNumUninitialized(PaddedImageWidth*PaddedImageHeight);
 		FColor* Scratch = ScratchBuffer.GetData();
@@ -140,7 +137,9 @@ namespace FMaterialBakingHelpersImpl
 			Scratch[(PaddedImageHeight - 1)*PaddedImageWidth + X] = FColor(BackgroundMask);
 		}
 
-		// Set leftg and right border pixels to background color and copy image data into rows
+		FColor* Current = InOutPixels.GetData();
+
+		// Set left and right border pixels to background color and copy image data into rows
 		for (int32 Y = 1; Y <= ImageHeight; ++Y)
 		{
 			const int32 YOffset = Y * PaddedImageWidth;
@@ -277,7 +276,7 @@ namespace FMaterialBakingHelpersImpl
 						PreviousY = Y - 1;
 						NextRowsLeft->Add(Y - 1);
 					}
-					if (!RowCompleted[Y] && RowRemainingPixels[RowIndex] && PreviousY < Y)
+					if (PreviousY < Y && !RowCompleted[Y] && RowRemainingPixels[RowIndex])
 					{
 						PreviousY = Y;
 						NextRowsLeft->Add(Y);
@@ -291,7 +290,7 @@ namespace FMaterialBakingHelpersImpl
 			}
 
 			// Swap rows left buffers
-			{			
+			{
 				TArray<uint32>* Temp = NextRowsLeft;
 				NextRowsLeft = CurrentRowsLeft;
 				CurrentRowsLeft = Temp;
@@ -301,13 +300,17 @@ namespace FMaterialBakingHelpersImpl
 		}
 
 		// If we finished before replacing all pixels, replace the remaining pixels with black.
-		if (CurrentRowsLeft->Num() > 0)
+		for (int32 Index = 0; Index < CurrentRowsLeft->Num(); Index++)
 		{
-			for (int32 i = 0; i < InOutPixels.Num(); ++i)
+			const int32 Y = (*CurrentRowsLeft)[Index];
+
+			int32 PixelIndex = (Y - 1) * ImageWidth;
+			for (int32 X = 1; X <= ImageWidth; X++)
 			{
-				if (InOutPixels[i].DWColor() == BackgroundMask)
+				FColor& Color = Current[PixelIndex++];
+				if (Color.DWColor() == BackgroundMask)
 				{
-					InOutPixels[i] = FColor::Black;
+					Color = FColor::Black;
 				}
 			}
 		}
@@ -317,11 +320,14 @@ namespace FMaterialBakingHelpersImpl
 void FGLTFMaterialBakingHelpers::PerformUVBorderSmear(TArray<FColor>& InOutPixels, int32 ImageWidth, int32 ImageHeight, int32 MaxIterations, FColor BackgroundColor)
 {
 	const uint32 BackgroundMask = BackgroundColor.DWColor();
-	FMaterialBakingHelpersImpl::PerformUVBorderSmear(InOutPixels, ImageWidth, ImageHeight, MaxIterations, false, BackgroundMask);
+	FMaterialBakingHelpersImpl::PerformUVBorderSmear(InOutPixels, ImageWidth, ImageHeight, MaxIterations, BackgroundMask);
 }
 
 void FGLTFMaterialBakingHelpers::PerformUVBorderSmearAndShrink(TArray<FColor>& InOutPixels, int32& InOutImageWidth, int32& InOutImageHeight, FColor BackgroundColor)
 {
 	const uint32 BackgroundMask = BackgroundColor.DWColor();
-	FMaterialBakingHelpersImpl::PerformUVBorderSmear(InOutPixels, InOutImageWidth, InOutImageHeight, -1, true, BackgroundMask);
-}
\ No newline at end of file
+	FMaterialBakingHelpersImpl::PerformShrinking(InOutPixels, InOutImageWidth, InOutImageHeight, BackgroundMask);
+
+	const int32 DefaultMaxIteration = -1; // Perform smearing over the whole image
+	FMaterialBakingHelpersImpl::PerformUVBorderSmear(InOutPixels, InOutImageWidth, InOutImageHeight, DefaultMaxIteration, BackgroundMask);
+}
diff --git a/Engine/Plugins/Enterprise/GLTFExporter/MigratedFromGit/Source/GLTFMaterialBaking/Private/GLTFMaterialBakingModule.cpp b/Engine/Plugins/Enterprise/GLTFExporter/MigratedFromGit/Source/GLTFMaterialBaking/Private/GLTFMaterialBakingModule.cpp
--- a/Engine/Plugins/Enterprise/GLTFExporter/MigratedFromGit/Source/GLTFMaterialBaking/Private/GLTFMaterialBakingModule.cpp
+++ b/Engine/Plugins/Enterprise/GLTFExporter/MigratedFromGit/Source/GLTFMaterialBaking/Private/GLTFMaterialBakingModule.cpp
@@ -17,12 +17,16 @@
 #include "RHISurfaceDataConversion.h"
 #include "Misc/ScopedSlowTask.h"
 #include "MeshDescription.h"
+#include "TextureCompiler.h"
+#include "RenderCaptureInterface.h"
 #if WITH_EDITOR
 #include "Misc/FileHelper.h"
 #endif
 
 IMPLEMENT_MODULE(FGLTFMaterialBakingModule, GLTFMaterialBaking);
 
+DEFINE_LOG_CATEGORY_STATIC(LogMaterialBaking, All, All);
+
 #define LOCTEXT_NAMESPACE "MaterialBakingModule"
 
 /** Cvars for advanced features */
@@ -31,7 +35,7 @@ static TAutoConsoleVariable<int32> CVarUseMaterialProxyCaching(
 	1,
 	TEXT("Determines whether or not Material Proxies should be cached to speed up material baking.\n")
 	TEXT("0: Turned Off\n")
-	TEXT("1: Turned On"),	
+	TEXT("1: Turned On"),
 	ECVF_Default);
 
 static TAutoConsoleVariable<int32> CVarSaveIntermediateTextures(
@@ -42,6 +46,14 @@ static TAutoConsoleVariable<int32> CVarSaveIntermediateTextures(
 	TEXT("1: Turned On"),
 	ECVF_Default);
 
+static TAutoConsoleVariable<int32> CVarMaterialBakingRDOCCapture(
+	TEXT("GLTFMaterialBaking.RenderDocCapture"),
+	0,
+	TEXT("Determines whether or not to trigger a RenderDoc capture.\n")
+	TEXT("0: Turned Off\n")
+	TEXT("1: Turned On"),
+	ECVF_Default);
+
 namespace FMaterialBakingModuleImpl
 {
 	// Custom dynamic mesh allocator specifically tailored for Material Baking.
@@ -54,8 +66,8 @@ namespace FMaterialBakingModuleImpl
 		// memory is first initialized.
 		const uint32 SmallestPooledBufferSize = 256*1024;
 
-		TArray<FIndexBufferRHIRef>  IndexBuffers;
-		TArray<FVertexBufferRHIRef> VertexBuffers;
+		TArray<FBufferRHIRef>  IndexBuffers;
+		TArray<FBufferRHIRef> VertexBuffers;
 
 		template <typename RefType>
 		RefType GetSmallestFit(uint32 SizeInBytes, TArray<RefType>& Array)
@@ -83,12 +95,12 @@ namespace FMaterialBakingModuleImpl
 			return Ref;
 		}
 
-		virtual FIndexBufferRHIRef AllocIndexBuffer(uint32 NumElements) override
+		virtual FBufferRHIRef AllocIndexBuffer(uint32 NumElements) override
 		{
 			uint32 BufferSize = GetIndexBufferSize(NumElements);
 			if (BufferSize > SmallestPooledBufferSize)
 			{
-				FIndexBufferRHIRef Ref = GetSmallestFit(GetIndexBufferSize(NumElements), IndexBuffers);
+				FBufferRHIRef Ref = GetSmallestFit(GetIndexBufferSize(NumElements), IndexBuffers);
 				if (Ref.IsValid())
 				{
 					return Ref;
@@ -98,7 +110,7 @@ namespace FMaterialBakingModuleImpl
 			return FDynamicMeshBufferAllocator::AllocIndexBuffer(NumElements);
 		}
 
-		virtual void ReleaseIndexBuffer(FIndexBufferRHIRef& IndexBufferRHI) override
+		virtual void ReleaseIndexBuffer(FBufferRHIRef& IndexBufferRHI) override
 		{
 			if (IndexBufferRHI->GetSize() > SmallestPooledBufferSize)
 			{
@@ -108,12 +120,12 @@ namespace FMaterialBakingModuleImpl
 			IndexBufferRHI = nullptr;
 		}
 
-		virtual FVertexBufferRHIRef AllocVertexBuffer(uint32 Stride, uint32 NumElements) override
+		virtual FBufferRHIRef AllocVertexBuffer(uint32 Stride, uint32 NumElements) override
 		{
 			uint32 BufferSize = GetVertexBufferSize(Stride, NumElements);
 			if (BufferSize > SmallestPooledBufferSize)
 			{
-				FVertexBufferRHIRef Ref = GetSmallestFit(BufferSize, VertexBuffers);
+				FBufferRHIRef Ref = GetSmallestFit(BufferSize, VertexBuffers);
 				if (Ref.IsValid())
 				{
 					return Ref;
@@ -123,7 +135,7 @@ namespace FMaterialBakingModuleImpl
 			return FDynamicMeshBufferAllocator::AllocVertexBuffer(Stride, NumElements);
 		}
 
-		virtual void ReleaseVertexBuffer(FVertexBufferRHIRef& VertexBufferRHI) override
+		virtual void ReleaseVertexBuffer(FBufferRHIRef& VertexBufferRHI) override
 		{
 			if (VertexBufferRHI->GetSize() > SmallestPooledBufferSize)
 			{
@@ -137,14 +149,14 @@ namespace FMaterialBakingModuleImpl
 	class FStagingBufferPool
 	{
 	public:
-		FTexture2DRHIRef CreateStagingBuffer_RenderThread(FRHICommandListImmediate& RHICmdList, int32 Width, int32 Height, EPixelFormat Format)
+		FTexture2DRHIRef CreateStagingBuffer_RenderThread(FRHICommandListImmediate& RHICmdList, int32 Width, int32 Height, EPixelFormat Format, bool bIsSRGB)
 		{
 			TRACE_CPUPROFILER_EVENT_SCOPE(CreateStagingBuffer_RenderThread)
 
-			auto StagingBufferPredicate = 
-				[Width, Height, Format](const FTexture2DRHIRef& Texture2DRHIRef)
+			auto StagingBufferPredicate =
+				[Width, Height, Format, bIsSRGB](const FTexture2DRHIRef& Texture2DRHIRef)
 				{
-					return Texture2DRHIRef->GetSizeX() == Width && Texture2DRHIRef->GetSizeY() == Height && Texture2DRHIRef->GetFormat() == Format;
+					return Texture2DRHIRef->GetSizeX() == Width && Texture2DRHIRef->GetSizeY() == Height && Texture2DRHIRef->GetFormat() == Format && bool(Texture2DRHIRef->GetFlags() & TexCreate_SRGB) == bIsSRGB;
 				};
 
 			// Process any staging buffers available for unmapping
@@ -173,8 +185,14 @@ namespace FMaterialBakingModuleImpl
 			}
 
 			TRACE_CPUPROFILER_EVENT_SCOPE(RHICreateTexture2D)
-			FRHIResourceCreateInfo CreateInfo;
-			return RHICreateTexture2D(Width, Height, Format, 1, 1, TexCreate_CPUReadback, CreateInfo);
+			FRHIResourceCreateInfo CreateInfo(TEXT("FStagingBufferPool_StagingBuffer"));
+			ETextureCreateFlags TextureCreateFlags = TexCreate_CPUReadback;
+			if (bIsSRGB)
+			{
+				TextureCreateFlags |= TexCreate_SRGB;
+			}
+
+			return RHICreateTexture2D(Width, Height, Format, 1, 1, TextureCreateFlags, CreateInfo);
 		}
 
 		void ReleaseStagingBufferForUnmap_AnyThread(FTexture2DRHIRef& Texture2DRHIRef)
@@ -238,7 +256,7 @@ void FGLTFMaterialBakingModule::StartupModule()
 	bEmissiveHDR = false;
 
 	// Set which properties should enforce gamma correction
-	SetLinearBake(false);
+	SetLinearBake(true);
 
 	// Set which pixel format should be used for the possible baked out material properties
 	PerPropertyFormat.Add(MP_EmissiveColor, PF_FloatRGBA);
@@ -263,6 +281,9 @@ void FGLTFMaterialBakingModule::StartupModule()
 	// Register callback for modified objects
 	FCoreUObjectDelegates::OnObjectModified.AddRaw(this, &FGLTFMaterialBakingModule::OnObjectModified);
 
+	// Register callback on garbage collection
+	FCoreUObjectDelegates::GetPreGarbageCollectDelegate().AddRaw(this, &FGLTFMaterialBakingModule::OnPreGarbageCollect);
+
 	// NOTE: Because the attribute definition of TransmittanceColor is not registered by the engine, we have to do it here
 	AddCustomAttributeHack(FGuid(0xF2D8C70E, 0x42ECA0D1, 0x4652D0AD, 0xB785A065), "TransmittanceColor", "GetThinTranslucentMaterialOutput", MCT_Float3, FVector4(0.5, 0.5, 0.5, 0));
 }
@@ -276,7 +297,11 @@ void FGLTFMaterialBakingModule::ShutdownModule()
 	{
 		PropertyEditorModule->UnregisterCustomPropertyTypeLayout(TEXT("PropertyEntry"));
 	}
+
 	FCoreUObjectDelegates::OnObjectModified.RemoveAll(this);
+	FCoreUObjectDelegates::GetPreGarbageCollectDelegate().RemoveAll(this);
+
+	CleanupMaterialProxies();
 }
 
 void FGLTFMaterialBakingModule::BakeMaterials(const TArray<FGLTFMaterialData*>& MaterialSettings, const TArray<FGLTFMeshRenderData*>& MeshSettings, TArray<FGLTFBakeOutput>& Output)
@@ -334,33 +359,19 @@ void FGLTFMaterialBakingModule::BakeMaterials(const TArray<FGLTFMaterialData*>&
 	}
 }
 
-#if !(ENGINE_MAJOR_VERSION > 4 || ENGINE_MINOR_VERSION >= 27)
-static inline void ConvertRawR16G16B16A16FDataToFFloat16Color(uint32 Width, uint32 Height, uint8* In, uint32 SrcPitch, FFloat16Color* Out)
+void FGLTFMaterialBakingModule::BakeMaterials(const TArray<FGLTFMaterialDataEx*>& MaterialSettings, const TArray<FGLTFMeshRenderData*>& MeshSettings, TArray<FGLTFBakeOutputEx>& Output)
 {
-	const uint32 DstPitch = Width * sizeof(FFloat16Color);
-
-	// If source & dest pitch matches, perform a single memcpy.
-	if (DstPitch == SrcPitch)
-	{
-		FPlatformMemory::Memcpy(Out, In, Width * Height * sizeof(FFloat16Color));
-	}
-	else
+	UE_LOG(LogMaterialBaking, Verbose, TEXT("Performing material baking for %d materials"), MaterialSettings.Num());
+	for (int32 i = 0; i < MaterialSettings.Num(); i++)
 	{
-		check(SrcPitch > DstPitch);
-
-		// Need to copy row wise since the Pitch does not match the Width.
-		for (uint32 Y = 0; Y < Height; Y++)
+		if (MaterialSettings[i]->Material && MeshSettings[i]->MeshDescription)
 		{
-			FFloat16Color* SrcPtr = (FFloat16Color*)(In + Y * SrcPitch);
-			FFloat16Color* DestPtr = Out + Y * Width;
-			FMemory::Memcpy(DestPtr, SrcPtr, DstPitch);
+			UE_LOG(LogMaterialBaking, Verbose, TEXT("    [%5d] Material: %-50s Vertices: %8d    Triangles: %8d"), i, *MaterialSettings[i]->Material->GetName(), MeshSettings[i]->MeshDescription->Vertices().Num(), MeshSettings[i]->MeshDescription->Triangles().Num());
 		}
 	}
-}
-#endif
 
-void FGLTFMaterialBakingModule::BakeMaterials(const TArray<FGLTFMaterialDataEx*>& MaterialSettings, const TArray<FGLTFMeshRenderData*>& MeshSettings, TArray<FGLTFBakeOutputEx>& Output)
-{
+	RenderCaptureInterface::FScopedCapture RenderCapture(CVarMaterialBakingRDOCCapture.GetValueOnAnyThread() == 1, TEXT("MaterialBaking"));
+
 	TRACE_CPUPROFILER_EVENT_SCOPE(FMaterialBakingModule::BakeMaterials)
 
 	checkf(MaterialSettings.Num() == MeshSettings.Num(), TEXT("Number of material settings does not match that of MeshSettings"));
@@ -386,7 +397,7 @@ void FGLTFMaterialBakingModule::BakeMaterials(const TArray<FGLTFMaterialDataEx*>
 	// Soft page faults are now incredibly expensive on Windows 10.
 	Algo::SortBy(
 		ProcessingOrder,
-		[&MeshSettings](const uint32 Index){ return MeshSettings[Index]->RawMeshDescription ? MeshSettings[Index]->RawMeshDescription->Vertices().Num() : 0; },
+		[&MeshSettings](const uint32 Index){ return MeshSettings[Index]->MeshDescription ? MeshSettings[Index]->MeshDescription->Vertices().Num() : 0; },
 		TGreater<>()
 	);
 
@@ -409,7 +420,7 @@ void FGLTFMaterialBakingModule::BakeMaterials(const TArray<FGLTFMaterialDataEx*>
 		[&](int32 MaterialIndex)
 	{
 		TRACE_CPUPROFILER_EVENT_SCOPE(PrepareRenderItems);
-		
+
 		TMap<FMaterialBakingModuleImpl::FRenderItemKey, FGLTFMeshMaterialRenderItem*>* RenderItems = new TMap<FRenderItemKey, FGLTFMeshMaterialRenderItem *>();
 		const FGLTFMaterialDataEx* CurrentMaterialSettings = MaterialSettings[MaterialIndex];
 		const FGLTFMeshRenderData* CurrentMeshSettings = MeshSettings[MaterialIndex];
@@ -431,7 +442,7 @@ void FGLTFMaterialBakingModule::BakeMaterials(const TArray<FGLTFMaterialDataEx*>
 	TFuture<TMap<FRenderItemKey, FGLTFMeshMaterialRenderItem*>*> PreparedRenderItems[PipelineDepth];
 	for (; NextRenderItem < NumMaterials && NextRenderItem < PipelineDepth; ++NextRenderItem)
 	{
-		PreparedRenderItems[NextRenderItem] = 
+		PreparedRenderItems[NextRenderItem] =
 			Async(
 				EAsyncExecution::ThreadPool,
 				[&PrepareRenderItems_AnyThread, &ProcessingOrder, NextRenderItem]()
@@ -457,6 +468,8 @@ void FGLTFMaterialBakingModule::BakeMaterials(const TArray<FGLTFMaterialDataEx*>
 			{
 				TRACE_CPUPROFILER_EVENT_SCOPE(LoadTexturesForMaterial)
 
+				FTextureCompilingManager::Get().FinishCompilation(MaterialTextures);
+
 				for (UTexture* Texture : MaterialTextures)
 				{
 					if (Texture != NULL)
@@ -500,7 +513,7 @@ void FGLTFMaterialBakingModule::BakeMaterials(const TArray<FGLTFMaterialDataEx*>
 		if (NextRenderItem < NumMaterials)
 		{
 			check((NextRenderItem % PipelineDepth) == (Index % PipelineDepth));
-			PreparedRenderItems[NextRenderItem % PipelineDepth] = 
+			PreparedRenderItems[NextRenderItem % PipelineDepth] =
 				Async(
 					EAsyncExecution::ThreadPool,
 					[&PrepareRenderItems_AnyThread, NextMaterialIndex = ProcessingOrder[NextRenderItem]]()
@@ -512,8 +525,11 @@ void FGLTFMaterialBakingModule::BakeMaterials(const TArray<FGLTFMaterialDataEx*>
 		}
 
 		const FGLTFMaterialDataEx* CurrentMaterialSettings = MaterialSettings[MaterialIndex];
+		const FGLTFMeshRenderData* CurrentMeshSettings = MeshSettings[MaterialIndex];
 		FGLTFBakeOutputEx& CurrentOutput = Output[MaterialIndex];
 
+		TRACE_CPUPROFILER_EVENT_SCOPE_TEXT(*CurrentMaterialSettings->Material->GetName())
+
 		TArray<FGLTFMaterialPropertyEx> MaterialPropertiesToBakeOut;
 		CurrentMaterialSettings->PropertySizes.GenerateKeyArray(MaterialPropertiesToBakeOut);
 
@@ -537,6 +553,9 @@ void FGLTFMaterialBakingModule::BakeMaterials(const TArray<FGLTFMaterialDataEx*>
 			for (int32 PropertyIndex = 0; PropertyIndex < NumPropertiesToRender; ++PropertyIndex)
 			{
 				const FGLTFMaterialPropertyEx& Property = MaterialPropertiesToBakeOut[PropertyIndex];
+
+				TRACE_CPUPROFILER_EVENT_SCOPE_TEXT(*Property.ToString())
+
 				FGLTFExportMaterialProxy* ExportMaterialProxy = CreateMaterialProxy(CurrentMaterialSettings, Property);
 
 				if (!ExportMaterialProxy->IsCompilationFinished())
@@ -568,14 +587,14 @@ void FGLTFMaterialBakingModule::BakeMaterials(const TArray<FGLTFMaterialDataEx*>
 
 							FSceneViewFamily ViewFamily(FSceneViewFamily::ConstructionValues(RenderTarget->GetRenderTargetResource(), nullptr,
 								FEngineShowFlags(ESFIM_Game))
-								.SetWorldTimes(0.0f, 0.0f, 0.0f)
+								.SetTime(FGameTime())
 								.SetGammaCorrection(RenderTarget->GetRenderTargetResource()->GetDisplayGamma()));
 
 							RenderItem.MaterialRenderProxy = ExportMaterialProxy;
 							RenderItem.ViewFamily = &ViewFamily;
 
 							FTextureRenderTargetResource* RenderTargetResource = RenderTarget->GetRenderTargetResource();
-							FCanvas Canvas(RenderTargetResource, nullptr, FApp::GetCurrentTime() - GStartTime, FApp::GetDeltaTime(), FApp::GetCurrentTime() - GStartTime, GMaxRHIFeatureLevel);
+							FCanvas Canvas(RenderTargetResource, nullptr, FGameTime::GetTimeSinceAppStart(), GMaxRHIFeatureLevel);
 							Canvas.SetAllowedModes(FCanvas::Allow_Flush);
 							Canvas.SetRenderTargetRect(FIntRect(0, 0, RenderTarget->GetSurfaceWidth(), RenderTarget->GetSurfaceHeight()));
 							Canvas.SetBaseTransform(Canvas.CalcBaseTransform2D(RenderTarget->GetSurfaceWidth(), RenderTarget->GetSurfaceHeight()));
@@ -587,11 +606,11 @@ void FGLTFMaterialBakingModule::BakeMaterials(const TArray<FGLTFMaterialDataEx*>
 							Canvas.Flush_RenderThread(RHICmdList);
 							SortElement.RenderBatchArray.Empty();
 
-							FTexture2DRHIRef StagingBufferRef = StagingBufferPool.CreateStagingBuffer_RenderThread(RHICmdList, RenderTargetResource->GetSizeX(), RenderTargetResource->GetSizeY(), RenderTarget->GetFormat());
+							FTexture2DRHIRef StagingBufferRef = StagingBufferPool.CreateStagingBuffer_RenderThread(RHICmdList, RenderTargetResource->GetSizeX(), RenderTargetResource->GetSizeY(), RenderTarget->GetFormat(), RenderTarget->IsSRGB());
 							FGPUFenceRHIRef GPUFence = RHICreateGPUFence(TEXT("MaterialBackingFence"));
 
 							FResolveRect Rect(0, 0, RenderTargetResource->GetSizeX(), RenderTargetResource->GetSizeY());
-							RHICmdList.CopyToResolveTarget(RenderTargetResource->GetRenderTargetTexture(), StagingBufferRef, FResolveParams(Rect));	
+							RHICmdList.CopyToResolveTarget(RenderTargetResource->GetRenderTargetTexture(), StagingBufferRef, FResolveParams(Rect));
 							RHICmdList.WriteGPUFence(GPUFence);
 
 							// Prepare a lambda for final processing that will be executed asynchronously
@@ -622,7 +641,7 @@ void FGLTFMaterialBakingModule::BakeMaterials(const TArray<FGLTFMaterialDataEx*>
 									else
 									{
 										TRACE_CPUPROFILER_EVENT_SCOPE(ConvertRawB8G8R8A8DataToFColor)
-										
+
 										check(StagingBuffer->GetFormat() == PF_B8G8R8A8);
 										ConvertRawB8G8R8A8DataToFColor(OutputSize.X, OutputSize.Y, (uint8*)Data, DataWidth * sizeof(FColor), OutputColor.GetData());
 									}
@@ -779,11 +798,20 @@ bool FGLTFMaterialBakingModule::IsLinearBake(FGLTFMaterialPropertyEx Property)
 	return ColorSpace == EPropertyColorSpace::Linear;
 }
 
+static void DeleteCachedMaterialProxy(FGLTFExportMaterialProxy* Proxy)
+{
+	ENQUEUE_RENDER_COMMAND(DeleteCachedMaterialProxy)(
+		[Proxy](FRHICommandListImmediate& RHICmdList)
+		{
+			delete Proxy;
+		});
+}
+
 void FGLTFMaterialBakingModule::CleanupMaterialProxies()
 {
 	for (auto Iterator : MaterialProxyPool)
 	{
-		delete Iterator.Value.Value;
+		DeleteCachedMaterialProxy(Iterator.Value.Value);
 	}
 	MaterialProxyPool.Reset();
 }
@@ -802,7 +830,7 @@ UTextureRenderTarget2D* FGLTFMaterialBakingModule::CreateRenderTarget(bool bInFo
 
 	// Find any pooled render target with suitable properties.
 	UTextureRenderTarget2D** FindResult = RenderTargetPool.FindByPredicate(RenderTargetComparison);
-	
+
 	if (FindResult)
 	{
 		RenderTarget = *FindResult;
@@ -873,7 +901,7 @@ void FGLTFMaterialBakingModule::ProcessEmissiveOutput(const FFloat16Color* Color
 	// Find maximum float value across texture
 	ParallelFor(NumThreads, [&Color16, LinesPerThread, MaxValue, OutputSize, Color16Pitch, BackgroundFloat16](int32 Index)
 	{
-		const int32 EndY = FMath::Min((Index + 1) * LinesPerThread, OutputSize.Y);			
+		const int32 EndY = FMath::Min((Index + 1) * LinesPerThread, OutputSize.Y);
 		float& CurrentMaxValue = MaxValue[Index];
 		for (int32 PixelY = Index * LinesPerThread; PixelY < EndY; ++PixelY)
 		{
@@ -900,10 +928,10 @@ void FGLTFMaterialBakingModule::ProcessEmissiveOutput(const FFloat16Color* Color
 
 		return TempValue;
 	}();
-		
+
 	if (GlobalMaxValue <= 0.01f)
 	{
-		// Black emissive, drop it			
+		// Black emissive, drop it
 	}
 
 	// Now convert Float16 to Color using the scale
@@ -932,7 +960,7 @@ void FGLTFMaterialBakingModule::ProcessEmissiveOutput(const FFloat16Color* Color
 					Pixel8.G = (uint8)FMath::RoundToInt(Pixel16.G.GetFloat() * Scale);
 					Pixel8.B = (uint8)FMath::RoundToInt(Pixel16.B.GetFloat() * Scale);
 				}
-					
+
 				Pixel8.A = 255;
 			}
 		}
@@ -987,14 +1015,7 @@ void FGLTFMaterialBakingModule::OnObjectModified(UObject* Object)
 				// We have a match, remove the entry from our pool
 				if (bMustDelete)
 				{
-					FGLTFExportMaterialProxy* Proxy = It.Value().Value;
-
-					ENQUEUE_RENDER_COMMAND(DeleteCachedMaterialProxy)(
-						[Proxy](FRHICommandListImmediate& RHICmdList)
-						{
-							delete Proxy;
-						});
-
+					DeleteCachedMaterialProxy(It.Value().Value);
 					It.RemoveCurrent();
 				}
 			}
@@ -1002,4 +1023,9 @@ void FGLTFMaterialBakingModule::OnObjectModified(UObject* Object)
 	}
 }
 
+void FGLTFMaterialBakingModule::OnPreGarbageCollect()
+{
+	CleanupMaterialProxies();
+}
+
 #undef LOCTEXT_NAMESPACE //"MaterialBakingModule"
diff --git a/Engine/Plugins/Enterprise/GLTFExporter/MigratedFromGit/Source/GLTFMaterialBaking/Private/GLTFMaterialBakingModule.h b/Engine/Plugins/Enterprise/GLTFExporter/MigratedFromGit/Source/GLTFMaterialBaking/Private/GLTFMaterialBakingModule.h
--- a/Engine/Plugins/Enterprise/GLTFExporter/MigratedFromGit/Source/GLTFMaterialBaking/Private/GLTFMaterialBakingModule.h
+++ b/Engine/Plugins/Enterprise/GLTFExporter/MigratedFromGit/Source/GLTFMaterialBaking/Private/GLTFMaterialBakingModule.h
@@ -57,6 +57,10 @@ protected:
 
 	/** Callback for modified objects which should be removed from MaterialProxyPool in that case */
 	void OnObjectModified(UObject* Object);
+
+	/** Callback used to clear material proxy cache on garbage collection */
+	void OnPreGarbageCollect();
+
 private:
 	enum EPropertyColorSpace
 	{
diff --git a/Engine/Plugins/Enterprise/GLTFExporter/MigratedFromGit/Source/GLTFMaterialBaking/Private/MaterialRenderItem.cpp b/Engine/Plugins/Enterprise/GLTFExporter/MigratedFromGit/Source/GLTFMaterialBaking/Private/MaterialRenderItem.cpp
--- a/Engine/Plugins/Enterprise/GLTFExporter/MigratedFromGit/Source/GLTFMaterialBaking/Private/MaterialRenderItem.cpp
+++ b/Engine/Plugins/Enterprise/GLTFExporter/MigratedFromGit/Source/GLTFMaterialBaking/Private/MaterialRenderItem.cpp
@@ -7,6 +7,7 @@
 #include "StaticMeshAttributes.h"
 #include "DynamicMeshBuilder.h"
 #include "MeshPassProcessor.h"
+#include "CanvasRender.h"
 
 #define SHOW_WIREFRAME_MESH 0
 
@@ -25,7 +26,7 @@ FGLTFMeshMaterialRenderItem::FGLTFMeshMaterialRenderItem(
 	LCI = new FGLTFMeshRenderInfo(InMeshSettings->LightMap, nullptr, nullptr, InMeshSettings->LightmapResourceCluster);
 }
 
-bool FGLTFMeshMaterialRenderItem::Render_RenderThread(FRHICommandListImmediate& RHICmdList, FMeshPassProcessorRenderState& DrawRenderState, const FCanvas* Canvas)
+bool FGLTFMeshMaterialRenderItem::Render_RenderThread(FCanvasRenderContext& RenderContext, FMeshPassProcessorRenderState& DrawRenderState, const FCanvas* Canvas)
 {
 	checkSlow(ViewFamily && MeshSettings && MaterialRenderProxy);
 	// current render target set for the canvas
@@ -57,20 +58,20 @@ bool FGLTFMeshMaterialRenderItem::Render_RenderThread(FRHICommandListImmediate&
 		LocalDrawRenderState.SetBlendState(TStaticBlendState<CW_RGBA>::GetRHI());
 		LocalDrawRenderState.SetDepthStencilState(TStaticDepthStencilState<false, CF_Always>::GetRHI());
 
-		QueueMaterial(RHICmdList, LocalDrawRenderState, &View);
+		QueueMaterial(RenderContext, LocalDrawRenderState, &View);
 	}
 
 	return true;
 }
 
-bool FGLTFMeshMaterialRenderItem::Render_GameThread(const FCanvas* Canvas, FRenderThreadScope& RenderScope)
+bool FGLTFMeshMaterialRenderItem::Render_GameThread(const FCanvas* Canvas, FCanvasRenderThreadScope& RenderScope)
 {
 	RenderScope.EnqueueRenderCommand(
-		[this, Canvas](FRHICommandListImmediate& RHICmdList)
+		[this, Canvas](FCanvasRenderContext& RenderContext)
 		{
 			// Render_RenderThread uses its own render state
 			FMeshPassProcessorRenderState DummyRenderState;
-			Render_RenderThread(RHICmdList, DummyRenderState, Canvas);
+			Render_RenderThread(RenderContext, DummyRenderState, Canvas);
 		}
 	);
 
@@ -84,7 +85,7 @@ void FGLTFMeshMaterialRenderItem::GenerateRenderData()
 	// Reset array without resizing
 	Vertices.SetNum(0, false);
 	Indices.SetNum(0, false);
-	if (MeshSettings->RawMeshDescription)
+	if (MeshSettings->MeshDescription)
 	{
 		// Use supplied FMeshDescription data to populate render data
 		PopulateWithMeshData();
@@ -106,7 +107,7 @@ FGLTFMeshMaterialRenderItem::~FGLTFMeshMaterialRenderItem()
 	);
 }
 
-void FGLTFMeshMaterialRenderItem::QueueMaterial(FRHICommandListImmediate& RHICmdList, FMeshPassProcessorRenderState& DrawRenderState, const FSceneView* View)
+void FGLTFMeshMaterialRenderItem::QueueMaterial(FCanvasRenderContext& RenderContext, FMeshPassProcessorRenderState& DrawRenderState, const FSceneView* View)
 {
 	TRACE_CPUPROFILER_EVENT_SCOPE(FMeshMaterialRenderItem::QueueMaterial)
 
@@ -126,26 +127,17 @@ void FGLTFMeshMaterialRenderItem::QueueMaterial(FRHICommandListImmediate& RHICmd
 		const FGLTFPrimitiveData* PrimitiveData = MeshSettings->PrimitiveData;
 		if (PrimitiveData != nullptr)
 		{
-			FPrimitiveUniformShaderParameters PrimitiveParams = GetPrimitiveUniformShaderParameters(
-				FMatrix::Identity,
-				FMatrix::Identity,
-				PrimitiveData->ActorPosition,
-				PrimitiveData->ObjectBounds,
-				PrimitiveData->ObjectLocalBounds,
-				PrimitiveData->PreSkinnedLocalBounds,
-				false,
-				false,
-				false,
-				false,
-				false,
-				/* bDrawsVelocity = */ false,
-				GetDefaultLightingChannelMask(),
-				0,
-				INDEX_NONE,
-				INDEX_NONE,
-				false,
-				PrimitiveData->CustomPrimitiveData
-			);
+			FPrimitiveUniformShaderParameters PrimitiveParams = FPrimitiveUniformShaderParametersBuilder{}
+				.Defaults()
+					.LocalToWorld(FMatrix::Identity)
+					.ActorWorldPosition(PrimitiveData->ActorPosition)
+					.WorldBounds(PrimitiveData->ObjectBounds)
+					.LocalBounds(PrimitiveData->ObjectLocalBounds)
+					.PreSkinnedLocalBounds(PrimitiveData->PreSkinnedLocalBounds)
+					.CustomPrimitiveData(PrimitiveData->CustomPrimitiveData)
+					.ReceivesDecals(false)
+					.DrawsVelocity(false)
+				.Build();
 
 			// Overwrite object orientation, since original calculation is derived from LocalToWorld matrix,
 			// which currently needs to be identity matrix to avoid transforming the "2D" vertices used for baking
@@ -167,7 +159,7 @@ void FGLTFMeshMaterialRenderItem::QueueMaterial(FRHICommandListImmediate& RHICmd
 	MeshElement.bWireframe = true;
 #endif
 
-	const int32 NumTris = FMath::TruncToInt(Indices.Num() / 3);
+	const int32 NumTris = FMath::TruncToInt((float)Indices.Num() / 3);
 	if (NumTris == 0)
 	{
 		// there's nothing to do here
@@ -175,7 +167,7 @@ void FGLTFMeshMaterialRenderItem::QueueMaterial(FRHICommandListImmediate& RHICmd
 	}
 
 	// Bake the material out to a tile
-	GetRendererModule().DrawTileMesh(RHICmdList, DrawRenderState, *View, MeshElement, false /*bIsHitTesting*/, FHitProxyId());
+	GetRendererModule().DrawTileMesh(RenderContext, DrawRenderState, *View, MeshElement, false /*bIsHitTesting*/, FHitProxyId());
 }
 
 void FGLTFMeshMaterialRenderItem::PopulateWithQuadData()
@@ -197,12 +189,12 @@ void FGLTFMeshMaterialRenderItem::PopulateWithQuadData()
 		const int32 X = VertIndex & 1;
 		const int32 Y = (VertIndex >> 1) & 1;
 		Vert->Position.Set(ScaleX * X, ScaleY * Y, 0);
-		Vert->SetTangents(FVector(1, 0, 0), FVector(0, 1, 0), FVector(0, 0, 1));
+		Vert->SetTangents(FVector3f(1, 0, 0), FVector3f(0, 1, 0), FVector3f(0, 0, 1));
 		FMemory::Memzero(&Vert->TextureCoordinate, sizeof(Vert->TextureCoordinate));
 		for (int32 TexcoordIndex = 0; TexcoordIndex < MAX_STATIC_TEXCOORDS; TexcoordIndex++)
 		{
 			Vert->TextureCoordinate[TexcoordIndex].Set(U + SizeU * X, V + SizeV * Y);
-		}		
+		}
 		Vert->Color = FColor::White;
 	}
 
@@ -213,14 +205,16 @@ void FGLTFMeshMaterialRenderItem::PopulateWithQuadData()
 
 void FGLTFMeshMaterialRenderItem::PopulateWithMeshData()
 {
-	const FMeshDescription* RawMesh = MeshSettings->RawMeshDescription;
-
-	TVertexAttributesConstRef<FVector> VertexPositions = RawMesh->VertexAttributes().GetAttributesRef<FVector>(MeshAttribute::Vertex::Position);
-	TVertexInstanceAttributesConstRef<FVector> VertexInstanceNormals = RawMesh->VertexInstanceAttributes().GetAttributesRef<FVector>(MeshAttribute::VertexInstance::Normal);
-	TVertexInstanceAttributesConstRef<FVector> VertexInstanceTangents = RawMesh->VertexInstanceAttributes().GetAttributesRef<FVector>(MeshAttribute::VertexInstance::Tangent);
-	TVertexInstanceAttributesConstRef<float> VertexInstanceBinormalSigns = RawMesh->VertexInstanceAttributes().GetAttributesRef<float>(MeshAttribute::VertexInstance::BinormalSign);
-	TVertexInstanceAttributesConstRef<FVector2D> VertexInstanceUVs = RawMesh->VertexInstanceAttributes().GetAttributesRef<FVector2D>(MeshAttribute::VertexInstance::TextureCoordinate);
-	TVertexInstanceAttributesConstRef<FVector4> VertexInstanceColors = RawMesh->VertexInstanceAttributes().GetAttributesRef<FVector4>(MeshAttribute::VertexInstance::Color);
+	const FMeshDescription* RawMesh = MeshSettings->MeshDescription;
+
+	FStaticMeshConstAttributes Attributes(*RawMesh);
+	TArrayView<const FVector3f> VertexPositions = Attributes.GetVertexPositions().GetRawArray();
+	TArrayView<const FVector3f> VertexInstanceNormals = Attributes.GetVertexInstanceNormals().GetRawArray();
+	TArrayView<const FVector3f> VertexInstanceTangents = Attributes.GetVertexInstanceTangents().GetRawArray();
+	TArrayView<const float> VertexInstanceBinormalSigns = Attributes.GetVertexInstanceBinormalSigns().GetRawArray();
+	TArrayView<const FVector4f> VertexInstanceColors = Attributes.GetVertexInstanceColors().GetRawArray();
+
+	TVertexInstanceAttributesConstRef<FVector2f> VertexInstanceUVs = Attributes.GetVertexInstanceUVs();
 	const int32 NumVerts = RawMesh->Vertices().Num();
 
 	// reserve renderer data
@@ -234,82 +228,75 @@ void FGLTFMeshMaterialRenderItem::PopulateWithMeshData()
 	// count number of texture coordinates for this mesh
 	const int32 NumTexcoords = [&]()
 	{
-		return FMath::Min(VertexInstanceUVs.GetNumIndices(), VertexPositionStoredUVChannel);
-	}();		
+		return FMath::Min(VertexInstanceUVs.GetNumChannels(), VertexPositionStoredUVChannel);
+	}();
 
 	// check if we should use NewUVs or original UV set
 	const bool bUseNewUVs = MeshSettings->CustomTextureCoordinates.Num() > 0;
 	if (bUseNewUVs)
 	{
-		check(MeshSettings->CustomTextureCoordinates.Num() == VertexInstanceUVs.GetNumElements() && VertexInstanceUVs.GetNumIndices() > MeshSettings->TextureCoordinateIndex);
+		check(MeshSettings->CustomTextureCoordinates.Num() == VertexInstanceUVs.GetNumElements() && VertexInstanceUVs.GetNumChannels() > MeshSettings->TextureCoordinateIndex);
 	}
 
 	// add vertices
 	int32 VertIndex = 0;
 	int32 FaceIndex = 0;
-	for(const FPolygonID PolygonID : RawMesh->Polygons().GetElementIDs())
+	for(const FTriangleID& TriangleID : RawMesh->Triangles().GetElementIDs())
 	{
-		const FPolygonGroupID PolygonGroupID = RawMesh->GetPolygonPolygonGroup(PolygonID);
-		const TArray<FTriangleID>& TriangleIDs = RawMesh->GetPolygonTriangleIDs(PolygonID);
-		for (const FTriangleID& TriangleID : TriangleIDs)
+		const FPolygonGroupID PolygonGroupID = RawMesh->GetTrianglePolygonGroup(TriangleID);
+		if (MeshSettings->MaterialIndices.Contains(PolygonGroupID.GetValue()))
 		{
-			if (MeshSettings->MaterialIndices.Contains(PolygonGroupID.GetValue()))
+			const int32 NUM_VERTICES = 3;
+			for (int32 Corner = 0; Corner < NUM_VERTICES; Corner++)
 			{
-				const int32 NUM_VERTICES = 3;
-				for (int32 Corner = 0; Corner < NUM_VERTICES; Corner++)
-				{
-					// Swap vertices order if mesh is mirrored
-					const int32 CornerIdx = !MeshSettings->bMirrored ? Corner : NUM_VERTICES - Corner - 1;
+				// Swap vertices order if mesh is mirrored
+				const int32 CornerIdx = !MeshSettings->bMirrored ? Corner : NUM_VERTICES - Corner - 1;
 
-					const int32 SrcVertIndex = FaceIndex * NUM_VERTICES + CornerIdx;
-					const FVertexInstanceID SrcVertexInstanceID = RawMesh->GetTriangleVertexInstance(TriangleID, Corner);
-					const FVertexID SrcVertexID = RawMesh->GetVertexInstanceVertex(SrcVertexInstanceID);
+				const int32 SrcVertIndex = FaceIndex * NUM_VERTICES + CornerIdx;
+				const FVertexInstanceID SrcVertexInstanceID = RawMesh->GetTriangleVertexInstance(TriangleID, Corner);
+				const FVertexID SrcVertexID = RawMesh->GetVertexInstanceVertex(SrcVertexInstanceID);
 
-					// add vertex
-					FDynamicMeshVertex* Vert = new(Vertices)FDynamicMeshVertex();
-					if (!bUseNewUVs)
-					{
-						// compute vertex position from original UV
-						const FVector2D& UV = VertexInstanceUVs.Get(SrcVertexInstanceID, MeshSettings->TextureCoordinateIndex);
-						Vert->Position.Set(UV.X * ScaleX, UV.Y * ScaleY, 0);
-					}
-					else
-					{
-						const FVector2D& UV = MeshSettings->CustomTextureCoordinates[SrcVertIndex];
-						Vert->Position.Set(UV.X * ScaleX, UV.Y * ScaleY, 0);
-					}
-					FVector TangentX = VertexInstanceTangents[SrcVertexInstanceID];
-					FVector TangentZ = VertexInstanceNormals[SrcVertexInstanceID];
-					FVector TangentY = FVector::CrossProduct(TangentZ, TangentX).GetSafeNormal() * VertexInstanceBinormalSigns[SrcVertexInstanceID];
-					Vert->SetTangents(TangentX, TangentY, TangentZ);
-					for (int32 TexcoordIndex = 0; TexcoordIndex < NumTexcoords; TexcoordIndex++)
-					{
-						Vert->TextureCoordinate[TexcoordIndex] = VertexInstanceUVs.Get(SrcVertexInstanceID, TexcoordIndex);
-					}
-				
-					if (NumTexcoords < VertexPositionStoredUVChannel)
+				// add vertex
+				FDynamicMeshVertex* Vert = new(Vertices)FDynamicMeshVertex();
+				if (!bUseNewUVs)
+				{
+					// compute vertex position from original UV
+					const FVector2D& UV = FVector2D(VertexInstanceUVs.Get(SrcVertexInstanceID, MeshSettings->TextureCoordinateIndex));
+					Vert->Position.Set(UV.X * ScaleX, UV.Y * ScaleY, 0);
+				}
+				else
+				{
+					const FVector2D& UV = MeshSettings->CustomTextureCoordinates[SrcVertIndex];
+					Vert->Position.Set(UV.X * ScaleX, UV.Y * ScaleY, 0);
+				}
+				FVector3f TangentX = VertexInstanceTangents[SrcVertexInstanceID];
+				FVector3f TangentZ = VertexInstanceNormals[SrcVertexInstanceID];
+				FVector3f TangentY = FVector3f::CrossProduct(TangentZ, TangentX).GetSafeNormal() * VertexInstanceBinormalSigns[SrcVertexInstanceID];
+				Vert->SetTangents(TangentX, TangentY, TangentZ);
+				for (int32 TexcoordIndex = 0; TexcoordIndex < NumTexcoords; TexcoordIndex++)
+				{
+					Vert->TextureCoordinate[TexcoordIndex] = VertexInstanceUVs.Get(SrcVertexInstanceID, TexcoordIndex);
+				}
+
+				if (NumTexcoords < VertexPositionStoredUVChannel)
+				{
+					for (int32 TexcoordIndex = NumTexcoords; TexcoordIndex < VertexPositionStoredUVChannel; TexcoordIndex++)
 					{
-						for (int32 TexcoordIndex = NumTexcoords; TexcoordIndex < VertexPositionStoredUVChannel; TexcoordIndex++)
-						{
-							Vert->TextureCoordinate[TexcoordIndex] = Vert->TextureCoordinate[FMath::Max(NumTexcoords - 1, 0)];
-						}
+						Vert->TextureCoordinate[TexcoordIndex] = Vert->TextureCoordinate[FMath::Max(NumTexcoords - 1, 0)];
 					}
-					// Store original vertex positions in texture coordinate data
-					Vert->TextureCoordinate[6].X = VertexPositions[SrcVertexID].X;
-					Vert->TextureCoordinate[6].Y = VertexPositions[SrcVertexID].Y;
-					Vert->TextureCoordinate[7].X = VertexPositions[SrcVertexID].Z;
-
-					Vert->Color = FLinearColor(VertexInstanceColors[SrcVertexInstanceID]).ToFColor(true);
-					// add index
-					Indices.Add(VertIndex);
-					VertIndex++;
 				}
-				// add the same triangle with opposite vertex order
-				Indices.Add(VertIndex - 3);
-				Indices.Add(VertIndex - 1);
-				Indices.Add(VertIndex - 2);
+				// Store original vertex positions in texture coordinate data
+				Vert->TextureCoordinate[6].X = VertexPositions[SrcVertexID].X;
+				Vert->TextureCoordinate[6].Y = VertexPositions[SrcVertexID].Y;
+				Vert->TextureCoordinate[7].X = VertexPositions[SrcVertexID].Z;
+				Vert->TextureCoordinate[7].Y = 0.0f;
+
+				Vert->Color = FLinearColor(VertexInstanceColors[SrcVertexInstanceID]).ToFColor(true);
+				// add index
+				Indices.Add(VertIndex);
+				VertIndex++;
 			}
-			FaceIndex++;
 		}
+		FaceIndex++;
 	}
 }
diff --git a/Engine/Plugins/Enterprise/GLTFExporter/MigratedFromGit/Source/GLTFMaterialBaking/Private/MaterialRenderItem.h b/Engine/Plugins/Enterprise/GLTFExporter/MigratedFromGit/Source/GLTFMaterialBaking/Private/MaterialRenderItem.h
--- a/Engine/Plugins/Enterprise/GLTFExporter/MigratedFromGit/Source/GLTFMaterialBaking/Private/MaterialRenderItem.h
+++ b/Engine/Plugins/Enterprise/GLTFExporter/MigratedFromGit/Source/GLTFMaterialBaking/Private/MaterialRenderItem.h
@@ -53,16 +53,16 @@ public:
 	virtual ~FGLTFMeshMaterialRenderItem();
 
 	/** Begin FCanvasBaseRenderItem overrides */
-	virtual bool Render_RenderThread(FRHICommandListImmediate& RHICmdList, FMeshPassProcessorRenderState& DrawRenderState, const FCanvas* Canvas) final;
-	virtual bool Render_GameThread(const FCanvas* Canvas, FRenderThreadScope& RenderScope) final;
+	virtual bool Render_RenderThread(FCanvasRenderContext& RenderContext, FMeshPassProcessorRenderState& DrawRenderState, const FCanvas* Canvas) final;
+	virtual bool Render_GameThread(const FCanvas* Canvas, FCanvasRenderThreadScope& RenderScope) final;
 	/** End FCanvasBaseRenderItem overrides */
 
 	/** Populate vertices and indices according to available mesh data and otherwise uses simple quad */
 	void GenerateRenderData();
 protected:
 	/** Enqueues the current material to be rendered */
-	void QueueMaterial(FRHICommandListImmediate& RHICmdList, FMeshPassProcessorRenderState& DrawRenderState, const FSceneView* View);
-	
+	void QueueMaterial(FCanvasRenderContext& RenderContext, FMeshPassProcessorRenderState& DrawRenderState, const FSceneView* View);
+
 	/** Helper functions to populate render data using either mesh data or a simple quad */
 	void PopulateWithQuadData();
 	void PopulateWithMeshData();
diff --git a/Engine/Plugins/Enterprise/GLTFExporter/MigratedFromGit/Source/GLTFMaterialBaking/Public/GLTFMaterialBakingStructures.h b/Engine/Plugins/Enterprise/GLTFExporter/MigratedFromGit/Source/GLTFMaterialBaking/Public/GLTFMaterialBakingStructures.h
--- a/Engine/Plugins/Enterprise/GLTFExporter/MigratedFromGit/Source/GLTFMaterialBaking/Public/GLTFMaterialBakingStructures.h
+++ b/Engine/Plugins/Enterprise/GLTFExporter/MigratedFromGit/Source/GLTFMaterialBaking/Public/GLTFMaterialBakingStructures.h
@@ -114,14 +114,14 @@ struct FGLTFPrimitiveData
 struct FGLTFMeshRenderData
 {
 	FGLTFMeshRenderData()
-		: RawMeshDescription(nullptr), Mesh(nullptr), bMirrored(false), VertexColorHash(0), TextureCoordinateIndex(0), LightMapIndex(0), LightMap(nullptr), LightmapResourceCluster(nullptr), PrimitiveData(nullptr)
+		: MeshDescription(nullptr), Mesh(nullptr), bMirrored(false), VertexColorHash(0), TextureCoordinateIndex(0), LightMapIndex(0), LightMap(nullptr), LightmapResourceCluster(nullptr), PrimitiveData(nullptr)
 	{}
 
 	/** Ptr to raw mesh data to use for baking out the material data, if nullptr a standard quad is used */
-	FMeshDescription* RawMeshDescription;
+	const FMeshDescription* MeshDescription;
 
 	/** Ptr to original static mesh this mesh data came from */
-	UStaticMesh* Mesh;
+	const UStaticMesh* Mesh;
 
 	/** Transform determinant used to detect mirroring */
 	bool bMirrored;
@@ -137,7 +137,7 @@ struct FGLTFMeshRenderData
 
 	/** Box which's space contains the UV coordinates used to bake out the material */
 	FBox2D TextureCoordinateBox;
-	
+
 	/** Specific texture coordinate index to use as texture coordinates to bake out the material (is overruled if CustomTextureCoordinates contains any data) */
 	int32 TextureCoordinateIndex;
 
