[GLTFExporter] feat: fix UE5 build issues in plugin's material analyzer module

commit 9de77a49a545928e669165b0c664552d2804727f
Date: Thu, 17 Mar 2022 17:15:36 +0100

    feat: fix UE5 build issues in plugin's material analyzer module
    (cherry picked from commit 3257773f3b4136fe99d2a09d0ac0e8c53946c08b​)

#skipci

diff --git a/Engine/Plugins/Enterprise/GLTFExporter/MigratedFromGit/Source/GLTFMaterialAnalyzer/Private/GLTFProxyMaterialCompiler.h b/Engine/Plugins/Enterprise/GLTFExporter/MigratedFromGit/Source/GLTFMaterialAnalyzer/Private/GLTFProxyMaterialCompiler.h
--- a/Engine/Plugins/Enterprise/GLTFExporter/MigratedFromGit/Source/GLTFMaterialAnalyzer/Private/GLTFProxyMaterialCompiler.h
+++ b/Engine/Plugins/Enterprise/GLTFExporter/MigratedFromGit/Source/GLTFMaterialAnalyzer/Private/GLTFProxyMaterialCompiler.h
@@ -53,7 +53,7 @@ public:
 #if (ENGINE_MAJOR_VERSION > 4 || ENGINE_MINOR_VERSION >= 26)
 	virtual int32 CustomExpression(class UMaterialExpressionCustom* Custom, int32 OutputIndex, TArray<int32>& CompiledInputs) override
 #else
-	virtual int32 CustomExpression(class UMaterialExpressionCustom* Custom, TArray<int32>& CompiledInputs) override 
+	virtual int32 CustomExpression(class UMaterialExpressionCustom* Custom, TArray<int32>& CompiledInputs) override
 #endif
 	{
 		// NOTE: because there is no overridable compiler method for ObjectLocalBounds in engine this is our only option currently
@@ -158,7 +158,7 @@ public:
 		return EMaterialCompilerType::MaterialProxy;
 	}
 
-#if (ENGINE_MAJOR_VERSION > 4 || ENGINE_MINOR_VERSION >= 26)
+#if (ENGINE_MAJOR_VERSION == 4 && ENGINE_MINOR_VERSION >= 26)
 	virtual int32 PreSkinVertexOffset() override
 	{
 		return Compiler->PreSkinVertexOffset();
