[GLTFExporter] fix: use new UE5 float-based primitives (e.g. FVector3f)

commit 356b22f6d85be2e58d573674dff4ffe27773d990
Date: Thu, 17 Mar 2022 17:20:48 +0100

    fix: use new UE5 float-based primitives (e.g. FVector3f)
    (cherry picked from commit 213307507564a831cd79757d75cb08854f4b3e1f​)

#skipci

diff --git a/Engine/Plugins/Enterprise/GLTFExporter/MigratedFromGit/Source/GLTFExporter/Private/Converters/GLTFAccessorConverters.cpp b/Engine/Plugins/Enterprise/GLTFExporter/MigratedFromGit/Source/GLTFExporter/Private/Converters/GLTFAccessorConverters.cpp
--- a/Engine/Plugins/Enterprise/GLTFExporter/MigratedFromGit/Source/GLTFExporter/Private/Converters/GLTFAccessorConverters.cpp
+++ b/Engine/Plugins/Enterprise/GLTFExporter/MigratedFromGit/Source/GLTFExporter/Private/Converters/GLTFAccessorConverters.cpp
@@ -46,7 +46,7 @@ FGLTFJsonAccessor* FGLTFPositionBufferConverter::Convert(const FGLTFMeshSection*
 	for (uint32 VertexIndex = 0; VertexIndex < VertexCount; ++VertexIndex)
 	{
 		const uint32 MappedVertexIndex = IndexMap[VertexIndex];
-		const FVector& Position = *reinterpret_cast<const FVector*>(SourceData + Stride * MappedVertexIndex);
+		const FVector3f& Position = *reinterpret_cast<const FVector3f*>(SourceData + Stride * MappedVertexIndex);
 		Positions[VertexIndex] = FGLTFCoreUtilities::ConvertPosition(Position, Builder.ExportOptions->ExportUniformScale);
 	}
 
@@ -193,9 +193,9 @@ FGLTFJsonBufferView* FGLTFNormalBufferConverter::ConvertBufferView(const FGLTFMe
 	for (uint32 VertexIndex = 0; VertexIndex < VertexCount; ++VertexIndex)
 	{
 		const uint32 MappedVertexIndex = IndexMap[VertexIndex];
-		const FVector SafeNormal = TangentData[MappedVertexIndex].TangentZ.ToFVector().GetSafeNormal();
+		const FVector3f SafeNormal = TangentData[MappedVertexIndex].TangentZ.ToFVector().GetSafeNormal();
 
-		typedef typename TConditional<TIsSame<DestinationType, FGLTFVector3>::Value, FVector, SourceType>::Type IntermediateType;
+		typedef typename TConditional<TIsSame<DestinationType, FGLTFVector3>::Value, FVector3f, SourceType>::Type IntermediateType;
 		Normals[VertexIndex] = FGLTFCoreUtilities::ConvertNormal(IntermediateType(SafeNormal));
 	}
 
@@ -272,9 +272,9 @@ FGLTFJsonBufferView* FGLTFTangentBufferConverter::ConvertBufferView(const FGLTFM
 	for (uint32 VertexIndex = 0; VertexIndex < VertexCount; ++VertexIndex)
 	{
 		const uint32 MappedVertexIndex = IndexMap[VertexIndex];
-		const FVector SafeTangent = VertexTangents[MappedVertexIndex].TangentX.ToFVector().GetSafeNormal();
+		const FVector3f SafeTangent = VertexTangents[MappedVertexIndex].TangentX.ToFVector().GetSafeNormal();
 
-		typedef typename TConditional<TIsSame<DestinationType, FGLTFVector4>::Value, FVector, SourceType>::Type IntermediateType;
+		typedef typename TConditional<TIsSame<DestinationType, FGLTFVector4>::Value, FVector3f, SourceType>::Type IntermediateType;
 		Tangents[VertexIndex] = FGLTFCoreUtilities::ConvertTangent(IntermediateType(SafeTangent));
 	}
 
diff --git a/Engine/Plugins/Enterprise/GLTFExporter/MigratedFromGit/Source/GLTFExporter/Private/Converters/GLTFCombinedTexturePreview.cpp b/Engine/Plugins/Enterprise/GLTFExporter/MigratedFromGit/Source/GLTFExporter/Private/Converters/GLTFCombinedTexturePreview.cpp
--- a/Engine/Plugins/Enterprise/GLTFExporter/MigratedFromGit/Source/GLTFExporter/Private/Converters/GLTFCombinedTexturePreview.cpp
+++ b/Engine/Plugins/Enterprise/GLTFExporter/MigratedFromGit/Source/GLTFExporter/Private/Converters/GLTFCombinedTexturePreview.cpp
@@ -43,7 +43,7 @@ public:
 	 * Set shader parameters.
 	 * @param SampleTexture - The texture to sample.
 	 */
-	void SetParameters(FRHICommandList& RHICmdList, const FTexture* InTextureA, const FTexture* InTextureB, const FMatrix& InColorTransformA, const FMatrix& InColorTransformB, const FLinearColor& InBackgroundColor)
+	void SetParameters(FRHICommandList& RHICmdList, const FTexture* InTextureA, const FTexture* InTextureB, const FMatrix44f& InColorTransformA, const FMatrix44f& InColorTransformB, const FLinearColor& InBackgroundColor)
 	{
 		FRHIPixelShader* PixelShaderRHI = RHICmdList.GetBoundPixelShader();
 		SetTextureParameter(RHICmdList, PixelShaderRHI,TextureA,TextureSamplerA,InTextureA);
diff --git a/Engine/Plugins/Enterprise/GLTFExporter/MigratedFromGit/Source/GLTFExporter/Private/Converters/GLTFCombinedTexturePreview.h b/Engine/Plugins/Enterprise/GLTFExporter/MigratedFromGit/Source/GLTFExporter/Private/Converters/GLTFCombinedTexturePreview.h
--- a/Engine/Plugins/Enterprise/GLTFExporter/MigratedFromGit/Source/GLTFExporter/Private/Converters/GLTFCombinedTexturePreview.h
+++ b/Engine/Plugins/Enterprise/GLTFExporter/MigratedFromGit/Source/GLTFExporter/Private/Converters/GLTFCombinedTexturePreview.h
@@ -22,8 +22,8 @@ private:
 
 	const FTexture* TextureA;
 	const FTexture* TextureB;
-	FMatrix ColorTransformA;
-	FMatrix ColorTransformB;
+	FMatrix44f ColorTransformA;
+	FMatrix44f ColorTransformB;
 	FLinearColor BackgroundColor;
 
 	/** Binds vertex and pixel shaders for this element */
diff --git a/Engine/Plugins/Enterprise/GLTFExporter/MigratedFromGit/Source/GLTFExporter/Private/Converters/GLTFLightMapConverters.cpp b/Engine/Plugins/Enterprise/GLTFExporter/MigratedFromGit/Source/GLTFExporter/Private/Converters/GLTFLightMapConverters.cpp
--- a/Engine/Plugins/Enterprise/GLTFExporter/MigratedFromGit/Source/GLTFExporter/Private/Converters/GLTFLightMapConverters.cpp
+++ b/Engine/Plugins/Enterprise/GLTFExporter/MigratedFromGit/Source/GLTFExporter/Private/Converters/GLTFLightMapConverters.cpp
@@ -66,10 +66,10 @@ FGLTFJsonLightMap* FGLTFLightMapConverter::Convert(const UStaticMeshComponent* S
 		return nullptr;
 	}
 
-	const FVector2D& CoordinateBias = LightMap2D->GetCoordinateBias();
-	const FVector2D& CoordinateScale = LightMap2D->GetCoordinateScale();
-	const FVector4& LightMapAdd = LightMapInteraction.GetAddArray()[0];
-	const FVector4& LightMapScale = LightMapInteraction.GetScaleArray()[0];
+	const FVector2f CoordinateBias = FVector2f(LightMap2D->GetCoordinateBias());
+	const FVector2f CoordinateScale = FVector2f(LightMap2D->GetCoordinateScale());
+	const FVector4f& LightMapAdd = LightMapInteraction.GetAddArray()[0];
+	const FVector4f& LightMapScale = LightMapInteraction.GetScaleArray()[0];
 
 	FGLTFJsonLightMap* JsonLightMap = Builder.AddLightMap();
 	StaticMeshComponent->GetName(JsonLightMap->Name); // TODO: use better name (similar to light and camera)
diff --git a/Engine/Plugins/Enterprise/GLTFExporter/MigratedFromGit/Source/GLTFExporter/Private/Converters/GLTFNodeConverters.cpp b/Engine/Plugins/Enterprise/GLTFExporter/MigratedFromGit/Source/GLTFExporter/Private/Converters/GLTFNodeConverters.cpp
--- a/Engine/Plugins/Enterprise/GLTFExporter/MigratedFromGit/Source/GLTFExporter/Private/Converters/GLTFNodeConverters.cpp
+++ b/Engine/Plugins/Enterprise/GLTFExporter/MigratedFromGit/Source/GLTFExporter/Private/Converters/GLTFNodeConverters.cpp
@@ -116,9 +116,9 @@ FGLTFJsonNode* FGLTFComponentConverter::Convert(const USceneComponent* SceneComp
 		}
 	}
 
-	const FTransform Transform = SceneComponent->GetComponentTransform();
-	const FTransform ParentTransform = ParentComponent != nullptr ? ParentComponent->GetSocketTransform(SocketName) : FTransform::Identity;
-	const FTransform RelativeTransform = bIsRootNode ? Transform : Transform.GetRelativeTransform(ParentTransform);
+	const FTransform3f Transform = FTransform3f(SceneComponent->GetComponentTransform());
+	const FTransform3f ParentTransform = ParentComponent != nullptr ? FTransform3f(ParentComponent->GetSocketTransform(SocketName)) : FTransform3f::Identity;
+	const FTransform3f RelativeTransform = bIsRootNode ? Transform : Transform.GetRelativeTransform(ParentTransform);
 
 	FGLTFJsonNode* Node = Builder.AddNode();
 	Node->Name = FGLTFNameUtility::GetName(SceneComponent);
@@ -233,9 +233,9 @@ FGLTFJsonNode* FGLTFStaticSocketConverter::Convert(FGLTFJsonNode* RootNode, cons
 	Node->Name = SocketName.ToString();
 
 	// TODO: add warning check for non-uniform scaling
-	Node->Translation = FGLTFCoreUtilities::ConvertPosition(Socket->RelativeLocation, Builder.ExportOptions->ExportUniformScale);
-	Node->Rotation = FGLTFCoreUtilities::ConvertRotation(Socket->RelativeRotation.Quaternion());
-	Node->Scale = FGLTFCoreUtilities::ConvertScale(Socket->RelativeScale);
+	Node->Translation = FGLTFCoreUtilities::ConvertPosition(FVector3f(Socket->RelativeLocation), Builder.ExportOptions->ExportUniformScale);
+	Node->Rotation = FGLTFCoreUtilities::ConvertRotation(FRotator3f(Socket->RelativeRotation).Quaternion());
+	Node->Scale = FGLTFCoreUtilities::ConvertScale(FVector3f(Socket->RelativeScale));
 
 	RootNode->Children.Add(Node);
 	return Node;
@@ -256,9 +256,9 @@ FGLTFJsonNode* FGLTFSkeletalSocketConverter::Convert(FGLTFJsonNode* RootNode, co
 		Node->Name = SocketName.ToString();
 
 		// TODO: add warning check for non-uniform scaling
-		Node->Translation = FGLTFCoreUtilities::ConvertPosition(Socket->RelativeLocation, Builder.ExportOptions->ExportUniformScale);
-		Node->Rotation = FGLTFCoreUtilities::ConvertRotation(Socket->RelativeRotation.Quaternion());
-		Node->Scale = FGLTFCoreUtilities::ConvertScale(Socket->RelativeScale);
+		Node->Translation = FGLTFCoreUtilities::ConvertPosition(FVector3f(Socket->RelativeLocation), Builder.ExportOptions->ExportUniformScale);
+		Node->Rotation = FGLTFCoreUtilities::ConvertRotation(FRotator3f(Socket->RelativeRotation).Quaternion());
+		Node->Scale = FGLTFCoreUtilities::ConvertScale(FVector3f(Socket->RelativeScale));
 
 		const int32 ParentBone = RefSkeleton.FindBoneIndex(Socket->BoneName);
 		FGLTFJsonNode* ParentNode = ParentBone != INDEX_NONE ? Builder.AddUniqueNode(RootNode, SkeletalMesh, ParentBone) : RootNode;
@@ -303,7 +303,7 @@ FGLTFJsonNode* FGLTFSkeletalBoneConverter::Convert(FGLTFJsonNode* RootNode, cons
 	if (BonePoses.IsValidIndex(BoneIndex))
 	{
 		// TODO: add warning check for non-uniform scaling
-		const FTransform& BonePose = BonePoses[BoneIndex];
+		const FTransform3f BonePose = FTransform3f(BonePoses[BoneIndex]);
 		Node->Translation = FGLTFCoreUtilities::ConvertPosition(BonePose.GetTranslation(), Builder.ExportOptions->ExportUniformScale);
 		Node->Rotation = FGLTFCoreUtilities::ConvertRotation(BonePose.GetRotation());
 		Node->Scale = FGLTFCoreUtilities::ConvertScale(BonePose.GetScale3D());
diff --git a/Engine/Plugins/Enterprise/GLTFExporter/MigratedFromGit/Source/GLTFExporter/Private/Converters/GLTFSkinConverters.cpp b/Engine/Plugins/Enterprise/GLTFExporter/MigratedFromGit/Source/GLTFExporter/Private/Converters/GLTFSkinConverters.cpp
--- a/Engine/Plugins/Enterprise/GLTFExporter/MigratedFromGit/Source/GLTFExporter/Private/Converters/GLTFSkinConverters.cpp
+++ b/Engine/Plugins/Enterprise/GLTFExporter/MigratedFromGit/Source/GLTFExporter/Private/Converters/GLTFSkinConverters.cpp
@@ -38,7 +38,7 @@ FGLTFJsonSkin* FGLTFSkinConverter::Convert(FGLTFJsonNode* RootNode, const USkele
 
 	for (int32 BoneIndex = 0; BoneIndex < BoneCount; ++BoneIndex)
 	{
-		const FTransform InverseBindTransform = FGLTFBoneUtility::GetBindTransform(RefSkeleton, BoneIndex).Inverse();
+		const FTransform3f InverseBindTransform = FTransform3f(FGLTFBoneUtility::GetBindTransform(RefSkeleton, BoneIndex).Inverse());
 		InverseBindMatrices[BoneIndex] = FGLTFCoreUtilities::ConvertTransform(InverseBindTransform, Builder.ExportOptions->ExportUniformScale);
 	}
 
diff --git a/Engine/Plugins/Enterprise/GLTFExporter/MigratedFromGit/Source/GLTFExporter/Private/Converters/GLTFUVDegenerateChecker.cpp b/Engine/Plugins/Enterprise/GLTFExporter/MigratedFromGit/Source/GLTFExporter/Private/Converters/GLTFUVDegenerateChecker.cpp
--- a/Engine/Plugins/Enterprise/GLTFExporter/MigratedFromGit/Source/GLTFExporter/Private/Converters/GLTFUVDegenerateChecker.cpp
+++ b/Engine/Plugins/Enterprise/GLTFExporter/MigratedFromGit/Source/GLTFExporter/Private/Converters/GLTFUVDegenerateChecker.cpp
@@ -23,7 +23,7 @@ void FGLTFUVDegenerateChecker::Sanitize(const FMeshDescription*& Description, FG
 
 	if (Description != nullptr)
 	{
-		const int32 TexCoordCount = Description->VertexInstanceAttributes().GetAttributesRef<FVector2D>(MeshAttribute::VertexInstance::TextureCoordinate).GetNumIndices();
+		const int32 TexCoordCount = Description->VertexInstanceAttributes().GetAttributesRef<FVector2f>(MeshAttribute::VertexInstance::TextureCoordinate).GetNumIndices();
 		if (TexCoord < 0 || TexCoord >= TexCoordCount)
 		{
 			Description = nullptr;
@@ -37,10 +37,10 @@ float FGLTFUVDegenerateChecker::Convert(const FMeshDescription* Description, FGL
 #if WITH_EDITOR
 	if (Description != nullptr)
 	{
-		const TVertexAttributesConstRef<FVector> Positions =
-		Description->VertexAttributes().GetAttributesRef<FVector>(MeshAttribute::Vertex::Position);
-		const TVertexInstanceAttributesConstRef<FVector2D> UVs =
-			Description->VertexInstanceAttributes().GetAttributesRef<FVector2D>(MeshAttribute::VertexInstance::TextureCoordinate);
+		const TVertexAttributesConstRef<FVector3f> Positions =
+		Description->VertexAttributes().GetAttributesRef<FVector3f>(MeshAttribute::Vertex::Position);
+		const TVertexInstanceAttributesConstRef<FVector2f> UVs =
+			Description->VertexInstanceAttributes().GetAttributesRef<FVector2f>(MeshAttribute::VertexInstance::TextureCoordinate);
 
 		int32 TriangleCount = 0;
 		int32 DegenerateCount = 0;
@@ -54,8 +54,8 @@ float FGLTFUVDegenerateChecker::Convert(const FMeshDescription* Description, FGL
 					const TStaticArray<FVertexID, 3> TriangleVertexIDs = Description->GetTriangleVertices(TriangleID);
 					TArrayView<const FVertexInstanceID> TriangleVertexInstanceIDs = Description->GetTriangleVertexInstances(TriangleID);
 
-					TStaticArray<FVector, 3> TrianglePositions;
-					TStaticArray<FVector2D, 3> TriangleUVs;
+					TStaticArray<FVector3f, 3> TrianglePositions;
+					TStaticArray<FVector2f, 3> TriangleUVs;
 
 					for (int32 Index = 0; Index < 3; Index++)
 					{
@@ -93,18 +93,18 @@ float FGLTFUVDegenerateChecker::Convert(const FMeshDescription* Description, FGL
 	return -1;
 }
 
-bool FGLTFUVDegenerateChecker::IsDegenerateTriangle(const TStaticArray<FVector2D, 3>& Points)
+bool FGLTFUVDegenerateChecker::IsDegenerateTriangle(const TStaticArray<FVector2f, 3>& Points)
 {
-	const FVector2D AB = Points[1] - Points[0];
-	const FVector2D AC = Points[2] - Points[0];
+	const FVector2f AB = Points[1] - Points[0];
+	const FVector2f AC = Points[2] - Points[0];
 	const float DoubleArea = FMath::Abs(AB ^ AC);
 	return DoubleArea < 2 * SMALL_NUMBER;
 }
 
-bool FGLTFUVDegenerateChecker::IsDegenerateTriangle(const TStaticArray<FVector, 3>& Points)
+bool FGLTFUVDegenerateChecker::IsDegenerateTriangle(const TStaticArray<FVector3f, 3>& Points)
 {
-	const FVector AB = Points[1] - Points[0];
-	const FVector AC = Points[2] - Points[0];
+	const FVector3f AB = Points[1] - Points[0];
+	const FVector3f AC = Points[2] - Points[0];
 	const float DoubleAreaSquared = (AB ^ AC).SizeSquared();
 	return DoubleAreaSquared < 2 * SMALL_NUMBER * SMALL_NUMBER;
 }
diff --git a/Engine/Plugins/Enterprise/GLTFExporter/MigratedFromGit/Source/GLTFExporter/Private/Converters/GLTFUVOverlapChecker.cpp b/Engine/Plugins/Enterprise/GLTFExporter/MigratedFromGit/Source/GLTFExporter/Private/Converters/GLTFUVOverlapChecker.cpp
--- a/Engine/Plugins/Enterprise/GLTFExporter/MigratedFromGit/Source/GLTFExporter/Private/Converters/GLTFUVOverlapChecker.cpp
+++ b/Engine/Plugins/Enterprise/GLTFExporter/MigratedFromGit/Source/GLTFExporter/Private/Converters/GLTFUVOverlapChecker.cpp
@@ -14,8 +14,8 @@ void FGLTFUVOverlapChecker::Sanitize(const FMeshDescription*& Description, FGLTF
 #if WITH_EDITOR
 	if (Description != nullptr)
 	{
-		const TVertexInstanceAttributesConstRef<FVector2D> VertexInstanceUVs =
-			Description->VertexInstanceAttributes().GetAttributesRef<FVector2D>(MeshAttribute::VertexInstance::TextureCoordinate);
+		const TVertexInstanceAttributesConstRef<FVector2f> VertexInstanceUVs =
+			Description->VertexInstanceAttributes().GetAttributesRef<FVector2f>(MeshAttribute::VertexInstance::TextureCoordinate);
 		const int32 TexCoordCount = VertexInstanceUVs.GetNumIndices();
 
 		if (TexCoord < 0 || TexCoord >= TexCoordCount)
diff --git a/Engine/Plugins/Enterprise/GLTFExporter/MigratedFromGit/Source/GLTFExporter/Private/Tasks/GLTFDelayedAnimationTasks.cpp b/Engine/Plugins/Enterprise/GLTFExporter/MigratedFromGit/Source/GLTFExporter/Private/Tasks/GLTFDelayedAnimationTasks.cpp
--- a/Engine/Plugins/Enterprise/GLTFExporter/MigratedFromGit/Source/GLTFExporter/Private/Tasks/GLTFDelayedAnimationTasks.cpp
+++ b/Engine/Plugins/Enterprise/GLTFExporter/MigratedFromGit/Source/GLTFExporter/Private/Tasks/GLTFDelayedAnimationTasks.cpp
@@ -61,7 +61,7 @@ void FGLTFDelayedAnimSequenceTask::Process()
 
 			for (int32 Frame = 0; Frame < FrameCount; ++Frame)
 			{
-				const FVector KeyPosition = FrameTransforms[Frame][BoneIndex].GetTranslation();
+				const FVector3f KeyPosition = FrameTransforms[Frame][BoneIndex].GetTranslation();
 				Translations[Frame] = FGLTFCoreUtilities::ConvertPosition(KeyPosition, Builder.ExportOptions->ExportUniformScale);
 			}
 
@@ -89,7 +89,7 @@ void FGLTFDelayedAnimSequenceTask::Process()
 
 			for (int32 Frame = 0; Frame < FrameCount; ++Frame)
 			{
-				const FQuat KeyRotation = FrameTransforms[Frame][BoneIndex].GetRotation();
+				const FQuat4f KeyRotation = FQuat4f(FrameTransforms[Frame][BoneIndex].GetRotation());
 				Rotations[Frame] = FGLTFCoreUtilities::ConvertRotation(KeyRotation);
 			}
 
@@ -117,7 +117,7 @@ void FGLTFDelayedAnimSequenceTask::Process()
 
 			for (int32 Frame = 0; Frame < FrameCount; ++Frame)
 			{
-				const FVector KeyScale = FrameTransforms[Frame][BoneIndex].GetScale3D();
+				const FVector3f KeyScale = FrameTransforms[Frame][BoneIndex].GetScale3D();
 				Scales[Frame] = FGLTFCoreUtilities::ConvertScale(KeyScale);
 			}
 
@@ -261,7 +261,7 @@ void FGLTFDelayedLevelSequenceTask::Process()
 							continue;
 						}
 
-						FVector DefaultTranslation = RelativeTransform.GetTranslation();
+						FVector3f DefaultTranslation = FVector3f(RelativeTransform.GetTranslation());
 						TArray<FGLTFVector3> Translations;
 						Translations.AddUninitialized(FrameCount);
 
@@ -269,7 +269,7 @@ void FGLTFDelayedLevelSequenceTask::Process()
 						{
 							FFrameTime& FrameTime = FrameTimes[Frame];
 
-							FVector Translation = DefaultTranslation;
+							FVector3f Translation = DefaultTranslation;
 							if (MaskX) Channels[IndexX]->Evaluate(FrameTime, Translation.X);
 							if (MaskY) Channels[IndexY]->Evaluate(FrameTime, Translation.Y);
 							if (MaskZ) Channels[IndexZ]->Evaluate(FrameTime, Translation.Z);
@@ -311,7 +311,7 @@ void FGLTFDelayedLevelSequenceTask::Process()
 							continue;
 						}
 
-						FRotator DefaultRotator = RelativeTransform.GetRotation().Rotator();
+						FRotator3f DefaultRotator = FRotator3f(RelativeTransform.GetRotation().Rotator());
 						TArray<FGLTFQuaternion> Rotations;
 						Rotations.AddUninitialized(FrameCount);
 
@@ -319,7 +319,7 @@ void FGLTFDelayedLevelSequenceTask::Process()
 						{
 							FFrameTime& FrameTime = FrameTimes[Frame];
 
-							FRotator Rotator = DefaultRotator;
+							FRotator3f Rotator = DefaultRotator;
 							if (MaskX) Channels[IndexX]->Evaluate(FrameTime, Rotator.Roll);
 							if (MaskY) Channels[IndexY]->Evaluate(FrameTime, Rotator.Pitch);
 							if (MaskZ) Channels[IndexZ]->Evaluate(FrameTime, Rotator.Yaw);
@@ -364,7 +364,7 @@ void FGLTFDelayedLevelSequenceTask::Process()
 							continue;
 						}
 
-						FVector DefaultScale = RelativeTransform.GetScale3D();
+						FVector3f DefaultScale = FVector3f(RelativeTransform.GetScale3D());
 						TArray<FGLTFVector3> Scales;
 						Scales.AddUninitialized(FrameCount);
 
@@ -372,7 +372,7 @@ void FGLTFDelayedLevelSequenceTask::Process()
 						{
 							FFrameTime& FrameTime = FrameTimes[Frame];
 
-							FVector Scale = DefaultScale;
+							FVector3f Scale = DefaultScale;
 							if (MaskX) Channels[IndexX]->Evaluate(FrameTime, Scale.X);
 							if (MaskY) Channels[IndexY]->Evaluate(FrameTime, Scale.Y);
 							if (MaskZ) Channels[IndexZ]->Evaluate(FrameTime, Scale.Z);
diff --git a/Engine/Plugins/Enterprise/GLTFExporter/MigratedFromGit/Source/GLTFExporter/Private/Utilities/GLTFCoreUtilities.cpp b/Engine/Plugins/Enterprise/GLTFExporter/MigratedFromGit/Source/GLTFExporter/Private/Utilities/GLTFCoreUtilities.cpp
--- a/Engine/Plugins/Enterprise/GLTFExporter/MigratedFromGit/Source/GLTFExporter/Private/Utilities/GLTFCoreUtilities.cpp
+++ b/Engine/Plugins/Enterprise/GLTFExporter/MigratedFromGit/Source/GLTFExporter/Private/Utilities/GLTFCoreUtilities.cpp
@@ -9,24 +9,24 @@ float FGLTFCoreUtilities::ConvertLength(const float Length, const float Conversi
 	return Length * ConversionScale;
 }
 
-FGLTFVector3 FGLTFCoreUtilities::ConvertVector(const FVector& Vector)
+FGLTFVector3 FGLTFCoreUtilities::ConvertVector(const FVector3f& Vector)
 {
 	// UE4 uses a left-handed coordinate system, with Z up.
 	// glTF uses a right-handed coordinate system, with Y up.
 	return { Vector.X, Vector.Z, Vector.Y };
 }
 
-FGLTFVector3 FGLTFCoreUtilities::ConvertPosition(const FVector& Position, const float ConversionScale)
+FGLTFVector3 FGLTFCoreUtilities::ConvertPosition(const FVector3f& Position, const float ConversionScale)
 {
 	return ConvertVector(Position * ConversionScale);
 }
 
-FGLTFVector3 FGLTFCoreUtilities::ConvertScale(const FVector& Scale)
+FGLTFVector3 FGLTFCoreUtilities::ConvertScale(const FVector3f& Scale)
 {
 	return ConvertVector(Scale);
 }
 
-FGLTFVector3 FGLTFCoreUtilities::ConvertNormal(const FVector& Normal)
+FGLTFVector3 FGLTFCoreUtilities::ConvertNormal(const FVector3f& Normal)
 {
 	return ConvertVector(Normal);
 }
@@ -41,7 +41,7 @@ FGLTFInt8Vector4 FGLTFCoreUtilities::ConvertNormal(const FPackedNormal& Normal)
 	return { Normal.Vector.X, Normal.Vector.Z, Normal.Vector.Y, 0 };
 }
 
-FGLTFVector4 FGLTFCoreUtilities::ConvertTangent(const FVector& Tangent)
+FGLTFVector4 FGLTFCoreUtilities::ConvertTangent(const FVector3f& Tangent)
 {
 	// glTF stores tangent as Vec4, with W component indicating handedness of tangent basis.
 	return { Tangent.X, Tangent.Z, Tangent.Y, 1.0f };
@@ -57,7 +57,7 @@ FGLTFInt8Vector4 FGLTFCoreUtilities::ConvertTangent(const FPackedNormal& Tangent
 	return { Tangent.Vector.X, Tangent.Vector.Z, Tangent.Vector.Y, MAX_int8 /* = 1.0 */ };
 }
 
-FGLTFVector2 FGLTFCoreUtilities::ConvertUV(const FVector2D& UV)
+FGLTFVector2 FGLTFCoreUtilities::ConvertUV(const FVector2f& UV)
 {
 	// No conversion actually needed, this is primarily for type-safety.
 	return { UV.X, UV.Y };
@@ -65,7 +65,7 @@ FGLTFVector2 FGLTFCoreUtilities::ConvertUV(const FVector2D& UV)
 
 FGLTFVector2 FGLTFCoreUtilities::ConvertUV(const FVector2DHalf& UV)
 {
-	return ConvertUV(FVector2D(UV));
+	return ConvertUV(FVector2f(UV));
 }
 
 FGLTFColor4 FGLTFCoreUtilities::ConvertColor(const FLinearColor& Color, bool bForceLDR)
@@ -115,12 +115,12 @@ FGLTFUInt8Color4 FGLTFCoreUtilities::ConvertColor(const FColor& Color)
 	return { Color.R, Color.G, Color.B, Color.A };
 }
 
-FGLTFQuaternion FGLTFCoreUtilities::ConvertRotation(const FRotator& Rotation)
+FGLTFQuaternion FGLTFCoreUtilities::ConvertRotation(const FRotator3f& Rotation)
 {
 	return ConvertRotation(Rotation.Quaternion());
 }
 
-FGLTFQuaternion FGLTFCoreUtilities::ConvertRotation(const FQuat& Rotation)
+FGLTFQuaternion FGLTFCoreUtilities::ConvertRotation(const FQuat4f& Rotation)
 {
 	// UE4 uses a left-handed coordinate system, with Z up.
 	// glTF uses a right-handed coordinate system, with Y up.
@@ -130,11 +130,11 @@ FGLTFQuaternion FGLTFCoreUtilities::ConvertRotation(const FQuat& Rotation)
 	// also, as handedness is changed rotation is inversed - hence negation
 	// therefore glTFRotation = (-qX, -qZ, -qY, qw)
 
-	const FQuat Normalized = Rotation.GetNormalized();
+	const FQuat4f Normalized = Rotation.GetNormalized();
 	return { -Normalized.X, -Normalized.Z, -Normalized.Y, Normalized.W };
 }
 
-FGLTFMatrix4 FGLTFCoreUtilities::ConvertMatrix(const FMatrix& Matrix)
+FGLTFMatrix4 FGLTFCoreUtilities::ConvertMatrix(const FMatrix44f& Matrix)
 {
 	// Unreal stores matrix elements in row major order.
 	// glTF stores matrix elements in column major order.
@@ -150,18 +150,18 @@ FGLTFMatrix4 FGLTFCoreUtilities::ConvertMatrix(const FMatrix& Matrix)
 	return Result;
 }
 
-FGLTFMatrix4 FGLTFCoreUtilities::ConvertTransform(const FTransform& Transform, const float ConversionScale)
+FGLTFMatrix4 FGLTFCoreUtilities::ConvertTransform(const FTransform3f& Transform, const float ConversionScale)
 {
-	const FQuat Rotation = Transform.GetRotation();
-	const FVector Translation = Transform.GetTranslation();
-	const FVector Scale = Transform.GetScale3D();
+	const FQuat4f Rotation = Transform.GetRotation();
+	const FVector3f Translation = Transform.GetTranslation();
+	const FVector3f Scale = Transform.GetScale3D();
 
-	const FQuat ConvertedRotation(-Rotation.X, -Rotation.Z, -Rotation.Y, Rotation.W);
-	const FVector ConvertedTranslation = FVector(Translation.X, Translation.Z, Translation.Y) * ConversionScale;
-	const FVector ConvertedScale(Scale.X, Scale.Z, Scale.Y);
+	const FQuat4f ConvertedRotation(-Rotation.X, -Rotation.Z, -Rotation.Y, Rotation.W);
+	const FVector3f ConvertedTranslation = FVector(Translation.X, Translation.Z, Translation.Y) * ConversionScale;
+	const FVector3f ConvertedScale(Scale.X, Scale.Z, Scale.Y);
 
-	const FTransform ConvertedTransform(ConvertedRotation, ConvertedTranslation, ConvertedScale);
-	const FMatrix Matrix = ConvertedTransform.ToMatrixWithScale(); // TODO: should it be ToMatrixNoScale()?
+	const FTransform3f ConvertedTransform(ConvertedRotation, ConvertedTranslation, ConvertedScale);
+	const FMatrix44f Matrix = ConvertedTransform.ToMatrixWithScale(); // TODO: should it be ToMatrixNoScale()?
 	return ConvertMatrix(Matrix);
 }
 
@@ -184,7 +184,7 @@ FGLTFQuaternion FGLTFCoreUtilities::GetLocalCameraRotation()
 	// Unreal uses +X axis as camera direction in Unreal coordinates.
 	// glTF uses +Y as camera direction in Unreal coordinates.
 
-	static FGLTFQuaternion Rotation = ConvertRotation(FRotator(0, 90, 0));
+	static FGLTFQuaternion Rotation = ConvertRotation(FRotator3f(0, 90, 0));
 	return Rotation;
 }
 
@@ -193,7 +193,7 @@ FGLTFQuaternion FGLTFCoreUtilities::GetLocalLightRotation()
 	// Unreal uses +X axis as light direction in Unreal coordinates.
 	// glTF uses +Y as light direction in Unreal coordinates.
 
-	static FGLTFQuaternion Rotation = ConvertRotation(FRotator(0, 90, 0));
+	static FGLTFQuaternion Rotation = ConvertRotation(FRotator3f(0, 90, 0));
 	return Rotation;
 }
 
diff --git a/Engine/Plugins/Enterprise/GLTFExporter/MigratedFromGit/Source/GLTFExporter/Public/Converters/GLTFUVDegenerateChecker.h b/Engine/Plugins/Enterprise/GLTFExporter/MigratedFromGit/Source/GLTFExporter/Public/Converters/GLTFUVDegenerateChecker.h
--- a/Engine/Plugins/Enterprise/GLTFExporter/MigratedFromGit/Source/GLTFExporter/Public/Converters/GLTFUVDegenerateChecker.h
+++ b/Engine/Plugins/Enterprise/GLTFExporter/MigratedFromGit/Source/GLTFExporter/Public/Converters/GLTFUVDegenerateChecker.h
@@ -13,6 +13,6 @@ class GLTFEXPORTER_API FGLTFUVDegenerateChecker final : public TGLTFConverter<fl
 
 	virtual float Convert(const FMeshDescription* Description, FGLTFIndexArray SectionIndices, int32 TexCoord) override;
 
-	static bool IsDegenerateTriangle(const TStaticArray<FVector2D, 3>& Points);
-	static bool IsDegenerateTriangle(const TStaticArray<FVector, 3>& Points);
+	static bool IsDegenerateTriangle(const TStaticArray<FVector2f, 3>& Points);
+	static bool IsDegenerateTriangle(const TStaticArray<FVector3f, 3>& Points);
 };
diff --git a/Engine/Plugins/Enterprise/GLTFExporter/MigratedFromGit/Source/GLTFExporter/Public/Utilities/GLTFCoreUtilities.h b/Engine/Plugins/Enterprise/GLTFExporter/MigratedFromGit/Source/GLTFExporter/Public/Utilities/GLTFCoreUtilities.h
--- a/Engine/Plugins/Enterprise/GLTFExporter/MigratedFromGit/Source/GLTFExporter/Public/Utilities/GLTFCoreUtilities.h
+++ b/Engine/Plugins/Enterprise/GLTFExporter/MigratedFromGit/Source/GLTFExporter/Public/Utilities/GLTFCoreUtilities.h
@@ -15,33 +15,33 @@ struct GLTFEXPORTER_API FGLTFCoreUtilities
 {
 	static float ConvertLength(const float Length, const float ConversionScale = 0.01);
 
-	static FGLTFVector3 ConvertVector(const FVector& Vector);
+	static FGLTFVector3 ConvertVector(const FVector3f& Vector);
 
-	static FGLTFVector3 ConvertPosition(const FVector& Position, const float ConversionScale = 0.01);
+	static FGLTFVector3 ConvertPosition(const FVector3f& Position, const float ConversionScale = 0.01);
 
-	static FGLTFVector3 ConvertScale(const FVector& Scale);
+	static FGLTFVector3 ConvertScale(const FVector3f& Scale);
 
-	static FGLTFVector3 ConvertNormal(const FVector& Normal);
+	static FGLTFVector3 ConvertNormal(const FVector3f& Normal);
 	static FGLTFInt16Vector4 ConvertNormal(const FPackedRGBA16N& Normal);
 	static FGLTFInt8Vector4 ConvertNormal(const FPackedNormal& Normal);
 
-	static FGLTFVector4 ConvertTangent(const FVector& Tangent);
+	static FGLTFVector4 ConvertTangent(const FVector3f& Tangent);
 	static FGLTFInt16Vector4 ConvertTangent(const FPackedRGBA16N& Tangent);
 	static FGLTFInt8Vector4 ConvertTangent(const FPackedNormal& Tangent);
 
-	static FGLTFVector2 ConvertUV(const FVector2D& UV);
+	static FGLTFVector2 ConvertUV(const FVector2f& UV);
 	static FGLTFVector2 ConvertUV(const FVector2DHalf& UV);
 
 	static FGLTFColor4 ConvertColor(const FLinearColor& Color, bool bForceLDR);
 	static FGLTFColor3 ConvertColor3(const FLinearColor& Color, bool bForceLDR);
 	static FGLTFUInt8Color4 ConvertColor(const FColor& Color);
 
-	static FGLTFQuaternion ConvertRotation(const FRotator& Rotation);
-	static FGLTFQuaternion ConvertRotation(const FQuat& Rotation);
+	static FGLTFQuaternion ConvertRotation(const FRotator3f& Rotation);
+	static FGLTFQuaternion ConvertRotation(const FQuat4f& Rotation);
 
-	static FGLTFMatrix4 ConvertMatrix(const FMatrix& Matrix);
+	static FGLTFMatrix4 ConvertMatrix(const FMatrix44f& Matrix);
 
-	static FGLTFMatrix4 ConvertTransform(const FTransform& Transform, const float ConversionScale = 0.01);
+	static FGLTFMatrix4 ConvertTransform(const FTransform3f& Transform, const float ConversionScale = 0.01);
 
 	static float ConvertFieldOfView(float FOVInDegrees, float AspectRatio);
 
