[GLTFExporter] fix: update FGLTFProxyAssetActions to new IAssetTypeActions interface in UE5

commit 15a19f2bbc66927b85e2a1cb26b0c71460b4ac70
Date: Thu, 17 Mar 2022 17:48:59 +0100

    fix: update FGLTFProxyAssetActions to new IAssetTypeActions interface in UE5
    (cherry picked from commit 85822b4b4fcdffd3baa8c8c6c6cb59f44c453448​)

#skipci

diff --git a/Engine/Plugins/Enterprise/GLTFExporter/MigratedFromGit/Source/GLTFExporter/Private/Actions/GLTFProxyAssetActions.h b/Engine/Plugins/Enterprise/GLTFExporter/MigratedFromGit/Source/GLTFExporter/Private/Actions/GLTFProxyAssetActions.h
--- a/Engine/Plugins/Enterprise/GLTFExporter/MigratedFromGit/Source/GLTFExporter/Private/Actions/GLTFProxyAssetActions.h
+++ b/Engine/Plugins/Enterprise/GLTFExporter/MigratedFromGit/Source/GLTFExporter/Private/Actions/GLTFProxyAssetActions.h
@@ -6,7 +6,7 @@
 
 #include "AssetTypeActions_Base.h"
 
-class FGLTFProxyAssetActions : public IAssetTypeActions
+class FGLTFProxyAssetActions final : public IAssetTypeActions
 {
 public:
 
@@ -20,10 +20,12 @@ public:
 	virtual FColor GetTypeColor() const override { return OriginalActions->GetTypeColor(); }
 	virtual void GetActions(const TArray<UObject*>& InObjects, FMenuBuilder& MenuBuilder) override { OriginalActions->GetActions(InObjects, MenuBuilder); }
 	virtual void OpenAssetEditor(const TArray<UObject*>& InObjects, TSharedPtr<IToolkitHost> EditWithinLevelEditor) override { OriginalActions->OpenAssetEditor(InObjects, EditWithinLevelEditor); }
-	virtual void AssetsActivated(const TArray<UObject*>& InObjects, EAssetTypeActivationMethod::Type ActivationType) override { OriginalActions->AssetsActivated(InObjects, ActivationType); }
 	virtual bool AssetsActivatedOverride(const TArray<UObject*>& InObjects, EAssetTypeActivationMethod::Type ActivationType) override { return OriginalActions->AssetsActivatedOverride(InObjects, ActivationType); }
+	virtual bool CanRename(const FAssetData& InAsset, FText* OutErrorMsg) const override { return OriginalActions->CanRename(InAsset, OutErrorMsg); }
+	virtual bool CanDuplicate(const FAssetData& InAsset, FText* OutErrorMsg) const override { return OriginalActions->CanDuplicate(InAsset, OutErrorMsg); }
 	virtual TArray<FAssetData> GetValidAssetsForPreviewOrEdit(TArrayView<const FAssetData> InAssetDatas, bool bIsPreview) override { return OriginalActions->GetValidAssetsForPreviewOrEdit(InAssetDatas, bIsPreview); }
 	virtual bool CanFilter() override { return OriginalActions->CanFilter(); }
+	virtual FName GetFilterName() const override { return OriginalActions->GetFilterName(); }
 	virtual bool CanLocalize() const override { return OriginalActions->CanLocalize(); }
 	virtual bool CanMerge() const override { return OriginalActions->CanMerge(); }
 	virtual void Merge( UObject* InObject) override { OriginalActions->Merge(InObject); }
