[GLTFExporter] fix: use new UE5 unified RHI buffer (instead of separate for Index and Vertex data)

commit fdb7f2bb657fc87247485084b0bbad1c4ec0ca12
Date: Thu, 17 Mar 2022 18:12:03 +0100

    fix: use new UE5 unified RHI buffer (instead of separate for Index and Vertex data)
    (cherry picked from commit ea532942943cc71be28aec2bfc2b8da6c417e2f0​)

#skipci

diff --git a/Engine/Plugins/Enterprise/GLTFExporter/MigratedFromGit/Source/GLTFExporter/Private/Converters/GLTFBufferAdapter.cpp b/Engine/Plugins/Enterprise/GLTFExporter/MigratedFromGit/Source/GLTFExporter/Private/Converters/GLTFBufferAdapter.cpp
--- a/Engine/Plugins/Enterprise/GLTFExporter/MigratedFromGit/Source/GLTFExporter/Private/Converters/GLTFBufferAdapter.cpp
+++ b/Engine/Plugins/Enterprise/GLTFExporter/MigratedFromGit/Source/GLTFExporter/Private/Converters/GLTFBufferAdapter.cpp
@@ -24,12 +24,7 @@ class FGLTFBufferAdapterGPU final : public IGLTFBufferAdapter
 {
 public:
 
-	FGLTFBufferAdapterGPU(FRHIIndexBuffer* RHIBuffer)
-	{
-		FGLTFBufferUtility::ReadRHIBuffer(RHIBuffer, DataBuffer);
-	}
-
-	FGLTFBufferAdapterGPU(FRHIVertexBuffer* RHIBuffer)
+	FGLTFBufferAdapterGPU(FRHIBuffer* RHIBuffer)
 	{
 		FGLTFBufferUtility::ReadRHIBuffer(RHIBuffer, DataBuffer);
 	}
diff --git a/Engine/Plugins/Enterprise/GLTFExporter/MigratedFromGit/Source/GLTFExporter/Private/Converters/GLTFBufferUtility.cpp b/Engine/Plugins/Enterprise/GLTFExporter/MigratedFromGit/Source/GLTFExporter/Private/Converters/GLTFBufferUtility.cpp
--- a/Engine/Plugins/Enterprise/GLTFExporter/MigratedFromGit/Source/GLTFExporter/Private/Converters/GLTFBufferUtility.cpp
+++ b/Engine/Plugins/Enterprise/GLTFExporter/MigratedFromGit/Source/GLTFExporter/Private/Converters/GLTFBufferUtility.cpp
@@ -138,7 +138,7 @@ const void* FGLTFBufferUtility::GetCPUBuffer(const FSkinWeightLookupVertexBuffer
 	return reinterpret_cast<const FSkinWeightLookupVertexBufferHack*>(VertexBuffer)->Data;
 }
 
-void FGLTFBufferUtility::ReadRHIBuffer(FRHIVertexBuffer* SourceBuffer, TArray<uint8>& OutData)
+void FGLTFBufferUtility::ReadRHIBuffer(FRHIBuffer* SourceBuffer, TArray<uint8>& OutData)
 {
 	OutData.Empty();
 
@@ -153,8 +153,8 @@ void FGLTFBufferUtility::ReadRHIBuffer(FRHIVertexBuffer* SourceBuffer, TArray<ui
 		return;
 	}
 
-	const uint32 Usage = SourceBuffer->GetUsage();
-	if ((Usage & BUF_Static) == 0)
+	const EBufferUsageFlags Usage = SourceBuffer->GetUsage();
+	if ((Usage & BUF_Static) != BUF_Static)
 	{
 		return; // Some RHI implementations only support reading static buffers
 	}
@@ -165,45 +165,9 @@ void FGLTFBufferUtility::ReadRHIBuffer(FRHIVertexBuffer* SourceBuffer, TArray<ui
 	ENQUEUE_RENDER_COMMAND(ReadRHIBuffer)(
 		[SourceBuffer, NumBytes, DstData](FRHICommandListImmediate& RHICmdList)
 		{
-			const void* SrcData = RHICmdList.LockVertexBuffer(SourceBuffer, 0, NumBytes, RLM_ReadOnly);
+			const void* SrcData = RHICmdList.LockBuffer(SourceBuffer, 0, NumBytes, RLM_ReadOnly);
 			FMemory::Memcpy(DstData, SrcData, NumBytes);
-			RHICmdList.UnlockVertexBuffer(SourceBuffer);
-		}
-	);
-
-	FlushRenderingCommands();
-}
-
-void FGLTFBufferUtility::ReadRHIBuffer(FRHIIndexBuffer* SourceBuffer, TArray<uint8>& OutData)
-{
-	OutData.Empty();
-
-	if (SourceBuffer == nullptr)
-	{
-		return;
-	}
-
-	const uint32 NumBytes = SourceBuffer->GetSize();
-	if (NumBytes == 0)
-	{
-		return;
-	}
-
-	const uint32 Usage = SourceBuffer->GetUsage();
-	if ((Usage & BUF_Static) == 0)
-	{
-		return; // Some RHI implementations only support reading static buffers
-	}
-
-	OutData.AddUninitialized(NumBytes);
-	void *DstData = OutData.GetData();
-
-	ENQUEUE_RENDER_COMMAND(ReadRHIBuffer)(
-		[SourceBuffer, NumBytes, DstData](FRHICommandListImmediate& RHICmdList)
-		{
-			const void* SrcData = RHICmdList.LockIndexBuffer(SourceBuffer, 0, NumBytes, RLM_ReadOnly);
-			FMemory::Memcpy(DstData, SrcData, NumBytes);
-			RHICmdList.UnlockIndexBuffer(SourceBuffer);
+			RHICmdList.UnlockBuffer(SourceBuffer);
 		}
 	);
 
diff --git a/Engine/Plugins/Enterprise/GLTFExporter/MigratedFromGit/Source/GLTFExporter/Private/Converters/GLTFBufferUtility.h b/Engine/Plugins/Enterprise/GLTFExporter/MigratedFromGit/Source/GLTFExporter/Private/Converters/GLTFBufferUtility.h
--- a/Engine/Plugins/Enterprise/GLTFExporter/MigratedFromGit/Source/GLTFExporter/Private/Converters/GLTFBufferUtility.h
+++ b/Engine/Plugins/Enterprise/GLTFExporter/MigratedFromGit/Source/GLTFExporter/Private/Converters/GLTFBufferUtility.h
@@ -18,6 +18,5 @@ struct FGLTFBufferUtility
 	static const void* GetCPUBuffer(const FSkinWeightDataVertexBuffer* VertexBuffer);
 	static const void* GetCPUBuffer(const FSkinWeightLookupVertexBuffer* VertexBuffer);
 
-	static void ReadRHIBuffer(FRHIIndexBuffer* SourceBuffer, TArray<uint8>& OutData);
-	static void ReadRHIBuffer(FRHIVertexBuffer* SourceBuffer, TArray<uint8>& OutData);
+	static void ReadRHIBuffer(FRHIBuffer* SourceBuffer, TArray<uint8>& OutData);
 };
