[GLTFExporter] fix: add missing json print policy methods for UE5

commit fadbf80365685ba88e0a4ba8ece380381bdb6a22
Date: Thu, 17 Mar 2022 18:19:04 +0100

    fix: add missing json print policy methods for UE5
    (cherry picked from commit d5e03def7130ab8ed34827568bcebb106c59388c​)

#skipci

diff --git a/Engine/Plugins/Enterprise/GLTFExporter/MigratedFromGit/Source/GLTFExporter/Private/Json/GLTFJsonWriter.cpp b/Engine/Plugins/Enterprise/GLTFExporter/MigratedFromGit/Source/GLTFExporter/Private/Json/GLTFJsonWriter.cpp
--- a/Engine/Plugins/Enterprise/GLTFExporter/MigratedFromGit/Source/GLTFExporter/Private/Json/GLTFJsonWriter.cpp
+++ b/Engine/Plugins/Enterprise/GLTFExporter/MigratedFromGit/Source/GLTFExporter/Private/Json/GLTFJsonWriter.cpp
@@ -62,16 +62,13 @@ public:
 
 	virtual void Write(float Number) override
 	{
-		// NOTE: Specifying 9 significant digits, this ensures no precision is lost
-		FString ExactString = FString::Printf(TEXT("%.9g"), Number);
-
 		if (CurrentIdentifier.IsEmpty())
 		{
-			JsonWriter->WriteRawJSONValue(ExactString);
+			JsonWriter->WriteValue(Number);
 		}
 		else
 		{
-			JsonWriter->WriteRawJSONValue(CurrentIdentifier, ExactString);
+			JsonWriter->WriteValue(CurrentIdentifier, Number);
 			CurrentIdentifier.Empty();
 		}
 	}
@@ -172,6 +169,16 @@ struct TGLTFJsonPrintPolicyUTF8
 			WriteChar(Stream, *CharPtr);
 		}
 	}
+
+	static void WriteFloat(FArchive* Stream, float Number)
+	{
+		WriteString(Stream, FString::Printf(TEXT("%.9g"), Number));
+	}
+
+	static void WriteDouble(FArchive* Stream, double Number)
+	{
+		WriteString(Stream, FString::Printf(TEXT("%.17g"), Number));
+	}
 };
 
 struct TGLTFCondensedJsonPrintPolicyUTF8 : TGLTFJsonPrintPolicyUTF8
