[GLTFExporter] fix: fix remaining build issues for UE5 related to some return types

commit e37f84d80f03ba3028bc2e49c6780f80e935bfee
Date: Thu, 17 Mar 2022 18:20:27 +0100

    fix: fix remaining build issues for UE5 related to some return types
    (cherry picked from commit 98e8cdd5bbf36b7b04627ab4731a61fc692141e6​)

#skipci

diff --git a/Engine/Plugins/Enterprise/GLTFExporter/MigratedFromGit/Source/GLTFExporter/Private/Converters/GLTFMeshUtility.cpp b/Engine/Plugins/Enterprise/GLTFExporter/MigratedFromGit/Source/GLTFExporter/Private/Converters/GLTFMeshUtility.cpp
--- a/Engine/Plugins/Enterprise/GLTFExporter/MigratedFromGit/Source/GLTFExporter/Private/Converters/GLTFMeshUtility.cpp
+++ b/Engine/Plugins/Enterprise/GLTFExporter/MigratedFromGit/Source/GLTFExporter/Private/Converters/GLTFMeshUtility.cpp
@@ -124,7 +124,7 @@ FGLTFIndexArray FGLTFMeshUtility::GetSectionIndices(const USkeletalMesh* Skeleta
 
 FGLTFIndexArray FGLTFMeshUtility::GetSectionIndices(const FStaticMeshLODResources& MeshLOD, int32 MaterialIndex)
 {
-	const FStaticMeshLODResources::FStaticMeshSectionArray& Sections = MeshLOD.Sections;
+	const FStaticMeshSectionArray& Sections = MeshLOD.Sections;
 
 	FGLTFIndexArray SectionIndices;
 	SectionIndices.Reserve(Sections.Num());
diff --git a/Engine/Plugins/Enterprise/GLTFExporter/MigratedFromGit/Source/GLTFExporter/Private/Converters/GLTFUVDegenerateChecker.cpp b/Engine/Plugins/Enterprise/GLTFExporter/MigratedFromGit/Source/GLTFExporter/Private/Converters/GLTFUVDegenerateChecker.cpp
--- a/Engine/Plugins/Enterprise/GLTFExporter/MigratedFromGit/Source/GLTFExporter/Private/Converters/GLTFUVDegenerateChecker.cpp
+++ b/Engine/Plugins/Enterprise/GLTFExporter/MigratedFromGit/Source/GLTFExporter/Private/Converters/GLTFUVDegenerateChecker.cpp
@@ -38,7 +38,7 @@ float FGLTFUVDegenerateChecker::Convert(const FMeshDescription* Description, FGL
 	if (Description != nullptr)
 	{
 		const TVertexAttributesConstRef<FVector3f> Positions =
-		Description->VertexAttributes().GetAttributesRef<FVector3f>(MeshAttribute::Vertex::Position);
+			Description->VertexAttributes().GetAttributesRef<FVector3f>(MeshAttribute::Vertex::Position);
 		const TVertexInstanceAttributesConstRef<FVector2f> UVs =
 			Description->VertexInstanceAttributes().GetAttributesRef<FVector2f>(MeshAttribute::VertexInstance::TextureCoordinate);
 
@@ -51,8 +51,8 @@ float FGLTFUVDegenerateChecker::Convert(const FMeshDescription* Description, FGL
 			{
 				for (const FTriangleID TriangleID : Description->GetPolygonTriangleIDs(PolygonID))
 				{
-					const TStaticArray<FVertexID, 3> TriangleVertexIDs = Description->GetTriangleVertices(TriangleID);
-					TArrayView<const FVertexInstanceID> TriangleVertexInstanceIDs = Description->GetTriangleVertexInstances(TriangleID);
+					const TArrayView<const FVertexID> TriangleVertexIDs = Description->GetTriangleVertices(TriangleID);
+					const TArrayView<const FVertexInstanceID> TriangleVertexInstanceIDs = Description->GetTriangleVertexInstances(TriangleID);
 
 					TStaticArray<FVector3f, 3> TrianglePositions;
 					TStaticArray<FVector2f, 3> TriangleUVs;
