[GLTFExporter] fix: add missing expected warnings in tests regarding vertex color

commit e8603ce9204f2eb992b0709edad5e1f664a3e588
Date: Fri, 25 Mar 2022 01:18:20 +0100

    fix: add missing expected warnings in tests regarding vertex color
    (cherry picked from commit ffbb73d4f59077eb576fb56ba824fdb646b09726​)

#skipci

diff --git a/Engine/Plugins/Enterprise/GLTFExporter/MigratedFromGit/Source/GLTFExporter/Private/Actors/GLTFHotspotActor.cpp b/Engine/Plugins/Enterprise/GLTFExporter/MigratedFromGit/Source/GLTFExporter/Private/Actors/GLTFHotspotActor.cpp
--- a/Engine/Plugins/Enterprise/GLTFExporter/MigratedFromGit/Source/GLTFExporter/Private/Actors/GLTFHotspotActor.cpp
+++ b/Engine/Plugins/Enterprise/GLTFExporter/MigratedFromGit/Source/GLTFExporter/Private/Actors/GLTFHotspotActor.cpp
@@ -465,11 +465,11 @@ void AGLTFHotspotActor::ValidateAnimation()
 			{
 				if (Skeleton != nullptr)
 				{
-					UE_LOG(LogGLTFHotspot, Warning, TEXT("Animation %s is incompatible with skeleton %s, removing animation from actor."), *AnimationSequence->GetName(), *Skeleton->GetName());
+					UE_LOG(LogGLTFHotspot, Warning, TEXT("Animation %s is incompatible with skeleton %s, removing animation from actor"), *AnimationSequence->GetName(), *Skeleton->GetName());
 				}
 				else
 				{
-					UE_LOG(LogGLTFHotspot, Warning, TEXT("Animation %s is incompatible because mesh %s has no skeleton, removing animation from actor."), *AnimationSequence->GetName(), *SkeletalMesh->GetName());
+					UE_LOG(LogGLTFHotspot, Warning, TEXT("Animation %s is incompatible because mesh %s has no skeleton, removing animation from actor"), *AnimationSequence->GetName(), *SkeletalMesh->GetName());
 				}
 
 				AnimationSequence = nullptr;
diff --git a/Engine/Plugins/Enterprise/GLTFExporter/MigratedFromGit/Source/GLTFExporter/Private/Tasks/GLTFDelayedMeshTasks.cpp b/Engine/Plugins/Enterprise/GLTFExporter/MigratedFromGit/Source/GLTFExporter/Private/Tasks/GLTFDelayedMeshTasks.cpp
--- a/Engine/Plugins/Enterprise/GLTFExporter/MigratedFromGit/Source/GLTFExporter/Private/Tasks/GLTFDelayedMeshTasks.cpp
+++ b/Engine/Plugins/Enterprise/GLTFExporter/MigratedFromGit/Source/GLTFExporter/Private/Tasks/GLTFDelayedMeshTasks.cpp
@@ -119,7 +119,7 @@ void FGLTFDelayedStaticMeshTask::Process()
 	if (Builder.ExportOptions->bExportVertexColors && HasVertexColors(ColorBuffer))
 	{
 		Builder.LogWarning(FString::Printf(
-			TEXT("Vertex colors in mesh %s will act as a multiplier for base color in glTF, regardless of material, which may produce undesirable results."),
+			TEXT("Vertex colors in mesh %s will act as a multiplier for base color in glTF, regardless of material, which may produce undesirable results"),
 			*StaticMesh->GetName()));
 	}
 	else
@@ -221,7 +221,7 @@ void FGLTFDelayedSkeletalMeshTask::Process()
 	if (Builder.ExportOptions->bExportVertexColors && HasVertexColors(ColorBuffer))
 	{
 		Builder.LogWarning(FString::Printf(
-			TEXT("Vertex colors in mesh %s will act as a multiplier for base color in glTF, regardless of material, which may produce undesirable results."),
+			TEXT("Vertex colors in mesh %s will act as a multiplier for base color in glTF, regardless of material, which may produce undesirable results"),
 			*SkeletalMesh->GetName()));
 	}
 	else
