[GLTFExporter] fix: simplify UV overlap checking to eliminate risk of inaccuracy due to color space or platform

commit 5a8f6468dbd49f6d7f070461844719a6ec412a62
Date: Fri, 25 Mar 2022 01:28:00 +0100

    fix: simplify UV overlap checking to eliminate risk of inaccuracy due to color space or platform
    (cherry picked from commit 461c6113eb2f14034a9c910c9b60e3d9a9717613​)

#skipci

diff --git a/Engine/Plugins/Enterprise/GLTFExporter/MigratedFromGit/Source/GLTFExporter/Private/Converters/GLTFUVOverlapChecker.cpp b/Engine/Plugins/Enterprise/GLTFExporter/MigratedFromGit/Source/GLTFExporter/Private/Converters/GLTFUVOverlapChecker.cpp
--- a/Engine/Plugins/Enterprise/GLTFExporter/MigratedFromGit/Source/GLTFExporter/Private/Converters/GLTFUVOverlapChecker.cpp
+++ b/Engine/Plugins/Enterprise/GLTFExporter/MigratedFromGit/Source/GLTFExporter/Private/Converters/GLTFUVOverlapChecker.cpp
@@ -42,7 +42,7 @@ float FGLTFUVOverlapChecker::Convert(const FMeshDescription* Description, FGLTFI
 	{
 		// TODO: investigate if the fixed size is high enough to properly calculate overlap
 		const FIntPoint TextureSize(512, 512);
-		const FGLTFMaterialPropertyEx Property = MP_EmissiveColor;
+		const FGLTFMaterialPropertyEx Property = MP_Opacity;
 
 		FGLTFMeshRenderData MeshSet;
 		MeshSet.TextureCoordinateBox = { { 0.0f, 0.0f }, { 1.0f, 1.0f } };
@@ -54,6 +54,7 @@ float FGLTFUVOverlapChecker::Convert(const FMeshDescription* Description, FGLTFI
 		MatSet.Material = GetMaterial();
 		MatSet.PropertySizes.Add(Property, TextureSize);
 		MatSet.BlendMode = MatSet.Material->GetBlendMode();
+		MatSet.BackgroundColor = FColor::Black;
 		MatSet.bPerformBorderSmear = false;
 
 		TArray<FGLTFMeshRenderData*> MeshSettings;
@@ -74,45 +75,37 @@ float FGLTFUVOverlapChecker::Convert(const FMeshDescription* Description, FGLTFI
 			return -1;
 		}
 
-		// NOTE: the emissive value of each pixel will be incremented by 10 each time a triangle is drawn on it.
-		// Therefore a value of 0.0 indicates an unreferenced pixel, a value of 10.0 indicates a uniquely referenced pixel,
-		// and a value of 20+ indicates an overlapping pixel.
-		// The value may differ somewhat from what was specified in the material due to conversions (color-space, gamma, float => uint etc),
-		// so to be safe we set the limit at 15 instead of 10 to prevent incorrectly flagging some pixels as overlapping.
-		const float EmissiveThreshold = 15.0f;
-
-		int32 ValidCount = 0;
+		int32 TotalCount = 0;
 		int32 OverlapCount = 0;
 
-		if (BakeOutput.EmissiveScale > EmissiveThreshold)
+		for (const FColor& Pixel: BakedPixels)
 		{
-			// This should never overflow since we know that the quotient will be <= 1.0
-			const uint8 ColorThreshold = FMath::RoundToInt((EmissiveThreshold / BakeOutput.EmissiveScale) * 255.0f);
+			const bool bIsBackground = Pixel.G < 64;
+			if (bIsBackground)
+			{
+				continue;
+			}
 
-			bool bIsMagenta;
-			bool bIsOverlapping;
+			TotalCount++;
 
-			for (const FColor& Pixel: BakedPixels)
+			const bool bIsOverlapping = Pixel.G > 192;
+			if (bIsOverlapping)
 			{
-				bIsMagenta = Pixel == FColor::Magenta;
-				bIsOverlapping = Pixel.G > ColorThreshold;
-
-				ValidCount += !bIsMagenta ? 1 : 0;
-				OverlapCount += !bIsMagenta && bIsOverlapping ? 1 : 0;
+				OverlapCount++;
 			}
 		}
 
-		if (ValidCount == 0)
+		if (TotalCount == 0)
 		{
 			return -1;
 		}
 
-		if (ValidCount == OverlapCount)
+		if (TotalCount == OverlapCount)
 		{
 			return 1;
 		}
 
-		return static_cast<float>(OverlapCount) / static_cast<float>(ValidCount);
+		return static_cast<float>(OverlapCount) / static_cast<float>(TotalCount);
 	}
 #endif
 
@@ -121,7 +114,7 @@ float FGLTFUVOverlapChecker::Convert(const FMeshDescription* Description, FGLTFI
 
 UMaterialInterface* FGLTFUVOverlapChecker::GetMaterial()
 {
-	static UMaterialInterface* Material = LoadObject<UMaterialInterface>(nullptr, TEXT("/GLTFExporter/Materials/OverlappingUVs.OverlappingUVs"));
+	static UMaterialInterface* Material = LoadObject<UMaterialInterface>(nullptr, TEXT("/GLTFExporter/Materials/M_UVOverlapChecker.M_UVOverlapChecker"));
 	check(Material);
 	return Material;
 }
diff --git a/Engine/Plugins/Enterprise/GLTFExporter/MigratedFromGit/Source/GLTFExporter/Private/Json/GLTFJsonWriter.cpp b/Engine/Plugins/Enterprise/GLTFExporter/MigratedFromGit/Source/GLTFExporter/Private/Json/GLTFJsonWriter.cpp
--- a/Engine/Plugins/Enterprise/GLTFExporter/MigratedFromGit/Source/GLTFExporter/Private/Json/GLTFJsonWriter.cpp
+++ b/Engine/Plugins/Enterprise/GLTFExporter/MigratedFromGit/Source/GLTFExporter/Private/Json/GLTFJsonWriter.cpp
@@ -2,8 +2,6 @@
 
 #include "Json/GLTFJsonWriter.h"
 #include "Serialization/JsonSerializer.h"
-#include "Policies/PrettyJsonPrintPolicy.h"
-#include "Policies/CondensedJsonPrintPolicy.h"
 
 template <class CharType, class PrintPolicy>
 class TGLTFJsonWriterImpl final : public IGLTFJsonWriter
diff --git a/Engine/Plugins/Enterprise/GLTFExporter/MigratedFromGit/Source/GLTFMaterialBaking/Private/GLTFExportMaterialProxy.h b/Engine/Plugins/Enterprise/GLTFExporter/MigratedFromGit/Source/GLTFMaterialBaking/Private/GLTFExportMaterialProxy.h
--- a/Engine/Plugins/Enterprise/GLTFExporter/MigratedFromGit/Source/GLTFMaterialBaking/Private/GLTFExportMaterialProxy.h
+++ b/Engine/Plugins/Enterprise/GLTFExporter/MigratedFromGit/Source/GLTFMaterialBaking/Private/GLTFExportMaterialProxy.h
@@ -538,7 +538,7 @@ public:
 	}
 	virtual bool IsMasked() const override { return ProxyBlendMode == BLEND_Masked; }
 	virtual enum EBlendMode GetBlendMode() const override { return ProxyBlendMode; }
-	virtual FMaterialShadingModelField GetShadingModels() const override { return MSM_DefaultLit; }
+	virtual FMaterialShadingModelField GetShadingModels() const override { return ProxyBlendMode == BLEND_Opaque ? MSM_DefaultLit : MSM_Unlit; }
 	virtual bool IsShadingModelFromMaterialExpression() const override { return false; }
 	virtual float GetOpacityMaskClipValue() const override { return 0.5f; }
 	virtual bool GetCastDynamicShadowAsMasked() const override { return false; }
