[GLTFExporter] fix: fix incorrect vector UE5 type when converting full-precision UV vertex buffer

commit 40c15424b886ca309a3e60d084ee8fa813116f7c
Date: Sun, 27 Mar 2022 08:52:46 +0200

    fix: fix incorrect vector UE5 type when converting full-precision UV vertex buffer
    (cherry picked from commit fbba5db48758dc4596fe71276f540eab2a3e2631​)

#skipci

diff --git a/Engine/Plugins/Enterprise/GLTFExporter/MigratedFromGit/Source/GLTFExporter/Private/Converters/GLTFAccessorConverters.cpp b/Engine/Plugins/Enterprise/GLTFExporter/MigratedFromGit/Source/GLTFExporter/Private/Converters/GLTFAccessorConverters.cpp
--- a/Engine/Plugins/Enterprise/GLTFExporter/MigratedFromGit/Source/GLTFExporter/Private/Converters/GLTFAccessorConverters.cpp
+++ b/Engine/Plugins/Enterprise/GLTFExporter/MigratedFromGit/Source/GLTFExporter/Private/Converters/GLTFAccessorConverters.cpp
@@ -307,7 +307,7 @@ FGLTFJsonAccessor* FGLTFUVBufferConverter::Convert(const FGLTFMeshSection* MeshS
 	}
 
 	return VertexBuffer->GetUseFullPrecisionUVs()
-		? Convert<FVector2D>(MeshSection, VertexBuffer, UVIndex, SourceData)
+		? Convert<FVector2f>(MeshSection, VertexBuffer, UVIndex, SourceData)
 		: Convert<FVector2DHalf>(MeshSection, VertexBuffer, UVIndex, SourceData);
 }
 
