[GLTFExporter] fix: make sure texture is fully loaded/streamed before converting

commit 102bebf71afa425bc23365d2732bbca48a899c86
Date: Sun, 27 Mar 2022 08:54:09 +0200

    fix: make sure texture is fully loaded/streamed before converting
    (cherry picked from commit 3ad23cb893547b8b32d8659ee61840cb3173d01d​)

#skipci

diff --git a/Engine/Plugins/Enterprise/GLTFExporter/MigratedFromGit/Source/GLTFExporter/Private/Converters/GLTFTextureUtility.cpp b/Engine/Plugins/Enterprise/GLTFExporter/MigratedFromGit/Source/GLTFExporter/Private/Converters/GLTFTextureUtility.cpp
--- a/Engine/Plugins/Enterprise/GLTFExporter/MigratedFromGit/Source/GLTFExporter/Private/Converters/GLTFTextureUtility.cpp
+++ b/Engine/Plugins/Enterprise/GLTFExporter/MigratedFromGit/Source/GLTFExporter/Private/Converters/GLTFTextureUtility.cpp
@@ -6,6 +6,7 @@
 #include "Converters/GLTFCubemapFacePreview.h"
 #include "CanvasItem.h"
 #include "CanvasTypes.h"
+#include "TextureCompiler.h"
 #include "Engine/Texture2D.h"
 #include "Engine/TextureCube.h"
 #include "Engine/TextureRenderTarget2D.h"
@@ -33,6 +34,18 @@ bool FGLTFTextureUtility::IsAlphaless(EPixelFormat PixelFormat)
 	}
 }
 
+void FGLTFTextureUtility::FullyLoad(const UTexture* Texture)
+{
+	UTexture* MutableTexture = const_cast<UTexture*>(Texture);
+
+#if WITH_EDITOR
+	FTextureCompilingManager::Get().FinishCompilation({ MutableTexture });
+#endif
+
+	MutableTexture->SetForceMipLevelsToBeResident(30.0f);
+	MutableTexture->WaitForStreaming();
+}
+
 bool FGLTFTextureUtility::IsHDR(const UTexture* Texture)
 {
 	switch (Texture->CompressionSettings)
@@ -175,10 +188,6 @@ bool FGLTFTextureUtility::DrawTexture(UTextureRenderTarget2D* OutTarget, const U
 		return false;
 	}
 
-	// Fully stream in the texture before drawing it.
-	const_cast<UTexture2D*>(InSource)->SetForceMipLevelsToBeResident(30.0f, true);
-	const_cast<UTexture2D*>(InSource)->WaitForStreaming();
-
 	TRefCountPtr<FBatchedElementParameters> BatchedElementParameters;
 
 	if (InSource->IsNormalMap())
diff --git a/Engine/Plugins/Enterprise/GLTFExporter/MigratedFromGit/Source/GLTFExporter/Private/Converters/GLTFTextureUtility.h b/Engine/Plugins/Enterprise/GLTFExporter/MigratedFromGit/Source/GLTFExporter/Private/Converters/GLTFTextureUtility.h
--- a/Engine/Plugins/Enterprise/GLTFExporter/MigratedFromGit/Source/GLTFExporter/Private/Converters/GLTFTextureUtility.h
+++ b/Engine/Plugins/Enterprise/GLTFExporter/MigratedFromGit/Source/GLTFExporter/Private/Converters/GLTFTextureUtility.h
@@ -8,6 +8,8 @@ struct FGLTFTextureUtility
 {
 	static bool IsAlphaless(EPixelFormat PixelFormat);
 
+	static void FullyLoad(const UTexture* Texture);
+
 	static bool IsHDR(const UTexture* Texture);
 	static bool IsCubemap(const UTexture* Texture);
 
diff --git a/Engine/Plugins/Enterprise/GLTFExporter/MigratedFromGit/Source/GLTFExporter/Private/Tasks/GLTFDelayedTextureTasks.cpp b/Engine/Plugins/Enterprise/GLTFExporter/MigratedFromGit/Source/GLTFExporter/Private/Tasks/GLTFDelayedTextureTasks.cpp
--- a/Engine/Plugins/Enterprise/GLTFExporter/MigratedFromGit/Source/GLTFExporter/Private/Tasks/GLTFDelayedTextureTasks.cpp
+++ b/Engine/Plugins/Enterprise/GLTFExporter/MigratedFromGit/Source/GLTFExporter/Private/Tasks/GLTFDelayedTextureTasks.cpp
@@ -18,6 +18,7 @@ FString FGLTFDelayedTexture2DTask::GetName()
 
 void FGLTFDelayedTexture2DTask::Process()
 {
+	FGLTFTextureUtility::FullyLoad(Texture2D);
 	Texture2D->GetName(JsonTexture->Name);
 
 	const bool bFromSRGB = Texture2D->SRGB;
@@ -69,6 +70,7 @@ FString FGLTFDelayedTextureCubeTask::GetName()
 
 void FGLTFDelayedTextureCubeTask::Process()
 {
+	FGLTFTextureUtility::FullyLoad(TextureCube);
 	JsonTexture->Name = TextureCube->GetName() + TEXT("_") + FGLTFJsonUtilities::GetValue(FGLTFCoreUtilities::ConvertCubeFace(CubeFace));
 
 	const bool bFromSRGB = TextureCube->SRGB;
@@ -126,6 +128,7 @@ FString FGLTFDelayedTextureRenderTarget2DTask::GetName()
 
 void FGLTFDelayedTextureRenderTarget2DTask::Process()
 {
+	FGLTFTextureUtility::FullyLoad(RenderTarget2D);
 	RenderTarget2D->GetName(JsonTexture->Name);
 
 	const bool bFromSRGB = RenderTarget2D->SRGB;
@@ -165,6 +168,7 @@ FString FGLTFDelayedTextureRenderTargetCubeTask::GetName()
 
 void FGLTFDelayedTextureRenderTargetCubeTask::Process()
 {
+	FGLTFTextureUtility::FullyLoad(RenderTargetCube);
 	JsonTexture->Name = RenderTargetCube->GetName() + TEXT("_") + FGLTFJsonUtilities::GetValue(FGLTFCoreUtilities::ConvertCubeFace(CubeFace));
 
 	const bool bFromSRGB = RenderTargetCube->SRGB;
