[GLTFExporter] fix: update detection of LocalObjectBounds to UE5

commit 1a76fc1a2c1b9d64e8d07a1d6e61fd64e63ca26b
Date: Sun, 27 Mar 2022 08:55:29 +0200

    fix: update detection of LocalObjectBounds to UE5
    (cherry picked from commit fbf6f8909a231d88d6738b4f0abd42780aa7aeae​)

#skipci

diff --git a/Engine/Plugins/Enterprise/GLTFExporter/MigratedFromGit/Source/GLTFMaterialAnalyzer/Private/GLTFProxyMaterialCompiler.h b/Engine/Plugins/Enterprise/GLTFExporter/MigratedFromGit/Source/GLTFMaterialAnalyzer/Private/GLTFProxyMaterialCompiler.h
--- a/Engine/Plugins/Enterprise/GLTFExporter/MigratedFromGit/Source/GLTFMaterialAnalyzer/Private/GLTFProxyMaterialCompiler.h
+++ b/Engine/Plugins/Enterprise/GLTFExporter/MigratedFromGit/Source/GLTFMaterialAnalyzer/Private/GLTFProxyMaterialCompiler.h
@@ -57,7 +57,7 @@ public:
 #endif
 	{
 		// NOTE: because there is no overridable compiler method for ObjectLocalBounds in engine this is our only option currently
-		if (Custom->Code.Contains(TEXT("GetPrimitiveData(Parameters.PrimitiveId).LocalObjectBounds")))
+		if (Custom->Code.Contains(TEXT("GetPrimitiveData(Parameters).LocalObjectBounds")))
 		{
 			bUsesObjectLocalBounds = true;
 		}
