[GLTFExporter] fix: add workaround for UE5 to access level sequence channels without GetChannelProxy (which doesn't work properly in UE5)

commit 311e8fbd26f8e69f0d4b42754f9e0bebf7504b8d
Date: Sun, 27 Mar 2022 09:01:02 +0200

    fix: add workaround for UE5 to access level sequence channels without GetChannelProxy (which doesn't work properly in UE5)
    (cherry picked from commit c0746c3701988c9c17febf8bd15110c958538889​)

#skipci

diff --git a/Engine/Plugins/Enterprise/GLTFExporter/MigratedFromGit/Source/GLTFExporter/Private/Converters/GLTFBoneUtility.cpp b/Engine/Plugins/Enterprise/GLTFExporter/MigratedFromGit/Source/GLTFExporter/Private/Converters/GLTFBoneUtility.cpp
--- a/Engine/Plugins/Enterprise/GLTFExporter/MigratedFromGit/Source/GLTFExporter/Private/Converters/GLTFBoneUtility.cpp
+++ b/Engine/Plugins/Enterprise/GLTFExporter/MigratedFromGit/Source/GLTFExporter/Private/Converters/GLTFBoneUtility.cpp
@@ -1,6 +1,8 @@
 // Copyright Epic Games, Inc. All Rights Reserved.
 
 #include "Converters/GLTFBoneUtility.h"
+#include "Sections/MovieScene3DTransformSection.h"
+
 #if (ENGINE_MAJOR_VERSION > 4 || ENGINE_MINOR_VERSION >= 26)
 #include "Animation/AnimationPoseData.h"
 #endif
@@ -91,3 +93,24 @@ void FGLTFBoneUtility::GetBoneTransformsByFrame(const UAnimSequence* AnimSequenc
 	Attributes.Empty();
 #endif
 }
+
+const FMovieSceneDoubleChannel* FGLTFBoneUtility::GetTranslationChannels(const UMovieScene3DTransformSection* TransformSection)
+{
+	// TODO: fix ugly workaround hack in engine api to access channels without GetChannelProxy (which doesn't work properly in UE5)
+	static const FProperty* Property = UMovieScene3DTransformSection::StaticClass()->FindPropertyByName(TEXT("Translation"));
+	return Property->ContainerPtrToValuePtr<FMovieSceneDoubleChannel>(TransformSection);
+}
+
+const FMovieSceneDoubleChannel* FGLTFBoneUtility::GetRotationChannels(const UMovieScene3DTransformSection* TransformSection)
+{
+	// TODO: fix ugly workaround hack in engine api to access channels without GetChannelProxy (which doesn't work properly in UE5)
+	static const FProperty* Property = UMovieScene3DTransformSection::StaticClass()->FindPropertyByName(TEXT("Rotation"));
+	return Property->ContainerPtrToValuePtr<FMovieSceneDoubleChannel>(TransformSection);
+}
+
+const FMovieSceneDoubleChannel* FGLTFBoneUtility::GetScaleChannels(const UMovieScene3DTransformSection* TransformSection)
+{
+	// TODO: fix ugly workaround hack in engine api to access channels without GetChannelProxy (which doesn't work properly in UE5)
+	static const FProperty* Property = UMovieScene3DTransformSection::StaticClass()->FindPropertyByName(TEXT("Scale"));
+	return Property->ContainerPtrToValuePtr<FMovieSceneDoubleChannel>(TransformSection);
+}
diff --git a/Engine/Plugins/Enterprise/GLTFExporter/MigratedFromGit/Source/GLTFExporter/Private/Converters/GLTFBoneUtility.h b/Engine/Plugins/Enterprise/GLTFExporter/MigratedFromGit/Source/GLTFExporter/Private/Converters/GLTFBoneUtility.h
--- a/Engine/Plugins/Enterprise/GLTFExporter/MigratedFromGit/Source/GLTFExporter/Private/Converters/GLTFBoneUtility.h
+++ b/Engine/Plugins/Enterprise/GLTFExporter/MigratedFromGit/Source/GLTFExporter/Private/Converters/GLTFBoneUtility.h
@@ -4,6 +4,9 @@
 
 #include "CoreMinimal.h"
 
+struct FMovieSceneDoubleChannel;
+class UMovieScene3DTransformSection;
+
 struct FGLTFBoneUtility
 {
 	static FTransform GetBindTransform(const FReferenceSkeleton& RefSkeleton, int32 BoneIndex);
@@ -13,4 +16,8 @@ struct FGLTFBoneUtility
 	static void GetBoneIndices(const USkeleton* Skeleton, TArray<FBoneIndexType>& OutBoneIndices);
 
 	static void GetBoneTransformsByFrame(const UAnimSequence* AnimSequence, const TArray<float>& FrameTimestamps, const TArray<FBoneIndexType>& BoneIndices, TArray<TArray<FTransform>>& OutBoneTransformsByFrame);
+
+	static const FMovieSceneDoubleChannel* GetTranslationChannels(const UMovieScene3DTransformSection* TransformSection);
+	static const FMovieSceneDoubleChannel* GetRotationChannels(const UMovieScene3DTransformSection* TransformSection);
+	static const FMovieSceneDoubleChannel* GetScaleChannels(const UMovieScene3DTransformSection* TransformSection);
 };
diff --git a/Engine/Plugins/Enterprise/GLTFExporter/MigratedFromGit/Source/GLTFExporter/Private/Tasks/GLTFDelayedAnimationTasks.cpp b/Engine/Plugins/Enterprise/GLTFExporter/MigratedFromGit/Source/GLTFExporter/Private/Tasks/GLTFDelayedAnimationTasks.cpp
--- a/Engine/Plugins/Enterprise/GLTFExporter/MigratedFromGit/Source/GLTFExporter/Private/Tasks/GLTFDelayedAnimationTasks.cpp
+++ b/Engine/Plugins/Enterprise/GLTFExporter/MigratedFromGit/Source/GLTFExporter/Private/Tasks/GLTFDelayedAnimationTasks.cpp
@@ -5,6 +5,7 @@
 #include "Utilities/GLTFCoreUtilities.h"
 #include "Converters/GLTFBoneUtility.h"
 #include "LevelSequence.h"
+#include "LevelSequenceActor.h"
 #include "LevelSequencePlayer.h"
 #include "MovieSceneTimeHelpers.h"
 #include "Tracks/MovieScene3DTransformTrack.h"
@@ -189,17 +190,17 @@ void FGLTFDelayedLevelSequenceTask::Process()
 		for (TWeakObjectPtr<UObject> Object : Player->FindBoundObjects(Binding.GetObjectGuid(), MovieSceneSequenceID::Root))
 		{
 			FGLTFJsonNode* Node;
-			FTransform RelativeTransform;
+			FTransform DefaultTransform;
 
-			if (AActor* Actor = Cast<AActor>(Object.Get()))
+			if (const AActor* Actor = Cast<AActor>(Object.Get()))
 			{
 				Node = Builder.AddUniqueNode(Actor);
-				RelativeTransform = Actor->GetRootComponent()->GetRelativeTransform();
+				DefaultTransform = Actor->GetRootComponent()->GetRelativeTransform();
 			}
-			else if (USceneComponent* SceneComponent = Cast<USceneComponent>(Object.Get()))
+			else if (const USceneComponent* SceneComponent = Cast<USceneComponent>(Object.Get()))
 			{
 				Node = Builder.AddUniqueNode(SceneComponent);
-				RelativeTransform = SceneComponent->GetRelativeTransform();
+				DefaultTransform = SceneComponent->GetRelativeTransform();
 			}
 			else
 			{
@@ -213,6 +214,10 @@ void FGLTFDelayedLevelSequenceTask::Process()
 				continue;
 			}
 
+			FVector3f DefaultTranslation = FVector3f(DefaultTransform.GetTranslation());
+			FRotator3f DefaultRotator = FRotator3f(DefaultTransform.GetRotation().Rotator());
+			FVector3f DefaultScale = FVector3f(DefaultTransform.GetScale3D());
+
 			for (UMovieSceneTrack* Track : Binding.GetTracks())
 			{
 				UMovieScene3DTransformTrack* TransformTrack = Cast<UMovieScene3DTransformTrack>(Track);
@@ -233,7 +238,7 @@ void FGLTFDelayedLevelSequenceTask::Process()
 						continue;
 					}
 
-					FOptionalMovieSceneBlendType BlendType = TransformSection->GetBlendType();
+					const FOptionalMovieSceneBlendType BlendType = TransformSection->GetBlendType();
 					if (BlendType.IsValid() && BlendType.Get() != EMovieSceneBlendType::Absolute)
 					{
 						// TODO: report warning
@@ -242,37 +247,26 @@ void FGLTFDelayedLevelSequenceTask::Process()
 
 					// TODO: do we need to account for TransformSection->GetOffsetTime() ?
 
-					TArrayView<FMovieSceneFloatChannel*> Channels = TransformSection->GetChannelProxy().GetChannels<FMovieSceneFloatChannel>();
-					EMovieSceneTransformChannel ChannelMask = TransformSection->GetMask().GetChannels();
+					const EMovieSceneTransformChannel ChannelMask = TransformSection->GetMask().GetChannels();
 
 					if (EnumHasAnyFlags(ChannelMask, EMovieSceneTransformChannel::Translation))
 					{
-						constexpr int32 IndexX = 0;
-						constexpr int32 IndexY = 1;
-						constexpr int32 IndexZ = 2;
-
-						const bool MaskX = EnumHasAllFlags(ChannelMask, EMovieSceneTransformChannel::TranslationX) && Channels.IsValidIndex(IndexX);
-						const bool MaskY = EnumHasAllFlags(ChannelMask, EMovieSceneTransformChannel::TranslationY) && Channels.IsValidIndex(IndexY);
-						const bool MaskZ = EnumHasAllFlags(ChannelMask, EMovieSceneTransformChannel::TranslationZ) && Channels.IsValidIndex(IndexZ);
+						const FMovieSceneDoubleChannel* Channels = FGLTFBoneUtility::GetTranslationChannels(TransformSection);
+						const bool MaskX = EnumHasAllFlags(ChannelMask, EMovieSceneTransformChannel::TranslationX);
+						const bool MaskY = EnumHasAllFlags(ChannelMask, EMovieSceneTransformChannel::TranslationY);
+						const bool MaskZ = EnumHasAllFlags(ChannelMask, EMovieSceneTransformChannel::TranslationZ);
 
-						if (!(MaskX | MaskY | MaskZ))
-						{
-							// TODO: report warning
-							continue;
-						}
-
-						FVector3f DefaultTranslation = FVector3f(RelativeTransform.GetTranslation());
 						TArray<FGLTFVector3> Translations;
 						Translations.AddUninitialized(FrameCount);
 
 						for (int32 Frame = 0; Frame < FrameCount; ++Frame)
 						{
-							FFrameTime& FrameTime = FrameTimes[Frame];
+							const FFrameTime& FrameTime = FrameTimes[Frame];
 
 							FVector3f Translation = DefaultTranslation;
-							if (MaskX) Channels[IndexX]->Evaluate(FrameTime, Translation.X);
-							if (MaskY) Channels[IndexY]->Evaluate(FrameTime, Translation.Y);
-							if (MaskZ) Channels[IndexZ]->Evaluate(FrameTime, Translation.Z);
+							if (MaskX) Channels[0].Evaluate(FrameTime, Translation.X);
+							if (MaskY) Channels[1].Evaluate(FrameTime, Translation.Y);
+							if (MaskZ) Channels[2].Evaluate(FrameTime, Translation.Z);
 
 							Translations[Frame] = FGLTFCoreUtilities::ConvertPosition(Translation, Builder.ExportOptions->ExportUniformScale);
 						}
@@ -297,39 +291,26 @@ void FGLTFDelayedLevelSequenceTask::Process()
 
 					if (EnumHasAnyFlags(ChannelMask, EMovieSceneTransformChannel::Rotation))
 					{
-						constexpr int32 IndexX = 3;
-						constexpr int32 IndexY = 4;
-						constexpr int32 IndexZ = 5;
-
-						const bool MaskX = EnumHasAllFlags(ChannelMask, EMovieSceneTransformChannel::RotationX) && Channels.IsValidIndex(IndexX);
-						const bool MaskY = EnumHasAllFlags(ChannelMask, EMovieSceneTransformChannel::RotationY) && Channels.IsValidIndex(IndexY);
-						const bool MaskZ = EnumHasAllFlags(ChannelMask, EMovieSceneTransformChannel::RotationZ) && Channels.IsValidIndex(IndexZ);
-
-						if (!(MaskX | MaskY | MaskZ))
-						{
-							// TODO: report warning
-							continue;
-						}
+						const FMovieSceneDoubleChannel* Channels = FGLTFBoneUtility::GetRotationChannels(TransformSection);
+						const bool MaskX = EnumHasAllFlags(ChannelMask, EMovieSceneTransformChannel::RotationX);
+						const bool MaskY = EnumHasAllFlags(ChannelMask, EMovieSceneTransformChannel::RotationY);
+						const bool MaskZ = EnumHasAllFlags(ChannelMask, EMovieSceneTransformChannel::RotationZ);
 
-						FRotator3f DefaultRotator = FRotator3f(RelativeTransform.GetRotation().Rotator());
 						TArray<FGLTFQuaternion> Rotations;
 						Rotations.AddUninitialized(FrameCount);
 
 						for (int32 Frame = 0; Frame < FrameCount; ++Frame)
 						{
-							FFrameTime& FrameTime = FrameTimes[Frame];
+							const FFrameTime& FrameTime = FrameTimes[Frame];
 
 							FRotator3f Rotator = DefaultRotator;
-							if (MaskX) Channels[IndexX]->Evaluate(FrameTime, Rotator.Roll);
-							if (MaskY) Channels[IndexY]->Evaluate(FrameTime, Rotator.Pitch);
-							if (MaskZ) Channels[IndexZ]->Evaluate(FrameTime, Rotator.Yaw);
+							if (MaskX) Channels[0].Evaluate(FrameTime, Rotator.Roll);
+							if (MaskY) Channels[1].Evaluate(FrameTime, Rotator.Pitch);
+							if (MaskZ) Channels[2].Evaluate(FrameTime, Rotator.Yaw);
 
 							Rotations[Frame] = FGLTFCoreUtilities::ConvertRotation(Rotator.Quaternion());
 						}
 
-						JsonInputAccessor->Count = Rotations.Num();
-						JsonInputAccessor->Max[0] = Timestamps[Rotations.Num() - 1];
-
 						FGLTFJsonAccessor* JsonOutputAccessor = Builder.AddAccessor();
 						JsonOutputAccessor->BufferView = Builder.AddBufferView(Rotations);
 						JsonOutputAccessor->ComponentType = EGLTFJsonComponentType::Float;
@@ -350,32 +331,22 @@ void FGLTFDelayedLevelSequenceTask::Process()
 
 					if (EnumHasAnyFlags(ChannelMask, EMovieSceneTransformChannel::Scale))
 					{
-						constexpr int32 IndexX = 6;
-						constexpr int32 IndexY = 7;
-						constexpr int32 IndexZ = 8;
-
-						const bool MaskX = EnumHasAllFlags(ChannelMask, EMovieSceneTransformChannel::ScaleX) && Channels.IsValidIndex(IndexX);
-						const bool MaskY = EnumHasAllFlags(ChannelMask, EMovieSceneTransformChannel::ScaleY) && Channels.IsValidIndex(IndexY);
-						const bool MaskZ = EnumHasAllFlags(ChannelMask, EMovieSceneTransformChannel::ScaleZ) && Channels.IsValidIndex(IndexZ);
-
-						if (!(MaskX | MaskY | MaskZ))
-						{
-							// TODO: report warning
-							continue;
-						}
+						const FMovieSceneDoubleChannel* Channels = FGLTFBoneUtility::GetScaleChannels(TransformSection);
+						const bool MaskX = EnumHasAllFlags(ChannelMask, EMovieSceneTransformChannel::ScaleX);
+						const bool MaskY = EnumHasAllFlags(ChannelMask, EMovieSceneTransformChannel::ScaleY);
+						const bool MaskZ = EnumHasAllFlags(ChannelMask, EMovieSceneTransformChannel::ScaleZ);
 
-						FVector3f DefaultScale = FVector3f(RelativeTransform.GetScale3D());
 						TArray<FGLTFVector3> Scales;
 						Scales.AddUninitialized(FrameCount);
 
 						for (int32 Frame = 0; Frame < FrameCount; ++Frame)
 						{
-							FFrameTime& FrameTime = FrameTimes[Frame];
+							const FFrameTime& FrameTime = FrameTimes[Frame];
 
 							FVector3f Scale = DefaultScale;
-							if (MaskX) Channels[IndexX]->Evaluate(FrameTime, Scale.X);
-							if (MaskY) Channels[IndexY]->Evaluate(FrameTime, Scale.Y);
-							if (MaskZ) Channels[IndexZ]->Evaluate(FrameTime, Scale.Z);
+							if (MaskX) Channels[0].Evaluate(FrameTime, Scale.X);
+							if (MaskY) Channels[1].Evaluate(FrameTime, Scale.Y);
+							if (MaskZ) Channels[2].Evaluate(FrameTime, Scale.Z);
 
 							Scales[Frame] = FGLTFCoreUtilities::ConvertScale(Scale);
 						}
diff --git a/Engine/Plugins/Enterprise/GLTFExporter/MigratedFromGit/Source/GLTFExporter/Private/Tasks/GLTFDelayedAnimationTasks.h b/Engine/Plugins/Enterprise/GLTFExporter/MigratedFromGit/Source/GLTFExporter/Private/Tasks/GLTFDelayedAnimationTasks.h
--- a/Engine/Plugins/Enterprise/GLTFExporter/MigratedFromGit/Source/GLTFExporter/Private/Tasks/GLTFDelayedAnimationTasks.h
+++ b/Engine/Plugins/Enterprise/GLTFExporter/MigratedFromGit/Source/GLTFExporter/Private/Tasks/GLTFDelayedAnimationTasks.h
@@ -4,7 +4,6 @@
 
 #include "Tasks/GLTFDelayedTask.h"
 #include "Builders/GLTFConvertBuilder.h"
-#include "LevelSequenceActor.h"
 #include "LevelSequence.h"
 
 class FGLTFDelayedAnimSequenceTask : public FGLTFDelayedTask
