[GLTFExporter] fix: add workaround to load level sequence synchronously in UE5

commit 505b6aa804dbe02e4ffc95bcf9c5c43bb6750041
Date: Sun, 27 Mar 2022 09:02:16 +0200

    fix: add workaround to load level sequence synchronously in UE5
    (cherry picked from commit f40a1de4d5b2efe6cc69c746551a0e55b1343a52​)

#skipci

diff --git a/Engine/Plugins/Enterprise/GLTFExporter/MigratedFromGit/Source/GLTFExporter/Private/Converters/GLTFAnimationConverters.cpp b/Engine/Plugins/Enterprise/GLTFExporter/MigratedFromGit/Source/GLTFExporter/Private/Converters/GLTFAnimationConverters.cpp
--- a/Engine/Plugins/Enterprise/GLTFExporter/MigratedFromGit/Source/GLTFExporter/Private/Converters/GLTFAnimationConverters.cpp
+++ b/Engine/Plugins/Enterprise/GLTFExporter/MigratedFromGit/Source/GLTFExporter/Private/Converters/GLTFAnimationConverters.cpp
@@ -4,6 +4,7 @@
 #include "Builders/GLTFConvertBuilder.h"
 #include "Tasks/GLTFDelayedAnimationTasks.h"
 #include "Animation/AnimSequence.h"
+#include "LevelSequenceActor.h"
 
 FGLTFJsonAnimation* FGLTFAnimationConverter::Convert(FGLTFJsonNode* RootNode, const USkeletalMesh* SkeletalMesh, const UAnimSequence* AnimSequence)
 {
@@ -40,9 +41,19 @@ FGLTFJsonAnimation* FGLTFAnimationConverter::Convert(FGLTFJsonNode* RootNode, co
 FGLTFJsonAnimation* FGLTFAnimationDataConverter::Convert(FGLTFJsonNode* RootNode, const USkeletalMeshComponent* SkeletalMeshComponent)
 {
 	const USkeletalMesh* SkeletalMesh = SkeletalMeshComponent->SkeletalMesh;
+	if (SkeletalMesh == nullptr)
+	{
+		return nullptr;
+	}
+
 	const UAnimSequence* AnimSequence = Cast<UAnimSequence>(SkeletalMeshComponent->AnimationData.AnimToPlay);
+	if (AnimSequence == nullptr)
+	{
+		return nullptr;
+	}
 
-	if (SkeletalMesh == nullptr || AnimSequence == nullptr || SkeletalMeshComponent->GetAnimationMode() != EAnimationMode::AnimationSingleNode)
+	const EAnimationMode::Type AnimationMode = SkeletalMeshComponent->GetAnimationMode();
+	if (AnimationMode != EAnimationMode::AnimationSingleNode)
 	{
 		return nullptr;
 	}
@@ -71,11 +82,24 @@ FGLTFJsonAnimation* FGLTFLevelSequenceConverter::Convert(const ULevel* Level, co
 FGLTFJsonAnimation* FGLTFLevelSequenceDataConverter::Convert(const ALevelSequenceActor* LevelSequenceActor)
 {
 	const ULevel* Level = LevelSequenceActor->GetLevel();
-	const ULevelSequence* LevelSequence = LevelSequenceActor->LoadSequence();
+	if (Level == nullptr)
+	{
+		return nullptr;
+	}
 
-	if (Level == nullptr || LevelSequence == nullptr)
+	const ULevelSequence* LevelSequence = LevelSequenceActor->GetSequence();
+	if (LevelSequence == nullptr)
 	{
+#if WITH_EDITORONLY_DATA
+		// TODO: find a nicer way to load level sequence synchronously without relying on a deprecated engine api
+		LevelSequence = Cast<ULevelSequence>(LevelSequenceActor->LevelSequence_DEPRECATED.TryLoad());
+		if (LevelSequence == nullptr)
+		{
+			return nullptr;
+		}
+#else
 		return nullptr;
+#endif
 	}
 
 	FGLTFJsonAnimation* JsonAnimation = Builder.AddUniqueAnimation(Level, LevelSequence);
