[GLTFExporter] refactor: move scene component conversion to own function and rename FGLTFActorUtility to FGLTFBlueprintUtility

commit 5d0f606d0fd4afcbb6ac3299a93fd71696bf3030
Date: Sun, 27 Mar 2022 12:27:30 +0200

    refactor: move scene component conversion to own function and rename FGLTFActorUtility to FGLTFBlueprintUtility
    (cherry picked from commit 0d101049b5062b2f520b98fe218d136d38fe8711​)

#skipci

diff --git a/Engine/Plugins/Enterprise/GLTFExporter/MigratedFromGit/Source/GLTFExporter/Private/Converters/GLTFActorUtility.cpp b/Engine/Plugins/Enterprise/GLTFExporter/MigratedFromGit/Source/GLTFExporter/Private/Converters/GLTFActorUtility.cpp
deleted file mode 100644
diff --git a/Engine/Plugins/Enterprise/GLTFExporter/MigratedFromGit/Source/GLTFExporter/Private/Converters/GLTFActorUtility.h b/Engine/Plugins/Enterprise/GLTFExporter/MigratedFromGit/Source/GLTFExporter/Private/Converters/GLTFActorUtility.h
deleted file mode 100644
diff --git a/Engine/Plugins/Enterprise/GLTFExporter/MigratedFromGit/Source/GLTFExporter/Private/Converters/GLTFBackdropConverters.cpp b/Engine/Plugins/Enterprise/GLTFExporter/MigratedFromGit/Source/GLTFExporter/Private/Converters/GLTFBackdropConverters.cpp
--- a/Engine/Plugins/Enterprise/GLTFExporter/MigratedFromGit/Source/GLTFExporter/Private/Converters/GLTFBackdropConverters.cpp
+++ b/Engine/Plugins/Enterprise/GLTFExporter/MigratedFromGit/Source/GLTFExporter/Private/Converters/GLTFBackdropConverters.cpp
@@ -3,12 +3,12 @@
 #include "Converters/GLTFBackdropConverters.h"
 #include "Utilities/GLTFCoreUtilities.h"
 #include "Converters/GLTFMaterialUtility.h"
-#include "Converters/GLTFActorUtility.h"
+#include "Converters/GLTFBlueprintUtility.h"
 #include "Builders/GLTFContainerBuilder.h"
 
 FGLTFJsonBackdrop* FGLTFBackdropConverter::Convert(const AActor* BackdropActor)
 {
-	if (!FGLTFActorUtility::IsHDRIBackdropBlueprint(FGLTFActorUtility::GetBlueprintPath(BackdropActor)))
+	if (!FGLTFBlueprintUtility::IsHDRIBackdrop(FGLTFBlueprintUtility::GetClassPath(BackdropActor)))
 	{
 		return nullptr;
 	}
@@ -27,7 +27,7 @@ FGLTFJsonBackdrop* FGLTFBackdropConverter::Convert(const AActor* BackdropActor)
 	JsonBackdrop->Angle = FRotator::ClampAxis(Rotation.Yaw * 2.0f * PI);
 
 	const UStaticMesh* Mesh;
-	if (FGLTFActorUtility::TryGetPropertyValue(BackdropActor, TEXT("Mesh"), Mesh))
+	if (FGLTFBlueprintUtility::TryGetPropertyValue(BackdropActor, TEXT("Mesh"), Mesh))
 	{
 		JsonBackdrop->Mesh = Builder.AddUniqueMesh(Mesh, { FGLTFMaterialUtility::GetDefaultMaterial() });
 	}
@@ -37,7 +37,7 @@ FGLTFJsonBackdrop* FGLTFBackdropConverter::Convert(const AActor* BackdropActor)
 	}
 
 	const UTextureCube* Cubemap;
-	if (FGLTFActorUtility::TryGetPropertyValue(BackdropActor, TEXT("Cubemap"), Cubemap))
+	if (FGLTFBlueprintUtility::TryGetPropertyValue(BackdropActor, TEXT("Cubemap"), Cubemap))
 	{
 		// TODO: add a proper custom gltf extension with its own converters for cubemaps
 
@@ -54,7 +54,7 @@ FGLTFJsonBackdrop* FGLTFBackdropConverter::Convert(const AActor* BackdropActor)
 	}
 
 	float Intensity;
-	if (FGLTFActorUtility::TryGetPropertyValue(BackdropActor, TEXT("Intensity"), Intensity))
+	if (FGLTFBlueprintUtility::TryGetPropertyValue(BackdropActor, TEXT("Intensity"), Intensity))
 	{
 		JsonBackdrop->Intensity = Intensity;
 	}
@@ -64,7 +64,7 @@ FGLTFJsonBackdrop* FGLTFBackdropConverter::Convert(const AActor* BackdropActor)
 	}
 
 	float Size;
-	if (FGLTFActorUtility::TryGetPropertyValue(BackdropActor, TEXT("Size"), Size))
+	if (FGLTFBlueprintUtility::TryGetPropertyValue(BackdropActor, TEXT("Size"), Size))
 	{
 		JsonBackdrop->Size = Size;
 	}
@@ -74,7 +74,7 @@ FGLTFJsonBackdrop* FGLTFBackdropConverter::Convert(const AActor* BackdropActor)
 	}
 
 	FVector ProjectionCenter;
-	if (FGLTFActorUtility::TryGetPropertyValue(BackdropActor, TEXT("ProjectionCenter"), ProjectionCenter))
+	if (FGLTFBlueprintUtility::TryGetPropertyValue(BackdropActor, TEXT("ProjectionCenter"), ProjectionCenter))
 	{
 		JsonBackdrop->ProjectionCenter = FGLTFCoreUtilities::ConvertPosition(ProjectionCenter, Builder.ExportOptions->ExportUniformScale);
 	}
@@ -84,7 +84,7 @@ FGLTFJsonBackdrop* FGLTFBackdropConverter::Convert(const AActor* BackdropActor)
 	}
 
 	float LightingDistanceFactor;
-	if (FGLTFActorUtility::TryGetPropertyValue(BackdropActor, TEXT("LightingDistanceFactor"), LightingDistanceFactor))
+	if (FGLTFBlueprintUtility::TryGetPropertyValue(BackdropActor, TEXT("LightingDistanceFactor"), LightingDistanceFactor))
 	{
 		JsonBackdrop->LightingDistanceFactor = LightingDistanceFactor;
 	}
@@ -94,7 +94,7 @@ FGLTFJsonBackdrop* FGLTFBackdropConverter::Convert(const AActor* BackdropActor)
 	}
 
 	bool UseCameraProjection;
-	if (FGLTFActorUtility::TryGetPropertyValue(BackdropActor, TEXT("UseCameraProjection"), UseCameraProjection))
+	if (FGLTFBlueprintUtility::TryGetPropertyValue(BackdropActor, TEXT("UseCameraProjection"), UseCameraProjection))
 	{
 		JsonBackdrop->UseCameraProjection = UseCameraProjection;
 	}
diff --git a/Engine/Plugins/Enterprise/GLTFExporter/MigratedFromGit/Source/GLTFExporter/Private/Converters/GLTFBlueprintUtility.cpp b/Engine/Plugins/Enterprise/GLTFExporter/MigratedFromGit/Source/GLTFExporter/Private/Converters/GLTFBlueprintUtility.cpp
new file mode 100644
--- /dev/null
+++ b/Engine/Plugins/Enterprise/GLTFExporter/MigratedFromGit/Source/GLTFExporter/Private/Converters/GLTFBlueprintUtility.cpp
@@ -0,0 +1,28 @@
+// Copyright Epic Games, Inc. All Rights Reserved.
+
+#include "Converters/GLTFBlueprintUtility.h"
+#include "Actors/GLTFHotspotActor.h"
+#include "LevelSequenceActor.h"
+
+FString FGLTFBlueprintUtility::GetClassPath(const AActor* Actor)
+{
+	UClass* Class = Actor->GetClass();
+	if (Class != nullptr && Class->HasAnyClassFlags(CLASS_CompiledFromBlueprint))
+	{
+		return Class->GetPathName();
+	}
+
+	return TEXT("");
+}
+
+bool FGLTFBlueprintUtility::IsSkySphere(const FString& Path)
+{
+	// TODO: what if a blueprint inherits BP_Sky_Sphere?
+	return Path.Equals(TEXT("/Engine/EngineSky/BP_Sky_Sphere.BP_Sky_Sphere_C"));
+}
+
+bool FGLTFBlueprintUtility::IsHDRIBackdrop(const FString& Path)
+{
+	// TODO: what if a blueprint inherits HDRIBackdrop?
+	return Path.Equals(TEXT("/HDRIBackdrop/Blueprints/HDRIBackdrop.HDRIBackdrop_C"));
+}
diff --git a/Engine/Plugins/Enterprise/GLTFExporter/MigratedFromGit/Source/GLTFExporter/Private/Converters/GLTFBlueprintUtility.h b/Engine/Plugins/Enterprise/GLTFExporter/MigratedFromGit/Source/GLTFExporter/Private/Converters/GLTFBlueprintUtility.h
new file mode 100644
--- /dev/null
+++ b/Engine/Plugins/Enterprise/GLTFExporter/MigratedFromGit/Source/GLTFExporter/Private/Converters/GLTFBlueprintUtility.h
@@ -0,0 +1,33 @@
+// Copyright Epic Games, Inc. All Rights Reserved.
+
+#pragma once
+
+#include "CoreMinimal.h"
+
+struct FGLTFBlueprintUtility
+{
+	static FString GetClassPath(const AActor* Actor);
+
+	static bool IsSkySphere(const FString& Path);
+
+	static bool IsHDRIBackdrop(const FString& Path);
+
+	template <class ValueType>
+	static bool TryGetPropertyValue(const UObject* Object, const TCHAR* PropertyName, ValueType& Value)
+	{
+		FProperty* Property = Object->GetClass()->FindPropertyByName(PropertyName);
+		if (Property == nullptr)
+		{
+			return false;
+		}
+
+		const ValueType* ValuePtr = Property->ContainerPtrToValuePtr<ValueType>(Object);
+		if (ValuePtr == nullptr)
+		{
+			return false;
+		}
+
+		Value = *ValuePtr;
+		return true;
+	}
+};
diff --git a/Engine/Plugins/Enterprise/GLTFExporter/MigratedFromGit/Source/GLTFExporter/Private/Converters/GLTFNodeConverters.cpp b/Engine/Plugins/Enterprise/GLTFExporter/MigratedFromGit/Source/GLTFExporter/Private/Converters/GLTFNodeConverters.cpp
--- a/Engine/Plugins/Enterprise/GLTFExporter/MigratedFromGit/Source/GLTFExporter/Private/Converters/GLTFNodeConverters.cpp
+++ b/Engine/Plugins/Enterprise/GLTFExporter/MigratedFromGit/Source/GLTFExporter/Private/Converters/GLTFNodeConverters.cpp
@@ -3,7 +3,7 @@
 #include "Converters/GLTFNodeConverters.h"
 #include "Builders/GLTFContainerBuilder.h"
 #include "Utilities/GLTFCoreUtilities.h"
-#include "Converters/GLTFActorUtility.h"
+#include "Converters/GLTFBlueprintUtility.h"
 #include "Converters/GLTFNameUtility.h"
 #include "Actors/GLTFHotspotActor.h"
 #include "LevelSequenceActor.h"
@@ -29,15 +29,15 @@ FGLTFJsonNode* FGLTFActorConverter::Convert(const AActor* Actor)
 
 	// TODO: process all components since any component can be attached to any other component in runtime
 
-	const FString BlueprintPath = FGLTFActorUtility::GetBlueprintPath(Actor);
-	if (FGLTFActorUtility::IsSkySphereBlueprint(BlueprintPath))
+	const FString BlueprintPath = FGLTFBlueprintUtility::GetClassPath(Actor);
+	if (FGLTFBlueprintUtility::IsSkySphere(BlueprintPath))
 	{
 		if (Builder.ExportOptions->bExportSkySpheres)
 		{
 			RootNode->SkySphere = Builder.AddUniqueSkySphere(Actor);
 		}
 	}
-	else if (FGLTFActorUtility::IsHDRIBackdropBlueprint(BlueprintPath))
+	else if (FGLTFBlueprintUtility::IsHDRIBackdrop(BlueprintPath))
 	{
 		if (Builder.ExportOptions->bExportHDRIBackdrops)
 		{
@@ -131,63 +131,79 @@ FGLTFJsonNode* FGLTFComponentConverter::Convert(const USceneComponent* SceneComp
 		ParentNode->Children.Add(Node);
 	}
 
+	ConvertComponentSpecialization(SceneComponent, Owner, Node);
+	return Node;
+}
+
+void FGLTFComponentConverter::ConvertComponentSpecialization(const USceneComponent* SceneComponent, const AActor* Owner, FGLTFJsonNode* Node)
+{
 	// TODO: don't export invisible components unless visibility is variable due to variant sets
 
-	// TODO: should hidden in game be configurable like this?
-	if ((Builder.ExportOptions->bExportHiddenInGame || (!SceneComponent->bHiddenInGame && !Owner->IsHidden())) && FGLTFActorUtility::IsGenericActor(Owner))
+	if (!Builder.ExportOptions->bExportHiddenInGame && (SceneComponent->bHiddenInGame || Owner->IsHidden()))
 	{
-		if (const UStaticMeshComponent* StaticMeshComponent = Cast<UStaticMeshComponent>(SceneComponent))
-		{
-			Node->Mesh = Builder.AddUniqueMesh(StaticMeshComponent);
+		return;
+	}
 
-			if (Builder.ExportOptions->bExportLightmaps)
-			{
-				Node->LightMap = Builder.AddUniqueLightMap(StaticMeshComponent);
-			}
-		}
-		else if (const USkeletalMeshComponent* SkeletalMeshComponent = Cast<USkeletalMeshComponent>(SceneComponent))
+	const FString BlueprintPath = FGLTFBlueprintUtility::GetClassPath(Owner);
+	if (FGLTFBlueprintUtility::IsSkySphere(BlueprintPath) || FGLTFBlueprintUtility::IsHDRIBackdrop(BlueprintPath))
+	{
+		return;
+	}
+
+	if (Owner->IsA<ALevelSequenceActor>() || Owner->IsA<APawn>() || Owner->IsA<AGLTFHotspotActor>())
+	{
+		return;
+	}
+
+	if (const UStaticMeshComponent* StaticMeshComponent = Cast<UStaticMeshComponent>(SceneComponent))
+	{
+		Node->Mesh = Builder.AddUniqueMesh(StaticMeshComponent);
+
+		if (Builder.ExportOptions->bExportLightmaps)
 		{
-			Node->Mesh = Builder.AddUniqueMesh(SkeletalMeshComponent);
+			Node->LightMap = Builder.AddUniqueLightMap(StaticMeshComponent);
+		}
+	}
+	else if (const USkeletalMeshComponent* SkeletalMeshComponent = Cast<USkeletalMeshComponent>(SceneComponent))
+	{
+		Node->Mesh = Builder.AddUniqueMesh(SkeletalMeshComponent);
 
-			if (Builder.ExportOptions->bExportVertexSkinWeights)
+		if (Builder.ExportOptions->bExportVertexSkinWeights)
+		{
+			Node->Skin = Builder.AddUniqueSkin(Node, SkeletalMeshComponent);
+			if (Node->Skin != nullptr)
 			{
-				Node->Skin = Builder.AddUniqueSkin(Node, SkeletalMeshComponent);
-				if (Node->Skin != nullptr)
+				if (Builder.ExportOptions->bExportAnimationSequences)
 				{
-					if (Builder.ExportOptions->bExportAnimationSequences)
-					{
-						Builder.AddUniqueAnimation(Node, SkeletalMeshComponent);
-					}
+					Builder.AddUniqueAnimation(Node, SkeletalMeshComponent);
 				}
 			}
 		}
-		else if (const UCameraComponent* CameraComponent = Cast<UCameraComponent>(SceneComponent))
+	}
+	else if (const UCameraComponent* CameraComponent = Cast<UCameraComponent>(SceneComponent))
+	{
+		if (Builder.ExportOptions->bExportCameras)
 		{
-			if (Builder.ExportOptions->bExportCameras)
-			{
-				// TODO: conversion of camera direction should be done in separate converter
-				FGLTFJsonNode* CameraNode = Builder.AddNode();
-				CameraNode->Name = FGLTFNameUtility::GetName(CameraComponent);
-				CameraNode->Rotation = FGLTFCoreUtilities::GetLocalCameraRotation();
-				CameraNode->Camera = Builder.AddUniqueCamera(CameraComponent);
-				Node->Children.Add(CameraNode);
-			}
+			// TODO: conversion of camera direction should be done in separate converter
+			FGLTFJsonNode* CameraNode = Builder.AddNode();
+			CameraNode->Name = FGLTFNameUtility::GetName(CameraComponent);
+			CameraNode->Rotation = FGLTFCoreUtilities::GetLocalCameraRotation();
+			CameraNode->Camera = Builder.AddUniqueCamera(CameraComponent);
+			Node->Children.Add(CameraNode);
 		}
-		else if (const ULightComponent* LightComponent = Cast<ULightComponent>(SceneComponent))
+	}
+	else if (const ULightComponent* LightComponent = Cast<ULightComponent>(SceneComponent))
+	{
+		if (Builder.ShouldExportLight(LightComponent->Mobility))
 		{
-			if (Builder.ShouldExportLight(LightComponent->Mobility))
-			{
-				// TODO: conversion of light direction should be done in separate converter
-				FGLTFJsonNode* LightNode = Builder.AddNode();
-				LightNode->Name = FGLTFNameUtility::GetName(LightComponent);
-				LightNode->Rotation = FGLTFCoreUtilities::GetLocalLightRotation();
-				LightNode->Light = Builder.AddUniqueLight(LightComponent);
-				Node->Children.Add(LightNode);
-			}
+			// TODO: conversion of light direction should be done in separate converter
+			FGLTFJsonNode* LightNode = Builder.AddNode();
+			LightNode->Name = FGLTFNameUtility::GetName(LightComponent);
+			LightNode->Rotation = FGLTFCoreUtilities::GetLocalLightRotation();
+			LightNode->Light = Builder.AddUniqueLight(LightComponent);
+			Node->Children.Add(LightNode);
 		}
 	}
-
-	return Node;
 }
 
 FGLTFJsonNode* FGLTFComponentSocketConverter::Convert(const USceneComponent* SceneComponent, FName SocketName)
diff --git a/Engine/Plugins/Enterprise/GLTFExporter/MigratedFromGit/Source/GLTFExporter/Private/Converters/GLTFSkySphereConverters.cpp b/Engine/Plugins/Enterprise/GLTFExporter/MigratedFromGit/Source/GLTFExporter/Private/Converters/GLTFSkySphereConverters.cpp
--- a/Engine/Plugins/Enterprise/GLTFExporter/MigratedFromGit/Source/GLTFExporter/Private/Converters/GLTFSkySphereConverters.cpp
+++ b/Engine/Plugins/Enterprise/GLTFExporter/MigratedFromGit/Source/GLTFExporter/Private/Converters/GLTFSkySphereConverters.cpp
@@ -3,7 +3,7 @@
 #include "Converters/GLTFSkySphereConverters.h"
 #include "Utilities/GLTFCoreUtilities.h"
 #include "Converters/GLTFMaterialUtility.h"
-#include "Converters/GLTFActorUtility.h"
+#include "Converters/GLTFBlueprintUtility.h"
 #include "Converters/GLTFCurveUtility.h"
 #include "Builders/GLTFContainerBuilder.h"
 #include "Engine/DirectionalLight.h"
@@ -12,7 +12,7 @@
 
 FGLTFJsonSkySphere* FGLTFSkySphereConverter::Convert(const AActor* SkySphereActor)
 {
-	if (!FGLTFActorUtility::IsSkySphereBlueprint(FGLTFActorUtility::GetBlueprintPath(SkySphereActor)))
+	if (!FGLTFBlueprintUtility::IsSkySphere(FGLTFBlueprintUtility::GetClassPath(SkySphereActor)))
 	{
 		return nullptr;
 	}
@@ -21,7 +21,7 @@ FGLTFJsonSkySphere* FGLTFSkySphereConverter::Convert(const AActor* SkySphereActo
 	SkySphereActor->GetName(JsonSkySphere->Name);
 
 	const UStaticMeshComponent* StaticMeshComponent = nullptr;
-	FGLTFActorUtility::TryGetPropertyValue(SkySphereActor, TEXT("SkySphereMesh"), StaticMeshComponent);
+	FGLTFBlueprintUtility::TryGetPropertyValue(SkySphereActor, TEXT("SkySphereMesh"), StaticMeshComponent);
 
 	if (StaticMeshComponent != nullptr)
 	{
@@ -41,7 +41,7 @@ FGLTFJsonSkySphere* FGLTFSkySphereConverter::Convert(const AActor* SkySphereActo
 	}
 
 	const ADirectionalLight* DirectionalLight = nullptr;
-	if (FGLTFActorUtility::TryGetPropertyValue(SkySphereActor, TEXT("Directional light actor"), DirectionalLight))
+	if (FGLTFBlueprintUtility::TryGetPropertyValue(SkySphereActor, TEXT("Directional light actor"), DirectionalLight))
 	{
 		JsonSkySphere->DirectionalLight = Builder.AddUniqueNode(DirectionalLight);
 	}
@@ -51,7 +51,7 @@ FGLTFJsonSkySphere* FGLTFSkySphereConverter::Convert(const AActor* SkySphereActo
 	}
 
 	const UMaterialInstance* SkyMaterial = nullptr;
-	FGLTFActorUtility::TryGetPropertyValue(SkySphereActor, TEXT("Sky material"), SkyMaterial);
+	FGLTFBlueprintUtility::TryGetPropertyValue(SkySphereActor, TEXT("Sky material"), SkyMaterial);
 
 	if (SkyMaterial != nullptr)
 	{
@@ -95,7 +95,7 @@ void FGLTFSkySphereConverter::ConvertProperty(const AActor* Actor, const TCHAR*
 {
 	check(Actor != nullptr);
 
-	if (!FGLTFActorUtility::TryGetPropertyValue(Actor, PropertyName, OutValue))
+	if (!FGLTFBlueprintUtility::TryGetPropertyValue(Actor, PropertyName, OutValue))
 	{
 		Builder.LogWarning(FString::Printf(
 			TEXT("Failed to export property %s for Sky Sphere %s"),
@@ -109,7 +109,7 @@ void FGLTFSkySphereConverter::ConvertColorProperty(const AActor* Actor, const TC
 	check(Actor != nullptr);
 
 	FLinearColor LinearColor;
-	if (FGLTFActorUtility::TryGetPropertyValue(Actor, PropertyName, LinearColor))
+	if (FGLTFBlueprintUtility::TryGetPropertyValue(Actor, PropertyName, LinearColor))
 	{
 		OutValue = FGLTFCoreUtilities::ConvertColor(LinearColor, Builder.ExportOptions->bStrictCompliance);
 	}
@@ -127,7 +127,7 @@ void FGLTFSkySphereConverter::ConvertColorCurveProperty(const AActor* Actor, con
 	check(Actor != nullptr);
 
 	const UCurveLinearColor* ColorCurve = nullptr;
-	FGLTFActorUtility::TryGetPropertyValue(Actor, PropertyName, ColorCurve);
+	FGLTFBlueprintUtility::TryGetPropertyValue(Actor, PropertyName, ColorCurve);
 
 	if (ColorCurve != nullptr)
 	{
diff --git a/Engine/Plugins/Enterprise/GLTFExporter/MigratedFromGit/Source/GLTFExporter/Public/Converters/GLTFNodeConverters.h b/Engine/Plugins/Enterprise/GLTFExporter/MigratedFromGit/Source/GLTFExporter/Public/Converters/GLTFNodeConverters.h
--- a/Engine/Plugins/Enterprise/GLTFExporter/MigratedFromGit/Source/GLTFExporter/Public/Converters/GLTFNodeConverters.h
+++ b/Engine/Plugins/Enterprise/GLTFExporter/MigratedFromGit/Source/GLTFExporter/Public/Converters/GLTFNodeConverters.h
@@ -25,6 +25,8 @@ class GLTFEXPORTER_API FGLTFComponentConverter final : public FGLTFBuilderContex
 	using FGLTFBuilderContext::FGLTFBuilderContext;
 
 	virtual FGLTFJsonNode* Convert(const USceneComponent* SceneComponent) override;
+
+	void ConvertComponentSpecialization(const USceneComponent* SceneComponent, const AActor* Owner, FGLTFJsonNode* Node);
 };
 
 class GLTFEXPORTER_API FGLTFComponentSocketConverter final : public FGLTFBuilderContext, public IGLTFComponentSocketConverter
