[GLTFExporter] fix: use double instead of float when reading properties in blueprint (like backdrop and skysphere)

commit 565840acd68742c6dd184d0c1ad6b92c25ba302b
Date: Sun, 27 Mar 2022 12:27:48 +0200

    fix: use double instead of float when reading properties in blueprint (like backdrop and skysphere)
    (cherry picked from commit 8947a1778f47a8c1457178d2c0291511693f36bb​)

#skipci

diff --git a/Engine/Plugins/Enterprise/GLTFExporter/MigratedFromGit/Source/GLTFExporter/Private/Converters/GLTFBlueprintUtility.h b/Engine/Plugins/Enterprise/GLTFExporter/MigratedFromGit/Source/GLTFExporter/Private/Converters/GLTFBlueprintUtility.h
--- a/Engine/Plugins/Enterprise/GLTFExporter/MigratedFromGit/Source/GLTFExporter/Private/Converters/GLTFBlueprintUtility.h
+++ b/Engine/Plugins/Enterprise/GLTFExporter/MigratedFromGit/Source/GLTFExporter/Private/Converters/GLTFBlueprintUtility.h
@@ -13,7 +13,7 @@ struct FGLTFBlueprintUtility
 	static bool IsHDRIBackdrop(const FString& Path);
 
 	template <class ValueType>
-	static bool TryGetPropertyValue(const UObject* Object, const TCHAR* PropertyName, ValueType& Value)
+	static bool TryGetPropertyValue(const UObject* Object, const TCHAR* PropertyName, ValueType& OutValue)
 	{
 		FProperty* Property = Object->GetClass()->FindPropertyByName(PropertyName);
 		if (Property == nullptr)
@@ -27,7 +27,21 @@ struct FGLTFBlueprintUtility
 			return false;
 		}
 
-		Value = *ValuePtr;
+		OutValue = *ValuePtr;
+		return true;
+	}
+
+	static bool TryGetPropertyValue(const UObject* Object, const TCHAR* PropertyName, float& OutValue)
+	{
+		// NOTE: blueprints always uses double instead of float in UE5
+
+		double Value;
+		if (!TryGetPropertyValue(Object, PropertyName, Value))
+		{
+			return false;
+		}
+
+		OutValue = static_cast<float>(Value);
 		return true;
 	}
 };
