[GLTFExporter] fix: replace FVector4 with FVector4f in FGLTFMaterialUtility

commit 5d174d98818044408e8826c0893960f63a23e2ad
Date: Sun, 27 Mar 2022 13:31:34 +0200

    fix: replace FVector4 with FVector4f in FGLTFMaterialUtility
    (cherry picked from commit aa9c87f9664884298fa6902d0ac6a75297da0454​)

#skipci

diff --git a/Engine/Plugins/Enterprise/GLTFExporter/MigratedFromGit/Source/GLTFExporter/Private/Converters/GLTFMaterialUtility.cpp b/Engine/Plugins/Enterprise/GLTFExporter/MigratedFromGit/Source/GLTFExporter/Private/Converters/GLTFMaterialUtility.cpp
--- a/Engine/Plugins/Enterprise/GLTFExporter/MigratedFromGit/Source/GLTFExporter/Private/Converters/GLTFMaterialUtility.cpp
+++ b/Engine/Plugins/Enterprise/GLTFExporter/MigratedFromGit/Source/GLTFExporter/Private/Converters/GLTFMaterialUtility.cpp
@@ -50,116 +50,24 @@ FGuid FGLTFMaterialUtility::GetAttributeIDChecked(const FGLTFMaterialPropertyEx&
 }
 #endif
 
-FVector4 FGLTFMaterialUtility::GetPropertyDefaultValue(const FGLTFMaterialPropertyEx& Property)
+FVector4f FGLTFMaterialUtility::GetPropertyDefaultValue(const FGLTFMaterialPropertyEx& Property)
 {
-#if (ENGINE_MAJOR_VERSION > 4 || ENGINE_MINOR_VERSION >= 27)
 	return FMaterialAttributeDefinitionMap::GetDefaultValue(GetAttributeIDChecked(Property));
-
-#else
-	switch (Property.Type)
-	{
-		case MP_EmissiveColor:          return FVector4(0,0,0,0);
-		case MP_Opacity:                return FVector4(1,0,0,0);
-		case MP_OpacityMask:            return FVector4(1,0,0,0);
-		case MP_BaseColor:              return FVector4(0,0,0,0);
-		case MP_Metallic:               return FVector4(0,0,0,0);
-		case MP_Specular:               return FVector4(.5,0,0,0);
-		case MP_Roughness:              return FVector4(.5,0,0,0);
-		case MP_Anisotropy:             return FVector4(0,0,0,0);
-		case MP_Normal:                 return FVector4(0,0,1,0);
-		case MP_Tangent:                return FVector4(1,0,0,0);
-		case MP_WorldPositionOffset:    return FVector4(0,0,0,0);
-		case MP_WorldDisplacement:      return FVector4(0,0,0,0);
-		case MP_TessellationMultiplier: return FVector4(1,0,0,0);
-		case MP_SubsurfaceColor:        return FVector4(1,1,1,0);
-		case MP_CustomData0:            return FVector4(1,0,0,0);
-		case MP_CustomData1:            return FVector4(.1,0,0,0);
-		case MP_AmbientOcclusion:       return FVector4(1,0,0,0);
-		case MP_Refraction:             return FVector4(1,0,0,0);
-		case MP_PixelDepthOffset:       return FVector4(0,0,0,0);
-		case MP_ShadingModel:           return FVector4(0,0,0,0);
-		default:                        break;
-	}
-
-	if (Property.Type >= MP_CustomizedUVs0 && Property.Type <= MP_CustomizedUVs7)
-	{
-		return FVector4(0,0,0,0);
-	}
-
-	if (Property == TEXT("ClearCoatBottomNormal"))
-	{
-		return FVector4(0,0,1,0);
-	}
-
-	if (Property == TEXT("TransmittanceColor"))
-	{
-		return FVector4(.5,.5,.5,0);
-	}
-
-	check(false);
-	return FVector4();
-#endif
 }
 
-FVector4 FGLTFMaterialUtility::GetPropertyMask(const FGLTFMaterialPropertyEx& Property)
+FVector4f FGLTFMaterialUtility::GetPropertyMask(const FGLTFMaterialPropertyEx& Property)
 {
-#if (ENGINE_MAJOR_VERSION > 4 || ENGINE_MINOR_VERSION >= 27)
 	switch (FMaterialAttributeDefinitionMap::GetValueType(GetAttributeIDChecked(Property)))
 	{
 		case MCT_Float:
-		case MCT_Float1: return FVector4(1, 0, 0, 0);
-		case MCT_Float2: return FVector4(1, 1, 0, 0);
-		case MCT_Float3: return FVector4(1, 1, 1, 0);
-		case MCT_Float4: return FVector4(1, 1, 1, 1);
+		case MCT_Float1: return FVector4f(1, 0, 0, 0);
+		case MCT_Float2: return FVector4f(1, 1, 0, 0);
+		case MCT_Float3: return FVector4f(1, 1, 1, 0);
+		case MCT_Float4: return FVector4f(1, 1, 1, 1);
 		default:
 			checkNoEntry();
-			return FVector4();
-	}
-
-#else
-	switch (Property.Type)
-	{
-		case MP_EmissiveColor:          return FVector4(1,1,1,0);
-		case MP_Opacity:                return FVector4(1,0,0,0);
-		case MP_OpacityMask:            return FVector4(1,0,0,0);
-		case MP_BaseColor:              return FVector4(1,1,1,0);
-		case MP_Metallic:               return FVector4(1,0,0,0);
-		case MP_Specular:               return FVector4(1,0,0,0);
-		case MP_Roughness:              return FVector4(1,0,0,0);
-		case MP_Anisotropy:             return FVector4(1,0,0,0);
-		case MP_Normal:                 return FVector4(1,1,1,0);
-		case MP_Tangent:                return FVector4(1,1,1,0);
-		case MP_WorldPositionOffset:    return FVector4(1,1,1,0);
-		case MP_WorldDisplacement:      return FVector4(1,1,1,0);
-		case MP_TessellationMultiplier: return FVector4(1,0,0,0);
-		case MP_SubsurfaceColor:        return FVector4(1,1,1,0);
-		case MP_CustomData0:            return FVector4(1,0,0,0);
-		case MP_CustomData1:            return FVector4(1,0,0,0);
-		case MP_AmbientOcclusion:       return FVector4(1,0,0,0);
-		case MP_Refraction:             return FVector4(1,1,0,0);
-		case MP_PixelDepthOffset:       return FVector4(1,0,0,0);
-		case MP_ShadingModel:           return FVector4(1,0,0,0);
-		default:                        break;
+			return FVector4f();
 	}
-
-	if (Property.Type >= MP_CustomizedUVs0 && Property.Type <= MP_CustomizedUVs7)
-	{
-		return FVector4(1,1,0,0);
-	}
-
-	if (Property == TEXT("ClearCoatBottomNormal"))
-	{
-		return FVector4(1,1,1,0);
-	}
-
-	if (Property == TEXT("TransmittanceColor"))
-	{
-		return FVector4(1,1,1,0);
-	}
-
-	check(false);
-	return FVector4();
-#endif
 }
 
 const FExpressionInput* FGLTFMaterialUtility::GetInputForProperty(const UMaterialInterface* Material, const FGLTFMaterialPropertyEx& Property)
diff --git a/Engine/Plugins/Enterprise/GLTFExporter/MigratedFromGit/Source/GLTFExporter/Private/Converters/GLTFMaterialUtility.h b/Engine/Plugins/Enterprise/GLTFExporter/MigratedFromGit/Source/GLTFExporter/Private/Converters/GLTFMaterialUtility.h
--- a/Engine/Plugins/Enterprise/GLTFExporter/MigratedFromGit/Source/GLTFExporter/Private/Converters/GLTFMaterialUtility.h
+++ b/Engine/Plugins/Enterprise/GLTFExporter/MigratedFromGit/Source/GLTFExporter/Private/Converters/GLTFMaterialUtility.h
@@ -38,8 +38,8 @@ struct FGLTFMaterialUtility
 	static FGuid GetAttributeIDChecked(const FGLTFMaterialPropertyEx& Property);
 #endif
 
-	static FVector4 GetPropertyDefaultValue(const FGLTFMaterialPropertyEx& Property);
-	static FVector4 GetPropertyMask(const FGLTFMaterialPropertyEx& Property);
+	static FVector4f GetPropertyDefaultValue(const FGLTFMaterialPropertyEx& Property);
+	static FVector4f GetPropertyMask(const FGLTFMaterialPropertyEx& Property);
 
 	static const FExpressionInput* GetInputForProperty(const UMaterialInterface* Material, const FGLTFMaterialPropertyEx& Property);
 
diff --git a/Engine/Plugins/Enterprise/GLTFExporter/MigratedFromGit/Source/GLTFExporter/Private/Tasks/GLTFDelayedMaterialTasks.cpp b/Engine/Plugins/Enterprise/GLTFExporter/MigratedFromGit/Source/GLTFExporter/Private/Tasks/GLTFDelayedMaterialTasks.cpp
--- a/Engine/Plugins/Enterprise/GLTFExporter/MigratedFromGit/Source/GLTFExporter/Private/Tasks/GLTFDelayedMaterialTasks.cpp
+++ b/Engine/Plugins/Enterprise/GLTFExporter/MigratedFromGit/Source/GLTFExporter/Private/Tasks/GLTFDelayedMaterialTasks.cpp
@@ -1434,7 +1434,7 @@ bool FGLTFDelayedMaterialTask::TryGetBakedMaterialProperty(FGLTFJsonTextureInfo&
 		return StoreBakedPropertyTexture(OutTexInfo, PropertyBakeOutput, PropertyName);
 	}
 
-	const FVector4 MaskedConstant = FVector4(PropertyBakeOutput.ConstantValue) * FGLTFMaterialUtility::GetPropertyMask(Property);
+	const FVector4f MaskedConstant = FVector4f(PropertyBakeOutput.ConstantValue) * FGLTFMaterialUtility::GetPropertyMask(Property);
 	if (MaskedConstant == FGLTFMaterialUtility::GetPropertyDefaultValue(Property))
 	{
 		// Constant value is the same as the property's default so we can set gltf to default.
