[GLTFExporter] fix: make sure mesh assets are fully loaded/streamed before converting

commit 7b48d0edce8cebeb933435ed862dfa49b7b789e7
Date: Mon, 28 Mar 2022 14:22:04 +0200

    fix: make sure mesh assets are fully loaded/streamed before converting
    (cherry picked from commit eb0a7d67645cf75ae90eafd425c618d03a7cbd58​)

#skipci

diff --git a/Engine/Plugins/Enterprise/GLTFExporter/MigratedFromGit/Source/GLTFExporter/Private/Converters/GLTFMeshUtility.cpp b/Engine/Plugins/Enterprise/GLTFExporter/MigratedFromGit/Source/GLTFExporter/Private/Converters/GLTFMeshUtility.cpp
--- a/Engine/Plugins/Enterprise/GLTFExporter/MigratedFromGit/Source/GLTFExporter/Private/Converters/GLTFMeshUtility.cpp
+++ b/Engine/Plugins/Enterprise/GLTFExporter/MigratedFromGit/Source/GLTFExporter/Private/Converters/GLTFMeshUtility.cpp
@@ -2,6 +2,32 @@
 
 #include "Converters/GLTFMeshUtility.h"
 #include "Rendering/SkeletalMeshRenderData.h"
+#include "StaticMeshCompiler.h"
+#include "SkeletalMeshCompiler.h"
+
+void FGLTFMeshUtility::FullyLoad(const UStaticMesh* InStaticMesh)
+{
+	UStaticMesh* StaticMesh = const_cast<UStaticMesh*>(InStaticMesh);
+
+#if WITH_EDITOR
+	FStaticMeshCompilingManager::Get().FinishCompilation({ StaticMesh });
+#endif
+
+	StaticMesh->SetForceMipLevelsToBeResident(30.0f);
+	StaticMesh->WaitForStreaming();
+}
+
+void FGLTFMeshUtility::FullyLoad(const USkeletalMesh* InSkeletalMesh)
+{
+	USkeletalMesh* SkeletalMesh = const_cast<USkeletalMesh*>(InSkeletalMesh);
+
+#if WITH_EDITOR
+	FSkeletalMeshCompilingManager::Get().FinishCompilation({ SkeletalMesh });
+#endif
+
+	SkeletalMesh->SetForceMipLevelsToBeResident(30.0f);
+	SkeletalMesh->WaitForStreaming();
+}
 
 const UMaterialInterface* FGLTFMeshUtility::GetDefaultMaterial()
 {
diff --git a/Engine/Plugins/Enterprise/GLTFExporter/MigratedFromGit/Source/GLTFExporter/Private/Converters/GLTFMeshUtility.h b/Engine/Plugins/Enterprise/GLTFExporter/MigratedFromGit/Source/GLTFExporter/Private/Converters/GLTFMeshUtility.h
--- a/Engine/Plugins/Enterprise/GLTFExporter/MigratedFromGit/Source/GLTFExporter/Private/Converters/GLTFMeshUtility.h
+++ b/Engine/Plugins/Enterprise/GLTFExporter/MigratedFromGit/Source/GLTFExporter/Private/Converters/GLTFMeshUtility.h
@@ -8,6 +8,9 @@ class FSkeletalMeshLODRenderData;
 
 struct FGLTFMeshUtility
 {
+	static void FullyLoad(const UStaticMesh* InStaticMesh);
+	static void FullyLoad(const USkeletalMesh* InSkeletalMesh);
+
 	static const UMaterialInterface* GetDefaultMaterial();
 
 	static const TArray<FStaticMaterial>& GetMaterials(const UStaticMesh* StaticMesh);
diff --git a/Engine/Plugins/Enterprise/GLTFExporter/MigratedFromGit/Source/GLTFExporter/Private/Converters/GLTFTextureUtility.cpp b/Engine/Plugins/Enterprise/GLTFExporter/MigratedFromGit/Source/GLTFExporter/Private/Converters/GLTFTextureUtility.cpp
--- a/Engine/Plugins/Enterprise/GLTFExporter/MigratedFromGit/Source/GLTFExporter/Private/Converters/GLTFTextureUtility.cpp
+++ b/Engine/Plugins/Enterprise/GLTFExporter/MigratedFromGit/Source/GLTFExporter/Private/Converters/GLTFTextureUtility.cpp
@@ -34,16 +34,16 @@ bool FGLTFTextureUtility::IsAlphaless(EPixelFormat PixelFormat)
 	}
 }
 
-void FGLTFTextureUtility::FullyLoad(const UTexture* Texture)
+void FGLTFTextureUtility::FullyLoad(const UTexture* InTexture)
 {
-	UTexture* MutableTexture = const_cast<UTexture*>(Texture);
+	UTexture* Texture = const_cast<UTexture*>(InTexture);
 
 #if WITH_EDITOR
-	FTextureCompilingManager::Get().FinishCompilation({ MutableTexture });
+	FTextureCompilingManager::Get().FinishCompilation({ Texture });
 #endif
 
-	MutableTexture->SetForceMipLevelsToBeResident(30.0f);
-	MutableTexture->WaitForStreaming();
+	Texture->SetForceMipLevelsToBeResident(30.0f);
+	Texture->WaitForStreaming();
 }
 
 bool FGLTFTextureUtility::IsHDR(const UTexture* Texture)
diff --git a/Engine/Plugins/Enterprise/GLTFExporter/MigratedFromGit/Source/GLTFExporter/Private/Converters/GLTFTextureUtility.h b/Engine/Plugins/Enterprise/GLTFExporter/MigratedFromGit/Source/GLTFExporter/Private/Converters/GLTFTextureUtility.h
--- a/Engine/Plugins/Enterprise/GLTFExporter/MigratedFromGit/Source/GLTFExporter/Private/Converters/GLTFTextureUtility.h
+++ b/Engine/Plugins/Enterprise/GLTFExporter/MigratedFromGit/Source/GLTFExporter/Private/Converters/GLTFTextureUtility.h
@@ -8,7 +8,7 @@ struct FGLTFTextureUtility
 {
 	static bool IsAlphaless(EPixelFormat PixelFormat);
 
-	static void FullyLoad(const UTexture* Texture);
+	static void FullyLoad(const UTexture* InTexture);
 
 	static bool IsHDR(const UTexture* Texture);
 	static bool IsCubemap(const UTexture* Texture);
diff --git a/Engine/Plugins/Enterprise/GLTFExporter/MigratedFromGit/Source/GLTFExporter/Private/Tasks/GLTFDelayedMeshTasks.cpp b/Engine/Plugins/Enterprise/GLTFExporter/MigratedFromGit/Source/GLTFExporter/Private/Tasks/GLTFDelayedMeshTasks.cpp
--- a/Engine/Plugins/Enterprise/GLTFExporter/MigratedFromGit/Source/GLTFExporter/Private/Tasks/GLTFDelayedMeshTasks.cpp
+++ b/Engine/Plugins/Enterprise/GLTFExporter/MigratedFromGit/Source/GLTFExporter/Private/Tasks/GLTFDelayedMeshTasks.cpp
@@ -109,6 +109,7 @@ namespace
 
 void FGLTFDelayedStaticMeshTask::Process()
 {
+	FGLTFMeshUtility::FullyLoad(StaticMesh);
 	JsonMesh->Name = StaticMeshComponent != nullptr ? FGLTFNameUtility::GetName(StaticMeshComponent) : StaticMesh->GetName();
 
 	const FStaticMeshLODResources& MeshLOD = StaticMesh->GetLODForExport(LODIndex);
@@ -206,6 +207,7 @@ void FGLTFDelayedStaticMeshTask::Process()
 
 void FGLTFDelayedSkeletalMeshTask::Process()
 {
+	FGLTFMeshUtility::FullyLoad(SkeletalMesh);
 	JsonMesh->Name = SkeletalMeshComponent != nullptr ? FGLTFNameUtility::GetName(SkeletalMeshComponent) : SkeletalMesh->GetName();
 
 	const FSkeletalMeshRenderData* RenderData = SkeletalMesh->GetResourceForRendering();
