[GLTFExporter] fix: don't use raw data for animation sequences in runtime

commit 0ce68c1ecbbd74aed4cb6fb652fd1bd3f7b81a82
Date: Mon, 28 Mar 2022 14:46:20 +0200

    fix: don't use raw data for animation sequences in runtime
    (cherry picked from commit ac622990f90638326be3b861fc8619baec06bc68​)

#skipci

diff --git a/Engine/Plugins/Enterprise/GLTFExporter/MigratedFromGit/Source/GLTFExporter/Private/Converters/GLTFBoneUtility.cpp b/Engine/Plugins/Enterprise/GLTFExporter/MigratedFromGit/Source/GLTFExporter/Private/Converters/GLTFBoneUtility.cpp
--- a/Engine/Plugins/Enterprise/GLTFExporter/MigratedFromGit/Source/GLTFExporter/Private/Converters/GLTFBoneUtility.cpp
+++ b/Engine/Plugins/Enterprise/GLTFExporter/MigratedFromGit/Source/GLTFExporter/Private/Converters/GLTFBoneUtility.cpp
@@ -2,10 +2,7 @@
 
 #include "Converters/GLTFBoneUtility.h"
 #include "Sections/MovieScene3DTransformSection.h"
-
-#if (ENGINE_MAJOR_VERSION > 4 || ENGINE_MINOR_VERSION >= 26)
 #include "Animation/AnimationPoseData.h"
-#endif
 
 FTransform FGLTFBoneUtility::GetBindTransform(const FReferenceSkeleton& RefSkeleton, int32 BoneIndex)
 {
@@ -51,14 +48,18 @@ void FGLTFBoneUtility::GetBoneIndices(const USkeleton* Skeleton, TArray<FBoneInd
 
 void FGLTFBoneUtility::GetBoneTransformsByFrame(const UAnimSequence* AnimSequence, const TArray<float>& FrameTimestamps, const TArray<FBoneIndexType>& BoneIndices, TArray<TArray<FTransform>>& OutBoneTransformsByFrame)
 {
-	FMemMark Mark(FMemStack::Get()); // Make sure to free stack allocations made by FCompactPose, FBlendedCurve, and FStackCustomAttributes when end of scope
+#if WITH_EDITOR
+	const bool bUseRawData = AnimSequence->OnlyUseRawData();
+#else
+	const bool bUseRawData = false;
+#endif
+
+	// Make sure to free stack allocations made by FCompactPose, FBlendedCurve, and FStackCustomAttributes when end of scope
+	FMemMark Mark(FMemStack::Get());
 
 	FBoneContainer BoneContainer;
 	BoneContainer.InitializeTo(BoneIndices, FCurveEvaluationOption(true), *AnimSequence->GetSkeleton());
-
-	// NOTE: Some assets may only have raw data, some only compressed data.
-	const bool bHasRawData = AnimSequence->GetRawAnimationData().Num() != 0;
-	BoneContainer.SetUseRAWData(bHasRawData);
+	BoneContainer.SetUseRAWData(bUseRawData);
 
 	const int32 FrameCount = FrameTimestamps.Num();
 	OutBoneTransformsByFrame.AddDefaulted(FrameCount);
@@ -69,29 +70,20 @@ void FGLTFBoneUtility::GetBoneTransformsByFrame(const UAnimSequence* AnimSequenc
 	FBlendedCurve Curve;
 	Curve.InitFrom(BoneContainer);
 
-#if (ENGINE_MAJOR_VERSION > 4 || ENGINE_MINOR_VERSION >= 26)
 	FStackCustomAttributes Attributes;
 	FAnimationPoseData PoseData(Pose, Curve, Attributes);
-#endif
 
 	for (int32 Frame = 0; Frame < FrameCount; ++Frame)
 	{
 		const FAnimExtractContext ExtractionContext(FrameTimestamps[Frame]); // TODO: set bExtractRootMotion?
-#if (ENGINE_MAJOR_VERSION > 4 || ENGINE_MINOR_VERSION >= 26)
-		AnimSequence->GetBonePose(PoseData, ExtractionContext);
-#else
-		AnimSequence->GetBonePose(Pose, Curve, ExtractionContext);
-#endif
+		AnimSequence->GetBonePose(PoseData, ExtractionContext, bUseRawData);
 		Pose.CopyBonesTo(OutBoneTransformsByFrame[Frame]);
 	}
 
 	// Clear all stack allocations to allow FMemMark to free them
 	Pose.Empty();
 	Curve.Empty();
-
-#if (ENGINE_MAJOR_VERSION > 4 || ENGINE_MINOR_VERSION >= 26)
 	Attributes.Empty();
-#endif
 }
 
 const FMovieSceneDoubleChannel* FGLTFBoneUtility::GetTranslationChannels(const UMovieScene3DTransformSection* TransformSection)
