[GLTFExporter] feat: make all default converters inheritable and allow any derived class to fall back on default behavior if needed

commit f2f429857b8c2c96b87e6f2ef4bf31426b5e26da
Date: Tue, 29 Mar 2022 22:47:25 +0200

    feat: make all default converters inheritable and allow any derived class to fall back on default behavior if needed
    (cherry picked from commit 88a855f32f7a9dcfc39f02605c5d28b51ee390a8​)

#skipci

diff --git a/Engine/Plugins/Enterprise/GLTFExporter/MigratedFromGit/Source/GLTFExporter/Public/Converters/GLTFAccessorConverters.h b/Engine/Plugins/Enterprise/GLTFExporter/MigratedFromGit/Source/GLTFExporter/Public/Converters/GLTFAccessorConverters.h
--- a/Engine/Plugins/Enterprise/GLTFExporter/MigratedFromGit/Source/GLTFExporter/Public/Converters/GLTFAccessorConverters.h
+++ b/Engine/Plugins/Enterprise/GLTFExporter/MigratedFromGit/Source/GLTFExporter/Public/Converters/GLTFAccessorConverters.h
@@ -16,56 +16,88 @@ typedef TGLTFConverter<FGLTFJsonAccessor*, const FGLTFMeshSection*, const FSkinW
 typedef TGLTFConverter<FGLTFJsonAccessor*, const FGLTFMeshSection*, const FSkinWeightVertexBuffer*, uint32> IGLTFBoneWeightBufferConverter;
 typedef TGLTFConverter<FGLTFJsonAccessor*, const FGLTFMeshSection*> IGLTFIndexBufferConverter;
 
-class GLTFEXPORTER_API FGLTFPositionBufferConverter final : public FGLTFBuilderContext, public IGLTFPositionBufferConverter
+class GLTFEXPORTER_API FGLTFPositionBufferConverter : public FGLTFBuilderContext, public IGLTFPositionBufferConverter
 {
+public:
+
 	using FGLTFBuilderContext::FGLTFBuilderContext;
 
+protected:
+
 	virtual FGLTFJsonAccessor* Convert(const FGLTFMeshSection* MeshSection, const FPositionVertexBuffer* VertexBuffer) override;
 };
 
-class GLTFEXPORTER_API FGLTFColorBufferConverter final : public FGLTFBuilderContext, public IGLTFColorBufferConverter
+class GLTFEXPORTER_API FGLTFColorBufferConverter : public FGLTFBuilderContext, public IGLTFColorBufferConverter
 {
+public:
+
 	using FGLTFBuilderContext::FGLTFBuilderContext;
 
+protected:
+
 	virtual FGLTFJsonAccessor* Convert(const FGLTFMeshSection* MeshSection, const FColorVertexBuffer* VertexBuffer) override;
 };
 
-class GLTFEXPORTER_API FGLTFNormalBufferConverter final : public FGLTFBuilderContext, public IGLTFNormalBufferConverter
+class GLTFEXPORTER_API FGLTFNormalBufferConverter : public FGLTFBuilderContext, public IGLTFNormalBufferConverter
 {
+public:
+
 	using FGLTFBuilderContext::FGLTFBuilderContext;
 
+protected:
+
 	virtual FGLTFJsonAccessor* Convert(const FGLTFMeshSection* MeshSection, const FStaticMeshVertexBuffer* VertexBuffer) override;
 
+private:
+
 	template <typename DestinationType, typename SourceType>
 	FGLTFJsonBufferView* ConvertBufferView(const FGLTFMeshSection* MeshSection, const void* TangentData);
 };
 
-class GLTFEXPORTER_API FGLTFTangentBufferConverter final : public FGLTFBuilderContext, public IGLTFTangentBufferConverter
+class GLTFEXPORTER_API FGLTFTangentBufferConverter : public FGLTFBuilderContext, public IGLTFTangentBufferConverter
 {
+public:
+
 	using FGLTFBuilderContext::FGLTFBuilderContext;
 
+protected:
+
 	virtual FGLTFJsonAccessor* Convert(const FGLTFMeshSection* MeshSection, const FStaticMeshVertexBuffer* VertexBuffer) override;
 
+private:
+
 	template <typename DestinationType, typename SourceType>
 	FGLTFJsonBufferView* ConvertBufferView(const FGLTFMeshSection* MeshSection, const void* TangentData);
 };
 
-class GLTFEXPORTER_API FGLTFUVBufferConverter final : public FGLTFBuilderContext, public IGLTFUVBufferConverter
+class GLTFEXPORTER_API FGLTFUVBufferConverter : public FGLTFBuilderContext, public IGLTFUVBufferConverter
 {
+public:
+
 	using FGLTFBuilderContext::FGLTFBuilderContext;
 
+protected:
+
 	virtual FGLTFJsonAccessor* Convert(const FGLTFMeshSection* MeshSection, const FStaticMeshVertexBuffer* VertexBuffer, uint32 UVIndex) override;
 
+private:
+
 	template <typename SourceType>
 	FGLTFJsonAccessor* Convert(const FGLTFMeshSection* MeshSection, const FStaticMeshVertexBuffer* VertexBuffer, uint32 UVIndex, const uint8* SourceData) const;
 };
 
-class GLTFEXPORTER_API FGLTFBoneIndexBufferConverter final : public FGLTFBuilderContext, public IGLTFBoneIndexBufferConverter
+class GLTFEXPORTER_API FGLTFBoneIndexBufferConverter : public FGLTFBuilderContext, public IGLTFBoneIndexBufferConverter
 {
+public:
+
 	using FGLTFBuilderContext::FGLTFBuilderContext;
 
+protected:
+
 	virtual FGLTFJsonAccessor* Convert(const FGLTFMeshSection* MeshSection, const FSkinWeightVertexBuffer* VertexBuffer, uint32 InfluenceOffset) override;
 
+private:
+
 	template <typename DestinationType>
 	FGLTFJsonAccessor* Convert(const FGLTFMeshSection* MeshSection, const FSkinWeightVertexBuffer* VertexBuffer, uint32 InfluenceOffset, const uint8* SourceData) const;
 
@@ -76,12 +108,18 @@ class GLTFEXPORTER_API FGLTFBoneIndexBufferConverter final : public FGLTFBuilder
 	FGLTFJsonAccessor* Convert(const FGLTFMeshSection* MeshSection, const FSkinWeightVertexBuffer* VertexBuffer, uint32 InfluenceOffset, const uint8* SourceData, CallbackType GetVertexInfluenceOffsetCount) const;
 };
 
-class GLTFEXPORTER_API FGLTFBoneWeightBufferConverter final : public FGLTFBuilderContext, public IGLTFBoneWeightBufferConverter
+class GLTFEXPORTER_API FGLTFBoneWeightBufferConverter : public FGLTFBuilderContext, public IGLTFBoneWeightBufferConverter
 {
+public:
+
 	using FGLTFBuilderContext::FGLTFBuilderContext;
 
+protected:
+
 	virtual FGLTFJsonAccessor* Convert(const FGLTFMeshSection* MeshSection, const FSkinWeightVertexBuffer* VertexBuffer, uint32 InfluenceOffset) override;
 
+private:
+
 	template <typename BoneIndexType>
 	FGLTFJsonAccessor* Convert(const FGLTFMeshSection* MeshSection, const FSkinWeightVertexBuffer* VertexBuffer, uint32 InfluenceOffset, const uint8* SourceData) const;
 
@@ -89,12 +127,18 @@ class GLTFEXPORTER_API FGLTFBoneWeightBufferConverter final : public FGLTFBuilde
 	FGLTFJsonAccessor* Convert(const FGLTFMeshSection* MeshSection, const FSkinWeightVertexBuffer* VertexBuffer, uint32 InfluenceOffset, const uint8* SourceData, CallbackType GetVertexInfluenceOffsetCount) const;
 };
 
-class GLTFEXPORTER_API FGLTFIndexBufferConverter final : public FGLTFBuilderContext, public IGLTFIndexBufferConverter
+class GLTFEXPORTER_API FGLTFIndexBufferConverter : public FGLTFBuilderContext, public IGLTFIndexBufferConverter
 {
+public:
+
 	using FGLTFBuilderContext::FGLTFBuilderContext;
 
+protected:
+
 	virtual FGLTFJsonAccessor* Convert(const FGLTFMeshSection* MeshSection) override;
 
+private:
+
 	template <typename IndexType>
 	FGLTFJsonAccessor* Convert(const FGLTFMeshSection* MeshSection) const;
 };
diff --git a/Engine/Plugins/Enterprise/GLTFExporter/MigratedFromGit/Source/GLTFExporter/Public/Converters/GLTFAnimationConverters.h b/Engine/Plugins/Enterprise/GLTFExporter/MigratedFromGit/Source/GLTFExporter/Public/Converters/GLTFAnimationConverters.h
--- a/Engine/Plugins/Enterprise/GLTFExporter/MigratedFromGit/Source/GLTFExporter/Public/Converters/GLTFAnimationConverters.h
+++ b/Engine/Plugins/Enterprise/GLTFExporter/MigratedFromGit/Source/GLTFExporter/Public/Converters/GLTFAnimationConverters.h
@@ -14,30 +14,46 @@ typedef TGLTFConverter<FGLTFJsonAnimation*, FGLTFJsonNode*, const USkeletalMeshC
 typedef TGLTFConverter<FGLTFJsonAnimation*, const ULevel*, const ULevelSequence*> IGLTFLevelSequenceConverter;
 typedef TGLTFConverter<FGLTFJsonAnimation*, const ALevelSequenceActor*> IGLTFLevelSequenceDataConverter;
 
-class GLTFEXPORTER_API FGLTFAnimationConverter final : public FGLTFBuilderContext, public IGLTFAnimationConverter
+class GLTFEXPORTER_API FGLTFAnimationConverter : public FGLTFBuilderContext, public IGLTFAnimationConverter
 {
+public:
+
 	using FGLTFBuilderContext::FGLTFBuilderContext;
 
+protected:
+
 	virtual FGLTFJsonAnimation* Convert(FGLTFJsonNode* RootNode, const USkeletalMesh* SkeletalMesh, const UAnimSequence* AnimSequence) override;
 };
 
-class GLTFEXPORTER_API FGLTFAnimationDataConverter final : public FGLTFBuilderContext, public IGLTFAnimationDataConverter
+class GLTFEXPORTER_API FGLTFAnimationDataConverter : public FGLTFBuilderContext, public IGLTFAnimationDataConverter
 {
+public:
+
 	using FGLTFBuilderContext::FGLTFBuilderContext;
 
+protected:
+
 	virtual FGLTFJsonAnimation* Convert(FGLTFJsonNode* RootNode, const USkeletalMeshComponent* SkeletalMeshComponent) override;
 };
 
-class GLTFEXPORTER_API FGLTFLevelSequenceConverter final : public FGLTFBuilderContext, public IGLTFLevelSequenceConverter
+class GLTFEXPORTER_API FGLTFLevelSequenceConverter : public FGLTFBuilderContext, public IGLTFLevelSequenceConverter
 {
+public:
+
 	using FGLTFBuilderContext::FGLTFBuilderContext;
 
+protected:
+
 	virtual FGLTFJsonAnimation* Convert(const ULevel* Level, const ULevelSequence*) override;
 };
 
-class GLTFEXPORTER_API FGLTFLevelSequenceDataConverter final : public FGLTFBuilderContext, public IGLTFLevelSequenceDataConverter
+class GLTFEXPORTER_API FGLTFLevelSequenceDataConverter : public FGLTFBuilderContext, public IGLTFLevelSequenceDataConverter
 {
+public:
+
 	using FGLTFBuilderContext::FGLTFBuilderContext;
 
+protected:
+
 	virtual FGLTFJsonAnimation* Convert(const ALevelSequenceActor* LevelSequenceActor) override;
 };
diff --git a/Engine/Plugins/Enterprise/GLTFExporter/MigratedFromGit/Source/GLTFExporter/Public/Converters/GLTFBackdropConverters.h b/Engine/Plugins/Enterprise/GLTFExporter/MigratedFromGit/Source/GLTFExporter/Public/Converters/GLTFBackdropConverters.h
--- a/Engine/Plugins/Enterprise/GLTFExporter/MigratedFromGit/Source/GLTFExporter/Public/Converters/GLTFBackdropConverters.h
+++ b/Engine/Plugins/Enterprise/GLTFExporter/MigratedFromGit/Source/GLTFExporter/Public/Converters/GLTFBackdropConverters.h
@@ -9,9 +9,13 @@
 typedef TGLTFConverter<FGLTFJsonBackdrop*, const AActor*> IGLTFBackdropConverter;
 
 
-class GLTFEXPORTER_API FGLTFBackdropConverter final : public FGLTFBuilderContext, public IGLTFBackdropConverter
+class GLTFEXPORTER_API FGLTFBackdropConverter : public FGLTFBuilderContext, public IGLTFBackdropConverter
 {
+public:
+
 	using FGLTFBuilderContext::FGLTFBuilderContext;
 
+protected:
+
 	virtual FGLTFJsonBackdrop* Convert(const AActor* BackdropActor) override;
 };
diff --git a/Engine/Plugins/Enterprise/GLTFExporter/MigratedFromGit/Source/GLTFExporter/Public/Converters/GLTFCameraConverters.h b/Engine/Plugins/Enterprise/GLTFExporter/MigratedFromGit/Source/GLTFExporter/Public/Converters/GLTFCameraConverters.h
--- a/Engine/Plugins/Enterprise/GLTFExporter/MigratedFromGit/Source/GLTFExporter/Public/Converters/GLTFCameraConverters.h
+++ b/Engine/Plugins/Enterprise/GLTFExporter/MigratedFromGit/Source/GLTFExporter/Public/Converters/GLTFCameraConverters.h
@@ -10,9 +10,13 @@ class UCameraComponent;
 
 typedef TGLTFConverter<FGLTFJsonCamera*, const UCameraComponent*> IGLTFCameraConverter;
 
-class GLTFEXPORTER_API FGLTFCameraConverter final : public FGLTFBuilderContext, public IGLTFCameraConverter
+class GLTFEXPORTER_API FGLTFCameraConverter : public FGLTFBuilderContext, public IGLTFCameraConverter
 {
+public:
+
 	using FGLTFBuilderContext::FGLTFBuilderContext;
 
+protected:
+
 	virtual FGLTFJsonCamera* Convert(const UCameraComponent* CameraComponent) override;
 };
diff --git a/Engine/Plugins/Enterprise/GLTFExporter/MigratedFromGit/Source/GLTFExporter/Public/Converters/GLTFEpicLevelVariantSetsConverters.h b/Engine/Plugins/Enterprise/GLTFExporter/MigratedFromGit/Source/GLTFExporter/Public/Converters/GLTFEpicLevelVariantSetsConverters.h
--- a/Engine/Plugins/Enterprise/GLTFExporter/MigratedFromGit/Source/GLTFExporter/Public/Converters/GLTFEpicLevelVariantSetsConverters.h
+++ b/Engine/Plugins/Enterprise/GLTFExporter/MigratedFromGit/Source/GLTFExporter/Public/Converters/GLTFEpicLevelVariantSetsConverters.h
@@ -16,12 +16,18 @@ class UPropertyValueMaterial;
 
 typedef TGLTFConverter<FGLTFJsonEpicLevelVariantSets*, const ULevelVariantSets*> IGLTFEpicLevelVariantSetsConverter;
 
-class GLTFEXPORTER_API FGLTFEpicLevelVariantSetsConverter final : public FGLTFBuilderContext, public IGLTFEpicLevelVariantSetsConverter
+class GLTFEXPORTER_API FGLTFEpicLevelVariantSetsConverter : public FGLTFBuilderContext, public IGLTFEpicLevelVariantSetsConverter
 {
+public:
+
 	using FGLTFBuilderContext::FGLTFBuilderContext;
 
+protected:
+
 	virtual FGLTFJsonEpicLevelVariantSets* Convert(const ULevelVariantSets* LevelVariantSets) override;
 
+private:
+
 	bool TryParseVariant(FGLTFJsonEpicVariant& OutVariant, const UVariant* Variant) const;
 	bool TryParseVariantBinding(FGLTFJsonEpicVariant& OutVariant, const UVariantObjectBinding* Binding) const;
 	bool TryParseVisibilityPropertyValue(FGLTFJsonEpicVariant& OutVariant, const UPropertyValue* Property) const;
diff --git a/Engine/Plugins/Enterprise/GLTFExporter/MigratedFromGit/Source/GLTFExporter/Public/Converters/GLTFHotspotConverters.h b/Engine/Plugins/Enterprise/GLTFExporter/MigratedFromGit/Source/GLTFExporter/Public/Converters/GLTFHotspotConverters.h
--- a/Engine/Plugins/Enterprise/GLTFExporter/MigratedFromGit/Source/GLTFExporter/Public/Converters/GLTFHotspotConverters.h
+++ b/Engine/Plugins/Enterprise/GLTFExporter/MigratedFromGit/Source/GLTFExporter/Public/Converters/GLTFHotspotConverters.h
@@ -9,9 +9,13 @@
 
 typedef TGLTFConverter<FGLTFJsonHotspot*, const AGLTFHotspotActor*> IGLTFHotspotConverter;
 
-class GLTFEXPORTER_API FGLTFHotspotConverter final : public FGLTFBuilderContext, public IGLTFHotspotConverter
+class GLTFEXPORTER_API FGLTFHotspotConverter : public FGLTFBuilderContext, public IGLTFHotspotConverter
 {
+public:
+
 	using FGLTFBuilderContext::FGLTFBuilderContext;
 
+protected:
+
 	virtual FGLTFJsonHotspot* Convert(const AGLTFHotspotActor* HotspotActor) override;
 };
diff --git a/Engine/Plugins/Enterprise/GLTFExporter/MigratedFromGit/Source/GLTFExporter/Public/Converters/GLTFImageConverters.h b/Engine/Plugins/Enterprise/GLTFExporter/MigratedFromGit/Source/GLTFExporter/Public/Converters/GLTFImageConverters.h
--- a/Engine/Plugins/Enterprise/GLTFExporter/MigratedFromGit/Source/GLTFExporter/Public/Converters/GLTFImageConverters.h
+++ b/Engine/Plugins/Enterprise/GLTFExporter/MigratedFromGit/Source/GLTFExporter/Public/Converters/GLTFImageConverters.h
@@ -11,11 +11,17 @@
 
 typedef TGLTFConverter<FGLTFJsonImage*, TGLTFSuperfluous<FString>, EGLTFTextureType, bool, FIntPoint, TGLTFSharedArray<FColor>> IGLTFImageConverter;
 
-class GLTFEXPORTER_API FGLTFImageConverter final : public FGLTFBuilderContext, public IGLTFImageConverter
+class GLTFEXPORTER_API FGLTFImageConverter : public FGLTFBuilderContext, public IGLTFImageConverter
 {
+public:
+
 	using FGLTFBuilderContext::FGLTFBuilderContext;
 
+protected:
+
 	virtual FGLTFJsonImage* Convert(TGLTFSuperfluous<FString> Name, EGLTFTextureType Type, bool bIgnoreAlpha, FIntPoint Size, TGLTFSharedArray<FColor> Pixels) override;
 
+private:
+
 	EGLTFJsonMimeType GetMimeType(const FColor* Pixels, FIntPoint Size, bool bIgnoreAlpha, EGLTFTextureType Type) const;
 };
diff --git a/Engine/Plugins/Enterprise/GLTFExporter/MigratedFromGit/Source/GLTFExporter/Public/Converters/GLTFKhrVariantConverters.h b/Engine/Plugins/Enterprise/GLTFExporter/MigratedFromGit/Source/GLTFExporter/Public/Converters/GLTFKhrVariantConverters.h
--- a/Engine/Plugins/Enterprise/GLTFExporter/MigratedFromGit/Source/GLTFExporter/Public/Converters/GLTFKhrVariantConverters.h
+++ b/Engine/Plugins/Enterprise/GLTFExporter/MigratedFromGit/Source/GLTFExporter/Public/Converters/GLTFKhrVariantConverters.h
@@ -12,11 +12,17 @@ class UPropertyValueMaterial;
 
 typedef TGLTFConverter<FGLTFJsonKhrMaterialVariant*, const UVariant*> IGLTFKhrMaterialVariantConverter;
 
-class GLTFEXPORTER_API FGLTFKhrMaterialVariantConverter final : public FGLTFBuilderContext, public IGLTFKhrMaterialVariantConverter
+class GLTFEXPORTER_API FGLTFKhrMaterialVariantConverter : public FGLTFBuilderContext, public IGLTFKhrMaterialVariantConverter
 {
+public:
+
 	using FGLTFBuilderContext::FGLTFBuilderContext;
 
+protected:
+
 	virtual FGLTFJsonKhrMaterialVariant* Convert(const UVariant* Variant) override;
 
+private:
+
 	bool TryParseMaterialProperty(FGLTFJsonPrimitive*& OutPrimitive, FGLTFJsonMaterial*& OutMaterial, const UPropertyValueMaterial* Property) const;
 };
diff --git a/Engine/Plugins/Enterprise/GLTFExporter/MigratedFromGit/Source/GLTFExporter/Public/Converters/GLTFLightConverters.h b/Engine/Plugins/Enterprise/GLTFExporter/MigratedFromGit/Source/GLTFExporter/Public/Converters/GLTFLightConverters.h
--- a/Engine/Plugins/Enterprise/GLTFExporter/MigratedFromGit/Source/GLTFExporter/Public/Converters/GLTFLightConverters.h
+++ b/Engine/Plugins/Enterprise/GLTFExporter/MigratedFromGit/Source/GLTFExporter/Public/Converters/GLTFLightConverters.h
@@ -8,9 +8,13 @@
 
 typedef TGLTFConverter<FGLTFJsonLight*, const ULightComponent*> IGLTFLightConverter;
 
-class GLTFEXPORTER_API FGLTFLightConverter final : public FGLTFBuilderContext, public IGLTFLightConverter
+class GLTFEXPORTER_API FGLTFLightConverter : public FGLTFBuilderContext, public IGLTFLightConverter
 {
+public:
+
 	using FGLTFBuilderContext::FGLTFBuilderContext;
 
+protected:
+
 	virtual FGLTFJsonLight* Convert(const ULightComponent* LightComponent) override;
 };
diff --git a/Engine/Plugins/Enterprise/GLTFExporter/MigratedFromGit/Source/GLTFExporter/Public/Converters/GLTFLightMapConverters.h b/Engine/Plugins/Enterprise/GLTFExporter/MigratedFromGit/Source/GLTFExporter/Public/Converters/GLTFLightMapConverters.h
--- a/Engine/Plugins/Enterprise/GLTFExporter/MigratedFromGit/Source/GLTFExporter/Public/Converters/GLTFLightMapConverters.h
+++ b/Engine/Plugins/Enterprise/GLTFExporter/MigratedFromGit/Source/GLTFExporter/Public/Converters/GLTFLightMapConverters.h
@@ -8,9 +8,13 @@
 
 typedef TGLTFConverter<FGLTFJsonLightMap*, const UStaticMeshComponent*> IGLTFLightMapConverter;
 
-class GLTFEXPORTER_API FGLTFLightMapConverter final : public FGLTFBuilderContext, public IGLTFLightMapConverter
+class GLTFEXPORTER_API FGLTFLightMapConverter : public FGLTFBuilderContext, public IGLTFLightMapConverter
 {
+public:
+
 	using FGLTFBuilderContext::FGLTFBuilderContext;
 
+protected:
+
 	virtual FGLTFJsonLightMap* Convert(const UStaticMeshComponent* StaticMeshComponent) override;
 };
diff --git a/Engine/Plugins/Enterprise/GLTFExporter/MigratedFromGit/Source/GLTFExporter/Public/Converters/GLTFMaterialConverters.h b/Engine/Plugins/Enterprise/GLTFExporter/MigratedFromGit/Source/GLTFExporter/Public/Converters/GLTFMaterialConverters.h
--- a/Engine/Plugins/Enterprise/GLTFExporter/MigratedFromGit/Source/GLTFExporter/Public/Converters/GLTFMaterialConverters.h
+++ b/Engine/Plugins/Enterprise/GLTFExporter/MigratedFromGit/Source/GLTFExporter/Public/Converters/GLTFMaterialConverters.h
@@ -12,14 +12,20 @@
 
 typedef TGLTFConverter<FGLTFJsonMaterial*, const UMaterialInterface*, const FGLTFMeshData*, FGLTFIndexArray> IGLTFMaterialConverter;
 
-class GLTFEXPORTER_API FGLTFMaterialConverter final : public FGLTFBuilderContext, public IGLTFMaterialConverter
+class GLTFEXPORTER_API FGLTFMaterialConverter : public FGLTFBuilderContext, public IGLTFMaterialConverter
 {
+public:
+
 	using FGLTFBuilderContext::FGLTFBuilderContext;
 
+protected:
+
 	virtual void Sanitize(const UMaterialInterface*& Material, const FGLTFMeshData*& MeshData, FGLTFIndexArray& SectionIndices) override;
 
 	virtual FGLTFJsonMaterial* Convert(const UMaterialInterface* Material, const FGLTFMeshData* MeshData, FGLTFIndexArray SectionIndices) override;
 
+private:
+
 	FGLTFUVOverlapChecker UVOverlapChecker;
 	FGLTFUVDegenerateChecker UVDegenerateChecker;
 };
diff --git a/Engine/Plugins/Enterprise/GLTFExporter/MigratedFromGit/Source/GLTFExporter/Public/Converters/GLTFMeshConverters.h b/Engine/Plugins/Enterprise/GLTFExporter/MigratedFromGit/Source/GLTFExporter/Public/Converters/GLTFMeshConverters.h
--- a/Engine/Plugins/Enterprise/GLTFExporter/MigratedFromGit/Source/GLTFExporter/Public/Converters/GLTFMeshConverters.h
+++ b/Engine/Plugins/Enterprise/GLTFExporter/MigratedFromGit/Source/GLTFExporter/Public/Converters/GLTFMeshConverters.h
@@ -11,24 +11,36 @@
 typedef TGLTFConverter<FGLTFJsonMesh*, const UStaticMesh*, const UStaticMeshComponent*, FGLTFMaterialArray, int32> IGLTFStaticMeshConverter;
 typedef TGLTFConverter<FGLTFJsonMesh*, const USkeletalMesh*, const USkeletalMeshComponent*, FGLTFMaterialArray, int32> IGLTFSkeletalMeshConverter;
 
-class GLTFEXPORTER_API FGLTFStaticMeshConverter final : public FGLTFBuilderContext, public IGLTFStaticMeshConverter
+class GLTFEXPORTER_API FGLTFStaticMeshConverter : public FGLTFBuilderContext, public IGLTFStaticMeshConverter
 {
+public:
+
 	using FGLTFBuilderContext::FGLTFBuilderContext;
 
+protected:
+
 	virtual void Sanitize(const UStaticMesh*& StaticMesh, const UStaticMeshComponent*& StaticMeshComponent, FGLTFMaterialArray& Materials, int32& LODIndex) override;
 
 	virtual FGLTFJsonMesh* Convert(const UStaticMesh* StaticMesh, const UStaticMeshComponent* StaticMeshComponent, FGLTFMaterialArray Materials, int32 LODIndex) override;
 
+private:
+
 	FGLTFStaticMeshSectionConverter MeshSectionConverter;
 };
 
-class GLTFEXPORTER_API FGLTFSkeletalMeshConverter final : public FGLTFBuilderContext, public IGLTFSkeletalMeshConverter
+class GLTFEXPORTER_API FGLTFSkeletalMeshConverter : public FGLTFBuilderContext, public IGLTFSkeletalMeshConverter
 {
+public:
+
 	using FGLTFBuilderContext::FGLTFBuilderContext;
 
+protected:
+
 	virtual void Sanitize(const USkeletalMesh*& SkeletalMesh, const USkeletalMeshComponent*& SkeletalMeshComponent, FGLTFMaterialArray& Materials, int32& LODIndex) override;
 
 	virtual FGLTFJsonMesh* Convert(const USkeletalMesh* SkeletalMesh, const USkeletalMeshComponent* SkeletalMeshComponent, FGLTFMaterialArray Materials, int32 LODIndex) override;
 
+private:
+
 	FGLTFSkeletalMeshSectionConverter MeshSectionConverter;
 };
diff --git a/Engine/Plugins/Enterprise/GLTFExporter/MigratedFromGit/Source/GLTFExporter/Public/Converters/GLTFMeshDataConverters.h b/Engine/Plugins/Enterprise/GLTFExporter/MigratedFromGit/Source/GLTFExporter/Public/Converters/GLTFMeshDataConverters.h
--- a/Engine/Plugins/Enterprise/GLTFExporter/MigratedFromGit/Source/GLTFExporter/Public/Converters/GLTFMeshDataConverters.h
+++ b/Engine/Plugins/Enterprise/GLTFExporter/MigratedFromGit/Source/GLTFExporter/Public/Converters/GLTFMeshDataConverters.h
@@ -22,14 +22,16 @@ public:
 	TGLTFMeshDataConverter(const TGLTFMeshDataConverter&) = delete;
 	TGLTFMeshDataConverter& operator=(const TGLTFMeshDataConverter&) = delete;
 
-private:
-
-	// TODO: standardize converters that are responsible for deleting their output
-	TArray<TUniquePtr<FGLTFMeshData>> Outputs;
+protected:
 
 	virtual void Sanitize(const MeshType*& Mesh, const MeshComponentType*& MeshComponent, int32& LODIndex) override;
 
 	virtual const FGLTFMeshData* Convert(const MeshType* Mesh, const MeshComponentType* MeshComponent, int32 LODIndex) override;
+
+private:
+
+	// TODO: standardize converters that are responsible for deleting their output
+	TArray<TUniquePtr<FGLTFMeshData>> Outputs;
 };
 
 typedef TGLTFMeshDataConverter<UStaticMesh, UStaticMeshComponent> FGLTFStaticMeshDataConverter;
diff --git a/Engine/Plugins/Enterprise/GLTFExporter/MigratedFromGit/Source/GLTFExporter/Public/Converters/GLTFMeshSectionConverters.h b/Engine/Plugins/Enterprise/GLTFExporter/MigratedFromGit/Source/GLTFExporter/Public/Converters/GLTFMeshSectionConverters.h
--- a/Engine/Plugins/Enterprise/GLTFExporter/MigratedFromGit/Source/GLTFExporter/Public/Converters/GLTFMeshSectionConverters.h
+++ b/Engine/Plugins/Enterprise/GLTFExporter/MigratedFromGit/Source/GLTFExporter/Public/Converters/GLTFMeshSectionConverters.h
@@ -7,7 +7,7 @@
 #include "Converters/GLTFIndexArray.h"
 
 template <typename MeshLODType>
-class TGLTFMeshSectionConverter final : public TGLTFConverter<const FGLTFMeshSection*, const MeshLODType*, FGLTFIndexArray>
+class TGLTFMeshSectionConverter : public TGLTFConverter<const FGLTFMeshSection*, const MeshLODType*, FGLTFIndexArray>
 {
 public:
 
@@ -18,14 +18,16 @@ public:
 	TGLTFMeshSectionConverter(const TGLTFMeshSectionConverter&) = delete;
 	TGLTFMeshSectionConverter& operator=(const TGLTFMeshSectionConverter&) = delete;
 
-private:
-
-	TArray<TUniquePtr<FGLTFMeshSection>> Outputs;
+protected:
 
 	virtual const FGLTFMeshSection* Convert(const MeshLODType* MeshLOD, FGLTFIndexArray SectionIndices) override
 	{
 		return Outputs.Add_GetRef(MakeUnique<FGLTFMeshSection>(MeshLOD, SectionIndices)).Get();
 	}
+
+private:
+
+	TArray<TUniquePtr<FGLTFMeshSection>> Outputs;
 };
 
 typedef TGLTFMeshSectionConverter<FStaticMeshLODResources> FGLTFStaticMeshSectionConverter;
diff --git a/Engine/Plugins/Enterprise/GLTFExporter/MigratedFromGit/Source/GLTFExporter/Public/Converters/GLTFNodeConverters.h b/Engine/Plugins/Enterprise/GLTFExporter/MigratedFromGit/Source/GLTFExporter/Public/Converters/GLTFNodeConverters.h
--- a/Engine/Plugins/Enterprise/GLTFExporter/MigratedFromGit/Source/GLTFExporter/Public/Converters/GLTFNodeConverters.h
+++ b/Engine/Plugins/Enterprise/GLTFExporter/MigratedFromGit/Source/GLTFExporter/Public/Converters/GLTFNodeConverters.h
@@ -13,46 +13,70 @@ typedef TGLTFConverter<FGLTFJsonNode*, FGLTFJsonNode*, const UStaticMesh*, FName
 typedef TGLTFConverter<FGLTFJsonNode*, FGLTFJsonNode*, const USkeletalMesh*, FName> IGLTFSkeletalSocketConverter;
 typedef TGLTFConverter<FGLTFJsonNode*, FGLTFJsonNode*, const USkeletalMesh*, int32> IGLTFSkeletalBoneConverter;
 
-class GLTFEXPORTER_API FGLTFActorConverter final : public FGLTFBuilderContext, public IGLTFActorConverter
+class GLTFEXPORTER_API FGLTFActorConverter : public FGLTFBuilderContext, public IGLTFActorConverter
 {
+public:
+
 	using FGLTFBuilderContext::FGLTFBuilderContext;
 
+protected:
+
 	virtual FGLTFJsonNode* Convert(const AActor* Actor) override;
 };
 
-class GLTFEXPORTER_API FGLTFComponentConverter final : public FGLTFBuilderContext, public IGLTFComponentConverter
+class GLTFEXPORTER_API FGLTFComponentConverter : public FGLTFBuilderContext, public IGLTFComponentConverter
 {
+public:
+
 	using FGLTFBuilderContext::FGLTFBuilderContext;
 
+protected:
+
 	virtual FGLTFJsonNode* Convert(const USceneComponent* SceneComponent) override;
 
 	void ConvertComponentSpecialization(const USceneComponent* SceneComponent, const AActor* Owner, FGLTFJsonNode* Node);
 };
 
-class GLTFEXPORTER_API FGLTFComponentSocketConverter final : public FGLTFBuilderContext, public IGLTFComponentSocketConverter
+class GLTFEXPORTER_API FGLTFComponentSocketConverter : public FGLTFBuilderContext, public IGLTFComponentSocketConverter
 {
+public:
+
 	using FGLTFBuilderContext::FGLTFBuilderContext;
 
+protected:
+
 	virtual FGLTFJsonNode* Convert(const USceneComponent* SceneComponent, FName SocketName) override;
 };
 
-class GLTFEXPORTER_API FGLTFStaticSocketConverter final : public FGLTFBuilderContext, public IGLTFStaticSocketConverter
+class GLTFEXPORTER_API FGLTFStaticSocketConverter : public FGLTFBuilderContext, public IGLTFStaticSocketConverter
 {
+public:
+
 	using FGLTFBuilderContext::FGLTFBuilderContext;
 
+protected:
+
 	virtual FGLTFJsonNode* Convert(FGLTFJsonNode* RootNode, const UStaticMesh* StaticMesh, FName SocketName) override;
 };
 
-class GLTFEXPORTER_API FGLTFSkeletalSocketConverter final : public FGLTFBuilderContext, public IGLTFSkeletalSocketConverter
+class GLTFEXPORTER_API FGLTFSkeletalSocketConverter : public FGLTFBuilderContext, public IGLTFSkeletalSocketConverter
 {
+public:
+
 	using FGLTFBuilderContext::FGLTFBuilderContext;
 
+protected:
+
 	virtual FGLTFJsonNode* Convert(FGLTFJsonNode* RootNode, const USkeletalMesh* SkeletalMesh, FName SocketName) override;
 };
 
-class GLTFEXPORTER_API FGLTFSkeletalBoneConverter final : public FGLTFBuilderContext, public IGLTFSkeletalBoneConverter
+class GLTFEXPORTER_API FGLTFSkeletalBoneConverter : public FGLTFBuilderContext, public IGLTFSkeletalBoneConverter
 {
+public:
+
 	using FGLTFBuilderContext::FGLTFBuilderContext;
 
+protected:
+
 	virtual FGLTFJsonNode* Convert(FGLTFJsonNode* RootNode, const USkeletalMesh* SkeletalMesh, int32 BoneIndex) override;
 };
diff --git a/Engine/Plugins/Enterprise/GLTFExporter/MigratedFromGit/Source/GLTFExporter/Public/Converters/GLTFSamplerConverters.h b/Engine/Plugins/Enterprise/GLTFExporter/MigratedFromGit/Source/GLTFExporter/Public/Converters/GLTFSamplerConverters.h
--- a/Engine/Plugins/Enterprise/GLTFExporter/MigratedFromGit/Source/GLTFExporter/Public/Converters/GLTFSamplerConverters.h
+++ b/Engine/Plugins/Enterprise/GLTFExporter/MigratedFromGit/Source/GLTFExporter/Public/Converters/GLTFSamplerConverters.h
@@ -8,9 +8,13 @@
 
 typedef TGLTFConverter<FGLTFJsonSampler*, const UTexture*> IGLTFSamplerConverter;
 
-class GLTFEXPORTER_API FGLTFSamplerConverter final : public FGLTFBuilderContext, public IGLTFSamplerConverter
+class GLTFEXPORTER_API FGLTFSamplerConverter : public FGLTFBuilderContext, public IGLTFSamplerConverter
 {
+public:
+
 	using FGLTFBuilderContext::FGLTFBuilderContext;
 
+protected:
+
 	virtual FGLTFJsonSampler* Convert(const UTexture* Texture) override;
 };
diff --git a/Engine/Plugins/Enterprise/GLTFExporter/MigratedFromGit/Source/GLTFExporter/Public/Converters/GLTFSceneConverters.h b/Engine/Plugins/Enterprise/GLTFExporter/MigratedFromGit/Source/GLTFExporter/Public/Converters/GLTFSceneConverters.h
--- a/Engine/Plugins/Enterprise/GLTFExporter/MigratedFromGit/Source/GLTFExporter/Public/Converters/GLTFSceneConverters.h
+++ b/Engine/Plugins/Enterprise/GLTFExporter/MigratedFromGit/Source/GLTFExporter/Public/Converters/GLTFSceneConverters.h
@@ -8,9 +8,13 @@
 
 typedef TGLTFConverter<FGLTFJsonScene*, const UWorld*> IGLTFSceneConverter;
 
-class GLTFEXPORTER_API FGLTFSceneConverter final : public FGLTFBuilderContext, public IGLTFSceneConverter
+class GLTFEXPORTER_API FGLTFSceneConverter : public FGLTFBuilderContext, public IGLTFSceneConverter
 {
+public:
+
 	using FGLTFBuilderContext::FGLTFBuilderContext;
 
+protected:
+
 	virtual FGLTFJsonScene* Convert(const UWorld* Level) override;
 };
diff --git a/Engine/Plugins/Enterprise/GLTFExporter/MigratedFromGit/Source/GLTFExporter/Public/Converters/GLTFSkinConverters.h b/Engine/Plugins/Enterprise/GLTFExporter/MigratedFromGit/Source/GLTFExporter/Public/Converters/GLTFSkinConverters.h
--- a/Engine/Plugins/Enterprise/GLTFExporter/MigratedFromGit/Source/GLTFExporter/Public/Converters/GLTFSkinConverters.h
+++ b/Engine/Plugins/Enterprise/GLTFExporter/MigratedFromGit/Source/GLTFExporter/Public/Converters/GLTFSkinConverters.h
@@ -8,9 +8,13 @@
 
 typedef TGLTFConverter<FGLTFJsonSkin*, FGLTFJsonNode*, const USkeletalMesh*> IGLTFSkinConverter;
 
-class GLTFEXPORTER_API FGLTFSkinConverter final : public FGLTFBuilderContext, public IGLTFSkinConverter
+class GLTFEXPORTER_API FGLTFSkinConverter : public FGLTFBuilderContext, public IGLTFSkinConverter
 {
+public:
+
 	using FGLTFBuilderContext::FGLTFBuilderContext;
 
+protected:
+
 	virtual FGLTFJsonSkin* Convert(FGLTFJsonNode* RootNode, const USkeletalMesh* SkeletalMesh) override;
 };
diff --git a/Engine/Plugins/Enterprise/GLTFExporter/MigratedFromGit/Source/GLTFExporter/Public/Converters/GLTFSkySphereConverters.h b/Engine/Plugins/Enterprise/GLTFExporter/MigratedFromGit/Source/GLTFExporter/Public/Converters/GLTFSkySphereConverters.h
--- a/Engine/Plugins/Enterprise/GLTFExporter/MigratedFromGit/Source/GLTFExporter/Public/Converters/GLTFSkySphereConverters.h
+++ b/Engine/Plugins/Enterprise/GLTFExporter/MigratedFromGit/Source/GLTFExporter/Public/Converters/GLTFSkySphereConverters.h
@@ -11,8 +11,18 @@ struct FGLTFJsonSkySphereColorCurve;
 
 typedef TGLTFConverter<FGLTFJsonSkySphere*, const AActor*> IGLTFSkySphereConverter;
 
-class GLTFEXPORTER_API FGLTFSkySphereConverter final : public FGLTFBuilderContext, public IGLTFSkySphereConverter
+class GLTFEXPORTER_API FGLTFSkySphereConverter : public FGLTFBuilderContext, public IGLTFSkySphereConverter
 {
+public:
+
+	using FGLTFBuilderContext::FGLTFBuilderContext;
+
+protected:
+
+	virtual FGLTFJsonSkySphere* Convert(const AActor* SkySphereActor) override;
+
+private:
+
 	enum class ESkySphereTextureParameter
 	{
 		SkyTexture,
@@ -20,10 +30,6 @@ class GLTFEXPORTER_API FGLTFSkySphereConverter final : public FGLTFBuilderContex
 		StarsTexture
 	};
 
-	using FGLTFBuilderContext::FGLTFBuilderContext;
-
-	virtual FGLTFJsonSkySphere* Convert(const AActor* SkySphereActor) override;
-
 	template <class ValueType>
 	void ConvertProperty(const AActor* Actor, const TCHAR* PropertyName, ValueType& OutValue) const;
 
diff --git a/Engine/Plugins/Enterprise/GLTFExporter/MigratedFromGit/Source/GLTFExporter/Public/Converters/GLTFTextureConverters.h b/Engine/Plugins/Enterprise/GLTFExporter/MigratedFromGit/Source/GLTFExporter/Public/Converters/GLTFTextureConverters.h
--- a/Engine/Plugins/Enterprise/GLTFExporter/MigratedFromGit/Source/GLTFExporter/Public/Converters/GLTFTextureConverters.h
+++ b/Engine/Plugins/Enterprise/GLTFExporter/MigratedFromGit/Source/GLTFExporter/Public/Converters/GLTFTextureConverters.h
@@ -16,45 +16,65 @@ typedef TGLTFConverter<FGLTFJsonTexture*, const UTextureRenderTarget2D*, bool> I
 typedef TGLTFConverter<FGLTFJsonTexture*, const UTextureRenderTargetCube*, ECubeFace, bool> IGLTFTextureRenderTargetCubeConverter;
 typedef TGLTFConverter<FGLTFJsonTexture*, const ULightMapTexture2D*> IGLTFTextureLightMapConverter;
 
-class GLTFEXPORTER_API FGLTFTexture2DConverter final : public FGLTFBuilderContext, public IGLTFTexture2DConverter
+class GLTFEXPORTER_API FGLTFTexture2DConverter : public FGLTFBuilderContext, public IGLTFTexture2DConverter
 {
+public:
+
 	using FGLTFBuilderContext::FGLTFBuilderContext;
 
+protected:
+
 	virtual void Sanitize(const UTexture2D*& Texture2D, bool& bToSRGB) override;
 
 	virtual FGLTFJsonTexture* Convert(const UTexture2D* Texture2D, bool bToSRGB) override;
 };
 
-class GLTFEXPORTER_API FGLTFTextureCubeConverter final : public FGLTFBuilderContext, public IGLTFTextureCubeConverter
+class GLTFEXPORTER_API FGLTFTextureCubeConverter : public FGLTFBuilderContext, public IGLTFTextureCubeConverter
 {
+public:
+
 	using FGLTFBuilderContext::FGLTFBuilderContext;
 
+protected:
+
 	virtual void Sanitize(const UTextureCube*& TextureCube, ECubeFace& CubeFace, bool& bToSRGB) override;
 
 	virtual FGLTFJsonTexture* Convert(const UTextureCube* TextureCube, ECubeFace CubeFace, bool bToSRGB) override;
 };
 
-class GLTFEXPORTER_API FGLTFTextureRenderTarget2DConverter final : public FGLTFBuilderContext, public IGLTFTextureRenderTarget2DConverter
+class GLTFEXPORTER_API FGLTFTextureRenderTarget2DConverter : public FGLTFBuilderContext, public IGLTFTextureRenderTarget2DConverter
 {
+public:
+
 	using FGLTFBuilderContext::FGLTFBuilderContext;
 
+protected:
+
 	virtual void Sanitize(const UTextureRenderTarget2D*& RenderTarget2D, bool& bToSRGB) override;
 
 	virtual FGLTFJsonTexture* Convert(const UTextureRenderTarget2D* RenderTarget2D, bool bToSRGB) override;
 };
 
-class GLTFEXPORTER_API FGLTFTextureRenderTargetCubeConverter final : public FGLTFBuilderContext, public IGLTFTextureRenderTargetCubeConverter
+class GLTFEXPORTER_API FGLTFTextureRenderTargetCubeConverter : public FGLTFBuilderContext, public IGLTFTextureRenderTargetCubeConverter
 {
+public:
+
 	using FGLTFBuilderContext::FGLTFBuilderContext;
 
+protected:
+
 	virtual void Sanitize(const UTextureRenderTargetCube*& RenderTargetCube, ECubeFace& CubeFace, bool& bToSRGB) override;
 
 	virtual FGLTFJsonTexture* Convert(const UTextureRenderTargetCube* RenderTargetCube, ECubeFace CubeFace, bool bToSRGB) override;
 };
 
-class GLTFEXPORTER_API FGLTFTextureLightMapConverter final : public FGLTFBuilderContext, public IGLTFTextureLightMapConverter
+class GLTFEXPORTER_API FGLTFTextureLightMapConverter : public FGLTFBuilderContext, public IGLTFTextureLightMapConverter
 {
+public:
+
 	using FGLTFBuilderContext::FGLTFBuilderContext;
 
+protected:
+
 	virtual FGLTFJsonTexture* Convert(const ULightMapTexture2D* LightMap) override;
 };
diff --git a/Engine/Plugins/Enterprise/GLTFExporter/MigratedFromGit/Source/GLTFExporter/Public/Converters/GLTFUVDegenerateChecker.h b/Engine/Plugins/Enterprise/GLTFExporter/MigratedFromGit/Source/GLTFExporter/Public/Converters/GLTFUVDegenerateChecker.h
--- a/Engine/Plugins/Enterprise/GLTFExporter/MigratedFromGit/Source/GLTFExporter/Public/Converters/GLTFUVDegenerateChecker.h
+++ b/Engine/Plugins/Enterprise/GLTFExporter/MigratedFromGit/Source/GLTFExporter/Public/Converters/GLTFUVDegenerateChecker.h
@@ -7,12 +7,16 @@
 
 struct FMeshDescription;
 
-class GLTFEXPORTER_API FGLTFUVDegenerateChecker final : public TGLTFConverter<float, const FMeshDescription*, FGLTFIndexArray, int32>
+class GLTFEXPORTER_API FGLTFUVDegenerateChecker : public TGLTFConverter<float, const FMeshDescription*, FGLTFIndexArray, int32>
 {
+protected:
+
 	virtual void Sanitize(const FMeshDescription*& Description, FGLTFIndexArray& SectionIndices, int32& TexCoord) override;
 
 	virtual float Convert(const FMeshDescription* Description, FGLTFIndexArray SectionIndices, int32 TexCoord) override;
 
+private:
+
 	static bool IsDegenerateTriangle(const TStaticArray<FVector2f, 3>& Points);
 	static bool IsDegenerateTriangle(const TStaticArray<FVector3f, 3>& Points);
 };
diff --git a/Engine/Plugins/Enterprise/GLTFExporter/MigratedFromGit/Source/GLTFExporter/Public/Converters/GLTFUVOverlapChecker.h b/Engine/Plugins/Enterprise/GLTFExporter/MigratedFromGit/Source/GLTFExporter/Public/Converters/GLTFUVOverlapChecker.h
--- a/Engine/Plugins/Enterprise/GLTFExporter/MigratedFromGit/Source/GLTFExporter/Public/Converters/GLTFUVOverlapChecker.h
+++ b/Engine/Plugins/Enterprise/GLTFExporter/MigratedFromGit/Source/GLTFExporter/Public/Converters/GLTFUVOverlapChecker.h
@@ -7,11 +7,15 @@
 
 struct FMeshDescription;
 
-class GLTFEXPORTER_API FGLTFUVOverlapChecker final : public TGLTFConverter<float, const FMeshDescription*, FGLTFIndexArray, int32>
+class GLTFEXPORTER_API FGLTFUVOverlapChecker : public TGLTFConverter<float, const FMeshDescription*, FGLTFIndexArray, int32>
 {
+protected:
+
 	virtual void Sanitize(const FMeshDescription*& Description, FGLTFIndexArray& SectionIndices, int32& TexCoord) override;
 
 	virtual float Convert(const FMeshDescription* Description, FGLTFIndexArray SectionIndices, int32 TexCoord) override;
 
+private:
+
 	static UMaterialInterface* GetMaterial();
 };
