[GLTFExporter] fix: fix compiler issues regarding FVector to FVector3f conversions for post-UE5-preview2

commit 3a12517b4024af76cb2fb12ca944b66a0fb5090c
Date: Wed, 30 Mar 2022 15:51:59 +0200

    fix: fix compiler issues regarding FVector to FVector3f conversions for post-UE5-preview2
    (cherry picked from commit 3fad206c5ccde6ffc0d2557855ef84403ef71469​)

#skipci

diff --git a/Engine/Plugins/Enterprise/GLTFExporter/MigratedFromGit/Source/GLTFExporter/Private/Converters/GLTFAccessorConverters.cpp b/Engine/Plugins/Enterprise/GLTFExporter/MigratedFromGit/Source/GLTFExporter/Private/Converters/GLTFAccessorConverters.cpp
--- a/Engine/Plugins/Enterprise/GLTFExporter/MigratedFromGit/Source/GLTFExporter/Private/Converters/GLTFAccessorConverters.cpp
+++ b/Engine/Plugins/Enterprise/GLTFExporter/MigratedFromGit/Source/GLTFExporter/Private/Converters/GLTFAccessorConverters.cpp
@@ -193,7 +193,7 @@ FGLTFJsonBufferView* FGLTFNormalBufferConverter::ConvertBufferView(const FGLTFMe
 	for (uint32 VertexIndex = 0; VertexIndex < VertexCount; ++VertexIndex)
 	{
 		const uint32 MappedVertexIndex = IndexMap[VertexIndex];
-		const FVector3f SafeNormal = TangentData[MappedVertexIndex].TangentZ.ToFVector().GetSafeNormal();
+		const FVector3f SafeNormal = TangentData[MappedVertexIndex].TangentZ.ToFVector3f().GetSafeNormal();
 
 		typedef typename TConditional<TIsSame<DestinationType, FGLTFVector3>::Value, FVector3f, SourceType>::Type IntermediateType;
 		Normals[VertexIndex] = FGLTFCoreUtilities::ConvertNormal(IntermediateType(SafeNormal));
@@ -272,7 +272,7 @@ FGLTFJsonBufferView* FGLTFTangentBufferConverter::ConvertBufferView(const FGLTFM
 	for (uint32 VertexIndex = 0; VertexIndex < VertexCount; ++VertexIndex)
 	{
 		const uint32 MappedVertexIndex = IndexMap[VertexIndex];
-		const FVector3f SafeTangent = VertexTangents[MappedVertexIndex].TangentX.ToFVector().GetSafeNormal();
+		const FVector3f SafeTangent = VertexTangents[MappedVertexIndex].TangentX.ToFVector3f().GetSafeNormal();
 
 		typedef typename TConditional<TIsSame<DestinationType, FGLTFVector4>::Value, FVector3f, SourceType>::Type IntermediateType;
 		Tangents[VertexIndex] = FGLTFCoreUtilities::ConvertTangent(IntermediateType(SafeTangent));
diff --git a/Engine/Plugins/Enterprise/GLTFExporter/MigratedFromGit/Source/GLTFExporter/Private/Converters/GLTFBackdropConverters.cpp b/Engine/Plugins/Enterprise/GLTFExporter/MigratedFromGit/Source/GLTFExporter/Private/Converters/GLTFBackdropConverters.cpp
--- a/Engine/Plugins/Enterprise/GLTFExporter/MigratedFromGit/Source/GLTFExporter/Private/Converters/GLTFBackdropConverters.cpp
+++ b/Engine/Plugins/Enterprise/GLTFExporter/MigratedFromGit/Source/GLTFExporter/Private/Converters/GLTFBackdropConverters.cpp
@@ -76,7 +76,7 @@ FGLTFJsonBackdrop* FGLTFBackdropConverter::Convert(const AActor* BackdropActor)
 	FVector ProjectionCenter;
 	if (FGLTFBlueprintUtility::TryGetPropertyValue(BackdropActor, TEXT("ProjectionCenter"), ProjectionCenter))
 	{
-		JsonBackdrop->ProjectionCenter = FGLTFCoreUtilities::ConvertPosition(ProjectionCenter, Builder.ExportOptions->ExportUniformScale);
+		JsonBackdrop->ProjectionCenter = FGLTFCoreUtilities::ConvertPosition(FVector3f(ProjectionCenter), Builder.ExportOptions->ExportUniformScale);
 	}
 	else
 	{
diff --git a/Engine/Plugins/Enterprise/GLTFExporter/MigratedFromGit/Source/GLTFExporter/Private/Converters/GLTFSkySphereConverters.cpp b/Engine/Plugins/Enterprise/GLTFExporter/MigratedFromGit/Source/GLTFExporter/Private/Converters/GLTFSkySphereConverters.cpp
--- a/Engine/Plugins/Enterprise/GLTFExporter/MigratedFromGit/Source/GLTFExporter/Private/Converters/GLTFSkySphereConverters.cpp
+++ b/Engine/Plugins/Enterprise/GLTFExporter/MigratedFromGit/Source/GLTFExporter/Private/Converters/GLTFSkySphereConverters.cpp
@@ -32,7 +32,7 @@ FGLTFJsonSkySphere* FGLTFSkySphereConverter::Convert(const AActor* SkySphereActo
 		const FTransform ParentTransform = ParentComponent->GetSocketTransform(SocketName);
 		const FTransform RelativeTransform = Transform.GetRelativeTransform(ParentTransform);
 
-		JsonSkySphere->Scale = FGLTFCoreUtilities::ConvertScale(RelativeTransform.GetScale3D());
+		JsonSkySphere->Scale = FGLTFCoreUtilities::ConvertScale(FVector3f(RelativeTransform.GetScale3D()));
 		JsonSkySphere->SkySphereMesh = Builder.AddUniqueMesh(StaticMeshComponent, { FGLTFMaterialUtility::GetDefaultMaterial() });
 	}
 	else
diff --git a/Engine/Plugins/Enterprise/GLTFExporter/MigratedFromGit/Source/GLTFExporter/Private/Tasks/GLTFDelayedAnimationTasks.cpp b/Engine/Plugins/Enterprise/GLTFExporter/MigratedFromGit/Source/GLTFExporter/Private/Tasks/GLTFDelayedAnimationTasks.cpp
--- a/Engine/Plugins/Enterprise/GLTFExporter/MigratedFromGit/Source/GLTFExporter/Private/Tasks/GLTFDelayedAnimationTasks.cpp
+++ b/Engine/Plugins/Enterprise/GLTFExporter/MigratedFromGit/Source/GLTFExporter/Private/Tasks/GLTFDelayedAnimationTasks.cpp
@@ -62,7 +62,7 @@ void FGLTFDelayedAnimSequenceTask::Process()
 
 			for (int32 Frame = 0; Frame < FrameCount; ++Frame)
 			{
-				const FVector3f KeyPosition = FrameTransforms[Frame][BoneIndex].GetTranslation();
+				const FVector3f KeyPosition = FVector3f(FrameTransforms[Frame][BoneIndex].GetTranslation());
 				Translations[Frame] = FGLTFCoreUtilities::ConvertPosition(KeyPosition, Builder.ExportOptions->ExportUniformScale);
 			}
 
@@ -118,7 +118,7 @@ void FGLTFDelayedAnimSequenceTask::Process()
 
 			for (int32 Frame = 0; Frame < FrameCount; ++Frame)
 			{
-				const FVector3f KeyScale = FrameTransforms[Frame][BoneIndex].GetScale3D();
+				const FVector3f KeyScale = FVector3f(FrameTransforms[Frame][BoneIndex].GetScale3D());
 				Scales[Frame] = FGLTFCoreUtilities::ConvertScale(KeyScale);
 			}
 
diff --git a/Engine/Plugins/Enterprise/GLTFExporter/MigratedFromGit/Source/GLTFExporter/Private/Utilities/GLTFCoreUtilities.cpp b/Engine/Plugins/Enterprise/GLTFExporter/MigratedFromGit/Source/GLTFExporter/Private/Utilities/GLTFCoreUtilities.cpp
--- a/Engine/Plugins/Enterprise/GLTFExporter/MigratedFromGit/Source/GLTFExporter/Private/Utilities/GLTFCoreUtilities.cpp
+++ b/Engine/Plugins/Enterprise/GLTFExporter/MigratedFromGit/Source/GLTFExporter/Private/Utilities/GLTFCoreUtilities.cpp
@@ -157,7 +157,7 @@ FGLTFMatrix4 FGLTFCoreUtilities::ConvertTransform(const FTransform3f& Transform,
 	const FVector3f Scale = Transform.GetScale3D();
 
 	const FQuat4f ConvertedRotation(-Rotation.X, -Rotation.Z, -Rotation.Y, Rotation.W);
-	const FVector3f ConvertedTranslation = FVector(Translation.X, Translation.Z, Translation.Y) * ConversionScale;
+	const FVector3f ConvertedTranslation = FVector3f(Translation.X, Translation.Z, Translation.Y) * ConversionScale;
 	const FVector3f ConvertedScale(Scale.X, Scale.Z, Scale.Y);
 
 	const FTransform3f ConvertedTransform(ConvertedRotation, ConvertedTranslation, ConvertedScale);
diff --git a/Engine/Plugins/Enterprise/GLTFExporter/MigratedFromGit/Source/GLTFMaterialBaking/Private/MaterialRenderItem.cpp b/Engine/Plugins/Enterprise/GLTFExporter/MigratedFromGit/Source/GLTFMaterialBaking/Private/MaterialRenderItem.cpp
--- a/Engine/Plugins/Enterprise/GLTFExporter/MigratedFromGit/Source/GLTFMaterialBaking/Private/MaterialRenderItem.cpp
+++ b/Engine/Plugins/Enterprise/GLTFExporter/MigratedFromGit/Source/GLTFMaterialBaking/Private/MaterialRenderItem.cpp
@@ -141,7 +141,7 @@ void FGLTFMeshMaterialRenderItem::QueueMaterial(FCanvasRenderContext& RenderCont
 
 			// Overwrite object orientation, since original calculation is derived from LocalToWorld matrix,
 			// which currently needs to be identity matrix to avoid transforming the "2D" vertices used for baking
-			PrimitiveParams.ObjectOrientation = PrimitiveData->ObjectOrientation;
+			PrimitiveParams.ObjectOrientation = FVector3f(PrimitiveData->ObjectOrientation);
 
 			const_cast<TUniformBuffer<FPrimitiveUniformShaderParameters>*>(MeshElement.Elements[0].PrimitiveUniformBufferResource)->SetContents(PrimitiveParams);
 		}
