[GLTFExporter] fix: add safety measure when iterating over level actors, in case temporary actors are added during export

commit 07168b1b8a47420711ce6ed220369611e1eeb17d
Date: Thu, 31 Mar 2022 03:06:44 +0200

    fix: add safety measure when iterating over level actors, in case temporary actors are added during export
    (cherry picked from commit a417d643553ce690c0a005c0c6a81de40671e28a​)

#skipci

diff --git a/Engine/Plugins/Enterprise/GLTFExporter/MigratedFromGit/Source/GLTFExporter/Private/Converters/GLTFSceneConverters.cpp b/Engine/Plugins/Enterprise/GLTFExporter/MigratedFromGit/Source/GLTFExporter/Private/Converters/GLTFSceneConverters.cpp
--- a/Engine/Plugins/Enterprise/GLTFExporter/MigratedFromGit/Source/GLTFExporter/Private/Converters/GLTFSceneConverters.cpp
+++ b/Engine/Plugins/Enterprise/GLTFExporter/MigratedFromGit/Source/GLTFExporter/Private/Converters/GLTFSceneConverters.cpp
@@ -23,9 +23,11 @@ FGLTFJsonScene* FGLTFSceneConverter::Convert(const UWorld* World)
 
 			// TODO: add support for exporting Level->Model?
 
+			TArray<AActor*> LevelActors(Level->Actors); // iterate using copy, in case new temporary actors are added during conversion
+
 			if (Builder.ExportOptions->VariantSetsMode != EGLTFVariantSetsMode::None)
 			{
-				for (const AActor* Actor : Level->Actors)
+				for (const AActor* Actor : LevelActors)
 				{
 					// TODO: should a LevelVariantSet be exported even if not selected for export?
 					if (const ALevelVariantSetsActor *LevelVariantSetsActor = Cast<ALevelVariantSetsActor>(Actor))
@@ -60,7 +62,7 @@ FGLTFJsonScene* FGLTFSceneConverter::Convert(const UWorld* World)
 				}
 			}
 
-			for (const AActor* Actor : Level->Actors)
+			for (const AActor* Actor : LevelActors)
 			{
 				FGLTFJsonNode* JsonNode = Builder.AddUniqueNode(Actor);
 				if (JsonNode != nullptr && Builder.IsRootActor(Actor))
