[GLTFExporter] fix: make sure lightmaps are fully loaded/streamed before converting

commit 9ab9b61f67efb9e9a19c8079f96920318ebc6401
Date: Thu, 31 Mar 2022 03:07:34 +0200

    fix: make sure lightmaps are fully loaded/streamed before converting
    (cherry picked from commit a95c58841c66aa80073e7dbee21457bbc5f0ee77​)

#skipci

diff --git a/Engine/Plugins/Enterprise/GLTFExporter/MigratedFromGit/Source/GLTFExporter/Private/Tasks/GLTFDelayedTextureTasks.cpp b/Engine/Plugins/Enterprise/GLTFExporter/MigratedFromGit/Source/GLTFExporter/Private/Tasks/GLTFDelayedTextureTasks.cpp
--- a/Engine/Plugins/Enterprise/GLTFExporter/MigratedFromGit/Source/GLTFExporter/Private/Tasks/GLTFDelayedTextureTasks.cpp
+++ b/Engine/Plugins/Enterprise/GLTFExporter/MigratedFromGit/Source/GLTFExporter/Private/Tasks/GLTFDelayedTextureTasks.cpp
@@ -223,6 +223,7 @@ FString FGLTFDelayedTextureLightMapTask::GetName()
 
 void FGLTFDelayedTextureLightMapTask::Process()
 {
+	FGLTFTextureUtility::FullyLoad(LightMap);
 	LightMap->GetName(JsonTexture->Name);
 
 	// NOTE: export of lightmaps via source data is used to work around issues with
