[GLTFExporter] feat: improve loading of test levels in test framework to ensure scene data is fully initialized (and potentially avoid crashes as well)

commit b42589e221961921edb407b15b07bd889412f888
Date: Thu, 31 Mar 2022 03:13:10 +0200

    feat: improve loading of test levels in test framework to ensure scene data is fully initialized (and potentially avoid crashes as well)
    (cherry picked from commit 1e60797916c9b9361cbe9efe68f61f6154449809​)

#skipci

diff --git a/Engine/Plugins/Enterprise/GLTFExporter/MigratedFromGit/Source/GLTFExporter/Private/Converters/GLTFObjectArrayScopeGuard.h b/Engine/Plugins/Enterprise/GLTFExporter/MigratedFromGit/Source/GLTFExporter/Private/Converters/GLTFObjectArrayScopeGuard.h
deleted file mode 100644
diff --git a/Engine/Plugins/Enterprise/GLTFExporter/MigratedFromGit/Source/GLTFExporter/Public/Converters/GLTFObjectArrayScopeGuard.h b/Engine/Plugins/Enterprise/GLTFExporter/MigratedFromGit/Source/GLTFExporter/Public/Converters/GLTFObjectArrayScopeGuard.h
new file mode 100644
--- /dev/null
+++ b/Engine/Plugins/Enterprise/GLTFExporter/MigratedFromGit/Source/GLTFExporter/Public/Converters/GLTFObjectArrayScopeGuard.h
@@ -0,0 +1,31 @@
+// Copyright Epic Games, Inc. All Rights Reserved.
+
+#pragma once
+
+#include "CoreMinimal.h"
+
+template <typename ObjectType = UObject, typename = typename TEnableIf<TIsDerivedFrom<ObjectType, UObject>::Value>::Type>
+class FGLTFObjectArrayScopeGuard : public FGCObject, public TArray<ObjectType*>
+{
+public:
+
+	FGLTFObjectArrayScopeGuard()
+	{
+	}
+
+	using TArray<ObjectType*>::TArray;
+
+	/** Non-copyable */
+	FGLTFObjectArrayScopeGuard(const FGLTFObjectArrayScopeGuard&) = delete;
+	FGLTFObjectArrayScopeGuard& operator=(const FGLTFObjectArrayScopeGuard&) = delete;
+
+	virtual void AddReferencedObjects(FReferenceCollector& Collector) override
+	{
+		Collector.AddReferencedObjects<ObjectType>(*this);
+	}
+
+	virtual FString GetReferencerName() const override
+	{
+		return TEXT("FGLTFObjectArrayScopeGuard");
+	}
+};
