[GLTFExporter] chore: change reference from UE4 to UE5 in documentation

commit 9f125eabf9942808c2b796eebe4072169ff0533a
Date: Thu, 31 Mar 2022 06:42:44 +0200

    chore: change reference from UE4 to UE5 in documentation
    (cherry picked from commit fbd31fbb8536e3608fdb1c71437a4fbf52bb926f​)

#skipci

diff --git a/Engine/Plugins/Enterprise/GLTFExporter/MigratedFromGit/Docs/what-can-be-exported.md b/Engine/Plugins/Enterprise/GLTFExporter/MigratedFromGit/Docs/what-can-be-exported.md
--- a/Engine/Plugins/Enterprise/GLTFExporter/MigratedFromGit/Docs/what-can-be-exported.md
+++ b/Engine/Plugins/Enterprise/GLTFExporter/MigratedFromGit/Docs/what-can-be-exported.md
@@ -1,6 +1,6 @@
 # What can be exported?
 
-Because glTF is a data-driven format (even with extensions) not everything in UE4 can be exported to glTF. The range of supported content can be separated into primary and secondary assets. Primary assets are directly exportable by right-clicking on an asset of the following type in the Content Browser and selecting `Asset Actions -> Export...`:
+Because glTF is a data-driven format (even with extensions) not everything in Unreal can be exported to glTF. The range of supported content can be separated into primary and secondary assets. Primary assets are directly exportable by right-clicking on an asset of the following type in the Content Browser and selecting `Asset Actions -> Export...`:
 
 - [Materials](#materials)
 - [Static Meshes](#static-meshes)
@@ -18,7 +18,7 @@ Secondary assets are exported indirectly by being used or referenced in a primar
 
 ## Materials
 
-Like UE4 glTF primarily uses a metallic/roughness PBR (physically based rendering) workflow enabling accurate photorealistic materials. However, unlike UE4, glTF materials do not support arbitrary material expressions. Instead, the format only allows a single texture and/or constant for each material input. Some inputs, like `Metallic`/`Roughness` and `Base Color`/`Opacity (Mask)`, must even share the same texture. So, to convert a UE4 material to glTF, the exporter will automatically use one of the following two methods (ordered by preference):
+Like Unreal, glTF primarily uses a metallic/roughness PBR (physically based rendering) workflow enabling accurate photorealistic materials. However, unlike Unreal, glTF materials do not support arbitrary material expressions. Instead, the format only allows a single texture and/or constant for each material input. Some inputs, like `Metallic`/`Roughness` and `Base Color`/`Opacity (Mask)`, must even share the same texture. So, to convert a Unreal material to glTF, the exporter will automatically use one of the following two methods (ordered by preference):
 
 - Material expression matching
 - Material baking
@@ -37,7 +37,7 @@ Material expression matching is the quicker and more accurate approach but limit
 
 > Currently, expression matching does not support multiplying a texture sample node with a constant.
 
-However, since most UE4 material inputs use more advanced expressions, the exporter will typically fall back to material baking if expression matching fails.
+However, since most Unreal material inputs use more advanced expressions, the exporter will typically fall back to material baking if expression matching fails.
 
 ### Material baking
 
@@ -67,7 +67,7 @@ Finally, its important to note that because each input expression is evaluated p
 
 ### Shading Models
 
-Not all UE4 shading models are supported by glTF. Currently, only the following can be exported:
+Not all Unreal shading models are supported by glTF. Currently, only the following can be exported:
 - `Default Lit`
 - `Clear Coat`
 - `Unlit`
@@ -80,11 +80,11 @@ More complicated expressions (that use non-static variables) may not be fully re
 
 ![Example of Shading Model Expression that can not be evaluated](figures/bad-shadingmodel-expression.png)
 
-It is also worth remembering that while a UE4 material can use different shading models simultaneously (i.e., for different regions), glTF is limited to a single shading model per material. Like before, the exporter will in such a case simply choose one of the final shading models.
+It is also worth remembering that while a Unreal material can use different shading models simultaneously (i.e., for different regions), glTF is limited to a single shading model per material. Like before, the exporter will in such a case simply choose one of the final shading models.
 
 #### Default Lit
 
-The most common and default shading model in UE4 is `Default Lit`, which shares almost the exact same material inputs with the default shading model in glTF:
+The most common and default shading model in Unreal is `Default Lit`, which shares almost the exact same material inputs with the default shading model in glTF:
 
 - `Base Color`
 - `Metallic`
@@ -110,7 +110,7 @@ Nonetheless, there are a few noteworthy differences:
 
    This means that unlike other inputs, `Normal` and `Ambient Occlusion` cannot have a constant non-default value, unless as a 1x1 texture. The reasoning is that constant values for such inputs should be avoided.
 
-3. glTF's metallic/roughness PBR workflow does not support `Specular` (like in UE4).
+3. glTF's metallic/roughness PBR workflow does not support `Specular` (like in Unreal).
 
 #### Clear Coat
 
@@ -170,7 +170,7 @@ Because glTF does not support multiple levels of detail. The exporter will selec
 
 ### Mesh Quantization
 
-To save disk and memory space, the following vertex attributes are quantized in UE4 and glTF:
+To save disk and memory space, the following vertex attributes are quantized in Unreal and glTF:
 
 - Vertex colors (8-bit per component)
 - Vertex normals (8-bit or 16-bit per component)
@@ -189,7 +189,7 @@ Apart from the caveats regarding static meshes, skeletal meshes in glTF also hav
 
 ## Animation Sequences
 
-UE4 animation sequences are fully supported in glTF, as long a vertex skin weights are also exported. Additionally, UE animation retargeting may also be accounted for in the export.
+Unreal animation sequences are fully supported in glTF, as long a vertex skin weights are also exported. Additionally, UE animation retargeting may also be accounted for in the export.
 
 ## Level Sequences
 
@@ -199,7 +199,7 @@ Support for level sequences is restricted to transform tracks in absolute space
 
 > Please note that export of Level Variant Sets uses the glTF extension `EPIC_level_variant_sets`, which can be turned off in the export options.
 
-In UE4, variant sets can be used to configure almost any property in a scene. The exporter currently only supports the following properties:
+In Unreal, variant sets can be used to configure almost any property in a scene. The exporter currently only supports the following properties:
 
 - `Material` - changing any material asset on a static or skeletal mesh component
 - `Static Mesh` - changing the mesh asset on a static mesh component
@@ -352,7 +352,7 @@ Property   | Description
 `Rotation` | Rotation of the component.
 `Scale`    | Scale of the component.
 
-Due to differences in how UE4 applies scale, non-uniform scale may be represented differently in glTF.
+Due to differences in how Unreal applies scale, non-uniform scale may be represented differently in glTF.
 
 ### Camera Components
 
@@ -443,7 +443,7 @@ The following texture types are supported by the exporter:
 
 > Please note that export of `Light Map Texture 2D` uses the glTF extension `EPIC_lightmap_textures`, which can be turned off in the export options.
 
-To support all texture settings (like color adjustments), the exporter uses the render data (i.e., platform data) stored internally by UE4 (rather than the source data). This ensure the exported texture is identical to the texture rendered in the UE4 editor or in-game. The downside of this approach is that any artifacts introduced by UE4's compression settings will also be exported. Thus, its recommended to use compression settings `UserInterface2D` and `HDR` wherever possible.
+To support all texture settings (like color adjustments), the exporter uses the render data (i.e., platform data) stored internally by Unreal (rather than the source data). This ensure the exported texture is identical to the texture rendered in the Unreal editor or in-game. The downside of this approach is that any artifacts introduced by Unreal's compression settings will also be exported. Thus, its recommended to use compression settings `UserInterface2D` and `HDR` wherever possible.
 
 ### HDR encoding
 
diff --git a/Engine/Plugins/Enterprise/GLTFExporter/MigratedFromGit/Docs/what-is-gltf.md b/Engine/Plugins/Enterprise/GLTFExporter/MigratedFromGit/Docs/what-is-gltf.md
--- a/Engine/Plugins/Enterprise/GLTFExporter/MigratedFromGit/Docs/what-is-gltf.md
+++ b/Engine/Plugins/Enterprise/GLTFExporter/MigratedFromGit/Docs/what-is-gltf.md
@@ -16,7 +16,7 @@ A glTF file can have one of two possible file extensions, `.gltf` (JSON) or `.gl
 
 ## Extensions
 
-Due to the complexities of modern 3D graphics, an engine-independent format such as glTF cannot cover all features, especially in a powerful game engine like UE4. The good news is that the glTF base format is specifically designed, as previously mentioned, with extensibility in mind. This extensibility is facilitated using so called glTF extensions.
+Due to the complexities of modern 3D graphics, an engine-independent format such as glTF cannot cover all features, especially in a powerful game engine like Unreal. The good news is that the glTF base format is specifically designed, as previously mentioned, with extensibility in mind. This extensibility is facilitated using so called glTF extensions.
 
 It is important to note that the extensibility does come at a cost - not all applications implement every glTF extension. If an application does not support an extension, it may still be able to load and show other parts of the glTF that do not use the specific extension - unless that extension is stated as explicitly required by the glTF, in which case the application will typically fail to open the entire glTF.
 
@@ -26,7 +26,7 @@ Each glTF extension has a unique name, allowing any application to easily identi
 - Multi-vendor extensions use the reserved prefix `EXT` and are typically supported by more than one vendor (i.e., company or application).
 - Vendor-specific extensions use registered prefixes (like `EPIC`) and are typically supported by at least that vendor. Applications by other vendors may also still support such an extension.
 
-To support many of UE4's features the exporter has implemented the following extensions (all of which can be turned off in the export options):
+To support many of Unreal's features the exporter has implemented the following extensions (all of which can be turned off in the export options):
 
 Extension                   | Description
 ----------------------------|--------------------------------------------------------
@@ -35,10 +35,10 @@ Extension                   | Description
 `KHR_materials_clearcoat`   | Materials with clear coat shading model
 `KHR_mesh_quantization`     | Decrease vertex data size and precision
 `KHR_texture_transform`     | Tiling and mirroring texture coordinates
-`EPIC_lightmap_textures`    | Lightmass baked UE4-encoded lightmaps
-`EPIC_level_variant_sets`   | Scene variants by UE4's variant manager
-`EPIC_hdri_backdrops`       | UE4 backdrop actors for HDR image projection
-`EPIC_sky_spheres`          | UE4's default sky sphere actors
+`EPIC_lightmap_textures`    | Lightmass baked Unreal-encoded lightmaps
+`EPIC_level_variant_sets`   | Scene variants by Unreal's variant manager
+`EPIC_hdri_backdrops`       | Unreal backdrop actors for HDR image projection
+`EPIC_sky_spheres`          | Unreal's default sky sphere actors
 `EPIC_animation_hotspots`   | Exporter-provided actors for interactive animations
 `EPIC_animation_playback`   | Animation looping, auto play, start time, play rate
 `EPIC_camera_controls`      | Exporter-provided actors for user-controlled cameras
diff --git a/Engine/Plugins/Enterprise/GLTFExporter/MigratedFromGit/README.md b/Engine/Plugins/Enterprise/GLTFExporter/MigratedFromGit/README.md
--- a/Engine/Plugins/Enterprise/GLTFExporter/MigratedFromGit/README.md
+++ b/Engine/Plugins/Enterprise/GLTFExporter/MigratedFromGit/README.md
@@ -1,11 +1,11 @@
 # glTF Exporter plugin
 
-The glTF Exporter lets UE4 content creators showcase any scene or model on the web in full interactive real-time 3D with minimal effort.
+The glTF Exporter lets Unreal content creators showcase any scene or model on the web in full interactive real-time 3D with minimal effort.
 
 
 ## Compatibility
 
-Currently, the plugin works with Unreal Engine **4.25**, **4.26**, and **4.27**. It has primarily been tested on Windows but is designed to be platform-independent.
+Currently, the plugin works with Unreal Engine 5.0 on Windows, Mac, and Linux.
 
 
 ## Usage
diff --git a/Engine/Plugins/Enterprise/GLTFExporter/MigratedFromGit/Source/GLTFExporter/Private/Utilities/GLTFCoreUtilities.cpp b/Engine/Plugins/Enterprise/GLTFExporter/MigratedFromGit/Source/GLTFExporter/Private/Utilities/GLTFCoreUtilities.cpp
--- a/Engine/Plugins/Enterprise/GLTFExporter/MigratedFromGit/Source/GLTFExporter/Private/Utilities/GLTFCoreUtilities.cpp
+++ b/Engine/Plugins/Enterprise/GLTFExporter/MigratedFromGit/Source/GLTFExporter/Private/Utilities/GLTFCoreUtilities.cpp
@@ -11,7 +11,7 @@ float FGLTFCoreUtilities::ConvertLength(const float Length, const float Conversi
 
 FGLTFVector3 FGLTFCoreUtilities::ConvertVector(const FVector3f& Vector)
 {
-	// UE4 uses a left-handed coordinate system, with Z up.
+	// Unreal uses a left-handed coordinate system, with Z up.
 	// glTF uses a right-handed coordinate system, with Y up.
 	return { Vector.X, Vector.Z, Vector.Y };
 }
@@ -110,7 +110,7 @@ FGLTFColor3 FGLTFCoreUtilities::ConvertColor3(const FLinearColor& Color, bool bF
 
 FGLTFUInt8Color4 FGLTFCoreUtilities::ConvertColor(const FColor& Color)
 {
-	// UE4 uses ABGR or ARGB depending on endianness.
+	// Unreal uses ABGR or ARGB depending on endianness.
 	// glTF always uses RGBA independent of endianness.
 	return { Color.R, Color.G, Color.B, Color.A };
 }
@@ -122,7 +122,7 @@ FGLTFQuaternion FGLTFCoreUtilities::ConvertRotation(const FRotator3f& Rotation)
 
 FGLTFQuaternion FGLTFCoreUtilities::ConvertRotation(const FQuat4f& Rotation)
 {
-	// UE4 uses a left-handed coordinate system, with Z up.
+	// Unreal uses a left-handed coordinate system, with Z up.
 	// glTF uses a right-handed coordinate system, with Y up.
 	// Rotation = (qX, qY, qZ, qW) = (sin(angle/2) * aX, sin(angle/2) * aY, sin(angle/2) * aZ, cons(angle/2))
 	// where (aX, aY, aZ) - rotation axis, angle - rotation angle
