[GLTFExporter] refactor: make warning about "vertex colors" into regular log message

commit 832f52fda49c6c1c1040fdb242b05d74dfa9bcdc
Date: Fri, 1 Apr 2022 17:16:33 +0200

    refactor: make warning about "vertex colors" into regular log message
    (cherry picked from commit c77339d77934aa3d415cd21df9f39dc475150f65​)

#skipci

diff --git a/Engine/Plugins/Enterprise/GLTFExporter/MigratedFromGit/Source/GLTFExporter/Private/Tasks/GLTFDelayedMeshTasks.cpp b/Engine/Plugins/Enterprise/GLTFExporter/MigratedFromGit/Source/GLTFExporter/Private/Tasks/GLTFDelayedMeshTasks.cpp
--- a/Engine/Plugins/Enterprise/GLTFExporter/MigratedFromGit/Source/GLTFExporter/Private/Tasks/GLTFDelayedMeshTasks.cpp
+++ b/Engine/Plugins/Enterprise/GLTFExporter/MigratedFromGit/Source/GLTFExporter/Private/Tasks/GLTFDelayedMeshTasks.cpp
@@ -119,7 +119,7 @@ void FGLTFDelayedStaticMeshTask::Process()
 
 	if (Builder.ExportOptions->bExportVertexColors && HasVertexColors(ColorBuffer))
 	{
-		Builder.LogWarning(FString::Printf(
+		Builder.LogSuggestion(FString::Printf(
 			TEXT("Vertex colors in mesh %s will act as a multiplier for base color in glTF, regardless of material, which may produce undesirable results"),
 			*StaticMesh->GetName()));
 	}
@@ -222,7 +222,7 @@ void FGLTFDelayedSkeletalMeshTask::Process()
 
 	if (Builder.ExportOptions->bExportVertexColors && HasVertexColors(ColorBuffer))
 	{
-		Builder.LogWarning(FString::Printf(
+		Builder.LogSuggestion(FString::Printf(
 			TEXT("Vertex colors in mesh %s will act as a multiplier for base color in glTF, regardless of material, which may produce undesirable results"),
 			*SkeletalMesh->GetName()));
 	}
