[GLTFExporter] chore: bump version to 1.3.0

commit b76b4f695ee12d8a18607e8bc1f3ffbdba678ef3
Date: Fri, 1 Apr 2022 19:16:48 +0200

    chore: bump version to 1.3.0
    (cherry picked from commit 36261ceeacf0274e197000df344fe807613c3e1f​)

#skipci

diff --git a/Engine/Plugins/Enterprise/GLTFExporter/MigratedFromGit/GLTFExporter.uplugin b/Engine/Plugins/Enterprise/GLTFExporter/MigratedFromGit/GLTFExporter.uplugin
--- a/Engine/Plugins/Enterprise/GLTFExporter/MigratedFromGit/GLTFExporter.uplugin
+++ b/Engine/Plugins/Enterprise/GLTFExporter/MigratedFromGit/GLTFExporter.uplugin
@@ -1,7 +1,7 @@
 {
 	"FileVersion": 3,
-	"Version": 120,
-	"VersionName": "1.2.0",
+	"Version": 130,
+	"VersionName": "1.3.0",
 	"FriendlyName": "glTF Exporter",
 	"Description": "An exporter for Khronos glTF 2.0.",
 	"Category": "Exporters",
