[GLTFExporter] fix: fix issue with adding user data when creating proxy material

commit c6b0b28bed11c72f518b5c884e580461dc6dd28b
Date: Thu, 14 Apr 2022 21:32:42 +0200

    fix: fix issue with adding user data when creating proxy material
    (cherry picked from commit f4a9b2e4dce93ddbca4ceb2a23401b0651cdd02e​)

#skipci

diff --git a/Engine/Plugins/Enterprise/GLTFExporter/MigratedFromGit/Source/GLTFExporter/Private/Converters/GLTFMaterialProxyFactory.cpp b/Engine/Plugins/Enterprise/GLTFExporter/MigratedFromGit/Source/GLTFExporter/Private/Converters/GLTFMaterialProxyFactory.cpp
--- a/Engine/Plugins/Enterprise/GLTFExporter/MigratedFromGit/Source/GLTFExporter/Private/Converters/GLTFMaterialProxyFactory.cpp
+++ b/Engine/Plugins/Enterprise/GLTFExporter/MigratedFromGit/Source/GLTFExporter/Private/Converters/GLTFMaterialProxyFactory.cpp
@@ -45,9 +45,9 @@ UMaterialInterface* FGLTFMaterialProxyFactory::Create(UMaterialInterface* Origin
 	UMaterialInstanceConstant* ProxyMaterial = CreateInstancedMaterial(OriginalMaterial, JsonMaterial->ShadingModel);
 	if (ProxyMaterial != nullptr)
 	{
-		SetUserData(ProxyMaterial, OriginalMaterial);
+		FGLTFProxyMaterialUtilities::SetProxyMaterial(OriginalMaterial, ProxyMaterial);
 		SetBaseProperties(ProxyMaterial, OriginalMaterial);
-		SetProxyProperties(ProxyMaterial, *JsonMaterial);
+		SetProxyParameters(ProxyMaterial, *JsonMaterial);
 	}
 
 	return ProxyMaterial;
@@ -61,19 +61,6 @@ void FGLTFMaterialProxyFactory::OpenLog()
 	}
 }
 
-void FGLTFMaterialProxyFactory::SetUserData(UMaterialInstanceConstant* ProxyMaterial, UMaterialInterface* OriginalMaterial)
-{
-	UGLTFMaterialExportOptions* UserData = OriginalMaterial->GetAssetUserData<UGLTFMaterialExportOptions>();
-	if (UserData == nullptr)
-	{
-		UserData = NewObject<UGLTFMaterialExportOptions>();
-		OriginalMaterial->AddAssetUserData(UserData);
-	}
-
-	UserData->Proxy = ProxyMaterial;
-	OriginalMaterial->Modify();
-}
-
 void FGLTFMaterialProxyFactory::SetBaseProperties(UMaterialInstanceConstant* ProxyMaterial, UMaterialInterface* OriginalMaterial)
 {
 	const UMaterial* BaseMaterial = ProxyMaterial->GetMaterial();
@@ -100,7 +87,7 @@ void FGLTFMaterialProxyFactory::SetBaseProperties(UMaterialInstanceConstant* Pro
 	}
 }
 
-void FGLTFMaterialProxyFactory::SetProxyProperties(UMaterialInstanceConstant* ProxyMaterial, const FGLTFJsonMaterial& JsonMaterial)
+void FGLTFMaterialProxyFactory::SetProxyParameters(UMaterialInstanceConstant* ProxyMaterial, const FGLTFJsonMaterial& JsonMaterial)
 {
 	SetProxyParameter(ProxyMaterial, FGLTFProxyMaterialInfo::BaseColorFactor, JsonMaterial.PBRMetallicRoughness.BaseColorFactor);
 	SetProxyParameter(ProxyMaterial, FGLTFProxyMaterialInfo::BaseColor, JsonMaterial.PBRMetallicRoughness.BaseColorTexture);
diff --git a/Engine/Plugins/Enterprise/GLTFExporter/MigratedFromGit/Source/GLTFExporter/Private/Converters/GLTFMaterialProxyFactory.h b/Engine/Plugins/Enterprise/GLTFExporter/MigratedFromGit/Source/GLTFExporter/Private/Converters/GLTFMaterialProxyFactory.h
--- a/Engine/Plugins/Enterprise/GLTFExporter/MigratedFromGit/Source/GLTFExporter/Private/Converters/GLTFMaterialProxyFactory.h
+++ b/Engine/Plugins/Enterprise/GLTFExporter/MigratedFromGit/Source/GLTFExporter/Private/Converters/GLTFMaterialProxyFactory.h
@@ -33,9 +33,8 @@ private:
 		TGLTFSharedArray<FColor> Pixels;
 	};
 
-	void SetUserData(UMaterialInstanceConstant* ProxyMaterial, UMaterialInterface* OriginalMaterial);
 	void SetBaseProperties(UMaterialInstanceConstant* ProxyMaterial, UMaterialInterface* OriginalMaterial);
-	void SetProxyProperties(UMaterialInstanceConstant* ProxyMaterial, const FGLTFJsonMaterial& JsonMaterial);
+	void SetProxyParameters(UMaterialInstanceConstant* ProxyMaterial, const FGLTFJsonMaterial& JsonMaterial);
 
 	void SetProxyParameter(UMaterialInstanceConstant* ProxyMaterial, const TGLTFProxyMaterialParameterInfo<float>& ParameterInfo, float Scalar);
 	void SetProxyParameter(UMaterialInstanceConstant* ProxyMaterial, const TGLTFProxyMaterialParameterInfo<FLinearColor>& ParameterInfo, const FGLTFJsonColor3& Color);
diff --git a/Engine/Plugins/Enterprise/GLTFExporter/MigratedFromGit/Source/GLTFExporter/Private/Utilities/GLTFProxyMaterialUtilities.cpp b/Engine/Plugins/Enterprise/GLTFExporter/MigratedFromGit/Source/GLTFExporter/Private/Utilities/GLTFProxyMaterialUtilities.cpp
--- a/Engine/Plugins/Enterprise/GLTFExporter/MigratedFromGit/Source/GLTFExporter/Private/Utilities/GLTFProxyMaterialUtilities.cpp
+++ b/Engine/Plugins/Enterprise/GLTFExporter/MigratedFromGit/Source/GLTFExporter/Private/Utilities/GLTFProxyMaterialUtilities.cpp
@@ -55,7 +55,7 @@ void FGLTFProxyMaterialUtilities::SetProxyMaterial(UMaterialInterface* OriginalM
 	UGLTFMaterialExportOptions* UserData = OriginalMaterial->GetAssetUserData<UGLTFMaterialExportOptions>();
 	if (UserData == nullptr)
 	{
-		UserData = NewObject<UGLTFMaterialExportOptions>();
+		UserData = NewObject<UGLTFMaterialExportOptions>(OriginalMaterial, NAME_None, RF_Public | RF_Transactional);
 		OriginalMaterial->AddAssetUserData(UserData);
 	}
 
