[GLTFExporter] Fixed GLTFExporter build break

commit 980d705d4357d3a1e4e1b0d280b4aef94d97034a
Date: Wed, 27 Apr 2022 11:55:56 -0400

    Fixed GLTFExporter build break
    preflight none

#skipci

diff --git a/Engine/Plugins/Enterprise/GLTFExporter/MigratedFromGit/Source/GLTFExporter/Private/Converters/GLTFAnimationConverters.cpp b/Engine/Plugins/Enterprise/GLTFExporter/MigratedFromGit/Source/GLTFExporter/Private/Converters/GLTFAnimationConverters.cpp
--- a/Engine/Plugins/Enterprise/GLTFExporter/MigratedFromGit/Source/GLTFExporter/Private/Converters/GLTFAnimationConverters.cpp
+++ b/Engine/Plugins/Enterprise/GLTFExporter/MigratedFromGit/Source/GLTFExporter/Private/Converters/GLTFAnimationConverters.cpp
@@ -8,7 +8,7 @@
 
 FGLTFJsonAnimation* FGLTFAnimationConverter::Convert(FGLTFJsonNode* RootNode, const USkeletalMesh* SkeletalMesh, const UAnimSequence* AnimSequence)
 {
-	if (AnimSequence->GetRawNumberOfFrames() < 0)
+	if (AnimSequence->GetNumberOfSampledKeys() < 0)
 	{
 		// TODO: report warning
 		return nullptr;
diff --git a/Engine/Plugins/Enterprise/GLTFExporter/MigratedFromGit/Source/GLTFExporter/Private/Converters/GLTFBoneUtility.cpp b/Engine/Plugins/Enterprise/GLTFExporter/MigratedFromGit/Source/GLTFExporter/Private/Converters/GLTFBoneUtility.cpp
--- a/Engine/Plugins/Enterprise/GLTFExporter/MigratedFromGit/Source/GLTFExporter/Private/Converters/GLTFBoneUtility.cpp
+++ b/Engine/Plugins/Enterprise/GLTFExporter/MigratedFromGit/Source/GLTFExporter/Private/Converters/GLTFBoneUtility.cpp
@@ -23,7 +23,7 @@ FTransform FGLTFBoneUtility::GetBindTransform(const FReferenceSkeleton& RefSkele
 
 void FGLTFBoneUtility::GetFrameTimestamps(const UAnimSequence* AnimSequence, TArray<float>& OutFrameTimestamps)
 {
-	const int32 FrameCount = AnimSequence->GetRawNumberOfFrames();
+	const int32 FrameCount = AnimSequence->GetNumberOfSampledKeys();
 	OutFrameTimestamps.AddUninitialized(FrameCount);
 
 	const float SequenceLength = AnimSequence->SequenceLength;
diff --git a/Engine/Plugins/Enterprise/GLTFExporter/MigratedFromGit/Source/GLTFExporter/Public/Converters/GLTFSharedArray.h b/Engine/Plugins/Enterprise/GLTFExporter/MigratedFromGit/Source/GLTFExporter/Public/Converters/GLTFSharedArray.h
--- a/Engine/Plugins/Enterprise/GLTFExporter/MigratedFromGit/Source/GLTFExporter/Public/Converters/GLTFSharedArray.h
+++ b/Engine/Plugins/Enterprise/GLTFExporter/MigratedFromGit/Source/GLTFExporter/Public/Converters/GLTFSharedArray.h
@@ -4,7 +4,7 @@
 
 #include "CoreMinimal.h"
 
-template <typename ElementType, typename ArrayType = TArray<ElementType>, ESPMode Mode = ESPMode::Fast>
+template <typename ElementType, typename ArrayType = TArray<ElementType>, ESPMode Mode = ESPMode::ThreadSafe>
 class TGLTFSharedArray : public TSharedRef<ArrayType, Mode>
 {
 public:
diff --git a/Engine/Plugins/Enterprise/GLTFExporter/MigratedFromGit/Source/GLTFMaterialBaking/Private/GLTFExportMaterialProxy.h b/Engine/Plugins/Enterprise/GLTFExporter/MigratedFromGit/Source/GLTFMaterialBaking/Private/GLTFExportMaterialProxy.h
--- a/Engine/Plugins/Enterprise/GLTFExporter/MigratedFromGit/Source/GLTFMaterialBaking/Private/GLTFExportMaterialProxy.h
+++ b/Engine/Plugins/Enterprise/GLTFExporter/MigratedFromGit/Source/GLTFMaterialBaking/Private/GLTFExportMaterialProxy.h
@@ -352,6 +352,8 @@ public:
 	{
 		return UMaterial::GetDefaultMaterial(MD_Surface)->GetRenderProxy();
 	}
+
+	virtual enum EStrataBlendMode GetStrataBlendMode() const override { return EStrataBlendMode::SBM_Opaque; }
 #else
 	virtual const FMaterial& GetMaterialWithFallback(ERHIFeatureLevel::Type FeatureLevel, const FMaterialRenderProxy*& OutFallbackMaterialRenderProxy) const override
 	{
