[GLTFExporter] Replacing legacy EditorStyle calls with AppStyle

commit 0ce068fcfba82bb36de5c886fc6f879cbc81faa1
Date: Mon, 9 May 2022 16:57:09 -0400

    Replacing legacy EditorStyle calls with AppStyle
    jira none
    preflight skip
    okforgithub public, internal

#skipci

diff --git a/Engine/Plugins/Enterprise/GLTFExporter/MigratedFromGit/Source/GLTFExporter/Private/UI/GLTFExportOptionsWindow.cpp b/Engine/Plugins/Enterprise/GLTFExporter/MigratedFromGit/Source/GLTFExporter/Private/UI/GLTFExportOptionsWindow.cpp
--- a/Engine/Plugins/Enterprise/GLTFExporter/MigratedFromGit/Source/GLTFExporter/Private/UI/GLTFExportOptionsWindow.cpp
+++ b/Engine/Plugins/Enterprise/GLTFExporter/MigratedFromGit/Source/GLTFExporter/Private/UI/GLTFExportOptionsWindow.cpp
@@ -9,7 +9,7 @@
 #include "Widgets/Layout/SBox.h"
 #include "Widgets/Layout/SUniformGridPanel.h"
 #include "Widgets/Input/SButton.h"
-#include "EditorStyleSet.h"
+#include "Styling/AppStyle.h"
 #include "PropertyEditorModule.h"
 #include "IDetailsView.h"
 #include "Widgets/SWindow.h"
@@ -53,14 +53,14 @@ void SGLTFExportOptionsWindow::Construct(const FArguments& InArgs)
 			[
 				SNew(SBorder)
 				.Padding(FMargin(3))
-				.BorderImage(FEditorStyle::GetBrush("ToolPanel.GroupBorder"))
+				.BorderImage(FAppStyle::GetBrush("ToolPanel.GroupBorder"))
 				[
 					SNew(SHorizontalBox)
 					+SHorizontalBox::Slot()
 					.AutoWidth()
 					[
 						SNew(STextBlock)
-						.Font(FEditorStyle::GetFontStyle("CurveEd.LabelFont"))
+						.Font(FAppStyle::GetFontStyle("CurveEd.LabelFont"))
 						.Text(LOCTEXT("Export_CurrentFileTitle", "Current File: "))
 					]
 					+SHorizontalBox::Slot()
@@ -69,7 +69,7 @@ void SGLTFExportOptionsWindow::Construct(const FArguments& InArgs)
 					.VAlign(VAlign_Center)
 					[
 						SNew(STextBlock)
-						.Font(FEditorStyle::GetFontStyle("CurveEd.InfoFont"))
+						.Font(FAppStyle::GetFontStyle("CurveEd.InfoFont"))
 						.Text(InArgs._FullPath)
 					]
 				]
@@ -126,7 +126,7 @@ void SGLTFExportOptionsWindow::Construct(const FArguments& InArgs)
 	HeaderToolBox->SetContent(
 		SNew(SBorder)
 		.Padding(FMargin(3))
-		.BorderImage(FEditorStyle::GetBrush("ToolPanel.GroupBorder"))
+		.BorderImage(FAppStyle::GetBrush("ToolPanel.GroupBorder"))
 		[
 			SNew(SHorizontalBox)
 			+ SHorizontalBox::Slot()
