[GLTFExporter] gltf Exporter: updating documentation for scripting usage

commit 24141d9749aa6f42b0a8cd3a4ff3cdfcd7f37b7d
Date: Tue, 10 May 2022 10:29:27 -0400

    gltf Exporter: updating documentation for scripting usage

#skipci

diff --git a/Engine/Plugins/Enterprise/GLTFExporter/MigratedFromGit/Docs/scripting-the-export.md b/Engine/Plugins/Enterprise/GLTFExporter/MigratedFromGit/Docs/scripting-the-export.md
new file mode 100644
--- /dev/null
+++ b/Engine/Plugins/Enterprise/GLTFExporter/MigratedFromGit/Docs/scripting-the-export.md
@@ -0,0 +1,41 @@
+# Scripting the export
+
+In addition to using the UI to do export of assets or scene it is also possible to use the plugin to export through blueprint/python in editor or at runtime.
+
+## Editor
+
+The main function to use is `Export to GLTF`. This function takes an object to export, a destination path, the export options and selected actors.
+The object can be any of the asset that are supported cf. [What can be exported?](Docs/what-can-be-exported.md)
+
+To export a level you should pass a reference to the current world. The `Selected Actors` set must be provided even if it is empty.
+It is used in the context of a level export to define only a subset of the scene to be exported.
+
+![The blueprint function to export an asset or a scene](figures/Blueprint-Export.png)
+
+Here is a python code snippet to export all assets in subfolder as individual gltf files
+
+```python
+import unreal
+
+assetPath = '/game/models/'
+outputDir = 'c:/Temp/glTFExport/'
+glTFExportOptions = unreal.GLTFExportOptions()
+selectedActors = set()
+
+staticMestPaths = unreal.EditorAssetLibrary.list_assets(assetPath)
+for smp in staticMestPaths:
+    sm = unreal.EditorAssetLibrary.load_asset(smp)
+    if unreal.MathLibrary.class_is_child_of(sm.get_class(), unreal.StaticMesh):
+        exportpath = outputDir+sm.get_name()+'/'+sm.get_name()+'.gltf'
+        unreal.GLTFExporter.export_to_gltf(sm,exportpath,glTFExportOptions,selectedActors)
+```
+
+## Runtime
+
+As in editor a glTF export can be triggered through the call to `Export to GLTF`. 
+
+Material baking not being available at runtime users have two options:
+- prebake materials in editor and ship them in their application. Those prebake materials will then be used on export.
+- use C++ and access the GLTFBuilder API at runtime to modify the content of the glTF file before export.
+
+![Proxy material creation and resulting changes](figures/Proxy-Material.png)
\ No newline at end of file
diff --git a/Engine/Plugins/Enterprise/GLTFExporter/MigratedFromGit/Docs/what-is-gltf.md b/Engine/Plugins/Enterprise/GLTFExporter/MigratedFromGit/Docs/what-is-gltf.md
--- a/Engine/Plugins/Enterprise/GLTFExporter/MigratedFromGit/Docs/what-is-gltf.md
+++ b/Engine/Plugins/Enterprise/GLTFExporter/MigratedFromGit/Docs/what-is-gltf.md
@@ -33,6 +33,7 @@ Extension                   | Description
 `KHR_lights_punctual`       | Point, spot, and directional lights
 `KHR_materials_unlit`       | Materials with unlit shading model
 `KHR_materials_clearcoat`   | Materials with clear coat shading model
+`KHR_materials_variants`    | Compact multiple material variants per asset
 `KHR_mesh_quantization`     | Decrease vertex data size and precision
 `KHR_texture_transform`     | Tiling and mirroring texture coordinates
 `EPIC_lightmap_textures`    | Lightmass baked Unreal-encoded lightmaps
diff --git a/Engine/Plugins/Enterprise/GLTFExporter/MigratedFromGit/README.md b/Engine/Plugins/Enterprise/GLTFExporter/MigratedFromGit/README.md
--- a/Engine/Plugins/Enterprise/GLTFExporter/MigratedFromGit/README.md
+++ b/Engine/Plugins/Enterprise/GLTFExporter/MigratedFromGit/README.md
@@ -20,6 +20,7 @@ Currently, the plugin works with Unreal Engine 5.0 on Windows, Mac, and Linux.
   - In the top menu, select `File -> Export Selected...`
   - Change `Save as type` to `.gltf` (or `.glb`) and click `Save`
   - When `glTF Export Options` window is shown, click `Export`
+- Alt 3: Script the export using `Blueprint` in `Editor` or at `Runtime`, or use `Python` in `Editor`.
 
 
 ## Documentation
@@ -27,3 +28,4 @@ Currently, the plugin works with Unreal Engine 5.0 on Windows, Mac, and Linux.
 - [What is glTF?](Docs/what-is-gltf.md)
 - [What can be exported?](Docs/what-can-be-exported.md)
 - [Export options reference](Docs/export-options-reference.md)
+- [Scripting the export](Docs/scripting-the-export.md)
