[GLTFExporter] Fixed compilation error of GLTFExporter

commit 8afff5cd67d19af49dff18910d093312d069a21d
Date: Tue, 24 May 2022 11:14:36 -0400

    Fixed compilation error of GLTFExporter
    jira UE-152567
    preflight none

#skipci

diff --git a/Engine/Plugins/Enterprise/GLTFExporter/MigratedFromGit/Source/GLTFExporter/Private/Converters/GLTFBoneUtility.cpp b/Engine/Plugins/Enterprise/GLTFExporter/MigratedFromGit/Source/GLTFExporter/Private/Converters/GLTFBoneUtility.cpp
--- a/Engine/Plugins/Enterprise/GLTFExporter/MigratedFromGit/Source/GLTFExporter/Private/Converters/GLTFBoneUtility.cpp
+++ b/Engine/Plugins/Enterprise/GLTFExporter/MigratedFromGit/Source/GLTFExporter/Private/Converters/GLTFBoneUtility.cpp
@@ -26,7 +26,7 @@ void FGLTFBoneUtility::GetFrameTimestamps(const UAnimSequence* AnimSequence, TAr
 	const int32 FrameCount = AnimSequence->GetNumberOfSampledKeys();
 	OutFrameTimestamps.AddUninitialized(FrameCount);
 
-	const float SequenceLength = AnimSequence->SequenceLength;
+	const float SequenceLength = AnimSequence->GetPlayLength();
 	const float FrameLength = FrameCount > 1 ? SequenceLength / (FrameCount - 1) : 0;
 
 	for (int32 FrameIndex = 0; FrameIndex < FrameCount; ++FrameIndex)
@@ -70,7 +70,7 @@ void FGLTFBoneUtility::GetBoneTransformsByFrame(const UAnimSequence* AnimSequenc
 	FBlendedCurve Curve;
 	Curve.InitFrom(BoneContainer);
 
-	FStackCustomAttributes Attributes;
+	UE::Anim::FStackAttributeContainer Attributes;
 	FAnimationPoseData PoseData(Pose, Curve, Attributes);
 
 	for (int32 Frame = 0; Frame < FrameCount; ++Frame)
diff --git a/Engine/Plugins/Enterprise/GLTFExporter/MigratedFromGit/Source/GLTFExporter/Private/Converters/GLTFCombinedTexturePreview.cpp b/Engine/Plugins/Enterprise/GLTFExporter/MigratedFromGit/Source/GLTFExporter/Private/Converters/GLTFCombinedTexturePreview.cpp
--- a/Engine/Plugins/Enterprise/GLTFExporter/MigratedFromGit/Source/GLTFExporter/Private/Converters/GLTFCombinedTexturePreview.cpp
+++ b/Engine/Plugins/Enterprise/GLTFExporter/MigratedFromGit/Source/GLTFExporter/Private/Converters/GLTFCombinedTexturePreview.cpp
@@ -85,7 +85,7 @@ void FGLTFCombinedTexturePreview::BindShaders(
 
 	GraphicsPSOInit.BlendState = TStaticBlendState<>::GetRHI();
 
-	SetGraphicsPipelineState(RHICmdList, GraphicsPSOInit, EApplyRendertargetOption::ForceApply);
+	SetGraphicsPipelineState(RHICmdList, GraphicsPSOInit, 0, EApplyRendertargetOption::CheckApply);
 
 	VertexShader->SetParameters(RHICmdList, InTransform);
 	PixelShader->SetParameters(RHICmdList, TextureA, TextureB, ColorTransformA, ColorTransformB, BackgroundColor);
diff --git a/Engine/Plugins/Enterprise/GLTFExporter/MigratedFromGit/Source/GLTFExporter/Private/Converters/GLTFCubemapFacePreview.cpp b/Engine/Plugins/Enterprise/GLTFExporter/MigratedFromGit/Source/GLTFExporter/Private/Converters/GLTFCubemapFacePreview.cpp
--- a/Engine/Plugins/Enterprise/GLTFExporter/MigratedFromGit/Source/GLTFExporter/Private/Converters/GLTFCubemapFacePreview.cpp
+++ b/Engine/Plugins/Enterprise/GLTFExporter/MigratedFromGit/Source/GLTFExporter/Private/Converters/GLTFCubemapFacePreview.cpp
@@ -75,7 +75,7 @@ void FGLTFCubemapFacePreview::BindShaders(
 	GraphicsPSOInit.PrimitiveType = PT_TriangleList;
 	GraphicsPSOInit.BlendState = TStaticBlendState<>::GetRHI();
 
-	SetGraphicsPipelineState(RHICmdList, GraphicsPSOInit, EApplyRendertargetOption::ForceApply);
+	SetGraphicsPipelineState(RHICmdList, GraphicsPSOInit, 0, EApplyRendertargetOption::CheckApply);
 
 	VertexShader->SetParameters(RHICmdList, InTransform);
 	PixelShader->SetParameters(RHICmdList, Texture, CubeFace);
diff --git a/Engine/Plugins/Enterprise/GLTFExporter/MigratedFromGit/Source/GLTFExporter/Private/Converters/GLTFMaterialUtility.cpp b/Engine/Plugins/Enterprise/GLTFExporter/MigratedFromGit/Source/GLTFExporter/Private/Converters/GLTFMaterialUtility.cpp
--- a/Engine/Plugins/Enterprise/GLTFExporter/MigratedFromGit/Source/GLTFExporter/Private/Converters/GLTFMaterialUtility.cpp
+++ b/Engine/Plugins/Enterprise/GLTFExporter/MigratedFromGit/Source/GLTFExporter/Private/Converters/GLTFMaterialUtility.cpp
@@ -86,7 +86,7 @@ const UMaterialExpressionCustomOutput* FGLTFMaterialUtility::GetCustomOutputByNa
 {
 	// TODO: should we also search inside material functions and attribute layers?
 
-	for (const UMaterialExpression* Expression : Material->GetMaterial()->Expressions)
+	for (const TObjectPtr<UMaterialExpression>& Expression : Material->GetMaterial()->GetExpressions())
 	{
 		const UMaterialExpressionCustomOutput* CustomOutput = Cast<UMaterialExpressionCustomOutput>(Expression);
 		if (CustomOutput != nullptr && CustomOutput->GetDisplayName() == Name)
diff --git a/Engine/Plugins/Enterprise/GLTFExporter/MigratedFromGit/Source/GLTFExporter/Private/Converters/GLTFNormalMapPreview.cpp b/Engine/Plugins/Enterprise/GLTFExporter/MigratedFromGit/Source/GLTFExporter/Private/Converters/GLTFNormalMapPreview.cpp
--- a/Engine/Plugins/Enterprise/GLTFExporter/MigratedFromGit/Source/GLTFExporter/Private/Converters/GLTFNormalMapPreview.cpp
+++ b/Engine/Plugins/Enterprise/GLTFExporter/MigratedFromGit/Source/GLTFExporter/Private/Converters/GLTFNormalMapPreview.cpp
@@ -80,7 +80,7 @@ void FGLTFNormalMapPreview::BindShaders(
 
 	GraphicsPSOInit.BlendState = TStaticBlendState<>::GetRHI();
 
-	SetGraphicsPipelineState(RHICmdList, GraphicsPSOInit, EApplyRendertargetOption::ForceApply);
+	SetGraphicsPipelineState(RHICmdList, GraphicsPSOInit, 0, EApplyRendertargetOption::CheckApply);
 
 	VertexShader->SetParameters(RHICmdList, InTransform);
 	PixelShader->SetParameters(RHICmdList, Texture);
diff --git a/Engine/Plugins/Enterprise/GLTFExporter/MigratedFromGit/Source/GLTFExporter/Private/Converters/GLTFSimpleTexture2DPreview.cpp b/Engine/Plugins/Enterprise/GLTFExporter/MigratedFromGit/Source/GLTFExporter/Private/Converters/GLTFSimpleTexture2DPreview.cpp
--- a/Engine/Plugins/Enterprise/GLTFExporter/MigratedFromGit/Source/GLTFExporter/Private/Converters/GLTFSimpleTexture2DPreview.cpp
+++ b/Engine/Plugins/Enterprise/GLTFExporter/MigratedFromGit/Source/GLTFExporter/Private/Converters/GLTFSimpleTexture2DPreview.cpp
@@ -71,7 +71,7 @@ void FGLTFSimpleTexture2DPreview::BindShaders(
 
 	GraphicsPSOInit.BlendState = TStaticBlendState<>::GetRHI();
 
-	SetGraphicsPipelineState(RHICmdList, GraphicsPSOInit, EApplyRendertargetOption::ForceApply);
+	SetGraphicsPipelineState(RHICmdList, GraphicsPSOInit, 0, EApplyRendertargetOption::CheckApply);
 
 	VertexShader->SetParameters(RHICmdList, InTransform);
 	PixelShader->SetParameters(RHICmdList, Texture);
diff --git a/Engine/Plugins/Enterprise/GLTFExporter/MigratedFromGit/Source/GLTFExporter/Private/Converters/GLTFTextureUtility.cpp b/Engine/Plugins/Enterprise/GLTFExporter/MigratedFromGit/Source/GLTFExporter/Private/Converters/GLTFTextureUtility.cpp
--- a/Engine/Plugins/Enterprise/GLTFExporter/MigratedFromGit/Source/GLTFExporter/Private/Converters/GLTFTextureUtility.cpp
+++ b/Engine/Plugins/Enterprise/GLTFExporter/MigratedFromGit/Source/GLTFExporter/Private/Converters/GLTFTextureUtility.cpp
@@ -148,9 +148,9 @@ UTexture2D* FGLTFTextureUtility::CreateTransientTexture(const void* RawData, int
 	// TODO: do these temp textures need to be part of the root set to avoid garbage collection?
 	UTexture2D* Texture = UTexture2D::CreateTransient(Size.X, Size.Y, Format);
 
-	void* MipData = Texture->PlatformData->Mips[0].BulkData.Lock(LOCK_READ_WRITE);
+	void* MipData = Texture->GetPlatformData()->Mips[0].BulkData.Lock(LOCK_READ_WRITE);
 	FMemory::Memcpy(MipData, RawData, ByteLength);
-	Texture->PlatformData->Mips[0].BulkData.Unlock();
+	Texture->GetPlatformData()->Mips[0].BulkData.Unlock();
 
 	Texture->SRGB = bSRGB ? 1 : 0;
 	Texture->CompressionSettings = TC_VectorDisplacementmap; // best quality
@@ -201,8 +201,8 @@ bool FGLTFTextureUtility::DrawTexture(UTextureRenderTarget2D* OutTarget, const U
 		BatchedElementParameters = new FGLTFSimpleTexture2DPreview();
 	}
 
-	FCanvas Canvas(RenderTarget, nullptr, 0.0f, 0.0f, 0.0f, GMaxRHIFeatureLevel);
-	FCanvasTileItem TileItem(InPosition, InSource->Resource, InSize, FLinearColor::White);
+	FCanvas Canvas(RenderTarget, nullptr, FGameTime::CreateDilated(0.0f, 0.0f, 0.0f, 0.0f), GMaxRHIFeatureLevel);
+	FCanvasTileItem TileItem(InPosition, InSource->GetResource(), InSize, FLinearColor::White);
 	TileItem.BatchedElementParameters = BatchedElementParameters;
 
 	Canvas.PushAbsoluteTransform(InTransform);
@@ -245,8 +245,8 @@ UTexture2D* FGLTFTextureUtility::CreateTextureFromCubeFace(const UTextureCube* T
 	UTextureRenderTarget2D* RenderTarget = CreateRenderTarget(Size, bIsHDR);
 	FTextureRenderTargetResource* RenderTargetResource = RenderTarget->GameThread_GetRenderTargetResource();
 
-	FCanvas Canvas(RenderTargetResource, nullptr, 0.0f, 0.0f, 0.0f, GMaxRHIFeatureLevel);
-	FCanvasTileItem TileItem(FVector2D::ZeroVector, TextureCube->Resource, FLinearColor::White);
+	FCanvas Canvas(RenderTargetResource, nullptr, FGameTime::CreateDilated(0.0f, 0.0f, 0.0f, 0.0f), GMaxRHIFeatureLevel);
+	FCanvasTileItem TileItem(FVector2D::ZeroVector, TextureCube->GetResource(), FLinearColor::White);
 	TileItem.BatchedElementParameters = new FGLTFCubemapFacePreview(CubeFace);
 	TileItem.Draw(&Canvas);
 
@@ -256,7 +256,7 @@ UTexture2D* FGLTFTextureUtility::CreateTextureFromCubeFace(const UTextureCube* T
 	FlushRenderingCommands();
 
 	UTexture2D* FaceTexture;
-	FTextureRenderTarget2DResource* Resource = static_cast<FTextureRenderTarget2DResource*>(RenderTarget->Resource);
+	FTextureRenderTarget2DResource* Resource = static_cast<FTextureRenderTarget2DResource*>(RenderTarget->GetResource());
 	check(Resource != nullptr);
 
 	if (bIsHDR)
@@ -286,7 +286,7 @@ UTexture2D* FGLTFTextureUtility::CreateTextureFromCubeFace(const UTextureCube* T
 
 UTexture2D* FGLTFTextureUtility::CreateTextureFromCubeFace(const UTextureRenderTargetCube* RenderTargetCube, ECubeFace CubeFace)
 {
-	FTextureRenderTargetCubeResource* Resource = static_cast<FTextureRenderTargetCubeResource*>(RenderTargetCube->Resource);
+	FTextureRenderTargetCubeResource* Resource = static_cast<FTextureRenderTargetCubeResource*>(const_cast<UTextureRenderTargetCube*>(RenderTargetCube)->GetResource());
 	if (Resource == nullptr)
 	{
 		return nullptr;
@@ -322,7 +322,7 @@ UTexture2D* FGLTFTextureUtility::CreateTextureFromCubeFace(const UTextureRenderT
 
 bool FGLTFTextureUtility::ReadPixels(const UTextureRenderTarget2D* InRenderTarget, TArray<FColor>& OutPixels, EGLTFJsonHDREncoding Encoding)
 {
-	FTextureRenderTarget2DResource* Resource = static_cast<FTextureRenderTarget2DResource*>(InRenderTarget->Resource);
+	FTextureRenderTarget2DResource* Resource = static_cast<FTextureRenderTarget2DResource*>(const_cast<UTextureRenderTarget2D*>(InRenderTarget)->GetResource());
 	if (Resource == nullptr)
 	{
 		return false;
@@ -397,38 +397,38 @@ void FGLTFTextureUtility::EncodeRGBE(const TArray<FFloat16Color>& InPixels, TArr
 
 bool FGLTFTextureUtility::LoadPlatformData(UTexture2D* Texture)
 {
-	if (Texture->PlatformData == nullptr || Texture->PlatformData->Mips.Num() == 0)
+	if (Texture->GetPlatformData() == nullptr || Texture->GetPlatformData()->Mips.Num() == 0)
 	{
 		return false;
 	}
 
 #if WITH_EDITOR
-	if (!Texture->PlatformData->Mips[0].BulkData.IsBulkDataLoaded())
+	if (!Texture->GetPlatformData()->Mips[0].BulkData.IsBulkDataLoaded())
 	{
 		// TODO: is this correct handling?
 		Texture->ForceRebuildPlatformData();
 	}
 #endif
 
-	return Texture->PlatformData->Mips[0].BulkData.IsBulkDataLoaded();
+	return Texture->GetPlatformData()->Mips[0].BulkData.IsBulkDataLoaded();
 }
 
 bool FGLTFTextureUtility::LoadPlatformData(UTextureCube* TextureCube)
 {
-	if (TextureCube->PlatformData == nullptr || TextureCube->PlatformData->Mips.Num() == 0)
+	if (TextureCube->GetPlatformData() == nullptr || TextureCube->GetPlatformData()->Mips.Num() == 0)
 	{
 		return false;
 	}
 
 #if WITH_EDITOR
-	if (!TextureCube->PlatformData->Mips[0].BulkData.IsBulkDataLoaded())
+	if (!TextureCube->GetPlatformData()->Mips[0].BulkData.IsBulkDataLoaded())
 	{
 		// TODO: is this correct handling?
 		TextureCube->ForceRebuildPlatformData();
 	}
 #endif
 
-	return TextureCube->PlatformData->Mips[0].BulkData.IsBulkDataLoaded();
+	return TextureCube->GetPlatformData()->Mips[0].BulkData.IsBulkDataLoaded();
 }
 
 void FGLTFTextureUtility::FlipGreenChannel(TArray<FColor>& Pixels)
diff --git a/Engine/Plugins/Enterprise/GLTFExporter/MigratedFromGit/Source/GLTFExporter/Private/Converters/GLTFUVDegenerateChecker.cpp b/Engine/Plugins/Enterprise/GLTFExporter/MigratedFromGit/Source/GLTFExporter/Private/Converters/GLTFUVDegenerateChecker.cpp
--- a/Engine/Plugins/Enterprise/GLTFExporter/MigratedFromGit/Source/GLTFExporter/Private/Converters/GLTFUVDegenerateChecker.cpp
+++ b/Engine/Plugins/Enterprise/GLTFExporter/MigratedFromGit/Source/GLTFExporter/Private/Converters/GLTFUVDegenerateChecker.cpp
@@ -23,7 +23,7 @@ void FGLTFUVDegenerateChecker::Sanitize(const FMeshDescription*& Description, FG
 
 	if (Description != nullptr)
 	{
-		const int32 TexCoordCount = Description->VertexInstanceAttributes().GetAttributesRef<FVector2f>(MeshAttribute::VertexInstance::TextureCoordinate).GetNumIndices();
+		const int32 TexCoordCount = Description->VertexInstanceAttributes().GetAttributesRef<FVector2f>(MeshAttribute::VertexInstance::TextureCoordinate).GetNumChannels();
 		if (TexCoord < 0 || TexCoord >= TexCoordCount)
 		{
 			Description = nullptr;
@@ -47,9 +47,9 @@ float FGLTFUVDegenerateChecker::Convert(const FMeshDescription* Description, FGL
 
 		for (const int32 SectionIndex : SectionIndices)
 		{
-			for (const FPolygonID PolygonID : Description->GetPolygonGroupPolygons(FPolygonGroupID(SectionIndex)))
+			for (const FPolygonID PolygonID : Description->GetPolygonGroupPolygonIDs(FPolygonGroupID(SectionIndex)))
 			{
-				for (const FTriangleID TriangleID : Description->GetPolygonTriangleIDs(PolygonID))
+				for (const FTriangleID TriangleID : Description->GetPolygonTriangles(PolygonID))
 				{
 					const TArrayView<const FVertexID> TriangleVertexIDs = Description->GetTriangleVertices(TriangleID);
 					const TArrayView<const FVertexInstanceID> TriangleVertexInstanceIDs = Description->GetTriangleVertexInstances(TriangleID);
diff --git a/Engine/Plugins/Enterprise/GLTFExporter/MigratedFromGit/Source/GLTFExporter/Private/Converters/GLTFUVOverlapChecker.cpp b/Engine/Plugins/Enterprise/GLTFExporter/MigratedFromGit/Source/GLTFExporter/Private/Converters/GLTFUVOverlapChecker.cpp
--- a/Engine/Plugins/Enterprise/GLTFExporter/MigratedFromGit/Source/GLTFExporter/Private/Converters/GLTFUVOverlapChecker.cpp
+++ b/Engine/Plugins/Enterprise/GLTFExporter/MigratedFromGit/Source/GLTFExporter/Private/Converters/GLTFUVOverlapChecker.cpp
@@ -16,7 +16,7 @@ void FGLTFUVOverlapChecker::Sanitize(const FMeshDescription*& Description, FGLTF
 	{
 		const TVertexInstanceAttributesConstRef<FVector2f> VertexInstanceUVs =
 			Description->VertexInstanceAttributes().GetAttributesRef<FVector2f>(MeshAttribute::VertexInstance::TextureCoordinate);
-		const int32 TexCoordCount = VertexInstanceUVs.GetNumIndices();
+		const int32 TexCoordCount = VertexInstanceUVs.GetNumChannels();
 
 		if (TexCoord < 0 || TexCoord >= TexCoordCount)
 		{
diff --git a/Engine/Plugins/Enterprise/GLTFExporter/MigratedFromGit/Source/GLTFExporter/Private/Exporters/GLTFExporterUtility.cpp b/Engine/Plugins/Enterprise/GLTFExporter/MigratedFromGit/Source/GLTFExporter/Private/Exporters/GLTFExporterUtility.cpp
--- a/Engine/Plugins/Enterprise/GLTFExporter/MigratedFromGit/Source/GLTFExporter/Private/Exporters/GLTFExporterUtility.cpp
+++ b/Engine/Plugins/Enterprise/GLTFExporter/MigratedFromGit/Source/GLTFExporter/Private/Exporters/GLTFExporterUtility.cpp
@@ -2,7 +2,7 @@
 
 #include "Exporters/GLTFExporterUtility.h"
 #include "Materials/MaterialInstance.h"
-#include "AssetRegistryModule.h"
+#include "AssetRegistry/AssetRegistryModule.h"
 #include "Engine/Selection.h"
 
 void FGLTFExporterUtility::GetSelectedActors(TSet<AActor*>& OutSelectedActors)
diff --git a/Engine/Plugins/Enterprise/GLTFExporter/MigratedFromGit/Source/GLTFExporter/Private/Utilities/GLTFTextureUtilities.cpp b/Engine/Plugins/Enterprise/GLTFExporter/MigratedFromGit/Source/GLTFExporter/Private/Utilities/GLTFTextureUtilities.cpp
--- a/Engine/Plugins/Enterprise/GLTFExporter/MigratedFromGit/Source/GLTFExporter/Private/Utilities/GLTFTextureUtilities.cpp
+++ b/Engine/Plugins/Enterprise/GLTFExporter/MigratedFromGit/Source/GLTFExporter/Private/Utilities/GLTFTextureUtilities.cpp
@@ -42,8 +42,8 @@ void FGLTFTextureUtilities::CombineTextures(const UTexture* TextureA, const FMat
 	FCanvas Canvas(RenderTarget, nullptr, 0.0f, 0.0f, 0.0f, GMaxRHIFeatureLevel);
 
 	// TODO: use different white texture depending on if target is SRGB or not
-	const FTexture* TextureResourceA = TextureA != nullptr ? TextureA->Resource : GetWhiteTexture();
-	const FTexture* TextureResourceB = TextureB != nullptr ? TextureB->Resource : GetWhiteTexture();
+	const FTexture* TextureResourceA = TextureA != nullptr ? TextureA->GetResource() : GetWhiteTexture();
+	const FTexture* TextureResourceB = TextureB != nullptr ? TextureB->GetResource() : GetWhiteTexture();
 
 	FCanvasTileItem TileItem(FVector2D::ZeroVector, FIntPoint(OutRenderTarget->SizeX, OutRenderTarget->SizeY), FLinearColor::White);
 	const TRefCountPtr<FBatchedElementParameters> BatchedElementParameters = new FGLTFCombinedTexturePreview(TextureResourceA, TextureResourceB, ColorTransformA, ColorTransformB, BackgroundColor);
@@ -72,5 +72,5 @@ FTexture* FGLTFTextureUtilities::GetWhiteTexture()
 {
 	// TODO: explicitly ensure this asset is always cooked
 	static UTexture* Texture = LoadObject<UTexture>(nullptr, TEXT("/GLTFExporter/Textures/Proxy/T_GLTF_DefaultWhite"));
-	return Texture != nullptr ? Texture->Resource : nullptr;
+	return Texture != nullptr ? Texture->GetResource() : nullptr;
 }
diff --git a/Engine/Plugins/Enterprise/GLTFExporter/MigratedFromGit/Source/GLTFMaterialAnalyzer/Private/GLTFMaterialAnalyzer.cpp b/Engine/Plugins/Enterprise/GLTFExporter/MigratedFromGit/Source/GLTFMaterialAnalyzer/Private/GLTFMaterialAnalyzer.cpp
--- a/Engine/Plugins/Enterprise/GLTFExporter/MigratedFromGit/Source/GLTFMaterialAnalyzer/Private/GLTFMaterialAnalyzer.cpp
+++ b/Engine/Plugins/Enterprise/GLTFExporter/MigratedFromGit/Source/GLTFMaterialAnalyzer/Private/GLTFMaterialAnalyzer.cpp
@@ -4,11 +4,6 @@
 #include "GLTFProxyMaterialCompiler.h"
 #include "Materials/HLSLMaterialTranslator.h"
 
-UGLTFMaterialAnalyzer::UGLTFMaterialAnalyzer()
-{
-	ResetToDefaults();
-}
-
 void UGLTFMaterialAnalyzer::ResetToDefaults()
 {
 	Property = MP_MAX;
@@ -53,7 +48,7 @@ void UGLTFMaterialAnalyzer::AnalyzeMaterialPropertyEx(const UMaterialInterface*
 
 UMaterialExpressionCustomOutput* UGLTFMaterialAnalyzer::GetCustomOutputExpression() const
 {
-	for (UMaterialExpression* Expression : Material->GetMaterial()->Expressions)
+	for (const TObjectPtr<UMaterialExpression>& Expression : Material->GetMaterial()->GetExpressions())
 	{
 		UMaterialExpressionCustomOutput* CustomOutputExpression = Cast<UMaterialExpressionCustomOutput>(Expression);
 		if (CustomOutputExpression != nullptr && CustomOutputExpression->GetDisplayName() == CustomOutput)
diff --git a/Engine/Plugins/Enterprise/GLTFExporter/MigratedFromGit/Source/GLTFMaterialAnalyzer/Public/GLTFMaterialAnalyzer.h b/Engine/Plugins/Enterprise/GLTFExporter/MigratedFromGit/Source/GLTFMaterialAnalyzer/Public/GLTFMaterialAnalyzer.h
--- a/Engine/Plugins/Enterprise/GLTFExporter/MigratedFromGit/Source/GLTFMaterialAnalyzer/Public/GLTFMaterialAnalyzer.h
+++ b/Engine/Plugins/Enterprise/GLTFExporter/MigratedFromGit/Source/GLTFMaterialAnalyzer/Public/GLTFMaterialAnalyzer.h
@@ -17,8 +17,6 @@ public:
 
 private:
 
-	UGLTFMaterialAnalyzer();
-
 	void ResetToDefaults();
 
 	UMaterialExpressionCustomOutput* GetCustomOutputExpression() const;
@@ -29,11 +27,11 @@ private:
 
 	virtual bool IsPropertyActive(EMaterialProperty InProperty) const override;
 
-	EMaterialProperty Property;
+	EMaterialProperty Property = MP_MAX;
 	FString CustomOutput;
 
 	// ReSharper disable once CppUE4ProbableMemoryIssuesWithUObject
-	UMaterialInterface* Material;
+	UMaterialInterface* Material = nullptr;
 
-	FGLTFMaterialAnalysis* Analysis;
+	FGLTFMaterialAnalysis* Analysis = nullptr;
 };
diff --git a/Engine/Plugins/Enterprise/GLTFExporter/MigratedFromGit/Source/GLTFMaterialBaking/Private/GLTFExportMaterialProxy.h b/Engine/Plugins/Enterprise/GLTFExporter/MigratedFromGit/Source/GLTFMaterialBaking/Private/GLTFExportMaterialProxy.h
--- a/Engine/Plugins/Enterprise/GLTFExporter/MigratedFromGit/Source/GLTFMaterialBaking/Private/GLTFExportMaterialProxy.h
+++ b/Engine/Plugins/Enterprise/GLTFExporter/MigratedFromGit/Source/GLTFMaterialBaking/Private/GLTFExportMaterialProxy.h
@@ -613,7 +613,7 @@ private:
 
 	UMaterialExpressionCustomOutput* GetCustomOutputExpression(const FString& FunctionName) const
 	{
-		for (UMaterialExpression* Expression : Material->Expressions)
+		for (const TObjectPtr<UMaterialExpression>& Expression : Material->GetExpressions())
 		{
 			UMaterialExpressionCustomOutput* CustomOutputExpression = Cast<UMaterialExpressionCustomOutput>(Expression);
 			if (CustomOutputExpression && CustomOutputExpression->GetFunctionName() == FunctionName)
diff --git a/Engine/Plugins/Enterprise/GLTFExporter/MigratedFromGit/Source/GLTFMaterialBaking/Private/GLTFMaterialBakingModule.cpp b/Engine/Plugins/Enterprise/GLTFExporter/MigratedFromGit/Source/GLTFMaterialBaking/Private/GLTFMaterialBakingModule.cpp
--- a/Engine/Plugins/Enterprise/GLTFExporter/MigratedFromGit/Source/GLTFMaterialBaking/Private/GLTFMaterialBakingModule.cpp
+++ b/Engine/Plugins/Enterprise/GLTFExporter/MigratedFromGit/Source/GLTFMaterialBaking/Private/GLTFMaterialBakingModule.cpp
@@ -192,7 +192,19 @@ namespace FMaterialBakingModuleImpl
 				TextureCreateFlags |= TexCreate_SRGB;
 			}
 
-			return RHICreateTexture2D(Width, Height, Format, 1, 1, TextureCreateFlags, CreateInfo);
+			return RHICreateTexture(
+				FRHITextureCreateDesc::Create2D(CreateInfo.DebugName)
+				.SetExtent(Width, Height)
+				.SetFormat((EPixelFormat)Format)
+				.SetNumMips(1)
+				.SetNumSamples(1)
+				.SetFlags(TextureCreateFlags)
+				.SetInitialState(ERHIAccess::Unknown)
+				.SetExtData(CreateInfo.ExtData)
+				.SetBulkData(CreateInfo.BulkData)
+				.SetGPUMask(CreateInfo.GPUMask)
+				.SetClearValue(CreateInfo.ClearValueBinding)
+			);
 		}
 
 		void ReleaseStagingBufferForUnmap_AnyThread(FTexture2DRHIRef& Texture2DRHIRef)
@@ -610,7 +622,9 @@ void FGLTFMaterialBakingModule::BakeMaterials(const TArray<FGLTFMaterialDataEx*>
 							FGPUFenceRHIRef GPUFence = RHICreateGPUFence(TEXT("MaterialBackingFence"));
 
 							FResolveRect Rect(0, 0, RenderTargetResource->GetSizeX(), RenderTargetResource->GetSizeY());
+							PRAGMA_DISABLE_DEPRECATION_WARNINGS
 							RHICmdList.CopyToResolveTarget(RenderTargetResource->GetRenderTargetTexture(), StagingBufferRef, FResolveParams(Rect));
+							PRAGMA_ENABLE_DEPRECATION_WARNINGS
 							RHICmdList.WriteGPUFence(GPUFence);
 
 							// Prepare a lambda for final processing that will be executed asynchronously
diff --git a/Engine/Plugins/Enterprise/GLTFExporter/MigratedFromGit/Source/GLTFMaterialBaking/Private/MaterialRenderItem.cpp b/Engine/Plugins/Enterprise/GLTFExporter/MigratedFromGit/Source/GLTFMaterialBaking/Private/MaterialRenderItem.cpp
--- a/Engine/Plugins/Enterprise/GLTFExporter/MigratedFromGit/Source/GLTFMaterialBaking/Private/MaterialRenderItem.cpp
+++ b/Engine/Plugins/Enterprise/GLTFExporter/MigratedFromGit/Source/GLTFMaterialBaking/Private/MaterialRenderItem.cpp
@@ -136,7 +136,7 @@ void FGLTFMeshMaterialRenderItem::QueueMaterial(FCanvasRenderContext& RenderCont
 					.PreSkinnedLocalBounds(PrimitiveData->PreSkinnedLocalBounds)
 					.CustomPrimitiveData(PrimitiveData->CustomPrimitiveData)
 					.ReceivesDecals(false)
-					.DrawsVelocity(false)
+					//.DrawsVelocity(false)
 				.Build();
 
 			// Overwrite object orientation, since original calculation is derived from LocalToWorld matrix,
