[GLTFExporter] Fixed non-inclusive terminology

commit ce65229c4048e7b98b77f450abdc13eab5e05900
Date: Fri, 8 Jul 2022 17:24:45 -0400

    Fixed non-inclusive terminology
    jira UE-159037
    preflight none

#skipci

diff --git a/Engine/Plugins/Enterprise/GLTFExporter/MigratedFromGit/Source/GLTFExporter/Private/Actions/GLTFProxyAssetActions.h b/Engine/Plugins/Enterprise/GLTFExporter/MigratedFromGit/Source/GLTFExporter/Private/Actions/GLTFProxyAssetActions.h
--- a/Engine/Plugins/Enterprise/GLTFExporter/MigratedFromGit/Source/GLTFExporter/Private/Actions/GLTFProxyAssetActions.h
+++ b/Engine/Plugins/Enterprise/GLTFExporter/MigratedFromGit/Source/GLTFExporter/Private/Actions/GLTFProxyAssetActions.h
@@ -46,6 +46,7 @@ public:
 	virtual FText GetDisplayNameFromAssetData(const FAssetData& AssetData) const override { return OriginalActions->GetDisplayNameFromAssetData(AssetData); }
 	virtual void SetSupported(bool bInSupported) override { OriginalActions->SetSupported(bInSupported); }
 	virtual bool IsSupported() const override { return OriginalActions->IsSupported(); }
+	virtual FTopLevelAssetPath GetClassPathName() const override { return TopLevelAssetPath; };
 
 private:
 
@@ -54,6 +55,8 @@ private:
 	void OnCreateProxy(TArray<FWeakObjectPtr> Objects) const;
 
 	TSharedRef<IAssetTypeActions> OriginalActions;
+
+	FTopLevelAssetPath TopLevelAssetPath;
 };
 
 #endif
diff --git a/Engine/Plugins/Enterprise/GLTFExporter/MigratedFromGit/Source/GLTFExporter/Private/Converters/GLTFNodeConverters.cpp b/Engine/Plugins/Enterprise/GLTFExporter/MigratedFromGit/Source/GLTFExporter/Private/Converters/GLTFNodeConverters.cpp
--- a/Engine/Plugins/Enterprise/GLTFExporter/MigratedFromGit/Source/GLTFExporter/Private/Converters/GLTFNodeConverters.cpp
+++ b/Engine/Plugins/Enterprise/GLTFExporter/MigratedFromGit/Source/GLTFExporter/Private/Converters/GLTFNodeConverters.cpp
@@ -301,7 +301,7 @@ FGLTFJsonNode* FGLTFSkeletalBoneConverter::Convert(FGLTFJsonNode* RootNode, cons
 	const FReferenceSkeleton& RefSkeleton = SkeletalMesh->RefSkeleton;
 #endif
 
-	// TODO: add support for MasterPoseComponent?
+	// TODO: add support for [Principal]PoseComponent?
 
 	const TArray<FMeshBoneInfo>& BoneInfos = RefSkeleton.GetRefBoneInfo();
 	if (!BoneInfos.IsValidIndex(BoneIndex))
diff --git a/Engine/Plugins/Enterprise/GLTFExporter/MigratedFromGit/Source/GLTFMaterialBaking/Private/MaterialRenderItem.cpp b/Engine/Plugins/Enterprise/GLTFExporter/MigratedFromGit/Source/GLTFMaterialBaking/Private/MaterialRenderItem.cpp
--- a/Engine/Plugins/Enterprise/GLTFExporter/MigratedFromGit/Source/GLTFMaterialBaking/Private/MaterialRenderItem.cpp
+++ b/Engine/Plugins/Enterprise/GLTFExporter/MigratedFromGit/Source/GLTFMaterialBaking/Private/MaterialRenderItem.cpp
@@ -47,8 +47,11 @@ bool FGLTFMeshMaterialRenderItem::Render_RenderThread(FCanvasRenderContext& Rend
 	View.FinalPostProcessSettings.bOverride_IndirectLightingIntensity = 1;
 	View.FinalPostProcessSettings.IndirectLightingIntensity = 0.0f;
 
-	const bool bNeedsToSwitchVerticalAxis = RHINeedsToSwitchVerticalAxis(Canvas->GetShaderPlatform()) && !Canvas->GetAllowSwitchVerticalAxis();
-	check(bNeedsToSwitchVerticalAxis == false);
+	// TODO: The 2 lines below has been disabled because of compilation errors:
+	//		RHINeedsToSwitchVerticalAxis does nt exist anymore
+	//		GetAllowSwitchVerticalAxis is not a member of FCanvas anymore
+	//const bool bNeedsToSwitchVerticalAxis = RHINeedsToSwitchVerticalAxis(Canvas->GetShaderPlatform()) && !Canvas->GetAllowSwitchVerticalAxis();
+	//check(bNeedsToSwitchVerticalAxis == false);
 
 	if (Vertices.Num() && Indices.Num())
 	{
